Prevent precompiling Vulcan RHI if there is no SDK installed.

diff --git a/Engine/Source/Runtime/VulkanRHI/VulkanRHI.Build.cs b/Engine/Source/Runtime/VulkanRHI/VulkanRHI.Build.cs
--- a/Engine/Source/Runtime/VulkanRHI/VulkanRHI.Build.cs
+++ b/Engine/Source/Runtime/VulkanRHI/VulkanRHI.Build.cs
@@ -61,5 +61,9 @@ public class VulkanRHI : ModuleRules
                     );
             }
 		}
+		else
+		{
+			PrecompileForTargets = PrecompileTargetsType.None;
+		}
     }
 }
