Copying //UE4/Dev-Rendering to Dev-Main (Source //UE4/Dev-Rendering@2932636)

diff --git a/Engine/Source/Runtime/VulkanRHI/Private/VulkanCommands.cpp b/Engine/Source/Runtime/VulkanRHI/Private/VulkanCommands.cpp
--- a/Engine/Source/Runtime/VulkanRHI/Private/VulkanCommands.cpp
+++ b/Engine/Source/Runtime/VulkanRHI/Private/VulkanCommands.cpp
@@ -524,9 +524,14 @@ void FVulkanCommandListContext::RHIDrawPrimitive(uint32 PrimitiveType, uint32 Ba
 	if (!BSS.HasError())
 #endif
 	{
+#if VULKAN_USE_NEW_COMMAND_BUFFERS
+		auto* CmdBuffer = CommandBufferManager->GetActiveCmdBuffer();
+		State.PrepareDraw(CmdBuffer, UEToVulkanType((EPrimitiveType)PrimitiveType));
+		vkCmdDraw(CmdBuffer->GetHandle(), NumVertices, NumInstances, BaseVertexIndex, 0);
+#else
 		State.PrepareDraw(UEToVulkanType((EPrimitiveType)PrimitiveType));
-
 		vkCmdDraw(Device->GetPendingState().GetCommandBuffer(), NumVertices, NumInstances, BaseVertexIndex, 0);
+#endif
 	}
 
 	//if (IsImmediate())
@@ -542,7 +547,7 @@ void FVulkanCommandListContext::RHIDrawPrimitiveIndirect(uint32 PrimitiveType, F
 	SCOPE_CYCLE_COUNTER(STAT_VulkanDrawCallTime);
 
 	//@NOTE: don't prepare draw without actually drawing
-#if !PLATFORM_ANDROID
+#if 0//!PLATFORM_ANDROID
 	Device->GetPendingState().PrepareDraw(UEToVulkanType((EPrimitiveType)PrimitiveType));
 	FVulkanVertexBuffer* ArgumentBuffer = ResourceCast(ArgumentBufferRHI);
 #endif
@@ -559,7 +564,6 @@ void FVulkanCommandListContext::RHIDrawIndexedPrimitive(FIndexBufferRHIParamRef
 	uint32 NumVertices, uint32 StartIndex, uint32 NumPrimitives, uint32 NumInstances)
 {
 	check(Device);
-
 	SCOPE_CYCLE_COUNTER(STAT_VulkanDrawCallTime);
 
 	FVulkanPendingState& State = Device->GetPendingState();
@@ -569,14 +573,23 @@ void FVulkanCommandListContext::RHIDrawIndexedPrimitive(FIndexBufferRHIParamRef
 		if(!BSS.HasError())
 	#endif
 	{
-		State.PrepareDraw(UEToVulkanType((EPrimitiveType)PrimitiveType));
-
 		FVulkanIndexBuffer* IndexBuffer = ResourceCast(IndexBufferRHI);
-		vkCmdBindIndexBuffer(State.GetCommandBuffer(),
-			IndexBuffer->GetBuffer()->GetBufferHandle(), IndexBuffer->GetOffset(), IndexBuffer->IndexType);
+#if VULKAN_USE_NEW_COMMAND_BUFFERS
+		FVulkanCmdBuffer* Cmd = CommandBufferManager->GetActiveCmdBuffer();
+		VkCommandBuffer CmdBuffer = Cmd->GetHandle();
+		State.PrepareDraw(Cmd, UEToVulkanType((EPrimitiveType)PrimitiveType));
+#else
+		VkCommandBuffer CmdBuffer = State.GetCommandBuffer();
+		State.PrepareDraw(UEToVulkanType((EPrimitiveType)PrimitiveType));
+#endif
+#if VULKAN_USE_NEW_RESOURCE_MANAGEMENT
+		vkCmdBindIndexBuffer(CmdBuffer, IndexBuffer->GetBuffer()->GetBuffer(), 0, IndexBuffer->IndexType);
+#else
+		vkCmdBindIndexBuffer(CmdBuffer, IndexBuffer->GetBuffer()->GetBufferHandle(), IndexBuffer->GetOffset(), IndexBuffer->IndexType);
+#endif
 
 		uint32 NumIndices = GetVertexCountForPrimitiveCount(NumPrimitives, PrimitiveType);
-		vkCmdDrawIndexed(Device->GetPendingState().GetCommandBuffer(), NumIndices, NumInstances, StartIndex, 0, FirstInstance);
+		vkCmdDrawIndexed(CmdBuffer, NumIndices, NumInstances, StartIndex, 0, FirstInstance);
 	}
 
 	//if (IsImmediate())
@@ -591,6 +604,7 @@ void FVulkanCommandListContext::RHIDrawIndexedIndirect(FIndexBufferRHIParamRef I
 
 	SCOPE_CYCLE_COUNTER(STAT_VulkanDrawCallTime);
 
+#if 0
 	//@NOTE: don't prepare draw without actually drawing
 #if !PLATFORM_ANDROID
 	Device->GetPendingState().PrepareDraw(UEToVulkanType((EPrimitiveType)PrimitiveType));
@@ -598,6 +612,7 @@ void FVulkanCommandListContext::RHIDrawIndexedIndirect(FIndexBufferRHIParamRef I
 #if 0
 	FVulkanStructuredBuffer* ArgumentsBuffer = ResourceCast(ArgumentsBufferRHI);
 #endif
+#endif
 #endif
 	VULKAN_SIGNAL_UNIMPLEMENTED();
 
@@ -612,7 +627,7 @@ void FVulkanCommandListContext::RHIDrawIndexedPrimitiveIndirect(uint32 Primitive
 	check(Device);
 
 	SCOPE_CYCLE_COUNTER(STAT_VulkanDrawCallTime);
-
+#if 0
 	//@NOTE: don't prepare draw without actually drawing
 #if !PLATFORM_ANDROID
 	Device->GetPendingState().PrepareDraw(UEToVulkanType((EPrimitiveType)PrimitiveType));
@@ -620,6 +635,7 @@ void FVulkanCommandListContext::RHIDrawIndexedPrimitiveIndirect(uint32 Primitive
 #if 0
 	FVulkanVertexBuffer* ArgumentBuffer = ResourceCast(ArgumentBufferRHI);
 #endif
+#endif
 #endif
 	VULKAN_SIGNAL_UNIMPLEMENTED();
 
@@ -666,9 +682,14 @@ void FVulkanCommandListContext::RHIEndDrawPrimitiveUP()
 		if(!Shader.HasError())
 	#endif
 	{
+#if VULKAN_USE_NEW_COMMAND_BUFFERS
+		auto* CmdBuffer = CommandBufferManager->GetActiveCmdBuffer();
+		State.PrepareDraw(CmdBuffer, UEToVulkanType((EPrimitiveType)PendingPrimitiveType));
+		vkCmdDraw(CmdBuffer->GetHandle(), PendingNumVertices, 1, PendingMinVertexIndex, 0);
+#else
 		State.PrepareDraw(UEToVulkanType((EPrimitiveType)PendingPrimitiveType));
-
 		vkCmdDraw(State.GetCommandBuffer(), PendingNumVertices, 1, PendingMinVertexIndex, 0);
+#endif
 	}
 
 	//if (IsImmediate())
@@ -715,11 +736,18 @@ void FVulkanCommandListContext::RHIEndDrawIndexedPrimitiveUP()
 		if(!Shader.HasError())
 	#endif
 	{
+#if VULKAN_USE_NEW_COMMAND_BUFFERS
+		FVulkanCmdBuffer* CmdBuffer = CommandBufferManager->GetActiveCmdBuffer();
+		VkCommandBuffer Cmd = CmdBuffer->GetHandle();
+		State.PrepareDraw(CmdBuffer, UEToVulkanType((EPrimitiveType)PendingPrimitiveType));
+#else
+		VkCommandBuffer Cmd = State.GetCommandBuffer();
 		State.PrepareDraw(UEToVulkanType((EPrimitiveType)PendingPrimitiveType));
+#endif
 		uint32 NumIndices = GetVertexCountForPrimitiveCount(PendingNumPrimitives, PendingPrimitiveType);
-		vkCmdBindIndexBuffer(State.GetCommandBuffer(), PendingDrawPrimitiveUPIndexData.Buffer->GetBufferHandle(), PendingDrawPrimitiveUPIndexData.Offset, PendingPrimitiveIndexType);
+		vkCmdBindIndexBuffer(Cmd, PendingDrawPrimitiveUPIndexData.Buffer->GetBufferHandle(), PendingDrawPrimitiveUPIndexData.Offset, PendingPrimitiveIndexType);
 
-		vkCmdDrawIndexed(State.GetCommandBuffer(), NumIndices, 1, PendingMinVertexIndex, 0, 0);
+		vkCmdDrawIndexed(Cmd, NumIndices, 1, PendingMinVertexIndex, 0, 0);
 	}
 
 	//if (IsImmediate())
@@ -748,7 +776,11 @@ void FVulkanCommandListContext::RHIClearMRT(bool bClearColor, int32 NumClearColo
 	check(bClearColor ? NumClearColors > 0 : true);
 
 	FVulkanPendingState& State = Device->GetPendingState();
+#if VULKAN_USE_NEW_COMMAND_BUFFERS
+	State.UpdateRenderPass(CommandBufferManager->GetActiveCmdBuffer());
+#else
 	State.UpdateRenderPass();
+#endif
 	
 #if VULKAN_ALLOW_MIDPASS_CLEAR
 	VkClearRect Rect;
@@ -844,3 +876,11 @@ void FVulkanCommandListContext::RHIEnableDepthBoundsTest(bool bEnable, float Min
 {
 	VULKAN_SIGNAL_UNIMPLEMENTED();
 }
+
+void FVulkanCommandListContext::SubmitCurrentCommands()
+{
+	if (IsImmediate())
+	{
+		//#todo-rco: Will submit the cmd buffers at this point
+	}
+}
diff --git a/Engine/Source/Runtime/VulkanRHI/Private/VulkanDebug.cpp b/Engine/Source/Runtime/VulkanRHI/Private/VulkanDebug.cpp
--- a/Engine/Source/Runtime/VulkanRHI/Private/VulkanDebug.cpp
+++ b/Engine/Source/Runtime/VulkanRHI/Private/VulkanDebug.cpp
@@ -56,19 +56,6 @@ static VkBool32 DebugReportFunction(
 	}
 	else if (msgFlags & VK_DEBUG_REPORT_WARN_BIT_EXT)
 	{
-		if (FCStringAnsi::Strstr(pMsg, "in output interface has no Location or Builtin decoration"))
-		{
-			return VK_FALSE;
-		}
-
-		#if VULKAN_HLSLCC_GEN_GLSL_IN_OUT_WITHOUT_STRUCTURES
-			if(FCStringAnsi::Strstr(pLayerPrefix, "SC") && msgCode == 3 && msgFlags == 2 && objType == VK_DEBUG_REPORT_OBJECT_TYPE_DEVICE_EXT)
-			{
-				// We are aware that some of the outputs may not used or are compiled out...
-				return VK_FALSE;
-			}
-		#endif
-
 		FString Message = FString::Printf(TEXT("VulkanRHI: VK WARNING: [%s] Code %d : %s\n"), ANSI_TO_TCHAR(pLayerPrefix), msgCode, ANSI_TO_TCHAR(pMsg));
 		FPlatformMisc::LowLevelOutputDebugString(*Message);
 	}
@@ -77,6 +64,7 @@ static VkBool32 DebugReportFunction(
 	{
 		FString Message = FString::Printf(TEXT("VulkanRHI: VK INFO: [%s] Code %d : %s\n"), ANSI_TO_TCHAR(pLayerPrefix), msgCode, ANSI_TO_TCHAR(pMsg));
 		FPlatformMisc::LowLevelOutputDebugString(*Message);
+		return VK_FALSE;
 	}
 #endif
 
diff --git a/Engine/Source/Runtime/VulkanRHI/Private/VulkanDevice.cpp b/Engine/Source/Runtime/VulkanRHI/Private/VulkanDevice.cpp
--- a/Engine/Source/Runtime/VulkanRHI/Private/VulkanDevice.cpp
+++ b/Engine/Source/Runtime/VulkanRHI/Private/VulkanDevice.cpp
@@ -13,6 +13,12 @@
 FVulkanDevice::FVulkanDevice(VkPhysicalDevice InGpu)
 	: Gpu(InGpu)
 	, Device(VK_NULL_HANDLE)
+#if VULKAN_USE_NEW_RESOURCE_MANAGEMENT
+	, ResourceHeapManager(this)
+#endif
+#if VULKAN_USE_NEW_COMMAND_BUFFERS && VULKAN_USE_NEW_RESOURCE_MANAGEMENT
+	, DeferredDeletionQueue(this)
+#endif
 	, DescriptorPool(nullptr)
 	, DefaultSampler(VK_NULL_HANDLE)
 	, Queue(nullptr)
@@ -45,8 +51,11 @@ FVulkanDevice::~FVulkanDevice()
 	}
 }
 
-void FVulkanDevice::InitDevice()
+void FVulkanDevice::CreateDevice()
 {
+	check(Device == VK_NULL_HANDLE);
+
+	// Setup extension and layer info
 	VkDeviceCreateInfo DeviceInfo;
 	FMemory::Memzero(DeviceInfo);
 	DeviceInfo.sType = VK_STRUCTURE_TYPE_DEVICE_CREATE_INFO;
@@ -61,36 +70,30 @@ void FVulkanDevice::InitDevice()
 	DeviceInfo.enabledLayerCount = ValidationLayers.Num();
 	DeviceInfo.ppEnabledLayerNames = (DeviceInfo.enabledLayerCount > 0) ? ValidationLayers.GetData() : nullptr;
 
-	TArray<VkDeviceQueueCreateInfo> QueueInfos;
-	QueueInfos.AddZeroed(QueueProps.Num());
-	int32 GfxQueueIndex = -1;
-	UE_LOG(LogVulkanRHI, Display, TEXT("Found %d Queues"), QueueProps.Num());
-	TArray<float> QueuePriorities;
-	for (int32 Index = 0; Index < QueueProps.Num(); ++Index)
+	// Setup Queue info
+	TArray<VkDeviceQueueCreateInfo> QueueFamilyInfos;
+	QueueFamilyInfos.AddZeroed(QueueFamilyProps.Num());
+	int32 GfxQueueFamilyIndex = -1;
+	UE_LOG(LogVulkanRHI, Display, TEXT("Found %d Queue Families"), QueueFamilyProps.Num());
+	uint32 NumPriorities = 0;
+	for (int32 FamilyIndex = 0; FamilyIndex < QueueFamilyProps.Num(); ++FamilyIndex)
 	{
-		const VkQueueFamilyProperties& CurrProps = QueueProps[Index];
-		VkDeviceQueueCreateInfo& CurrQueue = QueueInfos[Index];
+		const VkQueueFamilyProperties& CurrProps = QueueFamilyProps[FamilyIndex];
+		VkDeviceQueueCreateInfo& CurrQueue = QueueFamilyInfos[FamilyIndex];
 		CurrQueue.sType = VK_STRUCTURE_TYPE_DEVICE_QUEUE_CREATE_INFO;
-		CurrQueue.pNext = NULL;
-		CurrQueue.queueFamilyIndex = Index;
+		CurrQueue.queueFamilyIndex = FamilyIndex;
 		CurrQueue.queueCount = CurrProps.queueCount;
-		int32 StartIndex = QueuePriorities.Num();
-		QueuePriorities.AddUninitialized(CurrProps.queueCount);
-		for (int32 QueueIndex = 0; QueueIndex < (int32)CurrProps.queueCount; ++QueueIndex)
-		{
-			QueuePriorities[QueueIndex + StartIndex] = 1.0f;
-		}
-		CurrQueue.pQueuePriorities = QueuePriorities.GetData() + StartIndex;
+		NumPriorities += CurrProps.queueCount;
 
 		if ((CurrProps.queueFlags & VK_QUEUE_GRAPHICS_BIT) == VK_QUEUE_GRAPHICS_BIT)
 		{
-			if (GfxQueueIndex == -1)
+			if (GfxQueueFamilyIndex == -1)
 			{
-				GfxQueueIndex = Index;
+				GfxQueueFamilyIndex = FamilyIndex;
 			}
 			else
 			{
-				//#todo-rco: Support for multi-gpu/choose the best gpu!
+				//#todo-rco: Support for multi-queue/choose the best queue!
 			}
 		}
 
@@ -116,166 +119,174 @@ void FVulkanDevice::InitDevice()
 
 				return Info;
 			};
-		UE_LOG(LogVulkanRHI, Display, TEXT("Queue %d: %s"), Index, *GetQueueInfoString());
+		UE_LOG(LogVulkanRHI, Display, TEXT("Queue Family %d: %d queues%s"), FamilyIndex, CurrProps.queueCount, *GetQueueInfoString());
 	}
 
-	DeviceInfo.queueCreateInfoCount = QueueProps.Num();
-	DeviceInfo.pQueueCreateInfos = QueueInfos.GetData();
+	TArray<float> QueuePriorities;
+	QueuePriorities.AddUninitialized(NumPriorities);
+	float* CurrentPriority = QueuePriorities.GetData();
+	for (int32 Index = 0; Index < QueueFamilyProps.Num(); ++Index)
+	{
+		const VkQueueFamilyProperties& CurrProps = QueueFamilyProps[Index];
+		VkDeviceQueueCreateInfo& CurrQueue = QueueFamilyInfos[Index];
+		CurrQueue.pQueuePriorities = CurrentPriority;
+		for (int32 QueueIndex = 0; QueueIndex < (int32)CurrProps.queueCount; ++QueueIndex)
+		{
+			*CurrentPriority++ = 1.0f;
+		}
+	}
+
+	DeviceInfo.queueCreateInfoCount = QueueFamilyProps.Num();
+	DeviceInfo.pQueueCreateInfos = QueueFamilyInfos.GetData();
 
+	// Create the device
 	VERIFYVULKANRESULT(vkCreateDevice(Gpu, &DeviceInfo, nullptr, &Device));
 
-	MemoryManager.SetDevice(this->Device);
+	// Create Graphics Queue, here we submit command buffers for execution
+	Queue = new FVulkanQueue(this, GfxQueueFamilyIndex, 0);
+}
 
-	// Map formats
+void FVulkanDevice::SetupFormats()
+{
+	for (uint32 Index = 0; Index < VK_FORMAT_RANGE_SIZE; ++Index)
 	{
-		for (uint32 Index = 0; Index < VK_FORMAT_RANGE_SIZE; ++Index)
-		{
-			const VkFormat Format = (VkFormat)Index;
-			FMemory::Memzero(FormatProperties[Index]);
-			vkGetPhysicalDeviceFormatProperties(Gpu, Format, &FormatProperties[Index]);
-		}
+		const VkFormat Format = (VkFormat)Index;
+		FMemory::Memzero(FormatProperties[Index]);
+		vkGetPhysicalDeviceFormatProperties(Gpu, Format, &FormatProperties[Index]);
+	}
 
-		static_assert(sizeof(VkFormat) <= sizeof(GPixelFormats[0].PlatformFormat), "PlatformFormat must be increased!");
+	static_assert(sizeof(VkFormat) <= sizeof(GPixelFormats[0].PlatformFormat), "PlatformFormat must be increased!");
 
-		// Initialize the platform pixel format map.
-		for (int32 Index = 0; Index < PF_MAX; ++Index)
-		{
-			GPixelFormats[Index].PlatformFormat = VK_FORMAT_UNDEFINED;
-			GPixelFormats[Index].Supported = false;
-
-			// Set default component mapping
-			VkComponentMapping& ComponentMapping = PixelFormatComponentMapping[Index];
-			ComponentMapping.r = VK_COMPONENT_SWIZZLE_R;
-			ComponentMapping.g = VK_COMPONENT_SWIZZLE_G;
-			ComponentMapping.b = VK_COMPONENT_SWIZZLE_B;
-			ComponentMapping.a = VK_COMPONENT_SWIZZLE_A;
-		}
+	// Initialize the platform pixel format map.
+	for (int32 Index = 0; Index < PF_MAX; ++Index)
+	{
+		GPixelFormats[Index].PlatformFormat = VK_FORMAT_UNDEFINED;
+		GPixelFormats[Index].Supported = false;
+
+		// Set default component mapping
+		VkComponentMapping& ComponentMapping = PixelFormatComponentMapping[Index];
+		ComponentMapping.r = VK_COMPONENT_SWIZZLE_R;
+		ComponentMapping.g = VK_COMPONENT_SWIZZLE_G;
+		ComponentMapping.b = VK_COMPONENT_SWIZZLE_B;
+		ComponentMapping.a = VK_COMPONENT_SWIZZLE_A;
+	}
 
-		// Default formats
-		MapFormatSupport(PF_B8G8R8A8,		VK_FORMAT_B8G8R8A8_UNORM);
-		SetComponentMapping(PF_B8G8R8A8, VK_COMPONENT_SWIZZLE_R, VK_COMPONENT_SWIZZLE_G, VK_COMPONENT_SWIZZLE_B, VK_COMPONENT_SWIZZLE_A);
+	// Default formats
+	MapFormatSupport(PF_B8G8R8A8, VK_FORMAT_B8G8R8A8_UNORM);
+	SetComponentMapping(PF_B8G8R8A8, VK_COMPONENT_SWIZZLE_R, VK_COMPONENT_SWIZZLE_G, VK_COMPONENT_SWIZZLE_B, VK_COMPONENT_SWIZZLE_A);
 
-		MapFormatSupport(PF_G8,				VK_FORMAT_R8_UNORM);
-		SetComponentMapping(PF_G8, VK_COMPONENT_SWIZZLE_R, VK_COMPONENT_SWIZZLE_ZERO, VK_COMPONENT_SWIZZLE_ZERO, VK_COMPONENT_SWIZZLE_ZERO);
+	MapFormatSupport(PF_G8, VK_FORMAT_R8_UNORM);
+	SetComponentMapping(PF_G8, VK_COMPONENT_SWIZZLE_R, VK_COMPONENT_SWIZZLE_ZERO, VK_COMPONENT_SWIZZLE_ZERO, VK_COMPONENT_SWIZZLE_ZERO);
 
-		MapFormatSupport(PF_FloatRGB, VK_FORMAT_B10G11R11_UFLOAT_PACK32);
-		SetComponentMapping(PF_FloatRGB, VK_COMPONENT_SWIZZLE_R, VK_COMPONENT_SWIZZLE_G, VK_COMPONENT_SWIZZLE_B, VK_COMPONENT_SWIZZLE_ZERO);
+	MapFormatSupport(PF_FloatRGB, VK_FORMAT_B10G11R11_UFLOAT_PACK32);
+	SetComponentMapping(PF_FloatRGB, VK_COMPONENT_SWIZZLE_R, VK_COMPONENT_SWIZZLE_G, VK_COMPONENT_SWIZZLE_B, VK_COMPONENT_SWIZZLE_ZERO);
 
-		MapFormatSupport(PF_FloatRGBA,		VK_FORMAT_R16G16B16A16_SFLOAT,	8);
-		SetComponentMapping(PF_FloatRGBA, VK_COMPONENT_SWIZZLE_R, VK_COMPONENT_SWIZZLE_G, VK_COMPONENT_SWIZZLE_B, VK_COMPONENT_SWIZZLE_A);
+	MapFormatSupport(PF_FloatRGBA, VK_FORMAT_R16G16B16A16_SFLOAT, 8);
+	SetComponentMapping(PF_FloatRGBA, VK_COMPONENT_SWIZZLE_R, VK_COMPONENT_SWIZZLE_G, VK_COMPONENT_SWIZZLE_B, VK_COMPONENT_SWIZZLE_A);
 
-		MapFormatSupport(PF_DepthStencil,	VK_FORMAT_D24_UNORM_S8_UINT);
-		if (!GPixelFormats[PF_DepthStencil].Supported)
-		{
-			MapFormatSupport(PF_DepthStencil, VK_FORMAT_D32_SFLOAT_S8_UINT);
-		}
-		SetComponentMapping(PF_DepthStencil, VK_COMPONENT_SWIZZLE_R, VK_COMPONENT_SWIZZLE_ZERO, VK_COMPONENT_SWIZZLE_ZERO, VK_COMPONENT_SWIZZLE_ZERO);
+	MapFormatSupport(PF_DepthStencil, VK_FORMAT_D24_UNORM_S8_UINT);
+	if (!GPixelFormats[PF_DepthStencil].Supported)
+	{
+		MapFormatSupport(PF_DepthStencil, VK_FORMAT_D32_SFLOAT_S8_UINT);
+	}
+	SetComponentMapping(PF_DepthStencil, VK_COMPONENT_SWIZZLE_R, VK_COMPONENT_SWIZZLE_ZERO, VK_COMPONENT_SWIZZLE_ZERO, VK_COMPONENT_SWIZZLE_ZERO);
+
+	MapFormatSupport(PF_ShadowDepth, VK_FORMAT_D16_UNORM);
+	SetComponentMapping(PF_ShadowDepth, VK_COMPONENT_SWIZZLE_R, VK_COMPONENT_SWIZZLE_ZERO, VK_COMPONENT_SWIZZLE_ZERO, VK_COMPONENT_SWIZZLE_ZERO);
 
-		MapFormatSupport(PF_ShadowDepth,	VK_FORMAT_D16_UNORM);
-		SetComponentMapping(PF_ShadowDepth, VK_COMPONENT_SWIZZLE_R, VK_COMPONENT_SWIZZLE_ZERO, VK_COMPONENT_SWIZZLE_ZERO, VK_COMPONENT_SWIZZLE_ZERO);
+	// Requirement for GPU particles
+	MapFormatSupport(PF_G32R32F, VK_FORMAT_R32G32_SFLOAT, 8);
+	SetComponentMapping(PF_G32R32F, VK_COMPONENT_SWIZZLE_G, VK_COMPONENT_SWIZZLE_R, VK_COMPONENT_SWIZZLE_ZERO, VK_COMPONENT_SWIZZLE_ZERO);
 
-		// Requirement for GPU particles
-		MapFormatSupport(PF_G32R32F,		VK_FORMAT_R32G32_SFLOAT, 8);
-		SetComponentMapping(PF_G32R32F, VK_COMPONENT_SWIZZLE_G, VK_COMPONENT_SWIZZLE_R, VK_COMPONENT_SWIZZLE_ZERO, VK_COMPONENT_SWIZZLE_ZERO);
+	MapFormatSupport(PF_A32B32G32R32F, VK_FORMAT_R32G32B32A32_SFLOAT, 16);
+	SetComponentMapping(PF_A32B32G32R32F, VK_COMPONENT_SWIZZLE_R, VK_COMPONENT_SWIZZLE_G, VK_COMPONENT_SWIZZLE_B, VK_COMPONENT_SWIZZLE_A);
 
-		MapFormatSupport(PF_A32B32G32R32F,	VK_FORMAT_R32G32B32A32_SFLOAT, 16);
-		SetComponentMapping(PF_A32B32G32R32F, VK_COMPONENT_SWIZZLE_R, VK_COMPONENT_SWIZZLE_G, VK_COMPONENT_SWIZZLE_B, VK_COMPONENT_SWIZZLE_A);
+	MapFormatSupport(PF_G16R16F, VK_FORMAT_R16G16_SFLOAT);
+	SetComponentMapping(PF_G16R16F, VK_COMPONENT_SWIZZLE_G, VK_COMPONENT_SWIZZLE_R, VK_COMPONENT_SWIZZLE_ZERO, VK_COMPONENT_SWIZZLE_ZERO);
 
-		MapFormatSupport(PF_G16R16F,		VK_FORMAT_R16G16_SFLOAT);
-		SetComponentMapping(PF_G16R16F, VK_COMPONENT_SWIZZLE_G, VK_COMPONENT_SWIZZLE_R, VK_COMPONENT_SWIZZLE_ZERO, VK_COMPONENT_SWIZZLE_ZERO);
+	MapFormatSupport(PF_R32_UINT, VK_FORMAT_R32_UINT);
+	SetComponentMapping(PF_R32_UINT, VK_COMPONENT_SWIZZLE_R, VK_COMPONENT_SWIZZLE_ZERO, VK_COMPONENT_SWIZZLE_ZERO, VK_COMPONENT_SWIZZLE_ZERO);
 
-		MapFormatSupport(PF_R32_UINT,		VK_FORMAT_R32_UINT);
-		SetComponentMapping(PF_R32_UINT, VK_COMPONENT_SWIZZLE_R, VK_COMPONENT_SWIZZLE_ZERO, VK_COMPONENT_SWIZZLE_ZERO, VK_COMPONENT_SWIZZLE_ZERO);
-		
-		#if PLATFORM_ANDROID
-			MapFormatSupport(PF_D24,			VK_FORMAT_X8_D24_UNORM_PACK32);			
-		#elif PLATFORM_DESKTOP
-			MapFormatSupport(PF_D24,			VK_FORMAT_D24_UNORM_S8_UINT);	// VK_FORMAT_D24_UNORM_X8 seems to be unsupported
+#if PLATFORM_ANDROID
+	MapFormatSupport(PF_D24, VK_FORMAT_X8_D24_UNORM_PACK32);
+#elif PLATFORM_DESKTOP
+	MapFormatSupport(PF_D24, VK_FORMAT_D24_UNORM_S8_UINT);	// VK_FORMAT_D24_UNORM_X8 seems to be unsupported
+	if (!GPixelFormats[PF_D24].Supported)
+	{
+		MapFormatSupport(PF_D24, VK_FORMAT_D32_SFLOAT_S8_UINT);
+		if (!GPixelFormats[PF_D24].Supported)
+		{
+			MapFormatSupport(PF_D24, VK_FORMAT_D32_SFLOAT);
 			if (!GPixelFormats[PF_D24].Supported)
 			{
-				MapFormatSupport(PF_D24, VK_FORMAT_D32_SFLOAT_S8_UINT);
+				MapFormatSupport(PF_D24, VK_FORMAT_D16_UNORM_S8_UINT);
 				if (!GPixelFormats[PF_D24].Supported)
 				{
-					MapFormatSupport(PF_D24, VK_FORMAT_D32_SFLOAT);
-					if (!GPixelFormats[PF_D24].Supported)
-					{
-						MapFormatSupport(PF_D24, VK_FORMAT_D16_UNORM_S8_UINT);
-						if (!GPixelFormats[PF_D24].Supported)
-						{
-							MapFormatSupport(PF_D24, VK_FORMAT_D16_UNORM);
-						}
-					}
+					MapFormatSupport(PF_D24, VK_FORMAT_D16_UNORM);
 				}
 			}
-		#else
-			#error Unsupported Vulkan platform
-		#endif
-		SetComponentMapping(PF_D24, VK_COMPONENT_SWIZZLE_R, VK_COMPONENT_SWIZZLE_ZERO, VK_COMPONENT_SWIZZLE_ZERO, VK_COMPONENT_SWIZZLE_ZERO);
-		
-		MapFormatSupport(PF_R16F,			VK_FORMAT_R16_SFLOAT);
-		SetComponentMapping(PF_R16F, VK_COMPONENT_SWIZZLE_R, VK_COMPONENT_SWIZZLE_ZERO, VK_COMPONENT_SWIZZLE_ZERO, VK_COMPONENT_SWIZZLE_ZERO);
-
-		MapFormatSupport(PF_FloatR11G11B10,	VK_FORMAT_B10G11R11_UFLOAT_PACK32, 4);
-		SetComponentMapping(PF_FloatR11G11B10, VK_COMPONENT_SWIZZLE_R, VK_COMPONENT_SWIZZLE_G, VK_COMPONENT_SWIZZLE_B, VK_COMPONENT_SWIZZLE_ZERO);
+		}
+	}
+#else
+#error Unsupported Vulkan platform
+#endif
+	SetComponentMapping(PF_D24, VK_COMPONENT_SWIZZLE_R, VK_COMPONENT_SWIZZLE_ZERO, VK_COMPONENT_SWIZZLE_ZERO, VK_COMPONENT_SWIZZLE_ZERO);
 
-		MapFormatSupport(PF_A2B10G10R10,	VK_FORMAT_A2B10G10R10_UNORM_PACK32, 4);
-		SetComponentMapping(PF_A2B10G10R10, VK_COMPONENT_SWIZZLE_R, VK_COMPONENT_SWIZZLE_G, VK_COMPONENT_SWIZZLE_B, VK_COMPONENT_SWIZZLE_A);
+	MapFormatSupport(PF_R16F, VK_FORMAT_R16_SFLOAT);
+	SetComponentMapping(PF_R16F, VK_COMPONENT_SWIZZLE_R, VK_COMPONENT_SWIZZLE_ZERO, VK_COMPONENT_SWIZZLE_ZERO, VK_COMPONENT_SWIZZLE_ZERO);
 
-		MapFormatSupport(PF_A8,				VK_FORMAT_R8_UNORM);
-		SetComponentMapping(PF_A8, VK_COMPONENT_SWIZZLE_ZERO, VK_COMPONENT_SWIZZLE_ZERO, VK_COMPONENT_SWIZZLE_ZERO, VK_COMPONENT_SWIZZLE_R);
+	MapFormatSupport(PF_FloatR11G11B10, VK_FORMAT_B10G11R11_UFLOAT_PACK32, 4);
+	SetComponentMapping(PF_FloatR11G11B10, VK_COMPONENT_SWIZZLE_R, VK_COMPONENT_SWIZZLE_G, VK_COMPONENT_SWIZZLE_B, VK_COMPONENT_SWIZZLE_ZERO);
 
-		MapFormatSupport(PF_R8G8B8A8,		VK_FORMAT_R8G8B8A8_UNORM);
-		SetComponentMapping(PF_R8G8B8A8, VK_COMPONENT_SWIZZLE_R, VK_COMPONENT_SWIZZLE_G, VK_COMPONENT_SWIZZLE_B, VK_COMPONENT_SWIZZLE_A);
+	MapFormatSupport(PF_A2B10G10R10, VK_FORMAT_A2B10G10R10_UNORM_PACK32, 4);
+	SetComponentMapping(PF_A2B10G10R10, VK_COMPONENT_SWIZZLE_R, VK_COMPONENT_SWIZZLE_G, VK_COMPONENT_SWIZZLE_B, VK_COMPONENT_SWIZZLE_A);
 
-		MapFormatSupport(PF_R8G8,			VK_FORMAT_R8G8_UNORM);
-		SetComponentMapping(PF_R8G8, VK_COMPONENT_SWIZZLE_R, VK_COMPONENT_SWIZZLE_G, VK_COMPONENT_SWIZZLE_ZERO, VK_COMPONENT_SWIZZLE_ZERO);
+	MapFormatSupport(PF_A8, VK_FORMAT_R8_UNORM);
+	SetComponentMapping(PF_A8, VK_COMPONENT_SWIZZLE_ZERO, VK_COMPONENT_SWIZZLE_ZERO, VK_COMPONENT_SWIZZLE_ZERO, VK_COMPONENT_SWIZZLE_R);
 
-		MapFormatSupport(PF_R32_FLOAT,		VK_FORMAT_R32_SFLOAT);
-		SetComponentMapping(PF_R32_FLOAT, VK_COMPONENT_SWIZZLE_R, VK_COMPONENT_SWIZZLE_ZERO, VK_COMPONENT_SWIZZLE_ZERO, VK_COMPONENT_SWIZZLE_ZERO);
+	MapFormatSupport(PF_R8G8B8A8, VK_FORMAT_R8G8B8A8_UNORM);
+	SetComponentMapping(PF_R8G8B8A8, VK_COMPONENT_SWIZZLE_R, VK_COMPONENT_SWIZZLE_G, VK_COMPONENT_SWIZZLE_B, VK_COMPONENT_SWIZZLE_A);
 
-		#if PLATFORM_DESKTOP
-			MapFormatSupport(PF_DXT1,			VK_FORMAT_BC1_RGB_UNORM_BLOCK);	// Also what OpenGL expects (RGBA instead RGB, but not SRGB)
-			SetComponentMapping(PF_DXT1, VK_COMPONENT_SWIZZLE_R, VK_COMPONENT_SWIZZLE_G, VK_COMPONENT_SWIZZLE_B, VK_COMPONENT_SWIZZLE_ONE);
+	MapFormatSupport(PF_R8G8, VK_FORMAT_R8G8_UNORM);
+	SetComponentMapping(PF_R8G8, VK_COMPONENT_SWIZZLE_R, VK_COMPONENT_SWIZZLE_G, VK_COMPONENT_SWIZZLE_ZERO, VK_COMPONENT_SWIZZLE_ZERO);
 
-			MapFormatSupport(PF_DXT3,			VK_FORMAT_BC2_UNORM_BLOCK);
-			SetComponentMapping(PF_DXT3, VK_COMPONENT_SWIZZLE_R, VK_COMPONENT_SWIZZLE_G, VK_COMPONENT_SWIZZLE_B, VK_COMPONENT_SWIZZLE_A);
+	MapFormatSupport(PF_R32_FLOAT, VK_FORMAT_R32_SFLOAT);
+	SetComponentMapping(PF_R32_FLOAT, VK_COMPONENT_SWIZZLE_R, VK_COMPONENT_SWIZZLE_ZERO, VK_COMPONENT_SWIZZLE_ZERO, VK_COMPONENT_SWIZZLE_ZERO);
 
-			MapFormatSupport(PF_DXT5,			VK_FORMAT_BC3_UNORM_BLOCK);
-			SetComponentMapping(PF_DXT5, VK_COMPONENT_SWIZZLE_R, VK_COMPONENT_SWIZZLE_G, VK_COMPONENT_SWIZZLE_B, VK_COMPONENT_SWIZZLE_A);
+#if PLATFORM_DESKTOP
+	MapFormatSupport(PF_DXT1, VK_FORMAT_BC1_RGB_UNORM_BLOCK);	// Also what OpenGL expects (RGBA instead RGB, but not SRGB)
+	SetComponentMapping(PF_DXT1, VK_COMPONENT_SWIZZLE_R, VK_COMPONENT_SWIZZLE_G, VK_COMPONENT_SWIZZLE_B, VK_COMPONENT_SWIZZLE_ONE);
 
-			MapFormatSupport(PF_BC5,			VK_FORMAT_BC5_UNORM_BLOCK);
-			SetComponentMapping(PF_BC5, VK_COMPONENT_SWIZZLE_R, VK_COMPONENT_SWIZZLE_G, VK_COMPONENT_SWIZZLE_B, VK_COMPONENT_SWIZZLE_A);
-		#elif PLATFORM_ANDROID
-			MapFormatSupport(PF_ASTC_4x4,		VK_FORMAT_ASTC_4x4_UNORM_BLOCK);
-			SetComponentMapping(PF_ASTC_4x4, VK_COMPONENT_SWIZZLE_R, VK_COMPONENT_SWIZZLE_G, VK_COMPONENT_SWIZZLE_B, VK_COMPONENT_SWIZZLE_A);
+	MapFormatSupport(PF_DXT3, VK_FORMAT_BC2_UNORM_BLOCK);
+	SetComponentMapping(PF_DXT3, VK_COMPONENT_SWIZZLE_R, VK_COMPONENT_SWIZZLE_G, VK_COMPONENT_SWIZZLE_B, VK_COMPONENT_SWIZZLE_A);
 
-			MapFormatSupport(PF_ASTC_6x6,		VK_FORMAT_ASTC_6x6_UNORM_BLOCK);
-			SetComponentMapping(PF_ASTC_6x6, VK_COMPONENT_SWIZZLE_R, VK_COMPONENT_SWIZZLE_G, VK_COMPONENT_SWIZZLE_B, VK_COMPONENT_SWIZZLE_A);
+	MapFormatSupport(PF_DXT5, VK_FORMAT_BC3_UNORM_BLOCK);
+	SetComponentMapping(PF_DXT5, VK_COMPONENT_SWIZZLE_R, VK_COMPONENT_SWIZZLE_G, VK_COMPONENT_SWIZZLE_B, VK_COMPONENT_SWIZZLE_A);
 
-			MapFormatSupport(PF_ASTC_8x8,		VK_FORMAT_ASTC_8x8_UNORM_BLOCK);
-			SetComponentMapping(PF_ASTC_8x8, VK_COMPONENT_SWIZZLE_R, VK_COMPONENT_SWIZZLE_G, VK_COMPONENT_SWIZZLE_B, VK_COMPONENT_SWIZZLE_A);
+	MapFormatSupport(PF_BC5, VK_FORMAT_BC5_UNORM_BLOCK);
+	SetComponentMapping(PF_BC5, VK_COMPONENT_SWIZZLE_R, VK_COMPONENT_SWIZZLE_G, VK_COMPONENT_SWIZZLE_B, VK_COMPONENT_SWIZZLE_A);
+#elif PLATFORM_ANDROID
+	MapFormatSupport(PF_ASTC_4x4, VK_FORMAT_ASTC_4x4_UNORM_BLOCK);
+	SetComponentMapping(PF_ASTC_4x4, VK_COMPONENT_SWIZZLE_R, VK_COMPONENT_SWIZZLE_G, VK_COMPONENT_SWIZZLE_B, VK_COMPONENT_SWIZZLE_A);
 
-			MapFormatSupport(PF_ASTC_10x10,		VK_FORMAT_ASTC_10x10_UNORM_BLOCK);
-			SetComponentMapping(PF_ASTC_10x10, VK_COMPONENT_SWIZZLE_R, VK_COMPONENT_SWIZZLE_G, VK_COMPONENT_SWIZZLE_B, VK_COMPONENT_SWIZZLE_A);
+	MapFormatSupport(PF_ASTC_6x6, VK_FORMAT_ASTC_6x6_UNORM_BLOCK);
+	SetComponentMapping(PF_ASTC_6x6, VK_COMPONENT_SWIZZLE_R, VK_COMPONENT_SWIZZLE_G, VK_COMPONENT_SWIZZLE_B, VK_COMPONENT_SWIZZLE_A);
 
-			MapFormatSupport(PF_ASTC_12x12,		VK_FORMAT_ASTC_12x12_UNORM_BLOCK);
-			SetComponentMapping(PF_ASTC_12x12, VK_COMPONENT_SWIZZLE_R, VK_COMPONENT_SWIZZLE_G, VK_COMPONENT_SWIZZLE_B, VK_COMPONENT_SWIZZLE_A);
+	MapFormatSupport(PF_ASTC_8x8, VK_FORMAT_ASTC_8x8_UNORM_BLOCK);
+	SetComponentMapping(PF_ASTC_8x8, VK_COMPONENT_SWIZZLE_R, VK_COMPONENT_SWIZZLE_G, VK_COMPONENT_SWIZZLE_B, VK_COMPONENT_SWIZZLE_A);
 
-			MapFormatSupport(PF_ETC2_RGB,		VK_FORMAT_ETC2_R8G8B8_UNORM_BLOCK);
-			SetComponentMapping(PF_ETC2_RGB, VK_COMPONENT_SWIZZLE_R, VK_COMPONENT_SWIZZLE_G, VK_COMPONENT_SWIZZLE_B, VK_COMPONENT_SWIZZLE_ONE);
+	MapFormatSupport(PF_ASTC_10x10, VK_FORMAT_ASTC_10x10_UNORM_BLOCK);
+	SetComponentMapping(PF_ASTC_10x10, VK_COMPONENT_SWIZZLE_R, VK_COMPONENT_SWIZZLE_G, VK_COMPONENT_SWIZZLE_B, VK_COMPONENT_SWIZZLE_A);
 
-			MapFormatSupport(PF_ETC2_RGBA,		VK_FORMAT_ETC2_R8G8B8A8_UNORM_BLOCK);
-			SetComponentMapping(PF_ETC2_RGBA, VK_COMPONENT_SWIZZLE_R, VK_COMPONENT_SWIZZLE_G, VK_COMPONENT_SWIZZLE_B, VK_COMPONENT_SWIZZLE_A);
-		#endif
-	}
+	MapFormatSupport(PF_ASTC_12x12, VK_FORMAT_ASTC_12x12_UNORM_BLOCK);
+	SetComponentMapping(PF_ASTC_12x12, VK_COMPONENT_SWIZZLE_R, VK_COMPONENT_SWIZZLE_G, VK_COMPONENT_SWIZZLE_B, VK_COMPONENT_SWIZZLE_A);
 
-	// Create Graphics Queue, here we submit command buffers for execution
-	Queue = new FVulkanQueue(this, GfxQueueIndex);
+	MapFormatSupport(PF_ETC2_RGB, VK_FORMAT_ETC2_R8G8B8_UNORM_BLOCK);
+	SetComponentMapping(PF_ETC2_RGB, VK_COMPONENT_SWIZZLE_R, VK_COMPONENT_SWIZZLE_G, VK_COMPONENT_SWIZZLE_B, VK_COMPONENT_SWIZZLE_ONE);
 
-	{
-		FSamplerStateInitializerRHI Default(SF_Point);
-		DefaultSampler = new FVulkanSamplerState(Default, *this);
-	}
+	MapFormatSupport(PF_ETC2_RGBA, VK_FORMAT_ETC2_R8G8B8A8_UNORM_BLOCK);
+	SetComponentMapping(PF_ETC2_RGBA, VK_COMPONENT_SWIZZLE_R, VK_COMPONENT_SWIZZLE_G, VK_COMPONENT_SWIZZLE_B, VK_COMPONENT_SWIZZLE_A);
+#endif
 }
 
 void FVulkanDevice::MapFormatSupport(EPixelFormat UEFormat, VkFormat VulkanFormat)
@@ -350,33 +361,35 @@ bool FVulkanDevice::QueryGPU(int32 DeviceIndex)
 	vkGetPhysicalDeviceQueueFamilyProperties(Gpu, &QueueCount, nullptr);
 	check(QueueCount >= 1);
 
-	QueueProps.AddUninitialized(QueueCount);
-	vkGetPhysicalDeviceQueueFamilyProperties(Gpu, &QueueCount, QueueProps.GetData());
+	QueueFamilyProps.AddUninitialized(QueueCount);
+	vkGetPhysicalDeviceQueueFamilyProperties(Gpu, &QueueCount, QueueFamilyProps.GetData());
 
 	return bDiscrete;
 }
 
 void FVulkanDevice::InitGPU(int32 DeviceIndex)
 {
+	// Query features
 	vkGetPhysicalDeviceFeatures(Gpu, &Features);
 
 	UE_LOG(LogVulkanRHI, Display, TEXT("Geometry %d Tessellation %d"), Features.geometryShader, Features.tessellationShader);
+
+	CreateDevice();
+
+	SetupFormats();
+
 	MemoryManager.Init(this);
-#if VULKAN_USE_BUFFER_HEAPS
-	ResourceAllocationManager.Init(this);
+
+#if VULKAN_USE_NEW_RESOURCE_MANAGEMENT
+	ResourceHeapManager.Init();
 #endif
 
-	InitDevice();
 
 	DescriptorPool = new FVulkanDescriptorPool(this);
 
-#if VULKAN_USE_FENCE_MANAGER
 	FenceManager.Init(this);
-#endif
 
-#if VULKAN_USE_NEW_STAGING_BUFFERS
 	StagingManager.Init(this, Queue);
-#endif
 
 	// allocate ring buffer memory
 	VBIBRingBuffer = new FVulkanRingBuffer(this, VULKAN_VBIB_RING_BUFFER_SIZE, VK_BUFFER_USAGE_VERTEX_BUFFER_BIT | VK_BUFFER_USAGE_INDEX_BUFFER_BIT | VK_BUFFER_USAGE_UNIFORM_TEXEL_BUFFER_BIT | VK_BUFFER_USAGE_STORAGE_TEXEL_BUFFER_BIT);
@@ -414,6 +427,12 @@ void FVulkanDevice::InitGPU(int32 DeviceIndex)
 	// Create Pending state, contains pipeline states such as current shader and etc..
 	PendingState = new FVulkanPendingState(this);
 
+	// Setup default resource
+	{
+		FSamplerStateInitializerRHI Default(SF_Point);
+		DefaultSampler = new FVulkanSamplerState(Default, *this);
+	}
+
 #if VULKAN_ENABLE_DRAW_MARKERS
 	VkCmdDbgMarkerBegin = (PFN_vkCmdDbgMarkerBegin)(void*)vkGetDeviceProcAddr(Device, "vkCmdDbgMarkerBegin");
 	VkCmdDbgMarkerEnd = (PFN_vkCmdDbgMarkerEnd)(void*)vkGetDeviceProcAddr(Device, "vkCmdDbgMarkerEnd");
@@ -451,18 +470,15 @@ void FVulkanDevice::Destroy()
 
 	delete DescriptorPool;
 
-#if VULKAN_USE_NEW_STAGING_BUFFERS
 	StagingManager.Deinit();
-#endif
-#if VULKAN_USE_BUFFER_HEAPS
-	ResourceAllocationManager.Deinit();
+
+#if VULKAN_USE_NEW_RESOURCE_MANAGEMENT
+	ResourceHeapManager.Deinit();
 #endif
 
 	delete Queue;
 
-#if VULKAN_USE_FENCE_MANAGER
 	FenceManager.Deinit();
-#endif
 
 	MemoryManager.Deinit();
 
@@ -475,11 +491,14 @@ void FVulkanDevice::WaitUntilIdle()
 	VERIFYVULKANRESULT(vkDeviceWaitIdle(Device));
 }
 
+#if VULKAN_USE_NEW_COMMAND_BUFFERS
+#else
 void FVulkanDevice::EndCommandBufferBlock(FVulkanCmdBuffer* CmdBuffer)
 {
 	check(CmdBuffer);
 	GetQueue()->SubmitBlocking(CmdBuffer);
 }
+#endif
 
 bool FVulkanDevice::IsFormatSupported(VkFormat Format) const
 {
diff --git a/Engine/Source/Runtime/VulkanRHI/Private/VulkanDevice.h b/Engine/Source/Runtime/VulkanRHI/Private/VulkanDevice.h
--- a/Engine/Source/Runtime/VulkanRHI/Private/VulkanDevice.h
+++ b/Engine/Source/Runtime/VulkanRHI/Private/VulkanDevice.h
@@ -9,9 +9,11 @@
 #include "VulkanMemory.h"
 
 class FVulkanDescriptorPool;
+class FVulkanCommandListContext;
 
-struct FVulkanDevice
+class FVulkanDevice
 {
+public:
 	enum
 	{
 		NumTimestampPools = 3,	// Must be the same size as the number of the backbuffer images
@@ -26,7 +28,7 @@ struct FVulkanDevice
 
 	void InitGPU(int32 DeviceIndex);
 
-	void InitDevice();
+	void CreateDevice();
 
 	void Destroy();
 
@@ -38,7 +40,10 @@ struct FVulkanDevice
 		return *PendingState;
 	}
 
+#if VULKAN_USE_NEW_COMMAND_BUFFERS
+#else
 	void EndCommandBufferBlock(FVulkanCmdBuffer* CmdBuffer);
+#endif
 
 	inline FVulkanQueue* GetQueue()
 	{
@@ -110,10 +115,19 @@ struct FVulkanDevice
 		return MemoryManager;
 	}
 
-	VulkanRHI::FResourceAllocationManager& GetResourceAllocationManager()
+#if VULKAN_USE_NEW_RESOURCE_MANAGEMENT
+	VulkanRHI::FResourceHeapManager& GetResourceHeapManager()
+	{
+		return ResourceHeapManager;
+	}
+#endif
+
+#if VULKAN_USE_NEW_COMMAND_BUFFERS && VULKAN_USE_NEW_RESOURCE_MANAGEMENT
+	VulkanRHI::FDeferredDeletionQueue& GetDeferredDeletionQueue()
 	{
-		return ResourceAllocationManager;
+		return DeferredDeletionQueue;
 	}
+#endif
 
 	VulkanRHI::FStagingManager& GetStagingManager()
 	{
@@ -125,12 +139,10 @@ struct FVulkanDevice
 		return DescriptorPool;
 	}
 
-#if VULKAN_USE_FENCE_MANAGER
 	VulkanRHI::FFenceManager& GetFenceManager()
 	{
 		return FenceManager;
 	}
-#endif
 
 	FVulkanCommandListContext& GetImmediateContext()
 	{
@@ -158,21 +170,24 @@ private:
 	VkPhysicalDeviceProperties GpuProps;
 	VkPhysicalDeviceFeatures Features;
 	
+	VkDevice Device;
+
 	VulkanRHI::FDeviceMemoryManager MemoryManager;
-	VulkanRHI::FResourceAllocationManager ResourceAllocationManager;
+#if VULKAN_USE_NEW_RESOURCE_MANAGEMENT
+	VulkanRHI::FResourceHeapManager ResourceHeapManager;
+#endif
+#if VULKAN_USE_NEW_COMMAND_BUFFERS && VULKAN_USE_NEW_RESOURCE_MANAGEMENT
+	VulkanRHI::FDeferredDeletionQueue DeferredDeletionQueue;
+#endif
 	VulkanRHI::FStagingManager StagingManager;
 
-#if VULKAN_USE_FENCE_MANAGER
 	VulkanRHI::FFenceManager FenceManager;
-#endif
-
-	VkDevice Device;
 
 	FVulkanDescriptorPool* DescriptorPool;
 
 	FVulkanSamplerState* DefaultSampler;
 
-	TArray<VkQueueFamilyProperties> QueueProps;
+	TArray<VkQueueFamilyProperties> QueueFamilyProps;
 	VkFormatProperties FormatProperties[VK_FORMAT_RANGE_SIZE];
 
 	// Nullptr if not supported
@@ -183,12 +198,13 @@ private:
 	FVulkanPendingState* PendingState;
 	VkComponentMapping PixelFormatComponentMapping[PF_MAX];
 
-	class FVulkanCommandListContext* ImmediateContext;
+	FVulkanCommandListContext* ImmediateContext;
 
 	FVulkanRingBuffer* VBIBRingBuffer;
 	FVulkanRingBuffer* UBRingBuffer;
 
 	void GetDeviceExtensions(TArray<const ANSICHAR*>& OutDeviceExtensions, TArray<const ANSICHAR*>& OutDeviceLayers);
+	void SetupFormats();
 
 #if VULKAN_ENABLE_DRAW_MARKERS
 	PFN_vkCmdDbgMarkerBegin VkCmdDbgMarkerBegin;
diff --git a/Engine/Source/Runtime/VulkanRHI/Private/VulkanIndexBuffer.cpp b/Engine/Source/Runtime/VulkanRHI/Private/VulkanIndexBuffer.cpp
--- a/Engine/Source/Runtime/VulkanRHI/Private/VulkanIndexBuffer.cpp
+++ b/Engine/Source/Runtime/VulkanRHI/Private/VulkanIndexBuffer.cpp
@@ -5,17 +5,84 @@
 =============================================================================*/
 
 #include "VulkanRHIPrivate.h"
+#include "VulkanDevice.h"
+#include "VulkanContext.h"
 
 /** Constructor */
+#if VULKAN_USE_NEW_RESOURCE_MANAGEMENT
+
+static TMap<FVulkanIndexBuffer*, VulkanRHI::FPendingBuffer2Lock> GPendingLockIBs;
+
+FVulkanIndexBuffer::FVulkanIndexBuffer(FVulkanDevice* InDevice, VkBuffer InBuffer, VulkanRHI::FResourceAllocation* InResourceAllocation, uint32 InStride, uint32 InSize, uint32 InUsage)
+	: FRHIIndexBuffer(InStride, InSize, InUsage)
+	, IndexType(InStride == 4 ? VK_INDEX_TYPE_UINT32 : VK_INDEX_TYPE_UINT16)
+{
+	Buffer = new FVulkanBuffer2(InDevice, InBuffer, InResourceAllocation);
+}
+
+#else
 FVulkanIndexBuffer::FVulkanIndexBuffer(FVulkanDevice& InDevice, uint32 InStride, uint32 InSize, uint32 InUsage)
 	: FRHIIndexBuffer(InStride, InSize, InUsage)
 	, FVulkanMultiBuffer(InDevice, InSize, (EBufferUsageFlags)InUsage, VK_BUFFER_USAGE_INDEX_BUFFER_BIT)
 	, IndexType(InStride == 4 ? VK_INDEX_TYPE_UINT32 : VK_INDEX_TYPE_UINT16)
 {
 }
+#endif
 
 FIndexBufferRHIRef FVulkanDynamicRHI::RHICreateIndexBuffer(uint32 Stride, uint32 Size, uint32 InUsage, FRHIResourceCreateInfo& CreateInfo)
 {
+#if VULKAN_USE_NEW_RESOURCE_MANAGEMENT
+	check(Size > 0);
+
+	VkBufferUsageFlags BufferUsageFlags = VK_BUFFER_USAGE_INDEX_BUFFER_BIT | VK_BUFFER_USAGE_TRANSFER_DST_BIT;
+
+	const bool bStatic = (InUsage & BUF_Static) != 0;
+	const bool bDynamic = (InUsage & BUF_Dynamic) != 0;
+	const bool bVolatile = (InUsage & BUF_Volatile) != 0;
+
+	VkBuffer Buffer = VK_NULL_HANDLE;
+	VulkanRHI::FResourceAllocation* ResourceAllocation = FVulkanBuffer2::CreateResourceAllocationAndBuffer(Device, Size, InUsage, BufferUsageFlags, Buffer);
+	FVulkanIndexBuffer* NewBuffer = nullptr;
+	if (bDynamic)
+	{
+		//#todo-rco: Allocate in the CPU staging heap
+		NewBuffer = new FVulkanIndexBuffer(Device, Buffer, ResourceAllocation, Stride, Size, InUsage);
+
+		if (CreateInfo.ResourceArray)
+		{
+			uint32 CopyDataSize = FMath::Min(Size, CreateInfo.ResourceArray->GetResourceDataSize());
+
+			// Already mapped
+			FMemory::Memcpy(ResourceAllocation->GetMappedPointer(), CreateInfo.ResourceArray->GetResourceData(), CopyDataSize);
+
+			CreateInfo.ResourceArray->Discard();
+		}
+	}
+#if 0
+	else if (bVolatile)
+	{
+		check(0);
+	}
+#endif
+	else
+	{
+		// Static buffer path
+		ensure(bStatic);
+		NewBuffer = new FVulkanIndexBuffer(Device, Buffer, ResourceAllocation, Stride, Size, InUsage);
+
+		if (CreateInfo.ResourceArray)
+		{
+			uint32 CopyDataSize = FMath::Min(Size, CreateInfo.ResourceArray->GetResourceDataSize());
+			void* Data = this->RHILockIndexBuffer(NewBuffer, 0, CopyDataSize, RLM_WriteOnly);
+			FMemory::Memcpy(Data, CreateInfo.ResourceArray->GetResourceData(), CopyDataSize);
+			this->RHIUnlockIndexBuffer(NewBuffer);
+
+			CreateInfo.ResourceArray->Discard();
+		}
+	}
+
+	return NewBuffer;
+#else
 	// make the RHI object, which will allocate memory
 	FVulkanIndexBuffer* IndexBuffer = new FVulkanIndexBuffer(*Device, Stride, Size, InUsage);
 	if (CreateInfo.ResourceArray)
@@ -31,21 +98,119 @@ FIndexBufferRHIRef FVulkanDynamicRHI::RHICreateIndexBuffer(uint32 Stride, uint32
 		// Discard the resource array's contents.
 		CreateInfo.ResourceArray->Discard();
 	}
-
 	return IndexBuffer;
+#endif
 }
 
 void* FVulkanDynamicRHI::RHILockIndexBuffer(FIndexBufferRHIParamRef IndexBufferRHI, uint32 Offset, uint32 Size, EResourceLockMode LockMode)
 {
 	FVulkanIndexBuffer* IndexBuffer = ResourceCast(IndexBufferRHI);
+#if VULKAN_USE_NEW_RESOURCE_MANAGEMENT
+	void* Data = nullptr;
+
+	const bool bIsDynamic = (IndexBuffer->GetUsage() & BUF_AnyDynamic) != 0;
+	if (bIsDynamic)
+	{
+		Data = IndexBuffer->GetResourceAllocation()->GetMappedPointer();
+	}
+	else
+	{
+		if (LockMode == RLM_ReadOnly)
+		{
+			check(0);
+		}
+		else
+		{
+			check(LockMode == RLM_WriteOnly);
+			auto* StagingBuffer = new VulkanRHI::FStagingBuffer2(Device, Size);
+			VulkanRHI::FResourceAllocation* ResourceAllocation = Device->GetResourceHeapManager().AllocateCPUStagingResource(Size, StagingBuffer->GetAlignment(), __FILE__, __LINE__);
+			StagingBuffer->BindAllocation(ResourceAllocation);
+
+			{
+				FScopeLock ScopeLock(&LockBufferCS);
+				check(!GPendingLockIBs.Contains(IndexBuffer));
+
+				VulkanRHI::FPendingBuffer2Lock PendingLock;
+				PendingLock.ResourceAllocation = ResourceAllocation;
+				PendingLock.Offset = Offset;
+				PendingLock.Size = Size;
+				PendingLock.StagingBuffer = StagingBuffer;
+
+				GPendingLockIBs.Add(IndexBuffer, PendingLock);
+			}
+
+			Data = ResourceAllocation->GetMappedPointer();
+		}
+	}
+
+	check(Data);
+	return Data;
+#else
 	return IndexBuffer->Lock(LockMode, Size, Offset);
+#endif
 }
 
 void FVulkanDynamicRHI::RHIUnlockIndexBuffer(FIndexBufferRHIParamRef IndexBufferRHI)
 {
 	FVulkanIndexBuffer* IndexBuffer = ResourceCast(IndexBufferRHI);
-
+#if VULKAN_USE_NEW_RESOURCE_MANAGEMENT
+	VulkanRHI::FResourceAllocation* ResourceAllocation = nullptr;
+	uint32 LockSize = 0;
+	uint32 LockOffset = 0;
+	TRefCountPtr<VulkanRHI::FStagingBuffer2> StagingBuffer;
+
+	const bool bIsDynamic = (IndexBuffer->GetUsage() & BUF_AnyDynamic) != 0;
+	if (bIsDynamic)
+	{
+		//#todo-rco: Flush?
+	}
+	else
+	{
+		{
+			FScopeLock ScopeLock(&LockBufferCS);
+			VulkanRHI::FPendingBuffer2Lock PendingLock = GPendingLockIBs.FindAndRemoveChecked(IndexBuffer);
+			ResourceAllocation = PendingLock.ResourceAllocation;
+			LockSize = PendingLock.Size;
+			LockOffset = PendingLock.Offset;
+			StagingBuffer = PendingLock.StagingBuffer;
+		}
+#if VULKAN_USE_NEW_COMMAND_BUFFERS
+		FVulkanCmdBuffer* Cmd = Device->GetImmediateContext().GetCommandBufferManager()->GetActiveCmdBuffer();
+		ensure(Cmd->IsOutsideRenderPass());
+		VkCommandBuffer CmdBuffer = Cmd->GetHandle();
+#else
+		VkCommandBuffer CmdBuffer = VK_NULL_HANDLE;
+		FVulkanCmdBuffer* CmdObject = nullptr;
+
+		CmdObject = Device->GetImmediateContext().GetCommandBufferManager()->Create();
+		check(CmdObject);
+		CmdObject->Begin();
+		CmdBuffer = CmdObject->GetHandle();
+#endif
+
+		VkBufferCopy Region;
+		FMemory::Memzero(Region);
+		Region.size = LockSize;
+		Region.dstOffset = LockOffset;
+		Region.srcOffset = 0;
+		vkCmdCopyBuffer(CmdBuffer, StagingBuffer.GetReference()->Buffer, IndexBuffer->GetBuffer()->GetBuffer(), 1, &Region);
+		//UpdateBuffer(ResourceAllocation, IndexBuffer->GetBuffer(), LockSize, LockOffset);
+
+#if VULKAN_USE_NEW_COMMAND_BUFFERS && VULKAN_USE_NEW_RESOURCE_MANAGEMENT
+		Device->GetDeferredDeletionQueue().EnqueueResource(Cmd, StagingBuffer);
+#else
+		//if (!State)
+		{
+			check(CmdObject);
+			//#todo-rco: TEMP!
+			Device->GetQueue()->SubmitBlocking(CmdObject);
+			Device->GetImmediateContext().GetCommandBufferManager()->Destroy(CmdObject);
+		}
+#endif
+	}
+#else
 	IndexBuffer->Unlock();
+#endif
 }
 
 
diff --git a/Engine/Source/Runtime/VulkanRHI/Private/VulkanLayers.cpp b/Engine/Source/Runtime/VulkanRHI/Private/VulkanLayers.cpp
--- a/Engine/Source/Runtime/VulkanRHI/Private/VulkanLayers.cpp
+++ b/Engine/Source/Runtime/VulkanRHI/Private/VulkanLayers.cpp
@@ -16,7 +16,7 @@ TAutoConsoleVariable<int32> GValidationCvar(
 // List of validation layers which we want to activate for the (device)-instance (used in VulkanRHI.cpp)
 static const ANSICHAR* GValidationLayersInstance[] =
 {
-#if UE_VK_API_VERSION >= VK_MAKE_VERSION(1, 0, 5)
+#if VK_API_VERSION >= VK_MAKE_VERSION(1, 0, 5)
 	"VK_LAYER_GOOGLE_threading",
 #else
 	"VK_LAYER_LUNARG_threading",
@@ -38,7 +38,7 @@ static const ANSICHAR* GValidationLayersInstance[] =
 // List of validation layers which we want to activate for the device
 static const ANSICHAR* GValidationLayersDevice[] =
 {
-#if UE_VK_API_VERSION >= VK_MAKE_VERSION(1, 0, 5)
+#if VK_API_VERSION >= VK_MAKE_VERSION(1, 0, 5)
 	"VK_LAYER_GOOGLE_threading",
 #else
 	"VK_LAYER_LUNARG_threading",
diff --git a/Engine/Source/Runtime/VulkanRHI/Private/VulkanMemory.cpp b/Engine/Source/Runtime/VulkanRHI/Private/VulkanMemory.cpp
--- a/Engine/Source/Runtime/VulkanRHI/Private/VulkanMemory.cpp
+++ b/Engine/Source/Runtime/VulkanRHI/Private/VulkanMemory.cpp
@@ -14,7 +14,6 @@ namespace VulkanRHI
 	static FCriticalSection GFenceLock;
 
 	FDeviceMemoryManager::FDeviceMemoryManager() :
-		PhysicalDeviceHandle(VK_NULL_HANDLE),
 		DeviceHandle(VK_NULL_HANDLE),
 		Device(nullptr),
 		NumAllocations(0),
@@ -35,8 +34,8 @@ namespace VulkanRHI
 		NumAllocations = 0;
 		PeakNumAllocations = 0;
 
-		PhysicalDeviceHandle = InDevice->GetPhysicalHandle();
-		vkGetPhysicalDeviceMemoryProperties(PhysicalDeviceHandle, &MemoryProperties);
+		DeviceHandle = Device->GetInstanceHandle();
+		vkGetPhysicalDeviceMemoryProperties(InDevice->GetPhysicalHandle(), &MemoryProperties);
 
 		const uint32 MaxAllocations = InDevice->GetLimits().maxMemoryAllocationCount;
 
@@ -159,6 +158,26 @@ namespace VulkanRHI
 		Allocation = nullptr;
 	}
 
+#if UE_BUILD_DEBUG || UE_BUILD_DEVELOPMENT
+	void FDeviceMemoryManager::DumpMemory()
+	{
+		UE_LOG(LogVulkanRHI, Display, TEXT("Device Memory: %d allocations on %d heaps"), NumAllocations, HeapInfos.Num());
+		for (int32 Index = 0; Index < HeapInfos.Num(); ++Index)
+		{
+			auto& HeapInfo = HeapInfos[Index];
+			UE_LOG(LogVulkanRHI, Display, TEXT("\tHeap %d, %d allocations"), Index, HeapInfo.Allocations.Num());
+			uint64 TotalSize = 0;
+			for (int32 SubIndex = 0; SubIndex < HeapInfo.Allocations.Num(); ++SubIndex)
+			{
+				FDeviceMemoryAllocation* Allocation = HeapInfo.Allocations[SubIndex];
+				//UE_LOG(LogVulkanRHI, Display, TEXT("\t\tAllocation %d Size %d"), SubIndex, Allocation->Size);
+				TotalSize += Allocation->Size;
+			}
+			UE_LOG(LogVulkanRHI, Display, TEXT("\t\tTotal Allocated %u"), TotalSize);
+		}
+		NumAllocations = 0;
+	}
+#endif
 	FDeviceMemoryAllocation::~FDeviceMemoryAllocation()
 	{
 		checkf(bFreedBySystem, TEXT("Memory has to released calling FDeviceMemory::Free()!"));
@@ -167,110 +186,314 @@ namespace VulkanRHI
 	void* FDeviceMemoryAllocation::Map(VkDeviceSize InSize, VkDeviceSize Offset)
 	{
 		check(bCanBeMapped);
-		check(!bIsMapped);
+		check(!MappedPointer);
 		check(InSize + Offset <= Size);
 
-		void* MappedMemory = nullptr;
-
-		VERIFYVULKANRESULT(vkMapMemory(DeviceHandle, Handle, Offset, InSize, 0, &MappedMemory));
-		bIsMapped = true;
-
-		return MappedMemory;
+		VERIFYVULKANRESULT(vkMapMemory(DeviceHandle, Handle, Offset, InSize, 0, &MappedPointer));
+		return MappedPointer;
 	}
 
 	void FDeviceMemoryAllocation::Unmap()
 	{
-		check(bIsMapped);
+		check(MappedPointer);
 		vkUnmapMemory(DeviceHandle, Handle);
-		bIsMapped = false;
+		MappedPointer = nullptr;
 	}
 
 
-	void* FResourceAllocationInfo::Map(VkDeviceSize InSize, VkDeviceSize InOffset)
+	FResourceAllocation::FResourceAllocation(FResourceHeapPage* InOwner, FDeviceMemoryAllocation* InDeviceMemoryAllocation, uint32 InSize, uint32 InOffset, char* InFile, uint32 InLine)
+		: Owner(InOwner)
+		, DeviceMemoryAllocation(InDeviceMemoryAllocation)
+		, Size(InSize)
+		, Offset(InOffset)
+#if VULKAN_TRACK_MEMORY_USAGE
+		, File(InFile)
+		, Line(InLine)
+#endif
+	{
+	}
+
+	FResourceAllocation::~FResourceAllocation()
 	{
-		return Allocation->Map(InSize, InOffset + Offset);
+		Owner->ReleaseAllocation(this);
 	}
 
-	void FResourceAllocationInfo::Unmap()
+	void* FResourceAllocation::GetMappedPointer()
 	{
-		Allocation->Unmap();
+		check(DeviceMemoryAllocation->CanBeMapped());
+		check(DeviceMemoryAllocation->IsMapped());
+		return (uint8*)DeviceMemoryAllocation->GetMappedPointer() + Offset;
 	}
 
-	bool FResourceAllocationManager::FResourceEntry::TryAllocate(VkDeviceSize Size, VkDeviceSize Alignment, VkDeviceSize* OutOffset)
+
+	FResourceHeapPage::~FResourceHeapPage()
 	{
-		VkDeviceSize RequiredOffset = Align(FreeOffset, (int32)Alignment);
-		if (RequiredOffset + Size <= Allocation->GetSize())
+		check(!DeviceMemoryAllocation);
+	}
+
+	FResourceAllocation* FResourceHeapPage::TryAllocate(uint32 Size, uint32 Alignment, char* File, uint32 Line)
+	{
+		for (int32 Index = 0; Index < FreeList.Num(); ++Index)
 		{
-			*OutOffset = RequiredOffset;
-			FreeOffset = RequiredOffset + Size;
-			return true;
+			FPair& Entry = FreeList[Index];
+			uint32 NewOffset = Align(Entry.Offset, Alignment);
+			uint32 OffsetAdjustment = NewOffset - Entry.Offset;
+			if (OffsetAdjustment + Size <= Entry.Size)
+			{
+				FScopeLock ScopeLock(&GAllocationLock);
+				uint32 NewFreeOffset = NewOffset + Size;
+				if (OffsetAdjustment + Size < Entry.Size)
+				{
+					// Modify current Entry
+					Entry.Offset = NewFreeOffset;
+					Entry.Size -= Size;
+				}
+				else
+				{
+					check(OffsetAdjustment + Size == Entry.Size);
+					// Remove the complete entry
+					FreeList.RemoveAtSwap(Index, 1, false);
+				}
+
+				UsedSize += Size;
+
+				auto* NewResourceAllocation = new FResourceAllocation(this, DeviceMemoryAllocation, Size, NewOffset, File, Line);
+				ResourceAllocations.Add(NewResourceAllocation);
+
+				PeakNumAllocations = FMath::Max(PeakNumAllocations, ResourceAllocations.Num());
+				return NewResourceAllocation;
+			}
 		}
 
-		return false;
+		return nullptr;
 	}
 
-	FResourceAllocationManager::FResourceAllocationManager()
-		: Device(nullptr)
+	void FResourceHeapPage::ReleaseAllocation(FResourceAllocation* Allocation)
 	{
+		{
+			FScopeLock ScopeLock(&GAllocationLock);
+			ResourceAllocations.RemoveSingleSwap(Allocation, false);
+
+			FPair NewFree;
+			NewFree.Offset = Allocation->Offset;
+			NewFree.Size = Allocation->Size;
+			FreeList.Add(NewFree);
+		}
+
+		UsedSize -= Allocation->Size;
+		check(UsedSize >= 0);
+
+		if (JoinFreeBlocks())
+		{
+			Owner->FreePage(this);
+		}
 	}
 
-	FResourceAllocationManager::~FResourceAllocationManager()
+	bool FResourceHeapPage::JoinFreeBlocks()
 	{
+		FScopeLock ScopeLock(&GAllocationLock);
+		if (FreeList.Num() > 1)
+		{
+			FreeList.Sort();
+
+			for (int32 Index = FreeList.Num() - 1; Index > 0; --Index)
+			{
+				auto& Current = FreeList[Index];
+				auto& Prev = FreeList[Index - 1];
+				if (Prev.Offset + Prev.Size == Current.Offset)
+				{
+					Prev.Size += Current.Size;
+					FreeList.RemoveAt(Index, 1, false);
+				}
+			}
+		}
+
+		if (FreeList.Num() == 1)
+		{
+			if (ResourceAllocations.Num() == 0)
+			{
+				check(UsedSize == 0);
+				checkf(FreeList[0].Offset == 0 && FreeList[0].Size == MaxSize, TEXT("Memory leak, should have %d free, only have %d; missing %d bytes"), MaxSize, FreeList[0].Size, MaxSize - FreeList[0].Size);
+				return true;
+			}
+		}
+
+		return false;
 	}
 
-	void FResourceAllocationManager::Init(FVulkanDevice* InDevice)
+	FResourceHeap::FResourceHeap(FResourceHeapManager* InOwner, uint32 InMemoryTypeIndex, uint32 InPageSize)
+		: Owner(InOwner)
+		, MemoryTypeIndex(InMemoryTypeIndex)
+		, DefaultPageSize(InPageSize)
+		, PeakPageSize(0)
+		, UsedMemory(0)
 	{
-		check(!Device);
-		Device = InDevice;
+	}
+
+	FResourceHeap::~FResourceHeap()
+	{
+		FScopeLock ScopeLock(&CriticalSection);
 
-		Entries.AddDefaulted(Device->GetMemoryManager().GetNumMemoryTypes());
+		for (int32 Index = UsedPages.Num() - 1; Index >= 0; --Index)
+		{
+			FResourceHeapPage* Page = UsedPages[Index];
+			if (Page->JoinFreeBlocks())
+			{
+				Owner->GetParent()->GetMemoryManager().Free(Page->DeviceMemoryAllocation);
+				delete Page;
+			}
+			else
+			{
+				check(0);
+			}
+		}
+
+		check(UsedPages.Num() == 0);
+		for (int32 Index = 0; Index < FreePages.Num(); ++Index)
+		{
+			FResourceHeapPage* Page = FreePages[Index];
+			Owner->GetParent()->GetMemoryManager().Free(Page->DeviceMemoryAllocation);
+			delete Page;
+		}
 	}
 
-	void FResourceAllocationManager::Deinit()
+	void FResourceHeap::FreePage(FResourceHeapPage* InPage)
 	{
-		check(0);
+		FScopeLock ScopeLock(&CriticalSection);
+		check(InPage->JoinFreeBlocks());
+		UsedPages.RemoveSingleSwap(InPage, false);
+		FreePages.Add(InPage);
+
+		//#todo-rco: Give free pages back to the OS after N frames?
 	}
 
-	FResourceAllocationInfo* FResourceAllocationManager::AllocateBuffer(const VkMemoryRequirements& MemoryRequirements, VkMemoryPropertyFlags MemPropertyFlags)
+#if UE_BUILD_DEBUG || UE_BUILD_DEVELOPMENT
+	void FResourceHeap::DumpMemory()
 	{
-		auto& MemManager = Device->GetMemoryManager();
-		uint32 MemoryTypeIndex = ~0;
-		VERIFYVULKANRESULT(MemManager.GetMemoryTypeFromProperties(MemoryRequirements.memoryTypeBits, MemPropertyFlags, &MemoryTypeIndex));
+		UE_LOG(LogVulkanRHI, Display, TEXT("%d Free Pages, %d Used Pages, Peak Allocation Size on a Page %d"), FreePages.Num(), UsedPages.Num(), PeakPageSize);
+		uint64 UsedMemory = 0;
+		uint64 FreeMemory = 0;
+		uint32 NumSubAllocations = 0;
+		for (int32 Index = 0; Index < UsedPages.Num(); ++Index)
+		{
+			UsedMemory += UsedPages[Index]->UsedSize;
+			NumSubAllocations += UsedPages[Index]->ResourceAllocations.Num();
+		}
 
-		FScopeLock Lock(&GBufferAllocationLock);
-		auto& EntriesPerMemoryType = Entries[MemoryTypeIndex];
+		for (int32 Index = 0; Index < FreePages.Num(); ++Index)
+		{
+			FreeMemory += FreePages[Index]->MaxSize;
+		}
 
-		VkDeviceSize OutOffset = ~(VkDeviceSize)0;
+		UE_LOG(LogVulkanRHI, Display, TEXT("\tUsed Memory %d in %d Suballocations, Free Memory %d"), UsedMemory, NumSubAllocations, FreeMemory);
+	}
+#endif
 
-		for (int32 Index = 0; Index < EntriesPerMemoryType.Num(); ++Index)
+	FResourceAllocation* FResourceHeap::AllocateResource(uint32 Size, uint32 Alignment, bool bMapAllocation, char* File, uint32 Line)
+	{
+		FScopeLock ScopeLock(&CriticalSection);
+
+		if (Size < DefaultPageSize)
 		{
-			FResourceEntry* Entry = EntriesPerMemoryType[Index];
-			if (Entry->TryAllocate(MemoryRequirements.size, MemoryRequirements.alignment, &OutOffset))
+			// Check Used pages to see if we can fit this in
+			for (int32 Index = 0; Index < UsedPages.Num(); ++Index)
 			{
-				return new FResourceAllocationInfo(MemoryTypeIndex, Entry->Allocation, OutOffset);
+				FResourceHeapPage* Page = UsedPages[Index];
+				FResourceAllocation* ResourceAllocation = Page->TryAllocate(Size, Alignment, File, Line);
+				if (ResourceAllocation)
+				{
+					return ResourceAllocation;
+				}
 			}
 		}
 
-		// Allocate a new area
-		VkDeviceSize Size = FMath::Max(MemoryRequirements.size, (VkDeviceSize)PageSize);
-		auto* Block = MemManager.Alloc(Size, MemoryRequirements.memoryTypeBits, MemPropertyFlags, __FILE__, __LINE__);
-		check(Block);
-		FResourceEntry* NewEntry = new FResourceEntry(Block);
-		EntriesPerMemoryType.Add(NewEntry);
+		for (int32 Index = 0; Index < FreePages.Num(); ++Index)
+		{
+			FResourceHeapPage* Page = FreePages[Index];
+			FResourceAllocation* ResourceAllocation = Page->TryAllocate(Size, Alignment, File, Line);
+			if (ResourceAllocation)
+			{
+				FreePages.RemoveSingleSwap(Page, false);
+				UsedPages.Add(Page);
+				return ResourceAllocation;
+			}
+		}
+
+		uint32 AllocationSize = FMath::Max(Size, DefaultPageSize);
+		FDeviceMemoryAllocation* DeviceMemoryAllocation = Owner->GetParent()->GetMemoryManager().Alloc(AllocationSize, MemoryTypeIndex, __FILE__, __LINE__);
+		auto* NewPage = new FResourceHeapPage(this, DeviceMemoryAllocation, AllocationSize);
+		UsedPages.Add(NewPage);
+
+		UsedMemory += AllocationSize;
 
-		bool bResult = NewEntry->TryAllocate(MemoryRequirements.size, MemoryRequirements.alignment, &OutOffset);
-		check(bResult);
-		return new FResourceAllocationInfo(MemoryTypeIndex, NewEntry->Allocation, OutOffset);
+		PeakPageSize = FMath::Max(PeakPageSize, AllocationSize);
+
+		if (bMapAllocation)
+		{
+			DeviceMemoryAllocation->Map(AllocationSize, 0);
+		}
+
+		return NewPage->Allocate(Size, Alignment, File, Line);
 	}
 
-	void FResourceAllocationManager::FreeBuffer(FResourceAllocationInfo*& Info)
+	FResourceHeapManager::FResourceHeapManager(FVulkanDevice* InDevice)
+		: FDeviceChild(InDevice)
+		, GPUOnlyHeap(nullptr)
 	{
-		//#todo-rco: Verify!
-		check(0);
-		delete Info;
-		Info = nullptr;
 	}
 
+	FResourceHeapManager::~FResourceHeapManager()
+	{
+		Deinit();
+	}
+
+	void FResourceHeapManager::Init()
+	{
+		auto& MemoryManager = Device->GetMemoryManager();
+		const uint32 TypeBits = (1 << MemoryManager.GetNumMemoryTypes()) - 1;
+
+		auto& MemoryProperties = MemoryManager.GetMemoryProperties();
+
+		//#todo-rco: Hack! Figure out a better way for page size
+		float PercentageHeapForGPU = 0.75;
+
+		{
+			uint32 TypeIndex = 0;
+			VERIFYVULKANRESULT(MemoryManager.GetMemoryTypeFromProperties(TypeBits, VK_MEMORY_PROPERTY_DEVICE_LOCAL_BIT, &TypeIndex));
+
+			uint32 PageSize = (uint32)((uint64)(MemoryProperties.memoryHeaps[MemoryProperties.memoryTypes[TypeIndex].heapIndex].size * PercentageHeapForGPU) / (uint64)Device->GetLimits().maxMemoryAllocationCount);
+			GPUOnlyHeap = new FResourceHeap(this, TypeIndex, PageSize);
+		}
+
+		{
+			uint32 TypeIndex = 0;
+			VERIFYVULKANRESULT(MemoryManager.GetMemoryTypeFromProperties(TypeBits, VK_MEMORY_PROPERTY_HOST_VISIBLE_BIT, &TypeIndex));
+
+			uint32 PageSize = (uint32)((uint64)(MemoryProperties.memoryHeaps[MemoryProperties.memoryTypes[TypeIndex].heapIndex].size * (1.0f - PercentageHeapForGPU)) / (uint64)Device->GetLimits().maxMemoryAllocationCount);
+			CPUStagingHeap = new FResourceHeap(this, TypeIndex, PageSize);
+		}
+	}
+
+	void FResourceHeapManager::Deinit()
+	{
+		delete GPUOnlyHeap;
+		GPUOnlyHeap = nullptr;
+
+		delete CPUStagingHeap;
+		CPUStagingHeap = nullptr;
+	}
+
+#if UE_BUILD_DEBUG || UE_BUILD_DEVELOPMENT
+	void FResourceHeapManager::DumpMemory()
+	{
+		UE_LOG(LogVulkanRHI, Display, TEXT("GPU Only Heap"));
+		GPUOnlyHeap->DumpMemory();
+		UE_LOG(LogVulkanRHI, Display, TEXT("CPU Staging Heap"));
+		CPUStagingHeap->DumpMemory();
+	}
+#endif
+
+
 	FStagingResource::~FStagingResource()
 	{
 		checkf(!Allocation, TEXT("Staging Resource not released!"));
@@ -378,6 +601,18 @@ namespace VulkanRHI
 		return StagingImage;
 	}
 
+#if VULKAN_USE_NEW_COMMAND_BUFFERS
+	void FStagingManager::ReleaseImage(FVulkanCmdBuffer* CmdBuffer, FStagingImage*& StagingImage)
+	{
+		FScopeLock Lock(&GAllocationLock);
+		UsedStagingImages.RemoveSingleSwap(StagingImage);
+		FPendingItem* Entry = new (PendingFreeStagingImages) FPendingItem;
+		Entry->CmdBuffer = CmdBuffer;
+		Entry->FenceCounter = CmdBuffer->GetFenceSignaledCounter();
+		Entry->Resource = StagingImage;
+		StagingImage = nullptr;
+	}
+#else
 	void FStagingManager::ReleaseImage(FStagingImage*& StagingImage)
 	{
 		FScopeLock Lock(&GAllocationLock);
@@ -385,6 +620,7 @@ namespace VulkanRHI
 		FreeStagingImages.Add(StagingImage);
 		StagingImage = nullptr;
 	}
+#endif
 
 	FStagingBuffer* FStagingManager::AcquireBuffer(uint32 Size)
 	{
@@ -430,6 +666,18 @@ namespace VulkanRHI
 		return StagingBuffer;
 	}
 
+#if VULKAN_USE_NEW_COMMAND_BUFFERS
+	void FStagingManager::ReleaseBuffer(FVulkanCmdBuffer* CmdBuffer, FStagingBuffer*& StagingBuffer)
+	{
+		FScopeLock Lock(&GAllocationLock);
+		UsedStagingBuffers.RemoveSingleSwap(StagingBuffer);
+		FPendingItem* Entry = new (PendingFreeStagingBuffers) FPendingItem;
+		Entry->CmdBuffer = CmdBuffer;
+		Entry->FenceCounter = CmdBuffer->GetFenceSignaledCounter();
+		Entry->Resource = StagingBuffer;
+		StagingBuffer = nullptr;
+	}
+#else
 	void FStagingManager::ReleaseBuffer(FStagingBuffer*& StagingBuffer)
 	{
 		FScopeLock Lock(&GAllocationLock);
@@ -437,14 +685,44 @@ namespace VulkanRHI
 		FreeStagingBuffers.Add(StagingBuffer);
 		StagingBuffer = nullptr;
 	}
+#endif
+
+	void FStagingManager::ProcessPendingFree()
+	{
+#if VULKAN_USE_NEW_COMMAND_BUFFERS
+		FScopeLock Lock(&GAllocationLock);
+		for (int32 Index = PendingFreeStagingBuffers.Num() - 1; Index >= 0; --Index)
+		{
+			FPendingItem& Entry = PendingFreeStagingBuffers[Index];
+			if (Entry.FenceCounter < Entry.CmdBuffer->GetFenceSignaledCounter())
+			{
+				auto* StagingBuffer = (FStagingBuffer*)Entry.Resource;
+				PendingFreeStagingBuffers.RemoveAtSwap(Index, 1, false);
+				FreeStagingBuffers.Add(StagingBuffer);
+			}
+		}
+
+		for (int32 Index = PendingFreeStagingImages.Num() - 1; Index >= 0; --Index)
+		{
+			FPendingItem& Entry = PendingFreeStagingImages[Index];
+			if (Entry.FenceCounter < Entry.CmdBuffer->GetFenceSignaledCounter())
+			{
+				auto* StagingImage = (FStagingImage*)Entry.Resource;
+				PendingFreeStagingImages.RemoveAtSwap(Index, 1, false);
+				FreeStagingImages.Add(StagingImage);
+			}
+		}
+#endif
+	}
 
-	FFence::FFence(FVulkanDevice* Device)
+	FFence::FFence(FVulkanDevice* InDevice, FFenceManager* InOwner)
 		: State(FFence::EState::NotReady)
+		, Owner(InOwner)
 	{
 		VkFenceCreateInfo Info;
 		FMemory::Memzero(Info);
 		Info.sType = VK_STRUCTURE_TYPE_FENCE_CREATE_INFO;
-		VERIFYVULKANRESULT(vkCreateFence(Device->GetInstanceHandle(), &Info, nullptr, &Handle));
+		VERIFYVULKANRESULT(vkCreateFence(InDevice->GetInstanceHandle(), &Info, nullptr, &Handle));
 	}
 
 	FFence::~FFence()
@@ -486,7 +764,7 @@ namespace VulkanRHI
 			return Fence;
 		}
 
-		auto* NewFence = new FFence(Device);
+		auto* NewFence = new FFence(Device, this);
 		UsedFences.Add(NewFence);
 		return NewFence;
 	}
@@ -501,7 +779,21 @@ namespace VulkanRHI
 		Fence = nullptr;
 	}
 
-	FFence::EState FFenceManager::CheckFenceState(FFence* Fence)
+	void FFenceManager::WaitAndReleaseFence(FFence*& Fence, uint64 TimeInNanoseconds)
+	{
+		FScopeLock Lock(&GFenceLock);
+		if (!Fence->IsSignaled())
+		{
+			WaitForFence(Fence, TimeInNanoseconds);
+		}
+
+		ResetFence(Fence);
+		UsedFences.RemoveSingleSwap(Fence);
+		FreeFences.Add(Fence);
+		Fence = nullptr;
+	}
+
+	bool FFenceManager::CheckFenceState(FFence* Fence)
 	{
 		check(UsedFences.Contains(Fence));
 		check(Fence->State == FFence::EState::NotReady);
@@ -510,7 +802,7 @@ namespace VulkanRHI
 		{
 		case VK_SUCCESS:
 			Fence->State = FFence::EState::Signaled;
-			return FFence::EState::Signaled;
+			return true;
 
 		case VK_NOT_READY:
 			break;
@@ -520,7 +812,7 @@ namespace VulkanRHI
 			break;
 		}
 
-		return FFence::EState::NotReady;
+		return false;
 	}
 
 	bool FFenceManager::WaitForFence(FFence* Fence, uint64 TimeInNanoseconds)
@@ -551,4 +843,48 @@ namespace VulkanRHI
 			Fence->State = FFence::EState::NotReady;
 		}
 	}
+
+	FDeferredDeletionQueue::FDeferredDeletionQueue(FVulkanDevice* InDevice)
+		: FDeviceChild(InDevice)
+	{
+	}
+
+	FDeferredDeletionQueue::~FDeferredDeletionQueue()
+	{
+		Clear();
+	}
+
+	void FDeferredDeletionQueue::EnqueueResource(FVulkanCmdBuffer* CmdBuffer, FRefCount* Resource)
+	{
+#if VULKAN_USE_NEW_COMMAND_BUFFERS && VULKAN_USE_NEW_RESOURCE_MANAGEMENT
+		Resource->AddRef();
+
+		FEntry Entry;
+		Entry.CmdBuffer = CmdBuffer;
+		Entry.FenceCounter = CmdBuffer->GetFenceSignaledCounter();
+		Entry.Resource = Resource;
+		{
+			FScopeLock ScopeLock(&CS);
+			Entries.Add(Entry);
+		}
+#endif
+	}
+
+	void FDeferredDeletionQueue::ReleaseResources(/*bool bDeleteImmediately*/)
+	{
+#if VULKAN_USE_NEW_COMMAND_BUFFERS && VULKAN_USE_NEW_RESOURCE_MANAGEMENT
+		FScopeLock ScopeLock(&CS);
+
+		// Traverse list backwards to the swap switches to elements already tested
+		for (int32 Index = Entries.Num() - 1; Index >= 0; --Index)
+		{
+			FEntry* Entry = &Entries[Index];
+			if (Entry->FenceCounter < Entry->CmdBuffer->GetFenceSignaledCounter())
+			{
+				Entry->Resource->Release();
+				Entries.RemoveAtSwap(Index, 1, false);
+			}
+		}
+#endif
+	}
 }
diff --git a/Engine/Source/Runtime/VulkanRHI/Private/VulkanMemory.h b/Engine/Source/Runtime/VulkanRHI/Private/VulkanMemory.h
--- a/Engine/Source/Runtime/VulkanRHI/Private/VulkanMemory.h
+++ b/Engine/Source/Runtime/VulkanRHI/Private/VulkanMemory.h
@@ -8,14 +8,19 @@
 
 #define VULKAN_TRACK_MEMORY_USAGE	1
 
+class FVulkanQueue;
+class FVulkanCmdBuffer;
+
 namespace VulkanRHI
 {
+	class FFenceManager;
+
 	// Custom ref counting
-	class FVulkanRefCount
+	class FRefCount
 	{
 	public:
-		FVulkanRefCount() {}
-		virtual ~FVulkanRefCount()
+		FRefCount() {}
+		virtual ~FRefCount()
 		{
 			check(NumRefs.GetValue() == 0);
 		}
@@ -49,10 +54,10 @@ namespace VulkanRHI
 		FThreadSafeCounter NumRefs;
 	};
 
-	class FVulkanDeviceChild
+	class FDeviceChild
 	{
 	public:
-		FVulkanDeviceChild(FVulkanDevice* InDevice = nullptr) :
+		FDeviceChild(FVulkanDevice* InDevice = nullptr) :
 			Device(InDevice)
 		{
 		}
@@ -74,6 +79,7 @@ namespace VulkanRHI
 		FVulkanDevice* Device;
 	};
 
+	// An Allocation of a Device Heap. Lowest level of allocations and bounded by VkPhysicalDeviceLimits::maxMemoryAllocationCount.
 	class FDeviceMemoryAllocation
 	{
 	public:
@@ -83,7 +89,7 @@ namespace VulkanRHI
 			, Size(0)
 			, MemoryTypeIndex(0)
 			, bCanBeMapped(0)
-			, bIsMapped(0)
+			, MappedPointer(nullptr)
 			, bIsCoherent(0)
 			, bIsCached(0)
 			, bFreedBySystem(false)
@@ -97,17 +103,22 @@ namespace VulkanRHI
 		void* Map(VkDeviceSize Size, VkDeviceSize Offset);
 		void Unmap();
 
-		bool CanBeMapped() const
+		inline void* GetMappedPointer()
+		{
+			return MappedPointer;
+		}
+
+		inline bool CanBeMapped() const
 		{
 			return bCanBeMapped != 0;
 		}
 
-		bool IsMapped() const
+		inline bool IsMapped() const
 		{
-			return bIsMapped != 0;
+			return !!MappedPointer;
 		}
 
-		bool IsCoherent() const
+		inline bool IsCoherent() const
 		{
 			return bIsCoherent != 0;
 		}
@@ -148,13 +159,18 @@ namespace VulkanRHI
 			return Size;
 		}
 
+		inline uint32 GetMemoryTypeIndex() const
+		{
+			return MemoryTypeIndex;
+		}
+
 	protected:
 		VkDeviceSize Size;
 		VkDevice DeviceHandle;
 		VkDeviceMemory Handle;
+		void* MappedPointer;
 		uint32 MemoryTypeIndex : 8;
 		uint32 bCanBeMapped : 1;
-		uint32 bIsMapped : 1;
 		uint32 bIsCoherent : 1;
 		uint32 bIsCached : 1;
 		uint32 bFreedBySystem : 1;
@@ -170,18 +186,14 @@ namespace VulkanRHI
 		friend class FDeviceMemoryManager;
 	};
 
+	// Manager of Device Heap Allocations. Calling Alloc/Free is expensive!
 	class FDeviceMemoryManager
 	{
 	public:
 		FDeviceMemoryManager();
 		~FDeviceMemoryManager();
 
-		void Init(struct FVulkanDevice* InDevice);
-
-		void SetDevice(VkDevice InDevice)
-		{
-			DeviceHandle = InDevice;
-		}
+		void Init(FVulkanDevice* InDevice);
 
 		void Deinit();
 
@@ -190,6 +202,7 @@ namespace VulkanRHI
 			return MemoryProperties.memoryTypeCount;
 		}
 
+		//#todo-rco: Might need to revisit based on https://gitlab.khronos.org/vulkan/vulkan/merge_requests/1165
 		inline VkResult GetMemoryTypeFromProperties(uint32 TypeBits, VkMemoryPropertyFlags Properties, uint32* OutTypeIndex)
 		{
 			// Search memtypes to find first index with those properties
@@ -228,9 +241,12 @@ namespace VulkanRHI
 		// Sets the Allocation to nullptr
 		void Free(FDeviceMemoryAllocation*& Allocation);
 
+#if UE_BUILD_DEBUG || UE_BUILD_DEVELOPMENT
+		void DumpMemory();
+#endif
+
 	protected:
 		VkPhysicalDeviceMemoryProperties MemoryProperties;
-		VkPhysicalDevice PhysicalDeviceHandle;
 		VkDevice DeviceHandle;
 		FVulkanDevice* Device;
 		uint32 NumAllocations;
@@ -252,90 +268,166 @@ namespace VulkanRHI
 		TArray<FHeapInfo> HeapInfos;
 	};
 
-	class FResourceAllocationInfo
+	class FResourceHeap;
+	class FResourceHeapPage;
+	class FResourceHeapManager;
+
+	// A sub allocation for a specific memory type
+	class FResourceAllocation : public FRefCount
 	{
 	public:
-		inline VkDeviceMemory GetHandle() const
+		FResourceAllocation(FResourceHeapPage* InOwner, FDeviceMemoryAllocation* InDeviceMemoryAllocation, uint32 InSize, uint32 InOffset, char* File, uint32 Line);
+		virtual ~FResourceAllocation();
+
+		inline uint32 GetSize() const
 		{
-			return Allocation->GetHandle();
+			return Size;
 		}
 
-		inline VkDeviceSize GetOffset() const
+		inline uint32 GetOffset() const
 		{
 			return Offset;
 		}
 
-		FResourceAllocationInfo()
-			: Allocation(nullptr)
-			, Offset(~(VkDeviceSize)0)
-			, MemoryTypeIndex(~0)
-			, bReleased(true)
+		inline VkDeviceMemory GetHandle() const
 		{
+			return DeviceMemoryAllocation->GetHandle();
 		}
 
-		~FResourceAllocationInfo()
+		void* GetMappedPointer();
+
+		inline uint32 GetMemoryTypeIndex() const
 		{
-			check(bReleased);
+			return DeviceMemoryAllocation->GetMemoryTypeIndex();
 		}
 
-		void* Map(VkDeviceSize Size, VkDeviceSize Offset);
-		void Unmap();
+	private:
+		FResourceHeapPage* Owner;
+		uint32 Size;
+		uint32 Offset;
+		FDeviceMemoryAllocation* DeviceMemoryAllocation;
 
-	protected:
-		FDeviceMemoryAllocation* Allocation;
-		VkDeviceSize Offset;
-		uint32 MemoryTypeIndex;
-		bool bReleased;
+#if VULKAN_TRACK_MEMORY_USAGE
+		char* File;
+		uint32 Line;
+#endif
+
+		friend class FResourceHeapPage;
+	};
+
+	// One device allocation that is shared amongst different resources
+	class FResourceHeapPage
+	{
+	public:
+		FResourceHeapPage(FResourceHeap* InOwner, FDeviceMemoryAllocation* InDeviceMemoryAllocation, uint32 InSize)
+			: Owner(InOwner)
+			, DeviceMemoryAllocation(InDeviceMemoryAllocation)
+			, MaxSize(InSize)
+			, UsedSize(0)
+			, PeakNumAllocations(0)
+		{
+			FPair Pair;
+			Pair.Offset = 0;
+			Pair.Size = MaxSize;
+			FreeList.Add(Pair);
+		}
+
+		~FResourceHeapPage();
+
+		FResourceAllocation* TryAllocate(uint32 Size, uint32 Alignment, char* File, uint32 Line);
 
-		FResourceAllocationInfo(uint32 InMemoryTypeIndex,
-			FDeviceMemoryAllocation* InAllocation,
-			VkDeviceSize InOffset)
-			: Allocation(InAllocation)
-			, Offset(InOffset)
-			, MemoryTypeIndex(InMemoryTypeIndex)
-			, bReleased(false)
+		FResourceAllocation* Allocate(uint32 Size, uint32 Alignment, char* File, uint32 Line)
 		{
+			FResourceAllocation* ResourceAllocation = TryAllocate(Size, Alignment, File, Line);
+			check(ResourceAllocation);
+			return ResourceAllocation;
 		}
 
-		friend class FResourceAllocationManager;
+		void ReleaseAllocation(FResourceAllocation* Allocation);
+
+	protected:
+		FResourceHeap* Owner;
+		FDeviceMemoryAllocation* DeviceMemoryAllocation;
+		uint32 MaxSize;
+		uint32 UsedSize;
+		int32 PeakNumAllocations;
+		TArray<FResourceAllocation*>  ResourceAllocations;
+
+		bool JoinFreeBlocks();
+
+		struct FPair
+		{
+			uint32 Offset;
+			uint32 Size;
+
+			bool operator<(const FPair& In) const
+			{
+				return Offset < In.Offset;
+			}
+		};
+		TArray<FPair> FreeList;
+		friend class FResourceHeap;
 	};
 
-	class FResourceAllocationManager
+	// A set of Device Allocations (Heap Pages) for a specific memory type. This handles pooling allocations inside memory pages to avoid
+	// doing allocations directly off the device's heaps
+	class FResourceHeap
 	{
 	public:
-		FResourceAllocationManager();
-		~FResourceAllocationManager();
+		FResourceHeap(FResourceHeapManager* InOwner, uint32 InMemoryTypeIndex, uint32 InPageSize);
+		~FResourceHeap();
 
-		void Init(FVulkanDevice* InDevice);
-		void Deinit();
+		void FreePage(FResourceHeapPage* InPage);
 
-		FResourceAllocationInfo* AllocateBuffer(const VkMemoryRequirements& MemoryRequirements, VkMemoryPropertyFlags MemPropertyFlags);
-		// Sets to null
-		void FreeBuffer(FResourceAllocationInfo*& Info);
+#if UE_BUILD_DEBUG || UE_BUILD_DEVELOPMENT
+		void DumpMemory();
+#endif
 
 	protected:
-		enum
+		FResourceHeapManager* Owner;
+		uint32 MemoryTypeIndex;
+
+		uint32 DefaultPageSize;
+		uint32 PeakPageSize;
+		uint64 UsedMemory;
+
+		TArray<FResourceHeapPage*> UsedPages;
+		TArray<FResourceHeapPage*> FreePages;
+
+		FResourceAllocation* AllocateResource(uint32 Size, uint32 Alignment, bool bMapAllocation, char* File, uint32 Line);
+
+		FCriticalSection CriticalSection;
+
+		friend class FResourceHeapManager;
+	};
+
+	// Manages heaps and their interactions
+	class FResourceHeapManager : public FDeviceChild
+	{
+	public:
+		FResourceHeapManager(FVulkanDevice* InDevice);
+		~FResourceHeapManager();
+
+		void Init();
+		void Deinit();
+
+		inline FResourceAllocation* AllocateGPUOnlyResource(uint32 Size, uint32 Alignment, char* File, uint32 Line)
 		{
-			/*Host*/PageSize = 1 * 1024 * 1024,
-			//LocalPageSize = 1 * 1024 * 1024,
-		};
-		FVulkanDevice* Device;
+			return GPUOnlyHeap->AllocateResource(Size, Alignment, false, File, Line);
+		}
 
-		struct FResourceEntry
+		inline FResourceAllocation* AllocateCPUStagingResource(uint32 Size, uint32 Alignment, char* File, uint32 Line)
 		{
-			FResourceEntry(FDeviceMemoryAllocation* InAllocation)
-				: Allocation(InAllocation)
-				, FreeOffset(0)
-			{
-			}
+			return CPUStagingHeap->AllocateResource(Size, Alignment, true, File, Line);
+		}
 
-			bool TryAllocate(VkDeviceSize Size, VkDeviceSize Alignment, VkDeviceSize* OutOffset);
+#if UE_BUILD_DEBUG || UE_BUILD_DEVELOPMENT
+		void DumpMemory();
+#endif
 
-			FDeviceMemoryAllocation* Allocation;
-			VkDeviceSize FreeOffset;
-		};
-		// Array of entries per memory type
-		TArray< TArray<FResourceEntry*> > /*Free*/Entries;
+	protected:
+		FResourceHeap* GPUOnlyHeap;
+		FResourceHeap* CPUStagingHeap;
 	};
 
 	class FStagingResource
@@ -438,18 +530,41 @@ namespace VulkanRHI
 
 		FStagingBuffer* AcquireBuffer(uint32 Size);
 
+#if VULKAN_USE_NEW_COMMAND_BUFFERS
+		// Sets pointer to nullptr
+		void ReleaseImage(FVulkanCmdBuffer* CmdBuffer, FStagingImage*& StagingImage);
+
+		// Sets pointer to nullptr
+		void ReleaseBuffer(FVulkanCmdBuffer* CmdBuffer, FStagingBuffer*& StagingBuffer);
+#else
 		// Sets pointer to nullptr
 		void ReleaseImage(FStagingImage*& StagingImage);
 
 		// Sets pointer to nullptr
 		void ReleaseBuffer(FStagingBuffer*& StagingBuffer);
+#endif
+
+		void ProcessPendingFree();
 
 	protected:
-		TArray<FStagingImage*> FreeStagingImages;
+		struct FPendingItem
+		{
+			FVulkanCmdBuffer* CmdBuffer;
+			uint64 FenceCounter;
+			FStagingResource* Resource;
+		};
+
 		TArray<FStagingImage*> UsedStagingImages;
+#if VULKAN_USE_NEW_COMMAND_BUFFERS
+		TArray<FPendingItem> PendingFreeStagingImages;
+#endif
+		TArray<FStagingImage*> FreeStagingImages;
 
-		TArray<FStagingBuffer*> FreeStagingBuffers;
 		TArray<FStagingBuffer*> UsedStagingBuffers;
+#if VULKAN_USE_NEW_COMMAND_BUFFERS
+		TArray<FPendingItem> PendingFreeStagingBuffers;
+#endif
+		TArray<FStagingBuffer*> FreeStagingBuffers;
 
 		FVulkanDevice* Device;
 		FVulkanQueue* Queue;
@@ -458,13 +573,26 @@ namespace VulkanRHI
 	class FFence
 	{
 	public:
-		FFence(FVulkanDevice* Device);
+		FFence(FVulkanDevice* InDevice, FFenceManager* InOwner);
 
 		inline VkFence GetHandle() const
 		{
 			return Handle;
 		}
 
+		inline bool IsSignaled() const
+		{
+			return State == EState::Signaled;
+		}
+
+		FFenceManager* GetOwner()
+		{
+			return Owner;
+		}
+
+	protected:
+		VkFence Handle;
+
 		enum class EState
 		{
 			// Initial state
@@ -474,15 +602,10 @@ namespace VulkanRHI
 			Signaled,
 		};
 
-		inline EState GetState() const
-		{
-			return State;
-		}
-
-	protected:
-		VkFence Handle;
 		EState State;
 
+		FFenceManager* Owner;
+
 		// Only owner can delete!
 		~FFence();
 		friend class FFenceManager;
@@ -502,11 +625,11 @@ namespace VulkanRHI
 
 		FFence* AllocateFence();
 
-		inline FFence::EState GetFenceState(FFence* Fence)
+		inline bool IsFenceSignaled(FFence* Fence)
 		{
-			if (Fence->GetState() == FFence::EState::Signaled)
+			if (Fence->IsSignaled())
 			{
-				return FFence::EState::Signaled;
+				return true;
 			}
 
 			return CheckFenceState(Fence);
@@ -520,124 +643,35 @@ namespace VulkanRHI
 		// Sets it to nullptr
 		void ReleaseFence(FFence*& Fence);
 
+		// Sets it to nullptr
+		void WaitAndReleaseFence(FFence*& Fence, uint64 TimeInNanoseconds);
+
 	protected:
 		FVulkanDevice* Device;
 		TArray<FFence*> FreeFences;
 		TArray<FFence*> UsedFences;
 
-		FFence::EState CheckFenceState(FFence* Fence);
+		// Returns true if signaled
+		bool CheckFenceState(FFence* Fence);
 	};
 
-	struct FVulkanResource : public FVulkanRefCount
+	class FDeferredDeletionQueue : public FDeviceChild
 	{
-	};
+		//typedef TPair<FRefCountedObject*, uint64> FFencedObject;
+		//FThreadsafeQueue<FFencedObject> DeferredReleaseQueue;
 
-	//#todo-rco: Move to common RHI?
-	/**
-	 * The base class of threadsafe reference counted objects.
-	 */
-	template <class Type>
-	struct FThreadsafeQueue
-	{
-	private:
-		mutable FCriticalSection	SynchronizationObject; // made this mutable so this class can have const functions and still be thread safe
-		TQueue<Type>				Items;
 	public:
-		void Enqueue(const Type& Item)
-		{
-			FScopeLock ScopeLock(&SynchronizationObject);
-			Items.Enqueue(Item);
-		}
+		FDeferredDeletionQueue(FVulkanDevice* InDevice);
+		~FDeferredDeletionQueue();
 
-		bool Dequeue(Type& Result)
-		{
-			FScopeLock ScopeLock(&SynchronizationObject);
+		void EnqueueResource(FVulkanCmdBuffer* CmdBuffer, FRefCount* Resource);
+		void ReleaseResources(/*bool bDeleteImmediately = false*/);
 
-			return Items.Dequeue(Result);
-		}
-
-		template <typename CompareFunc>
-		bool Dequeue(Type& Result, CompareFunc Func)
+		inline void Clear()
 		{
-			FScopeLock ScopeLock(&SynchronizationObject);
-
-			if (Items.Peek(Result))
-			{
-				if (Func(Result))
-				{
-					Items.Dequeue(Result);
-
-					return true;
-				}
-			}
-
-			return false;
-		}
-
-		template <typename CompareFunc>
-		bool BatchDequeue(TQueue<Type>* Result, CompareFunc Func, uint32 MaxItems)
-		{
-			FScopeLock ScopeLock(&SynchronizationObject);
-
-			uint32 i = 0;
-			Type Item;
-			while (Items.Peek(Item) && i <= MaxItems)
-			{
-				if (Func(Item))
-				{
-					Items.Dequeue(Item);
-
-					Result->Enqueue(Item);
-
-					i++;
-				}
-				else
-				{
-					break;
-				}
-			}
-
-			return i > 0;
-		}
-
-		bool Peek(Type& Result) const
-		{
-			FScopeLock ScopeLock(&SynchronizationObject);
-			return Items.Peek(Result);
-		}
-
-		bool IsEmpty() const
-		{
-			FScopeLock ScopeLock(&SynchronizationObject);
-			return Items.IsEmpty();
-		}
-
-		void Empty()
-		{
-			FScopeLock ScopeLock(&SynchronizationObject);
-
-			Type Result;
-			while (Items.Dequeue(Result)) {}
+			ReleaseResources(/*true*/);
 		}
-	};
-
-	class FVulkanDeferredDeletionQueue : public FVulkanDeviceChild
-	{
-		typedef TPair<FVulkanResource*, uint64> FFencedObject;
-		FThreadsafeQueue<FFencedObject> DeferredReleaseQueue;
-
-	public:
-		FVulkanDeferredDeletionQueue(FVulkanDevice* InDevice);
-		~FVulkanDeferredDeletionQueue();
-
-		void EnqueueResource(FVulkanResource* Resource);
-		void ReleaseResources(bool bDeleteImmediately = false);
-
-		void Clear()
-		{
-			ReleaseResources(true);
-		}
-
+/*
 		class FVulkanAsyncDeletionWorker : public FVulkanDeviceChild, FNonAbandonableTask
 		{
 		public:
@@ -648,8 +682,16 @@ namespace VulkanRHI
 		private:
 			TQueue<FFencedObject> Queue;
 		};
-
+*/
 	private:
-		TQueue<FAsyncTask<FVulkanAsyncDeletionWorker>*> DeleteTasks;
+		//TQueue<FAsyncTask<FVulkanAsyncDeletionWorker>*> DeleteTasks;
+		struct FEntry
+		{
+			uint64 FenceCounter;
+			FVulkanCmdBuffer* CmdBuffer;
+			FRefCount* Resource;
+		};
+		FCriticalSection CS;
+		TArray<FEntry> Entries;
 	};
 }
diff --git a/Engine/Source/Runtime/VulkanRHI/Private/VulkanPendingState.cpp b/Engine/Source/Runtime/VulkanRHI/Private/VulkanPendingState.cpp
--- a/Engine/Source/Runtime/VulkanRHI/Private/VulkanPendingState.cpp
+++ b/Engine/Source/Runtime/VulkanRHI/Private/VulkanPendingState.cpp
@@ -7,6 +7,7 @@
 #include "VulkanRHIPrivate.h"
 #include "VulkanPendingState.h"
 #include "VulkanPipeline.h"
+#include "VulkanContext.h"
 
 // RTs 0-3
 #define NUMBITS_BLEND_STATE				4 //(x4=16) = 16 ++
@@ -261,16 +262,23 @@ FVulkanPendingState::FVulkanPendingState(FVulkanDevice* InDevice)
 	, bBeginRenderPass(false)
 	, bChangeRenderTarget(false)
 	, GlobalUniformPool(nullptr)
+#if VULKAN_USE_NEW_COMMAND_BUFFERS
+#else
 	, CurrentCmdBufferIndex(0)
+#endif
 	, bScissorEnable(false)
 {
 	check(Device);
 
+#if VULKAN_USE_NEW_COMMAND_BUFFERS
+	// State doesn't own or manage cmd buffers
+#else
 	for (uint32 CmdBufferIndex = 0; CmdBufferIndex < VULKAN_NUM_COMMAND_BUFFERS; ++CmdBufferIndex)
 	{
 		//#todo-rco: FIX ME!
 		CmdBuffers[CmdBufferIndex] = Device->GetImmediateContext().GetCommandBufferManager()->Create();
 	}
+#endif
 
 	Reset();
 
@@ -282,13 +290,16 @@ FVulkanPendingState::~FVulkanPendingState()
 {
 	check(Device);
 
+#if VULKAN_USE_NEW_COMMAND_BUFFERS
+	// State doesn't own or manage cmd buffers
+#else
 	for (uint32 CmdBufferIndex = 0; CmdBufferIndex < VULKAN_NUM_COMMAND_BUFFERS; ++CmdBufferIndex)
 	{
 		//#todo-rco: FIX ME!
 		Device->GetImmediateContext().GetCommandBufferManager()->Destroy(CmdBuffers[CmdBufferIndex]);
 		CmdBuffers[CmdBufferIndex] = nullptr;
 	}
-
+#endif
 	//Reset();
 
 	/*
@@ -338,7 +349,11 @@ FVulkanGlobalUniformPool& FVulkanPendingState::GetGlobalUniformPool()
     return *GlobalUniformPool;
 }
 
+#if VULKAN_USE_NEW_COMMAND_BUFFERS
+void FVulkanPendingState::RenderPassBegin(FVulkanCmdBuffer* CmdBuffer)
+#else
 void FVulkanPendingState::RenderPassBegin()
+#endif
 {
 	check(!bBeginRenderPass);
 
@@ -408,16 +423,29 @@ void FVulkanPendingState::RenderPassBegin()
 		// if either bit is set that we have a depth buffer attached
 	}
 		
+#if VULKAN_USE_NEW_COMMAND_BUFFERS
+	CmdBuffer->BeginRenderPass(CurrentState.RenderPass->GetLayout(), CurrentState.RenderPass->GetHandle(), CurrentState.FrameBuffer->GetHandle(), ClearValues);
+#else
 	CurrentState.RenderPass->Begin(GetCurrentCommandBuffer(), CurrentState.FrameBuffer->GetHandle(), ClearValues);
+#endif
 	bBeginRenderPass = true;
 }
 
+#if VULKAN_USE_NEW_COMMAND_BUFFERS
+void FVulkanPendingState::RenderPassEnd(FVulkanCmdBuffer* CmdBuffer)
+{
+	check(bBeginRenderPass);
+	CmdBuffer->EndRenderPass();
+	bBeginRenderPass = false;
+}
+#else
 void FVulkanPendingState::RenderPassEnd()
 {
 	check(bBeginRenderPass);
 	CurrentState.RenderPass->End(GetCurrentCommandBuffer());
 	bBeginRenderPass = false;
 }
+#endif
 
 // Expected to be called after render pass has been ended
 // and only from "FVulkanDynamicRHI::RHIEndDrawingViewport()"
@@ -427,11 +455,14 @@ void FVulkanPendingState::Reset()
 
 	CurrentState.Reset();
 
+#if VULKAN_USE_NEW_COMMAND_BUFFERS
+	// State doesn't own or manage cmd buffers
+#else
 	CurrentCmdBufferIndex = (CurrentCmdBufferIndex + 1) % VULKAN_NUM_COMMAND_BUFFERS;
 
 	GetCurrentCommandBuffer().Reset();
 	GetCurrentCommandBuffer().Begin();
-
+#endif
 	FMemory::Memzero(PendingStreams);
 }
 
@@ -452,10 +483,21 @@ FVulkanDescriptorPool::FVulkanDescriptorPool(FVulkanDevice* InDevice)
 	FMemory::Memzero(PeakAllocatedTypes);
 
 	//#todo-rco: Get some initial values
-	const uint32 LimitMaxUniformBuffers = 1 * 1024 * 1024;
-	const uint32 LimitMaxSamplers = 1 * 1024 * 1024;
-	const uint32 LimitMaxCombinedImageSamplers = 1 * 1024 * 1024;
-	const uint32 LimitMaxUniformTexelBuffers = 1 * 1024 * 1024;
+	uint32 LimitMaxUniformBuffers = 512 * 1024;
+	uint32 LimitMaxSamplers = FMath::Min(Limits.maxSamplerAllocationCount, (uint32)64 * 1024);
+	uint32 LimitMaxCombinedImageSamplers = 256 * 1024;
+	uint32 LimitMaxUniformTexelBuffers = 64 * 1024;
+
+#if VULKAN_HAS_DEBUGGING_ENABLED
+	if (GValidationCvar.GetValueOnAnyThread() != 0)
+	{
+		 //#todo-rco: Make sure this doesn't overflow the validation layers as of 1.0.5.0
+		LimitMaxUniformBuffers = FMath::Min(MAX_uint32 / MaxDescriptorSets, LimitMaxUniformBuffers);
+		LimitMaxSamplers = FMath::Min(MAX_uint32 / MaxDescriptorSets, LimitMaxSamplers);
+		LimitMaxCombinedImageSamplers = FMath::Min(MAX_uint32 / MaxDescriptorSets, LimitMaxCombinedImageSamplers);
+		LimitMaxUniformTexelBuffers = FMath::Min(MAX_uint32 / MaxDescriptorSets, LimitMaxUniformTexelBuffers);
+	}
+#endif
 
 	TArray<VkDescriptorPoolSize> Types;
 	auto* Type = new(Types) VkDescriptorPoolSize;
@@ -498,8 +540,8 @@ FVulkanDescriptorPool::FVulkanDescriptorPool(FVulkanDevice* InDevice)
 
 	for (int32 Index = 0; Index < Types.Num(); ++Index)
 	{
-		auto& Type = Types[Index];
-		MaxAllocatedTypes[Type.type] =Type.descriptorCount;
+		auto& Type2 = Types[Index];
+		MaxAllocatedTypes[Type2.type] = Type2.descriptorCount;
 	}
 }
 
@@ -582,7 +624,11 @@ inline void FVulkanBoundShaderState::BindDescriptorSets(FVulkanCmdBuffer* Cmd)
 	CurrDescriptorSets->Bind(Cmd, this);
 }
 
+#if VULKAN_USE_NEW_COMMAND_BUFFERS
+void FVulkanPendingState::PrepareDraw(FVulkanCmdBuffer* Cmd, VkPrimitiveTopology Topology)
+#else
 void FVulkanPendingState::PrepareDraw(VkPrimitiveTopology Topology)
+#endif
 {
 	SCOPE_CYCLE_COUNTER(STAT_VulkanDrawCallPrepareTime);
 
@@ -590,7 +636,11 @@ void FVulkanPendingState::PrepareDraw(VkPrimitiveTopology Topology)
 	SetKeyBits(CurrentKey, OFFSET_POLYTYPE, NUMBITS_POLYTYPE, Topology);
 
 	//@TODO: let's try not to do this per draw call?
+#if VULKAN_USE_NEW_COMMAND_BUFFERS
+	UpdateRenderPass(Cmd);
+#else
 	UpdateRenderPass();
+#endif
 
 	check(CurrentState.Shader);
     CurrentState.Shader->UpdateDescriptorSets(GlobalUniformPool);
@@ -603,7 +653,10 @@ void FVulkanPendingState::PrepareDraw(VkPrimitiveTopology Topology)
 
 	{
 		SCOPE_CYCLE_COUNTER(STAT_VulkanPipelineBind);
+#if VULKAN_USE_NEW_COMMAND_BUFFERS
+#else
 		auto* Cmd = &GetCurrentCommandBuffer();
+#endif
 		Pipeline->UpdateDynamicStates(Cmd, CurrentState);
 
 		VkPipeline NewPipeline = Pipeline->GetHandle();
@@ -632,6 +685,9 @@ void FVulkanPendingState::PrepareDraw(VkPrimitiveTopology Topology)
 
 void FVulkanPendingState::SubmitPendingCommandBuffers(FVulkanPendingState::TCallback* Callback, void* CallbackUserData)
 {
+#if VULKAN_USE_NEW_COMMAND_BUFFERS
+	check(0);
+#else
 	if (IsRenderPassActive())
 	{
 		RenderPassEnd();
@@ -645,20 +701,28 @@ void FVulkanPendingState::SubmitPendingCommandBuffers(FVulkanPendingState::TCall
 	check(Device);
 	//#todo-rco: FIX ME!
 	Device->GetImmediateContext().GetCommandBufferManager()->Submit(&GetCurrentCommandBuffer());
+#endif
 }
 
 void FVulkanPendingState::SubmitPendingCommandBuffersBlockingNoRenderPass()
 {
+#if VULKAN_USE_NEW_COMMAND_BUFFERS
+	check(0);
+#else
 	check(Device);
 	Device->GetQueue()->SubmitBlocking(&GetCurrentCommandBuffer());
+#endif
 }
 
 void FVulkanPendingState::SetRenderTargetsInfo(const FRHISetRenderTargetsInfo& InRTInfo)
 {
+	//#todo-rco: Check perf
+#if !VULKAN_USE_NEW_COMMAND_BUFFERS
 	if (NeedsToSetRenderTarget(InRTInfo) == false)
 	{
 		return;
 	}
+#endif
 
 	//@NOTE: this is only needed for the work-around below
 	FRHISetRenderTargetsInfo PrevRTInfo = RTInfo;
@@ -668,6 +732,7 @@ void FVulkanPendingState::SetRenderTargetsInfo(const FRHISetRenderTargetsInfo& I
 	RTInfo = InRTInfo; // probably don't need both of these
 	bChangeRenderTarget = true;
 
+	//#todo-rco: FIX THIS!!!
 	//@NOTE: this is an awkward work-around for not supporting a null color attachment properly
 	for (int32 Index = 0; Index < RTInfo.NumColorRenderTargets; Index++)
 	{
@@ -781,6 +846,8 @@ FVulkanRenderPass& FVulkanPendingState::GetRenderPass()
 	return *CurrentState.RenderPass;
 }
 
+#if VULKAN_USE_NEW_COMMAND_BUFFERS
+#else
 VkCommandBuffer& FVulkanPendingState::GetCommandBuffer()
 {
 	return GetCurrentCommandBuffer().GetHandle();
@@ -790,6 +857,7 @@ const VkBool32 FVulkanPendingState::GetIsCommandBufferEmpty() const
 {
 	return GetCurrentCommandBuffer().GetIsEmpty();
 }
+#endif
 
 void FVulkanPendingState::SetViewport(uint32 MinX, uint32 MinY, float MinZ, uint32 MaxX, uint32 MaxY, float MaxZ)
 {
diff --git a/Engine/Source/Runtime/VulkanRHI/Private/VulkanPendingState.h b/Engine/Source/Runtime/VulkanRHI/Private/VulkanPendingState.h
--- a/Engine/Source/Runtime/VulkanRHI/Private/VulkanPendingState.h
+++ b/Engine/Source/Runtime/VulkanRHI/Private/VulkanPendingState.h
@@ -10,8 +10,9 @@
 #include "VulkanRHI.h"
 #include "VulkanPipeline.h"
 
-struct FVulkanPendingState
+class FVulkanPendingState
 {
+public:
 	typedef TMap<FStateKey, FVulkanRenderPass*> FMapRenderRass;
 	typedef TMap<FStateKey, TArray<FVulkanFramebuffer*> > FMapFrameBufferArray;
 
@@ -28,6 +29,13 @@ struct FVulkanPendingState
 
 	void Reset();
 
+#if VULKAN_USE_NEW_COMMAND_BUFFERS
+	void RenderPassBegin(FVulkanCmdBuffer* CmdBuffer);
+
+	void PrepareDraw(FVulkanCmdBuffer* CmdBuffer, VkPrimitiveTopology Topology);
+
+	void RenderPassEnd(FVulkanCmdBuffer* CmdBuffer);
+#else
 	// Needs to be called before "PrepareDraw" and also starts writing in to the command buffer
 	void RenderPassBegin();
 
@@ -38,6 +46,7 @@ struct FVulkanPendingState
 	// Is currently linked to the command buffer
 	// Will get called in RHIEndDrawingViewport through SubmitPendingCommandBuffers
 	void RenderPassEnd();
+#endif
 
 	inline bool IsRenderPassActive() const
 	{
@@ -67,12 +76,35 @@ struct FVulkanPendingState
 	// Retuns constructed render pass
 	FVulkanRenderPass& GetRenderPass();
 
+#if VULKAN_USE_NEW_COMMAND_BUFFERS
+#else
 	// Returns currently bound command buffer
 	VkCommandBuffer& GetCommandBuffer();
 	const VkBool32 GetIsCommandBufferEmpty() const;
-
+#endif
 	FVulkanFramebuffer* GetFrameBuffer();
 
+#if VULKAN_USE_NEW_COMMAND_BUFFERS
+	inline void UpdateRenderPass(FVulkanCmdBuffer* CmdBuffer)
+	{
+		//#todo-rco: Don't test here, move earlier to SetRenderTarget
+		if (bChangeRenderTarget)
+		{
+			if (bBeginRenderPass)
+			{
+				RenderPassEnd(CmdBuffer);
+			}
+
+			RenderPassBegin(CmdBuffer);
+			bChangeRenderTarget = false;
+		}
+		if (!bBeginRenderPass)
+		{
+			RenderPassBegin(CmdBuffer);
+			bChangeRenderTarget = false;
+		}
+	}
+#else
 	inline void UpdateRenderPass()
 	{
 		//#todo-rco: Don't test here, move earlier to SetRenderTarget
@@ -87,6 +119,7 @@ struct FVulkanPendingState
 			bChangeRenderTarget = false;
 		}
 	}
+#endif
 
 	inline const FVulkanPipelineGraphicsKey& GetCurrentKey() const
 	{
@@ -117,6 +150,8 @@ struct FVulkanPendingState
 
 	void InitFrame();
 
+#if VULKAN_USE_NEW_COMMAND_BUFFERS
+#else
 	inline FVulkanCmdBuffer& GetCommandBuffer(uint32 Index)
 	{
 		check(Index<VULKAN_NUM_COMMAND_BUFFERS);
@@ -144,7 +179,7 @@ struct FVulkanPendingState
 		check(CmdBuffers[CurrentCmdBufferIndex]);
 		return *CmdBuffers[CurrentCmdBufferIndex];
 	}
-
+#endif
 private:
 	FVulkanRenderPass* GetOrCreateRenderPass(const FVulkanRenderTargetLayout& RTLayout);
 	FVulkanFramebuffer* GetOrCreateFramebuffer(const FRHISetRenderTargetsInfo& RHIRTInfo,
@@ -166,9 +201,13 @@ private:
 	//@TODO: probably needs to go somewhere else
 	FRHISetRenderTargetsInfo PrevRenderTargetsInfo;
 
+#if VULKAN_USE_NEW_COMMAND_BUFFERS
+	friend class FVulkanCommandListContext;
+#else
 	//#todo-rco: FIX ME! ASAP!!!
 	FVulkanCmdBuffer* CmdBuffers[VULKAN_NUM_COMMAND_BUFFERS];
 	uint32	CurrentCmdBufferIndex;
+#endif
 	
 	FRHISetRenderTargetsInfo RTInfo;
 
diff --git a/Engine/Source/Runtime/VulkanRHI/Private/VulkanPipeline.cpp b/Engine/Source/Runtime/VulkanRHI/Private/VulkanPipeline.cpp
--- a/Engine/Source/Runtime/VulkanRHI/Private/VulkanPipeline.cpp
+++ b/Engine/Source/Runtime/VulkanRHI/Private/VulkanPipeline.cpp
@@ -186,14 +186,8 @@ static TAutoConsoleVariable<int32> GEnablePipelineCacheLoadCvar(
 
 FVulkanPipelineStateCache::FDiskEntry::~FDiskEntry()
 {
-	if (bLoaded)
-	{
-		//free DescriptorSetLayouts;
-		//free PipelineLayout;
-		//free ShaderModules[SF_Compute];
-
-		delete RenderPass;
-	}
+	check(!bLoaded);
+	check(!RenderPass);
 }
 
 FVulkanPipelineStateCache::FVulkanPipelineStateCache(FVulkanDevice* InDevice)
@@ -204,6 +198,8 @@ FVulkanPipelineStateCache::FVulkanPipelineStateCache(FVulkanDevice* InDevice)
 
 FVulkanPipelineStateCache::~FVulkanPipelineStateCache()
 {
+	DestroyCache();
+
 	vkDestroyPipelineCache(Device->GetInstanceHandle(), PipelineCache, nullptr);
 	PipelineCache = VK_NULL_HANDLE;
 }
@@ -244,7 +240,8 @@ bool FVulkanPipelineStateCache::Load(const TArray<FString>& CacheFilenames, TArr
 				FVulkanPipelineStateKey CreateInfo(DiskEntry->GraphicsKey, DiskEntry->VertexInputKey, DiskEntry->ShaderHashes);
 
 				// Add to the cache
-				Cache.Add(CreateInfo, Pipeline);
+				KeyToPipelineMap.Add(CreateInfo, Pipeline);
+				CreatedPipelines.Add(DiskEntry, Pipeline);
 			}
 
 			Ar << OutDeviceCache;
@@ -349,7 +346,7 @@ void FVulkanPipelineStateCache::CreateAndAdd(const FVulkanPipelineStateKey& Crea
 	//UE_LOG(LogVulkanRHI, Display, TEXT("PK: Added Entry %llx Index %d"), CreateInfo.PipelineKey, DiskEntries.Num());
 
 	// Add to the cache
-	Cache.Add(CreateInfo, Pipeline);
+	KeyToPipelineMap.Add(CreateInfo, Pipeline);
 	DiskEntries.Add(DiskEntry);
 }
 
@@ -864,6 +861,7 @@ void FVulkanPipelineStateCache::CreatePipelineFromDiskEntry(const FDiskEntry* Di
 
 	PipelineInfo.pDynamicState = &DynamicState;
 
+	//#todo-rco: Group the pipelines and create multiple (derived) pipelines at once
 	VERIFYVULKANRESULT(vkCreateGraphicsPipelines(Device->GetInstanceHandle(), PipelineCache, 1, &PipelineInfo, nullptr, &Pipeline->Pipeline));
 }
 
@@ -999,9 +997,47 @@ TLinkedList<FVulkanBoundShaderState*>*& FVulkanPipelineStateCache::GetBSSList()
 	return GBSSList;
 }
 
-void FVulkanPipelineStateCache::RebuildCache()
+void FVulkanPipelineStateCache::DestroyCache()
 {
-	UE_LOG(LogVulkanRHI, Warning, TEXT("Rebuilding pipeline cache; ditching %d entries"), DiskEntries.Num());
+	VkDevice DeviceHandle = Device->GetInstanceHandle();
+	auto DestroyDiskEntry = [DeviceHandle](FDiskEntry& DiskEntry)
+	{
+		for (int32 Index = 0; Index < ARRAY_COUNT(DiskEntry.ShaderModules); ++Index)
+		{
+			if (DiskEntry.ShaderModules[Index] != VK_NULL_HANDLE)
+			{
+				vkDestroyShaderModule(DeviceHandle, DiskEntry.ShaderModules[Index], nullptr);
+			}
+		}
+
+		vkDestroyPipelineLayout(DeviceHandle, DiskEntry.PipelineLayout, nullptr);
+
+		for (int32 Index = 0; Index < DiskEntry.DescriptorSetLayouts.Num(); ++Index)
+		{
+			vkDestroyDescriptorSetLayout(DeviceHandle, DiskEntry.DescriptorSetLayouts[Index], nullptr);
+		}
+
+		delete DiskEntry.RenderPass;
+		DiskEntry.RenderPass = nullptr;
+	};
+
+	for (FDiskEntry& Entry : DiskEntries)
+	{
+		if (Entry.bLoaded)
+		{
+			DestroyDiskEntry(Entry);
+			Entry.bLoaded = false;
+
+			FVulkanPipeline* Pipeline = CreatedPipelines.FindAndRemoveChecked(&Entry);
+			check(Pipeline);
+			Pipeline->Destroy();
+			delete Pipeline;
+		}
+		else
+		{
+			Entry.RenderPass = nullptr;
+		}
+	}
 
 	DiskEntries.Empty();
 	for (TLinkedList<FVulkanBoundShaderState*>::TIterator It(GetBSSList()); It; It.Next())
@@ -1011,7 +1047,16 @@ void FVulkanPipelineStateCache::RebuildCache()
 		//BSS->DestroyPipelineCache();
 		BSS->PipelineCache.Empty(0);
 	}
-	Cache.Empty();
+
+	KeyToPipelineMap.Empty();
+	CreatedPipelines.Empty();
+}
+
+void FVulkanPipelineStateCache::RebuildCache()
+{
+	UE_LOG(LogVulkanRHI, Warning, TEXT("Rebuilding pipeline cache; ditching %d entries"), DiskEntries.Num());
+
+	DestroyCache();
 }
 
 #endif	// VULKAN_ENABLE_PIPELINE_CACHE
diff --git a/Engine/Source/Runtime/VulkanRHI/Private/VulkanPipeline.h b/Engine/Source/Runtime/VulkanRHI/Private/VulkanPipeline.h
--- a/Engine/Source/Runtime/VulkanRHI/Private/VulkanPipeline.h
+++ b/Engine/Source/Runtime/VulkanRHI/Private/VulkanPipeline.h
@@ -6,8 +6,8 @@
 
 #pragma once
 
-struct FVulkanDevice;
-struct FVulkanFramebuffer;
+class FVulkanDevice;
+class FVulkanFramebuffer;
 
 // High level description of the state
 struct FVulkanPipelineState
@@ -124,7 +124,7 @@ class FVulkanPipelineStateCache
 public:
 	FVulkanPipeline* Find(const FVulkanPipelineStateKey& CreateInfo)
 	{
-		auto** Found = Cache.Find(CreateInfo);
+		auto** Found = KeyToPipelineMap.Find(CreateInfo);
 #if 0
 		if (Found)
 		{
@@ -504,7 +504,11 @@ public:
 private:
 	FVulkanDevice* Device;
 
-	TMap<FVulkanPipelineStateKey, FVulkanPipeline*> Cache;
+	// Map the runtime key to a pipeline pointer
+	TMap<FVulkanPipelineStateKey, FVulkanPipeline*> KeyToPipelineMap;
+
+	// Map used to delete the created pipelines
+	TMap<FDiskEntry*, FVulkanPipeline*> CreatedPipelines;
 
 	TIndirectArray<FDiskEntry> DiskEntries;
 
@@ -514,6 +518,7 @@ private:
 	void PopulateDiskEntry(const FVulkanPipelineState& State, const FVulkanRenderPass* RenderPass, FDiskEntry* OutDiskEntry);
 	void CreateDiskEntryRuntimeObjects(FDiskEntry* DiskEntry);
 	bool Load(const TArray<FString>& CacheFilenames, TArray<uint8>& OutDeviceCache);
+	void DestroyCache();
 };
 
 #endif	// VULKAN_ENABLE_PIPELINE_CACHE
diff --git a/Engine/Source/Runtime/VulkanRHI/Private/VulkanQuery.cpp b/Engine/Source/Runtime/VulkanRHI/Private/VulkanQuery.cpp
--- a/Engine/Source/Runtime/VulkanRHI/Private/VulkanQuery.cpp
+++ b/Engine/Source/Runtime/VulkanRHI/Private/VulkanQuery.cpp
@@ -174,5 +174,5 @@ void FVulkanCommandListContext::RHIEndOcclusionQueryBatch()
 
 void FVulkanCommandListContext::RHISubmitCommandsHint()
 {
-	//#todo-rco: Submit at this point if possible
+	SubmitCurrentCommands();
 }
