Copying //UE4/Dev-Rendering to Dev-Main (Source //UE4/Dev-Rendering@2932636)

diff --git a/Engine/Source/Runtime/VulkanRHI/Private/VulkanQueue.cpp b/Engine/Source/Runtime/VulkanRHI/Private/VulkanQueue.cpp
--- a/Engine/Source/Runtime/VulkanRHI/Private/VulkanQueue.cpp
+++ b/Engine/Source/Runtime/VulkanRHI/Private/VulkanQueue.cpp
@@ -7,59 +7,54 @@
 #include "VulkanRHIPrivate.h"
 #include "VulkanQueue.h"
 #include "VulkanSwapChain.h"
+#include "VulkanMemory.h"
 
-#if !VULKAN_USE_FENCE_MANAGER
-static inline void CreateFence(VkDevice Device, VkFence& Fence)
-{
-	VkFenceCreateInfo FenceCreateInfo;
-	FMemory::Memzero(FenceCreateInfo);
-	FenceCreateInfo.sType = VK_STRUCTURE_TYPE_FENCE_CREATE_INFO;
-	VERIFYVULKANRESULT(vkCreateFence(Device, &FenceCreateInfo, nullptr, &Fence));
-}
-#endif
-
-FVulkanQueue::FVulkanQueue(FVulkanDevice* InDevice, uint32 InNodeIndex)
-	: AcquireNextImageKHR(VK_NULL_HANDLE)
+FVulkanQueue::FVulkanQueue(FVulkanDevice* InDevice, uint32 InFamilyIndex, uint32 InQueueIndex)
+	: 
+#if VULKAN_USE_NEW_COMMAND_BUFFERS
+#else
+	AcquireNextImageKHR(VK_NULL_HANDLE)
 	, QueuePresentKHR(VK_NULL_HANDLE)
 	, CurrentFenceIndex(0)
 	, CurrentImageIndex(-1)
-	, Queue(VK_NULL_HANDLE)
-	, NodeIndex(InNodeIndex)
+	, 
+#endif
+	Queue(VK_NULL_HANDLE)
+	, FamilyIndex(InFamilyIndex)
+	, QueueIndex(InQueueIndex)
 	, Device(InDevice)
 {
+#if VULKAN_USE_NEW_COMMAND_BUFFERS
+#else
 	check(Device);
-
 	for (int BufferIndex = 0; BufferIndex < VULKAN_NUM_IMAGE_BUFFERS; ++BufferIndex)
 	{
 		ImageAcquiredSemaphore[BufferIndex] = new FVulkanSemaphore(*Device);
 		RenderingCompletedSemaphore[BufferIndex] = new FVulkanSemaphore(*Device);
-#if VULKAN_USE_FENCE_MANAGER
 		Fences[BufferIndex] = nullptr;
-#else
-		Fences[BufferIndex] = VK_NULL_HANDLE;
-#endif
 	}
+#endif
+	vkGetDeviceQueue(Device->GetInstanceHandle(), FamilyIndex, InQueueIndex, &Queue);
 
-	vkGetDeviceQueue(Device->GetInstanceHandle(), NodeIndex, 0, &Queue);
-
+#if VULKAN_USE_NEW_COMMAND_BUFFERS
+#else
 	AcquireNextImageKHR = (PFN_vkAcquireNextImageKHR)(void*)vkGetDeviceProcAddr(Device->GetInstanceHandle(), "vkAcquireNextImageKHR");
 	check(AcquireNextImageKHR);
 
 	QueuePresentKHR = (PFN_vkQueuePresentKHR)(void*)vkGetDeviceProcAddr(Device->GetInstanceHandle(), "vkQueuePresentKHR");
 	check(QueuePresentKHR);
+#endif
 }
 
 FVulkanQueue::~FVulkanQueue()
 {
 	check(Device);
 
+#if VULKAN_USE_NEW_COMMAND_BUFFERS
+#else
 	if (Fences[CurrentFenceIndex] != VK_NULL_HANDLE)
 	{
-#if VULKAN_USE_FENCE_MANAGER
 		Device->GetFenceManager().WaitForFence(Fences[CurrentFenceIndex], UINT32_MAX);
-#else
-		VERIFYVULKANRESULT(vkWaitForFences(Device->GetInstanceHandle(), 1, &Fences[CurrentFenceIndex], VK_TRUE, UINT64_MAX));
-#endif
 	}
 
 	for (int BufferIndex = 0; BufferIndex < VULKAN_NUM_IMAGE_BUFFERS; ++BufferIndex)
@@ -69,22 +64,17 @@ FVulkanQueue::~FVulkanQueue()
 		delete RenderingCompletedSemaphore[BufferIndex];
 		RenderingCompletedSemaphore[BufferIndex] = nullptr;
 
-#if VULKAN_USE_FENCE_MANAGER
 		if (Fences[BufferIndex])
 		{
 			Device->GetFenceManager().ReleaseFence(Fences[BufferIndex]);
 			Fences[BufferIndex] = nullptr;
 		}
-#else
-		if (Fences[BufferIndex] != VK_NULL_HANDLE)
-		{
-			vkDestroyFence(Device->GetInstanceHandle(), Fences[BufferIndex], nullptr);
-			Fences[BufferIndex] = VK_NULL_HANDLE;
-		}
-#endif
 	}
+#endif
 }
 
+#if VULKAN_USE_NEW_COMMAND_BUFFERS
+#else
 uint32 FVulkanQueue::AquireImageIndex(FVulkanSwapChain* Swapchain)
 {	
 	check(Device);
@@ -104,13 +94,8 @@ uint32 FVulkanQueue::AquireImageIndex(FVulkanSwapChain* Swapchain)
 		#endif
 		{
 			// Grab the next fence and re-use it.
-#if VULKAN_USE_FENCE_MANAGER
 			Device->GetFenceManager().WaitForFence(Fences[CurrentFenceIndex], UINT32_MAX);
 			Device->GetFenceManager().ResetFence(Fences[CurrentFenceIndex]);
-#else
-			VERIFYVULKANRESULT(vkWaitForFences(Device->GetInstanceHandle(), 1, &Fences[CurrentFenceIndex], VK_TRUE, UINT64_MAX));
-			VERIFYVULKANRESULT(vkResetFences(Device->GetInstanceHandle(), 1, &Fences[CurrentFenceIndex]));
-#endif
 		}
 	}
 
@@ -131,7 +116,31 @@ uint32 FVulkanQueue::AquireImageIndex(FVulkanSwapChain* Swapchain)
 
 	return CurrentImageIndex;
 }
+#endif
+
+#if VULKAN_USE_NEW_COMMAND_BUFFERS
+void FVulkanQueue::Submit(FVulkanCmdBuffer* CmdBuffer)
+{
+	check(CmdBuffer->HasEnded());
+
+	VkSubmitInfo SubmitInfo;
+	FMemory::Memzero(SubmitInfo);
+
+	auto* Fence = CmdBuffer->GetFence();
+	check(!Fence->IsSignaled());
 
+	const VkCommandBuffer CmdBuffers[] = { CmdBuffer->GetHandle() };
+
+	SubmitInfo.sType = VK_STRUCTURE_TYPE_SUBMIT_INFO;
+	SubmitInfo.commandBufferCount = 1;
+	SubmitInfo.pCommandBuffers = CmdBuffers;
+	VERIFYVULKANRESULT(vkQueueSubmit(Queue, 1, &SubmitInfo, Fence->GetHandle()));
+
+	CmdBuffer->State = FVulkanCmdBuffer::EState::Submitted;
+
+	CmdBuffer->GetOwner()->RefreshFenceStatus();
+}
+#else
 void FVulkanQueue::Submit(FVulkanCmdBuffer* CmdBuffer)
 {
 	check(CmdBuffer);
@@ -156,27 +165,40 @@ void FVulkanQueue::Submit(FVulkanCmdBuffer* CmdBuffer)
 	SubmitInfo.commandBufferCount = 1;
 	SubmitInfo.pCommandBuffers = CmdBuffers;
 
-#if VULKAN_USE_FENCE_MANAGER
 	if (!Fences[CurrentFenceIndex])
 	{
 		Fences[CurrentFenceIndex] = Device->GetFenceManager().AllocateFence();
 	}
 	VERIFYVULKANRESULT(vkQueueSubmit(Queue, 1, &SubmitInfo, Fences[CurrentFenceIndex]->GetHandle()));
-#else
-	if(Fences[CurrentFenceIndex] == VK_NULL_HANDLE)
-	{
-		CreateFence(Device->GetInstanceHandle(), Fences[CurrentFenceIndex]);
-	}
-	VERIFYVULKANRESULT(vkQueueSubmit(Queue, 1, &SubmitInfo, Fences[CurrentFenceIndex]));
+
+#if VULKAN_FORCE_WAIT_FOR_QUEUE
+	VERIFYVULKANRESULT(vkQueueWaitIdle(Queue));
+#endif
+}
 #endif
 
+#if !VULKAN_USE_NEW_COMMAND_BUFFERS
+void FVulkanQueue::Submit2(VkCommandBuffer CmdBuffer, VulkanRHI::FFence* Fence)
+{
+	VkSubmitInfo SubmitInfo;
+	FMemory::Memzero(SubmitInfo);
+	SubmitInfo.sType = VK_STRUCTURE_TYPE_SUBMIT_INFO;
+	SubmitInfo.commandBufferCount = 1;
+	SubmitInfo.pCommandBuffers = &CmdBuffer;
+	check(!Fence || !Fence->IsSignaled());
+	VERIFYVULKANRESULT(vkQueueSubmit(Queue, 1, &SubmitInfo, Fence ? Fence->GetHandle() : VK_NULL_HANDLE));
+
 #if VULKAN_FORCE_WAIT_FOR_QUEUE
 	VERIFYVULKANRESULT(vkQueueWaitIdle(Queue));
 #endif
 }
+#endif
 
 void FVulkanQueue::SubmitBlocking(FVulkanCmdBuffer* CmdBuffer)
 {
+#if VULKAN_USE_NEW_COMMAND_BUFFERS
+	check(0);
+#else
 	check(CmdBuffer);
 	CmdBuffer->End();
 
@@ -188,7 +210,6 @@ void FVulkanQueue::SubmitBlocking(FVulkanCmdBuffer* CmdBuffer)
 	SubmitInfo.commandBufferCount = 1;
 	SubmitInfo.pCommandBuffers = CmdBufs;
 
-#if VULKAN_USE_FENCE_MANAGER
 	auto* Fence = Device->GetFenceManager().AllocateFence();
 	VERIFYVULKANRESULT(vkQueueSubmit(Queue, 1, &SubmitInfo, Fence->GetHandle()));
 
@@ -196,23 +217,11 @@ void FVulkanQueue::SubmitBlocking(FVulkanCmdBuffer* CmdBuffer)
 	check(bSuccess);
 
 	Device->GetFenceManager().ReleaseFence(Fence);
-#else
-	VkFence Fence = { VK_NULL_HANDLE };
-	VkFenceCreateInfo FenceCreateInfo;
-	FMemory::Memzero(FenceCreateInfo);
-	FenceCreateInfo.sType = VK_STRUCTURE_TYPE_FENCE_CREATE_INFO;
-	VERIFYVULKANRESULT(vkCreateFence(Device->GetInstanceHandle(), &FenceCreateInfo, nullptr, &Fence));
-
-	VERIFYVULKANRESULT(vkQueueSubmit(Queue, 1, &SubmitInfo, Fence));
-
-	VERIFYVULKANRESULT(vkWaitForFences(Device->GetInstanceHandle(), 1, &Fence, true, 0xFFFFFFFF));
-
-	VERIFYVULKANRESULT(vkResetFences(Device->GetInstanceHandle(), 1, &Fence));
-
-	vkDestroyFence(Device->GetInstanceHandle(), Fence, nullptr);
 #endif
 }
 
+#if VULKAN_USE_NEW_COMMAND_BUFFERS
+#else
 void FVulkanQueue::Present(FVulkanSwapChain* Swapchain, uint32 ImageIndex)
 {
 	check(Swapchain);
@@ -220,8 +229,6 @@ void FVulkanQueue::Present(FVulkanSwapChain* Swapchain, uint32 ImageIndex)
 	VkPresentInfoKHR Info;
 	FMemory::Memzero(Info);
 	Info.sType = VK_STRUCTURE_TYPE_PRESENT_INFO_KHR;
-	Info.pNext = nullptr;
-
 	Info.waitSemaphoreCount = 1;	// 0;	// 1;
 	VkSemaphore Local = RenderingCompletedSemaphore[CurrentFenceIndex]->GetHandle();
 	Info.pWaitSemaphores = &Local;
@@ -231,3 +238,4 @@ void FVulkanQueue::Present(FVulkanSwapChain* Swapchain, uint32 ImageIndex)
 
 	VERIFYVULKANRESULT(QueuePresentKHR(Queue, &Info));
 }
+#endif
diff --git a/Engine/Source/Runtime/VulkanRHI/Private/VulkanQueue.h b/Engine/Source/Runtime/VulkanRHI/Private/VulkanQueue.h
--- a/Engine/Source/Runtime/VulkanRHI/Private/VulkanQueue.h
+++ b/Engine/Source/Runtime/VulkanRHI/Private/VulkanQueue.h
@@ -8,29 +8,38 @@
 
 #include "VulkanRHIPrivate.h"
 
-struct FVulkanDevice;
+class FVulkanDevice;
 class FVulkanCmdBuffer;
 struct FVulkanSemaphore;
-struct FVulkanSwapChain;
+class FVulkanSwapChain;
 
-struct FVulkanQueue
+namespace VulkanRHI
 {
-	FVulkanQueue(FVulkanDevice* InDevice, uint32 InNodeIndex);
+	class FFence;
+}
+
+class FVulkanQueue
+{
+public:
+	FVulkanQueue(FVulkanDevice* InDevice, uint32 InFamilyIndex, uint32 InQueueIndex);
 
 	~FVulkanQueue();
 
-	inline uint32 GetNodeIndex() const
+	inline uint32 GetFamilyIndex() const
 	{
-		return NodeIndex;
+		return FamilyIndex;
 	}
 
 	void Submit(FVulkanCmdBuffer* CmdBuffer);
 
 	void SubmitBlocking(FVulkanCmdBuffer* CmdBuffer);
 
+#if VULKAN_USE_NEW_COMMAND_BUFFERS
+#else
+	void Submit2(VkCommandBuffer CmdBuffer, VulkanRHI::FFence* Fence);
 	uint32 AquireImageIndex(FVulkanSwapChain* Swapchain);
-	
 	void Present(FVulkanSwapChain* Swapchain, uint32 ImageIndex);
+#endif
 
 	inline VkQueue GetHandle() const
 	{
@@ -38,22 +47,20 @@ struct FVulkanQueue
 	}
 
 private:
+#if VULKAN_USE_NEW_COMMAND_BUFFERS
+#else
 	PFN_vkAcquireNextImageKHR AcquireNextImageKHR;
 	PFN_vkQueuePresentKHR QueuePresentKHR;
 
 	FVulkanSemaphore* ImageAcquiredSemaphore[VULKAN_NUM_IMAGE_BUFFERS];
 	FVulkanSemaphore* RenderingCompletedSemaphore[VULKAN_NUM_IMAGE_BUFFERS];
-#if VULKAN_USE_FENCE_MANAGER
 	VulkanRHI::FFence* Fences[VULKAN_NUM_IMAGE_BUFFERS];
-#else
-	VkFence Fences[VULKAN_NUM_IMAGE_BUFFERS];
-#endif
 	uint32	CurrentFenceIndex;
 
 	uint32 CurrentImageIndex;	// Perhaps move this to the FVulkanSwapChain?
-
+#endif
 	VkQueue Queue;
-	uint32 NodeIndex;
+	uint32 FamilyIndex;
+	uint32 QueueIndex;
 	FVulkanDevice* Device;
 };
-
diff --git a/Engine/Source/Runtime/VulkanRHI/Private/VulkanRHI.cpp b/Engine/Source/Runtime/VulkanRHI/Private/VulkanRHI.cpp
--- a/Engine/Source/Runtime/VulkanRHI/Private/VulkanRHI.cpp
+++ b/Engine/Source/Runtime/VulkanRHI/Private/VulkanRHI.cpp
@@ -197,7 +197,10 @@ FVulkanCommandListContext::FVulkanCommandListContext(FVulkanDynamicRHI* InRHI, F
 	, PendingIndexDataStride(0)
 	, DynamicVB(nullptr)
 	, DynamicIB(nullptr)
+#if VULKAN_USE_NEW_COMMAND_BUFFERS
+#else
 	, CommandBufferManager(nullptr)
+#endif
 {
 	DynamicIB = new FVulkanDynamicIndexBuffer(*Device);
 	DynamicVB = new FVulkanDynamicVertexBuffer(*Device);
@@ -235,6 +238,9 @@ FVulkanDynamicRHI::FVulkanDynamicRHI()
 	, SavePipelineCacheCmd(nullptr)
 	, RebuildPipelineCacheCmd(nullptr)
 #endif
+#if UE_BUILD_DEBUG || UE_BUILD_DEVELOPMENT
+	, DumpMemoryCmd(nullptr)
+#endif
 {
 	// This should be called once at the start 
 	check(IsInGameThread());
@@ -314,6 +320,10 @@ void FVulkanDynamicRHI::Shutdown()
 	IConsoleManager::Get().UnregisterConsoleObject(SavePipelineCacheCmd);
 	IConsoleManager::Get().UnregisterConsoleObject(RebuildPipelineCacheCmd);
 #endif
+
+#if UE_BUILD_DEBUG || UE_BUILD_DEVELOPMENT
+	IConsoleManager::Get().UnregisterConsoleObject(DumpMemoryCmd);
+#endif
 }
 
 void FVulkanDynamicRHI::CreateInstance()
@@ -462,10 +472,11 @@ void FVulkanDynamicRHI::InitInstance()
 		GSupportsRenderTargetFormat_PF_G8 = false;	// #todo-rco
 		GSupportsQuads = false;	// Not supported in Vulkan
 		GRHISupportsTextureStreaming = false;	// #todo-rco
+		GRHISupportsRHIThread = 0;//VULKAN_USE_NEW_RESOURCE_MANAGEMENT != 0;	// #todo-rco
 
 		//#todo-rco: Leaks/issues still present!
 		// Indicate that the RHI needs to use the engine's deferred deletion queue.
-		GRHINeedsExtraDeletionLatency = false;
+		GRHINeedsExtraDeletionLatency = 0;//VULKAN_USE_NEW_RESOURCE_MANAGEMENT != 0 && VULKAN_USE_NEW_COMMAND_BUFFERS != 0;
 
 		GMaxShadowDepthBufferSizeX =  FPlatformMath::Min<int32>(Props.limits.maxImageDimension2D, GMaxShadowDepthBufferSizeX);
 		GMaxShadowDepthBufferSizeY =  FPlatformMath::Min<int32>(Props.limits.maxImageDimension2D, GMaxShadowDepthBufferSizeY);
@@ -519,6 +530,15 @@ void FVulkanDynamicRHI::InitInstance()
 			ECVF_Default
 			);
 #endif
+
+#if UE_BUILD_DEBUG || UE_BUILD_DEVELOPMENT
+		DumpMemoryCmd = IConsoleManager::Get().RegisterConsoleCommand(
+			TEXT("DumpVulkanMemory"),
+			TEXT("Dumps memory map."),
+			FConsoleCommandDelegate::CreateStatic(DumpMemory),
+			ECVF_Default
+			);
+#endif
 	}
 }
 
@@ -544,8 +564,12 @@ void FVulkanCommandListContext::RHIBeginDrawingViewport(FViewportRHIParamRef Vie
 	FVulkanViewport* Viewport = ResourceCast(ViewportRHI);
 	if (Viewport->GetBackBufferIndex() < 0)
 	{
+#if VULKAN_USE_NEW_COMMAND_BUFFERS
+		Viewport->AcquireBackBuffer(Device->GetImmediateContext().GetCommandBufferManager()->GetActiveCmdBuffer());
+#else
 		// SMEDIS: This is the normal method of getting a new backbuffer. However, RHIGetViewportBackBuffer() can also do it if necessary.
 		Viewport->CurrentBackBuffer = Device->GetQueue()->AquireImageIndex(Viewport->SwapChain);
+#endif
 	}
 	RHI->DrawingViewport = Viewport;
 
@@ -584,6 +608,12 @@ void FVulkanCommandListContext::RHIEndDrawingViewport(FViewportRHIParamRef Viewp
 
 void FVulkanCommandListContext::RHIEndFrame()
 {
+#if VULKAN_USE_NEW_COMMAND_BUFFERS && VULKAN_USE_NEW_RESOURCE_MANAGEMENT
+	Device->GetDeferredDeletionQueue().Clear();
+#endif
+
+	Device->GetStagingManager().ProcessPendingFree();
+
 	Device->FrameCounter++;
 
 	DynamicIB->EndFrame();
@@ -670,15 +700,7 @@ IRHICommandContextContainer* FVulkanDynamicRHI::RHIGetCommandContextContainer()
 FVulkanBuffer::FVulkanBuffer(FVulkanDevice& InDevice, uint32 InSize, VkFlags InUsage, VkMemoryPropertyFlags InMemPropertyFlags, bool bInAllowMultiLock, const char* File, int32 Line) :
 	Device(InDevice),
 	Buf(VK_NULL_HANDLE),
-#if VULKAN_USE_MEMORY_SYSTEM
-	#if VULKAN_USE_BUFFER_HEAPS
-		MemoryInfo(nullptr),
-	#else
-		Allocation(nullptr),
-	#endif
-#else
-	Mem(VK_NULL_HANDLE),
-#endif
+	Allocation(nullptr),
 	Size(InSize),
 	Usage(InUsage),
 	BufferPtr(nullptr),
@@ -695,27 +717,9 @@ FVulkanBuffer::FVulkanBuffer(FVulkanDevice& InDevice, uint32 InSize, VkFlags InU
 	VkMemoryRequirements MemoryRequirements;
 	vkGetBufferMemoryRequirements(Device.GetInstanceHandle(), Buf, &MemoryRequirements);
 
-#if VULKAN_USE_MEMORY_SYSTEM
-	#if VULKAN_USE_BUFFER_HEAPS
-		MemoryInfo = InDevice.GetResourceAllocationManager().AllocateBuffer(MemoryRequirements, InMemPropertyFlags);
-		VERIFYVULKANRESULT_EXPANDED(vkBindBufferMemory(Device.GetInstanceHandle(), Buf, MemoryInfo->GetHandle(), MemoryInfo->GetOffset()));
-#else
-		Allocation = InDevice.GetMemoryManager().Alloc(MemoryRequirements.size, MemoryRequirements.memoryTypeBits, InMemPropertyFlags, File ? File : __FILE__, Line ? Line : __LINE__);
-		check(Allocation);
-		VERIFYVULKANRESULT_EXPANDED(vkBindBufferMemory(Device.GetInstanceHandle(), Buf, Allocation->GetHandle(), 0));
-	#endif
-#else
-	VkMemoryAllocateInfo AllocInfo;
-	FMemory::Memzero(AllocInfo);
-	AllocInfo.sType = VK_STRUCTURE_TYPE_MEMORY_ALLOCATE_INFO;
-	AllocInfo.pNext = nullptr;
-	AllocInfo.allocationSize = MemoryRequirements.size;
-	AllocInfo.memoryTypeIndex = 0;
-	VERIFYVULKANRESULT_EXPANDED(Device.GetMemoryManager().GetMemoryTypeFromProperties(MemoryRequirements.memoryTypeBits, InMemPropertyFlags, &AllocInfo.memoryTypeIndex));
-
-	VERIFYVULKANRESULT_EXPANDED(vkAllocateMemory(Device.GetInstanceHandle(), &AllocInfo, nullptr, &Mem));
-	VERIFYVULKANRESULT_EXPANDED(vkBindBufferMemory(Device.GetInstanceHandle(), Buf, Mem, 0));
-#endif
+	Allocation = InDevice.GetMemoryManager().Alloc(MemoryRequirements.size, MemoryRequirements.memoryTypeBits, InMemPropertyFlags, File ? File : __FILE__, Line ? Line : __LINE__);
+	check(Allocation);
+	VERIFYVULKANRESULT_EXPANDED(vkBindBufferMemory(Device.GetInstanceHandle(), Buf, Allocation->GetHandle(), 0));
 }
 
 FVulkanBuffer::~FVulkanBuffer()
@@ -726,17 +730,8 @@ FVulkanBuffer::~FVulkanBuffer()
 	vkDestroyBuffer(Device.GetInstanceHandle(), Buf, nullptr);
 	Buf = VK_NULL_HANDLE;
 
-#if VULKAN_USE_MEMORY_SYSTEM
-	#if VULKAN_USE_BUFFER_HEAPS
-		Device.GetResourceAllocationManager().FreeBuffer(MemoryInfo);
-	#else
-		Device.GetMemoryManager().Free(Allocation);
-		Allocation = nullptr;
-	#endif
-#else
-	vkFreeMemory(Device.GetInstanceHandle(), Mem, nullptr);
-	Mem = VK_NULL_HANDLE;
-#endif
+	Device.GetMemoryManager().Free(Allocation);
+	Allocation = nullptr;
 }
 
 void* FVulkanBuffer::Lock(uint32 InSize, uint32 InOffset)
@@ -752,15 +747,7 @@ void* FVulkanBuffer::Lock(uint32 InSize, uint32 InOffset)
 		if (LockStack == 0)
 		{
 			// lock the whole range
-#if VULKAN_USE_MEMORY_SYSTEM
-	#if VULKAN_USE_BUFFER_HEAPS
-			BufferPtr = MemoryInfo->Map(GetSize(), 0);
-	#else
 			BufferPtr = Allocation->Map(GetSize(), 0);
-	#endif
-#else
-			VERIFYVULKANRESULT_EXPANDED(vkMapMemory(Device.GetInstanceHandle(), Mem, 0, GetSize(), 0, &BufferPtr));
-#endif
 		}
 		// offset the whole range by the requested offset
 		BufferPtrOffset = InOffset;
@@ -769,15 +756,7 @@ void* FVulkanBuffer::Lock(uint32 InSize, uint32 InOffset)
 	else
 	{
 		check(BufferPtr == nullptr);
-#if VULKAN_USE_MEMORY_SYSTEM
-	#if VULKAN_USE_BUFFER_HEAPS
-		BufferPtr = MemoryInfo->Map(InSize, InOffset);
-	#else
 		BufferPtr = Allocation->Map(InSize, InOffset);
-	#endif
-#else
-		VERIFYVULKANRESULT_EXPANDED(vkMapMemory(Device.GetInstanceHandle(), Mem, InOffset, InSize, 0, &BufferPtr));
-#endif
 	}
 
 	return (uint8*)BufferPtr + BufferPtrOffset;
@@ -794,21 +773,15 @@ void FVulkanBuffer::Unlock()
 		return;
 	}
 
-#if VULKAN_USE_MEMORY_SYSTEM
-	#if VULKAN_USE_BUFFER_HEAPS
-		MemoryInfo->Unmap();
-	#else
-		Allocation->Unmap();
-	#endif
-#else
-	check(Mem);
-	vkUnmapMemory(Device.GetInstanceHandle(), Mem);
-#endif
+	Allocation->Unmap();
 	BufferPtr = nullptr;
 }
 
 void FVulkanBuffer::CopyTo(FVulkanSurface& Surface, const VkBufferImageCopy& CopyDesc, FVulkanPendingState* State)
 {
+#if VULKAN_USE_NEW_COMMAND_BUFFERS
+	check(0);
+#else
 	VkCommandBuffer Cmd = VK_NULL_HANDLE;
 	FVulkanCmdBuffer* CmdObject = nullptr;
 	if(State)
@@ -833,6 +806,7 @@ void FVulkanBuffer::CopyTo(FVulkanSurface& Surface, const VkBufferImageCopy& Cop
 		Device.GetQueue()->SubmitBlocking(CmdObject);
 		Device.GetImmediateContext().GetCommandBufferManager()->Destroy(CmdObject);
 	}
+#endif
 }
 
 
@@ -1202,6 +1176,8 @@ FVulkanRenderPass::~FVulkanRenderPass()
 	RenderPass = VK_NULL_HANDLE;
 }
 
+#if VULKAN_USE_NEW_COMMAND_BUFFERS
+#else
 void FVulkanRenderPass::Begin(FVulkanCmdBuffer& CmdBuf, const VkFramebuffer& Framebuffer, const VkClearValue* AttachmentClearValues)
 {
 	VkRenderPassBeginInfo Info;
@@ -1217,120 +1193,68 @@ void FVulkanRenderPass::Begin(FVulkanCmdBuffer& CmdBuf, const VkFramebuffer& Fra
 
 	vkCmdBeginRenderPass(CmdBuf.GetHandle(false), &Info, VK_SUBPASS_CONTENTS_INLINE);
 }
+#endif
+
 
 void VulkanSetImageLayout(
 	VkCommandBuffer CmdBuffer,
-	VkImage image,
-	VkImageAspectFlags aspectMask,
-	VkImageLayout old_image_layout,
-	VkImageLayout new_image_layout)
+	VkImage Image,
+	VkImageLayout OldLayout,
+	VkImageLayout NewLayout,
+	VkImageSubresourceRange SubresourceRange)
 {
 	check(CmdBuffer != VK_NULL_HANDLE);
 
-	VkImageMemoryBarrier image_memory_barrier;
-	FMemory::Memzero(image_memory_barrier);
-	image_memory_barrier.sType = VK_STRUCTURE_TYPE_IMAGE_MEMORY_BARRIER;
-	image_memory_barrier.pNext = NULL;
-	image_memory_barrier.srcAccessMask = 0;
-	image_memory_barrier.dstAccessMask = 0;
-	image_memory_barrier.oldLayout = old_image_layout;
-	image_memory_barrier.newLayout = new_image_layout;
-	image_memory_barrier.image = image;
-	image_memory_barrier.subresourceRange = { aspectMask, 0, 1, 0, 1 };
-
-	if (new_image_layout == VK_IMAGE_LAYOUT_TRANSFER_DST_OPTIMAL)
+	VkImageMemoryBarrier ImageBarrier;
+	FMemory::Memzero(ImageBarrier);
+	ImageBarrier.sType = VK_STRUCTURE_TYPE_IMAGE_MEMORY_BARRIER;
+	ImageBarrier.oldLayout = OldLayout;
+	ImageBarrier.newLayout = NewLayout;
+	ImageBarrier.image = Image;
+	ImageBarrier.subresourceRange = SubresourceRange;
+	ImageBarrier.srcQueueFamilyIndex = VK_QUEUE_FAMILY_IGNORED;
+	ImageBarrier.dstQueueFamilyIndex = VK_QUEUE_FAMILY_IGNORED;
+
+	switch (NewLayout)
 	{
-		image_memory_barrier.dstAccessMask = VK_ACCESS_TRANSFER_READ_BIT;
-	}
-	else if (new_image_layout == VK_IMAGE_LAYOUT_COLOR_ATTACHMENT_OPTIMAL)
-	{
-		image_memory_barrier.dstAccessMask = VK_ACCESS_COLOR_ATTACHMENT_WRITE_BIT;
-	}
-	else if (new_image_layout == VK_IMAGE_LAYOUT_DEPTH_STENCIL_ATTACHMENT_OPTIMAL)
-	{
-		image_memory_barrier.dstAccessMask = VK_ACCESS_DEPTH_STENCIL_ATTACHMENT_WRITE_BIT;
-	}
-	else if (new_image_layout == VK_IMAGE_LAYOUT_SHADER_READ_ONLY_OPTIMAL)
-	{
-		image_memory_barrier.dstAccessMask = VK_ACCESS_SHADER_READ_BIT | VK_ACCESS_INPUT_ATTACHMENT_READ_BIT;
+	case VK_IMAGE_LAYOUT_TRANSFER_DST_OPTIMAL:
+		ImageBarrier.dstAccessMask = VK_ACCESS_TRANSFER_READ_BIT;
+		break;
+	case VK_IMAGE_LAYOUT_COLOR_ATTACHMENT_OPTIMAL:
+		ImageBarrier.dstAccessMask = VK_ACCESS_COLOR_ATTACHMENT_WRITE_BIT;
+		break;
+	case VK_IMAGE_LAYOUT_DEPTH_STENCIL_ATTACHMENT_OPTIMAL:
+		ImageBarrier.dstAccessMask = VK_ACCESS_DEPTH_STENCIL_ATTACHMENT_WRITE_BIT;
+		break;
+	case VK_IMAGE_LAYOUT_SHADER_READ_ONLY_OPTIMAL:
+		ImageBarrier.dstAccessMask = VK_ACCESS_SHADER_READ_BIT | VK_ACCESS_INPUT_ATTACHMENT_READ_BIT;
+		break;
+	case VK_IMAGE_LAYOUT_PRESENT_SRC_KHR:
+		if (OldLayout == VK_IMAGE_LAYOUT_COLOR_ATTACHMENT_OPTIMAL)
+		{
+			ImageBarrier.srcAccessMask = VK_ACCESS_COLOR_ATTACHMENT_WRITE_BIT;
+		}
+		ImageBarrier.dstAccessMask = VK_ACCESS_MEMORY_READ_BIT;
+		break;
+	case VK_IMAGE_LAYOUT_GENERAL:
+		if (OldLayout == VK_IMAGE_LAYOUT_TRANSFER_DST_OPTIMAL)
+		{
+			ImageBarrier.srcAccessMask = VK_ACCESS_TRANSFER_WRITE_BIT;
+		}
+		break;
+	default:
+		check(0);
+		break;
 	}
 
-	VkImageMemoryBarrier *pmemory_barrier = &image_memory_barrier;
+	VkImageMemoryBarrier BarrierList[] ={ImageBarrier};
 
-	VkPipelineStageFlags src_stages = VK_PIPELINE_STAGE_TOP_OF_PIPE_BIT;
-	VkPipelineStageFlags dest_stages = VK_PIPELINE_STAGE_TOP_OF_PIPE_BIT;
+	VkPipelineStageFlags SourceStages = VK_PIPELINE_STAGE_TOP_OF_PIPE_BIT;
+	VkPipelineStageFlags DestStages = VK_PIPELINE_STAGE_TOP_OF_PIPE_BIT;
 
-	vkCmdPipelineBarrier(CmdBuffer, src_stages, dest_stages, 0, 0, nullptr, 0, nullptr, 1, &image_memory_barrier);
+	vkCmdPipelineBarrier(CmdBuffer, SourceStages, DestStages, 0, 0, nullptr, 0, nullptr, 1, BarrierList);
 }
 
-void VulkanImageCopy(VkCommandBuffer Cmd, FTextureRHIParamRef SourceTextureRHI, FTextureRHIParamRef DestTextureRHI)
-{	
-	FVulkanTextureBase* Src = FVulkanTextureBase::Cast(SourceTextureRHI);
-	FVulkanTextureBase* Dst = FVulkanTextureBase::Cast(DestTextureRHI);
-
-	check(Src->Surface.GetAspectMask() == Dst->Surface.GetAspectMask());
-	const VkImageAspectFlags AspectMask = Src->Surface.GetAspectMask();
-
-	VkImageLayout SrcLayoutOriginal = Src->Surface.ImageLayout;
-	VulkanSetImageLayout(
-			Cmd,
-			Src->Surface.Image,
-			AspectMask,
-			SrcLayoutOriginal,
-			VK_IMAGE_LAYOUT_TRANSFER_SRC_OPTIMAL);
-
-	
-	VkImageLayout DstLayoutOriginal = Dst->Surface.ImageLayout;
-	VulkanSetImageLayout(
-			Cmd,
-			Dst->Surface.Image,
-			AspectMask,
-			DstLayoutOriginal,
-			VK_IMAGE_LAYOUT_TRANSFER_DST_OPTIMAL);
-	
-	VkImageCopy CopyDesc;
-	FMemory::Memzero(CopyDesc);
-	CopyDesc.srcSubresource.baseArrayLayer = 0;
-	CopyDesc.srcSubresource.mipLevel = 0;
-	CopyDesc.srcSubresource.layerCount = 1;
-    CopyDesc.srcOffset.x = 0;
-    CopyDesc.srcOffset.y = 0;
-    CopyDesc.srcOffset.z = 0;
-	CopyDesc.dstSubresource.aspectMask = AspectMask;
-	CopyDesc.dstSubresource.baseArrayLayer = 0;
-	CopyDesc.dstSubresource.mipLevel = 0;
-	CopyDesc.dstSubresource.layerCount = 1;
-	CopyDesc.dstOffset.x = 0;
-	CopyDesc.dstOffset.y = 0;
-	CopyDesc.dstOffset.z = 0;
-    CopyDesc.extent.width = Src->Surface.Width;
-    CopyDesc.extent.height = Src->Surface.Height;
-    CopyDesc.extent.depth = 1;
-
-	// Do copy
-	vkCmdCopyImage(Cmd,
-		Src->Surface.Image, VK_IMAGE_LAYOUT_TRANSFER_SRC_OPTIMAL,
-		Dst->Surface.Image, VK_IMAGE_LAYOUT_TRANSFER_DST_OPTIMAL,
-		1, &CopyDesc);
-
-	// Restore SRC to the previouse layout
-	VulkanSetImageLayout(
-			Cmd,
-			Src->Surface.Image,
-			AspectMask,
-			VK_IMAGE_LAYOUT_TRANSFER_SRC_OPTIMAL,
-			SrcLayoutOriginal);
-	
-	// Restore DST to the previouse layout
-	VulkanSetImageLayout(
-			Cmd,
-			Dst->Surface.Image,
-			AspectMask,
-			VK_IMAGE_LAYOUT_TRANSFER_DST_OPTIMAL,
-			DstLayoutOriginal);
-}
-
-
 void VulkanResolveImage(VkCommandBuffer Cmd, FTextureRHIParamRef SourceTextureRHI, FTextureRHIParamRef DestTextureRHI)
 {
 	FVulkanTextureBase* Src = FVulkanTextureBase::Cast(SourceTextureRHI);
@@ -1514,3 +1438,14 @@ void FVulkanDynamicRHI::RebuildPipelineCache()
 	RHI->Device->PipelineStateCache->RebuildCache();
 }
 #endif
+
+#if UE_BUILD_DEBUG || UE_BUILD_DEVELOPMENT
+void FVulkanDynamicRHI::DumpMemory()
+{
+	auto* RHI = (FVulkanDynamicRHI*)GDynamicRHI;
+	RHI->Device->GetMemoryManager().DumpMemory();
+#if VULKAN_USE_NEW_RESOURCE_MANAGEMENT
+	RHI->Device->GetResourceHeapManager().DumpMemory();
+#endif
+}
+#endif
diff --git a/Engine/Source/Runtime/VulkanRHI/Private/VulkanRHIPrivate.h b/Engine/Source/Runtime/VulkanRHI/Private/VulkanRHIPrivate.h
--- a/Engine/Source/Runtime/VulkanRHI/Private/VulkanRHIPrivate.h
+++ b/Engine/Source/Runtime/VulkanRHI/Private/VulkanRHIPrivate.h
@@ -37,17 +37,17 @@ DECLARE_LOG_CATEGORY_EXTERN(LogVulkanRHI, Log, All);
 // Default is 1 (which is aniso off), the number is adjusted after the limits are queried.
 static int32 GMaxVulkanTextureFilterAnisotropic = 1;
 
-struct FVulkanQueue;
+class FVulkanQueue;
 class FVulkanCmdBuffer;
-struct FVulkanShader;
-struct FVulkanDescriptorSetsLayout;
+class FVulkanShader;
+class FVulkanDescriptorSetsLayout;
 class FVulkanBlendState;
 class FVulkanDepthStencilState;
 class FVulkanBoundShaderState;
 class FVulkanPipeline;
-struct FVulkanRenderPass;
-struct FVulkanCommandBufferManager;
-struct FVulkanPendingState;
+class FVulkanRenderPass;
+class FVulkanCommandBufferManager;
+class FVulkanPendingState;
 
 inline VkShaderStageFlagBits UEFrequencyToVKStageBit(EShaderFrequency InStage)
 {
@@ -165,54 +165,9 @@ private:
 #include "VulkanQueue.h"
 #include "VulkanCommandBuffer.h"
 
-struct FVulkanDepthStencilView
-{
-	FVulkanDepthStencilView() :
-#if 0
-		View(VK_NULL_HANDLE),
-#endif
-		Format(VK_FORMAT_UNDEFINED)
-	{
-	}
-
-	void Create(FVulkanDevice& Device, FVulkanTexture2D& Texture)
-	{
-		check(0);
-#if 0
-		VkDepthStencilViewCreateInfo ViewInfo;
-		ViewInfo.sType = VK_STRUCTURE_TYPE_DEPTH_STENCIL_VIEW_CREATE_INFO;
-		ViewInfo.pNext = nullptr;
-		ViewInfo.image = Texture.Surface.Image;
-		ViewInfo.mipLevel = 0;
-		ViewInfo.baseArraySlice = 0;
-		ViewInfo.arraySize = 1;
-		ViewInfo.flags = 0;
-		VERIFYVULKANRESULT(vkCreateDepthStencilView(Device.Device, &ViewInfo, &View));
-
-		Format = Texture.Surface.InternalFormat;
-#endif
-	}
-
-	void Destroy(FVulkanDevice& Device)
-	{
-		check(0);
-#if 0
-		if (View != VK_NULL_HANDLE)
-		{
-			vkDestroyObject(Device.Device, VK_OBJECT_TYPE_DEPTH_STENCIL_VIEW, View);
-			View = VK_NULL_HANDLE;
-		}
-#endif
-	}
-
-#if 0
-	VkDepthStencilView View;
-#endif
-	VkFormat Format;
-};
-
-struct FVulkanFramebuffer
+class FVulkanFramebuffer
 {
+public:
 	FVulkanFramebuffer(FVulkanDevice& Device, const FRHISetRenderTargetsInfo& RTInfo, const FVulkanRenderTargetLayout& RTLayout, const FVulkanRenderPass& RenderPass);
 
 	bool Matches(const FRHISetRenderTargetsInfo& RTInfo) const;
@@ -230,7 +185,11 @@ struct FVulkanFramebuffer
 	TArray<VkImageView> Attachments;
 	TArray<VkImageSubresourceRange> SubresourceRanges;
 
+#if VULKAN_USE_NEW_COMMAND_BUFFERS
+	void InsertWriteBarrier(FVulkanCmdBuffer* CmdBuffer);
+#else
 	void InsertWriteBarrier(VkCommandBuffer cmd);
+#endif
 
 private:
 	VkFramebuffer Framebuffer;
@@ -244,13 +203,14 @@ private:
 	TArray<VkImageMemoryBarrier> WriteBarriers;
 };
 
-struct FVulkanRenderPass
+class FVulkanRenderPass
 {
+public:
 	const FVulkanRenderTargetLayout& GetLayout() const { return Layout; }
 	VkRenderPass GetHandle() const { check(RenderPass != VK_NULL_HANDLE); return RenderPass; }
 
 private:
-	friend struct FVulkanPendingState;
+	friend class FVulkanPendingState;
 
 #if VULKAN_ENABLE_PIPELINE_CACHE
 	friend class FVulkanPipelineStateCache;
@@ -259,11 +219,14 @@ private:
 	FVulkanRenderPass(FVulkanDevice& Device, const FVulkanRenderTargetLayout& RTLayout);
 	~FVulkanRenderPass();
 
+#if VULKAN_USE_NEW_COMMAND_BUFFERS
+#else
 	void Begin(FVulkanCmdBuffer& CmdBuf, const VkFramebuffer& Framebuffer, const VkClearValue* AttachmentClearValues);
 	void End(FVulkanCmdBuffer& CmdBuf)
 	{
 		vkCmdEndRenderPass(CmdBuf.GetHandle());
 	}
+#endif
 
 private:
 	FVulkanRenderTargetLayout Layout;
@@ -271,8 +234,9 @@ private:
 	FVulkanDevice& Device;
 };
 
-struct FVulkanDescriptorSetsLayout
+class FVulkanDescriptorSetsLayout
 {
+public:
 	FVulkanDescriptorSetsLayout(FVulkanDevice* InDevice);
 	~FVulkanDescriptorSetsLayout();
 
@@ -355,7 +319,7 @@ struct FVulkanDescriptorSets
 
 private:
 	friend class FVulkanDescriptorPool;
-	friend struct FVulkanPendingState;
+	friend class FVulkanPendingState;
 
 	FVulkanDescriptorSets(FVulkanDevice* InDevice, const FVulkanBoundShaderState* InState, FVulkanDescriptorPool* InPool);
 	~FVulkanDescriptorSets();
@@ -366,10 +330,14 @@ private:
 	TArray<VkDescriptorSet> Sets;
 };
 
+void VulkanSetImageLayout(VkCommandBuffer CmdBuffer, VkImage Image, VkImageLayout OldLayout, VkImageLayout NewLayout, VkImageSubresourceRange SubresourceRange);
 
-void VulkanSetImageLayout(VkCommandBuffer CmdBuffer, VkImage image, VkImageAspectFlags aspectMask, VkImageLayout old_image_layout, VkImageLayout new_image_layout);
-
-void VulkanImageCopy(VkCommandBuffer Cmd, FTextureRHIParamRef SourceTextureRHI, FTextureRHIParamRef DestTextureRHI);
+// Transitions Color Images's first mip/layer/face
+inline void VulkanSetImageLayoutSimple(VkCommandBuffer CmdBuffer, VkImage Image, VkImageLayout OldLayout, VkImageLayout NewLayout, VkImageAspectFlags Aspect = VK_IMAGE_ASPECT_COLOR_BIT)
+{
+	VkImageSubresourceRange SubresourceRange = { Aspect, 0, 1, 0, 1 };
+	VulkanSetImageLayout(CmdBuffer, Image, OldLayout, NewLayout, SubresourceRange);
+}
 
 void VulkanResolveImage(VkCommandBuffer Cmd, FTextureRHIParamRef SourceTextureRHI, FTextureRHIParamRef DestTextureRHI);
 
@@ -456,6 +424,50 @@ public:
 	}
 };
 
+
+#if VULKAN_USE_NEW_RESOURCE_MANAGEMENT
+namespace VulkanRHI
+{
+	struct FStagingBuffer2 : public VulkanRHI::FDeviceChild, public VulkanRHI::FRefCount
+	{
+		VkBuffer Buffer;
+		VkMemoryRequirements MemoryRequirements;
+		TRefCountPtr<VulkanRHI::FResourceAllocation> ResourceAllocation;
+
+		FStagingBuffer2(FVulkanDevice* InDevice, uint32 Size)
+			: FDeviceChild(InDevice)
+		{
+			Buffer = FVulkanBuffer2::CreateVulkanBuffer(InDevice, Size, VK_BUFFER_USAGE_TRANSFER_SRC_BIT, MemoryRequirements);
+		}
+
+		uint32 GetAlignment() const
+		{
+			return (uint32)MemoryRequirements.alignment;
+		}
+
+		void BindAllocation(TRefCountPtr<VulkanRHI::FResourceAllocation> InResourceAllocation)
+		{
+			ResourceAllocation = InResourceAllocation;
+			check(Align(ResourceAllocation->GetOffset(), MemoryRequirements.alignment) == ResourceAllocation->GetOffset());
+			VERIFYVULKANRESULT(vkBindBufferMemory(Device->GetInstanceHandle(), Buffer, ResourceAllocation->GetHandle(), ResourceAllocation->GetOffset()));
+		}
+
+		~FStagingBuffer2()
+		{
+			vkDestroyBuffer(Device->GetInstanceHandle(), Buffer, nullptr);
+		}
+	};
+
+	struct FPendingBuffer2Lock
+	{
+		VulkanRHI::FResourceAllocation* ResourceAllocation;
+		TRefCountPtr<FStagingBuffer2> StagingBuffer;
+		uint32 Offset;
+		uint32 Size;
+	};
+}
+#endif
+
 #if VULKAN_HAS_DEBUGGING_ENABLED
 extern TAutoConsoleVariable<int32> GValidationCvar;
 #endif
diff --git a/Engine/Source/Runtime/VulkanRHI/Private/VulkanRenderTarget.cpp b/Engine/Source/Runtime/VulkanRHI/Private/VulkanRenderTarget.cpp
--- a/Engine/Source/Runtime/VulkanRHI/Private/VulkanRenderTarget.cpp
+++ b/Engine/Source/Runtime/VulkanRHI/Private/VulkanRenderTarget.cpp
@@ -7,6 +7,7 @@
 #include "VulkanRHIPrivate.h"
 #include "ScreenRendering.h"
 #include "VulkanPendingState.h"
+#include "VulkanContext.h"
 
 void FVulkanCommandListContext::RHISetRenderTargets(uint32 NumSimultaneousRenderTargets, const FRHIRenderTargetView* NewRenderTargets, const FRHIDepthRenderTargetView* NewDepthStencilTarget, uint32 NumUAVs, const FUnorderedAccessViewRHIParamRef* UAVs)
 {
@@ -36,26 +37,56 @@ void FVulkanCommandListContext::RHISetRenderTargetsAndClear(const FRHISetRenderT
 	FVulkanPendingState& PendingState = Device->GetPendingState();
 	PendingState.SetRenderTargetsInfo(RenderTargetsInfo);
 
-	if (!PendingState.GetIsCommandBufferEmpty())
+#if VULKAN_USE_NEW_COMMAND_BUFFERS
+	auto* CmdBuffer = CommandBufferManager->GetActiveCmdBuffer();
+	//#todo-rco: Write barrier always for now...
+	FVulkanFramebuffer* Framebuffer = PendingState.GetFrameBuffer();
+	if (Framebuffer)
 	{
-		FVulkanFramebuffer* Framebuffer = PendingState.GetFrameBuffer();
-		if (Framebuffer)
+		Framebuffer->InsertWriteBarrier(CmdBuffer);
+	}
+
+#if VULKAN_USE_NEW_RESOURCE_MANAGEMENT
+	if (CmdBuffer->IsInsideRenderPass())
+	{
+		if (
+			(RenderTargetsInfo.NumColorRenderTargets == 0 || (RenderTargetsInfo.NumColorRenderTargets == 1 && !RenderTargetsInfo.ColorRenderTarget[0].Texture)) &&
+			!RenderTargetsInfo.DepthStencilRenderTarget.Texture)
 		{
-			Framebuffer->InsertWriteBarrier(PendingState.GetCommandBuffer());
+			PendingState.RenderPassEnd(CmdBuffer);
+			PendingState.PrevRenderTargetsInfo = FRHISetRenderTargetsInfo();
 		}
 	}
+#endif
+#else
+#endif
 }
 
 void FVulkanCommandListContext::RHICopyToResolveTarget(FTextureRHIParamRef SourceTextureRHI, FTextureRHIParamRef DestTextureRHI, bool bKeepOriginalSurface, const FResolveParams& ResolveParams)
 {
-	// If we're using the pResoveSurface property of the subpass, so we don't need manual command buffer resolves and this function is a NoOp.
+	FVulkanPendingState& State = Device->GetPendingState();
+
+#if VULKAN_USE_NEW_COMMAND_BUFFERS
+	auto* CmdBuffer = CommandBufferManager->GetActiveCmdBuffer();
+	// Verify if we need to do some work (for the case of SetRT(), CopyToResolve() with no draw calls in between)
+	State.UpdateRenderPass(CmdBuffer);
+
+	const bool bRenderPassIsActive = State.IsRenderPassActive();
+
+	if (bRenderPassIsActive)
+	{
+		State.RenderPassEnd(CmdBuffer);
+	}
+
+	check(SourceTextureRHI->GetNumSamples() < 2);
+#else
+	// If we're using the pResolveAttachments property of the subpass, so we don't need manual command buffer resolves and this function is a NoOp.
 #if !VULKAN_USE_MSAA_RESOLVE_ATTACHMENTS
 	if (SourceTextureRHI->GetNumSamples() < 2)
 	{
 		return;
 	}
 
-	FVulkanPendingState& State = Device->GetPendingState();
 	const bool bRenderPassIsActive = State.IsRenderPassActive();
 
 	if (bRenderPassIsActive)
@@ -70,6 +101,7 @@ void FVulkanCommandListContext::RHICopyToResolveTarget(FTextureRHIParamRef Sourc
 		State.RenderPassBegin();
 	}
 #endif
+#endif
 }
 
 void FVulkanDynamicRHI::RHIReadSurfaceData(FTextureRHIParamRef TextureRHI, FIntRect Rect, TArray<FColor>& OutData, FReadSurfaceDataFlags InFlags)
@@ -105,3 +137,52 @@ void FVulkanDynamicRHI::RHIRead3DSurfaceFloatData(FTextureRHIParamRef TextureRHI
 {
 	VULKAN_SIGNAL_UNIMPLEMENTED();
 }
+
+
+void FVulkanCommandListContext::RHITransitionResources(EResourceTransitionAccess TransitionType, FTextureRHIParamRef* InTextures, int32 NumTextures)
+{
+	if (TransitionType == EResourceTransitionAccess::EReadable)
+	{
+		const FResolveParams ResolveParams;
+		for (int32 i = 0; i < NumTextures; ++i)
+		{
+			RHICopyToResolveTarget(InTextures[i], InTextures[i], true, ResolveParams);
+		}
+	}
+	else if (TransitionType == EResourceTransitionAccess::EWritable)
+	{
+		for (int32 i = 0; i < NumTextures; ++i)
+		{
+			FRHITexture* RHITexture = InTextures[i];
+			FRHITexture2D* RHITexture2D = RHITexture->GetTexture2D();
+			if (RHITexture2D)
+			{
+				FVulkanTexture2D* Texture = (FVulkanTexture2D*)RHITexture2D;
+				if (Texture->Surface.ImageLayout == VK_IMAGE_LAYOUT_PRESENT_SRC_KHR)
+				{
+#if VULKAN_USE_NEW_COMMAND_BUFFERS
+					check(0);
+#else
+					check(Texture->Surface.GetAspectMask() == VK_IMAGE_ASPECT_COLOR_BIT);
+					VulkanSetImageLayoutSimple(Device->GetPendingState().GetCommandBuffer(), Texture->Surface.Image, VK_IMAGE_LAYOUT_PRESENT_SRC_KHR, VK_IMAGE_LAYOUT_COLOR_ATTACHMENT_OPTIMAL);
+#endif
+				}
+			}
+			else
+			{
+				FRHITextureCube* RHITextureCube = RHITexture->GetTextureCube();
+				if (RHITextureCube)
+				{
+				}
+				else
+				{
+					ensure(0);
+				}
+			}
+		}
+	}
+	else
+	{
+		ensure(0);
+	}
+}
diff --git a/Engine/Source/Runtime/VulkanRHI/Private/VulkanShaders.cpp b/Engine/Source/Runtime/VulkanRHI/Private/VulkanShaders.cpp
--- a/Engine/Source/Runtime/VulkanRHI/Private/VulkanShaders.cpp
+++ b/Engine/Source/Runtime/VulkanRHI/Private/VulkanShaders.cpp
@@ -251,7 +251,12 @@ FVulkanBoundShaderState::FVulkanBoundShaderState(
 
 	check(Device);
 
+#if VULKAN_USE_NEW_COMMAND_BUFFERS
+	//#todo-rco: Fix this!
+	LastFrameRendered = GFrameNumberRenderThread;
+#else
 	LastFrameRendered = Device->GetPendingState().GetCurrentCommandBufferIndex();
+#endif
 
 	Layout = new FVulkanDescriptorSetsLayout(Device);
 
@@ -418,14 +423,18 @@ void FVulkanBoundShaderState::ResetState()
 
 	CurrDescriptorSets = nullptr;
 
+#if VULKAN_USE_NEW_COMMAND_BUFFERS
+	//#todo-rco: Fix this!
+	const uint32 CurrentImageIndex = GFrameNumberRenderThread % ARRAY_COUNT(CurrentDS);
+#else
 	// reset the DS list the first time it's used in a frame
 	const uint32 CurrentImageIndex = Device->GetPendingState().GetCurrentCommandBufferIndex();
+#endif
 	if (LastFrameRendered != CurrentImageIndex)
 	{
 		CurrentDS[CurrentImageIndex] = 0;
 		LastFrameRendered = CurrentImageIndex;
 	}
-
 	LastBoundPipeline = VK_NULL_HANDLE;
 	bDirtyVertexStreams = true;
 }
@@ -1051,8 +1060,13 @@ void FVulkanBoundShaderState::UpdateDescriptorSets(FVulkanGlobalUniformPool* Glo
 
 	// get the next available DS
 
+#if VULKAN_USE_NEW_COMMAND_BUFFERS
+	//#todo-rco: Fix this!
+	const uint32 CurrentImageIndex = GFrameNumberRenderThread % ARRAY_COUNT(CurrentDS);
+#else
 	// This is the image where we are going to execute command buffer into
 	const uint32 CurrentImageIndex = Device->GetPendingState().GetCurrentCommandBufferIndex();
+#endif
 	TArray<FVulkanDescriptorSets*>& CurrentDescriptorSet = DescriptorSets[CurrentImageIndex];
 	int32& CurrentDescriptorSetIndex = CurrentDS[CurrentImageIndex];
 
diff --git a/Engine/Source/Runtime/VulkanRHI/Private/VulkanSwapChain.cpp b/Engine/Source/Runtime/VulkanRHI/Private/VulkanSwapChain.cpp
--- a/Engine/Source/Runtime/VulkanRHI/Private/VulkanSwapChain.cpp
+++ b/Engine/Source/Runtime/VulkanRHI/Private/VulkanSwapChain.cpp
@@ -7,10 +7,14 @@
 #include "VulkanRHIPrivate.h"
 #include "VulkanSwapChain.h"
 
-FVulkanSwapChain::FVulkanSwapChain(VkInstance Instance, FVulkanDevice& Device, void* WindowHandle, EPixelFormat& InOutPixelFormat, uint32 Width, uint32 Height, 
+FVulkanSwapChain::FVulkanSwapChain(VkInstance Instance, FVulkanDevice& InDevice, void* WindowHandle, EPixelFormat& InOutPixelFormat, uint32 Width, uint32 Height, 
 	uint32* InOutDesiredNumBackBuffers, TArray<VkImage>& OutImages)
 	: SwapChain(VK_NULL_HANDLE)
+	, Device(InDevice)
 	, Surface(VK_NULL_HANDLE)
+#if VULKAN_USE_NEW_COMMAND_BUFFERS
+	, CurrentImageIndex(-1)
+#endif
 {
 #define FETCH_KHR_FPN(FunctionPointerName)\
 	{\
@@ -115,7 +119,7 @@ FVulkanSwapChain::FVulkanSwapChain(VkInstance Instance, FVulkanDevice& Device, v
 
 	//#todo-rco: Check multiple Gfx Queues?
 	VkBool32 bSupportsPresent = VK_FALSE;
-	VERIFYVULKANRESULT(vkGetPhysicalDeviceSurfaceSupportKHR(Device.GetPhysicalHandle(), Device.GetQueue()->GetNodeIndex(), Surface, &bSupportsPresent));
+	VERIFYVULKANRESULT(vkGetPhysicalDeviceSurfaceSupportKHR(Device.GetPhysicalHandle(), Device.GetQueue()->GetFamilyIndex(), Surface, &bSupportsPresent));
 	//#todo-rco: Find separate present queue if the gfx one doesn't support presents
 	check(bSupportsPresent);
 
@@ -154,25 +158,30 @@ FVulkanSwapChain::FVulkanSwapChain(VkInstance Instance, FVulkanDevice& Device, v
 	VERIFYVULKANRESULT_EXPANDED(vkGetPhysicalDeviceSurfaceCapabilitiesKHR(Device.GetPhysicalHandle(),
 		Surface,
 		&SurfProperties));
-
+	VkSurfaceTransformFlagBitsKHR PreTransform;
+	if (SurfProperties.supportedTransforms & VK_SURFACE_TRANSFORM_IDENTITY_BIT_KHR)
+	{
+		PreTransform = VK_SURFACE_TRANSFORM_IDENTITY_BIT_KHR;
+	}
+	else
+	{
+		PreTransform = SurfProperties.currentTransform;
+	}
 	uint32 DesiredNumBuffers = FMath::Clamp(*InOutDesiredNumBackBuffers, SurfProperties.minImageCount, SurfProperties.maxImageCount);
 	
 	VkSwapchainCreateInfoKHR SwapChainInfo;
 	FMemory::Memzero(SwapChainInfo);
 	SwapChainInfo.sType = VK_STRUCTURE_TYPE_SWAPCHAIN_CREATE_INFO_KHR;
-	SwapChainInfo.pNext = nullptr;
 	SwapChainInfo.surface = Surface;
 	SwapChainInfo.minImageCount = DesiredNumBuffers;
 	SwapChainInfo.imageFormat = CurrFormat.format;
 	SwapChainInfo.imageColorSpace = CurrFormat.colorSpace;
 	SwapChainInfo.imageExtent.width = PLATFORM_ANDROID ? Width : (SurfProperties.currentExtent.width == -1 ? Width : SurfProperties.currentExtent.width);
 	SwapChainInfo.imageExtent.height = PLATFORM_ANDROID ? Height : (SurfProperties.currentExtent.height == -1 ? Height : SurfProperties.currentExtent.height);
-	SwapChainInfo.imageUsage |= VK_IMAGE_USAGE_COLOR_ATTACHMENT_BIT;
-	SwapChainInfo.preTransform = SurfProperties.currentTransform;
+	SwapChainInfo.imageUsage = VK_IMAGE_USAGE_COLOR_ATTACHMENT_BIT;
+	SwapChainInfo.preTransform = PreTransform;
 	SwapChainInfo.imageArrayLayers = 1;
 	SwapChainInfo.imageSharingMode = VK_SHARING_MODE_EXCLUSIVE;
-	SwapChainInfo.queueFamilyIndexCount = 0;
-	SwapChainInfo.pQueueFamilyIndices = NULL;
 	SwapChainInfo.presentMode = PresentMode;
 	SwapChainInfo.oldSwapchain = VK_NULL_HANDLE;
 	SwapChainInfo.clipped = VK_TRUE;
@@ -187,10 +196,73 @@ FVulkanSwapChain::FVulkanSwapChain(VkInstance Instance, FVulkanDevice& Device, v
 
 	OutImages.AddUninitialized(NumSwapChainImages);
 	VERIFYVULKANRESULT_EXPANDED(GetSwapchainImagesKHR(Device.GetInstanceHandle(), SwapChain, &NumSwapChainImages, OutImages.GetData()));
+
+#if VULKAN_USE_NEW_COMMAND_BUFFERS
+	ImageAcquiredSemaphore.AddUninitialized(DesiredNumBuffers);
+	RenderingCompletedSemaphore.AddUninitialized(DesiredNumBuffers);
+	for (uint32 BufferIndex = 0; BufferIndex < DesiredNumBuffers; ++BufferIndex)
+	{
+		ImageAcquiredSemaphore[BufferIndex] = new FVulkanSemaphore(Device);
+		RenderingCompletedSemaphore[BufferIndex] = new FVulkanSemaphore(Device);
+	}
+#endif
 }
 
 void FVulkanSwapChain::Destroy(FVulkanDevice& Device)
 {
+#if VULKAN_USE_NEW_COMMAND_BUFFERS
+	for (int BufferIndex = 0; BufferIndex < ImageAcquiredSemaphore.Num(); ++BufferIndex)
+	{
+		delete ImageAcquiredSemaphore[BufferIndex];
+		delete RenderingCompletedSemaphore[BufferIndex];
+	}
+	ImageAcquiredSemaphore.Empty(0);
+	RenderingCompletedSemaphore.Empty(0);
+#endif
 	vkDestroySwapchainKHR(Device.GetInstanceHandle(), SwapChain, nullptr);
 	SwapChain = VK_NULL_HANDLE;
 }
+
+#if VULKAN_USE_NEW_COMMAND_BUFFERS
+int32 FVulkanSwapChain::AcquireImageIndex()
+{
+	// Get the index of the next swapchain image we should render to.
+	// We'll wait with an "infinite" timeout, the function will block until an image is ready.
+	// The ImageAcquiredSemaphore[ImageAcquiredSemaphoreIndex] will get signaled when the image is ready (upon function return).
+	// The Fences[CurrentFenceIndex] will also get signaled when the image is ready (upon function return).
+	// Note: Queues can still be filled in on the CPU side, but won't execute until the semaphore is signaled.
+	//CurrentImageIndex = -1;
+	uint32 ImageIndex = 0;
+	VkResult Result = AcquireNextImageKHR(
+		Device.GetInstanceHandle(),
+		SwapChain,
+		UINT64_MAX,
+		VK_NULL_HANDLE, //ImageAcquiredSemaphore[CurrentFenceIndex]->GetHandle(),
+		VK_NULL_HANDLE,
+		&ImageIndex);
+	checkf(Result == VK_SUCCESS || Result == VK_SUBOPTIMAL_KHR, TEXT("AcquireNextImageKHR failed Result = %d"), int32(Result));
+	CurrentImageIndex = (int32)ImageIndex;
+	return CurrentImageIndex;
+}
+
+void FVulkanSwapChain::Present(FVulkanQueue* Queue)
+{
+	check(CurrentImageIndex != -1);
+
+	VkPresentInfoKHR Info;
+	FMemory::Memzero(Info);
+	Info.sType = VK_STRUCTURE_TYPE_PRESENT_INFO_KHR;
+/*
+	Info.waitSemaphoreCount = 1;	// 0;	// 1;
+	VkSemaphore Semaphore = RenderingCompletedSemaphore[CurrentImageIndex]->GetHandle();
+	Info.pWaitSemaphores = &Semaphore;
+*/
+
+	Info.swapchainCount = 1;
+	Info.pSwapchains = &SwapChain;
+	Info.pImageIndices = (uint32*)&CurrentImageIndex;
+
+	VERIFYVULKANRESULT(QueuePresentKHR(Queue->GetHandle(), &Info));
+}
+#else
+#endif
diff --git a/Engine/Source/Runtime/VulkanRHI/Private/VulkanSwapChain.h b/Engine/Source/Runtime/VulkanRHI/Private/VulkanSwapChain.h
--- a/Engine/Source/Runtime/VulkanRHI/Private/VulkanSwapChain.h
+++ b/Engine/Source/Runtime/VulkanRHI/Private/VulkanSwapChain.h
@@ -6,19 +6,32 @@
 
 #pragma once
 
-struct FVulkanSwapChain
+class FVulkanQueue;
+
+class FVulkanSwapChain
 {
+public:
+#if VULKAN_USE_NEW_COMMAND_BUFFERS
+#else
 #if !PLATFORM_WINDOWS
 	//@HACK : maybe NUM_BUFFERS at VulkanViewport should be moved to here?
 	enum { NUM_BUFFERS = 2 };
+#endif
 #endif
 
-	FVulkanSwapChain(VkInstance Instance, FVulkanDevice& Device, void* WindowHandle, EPixelFormat& InOutPixelFormat, uint32 Width, uint32 Height,
+	FVulkanSwapChain(VkInstance Instance, FVulkanDevice& InDevice, void* WindowHandle, EPixelFormat& InOutPixelFormat, uint32 Width, uint32 Height,
 		uint32* InOutDesiredNumBackBuffers, TArray<VkImage>& OutImages);
 
 	void Destroy(FVulkanDevice& Device);
 
+#if VULKAN_USE_NEW_COMMAND_BUFFERS
+	void Present(FVulkanQueue* Queue);
+#else
+#endif
+
+protected:
 	VkSwapchainKHR SwapChain;
+	FVulkanDevice& Device;
 
 	PFN_vkCreateSwapchainKHR CreateSwapchainKHR;
 	PFN_vkDestroySwapchainKHR DestroySwapchainKHR;
@@ -27,4 +40,18 @@ struct FVulkanSwapChain
 	PFN_vkAcquireNextImageKHR AcquireNextImageKHR;
 
 	VkSurfaceKHR Surface;
-};
\ No newline at end of file
+
+#if VULKAN_USE_NEW_COMMAND_BUFFERS
+	int32 CurrentImageIndex;
+	TArray<FVulkanSemaphore*> ImageAcquiredSemaphore;
+	TArray<FVulkanSemaphore*> RenderingCompletedSemaphore;
+#else
+#endif
+
+#if VULKAN_USE_NEW_COMMAND_BUFFERS
+	int32 AcquireImageIndex();
+	friend class FVulkanViewport;
+#else
+#endif
+	friend class FVulkanQueue;
+};
diff --git a/Engine/Source/Runtime/VulkanRHI/Private/VulkanTexture.cpp b/Engine/Source/Runtime/VulkanRHI/Private/VulkanTexture.cpp
--- a/Engine/Source/Runtime/VulkanRHI/Private/VulkanTexture.cpp
+++ b/Engine/Source/Runtime/VulkanRHI/Private/VulkanTexture.cpp
@@ -8,8 +8,11 @@
 #include "VulkanMemory.h"
 #include "VulkanContext.h"
 #include "VulkanPendingState.h"
+
+
 static FCriticalSection GTextureMapLock;
-static TMap<FRHIResource*, VulkanRHI::FStagingImage*> GTextureMapEntries;
+static TMap<FRHIResource*, VulkanRHI::FStagingBuffer*> GPendingLockedBuffers;
+static TMap<FRHIResource*, VulkanRHI::FStagingImage*> GPendingLockedImages;
 
 static const VkImageTiling GVulkanViewTypeTilingMode[VK_IMAGE_VIEW_TYPE_RANGE_SIZE] =
 {
@@ -22,6 +25,7 @@ static const VkImageTiling GVulkanViewTypeTilingMode[VK_IMAGE_VIEW_TYPE_RANGE_SI
     VK_IMAGE_TILING_LINEAR,		// VK_IMAGE_VIEW_TYPE_CUBE_ARRAY
 };
 
+
 VkImage FVulkanSurface::CreateImage(
 	FVulkanDevice& InDevice,
 	VkImageViewType ResourceType,
@@ -254,11 +258,7 @@ FVulkanSurface::FVulkanSurface(FVulkanDevice& InDevice, VkImageViewType Resource
 	, Tiling(VK_IMAGE_TILING_MAX_ENUM)	// Can be expanded to a per-platform definition
 	, ViewType(ResourceType)
 	, bIsImageOwner(true)
-#if VULKAN_USE_MEMORY_SYSTEM
 	, Allocation(nullptr)
-#else
-	, DeviceMemory(VK_NULL_HANDLE)
-#endif
 	, NumMips(InNumMips)
 	, NumSamples(InNumSamples)
 	, AspectMask(0)
@@ -281,34 +281,8 @@ FVulkanSurface::FVulkanSurface(FVulkanDevice& InDevice, VkImageViewType Resource
 		MemProps |= VK_MEMORY_PROPERTY_HOST_VISIBLE_BIT;
 	}
 
-#if VULKAN_USE_MEMORY_SYSTEM
 	Allocation = InDevice.GetMemoryManager().Alloc(MemoryRequirements.size, MemoryRequirements.memoryTypeBits, MemProps, __FILE__, __LINE__);
 	VERIFYVULKANRESULT(vkBindImageMemory(Device->GetInstanceHandle(), Image, Allocation->GetHandle(), 0));
-#else
-	VkMemoryAllocateInfo AllocInfo;
-	FMemory::Memzero(AllocInfo);
-	AllocInfo.sType = VK_STRUCTURE_TYPE_MEMORY_ALLOCATE_INFO;
-	AllocInfo.pNext = nullptr;
-	AllocInfo.allocationSize = MemoryRequirements.size;
-	AllocInfo.memoryTypeIndex = 0;
-	VkResult Result = Device->GetMemoryManager().GetMemoryTypeFromProperties(MemoryRequirements.memoryTypeBits, MemProps, &AllocInfo.memoryTypeIndex);
-	if (Result != VK_SUCCESS)
-	{
-		// We are currently only hitting this case when trying to alloc CPU writable 2d texture arrays
-		if ((MemProps & VK_MEMORY_PROPERTY_HOST_VISIBLE_BIT) && CreateInfo.BulkData)
-		{
-			// Try again without CPU access, let's assume that since we have BulkData we can fill it in via DMA
-			MemProps &= ~VK_MEMORY_PROPERTY_HOST_VISIBLE_BIT;
-			Result = Device->GetMemoryManager().GetMemoryTypeFromProperties(MemoryRequirements.memoryTypeBits, MemProps, &AllocInfo.memoryTypeIndex);
-		}
-	}
-	VERIFYVULKANRESULT_EXPANDED(Result);
-
-	// Allocate/Free memory should have a wrapper function, which tracks "GVulkanBufferByteBalance"
-	VERIFYVULKANRESULT_EXPANDED(vkAllocateMemory(Device->GetInstanceHandle(), &AllocInfo, nullptr, &DeviceMemory));
-	VERIFYVULKANRESULT(vkBindImageMemory(Device->GetInstanceHandle(), Image, DeviceMemory, 0));
-#endif
-
 
 	Tiling = ImageCreateInfo.tiling;
 	check(Tiling == VK_IMAGE_TILING_LINEAR || Tiling == VK_IMAGE_TILING_OPTIMAL);
@@ -334,7 +308,12 @@ FVulkanSurface::FVulkanSurface(FVulkanDevice& InDevice, VkImageViewType Resource
 
 	if ((ImageCreateInfo.usage & VK_IMAGE_USAGE_SAMPLED_BIT) && (UEFlags & (TexCreate_RenderTargetable | TexCreate_DepthStencilTargetable)))
 	{
-		ClearBlocking(CreateInfo.ClearValueBinding);
+		const bool bTransitionToPresentable = ((UEFlags & TexCreate_Presentable) == TexCreate_Presentable);
+#if VULKAN_USE_NEW_COMMAND_BUFFERS
+		Clear(CreateInfo.ClearValueBinding, bTransitionToPresentable);
+#else
+		ClearBlocking(CreateInfo.ClearValueBinding, bTransitionToPresentable);
+#endif
 	}
 }
 
@@ -355,11 +334,7 @@ FVulkanSurface::FVulkanSurface(FVulkanDevice& InDevice, VkImageViewType Resource
 	, Tiling(VK_IMAGE_TILING_MAX_ENUM)	// Can be expanded to a per-platform definition
 	, ViewType(ResourceType)
 	, bIsImageOwner(false)
-#if VULKAN_USE_MEMORY_SYSTEM
 	, Allocation(nullptr)
-#else
-	, DeviceMemory(VK_NULL_HANDLE)
-#endif
 	, NumMips(0)
 	, NumSamples(1)	// This is defined by the MSAA setting...
 	, AspectMask(0)
@@ -377,7 +352,12 @@ FVulkanSurface::FVulkanSurface(FVulkanDevice& InDevice, VkImageViewType Resource
 	{
 		if (UEFlags & (TexCreate_RenderTargetable | TexCreate_DepthStencilTargetable))
 		{
-			ClearBlocking(CreateInfo.ClearValueBinding);
+			const bool bTransitionToPresentable = ((UEFlags & TexCreate_Presentable) == TexCreate_Presentable);
+#if VULKAN_USE_NEW_COMMAND_BUFFERS
+			Clear(CreateInfo.ClearValueBinding, bTransitionToPresentable);
+#else
+			ClearBlocking(CreateInfo.ClearValueBinding, bTransitionToPresentable);
+#endif
 		}
 	}
 }
@@ -412,26 +392,13 @@ void FVulkanSurface::Destroy()
 			ImageLayout = VK_IMAGE_LAYOUT_UNDEFINED;
 		}
 
-#if VULKAN_USE_MEMORY_SYSTEM
 		Device->GetMemoryManager().Free(Allocation);
 		Allocation = nullptr;
-#else
-		if (DeviceMemory != VK_NULL_HANDLE)
-		{
-			vkFreeMemory(Device->GetInstanceHandle(), DeviceMemory, nullptr);
-			DeviceMemory = VK_NULL_HANDLE;
-		}		
-#endif
 	}
 }
 
 void* FVulkanSurface::Lock(uint32 MipIndex, uint32 ArrayIndex, EResourceLockMode LockMode, uint32& DestStride)
 {
-#if VULKAN_USE_MEMORY_SYSTEM
-#else
-	checkf(DeviceMemory != VK_NULL_HANDLE, TEXT("Attempting to map unallocated image-memory"));
-#endif
-
 	DestStride = 0;
 
 	check((MemProps & VK_MEMORY_PROPERTY_HOST_VISIBLE_BIT) == 0 ? Tiling == VK_IMAGE_TILING_OPTIMAL : Tiling == VK_IMAGE_TILING_LINEAR);
@@ -462,11 +429,7 @@ void* FVulkanSurface::Lock(uint32 MipIndex, uint32 ArrayIndex, EResourceLockMode
 		check(DestStride == SubResourceLayout.rowPitch);
 
 		// Map buffer to a pointer
-#if VULKAN_USE_MEMORY_SYSTEM
 		Data = Allocation->Map(SubResourceLayout.size, SubResourceLayout.offset);
-#else
-		VERIFYVULKANRESULT(vkMapMemory(Device->GetInstanceHandle(), DeviceMemory, SubResourceLayout.offset, SubResourceLayout.size, 0, &Data));
-#endif
 		return Data;
 	}
 
@@ -497,11 +460,6 @@ void* FVulkanSurface::Lock(uint32 MipIndex, uint32 ArrayIndex, EResourceLockMode
 
 void FVulkanSurface::Unlock(uint32 MipIndex, uint32 ArrayIndex)
 {
-#if VULKAN_USE_MEMORY_SYSTEM
-#else
-	checkf(DeviceMemory != VK_NULL_HANDLE, TEXT("Attempting to unmap unallocated image-memory"));
-#endif
-
 	check((MemProps & VK_MEMORY_PROPERTY_HOST_VISIBLE_BIT) == 0 ? Tiling == VK_IMAGE_TILING_OPTIMAL : Tiling == VK_IMAGE_TILING_LINEAR);
 
 	if(Tiling == VK_IMAGE_TILING_LINEAR)
@@ -509,11 +467,7 @@ void FVulkanSurface::Unlock(uint32 MipIndex, uint32 ArrayIndex)
 		auto& Data = MipMapMapping.FindOrAdd(MipIndex);
 		checkf(Data != nullptr, TEXT("The buffer needs to be mapped, before it can be unmapped"));
 
-#if VULKAN_USE_MEMORY_SYSTEM
 		Allocation->Unmap();
-#else
-		vkUnmapMemory(Device->GetInstanceHandle(), DeviceMemory);
-#endif
 		Data = nullptr;
 		return;
 	}
@@ -638,48 +592,106 @@ VkImageAspectFlags FVulkanSurface::GetAspectMask() const
 	return AspectMask;
 }
 
-// Fixes the issue:
-//	"vkCmdBeginRenderPass(): Cannot read invalid swapchain image 1243cdf0, please fill the memory before using."
-void FVulkanSurface::ClearBlocking(const FClearValueBinding& ClearValueBinding)
+#if VULKAN_USE_NEW_COMMAND_BUFFERS
+void FVulkanSurface::Clear(const FClearValueBinding& ClearValueBinding, bool bTransitionToPresentable)
+{
+#if VULKAN_CLEAR_SURFACE_ON_CREATE
+	VkImageSubresourceRange Range;
+	FMemory::Memzero(Range);
+	Range.aspectMask = GetAspectMask();
+	Range.baseMipLevel = 0;
+	Range.levelCount = NumMips;
+	Range.baseArrayLayer = 0;
+	Range.layerCount = ViewType == VK_IMAGE_VIEW_TYPE_CUBE ? 6 : 1;
+
+#if VULKAN_USE_NEW_COMMAND_BUFFERS
+	FVulkanCmdBuffer* CmdBuffer = Device->GetImmediateContext().GetCommandBufferManager()->GetActiveCmdBuffer();
+	ensure(CmdBuffer->IsOutsideRenderPass());
+#else
+	//#todo-rco: FIX ME!
+	FVulkanCmdBuffer* CmdBuffer = Device->GetImmediateContext().GetCommandBufferManager()->Create();
+	CmdBuffer->Begin();
+#endif
+	if(Range.aspectMask == VK_IMAGE_ASPECT_COLOR_BIT)
+	{
+		VkClearColorValue Color;
+		FMemory::Memzero(Color);
+		Color.float32[0] = ClearValueBinding.Value.Color[0];
+		Color.float32[1] = ClearValueBinding.Value.Color[1];
+		Color.float32[2] = ClearValueBinding.Value.Color[2];
+		Color.float32[3] = ClearValueBinding.Value.Color[3];
+		VulkanSetImageLayoutSimple(CmdBuffer->GetHandle(), Image, ImageLayout, VK_IMAGE_LAYOUT_GENERAL);
+		vkCmdClearColorImage(CmdBuffer->GetHandle(), Image, VK_IMAGE_LAYOUT_GENERAL, &Color, 1, &Range);
+		VulkanSetImageLayoutSimple(CmdBuffer->GetHandle(), Image, VK_IMAGE_LAYOUT_GENERAL, bTransitionToPresentable ? VK_IMAGE_LAYOUT_PRESENT_SRC_KHR : VK_IMAGE_LAYOUT_COLOR_ATTACHMENT_OPTIMAL);
+		ImageLayout = VK_IMAGE_LAYOUT_COLOR_ATTACHMENT_OPTIMAL;
+	}
+	else
+	{
+		check(Range.aspectMask & (VK_IMAGE_ASPECT_DEPTH_BIT | VK_IMAGE_ASPECT_STENCIL_BIT));
+		check(!bTransitionToPresentable);
+		VkClearDepthStencilValue Value;
+		FMemory::Memzero(Value);
+		Value.depth = ClearValueBinding.Value.DSValue.Depth;
+		Value.stencil = ClearValueBinding.Value.DSValue.Stencil;
+		VulkanSetImageLayoutSimple(CmdBuffer->GetHandle(), Image, ImageLayout, VK_IMAGE_LAYOUT_GENERAL, Range.aspectMask);
+		vkCmdClearDepthStencilImage(CmdBuffer->GetHandle(), Image, VK_IMAGE_LAYOUT_GENERAL, &Value, 1, &Range);
+		VulkanSetImageLayoutSimple(CmdBuffer->GetHandle(), Image, VK_IMAGE_LAYOUT_GENERAL, VK_IMAGE_LAYOUT_DEPTH_STENCIL_ATTACHMENT_OPTIMAL, Range.aspectMask);
+		ImageLayout = VK_IMAGE_LAYOUT_DEPTH_STENCIL_ATTACHMENT_OPTIMAL;
+	}
+
+#if VULKAN_USE_NEW_COMMAND_BUFFERS
+#else
+	Device->EndCommandBufferBlock(CmdBuffer);
+#endif
+#endif
+}
+#else
+void FVulkanSurface::ClearBlocking(const FClearValueBinding& ClearValueBinding, bool bTransitionToPresentable)
 {
 #if VULKAN_CLEAR_SURFACE_ON_CREATE
-	VkImageSubresourceRange range;
-	FMemory::Memzero(range);
-	range.aspectMask = GetAspectMask();
-	range.baseMipLevel = 0;
-	range.levelCount = NumMips;
-	range.baseArrayLayer = 0;
-	range.layerCount = ViewType == VK_IMAGE_VIEW_TYPE_CUBE ? 6 : 1;
+	VkImageSubresourceRange Range;
+	FMemory::Memzero(Range);
+	Range.aspectMask = GetAspectMask();
+	Range.baseMipLevel = 0;
+	Range.levelCount = NumMips;
+	Range.baseArrayLayer = 0;
+	Range.layerCount = ViewType == VK_IMAGE_VIEW_TYPE_CUBE ? 6 : 1;
 
 	//#todo-rco: FIX ME!
-	FVulkanCmdBuffer* cmd = Device->GetImmediateContext().GetCommandBufferManager()->Create();
-	cmd->Begin();
+	FVulkanCmdBuffer* Cmd = Device->GetImmediateContext().GetCommandBufferManager()->Create();
+	Cmd->Begin();
 	
-	if(range.aspectMask == VK_IMAGE_ASPECT_COLOR_BIT)
+	if(Range.aspectMask == VK_IMAGE_ASPECT_COLOR_BIT)
 	{
-		VkClearColorValue color;
-		FMemory::Memzero(color);
-		color.float32[0] = ClearValueBinding.Value.Color[0];
-		color.float32[1] = ClearValueBinding.Value.Color[1];
-		color.float32[2] = ClearValueBinding.Value.Color[2];
-		color.float32[3] = ClearValueBinding.Value.Color[3];
-		VulkanSetImageLayout(cmd->GetHandle(), Image, range.aspectMask, VK_IMAGE_LAYOUT_UNDEFINED, VK_IMAGE_LAYOUT_COLOR_ATTACHMENT_OPTIMAL);
-		vkCmdClearColorImage(cmd->GetHandle(), Image, VK_IMAGE_LAYOUT_GENERAL, &color, 1, &range);
+		VkClearColorValue Color;
+		FMemory::Memzero(Color);
+		Color.float32[0] = ClearValueBinding.Value.Color[0];
+		Color.float32[1] = ClearValueBinding.Value.Color[1];
+		Color.float32[2] = ClearValueBinding.Value.Color[2];
+		Color.float32[3] = ClearValueBinding.Value.Color[3];
+		VulkanSetImageLayoutSimple(Cmd->GetHandle(), Image, ImageLayout, VK_IMAGE_LAYOUT_GENERAL);
+		vkCmdClearColorImage(Cmd->GetHandle(), Image, VK_IMAGE_LAYOUT_GENERAL, &Color, 1, &Range);
+		ImageLayout = bTransitionToPresentable ? VK_IMAGE_LAYOUT_PRESENT_SRC_KHR : VK_IMAGE_LAYOUT_COLOR_ATTACHMENT_OPTIMAL;
+		VulkanSetImageLayoutSimple(Cmd->GetHandle(), Image, VK_IMAGE_LAYOUT_GENERAL, ImageLayout);
 	}
 	else
 	{
-		check(range.aspectMask & (VK_IMAGE_ASPECT_DEPTH_BIT | VK_IMAGE_ASPECT_STENCIL_BIT));
-		VkClearDepthStencilValue value;
-		FMemory::Memzero(value);
-		value.depth = ClearValueBinding.Value.DSValue.Depth;
-		value.stencil = ClearValueBinding.Value.DSValue.Stencil;
-		VulkanSetImageLayout(cmd->GetHandle(), Image, range.aspectMask, VK_IMAGE_LAYOUT_UNDEFINED, VK_IMAGE_LAYOUT_DEPTH_STENCIL_ATTACHMENT_OPTIMAL);
-		vkCmdClearDepthStencilImage(cmd->GetHandle(), Image, VK_IMAGE_LAYOUT_DEPTH_STENCIL_ATTACHMENT_OPTIMAL, &value, 1, &range);
+		check(Range.aspectMask & (VK_IMAGE_ASPECT_DEPTH_BIT | VK_IMAGE_ASPECT_STENCIL_BIT));
+		check(!bTransitionToPresentable);
+		VkClearDepthStencilValue Value;
+		FMemory::Memzero(Value);
+		Value.depth = ClearValueBinding.Value.DSValue.Depth;
+		Value.stencil = ClearValueBinding.Value.DSValue.Stencil;
+		VulkanSetImageLayoutSimple(Cmd->GetHandle(), Image, ImageLayout, VK_IMAGE_LAYOUT_GENERAL, Range.aspectMask);
+		vkCmdClearDepthStencilImage(Cmd->GetHandle(), Image, VK_IMAGE_LAYOUT_DEPTH_STENCIL_ATTACHMENT_OPTIMAL, &Value, 1, &Range);
+		VulkanSetImageLayoutSimple(Cmd->GetHandle(), Image, VK_IMAGE_LAYOUT_GENERAL, VK_IMAGE_LAYOUT_DEPTH_STENCIL_ATTACHMENT_OPTIMAL, Range.aspectMask);
+		ImageLayout = VK_IMAGE_LAYOUT_DEPTH_STENCIL_ATTACHMENT_OPTIMAL;
 	}
 
-	Device->EndCommandBufferBlock(cmd);
+	Device->EndCommandBufferBlock(Cmd);
 #endif
 }
+#endif
 
 
 /*-----------------------------------------------------------------------------
@@ -778,86 +790,21 @@ void* FVulkanDynamicRHI::RHILockTexture2D(FTexture2DRHIParamRef TextureRHI,uint3
 	FVulkanTexture2D* Texture = ResourceCast(TextureRHI);
 	check(Texture);
 
-	EPixelFormat Format = Texture->Surface.Format;
-	uint32 MipWidth = FMath::Max<uint32>(Texture->Surface.Width >> MipIndex, GPixelFormats[Format].BlockSizeX);
-	uint32 MipHeight = FMath::Max<uint32>(Texture->Surface.Height >> MipIndex, GPixelFormats[Format].BlockSizeY);
-#if VULKAN_USE_NEW_STAGING_BUFFERS
-	VulkanRHI::FStagingImage** StagingImage = nullptr;
+	VulkanRHI::FStagingBuffer** StagingBuffer = nullptr;
 	{
 		FScopeLock Lock(&GTextureMapLock);
-		StagingImage = &GTextureMapEntries.FindOrAdd(TextureRHI);
-		checkf(!*StagingImage, TEXT("Can't lock the same texture twice!"));
+		StagingBuffer = &GPendingLockedBuffers.FindOrAdd(TextureRHI);
+		checkf(!*StagingBuffer, TEXT("Can't lock the same texture twice!"));
 	}
 
-	*StagingImage = Device->GetStagingManager().AcquireImage(Texture->Surface.InternalFormat, MipWidth, MipHeight);
-
-	VkImageSubresource Subresource;
-	FMemory::Memzero(Subresource);
-	Subresource.aspectMask = Texture->Surface.GetAspectMask();
-	Subresource.mipLevel = 0;
-	Subresource.arrayLayer = 0;
-	VkSubresourceLayout SubresourceLayout;
-	vkGetImageSubresourceLayout(Device->GetInstanceHandle(), (*StagingImage)->GetHandle(), &Subresource, &SubresourceLayout);
-	DestStride = SubresourceLayout.rowPitch;
-
-	return (*StagingImage)->Map();
-#else
-	FVulkanDevice* VulkanDevice = Texture->Surface.Device;
-	VkDevice LogicalDevice = VulkanDevice->GetInstanceHandle();
-
-	VkExtent3D StagingExtent;
-	FMemory::Memzero(StagingExtent);
-	StagingExtent.width = MipWidth;
-	StagingExtent.height = MipHeight;
-	StagingExtent.depth = 1;
-
-	VkImageCreateInfo StagingImageCreateInfo;
-	FMemory::Memzero(StagingImageCreateInfo);
-	StagingImageCreateInfo.sType = VK_STRUCTURE_TYPE_IMAGE_CREATE_INFO;
-	StagingImageCreateInfo.pNext = nullptr;
-	StagingImageCreateInfo.imageType = VK_IMAGE_TYPE_2D;
-	StagingImageCreateInfo.format = Texture->Surface.InternalFormat;
-	StagingImageCreateInfo.extent = StagingExtent;
-	StagingImageCreateInfo.mipLevels = 1;
-	StagingImageCreateInfo.arrayLayers = 1;
-	StagingImageCreateInfo.samples = VK_SAMPLE_COUNT_1_BIT;
-	StagingImageCreateInfo.tiling = VK_IMAGE_TILING_LINEAR;
-	StagingImageCreateInfo.usage = VK_IMAGE_USAGE_TRANSFER_SRC_BIT;
-
-	VERIFYVULKANRESULT(vkCreateImage(LogicalDevice, &StagingImageCreateInfo, nullptr, &Texture->StagingImage));
-
-	VkMemoryRequirements MemReqs;
-	vkGetImageMemoryRequirements(LogicalDevice, Texture->StagingImage, &MemReqs);
-
-	VkMemoryAllocateInfo AllocInfo;
-	FMemory::Memzero(AllocInfo);
-	AllocInfo.sType = VK_STRUCTURE_TYPE_MEMORY_ALLOCATE_INFO;
-	AllocInfo.allocationSize = MemReqs.size;
-
-	VkMemoryPropertyFlags MemoryTypeFlags = VK_MEMORY_PROPERTY_HOST_VISIBLE_BIT;
-	VERIFYVULKANRESULT(VulkanDevice->GetMemoryManager().GetMemoryTypeFromProperties(MemReqs.memoryTypeBits, MemoryTypeFlags, &AllocInfo.memoryTypeIndex));
-	#if VULKAN_USE_MEMORY_SYSTEM
-		check(0);
-	#else
-		VERIFYVULKANRESULT(vkAllocateMemory(LogicalDevice, &AllocInfo, nullptr, &Texture->PixelBufferMemory));
-	#endif
-	VERIFYVULKANRESULT(vkBindImageMemory(LogicalDevice, Texture->StagingImage, Texture->PixelBufferMemory, 0));
-
-	VkImageSubresource Subresource;
-	FMemory::Memzero(Subresource);
-	Subresource.aspectMask = Texture->Surface.GetAspectMask();
-	Subresource.mipLevel = 0;
-	Subresource.arrayLayer = 0;
-
-	VkSubresourceLayout SubresourceLayout;
-	vkGetImageSubresourceLayout(LogicalDevice, Texture->StagingImage, &Subresource, &SubresourceLayout);
-	DestStride = SubresourceLayout.rowPitch;
-
-	void* MappedData = nullptr;
-	VERIFYVULKANRESULT(vkMapMemory(LogicalDevice, Texture->PixelBufferMemory, 0, AllocInfo.allocationSize, 0, &MappedData));
+	uint32 BufferSize = 0;
+	DestStride = 0;
+	Texture->Surface.GetMipSize(MipIndex, BufferSize);
+	Texture->Surface.GetMipStride(MipIndex, DestStride);
+	*StagingBuffer = Device->GetStagingManager().AcquireBuffer(BufferSize);
 
-	return MappedData;
-#endif
+	void* Data = (*StagingBuffer)->Map();
+	return Data;
 }
 
 void FVulkanDynamicRHI::RHIUnlockTexture2D(FTexture2DRHIParamRef TextureRHI,uint32 MipIndex,bool bLockWithinMiptail)
@@ -867,114 +814,81 @@ void FVulkanDynamicRHI::RHIUnlockTexture2D(FTexture2DRHIParamRef TextureRHI,uint
 
 	VkDevice LogicalDevice = Device->GetInstanceHandle();
 
-#if VULKAN_USE_FENCE_MANAGER
-	auto* Fence = Device->GetFenceManager().AllocateFence();
-#else
-	VkFence Fence;
-	VkFenceCreateInfo FenceCreateInfo;
-	FMemory::Memzero(FenceCreateInfo);
-	FenceCreateInfo.sType = VK_STRUCTURE_TYPE_FENCE_CREATE_INFO;
-	VERIFYVULKANRESULT(vkCreateFence(LogicalDevice, &FenceCreateInfo, nullptr, &Fence));
-#endif
-
-#if VULKAN_USE_NEW_STAGING_BUFFERS
-	VulkanRHI::FStagingImage* StagingImage = nullptr;
+	VulkanRHI::FStagingBuffer* StagingBuffer = nullptr;
 	{
 		FScopeLock Lock(&GTextureMapLock);
-		bool bFound = GTextureMapEntries.RemoveAndCopyValue(TextureRHI, StagingImage);
+		bool bFound = GPendingLockedBuffers.RemoveAndCopyValue(TextureRHI, StagingBuffer);
 		checkf(bFound, TEXT("Texture was not locked!"));
 	}
 
-	StagingImage->Unmap();
-
-	//#todo-rco: TEMP!
-	Texture->StagingImage = StagingImage->GetHandle();
+	StagingBuffer->Unmap();
 
-	//#todo-rco: Put into immediate cmd list
-	{
-
-	}
+#if VULKAN_USE_NEW_COMMAND_BUFFERS
+	FVulkanCmdBuffer* CmdBuffer = Device->GetImmediateContext().GetCommandBufferManager()->GetActiveCmdBuffer();
+	ensure(CmdBuffer->IsOutsideRenderPass());
+	VkCommandBuffer StagingCommandBuffer = CmdBuffer->GetHandle();
 #else
-	//@TODO - share this with other texture unlocks
-	vkUnmapMemory(LogicalDevice, Texture->PixelBufferMemory);
-
-	// @TODO: set image layout?
-#endif
-
 	VkCommandBuffer StagingCommandBuffer;
-	VkCommandBufferAllocateInfo BufferCreateInfo;
-	FMemory::Memzero(BufferCreateInfo);
-	BufferCreateInfo.sType = VK_STRUCTURE_TYPE_COMMAND_BUFFER_ALLOCATE_INFO;
-	BufferCreateInfo.commandPool = Device->GetImmediateContext().GetCommandBufferManager()->GetHandle();
-	BufferCreateInfo.commandBufferCount = 1;
-	BufferCreateInfo.level = VK_COMMAND_BUFFER_LEVEL_PRIMARY;
-	VERIFYVULKANRESULT(vkAllocateCommandBuffers(LogicalDevice, &BufferCreateInfo, &StagingCommandBuffer));
+	VkCommandBufferAllocateInfo CmdBufferCreateInfo;
+	FMemory::Memzero(CmdBufferCreateInfo);
+	CmdBufferCreateInfo.sType = VK_STRUCTURE_TYPE_COMMAND_BUFFER_ALLOCATE_INFO;
+	CmdBufferCreateInfo.commandPool = Device->GetImmediateContext().GetCommandBufferManager()->GetHandle();
+	CmdBufferCreateInfo.commandBufferCount = 1;
+	CmdBufferCreateInfo.level = VK_COMMAND_BUFFER_LEVEL_PRIMARY;
+	VERIFYVULKANRESULT(vkAllocateCommandBuffers(LogicalDevice, &CmdBufferCreateInfo, &StagingCommandBuffer));
 
 	VkCommandBufferBeginInfo BufferBeginInfo;
 	FMemory::Memzero(BufferBeginInfo);
 	BufferBeginInfo.sType = VK_STRUCTURE_TYPE_COMMAND_BUFFER_BEGIN_INFO;
 	VERIFYVULKANRESULT(vkBeginCommandBuffer(StagingCommandBuffer, &BufferBeginInfo));
-
+#endif
 	EPixelFormat Format = Texture->Surface.Format;
 	uint32 MipWidth = FMath::Max<uint32>(Texture->Surface.Width >> MipIndex, GPixelFormats[Format].BlockSizeX);
 	uint32 MipHeight = FMath::Max<uint32>(Texture->Surface.Height >> MipIndex, GPixelFormats[Format].BlockSizeY);
 
-	VkImageCopy Region;
-	FMemory::Memzero(Region);
+	VkImageSubresourceRange SubresourceRange;
+	FMemory::Memzero(SubresourceRange);
+	SubresourceRange.aspectMask = VK_IMAGE_ASPECT_COLOR_BIT;
+	SubresourceRange.baseMipLevel = MipIndex;
+	SubresourceRange.levelCount = 1;
+	//SubresourceRange.baseArrayLayer = 0;
+	SubresourceRange.layerCount = 1;
+	VulkanSetImageLayout(StagingCommandBuffer, Texture->Surface.Image, VK_IMAGE_LAYOUT_UNDEFINED, VK_IMAGE_LAYOUT_GENERAL, SubresourceRange);
 
-	Region.srcSubresource.aspectMask = Texture->Surface.GetAspectMask();
-	Region.srcSubresource.mipLevel = 0;
-	Region.srcSubresource.baseArrayLayer = 0;
-	Region.srcSubresource.layerCount = 1;
-	Region.srcOffset.x = 0;
-	Region.srcOffset.y = 0;
-	Region.srcOffset.z = 0;
-
-	// dst image is expected to have the same dimensions/setup
-	Region.dstSubresource = Region.srcSubresource;
-	Region.dstSubresource.mipLevel = MipIndex;
-	
-	Region.extent.width = MipWidth;
-	Region.extent.height = MipHeight;
-	Region.extent.depth = 1;
+	VkBufferImageCopy Region;
+	FMemory::Memzero(Region);
+	//#todo-rco: Might need an offset here?
+	//Region.bufferOffset = 0;
+	//Region.bufferRowLength = 0;
+	//Region.bufferImageHeight = 0;
+	Region.imageSubresource.aspectMask = VK_IMAGE_ASPECT_COLOR_BIT;
+	Region.imageSubresource.mipLevel = MipIndex;
+	//Region.imageSubresource.baseArrayLayer = 0;
+	Region.imageSubresource.layerCount = 1;
+	Region.imageExtent.width = MipWidth;
+	Region.imageExtent.height = MipHeight;
+	Region.imageExtent.depth = 1;
+	vkCmdCopyBufferToImage(StagingCommandBuffer, StagingBuffer->GetHandle(), Texture->Surface.Image, VK_IMAGE_LAYOUT_GENERAL, 1, &Region);
 
-	vkCmdCopyImage(StagingCommandBuffer, Texture->StagingImage, VK_IMAGE_LAYOUT_TRANSFER_SRC_OPTIMAL, Texture->Surface.Image, VK_IMAGE_LAYOUT_TRANSFER_DST_OPTIMAL, 1, &Region);
+	VulkanSetImageLayout(StagingCommandBuffer, Texture->Surface.Image, VK_IMAGE_LAYOUT_GENERAL, VK_IMAGE_LAYOUT_SHADER_READ_ONLY_OPTIMAL, SubresourceRange);
 
+#if VULKAN_USE_NEW_COMMAND_BUFFERS
+	Device->GetStagingManager().ReleaseBuffer(CmdBuffer, StagingBuffer);
+#else
 	VERIFYVULKANRESULT(vkEndCommandBuffer(StagingCommandBuffer));
 
-	VkSubmitInfo SubmitInfo;
-	FMemory::Memzero(SubmitInfo);
-	SubmitInfo.sType = VK_STRUCTURE_TYPE_SUBMIT_INFO;
-	SubmitInfo.commandBufferCount = 1;
-	SubmitInfo.pCommandBuffers = &StagingCommandBuffer;
-
-#if VULKAN_USE_FENCE_MANAGER
-	VkResult Result = vkQueueSubmit(Device->GetQueue()->GetHandle(), 1, &SubmitInfo, Fence->GetHandle());
-#else
-	VkResult Result = vkQueueSubmit(Device->GetQueue()->GetHandle(), 1, &SubmitInfo, Fence);
-#endif
-	checkf(Result == VK_SUCCESS, TEXT("vkQueueSubmit failed with error %d, on texture with UE format %d, size %dx%d, mip %d"), (int32)Result, (int32)Format, MipWidth, MipHeight, MipIndex);
+	auto* Fence = Device->GetFenceManager().AllocateFence();
+	Device->GetQueue()->Submit2(StagingCommandBuffer, Fence);
 
-#if VULKAN_USE_FENCE_MANAGER
 	bool bSuccess = Device->GetFenceManager().WaitForFence(Fence, 0xFFFFFFFF);
 	check(bSuccess);
 
 	Device->GetFenceManager().ReleaseFence(Fence);
-#else
-	VERIFYVULKANRESULT(vkWaitForFences(LogicalDevice, 1, &Fence, true, 0xFFFFFFFF));
-	VERIFYVULKANRESULT(vkResetFences(LogicalDevice, 1, &Fence));
-	vkDestroyFence(LogicalDevice, Fence, nullptr);
-#endif
 
 	VERIFYVULKANRESULT(vkResetCommandBuffer(StagingCommandBuffer, 0x0));
 
-	//#todo-rco: FIX ME!
 	vkFreeCommandBuffers(LogicalDevice, Device->GetImmediateContext().GetCommandBufferManager()->GetHandle(), 1, &StagingCommandBuffer);
-#if VULKAN_USE_NEW_STAGING_BUFFERS
-	Device->GetStagingManager().ReleaseImage(StagingImage);
-#else
-	vkDestroyImage(LogicalDevice, Texture->StagingImage, nullptr);
-	vkFreeMemory(LogicalDevice, Texture->PixelBufferMemory, nullptr);
+	Device->GetStagingManager().ReleaseBuffer(StagingBuffer);
 #endif
 }
 
@@ -1093,7 +1007,6 @@ FVulkanTextureBase::FVulkanTextureBase(FVulkanDevice& Device, VkImageViewType Re
 	 , MSAASurface(nullptr)
 	#endif
 	, StagingImage(VK_NULL_HANDLE)
-	, PixelBufferMemory(VK_NULL_HANDLE)
 	#if VULKAN_HAS_DEBUGGING_ENABLED
 	, AllocationNumber(0)
 	#endif
@@ -1148,9 +1061,9 @@ FVulkanTextureBase::FVulkanTextureBase(FVulkanDevice& Device, VkImageViewType Re
 	FMemory::Memzero(BufferCreateInfo);
 	BufferCreateInfo.sType = VK_STRUCTURE_TYPE_BUFFER_CREATE_INFO;
 	BufferCreateInfo.flags = 0;
-	BufferCreateInfo.pQueueFamilyIndices = NULL;
-	BufferCreateInfo.queueFamilyIndexCount = 0;
-	BufferCreateInfo.sharingMode = VK_SHARING_MODE_EXCLUSIVE;
+	//BufferCreateInfo.pQueueFamilyIndices = NULL;
+	//BufferCreateInfo.queueFamilyIndexCount = 0;
+	//BufferCreateInfo.sharingMode = VK_SHARING_MODE_EXCLUSIVE;
 	BufferCreateInfo.size = CreateInfo.BulkData->GetResourceBulkDataSize();
 	BufferCreateInfo.usage = VK_BUFFER_USAGE_TRANSFER_SRC_BIT;
 
@@ -1162,52 +1075,55 @@ FVulkanTextureBase::FVulkanTextureBase(FVulkanDevice& Device, VkImageViewType Re
 
 	// Allocate device gmem
 	VkDeviceMemory Mem;
-#if 0//VULKAN_USE_NEW_STAGING_BUFFERS
+
 	//#todo-rco: USe StagingBuffers
-#else
-	#if VULKAN_USE_MEMORY_SYSTEM
-		VulkanRHI::FDeviceMemoryAllocation* Allocation = Device.GetMemoryManager().Alloc(MemReqs.size, MemReqs.memoryTypeBits, VK_MEMORY_PROPERTY_HOST_VISIBLE_BIT, __FILE__, __LINE__);
-		Mem = Allocation->GetHandle();
-	#else
-		VkMemoryAllocateInfo AllocInfo;
-		FMemory::Memzero(AllocInfo);
-		AllocInfo.sType = VK_STRUCTURE_TYPE_MEMORY_ALLOCATE_INFO;
-		AllocInfo.pNext = NULL;
-		AllocInfo.allocationSize = MemReqs.size;
-		AllocInfo.memoryTypeIndex = 0;
-
-		VERIFYVULKANRESULT(Device.GetMemoryManager().GetMemoryTypeFromProperties(MemReqs.memoryTypeBits, VK_MEMORY_PROPERTY_HOST_VISIBLE_BIT, &AllocInfo.memoryTypeIndex));
-		VERIFYVULKANRESULT(vkAllocateMemory(Device.GetInstanceHandle(), &AllocInfo, NULL, &Mem));
-	#endif
+	VulkanRHI::FDeviceMemoryAllocation* Allocation = Device.GetMemoryManager().Alloc(MemReqs.size, MemReqs.memoryTypeBits, VK_MEMORY_PROPERTY_HOST_VISIBLE_BIT, __FILE__, __LINE__);
+	Mem = Allocation->GetHandle();
 
 	VERIFYVULKANRESULT(vkBindBufferMemory(Device.GetInstanceHandle(), Buffer, Mem, 0));
 
 	void* Data = nullptr;
 	VERIFYVULKANRESULT(vkMapMemory(Device.GetInstanceHandle(), Mem, 0, MemReqs.size, 0, &Data));
-#endif
+
 	// Do copy
 	FMemory::Memcpy(Data, CreateInfo.BulkData->GetResourceBulkData(), CreateInfo.BulkData->GetResourceBulkDataSize());
 
 	vkUnmapMemory(Device.GetInstanceHandle(), Mem);
 
+#if VULKAN_USE_NEW_COMMAND_BUFFERS
+	FVulkanCmdBuffer* CmdBuffer = Device.GetImmediateContext().GetCommandBufferManager()->GetActiveCmdBuffer();
+	ensure(CmdBuffer->IsOutsideRenderPass());
+#else
 	//#todo-rco: FIX ME?
-	FVulkanCmdBuffer* Cmd = Device.GetImmediateContext().GetCommandBufferManager()->Create();
-	Cmd->Begin();
+	FVulkanCmdBuffer* CmdBuffer = Device.GetImmediateContext().GetCommandBufferManager()->Create();
+	CmdBuffer->Begin();
+#endif
+
+	VkImageSubresourceRange SubresourceRange;
+	FMemory::Memzero(SubresourceRange);
+	SubresourceRange.aspectMask = VK_IMAGE_ASPECT_COLOR_BIT;
+	SubresourceRange.baseMipLevel = 0;
+	SubresourceRange.levelCount = Surface.GetNumMips();
+	SubresourceRange.baseArrayLayer = 0;
+	SubresourceRange.layerCount = Surface.Depth;
+	VulkanSetImageLayout(CmdBuffer->GetHandle(), Surface.Image, VK_IMAGE_LAYOUT_UNDEFINED, VK_IMAGE_LAYOUT_GENERAL, SubresourceRange);
 
 	// Copy buffer to image
-	vkCmdCopyBufferToImage(Cmd->GetHandle(), Buffer, Surface.Image, Surface.ImageLayout, 1, &Region);
+	vkCmdCopyBufferToImage(CmdBuffer->GetHandle(), Buffer, Surface.Image, VK_IMAGE_LAYOUT_GENERAL, 1, &Region);
+
+	VulkanSetImageLayout(CmdBuffer->GetHandle(), Surface.Image, VK_IMAGE_LAYOUT_GENERAL, VK_IMAGE_LAYOUT_SHADER_READ_ONLY_OPTIMAL, SubresourceRange);
 
-	Device.GetQueue()->SubmitBlocking(Cmd);
+#if VULKAN_USE_NEW_COMMAND_BUFFERS
+#else
+	Device.GetQueue()->SubmitBlocking(CmdBuffer);
 	//#todo-rco: FIX ME!
-	Device.GetImmediateContext().GetCommandBufferManager()->Destroy(Cmd);
+	Device.GetImmediateContext().GetCommandBufferManager()->Destroy(CmdBuffer);
 
 	vkDestroyBuffer(Device.GetInstanceHandle(), Buffer, nullptr);
 
-	#if VULKAN_USE_MEMORY_SYSTEM
-		Device.GetMemoryManager().Free(Allocation);
-	#else
-		vkFreeMemory(Device.GetInstanceHandle(), Mem, nullptr);
-	#endif
+#endif
+	//#todo-rco: USe StagingBuffers
+	Device.GetMemoryManager().Free(Allocation);
 }
 
 FVulkanTextureBase::FVulkanTextureBase(FVulkanDevice& Device, VkImageViewType ResourceType, EPixelFormat Format, uint32 SizeX, uint32 SizeY, uint32 SizeZ, VkImage InImage, VkDeviceMemory InMem, uint32 UEFlags, const FRHIResourceCreateInfo& CreateInfo):
@@ -1216,7 +1132,6 @@ FVulkanTextureBase::FVulkanTextureBase(FVulkanDevice& Device, VkImageViewType Re
 	, MSAASurface(nullptr)
 	#endif
 	, StagingImage(VK_NULL_HANDLE)
-	, PixelBufferMemory(VK_NULL_HANDLE)
 	#if VULKAN_HAS_DEBUGGING_ENABLED
 	, AllocationNumber(0)
 	#endif
@@ -1314,93 +1229,22 @@ void* FVulkanDynamicRHI::RHILockTextureCubeFace(FTextureCubeRHIParamRef TextureC
 {
 	FVulkanTextureCube* Texture = ResourceCast(TextureCubeRHI);
 	check(Texture);
-	
-	EPixelFormat Format = Texture->Surface.Format;
-	uint32 MipWidth = FMath::Max<uint32>(Texture->Surface.Width >> MipIndex, GPixelFormats[Format].BlockSizeX);
-	uint32 MipHeight = FMath::Max<uint32>(Texture->Surface.Height >> MipIndex, GPixelFormats[Format].BlockSizeY);
 
-#if VULKAN_USE_NEW_STAGING_BUFFERS
-		
-	VulkanRHI::FStagingImage** StagingImage = nullptr;
+	VulkanRHI::FStagingBuffer** StagingBuffer = nullptr;
 	{
 		FScopeLock Lock(&GTextureMapLock);
-		StagingImage = &GTextureMapEntries.FindOrAdd(TextureCubeRHI);
-		checkf(!*StagingImage, TEXT("Can't lock the same texture twice!"));
+		StagingBuffer = &GPendingLockedBuffers.FindOrAdd(TextureCubeRHI);
+		checkf(!*StagingBuffer, TEXT("Can't lock the same texture twice!"));
 	}
 
-	*StagingImage = Device->GetStagingManager().AcquireImage(Texture->Surface.InternalFormat, MipWidth, MipHeight);
-
-	VkImageSubresource Subresource;
-	FMemory::Memzero(Subresource);
-	Subresource.aspectMask = Texture->Surface.GetAspectMask();
-	Subresource.mipLevel = 0;
-	Subresource.arrayLayer = 0;
-	VkSubresourceLayout SubresourceLayout;
-	vkGetImageSubresourceLayout(Device->GetInstanceHandle(), (*StagingImage)->GetHandle(), &Subresource, &SubresourceLayout);
-	DestStride = SubresourceLayout.rowPitch;
-
-	return (*StagingImage)->Map();
-#else
-	//@TODO - share this with other texture locks
-
-	FVulkanDevice* VulkanDevice = Texture->Surface.Device;
-	VkDevice LogicalDevice = VulkanDevice->GetInstanceHandle();
-
-
-	VkExtent3D StagingExtent;
-	FMemory::Memzero(StagingExtent);
-	StagingExtent.width = MipWidth;
-	StagingExtent.height = MipHeight;
-	StagingExtent.depth = 1;
-
-	VkImageCreateInfo StagingImageCreateInfo;
-	FMemory::Memzero(StagingImageCreateInfo);
-	StagingImageCreateInfo.sType = VK_STRUCTURE_TYPE_IMAGE_CREATE_INFO;
-	StagingImageCreateInfo.pNext = nullptr;
-	StagingImageCreateInfo.imageType = VK_IMAGE_TYPE_2D; // @TODO: check cube is 2D
-	StagingImageCreateInfo.format = Texture->Surface.InternalFormat;
-	StagingImageCreateInfo.extent = StagingExtent;
-	StagingImageCreateInfo.mipLevels = 1;
-	StagingImageCreateInfo.arrayLayers = 1;
-	StagingImageCreateInfo.samples = VK_SAMPLE_COUNT_1_BIT;
-	StagingImageCreateInfo.tiling = VK_IMAGE_TILING_LINEAR;
-	StagingImageCreateInfo.usage = VK_IMAGE_USAGE_TRANSFER_SRC_BIT;
-	
-	VERIFYVULKANRESULT(vkCreateImage(LogicalDevice, &StagingImageCreateInfo, nullptr, &Texture->StagingImage));
-
-	VkMemoryRequirements MemReqs;
-	vkGetImageMemoryRequirements(LogicalDevice, Texture->StagingImage, &MemReqs);
-
-	VkMemoryAllocateInfo AllocInfo;
-	FMemory::Memzero(AllocInfo);
-	AllocInfo.sType = VK_STRUCTURE_TYPE_MEMORY_ALLOCATE_INFO;
-	AllocInfo.allocationSize = MemReqs.size;
-
-	VkMemoryPropertyFlags MemoryTypeFlags = VK_MEMORY_PROPERTY_HOST_VISIBLE_BIT;
-	VERIFYVULKANRESULT(VulkanDevice->GetMemoryManager().GetMemoryTypeFromProperties(MemReqs.memoryTypeBits, MemoryTypeFlags, &AllocInfo.memoryTypeIndex));
-	#if VULKAN_USE_MEMORY_SYSTEM
-		check(0);
-	#else
-		VERIFYVULKANRESULT(vkAllocateMemory(LogicalDevice, &AllocInfo, nullptr, &Texture->PixelBufferMemory));
-	#endif
-
-	VERIFYVULKANRESULT(vkBindImageMemory(LogicalDevice, Texture->StagingImage, Texture->PixelBufferMemory, 0));
-
-	VkImageSubresource Subresource;
-	FMemory::Memzero(Subresource);
-	Subresource.aspectMask = VK_IMAGE_ASPECT_COLOR_BIT;
-	Subresource.mipLevel = 0;
-	Subresource.arrayLayer = 0;
-
-	VkSubresourceLayout SubresourceLayout;
-	vkGetImageSubresourceLayout(LogicalDevice, Texture->StagingImage, &Subresource, &SubresourceLayout);
-	DestStride = SubresourceLayout.rowPitch;
-
-	void* MappedData = nullptr;
-	VERIFYVULKANRESULT(vkMapMemory(LogicalDevice, Texture->PixelBufferMemory, 0, AllocInfo.allocationSize, 0, &MappedData));
+	uint32 BufferSize = 0;
+	DestStride = 0;
+	Texture->Surface.GetMipSize(MipIndex, BufferSize);
+	Texture->Surface.GetMipStride(MipIndex, DestStride);
+	*StagingBuffer = Device->GetStagingManager().AcquireBuffer(BufferSize);
 
-	return MappedData;
-	#endif
+	void* Data = (*StagingBuffer)->Map();
+	return Data;
 }
 
 void FVulkanDynamicRHI::RHIUnlockTextureCubeFace(FTextureCubeRHIParamRef TextureCubeRHI,uint32 FaceIndex,uint32 ArrayIndex,uint32 MipIndex,bool bLockWithinMiptail)
@@ -1410,112 +1254,81 @@ void FVulkanDynamicRHI::RHIUnlockTextureCubeFace(FTextureCubeRHIParamRef Texture
 
 	VkDevice LogicalDevice = Device->GetInstanceHandle();
 
-#if VULKAN_USE_FENCE_MANAGER
-	auto* Fence = Device->GetFenceManager().AllocateFence();
-#else
-	VkFence Fence;
-	VkFenceCreateInfo FenceCreateInfo;
-	FMemory::Memzero(FenceCreateInfo);
-	FenceCreateInfo.sType = VK_STRUCTURE_TYPE_FENCE_CREATE_INFO;
-	VERIFYVULKANRESULT(vkCreateFence(LogicalDevice, &FenceCreateInfo, nullptr, &Fence));
-#endif
-
-#if VULKAN_USE_NEW_STAGING_BUFFERS
-	VulkanRHI::FStagingImage* StagingImage = nullptr;
+	VulkanRHI::FStagingBuffer* StagingBuffer = nullptr;
 	{
 		FScopeLock Lock(&GTextureMapLock);
-		bool bFound = GTextureMapEntries.RemoveAndCopyValue(TextureCubeRHI, StagingImage);
+		bool bFound = GPendingLockedBuffers.RemoveAndCopyValue(TextureCubeRHI, StagingBuffer);
 		checkf(bFound, TEXT("Texture was not locked!"));
 	}
 
-	StagingImage->Unmap();
-
-	//#todo-rco: TEMP!
-	Texture->StagingImage = StagingImage->GetHandle();
-
-	//#todo-rco: Put into immediate cmd list
-	{
-	}
+	StagingBuffer->Unmap();
 
+#if VULKAN_USE_NEW_COMMAND_BUFFERS
+	FVulkanCmdBuffer* CmdBuffer = Device->GetImmediateContext().GetCommandBufferManager()->GetActiveCmdBuffer();
+	ensure(CmdBuffer->IsOutsideRenderPass());
+	VkCommandBuffer StagingCommandBuffer = CmdBuffer->GetHandle();
 #else
-	//@TODO - share this with other texture unlocks
-	vkUnmapMemory(LogicalDevice, Texture->PixelBufferMemory);
-
-	// @TODO: set image layout?
-#endif
-
 	VkCommandBuffer StagingCommandBuffer;
-	VkCommandBufferAllocateInfo BufferCreateInfo;
-	FMemory::Memzero(BufferCreateInfo);
-	BufferCreateInfo.sType = VK_STRUCTURE_TYPE_COMMAND_BUFFER_ALLOCATE_INFO;
-	BufferCreateInfo.commandPool = Device->GetImmediateContext().GetCommandBufferManager()->GetHandle();
-	BufferCreateInfo.commandBufferCount = 1;
-	BufferCreateInfo.level = VK_COMMAND_BUFFER_LEVEL_PRIMARY;
-	VERIFYVULKANRESULT(vkAllocateCommandBuffers(LogicalDevice, &BufferCreateInfo, &StagingCommandBuffer));
+	VkCommandBufferAllocateInfo CmdBufferCreateInfo;
+	FMemory::Memzero(CmdBufferCreateInfo);
+	CmdBufferCreateInfo.sType = VK_STRUCTURE_TYPE_COMMAND_BUFFER_ALLOCATE_INFO;
+	CmdBufferCreateInfo.commandPool = Device->GetImmediateContext().GetCommandBufferManager()->GetHandle();
+	CmdBufferCreateInfo.commandBufferCount = 1;
+	CmdBufferCreateInfo.level = VK_COMMAND_BUFFER_LEVEL_PRIMARY;
+	VERIFYVULKANRESULT(vkAllocateCommandBuffers(LogicalDevice, &CmdBufferCreateInfo, &StagingCommandBuffer));
 
 	VkCommandBufferBeginInfo BufferBeginInfo;
 	FMemory::Memzero(BufferBeginInfo);
 	BufferBeginInfo.sType = VK_STRUCTURE_TYPE_COMMAND_BUFFER_BEGIN_INFO;
 	VERIFYVULKANRESULT(vkBeginCommandBuffer(StagingCommandBuffer, &BufferBeginInfo));
-
+#endif
 	EPixelFormat Format = Texture->Surface.Format;
 	uint32 MipWidth = FMath::Max<uint32>(Texture->Surface.Width >> MipIndex, GPixelFormats[Format].BlockSizeX);
 	uint32 MipHeight = FMath::Max<uint32>(Texture->Surface.Height >> MipIndex, GPixelFormats[Format].BlockSizeY);
 
-	VkImageCopy Region;
+	VkImageSubresourceRange SubresourceRange;
+	FMemory::Memzero(SubresourceRange);
+	SubresourceRange.aspectMask = VK_IMAGE_ASPECT_COLOR_BIT;
+	SubresourceRange.baseMipLevel = MipIndex;
+	SubresourceRange.levelCount = 1;
+	SubresourceRange.baseArrayLayer = FaceIndex;
+	SubresourceRange.layerCount = 1;
+	VulkanSetImageLayout(StagingCommandBuffer, Texture->Surface.Image, VK_IMAGE_LAYOUT_UNDEFINED, VK_IMAGE_LAYOUT_GENERAL, SubresourceRange);
+
+	VkBufferImageCopy Region;
 	FMemory::Memzero(Region);
-	Region.srcSubresource.aspectMask = VK_IMAGE_ASPECT_COLOR_BIT;
-	Region.srcSubresource.mipLevel = 0;
-	Region.srcSubresource.baseArrayLayer = 0;
-	Region.srcSubresource.layerCount = 1;
-	Region.srcOffset.x = 0;
-	Region.srcOffset.y = 0;
-	Region.srcOffset.z = 0;
-	Region.dstSubresource.aspectMask = VK_IMAGE_ASPECT_COLOR_BIT;
-	Region.dstSubresource.mipLevel = MipIndex;
-	Region.dstSubresource.baseArrayLayer = FaceIndex;
-	Region.dstSubresource.layerCount = 1;
-	Region.dstOffset.x = 0;
-	Region.dstOffset.y = 0;
-	Region.dstOffset.z = 0;
-	Region.extent.width = MipWidth;
-	Region.extent.height = MipHeight;
-	Region.extent.depth = 1;
-
-	vkCmdCopyImage(StagingCommandBuffer, Texture->StagingImage, VK_IMAGE_LAYOUT_TRANSFER_SRC_OPTIMAL, Texture->Surface.Image, VK_IMAGE_LAYOUT_TRANSFER_DST_OPTIMAL, 1, &Region);
+	//#todo-rco: Might need an offset here?
+	//Region.bufferOffset = 0;
+	//Region.bufferRowLength = 0;
+	//Region.bufferImageHeight = 0;
+	Region.imageSubresource.aspectMask = VK_IMAGE_ASPECT_COLOR_BIT;
+	Region.imageSubresource.mipLevel = MipIndex;
+	Region.imageSubresource.baseArrayLayer = FaceIndex;
+	Region.imageSubresource.layerCount = 1;
+	Region.imageExtent.width = MipWidth;
+	Region.imageExtent.height = MipHeight;
+	Region.imageExtent.depth = 1;
+	vkCmdCopyBufferToImage(StagingCommandBuffer, StagingBuffer->GetHandle(), Texture->Surface.Image, VK_IMAGE_LAYOUT_GENERAL, 1, &Region);
 
-	VERIFYVULKANRESULT(vkEndCommandBuffer(StagingCommandBuffer));
+	VulkanSetImageLayout(StagingCommandBuffer, Texture->Surface.Image, VK_IMAGE_LAYOUT_GENERAL, VK_IMAGE_LAYOUT_SHADER_READ_ONLY_OPTIMAL, SubresourceRange);
 
-	VkSubmitInfo SubmitInfo;
-	FMemory::Memzero(SubmitInfo);
-	SubmitInfo.sType = VK_STRUCTURE_TYPE_SUBMIT_INFO;
-	SubmitInfo.commandBufferCount = 1;
-	SubmitInfo.pCommandBuffers = &StagingCommandBuffer;
+#if VULKAN_USE_NEW_COMMAND_BUFFERS
+	Device->GetStagingManager().ReleaseBuffer(CmdBuffer, StagingBuffer);
+#else
+	VERIFYVULKANRESULT(vkEndCommandBuffer(StagingCommandBuffer));
 
-#if VULKAN_USE_FENCE_MANAGER
-	VERIFYVULKANRESULT(vkQueueSubmit(Texture->Surface.Device->GetQueue()->GetHandle(), 1, &SubmitInfo, Fence->GetHandle()));
+	auto* Fence = Device->GetFenceManager().AllocateFence();
+	Device->GetQueue()->Submit2(StagingCommandBuffer, Fence);
 
 	bool bSuccess = Device->GetFenceManager().WaitForFence(Fence, 0xFFFFFFFF);
 	check(bSuccess);
 
 	Device->GetFenceManager().ReleaseFence(Fence);
-#else
-	VERIFYVULKANRESULT(vkQueueSubmit(Texture->Surface.Device->GetQueue()->GetHandle(), 1, &SubmitInfo, Fence));
-	
-	VERIFYVULKANRESULT(vkWaitForFences(LogicalDevice, 1, &Fence, true, 0xFFFFFFFF));
-	VERIFYVULKANRESULT(vkResetFences(LogicalDevice, 1, &Fence));
-	vkDestroyFence(LogicalDevice, Fence, nullptr);
-#endif
 
 	VERIFYVULKANRESULT(vkResetCommandBuffer(StagingCommandBuffer, 0x0));
 
-	//#todo-rco: FIX ME!
 	vkFreeCommandBuffers(LogicalDevice, Device->GetImmediateContext().GetCommandBufferManager()->GetHandle(), 1, &StagingCommandBuffer);
-#if VULKAN_USE_NEW_STAGING_BUFFERS
-	Device->GetStagingManager().ReleaseImage(StagingImage);
-#else
-	vkDestroyImage(LogicalDevice, Texture->StagingImage, nullptr);
-	vkFreeMemory(LogicalDevice, Texture->PixelBufferMemory, nullptr);
+	Device->GetStagingManager().ReleaseBuffer(StagingBuffer);
 #endif
 }
 
diff --git a/Engine/Source/Runtime/VulkanRHI/Private/VulkanUtil.cpp b/Engine/Source/Runtime/VulkanRHI/Private/VulkanUtil.cpp
--- a/Engine/Source/Runtime/VulkanRHI/Private/VulkanUtil.cpp
+++ b/Engine/Source/Runtime/VulkanRHI/Private/VulkanUtil.cpp
@@ -9,6 +9,7 @@
 #include "VulkanPendingState.h"
 #include "VulkanManager.h"
 #include "VulkanContext.h"
+#include "VulkanMemory.h"
 
 static FVulkanTimestampQueryPool* GTimestampQueryPool = nullptr;
 
@@ -67,7 +68,11 @@ void FVulkanGPUTiming::StartTiming()
 		if (StartTimestamp >= 0 && EndTimestamp >= 0)
 		{
 			auto& State = GTimestampQueryPool->Device->GetPendingState();
+#if VULKAN_USE_NEW_COMMAND_BUFFERS
+			check(0);
+#else
 			GTimestampQueryPool->WriteTimestamp(State.GetCommandBuffer(), StartTimestamp, VK_PIPELINE_STAGE_ALL_COMMANDS_BIT);
+#endif
 		}
 		bIsTiming = true;
 	}
@@ -85,7 +90,11 @@ void FVulkanGPUTiming::EndTiming()
 		if (StartTimestamp >= 0 && EndTimestamp >= 0)
 		{
 			auto& State = GTimestampQueryPool->Device->GetPendingState();
+#if VULKAN_USE_NEW_COMMAND_BUFFERS
+			check(0);
+#else
 			GTimestampQueryPool->WriteTimestamp(State.GetCommandBuffer(), EndTimestamp, VK_PIPELINE_STAGE_ALL_COMMANDS_BIT);
+#endif
 		}
 		bIsTiming = false;
 		bEndTimestampIssued = true;
@@ -363,6 +372,76 @@ void FVulkanDynamicRHI::IssueLongGPUTask()
 	}
 }
 
+#if VULKAN_USE_NEW_RESOURCE_MANAGEMENT
+FVulkanBuffer2::FVulkanBuffer2(FVulkanDevice* InDevice, VkBuffer InBuffer, VulkanRHI::FResourceAllocation* InResourceAllocation)
+	: FDeviceChild(InDevice)
+	, Buffer(InBuffer)
+	, ResourceAllocation(InResourceAllocation)
+{
+	check(Buffer != VK_NULL_HANDLE);
+	VERIFYVULKANRESULT(vkBindBufferMemory(InDevice->GetInstanceHandle(), Buffer, ResourceAllocation->GetHandle(), ResourceAllocation->GetOffset()));
+}
+
+FVulkanBuffer2::~FVulkanBuffer2()
+{
+	vkDestroyBuffer(Device->GetInstanceHandle(), Buffer, nullptr);
+	ResourceAllocation = nullptr;
+}
+
+
+VkBuffer FVulkanBuffer2::CreateVulkanBuffer(FVulkanDevice* InDevice, VkDeviceSize Size, VkBufferUsageFlags BufferUsageFlags, VkMemoryRequirements& OutMemoryRequirements)
+{
+	VkDevice Device = InDevice->GetInstanceHandle();
+	VkBuffer Buffer = VK_NULL_HANDLE;
+
+	VkBufferCreateInfo BufferCreateInfo;
+	FMemory::Memzero(BufferCreateInfo);
+	BufferCreateInfo.sType = VK_STRUCTURE_TYPE_BUFFER_CREATE_INFO;
+	BufferCreateInfo.size = Size;
+	BufferCreateInfo.usage = BufferUsageFlags;
+	VERIFYVULKANRESULT_EXPANDED(vkCreateBuffer(Device, &BufferCreateInfo, nullptr, &Buffer));
+
+	vkGetBufferMemoryRequirements(Device, Buffer, &OutMemoryRequirements);
+
+	return Buffer;
+}
+
+VulkanRHI::FResourceAllocation* FVulkanBuffer2::CreateResourceAllocationAndBuffer(FVulkanDevice* Device, VkDeviceSize Size, uint32 InUsage, VkBufferUsageFlags BufferUsageFlags, VkBuffer& OutBuffer)
+{
+	check(Size > 0);
+
+	VkMemoryRequirements MemoryRequirements;
+	OutBuffer = FVulkanBuffer2::CreateVulkanBuffer(Device, Size, BufferUsageFlags, MemoryRequirements);
+
+	const bool bStatic = (InUsage & BUF_Static) != 0;
+	const bool bDynamic = (InUsage & BUF_Dynamic) != 0;
+	const bool bVolatile = (InUsage & BUF_Volatile) != 0;
+
+	//#todo-rco: Not implemented yet
+	check((InUsage & BUF_KeepCPUAccessible) == 0);
+	VulkanRHI::FResourceAllocation* ResourceAllocation = nullptr;
+	if (bDynamic)
+	{
+		//#todo-rco: Allocate in the CPU staging heap
+		ResourceAllocation = Device->GetResourceHeapManager().AllocateCPUStagingResource(MemoryRequirements.size, MemoryRequirements.alignment, __FILE__, __LINE__);
+	}
+#if 0
+	else if (bVolatile)
+	{
+		check(0);
+	}
+#endif
+	else
+	{
+		// Static buffer path
+		ensure(bStatic);
+		ResourceAllocation = Device->GetResourceHeapManager().AllocateGPUOnlyResource(MemoryRequirements.size, MemoryRequirements.alignment, __FILE__, __LINE__);
+	}
+
+	return ResourceAllocation;
+}
+#endif
+
 
 namespace VulkanRHI
 {
diff --git a/Engine/Source/Runtime/VulkanRHI/Private/VulkanVertexDeclaration.cpp b/Engine/Source/Runtime/VulkanRHI/Private/VulkanVertexDeclaration.cpp
--- a/Engine/Source/Runtime/VulkanRHI/Private/VulkanVertexDeclaration.cpp
+++ b/Engine/Source/Runtime/VulkanRHI/Private/VulkanVertexDeclaration.cpp
@@ -19,7 +19,7 @@ struct FVulkanVertexDeclarationKey
     }
 };
 
-uint32 GetTypeHash(const FVulkanVertexDeclarationKey& Key)
+inline uint32 GetTypeHash(const FVulkanVertexDeclarationKey& Key)
 {
     return Key.Hash;
 }
