Copying //UE4/Dev-Rendering to Dev-Main (Source //UE4/Dev-Rendering@2932636)

diff --git a/Engine/Source/Runtime/VulkanRHI/Private/VulkanViewport.cpp b/Engine/Source/Runtime/VulkanRHI/Private/VulkanViewport.cpp
--- a/Engine/Source/Runtime/VulkanRHI/Private/VulkanViewport.cpp
+++ b/Engine/Source/Runtime/VulkanRHI/Private/VulkanViewport.cpp
@@ -37,27 +37,34 @@ FVulkanViewport::FVulkanViewport(FVulkanDynamicRHI* InRHI, void* WindowHandle, u
 
 	check(Images.Num() == NUM_BUFFERS);
 
-	//#todo-rco: FIX ME!
-	FVulkanCmdBuffer* cmd = RHI->Device->GetImmediateContext().GetCommandBufferManager()->Create();
-	cmd->Begin();
-
+#if VULKAN_USE_NEW_COMMAND_BUFFERS
+	FVulkanCmdBuffer* CmdBuffer = RHI->Device->GetImmediateContext().GetCommandBufferManager()->GetActiveCmdBuffer();
+	check(CmdBuffer->IsOutsideRenderPass());
+#else
+	FVulkanCmdBuffer* CmdBuffer = RHI->Device->GetImmediateContext().GetCommandBufferManager()->Create();
+	CmdBuffer->Begin();
+#endif
 	for (int32 Index = 0, Count = Images.Num(); Index < Count; ++Index)
 	{
-		VkImage img = Images[Index];
-
-		VulkanSetImageLayout(
-			cmd->GetHandle(),
-			img,
-			VK_IMAGE_ASPECT_COLOR_BIT,
-			VK_IMAGE_LAYOUT_UNDEFINED,
-			VK_IMAGE_LAYOUT_PRESENT_SRC_KHR);
-
-		BackBuffers[Index] = new FVulkanTexture2D(
-			*RHI->Device, PixelFormat, InSizeX, InSizeY,
-			img, TexCreate_Presentable | TexCreate_RenderTargetable, FClearValueBinding());
+		VkImage Image = Images[Index];
+
+		// Constructor will set to color optimal
+		BackBuffers[Index] = new FVulkanTexture2D(*RHI->Device, PixelFormat, InSizeX, InSizeY, Image, TexCreate_Presentable | TexCreate_RenderTargetable, FClearValueBinding());
+
+		FName Name = FName(*FString::Printf(TEXT("BackBuffer%d"), Index));
+		BackBuffers[Index]->SetName(Name);
+
+#if VULKAN_USE_NEW_COMMAND_BUFFERS
+#else
+		// Set to present src as the next calls will transition to color optimal
+		check(BackBuffers[Index]->Surface.ImageLayout == VK_IMAGE_LAYOUT_PRESENT_SRC_KHR);
+#endif
 	}
 
-	RHI->Device->EndCommandBufferBlock(cmd);
+#if VULKAN_USE_NEW_COMMAND_BUFFERS
+#else
+	RHI->Device->EndCommandBufferBlock(CmdBuffer);
+#endif
 }
 
 FVulkanViewport::~FVulkanViewport()
@@ -92,6 +99,9 @@ FVulkanFramebuffer::FVulkanFramebuffer(FVulkanDevice& Device, const FRHISetRende
 	#if VULKAN_USE_MSAA_RESOLVE_ATTACHMENTS
 		if (Texture->MSAASurface)
 		{
+#if VULKAN_USE_NEW_COMMAND_BUFFERS
+			check(0);
+#endif
 			Attachments.Add(Texture->MSAAView.View);
 
 		    // Create a write-barrier
@@ -113,6 +123,8 @@ FVulkanFramebuffer::FVulkanFramebuffer(FVulkanDevice& Device, const FRHISetRende
 	#endif
 		Attachments.Add(Texture->View.View);
 
+#if VULKAN_USE_NEW_COMMAND_BUFFERS
+#else
 		// Create a write-barrier
 		WriteBarriers.AddZeroed();
 		VkImageMemoryBarrier& barrier = WriteBarriers.Last();
@@ -120,15 +132,15 @@ FVulkanFramebuffer::FVulkanFramebuffer(FVulkanDevice& Device, const FRHISetRende
 		barrier.pNext = NULL;
 		barrier.srcAccessMask = VK_ACCESS_COLOR_ATTACHMENT_WRITE_BIT;
 		barrier.dstAccessMask = VK_ACCESS_MEMORY_READ_BIT;
-		barrier.oldLayout = VK_IMAGE_LAYOUT_GENERAL;
-		barrier.newLayout = VK_IMAGE_LAYOUT_GENERAL;
+		barrier.oldLayout = Texture->Surface.ImageLayout;
+		barrier.newLayout = VK_IMAGE_LAYOUT_COLOR_ATTACHMENT_OPTIMAL;
 		barrier.image = Texture->Surface.Image;
 		barrier.subresourceRange.aspectMask = Texture->Surface.GetAspectMask();
 		barrier.subresourceRange.baseMipLevel = 0;
 		barrier.subresourceRange.levelCount = 1;
 		barrier.subresourceRange.baseArrayLayer = 0;
 		barrier.subresourceRange.layerCount = 1;
-
+#endif
 		NumColorAttachments++;
 	}
 
@@ -149,6 +161,8 @@ FVulkanFramebuffer::FVulkanFramebuffer(FVulkanDevice& Device, const FRHISetRende
 
 		Attachments.Add(Texture->View.View);
 
+#if VULKAN_USE_NEW_COMMAND_BUFFERS
+#else
 		// Create a write-barrier
 		WriteBarriers.AddZeroed();
 		VkImageMemoryBarrier& barrier = WriteBarriers.Last();
@@ -157,7 +171,7 @@ FVulkanFramebuffer::FVulkanFramebuffer(FVulkanDevice& Device, const FRHISetRende
 		barrier.srcAccessMask = VK_ACCESS_DEPTH_STENCIL_ATTACHMENT_WRITE_BIT;
 		barrier.dstAccessMask = VK_ACCESS_MEMORY_READ_BIT;
 		barrier.oldLayout = VK_IMAGE_LAYOUT_GENERAL;
-		barrier.newLayout = VK_IMAGE_LAYOUT_GENERAL;
+		barrier.newLayout = VK_IMAGE_LAYOUT_DEPTH_STENCIL_ATTACHMENT_OPTIMAL;
 		barrier.image = Texture->Surface.Image;
 
 		barrier.subresourceRange.aspectMask = Texture->Surface.GetAspectMask();
@@ -165,6 +179,7 @@ FVulkanFramebuffer::FVulkanFramebuffer(FVulkanDevice& Device, const FRHISetRende
 		barrier.subresourceRange.levelCount = 1;
 		barrier.subresourceRange.baseArrayLayer = 0;
 		barrier.subresourceRange.layerCount = 1;
+#endif
 	}
 
 	VkFramebufferCreateInfo Info;
@@ -230,6 +245,34 @@ bool FVulkanFramebuffer::Matches(const FRHISetRenderTargetsInfo& InRTInfo) const
 	return true;
 }
 
+#if VULKAN_USE_NEW_COMMAND_BUFFERS
+void FVulkanViewport::AcquireBackBuffer(FVulkanCmdBuffer* CmdBuffer)
+{
+	CurrentBackBuffer = SwapChain->AcquireImageIndex();
+	check(CurrentBackBuffer != -1);
+	VulkanSetImageLayoutSimple(CmdBuffer->GetHandle(), BackBuffers[CurrentBackBuffer]->Surface.Image, VK_IMAGE_LAYOUT_PRESENT_SRC_KHR, VK_IMAGE_LAYOUT_COLOR_ATTACHMENT_OPTIMAL);
+}
+
+void FVulkanViewport::PrepareBackBufferForPresent(FVulkanCmdBuffer* CmdBuffer)
+{
+	check(CurrentBackBuffer != -1);
+	VulkanSetImageLayoutSimple(CmdBuffer->GetHandle(), BackBuffers[CurrentBackBuffer]->Surface.Image, VK_IMAGE_LAYOUT_COLOR_ATTACHMENT_OPTIMAL, VK_IMAGE_LAYOUT_PRESENT_SRC_KHR);
+}
+
+void FVulkanFramebuffer::InsertWriteBarrier(FVulkanCmdBuffer* CmdBuffer)
+{
+	if (WriteBarriers.Num() == 0)
+	{
+		return;
+	}
+
+	const VkPipelineStageFlags SrcStages = VK_PIPELINE_STAGE_COLOR_ATTACHMENT_OUTPUT_BIT;
+	const VkPipelineStageFlags DestStages = VK_PIPELINE_STAGE_BOTTOM_OF_PIPE_BIT;
+
+	check(CmdBuffer->IsOutsideRenderPass());
+	vkCmdPipelineBarrier(CmdBuffer->GetHandle(), SrcStages, DestStages, 0, 0, nullptr, 0, nullptr, WriteBarriers.Num(), WriteBarriers.GetData());
+}
+#else
 void FVulkanFramebuffer::InsertWriteBarrier(VkCommandBuffer cmd)
 {
 	if (WriteBarriers.Num() == 0)
@@ -237,12 +280,13 @@ void FVulkanFramebuffer::InsertWriteBarrier(VkCommandBuffer cmd)
 		return;
 	}
 
-	const VkPipelineStageFlags src_stages = VK_PIPELINE_STAGE_COLOR_ATTACHMENT_OUTPUT_BIT;
-	const VkPipelineStageFlags dest_stages = VK_PIPELINE_STAGE_BOTTOM_OF_PIPE_BIT;
+	const VkPipelineStageFlags SrcStages = VK_PIPELINE_STAGE_COLOR_ATTACHMENT_OUTPUT_BIT;
+	const VkPipelineStageFlags DestStages = VK_PIPELINE_STAGE_BOTTOM_OF_PIPE_BIT;
 
-	vkCmdPipelineBarrier(cmd, src_stages, dest_stages,
+	vkCmdPipelineBarrier(cmd, SrcStages, DestStages,
 		0, 0, nullptr, 0, nullptr, WriteBarriers.Num(), WriteBarriers.GetData());
 }
+#endif
 
 /*=============================================================================
  *	The following RHI functions must be called from the main thread.
@@ -272,8 +316,13 @@ void FVulkanDynamicRHI::WriteEndFrameTimestamp(void* Data)
 		auto* TimestampQueryPool = This->Device->GetTimestampQueryPool(This->PresentCount % (uint32)FVulkanDevice::NumTimestampPools);
 		if (TimestampQueryPool)
 		{
+#if VULKAN_USE_NEW_COMMAND_BUFFERS
+			auto* CmdBuffer = This->Device->GetImmediateContext().GetCommandBufferManager()->GetActiveCmdBuffer();
+			TimestampQueryPool->WriteEndFrame(CmdBuffer->GetHandle());
+#else
 			auto& PendingState = This->Device->GetPendingState();
 			TimestampQueryPool->WriteEndFrame(PendingState.GetCommandBuffer());
+#endif
 		}
 	}
 
@@ -288,9 +337,26 @@ void FVulkanDynamicRHI::Present()
 
 	check(Device);
 
+#if VULKAN_USE_NEW_COMMAND_BUFFERS
+	auto* CmdBufferManager = Device->GetImmediateContext().GetCommandBufferManager();
+	auto* CmdBuffer = CmdBufferManager->GetActiveCmdBuffer();
+	if (PendingState.IsRenderPassActive())
+	{
+		PendingState.RenderPassEnd(CmdBuffer);
+	}
+	DrawingViewport->PrepareBackBufferForPresent(CmdBuffer);
+	WriteEndFrameTimestamp(this);
+	CmdBuffer->End();
+	Device->GetQueue()->Submit(CmdBuffer);
+#else
 	PendingState.SubmitPendingCommandBuffers(&WriteEndFrameTimestamp, this);
+#endif
 
+#if VULKAN_USE_NEW_COMMAND_BUFFERS
+	DrawingViewport->GetSwapChain()->Present(Device->GetQueue());
+#else
 	Device->GetQueue()->Present(DrawingViewport->SwapChain, DrawingViewport->CurrentBackBuffer);
+#endif
 
 	bool bNativelyPresented = true;
 	if (bNativelyPresented)
@@ -323,6 +389,10 @@ void FVulkanDynamicRHI::Present()
 	//#todo-rco: This needs to go on RHIEndFrame but the CmdBuffer index is not the correct one to read the stats out!
 	VulkanRHI::GManager.GPUProfilingData.EndFrame();
 
+#if VULKAN_USE_NEW_COMMAND_BUFFERS
+	Device->GetImmediateContext().GetCommandBufferManager()->PrepareForNewActiveCommandBuffer();
+#endif
+
 	DrawingViewport->CurrentBackBuffer = -1;
 	DrawingViewport = nullptr;
 	PendingState.Reset();
@@ -339,7 +409,11 @@ void FVulkanDynamicRHI::Present()
 			TimestampQueryPool->CalculateFrameTime();
 		}
 
+#if VULKAN_USE_NEW_COMMAND_BUFFERS
+		Device->GetTimestampQueryPool(QueryNextFrameIndex)->WriteStartFrame(CmdBufferManager->GetActiveCmdBuffer()->GetHandle());
+#else
 		Device->GetTimestampQueryPool(QueryNextFrameIndex)->WriteStartFrame(PendingState.GetCommandBuffer());
+#endif
 	}
 
 	PresentCount++;
@@ -350,14 +424,22 @@ FTexture2DRHIRef FVulkanDynamicRHI::RHIGetViewportBackBuffer(FViewportRHIParamRe
 	FVulkanViewport* Viewport = ResourceCast(ViewportRHI);
 	if (Viewport->CurrentBackBuffer < 0)
 	{
+#if VULKAN_USE_NEW_COMMAND_BUFFERS
+		Viewport->AcquireBackBuffer(Device->GetImmediateContext().GetCommandBufferManager()->GetActiveCmdBuffer());
+#else
 		// SMEDIS: Acquire a new backbuffer now, in case RHIBeginDrawingViewport() hasn't been called yet.
 		Viewport->CurrentBackBuffer = Device->GetQueue()->AquireImageIndex(Viewport->SwapChain);
+#endif
 	}
 	return Viewport->GetBackBuffer();
 }
 
 void FVulkanDynamicRHI::RHIAdvanceFrameForGetViewportBackBuffer()
 {
+#if VULKAN_USE_NEW_COMMAND_BUFFERS
+	//#todo-rco: Do a pass to clear unused or expired elements in managers
+#else
+#endif
 }
 
 void FVulkanCommandListContext::RHISetViewport(uint32 MinX, uint32 MinY, float MinZ, uint32 MaxX, uint32 MaxY, float MaxZ)
diff --git a/Engine/Source/Runtime/VulkanRHI/Public/VulkanConfiguration.h b/Engine/Source/Runtime/VulkanRHI/Public/VulkanConfiguration.h
--- a/Engine/Source/Runtime/VulkanRHI/Public/VulkanConfiguration.h
+++ b/Engine/Source/Runtime/VulkanRHI/Public/VulkanConfiguration.h
@@ -9,9 +9,9 @@
 #include "RHIDefinitions.h"
 
 #if PLATFORM_WINDOWS
-	#define UE_VK_API_VERSION	VK_MAKE_VERSION(1, 0, 3)
+	#define UE_VK_API_VERSION	VK_MAKE_VERSION(1, 0, 5)
 #elif PLATFORM_ANDROID
-	#define UE_VK_API_VERSION	VK_MAKE_VERSION(1, 0, 1)
+	#define UE_VK_API_VERSION	VK_MAKE_VERSION(1, 0, 3)
 #else
 	#error Unsupported platform!
 #endif
@@ -50,17 +50,6 @@ enum class EDescriptorSetStage
 	Invalid		= -1,
 };
 
-// VULKAN_HLSLCC_GEN_GLSL_IN_OUT_WITHOUT_STRUCTURES:
-//	0:		"layout(location=%d) struct { vec4 Data; } varname;"
-//	1:		"layout(location=%d) vec4 varname;"
-//
-//	Originally, in/outs are generated with structures, which is fine. However,
-//	when the next shader stage has a compiled-out location and is first location is not starting with "0",
-//	the input is getting messed up. Therefore, we generate in/outs without structures.
-//
-//	NOTE: Changing the flag requires a full shader-rebuild.
-#define VULKAN_HLSLCC_GEN_GLSL_IN_OUT_WITHOUT_STRUCTURES		1
-
 inline EDescriptorSetStage GetDescriptorSetForStage(EShaderFrequency Stage)
 {
 	switch (Stage)
@@ -91,10 +80,10 @@ inline EDescriptorSetStage GetDescriptorSetForStage(EShaderFrequency Stage)
 #if PLATFORM_WINDOWS
 	#define VULKAN_DISABLE_DEBUG_CALLBACK						0	/* Disable the DebugReportFunction() callback in VulkanDebug.cpp */
 	#define VULKAN_CLEAR_SURFACE_ON_CREATE						1
-	#define VULKAN_USE_MSAA_RESOLVE_ATTACHMENTS					1	/* 1 = use resolve attachments, 0 = Use a command buffer vkResolveImage for MSAA resolve */
+	#define VULKAN_USE_MSAA_RESOLVE_ATTACHMENTS					0	/* 1 = use resolve attachments, 0 = Use a command buffer vkResolveImage for MSAA resolve */
 	#define VULKAN_USE_RING_BUFFER_FOR_GLOBAL_UBS				0	/* For some reason, using this on PC renders black - it may need lock/unlock all the time */
 	#define VULKAN_STRICT_TEXTURE_FLAGS							0	/* Checks the format feature flags when determining usage & flags for surfaces */
-	#define VULKAN_FORCE_WAIT_FOR_QUEUE							0	/* Ensures all work is finished inside the queue, before proceeding to the next frame. This should be only used to test for stability issues. The performance will be garbage, since we are not making use of any double/tripple buffering. */
+	#define VULKAN_FORCE_WAIT_FOR_QUEUE							0	/* Ensures all work is finished inside the queue, before proceeding to the next frame. This should be only used to test for stability issues. The performance will be garbage, since we are not making use of any double/triple buffering. */
 #else
 	#define VULKAN_DISABLE_DEBUG_CALLBACK						1	/* Disable the DebugReportFunction() callback in VulkanDebug.cpp */
 	#define VULKAN_CLEAR_SURFACE_ON_CREATE						0
@@ -109,16 +98,10 @@ inline EDescriptorSetStage GetDescriptorSetForStage(EShaderFrequency Stage)
 #define VULKAN_ENABLE_PIPELINE_CACHE							1
 
 //#todo-rco: Test on Android
-#define VULKAN_USE_MEMORY_SYSTEM								PLATFORM_WINDOWS	// Enables all vkAllocateMemory/vkFreeMemory to go through one system
-
-//#todo-rco: Test on Android
-#define VULKAN_USE_RESOURCE_ALLOCATIONS							0//PLATFORM_WINDOWS	// Does not directly allocate memory for buffers, instead partitions buffer ranges without reallocating memory
-
-//#todo-rco: Test on Android
-#define VULKAN_USE_NEW_STAGING_BUFFERS							PLATFORM_WINDOWS	// Enables reusing staging buffer memory
+#define VULKAN_USE_NEW_RESOURCE_MANAGEMENT						0//PLATFORM_WINDOWS
 
 //#todo-rco: Test on Android
-#define VULKAN_USE_FENCE_MANAGER								PLATFORM_WINDOWS
+#define VULKAN_USE_NEW_COMMAND_BUFFERS							0//PLATFORM_WINDOWS
 
 #define VULKAN_ENABLE_RHI_DEBUGGING								1
 
@@ -144,13 +127,3 @@ inline EDescriptorSetStage GetDescriptorSetForStage(EShaderFrequency Stage)
 		#define VULKAN_DISABLE_DEBUG_CALLBACK 0
 	#endif
 #endif
-
-
-#if !VULKAN_USE_MEMORY_SYSTEM
-	#if VULKAN_USE_NEW_STAGING_BUFFERS
-		#error VULKAN_USE_NEW_STAGING_BUFFERS requires VULKAN_USE_MEMORY_SYSTEM enabled!
-	#endif
-	#if VULKAN_USE_BUFFER_HEAPS
-		#error VULKAN_USE_BUFFER_HEAPS requires VULKAN_USE_MEMORY_SYSTEM enabled!
-	#endif
-#endif
diff --git a/Engine/Source/Runtime/VulkanRHI/Public/VulkanContext.h b/Engine/Source/Runtime/VulkanRHI/Public/VulkanContext.h
--- a/Engine/Source/Runtime/VulkanRHI/Public/VulkanContext.h
+++ b/Engine/Source/Runtime/VulkanRHI/Public/VulkanContext.h
@@ -8,8 +8,8 @@
 
 #include "VulkanResources.h"
 
-struct FVulkanDevice;
-struct FVulkanCommandBufferManager;
+class FVulkanDevice;
+class FVulkanCommandBufferManager;
 
 class FVulkanCommandListContext : public IRHICommandContext
 {
@@ -87,6 +87,7 @@ public:
 	virtual void RHISetMultipleViewports(uint32 Count, const FViewportBounds* Data) final override;
 	virtual void RHIClearUAV(FUnorderedAccessViewRHIParamRef UnorderedAccessViewRHI, const uint32* Values) final override;
 	virtual void RHICopyToResolveTarget(FTextureRHIParamRef SourceTexture, FTextureRHIParamRef DestTexture, bool bKeepOriginalSurface, const FResolveParams& ResolveParams) final override;
+	virtual void RHITransitionResources(EResourceTransitionAccess TransitionType, FTextureRHIParamRef* InRenderTargets, int32 NumTextures) final override;
 
 	// Render time measurement
 	virtual void RHIBeginRenderQuery(FRenderQueryRHIParamRef RenderQuery) final override;
@@ -145,4 +146,6 @@ protected:
 	TArray<FString> EventStack;
 
 	FVulkanCommandBufferManager* CommandBufferManager;
+
+	void SubmitCurrentCommands();
 };
diff --git a/Engine/Source/Runtime/VulkanRHI/Public/VulkanRHI.h b/Engine/Source/Runtime/VulkanRHI/Public/VulkanRHI.h
--- a/Engine/Source/Runtime/VulkanRHI/Public/VulkanRHI.h
+++ b/Engine/Source/Runtime/VulkanRHI/Public/VulkanRHI.h
@@ -18,9 +18,9 @@ typedef uint32 FStateKey;
 #include "VulkanViewport.h"
 
 class FVulkanTexture2D;
-struct FVulkanFramebuffer;
-struct FVulkanDevice;
-struct FVulkanQueue;
+class FVulkanFramebuffer;
+class FVulkanDevice;
+class FVulkanQueue;
 
 
 #if VULKAN_ENABLE_DRAW_MARKERS
@@ -186,12 +186,34 @@ protected:
 	static void RebuildPipelineCache();
 #endif
 
+#if UE_BUILD_DEBUG || UE_BUILD_DEVELOPMENT
+	IConsoleObject* DumpMemoryCmd;
+	static void DumpMemory();
+#endif
+
 #if VULKAN_HAS_DEBUGGING_ENABLED
 	VkDebugReportCallbackEXT MsgCallback;
 	void SetupDebugLayerCallback();
 	void RemoveDebugLayerCallback();
 #endif
 
+	inline bool ShouldDeferBufferLockOperation(FRHICommandList* RHICmdList)
+	{
+		if (RHICmdList == nullptr)
+		{
+			return false;
+		}
+
+		if (RHICmdList->Bypass() || !GRHIThread)
+		{
+			return false;
+		}
+
+		return true;
+	}
+
+	FCriticalSection LockBufferCS;
+
 	uint32 PresentCount;
 };
 
diff --git a/Engine/Source/Runtime/VulkanRHI/Public/VulkanResources.h b/Engine/Source/Runtime/VulkanRHI/Public/VulkanResources.h
--- a/Engine/Source/Runtime/VulkanRHI/Public/VulkanResources.h
+++ b/Engine/Source/Runtime/VulkanRHI/Public/VulkanResources.h
@@ -10,11 +10,12 @@
 #include "BoundShaderStateCache.h"
 #include "VulkanShaderResources.h"
 #include "VulkanState.h"
+#include "VulkanMemory.h"
 
-struct FVulkanDevice;
-struct FVulkanQueue;
+class FVulkanDevice;
+class FVulkanQueue;
 class FVulkanDescriptorPool;
-struct FVulkanDescriptorSetsLayout;
+class FVulkanDescriptorSetsLayout;
 struct FVulkanDescriptorSets;
 class FVulkanCmdBuffer;
 class FVulkanGlobalUniformPool;
@@ -27,7 +28,7 @@ struct FVulkanBufferView;
 namespace VulkanRHI
 {
 	class FDeviceMemoryAllocation;
-	class FResourceAllocationInfo;
+	class FResourceAllocation;
 }
 
 /** This represents a vertex declaration that hasn't been combined with a specific shader to create a bound shader. */
@@ -42,8 +43,9 @@ public:
 };
 
 
-struct FVulkanShader
+class FVulkanShader
 {
+public:
 	FVulkanShader(FVulkanDevice* InDevice) :
 		ShaderModule(VK_NULL_HANDLE),
 		Code(nullptr),
@@ -226,8 +228,13 @@ public:
 	uint8 FormatKey;
 private:
 
+#if VULKAN_USE_NEW_COMMAND_BUFFERS
+	// Used to clear render-target objects on creation
+	void Clear(const FClearValueBinding& ClearValueBinding, bool bTransitionToPresentable);
+#else
 	// Used to clear render-target objects on creation
-	void ClearBlocking(const FClearValueBinding& ClearValueBinding);
+	void ClearBlocking(const FClearValueBinding& ClearValueBinding, bool bTransitionToPresentable);
+#endif
 
 private:
 	// Linear or Optimal. Based on tiling, map/unmap is handled differently.
@@ -246,12 +253,7 @@ private:
 	VkImageViewType	ViewType;	// 2D, CUBE, 2DArray and etc..
 
 	bool bIsImageOwner;
-#if VULKAN_USE_MEMORY_SYSTEM
 	VulkanRHI::FDeviceMemoryAllocation* Allocation;
-#else
-	VkDeviceMemory DeviceMemory;
-#endif
-
 
 	// Used for optimal tiling mode
 	// Intermediate buffer, used to copy image data into the optimal image-buffer.
@@ -593,22 +595,14 @@ public:
 	void Unlock();
 
 	// By not providing a pending state, the copy is occuring immediately and is blocking..
-	void CopyTo(FVulkanSurface& Surface, const VkBufferImageCopy& CopyDesc, struct FVulkanPendingState* State = nullptr);
+	void CopyTo(FVulkanSurface& Surface, const VkBufferImageCopy& CopyDesc, class FVulkanPendingState* State = nullptr);
 
 	inline const VkFlags& GetFlags() const { return Usage; }
 
 private:
 	FVulkanDevice& Device;
 	VkBuffer Buf;
-#if VULKAN_USE_MEMORY_SYSTEM
-	#if VULKAN_USE_BUFFER_HEAPS
-		VulkanRHI::FResourceAllocationInfo* MemoryInfo;
-	#else
-		VulkanRHI::FDeviceMemoryAllocation* Allocation;
-	#endif
-#else
-	VkDeviceMemory Mem;
-#endif
+	VulkanRHI::FDeviceMemoryAllocation* Allocation;
 	uint32 Size;
 	VkFlags Usage;
 
@@ -673,6 +667,60 @@ protected:
 	uint64 RingBufferOffset;
 };
 
+
+#if VULKAN_USE_NEW_RESOURCE_MANAGEMENT
+
+class FVulkanBuffer2 : public VulkanRHI::FDeviceChild, public VulkanRHI::FRefCount
+{
+public:
+	FVulkanBuffer2(FVulkanDevice* InDevice, VkBuffer InBuffer, VulkanRHI::FResourceAllocation* InResourceAllocation);
+	virtual ~FVulkanBuffer2();
+
+	inline VkBuffer GetBuffer() const
+	{
+		return Buffer;
+	}
+
+	inline VulkanRHI::FResourceAllocation* GetResourceAllocation()
+	{
+		return ResourceAllocation;
+	}
+
+	static VulkanRHI::FResourceAllocation* FVulkanBuffer2::CreateResourceAllocationAndBuffer(FVulkanDevice* Device, VkDeviceSize Size, uint32 InUsage, VkBufferUsageFlags BufferUsageFlags, VkBuffer& OutBuffer);
+
+	static VkBuffer CreateVulkanBuffer(FVulkanDevice* Device, VkDeviceSize Size, VkBufferUsageFlags BufferUsageFlags, VkMemoryRequirements& OutMemoryRequirements);
+
+protected:
+	VkBuffer Buffer;
+	TRefCountPtr<VulkanRHI::FResourceAllocation> ResourceAllocation;
+};
+
+class FVulkanIndexBuffer : public FRHIIndexBuffer/*, public FVulkanBuffer2*/
+{
+public:
+	FVulkanIndexBuffer(FVulkanDevice* InDevice, VkBuffer InBuffer, VulkanRHI::FResourceAllocation* InResourceAllocation, uint32 InStride, uint32 InSize, uint32 InUsage);
+
+	inline VkIndexType GetIndexType() const
+	{
+		return IndexType;
+	}
+
+	inline FVulkanBuffer2* GetBuffer()
+	{
+		return Buffer;
+	}
+
+	inline VulkanRHI::FResourceAllocation* GetResourceAllocation()
+	{
+		return Buffer->GetResourceAllocation();
+	}
+
+private:
+	VkIndexType IndexType;
+	TRefCountPtr<FVulkanBuffer2> Buffer;
+	friend class FVulkanCommandListContext;
+};
+#else
 /** Index buffer resource class that stores stride information. */
 class FVulkanIndexBuffer : public FRHIIndexBuffer, public FVulkanMultiBuffer
 {
@@ -681,6 +729,7 @@ public:
 
 	VkIndexType IndexType;
 };
+#endif
 
 struct FVulkanDynamicLockInfo
 {
@@ -739,6 +788,8 @@ protected:
 	TArray<FPoolElement> Elements;
 };
 
+#if 0//VULKAN_USE_NEW_RESOURCE_MANAGEMENT
+#else
 /** Vertex buffer resource class that stores usage type. */
 class FVulkanVertexBuffer : public FRHIVertexBuffer, public FVulkanMultiBuffer
 {
@@ -750,6 +801,7 @@ public:
 private:
 	VkFlags Flags;
 };
+#endif
 
 class FVulkanUniformBuffer : public FRHIUniformBuffer
 {
diff --git a/Engine/Source/Runtime/VulkanRHI/Public/VulkanShaderResources.h b/Engine/Source/Runtime/VulkanRHI/Public/VulkanShaderResources.h
--- a/Engine/Source/Runtime/VulkanRHI/Public/VulkanShaderResources.h
+++ b/Engine/Source/Runtime/VulkanRHI/Public/VulkanShaderResources.h
@@ -9,8 +9,9 @@
 #include "BoundShaderStateCache.h"
 #include "CrossCompilerCommon.h"
 
-struct FVulkanShaderVarying
+class FVulkanShaderVarying
 {
+public:
 	FVulkanShaderVarying():
 		Location(0),
 		Components(0)
@@ -60,8 +61,9 @@ static void ClearBindings(CrossCompiler::FShaderBindings& Bindings)
 	Bindings.bHasRegularUniformBuffers = 0;
 }
 
-struct FVulkanShaderSerializedBindings : public CrossCompiler::FShaderBindings
+class FVulkanShaderSerializedBindings : public CrossCompiler::FShaderBindings
 {
+public:
 	FVulkanShaderSerializedBindings():
 		bFlattenUB(false)
 	{
@@ -175,8 +177,9 @@ inline FArchive& operator<<(FArchive& Ar, FVulkanCodeHeader& Header)
 }
 
 
-struct FVulkanShaderBindingTable
+class FVulkanShaderBindingTable
 {
+public:
 	TArray<uint32> SamplerBindingIndices;
 	//#todo-rco: FIX! SamplerBuffers share numbering with Samplers
 	TArray<uint32> SamplerBufferBindingIndices;
diff --git a/Engine/Source/Runtime/VulkanRHI/Public/VulkanState.h b/Engine/Source/Runtime/VulkanRHI/Public/VulkanState.h
--- a/Engine/Source/Runtime/VulkanRHI/Public/VulkanState.h
+++ b/Engine/Source/Runtime/VulkanRHI/Public/VulkanState.h
@@ -6,7 +6,7 @@
 
 #pragma once
 
-struct FVulkanDevice;
+class FVulkanDevice;
 
 class FVulkanSamplerState : public FRHISamplerState
 {
diff --git a/Engine/Source/Runtime/VulkanRHI/Public/VulkanViewport.h b/Engine/Source/Runtime/VulkanRHI/Public/VulkanViewport.h
--- a/Engine/Source/Runtime/VulkanRHI/Public/VulkanViewport.h
+++ b/Engine/Source/Runtime/VulkanRHI/Public/VulkanViewport.h
@@ -9,8 +9,8 @@
 #include "VulkanResources.h"
 
 class FVulkanDynamicRHI;
-struct FVulkanSwapChain;
-struct FVulkanQueue;
+class FVulkanSwapChain;
+class FVulkanQueue;
 
 class FVulkanViewport : public FRHIViewport
 {
@@ -21,7 +21,7 @@ public:
 	~FVulkanViewport();
 
 	FVulkanTexture2D* GetBackBuffer() const { return BackBuffers[CurrentBackBuffer]; }
-	uint32 GetBackBufferIndex() const { return CurrentBackBuffer; }
+	int32 GetBackBufferIndex() const { return CurrentBackBuffer; }
 
 	//bool Present(VkImage& Image, VkQueue& Queue, bool bLockToVsync);
 
@@ -33,6 +33,18 @@ public:
 
 	FIntPoint GetSizeXY() const { return FIntPoint(SizeX, SizeY); }
 
+#if VULKAN_USE_NEW_COMMAND_BUFFERS
+	FVulkanSwapChain* GetSwapChain()
+	{
+		return SwapChain;
+	}
+
+	void AcquireBackBuffer(FVulkanCmdBuffer* CmdBuffer);
+	void PrepareBackBufferForPresent(FVulkanCmdBuffer* CmdBuffer);
+#else
+#endif
+
+
 protected:
 	TRefCountPtr<FVulkanTexture2D> BackBuffers[NUM_BUFFERS];
 	FVulkanDynamicRHI* RHI;
@@ -40,7 +52,7 @@ protected:
 	uint32 SizeY;
 	bool bIsFullscreen;
 	EPixelFormat PixelFormat;
-	int CurrentBackBuffer;
+	int32 CurrentBackBuffer;
 	FVulkanSwapChain* SwapChain;
 
 	friend class FVulkanDynamicRHI;
