Copying //UE4/Release-Staging-4.12 to //UE4/Dev-Main (Source: //UE4/Release-4.12 @ 2955635)

diff --git a/Engine/Source/Runtime/Slate/Public/Widgets/Text/ISlateEditableTextWidget.h b/Engine/Source/Runtime/Slate/Public/Widgets/Text/ISlateEditableTextWidget.h
--- a/Engine/Source/Runtime/Slate/Public/Widgets/Text/ISlateEditableTextWidget.h
+++ b/Engine/Source/Runtime/Slate/Public/Widgets/Text/ISlateEditableTextWidget.h
@@ -182,6 +182,9 @@ public:
 	/** Is the text edit multi-line aware? */
 	virtual bool IsMultiLineTextEdit() const = 0;
 
+	/** Should the cursor be jumped to the end of the document when the widget gains focus? */
+	virtual bool ShouldJumpCursorToEndWhenFocused() const = 0;
+
 	/** Should the text be selected when the widget gains focus? */
 	virtual bool ShouldSelectAllTextWhenFocused() const = 0;
 
diff --git a/Engine/Source/Runtime/Slate/Public/Widgets/Text/SMultiLineEditableText.h b/Engine/Source/Runtime/Slate/Public/Widgets/Text/SMultiLineEditableText.h
--- a/Engine/Source/Runtime/Slate/Public/Widgets/Text/SMultiLineEditableText.h
+++ b/Engine/Source/Runtime/Slate/Public/Widgets/Text/SMultiLineEditableText.h
@@ -297,6 +297,7 @@ public:
 
 protected:
 	//~ Begin ISlateEditableTextWidget Interface
+	virtual bool ShouldJumpCursorToEndWhenFocused() const override;
 	virtual bool ShouldSelectAllTextWhenFocused() const override;
 	virtual bool ShouldClearTextSelectionOnFocusLoss() const override;
 	virtual bool ShouldRevertTextOnEscape() const override;
diff --git a/Engine/Source/Runtime/Slate/Public/Widgets/Views/STableViewBase.h b/Engine/Source/Runtime/Slate/Public/Widgets/Views/STableViewBase.h
--- a/Engine/Source/Runtime/Slate/Public/Widgets/Views/STableViewBase.h
+++ b/Engine/Source/Runtime/Slate/Public/Widgets/Views/STableViewBase.h
@@ -91,6 +91,9 @@ public:
 	/** Set the scroll offset of this view (in items) */
 	void SetScrollOffset( const float InScrollOffset );
 
+	/** Add the scroll offset of this view (in items) */
+	void AddScrollOffset(const float InScrollOffsetDelta, bool RefreshList = false);
+
 public:
 
 	// SWidget interface
diff --git a/Engine/Source/Runtime/SlateCore/Private/Fonts/FontCacheCompositeFont.cpp b/Engine/Source/Runtime/SlateCore/Private/Fonts/FontCacheCompositeFont.cpp
--- a/Engine/Source/Runtime/SlateCore/Private/Fonts/FontCacheCompositeFont.cpp
+++ b/Engine/Source/Runtime/SlateCore/Private/Fonts/FontCacheCompositeFont.cpp
@@ -181,29 +181,41 @@ const FFontData& FCompositeFontCache::GetFontDataForCharacter(const FSlateFontIn
 {
 	static const FFontData DummyFontData;
 
+	auto FontDataHasCharacter = [&](const FFontData& InFontData) -> bool
+	{
+#if WITH_FREETYPE
+		TSharedPtr<FFreeTypeFace> FaceAndMemory = GetFontFace(InFontData);
+		return FaceAndMemory.IsValid() && FT_Get_Char_Index(FaceAndMemory->GetFace(), InChar) != 0;
+#else  // WITH_FREETYPE
+		return false;
+#endif // WITH_FREETYPE
+	};
+
 	const FCompositeFont* const ResolvedCompositeFont = InFontInfo.GetCompositeFont();
 	const FCachedTypefaceData* const CachedTypefaceData = GetCachedTypefaceForCharacter(ResolvedCompositeFont, InChar);
 	if (CachedTypefaceData)
 	{
 		OutScalingFactor = CachedTypefaceData->GetScalingFactor();
 
+		const FCachedTypefaceData* const CachedDefaultTypefaceData = GetDefaultCachedTypeface(ResolvedCompositeFont);
+		const bool bIsDefaultTypeface = CachedTypefaceData == CachedDefaultTypefaceData;
+
 		// Try to find the correct font from the typeface
 		const FFontData* FoundFontData = CachedTypefaceData->GetFontData(InFontInfo.TypefaceFontName);
-		if (FoundFontData)
+		if (FoundFontData && (bIsDefaultTypeface || FontDataHasCharacter(*FoundFontData)))
 		{
 			return *FoundFontData;
 		}
 
 		// Failing that, try and find a font by the attributes of the default font with the given name
-		const FCachedTypefaceData* const CachedDefaultTypefaceData = GetDefaultCachedTypeface(ResolvedCompositeFont);
-		if (CachedDefaultTypefaceData && CachedTypefaceData != CachedDefaultTypefaceData)
+		if (!bIsDefaultTypeface && CachedDefaultTypefaceData)
 		{
 			const FFontData* const FoundDefaultFontData = CachedDefaultTypefaceData->GetFontData(InFontInfo.TypefaceFontName);
 			if (FoundDefaultFontData)
 			{
 				const TSet<FName>& DefaultFontAttributes = GetFontAttributes(*FoundDefaultFontData);
 				FoundFontData = GetBestMatchFontForAttributes(CachedTypefaceData, DefaultFontAttributes);
-				if (FoundFontData)
+				if (FoundFontData && FontDataHasCharacter(*FoundFontData))
 				{
 					return *FoundFontData;
 				}
@@ -212,10 +224,30 @@ const FFontData& FCompositeFontCache::GetFontDataForCharacter(const FSlateFontIn
 
 		// Failing that, return the first font available (the "None" font)
 		FoundFontData = CachedTypefaceData->GetFontData(NAME_None);
-		if (FoundFontData)
+		if (FoundFontData && (bIsDefaultTypeface || FontDataHasCharacter(*FoundFontData)))
 		{
 			return *FoundFontData;
 		}
+
+		// Failing that, try again using the default font (as the sub-font may not have actually supported the character we needed)
+		if (!bIsDefaultTypeface && CachedDefaultTypefaceData)
+		{
+			OutScalingFactor = CachedDefaultTypefaceData->GetScalingFactor();
+
+			// Try to find the correct font from the typeface
+			FoundFontData = CachedDefaultTypefaceData->GetFontData(InFontInfo.TypefaceFontName);
+			if (FoundFontData)
+			{
+				return *FoundFontData;
+			}
+
+			// Failing that, return the first font available (the "None" font)
+			FoundFontData = CachedDefaultTypefaceData->GetFontData(NAME_None);
+			if (FoundFontData)
+			{
+				return *FoundFontData;
+			}
+		}
 	}
 
 	OutScalingFactor = 1.0f;
diff --git a/Engine/Source/Runtime/SlateCore/Private/Styling/CoreStyle.cpp b/Engine/Source/Runtime/SlateCore/Private/Styling/CoreStyle.cpp
--- a/Engine/Source/Runtime/SlateCore/Private/Styling/CoreStyle.cpp
+++ b/Engine/Source/Runtime/SlateCore/Private/Styling/CoreStyle.cpp
@@ -583,7 +583,8 @@ TSharedRef<ISlateStyle> FCoreStyle::Create( const FName& InStyleSetName )
 			.SetNormalBarImage(FSlateColorBrush(FColor::White))
 			.SetDisabledBarImage(FSlateColorBrush(FLinearColor::Gray))
 			.SetNormalThumbImage( IMAGE_BRUSH( "Common/Button", FVector2D(8.0f, 14.0f) ) )
-			.SetDisabledThumbImage( IMAGE_BRUSH( "Common/Button_Disabled", FVector2D(8.0f, 14.0f) ) );
+			.SetDisabledThumbImage( IMAGE_BRUSH( "Common/Button_Disabled", FVector2D(8.0f, 14.0f) ) )
+			.SetBarThickness(2.0f);
 		Style->Set( "Slider", SliderStyle );
 
 		Style->Set( "VolumeControl", FVolumeControlStyle()
diff --git a/Engine/Source/Runtime/SlateCore/Private/Styling/SlateTypes.cpp b/Engine/Source/Runtime/SlateCore/Private/Styling/SlateTypes.cpp
--- a/Engine/Source/Runtime/SlateCore/Private/Styling/SlateTypes.cpp
+++ b/Engine/Source/Runtime/SlateCore/Private/Styling/SlateTypes.cpp
@@ -372,6 +372,7 @@ const FSearchBoxStyle& FSearchBoxStyle::GetDefault()
 }
 
 FSliderStyle::FSliderStyle()
+	: BarThickness(2.0f)
 {
 }
 
diff --git a/Engine/Source/Runtime/SlateCore/Public/Styling/SlateTypes.h b/Engine/Source/Runtime/SlateCore/Public/Styling/SlateTypes.h
--- a/Engine/Source/Runtime/SlateCore/Public/Styling/SlateTypes.h
+++ b/Engine/Source/Runtime/SlateCore/Public/Styling/SlateTypes.h
@@ -869,6 +869,10 @@ struct SLATECORE_API FSliderStyle : public FSlateWidgetStyle
 	UPROPERTY(EditAnywhere, BlueprintReadWrite, Category=Appearance)
 	FSlateBrush DisabledThumbImage;
 	FSliderStyle& SetDisabledThumbImage( const FSlateBrush& InDisabledThumbImage ){ DisabledThumbImage = InDisabledThumbImage; return *this; }
+
+	UPROPERTY(EditAnywhere, BlueprintReadWrite, Category=Appearance)
+	float BarThickness;
+	FSliderStyle& SetBarThickness(float InBarThickness) { BarThickness = InBarThickness; return *this; }
 };
 
 
diff --git a/Engine/Source/Runtime/VulkanRHI/Private/VulkanCommandBuffer.cpp b/Engine/Source/Runtime/VulkanRHI/Private/VulkanCommandBuffer.cpp
--- a/Engine/Source/Runtime/VulkanRHI/Private/VulkanCommandBuffer.cpp
+++ b/Engine/Source/Runtime/VulkanRHI/Private/VulkanCommandBuffer.cpp
@@ -11,17 +11,10 @@ FVulkanCmdBuffer::FVulkanCmdBuffer(FVulkanDevice* InDevice, FVulkanCommandBuffer
 	: Device(InDevice)
 	, CommandBufferManager(InCommandBufferManager)
 	, CommandBufferHandle(VK_NULL_HANDLE)
-#if VULKAN_USE_NEW_COMMAND_BUFFERS
 	, State(EState::ReadyForBegin)
 	, Fence(nullptr)
 	, FenceSignaledCounter(0)
-#else
-	, IsWriting(VK_FALSE)
-	, IsEmpty(VK_TRUE)
-#endif
 {
-	check(Device);
-
 	VkCommandBufferAllocateInfo CreateCmdBufInfo;
 	FMemory::Memzero(CreateCmdBufInfo);
 	CreateCmdBufInfo.sType = VK_STRUCTURE_TYPE_COMMAND_BUFFER_ALLOCATE_INFO;
@@ -31,27 +24,18 @@ FVulkanCmdBuffer::FVulkanCmdBuffer(FVulkanDevice* InDevice, FVulkanCommandBuffer
 	CreateCmdBufInfo.commandPool = CommandBufferManager->GetHandle();
 
 	VERIFYVULKANRESULT(vkAllocateCommandBuffers(Device->GetInstanceHandle(), &CreateCmdBufInfo, &CommandBufferHandle));
-#if VULKAN_USE_NEW_COMMAND_BUFFERS
 	Fence = Device->GetFenceManager().AllocateFence();
-#endif
 }
 
 FVulkanCmdBuffer::~FVulkanCmdBuffer()
 {
-#if VULKAN_USE_NEW_COMMAND_BUFFERS
 	auto& FenceManager = Device->GetFenceManager();
 	FenceManager.WaitAndReleaseFence(Fence, 0xffffffff);
-#else
-	check(IsWriting == VK_FALSE);
-	if (CommandBufferHandle != VK_NULL_HANDLE)
-#endif
-	{
-		vkFreeCommandBuffers(Device->GetInstanceHandle(), CommandBufferManager->GetHandle(), 1, &CommandBufferHandle);
-		CommandBufferHandle = VK_NULL_HANDLE;
-	}
+
+	vkFreeCommandBuffers(Device->GetInstanceHandle(), CommandBufferManager->GetHandle(), 1, &CommandBufferHandle);
+	CommandBufferHandle = VK_NULL_HANDLE;
 }
 
-#if VULKAN_USE_NEW_COMMAND_BUFFERS
 void FVulkanCmdBuffer::BeginRenderPass(const FVulkanRenderTargetLayout& Layout, VkRenderPass RenderPass, VkFramebuffer Framebuffer, const VkClearValue* AttachmentClearValues)
 {
 	check(IsOutsideRenderPass());
@@ -71,24 +55,11 @@ void FVulkanCmdBuffer::BeginRenderPass(const FVulkanRenderTargetLayout& Layout,
 
 	State = EState::IsInsideRenderPass;
 }
-#else
-VkCommandBuffer& FVulkanCmdBuffer::GetHandle(const VkBool32 WritingToCommandBuffer /* = true */)
-{
-	check(CommandBufferHandle != VK_NULL_HANDLE);
-	IsEmpty &= !WritingToCommandBuffer;
-	return CommandBufferHandle;
-}
-#endif
-
 
 void FVulkanCmdBuffer::Begin()
 {
-#if VULKAN_USE_NEW_COMMAND_BUFFERS
 	check(State == EState::ReadyForBegin);
-#else
-	check(!IsWriting);
-	check(IsEmpty);
-#endif
+
 	VkCommandBufferBeginInfo CmdBufBeginInfo;
 	FMemory::Memzero(CmdBufBeginInfo);
 	CmdBufBeginInfo.sType = VK_STRUCTURE_TYPE_COMMAND_BUFFER_BEGIN_INFO;
@@ -96,14 +67,9 @@ void FVulkanCmdBuffer::Begin()
 
 	VERIFYVULKANRESULT(vkBeginCommandBuffer(CommandBufferHandle, &CmdBufBeginInfo));
 
-#if VULKAN_USE_NEW_COMMAND_BUFFERS
 	State = EState::IsInsideBegin;
-#else
-	IsWriting = VK_TRUE;
-#endif
 }
 
-#if VULKAN_USE_NEW_COMMAND_BUFFERS
 void FVulkanCmdBuffer::RefreshFenceStatus()
 {
 	if (State == EState::Submitted)
@@ -121,44 +87,13 @@ void FVulkanCmdBuffer::RefreshFenceStatus()
 		check(!Fence->IsSignaled());
 	}
 }
-#else
-VkBool32 FVulkanCmdBuffer::GetIsWriting() const
-{
-	return IsWriting;
-}
 
-VkBool32 FVulkanCmdBuffer::GetIsEmpty() const
-{
-	return IsEmpty;
-}
-
-void FVulkanCmdBuffer::Reset()
-{
-	check(!IsWriting)
-
-	VERIFYVULKANRESULT(vkResetCommandBuffer(GetHandle(), VK_COMMAND_BUFFER_RESET_RELEASE_RESOURCES_BIT));
-
-	IsEmpty = VK_TRUE;
-}
-
-void FVulkanCmdBuffer::End()
-{
-	check(IsWriting)
-
-	VERIFYVULKANRESULT(vkEndCommandBuffer(GetHandle()));
-
-	IsWriting = VK_FALSE;
-	IsEmpty = VK_TRUE;
-}
-#endif
 
 FVulkanCommandBufferManager::FVulkanCommandBufferManager(FVulkanDevice* InDevice)
 	: Device(InDevice)
 	, Handle(VK_NULL_HANDLE)
-#if VULKAN_USE_NEW_COMMAND_BUFFERS
 	, ActiveCmdBuffer(nullptr)
-#else
-#endif
+	, UploadCmdBuffer(nullptr)
 {
 	check(Device);
 
@@ -170,11 +105,8 @@ FVulkanCommandBufferManager::FVulkanCommandBufferManager(FVulkanDevice* InDevice
 	CmdPoolInfo.flags = VK_COMMAND_POOL_CREATE_RESET_COMMAND_BUFFER_BIT;
 	VERIFYVULKANRESULT(vkCreateCommandPool(Device->GetInstanceHandle(), &CmdPoolInfo, nullptr, &Handle));
 
-#if VULKAN_USE_NEW_COMMAND_BUFFERS
 	ActiveCmdBuffer = Create();
 	ActiveCmdBuffer->Begin();
-#else
-#endif
 }
 
 FVulkanCommandBufferManager::~FVulkanCommandBufferManager()
@@ -188,6 +120,19 @@ FVulkanCommandBufferManager::~FVulkanCommandBufferManager()
 	vkDestroyCommandPool(Device->GetInstanceHandle(), Handle, nullptr);
 }
 
+FVulkanCmdBuffer* FVulkanCommandBufferManager::GetActiveCmdBuffer()
+{
+	if (UploadCmdBuffer)
+	{
+		check(UploadCmdBuffer->IsOutsideRenderPass());
+		UploadCmdBuffer->End();
+		Device->GetQueue()->Submit(UploadCmdBuffer);
+		UploadCmdBuffer = nullptr;
+	}
+
+	return ActiveCmdBuffer;
+}
+
 FVulkanCmdBuffer* FVulkanCommandBufferManager::Create()
 {
 	check(Device);
@@ -197,7 +142,6 @@ FVulkanCmdBuffer* FVulkanCommandBufferManager::Create()
 	return CmdBuffer;
 }
 
-#if VULKAN_USE_NEW_COMMAND_BUFFERS
 void FVulkanCommandBufferManager::RefreshFenceStatus()
 {
 	for (int32 Index = 0; Index < CmdBuffers.Num(); ++Index)
@@ -209,6 +153,8 @@ void FVulkanCommandBufferManager::RefreshFenceStatus()
 
 void FVulkanCommandBufferManager::PrepareForNewActiveCommandBuffer()
 {
+	check(!UploadCmdBuffer);
+
 	for (int32 Index = 0; Index < CmdBuffers.Num(); ++Index)
 	{
 		FVulkanCmdBuffer* CmdBuffer = CmdBuffers[Index];
@@ -229,22 +175,27 @@ void FVulkanCommandBufferManager::PrepareForNewActiveCommandBuffer()
 	ActiveCmdBuffer = Create();
 	ActiveCmdBuffer->Begin();
 }
-#else
-void FVulkanCommandBufferManager::Destroy(FVulkanCmdBuffer* CmdBuffer)
+
+FVulkanCmdBuffer* FVulkanCommandBufferManager::GetUploadCmdBuffer()
 {
-	CmdBuffers.Remove(CmdBuffer);
-	check(Device);
-	if (CmdBuffer->GetIsWriting())
+	if (!UploadCmdBuffer)
 	{
-		CmdBuffer->End();
+		for (int32 Index = 0; Index < CmdBuffers.Num(); ++Index)
+		{
+			FVulkanCmdBuffer* CmdBuffer = CmdBuffers[Index];
+			CmdBuffer->RefreshFenceStatus();
+			if (CmdBuffer->State == FVulkanCmdBuffer::EState::ReadyForBegin)
+			{
+				UploadCmdBuffer = CmdBuffer;
+				UploadCmdBuffer->Begin();
+				return UploadCmdBuffer;
+			}
+		}
+
+		// All cmd buffers are being executed still
+		UploadCmdBuffer = Create();
+		UploadCmdBuffer->Begin();
 	}
-	vkFreeCommandBuffers(Device->GetInstanceHandle(), GetHandle(), 1, &CmdBuffer->CommandBufferHandle);
-	CmdBuffer->CommandBufferHandle = VK_NULL_HANDLE;
-	delete CmdBuffer;
-}
 
-void FVulkanCommandBufferManager::Submit(FVulkanCmdBuffer* CmdBuffer)
-{
-	Device->GetQueue()->Submit(CmdBuffer);
+	return UploadCmdBuffer;
 }
-#endif
diff --git a/Engine/Source/Runtime/VulkanRHI/Private/VulkanCommandBuffer.h b/Engine/Source/Runtime/VulkanRHI/Private/VulkanCommandBuffer.h
--- a/Engine/Source/Runtime/VulkanRHI/Private/VulkanCommandBuffer.h
+++ b/Engine/Source/Runtime/VulkanRHI/Private/VulkanCommandBuffer.h
@@ -26,13 +26,6 @@ protected:
 	~FVulkanCmdBuffer();
 
 public:
-#if VULKAN_USE_NEW_COMMAND_BUFFERS
-/*
-	inline bool IsWritable() const
-	{
-		return State == EState::IsInsideBegin || State == EState::IsInsideRenderPass;
-	}
-*/
 	FVulkanCommandBufferManager* GetOwner()
 	{
 		return CommandBufferManager;
@@ -83,12 +76,9 @@ public:
 	{
 		return FenceSignaledCounter;
 	}
-#else
-	VkCommandBuffer& GetHandle(const VkBool32 MakeDirty = true);
-#endif
 
 	void Begin();
-#if VULKAN_USE_NEW_COMMAND_BUFFERS
+
 	enum class EState
 	{
 		ReadyForBegin,
@@ -97,27 +87,15 @@ public:
 		HasEnded,
 		Submitted,
 	};
-#else
-
-	void Reset();
 
-	void End();
-
-	VkBool32 GetIsWriting() const;
-	VkBool32 GetIsEmpty() const;
-#endif
 private:
 	FVulkanDevice* Device;
 	VkCommandBuffer CommandBufferHandle;
-#if VULKAN_USE_NEW_COMMAND_BUFFERS
 	EState State;
 	VulkanRHI::FFence* Fence;
 	uint64 FenceSignaledCounter;
+
 	void RefreshFenceStatus();
-#else
-	VkBool32 IsWriting;
-	VkBool32 IsEmpty;
-#endif
 
 	FVulkanCommandBufferManager* CommandBufferManager;
 };
@@ -129,37 +107,26 @@ public:
 
 	~FVulkanCommandBufferManager();
 
-#if VULKAN_USE_NEW_COMMAND_BUFFERS
-	FVulkanCmdBuffer* GetActiveCmdBuffer()
-	{
-		return ActiveCmdBuffer;
-	}
+	FVulkanCmdBuffer* GetActiveCmdBuffer();
+
+	FVulkanCmdBuffer* GetUploadCmdBuffer();
 
 	void RefreshFenceStatus();
 	void PrepareForNewActiveCommandBuffer();
-#else
-	FVulkanCmdBuffer* Create();
-
-	void Destroy(FVulkanCmdBuffer* CmdBuffer);
-
-	void Submit(FVulkanCmdBuffer* CmdBuffer);
 
-private:
-	friend class FVulkanCmdBuffer;
-	friend class FVulkanDynamicRHI;
-#endif
 	inline VkCommandPool GetHandle() const
 	{
 		check(Handle != VK_NULL_HANDLE);
 		return Handle;
 	}
+
 private:
 	FVulkanDevice* Device;
 	VkCommandPool Handle;
-#if VULKAN_USE_NEW_COMMAND_BUFFERS
 	FVulkanCmdBuffer* ActiveCmdBuffer;
+	FVulkanCmdBuffer* UploadCmdBuffer;
+
 	FVulkanCmdBuffer* Create();
-#else
-#endif
+
 	TArray<FVulkanCmdBuffer*> CmdBuffers;
 };
diff --git a/Engine/Source/Runtime/VulkanRHI/Private/VulkanCommands.cpp b/Engine/Source/Runtime/VulkanRHI/Private/VulkanCommands.cpp
--- a/Engine/Source/Runtime/VulkanRHI/Private/VulkanCommands.cpp
+++ b/Engine/Source/Runtime/VulkanRHI/Private/VulkanCommands.cpp
@@ -524,14 +524,9 @@ void FVulkanCommandListContext::RHIDrawPrimitive(uint32 PrimitiveType, uint32 Ba
 	if (!BSS.HasError())
 #endif
 	{
-#if VULKAN_USE_NEW_COMMAND_BUFFERS
 		auto* CmdBuffer = CommandBufferManager->GetActiveCmdBuffer();
-		State.PrepareDraw(CmdBuffer, UEToVulkanType((EPrimitiveType)PrimitiveType));
+		State.PrepareDraw(this, CmdBuffer, UEToVulkanType((EPrimitiveType)PrimitiveType));
 		vkCmdDraw(CmdBuffer->GetHandle(), NumVertices, NumInstances, BaseVertexIndex, 0);
-#else
-		State.PrepareDraw(UEToVulkanType((EPrimitiveType)PrimitiveType));
-		vkCmdDraw(Device->GetPendingState().GetCommandBuffer(), NumVertices, NumInstances, BaseVertexIndex, 0);
-#endif
 	}
 
 	//if (IsImmediate())
@@ -574,14 +569,9 @@ void FVulkanCommandListContext::RHIDrawIndexedPrimitive(FIndexBufferRHIParamRef
 	#endif
 	{
 		FVulkanIndexBuffer* IndexBuffer = ResourceCast(IndexBufferRHI);
-#if VULKAN_USE_NEW_COMMAND_BUFFERS
 		FVulkanCmdBuffer* Cmd = CommandBufferManager->GetActiveCmdBuffer();
 		VkCommandBuffer CmdBuffer = Cmd->GetHandle();
-		State.PrepareDraw(Cmd, UEToVulkanType((EPrimitiveType)PrimitiveType));
-#else
-		VkCommandBuffer CmdBuffer = State.GetCommandBuffer();
-		State.PrepareDraw(UEToVulkanType((EPrimitiveType)PrimitiveType));
-#endif
+		State.PrepareDraw(this, Cmd, UEToVulkanType((EPrimitiveType)PrimitiveType));
 
 		vkCmdBindIndexBuffer(CmdBuffer, IndexBuffer->GetHandle(), IndexBuffer->GetOffset(), IndexBuffer->GetIndexType());
 
@@ -668,20 +658,15 @@ void FVulkanCommandListContext::RHIEndDrawPrimitiveUP()
 	FVulkanPendingState& State = Device->GetPendingState();
 	FVulkanBoundShaderState& Shader = State.GetBoundShaderState();
 
-	State.SetStreamSource(0, PendingDrawPrimUPVertexAllocInfo.Buffer, PendingVertexDataStride, PendingDrawPrimUPVertexAllocInfo.Offset);
+	State.SetStreamSource(0, PendingDrawPrimUPVertexAllocInfo.GetHandle(), PendingVertexDataStride, PendingDrawPrimUPVertexAllocInfo.GetBindOffset());
 	
 	#if VULKAN_HAS_DEBUGGING_ENABLED
 		if(!Shader.HasError())
 	#endif
 	{
-#if VULKAN_USE_NEW_COMMAND_BUFFERS
 		auto* CmdBuffer = CommandBufferManager->GetActiveCmdBuffer();
-		State.PrepareDraw(CmdBuffer, UEToVulkanType((EPrimitiveType)PendingPrimitiveType));
+		State.PrepareDraw(this, CmdBuffer, UEToVulkanType((EPrimitiveType)PendingPrimitiveType));
 		vkCmdDraw(CmdBuffer->GetHandle(), PendingNumVertices, 1, PendingMinVertexIndex, 0);
-#else
-		State.PrepareDraw(UEToVulkanType((EPrimitiveType)PendingPrimitiveType));
-		vkCmdDraw(State.GetCommandBuffer(), PendingNumVertices, 1, PendingMinVertexIndex, 0);
-#endif
 	}
 
 	//if (IsImmediate())
@@ -721,22 +706,17 @@ void FVulkanCommandListContext::RHIEndDrawIndexedPrimitiveUP()
 	FVulkanPendingState& State = Device->GetPendingState();
 	FVulkanBoundShaderState& Shader = State.GetBoundShaderState();
 
-	State.SetStreamSource(0, PendingDrawPrimUPVertexAllocInfo.Buffer, PendingVertexDataStride, PendingDrawPrimUPVertexAllocInfo.Offset);
+	State.SetStreamSource(0, PendingDrawPrimUPVertexAllocInfo.GetHandle(), PendingVertexDataStride, PendingDrawPrimUPVertexAllocInfo.GetBindOffset());
 
 	#if VULKAN_HAS_DEBUGGING_ENABLED
 		if(!Shader.HasError())
 	#endif
 	{
-#if VULKAN_USE_NEW_COMMAND_BUFFERS
 		FVulkanCmdBuffer* CmdBuffer = CommandBufferManager->GetActiveCmdBuffer();
 		VkCommandBuffer Cmd = CmdBuffer->GetHandle();
-		State.PrepareDraw(CmdBuffer, UEToVulkanType((EPrimitiveType)PendingPrimitiveType));
-#else
-		VkCommandBuffer Cmd = State.GetCommandBuffer();
-		State.PrepareDraw(UEToVulkanType((EPrimitiveType)PendingPrimitiveType));
-#endif
+		State.PrepareDraw(this, CmdBuffer, UEToVulkanType((EPrimitiveType)PendingPrimitiveType));
 		uint32 NumIndices = GetVertexCountForPrimitiveCount(PendingNumPrimitives, PendingPrimitiveType);
-		vkCmdBindIndexBuffer(Cmd, PendingDrawPrimUPIndexAllocInfo.Buffer, PendingDrawPrimUPIndexAllocInfo.Offset, PendingPrimitiveIndexType);
+		vkCmdBindIndexBuffer(Cmd, PendingDrawPrimUPIndexAllocInfo.GetHandle(), PendingDrawPrimUPIndexAllocInfo.GetBindOffset(), PendingPrimitiveIndexType);
 		vkCmdDrawIndexed(Cmd, NumIndices, 1, PendingMinVertexIndex, 0, 0);
 	}
 
@@ -766,11 +746,7 @@ void FVulkanCommandListContext::RHIClearMRT(bool bClearColor, int32 NumClearColo
 	check(bClearColor ? NumClearColors > 0 : true);
 
 	FVulkanPendingState& State = Device->GetPendingState();
-#if VULKAN_USE_NEW_COMMAND_BUFFERS
 	State.UpdateRenderPass(CommandBufferManager->GetActiveCmdBuffer());
-#else
-	State.UpdateRenderPass();
-#endif
 	
 #if VULKAN_ALLOW_MIDPASS_CLEAR
 	VkClearRect Rect;
diff --git a/Engine/Source/Runtime/VulkanRHI/Private/VulkanDebug.cpp b/Engine/Source/Runtime/VulkanRHI/Private/VulkanDebug.cpp
--- a/Engine/Source/Runtime/VulkanRHI/Private/VulkanDebug.cpp
+++ b/Engine/Source/Runtime/VulkanRHI/Private/VulkanDebug.cpp
@@ -20,34 +20,39 @@ DEFINE_LOG_CATEGORY(LogVulkanRHI);
 #if VULKAN_HAS_DEBUGGING_ENABLED
 
 static VkBool32 VKAPI_PTR DebugReportFunction(
-	VkDebugReportFlagsEXT			msgFlags,
-	VkDebugReportObjectTypeEXT		objType,
-    uint64_t						srcObject,
-    size_t							location,
-    int32							msgCode,
-	const ANSICHAR*					pLayerPrefix,
-	const ANSICHAR*					pMsg,
-    void*							pUserData)
+	VkDebugReportFlagsEXT			MsgFlags,
+	VkDebugReportObjectTypeEXT		ObjType,
+    uint64_t						SrcObject,
+    size_t							Location,
+    int32							MsgCode,
+	const ANSICHAR*					LayerPrefix,
+	const ANSICHAR*					Msg,
+    void*							UserData)
 {
-	if (msgFlags != VK_DEBUG_REPORT_ERROR_BIT_EXT && 
-		msgFlags != VK_DEBUG_REPORT_WARNING_BIT_EXT &&
-		msgFlags != VK_DEBUG_REPORT_INFORMATION_BIT_EXT &&
-		msgFlags != VK_DEBUG_REPORT_DEBUG_BIT_EXT)
+	if (MsgFlags != VK_DEBUG_REPORT_ERROR_BIT_EXT && 
+		MsgFlags != VK_DEBUG_REPORT_WARNING_BIT_EXT &&
+		MsgFlags != VK_DEBUG_REPORT_INFORMATION_BIT_EXT &&
+		MsgFlags != VK_DEBUG_REPORT_DEBUG_BIT_EXT)
 	{
 		ensure(0);
 	}
 
-	if (msgFlags & VK_DEBUG_REPORT_ERROR_BIT_EXT)
+	if (MsgFlags & VK_DEBUG_REPORT_ERROR_BIT_EXT)
 	{
 		// Reaching this line should trigger a break/assert. 
 		// Check to see if this is a code we've seen before.
-		FString LayerCode = FString::Printf(TEXT("%s%d"), ANSI_TO_TCHAR(pLayerPrefix), msgCode);
+		FString LayerCode = FString::Printf(TEXT("%s%d"), ANSI_TO_TCHAR(LayerPrefix), MsgCode);
 		static TSet<FString> SeenCodes;
 		if (!SeenCodes.Contains(LayerCode))
 		{
-			FString Message = FString::Printf(TEXT("VulkanRHI: VK ERROR: [%s] Code %d : %s\n"), ANSI_TO_TCHAR(pLayerPrefix), msgCode, ANSI_TO_TCHAR(pMsg));
+			FString Message = FString::Printf(TEXT("VulkanRHI: VK ERROR: [%s] Code %d : %s\n"), ANSI_TO_TCHAR(LayerPrefix), MsgCode, ANSI_TO_TCHAR(Msg));
 			FPlatformMisc::LowLevelOutputDebugStringf(*Message);
     
+			if (!FCStringAnsi::Strcmp(LayerPrefix, "MEM"))
+			{
+				((FVulkanDynamicRHI*)GDynamicRHI)->DumpMemory();
+			}
+
 			// Break debugger on first instance of each message. 
 			// Continuing will ignore the error and suppress this message in future.
 			bool bIgnoreInFuture = true;
@@ -59,32 +64,33 @@ static VkBool32 VKAPI_PTR DebugReportFunction(
 		}
 		return VK_FALSE;
 	}
-	else if (msgFlags & VK_DEBUG_REPORT_WARNING_BIT_EXT)
+	else if (MsgFlags & VK_DEBUG_REPORT_WARNING_BIT_EXT)
 	{
-		FString Message = FString::Printf(TEXT("VulkanRHI: VK WARNING: [%s] Code %d : %s\n"), ANSI_TO_TCHAR(pLayerPrefix), msgCode, ANSI_TO_TCHAR(pMsg));
+		FString Message = FString::Printf(TEXT("VulkanRHI: VK WARNING: [%s] Code %d : %s\n"), ANSI_TO_TCHAR(LayerPrefix), MsgCode, ANSI_TO_TCHAR(Msg));
 		FPlatformMisc::LowLevelOutputDebugString(*Message);
+		return VK_FALSE;
 	}
 #if VULKAN_ENABLE_API_DUMP
-	else if (msgFlags & VK_DEBUG_REPORT_INFORMATION_BIT_EXT)
+	else if (MsgFlags & VK_DEBUG_REPORT_INFORMATION_BIT_EXT)
 	{
 #if !VULKAN_ENABLE_API_DUMP_DETAILED
-		if (!FCStringAnsi::Strcmp(pLayerPrefix, "MEM") || !FCStringAnsi::Strcmp(pLayerPrefix, "DS"))
+		if (!FCStringAnsi::Strcmp(LayerPrefix, "MEM") || !FCStringAnsi::Strcmp(LayerPrefix, "DS"))
 		{
 			// Skip Mem messages
 		}
 		else
 #endif
 		{
-			FString Message = FString::Printf(TEXT("VulkanRHI: VK INFO: [%s] Code %d : %s\n"), ANSI_TO_TCHAR(pLayerPrefix), msgCode, ANSI_TO_TCHAR(pMsg));
+			FString Message = FString::Printf(TEXT("VulkanRHI: VK INFO: [%s] Code %d : %s\n"), ANSI_TO_TCHAR(LayerPrefix), MsgCode, ANSI_TO_TCHAR(Msg));
 			FPlatformMisc::LowLevelOutputDebugString(*Message);
 		}
 
 		return VK_FALSE;
 	}
 #if VULKAN_ENABLE_API_DUMP_DETAILED
-	else if (msgFlags & VK_DEBUG_REPORT_DEBUG_BIT_EXT)
+	else if (MsgFlags & VK_DEBUG_REPORT_DEBUG_BIT_EXT)
 	{
-		FString Message = FString::Printf(TEXT("VulkanRHI: VK DEBUG: [%s] Code %d : %s\n"), ANSI_TO_TCHAR(pLayerPrefix), msgCode, ANSI_TO_TCHAR(pMsg));
+		FString Message = FString::Printf(TEXT("VulkanRHI: VK DEBUG: [%s] Code %d : %s\n"), ANSI_TO_TCHAR(LayerPrefix), MsgCode, ANSI_TO_TCHAR(Msg));
 		FPlatformMisc::LowLevelOutputDebugString(*Message);
 		return VK_FALSE;
 	}
diff --git a/Engine/Source/Runtime/VulkanRHI/Private/VulkanDevice.cpp b/Engine/Source/Runtime/VulkanRHI/Private/VulkanDevice.cpp
--- a/Engine/Source/Runtime/VulkanRHI/Private/VulkanDevice.cpp
+++ b/Engine/Source/Runtime/VulkanRHI/Private/VulkanDevice.cpp
@@ -14,9 +14,7 @@ FVulkanDevice::FVulkanDevice(VkPhysicalDevice InGpu)
 	: Gpu(InGpu)
 	, Device(VK_NULL_HANDLE)
 	, ResourceHeapManager(this)
-#if VULKAN_USE_NEW_COMMAND_BUFFERS
 	, DeferredDeletionQueue(this)
-#endif
 	, DescriptorPool(nullptr)
 	, DefaultSampler(VK_NULL_HANDLE)
 	, Queue(nullptr)
@@ -432,10 +430,16 @@ void FVulkanDevice::InitGPU(int32 DeviceIndex)
 #endif
 }
 
-void FVulkanDevice::Destroy()
+void FVulkanDevice::PrepareForDestroy()
 {
 	WaitUntilIdle();
 
+	delete ImmediateContext;
+	ImmediateContext = nullptr;
+}
+
+void FVulkanDevice::Destroy()
+{
 	if (TimestampQueryPool[0])
 	{
 		for (int32 Index = 0; Index < NumTimestampPools; ++Index)
@@ -450,9 +454,6 @@ void FVulkanDevice::Destroy()
 
 	delete UBRingBuffer;
 
-	delete ImmediateContext;
-	ImmediateContext = nullptr;
-
 #if VULKAN_ENABLE_PIPELINE_CACHE
 	delete PipelineStateCache;
 #endif
@@ -481,15 +482,6 @@ void FVulkanDevice::WaitUntilIdle()
 	VERIFYVULKANRESULT(vkDeviceWaitIdle(Device));
 }
 
-#if VULKAN_USE_NEW_COMMAND_BUFFERS
-#else
-void FVulkanDevice::EndCommandBufferBlock(FVulkanCmdBuffer* CmdBuffer)
-{
-	check(CmdBuffer);
-	GetQueue()->SubmitBlocking(CmdBuffer);
-}
-#endif
-
 bool FVulkanDevice::IsFormatSupported(VkFormat Format) const
 {
 	const VkFormatProperties& Prop = FormatProperties[Format];
diff --git a/Engine/Source/Runtime/VulkanRHI/Private/VulkanDevice.h b/Engine/Source/Runtime/VulkanRHI/Private/VulkanDevice.h
--- a/Engine/Source/Runtime/VulkanRHI/Private/VulkanDevice.h
+++ b/Engine/Source/Runtime/VulkanRHI/Private/VulkanDevice.h
@@ -30,6 +30,7 @@ public:
 
 	void CreateDevice();
 
+	void PrepareForDestroy();
 	void Destroy();
 
 	void WaitUntilIdle();
@@ -40,11 +41,6 @@ public:
 		return *PendingState;
 	}
 
-#if VULKAN_USE_NEW_COMMAND_BUFFERS
-#else
-	void EndCommandBufferBlock(FVulkanCmdBuffer* CmdBuffer);
-#endif
-
 	inline FVulkanQueue* GetQueue()
 	{
 		check(Queue);
@@ -115,12 +111,10 @@ public:
 		return ResourceHeapManager;
 	}
 
-#if VULKAN_USE_NEW_COMMAND_BUFFERS
 	VulkanRHI::FDeferredDeletionQueue& GetDeferredDeletionQueue()
 	{
 		return DeferredDeletionQueue;
 	}
-#endif
 
 	VulkanRHI::FStagingManager& GetStagingManager()
 	{
@@ -170,10 +164,11 @@ private:
 	VkDevice Device;
 
 	VulkanRHI::FDeviceMemoryManager MemoryManager;
+
 	VulkanRHI::FResourceHeapManager ResourceHeapManager;
-#if VULKAN_USE_NEW_COMMAND_BUFFERS
+
 	VulkanRHI::FDeferredDeletionQueue DeferredDeletionQueue;
-#endif
+
 	VulkanRHI::FStagingManager StagingManager;
 
 	VulkanRHI::FFenceManager FenceManager;
diff --git a/Engine/Source/Runtime/VulkanRHI/Private/VulkanIndexBuffer.cpp b/Engine/Source/Runtime/VulkanRHI/Private/VulkanIndexBuffer.cpp
--- a/Engine/Source/Runtime/VulkanRHI/Private/VulkanIndexBuffer.cpp
+++ b/Engine/Source/Runtime/VulkanRHI/Private/VulkanIndexBuffer.cpp
@@ -152,19 +152,10 @@ void FVulkanResourceMultiBuffer::Unlock()
 			uint32 LockOffset = PendingLock.Offset;
 			VulkanRHI::FStagingBuffer* StagingBuffer = PendingLock.StagingBuffer;
 			PendingLock.StagingBuffer = nullptr;
-#if VULKAN_USE_NEW_COMMAND_BUFFERS
-			FVulkanCmdBuffer* Cmd = Device->GetImmediateContext().GetCommandBufferManager()->GetActiveCmdBuffer();
+
+			FVulkanCmdBuffer* Cmd = Device->GetImmediateContext().GetCommandBufferManager()->GetUploadCmdBuffer();
 			ensure(Cmd->IsOutsideRenderPass());
 			VkCommandBuffer CmdBuffer = Cmd->GetHandle();
-#else
-			VkCommandBuffer CmdBuffer = VK_NULL_HANDLE;
-			FVulkanCmdBuffer* CmdObject = nullptr;
-
-			CmdObject = Device->GetImmediateContext().GetCommandBufferManager()->Create();
-			check(CmdObject);
-			CmdObject->Begin();
-			CmdBuffer = CmdObject->GetHandle();
-#endif
 
 			VkBufferCopy Region;
 			FMemory::Memzero(Region);
@@ -174,14 +165,8 @@ void FVulkanResourceMultiBuffer::Unlock()
 			vkCmdCopyBuffer(CmdBuffer, StagingBuffer->GetHandle(), Buffers[DynamicBufferIndex]->GetHandle(), 1, &Region);
 			//UpdateBuffer(ResourceAllocation, IndexBuffer->GetBuffer(), LockSize, LockOffset);
 
-#if VULKAN_USE_NEW_COMMAND_BUFFERS
-			Device->GetDeferredDeletionQueue().EnqueueResource(Cmd, StagingBuffer);
-#else
-			check(CmdObject);
-			Device->GetQueue()->SubmitBlocking(CmdObject);
-			Device->GetImmediateContext().GetCommandBufferManager()->Destroy(CmdObject);
-#endif
-			Device->GetStagingManager().ReleaseBuffer(StagingBuffer);
+			//Device->GetDeferredDeletionQueue().EnqueueResource(Cmd, StagingBuffer);
+			Device->GetStagingManager().ReleaseBuffer(Cmd, StagingBuffer);
 		}
 	}
 }
diff --git a/Engine/Source/Runtime/VulkanRHI/Private/VulkanMemory.cpp b/Engine/Source/Runtime/VulkanRHI/Private/VulkanMemory.cpp
--- a/Engine/Source/Runtime/VulkanRHI/Private/VulkanMemory.cpp
+++ b/Engine/Source/Runtime/VulkanRHI/Private/VulkanMemory.cpp
@@ -94,7 +94,7 @@ namespace VulkanRHI
 	{
 		for (int32 Index = 0; Index < HeapInfos.Num(); ++Index)
 		{
-			checkf(HeapInfos[Index].Allocations.Num() == 0, TEXT("Found %d unfreed allocations!"), HeapInfos[Index].Allocations.Num());
+			ensureMsgf(HeapInfos[Index].Allocations.Num() == 0, TEXT("Found %d unfreed allocations!"), HeapInfos[Index].Allocations.Num());
 		}
 		NumAllocations = 0;
 	}
@@ -172,7 +172,7 @@ namespace VulkanRHI
 			for (int32 SubIndex = 0; SubIndex < HeapInfo.Allocations.Num(); ++SubIndex)
 			{
 				FDeviceMemoryAllocation* Allocation = HeapInfo.Allocations[SubIndex];
-				//UE_LOG(LogVulkanRHI, Display, TEXT("\t\tAllocation %d Size %d"), SubIndex, Allocation->Size);
+				UE_LOG(LogVulkanRHI, Display, TEXT("\t\t%d Size %d Handle %p"), SubIndex, Allocation->Size, (void*)Allocation->Handle);
 				TotalSize += Allocation->Size;
 			}
 			UE_LOG(LogVulkanRHI, Display, TEXT("\t\tTotal Allocated %.2f MB"), TotalSize / 1024.0f / 1024.0f);
@@ -225,7 +225,7 @@ namespace VulkanRHI
 
 	void FOldResourceAllocation::BindBuffer(FVulkanDevice* Device, VkBuffer Buffer)
 	{
-		//FPlatformMisc::LowLevelOutputDebugStringf(TEXT("** vkBindBufferMemory Buf %p MemHandle %p MemOffset %d Size %u\n"), (void*)StagingBuffer->Buffer, (void*)StagingBuffer->ResourceAllocation->GetHandle(), (uint32)StagingBuffer->ResourceAllocation->GetOffset(), (uint32)MemReqs.size);
+		//FPlatformMisc::LowLevelOutputDebugStringf(TEXT("** vkBindBufferMemory Image %p MemHandle %p MemOffset %d Size %u\n"), (void*)Buffer, GetHandle(), GetOffset(), GetSize());
 		VkResult Result = vkBindBufferMemory(Device->GetInstanceHandle(), Buffer, GetHandle(), GetOffset());
 #if UE_BUILD_DEBUG || UE_BUILD_DEVELOPMENT
 		if (Result == VK_ERROR_OUT_OF_DEVICE_MEMORY || Result == VK_ERROR_OUT_OF_HOST_MEMORY)
@@ -239,7 +239,7 @@ namespace VulkanRHI
 
 	void FOldResourceAllocation::BindImage(FVulkanDevice* Device, VkImage Image)
 	{
-		//FPlatformMisc::LowLevelOutputDebugStringf(TEXT("** vkBindImageMemory Image %p MemHandle %p MemOffset %d Size %u\n"), (void*)Image, (void*)Allocation->GetHandle(), (uint32)0, (uint32)MemoryRequirements.size);
+		//FPlatformMisc::LowLevelOutputDebugStringf(TEXT("** vkBindImageMemory Image %p MemHandle %p MemOffset %d Size %u\n"), (void*)Image, GetHandle(), GetOffset(), GetSize());
 		VkResult Result = vkBindImageMemory(Device->GetInstanceHandle(), Image, GetHandle(), GetOffset());
 #if UE_BUILD_DEBUG || UE_BUILD_DEVELOPMENT
 		if (Result == VK_ERROR_OUT_OF_DEVICE_MEMORY || Result == VK_ERROR_OUT_OF_HOST_MEMORY)
@@ -467,7 +467,7 @@ namespace VulkanRHI
 				SubAllocUsedMemory += UsedPages[Index]->UsedSize;
 				NumSuballocations += UsedPages[Index]->ResourceAllocations.Num();
 
-				UE_LOG(LogVulkanRHI, Display, TEXT("\t\t%d: %4d suballocs, %4d free chunks (%d used/%d free/%d max)"), Index, UsedPages[Index]->ResourceAllocations.Num(), UsedPages[Index]->FreeList.Num(), UsedPages[Index]->UsedSize, UsedPages[Index]->MaxSize - UsedPages[Index]->UsedSize, UsedPages[Index]->MaxSize);
+				UE_LOG(LogVulkanRHI, Display, TEXT("\t\t%d: %4d suballocs, %4d free chunks (%d used/%d free/%d max) DeviceMemory %p"), Index, UsedPages[Index]->ResourceAllocations.Num(), UsedPages[Index]->FreeList.Num(), UsedPages[Index]->UsedSize, UsedPages[Index]->MaxSize - UsedPages[Index]->UsedSize, UsedPages[Index]->MaxSize, (void*)UsedPages[Index]->DeviceMemoryAllocation->GetHandle());
 			}
 
 			for (int32 Index = 0; Index < FreePages.Num(); ++Index)
@@ -498,10 +498,13 @@ namespace VulkanRHI
 			for (int32 Index = 0; Index < UsedPages.Num(); ++Index)
 			{
 				FOldResourceHeapPage* Page = UsedPages[Index];
-				FOldResourceAllocation* ResourceAllocation = Page->TryAllocate(Size, Alignment, File, Line);
-				if (ResourceAllocation)
+				if (Page->DeviceMemoryAllocation->IsMapped() == bMapAllocation)
 				{
-					return ResourceAllocation;
+					FOldResourceAllocation* ResourceAllocation = Page->TryAllocate(Size, Alignment, File, Line);
+					if (ResourceAllocation)
+					{
+						return ResourceAllocation;
+					}
 				}
 			}
 		}
@@ -509,17 +512,20 @@ namespace VulkanRHI
 		for (int32 Index = 0; Index < FreePages.Num(); ++Index)
 		{
 			FOldResourceHeapPage* Page = FreePages[Index];
-			FOldResourceAllocation* ResourceAllocation = Page->TryAllocate(Size, Alignment, File, Line);
-			if (ResourceAllocation)
+			if (Page->DeviceMemoryAllocation->IsMapped() == bMapAllocation)
 			{
-				FreePages.RemoveSingleSwap(Page, false);
-				UsedPages.Add(Page);
-				return ResourceAllocation;
+				FOldResourceAllocation* ResourceAllocation = Page->TryAllocate(Size, Alignment, File, Line);
+				if (ResourceAllocation)
+				{
+					FreePages.RemoveSingleSwap(Page, false);
+					UsedPages.Add(Page);
+					return ResourceAllocation;
+				}
 			}
 		}
 		uint32 AllocationSize = FMath::Max(Size, DefaultPageSize);
 #endif
-		FDeviceMemoryAllocation* DeviceMemoryAllocation = Owner->GetParent()->GetMemoryManager().Alloc(AllocationSize, MemoryTypeIndex, __FILE__, __LINE__);
+		FDeviceMemoryAllocation* DeviceMemoryAllocation = Owner->GetParent()->GetMemoryManager().Alloc(AllocationSize, MemoryTypeIndex, File, Line);
 		auto* NewPage = new FOldResourceHeapPage(this, DeviceMemoryAllocation);
 		UsedPages.Add(NewPage);
 
@@ -557,7 +563,7 @@ namespace VulkanRHI
 		const auto& MemoryProperties = MemoryManager.GetMemoryProperties();
 
 		//#todo-rco: Hack! Figure out a better way for page size
-		float PercentageAllocationForGPU = 0.95;
+		float PercentageAllocationForGPU = 0.8;
 
 		uint32 NumMemoryAllocations = (uint64)Device->GetLimits().maxMemoryAllocationCount;
 		uint32 NumGPUAllocations = (uint32)((float)NumMemoryAllocations * PercentageAllocationForGPU);
@@ -565,6 +571,19 @@ namespace VulkanRHI
 		ResourceHeaps.AddZeroed(MemoryProperties.memoryTypeCount);
 		PageSizes.AddZeroed(MemoryProperties.memoryTypeCount);
 
+		TArray<uint64> RemainingHeapSizes;
+		TArray<uint64> NumTypesPerHeap;
+		for (uint32 Index = 0; Index < MemoryProperties.memoryHeapCount; ++Index)
+		{
+			RemainingHeapSizes.Add(MemoryProperties.memoryHeaps[Index].size);
+			NumTypesPerHeap.Add(0);
+		}
+
+		for (uint32 Index = 0; Index < MemoryProperties.memoryTypeCount; ++Index)
+		{
+			++NumTypesPerHeap[MemoryProperties.memoryTypes[Index].heapIndex];
+		}
+
 		{
 			uint32 TypeIndex = 0;
 			VERIFYVULKANRESULT(MemoryManager.GetMemoryTypeFromProperties(TypeBits, VK_MEMORY_PROPERTY_DEVICE_LOCAL_BIT, &TypeIndex));
@@ -573,21 +592,27 @@ namespace VulkanRHI
 			GPUOnlyHeap = new FOldResourceHeap(this, TypeIndex, PageSize, HeapSize);
 			ResourceHeaps[TypeIndex] = GPUOnlyHeap;
 			PageSizes[TypeIndex] = PageSize;
+			RemainingHeapSizes[MemoryProperties.memoryTypes[TypeIndex].heapIndex] -= HeapSize;
 		}
 
+		// Allocate all remaining types
 		{
-			uint32 TypeIndex = 0;
-			VERIFYVULKANRESULT(MemoryManager.GetMemoryTypeFromProperties(TypeBits, VK_MEMORY_PROPERTY_HOST_VISIBLE_BIT, &TypeIndex));
+			uint32 NumRemainingAllocations = NumMemoryAllocations - NumGPUAllocations;
 
-			uint64 HeapSize = MemoryProperties.memoryHeaps[MemoryProperties.memoryTypes[TypeIndex].heapIndex].size;
-			uint32 NumCPUAllocations = NumMemoryAllocations - NumGPUAllocations;
-
-			// Less equal than 512MB might mean PCI express bus, so clamp for smaller pages
-			HeapSize = FMath::Min((uint64)512 * 1024 * 1024, HeapSize);
-			uint32 PageSize = HeapSize / NumCPUAllocations;
-			CPUStagingHeap = new FOldResourceHeap(this, TypeIndex, PageSize, HeapSize);
-			ResourceHeaps[TypeIndex] = CPUStagingHeap;
-			PageSizes[TypeIndex] = PageSize;
+			for (uint32 Index = 0; Index < MemoryProperties.memoryTypeCount; ++Index)
+			{
+				if (!ResourceHeaps[Index])
+				{
+					uint32 HeapIndex = MemoryProperties.memoryTypes[Index].heapIndex;
+					uint64 HeapSize = RemainingHeapSizes[HeapIndex] / NumTypesPerHeap[HeapIndex];
+					uint32 PageSize = HeapSize / NumRemainingAllocations;
+					ResourceHeaps[Index] = new FOldResourceHeap(this, Index, PageSize, HeapSize);
+					PageSizes[Index] = PageSize;
+
+					// Always grab the last one for now...
+					CPUStagingHeap = ResourceHeaps[Index];
+				}
+			}
 		}
 	}
 
@@ -595,11 +620,10 @@ namespace VulkanRHI
 	{
 		DestroyResourceAllocations();
 
-		delete GPUOnlyHeap;
-		GPUOnlyHeap = nullptr;
-
-		delete CPUStagingHeap;
-		CPUStagingHeap = nullptr;
+		for (int32 Index = 0; Index < ResourceHeaps.Num(); ++Index)
+		{
+			delete ResourceHeaps[Index];
+		}
 	}
 
 	void FResourceHeapManager::DestroyResourceAllocations()
@@ -825,19 +849,20 @@ namespace VulkanRHI
 	{
 		FScopeLock ScopeLock(&CS);
 
-		UE_LOG(LogVulkanRHI, Display, TEXT("GPU Only Heap, Memory Type Index %d"), GPUOnlyHeap->MemoryTypeIndex);
-		GPUOnlyHeap->DumpMemory();
-		UE_LOG(LogVulkanRHI, Display, TEXT("CPU Staging Heap, Memory Type Index %d"), CPUStagingHeap->MemoryTypeIndex);
-		CPUStagingHeap->DumpMemory();
+		for (int32 Index = 0; Index < ResourceHeaps.Num(); ++Index)
+		{
+			UE_LOG(LogVulkanRHI, Display, TEXT("Heap %d, Memory Type Index %d"), Index, ResourceHeaps[Index]->MemoryTypeIndex);
+			ResourceHeaps[Index]->DumpMemory();
+		}
 
 		UE_LOG(LogVulkanRHI, Display, TEXT("Buffer Allocations: %d Used / %d Free"), UsedBufferAllocations.Num(), FreeBufferAllocations.Num());
 		if (UsedBufferAllocations.Num() > 0)
 		{
-			UE_LOG(LogVulkanRHI, Display, TEXT("Index  MemFlags BufferFlags #Suballocs #FreeChunks UsedSize/MaxSize"));
+			UE_LOG(LogVulkanRHI, Display, TEXT("Index  BufferHandle   DeviceMemoryHandle MemFlags BufferFlags #Suballocs #FreeChunks UsedSize/MaxSize"));
 			for (int32 Index = 0; Index < UsedBufferAllocations.Num(); ++Index)
 			{
 				FBufferAllocation* BA = UsedBufferAllocations[Index];
-				UE_LOG(LogVulkanRHI, Display, TEXT("%6d 0x%06x 0x%08x %6d   %6d    %d/%d"), Index, BA->MemoryPropertyFlags, BA->BufferUsageFlags, BA->Suballocations.Num(), BA->FreeList.Num(), BA->UsedSize, BA->MaxSize);
+				UE_LOG(LogVulkanRHI, Display, TEXT("%6d %p %p 0x%06x 0x%08x %6d   %6d    %d/%d"), Index, (void*)BA->Buffer, (void*)BA->MemoryAllocation->GetHandle(), BA->MemoryPropertyFlags, BA->BufferUsageFlags, BA->Suballocations.Num(), BA->FreeList.Num(), BA->UsedSize, BA->MaxSize);
 			}
 		}
 	}
@@ -958,6 +983,7 @@ namespace VulkanRHI
 		check(ResourceAllocation);
 
 		vkDestroyBuffer(DeviceHandle, Buffer, nullptr);
+		Buffer = VK_NULL_HANDLE;
 		ResourceAllocation = nullptr;
 		//Memory.Free(Allocation);
 	}
@@ -965,9 +991,7 @@ namespace VulkanRHI
 	FStagingManager::~FStagingManager()
 	{
 		check(UsedStagingBuffers.Num() == 0);
-#if VULKAN_USE_NEW_COMMAND_BUFFERS
 		check(PendingFreeStagingBuffers.Num() == 0);
-#endif
 		check(FreeStagingBuffers.Num() == 0);
 	}
 	
@@ -1034,7 +1058,6 @@ namespace VulkanRHI
 		return StagingBuffer;
 	}
 
-#if VULKAN_USE_NEW_COMMAND_BUFFERS
 	void FStagingManager::ReleaseBuffer(FVulkanCmdBuffer* CmdBuffer, FStagingBuffer*& StagingBuffer)
 	{
 		FScopeLock Lock(&GAllocationLock);
@@ -1045,19 +1068,37 @@ namespace VulkanRHI
 		Entry->Resource = StagingBuffer;
 		StagingBuffer = nullptr;
 	}
-#else
-	void FStagingManager::ReleaseBuffer(FStagingBuffer*& StagingBuffer)
+
+#if UE_BUILD_DEBUG || UE_BUILD_DEVELOPMENT
+	void FStagingManager::DumpMemory()
 	{
-		FScopeLock Lock(&GAllocationLock);
-		UsedStagingBuffers.RemoveSingleSwap(StagingBuffer);
-		FreeStagingBuffers.Add(StagingBuffer);
-		StagingBuffer = nullptr;
+		UE_LOG(LogVulkanRHI, Display, TEXT("StagingManager %d Used %d Pending Free %d Free"), UsedStagingBuffers.Num(), PendingFreeStagingBuffers.Num(), FreeStagingBuffers.Num());
+		UE_LOG(LogVulkanRHI, Display, TEXT("Used   BufferHandle ResourceAllocation"));
+		for (int32 Index = 0; Index < UsedStagingBuffers.Num(); ++Index)
+		{
+			FStagingBuffer* Buffer = UsedStagingBuffers[Index];
+			UE_LOG(LogVulkanRHI, Display, TEXT("%6d %p %p"), Index, (void*)Buffer->GetHandle(), (void*)Buffer->ResourceAllocation->GetHandle());
+		}
+
+		UE_LOG(LogVulkanRHI, Display, TEXT("Pending CmdBuffer   Fence   BufferHandle ResourceAllocation"));
+		for (int32 Index = 0; Index < PendingFreeStagingBuffers.Num(); ++Index)
+		{
+			FPendingItem& Item = PendingFreeStagingBuffers[Index];
+			FStagingBuffer* Buffer = (FStagingBuffer*)Item.Resource;
+			UE_LOG(LogVulkanRHI, Display, TEXT("%6d %p %p %p %p"), Index, (void*)Item.CmdBuffer->GetHandle(), (void*)Item.FenceCounter, (void*)Buffer->GetHandle(), (void*)Buffer->ResourceAllocation->GetHandle());
+		}
+
+		UE_LOG(LogVulkanRHI, Display, TEXT("Free   BufferHandle ResourceAllocation"));
+		for (int32 Index = 0; Index < FreeStagingBuffers.Num(); ++Index)
+		{
+			FStagingBuffer* Buffer = FreeStagingBuffers[Index];
+			UE_LOG(LogVulkanRHI, Display, TEXT("%6d %p %p"), Index, (void*)Buffer->GetHandle(), (void*)Buffer->ResourceAllocation->GetHandle());
+		}
 	}
 #endif
 
 	void FStagingManager::ProcessPendingFree(bool bImmediately)
 	{
-#if VULKAN_USE_NEW_COMMAND_BUFFERS
 		FScopeLock Lock(&GAllocationLock);
 		for (int32 Index = PendingFreeStagingBuffers.Num() - 1; Index >= 0; --Index)
 		{
@@ -1071,20 +1112,6 @@ namespace VulkanRHI
 				//#todo-rco: Release buffers back to OS...
 			}
 		}
-
-		for (int32 Index = PendingFreeStagingImages.Num() - 1; Index >= 0; --Index)
-		{
-			FPendingItem& Entry = PendingFreeStagingImages[Index];
-			if (bImmediately || !Entry.CmdBuffer || Entry.FenceCounter < Entry.CmdBuffer->GetFenceSignaledCounter())
-			{
-				auto* StagingImage = (FStagingImage*)Entry.Resource;
-				PendingFreeStagingImages.RemoveAtSwap(Index, 1, false);
-				FreeStagingImages.Add(StagingImage);
-
-				//#todo-rco: Release images back to OS...
-			}
-		}
-#endif
 	}
 
 	FFence::FFence(FVulkanDevice* InDevice, FFenceManager* InOwner)
@@ -1228,7 +1255,6 @@ namespace VulkanRHI
 
 	void FDeferredDeletionQueue::EnqueueResource(FVulkanCmdBuffer* CmdBuffer, FRefCount* Resource)
 	{
-#if VULKAN_USE_NEW_COMMAND_BUFFERS
 		Resource->AddRef();
 
 		FEntry Entry;
@@ -1239,15 +1265,13 @@ namespace VulkanRHI
 			FScopeLock ScopeLock(&CS);
 			Entries.Add(Entry);
 		}
-#endif
 	}
 
 	void FDeferredDeletionQueue::ReleaseResources(/*bool bDeleteImmediately*/)
 	{
-#if VULKAN_USE_NEW_COMMAND_BUFFERS
 		FScopeLock ScopeLock(&CS);
 
-		// Traverse list backwards to the swap switches to elements already tested
+		// Traverse list backwards so the swap switches to elements already tested
 		for (int32 Index = Entries.Num() - 1; Index >= 0; --Index)
 		{
 			FEntry* Entry = &Entries[Index];
@@ -1257,7 +1281,6 @@ namespace VulkanRHI
 				Entries.RemoveAtSwap(Index, 1, false);
 			}
 		}
-#endif
 	}
 
 	FTempFrameAllocationBuffer::FTempFrameAllocationBuffer(FVulkanDevice* InDevice, uint32 InSize)
@@ -1268,26 +1291,25 @@ namespace VulkanRHI
 	{
 		for (int32 Index = 0; Index < NUM_RENDER_BUFFERS; ++Index)
 		{
-			Buffer[Index] = InDevice->GetStagingManager().AcquireBuffer(InSize, VK_BUFFER_USAGE_VERTEX_BUFFER_BIT | VK_BUFFER_USAGE_INDEX_BUFFER_BIT | VK_BUFFER_USAGE_UNIFORM_BUFFER_BIT | VK_BUFFER_USAGE_UNIFORM_TEXEL_BUFFER_BIT);
-			MappedData[Index] = (uint8*)Buffer[Index]->GetMappedPointer();
+			BufferSuballocations[Index] = InDevice->GetResourceHeapManager().AllocateBuffer(InSize,
+				VK_BUFFER_USAGE_VERTEX_BUFFER_BIT | VK_BUFFER_USAGE_INDEX_BUFFER_BIT | VK_BUFFER_USAGE_UNIFORM_BUFFER_BIT | VK_BUFFER_USAGE_UNIFORM_TEXEL_BUFFER_BIT,
+				VK_MEMORY_PROPERTY_HOST_VISIBLE_BIT,
+				__FILE__, __LINE__);
+			MappedData[Index] = (uint8*)BufferSuballocations[Index]->GetMappedPointer();
 			CurrentData[Index] = MappedData[Index];
 		}
 	}
 
 	FTempFrameAllocationBuffer::~FTempFrameAllocationBuffer()
 	{
-		check(!Buffer[0]);
+		Destroy();
 	}
 
 	void FTempFrameAllocationBuffer::Destroy()
 	{
 		for (int32 Index = 0; Index < NUM_RENDER_BUFFERS; ++Index)
 		{
-#if VULKAN_USE_NEW_COMMAND_BUFFERS
-			GetParent()->GetStagingManager().ReleaseBuffer(nullptr, Buffer[Index]);
-#else
-			GetParent()->GetStagingManager().ReleaseBuffer(Buffer[Index]);
-#endif
+			BufferSuballocations[Index] = nullptr;
 		}
 	}
 
@@ -1297,8 +1319,8 @@ namespace VulkanRHI
 		if (AlignedData + InSize <= MappedData[BufferIndex] + Size)
 		{
 			OutInfo.Data = AlignedData;
-			OutInfo.Buffer = Buffer[BufferIndex]->GetHandle();
-			OutInfo.Offset = (uint32)(AlignedData - MappedData[BufferIndex]);
+			OutInfo.BufferSuballocation = BufferSuballocations[BufferIndex];
+			OutInfo.CurrentOffset = (uint32)(AlignedData - MappedData[BufferIndex]);
 			CurrentData[BufferIndex] = AlignedData + InSize;
 			PeakUsed = FMath::Max(PeakUsed, (uint32)(CurrentData[BufferIndex] - MappedData[BufferIndex]));
 			return true;
diff --git a/Engine/Source/Runtime/VulkanRHI/Private/VulkanPendingState.cpp b/Engine/Source/Runtime/VulkanRHI/Private/VulkanPendingState.cpp
--- a/Engine/Source/Runtime/VulkanRHI/Private/VulkanPendingState.cpp
+++ b/Engine/Source/Runtime/VulkanRHI/Private/VulkanPendingState.cpp
@@ -262,24 +262,10 @@ FVulkanPendingState::FVulkanPendingState(FVulkanDevice* InDevice)
 	, bBeginRenderPass(false)
 	, bChangeRenderTarget(false)
 	, GlobalUniformPool(nullptr)
-#if VULKAN_USE_NEW_COMMAND_BUFFERS
-#else
-	, CurrentCmdBufferIndex(0)
-#endif
 	, bScissorEnable(false)
 {
 	check(Device);
 
-#if VULKAN_USE_NEW_COMMAND_BUFFERS
-	// State doesn't own or manage cmd buffers
-#else
-	for (uint32 CmdBufferIndex = 0; CmdBufferIndex < VULKAN_NUM_COMMAND_BUFFERS; ++CmdBufferIndex)
-	{
-		//#todo-rco: FIX ME!
-		CmdBuffers[CmdBufferIndex] = Device->GetImmediateContext().GetCommandBufferManager()->Create();
-	}
-#endif
-
 	Reset();
 
     // Create the global uniform pool
@@ -290,16 +276,6 @@ FVulkanPendingState::~FVulkanPendingState()
 {
 	check(Device);
 
-#if VULKAN_USE_NEW_COMMAND_BUFFERS
-	// State doesn't own or manage cmd buffers
-#else
-	for (uint32 CmdBufferIndex = 0; CmdBufferIndex < VULKAN_NUM_COMMAND_BUFFERS; ++CmdBufferIndex)
-	{
-		//#todo-rco: FIX ME!
-		Device->GetImmediateContext().GetCommandBufferManager()->Destroy(CmdBuffers[CmdBufferIndex]);
-		CmdBuffers[CmdBufferIndex] = nullptr;
-	}
-#endif
 	//Reset();
 
 	/*
@@ -331,32 +307,21 @@ FVulkanPendingState::~FVulkanPendingState()
 	RenderPassMap.Empty(0);
 }
 
-FVulkanDescriptorSets* FVulkanPendingState::AllocateDescriptorSet(const FVulkanBoundShaderState* BoundShaderState)
-{
-	FVulkanDescriptorSets* DescriptorSet = new FVulkanDescriptorSets(Device, BoundShaderState, Device->GetDescriptorPool());
-	return DescriptorSet;
-}
-
-void FVulkanPendingState::DeallocateDescriptorSet(FVulkanDescriptorSets*& DescriptorSet, const FVulkanBoundShaderState* BoundShaderState)
-{
-	delete DescriptorSet;
-	DescriptorSet = nullptr;
-}
-
 FVulkanGlobalUniformPool& FVulkanPendingState::GetGlobalUniformPool()
 {
     check(GlobalUniformPool);
     return *GlobalUniformPool;
 }
 
-#if VULKAN_USE_NEW_COMMAND_BUFFERS
-void FVulkanPendingState::RenderPassBegin(FVulkanCmdBuffer* CmdBuffer)
-#else
-void FVulkanPendingState::RenderPassBegin()
-#endif
+bool FVulkanPendingState::RenderPassBegin(FVulkanCmdBuffer* CmdBuffer)
 {
 	check(!bBeginRenderPass);
 
+	if (RTInfo.NumColorRenderTargets == 0 && !RTInfo.DepthStencilRenderTarget.Texture)
+	{
+		return false;
+	}
+
 	FVulkanRenderTargetLayout DesiredLayout(RTInfo);
 	FVulkanRenderPass* NewRenderPass = GetOrCreateRenderPass(DesiredLayout);
 
@@ -423,29 +388,19 @@ void FVulkanPendingState::RenderPassBegin()
 		// if either bit is set that we have a depth buffer attached
 	}
 		
-#if VULKAN_USE_NEW_COMMAND_BUFFERS
 	CmdBuffer->BeginRenderPass(CurrentState.RenderPass->GetLayout(), CurrentState.RenderPass->GetHandle(), CurrentState.FrameBuffer->GetHandle(), ClearValues);
-#else
-	CurrentState.RenderPass->Begin(GetCurrentCommandBuffer(), CurrentState.FrameBuffer->GetHandle(), ClearValues);
-#endif
+
 	bBeginRenderPass = true;
+
+	return true;
 }
 
-#if VULKAN_USE_NEW_COMMAND_BUFFERS
 void FVulkanPendingState::RenderPassEnd(FVulkanCmdBuffer* CmdBuffer)
 {
 	check(bBeginRenderPass);
 	CmdBuffer->EndRenderPass();
 	bBeginRenderPass = false;
 }
-#else
-void FVulkanPendingState::RenderPassEnd()
-{
-	check(bBeginRenderPass);
-	CurrentState.RenderPass->End(GetCurrentCommandBuffer());
-	bBeginRenderPass = false;
-}
-#endif
 
 // Expected to be called after render pass has been ended
 // and only from "FVulkanDynamicRHI::RHIEndDrawingViewport()"
@@ -455,14 +410,6 @@ void FVulkanPendingState::Reset()
 
 	CurrentState.Reset();
 
-#if VULKAN_USE_NEW_COMMAND_BUFFERS
-	// State doesn't own or manage cmd buffers
-#else
-	CurrentCmdBufferIndex = (CurrentCmdBufferIndex + 1) % VULKAN_NUM_COMMAND_BUFFERS;
-
-	GetCurrentCommandBuffer().Reset();
-	GetCurrentCommandBuffer().Begin();
-#endif
 	FMemory::Memzero(PendingStreams);
 }
 
@@ -624,11 +571,7 @@ inline void FVulkanBoundShaderState::BindDescriptorSets(FVulkanCmdBuffer* Cmd)
 	CurrDescriptorSets->Bind(Cmd, this);
 }
 
-#if VULKAN_USE_NEW_COMMAND_BUFFERS
-void FVulkanPendingState::PrepareDraw(FVulkanCmdBuffer* Cmd, VkPrimitiveTopology Topology)
-#else
-void FVulkanPendingState::PrepareDraw(VkPrimitiveTopology Topology)
-#endif
+void FVulkanPendingState::PrepareDraw(FVulkanCommandListContext* CmdListContext, FVulkanCmdBuffer* Cmd, VkPrimitiveTopology Topology)
 {
 	SCOPE_CYCLE_COUNTER(STAT_VulkanDrawCallPrepareTime);
 
@@ -636,14 +579,10 @@ void FVulkanPendingState::PrepareDraw(VkPrimitiveTopology Topology)
 	SetKeyBits(CurrentKey, OFFSET_POLYTYPE, NUMBITS_POLYTYPE, Topology);
 
 	//@TODO: let's try not to do this per draw call?
-#if VULKAN_USE_NEW_COMMAND_BUFFERS
 	UpdateRenderPass(Cmd);
-#else
-	UpdateRenderPass();
-#endif
 
 	check(CurrentState.Shader);
-    CurrentState.Shader->UpdateDescriptorSets(GlobalUniformPool);
+    CurrentState.Shader->UpdateDescriptorSets(CmdListContext, Cmd, GlobalUniformPool);
 
 	// let the BoundShaderState return a pipeline object for the full current state of things
 	CurrentState.InputAssembly.topology = Topology;
@@ -653,10 +592,6 @@ void FVulkanPendingState::PrepareDraw(VkPrimitiveTopology Topology)
 
 	{
 		SCOPE_CYCLE_COUNTER(STAT_VulkanPipelineBind);
-#if VULKAN_USE_NEW_COMMAND_BUFFERS
-#else
-		auto* Cmd = &GetCurrentCommandBuffer();
-#endif
 		Pipeline->UpdateDynamicStates(Cmd, CurrentState);
 
 		VkPipeline NewPipeline = Pipeline->GetHandle();
@@ -685,7 +620,7 @@ void FVulkanPendingState::PrepareDraw(VkPrimitiveTopology Topology)
 
 void FVulkanPendingState::SubmitPendingCommandBuffers(FVulkanPendingState::TCallback* Callback, void* CallbackUserData)
 {
-#if VULKAN_USE_NEW_COMMAND_BUFFERS
+#if 1//VULKAN_USE_NEW_COMMAND_BUFFERS
 	check(0);
 #else
 	if (IsRenderPassActive())
@@ -706,7 +641,7 @@ void FVulkanPendingState::SubmitPendingCommandBuffers(FVulkanPendingState::TCall
 
 void FVulkanPendingState::SubmitPendingCommandBuffersBlockingNoRenderPass()
 {
-#if VULKAN_USE_NEW_COMMAND_BUFFERS
+#if 1//VULKAN_USE_NEW_COMMAND_BUFFERS
 	check(0);
 #else
 	check(Device);
@@ -717,13 +652,25 @@ void FVulkanPendingState::SubmitPendingCommandBuffersBlockingNoRenderPass()
 void FVulkanPendingState::SetRenderTargetsInfo(const FRHISetRenderTargetsInfo& InRTInfo)
 {
 	//#todo-rco: Check perf
-#if !VULKAN_USE_NEW_COMMAND_BUFFERS
+#if 0//!VULKAN_USE_NEW_COMMAND_BUFFERS
 	if (NeedsToSetRenderTarget(InRTInfo) == false)
 	{
 		return;
 	}
 #endif
 
+#if 1
+	// Back this up for the next SetRenderTarget
+	RTInfo = InRTInfo;
+	if (RTInfo.NumColorRenderTargets == 1 && !RTInfo.ColorRenderTarget[0].Texture)
+	{
+		RTInfo.ColorRenderTarget[0].LoadAction = ERenderTargetLoadAction::ENoAction;
+		RTInfo.ColorRenderTarget[0].StoreAction = ERenderTargetStoreAction::ENoAction;
+		check(!RTInfo.bClearColor);
+		--RTInfo.NumColorRenderTargets;
+	}
+	bChangeRenderTarget = true;
+#else
 	//@NOTE: this is only needed for the work-around below
 	FRHISetRenderTargetsInfo PrevRTInfo = RTInfo;
 
@@ -743,6 +690,7 @@ void FVulkanPendingState::SetRenderTargetsInfo(const FRHISetRenderTargetsInfo& I
 			check(rtv.Texture);
 		}
 	}
+#endif
 }
 
 bool FVulkanPendingState::NeedsToSetRenderTarget(const FRHISetRenderTargetsInfo& InRTInfo)
@@ -846,19 +794,6 @@ FVulkanRenderPass& FVulkanPendingState::GetRenderPass()
 	return *CurrentState.RenderPass;
 }
 
-#if VULKAN_USE_NEW_COMMAND_BUFFERS
-#else
-VkCommandBuffer& FVulkanPendingState::GetCommandBuffer()
-{
-	return GetCurrentCommandBuffer().GetHandle();
-}
-
-const VkBool32 FVulkanPendingState::GetIsCommandBufferEmpty() const
-{
-	return GetCurrentCommandBuffer().GetIsEmpty();
-}
-#endif
-
 void FVulkanPendingState::SetViewport(uint32 MinX, uint32 MinY, float MinZ, uint32 MaxX, uint32 MaxY, float MaxZ)
 {
 	VkViewport& vp = CurrentState.Viewport;
diff --git a/Engine/Source/Runtime/VulkanRHI/Private/VulkanPendingState.h b/Engine/Source/Runtime/VulkanRHI/Private/VulkanPendingState.h
--- a/Engine/Source/Runtime/VulkanRHI/Private/VulkanPendingState.h
+++ b/Engine/Source/Runtime/VulkanRHI/Private/VulkanPendingState.h
@@ -20,33 +20,17 @@ public:
 
 	~FVulkanPendingState();
 
-	FVulkanDescriptorSets* AllocateDescriptorSet(const FVulkanBoundShaderState* BoundShaderState);
-	void DeallocateDescriptorSet(FVulkanDescriptorSets*& DescriptorSet, const FVulkanBoundShaderState* BoundShaderState);
-
     FVulkanGlobalUniformPool& GetGlobalUniformPool();
 
 	void SetRenderTargetsInfo(const FRHISetRenderTargetsInfo& InRTInfo);
 
 	void Reset();
 
-#if VULKAN_USE_NEW_COMMAND_BUFFERS
-	void RenderPassBegin(FVulkanCmdBuffer* CmdBuffer);
+	bool RenderPassBegin(FVulkanCmdBuffer* CmdBuffer);
 
-	void PrepareDraw(FVulkanCmdBuffer* CmdBuffer, VkPrimitiveTopology Topology);
+	void PrepareDraw(FVulkanCommandListContext* CmdListContext, FVulkanCmdBuffer* CmdBuffer, VkPrimitiveTopology Topology);
 
 	void RenderPassEnd(FVulkanCmdBuffer* CmdBuffer);
-#else
-	// Needs to be called before "PrepareDraw" and also starts writing in to the command buffer
-	void RenderPassBegin();
-
-	// Needs to be called before any calling a draw-call
-	// Binds graphics pipeline
-	void PrepareDraw(VkPrimitiveTopology Topology);
-
-	// Is currently linked to the command buffer
-	// Will get called in RHIEndDrawingViewport through SubmitPendingCommandBuffers
-	void RenderPassEnd();
-#endif
 
 	inline bool IsRenderPassActive() const
 	{
@@ -76,15 +60,8 @@ public:
 	// Retuns constructed render pass
 	FVulkanRenderPass& GetRenderPass();
 
-#if VULKAN_USE_NEW_COMMAND_BUFFERS
-#else
-	// Returns currently bound command buffer
-	VkCommandBuffer& GetCommandBuffer();
-	const VkBool32 GetIsCommandBufferEmpty() const;
-#endif
 	FVulkanFramebuffer* GetFrameBuffer();
 
-#if VULKAN_USE_NEW_COMMAND_BUFFERS
 	inline void UpdateRenderPass(FVulkanCmdBuffer* CmdBuffer)
 	{
 		//#todo-rco: Don't test here, move earlier to SetRenderTarget
@@ -95,7 +72,10 @@ public:
 				RenderPassEnd(CmdBuffer);
 			}
 
-			RenderPassBegin(CmdBuffer);
+			if (!RenderPassBegin(CmdBuffer))
+			{
+				return;
+			}
 			bChangeRenderTarget = false;
 		}
 		if (!bBeginRenderPass)
@@ -104,22 +84,6 @@ public:
 			bChangeRenderTarget = false;
 		}
 	}
-#else
-	inline void UpdateRenderPass()
-	{
-		//#todo-rco: Don't test here, move earlier to SetRenderTarget
-		if (bChangeRenderTarget)
-		{
-			if (bBeginRenderPass)
-			{
-				RenderPassEnd();
-			}
-
-			RenderPassBegin();
-			bChangeRenderTarget = false;
-		}
-	}
-#endif
 
 	inline const FVulkanPipelineGraphicsKey& GetCurrentKey() const
 	{
@@ -178,36 +142,6 @@ public:
 
 	void InitFrame();
 
-#if VULKAN_USE_NEW_COMMAND_BUFFERS
-#else
-	inline FVulkanCmdBuffer& GetCommandBuffer(uint32 Index)
-	{
-		check(Index<VULKAN_NUM_COMMAND_BUFFERS);
-		check(CmdBuffers[Index]);
-		return *CmdBuffers[Index];
-	}
-
-	inline uint32 GetCurrentCommandBufferIndex() const
-	{
-		return CurrentCmdBufferIndex;
-	}
-
-	// Returns current frame command buffer
-	// The index of the command buffer changes after "Reset()" is called
-	inline FVulkanCmdBuffer& GetCurrentCommandBuffer()
-	{
-		check(CmdBuffers[CurrentCmdBufferIndex]);
-		return *CmdBuffers[CurrentCmdBufferIndex];
-	}
-
-	// Returns current frame command buffer
-	// The index of the command buffer changes after "Reset()" is called
-	inline const FVulkanCmdBuffer& GetCurrentCommandBuffer() const
-	{
-		check(CmdBuffers[CurrentCmdBufferIndex]);
-		return *CmdBuffers[CurrentCmdBufferIndex];
-	}
-#endif
 private:
 	FVulkanRenderPass* GetOrCreateRenderPass(const FVulkanRenderTargetLayout& RTLayout);
 	FVulkanFramebuffer* GetOrCreateFramebuffer(const FRHISetRenderTargetsInfo& RHIRTInfo,
@@ -229,13 +163,7 @@ private:
 	//@TODO: probably needs to go somewhere else
 	FRHISetRenderTargetsInfo PrevRenderTargetsInfo;
 
-#if VULKAN_USE_NEW_COMMAND_BUFFERS
 	friend class FVulkanCommandListContext;
-#else
-	//#todo-rco: FIX ME! ASAP!!!
-	FVulkanCmdBuffer* CmdBuffers[VULKAN_NUM_COMMAND_BUFFERS];
-	uint32	CurrentCmdBufferIndex;
-#endif
 	
 	FRHISetRenderTargetsInfo RTInfo;
 
diff --git a/Engine/Source/Runtime/VulkanRHI/Private/VulkanQueue.cpp b/Engine/Source/Runtime/VulkanRHI/Private/VulkanQueue.cpp
--- a/Engine/Source/Runtime/VulkanRHI/Private/VulkanQueue.cpp
+++ b/Engine/Source/Runtime/VulkanRHI/Private/VulkanQueue.cpp
@@ -20,115 +20,19 @@ static FAutoConsoleVariableRef CVarVulkanUseExternalShaderCompiler(
 	);
 
 FVulkanQueue::FVulkanQueue(FVulkanDevice* InDevice, uint32 InFamilyIndex, uint32 InQueueIndex)
-	: 
-#if VULKAN_USE_NEW_COMMAND_BUFFERS
-#else
-	AcquireNextImageKHR(VK_NULL_HANDLE)
-	, QueuePresentKHR(VK_NULL_HANDLE)
-	, CurrentFenceIndex(0)
-	, CurrentImageIndex(-1)
-	, 
-#endif
-	Queue(VK_NULL_HANDLE)
+	: Queue(VK_NULL_HANDLE)
 	, FamilyIndex(InFamilyIndex)
 	, QueueIndex(InQueueIndex)
 	, Device(InDevice)
 {
-#if VULKAN_USE_NEW_COMMAND_BUFFERS
-#else
-	check(Device);
-	for (int BufferIndex = 0; BufferIndex < VULKAN_NUM_IMAGE_BUFFERS; ++BufferIndex)
-	{
-		ImageAcquiredSemaphore[BufferIndex] = new FVulkanSemaphore(*Device);
-		RenderingCompletedSemaphore[BufferIndex] = new FVulkanSemaphore(*Device);
-		Fences[BufferIndex] = nullptr;
-	}
-#endif
 	vkGetDeviceQueue(Device->GetInstanceHandle(), FamilyIndex, InQueueIndex, &Queue);
-
-#if VULKAN_USE_NEW_COMMAND_BUFFERS
-#else
-	AcquireNextImageKHR = (PFN_vkAcquireNextImageKHR)(void*)vkGetDeviceProcAddr(Device->GetInstanceHandle(), "vkAcquireNextImageKHR");
-	check(AcquireNextImageKHR);
-
-	QueuePresentKHR = (PFN_vkQueuePresentKHR)(void*)vkGetDeviceProcAddr(Device->GetInstanceHandle(), "vkQueuePresentKHR");
-	check(QueuePresentKHR);
-#endif
 }
 
 FVulkanQueue::~FVulkanQueue()
 {
 	check(Device);
-
-#if VULKAN_USE_NEW_COMMAND_BUFFERS
-#else
-	if (Fences[CurrentFenceIndex] != VK_NULL_HANDLE)
-	{
-		Device->GetFenceManager().WaitForFence(Fences[CurrentFenceIndex], UINT32_MAX);
-	}
-
-	for (int BufferIndex = 0; BufferIndex < VULKAN_NUM_IMAGE_BUFFERS; ++BufferIndex)
-	{
-		delete ImageAcquiredSemaphore[BufferIndex];
-		ImageAcquiredSemaphore[BufferIndex] = nullptr;
-		delete RenderingCompletedSemaphore[BufferIndex];
-		RenderingCompletedSemaphore[BufferIndex] = nullptr;
-
-		if (Fences[BufferIndex])
-		{
-			Device->GetFenceManager().ReleaseFence(Fences[BufferIndex]);
-			Fences[BufferIndex] = nullptr;
-		}
-	}
-#endif
-}
-
-#if VULKAN_USE_NEW_COMMAND_BUFFERS
-#else
-uint32 FVulkanQueue::AquireImageIndex(FVulkanSwapChain* Swapchain)
-{	
-	check(Device);
-	check(Swapchain);
-
-	//@TODO: Clean up code once the API and SDK is more robust.
-	//@SMEDIS: Since we only have one queue, we shouldn't actually need any semaphores in this .cpp file.
-	//@SMEDIS: I also don't think we need the Fences[] array.
-
-	CurrentFenceIndex = (CurrentFenceIndex + 1) % VULKAN_NUM_IMAGE_BUFFERS;
-
-	if (Fences[CurrentFenceIndex] != VK_NULL_HANDLE)
-	{
-		#if 0
-			VkResult Result = vkGetFenceStatus(Device->GetInstanceHandle(), Fences[CurrentFenceIndex]);
-			if(Result >= VK_SUCCESS)
-		#endif
-		{
-			// Grab the next fence and re-use it.
-			Device->GetFenceManager().WaitForFence(Fences[CurrentFenceIndex], UINT32_MAX);
-			Device->GetFenceManager().ResetFence(Fences[CurrentFenceIndex]);
-		}
-	}
-
-	// Get the index of the next swapchain image we should render to.
-	// We'll wait with an "infinite" timeout, the function will block until an image is ready.
-	// The ImageAcquiredSemaphore[ImageAcquiredSemaphoreIndex] will get signaled when the image is ready (upon function return).
-	// The Fences[CurrentFenceIndex] will also get signaled when the image is ready (upon function return).
-	// Note: Queues can still be filled in on the CPU side, but won't execute until the semaphore is signaled.
-	CurrentImageIndex = -1;
-	VkResult Result = AcquireNextImageKHR(
-		Device->GetInstanceHandle(),
-		Swapchain->SwapChain,
-		UINT64_MAX,
-		ImageAcquiredSemaphore[CurrentFenceIndex]->GetHandle(),
-		VK_NULL_HANDLE,
-		&CurrentImageIndex);
-	checkf(Result == VK_SUCCESS || Result == VK_SUBOPTIMAL_KHR, TEXT("AcquireNextImageKHR failed Result = %d"), int32(Result));
-
-	return CurrentImageIndex;
 }
-#endif
 
-#if VULKAN_USE_NEW_COMMAND_BUFFERS
 void FVulkanQueue::Submit(FVulkanCmdBuffer* CmdBuffer)
 {
 	check(CmdBuffer->HasEnded());
@@ -155,65 +59,10 @@ void FVulkanQueue::Submit(FVulkanCmdBuffer* CmdBuffer)
 
 	CmdBuffer->GetOwner()->RefreshFenceStatus();
 }
-#else
-void FVulkanQueue::Submit(FVulkanCmdBuffer* CmdBuffer)
-{
-	check(CmdBuffer);
-	CmdBuffer->End();
-
-	const VkCommandBuffer CmdBuffers[] = { CmdBuffer->GetHandle(false) };
-	const VkPipelineStageFlags WaitDstStageMask = VK_PIPELINE_STAGE_COLOR_ATTACHMENT_OUTPUT_BIT;
-	VkSubmitInfo SubmitInfo;
-	FMemory::Memzero(SubmitInfo);
-
-	VkSemaphore Local[2] =
-		{ 
-			ImageAcquiredSemaphore[CurrentFenceIndex]->GetHandle(),
-			RenderingCompletedSemaphore[CurrentFenceIndex]->GetHandle()
-		};
-	SubmitInfo.waitSemaphoreCount = 1;	// 0;	// 1;
-	SubmitInfo.pWaitSemaphores = &Local[0];
-	SubmitInfo.pWaitDstStageMask = &WaitDstStageMask;
-	SubmitInfo.signalSemaphoreCount = 1;	// 0;	// 1;
-	SubmitInfo.pSignalSemaphores = &Local[1];
-	SubmitInfo.sType = VK_STRUCTURE_TYPE_SUBMIT_INFO;
-	SubmitInfo.commandBufferCount = 1;
-	SubmitInfo.pCommandBuffers = CmdBuffers;
-
-	if (!Fences[CurrentFenceIndex])
-	{
-		Fences[CurrentFenceIndex] = Device->GetFenceManager().AllocateFence();
-	}
-	VERIFYVULKANRESULT(vkQueueSubmit(Queue, 1, &SubmitInfo, Fences[CurrentFenceIndex]->GetHandle()));
-
-	if (GWaitForIdleOnSubmit != 0)
-	{
-		VERIFYVULKANRESULT(vkQueueWaitIdle(Queue));
-	}
-}
-#endif
-
-#if !VULKAN_USE_NEW_COMMAND_BUFFERS
-void FVulkanQueue::Submit2(VkCommandBuffer CmdBuffer, VulkanRHI::FFence* Fence)
-{
-	VkSubmitInfo SubmitInfo;
-	FMemory::Memzero(SubmitInfo);
-	SubmitInfo.sType = VK_STRUCTURE_TYPE_SUBMIT_INFO;
-	SubmitInfo.commandBufferCount = 1;
-	SubmitInfo.pCommandBuffers = &CmdBuffer;
-	check(!Fence || !Fence->IsSignaled());
-	VERIFYVULKANRESULT(vkQueueSubmit(Queue, 1, &SubmitInfo, Fence ? Fence->GetHandle() : VK_NULL_HANDLE));
-
-	if (GWaitForIdleOnSubmit != 0)
-	{
-		VERIFYVULKANRESULT(vkQueueWaitIdle(Queue));
-	}
-}
-#endif
 
 void FVulkanQueue::SubmitBlocking(FVulkanCmdBuffer* CmdBuffer)
 {
-#if VULKAN_USE_NEW_COMMAND_BUFFERS
+#if 1//VULKAN_USE_NEW_COMMAND_BUFFERS
 	check(0);
 #else
 	check(CmdBuffer);
@@ -236,25 +85,3 @@ void FVulkanQueue::SubmitBlocking(FVulkanCmdBuffer* CmdBuffer)
 	Device->GetFenceManager().ReleaseFence(Fence);
 #endif
 }
-
-#if VULKAN_USE_NEW_COMMAND_BUFFERS
-#else
-void FVulkanQueue::Present(FVulkanSwapChain* Swapchain, uint32 ImageIndex)
-{
-	check(Swapchain);
-
-	VkPresentInfoKHR Info;
-	FMemory::Memzero(Info);
-	Info.sType = VK_STRUCTURE_TYPE_PRESENT_INFO_KHR;
-	Info.waitSemaphoreCount = 1;	// 0;	// 1;
-	VkSemaphore Local = RenderingCompletedSemaphore[CurrentFenceIndex]->GetHandle();
-	Info.pWaitSemaphores = &Local;
-	Info.swapchainCount = 1;
-	Info.pSwapchains = &Swapchain->SwapChain;
-	Info.pImageIndices = &ImageIndex;
-	VkResult OutResult = VK_SUCCESS;
-	Info.pResults = &OutResult;
-	VERIFYVULKANRESULT(QueuePresentKHR(Queue, &Info));
-	VERIFYVULKANRESULT(OutResult);
-}
-#endif
diff --git a/Engine/Source/Runtime/VulkanRHI/Private/VulkanQueue.h b/Engine/Source/Runtime/VulkanRHI/Private/VulkanQueue.h
--- a/Engine/Source/Runtime/VulkanRHI/Private/VulkanQueue.h
+++ b/Engine/Source/Runtime/VulkanRHI/Private/VulkanQueue.h
@@ -34,31 +34,12 @@ public:
 
 	void SubmitBlocking(FVulkanCmdBuffer* CmdBuffer);
 
-#if VULKAN_USE_NEW_COMMAND_BUFFERS
-#else
-	void Submit2(VkCommandBuffer CmdBuffer, VulkanRHI::FFence* Fence);
-	uint32 AquireImageIndex(FVulkanSwapChain* Swapchain);
-	void Present(FVulkanSwapChain* Swapchain, uint32 ImageIndex);
-#endif
-
 	inline VkQueue GetHandle() const
 	{
 		return Queue;
 	}
 
 private:
-#if VULKAN_USE_NEW_COMMAND_BUFFERS
-#else
-	PFN_vkAcquireNextImageKHR AcquireNextImageKHR;
-	PFN_vkQueuePresentKHR QueuePresentKHR;
-
-	FVulkanSemaphore* ImageAcquiredSemaphore[VULKAN_NUM_IMAGE_BUFFERS];
-	FVulkanSemaphore* RenderingCompletedSemaphore[VULKAN_NUM_IMAGE_BUFFERS];
-	VulkanRHI::FFence* Fences[VULKAN_NUM_IMAGE_BUFFERS];
-	uint32	CurrentFenceIndex;
-
-	uint32 CurrentImageIndex;	// Perhaps move this to the FVulkanSwapChain?
-#endif
 	VkQueue Queue;
 	uint32 FamilyIndex;
 	uint32 QueueIndex;
diff --git a/Engine/Source/Runtime/VulkanRHI/Private/VulkanRHI.cpp b/Engine/Source/Runtime/VulkanRHI/Private/VulkanRHI.cpp
--- a/Engine/Source/Runtime/VulkanRHI/Private/VulkanRHI.cpp
+++ b/Engine/Source/Runtime/VulkanRHI/Private/VulkanRHI.cpp
@@ -13,8 +13,12 @@
 #include "VulkanContext.h"
 #include "VulkanManager.h"
 
+#ifdef VK_API_VERSION
 // Check the SDK is least the API version we want to use
 static_assert(VK_API_VERSION >= UE_VK_API_VERSION, "Vulkan SDK is older than the version we want to support (UE_VK_API_VERSION). Please update your SDK.");
+#elif !defined(VK_HEADER_VERSION)
+	#error No VulkanSDK defines?
+#endif
 
 #if VK_HEADER_VERSION < 8 && (VK_API_VERSION < VK_MAKE_VERSION(1, 0, 3))
 	#include <vulkan/vk_ext_debug_report.h>
@@ -26,6 +30,12 @@ static TAutoConsoleVariable<int32> GLSLCvar(
 	TEXT("2 to use ES GLSL\n1 to use GLSL\n0 to use SPIRV")
 	);
 
+static TAutoConsoleVariable<int32> GRHIThreadCvar(
+	TEXT("r.Vulkan.RHIThread"),
+	0,
+	TEXT("1 to use RHI Thread")
+	);
+
 #if VULKAN_CUSTOM_MEMORY_MANAGER_ENABLED
 VkAllocationCallbacks GCallbacks;
 #endif
@@ -199,10 +209,7 @@ FVulkanCommandListContext::FVulkanCommandListContext(FVulkanDynamicRHI* InRHI, F
 	, PendingMinVertexIndex(0)
 	, PendingIndexDataStride(0)
 	, TempFrameAllocationBuffer(InDevice, VULKAN_TEMP_FRAME_ALLOCATOR_SIZE)
-#if VULKAN_USE_NEW_COMMAND_BUFFERS
-#else
 	, CommandBufferManager(nullptr)
-#endif
 {
 	// Create CommandBufferManager, contain all active buffers
 	CommandBufferManager = new FVulkanCommandBufferManager(InDevice);
@@ -261,6 +268,8 @@ void FVulkanDynamicRHI::Shutdown()
 	check(IsInGameThread() && IsInRenderingThread());
 	check(Device);
 
+	Device->PrepareForDestroy();
+
 	EmptyCachedBoundShaderStates();
 
 	FVulkanVertexDeclaration::EmptyCache();
@@ -403,7 +412,12 @@ void FVulkanDynamicRHI::InitInstance()
 	if (!Device)
 	{
 		check(!GIsRHIInitialized);
-		
+
+#ifdef PLATFORM_ANDROID
+		// Want to see the actual crash report on Android so unregister signal handlers
+		FPlatformMisc::SetCrashHandler((void(*)(const FGenericCrashContext& Context)) -1);
+#endif
+
 		GRHISupportsAsyncTextureCreation = false;
 
 		bool bDeviceSupportsGeometryShaders = true;
@@ -448,11 +462,11 @@ void FVulkanDynamicRHI::InitInstance()
 		GSupportsRenderTargetFormat_PF_G8 = false;	// #todo-rco
 		GSupportsQuads = false;	// Not supported in Vulkan
 		GRHISupportsTextureStreaming = false;	// #todo-rco
-		GRHISupportsRHIThread = false;	// #todo-rco
+		GRHISupportsRHIThread = GRHIThreadCvar->GetInt() != 0;	// #todo-rco
 
 		//#todo-rco: Leaks/issues still present!
 		// Indicate that the RHI needs to use the engine's deferred deletion queue.
-		GRHINeedsExtraDeletionLatency = false;
+		GRHINeedsExtraDeletionLatency = true;
 
 		GMaxShadowDepthBufferSizeX =  FPlatformMath::Min<int32>(Props.limits.maxImageDimension2D, GMaxShadowDepthBufferSizeX);
 		GMaxShadowDepthBufferSizeY =  FPlatformMath::Min<int32>(Props.limits.maxImageDimension2D, GMaxShadowDepthBufferSizeY);
@@ -532,14 +546,12 @@ void FVulkanCommandListContext::RHIBeginScene()
 
 void FVulkanCommandListContext::RHIEndScene()
 {
-#if VULKAN_USE_NEW_COMMAND_BUFFERS
 	auto* CmdBuffer = CommandBufferManager->GetActiveCmdBuffer();
 	auto& PendingState = Device->GetPendingState();
 	if (PendingState.IsRenderPassActive())
 	{
 		PendingState.RenderPassEnd(CmdBuffer);
 	}
-#endif
 }
 
 void FVulkanCommandListContext::RHIBeginDrawingViewport(FViewportRHIParamRef ViewportRHI, FTextureRHIParamRef RenderTargetRHI)
@@ -548,12 +560,7 @@ void FVulkanCommandListContext::RHIBeginDrawingViewport(FViewportRHIParamRef Vie
 	FVulkanViewport* Viewport = ResourceCast(ViewportRHI);
 	if (Viewport->GetBackBufferIndex() < 0)
 	{
-#if VULKAN_USE_NEW_COMMAND_BUFFERS
 		Viewport->AcquireBackBuffer(Device->GetImmediateContext().GetCommandBufferManager()->GetActiveCmdBuffer());
-#else
-		// SMEDIS: This is the normal method of getting a new backbuffer. However, RHIGetViewportBackBuffer() can also do it if necessary.
-		Viewport->CurrentBackBuffer = Device->GetQueue()->AquireImageIndex(Viewport->SwapChain);
-#endif
 	}
 	RHI->DrawingViewport = Viewport;
 
@@ -592,9 +599,7 @@ void FVulkanCommandListContext::RHIEndDrawingViewport(FViewportRHIParamRef Viewp
 
 void FVulkanCommandListContext::RHIEndFrame()
 {
-#if VULKAN_USE_NEW_COMMAND_BUFFERS
 	Device->GetDeferredDeletionQueue().Clear();
-#endif
 
 	Device->GetStagingManager().ProcessPendingFree();
 	Device->GetResourceHeapManager().ReleaseFreedPages();
@@ -613,12 +618,7 @@ void FVulkanCommandListContext::RHIPushEvent(const TCHAR* Name, FColor Color)
 #if VULKAN_ENABLE_DRAW_MARKERS
 		if (auto VkCmdDbgMarkerBegin = Device->GetCmdDbgMarkerBegin())
 		{
-#if VULKAN_USE_NEW_COMMAND_BUFFERS
 			VkCmdDbgMarkerBegin(GetCommandBufferManager()->GetActiveCmdBuffer()->GetHandle(), TCHAR_TO_ANSI(Name));
-#else
-			FVulkanPendingState& State = Device->GetPendingState();
-			VkCmdDbgMarkerBegin(State.GetCommandBuffer(), TCHAR_TO_ANSI(Name));
-#endif
 		}
 #endif
 
@@ -635,11 +635,7 @@ void FVulkanCommandListContext::RHIPopEvent()
 		if (auto VkCmdDbgMarkerEnd = Device->GetCmdDbgMarkerEnd())
 		{
 			FVulkanPendingState& State = Device->GetPendingState();
-#if VULKAN_USE_NEW_COMMAND_BUFFERS
 			VkCmdDbgMarkerEnd(GetCommandBufferManager()->GetActiveCmdBuffer()->GetHandle());
-#else
-			VkCmdDbgMarkerEnd(State.GetCommandBuffer());
-#endif
 		}
 #endif
 
@@ -770,7 +766,7 @@ void FVulkanBuffer::Unlock()
 
 void FVulkanBuffer::CopyTo(FVulkanSurface& Surface, const VkBufferImageCopy& CopyDesc, FVulkanPendingState* State)
 {
-#if VULKAN_USE_NEW_COMMAND_BUFFERS
+#if 1//VULKAN_USE_NEW_COMMAND_BUFFERS
 	check(0);
 #else
 	VkCommandBuffer Cmd = VK_NULL_HANDLE;
@@ -934,7 +930,7 @@ FVulkanDescriptorSets::~FVulkanDescriptorSets()
 	}
 }
 
-void FVulkanBufferView::Create(FVulkanDevice& Device, FVulkanBuffer& Buffer, EPixelFormat Format, uint32 Offset, uint32 Size)
+void FVulkanBufferView::Create(FVulkanBuffer& Buffer, EPixelFormat Format, uint32 Offset, uint32 Size)
 {
 	VkBufferViewCreateInfo ViewInfo;
 	FMemory::Memzero(ViewInfo);
@@ -952,10 +948,10 @@ void FVulkanBufferView::Create(FVulkanDevice& Device, FVulkanBuffer& Buffer, EPi
 	Flags = Buffer.GetFlags() & VK_BUFFER_USAGE_UNIFORM_TEXEL_BUFFER_BIT;
 	check(Flags);
 
-	VERIFYVULKANRESULT(vkCreateBufferView(Device.GetInstanceHandle(), &ViewInfo, nullptr, &View));
+	VERIFYVULKANRESULT(vkCreateBufferView(GetParent()->GetInstanceHandle(), &ViewInfo, nullptr, &View));
 }
 
-void FVulkanBufferView::Create(FVulkanDevice& Device, FVulkanResourceMultiBuffer* Buffer, EPixelFormat Format, uint32 Offset, uint32 Size)
+void FVulkanBufferView::Create(FVulkanResourceMultiBuffer* Buffer, EPixelFormat Format, uint32 Offset, uint32 Size)
 {
 	VkBufferViewCreateInfo ViewInfo;
 	FMemory::Memzero(ViewInfo);
@@ -975,14 +971,14 @@ void FVulkanBufferView::Create(FVulkanDevice& Device, FVulkanResourceMultiBuffer
 
 	ViewInfo.flags = Flags;
 
-	VERIFYVULKANRESULT(vkCreateBufferView(Device.GetInstanceHandle(), &ViewInfo, nullptr, &View));
+	VERIFYVULKANRESULT(vkCreateBufferView(GetParent()->GetInstanceHandle(), &ViewInfo, nullptr, &View));
 }
 
-void FVulkanBufferView::Destroy(FVulkanDevice& Device)
+void FVulkanBufferView::Destroy()
 {
 	if (View != VK_NULL_HANDLE)
 	{
-		vkDestroyBufferView(Device.GetInstanceHandle(), View, nullptr);
+		vkDestroyBufferView(GetParent()->GetInstanceHandle(), View, nullptr);
 		View = VK_NULL_HANDLE;
 	}
 }
@@ -1031,11 +1027,11 @@ FVulkanRenderTargetLayout::FVulkanRenderTargetLayout(const FRHISetRenderTargetsI
 #if !VULKAN_USE_MSAA_RESOLVE_ATTACHMENTS
 		    CurrDesc.stencilStoreOp = VK_ATTACHMENT_STORE_OP_DONT_CARE;
 #endif
-		    CurrDesc.initialLayout = VK_IMAGE_LAYOUT_COLOR_ATTACHMENT_OPTIMAL;
-		    CurrDesc.finalLayout = VK_IMAGE_LAYOUT_COLOR_ATTACHMENT_OPTIMAL;
+		    CurrDesc.initialLayout = VK_IMAGE_LAYOUT_GENERAL;
+		    CurrDesc.finalLayout = VK_IMAGE_LAYOUT_GENERAL;
     
 		    ColorReferences[NumColorAttachments].attachment = NumAttachments;
-		    ColorReferences[NumColorAttachments].layout = VK_IMAGE_LAYOUT_COLOR_ATTACHMENT_OPTIMAL;
+		    ColorReferences[NumColorAttachments].layout = VK_IMAGE_LAYOUT_GENERAL;
 		    NumAttachments++;
     
 #if VULKAN_USE_MSAA_RESOLVE_ATTACHMENTS
@@ -1190,25 +1186,6 @@ FVulkanRenderPass::~FVulkanRenderPass()
 	RenderPass = VK_NULL_HANDLE;
 }
 
-#if VULKAN_USE_NEW_COMMAND_BUFFERS
-#else
-void FVulkanRenderPass::Begin(FVulkanCmdBuffer& CmdBuf, const VkFramebuffer& Framebuffer, const VkClearValue* AttachmentClearValues)
-{
-	VkRenderPassBeginInfo Info;
-	FMemory::Memzero(Info);
-	Info.sType = VK_STRUCTURE_TYPE_RENDER_PASS_BEGIN_INFO;
-	Info.renderPass = GetHandle();
-	Info.framebuffer = Framebuffer;
-	Info.renderArea.offset.x = 0;
-	Info.renderArea.offset.y = 0;
-	Info.renderArea.extent = Layout.GetExtent2D();
-	Info.clearValueCount = Layout.GetNumAttachments();
-	Info.pClearValues = AttachmentClearValues;
-
-	vkCmdBeginRenderPass(CmdBuf.GetHandle(false), &Info, VK_SUBPASS_CONTENTS_INLINE);
-}
-#endif
-
 
 void VulkanSetImageLayout(
 	VkCommandBuffer CmdBuffer,
@@ -1232,7 +1209,7 @@ void VulkanSetImageLayout(
 	switch (NewLayout)
 	{
 	case VK_IMAGE_LAYOUT_TRANSFER_DST_OPTIMAL:
-		ImageBarrier.dstAccessMask = VK_ACCESS_TRANSFER_READ_BIT;
+		ImageBarrier.dstAccessMask = VK_ACCESS_TRANSFER_WRITE_BIT;
 		break;
 	case VK_IMAGE_LAYOUT_COLOR_ATTACHMENT_OPTIMAL:
 		ImageBarrier.dstAccessMask = VK_ACCESS_COLOR_ATTACHMENT_WRITE_BIT;
@@ -1241,6 +1218,10 @@ void VulkanSetImageLayout(
 		ImageBarrier.dstAccessMask = VK_ACCESS_DEPTH_STENCIL_ATTACHMENT_WRITE_BIT;
 		break;
 	case VK_IMAGE_LAYOUT_SHADER_READ_ONLY_OPTIMAL:
+		if (OldLayout == VK_IMAGE_LAYOUT_TRANSFER_DST_OPTIMAL)
+		{
+			ImageBarrier.srcAccessMask = VK_ACCESS_TRANSFER_WRITE_BIT;
+		}
 		ImageBarrier.dstAccessMask = VK_ACCESS_SHADER_READ_BIT | VK_ACCESS_INPUT_ATTACHMENT_READ_BIT;
 		break;
 	case VK_IMAGE_LAYOUT_PRESENT_SRC_KHR:
@@ -1310,17 +1291,13 @@ FVulkanRingBuffer::FVulkanRingBuffer(FVulkanDevice* InDevice, uint64 TotalSize,
 	, MinAlignment(0)
 {
 	FRHIResourceCreateInfo CreateInfo;
-	Buffer = InDevice->GetStagingManager().AcquireBuffer(TotalSize, Usage);
-
-	VkMemoryRequirements MemReqs;
-	vkGetBufferMemoryRequirements(InDevice->GetInstanceHandle(), Buffer->GetHandle(), &MemReqs);
-
-	MinAlignment = (uint32)MemReqs.alignment;
+	BufferSuballocation = InDevice->GetResourceHeapManager().AllocateBuffer(TotalSize, Usage, VK_MEMORY_PROPERTY_HOST_VISIBLE_BIT, __FILE__, __LINE__);
+	MinAlignment = BufferSuballocation->GetBufferAllocation()->GetAlignment();
 }
 
 FVulkanRingBuffer::~FVulkanRingBuffer()
 {
-	Device->GetStagingManager().ReleaseBuffer(Buffer);
+	delete BufferSuballocation;
 }
 
 uint64 FVulkanRingBuffer::AllocateMemory(uint64 Size, uint32 Alignment)
@@ -1363,5 +1340,6 @@ void FVulkanDynamicRHI::DumpMemory()
 	auto* RHI = (FVulkanDynamicRHI*)GDynamicRHI;
 	RHI->Device->GetMemoryManager().DumpMemory();
 	RHI->Device->GetResourceHeapManager().DumpMemory();
+	RHI->Device->GetStagingManager().DumpMemory();
 }
 #endif
diff --git a/Engine/Source/Runtime/VulkanRHI/Private/VulkanRHIPrivate.h b/Engine/Source/Runtime/VulkanRHI/Private/VulkanRHIPrivate.h
--- a/Engine/Source/Runtime/VulkanRHI/Private/VulkanRHIPrivate.h
+++ b/Engine/Source/Runtime/VulkanRHI/Private/VulkanRHIPrivate.h
@@ -185,11 +185,7 @@ public:
 	TArray<VkImageView> Attachments;
 	TArray<VkImageSubresourceRange> SubresourceRanges;
 
-#if VULKAN_USE_NEW_COMMAND_BUFFERS
-	void InsertWriteBarrier(FVulkanCmdBuffer* CmdBuffer);
-#else
-	void InsertWriteBarrier(VkCommandBuffer cmd);
-#endif
+	void InsertWriteBarriers(FVulkanCmdBuffer* CmdBuffer);
 
 private:
 	VkFramebuffer Framebuffer;
@@ -219,15 +215,6 @@ private:
 	FVulkanRenderPass(FVulkanDevice& Device, const FVulkanRenderTargetLayout& RTLayout);
 	~FVulkanRenderPass();
 
-#if VULKAN_USE_NEW_COMMAND_BUFFERS
-#else
-	void Begin(FVulkanCmdBuffer& CmdBuf, const VkFramebuffer& Framebuffer, const VkClearValue* AttachmentClearValues);
-	void End(FVulkanCmdBuffer& CmdBuf)
-	{
-		vkCmdEndRenderPass(CmdBuf.GetHandle());
-	}
-#endif
-
 private:
 	FVulkanRenderTargetLayout Layout;
 	VkRenderPass RenderPass;
@@ -328,6 +315,9 @@ private:
 	FVulkanDescriptorPool* Pool;
 	const FVulkanDescriptorSetsLayout& Layout;
 	TArray<VkDescriptorSet> Sets;
+
+	friend class FVulkanBoundShaderState;
+	friend class FVulkanCommandListContext;
 };
 
 void VulkanSetImageLayout(VkCommandBuffer CmdBuffer, VkImage Image, VkImageLayout OldLayout, VkImageLayout NewLayout, const VkImageSubresourceRange& SubresourceRange);
diff --git a/Engine/Source/Runtime/VulkanRHI/Private/VulkanRenderTarget.cpp b/Engine/Source/Runtime/VulkanRHI/Private/VulkanRenderTarget.cpp
--- a/Engine/Source/Runtime/VulkanRHI/Private/VulkanRenderTarget.cpp
+++ b/Engine/Source/Runtime/VulkanRHI/Private/VulkanRenderTarget.cpp
@@ -37,16 +37,8 @@ void FVulkanCommandListContext::RHISetRenderTargetsAndClear(const FRHISetRenderT
 	FVulkanPendingState& PendingState = Device->GetPendingState();
 	PendingState.SetRenderTargetsInfo(RenderTargetsInfo);
 
-#if VULKAN_USE_NEW_COMMAND_BUFFERS
-	auto* CmdBuffer = CommandBufferManager->GetActiveCmdBuffer();
-	//#todo-rco: Write barrier always for now...
-	FVulkanFramebuffer* Framebuffer = PendingState.GetFrameBuffer();
-	if (Framebuffer)
-	{
-		Framebuffer->InsertWriteBarrier(CmdBuffer);
-	}
-
 #if 0//VULKAN_USE_NEW_RESOURCE_MANAGEMENT
+	auto* CmdBuffer = CommandBufferManager->GetActiveCmdBuffer();
 	if (CmdBuffer->IsInsideRenderPass())
 	{
 		if (
@@ -58,15 +50,13 @@ void FVulkanCommandListContext::RHISetRenderTargetsAndClear(const FRHISetRenderT
 		}
 	}
 #endif
-#else
-#endif
 }
 
 void FVulkanCommandListContext::RHICopyToResolveTarget(FTextureRHIParamRef SourceTextureRHI, FTextureRHIParamRef DestTextureRHI, bool bKeepOriginalSurface, const FResolveParams& ResolveParams)
 {
 	FVulkanPendingState& State = Device->GetPendingState();
 
-#if VULKAN_USE_NEW_COMMAND_BUFFERS
+#if 1//VULKAN_USE_NEW_COMMAND_BUFFERS
 	auto* CmdBuffer = CommandBufferManager->GetActiveCmdBuffer();
 	// Verify if we need to do some work (for the case of SetRT(), CopyToResolve() with no draw calls in between)
 	State.UpdateRenderPass(CmdBuffer);
@@ -78,7 +68,14 @@ void FVulkanCommandListContext::RHICopyToResolveTarget(FTextureRHIParamRef Sourc
 		State.RenderPassEnd(CmdBuffer);
 	}
 
-	check(SourceTextureRHI->GetNumSamples() < 2);
+	check(!SourceTextureRHI || SourceTextureRHI->GetNumSamples() < 2);
+
+	FVulkanFramebuffer* Framebuffer = State.GetFrameBuffer();
+	if (Framebuffer)
+	{
+		Framebuffer->InsertWriteBarriers(GetCommandBufferManager()->GetActiveCmdBuffer());
+	}
+
 #else
 	// If we're using the pResolveAttachments property of the subpass, so we don't need manual command buffer resolves and this function is a NoOp.
 #if !VULKAN_USE_MSAA_RESOLVE_ATTACHMENTS
@@ -160,7 +157,7 @@ void FVulkanCommandListContext::RHITransitionResources(EResourceTransitionAccess
 				FVulkanTexture2D* Texture = (FVulkanTexture2D*)RHITexture2D;
 				if (Texture->Surface.ImageLayout == VK_IMAGE_LAYOUT_PRESENT_SRC_KHR)
 				{
-#if VULKAN_USE_NEW_COMMAND_BUFFERS
+#if 1//VULKAN_USE_NEW_COMMAND_BUFFERS
 					check(0);
 #else
 					check(Texture->Surface.GetAspectMask() == VK_IMAGE_ASPECT_COLOR_BIT);
diff --git a/Engine/Source/Runtime/VulkanRHI/Private/VulkanShaders.cpp b/Engine/Source/Runtime/VulkanRHI/Private/VulkanShaders.cpp
--- a/Engine/Source/Runtime/VulkanRHI/Private/VulkanShaders.cpp
+++ b/Engine/Source/Runtime/VulkanRHI/Private/VulkanShaders.cpp
@@ -225,7 +225,6 @@ FVulkanBoundShaderState::FVulkanBoundShaderState(
 	, BindingsNum(0)
 	, BindingsMask(0)
 	, AttributesNum(0)
-	, LastFrameRendered(-1)
 #if VULKAN_ENABLE_PIPELINE_CACHE
 	, GlobalListLink(this)
 #endif
@@ -243,7 +242,6 @@ FVulkanBoundShaderState::FVulkanBoundShaderState(
 	FMemory::Memzero(DirtyPackedUniformBufferStagingMask);
 	FMemory::Memzero(DescriptorBufferInfoForStage);
 	FMemory::Memzero(SRVWriteInfoForStage);
-	FMemory::Memzero(CurrentDS);
 
 	#if VULKAN_HAS_DEBUGGING_ENABLED
 		FMemory::Memzero(ImageDescCount);		
@@ -251,13 +249,6 @@ FVulkanBoundShaderState::FVulkanBoundShaderState(
 
 	check(Device);
 
-#if VULKAN_USE_NEW_COMMAND_BUFFERS
-	//#todo-rco: Fix this!
-	LastFrameRendered = GFrameNumberRenderThread;
-#else
-	LastFrameRendered = Device->GetPendingState().GetCurrentCommandBufferIndex();
-#endif
-
 	Layout = new FVulkanDescriptorSetsLayout(Device);
 
 	FVulkanVertexDeclaration* InVertexDeclaration = ResourceCast(InVertexDeclarationRHI);
@@ -333,15 +324,11 @@ FVulkanBoundShaderState::~FVulkanBoundShaderState()
 	GlobalListLink.Unlink();
 #endif
 
-	// Nuke descriptor sets
-	for(uint32 CommandIndex=0; CommandIndex<VULKAN_NUM_COMMAND_BUFFERS; ++CommandIndex)
+	for (int32 Index = 0; Index < DescriptorSetsEntries.Num(); ++Index)
 	{
-		TArray<FVulkanDescriptorSets*>& Sets = DescriptorSets[CommandIndex];
-		for(FVulkanDescriptorSets* DescriptorSet : Sets)
-		{
-			Device->GetPendingState().DeallocateDescriptorSet(DescriptorSet, this);
-		}
+		delete DescriptorSetsEntries[Index];
 	}
+	DescriptorSetsEntries.Empty(0);
 
 	// toss the pipeline states
 	for (auto& Pair : PipelineCache)
@@ -396,7 +383,9 @@ FVulkanPipeline* FVulkanBoundShaderState::PrepareForDraw(const FVulkanPipelineGr
 #if !UE_BUILD_SHIPPING
 			if (Device->FrameCounter > 3)
 			{
-				FPlatformMisc::LowLevelOutputDebugStringf(TEXT("Created a hitchy pipeline state for hash %llx %x (this = %p) VS %s %p %d PS %s %p %d\n"), PipelineKey, VertexInputKey, this, *VertexShader->DebugName, (void*)VertexShader.GetReference(), VertexShader->GlslSource.Num(), *PixelShader->DebugName, (void*)PixelShader.GetReference(), PixelShader->GlslSource.Num());
+				FPlatformMisc::LowLevelOutputDebugStringf(TEXT("Created a hitchy pipeline state for hash %llx%llx %x (this = %x) VS %s %x %d PS %s %x %d\n"), 
+					PipelineKey.Key[0], PipelineKey.Key[1], VertexInputKey, this, *VertexShader->DebugName, (void*)VertexShader.GetReference(), VertexShader->GlslSource.Num(), 
+					*PixelShader->DebugName, (void*)PixelShader.GetReference(), PixelShader->GlslSource.Num());
 			}
 #endif
 		}
@@ -423,18 +412,6 @@ void FVulkanBoundShaderState::ResetState()
 
 	CurrDescriptorSets = nullptr;
 
-#if VULKAN_USE_NEW_COMMAND_BUFFERS
-	//#todo-rco: Fix this!
-	const uint32 CurrentImageIndex = GFrameNumberRenderThread % ARRAY_COUNT(CurrentDS);
-#else
-	// reset the DS list the first time it's used in a frame
-	const uint32 CurrentImageIndex = Device->GetPendingState().GetCurrentCommandBufferIndex();
-#endif
-	if (LastFrameRendered != CurrentImageIndex)
-	{
-		CurrentDS[CurrentImageIndex] = 0;
-		LastFrameRendered = CurrentImageIndex;
-	}
 	LastBoundPipeline = VK_NULL_HANDLE;
 	bDirtyVertexStreams = true;
 }
@@ -1053,7 +1030,48 @@ void FVulkanBoundShaderState::SetLastError(const FString& Error)
 }
 #endif
 
-void FVulkanBoundShaderState::UpdateDescriptorSets(FVulkanGlobalUniformPool* GlobalUniformPool)
+FVulkanBoundShaderState::FDescriptorSetsPair::~FDescriptorSetsPair()
+{
+	delete DescriptorSets;
+}
+
+inline FVulkanDescriptorSets* FVulkanBoundShaderState::RequestDescriptorSets(FVulkanCmdBuffer* CmdBuffer)
+{
+	//#todo-rco: Make thread safe!
+	FDescriptorSetsEntry* FoundEntry = nullptr;
+	for (FDescriptorSetsEntry* DescriptorSetsEntry : DescriptorSetsEntries)
+	{
+		if (DescriptorSetsEntry->CmdBuffer == CmdBuffer)
+		{
+			FoundEntry = DescriptorSetsEntry;
+		}
+	}
+
+	if (!FoundEntry)
+	{
+		FoundEntry = new FDescriptorSetsEntry(CmdBuffer);
+		DescriptorSetsEntries.Add(FoundEntry);
+	}
+
+	const uint64 CmdBufferFenceSignaledCounter = CmdBuffer->GetFenceSignaledCounter();
+	for (int32 Index = 0; Index < FoundEntry->Pairs.Num(); ++Index)
+	{
+		FDescriptorSetsPair& Entry = FoundEntry->Pairs[Index];
+		if (Entry.FenceCounter < CmdBufferFenceSignaledCounter)
+		{
+			Entry.FenceCounter = CmdBufferFenceSignaledCounter;
+			return Entry.DescriptorSets;
+		}
+	}
+
+	FDescriptorSetsPair* NewEntry = new (FoundEntry->Pairs) FDescriptorSetsPair;
+	NewEntry->DescriptorSets = new FVulkanDescriptorSets(Device, this, Device->GetDescriptorPool());
+	NewEntry->FenceCounter = CmdBufferFenceSignaledCounter;
+	return NewEntry->DescriptorSets;
+}
+
+
+void FVulkanBoundShaderState::UpdateDescriptorSets(FVulkanCommandListContext* CmdListContext, FVulkanCmdBuffer* CmdBuffer, FVulkanGlobalUniformPool* GlobalUniformPool)
 {
 #if VULKAN_ENABLE_AGGRESSIVE_STATS
 	SCOPE_CYCLE_COUNTER(STAT_VulkanUpdateDescriptorSets);
@@ -1065,24 +1083,7 @@ void FVulkanBoundShaderState::UpdateDescriptorSets(FVulkanGlobalUniformPool* Glo
 
 	int32 WriteIndex = 0;
 
-	// get the next available DS
-
-#if VULKAN_USE_NEW_COMMAND_BUFFERS
-	//#todo-rco: Fix this!
-	const uint32 CurrentImageIndex = GFrameNumberRenderThread % ARRAY_COUNT(CurrentDS);
-#else
-	// This is the image where we are going to execute command buffer into
-	const uint32 CurrentImageIndex = Device->GetPendingState().GetCurrentCommandBufferIndex();
-#endif
-	TArray<FVulkanDescriptorSets*>& CurrentDescriptorSet = DescriptorSets[CurrentImageIndex];
-	int32& CurrentDescriptorSetIndex = CurrentDS[CurrentImageIndex];
-
-	if (CurrentDescriptorSetIndex == CurrentDescriptorSet.Num())
-	{
-		// @todo vulkan: we need to check we don't make more of these than MaxSets in FVulkanPendingState constructor!!
-		CurrentDescriptorSet.Push(State.AllocateDescriptorSet(this));
-	}
-	CurrDescriptorSets = CurrentDescriptorSet[CurrentDescriptorSetIndex++];
+	CurrDescriptorSets = RequestDescriptorSets(CmdBuffer);
 
 	auto& DescriptorSetHandles = CurrDescriptorSets->GetHandles();
 	int32 DescriptorSetIndex = 0;
@@ -1090,7 +1091,7 @@ void FVulkanBoundShaderState::UpdateDescriptorSets(FVulkanGlobalUniformPool* Glo
 #if VULKAN_USE_RING_BUFFER_FOR_GLOBAL_UBS
 	// this is an optimization for the ring buffer to only truly lock once for all uniforms
 	FVulkanRingBuffer* RingBuffer = Device->GetUBRingBuffer();
-	uint8* RingBufferBase = (uint8*)RingBuffer->Buffer->GetMappedPointer();
+	uint8* RingBufferBase = (uint8*)RingBuffer->GetMappedPointer();
 #endif
 
 	//#todo-rco: Compute!
@@ -1156,8 +1157,8 @@ void FVulkanBoundShaderState::UpdateDescriptorSets(FVulkanGlobalUniformPool* Glo
 					// Here we can specify a more precise buffer update
 					// However, this need to complemented with the buffer map/unmap functionality.
 					//@NOTE: bufferView is for texel buffers
-					BufferInfo->buffer = RingBuffer->Buffer->GetHandle();
-					BufferInfo->offset = RingBufferOffset;
+					BufferInfo->buffer = RingBuffer->GetHandle();
+					BufferInfo->offset = RingBufferOffset + RingBuffer->GetBufferOffset();
 					BufferInfo->range = UBSize;
 
 #else
diff --git a/Engine/Source/Runtime/VulkanRHI/Private/VulkanSwapChain.cpp b/Engine/Source/Runtime/VulkanRHI/Private/VulkanSwapChain.cpp
--- a/Engine/Source/Runtime/VulkanRHI/Private/VulkanSwapChain.cpp
+++ b/Engine/Source/Runtime/VulkanRHI/Private/VulkanSwapChain.cpp
@@ -12,9 +12,8 @@ FVulkanSwapChain::FVulkanSwapChain(VkInstance Instance, FVulkanDevice& InDevice,
 	: SwapChain(VK_NULL_HANDLE)
 	, Device(InDevice)
 	, Surface(VK_NULL_HANDLE)
-#if VULKAN_USE_NEW_COMMAND_BUFFERS
 	, CurrentImageIndex(-1)
-#endif
+	, SemaphoreIndex(0)
 {
 #define FETCH_KHR_FPN(FunctionPointerName)\
 	{\
@@ -197,34 +196,31 @@ FVulkanSwapChain::FVulkanSwapChain(VkInstance Instance, FVulkanDevice& InDevice,
 	OutImages.AddUninitialized(NumSwapChainImages);
 	VERIFYVULKANRESULT_EXPANDED(GetSwapchainImagesKHR(Device.GetInstanceHandle(), SwapChain, &NumSwapChainImages, OutImages.GetData()));
 
-#if VULKAN_USE_NEW_COMMAND_BUFFERS
 	ImageAcquiredSemaphore.AddUninitialized(DesiredNumBuffers);
-	RenderingCompletedSemaphore.AddUninitialized(DesiredNumBuffers);
+	//RenderingCompletedSemaphore.AddUninitialized(DesiredNumBuffers);
 	for (uint32 BufferIndex = 0; BufferIndex < DesiredNumBuffers; ++BufferIndex)
 	{
 		ImageAcquiredSemaphore[BufferIndex] = new FVulkanSemaphore(Device);
-		RenderingCompletedSemaphore[BufferIndex] = new FVulkanSemaphore(Device);
+		//RenderingCompletedSemaphore[BufferIndex] = new FVulkanSemaphore(Device);
 	}
-#endif
 }
 
 void FVulkanSwapChain::Destroy()
 {
-#if VULKAN_USE_NEW_COMMAND_BUFFERS
+	vkDestroySwapchainKHR(Device.GetInstanceHandle(), SwapChain, nullptr);
+	SwapChain = VK_NULL_HANDLE;
+
+	//#todo-rco: Enqueue for deletion as we first need to destroy the cmd buffers and queues otherwise validation fails
 	for (int BufferIndex = 0; BufferIndex < ImageAcquiredSemaphore.Num(); ++BufferIndex)
 	{
 		delete ImageAcquiredSemaphore[BufferIndex];
-		delete RenderingCompletedSemaphore[BufferIndex];
+		//delete RenderingCompletedSemaphore[BufferIndex];
 	}
 	ImageAcquiredSemaphore.Empty(0);
-	RenderingCompletedSemaphore.Empty(0);
-#endif
-	vkDestroySwapchainKHR(Device.GetInstanceHandle(), SwapChain, nullptr);
-	SwapChain = VK_NULL_HANDLE;
+	//RenderingCompletedSemaphore.Empty(0);
 }
 
-#if VULKAN_USE_NEW_COMMAND_BUFFERS
-int32 FVulkanSwapChain::AcquireImageIndex()
+int32 FVulkanSwapChain::AcquireImageIndex(FVulkanSemaphore** OutSemaphore)
 {
 	// Get the index of the next swapchain image we should render to.
 	// We'll wait with an "infinite" timeout, the function will block until an image is ready.
@@ -233,15 +229,18 @@ int32 FVulkanSwapChain::AcquireImageIndex()
 	// Note: Queues can still be filled in on the CPU side, but won't execute until the semaphore is signaled.
 	//CurrentImageIndex = -1;
 	uint32 ImageIndex = 0;
+	SemaphoreIndex = (SemaphoreIndex + 1) % ImageAcquiredSemaphore.Num();
+	*OutSemaphore = ImageAcquiredSemaphore[SemaphoreIndex];
 	VkResult Result = AcquireNextImageKHR(
 		Device.GetInstanceHandle(),
 		SwapChain,
 		UINT64_MAX,
-		VK_NULL_HANDLE, //ImageAcquiredSemaphore[CurrentFenceIndex]->GetHandle(),
+		ImageAcquiredSemaphore[SemaphoreIndex]->GetHandle(),
 		VK_NULL_HANDLE,
 		&ImageIndex);
 	checkf(Result == VK_SUCCESS || Result == VK_SUBOPTIMAL_KHR, TEXT("AcquireNextImageKHR failed Result = %d"), int32(Result));
 	CurrentImageIndex = (int32)ImageIndex;
+	check(CurrentImageIndex == ImageIndex);
 	return CurrentImageIndex;
 }
 
@@ -252,17 +251,12 @@ void FVulkanSwapChain::Present(FVulkanQueue* Queue)
 	VkPresentInfoKHR Info;
 	FMemory::Memzero(Info);
 	Info.sType = VK_STRUCTURE_TYPE_PRESENT_INFO_KHR;
-/*
-	Info.waitSemaphoreCount = 1;	// 0;	// 1;
-	VkSemaphore Semaphore = RenderingCompletedSemaphore[CurrentImageIndex]->GetHandle();
+	Info.waitSemaphoreCount = 1;
+	VkSemaphore Semaphore = ImageAcquiredSemaphore[SemaphoreIndex]->GetHandle();
 	Info.pWaitSemaphores = &Semaphore;
-*/
-
 	Info.swapchainCount = 1;
 	Info.pSwapchains = &SwapChain;
 	Info.pImageIndices = (uint32*)&CurrentImageIndex;
 
 	VERIFYVULKANRESULT(QueuePresentKHR(Queue->GetHandle(), &Info));
 }
-#else
-#endif
diff --git a/Engine/Source/Runtime/VulkanRHI/Private/VulkanSwapChain.h b/Engine/Source/Runtime/VulkanRHI/Private/VulkanSwapChain.h
--- a/Engine/Source/Runtime/VulkanRHI/Private/VulkanSwapChain.h
+++ b/Engine/Source/Runtime/VulkanRHI/Private/VulkanSwapChain.h
@@ -11,23 +11,12 @@ class FVulkanQueue;
 class FVulkanSwapChain
 {
 public:
-#if VULKAN_USE_NEW_COMMAND_BUFFERS
-#else
-#if !PLATFORM_WINDOWS
-	//@HACK : maybe NUM_BUFFERS at VulkanViewport should be moved to here?
-	enum { NUM_BUFFERS = 2 };
-#endif
-#endif
-
 	FVulkanSwapChain(VkInstance Instance, FVulkanDevice& InDevice, void* WindowHandle, EPixelFormat& InOutPixelFormat, uint32 Width, uint32 Height,
 		uint32* InOutDesiredNumBackBuffers, TArray<VkImage>& OutImages);
 
 	void Destroy();
 
-#if VULKAN_USE_NEW_COMMAND_BUFFERS
 	void Present(FVulkanQueue* Queue);
-#else
-#endif
 
 protected:
 	VkSwapchainKHR SwapChain;
@@ -41,17 +30,13 @@ protected:
 
 	VkSurfaceKHR Surface;
 
-#if VULKAN_USE_NEW_COMMAND_BUFFERS
 	int32 CurrentImageIndex;
+	int32 SemaphoreIndex;
 	TArray<FVulkanSemaphore*> ImageAcquiredSemaphore;
-	TArray<FVulkanSemaphore*> RenderingCompletedSemaphore;
-#else
-#endif
+	//TArray<FVulkanSemaphore*> RenderingCompletedSemaphore;
 
-#if VULKAN_USE_NEW_COMMAND_BUFFERS
-	int32 AcquireImageIndex();
+	int32 AcquireImageIndex(FVulkanSemaphore** OutSemaphore);
 	friend class FVulkanViewport;
-#else
-#endif
+
 	friend class FVulkanQueue;
 };
