Copying //UE4/Release-Staging-4.12 to //UE4/Dev-Main (Source: //UE4/Release-4.12 @ 2955635)

diff --git a/Engine/Source/Runtime/VulkanRHI/Private/VulkanTexture.cpp b/Engine/Source/Runtime/VulkanRHI/Private/VulkanTexture.cpp
--- a/Engine/Source/Runtime/VulkanRHI/Private/VulkanTexture.cpp
+++ b/Engine/Source/Runtime/VulkanRHI/Private/VulkanTexture.cpp
@@ -77,7 +77,7 @@ VkImage FVulkanSurface::CreateImage(
 	}
 
 	ImageCreateInfo.format = (VkFormat)GPixelFormats[InFormat].PlatformFormat;
-	if ((UEFlags & TexCreate_SRGB) == TexCreate_SRGB)
+	if ((UEFlags & TexCreate_SRGB) == TexCreate_SRGB && GMaxRHIFeatureLevel > ERHIFeatureLevel::ES3_1)
 	{
 		switch (ImageCreateInfo.format)
 		{
@@ -178,26 +178,32 @@ VkImage FVulkanSurface::CreateImage(
 		break;
 	}
 
+//#todo-rco: Verify flags work on newer Android drivers
+#if !PLATFORM_ANDROID
 	if (ImageCreateInfo.tiling == VK_IMAGE_TILING_LINEAR)
 	{
 		VkFormatFeatureFlags FormatFlags = InDevice.GetFormatProperties()[ImageCreateInfo.format].linearTilingFeatures;
 		if ((FormatFlags & VK_FORMAT_FEATURE_SAMPLED_IMAGE_BIT) == 0)
 		{
+			ensure((ImageCreateInfo.usage & VK_IMAGE_USAGE_SAMPLED_BIT) == 0);
 			ImageCreateInfo.usage &= ~VK_IMAGE_USAGE_SAMPLED_BIT;
 		}
 
 		if ((FormatFlags & VK_FORMAT_FEATURE_STORAGE_IMAGE_BIT) == 0)
 		{
+			ensure((ImageCreateInfo.usage & VK_IMAGE_USAGE_STORAGE_BIT) == 0);
 			ImageCreateInfo.usage &= ~VK_IMAGE_USAGE_STORAGE_BIT;
 		}
 
 		if ((FormatFlags & VK_FORMAT_FEATURE_COLOR_ATTACHMENT_BIT) == 0)
 		{
+			ensure((ImageCreateInfo.usage & VK_IMAGE_USAGE_COLOR_ATTACHMENT_BIT) == 0);
 			ImageCreateInfo.usage &= ~VK_IMAGE_USAGE_COLOR_ATTACHMENT_BIT;
 		}
 
 		if ((FormatFlags & VK_FORMAT_FEATURE_DEPTH_STENCIL_ATTACHMENT_BIT) == 0)
 		{
+			ensure((ImageCreateInfo.usage & VK_IMAGE_USAGE_DEPTH_STENCIL_ATTACHMENT_BIT) == 0);
 			ImageCreateInfo.usage &= ~VK_IMAGE_USAGE_DEPTH_STENCIL_ATTACHMENT_BIT;
 		}
 	}
@@ -206,24 +212,29 @@ VkImage FVulkanSurface::CreateImage(
 		VkFormatFeatureFlags FormatFlags = InDevice.GetFormatProperties()[ImageCreateInfo.format].optimalTilingFeatures;
 		if ((FormatFlags & VK_FORMAT_FEATURE_SAMPLED_IMAGE_BIT) == 0)
 		{
+			ensure((ImageCreateInfo.usage & VK_IMAGE_USAGE_SAMPLED_BIT) == 0);
 			ImageCreateInfo.usage &= ~VK_IMAGE_USAGE_SAMPLED_BIT;
 		}
 
 		if ((FormatFlags & VK_FORMAT_FEATURE_STORAGE_IMAGE_BIT) == 0)
 		{
+			ensure((ImageCreateInfo.usage & VK_IMAGE_USAGE_STORAGE_BIT) == 0);
 			ImageCreateInfo.usage &= ~VK_IMAGE_USAGE_STORAGE_BIT;
 		}
 
 		if ((FormatFlags & VK_FORMAT_FEATURE_COLOR_ATTACHMENT_BIT) == 0)
 		{
+			ensure((ImageCreateInfo.usage & VK_IMAGE_USAGE_COLOR_ATTACHMENT_BIT) == 0);
 			ImageCreateInfo.usage &= ~VK_IMAGE_USAGE_COLOR_ATTACHMENT_BIT;
 		}
 
 		if ((FormatFlags & VK_FORMAT_FEATURE_DEPTH_STENCIL_ATTACHMENT_BIT) == 0)
 		{
+			ensure((ImageCreateInfo.usage & VK_IMAGE_USAGE_DEPTH_STENCIL_ATTACHMENT_BIT) == 0);
 			ImageCreateInfo.usage &= ~VK_IMAGE_USAGE_DEPTH_STENCIL_ATTACHMENT_BIT;
 		}
 	}
+#endif
 
 	VkImage Image;
 	VERIFYVULKANRESULT(vkCreateImage(InDevice.GetInstanceHandle(), &ImageCreateInfo, nullptr, &Image));
@@ -320,11 +331,7 @@ FVulkanSurface::FVulkanSurface(FVulkanDevice& InDevice, VkImageViewType Resource
 	if ((ImageCreateInfo.usage & VK_IMAGE_USAGE_SAMPLED_BIT) && (UEFlags & (TexCreate_RenderTargetable | TexCreate_DepthStencilTargetable)))
 	{
 		const bool bTransitionToPresentable = ((UEFlags & TexCreate_Presentable) == TexCreate_Presentable);
-#if VULKAN_USE_NEW_COMMAND_BUFFERS
 		Clear(CreateInfo.ClearValueBinding, bTransitionToPresentable);
-#else
-		ClearBlocking(CreateInfo.ClearValueBinding, bTransitionToPresentable);
-#endif
 	}
 }
 
@@ -354,7 +361,7 @@ FVulkanSurface::FVulkanSurface(FVulkanDevice& InDevice, VkImageViewType Resource
 	AspectMask = GetAspectMask(InternalFormat);
 
 	// Purely informative patching, we know that "TexCreate_Presentable" uses optimal tiling
-	if (UEFlags & TexCreate_Presentable && GetTiling() == VK_IMAGE_TILING_MAX_ENUM)
+	if ((UEFlags & TexCreate_Presentable) == TexCreate_Presentable && GetTiling() == VK_IMAGE_TILING_MAX_ENUM)
 	{
 		Tiling = VK_IMAGE_TILING_OPTIMAL;
 	}
@@ -364,11 +371,7 @@ FVulkanSurface::FVulkanSurface(FVulkanDevice& InDevice, VkImageViewType Resource
 		if (UEFlags & (TexCreate_RenderTargetable | TexCreate_DepthStencilTargetable))
 		{
 			const bool bTransitionToPresentable = ((UEFlags & TexCreate_Presentable) == TexCreate_Presentable);
-#if VULKAN_USE_NEW_COMMAND_BUFFERS
 			Clear(CreateInfo.ClearValueBinding, bTransitionToPresentable);
-#else
-			ClearBlocking(CreateInfo.ClearValueBinding, bTransitionToPresentable);
-#endif
 		}
 	}
 }
@@ -607,7 +610,6 @@ VkImageAspectFlags FVulkanSurface::GetAspectMask() const
 	return AspectMask;
 }
 
-#if VULKAN_USE_NEW_COMMAND_BUFFERS
 void FVulkanSurface::Clear(const FClearValueBinding& ClearValueBinding, bool bTransitionToPresentable)
 {
 #if VULKAN_CLEAR_SURFACE_ON_CREATE
@@ -619,14 +621,10 @@ void FVulkanSurface::Clear(const FClearValueBinding& ClearValueBinding, bool bTr
 	Range.baseArrayLayer = 0;
 	Range.layerCount = ViewType == VK_IMAGE_VIEW_TYPE_CUBE ? 6 : 1;
 
-#if VULKAN_USE_NEW_COMMAND_BUFFERS
-	FVulkanCmdBuffer* CmdBuffer = Device->GetImmediateContext().GetCommandBufferManager()->GetActiveCmdBuffer();
+	//#todo-rco: This function is only used during loading currently, if used for regular RHIClear then use the ActiveCmdBuffer
+	FVulkanCmdBuffer* CmdBuffer = Device->GetImmediateContext().GetCommandBufferManager()->GetUploadCmdBuffer();
 	ensure(CmdBuffer->IsOutsideRenderPass());
-#else
-	//#todo-rco: FIX ME!
-	FVulkanCmdBuffer* CmdBuffer = Device->GetImmediateContext().GetCommandBufferManager()->Create();
-	CmdBuffer->Begin();
-#endif
+
 	if(Range.aspectMask == VK_IMAGE_ASPECT_COLOR_BIT)
 	{
 		VkClearColorValue Color;
@@ -653,61 +651,8 @@ void FVulkanSurface::Clear(const FClearValueBinding& ClearValueBinding, bool bTr
 		VulkanSetImageLayoutSimple(CmdBuffer->GetHandle(), Image, VK_IMAGE_LAYOUT_GENERAL, VK_IMAGE_LAYOUT_DEPTH_STENCIL_ATTACHMENT_OPTIMAL, Range.aspectMask);
 		ImageLayout = VK_IMAGE_LAYOUT_DEPTH_STENCIL_ATTACHMENT_OPTIMAL;
 	}
-
-#if VULKAN_USE_NEW_COMMAND_BUFFERS
-#else
-	Device->EndCommandBufferBlock(CmdBuffer);
-#endif
-#endif
-}
-#else
-void FVulkanSurface::ClearBlocking(const FClearValueBinding& ClearValueBinding, bool bTransitionToPresentable)
-{
-#if VULKAN_CLEAR_SURFACE_ON_CREATE
-	VkImageSubresourceRange Range;
-	FMemory::Memzero(Range);
-	Range.aspectMask = GetAspectMask();
-	Range.baseMipLevel = 0;
-	Range.levelCount = NumMips;
-	Range.baseArrayLayer = 0;
-	Range.layerCount = ViewType == VK_IMAGE_VIEW_TYPE_CUBE ? 6 : 1;
-
-	//#todo-rco: FIX ME!
-	FVulkanCmdBuffer* Cmd = Device->GetImmediateContext().GetCommandBufferManager()->Create();
-	Cmd->Begin();
-	
-	if(Range.aspectMask == VK_IMAGE_ASPECT_COLOR_BIT)
-	{
-		VkClearColorValue Color;
-		FMemory::Memzero(Color);
-		Color.float32[0] = ClearValueBinding.Value.Color[0];
-		Color.float32[1] = ClearValueBinding.Value.Color[1];
-		Color.float32[2] = ClearValueBinding.Value.Color[2];
-		Color.float32[3] = ClearValueBinding.Value.Color[3];
-		VulkanSetImageLayoutSimple(Cmd->GetHandle(), Image, ImageLayout, VK_IMAGE_LAYOUT_GENERAL);
-		vkCmdClearColorImage(Cmd->GetHandle(), Image, VK_IMAGE_LAYOUT_GENERAL, &Color, 1, &Range);
-		ImageLayout = bTransitionToPresentable ? VK_IMAGE_LAYOUT_PRESENT_SRC_KHR : VK_IMAGE_LAYOUT_COLOR_ATTACHMENT_OPTIMAL;
-		VulkanSetImageLayoutSimple(Cmd->GetHandle(), Image, VK_IMAGE_LAYOUT_GENERAL, ImageLayout);
-	}
-	else
-	{
-		check(Range.aspectMask & (VK_IMAGE_ASPECT_DEPTH_BIT | VK_IMAGE_ASPECT_STENCIL_BIT));
-		check(!bTransitionToPresentable);
-		VkClearDepthStencilValue Value;
-		FMemory::Memzero(Value);
-		Value.depth = ClearValueBinding.Value.DSValue.Depth;
-		Value.stencil = ClearValueBinding.Value.DSValue.Stencil;
-		VulkanSetImageLayoutSimple(Cmd->GetHandle(), Image, ImageLayout, VK_IMAGE_LAYOUT_GENERAL, Range.aspectMask);
-		vkCmdClearDepthStencilImage(Cmd->GetHandle(), Image, VK_IMAGE_LAYOUT_DEPTH_STENCIL_ATTACHMENT_OPTIMAL, &Value, 1, &Range);
-		VulkanSetImageLayoutSimple(Cmd->GetHandle(), Image, VK_IMAGE_LAYOUT_GENERAL, VK_IMAGE_LAYOUT_DEPTH_STENCIL_ATTACHMENT_OPTIMAL, Range.aspectMask);
-		ImageLayout = VK_IMAGE_LAYOUT_DEPTH_STENCIL_ATTACHMENT_OPTIMAL;
-	}
-
-	Device->EndCommandBufferBlock(Cmd);
 #endif
 }
-#endif
-
 
 /*-----------------------------------------------------------------------------
 	Texture allocator support.
@@ -847,25 +792,10 @@ void FVulkanDynamicRHI::RHIUnlockTexture2D(FTexture2DRHIParamRef TextureRHI,uint
 		checkf(bFound, TEXT("Texture was not locked!"));
 	}
 
-#if VULKAN_USE_NEW_COMMAND_BUFFERS
-	FVulkanCmdBuffer* CmdBuffer = Device->GetImmediateContext().GetCommandBufferManager()->GetActiveCmdBuffer();
+	FVulkanCmdBuffer* CmdBuffer = Device->GetImmediateContext().GetCommandBufferManager()->GetUploadCmdBuffer();
 	ensure(CmdBuffer->IsOutsideRenderPass());
 	VkCommandBuffer StagingCommandBuffer = CmdBuffer->GetHandle();
-#else
-	VkCommandBuffer StagingCommandBuffer;
-	VkCommandBufferAllocateInfo CmdBufferCreateInfo;
-	FMemory::Memzero(CmdBufferCreateInfo);
-	CmdBufferCreateInfo.sType = VK_STRUCTURE_TYPE_COMMAND_BUFFER_ALLOCATE_INFO;
-	CmdBufferCreateInfo.commandPool = Device->GetImmediateContext().GetCommandBufferManager()->GetHandle();
-	CmdBufferCreateInfo.commandBufferCount = 1;
-	CmdBufferCreateInfo.level = VK_COMMAND_BUFFER_LEVEL_PRIMARY;
-	VERIFYVULKANRESULT(vkAllocateCommandBuffers(LogicalDevice, &CmdBufferCreateInfo, &StagingCommandBuffer));
-
-	VkCommandBufferBeginInfo BufferBeginInfo;
-	FMemory::Memzero(BufferBeginInfo);
-	BufferBeginInfo.sType = VK_STRUCTURE_TYPE_COMMAND_BUFFER_BEGIN_INFO;
-	VERIFYVULKANRESULT(vkBeginCommandBuffer(StagingCommandBuffer, &BufferBeginInfo));
-#endif
+
 	EPixelFormat Format = Texture->Surface.Format;
 	uint32 MipWidth = FMath::Max<uint32>(Texture->Surface.Width >> MipIndex, GPixelFormats[Format].BlockSizeX);
 	uint32 MipHeight = FMath::Max<uint32>(Texture->Surface.Height >> MipIndex, GPixelFormats[Format].BlockSizeY);
@@ -877,7 +807,7 @@ void FVulkanDynamicRHI::RHIUnlockTexture2D(FTexture2DRHIParamRef TextureRHI,uint
 	SubresourceRange.levelCount = 1;
 	//SubresourceRange.baseArrayLayer = 0;
 	SubresourceRange.layerCount = 1;
-	VulkanSetImageLayout(StagingCommandBuffer, Texture->Surface.Image, VK_IMAGE_LAYOUT_UNDEFINED, VK_IMAGE_LAYOUT_GENERAL, SubresourceRange);
+	VulkanSetImageLayout(StagingCommandBuffer, Texture->Surface.Image, VK_IMAGE_LAYOUT_UNDEFINED, VK_IMAGE_LAYOUT_TRANSFER_DST_OPTIMAL, SubresourceRange);
 
 	VkBufferImageCopy Region;
 	FMemory::Memzero(Region);
@@ -892,28 +822,11 @@ void FVulkanDynamicRHI::RHIUnlockTexture2D(FTexture2DRHIParamRef TextureRHI,uint
 	Region.imageExtent.width = MipWidth;
 	Region.imageExtent.height = MipHeight;
 	Region.imageExtent.depth = 1;
-	vkCmdCopyBufferToImage(StagingCommandBuffer, StagingBuffer->GetHandle(), Texture->Surface.Image, VK_IMAGE_LAYOUT_GENERAL, 1, &Region);
+	vkCmdCopyBufferToImage(StagingCommandBuffer, StagingBuffer->GetHandle(), Texture->Surface.Image, VK_IMAGE_LAYOUT_TRANSFER_DST_OPTIMAL, 1, &Region);
 
-	VulkanSetImageLayout(StagingCommandBuffer, Texture->Surface.Image, VK_IMAGE_LAYOUT_GENERAL, VK_IMAGE_LAYOUT_SHADER_READ_ONLY_OPTIMAL, SubresourceRange);
+	VulkanSetImageLayout(StagingCommandBuffer, Texture->Surface.Image, VK_IMAGE_LAYOUT_TRANSFER_DST_OPTIMAL, VK_IMAGE_LAYOUT_SHADER_READ_ONLY_OPTIMAL, SubresourceRange);
 
-#if VULKAN_USE_NEW_COMMAND_BUFFERS
 	Device->GetStagingManager().ReleaseBuffer(CmdBuffer, StagingBuffer);
-#else
-	VERIFYVULKANRESULT(vkEndCommandBuffer(StagingCommandBuffer));
-
-	auto* Fence = Device->GetFenceManager().AllocateFence();
-	Device->GetQueue()->Submit2(StagingCommandBuffer, Fence);
-
-	bool bSuccess = Device->GetFenceManager().WaitForFence(Fence, 0xFFFFFFFF);
-	check(bSuccess);
-
-	Device->GetFenceManager().ReleaseFence(Fence);
-
-	VERIFYVULKANRESULT(vkResetCommandBuffer(StagingCommandBuffer, 0x0));
-
-	vkFreeCommandBuffers(LogicalDevice, Device->GetImmediateContext().GetCommandBufferManager()->GetHandle(), 1, &StagingCommandBuffer);
-	Device->GetStagingManager().ReleaseBuffer(StagingBuffer);
-#endif
 }
 
 void* FVulkanDynamicRHI::RHILockTexture2DArray(FTexture2DArrayRHIParamRef TextureRHI,uint32 TextureIndex,uint32 MipIndex,EResourceLockMode LockMode,uint32& DestStride,bool bLockWithinMiptail)
@@ -930,21 +843,65 @@ void FVulkanDynamicRHI::RHIUnlockTexture2DArray(FTexture2DArrayRHIParamRef Textu
 
 void FVulkanDynamicRHI::RHIUpdateTexture2D(FTexture2DRHIParamRef TextureRHI, uint32 MipIndex, const struct FUpdateTextureRegion2D& UpdateRegion, uint32 SourcePitch, const uint8* SourceData)
 {
-	VULKAN_SIGNAL_UNIMPLEMENTED();
-#if 0
 	FVulkanTexture2D* Texture = ResourceCast(TextureRHI);
-#endif
-
+	VULKAN_SIGNAL_UNIMPLEMENTED();
 }
 
-void FVulkanDynamicRHI::RHIUpdateTexture3D(FTexture3DRHIParamRef TextureRHI,uint32 MipIndex,const FUpdateTextureRegion3D& UpdateRegion,uint32 SourceRowPitch,uint32 SourceDepthPitch,const uint8* SourceData)
-{	
-	VULKAN_SIGNAL_UNIMPLEMENTED();
-#if 0
+void FVulkanDynamicRHI::RHIUpdateTexture3D(FTexture3DRHIParamRef TextureRHI, uint32 MipIndex, const FUpdateTextureRegion3D& UpdateRegion, uint32 SourceRowPitch, uint32 SourceDepthPitch, const uint8* SourceData)
+{
 	FVulkanTexture3D* Texture = ResourceCast(TextureRHI);
-#endif
+
+	EPixelFormat PixelFormat = Texture->GetFormat();
+	check(GPixelFormats[PixelFormat].BlockSizeX == 1);
+	check(GPixelFormats[PixelFormat].BlockSizeY == 1);
+
+	VkFormat Format = (VkFormat)GPixelFormats[PixelFormat].PlatformFormat;
+
+	// TO DO - add appropriate offsets to source data when necessary
+	check(UpdateRegion.SrcX == 0);
+	check(UpdateRegion.SrcY == 0);
+	check(UpdateRegion.SrcZ == 0);
+
+	auto& Context = Device->GetImmediateContext();
+	const auto& Limits = Device->GetLimits();
+	const uint32 AlignedSourcePitch = Align(SourceRowPitch, Limits.optimalBufferCopyRowPitchAlignment);
+	const uint32 BufferSize = Align(UpdateRegion.Height * UpdateRegion.Depth * AlignedSourcePitch, Limits.minMemoryMapAlignment);
+
+	auto* StagingBuffer = Device->GetStagingManager().AcquireBuffer(BufferSize);
+	void* Memory = StagingBuffer->GetMappedPointer();
+
+	uint32 Size = UpdateRegion.Height * UpdateRegion.Depth * SourceRowPitch;
+	check(Size <= BufferSize);
+
+	//#todo-rco: Probably needs some work (see D3D12)
+	//FMemory::Memcpy(Memory, SourceData, Size);
+	FMemory::Memzero(Memory, Size);
+
+	VkBufferImageCopy Region;
+	FMemory::Memzero(Region);
+	//#todo-rco: Might need an offset here?
+	check(UpdateRegion.SrcX == 0);
+	check(UpdateRegion.SrcY == 0);
+	check(UpdateRegion.SrcZ == 0);
+	//Region.bufferOffset = 0;
+	//Region.bufferRowLength = 0;
+	//Region.bufferImageHeight = 0;
+	Region.imageSubresource.aspectMask = VK_IMAGE_ASPECT_COLOR_BIT;
+	Region.imageSubresource.mipLevel = MipIndex;
+	Region.imageSubresource.baseArrayLayer = 0;
+	Region.imageSubresource.layerCount = 1;
+	Region.imageOffset.x = UpdateRegion.SrcX;
+	Region.imageOffset.y = UpdateRegion.SrcY;
+	Region.imageOffset.z = UpdateRegion.SrcZ;
+	Region.imageExtent.width = UpdateRegion.Width;
+	Region.imageExtent.height = UpdateRegion.Height;
+	Region.imageExtent.depth = UpdateRegion.Depth;
+	vkCmdCopyBufferToImage(Device->GetImmediateContext().GetCommandBufferManager()->GetUploadCmdBuffer()->GetHandle(), StagingBuffer->GetHandle(), Texture->Surface.Image, VK_IMAGE_LAYOUT_GENERAL, 1, &Region);
+
+	Device->GetStagingManager().ReleaseBuffer(nullptr, StagingBuffer);
 }
 
+
 void FVulkanTextureView::Create(FVulkanDevice& Device, FVulkanSurface& Surface, VkImageViewType ViewType, VkFormat Format, uint32 NumMips)
 {
 	VkImageViewCreateInfo ViewInfo;
@@ -1030,15 +987,31 @@ FVulkanTextureBase::FVulkanTextureBase(FVulkanDevice& Device, VkImageViewType Re
 #endif
 
 	// For now only 3D textures can have bulk-data parsed on creation time.
-	if(!CreateInfo.BulkData || Surface.GetTiling() != VK_IMAGE_TILING_OPTIMAL)
+	if (!CreateInfo.BulkData)
 	{
 		return;
 	}
 
+	if (Surface.GetTiling() != VK_IMAGE_TILING_OPTIMAL)
+	{
+		return;
+	}
+
+	check(Surface.GetViewType() == VK_IMAGE_VIEW_TYPE_3D);
+
 	// Transfer bulk data
+	VulkanRHI::FStagingBuffer* StagingBuffer = Device.GetStagingManager().AcquireBuffer(CreateInfo.BulkData->GetResourceBulkDataSize());
+	void* Data = StagingBuffer->GetMappedPointer();
+
+	// Do copy
+	FMemory::Memcpy(Data, CreateInfo.BulkData->GetResourceBulkData(), CreateInfo.BulkData->GetResourceBulkDataSize());
+
+	FVulkanCmdBuffer* CmdBuffer = Device.GetImmediateContext().GetCommandBufferManager()->GetUploadCmdBuffer();
+	ensure(CmdBuffer->IsOutsideRenderPass());
 
 	VkBufferImageCopy Region;
 	FMemory::Memzero(Region);
+	//#todo-rco: Use real Buffer offset when switching to suballocations!
 	Region.bufferOffset = 0;
 	Region.bufferRowLength = Surface.Width;
 	Region.bufferImageHeight = Surface.Height;
@@ -1052,50 +1025,6 @@ FVulkanTextureBase::FVulkanTextureBase(FVulkanDevice& Device, VkImageViewType Re
 	Region.imageExtent.height = Region.bufferImageHeight;
 	Region.imageExtent.depth = Surface.Depth;
 
-	// Temporary buffer which will hold bulkdata
-	VkBufferCreateInfo BufferCreateInfo;
-	FMemory::Memzero(BufferCreateInfo);
-	BufferCreateInfo.sType = VK_STRUCTURE_TYPE_BUFFER_CREATE_INFO;
-	BufferCreateInfo.flags = 0;
-	//BufferCreateInfo.pQueueFamilyIndices = NULL;
-	//BufferCreateInfo.queueFamilyIndexCount = 0;
-	//BufferCreateInfo.sharingMode = VK_SHARING_MODE_EXCLUSIVE;
-	BufferCreateInfo.size = CreateInfo.BulkData->GetResourceBulkDataSize();
-	BufferCreateInfo.usage = VK_BUFFER_USAGE_TRANSFER_SRC_BIT;
-
-	VkBuffer Buffer;
-	VERIFYVULKANRESULT(vkCreateBuffer(Device.GetInstanceHandle(), &BufferCreateInfo, NULL, &Buffer));
-
-	VkMemoryRequirements MemReqs;
-	vkGetBufferMemoryRequirements(Device.GetInstanceHandle(), Buffer, &MemReqs);
-
-	// Allocate device gmem
-	VkDeviceMemory Mem;
-
-	//#todo-rco: USe StagingBuffers
-	VulkanRHI::FDeviceMemoryAllocation* Allocation = Device.GetMemoryManager().Alloc(MemReqs.size, MemReqs.memoryTypeBits, VK_MEMORY_PROPERTY_HOST_VISIBLE_BIT, __FILE__, __LINE__);
-	Mem = Allocation->GetHandle();
-
-	//FPlatformMisc::LowLevelOutputDebugStringf(TEXT("** vkBindBufferMemory Buf %p MemHandle %p MemOffset %d Size %u\n"), (void*)Buffer, (void*)Mem, (uint32)0, (uint32)MemReqs.size);
-	VERIFYVULKANRESULT(vkBindBufferMemory(Device.GetInstanceHandle(), Buffer, Mem, 0));
-
-	void* Data = nullptr;
-	VERIFYVULKANRESULT(vkMapMemory(Device.GetInstanceHandle(), Mem, 0, MemReqs.size, 0, &Data));
-
-	// Do copy
-	FMemory::Memcpy(Data, CreateInfo.BulkData->GetResourceBulkData(), CreateInfo.BulkData->GetResourceBulkDataSize());
-
-	vkUnmapMemory(Device.GetInstanceHandle(), Mem);
-
-#if VULKAN_USE_NEW_COMMAND_BUFFERS
-	FVulkanCmdBuffer* CmdBuffer = Device.GetImmediateContext().GetCommandBufferManager()->GetActiveCmdBuffer();
-	ensure(CmdBuffer->IsOutsideRenderPass());
-#else
-	//#todo-rco: FIX ME?
-	FVulkanCmdBuffer* CmdBuffer = Device.GetImmediateContext().GetCommandBufferManager()->Create();
-	CmdBuffer->Begin();
-#endif
-
 	VkImageSubresourceRange SubresourceRange;
 	FMemory::Memzero(SubresourceRange);
 	SubresourceRange.aspectMask = VK_IMAGE_ASPECT_COLOR_BIT;
@@ -1103,24 +1032,14 @@ FVulkanTextureBase::FVulkanTextureBase(FVulkanDevice& Device, VkImageViewType Re
 	SubresourceRange.levelCount = Surface.GetNumMips();
 	//SubresourceRange.baseArrayLayer = 0;
 	SubresourceRange.layerCount = ArraySize;
-	VulkanSetImageLayout(CmdBuffer->GetHandle(), Surface.Image, VK_IMAGE_LAYOUT_UNDEFINED, VK_IMAGE_LAYOUT_GENERAL, SubresourceRange);
+	VulkanSetImageLayout(CmdBuffer->GetHandle(), Surface.Image, VK_IMAGE_LAYOUT_UNDEFINED, VK_IMAGE_LAYOUT_TRANSFER_DST_OPTIMAL, SubresourceRange);
 
 	// Copy buffer to image
-	vkCmdCopyBufferToImage(CmdBuffer->GetHandle(), Buffer, Surface.Image, VK_IMAGE_LAYOUT_GENERAL, 1, &Region);
-
-	VulkanSetImageLayout(CmdBuffer->GetHandle(), Surface.Image, VK_IMAGE_LAYOUT_GENERAL, VK_IMAGE_LAYOUT_SHADER_READ_ONLY_OPTIMAL, SubresourceRange);
-
-#if VULKAN_USE_NEW_COMMAND_BUFFERS
-#else
-	Device.GetQueue()->SubmitBlocking(CmdBuffer);
-	//#todo-rco: FIX ME!
-	Device.GetImmediateContext().GetCommandBufferManager()->Destroy(CmdBuffer);
+	vkCmdCopyBufferToImage(CmdBuffer->GetHandle(), StagingBuffer->GetHandle(), Surface.Image, VK_IMAGE_LAYOUT_TRANSFER_DST_OPTIMAL, 1, &Region);
 
-	vkDestroyBuffer(Device.GetInstanceHandle(), Buffer, nullptr);
+	VulkanSetImageLayout(CmdBuffer->GetHandle(), Surface.Image, VK_IMAGE_LAYOUT_TRANSFER_DST_OPTIMAL, VK_IMAGE_LAYOUT_SHADER_READ_ONLY_OPTIMAL, SubresourceRange);
 
-#endif
-	//#todo-rco: USe StagingBuffers
-	Device.GetMemoryManager().Free(Allocation);
+	Device.GetStagingManager().ReleaseBuffer(CmdBuffer, StagingBuffer);
 }
 
 FVulkanTextureBase::FVulkanTextureBase(FVulkanDevice& Device, VkImageViewType ResourceType, EPixelFormat Format, uint32 SizeX, uint32 SizeY, uint32 SizeZ, VkImage InImage, VkDeviceMemory InMem, uint32 UEFlags, const FRHIResourceCreateInfo& CreateInfo):
@@ -1244,25 +1163,9 @@ void FVulkanDynamicRHI::RHIUnlockTextureCubeFace(FTextureCubeRHIParamRef Texture
 		checkf(bFound, TEXT("Texture was not locked!"));
 	}
 
-#if VULKAN_USE_NEW_COMMAND_BUFFERS
-	FVulkanCmdBuffer* CmdBuffer = Device->GetImmediateContext().GetCommandBufferManager()->GetActiveCmdBuffer();
+	FVulkanCmdBuffer* CmdBuffer = Device->GetImmediateContext().GetCommandBufferManager()->GetUploadCmdBuffer();
 	ensure(CmdBuffer->IsOutsideRenderPass());
 	VkCommandBuffer StagingCommandBuffer = CmdBuffer->GetHandle();
-#else
-	VkCommandBuffer StagingCommandBuffer;
-	VkCommandBufferAllocateInfo CmdBufferCreateInfo;
-	FMemory::Memzero(CmdBufferCreateInfo);
-	CmdBufferCreateInfo.sType = VK_STRUCTURE_TYPE_COMMAND_BUFFER_ALLOCATE_INFO;
-	CmdBufferCreateInfo.commandPool = Device->GetImmediateContext().GetCommandBufferManager()->GetHandle();
-	CmdBufferCreateInfo.commandBufferCount = 1;
-	CmdBufferCreateInfo.level = VK_COMMAND_BUFFER_LEVEL_PRIMARY;
-	VERIFYVULKANRESULT(vkAllocateCommandBuffers(LogicalDevice, &CmdBufferCreateInfo, &StagingCommandBuffer));
-
-	VkCommandBufferBeginInfo BufferBeginInfo;
-	FMemory::Memzero(BufferBeginInfo);
-	BufferBeginInfo.sType = VK_STRUCTURE_TYPE_COMMAND_BUFFER_BEGIN_INFO;
-	VERIFYVULKANRESULT(vkBeginCommandBuffer(StagingCommandBuffer, &BufferBeginInfo));
-#endif
 	EPixelFormat Format = Texture->Surface.Format;
 	uint32 MipWidth = FMath::Max<uint32>(Texture->Surface.Width >> MipIndex, GPixelFormats[Format].BlockSizeX);
 	uint32 MipHeight = FMath::Max<uint32>(Texture->Surface.Height >> MipIndex, GPixelFormats[Format].BlockSizeY);
@@ -1274,7 +1177,7 @@ void FVulkanDynamicRHI::RHIUnlockTextureCubeFace(FTextureCubeRHIParamRef Texture
 	SubresourceRange.levelCount = 1;
 	SubresourceRange.baseArrayLayer = FaceIndex;
 	SubresourceRange.layerCount = 1;
-	VulkanSetImageLayout(StagingCommandBuffer, Texture->Surface.Image, VK_IMAGE_LAYOUT_UNDEFINED, VK_IMAGE_LAYOUT_GENERAL, SubresourceRange);
+	VulkanSetImageLayout(StagingCommandBuffer, Texture->Surface.Image, VK_IMAGE_LAYOUT_UNDEFINED, VK_IMAGE_LAYOUT_TRANSFER_DST_OPTIMAL, SubresourceRange);
 
 	VkBufferImageCopy Region;
 	FMemory::Memzero(Region);
@@ -1289,28 +1192,11 @@ void FVulkanDynamicRHI::RHIUnlockTextureCubeFace(FTextureCubeRHIParamRef Texture
 	Region.imageExtent.width = MipWidth;
 	Region.imageExtent.height = MipHeight;
 	Region.imageExtent.depth = 1;
-	vkCmdCopyBufferToImage(StagingCommandBuffer, StagingBuffer->GetHandle(), Texture->Surface.Image, VK_IMAGE_LAYOUT_GENERAL, 1, &Region);
+	vkCmdCopyBufferToImage(StagingCommandBuffer, StagingBuffer->GetHandle(), Texture->Surface.Image, VK_IMAGE_LAYOUT_TRANSFER_DST_OPTIMAL, 1, &Region);
 
-	VulkanSetImageLayout(StagingCommandBuffer, Texture->Surface.Image, VK_IMAGE_LAYOUT_GENERAL, VK_IMAGE_LAYOUT_SHADER_READ_ONLY_OPTIMAL, SubresourceRange);
+	VulkanSetImageLayout(StagingCommandBuffer, Texture->Surface.Image, VK_IMAGE_LAYOUT_TRANSFER_DST_OPTIMAL, VK_IMAGE_LAYOUT_SHADER_READ_ONLY_OPTIMAL, SubresourceRange);
 
-#if VULKAN_USE_NEW_COMMAND_BUFFERS
 	Device->GetStagingManager().ReleaseBuffer(CmdBuffer, StagingBuffer);
-#else
-	VERIFYVULKANRESULT(vkEndCommandBuffer(StagingCommandBuffer));
-
-	auto* Fence = Device->GetFenceManager().AllocateFence();
-	Device->GetQueue()->Submit2(StagingCommandBuffer, Fence);
-
-	bool bSuccess = Device->GetFenceManager().WaitForFence(Fence, 0xFFFFFFFF);
-	check(bSuccess);
-
-	Device->GetFenceManager().ReleaseFence(Fence);
-
-	VERIFYVULKANRESULT(vkResetCommandBuffer(StagingCommandBuffer, 0x0));
-
-	vkFreeCommandBuffers(LogicalDevice, Device->GetImmediateContext().GetCommandBufferManager()->GetHandle(), 1, &StagingCommandBuffer);
-	Device->GetStagingManager().ReleaseBuffer(StagingBuffer);
-#endif
 }
 
 void FVulkanDynamicRHI::RHIBindDebugLabelName(FTextureRHIParamRef TextureRHI, const TCHAR* Name)
diff --git a/Engine/Source/Runtime/VulkanRHI/Private/VulkanUAV.cpp b/Engine/Source/Runtime/VulkanRHI/Private/VulkanUAV.cpp
--- a/Engine/Source/Runtime/VulkanRHI/Private/VulkanUAV.cpp
+++ b/Engine/Source/Runtime/VulkanRHI/Private/VulkanUAV.cpp
@@ -20,8 +20,8 @@ void FVulkanShaderResourceView::UpdateView(FVulkanDevice* Device)
 		{
 			SCOPE_CYCLE_COUNTER(STAT_VulkanSRVUpdateTime);
 			// thanks to ref counting, overwriting the buffer will toss the old view
-			BufferView = new FVulkanBufferView();
-			BufferView->Create(*Device, SourceVertexBuffer.GetReference(), BufferViewFormat, SourceVertexBuffer->GetOffset(), SourceVertexBuffer->GetSize());
+			BufferView = new FVulkanBufferView(Device);
+			BufferView->Create(SourceVertexBuffer.GetReference(), BufferViewFormat, SourceVertexBuffer->GetOffset(), SourceVertexBuffer->GetSize());
 		}
 	}
 }
diff --git a/Engine/Source/Runtime/VulkanRHI/Private/VulkanUtil.cpp b/Engine/Source/Runtime/VulkanRHI/Private/VulkanUtil.cpp
--- a/Engine/Source/Runtime/VulkanRHI/Private/VulkanUtil.cpp
+++ b/Engine/Source/Runtime/VulkanRHI/Private/VulkanUtil.cpp
@@ -68,7 +68,7 @@ void FVulkanGPUTiming::StartTiming()
 		if (StartTimestamp >= 0 && EndTimestamp >= 0)
 		{
 			auto& State = GTimestampQueryPool->Device->GetPendingState();
-#if VULKAN_USE_NEW_COMMAND_BUFFERS
+#if 1//VULKAN_USE_NEW_COMMAND_BUFFERS
 			check(0);
 #else
 			GTimestampQueryPool->WriteTimestamp(State.GetCommandBuffer(), StartTimestamp, VK_PIPELINE_STAGE_ALL_COMMANDS_BIT);
@@ -90,7 +90,7 @@ void FVulkanGPUTiming::EndTiming()
 		if (StartTimestamp >= 0 && EndTimestamp >= 0)
 		{
 			auto& State = GTimestampQueryPool->Device->GetPendingState();
-#if VULKAN_USE_NEW_COMMAND_BUFFERS
+#if 1//VULKAN_USE_NEW_COMMAND_BUFFERS
 			check(0);
 #else
 			GTimestampQueryPool->WriteTimestamp(State.GetCommandBuffer(), EndTimestamp, VK_PIPELINE_STAGE_ALL_COMMANDS_BIT);
diff --git a/Engine/Source/Runtime/VulkanRHI/Private/VulkanViewport.cpp b/Engine/Source/Runtime/VulkanRHI/Private/VulkanViewport.cpp
--- a/Engine/Source/Runtime/VulkanRHI/Private/VulkanViewport.cpp
+++ b/Engine/Source/Runtime/VulkanRHI/Private/VulkanViewport.cpp
@@ -11,14 +11,15 @@
 #include "VulkanContext.h"
 
 
-FVulkanViewport::FVulkanViewport(FVulkanDynamicRHI* InRHI, void* WindowHandle, uint32 InSizeX, uint32 InSizeY, bool bInIsFullscreen, EPixelFormat InPreferredPixelFormat) :
-	RHI(InRHI),
-	SizeX(InSizeX),
-	SizeY(InSizeY),
-	bIsFullscreen(bInIsFullscreen),
-	PixelFormat(InPreferredPixelFormat),
-	CurrentBackBuffer(-1),
-	SwapChain(nullptr)
+FVulkanViewport::FVulkanViewport(FVulkanDynamicRHI* InRHI, void* WindowHandle, uint32 InSizeX, uint32 InSizeY, bool bInIsFullscreen, EPixelFormat InPreferredPixelFormat)
+	: RHI(InRHI)
+	, SizeX(InSizeX)
+	, SizeY(InSizeY)
+	, bIsFullscreen(bInIsFullscreen)
+	, PixelFormat(InPreferredPixelFormat)
+	, CurrentBackBuffer(-1)
+	, SwapChain(nullptr)
+	, PendingImageSemaphore(nullptr)
 {
 	check(IsInGameThread());
 	RHI->Viewports.Add(this);
@@ -37,13 +38,8 @@ FVulkanViewport::FVulkanViewport(FVulkanDynamicRHI* InRHI, void* WindowHandle, u
 
 	check(Images.Num() == NUM_BUFFERS);
 
-#if VULKAN_USE_NEW_COMMAND_BUFFERS
 	FVulkanCmdBuffer* CmdBuffer = RHI->Device->GetImmediateContext().GetCommandBufferManager()->GetActiveCmdBuffer();
 	check(CmdBuffer->IsOutsideRenderPass());
-#else
-	FVulkanCmdBuffer* CmdBuffer = RHI->Device->GetImmediateContext().GetCommandBufferManager()->Create();
-	CmdBuffer->Begin();
-#endif
 	for (int32 Index = 0, Count = Images.Num(); Index < Count; ++Index)
 	{
 		VkImage Image = Images[Index];
@@ -53,20 +49,7 @@ FVulkanViewport::FVulkanViewport(FVulkanDynamicRHI* InRHI, void* WindowHandle, u
 
 		FName Name = FName(*FString::Printf(TEXT("BackBuffer%d"), Index));
 		BackBuffers[Index]->SetName(Name);
-
-#if VULKAN_USE_NEW_COMMAND_BUFFERS
-#else
-#if !PLATFORM_ANDROID
-		// Set to present src as the next calls will transition to color optimal
-		checkf(BackBuffers[Index]->Surface.ImageLayout == VK_IMAGE_LAYOUT_PRESENT_SRC_KHR, TEXT("BackBuffers[%d]->Surface.ImageLayout is %d"), Index, BackBuffers[Index]->Surface.ImageLayout);
-#endif
-#endif
 	}
-
-#if VULKAN_USE_NEW_COMMAND_BUFFERS
-#else
-	RHI->Device->EndCommandBufferBlock(CmdBuffer);
-#endif
 }
 
 FVulkanViewport::~FVulkanViewport()
@@ -83,9 +66,10 @@ FVulkanViewport::~FVulkanViewport()
 	RHI->Viewports.Remove(this);
 }
 
+
 FVulkanFramebuffer::FVulkanFramebuffer(FVulkanDevice& Device, const FRHISetRenderTargetsInfo& InRTInfo, const FVulkanRenderTargetLayout& RTLayout, const FVulkanRenderPass& RenderPass)
-:	Framebuffer(VK_NULL_HANDLE)
-,	NumColorAttachments(0u)
+	: Framebuffer(VK_NULL_HANDLE)
+	, NumColorAttachments(0)
 {
 	Attachments.Empty(RTLayout.GetNumAttachments());
 
@@ -101,7 +85,7 @@ FVulkanFramebuffer::FVulkanFramebuffer(FVulkanDevice& Device, const FRHISetRende
 	#if VULKAN_USE_MSAA_RESOLVE_ATTACHMENTS
 		if (Texture->MSAASurface)
 		{
-#if VULKAN_USE_NEW_COMMAND_BUFFERS
+#if 1//VULKAN_USE_NEW_COMMAND_BUFFERS
 			check(0);
 #endif
 			Attachments.Add(Texture->MSAAView.View);
@@ -112,7 +96,7 @@ FVulkanFramebuffer::FVulkanFramebuffer(FVulkanDevice& Device, const FRHISetRende
 		    Barrier.sType = VK_STRUCTURE_TYPE_IMAGE_MEMORY_BARRIER;
 		    //Barrier.pNext = NULL;
 		    Barrier.srcAccessMask = VK_ACCESS_COLOR_ATTACHMENT_WRITE_BIT;
-		    Barrier.dstAccessMask = VK_ACCESS_MEMORY_READ_BIT;
+		    Barrier.dstAccessMask = VK_ACCESS_COLOR_ATTACHMENT_READ_BIT;
 		    Barrier.oldLayout = VK_IMAGE_LAYOUT_GENERAL;
 		    Barrier.newLayout = VK_IMAGE_LAYOUT_GENERAL;
 		    Barrier.image = Texture->Surface.Image;
@@ -127,8 +111,6 @@ FVulkanFramebuffer::FVulkanFramebuffer(FVulkanDevice& Device, const FRHISetRende
 	#endif
 		Attachments.Add(Texture->View.View);
 
-#if VULKAN_USE_NEW_COMMAND_BUFFERS
-#else
 		// Create a write-barrier
 		WriteBarriers.AddZeroed();
 		VkImageMemoryBarrier& Barrier = WriteBarriers.Last();
@@ -136,17 +118,15 @@ FVulkanFramebuffer::FVulkanFramebuffer(FVulkanDevice& Device, const FRHISetRende
 		//Barrier.pNext = NULL;
 		Barrier.srcAccessMask = VK_ACCESS_COLOR_ATTACHMENT_WRITE_BIT;
 		Barrier.dstAccessMask = VK_ACCESS_MEMORY_READ_BIT;
-		Barrier.oldLayout = Texture->Surface.ImageLayout;
-		Barrier.newLayout = VK_IMAGE_LAYOUT_COLOR_ATTACHMENT_OPTIMAL;
+		Barrier.oldLayout = VK_IMAGE_LAYOUT_GENERAL;
+		Barrier.newLayout = VK_IMAGE_LAYOUT_GENERAL;
 		Barrier.image = Texture->Surface.Image;
 		Barrier.subresourceRange.aspectMask = Texture->Surface.GetAspectMask();
 		//Barrier.subresourceRange.baseMipLevel = 0;
 		Barrier.subresourceRange.levelCount = 1;
 		//Barrier.subresourceRange.baseArrayLayer = 0;
 		Barrier.subresourceRange.layerCount = 1;
-		Barrier.srcQueueFamilyIndex = VK_QUEUE_FAMILY_IGNORED;
-		Barrier.dstQueueFamilyIndex = VK_QUEUE_FAMILY_IGNORED;
-#endif
+
 		NumColorAttachments++;
 	}
 
@@ -165,10 +145,8 @@ FVulkanFramebuffer::FVulkanFramebuffer(FVulkanDevice& Device, const FRHISetRende
 				Texture->Surface.GetNumMips());
 		}
 
-		Attachments.Add(Texture->View.View);
+		bool bHasStencil = (Texture->Surface.Format == PF_DepthStencil);
 
-#if VULKAN_USE_NEW_COMMAND_BUFFERS
-#else
 		// Create a write-barrier
 		WriteBarriers.AddZeroed();
 		VkImageMemoryBarrier& Barrier = WriteBarriers.Last();
@@ -177,17 +155,16 @@ FVulkanFramebuffer::FVulkanFramebuffer(FVulkanDevice& Device, const FRHISetRende
 		Barrier.srcAccessMask = VK_ACCESS_DEPTH_STENCIL_ATTACHMENT_WRITE_BIT;
 		Barrier.dstAccessMask = VK_ACCESS_MEMORY_READ_BIT;
 		Barrier.oldLayout = VK_IMAGE_LAYOUT_GENERAL;
-		Barrier.newLayout = VK_IMAGE_LAYOUT_DEPTH_STENCIL_ATTACHMENT_OPTIMAL;
+		Barrier.newLayout = VK_IMAGE_LAYOUT_GENERAL;
 		Barrier.image = Texture->Surface.Image;
-		Barrier.srcQueueFamilyIndex = VK_QUEUE_FAMILY_IGNORED;
-		Barrier.dstQueueFamilyIndex = VK_QUEUE_FAMILY_IGNORED;
 
 		Barrier.subresourceRange.aspectMask = Texture->Surface.GetAspectMask();
 		//Barrier.subresourceRange.baseMipLevel = 0;
 		Barrier.subresourceRange.levelCount = 1;
 		//Barrier.subresourceRange.baseArrayLayer = 0;
 		Barrier.subresourceRange.layerCount = 1;
-#endif
+
+		Attachments.Add(Texture->View.View);
 	}
 
 	VkFramebufferCreateInfo Info;
@@ -253,21 +230,20 @@ bool FVulkanFramebuffer::Matches(const FRHISetRenderTargetsInfo& InRTInfo) const
 	return true;
 }
 
-#if VULKAN_USE_NEW_COMMAND_BUFFERS
 void FVulkanViewport::AcquireBackBuffer(FVulkanCmdBuffer* CmdBuffer)
 {
-	CurrentBackBuffer = SwapChain->AcquireImageIndex();
+	CurrentBackBuffer = SwapChain->AcquireImageIndex(&PendingImageSemaphore);
 	check(CurrentBackBuffer != -1);
-	VulkanSetImageLayoutSimple(CmdBuffer->GetHandle(), BackBuffers[CurrentBackBuffer]->Surface.Image, VK_IMAGE_LAYOUT_PRESENT_SRC_KHR, VK_IMAGE_LAYOUT_COLOR_ATTACHMENT_OPTIMAL);
+	VulkanSetImageLayoutSimple(CmdBuffer->GetHandle(), BackBuffers[CurrentBackBuffer]->Surface.Image, VK_IMAGE_LAYOUT_PRESENT_SRC_KHR, VK_IMAGE_LAYOUT_GENERAL);
 }
 
 void FVulkanViewport::PrepareBackBufferForPresent(FVulkanCmdBuffer* CmdBuffer)
 {
 	check(CurrentBackBuffer != -1);
-	VulkanSetImageLayoutSimple(CmdBuffer->GetHandle(), BackBuffers[CurrentBackBuffer]->Surface.Image, VK_IMAGE_LAYOUT_COLOR_ATTACHMENT_OPTIMAL, VK_IMAGE_LAYOUT_PRESENT_SRC_KHR);
+	VulkanSetImageLayoutSimple(CmdBuffer->GetHandle(), BackBuffers[CurrentBackBuffer]->Surface.Image, VK_IMAGE_LAYOUT_GENERAL, VK_IMAGE_LAYOUT_PRESENT_SRC_KHR);
 }
 
-void FVulkanFramebuffer::InsertWriteBarrier(FVulkanCmdBuffer* CmdBuffer)
+void FVulkanFramebuffer::InsertWriteBarriers(FVulkanCmdBuffer* CmdBuffer)
 {
 	if (WriteBarriers.Num() == 0)
 	{
@@ -280,21 +256,6 @@ void FVulkanFramebuffer::InsertWriteBarrier(FVulkanCmdBuffer* CmdBuffer)
 	check(CmdBuffer->IsOutsideRenderPass());
 	vkCmdPipelineBarrier(CmdBuffer->GetHandle(), SrcStages, DestStages, 0, 0, nullptr, 0, nullptr, WriteBarriers.Num(), WriteBarriers.GetData());
 }
-#else
-void FVulkanFramebuffer::InsertWriteBarrier(VkCommandBuffer cmd)
-{
-	if (WriteBarriers.Num() == 0)
-	{
-		return;
-	}
-
-	const VkPipelineStageFlags SrcStages = VK_PIPELINE_STAGE_COLOR_ATTACHMENT_OUTPUT_BIT;
-	const VkPipelineStageFlags DestStages = VK_PIPELINE_STAGE_BOTTOM_OF_PIPE_BIT;
-
-	vkCmdPipelineBarrier(cmd, SrcStages, DestStages,
-		0, 0, nullptr, 0, nullptr, WriteBarriers.Num(), WriteBarriers.GetData());
-}
-#endif
 
 /*=============================================================================
  *	The following RHI functions must be called from the main thread.
@@ -324,13 +285,8 @@ void FVulkanDynamicRHI::WriteEndFrameTimestamp(void* Data)
 		auto* TimestampQueryPool = This->Device->GetTimestampQueryPool(This->PresentCount % (uint32)FVulkanDevice::NumTimestampPools);
 		if (TimestampQueryPool)
 		{
-#if VULKAN_USE_NEW_COMMAND_BUFFERS
 			auto* CmdBuffer = This->Device->GetImmediateContext().GetCommandBufferManager()->GetActiveCmdBuffer();
 			TimestampQueryPool->WriteEndFrame(CmdBuffer->GetHandle());
-#else
-			auto& PendingState = This->Device->GetPendingState();
-			TimestampQueryPool->WriteEndFrame(PendingState.GetCommandBuffer());
-#endif
 		}
 	}
 
@@ -345,7 +301,6 @@ void FVulkanDynamicRHI::Present()
 
 	check(Device);
 
-#if VULKAN_USE_NEW_COMMAND_BUFFERS
 	auto* CmdBufferManager = Device->GetImmediateContext().GetCommandBufferManager();
 	auto* CmdBuffer = CmdBufferManager->GetActiveCmdBuffer();
 	if (PendingState.IsRenderPassActive())
@@ -356,15 +311,8 @@ void FVulkanDynamicRHI::Present()
 	WriteEndFrameTimestamp(this);
 	CmdBuffer->End();
 	Device->GetQueue()->Submit(CmdBuffer);
-#else
-	PendingState.SubmitPendingCommandBuffers(&WriteEndFrameTimestamp, this);
-#endif
 
-#if VULKAN_USE_NEW_COMMAND_BUFFERS
 	DrawingViewport->GetSwapChain()->Present(Device->GetQueue());
-#else
-	Device->GetQueue()->Present(DrawingViewport->SwapChain, DrawingViewport->CurrentBackBuffer);
-#endif
 
 	bool bNativelyPresented = true;
 	if (bNativelyPresented)
@@ -397,9 +345,7 @@ void FVulkanDynamicRHI::Present()
 	//#todo-rco: This needs to go on RHIEndFrame but the CmdBuffer index is not the correct one to read the stats out!
 	VulkanRHI::GManager.GPUProfilingData.EndFrame();
 
-#if VULKAN_USE_NEW_COMMAND_BUFFERS
 	Device->GetImmediateContext().GetCommandBufferManager()->PrepareForNewActiveCommandBuffer();
-#endif
 
 	//#todo-rco: Consolidate 'end of frame'
 	Device->GetImmediateContext().GetTempFrameAllocationBuffer().Reset();
@@ -420,11 +366,7 @@ void FVulkanDynamicRHI::Present()
 			TimestampQueryPool->CalculateFrameTime();
 		}
 
-#if VULKAN_USE_NEW_COMMAND_BUFFERS
 		Device->GetTimestampQueryPool(QueryNextFrameIndex)->WriteStartFrame(CmdBufferManager->GetActiveCmdBuffer()->GetHandle());
-#else
-		Device->GetTimestampQueryPool(QueryNextFrameIndex)->WriteStartFrame(PendingState.GetCommandBuffer());
-#endif
 	}
 
 	PresentCount++;
@@ -435,22 +377,14 @@ FTexture2DRHIRef FVulkanDynamicRHI::RHIGetViewportBackBuffer(FViewportRHIParamRe
 	FVulkanViewport* Viewport = ResourceCast(ViewportRHI);
 	if (Viewport->CurrentBackBuffer < 0)
 	{
-#if VULKAN_USE_NEW_COMMAND_BUFFERS
 		Viewport->AcquireBackBuffer(Device->GetImmediateContext().GetCommandBufferManager()->GetActiveCmdBuffer());
-#else
-		// SMEDIS: Acquire a new backbuffer now, in case RHIBeginDrawingViewport() hasn't been called yet.
-		Viewport->CurrentBackBuffer = Device->GetQueue()->AquireImageIndex(Viewport->SwapChain);
-#endif
 	}
 	return Viewport->GetBackBuffer();
 }
 
 void FVulkanDynamicRHI::RHIAdvanceFrameForGetViewportBackBuffer()
 {
-#if VULKAN_USE_NEW_COMMAND_BUFFERS
 	//#todo-rco: Do a pass to clear unused or expired elements in managers
-#else
-#endif
 }
 
 void FVulkanCommandListContext::RHISetViewport(uint32 MinX, uint32 MinY, float MinZ, uint32 MaxX, uint32 MaxY, float MaxZ)
diff --git a/Engine/Source/Runtime/VulkanRHI/Public/VulkanConfiguration.h b/Engine/Source/Runtime/VulkanRHI/Public/VulkanConfiguration.h
--- a/Engine/Source/Runtime/VulkanRHI/Public/VulkanConfiguration.h
+++ b/Engine/Source/Runtime/VulkanRHI/Public/VulkanConfiguration.h
@@ -25,7 +25,7 @@
 
 // constants we probably will change a few times
 #define VULKAN_UB_RING_BUFFER_SIZE								(8 * 1024 * 1024)
-#define VULKAN_TEMP_FRAME_ALLOCATOR_SIZE						(1 * 1024 * 1024)
+#define VULKAN_TEMP_FRAME_ALLOCATOR_SIZE						(8 * 1024 * 1024)
 
 //@NOTE: VULKAN_NUM_IMAGE_BUFFERS should be smaller than VULKAN_NUM_COMMAND_BUFFERS, to make sure that we wait for the fence before we reset the cmd buffer
 //@TODO: Clean up VULKAN_NUM_IMAGE_BUFFERS and VULKAN_NUM_COMMAND_BUFFERS once the Vulkan API and SDK stabilizes.
@@ -94,10 +94,7 @@ inline EDescriptorSetStage GetDescriptorSetForStage(EShaderFrequency Stage)
 
 #define VULKAN_ENABLE_AGGRESSIVE_STATS							1
 
-#define VULKAN_ENABLE_PIPELINE_CACHE							1
-
-//#todo-rco: Test on Android
-#define VULKAN_USE_NEW_COMMAND_BUFFERS							0//PLATFORM_WINDOWS
+#define VULKAN_ENABLE_PIPELINE_CACHE							0
 
 #define VULKAN_ENABLE_RHI_DEBUGGING								1
 
diff --git a/Engine/Source/Runtime/VulkanRHI/Public/VulkanMemory.h b/Engine/Source/Runtime/VulkanRHI/Public/VulkanMemory.h
--- a/Engine/Source/Runtime/VulkanRHI/Public/VulkanMemory.h
+++ b/Engine/Source/Runtime/VulkanRHI/Public/VulkanMemory.h
@@ -108,11 +108,6 @@ namespace VulkanRHI
 		void* Map(VkDeviceSize Size, VkDeviceSize Offset);
 		void Unmap();
 
-		inline void* GetMappedPointer()
-		{
-			return MappedPointer;
-		}
-
 		inline bool CanBeMapped() const
 		{
 			return bCanBeMapped != 0;
@@ -123,6 +118,12 @@ namespace VulkanRHI
 			return !!MappedPointer;
 		}
 
+		inline void* GetMappedPointer()
+		{
+			check(IsMapped());
+			return MappedPointer;
+		}
+
 		inline bool IsCoherent() const
 		{
 			return bIsCoherent != 0;
@@ -409,6 +410,11 @@ namespace VulkanRHI
 			return AlignedOffset;
 		}
 
+		inline uint32 GetSize() const
+		{
+			return RequestedSize;
+		}
+
 	protected:
 		uint32 RequestedSize;
 		uint32 AlignedOffset;
@@ -435,6 +441,13 @@ namespace VulkanRHI
 			return Handle;
 		}
 
+		FBufferAllocation* GetBufferAllocation()
+		{
+			return Owner;
+		}
+
+		void* GetMappedPointer();
+
 	protected:
 		friend class FBufferAllocation;
 		FBufferAllocation* Owner;
@@ -475,6 +488,16 @@ namespace VulkanRHI
 			return TryAllocateNoLocking(Size, File, Line);
 		}
 
+		uint32 GetAlignment() const
+		{
+			return Alignment;
+		}
+
+		void* GetMappedPointer()
+		{
+			return MemoryAllocation->GetMappedPointer();
+		}
+
 	protected:
 		FResourceHeapManager* Owner;
 		uint32 MemoryTypeIndex;
@@ -670,6 +693,10 @@ namespace VulkanRHI
 			uint32 TypeIndex = 0;
 			VERIFYVULKANRESULT(DeviceMemoryManager->GetMemoryTypeFromProperties(MemoryReqs.memoryTypeBits, MemoryPropertyFlags, &TypeIndex));
 			bool bMapped = (MemoryPropertyFlags & VK_MEMORY_PROPERTY_HOST_VISIBLE_BIT) == VK_MEMORY_PROPERTY_HOST_VISIBLE_BIT;
+			if (!ResourceHeaps[TypeIndex])
+			{
+				UE_LOG(LogVulkanRHI, Fatal, TEXT("Missing memory type index %d, MemSize %d, MemPropTypeBits %u, MemPropertyFlags %u, %s(%d)"), TypeIndex, (uint32)MemoryReqs.size, (uint32)MemoryReqs.memoryTypeBits, (uint32)MemoryPropertyFlags, ANSI_TO_TCHAR(File), Line);
+			}
 			return ResourceHeaps[TypeIndex]->AllocateResource(MemoryReqs.size, MemoryReqs.alignment, false, bMapped, File, Line);
 		}
 
@@ -765,16 +792,15 @@ namespace VulkanRHI
 
 		FStagingBuffer* AcquireBuffer(uint32 Size, VkBufferUsageFlags InUsageFlags = VK_BUFFER_USAGE_TRANSFER_SRC_BIT);
 
-#if VULKAN_USE_NEW_COMMAND_BUFFERS
 		// Sets pointer to nullptr
 		void ReleaseBuffer(FVulkanCmdBuffer* CmdBuffer, FStagingBuffer*& StagingBuffer);
-#else
-		// Sets pointer to nullptr
-		void ReleaseBuffer(FStagingBuffer*& StagingBuffer);
-#endif
 
 		void ProcessPendingFree(bool bImmediately = false);
 
+#if UE_BUILD_DEBUG || UE_BUILD_DEVELOPMENT
+		void DumpMemory();
+#endif
+
 	protected:
 		struct FPendingItem
 		{
@@ -784,9 +810,7 @@ namespace VulkanRHI
 		};
 
 		TArray<FStagingBuffer*> UsedStagingBuffers;
-#if VULKAN_USE_NEW_COMMAND_BUFFERS
 		TArray<FPendingItem> PendingFreeStagingBuffers;
-#endif
 		TArray<FStagingBuffer*> FreeStagingBuffers;
 
 		FVulkanDevice* Device;
@@ -929,27 +953,40 @@ namespace VulkanRHI
 		struct FTempAllocInfo
 		{
 			void* Data;
-			uint32 Offset;
-			VkBuffer Buffer;
+
+			// Offset into the locked area
+			uint32 CurrentOffset;
+
+			FBufferSuballocation* BufferSuballocation;
+
+			inline uint32 GetBindOffset() const
+			{
+				return BufferSuballocation->GetOffset() + CurrentOffset;
+			}
+
+			inline VkBuffer GetHandle() const
+			{
+				return BufferSuballocation->GetHandle();
+			}
 		};
 
 		bool Alloc(uint32 InSize, uint32 InAlignment, FTempAllocInfo& OutInfo);
 
-		inline VkBuffer GetHandle() const
-		{
-			return Buffer[BufferIndex]->GetHandle();
-		}
-
 		void Reset();
 
 	protected:
 		uint8* MappedData[NUM_RENDER_BUFFERS];
 		uint8* CurrentData[NUM_RENDER_BUFFERS];
-		FStagingBuffer* Buffer[NUM_RENDER_BUFFERS];
+		TRefCountPtr<FBufferSuballocation> BufferSuballocations[NUM_RENDER_BUFFERS];
 		uint32 BufferIndex;
 		uint32 Size;
 		uint32 PeakUsed;
 
 		friend class FVulkanCommandListContext;
 	};
+
+	inline void* FBufferSuballocation::GetMappedPointer()
+	{
+		return Owner->GetMappedPointer();
+	}
 }
diff --git a/Engine/Source/Runtime/VulkanRHI/Public/VulkanRHI.h b/Engine/Source/Runtime/VulkanRHI/Public/VulkanRHI.h
--- a/Engine/Source/Runtime/VulkanRHI/Public/VulkanRHI.h
+++ b/Engine/Source/Runtime/VulkanRHI/Public/VulkanRHI.h
@@ -192,10 +192,12 @@ protected:
 
 #if UE_BUILD_DEBUG || UE_BUILD_DEVELOPMENT
 	IConsoleObject* DumpMemoryCmd;
+public:
 	static void DumpMemory();
 #endif
 
 #if VULKAN_HAS_DEBUGGING_ENABLED
+protected:
 	VkDebugReportCallbackEXT MsgCallback;
 	void SetupDebugLayerCallback();
 	void RemoveDebugLayerCallback();
diff --git a/Engine/Source/Runtime/VulkanRHI/Public/VulkanResources.h b/Engine/Source/Runtime/VulkanRHI/Public/VulkanResources.h
--- a/Engine/Source/Runtime/VulkanRHI/Public/VulkanResources.h
+++ b/Engine/Source/Runtime/VulkanRHI/Public/VulkanResources.h
@@ -229,13 +229,8 @@ public:
 	uint8 FormatKey;
 private:
 
-#if VULKAN_USE_NEW_COMMAND_BUFFERS
 	// Used to clear render-target objects on creation
 	void Clear(const FClearValueBinding& ClearValueBinding, bool bTransitionToPresentable);
-#else
-	// Used to clear render-target objects on creation
-	void ClearBlocking(const FClearValueBinding& ClearValueBinding, bool bTransitionToPresentable);
-#endif
 
 private:
 	// Linear or Optimal. Based on tiling, map/unmap is handled differently.
@@ -553,17 +548,23 @@ public:
 
 };
 
-struct FVulkanBufferView : public FRHIResource
+struct FVulkanBufferView : public FRHIResource, public VulkanRHI::FDeviceChild
 {
-	FVulkanBufferView() :
-		View(VK_NULL_HANDLE),
-		Flags(0)
+	FVulkanBufferView(FVulkanDevice* InDevice)
+		: VulkanRHI::FDeviceChild(InDevice)
+		, View(VK_NULL_HANDLE)
+		, Flags(0)
 	{
 	}
 
-	void Create(FVulkanDevice& Device, FVulkanBuffer& Buffer, EPixelFormat Format, uint32 Offset, uint32 Size);
-	void Create(FVulkanDevice& Device, FVulkanResourceMultiBuffer* Buffer, EPixelFormat Format, uint32 Offset, uint32 Size);
-	void Destroy(FVulkanDevice& Device);
+	virtual ~FVulkanBufferView()
+	{
+		Destroy();
+	}
+
+	void Create(FVulkanBuffer& Buffer, EPixelFormat Format, uint32 Offset, uint32 Size);
+	void Create(FVulkanResourceMultiBuffer* Buffer, EPixelFormat Format, uint32 Offset, uint32 Size);
+	void Destroy();
 
 	VkBufferView View;
 	VkFlags Flags;
@@ -611,12 +612,26 @@ public:
 	// allocate some space in the ring buffer
 	uint64 AllocateMemory(uint64 Size, uint32 Alignment);
 
-	VulkanRHI::FStagingBuffer* Buffer;
+	inline uint32 GetBufferOffset() const
+	{
+		return BufferSuballocation->GetOffset();
+	}
+
+	inline VkBuffer GetHandle() const
+	{
+		return BufferSuballocation->GetHandle();
+	}
+
+	inline void* GetMappedPointer()
+	{
+		return BufferSuballocation->GetMappedPointer();
+	}
 
 protected:
 	uint64 BufferSize;
 	uint64 BufferOffset;
 	uint32 MinAlignment;
+	VulkanRHI::FBufferSuballocation* BufferSuballocation;
 };
 
 class FVulkanResourceMultiBuffer : public VulkanRHI::FDeviceChild
@@ -629,7 +644,7 @@ public:
 	{
 		if (NumBuffers == 0)
 		{
-			return VolatileLockInfo.Buffer;
+			return VolatileLockInfo.GetHandle();
 		}
 		return Buffers[DynamicBufferIndex]->GetHandle();
 	}
@@ -639,11 +654,12 @@ public:
 		return NumBuffers > 1;
 	}
 
+	// Offset used for Binding a VkBuffer
 	inline uint32 GetOffset() const
 	{
 		if (NumBuffers == 0)
 		{
-			return VolatileLockInfo.Offset;
+			return VolatileLockInfo.GetBindOffset();
 		}
 		return Buffers[DynamicBufferIndex]->GetOffset();
 	}
@@ -815,7 +831,7 @@ public:
 
 	class FVulkanPipeline* PrepareForDraw(const struct FVulkanPipelineGraphicsKey& PipelineKey, uint32 VertexInputKey, const struct FVulkanPipelineState& State);
 
-    void UpdateDescriptorSets(FVulkanGlobalUniformPool* GlobalUniformPool);
+	void UpdateDescriptorSets(FVulkanCommandListContext* CmdListContext, FVulkanCmdBuffer* CmdBuffer, FVulkanGlobalUniformPool* GlobalUniformPool);
 
 	void BindDescriptorSets(FVulkanCmdBuffer* Cmd);
 
@@ -1052,10 +1068,33 @@ private:
 	// these are the cache pipeline state objects used in this BSS
 	TMap<FVulkanPipelineGraphicsKey, class FVulkanPipeline*> PipelineCache;
 
-	// descriptor set objects used by this BSS
-	uint64 LastFrameRendered;
-	int32 CurrentDS[VULKAN_NUM_COMMAND_BUFFERS];
-	TArray<FVulkanDescriptorSets*> DescriptorSets[VULKAN_NUM_COMMAND_BUFFERS];
+	struct FDescriptorSetsPair
+	{
+		uint64 FenceCounter;
+		FVulkanDescriptorSets* DescriptorSets;
+
+		FDescriptorSetsPair()
+			: FenceCounter(0)
+			, DescriptorSets(nullptr)
+		{
+		}
+
+		~FDescriptorSetsPair();
+	};
+
+	struct FDescriptorSetsEntry
+	{
+		FVulkanCmdBuffer* CmdBuffer;
+		TArray<FDescriptorSetsPair> Pairs;
+
+		FDescriptorSetsEntry(FVulkanCmdBuffer* InCmdBuffer)
+			: CmdBuffer(InCmdBuffer)
+		{
+		}
+	};
+	TArray<FDescriptorSetsEntry*> DescriptorSetsEntries;
+
+	FVulkanDescriptorSets* RequestDescriptorSets(FVulkanCmdBuffer* CmdBuffer);
 };
 
 template<class T>
diff --git a/Engine/Source/Runtime/VulkanRHI/Public/VulkanViewport.h b/Engine/Source/Runtime/VulkanRHI/Public/VulkanViewport.h
--- a/Engine/Source/Runtime/VulkanRHI/Public/VulkanViewport.h
+++ b/Engine/Source/Runtime/VulkanRHI/Public/VulkanViewport.h
@@ -11,6 +11,7 @@
 class FVulkanDynamicRHI;
 class FVulkanSwapChain;
 class FVulkanQueue;
+struct FVulkanSemaphore;
 
 class FVulkanViewport : public FRHIViewport
 {
@@ -33,17 +34,20 @@ public:
 
 	FIntPoint GetSizeXY() const { return FIntPoint(SizeX, SizeY); }
 
-#if VULKAN_USE_NEW_COMMAND_BUFFERS
 	FVulkanSwapChain* GetSwapChain()
 	{
 		return SwapChain;
 	}
 
+	FVulkanSemaphore* AcquirePendingSemaphore()
+	{
+		FVulkanSemaphore* Semaphore = PendingImageSemaphore;
+		PendingImageSemaphore = nullptr;
+		return Semaphore;
+	}
+
 	void AcquireBackBuffer(FVulkanCmdBuffer* CmdBuffer);
 	void PrepareBackBufferForPresent(FVulkanCmdBuffer* CmdBuffer);
-#else
-#endif
-
 
 protected:
 	TRefCountPtr<FVulkanTexture2D> BackBuffers[NUM_BUFFERS];
@@ -55,6 +59,8 @@ protected:
 	int32 CurrentBackBuffer;
 	FVulkanSwapChain* SwapChain;
 
+	FVulkanSemaphore* PendingImageSemaphore;
+
 	friend class FVulkanDynamicRHI;
 	friend class FVulkanCommandListContext;
 };
diff --git a/Engine/Source/Runtime/Windows/D3D11RHI/Private/D3D11Texture.cpp b/Engine/Source/Runtime/Windows/D3D11RHI/Private/D3D11Texture.cpp
--- a/Engine/Source/Runtime/Windows/D3D11RHI/Private/D3D11Texture.cpp
+++ b/Engine/Source/Runtime/Windows/D3D11RHI/Private/D3D11Texture.cpp
@@ -481,7 +481,7 @@ TD3D11Texture2D<BaseResourceType>* FD3D11DynamicRHI::CreateD3D11Texture2D(uint32
 
 	bool bPooledTexture = true;
 
-	if (GMaxRHIFeatureLevel <= ERHIFeatureLevel::ES2)
+	if (GMaxRHIFeatureLevel <= ERHIFeatureLevel::ES3_1)
 	{
 		// Remove sRGB read flag when not supported
 		Flags &= ~TexCreate_SRGB;
@@ -1007,7 +1007,7 @@ FTexture2DRHIRef FD3D11DynamicRHI::RHIAsyncCreateTexture2D(uint32 SizeX,uint32 S
 	check(GRHISupportsAsyncTextureCreation);
 	check((Flags & InvalidFlags) == 0);
 
-	if (GMaxRHIFeatureLevel <= ERHIFeatureLevel::ES2)
+	if (GMaxRHIFeatureLevel <= ERHIFeatureLevel::ES3_1)
 	{
 		// Remove sRGB read flag when not supported
 		Flags &= ~TexCreate_SRGB;
diff --git a/Engine/Source/Runtime/Windows/D3D12RHI/Private/D3D12Texture.cpp b/Engine/Source/Runtime/Windows/D3D12RHI/Private/D3D12Texture.cpp
--- a/Engine/Source/Runtime/Windows/D3D12RHI/Private/D3D12Texture.cpp
+++ b/Engine/Source/Runtime/Windows/D3D12RHI/Private/D3D12Texture.cpp
@@ -663,7 +663,7 @@ TD3D12Texture2D<BaseResourceType>* FD3D12DynamicRHI::CreateD3D11Texture2D(uint32
 
 	bool bPooledTexture = true;
 
-	if (GMaxRHIFeatureLevel <= ERHIFeatureLevel::ES2)
+	if (GMaxRHIFeatureLevel <= ERHIFeatureLevel::ES3_1)
 	{
 		// Remove sRGB read flag when not supported
 		Flags &= ~TexCreate_SRGB;
@@ -1244,7 +1244,7 @@ FTexture2DRHIRef FD3D12DynamicRHI::RHIAsyncCreateTexture2D(uint32 SizeX, uint32
 	check(GRHISupportsAsyncTextureCreation);
 	check((Flags & InvalidFlags) == 0);
 
-	if (GMaxRHIFeatureLevel <= ERHIFeatureLevel::ES2)
+	if (GMaxRHIFeatureLevel <= ERHIFeatureLevel::ES3_1)
 	{
 		// Remove sRGB read flag when not supported
 		Flags &= ~TexCreate_SRGB;
@@ -1992,7 +1992,7 @@ FTexture2DRHIRef FD3D12DynamicRHI::RHICreateTexture2DFromD3D12Resource(uint8 For
 
 	SCOPE_CYCLE_COUNTER(STAT_D3D12CreateTextureTime);
 
-	if (GMaxRHIFeatureLevel <= ERHIFeatureLevel::ES2)
+	if (GMaxRHIFeatureLevel <= ERHIFeatureLevel::ES3_1)
 	{
 		// Remove sRGB read flag when not supported
 		Flags &= ~TexCreate_SRGB;
diff --git a/README.md b/README.md
--- a/README.md
+++ b/README.md
@@ -23,7 +23,6 @@ We publish source for the engine in three rolling branches:
 
 The **[release branch](https://github.com/EpicGames/UnrealEngine/tree/release)** is extensively tested by our QA team and makes a great starting point for learning the engine or
 making your own games. We work hard to make releases stable and reliable, and aim to publish new releases every few months.
->**Important**: The Release branch has unintentionally been overwritten with code from the 4.11 Previews, and is *not* currently the latest and stable release. Please use Tags to grab the latest 4.10 release for our most stable version. This message will be removed when this matter has been resolved. 
 
 The **[promoted branch](https://github.com/EpicGames/UnrealEngine/tree/promoted)** is updated with builds for our artists and designers to use. We try to update it daily 
 (though we often catch things that prevent us from doing so) and it's a good balance between getting the latest cool stuff and knowing most things work.
