Fix Android Vulkan loader compile errors
#ue4
#android
#vulkan

diff --git a/Engine/Source/Runtime/VulkanRHI/Private/VulkanLayers.cpp b/Engine/Source/Runtime/VulkanRHI/Private/VulkanLayers.cpp
--- a/Engine/Source/Runtime/VulkanRHI/Private/VulkanLayers.cpp
+++ b/Engine/Source/Runtime/VulkanRHI/Private/VulkanLayers.cpp
@@ -234,8 +234,6 @@ void FVulkanDynamicRHI::GetInstanceLayersAndExtensions(TArray<const ANSICHAR*>&
 	}
 #endif	// VULKAN_HAS_DEBUGGING_ENABLED
 
-	//@TODO: Android driver hasn't implemented extensions yet
-#if !PLATFORM_ANDROID
 	for (int32 i = 0; i < GlobalExtensions.ExtensionProps.Num(); i++)
 	{
 		for (int32 j = 0; j < ARRAY_COUNT(GInstanceExtensions); j++)
@@ -247,7 +245,6 @@ void FVulkanDynamicRHI::GetInstanceLayersAndExtensions(TArray<const ANSICHAR*>&
 			}
 		}
 	}
-#endif	// !PLATFORM_ANDROID
 
 	if (OutInstanceExtensions.Num() > 0)
 	{
@@ -349,8 +346,6 @@ void FVulkanDevice::GetDeviceExtensions(TArray<const ANSICHAR*>& OutDeviceExtens
 	}
 #endif	// VULKAN_HAS_DEBUGGING_ENABLED
 
-	//@TODO: Extensions mechanisms are currently unavailable
-#if !PLATFORM_ANDROID
 	FLayerExtension Extensions;
 	FMemory::Memzero(Extensions.LayerProps);
 	GetDeviceLayerExtensions(Gpu, nullptr, Extensions);
@@ -384,9 +379,6 @@ void FVulkanDevice::GetDeviceExtensions(TArray<const ANSICHAR*>& OutDeviceExtens
 	}
 	#endif
 
-
-#endif	// !PLATFORM_ANDROID
-
 	if (OutDeviceExtensions.Num() > 0)
 	{
 		UE_LOG(LogVulkanRHI, Display, TEXT("Using device extensions"));
diff --git a/Engine/Source/Runtime/VulkanRHI/Private/VulkanLoader.h b/Engine/Source/Runtime/VulkanRHI/Private/VulkanLoader.h
--- a/Engine/Source/Runtime/VulkanRHI/Private/VulkanLoader.h
+++ b/Engine/Source/Runtime/VulkanRHI/Private/VulkanLoader.h
@@ -180,5 +180,7 @@
 
 // Declare all Vulkan functions
 #define DECLARE_VK_ENTRYPOINTS(Type,Func) extern Type Func;
-ENUM_VK_ENTRYPOINTS_ALL(DECLARE_VK_ENTRYPOINTS);
-
+namespace VulkanRHI
+{
+	ENUM_VK_ENTRYPOINTS_ALL(DECLARE_VK_ENTRYPOINTS);
+}
diff --git a/Engine/Source/Runtime/VulkanRHI/Private/VulkanRHI.cpp b/Engine/Source/Runtime/VulkanRHI/Private/VulkanRHI.cpp
--- a/Engine/Source/Runtime/VulkanRHI/Private/VulkanRHI.cpp
+++ b/Engine/Source/Runtime/VulkanRHI/Private/VulkanRHI.cpp
@@ -30,7 +30,7 @@ static_assert(VK_API_VERSION >= UE_VK_API_VERSION, "Vulkan SDK is older than the
 
 // Vulkan function pointers
 
-#define DEFINE_VK_ENTRYPOINTS(Type,Func) Type Func = NULL;
+#define DEFINE_VK_ENTRYPOINTS(Type,Func) Type VulkanRHI::Func = NULL;
 ENUM_VK_ENTRYPOINTS_ALL(DEFINE_VK_ENTRYPOINTS)
 
 #endif
@@ -197,10 +197,10 @@ static bool LoadVulkanLibrary()
 	}
 
 	bool bFoundAllEntryPoints = true;
-#define CHECK_VK_ENTRYPOINTS(Type,Func) if (Func == NULL) { bFoundAllEntryPoints = false; UE_LOG(LogRHI, Warning, TEXT("Failed to find entry point for %s"), TEXT(#Func)); }
+#define CHECK_VK_ENTRYPOINTS(Type,Func) if (VulkanRHI::Func == NULL) { bFoundAllEntryPoints = false; UE_LOG(LogRHI, Warning, TEXT("Failed to find entry point for %s"), TEXT(#Func)); }
 
 	// Initialize all of the entry points we have to query manually
-#define GET_VK_ENTRYPOINTS(Type,Func) Func = (Type)dlsym(VulkanLib, #Func);
+#define GET_VK_ENTRYPOINTS(Type,Func) VulkanRHI::Func = (Type)dlsym(VulkanLib, #Func);
 	ENUM_VK_ENTRYPOINTS_BASE(GET_VK_ENTRYPOINTS);
 	ENUM_VK_ENTRYPOINTS_BASE(CHECK_VK_ENTRYPOINTS);
 	if (!bFoundAllEntryPoints)
@@ -219,9 +219,9 @@ static bool LoadVulkanLibrary()
 static bool LoadVulkanInstanceFunctions(VkInstance inInstance)
 {
 	bool bFoundAllEntryPoints = true;
-#define CHECK_VK_ENTRYPOINTS(Type,Func) if (Func == NULL) { bFoundAllEntryPoints = false; UE_LOG(LogRHI, Warning, TEXT("Failed to find entry point for %s"), TEXT(#Func)); }
+#define CHECK_VK_ENTRYPOINTS(Type,Func) if (VulkanRHI::Func == NULL) { bFoundAllEntryPoints = false; UE_LOG(LogRHI, Warning, TEXT("Failed to find entry point for %s"), TEXT(#Func)); }
 
-#define GETINSTANCE_VK_ENTRYPOINTS(Type, Func) Func = (Type)vkGetInstanceProcAddr(inInstance, #Func);
+#define GETINSTANCE_VK_ENTRYPOINTS(Type, Func) VulkanRHI::Func = (Type)VulkanRHI::vkGetInstanceProcAddr(inInstance, #Func);
 	ENUM_VK_ENTRYPOINTS_INSTANCE(GETINSTANCE_VK_ENTRYPOINTS);
 	ENUM_VK_ENTRYPOINTS_INSTANCE(CHECK_VK_ENTRYPOINTS);
 
@@ -232,7 +232,7 @@ static void FreeVulkanLibrary()
 {
 	if (VulkanLib != nullptr)
 	{
-#define CLEAR_VK_ENTRYPOINTS(Type,Func) Func = nullptr;
+#define CLEAR_VK_ENTRYPOINTS(Type,Func) VulkanRHI::Func = nullptr;
 		ENUM_VK_ENTRYPOINTS_ALL(CLEAR_VK_ENTRYPOINTS);
 
 		dlclose(VulkanLib);
diff --git a/Engine/Source/Runtime/VulkanRHI/Private/VulkanRHIPrivate.h b/Engine/Source/Runtime/VulkanRHI/Private/VulkanRHIPrivate.h
--- a/Engine/Source/Runtime/VulkanRHI/Private/VulkanRHIPrivate.h
+++ b/Engine/Source/Runtime/VulkanRHI/Private/VulkanRHIPrivate.h
@@ -28,15 +28,19 @@ DECLARE_LOG_CATEGORY_EXTERN(LogVulkanRHI, Log, All);
 
 #if PLATFORM_ANDROID
 #include "VulkanLoader.h"
+#define VULKAN_COMMANDWRAPPERS_ENABLE 0
 #else
 #include <vulkan/vulkan.h>
+#define VULKAN_COMMANDWRAPPERS_ENABLE 1
 #endif
 
 #include "VulkanRHI.h"
 #include "VulkanGlobalUniformBuffer.h"
 #include "RHI.h"
 
+#if VULKAN_COMMANDWRAPPERS_ENABLE
 #include "VulkanCommandWrappers.h"
+#endif
 using namespace VulkanRHI;
 
 // Default is 1 (which is aniso off), the number is adjusted after the limits are queried.
diff --git a/Engine/Source/Runtime/VulkanRHI/Public/VulkanResources.h b/Engine/Source/Runtime/VulkanRHI/Public/VulkanResources.h
--- a/Engine/Source/Runtime/VulkanRHI/Public/VulkanResources.h
+++ b/Engine/Source/Runtime/VulkanRHI/Public/VulkanResources.h
@@ -977,7 +977,11 @@ public:
 		if (LastBoundPipeline != NewPipeline)
 		{
 			//#todo-rco: Compute
+#if	VULKAN_COMMANDWRAPPERS_ENABLE
 			vkCmdBindPipeline(CmdBuffer, VK_PIPELINE_BIND_POINT_GRAPHICS, NewPipeline);
+#else
+			VulkanRHI::vkCmdBindPipeline(CmdBuffer, VK_PIPELINE_BIND_POINT_GRAPHICS, NewPipeline);
+#endif
 			LastBoundPipeline = NewPipeline;
 		}
 	}
