Copying //UE4/Dev-Rendering to //UE4/Dev-Main (Source: //UE4/Dev-Rendering @ 3045398)

diff --git a/Engine/Source/Runtime/VulkanRHI/Private/VulkanCommandWrappers.h b/Engine/Source/Runtime/VulkanRHI/Private/VulkanCommandWrappers.h
--- a/Engine/Source/Runtime/VulkanRHI/Private/VulkanCommandWrappers.h
+++ b/Engine/Source/Runtime/VulkanRHI/Private/VulkanCommandWrappers.h
@@ -13,6 +13,8 @@ namespace VulkanRHI
 	void DevicePrintfBeginResult(VkDevice Device, const FString& String);
 	void PrintfBegin(const FString& String);
 	void DevicePrintfBegin(VkDevice Device, const FString& String);
+	void CmdPrintfBegin(VkCommandBuffer CmdBuffer, const FString& String);
+	void CmdPrintfBeginResult(VkCommandBuffer CmdBuffer, const FString& String);
 	void PrintResult(VkResult Result);
 	void PrintResultAndNamedHandle(VkResult Result, const TCHAR* HandleName, void* Handle);
 	void PrintResultAndPointer(VkResult Result, void* Handle);
@@ -21,9 +23,11 @@ namespace VulkanRHI
 	void DumpMemoryRequirements(VkMemoryRequirements* MemoryRequirements);
 	void DumpBufferCreate(VkDevice Device, const VkBufferCreateInfo* CreateInfo, VkBuffer* Buffer);
 	void DumpBufferViewCreate(VkDevice Device, const VkBufferViewCreateInfo* CreateInfo, VkBufferView* BufferView);
-	void DumpImageCreate(VkDevice Device, const VkImageCreateInfo* CreateInfo);
+	void DumpImageCreate(VkDevice Device, const VkImageCreateInfo* CreateInfo, VkImage* Image);
+	void DumpImageViewCreate(VkDevice Device, const VkImageViewCreateInfo* CreateInfo, VkImageView* ImageView);
 	void DumpFenceCreate(VkDevice Device, const VkFenceCreateInfo* CreateInfo, VkFence* Fence);
 	void DumpFenceList(uint32 FenceCount, const VkFence* Fences);
+	void DumpSemaphoreCreate(VkDevice Device, const VkSemaphoreCreateInfo* CreateInfo, VkSemaphore* Semaphore);
 	void DumpMappedMemoryRanges(uint32 MemoryRangeCount, const VkMappedMemoryRange* pMemoryRanges);
 	void DumpResolveImage(VkCommandBuffer CommandBuffer, VkImage SrcImage, VkImageLayout SrcImageLayout, VkImage DstImage, VkImageLayout DstImageLayout, uint32 RegionCount, const VkImageResolve* Regions);
 	void DumpFreeDescriptorSets(VkDevice Device, VkDescriptorPool DescriptorPool, uint32 DescriptorSetCount, const VkDescriptorSet* DescriptorSets);
@@ -32,16 +36,39 @@ namespace VulkanRHI
 	void DumpCmdPipelineBarrier(VkCommandBuffer CommandBuffer, VkPipelineStageFlags SrcStageMask, VkPipelineStageFlags DstStageMask, VkDependencyFlags DependencyFlags, uint32 MemoryBarrierCount, const VkMemoryBarrier* MemoryBarriers, uint32 BufferMemoryBarrierCount, const VkBufferMemoryBarrier* BufferMemoryBarriers, uint32 ImageMemoryBarrierCount, const VkImageMemoryBarrier* ImageMemoryBarriers);
 	void DumpCreateDescriptorSetLayout(VkDevice Device, const VkDescriptorSetLayoutCreateInfo* CreateInfo, VkDescriptorSetLayout* SetLayout);
 	void DumpAllocateDescriptorSets(VkDevice Device, const VkDescriptorSetAllocateInfo* AllocateInfo, VkDescriptorSet* DescriptorSets);
+	void DumpCreateFramebuffer(VkDevice Device, const VkFramebufferCreateInfo* CreateInfo, VkFramebuffer* Framebuffer);
 	void DumpCreateRenderPass(VkDevice Device, const VkRenderPassCreateInfo* CreateInfo, VkRenderPass* RenderPass);
 	void DumpQueueSubmit(VkQueue Queue, uint32 SubmitCount, const VkSubmitInfo* Submits, VkFence Fence);
 	void DumpCreateShaderModule(VkDevice Device, const VkShaderModuleCreateInfo* CreateInfo, VkShaderModule* ShaderModule);
+	void DumpCreateDevice(VkPhysicalDevice PhysicalDevice, const VkDeviceCreateInfo* CreateInfo, VkDevice* Device);
+	void DumpUpdateDescriptorSets(VkDevice Device, uint32 DescriptorWriteCount, const VkWriteDescriptorSet* DescriptorWrites, uint32 DescriptorCopyCount, const VkCopyDescriptorSet* DescriptorCopies);
+	void DumpBeginCommandBuffer(VkCommandBuffer CommandBuffer, const VkCommandBufferBeginInfo* BeginInfo);
+	void DumpCmdBeginRenderPass(VkCommandBuffer CommandBuffer, const VkRenderPassBeginInfo* RenderPassBegin, VkSubpassContents Contents);
+	void DumpCmdBindVertexBuffers(VkCommandBuffer CommandBuffer, uint32 FirstBinding, uint32 BindingCount, const VkBuffer* Buffers, const VkDeviceSize* pOffsets);
+	void DumpGetImageSubresourceLayout(VkDevice Device, VkImage Image, const VkImageSubresource* Subresource, VkSubresourceLayout* Layout);
+	void DumpImageSubresourceLayout(VkSubresourceLayout* Layout);
+	void DumpCmdCopyBufferToImage(VkCommandBuffer CommandBuffer, VkBuffer SrcBuffer, VkImage DstImage, VkImageLayout DstImageLayout, uint32 RegionCount, const VkBufferImageCopy* Regions);
+	void DumpCmdCopyBuffer(VkCommandBuffer CommandBuffer, VkBuffer SrcBuffer, VkBuffer DstBuffer, uint32 RegionCount, const VkBufferCopy* Regions);
+	void DumpCreatePipelineCache(VkDevice Device, const VkPipelineCacheCreateInfo* CreateInfo, VkPipelineCache* PipelineCache);
+	void DumpCreateCommandPool(VkDevice Device, const VkCommandPoolCreateInfo* CreateInfo, VkCommandPool* CommandPool);
+	void DumpCreateQueryPool(VkDevice Device, const VkQueryPoolCreateInfo* CreateInfo, VkQueryPool* QueryPool);
+	void DumpCreatePipelineLayout(VkDevice Device, const VkPipelineLayoutCreateInfo* CreateInfo, VkPipelineLayout* PipelineLayout);
+	void DumpCreateDescriptorPool(VkDevice Device, const VkDescriptorPoolCreateInfo* CreateInfo, VkDescriptorPool* DescriptorPool);
+	void DumpCreateSampler(VkDevice Device, const VkSamplerCreateInfo* CreateInfo, VkSampler* Sampler);
+	void DumpGetPhysicalDeviceFeatures(VkPhysicalDevice PhysicalDevice, VkPhysicalDeviceFeatures* Features);
+	void DumpPhysicalDeviceFeatures(VkPhysicalDeviceFeatures* Features);
+	void DumpBindDescriptorSets(VkCommandBuffer CommandBuffer, VkPipelineBindPoint PipelineBindPoint, VkPipelineLayout Layout, uint32 FirstSet, uint32 DescriptorSetCount, const VkDescriptorSet* DescriptorSets, uint32 DynamicOffsetCount, const uint32* DynamicOffsets);
+	void DumpSwapChainImages(VkResult Result, uint32* SwapchainImageCount, VkImage* SwapchainImages);
 #else
+	#define FlushDebugWrapperLog()
 	#define DevicePrintfBeginResult(d, x)
 	#define PrintfBeginResult(x)
 	#define DevicePrintfBegin(d, x)
+	#define CmdPrintfBegin(c, s)
+	#define CmdPrintfBeginResult(c, s)
 	#define PrintfBegin(x)
 	#define PrintResult(x)
-	#define PrintResultAndNamedHandle(r, m_, h)
+	#define PrintResultAndNamedHandle(r, n, h)
 	#define PrintResultAndPointer(r, h)
 	#define DumpPhysicalDeviceProperties(x)
 	#define DumpAllocateMemory(d, i, m)
@@ -49,8 +76,10 @@ namespace VulkanRHI
 	#define DumpBufferCreate(d, i, b)
 	#define DumpBufferViewCreate(d, i, v)
 	#define DumpImageCreate(d, x, i)
+	#define DumpImageViewCreate(d, i, v)
 	#define DumpFenceCreate(d, x, f)
 	#define DumpFenceList(c, p)
+	#define DumpSemaphoreCreate(d, i, s)
 	#define DumpMappedMemoryRanges(c, p)
 	#define DumpResolveImage(c, si, sil, di, dil, rc, r)
 	#define DumpFreeDescriptorSets(d, p, c, s)
@@ -59,9 +88,29 @@ namespace VulkanRHI
 	#define DumpCmdPipelineBarrier(c, sm, dm, df, mc, mb, bc, bb, ic, ib)
 	#define DumpCreateDescriptorSetLayout(d, i, s)
 	#define DumpAllocateDescriptorSets(d, i, s)
+	#define DumpCreateFramebuffer(d, i, f)
 	#define DumpCreateRenderPass(d, i, r)
 	#define DumpQueueSubmit(q, c, s, f)
 	#define DumpCreateShaderModule(d, i, s)
+	#define DumpCreateDevice(p, i, d)
+	#define DumpUpdateDescriptorSets(d, wc, w, cc, c)
+	#define DumpBeginCommandBuffer(c, i)
+	#define DumpCmdBeginRenderPass(c, b, cont)
+	#define DumpCmdBindVertexBuffers(cb, f, bc, b, o)
+	#define DumpGetImageSubresourceLayout(d, i, s, l)
+	#define DumpImageSubresourceLayout(l)
+	#define DumpCmdCopyBufferToImage(cb, sb, di, dil, rc, r)
+	#define DumpCmdCopyBuffer(cb, sb, db, rc, r)
+	#define DumpCreatePipelineCache(d, i, pc)
+	#define DumpCreateCommandPool(d, i, c)
+	#define DumpCreateQueryPool(d, i, q)
+	#define DumpCreatePipelineLayout(d, i, p)
+	#define DumpCreateDescriptorPool(d, i, dp)
+	#define DumpCreateSampler(d, i, s)
+	#define DumpGetPhysicalDeviceFeatures(pd, f)
+	#define DumpPhysicalDeviceFeatures(f)
+	#define DumpBindDescriptorSets(c, pbp, l, fs, dsc, ds, doc, do)
+	#define DumpSwapChainImages(r, c, i)
 #endif
 
 	static FORCEINLINE_DEBUGGABLE VkResult  vkCreateInstance(const VkInstanceCreateInfo* CreateInfo, const VkAllocationCallbacks* Allocator, VkInstance* Instance)
@@ -91,16 +140,22 @@ namespace VulkanRHI
 
 		return Result;
 	}
-#if 0
-	static FORCEINLINE_DEBUGGABLE void  vkGetPhysicalDeviceFeatures(
-		VkPhysicalDevice                            physicalDevice,
-		VkPhysicalDeviceFeatures*                   pFeatures);
 
-	static FORCEINLINE_DEBUGGABLE void  vkGetPhysicalDeviceFormatProperties(
-		VkPhysicalDevice                            physicalDevice,
-		VkFormat                                    format,
-		VkFormatProperties*                         pFormatProperties);
+	static FORCEINLINE_DEBUGGABLE void  vkGetPhysicalDeviceFeatures(VkPhysicalDevice PhysicalDevice, VkPhysicalDeviceFeatures* Features)
+	{
+		DumpGetPhysicalDeviceFeatures(PhysicalDevice, Features);
+		::vkGetPhysicalDeviceFeatures(PhysicalDevice, Features);
+		DumpPhysicalDeviceFeatures(Features);
+	}
+
+	static FORCEINLINE_DEBUGGABLE void  vkGetPhysicalDeviceFormatProperties(VkPhysicalDevice PhysicalDevice, VkFormat Format, VkFormatProperties* FormatProperties)
+	{
+		PrintfBegin(FString::Printf(TEXT("vkGetPhysicalDeviceFormatProperties(PhysicalDevice=%p, Format=%d, FormatProperties=%p)[...]"), PhysicalDevice, (int32)Format, FormatProperties));
 
+		::vkGetPhysicalDeviceFormatProperties(PhysicalDevice, Format, FormatProperties);
+	}
+
+#if 0
 	static FORCEINLINE_DEBUGGABLE VkResult  vkGetPhysicalDeviceImageFormatProperties(
 		VkPhysicalDevice                            physicalDevice,
 		VkFormat                                    format,
@@ -109,16 +164,21 @@ namespace VulkanRHI
 		VkImageUsageFlags                           usage,
 		VkImageCreateFlags                          flags,
 		VkImageFormatProperties*                    pImageFormatProperties);
+#endif
+	static FORCEINLINE_DEBUGGABLE void  vkGetPhysicalDeviceProperties(VkPhysicalDevice PhysicalDevice, VkPhysicalDeviceProperties* Properties)
+	{
+		PrintfBegin(FString::Printf(TEXT("vkGetPhysicalDeviceProperties(PhysicalDevice=%p, Properties=%p)[...]"), PhysicalDevice, Properties));
 
-	static FORCEINLINE_DEBUGGABLE void  vkGetPhysicalDeviceProperties(
-		VkPhysicalDevice                            physicalDevice,
-		VkPhysicalDeviceProperties*                 pProperties);
+		::vkGetPhysicalDeviceProperties(PhysicalDevice, Properties);
+	}
+
+	static FORCEINLINE_DEBUGGABLE void  vkGetPhysicalDeviceQueueFamilyProperties(VkPhysicalDevice PhysicalDevice, uint32* QueueFamilyPropertyCount, VkQueueFamilyProperties* QueueFamilyProperties)
+	{
+		PrintfBeginResult(FString::Printf(TEXT("vkGetPhysicalDeviceQueueFamilyProperties(PhysicalDevice=%p, QueueFamilyPropertyCount=%p, QueueFamilyProperties=%p)[...]"), PhysicalDevice, QueueFamilyPropertyCount, QueueFamilyProperties));
+
+		::vkGetPhysicalDeviceQueueFamilyProperties(PhysicalDevice, QueueFamilyPropertyCount, QueueFamilyProperties);
+	}
 
-	static FORCEINLINE_DEBUGGABLE void  vkGetPhysicalDeviceQueueFamilyProperties(
-		VkPhysicalDevice                            physicalDevice,
-		uint32*                                   pQueueFamilyPropertyCount,
-		VkQueueFamilyProperties*                    pQueueFamilyProperties);
-#endif
 	static FORCEINLINE_DEBUGGABLE void  vkGetPhysicalDeviceMemoryProperties(VkPhysicalDevice PhysicalDevice, VkPhysicalDeviceMemoryProperties* MemoryProperties)
 	{
 		PrintfBegin(FString::Printf(TEXT("vkGetPhysicalDeviceMemoryProperties(OutProp=%p)[...]"), MemoryProperties));
@@ -127,45 +187,78 @@ namespace VulkanRHI
 
 		DumpPhysicalDeviceProperties(MemoryProperties);
 	}
-#if 0
-	static FORCEINLINE_DEBUGGABLE PFN_vkVoidFunction  vkGetInstanceProcAddr(
-		VkInstance                                  instance,
-		const char*                                 pName);
 
-	static FORCEINLINE_DEBUGGABLE PFN_vkVoidFunction  vkGetDeviceProcAddr(
-		VkDevice                                    Device,
-		const char*                                 pName);
+	static FORCEINLINE_DEBUGGABLE PFN_vkVoidFunction  vkGetInstanceProcAddr(VkInstance Instance, const char* Name)
+	{
+		PrintfBeginResult(FString::Printf(TEXT("vkGetInstanceProcAddr(Instance=%p, Name=%s)[...]"), Instance, ANSI_TO_TCHAR(Name)));
 
-	static FORCEINLINE_DEBUGGABLE VkResult  vkCreateDevice(
-		VkPhysicalDevice                            physicalDevice,
-		const VkDeviceCreateInfo*                   CreateInfo,
-		const VkAllocationCallbacks*                Allocator,
-		VkDevice*                                   pDevice);
+		PFN_vkVoidFunction Function = ::vkGetInstanceProcAddr(Instance, Name);
+		PrintResultAndPointer(VK_SUCCESS, Function);
+		return Function;
+	}
 
-	static FORCEINLINE_DEBUGGABLE void  vkDestroyDevice(
-		VkDevice                                    Device,
-		const VkAllocationCallbacks*                Allocator);
+	static FORCEINLINE_DEBUGGABLE PFN_vkVoidFunction  vkGetDeviceProcAddr(VkDevice Device, const char* Name)
+	{
+		PrintfBeginResult(FString::Printf(TEXT("vkGetDeviceProcAddr(Device=%p, Name=%s)[...]"), Device, ANSI_TO_TCHAR(Name)));
 
-	static FORCEINLINE_DEBUGGABLE VkResult  vkEnumerateInstanceExtensionProperties(
-		const char*                                 pLayerName,
-		uint32*                                   pPropertyCount,
-		VkExtensionProperties*                      pProperties);
+		PFN_vkVoidFunction Function = ::vkGetDeviceProcAddr(Device, Name);
+		PrintResultAndPointer(VK_SUCCESS, Function);
+		return Function;
+	}
 
-	static FORCEINLINE_DEBUGGABLE VkResult  vkEnumerateDeviceExtensionProperties(
-		VkPhysicalDevice                            physicalDevice,
-		const char*                                 pLayerName,
-		uint32*                                   pPropertyCount,
-		VkExtensionProperties*                      pProperties);
+	static FORCEINLINE_DEBUGGABLE VkResult  vkCreateDevice(VkPhysicalDevice PhysicalDevice, const VkDeviceCreateInfo* CreateInfo, const VkAllocationCallbacks* Allocator, VkDevice* Device)
+	{
+		DumpCreateDevice(PhysicalDevice, CreateInfo, Device);
 
-	static FORCEINLINE_DEBUGGABLE VkResult  vkEnumerateInstanceLayerProperties(
-		uint32*                                   pPropertyCount,
-		VkLayerProperties*                          pProperties);
+		VkResult Result = ::vkCreateDevice(PhysicalDevice, CreateInfo, Allocator, Device);
+
+		PrintResultAndNamedHandle(Result, TEXT("Device"), *Device);
+		return Result;
+	}
+
+	static FORCEINLINE_DEBUGGABLE void  vkDestroyDevice(VkDevice Device, const VkAllocationCallbacks* Allocator)
+	{
+		PrintfBegin(FString::Printf(TEXT("vkDestroyDevice(Device=%p)"), Device));
+
+		::vkDestroyDevice(Device, Allocator);
+	}
+
+	static FORCEINLINE_DEBUGGABLE VkResult  vkEnumerateInstanceExtensionProperties(const char* LayerName, uint32* PropertyCount, VkExtensionProperties* Properties)
+	{
+		PrintfBeginResult(FString::Printf(TEXT("vkEnumerateInstanceExtensionProperties(LayerName=%s, PropertyCount=%p, Properties=%p)[...]"), ANSI_TO_TCHAR(LayerName), PropertyCount, Properties));
+
+		VkResult Result = ::vkEnumerateInstanceExtensionProperties(LayerName, PropertyCount, Properties);
+		PrintResultAndPointer(Result, (void*)(uint64)PropertyCount);
+		return Result;
+	}
+
+	static FORCEINLINE_DEBUGGABLE VkResult  vkEnumerateDeviceExtensionProperties(VkPhysicalDevice PhysicalDevice, const char* LayerName, uint32* PropertyCount, VkExtensionProperties* Properties)
+	{
+		PrintfBeginResult(FString::Printf(TEXT("vkEnumerateDeviceExtensionProperties(Device=%p, LayerName=%s, PropertyCount=%p, Properties=%p)[...]"), PhysicalDevice, ANSI_TO_TCHAR(LayerName), PropertyCount, Properties));
+
+		VkResult Result = ::vkEnumerateDeviceExtensionProperties(PhysicalDevice, LayerName, PropertyCount, Properties);
+		PrintResultAndPointer(Result, (void*)(uint64)PropertyCount);
+		return Result;
+	}
+
+	static FORCEINLINE_DEBUGGABLE VkResult  vkEnumerateInstanceLayerProperties(uint32* PropertyCount, VkLayerProperties* Properties)
+	{
+		PrintfBeginResult(FString::Printf(TEXT("vkEnumerateInstanceLayerProperties(PropertyCount=%p, Properties=%p)[...]"), PropertyCount, Properties));
+
+		VkResult Result = ::vkEnumerateInstanceLayerProperties(PropertyCount, Properties);
+		PrintResultAndPointer(Result, (void*)(uint64)PropertyCount);
+		return Result;
+	}
+
+	static FORCEINLINE_DEBUGGABLE VkResult  vkEnumerateDeviceLayerProperties(VkPhysicalDevice PhysicalDevice, uint32* PropertyCount, VkLayerProperties* Properties)
+	{
+		PrintfBeginResult(FString::Printf(TEXT("vkEnumerateDeviceLayerProperties(Device=%p, PropertyCount=%p, Properties=%p)[...]"), PhysicalDevice, PropertyCount, Properties));
+
+		VkResult Result = ::vkEnumerateDeviceLayerProperties(PhysicalDevice, PropertyCount, Properties);
+		PrintResultAndPointer(Result, (void*)(uint64)PropertyCount);
+		return Result;
+	}
 
-	static FORCEINLINE_DEBUGGABLE VkResult  vkEnumerateDeviceLayerProperties(
-		VkPhysicalDevice                            physicalDevice,
-		uint32*                                   pPropertyCount,
-		VkLayerProperties*                          pProperties);
-#endif
 	static FORCEINLINE_DEBUGGABLE void  vkGetDeviceQueue(VkDevice Device, uint32 QueueFamilyIndex, uint32 QueueIndex, VkQueue* Queue)
 	{
 		DevicePrintfBeginResult(Device, FString::Printf(TEXT("vkGetDeviceQueue(QueueFamilyIndex=%d, QueueIndex=%d, OutQueue=%p)"), QueueFamilyIndex, QueueIndex, Queue));
@@ -184,13 +277,26 @@ namespace VulkanRHI
 		PrintResult(Result);
 		return Result;
 	}
-#if 0
-	static FORCEINLINE_DEBUGGABLE VkResult  vkQueueWaitIdle(
-		VkQueue                                     queue);
+	static FORCEINLINE_DEBUGGABLE VkResult  vkQueueWaitIdle(VkQueue Queue)
+	{
+		PrintfBeginResult(FString::Printf(TEXT("vkQueueWaitIdle(Queue=%p)"), Queue));
+
+		VkResult Result = ::vkQueueWaitIdle(Queue);
+
+		PrintResult(Result);
+		return Result;
+	}
+
+	static FORCEINLINE_DEBUGGABLE VkResult  vkDeviceWaitIdle(VkDevice Device)
+	{
+		DevicePrintfBeginResult(Device, FString::Printf(TEXT("vkDeviceWaitIdle()")));
+
+		VkResult Result = ::vkDeviceWaitIdle(Device);
+
+		PrintResult(Result);
+		return Result;
+	}
 
-	static FORCEINLINE_DEBUGGABLE VkResult  vkDeviceWaitIdle(
-		VkDevice                                    Device);
-#endif
 	static FORCEINLINE_DEBUGGABLE VkResult  vkAllocateMemory(VkDevice Device, const VkMemoryAllocateInfo* AllocateInfo, const VkAllocationCallbacks* Allocator, VkDeviceMemory* Memory)
 	{
 		DumpAllocateMemory(Device, AllocateInfo, Memory);
@@ -218,9 +324,7 @@ namespace VulkanRHI
 		return Result;
 	}
 
-	static FORCEINLINE_DEBUGGABLE void  vkUnmapMemory(
-		VkDevice Device,
-		VkDeviceMemory Memory)
+	static FORCEINLINE_DEBUGGABLE void  vkUnmapMemory(VkDevice Device, VkDeviceMemory Memory)
 	{
 		DevicePrintfBegin(Device, FString::Printf(TEXT("vkUnmapMemory(DevMem=%p)"), Memory));
 		::vkUnmapMemory(Device, Memory);
@@ -275,18 +379,23 @@ namespace VulkanRHI
 
 	static FORCEINLINE_DEBUGGABLE void  vkGetBufferMemoryRequirements(VkDevice Device, VkBuffer Buffer, VkMemoryRequirements* MemoryRequirements)
 	{
-		DevicePrintfBegin(Device, FString::Printf(TEXT("vkGetBufferMemoryRequirements(Buffer=%p, OutReq=%p)"), Buffer, MemoryRequirements));
+		DevicePrintfBeginResult(Device, FString::Printf(TEXT("vkGetBufferMemoryRequirements(Buffer=%p, OutReq=%p)"), Buffer, MemoryRequirements));
 
 		::vkGetBufferMemoryRequirements(Device, Buffer, MemoryRequirements);
 
 		DumpMemoryRequirements(MemoryRequirements);
 	}
-#if 0
-	static FORCEINLINE_DEBUGGABLE void  vkGetImageMemoryRequirements(
-		VkDevice                                    Device,
-		VkImage                                     Image,
-		VkMemoryRequirements*                       MemoryRequirements);
 
+	static FORCEINLINE_DEBUGGABLE void  vkGetImageMemoryRequirements(VkDevice Device, VkImage Image, VkMemoryRequirements* MemoryRequirements)
+	{
+		DevicePrintfBeginResult(Device, FString::Printf(TEXT("vkGetImageMemoryRequirements(Image=%p, OutReq=%p)"), Image, MemoryRequirements));
+
+		::vkGetImageMemoryRequirements(Device, Image, MemoryRequirements);
+
+		DumpMemoryRequirements(MemoryRequirements);
+	}
+
+#if 0
 	static FORCEINLINE_DEBUGGABLE void  vkGetImageSparseMemoryRequirements(
 		VkDevice                                    Device,
 		VkImage                                     Image,
@@ -357,29 +466,39 @@ namespace VulkanRHI
 		PrintResult(Result);
 		return Result;
 	}
-#if 0
-	static FORCEINLINE_DEBUGGABLE VkResult  vkCreateSemaphore(
-		VkDevice                                    Device,
-		const VkSemaphoreCreateInfo*                CreateInfo,
-		const VkAllocationCallbacks*                Allocator,
-		VkSemaphore*                                pSemaphore);
 
-	static FORCEINLINE_DEBUGGABLE void  vkDestroySemaphore(
-		VkDevice                                    Device,
-		VkSemaphore                                 semaphore,
-		const VkAllocationCallbacks*                Allocator);
+	static FORCEINLINE_DEBUGGABLE VkResult  vkCreateSemaphore(VkDevice Device, const VkSemaphoreCreateInfo* CreateInfo, const VkAllocationCallbacks* Allocator, VkSemaphore* Semaphore)
+	{
+		DumpSemaphoreCreate(Device, CreateInfo, Semaphore);
+
+		VkResult Result = ::vkCreateSemaphore(Device, CreateInfo, Allocator, Semaphore);
+
+		PrintResultAndNamedHandle(Result, TEXT("Semaphore"), *Semaphore);
+		return Result;
+	}
 
+
+	static FORCEINLINE_DEBUGGABLE void  vkDestroySemaphore(VkDevice Device, VkSemaphore Semaphore, const VkAllocationCallbacks* Allocator)
+	{
+		DevicePrintfBegin(Device, FString::Printf(TEXT("vkDestroySemaphore(Semaphore=%p)"), Semaphore));
+
+		::vkDestroySemaphore(Device, Semaphore, Allocator);
+	}
+
+#if 0
 	static FORCEINLINE_DEBUGGABLE VkResult  vkCreateEvent(
 		VkDevice                                    Device,
 		const VkEventCreateInfo*                    CreateInfo,
 		const VkAllocationCallbacks*                Allocator,
 		VkEvent*                                    pEvent);
+#endif
+	static FORCEINLINE_DEBUGGABLE void  vkDestroyEvent(VkDevice Device, VkEvent Event, const VkAllocationCallbacks* Allocator)
+	{
+		DevicePrintfBegin(Device, FString::Printf(TEXT("vkDestroyEvent(Event=%p)"), Event));
 
-	static FORCEINLINE_DEBUGGABLE void  vkDestroyEvent(
-		VkDevice                                    Device,
-		VkEvent                                     event,
-		const VkAllocationCallbacks*                Allocator);
-
+		::vkDestroyEvent(Device, Event, Allocator);
+	}
+#if 0
 	static FORCEINLINE_DEBUGGABLE VkResult  vkGetEventStatus(
 		VkDevice                                    Device,
 		VkEvent                                     event);
@@ -391,28 +510,36 @@ namespace VulkanRHI
 	static FORCEINLINE_DEBUGGABLE VkResult  vkResetEvent(
 		VkDevice                                    Device,
 		VkEvent                                     event);
+#endif
+	static FORCEINLINE_DEBUGGABLE VkResult  vkCreateQueryPool(VkDevice Device, const VkQueryPoolCreateInfo* CreateInfo, const VkAllocationCallbacks* Allocator, VkQueryPool* QueryPool)
+	{
+		DumpCreateQueryPool(Device, CreateInfo, QueryPool);
 
-	static FORCEINLINE_DEBUGGABLE VkResult  vkCreateQueryPool(
-		VkDevice                                    Device,
-		const VkQueryPoolCreateInfo*                CreateInfo,
-		const VkAllocationCallbacks*                Allocator,
-		VkQueryPool*                                pQueryPool);
+		VkResult Result = ::vkCreateQueryPool(Device, CreateInfo, Allocator, QueryPool);
 
-	static FORCEINLINE_DEBUGGABLE void  vkDestroyQueryPool(
-		VkDevice                                    Device,
-		VkQueryPool                                 queryPool,
-		const VkAllocationCallbacks*                Allocator);
+		PrintResultAndNamedHandle(Result, TEXT("QueryPool"), *QueryPool);
+
+		return Result;
+	}
+
+	static FORCEINLINE_DEBUGGABLE void  vkDestroyQueryPool(VkDevice Device, VkQueryPool QueryPool, const VkAllocationCallbacks* Allocator)
+	{
+		DevicePrintfBegin(Device, FString::Printf(TEXT("vkDestroyQueryPool(QueryPool=%p)"), QueryPool));
+
+		::vkDestroyQueryPool(Device, QueryPool, Allocator);
+	}
+
+	static FORCEINLINE_DEBUGGABLE VkResult  vkGetQueryPoolResults(VkDevice Device, VkQueryPool QueryPool, 
+		uint32 FirstQuery, uint32 QueryCount, size_t DataSize, void* Data, VkDeviceSize Stride, VkQueryResultFlags Flags)
+	{
+		DevicePrintfBeginResult(Device, FString::Printf(TEXT("vkGetQueryPoolResults(QueryPool=%p, FirstQuery=%d, QueryCount=%d, DataSize=%d, Data=%p, Stride=%d, Flags=%d)[...]"), QueryPool, FirstQuery, QueryCount, (int32)DataSize, Data, Stride, Flags));
+
+		VkResult Result = ::vkGetQueryPoolResults(Device, QueryPool, FirstQuery, QueryCount, DataSize, Data, Stride, Flags);
+
+		PrintResult(Result);
+		return Result;
+	}
 
-	static FORCEINLINE_DEBUGGABLE VkResult  vkGetQueryPoolResults(
-		VkDevice                                    Device,
-		VkQueryPool                                 queryPool,
-		uint32                                    firstQuery,
-		uint32                                    queryCount,
-		size_t                                      dataSize,
-		void*                                       pData,
-		VkDeviceSize                                stride,
-		VkQueryResultFlags                          flags);
-#endif
 	static FORCEINLINE_DEBUGGABLE VkResult  vkCreateBuffer(VkDevice Device, const VkBufferCreateInfo* CreateInfo, const VkAllocationCallbacks* Allocator, VkBuffer* Buffer)
 	{
 		DumpBufferCreate(Device, CreateInfo, Buffer);
@@ -447,42 +574,54 @@ namespace VulkanRHI
 
 		::vkDestroyBufferView(Device, BufferView, Allocator);
 	}
-#if 0
-	static FORCEINLINE_DEBUGGABLE VkResult  vkCreateImage(
-		VkDevice                                    Device,
-		const VkImageCreateInfo*                    CreateInfo,
-		const VkAllocationCallbacks*                Allocator,
-		VkImage*                                    pImage);
 
-	static FORCEINLINE_DEBUGGABLE void  vkDestroyImage(
-		VkDevice                                    Device,
-		VkImage                                     Image,
-		const VkAllocationCallbacks*                Allocator);
+	static FORCEINLINE_DEBUGGABLE VkResult  vkCreateImage(VkDevice Device, const VkImageCreateInfo* CreateInfo, const VkAllocationCallbacks* Allocator, VkImage* Image)
+	{
+		DumpImageCreate(Device, CreateInfo, Image);
 
-	static FORCEINLINE_DEBUGGABLE void  vkGetImageSubresourceLayout(
-		VkDevice                                    Device,
-		VkImage                                     Image,
-		const VkImageSubresource*                   pSubresource,
-		VkSubresourceLayout*                        pLayout);
+		VkResult Result = ::vkCreateImage(Device, CreateInfo, Allocator, Image);
 
-	static FORCEINLINE_DEBUGGABLE VkResult  vkCreateImageView(
-		VkDevice                                    Device,
-		const VkImageViewCreateInfo*                CreateInfo,
-		const VkAllocationCallbacks*                Allocator,
-		VkImageView*                                pView);
+		PrintResultAndNamedHandle(Result, TEXT("Image"), *Image);
+		return Result;
+	}
 
-	static FORCEINLINE_DEBUGGABLE void  vkDestroyImageView(
-		VkDevice                                    Device,
-		VkImageView                                 imageView,
-		const VkAllocationCallbacks*                Allocator);
-#endif
-	static FORCEINLINE_DEBUGGABLE VkResult  vkCreateShaderModule(
-		VkDevice Device,
-		const VkShaderModuleCreateInfo* CreateInfo,
-		const VkAllocationCallbacks* Allocator,
-		VkShaderModule* ShaderModule)
+	static FORCEINLINE_DEBUGGABLE void  vkDestroyImage(VkDevice Device, VkImage Image, const VkAllocationCallbacks* Allocator)
+	{
+		DevicePrintfBegin(Device, FString::Printf(TEXT("vkDestroyImage(Image=%p)"), Image));
+
+		::vkDestroyImage(Device, Image, Allocator);
+	}
+
+	static FORCEINLINE_DEBUGGABLE void  vkGetImageSubresourceLayout(VkDevice Device, VkImage Image, const VkImageSubresource* Subresource, VkSubresourceLayout* Layout)
+	{
+		DumpGetImageSubresourceLayout(Device, Image, Subresource, Layout);
+
+		::vkGetImageSubresourceLayout(Device, Image, Subresource, Layout);
+
+		DumpImageSubresourceLayout(Layout);
+	}
+
+	static FORCEINLINE_DEBUGGABLE VkResult  vkCreateImageView(VkDevice Device, const VkImageViewCreateInfo* CreateInfo, const VkAllocationCallbacks* Allocator, VkImageView* View)
+	{
+		DumpImageViewCreate(Device, CreateInfo, View);
+
+		VkResult Result = ::vkCreateImageView(Device, CreateInfo, Allocator, View);
+
+		PrintResultAndNamedHandle(Result, TEXT("ImageView"), *View);
+		return Result;
+	}
+
+	static FORCEINLINE_DEBUGGABLE void  vkDestroyImageView(VkDevice Device, VkImageView ImageView, const VkAllocationCallbacks* Allocator)
+	{
+		DevicePrintfBegin(Device, FString::Printf(TEXT("vkDestroyImageView(ImageView=%p)"), ImageView));
+
+		::vkDestroyImageView(Device, ImageView, Allocator);
+	}
+
+	static FORCEINLINE_DEBUGGABLE VkResult  vkCreateShaderModule(VkDevice Device, const VkShaderModuleCreateInfo* CreateInfo, const VkAllocationCallbacks* Allocator, VkShaderModule* ShaderModule)
 	{
 		DumpCreateShaderModule(Device, CreateInfo, ShaderModule);
+
 		VkResult Result = ::vkCreateShaderModule(Device, CreateInfo, Allocator,ShaderModule);
 
 		PrintResultAndNamedHandle(Result, TEXT("ShaderModule"), *ShaderModule);
@@ -497,18 +636,25 @@ namespace VulkanRHI
 		::vkDestroyShaderModule(Device, ShaderModule, Allocator);
 	}
 
-#if 0
-	static FORCEINLINE_DEBUGGABLE VkResult  vkCreatePipelineCache(
-		VkDevice                                    Device,
-		const VkPipelineCacheCreateInfo*            CreateInfo,
-		const VkAllocationCallbacks*                Allocator,
-		VkPipelineCache*                            pPipelineCache);
+	static FORCEINLINE_DEBUGGABLE VkResult  vkCreatePipelineCache(VkDevice Device, const VkPipelineCacheCreateInfo* CreateInfo, const VkAllocationCallbacks* Allocator, VkPipelineCache* PipelineCache)
+	{
+		DumpCreatePipelineCache(Device, CreateInfo, PipelineCache);
 
-	static FORCEINLINE_DEBUGGABLE void  vkDestroyPipelineCache(
-		VkDevice                                    Device,
-		VkPipelineCache                             pipelineCache,
-		const VkAllocationCallbacks*                Allocator);
+		VkResult Result = ::vkCreatePipelineCache(Device, CreateInfo, Allocator, PipelineCache);
+
+		PrintResultAndNamedHandle(Result, TEXT("PipelineCache"), *PipelineCache);
+
+		return Result;
+	}
 
+	static FORCEINLINE_DEBUGGABLE void  vkDestroyPipelineCache(VkDevice Device, VkPipelineCache PipelineCache, const VkAllocationCallbacks* Allocator)
+	{
+		DevicePrintfBegin(Device, FString::Printf(TEXT("vkDestroyPipelineCache(PipelineCache=%p)"), PipelineCache));
+
+		::vkDestroyPipelineCache(Device, PipelineCache, Allocator);
+	}
+
+#if 0
 	static FORCEINLINE_DEBUGGABLE VkResult  vkGetPipelineCacheData(
 		VkDevice                                    Device,
 		VkPipelineCache                             pipelineCache,
@@ -520,50 +666,72 @@ namespace VulkanRHI
 		VkPipelineCache                             dstCache,
 		uint32                                    srcCacheCount,
 		const VkPipelineCache*                      pSrcCaches);
+#endif
+	static FORCEINLINE_DEBUGGABLE VkResult  vkCreateGraphicsPipelines(VkDevice Device, VkPipelineCache PipelineCache, uint32 CreateInfoCount, const VkGraphicsPipelineCreateInfo* CreateInfos, const VkAllocationCallbacks* Allocator, VkPipeline* Pipelines)
+	{
+		DevicePrintfBeginResult(Device, FString::Printf(TEXT("vkCreateGraphicsPipelines(PipelineCache=%p, CreateInfoCount=%d, CreateInfos=%p, OutPipelines=%p)[...]"), PipelineCache, CreateInfoCount, CreateInfos, Pipelines));
 
-	static FORCEINLINE_DEBUGGABLE VkResult  vkCreateGraphicsPipelines(
-		VkDevice                                    Device,
-		VkPipelineCache                             pipelineCache,
-		uint32                                    createInfoCount,
-		const VkGraphicsPipelineCreateInfo*         pCreateInfos,
-		const VkAllocationCallbacks*                Allocator,
-		VkPipeline*                                 pPipelines);
+		VkResult Result = ::vkCreateGraphicsPipelines(Device, PipelineCache, CreateInfoCount, CreateInfos, Allocator, Pipelines);
 
-	static FORCEINLINE_DEBUGGABLE VkResult  vkCreateComputePipelines(
-		VkDevice                                    Device,
-		VkPipelineCache                             pipelineCache,
-		uint32                                    createInfoCount,
-		const VkComputePipelineCreateInfo*          pCreateInfos,
-		const VkAllocationCallbacks*                Allocator,
-		VkPipeline*                                 pPipelines);
+		//#todo-rco: Multiple pipelines!
+		PrintResultAndNamedHandle(Result, TEXT("Pipeline"), Pipelines[0]);
+		return Result;
+	}
 
-	static FORCEINLINE_DEBUGGABLE void  vkDestroyPipeline(
-		VkDevice                                    Device,
-		VkPipeline                                  pipeline,
-		const VkAllocationCallbacks*                Allocator);
+	static FORCEINLINE_DEBUGGABLE VkResult  vkCreateComputePipelines(VkDevice Device, VkPipelineCache PipelineCache, uint32 CreateInfoCount, const VkComputePipelineCreateInfo* CreateInfos, const VkAllocationCallbacks* Allocator, VkPipeline* Pipelines)
+	{
+		DevicePrintfBeginResult(Device, FString::Printf(TEXT("vkCreateComputePipelines(PipelineCache=%p, CreateInfoCount=%d, CreateInfos=%p, OutPipelines=%p)[...]"), PipelineCache, CreateInfoCount, CreateInfos, Pipelines));
 
-	static FORCEINLINE_DEBUGGABLE VkResult  vkCreatePipelineLayout(
-		VkDevice                                    Device,
-		const VkPipelineLayoutCreateInfo*           CreateInfo,
-		const VkAllocationCallbacks*                Allocator,
-		VkPipelineLayout*                           pPipelineLayout);
+		VkResult Result = ::vkCreateComputePipelines(Device, PipelineCache, CreateInfoCount, CreateInfos, Allocator, Pipelines);
 
-	static FORCEINLINE_DEBUGGABLE void  vkDestroyPipelineLayout(
-		VkDevice                                    Device,
-		VkPipelineLayout                            pipelineLayout,
-		const VkAllocationCallbacks*                Allocator);
+		//#todo-rco: Multiple pipelines!
+		PrintResultAndNamedHandle(Result, TEXT("Pipeline"), Pipelines[0]);
+		return Result;
+	}
 
-	static FORCEINLINE_DEBUGGABLE VkResult  vkCreateSampler(
-		VkDevice                                    Device,
-		const VkSamplerCreateInfo*                  CreateInfo,
-		const VkAllocationCallbacks*                Allocator,
-		VkSampler*                                  pSampler);
+	static FORCEINLINE_DEBUGGABLE void  vkDestroyPipeline(VkDevice Device, VkPipeline Pipeline, const VkAllocationCallbacks* Allocator)
+	{
+		DevicePrintfBegin(Device, FString::Printf(TEXT("vkDestroyPipeline(Pipeline=%p)"), Pipeline));
+
+		::vkDestroyPipeline(Device, Pipeline, Allocator);
+	}
+
+	static FORCEINLINE_DEBUGGABLE VkResult  vkCreatePipelineLayout(VkDevice Device, const VkPipelineLayoutCreateInfo* CreateInfo, const VkAllocationCallbacks* Allocator, VkPipelineLayout* PipelineLayout)
+	{
+		DumpCreatePipelineLayout(Device, CreateInfo, PipelineLayout);
+
+		VkResult Result = ::vkCreatePipelineLayout(Device, CreateInfo, Allocator, PipelineLayout);
+
+		PrintResultAndNamedHandle(Result, TEXT("PipelineLayout"), *PipelineLayout);
+
+		return Result;
+	}
+
+	static FORCEINLINE_DEBUGGABLE void  vkDestroyPipelineLayout(VkDevice Device, VkPipelineLayout PipelineLayout, const VkAllocationCallbacks* Allocator)
+	{
+		DevicePrintfBegin(Device, FString::Printf(TEXT("vkDestroyPipelineLayout(PipelineLayout=%p)"), PipelineLayout));
+
+		::vkDestroyPipelineLayout(Device, PipelineLayout, Allocator);
+	}
+
+	static FORCEINLINE_DEBUGGABLE VkResult  vkCreateSampler(VkDevice Device, const VkSamplerCreateInfo* CreateInfo, const VkAllocationCallbacks* Allocator, VkSampler* Sampler)
+	{
+		DumpCreateSampler(Device, CreateInfo, Sampler);
+
+		VkResult Result = ::vkCreateSampler(Device, CreateInfo, Allocator, Sampler);
+
+		PrintResultAndNamedHandle(Result, TEXT("Sampler"), *Sampler);
+
+		return Result;
+	}
+
+	static FORCEINLINE_DEBUGGABLE void  vkDestroySampler(VkDevice Device, VkSampler Sampler, const VkAllocationCallbacks* Allocator)
+	{
+		DevicePrintfBegin(Device, FString::Printf(TEXT("vkDestroySampler(Sampler=%p)"), Sampler));
+
+		::vkDestroySampler(Device, Sampler, Allocator);
+	}
 
-	static FORCEINLINE_DEBUGGABLE void  vkDestroySampler(
-		VkDevice                                    Device,
-		VkSampler                                   sampler,
-		const VkAllocationCallbacks*                Allocator);
-#endif
 	static FORCEINLINE_DEBUGGABLE VkResult  vkCreateDescriptorSetLayout(VkDevice Device, const VkDescriptorSetLayoutCreateInfo* CreateInfo, const VkAllocationCallbacks* Allocator, VkDescriptorSetLayout* SetLayout)
 	{
 		DumpCreateDescriptorSetLayout(Device, CreateInfo, SetLayout);
@@ -580,18 +748,25 @@ namespace VulkanRHI
 
 		::vkDestroyDescriptorSetLayout(Device, DescriptorSetLayout, Allocator);
 	}
-#if 0
-	static FORCEINLINE_DEBUGGABLE VkResult  vkCreateDescriptorPool(
-		VkDevice                                    Device,
-		const VkDescriptorPoolCreateInfo*           CreateInfo,
-		const VkAllocationCallbacks*                Allocator,
-		VkDescriptorPool*                           pDescriptorPool);
 
-	static FORCEINLINE_DEBUGGABLE void  vkDestroyDescriptorPool(
-		VkDevice                                    Device,
-		VkDescriptorPool                            descriptorPool,
-		const VkAllocationCallbacks*                Allocator);
+	static FORCEINLINE_DEBUGGABLE VkResult  vkCreateDescriptorPool(VkDevice Device, const VkDescriptorPoolCreateInfo* CreateInfo, const VkAllocationCallbacks* Allocator, VkDescriptorPool* DescriptorPool)
+	{
+		DumpCreateDescriptorPool(Device, CreateInfo, DescriptorPool);
+
+		VkResult Result = ::vkCreateDescriptorPool(Device, CreateInfo, Allocator, DescriptorPool);
+
+		PrintResultAndNamedHandle(Result, TEXT("DescriptorPool"), *DescriptorPool);
+
+		return Result;
+	}
+
+	static FORCEINLINE_DEBUGGABLE void  vkDestroyDescriptorPool(VkDevice Device, VkDescriptorPool DescriptorPool, const VkAllocationCallbacks* Allocator)
+	{
+		DevicePrintfBegin(Device, FString::Printf(TEXT("vkDestroyDescriptorPool(DescriptorPool=%p)"), DescriptorPool));
 
+		::vkDestroyDescriptorPool(Device, DescriptorPool, Allocator);
+	}
+#if 0
 	static FORCEINLINE_DEBUGGABLE VkResult  vkResetDescriptorPool(
 		VkDevice                                    Device,
 		VkDescriptorPool                            descriptorPool,
@@ -616,28 +791,35 @@ namespace VulkanRHI
 		PrintResult(Result);
 		return Result;
 	}
-#if 0
-	static FORCEINLINE_DEBUGGABLE void  vkUpdateDescriptorSets(
-		VkDevice                                    Device,
-		uint32                                    descriptorWriteCount,
-		const VkWriteDescriptorSet*                 pDescriptorWrites,
-		uint32                                    descriptorCopyCount,
-		const VkCopyDescriptorSet*                  pDescriptorCopies);
 
-	static FORCEINLINE_DEBUGGABLE VkResult  vkCreateFramebuffer(
-		VkDevice                                    Device,
-		const VkFramebufferCreateInfo*              CreateInfo,
-		const VkAllocationCallbacks*                Allocator,
-		VkFramebuffer*                              pFramebuffer);
+	static FORCEINLINE_DEBUGGABLE void  vkUpdateDescriptorSets(VkDevice Device, uint32 DescriptorWriteCount, const VkWriteDescriptorSet* DescriptorWrites, uint32 DescriptorCopyCount, const VkCopyDescriptorSet* DescriptorCopies)
+	{
+		DumpUpdateDescriptorSets(Device, DescriptorWriteCount, DescriptorWrites, DescriptorCopyCount, DescriptorCopies);
+
+		::vkUpdateDescriptorSets(Device, DescriptorWriteCount, DescriptorWrites, DescriptorCopyCount, DescriptorCopies);
+	}
+
+	static FORCEINLINE_DEBUGGABLE VkResult  vkCreateFramebuffer(VkDevice Device, const VkFramebufferCreateInfo* CreateInfo, const VkAllocationCallbacks* Allocator, VkFramebuffer* Framebuffer)
+	{
+		DumpCreateFramebuffer(Device, CreateInfo, Framebuffer);
+
+		VkResult Result = ::vkCreateFramebuffer(Device, CreateInfo, Allocator, Framebuffer);
+
+		PrintResultAndNamedHandle(Result, TEXT("Framebuffer"), *Framebuffer);
+		return Result;
+	}
+
+	static FORCEINLINE_DEBUGGABLE void  vkDestroyFramebuffer(VkDevice Device, VkFramebuffer Framebuffer, const VkAllocationCallbacks* Allocator)
+	{
+		DevicePrintfBegin(Device, FString::Printf(TEXT("vkDestroyFramebuffer(Framebuffer=%p)"), Framebuffer));
+
+		::vkDestroyFramebuffer(Device, Framebuffer, Allocator);
+	}
 
-	static FORCEINLINE_DEBUGGABLE void  vkDestroyFramebuffer(
-		VkDevice                                    Device,
-		VkFramebuffer                               framebuffer,
-		const VkAllocationCallbacks*                Allocator);
-#endif
 	static FORCEINLINE_DEBUGGABLE VkResult  vkCreateRenderPass(VkDevice Device, const VkRenderPassCreateInfo* CreateInfo, const VkAllocationCallbacks* Allocator, VkRenderPass* RenderPass)
 	{
 		DumpCreateRenderPass(Device, CreateInfo, RenderPass);
+
 		VkResult Result = ::vkCreateRenderPass(Device, CreateInfo, Allocator, RenderPass);
 
 		PrintResultAndNamedHandle(Result, TEXT("RenderPass"), *RenderPass);
@@ -655,66 +837,106 @@ namespace VulkanRHI
 		VkDevice                                    Device,
 		VkRenderPass                                renderPass,
 		VkExtent2D*                                 pGranularity);
+#endif
+	static FORCEINLINE_DEBUGGABLE VkResult  vkCreateCommandPool(VkDevice Device, const VkCommandPoolCreateInfo* CreateInfo, const VkAllocationCallbacks* Allocator, VkCommandPool* CommandPool)
+	{
+		DumpCreateCommandPool(Device, CreateInfo, CommandPool);
 
-	static FORCEINLINE_DEBUGGABLE VkResult  vkCreateCommandPool(
-		VkDevice                                    Device,
-		const VkCommandPoolCreateInfo*              CreateInfo,
-		const VkAllocationCallbacks*                Allocator,
-		VkCommandPool*                              pCommandPool);
+		VkResult Result = ::vkCreateCommandPool(Device, CreateInfo, Allocator, CommandPool);
 
-	static FORCEINLINE_DEBUGGABLE void  vkDestroyCommandPool(
-		VkDevice                                    Device,
-		VkCommandPool                               commandPool,
-		const VkAllocationCallbacks*                Allocator);
+		PrintResultAndNamedHandle(Result, TEXT("CommandPool"), *CommandPool);
+
+		return Result;
+	}
 
+	static FORCEINLINE_DEBUGGABLE void  vkDestroyCommandPool(VkDevice Device, VkCommandPool CommandPool, const VkAllocationCallbacks* Allocator)
+	{
+		DevicePrintfBegin(Device, FString::Printf(TEXT("vkDestroyCommandPool(CommandPool=%p)"), CommandPool));
+
+		::vkDestroyCommandPool(Device, CommandPool, Allocator);
+	}
+#if 0
 	static FORCEINLINE_DEBUGGABLE VkResult  vkResetCommandPool(
 		VkDevice                                    Device,
 		VkCommandPool                               commandPool,
 		VkCommandPoolResetFlags                     flags);
+#endif
+	static FORCEINLINE_DEBUGGABLE VkResult  vkAllocateCommandBuffers(VkDevice Device, const VkCommandBufferAllocateInfo* AllocateInfo, VkCommandBuffer* CommandBuffers)
+	{
+		DevicePrintfBeginResult(Device, FString::Printf(TEXT("vkAllocateCommandBuffers(AllocateInfo=%p, OutCommandBuffers=%p)[...]"), AllocateInfo, CommandBuffers));
 
-	static FORCEINLINE_DEBUGGABLE VkResult  vkAllocateCommandBuffers(
-		VkDevice                                    Device,
-		const VkCommandBufferAllocateInfo*          pAllocateInfo,
-		VkCommandBuffer*                            pCommandBuffers);
+		VkResult Result = ::vkAllocateCommandBuffers(Device, AllocateInfo, CommandBuffers);
+		PrintResultAndNamedHandle(Result, TEXT("CommandBuffers"), *CommandBuffers);
+		return Result;
+	}
 
-	static FORCEINLINE_DEBUGGABLE void  vkFreeCommandBuffers(
-		VkDevice                                    Device,
-		VkCommandPool                               commandPool,
-		uint32                                    commandBufferCount,
-		const VkCommandBuffer*                      pCommandBuffers);
+	static FORCEINLINE_DEBUGGABLE void  vkFreeCommandBuffers(VkDevice Device, VkCommandPool CommandPool, uint32 CommandBufferCount, const VkCommandBuffer* CommandBuffers)
+	{
+		DevicePrintfBegin(Device, FString::Printf(TEXT("vkFreeCommandBuffers(CommandPool=%p, CommandBufferCount=%d, CommandBuffers=%p)[...]"), CommandPool, CommandBufferCount, CommandBuffers));
 
-	static FORCEINLINE_DEBUGGABLE VkResult  vkBeginCommandBuffer(
-		VkCommandBuffer                             commandBuffer,
-		const VkCommandBufferBeginInfo*             pBeginInfo);
+		::vkFreeCommandBuffers(Device, CommandPool, CommandBufferCount, CommandBuffers);
+	}
+
+	static FORCEINLINE_DEBUGGABLE VkResult  vkBeginCommandBuffer(VkCommandBuffer CommandBuffer, const VkCommandBufferBeginInfo* BeginInfo)
+	{
+		DumpBeginCommandBuffer(CommandBuffer, BeginInfo);
 
-	static FORCEINLINE_DEBUGGABLE VkResult  vkEndCommandBuffer(
-		VkCommandBuffer                             commandBuffer);
+		VkResult Result = ::vkBeginCommandBuffer(CommandBuffer, BeginInfo);
 
-	static FORCEINLINE_DEBUGGABLE VkResult  vkResetCommandBuffer(
-		VkCommandBuffer                             commandBuffer,
-		VkCommandBufferResetFlags                   flags);
+		PrintResult(Result);
 
-	static FORCEINLINE_DEBUGGABLE void  vkCmdBindPipeline(
-		VkCommandBuffer                             commandBuffer,
-		VkPipelineBindPoint                         pipelineBindPoint,
-		VkPipeline                                  pipeline);
+		return Result;
+	}
 
-	static FORCEINLINE_DEBUGGABLE void  vkCmdSetViewport(
-		VkCommandBuffer                             commandBuffer,
-		uint32                                    firstViewport,
-		uint32                                    viewportCount,
-		const VkViewport*                           pViewports);
+	static FORCEINLINE_DEBUGGABLE VkResult  vkEndCommandBuffer(VkCommandBuffer CommandBuffer)
+	{
+		CmdPrintfBeginResult(CommandBuffer, FString::Printf(TEXT("vkEndCommandBuffer(Cmd=%p)")));
 
-	static FORCEINLINE_DEBUGGABLE void  vkCmdSetScissor(
-		VkCommandBuffer                             commandBuffer,
-		uint32                                    firstScissor,
-		uint32                                    scissorCount,
-		const VkRect2D*                             pScissors);
+		VkResult Result = ::vkEndCommandBuffer(CommandBuffer);
 
-	static FORCEINLINE_DEBUGGABLE void  vkCmdSetLineWidth(
-		VkCommandBuffer                             commandBuffer,
-		float                                       lineWidth);
+		PrintResult(Result);
+		return Result;
+	}
+
+	static FORCEINLINE_DEBUGGABLE VkResult  vkResetCommandBuffer(VkCommandBuffer CommandBuffer, VkCommandBufferResetFlags Flags)
+	 {
+		PrintfBeginResult(FString::Printf(TEXT("vkResetCommandBuffer(Cmd=%p, Flags=%d)"), CommandBuffer, Flags));
 
+		VkResult Result = ::vkResetCommandBuffer(CommandBuffer, Flags);
+
+		PrintResult(Result);
+		return Result;
+	}
+
+	static FORCEINLINE_DEBUGGABLE void  vkCmdBindPipeline(VkCommandBuffer CommandBuffer, VkPipelineBindPoint PipelineBindPoint, VkPipeline Pipeline)
+	{
+		CmdPrintfBegin(CommandBuffer, FString::Printf(TEXT("vkCmdBindPipeline(BindPoint=%d, Pipeline=%p)[...]"), (int32)PipelineBindPoint, Pipeline));
+
+		::vkCmdBindPipeline(CommandBuffer, PipelineBindPoint, Pipeline);
+	}
+
+	static FORCEINLINE_DEBUGGABLE void  vkCmdSetViewport(VkCommandBuffer CommandBuffer, uint32 FirstViewport, uint32 ViewportCount, const VkViewport* Viewports)
+	{
+		CmdPrintfBegin(CommandBuffer, FString::Printf(TEXT("vkCmdSetViewport(FirstViewport=%d, ViewportCount=%d, Viewports=%p)[...]"), FirstViewport, ViewportCount, Viewports));
+
+		::vkCmdSetViewport(CommandBuffer, FirstViewport, ViewportCount, Viewports);
+	}
+
+	static FORCEINLINE_DEBUGGABLE void  vkCmdSetScissor(VkCommandBuffer CommandBuffer, uint32 FirstScissor, uint32 ScissorCount, const VkRect2D* Scissors)
+	{
+		CmdPrintfBegin(CommandBuffer, FString::Printf(TEXT("vkCmdSetScissor(FirstScissor=%d, ScissorCount=%d, Scissors=%p)[...]"), FirstScissor, ScissorCount, Scissors));
+		
+		::vkCmdSetScissor(CommandBuffer, FirstScissor, ScissorCount, Scissors);
+	}
+
+	static FORCEINLINE_DEBUGGABLE void  vkCmdSetLineWidth(VkCommandBuffer CommandBuffer, float LineWidth)
+	{
+		CmdPrintfBegin(CommandBuffer, FString::Printf(TEXT("vkCmdSetLineWidth(Width=%f)"), LineWidth));
+
+		::vkCmdSetLineWidth(CommandBuffer, LineWidth);
+	}
+
+#if 0
 	static FORCEINLINE_DEBUGGABLE void  vkCmdSetDepthBias(
 		VkCommandBuffer                             commandBuffer,
 		float                                       depthBiasConstantFactor,
@@ -744,45 +966,43 @@ namespace VulkanRHI
 		VkCommandBuffer                             commandBuffer,
 		VkStencilFaceFlags                          faceMask,
 		uint32                                    reference);
+#endif
+	static FORCEINLINE_DEBUGGABLE void  vkCmdBindDescriptorSets(VkCommandBuffer CommandBuffer, VkPipelineBindPoint PipelineBindPoint, VkPipelineLayout Layout, uint32 FirstSet, uint32 DescriptorSetCount, const VkDescriptorSet* DescriptorSets, uint32 DynamicOffsetCount, const uint32* DynamicOffsets)
+	{
+		DumpBindDescriptorSets(CommandBuffer, PipelineBindPoint, Layout, FirstSet, DescriptorSetCount, DescriptorSets, DynamicOffsetCount, DynamicOffsets);
 
-	static FORCEINLINE_DEBUGGABLE void  vkCmdBindDescriptorSets(
-		VkCommandBuffer                             commandBuffer,
-		VkPipelineBindPoint                         pipelineBindPoint,
-		VkPipelineLayout                            layout,
-		uint32                                    firstSet,
-		uint32                                    descriptorSetCount,
-		const VkDescriptorSet*                      pDescriptorSets,
-		uint32                                    dynamicOffsetCount,
-		const uint32*                             pDynamicOffsets);
+		::vkCmdBindDescriptorSets(CommandBuffer, PipelineBindPoint, Layout, FirstSet, DescriptorSetCount, DescriptorSets, DynamicOffsetCount, DynamicOffsets);
+	}
 
-	static FORCEINLINE_DEBUGGABLE void  vkCmdBindIndexBuffer(
-		VkCommandBuffer                             commandBuffer,
-		VkBuffer                                    Buffer,
-		VkDeviceSize                                offset,
-		VkIndexType                                 indexType);
+	static FORCEINLINE_DEBUGGABLE void  vkCmdBindIndexBuffer(VkCommandBuffer CommandBuffer, VkBuffer Buffer, VkDeviceSize Offset, VkIndexType IndexType)
+	{
+		CmdPrintfBegin(CommandBuffer, FString::Printf(TEXT("vkCmdBindIndexBuffer(Buffer=%p, Offset=%d, IndexType=%d)"), Buffer, (int32)Offset, (int32)IndexType));
 
-	static FORCEINLINE_DEBUGGABLE void  vkCmdBindVertexBuffers(
-		VkCommandBuffer                             commandBuffer,
-		uint32                                    firstBinding,
-		uint32                                    bindingCount,
-		const VkBuffer*                             pBuffers,
-		const VkDeviceSize*                         pOffsets);
+		::vkCmdBindIndexBuffer(CommandBuffer, Buffer, Offset, IndexType);
+	}
 
-	static FORCEINLINE_DEBUGGABLE void  vkCmdDraw(
-		VkCommandBuffer                             commandBuffer,
-		uint32                                    vertexCount,
-		uint32                                    instanceCount,
-		uint32                                    firstVertex,
-		uint32                                    firstInstance);
+	static FORCEINLINE_DEBUGGABLE void  vkCmdBindVertexBuffers(VkCommandBuffer CommandBuffer, uint32 FirstBinding, uint32 BindingCount, const VkBuffer* Buffers, const VkDeviceSize* Offsets)
+	{
+		DumpCmdBindVertexBuffers(CommandBuffer, FirstBinding, BindingCount, Buffers, Offsets);
 
-	static FORCEINLINE_DEBUGGABLE void  vkCmdDrawIndexed(
-		VkCommandBuffer                             commandBuffer,
-		uint32                                    indexCount,
-		uint32                                    instanceCount,
-		uint32                                    firstIndex,
-		int32_t                                     vertexOffset,
-		uint32                                    firstInstance);
+		::vkCmdBindVertexBuffers(CommandBuffer, FirstBinding, BindingCount, Buffers, Offsets);
+	}
+
+	static FORCEINLINE_DEBUGGABLE void  vkCmdDraw(VkCommandBuffer CommandBuffer, uint32 VertexCount, uint32 InstanceCount, uint32 FirstVertex, uint32 FirstInstance)
+	{
+		CmdPrintfBegin(CommandBuffer, FString::Printf(TEXT("vkCmdDraw(NumVertices=%d, NumInstances=%d, FirstVertex=%d, FirstInstance=%d)"), VertexCount, InstanceCount, FirstVertex, FirstInstance));
+
+		::vkCmdDraw(CommandBuffer, VertexCount, InstanceCount, FirstVertex, FirstInstance);
+	}
+
+	static FORCEINLINE_DEBUGGABLE void  vkCmdDrawIndexed(VkCommandBuffer CommandBuffer, uint32 IndexCount, uint32 InstanceCount, uint32 FirstIndex, int32_t VertexOffset, uint32 FirstInstance)
+	{
+		CmdPrintfBegin(CommandBuffer, FString::Printf(TEXT("vkCmdDrawIndexed(IndexCount=%d, NumInstances=%d, FirstIndex=%d, VertexOffset=%d, FirstInstance=%d)"), IndexCount, InstanceCount, FirstIndex, VertexOffset, FirstInstance));
 
+		::vkCmdDrawIndexed(CommandBuffer, IndexCount, InstanceCount, FirstIndex, VertexOffset, FirstInstance);
+	}
+
+#if 0
 	static FORCEINLINE_DEBUGGABLE void  vkCmdDrawIndirect(
 		VkCommandBuffer                             commandBuffer,
 		VkBuffer                                    Buffer,
@@ -807,23 +1027,22 @@ namespace VulkanRHI
 		VkCommandBuffer                             commandBuffer,
 		VkBuffer                                    Buffer,
 		VkDeviceSize                                offset);
+#endif
+	static FORCEINLINE_DEBUGGABLE void  vkCmdCopyBuffer(VkCommandBuffer CommandBuffer, VkBuffer SrcBuffer, VkBuffer DstBuffer, uint32 RegionCount, const VkBufferCopy* Regions)
+	{
+		DumpCmdCopyBuffer(CommandBuffer, SrcBuffer, DstBuffer, RegionCount, Regions);
 
-	static FORCEINLINE_DEBUGGABLE void  vkCmdCopyBuffer(
-		VkCommandBuffer                             commandBuffer,
-		VkBuffer                                    srcBuffer,
-		VkBuffer                                    dstBuffer,
-		uint32                                    regionCount,
-		const VkBufferCopy*                         pRegions);
+		::vkCmdCopyBuffer(CommandBuffer, SrcBuffer, DstBuffer, RegionCount, Regions);
+	}
 
-	static FORCEINLINE_DEBUGGABLE void  vkCmdCopyImage(
-		VkCommandBuffer                             commandBuffer,
-		VkImage                                     srcImage,
-		VkImageLayout                               srcImageLayout,
-		VkImage                                     dstImage,
-		VkImageLayout                               dstImageLayout,
-		uint32                                    regionCount,
-		const VkImageCopy*                          pRegions);
+	static FORCEINLINE_DEBUGGABLE void  vkCmdCopyImage(VkCommandBuffer CommandBuffer, VkImage SrcImage, VkImageLayout SrcImageLayout, VkImage DstImage, VkImageLayout DstImageLayout, uint32 RegionCount, const VkImageCopy* Regions)
+	{
+		CmdPrintfBegin(CommandBuffer, FString::Printf(TEXT("vkCmdCopyImage(SrcImage=%p, SrcImageLayout=%d, DstImage=%p, DstImageLayout=%d, RegionCount=%d, Regions=%p)[...]"), SrcImage, (int32)SrcImageLayout, DstImage, (int32)DstImageLayout, RegionCount, Regions));
+
+		::vkCmdCopyImage(CommandBuffer, SrcImage, SrcImageLayout, DstImage, DstImageLayout, RegionCount, Regions);
+	}
 
+#if 0
 	static FORCEINLINE_DEBUGGABLE void  vkCmdBlitImage(
 		VkCommandBuffer                             commandBuffer,
 		VkImage                                     srcImage,
@@ -833,15 +1052,14 @@ namespace VulkanRHI
 		uint32                                    regionCount,
 		const VkImageBlit*                          pRegions,
 		VkFilter                                    filter);
+#endif
+	static FORCEINLINE_DEBUGGABLE void  vkCmdCopyBufferToImage(VkCommandBuffer CommandBuffer, VkBuffer SrcBuffer, VkImage DstImage, VkImageLayout DstImageLayout, uint32 RegionCount, const VkBufferImageCopy* Regions)
+	{
+		DumpCmdCopyBufferToImage(CommandBuffer, SrcBuffer, DstImage, DstImageLayout, RegionCount, Regions);
 
-	static FORCEINLINE_DEBUGGABLE void  vkCmdCopyBufferToImage(
-		VkCommandBuffer                             commandBuffer,
-		VkBuffer                                    srcBuffer,
-		VkImage                                     dstImage,
-		VkImageLayout                               dstImageLayout,
-		uint32                                    regionCount,
-		const VkBufferImageCopy*                    pRegions);
-
+		::vkCmdCopyBufferToImage(CommandBuffer, SrcBuffer, DstImage, DstImageLayout, RegionCount, Regions);
+	}
+#if 0
 	static FORCEINLINE_DEBUGGABLE void  vkCmdCopyImageToBuffer(
 		VkCommandBuffer                             commandBuffer,
 		VkImage                                     srcImage,
@@ -863,30 +1081,28 @@ namespace VulkanRHI
 		VkDeviceSize                                dstOffset,
 		VkDeviceSize                                size,
 		uint32                                    data);
+#endif
+	static FORCEINLINE_DEBUGGABLE void  vkCmdClearColorImage(VkCommandBuffer CommandBuffer, VkImage Image, VkImageLayout ImageLayout, const VkClearColorValue* Color, uint32 RangeCount, const VkImageSubresourceRange* Ranges)
+	{
+		CmdPrintfBegin(CommandBuffer, FString::Printf(TEXT("vkCmdClearColorImage(Image=%p, ImageLayout=%d, Color=%p, RangeCount=%d, Ranges=%p)[...]"), Image, (int32)ImageLayout, Color, RangeCount, Ranges));
 
-	static FORCEINLINE_DEBUGGABLE void  vkCmdClearColorImage(
-		VkCommandBuffer                             commandBuffer,
-		VkImage                                     Image,
-		VkImageLayout                               imageLayout,
-		const VkClearColorValue*                    pColor,
-		uint32                                    rangeCount,
-		const VkImageSubresourceRange*              pRanges);
+		::vkCmdClearColorImage(CommandBuffer, Image, ImageLayout, Color, RangeCount, Ranges);
+	}
 
-	static FORCEINLINE_DEBUGGABLE void  vkCmdClearDepthStencilImage(
-		VkCommandBuffer                             commandBuffer,
-		VkImage                                     Image,
-		VkImageLayout                               imageLayout,
-		const VkClearDepthStencilValue*             pDepthStencil,
-		uint32                                    rangeCount,
-		const VkImageSubresourceRange*              pRanges);
+	static FORCEINLINE_DEBUGGABLE void  vkCmdClearDepthStencilImage(VkCommandBuffer CommandBuffer, VkImage Image, VkImageLayout ImageLayout, const VkClearDepthStencilValue* DepthStencil, uint32 RangeCount, const VkImageSubresourceRange* Ranges)
+	{
+		CmdPrintfBegin(CommandBuffer, FString::Printf(TEXT("vkCmdClearDepthStencilImage(Image=%p, ImageLayout=%d, DepthStencil=%p, RangeCount=%d, Ranges=%p)[...]"), Image, (int32)ImageLayout, DepthStencil, RangeCount, Ranges));
+
+		::vkCmdClearDepthStencilImage(CommandBuffer, Image, ImageLayout, DepthStencil, RangeCount, Ranges);
+	}
+
+	static FORCEINLINE_DEBUGGABLE void  vkCmdClearAttachments(VkCommandBuffer CommandBuffer, uint32 AttachmentCount, const VkClearAttachment* Attachments, uint32 RectCount, const VkClearRect* Rects)
+	{
+		CmdPrintfBegin(CommandBuffer, FString::Printf(TEXT("vkCmdClearAttachments(AttachmentCount=%d, Attachments=%p, RectCount=%d, Rects=%p)[...]"), AttachmentCount, Attachments, RectCount, Rects));
+
+		::vkCmdClearAttachments(CommandBuffer, AttachmentCount, Attachments, RectCount, Rects);
+	}
 
-	static FORCEINLINE_DEBUGGABLE void  vkCmdClearAttachments(
-		VkCommandBuffer                             commandBuffer,
-		uint32                                    attachmentCount,
-		const VkClearAttachment*                    pAttachments,
-		uint32                                    rectCount,
-		const VkClearRect*                          pRects);
-#endif
 	static FORCEINLINE_DEBUGGABLE void  vkCmdResolveImage(
 		VkCommandBuffer CommandBuffer,
 		VkImage SrcImage, VkImageLayout SrcImageLayout,
@@ -942,19 +1158,21 @@ namespace VulkanRHI
 		VkCommandBuffer                             commandBuffer,
 		VkQueryPool                                 queryPool,
 		uint32                                    query);
+#endif
+	static FORCEINLINE_DEBUGGABLE void  vkCmdResetQueryPool(VkCommandBuffer CommandBuffer, VkQueryPool QueryPool, uint32 FirstQuery, uint32 QueryCount)
+	{
+		CmdPrintfBegin(CommandBuffer, FString::Printf(TEXT("vkCmdResetQueryPool(QueryPool=%p, FirstQuery=%d, NumQueries=%d)"), QueryPool, FirstQuery, QueryCount));
 
-	static FORCEINLINE_DEBUGGABLE void  vkCmdResetQueryPool(
-		VkCommandBuffer                             commandBuffer,
-		VkQueryPool                                 queryPool,
-		uint32                                    firstQuery,
-		uint32                                    queryCount);
+		::vkCmdResetQueryPool(CommandBuffer, QueryPool, FirstQuery, QueryCount);
+	}
 
-	static FORCEINLINE_DEBUGGABLE void  vkCmdWriteTimestamp(
-		VkCommandBuffer                             commandBuffer,
-		VkPipelineStageFlagBits                     pipelineStage,
-		VkQueryPool                                 queryPool,
-		uint32                                    query);
+	static FORCEINLINE_DEBUGGABLE void  vkCmdWriteTimestamp(VkCommandBuffer CommandBuffer, VkPipelineStageFlagBits PipelineStage, VkQueryPool QueryPool, uint32 Query)
+	{
+		CmdPrintfBegin(CommandBuffer, FString::Printf(TEXT("vkCmdWriteTimestamp(PipelineStage=%d, QueryPool=%p, Query=%d)"), (int32)PipelineStage, QueryPool, Query));
 
+		::vkCmdWriteTimestamp(CommandBuffer, PipelineStage, QueryPool, Query);
+	}
+#if 0
 	static FORCEINLINE_DEBUGGABLE void  vkCmdCopyQueryPoolResults(
 		VkCommandBuffer                             commandBuffer,
 		VkQueryPool                                 queryPool,
@@ -972,25 +1190,80 @@ namespace VulkanRHI
 		uint32                                    offset,
 		uint32                                    size,
 		const void*                                 pValues);
-
+#endif
 	static FORCEINLINE_DEBUGGABLE void  vkCmdBeginRenderPass(
-		VkCommandBuffer                             commandBuffer,
-		const VkRenderPassBeginInfo*                pRenderPassBegin,
-		VkSubpassContents                           contents);
+		VkCommandBuffer CommandBuffer,
+		const VkRenderPassBeginInfo* RenderPassBegin,
+		VkSubpassContents Contents)
+	{
+		DumpCmdBeginRenderPass(CommandBuffer, RenderPassBegin, Contents);
 
+		::vkCmdBeginRenderPass(CommandBuffer, RenderPassBegin, Contents);
+	}
+#if 0
 	static FORCEINLINE_DEBUGGABLE void  vkCmdNextSubpass(
 		VkCommandBuffer                             commandBuffer,
 		VkSubpassContents                           contents);
+#endif
+	static FORCEINLINE_DEBUGGABLE void  vkCmdEndRenderPass(VkCommandBuffer CommandBuffer)
+	{
+		CmdPrintfBegin(CommandBuffer, TEXT("vkCmdEndRenderPass()"));
 
-	static FORCEINLINE_DEBUGGABLE void  vkCmdEndRenderPass(
-		VkCommandBuffer                             commandBuffer);
-
+		::vkCmdEndRenderPass(CommandBuffer);
+	}
+#if 0
 	static FORCEINLINE_DEBUGGABLE void  vkCmdExecuteCommands(
 		VkCommandBuffer                             commandBuffer,
 		uint32                                    commandBufferCount,
 		const VkCommandBuffer*                      pCommandBuffers);
 #endif
 
+	static FORCEINLINE_DEBUGGABLE VkResult vkCreateSwapchainKHR(VkDevice Device, const VkSwapchainCreateInfoKHR* CreateInfo, const VkAllocationCallbacks* Allocator, VkSwapchainKHR* Swapchain)
+	{
+		DevicePrintfBeginResult(Device, FString::Printf(TEXT("vkCreateSwapchainKHR(SwapChainInfo=%p, OutSwapChain=%p)[...]"), CreateInfo, Swapchain));
+
+		VkResult Result = ::vkCreateSwapchainKHR(Device, CreateInfo, nullptr, Swapchain);
+
+		PrintResultAndNamedHandle(Result, TEXT("SwapChain"), *Swapchain);
+		return Result;
+	}
+
+	static FORCEINLINE_DEBUGGABLE void vkDestroySwapchainKHR(VkDevice Device, VkSwapchainKHR Swapchain, const VkAllocationCallbacks* Allocator)
+	{
+		DevicePrintfBegin(Device, FString::Printf(TEXT("vkDestroySwapchainKHR(SwapChain=%p)[...]"), Swapchain));
+
+		::vkDestroySwapchainKHR(Device, Swapchain, Allocator);
+	}
+
+	static FORCEINLINE_DEBUGGABLE VkResult vkGetSwapchainImagesKHR(VkDevice Device, VkSwapchainKHR Swapchain, uint32_t* SwapchainImageCount, VkImage* SwapchainImages)
+	{
+		DevicePrintfBeginResult(Device, FString::Printf(TEXT("vkGetSwapchainImagesKHR(Swapchain=%p, OutSwapchainImageCount=%p, OutSwapchainImages=%p)"), Swapchain, SwapchainImageCount, SwapchainImages));
+
+		VkResult Result = ::vkGetSwapchainImagesKHR(Device, Swapchain, SwapchainImageCount, SwapchainImages);
+
+		DumpSwapChainImages(Result, SwapchainImageCount, SwapchainImages);
+		return Result;
+	}
+
+	static FORCEINLINE_DEBUGGABLE VkResult vkAcquireNextImageKHR(VkDevice Device, VkSwapchainKHR Swapchain, uint64_t Timeout, VkSemaphore Semaphore, VkFence Fence, uint32_t* ImageIndex)
+	{
+		DevicePrintfBeginResult(Device, FString::Printf(TEXT("vkAcquireNextImageKHR(Swapchain=%p, Timeout=%p, Semaphore=%p, Fence=%p, OutImageIndex=%p)[...]"), Swapchain, (void*)Timeout, Semaphore, Fence, ImageIndex));
+
+		VkResult Result = ::vkAcquireNextImageKHR(Device, Swapchain, Timeout, Semaphore, Fence, ImageIndex);
+
+		PrintResultAndNamedHandle(Result, TEXT("ImageIndex"), ImageIndex);
+		return Result;
+	}
+
+	static FORCEINLINE_DEBUGGABLE VkResult vkQueuePresentKHR(VkQueue Queue, const VkPresentInfoKHR* PresentInfo)
+	{
+		PrintfBeginResult(FString::Printf(TEXT("vkQueuePresentKHR(Queue=%p, Info=%p)[...]"), Queue, PresentInfo));
+
+		VkResult Result = ::vkQueuePresentKHR(Queue, PresentInfo);
+		PrintResult(Result);
+		return Result;
+	}
+
 #if VULKAN_ENABLE_DUMP_LAYER
 	void FlushDebugWrapperLog();
 #endif
diff --git a/Engine/Source/Runtime/VulkanRHI/Private/VulkanCommands.cpp b/Engine/Source/Runtime/VulkanRHI/Private/VulkanCommands.cpp
--- a/Engine/Source/Runtime/VulkanRHI/Private/VulkanCommands.cpp
+++ b/Engine/Source/Runtime/VulkanRHI/Private/VulkanCommands.cpp
@@ -511,7 +511,8 @@ void FVulkanCommandListContext::RHIDrawPrimitive(uint32 PrimitiveType, uint32 Ba
 
 	FVulkanCmdBuffer* CmdBuffer = CommandBufferManager->GetActiveCmdBuffer();
 	PendingState->PrepareDraw(this, CmdBuffer, UEToVulkanType((EPrimitiveType)PrimitiveType));
-	vkCmdDraw(CmdBuffer->GetHandle(), NumVertices, NumInstances, BaseVertexIndex, 0);
+	NumInstances = FMath::Min(1U, NumInstances);
+	VulkanRHI::vkCmdDraw(CmdBuffer->GetHandle(), NumVertices, NumInstances, BaseVertexIndex, 0);
 
 	//if (IsImmediate())
 	{
@@ -552,10 +553,11 @@ void FVulkanCommandListContext::RHIDrawIndexedPrimitive(FIndexBufferRHIParamRef
 	VkCommandBuffer CmdBuffer = Cmd->GetHandle();
 	PendingState->PrepareDraw(this, Cmd, UEToVulkanType((EPrimitiveType)PrimitiveType));
 
-	vkCmdBindIndexBuffer(CmdBuffer, IndexBuffer->GetHandle(), IndexBuffer->GetOffset(), IndexBuffer->GetIndexType());
+	VulkanRHI::vkCmdBindIndexBuffer(CmdBuffer, IndexBuffer->GetHandle(), IndexBuffer->GetOffset(), IndexBuffer->GetIndexType());
 
 	uint32 NumIndices = GetVertexCountForPrimitiveCount(NumPrimitives, PrimitiveType);
-	vkCmdDrawIndexed(CmdBuffer, NumIndices, NumInstances, StartIndex, 0, FirstInstance);
+	NumInstances = FMath::Min(1U, NumInstances);
+	VulkanRHI::vkCmdDrawIndexed(CmdBuffer, NumIndices, NumInstances, StartIndex, BaseVertexIndex, FirstInstance);
 
 	//if (IsImmediate())
 	{
@@ -639,7 +641,7 @@ void FVulkanCommandListContext::RHIEndDrawPrimitiveUP()
 	
 	FVulkanCmdBuffer* CmdBuffer = CommandBufferManager->GetActiveCmdBuffer();
 	PendingState->PrepareDraw(this, CmdBuffer, UEToVulkanType((EPrimitiveType)PendingPrimitiveType));
-	vkCmdDraw(CmdBuffer->GetHandle(), PendingNumVertices, 1, PendingMinVertexIndex, 0);
+	VulkanRHI::vkCmdDraw(CmdBuffer->GetHandle(), PendingNumVertices, 1, PendingMinVertexIndex, 0);
 
 	//if (IsImmediate())
 	{
@@ -683,8 +685,8 @@ void FVulkanCommandListContext::RHIEndDrawIndexedPrimitiveUP()
 	VkCommandBuffer Cmd = CmdBuffer->GetHandle();
 	PendingState->PrepareDraw(this, CmdBuffer, UEToVulkanType((EPrimitiveType)PendingPrimitiveType));
 	uint32 NumIndices = GetVertexCountForPrimitiveCount(PendingNumPrimitives, PendingPrimitiveType);
-	vkCmdBindIndexBuffer(Cmd, PendingDrawPrimUPIndexAllocInfo.GetHandle(), PendingDrawPrimUPIndexAllocInfo.GetBindOffset(), PendingPrimitiveIndexType);
-	vkCmdDrawIndexed(Cmd, NumIndices, 1, PendingMinVertexIndex, 0, 0);
+	VulkanRHI::vkCmdBindIndexBuffer(Cmd, PendingDrawPrimUPIndexAllocInfo.GetHandle(), PendingDrawPrimUPIndexAllocInfo.GetBindOffset(), PendingPrimitiveIndexType);
+	VulkanRHI::vkCmdDrawIndexed(Cmd, NumIndices, 1, PendingMinVertexIndex, 0, 0);
 
 	//if (IsImmediate())
 	{
@@ -694,15 +696,20 @@ void FVulkanCommandListContext::RHIEndDrawIndexedPrimitiveUP()
 
 void FVulkanCommandListContext::RHIClear(bool bClearColor,const FLinearColor& Color, bool bClearDepth,float Depth, bool bClearStencil,uint32 Stencil, FIntRect ExcludeRect)
 {
-	check(Device);
+	if (!(bClearColor || bClearDepth || bClearStencil))
+	{
+		return;
+	}
+
+	FVulkanCmdBuffer* CmdBuffer = CommandBufferManager->GetActiveCmdBuffer();
+	PendingState->UpdateRenderPass(CmdBuffer);
+
 	const uint32 NumColorAttachments = PendingState->GetFrameBuffer()->GetNumColorAttachments();
-	FVulkanCommandListContext::RHIClearMRT(bClearColor, NumColorAttachments, &Color, bClearDepth, Depth, bClearStencil, Stencil, ExcludeRect);
+	FVulkanCommandListContext::InternalClearMRT(CmdBuffer, bClearColor, NumColorAttachments, &Color, bClearDepth, Depth, bClearStencil, Stencil, ExcludeRect);
 }
 
 void FVulkanCommandListContext::RHIClearMRT(bool bClearColor, int32 NumClearColors, const FLinearColor* ClearColorArray, bool bClearDepth, float Depth, bool bClearStencil, uint32 Stencil, FIntRect ExcludeRect)
 {
-	check(Device);
-
 	if (!(bClearColor || bClearDepth || bClearStencil))
 	{
 		return;
@@ -710,8 +717,16 @@ void FVulkanCommandListContext::RHIClearMRT(bool bClearColor, int32 NumClearColo
 
 	check(bClearColor ? NumClearColors > 0 : true);
 
-	PendingState->UpdateRenderPass(CommandBufferManager->GetActiveCmdBuffer());
-	
+	FVulkanCmdBuffer* CmdBuffer = CommandBufferManager->GetActiveCmdBuffer();
+	PendingState->UpdateRenderPass(CmdBuffer);
+
+	const uint32 NumColorAttachments = PendingState->GetFrameBuffer()->GetNumColorAttachments();
+	check((uint32)NumClearColors <= NumColorAttachments);
+	InternalClearMRT(CmdBuffer, bClearColor, NumClearColors, ClearColorArray, bClearDepth, Depth, bClearStencil, Stencil, ExcludeRect);
+}
+
+void FVulkanCommandListContext::InternalClearMRT(FVulkanCmdBuffer* CmdBuffer, bool bClearColor, int32 NumClearColors, const FLinearColor* ClearColorArray, bool bClearDepth, float Depth, bool bClearStencil, uint32 Stencil, FIntRect ExcludeRect)
+{
 #if VULKAN_ALLOW_MIDPASS_CLEAR
 	VkClearRect Rect;
 	FMemory::Memzero(Rect);
@@ -719,33 +734,34 @@ void FVulkanCommandListContext::RHIClearMRT(bool bClearColor, int32 NumClearColo
 	Rect.rect.offset.y = 0;
 	Rect.rect.extent = PendingState->GetRenderPass().GetLayout().GetExtent2D();
 
-	VkClearAttachment Attachement;
-	FMemory::Memzero(Attachement);
+	VkClearAttachment Attachments[MaxSimultaneousRenderTargets + 1];
+	FMemory::Memzero(Attachments);
 
-	if(bClearColor)
+	uint32 NumAttachments = NumClearColors;
+	if (bClearColor)
 	{
 		for (int32 i = 0; i < NumClearColors; ++i)
 		{
-			Attachement.aspectMask = VK_IMAGE_ASPECT_COLOR_BIT;
-			Attachement.colorAttachment = i;
-			Attachement.clearValue.color.float32[0] = ClearColorArray[i].R;
-			Attachement.clearValue.color.float32[1] = ClearColorArray[i].G;
-			Attachement.clearValue.color.float32[2] = ClearColorArray[i].B;
-			Attachement.clearValue.color.float32[3] = ClearColorArray[i].A;
-			vkCmdClearAttachments(PendingState->GetCommandBuffer(), 1, &Attachement, 1, &Rect);
+			Attachments[i].aspectMask = VK_IMAGE_ASPECT_COLOR_BIT;
+			Attachments[i].colorAttachment = i;
+			Attachments[i].clearValue.color.float32[0] = ClearColorArray[i].R;
+			Attachments[i].clearValue.color.float32[1] = ClearColorArray[i].G;
+			Attachments[i].clearValue.color.float32[2] = ClearColorArray[i].B;
+			Attachments[i].clearValue.color.float32[3] = ClearColorArray[i].A;
 		}
 	}
 
 	if(bClearDepth || bClearStencil)
 	{
-		Attachement.aspectMask = 0;
-		Attachement.aspectMask |= bClearDepth ? VK_IMAGE_ASPECT_DEPTH_BIT : 0;
-		Attachement.aspectMask |= bClearStencil ? VK_IMAGE_ASPECT_STENCIL_BIT : 0;
-		Attachement.colorAttachment = NumClearColors;	// Depth attachment is always stored after the color attachment
-		Attachement.clearValue.depthStencil.depth = Depth;
-		Attachement.clearValue.depthStencil.stencil = Stencil;
-		vkCmdClearAttachments(PendingState->GetCommandBuffer(), 1, &Attachement, 1, &Rect);
+		Attachments[NumClearColors].aspectMask = bClearDepth ? VK_IMAGE_ASPECT_DEPTH_BIT : 0;
+		Attachments[NumClearColors].aspectMask |= bClearStencil ? VK_IMAGE_ASPECT_STENCIL_BIT : 0;
+		Attachments[NumClearColors].colorAttachment = 0;
+		Attachments[NumClearColors].clearValue.depthStencil.depth = Depth;
+		Attachments[NumClearColors].clearValue.depthStencil.stencil = Stencil;
+		++NumAttachments;
 	}
+
+	VulkanRHI::vkCmdClearAttachments(CmdBuffer->GetHandle(), NumAttachments, Attachments, 1, &Rect);
 #endif
 }
 
