Copying //UE4/Dev-Rendering to //UE4/Dev-Main (Source: //UE4/Dev-Rendering @ 3045398)

diff --git a/Engine/Source/Runtime/VulkanRHI/Private/VulkanDebug.cpp b/Engine/Source/Runtime/VulkanRHI/Private/VulkanDebug.cpp
--- a/Engine/Source/Runtime/VulkanRHI/Private/VulkanDebug.cpp
+++ b/Engine/Source/Runtime/VulkanRHI/Private/VulkanDebug.cpp
@@ -45,14 +45,17 @@ static VkBool32 VKAPI_PTR DebugReportFunction(
 		static TSet<FString> SeenCodes;
 		if (!SeenCodes.Contains(LayerCode))
 		{
+#if VULKAN_ENABLE_DUMP_LAYER
+			VulkanRHI::FlushDebugWrapperLog();
+#endif
 			FString Message = FString::Printf(TEXT("VK ERROR: [%s] Code %d : %s"), ANSI_TO_TCHAR(LayerPrefix), MsgCode, ANSI_TO_TCHAR(Msg));
 			FPlatformMisc::LowLevelOutputDebugStringf(TEXT("VulkanRHI: %s\n"), *Message);
 			UE_LOG(LogVulkanRHI, Error, TEXT("%s"), *Message);
     
 #if UE_BUILD_DEBUG || UE_BUILD_DEVELOPMENT
-			if (!FCStringAnsi::Strcmp(LayerPrefix, "MEM"))
+			//if (!FCStringAnsi::Strcmp(LayerPrefix, "MEM"))
 			{
-				((FVulkanDynamicRHI*)GDynamicRHI)->DumpMemory();
+				//((FVulkanDynamicRHI*)GDynamicRHI)->DumpMemory();
 			}
 #endif
 
@@ -69,6 +72,9 @@ static VkBool32 VKAPI_PTR DebugReportFunction(
 	}
 	else if (MsgFlags & VK_DEBUG_REPORT_WARNING_BIT_EXT)
 	{
+#if VULKAN_ENABLE_DUMP_LAYER
+		VulkanRHI::FlushDebugWrapperLog();
+#endif
 		FString Message = FString::Printf(TEXT("VK WARNING: [%s] Code %d : %s\n"), ANSI_TO_TCHAR(LayerPrefix), MsgCode, ANSI_TO_TCHAR(Msg));
 		FPlatformMisc::LowLevelOutputDebugStringf(TEXT("VulkanRHI: %s\n"), *Message);
 		UE_LOG(LogVulkanRHI, Warning, TEXT("%s"), *Message);
@@ -109,7 +115,7 @@ static VkBool32 VKAPI_PTR DebugReportFunction(
 void FVulkanDynamicRHI::SetupDebugLayerCallback()
 {
 #if !VULKAN_DISABLE_DEBUG_CALLBACK
-	PFN_vkCreateDebugReportCallbackEXT CreateMsgCallback = (PFN_vkCreateDebugReportCallbackEXT)(void*)vkGetInstanceProcAddr(Instance, CREATE_MSG_CALLBACK);
+	PFN_vkCreateDebugReportCallbackEXT CreateMsgCallback = (PFN_vkCreateDebugReportCallbackEXT)(void*)VulkanRHI::vkGetInstanceProcAddr(Instance, CREATE_MSG_CALLBACK);
 	if (CreateMsgCallback)
 	{
 		VkDebugReportCallbackCreateInfoEXT CreateInfo;
@@ -148,7 +154,7 @@ void FVulkanDynamicRHI::RemoveDebugLayerCallback()
 #if !VULKAN_DISABLE_DEBUG_CALLBACK
 	if (MsgCallback != VK_NULL_HANDLE)
 	{
-		PFN_vkDestroyDebugReportCallbackEXT DestroyMsgCallback = (PFN_vkDestroyDebugReportCallbackEXT)(void*)vkGetInstanceProcAddr(Instance, DESTROY_MSG_CALLBACK);
+		PFN_vkDestroyDebugReportCallbackEXT DestroyMsgCallback = (PFN_vkDestroyDebugReportCallbackEXT)(void*)VulkanRHI::vkGetInstanceProcAddr(Instance, DESTROY_MSG_CALLBACK);
 		checkf(DestroyMsgCallback, TEXT("GetProcAddr: Unable to find vkDbgCreateMsgCallback\vkGetInstanceProcAddr Failure"));
 		DestroyMsgCallback(Instance, MsgCallback, nullptr);
 	}
@@ -162,7 +168,7 @@ namespace VulkanRHI
 	static FString DebugLog;
 	static int32 DebugLine = 1;
 
-	static const TCHAR* TenTabs = TEXT("\t\t\t\t\t\t\t\t\t\t");
+	static const TCHAR* Tabs = TEXT("\t\t\t\t\t\t\t\t\t");
 
 	void FlushDebugWrapperLog()
 	{
@@ -175,6 +181,23 @@ namespace VulkanRHI
 		}
 	}
 
+	static void HandleFlushWrapperLog(const TArray<FString>& Args)
+	{
+		FlushDebugWrapperLog();
+	}
+
+	static FAutoConsoleCommand CVarVulkanFlushLog(
+		TEXT("r.Vulkan.FlushLog"),
+		TEXT("\n"),
+		FConsoleCommandWithArgsDelegate::CreateStatic(&HandleFlushWrapperLog)
+		);
+
+	static TAutoConsoleVariable<int32> CVarVulkanDumpLayer(
+		TEXT("r.Vulkan.DumpLayer"),
+		0,
+		TEXT("1 to enable dump layer, 0 to disable (default)")
+		);
+
 	static FString GetVkResultErrorString(VkResult Result)
 	{
 		FString ErrorString;
@@ -239,6 +262,106 @@ namespace VulkanRHI
 		return String;
 	}
 
+	static FString GetComponentMappingString(const VkComponentMapping& Mapping)
+	{
+		auto GetSwizzle = [](VkComponentSwizzle Swizzle) -> const TCHAR*
+			{
+				switch (Swizzle)
+				{
+				case VK_COMPONENT_SWIZZLE_IDENTITY: return TEXT("ID");
+				case VK_COMPONENT_SWIZZLE_ZERO:		return TEXT("0");
+				case VK_COMPONENT_SWIZZLE_ONE:		return TEXT("1");
+				case VK_COMPONENT_SWIZZLE_R:		return TEXT("R");
+				case VK_COMPONENT_SWIZZLE_G:		return TEXT("G");
+				case VK_COMPONENT_SWIZZLE_B:		return TEXT("B");
+				case VK_COMPONENT_SWIZZLE_A:		return TEXT("A");
+				default:
+					check(0);
+					return TEXT("-");
+				}
+			};
+		return FString::Printf(TEXT("(r=%s, g=%s, b=%s, a=%s)"), GetSwizzle(Mapping.r), GetSwizzle(Mapping.g), GetSwizzle(Mapping.b), GetSwizzle(Mapping.a));
+	}
+
+
+#define AppendBitFieldName(BitField, Name) \
+	if ((Flags & BitField) == BitField)\
+	{\
+		Flags &= ~BitField;\
+		String += Name;\
+		if (bFirst)\
+		{\
+			bFirst = false;\
+		}\
+		else\
+		{\
+			if (Flags !=0)\
+			{\
+				String += TEXT("|");\
+			}\
+		}\
+	}
+
+	static FString GetAspectMaskString(VkImageAspectFlags Flags)
+	{
+		if (Flags == 0)
+		{
+			return TEXT("0");
+		}
+		bool bFirst = true;
+		FString String;
+		AppendBitFieldName(VK_IMAGE_ASPECT_COLOR_BIT, TEXT("COLOR"));
+		AppendBitFieldName(VK_IMAGE_ASPECT_DEPTH_BIT, TEXT("DEPTH"));
+		AppendBitFieldName(VK_IMAGE_ASPECT_STENCIL_BIT, TEXT("STENCIL"));
+		AppendBitFieldName(VK_IMAGE_ASPECT_METADATA_BIT, TEXT("METADATA"));
+		if (Flags != 0)
+		{
+			FString Unknown = FString::Printf(TEXT("%d"), Flags);
+			AppendBitFieldName(Flags, Unknown);
+		}
+		return String;
+	}
+
+	static FString GetAccessFlagString(VkAccessFlags Flags)
+	{
+		if (Flags == 0)
+		{
+			return TEXT("0");
+		}
+		bool bFirst = true;
+		FString String;
+		AppendBitFieldName(VK_ACCESS_INDIRECT_COMMAND_READ_BIT, TEXT("INDIRECT_COMMAND"));
+		AppendBitFieldName(VK_ACCESS_INDEX_READ_BIT, TEXT("INDEX_READ"));
+		AppendBitFieldName(VK_ACCESS_VERTEX_ATTRIBUTE_READ_BIT, TEXT("VERTEX_ATTR_READ"));
+		AppendBitFieldName(VK_ACCESS_UNIFORM_READ_BIT, TEXT("UNIF_READ"));
+		AppendBitFieldName(VK_ACCESS_INPUT_ATTACHMENT_READ_BIT, TEXT("INPUT_ATT_READ"));
+		AppendBitFieldName(VK_ACCESS_SHADER_READ_BIT, TEXT("SHADER_READ"));
+		AppendBitFieldName(VK_ACCESS_SHADER_WRITE_BIT, TEXT("SHADER_WRITE"));
+		AppendBitFieldName(VK_ACCESS_COLOR_ATTACHMENT_READ_BIT, TEXT("COLOR_ATT_READ"));
+		AppendBitFieldName(VK_ACCESS_COLOR_ATTACHMENT_WRITE_BIT, TEXT("COLOR_ATT_WRITE"));
+		AppendBitFieldName(VK_ACCESS_DEPTH_STENCIL_ATTACHMENT_READ_BIT, TEXT("DS_ATT_READ"));
+		AppendBitFieldName(VK_ACCESS_DEPTH_STENCIL_ATTACHMENT_WRITE_BIT, TEXT("DS_ATT_WIRTE"));
+		AppendBitFieldName(VK_ACCESS_TRANSFER_READ_BIT, TEXT("TRANSFER_READ"));
+		AppendBitFieldName(VK_ACCESS_TRANSFER_WRITE_BIT, TEXT("TRANSFER_WRITE"));
+		AppendBitFieldName(VK_ACCESS_HOST_READ_BIT, TEXT("HOST_READ"));
+		AppendBitFieldName(VK_ACCESS_HOST_WRITE_BIT, TEXT("HOST_WRITE"));
+		AppendBitFieldName(VK_ACCESS_MEMORY_READ_BIT, TEXT("MEM_READ"));
+		AppendBitFieldName(VK_ACCESS_MEMORY_WRITE_BIT, TEXT("MEM_WRITE"));
+		if (Flags != 0)
+		{
+			FString Unknown = FString::Printf(TEXT("%d"), Flags);
+			AppendBitFieldName(Flags, Unknown);
+		}
+		return String;
+	}
+
+#undef  AppendBitFieldName
+
+	static FString GetSubResourceRangeString(const VkImageSubresourceRange& Range)
+	{
+		return FString::Printf(TEXT("aspectMask=%s, baseMipLevel=%d, levelCount=%d, baseArrayLayer=%d, layerCount=%d"), *GetAspectMaskString(Range.aspectMask), Range.baseMipLevel, Range.levelCount, Range.baseArrayLayer, Range.layerCount);		
+	}
+
 	static FString GetStageMaskString(VkPipelineStageFlags Flags)
 	{
 		return FString::Printf(TEXT("VkPipelineStageFlags=0x%x"), (uint32)Flags);
@@ -246,59 +369,117 @@ namespace VulkanRHI
 
 	void PrintfBeginResult(const FString& String)
 	{
-		DebugLog += FString::Printf(TEXT("[GLOBAL METHOD]   %8d: %s"), DebugLine++, *String);
+		if (CVarVulkanDumpLayer.GetValueOnAnyThread())
+		{
+			DebugLog += FString::Printf(TEXT("[GLOBAL METHOD]     %8d: %s"), DebugLine++, *String);
+		}
 	}
 
 	void PrintfBegin(const FString& String)
 	{
-		DebugLog += FString::Printf(TEXT("[GLOBAL METHOD]   %8d: %s\n"), DebugLine++, *String);
+		if (CVarVulkanDumpLayer.GetValueOnAnyThread())
+		{
+			DebugLog += FString::Printf(TEXT("[GLOBAL METHOD]     %8d: %s\n"), DebugLine++, *String);
+		}
 	}
 
 	void DevicePrintfBeginResult(VkDevice Device, const FString& String)
 	{
-		DebugLog += FString::Printf(TEXT("[%p]%8d: %s"), Device, DebugLine++, *String);
+		if (CVarVulkanDumpLayer.GetValueOnAnyThread())
+		{
+			DebugLog += FString::Printf(TEXT("[D:%p]%8d: %s"), Device, DebugLine++, *String);
+		}
 	}
 
 	void DevicePrintfBegin(VkDevice Device, const FString& String)
 	{
-		DebugLog += FString::Printf(TEXT("[%p]%8d: %s\n"), Device, DebugLine++, *String);
+		if (CVarVulkanDumpLayer.GetValueOnAnyThread())
+		{
+			DebugLog += FString::Printf(TEXT("[D:%p]%8d: %s\n"), Device, DebugLine++, *String);
+		}
+	}
+
+	void CmdPrintfBegin(VkCommandBuffer CmdBuffer, const FString& String)
+	{
+		if (CVarVulkanDumpLayer.GetValueOnAnyThread())
+		{
+			DebugLog += FString::Printf(TEXT("[C:%p]%8d: %s\n"), CmdBuffer, DebugLine++, *String);
+		}
+	}
+
+	void CmdPrintfBeginResult(VkCommandBuffer CmdBuffer, const FString& String)
+	{
+		if (CVarVulkanDumpLayer.GetValueOnAnyThread())
+		{
+			DebugLog += FString::Printf(TEXT("[C:%p]%8d: %s"), CmdBuffer, DebugLine++, *String);
+		}
 	}
 
 	void PrintResult(VkResult Result)
 	{
-		DebugLog += FString::Printf(TEXT(" -> %s\n"), *GetVkResultErrorString(Result));
+		if (CVarVulkanDumpLayer.GetValueOnAnyThread())
+		{
+			DebugLog += FString::Printf(TEXT(" -> %s\n"), *GetVkResultErrorString(Result));
+			if (Result < VK_SUCCESS)
+			{
+				FlushDebugWrapperLog();
+			}
+		}
 	}
 
 	void PrintResultAndPointer(VkResult Result, void* Handle)
 	{
-		DebugLog += FString::Printf(TEXT(" -> %s => %p\n"), *GetVkResultErrorString(Result), Handle);
+		if (CVarVulkanDumpLayer.GetValueOnAnyThread())
+		{
+			DebugLog += FString::Printf(TEXT(" -> %s => %p\n"), *GetVkResultErrorString(Result), Handle);
+			if (Result < VK_SUCCESS)
+			{
+				FlushDebugWrapperLog();
+			}
+		}
 	}
 
 	void PrintResultAndNamedHandle(VkResult Result, const TCHAR* HandleName, void* Handle)
 	{
-		DebugLog += FString::Printf(TEXT(" -> %s => %s=%p\n"), *GetVkResultErrorString(Result), HandleName, Handle);
+		if (CVarVulkanDumpLayer.GetValueOnAnyThread())
+		{
+			DebugLog += FString::Printf(TEXT(" -> %s => %s=%p\n"), *GetVkResultErrorString(Result), HandleName, Handle);
+			if (Result < VK_SUCCESS)
+			{
+				FlushDebugWrapperLog();
+			}
+		}
 	}
 
 	void DumpPhysicalDeviceProperties(VkPhysicalDeviceMemoryProperties* Properties)
 	{
-		DebugLog += FString::Printf(TEXT(" -> VkPhysicalDeviceMemoryProperties[...]\n"));
+		if (CVarVulkanDumpLayer.GetValueOnAnyThread())
+		{
+			DebugLog += FString::Printf(TEXT(" -> VkPhysicalDeviceMemoryProperties[...]\n"));
+		}
 	}
 
 	void DumpAllocateMemory(VkDevice Device, const VkMemoryAllocateInfo* AllocateInfo, VkDeviceMemory* Memory)
 	{
-		DevicePrintfBeginResult(Device, FString::Printf(TEXT("vkAllocateMemory(OutMem=%p)\n"), AllocateInfo, Memory));
-		DebugLog += TenTabs[3];
-		DebugLog += FString::Printf(TEXT("VkMemoryAllocateInfo %p: Size=%d, MemoryTypeIndex=%d"), AllocateInfo, (uint32)AllocateInfo->allocationSize, AllocateInfo->memoryTypeIndex);
+		if (CVarVulkanDumpLayer.GetValueOnAnyThread())
+		{
+			DevicePrintfBeginResult(Device, FString::Printf(TEXT("vkAllocateMemory(OutMem=%p)\n"), AllocateInfo, Memory));
+			DebugLog += Tabs;
+			DebugLog += FString::Printf(TEXT("VkMemoryAllocateInfo %p: Size=%d, MemoryTypeIndex=%d"), AllocateInfo, (uint32)AllocateInfo->allocationSize, AllocateInfo->memoryTypeIndex);
+		}
 	}
 
 	void DumpMemoryRequirements(VkMemoryRequirements* MemoryRequirements)
 	{
-		DebugLog += TenTabs[3];
-		DebugLog += FString::Printf(TEXT("VkMemoryRequirements: Size=%d Align=%d MemoryTypeBits=0x%x\n"), (uint32)MemoryRequirements->size, (uint32)MemoryRequirements->alignment, MemoryRequirements->memoryTypeBits);
+		if (CVarVulkanDumpLayer.GetValueOnAnyThread())
+		{
+			DebugLog += FString::Printf(TEXT(" -> Size=%d Align=%d MemoryTypeBits=0x%x\n"), (uint32)MemoryRequirements->size, (uint32)MemoryRequirements->alignment, MemoryRequirements->memoryTypeBits);
+		}
 	}
 
 	void DumpBufferCreate(VkDevice Device, const VkBufferCreateInfo* CreateInfo, VkBuffer* Buffer)
 	{
+		FlushDebugWrapperLog();
 		DevicePrintfBeginResult(Device, FString::Printf(TEXT("vkCreateBuffer(Info=%p, OutBuffer=%p)[...]"), CreateInfo, Buffer));
 	}
 
@@ -307,25 +488,50 @@ namespace VulkanRHI
 		DevicePrintfBeginResult(Device, FString::Printf(TEXT("VkBufferViewCreateInfo(Info=%p, OutBufferView=%p)[...]"), CreateInfo, BufferView));
 	}
 
-	void DumpImageCreate(VkDevice Device, const VkImageCreateInfo* CreateInfo)
+	void DumpImageCreate(VkDevice Device, const VkImageCreateInfo* CreateInfo, VkImage* Image)
 	{
-		ensure(0);
+		FlushDebugWrapperLog();
+		DevicePrintfBeginResult(Device, FString::Printf(TEXT("vkCreateImage(Info=%p, OutImage=%p)[...]"), CreateInfo, Image));
+	}
+
+	void DumpImageViewCreate(VkDevice Device, const VkImageViewCreateInfo* CreateInfo, VkImageView* ImageView)
+	{
+		DevicePrintfBegin(Device, FString::Printf(TEXT("VkImageViewCreateInfo(Info=%p, OutImageView=%p)"), CreateInfo, ImageView));
+
+		if (CVarVulkanDumpLayer.GetValueOnAnyThread())
+		{
+			DebugLog += FString::Printf(TEXT("%sVkImageViewCreateInfo: Flags=%d, Image=%p, ViewType=%d, Format=%p, Components=%s\n"), Tabs, CreateInfo->flags, CreateInfo->image, (int32)CreateInfo->viewType, (void*)CreateInfo->format, *GetComponentMappingString(CreateInfo->components));
+			DebugLog += FString::Printf(TEXT("%s\tSubresourceRange=(%s)"), Tabs, *GetSubResourceRangeString(CreateInfo->subresourceRange));
+		}
 	}
 
 	void DumpFenceCreate(VkDevice Device, const VkFenceCreateInfo* CreateInfo, VkFence* Fence)
 	{
 		DevicePrintfBeginResult(Device, FString::Printf(TEXT("vkCreateFence(CreateInfo=%p, OutFence=%p)[...]"), CreateInfo, Fence));
 	}
-	
+
 	void DumpFenceList(uint32 FenceCount, const VkFence* Fences)
 	{
-		for (uint32 Index = 0; Index < 0; ++Index)
+		if (CVarVulkanDumpLayer.GetValueOnAnyThread())
 		{
-			DebugLog += TenTabs[2];
-			DebugLog += FString::Printf(TEXT("Fence[%d]=%p\n"), Index, Fences[Index]);
+			for (uint32 Index = 0; Index < FenceCount; ++Index)
+			{
+				DebugLog += Tabs;
+				DebugLog += '\t';
+				DebugLog += FString::Printf(TEXT("Fence[%d]=%p"), Index, Fences[Index]);
+				if (Index < FenceCount - 1)
+				{
+					DebugLog += TEXT("\n");
+				}
+			}
 		}
 	}
 	
+	void DumpSemaphoreCreate(VkDevice Device, const VkSemaphoreCreateInfo* CreateInfo, VkSemaphore* Semaphore)
+	{
+		DevicePrintfBeginResult(Device, FString::Printf(TEXT("vkCreateSemaphore(CreateInfo=%p, OutSemaphore=%p)[...]"), CreateInfo, Semaphore));
+	}
+
 	void DumpMappedMemoryRanges(uint32 memoryRangeCount, const VkMappedMemoryRange* MemoryRanges)
 	{
 		ensure(0);
@@ -333,31 +539,37 @@ namespace VulkanRHI
 
 	void DumpResolveImage(VkCommandBuffer CommandBuffer, VkImage SrcImage, VkImageLayout SrcImageLayout, VkImage DstImage, VkImageLayout DstImageLayout, uint32 RegionCount, const VkImageResolve* Regions)
 	{
-		PrintfBeginResult(FString::Printf(TEXT("vkCmdResolveImage(CmdBuffer=%p, SrcImage=%p, SrcImageLayout=%s, DestImage=%p, DestImageLayout=%s, NumRegions=%d, Regions=%p)[...]"),
-			CommandBuffer, SrcImage, *GetImageLayoutString(SrcImageLayout), DstImage, *GetImageLayoutString(DstImageLayout), RegionCount, Regions));
-		for (uint32 Index = 0; Index < RegionCount; ++Index)
+		if (CVarVulkanDumpLayer.GetValueOnAnyThread())
 		{
-			DebugLog += TenTabs[3];
-			DebugLog += FString::Printf(TEXT("Region %d: "), Index);
-/*
-			typedef struct VkImageResolve {
-				VkImageSubresourceLayers    srcSubresource;
-				VkOffset3D                  srcOffset;
-				VkImageSubresourceLayers    dstSubresource;
-				VkOffset3D                  dstOffset;
-				VkExtent3D                  extent;
-
-*/
+			CmdPrintfBegin(CommandBuffer, FString::Printf(TEXT("vkCmdResolveImage(SrcImage=%p, SrcImageLayout=%s, DestImage=%p, DestImageLayout=%s, NumRegions=%d, Regions=%p)[...]"),
+				CommandBuffer, SrcImage, *GetImageLayoutString(SrcImageLayout), DstImage, *GetImageLayoutString(DstImageLayout), RegionCount, Regions));
+			for (uint32 Index = 0; Index < RegionCount; ++Index)
+			{
+				DebugLog += Tabs;
+				DebugLog += FString::Printf(TEXT("Region %d: "), Index);
+				/*
+							typedef struct VkImageResolve {
+								VkImageSubresourceLayers    srcSubresource;
+								VkOffset3D                  srcOffset;
+								VkImageSubresourceLayers    dstSubresource;
+								VkOffset3D                  dstOffset;
+								VkExtent3D                  extent;
+
+				*/
+			}
 		}
 	}
 
 	void DumpFreeDescriptorSets(VkDevice Device, VkDescriptorPool DescriptorPool, uint32 DescriptorSetCount, const VkDescriptorSet* DescriptorSets)
 	{
-		DevicePrintfBegin(Device, FString::Printf(TEXT("vkFreeDescriptorSets(Pool=%p, Count=%d, Sets=%p)"), DescriptorPool, DescriptorSetCount, DescriptorSets));
-		for (uint32 Index = 0; Index < DescriptorSetCount; ++Index)
+		if (CVarVulkanDumpLayer.GetValueOnAnyThread())
 		{
-			DebugLog += TenTabs[3];
-			DebugLog += FString::Printf(TEXT("Set %d: %p\n"), Index, DescriptorSets[Index]);
+			DevicePrintfBegin(Device, FString::Printf(TEXT("vkFreeDescriptorSets(Pool=%p, NumSets=%d, Sets=%p)"), DescriptorPool, DescriptorSetCount, DescriptorSets));
+			for (uint32 Index = 0; Index < DescriptorSetCount; ++Index)
+			{
+				DebugLog += Tabs;
+				DebugLog += FString::Printf(TEXT("Set %d: %p\n"), Index, DescriptorSets[Index]);
+			}
 		}
 	}
 
@@ -368,16 +580,19 @@ namespace VulkanRHI
 
 	void DumpEnumeratePhysicalDevicesEpilog(uint32* PhysicalDeviceCount, VkPhysicalDevice* PhysicalDevices)
 	{
-		if (PhysicalDeviceCount)
+		if (CVarVulkanDumpLayer.GetValueOnAnyThread())
 		{
-			DebugLog += TenTabs[3];
-			DebugLog += FString::Printf(TEXT("OutCount=%d\n"), *PhysicalDeviceCount);
-			if (PhysicalDevices)
+			if (PhysicalDeviceCount)
 			{
-				for (uint32 Index = 0; Index < *PhysicalDeviceCount; ++Index)
+				DebugLog += Tabs;
+				DebugLog += FString::Printf(TEXT("OutCount=%d\n"), *PhysicalDeviceCount);
+				if (PhysicalDevices)
 				{
-					DebugLog += TenTabs[2];
-					DebugLog += FString::Printf(TEXT("OutDevice[%d]=%p\n"), Index, PhysicalDevices[Index]);
+					for (uint32 Index = 0; Index < *PhysicalDeviceCount; ++Index)
+					{
+						DebugLog += Tabs;
+						DebugLog += FString::Printf(TEXT("\tOutDevice[%d]=%p\n"), Index, PhysicalDevices[Index]);
+					}
 				}
 			}
 		}
@@ -385,7 +600,36 @@ namespace VulkanRHI
 
 	void DumpCmdPipelineBarrier(VkCommandBuffer CommandBuffer, VkPipelineStageFlags SrcStageMask, VkPipelineStageFlags DstStageMask, VkDependencyFlags DependencyFlags, uint32 MemoryBarrierCount, const VkMemoryBarrier* MemoryBarriers, uint32 BufferMemoryBarrierCount, const VkBufferMemoryBarrier* BufferMemoryBarriers, uint32 ImageMemoryBarrierCount, const VkImageMemoryBarrier* ImageMemoryBarriers)
 	{
-		PrintfBegin(FString::Printf(TEXT("vkCmdPipelineBarrier(CmdBuffer=%p, SrcMask=%s, DestMask=%s, Flags=%d, MemCount=%d, Mem=%p, BufferCount=%d, Buffers=%p, ImageCount=%d, Images=%p)[...]"), CommandBuffer, *GetStageMaskString(SrcStageMask), *GetStageMaskString(DstStageMask), (uint32)DependencyFlags, MemoryBarrierCount, MemoryBarriers, BufferMemoryBarrierCount, BufferMemoryBarriers, ImageMemoryBarrierCount, ImageMemoryBarriers));
+		CmdPrintfBegin(CommandBuffer, FString::Printf(TEXT("vkCmdPipelineBarrier(SrcMask=%s, DestMask=%s, Flags=%d, NumMemB=%d, MemB=%p,"), *GetStageMaskString(SrcStageMask), *GetStageMaskString(DstStageMask), (uint32)DependencyFlags, MemoryBarrierCount, MemoryBarriers));
+		if (CVarVulkanDumpLayer.GetValueOnAnyThread())
+		{
+			DebugLog += FString::Printf(TEXT("%s\tNumBufferB=%d, BufferB=%p, NumImageB=%d, ImageB=%p)[...]\n"), Tabs, BufferMemoryBarrierCount, BufferMemoryBarriers, ImageMemoryBarrierCount, ImageMemoryBarriers);
+			if (ImageMemoryBarrierCount)
+			{
+				for (uint32 Index = 0; Index < ImageMemoryBarrierCount; ++Index)
+				{
+					DebugLog += FString::Printf(TEXT("%s\tImageBarrier[%d]: srcAccess=%s, oldLayout=%s, srcQueueFamilyIndex=%d\n"), Tabs, Index, *GetAccessFlagString(ImageMemoryBarriers[Index].srcAccessMask), *GetImageLayoutString(ImageMemoryBarriers[Index].oldLayout), ImageMemoryBarriers[Index].srcQueueFamilyIndex);
+					DebugLog += FString::Printf(TEXT("%s\t\tdstAccess=%s, newLayout=%s, dstQueueFamilyIndex=%d\n"), Tabs, *GetAccessFlagString(ImageMemoryBarriers[Index].dstAccessMask), *GetImageLayoutString(ImageMemoryBarriers[Index].newLayout), ImageMemoryBarriers[Index].dstQueueFamilyIndex);
+					DebugLog += FString::Printf(TEXT("%s\t\tImage=%p, subresourceRange=(%s)\n"), Tabs, ImageMemoryBarriers[Index].image, *GetSubResourceRangeString(ImageMemoryBarriers[Index].subresourceRange));
+				}
+			}
+		}
+	}
+
+	void DumpBindDescriptorSets(VkCommandBuffer CommandBuffer, VkPipelineBindPoint PipelineBindPoint, VkPipelineLayout Layout, uint32 FirstSet, uint32 DescriptorSetCount, const VkDescriptorSet* DescriptorSets, uint32 DynamicOffsetCount, const uint32* DynamicOffsets)
+	{
+		CmdPrintfBegin(CommandBuffer, FString::Printf(TEXT("vkCmdBindDescriptorSets(PipelineBindPoint=%d, Layout=%p, FirstSet=%d, NumDS=%d, DS=%p, NumDynamicOffset=%d, DynamicOffsets=%p)[...]"), (int32)PipelineBindPoint, Layout, FirstSet, DescriptorSetCount, DescriptorSets, DynamicOffsetCount, DynamicOffsets));
+		if (CVarVulkanDumpLayer.GetValueOnAnyThread())
+		{
+			for (uint32 Index = 0; Index < DescriptorSetCount; ++Index)
+			{
+				DebugLog += FString::Printf(TEXT("%s\tDS[%d]=%p\n"), Tabs, Index, DescriptorSets[Index]);
+			}
+			for (uint32 Index = 0; Index < DynamicOffsetCount; ++Index)
+			{
+				DebugLog += FString::Printf(TEXT("%s\tDynamicOffset[%d]=%p\n"), Tabs, Index, DynamicOffsets[Index]);
+			}
+		}
 	}
 
 	void DumpCreateDescriptorSetLayout(VkDevice Device, const VkDescriptorSetLayoutCreateInfo* CreateInfo, VkDescriptorSetLayout* SetLayout)
@@ -398,17 +642,66 @@ namespace VulkanRHI
 		DevicePrintfBegin(Device, FString::Printf(TEXT("vkAllocateDescriptorSets(Info=%p, OutSets=%p)[...]"), AllocateInfo, DescriptorSets));
 	}
 
+	void DumpUpdateDescriptorSets(VkDevice Device, uint32 DescriptorWriteCount, const VkWriteDescriptorSet* DescriptorWrites, uint32 DescriptorCopyCount, const VkCopyDescriptorSet* DescriptorCopies)
+	{
+		DevicePrintfBegin(Device, FString::Printf(TEXT("vkUpdateDescriptorSets(NumWrites=%d, Writes=%p, NumCopies=%d, Copies=%p)[...]"), DescriptorWriteCount, DescriptorWrites, DescriptorCopyCount, DescriptorCopies));
+	}
+
+	void DumpCreateFramebuffer(VkDevice Device, const VkFramebufferCreateInfo* CreateInfo, VkFramebuffer* Framebuffer)
+	{
+		if (CVarVulkanDumpLayer.GetValueOnAnyThread())
+		{
+			DevicePrintfBegin(Device, FString::Printf(TEXT("vkCreateFramebuffer(Info=%p, OutFramebuffer=%p)"), CreateInfo, Framebuffer));
+			DebugLog += FString::Printf(TEXT("%sVkFramebufferCreateInfo: Flags=%d, RenderPass=%p, NumAttachments=%d\n"), Tabs, CreateInfo->flags, CreateInfo->renderPass, CreateInfo->attachmentCount);
+			for (uint32 Index = 0; Index < CreateInfo->attachmentCount; ++Index)
+			{
+				DebugLog += FString::Printf(TEXT("%s\tAttachment[%d]: ImageView=%p\n"), Tabs, Index, CreateInfo->pAttachments[Index]);
+			}
+			DebugLog += FString::Printf(TEXT("%s\twidth=%d, height=%d, layers=%d"), Tabs, CreateInfo->width, CreateInfo->height, CreateInfo->layers);
+		}
+	}
+
 	void DumpCreateRenderPass(VkDevice Device, const VkRenderPassCreateInfo* CreateInfo, VkRenderPass* RenderPass)
 	{
-		DevicePrintfBegin(Device, FString::Printf(TEXT("vkCreateRenderPass(Info=%p, OutRenderPass=%p)[...]"), CreateInfo, RenderPass));
+		DevicePrintfBeginResult(Device, FString::Printf(TEXT("vkCreateRenderPass(Info=%p, OutRenderPass=%p)[...]"), CreateInfo, RenderPass));
 	}
 
 	void DumpQueueSubmit(VkQueue Queue, uint32 SubmitCount, const VkSubmitInfo* Submits, VkFence Fence)
 	{
-		// Probably good point to submit
-		FlushDebugWrapperLog();
+		if (CVarVulkanDumpLayer.GetValueOnAnyThread())
+		{
+			PrintfBeginResult(FString::Printf(TEXT("vkQueueSubmit(Queue=%p, Count=%d, Submits=%p, Fence=%p)"), Queue, SubmitCount, Submits, Fence));
+			for (uint32 Index = 0; Index < SubmitCount; ++Index)
+			{
+				DebugLog += FString::Printf(TEXT("\n%sSubmit[%d]:"), Tabs, Index);
+				if (Submits[Index].waitSemaphoreCount > 0)
+				{
+					DebugLog += FString::Printf(TEXT("\n%s\tWaitSemaphores(Mask): "), Tabs, Index);
+					for (uint32 SubIndex = 0; SubIndex < Submits[Index].waitSemaphoreCount; ++SubIndex)
+					{
+						DebugLog += FString::Printf(TEXT("%p(%d) "), Submits[Index].pWaitSemaphores[SubIndex], (int32)Submits[Index].pWaitDstStageMask[SubIndex]);
+					}
+				}
+				if (Submits[Index].commandBufferCount > 0)
+				{
+					DebugLog += FString::Printf(TEXT("\n%s\tCommandBuffers: "), Tabs, Index);
+					for (uint32 SubIndex = 0; SubIndex < Submits[Index].commandBufferCount; ++SubIndex)
+					{
+						DebugLog += FString::Printf(TEXT("%p "), Submits[Index].pCommandBuffers[SubIndex]);
+					}
+				}
+				if (Submits[Index].signalSemaphoreCount > 0)
+				{
+					DebugLog += FString::Printf(TEXT("\n%s\tSignalSemaphore: "), Tabs, Index);
+					for (uint32 SubIndex = 0; SubIndex < Submits[Index].signalSemaphoreCount; ++SubIndex)
+					{
+						DebugLog += FString::Printf(TEXT("%p "), Submits[Index].pSignalSemaphores[SubIndex]);
+					}
+				}
+			}
 
-		PrintfBeginResult(FString::Printf(TEXT("vkQueueSubmit(Queue=%p, Count=%d, Submits=%p, Fence=%p)[...]"), Queue, SubmitCount, Submits, Fence));
+			FlushDebugWrapperLog();
+		}
 	}
 
 	void DumpCreateShaderModule(VkDevice Device, const VkShaderModuleCreateInfo* CreateInfo, VkShaderModule* ShaderModule)
@@ -416,16 +709,139 @@ namespace VulkanRHI
 		DevicePrintfBeginResult(Device, FString::Printf(TEXT("vkCreateShaderModule(CreateInfo=%p, OutShaderModule=%p)[...]"), CreateInfo, ShaderModule));
 	}
 
-	static void HandleFlushWrapperLog(const TArray<FString>& Args)
+	void DumpCreatePipelineCache(VkDevice Device, const VkPipelineCacheCreateInfo* CreateInfo, VkPipelineCache* PipelineCache)
+	{
+		DevicePrintfBeginResult(Device, FString::Printf(TEXT("vkCreatePipelineCache(CreateInfo=%p, OutPipelineCache=%p)[...]"), CreateInfo, PipelineCache));
+	}
+
+	void DumpCreateCommandPool(VkDevice Device, const VkCommandPoolCreateInfo* CreateInfo, VkCommandPool* CommandPool)
+	{
+		DevicePrintfBeginResult(Device, FString::Printf(TEXT("vkCreateCommandPool(CreateInfo=%p, OutCommandPool=%p)[...]"), CreateInfo, CommandPool));
+	}
+
+	void DumpCreateQueryPool(VkDevice Device, const VkQueryPoolCreateInfo* CreateInfo, VkQueryPool* QueryPool)
+	{
+		DevicePrintfBeginResult(Device, FString::Printf(TEXT("vkCreateQueryPool(CreateInfo=%p, OutQueryPool=%p)[...]"), CreateInfo, QueryPool));
+	}
+
+	void DumpCreatePipelineLayout(VkDevice Device, const VkPipelineLayoutCreateInfo* CreateInfo, VkPipelineLayout* PipelineLayout)
+	{
+		DevicePrintfBeginResult(Device, FString::Printf(TEXT("vkCreatePipelineLayout(CreateInfo=%p, OutPipelineLayout=%p)[...]"), CreateInfo, PipelineLayout));
+	}
+
+	void DumpCreateDescriptorPool(VkDevice Device, const VkDescriptorPoolCreateInfo* CreateInfo, VkDescriptorPool* DescriptorPool)
+	{
+		DevicePrintfBeginResult(Device, FString::Printf(TEXT("vkCreateDescriptorPool(CreateInfo=%p, OutDescriptorPool=%p)[...]"), CreateInfo, DescriptorPool));
+	}
+
+	void DumpCreateSampler(VkDevice Device, const VkSamplerCreateInfo* CreateInfo, VkSampler* Sampler)
+	{
+		DevicePrintfBeginResult(Device, FString::Printf(TEXT("vkCreateSampler(CreateInfo=%p, OutSampler=%p)[...]"), CreateInfo, Sampler));
+	}
+
+	void DumpCreateDevice(VkPhysicalDevice PhysicalDevice, const VkDeviceCreateInfo* CreateInfo, VkDevice* Device)
+	{
+		PrintfBeginResult(FString::Printf(TEXT("vkCreateDevice(PhysicalDevice=%p, CreateInfo=%p, OutDevice=%p)[...]"), PhysicalDevice, CreateInfo, Device));
+	}
+
+	void DumpGetPhysicalDeviceFeatures(VkPhysicalDevice PhysicalDevice, VkPhysicalDeviceFeatures* Features)
+	{
+		PrintfBeginResult(FString::Printf(TEXT("GetPhysicalDeviceFeatures(PhysicalDevice=%p, Features=%p)[...]"), PhysicalDevice, Features));
+	}
+
+	void DumpPhysicalDeviceFeatures(VkPhysicalDeviceFeatures* Features)
+	{
+		if (CVarVulkanDumpLayer.GetValueOnAnyThread())
+		{
+			DebugLog += FString::Printf(TEXT("VkPhysicalDeviceFeatures [...]\n"));
+		}
+	}
+
+	void DumpBeginCommandBuffer(VkCommandBuffer CommandBuffer, const VkCommandBufferBeginInfo* BeginInfo)
 	{
 		FlushDebugWrapperLog();
+
+		PrintfBeginResult(FString::Printf(TEXT("vkBeginCommandBuffer(CmdBuffer=%p, Info=%p)[...]"), CommandBuffer, BeginInfo));
 	}
 
-	static FAutoConsoleCommand CVarRHIThreadEnable(
-		TEXT("r.Vulkan.FlushLog"),
-		TEXT("\n"),
-		FConsoleCommandWithArgsDelegate::CreateStatic(&HandleFlushWrapperLog)
-		);
+	void DumpCmdBeginRenderPass(VkCommandBuffer CommandBuffer, const VkRenderPassBeginInfo* RenderPassBegin, VkSubpassContents Contents)
+	{
+		if (CVarVulkanDumpLayer.GetValueOnAnyThread())
+		{
+			auto GetSubpassContents = [](VkSubpassContents Contents) -> FString
+				{
+					switch (Contents)
+					{
+					case VK_SUBPASS_CONTENTS_INLINE: return TEXT("INLINE");
+					case VK_SUBPASS_CONTENTS_SECONDARY_COMMAND_BUFFERS: return TEXT("SECONDARY_CMD_BUFS");
+					default: return FString::Printf(TEXT("%d"), (int32)Contents);
+					}					
+				};
+			CmdPrintfBegin(CommandBuffer, FString::Printf(TEXT("vkCmdBeginRenderPass(BeginInfo=%p, Contents=%s)"), RenderPassBegin, *GetSubpassContents(Contents)));
+			DebugLog += FString::Printf(TEXT("%sBeginInfo: RenderPass=%p, Framebuffer=%p, renderArea=(x:%d, y:%d, w:%d, h:%d), clearValues=%d\n"),
+				Tabs, RenderPassBegin->renderPass, RenderPassBegin->framebuffer, 
+				RenderPassBegin->renderArea.offset.x, RenderPassBegin->renderArea.offset.y, 
+				RenderPassBegin->renderArea.extent.width, RenderPassBegin->renderArea.extent.height,
+				RenderPassBegin->clearValueCount);
+			for (uint32 Index = 0; Index < RenderPassBegin->clearValueCount; ++Index)
+			{
+				DebugLog += FString::Printf(TEXT("%s\tclearValue[%d]=(%d(%f), %d(%f), %d(%f), %d(%f))\n"),
+					Tabs, Index,
+					RenderPassBegin->pClearValues[Index].color.uint32[0], RenderPassBegin->pClearValues[Index].color.float32[0],
+					RenderPassBegin->pClearValues[Index].color.uint32[1], RenderPassBegin->pClearValues[Index].color.float32[1],
+					RenderPassBegin->pClearValues[Index].color.uint32[2], RenderPassBegin->pClearValues[Index].color.float32[2],
+					RenderPassBegin->pClearValues[Index].color.uint32[3], RenderPassBegin->pClearValues[Index].color.float32[3]);
+			}
+		}
+	}
+
+	void DumpCmdBindVertexBuffers(VkCommandBuffer CommandBuffer, uint32 FirstBinding, uint32 BindingCount, const VkBuffer* Buffers, const VkDeviceSize* Offsets)
+	{
+		CmdPrintfBegin(CommandBuffer, FString::Printf(TEXT("vkCmdBindVertexBuffers(FirstBinding=%d, NumBindings=%d, Buffers=%p, Offsets=%p)[...]"), FirstBinding, BindingCount, Buffers, Offsets));
+	}
+
+	void DumpCmdCopyBufferToImage(VkCommandBuffer CommandBuffer, VkBuffer SrcBuffer, VkImage DstImage, VkImageLayout DstImageLayout, uint32 RegionCount, const VkBufferImageCopy* Regions)
+	{
+		CmdPrintfBegin(CommandBuffer, FString::Printf(TEXT("vkCmdCopyBufferToImage(SrcBuffer=%p, DstImage=%p, DstImageLayout=%s, NumRegions=%d, Regions=%p)[...]"), 
+			SrcBuffer, DstImage, *GetImageLayoutString(DstImageLayout), RegionCount, Regions));
+	}
+
+	void DumpCmdCopyBuffer(VkCommandBuffer CommandBuffer, VkBuffer SrcBuffer, VkBuffer DstBuffer, uint32 RegionCount, const VkBufferCopy* Regions)
+	{
+		CmdPrintfBegin(CommandBuffer, FString::Printf(TEXT("vkCmdCopyBuffer(SrcBuffer=%p, DstBuffer=%p, NumRegions=%d, Regions=%p)[...]"), SrcBuffer, DstBuffer, RegionCount, Regions));
+	}
+
+	void DumpGetImageSubresourceLayout(VkDevice Device, VkImage Image, const VkImageSubresource* Subresource, VkSubresourceLayout* Layout)
+	{
+		DevicePrintfBeginResult(Device, FString::Printf(TEXT("vkGetImageSubresourceLayout(Image=%p, Subresource=%p, OutLayout=%p)"), Image, Subresource, Layout));
+	}
+
+	void DumpImageSubresourceLayout(VkSubresourceLayout* Layout)
+	{
+		if (CVarVulkanDumpLayer.GetValueOnAnyThread())
+		{
+			DebugLog += FString::Printf(TEXT("VkSubresourceLayout: [...]\n"));
+		}
+	}
+
+	void DumpSwapChainImages(VkResult Result, uint32* SwapchainImageCount, VkImage* SwapchainImages)
+	{
+		if (CVarVulkanDumpLayer.GetValueOnAnyThread())
+		{
+			PrintResult(Result);
+			if (SwapchainImages)
+			{
+				for (uint32 Index = 0; Index < *SwapchainImageCount; ++Index)
+				{
+					DebugLog += FString::Printf(TEXT("%sImage[%d]=%p\n"), Tabs, Index, SwapchainImages[Index]);
+				}
+			}
+			else
+			{
+				DebugLog += FString::Printf(TEXT("%sNumImages=%d\n"), Tabs, *SwapchainImageCount);
+			}
+		}
+	}
 
 	static struct FGlobalDumpLog
 	{
diff --git a/Engine/Source/Runtime/VulkanRHI/Private/VulkanDevice.cpp b/Engine/Source/Runtime/VulkanRHI/Private/VulkanDevice.cpp
--- a/Engine/Source/Runtime/VulkanRHI/Private/VulkanDevice.cpp
+++ b/Engine/Source/Runtime/VulkanRHI/Private/VulkanDevice.cpp
@@ -137,7 +137,7 @@ void FVulkanDevice::CreateDevice()
 	DeviceInfo.pEnabledFeatures = &Features;
 
 	// Create the device
-	VERIFYVULKANRESULT(vkCreateDevice(Gpu, &DeviceInfo, nullptr, &Device));
+	VERIFYVULKANRESULT(VulkanRHI::vkCreateDevice(Gpu, &DeviceInfo, nullptr, &Device));
 
 	// Create Graphics Queue, here we submit command buffers for execution
 	Queue = new FVulkanQueue(this, GfxQueueFamilyIndex, 0);
@@ -145,8 +145,8 @@ void FVulkanDevice::CreateDevice()
 #if VULKAN_ENABLE_DRAW_MARKERS
 	if (bDebugMarkersFound)
 	{
-		CmdDbgMarkerBegin = (PFN_vkCmdDebugMarkerBeginEXT)(void*)vkGetDeviceProcAddr(Device, "vkCmdDebugMarkerBeginEXT");
-		CmdDbgMarkerEnd = (PFN_vkCmdDebugMarkerEndEXT)(void*)vkGetDeviceProcAddr(Device, "vkCmdDebugMarkerEndEXT");
+		CmdDbgMarkerBegin = (PFN_vkCmdDebugMarkerBeginEXT)(void*)VulkanRHI::vkGetDeviceProcAddr(Device, "vkCmdDebugMarkerBeginEXT");
+		CmdDbgMarkerEnd = (PFN_vkCmdDebugMarkerEndEXT)(void*)VulkanRHI::vkGetDeviceProcAddr(Device, "vkCmdDebugMarkerEndEXT");
 
 		// We're running under RenderDoc or other trace tool, so enable capturing mode
 		GDynamicRHI->EnableIdealGPUCaptureOptions(true);
@@ -160,7 +160,7 @@ void FVulkanDevice::SetupFormats()
 	{
 		const VkFormat Format = (VkFormat)Index;
 		FMemory::Memzero(FormatProperties[Index]);
-		vkGetPhysicalDeviceFormatProperties(Gpu, Format, &FormatProperties[Index]);
+		VulkanRHI::vkGetPhysicalDeviceFormatProperties(Gpu, Format, &FormatProperties[Index]);
 	}
 
 	static_assert(sizeof(VkFormat) <= sizeof(GPixelFormats[0].PlatformFormat), "PlatformFormat must be increased!");
@@ -336,7 +336,7 @@ void FVulkanDevice::MapFormatSupport(EPixelFormat UEFormat, VkFormat VulkanForma
 bool FVulkanDevice::QueryGPU(int32 DeviceIndex)
 {
 	bool bDiscrete = false;
-	vkGetPhysicalDeviceProperties(Gpu, &GpuProps);
+	VulkanRHI::vkGetPhysicalDeviceProperties(Gpu, &GpuProps);
 
 	auto GetDeviceTypeString = [&]()
 	{
@@ -372,11 +372,11 @@ bool FVulkanDevice::QueryGPU(int32 DeviceIndex)
 	UE_LOG(LogVulkanRHI, Display, TEXT("Max Descriptor Sets Bound %d Timestamps %d"), GpuProps.limits.maxBoundDescriptorSets, GpuProps.limits.timestampComputeAndGraphics);
 
 	uint32 QueueCount = 0;
-	vkGetPhysicalDeviceQueueFamilyProperties(Gpu, &QueueCount, nullptr);
+	VulkanRHI::vkGetPhysicalDeviceQueueFamilyProperties(Gpu, &QueueCount, nullptr);
 	check(QueueCount >= 1);
 
 	QueueFamilyProps.AddUninitialized(QueueCount);
-	vkGetPhysicalDeviceQueueFamilyProperties(Gpu, &QueueCount, QueueFamilyProps.GetData());
+	VulkanRHI::vkGetPhysicalDeviceQueueFamilyProperties(Gpu, &QueueCount, QueueFamilyProps.GetData());
 
 	return bDiscrete;
 }
@@ -384,7 +384,7 @@ bool FVulkanDevice::QueryGPU(int32 DeviceIndex)
 void FVulkanDevice::InitGPU(int32 DeviceIndex)
 {
 	// Query features
-	vkGetPhysicalDeviceFeatures(Gpu, &Features);
+	VulkanRHI::vkGetPhysicalDeviceFeatures(Gpu, &Features);
 
 	UE_LOG(LogVulkanRHI, Display, TEXT("Geometry %d Tessellation %d"), Features.geometryShader, Features.tessellationShader);
 
@@ -478,13 +478,15 @@ void FVulkanDevice::Destroy()
 
 	MemoryManager.Deinit();
 
-	vkDestroyDevice(Device, nullptr);
+	DeferredDeletionQueue.Clear();
+
+	VulkanRHI::vkDestroyDevice(Device, nullptr);
 	Device = VK_NULL_HANDLE;
 }
 
 void FVulkanDevice::WaitUntilIdle()
 {
-	VERIFYVULKANRESULT(vkDeviceWaitIdle(Device));
+	VERIFYVULKANRESULT(VulkanRHI::vkDeviceWaitIdle(Device));
 }
 
 bool FVulkanDevice::IsFormatSupported(VkFormat Format) const
@@ -510,7 +512,7 @@ bool FVulkanDevice::IsFormatSupported(VkFormat Format) const
 	// Add it for faster caching next time
 	VkFormatProperties& NewProperties = ExtensionFormatProperties.Add(Format);
 	FMemory::Memzero(NewProperties);
-	vkGetPhysicalDeviceFormatProperties(Gpu, Format, &NewProperties);
+	VulkanRHI::vkGetPhysicalDeviceFormatProperties(Gpu, Format, &NewProperties);
 
 	return ArePropertiesSupported(NewProperties);
 }
diff --git a/Engine/Source/Runtime/VulkanRHI/Private/VulkanIndexBuffer.cpp b/Engine/Source/Runtime/VulkanRHI/Private/VulkanIndexBuffer.cpp
--- a/Engine/Source/Runtime/VulkanRHI/Private/VulkanIndexBuffer.cpp
+++ b/Engine/Source/Runtime/VulkanRHI/Private/VulkanIndexBuffer.cpp
@@ -163,7 +163,7 @@ void FVulkanResourceMultiBuffer::Unlock()
 			Region.size = LockSize;
 			//Region.srcOffset = 0;
 			Region.dstOffset = LockOffset + Buffers[DynamicBufferIndex]->GetOffset();
-			vkCmdCopyBuffer(CmdBuffer, StagingBuffer->GetHandle(), Buffers[DynamicBufferIndex]->GetHandle(), 1, &Region);
+			VulkanRHI::vkCmdCopyBuffer(CmdBuffer, StagingBuffer->GetHandle(), Buffers[DynamicBufferIndex]->GetHandle(), 1, &Region);
 			//UpdateBuffer(ResourceAllocation, IndexBuffer->GetBuffer(), LockSize, LockOffset);
 
 			//Device->GetDeferredDeletionQueue().EnqueueResource(Cmd, StagingBuffer);
diff --git a/Engine/Source/Runtime/VulkanRHI/Private/VulkanLayers.cpp b/Engine/Source/Runtime/VulkanRHI/Private/VulkanLayers.cpp
--- a/Engine/Source/Runtime/VulkanRHI/Private/VulkanLayers.cpp
+++ b/Engine/Source/Runtime/VulkanRHI/Private/VulkanLayers.cpp
@@ -115,14 +115,14 @@ static inline void GetInstanceLayerExtensions(const ANSICHAR* LayerName, FLayerE
 	{
 		//@TODO: Currently unsupported on device, so just make sure it doesn't cause problems
 		uint32 Count = 0;
-		Result = vkEnumerateInstanceExtensionProperties(LayerName, &Count, nullptr);
+		Result = VulkanRHI::vkEnumerateInstanceExtensionProperties(LayerName, &Count, nullptr);
 		check(Result >= VK_SUCCESS);
 
 		if (Count > 0)
 		{
 			OutLayer.ExtensionProps.Empty(Count);
 			OutLayer.ExtensionProps.AddUninitialized(Count);
-			Result = vkEnumerateInstanceExtensionProperties(LayerName, &Count, OutLayer.ExtensionProps.GetData());
+			Result = VulkanRHI::vkEnumerateInstanceExtensionProperties(LayerName, &Count, OutLayer.ExtensionProps.GetData());
 			check(Result >= VK_SUCCESS);
 		}
 	}
@@ -135,14 +135,14 @@ static inline void GetDeviceLayerExtensions(VkPhysicalDevice Device, const ANSIC
 	do
 	{
 		uint32 Count = 0;
-		Result = vkEnumerateDeviceExtensionProperties(Device, LayerName, &Count, nullptr);
+		Result = VulkanRHI::vkEnumerateDeviceExtensionProperties(Device, LayerName, &Count, nullptr);
 		check(Result >= VK_SUCCESS);
 
 		if (Count > 0)
 		{
 			OutLayer.ExtensionProps.Empty(Count);
 			OutLayer.ExtensionProps.AddUninitialized(Count);
-			Result = vkEnumerateDeviceExtensionProperties(Device, LayerName, &Count, OutLayer.ExtensionProps.GetData());
+			Result = VulkanRHI::vkEnumerateDeviceExtensionProperties(Device, LayerName, &Count, OutLayer.ExtensionProps.GetData());
 			check(Result >= VK_SUCCESS);
 		}
 	}
@@ -166,14 +166,14 @@ void FVulkanDynamicRHI::GetInstanceLayersAndExtensions(TArray<const ANSICHAR*>&
 	do
 	{
 		uint32 InstanceLayerCount = 0;
-		Result = vkEnumerateInstanceLayerProperties(&InstanceLayerCount, nullptr);
+		Result = VulkanRHI::vkEnumerateInstanceLayerProperties(&InstanceLayerCount, nullptr);
 		check(Result >= VK_SUCCESS);
 
 		if (InstanceLayerCount > 0)
 		{
 			GlobalLayers.Empty(InstanceLayerCount);
 			GlobalLayerProperties.AddZeroed(InstanceLayerCount);
-			Result = vkEnumerateInstanceLayerProperties(&InstanceLayerCount, &GlobalLayerProperties[GlobalLayerProperties.Num() - InstanceLayerCount]);
+			Result = VulkanRHI::vkEnumerateInstanceLayerProperties(&InstanceLayerCount, &GlobalLayerProperties[GlobalLayerProperties.Num() - InstanceLayerCount]);
 			check(Result >= VK_SUCCESS);
 		}
 	}
@@ -274,9 +274,9 @@ void FVulkanDevice::GetDeviceExtensions(TArray<const ANSICHAR*>& OutDeviceExtens
 	TArray<VkLayerProperties> LayerProperties;
 	{
 		uint32 Count = 0;
-		VERIFYVULKANRESULT(vkEnumerateDeviceLayerProperties(Gpu, &Count, nullptr));
+		VERIFYVULKANRESULT(VulkanRHI::vkEnumerateDeviceLayerProperties(Gpu, &Count, nullptr));
 		LayerProperties.AddZeroed(Count);
-		VERIFYVULKANRESULT(vkEnumerateDeviceLayerProperties(Gpu, &Count, (VkLayerProperties*)LayerProperties.GetData()));
+		VERIFYVULKANRESULT(VulkanRHI::vkEnumerateDeviceLayerProperties(Gpu, &Count, (VkLayerProperties*)LayerProperties.GetData()));
 		check(Count == LayerProperties.Num());
 	}
 
diff --git a/Engine/Source/Runtime/VulkanRHI/Private/VulkanMemory.cpp b/Engine/Source/Runtime/VulkanRHI/Private/VulkanMemory.cpp
--- a/Engine/Source/Runtime/VulkanRHI/Private/VulkanMemory.cpp
+++ b/Engine/Source/Runtime/VulkanRHI/Private/VulkanMemory.cpp
@@ -66,15 +66,15 @@ namespace VulkanRHI
 					}
 					if ((Flags & VK_MEMORY_PROPERTY_HOST_VISIBLE_BIT) == VK_MEMORY_PROPERTY_HOST_VISIBLE_BIT)
 					{
-						String += TEXT(" HVisible");
+						String += TEXT(" HostVisible");
 					}
 					if ((Flags & VK_MEMORY_PROPERTY_HOST_COHERENT_BIT) == VK_MEMORY_PROPERTY_HOST_COHERENT_BIT)
 					{
-						String += TEXT(" HCoherent");
+						String += TEXT(" HostCoherent");
 					}
 					if ((Flags & VK_MEMORY_PROPERTY_HOST_CACHED_BIT) == VK_MEMORY_PROPERTY_HOST_CACHED_BIT)
 					{
-						String += TEXT(" HCached");
+						String += TEXT(" HostCached");
 					}
 					if ((Flags & VK_MEMORY_PROPERTY_LAZILY_ALLOCATED_BIT) == VK_MEMORY_PROPERTY_LAZILY_ALLOCATED_BIT)
 					{
@@ -777,7 +777,7 @@ namespace VulkanRHI
 
 		FDeviceMemoryAllocation* DeviceMemoryAllocation = Device->GetMemoryManager().Alloc(BufferSize, MemoryTypeIndex, File, Line);
 		VERIFYVULKANRESULT(VulkanRHI::vkBindBufferMemory(Device->GetInstanceHandle(), Buffer, DeviceMemoryAllocation->GetHandle(), 0));
-		if (MemoryPropertyFlags == VK_MEMORY_PROPERTY_HOST_VISIBLE_BIT)
+		if (DeviceMemoryAllocation->CanBeMapped())
 		{
 			DeviceMemoryAllocation->Map(BufferSize, 0);
 		}
@@ -840,17 +840,17 @@ namespace VulkanRHI
 		ImageCreateInfo.sType = VK_STRUCTURE_TYPE_IMAGE_CREATE_INFO
 		ImageCreateInfo.size = ImageSize;
 		ImageCreateInfo.usage = ImageUsageFlags;
-		VERIFYVULKANRESULT(vkCreateImage(Device->GetInstanceHandle(), &ImageCreateInfo, nullptr, &Image));
+		VERIFYVULKANRESULT(VulkanRHI::vkCreateImage(Device->GetInstanceHandle(), &ImageCreateInfo, nullptr, &Image));
 
 		VkMemoryRequirements MemReqs;
-		vkGetImageMemoryRequirements(Device->GetInstanceHandle(), Image, &MemReqs);
+		VulkanRHI::vkGetImageMemoryRequirements(Device->GetInstanceHandle(), Image, &MemReqs);
 
 		uint32 MemoryTypeIndex;
 		VERIFYVULKANRESULT(Device->GetMemoryManager().GetMemoryTypeFromProperties(MemReqs.memoryTypeBits, MemoryPropertyFlags, &MemoryTypeIndex));
 
 		FDeviceMemoryAllocation* DeviceMemoryAllocation = Device->GetMemoryManager().Alloc(ImageSize, MemoryTypeIndex, File, Line);
 		VERIFYVULKANRESULT(VulkanRHI::vkBindImageMemory(Device->GetInstanceHandle(), Image, DeviceMemoryAllocation->GetHandle(), 0));
-		if (MemoryPropertyFlags == VK_MEMORY_PROPERTY_HOST_VISIBLE_BIT)
+		if (DeviceMemoryAllocation->CanBeMapped())
 		{
 			DeviceMemoryAllocation->Map(ImageSize, 0);
 		}
@@ -997,6 +997,7 @@ namespace VulkanRHI
 
 	void FBufferAllocation::Destroy(FVulkanDevice* Device)
 	{
+		// Does not need to go in the deferred deletion queue
 		VulkanRHI::vkDestroyBuffer(Device->GetInstanceHandle(), Buffer, nullptr);
 		Buffer = VK_NULL_HANDLE;
 	}
@@ -1006,11 +1007,12 @@ namespace VulkanRHI
 		checkf(!ResourceAllocation, TEXT("Staging Buffer not released!"));
 	}
 
-	void FStagingBuffer::Destroy(VkDevice DeviceHandle)
+	void FStagingBuffer::Destroy(FVulkanDevice* Device)
 	{
 		check(ResourceAllocation);
 
-		VulkanRHI::vkDestroyBuffer(DeviceHandle, Buffer, nullptr);
+		// Does not need to go in the deferred deletion queue
+		VulkanRHI::vkDestroyBuffer(Device->GetInstanceHandle(), Buffer, nullptr);
 		Buffer = VK_NULL_HANDLE;
 		ResourceAllocation = nullptr;
 		//Memory.Free(Allocation);
@@ -1072,7 +1074,7 @@ namespace VulkanRHI
 		VkMemoryRequirements MemReqs;
 		VulkanRHI::vkGetBufferMemoryRequirements(VulkanDevice, StagingBuffer->Buffer, &MemReqs);
 
-		StagingBuffer->ResourceAllocation = Device->GetResourceHeapManager().AllocateBufferMemory(MemReqs, VK_MEMORY_PROPERTY_HOST_VISIBLE_BIT, __FILE__, __LINE__);
+		StagingBuffer->ResourceAllocation = Device->GetResourceHeapManager().AllocateBufferMemory(MemReqs, VK_MEMORY_PROPERTY_HOST_VISIBLE_BIT | VK_MEMORY_PROPERTY_HOST_COHERENT_BIT, __FILE__, __LINE__);
 		StagingBuffer->ResourceAllocation->BindBuffer(Device, StagingBuffer->Buffer);
 
 		{
@@ -1143,7 +1145,7 @@ namespace VulkanRHI
 			FPendingItem& Entry = FreeStagingBuffers[Index];
 			if (bImmediately || !Entry.CmdBuffer || Entry.FenceCounter  + NUM_FRAMES_TO_WAIT_BEFORE_RELEASING_TO_OS < Entry.CmdBuffer->GetFenceSignaledCounter())
 			{
-				Entry.Resource->Destroy(Device->GetInstanceHandle());
+				Entry.Resource->Destroy(Device);
 				delete Entry.Resource;
 				FreeStagingBuffers.RemoveAtSwap(Index, 1, false);
 			}
@@ -1172,6 +1174,7 @@ namespace VulkanRHI
 
 	inline void FFenceManager::DestroyFence(FFence* Fence)
 	{
+		// Does not need to go in the deferred deletion queue
 		VulkanRHI::vkDestroyFence(Device->GetInstanceHandle(), Fence->GetHandle(), nullptr);
 		Fence->Handle = VK_NULL_HANDLE;
 		delete Fence;
@@ -1296,34 +1299,57 @@ namespace VulkanRHI
 
 	FDeferredDeletionQueue::~FDeferredDeletionQueue()
 	{
-		Clear();
+		check(Entries.Num() == 0);
 	}
 
-	void FDeferredDeletionQueue::EnqueueResource(FVulkanCmdBuffer* CmdBuffer, FRefCount* Resource)
+	void FDeferredDeletionQueue::EnqueueGenericResource(EType Type, void* Handle)
 	{
-		Resource->AddRef();
+		FVulkanQueue* Queue = Device->GetQueue();
 
 		FEntry Entry;
-		Entry.CmdBuffer = CmdBuffer;
-		Entry.FenceCounter = CmdBuffer->GetFenceSignaledCounter();
-		Entry.Resource = Resource;
+		Queue->GetLastSubmittedInfo(Entry.CmdBuffer, Entry.FenceCounter);
+		Entry.Handle = Handle;
+		Entry.StructureType = Type;
+
 		{
 			FScopeLock ScopeLock(&CS);
 			Entries.Add(Entry);
 		}
 	}
 
-	void FDeferredDeletionQueue::ReleaseResources(/*bool bDeleteImmediately*/)
+	void FDeferredDeletionQueue::ReleaseResources(bool bDeleteImmediately)
 	{
+		SCOPE_CYCLE_COUNTER(STAT_VulkanDeletionQueue);
 		FScopeLock ScopeLock(&CS);
 
+		VkDevice DeviceHandle = Device->GetInstanceHandle();
+
 		// Traverse list backwards so the swap switches to elements already tested
 		for (int32 Index = Entries.Num() - 1; Index >= 0; --Index)
 		{
 			FEntry* Entry = &Entries[Index];
-			if (Entry->FenceCounter < Entry->CmdBuffer->GetFenceSignaledCounter())
+			if (bDeleteImmediately || Entry->FenceCounter < Entry->CmdBuffer->GetFenceSignaledCounter())
 			{
-				Entry->Resource->Release();
+				switch (Entry->StructureType)
+				{
+#define VKSWITCH(Type)	case EType::Type: VulkanRHI::vkDestroy##Type(DeviceHandle, (Vk##Type)Entry->Handle, nullptr); break
+				VKSWITCH(RenderPass);
+				VKSWITCH(Buffer);
+				VKSWITCH(BufferView);
+				VKSWITCH(Image);
+				VKSWITCH(ImageView);
+				VKSWITCH(Pipeline);
+				VKSWITCH(PipelineLayout);
+				VKSWITCH(Framebuffer);
+				VKSWITCH(DescriptorSetLayout);
+				VKSWITCH(Sampler);
+				VKSWITCH(Semaphore);
+				VKSWITCH(ShaderModule);
+#undef VKSWTICH
+				default:
+					check(0);
+					break;
+				}
 				Entries.RemoveAtSwap(Index, 1, false);
 			}
 		}
@@ -1339,7 +1365,7 @@ namespace VulkanRHI
 		{
 			BufferSuballocations[Index] = InDevice->GetResourceHeapManager().AllocateBuffer(InSize,
 				VK_BUFFER_USAGE_VERTEX_BUFFER_BIT | VK_BUFFER_USAGE_INDEX_BUFFER_BIT | VK_BUFFER_USAGE_UNIFORM_BUFFER_BIT | VK_BUFFER_USAGE_UNIFORM_TEXEL_BUFFER_BIT,
-				VK_MEMORY_PROPERTY_HOST_VISIBLE_BIT,
+				VK_MEMORY_PROPERTY_HOST_VISIBLE_BIT | VK_MEMORY_PROPERTY_HOST_COHERENT_BIT,
 				__FILE__, __LINE__);
 			MappedData[Index] = (uint8*)BufferSuballocations[Index]->GetMappedPointer();
 			CurrentData[Index] = MappedData[Index];
diff --git a/Engine/Source/Runtime/VulkanRHI/Private/VulkanPendingState.cpp b/Engine/Source/Runtime/VulkanRHI/Private/VulkanPendingState.cpp
--- a/Engine/Source/Runtime/VulkanRHI/Private/VulkanPendingState.cpp
+++ b/Engine/Source/Runtime/VulkanRHI/Private/VulkanPendingState.cpp
@@ -374,7 +374,11 @@ bool FVulkanPendingState::RenderPassBegin(FVulkanCmdBuffer* CmdBuffer)
 		// @todo vulkan - we don't track the format of the depth buffer in the key, only if depth is enabled (write/test) - should be enough, but worth checking that
 		// if either bit is set that we have a depth buffer attached
 	}
-		
+
+	{
+		const VkExtent2D& LayoutExtents = CurrentState.RenderPass->GetLayout().GetExtent2D();
+		ensure(LayoutExtents.width == CurrentState.FrameBuffer->GetWidth() && LayoutExtents.height == CurrentState.FrameBuffer->GetHeight());
+	}
 	CmdBuffer->BeginRenderPass(CurrentState.RenderPass->GetLayout(), CurrentState.RenderPass->GetHandle(), CurrentState.FrameBuffer->GetHandle(), ClearValues);
 
 	bBeginRenderPass = true;
@@ -469,7 +473,7 @@ FVulkanDescriptorPool::FVulkanDescriptorPool(FVulkanDevice* InDevice)
 	PoolInfo.pPoolSizes = Types.GetData();
 	PoolInfo.maxSets = MaxDescriptorSets;
 
-	VERIFYVULKANRESULT(vkCreateDescriptorPool(Device->GetInstanceHandle(), &PoolInfo, nullptr, &DescriptorPool));
+	VERIFYVULKANRESULT(VulkanRHI::vkCreateDescriptorPool(Device->GetInstanceHandle(), &PoolInfo, nullptr, &DescriptorPool));
 
 	for (int32 Index = 0; Index < Types.Num(); ++Index)
 	{
@@ -482,7 +486,7 @@ FVulkanDescriptorPool::~FVulkanDescriptorPool()
 {
 	if (DescriptorPool != VK_NULL_HANDLE)
 	{
-		vkDestroyDescriptorPool(Device->GetInstanceHandle(), DescriptorPool, nullptr);
+		VulkanRHI::vkDestroyDescriptorPool(Device->GetInstanceHandle(), DescriptorPool, nullptr);
 		DescriptorPool = VK_NULL_HANDLE;
 	}
 }
@@ -544,7 +548,7 @@ void FVulkanDescriptorPool::TrackRemoveUsage(const FVulkanDescriptorSetsLayout&
 inline void FVulkanDescriptorSets::Bind(FVulkanCmdBuffer* Cmd, FVulkanBoundShaderState* State)
 {
 	//#todo-rco: Compute
-	vkCmdBindDescriptorSets(Cmd->GetHandle(),
+	VulkanRHI::vkCmdBindDescriptorSets(Cmd->GetHandle(),
 		VK_PIPELINE_BIND_POINT_GRAPHICS,
 		State->GetPipelineLayout(),
 		0, Sets.Num(), Sets.GetData(),
@@ -841,6 +845,8 @@ void FVulkanPendingState::SetRasterizerState(FVulkanRasterizerState* NewState)
 
 void FVulkanPendingState::NotifyDeletedRenderTarget(const FVulkanTextureBase* Texture)
 {
+	check(IsInRenderingThread());
+	//TArray<VkFramebuffer> FramebuffersToDelete;
 	for (auto& Pair : FrameBufferMap)
 	{
 		TArray<FVulkanFramebuffer*>& FrameBuffers = Pair.Value;
@@ -849,9 +855,10 @@ void FVulkanPendingState::NotifyDeletedRenderTarget(const FVulkanTextureBase* Te
 			FVulkanFramebuffer* FB = FrameBuffers[Index];
 			if (FB->ContainsRenderTarget(Texture))
 			{
+				//FramebuffersToDelete.Add(FB->GetHandle());
+				FrameBuffers.RemoveAtSwap(Index, 1, false);
 				FB->Destroy(*Device);
 				delete FB;
-				FrameBuffers.RemoveAtSwap(Index, 1, false);
 			}
 		}
 	}
diff --git a/Engine/Source/Runtime/VulkanRHI/Private/VulkanPipeline.cpp b/Engine/Source/Runtime/VulkanRHI/Private/VulkanPipeline.cpp
--- a/Engine/Source/Runtime/VulkanRHI/Private/VulkanPipeline.cpp
+++ b/Engine/Source/Runtime/VulkanRHI/Private/VulkanPipeline.cpp
@@ -57,7 +57,8 @@ FVulkanPipeline::FVulkanPipeline(FVulkanDevice* InDevice)
 FVulkanPipeline::~FVulkanPipeline()
 {
 #if VULKAN_ENABLE_PIPELINE_CACHE
-	vkDestroyPipeline(Device->GetInstanceHandle(), Pipeline, nullptr);
+	Device->GetDeferredDeletionQueue().EnqueueResource(VulkanRHI::FDeferredDeletionQueue::EType::Pipeline, Pipeline);
+	Pipeline = VK_NULL_HANDLE;
 #endif
 }
 
@@ -121,7 +122,7 @@ void FVulkanPipeline::Create(const FVulkanPipelineState& State)
 	VkPipelineCacheCreateInfo PipelineCacheInfo;
 	FMemory::Memzero(PipelineCacheInfo);
 	PipelineCacheInfo.sType = VK_STRUCTURE_TYPE_PIPELINE_CACHE_CREATE_INFO;
-	VERIFYVULKANRESULT(vkCreatePipelineCache(Device->GetInstanceHandle(), &PipelineCacheInfo, nullptr, &PipelineCache));
+	VERIFYVULKANRESULT(VulkanRHI::vkCreatePipelineCache(Device->GetInstanceHandle(), &PipelineCacheInfo, nullptr, &PipelineCache));
 
 	check(State.RenderPass);
 	PipelineInfo.pColorBlendState = &CBInfo;
@@ -133,7 +134,7 @@ void FVulkanPipeline::Create(const FVulkanPipelineState& State)
 	PipelineInfo.pDepthStencilState = &State.DepthStencilState->DepthStencilState;
 	PipelineInfo.pDynamicState = &State.DynamicState;
 
-	VERIFYVULKANRESULT(vkCreateGraphicsPipelines(Device->GetInstanceHandle(), PipelineCache, 1, &PipelineInfo, nullptr, &Pipeline));
+	VERIFYVULKANRESULT(VulkanRHI::vkCreateGraphicsPipelines(Device->GetInstanceHandle(), PipelineCache, 1, &PipelineInfo, nullptr, &Pipeline));
 }
 #endif
 
@@ -163,7 +164,7 @@ void FVulkanPipeline::InternalUpdateDynamicStates(FVulkanCmdBuffer* Cmd, FVulkan
 	{
 		check(State.Viewport.width > 0);
 		check(State.Viewport.height > 0);
-		vkCmdSetViewport(Cmd->GetHandle(), 0, 1, &State.Viewport);
+		VulkanRHI::vkCmdSetViewport(Cmd->GetHandle(), 0, 1, &State.Viewport);
 
 		State.bNeedsViewportUpdate = false;
 	}
@@ -178,7 +179,7 @@ void FVulkanPipeline::InternalUpdateDynamicStates(FVulkanCmdBuffer* Cmd, FVulkan
 			Scissor.extent.height = State.Viewport.height;
 		}
 
-		vkCmdSetScissor(Cmd->GetHandle(), 0, 1, &Scissor);
+		VulkanRHI::vkCmdSetScissor(Cmd->GetHandle(), 0, 1, &Scissor);
 
 		State.bNeedsScissorUpdate = false;
 	}
@@ -209,7 +210,7 @@ FVulkanPipelineStateCache::~FVulkanPipelineStateCache()
 {
 	DestroyCache();
 
-	vkDestroyPipelineCache(Device->GetInstanceHandle(), PipelineCache, nullptr);
+	VulkanRHI::vkDestroyPipelineCache(Device->GetInstanceHandle(), PipelineCache, nullptr);
 	PipelineCache = VK_NULL_HANDLE;
 }
 
@@ -320,7 +321,7 @@ void FVulkanPipelineStateCache::InitAndLoad(const TArray<FString>& CacheFilename
 	PipelineCacheInfo.sType = VK_STRUCTURE_TYPE_PIPELINE_CACHE_CREATE_INFO;
 	PipelineCacheInfo.initialDataSize = bLoaded ? DeviceCache.Num() : 0;
 	PipelineCacheInfo.pInitialData = bLoaded ? DeviceCache.GetData() : 0;
-	VERIFYVULKANRESULT(vkCreatePipelineCache(Device->GetInstanceHandle(), &PipelineCacheInfo, nullptr, &PipelineCache));
+	VERIFYVULKANRESULT(VulkanRHI::vkCreatePipelineCache(Device->GetInstanceHandle(), &PipelineCacheInfo, nullptr, &PipelineCache));
 }
 
 void FVulkanPipelineStateCache::Save(FString& CacheFilename)
@@ -881,7 +882,7 @@ void FVulkanPipelineStateCache::CreatePipelineFromDiskEntry(const FDiskEntry* Di
 	PipelineInfo.pDynamicState = &DynamicState;
 
 	//#todo-rco: Group the pipelines and create multiple (derived) pipelines at once
-	VERIFYVULKANRESULT(vkCreateGraphicsPipelines(Device->GetInstanceHandle(), PipelineCache, 1, &PipelineInfo, nullptr, &Pipeline->Pipeline));
+	VERIFYVULKANRESULT(VulkanRHI::vkCreateGraphicsPipelines(Device->GetInstanceHandle(), PipelineCache, 1, &PipelineInfo, nullptr, &Pipeline->Pipeline));
 }
 
 void FVulkanPipelineStateCache::PopulateDiskEntry(const FVulkanPipelineState& State, const FVulkanRenderPass* RenderPass, FVulkanPipelineStateCache::FDiskEntry* OutDiskEntry)
@@ -982,7 +983,7 @@ void FVulkanPipelineStateCache::CreateDiskEntryRuntimeObjects(FDiskEntry* DiskEn
 		//PipelineLayoutInfo.pushConstantRangeCount = 0;
 		//PipelineLayoutInfo.pPushConstantRanges = nullptr;
 
-		VERIFYVULKANRESULT(vkCreatePipelineLayout(Device->GetInstanceHandle(), &PipelineLayoutInfo, nullptr, &DiskEntry->PipelineLayout));
+		VERIFYVULKANRESULT(VulkanRHI::vkCreatePipelineLayout(Device->GetInstanceHandle(), &PipelineLayoutInfo, nullptr, &DiskEntry->PipelineLayout));
 	}
 
 	{
@@ -1029,7 +1030,7 @@ void FVulkanPipelineStateCache::DestroyCache()
 			}
 		}
 
-		vkDestroyPipelineLayout(DeviceHandle, DiskEntry.PipelineLayout, nullptr);
+		VulkanRHI::vkDestroyPipelineLayout(DeviceHandle, DiskEntry.PipelineLayout, nullptr);
 
 		for (int32 Index = 0; Index < DiskEntry.DescriptorSetLayouts.Num(); ++Index)
 		{
diff --git a/Engine/Source/Runtime/VulkanRHI/Private/VulkanQuery.cpp b/Engine/Source/Runtime/VulkanRHI/Private/VulkanQuery.cpp
--- a/Engine/Source/Runtime/VulkanRHI/Private/VulkanQuery.cpp
+++ b/Engine/Source/Runtime/VulkanRHI/Private/VulkanQuery.cpp
@@ -38,12 +38,12 @@ FVulkanQueryPool::FVulkanQueryPool(FVulkanDevice* InDevice, uint32 InNumQueries,
 	PoolCreateInfo.queryType = QueryType;
 	PoolCreateInfo.queryCount = NumQueries;
 
-	VERIFYVULKANRESULT(vkCreateQueryPool(Device->GetInstanceHandle(), &PoolCreateInfo, nullptr, &QueryPool));
+	VERIFYVULKANRESULT(VulkanRHI::vkCreateQueryPool(Device->GetInstanceHandle(), &PoolCreateInfo, nullptr, &QueryPool));
 }
 
 void FVulkanQueryPool::Destroy()
 {
-	vkDestroyQueryPool(Device->GetInstanceHandle(), QueryPool, nullptr);
+	VulkanRHI::vkDestroyQueryPool(Device->GetInstanceHandle(), QueryPool, nullptr);
 	QueryPool = VK_NULL_HANDLE;
 }
 
@@ -70,10 +70,10 @@ void FVulkanTimestampQueryPool::WriteStartFrame(VkCommandBuffer CmdBuffer)
 	}
 	static_assert(StartFrame + 1 == EndFrame, "Enums required to be contiguous!");
 
-	vkCmdResetQueryPool(CmdBuffer, QueryPool, StartFrame, TotalQueries);
+	VulkanRHI::vkCmdResetQueryPool(CmdBuffer, QueryPool, StartFrame, TotalQueries);
 
 	// Start Frame Timestamp
-	vkCmdWriteTimestamp(CmdBuffer, VK_PIPELINE_STAGE_ALL_COMMANDS_BIT, QueryPool, StartFrame);
+	VulkanRHI::vkCmdWriteTimestamp(CmdBuffer, VK_PIPELINE_STAGE_ALL_COMMANDS_BIT, QueryPool, StartFrame);
 }
 
 void FVulkanTimestampQueryPool::WriteEndFrame(VkCommandBuffer CmdBuffer)
@@ -81,7 +81,7 @@ void FVulkanTimestampQueryPool::WriteEndFrame(VkCommandBuffer CmdBuffer)
 	if (!bFirst)
 	{
 		// End Frame Timestamp
-		vkCmdWriteTimestamp(CmdBuffer, VK_PIPELINE_STAGE_ALL_COMMANDS_BIT, QueryPool, EndFrame);
+		VulkanRHI::vkCmdWriteTimestamp(CmdBuffer, VK_PIPELINE_STAGE_ALL_COMMANDS_BIT, QueryPool, EndFrame);
 	}
 }
 
@@ -92,7 +92,7 @@ void FVulkanTimestampQueryPool::CalculateFrameTime()
 	{
 		VkDevice DeviceHandle = Device->GetInstanceHandle();
 		VkResult Result;
-		Result = vkGetQueryPoolResults(DeviceHandle, QueryPool, StartFrame, 2, sizeof(Results), Results, sizeof(uint64), /*VK_QUERY_RESULT_WAIT_BIT | */VK_QUERY_RESULT_64_BIT);
+		Result = VulkanRHI::vkGetQueryPoolResults(DeviceHandle, QueryPool, StartFrame, 2, sizeof(Results), Results, sizeof(uint64), /*VK_QUERY_RESULT_WAIT_BIT | */VK_QUERY_RESULT_64_BIT);
 		if (Result != VK_SUCCESS)
 		{
 			GGPUFrameTime = 0;
@@ -122,7 +122,7 @@ int32 FVulkanTimestampQueryPool::AllocateUserQuery()
 void FVulkanTimestampQueryPool::WriteTimestamp(VkCommandBuffer CmdBuffer, int32 UserQuery, VkPipelineStageFlagBits PipelineStageBits)
 {
 	check(UserQuery != -1);
-	vkCmdWriteTimestamp(CmdBuffer, PipelineStageBits, QueryPool, StartUser + UserQuery);
+	VulkanRHI::vkCmdWriteTimestamp(CmdBuffer, PipelineStageBits, QueryPool, StartUser + UserQuery);
 }
 
 uint32 FVulkanTimestampQueryPool::CalculateTimeFromUserQueries(int32 UserBegin, int32 UserEnd, bool bWait)
@@ -133,13 +133,13 @@ uint32 FVulkanTimestampQueryPool::CalculateTimeFromUserQueries(int32 UserBegin,
 	VkDevice DeviceHandle = Device->GetInstanceHandle();
 	VkResult Result;
 
-	Result = vkGetQueryPoolResults(DeviceHandle, QueryPool, StartUser + UserBegin, 1, sizeof(uint64_t), &Begin, 0, /*(bWait ? VK_QUERY_RESULT_WAIT_BIT : 0) |*/ VK_QUERY_RESULT_64_BIT);
+	Result = VulkanRHI::vkGetQueryPoolResults(DeviceHandle, QueryPool, StartUser + UserBegin, 1, sizeof(uint64_t), &Begin, 0, /*(bWait ? VK_QUERY_RESULT_WAIT_BIT : 0) |*/ VK_QUERY_RESULT_64_BIT);
 	if (Result != VK_SUCCESS)
 	{
 		return 0;
 	}
 
-	Result = vkGetQueryPoolResults(DeviceHandle, QueryPool, StartUser + UserEnd, 1, sizeof(uint64_t), &End, 0, (bWait ? VK_QUERY_RESULT_WAIT_BIT : 0) | VK_QUERY_RESULT_64_BIT);
+	Result = VulkanRHI::vkGetQueryPoolResults(DeviceHandle, QueryPool, StartUser + UserEnd, 1, sizeof(uint64_t), &End, 0, (bWait ? VK_QUERY_RESULT_WAIT_BIT : 0) | VK_QUERY_RESULT_64_BIT);
 	if (Result != VK_SUCCESS)
 	{
 		return 0;
diff --git a/Engine/Source/Runtime/VulkanRHI/Private/VulkanQueue.cpp b/Engine/Source/Runtime/VulkanRHI/Private/VulkanQueue.cpp
--- a/Engine/Source/Runtime/VulkanRHI/Private/VulkanQueue.cpp
+++ b/Engine/Source/Runtime/VulkanRHI/Private/VulkanQueue.cpp
@@ -24,6 +24,9 @@ FVulkanQueue::FVulkanQueue(FVulkanDevice* InDevice, uint32 InFamilyIndex, uint32
 	, FamilyIndex(InFamilyIndex)
 	, QueueIndex(InQueueIndex)
 	, Device(InDevice)
+	, LastSubmittedCmdBuffer(nullptr)
+	, LastSubmittedCmdBufferFenceCounter(0)
+	, SubmitCounter(0)
 {
 	VulkanRHI::vkGetDeviceQueue(Device->GetInstanceHandle(), FamilyIndex, InQueueIndex, &Queue);
 }
@@ -66,10 +69,20 @@ void FVulkanQueue::Submit(FVulkanCmdBuffer* CmdBuffer, FVulkanSemaphore* WaitSem
 
 	if (GWaitForIdleOnSubmit != 0)
 	{
-		VERIFYVULKANRESULT(vkQueueWaitIdle(Queue));
+		VERIFYVULKANRESULT(VulkanRHI::vkQueueWaitIdle(Queue));
 	}
 
 	CmdBuffer->State = FVulkanCmdBuffer::EState::Submitted;
 
+	UpdateLastSubmittedCommandBuffer(CmdBuffer);
+
 	CmdBuffer->GetOwner()->RefreshFenceStatus();
 }
+
+void FVulkanQueue::UpdateLastSubmittedCommandBuffer(FVulkanCmdBuffer* CmdBuffer)
+{
+	FScopeLock ScopeLock(&CS);
+	LastSubmittedCmdBuffer = CmdBuffer;
+	LastSubmittedCmdBufferFenceCounter = CmdBuffer->GetFenceSignaledCounter();
+	++SubmitCounter;
+}
diff --git a/Engine/Source/Runtime/VulkanRHI/Private/VulkanQueue.h b/Engine/Source/Runtime/VulkanRHI/Private/VulkanQueue.h
--- a/Engine/Source/Runtime/VulkanRHI/Private/VulkanQueue.h
+++ b/Engine/Source/Runtime/VulkanRHI/Private/VulkanQueue.h
@@ -38,9 +38,28 @@ public:
 		return Queue;
 	}
 
+	void GetLastSubmittedInfo(FVulkanCmdBuffer*& OutCmdBuffer, uint64& OutFenceCounter) const
+	{
+		FScopeLock ScopeLock(&CS);
+		OutCmdBuffer = LastSubmittedCmdBuffer;
+		OutFenceCounter = LastSubmittedCmdBufferFenceCounter;
+	}
+
+	inline uint64 GetSubmitCount() const
+	{
+		return SubmitCounter;
+	}
+
 private:
 	VkQueue Queue;
 	uint32 FamilyIndex;
 	uint32 QueueIndex;
 	FVulkanDevice* Device;
+
+	mutable FCriticalSection CS;
+	FVulkanCmdBuffer* LastSubmittedCmdBuffer;
+	uint64 LastSubmittedCmdBufferFenceCounter;
+	uint64 SubmitCounter;
+
+	void UpdateLastSubmittedCommandBuffer(FVulkanCmdBuffer* CmdBuffer);
 };
diff --git a/Engine/Source/Runtime/VulkanRHI/Private/VulkanRHI.cpp b/Engine/Source/Runtime/VulkanRHI/Private/VulkanRHI.cpp
--- a/Engine/Source/Runtime/VulkanRHI/Private/VulkanRHI.cpp
+++ b/Engine/Source/Runtime/VulkanRHI/Private/VulkanRHI.cpp
@@ -127,39 +127,6 @@ struct FVulkanMemManager
 
 static FVulkanMemManager GVulkanMemMgr;
 
-static inline VkAttachmentLoadOp RenderTargetLoadActionToVulkan(ERenderTargetLoadAction InLoadAction)
-{
-	VkAttachmentLoadOp OutLoadAction = VK_ATTACHMENT_LOAD_OP_MAX_ENUM;
-
-	switch (InLoadAction)
-	{
-	case ERenderTargetLoadAction::ELoad:		OutLoadAction = VK_ATTACHMENT_LOAD_OP_LOAD;			break;
-	case ERenderTargetLoadAction::EClear:		OutLoadAction = VK_ATTACHMENT_LOAD_OP_CLEAR;		break;
-	case ERenderTargetLoadAction::ENoAction:	OutLoadAction = VK_ATTACHMENT_LOAD_OP_DONT_CARE;	break;
-	default:																						break;
-	}
-
-	// Check for missing translation
-	check(OutLoadAction != VK_ATTACHMENT_LOAD_OP_MAX_ENUM);
-	return OutLoadAction;
-}
-
-static inline VkAttachmentStoreOp RenderTargetStoreActionToVulkan(ERenderTargetStoreAction InStoreAction)
-{
-	VkAttachmentStoreOp OutStoreAction = VK_ATTACHMENT_STORE_OP_MAX_ENUM;
-
-	switch (InStoreAction)
-	{
-	case ERenderTargetStoreAction::EStore:		OutStoreAction = VK_ATTACHMENT_STORE_OP_STORE;		break;
-	case ERenderTargetStoreAction::ENoAction:	OutStoreAction = VK_ATTACHMENT_STORE_OP_DONT_CARE;	break;
-	default:																						break;
-	}
-
-	// Check for missing translation
-	check(OutStoreAction != VK_ATTACHMENT_STORE_OP_MAX_ENUM);
-	return OutStoreAction;
-}
-
 static inline int32 CountSetBits(int32 n)
 {
 	uint32 Count = 0;
@@ -730,7 +697,7 @@ void FVulkanCommandListContext::RHIEndDrawingViewport(FViewportRHIParamRef Viewp
 void FVulkanCommandListContext::RHIEndFrame()
 {
 	//FRCLog::Printf(FString::Printf(TEXT("FVulkanCommandListContext::RHIEndFrame\n")));
-	Device->GetDeferredDeletionQueue().Clear();
+	Device->GetDeferredDeletionQueue().ReleaseResources();
 
 	Device->GetStagingManager().ProcessPendingFree();
 	Device->GetResourceHeapManager().ReleaseFreedPages();
@@ -747,6 +714,9 @@ void FVulkanCommandListContext::RHIPushEvent(const TCHAR* Name, FColor Color)
 	if (IsImmediate())
 	{
 #if VULKAN_ENABLE_DRAW_MARKERS
+#if VULKAN_ENABLE_DUMP_LAYER
+		VulkanRHI::PrintfBegin(FString::Printf(TEXT("vkCmdDbgMarkerBeginEXT(%s)"), Name));
+#endif
 		if (auto CmdDbgMarkerBegin = Device->GetCmdDbgMarkerBegin())
 		{
 			VkDebugMarkerMarkerInfoEXT Info;
@@ -771,6 +741,9 @@ void FVulkanCommandListContext::RHIPopEvent()
 	if (IsImmediate())
 	{
 #if VULKAN_ENABLE_DRAW_MARKERS
+#if VULKAN_ENABLE_DUMP_LAYER
+		VulkanRHI::PrintfBegin(TEXT("vkCmdDbgMarkerEndEXT()"));
+#endif
 		if (auto CmdDbgMarkerEnd = Device->GetCmdDbgMarkerEnd())
 		{
 			CmdDbgMarkerEnd(GetCommandBufferManager()->GetActiveCmdBuffer()->GetHandle());
@@ -851,7 +824,7 @@ FVulkanBuffer::~FVulkanBuffer()
 	// The buffer should be unmapped
 	check(BufferPtr == nullptr);
 
-	VulkanRHI::vkDestroyBuffer(Device.GetInstanceHandle(), Buf, nullptr);
+	Device.GetDeferredDeletionQueue().EnqueueResource(FDeferredDeletionQueue::EType::Buffer, Buf);
 	Buf = VK_NULL_HANDLE;
 
 	Device.GetMemoryManager().Free(Allocation);
@@ -911,10 +884,13 @@ FVulkanDescriptorSetsLayout::FVulkanDescriptorSetsLayout(FVulkanDevice* InDevice
 FVulkanDescriptorSetsLayout::~FVulkanDescriptorSetsLayout()
 {
 	check(Device);
-	for (VkDescriptorSetLayout Handle : LayoutHandles)
+	VulkanRHI::FDeferredDeletionQueue& DeletionQueue = Device->GetDeferredDeletionQueue();
+	for (VkDescriptorSetLayout& Handle : LayoutHandles)
 	{
-		VulkanRHI::vkDestroyDescriptorSetLayout(Device->GetInstanceHandle(), Handle, nullptr);
+		DeletionQueue.EnqueueResource(VulkanRHI::FDeferredDeletionQueue::EType::DescriptorSetLayout, Handle);
 	}
+
+	LayoutHandles.Reset(0);
 }
 
 void FVulkanDescriptorSetsLayout::AddDescriptor(int32 DescriptorSetIndex, const VkDescriptorSetLayoutBinding& Descriptor, int32 BindingIndex)
@@ -1081,172 +1057,11 @@ void FVulkanBufferView::Destroy()
 {
 	if (View != VK_NULL_HANDLE)
 	{
-		VulkanRHI::vkDestroyBufferView(GetParent()->GetInstanceHandle(), View, nullptr);
+		Device->GetDeferredDeletionQueue().EnqueueResource(FDeferredDeletionQueue::EType::BufferView, View);
 		View = VK_NULL_HANDLE;
 	}
 }
 
-FVulkanRenderTargetLayout::FVulkanRenderTargetLayout(const FRHISetRenderTargetsInfo& RTInfo)
-	: NumAttachments(0)
-	, NumColorAttachments(0)
-	, bHasDepthStencil(false)
-	, bHasResolveAttachments(false)
-	, Hash(0)
-{
-	FMemory::Memzero(ColorReferences);
-	FMemory::Memzero(ResolveReferences);
-	FMemory::Memzero(DepthStencilReference);
-	FMemory::Memzero(Desc);
-	FMemory::Memzero(Extent);
-
-	bool bSetExtent = false;
-		
-	for (int32 Index = 0; Index < RTInfo.NumColorRenderTargets; ++Index)
-	{
-		const FRHIRenderTargetView& rtv = RTInfo.ColorRenderTarget[Index];
-		if(rtv.Texture)
-		{
-		    FVulkanTextureBase* Texture = FVulkanTextureBase::Cast(rtv.Texture);
-		    check(Texture);
-    
-		    if (!bSetExtent)
-		    {
-			    bSetExtent = true;
-			    Extent.Extent3D.width = Texture->Surface.Width;
-			    Extent.Extent3D.height = Texture->Surface.Height;
-			    Extent.Extent3D.depth = 1;
-		    }
-    
-		    uint32 NumSamples = rtv.Texture->GetNumSamples();
-	    
-		    VkAttachmentDescription& CurrDesc = Desc[NumAttachments];
-		    
-		    //@TODO: Check this, it should be a power-of-two. Might be a VulkanConvert helper function.
-		    CurrDesc.samples = static_cast<VkSampleCountFlagBits>(NumSamples);
-		    CurrDesc.format = (VkFormat)GPixelFormats[rtv.Texture->GetFormat()].PlatformFormat;
-		    CurrDesc.loadOp = RenderTargetLoadActionToVulkan(rtv.LoadAction);
-		    CurrDesc.storeOp = RenderTargetStoreActionToVulkan(rtv.StoreAction);
-		    CurrDesc.stencilLoadOp = VK_ATTACHMENT_LOAD_OP_DONT_CARE;
-#if !VULKAN_USE_MSAA_RESOLVE_ATTACHMENTS
-		    CurrDesc.stencilStoreOp = VK_ATTACHMENT_STORE_OP_DONT_CARE;
-#endif
-		    CurrDesc.initialLayout = VK_IMAGE_LAYOUT_GENERAL;
-		    CurrDesc.finalLayout = VK_IMAGE_LAYOUT_GENERAL;
-    
-		    ColorReferences[NumColorAttachments].attachment = NumAttachments;
-		    ColorReferences[NumColorAttachments].layout = VK_IMAGE_LAYOUT_GENERAL;
-		    NumAttachments++;
-    
-#if VULKAN_USE_MSAA_RESOLVE_ATTACHMENTS
-			if (NumSamples == 1)
-			{
-			    CurrDesc.storeOp = RenderTargetStoreActionToVulkan(rtv.StoreAction);
-    
-			    ResolveReferences[NumColorAttachments].attachment = VK_ATTACHMENT_UNUSED;
-			    ResolveReferences[NumColorAttachments].layout = VK_IMAGE_LAYOUT_UNDEFINED;
-		    }
-			else
-		    {
-			    // discard MSAA color target after resolve to 1x surface
-			    CurrDesc.storeOp = VK_ATTACHMENT_STORE_OP_DONT_CARE;
-    
-			    // Resolve attachment
-			    VkAttachmentDescription& ResolveDesc = Desc[NumAttachments];
-			    ResolveDesc.samples = VK_SAMPLE_COUNT_1_BIT;
-			    ResolveDesc.format = (VkFormat)GPixelFormats[rtv.Texture->GetFormat()].PlatformFormat;
-			    ResolveDesc.loadOp = VK_ATTACHMENT_LOAD_OP_DONT_CARE;
-			    ResolveDesc.storeOp = RenderTargetStoreActionToVulkan(rtv.StoreAction);
-			    ResolveDesc.stencilLoadOp = VK_ATTACHMENT_LOAD_OP_DONT_CARE;
-			    ResolveDesc.stencilStoreOp = VK_ATTACHMENT_STORE_OP_DONT_CARE;
-			    ResolveDesc.initialLayout = VK_IMAGE_LAYOUT_COLOR_ATTACHMENT_OPTIMAL;
-			    ResolveDesc.finalLayout = VK_IMAGE_LAYOUT_COLOR_ATTACHMENT_OPTIMAL;
-    
-			    ResolveReferences[NumColorAttachments].attachment = NumAttachments;
-			    ResolveReferences[NumColorAttachments].layout = VK_IMAGE_LAYOUT_COLOR_ATTACHMENT_OPTIMAL;
-			    NumAttachments++;
-			    bHasResolveAttachments = true;
-		    }
-#endif
-	    
-		    NumColorAttachments++;
-		}
-	}
-
-	if (RTInfo.DepthStencilRenderTarget.Texture)
-	{
-		VkAttachmentDescription& CurrDesc = Desc[NumAttachments];
-		FMemory::Memzero(CurrDesc);
-		FVulkanTextureBase* Texture = FVulkanTextureBase::Cast(RTInfo.DepthStencilRenderTarget.Texture);
-		check(Texture);
-
-		//@TODO: Check this, it should be a power-of-two. Might be a VulkanConvert helper function.
-		CurrDesc.samples = static_cast<VkSampleCountFlagBits>(RTInfo.DepthStencilRenderTarget.Texture->GetNumSamples());
-		CurrDesc.format = (VkFormat)GPixelFormats[RTInfo.DepthStencilRenderTarget.Texture->GetFormat()].PlatformFormat;
-		CurrDesc.loadOp = RenderTargetLoadActionToVulkan(RTInfo.DepthStencilRenderTarget.DepthLoadAction);
-		CurrDesc.stencilLoadOp = RenderTargetLoadActionToVulkan(RTInfo.DepthStencilRenderTarget.StencilLoadAction);
-		if (CurrDesc.samples == VK_SAMPLE_COUNT_1_BIT)
-		{
-			CurrDesc.storeOp = RenderTargetStoreActionToVulkan(RTInfo.DepthStencilRenderTarget.DepthStoreAction);
-			CurrDesc.stencilStoreOp = RenderTargetStoreActionToVulkan(RTInfo.DepthStencilRenderTarget.GetStencilStoreAction());
-		}
-		else
-		{
-			// Never want to store MSAA depth/stencil
-			CurrDesc.storeOp = VK_ATTACHMENT_STORE_OP_DONT_CARE;
-			CurrDesc.stencilStoreOp = VK_ATTACHMENT_STORE_OP_DONT_CARE;
-		}
-		CurrDesc.initialLayout = VK_IMAGE_LAYOUT_DEPTH_STENCIL_ATTACHMENT_OPTIMAL;
-		CurrDesc.finalLayout = VK_IMAGE_LAYOUT_DEPTH_STENCIL_ATTACHMENT_OPTIMAL;
-		
-		DepthStencilReference.attachment = NumAttachments;
-		DepthStencilReference.layout = VK_IMAGE_LAYOUT_DEPTH_STENCIL_ATTACHMENT_OPTIMAL;
-
-		bHasDepthStencil = true;
-		NumAttachments++;
-
-		if (!bSetExtent)
-		{
-			bSetExtent = true;
-			Extent.Extent3D.width = Texture->Surface.Width;
-			Extent.Extent3D.height = Texture->Surface.Height;
-			Extent.Extent3D.depth = 1;
-		}
-	}
-
-	check(bSetExtent);
-
-	//@TODO: remove loads on device
-#if 0
-	//@HACK: don't used the load/store info in the hash (avoids a store/resolve between the tonemapping and slate passes)
-	FRHISetRenderTargetsInfo HashedRTInfo;
-	FMemory::Memcpy(HashedRTInfo, RTInfo);
-	for (int32 Index = 0; Index < HashedRTInfo.NumColorRenderTargets; ++Index)
-	{
-		FRHIRenderTargetView* View = &HashedRTInfo.ColorRenderTarget[Index];
-		View->LoadAction = ERenderTargetLoadAction::ENoAction;
-		View->StoreAction = ERenderTargetStoreAction::ENoAction;
-	}
-
-	if (HashedRTInfo.DepthStencilRenderTarget.Texture)
-	{
-		FRHIDepthRenderTargetView* View = &HashedRTInfo.DepthStencilRenderTarget;
-		View->DepthLoadAction = ERenderTargetLoadAction::ENoAction;
-		View->DepthStoreAction = ERenderTargetStoreAction::ENoAction;
-		View->StencilLoadAction = ERenderTargetLoadAction::ENoAction;
-		View->StencilStoreAction = ERenderTargetStoreAction::ENoAction;
-		View->DepthStencilAccess = FExclusiveDepthStencil();
-		// doubt we care about the stencil store action anyway
-	}
-	HashedRTInfo.bClearDepth = false;
-	HashedRTInfo.bClearStencil = false;
-	HashedRTInfo.bClearColor = false;
-
-	Hash = HashedRTInfo.CalculateHash();
-#else
-	Hash = RTInfo.CalculateHash();
-#endif
-}
-
 FVulkanRenderPass::FVulkanRenderPass(FVulkanDevice& InDevice, const FVulkanRenderTargetLayout& RTLayout) :
 	Layout(RTLayout),
 	RenderPass(VK_NULL_HANDLE),
@@ -1285,7 +1100,7 @@ FVulkanRenderPass::~FVulkanRenderPass()
 {
 	DEC_DWORD_STAT(STAT_VulkanNumRenderPasses);
 
-	VulkanRHI::vkDestroyRenderPass(Device.GetInstanceHandle(), RenderPass, nullptr);
+	Device.GetDeferredDeletionQueue().EnqueueResource(FDeferredDeletionQueue::EType::RenderPass, RenderPass);
 	RenderPass = VK_NULL_HANDLE;
 }
 
@@ -1409,7 +1224,7 @@ FVulkanRingBuffer::FVulkanRingBuffer(FVulkanDevice* InDevice, uint64 TotalSize,
 	, MinAlignment(0)
 {
 	FRHIResourceCreateInfo CreateInfo;
-	BufferSuballocation = InDevice->GetResourceHeapManager().AllocateBuffer(TotalSize, Usage, VK_MEMORY_PROPERTY_HOST_VISIBLE_BIT, __FILE__, __LINE__);
+	BufferSuballocation = InDevice->GetResourceHeapManager().AllocateBuffer(TotalSize, Usage, VK_MEMORY_PROPERTY_HOST_VISIBLE_BIT | VK_MEMORY_PROPERTY_HOST_COHERENT_BIT, __FILE__, __LINE__);
 	MinAlignment = BufferSuballocation->GetBufferAllocation()->GetAlignment();
 }
 
diff --git a/Engine/Source/Runtime/VulkanRHI/Private/VulkanRHIPrivate.h b/Engine/Source/Runtime/VulkanRHI/Private/VulkanRHIPrivate.h
--- a/Engine/Source/Runtime/VulkanRHI/Private/VulkanRHIPrivate.h
+++ b/Engine/Source/Runtime/VulkanRHI/Private/VulkanRHIPrivate.h
@@ -151,13 +151,14 @@ struct FVulkanSemaphore
 		PresentCompleteSemaphoreCreateInfo.sType = VK_STRUCTURE_TYPE_SEMAPHORE_CREATE_INFO;
 		PresentCompleteSemaphoreCreateInfo.pNext = nullptr;
 		PresentCompleteSemaphoreCreateInfo.flags = 0;
-		VERIFYVULKANRESULT(vkCreateSemaphore(Device.GetInstanceHandle(), &PresentCompleteSemaphoreCreateInfo, nullptr, &SemaphoreHandle));
+		VERIFYVULKANRESULT(VulkanRHI::vkCreateSemaphore(Device.GetInstanceHandle(), &PresentCompleteSemaphoreCreateInfo, nullptr, &SemaphoreHandle));
 	}
 
 	~FVulkanSemaphore()
 	{
 		check(SemaphoreHandle != VK_NULL_HANDLE);
-		vkDestroySemaphore(Device.GetInstanceHandle(), SemaphoreHandle, nullptr);
+		Device.GetDeferredDeletionQueue().EnqueueResource(VulkanRHI::FDeferredDeletionQueue::EType::Semaphore, SemaphoreHandle);
+		SemaphoreHandle = VK_NULL_HANDLE;
 	}
 
 	VkSemaphore GetHandle() const
@@ -224,8 +225,19 @@ public:
 		return Texture == (FVulkanTexture2D*)RTInfo.DepthStencilRenderTarget.Texture;
 	}
 
+	inline uint32 GetWidth() const
+	{
+		return Extents.width;
+	}
+
+	inline uint32 GetHeight() const
+	{
+		return Extents.height;
+	}
+
 private:
 	VkFramebuffer Framebuffer;
+	VkExtent2D Extents;
 
 	// We do not adjust RTInfo, since it used for hashing and is what the UE provides,
 	// it's up to VulkanRHI to handle this correctly.
@@ -399,6 +411,7 @@ DECLARE_CYCLE_STAT_EXTERN(TEXT("DrawPrim UP Prep Time"), STAT_VulkanUPPrepTime,
 DECLARE_CYCLE_STAT_EXTERN(TEXT("Uniform Buffer Creation Time"), STAT_VulkanUniformBufferCreateTime, STATGROUP_VulkanRHI, );
 DECLARE_CYCLE_STAT_EXTERN(TEXT("Apply DS Uniform Buffers"), STAT_VulkanApplyDSUniformBuffers, STATGROUP_VulkanRHI, );
 DECLARE_CYCLE_STAT_EXTERN(TEXT("SRV Update Time"), STAT_VulkanSRVUpdateTime, STATGROUP_VulkanRHI, );
+DECLARE_CYCLE_STAT_EXTERN(TEXT("Deletion Queue"), STAT_VulkanDeletionQueue, STATGROUP_VulkanRHI, );
 #if VULKAN_ENABLE_AGGRESSIVE_STATS
 DECLARE_CYCLE_STAT_EXTERN(TEXT("Apply DS Shader Resources"), STAT_VulkanApplyDSResources, STATGROUP_VulkanRHI, );
 DECLARE_CYCLE_STAT_EXTERN(TEXT("Update DescriptorSets"), STAT_VulkanUpdateDescriptorSets, STATGROUP_VulkanRHI, );
@@ -409,7 +422,6 @@ DECLARE_CYCLE_STAT_EXTERN(TEXT("Set unif Buffer"), STAT_VulkanSetUniformBufferTi
 DECLARE_CYCLE_STAT_EXTERN(TEXT("VkUpdate DS"), STAT_VulkanVkUpdateDS, STATGROUP_VulkanRHI, );
 DECLARE_CYCLE_STAT_EXTERN(TEXT("Clear Dirty DS State"), STAT_VulkanClearDirtyDSState, STATGROUP_VulkanRHI, );
 DECLARE_CYCLE_STAT_EXTERN(TEXT("Bind Vertex Streams"), STAT_VulkanBindVertexStreamsTime, STATGROUP_VulkanRHI, );
-
 #endif
 
 namespace VulkanRHI
@@ -427,6 +439,39 @@ namespace VulkanRHI
 extern TAutoConsoleVariable<int32> GValidationCvar;
 #endif
 
+static inline VkAttachmentLoadOp RenderTargetLoadActionToVulkan(ERenderTargetLoadAction InLoadAction)
+{
+	VkAttachmentLoadOp OutLoadAction = VK_ATTACHMENT_LOAD_OP_MAX_ENUM;
+
+	switch (InLoadAction)
+	{
+	case ERenderTargetLoadAction::ELoad:		OutLoadAction = VK_ATTACHMENT_LOAD_OP_LOAD;			break;
+	case ERenderTargetLoadAction::EClear:		OutLoadAction = VK_ATTACHMENT_LOAD_OP_CLEAR;		break;
+	case ERenderTargetLoadAction::ENoAction:	OutLoadAction = VK_ATTACHMENT_LOAD_OP_DONT_CARE;	break;
+	default:																						break;
+	}
+
+	// Check for missing translation
+	check(OutLoadAction != VK_ATTACHMENT_LOAD_OP_MAX_ENUM);
+	return OutLoadAction;
+}
+
+static inline VkAttachmentStoreOp RenderTargetStoreActionToVulkan(ERenderTargetStoreAction InStoreAction)
+{
+	VkAttachmentStoreOp OutStoreAction = VK_ATTACHMENT_STORE_OP_MAX_ENUM;
+
+	switch (InStoreAction)
+	{
+	case ERenderTargetStoreAction::EStore:		OutStoreAction = VK_ATTACHMENT_STORE_OP_STORE;		break;
+	case ERenderTargetStoreAction::ENoAction:	OutStoreAction = VK_ATTACHMENT_STORE_OP_DONT_CARE;	break;
+	default:																						break;
+	}
+
+	// Check for missing translation
+	check(OutStoreAction != VK_ATTACHMENT_STORE_OP_MAX_ENUM);
+	return OutStoreAction;
+}
+
 
 #if 0
 namespace FRCLog
diff --git a/Engine/Source/Runtime/VulkanRHI/Private/VulkanRenderTarget.cpp b/Engine/Source/Runtime/VulkanRHI/Private/VulkanRenderTarget.cpp
--- a/Engine/Source/Runtime/VulkanRHI/Private/VulkanRenderTarget.cpp
+++ b/Engine/Source/Runtime/VulkanRHI/Private/VulkanRenderTarget.cpp
@@ -213,3 +213,187 @@ void FVulkanCommandListContext::RHITransitionResources(EResourceTransitionAccess
 		ensure(0);
 	}
 }
+
+// Need a separate struct so we can memzero/remove dependencies on reference counts
+struct FRenderTargetLayoutHashableStruct
+{
+	// Depth goes in the last slot
+	FRHITexture* Texture[MaxSimultaneousRenderTargets + 1];
+	uint32 MipIndex[MaxSimultaneousRenderTargets];
+	uint32 ArraySliceIndex[MaxSimultaneousRenderTargets];
+	ERenderTargetLoadAction LoadAction[MaxSimultaneousRenderTargets];
+	ERenderTargetStoreAction StoreAction[MaxSimultaneousRenderTargets];
+
+	ERenderTargetLoadAction		DepthLoadAction;
+	ERenderTargetStoreAction	DepthStoreAction;
+	ERenderTargetLoadAction		StencilLoadAction;
+	ERenderTargetStoreAction	StencilStoreAction;
+	FExclusiveDepthStencil		DepthStencilAccess;
+
+	// Fill this outside Set() function
+	VkExtent2D					Extents;
+
+	bool bClearDepth;
+	bool bClearStencil;
+	bool bClearColor;
+
+	void Set(const FRHISetRenderTargetsInfo& RTInfo)
+	{
+		FMemory::Memzero(*this);
+		for (int32 Index = 0; Index < RTInfo.NumColorRenderTargets; ++Index)
+		{
+			Texture[Index] = RTInfo.ColorRenderTarget[Index].Texture;
+			MipIndex[Index] = RTInfo.ColorRenderTarget[Index].MipIndex;
+			ArraySliceIndex[Index] = RTInfo.ColorRenderTarget[Index].ArraySliceIndex;
+			LoadAction[Index] = RTInfo.ColorRenderTarget[Index].LoadAction;
+			StoreAction[Index] = RTInfo.ColorRenderTarget[Index].StoreAction;
+		}
+
+		Texture[MaxSimultaneousRenderTargets] = RTInfo.DepthStencilRenderTarget.Texture;
+		DepthLoadAction = RTInfo.DepthStencilRenderTarget.DepthLoadAction;
+		DepthStoreAction = RTInfo.DepthStencilRenderTarget.DepthStoreAction;
+		StencilLoadAction = RTInfo.DepthStencilRenderTarget.StencilLoadAction;
+		StencilStoreAction = RTInfo.DepthStencilRenderTarget.GetStencilStoreAction();
+		DepthStencilAccess = RTInfo.DepthStencilRenderTarget.GetDepthStencilAccess();
+
+		bClearDepth = RTInfo.bClearDepth;
+		bClearStencil = RTInfo.bClearStencil;
+		bClearColor = RTInfo.bClearColor;
+	}
+};
+
+FVulkanRenderTargetLayout::FVulkanRenderTargetLayout(const FRHISetRenderTargetsInfo& RTInfo)
+	: NumAttachments(0)
+	, NumColorAttachments(0)
+	, bHasDepthStencil(false)
+	, bHasResolveAttachments(false)
+	, Hash(0)
+{
+	FMemory::Memzero(ColorReferences);
+	FMemory::Memzero(ResolveReferences);
+	FMemory::Memzero(DepthStencilReference);
+	FMemory::Memzero(Desc);
+	FMemory::Memzero(Extent);
+
+	bool bSetExtent = false;
+		
+	for (int32 Index = 0; Index < RTInfo.NumColorRenderTargets; ++Index)
+	{
+		const FRHIRenderTargetView& RTView = RTInfo.ColorRenderTarget[Index];
+		if (RTView.Texture)
+		{
+		    FVulkanTextureBase* Texture = FVulkanTextureBase::Cast(RTView.Texture);
+		    check(Texture);
+    
+		    if (!bSetExtent)
+		    {
+			    bSetExtent = true;
+			    Extent.Extent3D.width = Texture->Surface.Width;
+			    Extent.Extent3D.height = Texture->Surface.Height;
+			    Extent.Extent3D.depth = 1;
+		    }
+    
+		    uint32 NumSamples = RTView.Texture->GetNumSamples();
+	    
+		    VkAttachmentDescription& CurrDesc = Desc[NumAttachments];
+		    
+		    //@TODO: Check this, it should be a power-of-two. Might be a VulkanConvert helper function.
+		    CurrDesc.samples = static_cast<VkSampleCountFlagBits>(NumSamples);
+		    CurrDesc.format = (VkFormat)GPixelFormats[RTView.Texture->GetFormat()].PlatformFormat;
+		    CurrDesc.loadOp = RenderTargetLoadActionToVulkan(RTView.LoadAction);
+		    CurrDesc.storeOp = RenderTargetStoreActionToVulkan(RTView.StoreAction);
+		    CurrDesc.stencilLoadOp = VK_ATTACHMENT_LOAD_OP_DONT_CARE;
+#if !VULKAN_USE_MSAA_RESOLVE_ATTACHMENTS
+		    CurrDesc.stencilStoreOp = VK_ATTACHMENT_STORE_OP_DONT_CARE;
+#endif
+		    CurrDesc.initialLayout = VK_IMAGE_LAYOUT_GENERAL;
+		    CurrDesc.finalLayout = VK_IMAGE_LAYOUT_GENERAL;
+    
+		    ColorReferences[NumColorAttachments].attachment = NumAttachments;
+		    ColorReferences[NumColorAttachments].layout = VK_IMAGE_LAYOUT_GENERAL;
+		    NumAttachments++;
+    
+#if VULKAN_USE_MSAA_RESOLVE_ATTACHMENTS
+			if (NumSamples == 1)
+			{
+			    CurrDesc.storeOp = RenderTargetStoreActionToVulkan(RTView.StoreAction);
+    
+			    ResolveReferences[NumColorAttachments].attachment = VK_ATTACHMENT_UNUSED;
+			    ResolveReferences[NumColorAttachments].layout = VK_IMAGE_LAYOUT_UNDEFINED;
+		    }
+			else
+		    {
+			    // discard MSAA color target after resolve to 1x surface
+			    CurrDesc.storeOp = VK_ATTACHMENT_STORE_OP_DONT_CARE;
+    
+			    // Resolve attachment
+			    VkAttachmentDescription& ResolveDesc = Desc[NumAttachments];
+			    ResolveDesc.samples = VK_SAMPLE_COUNT_1_BIT;
+			    ResolveDesc.format = (VkFormat)GPixelFormats[RTView.Texture->GetFormat()].PlatformFormat;
+			    ResolveDesc.loadOp = VK_ATTACHMENT_LOAD_OP_DONT_CARE;
+			    ResolveDesc.storeOp = RenderTargetStoreActionToVulkan(RTView.StoreAction);
+			    ResolveDesc.stencilLoadOp = VK_ATTACHMENT_LOAD_OP_DONT_CARE;
+			    ResolveDesc.stencilStoreOp = VK_ATTACHMENT_STORE_OP_DONT_CARE;
+			    ResolveDesc.initialLayout = VK_IMAGE_LAYOUT_COLOR_ATTACHMENT_OPTIMAL;
+			    ResolveDesc.finalLayout = VK_IMAGE_LAYOUT_COLOR_ATTACHMENT_OPTIMAL;
+    
+			    ResolveReferences[NumColorAttachments].attachment = NumAttachments;
+			    ResolveReferences[NumColorAttachments].layout = VK_IMAGE_LAYOUT_COLOR_ATTACHMENT_OPTIMAL;
+			    NumAttachments++;
+			    bHasResolveAttachments = true;
+		    }
+#endif
+	    
+		    NumColorAttachments++;
+		}
+	}
+
+	if (RTInfo.DepthStencilRenderTarget.Texture)
+	{
+		VkAttachmentDescription& CurrDesc = Desc[NumAttachments];
+		FMemory::Memzero(CurrDesc);
+		FVulkanTextureBase* Texture = FVulkanTextureBase::Cast(RTInfo.DepthStencilRenderTarget.Texture);
+		check(Texture);
+
+		//@TODO: Check this, it should be a power-of-two. Might be a VulkanConvert helper function.
+		CurrDesc.samples = static_cast<VkSampleCountFlagBits>(RTInfo.DepthStencilRenderTarget.Texture->GetNumSamples());
+		CurrDesc.format = (VkFormat)GPixelFormats[RTInfo.DepthStencilRenderTarget.Texture->GetFormat()].PlatformFormat;
+		CurrDesc.loadOp = RenderTargetLoadActionToVulkan(RTInfo.DepthStencilRenderTarget.DepthLoadAction);
+		CurrDesc.stencilLoadOp = RenderTargetLoadActionToVulkan(RTInfo.DepthStencilRenderTarget.StencilLoadAction);
+		if (CurrDesc.samples == VK_SAMPLE_COUNT_1_BIT)
+		{
+			CurrDesc.storeOp = RenderTargetStoreActionToVulkan(RTInfo.DepthStencilRenderTarget.DepthStoreAction);
+			CurrDesc.stencilStoreOp = RenderTargetStoreActionToVulkan(RTInfo.DepthStencilRenderTarget.GetStencilStoreAction());
+		}
+		else
+		{
+			// Never want to store MSAA depth/stencil
+			CurrDesc.storeOp = VK_ATTACHMENT_STORE_OP_DONT_CARE;
+			CurrDesc.stencilStoreOp = VK_ATTACHMENT_STORE_OP_DONT_CARE;
+		}
+		CurrDesc.initialLayout = VK_IMAGE_LAYOUT_DEPTH_STENCIL_ATTACHMENT_OPTIMAL;
+		CurrDesc.finalLayout = VK_IMAGE_LAYOUT_DEPTH_STENCIL_ATTACHMENT_OPTIMAL;
+		
+		DepthStencilReference.attachment = NumAttachments;
+		DepthStencilReference.layout = VK_IMAGE_LAYOUT_DEPTH_STENCIL_ATTACHMENT_OPTIMAL;
+
+		bHasDepthStencil = true;
+		NumAttachments++;
+
+		if (!bSetExtent)
+		{
+			bSetExtent = true;
+			Extent.Extent3D.width = Texture->Surface.Width;
+			Extent.Extent3D.height = Texture->Surface.Height;
+			Extent.Extent3D.depth = 1;
+		}
+	}
+
+	check(bSetExtent);
+
+	FRenderTargetLayoutHashableStruct RTHash;
+	FMemory::Memzero(RTHash);
+	RTHash.Set(RTInfo);
+	RTHash.Extents = Extent.Extent2D;
+	Hash = FCrc::MemCrc32(&RTHash, sizeof(RTHash));
+}
diff --git a/Engine/Source/Runtime/VulkanRHI/Private/VulkanShaders.cpp b/Engine/Source/Runtime/VulkanRHI/Private/VulkanShaders.cpp
--- a/Engine/Source/Runtime/VulkanRHI/Private/VulkanShaders.cpp
+++ b/Engine/Source/Runtime/VulkanRHI/Private/VulkanShaders.cpp
@@ -118,7 +118,7 @@ FVulkanShader::~FVulkanShader()
 
 	if (ShaderModule != VK_NULL_HANDLE)
 	{
-		VulkanRHI::vkDestroyShaderModule(Device->GetInstanceHandle(), ShaderModule, nullptr);
+		Device->GetDeferredDeletionQueue().EnqueueResource(VulkanRHI::FDeferredDeletionQueue::EType::ShaderModule, ShaderModule);
 		ShaderModule = VK_NULL_HANDLE;
 	}
 }
@@ -343,7 +343,8 @@ FVulkanBoundShaderState::~FVulkanBoundShaderState()
 
 	if (PipelineLayout != VK_NULL_HANDLE)
 	{
-		vkDestroyPipelineLayout(Device->GetInstanceHandle(), PipelineLayout, nullptr);
+		Device->GetDeferredDeletionQueue().EnqueueResource(VulkanRHI::FDeferredDeletionQueue::EType::PipelineLayout, PipelineLayout);
+		PipelineLayout = VK_NULL_HANDLE;
 	}
 
 	DEC_DWORD_STAT(STAT_VulkanNumBoundShaderState);
@@ -436,11 +437,21 @@ const VkPipelineLayout& FVulkanBoundShaderState::GetPipelineLayout() const
 	CreateInfo.pushConstantRangeCount = 0;
 	CreateInfo.pPushConstantRanges = nullptr;
 
-	VERIFYVULKANRESULT(vkCreatePipelineLayout(Device->GetInstanceHandle(), &CreateInfo, nullptr, &PipelineLayout));
+	VERIFYVULKANRESULT(VulkanRHI::vkCreatePipelineLayout(Device->GetInstanceHandle(), &CreateInfo, nullptr, &PipelineLayout));
 
 	return PipelineLayout;
 }
 
+void FVulkanBoundShaderState::BindPipeline(VkCommandBuffer CmdBuffer, VkPipeline NewPipeline)
+{
+	if (LastBoundPipeline != NewPipeline)
+	{
+		//#todo-rco: Compute
+		VulkanRHI::vkCmdBindPipeline(CmdBuffer, VK_PIPELINE_BIND_POINT_GRAPHICS, NewPipeline);
+		LastBoundPipeline = NewPipeline;
+	}
+}
+
 void FVulkanBoundShaderState::InitGlobalUniforms(EShaderFrequency Stage, const FVulkanShaderSerializedBindings& SerializedBindings)
 {
 	auto FindPackedArraySize = [&](int32 PackedTypeIndex) -> uint16
@@ -690,7 +701,7 @@ void FVulkanBoundShaderState::InternalBindVertexStreams(FVulkanCmdBuffer* Cmd, c
 		// Incorrect:	1, 0, 2, 3
 		// Incorrect:	0, 2, 3, 5
 		// Reordering and creation of stream binding index is done in "GenerateVertexInputStateInfo()"
-		vkCmdBindVertexBuffers(Cmd->GetHandle(), 0, Tmp.VertexBuffers.Num(), Tmp.VertexBuffers.GetData(), Tmp.VertexOffsets.GetData());
+		VulkanRHI::vkCmdBindVertexBuffers(Cmd->GetHandle(), 0, Tmp.VertexBuffers.Num(), Tmp.VertexBuffers.GetData(), Tmp.VertexOffsets.GetData());
 	}
 }
 
@@ -1202,7 +1213,7 @@ void FVulkanBoundShaderState::UpdateDescriptorSets(FVulkanCommandListContext* Cm
 #if VULKAN_ENABLE_AGGRESSIVE_STATS
 		SCOPE_CYCLE_COUNTER(STAT_VulkanVkUpdateDS);
 #endif
-		vkUpdateDescriptorSets(Device->GetInstanceHandle(), DescriptorWrites.Num(), DescriptorWrites.GetData(), 0, nullptr);
+		VulkanRHI::vkUpdateDescriptorSets(Device->GetInstanceHandle(), DescriptorWrites.Num(), DescriptorWrites.GetData(), 0, nullptr);
 	}
 
 	{
diff --git a/Engine/Source/Runtime/VulkanRHI/Private/VulkanState.cpp b/Engine/Source/Runtime/VulkanRHI/Private/VulkanState.cpp
--- a/Engine/Source/Runtime/VulkanRHI/Private/VulkanState.cpp
+++ b/Engine/Source/Runtime/VulkanRHI/Private/VulkanState.cpp
@@ -266,12 +266,12 @@ FVulkanSamplerState::FVulkanSamplerState(const FSamplerStateInitializerRHI& Init
 	SamplerInfo.maxLod = Initializer.MaxMipLevel;
 	SamplerInfo.borderColor = Initializer.BorderColor == 0 ? VK_BORDER_COLOR_FLOAT_TRANSPARENT_BLACK : VK_BORDER_COLOR_FLOAT_OPAQUE_WHITE;
 
-	VERIFYVULKANRESULT(vkCreateSampler(Device.GetInstanceHandle(), &SamplerInfo, nullptr, &Sampler));
+	VERIFYVULKANRESULT(VulkanRHI::vkCreateSampler(Device.GetInstanceHandle(), &SamplerInfo, nullptr, &Sampler));
 }
 
 FVulkanSamplerState::~FVulkanSamplerState()
 {
-	vkDestroySampler(Device.GetInstanceHandle(), Sampler, nullptr);
+	Device.GetDeferredDeletionQueue().EnqueueResource(VulkanRHI::FDeferredDeletionQueue::EType::Sampler, Sampler);
 	Sampler = VK_NULL_HANDLE;
 }
 
diff --git a/Engine/Source/Runtime/VulkanRHI/Private/VulkanSwapChain.cpp b/Engine/Source/Runtime/VulkanRHI/Private/VulkanSwapChain.cpp
--- a/Engine/Source/Runtime/VulkanRHI/Private/VulkanSwapChain.cpp
+++ b/Engine/Source/Runtime/VulkanRHI/Private/VulkanSwapChain.cpp
@@ -15,39 +15,6 @@ FVulkanSwapChain::FVulkanSwapChain(VkInstance Instance, FVulkanDevice& InDevice,
 	, CurrentImageIndex(-1)
 	, SemaphoreIndex(0)
 {
-#define FETCH_KHR_FPN(FunctionPointerName)\
-	{\
-		FunctionPointerName = (PFN_vk##FunctionPointerName)vkGetDeviceProcAddr(Device.GetInstanceHandle(), "vk"#FunctionPointerName);\
-		check(FunctionPointerName);\
-	}
-
-#if PLATFORM_ANDROID || PLATFORM_WINDOWS
-	#define FETCH_KHR_INSTANCE_FPN(FunctionPointerName) \
-	{ \
-		FunctionPointerName = (PFN_vk##FunctionPointerName)vkGetInstanceProcAddr(Instance, "vk"#FunctionPointerName); \
-		check(FunctionPointerName); \
-	}
-#else
-	#define FETCH_KHR_INSTANCE_FPN(FunctionPointerName) FETCH_KHR_FPN(FunctionPointerName)
-#endif
-
-	#ifdef _MSC_VER
-		#pragma warning(push)
-		#pragma warning(disable:4191)
-	#endif
-
-	FETCH_KHR_FPN(CreateSwapchainKHR);
-	FETCH_KHR_FPN(DestroySwapchainKHR);
-	FETCH_KHR_FPN(GetSwapchainImagesKHR);
-	FETCH_KHR_FPN(QueuePresentKHR);
-	FETCH_KHR_FPN(AcquireNextImageKHR);
-
-	#ifdef _MSC_VER
-		#pragma warning(pop)
-	#endif
-
-	#undef FETCH_KHR_FPN
-
 #if PLATFORM_WINDOWS
 	VkWin32SurfaceCreateInfoKHR SurfaceCreateInfo;
 	FMemory::Memzero(SurfaceCreateInfo);
@@ -188,13 +155,13 @@ FVulkanSwapChain::FVulkanSwapChain(VkInstance Instance, FVulkanDevice& InDevice,
 
 	*InOutDesiredNumBackBuffers = DesiredNumBuffers;
 
-	VERIFYVULKANRESULT_EXPANDED(vkCreateSwapchainKHR(Device.GetInstanceHandle(), &SwapChainInfo, nullptr, &SwapChain));
+	VERIFYVULKANRESULT_EXPANDED(VulkanRHI::vkCreateSwapchainKHR(Device.GetInstanceHandle(), &SwapChainInfo, nullptr, &SwapChain));
 
 	uint32 NumSwapChainImages;
-	VERIFYVULKANRESULT_EXPANDED(GetSwapchainImagesKHR(Device.GetInstanceHandle(), SwapChain, &NumSwapChainImages, nullptr));
+	VERIFYVULKANRESULT_EXPANDED(VulkanRHI::vkGetSwapchainImagesKHR(Device.GetInstanceHandle(), SwapChain, &NumSwapChainImages, nullptr));
 
 	OutImages.AddUninitialized(NumSwapChainImages);
-	VERIFYVULKANRESULT_EXPANDED(GetSwapchainImagesKHR(Device.GetInstanceHandle(), SwapChain, &NumSwapChainImages, OutImages.GetData()));
+	VERIFYVULKANRESULT_EXPANDED(VulkanRHI::vkGetSwapchainImagesKHR(Device.GetInstanceHandle(), SwapChain, &NumSwapChainImages, OutImages.GetData()));
 
 	ImageAcquiredSemaphore.AddUninitialized(DesiredNumBuffers);
 	for (uint32 BufferIndex = 0; BufferIndex < DesiredNumBuffers; ++BufferIndex)
@@ -205,7 +172,7 @@ FVulkanSwapChain::FVulkanSwapChain(VkInstance Instance, FVulkanDevice& InDevice,
 
 void FVulkanSwapChain::Destroy()
 {
-	vkDestroySwapchainKHR(Device.GetInstanceHandle(), SwapChain, nullptr);
+	VulkanRHI::vkDestroySwapchainKHR(Device.GetInstanceHandle(), SwapChain, nullptr);
 	SwapChain = VK_NULL_HANDLE;
 
 	//#todo-rco: Enqueue for deletion as we first need to destroy the cmd buffers and queues otherwise validation fails
@@ -223,7 +190,7 @@ int32 FVulkanSwapChain::AcquireImageIndex(FVulkanSemaphore** OutSemaphore)
 	uint32 ImageIndex = 0;
 	SemaphoreIndex = (SemaphoreIndex + 1) % ImageAcquiredSemaphore.Num();
 	*OutSemaphore = ImageAcquiredSemaphore[SemaphoreIndex];
-	VkResult Result = AcquireNextImageKHR(
+	VkResult Result = VulkanRHI::vkAcquireNextImageKHR(
 		Device.GetInstanceHandle(),
 		SwapChain,
 		UINT64_MAX,
@@ -233,7 +200,6 @@ int32 FVulkanSwapChain::AcquireImageIndex(FVulkanSemaphore** OutSemaphore)
 	checkf(Result == VK_SUCCESS || Result == VK_SUBOPTIMAL_KHR, TEXT("AcquireNextImageKHR failed Result = %d"), int32(Result));
 	CurrentImageIndex = (int32)ImageIndex;
 	check(CurrentImageIndex == ImageIndex);
-	//FRCLog::Printf(FString::Printf(TEXT("FVulkanSwapChain::AcquireImageIndex() => %d\n"), CurrentImageIndex));
 	return CurrentImageIndex;
 }
 
@@ -255,6 +221,7 @@ bool FVulkanSwapChain::Present(FVulkanQueue* Queue, FVulkanSemaphore* BackBuffer
 	Info.pSwapchains = &SwapChain;
 	Info.pImageIndices = (uint32*)&CurrentImageIndex;
 
-	VERIFYVULKANRESULT(QueuePresentKHR(Queue->GetHandle(), &Info));
+	VERIFYVULKANRESULT(VulkanRHI::vkQueuePresentKHR(Queue->GetHandle(), &Info));
+
 	return true;
 }
diff --git a/Engine/Source/Runtime/VulkanRHI/Private/VulkanSwapChain.h b/Engine/Source/Runtime/VulkanRHI/Private/VulkanSwapChain.h
--- a/Engine/Source/Runtime/VulkanRHI/Private/VulkanSwapChain.h
+++ b/Engine/Source/Runtime/VulkanRHI/Private/VulkanSwapChain.h
@@ -22,12 +22,6 @@ protected:
 	VkSwapchainKHR SwapChain;
 	FVulkanDevice& Device;
 
-	PFN_vkCreateSwapchainKHR CreateSwapchainKHR;
-	PFN_vkDestroySwapchainKHR DestroySwapchainKHR;
-	PFN_vkGetSwapchainImagesKHR GetSwapchainImagesKHR;
-	PFN_vkQueuePresentKHR QueuePresentKHR;
-	PFN_vkAcquireNextImageKHR AcquireNextImageKHR;
-
 	VkSurfaceKHR Surface;
 
 	int32 CurrentImageIndex;
