Copying //UE4/Dev-Rendering to //UE4/Dev-Main (Source: //UE4/Dev-Rendering @ 3045398)

diff --git a/Engine/Source/Runtime/VulkanRHI/Private/VulkanTexture.cpp b/Engine/Source/Runtime/VulkanRHI/Private/VulkanTexture.cpp
--- a/Engine/Source/Runtime/VulkanRHI/Private/VulkanTexture.cpp
+++ b/Engine/Source/Runtime/VulkanRHI/Private/VulkanTexture.cpp
@@ -242,12 +242,12 @@ VkImage FVulkanSurface::CreateImage(
 #endif
 
 	VkImage Image;
-	VERIFYVULKANRESULT(vkCreateImage(InDevice.GetInstanceHandle(), &ImageCreateInfo, nullptr, &Image));
+	VERIFYVULKANRESULT(VulkanRHI::vkCreateImage(InDevice.GetInstanceHandle(), &ImageCreateInfo, nullptr, &Image));
 
 	// Fetch image size
 	VkMemoryRequirements tmpMemReq;
 	VkMemoryRequirements& MemReq = OutMemoryRequirements ? *OutMemoryRequirements : tmpMemReq;
-	vkGetImageMemoryRequirements(InDevice.GetInstanceHandle(), Image, &MemReq);
+	VulkanRHI::vkGetImageMemoryRequirements(InDevice.GetInstanceHandle(), Image, &MemReq);
 
 	return Image;
 }
@@ -406,7 +406,7 @@ void FVulkanSurface::Destroy()
 
 		if (Image != VK_NULL_HANDLE)
 		{
-			vkDestroyImage(Device->GetInstanceHandle(), Image, nullptr);
+			Device->GetDeferredDeletionQueue().EnqueueResource(VulkanRHI::FDeferredDeletionQueue::EType::Image, Image);
 			Image = VK_NULL_HANDLE;
 			ImageLayout = VK_IMAGE_LAYOUT_UNDEFINED;
 		}
@@ -441,7 +441,7 @@ void* FVulkanSurface::Lock(uint32 MipIndex, uint32 ArrayIndex, EResourceLockMode
 	// Get buffer size
 	// Pitch can be only retrieved from linear textures.
 	VkSubresourceLayout SubResourceLayout;
-	vkGetImageSubresourceLayout(Device->GetInstanceHandle(), Image, &ImageSubResource, &SubResourceLayout);
+	VulkanRHI::vkGetImageSubresourceLayout(Device->GetInstanceHandle(), Image, &ImageSubResource, &SubResourceLayout);
 
 	// Set linear row-pitch
 	GetMipStride(MipIndex, DestStride);
@@ -471,7 +471,7 @@ void* FVulkanSurface::Lock(uint32 MipIndex, uint32 ArrayIndex, EResourceLockMode
 	Usage |= VK_BUFFER_USAGE_TRANSFER_SRC_BIT;
 
 	VkMemoryPropertyFlags Flags = 0;
-	Flags |= VK_MEMORY_PROPERTY_HOST_VISIBLE_BIT;
+	Flags |= VK_MEMORY_PROPERTY_HOST_VISIBLE_BIT | VK_MEMORY_PROPERTY_HOST_COHERENT;
 
 	LinearBuffer = new FVulkanBuffer(*Device, Bytes, Usage, Flags, false, __FILE__, __LINE__);
 
@@ -507,7 +507,7 @@ void FVulkanSurface::Unlock(uint32 MipIndex, uint32 ArrayIndex)
 
 	VkSubresourceLayout SubResourceLayout;
 	
-	vkGetImageSubresourceLayout(Device->GetInstanceHandle(), Image, &ImageSubResource, &SubResourceLayout);
+	VulkanRHI::vkGetImageSubresourceLayout(Device->GetInstanceHandle(), Image, &ImageSubResource, &SubResourceLayout);
 
 	VkBufferImageCopy Region;
 	FMemory::Memzero(Region);
@@ -537,7 +537,7 @@ uint32 FVulkanSurface::GetMemorySize()
 {
 	//#todo-rco: Slow!
 	VkMemoryRequirements MemReqs;
-	vkGetImageMemoryRequirements(Device->GetInstanceHandle(), Image, &MemReqs);
+	VulkanRHI::vkGetImageMemoryRequirements(Device->GetInstanceHandle(), Image, &MemReqs);
 
 	return MemReqs.size;
 }
@@ -640,7 +640,7 @@ void FVulkanSurface::Clear(const FClearValueBinding& ClearValueBinding, bool bTr
 		Color.float32[2] = ClearValueBinding.Value.Color[2];
 		Color.float32[3] = ClearValueBinding.Value.Color[3];
 		VulkanSetImageLayoutSimple(CmdBuffer->GetHandle(), Image, ImageLayout, VK_IMAGE_LAYOUT_GENERAL);
-		vkCmdClearColorImage(CmdBuffer->GetHandle(), Image, VK_IMAGE_LAYOUT_GENERAL, &Color, 1, &Range);
+		VulkanRHI::vkCmdClearColorImage(CmdBuffer->GetHandle(), Image, VK_IMAGE_LAYOUT_GENERAL, &Color, 1, &Range);
 		VulkanSetImageLayoutSimple(CmdBuffer->GetHandle(), Image, VK_IMAGE_LAYOUT_GENERAL, bTransitionToPresentable ? VK_IMAGE_LAYOUT_PRESENT_SRC_KHR : VK_IMAGE_LAYOUT_COLOR_ATTACHMENT_OPTIMAL);
 		ImageLayout = VK_IMAGE_LAYOUT_COLOR_ATTACHMENT_OPTIMAL;
 	}
@@ -653,7 +653,7 @@ void FVulkanSurface::Clear(const FClearValueBinding& ClearValueBinding, bool bTr
 		Value.depth = ClearValueBinding.Value.DSValue.Depth;
 		Value.stencil = ClearValueBinding.Value.DSValue.Stencil;
 		VulkanSetImageLayoutSimple(CmdBuffer->GetHandle(), Image, ImageLayout, VK_IMAGE_LAYOUT_GENERAL, Range.aspectMask);
-		vkCmdClearDepthStencilImage(CmdBuffer->GetHandle(), Image, VK_IMAGE_LAYOUT_GENERAL, &Value, 1, &Range);
+		VulkanRHI::vkCmdClearDepthStencilImage(CmdBuffer->GetHandle(), Image, VK_IMAGE_LAYOUT_GENERAL, &Value, 1, &Range);
 		VulkanSetImageLayoutSimple(CmdBuffer->GetHandle(), Image, VK_IMAGE_LAYOUT_GENERAL, VK_IMAGE_LAYOUT_DEPTH_STENCIL_ATTACHMENT_OPTIMAL, Range.aspectMask);
 		ImageLayout = VK_IMAGE_LAYOUT_DEPTH_STENCIL_ATTACHMENT_OPTIMAL;
 	}
@@ -827,7 +827,7 @@ void FVulkanDynamicRHI::RHIUnlockTexture2D(FTexture2DRHIParamRef TextureRHI,uint
 	Region.imageExtent.width = MipWidth;
 	Region.imageExtent.height = MipHeight;
 	Region.imageExtent.depth = 1;
-	vkCmdCopyBufferToImage(StagingCommandBuffer, StagingBuffer->GetHandle(), Texture->Surface.Image, VK_IMAGE_LAYOUT_TRANSFER_DST_OPTIMAL, 1, &Region);
+	VulkanRHI::vkCmdCopyBufferToImage(StagingCommandBuffer, StagingBuffer->GetHandle(), Texture->Surface.Image, VK_IMAGE_LAYOUT_TRANSFER_DST_OPTIMAL, 1, &Region);
 
 	VulkanSetImageLayout(StagingCommandBuffer, Texture->Surface.Image, VK_IMAGE_LAYOUT_TRANSFER_DST_OPTIMAL, VK_IMAGE_LAYOUT_SHADER_READ_ONLY_OPTIMAL, SubresourceRange);
 
@@ -903,7 +903,7 @@ void FVulkanDynamicRHI::RHIUpdateTexture3D(FTexture3DRHIParamRef TextureRHI, uin
 	Region.imageExtent.width = UpdateRegion.Width;
 	Region.imageExtent.height = UpdateRegion.Height;
 	Region.imageExtent.depth = UpdateRegion.Depth;
-	vkCmdCopyBufferToImage(Device->GetImmediateContext().GetCommandBufferManager()->GetUploadCmdBuffer()->GetHandle(), StagingBuffer->GetHandle(), Texture->Surface.Image, VK_IMAGE_LAYOUT_GENERAL, 1, &Region);
+	VulkanRHI::vkCmdCopyBufferToImage(Device->GetImmediateContext().GetCommandBufferManager()->GetUploadCmdBuffer()->GetHandle(), StagingBuffer->GetHandle(), Texture->Surface.Image, VK_IMAGE_LAYOUT_GENERAL, 1, &Region);
 
 	Device->GetStagingManager().ReleaseBuffer(nullptr, StagingBuffer);
 }
@@ -925,14 +925,14 @@ void FVulkanTextureView::Create(FVulkanDevice& Device, VkImage Image, VkImageVie
 	ViewInfo.subresourceRange = { VK_IMAGE_ASPECT_COLOR_BIT, 0, NumMips, 0, ArraySize*FaceCount };
 	ViewInfo.subresourceRange.aspectMask = AspectFlags;
 
-	VERIFYVULKANRESULT(vkCreateImageView(Device.GetInstanceHandle(), &ViewInfo, nullptr, &View));
+	VERIFYVULKANRESULT(VulkanRHI::vkCreateImageView(Device.GetInstanceHandle(), &ViewInfo, nullptr, &View));
 }
 
 void FVulkanTextureView::Destroy(FVulkanDevice& Device)
 {
 	if (View)
 	{
-		vkDestroyImageView(Device.GetInstanceHandle(), View, nullptr);
+		Device.GetDeferredDeletionQueue().EnqueueResource(VulkanRHI::FDeferredDeletionQueue::EType::ImageView, View);
 		View = VK_NULL_HANDLE;
 	}
 }
@@ -1034,7 +1034,7 @@ FVulkanTextureBase::FVulkanTextureBase(FVulkanDevice& Device, VkImageViewType Re
 	VulkanSetImageLayout(CmdBuffer->GetHandle(), Surface.Image, VK_IMAGE_LAYOUT_UNDEFINED, VK_IMAGE_LAYOUT_TRANSFER_DST_OPTIMAL, SubresourceRange);
 
 	// Copy buffer to image
-	vkCmdCopyBufferToImage(CmdBuffer->GetHandle(), StagingBuffer->GetHandle(), Surface.Image, VK_IMAGE_LAYOUT_TRANSFER_DST_OPTIMAL, 1, &Region);
+	VulkanRHI::vkCmdCopyBufferToImage(CmdBuffer->GetHandle(), StagingBuffer->GetHandle(), Surface.Image, VK_IMAGE_LAYOUT_TRANSFER_DST_OPTIMAL, 1, &Region);
 
 	VulkanSetImageLayout(CmdBuffer->GetHandle(), Surface.Image, VK_IMAGE_LAYOUT_TRANSFER_DST_OPTIMAL, VK_IMAGE_LAYOUT_SHADER_READ_ONLY_OPTIMAL, SubresourceRange);
 
@@ -1085,14 +1085,6 @@ FVulkanTexture2D::~FVulkanTexture2D()
 	{
 		Surface.Device->NotifyDeletedRenderTarget(this);
 	}
-
-	Destroy(*Surface.Device);
-}
-
-void FVulkanTexture2D::Destroy(FVulkanDevice& Device)
-{
-	View.Destroy(Device);
-	Surface.Destroy();
 }
 
 
@@ -1215,7 +1207,7 @@ void FVulkanDynamicRHI::RHIUnlockTextureCubeFace(FTextureCubeRHIParamRef Texture
 	Region.imageExtent.width = MipWidth;
 	Region.imageExtent.height = MipHeight;
 	Region.imageExtent.depth = 1;
-	vkCmdCopyBufferToImage(StagingCommandBuffer, StagingBuffer->GetHandle(), Texture->Surface.Image, VK_IMAGE_LAYOUT_TRANSFER_DST_OPTIMAL, 1, &Region);
+	VulkanRHI::vkCmdCopyBufferToImage(StagingCommandBuffer, StagingBuffer->GetHandle(), Texture->Surface.Image, VK_IMAGE_LAYOUT_TRANSFER_DST_OPTIMAL, 1, &Region);
 
 	VulkanSetImageLayout(StagingCommandBuffer, Texture->Surface.Image, VK_IMAGE_LAYOUT_TRANSFER_DST_OPTIMAL, VK_IMAGE_LAYOUT_SHADER_READ_ONLY_OPTIMAL, SubresourceRange);
 
@@ -1224,6 +1216,16 @@ void FVulkanDynamicRHI::RHIUnlockTextureCubeFace(FTextureCubeRHIParamRef Texture
 
 void FVulkanDynamicRHI::RHIBindDebugLabelName(FTextureRHIParamRef TextureRHI, const TCHAR* Name)
 {
+#if VULKAN_ENABLE_DUMP_LAYER
+	{
+		if (FRHITexture2D* Tex2d = TextureRHI->GetTexture2D())
+		{
+			FVulkanTexture2D* VTex2d = (FVulkanTexture2D*)Tex2d;
+			VulkanRHI::PrintfBegin(FString::Printf(TEXT("vkDebugMarkerSetObjectNameEXT(%p=%s)"), VTex2d->Surface.Image, Name));
+		}
+	}
+#endif
+
 #if VULKAN_ENABLE_DRAW_MARKERS
 	//if (Device->SupportsDebugMarkers())
 	/*
@@ -1296,7 +1298,7 @@ uint64 FVulkanDynamicRHI::RHICalcTexture2DPlatformSize(uint32 SizeX, uint32 Size
 		false, 0, NumMips, Flags,
 		&InternalFormat, &CreateInfo, &MemReq);
 
-	vkDestroyImage(Device->GetInstanceHandle(), TmpImage, nullptr);
+	VulkanRHI::vkDestroyImage(Device->GetInstanceHandle(), TmpImage, nullptr);
 
 	OutAlign = MemReq.alignment;
 
diff --git a/Engine/Source/Runtime/VulkanRHI/Private/VulkanUniformBuffer.cpp b/Engine/Source/Runtime/VulkanRHI/Private/VulkanUniformBuffer.cpp
--- a/Engine/Source/Runtime/VulkanRHI/Private/VulkanUniformBuffer.cpp
+++ b/Engine/Source/Runtime/VulkanRHI/Private/VulkanUniformBuffer.cpp
@@ -13,12 +13,12 @@ static TRefCountPtr<FVulkanBuffer> AllocateBufferFromPool(FVulkanDevice& Device,
 {
 	if (Usage == UniformBuffer_MultiFrame)
 	{
-		return new FVulkanBuffer(Device, ConstantBufferSize, VK_BUFFER_USAGE_UNIFORM_BUFFER_BIT, VK_MEMORY_PROPERTY_HOST_VISIBLE_BIT, false, __FILE__, __LINE__);
+		return new FVulkanBuffer(Device, ConstantBufferSize, VK_BUFFER_USAGE_UNIFORM_BUFFER_BIT, VK_MEMORY_PROPERTY_HOST_VISIBLE_BIT | VK_MEMORY_PROPERTY_HOST_COHERENT_BIT, false, __FILE__, __LINE__);
 	}
 
 	int32 BufferIndex = GFrameNumberRenderThread % NUM_SAFE_FRAMES;
 
-	GUBPool[BufferIndex].Add(new FVulkanBuffer(Device, ConstantBufferSize, VK_BUFFER_USAGE_UNIFORM_BUFFER_BIT, VK_MEMORY_PROPERTY_HOST_VISIBLE_BIT, false, __FILE__, __LINE__));
+	GUBPool[BufferIndex].Add(new FVulkanBuffer(Device, ConstantBufferSize, VK_BUFFER_USAGE_UNIFORM_BUFFER_BIT, VK_MEMORY_PROPERTY_HOST_VISIBLE_BIT | VK_MEMORY_PROPERTY_HOST_COHERENT_BIT, false, __FILE__, __LINE__));
 	return GUBPool[BufferIndex].Last();
 }
 
@@ -123,7 +123,7 @@ FUniformBufferRHIRef FVulkanDynamicRHI::RHICreateUniformBuffer(const void* Conte
 }
 
 FVulkanPooledUniformBuffer::FVulkanPooledUniformBuffer(FVulkanDevice& InDevice, uint32 InSize)
-    : Buffer(InDevice, InSize, VK_BUFFER_USAGE_UNIFORM_BUFFER_BIT, VK_MEMORY_PROPERTY_HOST_VISIBLE_BIT, false, __FILE__, __LINE__)
+    : Buffer(InDevice, InSize, VK_BUFFER_USAGE_UNIFORM_BUFFER_BIT, VK_MEMORY_PROPERTY_HOST_VISIBLE_BIT | VK_MEMORY_PROPERTY_HOST_COHERENT_BIT, false, __FILE__, __LINE__)
 {
 }
 
diff --git a/Engine/Source/Runtime/VulkanRHI/Private/VulkanUtil.cpp b/Engine/Source/Runtime/VulkanRHI/Private/VulkanUtil.cpp
--- a/Engine/Source/Runtime/VulkanRHI/Private/VulkanUtil.cpp
+++ b/Engine/Source/Runtime/VulkanRHI/Private/VulkanUtil.cpp
@@ -462,6 +462,7 @@ DEFINE_STAT(STAT_VulkanUPPrepTime);
 DEFINE_STAT(STAT_VulkanUniformBufferCreateTime);
 DEFINE_STAT(STAT_VulkanApplyDSUniformBuffers);
 DEFINE_STAT(STAT_VulkanSRVUpdateTime);
+DEFINE_STAT(STAT_VulkanDeletionQueue);
 #if VULKAN_ENABLE_AGGRESSIVE_STATS
 DEFINE_STAT(STAT_VulkanApplyDSResources);
 DEFINE_STAT(STAT_VulkanUpdateDescriptorSets);
diff --git a/Engine/Source/Runtime/VulkanRHI/Private/VulkanViewport.cpp b/Engine/Source/Runtime/VulkanRHI/Private/VulkanViewport.cpp
--- a/Engine/Source/Runtime/VulkanRHI/Private/VulkanViewport.cpp
+++ b/Engine/Source/Runtime/VulkanRHI/Private/VulkanViewport.cpp
@@ -38,6 +38,7 @@ FVulkanViewport::FVulkanViewport(FVulkanDynamicRHI* InRHI, FVulkanDevice* InDevi
 	, AcquiredImageIndex(-1)
 	, SwapChain(nullptr)
 	, WindowHandle(InWindowHandle)
+	, PresentCount(0)
 	, AcquiredSemaphore(nullptr)
 {
 	check(IsInGameThread());
@@ -255,14 +256,17 @@ FVulkanFramebuffer::FVulkanFramebuffer(FVulkanDevice& Device, const FRHISetRende
 	Info.width  = RTLayout.GetExtent3D().width;
 	Info.height = RTLayout.GetExtent3D().height;
 	Info.layers = 1;
-	VERIFYVULKANRESULT_EXPANDED(vkCreateFramebuffer(Device.GetInstanceHandle(), &Info, nullptr, &Framebuffer));
+	VERIFYVULKANRESULT_EXPANDED(VulkanRHI::vkCreateFramebuffer(Device.GetInstanceHandle(), &Info, nullptr, &Framebuffer));
 
 	RTInfo = InRTInfo;
+	Extents.width = Info.width;
+	Extents.height = Info.height;
 }
 
 void FVulkanFramebuffer::Destroy(FVulkanDevice& Device)
 {
-	vkDestroyFramebuffer(Device.GetInstanceHandle(), Framebuffer, nullptr);
+	Device.GetDeferredDeletionQueue().EnqueueResource(VulkanRHI::FDeferredDeletionQueue::EType::Framebuffer, Framebuffer);
+	Framebuffer = VK_NULL_HANDLE;
 }
 
 bool FVulkanFramebuffer::Matches(const FRHISetRenderTargetsInfo& InRTInfo) const
@@ -347,7 +351,7 @@ void FVulkanViewport::CreateSwapchain()
 			VkClearColorValue Color;
 			FMemory::Memzero(Color);
 			VulkanSetImageLayoutSimple(CmdBuffer->GetHandle(), Images[Index], VK_IMAGE_LAYOUT_UNDEFINED, VK_IMAGE_LAYOUT_GENERAL);
-			vkCmdClearColorImage(CmdBuffer->GetHandle(), Images[Index], VK_IMAGE_LAYOUT_GENERAL, &Color, 1, &Range);
+			VulkanRHI::vkCmdClearColorImage(CmdBuffer->GetHandle(), Images[Index], VK_IMAGE_LAYOUT_GENERAL, &Color, 1, &Range);
 			VulkanSetImageLayoutSimple(CmdBuffer->GetHandle(), Images[Index], VK_IMAGE_LAYOUT_GENERAL, VK_IMAGE_LAYOUT_COLOR_ATTACHMENT_OPTIMAL);
 		}
 	}
@@ -439,6 +443,7 @@ bool FVulkanViewport::Present(FVulkanCmdBuffer* CmdBuffer, FVulkanQueue* Queue,
 		Device->GetTimestampQueryPool(QueryNextFrameIndex)->WriteStartFrame(CmdBufferManager->GetActiveCmdBuffer()->GetHandle());
 	}
 #endif
+	++PresentCount;
 	return bResult;
 }
 
@@ -492,6 +497,7 @@ void FVulkanDynamicRHI::WriteEndFrameTimestamp(void* Data)
 	VulkanRHI::GManager.GPUProfilingData.EndFrameBeforeSubmit();
 }
 
+#if 0
 void FVulkanDynamicRHI::Present()
 {
 #if 0
@@ -597,6 +603,7 @@ void FVulkanDynamicRHI::Present()
 #endif
 	PresentCount++;
 }
+#endif
 
 FTexture2DRHIRef FVulkanDynamicRHI::RHIGetViewportBackBuffer(FViewportRHIParamRef ViewportRHI)
 {
diff --git a/Engine/Source/Runtime/VulkanRHI/Public/VulkanContext.h b/Engine/Source/Runtime/VulkanRHI/Public/VulkanContext.h
--- a/Engine/Source/Runtime/VulkanRHI/Public/VulkanContext.h
+++ b/Engine/Source/Runtime/VulkanRHI/Public/VulkanContext.h
@@ -159,4 +159,6 @@ protected:
 	FVulkanPendingState* PendingState;
 
 	void SubmitCurrentCommands();
+
+	void InternalClearMRT(FVulkanCmdBuffer* CmdBuffer, bool bClearColor, int32 NumClearColors, const FLinearColor* ColorArray, bool bClearDepth, float Depth, bool bClearStencil, uint32 Stencil, FIntRect ExcludeRect);
 };
diff --git a/Engine/Source/Runtime/VulkanRHI/Public/VulkanMemory.h b/Engine/Source/Runtime/VulkanRHI/Public/VulkanMemory.h
--- a/Engine/Source/Runtime/VulkanRHI/Public/VulkanMemory.h
+++ b/Engine/Source/Runtime/VulkanRHI/Public/VulkanMemory.h
@@ -729,7 +729,7 @@ namespace VulkanRHI
 		// Owner maintains lifetime
 		virtual ~FStagingBuffer();
 
-		void Destroy(VkDevice DeviceHandle);
+		void Destroy(FVulkanDevice* Device);
 
 		friend class FStagingManager;
 	};
@@ -870,12 +870,33 @@ namespace VulkanRHI
 		FDeferredDeletionQueue(FVulkanDevice* InDevice);
 		~FDeferredDeletionQueue();
 
-		void EnqueueResource(FVulkanCmdBuffer* CmdBuffer, FRefCount* Resource);
-		void ReleaseResources(/*bool bDeleteImmediately = false*/);
+		enum class EType
+		{
+			RenderPass,
+			Buffer,
+			BufferView,
+			Image,
+			ImageView,
+			Pipeline,
+			PipelineLayout,
+			Framebuffer,
+			DescriptorSetLayout,
+			Sampler,
+			Semaphore,
+			ShaderModule,
+		};
+
+		template <typename T>
+		inline void EnqueueResource(EType Type, T Handle)
+		{
+			EnqueueGenericResource(Type, (void*)Handle);
+		}
+
+		void ReleaseResources(bool bDeleteImmediately = false);
 
 		inline void Clear()
 		{
-			ReleaseResources(/*true*/);
+			ReleaseResources(true);
 		}
 /*
 		class FVulkanAsyncDeletionWorker : public FVulkanDeviceChild, FNonAbandonableTask
@@ -891,11 +912,17 @@ namespace VulkanRHI
 */
 	private:
 		//TQueue<FAsyncTask<FVulkanAsyncDeletionWorker>*> DeleteTasks;
+
+		void EnqueueGenericResource(EType Type, void* Handle);
+
 		struct FEntry
 		{
 			uint64 FenceCounter;
 			FVulkanCmdBuffer* CmdBuffer;
-			FRefCount* Resource;
+			void* Handle;
+			EType StructureType;
+
+			//FRefCount* Resource;
 		};
 		FCriticalSection CS;
 		TArray<FEntry> Entries;
diff --git a/Engine/Source/Runtime/VulkanRHI/Public/VulkanRHI.h b/Engine/Source/Runtime/VulkanRHI/Public/VulkanRHI.h
--- a/Engine/Source/Runtime/VulkanRHI/Public/VulkanRHI.h
+++ b/Engine/Source/Runtime/VulkanRHI/Public/VulkanRHI.h
@@ -175,7 +175,6 @@ protected:
 	void InitGPU(FVulkanDevice* Device);
 	void InitDevice(FVulkanDevice* Device);
 	
-	void Present();
 	friend class FVulkanCommandListContext;
 
 	friend class FVulkanViewport;
diff --git a/Engine/Source/Runtime/VulkanRHI/Public/VulkanResources.h b/Engine/Source/Runtime/VulkanRHI/Public/VulkanResources.h
--- a/Engine/Source/Runtime/VulkanRHI/Public/VulkanResources.h
+++ b/Engine/Source/Runtime/VulkanRHI/Public/VulkanResources.h
@@ -318,8 +318,6 @@ public:
 	FVulkanTexture2D(FVulkanDevice& Device, EPixelFormat Format, uint32 SizeX, uint32 SizeY, uint32 NumMips, uint32 NumSamples, uint32 UEFlags, const FRHIResourceCreateInfo& CreateInfo);
 	virtual ~FVulkanTexture2D();
 
-	void Destroy(FVulkanDevice& Device);
-
 	// IRefCountedObject interface.
 	virtual uint32 AddRef() const override final
 	{
@@ -972,19 +970,7 @@ public:
 		return *Shader;
 	}
 
-	inline void BindPipeline(VkCommandBuffer CmdBuffer, VkPipeline NewPipeline)
-	{
-		if (LastBoundPipeline != NewPipeline)
-		{
-			//#todo-rco: Compute
-#if	VULKAN_COMMANDWRAPPERS_ENABLE
-			vkCmdBindPipeline(CmdBuffer, VK_PIPELINE_BIND_POINT_GRAPHICS, NewPipeline);
-#else
-			VulkanRHI::vkCmdBindPipeline(CmdBuffer, VK_PIPELINE_BIND_POINT_GRAPHICS, NewPipeline);
-#endif
-			LastBoundPipeline = NewPipeline;
-		}
-	}
+	void BindPipeline(VkCommandBuffer CmdBuffer, VkPipeline NewPipeline);
 
 private:
 	void InitGlobalUniforms(EShaderFrequency Stage, const FVulkanShaderSerializedBindings& SerializedBindings);
diff --git a/Engine/Source/Runtime/Windows/D3D11RHI/Private/D3D11Commands.cpp b/Engine/Source/Runtime/Windows/D3D11RHI/Private/D3D11Commands.cpp
--- a/Engine/Source/Runtime/Windows/D3D11RHI/Private/D3D11Commands.cpp
+++ b/Engine/Source/Runtime/Windows/D3D11RHI/Private/D3D11Commands.cpp
@@ -2140,7 +2140,11 @@ void FD3D11DynamicRHI::RHIClearMRTImpl(bool bClearColor, int32 NumClearColors, c
 		{
 			for (int32 TargetIndex = 0; TargetIndex < BoundRenderTargets.GetNumActiveTargets(); TargetIndex++)
 			{				
-				Direct3DDeviceIMContext->ClearRenderTargetView(BoundRenderTargets.GetRenderTargetView(TargetIndex),(float*)&ClearColorArray[TargetIndex]);
+				ID3D11RenderTargetView* RenderTargetView = BoundRenderTargets.GetRenderTargetView(TargetIndex);
+				if (RenderTargetView != nullptr)
+				{
+					Direct3DDeviceIMContext->ClearRenderTargetView(RenderTargetView, (float*)&ClearColorArray[TargetIndex]);
+				}
 			}
 		}
 
@@ -2294,4 +2298,4 @@ void FD3D11DynamicRHI::RHITransitionResources(EResourceTransitionAccess Transiti
 	{
 		WriteFence->WriteFence();
 	}
-}
\ No newline at end of file
+}
diff --git a/Engine/Source/Runtime/Windows/D3D11RHI/Private/D3D11Util.cpp b/Engine/Source/Runtime/Windows/D3D11RHI/Private/D3D11Util.cpp
--- a/Engine/Source/Runtime/Windows/D3D11RHI/Private/D3D11Util.cpp
+++ b/Engine/Source/Runtime/Windows/D3D11RHI/Private/D3D11Util.cpp
@@ -324,9 +324,12 @@ FD3D11BoundRenderTargets::FD3D11BoundRenderTargets(ID3D11DeviceContext* InDevice
 		&RenderTargetViews[0],
 		&DepthStencilView
 		);
-	for (NumActiveTargets = 0; NumActiveTargets < MaxSimultaneousRenderTargets; ++NumActiveTargets)
+
+	// Find the last non-null rendertarget to determine the max 
+	// We traverse the array backwards, since they can be sparse
+	for (NumActiveTargets = MaxSimultaneousRenderTargets; NumActiveTargets > 0; --NumActiveTargets)
 	{
-		if (RenderTargetViews[NumActiveTargets] == NULL)
+		if (RenderTargetViews[NumActiveTargets-1] != NULL)
 		{
 			break;
 		}
@@ -339,7 +342,10 @@ FD3D11BoundRenderTargets::~FD3D11BoundRenderTargets()
 	// to make a corresponding call to Release.
 	for (int32 TargetIndex = 0; TargetIndex < NumActiveTargets; ++TargetIndex)
 	{
-		RenderTargetViews[TargetIndex]->Release();
+		if (RenderTargetViews[TargetIndex] != nullptr)
+		{
+			RenderTargetViews[TargetIndex]->Release();
+		}
 	}
 	if (DepthStencilView)
 	{
diff --git a/Engine/Source/Runtime/Windows/D3D12RHI/Private/D3D12Commands.cpp b/Engine/Source/Runtime/Windows/D3D12RHI/Private/D3D12Commands.cpp
--- a/Engine/Source/Runtime/Windows/D3D12RHI/Private/D3D12Commands.cpp
+++ b/Engine/Source/Runtime/Windows/D3D12RHI/Private/D3D12Commands.cpp
@@ -2059,15 +2059,22 @@ void FD3D12CommandContext::RHIClearMRTImpl(bool bClearColor, int32 NumClearColor
 	uint32 ClearFlags = 0;
 	if (ClearDSV)
 	{
-		if (bClearDepth)
+		if (bClearDepth && DepthStencilView->HasDepth())
 		{
 			ClearFlags |= D3D12_CLEAR_FLAG_DEPTH;
-			check(DepthStencilView->HasDepth());
 		}
-		if (bClearStencil)
+		else if (bClearDepth)
+		{
+			UE_LOG(LogD3D12RHI, Warning, TEXT("RHIClearMRTImpl: Asking to clear a DSV that does not store depth."));
+		}
+
+		if (bClearStencil && DepthStencilView->HasStencil())
 		{
 			ClearFlags |= D3D12_CLEAR_FLAG_STENCIL;
-			check(DepthStencilView->HasStencil());
+		}
+		else if (bClearStencil)
+		{
+			UE_LOG(LogD3D12RHI, Warning, TEXT("RHIClearMRTImpl: Asking to clear a DSV that does not store stencil."));
 		}
 
 		if (bClearDepth && (!DepthStencilView->HasStencil() || bClearStencil))
diff --git a/Engine/Source/Runtime/Windows/D3D12RHI/Private/D3D12UAV.cpp b/Engine/Source/Runtime/Windows/D3D12RHI/Private/D3D12UAV.cpp
--- a/Engine/Source/Runtime/Windows/D3D12RHI/Private/D3D12UAV.cpp
+++ b/Engine/Source/Runtime/Windows/D3D12RHI/Private/D3D12UAV.cpp
@@ -237,6 +237,7 @@ void FD3D12CommandContext::RHIClearUAV(FUnorderedAccessViewRHIParamRef Unordered
 	D3D12_GPU_DESCRIPTOR_HANDLE GPUHandle = StateCache.GetDescriptorCache()->GetCurrentViewHeap()->GetGPUSlotHandle(ReservedSlot);
 	GetParentDevice()->GetDevice()->CopyDescriptorsSimple(1, DestSlot, CPUHandle, D3D12_DESCRIPTOR_HEAP_TYPE_CBV_SRV_UAV);
 	FD3D12DynamicRHI::TransitionResource(CommandListHandle, UnorderedAccessView, D3D12_RESOURCE_STATE_UNORDERED_ACCESS);
+	CommandListHandle.FlushResourceBarriers();
 	numClears++;
 	CommandListHandle->ClearUnorderedAccessViewUint(GPUHandle, CPUHandle, UnorderedAccessView->GetResource()->GetResource(), Values, 0, nullptr);
 	CommandListHandle.UpdateResidency(UnorderedAccessView->GetResource());
