Copying //UE4/Dev-Rendering to //UE4/Dev-Main (Source: //UE4/Dev-Rendering @ 3054480)

diff --git a/Engine/Source/Runtime/VulkanRHI/Private/VulkanDebug.cpp b/Engine/Source/Runtime/VulkanRHI/Private/VulkanDebug.cpp
--- a/Engine/Source/Runtime/VulkanRHI/Private/VulkanDebug.cpp
+++ b/Engine/Source/Runtime/VulkanRHI/Private/VulkanDebug.cpp
@@ -198,12 +198,210 @@
 		TEXT("1 to enable dump layer, 0 to disable (default)")
 		);
 
+	static FString GetVkFormatString(VkFormat Format)
+	{
+		switch (Format)
+		{
+			// + 10 to skip "VK_FORMAT"
+#define VKSWITCHCASE(x)	case x: return TEXT(#x) + 10;
+			VKSWITCHCASE(VK_FORMAT_UNDEFINED)
+			VKSWITCHCASE(VK_FORMAT_R4G4_UNORM_PACK8)
+			VKSWITCHCASE(VK_FORMAT_R4G4B4A4_UNORM_PACK16)
+			VKSWITCHCASE(VK_FORMAT_B4G4R4A4_UNORM_PACK16)
+			VKSWITCHCASE(VK_FORMAT_R5G6B5_UNORM_PACK16)
+			VKSWITCHCASE(VK_FORMAT_B5G6R5_UNORM_PACK16)
+			VKSWITCHCASE(VK_FORMAT_R5G5B5A1_UNORM_PACK16)
+			VKSWITCHCASE(VK_FORMAT_B5G5R5A1_UNORM_PACK16)
+			VKSWITCHCASE(VK_FORMAT_A1R5G5B5_UNORM_PACK16)
+			VKSWITCHCASE(VK_FORMAT_R8_UNORM)
+			VKSWITCHCASE(VK_FORMAT_R8_SNORM)
+			VKSWITCHCASE(VK_FORMAT_R8_USCALED)
+			VKSWITCHCASE(VK_FORMAT_R8_SSCALED)
+			VKSWITCHCASE(VK_FORMAT_R8_UINT)
+			VKSWITCHCASE(VK_FORMAT_R8_SINT)
+			VKSWITCHCASE(VK_FORMAT_R8_SRGB)
+			VKSWITCHCASE(VK_FORMAT_R8G8_UNORM)
+			VKSWITCHCASE(VK_FORMAT_R8G8_SNORM)
+			VKSWITCHCASE(VK_FORMAT_R8G8_USCALED)
+			VKSWITCHCASE(VK_FORMAT_R8G8_SSCALED)
+			VKSWITCHCASE(VK_FORMAT_R8G8_UINT)
+			VKSWITCHCASE(VK_FORMAT_R8G8_SINT)
+			VKSWITCHCASE(VK_FORMAT_R8G8_SRGB)
+			VKSWITCHCASE(VK_FORMAT_R8G8B8_UNORM)
+			VKSWITCHCASE(VK_FORMAT_R8G8B8_SNORM)
+			VKSWITCHCASE(VK_FORMAT_R8G8B8_USCALED)
+			VKSWITCHCASE(VK_FORMAT_R8G8B8_SSCALED)
+			VKSWITCHCASE(VK_FORMAT_R8G8B8_UINT)
+			VKSWITCHCASE(VK_FORMAT_R8G8B8_SINT)
+			VKSWITCHCASE(VK_FORMAT_R8G8B8_SRGB)
+			VKSWITCHCASE(VK_FORMAT_B8G8R8_UNORM)
+			VKSWITCHCASE(VK_FORMAT_B8G8R8_SNORM)
+			VKSWITCHCASE(VK_FORMAT_B8G8R8_USCALED)
+			VKSWITCHCASE(VK_FORMAT_B8G8R8_SSCALED)
+			VKSWITCHCASE(VK_FORMAT_B8G8R8_UINT)
+			VKSWITCHCASE(VK_FORMAT_B8G8R8_SINT)
+			VKSWITCHCASE(VK_FORMAT_B8G8R8_SRGB)
+			VKSWITCHCASE(VK_FORMAT_R8G8B8A8_UNORM)
+			VKSWITCHCASE(VK_FORMAT_R8G8B8A8_SNORM)
+			VKSWITCHCASE(VK_FORMAT_R8G8B8A8_USCALED)
+			VKSWITCHCASE(VK_FORMAT_R8G8B8A8_SSCALED)
+			VKSWITCHCASE(VK_FORMAT_R8G8B8A8_UINT)
+			VKSWITCHCASE(VK_FORMAT_R8G8B8A8_SINT)
+			VKSWITCHCASE(VK_FORMAT_R8G8B8A8_SRGB)
+			VKSWITCHCASE(VK_FORMAT_B8G8R8A8_UNORM)
+			VKSWITCHCASE(VK_FORMAT_B8G8R8A8_SNORM)
+			VKSWITCHCASE(VK_FORMAT_B8G8R8A8_USCALED)
+			VKSWITCHCASE(VK_FORMAT_B8G8R8A8_SSCALED)
+			VKSWITCHCASE(VK_FORMAT_B8G8R8A8_UINT)
+			VKSWITCHCASE(VK_FORMAT_B8G8R8A8_SINT)
+			VKSWITCHCASE(VK_FORMAT_B8G8R8A8_SRGB)
+			VKSWITCHCASE(VK_FORMAT_A8B8G8R8_UNORM_PACK32)
+			VKSWITCHCASE(VK_FORMAT_A8B8G8R8_SNORM_PACK32)
+			VKSWITCHCASE(VK_FORMAT_A8B8G8R8_USCALED_PACK32)
+			VKSWITCHCASE(VK_FORMAT_A8B8G8R8_SSCALED_PACK32)
+			VKSWITCHCASE(VK_FORMAT_A8B8G8R8_UINT_PACK32)
+			VKSWITCHCASE(VK_FORMAT_A8B8G8R8_SINT_PACK32)
+			VKSWITCHCASE(VK_FORMAT_A8B8G8R8_SRGB_PACK32)
+			VKSWITCHCASE(VK_FORMAT_A2R10G10B10_UNORM_PACK32)
+			VKSWITCHCASE(VK_FORMAT_A2R10G10B10_SNORM_PACK32)
+			VKSWITCHCASE(VK_FORMAT_A2R10G10B10_USCALED_PACK32)
+			VKSWITCHCASE(VK_FORMAT_A2R10G10B10_SSCALED_PACK32)
+			VKSWITCHCASE(VK_FORMAT_A2R10G10B10_UINT_PACK32)
+			VKSWITCHCASE(VK_FORMAT_A2R10G10B10_SINT_PACK32)
+			VKSWITCHCASE(VK_FORMAT_A2B10G10R10_UNORM_PACK32)
+			VKSWITCHCASE(VK_FORMAT_A2B10G10R10_SNORM_PACK32)
+			VKSWITCHCASE(VK_FORMAT_A2B10G10R10_USCALED_PACK32)
+			VKSWITCHCASE(VK_FORMAT_A2B10G10R10_SSCALED_PACK32)
+			VKSWITCHCASE(VK_FORMAT_A2B10G10R10_UINT_PACK32)
+			VKSWITCHCASE(VK_FORMAT_A2B10G10R10_SINT_PACK32)
+			VKSWITCHCASE(VK_FORMAT_R16_UNORM)
+			VKSWITCHCASE(VK_FORMAT_R16_SNORM)
+			VKSWITCHCASE(VK_FORMAT_R16_USCALED)
+			VKSWITCHCASE(VK_FORMAT_R16_SSCALED)
+			VKSWITCHCASE(VK_FORMAT_R16_UINT)
+			VKSWITCHCASE(VK_FORMAT_R16_SINT)
+			VKSWITCHCASE(VK_FORMAT_R16_SFLOAT)
+			VKSWITCHCASE(VK_FORMAT_R16G16_UNORM)
+			VKSWITCHCASE(VK_FORMAT_R16G16_SNORM)
+			VKSWITCHCASE(VK_FORMAT_R16G16_USCALED)
+			VKSWITCHCASE(VK_FORMAT_R16G16_SSCALED)
+			VKSWITCHCASE(VK_FORMAT_R16G16_UINT)
+			VKSWITCHCASE(VK_FORMAT_R16G16_SINT)
+			VKSWITCHCASE(VK_FORMAT_R16G16_SFLOAT)
+			VKSWITCHCASE(VK_FORMAT_R16G16B16_UNORM)
+			VKSWITCHCASE(VK_FORMAT_R16G16B16_SNORM)
+			VKSWITCHCASE(VK_FORMAT_R16G16B16_USCALED)
+			VKSWITCHCASE(VK_FORMAT_R16G16B16_SSCALED)
+			VKSWITCHCASE(VK_FORMAT_R16G16B16_UINT)
+			VKSWITCHCASE(VK_FORMAT_R16G16B16_SINT)
+			VKSWITCHCASE(VK_FORMAT_R16G16B16_SFLOAT)
+			VKSWITCHCASE(VK_FORMAT_R16G16B16A16_UNORM)
+			VKSWITCHCASE(VK_FORMAT_R16G16B16A16_SNORM)
+			VKSWITCHCASE(VK_FORMAT_R16G16B16A16_USCALED)
+			VKSWITCHCASE(VK_FORMAT_R16G16B16A16_SSCALED)
+			VKSWITCHCASE(VK_FORMAT_R16G16B16A16_UINT)
+			VKSWITCHCASE(VK_FORMAT_R16G16B16A16_SINT)
+			VKSWITCHCASE(VK_FORMAT_R16G16B16A16_SFLOAT)
+			VKSWITCHCASE(VK_FORMAT_R32_UINT)
+			VKSWITCHCASE(VK_FORMAT_R32_SINT)
+			VKSWITCHCASE(VK_FORMAT_R32_SFLOAT)
+			VKSWITCHCASE(VK_FORMAT_R32G32_UINT)
+			VKSWITCHCASE(VK_FORMAT_R32G32_SINT)
+			VKSWITCHCASE(VK_FORMAT_R32G32_SFLOAT)
+			VKSWITCHCASE(VK_FORMAT_R32G32B32_UINT)
+			VKSWITCHCASE(VK_FORMAT_R32G32B32_SINT)
+			VKSWITCHCASE(VK_FORMAT_R32G32B32_SFLOAT)
+			VKSWITCHCASE(VK_FORMAT_R32G32B32A32_UINT)
+			VKSWITCHCASE(VK_FORMAT_R32G32B32A32_SINT)
+			VKSWITCHCASE(VK_FORMAT_R32G32B32A32_SFLOAT)
+			VKSWITCHCASE(VK_FORMAT_R64_UINT)
+			VKSWITCHCASE(VK_FORMAT_R64_SINT)
+			VKSWITCHCASE(VK_FORMAT_R64_SFLOAT)
+			VKSWITCHCASE(VK_FORMAT_R64G64_UINT)
+			VKSWITCHCASE(VK_FORMAT_R64G64_SINT)
+			VKSWITCHCASE(VK_FORMAT_R64G64_SFLOAT)
+			VKSWITCHCASE(VK_FORMAT_R64G64B64_UINT)
+			VKSWITCHCASE(VK_FORMAT_R64G64B64_SINT)
+			VKSWITCHCASE(VK_FORMAT_R64G64B64_SFLOAT)
+			VKSWITCHCASE(VK_FORMAT_R64G64B64A64_UINT)
+			VKSWITCHCASE(VK_FORMAT_R64G64B64A64_SINT)
+			VKSWITCHCASE(VK_FORMAT_R64G64B64A64_SFLOAT)
+			VKSWITCHCASE(VK_FORMAT_B10G11R11_UFLOAT_PACK32)
+			VKSWITCHCASE(VK_FORMAT_E5B9G9R9_UFLOAT_PACK32)
+			VKSWITCHCASE(VK_FORMAT_D16_UNORM)
+			VKSWITCHCASE(VK_FORMAT_X8_D24_UNORM_PACK32)
+			VKSWITCHCASE(VK_FORMAT_D32_SFLOAT)
+			VKSWITCHCASE(VK_FORMAT_S8_UINT)
+			VKSWITCHCASE(VK_FORMAT_D16_UNORM_S8_UINT)
+			VKSWITCHCASE(VK_FORMAT_D24_UNORM_S8_UINT)
+			VKSWITCHCASE(VK_FORMAT_D32_SFLOAT_S8_UINT)
+			VKSWITCHCASE(VK_FORMAT_BC1_RGB_UNORM_BLOCK)
+			VKSWITCHCASE(VK_FORMAT_BC1_RGB_SRGB_BLOCK)
+			VKSWITCHCASE(VK_FORMAT_BC1_RGBA_UNORM_BLOCK)
+			VKSWITCHCASE(VK_FORMAT_BC1_RGBA_SRGB_BLOCK)
+			VKSWITCHCASE(VK_FORMAT_BC2_UNORM_BLOCK)
+			VKSWITCHCASE(VK_FORMAT_BC2_SRGB_BLOCK)
+			VKSWITCHCASE(VK_FORMAT_BC3_UNORM_BLOCK)
+			VKSWITCHCASE(VK_FORMAT_BC3_SRGB_BLOCK)
+			VKSWITCHCASE(VK_FORMAT_BC4_UNORM_BLOCK)
+			VKSWITCHCASE(VK_FORMAT_BC4_SNORM_BLOCK)
+			VKSWITCHCASE(VK_FORMAT_BC5_UNORM_BLOCK)
+			VKSWITCHCASE(VK_FORMAT_BC5_SNORM_BLOCK)
+			VKSWITCHCASE(VK_FORMAT_BC6H_UFLOAT_BLOCK)
+			VKSWITCHCASE(VK_FORMAT_BC6H_SFLOAT_BLOCK)
+			VKSWITCHCASE(VK_FORMAT_BC7_UNORM_BLOCK)
+			VKSWITCHCASE(VK_FORMAT_BC7_SRGB_BLOCK)
+			VKSWITCHCASE(VK_FORMAT_ETC2_R8G8B8_UNORM_BLOCK)
+			VKSWITCHCASE(VK_FORMAT_ETC2_R8G8B8_SRGB_BLOCK)
+			VKSWITCHCASE(VK_FORMAT_ETC2_R8G8B8A1_UNORM_BLOCK)
+			VKSWITCHCASE(VK_FORMAT_ETC2_R8G8B8A1_SRGB_BLOCK)
+			VKSWITCHCASE(VK_FORMAT_ETC2_R8G8B8A8_UNORM_BLOCK)
+			VKSWITCHCASE(VK_FORMAT_ETC2_R8G8B8A8_SRGB_BLOCK)
+			VKSWITCHCASE(VK_FORMAT_EAC_R11_UNORM_BLOCK)
+			VKSWITCHCASE(VK_FORMAT_EAC_R11_SNORM_BLOCK)
+			VKSWITCHCASE(VK_FORMAT_EAC_R11G11_UNORM_BLOCK)
+			VKSWITCHCASE(VK_FORMAT_EAC_R11G11_SNORM_BLOCK)
+			VKSWITCHCASE(VK_FORMAT_ASTC_4x4_UNORM_BLOCK)
+			VKSWITCHCASE(VK_FORMAT_ASTC_4x4_SRGB_BLOCK)
+			VKSWITCHCASE(VK_FORMAT_ASTC_5x4_UNORM_BLOCK)
+			VKSWITCHCASE(VK_FORMAT_ASTC_5x4_SRGB_BLOCK)
+			VKSWITCHCASE(VK_FORMAT_ASTC_5x5_UNORM_BLOCK)
+			VKSWITCHCASE(VK_FORMAT_ASTC_5x5_SRGB_BLOCK)
+			VKSWITCHCASE(VK_FORMAT_ASTC_6x5_UNORM_BLOCK)
+			VKSWITCHCASE(VK_FORMAT_ASTC_6x5_SRGB_BLOCK)
+			VKSWITCHCASE(VK_FORMAT_ASTC_6x6_UNORM_BLOCK)
+			VKSWITCHCASE(VK_FORMAT_ASTC_6x6_SRGB_BLOCK)
+			VKSWITCHCASE(VK_FORMAT_ASTC_8x5_UNORM_BLOCK)
+			VKSWITCHCASE(VK_FORMAT_ASTC_8x5_SRGB_BLOCK)
+			VKSWITCHCASE(VK_FORMAT_ASTC_8x6_UNORM_BLOCK)
+			VKSWITCHCASE(VK_FORMAT_ASTC_8x6_SRGB_BLOCK)
+			VKSWITCHCASE(VK_FORMAT_ASTC_8x8_UNORM_BLOCK)
+			VKSWITCHCASE(VK_FORMAT_ASTC_8x8_SRGB_BLOCK)
+			VKSWITCHCASE(VK_FORMAT_ASTC_10x5_UNORM_BLOCK)
+			VKSWITCHCASE(VK_FORMAT_ASTC_10x5_SRGB_BLOCK)
+			VKSWITCHCASE(VK_FORMAT_ASTC_10x6_UNORM_BLOCK)
+			VKSWITCHCASE(VK_FORMAT_ASTC_10x6_SRGB_BLOCK)
+			VKSWITCHCASE(VK_FORMAT_ASTC_10x8_UNORM_BLOCK)
+			VKSWITCHCASE(VK_FORMAT_ASTC_10x8_SRGB_BLOCK)
+			VKSWITCHCASE(VK_FORMAT_ASTC_10x10_UNORM_BLOCK)
+			VKSWITCHCASE(VK_FORMAT_ASTC_10x10_SRGB_BLOCK)
+			VKSWITCHCASE(VK_FORMAT_ASTC_12x10_UNORM_BLOCK)
+			VKSWITCHCASE(VK_FORMAT_ASTC_12x10_SRGB_BLOCK)
+			VKSWITCHCASE(VK_FORMAT_ASTC_12x12_UNORM_BLOCK)
+			VKSWITCHCASE(VK_FORMAT_ASTC_12x12_SRGB_BLOCK)
+#undef VKSWITCHCASE
+		default:
+			break;
+		}
+		return FString::Printf(TEXT("Unknown VkFormat %d"), (int32)Format);
+	}
+
 	static FString GetVkResultErrorString(VkResult Result)
 	{
-		FString ErrorString;
 		switch (Result)
 		{
-#define VKSWITCHCASE(x)	case x: ErrorString = TEXT(#x); break;
+			// + 3 to skip "VK_"
+#define VKSWITCHCASE(x)	case x: return TEXT(#x) + 3;
 			VKSWITCHCASE(VK_SUCCESS)
 			VKSWITCHCASE(VK_NOT_READY)
 			VKSWITCHCASE(VK_TIMEOUT)
@@ -230,19 +428,34 @@
 			VKSWITCHCASE(VK_ERROR_INVALID_SHADER_NV)
 #undef VKSWITCHCASE
 		default:
-			ErrorString = FString::Printf(TEXT("Unknown VkResult %d"), (int32)Result);
 			break;
 		}
 
-		return ErrorString;
+		return FString::Printf(TEXT("Unknown VkResult %d"), (int32)Result);
 	}
 
+	FString GetImageTilingString(VkImageTiling Tiling)
+	{
+		switch (Tiling)
+		{
+			// + 16 to skip "VK_IMAGE_TILING_"
+#define VKSWITCHCASE(x)	case x: return TEXT(#x) + 16;
+			VKSWITCHCASE(VK_IMAGE_TILING_OPTIMAL)
+			VKSWITCHCASE(VK_IMAGE_TILING_LINEAR)
+#undef VKSWITCHCASE
+		default:
+			break;
+		}
+
+		return FString::Printf(TEXT("Unknown VkImageTiling %d"), (int32)Tiling);
+	}
+
 	static FString GetImageLayoutString(VkImageLayout Layout)
 	{
-		FString String;
 		switch (Layout)
 		{
-#define VKSWITCHCASE(x)	case x: String = TEXT(#x); break;
+			// + 16 to skip "VK_IMAGE_LAYOUT"
+#define VKSWITCHCASE(x)	case x: return TEXT(#x) + 16;
 			VKSWITCHCASE(VK_IMAGE_LAYOUT_UNDEFINED)
 			VKSWITCHCASE(VK_IMAGE_LAYOUT_GENERAL)
 			VKSWITCHCASE(VK_IMAGE_LAYOUT_COLOR_ATTACHMENT_OPTIMAL)
@@ -255,13 +468,33 @@
 			VKSWITCHCASE(VK_IMAGE_LAYOUT_PRESENT_SRC_KHR)
 #undef VKSWITCHCASE
 		default:
-			String = FString::Printf(TEXT("Unknown VkImageLayout %d"), (int32)Layout);
 			break;
 		}
 
-		return String;
+		return FString::Printf(TEXT("Unknown VkImageLayout %d"), (int32)Layout);
 	}
 
+	static FString GetImageViewTypeString(VkImageViewType Type)
+	{
+		switch (Type)
+		{
+			// + 19 to skip "VK_IMAGE_VIEW_TYPE_"
+#define VKSWITCHCASE(x)	case x: return TEXT(#x) + 16;
+		VKSWITCHCASE(VK_IMAGE_VIEW_TYPE_1D)
+		VKSWITCHCASE(VK_IMAGE_VIEW_TYPE_2D)
+		VKSWITCHCASE(VK_IMAGE_VIEW_TYPE_3D)
+		VKSWITCHCASE(VK_IMAGE_VIEW_TYPE_CUBE)
+		VKSWITCHCASE(VK_IMAGE_VIEW_TYPE_1D_ARRAY)
+		VKSWITCHCASE(VK_IMAGE_VIEW_TYPE_2D_ARRAY)
+		VKSWITCHCASE(VK_IMAGE_VIEW_TYPE_CUBE_ARRAY)
+#undef VKSWITCHCASE
+		default:
+			break;
+		}
+
+		return FString::Printf(TEXT("Unknown VkImageViewType %d"), (int32)Type);
+	}
+
 	static FString GetComponentMappingString(const VkComponentMapping& Mapping)
 	{
 		auto GetSwizzle = [](VkComponentSwizzle Swizzle) -> const TCHAR*
@@ -288,18 +521,11 @@
 	if ((Flags & BitField) == BitField)\
 	{\
 		Flags &= ~BitField;\
+		if (String.Len() > 0)\
+		{\
+			String += TEXT("|");\
+		}\
 		String += Name;\
-		if (bFirst)\
-		{\
-			bFirst = false;\
-		}\
-		else\
-		{\
-			if (Flags !=0)\
-			{\
-				String += TEXT("|");\
-			}\
-		}\
 	}
 
 	static FString GetAspectMaskString(VkImageAspectFlags Flags)
@@ -308,7 +534,6 @@
 		{
 			return TEXT("0");
 		}
-		bool bFirst = true;
 		FString String;
 		AppendBitFieldName(VK_IMAGE_ASPECT_COLOR_BIT, TEXT("COLOR"));
 		AppendBitFieldName(VK_IMAGE_ASPECT_DEPTH_BIT, TEXT("DEPTH"));
@@ -328,7 +553,6 @@
 		{
 			return TEXT("0");
 		}
-		bool bFirst = true;
 		FString String;
 		AppendBitFieldName(VK_ACCESS_INDIRECT_COMMAND_READ_BIT, TEXT("INDIRECT_COMMAND"));
 		AppendBitFieldName(VK_ACCESS_INDEX_READ_BIT, TEXT("INDEX_READ"));
@@ -355,8 +579,62 @@
 		return String;
 	}
 
+	FString GetSampleCountString(VkSampleCountFlags Flags)
+	{
+		if (Flags == 0)
+		{
+			return TEXT("0");
+		}
+		FString String;
+		AppendBitFieldName(VK_SAMPLE_COUNT_1_BIT, TEXT("1"));
+		AppendBitFieldName(VK_SAMPLE_COUNT_2_BIT, TEXT("2"));
+		AppendBitFieldName(VK_SAMPLE_COUNT_4_BIT, TEXT("4"));
+		AppendBitFieldName(VK_SAMPLE_COUNT_8_BIT, TEXT("8"));
+		AppendBitFieldName(VK_SAMPLE_COUNT_16_BIT, TEXT("16"));
+		AppendBitFieldName(VK_SAMPLE_COUNT_32_BIT, TEXT("32"));
+		AppendBitFieldName(VK_SAMPLE_COUNT_64_BIT, TEXT("64"));
+		if (Flags != 0)
+		{
+			FString Unknown = FString::Printf(TEXT("%d"), Flags);
+			AppendBitFieldName(Flags, Unknown);
+		}
+		return String;
+	}
+
+	FString GetImageUsageString(VkImageUsageFlags Flags)
+	{
+		if (Flags == 0)
+		{
+			return TEXT("0");
+		}
+		FString String;
+		AppendBitFieldName(VK_IMAGE_USAGE_TRANSFER_SRC_BIT, TEXT("XFER_SRC"));
+		AppendBitFieldName(VK_IMAGE_USAGE_TRANSFER_DST_BIT, TEXT("XFER_DST"));
+		AppendBitFieldName(VK_IMAGE_USAGE_SAMPLED_BIT, TEXT("SAMPLED"));
+		AppendBitFieldName(VK_IMAGE_USAGE_STORAGE_BIT, TEXT("STORAGE"));
+		AppendBitFieldName(VK_IMAGE_USAGE_COLOR_ATTACHMENT_BIT, TEXT("COLOR_ATT"));
+		AppendBitFieldName(VK_IMAGE_USAGE_DEPTH_STENCIL_ATTACHMENT_BIT, TEXT("DS_ATT"));
+		AppendBitFieldName(VK_IMAGE_USAGE_TRANSIENT_ATTACHMENT_BIT, TEXT("TRANS_ATT"));
+		AppendBitFieldName(VK_IMAGE_USAGE_INPUT_ATTACHMENT_BIT, TEXT("IN_ATT"));
+		if (Flags != 0)
+		{
+			FString Unknown = FString::Printf(TEXT("%d"), Flags);
+			AppendBitFieldName(Flags, Unknown);
+		}
+		return String;
+	}
 #undef  AppendBitFieldName
 
+	FString GetExtentString(const VkExtent3D& Extent)
+	{
+		return FString::Printf(TEXT("w:%d h:%d d:%d"), Extent.width, Extent.height, Extent.depth);
+	}
+
+	FString GetExtentString(const VkExtent2D& Extent)
+	{
+		return FString::Printf(TEXT("w:%d h:%d"), Extent.width, Extent.height);
+	}
+
 	static FString GetSubResourceRangeString(const VkImageSubresourceRange& Range)
 	{
 		return FString::Printf(TEXT("aspectMask=%s, baseMipLevel=%d, levelCount=%d, baseArrayLayer=%d, layerCount=%d"), *GetAspectMaskString(Range.aspectMask), Range.baseMipLevel, Range.levelCount, Range.baseArrayLayer, Range.layerCount);		
@@ -463,9 +741,7 @@
 	{
 		if (CVarVulkanDumpLayer.GetValueOnAnyThread())
 		{
-			DevicePrintfBeginResult(Device, FString::Printf(TEXT("vkAllocateMemory(OutMem=%p)\n"), AllocateInfo, Memory));
-			DebugLog += Tabs;
-			DebugLog += FString::Printf(TEXT("VkMemoryAllocateInfo %p: Size=%d, MemoryTypeIndex=%d"), AllocateInfo, (uint32)AllocateInfo->allocationSize, AllocateInfo->memoryTypeIndex);
+			DevicePrintfBeginResult(Device, FString::Printf(TEXT("vkAllocateMemory(OutMem=%p): Size=%d, MemTypeIndex=%d"), AllocateInfo, Memory, (uint32)AllocateInfo->allocationSize, AllocateInfo->memoryTypeIndex));
 		}
 	}
 
@@ -473,34 +749,41 @@
 	{
 		if (CVarVulkanDumpLayer.GetValueOnAnyThread())
 		{
-			DebugLog += FString::Printf(TEXT(" -> Size=%d Align=%d MemoryTypeBits=0x%x\n"), (uint32)MemoryRequirements->size, (uint32)MemoryRequirements->alignment, MemoryRequirements->memoryTypeBits);
+			DebugLog += FString::Printf(TEXT(" -> Size=%d Align=%d MemTypeBits=0x%x\n"), (uint32)MemoryRequirements->size, (uint32)MemoryRequirements->alignment, MemoryRequirements->memoryTypeBits);
 		}
 	}
 
-	void DumpBufferCreate(VkDevice Device, const VkBufferCreateInfo* CreateInfo, VkBuffer* Buffer)
+	void DumpCreateBuffer(VkDevice Device, const VkBufferCreateInfo* CreateInfo, VkBuffer* Buffer)
 	{
 		FlushDebugWrapperLog();
 		DevicePrintfBeginResult(Device, FString::Printf(TEXT("vkCreateBuffer(Info=%p, OutBuffer=%p)[...]"), CreateInfo, Buffer));
 	}
 
-	void DumpBufferViewCreate(VkDevice Device, const VkBufferViewCreateInfo* CreateInfo, VkBufferView* BufferView)
+	void DumpCreateBufferView(VkDevice Device, const VkBufferViewCreateInfo* CreateInfo, VkBufferView* BufferView)
 	{
 		DevicePrintfBeginResult(Device, FString::Printf(TEXT("VkBufferViewCreateInfo(Info=%p, OutBufferView=%p)[...]"), CreateInfo, BufferView));
 	}
 
-	void DumpImageCreate(VkDevice Device, const VkImageCreateInfo* CreateInfo, VkImage* Image)
+	void DumpCreateImage(VkDevice Device, const VkImageCreateInfo* CreateInfo, VkImage* Image)
 	{
-		FlushDebugWrapperLog();
-		DevicePrintfBeginResult(Device, FString::Printf(TEXT("vkCreateImage(Info=%p, OutImage=%p)[...]"), CreateInfo, Image));
+		if (CVarVulkanDumpLayer.GetValueOnAnyThread())
+		{
+			FlushDebugWrapperLog();
+			DevicePrintfBegin(Device, FString::Printf(TEXT("vkCreateImage(Info=%p, OutImage=%p)"), CreateInfo, Image));
+			DebugLog += FString::Printf(TEXT("%sVkImageCreateInfo: Flags=%d, ImageType=%d, Format=%s, MipLevels=%d, ArrayLayers=%d, Samples=%s\n"), Tabs, CreateInfo->flags, (uint32)CreateInfo->imageType,
+				*GetVkFormatString(CreateInfo->format), CreateInfo->mipLevels, CreateInfo->arrayLayers, *GetSampleCountString(CreateInfo->samples));
+			DebugLog += FString::Printf(TEXT("%s\tExtent=(%s) Tiling=%s, Usage=%s, Initial=%s\n"), Tabs, *GetExtentString(CreateInfo->extent), 
+				*GetImageTilingString(CreateInfo->tiling), *GetImageUsageString(CreateInfo->usage), *GetImageLayoutString(CreateInfo->initialLayout));
+		}
 	}
 
-	void DumpImageViewCreate(VkDevice Device, const VkImageViewCreateInfo* CreateInfo, VkImageView* ImageView)
+	void DumpCreateImageView(VkDevice Device, const VkImageViewCreateInfo* CreateInfo, VkImageView* ImageView)
 	{
-		DevicePrintfBegin(Device, FString::Printf(TEXT("VkImageViewCreateInfo(Info=%p, OutImageView=%p)"), CreateInfo, ImageView));
-
 		if (CVarVulkanDumpLayer.GetValueOnAnyThread())
 		{
-			DebugLog += FString::Printf(TEXT("%sVkImageViewCreateInfo: Flags=%d, Image=%p, ViewType=%d, Format=%p, Components=%s\n"), Tabs, CreateInfo->flags, CreateInfo->image, (int32)CreateInfo->viewType, (void*)CreateInfo->format, *GetComponentMappingString(CreateInfo->components));
+			DevicePrintfBegin(Device, FString::Printf(TEXT("vkCreateImageView(Info=%p, OutImageView=%p)"), CreateInfo, ImageView));
+			DebugLog += FString::Printf(TEXT("%sVkImageViewCreateInfo: Flags=%d, Image=%p, ViewType=%s, Format=%s, Components=%s\n"), Tabs, CreateInfo->flags, CreateInfo->image, 
+				*GetImageViewTypeString(CreateInfo->viewType), *GetVkFormatString(CreateInfo->format), *GetComponentMappingString(CreateInfo->components));
 			DebugLog += FString::Printf(TEXT("%s\tSubresourceRange=(%s)"), Tabs, *GetSubResourceRangeString(CreateInfo->subresourceRange));
 		}
 	}
@@ -663,7 +946,98 @@
 
 	void DumpCreateRenderPass(VkDevice Device, const VkRenderPassCreateInfo* CreateInfo, VkRenderPass* RenderPass)
 	{
-		DevicePrintfBeginResult(Device, FString::Printf(TEXT("vkCreateRenderPass(Info=%p, OutRenderPass=%p)[...]"), CreateInfo, RenderPass));
+		if (CVarVulkanDumpLayer.GetValueOnAnyThread())
+		{
+			DevicePrintfBegin(Device, FString::Printf(TEXT("vkCreateRenderPass(Info=%p, OutRenderPass=%p)[...]"), CreateInfo, RenderPass));
+			DebugLog += FString::Printf(TEXT("%s\tVkRenderPassCreateInfo: NumAttachments=%d, Attachments=%p, NumSubPasses=%d, SubPasses=%p\n"), Tabs, CreateInfo->attachmentCount, CreateInfo->pAttachments, CreateInfo->subpassCount, CreateInfo->pSubpasses);
+			for (uint32 Index = 0; Index < CreateInfo->attachmentCount; ++Index)
+			{
+				auto GetLoadOpString = [](VkAttachmentLoadOp Op) -> FString
+					{
+						switch (Op)
+						{
+						case VK_ATTACHMENT_LOAD_OP_LOAD: return TEXT("LOAD");
+						case VK_ATTACHMENT_LOAD_OP_CLEAR: return TEXT("CLEAR");
+						case VK_ATTACHMENT_LOAD_OP_DONT_CARE: return TEXT("DONT_CARE");
+						default: return FString::Printf(TEXT("Invalid(%d)"), (uint32)Op);
+						}
+					};
+				auto GetStoreOpString = [](VkAttachmentStoreOp Op) -> FString
+					{
+						switch(Op)
+						{
+						case VK_ATTACHMENT_STORE_OP_STORE: return TEXT("STORE");
+						case VK_ATTACHMENT_STORE_OP_DONT_CARE: return TEXT("DONT_CARE");
+						default: return FString::Printf(TEXT("Invalid(%d)"), (uint32)Op);
+						}
+					};
+				/*
+				typedef struct VkAttachmentDescription {
+					VkAttachmentDescriptionFlags    flags;
+					VkFormat                        format;
+					VkSampleCountFlagBits           samples;
+					VkAttachmentLoadOp              loadOp;
+					VkAttachmentStoreOp             storeOp;
+					VkAttachmentLoadOp              stencilLoadOp;
+					VkAttachmentStoreOp             stencilStoreOp;
+					VkImageLayout                   initialLayout;
+					VkImageLayout                   finalLayout;
+				} VkAttachmentDescription;
+*/
+				const VkAttachmentDescription& Desc = CreateInfo->pAttachments[Index];
+				DebugLog += FString::Printf(TEXT("%s\t\tAttachment[%d]: Flags=%s, Format=%s, Samples=%s, Load=%s, Store=%s\n"), Tabs, Index,
+					(Desc.flags == VK_ATTACHMENT_DESCRIPTION_MAY_ALIAS_BIT ? TEXT("MAY_ALIAS") : TEXT("0")),
+					*GetVkFormatString(Desc.format), *GetSampleCountString(Desc.samples), *GetLoadOpString(Desc.loadOp), *GetStoreOpString(Desc.storeOp));
+				DebugLog += FString::Printf(TEXT("%s\t\t\tLoadStencil=%s, StoreStencil=%s, Initial=%s, Final=%s\n"), Tabs, 
+					*GetLoadOpString(Desc.stencilLoadOp), *GetStoreOpString(Desc.stencilStoreOp), *VulkanRHI::GetImageLayoutString(Desc.initialLayout), *VulkanRHI::GetImageLayoutString(Desc.finalLayout));
+			}
+
+			for (uint32 Index = 0; Index < CreateInfo->subpassCount; ++Index)
+			{
+				const VkSubpassDescription& Desc = CreateInfo->pSubpasses[Index];
+				DebugLog += FString::Printf(TEXT("%s\t\tSubpass[%d]: Flags=%d, Bind=%s, NumInputAttach=%d, InputAttach=%p, NumColorAttach=%d, ColorAttach=%p, DSAttch=%p\n"), Tabs, Index,
+					Desc.flags,
+					Desc.pipelineBindPoint == VK_PIPELINE_BIND_POINT_COMPUTE ? TEXT("Compute") : TEXT("Gfx"),
+					Desc.inputAttachmentCount, Desc.pInputAttachments, Desc.colorAttachmentCount, Desc.pColorAttachments, Desc.pDepthStencilAttachment);
+				for (uint32 SubIndex = 0; SubIndex < Desc.inputAttachmentCount; ++SubIndex)
+				{
+					DebugLog += FString::Printf(TEXT("%s\t\t\tInputAttach[%d]: Attach=%d, Layout=%s\n"), Tabs, Index,
+						Desc.pInputAttachments[SubIndex].attachment, *GetImageLayoutString(Desc.pInputAttachments[SubIndex].layout));
+				}
+				for (uint32 SubIndex = 0; SubIndex < Desc.inputAttachmentCount; ++SubIndex)
+				{
+					DebugLog += FString::Printf(TEXT("%s\t\t\tColorAttach[%d]: Attach=%d, Layout=%s\n"), Tabs, Index,
+						Desc.pColorAttachments[SubIndex].attachment, *GetImageLayoutString(Desc.pColorAttachments[SubIndex].layout));
+				}
+				if (Desc.pDepthStencilAttachment)
+				{
+					DebugLog += FString::Printf(TEXT("%s\t\t\tDSAttach: Attach=%d, Layout=%s\n"), Tabs, Desc.pDepthStencilAttachment->attachment, *GetImageLayoutString(Desc.pDepthStencilAttachment->layout));
+				}
+				/*
+				typedef struct VkAttachmentReference {
+					uint32_t         attachment;
+					VkImageLayout    layout;
+				} VkAttachmentReference;
+
+				typedef struct VkSubpassDescription {
+					const VkAttachmentReference*    pResolveAttachments;
+					uint32_t                        preserveAttachmentCount;
+					const uint32_t*                 pPreserveAttachments;
+				} VkSubpassDescription;*/
+			}
+
+/*
+			typedef struct VkRenderPassCreateInfo {
+				uint32_t                          attachmentCount;
+				const VkAttachmentDescription*    pAttachments;
+				uint32_t                          subpassCount;
+				const VkSubpassDescription*       pSubpasses;
+				uint32_t                          dependencyCount;
+				const VkSubpassDependency*        pDependencies;
+			} VkRenderPassCreateInfo;
+*/
+
+		}
 	}
 
 	void DumpQueueSubmit(VkQueue Queue, uint32 SubmitCount, const VkSubmitInfo* Submits, VkFence Fence)
@@ -716,6 +1090,7 @@
 
 	void DumpCreateCommandPool(VkDevice Device, const VkCommandPoolCreateInfo* CreateInfo, VkCommandPool* CommandPool)
 	{
+		FlushDebugWrapperLog();
 		DevicePrintfBeginResult(Device, FString::Printf(TEXT("vkCreateCommandPool(CreateInfo=%p, OutCommandPool=%p)[...]"), CreateInfo, CommandPool));
 	}
 
@@ -741,6 +1116,7 @@
 
 	void DumpCreateDevice(VkPhysicalDevice PhysicalDevice, const VkDeviceCreateInfo* CreateInfo, VkDevice* Device)
 	{
+		FlushDebugWrapperLog();
 		PrintfBeginResult(FString::Printf(TEXT("vkCreateDevice(PhysicalDevice=%p, CreateInfo=%p, OutDevice=%p)[...]"), PhysicalDevice, CreateInfo, Device));
 	}
 
@@ -778,10 +1154,10 @@
 					}					
 				};
 			CmdPrintfBegin(CommandBuffer, FString::Printf(TEXT("vkCmdBeginRenderPass(BeginInfo=%p, Contents=%s)"), RenderPassBegin, *GetSubpassContents(Contents)));
-			DebugLog += FString::Printf(TEXT("%sBeginInfo: RenderPass=%p, Framebuffer=%p, renderArea=(x:%d, y:%d, w:%d, h:%d), clearValues=%d\n"),
+			DebugLog += FString::Printf(TEXT("%sBeginInfo: RenderPass=%p, Framebuffer=%p, renderArea=(x:%d, y:%d, %s), clearValues=%d\n"),
 				Tabs, RenderPassBegin->renderPass, RenderPassBegin->framebuffer, 
 				RenderPassBegin->renderArea.offset.x, RenderPassBegin->renderArea.offset.y, 
-				RenderPassBegin->renderArea.extent.width, RenderPassBegin->renderArea.extent.height,
+				*GetExtentString(RenderPassBegin->renderArea.extent),
 				RenderPassBegin->clearValueCount);
 			for (uint32 Index = 0; Index < RenderPassBegin->clearValueCount; ++Index)
 			{
diff --git a/Engine/Source/Runtime/VulkanRHI/Private/VulkanDevice.cpp b/Engine/Source/Runtime/VulkanRHI/Private/VulkanDevice.cpp
--- a/Engine/Source/Runtime/VulkanRHI/Private/VulkanDevice.cpp
+++ b/Engine/Source/Runtime/VulkanRHI/Private/VulkanDevice.cpp
@@ -215,6 +215,9 @@ void FVulkanDevice::SetupFormats()
 	MapFormatSupport(PF_R32_UINT, VK_FORMAT_R32_UINT);
 	SetComponentMapping(PF_R32_UINT, VK_COMPONENT_SWIZZLE_R, VK_COMPONENT_SWIZZLE_ZERO, VK_COMPONENT_SWIZZLE_ZERO, VK_COMPONENT_SWIZZLE_ZERO);
 
+	MapFormatSupport(PF_R8_UINT, VK_FORMAT_R8_UINT);
+	SetComponentMapping(PF_R8_UINT, VK_COMPONENT_SWIZZLE_R, VK_COMPONENT_SWIZZLE_ZERO, VK_COMPONENT_SWIZZLE_ZERO, VK_COMPONENT_SWIZZLE_ZERO);
+
 #if PLATFORM_ANDROID
 	MapFormatSupport(PF_D24, VK_FORMAT_X8_D24_UNORM_PACK32);
 #elif PLATFORM_DESKTOP
diff --git a/Engine/Source/Runtime/VulkanRHI/Private/VulkanLayers.cpp b/Engine/Source/Runtime/VulkanRHI/Private/VulkanLayers.cpp
--- a/Engine/Source/Runtime/VulkanRHI/Private/VulkanLayers.cpp
+++ b/Engine/Source/Runtime/VulkanRHI/Private/VulkanLayers.cpp
@@ -25,22 +25,30 @@ static const ANSICHAR* GRequiredLayersInstance[] =
 	"VK_LAYER_LUNARG_swapchain",
 };
 
+#define VULKAN_ENABLE_STANDARD_VALIDATION	1
+
 // List of validation layers which we want to activate for the instance
 static const ANSICHAR* GValidationLayersInstance[] =
 {
 #if VULKAN_ENABLE_API_DUMP
 	"VK_LAYER_LUNARG_api_dump",
 #endif
+
+#if VULKAN_ENABLE_STANDARD_VALIDATION
+	"VK_LAYER_LUNARG_standard_validation",
+#else
+	"VK_LAYER_GOOGLE_threading",
+	"VK_LAYER_LUNARG_parameter_validation",
+	"VK_LAYER_LUNARG_object_tracker",
+	"VK_LAYER_LUNARG_image",
 	"VK_LAYER_LUNARG_core_validation",
+	"VK_LAYER_LUNARG_swapchain",
+	"VK_LAYER_GOOGLE_unique_objects",
+#endif
+
 	"VK_LAYER_LUNARG_device_limits",
-	"VK_LAYER_LUNARG_image",
-	"VK_LAYER_LUNARG_object_tracker",	// The framebuffer is not registered for some reason by the object tracker... the steps are exactly the same as in the demo. For now ObjectTracker is disabled...
-	"VK_LAYER_LUNARG_parameter_validation",
 	//"VK_LAYER_LUNARG_screenshot",
-	"VK_LAYER_GOOGLE_threading",
-	"VK_LAYER_GOOGLE_unique_objects",
 	//"VK_LAYER_NV_optimus",
-	"VK_LAYER_LUNARG_standard_validation",
 	//"VK_LAYER_LUNARG_vktrace",		// Useful for future
 };
 
@@ -56,16 +64,22 @@ static const ANSICHAR* GValidationLayersDevice[] =
 #if VULKAN_ENABLE_API_DUMP
 	"VK_LAYER_LUNARG_api_dump",
 #endif
+
+#if VULKAN_ENABLE_STANDARD_VALIDATION
+	"VK_LAYER_LUNARG_standard_validation",
+#else
+	"VK_LAYER_GOOGLE_threading",
+	"VK_LAYER_LUNARG_parameter_validation",
+	"VK_LAYER_LUNARG_object_tracker",
+	"VK_LAYER_LUNARG_image",
 	"VK_LAYER_LUNARG_core_validation",
+	"VK_LAYER_LUNARG_swapchain",
+	"VK_LAYER_GOOGLE_unique_objects",
+#endif
+
 	"VK_LAYER_LUNARG_device_limits",
-	"VK_LAYER_LUNARG_image",
-	"VK_LAYER_LUNARG_object_tracker",	// The framebuffer is not registered for some reason by the object tracker... the steps are exactly the same as in the demo. For now ObjectTracker is disabled...
-	"VK_LAYER_LUNARG_parameter_validation",
 	//"VK_LAYER_LUNARG_screenshot",
-	"VK_LAYER_GOOGLE_threading",
-	"VK_LAYER_GOOGLE_unique_objects",
 	//"VK_LAYER_NV_optimus",
-	"VK_LAYER_LUNARG_standard_validation",
 	//"VK_LAYER_LUNARG_vktrace",		// Useful for future
 };
 #endif // VULKAN_HAS_DEBUGGING_ENABLED
diff --git a/Engine/Source/Runtime/VulkanRHI/Private/VulkanPendingState.cpp b/Engine/Source/Runtime/VulkanRHI/Private/VulkanPendingState.cpp
--- a/Engine/Source/Runtime/VulkanRHI/Private/VulkanPendingState.cpp
+++ b/Engine/Source/Runtime/VulkanRHI/Private/VulkanPendingState.cpp
@@ -370,7 +370,17 @@ bool FVulkanPendingState::RenderPassBegin(FVulkanCmdBuffer* CmdBuffer)
 	if (RTInfo.DepthStencilRenderTarget.Texture)
 	{
 		VkClearValue& DstColor = ClearValues[Index];
-		RTInfo.DepthStencilRenderTarget.Texture->GetClearBinding().GetDepthStencil(DstColor.depthStencil.depth, DstColor.depthStencil.stencil);
+		const FClearValueBinding& ClearBinding = RTInfo.DepthStencilRenderTarget.Texture->GetClearBinding();
+		if (ClearBinding.ColorBinding == EClearBinding::EDepthStencilBound)
+		{
+			ClearBinding.GetDepthStencil(DstColor.depthStencil.depth, DstColor.depthStencil.stencil);
+		}
+		else
+		{
+			ensure(!RTInfo.bClearDepth && !RTInfo.bClearStencil);
+			DstColor.depthStencil.depth = 1.0f;
+			DstColor.depthStencil.stencil = 0;
+		}
 		// @todo vulkan - we don't track the format of the depth buffer in the key, only if depth is enabled (write/test) - should be enough, but worth checking that
 		// if either bit is set that we have a depth buffer attached
 	}
diff --git a/Engine/Source/Runtime/VulkanRHI/Private/VulkanQuery.cpp b/Engine/Source/Runtime/VulkanRHI/Private/VulkanQuery.cpp
--- a/Engine/Source/Runtime/VulkanRHI/Private/VulkanQuery.cpp
+++ b/Engine/Source/Runtime/VulkanRHI/Private/VulkanQuery.cpp
@@ -53,7 +53,9 @@ FVulkanTimestampQueryPool::FVulkanTimestampQueryPool(FVulkanDevice* InDevice) :
 	TimeStampsPerSeconds(0),
 	SecondsPerTimestamp(0),
 	UsedUserQueries(0),
-	bFirst(true)
+	bFirst(true),
+	LastEndCmdBuffer(nullptr),
+	LastFenceSignaledCounter(0)
 {
 	// The number of nanoseconds it takes for a timestamp value to be incremented by 1 can be obtained from VkPhysicalDeviceLimits::timestampPeriod after a call to vkGetPhysicalDeviceProperties.
 	double NanoSecondsPerTimestamp = Device->GetDeviceProperties().limits.timestampPeriod;
@@ -76,12 +78,14 @@ void FVulkanTimestampQueryPool::WriteStartFrame(VkCommandBuffer CmdBuffer)
 	VulkanRHI::vkCmdWriteTimestamp(CmdBuffer, VK_PIPELINE_STAGE_ALL_COMMANDS_BIT, QueryPool, StartFrame);
 }
 
-void FVulkanTimestampQueryPool::WriteEndFrame(VkCommandBuffer CmdBuffer)
+void FVulkanTimestampQueryPool::WriteEndFrame(FVulkanCmdBuffer* CmdBuffer)
 {
 	if (!bFirst)
 	{
 		// End Frame Timestamp
-		VulkanRHI::vkCmdWriteTimestamp(CmdBuffer, VK_PIPELINE_STAGE_ALL_COMMANDS_BIT, QueryPool, EndFrame);
+		VulkanRHI::vkCmdWriteTimestamp(CmdBuffer->GetHandle(), VK_PIPELINE_STAGE_ALL_COMMANDS_BIT, QueryPool, EndFrame);
+		LastEndCmdBuffer = CmdBuffer;
+		LastFenceSignaledCounter = CmdBuffer->GetFenceSignaledCounter();
 	}
 }
 
@@ -90,17 +94,25 @@ void FVulkanTimestampQueryPool::CalculateFrameTime()
 	uint64_t Results[2] = { 0, 0 };
 	if (!bFirst)
 	{
-		VkDevice DeviceHandle = Device->GetInstanceHandle();
-		VkResult Result;
-		Result = VulkanRHI::vkGetQueryPoolResults(DeviceHandle, QueryPool, StartFrame, 2, sizeof(Results), Results, sizeof(uint64), /*VK_QUERY_RESULT_WAIT_BIT | */VK_QUERY_RESULT_64_BIT);
-		if (Result != VK_SUCCESS)
+		if (LastEndCmdBuffer && LastFenceSignaledCounter < LastEndCmdBuffer->GetFenceSignaledCounter())
+		{
+			VkDevice DeviceHandle = Device->GetInstanceHandle();
+			VkResult Result;
+			Result = VulkanRHI::vkGetQueryPoolResults(DeviceHandle, QueryPool, StartFrame, 2, sizeof(Results), Results, sizeof(uint64), /*VK_QUERY_RESULT_PARTIAL_BIT |*/ VK_QUERY_RESULT_64_BIT);
+			if (Result != VK_SUCCESS)
+			{
+				GGPUFrameTime = 0;
+				return;
+			}
+		}
+		else
 		{
-			GGPUFrameTime = 0;
 			return;
 		}
 	}
 
-	double ValueInSeconds = (double)(Results[1] - Results[0]) * SecondsPerTimestamp;
+	const uint64 Delta = Results[1] - Results[0];
+	double ValueInSeconds = (double)Delta * SecondsPerTimestamp;
 	GGPUFrameTime = (uint32)(ValueInSeconds / FPlatformTime::GetSecondsPerCycle());
 }
 
diff --git a/Engine/Source/Runtime/VulkanRHI/Private/VulkanQueue.cpp b/Engine/Source/Runtime/VulkanRHI/Private/VulkanQueue.cpp
--- a/Engine/Source/Runtime/VulkanRHI/Private/VulkanQueue.cpp
+++ b/Engine/Source/Runtime/VulkanRHI/Private/VulkanQueue.cpp
@@ -65,7 +65,10 @@ void FVulkanQueue::Submit(FVulkanCmdBuffer* CmdBuffer, FVulkanSemaphore* WaitSem
 		SubmitInfo.pWaitSemaphores = &Semaphores[1];
 		SubmitInfo.pWaitDstStageMask = &WaitStageFlags;
 	}
-	VERIFYVULKANRESULT(VulkanRHI::vkQueueSubmit(Queue, 1, &SubmitInfo, Fence->GetHandle()));
+	{
+		SCOPE_CYCLE_COUNTER(STAT_VulkanQueueSubmit)
+		VERIFYVULKANRESULT(VulkanRHI::vkQueueSubmit(Queue, 1, &SubmitInfo, Fence->GetHandle()));
+	}
 
 	if (GWaitForIdleOnSubmit != 0)
 	{
diff --git a/Engine/Source/Runtime/VulkanRHI/Private/VulkanRHIPrivate.h b/Engine/Source/Runtime/VulkanRHI/Private/VulkanRHIPrivate.h
--- a/Engine/Source/Runtime/VulkanRHI/Private/VulkanRHIPrivate.h
+++ b/Engine/Source/Runtime/VulkanRHI/Private/VulkanRHIPrivate.h
@@ -412,6 +412,8 @@ DECLARE_CYCLE_STAT_EXTERN(TEXT("Uniform Buffer Creation Time"), STAT_VulkanUnifo
 DECLARE_CYCLE_STAT_EXTERN(TEXT("Apply DS Uniform Buffers"), STAT_VulkanApplyDSUniformBuffers, STATGROUP_VulkanRHI, );
 DECLARE_CYCLE_STAT_EXTERN(TEXT("SRV Update Time"), STAT_VulkanSRVUpdateTime, STATGROUP_VulkanRHI, );
 DECLARE_CYCLE_STAT_EXTERN(TEXT("Deletion Queue"), STAT_VulkanDeletionQueue, STATGROUP_VulkanRHI, );
+DECLARE_CYCLE_STAT_EXTERN(TEXT("Queue Submit"), STAT_VulkanQueueSubmit, STATGROUP_VulkanRHI, );
+DECLARE_CYCLE_STAT_EXTERN(TEXT("Queue Present"), STAT_VulkanQueuePresent, STATGROUP_VulkanRHI, );
 #if VULKAN_ENABLE_AGGRESSIVE_STATS
 DECLARE_CYCLE_STAT_EXTERN(TEXT("Apply DS Shader Resources"), STAT_VulkanApplyDSResources, STATGROUP_VulkanRHI, );
 DECLARE_CYCLE_STAT_EXTERN(TEXT("Update DescriptorSets"), STAT_VulkanUpdateDescriptorSets, STATGROUP_VulkanRHI, );
diff --git a/Engine/Source/Runtime/VulkanRHI/Private/VulkanRenderTarget.cpp b/Engine/Source/Runtime/VulkanRHI/Private/VulkanRenderTarget.cpp
--- a/Engine/Source/Runtime/VulkanRHI/Private/VulkanRenderTarget.cpp
+++ b/Engine/Source/Runtime/VulkanRHI/Private/VulkanRenderTarget.cpp
@@ -38,10 +38,6 @@ void FVulkanCommandListContext::RHISetRenderTargets(uint32 NumSimultaneousRender
 	checkf(NumUAVs == 0, TEXT("Calling SetRenderTargets with UAVs is not supported in Vulkan yet"));
 }
 
-void FVulkanDynamicRHI::RHIDiscardRenderTargets(bool Depth, bool Stencil, uint32 ColorBitMask)
-{
-}
-
 void FVulkanCommandListContext::RHISetRenderTargetsAndClear(const FRHISetRenderTargetsInfo& RenderTargetsInfo)
 {
 	//FRCLog::Printf(FString::Printf(TEXT("RHISetRenderTargetsAndClear\n")));
@@ -203,7 +199,14 @@ void FVulkanCommandListContext::RHITransitionResources(EResourceTransitionAccess
 				}
 				else
 				{
-					ensure(0);
+					FRHITexture3D* RHITexture3D = RHITexture->GetTexture3D();
+					if (RHITexture3D)
+					{
+					}
+					else
+					{
+						ensure(0);
+					}
 				}
 			}
 		}
diff --git a/Engine/Source/Runtime/VulkanRHI/Private/VulkanShaders.cpp b/Engine/Source/Runtime/VulkanRHI/Private/VulkanShaders.cpp
--- a/Engine/Source/Runtime/VulkanRHI/Private/VulkanShaders.cpp
+++ b/Engine/Source/Runtime/VulkanRHI/Private/VulkanShaders.cpp
@@ -235,7 +235,7 @@ FVulkanBoundShaderState::FVulkanBoundShaderState(
 	FMemory::Memzero(VertexInputBindings);
 	FMemory::Memzero(Attributes);
 	FMemory::Memzero(VertexInputStateInfo);
-	FMemory::Memzero(DescriptorImageInfoForStage);
+	FMemory::Memzero(DescriptorSamplerImageInfoForStage);
 
 	static_assert(CrossCompiler::PACKED_TYPEINDEX_MAX <= sizeof(DirtyPackedUniformBufferStagingMask[0]) * 8, "Won't fit!");
 	FMemory::Memzero(DescriptorBufferInfoForPackedUBForStage);
@@ -306,13 +306,13 @@ FVulkanBoundShaderState::FVulkanBoundShaderState(
 
 	GenerateLayoutBindings();
 
-	uint32 UniformSamplerCount = 0;
+	uint32 UniformCombinedSamplerCount = 0;
 	uint32 UniformSamplerBufferCount = 0;
 	uint32 UniformBufferCount = 0;
-	GetDescriptorInfoCounts(UniformSamplerCount, UniformSamplerBufferCount, UniformBufferCount);
+	GetDescriptorInfoCounts(UniformCombinedSamplerCount, UniformSamplerBufferCount, UniformBufferCount);
 
 	// Setup Write information
-	CreateDescriptorWriteInfo(UniformSamplerCount, UniformSamplerBufferCount, UniformBufferCount);
+	CreateDescriptorWriteInfo(UniformCombinedSamplerCount, UniformSamplerBufferCount, UniformBufferCount);
 }
 
 FVulkanBoundShaderState::~FVulkanBoundShaderState()
@@ -533,20 +533,20 @@ void FVulkanBoundShaderState::GenerateLayoutBindings(EShaderFrequency Stage, con
 			}
 		};
 
-	LoopPerType(FVulkanShaderSerializedBindings::TYPE_SAMPLER, VK_DESCRIPTOR_TYPE_COMBINED_IMAGE_SAMPLER);
+	LoopPerType(FVulkanShaderSerializedBindings::TYPE_COMBINED_IMAGE_SAMPLER, VK_DESCRIPTOR_TYPE_COMBINED_IMAGE_SAMPLER);
 	LoopPerType(FVulkanShaderSerializedBindings::TYPE_SAMPLER_BUFFER, VK_DESCRIPTOR_TYPE_UNIFORM_TEXEL_BUFFER);
 	LoopPerType(FVulkanShaderSerializedBindings::TYPE_UNIFORM_BUFFER, VK_DESCRIPTOR_TYPE_UNIFORM_BUFFER);
 	LoopPerType(FVulkanShaderSerializedBindings::TYPE_PACKED_UNIFORM_BUFFER, VK_DESCRIPTOR_TYPE_UNIFORM_BUFFER);
 }
 
-void FVulkanBoundShaderState::CreateDescriptorWriteInfo(uint32 InUniformSamplerCount, uint32 InUniformSamplerBufferCount, uint32 InUniformBufferCount)
+void FVulkanBoundShaderState::CreateDescriptorWriteInfo(uint32 InUniformCombinedSamplerCount, uint32 InUniformSamplerBufferCount, uint32 InUniformBufferCount)
 {
 	check(DescriptorWrites.Num() == 0);
-	DescriptorWrites.AddZeroed(InUniformSamplerCount + InUniformSamplerBufferCount + InUniformBufferCount);
-	DescriptorImageInfo.AddZeroed(InUniformSamplerCount);
+	DescriptorWrites.AddZeroed(InUniformCombinedSamplerCount + InUniformSamplerBufferCount + InUniformBufferCount);
+	DescriptorImageInfo.AddZeroed(InUniformCombinedSamplerCount);
 	DescriptorBufferInfo.AddZeroed(InUniformBufferCount);
 
-	for (uint32 Index = 0; Index < InUniformSamplerCount; ++Index)
+	for (uint32 Index = 0; Index < InUniformCombinedSamplerCount; ++Index)
 	{
 		VkDescriptorImageInfo& Sampler = DescriptorImageInfo[Index];
 		// Texture.Load() still requires a default sampler...
@@ -574,13 +574,13 @@ void FVulkanBoundShaderState::CreateDescriptorWriteInfo(uint32 InUniformSamplerC
 		//#todo-rco: Use BindingTable instead...
 		auto& Bindings = StageShader->CodeHeader.SerializedBindings.Bindings;
 
-		if (InUniformSamplerCount)
+		if (InUniformCombinedSamplerCount)
 		{
-			DescriptorImageInfoForStage[Stage] = DescriptorImageInfo.GetData() + ImageIndex;
+			DescriptorSamplerImageInfoForStage[Stage] = DescriptorImageInfo.GetData() + ImageIndex;
 
-			for (int32 Index = 0; Index < Bindings[FVulkanShaderSerializedBindings::TYPE_SAMPLER].Num(); ++Index)
+			for (int32 Index = 0; Index < Bindings[FVulkanShaderSerializedBindings::TYPE_COMBINED_IMAGE_SAMPLER].Num(); ++Index)
 			{
-				FVulkanShaderSerializedBindings::FBindMap& Mapping = Bindings[FVulkanShaderSerializedBindings::TYPE_SAMPLER][Index];
+				FVulkanShaderSerializedBindings::FBindMap& Mapping = Bindings[FVulkanShaderSerializedBindings::TYPE_COMBINED_IMAGE_SAMPLER][Index];
 				VkWriteDescriptorSet* WriteDesc = &DescriptorWrites[WriteIndex++];
 				WriteDesc->sType = VK_STRUCTURE_TYPE_WRITE_DESCRIPTOR_SET;
 				WriteDesc->dstBinding = Mapping.VulkanBindingIndex;
@@ -635,7 +635,7 @@ void FVulkanBoundShaderState::CreateDescriptorWriteInfo(uint32 InUniformSamplerC
 		}
 	}
 
-	check(InUniformSamplerCount + InUniformSamplerBufferCount + InUniformBufferCount == DescriptorWrites.Num());
+	check(InUniformCombinedSamplerCount + InUniformSamplerBufferCount + InUniformBufferCount == DescriptorWrites.Num());
 }
 
 void FVulkanBoundShaderState::InternalBindVertexStreams(FVulkanCmdBuffer* Cmd, const void* InVertexStreams)
@@ -919,8 +919,8 @@ void FVulkanBoundShaderState::SetTexture(EShaderFrequency Stage, uint32 BindPoin
 	check(VulkanTextureBase);
 
 	FVulkanShader* StageShader = GetShaderPtr((EShaderFrequency)Stage);
-	uint32 VulkanBindingPoint = StageShader->GetBindingTable().SamplerBindingIndices[BindPoint];
-	DescriptorImageInfoForStage[Stage][VulkanBindingPoint].imageView = VulkanTextureBase ? VulkanTextureBase->View.View : VK_NULL_HANDLE;
+	uint32 VulkanBindingPoint = StageShader->GetBindingTable().CombinedSamplerBindingIndices[BindPoint];
+	DescriptorSamplerImageInfoForStage[Stage][VulkanBindingPoint].imageView = VulkanTextureBase ? VulkanTextureBase->View.View : VK_NULL_HANDLE;
 
 	DirtyTextures[Stage] = true;
 #if VULKAN_ENABLE_RHI_DEBUGGING
@@ -937,8 +937,8 @@ void FVulkanBoundShaderState::SetSamplerState(EShaderFrequency Stage, uint32 Bin
 	#endif
 
 	FVulkanShader* StageShader = GetShaderPtr((EShaderFrequency)Stage);
-	uint32 VulkanBindingPoint = StageShader->GetBindingTable().SamplerBindingIndices[BindPoint];
-	DescriptorImageInfoForStage[Stage][VulkanBindingPoint].sampler = Sampler ? Sampler->Sampler : VK_NULL_HANDLE;
+	uint32 VulkanBindingPoint = StageShader->GetBindingTable().CombinedSamplerBindingIndices[BindPoint];
+	DescriptorSamplerImageInfoForStage[Stage][VulkanBindingPoint].sampler = Sampler ? Sampler->Sampler : VK_NULL_HANDLE;
 
 	DirtySamplerStates[Stage] = true;
 #if VULKAN_ENABLE_RHI_DEBUGGING
@@ -1018,9 +1018,9 @@ void FVulkanBoundShaderState::SetUniformBuffer(FVulkanPendingState& PendingState
 #endif
 }
 
-void FVulkanBoundShaderState::GetDescriptorInfoCounts(uint32& OutSamplerCount, uint32& OutSamplerBufferCount, uint32& OutUniformBufferCount)
+void FVulkanBoundShaderState::GetDescriptorInfoCounts(uint32& OutCombinedSamplerCount, uint32& OutSamplerBufferCount, uint32& OutUniformBufferCount)
 {
-	OutSamplerCount = 0;
+	OutCombinedSamplerCount = 0;
 	OutSamplerBufferCount = 0;
 	OutUniformBufferCount = 0;
 
@@ -1033,7 +1033,7 @@ void FVulkanBoundShaderState::GetDescriptorInfoCounts(uint32& OutSamplerCount, u
 			{
 				auto& Layouts = StageShader->CodeHeader.SerializedBindings.Bindings;
 
-				OutSamplerCount += Layouts[FVulkanShaderSerializedBindings::TYPE_SAMPLER].Num();
+				OutCombinedSamplerCount += Layouts[FVulkanShaderSerializedBindings::TYPE_COMBINED_IMAGE_SAMPLER].Num();
 				OutSamplerBufferCount += Layouts[FVulkanShaderSerializedBindings::TYPE_SAMPLER_BUFFER].Num();
 				OutUniformBufferCount += Layouts[FVulkanShaderSerializedBindings::TYPE_PACKED_UNIFORM_BUFFER].Num() + Layouts[FVulkanShaderSerializedBindings::TYPE_UNIFORM_BUFFER].Num();
 
diff --git a/Engine/Source/Runtime/VulkanRHI/Private/VulkanSwapChain.cpp b/Engine/Source/Runtime/VulkanRHI/Private/VulkanSwapChain.cpp
--- a/Engine/Source/Runtime/VulkanRHI/Private/VulkanSwapChain.cpp
+++ b/Engine/Source/Runtime/VulkanRHI/Private/VulkanSwapChain.cpp
@@ -221,7 +221,10 @@ bool FVulkanSwapChain::Present(FVulkanQueue* Queue, FVulkanSemaphore* BackBuffer
 	Info.pSwapchains = &SwapChain;
 	Info.pImageIndices = (uint32*)&CurrentImageIndex;
 
-	VERIFYVULKANRESULT(VulkanRHI::vkQueuePresentKHR(Queue->GetHandle(), &Info));
+	{
+		SCOPE_CYCLE_COUNTER(STAT_VulkanQueuePresent);
+		VERIFYVULKANRESULT(VulkanRHI::vkQueuePresentKHR(Queue->GetHandle(), &Info));
+	}
 
 	return true;
 }
diff --git a/Engine/Source/Runtime/VulkanRHI/Private/VulkanTexture.cpp b/Engine/Source/Runtime/VulkanRHI/Private/VulkanTexture.cpp
--- a/Engine/Source/Runtime/VulkanRHI/Private/VulkanTexture.cpp
+++ b/Engine/Source/Runtime/VulkanRHI/Private/VulkanTexture.cpp
@@ -1223,7 +1223,17 @@ void FVulkanDynamicRHI::RHIBindDebugLabelName(FTextureRHIParamRef TextureRHI, co
 			FVulkanTexture2D* VTex2d = (FVulkanTexture2D*)Tex2d;
 			VulkanRHI::PrintfBegin(FString::Printf(TEXT("vkDebugMarkerSetObjectNameEXT(%p=%s)"), VTex2d->Surface.Image, Name));
 		}
-	}
+		else if (FRHITextureCube* TexCube = TextureRHI->GetTextureCube())
+		{
+			FVulkanTextureCube* VTexCube = (FVulkanTextureCube*)TexCube;
+			VulkanRHI::PrintfBegin(FString::Printf(TEXT("vkDebugMarkerSetObjectNameEXT(%p=%s)"), VTexCube->Surface.Image, Name));
+		}
+		else if (FRHITexture3D* Tex3d = TextureRHI->GetTexture3D())
+		{
+			FVulkanTexture3D* VTex3d = (FVulkanTexture3D*)Tex3d;
+			VulkanRHI::PrintfBegin(FString::Printf(TEXT("vkDebugMarkerSetObjectNameEXT(%p=%s)"), VTex3d->Surface.Image, Name));
+		}
+}
 #endif
 
 #if VULKAN_ENABLE_DRAW_MARKERS
diff --git a/Engine/Source/Runtime/VulkanRHI/Private/VulkanUtil.cpp b/Engine/Source/Runtime/VulkanRHI/Private/VulkanUtil.cpp
--- a/Engine/Source/Runtime/VulkanRHI/Private/VulkanUtil.cpp
+++ b/Engine/Source/Runtime/VulkanRHI/Private/VulkanUtil.cpp
@@ -463,6 +463,8 @@ DEFINE_STAT(STAT_VulkanUniformBufferCreateTime);
 DEFINE_STAT(STAT_VulkanApplyDSUniformBuffers);
 DEFINE_STAT(STAT_VulkanSRVUpdateTime);
 DEFINE_STAT(STAT_VulkanDeletionQueue);
+DEFINE_STAT(STAT_VulkanQueueSubmit);
+DEFINE_STAT(STAT_VulkanQueuePresent);
 #if VULKAN_ENABLE_AGGRESSIVE_STATS
 DEFINE_STAT(STAT_VulkanApplyDSResources);
 DEFINE_STAT(STAT_VulkanUpdateDescriptorSets);
diff --git a/Engine/Source/Runtime/VulkanRHI/Private/VulkanViewport.cpp b/Engine/Source/Runtime/VulkanRHI/Private/VulkanViewport.cpp
--- a/Engine/Source/Runtime/VulkanRHI/Private/VulkanViewport.cpp
+++ b/Engine/Source/Runtime/VulkanRHI/Private/VulkanViewport.cpp
@@ -369,6 +369,21 @@ bool FVulkanViewport::Present(FVulkanCmdBuffer* CmdBuffer, FVulkanQueue* Queue,
 
 	//#todo-rco: Might need to NOT be undefined...
 	VulkanSetImageLayoutSimple(CmdBuffer->GetHandle(), BackBufferImages[AcquiredImageIndex], VK_IMAGE_LAYOUT_UNDEFINED/*VK_IMAGE_LAYOUT_COLOR_ATTACHMENT_OPTIMAL*/, VK_IMAGE_LAYOUT_PRESENT_SRC_KHR);
+
+	{
+		FVulkanTimestampQueryPool* TimestampPool = Device->GetTimestampQueryPool(PresentCount % FVulkanDevice::NumTimestampPools);
+		if (TimestampPool)
+		{
+			TimestampPool->WriteEndFrame(CmdBuffer);
+		}
+
+		if (PresentCount >= FVulkanDevice::NumTimestampPools)
+		{
+			FVulkanTimestampQueryPool* PrevTimestampPool = Device->GetTimestampQueryPool((PresentCount + FVulkanDevice::NumTimestampPools - 1) % FVulkanDevice::NumTimestampPools);
+			PrevTimestampPool->CalculateFrameTime();
+		}
+	}
+
 	CmdBuffer->End();
 	Queue->Submit(CmdBuffer, nullptr, 0, RenderingDoneSemaphores[AcquiredImageIndex]);
 
@@ -420,7 +435,8 @@ bool FVulkanViewport::Present(FVulkanCmdBuffer* CmdBuffer, FVulkanQueue* Queue,
 	//#todo-rco: This needs to go on RHIEndFrame but the CmdBuffer index is not the correct one to read the stats out!
 	VulkanRHI::GManager.GPUProfilingData.EndFrame();
 
-	Device->GetImmediateContext().GetCommandBufferManager()->PrepareForNewActiveCommandBuffer();
+	FVulkanCommandBufferManager* ImmediateCmdBufMgr = Device->GetImmediateContext().GetCommandBufferManager();
+	ImmediateCmdBufMgr->PrepareForNewActiveCommandBuffer();
 
 	//#todo-rco: Consolidate 'end of frame'
 	Device->GetImmediateContext().GetTempFrameAllocationBuffer().Reset();
@@ -444,6 +460,15 @@ bool FVulkanViewport::Present(FVulkanCmdBuffer* CmdBuffer, FVulkanQueue* Queue,
 	}
 #endif
 	++PresentCount;
+	{
+		FVulkanTimestampQueryPool* TimestampPool = Device->GetTimestampQueryPool(PresentCount % FVulkanDevice::NumTimestampPools);
+		if (TimestampPool)
+		{
+			FVulkanCmdBuffer* ActiveCmdBuffer = ImmediateCmdBufMgr->GetActiveCmdBuffer();
+			TimestampPool->WriteStartFrame(ActiveCmdBuffer->GetHandle());
+		}
+	}
+
 	return bResult;
 }
 
@@ -476,11 +501,12 @@ void FVulkanDynamicRHI::RHIResizeViewport(FViewportRHIParamRef ViewportRHI, uint
 	FVulkanViewport* Viewport = ResourceCast(ViewportRHI);
 }
 
-void FVulkanDynamicRHI::RHITick( float DeltaTime )
+void FVulkanDynamicRHI::RHITick(float DeltaTime)
 {
-	check( IsInGameThread() );
+	check(IsInGameThread());
 }
 
+/*
 void FVulkanDynamicRHI::WriteEndFrameTimestamp(void* Data)
 {
 	FVulkanDynamicRHI* This = (FVulkanDynamicRHI*)Data;
@@ -496,6 +522,7 @@ void FVulkanDynamicRHI::WriteEndFrameTimestamp(void* Data)
 
 	VulkanRHI::GManager.GPUProfilingData.EndFrameBeforeSubmit();
 }
+*/
 
 #if 0
 void FVulkanDynamicRHI::Present()
diff --git a/Engine/Source/Runtime/VulkanRHI/Public/VulkanRHI.h b/Engine/Source/Runtime/VulkanRHI/Public/VulkanRHI.h
--- a/Engine/Source/Runtime/VulkanRHI/Public/VulkanRHI.h
+++ b/Engine/Source/Runtime/VulkanRHI/Public/VulkanRHI.h
@@ -123,7 +123,7 @@ public:
 	virtual void RHIResizeViewport(FViewportRHIParamRef Viewport, uint32 SizeX, uint32 SizeY, bool bIsFullscreen) final override;
 	virtual void RHITick(float DeltaTime) final override;
 	virtual void RHISetStreamOutTargets(uint32 NumTargets,const FVertexBufferRHIParamRef* VertexBuffers,const uint32* Offsets) final override;
-	virtual void RHIDiscardRenderTargets(bool Depth,bool Stencil,uint32 ColorBitMask) final override;
+	virtual void RHIDiscardRenderTargets(bool Depth,bool Stencil,uint32 ColorBitMask) final override {}
 	virtual void RHIBlockUntilGPUIdle() final override;
 	virtual void RHISuspendRendering() final override;
 	virtual void RHIResumeRendering() final override;
@@ -179,7 +179,9 @@ protected:
 
 	friend class FVulkanViewport;
 
+/*
 	static void WriteEndFrameTimestamp(void*);
+*/
 
 	static void GetInstanceLayersAndExtensions(TArray<const ANSICHAR*>& OutInstanceExtensions, TArray<const ANSICHAR*>& OutInstanceLayers);
 
diff --git a/Engine/Source/Runtime/VulkanRHI/Public/VulkanResources.h b/Engine/Source/Runtime/VulkanRHI/Public/VulkanResources.h
--- a/Engine/Source/Runtime/VulkanRHI/Public/VulkanResources.h
+++ b/Engine/Source/Runtime/VulkanRHI/Public/VulkanResources.h
@@ -523,7 +523,7 @@ public:
 	virtual ~FVulkanTimestampQueryPool(){}
 
 	void WriteStartFrame(VkCommandBuffer CmdBuffer);
-	void WriteEndFrame(VkCommandBuffer CmdBuffer);
+	void WriteEndFrame(FVulkanCmdBuffer* CmdBuffer);
 
 	void CalculateFrameTime();
 
@@ -551,6 +551,9 @@ protected:
 	double SecondsPerTimestamp;
 	int32 UsedUserQueries;
 	bool bFirst;
+
+	FVulkanCmdBuffer* LastEndCmdBuffer;
+	uint64 LastFenceSignaledCounter;
 };
 
 /** Vulkan occlusion query */
@@ -981,7 +984,7 @@ private:
 
 	void GenerateVertexInputStateInfo();
 
-	void GetDescriptorInfoCounts(uint32& OutSamplerCount, uint32& OutSamplerBufferCount, uint32& OutUniformBufferCount);
+	void GetDescriptorInfoCounts(uint32& OutCombinedSamplerCount, uint32& OutSamplerBufferCount, uint32& OutUniformBufferCount);
 
 	void InternalBindVertexStreams(FVulkanCmdBuffer* Cmd, const void* VertexStreams);
 
@@ -1030,7 +1033,7 @@ private:
 	// At this moment unused, the samplers are mapped together with texture/image
 	// we are using "VK_DESCRIPTOR_TYPE_COMBINED_IMAGE_SAMPLER"
 	bool DirtySamplerStates[SF_Compute];
-	VkDescriptorImageInfo* DescriptorImageInfoForStage[SF_Compute];
+	VkDescriptorImageInfo* DescriptorSamplerImageInfoForStage[SF_Compute];
 
 	TArray< VkWriteDescriptorSet* > SRVWriteInfoForStage[SF_Compute];
 	bool DirtySRVs[SF_Compute];
@@ -1051,7 +1054,7 @@ private:
 	// InOutMask is used from VertexShader-Header to filter out active attributs
 	void ProcessVertexElementForBinding(const FVertexElement& Element);
 
-	void CreateDescriptorWriteInfo(uint32 UniformSamplerCount, uint32 UniformSamplerBufferCount, uint32 UniformBufferCount);
+	void CreateDescriptorWriteInfo(uint32 UniformCombinedSamplerCount, uint32 UniformSamplerBufferCount, uint32 UniformBufferCount);
 
 	uint32 BindingsNum;
 	uint32 BindingsMask;
diff --git a/Engine/Source/Runtime/VulkanRHI/Public/VulkanShaderResources.h b/Engine/Source/Runtime/VulkanRHI/Public/VulkanShaderResources.h
--- a/Engine/Source/Runtime/VulkanRHI/Public/VulkanShaderResources.h
+++ b/Engine/Source/Runtime/VulkanRHI/Public/VulkanShaderResources.h
@@ -73,10 +73,12 @@ public:
 
 	enum EBindingType : uint16
 	{
-		TYPE_SAMPLER,
+		TYPE_COMBINED_IMAGE_SAMPLER,
 		TYPE_SAMPLER_BUFFER,
 		TYPE_UNIFORM_BUFFER,
 		TYPE_PACKED_UNIFORM_BUFFER,
+		//TYPE_SAMPLER,
+		//TYPE_IMAGE,
 
 		TYPE_MAX,
 	};
@@ -180,7 +182,7 @@ inline FArchive& operator<<(FArchive& Ar, FVulkanCodeHeader& Header)
 class FVulkanShaderBindingTable
 {
 public:
-	TArray<uint32> SamplerBindingIndices;
+	TArray<uint32> CombinedSamplerBindingIndices;
 	//#todo-rco: FIX! SamplerBuffers share numbering with Samplers
 	TArray<uint32> SamplerBufferBindingIndices;
 	TArray<uint32> UniformBufferBindingIndices;
@@ -199,7 +201,7 @@ public:
 
 inline FArchive& operator<<(FArchive& Ar, FVulkanShaderBindingTable& BindingTable)
 {
-	Ar << BindingTable.SamplerBindingIndices;
+	Ar << BindingTable.CombinedSamplerBindingIndices;
 	Ar << BindingTable.SamplerBufferBindingIndices;
 	Ar << BindingTable.UniformBufferBindingIndices;
 	for (int32 Index = 0; Index < CrossCompiler::PACKED_TYPEINDEX_MAX; ++Index)
diff --git a/Engine/Source/Runtime/Windows/D3D11RHI/Private/D3D11Device.cpp b/Engine/Source/Runtime/Windows/D3D11RHI/Private/D3D11Device.cpp
--- a/Engine/Source/Runtime/Windows/D3D11RHI/Private/D3D11Device.cpp
+++ b/Engine/Source/Runtime/Windows/D3D11RHI/Private/D3D11Device.cpp
@@ -167,6 +167,7 @@ FD3D11DynamicRHI::FD3D11DynamicRHI(IDXGIFactory1* InDXGIFactory1,D3D_FEATURE_LEV
 
 	GPixelFormats[ PF_BC6H			].PlatformFormat	= DXGI_FORMAT_BC6H_UF16;
 	GPixelFormats[ PF_BC7			].PlatformFormat	= DXGI_FORMAT_BC7_TYPELESS;
+	GPixelFormats[ PF_R8_UINT		].PlatformFormat	= DXGI_FORMAT_R8_UINT;
 
 	if (FeatureLevel >= D3D_FEATURE_LEVEL_11_0)
 	{
@@ -469,6 +470,14 @@ void FD3D11DynamicRHI::CleanupD3DDevice()
 		{
 			Direct3DDeviceIMContext->ClearState();
 			Direct3DDeviceIMContext->Flush();
+
+			// Perform a detailed live object report (with resource types)
+			ID3D11Debug* D3D11Debug;
+			Direct3DDevice->QueryInterface(__uuidof(ID3D11Debug), (void**)(&D3D11Debug));
+			if (D3D11Debug)
+			{
+				D3D11Debug->ReportLiveDeviceObjects(D3D11_RLDO_DETAIL);
+			}
 		}
 
 		Direct3DDeviceIMContext = nullptr;
diff --git a/Engine/Source/Runtime/Windows/D3D12RHI/Private/D3D12Commands.cpp b/Engine/Source/Runtime/Windows/D3D12RHI/Private/D3D12Commands.cpp
--- a/Engine/Source/Runtime/Windows/D3D12RHI/Private/D3D12Commands.cpp
+++ b/Engine/Source/Runtime/Windows/D3D12RHI/Private/D3D12Commands.cpp
@@ -1221,6 +1221,7 @@ void FD3D12CommandContext::RHIBeginRenderQuery(FRenderQueryRHIParamRef QueryRHI)
 	{
 		Query->bResultIsCached = false;
 		Query->HeapIndex = GetParentDevice()->GetQueryHeap()->BeginQuery(*this, D3D12_QUERY_TYPE_OCCLUSION);
+		Query->OwningContext = this;
 	}
 	else
 	{
@@ -1251,6 +1252,7 @@ void FD3D12CommandContext::RHIEndRenderQuery(FRenderQueryRHIParamRef QueryRHI)
 		Query->CLSyncPoint = CommandListHandle;
 
 		CommandListHandle.UpdateResidency(pQueryHeap->GetResultBuffer());
+		check(Query->OwningContext == this);
 		break;
 	}
 
@@ -1258,6 +1260,7 @@ void FD3D12CommandContext::RHIEndRenderQuery(FRenderQueryRHIParamRef QueryRHI)
 	{
 		Query->bResultIsCached = false;
 		Query->CLSyncPoint = CommandListHandle;
+		Query->OwningContext = this;
 		this->otherWorkCounter += 2;	// +2 For the EndQuery and the ResolveQueryData
 		CommandListHandle->EndQuery(Query->QueryHeap, D3D12_QUERY_TYPE_TIMESTAMP, Query->HeapIndex);
 		CommandListHandle->ResolveQueryData(Query->QueryHeap, D3D12_QUERY_TYPE_TIMESTAMP, Query->HeapIndex, 1, Query->ResultBuffer, sizeof(uint64) * Query->HeapIndex);
diff --git a/Engine/Source/Runtime/Windows/D3D12RHI/Private/D3D12Device.cpp b/Engine/Source/Runtime/Windows/D3D12RHI/Private/D3D12Device.cpp
--- a/Engine/Source/Runtime/Windows/D3D12RHI/Private/D3D12Device.cpp
+++ b/Engine/Source/Runtime/Windows/D3D12RHI/Private/D3D12Device.cpp
@@ -438,6 +438,7 @@ FD3D12DynamicRHI::FD3D12DynamicRHI(IDXGIFactory4* InDXGIFactory, FD3D12Adapter&
 
 	GPixelFormats[ PF_BC6H			].PlatformFormat	= DXGI_FORMAT_BC6H_UF16;
 	GPixelFormats[ PF_BC7			].PlatformFormat	= DXGI_FORMAT_BC7_TYPELESS;
+	GPixelFormats[ PF_R8_UINT		].PlatformFormat	= DXGI_FORMAT_R8_UINT;
 
 	// MS - Not doing any feature level checks. D3D12 currently supports these limits.
 	// However this may need to be revisited if new feature levels are introduced with different HW requirement
diff --git a/Engine/Source/Runtime/Windows/D3D12RHI/Private/D3D12Query.cpp b/Engine/Source/Runtime/Windows/D3D12RHI/Private/D3D12Query.cpp
--- a/Engine/Source/Runtime/Windows/D3D12RHI/Private/D3D12Query.cpp
+++ b/Engine/Source/Runtime/Windows/D3D12RHI/Private/D3D12Query.cpp
@@ -92,6 +92,7 @@ bool FD3D12Device::GetQueryData(FD3D12OcclusionQuery& Query, bool bWait)
 {
 	// Wait for the query result to be ready (if requested).
 	const FD3D12CLSyncPoint& SyncPoint = Query.CLSyncPoint;
+
 	if (!SyncPoint.IsComplete())
 	{
 		if (!bWait)
@@ -99,8 +100,25 @@ bool FD3D12Device::GetQueryData(FD3D12OcclusionQuery& Query, bool bWait)
 			return false;
 		}
 
-		check(!SyncPoint.IsOpen());
+		SCOPE_CYCLE_COUNTER(STAT_RenderQueryResultTime);
+		uint32 IdleStart = FPlatformTime::Cycles();
+		double StartTime = FPlatformTime::Seconds();
+
+		if (SyncPoint.IsOpen())
+		{
+			// The query is on a command list that hasn't been submitted yet.
+			// We need to flush, but the RHI thread may be using the default command list...so stall it first.
+			check(IsInRenderingThread());
+			check(Query.OwningContext->IsDefaultContext());
+
+			FScopedRHIThreadStaller StallRHIThread(FRHICommandListExecutor::GetImmediateCommandList());
+			Query.OwningContext->FlushCommands(); // Don't wait yet, since we're stalling the RHI thread.
+		}
+
 		SyncPoint.WaitForCompletion();
+
+		GRenderThreadIdle[ERenderThreadIdleTypes::WaitingForGPUQuery] += FPlatformTime::Cycles() - IdleStart;
+		GRenderThreadNumIdle[ERenderThreadIdleTypes::WaitingForGPUQuery]++;
 	}
 
 	// Read the data from the query's buffer.
diff --git a/Engine/Source/Runtime/Windows/D3D12RHI/Private/D3D12Query.h b/Engine/Source/Runtime/Windows/D3D12RHI/Private/D3D12Query.h
--- a/Engine/Source/Runtime/Windows/D3D12RHI/Private/D3D12Query.h
+++ b/Engine/Source/Runtime/Windows/D3D12RHI/Private/D3D12Query.h
@@ -26,6 +26,8 @@ public:
 	// todo: memory optimize
 	const ERenderQueryType Type;
 
+	class FD3D12CommandContext* OwningContext;
+
 	// When the query result is ready on the GPU.
 	FD3D12CLSyncPoint CLSyncPoint;
 
@@ -43,6 +45,7 @@ public:
 	{
 		HeapIndex = -1;
 		bResultIsCached = false;
+		OwningContext = nullptr;
 	}
 };
 
