Compile error fix in Vulkan header.

diff --git a/Engine/Source/Runtime/VulkanRHI/Public/VulkanResources.h b/Engine/Source/Runtime/VulkanRHI/Public/VulkanResources.h
--- a/Engine/Source/Runtime/VulkanRHI/Public/VulkanResources.h
+++ b/Engine/Source/Runtime/VulkanRHI/Public/VulkanResources.h
@@ -278,7 +278,7 @@ struct FVulkanTextureView
 	{
 	}
 
-	static VkImageView FVulkanTextureView::StaticCreate(FVulkanDevice& Device, VkImage Image, VkImageViewType ViewType, VkImageAspectFlags AspectFlags, EPixelFormat UEFormat, VkFormat Format, uint32 FirstMip, uint32 NumMips, uint32 ArraySliceIndex, uint32 NumArraySlices);
+	static VkImageView StaticCreate(FVulkanDevice& Device, VkImage Image, VkImageViewType ViewType, VkImageAspectFlags AspectFlags, EPixelFormat UEFormat, VkFormat Format, uint32 FirstMip, uint32 NumMips, uint32 ArraySliceIndex, uint32 NumArraySlices);
 
 	void Create(FVulkanDevice& Device, VkImage Image, VkImageViewType ViewType, VkImageAspectFlags AspectFlags, EPixelFormat UEFormat, VkFormat Format, uint32 NumMips, uint32 NumArraySlices);
 	void Destroy(FVulkanDevice& Device);
@@ -299,7 +299,7 @@ struct FVulkanTextureBase : public FVulkanBaseShaderResource
 	FVulkanTextureBase(FVulkanDevice& Device, VkImageViewType ResourceType, EPixelFormat Format, uint32 SizeX, uint32 SizeY, uint32 SizeZ, VkImage InImage, VkDeviceMemory InMem, uint32 UEFlags, const FRHIResourceCreateInfo& CreateInfo = FRHIResourceCreateInfo());
 	virtual ~FVulkanTextureBase();
 
-	VkImageView FVulkanTextureBase::CreateRenderTargetView(uint32 MipIndex, uint32 ArraySliceIndex);
+	VkImageView CreateRenderTargetView(uint32 MipIndex, uint32 ArraySliceIndex);
 
 	FVulkanSurface Surface;
 
