Copying //UE4/Dev-Build to //UE4/Dev-Main (Source: //UE4/Dev-Build @ 3092544)

diff --git a/Engine/Source/Runtime/RenderCore/Public/TickableObjectRenderThread.h b/Engine/Source/Runtime/RenderCore/Public/TickableObjectRenderThread.h
--- a/Engine/Source/Runtime/RenderCore/Public/TickableObjectRenderThread.h
+++ b/Engine/Source/Runtime/RenderCore/Public/TickableObjectRenderThread.h
@@ -4,8 +4,7 @@
 	TickableObjectRenderThread.h: Rendering thread tickable object definition.
 =============================================================================*/
 
-#ifndef __TickableObjectRenderThread_H__
-#define __TickableObjectRenderThread_H__
+#pragma once
 
 #include "RenderCore.h"
 
@@ -135,4 +134,3 @@ private:
 	bool bHighFrequency;
 };
 
-#endif
diff --git a/Engine/Source/Runtime/Renderer/Private/DynamicPrimitiveDrawing.h b/Engine/Source/Runtime/Renderer/Private/DynamicPrimitiveDrawing.h
--- a/Engine/Source/Runtime/Renderer/Private/DynamicPrimitiveDrawing.h
+++ b/Engine/Source/Runtime/Renderer/Private/DynamicPrimitiveDrawing.h
@@ -4,8 +4,7 @@
 	DynamicPrimitiveDrawing.h: Dynamic primitive drawing definitions.
 =============================================================================*/
 
-#ifndef __DYNAMICPRIMITIVEDRAWING_H__
-#define __DYNAMICPRIMITIVEDRAWING_H__
+#pragma once
 
 /**
 * Draws a range of view's elements with the specified drawing policy factory type.
@@ -112,4 +111,3 @@ private:
 
 #include "DynamicPrimitiveDrawing.inl"
 
-#endif
diff --git a/Engine/Source/Runtime/Renderer/Private/DynamicPrimitiveDrawing.inl b/Engine/Source/Runtime/Renderer/Private/DynamicPrimitiveDrawing.inl
--- a/Engine/Source/Runtime/Renderer/Private/DynamicPrimitiveDrawing.inl
+++ b/Engine/Source/Runtime/Renderer/Private/DynamicPrimitiveDrawing.inl
@@ -4,8 +4,7 @@
 	DynamicPrimitiveDrawing.inl: Dynamic primitive drawing implementation.
 =============================================================================*/
 
-#ifndef __DYNAMICPRIMITIVEDRAWING_INL__
-#define __DYNAMICPRIMITIVEDRAWING_INL__
+#pragma once
 
 template<class DrawingPolicyFactoryType>
 void DrawViewElementsInner(
@@ -307,4 +306,3 @@ inline int32 FViewElementPDI::DrawMesh(const FMeshBatch& Mesh)
 	return 0;
 }
 
-#endif
diff --git a/Engine/Source/Runtime/Renderer/Private/LightRendering.h b/Engine/Source/Runtime/Renderer/Private/LightRendering.h
--- a/Engine/Source/Runtime/Renderer/Private/LightRendering.h
+++ b/Engine/Source/Runtime/Renderer/Private/LightRendering.h
@@ -4,8 +4,7 @@
 	LightRendering.h: Light rendering declarations.
 =============================================================================*/
 
-#ifndef __LIGHTRENDERING_H__
-#define __LIGHTRENDERING_H__
+#pragma once
 
 #include "UniformBuffer.h"
 
@@ -121,4 +120,3 @@ private:
 
 extern void SetBoundingGeometryRasterizerAndDepthState(FRHICommandList& RHICmdList, const FViewInfo& View, const FSphere& LightBounds);
 
-#endif
diff --git a/Engine/Source/Runtime/Renderer/Private/PostProcess/SceneFilterRendering.h b/Engine/Source/Runtime/Renderer/Private/PostProcess/SceneFilterRendering.h
--- a/Engine/Source/Runtime/Renderer/Private/PostProcess/SceneFilterRendering.h
+++ b/Engine/Source/Runtime/Renderer/Private/PostProcess/SceneFilterRendering.h
@@ -4,8 +4,7 @@
 	SceneFilterRendering.h: Filter rendering definitions.
 =============================================================================*/
 
-#ifndef _INC_SCENEFILTERRENDERING
-#define _INC_SCENEFILTERRENDERING
+#pragma once
 
 #include "SceneRenderTargets.h"
 
@@ -208,4 +207,3 @@ public:
 	uint32 NumPrimitives() const;
 };
 
-#endif
diff --git a/Engine/Source/Runtime/Renderer/Private/RendererPrivate.h b/Engine/Source/Runtime/Renderer/Private/RendererPrivate.h
--- a/Engine/Source/Runtime/Renderer/Private/RendererPrivate.h
+++ b/Engine/Source/Runtime/Renderer/Private/RendererPrivate.h
@@ -4,8 +4,7 @@
 	RendererPrivate.h: Renderer interface private definitions.
 =============================================================================*/
 
-#ifndef __RendererPrivate_h__
-#define __RendererPrivate_h__
+#pragma once
 
 #include "Engine.h"
 #include "ShaderCore.h"
@@ -86,4 +85,3 @@ private:
 
 extern ICustomCulling* GCustomCullingImpl;
 
-#endif
diff --git a/Engine/Source/Runtime/Renderer/Private/ScenePrivate.h b/Engine/Source/Runtime/Renderer/Private/ScenePrivate.h
--- a/Engine/Source/Runtime/Renderer/Private/ScenePrivate.h
+++ b/Engine/Source/Runtime/Renderer/Private/ScenePrivate.h
@@ -4,8 +4,7 @@
 	ScenePrivate.h: Private scene manager definitions.
 =============================================================================*/
 
-#ifndef __SCENEPRIVATE_H__
-#define __SCENEPRIVATE_H__
+#pragma once
 
 class SceneRenderingAllocator;
 class USceneCaptureComponent;
@@ -2264,4 +2263,3 @@ private:
 
 #include "BasePassRendering.inl"
 
-#endif // __SCENEPRIVATE_H__
diff --git a/Engine/Source/Runtime/Renderer/Private/ShadowRendering.h b/Engine/Source/Runtime/Renderer/Private/ShadowRendering.h
--- a/Engine/Source/Runtime/Renderer/Private/ShadowRendering.h
+++ b/Engine/Source/Runtime/Renderer/Private/ShadowRendering.h
@@ -4,8 +4,7 @@
 	ShadowRendering.h: Shadow rendering definitions.
 =============================================================================*/
 
-#ifndef __ShadowRendering_H__
-#define __ShadowRendering_H__
+#pragma once
 
 #include "ShaderParameterUtils.h"
 #include "SceneCore.h"
@@ -1887,4 +1886,3 @@ struct FCompareFProjectedShadowInfoBySplitIndex
 	}
 };
 
-#endif
diff --git a/Engine/Source/Runtime/Renderer/Private/StaticMeshDrawList.h b/Engine/Source/Runtime/Renderer/Private/StaticMeshDrawList.h
--- a/Engine/Source/Runtime/Renderer/Private/StaticMeshDrawList.h
+++ b/Engine/Source/Runtime/Renderer/Private/StaticMeshDrawList.h
@@ -4,8 +4,7 @@
 	StaticMeshDrawList.h: Static mesh draw list definition.
 =============================================================================*/
 
-#ifndef __STATICMESHDRAWLIST_H__
-#define __STATICMESHDRAWLIST_H__
+#pragma once
 
 extern ENGINE_API bool GDrawListsLocked;
 /** View state for instanced stereo rendering. */
@@ -499,4 +498,3 @@ public:
 
 #include "StaticMeshDrawList.inl"
 
-#endif
diff --git a/Engine/Source/Runtime/Renderer/Public/DrawingPolicy.h b/Engine/Source/Runtime/Renderer/Public/DrawingPolicy.h
--- a/Engine/Source/Runtime/Renderer/Public/DrawingPolicy.h
+++ b/Engine/Source/Runtime/Renderer/Public/DrawingPolicy.h
@@ -4,8 +4,7 @@
 	DrawingPolicy.h: Drawing policy definitions.
 =============================================================================*/
 
-#ifndef __DRAWINGPOLICY_H__
-#define __DRAWINGPOLICY_H__
+#pragma once
 
 #include "RHICommandList.h"
 
@@ -276,4 +275,3 @@ protected:
 
 uint32 GetTypeHash(const FBoundShaderStateRHIRef &Key);
 
-#endif
diff --git a/Engine/Source/Runtime/Renderer/Public/PrimitiveSceneInfo.h b/Engine/Source/Runtime/Renderer/Public/PrimitiveSceneInfo.h
--- a/Engine/Source/Runtime/Renderer/Public/PrimitiveSceneInfo.h
+++ b/Engine/Source/Runtime/Renderer/Public/PrimitiveSceneInfo.h
@@ -4,8 +4,7 @@
 	PrimitiveSceneInfo.h: Primitive scene info definitions.
 =============================================================================*/
 
-#ifndef __PRIMITIVESCENEINFO_H__
-#define __PRIMITIVESCENEINFO_H__
+#pragma once
 
 #include "GenericOctree.h"
 
@@ -460,4 +459,3 @@ struct FPrimitiveOctreeSemantics
 	}
 };
 
-#endif
diff --git a/Engine/Source/Runtime/Renderer/Public/RendererInterface.h b/Engine/Source/Runtime/Renderer/Public/RendererInterface.h
--- a/Engine/Source/Runtime/Renderer/Public/RendererInterface.h
+++ b/Engine/Source/Runtime/Renderer/Public/RendererInterface.h
@@ -4,8 +4,7 @@
 	RendererInterface.h: Renderer interface definition.
 =============================================================================*/
 
-#ifndef __RendererInterface_H__
-#define __RendererInterface_H__
+#pragma once
 
 #include "ModuleInterface.h"
 #include "ModuleManager.h"
@@ -674,6 +673,3 @@ public:
 	virtual bool HasPostOpaqueExtentions() const = 0;
 };
 
-
-
-#endif
diff --git a/Engine/Source/Runtime/ShaderCore/Private/Shader.cpp b/Engine/Source/Runtime/ShaderCore/Private/Shader.cpp
--- a/Engine/Source/Runtime/ShaderCore/Private/Shader.cpp
+++ b/Engine/Source/Runtime/ShaderCore/Private/Shader.cpp
@@ -4,6 +4,7 @@
 	Shader.cpp: Shader implementation.
 =============================================================================*/
 
+#include "ShaderCorePrivatePCH.h"
 #include "ShaderCore.h"
 #include "Shader.h"
 #include "VertexFactory.h"
diff --git a/Engine/Source/Runtime/ShaderCore/Private/ShaderCache.cpp b/Engine/Source/Runtime/ShaderCore/Private/ShaderCache.cpp
--- a/Engine/Source/Runtime/ShaderCore/Private/ShaderCache.cpp
+++ b/Engine/Source/Runtime/ShaderCore/Private/ShaderCache.cpp
@@ -4,6 +4,7 @@
 	ShaderCache.cpp: Bound shader state cache implementation.
 =============================================================================*/
 
+#include "ShaderCorePrivatePCH.h"
 #include "ShaderCore.h"
 #include "ShaderCache.h"
 #include "Shader.h"
diff --git a/Engine/Source/Runtime/ShaderCore/Private/ShaderCorePrivatePCH.h b/Engine/Source/Runtime/ShaderCore/Private/ShaderCorePrivatePCH.h
new file mode 100644
--- /dev/null
+++ b/Engine/Source/Runtime/ShaderCore/Private/ShaderCorePrivatePCH.h
@@ -0,0 +1,5 @@
+// Copyright 1998-2016 Epic Games, Inc. All Rights Reserved.
+
+#pragma once
+
+#include "ShaderCore.h"
diff --git a/Engine/Source/Runtime/ShaderCore/Public/Shader.h b/Engine/Source/Runtime/ShaderCore/Public/Shader.h
--- a/Engine/Source/Runtime/ShaderCore/Public/Shader.h
+++ b/Engine/Source/Runtime/ShaderCore/Public/Shader.h
@@ -252,7 +252,7 @@ private:
 	uint32 Canary;
 
 	/** Initialize the shader RHI resources. */
-	void InitializeShaderRHI();
+	SHADERCORE_API void InitializeShaderRHI();
 
 	/** Tracks loaded shader resources by id. */
 	static TMap<FShaderResourceId, FShaderResource*> ShaderResourceIdMap;
diff --git a/Engine/Source/Runtime/Slate/Private/Framework/Docking/DockingPrivate.h b/Engine/Source/Runtime/Slate/Private/Framework/Docking/DockingPrivate.h
--- a/Engine/Source/Runtime/Slate/Private/Framework/Docking/DockingPrivate.h
+++ b/Engine/Source/Runtime/Slate/Private/Framework/Docking/DockingPrivate.h
@@ -5,8 +5,6 @@
 
 #include "SlateBasics.h"
 
-#define DEBUG_TAB_MANAGEMENT 0
-
 #include "SDockingNode.h"
 #include "SDockingSplitter.h"
 #include "SDockingTabStack.h"
diff --git a/Engine/Source/Runtime/Slate/Public/Framework/Commands/Commands.h b/Engine/Source/Runtime/Slate/Public/Framework/Commands/Commands.h
--- a/Engine/Source/Runtime/Slate/Public/Framework/Commands/Commands.h
+++ b/Engine/Source/Runtime/Slate/Public/Framework/Commands/Commands.h
@@ -27,7 +27,7 @@ public:
 	
 
 	/** Use this method to register commands. Usually done in StartupModule(). */
-	static void Register()
+	FORCENOINLINE static void Register()
 	{
 		if ( !Instance.IsValid() )
 		{
@@ -56,13 +56,13 @@ public:
 	}
 
 	/** Get the singleton instance of this set of commands. */
-	static const CommandContextType& Get()
+	FORCENOINLINE static const CommandContextType& Get()
 	{
 		return *(Instance.Pin());
 	}
 
 	/** Use this method to clean up any resources used by the command set. Usually done in ShutdownModule() */
-	static void Unregister()
+	FORCENOINLINE static void Unregister()
 	{
 		// The instance may not be valid if it was never used.
 		if( Instance.IsValid() )
@@ -77,7 +77,7 @@ public:
 	}
 
 	/** Get the BindingContext for this set of commands. */
-	static const FBindingContext& GetContext()
+	FORCENOINLINE static const FBindingContext& GetContext()
 	{
 		check( Instance.IsValid() );
 		return *(Instance.Pin());
diff --git a/Engine/Source/Runtime/Slate/Public/Framework/Commands/UIAction.h b/Engine/Source/Runtime/Slate/Public/Framework/Commands/UIAction.h
--- a/Engine/Source/Runtime/Slate/Public/Framework/Commands/UIAction.h
+++ b/Engine/Source/Runtime/Slate/Public/Framework/Commands/UIAction.h
@@ -30,7 +30,7 @@ enum class EUIActionRepeatMode
 /**
  * Implements an UI action.
  */
-struct FUIAction
+struct SLATE_API FUIAction
 {
 	/** Holds a delegate that is executed when this action is activated. */
 	FExecuteAction ExecuteAction;
@@ -50,13 +50,7 @@ struct FUIAction
 public:
 
 	/** Default constructor. */
-	FUIAction( )
-		: ExecuteAction()
-		, CanExecuteAction()
-		, GetActionCheckState()
-		, IsActionVisibleDelegate()
-		, RepeatMode(EUIActionRepeatMode::RepeatDisabled)
-	{ }
+	FUIAction();
 
 	/**
 	 * Constructor that takes delegates to initialize the action with
@@ -64,13 +58,7 @@ public:
 	 * @param	ExecuteAction		The delegate to call when the action should be executed
 	 * @param	RepeatMode			Can this action can be repeated if the chord used to call it is held down?
 	 */
-	FUIAction( FExecuteAction InitExecuteAction, EUIActionRepeatMode InitRepeatMode = EUIActionRepeatMode::RepeatDisabled )
-		: ExecuteAction(InitExecuteAction)
-		, CanExecuteAction(FCanExecuteAction::CreateRaw(&FSlateApplication::Get(), &FSlateApplication::IsNormalExecution))
-		, GetActionCheckState()
-		, IsActionVisibleDelegate()
-		, RepeatMode(InitRepeatMode)
-	{ }
+	FUIAction(FExecuteAction InitExecuteAction, EUIActionRepeatMode InitRepeatMode = EUIActionRepeatMode::RepeatDisabled);
 
 	/**
 	 * Constructor that takes delegates to initialize the action with
@@ -79,13 +67,7 @@ public:
 	 * @param	CanExecuteAction	The delegate to call to see if the action can be executed
 	 * @param	RepeatMode			Can this action can be repeated if the chord used to call it is held down?
 	 */
-	FUIAction( FExecuteAction InitExecuteAction, FCanExecuteAction InitCanExecuteAction, EUIActionRepeatMode InitRepeatMode = EUIActionRepeatMode::RepeatDisabled )
-		: ExecuteAction(InitExecuteAction )
-		, CanExecuteAction(InitCanExecuteAction)
-		, GetActionCheckState()
-		, IsActionVisibleDelegate()
-		, RepeatMode(InitRepeatMode)
-	{ }
+	FUIAction(FExecuteAction InitExecuteAction, FCanExecuteAction InitCanExecuteAction, EUIActionRepeatMode InitRepeatMode = EUIActionRepeatMode::RepeatDisabled);
 
 	/**
 	 * Constructor that takes delegates to initialize the action with
@@ -95,13 +77,7 @@ public:
 	 * @param	IsCheckedDelegate	The delegate to call to see if the action should appear checked when visualized
 	 * @param	RepeatMode			Can this action can be repeated if the chord used to call it is held down?
 	 */
-	FUIAction( FExecuteAction InitExecuteAction, FCanExecuteAction InitCanExecuteAction, FIsActionChecked InitIsCheckedDelegate, EUIActionRepeatMode InitRepeatMode = EUIActionRepeatMode::RepeatDisabled )
-		: ExecuteAction(InitExecuteAction)
-		, CanExecuteAction(InitCanExecuteAction)
-		, GetActionCheckState(FGetActionCheckState::CreateStatic(&FUIAction::IsActionCheckedPassthrough, InitIsCheckedDelegate))
-		, IsActionVisibleDelegate()
-		, RepeatMode(InitRepeatMode)
-	{ }
+	FUIAction(FExecuteAction InitExecuteAction, FCanExecuteAction InitCanExecuteAction, FIsActionChecked InitIsCheckedDelegate, EUIActionRepeatMode InitRepeatMode = EUIActionRepeatMode::RepeatDisabled);
 
 	/**
 	 * Constructor that takes delegates to initialize the action with
@@ -111,13 +87,7 @@ public:
 	 * @param	GetActionCheckState	The delegate to call to see what the check state of the action should be
 	 * @param	RepeatMode			Can this action can be repeated if the chord used to call it is held down?
 	 */
-	FUIAction( FExecuteAction InitExecuteAction, FCanExecuteAction InitCanExecuteAction, FGetActionCheckState InitGetActionCheckStateDelegate, EUIActionRepeatMode InitRepeatMode = EUIActionRepeatMode::RepeatDisabled )
-		: ExecuteAction(InitExecuteAction)
-		, CanExecuteAction(InitCanExecuteAction)
-		, GetActionCheckState(InitGetActionCheckStateDelegate)
-		, IsActionVisibleDelegate()
-		, RepeatMode(InitRepeatMode)
-	{ }
+	FUIAction(FExecuteAction InitExecuteAction, FCanExecuteAction InitCanExecuteAction, FGetActionCheckState InitGetActionCheckStateDelegate, EUIActionRepeatMode InitRepeatMode = EUIActionRepeatMode::RepeatDisabled);
 
 	/**
 	 * Constructor that takes delegates to initialize the action with
@@ -128,13 +98,7 @@ public:
 	 * @param	IsActionVisible		The delegate to call to see if the action should be visible
 	 * @param	RepeatMode			Can this action can be repeated if the chord used to call it is held down?
 	 */
-	FUIAction( FExecuteAction InitExecuteAction, FCanExecuteAction InitCanExecuteAction, FIsActionChecked InitIsCheckedDelegate, FIsActionButtonVisible InitIsActionVisibleDelegate, EUIActionRepeatMode InitRepeatMode = EUIActionRepeatMode::RepeatDisabled )
-		: ExecuteAction(InitExecuteAction)
-		, CanExecuteAction(InitCanExecuteAction)
-		, GetActionCheckState(FGetActionCheckState::CreateStatic(&FUIAction::IsActionCheckedPassthrough, InitIsCheckedDelegate))
-		, IsActionVisibleDelegate(InitIsActionVisibleDelegate)
-		, RepeatMode(InitRepeatMode)
-	{ }
+	FUIAction(FExecuteAction InitExecuteAction, FCanExecuteAction InitCanExecuteAction, FIsActionChecked InitIsCheckedDelegate, FIsActionButtonVisible InitIsActionVisibleDelegate, EUIActionRepeatMode InitRepeatMode = EUIActionRepeatMode::RepeatDisabled);
 
 	/**
 	 * Constructor that takes delegates to initialize the action with
@@ -145,13 +109,7 @@ public:
 	 * @param	IsActionVisible		The delegate to call to see if the action should be visible
 	 * @param	RepeatMode			Can this action can be repeated if the chord used to call it is held down?
 	 */
-	FUIAction( FExecuteAction InitExecuteAction, FCanExecuteAction InitCanExecuteAction, FGetActionCheckState InitGetActionCheckStateDelegate, FIsActionButtonVisible InitIsActionVisibleDelegate, EUIActionRepeatMode InitRepeatMode = EUIActionRepeatMode::RepeatDisabled )
-		: ExecuteAction(InitExecuteAction)
-		, CanExecuteAction(InitCanExecuteAction)
-		, GetActionCheckState(InitGetActionCheckStateDelegate)
-		, IsActionVisibleDelegate(InitIsActionVisibleDelegate)
-		, RepeatMode(InitRepeatMode)
-	{ }
+	FUIAction(FExecuteAction InitExecuteAction, FCanExecuteAction InitCanExecuteAction, FGetActionCheckState InitGetActionCheckStateDelegate, FIsActionButtonVisible InitIsActionVisibleDelegate, EUIActionRepeatMode InitRepeatMode = EUIActionRepeatMode::RepeatDisabled);
 
 public:
 
diff --git a/Engine/Source/Runtime/Slate/Public/SlateBasics.h b/Engine/Source/Runtime/Slate/Public/SlateBasics.h
--- a/Engine/Source/Runtime/Slate/Public/SlateBasics.h
+++ b/Engine/Source/Runtime/Slate/Public/SlateBasics.h
@@ -38,11 +38,6 @@
 #include "UICommandList.h"
 
 // Legacy
-class SMissingWidget;
-class FSlateStyleRegistry;
-class FUICommandDragDropOp;
-class FInertialScrollManager;
-class FGenericCommands;
 #include "SWeakWidget.h"
 #include "TextLayoutEngine.h"
 #include "SPanel.h"
@@ -75,7 +70,6 @@ class FGenericCommands;
 #include "SMultiLineEditableTextBox.h"
 #include "SEditableText.h"
 #include "SEditableTextBox.h"
-class SSearchBox;
 #include "SButton.h"
 #include "SToolTip.h"
 #include "SScrollBar.h"
@@ -99,50 +93,13 @@ class SSearchBox;
 #include "SCheckBox.h"
 #include "SSpinBox.h"
 #include "SSlider.h"
-class SVolumeControl;
-class SColorSpectrum;
-class SColorWheel;
 
-class FToolBarButtonBlock;
-class SToolBarButtonBlock;
-class FToolBarComboButtonBlock;
-class SToolBarComboButtonBlock;
-class SPopup;
 ///
 #include "SComboButton.h"
 #include "SComboBox.h"
-class SHyperlink;
-class SRichTextHyperlink;
-class SThrobber;
-class SCircularThrobber;
-class STextEntryPopup;
-class STextComboPopup;
-class SExpandableButton;
-class SExpandableArea;
-class SNotificationItem;
-class SNotificationList;
-class SWidgetSwitcher;
-class SSuggestionTextBox;
-template <typename ItemType> class SBreadcrumbTrail;
-class STextComboBox;
-template<typename NumericType> class SNumericEntryBox;
-template<typename OptionType> class SEditableComboBox;
-class FSlateNotificationManager;
-class SDPIScaler;
-class SInlineEditableTextBlock;
-class SVirtualKeyboardEntry;
-class SSafeZone;
-struct FMarqueeRect;
-class SRotatorInputBox;
-class SVectorInputBox;
-class SVirtualJoystick;
-class SSearchBox;
 
 // Docking Framework
 #include "WorkspaceItem.h"
 #include "TabManager.h"
 #include "LayoutService.h"
 
-enum ETabActivationCause : uint8;
-enum ETabRole : uint8;
-class SDockTab;
diff --git a/Engine/Source/Runtime/Slate/Public/SlateFwd.h b/Engine/Source/Runtime/Slate/Public/SlateFwd.h
--- a/Engine/Source/Runtime/Slate/Public/SlateFwd.h
+++ b/Engine/Source/Runtime/Slate/Public/SlateFwd.h
@@ -42,3 +42,52 @@ class SSearchBox;
 enum ETabActivationCause : uint8;
 enum ETabRole : uint8;
 class SDockTab;
+
+// Legacy
+class SMissingWidget;
+class FSlateStyleRegistry;
+class FUICommandDragDropOp;
+class FInertialScrollManager;
+class FGenericCommands;
+class SSearchBox;
+
+class SVolumeControl;
+class SColorSpectrum;
+class SColorWheel;
+
+class FToolBarButtonBlock;
+class SToolBarButtonBlock;
+class FToolBarComboButtonBlock;
+class SToolBarComboButtonBlock;
+class SPopup;
+
+class SHyperlink;
+class SRichTextHyperlink;
+class SThrobber;
+class SCircularThrobber;
+class STextEntryPopup;
+class STextComboPopup;
+class SExpandableButton;
+class SExpandableArea;
+class SNotificationItem;
+class SNotificationList;
+class SWidgetSwitcher;
+class SSuggestionTextBox;
+template <typename ItemType> class SBreadcrumbTrail;
+class STextComboBox;
+template<typename NumericType> class SNumericEntryBox;
+template<typename OptionType> class SEditableComboBox;
+class FSlateNotificationManager;
+class SDPIScaler;
+class SInlineEditableTextBlock;
+class SVirtualKeyboardEntry;
+class SSafeZone;
+struct FMarqueeRect;
+class SRotatorInputBox;
+class SVectorInputBox;
+class SVirtualJoystick;
+class SSearchBox;
+
+enum ETabActivationCause : uint8;
+enum ETabRole : uint8;
+class SDockTab;
diff --git a/Engine/Source/Runtime/SlateCore/Private/Input/Events.cpp b/Engine/Source/Runtime/SlateCore/Private/Input/Events.cpp
--- a/Engine/Source/Runtime/SlateCore/Private/Input/Events.cpp
+++ b/Engine/Source/Runtime/SlateCore/Private/Input/Events.cpp
@@ -9,6 +9,16 @@
 const FTouchKeySet FTouchKeySet::StandardSet(EKeys::LeftMouseButton);
 const FTouchKeySet FTouchKeySet::EmptySet(EKeys::Invalid);
 
+FGeometry FInputEvent::FindGeometry(const TSharedRef<SWidget>& WidgetToFind) const
+{
+	return EventPath->FindArrangedWidget(WidgetToFind).Get(FArrangedWidget::NullWidget).Geometry;
+}
+
+TSharedRef<SWindow> FInputEvent::GetWindow() const
+{
+	return EventPath->GetWindow();
+}
+
 FText FInputEvent::ToText() const
 {
 	return NSLOCTEXT("Events", "Unimplemented", "Unimplemented");
diff --git a/Engine/Source/Runtime/SlateCore/Public/Input/CursorReply.h b/Engine/Source/Runtime/SlateCore/Public/Input/CursorReply.h
--- a/Engine/Source/Runtime/SlateCore/Public/Input/CursorReply.h
+++ b/Engine/Source/Runtime/SlateCore/Public/Input/CursorReply.h
@@ -2,6 +2,7 @@
 
 #pragma once
 
+class SWindow;
 
 /**
  * A reply to the OnQueryCursor event.
diff --git a/Engine/Source/Runtime/SlateCore/Public/Input/Events.h b/Engine/Source/Runtime/SlateCore/Public/Input/Events.h
--- a/Engine/Source/Runtime/SlateCore/Public/Input/Events.h
+++ b/Engine/Source/Runtime/SlateCore/Public/Input/Events.h
@@ -7,6 +7,7 @@
 
 
 class SWindow;
+class FWidgetPath;
 
 
 /**
@@ -318,15 +319,9 @@ public:
 	}
 
 	/** The event path provides additional context for handling */
-	FGeometry FindGeometry( const TSharedRef<SWidget>& WidgetToFind ) const
-	{
-		return EventPath->FindArrangedWidget(WidgetToFind).Get(FArrangedWidget::NullWidget).Geometry;
-	}
+	SLATECORE_API FGeometry FindGeometry(const TSharedRef<SWidget>& WidgetToFind) const;
 
-	TSharedRef<SWindow> GetWindow() const
-	{
-		return EventPath->GetWindow();
-	}
+	SLATECORE_API TSharedRef<SWindow> GetWindow() const;
 
 	/** Set the widget path along which this event will be routed */
 	void SetEventPath( const FWidgetPath& InEventPath )
diff --git a/Engine/Source/Runtime/SlateCore/Public/Layout/WidgetPath.h b/Engine/Source/Runtime/SlateCore/Public/Layout/WidgetPath.h
--- a/Engine/Source/Runtime/SlateCore/Public/Layout/WidgetPath.h
+++ b/Engine/Source/Runtime/SlateCore/Public/Layout/WidgetPath.h
@@ -247,3 +247,5 @@ public:
 	TArray< TWeakPtr<SWidget> > Widgets;
 	TWeakPtr< SWindow > Window;
 };
+
+#include "WidgetPath.inl"
diff --git a/Engine/Source/Runtime/SlateCore/Public/SlateCore.h b/Engine/Source/Runtime/SlateCore/Public/SlateCore.h
--- a/Engine/Source/Runtime/SlateCore/Public/SlateCore.h
+++ b/Engine/Source/Runtime/SlateCore/Public/SlateCore.h
@@ -10,25 +10,10 @@
 #include "InputCore.h"
 #include "ModuleManager.h"
 
-
 /* Globals
  *****************************************************************************/
 
-SLATECORE_API DECLARE_LOG_CATEGORY_EXTERN(LogSlate, Log, All);
-SLATECORE_API DECLARE_LOG_CATEGORY_EXTERN(LogSlateStyles, Log, All);
-
-DECLARE_STATS_GROUP(TEXT("Slate Memory"), STATGROUP_SlateMemory, STATCAT_Advanced);
-DECLARE_STATS_GROUP(TEXT("Slate"), STATGROUP_Slate, STATCAT_Advanced);
-DECLARE_STATS_GROUP_VERBOSE(TEXT("SlateVerbose"), STATGROUP_SlateVerbose, STATCAT_Advanced);
-
-
-// Compile all the RichText and MultiLine editable text?
-#define WITH_FANCY_TEXT 1
-
-/* Forward declarations
- *****************************************************************************/
-class FActiveTimerHandle;
-enum class EActiveTimerReturnType : uint8;
+#include "SlateGlobals.h"
 
 /* Delegates
  *****************************************************************************/
@@ -56,7 +41,6 @@ enum class EActiveTimerReturnType : uint8;
 #include "ArrangedChildren.h"
 #include "SNullWidget.h"
 #include "Children.h"
-#include "WidgetPath.h"
 #include "LayoutUtils.h"
 
 // Animation
@@ -67,6 +51,7 @@ enum class EActiveTimerReturnType : uint8;
 // Sound
 #include "SlateSound.h"
 #include "ISlateSoundDevice.h"
+#include "NullSlateSoundDevice.h"
 
 // Styling
 #include "WidgetStyle.h"
@@ -131,7 +116,7 @@ enum class EActiveTimerReturnType : uint8;
 #include "SUserWidget.h"
 #include "SLeafWidget.h"
 #include "SPanel.h"
-#include "WidgetPath.inl"
+#include "WidgetPath.h"
 ////
 #include "SBoxPanel.h"
 #include "SOverlay.h"
diff --git a/Engine/Source/Runtime/SlateCore/Public/Sound/ISlateSoundDevice.h b/Engine/Source/Runtime/SlateCore/Public/Sound/ISlateSoundDevice.h
--- a/Engine/Source/Runtime/SlateCore/Public/Sound/ISlateSoundDevice.h
+++ b/Engine/Source/Runtime/SlateCore/Public/Sound/ISlateSoundDevice.h
@@ -16,19 +16,3 @@ public:
 	/** Virtual destructor: because it is an interface / pure virtual base class. */
 	virtual ~ISlateSoundDevice(){}
 };
-
-/** Silent implementation of ISlateSoundDevice; it plays nothing. */
-class SLATECORE_API FNullSlateSoundDevice : public ISlateSoundDevice
-{
-public:
-	virtual void PlaySound(const FSlateSound&, int32) const override
-	{
-	}
-
-	virtual float GetSoundDuration(const FSlateSound& Sound) const override
-	{
-		return 0.0f;
-	}
-
-	virtual ~FNullSlateSoundDevice(){}
-};
diff --git a/Engine/Source/Runtime/SlateCore/Public/Sound/NullSlateSoundDevice.h b/Engine/Source/Runtime/SlateCore/Public/Sound/NullSlateSoundDevice.h
new file mode 100644
--- /dev/null
+++ b/Engine/Source/Runtime/SlateCore/Public/Sound/NullSlateSoundDevice.h
@@ -0,0 +1,14 @@
+// Copyright 1998-2016 Epic Games, Inc. All Rights Reserved.
+
+#pragma once
+
+#include "ISlateSoundDevice.h"
+
+/** Silent implementation of ISlateSoundDevice; it plays nothing. */
+class SLATECORE_API FNullSlateSoundDevice : public ISlateSoundDevice
+{
+public:
+	virtual void PlaySound(const FSlateSound&, int32) const override;
+	virtual float GetSoundDuration(const FSlateSound& Sound) const override;
+	virtual ~FNullSlateSoundDevice();
+};
diff --git a/Engine/Source/Runtime/Sockets/Private/SocketsPrivatePCH.h b/Engine/Source/Runtime/Sockets/Private/SocketsPrivatePCH.h
--- a/Engine/Source/Runtime/Sockets/Private/SocketsPrivatePCH.h
+++ b/Engine/Source/Runtime/Sockets/Private/SocketsPrivatePCH.h
@@ -1,11 +1,9 @@
 // Copyright 1998-2016 Epic Games, Inc. All Rights Reserved.
 
-#ifndef UNREAL_SOCKETS_PRIVATE_PCH
-#define UNREAL_SOCKETS_PRIVATE_PCH
+#pragma once
 
 #include "Core.h"
 #include "SocketSubsystem.h"
 #include "SocketSubsystemModule.h"
 #include "ModuleManager.h"
 
-#endif
diff --git a/Engine/Source/Runtime/Sockets/Public/MultichannelTCP.h b/Engine/Source/Runtime/Sockets/Public/MultichannelTCP.h
--- a/Engine/Source/Runtime/Sockets/Public/MultichannelTCP.h
+++ b/Engine/Source/Runtime/Sockets/Public/MultichannelTCP.h
@@ -5,17 +5,7 @@
 #include "SocketSubsystem.h"
 #include "Sockets.h"
 
-
-SOCKETS_API DECLARE_LOG_CATEGORY_EXTERN(LogMultichannelTCP, Log, All);
-
-
-/** Magic number used to verify packet header **/
-enum
-{
-	MultichannelMagic= 0xa692339f
-};
-
-
+#include "MultichannelTcpGlobals.h"
 #include "NetworkMessage.h"
 #include "MultichannelTcpReceiver.h"
 #include "MultichannelTcpSender.h"
diff --git a/Engine/Source/Runtime/SynthBenchmark/Private/SynthBenchmarkPrivatePCH.h b/Engine/Source/Runtime/SynthBenchmark/Private/SynthBenchmarkPrivatePCH.h
--- a/Engine/Source/Runtime/SynthBenchmark/Private/SynthBenchmarkPrivatePCH.h
+++ b/Engine/Source/Runtime/SynthBenchmark/Private/SynthBenchmarkPrivatePCH.h
@@ -1,6 +1,5 @@
 // Copyright 1998-2016 Epic Games, Inc. All Rights Reserved.
 
-#include "SynthBenchmark.h"
+#pragma once
 
-// You should place include statements to your module's private header files here.  You only need to
-// add includes for headers that are used in most of your module's source files though.
+#include "SynthBenchmark.h"
diff --git a/Engine/Source/Runtime/UnrealAudio/Private/UnrealAudioDecode.h b/Engine/Source/Runtime/UnrealAudio/Private/UnrealAudioDecode.h
--- a/Engine/Source/Runtime/UnrealAudio/Private/UnrealAudioDecode.h
+++ b/Engine/Source/Runtime/UnrealAudio/Private/UnrealAudioDecode.h
@@ -2,7 +2,7 @@
 
 #pragma once 
 
-#include "UnrealAudioModule.h"
+#include "UnrealAudioTypes.h"
 #include "UnrealAudioHandles.h"
 #include "UnrealAudioSoundFileManager.h"
 #include "UnrealAudioTypes.h"
diff --git a/Engine/Source/Runtime/UnrealAudio/Private/UnrealAudioDeviceFormat.h b/Engine/Source/Runtime/UnrealAudio/Private/UnrealAudioDeviceFormat.h
--- a/Engine/Source/Runtime/UnrealAudio/Private/UnrealAudioDeviceFormat.h
+++ b/Engine/Source/Runtime/UnrealAudio/Private/UnrealAudioDeviceFormat.h
@@ -2,9 +2,7 @@
 
 #pragma once 
 
-#include "UnrealAudioModule.h"
-#include "UnrealAudioPrivate.h"
-#include "UnrealAudioDeviceModule.h"
+#include "UnrealAudioTypes.h"
 
 #if ENABLE_UNREAL_AUDIO
 
diff --git a/Engine/Source/Runtime/UnrealAudio/Private/UnrealAudioEmitterInternal.h b/Engine/Source/Runtime/UnrealAudio/Private/UnrealAudioEmitterInternal.h
--- a/Engine/Source/Runtime/UnrealAudio/Private/UnrealAudioEmitterInternal.h
+++ b/Engine/Source/Runtime/UnrealAudio/Private/UnrealAudioEmitterInternal.h
@@ -2,11 +2,8 @@
 
 #pragma once 
 
-#include "UnrealAudioModule.h"
-#include "UnrealAudioPrivate.h"
 #include "UnrealAudioTypes.h"
 #include "UnrealAudioEmitter.h"
-#include "UnrealAudioPrivate.h"
 #include "UnrealAudioEntityManager.h"
 #include "UnrealAudioUtilities.h"
 
diff --git a/Engine/Source/Runtime/UnrealAudio/Private/UnrealAudioEmitterManager.h b/Engine/Source/Runtime/UnrealAudio/Private/UnrealAudioEmitterManager.h
--- a/Engine/Source/Runtime/UnrealAudio/Private/UnrealAudioEmitterManager.h
+++ b/Engine/Source/Runtime/UnrealAudio/Private/UnrealAudioEmitterManager.h
@@ -2,8 +2,6 @@
 
 #pragma once 
 
-#include "UnrealAudioPrivate.h"
-#include "UnrealAudioModule.h"
 #include "UnrealAudioTypes.h"
 #include "UnrealAudioEmitter.h"
 #include "UnrealAudioEntityManager.h"
diff --git a/Engine/Source/Runtime/UnrealAudio/Private/UnrealAudioPitchManager.h b/Engine/Source/Runtime/UnrealAudio/Private/UnrealAudioPitchManager.h
--- a/Engine/Source/Runtime/UnrealAudio/Private/UnrealAudioPitchManager.h
+++ b/Engine/Source/Runtime/UnrealAudio/Private/UnrealAudioPitchManager.h
@@ -2,7 +2,6 @@
 
 #pragma once 
 
-#include "UnrealAudioModule.h"
 #include "UnrealAudioTypes.h"
 #include "UnrealAudioVoice.h"
 #include "UnrealAudioPrivate.h"
diff --git a/Engine/Source/Runtime/UnrealAudio/Private/UnrealAudioSampleRateConverter.h b/Engine/Source/Runtime/UnrealAudio/Private/UnrealAudioSampleRateConverter.h
--- a/Engine/Source/Runtime/UnrealAudio/Private/UnrealAudioSampleRateConverter.h
+++ b/Engine/Source/Runtime/UnrealAudio/Private/UnrealAudioSampleRateConverter.h
@@ -2,7 +2,7 @@
 
 #pragma once
 
-#include "UnrealAudioPrivate.h"
+#include "UnrealAudioTypes.h"
 
 #if ENABLE_UNREAL_AUDIO
 
diff --git a/Engine/Source/Runtime/UnrealAudio/Private/UnrealAudioSoundFileInternal.h b/Engine/Source/Runtime/UnrealAudio/Private/UnrealAudioSoundFileInternal.h
--- a/Engine/Source/Runtime/UnrealAudio/Private/UnrealAudioSoundFileInternal.h
+++ b/Engine/Source/Runtime/UnrealAudio/Private/UnrealAudioSoundFileInternal.h
@@ -2,9 +2,7 @@
 
 #pragma once 
 
-#include "UnrealAudioModule.h"
-#include "UnrealAudioPrivate.h"
-#include "UnrealAudioSoundFileManager.h"
+#include "UnrealAudioTypes.h"
 #include "UnrealAudioTypes.h"
 #include "UnrealAudioSoundFile.h"
 #include "UnrealAudioEntityManager.h"
diff --git a/Engine/Source/Runtime/UnrealAudio/Private/UnrealAudioSoundFileManager.h b/Engine/Source/Runtime/UnrealAudio/Private/UnrealAudioSoundFileManager.h
--- a/Engine/Source/Runtime/UnrealAudio/Private/UnrealAudioSoundFileManager.h
+++ b/Engine/Source/Runtime/UnrealAudio/Private/UnrealAudioSoundFileManager.h
@@ -2,8 +2,6 @@
 
 #pragma once 
 
-#include "UnrealAudioModule.h"
-#include "UnrealAudioPrivate.h"
 #include "UnrealAudioTypes.h"
 #include "UnrealAudioSoundFile.h"
 #include "UnrealAudioSoundFileInternal.h"
diff --git a/Engine/Source/Runtime/UnrealAudio/Private/UnrealAudioUtilities.h b/Engine/Source/Runtime/UnrealAudio/Private/UnrealAudioUtilities.h
--- a/Engine/Source/Runtime/UnrealAudio/Private/UnrealAudioUtilities.h
+++ b/Engine/Source/Runtime/UnrealAudio/Private/UnrealAudioUtilities.h
@@ -2,8 +2,7 @@
 
 #pragma once 
 
-#include "UnrealAudioPrivate.h"
-#include "UnrealAudioModule.h"
+#include "UnrealAudioTypes.h"
 #include "UnrealAudioEntityManager.h"
 
 #if ENABLE_UNREAL_AUDIO
diff --git a/Engine/Source/Runtime/UnrealAudio/Private/UnrealAudioVoiceInternal.h b/Engine/Source/Runtime/UnrealAudio/Private/UnrealAudioVoiceInternal.h
--- a/Engine/Source/Runtime/UnrealAudio/Private/UnrealAudioVoiceInternal.h
+++ b/Engine/Source/Runtime/UnrealAudio/Private/UnrealAudioVoiceInternal.h
@@ -2,8 +2,6 @@
 
 #pragma once 
 
-#include "UnrealAudioModule.h"
-#include "UnrealAudioPrivate.h"
 #include "UnrealAudioTypes.h"
 #include "UnrealAudioVoice.h"
 #include "UnrealAudioPrivate.h"
diff --git a/Engine/Source/Runtime/UnrealAudio/Private/UnrealAudioVoiceManager.h b/Engine/Source/Runtime/UnrealAudio/Private/UnrealAudioVoiceManager.h
--- a/Engine/Source/Runtime/UnrealAudio/Private/UnrealAudioVoiceManager.h
+++ b/Engine/Source/Runtime/UnrealAudio/Private/UnrealAudioVoiceManager.h
@@ -2,8 +2,6 @@
 
 #pragma once 
 
-#include "UnrealAudioPrivate.h"
-#include "UnrealAudioModule.h"
 #include "UnrealAudioTypes.h"
 #include "UnrealAudioVoice.h"
 #include "UnrealAudioEntityManager.h"
diff --git a/Engine/Source/Runtime/UnrealAudio/Private/UnrealAudioVoiceMixer.h b/Engine/Source/Runtime/UnrealAudio/Private/UnrealAudioVoiceMixer.h
--- a/Engine/Source/Runtime/UnrealAudio/Private/UnrealAudioVoiceMixer.h
+++ b/Engine/Source/Runtime/UnrealAudio/Private/UnrealAudioVoiceMixer.h
@@ -2,13 +2,10 @@
 
 #pragma once 
 
-#include "UnrealAudioPrivate.h"
-#include "UnrealAudioModule.h"
 #include "UnrealAudioTypes.h"
 #include "UnrealAudioVoice.h"
 #include "UnrealAudioVoiceInternal.h"
 #include "UnrealAudioUtilities.h"
-#include "UnrealAudioVoiceManager.h"
 #include "UnrealAudioDecode.h"
 #include "UnrealAudioSampleRateConverter.h"
 
diff --git a/Engine/Source/Runtime/UnrealAudio/Private/UnrealAudioVolumeManager.h b/Engine/Source/Runtime/UnrealAudio/Private/UnrealAudioVolumeManager.h
--- a/Engine/Source/Runtime/UnrealAudio/Private/UnrealAudioVolumeManager.h
+++ b/Engine/Source/Runtime/UnrealAudio/Private/UnrealAudioVolumeManager.h
@@ -2,7 +2,6 @@
 
 #pragma once 
 
-#include "UnrealAudioModule.h"
 #include "UnrealAudioTypes.h"
 #include "UnrealAudioVoice.h"
 #include "UnrealAudioPrivate.h"
diff --git a/Engine/Source/Runtime/UnrealAudio/Public/UnrealAudioTypes.h b/Engine/Source/Runtime/UnrealAudio/Public/UnrealAudioTypes.h
--- a/Engine/Source/Runtime/UnrealAudio/Public/UnrealAudioTypes.h
+++ b/Engine/Source/Runtime/UnrealAudio/Public/UnrealAudioTypes.h
@@ -9,8 +9,6 @@
 #define ENABLE_UNREAL_AUDIO 0
 #endif
 
-#include "UnrealAudioModule.h"
-
 // Turns on extra checks and asserts for debugging threading code
 #define ENABLE_UNREAL_AUDIO_THREAD_DEBUGGING		1
 #define ENABLE_UNREAL_AUDIO_EXTRA_DEBUG_CHECKS		1
diff --git a/Engine/Source/Runtime/VulkanRHI/Private/VulkanCommandBuffer.h b/Engine/Source/Runtime/VulkanRHI/Private/VulkanCommandBuffer.h
--- a/Engine/Source/Runtime/VulkanRHI/Private/VulkanCommandBuffer.h
+++ b/Engine/Source/Runtime/VulkanRHI/Private/VulkanCommandBuffer.h
@@ -6,7 +6,7 @@
 
 #pragma once
 
-#include "VulkanRHIPrivate.h"
+#include "VulkanConfiguration.h"
 
 class FVulkanDevice;
 class FVulkanCommandBufferManager;
diff --git a/Engine/Source/Runtime/VulkanRHI/Private/VulkanPendingState.h b/Engine/Source/Runtime/VulkanRHI/Private/VulkanPendingState.h
--- a/Engine/Source/Runtime/VulkanRHI/Private/VulkanPendingState.h
+++ b/Engine/Source/Runtime/VulkanRHI/Private/VulkanPendingState.h
@@ -10,6 +10,8 @@
 #include "VulkanRHI.h"
 #include "VulkanPipeline.h"
 
+typedef uint32 FStateKey;
+
 class FVulkanPendingState
 {
 public:
diff --git a/Engine/Source/Runtime/VulkanRHI/Private/VulkanQueue.h b/Engine/Source/Runtime/VulkanRHI/Private/VulkanQueue.h
--- a/Engine/Source/Runtime/VulkanRHI/Private/VulkanQueue.h
+++ b/Engine/Source/Runtime/VulkanRHI/Private/VulkanQueue.h
@@ -6,7 +6,7 @@
 
 #pragma once
 
-#include "VulkanRHIPrivate.h"
+#include "VulkanConfiguration.h"
 
 class FVulkanDevice;
 class FVulkanCmdBuffer;
diff --git a/Engine/Source/Runtime/VulkanRHI/Public/VulkanDynamicRHI.h b/Engine/Source/Runtime/VulkanRHI/Public/VulkanDynamicRHI.h
new file mode 100644
--- /dev/null
+++ b/Engine/Source/Runtime/VulkanRHI/Public/VulkanDynamicRHI.h
@@ -0,0 +1,227 @@
+// Copyright 1998-2016 Epic Games, Inc. All Rights Reserved.
+
+/*=============================================================================
+	VulkanRHI.h: Public Vulkan RHI definitions.
+=============================================================================*/
+
+#pragma once 
+
+#include "VulkanGlobals.h"
+
+class FVulkanTexture2D;
+class FVulkanFramebuffer;
+class FVulkanDevice;
+class FVulkanQueue;
+
+
+/** The interface which is implemented by the dynamically bound RHI. */
+class FVulkanDynamicRHI : public FDynamicRHI
+{
+public:
+
+	/** Initialization constructor. */
+	FVulkanDynamicRHI();
+
+	/** Destructor */
+	~FVulkanDynamicRHI() {}
+
+
+	// FDynamicRHI interface.
+	virtual void Init();
+	virtual void Shutdown();
+
+	void InitInstance();
+	void IssueLongGPUTask();
+
+	virtual FSamplerStateRHIRef RHICreateSamplerState(const FSamplerStateInitializerRHI& Initializer) final override;
+	virtual FRasterizerStateRHIRef RHICreateRasterizerState(const FRasterizerStateInitializerRHI& Initializer) final override;
+	virtual FDepthStencilStateRHIRef RHICreateDepthStencilState(const FDepthStencilStateInitializerRHI& Initializer) final override;
+	virtual FBlendStateRHIRef RHICreateBlendState(const FBlendStateInitializerRHI& Initializer) final override;
+	virtual FVertexDeclarationRHIRef RHICreateVertexDeclaration(const FVertexDeclarationElementList& Elements) final override;
+	virtual FPixelShaderRHIRef RHICreatePixelShader(const TArray<uint8>& Code) final override;
+	virtual FVertexShaderRHIRef RHICreateVertexShader(const TArray<uint8>& Code) final override;
+	virtual FHullShaderRHIRef RHICreateHullShader(const TArray<uint8>& Code) final override;
+	virtual FDomainShaderRHIRef RHICreateDomainShader(const TArray<uint8>& Code) final override;
+	virtual FGeometryShaderRHIRef RHICreateGeometryShader(const TArray<uint8>& Code) final override;
+	virtual FGeometryShaderRHIRef RHICreateGeometryShaderWithStreamOutput(const TArray<uint8>& Code, const FStreamOutElementList& ElementList, uint32 NumStrides, const uint32* Strides, int32 RasterizedStream) final override;
+	virtual FComputeShaderRHIRef RHICreateComputeShader(const TArray<uint8>& Code) final override;
+	virtual FBoundShaderStateRHIRef RHICreateBoundShaderState(FVertexDeclarationRHIParamRef VertexDeclaration, FVertexShaderRHIParamRef VertexShader, FHullShaderRHIParamRef HullShader, FDomainShaderRHIParamRef DomainShader, FPixelShaderRHIParamRef PixelShader, FGeometryShaderRHIParamRef GeometryShader) final override;
+	virtual FUniformBufferRHIRef RHICreateUniformBuffer(const void* Contents, const FRHIUniformBufferLayout& Layout, EUniformBufferUsage Usage) final override;
+	virtual FIndexBufferRHIRef RHICreateIndexBuffer(uint32 Stride, uint32 Size, uint32 InUsage, FRHIResourceCreateInfo& CreateInfo) final override;
+	virtual void* RHILockIndexBuffer(FIndexBufferRHIParamRef IndexBuffer, uint32 Offset, uint32 Size, EResourceLockMode LockMode) final override;
+	virtual void RHIUnlockIndexBuffer(FIndexBufferRHIParamRef IndexBuffer) final override;
+	virtual FVertexBufferRHIRef RHICreateVertexBuffer(uint32 Size, uint32 InUsage, FRHIResourceCreateInfo& CreateInfo) final override;
+	virtual void* RHILockVertexBuffer(FVertexBufferRHIParamRef VertexBuffer, uint32 Offset, uint32 SizeRHI, EResourceLockMode LockMode) final override;
+	virtual void RHIUnlockVertexBuffer(FVertexBufferRHIParamRef VertexBuffer) final override;
+	virtual void RHICopyVertexBuffer(FVertexBufferRHIParamRef SourceBuffer, FVertexBufferRHIParamRef DestBuffer) final override;
+	virtual FStructuredBufferRHIRef RHICreateStructuredBuffer(uint32 Stride, uint32 Size, uint32 InUsage, FRHIResourceCreateInfo& CreateInfo) final override;
+	virtual void* RHILockStructuredBuffer(FStructuredBufferRHIParamRef StructuredBuffer, uint32 Offset, uint32 SizeRHI, EResourceLockMode LockMode) final override;
+	virtual void RHIUnlockStructuredBuffer(FStructuredBufferRHIParamRef StructuredBuffer) final override;
+	virtual FUnorderedAccessViewRHIRef RHICreateUnorderedAccessView(FStructuredBufferRHIParamRef StructuredBuffer, bool bUseUAVCounter, bool bAppendBuffer) final override;
+	virtual FUnorderedAccessViewRHIRef RHICreateUnorderedAccessView(FTextureRHIParamRef Texture, uint32 MipLevel) final override;
+	virtual FUnorderedAccessViewRHIRef RHICreateUnorderedAccessView(FVertexBufferRHIParamRef VertexBuffer, uint8 Format) final override;
+	virtual FShaderResourceViewRHIRef RHICreateShaderResourceView(FStructuredBufferRHIParamRef StructuredBuffer) final override;
+	virtual FShaderResourceViewRHIRef RHICreateShaderResourceView(FVertexBufferRHIParamRef VertexBuffer, uint32 Stride, uint8 Format) final override;
+	virtual FShaderResourceViewRHIRef RHICreateShaderResourceView(FIndexBufferRHIParamRef Buffer) final override;
+	virtual uint64 RHICalcTexture2DPlatformSize(uint32 SizeX, uint32 SizeY, uint8 Format, uint32 NumMips, uint32 NumSamples, uint32 Flags, uint32& OutAlign) final override;
+	virtual uint64 RHICalcTexture3DPlatformSize(uint32 SizeX, uint32 SizeY, uint32 SizeZ, uint8 Format, uint32 NumMips, uint32 Flags, uint32& OutAlign) final override;
+	virtual uint64 RHICalcTextureCubePlatformSize(uint32 Size, uint8 Format, uint32 NumMips, uint32 Flags, uint32& OutAlign) final override;
+	virtual void RHIGetTextureMemoryStats(FTextureMemoryStats& OutStats) final override;
+	virtual bool RHIGetTextureMemoryVisualizeData(FColor* TextureData, int32 SizeX, int32 SizeY, int32 Pitch, int32 PixelSize) final override;
+	virtual FTextureReferenceRHIRef RHICreateTextureReference(FLastRenderTimeContainer* LastRenderTime) final override;
+	virtual FTexture2DRHIRef RHICreateTexture2D(uint32 SizeX, uint32 SizeY, uint8 Format, uint32 NumMips, uint32 NumSamples, uint32 Flags, FRHIResourceCreateInfo& CreateInfo) final override;
+	virtual FTexture2DRHIRef RHIAsyncCreateTexture2D(uint32 SizeX, uint32 SizeY, uint8 Format, uint32 NumMips, uint32 Flags, void** InitialMipData, uint32 NumInitialMips) final override;
+	virtual void RHICopySharedMips(FTexture2DRHIParamRef DestTexture2D, FTexture2DRHIParamRef SrcTexture2D) final override;
+	virtual FTexture2DArrayRHIRef RHICreateTexture2DArray(uint32 SizeX, uint32 SizeY, uint32 SizeZ, uint8 Format, uint32 NumMips, uint32 Flags, FRHIResourceCreateInfo& CreateInfo) final override;
+	virtual FTexture3DRHIRef RHICreateTexture3D(uint32 SizeX, uint32 SizeY, uint32 SizeZ, uint8 Format, uint32 NumMips, uint32 Flags, FRHIResourceCreateInfo& CreateInfo) final override;
+	virtual void RHIGetResourceInfo(FTextureRHIParamRef Ref, FRHIResourceInfo& OutInfo) final override;
+	virtual FShaderResourceViewRHIRef RHICreateShaderResourceView(FTexture2DRHIParamRef Texture2DRHI, uint8 MipLevel) final override;
+	virtual FShaderResourceViewRHIRef RHICreateShaderResourceView(FTexture2DRHIParamRef Texture2DRHI, uint8 MipLevel, uint8 NumMipLevels, uint8 Format) final override;
+	virtual FShaderResourceViewRHIRef RHICreateShaderResourceView(FTexture3DRHIParamRef Texture3DRHI, uint8 MipLevel) final override;
+	virtual FShaderResourceViewRHIRef RHICreateShaderResourceView(FTexture2DArrayRHIParamRef Texture2DArrayRHI, uint8 MipLevel) final override;
+	virtual FShaderResourceViewRHIRef RHICreateShaderResourceView(FTextureCubeRHIParamRef TextureCubeRHI, uint8 MipLevel) final override;
+	virtual void RHIGenerateMips(FTextureRHIParamRef Texture) final override;
+	virtual uint32 RHIComputeMemorySize(FTextureRHIParamRef TextureRHI) final override;
+	virtual FTexture2DRHIRef RHIAsyncReallocateTexture2D(FTexture2DRHIParamRef Texture2D, int32 NewMipCount, int32 NewSizeX, int32 NewSizeY, FThreadSafeCounter* RequestStatus) final override;
+	virtual ETextureReallocationStatus RHIFinalizeAsyncReallocateTexture2D(FTexture2DRHIParamRef Texture2D, bool bBlockUntilCompleted) final override;
+	virtual ETextureReallocationStatus RHICancelAsyncReallocateTexture2D(FTexture2DRHIParamRef Texture2D, bool bBlockUntilCompleted) final override;
+	virtual void* RHILockTexture2D(FTexture2DRHIParamRef Texture, uint32 MipIndex, EResourceLockMode LockMode, uint32& DestStride, bool bLockWithinMiptail) final override;
+	virtual void RHIUnlockTexture2D(FTexture2DRHIParamRef Texture, uint32 MipIndex, bool bLockWithinMiptail) final override;
+	virtual void* RHILockTexture2DArray(FTexture2DArrayRHIParamRef Texture, uint32 TextureIndex, uint32 MipIndex, EResourceLockMode LockMode, uint32& DestStride, bool bLockWithinMiptail) final override;
+	virtual void RHIUnlockTexture2DArray(FTexture2DArrayRHIParamRef Texture, uint32 TextureIndex, uint32 MipIndex, bool bLockWithinMiptail) final override;
+	virtual void RHIUpdateTexture2D(FTexture2DRHIParamRef Texture, uint32 MipIndex, const struct FUpdateTextureRegion2D& UpdateRegion, uint32 SourcePitch, const uint8* SourceData) final override;
+	virtual void RHIUpdateTexture3D(FTexture3DRHIParamRef Texture, uint32 MipIndex, const struct FUpdateTextureRegion3D& UpdateRegion, uint32 SourceRowPitch, uint32 SourceDepthPitch, const uint8* SourceData) final override;
+	virtual FTextureCubeRHIRef RHICreateTextureCube(uint32 Size, uint8 Format, uint32 NumMips, uint32 Flags, FRHIResourceCreateInfo& CreateInfo) final override;
+	virtual FTextureCubeRHIRef RHICreateTextureCubeArray(uint32 Size, uint32 ArraySize, uint8 Format, uint32 NumMips, uint32 Flags, FRHIResourceCreateInfo& CreateInfo) final override;
+	virtual void* RHILockTextureCubeFace(FTextureCubeRHIParamRef Texture, uint32 FaceIndex, uint32 ArrayIndex, uint32 MipIndex, EResourceLockMode LockMode, uint32& DestStride, bool bLockWithinMiptail) final override;
+	virtual void RHIUnlockTextureCubeFace(FTextureCubeRHIParamRef Texture, uint32 FaceIndex, uint32 ArrayIndex, uint32 MipIndex, bool bLockWithinMiptail) final override;
+	virtual void RHIBindDebugLabelName(FTextureRHIParamRef Texture, const TCHAR* Name) final override;
+	virtual void RHIBindDebugLabelName(FUnorderedAccessViewRHIParamRef UnorderedAccessViewRHI, const TCHAR* Name) final override;
+	virtual void RHIReadSurfaceData(FTextureRHIParamRef Texture, FIntRect Rect, TArray<FColor>& OutData, FReadSurfaceDataFlags InFlags) final override;
+	virtual void RHIMapStagingSurface(FTextureRHIParamRef Texture, void*& OutData, int32& OutWidth, int32& OutHeight) final override;
+	virtual void RHIUnmapStagingSurface(FTextureRHIParamRef Texture) final override;
+	virtual void RHIReadSurfaceFloatData(FTextureRHIParamRef Texture, FIntRect Rect, TArray<FFloat16Color>& OutData, ECubeFace CubeFace, int32 ArrayIndex, int32 MipIndex) final override;
+	virtual void RHIRead3DSurfaceFloatData(FTextureRHIParamRef Texture, FIntRect Rect, FIntPoint ZMinMax, TArray<FFloat16Color>& OutData) final override;
+	virtual FRenderQueryRHIRef RHICreateRenderQuery(ERenderQueryType QueryType) final override;
+	virtual bool RHIGetRenderQueryResult(FRenderQueryRHIParamRef RenderQuery, uint64& OutResult, bool bWait) final override;
+	virtual FTexture2DRHIRef RHIGetViewportBackBuffer(FViewportRHIParamRef Viewport) final override;
+	virtual void RHIAdvanceFrameForGetViewportBackBuffer() final override;
+	virtual void RHIAcquireThreadOwnership() final override;
+	virtual void RHIReleaseThreadOwnership() final override;
+	virtual void RHIFlushResources() final override;
+	virtual uint32 RHIGetGPUFrameCycles() final override;
+	virtual FViewportRHIRef RHICreateViewport(void* WindowHandle, uint32 SizeX, uint32 SizeY, bool bIsFullscreen, EPixelFormat PreferredPixelFormat) final override;
+	virtual void RHIResizeViewport(FViewportRHIParamRef Viewport, uint32 SizeX, uint32 SizeY, bool bIsFullscreen) final override;
+	virtual void RHITick(float DeltaTime) final override;
+	virtual void RHISetStreamOutTargets(uint32 NumTargets, const FVertexBufferRHIParamRef* VertexBuffers, const uint32* Offsets) final override;
+	virtual void RHIDiscardRenderTargets(bool Depth, bool Stencil, uint32 ColorBitMask) final override {}
+	virtual void RHIBlockUntilGPUIdle() final override;
+	virtual void RHISuspendRendering() final override;
+	virtual void RHIResumeRendering() final override;
+	virtual bool RHIIsRenderingSuspended() final override;
+	virtual bool RHIGetAvailableResolutions(FScreenResolutionArray& Resolutions, bool bIgnoreRefreshRate) final override;
+	virtual void RHIGetSupportedResolution(uint32& Width, uint32& Height) final override;
+	virtual void RHIVirtualTextureSetFirstMipInMemory(FTexture2DRHIParamRef Texture, uint32 FirstMip) final override;
+	virtual void RHIVirtualTextureSetFirstMipVisible(FTexture2DRHIParamRef Texture, uint32 FirstMip) final override;
+	virtual void RHIExecuteCommandList(FRHICommandList* CmdList) final override;
+	virtual void* RHIGetNativeDevice() final override;
+	virtual class IRHICommandContext* RHIGetDefaultContext() final override;
+	virtual class IRHICommandContextContainer* RHIGetCommandContextContainer() final override;
+
+	inline uint32 GetPresentCount() const
+	{
+		return PresentCount;
+	}
+
+	const TArray<const ANSICHAR*>& GetInstanceExtensions() const
+	{
+		return InstanceExtensions;
+	}
+
+	const TArray<const ANSICHAR*>& GetInstanceLayers() const
+	{
+		return InstanceLayers;
+	}
+
+private:
+	void PooledUniformBuffersBeginFrame();
+	void ReleasePooledUniformBuffers();
+
+protected:
+	VkInstance Instance;
+	TArray<const ANSICHAR*> InstanceExtensions;
+	TArray<const ANSICHAR*> InstanceLayers;
+
+	TArray<FVulkanDevice*> Devices;
+
+	FVulkanDevice* Device;
+
+	/** A list of all viewport RHIs that have been created. */
+	TArray<FVulkanViewport*> Viewports;
+
+	/** The viewport which is currently being drawn. */
+	TRefCountPtr<FVulkanViewport> DrawingViewport;
+
+	void CreateInstance();
+	void InitGPU(FVulkanDevice* Device);
+	void InitDevice(FVulkanDevice* Device);
+
+	friend class FVulkanCommandListContext;
+
+	friend class FVulkanViewport;
+
+	/*
+	static void WriteEndFrameTimestamp(void*);
+	*/
+
+	static void GetInstanceLayersAndExtensions(TArray<const ANSICHAR*>& OutInstanceExtensions, TArray<const ANSICHAR*>& OutInstanceLayers);
+
+#if VULKAN_ENABLE_PIPELINE_CACHE
+	IConsoleObject* SavePipelineCacheCmd;
+	IConsoleObject* RebuildPipelineCacheCmd;
+	static void SavePipelineCache();
+	static void RebuildPipelineCache();
+#endif
+
+#if UE_BUILD_DEBUG || UE_BUILD_DEVELOPMENT
+	IConsoleObject* DumpMemoryCmd;
+public:
+	static void DumpMemory();
+#endif
+
+#if VULKAN_HAS_DEBUGGING_ENABLED
+protected:
+	VkDebugReportCallbackEXT MsgCallback;
+	void SetupDebugLayerCallback();
+	void RemoveDebugLayerCallback();
+#endif
+
+	inline bool ShouldDeferBufferLockOperation(FRHICommandList* RHICmdList)
+	{
+		if (RHICmdList == nullptr)
+		{
+			return false;
+		}
+
+		if (RHICmdList->Bypass() || !GRHIThread)
+		{
+			return false;
+		}
+
+		return true;
+	}
+
+	FCriticalSection LockBufferCS;
+
+	uint32 PresentCount;
+};
+
+/** Implements the Vulkan module as a dynamic RHI providing module. */
+class FVulkanDynamicRHIModule : public IDynamicRHIModule
+{
+public:
+	// IDynamicRHIModule
+	virtual bool IsSupported() override;
+
+	virtual FDynamicRHI* CreateRHI(ERHIFeatureLevel::Type RequestedFeatureLevel = ERHIFeatureLevel::Num) override;
+};
diff --git a/Engine/Source/Runtime/VulkanRHI/Public/VulkanGlobals.h b/Engine/Source/Runtime/VulkanRHI/Public/VulkanGlobals.h
new file mode 100644
--- /dev/null
+++ b/Engine/Source/Runtime/VulkanRHI/Public/VulkanGlobals.h
@@ -0,0 +1,9 @@
+// Copyright 1998-2016 Epic Games, Inc. All Rights Reserved.
+
+/*=============================================================================
+	VulkanGlobals.h: Global Vulkan RHI definitions.
+=============================================================================*/
+
+#pragma once 
+
+DECLARE_LOG_CATEGORY_EXTERN(LogVulkan, Display, All);
diff --git a/Engine/Source/Runtime/VulkanRHI/Public/VulkanRHI.h b/Engine/Source/Runtime/VulkanRHI/Public/VulkanRHI.h
--- a/Engine/Source/Runtime/VulkanRHI/Public/VulkanRHI.h
+++ b/Engine/Source/Runtime/VulkanRHI/Public/VulkanRHI.h
@@ -6,231 +6,12 @@
 
 #pragma once 
 
-DECLARE_LOG_CATEGORY_EXTERN(LogVulkan, Display, All);
-
-typedef uint32 FStateKey;
-
 // Vulkan RHI public headers.
+#include "VulkanGlobals.h"
 #include "VulkanConfiguration.h"
 #include "VulkanState.h"
 #include "VulkanUtil.h"
 #include "VulkanResources.h"
 #include "VulkanViewport.h"
+#include "VulkanDynamicRHI.h"
 
-class FVulkanTexture2D;
-class FVulkanFramebuffer;
-class FVulkanDevice;
-class FVulkanQueue;
-
-
-/** The interface which is implemented by the dynamically bound RHI. */
-class FVulkanDynamicRHI : public FDynamicRHI
-{
-public:
-
-	/** Initialization constructor. */
-	FVulkanDynamicRHI();
-
-	/** Destructor */
-	~FVulkanDynamicRHI() {}
-
-
-	// FDynamicRHI interface.
-	virtual void Init();
-	virtual void Shutdown();
-
-	void InitInstance();
-	void IssueLongGPUTask();
-
-	virtual FSamplerStateRHIRef RHICreateSamplerState(const FSamplerStateInitializerRHI& Initializer) final override;
-	virtual FRasterizerStateRHIRef RHICreateRasterizerState(const FRasterizerStateInitializerRHI& Initializer) final override;
-	virtual FDepthStencilStateRHIRef RHICreateDepthStencilState(const FDepthStencilStateInitializerRHI& Initializer) final override;
-	virtual FBlendStateRHIRef RHICreateBlendState(const FBlendStateInitializerRHI& Initializer) final override;
-	virtual FVertexDeclarationRHIRef RHICreateVertexDeclaration(const FVertexDeclarationElementList& Elements) final override;
-	virtual FPixelShaderRHIRef RHICreatePixelShader(const TArray<uint8>& Code) final override;
-	virtual FVertexShaderRHIRef RHICreateVertexShader(const TArray<uint8>& Code) final override;
-	virtual FHullShaderRHIRef RHICreateHullShader(const TArray<uint8>& Code) final override;
-	virtual FDomainShaderRHIRef RHICreateDomainShader(const TArray<uint8>& Code) final override;
-	virtual FGeometryShaderRHIRef RHICreateGeometryShader(const TArray<uint8>& Code) final override;
-	virtual FGeometryShaderRHIRef RHICreateGeometryShaderWithStreamOutput(const TArray<uint8>& Code, const FStreamOutElementList& ElementList, uint32 NumStrides, const uint32* Strides, int32 RasterizedStream) final override;
-	virtual FComputeShaderRHIRef RHICreateComputeShader(const TArray<uint8>& Code) final override;
-	virtual FBoundShaderStateRHIRef RHICreateBoundShaderState(FVertexDeclarationRHIParamRef VertexDeclaration, FVertexShaderRHIParamRef VertexShader, FHullShaderRHIParamRef HullShader, FDomainShaderRHIParamRef DomainShader, FPixelShaderRHIParamRef PixelShader, FGeometryShaderRHIParamRef GeometryShader) final override;
-	virtual FUniformBufferRHIRef RHICreateUniformBuffer(const void* Contents, const FRHIUniformBufferLayout& Layout, EUniformBufferUsage Usage) final override;
-	virtual FIndexBufferRHIRef RHICreateIndexBuffer(uint32 Stride, uint32 Size, uint32 InUsage, FRHIResourceCreateInfo& CreateInfo) final override;
-	virtual void* RHILockIndexBuffer(FIndexBufferRHIParamRef IndexBuffer, uint32 Offset, uint32 Size, EResourceLockMode LockMode) final override;
-	virtual void RHIUnlockIndexBuffer(FIndexBufferRHIParamRef IndexBuffer) final override;
-	virtual FVertexBufferRHIRef RHICreateVertexBuffer(uint32 Size, uint32 InUsage, FRHIResourceCreateInfo& CreateInfo) final override;
-	virtual void* RHILockVertexBuffer(FVertexBufferRHIParamRef VertexBuffer, uint32 Offset, uint32 SizeRHI, EResourceLockMode LockMode) final override;
-	virtual void RHIUnlockVertexBuffer(FVertexBufferRHIParamRef VertexBuffer) final override;
-	virtual void RHICopyVertexBuffer(FVertexBufferRHIParamRef SourceBuffer,FVertexBufferRHIParamRef DestBuffer) final override;
-	virtual FStructuredBufferRHIRef RHICreateStructuredBuffer(uint32 Stride, uint32 Size, uint32 InUsage, FRHIResourceCreateInfo& CreateInfo) final override;
-	virtual void* RHILockStructuredBuffer(FStructuredBufferRHIParamRef StructuredBuffer, uint32 Offset, uint32 SizeRHI, EResourceLockMode LockMode) final override;
-	virtual void RHIUnlockStructuredBuffer(FStructuredBufferRHIParamRef StructuredBuffer) final override;
-	virtual FUnorderedAccessViewRHIRef RHICreateUnorderedAccessView(FStructuredBufferRHIParamRef StructuredBuffer, bool bUseUAVCounter, bool bAppendBuffer) final override;
-	virtual FUnorderedAccessViewRHIRef RHICreateUnorderedAccessView(FTextureRHIParamRef Texture, uint32 MipLevel) final override;
-	virtual FUnorderedAccessViewRHIRef RHICreateUnorderedAccessView(FVertexBufferRHIParamRef VertexBuffer, uint8 Format) final override;
-	virtual FShaderResourceViewRHIRef RHICreateShaderResourceView(FStructuredBufferRHIParamRef StructuredBuffer) final override;
-	virtual FShaderResourceViewRHIRef RHICreateShaderResourceView(FVertexBufferRHIParamRef VertexBuffer, uint32 Stride, uint8 Format) final override;
-	virtual FShaderResourceViewRHIRef RHICreateShaderResourceView(FIndexBufferRHIParamRef Buffer) final override;
-	virtual uint64 RHICalcTexture2DPlatformSize(uint32 SizeX, uint32 SizeY, uint8 Format, uint32 NumMips, uint32 NumSamples, uint32 Flags, uint32& OutAlign) final override;
-	virtual uint64 RHICalcTexture3DPlatformSize(uint32 SizeX, uint32 SizeY, uint32 SizeZ, uint8 Format, uint32 NumMips, uint32 Flags, uint32& OutAlign) final override;
-	virtual uint64 RHICalcTextureCubePlatformSize(uint32 Size, uint8 Format, uint32 NumMips, uint32 Flags, uint32& OutAlign) final override;
-	virtual void RHIGetTextureMemoryStats(FTextureMemoryStats& OutStats) final override;
-	virtual bool RHIGetTextureMemoryVisualizeData(FColor* TextureData,int32 SizeX,int32 SizeY,int32 Pitch,int32 PixelSize) final override;
-	virtual FTextureReferenceRHIRef RHICreateTextureReference(FLastRenderTimeContainer* LastRenderTime) final override;
-	virtual FTexture2DRHIRef RHICreateTexture2D(uint32 SizeX, uint32 SizeY, uint8 Format, uint32 NumMips, uint32 NumSamples, uint32 Flags, FRHIResourceCreateInfo& CreateInfo) final override;
-	virtual FTexture2DRHIRef RHIAsyncCreateTexture2D(uint32 SizeX, uint32 SizeY, uint8 Format, uint32 NumMips, uint32 Flags, void** InitialMipData, uint32 NumInitialMips) final override;
-	virtual void RHICopySharedMips(FTexture2DRHIParamRef DestTexture2D, FTexture2DRHIParamRef SrcTexture2D) final override;
-	virtual FTexture2DArrayRHIRef RHICreateTexture2DArray(uint32 SizeX, uint32 SizeY, uint32 SizeZ, uint8 Format, uint32 NumMips, uint32 Flags, FRHIResourceCreateInfo& CreateInfo) final override;
-	virtual FTexture3DRHIRef RHICreateTexture3D(uint32 SizeX, uint32 SizeY, uint32 SizeZ, uint8 Format, uint32 NumMips, uint32 Flags, FRHIResourceCreateInfo& CreateInfo) final override;
-	virtual void RHIGetResourceInfo(FTextureRHIParamRef Ref, FRHIResourceInfo& OutInfo) final override;
-	virtual FShaderResourceViewRHIRef RHICreateShaderResourceView(FTexture2DRHIParamRef Texture2DRHI, uint8 MipLevel) final override;
-	virtual FShaderResourceViewRHIRef RHICreateShaderResourceView(FTexture2DRHIParamRef Texture2DRHI, uint8 MipLevel, uint8 NumMipLevels, uint8 Format) final override;
-	virtual FShaderResourceViewRHIRef RHICreateShaderResourceView(FTexture3DRHIParamRef Texture3DRHI, uint8 MipLevel) final override;
-	virtual FShaderResourceViewRHIRef RHICreateShaderResourceView(FTexture2DArrayRHIParamRef Texture2DArrayRHI, uint8 MipLevel) final override;
-	virtual FShaderResourceViewRHIRef RHICreateShaderResourceView(FTextureCubeRHIParamRef TextureCubeRHI, uint8 MipLevel) final override;
-	virtual void RHIGenerateMips(FTextureRHIParamRef Texture) final override;
-	virtual uint32 RHIComputeMemorySize(FTextureRHIParamRef TextureRHI) final override;
-	virtual FTexture2DRHIRef RHIAsyncReallocateTexture2D(FTexture2DRHIParamRef Texture2D, int32 NewMipCount, int32 NewSizeX, int32 NewSizeY, FThreadSafeCounter* RequestStatus) final override;
-	virtual ETextureReallocationStatus RHIFinalizeAsyncReallocateTexture2D(FTexture2DRHIParamRef Texture2D, bool bBlockUntilCompleted) final override;
-	virtual ETextureReallocationStatus RHICancelAsyncReallocateTexture2D(FTexture2DRHIParamRef Texture2D, bool bBlockUntilCompleted) final override;
-	virtual void* RHILockTexture2D(FTexture2DRHIParamRef Texture, uint32 MipIndex, EResourceLockMode LockMode, uint32& DestStride, bool bLockWithinMiptail) final override;
-	virtual void RHIUnlockTexture2D(FTexture2DRHIParamRef Texture, uint32 MipIndex, bool bLockWithinMiptail) final override;
-	virtual void* RHILockTexture2DArray(FTexture2DArrayRHIParamRef Texture, uint32 TextureIndex, uint32 MipIndex, EResourceLockMode LockMode, uint32& DestStride, bool bLockWithinMiptail) final override;
-	virtual void RHIUnlockTexture2DArray(FTexture2DArrayRHIParamRef Texture, uint32 TextureIndex, uint32 MipIndex, bool bLockWithinMiptail) final override;
-	virtual void RHIUpdateTexture2D(FTexture2DRHIParamRef Texture, uint32 MipIndex, const struct FUpdateTextureRegion2D& UpdateRegion, uint32 SourcePitch, const uint8* SourceData) final override;
-	virtual void RHIUpdateTexture3D(FTexture3DRHIParamRef Texture, uint32 MipIndex, const struct FUpdateTextureRegion3D& UpdateRegion, uint32 SourceRowPitch, uint32 SourceDepthPitch, const uint8* SourceData) final override;
-	virtual FTextureCubeRHIRef RHICreateTextureCube(uint32 Size, uint8 Format, uint32 NumMips, uint32 Flags, FRHIResourceCreateInfo& CreateInfo) final override;
-	virtual FTextureCubeRHIRef RHICreateTextureCubeArray(uint32 Size, uint32 ArraySize, uint8 Format, uint32 NumMips, uint32 Flags, FRHIResourceCreateInfo& CreateInfo) final override;
-	virtual void* RHILockTextureCubeFace(FTextureCubeRHIParamRef Texture, uint32 FaceIndex, uint32 ArrayIndex, uint32 MipIndex, EResourceLockMode LockMode, uint32& DestStride, bool bLockWithinMiptail) final override;
-	virtual void RHIUnlockTextureCubeFace(FTextureCubeRHIParamRef Texture, uint32 FaceIndex, uint32 ArrayIndex, uint32 MipIndex, bool bLockWithinMiptail) final override;
-	virtual void RHIBindDebugLabelName(FTextureRHIParamRef Texture, const TCHAR* Name) final override;
-	virtual void RHIBindDebugLabelName(FUnorderedAccessViewRHIParamRef UnorderedAccessViewRHI, const TCHAR* Name) final override;
-	virtual void RHIReadSurfaceData(FTextureRHIParamRef Texture,FIntRect Rect,TArray<FColor>& OutData,FReadSurfaceDataFlags InFlags) final override;
-	virtual void RHIMapStagingSurface(FTextureRHIParamRef Texture,void*& OutData,int32& OutWidth,int32& OutHeight) final override;
-	virtual void RHIUnmapStagingSurface(FTextureRHIParamRef Texture) final override;
-	virtual void RHIReadSurfaceFloatData(FTextureRHIParamRef Texture,FIntRect Rect,TArray<FFloat16Color>& OutData,ECubeFace CubeFace,int32 ArrayIndex,int32 MipIndex) final override;
-	virtual void RHIRead3DSurfaceFloatData(FTextureRHIParamRef Texture,FIntRect Rect,FIntPoint ZMinMax,TArray<FFloat16Color>& OutData) final override;
-	virtual FRenderQueryRHIRef RHICreateRenderQuery(ERenderQueryType QueryType) final override;
-	virtual bool RHIGetRenderQueryResult(FRenderQueryRHIParamRef RenderQuery, uint64& OutResult, bool bWait) final override;
-	virtual FTexture2DRHIRef RHIGetViewportBackBuffer(FViewportRHIParamRef Viewport) final override;
-	virtual void RHIAdvanceFrameForGetViewportBackBuffer() final override;
-	virtual void RHIAcquireThreadOwnership() final override;
-	virtual void RHIReleaseThreadOwnership() final override;
-	virtual void RHIFlushResources() final override;
-	virtual uint32 RHIGetGPUFrameCycles() final override;
-	virtual FViewportRHIRef RHICreateViewport(void* WindowHandle, uint32 SizeX, uint32 SizeY, bool bIsFullscreen, EPixelFormat PreferredPixelFormat) final override;
-	virtual void RHIResizeViewport(FViewportRHIParamRef Viewport, uint32 SizeX, uint32 SizeY, bool bIsFullscreen) final override;
-	virtual void RHITick(float DeltaTime) final override;
-	virtual void RHISetStreamOutTargets(uint32 NumTargets,const FVertexBufferRHIParamRef* VertexBuffers,const uint32* Offsets) final override;
-	virtual void RHIDiscardRenderTargets(bool Depth,bool Stencil,uint32 ColorBitMask) final override {}
-	virtual void RHIBlockUntilGPUIdle() final override;
-	virtual void RHISuspendRendering() final override;
-	virtual void RHIResumeRendering() final override;
-	virtual bool RHIIsRenderingSuspended() final override;
-	virtual bool RHIGetAvailableResolutions(FScreenResolutionArray& Resolutions, bool bIgnoreRefreshRate) final override;
-	virtual void RHIGetSupportedResolution(uint32& Width, uint32& Height) final override;
-	virtual void RHIVirtualTextureSetFirstMipInMemory(FTexture2DRHIParamRef Texture, uint32 FirstMip) final override;
-	virtual void RHIVirtualTextureSetFirstMipVisible(FTexture2DRHIParamRef Texture, uint32 FirstMip) final override;
-	virtual void RHIExecuteCommandList(FRHICommandList* CmdList) final override;
-	virtual void* RHIGetNativeDevice() final override;
-	virtual class IRHICommandContext* RHIGetDefaultContext() final override;
-	virtual class IRHICommandContextContainer* RHIGetCommandContextContainer() final override;
-
-	inline uint32 GetPresentCount() const
-	{
-		return PresentCount;
-	}
-
-	const TArray<const ANSICHAR*>& GetInstanceExtensions() const
-	{
-		return InstanceExtensions;
-	}
-
-	const TArray<const ANSICHAR*>& GetInstanceLayers() const
-	{
-		return InstanceLayers;
-	}
-
-private:
-    void PooledUniformBuffersBeginFrame();
-    void ReleasePooledUniformBuffers();
-
-protected:
-	VkInstance Instance;
-	TArray<const ANSICHAR*> InstanceExtensions;
-	TArray<const ANSICHAR*> InstanceLayers;
-
-	TArray<FVulkanDevice*> Devices;
-
-	FVulkanDevice* Device;
-
-	/** A list of all viewport RHIs that have been created. */
-	TArray<FVulkanViewport*> Viewports;
-
-	/** The viewport which is currently being drawn. */
-	TRefCountPtr<FVulkanViewport> DrawingViewport;
-
-	void CreateInstance();
-	void InitGPU(FVulkanDevice* Device);
-	void InitDevice(FVulkanDevice* Device);
-	
-	friend class FVulkanCommandListContext;
-
-	friend class FVulkanViewport;
-
-/*
-	static void WriteEndFrameTimestamp(void*);
-*/
-
-	static void GetInstanceLayersAndExtensions(TArray<const ANSICHAR*>& OutInstanceExtensions, TArray<const ANSICHAR*>& OutInstanceLayers);
-
-#if VULKAN_ENABLE_PIPELINE_CACHE
-	IConsoleObject* SavePipelineCacheCmd;
-	IConsoleObject* RebuildPipelineCacheCmd;
-	static void SavePipelineCache();
-	static void RebuildPipelineCache();
-#endif
-
-#if UE_BUILD_DEBUG || UE_BUILD_DEVELOPMENT
-	IConsoleObject* DumpMemoryCmd;
-public:
-	static void DumpMemory();
-#endif
-
-#if VULKAN_HAS_DEBUGGING_ENABLED
-protected:
-	VkDebugReportCallbackEXT MsgCallback;
-	void SetupDebugLayerCallback();
-	void RemoveDebugLayerCallback();
-#endif
-
-	inline bool ShouldDeferBufferLockOperation(FRHICommandList* RHICmdList)
-	{
-		if (RHICmdList == nullptr)
-		{
-			return false;
-		}
-
-		if (RHICmdList->Bypass() || !GRHIThread)
-		{
-			return false;
-		}
-
-		return true;
-	}
-
-	FCriticalSection LockBufferCS;
-
-	uint32 PresentCount;
-};
-
-/** Implements the Vulkan module as a dynamic RHI providing module. */
-class FVulkanDynamicRHIModule : public IDynamicRHIModule
-{
-public:
-	// IDynamicRHIModule
-	virtual bool IsSupported() override;
-
-	virtual FDynamicRHI* CreateRHI(ERHIFeatureLevel::Type RequestedFeatureLevel = ERHIFeatureLevel::Num) override;
-};
diff --git a/Engine/Source/Runtime/VulkanRHI/Public/VulkanResources.h b/Engine/Source/Runtime/VulkanRHI/Public/VulkanResources.h
--- a/Engine/Source/Runtime/VulkanRHI/Public/VulkanResources.h
+++ b/Engine/Source/Runtime/VulkanRHI/Public/VulkanResources.h
@@ -6,7 +6,9 @@
 
 #pragma once
 
-#include "VulkanRHI.h"
+#include "VulkanConfiguration.h"
+#include "VulkanState.h"
+#include "VulkanUtil.h"
 #include "BoundShaderStateCache.h"
 #include "VulkanShaderResources.h"
 #include "VulkanState.h"
diff --git a/Engine/Source/Runtime/Windows/D3D11RHI/Public/D3D11ConstantBuffer.h b/Engine/Source/Runtime/Windows/D3D11RHI/Public/D3D11ConstantBuffer.h
--- a/Engine/Source/Runtime/Windows/D3D11RHI/Public/D3D11ConstantBuffer.h
+++ b/Engine/Source/Runtime/Windows/D3D11RHI/Public/D3D11ConstantBuffer.h
@@ -4,8 +4,7 @@
 	D3D11ConstantBuffer.h: Public D3D Constant Buffer definitions.
 =============================================================================*/
 
-#ifndef __D3D11ConstantBuffer_H__
-#define __D3D11ConstantBuffer_H__
+#pragma once
 
 #include "RenderResource.h"
 
@@ -69,4 +68,3 @@ protected:
 
 DECLARE_CYCLE_STAT_EXTERN(TEXT("Global Constant buffer update time"),STAT_D3D11GlobalConstantBufferUpdateTime,STATGROUP_D3D11RHI, );
 
-#endif
diff --git a/Engine/Source/Runtime/Windows/D3D11RHI/Public/D3D11RHI.h b/Engine/Source/Runtime/Windows/D3D11RHI/Public/D3D11RHI.h
--- a/Engine/Source/Runtime/Windows/D3D11RHI/Public/D3D11RHI.h
+++ b/Engine/Source/Runtime/Windows/D3D11RHI/Public/D3D11RHI.h
@@ -4,7 +4,4 @@
 	D3D11RHI.h: Public D3D RHI definitions.
 =============================================================================*/
 
-#ifndef _INC_D3D11RHI
-#define _INC_D3D11RHI
-
-#endif
+#pragma once
diff --git a/Engine/Source/Runtime/Windows/D3D12RHI/Private/D3D12Device.h b/Engine/Source/Runtime/Windows/D3D12RHI/Private/D3D12Device.h
--- a/Engine/Source/Runtime/Windows/D3D12RHI/Private/D3D12Device.h
+++ b/Engine/Source/Runtime/Windows/D3D12RHI/Private/D3D12Device.h
@@ -7,6 +7,7 @@ D3D12Device.h: D3D12 Device Interfaces
 #pragma once
 
 class FD3D12DynamicRHI;
+class FD3D12OcclusionQuery;
 
 struct FD3D12Adapter
 {
@@ -253,4 +254,63 @@ protected:
 	FD3D12FenceCorePool FenceCorePool;
 
 	FD3D12ResidencyManager ResidencyManager;
-};
\ No newline at end of file
+};
+
+// Inline functions for FD3D12View which are dependent on the full FD3DDevice declaration
+
+template <typename TDesc>
+void FD3D12View<TDesc>::AllocateHeapSlot()
+{
+	FDescriptorHeapManager& DescriptorAllocator = GetParentDevice()->GetViewDescriptorAllocator<TDesc>();
+	Descriptor = DescriptorAllocator.AllocateHeapSlot(DescriptorHeapIndex);
+	check(Descriptor.ptr != 0);
+}
+
+template <typename TDesc>
+void FD3D12View<TDesc>::FreeHeapSlot()
+{
+	if (Descriptor.ptr)
+	{
+		FDescriptorHeapManager& DescriptorAllocator = GetParentDevice()->GetViewDescriptorAllocator<TDesc>();
+		DescriptorAllocator.FreeHeapSlot(Descriptor, DescriptorHeapIndex);
+		Descriptor.ptr = 0;
+	}
+}
+
+
+template <typename TDesc>
+void FD3D12View<TDesc>::CreateView(FD3D12Resource* InResource, FD3D12Resource* InCounterResource)
+{
+	if (!InResource)
+	{
+		InResource = GetResource();
+	}
+	else
+	{
+		// Only need to update the view's subresource subset if a new resource is used
+		UpdateViewSubresourceSubset(InResource);
+	}
+
+	check(Descriptor.ptr != 0);
+	(GetParentDevice()->GetDevice()->*TCreateViewMap<TDesc>::GetCreate()) (
+		InResource ? InResource->GetResource() : nullptr, &Desc, Descriptor);
+}
+
+template <typename TDesc>
+void FD3D12View<TDesc>::CreateViewWithCounter(FD3D12Resource* InResource, FD3D12Resource* InCounterResource)
+{
+	if (!InResource)
+	{
+		InResource = GetResource();
+	}
+	else
+	{
+		// Only need to update the view's subresource subset if a new resource is used
+		UpdateViewSubresourceSubset(InResource);
+	}
+
+	check(Descriptor.ptr != 0);
+	(GetParentDevice()->GetDevice()->*TCreateViewMap<TDesc>::GetCreate()) (
+		InResource ? InResource->GetResource() : nullptr,
+		InCounterResource ? InCounterResource->GetResource() : nullptr, &Desc, Descriptor);
+}
diff --git a/Engine/Source/Runtime/Windows/D3D12RHI/Public/D3D12Resources.h b/Engine/Source/Runtime/Windows/D3D12RHI/Public/D3D12Resources.h
--- a/Engine/Source/Runtime/Windows/D3D12RHI/Public/D3D12Resources.h
+++ b/Engine/Source/Runtime/Windows/D3D12RHI/Public/D3D12Resources.h
@@ -635,6 +635,80 @@ public:
 	void UpdateDefaultStateCache();
 };
 
+/** Information about a D3D resource that is currently locked. */
+struct FD3D12LockedData
+{
+	TRefCountPtr<FD3D12Resource> StagingResource;
+	TRefCountPtr<FD3D12Resource> UploadHeapResource;
+	TRefCountPtr<FD3D12ResourceLocation> UploadHeapLocation;
+	uint32 Pitch;
+	uint32 DepthPitch;
+
+	// constructor
+	FD3D12LockedData()
+		: bAllocDataWasUsed(false)
+	{
+	}
+
+	// 16 byte alignment for best performance  (can be 30x faster than unaligned)
+	void AllocData(uint32 Size)
+	{
+		Data = (uint8*)FMemory::Malloc(Size, 16);
+		bAllocDataWasUsed = true;
+	}
+
+	// Some driver might return aligned memory so we don't enforce the alignment
+	void SetData(void* InData)
+	{
+		check(!bAllocDataWasUsed); Data = (uint8*)InData;
+	}
+
+	uint8* GetData() const
+	{
+		return Data;
+	}
+
+	// only call if AllocData() was used
+	void FreeData()
+	{
+		check(bAllocDataWasUsed);
+		FMemory::Free(Data);
+		Data = 0;
+	}
+
+private:
+	uint8* Data;
+	// then FreeData
+	bool bAllocDataWasUsed;
+};
+
+/**
+* Class for managing dynamic buffers.
+*/
+class FD3D12DynamicBuffer : public FRenderResource, public FRHIResource, public FD3D12DeviceChild
+{
+public:
+	/** Initialization constructor. */
+	FD3D12DynamicBuffer(FD3D12Device* InParent, class FD3D12FastAllocator& Allocator);
+	/** Destructor. */
+	~FD3D12DynamicBuffer();
+
+	/** Locks the buffer returning at least Size bytes. */
+	void* Lock(uint32 Size);
+	/** Unlocks the buffer returning the underlying D3D12 buffer to use as a resource. */
+	FD3D12ResourceLocation* Unlock();
+
+	// Begin FRenderResource interface.
+	virtual void InitRHI() override;
+	virtual void ReleaseRHI() override;
+	// End FRenderResource interface.
+
+	void ReleaseResourceLocation() { ResourceLocation = nullptr; }
+
+private:
+	TRefCountPtr<FD3D12ResourceLocation> ResourceLocation;
+	class FD3D12FastAllocator& FastAllocator;
+};
 
 class FD3D12DeferredDeletionQueue : public FD3D12DeviceChild
 {
@@ -1968,22 +2042,11 @@ protected:
 		AllocateHeapSlot();
 	}
 
-	void AllocateHeapSlot()
-	{
-		FDescriptorHeapManager& DescriptorAllocator = GetParentDevice()->GetViewDescriptorAllocator<TDesc>();
-		Descriptor = DescriptorAllocator.AllocateHeapSlot(DescriptorHeapIndex);
-		check(Descriptor.ptr != 0);
-	}
+	// Implemented in D3D12Device.h due to dependency on FD3D12Device declaration
+	void AllocateHeapSlot();
 
-	void FreeHeapSlot()
-	{
-		if (Descriptor.ptr)
-		{
-			FDescriptorHeapManager& DescriptorAllocator = GetParentDevice()->GetViewDescriptorAllocator<TDesc>();
-			DescriptorAllocator.FreeHeapSlot(Descriptor, DescriptorHeapIndex);
-			Descriptor.ptr = 0;
-		}
-	}
+	// Implemented in D3D12Device.h due to dependency on FD3D12Device declaration
+	void FreeHeapSlot();
 
 	void UpdateViewSubresourceSubset(const FD3D12Resource* const InResource)
 	{
@@ -2011,40 +2074,11 @@ protected:
 
 
 public:
-	void CreateView(FD3D12Resource* InResource = nullptr, FD3D12Resource* InCounterResource = nullptr)
-	{
-		if (!InResource)
-		{
-			InResource = GetResource();
-		}
-		else
-		{
-			// Only need to update the view's subresource subset if a new resource is used
-			UpdateViewSubresourceSubset(InResource);
-		}
+	// Implemented in D3D12Device.h due to dependency on FD3D12Device declaration
+	void CreateView(FD3D12Resource* InResource = nullptr, FD3D12Resource* InCounterResource = nullptr);
 
-		check(Descriptor.ptr != 0);
-		(GetParentDevice()->GetDevice()->*TCreateViewMap<TDesc>::GetCreate()) (
-			InResource ? InResource->GetResource() : nullptr, &Desc, Descriptor);
-	}
-
-	void CreateViewWithCounter(FD3D12Resource* InResource, FD3D12Resource* InCounterResource)
-	{
-		if (!InResource)
-		{
-			InResource = GetResource();
-		}
-		else
-		{
-			// Only need to update the view's subresource subset if a new resource is used
-			UpdateViewSubresourceSubset(InResource);
-		}
-
-		check(Descriptor.ptr != 0);
-		(GetParentDevice()->GetDevice()->*TCreateViewMap<TDesc>::GetCreate()) (
-			InResource ? InResource->GetResource() : nullptr,
-			InCounterResource ? InCounterResource->GetResource() : nullptr, &Desc, Descriptor);
-	}
+	// Implemented in D3D12Device.h due to dependency on FD3D12Device declaration
+	void CreateViewWithCounter(FD3D12Resource* InResource, FD3D12Resource* InCounterResource);
 
 	const TDesc& GetDesc() const { return Desc; }
 	const CViewSubresourceSubset& GetViewSubresourceSubset() const { return ViewSubresourceSubset; }
diff --git a/Engine/Source/Runtime/Windows/D3D12RHI/Public/D3D12Util.h b/Engine/Source/Runtime/Windows/D3D12RHI/Public/D3D12Util.h
--- a/Engine/Source/Runtime/Windows/D3D12RHI/Public/D3D12Util.h
+++ b/Engine/Source/Runtime/Windows/D3D12RHI/Public/D3D12Util.h
@@ -324,53 +324,6 @@ public:
 	}
 };
 
-/** Information about a D3D resource that is currently locked. */
-struct FD3D12LockedData
-{
-	TRefCountPtr<FD3D12Resource> StagingResource;
-	TRefCountPtr<FD3D12Resource> UploadHeapResource;
-	TRefCountPtr<FD3D12ResourceLocation> UploadHeapLocation;
-	uint32 Pitch;
-	uint32 DepthPitch;
-
-	// constructor
-	FD3D12LockedData()
-		: bAllocDataWasUsed(false)
-	{
-	}
-
-	// 16 byte alignment for best performance  (can be 30x faster than unaligned)
-	void AllocData(uint32 Size)
-	{
-		Data = (uint8*)FMemory::Malloc(Size, 16);
-		bAllocDataWasUsed = true;
-	}
-
-	// Some driver might return aligned memory so we don't enforce the alignment
-	void SetData(void* InData)
-	{
-		check(!bAllocDataWasUsed); Data = (uint8*)InData;
-	}
-
-	uint8* GetData() const
-	{
-		return Data;
-	}
-
-	// only call if AllocData() was used
-	void FreeData()
-	{
-		check(bAllocDataWasUsed);
-		FMemory::Free(Data);
-		Data = 0;
-	}
-
-private:
-	uint8* Data;
-	// then FreeData
-	bool bAllocDataWasUsed;
-};
-
 class FD3D12RenderTargetView;
 class FD3D12DepthStencilView;
 
@@ -402,34 +355,6 @@ private:
 	int32 NumActiveTargets;
 };
 
-/**
- * Class for managing dynamic buffers.
- */
-class FD3D12DynamicBuffer : public FRenderResource, public FRHIResource, public FD3D12DeviceChild
-{
-public:
-	/** Initialization constructor. */
-	FD3D12DynamicBuffer(FD3D12Device* InParent, class FD3D12FastAllocator& Allocator);
-	/** Destructor. */
-	~FD3D12DynamicBuffer();
-
-	/** Locks the buffer returning at least Size bytes. */
-	void* Lock(uint32 Size);
-	/** Unlocks the buffer returning the underlying D3D12 buffer to use as a resource. */
-	FD3D12ResourceLocation* Unlock();
-
-	// Begin FRenderResource interface.
-	virtual void InitRHI() override;
-	virtual void ReleaseRHI() override;
-	// End FRenderResource interface.
-
-	void ReleaseResourceLocation() { ResourceLocation = nullptr; }
-
-private:
-	TRefCountPtr<FD3D12ResourceLocation> ResourceLocation;
-	class FD3D12FastAllocator& FastAllocator;
-};
-
 static D3D12_DESCRIPTOR_HEAP_DESC CreateDHD(D3D12_DESCRIPTOR_HEAP_TYPE Type, uint32 NumDescriptorsPerHeap, D3D12_DESCRIPTOR_HEAP_FLAGS Flags)
 {
 	D3D12_DESCRIPTOR_HEAP_DESC DHD = {Type, NumDescriptorsPerHeap, Flags};
diff --git a/Engine/Source/Runtime/Windows/XAudio2/Private/XAudio2Buffer.cpp b/Engine/Source/Runtime/Windows/XAudio2/Private/XAudio2Buffer.cpp
--- a/Engine/Source/Runtime/Windows/XAudio2/Private/XAudio2Buffer.cpp
+++ b/Engine/Source/Runtime/Windows/XAudio2/Private/XAudio2Buffer.cpp
@@ -11,6 +11,7 @@
 	Audio includes.
 ------------------------------------------------------------------------------------*/
 
+#include "XAudio2PrivatePCH.h"
 #include "XAudio2Device.h"
 #include "AudioDecompress.h"
 #include "AudioEffect.h"
diff --git a/Engine/Source/Runtime/Windows/XAudio2/Private/XAudio2Device.cpp b/Engine/Source/Runtime/Windows/XAudio2/Private/XAudio2Device.cpp
--- a/Engine/Source/Runtime/Windows/XAudio2/Private/XAudio2Device.cpp
+++ b/Engine/Source/Runtime/Windows/XAudio2/Private/XAudio2Device.cpp
@@ -11,6 +11,7 @@
 	Audio includes.
 ------------------------------------------------------------------------------------*/
 
+#include "XAudio2PrivatePCH.h"
 #include "XAudio2Device.h"
 #include "AudioEffect.h"
 #include "OpusAudioInfo.h"
diff --git a/Engine/Source/Runtime/Windows/XAudio2/Private/XAudio2PrivatePCH.h b/Engine/Source/Runtime/Windows/XAudio2/Private/XAudio2PrivatePCH.h
new file mode 100644
--- /dev/null
+++ b/Engine/Source/Runtime/Windows/XAudio2/Private/XAudio2PrivatePCH.h
@@ -0,0 +1,5 @@
+// Copyright 1998-2016 Epic Games, Inc. All Rights Reserved.
+
+#pragma once
+
+#include "Engine.h"
diff --git a/Engine/Source/Runtime/Windows/XAudio2/XAudio2.Build.cs b/Engine/Source/Runtime/Windows/XAudio2/XAudio2.Build.cs
--- a/Engine/Source/Runtime/Windows/XAudio2/XAudio2.Build.cs
+++ b/Engine/Source/Runtime/Windows/XAudio2/XAudio2.Build.cs
@@ -33,5 +33,7 @@ public class XAudio2 : ModuleRules
 			"Vorbis",
 			"VorbisFile"
 			);
+
+		Definitions.Add("WITH_OGGVORBIS=1");
 	}
 }
