UE4 - Android compile fix

diff --git a/Engine/Source/Runtime/VulkanRHI/Private/VulkanTexture.cpp b/Engine/Source/Runtime/VulkanRHI/Private/VulkanTexture.cpp
--- a/Engine/Source/Runtime/VulkanRHI/Private/VulkanTexture.cpp
+++ b/Engine/Source/Runtime/VulkanRHI/Private/VulkanTexture.cpp
@@ -1006,7 +1006,7 @@ FVulkanTextureBase::FVulkanTextureBase(FVulkanDevice& Device, VkImageViewType Re
 		MSAASurface = new FVulkanSurface(Device, ResourceType, InFormat, SizeX, SizeY, SizeZ, /*bArray=*/ false, 1, NumMips, NumSamples, UEFlags, CreateInfo);
 		if (ResourceType != VK_IMAGE_VIEW_TYPE_MAX_ENUM)
 		{
-			MSAAView.Create(Device, MSAASurface->Image, ResourceType, MSAASurface->GetFullAspectMask(), MSAASurface->Format, MSAASurface->InternalFormat, FMath::Max(NumMips, 1u), FMath::Max(1u, ArraySize));
+			MSAAView.Create(Device, MSAASurface->Image, ResourceType, MSAASurface->GetFullAspectMask(), MSAASurface->Format, MSAASurface->InternalFormat, 0, FMath::Max(NumMips, 1u), 0, FMath::Max(1u, ArraySize));
 		}
 	}
 #endif
