UE4 - Fix Android compile error

diff --git a/Engine/Source/Runtime/VulkanRHI/Private/VulkanViewport.cpp b/Engine/Source/Runtime/VulkanRHI/Private/VulkanViewport.cpp
--- a/Engine/Source/Runtime/VulkanRHI/Private/VulkanViewport.cpp
+++ b/Engine/Source/Runtime/VulkanRHI/Private/VulkanViewport.cpp
@@ -171,7 +171,7 @@ FVulkanFramebuffer::FVulkanFramebuffer(FVulkanDevice& Device, const FRHISetRende
 		    Barrier.oldLayout = VK_IMAGE_LAYOUT_GENERAL;
 		    Barrier.newLayout = VK_IMAGE_LAYOUT_GENERAL;
 		    Barrier.image = Texture->Surface.Image;
-		    Barrier.subresourceRange.aspectMask = Texture->MSAASurface->GetAspectMask();
+		    Barrier.subresourceRange.aspectMask = Texture->MSAASurface->GetFullAspectMask();
 		    //Barrier.subresourceRange.baseMipLevel = 0;
 		    Barrier.subresourceRange.levelCount = 1;
 		    //Barrier.subresourceRange.baseArrayLayer = 0;
