Copying //UE4/Dev-Mobile to //UE4/Dev-Main (Source: //UE4/Dev-Mobile @ 3116515)

==========================
MAJOR FEATURES + CHANGES
==========================

Change 3065209 on 2016/07/26 by [REDACTED_PERSON]

Add Android Texture Format used for packaging/cooking to the Manifest File

#ue4
#android

Change 3068915 on 2016/07/28 by [REDACTED_PERSON]

Add an additional Texture Compression support line in the manifest for DXT

#ue4
#android

Change 3075911 on 2016/08/03 by [REDACTED_PERSON]

Make the "Running {ProjectName} on {Device}" toast stay up when launching a game to an Android device until we've finished running it, as it does on other platforms. This logic already existed but only ran if the "Prebuilt" flag was passed in, however we want this to always run now. Re-writing to run through waiting on each process to finish for each device launched on

#ue4
#android

Change 3080981 on 2016/08/08 by [REDACTED_PERSON]

Clear any input before removing the TouchInterface from the screen to prevent infinite input after it is cleared

#ue4
#platform

Change 3092587 on 2016/08/17 by [REDACTED_PERSON]

Adding "IsGamepadAttached" functionality to Android Application

#ue4
#android

Change 3095840 on 2016/08/21 by [REDACTED_PERSON]

Fixed: Particle Cutout Crashes On Certain Devices (Samsung Galaxy Note 2)
Happens only with non-instanced path

Change 3095855 on 2016/08/22 by [REDACTED_PERSON]

Allow UWorldComposition::GetTilesList to be used in runtime code
Licensee request https://udn.unrealengine.com/questions/307586/world-compositions-world-dimensions.html

Change 3096093 on 2016/08/22 by [REDACTED_PERSON]

Allow Vulkan api logging on android

Change 3096361 on 2016/08/22 by [REDACTED_PERSON]

Github 2663 pull - Pass any extras used to launch SplashActivity down to GameActivity (Contributed by [REDACTED_PERSON])

#github #2663
#ue4
#android

Change 3097412 on 2016/08/23 by [REDACTED_PERSON]

Using BulkSerialize for cooked collision data to speed up serialization

Change 3098957 on 2016/08/23 by [REDACTED_PERSON]

Merging //UE4/Dev-Main to Dev-Mobile (//UE4/Dev-Mobile)

Change 3099058 on 2016/08/24 by [REDACTED_PERSON]

Check EXT_debug_label and EXT_debug_marker functions were found before calling as a few devices to not implement these extensions
#UE-35087

Change 3099131 on 2016/08/24 by [REDACTED_PERSON]

Fixed: HDR compressed texture become black in some mali devices
Use sized internal format for half-float textures on ES3 devices, as ES3 spec expects it

Change 3099150 on 2016/08/24 by [REDACTED_PERSON]

Enable HALF_FLOAT and UNSIGNED_INT_2_10_10_10_REV vertex formats on ES3+ devices, spec req

Change 3102252 on 2016/08/26 by [REDACTED_PERSON]

Prevent view uniform buffer crash on ES2 devices that do not support 3D textures

Change 3102258 on 2016/08/26 by [REDACTED_PERSON]

Enabled refraction on iPadMini4

Change 3102651 on 2016/08/26 by [REDACTED_PERSON]

Fixed: instanced static mesh world normals
Also removed unnecessary instance matrix transposing

Change 3109397 on 2016/09/01 by [REDACTED_PERSON]

Fix problem with Android sessions not appearing in Session Frontend

Change 3109490 on 2016/09/01 by [REDACTED_PERSON]

Merging //UE4/Dev-Main to Dev-Mobile (//UE4/Dev-Mobile)

Change 3111628 on 2016/09/02 by [REDACTED_PERSON]

Landscape wireframe LOD visualization with tessellation

Change 3112809 on 2016/09/02 by [REDACTED_PERSON]

Update cached length when file is written to on Android
#ue4
#android

Change 3113245 on 2016/09/04 by [REDACTED_PERSON]

Fixed: Subway Sequencer plays only a black screen when packaged for ESDSR (3.1+AEP)

Change 3113249 on 2016/09/04 by [REDACTED_PERSON]

Replicated fix from 4.13: GPU particles no longer work on iOS or TVOS Metal devices

Change 3113513 on 2016/09/05 by [REDACTED_PERSON]

Add vulkan version parameter to android device profile selector's source inputs .
reinstate Vulkan Disable cvar functionality.
Added mali no vulkan device profile.

Change 3113519 on 2016/09/05 by [REDACTED_PERSON]

Remove temp 4.13 hack to avoid public header changes.

Change 3113535 on 2016/09/05 by [REDACTED_PERSON]

Decode 32 bit HDR formats when using scene captures.

Change 3113813 on 2016/09/06 by [REDACTED_PERSON]

Resend to server sub-levels visibility state right after world actors are initialized.
During seamless travel client loads always-loaded sub-levels before world actors are initialized and ServerUpdateLevelVisibility calls in UWorld::AddToWorld are skipped.

Change 3113870 on 2016/09/06 by [REDACTED_PERSON]

Fix issue with ES2 Feature Level preview and Mobile Preview PIE not limiting materials to 8 textures

Change 3115031 on 2016/09/06 by [REDACTED_PERSON]

Add Vulkan version code not moved over from 4.13.1
#ue4
#android

Change 3115496 on 2016/09/07 by [REDACTED_PERSON]

Use the same DDC key for source reflection data and encoded data. Fixes broken reflections on mobile

diff --git a/Engine/Build/Android/Java/src/com/epicgames/ue4/GameActivity.java b/Engine/Build/Android/Java/src/com/epicgames/ue4/GameActivity.java
--- a/Engine/Build/Android/Java/src/com/epicgames/ue4/GameActivity.java
+++ b/Engine/Build/Android/Java/src/com/epicgames/ue4/GameActivity.java
@@ -32,6 +32,7 @@ import android.content.IntentSender.SendIntentException;
 import android.content.pm.ApplicationInfo;
 import android.content.pm.PackageManager;
 import android.content.pm.PackageManager.NameNotFoundException;
+import android.content.pm.FeatureInfo;
 
 import android.media.AudioManager;
 import android.util.DisplayMetrics;
@@ -193,7 +194,11 @@ public class GameActivity extends NativeActivity implements SurfaceHolder.Callba
 	private SurfaceView MySurfaceView;
 	private int DesiredHolderWidth = 0;
 	private int DesiredHolderHeight = 0;
-	
+
+	/** Discovered Vulkan Version and Level from getSystemAvailableFeatures() */
+	private int VulkanVersion = 0;
+	private int VulkanLevel = 0;
+
 	enum EAlertDialogType
 	{
 		None,
@@ -376,7 +381,54 @@ public class GameActivity extends NativeActivity implements SurfaceHolder.Callba
 
 		// tell Android that we want volume controls to change the media volume, aka music
 		setVolumeControlStream(AudioManager.STREAM_MUSIC);
-		
+
+		// Look for Vulkan support if Nougat or later
+		if (ANDROID_BUILD_VERSION >= 24)
+		{
+			FeatureInfo[] features = getPackageManager().getSystemAvailableFeatures();
+			for (FeatureInfo feature : features) {
+				if (feature.name != null)
+				{
+					if (feature.name.equals("android.hardware.vulkan.level"))
+					{
+						// since we may not be compiled against android-24 or higher, use .toString to get the version field
+						String dump = feature.toString();
+						int index = dump.indexOf("v=");
+						if (index >= 0)
+						{
+							dump = dump.substring(index+2);
+							index = dump.indexOf(" ");
+							if (index >= 0)
+							{
+								VulkanLevel = Integer.parseInt(dump.substring(0, index));
+								Log.debug("Vulkan level: " + VulkanLevel);
+							}
+						}
+					}
+					else
+					if (feature.name.equals("android.hardware.vulkan.version"))
+					{
+						// since we may not be compiled against android-24 or higher, use .toString to get the version field
+						String dump = feature.toString();
+						int index = dump.indexOf("v=");
+						if (index >= 0)
+						{
+							dump = dump.substring(index+2);
+							index = dump.indexOf(" ");
+							if (index >= 0)
+							{
+								VulkanVersion = Integer.parseInt(dump.substring(0, index));
+								int VersionMajor = (VulkanVersion >> 22) & 0x03ff;
+								int VersionMinor = (VulkanVersion >> 12) & 0x03ff;
+								int VersionPatch = VulkanVersion & 0x0fff;
+								Log.debug("Vulkan version: " + VersionMajor + "." + VersionMinor + "." + VersionPatch);
+							}
+						}
+					}
+				}
+			}
+		}
+
 		// is this a native landscape device (tablet, tv)?
 		if ( getDeviceDefaultOrientation() == Configuration.ORIENTATION_LANDSCAPE )
 		{
@@ -1699,6 +1751,24 @@ public class GameActivity extends NativeActivity implements SurfaceHolder.Callba
 		}
 	}
 
+	public boolean AndroidThunkJava_IsGamepadAttached()
+	{
+		int[] deviceIds = InputDevice.getDeviceIds();
+		for (int deviceIndex=0; deviceIndex < deviceIds.length; deviceIndex++)
+		{
+			InputDevice inputDevice = InputDevice.getDevice(deviceIds[deviceIndex]);
+			// is it a joystick, dpad, or gamepad?
+			if (((inputDevice.getSources() & InputDevice.SOURCE_GAMEPAD) == InputDevice.SOURCE_GAMEPAD)
+                || ((inputDevice.getSources() & InputDevice.SOURCE_JOYSTICK) == InputDevice.SOURCE_JOYSTICK)
+				|| ((inputDevice.getSources() & InputDevice.SOURCE_DPAD) == InputDevice.SOURCE_DPAD))
+			{
+				return true;
+			}
+		}
+
+		return false;
+	}
+
 	public boolean AndroidThunkJava_HasMetaDataKey(String key)
 	{
 		if (_bundle == null || key == null)
@@ -1719,6 +1789,15 @@ public class GameActivity extends NativeActivity implements SurfaceHolder.Callba
 
 	public int AndroidThunkJava_GetMetaDataInt(String key)
 	{
+		if (key.equals("android.hardware.vulkan.version"))
+		{
+			return VulkanVersion;
+		}
+		else
+		if (key.equals("android.hardware.vulkan.level"))
+		{
+			return VulkanLevel;
+		}
 		if (_bundle == null || key == null)
 		{
 			return 0;
diff --git a/Engine/Build/Android/Java/src/com/epicgames/ue4/SplashActivity.java b/Engine/Build/Android/Java/src/com/epicgames/ue4/SplashActivity.java
--- a/Engine/Build/Android/Java/src/com/epicgames/ue4/SplashActivity.java
+++ b/Engine/Build/Android/Java/src/com/epicgames/ue4/SplashActivity.java
@@ -55,6 +55,15 @@ public class SplashActivity extends Activity
 		{
 			intent.putExtra("ShouldHideUI", "true");
 		}
+
+        //pass down any extras added to this Activity's intent to the GameActivity intent (GCM data, for example)
+        Intent intentFromActivity = getIntent();
+        Bundle intentBundle = intentFromActivity.getExtras();
+        if(intentBundle != null)
+		{
+            intent.putExtras(intentBundle);
+        }
+
 		startActivity(intent);
 		finish();
 		overridePendingTransition(0, 0);
@@ -68,4 +77,4 @@ public class SplashActivity extends Activity
 		overridePendingTransition(0, 0);
 	}
 
-}
+}
\ No newline at end of file
diff --git a/Engine/Config/BaseDeviceProfiles.ini b/Engine/Config/BaseDeviceProfiles.ini
--- a/Engine/Config/BaseDeviceProfiles.ini
+++ b/Engine/Config/BaseDeviceProfiles.ini
@@ -54,9 +54,7 @@
 +DeviceProfileNameAndTypes=Android_Mali_T6xx,Android
 +DeviceProfileNameAndTypes=Android_Mali_T7xx,Android
 +DeviceProfileNameAndTypes=Android_Mali_T8xx,Android
-+DeviceProfileNameAndTypes=Android_Vulkan_GalaxyS7_Adreno,Android
-+DeviceProfileNameAndTypes=Android_Vulkan_GalaxyS7_Mali,Android
-+DeviceProfileNameAndTypes=Android_Vulkan_HUAWEI_Mali,Android
++DeviceProfileNameAndTypes=Android_Mali_T8xx_No_Vulkan,Android
 
 [/Script/Engine.TextureLODSettings]
 TextureLODGroups=(Group=TEXTUREGROUP_World,MinLODSize=1,MaxLODSize=4096,LODBias=0,MinMagFilter=aniso,MipFilter=point,MipGenSettings=TMGS_SimpleAverage)
@@ -163,6 +161,7 @@ BaseProfileName=iPadAir
 [iPadMini4 DeviceProfile]
 DeviceType=IOS
 BaseProfileName=iPadMini2
++CVars=r.RefractionQuality=1
 
 [iPhone4 DeviceProfile]
 DeviceType=IOS
@@ -271,12 +270,10 @@ MatchProfile=(Profile="Android_PowerVR54x",Match=((SourceType=SRC_GpuFamily,Comp
 +MatchProfile=(Profile="Android_Mali_4xx",Match=((SourceType=SRC_GpuFamily,CompareType=CMP_Regex,MatchString="^Mali\\-4")))
 +MatchProfile=(Profile="Android_Mali_T6xx",Match=((SourceType=SRC_GpuFamily,CompareType=CMP_Regex,MatchString="^Mali\\-T6")))
 +MatchProfile=(Profile="Android_Mali_T7xx",Match=((SourceType=SRC_GpuFamily,CompareType=CMP_Regex,MatchString="^Mali\\-T7")))
++MatchProfile=(Profile="Android_Mali_T8xx_No_Vulkan",Match=((SourceType=SRC_GpuFamily,CompareType=CMP_Regex,MatchString="^Mali\\-T8"),(SourceType=SRC_VulkanVersion,CompareType=CMP_Regex,MatchString="^1\\.0\\.([0-9]+)"),(SourceType=SRC_PreviousRegexMatch,CompareType=CMP_Less,MatchString="11")))
 +MatchProfile=(Profile="Android_Mali_T8xx",Match=((SourceType=SRC_GpuFamily,CompareType=CMP_Regex,MatchString="^Mali\\-T8")))
 +MatchProfile=(Profile="Android_Mid",Match=((SourceType=SRC_GlVersion,CompareType=CMP_Regex,MatchString="^OpenGL ES 3\\.")))
 +MatchProfile=(Profile="Android_Low",Match=((SourceType=SRC_GlVersion,CompareType=CMP_Regex,MatchString="^OpenGL ES 2\\.")))
-+MatchProfile=(Profile="Android_Vulkan_GalaxyS7_Mali",Match=((SourceType=SRC_DeviceModel,CompareType=CMP_Regex,MatchString="SM\\-G93[05][WXFKLS]"),(SourceType=SRC_GpuFamily,CompareType=CMP_Equal,MatchString="Vulkan Mobile")))
-+MatchProfile=(Profile="Android_Vulkan_GalaxyS7_Adreno",Match=((SourceType=SRC_DeviceModel,CompareType=CMP_Regex,MatchString="SM\\-G93[05]"),(SourceType=SRC_GpuFamily,CompareType=CMP_Equal,MatchString="Vulkan Mobile")))
-+MatchProfile=(Profile="Android_Vulkan_HUAWEI_Mali",Match=((SourceType=SRC_DeviceModel,CompareType=CMP_Equal,MatchString="HUAWEI"),(SourceType=SRC_GpuFamily,CompareType=CMP_Equal,MatchString="Vulkan Mobile")))
 
 [/Script/AndroidDeviceProfileSelector.AndroidJavaSurfaceViewDevices]
 SurfaceViewDevices=(Manufacturer="Amazon",Model="Kindle Fire")	; Kindle Fire (2011)
@@ -345,6 +342,7 @@ BaseProfileName=Android_Adreno330
 [Android_Adreno5xx DeviceProfile]
 DeviceType=Android
 BaseProfileName=Android_High
++CVars=r.Vulkan.UseGLSL=2
 
 [Android_Adreno4xx DeviceProfile]
 DeviceType=Android
@@ -411,18 +409,10 @@ BaseProfileName=Android_High
 DeviceType=Android
 BaseProfileName=Android_Mid
 
-[Android_Vulkan_GalaxyS7_Mali DeviceProfile]
-DeviceType=Android
-BaseProfileName=Android_High
-
-[Android_Vulkan_GalaxyS7_Adreno DeviceProfile]
+[Android_Mali_T8xx_No_Vulkan DeviceProfile]
 DeviceType=Android
-BaseProfileName=Android_High
-+CVars=r.Vulkan.UseGLSL=2
-
-[Android_Vulkan_HUAWEI_Mali DeviceProfile]
-DeviceType=Android
-BaseProfileName=Android_High
+BaseProfileName=Android_Mali_T8xx
++CVars=r.Android.DisableVulkanSupport=1
 
 [PS4 DeviceProfile]
 DeviceType=PS4
diff --git a/Engine/Plugins/Runtime/AndroidDeviceProfileSelector/Source/AndroidDeviceProfileSelector/Private/AndroidDeviceProfileSelectorModule.cpp b/Engine/Plugins/Runtime/AndroidDeviceProfileSelector/Source/AndroidDeviceProfileSelector/Private/AndroidDeviceProfileSelectorModule.cpp
--- a/Engine/Plugins/Runtime/AndroidDeviceProfileSelector/Source/AndroidDeviceProfileSelector/Private/AndroidDeviceProfileSelectorModule.cpp
+++ b/Engine/Plugins/Runtime/AndroidDeviceProfileSelector/Source/AndroidDeviceProfileSelector/Private/AndroidDeviceProfileSelectorModule.cpp
@@ -38,6 +38,8 @@ FString const FAndroidDeviceProfileSelectorModule::GetRuntimeDeviceProfileName()
 
 		FString GPUFamily = FAndroidMisc::GetGPUFamily();
 		FString GLVersion = FAndroidMisc::GetGLVersion();
+
+		FString VulkanVersion = FAndroidMisc::GetVulkanVersion();
 		FString AndroidVersion = FAndroidMisc::GetAndroidVersion();
 		FString DeviceMake = FAndroidMisc::GetDeviceMake();
 		FString DeviceModel = FAndroidMisc::GetDeviceModel();
@@ -54,6 +56,7 @@ FString const FAndroidDeviceProfileSelectorModule::GetRuntimeDeviceProfileName()
 		UE_LOG(LogAndroid, Log, TEXT("  Default profile: %s"), *ProfileName);
 		UE_LOG(LogAndroid, Log, TEXT("  GpuFamily: %s"), *GPUFamily);
 		UE_LOG(LogAndroid, Log, TEXT("  GlVersion: %s"), *GLVersion);
+		UE_LOG(LogAndroid, Log, TEXT("  VulkanVersion: %s"), *VulkanVersion);
 		UE_LOG(LogAndroid, Log, TEXT("  AndroidVersion: %s"), *AndroidVersion);
 		UE_LOG(LogAndroid, Log, TEXT("  DeviceMake: %s"), *DeviceMake);
 		UE_LOG(LogAndroid, Log, TEXT("  DeviceModel: %s"), *DeviceModel);
@@ -87,6 +90,9 @@ FString const FAndroidDeviceProfileSelectorModule::GetRuntimeDeviceProfileName()
 				case SRC_DeviceModel:
 					SourceString = &DeviceModel;
 					break;
+				case SRC_VulkanVersion:
+					SourceString = &VulkanVersion;
+					break;
 				default:
 					continue;
 				}
diff --git a/Engine/Plugins/Runtime/AndroidDeviceProfileSelector/Source/AndroidDeviceProfileSelector/Public/AndroidDeviceProfileMatchingRules.h b/Engine/Plugins/Runtime/AndroidDeviceProfileSelector/Source/AndroidDeviceProfileSelector/Public/AndroidDeviceProfileMatchingRules.h
--- a/Engine/Plugins/Runtime/AndroidDeviceProfileSelector/Source/AndroidDeviceProfileSelector/Public/AndroidDeviceProfileMatchingRules.h
+++ b/Engine/Plugins/Runtime/AndroidDeviceProfileSelector/Source/AndroidDeviceProfileSelector/Public/AndroidDeviceProfileMatchingRules.h
@@ -13,6 +13,7 @@ enum ESourceType
 	SRC_AndroidVersion,
 	SRC_DeviceMake,
 	SRC_DeviceModel,
+	SRC_VulkanVersion,
 	SRC_MAX,
 };
 
diff --git a/Engine/Shaders/ACES.usf b/Engine/Shaders/ACES.usf
--- a/Engine/Shaders/ACES.usf
+++ b/Engine/Shaders/ACES.usf
@@ -238,7 +238,7 @@ float cubic_basis_shaper
 		int j = knot_coord;
 		float t = knot_coord - j;
       
-		float monomials[4] = { t*t*t, t*t, t, 1. };
+		float monomials[4] = { t*t*t, t*t, t, 1.0 };
 
 		// (if/else structure required for compatibility with CTL < v1.5.)
 		if ( j == 3) {
@@ -340,7 +340,7 @@ float segmented_spline_c5_fwd( float x )
 
 		float3 cf = { C.coefsLow[ j], C.coefsLow[ j + 1], C.coefsLow[ j + 2]};
     
-		float3 monomials = { t * t, t, 1 };
+		float3 monomials = { t * t, t, 1.0 };
 		logy = dot( monomials, mul( cf, M));
 	}
 	else if (( logx >= log10(C.midPoint.x) ) && ( logx < log10(C.maxPoint.x) ))
@@ -351,7 +351,7 @@ float segmented_spline_c5_fwd( float x )
 
 		float3 cf = { C.coefsHigh[ j], C.coefsHigh[ j + 1], C.coefsHigh[ j + 2]};
 
-		float3 monomials = { t * t, t, 1 };
+		float3 monomials = { t * t, t, 1.0 };
 		logy = dot( monomials, mul( cf, M));
 	}
 	else
@@ -488,7 +488,7 @@ float segmented_spline_c9_fwd( float x, const SegmentedSplineParams_c9 C )
 
 		float3 cf = { C.coefsLow[j], C.coefsLow[j + 1], C.coefsLow[j + 2] };
 
-		float3 monomials = { t * t, t, 1 };
+		float3 monomials = { t * t, t, 1.0 };
 		logy = dot(monomials, mul(cf, M));
 	}
 	else if ((logx >= log10(C.midPoint.x)) && (logx < log10(C.maxPoint.x)))
@@ -499,7 +499,7 @@ float segmented_spline_c9_fwd( float x, const SegmentedSplineParams_c9 C )
 
 		float3 cf = { C.coefsHigh[j], C.coefsHigh[j + 1], C.coefsHigh[j + 2] };
 
-		float3 monomials = { t * t, t, 1 };
+		float3 monomials = { t * t, t, 1.0 };
 		logy = dot(monomials, mul(cf, M));
 	}
 	else//if ( logIn >= log10(C.maxPoint.x) )
diff --git a/Engine/Shaders/Common.usf b/Engine/Shaders/Common.usf
--- a/Engine/Shaders/Common.usf
+++ b/Engine/Shaders/Common.usf
@@ -954,6 +954,23 @@ MaterialFloat4 Decode32BPPHDR(MaterialFloat4 Encoded, MaterialFloat3 OtherEncode
 	}
 }
 
+// Demosaic capable 32bpp hdr decode.
+float4 Decode32BPPHDR(float4 Pixel, float2 SvPosition, Texture2D SourceTexture, SamplerState SourceSampler, float2 InTexCoords[4])
+{
+	half3 PixelColorN = SourceTexture.Sample(SourceSampler, InTexCoords[0].xy).rgb;
+	half3 PixelColorE = SourceTexture.Sample(SourceSampler, InTexCoords[1].xy).rgb;
+	half3 PixelColorW = SourceTexture.Sample(SourceSampler, InTexCoords[2].xy).rgb;
+	half3 PixelColorS = SourceTexture.Sample(SourceSampler, InTexCoords[3].xy).rgb;
+	half3 PixelColorV = PixelColorN * 0.5 + PixelColorS * 0.5;
+	half3 PixelColorH = PixelColorW * 0.5 + PixelColorE * 0.5;
+	if(abs(PixelColorN.g - PixelColorS.g) < abs(PixelColorW.g - PixelColorE.g)) 
+	{
+		PixelColorH = PixelColorV;
+	}
+	Pixel.rgb = Decode32BPPHDR(Pixel, PixelColorH, SvPosition.xy).rgb;
+	Pixel.a = 0.0;
+	return Pixel;
+}
 
 /** Get render target write mask value
   * This gets a bit from a write mask texture created with FRTWriteMaskDecodeCS. Only supprted on some platforms.
diff --git a/Engine/Shaders/Decode32bppHDR.usf b/Engine/Shaders/Decode32bppHDR.usf
new file mode 100644
--- /dev/null
+++ b/Engine/Shaders/Decode32bppHDR.usf
@@ -0,0 +1,49 @@
+// Copyright 1998-2016 Epic Games, Inc. All Rights Reserved.
+
+/*=============================================================================
+	Decode32bppHDR.usf: Decode 32bpp HDR encoding
+=============================================================================*/
+
+#include "Common.usf"
+
+float2 InvTexSize;
+
+void MainVS(
+	float2 InPosition	: ATTRIBUTE0,
+	float2 InUV			: ATTRIBUTE1,
+#if DECODING_MOSAIC
+	out float2 OutTexCoords[4] : TEXCOORD4,
+#endif
+	out FScreenVertexOutput Output
+	)
+{
+	DrawRectangle( float4( InPosition, 0, 1 ), InUV, Output.Position, Output.UV);
+
+#if DECODING_MOSAIC
+	OutTexCoords[0] = Output.UV + InvTexSize.xy * float2( 0,-1);
+	OutTexCoords[1] = Output.UV + InvTexSize.xy * float2( 1, 0);
+	OutTexCoords[2] = Output.UV + InvTexSize.xy * float2(-1, 0);
+	OutTexCoords[3] = Output.UV + InvTexSize.xy * float2( 0, 1);
+#endif
+}
+
+
+Texture2D InTexture;
+SamplerState InTextureSampler;
+
+void MainPS(
+#if DECODING_MOSAIC
+	in float2 InTexCoords[4] : TEXCOORD4,
+#endif
+	FScreenVertexOutput Input,
+	out float4 OutColor : SV_Target0
+	)
+{
+    float4 EncodedColor = Texture2DSample(InTexture, InTextureSampler, Input.UV);
+	// Decode EncodedColor, OutColor will saturate on 8 bit RTs.
+#if DECODING_MOSAIC
+	OutColor = Decode32BPPHDR(EncodedColor, Input.Position, InTexture, InTextureSampler, InTexCoords);
+#else
+	OutColor = Decode32BPPHDR(EncodedColor);
+#endif
+}
diff --git a/Engine/Shaders/LocalVertexFactory.usf b/Engine/Shaders/LocalVertexFactory.usf
--- a/Engine/Shaders/LocalVertexFactory.usf
+++ b/Engine/Shaders/LocalVertexFactory.usf
@@ -182,19 +182,53 @@ struct FVertexFactoryIntermediates
 float4x4 GetInstanceTransform(FVertexFactoryInput Input)
 {
 	#if !USE_INSTANCING_EMULATED
-	return float4x4(float4(Input.InstanceTransform1.xyz, Input.InstanceOrigin.x), float4(Input.InstanceTransform2.xyz, Input.InstanceOrigin.y), float4(Input.InstanceTransform3.xyz, Input.InstanceOrigin.z), float4(0, 0, 0, 1));
+	return float4x4(
+		float4(Input.InstanceTransform1.xyz, 0.0f), 
+		float4(Input.InstanceTransform2.xyz, 0.0f), 
+		float4(Input.InstanceTransform3.xyz, 0.0f), 
+		float4(Input.InstanceOrigin.xyz, 1.0f));
 	#else
-	return float4x4(float4(CPUInstanceTransform[0].xyz, CPUInstanceOrigin.x), float4(CPUInstanceTransform[1].xyz, CPUInstanceOrigin.y), float4(CPUInstanceTransform[2].xyz, CPUInstanceOrigin.z), float4(0,0,0,1));
+	return float4x4(
+		float4(CPUInstanceTransform[0].xyz, 0.0f), 
+		float4(CPUInstanceTransform[1].xyz, 0.0f), 
+		float4(CPUInstanceTransform[2].xyz, 0.0f), 
+		float4(CPUInstanceOrigin.xyz, 1.0f));
 	#endif	// !USE_INSTANCING_EMULATED
 }
 
 float4x4 GetInstanceTransform(FPositionOnlyVertexFactoryInput Input)
 {
 	#if !USE_INSTANCING_EMULATED
-	return float4x4(float4(Input.InstanceTransform1.xyz, Input.InstanceOrigin.x), float4(Input.InstanceTransform2.xyz, Input.InstanceOrigin.y), float4(Input.InstanceTransform3.xyz, Input.InstanceOrigin.z), float4(0, 0, 0, 1));
+	return float4x4(
+		float4(Input.InstanceTransform1.xyz, 0.0f), 
+		float4(Input.InstanceTransform2.xyz, 0.0f), 
+		float4(Input.InstanceTransform3.xyz, 0.0f), 
+		float4(Input.InstanceOrigin.xyz, 1.0f));
 	#else
-	return float4x4(float4(CPUInstanceTransform[0].xyz, CPUInstanceOrigin.x), float4(CPUInstanceTransform[1].xyz, CPUInstanceOrigin.y), float4(CPUInstanceTransform[2].xyz, CPUInstanceOrigin.z), float4(0,0,0,1));
-#endif	// !USE_INSTANCING_EMULATED
+	return float4x4(
+		float4(CPUInstanceTransform[0].xyz, 0.0f), 
+		float4(CPUInstanceTransform[1].xyz, 0.0f), 
+		float4(CPUInstanceTransform[2].xyz, 0.0f), 
+		float4(CPUInstanceOrigin.xyz, 1.0f));
+	#endif	// !USE_INSTANCING_EMULATED
+}
+
+half3x3 GetInstanceToLocal3x3(FVertexFactoryInput Input)
+{
+	#if !USE_INSTANCING_EMULATED
+	return half3x3(Input.InstanceTransform1.xyz, Input.InstanceTransform2.xyz, Input.InstanceTransform3.xyz);
+	#else
+	return half3x3(half3(CPUInstanceTransform[0].xyz), half3(CPUInstanceTransform[1].xyz), half3(CPUInstanceTransform[2].xyz));
+	#endif	// !USE_INSTANCING_EMULATED
+}
+
+half3x3 GetInstanceToLocal3x3(FPositionOnlyVertexFactoryInput Input)
+{
+	#if !USE_INSTANCING_EMULATED
+	return half3x3(Input.InstanceTransform1.xyz, Input.InstanceTransform2.xyz, Input.InstanceTransform3.xyz);
+	#else
+	return half3x3(half3(CPUInstanceTransform[0].xyz), half3(CPUInstanceTransform[1].xyz), half3(CPUInstanceTransform[2].xyz));
+	#endif	// !USE_INSTANCING_EMULATED
 }
 
 float2 GetInstanceShadowMapBias(FVertexFactoryInput Input)
@@ -317,7 +351,7 @@ FMaterialVertexParameters GetMaterialVertexParameters(FVertexFactoryInput Input,
 	Result.TangentToWorld = Intermediates.TangentToWorld;
 
 #if USE_INSTANCING
-	Result.InstanceLocalToWorld = mul(transpose(GetInstanceTransform(Input)), Primitive.LocalToWorld);
+	Result.InstanceLocalToWorld = mul(GetInstanceTransform(Input), Primitive.LocalToWorld);
 	Result.InstanceLocalPosition = Input.Position.xyz;
 	Result.PerInstanceParams = Intermediates.PerInstanceParams;
 #endif	// USE_INSTANCING
@@ -454,7 +488,7 @@ float4 CalcWorldPosition(float4 Position)
 #endif	// USE_INSTANCING
 {
 #if USE_INSTANCING
-	return TransformLocalToTranslatedWorld(mul(Position, transpose(InstanceTransform)).xyz);
+	return TransformLocalToTranslatedWorld(mul(Position, InstanceTransform).xyz);
 #elif USE_SPLINEDEFORM
 /*
 	// Make transform for this point along spline
@@ -502,9 +536,12 @@ half3x3 CalcTangentToLocal(FVertexFactoryInput Input)
 half3x3 CalcTangentToWorld(FVertexFactoryInput Input, half3x3 TangentToLocal)
 {
 #if USE_INSTANCING
-	half3x3 TangentToPrimitive = mul(TangentToLocal, transpose((float3x3)GetInstanceTransform(Input)));
-	half3x3 TangentToWorldWithScaling = mul(TangentToPrimitive, (float3x3)Primitive.LocalToWorld);
-	half3x3 TangentToWorld = float3x3(normalize(TangentToWorldWithScaling[0]), normalize(TangentToWorldWithScaling[1]), normalize(TangentToWorldWithScaling[2]));
+	half3x3 InstanceToWorld = mul(GetInstanceToLocal3x3(Input), GetLocalToWorld3x3());
+	// remove scaling
+	InstanceToWorld[0] = normalize(InstanceToWorld[0]);
+	InstanceToWorld[1] = normalize(InstanceToWorld[1]);
+	InstanceToWorld[2] = normalize(InstanceToWorld[2]);
+	half3x3 TangentToWorld = mul(TangentToLocal, InstanceToWorld);
 #else
 	half3x3 TangentToWorld = CalcTangentToWorldNoScale(TangentToLocal);
 #endif	// USE_INSTANCING
@@ -659,7 +696,7 @@ float4 VertexFactoryGetPreviousWorldPosition(FVertexFactoryInput Input, FVertexF
 	PreviousLocalToWorldTranslated[3][2] += ResolvedView.PrevPreViewTranslation.z;
 
 #if USE_INSTANCING
-	float4x4 InstanceTransform = transpose(GetInstanceTransform(Input));
+	float4x4 InstanceTransform = GetInstanceTransform(Input);
 	return mul(mul(Input.Position, InstanceTransform), PreviousLocalToWorldTranslated);
 #elif GPUSKIN_PASS_THROUGH
 	uint Offset = Input.VertexId * GPUSKIN_RWBUFFER_NUM_FLOATS + GPUSkinCachePreviousFloatOffset;
diff --git a/Engine/Shaders/PostProcessTonemap.usf b/Engine/Shaders/PostProcessTonemap.usf
--- a/Engine/Shaders/PostProcessTonemap.usf
+++ b/Engine/Shaders/PostProcessTonemap.usf
@@ -722,26 +722,7 @@ void MainPS_ES2(
 
 		#if USE_32BPP_HDR
 		{
-			// TODO: Support odd frame inversion of mosaic pattern?
-			#if 0
-				half3 SceneColorE = PostprocessInput0.Sample(PostprocessInput0Sampler, InTexCoords[1].xy).rgb;
-				SceneColor.rgb = HdrDemosaic(SceneColor.rgb, SceneColorE, SvPosition.xy);
-			#endif
-			#if 1
-				// Higher quality path.
-				half3 SceneColorN = PostprocessInput0.Sample(PostprocessInput0Sampler, InTexCoords[0].xy).rgb;
-				half3 SceneColorE = PostprocessInput0.Sample(PostprocessInput0Sampler, InTexCoords[1].xy).rgb;
-				half3 SceneColorW = PostprocessInput0.Sample(PostprocessInput0Sampler, InTexCoords[2].xy).rgb;
-				half3 SceneColorS = PostprocessInput0.Sample(PostprocessInput0Sampler, InTexCoords[3].xy).rgb;
-				half3 SceneColorV = SceneColorN * 0.5 + SceneColorS * 0.5;
-				half3 SceneColorH = SceneColorW * 0.5 + SceneColorE * 0.5;
-				if(abs(SceneColorN.g - SceneColorS.g) < abs(SceneColorW.g - SceneColorE.g)) 
-				{
-					SceneColorH = SceneColorV;
-				}
-				SceneColor.rgb = Decode32BPPHDR(SceneColor, SceneColorH, SvPosition.xy).rgb;
-				SceneColor.a = 0.0;
-			#endif
+			SceneColor = Decode32BPPHDR(SceneColor, SvPosition, PostprocessInput0, PostprocessInput0Sampler, InTexCoords);
 		}
 		#endif
 
diff --git a/Engine/Source/Developer/Windows/ShaderFormatD3D/Private/D3D11ShaderCompiler.cpp b/Engine/Source/Developer/Windows/ShaderFormatD3D/Private/D3D11ShaderCompiler.cpp
--- a/Engine/Source/Developer/Windows/ShaderFormatD3D/Private/D3D11ShaderCompiler.cpp
+++ b/Engine/Source/Developer/Windows/ShaderFormatD3D/Private/D3D11ShaderCompiler.cpp
@@ -832,7 +832,7 @@ static bool CompileAndProcessD3DShader(FString& PreprocessedShaderSource, const
 
 		if (Input.Target.Platform == SP_PCD3D_ES2)
 		{
-			if (Output.NumTextureSamplers > 16)
+			if (Output.NumTextureSamplers > 8)
 			{
 				FilteredErrors.Add(FString::Printf(TEXT("Shader uses more than 8 texture samplers which is not supported by ES2!  Used: %u"), Output.NumTextureSamplers));
 				Result = E_FAIL;
diff --git a/Engine/Source/Programs/AutomationTool/Android/AndroidPlatform.Automation.cs b/Engine/Source/Programs/AutomationTool/Android/AndroidPlatform.Automation.cs
--- a/Engine/Source/Programs/AutomationTool/Android/AndroidPlatform.Automation.cs
+++ b/Engine/Source/Programs/AutomationTool/Android/AndroidPlatform.Automation.cs
@@ -233,27 +233,27 @@ public class AndroidPlatform : Platform
 				//helper delegate to prevent code duplication but allow us access to all the local variables we need
 				var CreateInstallFilesAction = new Action<bool>(bPCInstall =>
 					{
-				// Write install batch file(s).
+						// Write install batch file(s).
 						string BatchName = GetFinalBatchName(ApkName, Params, bMakeSeparateApks ? Architecture : "", bMakeSeparateApks ? GPUArchitecture : "", false, false, bPCInstall);
-                string PackageName = GetPackageInfo(ApkName, false);
-				// make a batch file that can be used to install the .apk and .obb files
+						string PackageName = GetPackageInfo(ApkName, false);
+						// make a batch file that can be used to install the .apk and .obb files
 						string[] BatchLines = GenerateInstallBatchFile(bPackageDataInsideApk, PackageName, ApkName, Params, ObbName, DeviceObbName, false, bPCInstall);
-				File.WriteAllLines(BatchName, BatchLines);
-				// make a batch file that can be used to uninstall the .apk and .obb files
+						File.WriteAllLines(BatchName, BatchLines);
+						// make a batch file that can be used to uninstall the .apk and .obb files
 						string UninstallBatchName = GetFinalBatchName(ApkName, Params, bMakeSeparateApks ? Architecture : "", bMakeSeparateApks ? GPUArchitecture : "", false, true, bPCInstall);
 						BatchLines = GenerateUninstallBatchFile(bPackageDataInsideApk, PackageName, ApkName, Params, ObbName, DeviceObbName, false, bPCInstall);
-				File.WriteAllLines(UninstallBatchName, BatchLines);
+						File.WriteAllLines(UninstallBatchName, BatchLines);
 
-				if (Utils.IsRunningOnMono)
-				{
-					CommandUtils.FixUnixFilePermissions(BatchName);
-					CommandUtils.FixUnixFilePermissions(UninstallBatchName);
-					//if(File.Exists(NoInstallBatchName)) 
-					//{
-					//    CommandUtils.FixUnixFilePermissions(NoInstallBatchName);
-					//}
-				}
-			}
+						if (Utils.IsRunningOnMono)
+						{
+							CommandUtils.FixUnixFilePermissions(BatchName);
+							CommandUtils.FixUnixFilePermissions(UninstallBatchName);
+							//if(File.Exists(NoInstallBatchName)) 
+							//{
+							//    CommandUtils.FixUnixFilePermissions(NoInstallBatchName);
+							//}
+						}
+					}
 				);
 
 				if (bNeedsPCInstall)
@@ -473,9 +473,9 @@ public class AndroidPlatform : Platform
 						string BatchName = GetFinalBatchName(ApkName, Params, bMakeSeparateApks ? Architecture : "", bMakeSeparateApks ? GPUArchitecture : "", false, false, bPCInstall);
 						string UninstallBatchName = GetFinalBatchName(ApkName, Params, bMakeSeparateApks ? Architecture : "", bMakeSeparateApks ? GPUArchitecture : "", false, true, bPCInstall);
 
-				SC.ArchiveFiles(Path.GetDirectoryName(BatchName), Path.GetFileName(BatchName));
-                SC.ArchiveFiles(Path.GetDirectoryName(UninstallBatchName), Path.GetFileName(UninstallBatchName));
-				//SC.ArchiveFiles(Path.GetDirectoryName(NoOBBBatchName), Path.GetFileName(NoOBBBatchName));
+						SC.ArchiveFiles(Path.GetDirectoryName(BatchName), Path.GetFileName(BatchName));
+						SC.ArchiveFiles(Path.GetDirectoryName(UninstallBatchName), Path.GetFileName(UninstallBatchName));
+						//SC.ArchiveFiles(Path.GetDirectoryName(NoOBBBatchName), Path.GetFileName(NoOBBBatchName));
 					}
 				);
 
@@ -507,11 +507,11 @@ public class AndroidPlatform : Platform
 			bNeedsMonoInstall = Utils.IsRunningOnMono;
 			bNeedsPCInstall = !bNeedsMonoInstall;
 		}
-		}
+	}
 
 	private string GetAdbCommandLine(ProjectParams Params, string SerialNumber, string Args)
 	{
-		if (SerialNumber != "")
+	    if (SerialNumber != "")
 		{
 			SerialNumber = "-s " + SerialNumber;
 		}
@@ -630,8 +630,8 @@ public class AndroidPlatform : Platform
 		UFSManifests = null;
 		NonUFSManifests = null;
 
-		// Query the storage path from the device
-		string DeviceStorageQueryCommand = GetStorageQueryCommand();
+        // Query the storage path from the device
+        string DeviceStorageQueryCommand = GetStorageQueryCommand();
 		ProcessResult StorageResult = RunAdbCommand(Params, DeviceName, DeviceStorageQueryCommand, null, ERunOptions.AppMustExist);
 		String StorageLocation = StorageResult.Output.Trim();
 		string RemoteDir = StorageLocation + "/UE4Game/" + Params.ShortProjectName;
@@ -676,307 +676,307 @@ public class AndroidPlatform : Platform
 		}
 	}
 
-	public override void Deploy(ProjectParams Params, DeploymentContext SC)
-	{
+    public override void Deploy(ProjectParams Params, DeploymentContext SC)
+    {
         foreach (var DeviceName in Params.DeviceNames)
         {
             string DeviceArchitecture = GetBestDeviceArchitecture(Params, DeviceName);
             string GPUArchitecture = GetBestGPUArchitecture(Params, DeviceName);
 
-		string ApkName = GetFinalApkName(Params, SC.StageExecutables[0], true, DeviceArchitecture, GPUArchitecture);
+            string ApkName = GetFinalApkName(Params, SC.StageExecutables[0], true, DeviceArchitecture, GPUArchitecture);
 
-		// make sure APK is up to date (this is fast if so)
-		var Deploy = new UEDeployAndroid(Params.RawProjectPath);
-		if (!Params.Prebuilt)
-		{
-			string CookFlavor = SC.FinalCookPlatform.IndexOf("_") > 0 ? SC.FinalCookPlatform.Substring(SC.FinalCookPlatform.IndexOf("_")) : "";
-			string SOName = GetSONameWithoutArchitecture(Params, SC.StageExecutables[0]);
-			List<string> Architectures = new List<string>();
-			Architectures.Add(DeviceArchitecture);
-			Deploy.SetAndroidPluginData(Architectures, CollectPluginDataPaths(SC));
-			Deploy.PrepForUATPackageOrDeploy(Params.RawProjectPath, Params.ShortProjectName, SC.ProjectRoot, SOName, SC.LocalRoot + "/Engine", Params.Distribution, CookFlavor, true);
-		}
+            // make sure APK is up to date (this is fast if so)
+            var Deploy = new UEDeployAndroid(Params.RawProjectPath);
+            if (!Params.Prebuilt)
+            {
+                string CookFlavor = SC.FinalCookPlatform.IndexOf("_") > 0 ? SC.FinalCookPlatform.Substring(SC.FinalCookPlatform.IndexOf("_")) : "";
+                string SOName = GetSONameWithoutArchitecture(Params, SC.StageExecutables[0]);
+                List<string> Architectures = new List<string>();
+                Architectures.Add(DeviceArchitecture);
+                Deploy.SetAndroidPluginData(Architectures, CollectPluginDataPaths(SC));
+                Deploy.PrepForUATPackageOrDeploy(Params.RawProjectPath, Params.ShortProjectName, SC.ProjectRoot, SOName, SC.LocalRoot + "/Engine", Params.Distribution, CookFlavor, true);
+            }
 
-		// now we can use the apk to get more info
-		string PackageName = GetPackageInfo(ApkName, false);
+            // now we can use the apk to get more info
+            string PackageName = GetPackageInfo(ApkName, false);
 
-		// Setup the OBB name and add the storage path (queried from the device) to it
-		string DeviceStorageQueryCommand = GetStorageQueryCommand();
+            // Setup the OBB name and add the storage path (queried from the device) to it
+            string DeviceStorageQueryCommand = GetStorageQueryCommand();
             ProcessResult Result = RunAdbCommand(Params, DeviceName, DeviceStorageQueryCommand, null, ERunOptions.AppMustExist);
-        String StorageLocation = Result.Output.Trim(); // "/mnt/sdcard";
-		string DeviceObbName = StorageLocation + "/" + GetDeviceObbName(ApkName);
-		string RemoteDir = StorageLocation + "/UE4Game/" + Params.ShortProjectName;
-
-		// determine if APK out of date
-		string APKLastUpdateTime = new FileInfo(ApkName).LastWriteTime.ToString();
-		bool bNeedAPKInstall = true;
-		if (Params.IterativeDeploy)
-		{
-			// Check for apk installed with this package name on the device
+            String StorageLocation = Result.Output.Trim(); // "/mnt/sdcard";
+            string DeviceObbName = StorageLocation + "/" + GetDeviceObbName(ApkName);
+            string RemoteDir = StorageLocation + "/UE4Game/" + Params.ShortProjectName;
+
+            // determine if APK out of date
+            string APKLastUpdateTime = new FileInfo(ApkName).LastWriteTime.ToString();
+            bool bNeedAPKInstall = true;
+            if (Params.IterativeDeploy)
+            {
+                // Check for apk installed with this package name on the device
                 ProcessResult InstalledResult = RunAdbCommand(Params, DeviceName, "shell pm list packages " + PackageName, null, ERunOptions.AppMustExist);
-			if (InstalledResult.Output.Contains(PackageName))
-			{
-				// See if apk is up to date on device
+                if (InstalledResult.Output.Contains(PackageName))
+                {
+                    // See if apk is up to date on device
                     InstalledResult = RunAdbCommand(Params, DeviceName, "shell cat " + RemoteDir + "/APKFileStamp.txt", null, ERunOptions.AppMustExist);
-				if (InstalledResult.Output.StartsWith("APK: "))
-				{
-					if (InstalledResult.Output.Substring(5).Trim() == APKLastUpdateTime)
-						bNeedAPKInstall = false;
+                    if (InstalledResult.Output.StartsWith("APK: "))
+                    {
+                        if (InstalledResult.Output.Substring(5).Trim() == APKLastUpdateTime)
+                            bNeedAPKInstall = false;
 
-					// Stop the previously running copy (uninstall/install did this before)
+                        // Stop the previously running copy (uninstall/install did this before)
                         InstalledResult = RunAdbCommand(Params, DeviceName, "shell am force-stop " + PackageName, null, ERunOptions.AppMustExist);
-					if (InstalledResult.Output.Contains("Error"))
-					{
-						// force-stop not supported (Android < 3.0) so check if package is actually running
-						// Note: cannot use grep here since it may not be installed on device
+                        if (InstalledResult.Output.Contains("Error"))
+                        {
+                            // force-stop not supported (Android < 3.0) so check if package is actually running
+                            // Note: cannot use grep here since it may not be installed on device
                             InstalledResult = RunAdbCommand(Params, DeviceName, "shell ps", null, ERunOptions.AppMustExist);
-						if (InstalledResult.Output.Contains(PackageName))
-						{
-							// it is actually running so use the slow way to kill it (uninstall and reinstall)
-							bNeedAPKInstall = true;
-						}
-
-					}
-				}
-			}
-		}
+                            if (InstalledResult.Output.Contains(PackageName))
+                            {
+                                // it is actually running so use the slow way to kill it (uninstall and reinstall)
+                                bNeedAPKInstall = true;
+                            }
+
+                        }
+                    }
+                }
+            }
 
-		// install new APK if needed
-		if (bNeedAPKInstall)
-		{
-			// try uninstalling an old app with the same identifier.
-			int SuccessCode = 0;
-			string UninstallCommandline = "uninstall " + PackageName;
+            // install new APK if needed
+            if (bNeedAPKInstall)
+            {
+                // try uninstalling an old app with the same identifier.
+                int SuccessCode = 0;
+                string UninstallCommandline = "uninstall " + PackageName;
                 RunAndLogAdbCommand(Params, DeviceName, UninstallCommandline, out SuccessCode);
 
-			// install the apk
-			string InstallCommandline = "install \"" + ApkName + "\"";
+                // install the apk
+                string InstallCommandline = "install \"" + ApkName + "\"";
                 string InstallOutput = RunAndLogAdbCommand(Params, DeviceName, InstallCommandline, out SuccessCode);
-			int FailureIndex = InstallOutput.IndexOf("Failure"); 
+                int FailureIndex = InstallOutput.IndexOf("Failure");
+
+                // adb install doesn't always return an error code on failure, and instead prints "Failure", followed by an error code.
+                if (SuccessCode != 0 || FailureIndex != -1)
+                {
+                    string ErrorMessage = string.Format("Installation of apk '{0}' failed", ApkName);
+                    if (FailureIndex != -1)
+                    {
+                        string FailureString = InstallOutput.Substring(FailureIndex + 7).Trim();
+                        if (FailureString != "")
+                        {
+                            ErrorMessage += ": " + FailureString;
+                        }
+                    }
+                    if (ErrorMessage.Contains("OLDER_SDK"))
+                    {
+                        LogError("minSdkVersion is higher than Android version installed on device, possibly due to NDK API Level");
+                    }
+                    throw new AutomationException(ExitCode.Error_AppInstallFailed, ErrorMessage);
+                }
+            }
 
-			// adb install doesn't always return an error code on failure, and instead prints "Failure", followed by an error code.
-			if (SuccessCode != 0 || FailureIndex != -1)
-			{
-				string ErrorMessage = string.Format("Installation of apk '{0}' failed", ApkName);
-				if (FailureIndex != -1)
-				{
-					string FailureString = InstallOutput.Substring(FailureIndex + 7).Trim();
-					if (FailureString != "")
-					{
-						ErrorMessage += ": " + FailureString;
-					}
-				}
-				if (ErrorMessage.Contains("OLDER_SDK"))
-				{
-					LogError("minSdkVersion is higher than Android version installed on device, possibly due to NDK API Level");
-				}
-				throw new AutomationException(ExitCode.Error_AppInstallFailed, ErrorMessage);
-			}
-		}
- 
-		// update the ue4commandline.txt
-		// update and deploy ue4commandline.txt
-		// always delete the existing commandline text file, so it doesn't reuse an old one
-		string IntermediateCmdLineFile = CombinePaths(SC.StageDirectory, "UE4CommandLine.txt");
-		Project.WriteStageCommandline(IntermediateCmdLineFile, Params, SC);
-
-		// copy files to device if we were staging
-		if (SC.Stage)
-		{
-			// cache some strings
-			string BaseCommandline = "push";
+            // update the ue4commandline.txt
+            // update and deploy ue4commandline.txt
+            // always delete the existing commandline text file, so it doesn't reuse an old one
+            string IntermediateCmdLineFile = CombinePaths(SC.StageDirectory, "UE4CommandLine.txt");
+            Project.WriteStageCommandline(IntermediateCmdLineFile, Params, SC);
 
-            HashSet<string> EntriesToDeploy = new HashSet<string>();
+            // copy files to device if we were staging
+            if (SC.Stage)
+            {
+                // cache some strings
+                string BaseCommandline = "push";
 
-			if (Params.IterativeDeploy)
-			{
-				// always send UE4CommandLine.txt (it was written above after delta checks applied)
-				EntriesToDeploy.Add(IntermediateCmdLineFile);
+                HashSet<string> EntriesToDeploy = new HashSet<string>();
 
-				// Add non UFS files if any to deploy
-                    String NonUFSManifestPath = SC.GetNonUFSDeploymentDeltaPath(DeviceName);
-				if (File.Exists(NonUFSManifestPath))
-				{
-					string NonUFSFiles = File.ReadAllText(NonUFSManifestPath);
-					foreach (string Filename in NonUFSFiles.Split('\n'))
-					{
-						if (!string.IsNullOrEmpty(Filename) && !string.IsNullOrWhiteSpace(Filename))
-						{
-							EntriesToDeploy.Add(CombinePaths(SC.StageDirectory, Filename.Trim()));
-						}
-					}
-				}
+                if (Params.IterativeDeploy)
+                {
+                    // always send UE4CommandLine.txt (it was written above after delta checks applied)
+                    EntriesToDeploy.Add(IntermediateCmdLineFile);
 
-				// Add UFS files if any to deploy
+                    // Add non UFS files if any to deploy
+                    String NonUFSManifestPath = SC.GetNonUFSDeploymentDeltaPath(DeviceName);
+                    if (File.Exists(NonUFSManifestPath))
+                    {
+                        string NonUFSFiles = File.ReadAllText(NonUFSManifestPath);
+                        foreach (string Filename in NonUFSFiles.Split('\n'))
+                        {
+                            if (!string.IsNullOrEmpty(Filename) && !string.IsNullOrWhiteSpace(Filename))
+                            {
+                                EntriesToDeploy.Add(CombinePaths(SC.StageDirectory, Filename.Trim()));
+                            }
+                        }
+                    }
+
+                    // Add UFS files if any to deploy
                     String UFSManifestPath = SC.GetUFSDeploymentDeltaPath(DeviceName);
-				if (File.Exists(UFSManifestPath))
-				{
-					string UFSFiles = File.ReadAllText(UFSManifestPath);
-					foreach (string Filename in UFSFiles.Split('\n'))
-					{
-						if (!string.IsNullOrEmpty(Filename) && !string.IsNullOrWhiteSpace(Filename))
-						{
-							EntriesToDeploy.Add(CombinePaths(SC.StageDirectory, Filename.Trim()));
-						}
-					}
-				}
-
-				// For now, if too many files may be better to just push them all
-				if (EntriesToDeploy.Count > 500)
-				{
-					// make sure device is at a clean state
+                    if (File.Exists(UFSManifestPath))
+                    {
+                        string UFSFiles = File.ReadAllText(UFSManifestPath);
+                        foreach (string Filename in UFSFiles.Split('\n'))
+                        {
+                            if (!string.IsNullOrEmpty(Filename) && !string.IsNullOrWhiteSpace(Filename))
+                            {
+                                EntriesToDeploy.Add(CombinePaths(SC.StageDirectory, Filename.Trim()));
+                            }
+                        }
+                    }
+
+                    // For now, if too many files may be better to just push them all
+                    if (EntriesToDeploy.Count > 500)
+                    {
+                        // make sure device is at a clean state
                         RunAdbCommand(Params, DeviceName, "shell rm -r " + RemoteDir);
 
-					EntriesToDeploy.Clear();
-					EntriesToDeploy.TrimExcess();
-					EntriesToDeploy.Add(SC.StageDirectory);
-				}
-			}
-			else
-			{
-				// make sure device is at a clean state
+                        EntriesToDeploy.Clear();
+                        EntriesToDeploy.TrimExcess();
+                        EntriesToDeploy.Add(SC.StageDirectory);
+                    }
+                }
+                else
+                {
+                    // make sure device is at a clean state
                     RunAdbCommand(Params, DeviceName, "shell rm -r " + RemoteDir);
 
-				// Copy UFS files..
-				string[] Files = Directory.GetFiles(SC.StageDirectory, "*", SearchOption.AllDirectories);
-				System.Array.Sort(Files);
-
-				// Find all the files we exclude from copying. And include
-				// the directories we need to individually copy.
-				HashSet<string> ExcludedFiles = new HashSet<string>();
-				SortedSet<string> IndividualCopyDirectories
-					= new SortedSet<string>((IComparer<string>)new LongestFirst());
-				foreach (string Filename in Files)
-				{
-					bool Exclude = false;
-					// Don't push the apk, we install it
-					Exclude |= Path.GetExtension(Filename).Equals(".apk", StringComparison.InvariantCultureIgnoreCase);
-					// For excluded files we add the parent dirs to our
-					// tracking of stuff to individually copy.
-					if (Exclude)
-					{
-						ExcludedFiles.Add(Filename);
-						// We include all directories up to the stage root in having
-						// to individually copy the files.
-						for (string FileDirectory = Path.GetDirectoryName(Filename);
-							!FileDirectory.Equals(SC.StageDirectory);
-							FileDirectory = Path.GetDirectoryName(FileDirectory))
-						{
-							if (!IndividualCopyDirectories.Contains(FileDirectory))
-							{
-								IndividualCopyDirectories.Add(FileDirectory);
-							}
-						}
-						if (!IndividualCopyDirectories.Contains(SC.StageDirectory))
-						{
-							IndividualCopyDirectories.Add(SC.StageDirectory);
-						}
-					}
-				}
-
-				// The directories are sorted above in "deepest" first. We can
-				// therefore start copying those individual dirs which will
-				// recreate the tree. As the subtrees will get copied at each
-				// possible individual level.
-				foreach (string DirectoryName in IndividualCopyDirectories)
-				{
-					string[] Entries
-						= Directory.GetFileSystemEntries(DirectoryName, "*", SearchOption.TopDirectoryOnly);
-					foreach (string Entry in Entries)
-					{
-						// We avoid excluded files and the individual copy dirs
-						// (the individual copy dirs will get handled as we iterate).
-						if (ExcludedFiles.Contains(Entry) || IndividualCopyDirectories.Contains(Entry))
-						{
-							continue;
-						}
-						else
-						{
-							EntriesToDeploy.Add(Entry);
-						}
-					}
-				}
-
-				if (EntriesToDeploy.Count == 0)
-				{
-					EntriesToDeploy.Add(SC.StageDirectory);
-				}
-			}
-
-			// We now have a minimal set of file & dir entries we need
-			// to deploy. Files we deploy will get individually copied
-			// and dirs will get the tree copies by default (that's
-			// what ADB does).
-			HashSet<ProcessResult> DeployCommands = new HashSet<ProcessResult>();
-			foreach (string Entry in EntriesToDeploy)
-			{
-				string FinalRemoteDir = RemoteDir;
-				string RemotePath = Entry.Replace(SC.StageDirectory, FinalRemoteDir).Replace("\\", "/");
-				string Commandline = string.Format("{0} \"{1}\" \"{2}\"", BaseCommandline, Entry, RemotePath);
-				// We run deploy commands in parallel to maximize the connection
-				// throughput.
-				DeployCommands.Add(
+                    // Copy UFS files..
+                    string[] Files = Directory.GetFiles(SC.StageDirectory, "*", SearchOption.AllDirectories);
+                    System.Array.Sort(Files);
+
+                    // Find all the files we exclude from copying. And include
+                    // the directories we need to individually copy.
+                    HashSet<string> ExcludedFiles = new HashSet<string>();
+                    SortedSet<string> IndividualCopyDirectories
+                        = new SortedSet<string>((IComparer<string>)new LongestFirst());
+                    foreach (string Filename in Files)
+                    {
+                        bool Exclude = false;
+                        // Don't push the apk, we install it
+                        Exclude |= Path.GetExtension(Filename).Equals(".apk", StringComparison.InvariantCultureIgnoreCase);
+                        // For excluded files we add the parent dirs to our
+                        // tracking of stuff to individually copy.
+                        if (Exclude)
+                        {
+                            ExcludedFiles.Add(Filename);
+                            // We include all directories up to the stage root in having
+                            // to individually copy the files.
+                            for (string FileDirectory = Path.GetDirectoryName(Filename);
+                                !FileDirectory.Equals(SC.StageDirectory);
+                                FileDirectory = Path.GetDirectoryName(FileDirectory))
+                            {
+                                if (!IndividualCopyDirectories.Contains(FileDirectory))
+                                {
+                                    IndividualCopyDirectories.Add(FileDirectory);
+                                }
+                            }
+                            if (!IndividualCopyDirectories.Contains(SC.StageDirectory))
+                            {
+                                IndividualCopyDirectories.Add(SC.StageDirectory);
+                            }
+                        }
+                    }
+
+                    // The directories are sorted above in "deepest" first. We can
+                    // therefore start copying those individual dirs which will
+                    // recreate the tree. As the subtrees will get copied at each
+                    // possible individual level.
+                    foreach (string DirectoryName in IndividualCopyDirectories)
+                    {
+                        string[] Entries
+                            = Directory.GetFileSystemEntries(DirectoryName, "*", SearchOption.TopDirectoryOnly);
+                        foreach (string Entry in Entries)
+                        {
+                            // We avoid excluded files and the individual copy dirs
+                            // (the individual copy dirs will get handled as we iterate).
+                            if (ExcludedFiles.Contains(Entry) || IndividualCopyDirectories.Contains(Entry))
+                            {
+                                continue;
+                            }
+                            else
+                            {
+                                EntriesToDeploy.Add(Entry);
+                            }
+                        }
+                    }
+
+                    if (EntriesToDeploy.Count == 0)
+                    {
+                        EntriesToDeploy.Add(SC.StageDirectory);
+                    }
+                }
+
+                // We now have a minimal set of file & dir entries we need
+                // to deploy. Files we deploy will get individually copied
+                // and dirs will get the tree copies by default (that's
+                // what ADB does).
+                HashSet<ProcessResult> DeployCommands = new HashSet<ProcessResult>();
+                foreach (string Entry in EntriesToDeploy)
+                {
+                    string FinalRemoteDir = RemoteDir;
+                    string RemotePath = Entry.Replace(SC.StageDirectory, FinalRemoteDir).Replace("\\", "/");
+                    string Commandline = string.Format("{0} \"{1}\" \"{2}\"", BaseCommandline, Entry, RemotePath);
+                    // We run deploy commands in parallel to maximize the connection
+                    // throughput.
+                    DeployCommands.Add(
                         RunAdbCommand(Params, DeviceName, Commandline, null,
-						ERunOptions.Default | ERunOptions.NoWaitForExit));
-				// But we limit the parallel commands to avoid overwhelming
-				// memory resources.
-				if (DeployCommands.Count == DeployMaxParallelCommands)
-				{
-					while (DeployCommands.Count > DeployMaxParallelCommands / 2)
-					{
-						Thread.Sleep(1);
-						DeployCommands.RemoveWhere(
-							delegate(ProcessResult r)
-							{
-								return r.HasExited;
-							});
-					}
-				}
-			}
-			foreach (ProcessResult deploy_result in DeployCommands)
-			{
-				deploy_result.WaitForExit();
-			}
-
-			// delete the .obb file, since it will cause nothing we just deployed to be used
+                            ERunOptions.Default | ERunOptions.NoWaitForExit));
+                    // But we limit the parallel commands to avoid overwhelming
+                    // memory resources.
+                    if (DeployCommands.Count == DeployMaxParallelCommands)
+                    {
+                        while (DeployCommands.Count > DeployMaxParallelCommands / 2)
+                        {
+                            Thread.Sleep(1);
+                            DeployCommands.RemoveWhere(
+                                delegate (ProcessResult r)
+                                {
+                                    return r.HasExited;
+                                });
+                        }
+                    }
+                }
+                foreach (ProcessResult deploy_result in DeployCommands)
+                {
+                    deploy_result.WaitForExit();
+                }
+
+                // delete the .obb file, since it will cause nothing we just deployed to be used
                 RunAdbCommand(Params, DeviceName, "shell rm " + DeviceObbName);
-		}
-		else if (SC.Archive)
-		{
-			// deploy the obb if there is one
-			string ObbPath = Path.Combine(SC.StageDirectory, GetFinalObbName(ApkName));
-			if (File.Exists(ObbPath))
-			{
-				// cache some strings
-				string BaseCommandline = "push";
-
-				string Commandline = string.Format("{0} \"{1}\" \"{2}\"", BaseCommandline, ObbPath, DeviceObbName);
+            }
+            else if (SC.Archive)
+            {
+                // deploy the obb if there is one
+                string ObbPath = Path.Combine(SC.StageDirectory, GetFinalObbName(ApkName));
+                if (File.Exists(ObbPath))
+                {
+                    // cache some strings
+                    string BaseCommandline = "push";
+
+                    string Commandline = string.Format("{0} \"{1}\" \"{2}\"", BaseCommandline, ObbPath, DeviceObbName);
                     RunAdbCommand(Params, DeviceName, Commandline);
-			}
-		}
-		else
-		{
-			// cache some strings
-			string BaseCommandline = "push";
+                }
+            }
+            else
+            {
+                // cache some strings
+                string BaseCommandline = "push";
 
-			string FinalRemoteDir = RemoteDir;
-			/*
-			// handle the special case of the UE4Commandline.txt when using content only game (UE4Game)
-			if (!Params.IsCodeBasedProject)
-			{
-				FinalRemoteDir = "/mnt/sdcard/UE4Game";
-			}
-			*/
+                string FinalRemoteDir = RemoteDir;
+                /*
+			    // handle the special case of the UE4Commandline.txt when using content only game (UE4Game)
+			    if (!Params.IsCodeBasedProject)
+			    {
+				    FinalRemoteDir = "/mnt/sdcard/UE4Game";
+			    }
+			    */
 
-			string RemoteFilename = IntermediateCmdLineFile.Replace(SC.StageDirectory, FinalRemoteDir).Replace("\\", "/");
-			string Commandline = string.Format("{0} \"{1}\" \"{2}\"", BaseCommandline, IntermediateCmdLineFile, RemoteFilename);
+                string RemoteFilename = IntermediateCmdLineFile.Replace(SC.StageDirectory, FinalRemoteDir).Replace("\\", "/");
+                string Commandline = string.Format("{0} \"{1}\" \"{2}\"", BaseCommandline, IntermediateCmdLineFile, RemoteFilename);
                 RunAdbCommand(Params, DeviceName, Commandline);
-		}
-    
-        // write new timestamp for APK (do it here since RemoteDir will now exist)
-        if (bNeedAPKInstall)
-        {
-            int SuccessCode = 0;
+            }
+
+            // write new timestamp for APK (do it here since RemoteDir will now exist)
+            if (bNeedAPKInstall)
+            {
+                int SuccessCode = 0;
                 RunAndLogAdbCommand(Params, DeviceName, "shell \"echo 'APK: " + APKLastUpdateTime + "' > " + RemoteDir + "/APKFileStamp.txt\"", out SuccessCode);
             }
         }
@@ -1244,85 +1244,103 @@ public class AndroidPlatform : Platform
 
 	public override ProcessResult RunClient(ERunOptions ClientRunFlags, string ClientApp, string ClientCmdLine, ProjectParams Params)
 	{
-        ProcessResult OutClientProcess = null;
+		//make a copy of the device names, we'll be working through them
+		List<string> DeviceNames = new List<string>();
+		//same with the package names
+		List<string> PackageNames = new List<string>();
 
-        foreach (var DeviceName in Params.DeviceNames)
-        {
-            string DeviceArchitecture = GetBestDeviceArchitecture(Params, DeviceName);
-            string GPUArchitecture = GetBestGPUArchitecture(Params, DeviceName); ;
-
-		string ApkName = ClientApp + DeviceArchitecture + ".apk";
-		if (!File.Exists(ApkName))
+		foreach (string DeviceName in Params.DeviceNames)
 		{
-			ApkName = GetFinalApkName(Params, Path.GetFileNameWithoutExtension(ClientApp), true, DeviceArchitecture, GPUArchitecture);
-		}
-
-		Console.WriteLine("Apk='{0}', ClientApp='{1}', ExeName='{2}'", ApkName, ClientApp, Params.ProjectGameExeFilename);
+			//save the device name
+			DeviceNames.Add(DeviceName);
+
+			//get the package name and save that
+			string DeviceArchitecture = GetBestDeviceArchitecture(Params, DeviceName);
+			string GPUArchitecture = GetBestGPUArchitecture(Params, DeviceName);
+			string ApkName = ClientApp + DeviceArchitecture + ".apk";
+			if (!File.Exists(ApkName))
+			{
+				ApkName = GetFinalApkName(Params, Path.GetFileNameWithoutExtension(ClientApp), true, DeviceArchitecture, GPUArchitecture);
+			}
+			Console.WriteLine("Apk='{0}', ClientApp='{1}', ExeName='{2}'", ApkName, ClientApp, Params.ProjectGameExeFilename);
 
-		// run aapt to get the name of the intent
-		string PackageName = GetPackageInfo(ApkName, false);
-		if (PackageName == null)
-		{
-            throw new AutomationException(ExitCode.Error_FailureGettingPackageInfo, "Failed to get package name from " + ClientApp);
-		}
+			// run aapt to get the name of the intent
+			string PackageName = GetPackageInfo(ApkName, false);
+			if (PackageName == null)
+			{
+				throw new AutomationException(ExitCode.Error_FailureGettingPackageInfo, "Failed to get package name from " + ClientApp);
+			}
 
-		if (Params.Prebuilt)
-		{
-			// clear the log
-                RunAdbCommand(Params, DeviceName, "logcat -c");
-		}
+			PackageNames.Add(PackageName);
 
-		// start the app on device!
-		string CommandLine = "shell am start -n " + PackageName + "/" + GetLaunchableActivityName();
-            ProcessResult ClientProcess = RunAdbCommand(Params, DeviceName, CommandLine, null, ClientRunFlags);
+			// clear the log for the device
+			RunAdbCommand(Params, DeviceName, "logcat -c");
 
-            OutClientProcess = ClientProcess;
+			// start the app on device!
+			string CommandLine = "shell am start -n " + PackageName + "/" + GetLaunchableActivityName();
+			RunAdbCommand(Params, DeviceName, CommandLine, null, ClientRunFlags);
 
-		if (Params.Prebuilt)
-		{
 			// save the output to the staging directory
 			string LogPath = Path.Combine(Params.BaseStageDirectory, "Android\\logs");
-                string LogFilename = Path.Combine(LogPath, "devicelog" + DeviceName + ".log");
-                string ServerLogFilename = Path.Combine(CmdEnv.LogFolder, "devicelog" + DeviceName + ".log");
-
 			Directory.CreateDirectory(LogPath);
+		}
 
-			// check if the game is still running 
-			// time out if it takes to long
-			DateTime StartTime = DateTime.Now;
-			int TimeOutSeconds = Params.RunTimeoutSeconds;
+		//now check if each device still has the game running, and time out if it's taking too long
+		DateTime StartTime = DateTime.Now;
+		int TimeOutSeconds = Params.RunTimeoutSeconds;
 
-			while (true)
+		while (DeviceNames.Count > 0)
+		{
+			for(int DeviceIndex = 0; DeviceIndex < DeviceNames.Count; DeviceIndex++)
 			{
-                    ProcessResult ProcessesResult = RunAdbCommand(Params, DeviceName, "shell ps", null, ERunOptions.SpewIsVerbose);
+				string DeviceName = DeviceNames[DeviceIndex];
+
+				bool FinishedRunning = false;
+				ProcessResult ProcessesResult = RunAdbCommand(Params, DeviceName, "shell ps", null, ERunOptions.SpewIsVerbose);
 
 				string RunningProcessList = ProcessesResult.Output;
-				if (!RunningProcessList.Contains(PackageName))
+				if (!RunningProcessList.Contains(PackageNames[DeviceIndex]))
 				{
-					break;
+					FinishedRunning = true;
 				}
+
 				Thread.Sleep(10);
 
-				TimeSpan DeltaRunTime = DateTime.Now - StartTime;
-				if ((DeltaRunTime.TotalSeconds > TimeOutSeconds) && (TimeOutSeconds != 0))
+				if(!FinishedRunning)
 				{
-                        Log("Device: " + DeviceName + " timed out while waiting for run to finish");
-					break;
+					TimeSpan DeltaRunTime = DateTime.Now - StartTime;
+					if ((DeltaRunTime.TotalSeconds > TimeOutSeconds) && (TimeOutSeconds != 0))
+					{
+						Log("Device: " + DeviceName + " timed out while waiting for run to finish");
+						FinishedRunning = true;
+					}
 				}
-			}
 
-			// this is just to get the ue4 log to go to the output
-                RunAdbCommand(Params, DeviceName, "logcat -d -s UE4 -s Debug");
+				//log the results, then clear out the device from our list
+				if(FinishedRunning)
+				{
+					// this is just to get the ue4 log to go to the output
+					RunAdbCommand(Params, DeviceName, "logcat -d -s UE4 -s Debug");
+
+					// get the log we actually want to save
+					ProcessResult LogFileProcess = RunAdbCommand(Params, DeviceName, "logcat -d", null, ERunOptions.AppMustExist);
+
+					string LogPath = Path.Combine(Params.BaseStageDirectory, "Android\\logs");
+					string LogFilename = Path.Combine(LogPath, "devicelog" + DeviceName + ".log");
+					string ServerLogFilename = Path.Combine(CmdEnv.LogFolder, "devicelog" + DeviceName + ".log");
+
+					File.WriteAllText(LogFilename, LogFileProcess.Output);
+					File.WriteAllText(ServerLogFilename, LogFileProcess.Output);
 
-			// get the log we actually want to save
-                ProcessResult LogFileProcess = RunAdbCommand(Params, DeviceName, "logcat -d", null, ERunOptions.AppMustExist);
+					DeviceNames.RemoveAt(DeviceIndex);
+					PackageNames.RemoveAt(DeviceIndex);
 
-			File.WriteAllText(LogFilename, LogFileProcess.Output);
-			File.WriteAllText(ServerLogFilename, LogFileProcess.Output);
+					--DeviceIndex;
+				}
+			}
 		}
-        }
 
-		return OutClientProcess;
+		return null;
 	}
 
 	public override void GetFilesToDeployOrStage(ProjectParams Params, DeploymentContext SC)
@@ -1346,9 +1364,9 @@ public class AndroidPlatform : Platform
 	/// </summary>
 	/// <returns>Cook platform string.</returns>
 	public override string GetCookPlatform(bool bDedicatedServer, bool bIsClientOnly)
-		{
-			return "Android";
-		}
+	{
+		return "Android";
+	}
 
 	public override bool DeployPakInternalLowerCaseFilenames()
 	{
@@ -1385,16 +1403,16 @@ public class AndroidPlatform : Platform
         }
     }
 
-/*
-	public override bool RequiresPackageToDeploy
-	{
-		get { return true; }
-	}
-*/
-    
-	#region Hooks
+    /*
+        public override bool RequiresPackageToDeploy
+        {
+            get { return true; }
+        }
+    */
+
+    #region Hooks
 
-	public override void PostBuildTarget(UE4Build Build, FileReference UProjectPath, string TargetName, string Config)
+    public override void PostBuildTarget(UE4Build Build, FileReference UProjectPath, string TargetName, string Config)
 	{
 		// Run UBT w/ the prep for deployment only option
 		// This is required as UBT will 'fake' success when building via UAT and run
diff --git a/Engine/Source/Programs/UnrealBuildTool/Android/UEDeployAndroid.cs b/Engine/Source/Programs/UnrealBuildTool/Android/UEDeployAndroid.cs
--- a/Engine/Source/Programs/UnrealBuildTool/Android/UEDeployAndroid.cs
+++ b/Engine/Source/Programs/UnrealBuildTool/Android/UEDeployAndroid.cs
@@ -1201,7 +1201,7 @@ namespace UnrealBuildTool
 
 
 
-		private string GenerateManifest(AndroidToolChain ToolChain, string ProjectName, bool bIsForDistribution, bool bPackageDataInsideApk, string GameBuildFilesPath, bool bHasOBBFiles, bool bDisableVerifyOBBOnStartUp, string UE4Arch, string GPUArch)
+		private string GenerateManifest(AndroidToolChain ToolChain, string ProjectName, bool bIsForDistribution, bool bPackageDataInsideApk, string GameBuildFilesPath, bool bHasOBBFiles, bool bDisableVerifyOBBOnStartUp, string UE4Arch, string GPUArch, string CookFlavor)
 		{
 			string Arch = GetNDKArch(UE4Arch);
 			int NDKLevelInt = ToolChain.GetNdkApiLevelInt();
@@ -1304,6 +1304,51 @@ namespace UnrealBuildTool
 				}
 			}
 
+			//figure out which texture compressions are supported
+			bool bETC1Enabled, bETC2Enabled, bDXTEnabled, bATCEnabled, bPVRTCEnabled, bASTCEnabled;
+			bETC1Enabled = bETC2Enabled = bDXTEnabled = bATCEnabled = bPVRTCEnabled = bASTCEnabled = false;
+			if (CookFlavor.Length < 1)
+			{
+				//All values supproted
+				bETC1Enabled = bETC2Enabled = bDXTEnabled = bATCEnabled = bPVRTCEnabled = bASTCEnabled = true;
+			}
+			else
+			{
+				switch(CookFlavor)
+				{
+					case "_Multi":
+						//need to check ini to determine which are supported
+						Ini.GetBool("/Script/AndroidRuntimeSettings.AndroidRuntimeSettings", "bMultiTargetFormat_ETC1", out bETC1Enabled);
+						Ini.GetBool("/Script/AndroidRuntimeSettings.AndroidRuntimeSettings", "bMultiTargetFormat_ETC2", out bETC2Enabled);
+						Ini.GetBool("/Script/AndroidRuntimeSettings.AndroidRuntimeSettings", "bMultiTargetFormat_DXT", out bDXTEnabled);
+						Ini.GetBool("/Script/AndroidRuntimeSettings.AndroidRuntimeSettings", "bMultiTargetFormat_ATC", out bATCEnabled);
+						Ini.GetBool("/Script/AndroidRuntimeSettings.AndroidRuntimeSettings", "bMultiTargetFormat_PVRTC", out bPVRTCEnabled);
+						Ini.GetBool("/Script/AndroidRuntimeSettings.AndroidRuntimeSettings", "bMultiTargetFormat_ASTC", out bASTCEnabled);
+						break;
+					case "_ETC1":
+						bETC1Enabled = true;
+						break;
+					case "_ETC2":
+						bETC2Enabled = true;
+						break;
+					case "_DXT":
+						bDXTEnabled = true;
+						break;
+					case "_ATC":
+						bATCEnabled = true;
+						break;
+					case "_PVRTC":
+						bPVRTCEnabled = true;
+						break;
+					case "_ASTC":
+						bASTCEnabled = true;
+						break;
+					default:
+						Log.TraceWarning("Invalid or unknown CookFlavor used in GenerateManifest: {0}", CookFlavor);
+						break;
+				}
+			}
+
 			StringBuilder Text = new StringBuilder();
 			Text.AppendLine("<?xml version=\"1.0\" encoding=\"utf-8\"?>");
 			Text.AppendLine("<manifest xmlns:android=\"http://schemas.android.com/apk/res/android\"");
@@ -1491,6 +1536,36 @@ namespace UnrealBuildTool
 						Text.AppendLine("\t" + Line);
 					}
 				}
+
+				Text.AppendLine("\t<!-- Supported texture compression formats (cooked) -->");
+				if (bETC1Enabled)
+				{
+					Text.AppendLine("\t<supports-gl-texture android:name=\"GL_OES_compressed_ETC1_RGB8_texture\" />");
+				}
+				if (bETC2Enabled)
+				{
+					Text.AppendLine("\t<supports-gl-texture android:name=\"GL_COMPRESSED_RGB8_ETC2\" />");
+					Text.AppendLine("\t<supports-gl-texture android:name=\"GL_COMPRESSED_RGBA8_ETC2_EAC\" />");
+				}
+				if (bATCEnabled)
+				{
+					Text.AppendLine("\t<supports-gl-texture android:name=\"GL_AMD_compressed_ATC_texture\" />");
+					Text.AppendLine("\t<supports-gl-texture android:name=\"GL_ATI_texture_compression_atitc\" />");
+				}
+				if (bDXTEnabled)
+				{
+					Text.AppendLine("\t<supports-gl-texture android:name=\"GL_EXT_texture_compression_dxt1\" />");
+					Text.AppendLine("\t<supports-gl-texture android:name=\"GL_EXT_texture_compression_s3tc\" />");
+					Text.AppendLine("\t<supports-gl-texture android:name=\"GL_NV_texture_compression_s3tc\" />");
+				}
+				if (bPVRTCEnabled)
+				{
+					Text.AppendLine("\t<supports-gl-texture android:name=\"GL_IMG_texture_compression_pvrtc\" />");
+				}
+				if (bASTCEnabled)
+				{
+					Text.AppendLine("\t<supports-gl-texture android:name=\"GL_KHR_texture_compression_astc_ldr\" />");
+				}
 			}
 
 			Text.AppendLine("</manifest>");
@@ -1864,7 +1939,7 @@ namespace UnrealBuildTool
 			{
 				BuildList = from Arch in Arches
 							from GPUArch in GPUArchitectures
-							let manifest = GenerateManifest(ToolChain, ProjectName, bForDistribution, bPackageDataInsideApk, GameBuildFilesPath, RequiresOBB(bDisallowPackagingDataInApk, ObbFileLocation), bDisableVerifyOBBOnStartUp, Arch, GPUArch)
+							let manifest = GenerateManifest(ToolChain, ProjectName, bForDistribution, bPackageDataInsideApk, GameBuildFilesPath, RequiresOBB(bDisallowPackagingDataInApk, ObbFileLocation), bDisableVerifyOBBOnStartUp, Arch, GPUArch, CookFlavor)
 							select Tuple.Create(Arch, GPUArch, manifest);
 			}
 			else
@@ -1872,7 +1947,7 @@ namespace UnrealBuildTool
 				BuildList = from Arch in Arches
 							from GPUArch in GPUArchitectures
 							let manifestFile = Path.Combine(IntermediateAndroidPath, Arch + "_" + GPUArch + "_AndroidManifest.xml")
-							let manifest = GenerateManifest(ToolChain, ProjectName, bForDistribution, bPackageDataInsideApk, GameBuildFilesPath, RequiresOBB(bDisallowPackagingDataInApk, ObbFileLocation), bDisableVerifyOBBOnStartUp, Arch, GPUArch)
+							let manifest = GenerateManifest(ToolChain, ProjectName, bForDistribution, bPackageDataInsideApk, GameBuildFilesPath, RequiresOBB(bDisallowPackagingDataInApk, ObbFileLocation), bDisableVerifyOBBOnStartUp, Arch, GPUArch, CookFlavor)
 							let OldManifest = File.Exists(manifestFile) ? File.ReadAllText(manifestFile) : ""
 							where manifest != OldManifest
 							select Tuple.Create(Arch, GPUArch, manifest);
diff --git a/Engine/Source/Runtime/Apple/MetalRHI/Private/MetalRHI.cpp b/Engine/Source/Runtime/Apple/MetalRHI/Private/MetalRHI.cpp
--- a/Engine/Source/Runtime/Apple/MetalRHI/Private/MetalRHI.cpp
+++ b/Engine/Source/Runtime/Apple/MetalRHI/Private/MetalRHI.cpp
@@ -68,15 +68,11 @@ FMetalDynamicRHI::FMetalDynamicRHI()
     {
         GMaxRHIFeatureLevel = ERHIFeatureLevel::SM4;
         GMaxRHIShaderPlatform = SP_METAL_MRT;
-		GSupportsMultipleRenderTargets = true;
     }
     else
     {
         GMaxRHIFeatureLevel = ERHIFeatureLevel::ES3_1;
         GMaxRHIShaderPlatform = SP_METAL;
-		// disable MRTs completely, because there is code in the engine that assumes the ability to use
-		// 32 bytes of pixel storage (like the 2 VERY WIDE ones in GPU particles)
-		GSupportsMultipleRenderTargets = false;
 	}
 	
 	GShaderPlatformForFeatureLevel[ERHIFeatureLevel::ES2] = SP_METAL;
diff --git a/Engine/Source/Runtime/Core/Private/Android/AndroidApplication.cpp b/Engine/Source/Runtime/Core/Private/Android/AndroidApplication.cpp
--- a/Engine/Source/Runtime/Core/Private/Android/AndroidApplication.cpp
+++ b/Engine/Source/Runtime/Core/Private/Android/AndroidApplication.cpp
@@ -95,6 +95,18 @@ IInputInterface* FAndroidApplication::GetInputInterface()
 	return InputInterface.Get();
 }
 
+bool FAndroidApplication::IsGamepadAttached() const
+{
+	FAndroidInputInterface* AndroidInputInterface = (FAndroidInputInterface*)InputInterface.Get();
+
+	if (AndroidInputInterface)
+	{
+		return AndroidInputInterface->IsGamepadAttached();
+	}
+
+	return false;
+}
+
 void FDisplayMetrics::GetDisplayMetrics( FDisplayMetrics& OutDisplayMetrics )
 {
 	// Get screen rect
diff --git a/Engine/Source/Runtime/Core/Private/Android/AndroidFile.cpp b/Engine/Source/Runtime/Core/Private/Android/AndroidFile.cpp
--- a/Engine/Source/Runtime/Core/Private/Android/AndroidFile.cpp
+++ b/Engine/Source/Runtime/Core/Private/Android/AndroidFile.cpp
@@ -310,6 +310,10 @@ public:
 			Source += ThisSize;
 			BytesToWrite -= ThisSize;
 		}
+		
+		// Update the cached file length
+		Length = FMath::Max(Length, CurrentOffset);
+
 		return true;
 	}
 
diff --git a/Engine/Source/Runtime/Core/Private/Android/AndroidInputInterface.cpp b/Engine/Source/Runtime/Core/Private/Android/AndroidInputInterface.cpp
--- a/Engine/Source/Runtime/Core/Private/Android/AndroidInputInterface.cpp
+++ b/Engine/Source/Runtime/Core/Private/Android/AndroidInputInterface.cpp
@@ -198,6 +198,33 @@ void FAndroidInputInterface::SetForceFeedbackChannelValues(int32 ControllerId, c
 	UpdateVibeMotors();
 }
 
+extern bool AndroidThunkCpp_IsGamepadAttached();
+
+bool FAndroidInputInterface::IsGamepadAttached() const
+{
+	// Check for gamepads that have already been validated
+	for (int32 DeviceIndex = 0; DeviceIndex < MAX_NUM_CONTROLLERS; DeviceIndex++)
+	{
+		FAndroidGamepadDeviceMapping& CurrentDevice = DeviceMapping[DeviceIndex];
+
+		if (CurrentDevice.DeviceState == MappingState::Valid)
+		{
+			return true;
+		}
+	}
+
+	for (auto DeviceIt = ExternalInputDevices.CreateConstIterator(); DeviceIt; ++DeviceIt)
+	{
+		if ((*DeviceIt)->IsGamepadAttached())
+		{
+			return true;
+		}
+	}
+
+	//if all of this fails, do a check on the Java side to see if the gamepad is attached
+	return AndroidThunkCpp_IsGamepadAttached();
+}
+
 extern void AndroidThunkCpp_Vibrate(int32 Duration);
 
 void FAndroidInputInterface::UpdateVibeMotors()
diff --git a/Engine/Source/Runtime/Core/Private/Android/AndroidMisc.cpp b/Engine/Source/Runtime/Core/Private/Android/AndroidMisc.cpp
--- a/Engine/Source/Runtime/Core/Private/Android/AndroidMisc.cpp
+++ b/Engine/Source/Runtime/Core/Private/Android/AndroidMisc.cpp
@@ -111,9 +111,13 @@ void FAndroidMisc::LoadPreInitModules()
 	FModuleManager::Get().LoadModule(TEXT("AndroidAudio"));
 }
 
+// Test for device vulkan support.
+static void EstablishVulkanDeviceSupport();
+
 void FAndroidMisc::PlatformPreInit()
 {
 	FGenericPlatformMisc::PlatformPreInit();
+	EstablishVulkanDeviceSupport();
 	FAndroidAppEntry::PlatformInit();
 }
 
@@ -1096,11 +1100,17 @@ int32 FAndroidMisc::GetAndroidBuildVersion()
 #define VK_MAKE_VERSION(major, minor, patch) \
     (((major) << 22) | ((minor) << 12) | (patch))
 
+#define VK_VERSION_MAJOR(version) ((uint32)(version) >> 22)
+#define VK_VERSION_MINOR(version) (((uint32)(version) >> 12) & 0x3ff)
+#define VK_VERSION_PATCH(version) ((uint32)(version) & 0xfff)
+
 typedef uint32 VkFlags;
+typedef uint32 VkBool32;
 
 #define VK_DEFINE_HANDLE(object) typedef struct object##_T* object;
 
 VK_DEFINE_HANDLE(VkInstance)
+VK_DEFINE_HANDLE(VkPhysicalDevice)
 
 typedef enum VkResult {
 	VK_SUCCESS = 0,
@@ -1171,26 +1181,188 @@ typedef struct VkAllocationCallbacks {
 	void*	pfnInternalFree;
 } VkAllocationCallbacks;
 
+typedef uint64 VkDeviceSize;
+typedef VkFlags VkSampleCountFlags;
+
+typedef struct VkPhysicalDeviceLimits {
+	uint32				maxImageDimension1D;
+	uint32				maxImageDimension2D;
+	uint32				maxImageDimension3D;
+	uint32				maxImageDimensionCube;
+	uint32				maxImageArrayLayers;
+	uint32				maxTexelBufferElements;
+	uint32				maxUniformBufferRange;
+	uint32				maxStorageBufferRange;
+	uint32				maxPushConstantsSize;
+	uint32				maxMemoryAllocationCount;
+	uint32				maxSamplerAllocationCount;
+	VkDeviceSize		bufferImageGranularity;
+	VkDeviceSize		sparseAddressSpaceSize;
+	uint32				maxBoundDescriptorSets;
+	uint32				maxPerStageDescriptorSamplers;
+	uint32				maxPerStageDescriptorUniformBuffers;
+	uint32				maxPerStageDescriptorStorageBuffers;
+	uint32				maxPerStageDescriptorSampledImages;
+	uint32				maxPerStageDescriptorStorageImages;
+	uint32				maxPerStageDescriptorInputAttachments;
+	uint32				maxPerStageResources;
+	uint32				maxDescriptorSetSamplers;
+	uint32				maxDescriptorSetUniformBuffers;
+	uint32				maxDescriptorSetUniformBuffersDynamic;
+	uint32				maxDescriptorSetStorageBuffers;
+	uint32				maxDescriptorSetStorageBuffersDynamic;
+	uint32				maxDescriptorSetSampledImages;
+	uint32				maxDescriptorSetStorageImages;
+	uint32				maxDescriptorSetInputAttachments;
+	uint32				maxVertexInputAttributes;
+	uint32				maxVertexInputBindings;
+	uint32				maxVertexInputAttributeOffset;
+	uint32				maxVertexInputBindingStride;
+	uint32				maxVertexOutputComponents;
+	uint32				maxTessellationGenerationLevel;
+	uint32				maxTessellationPatchSize;
+	uint32				maxTessellationControlPerVertexInputComponents;
+	uint32				maxTessellationControlPerVertexOutputComponents;
+	uint32				maxTessellationControlPerPatchOutputComponents;
+	uint32				maxTessellationControlTotalOutputComponents;
+	uint32				maxTessellationEvaluationInputComponents;
+	uint32				maxTessellationEvaluationOutputComponents;
+	uint32				maxGeometryShaderInvocations;
+	uint32				maxGeometryInputComponents;
+	uint32				maxGeometryOutputComponents;
+	uint32				maxGeometryOutputVertices;
+	uint32				maxGeometryTotalOutputComponents;
+	uint32				maxFragmentInputComponents;
+	uint32				maxFragmentOutputAttachments;
+	uint32				maxFragmentDualSrcAttachments;
+	uint32				maxFragmentCombinedOutputResources;
+	uint32				maxComputeSharedMemorySize;
+	uint32				maxComputeWorkGroupCount[3];
+	uint32				maxComputeWorkGroupInvocations;
+	uint32				maxComputeWorkGroupSize[3];
+	uint32				subPixelPrecisionBits;
+	uint32				subTexelPrecisionBits;
+	uint32				mipmapPrecisionBits;
+	uint32				maxDrawIndexedIndexValue;
+	uint32				maxDrawIndirectCount;
+	float				maxSamplerLodBias;
+	float				maxSamplerAnisotropy;
+	uint32				maxViewports;
+	uint32				maxViewportDimensions[2];
+	float				viewportBoundsRange[2];
+	uint32				viewportSubPixelBits;
+	size_t				minMemoryMapAlignment;
+	VkDeviceSize		minTexelBufferOffsetAlignment;
+	VkDeviceSize		minUniformBufferOffsetAlignment;
+	VkDeviceSize		minStorageBufferOffsetAlignment;
+	int32				minTexelOffset;
+	uint32				maxTexelOffset;
+	int32				minTexelGatherOffset;
+	uint32				maxTexelGatherOffset;
+	float				minInterpolationOffset;
+	float				maxInterpolationOffset;
+	uint32				subPixelInterpolationOffsetBits;
+	uint32				maxFramebufferWidth;
+	uint32				maxFramebufferHeight;
+	uint32				maxFramebufferLayers;
+	VkSampleCountFlags	framebufferColorSampleCounts;
+	VkSampleCountFlags	framebufferDepthSampleCounts;
+	VkSampleCountFlags	framebufferStencilSampleCounts;
+	VkSampleCountFlags	framebufferNoAttachmentsSampleCounts;
+	uint32				maxColorAttachments;
+	VkSampleCountFlags	sampledImageColorSampleCounts;
+	VkSampleCountFlags	sampledImageIntegerSampleCounts;
+	VkSampleCountFlags	sampledImageDepthSampleCounts;
+	VkSampleCountFlags	sampledImageStencilSampleCounts;
+	VkSampleCountFlags	storageImageSampleCounts;
+	uint32				maxSampleMaskWords;
+	VkBool32			timestampComputeAndGraphics;
+	float				timestampPeriod;
+	uint32				maxClipDistances;
+	uint32				maxCullDistances;
+	uint32				maxCombinedClipAndCullDistances;
+	uint32				discreteQueuePriorities;
+	float				pointSizeRange[2];
+	float				lineWidthRange[2];
+	float				pointSizeGranularity;
+	float				lineWidthGranularity;
+	VkBool32			strictLines;
+	VkBool32			standardSampleLocations;
+	VkDeviceSize		optimalBufferCopyOffsetAlignment;
+	VkDeviceSize		optimalBufferCopyRowPitchAlignment;
+	VkDeviceSize		nonCoherentAtomSize;
+} VkPhysicalDeviceLimits;
+
+typedef struct VkPhysicalDeviceSparseProperties {
+	VkBool32	residencyStandard2DBlockShape;
+	VkBool32	residencyStandard2DMultisampleBlockShape;
+	VkBool32	residencyStandard3DBlockShape;
+	VkBool32	residencyAlignedMipSize;
+	VkBool32	residencyNonResidentStrict;
+} VkPhysicalDeviceSparseProperties;
+
+typedef enum VkPhysicalDeviceType {
+	VK_PHYSICAL_DEVICE_TYPE_OTHER = 0,
+	VK_PHYSICAL_DEVICE_TYPE_INTEGRATED_GPU = 1,
+	VK_PHYSICAL_DEVICE_TYPE_DISCRETE_GPU = 2,
+	VK_PHYSICAL_DEVICE_TYPE_VIRTUAL_GPU = 3,
+	VK_PHYSICAL_DEVICE_TYPE_CPU = 4,
+	VK_PHYSICAL_DEVICE_TYPE_BEGIN_RANGE = VK_PHYSICAL_DEVICE_TYPE_OTHER,
+	VK_PHYSICAL_DEVICE_TYPE_END_RANGE = VK_PHYSICAL_DEVICE_TYPE_CPU,
+	VK_PHYSICAL_DEVICE_TYPE_RANGE_SIZE = (VK_PHYSICAL_DEVICE_TYPE_CPU - VK_PHYSICAL_DEVICE_TYPE_OTHER + 1),
+	VK_PHYSICAL_DEVICE_TYPE_MAX_ENUM = 0x7FFFFFFF
+} VkPhysicalDeviceType;
+
+#define VK_MAX_PHYSICAL_DEVICE_NAME_SIZE 256
+#define VK_UUID_SIZE 16
+
+typedef struct VkPhysicalDeviceProperties {
+	uint32								apiVersion;
+	uint32								driverVersion;
+	uint32								vendorID;
+	uint32								deviceID;
+	VkPhysicalDeviceType				deviceType;
+	char								deviceName[VK_MAX_PHYSICAL_DEVICE_NAME_SIZE];
+	uint8_t								pipelineCacheUUID[VK_UUID_SIZE];
+	VkPhysicalDeviceLimits				limits;
+	VkPhysicalDeviceSparseProperties	sparseProperties;
+} VkPhysicalDeviceProperties;
+
 typedef VkResult(VKAPI_PTR *PFN_vkCreateInstance)(const VkInstanceCreateInfo* pCreateInfo, const VkAllocationCallbacks* pAllocator, VkInstance* pInstance);
 typedef void (VKAPI_PTR *PFN_vkDestroyInstance)(VkInstance instance, const VkAllocationCallbacks* pAllocator);
+typedef VkResult(VKAPI_PTR *PFN_vkEnumeratePhysicalDevices)(VkInstance instance, uint32* pPhysicalDeviceCount, VkPhysicalDevice* pPhysicalDevices);
+typedef void (VKAPI_PTR *PFN_vkGetPhysicalDeviceProperties)(VkPhysicalDevice physicalDevice, VkPhysicalDeviceProperties* pProperties);
 
 ///////////////////////////////////////////////////////////////////////////////
 
 #define UE_VK_API_VERSION	VK_MAKE_VERSION(1, 0, 1)
 
-static bool AttemptVulkanInit(void* VulkanLib)
+enum class EDeviceVulkanSupportStatus
+{
+	Uninitialized,
+	NotSupported,
+	Supported
+};
+
+static FString VulkanVersionString;
+static EDeviceVulkanSupportStatus VulkanSupport = EDeviceVulkanSupportStatus::Uninitialized;
+
+static EDeviceVulkanSupportStatus AttemptVulkanInit(void* VulkanLib)
 {
 	if (VulkanLib == nullptr)
 	{
-		return false;
+		return EDeviceVulkanSupportStatus::NotSupported;
 	}
 
 	// Try to get required functions to check for driver
 	PFN_vkCreateInstance vkCreateInstance = (PFN_vkCreateInstance)dlsym(VulkanLib, "vkCreateInstance");
 	PFN_vkDestroyInstance vkDestroyInstance = (PFN_vkDestroyInstance)dlsym(VulkanLib, "vkDestroyInstance");
-	if (!vkCreateInstance || !vkDestroyInstance)
+	PFN_vkEnumeratePhysicalDevices vkEnumeratePhysicalDevices = (PFN_vkEnumeratePhysicalDevices)dlsym(VulkanLib, "vkEnumeratePhysicalDevices");
+	PFN_vkGetPhysicalDeviceProperties vkGetPhysicalDeviceProperties = (PFN_vkGetPhysicalDeviceProperties)dlsym(VulkanLib, "vkGetPhysicalDeviceProperties");
+	
+	if (!vkCreateInstance || !vkDestroyInstance || !vkEnumeratePhysicalDevices || !vkGetPhysicalDeviceProperties)
 	{
-		return false;
+		return EDeviceVulkanSupportStatus::NotSupported;
 	}
 		
 	// try to create instance to verify driver available
@@ -1215,73 +1387,118 @@ static bool AttemptVulkanInit(void* VulkanLib)
 	VkResult Result = vkCreateInstance(&InstInfo, nullptr, &Instance);
 	if (Result != VK_SUCCESS)
 	{
-		return false;
+		return EDeviceVulkanSupportStatus::NotSupported;
+	}
+
+	// Determine Vulkan device's API level.
+	uint32 GpuCount = 0;
+	Result = vkEnumeratePhysicalDevices(Instance, &GpuCount, nullptr);
+	if (Result != VK_SUCCESS || GpuCount == 0)
+	{
+		vkDestroyInstance(Instance, nullptr);
+		return EDeviceVulkanSupportStatus::NotSupported;
+	}
+
+	TArray<VkPhysicalDevice> PhysicalDevices;
+	PhysicalDevices.AddZeroed(GpuCount);
+	Result = vkEnumeratePhysicalDevices(Instance, &GpuCount, PhysicalDevices.GetData());
+	if (Result != VK_SUCCESS)
+	{
+		vkDestroyInstance(Instance, nullptr);
+		return EDeviceVulkanSupportStatus::NotSupported;
 	}
+
+	// Don't care which device - This code is making the assumption that all devices will have same api version.
+	VkPhysicalDeviceProperties DeviceProperties;
+	vkGetPhysicalDeviceProperties(PhysicalDevices[0], &DeviceProperties);
+
+	VulkanVersionString = FString::Printf(TEXT("%d.%d.%d"), VK_VERSION_MAJOR(DeviceProperties.apiVersion), VK_VERSION_MINOR(DeviceProperties.apiVersion), VK_VERSION_PATCH(DeviceProperties.apiVersion));
 	vkDestroyInstance(Instance, nullptr);
-	return true;
+
+	return EDeviceVulkanSupportStatus::Supported;
 }
 
-extern bool AndroidThunkCpp_GetMetaDataBoolean(const FString& Key);
+extern int32 AndroidThunkCpp_GetMetaDataInt(const FString& Key);
 
-bool FAndroidMisc::ShouldUseVulkan()
+// Test for device vulkan support.
+static void EstablishVulkanDeviceSupport()
 {
-#if PLATFORM_ANDROID_VULKAN
 	// just do this check once
-	static int32 ShouldUseVulkanFlag = -1;
-	if (ShouldUseVulkanFlag == -1)
+	check(VulkanSupport == EDeviceVulkanSupportStatus::Uninitialized);
+	// assume no
+	VulkanSupport = EDeviceVulkanSupportStatus::NotSupported;
+	VulkanVersionString = TEXT("0.0.0");
+
+	// make sure the Vulkan RHI is compiled in
+	if (FModuleManager::Get().ModuleExists(TEXT("VulkanRHI")))
 	{
-		// assume no
-		ShouldUseVulkanFlag = 0;
+		FPlatformMisc::LowLevelOutputDebugString(TEXT("Compiled with Vulkan support"));
 
-		// make sure the project setting has enabled Vulkan support (per-project user settings in the editor) from AndroidManifest.xml
-		bool bSupportsVulkan = AndroidThunkCpp_GetMetaDataBoolean(TEXT("com.epicgames.ue4.GameActivity.bSupportsVulkan"));
-		if (bSupportsVulkan && FModuleManager::Get().ModuleExists(TEXT("VulkanRHI")))
-		{
-			FPlatformMisc::LowLevelOutputDebugString(TEXT("Compiled with Vulkan support"));
+		// does commandline override (using GL or ES2 for legacy commandlines)
+		bool bForceOpenGL = FParse::Param(FCommandLine::Get(), TEXT("GL")) || FParse::Param(FCommandLine::Get(), TEXT("OpenGL")) || FParse::Param(FCommandLine::Get(), TEXT("ES2"));
 
-			// does commandline override (using GL or ES2 for legacy commandlines)
-			bool bForceOpenGL = FParse::Param(FCommandLine::Get(), TEXT("GL")) || FParse::Param(FCommandLine::Get(), TEXT("OpenGL")) || FParse::Param(FCommandLine::Get(), TEXT("ES2"));
-			static const auto CVarDisableVulkan = IConsoleManager::Get().FindTConsoleVariableDataInt(TEXT("r.Android.DisableVulkanSupport"));
-			bForceOpenGL = bForceOpenGL || CVarDisableVulkan->GetValueOnAnyThread() == 1;
-			if (!bForceOpenGL)
+		if (!bForceOpenGL)
+		{
+			// check for libvulkan.so
+			void* VulkanLib = dlopen("libvulkan.so", RTLD_NOW | RTLD_LOCAL);
+			if (VulkanLib != nullptr)
 			{
-				// check for libvulkan.so
-				void* VulkanLib = dlopen("libvulkan.so", RTLD_NOW | RTLD_LOCAL);
-				if (VulkanLib != nullptr)
-				{
-					FPlatformMisc::LowLevelOutputDebugString(TEXT("Vulkan library detected, checking for available driver"));
-					ShouldUseVulkanFlag = AttemptVulkanInit(VulkanLib);
-					dlclose(VulkanLib);
+				FPlatformMisc::LowLevelOutputDebugString(TEXT("Vulkan library detected, checking for available driver"));
 
-					if (ShouldUseVulkanFlag)
-					{
-						FPlatformMisc::LowLevelOutputDebugString(TEXT("Vulkan driver available, will use VulkanRHI"));
-					}
-					else
+				// if Nougat, we can check the Vulkan version
+				if (FAndroidMisc::GetAndroidBuildVersion() >= 24)
+				{
+					int32 VulkanVersion = AndroidThunkCpp_GetMetaDataInt(TEXT("android.hardware.vulkan.version"));
+					if (VulkanVersion >= UE_VK_API_VERSION)
 					{
-						FPlatformMisc::LowLevelOutputDebugString(TEXT("Vulkan driver NOT available, falling back to OpenGL ES"));
+						// final check, try initializing the instance
+						VulkanSupport = AttemptVulkanInit(VulkanLib);
 					}
 				}
 				else
 				{
-					FPlatformMisc::LowLevelOutputDebugString(TEXT("Vulkan library NOT detected, falling back to OpenGL ES"));
+					// otherwise, we need to try initializing the instance
+					VulkanSupport = AttemptVulkanInit(VulkanLib);
+				}
+
+				dlclose(VulkanLib);
+
+				if (VulkanSupport == EDeviceVulkanSupportStatus::Supported)
+				{
+					FPlatformMisc::LowLevelOutputDebugString(TEXT("Vulkan driver available, will use VulkanRHI"));
+				}
+				else
+				{
+					FPlatformMisc::LowLevelOutputDebugString(TEXT("Vulkan driver NOT available, falling back to OpenGL ES"));
 				}
 			}
 			else
 			{
-				FPlatformMisc::LowLevelOutputDebugString(TEXT("Forced OpenGL ES"));
+				FPlatformMisc::LowLevelOutputDebugString(TEXT("Vulkan library NOT detected, falling back to OpenGL ES"));
 			}
 		}
 		else
 		{
-			FPlatformMisc::LowLevelOutputDebugString(TEXT("Compiled with OpenGL ES support only"));
+			FPlatformMisc::LowLevelOutputDebugString(TEXT("Forced OpenGL ES"));
 		}
 	}
+	else
+	{
+		FPlatformMisc::LowLevelOutputDebugString(TEXT("Compiled with OpenGL ES support only"));
+	}
+}
 
-	return ShouldUseVulkanFlag == 1;
-#else
-	return false;
-#endif
+bool FAndroidMisc::ShouldUseVulkan()
+{
+	check(VulkanSupport != EDeviceVulkanSupportStatus::Uninitialized);
+	static const auto CVarDisableVulkan = IConsoleManager::Get().FindTConsoleVariableDataInt(TEXT("r.Android.DisableVulkanSupport"));
+	return VulkanSupport == EDeviceVulkanSupportStatus::Supported && CVarDisableVulkan->GetValueOnAnyThread() == 0;
+}
+
+FString FAndroidMisc::GetVulkanVersion()
+{
+	check(VulkanSupport != EDeviceVulkanSupportStatus::Uninitialized);
+	return VulkanVersionString;
 }
 
 #if !UE_BUILD_SHIPPING
diff --git a/Engine/Source/Runtime/Core/Private/HAL/ConsoleManager.cpp b/Engine/Source/Runtime/Core/Private/HAL/ConsoleManager.cpp
--- a/Engine/Source/Runtime/Core/Private/HAL/ConsoleManager.cpp
+++ b/Engine/Source/Runtime/Core/Private/HAL/ConsoleManager.cpp
@@ -2323,4 +2323,10 @@ static TAutoConsoleVariable<int32> CVarDisableOpenGLES31Support(
 	TEXT("Disable support for OpenGLES 3.1 API. (Android Only)\n")
 	TEXT("  0 = OpenGLES 3.1 API will be used (providing device and project supports it) [default]\n")
 	TEXT("  1 = OpenGLES 3.1 will be disabled, OpenGL ES2 fall back will be used."),
-	ECVF_ReadOnly);
\ No newline at end of file
+	ECVF_ReadOnly);
+
+static TAutoConsoleVariable<int32> GLSLCvar(
+	TEXT("r.Vulkan.UseGLSL"),
+	0,
+	TEXT("2 to use ES GLSL\n1 to use GLSL\n0 to use SPIRV")
+);
\ No newline at end of file
diff --git a/Engine/Source/Runtime/Core/Public/Android/AndroidApplication.h b/Engine/Source/Runtime/Core/Public/Android/AndroidApplication.h
--- a/Engine/Source/Runtime/Core/Public/Android/AndroidApplication.h
+++ b/Engine/Source/Runtime/Core/Public/Android/AndroidApplication.h
@@ -70,6 +70,8 @@ public:
 
 	static void OnWindowSizeChanged();
 
+	virtual bool IsGamepadAttached() const override;
+
 private:
 
 	FAndroidApplication();
diff --git a/Engine/Source/Runtime/Core/Public/Android/AndroidInputInterface.h b/Engine/Source/Runtime/Core/Public/Android/AndroidInputInterface.h
--- a/Engine/Source/Runtime/Core/Public/Android/AndroidInputInterface.h
+++ b/Engine/Source/Runtime/Core/Public/Android/AndroidInputInterface.h
@@ -214,6 +214,8 @@ public:
 	virtual void SetForceFeedbackChannelValues(int32 ControllerId, const FForceFeedbackValues &values) override;
 	virtual void SetLightColor(int32 ControllerId, FColor Color) override {}
 
+	virtual bool IsGamepadAttached() const;
+
 	virtual void AddExternalInputDevice(TSharedPtr<class IInputDevice> InputDevice);
 
 private:
diff --git a/Engine/Source/Runtime/Core/Public/Android/AndroidMisc.h b/Engine/Source/Runtime/Core/Public/Android/AndroidMisc.h
--- a/Engine/Source/Runtime/Core/Public/Android/AndroidMisc.h
+++ b/Engine/Source/Runtime/Core/Public/Android/AndroidMisc.h
@@ -120,6 +120,7 @@ struct CORE_API FAndroidMisc : public FGenericPlatformMisc
 	static bool SupportsShaderFramebufferFetch();
 	static int GetAndroidBuildVersion();
 	static bool ShouldUseVulkan();
+	static FString GetVulkanVersion();
 
 #if !UE_BUILD_SHIPPING
 	static bool IsDebuggerPresent();
diff --git a/Engine/Source/Runtime/Engine/Classes/Particles/SubUVAnimation.h b/Engine/Source/Runtime/Engine/Classes/Particles/SubUVAnimation.h
--- a/Engine/Source/Runtime/Engine/Classes/Particles/SubUVAnimation.h
+++ b/Engine/Source/Runtime/Engine/Classes/Particles/SubUVAnimation.h
@@ -180,7 +180,7 @@ public:
 
 	inline const FVector2D* GetFrameData(int32 FrameIndex) const
 	{
-		return &DerivedData.BoundingGeometry[FrameIndex * GetNumFrames()];
+		return &DerivedData.BoundingGeometry[FrameIndex * GetNumBoundingVertices()];
 	}
 
 	inline FShaderResourceViewRHIParamRef GetBoundingGeometrySRV() const
diff --git a/Engine/Source/Runtime/Engine/Private/Components/ReflectionCaptureComponent.cpp b/Engine/Source/Runtime/Engine/Private/Components/ReflectionCaptureComponent.cpp
--- a/Engine/Source/Runtime/Engine/Private/Components/ReflectionCaptureComponent.cpp
+++ b/Engine/Source/Runtime/Engine/Private/Components/ReflectionCaptureComponent.cpp
@@ -645,14 +645,11 @@ void FReflectionCaptureEncodedHDRDerivedData::GenerateFromDerivedDataSource(cons
 	}
 }
 
-// Generate a new guid to force a recache of all encoded HDR derived data
-#define REFLECTIONCAPTURE_ENCODED_DERIVEDDATA_VER TEXT("6B6DFE9DF44888914934C082283C3296")
-
 FString FReflectionCaptureEncodedHDRDerivedData::GetDDCKeyString(const FGuid& StateId, int32 CubemapDimension)
 {
 	return FDerivedDataCacheInterface::BuildCacheKey(
-		TEXT("REFL_ENC"), 
-		REFLECTIONCAPTURE_ENCODED_DERIVEDDATA_VER, 
+		TEXT("REFL_ENC"),
+		*ReflectionCaptureDDCVer.ToString(),
 		*StateId.ToString().Append("_").Append(FString::FromInt(CubemapDimension))
 		);
 }
diff --git a/Engine/Source/Runtime/Engine/Private/PlayerController.cpp b/Engine/Source/Runtime/Engine/Private/PlayerController.cpp
--- a/Engine/Source/Runtime/Engine/Private/PlayerController.cpp
+++ b/Engine/Source/Runtime/Engine/Private/PlayerController.cpp
@@ -4563,6 +4563,8 @@ void APlayerController::ActivateTouchInterface(UTouchInterface* NewTouchInterfac
 		{
 			LocalPlayer->ViewportClient->RemoveViewportWidgetContent(VirtualJoystick.ToSharedRef());
 		}
+		//clear any input before clearing the VirtualJoystick
+		FlushPressedKeys();
 		VirtualJoystick = NULL;
 	}
 }
diff --git a/Engine/Source/Runtime/Engine/Private/World.cpp b/Engine/Source/Runtime/Engine/Private/World.cpp
--- a/Engine/Source/Runtime/Engine/Private/World.cpp
+++ b/Engine/Source/Runtime/Engine/Private/World.cpp
@@ -3186,6 +3186,28 @@ void UWorld::InitializeActorsForPlay(const FURL& InURL, bool bResetTime)
 			Level->RouteActorInitialize();
 		}
 
+		// Let server know client sub-levels visibility state
+		for (int32 LevelIndex = 1; LevelIndex < Levels.Num(); LevelIndex++)
+		{
+			ULevel*	SubLevel = Levels[LevelIndex];
+			// Remap packagename for PIE networking before sending out to server
+			FName PackageName = SubLevel->GetOutermost()->GetFName();
+			FString PackageNameStr = PackageName.ToString();
+			if (GEngine->NetworkRemapPath(this, PackageNameStr, false))
+			{
+				PackageName = FName(*PackageNameStr);
+			}
+
+			for (FLocalPlayerIterator It(GEngine, this); It; ++It)
+			{
+				APlayerController* LocalPlayerController = It->GetPlayerController(this);
+				if (LocalPlayerController != NULL)
+				{
+					LocalPlayerController->ServerUpdateLevelVisibility(PackageName, SubLevel->bIsVisible);
+				}
+			}
+		}
+
 		bStartup = 0;
 	}
 
diff --git a/Engine/Source/Runtime/Engine/Public/MaterialShared.h b/Engine/Source/Runtime/Engine/Public/MaterialShared.h
--- a/Engine/Source/Runtime/Engine/Public/MaterialShared.h
+++ b/Engine/Source/Runtime/Engine/Public/MaterialShared.h
@@ -1548,10 +1548,10 @@ public:
 	{}
 
 	// FMaterialRenderProxy interface.
-	virtual const class FMaterial* GetMaterial(ERHIFeatureLevel::Type InFeatureLevel) const;
-	virtual bool GetVectorValue(const FName ParameterName, FLinearColor* OutValue, const FMaterialRenderContext& Context) const;
-	virtual bool GetScalarValue(const FName ParameterName, float* OutValue, const FMaterialRenderContext& Context) const;
-	virtual bool GetTextureValue(const FName ParameterName,const UTexture** OutValue, const FMaterialRenderContext& Context) const;
+	ENGINE_API virtual const class FMaterial* GetMaterial(ERHIFeatureLevel::Type InFeatureLevel) const;
+	ENGINE_API virtual bool GetVectorValue(const FName ParameterName, FLinearColor* OutValue, const FMaterialRenderContext& Context) const;
+	ENGINE_API virtual bool GetScalarValue(const FName ParameterName, float* OutValue, const FMaterialRenderContext& Context) const;
+	ENGINE_API virtual bool GetTextureValue(const FName ParameterName,const UTexture** OutValue, const FMaterialRenderContext& Context) const;
 };
 
 
diff --git a/Engine/Source/Runtime/Engine/Public/StaticMeshResources.h b/Engine/Source/Runtime/Engine/Public/StaticMeshResources.h
--- a/Engine/Source/Runtime/Engine/Public/StaticMeshResources.h
+++ b/Engine/Source/Runtime/Engine/Public/StaticMeshResources.h
@@ -1164,17 +1164,17 @@ struct FInstanceStream
 		Me->InstanceOrigin.W = RandomInstanceID;
 
 		Me->InstanceTransform1[0] = Transform.M[0][0];
-		Me->InstanceTransform1[1] = Transform.M[1][0];
-		Me->InstanceTransform1[2] = Transform.M[2][0];
+		Me->InstanceTransform1[1] = Transform.M[0][1];
+		Me->InstanceTransform1[2] = Transform.M[0][2];
 		Me->InstanceTransform1[3] = FloatType();
 
-		Me->InstanceTransform2[0] = Transform.M[0][1];
+		Me->InstanceTransform2[0] = Transform.M[1][1];
 		Me->InstanceTransform2[1] = Transform.M[1][1];
-		Me->InstanceTransform2[2] = Transform.M[2][1];
+		Me->InstanceTransform2[2] = Transform.M[1][2];
 		Me->InstanceTransform2[3] = FloatType();
 
-		Me->InstanceTransform3[0] = Transform.M[0][2];
-		Me->InstanceTransform3[1] = Transform.M[1][2];
+		Me->InstanceTransform3[0] = Transform.M[2][0];
+		Me->InstanceTransform3[1] = Transform.M[2][1];
 		Me->InstanceTransform3[2] = Transform.M[2][2];
 		Me->InstanceTransform3[3] = FloatType();
 
@@ -1191,15 +1191,15 @@ struct FInstanceStream
 		Transform.M[3][2] = InstanceOrigin.Z;
 
 		Transform.M[0][0] = InstanceTransform1[0];
-		Transform.M[1][0] = InstanceTransform1[1];
-		Transform.M[2][0] = InstanceTransform1[2];
+		Transform.M[0][1] = InstanceTransform1[1];
+		Transform.M[0][2] = InstanceTransform1[2];
 
-		Transform.M[0][1] = InstanceTransform2[0];
+		Transform.M[1][0] = InstanceTransform2[0];
 		Transform.M[1][1] = InstanceTransform2[1];
-		Transform.M[2][1] = InstanceTransform2[2];
+		Transform.M[1][2] = InstanceTransform2[2];
 
-		Transform.M[0][2] = InstanceTransform3[0];
-		Transform.M[1][2] = InstanceTransform3[1];
+		Transform.M[2][0] = InstanceTransform3[0];
+		Transform.M[2][1] = InstanceTransform3[1];
 		Transform.M[2][2] = InstanceTransform3[2];
 
 		Transform.M[0][3] = 0.f;
@@ -1250,17 +1250,17 @@ struct FInstanceStream
 		Me->InstanceOrigin.W = RandomInstanceID;
 
 		Me->InstanceTransform1[0] = Transform.M[0][0];
-		Me->InstanceTransform1[1] = Transform.M[1][0];
-		Me->InstanceTransform1[2] = Transform.M[2][0];
+		Me->InstanceTransform1[1] = Transform.M[0][1];
+		Me->InstanceTransform1[2] = Transform.M[0][2];
 		Me->InstanceTransform1[3] = ((float)HitProxyColor.R) + (bSelected ? 256.f : 0.0f);
 
-		Me->InstanceTransform2[0] = Transform.M[0][1];
+		Me->InstanceTransform2[0] = Transform.M[1][0];
 		Me->InstanceTransform2[1] = Transform.M[1][1];
-		Me->InstanceTransform2[2] = Transform.M[2][1];
+		Me->InstanceTransform2[2] = Transform.M[1][2];
 		Me->InstanceTransform2[3] = (float)HitProxyColor.G;
 
-		Me->InstanceTransform3[0] = Transform.M[0][2];
-		Me->InstanceTransform3[1] = Transform.M[1][2];
+		Me->InstanceTransform3[0] = Transform.M[2][0];
+		Me->InstanceTransform3[1] = Transform.M[2][1];
 		Me->InstanceTransform3[2] = Transform.M[2][2];
 		Me->InstanceTransform3[3] = (float)HitProxyColor.B;
 
@@ -1283,17 +1283,17 @@ struct FInstanceStream
 		Me->InstanceOrigin.W = RandomInstanceID;
 
 		Me->InstanceTransform1[0] = Transform.M[0][0];
-		Me->InstanceTransform1[1] = Transform.M[1][0];
-		Me->InstanceTransform1[2] = Transform.M[2][0];
+		Me->InstanceTransform1[1] = Transform.M[0][1];
+		Me->InstanceTransform1[2] = Transform.M[0][2];
 		Me->InstanceTransform1[3] = FloatType();
 
-		Me->InstanceTransform2[0] = Transform.M[0][1];
+		Me->InstanceTransform2[0] = Transform.M[1][0];
 		Me->InstanceTransform2[1] = Transform.M[1][1];
-		Me->InstanceTransform2[2] = Transform.M[2][1];
+		Me->InstanceTransform2[2] = Transform.M[1][2];
 		Me->InstanceTransform2[3] = FloatType();
 
-		Me->InstanceTransform3[0] = Transform.M[0][2];
-		Me->InstanceTransform3[1] = Transform.M[1][2];
+		Me->InstanceTransform3[0] = Transform.M[2][0];
+		Me->InstanceTransform3[1] = Transform.M[2][1];
 		Me->InstanceTransform3[2] = Transform.M[2][2];
 		Me->InstanceTransform3[3] = FloatType();
 
diff --git a/Engine/Source/Runtime/Landscape/Private/LandscapeCollision.cpp b/Engine/Source/Runtime/Landscape/Private/LandscapeCollision.cpp
--- a/Engine/Source/Runtime/Landscape/Private/LandscapeCollision.cpp
+++ b/Engine/Source/Runtime/Landscape/Private/LandscapeCollision.cpp
@@ -1384,7 +1384,7 @@ void ULandscapeHeightfieldCollisionComponent::Serialize(FArchive& Ar)
 
 		if (bCooked)
 		{
-			Ar << CookedCollisionData;
+			CookedCollisionData.BulkSerialize(Ar);
 		}
 		else
 		{
diff --git a/Engine/Source/Runtime/Landscape/Private/LandscapeRender.cpp b/Engine/Source/Runtime/Landscape/Private/LandscapeRender.cpp
--- a/Engine/Source/Runtime/Landscape/Private/LandscapeRender.cpp
+++ b/Engine/Source/Runtime/Landscape/Private/LandscapeRender.cpp
@@ -1414,6 +1414,21 @@ void FLandscapeComponentSceneProxy::CalcLODParamsForSubsection(const FSceneView&
 	OutNeighborLODs[3] = FMath::Max<float>(OutfLOD, CalcDesiredLOD(View, CameraLocalPos, SubX,     SubY + 1));
 }
 
+namespace
+{
+    FLinearColor GetColorForLod(int32 CurrentLOD, int32 ForcedLOD)
+    {
+	    int32 ColorIndex = INDEX_NONE;
+	    if (GEngine->LODColorationColors.Num() > 0)
+	    {
+		    ColorIndex = CurrentLOD;
+		    ColorIndex = FMath::Clamp(ColorIndex, 0, GEngine->LODColorationColors.Num() - 1);
+	    }
+	    const FLinearColor& LODColor = ColorIndex != INDEX_NONE ? GEngine->LODColorationColors[ColorIndex] : FLinearColor::Gray;
+	    return ForcedLOD >= 0 ? LODColor : LODColor * 0.2f;
+	}
+}
+
 void FLandscapeComponentSceneProxy::GetDynamicMeshElements(const TArray<const FSceneView*>& Views, const FSceneViewFamily& ViewFamily, uint32 VisibilityMap, FMeshElementCollector& Collector) const
 {
 	QUICK_SCOPE_CYCLE_COUNTER(STAT_FLandscapeComponentSceneProxy_GetMeshElements);
@@ -1547,6 +1562,8 @@ void FLandscapeComponentSceneProxy::GetDynamicMeshElements(const TArray<const FS
 
 			// Render the landscape component
 #if WITH_EDITOR
+			const bool bMaterialModifiesMeshPosition = MaterialInterface->GetRenderProxy(false)->GetMaterial(View->GetFeatureLevel())->MaterialModifiesMeshPosition_RenderThread();
+
 			switch (GLandscapeViewMode)
 			{
 			case ELandscapeViewMode::DebugLayer:
@@ -1613,43 +1630,29 @@ void FLandscapeComponentSceneProxy::GetDynamicMeshElements(const TArray<const FS
 
 			case ELandscapeViewMode::LOD:
 			{
-				for (int32 i = 0; i < MeshTools.Elements.Num(); i++)
+				auto& TemplateMesh = bIsWireframe ? Mesh : MeshTools;
+				for (int32 i = 0; i < TemplateMesh.Elements.Num(); i++)
 				{
 					FMeshBatch& LODMesh = Collector.AllocateMesh();
-					LODMesh = MeshTools;
+					LODMesh = TemplateMesh;
 					LODMesh.Elements.Empty(1);
-					LODMesh.Elements.Add(MeshTools.Elements[i]);
-
-					int32 ColorIndex = INDEX_NONE;
-					if (GEngine->LODColorationColors.Num() > 0)
-					{
-						ColorIndex = ((FLandscapeBatchElementParams*)MeshTools.Elements[i].UserData)->CurrentLOD;
-						ColorIndex = FMath::Clamp(ColorIndex, 0, GEngine->LODColorationColors.Num() - 1);
-						LODMesh.VisualizeLODIndex = ColorIndex;
-					}
-
-					const FLinearColor& LODColor = ColorIndex != INDEX_NONE ? GEngine->LODColorationColors[ColorIndex] : FLinearColor::Gray;
-					FLinearColor Color = ForcedLOD >= 0 ? LODColor : LODColor * 0.2f;
-
-					auto LODMaterialInstance = new FColoredMaterialRenderProxy(GEngine->LevelColorationUnlitMaterial->GetRenderProxy(false), Color);
-					LODMesh.MaterialRenderProxy = LODMaterialInstance;
-					Collector.RegisterOneFrameMaterialProxy(LODMaterialInstance);
-
-					if (ViewFamily.EngineShowFlags.Wireframe)
-					{
-						LODMesh.bCanApplyViewModeOverrides = false;
-						LODMesh.bWireframe = true;
-					}
-					else
-					{
-						LODMesh.bCanApplyViewModeOverrides = true;
-						LODMesh.bUseWireframeSelectionColoring = IsSelected();
-					}
-
+					LODMesh.Elements.Add(TemplateMesh.Elements[i]);
+					int32 CurrentLOD = ((FLandscapeBatchElementParams*)TemplateMesh.Elements[i].UserData)->CurrentLOD;
+					LODMesh.VisualizeLODIndex = CurrentLOD;
+					FLinearColor Color = GetColorForLod(CurrentLOD, ForcedLOD);
+					FMaterialRenderProxy* LODMaterialProxy =
+						bMaterialModifiesMeshPosition && bIsWireframe
+						? (FMaterialRenderProxy*)new FOverrideSelectionColorMaterialRenderProxy(MaterialInterface->GetRenderProxy(false), Color)
+						: (FMaterialRenderProxy*)new FColoredMaterialRenderProxy(GEngine->LevelColorationUnlitMaterial->GetRenderProxy(false), Color);
+					Collector.RegisterOneFrameMaterialProxy(LODMaterialProxy);
+					LODMesh.MaterialRenderProxy = LODMaterialProxy;
+					LODMesh.bCanApplyViewModeOverrides = !bIsWireframe;
+					LODMesh.bWireframe = bIsWireframe;
+					LODMesh.bUseWireframeSelectionColoring = IsSelected();
 					Collector.AddMesh(ViewIndex, LODMesh);
 
 					NumPasses++;
-					NumTriangles += MeshTools.Elements[i].NumPrimitives;
+					NumTriangles += TemplateMesh.Elements[i].NumPrimitives;
 					NumDrawCalls++;
 				}
 			}
diff --git a/Engine/Source/Runtime/Launch/Private/Android/AndroidJNI.cpp b/Engine/Source/Runtime/Launch/Private/Android/AndroidJNI.cpp
--- a/Engine/Source/Runtime/Launch/Private/Android/AndroidJNI.cpp
+++ b/Engine/Source/Runtime/Launch/Private/Android/AndroidJNI.cpp
@@ -57,6 +57,7 @@ void FJavaWrapper::FindClassesAndMethods(JNIEnv* Env)
 	AndroidThunkJava_InitHMDs = FindMethod(Env, GameActivityClassID, "AndroidThunkJava_InitHMDs", "()V", bIsOptional);
 	AndroidThunkJava_DismissSplashScreen = FindMethod(Env, GameActivityClassID, "AndroidThunkJava_DismissSplashScreen", "()V", bIsOptional);
 	AndroidThunkJava_GetInputDeviceInfo = FindMethod(Env, GameActivityClassID, "AndroidThunkJava_GetInputDeviceInfo", "(I)Lcom/epicgames/ue4/GameActivity$InputDeviceInfo;", bIsOptional);
+	AndroidThunkJava_IsGamepadAttached = FindMethod(Env, GameActivityClassID, "AndroidThunkJava_IsGamepadAttached", "()Z", bIsOptional);
 	AndroidThunkJava_HasMetaDataKey = FindMethod(Env, GameActivityClassID, "AndroidThunkJava_HasMetaDataKey", "(Ljava/lang/String;)Z", bIsOptional);
 	AndroidThunkJava_GetMetaDataBoolean = FindMethod(Env, GameActivityClassID, "AndroidThunkJava_GetMetaDataBoolean", "(Ljava/lang/String;)Z", bIsOptional);
 	AndroidThunkJava_GetMetaDataInt = FindMethod(Env, GameActivityClassID, "AndroidThunkJava_GetMetaDataInt", "(Ljava/lang/String;)I", bIsOptional);
@@ -206,6 +207,7 @@ jmethodID FJavaWrapper::AndroidThunkJava_KeepScreenOn;
 jmethodID FJavaWrapper::AndroidThunkJava_InitHMDs;
 jmethodID FJavaWrapper::AndroidThunkJava_DismissSplashScreen;
 jmethodID FJavaWrapper::AndroidThunkJava_GetInputDeviceInfo;
+jmethodID FJavaWrapper::AndroidThunkJava_IsGamepadAttached;
 jmethodID FJavaWrapper::AndroidThunkJava_HasMetaDataKey;
 jmethodID FJavaWrapper::AndroidThunkJava_GetMetaDataBoolean;
 jmethodID FJavaWrapper::AndroidThunkJava_GetMetaDataInt;
@@ -344,6 +346,16 @@ bool AndroidThunkCpp_GetInputDeviceInfo(int32 deviceId, FAndroidInputDeviceInfo
 	return false;
 }
 
+bool AndroidThunkCpp_IsGamepadAttached()
+{
+	bool Result = false;
+	if (JNIEnv* Env = FAndroidApplication::GetJavaEnv())
+	{
+		Result = FJavaWrapper::CallBooleanMethod(Env, FJavaWrapper::GameActivityThis, FJavaWrapper::AndroidThunkJava_IsGamepadAttached);
+	}
+	return Result;
+}
+
 bool AndroidThunkCpp_HasMetaDataKey(const FString& Key)
 {
 	bool Result = false;
diff --git a/Engine/Source/Runtime/Launch/Public/Android/AndroidJNI.h b/Engine/Source/Runtime/Launch/Public/Android/AndroidJNI.h
--- a/Engine/Source/Runtime/Launch/Public/Android/AndroidJNI.h
+++ b/Engine/Source/Runtime/Launch/Public/Android/AndroidJNI.h
@@ -30,6 +30,7 @@ public:
 	static jmethodID AndroidThunkJava_InitHMDs;
 	static jmethodID AndroidThunkJava_DismissSplashScreen;
 	static jmethodID AndroidThunkJava_GetInputDeviceInfo;
+	static jmethodID AndroidThunkJava_IsGamepadAttached;
 	static jmethodID AndroidThunkJava_HasMetaDataKey;
 	static jmethodID AndroidThunkJava_GetMetaDataBoolean;
 	static jmethodID AndroidThunkJava_GetMetaDataInt;
diff --git a/Engine/Source/Runtime/OpenGLDrv/Private/Android/AndroidESDeferredOpenGL.cpp b/Engine/Source/Runtime/OpenGLDrv/Private/Android/AndroidESDeferredOpenGL.cpp
--- a/Engine/Source/Runtime/OpenGLDrv/Private/Android/AndroidESDeferredOpenGL.cpp
+++ b/Engine/Source/Runtime/OpenGLDrv/Private/Android/AndroidESDeferredOpenGL.cpp
@@ -594,4 +594,14 @@ void FAndroidAppEntry::PlatformInit()
 	}
 }
 
+void FAndroidAppEntry::ReleaseEGL()
+{
+	AndroidEGL* EGL = AndroidEGL::GetInstance();
+	if (EGL->IsInitialized())
+	{
+		EGL->DestroyBackBuffer();
+		EGL->Terminate();
+	}
+}
+
 #endif
diff --git a/Engine/Source/Runtime/OpenGLDrv/Private/Android/AndroidOpenGL.cpp b/Engine/Source/Runtime/OpenGLDrv/Private/Android/AndroidOpenGL.cpp
--- a/Engine/Source/Runtime/OpenGLDrv/Private/Android/AndroidOpenGL.cpp
+++ b/Engine/Source/Runtime/OpenGLDrv/Private/Android/AndroidOpenGL.cpp
@@ -484,6 +484,8 @@ void FAndroidOpenGL::ProcessExtensions(const FString& ExtensionsString)
 		bSupportsInstancing = true;
 		bSupportsTextureFloat = true;
 		bSupportsTextureHalfFloat = true;
+		bSupportsRGB10A2 = true;
+		bSupportsVertexHalfFloat = true;
 		
 		// According to https://www.khronos.org/registry/gles/extensions/EXT/EXT_color_buffer_float.txt
 		bSupportsColorBufferHalfFloat = (bSupportsColorBufferHalfFloat || bSupportsColorBufferFloat);
@@ -587,13 +589,8 @@ void FAndroidMisc::GetValidTargetPlatforms(TArray<FString>& TargetPlatformNames)
 
 void FAndroidAppEntry::PlatformInit()
 {
-	// Note: UE-33593,
-	// Creating & destroying EGL context breaks vulkan initialization on mali devices.
-	if (!FAndroidMisc::ShouldUseVulkan())
-	{
-		// Create an ES2 EGL here for gpu queries.
-		AndroidEGL::GetInstance()->Init(AndroidEGL::AV_OpenGLES, 2, 0, false);
-	}
+	// create an ES2 EGL here for gpu queries.
+	AndroidEGL::GetInstance()->Init(AndroidEGL::AV_OpenGLES, 2, 0, false);
 }
 
 void FAndroidAppEntry::ReleaseEGL()
diff --git a/Engine/Source/Runtime/OpenGLDrv/Private/Android/AndroidOpenGL.h b/Engine/Source/Runtime/OpenGLDrv/Private/Android/AndroidOpenGL.h
--- a/Engine/Source/Runtime/OpenGLDrv/Private/Android/AndroidOpenGL.h
+++ b/Engine/Source/Runtime/OpenGLDrv/Private/Android/AndroidOpenGL.h
@@ -57,8 +57,6 @@ extern PFNBLITFRAMEBUFFERNVPROC					glBlitFramebufferNV ;
 #define GL_SAMPLES_PASSED_EXT					0x8914
 #define GL_ANY_SAMPLES_PASSED_EXT				0x8C2F
 
-/* from GL_EXT_color_buffer_half_float */
-#define GL_RGBA16F_EXT                          0x881A
 
 typedef void (GL_APIENTRYP PFNGLGENQUERIESEXTPROC) (GLsizei n, GLuint *ids);
 typedef void (GL_APIENTRYP PFNGLDELETEQUERIESEXTPROC) (GLsizei n, const GLuint *ids);
@@ -273,7 +271,7 @@ struct FAndroidOpenGL : public FOpenGLES2
 
 	static FORCEINLINE bool TexStorage2D(GLenum Target, GLint Levels, GLint InternalFormat, GLsizei Width, GLsizei Height, GLenum Format, GLenum Type, uint32 Flags)
 	{
-		if( bUseHalfFloatTexStorage && Type == GL_HALF_FLOAT_OES && (Flags & TexCreate_RenderTargetable) != 0 )
+		if( bUseHalfFloatTexStorage && Type == GetTextureHalfFloatPixelType() && (Flags & TexCreate_RenderTargetable) != 0 )
 		{
 			glTexStorage2D(Target, Levels, InternalFormat, Width, Height);
 			VERIFY_GL(glTexStorage2D)
@@ -428,6 +426,16 @@ struct FAndroidOpenGL : public FOpenGLES2
 	// Adreno doesn't support HALF_FLOAT
 	static FORCEINLINE int32 GetReadHalfFloatPixelsEnum()				{ return GL_FLOAT; }
 
+	static FORCEINLINE GLenum GetTextureHalfFloatPixelType()			
+	{ 
+		return bES30Support ? GL_HALF_FLOAT : GL_HALF_FLOAT_OES;
+	}
+
+	static FORCEINLINE GLenum GetTextureHalfFloatInternalFormat()		
+	{ 
+		return bES30Support ? GL_RGBA16F : GL_RGBA;
+	}
+
 	// Android ES2 shaders have code that allows compile selection of
 	// 32 bpp HDR encoding mode via 'intrinsic_GetHDR32bppEncodeModeES2()'.
 	static FORCEINLINE bool SupportsHDR32bppEncodeModeIntrinsic()		{ return true; }
diff --git a/Engine/Source/Runtime/OpenGLDrv/Private/Android/AndroidOpenGLPrivate.h b/Engine/Source/Runtime/OpenGLDrv/Private/Android/AndroidOpenGLPrivate.h
--- a/Engine/Source/Runtime/OpenGLDrv/Private/Android/AndroidOpenGLPrivate.h
+++ b/Engine/Source/Runtime/OpenGLDrv/Private/Android/AndroidOpenGLPrivate.h
@@ -30,95 +30,82 @@ public:
 private:
 	FAndroidGPUInfo()
 	{
-		if (FAndroidMisc::ShouldUseVulkan())
-		{
-			// hard coded for the time being
-			GPUFamily = "Vulkan Mobile";
-			GLVersion = "3.1.0";
-			bSupportsFloatingPointRenderTargets = true;
-			bSupportsFrameBufferFetch = false;
-			TargetPlatformNames.Add(TEXT("Android_ASTC"));
-			GAndroidGPUInfoReady = true;
-		}
-		else
-		{
-			// this is only valid in the game thread, make sure we are initialized there before being called on other threads!
-			check(IsInGameThread())
+		// this is only valid in the game thread, make sure we are initialized there before being called on other threads!
+		check(IsInGameThread())
 
-			// make sure GL is started so we can get the supported formats
-			AndroidEGL* EGL = AndroidEGL::GetInstance();
+		// make sure GL is started so we can get the supported formats
+		AndroidEGL* EGL = AndroidEGL::GetInstance();
 
-			if (!EGL->IsInitialized())
-			{
-				FAndroidAppEntry::PlatformInit();
+		if (!EGL->IsInitialized())
+		{
+			FAndroidAppEntry::PlatformInit();
 #if PLATFORM_ANDROIDESDEFERRED
-				EGL->InitSurface(false, true);
+			EGL->InitSurface(false, true);
 #endif
-			}
+		}
 #if !PLATFORM_ANDROIDESDEFERRED
-			// Do not create a window surface if the app is for GearVR (use small buffer)
-			bool bCreateSurface = !AndroidThunkCpp_IsGearVRApplication();
-			FPlatformMisc::LowLevelOutputDebugString(TEXT("FAndroidGPUInfo"));
-			EGL->InitSurface(bCreateSurface, bCreateSurface);
+		// Do not create a window surface if the app is for GearVR (use small buffer)
+		bool bCreateSurface = !AndroidThunkCpp_IsGearVRApplication();
+		FPlatformMisc::LowLevelOutputDebugString(TEXT("FAndroidGPUInfo"));
+		EGL->InitSurface(bCreateSurface, bCreateSurface);
 #endif
-			EGL->SetCurrentSharedContext();
+		EGL->SetCurrentSharedContext();
 
-			// get extensions
-			// Process the extension caps directly here, as FOpenGL might not yet be setup
-			// Do not process extensions here, because extension pointers may not be setup
-			const ANSICHAR* GlGetStringOutput = (const ANSICHAR*) glGetString(GL_EXTENSIONS);
-			FString ExtensionsString = GlGetStringOutput;
+		// get extensions
+		// Process the extension caps directly here, as FOpenGL might not yet be setup
+		// Do not process extensions here, because extension pointers may not be setup
+		const ANSICHAR* GlGetStringOutput = (const ANSICHAR*) glGetString(GL_EXTENSIONS);
+		FString ExtensionsString = GlGetStringOutput;
 
-			GPUFamily = (const ANSICHAR*)glGetString(GL_RENDERER);
-			check(!GPUFamily.IsEmpty());
+		GPUFamily = (const ANSICHAR*)glGetString(GL_RENDERER);
+		check(!GPUFamily.IsEmpty());
 
-			GLVersion = (const ANSICHAR*)glGetString(GL_VERSION);
+		GLVersion = (const ANSICHAR*)glGetString(GL_VERSION);
 
-			const bool bES30Support = GLVersion.Contains(TEXT("OpenGL ES 3."));
+		const bool bES30Support = GLVersion.Contains(TEXT("OpenGL ES 3."));
 
 #if PLATFORM_ANDROIDESDEFERRED
-			TargetPlatformNames.Add(TEXT("Android_ESDEFERRED"));
+		TargetPlatformNames.Add(TEXT("Android_ESDEFERRED"));
 #else
-			// highest priority is the per-texture version
-			if (ExtensionsString.Contains(TEXT("GL_KHR_texture_compression_astc_ldr")))
-			{
-				TargetPlatformNames.Add(TEXT("Android_ASTC"));
-			}
-			if (ExtensionsString.Contains(TEXT("GL_NV_texture_compression_s3tc")) || ExtensionsString.Contains(TEXT("GL_EXT_texture_compression_s3tc")))
-			{
-				TargetPlatformNames.Add(TEXT("Android_DXT"));
-			}
-			if (ExtensionsString.Contains(TEXT("GL_ATI_texture_compression_atitc")) || ExtensionsString.Contains(TEXT("GL_AMD_compressed_ATC_texture")))
-			{
-				TargetPlatformNames.Add(TEXT("Android_ATC"));
-			}
-			if (ExtensionsString.Contains(TEXT("GL_IMG_texture_compression_pvrtc")))
-			{
-				TargetPlatformNames.Add(TEXT("Android_PVRTC"));
-			}
-			if (bES30Support)
-			{
-				TargetPlatformNames.Add(TEXT("Android_ETC2"));
-			}
-			if (ExtensionsString.Contains(TEXT("GL_KHR_texture_compression_astc_ldr")))
-			{
-				TargetPlatformNames.Add(TEXT("Android_ASTC"));
-			}
-
-			// all devices support ETC
-			TargetPlatformNames.Add(TEXT("Android_ETC1"));
-
-			// finally, generic Android
-			TargetPlatformNames.Add(TEXT("Android"));
+		// highest priority is the per-texture version
+		if (ExtensionsString.Contains(TEXT("GL_KHR_texture_compression_astc_ldr")))
+		{
+			TargetPlatformNames.Add(TEXT("Android_ASTC"));
+		}
+		if (ExtensionsString.Contains(TEXT("GL_NV_texture_compression_s3tc")) || ExtensionsString.Contains(TEXT("GL_EXT_texture_compression_s3tc")))
+		{
+			TargetPlatformNames.Add(TEXT("Android_DXT"));
+		}
+		if (ExtensionsString.Contains(TEXT("GL_ATI_texture_compression_atitc")) || ExtensionsString.Contains(TEXT("GL_AMD_compressed_ATC_texture")))
+		{
+			TargetPlatformNames.Add(TEXT("Android_ATC"));
+		}
+		if (ExtensionsString.Contains(TEXT("GL_IMG_texture_compression_pvrtc")))
+		{
+			TargetPlatformNames.Add(TEXT("Android_PVRTC"));
+		}
+		if (bES30Support)
+		{
+			TargetPlatformNames.Add(TEXT("Android_ETC2"));
+		}
+		if (ExtensionsString.Contains(TEXT("GL_KHR_texture_compression_astc_ldr")))
+		{
+			TargetPlatformNames.Add(TEXT("Android_ASTC"));
+		}
+
+		// all devices support ETC
+		TargetPlatformNames.Add(TEXT("Android_ETC1"));
+
+		// finally, generic Android
+		TargetPlatformNames.Add(TEXT("Android"));
 
 #endif
-			bSupportsFloatingPointRenderTargets = 
-				ExtensionsString.Contains(TEXT("GL_EXT_color_buffer_half_float")) 
-				// According to https://www.khronos.org/registry/gles/extensions/EXT/EXT_color_buffer_float.txt
-				|| (bES30Support && ExtensionsString.Contains(TEXT("GL_EXT_color_buffer_float")));
+		bSupportsFloatingPointRenderTargets = 
+			ExtensionsString.Contains(TEXT("GL_EXT_color_buffer_half_float")) 
+			// According to https://www.khronos.org/registry/gles/extensions/EXT/EXT_color_buffer_float.txt
+			|| (bES30Support && ExtensionsString.Contains(TEXT("GL_EXT_color_buffer_float")));
 
-			bSupportsFrameBufferFetch = ExtensionsString.Contains(TEXT("GL_EXT_shader_framebuffer_fetch")) || ExtensionsString.Contains(TEXT("GL_NV_shader_framebuffer_fetch")) || ExtensionsString.Contains(TEXT("GL_ARM_shader_framebuffer_fetch"));
-		}
+		bSupportsFrameBufferFetch = ExtensionsString.Contains(TEXT("GL_EXT_shader_framebuffer_fetch")) || ExtensionsString.Contains(TEXT("GL_NV_shader_framebuffer_fetch")) || ExtensionsString.Contains(TEXT("GL_ARM_shader_framebuffer_fetch"));
 		GAndroidGPUInfoReady = true;
 	}
 };
\ No newline at end of file
diff --git a/Engine/Source/Runtime/OpenGLDrv/Private/OpenGLDevice.cpp b/Engine/Source/Runtime/OpenGLDrv/Private/OpenGLDevice.cpp
--- a/Engine/Source/Runtime/OpenGLDrv/Private/OpenGLDevice.cpp
+++ b/Engine/Source/Runtime/OpenGLDrv/Private/OpenGLDevice.cpp
@@ -904,7 +904,9 @@ static void InitRHICapabilitiesForGL()
 		if (FOpenGL::SupportsColorBufferHalfFloat() && FOpenGL::SupportsTextureHalfFloat())
 		{
 #if PLATFORM_ANDROID
-			SetupTextureFormat(PF_FloatRGBA, FOpenGLTextureFormat(GL_RGBA, GL_RGBA, GL_RGBA16F_EXT, GL_RGBA16F_EXT, GL_RGBA, GL_HALF_FLOAT_OES, false, false));
+			GLenum InternalFormatRGBA16 = FOpenGL::GetTextureHalfFloatInternalFormat();
+			GLenum PixelTypeRGBA16 = FOpenGL::GetTextureHalfFloatPixelType();
+			SetupTextureFormat(PF_FloatRGBA, FOpenGLTextureFormat(InternalFormatRGBA16, InternalFormatRGBA16, GL_RGBA, PixelTypeRGBA16, false, false));
 #else
 			SetupTextureFormat(PF_FloatRGBA, FOpenGLTextureFormat(GL_RGBA, GL_RGBA, GL_RGBA, GL_HALF_FLOAT_OES, false, false));
 #endif
diff --git a/Engine/Source/Runtime/OpenGLDrv/Private/OpenGLUAV.cpp b/Engine/Source/Runtime/OpenGLDrv/Private/OpenGLUAV.cpp
--- a/Engine/Source/Runtime/OpenGLDrv/Private/OpenGLUAV.cpp
+++ b/Engine/Source/Runtime/OpenGLDrv/Private/OpenGLUAV.cpp
@@ -110,6 +110,12 @@ FOpenGLVertexBufferUnorderedAccessView::FOpenGLVertexBufferUnorderedAccessView(
 	this->Format = GLFormat.InternalFormat[0];
 }
 
+uint32 FOpenGLVertexBufferUnorderedAccessView::GetBufferSize()
+{
+	FOpenGLVertexBuffer* VertexBuffer = FOpenGLDynamicRHI::ResourceCast(VertexBufferRHI.GetReference());
+	return VertexBufferRHI->GetSize();
+}
+
 FOpenGLVertexBufferUnorderedAccessView::~FOpenGLVertexBufferUnorderedAccessView()
 {
 	if (Resource)
@@ -135,11 +141,24 @@ FShaderResourceViewRHIRef FOpenGLDynamicRHI::RHICreateShaderResourceView(FStruct
 
 void FOpenGLDynamicRHI::RHIClearUAV(FUnorderedAccessViewRHIParamRef UnorderedAccessViewRHI, const uint32* Values)
 {
-#if OPENGL_GL4
 	FOpenGLUnorderedAccessView* Texture = ResourceCast(UnorderedAccessViewRHI);
+
+#if OPENGL_GL4
 	glBindBuffer(GL_TEXTURE_BUFFER, Texture->BufferResource);
 	FOpenGL::ClearBufferData(GL_TEXTURE_BUFFER, Texture->Format, GL_RED_INTEGER, GL_UNSIGNED_INT, Values);
 	GPUProfilingData.RegisterGPUWork(1);
+
+#elif OPENGL_ESDEFERRED
+	glBindBuffer(GL_TEXTURE_BUFFER, Texture->BufferResource);
+	uint32 BufferSize = Texture->GetBufferSize();
+	if (BufferSize > 0)
+	{
+		void* BufferData = FOpenGL::MapBufferRange(GL_TEXTURE_BUFFER, 0, BufferSize, FOpenGLBase::RLM_WriteOnly);
+		uint8 ClearValue = uint8(Values[0] & 0xff);
+		FPlatformMemory::Memset(BufferData, ClearValue, BufferSize);
+		FOpenGL::UnmapBufferRange(GL_TEXTURE_BUFFER, 0, BufferSize);
+		GPUProfilingData.RegisterGPUWork(1);
+	}
 #else
 	UE_LOG(LogRHI, Fatal, TEXT("Only OpenGL4 supports RHIClearUAV."));
 #endif
diff --git a/Engine/Source/Runtime/OpenGLDrv/Public/OpenGL.h b/Engine/Source/Runtime/OpenGLDrv/Public/OpenGL.h
--- a/Engine/Source/Runtime/OpenGLDrv/Public/OpenGL.h
+++ b/Engine/Source/Runtime/OpenGLDrv/Public/OpenGL.h
@@ -39,6 +39,10 @@ struct FPlatformOpenGLContext;
 #define GL_HALF_FLOAT 0x140B
 #endif
 
+#ifndef GL_RGBA16F
+#define GL_RGBA16F    0x881A
+#endif
+
 // Base static class
 class FOpenGLBase
 {
@@ -161,6 +165,8 @@ public:
 	static FORCEINLINE GLenum GetDepthFormat()							{ return GL_DEPTH_COMPONENT16; }
 	static FORCEINLINE GLenum GetShadowDepthFormat()					{ return GL_DEPTH_COMPONENT16; }
 	static FORCEINLINE GLenum GetVertexHalfFloatFormat()				{ return GL_HALF_FLOAT; }
+	static FORCEINLINE GLenum GetTextureHalfFloatPixelType()			{ return GL_HALF_FLOAT; }
+	static FORCEINLINE GLenum GetTextureHalfFloatInternalFormat()		{ return GL_RGBA16F; }
 
 	static FORCEINLINE GLint GetMaxTextureImageUnits()			{ check(MaxTextureImageUnits != -1); return MaxTextureImageUnits; }
 	static FORCEINLINE GLint GetMaxVertexTextureImageUnits()	{ check(MaxVertexTextureImageUnits != -1); return MaxVertexTextureImageUnits; }
diff --git a/Engine/Source/Runtime/OpenGLDrv/Public/OpenGLES2.h b/Engine/Source/Runtime/OpenGLDrv/Public/OpenGLES2.h
--- a/Engine/Source/Runtime/OpenGLDrv/Public/OpenGLES2.h
+++ b/Engine/Source/Runtime/OpenGLDrv/Public/OpenGLES2.h
@@ -150,6 +150,10 @@ struct FOpenGLES2 : public FOpenGLBase
 	static FORCEINLINE int32 GetReadHalfFloatPixelsEnum()				{ return GL_HALF_FLOAT_OES; }
 
 	static FORCEINLINE GLenum GetVertexHalfFloatFormat()				{ return GL_HALF_FLOAT_OES; }
+	static FORCEINLINE GLenum GetTextureHalfFloatPixelType()			{ return GL_HALF_FLOAT_OES; }
+	static FORCEINLINE GLenum GetTextureHalfFloatInternalFormat()		{ return GL_RGBA; }
+
+	
 	static FORCEINLINE bool NeedsVertexAttribRemapTable()				{ return bNeedsVertexAttribRemap; }
 
 	// On iOS both glMapBufferOES() and glBufferSubData() for immediate vertex and index data
@@ -196,27 +200,36 @@ struct FOpenGLES2 : public FOpenGLBase
 
 	static FORCEINLINE void LabelObject(GLenum Type, GLuint Object, const ANSICHAR* Name)
 	{
-		// @todo-mobile: Tie this to be optional; 
-		glLabelObjectEXT(Type, Object, 0, Name);
+		if (glLabelObjectEXT != nullptr)
+		{
+			glLabelObjectEXT(Type, Object, 0, Name);
+		}
 	}
 
 	static FORCEINLINE GLsizei GetLabelObject(GLenum Type, GLuint Object, GLsizei BufferSize, ANSICHAR* OutName)
 	{
 		GLsizei Length = 0;
-		glGetObjectLabelEXT(Type, Object, BufferSize, &Length, OutName);
+		if (glGetObjectLabelEXT != nullptr)
+		{
+			glGetObjectLabelEXT(Type, Object, BufferSize, &Length, OutName);
+		}
 		return Length;
 	}
 
 	static FORCEINLINE void PushGroupMarker(const ANSICHAR* Name)
 	{
-		// @todo-mobile: Tie this to be optional; 
-		glPushGroupMarkerEXT(0, Name);
+		if (glPushGroupMarkerEXT != nullptr)
+		{
+			glPushGroupMarkerEXT(0, Name);
+		}
 	}
 
 	static FORCEINLINE void PopGroupMarker()
 	{
-		// @todo-mobile: Tie this to be optional; 
-		glPopGroupMarkerEXT();
+		if (glPopGroupMarkerEXT != nullptr)
+		{
+			glPopGroupMarkerEXT();
+		}
 	}
 
 	// Required
@@ -728,9 +741,6 @@ public:
 #ifndef GL_RGBA16
 #define GL_RGBA16 0x805B
 #endif
-#ifndef GL_RGBA16F
-#define GL_RGBA16F 0x881A
-#endif
 #ifndef GL_RGBA32F
 #define GL_RGBA32F 0x8814
 #endif
diff --git a/Engine/Source/Runtime/OpenGLDrv/Public/OpenGLResources.h b/Engine/Source/Runtime/OpenGLDrv/Public/OpenGLResources.h
--- a/Engine/Source/Runtime/OpenGLDrv/Public/OpenGLResources.h
+++ b/Engine/Source/Runtime/OpenGLDrv/Public/OpenGLResources.h
@@ -1281,6 +1281,11 @@ public:
 	GLuint	Resource;
 	GLuint  BufferResource;
 	GLenum	Format;
+
+	virtual uint32 GetBufferSize()
+	{
+		return 0;
+	}
 };
 
 class FOpenGLTextureUnorderedAccessView : public FOpenGLUnorderedAccessView
@@ -1306,6 +1311,8 @@ public:
 	FVertexBufferRHIRef VertexBufferRHI; // to keep the vertex buffer alive
 
 	FOpenGLDynamicRHI* OpenGLRHI;
+
+	virtual uint32 GetBufferSize() override;
 };
 
 class FOpenGLShaderResourceView : public FRHIShaderResourceView
diff --git a/Engine/Source/Runtime/Renderer/Private/GlobalDistanceField.cpp b/Engine/Source/Runtime/Renderer/Private/GlobalDistanceField.cpp
--- a/Engine/Source/Runtime/Renderer/Private/GlobalDistanceField.cpp
+++ b/Engine/Source/Runtime/Renderer/Private/GlobalDistanceField.cpp
@@ -1002,13 +1002,13 @@ void FViewInfo::SetupDefaultGlobalDistanceFieldUniformBufferParameters(FViewUnif
 	ViewUniformShaderParameters.GlobalVolumeTexelSize_UB = 0.0f;
 	ViewUniformShaderParameters.MaxGlobalDistance_UB = 0.0f;
 
-	ViewUniformShaderParameters.GlobalDistanceFieldTexture0_UB = GBlackVolumeTexture->TextureRHI.GetReference();
+	ViewUniformShaderParameters.GlobalDistanceFieldTexture0_UB = OrBlack3DIfNull(GBlackVolumeTexture->TextureRHI.GetReference());
 	ViewUniformShaderParameters.GlobalDistanceFieldSampler0_UB = TStaticSamplerState<SF_Bilinear, AM_Wrap, AM_Wrap, AM_Wrap>::GetRHI();
-	ViewUniformShaderParameters.GlobalDistanceFieldTexture1_UB = GBlackVolumeTexture->TextureRHI.GetReference();
+	ViewUniformShaderParameters.GlobalDistanceFieldTexture1_UB = OrBlack3DIfNull(GBlackVolumeTexture->TextureRHI.GetReference());
 	ViewUniformShaderParameters.GlobalDistanceFieldSampler1_UB = TStaticSamplerState<SF_Bilinear, AM_Wrap, AM_Wrap, AM_Wrap>::GetRHI();
-	ViewUniformShaderParameters.GlobalDistanceFieldTexture2_UB = GBlackVolumeTexture->TextureRHI.GetReference();
+	ViewUniformShaderParameters.GlobalDistanceFieldTexture2_UB = OrBlack3DIfNull(GBlackVolumeTexture->TextureRHI.GetReference());
 	ViewUniformShaderParameters.GlobalDistanceFieldSampler2_UB = TStaticSamplerState<SF_Bilinear, AM_Wrap, AM_Wrap, AM_Wrap>::GetRHI();
-	ViewUniformShaderParameters.GlobalDistanceFieldTexture3_UB = GBlackVolumeTexture->TextureRHI.GetReference();
+	ViewUniformShaderParameters.GlobalDistanceFieldTexture3_UB = OrBlack3DIfNull(GBlackVolumeTexture->TextureRHI.GetReference());
 	ViewUniformShaderParameters.GlobalDistanceFieldSampler3_UB = TStaticSamplerState<SF_Bilinear, AM_Wrap, AM_Wrap, AM_Wrap>::GetRHI();
 }
 
diff --git a/Engine/Source/Runtime/Renderer/Private/SceneCaptureRendering.cpp b/Engine/Source/Runtime/Renderer/Private/SceneCaptureRendering.cpp
--- a/Engine/Source/Runtime/Renderer/Private/SceneCaptureRendering.cpp
+++ b/Engine/Source/Runtime/Renderer/Private/SceneCaptureRendering.cpp
@@ -176,6 +176,113 @@ void FDeferredShadingSceneRenderer::CopySceneCaptureComponentToTarget(FRHIComman
 	}
 }
 
+
+/**
+* Shader set for decoding a mosaic or RGBE encoded HDR image as part of a copy operation.
+*/
+template<bool bDemosaic>
+class FHDRDecodePS : public FGlobalShader
+{
+	DECLARE_SHADER_TYPE(FHDRDecodePS, Global)
+public:
+
+	static bool ShouldCache(EShaderPlatform Platform) { return IsES2Platform(Platform); }
+
+	FHDRDecodePS(const ShaderMetaType::CompiledShaderInitializerType& Initializer) :
+		FGlobalShader(Initializer)
+	{
+		InTexture.Bind(Initializer.ParameterMap, TEXT("InTexture"), SPF_Mandatory);
+		InTextureSampler.Bind(Initializer.ParameterMap, TEXT("InTextureSampler"));
+	}
+	FHDRDecodePS() {}
+
+	static void ModifyCompilationEnvironment(EShaderPlatform Platform, FShaderCompilerEnvironment& OutEnvironment)
+	{
+		OutEnvironment.SetDefine(TEXT("DECODING_MOSAIC"), bDemosaic ? 1u : 0u);
+	}
+
+	void SetParameters(FRHICommandList& RHICmdList, const FViewInfo& View, FSamplerStateRHIParamRef SamplerStateRHI, FTextureRHIParamRef TextureRHI)
+	{
+		FGlobalShader::SetParameters(RHICmdList, GetPixelShader(), View);
+		SetTextureParameter(RHICmdList, GetPixelShader(), InTexture, InTextureSampler, SamplerStateRHI, TextureRHI);
+	}
+
+	virtual bool Serialize(FArchive& Ar) override
+	{
+		bool bShaderHasOutdatedParameters = FGlobalShader::Serialize(Ar);
+		Ar << InTexture;
+		Ar << InTextureSampler;
+		return bShaderHasOutdatedParameters;
+	}
+
+private:
+	FShaderResourceParameter InTexture;
+	FShaderResourceParameter InTextureSampler;
+};
+
+/**
+* A vertex shader for rendering a textured screen element.
+*/
+template<bool bDemosaic>
+class FHDRDecodeVS : public FGlobalShader
+{
+	DECLARE_SHADER_TYPE(FHDRDecodeVS, Global)
+public:
+
+	static bool ShouldCache(EShaderPlatform Platform) { return IsES2Platform(Platform); }
+
+	FHDRDecodeVS(const ShaderMetaType::CompiledShaderInitializerType& Initializer) :
+		FGlobalShader(Initializer)
+	{
+		InvTexSizeParameter.Bind(Initializer.ParameterMap, TEXT("InvTexSize"));
+	}
+	FHDRDecodeVS() {}
+	
+	static void ModifyCompilationEnvironment(EShaderPlatform Platform, FShaderCompilerEnvironment& OutEnvironment)
+	{
+		OutEnvironment.SetDefine(TEXT("DECODING_MOSAIC"), bDemosaic ? 1u : 0u);
+	}
+
+	void SetParameters(FRHICommandList& RHICmdList, const FViewInfo& View, const FIntPoint& SourceTexSize)
+	{
+		FGlobalShader::SetParameters(RHICmdList, GetVertexShader(), View);
+		if(InvTexSizeParameter.IsBound())
+		{
+			FVector2D InvTexSize(1.0f / SourceTexSize.X, 1.0f / SourceTexSize.Y);
+			SetShaderValue(RHICmdList, GetVertexShader(), InvTexSizeParameter, InvTexSize);
+		}
+	}
+
+	virtual bool Serialize(FArchive& Ar) override
+	{
+		bool bShaderHasOutdatedParameters = FGlobalShader::Serialize(Ar);
+		Ar << InvTexSizeParameter;
+		return bShaderHasOutdatedParameters;
+	}
+
+	FShaderParameter InvTexSizeParameter;
+};
+
+IMPLEMENT_SHADER_TYPE(template<>, FHDRDecodePS<false>, TEXT("Decode32bppHDR"), TEXT("MainPS"), SF_Pixel);
+IMPLEMENT_SHADER_TYPE(template<>, FHDRDecodePS<true>, TEXT("Decode32bppHDR"), TEXT("MainPS"), SF_Pixel);
+
+IMPLEMENT_SHADER_TYPE(template<>, FHDRDecodeVS<false>, TEXT("Decode32bppHDR"), TEXT("MainVS"), SF_Vertex);
+IMPLEMENT_SHADER_TYPE(template<>, FHDRDecodeVS<true>, TEXT("Decode32bppHDR"), TEXT("MainVS"), SF_Vertex);
+
+template <bool bDemosaic>
+static FShader* SetCaptureToTargetShaders(FRHICommandListImmediate& RHICmdList, FViewInfo& View, const FIntPoint& SourceTexSize, FTextureRHIParamRef SourceTextureRHI)
+{
+	TShaderMapRef<FHDRDecodeVS<bDemosaic>> VertexShader(View.ShaderMap);
+	TShaderMapRef<FHDRDecodePS<bDemosaic>> PixelShader(View.ShaderMap);
+	static FGlobalBoundShaderState BoundShaderState;
+	SetGlobalBoundShaderState(RHICmdList, View.GetFeatureLevel(), BoundShaderState, GFilterVertexDeclaration.VertexDeclarationRHI, *VertexShader, *PixelShader);
+
+	VertexShader->SetParameters(RHICmdList, View, SourceTexSize);
+	PixelShader->SetParameters(RHICmdList, View, TStaticSamplerState<SF_Point>::GetRHI(), SourceTextureRHI);
+
+	return *VertexShader;
+}
+
 // Copies into render target, optionally flipping it in the Y-axis
 static void CopyCaptureToTarget(FRHICommandListImmediate& RHICmdList, const FRenderTarget* Target, const FIntPoint& TargetSize, FViewInfo& View, const FIntRect& ViewRect, FTextureRHIParamRef SourceTextureRHI, bool bNeedsFlippedRenderTarget)
 {
@@ -187,13 +294,17 @@ static void CopyCaptureToTarget(FRHICommandListImmediate& RHICmdList, const FRen
 	RHICmdList.SetDepthStencilState(TStaticDepthStencilState<false, CF_Always>::GetRHI());
 	RHICmdList.SetBlendState(TStaticBlendState<>::GetRHI());
 
-	TShaderMapRef<FScreenVS> VertexShader(View.ShaderMap);
-	TShaderMapRef<FScreenPS> PixelShader(View.ShaderMap);
-	static FGlobalBoundShaderState BoundShaderState;
-	SetGlobalBoundShaderState(RHICmdList, View.GetFeatureLevel(), BoundShaderState, GFilterVertexDeclaration.VertexDeclarationRHI, *VertexShader, *PixelShader);
-
-	VertexShader->SetParameters(RHICmdList, View);
-	PixelShader->SetParameters(RHICmdList, TStaticSamplerState<SF_Point>::GetRHI(), SourceTextureRHI);
+	const bool bUsingDemosaic = IsMobileHDRMosaic();
+	FShader* VertexShader;
+	FIntPoint SourceTexSize = bNeedsFlippedRenderTarget ? TargetSize : FSceneRenderTargets::Get(RHICmdList).GetBufferSizeXY();
+	if (bUsingDemosaic)
+	{
+		VertexShader = SetCaptureToTargetShaders<true>(RHICmdList, View, SourceTexSize, SourceTextureRHI);
+	}
+	else
+	{
+		VertexShader = SetCaptureToTargetShaders<false>(RHICmdList, View, SourceTexSize, SourceTextureRHI);
+	}
 
 	if (bNeedsFlippedRenderTarget)
 	{
@@ -204,8 +315,8 @@ static void CopyCaptureToTarget(FRHICommandListImmediate& RHICmdList, const FRen
 			ViewRect.Min.X, ViewRect.Height() - ViewRect.Min.Y,
 			ViewRect.Width(), -ViewRect.Height(),
 			TargetSize,
-			TargetSize,
-			*VertexShader,
+			SourceTexSize,
+			VertexShader,
 			EDRF_UseTriangleOptimization);
 	}
 	else
@@ -217,8 +328,8 @@ static void CopyCaptureToTarget(FRHICommandListImmediate& RHICmdList, const FRen
 			ViewRect.Min.X, ViewRect.Min.Y,
 			ViewRect.Width(), ViewRect.Height(),
 			TargetSize,
-			FSceneRenderTargets::Get(RHICmdList).GetBufferSizeXY(),
-			*VertexShader,
+			SourceTexSize,
+			VertexShader,
 			EDRF_UseTriangleOptimization);
 	}
 }
diff --git a/Engine/Source/Runtime/VulkanRHI/Private/VulkanCommandWrappers.h b/Engine/Source/Runtime/VulkanRHI/Private/VulkanCommandWrappers.h
--- a/Engine/Source/Runtime/VulkanRHI/Private/VulkanCommandWrappers.h
+++ b/Engine/Source/Runtime/VulkanRHI/Private/VulkanCommandWrappers.h
@@ -18,6 +18,8 @@ namespace VulkanRHI
 	void PrintResult(VkResult Result);
 	void PrintResultAndNamedHandle(VkResult Result, const TCHAR* HandleName, void* Handle);
 	void PrintResultAndPointer(VkResult Result, void* Handle);
+	void PrintResultAndNamedHandle(VkResult Result, const TCHAR* HandleName, uint64 Handle);
+	void PrintResultAndPointer(VkResult Result, uint64 Handle);
 	void DumpPhysicalDeviceProperties(VkPhysicalDeviceMemoryProperties* Properties);
 	void DumpAllocateMemory(VkDevice Device, const VkMemoryAllocateInfo* AllocateInfo, VkDeviceMemory* Memory);
 	void DumpMemoryRequirements(VkMemoryRequirements* MemoryRequirements);
@@ -121,11 +123,11 @@ namespace VulkanRHI
 	#define DumpCmdClearDepthStencilImage(c, i, il, ds, rc, r)
 #endif
 
-	static FORCEINLINE_DEBUGGABLE VkResult  vkCreateInstance(const VkInstanceCreateInfo* CreateInfo, const VkAllocationCallbacks* Allocator, VkInstance* Instance)
+	FORCEINLINE_DEBUGGABLE VkResult  vkCreateInstance(const VkInstanceCreateInfo* CreateInfo, const VkAllocationCallbacks* Allocator, VkInstance* Instance)
 	{
 		DumpCreateInstance(CreateInfo, Instance);
 
-		VkResult Result = ::vkCreateInstance(CreateInfo, Allocator, Instance);
+		VkResult Result = VULKANAPINAMESPACE::vkCreateInstance(CreateInfo, Allocator, Instance);
 
 		PrintResultAndNamedHandle(Result, TEXT("Instance"), *Instance);
 		return Result;
@@ -135,14 +137,14 @@ namespace VulkanRHI
 	{
 		PrintfBegin(FString::Printf(TEXT("vkDestroyInstance(Instance=%p)"), Instance));
 
-		::vkDestroyInstance(Instance, Allocator);
+		VULKANAPINAMESPACE::vkDestroyInstance(Instance, Allocator);
 	}
 
 	static FORCEINLINE_DEBUGGABLE VkResult  vkEnumeratePhysicalDevices(VkInstance Instance, uint32* PhysicalDeviceCount, VkPhysicalDevice* PhysicalDevices)
 	{
 		PrintfBegin(FString::Printf(TEXT("vkEnumeratePhysicalDevices(Instance=%p, Count=%p, Devices=%p)"), Instance, PhysicalDeviceCount, PhysicalDevices));
 
-		VkResult Result = ::vkEnumeratePhysicalDevices(Instance, PhysicalDeviceCount, PhysicalDevices);
+		VkResult Result = VULKANAPINAMESPACE::vkEnumeratePhysicalDevices(Instance, PhysicalDeviceCount, PhysicalDevices);
 
 		DumpEnumeratePhysicalDevicesEpilog(PhysicalDeviceCount, PhysicalDevices);
 
@@ -152,7 +154,7 @@ namespace VulkanRHI
 	static FORCEINLINE_DEBUGGABLE void  vkGetPhysicalDeviceFeatures(VkPhysicalDevice PhysicalDevice, VkPhysicalDeviceFeatures* Features)
 	{
 		DumpGetPhysicalDeviceFeatures(PhysicalDevice, Features);
-		::vkGetPhysicalDeviceFeatures(PhysicalDevice, Features);
+		VULKANAPINAMESPACE::vkGetPhysicalDeviceFeatures(PhysicalDevice, Features);
 		DumpPhysicalDeviceFeatures(Features);
 	}
 
@@ -160,7 +162,7 @@ namespace VulkanRHI
 	{
 		PrintfBegin(FString::Printf(TEXT("vkGetPhysicalDeviceFormatProperties(PhysicalDevice=%p, Format=%d, FormatProperties=%p)[...]"), PhysicalDevice, (int32)Format, FormatProperties));
 
-		::vkGetPhysicalDeviceFormatProperties(PhysicalDevice, Format, FormatProperties);
+		VULKANAPINAMESPACE::vkGetPhysicalDeviceFormatProperties(PhysicalDevice, Format, FormatProperties);
 	}
 
 #if 0
@@ -177,21 +179,21 @@ namespace VulkanRHI
 	{
 		PrintfBegin(FString::Printf(TEXT("vkGetPhysicalDeviceProperties(PhysicalDevice=%p, Properties=%p)[...]"), PhysicalDevice, Properties));
 
-		::vkGetPhysicalDeviceProperties(PhysicalDevice, Properties);
+		VULKANAPINAMESPACE::vkGetPhysicalDeviceProperties(PhysicalDevice, Properties);
 	}
 
 	static FORCEINLINE_DEBUGGABLE void  vkGetPhysicalDeviceQueueFamilyProperties(VkPhysicalDevice PhysicalDevice, uint32* QueueFamilyPropertyCount, VkQueueFamilyProperties* QueueFamilyProperties)
 	{
 		PrintfBeginResult(FString::Printf(TEXT("vkGetPhysicalDeviceQueueFamilyProperties(PhysicalDevice=%p, QueueFamilyPropertyCount=%p, QueueFamilyProperties=%p)[...]"), PhysicalDevice, QueueFamilyPropertyCount, QueueFamilyProperties));
 
-		::vkGetPhysicalDeviceQueueFamilyProperties(PhysicalDevice, QueueFamilyPropertyCount, QueueFamilyProperties);
+		VULKANAPINAMESPACE::vkGetPhysicalDeviceQueueFamilyProperties(PhysicalDevice, QueueFamilyPropertyCount, QueueFamilyProperties);
 	}
 
 	static FORCEINLINE_DEBUGGABLE void  vkGetPhysicalDeviceMemoryProperties(VkPhysicalDevice PhysicalDevice, VkPhysicalDeviceMemoryProperties* MemoryProperties)
 	{
 		PrintfBegin(FString::Printf(TEXT("vkGetPhysicalDeviceMemoryProperties(OutProp=%p)[...]"), MemoryProperties));
 
-		::vkGetPhysicalDeviceMemoryProperties(PhysicalDevice, MemoryProperties);
+		VULKANAPINAMESPACE::vkGetPhysicalDeviceMemoryProperties(PhysicalDevice, MemoryProperties);
 
 		DumpPhysicalDeviceProperties(MemoryProperties);
 	}
@@ -200,8 +202,8 @@ namespace VulkanRHI
 	{
 		PrintfBeginResult(FString::Printf(TEXT("vkGetInstanceProcAddr(Instance=%p, Name=%s)[...]"), Instance, ANSI_TO_TCHAR(Name)));
 
-		PFN_vkVoidFunction Function = ::vkGetInstanceProcAddr(Instance, Name);
-		PrintResultAndPointer(VK_SUCCESS, Function);
+		PFN_vkVoidFunction Function = VULKANAPINAMESPACE::vkGetInstanceProcAddr(Instance, Name);
+		PrintResultAndPointer(VK_SUCCESS, (void*)Function);
 		return Function;
 	}
 
@@ -209,8 +211,8 @@ namespace VulkanRHI
 	{
 		PrintfBeginResult(FString::Printf(TEXT("vkGetDeviceProcAddr(Device=%p, Name=%s)[...]"), Device, ANSI_TO_TCHAR(Name)));
 
-		PFN_vkVoidFunction Function = ::vkGetDeviceProcAddr(Device, Name);
-		PrintResultAndPointer(VK_SUCCESS, Function);
+		PFN_vkVoidFunction Function = VULKANAPINAMESPACE::vkGetDeviceProcAddr(Device, Name);
+		PrintResultAndPointer(VK_SUCCESS, (void*)Function);
 		return Function;
 	}
 
@@ -218,7 +220,7 @@ namespace VulkanRHI
 	{
 		DumpCreateDevice(PhysicalDevice, CreateInfo, Device);
 
-		VkResult Result = ::vkCreateDevice(PhysicalDevice, CreateInfo, Allocator, Device);
+		VkResult Result = VULKANAPINAMESPACE::vkCreateDevice(PhysicalDevice, CreateInfo, Allocator, Device);
 
 		PrintResultAndNamedHandle(Result, TEXT("Device"), *Device);
 		return Result;
@@ -228,14 +230,14 @@ namespace VulkanRHI
 	{
 		PrintfBegin(FString::Printf(TEXT("vkDestroyDevice(Device=%p)"), Device));
 
-		::vkDestroyDevice(Device, Allocator);
+		VULKANAPINAMESPACE::vkDestroyDevice(Device, Allocator);
 	}
 
 	static FORCEINLINE_DEBUGGABLE VkResult  vkEnumerateInstanceExtensionProperties(const char* LayerName, uint32* PropertyCount, VkExtensionProperties* Properties)
 	{
 		PrintfBeginResult(FString::Printf(TEXT("vkEnumerateInstanceExtensionProperties(LayerName=%s, PropertyCount=%p, Properties=%p)[...]"), ANSI_TO_TCHAR(LayerName), PropertyCount, Properties));
 
-		VkResult Result = ::vkEnumerateInstanceExtensionProperties(LayerName, PropertyCount, Properties);
+		VkResult Result = VULKANAPINAMESPACE::vkEnumerateInstanceExtensionProperties(LayerName, PropertyCount, Properties);
 		PrintResultAndPointer(Result, (void*)(uint64)PropertyCount);
 		return Result;
 	}
@@ -244,7 +246,7 @@ namespace VulkanRHI
 	{
 		PrintfBeginResult(FString::Printf(TEXT("vkEnumerateDeviceExtensionProperties(Device=%p, LayerName=%s, PropertyCount=%p, Properties=%p)[...]"), PhysicalDevice, ANSI_TO_TCHAR(LayerName), PropertyCount, Properties));
 
-		VkResult Result = ::vkEnumerateDeviceExtensionProperties(PhysicalDevice, LayerName, PropertyCount, Properties);
+		VkResult Result = VULKANAPINAMESPACE::vkEnumerateDeviceExtensionProperties(PhysicalDevice, LayerName, PropertyCount, Properties);
 		PrintResultAndPointer(Result, (void*)(uint64)PropertyCount);
 		return Result;
 	}
@@ -253,7 +255,7 @@ namespace VulkanRHI
 	{
 		PrintfBeginResult(FString::Printf(TEXT("vkEnumerateInstanceLayerProperties(PropertyCount=%p, Properties=%p)[...]"), PropertyCount, Properties));
 
-		VkResult Result = ::vkEnumerateInstanceLayerProperties(PropertyCount, Properties);
+		VkResult Result = VULKANAPINAMESPACE::vkEnumerateInstanceLayerProperties(PropertyCount, Properties);
 		PrintResultAndPointer(Result, (void*)(uint64)PropertyCount);
 		return Result;
 	}
@@ -262,7 +264,7 @@ namespace VulkanRHI
 	{
 		PrintfBeginResult(FString::Printf(TEXT("vkEnumerateDeviceLayerProperties(Device=%p, PropertyCount=%p, Properties=%p)[...]"), PhysicalDevice, PropertyCount, Properties));
 
-		VkResult Result = ::vkEnumerateDeviceLayerProperties(PhysicalDevice, PropertyCount, Properties);
+		VkResult Result = VULKANAPINAMESPACE::vkEnumerateDeviceLayerProperties(PhysicalDevice, PropertyCount, Properties);
 		PrintResultAndPointer(Result, (void*)(uint64)PropertyCount);
 		return Result;
 	}
@@ -271,7 +273,7 @@ namespace VulkanRHI
 	{
 		DevicePrintfBeginResult(Device, FString::Printf(TEXT("vkGetDeviceQueue(QueueFamilyIndex=%d, QueueIndex=%d, OutQueue=%p)"), QueueFamilyIndex, QueueIndex, Queue));
 
-		::vkGetDeviceQueue(Device, QueueFamilyIndex, QueueIndex, Queue);
+		VULKANAPINAMESPACE::vkGetDeviceQueue(Device, QueueFamilyIndex, QueueIndex, Queue);
 
 		PrintResultAndNamedHandle(VK_SUCCESS, TEXT("Queue"), *Queue);
 	}
@@ -280,7 +282,7 @@ namespace VulkanRHI
 	{
 		DumpQueueSubmit(Queue, SubmitCount, Submits, Fence);
 
-		VkResult Result = ::vkQueueSubmit(Queue, SubmitCount, Submits, Fence);
+		VkResult Result = VULKANAPINAMESPACE::vkQueueSubmit(Queue, SubmitCount, Submits, Fence);
 
 		PrintResult(Result);
 		return Result;
@@ -289,7 +291,7 @@ namespace VulkanRHI
 	{
 		PrintfBeginResult(FString::Printf(TEXT("vkQueueWaitIdle(Queue=%p)"), Queue));
 
-		VkResult Result = ::vkQueueWaitIdle(Queue);
+		VkResult Result = VULKANAPINAMESPACE::vkQueueWaitIdle(Queue);
 
 		PrintResult(Result);
 		return Result;
@@ -299,7 +301,7 @@ namespace VulkanRHI
 	{
 		DevicePrintfBeginResult(Device, FString::Printf(TEXT("vkDeviceWaitIdle()")));
 
-		VkResult Result = ::vkDeviceWaitIdle(Device);
+		VkResult Result = VULKANAPINAMESPACE::vkDeviceWaitIdle(Device);
 
 		PrintResult(Result);
 		return Result;
@@ -309,7 +311,7 @@ namespace VulkanRHI
 	{
 		DumpAllocateMemory(Device, AllocateInfo, Memory);
 
-		VkResult Result = ::vkAllocateMemory(Device, AllocateInfo, Allocator, Memory);
+		VkResult Result = VULKANAPINAMESPACE::vkAllocateMemory(Device, AllocateInfo, Allocator, Memory);
 
 		PrintResultAndNamedHandle(Result, TEXT("DevMem"), *Memory);
 		return Result;
@@ -319,14 +321,14 @@ namespace VulkanRHI
 	{
 		DevicePrintfBegin(Device, FString::Printf(TEXT("vkFreeMemory(DevMem=%p)"), Memory));
 
-		::vkFreeMemory(Device, Memory, Allocator);
+		VULKANAPINAMESPACE::vkFreeMemory(Device, Memory, Allocator);
 	}
 
 	static FORCEINLINE_DEBUGGABLE VkResult  vkMapMemory(VkDevice Device, VkDeviceMemory Memory, VkDeviceSize Offset, VkDeviceSize Size, VkMemoryMapFlags Flags, void** Data)
 	{
 		DevicePrintfBeginResult(Device, FString::Printf(TEXT("vkMapMemory(DevMem=%p, Off=%d, Size=%d, Flags=0x%x, OutData=%p)"), Memory, (uint32)Offset, (uint32)Size, Flags, Data));
 
-		VkResult Result = ::vkMapMemory(Device, Memory, Offset, Size, Flags, Data);
+		VkResult Result = VULKANAPINAMESPACE::vkMapMemory(Device, Memory, Offset, Size, Flags, Data);
 
 		PrintResultAndPointer(Result, *Data);
 		return Result;
@@ -335,7 +337,7 @@ namespace VulkanRHI
 	static FORCEINLINE_DEBUGGABLE void  vkUnmapMemory(VkDevice Device, VkDeviceMemory Memory)
 	{
 		DevicePrintfBegin(Device, FString::Printf(TEXT("vkUnmapMemory(DevMem=%p)"), Memory));
-		::vkUnmapMemory(Device, Memory);
+		VULKANAPINAMESPACE::vkUnmapMemory(Device, Memory);
 	}
 
 	static FORCEINLINE_DEBUGGABLE VkResult  vkFlushMappedMemoryRanges(VkDevice Device, uint32 MemoryRangeCount, const VkMappedMemoryRange* MemoryRanges)
@@ -343,7 +345,7 @@ namespace VulkanRHI
 		DumpMappedMemoryRanges(MemoryRangeCount, MemoryRanges);
 		DevicePrintfBeginResult(Device, FString::Printf(TEXT("vkFlushMappedMemoryRanges(Count=%d, Ranges=%p)"), MemoryRangeCount, MemoryRanges));
 
-		VkResult Result = ::vkFlushMappedMemoryRanges(Device, MemoryRangeCount, MemoryRanges);
+		VkResult Result = VULKANAPINAMESPACE::vkFlushMappedMemoryRanges(Device, MemoryRangeCount, MemoryRanges);
 
 		PrintResult(Result);
 		return Result;
@@ -354,7 +356,7 @@ namespace VulkanRHI
 		DumpMappedMemoryRanges(MemoryRangeCount, MemoryRanges);
 		DevicePrintfBeginResult(Device, FString::Printf(TEXT("vkInvalidateMappedMemoryRanges(Count=%d, Ranges=%p)"), MemoryRangeCount, MemoryRanges));
 
-		VkResult Result = ::vkInvalidateMappedMemoryRanges(Device, MemoryRangeCount, MemoryRanges);
+		VkResult Result = VULKANAPINAMESPACE::vkInvalidateMappedMemoryRanges(Device, MemoryRangeCount, MemoryRanges);
 
 		PrintResult(Result);
 		return Result;
@@ -369,7 +371,7 @@ namespace VulkanRHI
 	{
 		DevicePrintfBeginResult(Device, FString::Printf(TEXT("vkBindBufferMemory(Buffer=%p, DevMem=%p, MemOff=%d)"), Buffer, Memory, (uint32)MemoryOffset));
 
-		VkResult Result = ::vkBindBufferMemory(Device, Buffer, Memory, MemoryOffset);
+		VkResult Result = VULKANAPINAMESPACE::vkBindBufferMemory(Device, Buffer, Memory, MemoryOffset);
 
 		PrintResult(Result);
 		return Result;
@@ -379,7 +381,7 @@ namespace VulkanRHI
 	{
 		DevicePrintfBeginResult(Device, FString::Printf(TEXT("vkBindImageMemory(Image=%p, DevMem=%p, MemOff=%d)"), Image, Memory, (uint32)MemoryOffset));
 
-		VkResult Result = ::vkBindImageMemory(Device, Image, Memory, MemoryOffset);
+		VkResult Result = VULKANAPINAMESPACE::vkBindImageMemory(Device, Image, Memory, MemoryOffset);
 
 		PrintResult(Result);
 		return Result;
@@ -389,7 +391,7 @@ namespace VulkanRHI
 	{
 		DevicePrintfBeginResult(Device, FString::Printf(TEXT("vkGetBufferMemoryRequirements(Buffer=%p, OutReq=%p)"), Buffer, MemoryRequirements));
 
-		::vkGetBufferMemoryRequirements(Device, Buffer, MemoryRequirements);
+		VULKANAPINAMESPACE::vkGetBufferMemoryRequirements(Device, Buffer, MemoryRequirements);
 
 		DumpMemoryRequirements(MemoryRequirements);
 	}
@@ -398,7 +400,7 @@ namespace VulkanRHI
 	{
 		DevicePrintfBeginResult(Device, FString::Printf(TEXT("vkGetImageMemoryRequirements(Image=%p, OutReq=%p)"), Image, MemoryRequirements));
 
-		::vkGetImageMemoryRequirements(Device, Image, MemoryRequirements);
+		VULKANAPINAMESPACE::vkGetImageMemoryRequirements(Device, Image, MemoryRequirements);
 
 		DumpMemoryRequirements(MemoryRequirements);
 	}
@@ -430,7 +432,7 @@ namespace VulkanRHI
 	{
 		DumpFenceCreate(Device, CreateInfo, Fence);
 
-		VkResult Result = ::vkCreateFence(Device, CreateInfo, Allocator, Fence);
+		VkResult Result = VULKANAPINAMESPACE::vkCreateFence(Device, CreateInfo, Allocator, Fence);
 
 		PrintResultAndNamedHandle(Result, TEXT("Fence"), *Fence);
 		return Result;
@@ -440,7 +442,7 @@ namespace VulkanRHI
 	{
 		DevicePrintfBegin(Device, FString::Printf(TEXT("vkDestroyFence(Fence=%p)"), Fence));
 
-		::vkDestroyFence(Device, Fence, Allocator);
+		VULKANAPINAMESPACE::vkDestroyFence(Device, Fence, Allocator);
 	}
 
 	static FORCEINLINE_DEBUGGABLE VkResult  vkResetFences(VkDevice Device, uint32 FenceCount, const VkFence* Fences)
@@ -448,7 +450,7 @@ namespace VulkanRHI
 		DevicePrintfBegin(Device, FString::Printf(TEXT("vkResetFences(Count=%d, Fences=%p)"), FenceCount, Fences));
 		DumpFenceList(FenceCount, Fences);
 
-		VkResult Result = ::vkResetFences(Device, FenceCount, Fences);
+		VkResult Result = VULKANAPINAMESPACE::vkResetFences(Device, FenceCount, Fences);
 
 		PrintResult(Result);
 		return Result;
@@ -458,7 +460,7 @@ namespace VulkanRHI
 	{
 		DevicePrintfBeginResult(Device, FString::Printf(TEXT("vkGetFenceStatus(Fence=%p)"), Fence));
 
-		VkResult Result = ::vkGetFenceStatus(Device, Fence);
+		VkResult Result = VULKANAPINAMESPACE::vkGetFenceStatus(Device, Fence);
 
 		PrintResult(Result);
 		return Result;
@@ -469,7 +471,7 @@ namespace VulkanRHI
 		DevicePrintfBegin(Device, FString::Printf(TEXT("vkWaitForFences(Count=%p, Fences=%d, WaitAll=%d, Timeout=%p)"), FenceCount, Fences, (uint32)bWaitAll, Timeout));
 		DumpFenceList(FenceCount, Fences);
 
-		VkResult Result = ::vkWaitForFences(Device, FenceCount, Fences, bWaitAll, Timeout);
+		VkResult Result = VULKANAPINAMESPACE::vkWaitForFences(Device, FenceCount, Fences, bWaitAll, Timeout);
 
 		PrintResult(Result);
 		return Result;
@@ -479,7 +481,7 @@ namespace VulkanRHI
 	{
 		DumpSemaphoreCreate(Device, CreateInfo, Semaphore);
 
-		VkResult Result = ::vkCreateSemaphore(Device, CreateInfo, Allocator, Semaphore);
+		VkResult Result = VULKANAPINAMESPACE::vkCreateSemaphore(Device, CreateInfo, Allocator, Semaphore);
 
 		PrintResultAndNamedHandle(Result, TEXT("Semaphore"), *Semaphore);
 		return Result;
@@ -490,7 +492,7 @@ namespace VulkanRHI
 	{
 		DevicePrintfBegin(Device, FString::Printf(TEXT("vkDestroySemaphore(Semaphore=%p)"), Semaphore));
 
-		::vkDestroySemaphore(Device, Semaphore, Allocator);
+		VULKANAPINAMESPACE::vkDestroySemaphore(Device, Semaphore, Allocator);
 	}
 
 #if 0
@@ -504,7 +506,7 @@ namespace VulkanRHI
 	{
 		DevicePrintfBegin(Device, FString::Printf(TEXT("vkDestroyEvent(Event=%p)"), Event));
 
-		::vkDestroyEvent(Device, Event, Allocator);
+		VULKANAPINAMESPACE::vkDestroyEvent(Device, Event, Allocator);
 	}
 #if 0
 	static FORCEINLINE_DEBUGGABLE VkResult  vkGetEventStatus(
@@ -523,7 +525,7 @@ namespace VulkanRHI
 	{
 		DumpCreateQueryPool(Device, CreateInfo, QueryPool);
 
-		VkResult Result = ::vkCreateQueryPool(Device, CreateInfo, Allocator, QueryPool);
+		VkResult Result = VULKANAPINAMESPACE::vkCreateQueryPool(Device, CreateInfo, Allocator, QueryPool);
 
 		PrintResultAndNamedHandle(Result, TEXT("QueryPool"), *QueryPool);
 
@@ -534,7 +536,7 @@ namespace VulkanRHI
 	{
 		DevicePrintfBegin(Device, FString::Printf(TEXT("vkDestroyQueryPool(QueryPool=%p)"), QueryPool));
 
-		::vkDestroyQueryPool(Device, QueryPool, Allocator);
+		VULKANAPINAMESPACE::vkDestroyQueryPool(Device, QueryPool, Allocator);
 	}
 
 	static FORCEINLINE_DEBUGGABLE VkResult  vkGetQueryPoolResults(VkDevice Device, VkQueryPool QueryPool, 
@@ -542,7 +544,7 @@ namespace VulkanRHI
 	{
 		DevicePrintfBeginResult(Device, FString::Printf(TEXT("vkGetQueryPoolResults(QueryPool=%p, FirstQuery=%d, QueryCount=%d, DataSize=%d, Data=%p, Stride=%d, Flags=%d)[...]"), QueryPool, FirstQuery, QueryCount, (int32)DataSize, Data, Stride, Flags));
 
-		VkResult Result = ::vkGetQueryPoolResults(Device, QueryPool, FirstQuery, QueryCount, DataSize, Data, Stride, Flags);
+		VkResult Result = VULKANAPINAMESPACE::vkGetQueryPoolResults(Device, QueryPool, FirstQuery, QueryCount, DataSize, Data, Stride, Flags);
 
 		PrintResult(Result);
 		return Result;
@@ -552,7 +554,7 @@ namespace VulkanRHI
 	{
 		DumpCreateBuffer(Device, CreateInfo, Buffer);
 
-		VkResult Result = ::vkCreateBuffer(Device, CreateInfo, Allocator, Buffer);
+		VkResult Result = VULKANAPINAMESPACE::vkCreateBuffer(Device, CreateInfo, Allocator, Buffer);
 
 		PrintResultAndNamedHandle(Result, TEXT("Buffer"), *Buffer);
 		return Result;
@@ -562,7 +564,7 @@ namespace VulkanRHI
 	{
 		DevicePrintfBegin(Device, FString::Printf(TEXT("vkDestroyBuffer(Buffer=%p)"), Buffer));
 
-		::vkDestroyBuffer(Device, Buffer, Allocator);
+		VULKANAPINAMESPACE::vkDestroyBuffer(Device, Buffer, Allocator);
 	}
 
 	static FORCEINLINE_DEBUGGABLE VkResult  vkCreateBufferView(VkDevice Device, const VkBufferViewCreateInfo* CreateInfo, const VkAllocationCallbacks* Allocator, VkBufferView* View)
@@ -570,7 +572,7 @@ namespace VulkanRHI
 		DumpCreateBufferView(Device, CreateInfo, View);
 		DevicePrintfBeginResult(Device, FString::Printf(TEXT("vkCreateBufferView(Info=%p, OutView=%p)"), CreateInfo, View));
 
-		VkResult Result = ::vkCreateBufferView(Device, CreateInfo, Allocator, View);
+		VkResult Result = VULKANAPINAMESPACE::vkCreateBufferView(Device, CreateInfo, Allocator, View);
 
 		PrintResultAndNamedHandle(Result, TEXT("BufferView"), *View);
 		return Result;
@@ -580,14 +582,14 @@ namespace VulkanRHI
 	{
 		DevicePrintfBegin(Device, FString::Printf(TEXT("vkDestroyBufferView(BufferView=%p)"), BufferView));
 
-		::vkDestroyBufferView(Device, BufferView, Allocator);
+		VULKANAPINAMESPACE::vkDestroyBufferView(Device, BufferView, Allocator);
 	}
 
 	static FORCEINLINE_DEBUGGABLE VkResult  vkCreateImage(VkDevice Device, const VkImageCreateInfo* CreateInfo, const VkAllocationCallbacks* Allocator, VkImage* Image)
 	{
 		DumpCreateImage(Device, CreateInfo, Image);
 
-		VkResult Result = ::vkCreateImage(Device, CreateInfo, Allocator, Image);
+		VkResult Result = VULKANAPINAMESPACE::vkCreateImage(Device, CreateInfo, Allocator, Image);
 
 		PrintResultAndNamedHandle(Result, TEXT("Image"), *Image);
 		return Result;
@@ -597,14 +599,14 @@ namespace VulkanRHI
 	{
 		DevicePrintfBegin(Device, FString::Printf(TEXT("vkDestroyImage(Image=%p)"), Image));
 
-		::vkDestroyImage(Device, Image, Allocator);
+		VULKANAPINAMESPACE::vkDestroyImage(Device, Image, Allocator);
 	}
 
 	static FORCEINLINE_DEBUGGABLE void  vkGetImageSubresourceLayout(VkDevice Device, VkImage Image, const VkImageSubresource* Subresource, VkSubresourceLayout* Layout)
 	{
 		DumpGetImageSubresourceLayout(Device, Image, Subresource, Layout);
 
-		::vkGetImageSubresourceLayout(Device, Image, Subresource, Layout);
+		VULKANAPINAMESPACE::vkGetImageSubresourceLayout(Device, Image, Subresource, Layout);
 
 		DumpImageSubresourceLayout(Layout);
 	}
@@ -613,7 +615,7 @@ namespace VulkanRHI
 	{
 		DumpCreateImageView(Device, CreateInfo, View);
 
-		VkResult Result = ::vkCreateImageView(Device, CreateInfo, Allocator, View);
+		VkResult Result = VULKANAPINAMESPACE::vkCreateImageView(Device, CreateInfo, Allocator, View);
 
 		PrintResultAndNamedHandle(Result, TEXT("ImageView"), *View);
 		return Result;
@@ -623,14 +625,14 @@ namespace VulkanRHI
 	{
 		DevicePrintfBegin(Device, FString::Printf(TEXT("vkDestroyImageView(ImageView=%p)"), ImageView));
 
-		::vkDestroyImageView(Device, ImageView, Allocator);
+		VULKANAPINAMESPACE::vkDestroyImageView(Device, ImageView, Allocator);
 	}
 
 	static FORCEINLINE_DEBUGGABLE VkResult  vkCreateShaderModule(VkDevice Device, const VkShaderModuleCreateInfo* CreateInfo, const VkAllocationCallbacks* Allocator, VkShaderModule* ShaderModule)
 	{
 		DumpCreateShaderModule(Device, CreateInfo, ShaderModule);
 
-		VkResult Result = ::vkCreateShaderModule(Device, CreateInfo, Allocator,ShaderModule);
+		VkResult Result = VULKANAPINAMESPACE::vkCreateShaderModule(Device, CreateInfo, Allocator,ShaderModule);
 
 		PrintResultAndNamedHandle(Result, TEXT("ShaderModule"), *ShaderModule);
 
@@ -641,14 +643,14 @@ namespace VulkanRHI
 	{
 		DevicePrintfBegin(Device, FString::Printf(TEXT("vkDestroyShaderModule(ShaderModule=%p)"), ShaderModule));
 
-		::vkDestroyShaderModule(Device, ShaderModule, Allocator);
+		VULKANAPINAMESPACE::vkDestroyShaderModule(Device, ShaderModule, Allocator);
 	}
 
 	static FORCEINLINE_DEBUGGABLE VkResult  vkCreatePipelineCache(VkDevice Device, const VkPipelineCacheCreateInfo* CreateInfo, const VkAllocationCallbacks* Allocator, VkPipelineCache* PipelineCache)
 	{
 		DumpCreatePipelineCache(Device, CreateInfo, PipelineCache);
 
-		VkResult Result = ::vkCreatePipelineCache(Device, CreateInfo, Allocator, PipelineCache);
+		VkResult Result = VULKANAPINAMESPACE::vkCreatePipelineCache(Device, CreateInfo, Allocator, PipelineCache);
 
 		PrintResultAndNamedHandle(Result, TEXT("PipelineCache"), *PipelineCache);
 
@@ -659,27 +661,38 @@ namespace VulkanRHI
 	{
 		DevicePrintfBegin(Device, FString::Printf(TEXT("vkDestroyPipelineCache(PipelineCache=%p)"), PipelineCache));
 
-		::vkDestroyPipelineCache(Device, PipelineCache, Allocator);
+		VULKANAPINAMESPACE::vkDestroyPipelineCache(Device, PipelineCache, Allocator);
 	}
 
-#if 0
 	static FORCEINLINE_DEBUGGABLE VkResult  vkGetPipelineCacheData(
 		VkDevice                                    Device,
-		VkPipelineCache                             pipelineCache,
-		size_t*                                     pDataSize,
-		void*                                       pData);
+		VkPipelineCache                             PipelineCache,
+		size_t*                                     DataSize,
+		void*                                       Data)
+	{
+		DevicePrintfBegin(Device, FString::Printf(TEXT("vkGetPipelineCacheData(PipelineCache=%p, DataSize=%d, [Data])"), PipelineCache, DataSize));
+		VkResult Result = VULKANAPINAMESPACE::vkGetPipelineCacheData(Device, PipelineCache, DataSize, Data);
+		PrintResult(Result);
+		return Result;
+	}
 
 	static FORCEINLINE_DEBUGGABLE VkResult  vkMergePipelineCaches(
 		VkDevice                                    Device,
-		VkPipelineCache                             dstCache,
-		uint32                                    srcCacheCount,
-		const VkPipelineCache*                      pSrcCaches);
-#endif
+		VkPipelineCache                             DestCache,
+		uint32                                    SourceCacheCount,
+		const VkPipelineCache*                      SrcCaches)
+	{
+		DevicePrintfBegin(Device, FString::Printf(TEXT("vkMergePipelineCaches(DestCache=%p, SourceCacheCount=%d, [SrcCaches])"), DestCache, SourceCacheCount));
+		VkResult Result = VULKANAPINAMESPACE::vkMergePipelineCaches(Device, DestCache, SourceCacheCount, SrcCaches);
+		PrintResult(Result);
+		return Result;
+	}
+
 	static FORCEINLINE_DEBUGGABLE VkResult  vkCreateGraphicsPipelines(VkDevice Device, VkPipelineCache PipelineCache, uint32 CreateInfoCount, const VkGraphicsPipelineCreateInfo* CreateInfos, const VkAllocationCallbacks* Allocator, VkPipeline* Pipelines)
 	{
 		DevicePrintfBeginResult(Device, FString::Printf(TEXT("vkCreateGraphicsPipelines(PipelineCache=%p, CreateInfoCount=%d, CreateInfos=%p, OutPipelines=%p)[...]"), PipelineCache, CreateInfoCount, CreateInfos, Pipelines));
 
-		VkResult Result = ::vkCreateGraphicsPipelines(Device, PipelineCache, CreateInfoCount, CreateInfos, Allocator, Pipelines);
+		VkResult Result = VULKANAPINAMESPACE::vkCreateGraphicsPipelines(Device, PipelineCache, CreateInfoCount, CreateInfos, Allocator, Pipelines);
 
 		//#todo-rco: Multiple pipelines!
 		PrintResultAndNamedHandle(Result, TEXT("Pipeline"), Pipelines[0]);
@@ -690,7 +703,7 @@ namespace VulkanRHI
 	{
 		DevicePrintfBeginResult(Device, FString::Printf(TEXT("vkCreateComputePipelines(PipelineCache=%p, CreateInfoCount=%d, CreateInfos=%p, OutPipelines=%p)[...]"), PipelineCache, CreateInfoCount, CreateInfos, Pipelines));
 
-		VkResult Result = ::vkCreateComputePipelines(Device, PipelineCache, CreateInfoCount, CreateInfos, Allocator, Pipelines);
+		VkResult Result = VULKANAPINAMESPACE::vkCreateComputePipelines(Device, PipelineCache, CreateInfoCount, CreateInfos, Allocator, Pipelines);
 
 		//#todo-rco: Multiple pipelines!
 		PrintResultAndNamedHandle(Result, TEXT("Pipeline"), Pipelines[0]);
@@ -701,14 +714,14 @@ namespace VulkanRHI
 	{
 		DevicePrintfBegin(Device, FString::Printf(TEXT("vkDestroyPipeline(Pipeline=%p)"), Pipeline));
 
-		::vkDestroyPipeline(Device, Pipeline, Allocator);
+		VULKANAPINAMESPACE::vkDestroyPipeline(Device, Pipeline, Allocator);
 	}
 
 	static FORCEINLINE_DEBUGGABLE VkResult  vkCreatePipelineLayout(VkDevice Device, const VkPipelineLayoutCreateInfo* CreateInfo, const VkAllocationCallbacks* Allocator, VkPipelineLayout* PipelineLayout)
 	{
 		DumpCreatePipelineLayout(Device, CreateInfo, PipelineLayout);
 
-		VkResult Result = ::vkCreatePipelineLayout(Device, CreateInfo, Allocator, PipelineLayout);
+		VkResult Result = VULKANAPINAMESPACE::vkCreatePipelineLayout(Device, CreateInfo, Allocator, PipelineLayout);
 
 		PrintResultAndNamedHandle(Result, TEXT("PipelineLayout"), *PipelineLayout);
 
@@ -719,14 +732,14 @@ namespace VulkanRHI
 	{
 		DevicePrintfBegin(Device, FString::Printf(TEXT("vkDestroyPipelineLayout(PipelineLayout=%p)"), PipelineLayout));
 
-		::vkDestroyPipelineLayout(Device, PipelineLayout, Allocator);
+		VULKANAPINAMESPACE::vkDestroyPipelineLayout(Device, PipelineLayout, Allocator);
 	}
 
 	static FORCEINLINE_DEBUGGABLE VkResult  vkCreateSampler(VkDevice Device, const VkSamplerCreateInfo* CreateInfo, const VkAllocationCallbacks* Allocator, VkSampler* Sampler)
 	{
 		DumpCreateSampler(Device, CreateInfo, Sampler);
 
-		VkResult Result = ::vkCreateSampler(Device, CreateInfo, Allocator, Sampler);
+		VkResult Result = VULKANAPINAMESPACE::vkCreateSampler(Device, CreateInfo, Allocator, Sampler);
 
 		PrintResultAndNamedHandle(Result, TEXT("Sampler"), *Sampler);
 
@@ -737,14 +750,14 @@ namespace VulkanRHI
 	{
 		DevicePrintfBegin(Device, FString::Printf(TEXT("vkDestroySampler(Sampler=%p)"), Sampler));
 
-		::vkDestroySampler(Device, Sampler, Allocator);
+		VULKANAPINAMESPACE::vkDestroySampler(Device, Sampler, Allocator);
 	}
 
 	static FORCEINLINE_DEBUGGABLE VkResult  vkCreateDescriptorSetLayout(VkDevice Device, const VkDescriptorSetLayoutCreateInfo* CreateInfo, const VkAllocationCallbacks* Allocator, VkDescriptorSetLayout* SetLayout)
 	{
 		DumpCreateDescriptorSetLayout(Device, CreateInfo, SetLayout);
 
-		VkResult Result = ::vkCreateDescriptorSetLayout(Device, CreateInfo, Allocator, SetLayout);
+		VkResult Result = VULKANAPINAMESPACE::vkCreateDescriptorSetLayout(Device, CreateInfo, Allocator, SetLayout);
 
 		PrintResultAndNamedHandle(Result, TEXT("DescriptorSetLayout"), *SetLayout);
 		return Result;
@@ -754,14 +767,14 @@ namespace VulkanRHI
 	{
 		DevicePrintfBegin(Device, FString::Printf(TEXT("vkDestroyDescriptorSetLayout(DescriptorSetLayout=%p)"), DescriptorSetLayout));
 
-		::vkDestroyDescriptorSetLayout(Device, DescriptorSetLayout, Allocator);
+		VULKANAPINAMESPACE::vkDestroyDescriptorSetLayout(Device, DescriptorSetLayout, Allocator);
 	}
 
 	static FORCEINLINE_DEBUGGABLE VkResult  vkCreateDescriptorPool(VkDevice Device, const VkDescriptorPoolCreateInfo* CreateInfo, const VkAllocationCallbacks* Allocator, VkDescriptorPool* DescriptorPool)
 	{
 		DumpCreateDescriptorPool(Device, CreateInfo, DescriptorPool);
 
-		VkResult Result = ::vkCreateDescriptorPool(Device, CreateInfo, Allocator, DescriptorPool);
+		VkResult Result = VULKANAPINAMESPACE::vkCreateDescriptorPool(Device, CreateInfo, Allocator, DescriptorPool);
 
 		PrintResultAndNamedHandle(Result, TEXT("DescriptorPool"), *DescriptorPool);
 
@@ -772,7 +785,7 @@ namespace VulkanRHI
 	{
 		DevicePrintfBegin(Device, FString::Printf(TEXT("vkDestroyDescriptorPool(DescriptorPool=%p)"), DescriptorPool));
 
-		::vkDestroyDescriptorPool(Device, DescriptorPool, Allocator);
+		VULKANAPINAMESPACE::vkDestroyDescriptorPool(Device, DescriptorPool, Allocator);
 	}
 #if 0
 	static FORCEINLINE_DEBUGGABLE VkResult  vkResetDescriptorPool(
@@ -784,7 +797,7 @@ namespace VulkanRHI
 	{
 		DumpAllocateDescriptorSets(Device, AllocateInfo, DescriptorSets);
 
-		VkResult Result = ::vkAllocateDescriptorSets(Device, AllocateInfo, DescriptorSets);
+		VkResult Result = VULKANAPINAMESPACE::vkAllocateDescriptorSets(Device, AllocateInfo, DescriptorSets);
 
 		PrintResultAndNamedHandle(Result, TEXT("DescriptorSet"), *DescriptorSets);
 		return Result;
@@ -794,7 +807,7 @@ namespace VulkanRHI
 	{
 		DumpFreeDescriptorSets(Device, DescriptorPool, DescriptorSetCount, DescriptorSets);
 
-		VkResult Result = ::vkFreeDescriptorSets(Device, DescriptorPool, DescriptorSetCount, DescriptorSets);
+		VkResult Result = VULKANAPINAMESPACE::vkFreeDescriptorSets(Device, DescriptorPool, DescriptorSetCount, DescriptorSets);
 
 		PrintResult(Result);
 		return Result;
@@ -804,14 +817,14 @@ namespace VulkanRHI
 	{
 		DumpUpdateDescriptorSets(Device, DescriptorWriteCount, DescriptorWrites, DescriptorCopyCount, DescriptorCopies);
 
-		::vkUpdateDescriptorSets(Device, DescriptorWriteCount, DescriptorWrites, DescriptorCopyCount, DescriptorCopies);
+		VULKANAPINAMESPACE::vkUpdateDescriptorSets(Device, DescriptorWriteCount, DescriptorWrites, DescriptorCopyCount, DescriptorCopies);
 	}
 
 	static FORCEINLINE_DEBUGGABLE VkResult  vkCreateFramebuffer(VkDevice Device, const VkFramebufferCreateInfo* CreateInfo, const VkAllocationCallbacks* Allocator, VkFramebuffer* Framebuffer)
 	{
 		DumpCreateFramebuffer(Device, CreateInfo, Framebuffer);
 
-		VkResult Result = ::vkCreateFramebuffer(Device, CreateInfo, Allocator, Framebuffer);
+		VkResult Result = VULKANAPINAMESPACE::vkCreateFramebuffer(Device, CreateInfo, Allocator, Framebuffer);
 
 		PrintResultAndNamedHandle(Result, TEXT("Framebuffer"), *Framebuffer);
 		return Result;
@@ -821,14 +834,14 @@ namespace VulkanRHI
 	{
 		DevicePrintfBegin(Device, FString::Printf(TEXT("vkDestroyFramebuffer(Framebuffer=%p)"), Framebuffer));
 
-		::vkDestroyFramebuffer(Device, Framebuffer, Allocator);
+		VULKANAPINAMESPACE::vkDestroyFramebuffer(Device, Framebuffer, Allocator);
 	}
 
 	static FORCEINLINE_DEBUGGABLE VkResult  vkCreateRenderPass(VkDevice Device, const VkRenderPassCreateInfo* CreateInfo, const VkAllocationCallbacks* Allocator, VkRenderPass* RenderPass)
 	{
 		DumpCreateRenderPass(Device, CreateInfo, RenderPass);
 
-		VkResult Result = ::vkCreateRenderPass(Device, CreateInfo, Allocator, RenderPass);
+		VkResult Result = VULKANAPINAMESPACE::vkCreateRenderPass(Device, CreateInfo, Allocator, RenderPass);
 
 		PrintResultAndNamedHandle(Result, TEXT("RenderPass"), *RenderPass);
 		return Result;
@@ -838,7 +851,7 @@ namespace VulkanRHI
 	{
 		DevicePrintfBegin(Device, FString::Printf(TEXT("vkDestroyRenderPass(RenderPass=%p)"), RenderPass));
 
-		::vkDestroyRenderPass(Device, RenderPass, Allocator);
+		VULKANAPINAMESPACE::vkDestroyRenderPass(Device, RenderPass, Allocator);
 	}
 #if 0
 	static FORCEINLINE_DEBUGGABLE void  vkGetRenderAreaGranularity(
@@ -850,7 +863,7 @@ namespace VulkanRHI
 	{
 		DumpCreateCommandPool(Device, CreateInfo, CommandPool);
 
-		VkResult Result = ::vkCreateCommandPool(Device, CreateInfo, Allocator, CommandPool);
+		VkResult Result = VULKANAPINAMESPACE::vkCreateCommandPool(Device, CreateInfo, Allocator, CommandPool);
 
 		PrintResultAndNamedHandle(Result, TEXT("CommandPool"), *CommandPool);
 
@@ -861,7 +874,7 @@ namespace VulkanRHI
 	{
 		DevicePrintfBegin(Device, FString::Printf(TEXT("vkDestroyCommandPool(CommandPool=%p)"), CommandPool));
 
-		::vkDestroyCommandPool(Device, CommandPool, Allocator);
+		VULKANAPINAMESPACE::vkDestroyCommandPool(Device, CommandPool, Allocator);
 	}
 #if 0
 	static FORCEINLINE_DEBUGGABLE VkResult  vkResetCommandPool(
@@ -873,7 +886,7 @@ namespace VulkanRHI
 	{
 		DevicePrintfBeginResult(Device, FString::Printf(TEXT("vkAllocateCommandBuffers(AllocateInfo=%p, OutCommandBuffers=%p)[...]"), AllocateInfo, CommandBuffers));
 
-		VkResult Result = ::vkAllocateCommandBuffers(Device, AllocateInfo, CommandBuffers);
+		VkResult Result = VULKANAPINAMESPACE::vkAllocateCommandBuffers(Device, AllocateInfo, CommandBuffers);
 		PrintResultAndNamedHandle(Result, TEXT("CommandBuffers"), *CommandBuffers);
 		return Result;
 	}
@@ -882,14 +895,14 @@ namespace VulkanRHI
 	{
 		DevicePrintfBegin(Device, FString::Printf(TEXT("vkFreeCommandBuffers(CommandPool=%p, CommandBufferCount=%d, CommandBuffers=%p)[...]"), CommandPool, CommandBufferCount, CommandBuffers));
 
-		::vkFreeCommandBuffers(Device, CommandPool, CommandBufferCount, CommandBuffers);
+		VULKANAPINAMESPACE::vkFreeCommandBuffers(Device, CommandPool, CommandBufferCount, CommandBuffers);
 	}
 
 	static FORCEINLINE_DEBUGGABLE VkResult  vkBeginCommandBuffer(VkCommandBuffer CommandBuffer, const VkCommandBufferBeginInfo* BeginInfo)
 	{
 		DumpBeginCommandBuffer(CommandBuffer, BeginInfo);
 
-		VkResult Result = ::vkBeginCommandBuffer(CommandBuffer, BeginInfo);
+		VkResult Result = VULKANAPINAMESPACE::vkBeginCommandBuffer(CommandBuffer, BeginInfo);
 
 		PrintResult(Result);
 
@@ -900,7 +913,7 @@ namespace VulkanRHI
 	{
 		CmdPrintfBeginResult(CommandBuffer, FString::Printf(TEXT("vkEndCommandBuffer(Cmd=%p)")));
 
-		VkResult Result = ::vkEndCommandBuffer(CommandBuffer);
+		VkResult Result = VULKANAPINAMESPACE::vkEndCommandBuffer(CommandBuffer);
 
 		PrintResult(Result);
 		return Result;
@@ -910,7 +923,7 @@ namespace VulkanRHI
 	 {
 		PrintfBeginResult(FString::Printf(TEXT("vkResetCommandBuffer(Cmd=%p, Flags=%d)"), CommandBuffer, Flags));
 
-		VkResult Result = ::vkResetCommandBuffer(CommandBuffer, Flags);
+		VkResult Result = VULKANAPINAMESPACE::vkResetCommandBuffer(CommandBuffer, Flags);
 
 		PrintResult(Result);
 		return Result;
@@ -920,28 +933,28 @@ namespace VulkanRHI
 	{
 		CmdPrintfBegin(CommandBuffer, FString::Printf(TEXT("vkCmdBindPipeline(BindPoint=%d, Pipeline=%p)[...]"), (int32)PipelineBindPoint, Pipeline));
 
-		::vkCmdBindPipeline(CommandBuffer, PipelineBindPoint, Pipeline);
+		VULKANAPINAMESPACE::vkCmdBindPipeline(CommandBuffer, PipelineBindPoint, Pipeline);
 	}
 
 	static FORCEINLINE_DEBUGGABLE void  vkCmdSetViewport(VkCommandBuffer CommandBuffer, uint32 FirstViewport, uint32 ViewportCount, const VkViewport* Viewports)
 	{
 		CmdPrintfBegin(CommandBuffer, FString::Printf(TEXT("vkCmdSetViewport(FirstViewport=%d, ViewportCount=%d, Viewports=%p)[...]"), FirstViewport, ViewportCount, Viewports));
 
-		::vkCmdSetViewport(CommandBuffer, FirstViewport, ViewportCount, Viewports);
+		VULKANAPINAMESPACE::vkCmdSetViewport(CommandBuffer, FirstViewport, ViewportCount, Viewports);
 	}
 
 	static FORCEINLINE_DEBUGGABLE void  vkCmdSetScissor(VkCommandBuffer CommandBuffer, uint32 FirstScissor, uint32 ScissorCount, const VkRect2D* Scissors)
 	{
 		CmdPrintfBegin(CommandBuffer, FString::Printf(TEXT("vkCmdSetScissor(FirstScissor=%d, ScissorCount=%d, Scissors=%p)[...]"), FirstScissor, ScissorCount, Scissors));
 		
-		::vkCmdSetScissor(CommandBuffer, FirstScissor, ScissorCount, Scissors);
+		VULKANAPINAMESPACE::vkCmdSetScissor(CommandBuffer, FirstScissor, ScissorCount, Scissors);
 	}
 
 	static FORCEINLINE_DEBUGGABLE void  vkCmdSetLineWidth(VkCommandBuffer CommandBuffer, float LineWidth)
 	{
 		CmdPrintfBegin(CommandBuffer, FString::Printf(TEXT("vkCmdSetLineWidth(Width=%f)"), LineWidth));
 
-		::vkCmdSetLineWidth(CommandBuffer, LineWidth);
+		VULKANAPINAMESPACE::vkCmdSetLineWidth(CommandBuffer, LineWidth);
 	}
 
 #if 0
@@ -982,35 +995,35 @@ namespace VulkanRHI
 	{
 		DumpBindDescriptorSets(CommandBuffer, PipelineBindPoint, Layout, FirstSet, DescriptorSetCount, DescriptorSets, DynamicOffsetCount, DynamicOffsets);
 
-		::vkCmdBindDescriptorSets(CommandBuffer, PipelineBindPoint, Layout, FirstSet, DescriptorSetCount, DescriptorSets, DynamicOffsetCount, DynamicOffsets);
+		VULKANAPINAMESPACE::vkCmdBindDescriptorSets(CommandBuffer, PipelineBindPoint, Layout, FirstSet, DescriptorSetCount, DescriptorSets, DynamicOffsetCount, DynamicOffsets);
 	}
 
 	static FORCEINLINE_DEBUGGABLE void  vkCmdBindIndexBuffer(VkCommandBuffer CommandBuffer, VkBuffer Buffer, VkDeviceSize Offset, VkIndexType IndexType)
 	{
 		CmdPrintfBegin(CommandBuffer, FString::Printf(TEXT("vkCmdBindIndexBuffer(Buffer=%p, Offset=%d, IndexType=%d)"), Buffer, (int32)Offset, (int32)IndexType));
 
-		::vkCmdBindIndexBuffer(CommandBuffer, Buffer, Offset, IndexType);
+		VULKANAPINAMESPACE::vkCmdBindIndexBuffer(CommandBuffer, Buffer, Offset, IndexType);
 	}
 
 	static FORCEINLINE_DEBUGGABLE void  vkCmdBindVertexBuffers(VkCommandBuffer CommandBuffer, uint32 FirstBinding, uint32 BindingCount, const VkBuffer* Buffers, const VkDeviceSize* Offsets)
 	{
 		DumpCmdBindVertexBuffers(CommandBuffer, FirstBinding, BindingCount, Buffers, Offsets);
 
-		::vkCmdBindVertexBuffers(CommandBuffer, FirstBinding, BindingCount, Buffers, Offsets);
+		VULKANAPINAMESPACE::vkCmdBindVertexBuffers(CommandBuffer, FirstBinding, BindingCount, Buffers, Offsets);
 	}
 
 	static FORCEINLINE_DEBUGGABLE void  vkCmdDraw(VkCommandBuffer CommandBuffer, uint32 VertexCount, uint32 InstanceCount, uint32 FirstVertex, uint32 FirstInstance)
 	{
 		CmdPrintfBegin(CommandBuffer, FString::Printf(TEXT("vkCmdDraw(NumVertices=%d, NumInstances=%d, FirstVertex=%d, FirstInstance=%d)"), VertexCount, InstanceCount, FirstVertex, FirstInstance));
 
-		::vkCmdDraw(CommandBuffer, VertexCount, InstanceCount, FirstVertex, FirstInstance);
+		VULKANAPINAMESPACE::vkCmdDraw(CommandBuffer, VertexCount, InstanceCount, FirstVertex, FirstInstance);
 	}
 
 	static FORCEINLINE_DEBUGGABLE void  vkCmdDrawIndexed(VkCommandBuffer CommandBuffer, uint32 IndexCount, uint32 InstanceCount, uint32 FirstIndex, int32_t VertexOffset, uint32 FirstInstance)
 	{
 		CmdPrintfBegin(CommandBuffer, FString::Printf(TEXT("vkCmdDrawIndexed(IndexCount=%d, NumInstances=%d, FirstIndex=%d, VertexOffset=%d, FirstInstance=%d)"), IndexCount, InstanceCount, FirstIndex, VertexOffset, FirstInstance));
 
-		::vkCmdDrawIndexed(CommandBuffer, IndexCount, InstanceCount, FirstIndex, VertexOffset, FirstInstance);
+		VULKANAPINAMESPACE::vkCmdDrawIndexed(CommandBuffer, IndexCount, InstanceCount, FirstIndex, VertexOffset, FirstInstance);
 	}
 
 #if 0
@@ -1043,14 +1056,14 @@ namespace VulkanRHI
 	{
 		DumpCmdCopyBuffer(CommandBuffer, SrcBuffer, DstBuffer, RegionCount, Regions);
 
-		::vkCmdCopyBuffer(CommandBuffer, SrcBuffer, DstBuffer, RegionCount, Regions);
+		VULKANAPINAMESPACE::vkCmdCopyBuffer(CommandBuffer, SrcBuffer, DstBuffer, RegionCount, Regions);
 	}
 
 	static FORCEINLINE_DEBUGGABLE void  vkCmdCopyImage(VkCommandBuffer CommandBuffer, VkImage SrcImage, VkImageLayout SrcImageLayout, VkImage DstImage, VkImageLayout DstImageLayout, uint32 RegionCount, const VkImageCopy* Regions)
 	{
 		CmdPrintfBegin(CommandBuffer, FString::Printf(TEXT("vkCmdCopyImage(SrcImage=%p, SrcImageLayout=%d, DstImage=%p, DstImageLayout=%d, RegionCount=%d, Regions=%p)[...]"), SrcImage, (int32)SrcImageLayout, DstImage, (int32)DstImageLayout, RegionCount, Regions));
 
-		::vkCmdCopyImage(CommandBuffer, SrcImage, SrcImageLayout, DstImage, DstImageLayout, RegionCount, Regions);
+		VULKANAPINAMESPACE::vkCmdCopyImage(CommandBuffer, SrcImage, SrcImageLayout, DstImage, DstImageLayout, RegionCount, Regions);
 	}
 
 #if 0
@@ -1068,7 +1081,7 @@ namespace VulkanRHI
 	{
 		DumpCmdCopyBufferToImage(CommandBuffer, SrcBuffer, DstImage, DstImageLayout, RegionCount, Regions);
 
-		::vkCmdCopyBufferToImage(CommandBuffer, SrcBuffer, DstImage, DstImageLayout, RegionCount, Regions);
+		VULKANAPINAMESPACE::vkCmdCopyBufferToImage(CommandBuffer, SrcBuffer, DstImage, DstImageLayout, RegionCount, Regions);
 	}
 
 	static FORCEINLINE_DEBUGGABLE void  vkCmdCopyImageToBuffer(VkCommandBuffer CommandBuffer, VkImage SrcImage, VkImageLayout SrcImageLayout, VkBuffer DstBuffer, uint32 RegionCount, const VkBufferImageCopy* Regions)
@@ -1097,21 +1110,21 @@ namespace VulkanRHI
 	{
 		DumpCmdClearColorImage(CommandBuffer, Image, ImageLayout, Color, RangeCount, Ranges);
 
-		::vkCmdClearColorImage(CommandBuffer, Image, ImageLayout, Color, RangeCount, Ranges);
+		VULKANAPINAMESPACE::vkCmdClearColorImage(CommandBuffer, Image, ImageLayout, Color, RangeCount, Ranges);
 	}
 
 	static FORCEINLINE_DEBUGGABLE void  vkCmdClearDepthStencilImage(VkCommandBuffer CommandBuffer, VkImage Image, VkImageLayout ImageLayout, const VkClearDepthStencilValue* DepthStencil, uint32 RangeCount, const VkImageSubresourceRange* Ranges)
 	{
 		DumpCmdClearDepthStencilImage(CommandBuffer, Image, ImageLayout, DepthStencil, RangeCount, Ranges);
 
-		::vkCmdClearDepthStencilImage(CommandBuffer, Image, ImageLayout, DepthStencil, RangeCount, Ranges);
+		VULKANAPINAMESPACE::vkCmdClearDepthStencilImage(CommandBuffer, Image, ImageLayout, DepthStencil, RangeCount, Ranges);
 	}
 
 	static FORCEINLINE_DEBUGGABLE void  vkCmdClearAttachments(VkCommandBuffer CommandBuffer, uint32 AttachmentCount, const VkClearAttachment* Attachments, uint32 RectCount, const VkClearRect* Rects)
 	{
 		DumpCmdClearAttachments(CommandBuffer, AttachmentCount, Attachments, RectCount, Rects);
 
-		::vkCmdClearAttachments(CommandBuffer, AttachmentCount, Attachments, RectCount, Rects);
+		VULKANAPINAMESPACE::vkCmdClearAttachments(CommandBuffer, AttachmentCount, Attachments, RectCount, Rects);
 	}
 
 	static FORCEINLINE_DEBUGGABLE void  vkCmdResolveImage(
@@ -1122,7 +1135,7 @@ namespace VulkanRHI
 	{
 		DumpResolveImage(CommandBuffer, SrcImage, SrcImageLayout, DstImage, DstImageLayout, RegionCount, Regions);
 
-		::vkCmdResolveImage(CommandBuffer, SrcImage, SrcImageLayout, DstImage, DstImageLayout, RegionCount, Regions);
+		VULKANAPINAMESPACE::vkCmdResolveImage(CommandBuffer, SrcImage, SrcImageLayout, DstImage, DstImageLayout, RegionCount, Regions);
 	}
 #if 0
 	static FORCEINLINE_DEBUGGABLE void  vkCmdSetEvent(
@@ -1156,7 +1169,7 @@ namespace VulkanRHI
 	{
 		DumpCmdPipelineBarrier(CommandBuffer, SrcStageMask, DstStageMask, DependencyFlags, MemoryBarrierCount, MemoryBarriers, BufferMemoryBarrierCount, BufferMemoryBarriers, ImageMemoryBarrierCount, ImageMemoryBarriers);
 
-		::vkCmdPipelineBarrier(CommandBuffer, SrcStageMask, DstStageMask, DependencyFlags, MemoryBarrierCount, MemoryBarriers, BufferMemoryBarrierCount, BufferMemoryBarriers, ImageMemoryBarrierCount, ImageMemoryBarriers);
+		VULKANAPINAMESPACE::vkCmdPipelineBarrier(CommandBuffer, SrcStageMask, DstStageMask, DependencyFlags, MemoryBarrierCount, MemoryBarriers, BufferMemoryBarrierCount, BufferMemoryBarriers, ImageMemoryBarrierCount, ImageMemoryBarriers);
 	}
 #if 0
 	static FORCEINLINE_DEBUGGABLE void  vkCmdBeginQuery(
@@ -1174,14 +1187,14 @@ namespace VulkanRHI
 	{
 		CmdPrintfBegin(CommandBuffer, FString::Printf(TEXT("vkCmdResetQueryPool(QueryPool=%p, FirstQuery=%d, NumQueries=%d)"), QueryPool, FirstQuery, QueryCount));
 
-		::vkCmdResetQueryPool(CommandBuffer, QueryPool, FirstQuery, QueryCount);
+		VULKANAPINAMESPACE::vkCmdResetQueryPool(CommandBuffer, QueryPool, FirstQuery, QueryCount);
 	}
 
 	static FORCEINLINE_DEBUGGABLE void  vkCmdWriteTimestamp(VkCommandBuffer CommandBuffer, VkPipelineStageFlagBits PipelineStage, VkQueryPool QueryPool, uint32 Query)
 	{
 		CmdPrintfBegin(CommandBuffer, FString::Printf(TEXT("vkCmdWriteTimestamp(PipelineStage=%d, QueryPool=%p, Query=%d)"), (int32)PipelineStage, QueryPool, Query));
 
-		::vkCmdWriteTimestamp(CommandBuffer, PipelineStage, QueryPool, Query);
+		VULKANAPINAMESPACE::vkCmdWriteTimestamp(CommandBuffer, PipelineStage, QueryPool, Query);
 	}
 #if 0
 	static FORCEINLINE_DEBUGGABLE void  vkCmdCopyQueryPoolResults(
@@ -1209,7 +1222,7 @@ namespace VulkanRHI
 	{
 		DumpCmdBeginRenderPass(CommandBuffer, RenderPassBegin, Contents);
 
-		::vkCmdBeginRenderPass(CommandBuffer, RenderPassBegin, Contents);
+		VULKANAPINAMESPACE::vkCmdBeginRenderPass(CommandBuffer, RenderPassBegin, Contents);
 	}
 #if 0
 	static FORCEINLINE_DEBUGGABLE void  vkCmdNextSubpass(
@@ -1220,7 +1233,7 @@ namespace VulkanRHI
 	{
 		CmdPrintfBegin(CommandBuffer, TEXT("vkCmdEndRenderPass()"));
 
-		::vkCmdEndRenderPass(CommandBuffer);
+		VULKANAPINAMESPACE::vkCmdEndRenderPass(CommandBuffer);
 	}
 #if 0
 	static FORCEINLINE_DEBUGGABLE void  vkCmdExecuteCommands(
@@ -1233,7 +1246,7 @@ namespace VulkanRHI
 	{
 		DevicePrintfBeginResult(Device, FString::Printf(TEXT("vkCreateSwapchainKHR(SwapChainInfo=%p, OutSwapChain=%p)[...]"), CreateInfo, Swapchain));
 
-		VkResult Result = ::vkCreateSwapchainKHR(Device, CreateInfo, nullptr, Swapchain);
+		VkResult Result = VULKANAPINAMESPACE::vkCreateSwapchainKHR(Device, CreateInfo, nullptr, Swapchain);
 
 		PrintResultAndNamedHandle(Result, TEXT("SwapChain"), *Swapchain);
 		return Result;
@@ -1243,14 +1256,14 @@ namespace VulkanRHI
 	{
 		DevicePrintfBegin(Device, FString::Printf(TEXT("vkDestroySwapchainKHR(SwapChain=%p)[...]"), Swapchain));
 
-		::vkDestroySwapchainKHR(Device, Swapchain, Allocator);
+		VULKANAPINAMESPACE::vkDestroySwapchainKHR(Device, Swapchain, Allocator);
 	}
 
 	static FORCEINLINE_DEBUGGABLE VkResult vkGetSwapchainImagesKHR(VkDevice Device, VkSwapchainKHR Swapchain, uint32_t* SwapchainImageCount, VkImage* SwapchainImages)
 	{
 		DevicePrintfBeginResult(Device, FString::Printf(TEXT("vkGetSwapchainImagesKHR(Swapchain=%p, OutSwapchainImageCount=%p, OutSwapchainImages=%p)"), Swapchain, SwapchainImageCount, SwapchainImages));
 
-		VkResult Result = ::vkGetSwapchainImagesKHR(Device, Swapchain, SwapchainImageCount, SwapchainImages);
+		VkResult Result = VULKANAPINAMESPACE::vkGetSwapchainImagesKHR(Device, Swapchain, SwapchainImageCount, SwapchainImages);
 
 		DumpSwapChainImages(Result, SwapchainImageCount, SwapchainImages);
 		return Result;
@@ -1260,7 +1273,7 @@ namespace VulkanRHI
 	{
 		DevicePrintfBeginResult(Device, FString::Printf(TEXT("vkAcquireNextImageKHR(Swapchain=%p, Timeout=%p, Semaphore=%p, Fence=%p, OutImageIndex=%p)[...]"), Swapchain, (void*)Timeout, Semaphore, Fence, ImageIndex));
 
-		VkResult Result = ::vkAcquireNextImageKHR(Device, Swapchain, Timeout, Semaphore, Fence, ImageIndex);
+		VkResult Result = VULKANAPINAMESPACE::vkAcquireNextImageKHR(Device, Swapchain, Timeout, Semaphore, Fence, ImageIndex);
 
 		PrintResultAndNamedHandle(Result, TEXT("ImageIndex"), ImageIndex);
 		return Result;
@@ -1270,11 +1283,55 @@ namespace VulkanRHI
 	{
 		PrintfBeginResult(FString::Printf(TEXT("vkQueuePresentKHR(Queue=%p, Info=%p)[...]"), Queue, PresentInfo));
 
-		VkResult Result = ::vkQueuePresentKHR(Queue, PresentInfo);
+		VkResult Result = VULKANAPINAMESPACE::vkQueuePresentKHR(Queue, PresentInfo);
+		PrintResult(Result);
+		return Result;
+	}
+
+	static FORCEINLINE_DEBUGGABLE VkResult vkGetPhysicalDeviceSurfaceCapabilitiesKHR(VkPhysicalDevice PhysicalDevice, VkSurfaceKHR Surface, VkSurfaceCapabilitiesKHR* SurfaceCapabilities)
+	{
+		PrintfBeginResult(FString::Printf(TEXT("vkGetPhysicalDeviceSurfaceCapabilitiesKHR(PhysicalDevice=%p, Surface=%p)[...]"), PhysicalDevice, Surface));
+		VkResult Result = VULKANAPINAMESPACE::vkGetPhysicalDeviceSurfaceCapabilitiesKHR(PhysicalDevice, Surface, SurfaceCapabilities);
+		PrintResult(Result);
+		return Result;
+	}
+
+	static FORCEINLINE_DEBUGGABLE VkResult vkGetPhysicalDeviceSurfaceFormatsKHR(VkPhysicalDevice PhysicalDevice, VkSurfaceKHR Surface, uint32_t* SurfaceFormatCountPtr, VkSurfaceFormatKHR* SurfaceFormats)
+	{
+		PrintfBeginResult(FString::Printf(TEXT("vkGetPhysicalDeviceSurfaceFormatsKHR(PhysicalDevice=%p, Surface=%p)[...]"), PhysicalDevice, Surface));
+		VkResult Result = VULKANAPINAMESPACE::vkGetPhysicalDeviceSurfaceFormatsKHR(PhysicalDevice, Surface, SurfaceFormatCountPtr, SurfaceFormats);
+		PrintResult(Result);
+		return Result;
+	}
+
+	static FORCEINLINE_DEBUGGABLE VkResult vkGetPhysicalDeviceSurfaceSupportKHR(VkPhysicalDevice PhysicalDevice, uint32_t QueueFamilyIndex,	VkSurfaceKHR Surface, VkBool32* SupportedPtr)
+	{
+		PrintfBeginResult(FString::Printf(TEXT("vkGetPhysicalDeviceSurfaceSupportKHR(PhysicalDevice=%p, QueueFamilyIndex=%d)[...]"), PhysicalDevice, QueueFamilyIndex));
+		VkResult Result = VULKANAPINAMESPACE::vkGetPhysicalDeviceSurfaceSupportKHR(PhysicalDevice, QueueFamilyIndex, Surface, SupportedPtr);
 		PrintResult(Result);
 		return Result;
 	}
 
+	static FORCEINLINE_DEBUGGABLE VkResult vkGetPhysicalDeviceSurfacePresentModesKHR(VkPhysicalDevice PhysicalDevice, VkSurfaceKHR Surface, uint32_t* PresentModeCountPtr, VkPresentModeKHR* PresentModesPtr)
+	{
+		PrintfBeginResult(FString::Printf(TEXT("vkGetPhysicalDeviceSurfacePresentModesKHR(PhysicalDevice=%p, Surface=%p, PresentModeCountPtr=%d, PresentModesPtr=%p)"), PhysicalDevice, Surface, PresentModeCountPtr ? *PresentModeCountPtr:0, PresentModesPtr));
+		VkResult Result = VULKANAPINAMESPACE::vkGetPhysicalDeviceSurfacePresentModesKHR(PhysicalDevice, Surface, PresentModeCountPtr, PresentModesPtr);
+		PrintResult(Result);
+		return Result;
+	}
+
+#if PLATFORM_ANDROID
+	static FORCEINLINE_DEBUGGABLE VkResult vkCreateAndroidSurfaceKHR(VkInstance Instance, const VkAndroidSurfaceCreateInfoKHR* CreateInfo, const VkAllocationCallbacks* Allocator, VkSurfaceKHR* Surface)
+	{
+		PrintfBeginResult(FString::Printf(TEXT("vkCreateAndroidSurfaceKHR(Instance=%p, CreateInfo=%p, Allocator=%p, Surface=%p)[...]"), Instance, CreateInfo, Allocator, Surface));
+
+		VkResult Result = VULKANAPINAMESPACE::vkCreateAndroidSurfaceKHR(Instance, CreateInfo, Allocator, Surface);
+		PrintResult(Result);
+		return Result;
+	}
+#endif
+	
+
 #if VULKAN_ENABLE_DUMP_LAYER
 	void FlushDebugWrapperLog();
 #endif
diff --git a/Engine/Source/Runtime/VulkanRHI/Private/VulkanDebug.cpp b/Engine/Source/Runtime/VulkanRHI/Private/VulkanDebug.cpp
--- a/Engine/Source/Runtime/VulkanRHI/Private/VulkanDebug.cpp
+++ b/Engine/Source/Runtime/VulkanRHI/Private/VulkanDebug.cpp
@@ -439,7 +439,9 @@ namespace VulkanRHI
 			VKSWITCHCASE(VK_ERROR_OUT_OF_DATE_KHR)
 			VKSWITCHCASE(VK_ERROR_INCOMPATIBLE_DISPLAY_KHR)
 			VKSWITCHCASE(VK_ERROR_VALIDATION_FAILED_EXT)
+#if VK_HEADER_VERSION >= 13
 			VKSWITCHCASE(VK_ERROR_INVALID_SHADER_NV)
+#endif
 #undef VKSWITCHCASE
 		default:
 			break;
@@ -831,6 +833,30 @@ namespace VulkanRHI
 		}
 	}
 
+	void PrintResultAndPointer(VkResult Result, uint64 Handle)
+	{
+		if (CVarVulkanDumpLayer.GetValueOnAnyThread())
+		{
+			DebugLog += FString::Printf(TEXT(" -> %s => %ull\n"), *GetVkResultErrorString(Result), Handle);
+			if (Result < VK_SUCCESS)
+			{
+				FlushDebugWrapperLog();
+			}
+		}
+	}
+
+	void PrintResultAndNamedHandle(VkResult Result, const TCHAR* HandleName, uint64 Handle)
+	{
+		if (CVarVulkanDumpLayer.GetValueOnAnyThread())
+		{
+			DebugLog += FString::Printf(TEXT(" -> %s => %s=%ull\n"), *GetVkResultErrorString(Result), HandleName, Handle);
+			if (Result < VK_SUCCESS)
+			{
+				FlushDebugWrapperLog();
+			}
+		}
+	}
+
 	void DumpPhysicalDeviceProperties(VkPhysicalDeviceMemoryProperties* Properties)
 	{
 		if (CVarVulkanDumpLayer.GetValueOnAnyThread())
@@ -1352,13 +1378,13 @@ namespace VulkanRHI
 		if (CVarVulkanDumpLayer.GetValueOnAnyThread())
 		{
 			FlushDebugWrapperLog();
-			auto GetSubpassContents = [](VkSubpassContents Contents) -> FString
+			auto GetSubpassContents = [](VkSubpassContents InContents) -> FString
 				{
-					switch (Contents)
+					switch (InContents)
 					{
-					case VK_SUBPASS_CONTENTS_INLINE: return TEXT("INLINE");
-					case VK_SUBPASS_CONTENTS_SECONDARY_COMMAND_BUFFERS: return TEXT("SECONDARY_CMD_BUFS");
-					default: return FString::Printf(TEXT("%d"), (int32)Contents);
+						case VK_SUBPASS_CONTENTS_INLINE: return TEXT("INLINE");
+						case VK_SUBPASS_CONTENTS_SECONDARY_COMMAND_BUFFERS: return TEXT("SECONDARY_CMD_BUFS");
+						default: return FString::Printf(TEXT("%d"), (int32)InContents);
 					}					
 				};
 			CmdPrintfBegin(CommandBuffer, FString::Printf(TEXT("vkCmdBeginRenderPass(BeginInfo=%p, Contents=%s)"), RenderPassBegin, *GetSubpassContents(Contents)));
diff --git a/Engine/Source/Runtime/VulkanRHI/Private/VulkanLoader.h b/Engine/Source/Runtime/VulkanRHI/Private/VulkanLoader.h
--- a/Engine/Source/Runtime/VulkanRHI/Private/VulkanLoader.h
+++ b/Engine/Source/Runtime/VulkanRHI/Private/VulkanLoader.h
@@ -180,7 +180,7 @@
 
 // Declare all Vulkan functions
 #define DECLARE_VK_ENTRYPOINTS(Type,Func) extern Type Func;
-namespace VulkanRHI
+namespace VulkanDynamicAPI
 {
 	ENUM_VK_ENTRYPOINTS_ALL(DECLARE_VK_ENTRYPOINTS);
 }
diff --git a/Engine/Source/Runtime/VulkanRHI/Private/VulkanPipeline.cpp b/Engine/Source/Runtime/VulkanRHI/Private/VulkanPipeline.cpp
--- a/Engine/Source/Runtime/VulkanRHI/Private/VulkanPipeline.cpp
+++ b/Engine/Source/Runtime/VulkanRHI/Private/VulkanPipeline.cpp
@@ -340,11 +340,11 @@ void FVulkanPipelineStateCache::Save(FString& CacheFilename)
 	Ar << DiskEntries;
 	TArray<uint8> DeviceCache;
 	size_t Size = 0;
-	VERIFYVULKANRESULT(vkGetPipelineCacheData(Device->GetInstanceHandle(), PipelineCache, &Size, nullptr));
+	VERIFYVULKANRESULT(VulkanRHI::vkGetPipelineCacheData(Device->GetInstanceHandle(), PipelineCache, &Size, nullptr));
 	if (Size > 0)
 	{
 		DeviceCache.AddUninitialized(Size);
-		VERIFYVULKANRESULT(vkGetPipelineCacheData(Device->GetInstanceHandle(), PipelineCache, &Size, DeviceCache.GetData()));
+		VERIFYVULKANRESULT(VulkanRHI::vkGetPipelineCacheData(Device->GetInstanceHandle(), PipelineCache, &Size, DeviceCache.GetData()));
 	}
 	Ar << DeviceCache;
 
diff --git a/Engine/Source/Runtime/VulkanRHI/Private/VulkanRHI.cpp b/Engine/Source/Runtime/VulkanRHI/Private/VulkanRHI.cpp
--- a/Engine/Source/Runtime/VulkanRHI/Private/VulkanRHI.cpp
+++ b/Engine/Source/Runtime/VulkanRHI/Private/VulkanRHI.cpp
@@ -30,7 +30,7 @@ static_assert(VK_API_VERSION >= UE_VK_API_VERSION, "Vulkan SDK is older than the
 
 // Vulkan function pointers
 
-#define DEFINE_VK_ENTRYPOINTS(Type,Func) Type VulkanRHI::Func = NULL;
+#define DEFINE_VK_ENTRYPOINTS(Type,Func) Type VulkanDynamicAPI::Func = NULL;
 ENUM_VK_ENTRYPOINTS_ALL(DEFINE_VK_ENTRYPOINTS)
 
 #endif
@@ -158,10 +158,10 @@ static bool LoadVulkanLibrary()
 	}
 
 	bool bFoundAllEntryPoints = true;
-#define CHECK_VK_ENTRYPOINTS(Type,Func) if (VulkanRHI::Func == NULL) { bFoundAllEntryPoints = false; UE_LOG(LogRHI, Warning, TEXT("Failed to find entry point for %s"), TEXT(#Func)); }
+#define CHECK_VK_ENTRYPOINTS(Type,Func) if (VulkanDynamicAPI::Func == NULL) { bFoundAllEntryPoints = false; UE_LOG(LogRHI, Warning, TEXT("Failed to find entry point for %s"), TEXT(#Func)); }
 
 	// Initialize all of the entry points we have to query manually
-#define GET_VK_ENTRYPOINTS(Type,Func) VulkanRHI::Func = (Type)dlsym(VulkanLib, #Func);
+#define GET_VK_ENTRYPOINTS(Type,Func) VulkanDynamicAPI::Func = (Type)dlsym(VulkanLib, #Func);
 	ENUM_VK_ENTRYPOINTS_BASE(GET_VK_ENTRYPOINTS);
 	ENUM_VK_ENTRYPOINTS_BASE(CHECK_VK_ENTRYPOINTS);
 	if (!bFoundAllEntryPoints)
@@ -180,9 +180,9 @@ static bool LoadVulkanLibrary()
 static bool LoadVulkanInstanceFunctions(VkInstance inInstance)
 {
 	bool bFoundAllEntryPoints = true;
-#define CHECK_VK_ENTRYPOINTS(Type,Func) if (VulkanRHI::Func == NULL) { bFoundAllEntryPoints = false; UE_LOG(LogRHI, Warning, TEXT("Failed to find entry point for %s"), TEXT(#Func)); }
+#define CHECK_VK_ENTRYPOINTS(Type,Func) if (VulkanDynamicAPI::Func == NULL) { bFoundAllEntryPoints = false; UE_LOG(LogRHI, Warning, TEXT("Failed to find entry point for %s"), TEXT(#Func)); }
 
-#define GETINSTANCE_VK_ENTRYPOINTS(Type, Func) VulkanRHI::Func = (Type)VulkanRHI::vkGetInstanceProcAddr(inInstance, #Func);
+#define GETINSTANCE_VK_ENTRYPOINTS(Type, Func) VulkanDynamicAPI::Func = (Type)VulkanDynamicAPI::vkGetInstanceProcAddr(inInstance, #Func);
 	ENUM_VK_ENTRYPOINTS_INSTANCE(GETINSTANCE_VK_ENTRYPOINTS);
 	ENUM_VK_ENTRYPOINTS_INSTANCE(CHECK_VK_ENTRYPOINTS);
 
@@ -193,7 +193,7 @@ static void FreeVulkanLibrary()
 {
 	if (VulkanLib != nullptr)
 	{
-#define CLEAR_VK_ENTRYPOINTS(Type,Func) VulkanRHI::Func = nullptr;
+#define CLEAR_VK_ENTRYPOINTS(Type,Func) VulkanDynamicAPI::Func = nullptr;
 		ENUM_VK_ENTRYPOINTS_ALL(CLEAR_VK_ENTRYPOINTS);
 
 		dlclose(VulkanLib);
diff --git a/Engine/Source/Runtime/VulkanRHI/Private/VulkanRHIPrivate.h b/Engine/Source/Runtime/VulkanRHI/Private/VulkanRHIPrivate.h
--- a/Engine/Source/Runtime/VulkanRHI/Private/VulkanRHIPrivate.h
+++ b/Engine/Source/Runtime/VulkanRHI/Private/VulkanRHIPrivate.h
@@ -27,20 +27,45 @@ DECLARE_LOG_CATEGORY_EXTERN(LogVulkanRHI, Log, All);
 #endif
 
 #if PLATFORM_ANDROID
-#include "VulkanLoader.h"
-#define VULKAN_COMMANDWRAPPERS_ENABLE 0
+	#define VULKAN_COMMANDWRAPPERS_ENABLE VULKAN_ENABLE_API_DUMP
+	#define VULKAN_DYNAMICALLYLOADED 1
 #else
-#include <vulkan/vulkan.h>
-#define VULKAN_COMMANDWRAPPERS_ENABLE 1
+	#define VULKAN_COMMANDWRAPPERS_ENABLE 1
+	#define VULKAN_DYNAMICALLYLOADED 0
+#endif
+
+#if VULKAN_DYNAMICALLYLOADED
+	#include "VulkanLoader.h"
+#else
+	#include <vulkan/vulkan.h>
+#endif
+
+#if VULKAN_COMMANDWRAPPERS_ENABLE
+	#if VULKAN_DYNAMICALLYLOADED
+		// Vulkan API is defined in VulkanDynamicAPI namespace.
+		#define VULKANAPINAMESPACE VulkanDynamicAPI
+	#else
+		// Vulkan API is in the global namespace.
+		#define VULKANAPINAMESPACE
+	#endif
+	#include "VulkanCommandWrappers.h"
+#else
+	#if VULKAN_DYNAMICALLYLOADED
+		// Bring functions from VulkanDynamicAPI to VulkanRHI
+		#define VK_DYNAMICAPI_TO_VULKANRHI(Type,Func) using VulkanDynamicAPI::Func;
+		namespace VulkanRHI
+		{
+			ENUM_VK_ENTRYPOINTS_ALL(VK_DYNAMICAPI_TO_VULKANRHI);
+		}
+	#else
+		#error "Statically linked vulkan api must be wrapped!"
+	#endif
 #endif
 
 #include "VulkanRHI.h"
 #include "VulkanGlobalUniformBuffer.h"
 #include "RHI.h"
 
-#if VULKAN_COMMANDWRAPPERS_ENABLE
-#include "VulkanCommandWrappers.h"
-#endif
 using namespace VulkanRHI;
 
 // Default is 1 (which is aniso off), the number is adjusted after the limits are queried.
diff --git a/Engine/Source/Runtime/VulkanRHI/Private/VulkanSwapChain.cpp b/Engine/Source/Runtime/VulkanRHI/Private/VulkanSwapChain.cpp
--- a/Engine/Source/Runtime/VulkanRHI/Private/VulkanSwapChain.cpp
+++ b/Engine/Source/Runtime/VulkanRHI/Private/VulkanSwapChain.cpp
@@ -38,12 +38,12 @@ FVulkanSwapChain::FVulkanSwapChain(VkInstance Instance, FVulkanDevice& InDevice,
 	FMemory::Memzero(CurrFormat);
 	{
 		uint32 NumFormats;
-		VERIFYVULKANRESULT_EXPANDED(vkGetPhysicalDeviceSurfaceFormatsKHR(Device.GetPhysicalHandle(), Surface, &NumFormats, nullptr));
+		VERIFYVULKANRESULT_EXPANDED(VulkanRHI::vkGetPhysicalDeviceSurfaceFormatsKHR(Device.GetPhysicalHandle(), Surface, &NumFormats, nullptr));
 		check(NumFormats > 0);
 
 		TArray<VkSurfaceFormatKHR> Formats;
 		Formats.AddZeroed(NumFormats);
-		VERIFYVULKANRESULT_EXPANDED(vkGetPhysicalDeviceSurfaceFormatsKHR(Device.GetPhysicalHandle(), Surface, &NumFormats, Formats.GetData()));
+		VERIFYVULKANRESULT_EXPANDED(VulkanRHI::vkGetPhysicalDeviceSurfaceFormatsKHR(Device.GetPhysicalHandle(), Surface, &NumFormats, Formats.GetData()));
 
 		if (Formats.Num() == 1 && Formats[0].format == VK_FORMAT_UNDEFINED && InOutPixelFormat == PF_Unknown)
 		{
@@ -85,7 +85,7 @@ FVulkanSwapChain::FVulkanSwapChain(VkInstance Instance, FVulkanDevice& InDevice,
 
 	//#todo-rco: Check multiple Gfx Queues?
 	VkBool32 bSupportsPresent = VK_FALSE;
-	VERIFYVULKANRESULT(vkGetPhysicalDeviceSurfaceSupportKHR(Device.GetPhysicalHandle(), Device.GetQueue()->GetFamilyIndex(), Surface, &bSupportsPresent));
+	VERIFYVULKANRESULT(VulkanRHI::vkGetPhysicalDeviceSurfaceSupportKHR(Device.GetPhysicalHandle(), Device.GetQueue()->GetFamilyIndex(), Surface, &bSupportsPresent));
 	//#todo-rco: Find separate present queue if the gfx one doesn't support presents
 	check(bSupportsPresent);
 
@@ -94,12 +94,12 @@ FVulkanSwapChain::FVulkanSwapChain(VkInstance Instance, FVulkanDevice& InDevice,
 #if !PLATFORM_ANDROID
 	{
 		uint32 NumFoundPresentModes = 0;
-		VERIFYVULKANRESULT(vkGetPhysicalDeviceSurfacePresentModesKHR(Device.GetPhysicalHandle(), Surface, &NumFoundPresentModes, nullptr));
+		VERIFYVULKANRESULT(VulkanRHI::vkGetPhysicalDeviceSurfacePresentModesKHR(Device.GetPhysicalHandle(), Surface, &NumFoundPresentModes, nullptr));
 		check(NumFoundPresentModes > 0);
 
 		TArray<VkPresentModeKHR> FoundPresentModes;
 		FoundPresentModes.AddZeroed(NumFoundPresentModes);
-		VERIFYVULKANRESULT(vkGetPhysicalDeviceSurfacePresentModesKHR(Device.GetPhysicalHandle(), Surface, &NumFoundPresentModes, FoundPresentModes.GetData()));
+		VERIFYVULKANRESULT(VulkanRHI::vkGetPhysicalDeviceSurfacePresentModesKHR(Device.GetPhysicalHandle(), Surface, &NumFoundPresentModes, FoundPresentModes.GetData()));
 
 		bool bFoundDesiredMode = false;
 		for (size_t i = 0; i < NumFoundPresentModes; i++)
@@ -121,7 +121,7 @@ FVulkanSwapChain::FVulkanSwapChain(VkInstance Instance, FVulkanDevice& InDevice,
 	// Check the surface properties and formats
 	
 	VkSurfaceCapabilitiesKHR SurfProperties;
-	VERIFYVULKANRESULT_EXPANDED(vkGetPhysicalDeviceSurfaceCapabilitiesKHR(Device.GetPhysicalHandle(),
+	VERIFYVULKANRESULT_EXPANDED(VulkanRHI::vkGetPhysicalDeviceSurfaceCapabilitiesKHR(Device.GetPhysicalHandle(),
 		Surface,
 		&SurfProperties));
 	VkSurfaceTransformFlagBitsKHR PreTransform;
diff --git a/Engine/Source/Runtime/VulkanRHI/Private/VulkanTexture.cpp b/Engine/Source/Runtime/VulkanRHI/Private/VulkanTexture.cpp
--- a/Engine/Source/Runtime/VulkanRHI/Private/VulkanTexture.cpp
+++ b/Engine/Source/Runtime/VulkanRHI/Private/VulkanTexture.cpp
@@ -1269,6 +1269,8 @@ void FVulkanDynamicRHI::RHIBindDebugLabelName(FTextureRHIParamRef TextureRHI, co
 {
 #if VULKAN_ENABLE_DUMP_LAYER
 	{
+// TODO: this dies in the printf on android. Needs investigation.
+#if !PLATFORM_ANDROID
 		if (FRHITexture2D* Tex2d = TextureRHI->GetTexture2D())
 		{
 			FVulkanTexture2D* VTex2d = (FVulkanTexture2D*)Tex2d;
@@ -1284,7 +1286,8 @@ void FVulkanDynamicRHI::RHIBindDebugLabelName(FTextureRHIParamRef TextureRHI, co
 			FVulkanTexture3D* VTex3d = (FVulkanTexture3D*)Tex3d;
 			VulkanRHI::PrintfBegin(FString::Printf(TEXT("vkDebugMarkerSetObjectNameEXT(%p=%s)"), VTex3d->Surface.Image, Name));
 		}
-}
+	#endif
+	}
 #endif
 
 #if VULKAN_ENABLE_DRAW_MARKERS
