UE - Fix static analysis warnings

diff --git a/Engine/Source/Runtime/VulkanRHI/Private/VulkanRenderTarget.cpp b/Engine/Source/Runtime/VulkanRHI/Private/VulkanRenderTarget.cpp
--- a/Engine/Source/Runtime/VulkanRHI/Private/VulkanRenderTarget.cpp
+++ b/Engine/Source/Runtime/VulkanRHI/Private/VulkanRenderTarget.cpp
@@ -382,30 +382,30 @@ void FVulkanDynamicRHI::RHIUnmapStagingSurface(FTextureRHIParamRef TextureRHI)
 
 void FVulkanDynamicRHI::RHIReadSurfaceFloatData(FTextureRHIParamRef TextureRHI, FIntRect Rect, TArray<FFloat16Color>& OutData, ECubeFace CubeFace,int32 ArrayIndex,int32 MipIndex)
 {
-	auto DoCopyFloat = [](FVulkanDevice* InDevice, FVulkanCmdBuffer* CmdBuffer, const FVulkanSurface& Surface, uint32 MipIndex, uint32 SrcBaseArrayLayer, FIntRect Rect, TArray<FFloat16Color>& OutData)
+	auto DoCopyFloat = [](FVulkanDevice* InDevice, FVulkanCmdBuffer* InCmdBuffer, const FVulkanSurface& Surface, uint32 InMipIndex, uint32 SrcBaseArrayLayer, FIntRect InRect, TArray<FFloat16Color>& OutputData)
 	{
 		ensure(Surface.InternalFormat == VK_FORMAT_R16G16B16A16_SFLOAT);
 
-		uint32 NumPixels = (Surface.Width >> MipIndex) * (Surface.Height >> MipIndex);
+		uint32 NumPixels = (Surface.Width >> InMipIndex) * (Surface.Height >> InMipIndex);
 		const uint32 Size = NumPixels * sizeof(FLinearColor);
 		VulkanRHI::FStagingBuffer* StagingBuffer = InDevice->GetStagingManager().AcquireBuffer(Size, VK_BUFFER_USAGE_TRANSFER_DST_BIT, true);
 		VkBufferImageCopy CopyRegion;
 		FMemory::Memzero(CopyRegion);
 		//Region.bufferOffset = 0;
-		CopyRegion.bufferRowLength = Surface.Width >> MipIndex;
-		CopyRegion.bufferImageHeight = Surface.Height >> MipIndex;
+		CopyRegion.bufferRowLength = Surface.Width >> InMipIndex;
+		CopyRegion.bufferImageHeight = Surface.Height >> InMipIndex;
 		CopyRegion.imageSubresource.aspectMask = VK_IMAGE_ASPECT_COLOR_BIT;
-		CopyRegion.imageSubresource.mipLevel = MipIndex;
+		CopyRegion.imageSubresource.mipLevel = InMipIndex;
 		CopyRegion.imageSubresource.baseArrayLayer = SrcBaseArrayLayer;
 		CopyRegion.imageSubresource.layerCount = 1;
-		CopyRegion.imageExtent.width = Surface.Width >> MipIndex;
-		CopyRegion.imageExtent.height = Surface.Height >> MipIndex;
+		CopyRegion.imageExtent.width = Surface.Width >> InMipIndex;
+		CopyRegion.imageExtent.height = Surface.Height >> InMipIndex;
 		CopyRegion.imageExtent.depth = 1;
-		VulkanRHI::vkCmdCopyImageToBuffer(CmdBuffer->GetHandle(), Surface.Image, VK_IMAGE_LAYOUT_GENERAL, StagingBuffer->GetHandle(), 1, &CopyRegion);
+		VulkanRHI::vkCmdCopyImageToBuffer(InCmdBuffer->GetHandle(), Surface.Image, VK_IMAGE_LAYOUT_GENERAL, StagingBuffer->GetHandle(), 1, &CopyRegion);
 
 		VkBufferMemoryBarrier Barrier;
 		VulkanRHI::SetupAndZeroBufferBarrier(Barrier, VK_ACCESS_TRANSFER_WRITE_BIT, VK_ACCESS_MEMORY_READ_BIT, StagingBuffer->GetHandle(), StagingBuffer->GetOffset(), Size);
-		VulkanRHI::vkCmdPipelineBarrier(CmdBuffer->GetHandle(), VK_PIPELINE_STAGE_TRANSFER_BIT, VK_PIPELINE_STAGE_HOST_BIT, 0, 0, nullptr, 1, &Barrier, 0, nullptr);
+		VulkanRHI::vkCmdPipelineBarrier(InCmdBuffer->GetHandle(), VK_PIPELINE_STAGE_TRANSFER_BIT, VK_PIPELINE_STAGE_HOST_BIT, 0, 0, nullptr, 1, &Barrier, 0, nullptr);
 
 		// Force upload
 		InDevice->GetImmediateContext().GetCommandBufferManager()->SubmitUploadCmdBuffer(true);
@@ -420,17 +420,17 @@ void FVulkanDynamicRHI::RHIReadSurfaceFloatData(FTextureRHIParamRef TextureRHI,
 		VulkanRHI::vkInvalidateMappedMemoryRanges(InDevice->GetInstanceHandle(), 1, &MappedRange);
 		VulkanRHI::vkFlushMappedMemoryRanges(InDevice->GetInstanceHandle(), 1, &MappedRange);
 
-		OutData.SetNum(NumPixels);
-		FFloat16Color* Dest = OutData.GetData();
-		for (int32 Row = Rect.Min.Y; Row < Rect.Max.Y; ++Row)
+		OutputData.SetNum(NumPixels);
+		FFloat16Color* Dest = OutputData.GetData();
+		for (int32 Row = InRect.Min.Y; Row < InRect.Max.Y; ++Row)
 		{
-			FLinearColor* Src = (FLinearColor*)StagingBuffer->GetMappedPointer() + Row * (Surface.Width >> MipIndex) + Rect.Min.X;
-			for (int32 Col = Rect.Min.X; Col < Rect.Max.X; ++Col)
+			FLinearColor* Src = (FLinearColor*)StagingBuffer->GetMappedPointer() + Row * (Surface.Width >> InMipIndex) + InRect.Min.X;
+			for (int32 Col = InRect.Min.X; Col < InRect.Max.X; ++Col)
 			{
 				*Dest++ = FFloat16Color(*Src++);
 			}
 		}
-		InDevice->GetStagingManager().ReleaseBuffer(CmdBuffer, StagingBuffer);
+		InDevice->GetStagingManager().ReleaseBuffer(InCmdBuffer, StagingBuffer);
 	};
 
 	FVulkanCmdBuffer* CmdBuffer = Device->GetImmediateContext().GetCommandBufferManager()->GetUploadCmdBuffer();
