Temporarily comment out VK_FORMAT_PVRTC* not in current TP VulkanSDK

diff --git a/Engine/Source/Runtime/VulkanRHI/Private/VulkanRHIPrivate.h b/Engine/Source/Runtime/VulkanRHI/Private/VulkanRHIPrivate.h
--- a/Engine/Source/Runtime/VulkanRHI/Private/VulkanRHIPrivate.h
+++ b/Engine/Source/Runtime/VulkanRHI/Private/VulkanRHIPrivate.h
@@ -705,10 +705,10 @@ inline VkFormat UEToVkFormat(EPixelFormat UEFormat, const bool bIsSRGB)
 		case VK_FORMAT_ASTC_10x10_UNORM_BLOCK:		Format = VK_FORMAT_ASTC_10x10_SRGB_BLOCK; break;
 		case VK_FORMAT_ASTC_12x10_UNORM_BLOCK:		Format = VK_FORMAT_ASTC_12x10_SRGB_BLOCK; break;
 		case VK_FORMAT_ASTC_12x12_UNORM_BLOCK:		Format = VK_FORMAT_ASTC_12x12_SRGB_BLOCK; break;
-		case VK_FORMAT_PVRTC1_2BPP_UNORM_BLOCK_IMG:	Format = VK_FORMAT_PVRTC1_2BPP_SRGB_BLOCK_IMG; break;
-		case VK_FORMAT_PVRTC1_4BPP_UNORM_BLOCK_IMG:	Format = VK_FORMAT_PVRTC1_4BPP_SRGB_BLOCK_IMG; break;
-		case VK_FORMAT_PVRTC2_2BPP_UNORM_BLOCK_IMG:	Format = VK_FORMAT_PVRTC2_2BPP_SRGB_BLOCK_IMG; break;
-		case VK_FORMAT_PVRTC2_4BPP_UNORM_BLOCK_IMG:	Format = VK_FORMAT_PVRTC2_4BPP_SRGB_BLOCK_IMG; break;
+//		case VK_FORMAT_PVRTC1_2BPP_UNORM_BLOCK_IMG:	Format = VK_FORMAT_PVRTC1_2BPP_SRGB_BLOCK_IMG; break;
+//		case VK_FORMAT_PVRTC1_4BPP_UNORM_BLOCK_IMG:	Format = VK_FORMAT_PVRTC1_4BPP_SRGB_BLOCK_IMG; break;
+//		case VK_FORMAT_PVRTC2_2BPP_UNORM_BLOCK_IMG:	Format = VK_FORMAT_PVRTC2_2BPP_SRGB_BLOCK_IMG; break;
+//		case VK_FORMAT_PVRTC2_4BPP_UNORM_BLOCK_IMG:	Format = VK_FORMAT_PVRTC2_4BPP_SRGB_BLOCK_IMG; break;
 		default:	break;
 		}
 	}
