UE4 - Fixes for static analysis

diff --git a/Engine/Source/Runtime/VulkanRHI/Private/VulkanQuery.cpp b/Engine/Source/Runtime/VulkanRHI/Private/VulkanQuery.cpp
--- a/Engine/Source/Runtime/VulkanRHI/Private/VulkanQuery.cpp
+++ b/Engine/Source/Runtime/VulkanRHI/Private/VulkanQuery.cpp
@@ -321,7 +321,6 @@ void FVulkanCommandListContext::RHIBeginRenderQuery(FRenderQueryRHIParamRef Quer
 
 	if (Query->QueryType == RQT_Occlusion)
 	{
-		ensure(CurrentOcclusionQueryPool);
 		FVulkanOcclusionQueryPool& PoolRef = Device->GetAvailableOcclusionQueryPool();
 		FVulkanOcclusionQueryPool* Pool = &PoolRef;
 		if (CurrentOcclusionQueryPool != Pool)
diff --git a/Engine/Source/Runtime/VulkanRHI/Private/VulkanRenderTarget.cpp b/Engine/Source/Runtime/VulkanRHI/Private/VulkanRenderTarget.cpp
--- a/Engine/Source/Runtime/VulkanRHI/Private/VulkanRenderTarget.cpp
+++ b/Engine/Source/Runtime/VulkanRHI/Private/VulkanRenderTarget.cpp
@@ -696,7 +696,7 @@ void FVulkanCommandListContext::RHITransitionResources(EResourceTransitionAccess
 		FVulkanCmdBuffer* CmdBuffer = CommandBufferManager->GetActiveCmdBuffer();
 		check(CmdBuffer->HasBegun());
 
-		auto SetImageLayout = [](FRenderPassState& InRenderPassState, VkCommandBuffer CmdBuffer, FVulkanSurface& Surface, uint32 NumArraySlices)
+		auto SetImageLayout = [](FRenderPassState& InRenderPassState, VkCommandBuffer InCmdBuffer, FVulkanSurface& Surface, uint32 NumArraySlices)
 		{
 			if (InRenderPassState.CurrentRenderPass)
 			{
@@ -732,13 +732,13 @@ void FVulkanCommandListContext::RHITransitionResources(EResourceTransitionAccess
 
 			if ((AspectMask & VK_IMAGE_ASPECT_COLOR_BIT) != 0)
 			{
-				VulkanSetImageLayout(CmdBuffer, Surface.Image, SrcLayout, VK_IMAGE_LAYOUT_COLOR_ATTACHMENT_OPTIMAL, SubresourceRange);
+				VulkanSetImageLayout(InCmdBuffer, Surface.Image, SrcLayout, VK_IMAGE_LAYOUT_COLOR_ATTACHMENT_OPTIMAL, SubresourceRange);
 				InRenderPassState.CurrentLayout.FindOrAdd(Surface.Image) = VK_IMAGE_LAYOUT_COLOR_ATTACHMENT_OPTIMAL;
 			}
 			else
 			{
 				check((AspectMask & (VK_IMAGE_ASPECT_DEPTH_BIT | VK_IMAGE_ASPECT_STENCIL_BIT)) != 0);
-				VulkanSetImageLayout(CmdBuffer, Surface.Image, SrcLayout, VK_IMAGE_LAYOUT_DEPTH_STENCIL_ATTACHMENT_OPTIMAL, SubresourceRange);
+				VulkanSetImageLayout(InCmdBuffer, Surface.Image, SrcLayout, VK_IMAGE_LAYOUT_DEPTH_STENCIL_ATTACHMENT_OPTIMAL, SubresourceRange);
 				InRenderPassState.CurrentLayout.FindOrAdd(Surface.Image) = VK_IMAGE_LAYOUT_DEPTH_STENCIL_ATTACHMENT_OPTIMAL;
 			}
 		};
