UE-43733 //UE4/Dev-VR: Step 'IncludeTool Static Analysis Win64'
-Fixes two 'includes after first code block' warnings.
-Fixes 'missing header guard' warning.  Just had to add one.

diff --git a/Engine/Source/Runtime/VulkanRHI/Private/VulkanBoundShaderState.h b/Engine/Source/Runtime/VulkanRHI/Private/VulkanBoundShaderState.h
new file mode 100644
--- /dev/null
+++ b/Engine/Source/Runtime/VulkanRHI/Private/VulkanBoundShaderState.h
@@ -0,0 +1,236 @@
+// Copyright 1998-2017 Epic Games, Inc. All Rights Reserved.
+
+/*=============================================================================
+	VulkanBoundShaderState.h: 
+=============================================================================*/
+
+#pragma once 
+
+#include "VulkanResources.h"
+#include "../Private/VulkanDescriptorSets.h"
+
+/**
+* Combined shader state and vertex definition for rendering geometry.
+* Each unique instance consists of a vertex decl, vertex shader, and pixel shader.
+*/
+class FVulkanBoundShaderState : public FVulkanDescriptorSetRingBuffer, public FVulkanDescriptorSetWriteContainer, public FRHIBoundShaderState
+{
+public:
+	FVulkanBoundShaderState(
+		FVulkanDevice* InDevice,
+		FVertexDeclarationRHIParamRef InVertexDeclarationRHI,
+		FVertexShaderRHIParamRef InVertexShaderRHI,
+		FPixelShaderRHIParamRef InPixelShaderRHI,
+		FHullShaderRHIParamRef InHullShaderRHI,
+		FDomainShaderRHIParamRef InDomainShaderRHI,
+		FGeometryShaderRHIParamRef InGeometryShaderRHI);
+
+	virtual ~FVulkanBoundShaderState();
+
+	// Called when the shader is set by the engine as current pending state
+	// After the shader is set, it is expected that all require states are provided/set by the engine
+	void ResetState();
+
+	const FVulkanLayout& GetLayout() const
+	{
+		return Layout;
+	}
+
+	class FVulkanGfxPipeline* PrepareForDraw(class FVulkanRenderPass* RenderPass, const struct FVulkanPipelineGraphicsKey& PipelineKey, uint32 VertexInputKey, const struct FVulkanGfxPipelineState& State);
+
+	bool UpdateDescriptorSets(FVulkanCommandListContext* CmdListContext, FVulkanCmdBuffer* CmdBuffer, FVulkanGlobalUniformPool* GlobalUniformPool);
+
+	void BindDescriptorSets(FVulkanCmdBuffer* Cmd);
+
+	inline void BindVertexStreams(FVulkanCmdBuffer* Cmd, const void* VertexStreams)
+	{
+		if (bDirtyVertexStreams)
+		{
+			InternalBindVertexStreams(Cmd, VertexStreams);
+			bDirtyVertexStreams = false;
+		}
+	}
+
+	inline void MarkDirtyVertexStreams()
+	{
+		bDirtyVertexStreams = true;
+	}
+
+	inline const FVulkanVertexInputStateInfo& GetVertexInputStateInfo() const
+	{
+		check(VertexInputStateInfo.Info.sType == VK_STRUCTURE_TYPE_PIPELINE_VERTEX_INPUT_STATE_CREATE_INFO);
+		return VertexInputStateInfo;
+	}
+
+	inline FSHAHash const* GetShaderHashes() const
+	{
+		return ShaderHashes;
+	}
+
+	// Binding
+	inline void SetShaderParameter(EShaderFrequency Stage, uint32 BufferIndex, uint32 ByteOffset, uint32 NumBytes, const void* NewValue)
+	{
+#if VULKAN_ENABLE_AGGRESSIVE_STATS
+		SCOPE_CYCLE_COUNTER(STAT_VulkanSetShaderParamTime);
+#endif
+		NEWPackedUniformBufferStaging[Stage].SetPackedGlobalParameter(BufferIndex, ByteOffset, NumBytes, NewValue, NEWPackedUniformBufferStagingDirty[Stage]);
+	}
+
+	inline void SetTexture(EShaderFrequency Stage, uint32 BindPoint, const FVulkanTextureBase* VulkanTextureBase)
+	{
+		check(VulkanTextureBase);
+		DescriptorWriteSet[Stage].WriteImage(BindPoint, VulkanTextureBase->PartialView->View,
+			(VulkanTextureBase->Surface.GetFullAspectMask() & (VK_IMAGE_ASPECT_DEPTH_BIT | VK_IMAGE_ASPECT_STENCIL_BIT)) != 0 ? VK_IMAGE_LAYOUT_DEPTH_STENCIL_READ_ONLY_OPTIMAL : VK_IMAGE_LAYOUT_GENERAL);
+	}
+
+	inline void SetSamplerState(EShaderFrequency Stage, uint32 BindPoint, FVulkanSamplerState* Sampler)
+	{
+		check(Sampler);
+		DescriptorWriteSet[Stage].WriteSampler(BindPoint, Sampler->Sampler);
+	}
+
+	inline void SetBufferViewState(EShaderFrequency Stage, uint32 BindPoint, FVulkanBufferView* View)
+	{
+		check(View && (View->Flags & VK_BUFFER_USAGE_UNIFORM_TEXEL_BUFFER_BIT) == VK_BUFFER_USAGE_UNIFORM_TEXEL_BUFFER_BIT);
+		DescriptorWriteSet[Stage].WriteUniformTexelBuffer(BindPoint, View);
+	}
+
+	inline void SetTextureView(EShaderFrequency Stage, uint32 BindPoint, const FVulkanTextureView& TextureView)
+	{
+		DescriptorWriteSet[Stage].WriteImage(BindPoint, TextureView.View, VK_IMAGE_LAYOUT_GENERAL);
+	}
+
+	inline void SetUniformBufferConstantData(EShaderFrequency Stage, uint32 BindPoint, const TArray<uint8>& ConstantData)
+	{
+		FVulkanShader* Shader = GetShaderPtr(Stage);
+		NEWPackedUniformBufferStaging[Stage].SetEmulatedUniformBufferIntoPacked(BindPoint, ConstantData, NEWPackedUniformBufferStagingDirty[Stage]);
+	}
+
+	void SetUniformBuffer(EShaderFrequency Stage, uint32 BindPoint, const FVulkanUniformBuffer* UniformBuffer);
+
+	void SetSRV(EShaderFrequency Stage, uint32 TextureIndex, FVulkanShaderResourceView* SRV);
+
+	inline const FVulkanShader* GetShaderPtr(EShaderFrequency Stage) const
+	{
+		switch (Stage)
+		{
+		case SF_Vertex:		return VertexShader;
+		case SF_Hull:		return HullShader;
+		case SF_Domain:		return DomainShader;
+		case SF_Pixel:		return PixelShader;
+		case SF_Geometry:	return GeometryShader;
+			//#todo-rco: Really should assert here...
+		default:			return nullptr;
+		}
+	}
+
+	// Has no internal verify.
+	// If stage does not exists, returns a "nullptr".
+	inline FVulkanShader* GetShaderPtr(EShaderFrequency Stage)
+	{
+		switch (Stage)
+		{
+		case SF_Vertex:		return VertexShader;
+		case SF_Hull:		return HullShader;
+		case SF_Domain:		return DomainShader;
+		case SF_Pixel:		return PixelShader;
+		case SF_Geometry:	return GeometryShader;
+			//#todo-rco: Really should assert here...
+		default:			return nullptr;
+		}
+	}
+
+	inline bool HasShaderStage(EShaderFrequency Stage) const
+	{
+		return (GetShaderPtr(Stage) != nullptr);
+	}
+
+	// Verifies if the shader stage exists
+	inline const FVulkanShader& GetShader(EShaderFrequency Stage) const
+	{
+		const FVulkanShader* Shader = GetShaderPtr(Stage);
+		check(Shader);
+		return *Shader;
+	}
+
+	inline void BindPipeline(VkCommandBuffer CmdBuffer, VkPipeline NewPipeline)
+	{
+		if (LastBoundPipeline != NewPipeline)
+		{
+			VulkanRHI::vkCmdBindPipeline(CmdBuffer, VK_PIPELINE_BIND_POINT_GRAPHICS, NewPipeline);
+			LastBoundPipeline = NewPipeline;
+		}
+	}
+
+private:
+	void InternalBindVertexStreams(FVulkanCmdBuffer* Cmd, const void* VertexStreams);
+	FCachedBoundShaderStateLink CacheLink;
+
+	/** Cached vertex structure */
+	TRefCountPtr<FVulkanVertexDeclaration> VertexDeclaration;
+
+	friend class FVulkanPipelineStateCache;
+
+	FSHAHash ShaderHashes[SF_NumFrequencies];
+
+	TLinkedList<FVulkanBoundShaderState*> GlobalListLink;
+
+	/** Cached shaders */
+	TRefCountPtr<FVulkanVertexShader> VertexShader;
+	TRefCountPtr<FVulkanPixelShader> PixelShader;
+	TRefCountPtr<FVulkanHullShader> HullShader;
+	TRefCountPtr<FVulkanDomainShader> DomainShader;
+	TRefCountPtr<FVulkanGeometryShader> GeometryShader;
+
+	// For debugging only
+	int32 ID;
+
+	bool bDirtyVertexStreams;
+
+	FVulkanDescriptorSetWriter DescriptorWriteSet[SF_Compute];
+
+	FPackedUniformBuffers NEWPackedUniformBufferStaging[SF_Compute];
+	uint64 NEWPackedUniformBufferStagingDirty[SF_Compute];
+
+#if VULKAN_ENABLE_RHI_DEBUGGING
+	struct FDebugInfo
+	{
+		const FVulkanTextureBase* Textures[SF_Compute][32];
+		const FVulkanSamplerState* SamplerStates[SF_Compute][32];
+		union
+		{
+			const FVulkanTextureView* SRVIs[SF_Compute][32];
+			const FVulkanBufferView* SRVBs[SF_Compute][32];
+		};
+		const FVulkanUniformBuffer* UBs[SF_Compute][32];
+	};
+	FDebugInfo DebugInfo;
+#endif
+
+	void CreateDescriptorWriteInfos();
+
+
+	FVulkanLayout Layout;
+
+	VkPipeline LastBoundPipeline;
+
+	// Vertex input configuration
+	FVulkanVertexInputStateInfo VertexInputStateInfo;
+
+	// Members in Tmp are normally allocated each frame
+	// To reduce the allocation, we reuse these array
+	struct FTmp
+	{
+		TArray<VkBuffer> VertexBuffers;
+		TArray<VkDeviceSize> VertexOffsets;
+	} Tmp;
+
+	// these are the cache pipeline state objects used in this BSS; RefCounts must be manually controlled for the FVulkanPipelines!
+	TMap<FVulkanPipelineGraphicsKey, FVulkanGfxPipeline* > PipelineCache;
+};
+
+template<>
+struct TVulkanResourceTraits<FRHIBoundShaderState>
+{
+	typedef FVulkanBoundShaderState TConcreteType;
+};
\ No newline at end of file
diff --git a/Engine/Source/Runtime/VulkanRHI/Private/VulkanCommandsDirect.h b/Engine/Source/Runtime/VulkanRHI/Private/VulkanCommandsDirect.h
new file mode 100644
--- /dev/null
+++ b/Engine/Source/Runtime/VulkanRHI/Private/VulkanCommandsDirect.h
@@ -0,0 +1,21 @@
+// Copyright 1998-2017 Epic Games, Inc. All Rights Reserved.
+
+/*=============================================================================
+	VulkanCommandsDirect.h: Enumerate the VulkanDynamicAPI.  Doing this in a header to avoid static analysis warnings.
+=============================================================================*/
+
+#pragma once 
+
+#if VULKAN_COMMANDWRAPPERS_ENABLE
+	#error "VulkanCommandsDirect should not be used when CommandWrappers are enabled."
+#endif
+#if !VULKAN_DYNAMICALLYLOADED
+	#error "VulkanCommandsDirect is for the Dynamic api."
+#endif
+
+// Bring functions from VulkanDynamicAPI to VulkanRHI
+#define VK_DYNAMICAPI_TO_VULKANRHI(Type,Func) using VulkanDynamicAPI::Func;
+namespace VulkanRHI
+{
+	ENUM_VK_ENTRYPOINTS_ALL(VK_DYNAMICAPI_TO_VULKANRHI);
+}
\ No newline at end of file
diff --git a/Engine/Source/Runtime/VulkanRHI/Private/VulkanGlobalUniformBuffer.h b/Engine/Source/Runtime/VulkanRHI/Private/VulkanGlobalUniformBuffer.h
--- a/Engine/Source/Runtime/VulkanRHI/Private/VulkanGlobalUniformBuffer.h
+++ b/Engine/Source/Runtime/VulkanRHI/Private/VulkanGlobalUniformBuffer.h
@@ -6,6 +6,8 @@
 
 #pragma once
 
+#include "VulkanResources.h"
+
 struct FVulkanPooledUniformBuffer : public FRefCountedObject
 {
     FVulkanPooledUniformBuffer(FVulkanDevice& InDevice, uint32 InSize);
diff --git a/Engine/Source/Runtime/VulkanRHI/Private/VulkanLoader.h b/Engine/Source/Runtime/VulkanRHI/Private/VulkanLoader.h
--- a/Engine/Source/Runtime/VulkanRHI/Private/VulkanLoader.h
+++ b/Engine/Source/Runtime/VulkanRHI/Private/VulkanLoader.h
@@ -1,5 +1,7 @@
 // Copyright 1998-2017 Epic Games, Inc. All Rights Reserved.
 
+#pragma once
+
 #define VK_NO_PROTOTYPES
 #include <vulkan.h>
 
diff --git a/Engine/Source/Runtime/VulkanRHI/Private/VulkanPendingState.h b/Engine/Source/Runtime/VulkanRHI/Private/VulkanPendingState.h
--- a/Engine/Source/Runtime/VulkanRHI/Private/VulkanPendingState.h
+++ b/Engine/Source/Runtime/VulkanRHI/Private/VulkanPendingState.h
@@ -11,6 +11,7 @@
 #include "VulkanPipeline.h"
 #include "VulkanGlobalUniformBuffer.h"
 #include "VulkanPipelineState.h"
+#include "VulkanBoundShaderState.h"
 
 typedef uint32 FStateKey;
 
diff --git a/Engine/Source/Runtime/VulkanRHI/Private/VulkanPipeline.h b/Engine/Source/Runtime/VulkanRHI/Private/VulkanPipeline.h
--- a/Engine/Source/Runtime/VulkanRHI/Private/VulkanPipeline.h
+++ b/Engine/Source/Runtime/VulkanRHI/Private/VulkanPipeline.h
@@ -7,6 +7,7 @@
 #pragma once
 
 #include "VulkanResources.h"
+#include "VulkanDescriptorSets.h"
 
 class FVulkanDevice;
 class FVulkanFramebuffer;
diff --git a/Engine/Source/Runtime/VulkanRHI/Private/VulkanRHIPrivate.h b/Engine/Source/Runtime/VulkanRHI/Private/VulkanRHIPrivate.h
--- a/Engine/Source/Runtime/VulkanRHI/Private/VulkanRHIPrivate.h
+++ b/Engine/Source/Runtime/VulkanRHI/Private/VulkanRHIPrivate.h
@@ -17,14 +17,6 @@
 #include "WindowsHWrapper.h"
 #endif
 
-DECLARE_LOG_CATEGORY_EXTERN(LogVulkanRHI, Log, All);
-
-/** How many back buffers to cycle through */
-enum
-{
-	NUM_RENDER_BUFFERS = 3,
-};
-
 #ifndef VK_PROTOTYPES
 #define VK_PROTOTYPES	1
 #endif
@@ -72,12 +64,7 @@ enum
 	#include "VulkanCommandWrappers.h"
 #else
 	#if VULKAN_DYNAMICALLYLOADED
-		// Bring functions from VulkanDynamicAPI to VulkanRHI
-		#define VK_DYNAMICAPI_TO_VULKANRHI(Type,Func) using VulkanDynamicAPI::Func;
-		namespace VulkanRHI
-		{
-			ENUM_VK_ENTRYPOINTS_ALL(VK_DYNAMICAPI_TO_VULKANRHI);
-		}
+		#include "VulkanCommandsDirect.h"
 	#else
 		#error "Statically linked vulkan api must be wrapped!"
 	#endif
@@ -87,6 +74,10 @@ enum
 #include "VulkanRHI.h"
 #include "VulkanGlobalUniformBuffer.h"
 #include "RHI.h"
+#include "VulkanDevice.h"
+#include "VulkanQueue.h"
+#include "VulkanCommandBuffer.h"
+#include "Stats2.h"
 
 using namespace VulkanRHI;
 
@@ -189,8 +180,6 @@ protected:
 	friend class FVulkanPipelineStateCache;
 };
 
-#include "VulkanDevice.h"
-
 struct FVulkanSemaphore
 {
 	FVulkanSemaphore(FVulkanDevice& InDevice):
@@ -224,9 +213,6 @@ private:
 	VkSemaphore SemaphoreHandle;
 };
 
-#include "VulkanQueue.h"
-#include "VulkanCommandBuffer.h"
-
 class FVulkanFramebuffer
 {
 public:
@@ -406,7 +392,6 @@ inline void VulkanSetImageLayoutSimple(VkCommandBuffer CmdBuffer, VkImage Image,
 void VulkanResolveImage(VkCommandBuffer Cmd, FTextureRHIParamRef SourceTextureRHI, FTextureRHIParamRef DestTextureRHI);
 
 // Stats
-#include "Stats2.h"
 DECLARE_STATS_GROUP(TEXT("Vulkan RHI"), STATGROUP_VulkanRHI, STATCAT_Advanced);
 //DECLARE_STATS_GROUP(TEXT("Vulkan RHI Verbose"), STATGROUP_VulkanRHIVERBOSE, STATCAT_Advanced);
 //DECLARE_DWORD_COUNTER_STAT_EXTERN(TEXT("DrawPrimitive calls"), STAT_VulkanDrawPrimitiveCalls, STATGROUP_VulkanRHI, );
diff --git a/Engine/Source/Runtime/VulkanRHI/Public/VulkanConfiguration.h b/Engine/Source/Runtime/VulkanRHI/Public/VulkanConfiguration.h
--- a/Engine/Source/Runtime/VulkanRHI/Public/VulkanConfiguration.h
+++ b/Engine/Source/Runtime/VulkanRHI/Public/VulkanConfiguration.h
@@ -166,3 +166,11 @@ namespace EVulkanBindingType
 		return 0;
 	}
 }
+
+/** How many back buffers to cycle through */
+enum
+{
+	NUM_RENDER_BUFFERS = 3,
+};
+
+DECLARE_LOG_CATEGORY_EXTERN(LogVulkanRHI, Log, All);
diff --git a/Engine/Source/Runtime/VulkanRHI/Public/VulkanResources.h b/Engine/Source/Runtime/VulkanRHI/Public/VulkanResources.h
--- a/Engine/Source/Runtime/VulkanRHI/Public/VulkanResources.h
+++ b/Engine/Source/Runtime/VulkanRHI/Public/VulkanResources.h
@@ -1115,227 +1115,7 @@ protected:
 };
 
 
-#include "../Private/VulkanDescriptorSets.h"
 
-/**
- * Combined shader state and vertex definition for rendering geometry. 
- * Each unique instance consists of a vertex decl, vertex shader, and pixel shader.
- */
-class FVulkanBoundShaderState : public FVulkanDescriptorSetRingBuffer, public FVulkanDescriptorSetWriteContainer, public FRHIBoundShaderState
-{
-public:
-	FVulkanBoundShaderState(
-		FVulkanDevice* InDevice,
-		FVertexDeclarationRHIParamRef InVertexDeclarationRHI,
-		FVertexShaderRHIParamRef InVertexShaderRHI,
-		FPixelShaderRHIParamRef InPixelShaderRHI,
-		FHullShaderRHIParamRef InHullShaderRHI,
-		FDomainShaderRHIParamRef InDomainShaderRHI,
-		FGeometryShaderRHIParamRef InGeometryShaderRHI);
-
-	virtual ~FVulkanBoundShaderState();
-
-	// Called when the shader is set by the engine as current pending state
-	// After the shader is set, it is expected that all require states are provided/set by the engine
-	void ResetState();
-
-	const FVulkanLayout& GetLayout() const
-	{
-		return Layout;
-	}
-
-	class FVulkanGfxPipeline* PrepareForDraw(class FVulkanRenderPass* RenderPass, const struct FVulkanPipelineGraphicsKey& PipelineKey, uint32 VertexInputKey, const struct FVulkanGfxPipelineState& State);
-
-	bool UpdateDescriptorSets(FVulkanCommandListContext* CmdListContext, FVulkanCmdBuffer* CmdBuffer, FVulkanGlobalUniformPool* GlobalUniformPool);
-
-	void BindDescriptorSets(FVulkanCmdBuffer* Cmd);
-
-	inline void BindVertexStreams(FVulkanCmdBuffer* Cmd, const void* VertexStreams)
-	{
-		if (bDirtyVertexStreams)
-		{
-			InternalBindVertexStreams(Cmd, VertexStreams);
-			bDirtyVertexStreams = false;
-		}
-	}
-
-	inline void MarkDirtyVertexStreams()
-	{
-		bDirtyVertexStreams = true;
-	}
-
-	inline const FVulkanVertexInputStateInfo& GetVertexInputStateInfo() const
-	{
-		check(VertexInputStateInfo.Info.sType == VK_STRUCTURE_TYPE_PIPELINE_VERTEX_INPUT_STATE_CREATE_INFO);
-		return VertexInputStateInfo;
-	}
-
-	inline FSHAHash const* GetShaderHashes() const
-	{
-		return ShaderHashes;
-	}
-
-	// Binding
-	inline void SetShaderParameter(EShaderFrequency Stage, uint32 BufferIndex, uint32 ByteOffset, uint32 NumBytes, const void* NewValue)
-	{
-#if VULKAN_ENABLE_AGGRESSIVE_STATS
-		SCOPE_CYCLE_COUNTER(STAT_VulkanSetShaderParamTime);
-#endif
-		NEWPackedUniformBufferStaging[Stage].SetPackedGlobalParameter(BufferIndex, ByteOffset, NumBytes, NewValue, NEWPackedUniformBufferStagingDirty[Stage]);
-	}
-
-	inline void SetTexture(EShaderFrequency Stage, uint32 BindPoint, const FVulkanTextureBase* VulkanTextureBase)
-	{
-		check(VulkanTextureBase);
-		DescriptorWriteSet[Stage].WriteImage(BindPoint, VulkanTextureBase->PartialView->View,
-			(VulkanTextureBase->Surface.GetFullAspectMask() & (VK_IMAGE_ASPECT_DEPTH_BIT | VK_IMAGE_ASPECT_STENCIL_BIT)) != 0 ? VK_IMAGE_LAYOUT_DEPTH_STENCIL_READ_ONLY_OPTIMAL : VK_IMAGE_LAYOUT_GENERAL);
-	}
-
-	inline void SetSamplerState(EShaderFrequency Stage, uint32 BindPoint, FVulkanSamplerState* Sampler)
-	{
-		check(Sampler);
-		DescriptorWriteSet[Stage].WriteSampler(BindPoint, Sampler->Sampler);
-	}
-
-	inline void SetBufferViewState(EShaderFrequency Stage, uint32 BindPoint, FVulkanBufferView* View)
-	{
-		check(View && (View->Flags & VK_BUFFER_USAGE_UNIFORM_TEXEL_BUFFER_BIT) == VK_BUFFER_USAGE_UNIFORM_TEXEL_BUFFER_BIT);
-		DescriptorWriteSet[Stage].WriteUniformTexelBuffer(BindPoint, View);
-	}
-
-	inline void SetTextureView(EShaderFrequency Stage, uint32 BindPoint, const FVulkanTextureView& TextureView)
-	{
-		DescriptorWriteSet[Stage].WriteImage(BindPoint, TextureView.View, VK_IMAGE_LAYOUT_GENERAL);
-	}
-
-	inline void SetUniformBufferConstantData(EShaderFrequency Stage, uint32 BindPoint, const TArray<uint8>& ConstantData)
-	{
-		FVulkanShader* Shader = GetShaderPtr(Stage);
-		NEWPackedUniformBufferStaging[Stage].SetEmulatedUniformBufferIntoPacked(BindPoint, ConstantData, NEWPackedUniformBufferStagingDirty[Stage]);
-	}
-
-	void SetUniformBuffer(EShaderFrequency Stage, uint32 BindPoint, const FVulkanUniformBuffer* UniformBuffer);
-
-	void SetSRV(EShaderFrequency Stage, uint32 TextureIndex, FVulkanShaderResourceView* SRV);
-
-	inline const FVulkanShader* GetShaderPtr(EShaderFrequency Stage) const
-	{
-		switch (Stage)
-		{
-		case SF_Vertex:		return VertexShader;
-		case SF_Hull:		return HullShader;
-		case SF_Domain:		return DomainShader;
-		case SF_Pixel:		return PixelShader;
-		case SF_Geometry:	return GeometryShader;
-			//#todo-rco: Really should assert here...
-		default:			return nullptr;
-		}
-	}
-
-	// Has no internal verify.
-	// If stage does not exists, returns a "nullptr".
-	inline FVulkanShader* GetShaderPtr(EShaderFrequency Stage)
-	{
-		switch (Stage)
-		{
-		case SF_Vertex:		return VertexShader;
-		case SF_Hull:		return HullShader;
-		case SF_Domain:		return DomainShader;
-		case SF_Pixel:		return PixelShader;
-		case SF_Geometry:	return GeometryShader;
-			//#todo-rco: Really should assert here...
-		default:			return nullptr;
-		}
-	}
-
-	inline bool HasShaderStage(EShaderFrequency Stage) const
-	{
-		return (GetShaderPtr(Stage) != nullptr);
-	}
-
-	// Verifies if the shader stage exists
-	inline const FVulkanShader& GetShader(EShaderFrequency Stage) const
-	{
-		const FVulkanShader* Shader = GetShaderPtr(Stage);
-		check(Shader);
-		return *Shader;
-	}
-
-	inline void BindPipeline(VkCommandBuffer CmdBuffer, VkPipeline NewPipeline)
-	{
-		if (LastBoundPipeline != NewPipeline)
-		{
-			VulkanRHI::vkCmdBindPipeline(CmdBuffer, VK_PIPELINE_BIND_POINT_GRAPHICS, NewPipeline);
-			LastBoundPipeline = NewPipeline;
-		}
-	}
-
-private:
-	void InternalBindVertexStreams(FVulkanCmdBuffer* Cmd, const void* VertexStreams);
-	FCachedBoundShaderStateLink CacheLink;
-
-	/** Cached vertex structure */
-	TRefCountPtr<FVulkanVertexDeclaration> VertexDeclaration;
-
-	friend class FVulkanPipelineStateCache;
-
-	FSHAHash ShaderHashes[SF_NumFrequencies];
-
-	TLinkedList<FVulkanBoundShaderState*> GlobalListLink;
-
-	/** Cached shaders */
-	TRefCountPtr<FVulkanVertexShader> VertexShader;
-	TRefCountPtr<FVulkanPixelShader> PixelShader;
-	TRefCountPtr<FVulkanHullShader> HullShader;
-	TRefCountPtr<FVulkanDomainShader> DomainShader;
-	TRefCountPtr<FVulkanGeometryShader> GeometryShader;
-
-	// For debugging only
-	int32 ID;
-
-	bool bDirtyVertexStreams;
-
-	FVulkanDescriptorSetWriter DescriptorWriteSet[SF_Compute];
-
-	FPackedUniformBuffers NEWPackedUniformBufferStaging[SF_Compute];
-	uint64 NEWPackedUniformBufferStagingDirty[SF_Compute];
-
-#if VULKAN_ENABLE_RHI_DEBUGGING
-	struct FDebugInfo
-	{
-		const FVulkanTextureBase* Textures[SF_Compute][32];
-		const FVulkanSamplerState* SamplerStates[SF_Compute][32];
-		union
-		{
-			const FVulkanTextureView* SRVIs[SF_Compute][32];
-			const FVulkanBufferView* SRVBs[SF_Compute][32];
-		};
-		const FVulkanUniformBuffer* UBs[SF_Compute][32];
-	};
-	FDebugInfo DebugInfo;
-#endif
-
-	void CreateDescriptorWriteInfos();
-
-
-	FVulkanLayout Layout;
-
-	VkPipeline LastBoundPipeline;
-
-	// Vertex input configuration
-	FVulkanVertexInputStateInfo VertexInputStateInfo;
-
-	// Members in Tmp are normally allocated each frame
-	// To reduce the allocation, we reuse these array
-	struct FTmp
-	{
-		TArray<VkBuffer> VertexBuffers;
-		TArray<VkDeviceSize> VertexOffsets;
-	} Tmp;
-
-	// these are the cache pipeline state objects used in this BSS; RefCounts must be manually controlled for the FVulkanPipelines!
-	TMap<FVulkanPipelineGraphicsKey, FVulkanGfxPipeline* > PipelineCache;
-};
 
 template<class T>
 struct TVulkanResourceTraits
@@ -1377,11 +1157,6 @@ struct TVulkanResourceTraits<FRHIComputeShader>
 	typedef FVulkanComputeShader TConcreteType;
 };
 template<>
-struct TVulkanResourceTraits<FRHIBoundShaderState>
-{
-	typedef FVulkanBoundShaderState TConcreteType;
-};
-template<>
 struct TVulkanResourceTraits<FRHITexture3D>
 {
 	typedef FVulkanTexture3D TConcreteType;
