Fix Vulkan crash on start (UE-46603).

- GDeviceExtensions and similar arrays are now always null-terminated because otherwise they can become zero-sized depending on the defines.

diff --git a/Engine/Source/Runtime/VulkanRHI/Private/VulkanLayers.cpp b/Engine/Source/Runtime/VulkanRHI/Private/VulkanLayers.cpp
--- a/Engine/Source/Runtime/VulkanRHI/Private/VulkanLayers.cpp
+++ b/Engine/Source/Runtime/VulkanRHI/Private/VulkanLayers.cpp
@@ -439,7 +439,7 @@ void FVulkanDevice::GetDeviceExtensions(TArray<const ANSICHAR*>& OutDeviceExtens
 		}
 	}
 
-	for (uint32 Index = 0; Index < ARRAY_COUNT(GDeviceExtensions); ++Index)
+	for (uint32 Index = 0; Index < ARRAY_COUNT(GDeviceExtensions) && GDeviceExtensions[Index] != nullptr; ++Index)
 	{
 		for (int32 i = 0; i < Extensions.ExtensionProps.Num(); i++)
 		{
