UE4 - Copy from 3527113, fix Vulkan not presenting

diff --git a/Engine/Source/Runtime/VulkanRHI/Private/VulkanSwapChain.cpp b/Engine/Source/Runtime/VulkanRHI/Private/VulkanSwapChain.cpp
--- a/Engine/Source/Runtime/VulkanRHI/Private/VulkanSwapChain.cpp
+++ b/Engine/Source/Runtime/VulkanRHI/Private/VulkanSwapChain.cpp
@@ -11,6 +11,8 @@
 #include <SDL.h>
 #endif
 
+extern FAutoConsoleVariable GCVarDelayAcquireBackBuffer;
+
 FVulkanSwapChain::FVulkanSwapChain(VkInstance InInstance, FVulkanDevice& InDevice, void* WindowHandle, EPixelFormat& InOutPixelFormat, uint32 Width, uint32 Height,
 	uint32* InOutDesiredNumBackBuffers, TArray<VkImage>& OutImages)
 	: SwapChain(VK_NULL_HANDLE)
@@ -18,6 +20,8 @@ FVulkanSwapChain::FVulkanSwapChain(VkInstance InInstance, FVulkanDevice& InDevic
 	, Surface(VK_NULL_HANDLE)
 	, CurrentImageIndex(-1)
 	, SemaphoreIndex(0)
+	, NumPresentCalls(0)
+	, NumAcquireCalls(0)
 	, Instance(InInstance)
 {
 #if PLATFORM_WINDOWS
@@ -214,7 +218,11 @@ FVulkanSwapChain::FVulkanSwapChain(VkInstance InInstance, FVulkanDevice& InDevic
 	SwapChainInfo.imageColorSpace = CurrFormat.colorSpace;
 	SwapChainInfo.imageExtent.width = PLATFORM_ANDROID ? Width : (SurfProperties.currentExtent.width == 0xFFFFFFFF ? Width : SurfProperties.currentExtent.width);
 	SwapChainInfo.imageExtent.height = PLATFORM_ANDROID ? Height : (SurfProperties.currentExtent.height == 0xFFFFFFFF ? Height : SurfProperties.currentExtent.height);
-	SwapChainInfo.imageUsage = VK_IMAGE_USAGE_COLOR_ATTACHMENT_BIT | VK_IMAGE_USAGE_TRANSFER_DST_BIT | VK_IMAGE_USAGE_SAMPLED_BIT;
+	SwapChainInfo.imageUsage = VK_IMAGE_USAGE_TRANSFER_DST_BIT;
+	if (GCVarDelayAcquireBackBuffer->GetInt() != 0)
+	{
+		SwapChainInfo.imageUsage |= VK_IMAGE_USAGE_COLOR_ATTACHMENT_BIT | VK_IMAGE_USAGE_SAMPLED_BIT;
+	}
 	SwapChainInfo.preTransform = PreTransform;
 	SwapChainInfo.imageArrayLayers = 1;
 	SwapChainInfo.imageSharingMode = VK_SHARING_MODE_EXCLUSIVE;
@@ -292,6 +300,8 @@ int32 FVulkanSwapChain::AcquireImageIndex(FVulkanSemaphore** OutSemaphore)
 	uint32 ImageIndex = 0;
 	SemaphoreIndex = (SemaphoreIndex + 1) % ImageAcquiredSemaphore.Num();
 
+	// If we have not called present for any of the swapchain images, it will cause a crash/hang
+	checkf(!(NumAcquireCalls == ImageAcquiredSemaphore.Num() - 1 && NumPresentCalls == 0), TEXT("vkAcquireNextImageKHR will fail as no images have been presented before acquiring all of them"));
 	VulkanRHI::FFenceManager& FenceMgr = Device.GetFenceManager();
 	FenceMgr.ResetFence(ImageAcquiredFences[SemaphoreIndex]);
 
@@ -302,12 +312,22 @@ int32 FVulkanSwapChain::AcquireImageIndex(FVulkanSemaphore** OutSemaphore)
 		ImageAcquiredSemaphore[SemaphoreIndex]->GetHandle(),
 		ImageAcquiredFences[SemaphoreIndex]->GetHandle(),
 		&ImageIndex);
-
+	++NumAcquireCalls;
 	*OutSemaphore = ImageAcquiredSemaphore[SemaphoreIndex];
 
-	checkf(Result == VK_SUCCESS || Result == VK_SUBOPTIMAL_KHR, TEXT("AcquireNextImageKHR failed Result = %d"), int32(Result));
+	if (Result == VK_ERROR_VALIDATION_FAILED_EXT)
+	{
+		extern TAutoConsoleVariable<int32> GValidationCvar;
+		if (GValidationCvar.GetValueOnRenderThread() == 0)
+		{
+			UE_LOG(LogVulkanRHI, Fatal, TEXT("vkAcquireNextImageKHR failed with Validation error. Try running with r.Vulkan.EnableValidation=1 to get information from the driver"));
+		}
+	}
+	else
+	{
+		checkf(Result == VK_SUCCESS || Result == VK_SUBOPTIMAL_KHR, TEXT("vkAcquireNextImageKHR failed Result = %d"), int32(Result));
+	}
 	CurrentImageIndex = (int32)ImageIndex;
-	check(CurrentImageIndex == ImageIndex);
 	
 	{
 		SCOPE_CYCLE_COUNTER(STAT_VulkanWaitSwapchain);
@@ -340,5 +360,7 @@ bool FVulkanSwapChain::Present(FVulkanQueue* Queue, FVulkanSemaphore* BackBuffer
 		VERIFYVULKANRESULT(VulkanRHI::vkQueuePresentKHR(Queue->GetHandle(), &Info));
 	}
 
+	++NumPresentCalls;
+
 	return true;
 }
diff --git a/Engine/Source/Runtime/VulkanRHI/Private/VulkanSwapChain.h b/Engine/Source/Runtime/VulkanRHI/Private/VulkanSwapChain.h
--- a/Engine/Source/Runtime/VulkanRHI/Private/VulkanSwapChain.h
+++ b/Engine/Source/Runtime/VulkanRHI/Private/VulkanSwapChain.h
@@ -31,6 +31,8 @@ protected:
 
 	int32 CurrentImageIndex;
 	int32 SemaphoreIndex;
+	uint32 NumPresentCalls;
+	uint32 NumAcquireCalls;
 	VkInstance Instance;
 	TArray<FVulkanSemaphore*> ImageAcquiredSemaphore;
 	TArray<VulkanRHI::FFence*> ImageAcquiredFences;
diff --git a/Engine/Source/Runtime/VulkanRHI/Private/VulkanViewport.cpp b/Engine/Source/Runtime/VulkanRHI/Private/VulkanViewport.cpp
--- a/Engine/Source/Runtime/VulkanRHI/Private/VulkanViewport.cpp
+++ b/Engine/Source/Runtime/VulkanRHI/Private/VulkanViewport.cpp
@@ -9,7 +9,7 @@
 #include "VulkanPendingState.h"
 #include "VulkanContext.h"
 
-static FAutoConsoleVariable CVarDelayAcquireBackBuffer(
+FAutoConsoleVariable GCVarDelayAcquireBackBuffer(
 	TEXT("r.Vulkan.DelayAcquireBackBuffer"),
 	1,
 	TEXT("Delay acquiring the back buffer until preset"),
@@ -18,7 +18,7 @@ static FAutoConsoleVariable CVarDelayAcquireBackBuffer(
 
 inline static bool DelayAcquireBackBuffer()
 {
-	return CVarDelayAcquireBackBuffer->GetInt() != 0;
+	return GCVarDelayAcquireBackBuffer->GetInt() != 0;
 }
 
 struct FRHICommandAcquireBackBuffer : public FRHICommand<FRHICommandAcquireBackBuffer>
@@ -529,7 +529,7 @@ bool FVulkanViewport::Present(FVulkanCmdBuffer* CmdBuffer, FVulkanQueue* Queue,
 	}
 
 	bool bResult = false;
-	if (bNeedNativePresent && RHIBackBuffer != nullptr)
+	if (bNeedNativePresent && (GCVarDelayAcquireBackBuffer->GetInt() != 0 || RHIBackBuffer != nullptr))
 	{
 		// Present the back buffer to the viewport window.
 		bResult = SwapChain->Present(Queue, RenderingDoneSemaphores[AcquiredImageIndex]);//, SyncInterval, 0);
