Fixed incremental/non-unity build failures (UE-56421 et al).

diff --git a/Engine/Source/Runtime/Core/Public/Unix/UnixPlatformTLS.h b/Engine/Source/Runtime/Core/Public/Unix/UnixPlatformTLS.h
--- a/Engine/Source/Runtime/Core/Public/Unix/UnixPlatformTLS.h
+++ b/Engine/Source/Runtime/Core/Public/Unix/UnixPlatformTLS.h
@@ -7,6 +7,7 @@
 #pragma once
 
 #include "CoreTypes.h"
+#include "Unix/UnixSystemIncludes.h"
 #include "Misc/AssertionMacros.h"
 #include "GenericPlatform/GenericPlatformTLS.h"
 
diff --git a/Engine/Source/Runtime/VulkanRHI/Private/Linux/VulkanLinuxPlatform.h b/Engine/Source/Runtime/VulkanRHI/Private/Linux/VulkanLinuxPlatform.h
--- a/Engine/Source/Runtime/VulkanRHI/Private/Linux/VulkanLinuxPlatform.h
+++ b/Engine/Source/Runtime/VulkanRHI/Private/Linux/VulkanLinuxPlatform.h
@@ -2,6 +2,8 @@
 
 #pragma once
 
+#include "CoreMinimal.h"
+
 #define VULKAN_HAS_PHYSICAL_DEVICE_PROPERTIES2		1
 #define VULKAN_COMMANDWRAPPERS_ENABLE				0
 #define VULKAN_DYNAMICALLYLOADED					1
diff --git a/Engine/Source/Runtime/VulkanRHI/Private/VulkanLoader.h b/Engine/Source/Runtime/VulkanRHI/Private/VulkanLoader.h
--- a/Engine/Source/Runtime/VulkanRHI/Private/VulkanLoader.h
+++ b/Engine/Source/Runtime/VulkanRHI/Private/VulkanLoader.h
@@ -5,6 +5,8 @@
 #define VK_NO_PROTOTYPES
 #include <vulkan.h>
 
+#include "CoreMinimal.h"
+
 // List all base Vulkan entry points used by Unreal that need to be loaded manually
 #define ENUM_VK_ENTRYPOINTS_BASE(EnumMacro) \
 	EnumMacro(PFN_vkCreateInstance, vkCreateInstance) \
