Fixed VulkanSwapChain missing the handling of the unknown pixel format

diff --git a/Engine/Source/Runtime/VulkanRHI/Private/VulkanSwapChain.cpp b/Engine/Source/Runtime/VulkanRHI/Private/VulkanSwapChain.cpp
--- a/Engine/Source/Runtime/VulkanRHI/Private/VulkanSwapChain.cpp
+++ b/Engine/Source/Runtime/VulkanRHI/Private/VulkanSwapChain.cpp
@@ -50,6 +50,12 @@ FVulkanSwapChain::FVulkanSwapChain(VkInstance InInstance, FVulkanDevice& InDevic
 		Formats.AddZeroed(NumFormats);
 		VERIFYVULKANRESULT_EXPANDED(VulkanRHI::vkGetPhysicalDeviceSurfaceFormatsKHR(Device.GetPhysicalHandle(), Surface, &NumFormats, Formats.GetData()));
 
+		if (InOutPixelFormat == PF_Unknown)
+		{
+			static const auto* CVarDefaultBackBufferPixelFormat = IConsoleManager::Get().FindTConsoleVariableDataInt(TEXT("r.DefaultBackBufferPixelFormat"));
+			InOutPixelFormat = CVarDefaultBackBufferPixelFormat ? EDefaultBackBufferPixelFormat::Convert2PixelFormat(CVarDefaultBackBufferPixelFormat->GetValueOnGameThread()) : PF_Unknown;
+		}
+
 		if (InOutPixelFormat != PF_Unknown)
 		{
 			bool bFound = false;
