Vulkan: fix missing debug markers (TCHAR_TO_ANSI lifetime issues)

- Fixed by BrandonS.

(merged from Dev-Console)

diff --git a/Engine/Source/Runtime/VulkanRHI/Private/VulkanRHI.cpp b/Engine/Source/Runtime/VulkanRHI/Private/VulkanRHI.cpp
--- a/Engine/Source/Runtime/VulkanRHI/Private/VulkanRHI.cpp
+++ b/Engine/Source/Runtime/VulkanRHI/Private/VulkanRHI.cpp
@@ -764,10 +764,11 @@ void FVulkanCommandListContext::RHIPushEvent(const TCHAR* Name, FColor Color)
 #if VULKAN_ENABLE_DRAW_MARKERS
 	if (auto CmdDbgMarkerBegin = Device->GetCmdDbgMarkerBegin())
 	{
+		FTCHARToUTF8 Converter(Name);
 		VkDebugMarkerMarkerInfoEXT Info;
 		FMemory::Memzero(Info);
 		Info.sType = VK_STRUCTURE_TYPE_DEBUG_MARKER_MARKER_INFO_EXT;
-		Info.pMarkerName = TCHAR_TO_ANSI(Name);
+		Info.pMarkerName = Converter.Get();
 		FLinearColor LColor(Color);
 		Info.color[0] = LColor.R;
 		Info.color[1] = LColor.G;
