Ensure we prefer VK_PRESENT_MODE_MAILBOX_KHR or VK_PRESENT_MODE_IMMEDIATE_KHR over VK_PRESENT_MODE_FIFO_KHR when creating a swap chain when running inVRr. VSync will interfere with compositor cadence.

diff --git a/Engine/Source/Runtime/HeadMountedDisplay/HeadMountedDisplay.Build.cs b/Engine/Source/Runtime/HeadMountedDisplay/HeadMountedDisplay.Build.cs
--- a/Engine/Source/Runtime/HeadMountedDisplay/HeadMountedDisplay.Build.cs
+++ b/Engine/Source/Runtime/HeadMountedDisplay/HeadMountedDisplay.Build.cs
@@ -1,4 +1,4 @@
-// Copyright 1998-2018 Epic Games, Inc. All Rights Reserved.
+﻿// Copyright 1998-2018 Epic Games, Inc. All Rights Reserved.
 
 using UnrealBuildTool;
 
@@ -31,7 +31,8 @@ public class HeadMountedDisplay : ModuleRules
                 "ShaderCore",
                 "RenderCore",
                 "UtilityShaders",
-                "Analytics"
+                "Analytics",
+                "EngineSettings",
             }
         );
 
diff --git a/Engine/Source/Runtime/HeadMountedDisplay/Private/HeadMountedDisplayVulkanExtensions.cpp b/Engine/Source/Runtime/HeadMountedDisplay/Private/HeadMountedDisplayVulkanExtensions.cpp
new file mode 100644
--- /dev/null
+++ b/Engine/Source/Runtime/HeadMountedDisplay/Private/HeadMountedDisplayVulkanExtensions.cpp
@@ -0,0 +1,10 @@
+﻿// Copyright 1998-2018 Epic Games, Inc. All Rights Reserved.
+
+#include "IHeadMountedDisplayVulkanExtensions.h"
+#include "GeneralProjectSettings.h"
+#include "Misc/CommandLine.h"
+
+bool IHeadMountedDisplayVulkanExtensions::ShouldDisableVulkanVSync() const
+{
+	return FParse::Param(FCommandLine::Get(), TEXT("vr")) || GetDefault<UGeneralProjectSettings>()->bStartInVR;
+}
diff --git a/Engine/Source/Runtime/HeadMountedDisplay/Public/IHeadMountedDisplayVulkanExtensions.h b/Engine/Source/Runtime/HeadMountedDisplay/Public/IHeadMountedDisplayVulkanExtensions.h
--- a/Engine/Source/Runtime/HeadMountedDisplay/Public/IHeadMountedDisplayVulkanExtensions.h
+++ b/Engine/Source/Runtime/HeadMountedDisplay/Public/IHeadMountedDisplayVulkanExtensions.h
@@ -1,15 +1,23 @@
-// Copyright 1998-2018 Epic Games, Inc. All Rights Reserved.
+﻿// Copyright 1998-2018 Epic Games, Inc. All Rights Reserved.
 
 #pragma once
 
+#include "Containers/Array.h"
+
 /**
 * Query Vulkan extensions required by the HMD.
 *
 * Knowledge of these extensions is required during Vulkan RHI initialization.
 */
-class IHeadMountedDisplayVulkanExtensions
+class HEADMOUNTEDDISPLAY_API IHeadMountedDisplayVulkanExtensions
 {
 public:
 	virtual bool GetVulkanInstanceExtensionsRequired(TArray<const ANSICHAR*>& Out) = 0;
 	virtual bool GetVulkanDeviceExtensionsRequired(struct VkPhysicalDevice_T *pPhysicalDevice, TArray<const ANSICHAR*>& Out) = 0;
-};
\ No newline at end of file
+
+	/**
+	* If true and VK_PRESENT_MODE_MAILBOX_KHR is not available, VK_PRESENT_MODE_IMMEDIATE_KHR will be preferred over VK_PRESENT_MODE_FIFO_KHR
+	* when creating the swap chain for the spectator window.
+	*/
+	virtual bool ShouldDisableVulkanVSync() const;
+};
diff --git a/Engine/Source/Runtime/VulkanRHI/Private/VulkanRHI.cpp b/Engine/Source/Runtime/VulkanRHI/Private/VulkanRHI.cpp
--- a/Engine/Source/Runtime/VulkanRHI/Private/VulkanRHI.cpp
+++ b/Engine/Source/Runtime/VulkanRHI/Private/VulkanRHI.cpp
@@ -1,4 +1,4 @@
-// Copyright 1998-2018 Epic Games, Inc. All Rights Reserved.
+﻿// Copyright 1998-2018 Epic Games, Inc. All Rights Reserved.
 
 /*=============================================================================
 	VulkanRHI.cpp: Vulkan device RHI implementation.
@@ -270,6 +270,9 @@ void FVulkanDynamicRHI::Shutdown()
 	delete Device;
 	Device = nullptr;
 
+	// Release the early HMD interface used to query extra extensions - if any was used
+	HMDVulkanExtensions = nullptr;
+
 #if VULKAN_HAS_DEBUGGING_ENABLED
 	RemoveDebugLayerCallback();
 #endif
@@ -633,9 +636,6 @@ void FVulkanDynamicRHI::InitInstance()
 
 		GProjectionSignY = 1.0f;
 
-		// Release the early HMD interface used to query extra extensions - if any was used
-		HMDVulkanExtensions = nullptr;
-
 		GIsRHIInitialized = true;
 
 		SavePipelineCacheCmd = IConsoleManager::Get().RegisterConsoleCommand(
diff --git a/Engine/Source/Runtime/VulkanRHI/Private/VulkanSwapChain.cpp b/Engine/Source/Runtime/VulkanRHI/Private/VulkanSwapChain.cpp
--- a/Engine/Source/Runtime/VulkanRHI/Private/VulkanSwapChain.cpp
+++ b/Engine/Source/Runtime/VulkanRHI/Private/VulkanSwapChain.cpp
@@ -1,4 +1,4 @@
-// Copyright 1998-2018 Epic Games, Inc. All Rights Reserved.
+﻿// Copyright 1998-2018 Epic Games, Inc. All Rights Reserved.
 
 /*=============================================================================
 	VulkanSwapChain.h: Vulkan viewport RHI definitions.
@@ -8,6 +8,8 @@
 #include "VulkanSwapChain.h"
 #include "VulkanPlatform.h"
 #include "Engine/RendererSettings.h"
+#include "IHeadMountedDisplayModule.h"
+
 
 int32 GShouldCpuWaitForFence = 1;
 static FAutoConsoleVariableRef CVarCpuWaitForFence(
@@ -169,28 +171,44 @@ FVulkanSwapChain::FVulkanSwapChain(VkInstance InInstance, FVulkanDevice& InDevic
 		FoundPresentModes.AddZeroed(NumFoundPresentModes);
 		VERIFYVULKANRESULT(VulkanRHI::vkGetPhysicalDeviceSurfacePresentModesKHR(Device.GetPhysicalHandle(), Surface, &NumFoundPresentModes, FoundPresentModes.GetData()));
 
-		bool bFoundDesiredMode = false;
-		VkPresentModeKHR PotentialModes[] = {VK_PRESENT_MODE_MAILBOX_KHR, VK_PRESENT_MODE_FIFO_KHR};
-		for (VkPresentModeKHR PotentialMode : PotentialModes)
-		{
-			for (size_t i = 0; i < NumFoundPresentModes; i++)
-			{
-				if (FoundPresentModes[i] == PotentialMode)
-				{
-					bFoundDesiredMode = true;
-					PresentMode = PotentialMode;
-					break;
-				}
-			}
+		bool bFoundPresentModeMailbox = false;
+		bool bFoundPresentModeImmediate = false;
+		bool bFoundPresentModeFIFO = false;
 
-			if (bFoundDesiredMode)
+		for (size_t i = 0; i < NumFoundPresentModes; i++)
+		{
+			switch (FoundPresentModes[i])
 			{
+			case VK_PRESENT_MODE_MAILBOX_KHR:
+				bFoundPresentModeMailbox = true;
+				break;
+			case VK_PRESENT_MODE_IMMEDIATE_KHR:
+				bFoundPresentModeImmediate = true;
+				break;
+			case VK_PRESENT_MODE_FIFO_KHR:
+				bFoundPresentModeFIFO = true;
 				break;
 			}
 		}
-		if (!bFoundDesiredMode)
+
+		// Until FVulkanViewport::Present honors SyncInterval, we need to disable vsync for the spectator window if using an HMD.
+		const bool bDisableVsyncForHMD = (FVulkanDynamicRHI::HMDVulkanExtensions.IsValid()) ? FVulkanDynamicRHI::HMDVulkanExtensions->ShouldDisableVulkanVSync() : false;
+		
+		if (bFoundPresentModeMailbox)
+		{
+			PresentMode = VK_PRESENT_MODE_MAILBOX_KHR;
+		}
+		else if (bFoundPresentModeImmediate && bDisableVsyncForHMD)
+		{
+			PresentMode = VK_PRESENT_MODE_IMMEDIATE_KHR;
+		}
+		else if (bFoundPresentModeFIFO)
+		{
+			PresentMode = VK_PRESENT_MODE_FIFO_KHR;
+		}
+		else
 		{
-			UE_LOG(LogVulkanRHI, Warning, TEXT("Couldn't find desired VkPresentModeKHR %d! Using %d"), (int32)PresentMode, FoundPresentModes[0]);
+			UE_LOG(LogVulkanRHI, Warning, TEXT("Couldn't find desired PresentMode! Using %d"), static_cast<int32>(FoundPresentModes[0]));
 			PresentMode = FoundPresentModes[0];
 		}
 	}
