Vulkan on Linux does not allow debugging with markers

This is required for debugging wtih Render Doc or anything that uses VK debug markers

diff --git a/Engine/Source/Runtime/VulkanRHI/Private/Linux/VulkanLinuxPlatform.h b/Engine/Source/Runtime/VulkanRHI/Private/Linux/VulkanLinuxPlatform.h
--- a/Engine/Source/Runtime/VulkanRHI/Private/Linux/VulkanLinuxPlatform.h
+++ b/Engine/Source/Runtime/VulkanRHI/Private/Linux/VulkanLinuxPlatform.h
@@ -7,8 +7,8 @@
 #define VULKAN_HAS_PHYSICAL_DEVICE_PROPERTIES2		1
 #define VULKAN_COMMANDWRAPPERS_ENABLE				0
 #define VULKAN_DYNAMICALLYLOADED					1
-#define VULKAN_SHOULD_DEBUG_IN_DEVELOPMENT			0
-#define VULKAN_SHOULD_ENABLE_DRAW_MARKERS			0
+#define VULKAN_SHOULD_DEBUG_IN_DEVELOPMENT			1
+#define VULKAN_SHOULD_ENABLE_DRAW_MARKERS			(UE_BUILD_DEBUG || UE_BUILD_DEVELOPMENT)
 #define VULKAN_SIGNAL_UNIMPLEMENTED()				checkf(false, TEXT("Unimplemented vulkan functionality: %s"), __PRETTY_FUNCTION__)
 
 #define	VULKAN_SUPPORTS_DEDICATED_ALLOCATION				0
