Vulkan: disable unused device features (UE-59930).

(Edigrating CL 4111492 from Dev-Console to Release-4.20)

diff --git a/Engine/Source/Runtime/VulkanRHI/Private/VulkanDevice.cpp b/Engine/Source/Runtime/VulkanRHI/Private/VulkanDevice.cpp
--- a/Engine/Source/Runtime/VulkanRHI/Private/VulkanDevice.cpp
+++ b/Engine/Source/Runtime/VulkanRHI/Private/VulkanDevice.cpp
@@ -206,7 +206,9 @@ void FVulkanDevice::CreateDevice()
 	DeviceInfo.queueCreateInfoCount = QueueFamilyInfos.Num();
 	DeviceInfo.pQueueCreateInfos = QueueFamilyInfos.GetData();
 
-	DeviceInfo.pEnabledFeatures = &Features;
+	VkPhysicalDeviceFeatures EnabledFeatures;
+	FVulkanPlatform::RestrictEnabledPhysicalDeviceFeatures(Features, EnabledFeatures);
+	DeviceInfo.pEnabledFeatures = &EnabledFeatures;
 
 	// Create the device
 	VERIFYVULKANRESULT(VulkanRHI::vkCreateDevice(Gpu, &DeviceInfo, nullptr, &Device));
diff --git a/Engine/Source/Runtime/VulkanRHI/Private/VulkanGenericPlatform.h b/Engine/Source/Runtime/VulkanRHI/Private/VulkanGenericPlatform.h
--- a/Engine/Source/Runtime/VulkanRHI/Private/VulkanGenericPlatform.h
+++ b/Engine/Source/Runtime/VulkanRHI/Private/VulkanGenericPlatform.h
@@ -54,4 +54,22 @@ public:
 	static bool RegisterGPUWork() { return true; }
 
 	static void WriteBufferMarkerAMD(VkCommandBuffer CmdBuffer, VkBuffer DestBuffer, const TArrayView<uint32>& Entries, bool bAdding) {}
+
+	// allow the platform code to restrict the device features
+	static void RestrictEnabledPhysicalDeviceFeatures(const VkPhysicalDeviceFeatures& DeviceFeatures, VkPhysicalDeviceFeatures& FeaturesToEnable)
+	{ 
+		FeaturesToEnable = DeviceFeatures; 
+
+		// disable everything sparse-related
+		FeaturesToEnable.shaderResourceResidency	= VK_FALSE;
+		FeaturesToEnable.shaderResourceMinLod		= VK_FALSE;
+		FeaturesToEnable.sparseBinding				= VK_FALSE;
+		FeaturesToEnable.sparseResidencyBuffer		= VK_FALSE;
+		FeaturesToEnable.sparseResidencyImage2D		= VK_FALSE;
+		FeaturesToEnable.sparseResidencyImage3D		= VK_FALSE;
+		FeaturesToEnable.sparseResidency2Samples	= VK_FALSE;
+		FeaturesToEnable.sparseResidency4Samples	= VK_FALSE;
+		FeaturesToEnable.sparseResidency8Samples	= VK_FALSE;
+		FeaturesToEnable.sparseResidencyAliased		= VK_FALSE;
+	}
 };
