Vulkan: expose some additional driver information for Unix platform

diff --git a/Engine/Source/Runtime/VulkanRHI/Private/VulkanRHI.cpp b/Engine/Source/Runtime/VulkanRHI/Private/VulkanRHI.cpp
--- a/Engine/Source/Runtime/VulkanRHI/Private/VulkanRHI.cpp
+++ b/Engine/Source/Runtime/VulkanRHI/Private/VulkanRHI.cpp
@@ -553,6 +553,12 @@ void FVulkanDynamicRHI::SelectAndInitDevice()
 
 		// Ignore GRHIAdapterInternalDriverVersion for now as the device name doesn't match
 	}
+	else if(PLATFORM_UNIX)
+	{
+		GRHIAdapterInternalDriverVersion = FString::Printf(TEXT("%d.%d.%d (0x%X)"), VK_VERSION_MAJOR(Props.apiVersion), VK_VERSION_MINOR(Props.apiVersion), VK_VERSION_PATCH(Props.apiVersion), Props.apiVersion);
+		GRHIAdapterUserDriverVersion = FString::Printf(TEXT("%d.%d.%d (0x%X)"), VK_VERSION_MAJOR(Props.driverVersion), VK_VERSION_MINOR(Props.driverVersion), VK_VERSION_PATCH(Props.driverVersion), Props.driverVersion);
+		GRHIDeviceId = Props.deviceID;
+	}
 }
 
 void FVulkanDynamicRHI::InitInstance()
