UE4.20 - Add deletion delay for Vulkan objects

diff --git a/Engine/Source/Runtime/VulkanRHI/Private/VulkanCommandWrappers.h b/Engine/Source/Runtime/VulkanRHI/Private/VulkanCommandWrappers.h
--- a/Engine/Source/Runtime/VulkanRHI/Private/VulkanCommandWrappers.h
+++ b/Engine/Source/Runtime/VulkanRHI/Private/VulkanCommandWrappers.h
@@ -7,9 +7,11 @@
 #pragma once 
 
 #if VULKAN_ENABLE_CUSTOM_LAYER
-	#define  VULKAN_LAYER_BODY	;
+	#define  VULKAN_LAYER_BODY		;
+	#define	VULKAN_EXTERN_EXPORT	VULKANRHI_API
 #else
-	#define VULKAN_LAYER_BODY	{}
+	#define VULKAN_LAYER_BODY		{}
+	#define	VULKAN_EXTERN_EXPORT
 #endif
 
 struct FWrapLayer
@@ -20,7 +22,7 @@ struct FWrapLayer
 	static void EnumeratePhysicalDevices(VkResult Result, VkInstance Instance, uint32* PhysicalDeviceCount, VkPhysicalDevice* PhysicalDevices) VULKAN_LAYER_BODY
 	static void DestroyInstance(VkResult Result, VkInstance Instance) VULKAN_LAYER_BODY
 	static void GetInstanceProcAddr(VkResult Result, VkInstance Instance, const char* Name, PFN_vkVoidFunction VoidFunction) VULKAN_LAYER_BODY
-	VULKANRHI_API static void EnumerateInstanceExtensionProperties(VkResult Result, const char* LayerName, uint32* PropertyCount, VkExtensionProperties* Properties) VULKAN_LAYER_BODY
+	VULKAN_EXTERN_EXPORT static void EnumerateInstanceExtensionProperties(VkResult Result, const char* LayerName, uint32* PropertyCount, VkExtensionProperties* Properties) VULKAN_LAYER_BODY
 	static void EnumerateInstanceLayerProperties(VkResult Result, uint32* PropertyCount, VkLayerProperties* Properties) VULKAN_LAYER_BODY
 	static void AllocateMemory(VkResult Result, VkDevice Device, const VkMemoryAllocateInfo* AllocateInfo, VkDeviceMemory* Memory) VULKAN_LAYER_BODY
 	static void FlushMappedMemoryRanges(VkResult Result, VkDevice Device, uint32 MemoryRangeCount, const VkMappedMemoryRange* pMemoryRanges) VULKAN_LAYER_BODY
@@ -34,7 +36,7 @@ struct FWrapLayer
 	static void CreateQueryPool(VkResult Result, VkDevice Device, const VkQueryPoolCreateInfo* CreateInfo, VkQueryPool* QueryPool) VULKAN_LAYER_BODY
 	static void CreateBuffer(VkResult Result, VkDevice Device, const VkBufferCreateInfo* CreateInfo, VkBuffer* Buffer) VULKAN_LAYER_BODY
 	static void CreateBufferView(VkResult Result, VkDevice Device, const VkBufferViewCreateInfo* CreateInfo, VkBufferView* BufferView) VULKAN_LAYER_BODY
-	static void CreateImage(VkResult Result, VkDevice Device, const VkImageCreateInfo* CreateInfo, VkImage* Image) VULKAN_LAYER_BODY
+	VULKAN_EXTERN_EXPORT static void CreateImage(VkResult Result, VkDevice Device, const VkImageCreateInfo* CreateInfo, VkImage* Image) VULKAN_LAYER_BODY
 	static void DestroyImage(VkResult Result, VkDevice Device, VkImage Image) VULKAN_LAYER_BODY
 	static void GetImageSubresourceLayout(VkResult Result, VkDevice Device, VkImage Image, const VkImageSubresource* Subresource, VkSubresourceLayout* Layout) VULKAN_LAYER_BODY
 	static void CreateImageView(VkResult Result, VkDevice Device, const VkImageViewCreateInfo* CreateInfo, VkImageView* ImageView) VULKAN_LAYER_BODY
@@ -58,9 +60,9 @@ struct FWrapLayer
 	static void BindIndexBuffer(VkResult Result, VkCommandBuffer CommandBuffer, VkBuffer IndexBuffer, VkDeviceSize Offset, VkIndexType IndexType) VULKAN_LAYER_BODY
 	static void CopyBuffer(VkResult Result, VkCommandBuffer CommandBuffer, VkBuffer SrcBuffer, VkBuffer DstBuffer, uint32 RegionCount, const VkBufferCopy* Regions) VULKAN_LAYER_BODY
 	static void CopyBufferToImage(VkResult Result, VkCommandBuffer CommandBuffer, VkBuffer SrcBuffer, VkImage DstImage, VkImageLayout DstImageLayout, uint32 RegionCount, const VkBufferImageCopy* Regions) VULKAN_LAYER_BODY
-	static void BlitImage(VkResult Result, VkCommandBuffer CommandBuffer, VkImage SrcImage, VkImageLayout SrcImageLayout, VkImage DstImage, VkImageLayout DstImageLayout, uint32 RegionCount, const VkImageBlit* Regions, VkFilter Filter) VULKAN_LAYER_BODY
+	VULKAN_EXTERN_EXPORT static void BlitImage(VkResult Result, VkCommandBuffer CommandBuffer, VkImage SrcImage, VkImageLayout SrcImageLayout, VkImage DstImage, VkImageLayout DstImageLayout, uint32 RegionCount, const VkImageBlit* Regions, VkFilter Filter) VULKAN_LAYER_BODY
 	static void CopyImageToBuffer(VkResult Result, VkCommandBuffer CommandBuffer, VkImage SrcImage, VkImageLayout SrcImageLayout, VkBuffer DstBuffer, uint32 RegionCount, const VkBufferImageCopy* Regions) VULKAN_LAYER_BODY
-	static void ClearColorImage(VkResult Result, VkCommandBuffer CommandBuffer, VkImage Image, VkImageLayout ImageLayout, const VkClearColorValue* ColorValue, uint32 RangeCount, const VkImageSubresourceRange* Ranges) VULKAN_LAYER_BODY
+	VULKAN_EXTERN_EXPORT static void ClearColorImage(VkResult Result, VkCommandBuffer CommandBuffer, VkImage Image, VkImageLayout ImageLayout, const VkClearColorValue* ColorValue, uint32 RangeCount, const VkImageSubresourceRange* Ranges) VULKAN_LAYER_BODY
 	static void ClearDepthStencilImage(VkResult Result, VkCommandBuffer CommandBuffer, VkImage Image, VkImageLayout ImageLayout, const VkClearDepthStencilValue* DepthStencil, uint32 RangeCount, const VkImageSubresourceRange* Ranges) VULKAN_LAYER_BODY
 	static void ClearAttachments(VkResult Result, VkCommandBuffer CommandBuffer, uint32 AttachmentCount, const VkClearAttachment* Attachments, uint32 RectCount, const VkClearRect* Rects) VULKAN_LAYER_BODY
 	static void ResolveImage(VkResult Result, VkCommandBuffer CommandBuffer, VkImage SrcImage, VkImageLayout SrcImageLayout, VkImage DstImage, VkImageLayout DstImageLayout, uint32 RegionCount, const VkImageResolve* Regions) VULKAN_LAYER_BODY
@@ -80,7 +82,7 @@ struct FWrapLayer
 	static void MapMemory(VkResult Result, VkDevice Device, VkDeviceMemory Memory, VkDeviceSize Offset, VkDeviceSize Size, VkMemoryMapFlags Flags, void** Data) VULKAN_LAYER_BODY
 	static void UnmapMemory(VkResult Result, VkDevice Device, VkDeviceMemory Memory) VULKAN_LAYER_BODY
 	static void BindBufferMemory(VkResult Result, VkDevice Device, VkBuffer Buffer, VkDeviceMemory Memory, VkDeviceSize MemoryOffset) VULKAN_LAYER_BODY
-	static void BindImageMemory(VkResult Result, VkDevice Device, VkImage Image, VkDeviceMemory Memory, VkDeviceSize MemoryOffset) VULKAN_LAYER_BODY
+	VULKAN_EXTERN_EXPORT static void BindImageMemory(VkResult Result, VkDevice Device, VkImage Image, VkDeviceMemory Memory, VkDeviceSize MemoryOffset) VULKAN_LAYER_BODY
 	static void GetFenceStatus(VkResult Result, VkDevice Device, VkFence Fence) VULKAN_LAYER_BODY
 	static void GetQueryPoolResults(VkResult Result, VkDevice Device, VkQueryPool QueryPool, uint32 FirstQuery, uint32 QueryCount, size_t DataSize, void* Data, VkDeviceSize Stride, VkQueryResultFlags Flags) VULKAN_LAYER_BODY
 	static void CreateComputePipelines(VkResult Result, VkDevice Device, VkPipelineCache PipelineCache, uint32 CreateInfoCount, const VkComputePipelineCreateInfo* CreateInfos, VkPipeline* Pipelines) VULKAN_LAYER_BODY
@@ -150,7 +152,7 @@ struct FWrapLayer
 	static void ResetEvent(VkResult Result, VkCommandBuffer CommandBuffer, VkEvent Event, VkPipelineStageFlags StageMask) VULKAN_LAYER_BODY
 	static void CopyQueryPoolResults(VkResult Result, VkCommandBuffer CommandBuffer, VkQueryPool QueryPool, uint32 FirstQuery, uint32 QueryCount, VkBuffer DstBuffer, VkDeviceSize DstOffset, VkDeviceSize Stride, VkQueryResultFlags Flags) VULKAN_LAYER_BODY
 	static void GetDeviceProcAddr(VkResult Result, VkDevice Device, const char* Name, PFN_vkVoidFunction VoidFunction) VULKAN_LAYER_BODY
-	VULKANRHI_API static void EnumerateDeviceExtensionProperties(VkResult Result, VkPhysicalDevice PhysicalDevice, const char* LayerName, uint32* PropertyCount, VkExtensionProperties* Properties) VULKAN_LAYER_BODY
+	VULKAN_EXTERN_EXPORT static void EnumerateDeviceExtensionProperties(VkResult Result, VkPhysicalDevice PhysicalDevice, const char* LayerName, uint32* PropertyCount, VkExtensionProperties* Properties) VULKAN_LAYER_BODY
 	static void EnumerateDeviceLayerProperties(VkResult Result, VkPhysicalDevice PhysicalDevice, uint32* PropertyCount, VkLayerProperties* Properties) VULKAN_LAYER_BODY
 	static void GetPhysicalDeviceSurfaceCapabilitiesKHR(VkResult Result, VkPhysicalDevice PhysicalDevice, VkSurfaceKHR Surface, VkSurfaceCapabilitiesKHR* SurfaceCapabilities) VULKAN_LAYER_BODY
 	static void GetPhysicalDeviceSurfaceFormatsKHR(VkResult Result, VkPhysicalDevice PhysicalDevice, VkSurfaceKHR Surface, uint32_t* SurfaceFormatCountPtr, VkSurfaceFormatKHR* SurfaceFormats) VULKAN_LAYER_BODY
diff --git a/Engine/Source/Runtime/VulkanRHI/Private/VulkanMemory.cpp b/Engine/Source/Runtime/VulkanRHI/Private/VulkanMemory.cpp
--- a/Engine/Source/Runtime/VulkanRHI/Private/VulkanMemory.cpp
+++ b/Engine/Source/Runtime/VulkanRHI/Private/VulkanMemory.cpp
@@ -9,6 +9,12 @@
 #include "Misc/OutputDeviceRedirector.h"
 #include "HAL/PlatformStackWalk.h"
 
+
+// This 'frame number' should only be used for the deletion queue
+uint32 GVulkanRHIDeletionFrameNumber = 0;
+const uint32 NUM_FRAMES_TO_WAIT_FOR_RESOURCE_DELETE = 2;
+
+
 #if VULKAN_MEMORY_TRACK_CALLSTACK
 static FCriticalSection GStackTraceMutex;
 static char GStackTrace[65536];
@@ -1610,14 +1616,16 @@ namespace VulkanRHI
 		Queue->GetLastSubmittedInfo(Entry.CmdBuffer, Entry.FenceCounter);
 		Entry.Handle = Handle;
 		Entry.StructureType = Type;
+		Entry.FrameNumber = GVulkanRHIDeletionFrameNumber;
 
 		{
 			FScopeLock ScopeLock(&CS);
 
 #if VULKAN_HAS_DEBUGGING_ENABLED
-			FEntry* ExistingEntry = Entries.FindByPredicate([&](const FEntry& InEntry) { 
-				return InEntry.Handle == Entry.Handle; 
-			});
+			FEntry* ExistingEntry = Entries.FindByPredicate([&](const FEntry& InEntry)
+				{ 
+					return InEntry.Handle == Entry.Handle; 
+				});
 			checkf(ExistingEntry == nullptr, TEXT("Attempt to double-delete resource, Type: %d, Handle: %llu"), (int32)Type, Handle);
 #endif
 
@@ -1639,7 +1647,10 @@ namespace VulkanRHI
 		{
 			FEntry* Entry = &Entries[Index];
 			// #todo-rco: Had to add this check, we were getting null CmdBuffers on the first frame, or before first frame maybe
-			if (bDeleteImmediately || Entry->CmdBuffer == nullptr || Entry->FenceCounter < Entry->CmdBuffer->GetFenceSignaledCounter())
+			if (bDeleteImmediately ||
+				(GVulkanRHIDeletionFrameNumber > Entry->FrameNumber + NUM_FRAMES_TO_WAIT_FOR_RESOURCE_DELETE &&
+				(Entry->CmdBuffer == nullptr || Entry->FenceCounter < Entry->CmdBuffer->GetFenceSignaledCounter()))
+				)
 			{
 				switch (Entry->StructureType)
 				{
diff --git a/Engine/Source/Runtime/VulkanRHI/Private/VulkanRHI.cpp b/Engine/Source/Runtime/VulkanRHI/Private/VulkanRHI.cpp
--- a/Engine/Source/Runtime/VulkanRHI/Private/VulkanRHI.cpp
+++ b/Engine/Source/Runtime/VulkanRHI/Private/VulkanRHI.cpp
@@ -1,4 +1,4 @@
-﻿// Copyright 1998-2018 Epic Games, Inc. All Rights Reserved.
+// Copyright 1998-2018 Epic Games, Inc. All Rights Reserved.
 
 /*=============================================================================
 	VulkanRHI.cpp: Vulkan device RHI implementation.
@@ -683,6 +683,9 @@ void FVulkanCommandListContext::RHIBeginFrame()
 	check(IsImmediate());
 	RHIPrivateBeginFrame();
 
+	extern uint32 GVulkanRHIDeletionFrameNumber;
+	++GVulkanRHIDeletionFrameNumber;
+
 	GpuProfiler.BeginFrame();
 }
 
diff --git a/Engine/Source/Runtime/VulkanRHI/Public/VulkanMemory.h b/Engine/Source/Runtime/VulkanRHI/Public/VulkanMemory.h
--- a/Engine/Source/Runtime/VulkanRHI/Public/VulkanMemory.h
+++ b/Engine/Source/Runtime/VulkanRHI/Public/VulkanMemory.h
@@ -1099,11 +1099,10 @@ namespace VulkanRHI
 		struct FEntry
 		{
 			uint64 FenceCounter;
-			FVulkanCmdBuffer* CmdBuffer;
 			uint64 Handle;
+			FVulkanCmdBuffer* CmdBuffer;
 			EType StructureType;
-
-			//FRefCount* Resource;
+			uint32 FrameNumber;
 		};
 		FCriticalSection CS;
 		TArray<FEntry> Entries;
