UE4.20 - Support for compressed Vulkan volume textures
Copy from 4078994 and 4080845

diff --git a/Engine/Source/Developer/Windows/WindowsTargetPlatform/Private/GenericWindowsTargetPlatform.h b/Engine/Source/Developer/Windows/WindowsTargetPlatform/Private/GenericWindowsTargetPlatform.h
--- a/Engine/Source/Developer/Windows/WindowsTargetPlatform/Private/GenericWindowsTargetPlatform.h
+++ b/Engine/Source/Developer/Windows/WindowsTargetPlatform/Private/GenericWindowsTargetPlatform.h
@@ -50,6 +50,7 @@ public:
 
 		static FName NAME_PCD3D_SM5(TEXT("PCD3D_SM5"));
 		static FName NAME_PCD3D_SM4(TEXT("PCD3D_SM4"));
+		static FName NAME_VULKAN_SM5(TEXT("SF_VULKAN_SM5"));
 
 		bSupportDX11TextureFormats = true;
 		bSupportCompressedVolumeTexture = true;
@@ -67,10 +68,14 @@ public:
 				{
 					ShaderPlatform = SP_PCD3D_SM4;
 				}
+				else if (TargetedShaderFormat == NAME_VULKAN_SM5)
+				{
+					ShaderPlatform = SP_VULKAN_SM5;
+				}
 			}
 
 			// If we're targeting only DX11 we can use DX11 texture formats. Otherwise we'd have to compress fallbacks and increase the size of cooked content significantly.
-			if (ShaderPlatform != SP_PCD3D_SM5)
+			if (ShaderPlatform != SP_PCD3D_SM5 && ShaderPlatform != SP_VULKAN_SM5)
 			{
 				bSupportDX11TextureFormats = false;
 			}
diff --git a/Engine/Source/Runtime/Engine/Private/VolumeTexture.cpp b/Engine/Source/Runtime/Engine/Private/VolumeTexture.cpp
--- a/Engine/Source/Runtime/Engine/Private/VolumeTexture.cpp
+++ b/Engine/Source/Runtime/Engine/Private/VolumeTexture.cpp
@@ -557,7 +557,11 @@ bool UVolumeTexture::ShaderPlatformSupportsCompression(EShaderPlatform ShaderPla
 	case SP_PCD3D_SM5:
 	case SP_PS4:
 	case SP_XBOXONE_D3D12:
+	case SP_VULKAN_SM5:
+	case SP_VULKAN_SM4:
+	case SP_VULKAN_SM5_LUMIN:
 		return true;
+
 	default:
 		return false;
 	}
diff --git a/Engine/Source/Runtime/VulkanRHI/Private/VulkanRHI.cpp b/Engine/Source/Runtime/VulkanRHI/Private/VulkanRHI.cpp
--- a/Engine/Source/Runtime/VulkanRHI/Private/VulkanRHI.cpp
+++ b/Engine/Source/Runtime/VulkanRHI/Private/VulkanRHI.cpp
@@ -16,6 +16,8 @@
 #include "VulkanPipelineState.h"
 #include "Misc/FileHelper.h"
 
+extern RHI_API bool GUseTexture3DBulkDataRHI;
+
 #if VULKAN_ENABLE_DESKTOP_HMD_SUPPORT
 #include "Runtime/HeadMountedDisplay/Public/IHeadMountedDisplayModule.h"
 #endif
@@ -61,7 +63,7 @@ FDynamicRHI* FVulkanDynamicRHIModule::CreateRHI(ERHIFeatureLevel::Type InRequest
 			FParse::Param(FCommandLine::Get(), TEXT("featureleveles31")) || FParse::Param(FCommandLine::Get(), TEXT("featureleveles2"))))
 	{
 		GMaxRHIFeatureLevel = ERHIFeatureLevel::ES3_1;
-		GMaxRHIShaderPlatform = PLATFORM_LUMIN ? SP_VULKAN_ES3_1_LUMIN : (PLATFORM_ANDROID ? SP_VULKAN_ES3_1_ANDROID : SP_VULKAN_PCES3_1);
+		GMaxRHIShaderPlatform = PLATFORM_ANDROID ? SP_VULKAN_ES3_1_ANDROID : (PLATFORM_LUMIN ?  SP_VULKAN_ES3_1_LUMIN : SP_VULKAN_PCES3_1);
 	}
 	else if (InRequestedFeatureLevel == ERHIFeatureLevel::SM4)
 	{
@@ -625,6 +627,8 @@ void FVulkanDynamicRHI::InitInstance()
 
 		GRHIRequiresRenderTargetForPixelShaderUAVs = true;
 
+		GUseTexture3DBulkDataRHI = true;
+
 		GDynamicRHI = this;
 
 		// Notify all initialized FRenderResources that there's a valid RHI device to create their RHI resources for now.
diff --git a/Engine/Source/Runtime/VulkanRHI/Private/VulkanTexture.cpp b/Engine/Source/Runtime/VulkanRHI/Private/VulkanTexture.cpp
--- a/Engine/Source/Runtime/VulkanRHI/Private/VulkanTexture.cpp
+++ b/Engine/Source/Runtime/VulkanRHI/Private/VulkanTexture.cpp
@@ -1224,21 +1224,32 @@ void FVulkanDynamicRHI::InternalUpdateTexture3D(bool bFromRenderingThread, FText
 {
 	FVulkanTexture3D* Texture = ResourceCast(TextureRHI);
 
-	EPixelFormat PixelFormat = Texture->GetFormat();
-	check(GPixelFormats[PixelFormat].BlockSizeX == 1);
-	check(GPixelFormats[PixelFormat].BlockSizeY == 1);
+	const EPixelFormat PixelFormat = Texture->GetFormat();
+	const int32 BlockSizeX = GPixelFormats[PixelFormat].BlockSizeX;
+	const int32 BlockSizeY = GPixelFormats[PixelFormat].BlockSizeY;
+	const int32 BlockSizeZ = GPixelFormats[PixelFormat].BlockSizeZ;
+	const int32 BlockBytes = GPixelFormats[PixelFormat].BlockBytes;
+	const VkFormat Format = UEToVkFormat(PixelFormat, false);
 
-	VkFormat Format = UEToVkFormat(PixelFormat, false);
+	ensure(BlockSizeZ == 1);
 
 	FVulkanCommandListContext& Context = Device->GetImmediateContext();
 	const VkPhysicalDeviceLimits& Limits = Device->GetLimits();
 
-	const uint32 AlignedSourcePitch = Align(SourceRowPitch, Limits.optimalBufferCopyRowPitchAlignment);	
-	const int32 SlicePitch = AlignedSourcePitch * UpdateRegion.Height;
-	const uint32 BufferSize = Align(UpdateRegion.Depth * SlicePitch, Limits.minMemoryMapAlignment);
+	VkBufferImageCopy Region;
+	FMemory::Memzero(Region);
+	VulkanRHI::FStagingBuffer* StagingBuffer = nullptr;
+	const uint32 NumBlocksX = (uint32)FMath::DivideAndRoundUp<int32>(UpdateRegion.Width, (uint32)BlockSizeX);
+	const uint32 NumBlocksY = (uint32)FMath::DivideAndRoundUp<int32>(UpdateRegion.Height, (uint32)BlockSizeY);
+	check(NumBlocksX * BlockBytes <= SourceRowPitch);
+	check(NumBlocksX * BlockBytes * NumBlocksY <= SourceDepthPitch);
 
-	VulkanRHI::FStagingBuffer* StagingBuffer = Device->GetStagingManager().AcquireBuffer(BufferSize);
-	void* Memory = StagingBuffer->GetMappedPointer();
+	const uint32 DestRowPitch = NumBlocksX * BlockBytes;
+	const uint32 DestSlicePitch = DestRowPitch * NumBlocksY;
+
+	const uint32 BufferSize = Align(DestSlicePitch * UpdateRegion.Depth, Limits.minMemoryMapAlignment);
+	StagingBuffer = Device->GetStagingManager().AcquireBuffer(BufferSize);
+	void* RESTRICT Memory = StagingBuffer->GetMappedPointer();
 
 	VkImageSubresourceRange SubresourceRange;
 	FMemory::Memzero(SubresourceRange);
@@ -1248,29 +1259,25 @@ void FVulkanDynamicRHI::InternalUpdateTexture3D(bool bFromRenderingThread, FText
 	//SubresourceRange.baseArrayLayer = 0;
 	SubresourceRange.layerCount = 1;
 
-	uint32 CopyPitch = UpdateRegion.Width * GPixelFormats[Texture->GetFormat()].BlockBytes;
-	check(CopyPitch <= SourceRowPitch);
-	uint8* RowData = (uint8*)Memory;
-	for (uint32 i = 0; i < UpdateRegion.Depth; i++)
+	uint8* RESTRICT DestData = (uint8*)Memory;
+	for (uint32 Depth = 0; Depth < UpdateRegion.Depth; Depth++)
 	{
-		uint8* DestRowData = RowData + SlicePitch * i;
-		const uint8* SourceRowData = SourceData + SourceDepthPitch * i;
-		for (uint32 j = 0; j < UpdateRegion.Height; j++)
+		uint8* RESTRICT SourceRowData = (uint8*)SourceData + SourceDepthPitch * Depth;
+		for (uint32 Height = 0; Height < NumBlocksY; ++Height)
 		{
-			FMemory::Memcpy(DestRowData, SourceRowData, CopyPitch);
+			FMemory::Memcpy(DestData, SourceRowData, NumBlocksX * BlockBytes);
+			DestData += DestRowPitch;
 			SourceRowData += SourceRowPitch;
-			DestRowData += AlignedSourcePitch;
 		}
 	}
 
-	VkBufferImageCopy Region;
-	FMemory::Memzero(Region);
 	//Region.bufferOffset = 0;
-	Region.bufferRowLength = UpdateRegion.Width;
-	Region.bufferImageHeight = UpdateRegion.Height;
+	// Set these to zero to assume tightly packed buffer
+	//Region.bufferRowLength = 0;
+	//Region.bufferImageHeight = 0;
 	Region.imageSubresource.aspectMask = VK_IMAGE_ASPECT_COLOR_BIT;
 	Region.imageSubresource.mipLevel = MipIndex;
-	Region.imageSubresource.baseArrayLayer = 0;
+	//Region.imageSubresource.baseArrayLayer = 0;
 	Region.imageSubresource.layerCount = 1;
 	Region.imageOffset.x = UpdateRegion.DestX;
 	Region.imageOffset.y = UpdateRegion.DestY;
