Reverting logic to determine Lumin shader platform

diff --git a/Engine/Source/Runtime/VulkanRHI/Private/VulkanRHI.cpp b/Engine/Source/Runtime/VulkanRHI/Private/VulkanRHI.cpp
--- a/Engine/Source/Runtime/VulkanRHI/Private/VulkanRHI.cpp
+++ b/Engine/Source/Runtime/VulkanRHI/Private/VulkanRHI.cpp
@@ -63,7 +63,7 @@ FDynamicRHI* FVulkanDynamicRHIModule::CreateRHI(ERHIFeatureLevel::Type InRequest
 			FParse::Param(FCommandLine::Get(), TEXT("featureleveles31")) || FParse::Param(FCommandLine::Get(), TEXT("featureleveles2"))))
 	{
 		GMaxRHIFeatureLevel = ERHIFeatureLevel::ES3_1;
-		GMaxRHIShaderPlatform = PLATFORM_ANDROID ? SP_VULKAN_ES3_1_ANDROID : (PLATFORM_LUMIN ?  SP_VULKAN_ES3_1_LUMIN : SP_VULKAN_PCES3_1);
+		GMaxRHIShaderPlatform = PLATFORM_LUMIN ? SP_VULKAN_ES3_1_LUMIN : (PLATFORM_ANDROID ? SP_VULKAN_ES3_1_ANDROID : SP_VULKAN_PCES3_1);
 	}
 	else if (InRequestedFeatureLevel == ERHIFeatureLevel::SM4)
 	{
