UE4.20 - Fixes for Vulkan

diff --git a/Engine/Source/Developer/ShaderCompilerCommon/Private/ShaderCompilerCommon.cpp b/Engine/Source/Developer/ShaderCompilerCommon/Private/ShaderCompilerCommon.cpp
--- a/Engine/Source/Developer/ShaderCompilerCommon/Private/ShaderCompilerCommon.cpp
+++ b/Engine/Source/Developer/ShaderCompilerCommon/Private/ShaderCompilerCommon.cpp
@@ -780,6 +780,10 @@ namespace CrossCompiler
 
 		// +1 for EOL
 		const TCHAR* Ptr = &String[FoundBegin + 1 + Prolog.Len()];
+		while (*Ptr == '\r' || *Ptr == '\n')
+		{
+			++Ptr;
+		}
 		const TCHAR* PtrEnd = &String[FoundEnd];
 		while (Ptr < PtrEnd)
 		{
@@ -797,6 +801,8 @@ namespace CrossCompiler
 			{
 				return;
 			}
+			// Optional \r
+			CrossCompiler::Match(Ptr, '\r');
 			if (!CrossCompiler::Match(Ptr, '\n'))
 			{
 				return;
@@ -843,6 +849,8 @@ namespace CrossCompiler
 			{
 				return;
 			}
+			// Optional
+			CrossCompiler::Match(Ptr, '\r');
 			if (!CrossCompiler::Match(Ptr, '\n'))
 			{
 				return;
diff --git a/Engine/Source/Developer/VulkanShaderFormat/Private/VulkanBackend.cpp b/Engine/Source/Developer/VulkanShaderFormat/Private/VulkanBackend.cpp
--- a/Engine/Source/Developer/VulkanShaderFormat/Private/VulkanBackend.cpp
+++ b/Engine/Source/Developer/VulkanShaderFormat/Private/VulkanBackend.cpp
@@ -144,6 +144,15 @@ static inline std::string FixHlslName(const glsl_type* Type, bool bUseTextureIns
 	return Name;
 }
 
+static inline void Scanf(char* Dest, const char* Format, float* OutValue)
+{
+#if PLATFORM_WINDOWS
+	sscanf_s(Dest, Format, OutValue);
+#else
+	sscanf(Dest, Format, OutValue);
+#endif
+}
+
 /**
 * This table must match the ir_expression_operation enum.
 */
@@ -602,20 +611,20 @@ static void DumpSortedRanges(TDMARangeList& SortedRanges)
 	}
 }
 
-static inline EDescriptorSetStage GetDescriptorSetForStage(_mesa_glsl_parser_targets Target)
+static inline DescriptorSet::EStage GetDescriptorSetForStage(_mesa_glsl_parser_targets Target)
 {
 	switch (Target)
 	{
-	case vertex_shader: return GetDescriptorSetForStage(SF_Vertex);
-	case fragment_shader: return GetDescriptorSetForStage(SF_Pixel);
-	case compute_shader: return GetDescriptorSetForStage(SF_Compute);
-	case geometry_shader: return GetDescriptorSetForStage(SF_Geometry);
-	case tessellation_evaluation_shader: return GetDescriptorSetForStage(SF_Domain);
-	case tessellation_control_shader: return GetDescriptorSetForStage(SF_Hull);
+	case vertex_shader:						return DescriptorSet::GetSetForFrequency(SF_Vertex);
+	case fragment_shader:					return DescriptorSet::GetSetForFrequency(SF_Pixel);
+	case compute_shader:					return DescriptorSet::GetSetForFrequency(SF_Compute);
+	case geometry_shader:					return DescriptorSet::GetSetForFrequency(SF_Geometry);
+	case tessellation_evaluation_shader:	return DescriptorSet::GetSetForFrequency(SF_Domain);
+	case tessellation_control_shader:		return DescriptorSet::GetSetForFrequency(SF_Hull);
 	default: check(0);	break;	// NOT IMPLEMENTED!
 	}
 
-	return EDescriptorSetStage::Invalid;
+	return DescriptorSet::Invalid;
 }
 
 /**
@@ -1195,10 +1204,16 @@ class FGenerateVulkanVisitor : public ir_visitor
 				uint32 Interpolation = var->interpolation;
 				if (var->type->is_sampler())
 				{
+					check(var->name);
+					EVulkanBindingType::EType BindingType = var->type->sampler_buffer
+						? EVulkanBindingType::UniformTexelBuffer
+						: (SamplerMapping.StandaloneTextures.find(std::string(var->name)) != SamplerMapping.StandaloneTextures.end()
+							? EVulkanBindingType::Image
+							: EVulkanBindingType::CombinedImageSampler);
 					layout = ralloc_asprintf(nullptr,
 						"layout(set=%d, binding=BINDING_%d) ",
 						GetDescriptorSetForStage(ParseState->target),
-						BindingTable.RegisterBinding(var->name, "s", var->type->sampler_buffer ? EVulkanBindingType::UniformTexelBuffer : EVulkanBindingType::CombinedImageSampler));
+						BindingTable.RegisterBinding(var->name, "s", BindingType));
 				}
 				else if (bGenerateLayoutLocations && var->explicit_location)
 				{
@@ -1490,9 +1505,7 @@ class FGenerateVulkanVisitor : public ir_visitor
 			!SamplerMapping.UseCombinedImageSamplerForTexture(tex->sampler->variable_referenced()->name))
 		{
 			uint32 SSIndex = AddUniqueSamplerState(tex->SamplerStateName);
-			char PackedName[256];
-			FCStringAnsi::Sprintf(PackedName, "%sz%u", glsl_variable_tag_from_parser_target(ParseState->target), SSIndex);
-			BindingTable.RegisterBinding(PackedName, "z", EVulkanBindingType::Sampler);
+			BindingTable.RegisterBinding(tex->SamplerStateName, "z", EVulkanBindingType::Sampler);
 
 			auto GetSamplerSuffix = [](int32 Dim)
 			{
@@ -1511,7 +1524,7 @@ class FGenerateVulkanVisitor : public ir_visitor
 			ralloc_asprintf_append(buffer, "sampler%s%s(", GetSamplerSuffix(tex->sampler->type->sampler_dimensionality),
 				tex->sampler->type->sampler_array ? "Array" : "");
 			tex->sampler->accept(this);
-			ralloc_asprintf_append(buffer, ", %s)", PackedName);
+			ralloc_asprintf_append(buffer, ", %s)", tex->SamplerStateName);
 		}
 		else
 		{
@@ -1915,14 +1928,97 @@ class FGenerateVulkanVisitor : public ir_visitor
 			{
 				float value = constant->value.f[index];
 				// Original formatting code relied on %f style formatting
-				// %e is more accureate, and has been available since at least ES 2.0
+				// %e is more accurate, and has been available since at least ES 2.0
 				// leaving original code in place, in case some drivers don't properly handle it
 #if 0
 				const char *format = (fabsf(fmodf(value, 1.0f)) < 1.e-8f) ? "%.1f" : "%.8f";
+				ralloc_asprintf_append(buffer, format, value);
 #else
-				const char *format = "%e";
+/*
+				f= 0;
+					%f =	0.000000
+					%e =	0.000000e+00
+					%.16g =	0
+					%g =	0
+				f=  256.0 / 255.0;
+					%f =	1.003922
+					%e =	1.003922e+00
+					%.16g =	1.003921627998352
+					%g =	1.00392
+				f=  2.3283064365386963e-10;
+					%f =	0.000000
+					%e =	2.328306e-10
+					%.16g =	2.328306436538696e-10
+					%g =	2.32831e-10
+				f=  1e-6;
+					%f =	0.000001
+					%e =	1.000000e-06
+					%.16g =	9.999999974752427e-07
+					%g =	1e-06
+				f=  -1;
+					%f =	-1.000000
+					%e =	-1.000000e+00
+					%.16g =	-1
+					%g =	-1
+				f=  -1000;
+					%f =	-1000.000000
+					%e =	-1.000000e+03
+					%.16g =	-1000
+					%g =	-1000
+				f=  -1048576;
+					%f =	-1048576.000000
+					%e =	-1.048576e+06
+					%.16g =	-1048576
+					%g =	-1.04858e+06
+				f=  UINT_MAX;
+					%f =	4294967296.000000
+					%e =	4.294967e+09
+					%.16g =	4294967296
+					%g =	4.29497e+09
+				f=	3.402823466e+38
+					%f=		340282346638528859811704183484516925440.000000
+					%e=		3.402823e+38
+					%.16g=	3.402823466385289e+38
+					%g=		3.40282e+38
+*/
+				// Not fast, but way more precise
+				{
+					// 128 to handle the case of 3.402823466e+38f (which we use!)
+					char f[128], g[128], g10[128];
+					float ReadValue;
+					// Always add decimal point as glsl is painful
+					FCStringAnsi::Sprintf(g, "%g", value);
+					Scanf(g, "%f", &ReadValue);
+					bool bHasDecimalPoint = FCStringAnsi::Strchr(g, '.') != nullptr;
+					if (bHasDecimalPoint && ReadValue == value)
+					{
+						ralloc_strcat(buffer, g);
+					}
+					else
+					{
+						FCStringAnsi::Sprintf(f, "%f", value);
+						Scanf(f, "%f", &ReadValue);
+						if (ReadValue == value)
+						{
+							ralloc_strcat(buffer, f);
+						}
+						else
+						{
+							bHasDecimalPoint = FCStringAnsi::Strchr(g, '.') != nullptr;
+							FCStringAnsi::Sprintf(g10, "%.10g", value);
+							Scanf(g10, "%f", &ReadValue);
+							if (ReadValue == value)
+							{
+								ralloc_strcat(buffer, g10);
+							}
+							else
+							{
+								ralloc_asprintf_append(buffer, "%.16e", value);
+							}
+						}
+					}
+				}
 #endif
-				ralloc_asprintf_append(buffer, format, value);
 			}
 			else
 			{
@@ -3308,24 +3404,27 @@ public:
 		// Here since the code_buffer must have been populated beforehand
 		if (ParseState->LanguageSpec->AllowsSharingSamplers())
 		{
-			auto FindPrecision = [&](int32 Index)
+			auto FindPrecision = [&](const char* Name)
 			{
 				for (auto& Pair : state->TextureToSamplerMap)
 				{
 					for (auto& Entry : Pair.second)
 					{
-						if (Entry == SamplerStateNames[Index])
+						for (const std::string& SSName : SamplerStateNames)
 						{
-							for (auto& PackedEntry : state->GlobalPackedArraysMap['s'])
+							if (Entry == SSName)
 							{
-								if (!FCStringAnsi::Strcmp(Pair.first.c_str(), PackedEntry.Name.c_str()))
+								for (auto& PackedEntry : state->GlobalPackedArraysMap['s'])
 								{
-									foreach_iter(exec_list_iterator, iter, sampler_variables)
+									if (!FCStringAnsi::Strcmp(Pair.first.c_str(), PackedEntry.Name.c_str()))
 									{
-										ir_variable* var = ((extern_var*)iter.get())->var;
-										if (!FCStringAnsi::Strcmp(var->name, PackedEntry.CB_PackedSampler.c_str()))
+										foreach_iter(exec_list_iterator, iter, sampler_variables)
 										{
-											return GetPrecisionModifierName(GetPrecisionModifier(var->type));
+											ir_variable* var = ((extern_var*)iter.get())->var;
+											if (!FCStringAnsi::Strcmp(var->name, PackedEntry.CB_PackedSampler.c_str()))
+											{
+												return GetPrecisionModifierName(GetPrecisionModifier(var->type));
+											}
 										}
 									}
 								}
@@ -3341,14 +3440,13 @@ public:
 			{
 				if (BindingTable.Bindings[Index].Type == EVulkanBindingType::Sampler)
 				{
-					int32 Binding = atoi(BindingTable.Bindings[Index].Name + 2);
-					const char* Precision = FindPrecision(Binding);
+					const char* Precision = FindPrecision(BindingTable.Bindings[Index].Name);
 
-					ralloc_asprintf_append(buffer, "layout(set=%d, binding=BINDING_%d) uniform %s sampler %sz%d;\n",
+					ralloc_asprintf_append(buffer, "layout(set=%d, binding=BINDING_%d) uniform %s sampler %s;\n",
 						GetDescriptorSetForStage(ParseState->target),
 						Index,
 						Precision,
-						glsl_variable_tag_from_parser_target(ParseState->target), Binding);
+						BindingTable.Bindings[Index].Name);
 				}
 			}
 		}
@@ -3771,8 +3869,8 @@ static FSystemValue GeometrySystemValueTable[] =
 {
 	{ "SV_VertexID", glsl_type::int_type, "gl_VertexID", ir_var_in, false, false, false, false },
 	{ "SV_InstanceID", glsl_type::int_type, "gl_InstanceID", ir_var_in, false, false, false, false },
-	{ "SV_Position", glsl_type::vec4_type, "gl_Position", ir_var_in, false, true, false, false },
-	{ "SV_Position", glsl_type::vec4_type, "gl_Position", ir_var_out, false, false, false, false },
+	{ "SV_Position", glsl_type::vec4_type, "gl_Position", ir_var_in, false, true, true, false },
+	{ "SV_Position", glsl_type::vec4_type, "gl_Position", ir_var_out, false, false, true, false },
 	{ "SV_RenderTargetArrayIndex", glsl_type::int_type, "gl_Layer", ir_var_out, false, false, false, false },
 	{ "SV_PrimitiveID", glsl_type::int_type, "gl_PrimitiveID", ir_var_out, false, false, false, false },
 	{ "SV_PrimitiveID", glsl_type::int_type, "gl_PrimitiveIDIn", ir_var_in, false, false, false, false },
@@ -4497,6 +4595,7 @@ static void GenShaderInputForVariable(
 			if (ParseState->adjust_clip_space_dx11_to_opengl && ApplyClipSpaceAdjustment)
 			{
 				// This is for input of gl_Position into geometry shader only.
+				check(Frequency == HSF_GeometryShader && InputSemantic && FCStringAnsi::Stricmp(InputSemantic, "SV_Position") == 0);
 
 				// Generate a local variable to do the conversion in, keeping source type.
 				ir_variable* TempVariable = new(ParseState)ir_variable(SrcValue->type, NULL, ir_var_temporary);
@@ -4523,6 +4622,9 @@ static void GenShaderInputForVariable(
 						)
 						);
 				}
+
+				// Use TempVariable anywhere you would have otherwise used SrcValue going forward...
+				SrcValue = TempVariableDeref->clone(ParseState, NULL);
 			}
 
 			apply_type_conversion(InputType, SrcValue, PreCallInstructions, ParseState, true, &loc);
@@ -5697,6 +5799,7 @@ int32 FVulkanBindingTable::RegisterBinding(const char* InName, const char* Block
 
 void FVulkanBindingTable::SortBindings()
 {
+	// Order is guaranteed to match EVulkanBindingType::EType
 	check(!bSorted);
 	Bindings.Sort([](const FBinding& A, const FBinding& B)
 		{
@@ -5710,12 +5813,36 @@ void FVulkanBindingTable::SortBindings()
 	bSorted = true;
 }
 
-void FVulkanBindingTable::PrintBindingTableDefines(char** Buffer)
+void FVulkanBindingTable::PrintBindingTableDefines(char** OutBuffer) const
 {
+	auto GetVulkanBindingTypeName = [](EVulkanBindingType::EType Type)
+		{
+			switch(Type)
+			{
+			case EVulkanBindingType::PackedUniformBuffer:	return "Packed UB";
+			case EVulkanBindingType::UniformBuffer:			return "Uniform Buffer";
+			case EVulkanBindingType::CombinedImageSampler:	return "Combined Image Sampler";
+			case EVulkanBindingType::Sampler:				return "Sampler";
+			case EVulkanBindingType::Image:					return "Image";
+			case EVulkanBindingType::UniformTexelBuffer:	return "Uniform Texel Buffer";
+			case EVulkanBindingType::StorageImage:			return "Storage Image";
+			case EVulkanBindingType::StorageTexelBuffer:	return "Storage TexelBuffer";
+			case EVulkanBindingType::StorageBuffer:			return "Storage Buffer";
+			default:										return "INVALID!";
+			}
+		};
+	EVulkanBindingType::EType PreviousType = EVulkanBindingType::Count;
+	ralloc_asprintf_append(OutBuffer, "\n");
 	for (int32 Index = 0; Index < Bindings.Num(); ++Index)
 	{
-		ralloc_asprintf_append(Buffer, "#define BINDING_%d\t%d\n", Bindings[Index].VirtualIndex, Index);
+		if (PreviousType != Bindings[Index].Type)
+		{
+			ralloc_asprintf_append(OutBuffer, "// %s\n", GetVulkanBindingTypeName(Bindings[Index].Type));
+			PreviousType = Bindings[Index].Type;
+		}
+		ralloc_asprintf_append(OutBuffer, "#define BINDING_%d\t%d\n", Bindings[Index].VirtualIndex, Index);
 	}
+	ralloc_asprintf_append(OutBuffer, "\n");
 }
 
 struct FFixIntrinsicsVisitor : public ir_rvalue_visitor
diff --git a/Engine/Source/Developer/VulkanShaderFormat/Private/VulkanBackend.h b/Engine/Source/Developer/VulkanShaderFormat/Private/VulkanBackend.h
--- a/Engine/Source/Developer/VulkanShaderFormat/Private/VulkanBackend.h
+++ b/Engine/Source/Developer/VulkanShaderFormat/Private/VulkanBackend.h
@@ -72,7 +72,7 @@ struct FVulkanBindingTable
 	}
 
 	void SortBindings();
-	void PrintBindingTableDefines(char** Buffer);
+	void PrintBindingTableDefines(char** OutBuffer) const;
 
 private:
 	// Previous implementation supported bindings only for textures.
diff --git a/Engine/Source/Developer/VulkanShaderFormat/Private/VulkanGlslangBridge.cpp b/Engine/Source/Developer/VulkanShaderFormat/Private/VulkanGlslangBridge.cpp
--- a/Engine/Source/Developer/VulkanShaderFormat/Private/VulkanGlslangBridge.cpp
+++ b/Engine/Source/Developer/VulkanShaderFormat/Private/VulkanGlslangBridge.cpp
@@ -23,6 +23,7 @@ THIRD_PARTY_INCLUDES_START
 	#include "SPIRV/GLSL.std.450.h"
 	#include "SPIRV/doc.h"
 	#include "SPIRV/disassemble.h"
+	#include "SPIRV/spirv.hpp"
 THIRD_PARTY_INCLUDES_END
 
 #if defined(_MSC_VER) && _MSC_VER >= 1800
@@ -164,6 +165,158 @@ static EShLanguage GetStage(EHlslShaderFrequency Frequency)
 	return EShLangCount;
 }
 
+static void ComputeMovableWordIndices(FSpirv& Spirv)
+{
+	// SPIRV Header
+	const uint32* PtrStart = Spirv.Data.GetData();
+	const uint32* Ptr = PtrStart;
+	const uint32* PtrEnd = Spirv.Data.GetData() + Spirv.Data.Num();
+	check(*Ptr++ == spv::MagicNumber);
+	uint32_t Version = *Ptr++;
+	uint32_t Generator = *Ptr++;
+	uint32_t Bound = *Ptr++;
+	check(*Ptr++ == 0);
+
+	auto ReadLiteralString = [](const uint32_t* Ptr)
+	{
+		FString S;
+		const char* Str = (const char*)Ptr;
+		// Empty string is allowed...
+		while (*Str)
+		{
+			S += *Str;
+			++Str;
+		}
+		return S;
+	};
+
+	TMap<uint32, FString> Names;
+	struct FDecorations
+	{
+		uint32 BindingIndex = UINT32_MAX;
+		uint32 WordBindingIndex = UINT32_MAX;
+		uint32 DescriptorSet = UINT32_MAX;
+		uint32 WordDescriptorSet = UINT32_MAX;
+	};
+	TMap<uint32, FDecorations> Decorations;
+	TMap<uint32, uint32> TypePointerUniforms;
+	TMap<uint32, uint32> VariableUniformTypes;
+	bool bDone = false;
+	while (Ptr < PtrEnd && !bDone)
+	{
+		uint32_t WordCount = (*Ptr >> spv::WordCountShift) & spv::OpCodeMask;
+		spv::Op OpCode = (spv::Op)(*Ptr & spv::OpCodeMask);
+		switch (OpCode)
+		{
+		case spv::OpName:
+			{
+				uint32 TargetId = Ptr[1];
+				FString Name = ReadLiteralString(Ptr + 2);
+				Names.Add(TargetId, Name);
+			}
+			break;
+		case spv::OpDecorate:
+		{
+			uint32 TargetId = Ptr[1];
+			spv::Decoration Decoration = (spv::Decoration)Ptr[2];
+			switch (Decoration)
+			{
+			case spv::DecorationDescriptorSet:
+			{
+				uint32 Value = Ptr[3];
+				uint32 WordValueIndex = (uint32)(&Ptr[3] - PtrStart);
+				FDecorations& UBDecoration = Decorations.FindOrAdd(TargetId);
+				UBDecoration.DescriptorSet = Value;
+				UBDecoration.WordDescriptorSet = WordValueIndex;
+				break;
+			}
+			case spv::DecorationBinding:
+			{
+				uint32 Value = Ptr[3];
+				uint32 WordValueIndex = (uint32)(&Ptr[3] - PtrStart);
+				FDecorations& UBDecoration = Decorations.FindOrAdd(TargetId);
+				UBDecoration.BindingIndex = Value;
+				UBDecoration.WordBindingIndex = WordValueIndex;
+				break;
+			}
+			default:
+				break;
+			}
+		}
+			break;
+		case spv::OpTypePointer:
+		{
+			uint32 Result = Ptr[1];
+			spv::StorageClass Storage = (spv::StorageClass)Ptr[2];
+			if (Storage == spv::StorageClassUniform || spv::StorageClassUniformConstant)
+			{
+				uint32 Type = Ptr[3];
+				TypePointerUniforms.Add(Result, Type);
+			}
+		}
+			break;
+		case spv::OpVariable:
+		{
+			uint32 Type = Ptr[1];
+			uint32 Id = Ptr[2];
+			spv::StorageClass Storage = (spv::StorageClass)Ptr[3];
+			if (Storage == spv::StorageClassUniform || Storage == spv::StorageClassUniformConstant)
+			{
+				VariableUniformTypes.Add(Id, Type);
+			}
+		}
+			break;
+		case spv::OpFunction:
+			bDone = true;
+			break;
+		default:
+			break;
+		}
+
+		Ptr += WordCount;
+	}
+
+	check((bDone && Ptr < PtrEnd) || (!bDone && Ptr == PtrEnd));
+
+	// Go through all found uniform variables and make sure we found the right info
+	for (const auto Pair : VariableUniformTypes)
+	{
+		const uint32 VariableId = Pair.Key;
+		const FString* FoundVariableName = Names.Find(VariableId);
+		if (FoundVariableName)
+		{
+			const uint32 VariableType = Pair.Value;
+			if (FoundVariableName->Len() == 0)
+			{
+				// Uniform buffer
+				uint32 TypePointer = TypePointerUniforms.FindChecked(VariableType);
+				const FString* FoundTypeName = Names.Find(TypePointer);
+				if (FoundTypeName && FoundTypeName->Len() > 0)
+				{
+					FSpirv::FEntry* FoundEntry = Spirv.GetEntry(*FoundTypeName);
+					check(FoundEntry);
+					FDecorations& FoundDecorations = Decorations.FindChecked(VariableId);
+					FoundEntry->Binding = FoundDecorations.BindingIndex;
+					FoundEntry->WordBindingIndex = FoundDecorations.WordBindingIndex;
+					FoundEntry->DescriptorSet = FoundDecorations.DescriptorSet;
+					FoundEntry->WordDescriptorSetIndex = FoundDecorations.WordDescriptorSet;
+				}
+			}
+			else
+			{
+				// Standalone global var
+				FSpirv::FEntry* FoundEntry = Spirv.GetEntry(*FoundVariableName);
+				check(FoundEntry);
+				FDecorations& FoundDecorations = Decorations.FindChecked(VariableId);
+				FoundEntry->Binding = FoundDecorations.BindingIndex;
+				FoundEntry->WordBindingIndex = FoundDecorations.WordBindingIndex;
+				FoundEntry->DescriptorSet = FoundDecorations.DescriptorSet;
+				FoundEntry->WordDescriptorSetIndex = FoundDecorations.WordDescriptorSet;
+			}
+		}
+	}
+}
+
 bool GenerateSpirv(const ANSICHAR* Source, FCompilerInfo& CompilerInfo, FString& OutErrors, const FString& DumpDebugInfoPath, FSpirv& OutSpirv)
 {
 	glslang::TProgram* Program = new glslang::TProgram;
@@ -209,9 +362,9 @@ bool GenerateSpirv(const ANSICHAR* Source, FCompilerInfo& CompilerInfo, FString&
 		std::vector<unsigned int> Spirv;
 		glslang::GlslangToSpv(*Program->getIntermediate((EShLanguage)Stage), Spirv);
 
-		uint32 SizeInBytes = Spirv.size() * sizeof(Spirv[0]);
-		OutSpirv.Data.AddZeroed(SizeInBytes);
-		FMemory::Memcpy(OutSpirv.Data.GetData(), &Spirv[0], SizeInBytes);
+		uint32 SizeInWords = Spirv.size();
+		OutSpirv.Data.AddZeroed(SizeInWords);
+		FMemory::Memcpy(OutSpirv.Data.GetData(), &Spirv[0], SizeInWords * sizeof(uint32));
 
 		if (CompilerInfo.bDebugDump)
 		{
@@ -277,6 +430,8 @@ bool GenerateSpirv(const ANSICHAR* Source, FCompilerInfo& CompilerInfo, FString&
 			OutSpirv.ReflectionInfo.Add(Entry);
 		}
 
+		ComputeMovableWordIndices(OutSpirv);
+
 		return true;
 	};
 
diff --git a/Engine/Source/Developer/VulkanShaderFormat/Private/VulkanShaderCompiler.cpp b/Engine/Source/Developer/VulkanShaderFormat/Private/VulkanShaderCompiler.cpp
--- a/Engine/Source/Developer/VulkanShaderFormat/Private/VulkanShaderCompiler.cpp
+++ b/Engine/Source/Developer/VulkanShaderFormat/Private/VulkanShaderCompiler.cpp
@@ -580,39 +580,13 @@ static void BuildShaderOutput(
 		Header.NEWPackedGlobalUBSizes[PackedArrayIndex] = Align((uint32)Pair.Value, (uint32)16);
 	}
 
-	auto FindHlslccBindingByIndex = [&HlslccBindings](int32 TestIndex, EVulkanBindingType::EType Type)
-	{
-		int32 VulkanBindingIndex = -1;
-		int32 Index = 0;
-		for (Index = 0; Index < HlslccBindings.Num(); ++Index)
-		{
-			const FVulkanBindingTable::FBinding& Binding = HlslccBindings[Index];
-			if (Binding.Type == Type)
-			{
-				check(IsAlpha(Binding.Name[0]) && IsAlpha(Binding.Name[1]));
-				check(IsDigit(Binding.Name[2]));
-				const char* IndexStr = &Binding.Name[2];
-				int32 IndexFromName = FCStringAnsi::Atoi(IndexStr);
-				if (IndexFromName == TestIndex)
-				{
-					return Index;
-				}
-			}
-	}
-
-		return -1;
-	};
-
-	// shared samplers that are actually shared end up being generated as raw sampler uniforms by hlslcc, which follow special naming rules.  they are output as "pz0", "pz1", "pz2", etc...
-	//  take advantage of the FindHlslccBindingByIndex function someone already defined above (but never used?) and scrape these out of the file.
-	//  Note that hlslcc appears to *only* dump raw samplers to its SamplerStates array, and these raw samplers are in indexed order with the 'p'z enumeration.
 	TSet<FString> SharedSamplerStates;
 	for (int32 i = 0; i < CCHeader.SamplerStates.Num(); i++)
 	{
-		int32 HlslccBindingIndex = FindHlslccBindingByIndex(i, EVulkanBindingType::Sampler);
+		const FString& Name = CCHeader.SamplerStates[i].Name;
+		int32 HlslccBindingIndex = Spirv.FindBinding(Name);
 		check(HlslccBindingIndex != -1);
 
-		const FString& Name = CCHeader.SamplerStates[i].Name;
 		SharedSamplerStates.Add(Name);
 		auto& Binding = HlslccBindings[HlslccBindingIndex];
 		int32 BindingIndex = Spirv.FindBinding(Binding.Name, true);
@@ -739,6 +713,7 @@ static void BuildShaderOutput(
 	AppendCString(DebugNameArray, TCHAR_TO_ANSI(*DebugName));
 	Ar << DebugNameArray;
 
+	check(Spirv.Data.Num() != 0);
 	Ar << Spirv.Data;
 
 	// store data we can pickup later with ShaderCode.FindOptionalData('n'), could be removed for shipping
diff --git a/Engine/Source/Developer/VulkanShaderFormat/Private/VulkanShaderFormat.cpp b/Engine/Source/Developer/VulkanShaderFormat/Private/VulkanShaderFormat.cpp
--- a/Engine/Source/Developer/VulkanShaderFormat/Private/VulkanShaderFormat.cpp
+++ b/Engine/Source/Developer/VulkanShaderFormat/Private/VulkanShaderFormat.cpp
@@ -23,8 +23,8 @@ class FShaderFormatVulkan : public IShaderFormat
 {
 	enum 
 	{
-		UE_SHADER_VULKAN_ES3_1_VER = 15,
-		UE_SHADER_VULKAN_SM5_VER = 15,
+		UE_SHADER_VULKAN_ES3_1_VER = 17,
+		UE_SHADER_VULKAN_SM5_VER = 17,
 	};
 
 	int32 InternalGetVersion(FName Format) const
diff --git a/Engine/Source/Developer/VulkanShaderFormat/Private/VulkanShaderFormat.h b/Engine/Source/Developer/VulkanShaderFormat/Private/VulkanShaderFormat.h
--- a/Engine/Source/Developer/VulkanShaderFormat/Private/VulkanShaderFormat.h
+++ b/Engine/Source/Developer/VulkanShaderFormat/Private/VulkanShaderFormat.h
@@ -36,11 +36,19 @@ struct FCompilerInfo
 
 struct FSpirv
 {
-	TArray<uint8> Data;
+	TArray<uint32> Data;
 	struct FEntry
 	{
 		FString Name;
-		int32 Binding;
+		int32 Binding = -1;
+
+		uint32 DescriptorSet = UINT32_MAX;
+
+		// Index into the Spirv Word containing the descriptor set decoration
+		uint32 WordDescriptorSetIndex = UINT32_MAX;
+
+		// Index into the Spirv Word containing the binding index decoration
+		uint32 WordBindingIndex = UINT32_MAX;
 	};
 	TArray<FEntry> ReflectionInfo;
 
@@ -65,5 +73,44 @@ struct FSpirv
 
 		return -1;
 	}
+
+	const FEntry* GetEntryByBindingIndex(int32 BindingIndex) const
+	{
+		for (int32 Index = 0; Index < ReflectionInfo.Num(); ++Index)
+		{
+			if (ReflectionInfo[Index].Binding == BindingIndex)
+			{
+				return &ReflectionInfo[Index];
+			}
+		}
+
+		return nullptr;
+	}
+
+	FEntry* GetEntry(const FString& Name)
+	{
+		for (int32 Index = 0; Index < ReflectionInfo.Num(); ++Index)
+		{
+			if (ReflectionInfo[Index].Name == Name)
+			{
+				return &ReflectionInfo[Index];
+			}
+		}
+
+		return nullptr;
+	}
+
+	FEntry const* GetEntry(const FString& Name) const
+	{
+		for (int32 Index = 0; Index < ReflectionInfo.Num(); ++Index)
+		{
+			if (ReflectionInfo[Index].Name == Name)
+			{
+				return &ReflectionInfo[Index];
+			}
+		}
+
+		return nullptr;
+	}
 };
 extern bool GenerateSpirv(const ANSICHAR* Source, FCompilerInfo& CompilerInfo, FString& OutErrors, const FString& DumpDebugInfoPath, FSpirv& OutSpirv);
diff --git a/Engine/Source/Runtime/Renderer/Private/PostProcess/PostProcessTemporalAA.cpp b/Engine/Source/Runtime/Renderer/Private/PostProcess/PostProcessTemporalAA.cpp
--- a/Engine/Source/Runtime/Renderer/Private/PostProcess/PostProcessTemporalAA.cpp
+++ b/Engine/Source/Runtime/Renderer/Private/PostProcess/PostProcessTemporalAA.cpp
@@ -544,6 +544,11 @@ class FPostProcessTemporalAACS : public FGlobalShader
 
 IMPLEMENT_GLOBAL_SHADER(FPostProcessTemporalAACS, "/Engine/Private/PostProcessTemporalAA.usf", "MainCS", SF_Compute);
 
+static void TransitionShaderResources(FRenderingCompositePassContext& Context)
+{
+	TShaderMapRef< FPostProcessTonemapVS > VertexShader(Context.GetShaderMap());
+	VertexShader->TransitionResources(Context);
+}
 
 void DrawPixelPassTemplate(
 	FRenderingCompositePassContext& Context,
@@ -570,7 +575,6 @@ void DrawPixelPassTemplate(
 	GraphicsPSOInit.PrimitiveType = PT_TriangleList;
 
 	SetGraphicsPipelineState(Context.RHICmdList, GraphicsPSOInit);
-	VertexShader->TransitionResources(Context);
 	VertexShader->SetVS(Context);
 	PixelShader->SetParameters(Context.RHICmdList, Context, InputHistory, PassParameters, bUseDither, SrcSize);
 
@@ -852,6 +856,8 @@ void FRCPassPostProcessTemporalAA::Process(FRenderingCompositePassContext& Conte
 		// Inform MultiGPU systems that we're starting to update this resource
 		Context.RHICmdList.BeginUpdateMultiFrameResource(DestRenderTarget[0]->ShaderResourceTexture);
 
+		TransitionShaderResources(Context);
+
 		// Setup render targets.
 		{
 			FRHIDepthRenderTargetView DepthRTV(SceneContext.GetSceneDepthTexture(),
diff --git a/Engine/Source/Runtime/VulkanRHI/Private/Android/VulkanAndroidPlatform.cpp b/Engine/Source/Runtime/VulkanRHI/Private/Android/VulkanAndroidPlatform.cpp
--- a/Engine/Source/Runtime/VulkanRHI/Private/Android/VulkanAndroidPlatform.cpp
+++ b/Engine/Source/Runtime/VulkanRHI/Private/Android/VulkanAndroidPlatform.cpp
@@ -1,5 +1,13 @@
 // Copyright 1998-2018 Epic Games, Inc. All Rights Reserved.
 
+#ifndef PLATFORM_LUMIN
+	#define	PLATFORM_LUMIN	0
+#endif
+
+#ifndef PLATFORM_LUMINGL4
+	#define	PLATFORM_LUMINGL4	0
+#endif
+
 //#todo-Lumin: Remove this define when it becomes untangled from Android
 #if !PLATFORM_LUMIN && !PLATFORM_LUMINGL4
 #include "VulkanAndroidPlatform.h"
@@ -102,8 +110,7 @@ void FVulkanAndroidPlatform::FreeVulkanLibrary()
 void FVulkanAndroidPlatform::CreateSurface(void* WindowHandle, VkInstance Instance, VkSurfaceKHR* OutSurface)
 {
 	VkAndroidSurfaceCreateInfoKHR SurfaceCreateInfo;
-	FMemory::Memzero(SurfaceCreateInfo);
-	SurfaceCreateInfo.sType = VK_STRUCTURE_TYPE_ANDROID_SURFACE_CREATE_INFO_KHR;
+	ZeroVulkanStruct(SurfaceCreateInfo, VK_STRUCTURE_TYPE_ANDROID_SURFACE_CREATE_INFO_KHR);
 	SurfaceCreateInfo.window = (ANativeWindow*)WindowHandle;
 
 	VERIFYVULKANRESULT(vkCreateAndroidSurfaceKHR(Instance, &SurfaceCreateInfo, nullptr, OutSurface));
diff --git a/Engine/Source/Runtime/VulkanRHI/Private/Android/VulkanAndroidPlatform.h b/Engine/Source/Runtime/VulkanRHI/Private/Android/VulkanAndroidPlatform.h
--- a/Engine/Source/Runtime/VulkanRHI/Private/Android/VulkanAndroidPlatform.h
+++ b/Engine/Source/Runtime/VulkanRHI/Private/Android/VulkanAndroidPlatform.h
@@ -54,6 +54,8 @@ public:
 	//#todo-rco: Detect Mali?
 	static bool RequiresPresentLayoutFix() { return true; }
 
+	static bool HasUnifiedMemory() { return true; }
+
 	static bool RegisterGPUWork() { return false; }
 
 protected:
diff --git a/Engine/Source/Runtime/VulkanRHI/Private/Lumin/VulkanLuminPlatform.h b/Engine/Source/Runtime/VulkanRHI/Private/Lumin/VulkanLuminPlatform.h
--- a/Engine/Source/Runtime/VulkanRHI/Private/Lumin/VulkanLuminPlatform.h
+++ b/Engine/Source/Runtime/VulkanRHI/Private/Lumin/VulkanLuminPlatform.h
@@ -49,6 +49,8 @@ public:
 
 	static bool ForceEnableDebugMarkers();
 
+	static bool HasUnifiedMemory() { return true; }
+
 protected:
 	static void* VulkanLib;
 	static bool bAttemptedLoad;
diff --git a/Engine/Source/Runtime/VulkanRHI/Private/Mac/VulkanPlatformDefines.h b/Engine/Source/Runtime/VulkanRHI/Private/Mac/VulkanPlatformDefines.h
new file mode 100644
--- /dev/null
+++ b/Engine/Source/Runtime/VulkanRHI/Private/Mac/VulkanPlatformDefines.h
@@ -0,0 +1,3 @@
+// Copyright 1998-2018 Epic Games, Inc. All Rights Reserved.
+
+#pragma once
diff --git a/Engine/Source/Runtime/VulkanRHI/Private/VulkanCommandBuffer.cpp b/Engine/Source/Runtime/VulkanRHI/Private/VulkanCommandBuffer.cpp
--- a/Engine/Source/Runtime/VulkanRHI/Private/VulkanCommandBuffer.cpp
+++ b/Engine/Source/Runtime/VulkanRHI/Private/VulkanCommandBuffer.cpp
@@ -49,9 +49,7 @@ FVulkanCmdBuffer::FVulkanCmdBuffer(FVulkanDevice* InDevice, FVulkanCommandBuffer
 	FMemory::Memzero(CurrentScissor);
 	
 	VkCommandBufferAllocateInfo CreateCmdBufInfo;
-	FMemory::Memzero(CreateCmdBufInfo);
-	CreateCmdBufInfo.sType = VK_STRUCTURE_TYPE_COMMAND_BUFFER_ALLOCATE_INFO;
-	CreateCmdBufInfo.pNext = nullptr;
+	ZeroVulkanStruct(CreateCmdBufInfo, VK_STRUCTURE_TYPE_COMMAND_BUFFER_ALLOCATE_INFO);
 	CreateCmdBufInfo.level = VK_COMMAND_BUFFER_LEVEL_PRIMARY;
 	CreateCmdBufInfo.commandBufferCount = 1;
 	CreateCmdBufInfo.commandPool = CommandBufferPool->GetHandle();
@@ -89,8 +87,7 @@ void FVulkanCmdBuffer::BeginRenderPass(const FVulkanRenderTargetLayout& Layout,
 	checkf(IsOutsideRenderPass(), TEXT("Can't BeginRP as already inside one! CmdBuffer 0x%p State=%d"), CommandBufferHandle, (int32)State);
 
 	VkRenderPassBeginInfo Info;
-	FMemory::Memzero(Info);
-	Info.sType = VK_STRUCTURE_TYPE_RENDER_PASS_BEGIN_INFO;
+	ZeroVulkanStruct(Info, VK_STRUCTURE_TYPE_RENDER_PASS_BEGIN_INFO);
 	Info.renderPass = RenderPass->GetHandle();
 	Info.framebuffer = Framebuffer->GetHandle();
 	Info.renderArea.offset.x = 0;
@@ -155,8 +152,7 @@ void FVulkanCmdBuffer::Begin()
 	checkf(State == EState::ReadyForBegin, TEXT("Can't Begin as we're NOT ready! CmdBuffer 0x%p State=%d"), CommandBufferHandle, (int32)State);
 
 	VkCommandBufferBeginInfo CmdBufBeginInfo;
-	FMemory::Memzero(CmdBufBeginInfo);
-	CmdBufBeginInfo.sType = VK_STRUCTURE_TYPE_COMMAND_BUFFER_BEGIN_INFO;
+	ZeroVulkanStruct(CmdBufBeginInfo, VK_STRUCTURE_TYPE_COMMAND_BUFFER_BEGIN_INFO);
 	CmdBufBeginInfo.flags = VK_COMMAND_BUFFER_USAGE_ONE_TIME_SUBMIT_BIT;
 
 	VERIFYVULKANRESULT(VulkanRHI::vkBeginCommandBuffer(CommandBufferHandle, &CmdBufBeginInfo));
@@ -254,8 +250,7 @@ FVulkanCommandBufferPool::~FVulkanCommandBufferPool()
 void FVulkanCommandBufferPool::Create(uint32 QueueFamilyIndex)
 {
 	VkCommandPoolCreateInfo CmdPoolInfo;
-	FMemory::Memzero(CmdPoolInfo);
-	CmdPoolInfo.sType = VK_STRUCTURE_TYPE_COMMAND_POOL_CREATE_INFO;
+	ZeroVulkanStruct(CmdPoolInfo, VK_STRUCTURE_TYPE_COMMAND_POOL_CREATE_INFO);
 	CmdPoolInfo.queueFamilyIndex =  QueueFamilyIndex;
 	//#todo-rco: Should we use VK_COMMAND_POOL_CREATE_TRANSIENT_BIT?
 	CmdPoolInfo.flags = VK_COMMAND_POOL_CREATE_RESET_COMMAND_BUFFER_BIT;
diff --git a/Engine/Source/Runtime/VulkanRHI/Private/VulkanCommands.cpp b/Engine/Source/Runtime/VulkanRHI/Private/VulkanCommands.cpp
--- a/Engine/Source/Runtime/VulkanRHI/Private/VulkanCommands.cpp
+++ b/Engine/Source/Runtime/VulkanRHI/Private/VulkanCommands.cpp
@@ -202,42 +202,48 @@ void FVulkanCommandListContext::RHISetUAVParameter(FComputeShaderRHIParamRef Com
 
 void FVulkanCommandListContext::RHISetShaderTexture(FVertexShaderRHIParamRef VertexShaderRHI, uint32 TextureIndex, FTextureRHIParamRef NewTextureRHI)
 {
-	check(PendingGfxState->CurrentBSS && PendingGfxState->CurrentBSS->GetShader(SF_Vertex) == ResourceCast(VertexShaderRHI));
+	check(PendingGfxState->CurrentBSS && PendingGfxState->CurrentBSS->GetShader(DescriptorSet::Vertex) == ResourceCast(VertexShaderRHI));
 	FVulkanTextureBase* Texture = GetVulkanTextureFromRHITexture(NewTextureRHI);
 	VkImageLayout Layout = GetLayoutForDescriptor(Texture->Surface);
-	PendingGfxState->SetTexture(SF_Vertex, TextureIndex, Texture, Layout);
+	PendingGfxState->SetTexture(DescriptorSet::Vertex, TextureIndex, Texture, Layout);
 }
 
 void FVulkanCommandListContext::RHISetShaderTexture(FHullShaderRHIParamRef HullShaderRHI, uint32 TextureIndex, FTextureRHIParamRef NewTextureRHI)
 {
-	check(PendingGfxState->CurrentBSS && PendingGfxState->CurrentBSS->GetShader(SF_Hull) == ResourceCast(HullShaderRHI));
+	ensureMsgf(0, TEXT("Tessellation not supported yet!"));
+/*
+	check(PendingGfxState->CurrentBSS && PendingGfxState->CurrentBSS->GetShader(DescriptorSet::Hull) == ResourceCast(HullShaderRHI));
 	FVulkanTextureBase* Texture = GetVulkanTextureFromRHITexture(NewTextureRHI);
 	VkImageLayout Layout = GetLayoutForDescriptor(Texture->Surface);
-	PendingGfxState->SetTexture(SF_Hull, TextureIndex, Texture, Layout);
+	PendingGfxState->SetTexture(DescriptorSet::Hull, TextureIndex, Texture, Layout);
+*/
 }
 
 void FVulkanCommandListContext::RHISetShaderTexture(FDomainShaderRHIParamRef DomainShaderRHI, uint32 TextureIndex, FTextureRHIParamRef NewTextureRHI)
 {
-	check(PendingGfxState->CurrentBSS && PendingGfxState->CurrentBSS->GetShader(SF_Domain) == ResourceCast(DomainShaderRHI));
+	ensureMsgf(0, TEXT("Tessellation not supported yet!"));
+/*
+	check(PendingGfxState->CurrentBSS && PendingGfxState->CurrentBSS->GetShader(DescriptorSet::Domain) == ResourceCast(DomainShaderRHI));
 	FVulkanTextureBase* Texture = GetVulkanTextureFromRHITexture(NewTextureRHI);
 	VkImageLayout Layout = GetLayoutForDescriptor(Texture->Surface);
-	PendingGfxState->SetTexture(SF_Domain, TextureIndex, Texture, Layout);
+	PendingGfxState->SetTexture(DescriptorSet::Domain, TextureIndex, Texture, Layout);
+*/
 }
 
 void FVulkanCommandListContext::RHISetShaderTexture(FGeometryShaderRHIParamRef GeometryShaderRHI, uint32 TextureIndex, FTextureRHIParamRef NewTextureRHI)
 {
-	check(PendingGfxState->CurrentBSS && PendingGfxState->CurrentBSS->GetShader(SF_Geometry) == ResourceCast(GeometryShaderRHI));
+	check(PendingGfxState->CurrentBSS && PendingGfxState->CurrentBSS->GetShader(DescriptorSet::Geometry) == ResourceCast(GeometryShaderRHI));
 	FVulkanTextureBase* Texture = GetVulkanTextureFromRHITexture(NewTextureRHI);
 	VkImageLayout Layout = GetLayoutForDescriptor(Texture->Surface);
-	PendingGfxState->SetTexture(SF_Geometry, TextureIndex, Texture, Layout);
+	PendingGfxState->SetTexture(DescriptorSet::Geometry, TextureIndex, Texture, Layout);
 }
 
 void FVulkanCommandListContext::RHISetShaderTexture(FPixelShaderRHIParamRef PixelShaderRHI, uint32 TextureIndex, FTextureRHIParamRef NewTextureRHI)
 {
-	check(PendingGfxState->CurrentBSS && PendingGfxState->CurrentBSS->GetShader(SF_Pixel) == ResourceCast(PixelShaderRHI));
+	check(PendingGfxState->CurrentBSS && PendingGfxState->CurrentBSS->GetShader(DescriptorSet::Pixel) == ResourceCast(PixelShaderRHI));
 	FVulkanTextureBase* Texture = GetVulkanTextureFromRHITexture(NewTextureRHI);
 	VkImageLayout Layout = GetLayoutForDescriptor(Texture->Surface);
-	PendingGfxState->SetTexture(SF_Pixel, TextureIndex, Texture, Layout);
+	PendingGfxState->SetTexture(DescriptorSet::Pixel, TextureIndex, Texture, Layout);
 }
 
 void FVulkanCommandListContext::RHISetShaderTexture(FComputeShaderRHIParamRef ComputeShader, uint32 TextureIndex, FTextureRHIParamRef NewTextureRHI)
@@ -251,37 +257,43 @@ void FVulkanCommandListContext::RHISetShaderTexture(FComputeShaderRHIParamRef Co
 
 void FVulkanCommandListContext::RHISetShaderResourceViewParameter(FVertexShaderRHIParamRef VertexShaderRHI, uint32 TextureIndex, FShaderResourceViewRHIParamRef SRVRHI)
 {
-	check(PendingGfxState->CurrentBSS && PendingGfxState->CurrentBSS->GetShader(SF_Vertex) == ResourceCast(VertexShaderRHI));
+	check(PendingGfxState->CurrentBSS && PendingGfxState->CurrentBSS->GetShader(DescriptorSet::Vertex) == ResourceCast(VertexShaderRHI));
 	FVulkanShaderResourceView* SRV = ResourceCast(SRVRHI);
-	PendingGfxState->SetSRV(SF_Vertex, TextureIndex, SRV);
+	PendingGfxState->SetSRV(DescriptorSet::Vertex, TextureIndex, SRV);
 }
 
 void FVulkanCommandListContext::RHISetShaderResourceViewParameter(FHullShaderRHIParamRef HullShaderRHI,uint32 TextureIndex,FShaderResourceViewRHIParamRef SRVRHI)
 {
-	check(PendingGfxState->CurrentBSS && PendingGfxState->CurrentBSS->GetShader(SF_Hull) == ResourceCast(HullShaderRHI));
+	ensureMsgf(0, TEXT("Tessellation not supported yet!"));
+/*
+	check(PendingGfxState->CurrentBSS && PendingGfxState->CurrentBSS->GetShader(DescriptorSet::Hull) == ResourceCast(HullShaderRHI));
 	FVulkanShaderResourceView* SRV = ResourceCast(SRVRHI);
-	PendingGfxState->SetSRV(SF_Hull, TextureIndex, SRV);
+	PendingGfxState->SetSRV(DescriptorSet::Hull, TextureIndex, SRV);
+*/
 }
 
 void FVulkanCommandListContext::RHISetShaderResourceViewParameter(FDomainShaderRHIParamRef DomainShaderRHI,uint32 TextureIndex,FShaderResourceViewRHIParamRef SRVRHI)
 {
-	check(PendingGfxState->CurrentBSS && PendingGfxState->CurrentBSS->GetShader(SF_Domain) == ResourceCast(DomainShaderRHI));
+	ensureMsgf(0, TEXT("Tessellation not supported yet!"));
+/*
+	check(PendingGfxState->CurrentBSS && PendingGfxState->CurrentBSS->GetShader(DescriptorSet::Domain) == ResourceCast(DomainShaderRHI));
 	FVulkanShaderResourceView* SRV = ResourceCast(SRVRHI);
-	PendingGfxState->SetSRV(SF_Domain, TextureIndex, SRV);
+	PendingGfxState->SetSRV(DescriptorSet::Domain, TextureIndex, SRV);
+*/
 }
 
 void FVulkanCommandListContext::RHISetShaderResourceViewParameter(FGeometryShaderRHIParamRef GeometryShaderRHI,uint32 TextureIndex,FShaderResourceViewRHIParamRef SRVRHI)
 {
-	check(PendingGfxState->CurrentBSS && PendingGfxState->CurrentBSS->GetShader(SF_Geometry) == ResourceCast(GeometryShaderRHI));
+	check(PendingGfxState->CurrentBSS && PendingGfxState->CurrentBSS->GetShader(DescriptorSet::Geometry) == ResourceCast(GeometryShaderRHI));
 	FVulkanShaderResourceView* SRV = ResourceCast(SRVRHI);
-	PendingGfxState->SetSRV(SF_Geometry, TextureIndex, SRV);
+	PendingGfxState->SetSRV(DescriptorSet::Geometry, TextureIndex, SRV);
 }
 
 void FVulkanCommandListContext::RHISetShaderResourceViewParameter(FPixelShaderRHIParamRef PixelShaderRHI,uint32 TextureIndex,FShaderResourceViewRHIParamRef SRVRHI)
 {
-	check(PendingGfxState->CurrentBSS && PendingGfxState->CurrentBSS->GetShader(SF_Pixel) == ResourceCast(PixelShaderRHI));
+	check(PendingGfxState->CurrentBSS && PendingGfxState->CurrentBSS->GetShader(DescriptorSet::Pixel) == ResourceCast(PixelShaderRHI));
 	FVulkanShaderResourceView* SRV = ResourceCast(SRVRHI);
-	PendingGfxState->SetSRV(SF_Pixel, TextureIndex, SRV);
+	PendingGfxState->SetSRV(DescriptorSet::Pixel, TextureIndex, SRV);
 }
 
 void FVulkanCommandListContext::RHISetShaderResourceViewParameter(FComputeShaderRHIParamRef ComputeShaderRHI,uint32 TextureIndex,FShaderResourceViewRHIParamRef SRVRHI)
@@ -293,37 +305,43 @@ void FVulkanCommandListContext::RHISetShaderResourceViewParameter(FComputeShader
 
 void FVulkanCommandListContext::RHISetShaderSampler(FVertexShaderRHIParamRef VertexShaderRHI, uint32 SamplerIndex, FSamplerStateRHIParamRef NewStateRHI)
 {
-	check(PendingGfxState->CurrentBSS && PendingGfxState->CurrentBSS->GetShader(SF_Vertex) == ResourceCast(VertexShaderRHI));
+	check(PendingGfxState->CurrentBSS && PendingGfxState->CurrentBSS->GetShader(DescriptorSet::Vertex) == ResourceCast(VertexShaderRHI));
 	FVulkanSamplerState* Sampler = ResourceCast(NewStateRHI);
-	PendingGfxState->SetSamplerState(SF_Vertex, SamplerIndex, Sampler);
+	PendingGfxState->SetSamplerState(DescriptorSet::Vertex, SamplerIndex, Sampler);
 }
 
 void FVulkanCommandListContext::RHISetShaderSampler(FHullShaderRHIParamRef HullShaderRHI, uint32 SamplerIndex, FSamplerStateRHIParamRef NewStateRHI)
 {
-	check(PendingGfxState->CurrentBSS && PendingGfxState->CurrentBSS->GetShader(SF_Hull) == ResourceCast(HullShaderRHI));
+	ensureMsgf(0, TEXT("Tessellation not supported yet!"));
+/*
+	check(PendingGfxState->CurrentBSS && PendingGfxState->CurrentBSS->GetShader(DescriptorSet::Hull) == ResourceCast(HullShaderRHI));
 	FVulkanSamplerState* Sampler = ResourceCast(NewStateRHI);
-	PendingGfxState->SetSamplerState(SF_Hull, SamplerIndex, Sampler);
+	PendingGfxState->SetSamplerState(DescriptorSet::Hull, SamplerIndex, Sampler);
+*/
 }
 
 void FVulkanCommandListContext::RHISetShaderSampler(FDomainShaderRHIParamRef DomainShaderRHI, uint32 SamplerIndex, FSamplerStateRHIParamRef NewStateRHI)
 {
-	check(PendingGfxState->CurrentBSS && PendingGfxState->CurrentBSS->GetShader(SF_Domain) == ResourceCast(DomainShaderRHI));
+	ensureMsgf(0, TEXT("Tessellation not supported yet!"));
+/*
+	check(PendingGfxState->CurrentBSS && PendingGfxState->CurrentBSS->GetShader(DescriptorSet::Domain) == ResourceCast(DomainShaderRHI));
 	FVulkanSamplerState* Sampler = ResourceCast(NewStateRHI);
-	PendingGfxState->SetSamplerState(SF_Domain, SamplerIndex, Sampler);
+	PendingGfxState->SetSamplerState(DescriptorSet::Domain, SamplerIndex, Sampler);
+*/
 }
 
 void FVulkanCommandListContext::RHISetShaderSampler(FGeometryShaderRHIParamRef GeometryShaderRHI, uint32 SamplerIndex, FSamplerStateRHIParamRef NewStateRHI)
 {
-	check(PendingGfxState->CurrentBSS && PendingGfxState->CurrentBSS->GetShader(SF_Geometry) == ResourceCast(GeometryShaderRHI));
+	check(PendingGfxState->CurrentBSS && PendingGfxState->CurrentBSS->GetShader(DescriptorSet::Geometry) == ResourceCast(GeometryShaderRHI));
 	FVulkanSamplerState* Sampler = ResourceCast(NewStateRHI);
-	PendingGfxState->SetSamplerState(SF_Geometry, SamplerIndex, Sampler);
+	PendingGfxState->SetSamplerState(DescriptorSet::Geometry, SamplerIndex, Sampler);
 }
 
 void FVulkanCommandListContext::RHISetShaderSampler(FPixelShaderRHIParamRef PixelShaderRHI, uint32 SamplerIndex, FSamplerStateRHIParamRef NewStateRHI)
 {
-	check(PendingGfxState->CurrentBSS && PendingGfxState->CurrentBSS->GetShader(SF_Pixel) == ResourceCast(PixelShaderRHI));
+	check(PendingGfxState->CurrentBSS && PendingGfxState->CurrentBSS->GetShader(DescriptorSet::Pixel) == ResourceCast(PixelShaderRHI));
 	FVulkanSamplerState* Sampler = ResourceCast(NewStateRHI);
-	PendingGfxState->SetSamplerState(SF_Pixel, SamplerIndex, Sampler);
+	PendingGfxState->SetSamplerState(DescriptorSet::Pixel, SamplerIndex, Sampler);
 }
 
 void FVulkanCommandListContext::RHISetShaderSampler(FComputeShaderRHIParamRef ComputeShaderRHI, uint32 SamplerIndex, FSamplerStateRHIParamRef NewStateRHI)
@@ -335,37 +353,43 @@ void FVulkanCommandListContext::RHISetShaderSampler(FComputeShaderRHIParamRef Co
 
 void FVulkanCommandListContext::RHISetShaderParameter(FVertexShaderRHIParamRef VertexShaderRHI, uint32 BufferIndex, uint32 BaseIndex, uint32 NumBytes, const void* NewValue)
 {
-	check(PendingGfxState->CurrentBSS && PendingGfxState->CurrentBSS->GetShader(SF_Vertex) == ResourceCast(VertexShaderRHI));
+	check(PendingGfxState->CurrentBSS && PendingGfxState->CurrentBSS->GetShader(DescriptorSet::Vertex) == ResourceCast(VertexShaderRHI));
 
-	PendingGfxState->SetShaderParameter(SF_Vertex, BufferIndex, BaseIndex, NumBytes, NewValue);
+	PendingGfxState->SetShaderParameter(DescriptorSet::Vertex, BufferIndex, BaseIndex, NumBytes, NewValue);
 }
 
 void FVulkanCommandListContext::RHISetShaderParameter(FHullShaderRHIParamRef HullShaderRHI, uint32 BufferIndex, uint32 BaseIndex, uint32 NumBytes, const void* NewValue)
 {
-	check(PendingGfxState->CurrentBSS && PendingGfxState->CurrentBSS->GetShader(SF_Hull) == ResourceCast(HullShaderRHI));
+	ensureMsgf(0, TEXT("Tessellation not supported yet!"));
+/*
+	check(PendingGfxState->CurrentBSS && PendingGfxState->CurrentBSS->GetShader(DescriptorSet::Hull) == ResourceCast(HullShaderRHI));
 
-	PendingGfxState->SetShaderParameter(SF_Hull, BufferIndex, BaseIndex, NumBytes, NewValue);
+	PendingGfxState->SetShaderParameter(DescriptorSet::Hull, BufferIndex, BaseIndex, NumBytes, NewValue);
+*/
 }
 
 void FVulkanCommandListContext::RHISetShaderParameter(FDomainShaderRHIParamRef DomainShaderRHI, uint32 BufferIndex, uint32 BaseIndex, uint32 NumBytes, const void* NewValue)
 {
-	check(PendingGfxState->CurrentBSS && PendingGfxState->CurrentBSS->GetShader(SF_Domain) == ResourceCast(DomainShaderRHI));
+	ensureMsgf(0, TEXT("Tessellation not supported yet!"));
+/*
+	check(PendingGfxState->CurrentBSS && PendingGfxState->CurrentBSS->GetShader(DescriptorSet::Domain) == ResourceCast(DomainShaderRHI));
 
-	PendingGfxState->SetShaderParameter(SF_Domain, BufferIndex, BaseIndex, NumBytes, NewValue);
+	PendingGfxState->SetShaderParameter(DescriptorSet::Domain, BufferIndex, BaseIndex, NumBytes, NewValue);
+*/
 }
 
 void FVulkanCommandListContext::RHISetShaderParameter(FGeometryShaderRHIParamRef GeometryShaderRHI, uint32 BufferIndex, uint32 BaseIndex, uint32 NumBytes, const void* NewValue)
 {
-	check(PendingGfxState->CurrentBSS && PendingGfxState->CurrentBSS->GetShader(SF_Geometry) == ResourceCast(GeometryShaderRHI));
+	check(PendingGfxState->CurrentBSS && PendingGfxState->CurrentBSS->GetShader(DescriptorSet::Geometry) == ResourceCast(GeometryShaderRHI));
 
-	PendingGfxState->SetShaderParameter(SF_Geometry, BufferIndex, BaseIndex, NumBytes, NewValue);
+	PendingGfxState->SetShaderParameter(DescriptorSet::Geometry, BufferIndex, BaseIndex, NumBytes, NewValue);
 }
 
 void FVulkanCommandListContext::RHISetShaderParameter(FPixelShaderRHIParamRef PixelShaderRHI, uint32 BufferIndex, uint32 BaseIndex, uint32 NumBytes, const void* NewValue)
 {
-	check(PendingGfxState->CurrentBSS && PendingGfxState->CurrentBSS->GetShader(SF_Pixel) == ResourceCast(PixelShaderRHI));
+	check(PendingGfxState->CurrentBSS && PendingGfxState->CurrentBSS->GetShader(DescriptorSet::Pixel) == ResourceCast(PixelShaderRHI));
 
-	PendingGfxState->SetShaderParameter(SF_Pixel, BufferIndex, BaseIndex, NumBytes, NewValue);
+	PendingGfxState->SetShaderParameter(DescriptorSet::Pixel, BufferIndex, BaseIndex, NumBytes, NewValue);
 }
 
 void FVulkanCommandListContext::RHISetShaderParameter(FComputeShaderRHIParamRef ComputeShaderRHI,uint32 BufferIndex, uint32 BaseIndex, uint32 NumBytes, const void* NewValue)
@@ -446,7 +470,7 @@ static void GatherUniformBufferResources(
 	}
 }
 
-inline void FVulkanCommandListContext::SetShaderUniformBuffer(EShaderFrequency Stage, const FVulkanUniformBuffer* UniformBuffer, int32 BindingIndex, const FVulkanShader* Shader)
+inline void FVulkanCommandListContext::SetShaderUniformBuffer(DescriptorSet::EStage Stage, const FVulkanUniformBuffer* UniformBuffer, int32 BindingIndex, const FVulkanShader* Shader)
 {
 #if VULKAN_ENABLE_AGGRESSIVE_STATS
 	SCOPE_CYCLE_COUNTER(STAT_VulkanSetUniformBufferTime);
@@ -552,31 +576,37 @@ inline void FVulkanCommandListContext::SetShaderUniformBuffer(EShaderFrequency S
 void FVulkanCommandListContext::RHISetShaderUniformBuffer(FVertexShaderRHIParamRef VertexShaderRHI, uint32 BufferIndex, FUniformBufferRHIParamRef BufferRHI)
 {
 	FVulkanUniformBuffer* UniformBuffer = ResourceCast(BufferRHI);
-	SetShaderUniformBuffer(SF_Vertex, UniformBuffer, BufferIndex, ResourceCast(VertexShaderRHI));
+	SetShaderUniformBuffer(DescriptorSet::Vertex, UniformBuffer, BufferIndex, ResourceCast(VertexShaderRHI));
 }
 
 void FVulkanCommandListContext::RHISetShaderUniformBuffer(FHullShaderRHIParamRef HullShaderRHI, uint32 BufferIndex, FUniformBufferRHIParamRef BufferRHI)
 {
+	ensureMsgf(0, TEXT("Tessellation not supported yet!"));
+/*
 	FVulkanUniformBuffer* UniformBuffer = ResourceCast(BufferRHI);
-	SetShaderUniformBuffer(SF_Hull, UniformBuffer, BufferIndex, ResourceCast(HullShaderRHI));
+	SetShaderUniformBuffer(DescriptorSet::Hull, UniformBuffer, BufferIndex, ResourceCast(HullShaderRHI));
+*/
 }
 
 void FVulkanCommandListContext::RHISetShaderUniformBuffer(FDomainShaderRHIParamRef DomainShaderRHI, uint32 BufferIndex, FUniformBufferRHIParamRef BufferRHI)
 {
+	ensureMsgf(0, TEXT("Tessellation not supported yet!"));
+/*
 	FVulkanUniformBuffer* UniformBuffer = ResourceCast(BufferRHI);
-	SetShaderUniformBuffer(SF_Domain, UniformBuffer, BufferIndex, ResourceCast(DomainShaderRHI));
+	SetShaderUniformBuffer(DescriptorSet::Domain, UniformBuffer, BufferIndex, ResourceCast(DomainShaderRHI));
+*/
 }
 
 void FVulkanCommandListContext::RHISetShaderUniformBuffer(FGeometryShaderRHIParamRef GeometryShaderRHI, uint32 BufferIndex, FUniformBufferRHIParamRef BufferRHI)
 {
 	FVulkanUniformBuffer* UniformBuffer = ResourceCast(BufferRHI);
-	SetShaderUniformBuffer(SF_Geometry, UniformBuffer, BufferIndex, ResourceCast(GeometryShaderRHI));
+	SetShaderUniformBuffer(DescriptorSet::Geometry, UniformBuffer, BufferIndex, ResourceCast(GeometryShaderRHI));
 }
 
 void FVulkanCommandListContext::RHISetShaderUniformBuffer(FPixelShaderRHIParamRef PixelShaderRHI, uint32 BufferIndex, FUniformBufferRHIParamRef BufferRHI)
 {
 	FVulkanUniformBuffer* UniformBuffer = ResourceCast(BufferRHI);
-	SetShaderUniformBuffer(SF_Pixel, UniformBuffer, BufferIndex, ResourceCast(PixelShaderRHI));
+	SetShaderUniformBuffer(DescriptorSet::Pixel, UniformBuffer, BufferIndex, ResourceCast(PixelShaderRHI));
 }
 
 void FVulkanCommandListContext::RHISetShaderUniformBuffer(FComputeShaderRHIParamRef ComputeShaderRHI, uint32 BufferIndex, FUniformBufferRHIParamRef BufferRHI)
diff --git a/Engine/Source/Runtime/VulkanRHI/Private/VulkanDebug.cpp b/Engine/Source/Runtime/VulkanRHI/Private/VulkanDebug.cpp
--- a/Engine/Source/Runtime/VulkanRHI/Private/VulkanDebug.cpp
+++ b/Engine/Source/Runtime/VulkanRHI/Private/VulkanDebug.cpp
@@ -138,65 +138,197 @@ static VkBool32 VKAPI_PTR DebugReportFunction(
 	return VK_FALSE;
 }
 
-void FVulkanDynamicRHI::SetupDebugLayerCallback()
+#if VULKAN_SUPPORTS_DEBUG_UTILS
+static VkBool32 DebugUtilsCallback(VkDebugUtilsMessageSeverityFlagBitsEXT MsgSeverity, VkDebugUtilsMessageTypeFlagsEXT MsgType,
+	const VkDebugUtilsMessengerCallbackDataEXT* CallbackData, void* UserData)
 {
-	if (!bSupportsDebugCallbackExt)
+	if (CallbackData->pMessageIdName)
 	{
-		UE_LOG(LogVulkanRHI, Warning, TEXT("Instance does not support 'VK_EXT_debug_report' extension; debug reporting skipped!"));
-		return;
 	}
-	
-	PFN_vkCreateDebugReportCallbackEXT CreateMsgCallback = (PFN_vkCreateDebugReportCallbackEXT)(void*)VulkanRHI::vkGetInstanceProcAddr(Instance, CREATE_MSG_CALLBACK);
-	if (CreateMsgCallback)
+	else
 	{
-		VkDebugReportCallbackCreateInfoEXT CreateInfo;
-		FMemory::Memzero(CreateInfo);
-		CreateInfo.sType = VK_STRUCTURE_TYPE_DEBUG_REPORT_CREATE_INFO_EXT;
-		CreateInfo.pfnCallback = DebugReportFunction;
+		if (MsgType == 2 && CallbackData->messageIdNumber == 5)
+		{
+			// SPIR-V module not valid: MemoryBarrier: Vulkan specification requires Memory Semantics to have one of the following bits set: Acquire, Release, AcquireRelease or SequentiallyConsistent
+			return VK_FALSE;
+		}
+		else if (MsgType == 2 && CallbackData->messageIdNumber == 2)
+		{
+			// fragment shader writes to output location 0 with no matching attachment
+			return VK_FALSE;
+		}
+		else if (MsgType == 2 && CallbackData->messageIdNumber == 3)
+		{
+			// Attachment 2 not written by fragment shader
+			return VK_FALSE;
+		}
+		else if (MsgType == 6 && CallbackData->messageIdNumber == 2)
+		{
+			// Vertex shader writes to output location 0.0 which is not consumed by fragment shader
+			return VK_FALSE;
+		}
+	}
 
-		int32 CVar = GValidationCvar.GetValueOnRenderThread();
-		switch (CVar)
+	const TCHAR* Severity = TEXT("");
+	if (MsgSeverity & VK_DEBUG_UTILS_MESSAGE_SEVERITY_ERROR_BIT_EXT)
+	{
+		ensure((MsgSeverity & ~VK_DEBUG_UTILS_MESSAGE_SEVERITY_ERROR_BIT_EXT) == 0);
+		Severity = TEXT("Error");
+	}
+	else if (MsgSeverity & VK_DEBUG_UTILS_MESSAGE_SEVERITY_WARNING_BIT_EXT)
+	{
+		ensure((MsgSeverity & ~VK_DEBUG_UTILS_MESSAGE_SEVERITY_WARNING_BIT_EXT) == 0);
+		Severity = TEXT("Warning");
+	}
+	else if (MsgSeverity & VK_DEBUG_UTILS_MESSAGE_SEVERITY_INFO_BIT_EXT)
+	{
+		ensure((MsgSeverity & ~VK_DEBUG_UTILS_MESSAGE_SEVERITY_INFO_BIT_EXT) == 0);
+		Severity = TEXT("Info");
+	}
+	else if (MsgSeverity & VK_DEBUG_UTILS_MESSAGE_SEVERITY_VERBOSE_BIT_EXT)
+	{
+		ensure((MsgSeverity & ~VK_DEBUG_UTILS_MESSAGE_SEVERITY_VERBOSE_BIT_EXT) == 0);
+		Severity = TEXT("Verbose");
+	}
+
+	const TCHAR* Type = TEXT("");
+	if (MsgType & VK_DEBUG_UTILS_MESSAGE_TYPE_GENERAL_BIT_EXT)
+	{
+		ensure((MsgType & ~VK_DEBUG_UTILS_MESSAGE_TYPE_GENERAL_BIT_EXT) == 0);
+		Type = TEXT(" General");
+	}
+	else if (MsgType & VK_DEBUG_UTILS_MESSAGE_TYPE_VALIDATION_BIT_EXT)
+	{
+		if (MsgType & VK_DEBUG_UTILS_MESSAGE_TYPE_PERFORMANCE_BIT_EXT)
 		{
-		default:
-			CreateInfo.flags |= VK_DEBUG_REPORT_DEBUG_BIT_EXT;
-			// Fall-through...
-		case 4:
-			CreateInfo.flags |= VK_DEBUG_REPORT_INFORMATION_BIT_EXT;
-			// Fall-through...
-		case 3:
-			CreateInfo.flags |= VK_DEBUG_REPORT_PERFORMANCE_WARNING_BIT_EXT;
-			// Fall-through...
-		case 2:
-			CreateInfo.flags |= VK_DEBUG_REPORT_WARNING_BIT_EXT;
-			// Fall-through...
-		case 1:
-			CreateInfo.flags |= VK_DEBUG_REPORT_ERROR_BIT_EXT;
-			break;
-		case 0:
-			// Nothing to do!
-			break;
+			ensure((MsgType & ~(VK_DEBUG_UTILS_MESSAGE_TYPE_VALIDATION_BIT_EXT | VK_DEBUG_UTILS_MESSAGE_TYPE_PERFORMANCE_BIT_EXT)) == 0);
+			Type = TEXT("Perf/Validation");
 		}
-		VkResult Result = CreateMsgCallback(Instance, &CreateInfo, nullptr, &MsgCallback);
-		switch (Result)
+		else
 		{
-		case VK_SUCCESS:
-			break;
-		case VK_ERROR_OUT_OF_HOST_MEMORY:
-			UE_LOG(LogVulkanRHI, Warning, TEXT("CreateMsgCallback: out of host memory/CreateMsgCallback Failure; debug reporting skipped"));
-			break;
-		default:
-			UE_LOG(LogVulkanRHI, Warning, TEXT("CreateMsgCallback: unknown failure %d/CreateMsgCallback Failure; debug reporting skipped"), (int32)Result);
-			break;
+			ensure((MsgType & ~VK_DEBUG_UTILS_MESSAGE_TYPE_VALIDATION_BIT_EXT) == 0);
+			Type = TEXT("Validation");
+		}
+	}
+	else if (MsgType & VK_DEBUG_UTILS_MESSAGE_TYPE_PERFORMANCE_BIT_EXT)
+	{
+		ensure((MsgType & ~VK_DEBUG_UTILS_MESSAGE_TYPE_PERFORMANCE_BIT_EXT) == 0);
+		Type = TEXT("Perf");
+	}
+
+	static TMap<const char*, TSet<int32>> SeenCodes;
+	if (GCVarUniqueValidationMessages->GetInt() == 0 || !SeenCodes.FindOrAdd(CallbackData->pMessageIdName).Contains(CallbackData->messageIdNumber))
+	{
+		if (CallbackData->pMessageIdName)
+		{
+			FPlatformMisc::LowLevelOutputDebugStringf(TEXT("*** %s %s:%d(%s) %s\n"), Type, Severity, CallbackData->messageIdNumber, ANSI_TO_TCHAR(CallbackData->pMessageIdName), ANSI_TO_TCHAR(CallbackData->pMessage));
+		}
+		else
+		{
+			FPlatformMisc::LowLevelOutputDebugStringf(TEXT("*** %s %s:%d %s\n"), Type, Severity, CallbackData->messageIdNumber, ANSI_TO_TCHAR(CallbackData->pMessage));
+		}
+		if (GCVarUniqueValidationMessages->GetInt() == 1)
+		{
+			SeenCodes.Find(CallbackData->pMessageIdName)->Add(CallbackData->messageIdNumber);
+		}
+	}
+
+	return VK_FALSE;
+}
+#endif
+
+void FVulkanDynamicRHI::SetupDebugLayerCallback()
+{
+#if VULKAN_SUPPORTS_DEBUG_UTILS
+	if (bSupportsDebugUtilsExt)
+	{
+		PFN_vkCreateDebugUtilsMessengerEXT CreateDebugUtilsMessengerEXT = (PFN_vkCreateDebugUtilsMessengerEXT)(void*)VulkanRHI::vkGetInstanceProcAddr(Instance, "vkCreateDebugUtilsMessengerEXT");
+		if (CreateDebugUtilsMessengerEXT)
+		{
+			VkDebugUtilsMessengerCreateInfoEXT CreateInfo;
+			ZeroVulkanStruct(CreateInfo, VK_STRUCTURE_TYPE_DEBUG_UTILS_MESSENGER_CREATE_INFO_EXT);
+
+			int32 CVar = GValidationCvar.GetValueOnRenderThread();
+			CreateInfo.messageSeverity = (CVar >= 1 ? VK_DEBUG_UTILS_MESSAGE_SEVERITY_ERROR_BIT_EXT : 0) |
+				(CVar >= 2 ? VK_DEBUG_UTILS_MESSAGE_SEVERITY_WARNING_BIT_EXT : 0) | (CVar >= 3 ? VK_DEBUG_UTILS_MESSAGE_SEVERITY_INFO_BIT_EXT : 0);
+			CreateInfo.messageType = (CVar >= 1 ? (VK_DEBUG_UTILS_MESSAGE_TYPE_GENERAL_BIT_EXT | VK_DEBUG_UTILS_MESSAGE_TYPE_VALIDATION_BIT_EXT) : 0) |
+				(CVar >= 3 ? VK_DEBUG_UTILS_MESSAGE_TYPE_PERFORMANCE_BIT_EXT : 0);
+			CreateInfo.pfnUserCallback = (PFN_vkDebugUtilsMessengerCallbackEXT)(void*)DebugUtilsCallback;
+			VkResult Result = (*CreateDebugUtilsMessengerEXT)(Instance, &CreateInfo, nullptr, &Messenger);
+			ensure(Result == VK_SUCCESS);
 		}
 	}
 	else
+#endif
+	if (bSupportsDebugCallbackExt)
 	{
-		UE_LOG(LogVulkanRHI, Warning, TEXT("GetProcAddr: Unable to find vkDbgCreateMsgCallback/vkGetInstanceProcAddr; debug reporting skipped!"));
+		PFN_vkCreateDebugReportCallbackEXT CreateMsgCallback = (PFN_vkCreateDebugReportCallbackEXT)(void*)VulkanRHI::vkGetInstanceProcAddr(Instance, CREATE_MSG_CALLBACK);
+		if (CreateMsgCallback)
+		{
+			VkDebugReportCallbackCreateInfoEXT CreateInfo;
+			ZeroVulkanStruct(CreateInfo, VK_STRUCTURE_TYPE_DEBUG_REPORT_CREATE_INFO_EXT);
+			CreateInfo.pfnCallback = DebugReportFunction;
+
+			int32 CVar = GValidationCvar.GetValueOnRenderThread();
+			switch (CVar)
+			{
+			default:
+				CreateInfo.flags |= VK_DEBUG_REPORT_DEBUG_BIT_EXT;
+				// Fall-through...
+			case 4:
+				CreateInfo.flags |= VK_DEBUG_REPORT_INFORMATION_BIT_EXT;
+				// Fall-through...
+			case 3:
+				CreateInfo.flags |= VK_DEBUG_REPORT_PERFORMANCE_WARNING_BIT_EXT;
+				// Fall-through...
+			case 2:
+				CreateInfo.flags |= VK_DEBUG_REPORT_WARNING_BIT_EXT;
+				// Fall-through...
+			case 1:
+				CreateInfo.flags |= VK_DEBUG_REPORT_ERROR_BIT_EXT;
+				break;
+			case 0:
+				// Nothing to do!
+				break;
+			}
+			VkResult Result = CreateMsgCallback(Instance, &CreateInfo, nullptr, &MsgCallback);
+			switch (Result)
+			{
+			case VK_SUCCESS:
+				break;
+			case VK_ERROR_OUT_OF_HOST_MEMORY:
+				UE_LOG(LogVulkanRHI, Warning, TEXT("CreateMsgCallback: out of host memory/CreateMsgCallback Failure; debug reporting skipped"));
+				break;
+			default:
+				UE_LOG(LogVulkanRHI, Warning, TEXT("CreateMsgCallback: unknown failure %d/CreateMsgCallback Failure; debug reporting skipped"), (int32)Result);
+				break;
+			}
+		}
+		else
+		{
+			UE_LOG(LogVulkanRHI, Warning, TEXT("GetProcAddr: Unable to find vkDbgCreateMsgCallback/vkGetInstanceProcAddr; debug reporting skipped!"));
+		}
+	}
+	else
+	{
+		UE_LOG(LogVulkanRHI, Warning, TEXT("Instance does not support 'VK_EXT_debug_report' extension; debug reporting skipped!"));
+		return;
 	}
 }
 
 void FVulkanDynamicRHI::RemoveDebugLayerCallback()
 {
+#if VULKAN_SUPPORTS_DEBUG_UTILS
+	if (Messenger != VK_NULL_HANDLE)
+	{
+		PFN_vkDestroyDebugUtilsMessengerEXT DestroyDebugUtilsMessengerEXT = (PFN_vkDestroyDebugUtilsMessengerEXT)(void*)VulkanRHI::vkGetInstanceProcAddr(Instance, "vkDestroyDebugUtilsMessengerEXT");
+		if (DestroyDebugUtilsMessengerEXT)
+		{
+			(*DestroyDebugUtilsMessengerEXT)(Instance, Messenger, nullptr);
+		}
+	}
+	else
+#endif
 	if (MsgCallback != VK_NULL_HANDLE)
 	{
 		PFN_vkDestroyDebugReportCallbackEXT DestroyMsgCallback = (PFN_vkDestroyDebugReportCallbackEXT)(void*)VulkanRHI::vkGetInstanceProcAddr(Instance, DESTROY_MSG_CALLBACK);
@@ -295,6 +427,7 @@ static TMap<VkImageView, TTrackingResource<VkImageViewCreateInfo>> GVulkanTracki
 
 #if VULKAN_ENABLE_BUFFER_TRACKING_LAYER
 static TMap<VkBuffer, TTrackingResource<VkBufferCreateInfo>> GVulkanTrackingBuffers;
+static TMap<VkBuffer, TArray<VkBufferView>> GVulkanTrackingBufferToBufferViews;
 static TMap<VkBufferView, TTrackingResource<VkBufferViewCreateInfo>> GVulkanTrackingBufferViews;
 #endif
 
@@ -1267,6 +1400,7 @@ void FWrapLayer::CreateBufferView(VkResult Result, VkDevice Device, const VkBuff
 #if VULKAN_ENABLE_TRACKING_CALLSTACK
 			CaptureCallStack(TrackingBuffer.CreateCallstack, 3);
 #endif
+			GVulkanTrackingBufferToBufferViews.FindOrAdd(CreateInfo->buffer).Add(*BufferView);
 		}
 #endif
 	}
@@ -3024,6 +3158,12 @@ void FWrapLayer::DestroyBuffer(VkResult Result, VkDevice Device, VkBuffer Buffer
 			FScopeLock ScopeLock(&GTrackingCS);
 			int32 NumRemoved = GVulkanTrackingBuffers.Remove(Buffer);
 			ensure(NumRemoved > 0);
+
+			auto* Found = GVulkanTrackingBufferToBufferViews.Find(Buffer);
+			if (Found)
+			{
+				ensure(Found->Num() > 0);
+			}
 		}
 #endif
 	}
diff --git a/Engine/Source/Runtime/VulkanRHI/Private/VulkanDescriptorSets.h b/Engine/Source/Runtime/VulkanRHI/Private/VulkanDescriptorSets.h
--- a/Engine/Source/Runtime/VulkanRHI/Private/VulkanDescriptorSets.h
+++ b/Engine/Source/Runtime/VulkanRHI/Private/VulkanDescriptorSets.h
@@ -36,7 +36,7 @@ public:
 		return SetLayouts;
 	}
 
-	void AddBindingsForStage(VkShaderStageFlagBits StageFlags, EDescriptorSetStage DescSet, const FVulkanCodeHeader& CodeHeader);
+	void AddBindingsForStage(VkShaderStageFlagBits StageFlags, DescriptorSet::EStage DescSet, const FVulkanCodeHeader& CodeHeader);
 
 	friend uint32 GetTypeHash(const FVulkanDescriptorSetsLayoutInfo& In)
 	{
@@ -360,7 +360,7 @@ class FOLDVulkanDescriptorSets
 public:
 	~FOLDVulkanDescriptorSets();
 
-	typedef TArray<VkDescriptorSet, TInlineAllocator<SF_Compute>> FDescriptorSetArray;
+	typedef TArray<VkDescriptorSet, TInlineAllocator<DescriptorSet::NumGfxStages>> FDescriptorSetArray;
 
 	inline const FDescriptorSetArray& GetHandles() const
 	{
@@ -431,7 +431,7 @@ protected:
 	FVulkanDescriptorSetsLayout DescriptorSetLayout;
 	VkPipelineLayout PipelineLayout;
 
-	inline void AddBindingsForStage(VkShaderStageFlagBits StageFlags, EDescriptorSetStage DescSet, const FVulkanCodeHeader& CodeHeader)
+	inline void AddBindingsForStage(VkShaderStageFlagBits StageFlags, DescriptorSet::EStage DescSet, const FVulkanCodeHeader& CodeHeader)
 	{
 		// Setting descriptor is only allowed prior to compiling the layout
 		check(DescriptorSetLayout.GetHandles().Num() == 0);
diff --git a/Engine/Source/Runtime/VulkanRHI/Private/VulkanDevice.cpp b/Engine/Source/Runtime/VulkanRHI/Private/VulkanDevice.cpp
--- a/Engine/Source/Runtime/VulkanRHI/Private/VulkanDevice.cpp
+++ b/Engine/Source/Runtime/VulkanRHI/Private/VulkanDevice.cpp
@@ -48,6 +48,66 @@ static void EnableDrawMarkers()
 	}
 }
 
+#if VULKAN_SUPPORTS_VALIDATION_CACHE
+static void LoadValidationCache(VkDevice Device, VkValidationCacheEXT& OutValidationCache)
+{
+	VkValidationCacheCreateInfoEXT ValidationCreateInfo;
+	ZeroVulkanStruct(ValidationCreateInfo, VK_STRUCTURE_TYPE_VALIDATION_CACHE_CREATE_INFO_EXT);
+	TArray<uint8> InData;
+
+	const FString& CacheFilename = VulkanRHI::GetValidationCacheFilename();
+	UE_LOG(LogVulkanRHI, Display, TEXT("Trying validation cache file %s"), *CacheFilename);
+	if (FFileHelper::LoadFileToArray(InData, *CacheFilename, FILEREAD_Silent) && InData.Num() > 0)
+	{
+		// The code below supports SDK 1.0.65 Vulkan spec, which contains the following table:
+		//
+		// Offset	 Size            Meaning
+		// ------    ------------    ------------------------------------------------------------------
+		//      0               4    length in bytes of the entire validation cache header written as a
+		//                           stream of bytes, with the least significant byte first
+		//      4               4    a VkValidationCacheHeaderVersionEXT value written as a stream of
+		//                           bytes, with the least significant byte first
+		//      8    VK_UUID_SIZE    a layer commit ID expressed as a UUID, which uniquely identifies
+		//                           the version of the validation layers used to generate these
+		//                           validation results
+		int32* DataPtr = (int32*)InData.GetData();
+		if (*DataPtr > 0)
+		{
+			++DataPtr;
+			int32 Version = *DataPtr++;
+			if (Version == VK_PIPELINE_CACHE_HEADER_VERSION_ONE)
+			{
+				DataPtr += VK_UUID_SIZE / sizeof(int32);
+			}
+			else
+			{
+				UE_LOG(LogVulkanRHI, Warning, TEXT("Bad validation cache file %s, version=%d, expected %d"), *CacheFilename, Version, VK_PIPELINE_CACHE_HEADER_VERSION_ONE);
+				InData.Reset(0);
+			}
+		}
+		else
+		{
+			UE_LOG(LogVulkanRHI, Warning, TEXT("Bad validation cache file %s, header size=%d"), *CacheFilename, *DataPtr);
+			InData.Reset(0);
+		}
+	}
+
+	ValidationCreateInfo.initialDataSize = InData.Num();
+	ValidationCreateInfo.pInitialData = InData.Num() > 0 ? InData.GetData() : nullptr;
+	//ValidationCreateInfo.flags = 0;
+	PFN_vkCreateValidationCacheEXT vkCreateValidationCache = (PFN_vkCreateValidationCacheEXT)(void*)VulkanRHI::vkGetDeviceProcAddr(Device, "vkCreateValidationCacheEXT");
+	if (vkCreateValidationCache)
+	{
+		VkResult Result = vkCreateValidationCache(Device, &ValidationCreateInfo, VulkanRHI::GetMemoryAllocator(nullptr), &OutValidationCache);
+		if (Result != VK_SUCCESS)
+		{
+			UE_LOG(LogVulkanRHI, Warning, TEXT("Failed to create Vulkan validation cache, VkResult=%d"), Result);
+		}
+	}
+}
+#endif
+
+
 FVulkanDevice::FVulkanDevice(VkPhysicalDevice InGpu)
 	: Gpu(InGpu)
 	, Device(VK_NULL_HANDLE)
@@ -88,8 +148,7 @@ void FVulkanDevice::CreateDevice()
 
 	// Setup extension and layer info
 	VkDeviceCreateInfo DeviceInfo;
-	FMemory::Memzero(DeviceInfo);
-	DeviceInfo.sType = VK_STRUCTURE_TYPE_DEVICE_CREATE_INFO;
+	ZeroVulkanStruct(DeviceInfo, VK_STRUCTURE_TYPE_DEVICE_CREATE_INFO);
 
 	bool bDebugMarkersFound = false;
 	TArray<const ANSICHAR*> DeviceExtensions;
@@ -237,6 +296,20 @@ void FVulkanDevice::CreateDevice()
 	TransferQueue = new FVulkanQueue(this, TransferQueueFamilyIndex);
 
 #if VULKAN_ENABLE_DRAW_MARKERS
+#if 0//VULKAN_SUPPORTS_DEBUG_UTILS
+	FVulkanDynamicRHI* RHI = (FVulkanDynamicRHI*)GDynamicRHI;
+	if (RHI->SupportsDebugUtilsExt() && GRenderDocFound)
+	{
+		DebugMarkers.CmdBeginDebugLabel = (PFN_vkCmdBeginDebugUtilsLabelEXT)(void*)VulkanRHI::vkGetInstanceProcAddr(RHI->GetInstance(), "vkCmdBeginDebugUtilsLabelEXT");
+		DebugMarkers.CmdEndDebugLabel = (PFN_vkCmdEndDebugUtilsLabelEXT)(void*)VulkanRHI::vkGetInstanceProcAddr(RHI->GetInstance(), "vkCmdEndDebugUtilsLabelEXT");
+		DebugMarkers.SetDebugName = (PFN_vkSetDebugUtilsObjectNameEXT)(void*)VulkanRHI::vkGetInstanceProcAddr(RHI->GetInstance(), "vkSetDebugUtilsObjectNameEXT");
+		if (DebugMarkers.CmdBeginDebugLabel && DebugMarkers.CmdEndDebugLabel && DebugMarkers.SetDebugName)
+		{
+			bDebugMarkersFound = true;
+		}
+	}
+	else
+#endif	// VULKAN_SUPPORTS_DEBUG_UTILS
 	if (bDebugMarkersFound || FVulkanPlatform::ForceEnableDebugMarkers())
 	{
 		DebugMarkers.CmdBegin = (PFN_vkCmdDebugMarkerBeginEXT)(void*)VulkanRHI::vkGetDeviceProcAddr(Device, "vkCmdDebugMarkerBeginEXT");
@@ -565,10 +638,9 @@ bool FVulkanDevice::QueryGPU(int32 DeviceIndex)
 	if (GetOptionalExtensions().HasKHRGetPhysicalDeviceProperties2)
 	{
 		VkPhysicalDeviceProperties2KHR GpuProps2;
-		GpuProps2.sType = VK_STRUCTURE_TYPE_PHYSICAL_DEVICE_PROPERTIES_2_KHR;
+		ZeroVulkanStruct(GpuProps2, VK_STRUCTURE_TYPE_PHYSICAL_DEVICE_PROPERTIES_2_KHR);
 		GpuProps2.pNext = &GpuIdProps;
-		GpuIdProps.sType = VK_STRUCTURE_TYPE_PHYSICAL_DEVICE_ID_PROPERTIES_KHR;
-		GpuIdProps.pNext = nullptr;
+		ZeroVulkanStruct(GpuIdProps, VK_STRUCTURE_TYPE_PHYSICAL_DEVICE_ID_PROPERTIES_KHR);
 		VulkanRHI::vkGetPhysicalDeviceProperties2KHR(Gpu, &GpuProps2);
 	}
 #endif
@@ -635,11 +707,10 @@ void FVulkanDevice::InitGPU(int32 DeviceIndex)
 	StagingManager.Init(this);
 
 #if VULKAN_SUPPORTS_AMD_BUFFER_MARKER
-	if (OptionalDeviceExtensions.HasAMDBufferMarker)
+	if (GGPUCrashDebuggingEnabled && OptionalDeviceExtensions.HasAMDBufferMarker)
 	{
 		VkBufferCreateInfo CreateInfo;
-		FMemory::Memzero(CreateInfo);
-		CreateInfo.sType = VK_STRUCTURE_TYPE_BUFFER_CREATE_INFO;
+		ZeroVulkanStruct(CreateInfo, VK_STRUCTURE_TYPE_BUFFER_CREATE_INFO);
 		CreateInfo.size = GMaxCrashBufferEntries * sizeof(uint32_t);
 		CreateInfo.usage = VK_BUFFER_USAGE_TRANSFER_DST_BIT;
 		VERIFYVULKANRESULT(VulkanRHI::vkCreateBuffer(Device, &CreateInfo, nullptr, &CrashMarker.Buffer));
@@ -703,6 +774,13 @@ void FVulkanDevice::InitGPU(int32 DeviceIndex)
 		}
 	}
 
+#if VULKAN_SUPPORTS_VALIDATION_CACHE
+	if (OptionalDeviceExtensions.HasEXTValidationCache)
+	{
+		LoadValidationCache(Device, ValidationCache);
+	}
+#endif
+
 	PipelineStateCache->InitAndLoad(CacheFilenames);
 
 	// Setup default resource
@@ -715,66 +793,6 @@ void FVulkanDevice::InitGPU(int32 DeviceIndex)
 		DefaultImageView = FVulkanTextureView::StaticCreate(*this, DefaultImage->Image, VK_IMAGE_VIEW_TYPE_2D, DefaultImage->GetFullAspectMask(), PF_B8G8R8A8, VK_FORMAT_B8G8R8A8_UNORM, 0, 1, 0, 1);
 	}
 
-#if VULKAN_SUPPORTS_VALIDATION_CACHE
-	if (OptionalDeviceExtensions.HasEXTValidationCache)
-	{
-		VkValidationCacheCreateInfoEXT ValidationCreateInfo;
-		FMemory::Memzero(ValidationCreateInfo);
-		ValidationCreateInfo.sType = VK_STRUCTURE_TYPE_VALIDATION_CACHE_CREATE_INFO_EXT;
-		TArray<uint8> InData;
-
-		const FString& CacheFilename = VulkanRHI::GetValidationCacheFilename();
-		UE_LOG(LogVulkanRHI, Display, TEXT("Trying pipeline cache file %s"), *CacheFilename);
-		if (FFileHelper::LoadFileToArray(InData, *CacheFilename, FILEREAD_Silent) && InData.Num() > 0)
-		{
-			// The code below supports SDK 1.0.65 Vulkan spec, which contains the following table:
-			//
-			// Offset	 Size            Meaning
-			// ------    ------------    ------------------------------------------------------------------
-			//      0               4    length in bytes of the entire validation cache header written as a
-			//                           stream of bytes, with the least significant byte first
-			//      4               4    a VkValidationCacheHeaderVersionEXT value written as a stream of
-			//                           bytes, with the least significant byte first
-			//      8    VK_UUID_SIZE    a layer commit ID expressed as a UUID, which uniquely identifies
-			//                           the version of the validation layers used to generate these
-			//                           validation results
-			int32* DataPtr = (int32*)InData.GetData();
-			if (*DataPtr > 0)
-			{
-				++DataPtr;
-				int32 Version = *DataPtr++;
-				if (Version != VK_PIPELINE_CACHE_HEADER_VERSION_ONE)
-				{
-					DataPtr += VK_UUID_SIZE / sizeof(int32);
-				}
-				else
-				{
-					UE_LOG(LogVulkanRHI, Warning, TEXT("Bad validation cache file %s, version=%d, expected %d"), *CacheFilename, Version, VK_PIPELINE_CACHE_HEADER_VERSION_ONE);
-					InData.Reset(0);
-				}
-			}
-			else
-			{
-				UE_LOG(LogVulkanRHI, Warning, TEXT("Bad validation cache file %s, header size=%d"), *CacheFilename, *DataPtr);
-				InData.Reset(0);
-			}
-		}
-
-		ValidationCreateInfo.initialDataSize = InData.Num();
-		ValidationCreateInfo.pInitialData = InData.Num() > 0 ? InData.GetData() : nullptr;
-		//ValidationCreateInfo.flags = 0;
-		PFN_vkCreateValidationCacheEXT vkCreateValidationCache = (PFN_vkCreateValidationCacheEXT)(void*)VulkanRHI::vkGetDeviceProcAddr(Device, "vkCreateValidationCacheEXT");
-		if (vkCreateValidationCache)
-		{
-			VkResult Result = vkCreateValidationCache(Device, &ValidationCreateInfo, VulkanRHI::GetMemoryAllocator(nullptr), &ValidationCache);
-			if (Result != VK_SUCCESS)
-			{
-				UE_LOG(LogVulkanRHI, Warning, TEXT("Failed to create Vulkan validation cache, VkResult=%d"), Result);
-			}
-		}
-	}
-#endif
-
 #if VULKAN_USE_NEW_QUERIES
 	const auto* NumBufferedQueriesVar = IConsoleManager::Get().FindTConsoleVariableDataInt(TEXT("r.NumBufferedOcclusionQueries"));
 	int32 NumOcclusionQueryPools = (NumBufferedQueriesVar ? NumBufferedQueriesVar->GetValueOnAnyThread() : 3);
@@ -859,7 +877,7 @@ void FVulkanDevice::Destroy()
 
 
 #if VULKAN_SUPPORTS_AMD_BUFFER_MARKER
-	if (OptionalDeviceExtensions.HasAMDBufferMarker)
+	if (GGPUCrashDebuggingEnabled && OptionalDeviceExtensions.HasAMDBufferMarker)
 	{
 		CrashMarker.Allocation->Unmap();
 		VulkanRHI::vkDestroyBuffer(Device, CrashMarker.Buffer, nullptr);
diff --git a/Engine/Source/Runtime/VulkanRHI/Private/VulkanDevice.h b/Engine/Source/Runtime/VulkanRHI/Private/VulkanDevice.h
--- a/Engine/Source/Runtime/VulkanRHI/Private/VulkanDevice.h
+++ b/Engine/Source/Runtime/VulkanRHI/Private/VulkanDevice.h
@@ -189,20 +189,38 @@ public:
 	void NotifyDeletedImage(VkImage Image);
 
 #if VULKAN_ENABLE_DRAW_MARKERS
-	PFN_vkCmdDebugMarkerBeginEXT GetCmdDbgMarkerBegin() const
+	inline PFN_vkCmdDebugMarkerBeginEXT GetCmdDbgMarkerBegin() const
 	{
 		return DebugMarkers.CmdBegin;
 	}
 
-	PFN_vkCmdDebugMarkerEndEXT GetCmdDbgMarkerEnd() const
+	inline PFN_vkCmdDebugMarkerEndEXT GetCmdDbgMarkerEnd() const
 	{
 		return DebugMarkers.CmdEnd;
 	}
 
-	PFN_vkDebugMarkerSetObjectNameEXT GetDebugMarkerSetObjectName() const
+	inline PFN_vkDebugMarkerSetObjectNameEXT GetDebugMarkerSetObjectName() const
 	{
 		return DebugMarkers.CmdSetObjectName;
 	}
+
+#if 0//VULKAN_SUPPORTS_DEBUG_UTILS
+	inline PFN_vkCmdBeginDebugUtilsLabelEXT GetCmdBeginDebugLabel() const
+	{
+		return DebugMarkers.CmdBeginDebugLabel;
+	}
+
+	inline PFN_vkCmdEndDebugUtilsLabelEXT GetCmdEndDebugLabel() const
+	{
+		return DebugMarkers.CmdEndDebugLabel;
+	}
+
+	inline PFN_vkSetDebugUtilsObjectNameEXT GetSetDebugName() const
+	{
+		return DebugMarkers.SetDebugName;
+	}
+#endif
+
 #endif
 
 	void PrepareForCPURead();
@@ -370,6 +388,12 @@ private:
 		PFN_vkCmdDebugMarkerBeginEXT		CmdBegin = nullptr;
 		PFN_vkCmdDebugMarkerEndEXT			CmdEnd = nullptr;
 		PFN_vkDebugMarkerSetObjectNameEXT	CmdSetObjectName = nullptr;
+
+#if 0//VULKAN_SUPPORTS_DEBUG_UTILS
+		PFN_vkCmdBeginDebugUtilsLabelEXT	CmdBeginDebugLabel = nullptr;
+		PFN_vkCmdEndDebugUtilsLabelEXT		CmdEndDebugLabel = nullptr;
+		PFN_vkSetDebugUtilsObjectNameEXT	SetDebugName = nullptr;
+#endif
 	} DebugMarkers;
 	friend class FVulkanCommandListContext;
 #endif
diff --git a/Engine/Source/Runtime/VulkanRHI/Private/VulkanGenericPlatform.h b/Engine/Source/Runtime/VulkanRHI/Private/VulkanGenericPlatform.h
--- a/Engine/Source/Runtime/VulkanRHI/Private/VulkanGenericPlatform.h
+++ b/Engine/Source/Runtime/VulkanRHI/Private/VulkanGenericPlatform.h
@@ -51,8 +51,13 @@ public:
 
 	static bool SupportsDeviceLocalHostVisibleWithNoPenalty() { return false; }
 
+	static bool HasUnifiedMemory() { return false; }
+
 	static bool RegisterGPUWork() { return true; }
 
+	// Does the platform only support Vertex & Pixel stages (not Geometry or Tessellation)?
+	static bool IsVSPSOnly() { return false; }
+
 	static void WriteBufferMarkerAMD(VkCommandBuffer CmdBuffer, VkBuffer DestBuffer, const TArrayView<uint32>& Entries, bool bAdding) {}
 
 	// allow the platform code to restrict the device features
diff --git a/Engine/Source/Runtime/VulkanRHI/Private/VulkanLayers.cpp b/Engine/Source/Runtime/VulkanRHI/Private/VulkanLayers.cpp
--- a/Engine/Source/Runtime/VulkanRHI/Private/VulkanLayers.cpp
+++ b/Engine/Source/Runtime/VulkanRHI/Private/VulkanLayers.cpp
@@ -49,9 +49,6 @@ static const ANSICHAR* GValidationLayersDevice[] =
 // Instance Extensions to enable for all platforms
 static const ANSICHAR* GInstanceExtensions[] =
 {
-#if VULKAN_HAS_DEBUGGING_ENABLED
-	VK_EXT_DEBUG_REPORT_EXTENSION_NAME,
-#endif
 #if VULKAN_ENABLE_DESKTOP_HMD_SUPPORT
 	VK_KHR_EXTERNAL_MEMORY_CAPABILITIES_EXTENSION_NAME,
 	VK_KHR_GET_PHYSICAL_DEVICE_PROPERTIES_2_EXTENSION_NAME,
@@ -191,8 +188,34 @@ static inline bool FindLayerInList(const TArray<FLayerExtension>& List, const ch
 	return FindLayerIndexInList(List, LayerName) != INDEX_NONE;
 }
 
-void FVulkanDynamicRHI::GetInstanceLayersAndExtensions(TArray<const ANSICHAR*>& OutInstanceExtensions, TArray<const ANSICHAR*>& OutInstanceLayers)
+static inline bool FindLayerExtensionInList(const TArray<FLayerExtension>& List, const char* ExtensionName, const char*& FoundLayer)
+{
+	for (int32 Index = 0; Index < List.Num(); ++Index)
+	{
+		for (int32 ExtIndex = 0; ExtIndex < List[Index].ExtensionProps.Num(); ++ExtIndex)
+		{
+			if (!FCStringAnsi::Strcmp(List[Index].ExtensionProps[ExtIndex].extensionName, ExtensionName))
+			{
+				FoundLayer = List[Index].LayerProps.layerName;
+				return true;
+			}
+		}
+	}
+
+	return false;
+}
+
+static inline bool FindLayerExtensionInList(const TArray<FLayerExtension>& List, const char* ExtensionName)
+{
+	const char* Dummy = nullptr;
+	return FindLayerExtensionInList(List, ExtensionName, Dummy);
+}
+
+
+void FVulkanDynamicRHI::GetInstanceLayersAndExtensions(TArray<const ANSICHAR*>& OutInstanceExtensions, TArray<const ANSICHAR*>& OutInstanceLayers, bool& bOutDebugUtils)
 {
+	bOutDebugUtils = false;
+
 	TArray<FLayerExtension> GlobalLayerExtensions;
 	// 0 is reserved for NULL/instance
 	GlobalLayerExtensions.AddDefaulted();
@@ -250,17 +273,20 @@ void FVulkanDynamicRHI::GetInstanceLayersAndExtensions(TArray<const ANSICHAR*>&
 
 	FVulkanPlatform::NotifyFoundInstanceLayersAndExtensions(FoundUniqueLayers, FoundUniqueExtensions);
 
+	bool bVkTrace = false;
 	if (FParse::Param(FCommandLine::Get(), TEXT("vktrace")))
 	{
 		const char* VkTraceName = "VK_LAYER_LUNARG_vktrace";
 		if (FindLayerInList(GlobalLayerExtensions, VkTraceName))
 		{
 			OutInstanceLayers.Add(VkTraceName);
+			bVkTrace = true;
 		}
 	}
 
 #if VULKAN_HAS_DEBUGGING_ENABLED
 #if VULKAN_ENABLE_API_DUMP
+	if (!bVkTrace)
 	{
 		const char* VkApiDumpName = "VK_LAYER_LUNARG_api_dump";
 		bool bApiDumpFound = FindLayerInList(GlobalLayerExtensions, VkApiDumpName);
@@ -281,7 +307,7 @@ void FVulkanDynamicRHI::GetInstanceLayersAndExtensions(TArray<const ANSICHAR*>&
 		GValidationCvar->Set(VulkanValidationOption, ECVF_SetByCommandline);
 	}
 
-	if (VulkanValidationOption > 0)
+	if (!bVkTrace && VulkanValidationOption > 0)
 	{
 		// Verify that all requested debugging device-layers are available
 		for (uint32 LayerIndex = 0; GValidationLayersInstance[LayerIndex] != nullptr; ++LayerIndex)
@@ -298,6 +324,18 @@ void FVulkanDynamicRHI::GetInstanceLayersAndExtensions(TArray<const ANSICHAR*>&
 			}
 		}
 	}
+
+#if VULKAN_SUPPORTS_DEBUG_UTILS
+	if (!bVkTrace && GValidationCvar.GetValueOnAnyThread() > 0)
+	{
+		const char* FoundDebugUtilsLayer = nullptr;
+		bOutDebugUtils = FindLayerExtensionInList(GlobalLayerExtensions, VK_EXT_DEBUG_UTILS_EXTENSION_NAME, FoundDebugUtilsLayer);
+		if (bOutDebugUtils && *FoundDebugUtilsLayer)
+		{
+			OutInstanceLayers.Add(FoundDebugUtilsLayer);
+		}
+	}
+#endif
 #endif	// VULKAN_HAS_DEBUGGING_ENABLED
 
 	// Check to see if the HMD requires any specific Vulkan extensions to operate
@@ -317,26 +355,34 @@ void FVulkanDynamicRHI::GetInstanceLayersAndExtensions(TArray<const ANSICHAR*>&
 	TArray<const ANSICHAR*> PlatformExtensions;
 	FVulkanPlatform::GetInstanceExtensions(PlatformExtensions);
 
-	for (int32 Index = 0; Index < GlobalLayerExtensions.Num(); ++Index)
+	for (const ANSICHAR* PlatformExtension : PlatformExtensions)
 	{
-		for (int32 i = 0; i < GlobalLayerExtensions[Index].ExtensionProps.Num(); i++)
+		if (FindLayerExtensionInList(GlobalLayerExtensions, PlatformExtension))
 		{
-			for (const ANSICHAR* PlatformExtension : PlatformExtensions)
-			{
-				if (!FCStringAnsi::Strcmp(GlobalLayerExtensions[Index].ExtensionProps[i].extensionName, PlatformExtension))
-				{
-					OutInstanceExtensions.Add(PlatformExtension);
-					break;
-				}
-			}
-			for (int32 j = 0; GInstanceExtensions[j] != nullptr; j++)
-			{
-				if (!FCStringAnsi::Strcmp(GlobalLayerExtensions[Index].ExtensionProps[i].extensionName, GInstanceExtensions[j]))
-				{
-					OutInstanceExtensions.Add(GInstanceExtensions[j]);
-					break;
-				}
-			}
+			OutInstanceExtensions.Add(PlatformExtension);
+		}
+	}
+
+	for (int32 j = 0; GInstanceExtensions[j] != nullptr; j++)
+	{
+		if (FindLayerExtensionInList(GlobalLayerExtensions, GInstanceExtensions[j]))
+		{
+			OutInstanceExtensions.Add(GInstanceExtensions[j]);
+		}
+	}
+
+#if VULKAN_SUPPORTS_DEBUG_UTILS
+	if (!bVkTrace && bOutDebugUtils && FindLayerExtensionInList(GlobalLayerExtensions, VK_EXT_DEBUG_UTILS_EXTENSION_NAME))
+	{
+		OutInstanceExtensions.Add(VK_EXT_DEBUG_UTILS_EXTENSION_NAME);
+	}
+	//else if (!bOutDebugUtils)
+#endif
+	if (!bVkTrace && GValidationCvar.GetValueOnAnyThread() == 0)
+	{
+		if (FindLayerExtensionInList(GlobalLayerExtensions, VK_EXT_DEBUG_REPORT_EXTENSION_NAME))
+		{
+			OutInstanceExtensions.Add(VK_EXT_DEBUG_REPORT_EXTENSION_NAME);
 		}
 	}
 
@@ -541,7 +587,8 @@ void FVulkanDevice::GetDeviceExtensionsAndLayers(TArray<const ANSICHAR*>& OutDev
 	}
 
 #if VULKAN_ENABLE_DRAW_MARKERS
-	if (!bOutDebugMarkers && ListContains(AvailableExtensions, VK_EXT_DEBUG_MARKER_EXTENSION_NAME) && (GRenderDocFound || FVulkanPlatform::ForceEnableDebugMarkers()))
+	if (!bOutDebugMarkers &&
+		((GValidationCvar.GetValueOnAnyThread() == 0 && ListContains(AvailableExtensions, VK_EXT_DEBUG_MARKER_EXTENSION_NAME)) || FVulkanPlatform::ForceEnableDebugMarkers()))
 	{
 		OutDeviceExtensions.Add(VK_EXT_DEBUG_MARKER_EXTENSION_NAME);
 		bOutDebugMarkers = true;
diff --git a/Engine/Source/Runtime/VulkanRHI/Private/VulkanLoader.h b/Engine/Source/Runtime/VulkanRHI/Private/VulkanLoader.h
--- a/Engine/Source/Runtime/VulkanRHI/Private/VulkanLoader.h
+++ b/Engine/Source/Runtime/VulkanRHI/Private/VulkanLoader.h
@@ -169,10 +169,10 @@
 
 // List all optional Vulkan entry points used by Unreal that need to be loaded manually
 #define ENUM_VK_ENTRYPOINTS_OPTIONAL_INSTANCE(EnumMacro) \
-	EnumMacro(PFN_vkCreateDisplayPlaneSurfaceKHR, vkCreateDisplayPlaneSurfaceKHR) \
-	EnumMacro(PFN_vkCreateSharedSwapchainsKHR, vkCreateSharedSwapchainsKHR) \
-	EnumMacro(PFN_vkCreateDebugReportCallbackEXT, vkCreateDebugReportCallbackEXT) \
-	EnumMacro(PFN_vkDestroyDebugReportCallbackEXT, vkDestroyDebugReportCallbackEXT)
+	EnumMacro(PFN_vkCreateSharedSwapchainsKHR, vkCreateSharedSwapchainsKHR)
+	//EnumMacro(PFN_vkCreateDebugReportCallbackEXT, vkCreateDebugReportCallbackEXT)
+	//EnumMacro(PFN_vkDestroyDebugReportCallbackEXT, vkDestroyDebugReportCallbackEXT)
+	//EnumMacro(PFN_vkCreateDisplayPlaneSurfaceKHR, vkCreateDisplayPlaneSurfaceKHR)
 
 // List of all Vulkan entry points
 #define ENUM_VK_ENTRYPOINTS_ALL(EnumMacro) \
diff --git a/Engine/Source/Runtime/VulkanRHI/Private/VulkanMemory.cpp b/Engine/Source/Runtime/VulkanRHI/Private/VulkanMemory.cpp
--- a/Engine/Source/Runtime/VulkanRHI/Private/VulkanMemory.cpp
+++ b/Engine/Source/Runtime/VulkanRHI/Private/VulkanMemory.cpp
@@ -88,8 +88,8 @@ namespace VulkanRHI
 			HeapInfos[Index].TotalSize = MemoryProperties.memoryHeaps[Index].size;
 		}
 
-		bHasUnifiedMemory = (MemoryProperties.memoryHeapCount == 1);
-		UE_LOG(LogVulkanRHI, Display, TEXT("%d Device Memory Types"), MemoryProperties.memoryTypeCount);
+		bHasUnifiedMemory = FVulkanPlatform::HasUnifiedMemory();
+		UE_LOG(LogVulkanRHI, Display, TEXT("%d Device Memory Types (%sunified)"), MemoryProperties.memoryTypeCount, bHasUnifiedMemory ? TEXT("") : TEXT("Not "));
 		for (uint32 Index = 0; Index < MemoryProperties.memoryTypeCount; ++Index)
 		{
 			auto GetFlagsString = [](VkMemoryPropertyFlags Flags)
@@ -171,8 +171,7 @@ namespace VulkanRHI
 		check(MemoryTypeIndex < MemoryProperties.memoryTypeCount);
 
 		VkMemoryAllocateInfo Info;
-		FMemory::Memzero(Info);
-		Info.sType = VK_STRUCTURE_TYPE_MEMORY_ALLOCATE_INFO;
+		ZeroVulkanStruct(Info, VK_STRUCTURE_TYPE_MEMORY_ALLOCATE_INFO);
 		Info.allocationSize = AllocationSize;
 		Info.memoryTypeIndex = MemoryTypeIndex;
 
@@ -340,8 +339,7 @@ namespace VulkanRHI
 			check(IsMapped());
 			check(InOffset + InSize <= Size);
 			VkMappedMemoryRange Range;
-			FMemory::Memzero(Range);
-			Range.sType = VK_STRUCTURE_TYPE_MAPPED_MEMORY_RANGE;
+			ZeroVulkanStruct(Range, VK_STRUCTURE_TYPE_MAPPED_MEMORY_RANGE);
 			Range.memory = Handle;
 			Range.offset = InOffset;
 			Range.size = InSize;
@@ -356,8 +354,7 @@ namespace VulkanRHI
 			check(IsMapped());
 			check(InOffset + InSize <= Size);
 			VkMappedMemoryRange Range;
-			FMemory::Memzero(Range);
-			Range.sType = VK_STRUCTURE_TYPE_MAPPED_MEMORY_RANGE;
+			ZeroVulkanStruct(Range, VK_STRUCTURE_TYPE_MAPPED_MEMORY_RANGE);
 			Range.memory = Handle;
 			Range.offset = InOffset;
 			Range.size = InSize;
@@ -623,7 +620,9 @@ namespace VulkanRHI
 
 		{
 			FScopeLock ScopeLock(&GOldResourceLock);
-			for (int32 Index = 0; Index < FreePages.Num(); ++Index)
+
+			// Leave a page not freed to avoid potential hitching
+			for (int32 Index = 1; Index < FreePages.Num(); ++Index)
 			{
 				FOldResourceHeapPage* Page = FreePages[Index];
 				if (bImmediately || Page->FrameFreed + NUM_FRAMES_TO_WAIT_BEFORE_RELEASING_TO_OS < GFrameNumberRenderThread)
@@ -753,8 +752,7 @@ namespace VulkanRHI
 
 		check(Image != VK_NULL_HANDLE);
 		VkMemoryDedicatedAllocateInfoKHR DedicatedAllocInfo;
-		FMemory::Memzero(DedicatedAllocInfo);
-		DedicatedAllocInfo.sType = VK_STRUCTURE_TYPE_MEMORY_DEDICATED_ALLOCATE_INFO_KHR;
+		ZeroVulkanStruct(DedicatedAllocInfo, VK_STRUCTURE_TYPE_MEMORY_DEDICATED_ALLOCATE_INFO_KHR);
 		DedicatedAllocInfo.image = Image;
 		FDeviceMemoryAllocation* DeviceMemoryAllocation = Owner->GetParent()->GetMemoryManager().Alloc(false, AllocationSize, MemoryTypeIndex, &DedicatedAllocInfo, File, Line);
 
@@ -993,8 +991,7 @@ namespace VulkanRHI
 
 		VkBuffer Buffer;
 		VkBufferCreateInfo BufferCreateInfo;
-		FMemory::Memzero(BufferCreateInfo);
-		BufferCreateInfo.sType = VK_STRUCTURE_TYPE_BUFFER_CREATE_INFO;
+		ZeroVulkanStruct(BufferCreateInfo, VK_STRUCTURE_TYPE_BUFFER_CREATE_INFO);
 		BufferCreateInfo.size = BufferSize;
 		BufferCreateInfo.usage = BufferUsageFlags;
 		VERIFYVULKANRESULT(VulkanRHI::vkCreateBuffer(Device->GetInstanceHandle(), &BufferCreateInfo, nullptr, &Buffer));
@@ -1032,17 +1029,14 @@ namespace VulkanRHI
 	FOldResourceAllocation* FResourceHeapManager::AllocateDedicatedImageMemory(VkImage Image, const VkMemoryRequirements& MemoryReqs, VkMemoryPropertyFlags MemoryPropertyFlags, const char* File, uint32 Line)
 	{
 		VkImageMemoryRequirementsInfo2KHR ImageMemoryReqs2;
-		FMemory::Memzero(ImageMemoryReqs2);
-		ImageMemoryReqs2.sType = VK_STRUCTURE_TYPE_IMAGE_MEMORY_REQUIREMENTS_INFO_2_KHR;
+		ZeroVulkanStruct(ImageMemoryReqs2, VK_STRUCTURE_TYPE_IMAGE_MEMORY_REQUIREMENTS_INFO_2_KHR);
 		ImageMemoryReqs2.image = Image;
 
 		VkMemoryDedicatedRequirementsKHR DedMemoryReqs;
-		FMemory::Memzero(DedMemoryReqs);
-		DedMemoryReqs.sType = VK_STRUCTURE_TYPE_MEMORY_DEDICATED_REQUIREMENTS_KHR;
+		ZeroVulkanStruct(DedMemoryReqs, VK_STRUCTURE_TYPE_MEMORY_DEDICATED_REQUIREMENTS_KHR);
 
 		VkMemoryRequirements2KHR MemoryReqs2;
-		FMemory::Memzero(MemoryReqs2);
-		MemoryReqs2.sType = VK_STRUCTURE_TYPE_MEMORY_REQUIREMENTS_2_KHR;
+		ZeroVulkanStruct(MemoryReqs2, VK_STRUCTURE_TYPE_MEMORY_REQUIREMENTS_2_KHR);
 		MemoryReqs2.pNext = &DedMemoryReqs;
 
 		VulkanRHI::vkGetImageMemoryRequirements2KHR(Device->GetInstanceHandle(), &ImageMemoryReqs2, &MemoryReqs2);
@@ -1261,8 +1255,7 @@ namespace VulkanRHI
 		FStagingBuffer* StagingBuffer = new FStagingBuffer();
 
 		VkBufferCreateInfo StagingBufferCreateInfo;
-		FMemory::Memzero(StagingBufferCreateInfo);
-		StagingBufferCreateInfo.sType = VK_STRUCTURE_TYPE_BUFFER_CREATE_INFO;
+		ZeroVulkanStruct(StagingBufferCreateInfo, VK_STRUCTURE_TYPE_BUFFER_CREATE_INFO);
 		StagingBufferCreateInfo.size = Size;
 		StagingBufferCreateInfo.usage = InUsageFlags;
 
@@ -1438,8 +1431,7 @@ namespace VulkanRHI
 		, Owner(InOwner)
 	{
 		VkFenceCreateInfo Info;
-		FMemory::Memzero(Info);
-		Info.sType = VK_STRUCTURE_TYPE_FENCE_CREATE_INFO;
+		ZeroVulkanStruct(Info, VK_STRUCTURE_TYPE_FENCE_CREATE_INFO);
 		Info.flags = bCreateSignaled ? VK_FENCE_CREATE_SIGNALED_BIT : 0;
 		VERIFYVULKANRESULT(VulkanRHI::vkCreateFence(InDevice->GetInstanceHandle(), &Info, nullptr, &Handle));
 	}
@@ -1587,8 +1579,7 @@ namespace VulkanRHI
 		: FDeviceChild(InDevice)
 	{
 		VkEventCreateInfo Info;
-		FMemory::Memzero(Info);
-		Info.sType = VK_STRUCTURE_TYPE_EVENT_CREATE_INFO;
+		ZeroVulkanStruct(Info, VK_STRUCTURE_TYPE_EVENT_CREATE_INFO);
 		VERIFYVULKANRESULT(VulkanRHI::vkCreateEvent(InDevice->GetInstanceHandle(), &Info, nullptr, &Handle));
 	}
 
@@ -1768,8 +1759,7 @@ namespace VulkanRHI
 		EImageLayoutBarrier Source, EImageLayoutBarrier Dest, const VkImageSubresourceRange& SubresourceRange)
 	{
 		VkImageMemoryBarrier ImageBarrier;
-		FMemory::Memzero(ImageBarrier);
-		ImageBarrier.sType = VK_STRUCTURE_TYPE_IMAGE_MEMORY_BARRIER;
+		ZeroVulkanStruct(ImageBarrier, VK_STRUCTURE_TYPE_IMAGE_MEMORY_BARRIER);
 		ImageBarrier.image = Image;
 		ImageBarrier.subresourceRange = SubresourceRange;
 		ImageBarrier.srcQueueFamilyIndex = VK_QUEUE_FAMILY_IGNORED;
@@ -1797,9 +1787,12 @@ namespace VulkanRHI
 		VulkanRHI::vkCmdPipelineBarrier(CmdBuffer, SourceStages, DestStages, 0, 0, nullptr, 0, nullptr, 1, &ImageBarrier);
 	}
 
-	void FPendingBarrier::InnerExecute(FVulkanCmdBuffer* CmdBuffer)
+	void FPendingBarrier::InnerExecute(FVulkanCmdBuffer* CmdBuffer, bool bEnsure)
 	{
-		ensure(CmdBuffer->IsOutsideRenderPass());
+		if (bEnsure)
+		{
+			ensure(CmdBuffer->IsOutsideRenderPass());
+		}
 		VulkanRHI::vkCmdPipelineBarrier(CmdBuffer->GetHandle(),
 			SourceStage, DestStage, 0,
 			0, nullptr,
@@ -1814,8 +1807,7 @@ namespace VulkanRHI
 	{
 		// Create semaphore
 		VkSemaphoreCreateInfo CreateInfo;
-		FMemory::Memzero(CreateInfo);
-		CreateInfo.sType = VK_STRUCTURE_TYPE_SEMAPHORE_CREATE_INFO;
+		ZeroVulkanStruct(CreateInfo, VK_STRUCTURE_TYPE_SEMAPHORE_CREATE_INFO);
 		//CreateInfo.pNext = nullptr;
 		//CreateInfo.flags = 0;
 		VERIFYVULKANRESULT(VulkanRHI::vkCreateSemaphore(Device.GetInstanceHandle(), &CreateInfo, nullptr, &SemaphoreHandle));
diff --git a/Engine/Source/Runtime/VulkanRHI/Private/VulkanPendingState.cpp b/Engine/Source/Runtime/VulkanRHI/Private/VulkanPendingState.cpp
--- a/Engine/Source/Runtime/VulkanRHI/Private/VulkanPendingState.cpp
+++ b/Engine/Source/Runtime/VulkanRHI/Private/VulkanPendingState.cpp
@@ -113,8 +113,7 @@ FVulkanDescriptorPool::FVulkanDescriptorPool(FVulkanDevice* InDevice)
 #endif
 
 	VkDescriptorPoolCreateInfo PoolInfo;
-	FMemory::Memzero(PoolInfo);
-	PoolInfo.sType = VK_STRUCTURE_TYPE_DESCRIPTOR_POOL_CREATE_INFO;
+	ZeroVulkanStruct(PoolInfo, VK_STRUCTURE_TYPE_DESCRIPTOR_POOL_CREATE_INFO);
 	PoolInfo.flags = VK_DESCRIPTOR_POOL_CREATE_FREE_DESCRIPTOR_SET_BIT;
 	PoolInfo.poolSizeCount = Types.Num();
 	PoolInfo.pPoolSizes = Types.GetData();
@@ -569,7 +568,7 @@ void FVulkanPendingGfxState::PrepareForDraw(FVulkanCmdBuffer* CmdBuffer)
 					{
 						UE_LOG(LogVulkanRHI, Warning, TEXT("Missing binding on location %d in '%s' vertex shader"),
 							CurrBinding.binding,
-							*CurrentBSS->GetShader(SF_Vertex)->GetDebugName());
+							*CurrentBSS->GetShader(DescriptorSet::Vertex)->GetDebugName());
 						ensure(0);
 					}
 				}
@@ -627,7 +626,7 @@ void FVulkanPendingGfxState::InternalUpdateDynamicStates(FVulkanCmdBuffer* Cmd)
 	Cmd->bNeedsDynamicStateSet = false;
 }
 
-void FVulkanPendingGfxState::SetSRV(EShaderFrequency Stage, uint32 BindIndex, FVulkanShaderResourceView* SRV)
+void FVulkanPendingGfxState::SetSRV(DescriptorSet::EStage Stage, uint32 BindIndex, FVulkanShaderResourceView* SRV)
 {
 	if (SRV)
 	{
@@ -657,7 +656,7 @@ void FVulkanPendingGfxState::SetSRV(EShaderFrequency Stage, uint32 BindIndex, FV
 	}
 }
 
-void FVulkanPendingGfxState::SetUAV(EShaderFrequency Stage, uint32 UAVIndex, FVulkanUnorderedAccessView* UAV)
+void FVulkanPendingGfxState::SetUAV(DescriptorSet::EStage Stage, uint32 UAVIndex, FVulkanUnorderedAccessView* UAV)
 {
 	if (UAV)
 	{
diff --git a/Engine/Source/Runtime/VulkanRHI/Private/VulkanPendingState.h b/Engine/Source/Runtime/VulkanRHI/Private/VulkanPendingState.h
--- a/Engine/Source/Runtime/VulkanRHI/Private/VulkanPendingState.h
+++ b/Engine/Source/Runtime/VulkanRHI/Private/VulkanPendingState.h
@@ -182,31 +182,31 @@ public:
 		bDirtyVertexStreams = true;
 	}
 
-	inline void SetTexture(EShaderFrequency Stage, uint32 BindPoint, const FVulkanTextureBase* TextureBase, VkImageLayout Layout)
+	inline void SetTexture(DescriptorSet::EStage Stage, uint32 BindPoint, const FVulkanTextureBase* TextureBase, VkImageLayout Layout)
 	{
 		CurrentState->SetTexture(Stage, BindPoint, TextureBase, Layout);
 	}
 
-	inline void SetUniformBufferConstantData(EShaderFrequency Stage, uint32 BindPoint, const TArray<uint8>& ConstantData)
+	inline void SetUniformBufferConstantData(DescriptorSet::EStage Stage, uint32 BindPoint, const TArray<uint8>& ConstantData)
 	{
 		CurrentState->SetUniformBufferConstantData(Stage, BindPoint, ConstantData);
 	}
 
-	inline void SetUniformBuffer(EShaderFrequency Stage, uint32 BindPoint, const FVulkanUniformBuffer* UniformBuffer)
+	inline void SetUniformBuffer(DescriptorSet::EStage Stage, uint32 BindPoint, const FVulkanUniformBuffer* UniformBuffer)
 	{
 		CurrentState->SetUniformBuffer(Stage, BindPoint, UniformBuffer);
 	}
 
-	void SetUAV(EShaderFrequency Stage, uint32 UAVIndex, FVulkanUnorderedAccessView* UAV);
+	void SetUAV(DescriptorSet::EStage Stage, uint32 UAVIndex, FVulkanUnorderedAccessView* UAV);
 
-	void SetSRV(EShaderFrequency Stage, uint32 BindIndex, FVulkanShaderResourceView* SRV);
+	void SetSRV(DescriptorSet::EStage Stage, uint32 BindIndex, FVulkanShaderResourceView* SRV);
 
-	inline void SetSamplerState(EShaderFrequency Stage, uint32 BindPoint, FVulkanSamplerState* Sampler)
+	inline void SetSamplerState(DescriptorSet::EStage Stage, uint32 BindPoint, FVulkanSamplerState* Sampler)
 	{
 		CurrentState->SetSamplerState(Stage, BindPoint, Sampler);
 	}
 
-	inline void SetShaderParameter(EShaderFrequency Stage, uint32 BufferIndex, uint32 ByteOffset, uint32 NumBytes, const void* NewValue)
+	inline void SetShaderParameter(DescriptorSet::EStage Stage, uint32 BufferIndex, uint32 ByteOffset, uint32 NumBytes, const void* NewValue)
 	{
 		CurrentState->SetShaderParameter(Stage, BufferIndex, ByteOffset, NumBytes, NewValue);
 	}
diff --git a/Engine/Source/Runtime/VulkanRHI/Private/VulkanPipeline.cpp b/Engine/Source/Runtime/VulkanRHI/Private/VulkanPipeline.cpp
--- a/Engine/Source/Runtime/VulkanRHI/Private/VulkanPipeline.cpp
+++ b/Engine/Source/Runtime/VulkanRHI/Private/VulkanPipeline.cpp
@@ -14,6 +14,7 @@
 #include "GlobalShader.h"
 
 static const double HitchTime = 1.0 / 1000.0;
+const int32 NumGfxStages = FVulkanPlatform::IsVSPSOnly() ? DescriptorSet::NumMobileStages : DescriptorSet::NumGfxStages;
 
 template <typename TRHIType, typename TVulkanType>
 static inline FSHAHash GetShaderHash(TRHIType* RHIShader)
@@ -30,15 +31,15 @@ static inline FSHAHash GetShaderHash(TRHIType* RHIShader)
 	return Dummy;
 }
 
-static inline FSHAHash GetShaderHashForStage(const FGraphicsPipelineStateInitializer& PSOInitializer, int32 Stage)
+static inline FSHAHash GetShaderHashForStage(const FGraphicsPipelineStateInitializer& PSOInitializer, DescriptorSet::EStage Stage)
 {
 	switch (Stage)
 	{
-	case SF_Vertex:		return GetShaderHash<FRHIVertexShader, FVulkanVertexShader>(PSOInitializer.BoundShaderState.VertexShaderRHI);
-	case SF_Pixel:		return GetShaderHash<FRHIPixelShader, FVulkanPixelShader>(PSOInitializer.BoundShaderState.PixelShaderRHI);
-	case SF_Geometry:	return GetShaderHash<FRHIGeometryShader, FVulkanGeometryShader>(PSOInitializer.BoundShaderState.GeometryShaderRHI);
-	case SF_Hull:		return GetShaderHash<FRHIHullShader, FVulkanHullShader>(PSOInitializer.BoundShaderState.HullShaderRHI);
-	case SF_Domain:		return GetShaderHash<FRHIDomainShader, FVulkanDomainShader>(PSOInitializer.BoundShaderState.DomainShaderRHI);
+	case DescriptorSet::Vertex:		return GetShaderHash<FRHIVertexShader, FVulkanVertexShader>(PSOInitializer.BoundShaderState.VertexShaderRHI);
+	case DescriptorSet::Pixel:		return GetShaderHash<FRHIPixelShader, FVulkanPixelShader>(PSOInitializer.BoundShaderState.PixelShaderRHI);
+	case DescriptorSet::Geometry:	return GetShaderHash<FRHIGeometryShader, FVulkanGeometryShader>(PSOInitializer.BoundShaderState.GeometryShaderRHI);
+	//case DescriptorSet::Hull:		return GetShaderHash<FRHIHullShader, FVulkanHullShader>(PSOInitializer.BoundShaderState.HullShaderRHI);
+	//case DescriptorSet::Domain:		return GetShaderHash<FRHIDomainShader, FVulkanDomainShader>(PSOInitializer.BoundShaderState.DomainShaderRHI);
 	default:			check(0); break;
 	}
 
@@ -158,8 +159,7 @@ void FVulkanPipelineStateCache::Load(const TArray<FString>& CacheFilenames)
 			if (FVulkanPipelineStateCacheFile::BinaryCacheMatches(Device, DeviceCache))
 			{
 				VkPipelineCacheCreateInfo PipelineCacheInfo;
-				FMemory::Memzero(PipelineCacheInfo);
-				PipelineCacheInfo.sType = VK_STRUCTURE_TYPE_PIPELINE_CACHE_CREATE_INFO;
+				ZeroVulkanStruct(PipelineCacheInfo, VK_STRUCTURE_TYPE_PIPELINE_CACHE_CREATE_INFO);
 				PipelineCacheInfo.initialDataSize = DeviceCache.Num();
 				PipelineCacheInfo.pInitialData = DeviceCache.GetData();
 
@@ -193,7 +193,7 @@ void FVulkanPipelineStateCache::Load(const TArray<FString>& CacheFilenames)
 
 			FVulkanPipelineStateCacheFile File;
 
-			FShaderUCodeCache::TDataMap FileShaderCacheData;
+			FShaderUCodeCache::THashToMicrocode FileShaderCacheData;
 			File.ShaderCache = &FileShaderCacheData;
 
 			bool Valid = File.Load(Ar, *CacheFilename);
@@ -207,8 +207,7 @@ void FVulkanPipelineStateCache::Load(const TArray<FString>& CacheFilenames)
 			if (PipelineCache == VK_NULL_HANDLE)
 			{
 				VkPipelineCacheCreateInfo PipelineCacheInfo;
-				FMemory::Memzero(PipelineCacheInfo);
-				PipelineCacheInfo.sType = VK_STRUCTURE_TYPE_PIPELINE_CACHE_CREATE_INFO;
+				ZeroVulkanStruct(PipelineCacheInfo, VK_STRUCTURE_TYPE_PIPELINE_CACHE_CREATE_INFO);
 				VERIFYVULKANRESULT(VulkanRHI::vkCreatePipelineCache(Device->GetInstanceHandle(), &PipelineCacheInfo, nullptr, &PipelineCache));
 			}
 
@@ -228,7 +227,7 @@ void FVulkanPipelineStateCache::Load(const TArray<FString>& CacheFilenames)
 				FGfxPipelineEntry* GfxEntry = File.GfxPipelineEntries[Index];
 
 				FShaderHashes ShaderHashes;
-				for (int32 i = 0; i < SF_Compute; ++i)
+				for (int32 i = 0; i < NumGfxStages; ++i)
 				{
 					ShaderHashes.Stages[i] = GfxEntry->ShaderHashes[i];
 
@@ -300,8 +299,7 @@ void FVulkanPipelineStateCache::Load(const TArray<FString>& CacheFilenames)
 		if (PipelineCache == VK_NULL_HANDLE)
 		{
 			VkPipelineCacheCreateInfo PipelineCacheInfo;
-			FMemory::Memzero(PipelineCacheInfo);
-			PipelineCacheInfo.sType = VK_STRUCTURE_TYPE_PIPELINE_CACHE_CREATE_INFO;
+			ZeroVulkanStruct(PipelineCacheInfo, VK_STRUCTURE_TYPE_PIPELINE_CACHE_CREATE_INFO);
 			VERIFYVULKANRESULT(VulkanRHI::vkCreatePipelineCache(Device->GetInstanceHandle(), &PipelineCacheInfo, nullptr, &PipelineCache));
 		}
 	}
@@ -334,8 +332,7 @@ void FVulkanPipelineStateCache::InitAndLoad(const TArray<FString>& CacheFilename
 	if (PipelineCache == VK_NULL_HANDLE)
 	{
 		VkPipelineCacheCreateInfo PipelineCacheInfo;
-		FMemory::Memzero(PipelineCacheInfo);
-		PipelineCacheInfo.sType = VK_STRUCTURE_TYPE_PIPELINE_CACHE_CREATE_INFO;
+		ZeroVulkanStruct(PipelineCacheInfo, VK_STRUCTURE_TYPE_PIPELINE_CACHE_CREATE_INFO);
 		VERIFYVULKANRESULT(VulkanRHI::vkCreatePipelineCache(Device->GetInstanceHandle(), &PipelineCacheInfo, nullptr, &PipelineCache));
 	}
 }
@@ -848,17 +845,15 @@ void FVulkanPipelineStateCache::CreateGfxPipelineFromEntry(const FGfxPipelineEnt
 {
 	// Pipeline
 	VkGraphicsPipelineCreateInfo PipelineInfo;
-	FMemory::Memzero(PipelineInfo);
-	PipelineInfo.sType = VK_STRUCTURE_TYPE_GRAPHICS_PIPELINE_CREATE_INFO;
+	ZeroVulkanStruct(PipelineInfo, VK_STRUCTURE_TYPE_GRAPHICS_PIPELINE_CREATE_INFO);
 	PipelineInfo.layout = GfxEntry->Layout->GetPipelineLayout();
 
 	// Color Blend
 	VkPipelineColorBlendStateCreateInfo CBInfo;
-	FMemory::Memzero(CBInfo);
+	ZeroVulkanStruct(CBInfo, VK_STRUCTURE_TYPE_PIPELINE_COLOR_BLEND_STATE_CREATE_INFO);
+	CBInfo.attachmentCount = GfxEntry->ColorAttachmentStates.Num();
 	VkPipelineColorBlendAttachmentState BlendStates[MaxSimultaneousRenderTargets];
 	FMemory::Memzero(BlendStates);
-	CBInfo.sType = VK_STRUCTURE_TYPE_PIPELINE_COLOR_BLEND_STATE_CREATE_INFO;
-	CBInfo.attachmentCount = GfxEntry->ColorAttachmentStates.Num();
 	for (int32 Index = 0; Index < GfxEntry->ColorAttachmentStates.Num(); ++Index)
 	{
 		GfxEntry->ColorAttachmentStates[Index].WriteInto(BlendStates[Index]);
@@ -871,33 +866,30 @@ void FVulkanPipelineStateCache::CreateGfxPipelineFromEntry(const FGfxPipelineEnt
 
 	// Viewport
 	VkPipelineViewportStateCreateInfo VPInfo;
-	FMemory::Memzero(VPInfo);
-	VPInfo.sType = VK_STRUCTURE_TYPE_PIPELINE_VIEWPORT_STATE_CREATE_INFO;
+	ZeroVulkanStruct(VPInfo, VK_STRUCTURE_TYPE_PIPELINE_VIEWPORT_STATE_CREATE_INFO);
 	VPInfo.viewportCount = 1;
 	VPInfo.scissorCount = 1;
 
 	// Multisample
 	VkPipelineMultisampleStateCreateInfo MSInfo;
-	FMemory::Memzero(MSInfo);
-	MSInfo.sType = VK_STRUCTURE_TYPE_PIPELINE_MULTISAMPLE_STATE_CREATE_INFO;
-	MSInfo.pSampleMask = NULL;
+	ZeroVulkanStruct(MSInfo, VK_STRUCTURE_TYPE_PIPELINE_MULTISAMPLE_STATE_CREATE_INFO);
 	MSInfo.rasterizationSamples = (VkSampleCountFlagBits)FMath::Max(1u, GfxEntry->RasterizationSamples);
 
 	// Two stages: vs and fs
-	VkPipelineShaderStageCreateInfo ShaderStages[SF_Compute];
+	VkPipelineShaderStageCreateInfo ShaderStages[DescriptorSet::NumGfxStages];
 	FMemory::Memzero(ShaderStages);
 	PipelineInfo.stageCount = 0;
 	PipelineInfo.pStages = ShaderStages;
-	for (uint32 ShaderStage = 0; ShaderStage < SF_Compute; ++ShaderStage)
+	for (int32 ShaderStage = 0; ShaderStage < NumGfxStages; ++ShaderStage)
 	{
 		if (GfxEntry->ShaderMicrocodes[ShaderStage] == nullptr)
 		{
 			continue;
 		}
-		const EShaderFrequency CurrStage = (EShaderFrequency)ShaderStage;
+		const DescriptorSet::EStage CurrStage = (DescriptorSet::EStage)ShaderStage;
 
 		ShaderStages[PipelineInfo.stageCount].sType  = VK_STRUCTURE_TYPE_PIPELINE_SHADER_STAGE_CREATE_INFO;
-		ShaderStages[PipelineInfo.stageCount].stage = UEFrequencyToVKStageBit(CurrStage);
+		ShaderStages[PipelineInfo.stageCount].stage = UEFrequencyToVKStageBit(DescriptorSet::GetFrequencyForGfxSet(CurrStage));
 		ShaderStages[PipelineInfo.stageCount].module = GfxEntry->ShaderModules[CurrStage];
 		ShaderStages[PipelineInfo.stageCount].pName = "main";
 		PipelineInfo.stageCount++;
@@ -907,8 +899,7 @@ void FVulkanPipelineStateCache::CreateGfxPipelineFromEntry(const FGfxPipelineEnt
 
 	// Vertex Input. The structure is mandatory even without vertex attributes.
 	VkPipelineVertexInputStateCreateInfo VBInfo;
-	FMemory::Memzero(VBInfo);
-	VBInfo.sType = VK_STRUCTURE_TYPE_PIPELINE_VERTEX_INPUT_STATE_CREATE_INFO;
+	ZeroVulkanStruct(VBInfo, VK_STRUCTURE_TYPE_PIPELINE_VERTEX_INPUT_STATE_CREATE_INFO);
 	TArray<VkVertexInputBindingDescription> VBBindings;
 	for (const FGfxPipelineEntry::FVertexBinding& SourceBinding : GfxEntry->VertexBindings)
 	{
@@ -934,8 +925,7 @@ void FVulkanPipelineStateCache::CreateGfxPipelineFromEntry(const FGfxPipelineEnt
 	PipelineInfo.renderPass = GfxEntry->RenderPass->GetHandle();
 
 	VkPipelineInputAssemblyStateCreateInfo InputAssembly;
-	FMemory::Memzero(InputAssembly);
-	InputAssembly.sType = VK_STRUCTURE_TYPE_PIPELINE_INPUT_ASSEMBLY_STATE_CREATE_INFO;
+	ZeroVulkanStruct(InputAssembly, VK_STRUCTURE_TYPE_PIPELINE_INPUT_ASSEMBLY_STATE_CREATE_INFO);
 	InputAssembly.topology = (VkPrimitiveTopology)GfxEntry->Topology;
 
 	PipelineInfo.pInputAssemblyState = &InputAssembly;
@@ -945,16 +935,14 @@ void FVulkanPipelineStateCache::CreateGfxPipelineFromEntry(const FGfxPipelineEnt
 	GfxEntry->Rasterizer.WriteInto(RasterizerState);
 
 	VkPipelineDepthStencilStateCreateInfo DepthStencilState;
-	FMemory::Memzero(DepthStencilState);
-	DepthStencilState.sType = VK_STRUCTURE_TYPE_PIPELINE_DEPTH_STENCIL_STATE_CREATE_INFO;
+	ZeroVulkanStruct(DepthStencilState, VK_STRUCTURE_TYPE_PIPELINE_DEPTH_STENCIL_STATE_CREATE_INFO);
 	GfxEntry->DepthStencil.WriteInto(DepthStencilState);
 
 	PipelineInfo.pRasterizationState = &RasterizerState;
 	PipelineInfo.pDepthStencilState = &DepthStencilState;
 
 	VkPipelineDynamicStateCreateInfo DynamicState;
-	FMemory::Memzero(DynamicState);
-	DynamicState.sType = VK_STRUCTURE_TYPE_PIPELINE_DYNAMIC_STATE_CREATE_INFO;
+	ZeroVulkanStruct(DynamicState, VK_STRUCTURE_TYPE_PIPELINE_DYNAMIC_STATE_CREATE_INFO);
 	VkDynamicState DynamicStatesEnabled[VK_DYNAMIC_STATE_RANGE_SIZE];
 	DynamicState.pDynamicStates = DynamicStatesEnabled;
 	FMemory::Memzero(DynamicStatesEnabled);
@@ -1007,8 +995,7 @@ void FVulkanPipelineStateCache::CreatGfxEntryRuntimeObjects(FGfxPipelineEntry* G
 			if (GfxEntry->ShaderMicrocodes[Index] != nullptr)
 			{
 				VkShaderModuleCreateInfo ModuleCreateInfo;
-				FMemory::Memzero(ModuleCreateInfo);
-				ModuleCreateInfo.sType = VK_STRUCTURE_TYPE_SHADER_MODULE_CREATE_INFO;
+				ZeroVulkanStruct(ModuleCreateInfo, VK_STRUCTURE_TYPE_SHADER_MODULE_CREATE_INFO);
 				ModuleCreateInfo.codeSize = GfxEntry->ShaderMicrocodes[Index]->Num();
 				ModuleCreateInfo.pCode = (uint32*)GfxEntry->ShaderMicrocodes[Index]->GetData();
 				VERIFYVULKANRESULT(VulkanRHI::vkCreateShaderModule(Device->GetInstanceHandle(), &ModuleCreateInfo, nullptr, &GfxEntry->ShaderModules[Index]));
@@ -1110,11 +1097,11 @@ void FVulkanPipelineStateCache::RebuildCache()
 
 FVulkanPipelineStateCache::FShaderHashes::FShaderHashes(const FGraphicsPipelineStateInitializer& PSOInitializer)
 {
-	Stages[SF_Vertex] = GetShaderHash<FRHIVertexShader, FVulkanVertexShader>(PSOInitializer.BoundShaderState.VertexShaderRHI);
-	Stages[SF_Pixel] = GetShaderHash<FRHIPixelShader, FVulkanPixelShader>(PSOInitializer.BoundShaderState.PixelShaderRHI);
-	Stages[SF_Geometry] = GetShaderHash<FRHIGeometryShader, FVulkanGeometryShader>(PSOInitializer.BoundShaderState.GeometryShaderRHI);
-	Stages[SF_Hull] = GetShaderHash<FRHIHullShader, FVulkanHullShader>(PSOInitializer.BoundShaderState.HullShaderRHI);
-	Stages[SF_Domain] = GetShaderHash<FRHIDomainShader, FVulkanDomainShader>(PSOInitializer.BoundShaderState.DomainShaderRHI);
+	Stages[DescriptorSet::Vertex] = GetShaderHash<FRHIVertexShader, FVulkanVertexShader>(PSOInitializer.BoundShaderState.VertexShaderRHI);
+	Stages[DescriptorSet::Pixel] = GetShaderHash<FRHIPixelShader, FVulkanPixelShader>(PSOInitializer.BoundShaderState.PixelShaderRHI);
+	Stages[DescriptorSet::Geometry] = GetShaderHash<FRHIGeometryShader, FVulkanGeometryShader>(PSOInitializer.BoundShaderState.GeometryShaderRHI);
+	//Stages[DescriptorSet::Hull] = GetShaderHash<FRHIHullShader, FVulkanHullShader>(PSOInitializer.BoundShaderState.HullShaderRHI);
+	//Stages[DescriptorSet::Domain] = GetShaderHash<FRHIDomainShader, FVulkanDomainShader>(PSOInitializer.BoundShaderState.DomainShaderRHI);
 	Finalize();
 }
 
@@ -1140,28 +1127,20 @@ inline FVulkanLayout* FVulkanPipelineStateCache::FindOrAddLayout(const FVulkanDe
 	return Layout;
 }
 
-FVulkanPipelineStateCache::FGfxPipelineEntry* FVulkanPipelineStateCache::CreateGfxEntry(const FGraphicsPipelineStateInitializer& PSOInitializer)
+FVulkanLayout* FVulkanPipelineStateCache::GetOrGenerateGfxLayout(const FGraphicsPipelineStateInitializer& PSOInitializer,
+	FVulkanShader** OutShaders, FVulkanVertexInputStateInfo& OutVertexInputState)
 {
-	FGfxPipelineEntry* OutGfxEntry = new FGfxPipelineEntry();
-
 	const FBoundShaderStateInput& BSI = PSOInitializer.BoundShaderState;
-	FVulkanShader* Shaders[SF_Compute] = {};
-
-	OutGfxEntry->RenderPass = Device->GetImmediateContext().PrepareRenderPassForPSOCreation(PSOInitializer);
 
-	// generate a FVulkanVertexInputStateInfo
-	FVulkanVertexInputStateInfo VertexInputState;
-	check(BSI.VertexShaderRHI);
 	FVulkanVertexShader* VS = ResourceCast(BSI.VertexShaderRHI);
 	const FVulkanCodeHeader& VSHeader = VS->GetCodeHeader();
-	Shaders[SF_Vertex] = VS;
-	VertexInputState.Generate(ResourceCast(PSOInitializer.BoundShaderState.VertexDeclarationRHI), VSHeader.SerializedBindings.InOutMask);
-
+	OutShaders[DescriptorSet::Vertex] = VS;
+	OutVertexInputState.Generate(ResourceCast(PSOInitializer.BoundShaderState.VertexDeclarationRHI), VSHeader.SerializedBindings.InOutMask);
 
 	// Generate a layout
 	FVulkanDescriptorSetsLayoutInfo DescriptorSetLayoutInfo;
-	DescriptorSetLayoutInfo.AddBindingsForStage(VK_SHADER_STAGE_VERTEX_BIT, EDescriptorSetStage::Vertex, VSHeader);
-	
+	DescriptorSetLayoutInfo.AddBindingsForStage(VK_SHADER_STAGE_VERTEX_BIT, DescriptorSet::Vertex, VSHeader);
+
 	FVulkanPixelShader* PS = nullptr;
 	if (BSI.PixelShaderRHI)
 	{
@@ -1172,33 +1151,37 @@ FVulkanPipelineStateCache::FGfxPipelineEntry* FVulkanPipelineStateCache::CreateG
 		// Some mobile devices expect PS stage (S7 Adreno)
 		PS = ResourceCast(TShaderMapRef<FNULLPS>(GetGlobalShaderMap(GMaxRHIFeatureLevel))->GetPixelShader());
 	}
-	
+
 	if (PS)
 	{
-		Shaders[SF_Pixel] = PS;
+		OutShaders[DescriptorSet::Pixel] = PS;
 		const FVulkanCodeHeader& PSHeader = PS->GetCodeHeader();
-		DescriptorSetLayoutInfo.AddBindingsForStage(VK_SHADER_STAGE_FRAGMENT_BIT, EDescriptorSetStage::Pixel, PSHeader);
+		DescriptorSetLayoutInfo.AddBindingsForStage(VK_SHADER_STAGE_FRAGMENT_BIT, DescriptorSet::Pixel, PSHeader);
 	}
 
 	if (BSI.GeometryShaderRHI)
 	{
 		FVulkanGeometryShader* GS = ResourceCast(BSI.GeometryShaderRHI);
-		Shaders[SF_Geometry] = GS;
+		OutShaders[DescriptorSet::Geometry] = GS;
 		const FVulkanCodeHeader& GSHeader = GS->GetCodeHeader();
-		DescriptorSetLayoutInfo.AddBindingsForStage(VK_SHADER_STAGE_GEOMETRY_BIT, EDescriptorSetStage::Geometry, GSHeader);
+		DescriptorSetLayoutInfo.AddBindingsForStage(VK_SHADER_STAGE_GEOMETRY_BIT, DescriptorSet::Geometry, GSHeader);
 	}
+
 	if (BSI.HullShaderRHI)
 	{
+		ensureMsgf(0, TEXT("Tessellation not supported yet!"));
+		/*
 		// Can't have Hull w/o Domain
 		check(BSI.DomainShaderRHI);
 		FVulkanHullShader* HS = ResourceCast(BSI.HullShaderRHI);
 		FVulkanDomainShader* DS = ResourceCast(BSI.DomainShaderRHI);
-		Shaders[SF_Hull] = HS;
-		Shaders[SF_Domain] = DS;
+		Shaders[DescriptorSet::Hull] = HS;
+		Shaders[DescriptorSet::Domain] = DS;
 		const FVulkanCodeHeader& HSHeader = HS->GetCodeHeader();
 		const FVulkanCodeHeader& DSHeader = DS->GetCodeHeader();
-		DescriptorSetLayoutInfo.AddBindingsForStage(VK_SHADER_STAGE_TESSELLATION_CONTROL_BIT, EDescriptorSetStage::Hull, HSHeader);
-		DescriptorSetLayoutInfo.AddBindingsForStage(VK_SHADER_STAGE_TESSELLATION_EVALUATION_BIT, EDescriptorSetStage::Domain, DSHeader);
+		DescriptorSetLayoutInfo.AddBindingsForStage(VK_SHADER_STAGE_TESSELLATION_CONTROL_BIT, DescriptorSet::Hull, HSHeader);
+		DescriptorSetLayoutInfo.AddBindingsForStage(VK_SHADER_STAGE_TESSELLATION_EVALUATION_BIT, DescriptorSet::Domain, DSHeader);
+		*/
 	}
 	else
 	{
@@ -1206,9 +1189,20 @@ FVulkanPipelineStateCache::FGfxPipelineEntry* FVulkanPipelineStateCache::CreateG
 		check(BSI.DomainShaderRHI == nullptr);
 	}
 
-	OutGfxEntry->Layout = FindOrAddLayout(DescriptorSetLayoutInfo);
+	return FindOrAddLayout(DescriptorSetLayoutInfo);
+}
+
+FVulkanPipelineStateCache::FGfxPipelineEntry* FVulkanPipelineStateCache::CreateGfxEntry(const FGraphicsPipelineStateInitializer& PSOInitializer)
+{
+	FGfxPipelineEntry* OutGfxEntry = new FGfxPipelineEntry();
+
+	FVulkanShader* Shaders[DescriptorSet::NumGfxStages];
+	FMemory::Memzero(Shaders);
+
+	OutGfxEntry->RenderPass = Device->GetImmediateContext().PrepareRenderPassForPSOCreation(PSOInitializer);
 
-	// now we should have everything we need ??
+	FVulkanVertexInputStateInfo VertexInputState;
+	OutGfxEntry->Layout = GetOrGenerateGfxLayout(PSOInitializer, Shaders, VertexInputState);
 
 	OutGfxEntry->RasterizationSamples = OutGfxEntry->RenderPass->GetLayout().GetAttachmentDescriptions()[0].samples;
 	ensure(OutGfxEntry->RasterizationSamples == PSOInitializer.NumSamples);
@@ -1254,14 +1248,14 @@ FVulkanPipelineStateCache::FGfxPipelineEntry* FVulkanPipelineStateCache::CreateG
 	}
 
 	int32 NumShaders = 0;
-	for (int32 Index = 0; Index < SF_Compute; ++Index)
+	for (int32 Index = 0; Index < NumGfxStages; ++Index)
 	{
 		FVulkanShader* Shader = Shaders[Index];
 		if (Shader)
 		{
-			check(Shader->CodeSize != 0);
+			check(Shader->Spirv.Num() != 0);
 
-			FSHAHash Hash = GetShaderHashForStage(PSOInitializer, Index);
+			FSHAHash Hash = GetShaderHashForStage(PSOInitializer, (DescriptorSet::EStage)Index);
 			OutGfxEntry->ShaderHashes[Index] = Hash;
 
 			OutGfxEntry->ShaderMicrocodes[Index] = ShaderCache.Get(Hash);
@@ -1276,7 +1270,6 @@ FVulkanPipelineStateCache::FGfxPipelineEntry* FVulkanPipelineStateCache::CreateG
 	}
 	check(NumShaders > 0);
 
-
 	OutGfxEntry->RenderTargets.ReadFrom(OutGfxEntry->RenderPass->GetLayout());
 
 	return OutGfxEntry;
@@ -1422,10 +1415,9 @@ FVulkanPipelineStateCache::FComputePipelineEntry* FVulkanPipelineStateCache::Cre
 {
 	FComputePipelineEntry* OutComputeEntry = new FComputePipelineEntry();
 
-	check(ComputeShader->CodeSize != 0);
 	OutComputeEntry->ShaderHash = ComputeShader->GetHash();
 	OutComputeEntry->ShaderMicrocode = ShaderCache.Get(ComputeShader->GetHash());
-	if (OutComputeEntry->ShaderMicrocode == nullptr)
+	if (!OutComputeEntry->ShaderMicrocode)
 	{
 		OutComputeEntry->ShaderMicrocode = ShaderCache.Add(ComputeShader->GetHash(), ComputeShader);
 	}
@@ -1433,7 +1425,7 @@ FVulkanPipelineStateCache::FComputePipelineEntry* FVulkanPipelineStateCache::Cre
 	OutComputeEntry->ShaderModule = ComputeShader->GetHandle();
 
 	FVulkanDescriptorSetsLayoutInfo DescriptorSetLayoutInfo;
-	DescriptorSetLayoutInfo.AddBindingsForStage(VK_SHADER_STAGE_COMPUTE_BIT, EDescriptorSetStage::Compute, ComputeShader->GetCodeHeader());
+	DescriptorSetLayoutInfo.AddBindingsForStage(VK_SHADER_STAGE_COMPUTE_BIT, DescriptorSet::Compute, ComputeShader->GetCodeHeader());
 	OutComputeEntry->Layout = FindOrAddLayout(DescriptorSetLayoutInfo);
 
 	const TArray<FVulkanDescriptorSetsLayout::FSetLayout>& Layouts = DescriptorSetLayoutInfo.GetLayouts();
@@ -1456,8 +1448,7 @@ FVulkanComputePipeline* FVulkanPipelineStateCache::CreateComputePipelineFromEntr
 	FVulkanComputePipeline* Pipeline = new FVulkanComputePipeline(Device);
 
 	VkComputePipelineCreateInfo PipelineInfo;
-	FMemory::Memzero(PipelineInfo);
-	PipelineInfo.sType = VK_STRUCTURE_TYPE_COMPUTE_PIPELINE_CREATE_INFO;
+	ZeroVulkanStruct(PipelineInfo, VK_STRUCTURE_TYPE_COMPUTE_PIPELINE_CREATE_INFO);
 	PipelineInfo.stage.sType = VK_STRUCTURE_TYPE_PIPELINE_SHADER_STAGE_CREATE_INFO;
 	PipelineInfo.stage.stage = VK_SHADER_STAGE_COMPUTE_BIT;
 	PipelineInfo.stage.module = ComputeEntry->ShaderModule;
@@ -1498,8 +1489,7 @@ void FVulkanPipelineStateCache::CreateComputeEntryRuntimeObjects(FComputePipelin
 		if (ComputeEntry->ShaderMicrocode != nullptr)
 		{
 			VkShaderModuleCreateInfo ModuleCreateInfo;
-			FMemory::Memzero(ModuleCreateInfo);
-			ModuleCreateInfo.sType = VK_STRUCTURE_TYPE_SHADER_MODULE_CREATE_INFO;
+			ZeroVulkanStruct(ModuleCreateInfo, VK_STRUCTURE_TYPE_SHADER_MODULE_CREATE_INFO);
 			ModuleCreateInfo.codeSize = ComputeEntry->ShaderMicrocode->Num();
 			ModuleCreateInfo.pCode = (uint32*)ComputeEntry->ShaderMicrocode->GetData();
 			VERIFYVULKANRESULT(VulkanRHI::vkCreateShaderModule(Device->GetInstanceHandle(), &ModuleCreateInfo, nullptr, &ComputeEntry->ShaderModule));
diff --git a/Engine/Source/Runtime/VulkanRHI/Private/VulkanPipeline.h b/Engine/Source/Runtime/VulkanRHI/Private/VulkanPipeline.h
--- a/Engine/Source/Runtime/VulkanRHI/Private/VulkanPipeline.h
+++ b/Engine/Source/Runtime/VulkanRHI/Private/VulkanPipeline.h
@@ -111,6 +111,7 @@ class FVulkanPipelineStateCache
 public:
 	inline FVulkanGraphicsPipelineState* FindInRuntimeCache(const FGraphicsPipelineStateInitializer& Initializer, uint32& OutHash)
 	{
+		//#todo-rco: Improve hash: Don't use BSS pointers; skip load/store actions?? etc
 		OutHash = FCrc::MemCrc32(&Initializer, sizeof(Initializer));
 		
 		{
@@ -139,7 +140,7 @@ public:
 	enum
 	{
 		// Bump every time serialization changes
-		VERSION = 18,
+		VERSION = 19,
 	};
 
 	struct FDescriptorSetLayoutBinding
@@ -164,22 +165,20 @@ public:
 	// Shader microcode is shared between pipeline entries so keep a cache around to prevent duplicated storage
 	struct FShaderUCodeCache
 	{
-		using TCodeHandle = TArray<uint8>;
-		using TDataMap = TMap<FSHAHash, TCodeHandle>;
-		TDataMap Data;
+		typedef TMap<FSHAHash, TArray<uint32>> THashToMicrocode;
+		THashToMicrocode	Data;
 
-		TCodeHandle* Add(const FSHAHash& Hash, const FVulkanShader* Shader)
+		TArray<uint32>* Add(const FSHAHash& Hash, const FVulkanShader* Shader)
 		{
-			check(Shader->CodeSize != 0);
+			check(Shader->Spirv.Num() != 0);
 
-			TCodeHandle& Code = Data.Add(Hash);
-			Code.AddUninitialized(Shader->CodeSize);
-			FMemory::Memcpy(Code.GetData(), Shader->Code, Shader->CodeSize);
+			TArray<uint32>& Code = Data.Add(Hash);
+			Code = Shader->Spirv;
 
 			return &Data[Hash];
 		}
 
-		TCodeHandle* Get(const FSHAHash& Hash)
+		TArray<uint32>* Get(const FSHAHash& Hash)
 		{
 			return Data.Find(Hash);
 		}
@@ -331,8 +330,8 @@ public:
 		};
 		FDepthStencil DepthStencil;
 
-		FShaderUCodeCache::TCodeHandle* ShaderMicrocodes[SF_Compute];
-		FSHAHash ShaderHashes[SF_Compute];
+		TArray<uint32>* ShaderMicrocodes[DescriptorSet::NumGfxStages];
+		FSHAHash ShaderHashes[DescriptorSet::NumGfxStages];
 
 		struct FRenderTargets
 		{
@@ -429,7 +428,7 @@ public:
 		~FGfxPipelineEntry();
 
 		// Vulkan Runtime Data/Objects
-		VkShaderModule ShaderModules[SF_Compute];
+		VkShaderModule ShaderModules[DescriptorSet::NumGfxStages];
 		const FVulkanRenderPass* RenderPass;
 		FVulkanLayout* Layout;
 
@@ -517,7 +516,7 @@ public:
 
 		bool bLoaded;
 
-		FShaderUCodeCache::TCodeHandle* ShaderMicrocode;
+		TArray<uint32>* ShaderMicrocode;
 		FSHAHash ShaderHash;
 		TArray<TArray<FDescriptorSetLayoutBinding>> DescriptorSetLayoutBindings;
 
@@ -563,10 +562,13 @@ private:
 	void Load(const TArray<FString>& CacheFilenames);
 	void DestroyCache();
 
+	FVulkanLayout* GetOrGenerateGfxLayout(const FGraphicsPipelineStateInitializer& PSOInitializer, FVulkanShader** OutShaders, FVulkanVertexInputStateInfo& OutVertexInputState);
+
+
 	struct FShaderHashes
 	{
 		uint32 Hash;
-		FSHAHash Stages[SF_Compute];
+		FSHAHash Stages[DescriptorSet::NumGfxStages];
 
 		FShaderHashes();
 		FShaderHashes(const FGraphicsPipelineStateInitializer& PSOInitializer);
@@ -584,7 +586,7 @@ private:
 
 		friend inline bool operator == (const FShaderHashes& A, const FShaderHashes& B)
 		{
-			for (int32 Index = 0; Index < SF_Compute; ++Index)
+			for (int32 Index = 0; Index < DescriptorSet::NumGfxStages; ++Index)
 			{
 				if (A.Stages[Index] != B.Stages[Index])
 				{
@@ -624,7 +626,7 @@ private:
 			int32 UncompressedSize = 0; // 0 == file is uncompressed
 		} Header;
 
-		FShaderUCodeCache::TDataMap* ShaderCache = nullptr;
+		FShaderUCodeCache::THashToMicrocode* ShaderCache = nullptr;
 
 		TArray<FGfxPipelineEntry*> GfxPipelineEntries;		
 
diff --git a/Engine/Source/Runtime/VulkanRHI/Private/VulkanPipelineState.cpp b/Engine/Source/Runtime/VulkanRHI/Private/VulkanPipelineState.cpp
--- a/Engine/Source/Runtime/VulkanRHI/Private/VulkanPipelineState.cpp
+++ b/Engine/Source/Runtime/VulkanRHI/Private/VulkanPipelineState.cpp
@@ -17,6 +17,9 @@ enum
 	NumAllocationsPerPool = 8,
 };
 
+
+const int32 NumGfxStages = FVulkanPlatform::IsVSPSOnly() ? DescriptorSet::NumMobileStages : DescriptorSet::NumGfxStages;
+
 extern TAutoConsoleVariable<int32> GDynamicGlobalUBs;
 
 static TAutoConsoleVariable<int32> GAlwaysWriteDS(
@@ -167,9 +170,9 @@ FVulkanGfxPipelineState::FVulkanGfxPipelineState(FVulkanDevice* InDevice, FVulka
 	FMemory::Memzero(ResourcesDirty);
 	FMemory::Memzero(ResourcesDirtyMask);
 
-	for (int32 Index = 0; Index < SF_Compute; ++Index)
+	for (int32 Index = 0; Index < (int32)NumGfxStages; ++Index)
 	{
-		const FVulkanShader* StageShader = BSS->GetShader((EShaderFrequency)Index);
+		const FVulkanShader* StageShader = BSS->GetShader((DescriptorSet::EStage)Index);
 		if (StageShader)
 		{
 			UsedStagesMask |= (1 << Index);
@@ -183,20 +186,23 @@ FVulkanGfxPipelineState::FVulkanGfxPipelineState(FVulkanDevice* InDevice, FVulka
 		}
 	}
 
-	PackedUniformBuffers[SF_Vertex].Init(CodeHeaderPerStage[SF_Vertex], PackedUniformBuffersMask[SF_Vertex], UniformBuffersWithDataMask[SF_Vertex], ResourcesDirtyMask[SF_Vertex]);
+	PackedUniformBuffers[DescriptorSet::Vertex].Init(CodeHeaderPerStage[DescriptorSet::Vertex], PackedUniformBuffersMask[DescriptorSet::Vertex], UniformBuffersWithDataMask[DescriptorSet::Vertex], ResourcesDirtyMask[DescriptorSet::Vertex]);
 
 	if (BSS->GetPixelShader())
 	{
-		PackedUniformBuffers[SF_Pixel].Init(CodeHeaderPerStage[SF_Pixel], PackedUniformBuffersMask[SF_Pixel], UniformBuffersWithDataMask[SF_Pixel], ResourcesDirtyMask[SF_Pixel]);
+		PackedUniformBuffers[DescriptorSet::Pixel].Init(CodeHeaderPerStage[DescriptorSet::Pixel], PackedUniformBuffersMask[DescriptorSet::Pixel], UniformBuffersWithDataMask[DescriptorSet::Pixel], ResourcesDirtyMask[DescriptorSet::Pixel]);
 	}
 	if (BSS->GetGeometryShader())
 	{
-		PackedUniformBuffers[SF_Geometry].Init(CodeHeaderPerStage[SF_Geometry], PackedUniformBuffersMask[SF_Geometry], UniformBuffersWithDataMask[SF_Geometry], ResourcesDirtyMask[SF_Geometry]);
+		PackedUniformBuffers[DescriptorSet::Geometry].Init(CodeHeaderPerStage[DescriptorSet::Geometry], PackedUniformBuffersMask[DescriptorSet::Geometry], UniformBuffersWithDataMask[DescriptorSet::Geometry], ResourcesDirtyMask[DescriptorSet::Geometry]);
 	}
 	if (BSS->GetHullShader())
 	{
-		PackedUniformBuffers[SF_Domain].Init(CodeHeaderPerStage[SF_Domain], PackedUniformBuffersMask[SF_Domain], UniformBuffersWithDataMask[SF_Domain], ResourcesDirtyMask[SF_Domain]);
-		PackedUniformBuffers[SF_Hull].Init(CodeHeaderPerStage[SF_Hull], PackedUniformBuffersMask[SF_Hull], UniformBuffersWithDataMask[SF_Hull], ResourcesDirtyMask[SF_Domain]);
+		ensureMsgf(0, TEXT("Tessellation not supported yet!"));
+/*
+		PackedUniformBuffers[DescriptorSet::Domain].Init(CodeHeaderPerStage[DescriptorSet::Domain], PackedUniformBuffersMask[DescriptorSet::Domain], UniformBuffersWithDataMask[DescriptorSet::Domain], ResourcesDirtyMask[DescriptorSet::Domain]);
+		PackedUniformBuffers[DescriptorSet::Hull].Init(CodeHeaderPerStage[DescriptorSet::Hull], PackedUniformBuffersMask[DescriptorSet::Hull], UniformBuffersWithDataMask[DescriptorSet::Hull], ResourcesDirtyMask[DescriptorSet::Domain]);
+*/
 	}
 
 	check(InGfxPipeline);
@@ -220,11 +226,9 @@ void FVulkanGfxPipelineState::CreateDescriptorWriteInfos()
 {
 	check(DSWriteContainer.DescriptorWrites.Num() == 0);
 
-	static_assert(SF_Geometry + 1 == SF_Compute, "Loop assumes compute is after gfx stages!");
-
-	for (uint32 Stage = 0; Stage < SF_Compute; Stage++)
+	for (int32 Stage = 0; Stage < NumGfxStages; Stage++)
 	{
-		const FVulkanShader* StageShader = BSS->GetShader((EShaderFrequency)Stage);
+		const FVulkanShader* StageShader = BSS->GetShader((DescriptorSet::EStage)Stage);
 		if (!StageShader)
 		{
 			continue;
@@ -254,13 +258,13 @@ void FVulkanGfxPipelineState::CreateDescriptorWriteInfos()
 	VkDescriptorImageInfo* CurrentImageInfo = DSWriteContainer.DescriptorImageInfo.GetData();
 	VkDescriptorBufferInfo* CurrentBufferInfo = DSWriteContainer.DescriptorBufferInfo.GetData();
 	uint8* CurrentBindingToDynamicOffsetMap = DSWriteContainer.BindingToDynamicOffsetMap.GetData();
-	uint32 DynamicOffsetsStart[SF_Compute];
+	uint32 DynamicOffsetsStart[DescriptorSet::NumGfxStages];
 	uint32 TotalNumDynamicOffsets = 0;
-	for (uint32 Stage = 0; Stage < SF_Compute; Stage++)
+	for (int32 Stage = 0; Stage < NumGfxStages; Stage++)
 	{
 		DynamicOffsetsStart[Stage] = TotalNumDynamicOffsets;
 
-		const FVulkanShader* StageShader = BSS->GetShader((EShaderFrequency)Stage);
+		const FVulkanShader* StageShader = BSS->GetShader((DescriptorSet::EStage)Stage);
 		if (!StageShader)
 		{
 			continue;
@@ -277,7 +281,7 @@ void FVulkanGfxPipelineState::CreateDescriptorWriteInfos()
 	}
 
 	DynamicOffsets.AddZeroed(TotalNumDynamicOffsets);
-	for (uint32 Stage = 0; Stage < SF_Compute; Stage++)
+	for (int32 Stage = 0; Stage < NumGfxStages; Stage++)
 	{
 		DSWriter[Stage].DynamicOffsets = DynamicOffsetsStart[Stage] + DynamicOffsets.GetData();
 	}
@@ -319,7 +323,6 @@ bool FVulkanGfxPipelineState::UpdateDescriptorSets(FVulkanCommandListContext* Cm
 	FVulkanUniformBufferUploader* UniformBufferUploader = CmdListContext->GetUniformBufferUploader();
 	uint8* CPURingBufferBase = (uint8*)UniformBufferUploader->GetCPUMappedPointer();
 
-	static_assert(SF_Geometry + 1 == SF_Compute, "Loop assumes compute is after gfx stages!");
 	const bool bUseDynamicGlobalUBs = (GDynamicGlobalUBs->GetInt() > 0);
 	const VkDeviceSize UBOffsetAlignment = Device->GetLimits().minUniformBufferOffsetAlignment;
 	int32 NumNonDirtyStages = 0;
@@ -328,7 +331,7 @@ bool FVulkanGfxPipelineState::UpdateDescriptorSets(FVulkanCommandListContext* Cm
 	{
 		uint32 RemainingStagesMask = UsedStagesMask;
 		uint32 RemainingHasDescriptorsPerStageMask = HasDescriptorsPerStageMask;
-		for (uint32 Stage = 0; Stage < SF_Compute && (RemainingStagesMask > 0); Stage++)
+		for (int32 Stage = 0; Stage < NumGfxStages && (RemainingStagesMask > 0); Stage++)
 		{
 			// Only process Shader Stages that exist in this pipeline
 			if (RemainingStagesMask & 1)
@@ -429,7 +432,7 @@ void FVulkanCommandListContext::RHISetGraphicsPipelineState(FGraphicsPipelineSta
 	{
 		for (int32 Index = 0; Index < PendingPixelUAVs.Num(); ++Index)
 		{
-			PendingGfxState->SetUAV(SF_Pixel, PendingPixelUAVs[Index].BindIndex, PendingPixelUAVs[Index].UAV);
+			PendingGfxState->SetUAV(DescriptorSet::Pixel, PendingPixelUAVs[Index].BindIndex, PendingPixelUAVs[Index].UAV);
 		}
 	}
 }
diff --git a/Engine/Source/Runtime/VulkanRHI/Private/VulkanPipelineState.h b/Engine/Source/Runtime/VulkanRHI/Private/VulkanPipelineState.h
--- a/Engine/Source/Runtime/VulkanRHI/Private/VulkanPipelineState.h
+++ b/Engine/Source/Runtime/VulkanRHI/Private/VulkanPipelineState.h
@@ -194,91 +194,91 @@ public:
 	}
 
 	template<VkDescriptorType Type>
-	inline void MarkDirty(EShaderFrequency Stage, bool bDirty)
+	inline void MarkDirty(DescriptorSet::EStage Stage, bool bDirty)
 	{
-		ResourcesDirty[Stage] |= ((uint64)(bDirty ? 1 : 0)) << (uint64)Type;
+		ResourcesDirty[(int8)Stage] |= ((uint64)(bDirty ? 1 : 0)) << (uint64)Type;
 	}
 
-	inline void SetStorageBuffer(EShaderFrequency Stage, uint32 BindPoint, VkBuffer Buffer, uint32 Offset, uint32 Size, VkBufferUsageFlags UsageFlags)
+	inline void SetStorageBuffer(DescriptorSet::EStage Stage, uint32 BindPoint, VkBuffer Buffer, uint32 Offset, uint32 Size, VkBufferUsageFlags UsageFlags)
 	{
 		check((UsageFlags & VK_BUFFER_USAGE_STORAGE_BUFFER_BIT) == VK_BUFFER_USAGE_STORAGE_BUFFER_BIT);
-		const bool bDirty = DSWriter[Stage].WriteStorageBuffer(BindPoint, Buffer, Offset, Size);
+		const bool bDirty = DSWriter[(int8)Stage].WriteStorageBuffer(BindPoint, Buffer, Offset, Size);
 		MarkDirty<VK_DESCRIPTOR_TYPE_STORAGE_BUFFER>(Stage, bDirty);
 	}
 
-	inline void SetUAVTexelBufferViewState(EShaderFrequency Stage, uint32 BindPoint, FVulkanBufferView* View)
+	inline void SetUAVTexelBufferViewState(DescriptorSet::EStage Stage, uint32 BindPoint, FVulkanBufferView* View)
 	{
 		check(View && (View->Flags & VK_BUFFER_USAGE_STORAGE_TEXEL_BUFFER_BIT) == VK_BUFFER_USAGE_STORAGE_TEXEL_BUFFER_BIT);
-		const bool bDirty = DSWriter[Stage].WriteStorageTexelBuffer(BindPoint, View);
+		const bool bDirty = DSWriter[(int8)Stage].WriteStorageTexelBuffer(BindPoint, View);
 		MarkDirty<VK_DESCRIPTOR_TYPE_STORAGE_TEXEL_BUFFER>(Stage, bDirty);
 	}
 
-	inline void SetUAVTextureView(EShaderFrequency Stage, uint32 BindPoint, const FVulkanTextureView& TextureView, VkImageLayout Layout)
+	inline void SetUAVTextureView(DescriptorSet::EStage Stage, uint32 BindPoint, const FVulkanTextureView& TextureView, VkImageLayout Layout)
 	{
-		const bool bDirty = DSWriter[Stage].WriteStorageImage(BindPoint, TextureView.View, Layout);
+		const bool bDirty = DSWriter[(int8)Stage].WriteStorageImage(BindPoint, TextureView.View, Layout);
 		MarkDirty<VK_DESCRIPTOR_TYPE_STORAGE_IMAGE>(Stage, bDirty);
 	}
 
-	inline void SetTexture(EShaderFrequency Stage, uint32 BindPoint, const FVulkanTextureBase* TextureBase, VkImageLayout Layout)
+	inline void SetTexture(DescriptorSet::EStage Stage, uint32 BindPoint, const FVulkanTextureBase* TextureBase, VkImageLayout Layout)
 	{
 		check(TextureBase);
-		const bool bDirty = DSWriter[Stage].WriteImage(BindPoint, TextureBase->PartialView->View, Layout);
+		const bool bDirty = DSWriter[(int8)Stage].WriteImage(BindPoint, TextureBase->PartialView->View, Layout);
 		MarkDirty<VK_DESCRIPTOR_TYPE_SAMPLED_IMAGE>(Stage, bDirty);
 	}
 
-	inline void SetSRVBufferViewState(EShaderFrequency Stage, uint32 BindPoint, FVulkanBufferView* View)
+	inline void SetSRVBufferViewState(DescriptorSet::EStage Stage, uint32 BindPoint, FVulkanBufferView* View)
 	{
 		check(View && (View->Flags & VK_BUFFER_USAGE_UNIFORM_TEXEL_BUFFER_BIT) == VK_BUFFER_USAGE_UNIFORM_TEXEL_BUFFER_BIT);
-		const bool bDirty = DSWriter[Stage].WriteUniformTexelBuffer(BindPoint, View);
+		const bool bDirty = DSWriter[(int8)Stage].WriteUniformTexelBuffer(BindPoint, View);
 		MarkDirty<VK_DESCRIPTOR_TYPE_UNIFORM_TEXEL_BUFFER>(Stage, bDirty);
 	}
 
-	inline void SetSRVTextureView(EShaderFrequency Stage, uint32 BindPoint, const FVulkanTextureView& TextureView, VkImageLayout Layout)
+	inline void SetSRVTextureView(DescriptorSet::EStage Stage, uint32 BindPoint, const FVulkanTextureView& TextureView, VkImageLayout Layout)
 	{
-		const bool bDirty = DSWriter[Stage].WriteImage(BindPoint, TextureView.View, Layout);
+		const bool bDirty = DSWriter[(int8)Stage].WriteImage(BindPoint, TextureView.View, Layout);
 		MarkDirty<VK_DESCRIPTOR_TYPE_SAMPLED_IMAGE>(Stage, bDirty);
 	}
 
-	inline void SetSamplerState(EShaderFrequency Stage, uint32 BindPoint, FVulkanSamplerState* Sampler)
+	inline void SetSamplerState(DescriptorSet::EStage Stage, uint32 BindPoint, FVulkanSamplerState* Sampler)
 	{
 		check(Sampler && Sampler->Sampler != VK_NULL_HANDLE);
-		const bool bDirty = DSWriter[Stage].WriteSampler(BindPoint, Sampler->Sampler);
+		const bool bDirty = DSWriter[(int8)Stage].WriteSampler(BindPoint, Sampler->Sampler);
 		MarkDirty<VK_DESCRIPTOR_TYPE_SAMPLER>(Stage, bDirty);
 	}
 
-	inline void SetShaderParameter(EShaderFrequency Stage, uint32 BufferIndex, uint32 ByteOffset, uint32 NumBytes, const void* NewValue)
+	inline void SetShaderParameter(DescriptorSet::EStage Stage, uint32 BufferIndex, uint32 ByteOffset, uint32 NumBytes, const void* NewValue)
 	{
-		PackedUniformBuffers[Stage].SetPackedGlobalParameter(BufferIndex, ByteOffset, NumBytes, NewValue, PackedUniformBuffersDirty[Stage]);
+		PackedUniformBuffers[(int8)Stage].SetPackedGlobalParameter(BufferIndex, ByteOffset, NumBytes, NewValue, PackedUniformBuffersDirty[(int8)Stage]);
 	}
 
-	inline void SetUniformBufferConstantData(EShaderFrequency Stage, uint32 BindPoint, const TArray<uint8>& ConstantData)
+	inline void SetUniformBufferConstantData(DescriptorSet::EStage Stage, uint32 BindPoint, const TArray<uint8>& ConstantData)
 	{
-		PackedUniformBuffers[Stage].SetEmulatedUniformBufferIntoPacked(BindPoint, ConstantData, PackedUniformBuffersDirty[Stage]);
+		PackedUniformBuffers[(int8)Stage].SetEmulatedUniformBufferIntoPacked(BindPoint, ConstantData, PackedUniformBuffersDirty[(int8)Stage]);
 	}
 
-	inline void SetUniformBuffer(EShaderFrequency Stage, uint32 BindPoint, const FVulkanUniformBuffer* UniformBuffer)
+	inline void SetUniformBuffer(DescriptorSet::EStage Stage, uint32 BindPoint, const FVulkanUniformBuffer* UniformBuffer)
 	{
-		if ((UniformBuffersWithDataMask[Stage] & (1ULL << (uint64)BindPoint)) != 0)
+		if ((UniformBuffersWithDataMask[(int8)Stage] & (1ULL << (uint64)BindPoint)) != 0)
 		{
 			extern TAutoConsoleVariable<int32> GDynamicGlobalUBs;
 			if (GDynamicGlobalUBs.GetValueOnRenderThread() > 1)
 			{
-				const bool bDirty = DSWriter[Stage].WriteDynamicUniformBuffer(BindPoint, UniformBuffer->GetHandle(), 0, UniformBuffer->GetSize(), UniformBuffer->GetOffset());
+				const bool bDirty = DSWriter[(int8)Stage].WriteDynamicUniformBuffer(BindPoint, UniformBuffer->GetHandle(), 0, UniformBuffer->GetSize(), UniformBuffer->GetOffset());
 				MarkDirty<VK_DESCRIPTOR_TYPE_UNIFORM_BUFFER_DYNAMIC>(Stage, bDirty);
 			}
 			else
 			{
-				const bool bDirty = DSWriter[Stage].WriteUniformBuffer(BindPoint, UniformBuffer->GetHandle(), UniformBuffer->GetOffset(), UniformBuffer->GetSize());
+				const bool bDirty = DSWriter[(int8)Stage].WriteUniformBuffer(BindPoint, UniformBuffer->GetHandle(), UniformBuffer->GetOffset(), UniformBuffer->GetSize());
 				MarkDirty<VK_DESCRIPTOR_TYPE_UNIFORM_BUFFER>(Stage, bDirty);
 			}
 		}
 	}
 
-	inline void SetDynamicUniformBuffer(EShaderFrequency Stage, uint32 BindPoint, const FVulkanUniformBuffer* UniformBuffer)
+	inline void SetDynamicUniformBuffer(DescriptorSet::EStage Stage, uint32 BindPoint, const FVulkanUniformBuffer* UniformBuffer)
 	{
-		if ((UniformBuffersWithDataMask[Stage] & (1ULL << (uint64)BindPoint)) != 0)
+		if ((UniformBuffersWithDataMask[(int8)Stage] & (1ULL << (uint64)BindPoint)) != 0)
 		{
-			const bool bDirty = DSWriter[Stage].WriteDynamicUniformBuffer(BindPoint, UniformBuffer->GetHandle(), UniformBuffer->GetOffset(), UniformBuffer->GetSize(), 0);
+			const bool bDirty = DSWriter[(int8)Stage].WriteDynamicUniformBuffer(BindPoint, UniformBuffer->GetHandle(), UniformBuffer->GetOffset(), UniformBuffer->GetSize(), 0);
 			MarkDirty<VK_DESCRIPTOR_TYPE_UNIFORM_BUFFER_DYNAMIC>(Stage, bDirty);
 		}
 	}
@@ -311,7 +311,7 @@ public:
 		int32 TotalNumBuffer = 0;
 		int32 TotalNumImages = 0;
 
-		for (int32 Index = 0; Index < SF_Compute; ++Index)
+		for (int32 Index = 0; Index < DescriptorSet::NumGfxStages; ++Index)
 		{
 			FVulkanShader* Shader = GfxPipeline->Shaders[Index];
 			if (!Shader)
@@ -336,16 +336,16 @@ protected:
 	uint32 UsedStagesMask = 0;
 	// Bitmask of stages that have descriptors
 	uint32 HasDescriptorsPerStageMask = 0;
-	const FVulkanCodeHeader* CodeHeaderPerStage[SF_Compute];
+	const FVulkanCodeHeader* CodeHeaderPerStage[DescriptorSet::NumGfxStages];
 
-	uint64 ResourcesDirty[SF_Compute];
-	uint64 ResourcesDirtyMask[SF_Compute];
+	uint64 ResourcesDirty[DescriptorSet::NumGfxStages];
+	uint64 ResourcesDirtyMask[DescriptorSet::NumGfxStages];
 
-	FPackedUniformBuffers PackedUniformBuffers[SF_Compute];
-	uint64 PackedUniformBuffersMask[SF_Compute];
-	uint64 PackedUniformBuffersDirty[SF_Compute];
-	uint64 UniformBuffersWithDataMask[SF_Compute];
-	FVulkanDescriptorSetWriter DSWriter[SF_Compute];
+	FPackedUniformBuffers PackedUniformBuffers[DescriptorSet::NumGfxStages];
+	uint64 PackedUniformBuffersMask[DescriptorSet::NumGfxStages];
+	uint64 PackedUniformBuffersDirty[DescriptorSet::NumGfxStages];
+	uint64 UniformBuffersWithDataMask[DescriptorSet::NumGfxStages];
+	FVulkanDescriptorSetWriter DSWriter[DescriptorSet::EStage::NumGfxStages];
 
 	FVulkanGraphicsPipelineState* GfxPipeline;
 	FVulkanBoundShaderState* BSS;
diff --git a/Engine/Source/Runtime/VulkanRHI/Private/VulkanQueue.cpp b/Engine/Source/Runtime/VulkanRHI/Private/VulkanQueue.cpp
--- a/Engine/Source/Runtime/VulkanRHI/Private/VulkanQueue.cpp
+++ b/Engine/Source/Runtime/VulkanRHI/Private/VulkanQueue.cpp
@@ -44,8 +44,7 @@ void FVulkanQueue::Submit(FVulkanCmdBuffer* CmdBuffer, uint32 NumSignalSemaphore
 	const VkCommandBuffer CmdBuffers[] = { CmdBuffer->GetHandle() };
 
 	VkSubmitInfo SubmitInfo;
-	FMemory::Memzero(SubmitInfo);
-	SubmitInfo.sType = VK_STRUCTURE_TYPE_SUBMIT_INFO;
+	ZeroVulkanStruct(SubmitInfo, VK_STRUCTURE_TYPE_SUBMIT_INFO);
 	SubmitInfo.commandBufferCount = 1;
 	SubmitInfo.pCommandBuffers = CmdBuffers;
 	SubmitInfo.signalSemaphoreCount = NumSignalSemaphores;
diff --git a/Engine/Source/Runtime/VulkanRHI/Private/VulkanRHI.cpp b/Engine/Source/Runtime/VulkanRHI/Private/VulkanRHI.cpp
--- a/Engine/Source/Runtime/VulkanRHI/Private/VulkanRHI.cpp
+++ b/Engine/Source/Runtime/VulkanRHI/Private/VulkanRHI.cpp
@@ -45,6 +45,9 @@ TAutoConsoleVariable<int32> GRHIThreadCvar(
 	TEXT("2 to use multiple RHI Thread\n")
 );
 
+bool GGPUCrashDebuggingEnabled = false;
+
+
 extern TAutoConsoleVariable<int32> GRHIAllowAsyncComputeCvar;
 
 
@@ -195,6 +198,11 @@ void FVulkanDynamicRHI::Init()
 		UE_LOG(LogVulkanRHI, Fatal, TEXT("Failed to find all required Vulkan entry points; make sure your driver supports Vulkan!"));
 	}
 
+	{
+		IConsoleVariable* GPUCrashDebuggingCVar = IConsoleManager::Get().FindConsoleVariable(TEXT("r.GPUCrashDebugging"));
+		GGPUCrashDebuggingEnabled = GPUCrashDebuggingCVar && GPUCrashDebuggingCVar->GetInt() != 0;
+	}
+
 	InitInstance();
 
 	if (GPoolSizeVRAMPercentage > 0)
@@ -302,22 +310,19 @@ void FVulkanDynamicRHI::CreateInstance()
 	auto* CVarDisableEngineAndAppRegistration = IConsoleManager::Get().FindTConsoleVariableDataInt(TEXT("r.DisableEngineAndAppRegistration"));
 	bool bDisableEngineRegistration = (CVarDisableEngineAndAppRegistration && CVarDisableEngineAndAppRegistration->GetValueOnAnyThread() != 0) ||
 		(CVarShaderDevelopmentMode && CVarShaderDevelopmentMode->GetValueOnAnyThread() != 0);
-	VkApplicationInfo App;
-	FMemory::Memzero(App);
-	App.sType = VK_STRUCTURE_TYPE_APPLICATION_INFO;
-	App.pApplicationName = bDisableEngineRegistration ? "" : "UE4";
-	App.applicationVersion = 0;
-	App.pEngineName = bDisableEngineRegistration ? "" : "UE4";
-	App.engineVersion = 15;
-	App.apiVersion = UE_VK_API_VERSION;
+	VkApplicationInfo AppInfo;
+	ZeroVulkanStruct(AppInfo, VK_STRUCTURE_TYPE_APPLICATION_INFO);
+	AppInfo.pApplicationName = bDisableEngineRegistration ? "" : "UE4";
+	//AppInfo.applicationVersion = 0;
+	AppInfo.pEngineName = bDisableEngineRegistration ? "" : "UE4";
+	AppInfo.engineVersion = 15;
+	AppInfo.apiVersion = UE_VK_API_VERSION;
 
 	VkInstanceCreateInfo InstInfo;
-	FMemory::Memzero(InstInfo);
-	InstInfo.sType = VK_STRUCTURE_TYPE_INSTANCE_CREATE_INFO;
-	InstInfo.pNext = nullptr;
-	InstInfo.pApplicationInfo = &App;
+	ZeroVulkanStruct(InstInfo, VK_STRUCTURE_TYPE_INSTANCE_CREATE_INFO);
+	InstInfo.pApplicationInfo = &AppInfo;
 
-	GetInstanceLayersAndExtensions(InstanceExtensions, InstanceLayers);
+	GetInstanceLayersAndExtensions(InstanceExtensions, InstanceLayers, bSupportsDebugUtilsExt);
 
 	InstInfo.enabledExtensionCount = InstanceExtensions.Num();
 	InstInfo.ppEnabledExtensionNames = InstInfo.enabledExtensionCount > 0 ? (const ANSICHAR* const*)InstanceExtensions.GetData() : nullptr;
@@ -325,7 +330,7 @@ void FVulkanDynamicRHI::CreateInstance()
 	InstInfo.enabledLayerCount = InstanceLayers.Num();
 	InstInfo.ppEnabledLayerNames = InstInfo.enabledLayerCount > 0 ? InstanceLayers.GetData() : nullptr;
 #if VULKAN_HAS_DEBUGGING_ENABLED
-	bSupportsDebugCallbackExt = InstanceExtensions.ContainsByPredicate([](const ANSICHAR* Key)
+	bSupportsDebugCallbackExt = !bSupportsDebugUtilsExt && InstanceExtensions.ContainsByPredicate([](const ANSICHAR* Key)
 		{ 
 			return Key && !FCStringAnsi::Strcmp(Key, VK_EXT_DEBUG_REPORT_EXTENSION_NAME); 
 		});
@@ -663,12 +668,17 @@ void FVulkanDynamicRHI::InitInstance()
 			);
 
 #if VULKAN_SUPPORTS_VALIDATION_CACHE
-		SaveValidationCacheCmd = IConsoleManager::Get().RegisterConsoleCommand(
-			TEXT("r.Vulkan.SaveValidationCache"),
-			TEXT("Save validation cache."),
-			FConsoleCommandDelegate::CreateStatic(SaveValidationCache),
-			ECVF_Default
-		);
+#if VULKAN_HAS_DEBUGGING_ENABLED
+		if (GValidationCvar.GetValueOnAnyThread() > 0)
+		{
+			SaveValidationCacheCmd = IConsoleManager::Get().RegisterConsoleCommand(
+				TEXT("r.Vulkan.SaveValidationCache"),
+				TEXT("Save validation cache."),
+				FConsoleCommandDelegate::CreateStatic(SaveValidationCache),
+				ECVF_Default
+				);
+		}
+#endif
 #endif
 
 #if UE_BUILD_DEBUG || UE_BUILD_DEVELOPMENT
@@ -775,12 +785,27 @@ void FVulkanCommandListContext::RHIPushEvent(const TCHAR* Name, FColor Color)
 	EventStack.Add(Name);
 
 #if VULKAN_ENABLE_DRAW_MARKERS
+#if 0//VULKAN_SUPPORTS_DEBUG_UTILS
+	if (auto CmdBeginLabel = Device->GetCmdBeginDebugLabel())
+	{
+		FTCHARToUTF8 Converter(Name);
+		VkDebugUtilsLabelEXT Label;
+		ZeroVulkanStruct(Label, VK_STRUCTURE_TYPE_DEBUG_UTILS_LABEL_EXT);
+		Label.pLabelName = Converter.Get();
+		FLinearColor LColor(Color);
+		Label.color[0] = LColor.R;
+		Label.color[1] = LColor.G;
+		Label.color[2] = LColor.B;
+		Label.color[3] = LColor.A;
+		CmdBeginLabel(GetCommandBufferManager()->GetActiveCmdBuffer()->GetHandle(), &Label);
+	}
+	else
+#endif
 	if (auto CmdDbgMarkerBegin = Device->GetCmdDbgMarkerBegin())
 	{
 		FTCHARToUTF8 Converter(Name);
 		VkDebugMarkerMarkerInfoEXT Info;
-		FMemory::Memzero(Info);
-		Info.sType = VK_STRUCTURE_TYPE_DEBUG_MARKER_MARKER_INFO_EXT;
+		ZeroVulkanStruct(Info, VK_STRUCTURE_TYPE_DEBUG_MARKER_MARKER_INFO_EXT);
 		Info.pMarkerName = Converter.Get();
 		FLinearColor LColor(Color);
 		Info.color[0] = LColor.R;
@@ -812,6 +837,13 @@ void FVulkanCommandListContext::RHIPushEvent(const TCHAR* Name, FColor Color)
 void FVulkanCommandListContext::RHIPopEvent()
 {
 #if VULKAN_ENABLE_DRAW_MARKERS
+#if 0//VULKAN_SUPPORTS_DEBUG_UTILS
+	if (auto CmdEndLabel = Device->GetCmdEndDebugLabel())
+	{
+		CmdEndLabel(GetCommandBufferManager()->GetActiveCmdBuffer()->GetHandle());
+	}
+	else
+#endif
 	if (auto CmdDbgMarkerEnd = Device->GetCmdDbgMarkerEnd())
 	{
 		CmdDbgMarkerEnd(GetCommandBufferManager()->GetActiveCmdBuffer()->GetHandle());
@@ -931,8 +963,7 @@ FVulkanBuffer::FVulkanBuffer(FVulkanDevice& InDevice, uint32 InSize, VkFlags InU
 	LockStack(0)
 {
 	VkBufferCreateInfo BufInfo;
-	FMemory::Memzero(BufInfo);
-	BufInfo.sType = VK_STRUCTURE_TYPE_BUFFER_CREATE_INFO;
+	ZeroVulkanStruct(BufInfo, VK_STRUCTURE_TYPE_BUFFER_CREATE_INFO);
 	BufInfo.size = Size;
 	BufInfo.usage = Usage;
 	VERIFYVULKANRESULT_EXPANDED(VulkanRHI::vkCreateBuffer(Device.GetInstanceHandle(), &BufInfo, nullptr, &Buf));
@@ -1114,9 +1145,7 @@ void FVulkanDescriptorSetsLayout::Compile()
 	for (FSetLayout& Layout : SetLayouts)
 	{
 		VkDescriptorSetLayoutCreateInfo DescriptorLayoutInfo;
-		FMemory::Memzero(DescriptorLayoutInfo);
-		DescriptorLayoutInfo.sType = VK_STRUCTURE_TYPE_DESCRIPTOR_SET_LAYOUT_CREATE_INFO;
-		DescriptorLayoutInfo.pNext = nullptr;
+		ZeroVulkanStruct(DescriptorLayoutInfo, VK_STRUCTURE_TYPE_DESCRIPTOR_SET_LAYOUT_CREATE_INFO);
 		DescriptorLayoutInfo.bindingCount = Layout.LayoutBindings.Num();
 		DescriptorLayoutInfo.pBindings = Layout.LayoutBindings.GetData();
 
@@ -1146,8 +1175,7 @@ FOLDVulkanDescriptorSets::FOLDVulkanDescriptorSets(FVulkanDevice* InDevice, cons
 	if (LayoutHandles.Num() > 0)
 	{
 		VkDescriptorSetAllocateInfo DescriptorSetAllocateInfo;
-		FMemory::Memzero(DescriptorSetAllocateInfo);
-		DescriptorSetAllocateInfo.sType = VK_STRUCTURE_TYPE_DESCRIPTOR_SET_ALLOCATE_INFO;
+		ZeroVulkanStruct(DescriptorSetAllocateInfo, VK_STRUCTURE_TYPE_DESCRIPTOR_SET_ALLOCATE_INFO);
 		// Pool will be filled in by FVulkanCommandListContext::AllocateDescriptorSets
 		//DescriptorSetAllocateInfo.descriptorPool = Pool->GetHandle();
 		DescriptorSetAllocateInfo.descriptorSetCount = LayoutHandles.Num();
@@ -1188,8 +1216,7 @@ void FVulkanBufferView::Create(FVulkanBuffer& Buffer, EPixelFormat Format, uint3
 	check(FormatInfo.Supported);
 
 	VkBufferViewCreateInfo ViewInfo;
-	FMemory::Memzero(ViewInfo);
-	ViewInfo.sType = VK_STRUCTURE_TYPE_BUFFER_VIEW_CREATE_INFO;
+	ZeroVulkanStruct(ViewInfo, VK_STRUCTURE_TYPE_BUFFER_VIEW_CREATE_INFO);
 	ViewInfo.buffer = Buffer.GetBufferHandle();
 	ViewInfo.format = (VkFormat)FormatInfo.PlatformFormat;
 	ViewInfo.offset = Offset;
@@ -1217,8 +1244,7 @@ void FVulkanBufferView::Create(VkFormat Format, FVulkanResourceMultiBuffer* Buff
 	check(Format != VK_FORMAT_UNDEFINED);
 
 	VkBufferViewCreateInfo ViewInfo;
-	FMemory::Memzero(ViewInfo);
-	ViewInfo.sType = VK_STRUCTURE_TYPE_BUFFER_VIEW_CREATE_INFO;
+	ZeroVulkanStruct(ViewInfo, VK_STRUCTURE_TYPE_BUFFER_VIEW_CREATE_INFO);
 	ViewInfo.buffer = Buffer->GetHandle();
 	ViewInfo.format = Format;
 	ViewInfo.offset = Offset;
@@ -1251,26 +1277,21 @@ FVulkanRenderPass::FVulkanRenderPass(FVulkanDevice& InDevice, const FVulkanRende
 	VkSubpassDescription SubpassDesc;
 	FMemory::Memzero(SubpassDesc);
 	SubpassDesc.pipelineBindPoint = VK_PIPELINE_BIND_POINT_GRAPHICS;
-	SubpassDesc.flags = 0;
-	SubpassDesc.inputAttachmentCount = 0;
-	SubpassDesc.pInputAttachments = nullptr;
+	//SubpassDesc.flags = 0;
+	//SubpassDesc.inputAttachmentCount = 0;
 	SubpassDesc.colorAttachmentCount = InRTLayout.GetNumColorAttachments();
 	SubpassDesc.pColorAttachments = InRTLayout.GetColorAttachmentReferences();
 	SubpassDesc.pResolveAttachments = InRTLayout.GetResolveAttachmentReferences();
 	SubpassDesc.pDepthStencilAttachment = InRTLayout.GetDepthStencilAttachmentReference();
-	SubpassDesc.preserveAttachmentCount = 0;
-	SubpassDesc.pPreserveAttachments = nullptr;
+	//SubpassDesc.preserveAttachmentCount = 0;
 
 	VkRenderPassCreateInfo CreateInfo;
-	FMemory::Memzero(CreateInfo);
-	CreateInfo.sType = VK_STRUCTURE_TYPE_RENDER_PASS_CREATE_INFO;
-	CreateInfo.pNext = nullptr;
+	ZeroVulkanStruct(CreateInfo, VK_STRUCTURE_TYPE_RENDER_PASS_CREATE_INFO);
 	CreateInfo.attachmentCount = InRTLayout.GetNumAttachmentDescriptions();
 	CreateInfo.pAttachments = InRTLayout.GetAttachmentDescriptions();
 	CreateInfo.subpassCount = 1;
 	CreateInfo.pSubpasses = &SubpassDesc;
-	CreateInfo.dependencyCount = 0;
-	CreateInfo.pDependencies = nullptr;
+	//CreateInfo.dependencyCount = 0;
 
 	VERIFYVULKANRESULT_EXPANDED(VulkanRHI::vkCreateRenderPass(Device.GetInstanceHandle(), &CreateInfo, nullptr, &RenderPass));
 }
@@ -1292,8 +1313,7 @@ void VulkanSetImageLayout(
 	const VkImageSubresourceRange& SubresourceRange)
 {
 	VkImageMemoryBarrier ImageBarrier;
-	FMemory::Memzero(ImageBarrier);
-	ImageBarrier.sType = VK_STRUCTURE_TYPE_IMAGE_MEMORY_BARRIER;
+	ZeroVulkanStruct(ImageBarrier, VK_STRUCTURE_TYPE_IMAGE_MEMORY_BARRIER);
 	ImageBarrier.oldLayout = OldLayout;
 	ImageBarrier.newLayout = NewLayout;
 	ImageBarrier.image = Image;
diff --git a/Engine/Source/Runtime/VulkanRHI/Private/VulkanRHIPrivate.h b/Engine/Source/Runtime/VulkanRHI/Private/VulkanRHIPrivate.h
--- a/Engine/Source/Runtime/VulkanRHI/Private/VulkanRHIPrivate.h
+++ b/Engine/Source/Runtime/VulkanRHI/Private/VulkanRHIPrivate.h
@@ -721,16 +721,28 @@ namespace VulkanRHI
 	}
 
 #if VULKAN_ENABLE_DRAW_MARKERS
-	inline void SetDebugObjectName(PFN_vkDebugMarkerSetObjectNameEXT DebugMarkerSetObjectName, VkDevice VulkanDevice, VkImage Image, const char* ObjectName)
+	inline void SetDebugMarkerName(PFN_vkDebugMarkerSetObjectNameEXT DebugMarkerSetObjectName, VkDevice VulkanDevice, VkImage Image, const char* ObjectName)
 	{
 		VkDebugMarkerObjectNameInfoEXT Info;
-		FMemory::Memzero(Info);
-		Info.sType = VK_STRUCTURE_TYPE_DEBUG_MARKER_OBJECT_NAME_INFO_EXT;
+		ZeroVulkanStruct(Info, VK_STRUCTURE_TYPE_DEBUG_MARKER_OBJECT_NAME_INFO_EXT);
 		Info.objectType = VK_DEBUG_REPORT_OBJECT_TYPE_IMAGE_EXT;
 		Info.object = (uint64)Image;
 		Info.pObjectName = ObjectName;
 		DebugMarkerSetObjectName(VulkanDevice, &Info);
 };
+
+#if 0//VULKAN_SUPPORTS_DEBUG_UTILS
+	inline void SetDebugName(PFN_vkSetDebugUtilsObjectNameEXT SetDebugName, VkDevice Device, VkImage Image, const char* Name)
+	{
+		FTCHARToUTF8 Converter(Name);
+		VkDebugUtilsObjectNameInfoEXT Info;
+		ZeroVulkanStruct(Info, VK_STRUCTURE_TYPE_DEBUG_UTILS_OBJECT_NAME_INFO_EXT);
+		Info.objectType = VK_OBJECT_TYPE_IMAGE;
+		Info.objectHandle = (uint64)Image;
+		Info.pObjectName = Converter.Get();
+		SetDebugName(Device, &Info);
+}
+#endif
 #endif
 
 	// For cases when we want to use DepthRead_StencilDONTCARE
@@ -765,6 +777,7 @@ namespace VulkanRHI
 
 extern int32 GVulkanSubmitAfterEveryEndRenderPass;
 extern int32 GWaitForIdleOnSubmit;
+extern bool GGPUCrashDebuggingEnabled;
 
 #if VULKAN_HAS_DEBUGGING_ENABLED
 extern bool GRenderDocFound;
diff --git a/Engine/Source/Runtime/VulkanRHI/Private/VulkanRenderTarget.cpp b/Engine/Source/Runtime/VulkanRHI/Private/VulkanRenderTarget.cpp
--- a/Engine/Source/Runtime/VulkanRHI/Private/VulkanRenderTarget.cpp
+++ b/Engine/Source/Runtime/VulkanRHI/Private/VulkanRenderTarget.cpp
@@ -216,6 +216,11 @@ void FTransitionAndLayoutManager::BeginRealRenderPass(FVulkanCommandListContext&
 	int32 NumColorTargets = RPInfo.GetNumColorRenderTargets();
 	int32 Index = 0;
 	FPendingBarrier Barrier;
+	if (RPInfo.bGeneratingMips)
+	{
+		GenerateMipsInfo.NumRenderTargets = NumColorTargets;
+	}
+
 	for (Index = 0; Index < NumColorTargets; ++Index)
 	{
 		FTextureRHIParamRef Texture = RPInfo.ColorRenderTargets[Index].RenderTarget;
@@ -235,15 +240,15 @@ void FTransitionAndLayoutManager::BeginRealRenderPass(FVulkanCommandListContext&
 #endif
 				int32 NumSlices = Surface.GetNumberOfArrayLevels();
 				GenerateMipsInfo.bInsideGenerateMips = true;
-				GenerateMipsInfo.CurrentImage = Surface.Image;
+				GenerateMipsInfo.Target[Index].CurrentImage = Surface.Image;
 
-				GenerateMipsInfo.Layouts.Reset(0);
+				GenerateMipsInfo.Target[Index].Layouts.Reset(0);
 				for (int32 SliceIndex = 0; SliceIndex < NumSlices; ++SliceIndex)
 				{
-					GenerateMipsInfo.Layouts.AddDefaulted();
+					GenerateMipsInfo.Target[Index].Layouts.AddDefaulted();
 					for (int32 MipIndex = 0; MipIndex < NumMips; ++MipIndex)
 					{
-						GenerateMipsInfo.Layouts[SliceIndex].Add(*Found);
+						GenerateMipsInfo.Target[Index].Layouts[SliceIndex].Add(*Found);
 					}
 				}
 
@@ -253,7 +258,7 @@ void FTransitionAndLayoutManager::BeginRealRenderPass(FVulkanCommandListContext&
 				}
 			}
 
-			ensure(GenerateMipsInfo.CurrentImage == Surface.Image);
+			ensure(GenerateMipsInfo.Target[Index].CurrentImage == Surface.Image);
 
 			int32 SliceIndex = (uint32)FMath::Max<int32>(RPInfo.ColorRenderTargets[Index].ArraySlice, 0);
 			int32 RTMipIndex = RPInfo.ColorRenderTargets[Index].MipIndex;
@@ -261,7 +266,7 @@ void FTransitionAndLayoutManager::BeginRealRenderPass(FVulkanCommandListContext&
 			GenerateMipsInfo.CurrentSlice = SliceIndex;
 			GenerateMipsInfo.CurrentMip = RTMipIndex;
 			GenerateMipsInfo.bLastMip = (RTMipIndex == (NumMips - 1));
-			if (GenerateMipsInfo.Layouts[SliceIndex][RTMipIndex - 1] != VK_IMAGE_LAYOUT_SHADER_READ_ONLY_OPTIMAL)
+			if (GenerateMipsInfo.Target[Index].Layouts[SliceIndex][RTMipIndex - 1] != VK_IMAGE_LAYOUT_SHADER_READ_ONLY_OPTIMAL)
 			{
 				// Transition to readable
 				int32 BarrierIndex = Barrier.AddImageBarrier(Surface.Image, VK_IMAGE_ASPECT_COLOR_BIT, 1);
@@ -269,13 +274,13 @@ void FTransitionAndLayoutManager::BeginRealRenderPass(FVulkanCommandListContext&
 				Range.baseMipLevel = RTMipIndex - 1;
 				Range.baseArrayLayer = SliceIndex;
 #if !USING_CODE_ANALYSIS
-				ensure(GenerateMipsInfo.Layouts[SliceIndex][RTMipIndex - 1] == VK_IMAGE_LAYOUT_COLOR_ATTACHMENT_OPTIMAL);
+				ensure(GenerateMipsInfo.Target[Index].Layouts[SliceIndex][RTMipIndex - 1] == VK_IMAGE_LAYOUT_COLOR_ATTACHMENT_OPTIMAL);
 #endif
 				Barrier.SetTransition(BarrierIndex, EImageLayoutBarrier::ColorAttachment, EImageLayoutBarrier::PixelShaderRead);
-				GenerateMipsInfo.Layouts[SliceIndex][RTMipIndex - 1] = VK_IMAGE_LAYOUT_SHADER_READ_ONLY_OPTIMAL;
+				GenerateMipsInfo.Target[Index].Layouts[SliceIndex][RTMipIndex - 1] = VK_IMAGE_LAYOUT_SHADER_READ_ONLY_OPTIMAL;
 			}
 
-			if (GenerateMipsInfo.Layouts[SliceIndex][RTMipIndex] != VK_IMAGE_LAYOUT_COLOR_ATTACHMENT_OPTIMAL)
+			if (GenerateMipsInfo.Target[Index].Layouts[SliceIndex][RTMipIndex] != VK_IMAGE_LAYOUT_COLOR_ATTACHMENT_OPTIMAL)
 			{
 				// Transition to writeable
 				int32 BarrierIndex = Barrier.AddImageBarrier(Surface.Image, VK_IMAGE_ASPECT_COLOR_BIT, 1);
@@ -283,10 +288,10 @@ void FTransitionAndLayoutManager::BeginRealRenderPass(FVulkanCommandListContext&
 				Range.baseMipLevel = RTMipIndex;
 				Range.baseArrayLayer = SliceIndex;
 #if !USING_CODE_ANALYSIS
-				ensure(GenerateMipsInfo.Layouts[SliceIndex][RTMipIndex] == VK_IMAGE_LAYOUT_SHADER_READ_ONLY_OPTIMAL);
+				ensure(GenerateMipsInfo.Target[Index].Layouts[SliceIndex][RTMipIndex] == VK_IMAGE_LAYOUT_SHADER_READ_ONLY_OPTIMAL);
 #endif
 				Barrier.SetTransition(BarrierIndex, EImageLayoutBarrier::PixelShaderRead, EImageLayoutBarrier::ColorAttachment);
-				GenerateMipsInfo.Layouts[SliceIndex][RTMipIndex] = VK_IMAGE_LAYOUT_COLOR_ATTACHMENT_OPTIMAL;
+				GenerateMipsInfo.Target[Index].Layouts[SliceIndex][RTMipIndex] = VK_IMAGE_LAYOUT_COLOR_ATTACHMENT_OPTIMAL;
 			}
 		}
 		else
@@ -363,18 +368,19 @@ void FTransitionAndLayoutManager::EndRealRenderPass(FVulkanCmdBuffer* CmdBuffer)
 	{
 		if (GenerateMipsInfo.bLastMip)
 		{
-			ensure(GenerateMipsInfo.Layouts[GenerateMipsInfo.CurrentSlice][GenerateMipsInfo.CurrentMip] == VK_IMAGE_LAYOUT_COLOR_ATTACHMENT_OPTIMAL);
-
 			FPendingBarrier Barrier;
+			for (int32 Index = 0; Index < GenerateMipsInfo.NumRenderTargets; ++Index)
 			{
+				ensure(GenerateMipsInfo.Target[Index].Layouts[GenerateMipsInfo.CurrentSlice][GenerateMipsInfo.CurrentMip] == VK_IMAGE_LAYOUT_COLOR_ATTACHMENT_OPTIMAL);
+
 				// Transition to readable
-				int32 BarrierIndex = Barrier.AddImageBarrier(GenerateMipsInfo.CurrentImage, VK_IMAGE_ASPECT_COLOR_BIT, 1);
+				int32 BarrierIndex = Barrier.AddImageBarrier(GenerateMipsInfo.Target[Index].CurrentImage, VK_IMAGE_ASPECT_COLOR_BIT, 1);
 				VkImageSubresourceRange& Range = Barrier.GetSubresource(BarrierIndex);
 				Range.baseMipLevel = GenerateMipsInfo.CurrentMip;
 				Range.baseArrayLayer = GenerateMipsInfo.CurrentSlice;
 				Barrier.SetTransition(BarrierIndex, EImageLayoutBarrier::ColorAttachment, EImageLayoutBarrier::PixelShaderRead);
 				// This could really be ignored...
-				GenerateMipsInfo.Layouts[GenerateMipsInfo.CurrentSlice][GenerateMipsInfo.CurrentMip] = VK_IMAGE_LAYOUT_SHADER_READ_ONLY_OPTIMAL;
+				GenerateMipsInfo.Target[Index].Layouts[GenerateMipsInfo.CurrentSlice][GenerateMipsInfo.CurrentMip] = VK_IMAGE_LAYOUT_SHADER_READ_ONLY_OPTIMAL;
 			}
 			Barrier.Execute(CmdBuffer);
 		}
@@ -663,7 +669,8 @@ void FVulkanDynamicRHI::RHIReadSurfaceData(FTextureRHIParamRef TextureRHI, FIntR
 	FVulkanCmdBuffer* CmdBuffer = ImmediateContext.GetCommandBufferManager()->GetUploadCmdBuffer();
 
 	ensure(Texture2D->Surface.StorageFormat == VK_FORMAT_R8G8B8A8_UNORM || Texture2D->Surface.StorageFormat == VK_FORMAT_B8G8R8A8_UNORM || Texture2D->Surface.StorageFormat == VK_FORMAT_R16G16B16A16_SFLOAT || Texture2D->Surface.StorageFormat == VK_FORMAT_A2B10G10R10_UNORM_PACK32);
-	const uint32 Size = NumPixels * sizeof(FColor);
+	const bool bIs8Bpp = (Texture2D->Surface.StorageFormat == VK_FORMAT_R16G16B16A16_SFLOAT);
+	const uint32 Size = NumPixels * sizeof(FColor) * (bIs8Bpp ? 2 : 1);
 	VulkanRHI::FStagingBuffer* StagingBuffer = Device->GetStagingManager().AcquireBuffer(Size, VK_BUFFER_USAGE_TRANSFER_DST_BIT, true);
 	if (GIgnoreCPUReads == 0)
 	{
@@ -715,8 +722,7 @@ void FVulkanDynamicRHI::RHIReadSurfaceData(FTextureRHIParamRef TextureRHI, FIntR
 
 /*
 	VkMappedMemoryRange MappedRange;
-	FMemory::Memzero(MappedRange);
-	MappedRange.sType = VK_STRUCTURE_TYPE_MAPPED_MEMORY_RANGE;
+	ZeroVulkanStruct(MappedRange, VK_STRUCTURE_TYPE_MAPPED_MEMORY_RANGE);
 	MappedRange.memory = StagingBuffer->GetDeviceMemoryHandle();
 	MappedRange.offset = StagingBuffer->GetAllocationOffset();
 	MappedRange.size = Size;
@@ -949,8 +955,7 @@ void FVulkanDynamicRHI::RHIReadSurfaceFloatData(FTextureRHIParamRef TextureRHI,
 		StagingBuffer->InvalidateMappedMemory();
 /*
 		VkMappedMemoryRange MappedRange;
-		FMemory::Memzero(MappedRange);
-		MappedRange.sType = VK_STRUCTURE_TYPE_MAPPED_MEMORY_RANGE;
+		ZeroVulkanStruct(MappedRange, VK_STRUCTURE_TYPE_MAPPED_MEMORY_RANGE);
 		MappedRange.memory = StagingBuffer->GetDeviceMemoryHandle();
 		MappedRange.offset = StagingBuffer->GetAllocationOffset();
 		MappedRange.size = Size;
@@ -1198,22 +1203,31 @@ void FVulkanCommandListContext::TransitionResources(const FPendingTransition& Pe
 				}
 			}
 
-			for (int32 Index = 0; Index < PendingTransition.Textures.Num(); ++Index)
+			if (bShowTransitionEvents)
 			{
-				SCOPED_RHI_CONDITIONAL_DRAW_EVENTF(*this, RHITransitionResourcesLoop, bShowTransitionEvents, TEXT("To:%i - %s"), Index, *PendingTransition.Textures[Index]->GetName().ToString());
+				for (int32 Index = 0; Index < PendingTransition.Textures.Num(); ++Index)
+				{
+					SCOPED_RHI_DRAW_EVENTF(*this, RHITransitionResourcesLoop, TEXT("To:%i - %s"), Index, *PendingTransition.Textures[Index]->GetName().ToString());
+				}
+			}
 
+			VulkanRHI::FPendingBarrier Barrier;
+			for (int32 Index = 0; Index < PendingTransition.Textures.Num(); ++Index)
+			{
 				FVulkanTextureBase* VulkanTexture = FVulkanTextureBase::Cast(PendingTransition.Textures[Index]);
 				VkImageLayout& SrcLayout = TransitionAndLayoutManager.FindOrAddLayoutRW(VulkanTexture->Surface.Image, VK_IMAGE_LAYOUT_UNDEFINED);
 				bool bIsDepthStencil = VulkanTexture->Surface.IsDepthOrStencilAspect();
 				// During HMD rendering we get a frame where nothing is rendered into the depth buffer, but CopyToTexture is still called...
 				ensure(SrcLayout != VK_IMAGE_LAYOUT_UNDEFINED || bIsDepthStencil);
 				VkImageLayout DstLayout = bIsDepthStencil ? VK_IMAGE_LAYOUT_DEPTH_STENCIL_READ_ONLY_OPTIMAL : VK_IMAGE_LAYOUT_SHADER_READ_ONLY_OPTIMAL;
-				VkImageSubresourceRange Range = VulkanRHI::SetupImageSubresourceRange(VulkanTexture->Surface.GetFullAspectMask());
-				Range.layerCount = VK_REMAINING_ARRAY_LAYERS;
-				Range.levelCount = VK_REMAINING_MIP_LEVELS;
-				VulkanSetImageLayout(CmdBuffer->GetHandle(), VulkanTexture->Surface.Image, SrcLayout, DstLayout, Range);
+
+				int32 BarrierIndex = Barrier.AddImageBarrier(VulkanTexture->Surface.Image, VulkanTexture->Surface.GetFullAspectMask(), VulkanTexture->Surface.GetNumMips(), VulkanTexture->Surface.GetNumberOfArrayLevels());
+				Barrier.SetTransition(BarrierIndex, VulkanRHI::GetImageLayoutFromVulkanLayout(SrcLayout), VulkanRHI::GetImageLayoutFromVulkanLayout(DstLayout));
+
 				SrcLayout = DstLayout;
 			}
+			//#todo-rco: Temp ensure disabled
+			Barrier.Execute(CmdBuffer, false);
 		}
 		else if (PendingTransition.TransitionType == EResourceTransitionAccess::EWritable)
 		{
@@ -1229,45 +1243,50 @@ void FVulkanCommandListContext::TransitionResources(const FPendingTransition& Pe
 				}
 			}
 
-			auto SetImageLayout = [](FTransitionAndLayoutManager& InRenderPassState, VkCommandBuffer InCmdBuffer, FVulkanSurface& Surface)
+			if (bShowTransitionEvents)
+			{
+				for (int32 i = 0; i < PendingTransition.Textures.Num(); ++i)
+				{
+					FRHITexture* RHITexture = PendingTransition.Textures[i];
+					SCOPED_RHI_DRAW_EVENTF(*this, RHITransitionResourcesLoop, TEXT("To:%i - %s"), i, *PendingTransition.Textures[i]->GetName().ToString());
+				}
+			}
+
+			VulkanRHI::FPendingBarrier Barrier;
+
+			for (int32 Index = 0; Index < PendingTransition.Textures.Num(); ++Index)
 			{
+				FVulkanSurface& Surface = FVulkanTextureBase::Cast(PendingTransition.Textures[Index])->Surface;
+
 				const VkImageAspectFlags AspectMask = Surface.GetFullAspectMask();
 				VkImageSubresourceRange SubresourceRange = {AspectMask, 0, Surface.GetNumMips(), 0, Surface.GetNumberOfArrayLevels()};
 
-				VkImageLayout& SrcLayout = InRenderPassState.FindOrAddLayoutRW(Surface.Image, VK_IMAGE_LAYOUT_UNDEFINED);
+				VkImageLayout& SrcLayout = TransitionAndLayoutManager.FindOrAddLayoutRW(Surface.Image, VK_IMAGE_LAYOUT_UNDEFINED);
+
+				int32 BarrierIndex = Barrier.AddImageBarrier(Surface.Image, Surface.GetFullAspectMask(), Surface.GetNumMips(), Surface.GetNumberOfArrayLevels());
 
 				if ((AspectMask & VK_IMAGE_ASPECT_COLOR_BIT) != 0)
 				{
-					if (SrcLayout != VK_IMAGE_LAYOUT_COLOR_ATTACHMENT_OPTIMAL)
+					// When transitioning a full 3d/cube texture just assume it's undefined
+					if (SrcLayout == VK_IMAGE_LAYOUT_UNDEFINED || (Surface.GetNumMips() > 1 && Surface.GetNumberOfArrayLevels() > 1))
 					{
-						VulkanSetImageLayout(InCmdBuffer, Surface.Image, SrcLayout, VK_IMAGE_LAYOUT_COLOR_ATTACHMENT_OPTIMAL, SubresourceRange);
-						SrcLayout = VK_IMAGE_LAYOUT_COLOR_ATTACHMENT_OPTIMAL;
+						Barrier.SetTransition(BarrierIndex, EImageLayoutBarrier::Undefined, EImageLayoutBarrier::ColorAttachment);
 					}
-				}
-				else
-				{
-					if (SrcLayout != VK_IMAGE_LAYOUT_DEPTH_STENCIL_ATTACHMENT_OPTIMAL)
+					else
 					{
-						check(Surface.IsDepthOrStencilAspect());
-						VulkanSetImageLayout(InCmdBuffer, Surface.Image, SrcLayout, VK_IMAGE_LAYOUT_DEPTH_STENCIL_ATTACHMENT_OPTIMAL, SubresourceRange);
-						SrcLayout = VK_IMAGE_LAYOUT_DEPTH_STENCIL_ATTACHMENT_OPTIMAL;
+						Barrier.SetTransition(BarrierIndex, VulkanRHI::GetImageLayoutFromVulkanLayout(SrcLayout), VulkanRHI::EImageLayoutBarrier::ColorAttachment);
 					}
+					SrcLayout = VK_IMAGE_LAYOUT_COLOR_ATTACHMENT_OPTIMAL;
 				}
-			};
-
-			for (int32 i = 0; i < PendingTransition.Textures.Num(); ++i)
-			{
-				FRHITexture* RHITexture = PendingTransition.Textures[i];
-				if (!RHITexture)
+				else
 				{
-					continue;
+					check(Surface.IsDepthOrStencilAspect());
+					Barrier.SetTransition(BarrierIndex, VulkanRHI::GetImageLayoutFromVulkanLayout(SrcLayout), VulkanRHI::EImageLayoutBarrier::DepthStencilAttachment);
+					SrcLayout = VK_IMAGE_LAYOUT_DEPTH_STENCIL_ATTACHMENT_OPTIMAL;
 				}
-
-				SCOPED_RHI_CONDITIONAL_DRAW_EVENTF(*this, RHITransitionResourcesLoop, bShowTransitionEvents, TEXT("To:%i - %s"), i, *PendingTransition.Textures[i]->GetName().ToString());
-
-				FVulkanTextureBase* Base = GetVulkanTextureFromRHITexture(RHITexture);
-				SetImageLayout(TransitionAndLayoutManager, CmdBuffer->GetHandle(), Base->Surface);
 			}
+
+			Barrier.Execute(CmdBuffer);
 		}
 		else if (PendingTransition.TransitionType == EResourceTransitionAccess::ERWSubResBarrier)
 		{
diff --git a/Engine/Source/Runtime/VulkanRHI/Private/VulkanShaders.cpp b/Engine/Source/Runtime/VulkanRHI/Private/VulkanShaders.cpp
--- a/Engine/Source/Runtime/VulkanRHI/Private/VulkanShaders.cpp
+++ b/Engine/Source/Runtime/VulkanRHI/Private/VulkanShaders.cpp
@@ -19,9 +19,6 @@ TAutoConsoleVariable<int32> GDynamicGlobalUBs(
 	ECVF_ReadOnly | ECVF_RenderThreadSafe
 );
 
-static_assert(SF_Geometry + 1 == SF_Compute, "Assumes compute is after gfx stages!");
-
-
 static void ConvertPackedUBsToDynamic(FVulkanCodeHeader& CodeHeader)
 {
 	if (GDynamicGlobalUBs.GetValueOnAnyThread() > 1)
@@ -65,43 +62,29 @@ void FVulkanShader::Create(EShaderFrequency Frequency, const TArray<uint8>& InSh
 	Ar << DebugNameArray;
 	DebugName = ANSI_TO_TCHAR(DebugNameArray.GetData());
 
-	TArray<uint8> Spirv;
 	Ar << Spirv;
+	check(Spirv.Num() != 0);
 
 	int32 CodeOffset = Ar.Tell();
 
 	VkShaderModuleCreateInfo ModuleCreateInfo;
-	FMemory::Memzero(ModuleCreateInfo);
-	ModuleCreateInfo.sType = VK_STRUCTURE_TYPE_SHADER_MODULE_CREATE_INFO;
-	ModuleCreateInfo.pNext = nullptr;
-	ModuleCreateInfo.flags = 0;
-
-	check(Code == nullptr)
-	ModuleCreateInfo.codeSize = Spirv.Num();
-	Code = (uint32*)FMemory::Malloc(ModuleCreateInfo.codeSize);
-	FMemory::Memcpy(Code, Spirv.GetData(), ModuleCreateInfo.codeSize);
+	ZeroVulkanStruct(ModuleCreateInfo, VK_STRUCTURE_TYPE_SHADER_MODULE_CREATE_INFO);
+	//ModuleCreateInfo.flags = 0;
 
-	check(Code != nullptr);
-	CodeSize = ModuleCreateInfo.codeSize;
-	ModuleCreateInfo.pCode = Code;
+	ModuleCreateInfo.codeSize = Spirv.Num() * sizeof(uint32);
+	ModuleCreateInfo.pCode = Spirv.GetData();
 
 #if VULKAN_SUPPORTS_VALIDATION_CACHE
 	VkShaderModuleValidationCacheCreateInfoEXT ValidationInfo;
 	if (Device->GetOptionalExtensions().HasEXTValidationCache)
 	{
-		FMemory::Memzero(ValidationInfo);
-		ValidationInfo.sType = VK_STRUCTURE_TYPE_SHADER_MODULE_VALIDATION_CACHE_CREATE_INFO_EXT;
+		ZeroVulkanStruct(ValidationInfo, VK_STRUCTURE_TYPE_SHADER_MODULE_VALIDATION_CACHE_CREATE_INFO_EXT);
 		ValidationInfo.validationCache = Device->GetValidationCache();
 		ModuleCreateInfo.pNext = &ValidationInfo;
 	}
 #endif
 
 	VERIFYVULKANRESULT(VulkanRHI::vkCreateShaderModule(Device->GetInstanceHandle(), &ModuleCreateInfo, nullptr, &ShaderModule));
-
-	if (Frequency == SF_Compute && CodeHeader.NEWDescriptorInfo.DescriptorTypes.Num() == 0)
-	{
-		UE_LOG(LogVulkanRHI, Warning, TEXT("Compute shader %s %s has no resources!"), *CodeHeader.ShaderName, *DebugName);
-	}
 }
 
 #if VULKAN_HAS_DEBUGGING_ENABLED
@@ -160,14 +143,6 @@ void TVulkanBaseShader<BaseResourceType, Frequency>::Create(const TArray<uint8>&
 
 FVulkanShader::~FVulkanShader()
 {
-	check(Device);
-
-	if (Code)
-	{
-		FMemory::Free(Code);
-		Code = nullptr;
-	}
-
 	if (ShaderModule != VK_NULL_HANDLE)
 	{
 		Device->GetDeferredDeletionQueue().EnqueueResource(VulkanRHI::FDeferredDeletionQueue::EType::ShaderModule, ShaderModule);
@@ -250,13 +225,10 @@ void FVulkanLayout::Compile()
 	const TArray<VkDescriptorSetLayout>& LayoutHandles = DescriptorSetLayout.GetHandles();
 
 	VkPipelineLayoutCreateInfo PipelineLayoutCreateInfo;
-	FMemory::Memzero(PipelineLayoutCreateInfo);
-	PipelineLayoutCreateInfo.sType = VK_STRUCTURE_TYPE_PIPELINE_LAYOUT_CREATE_INFO;
-	//PipelineLayoutCreateInfo.pNext = nullptr;
+	ZeroVulkanStruct(PipelineLayoutCreateInfo, VK_STRUCTURE_TYPE_PIPELINE_LAYOUT_CREATE_INFO);
 	PipelineLayoutCreateInfo.setLayoutCount = LayoutHandles.Num();
 	PipelineLayoutCreateInfo.pSetLayouts = LayoutHandles.GetData();
 	//PipelineLayoutCreateInfo.pushConstantRangeCount = 0;
-	//PipelineLayoutCreateInfo.pPushConstantRanges = nullptr;
 
 	VERIFYVULKANRESULT(VulkanRHI::vkCreatePipelineLayout(Device->GetInstanceHandle(), &PipelineLayoutCreateInfo, nullptr, &PipelineLayout));
 }
@@ -320,7 +292,7 @@ uint32 FVulkanDescriptorSetWriter::SetupDescriptorWrites(const FNEWVulkanShaderD
 	return DynamicOffsetIndex;
 }
 
-void FVulkanDescriptorSetsLayoutInfo::AddBindingsForStage(VkShaderStageFlagBits StageFlags, EDescriptorSetStage DescSet, const FVulkanCodeHeader& CodeHeader)
+void FVulkanDescriptorSetsLayoutInfo::AddBindingsForStage(VkShaderStageFlagBits StageFlags, DescriptorSet::EStage DescSet, const FVulkanCodeHeader& CodeHeader)
 {
 	//#todo-rco: Mobile assumption!
 	int32 DescriptorSetIndex = (int32)DescSet;
diff --git a/Engine/Source/Runtime/VulkanRHI/Private/VulkanState.cpp b/Engine/Source/Runtime/VulkanRHI/Private/VulkanState.cpp
--- a/Engine/Source/Runtime/VulkanRHI/Private/VulkanState.cpp
+++ b/Engine/Source/Runtime/VulkanRHI/Private/VulkanState.cpp
@@ -224,8 +224,7 @@ static inline VkCullModeFlags RasterizerCullModeToVulkan(ERasterizerCullMode InC
 
 void FVulkanSamplerState::SetupSamplerCreateInfo(const FSamplerStateInitializerRHI& Initializer, FVulkanDevice& InDevice, VkSamplerCreateInfo& OutSamplerInfo)
 {
-	FMemory::Memzero(OutSamplerInfo);
-	OutSamplerInfo.sType = VK_STRUCTURE_TYPE_SAMPLER_CREATE_INFO;
+	ZeroVulkanStruct(OutSamplerInfo, VK_STRUCTURE_TYPE_SAMPLER_CREATE_INFO);
 
 	OutSamplerInfo.magFilter = TranslateMagFilterMode(Initializer.Filter);
 	OutSamplerInfo.minFilter = TranslateMinFilterMode(Initializer.Filter);
@@ -275,8 +274,7 @@ FVulkanRasterizerState::FVulkanRasterizerState(const FRasterizerStateInitializer
 
 void FVulkanDepthStencilState::SetupCreateInfo(const FGraphicsPipelineStateInitializer& GfxPSOInit, VkPipelineDepthStencilStateCreateInfo& OutDepthStencilState)
 {
-	FMemory::Memzero(OutDepthStencilState);
-	OutDepthStencilState.sType = VK_STRUCTURE_TYPE_PIPELINE_DEPTH_STENCIL_STATE_CREATE_INFO;
+	ZeroVulkanStruct(OutDepthStencilState, VK_STRUCTURE_TYPE_PIPELINE_DEPTH_STENCIL_STATE_CREATE_INFO);
 
 	OutDepthStencilState.depthTestEnable = (Initializer.DepthTest != CF_Always || Initializer.bEnableDepthWrite) ? VK_TRUE : VK_FALSE;
 	OutDepthStencilState.depthCompareOp = CompareOpToVulkan(Initializer.DepthTest);
diff --git a/Engine/Source/Runtime/VulkanRHI/Private/VulkanSwapChain.cpp b/Engine/Source/Runtime/VulkanRHI/Private/VulkanSwapChain.cpp
--- a/Engine/Source/Runtime/VulkanRHI/Private/VulkanSwapChain.cpp
+++ b/Engine/Source/Runtime/VulkanRHI/Private/VulkanSwapChain.cpp
@@ -244,8 +244,7 @@ FVulkanSwapChain::FVulkanSwapChain(VkInstance InInstance, FVulkanDevice& InDevic
 
 
 	VkSwapchainCreateInfoKHR SwapChainInfo;
-	FMemory::Memzero(SwapChainInfo);
-	SwapChainInfo.sType = VK_STRUCTURE_TYPE_SWAPCHAIN_CREATE_INFO_KHR;
+	ZeroVulkanStruct(SwapChainInfo, VK_STRUCTURE_TYPE_SWAPCHAIN_CREATE_INFO_KHR);
 	SwapChainInfo.surface = Surface;
 	SwapChainInfo.minImageCount = DesiredNumBuffers;
 	SwapChainInfo.imageFormat = CurrFormat.format;
@@ -419,8 +418,7 @@ FVulkanSwapChain::EStatus FVulkanSwapChain::Present(FVulkanQueue* GfxQueue, FVul
 	//ensure(GfxQueue == PresentQueue);
 
 	VkPresentInfoKHR Info;
-	FMemory::Memzero(Info);
-	Info.sType = VK_STRUCTURE_TYPE_PRESENT_INFO_KHR;
+	ZeroVulkanStruct(Info, VK_STRUCTURE_TYPE_PRESENT_INFO_KHR);
 	VkSemaphore Semaphore = VK_NULL_HANDLE;
 	if (BackBufferRenderingDoneSemaphore)
 	{
diff --git a/Engine/Source/Runtime/VulkanRHI/Private/VulkanTexture.cpp b/Engine/Source/Runtime/VulkanRHI/Private/VulkanTexture.cpp
--- a/Engine/Source/Runtime/VulkanRHI/Private/VulkanTexture.cpp
+++ b/Engine/Source/Runtime/VulkanRHI/Private/VulkanTexture.cpp
@@ -179,8 +179,7 @@ VkImage FVulkanSurface::CreateImage(
 	VkImageCreateInfo TmpCreateInfo;
 	VkImageCreateInfo* ImageCreateInfoPtr = OutInfo ? OutInfo : &TmpCreateInfo;
 	VkImageCreateInfo& ImageCreateInfo = *ImageCreateInfoPtr;
-	FMemory::Memzero(ImageCreateInfo);
-	ImageCreateInfo.sType = VK_STRUCTURE_TYPE_IMAGE_CREATE_INFO;
+	ZeroVulkanStruct(ImageCreateInfo, VK_STRUCTURE_TYPE_IMAGE_CREATE_INFO);
 
 	switch(ResourceType)
 	{
@@ -1304,9 +1303,7 @@ VkImageView FVulkanTextureView::StaticCreate(FVulkanDevice& Device, VkImage Imag
 	VkImageView View = VK_NULL_HANDLE;
 
 	VkImageViewCreateInfo ViewInfo;
-	FMemory::Memzero(ViewInfo);
-	ViewInfo.sType = VK_STRUCTURE_TYPE_IMAGE_VIEW_CREATE_INFO;
-	ViewInfo.pNext = nullptr;
+	ZeroVulkanStruct(ViewInfo, VK_STRUCTURE_TYPE_IMAGE_VIEW_CREATE_INFO);
 	ViewInfo.image = Image;
 	ViewInfo.viewType = ViewType;
 	ViewInfo.format = Format;
@@ -1809,10 +1806,20 @@ void FVulkanDynamicRHI::RHIBindDebugLabelName(FTextureRHIParamRef TextureRHI, co
 #endif
 
 #if VULKAN_ENABLE_DRAW_MARKERS
-	if (Device->GetDebugMarkerSetObjectName())
+#if 0//VULKAN_SUPPORTS_DEBUG_UTILS
+	if (auto* SetDebugName = Device->GetSetDebugName())
 	{
 		FVulkanTextureBase* Base = (FVulkanTextureBase*)TextureRHI->GetTextureBaseRHI();
-		VulkanRHI::SetDebugObjectName(Device->GetDebugMarkerSetObjectName(), Device->GetInstanceHandle(), Base->Surface.Image, TCHAR_TO_ANSI(Name));
+		FTCHARToUTF8 Converter(Name);
+		VulkanRHI::SetDebugName(SetDebugName, Device->GetInstanceHandle(), Base->Surface.Image, Converter.Get());
+	}
+	else
+#endif
+	if (auto* SetObjectName = Device->GetDebugMarkerSetObjectName())
+	{
+		FVulkanTextureBase* Base = (FVulkanTextureBase*)TextureRHI->GetTextureBaseRHI();
+		FTCHARToUTF8 Converter(Name);
+		VulkanRHI::SetDebugMarkerName(SetObjectName, Device->GetInstanceHandle(), Base->Surface.Image, Converter.Get());
 	}
 #endif
 	FName DebugName(Name);
@@ -1828,8 +1835,7 @@ void FVulkanDynamicRHI::RHIBindDebugLabelName(FUnorderedAccessViewRHIParamRef Un
 		//{
 		//	FVulkanTexture2D* VulkanTexture = (FVulkanTexture2D*)Tex2d;
 		//	VkDebugMarkerObjectTagInfoEXT Info;
-		//	FMemory::Memzero(Info);
-		//	Info.sType = VK_STRUCTURE_TYPE_DEBUG_MARKER_OBJECT_NAME_INFO_EXT;
+		//	ZeroVulkanStruct(Info, VK_STRUCTURE_TYPE_DEBUG_MARKER_OBJECT_NAME_INFO_EXT);
 		//	Info.objectType = VK_DEBUG_REPORT_OBJECT_TYPE_IMAGE_EXT;
 		//	Info.object = VulkanTexture->Surface.Image;
 		//	vkDebugMarkerSetObjectNameEXT(Device->GetInstanceHandle(), &Info);
diff --git a/Engine/Source/Runtime/VulkanRHI/Private/VulkanUtil.cpp b/Engine/Source/Runtime/VulkanRHI/Private/VulkanUtil.cpp
--- a/Engine/Source/Runtime/VulkanRHI/Private/VulkanUtil.cpp
+++ b/Engine/Source/Runtime/VulkanRHI/Private/VulkanUtil.cpp
@@ -334,7 +334,7 @@ float FVulkanEventNode::GetTiming()
 void FVulkanGPUProfiler::BeginFrame()
 {
 #if VULKAN_SUPPORTS_AMD_BUFFER_MARKER
-	if (Device->GetOptionalExtensions().HasAMDBufferMarker)
+	if (GGPUCrashDebuggingEnabled && Device->GetOptionalExtensions().HasAMDBufferMarker)
 	{
 		static auto* CrashCollectionEnableCvar = IConsoleManager::Get().FindTConsoleVariableDataInt(TEXT("r.gpucrash.collectionenable"));
 		static auto* CrashCollectionDataDepth = IConsoleManager::Get().FindTConsoleVariableDataInt(TEXT("r.gpucrash.datadepth"));
@@ -532,8 +532,7 @@ namespace VulkanRHI
 		VkBuffer Buffer = VK_NULL_HANDLE;
 
 		VkBufferCreateInfo BufferCreateInfo;
-		FMemory::Memzero(BufferCreateInfo);
-		BufferCreateInfo.sType = VK_STRUCTURE_TYPE_BUFFER_CREATE_INFO;
+		ZeroVulkanStruct(BufferCreateInfo, VK_STRUCTURE_TYPE_BUFFER_CREATE_INFO);
 		BufferCreateInfo.size = Size;
 		BufferCreateInfo.usage = BufferUsageFlags;
 		VERIFYVULKANRESULT_EXPANDED(VulkanRHI::vkCreateBuffer(Device, &BufferCreateInfo, nullptr, &Buffer));
@@ -601,7 +600,7 @@ namespace VulkanRHI
 			ANSI_TO_TCHAR(VkFunction), (int32)Result, ANSI_TO_TCHAR(Filename), Line, *ErrorString);
 
 #if VULKAN_SUPPORTS_AMD_BUFFER_MARKER
-		if (GIsGPUCrashed)
+		if (GIsGPUCrashed && GGPUCrashDebuggingEnabled)
 		{
 			FVulkanDynamicRHI* RHI = (FVulkanDynamicRHI*)GDynamicRHI;
 			FVulkanDevice* Device = RHI->GetDevice();
diff --git a/Engine/Source/Runtime/VulkanRHI/Private/VulkanViewport.cpp b/Engine/Source/Runtime/VulkanRHI/Private/VulkanViewport.cpp
--- a/Engine/Source/Runtime/VulkanRHI/Private/VulkanViewport.cpp
+++ b/Engine/Source/Runtime/VulkanRHI/Private/VulkanViewport.cpp
@@ -200,7 +200,7 @@ FVulkanTexture2D* FVulkanViewport::GetBackBuffer(FRHICommandList& RHICmdList)
 #if VULKAN_ENABLE_DRAW_MARKERS
 		if (Device->GetDebugMarkerSetObjectName())
 		{
-			VulkanRHI::SetDebugObjectName(Device->GetDebugMarkerSetObjectName(), Device->GetInstanceHandle(), RenderingBackBuffer->Surface.Image, "RenderingBackBuffer");
+			VulkanRHI::SetDebugMarkerName(Device->GetDebugMarkerSetObjectName(), Device->GetInstanceHandle(), RenderingBackBuffer->Surface.Image, "RenderingBackBuffer");
 		}
 #endif
 
@@ -310,9 +310,7 @@ FVulkanFramebuffer::FVulkanFramebuffer(FVulkanDevice& Device, const FRHISetRende
 	}
 
 	VkFramebufferCreateInfo CreateInfo;
-	FMemory::Memzero(CreateInfo);
-	CreateInfo.sType = VK_STRUCTURE_TYPE_FRAMEBUFFER_CREATE_INFO;
-	CreateInfo.pNext = nullptr;
+	ZeroVulkanStruct(CreateInfo, VK_STRUCTURE_TYPE_FRAMEBUFFER_CREATE_INFO);
 	CreateInfo.renderPass = RenderPass.GetHandle();
 	CreateInfo.attachmentCount = AttachmentViews.Num();
 	CreateInfo.pAttachments = AttachmentViews.GetData();
@@ -539,7 +537,7 @@ void FVulkanViewport::CreateSwapchain()
 #if VULKAN_ENABLE_DRAW_MARKERS
 		if (Device->GetDebugMarkerSetObjectName())
 		{
-			VulkanRHI::SetDebugObjectName(Device->GetDebugMarkerSetObjectName(), Device->GetInstanceHandle(), RenderingBackBuffer->Surface.Image, "RenderingBackBuffer");
+			VulkanRHI::SetDebugMarkerName(Device->GetDebugMarkerSetObjectName(), Device->GetInstanceHandle(), RenderingBackBuffer->Surface.Image, "RenderingBackBuffer");
 		}
 #endif
 	}
diff --git a/Engine/Source/Runtime/VulkanRHI/Private/Windows/VulkanWindowsPlatform.cpp b/Engine/Source/Runtime/VulkanRHI/Private/Windows/VulkanWindowsPlatform.cpp
--- a/Engine/Source/Runtime/VulkanRHI/Private/Windows/VulkanWindowsPlatform.cpp
+++ b/Engine/Source/Runtime/VulkanRHI/Private/Windows/VulkanWindowsPlatform.cpp
@@ -117,14 +117,16 @@ void FVulkanWindowsPlatform::GetDeviceExtensions(TArray<const ANSICHAR*>& OutExt
 	OutExtensions.Add(VK_KHR_GET_MEMORY_REQUIREMENTS_2_EXTENSION_NAME);
 	OutExtensions.Add(VK_KHR_DEDICATED_ALLOCATION_EXTENSION_NAME);
 #endif
-	OutExtensions.Add(VK_AMD_BUFFER_MARKER_EXTENSION_NAME);
+	if (GGPUCrashDebuggingEnabled)
+	{
+		OutExtensions.Add(VK_AMD_BUFFER_MARKER_EXTENSION_NAME);
+	}
 }
 
 void FVulkanWindowsPlatform::CreateSurface(void* WindowHandle, VkInstance Instance, VkSurfaceKHR* OutSurface)
 {
 	VkWin32SurfaceCreateInfoKHR SurfaceCreateInfo;
-	FMemory::Memzero(SurfaceCreateInfo);
-	SurfaceCreateInfo.sType = VK_STRUCTURE_TYPE_WIN32_SURFACE_CREATE_INFO_KHR;
+	ZeroVulkanStruct(SurfaceCreateInfo, VK_STRUCTURE_TYPE_WIN32_SURFACE_CREATE_INFO_KHR);
 	SurfaceCreateInfo.hinstance = GetModuleHandle(nullptr);
 	SurfaceCreateInfo.hwnd = (HWND)WindowHandle;
 	VERIFYVULKANRESULT(VulkanDynamicAPI::vkCreateWin32SurfaceKHR(Instance, &SurfaceCreateInfo, nullptr, OutSurface));
diff --git a/Engine/Source/Runtime/VulkanRHI/Public/VulkanCommon.h b/Engine/Source/Runtime/VulkanRHI/Public/VulkanCommon.h
--- a/Engine/Source/Runtime/VulkanRHI/Public/VulkanCommon.h
+++ b/Engine/Source/Runtime/VulkanRHI/Public/VulkanCommon.h
@@ -8,52 +8,78 @@
 
 #include "RHIDefinitions.h"
 
-enum class EDescriptorSetStage
+namespace DescriptorSet
 {
-	// Adjusting these requires a full shader rebuild (ie modify the guid on VulkanCommon.usf)
-	// Keep the values in sync with EShaderFrequency
-	Vertex		= 0,
-	Hull		= 1,
-	Domain		= 2,
-	Pixel		= 3,
-	Geometry	= 4,
-
-	// Compute is its own pipeline, so it can all live as set 0
-	Compute		= 0,
-
-	Invalid		= -1,
-};
+	enum EStage
+	{
+		// Adjusting these requires a full shader rebuild (ie modify the guid on VulkanCommon.usf)
+		// Keep the values in sync with EShaderFrequency
+		Vertex			= 0,
+		Pixel			= 1,
+		Geometry		= 2,
+		//Hull			= 3,
+		//Domain			= 4,
 
-inline EDescriptorSetStage GetDescriptorSetForStage(EShaderFrequency Stage)
-{
-	switch (Stage)
+		// We don't support tessellation on desktop currently
+		NumGfxStages	= 3,
+
+		// Mobile only supports VS + PS currently
+		NumMobileStages = 2,
+
+		// Compute is its own pipeline, so it can all live as set 0
+		Compute			= 0,
+
+		Invalid			= -1,
+	};
+
+	inline EStage GetSetForFrequency(EShaderFrequency Stage)
 	{
-	case SF_Vertex:		return EDescriptorSetStage::Vertex;
-	case SF_Hull:		return EDescriptorSetStage::Hull;
-	case SF_Domain:		return EDescriptorSetStage::Domain;
-	case SF_Pixel:		return EDescriptorSetStage::Pixel;
-	case SF_Geometry:	return EDescriptorSetStage::Geometry;
-	case SF_Compute:	return EDescriptorSetStage::Compute;
-	default:
-		checkf(0, TEXT("Invalid shader Stage %d"), (int32)Stage);
-		break;
+		switch (Stage)
+		{
+		case SF_Vertex:		return Vertex;
+		//case SF_Hull:		return Hull;
+		//case SF_Domain:		return Domain;
+		case SF_Pixel:		return Pixel;
+		case SF_Geometry:	return Geometry;
+		case SF_Compute:	return Compute;
+		default:
+			checkf(0, TEXT("Invalid shader Stage %d"), (int32)Stage);
+			break;
+		}
+
+		return Invalid;
 	}
 
-	return EDescriptorSetStage::Invalid;
-}
+	inline EShaderFrequency GetFrequencyForGfxSet(EStage Stage)
+	{
+		switch (Stage)
+		{
+		case EStage::Vertex:	return SF_Vertex;
+		//case EStage::Hull:		return SF_Hull;
+		//case EStage::Domain:	return SF_Domain;
+		case EStage::Pixel:		return SF_Pixel;
+		case EStage::Geometry:	return SF_Geometry;
+		default:
+			checkf(0, TEXT("Invalid shader Stage %d"), (int32)Stage);
+			break;
+		}
+
+		return SF_NumFrequencies;
+	}
+};
 
 namespace EVulkanBindingType
 {
 	enum EType : uint8
 	{
-		PackedUniformBuffer,		//VK_DESCRIPTOR_TYPE_UNIFORM_BUFFER
+		PackedUniformBuffer,	//VK_DESCRIPTOR_TYPE_UNIFORM_BUFFER
 		UniformBuffer,			//VK_DESCRIPTOR_TYPE_UNIFORM_BUFFER
 
 		CombinedImageSampler,	//VK_DESCRIPTOR_TYPE_COMBINED_IMAGE_SAMPLER
-		Sampler,					//VK_DESCRIPTOR_TYPE_SAMPLER
-		Image,						//VK_DESCRIPTOR_TYPE_SAMPLED_IMAGE
+		Sampler,				//VK_DESCRIPTOR_TYPE_SAMPLER
+		Image,					//VK_DESCRIPTOR_TYPE_SAMPLED_IMAGE
 
-		UniformTexelBuffer,			//VK_DESCRIPTOR_TYPE_UNIFORM_TEXEL_BUFFER	Buffer<>
+		UniformTexelBuffer,		//VK_DESCRIPTOR_TYPE_UNIFORM_TEXEL_BUFFER	Buffer<>
 
 		//A storage image (VK_DESCRIPTOR_TYPE_STORAGE_IMAGE) is a descriptor type that is used for load, store, and atomic operations on image memory from within shaders bound to pipelines.
 		StorageImage,				//VK_DESCRIPTOR_TYPE_STORAGE_IMAGE		RWTexture
diff --git a/Engine/Source/Runtime/VulkanRHI/Public/VulkanConfiguration.h b/Engine/Source/Runtime/VulkanRHI/Public/VulkanConfiguration.h
--- a/Engine/Source/Runtime/VulkanRHI/Public/VulkanConfiguration.h
+++ b/Engine/Source/Runtime/VulkanRHI/Public/VulkanConfiguration.h
@@ -122,6 +122,12 @@
 	#define VULKAN_SUPPORTS_AMD_BUFFER_MARKER					0
 #endif
 
+#ifdef VK_EXT_debug_utils
+	#define VULKAN_SUPPORTS_DEBUG_UTILS							1
+#else
+	#define VULKAN_SUPPORTS_DEBUG_UTILS							0
+#endif
+
 #define VULKAN_ENABLE_TRACKING_LAYER							(VULKAN_ENABLE_BUFFER_TRACKING_LAYER || VULKAN_ENABLE_IMAGE_TRACKING_LAYER)
 #define VULKAN_ENABLE_CUSTOM_LAYER								(VULKAN_ENABLE_DUMP_LAYER || VULKAN_ENABLE_TRACKING_LAYER)
 
diff --git a/Engine/Source/Runtime/VulkanRHI/Public/VulkanContext.h b/Engine/Source/Runtime/VulkanRHI/Public/VulkanContext.h
--- a/Engine/Source/Runtime/VulkanRHI/Public/VulkanContext.h
+++ b/Engine/Source/Runtime/VulkanRHI/Public/VulkanContext.h
@@ -77,14 +77,19 @@ public:
 
 	struct FGenerateMipsInfo
 	{
-		// Per face/slice array of mip layouts
-		TArray<TArray<VkImageLayout>> Layouts;
+		int32 NumRenderTargets = 0;
 
 		bool bInsideGenerateMips;
 		bool bLastMip;
 		int32 CurrentSlice;
 		int32 CurrentMip;
-		VkImage CurrentImage;
+
+		struct
+		{
+			// Per face/slice array of mip layouts
+			TArray<TArray<VkImageLayout>> Layouts;
+			VkImage CurrentImage;
+		} Target[MaxSimultaneousRenderTargets];
 
 		FGenerateMipsInfo()
 		{
@@ -93,12 +98,16 @@ public:
 
 		void Reset()
 		{
+			NumRenderTargets = 0;
 			bInsideGenerateMips = false;
 			bLastMip = false;
 			CurrentSlice = -1;
 			CurrentMip = -1;
-			CurrentImage = VK_NULL_HANDLE;
-			Layouts.Reset(0);
+			for (int32 Index = 0; Index < MaxSimultaneousRenderTargets; ++Index)
+			{
+				Target[Index].Layouts.Reset(0);
+				Target[Index].CurrentImage = VK_NULL_HANDLE;
+			}
 		}
 	} GenerateMipsInfo;
 
@@ -371,7 +380,7 @@ protected:
 	bool bAutomaticFlushAfterComputeShader;
 	FVulkanUniformBufferUploader* UniformBufferUploader;
 
-	void SetShaderUniformBuffer(EShaderFrequency Stage, const FVulkanUniformBuffer* UniformBuffer, int32 BindingIndex, const FVulkanShader* Shader);
+	void SetShaderUniformBuffer(DescriptorSet::EStage Stage, const FVulkanUniformBuffer* UniformBuffer, int32 BindingIndex, const FVulkanShader* Shader);
 
 	struct
 	{
diff --git a/Engine/Source/Runtime/VulkanRHI/Public/VulkanDynamicRHI.h b/Engine/Source/Runtime/VulkanRHI/Public/VulkanDynamicRHI.h
--- a/Engine/Source/Runtime/VulkanRHI/Public/VulkanDynamicRHI.h
+++ b/Engine/Source/Runtime/VulkanRHI/Public/VulkanDynamicRHI.h
@@ -341,28 +341,33 @@ public:
 	// Historical number of times we've presented any and all viewports
 	uint32 TotalPresentCount = 0;
 
-	const TArray<const ANSICHAR*>& GetInstanceExtensions() const
+	inline const TArray<const ANSICHAR*>& GetInstanceExtensions() const
 	{
 		return InstanceExtensions;
 	}
 
-	const TArray<const ANSICHAR*>& GetInstanceLayers() const
+	inline const TArray<const ANSICHAR*>& GetInstanceLayers() const
 	{
 		return InstanceLayers;
 	}
 
 	static void RecreateSwapChain(void* NewNativeWindow);
 
-	VkInstance GetInstance() const
+	inline VkInstance GetInstance() const
 	{
 		return Instance;
 	}
 	
-	FVulkanDevice* GetDevice()
+	inline FVulkanDevice* GetDevice()
 	{
 		return Device;
 	}
 
+	inline bool SupportsDebugUtilsExt() const
+	{
+		return bSupportsDebugUtilsExt;
+	}
+
 	virtual void VulkanSetImageLayout( VkCommandBuffer CmdBuffer, VkImage Image, VkImageLayout OldLayout, VkImageLayout NewLayout, const VkImageSubresourceRange& SubresourceRange );
 	
 private:
@@ -397,7 +402,7 @@ protected:
 	static void WriteEndFrameTimestamp(void*);
 	*/
 
-	static void GetInstanceLayersAndExtensions(TArray<const ANSICHAR*>& OutInstanceExtensions, TArray<const ANSICHAR*>& OutInstanceLayers);
+	static void GetInstanceLayersAndExtensions(TArray<const ANSICHAR*>& OutInstanceExtensions, TArray<const ANSICHAR*>& OutInstanceLayers, bool& bOutDebugUtils);
 
 	IConsoleObject* SavePipelineCacheCmd = nullptr;
 	IConsoleObject* RebuildPipelineCacheCmd = nullptr;
@@ -416,8 +421,13 @@ public:
 	static void DumpMemory();
 #endif
 
-#if VULKAN_HAS_DEBUGGING_ENABLED
 protected:
+	bool bSupportsDebugUtilsExt = false;
+#if VULKAN_HAS_DEBUGGING_ENABLED
+#if VULKAN_SUPPORTS_DEBUG_UTILS
+	VkDebugUtilsMessengerEXT Messenger = VK_NULL_HANDLE;
+#endif
+
 	bool bSupportsDebugCallbackExt = false;
 	VkDebugReportCallbackEXT MsgCallback = VK_NULL_HANDLE;
 	void SetupDebugLayerCallback();
diff --git a/Engine/Source/Runtime/VulkanRHI/Public/VulkanGlobals.h b/Engine/Source/Runtime/VulkanRHI/Public/VulkanGlobals.h
--- a/Engine/Source/Runtime/VulkanRHI/Public/VulkanGlobals.h
+++ b/Engine/Source/Runtime/VulkanRHI/Public/VulkanGlobals.h
@@ -7,3 +7,12 @@
 #pragma once 
 
 DECLARE_LOG_CATEGORY_EXTERN(LogVulkan, Display, All);
+
+template< class T >
+static FORCEINLINE void ZeroVulkanStruct(T& Struct, VkStructureType Type)
+{
+	static_assert(!TIsPointer<T>::Value, "Don't use a pointer!");
+	static_assert(STRUCT_OFFSET(T, sType) == 0, "Assumes sType is the first member in the Vulkan type!");
+	Struct.sType = Type;
+	FMemory::Memzero(((uint8*)&Struct) + sizeof(VkStructureType), sizeof(T) - sizeof(VkStructureType));
+}
diff --git a/Engine/Source/Runtime/VulkanRHI/Public/VulkanMemory.h b/Engine/Source/Runtime/VulkanRHI/Public/VulkanMemory.h
--- a/Engine/Source/Runtime/VulkanRHI/Public/VulkanMemory.h
+++ b/Engine/Source/Runtime/VulkanRHI/Public/VulkanMemory.h
@@ -1233,8 +1233,8 @@ namespace VulkanRHI
 		case VK_IMAGE_LAYOUT_DEPTH_STENCIL_READ_ONLY_OPTIMAL:
 			return EImageLayoutBarrier::PixelDepthStencilRead;
 
-		//case VK_IMAGE_LAYOUT_GENERAL:
-		//	return EImageLayoutBarrier::ComputeGeneral;
+		case VK_IMAGE_LAYOUT_GENERAL:
+			return EImageLayoutBarrier::PixelGeneralRW;
 
 		default:
 			checkf(0, TEXT("Unknown VkImageLayout %d"), (int32)Layout);
@@ -1356,8 +1356,7 @@ namespace VulkanRHI
 	inline VkImageMemoryBarrier SetupImageMemoryBarrier(VkImage Image, VkImageAspectFlags Aspect, uint32 NumMips = 1)
 	{
 		VkImageMemoryBarrier Barrier;
-		FMemory::Memzero(Barrier);
-		Barrier.sType = VK_STRUCTURE_TYPE_IMAGE_MEMORY_BARRIER;
+		ZeroVulkanStruct(Barrier, VK_STRUCTURE_TYPE_IMAGE_MEMORY_BARRIER);
 		Barrier.image = Image;
 		Barrier.subresourceRange.aspectMask = Aspect;
 		Barrier.subresourceRange.levelCount = NumMips;
@@ -1375,7 +1374,7 @@ namespace VulkanRHI
 		TArray<VkImageMemoryBarrier> ImageBarriers;
 		TArray<VkBufferMemoryBarrier> BufferBarriers;
 
-		void InnerExecute(FVulkanCmdBuffer* CmdBuffer);
+		void InnerExecute(FVulkanCmdBuffer* CmdBuffer, bool bEnsure);
 
 	public:
 		inline void ResetStages()
@@ -1384,7 +1383,7 @@ namespace VulkanRHI
 			DestStage = 0;
 		}
 
-		inline int32 AddImageBarrier(VkImage Image, VkImageAspectFlags Aspect, uint32 NumMips)
+		inline int32 AddImageBarrier(VkImage Image, VkImageAspectFlags Aspect, uint32 NumMips, uint32 NumLayers = 1)
 		{
 			int32 Index = ImageBarriers.AddZeroed();
 			VkImageMemoryBarrier& Barrier = ImageBarriers[Index];
@@ -1392,7 +1391,7 @@ namespace VulkanRHI
 			Barrier.image = Image;
 			Barrier.subresourceRange.aspectMask = Aspect;
 			Barrier.subresourceRange.levelCount = NumMips;
-			Barrier.subresourceRange.layerCount = 1;
+			Barrier.subresourceRange.layerCount = NumLayers;
 			Barrier.srcQueueFamilyIndex = VK_QUEUE_FAMILY_IGNORED;
 			Barrier.dstQueueFamilyIndex = VK_QUEUE_FAMILY_IGNORED;
 
@@ -1442,11 +1441,11 @@ namespace VulkanRHI
 		}
 
 		// Actually Insert the cmd in cmdbuffer 
-		void Execute(FVulkanCmdBuffer* CmdBuffer)
+		void Execute(FVulkanCmdBuffer* CmdBuffer, bool bEnsure = true)
 		{
 			if (ImageBarriers.Num() > 0 || BufferBarriers.Num() > 0)
 			{
-				InnerExecute(CmdBuffer);
+				InnerExecute(CmdBuffer, bEnsure);
 			}
 		}
 	};
diff --git a/Engine/Source/Runtime/VulkanRHI/Public/VulkanResources.h b/Engine/Source/Runtime/VulkanRHI/Public/VulkanResources.h
--- a/Engine/Source/Runtime/VulkanRHI/Public/VulkanResources.h
+++ b/Engine/Source/Runtime/VulkanRHI/Public/VulkanResources.h
@@ -56,8 +56,6 @@ class FVulkanShader
 public:
 	FVulkanShader(FVulkanDevice* InDevice) :
 		ShaderModule(VK_NULL_HANDLE),
-		Code(nullptr),
-		CodeSize(0),
 		Device(InDevice)
 	{
 	}
@@ -84,16 +82,12 @@ public:
 protected:
 	/** External bindings for this shader. */
 	FVulkanCodeHeader CodeHeader;
-
 	VkShaderModule ShaderModule;
 
-	uint32* Code;
-	uint32 CodeSize;
-	FString DebugName;
-
-	TArray<ANSICHAR> GlslSource;
+	TArray<uint32>		Spirv;
+	FString				DebugName;
 
-	FVulkanDevice* Device;
+	FVulkanDevice*		Device;
 
 	friend class FVulkanCommandListContext;
 	friend class FVulkanPipelineStateCache;
@@ -130,12 +124,12 @@ public:
 	}
 };
 
-typedef TVulkanBaseShader<FRHIVertexShader, SF_Vertex> FVulkanVertexShader;
-typedef TVulkanBaseShader<FRHIPixelShader, SF_Pixel> FVulkanPixelShader;
-typedef TVulkanBaseShader<FRHIHullShader, SF_Hull> FVulkanHullShader;
-typedef TVulkanBaseShader<FRHIDomainShader, SF_Domain> FVulkanDomainShader;
-typedef TVulkanBaseShader<FRHIComputeShader, SF_Compute> FVulkanComputeShader;
-typedef TVulkanBaseShader<FRHIGeometryShader, SF_Geometry> FVulkanGeometryShader;
+typedef TVulkanBaseShader<FRHIVertexShader, SF_Vertex>		FVulkanVertexShader;
+typedef TVulkanBaseShader<FRHIPixelShader, SF_Pixel>		FVulkanPixelShader;
+typedef TVulkanBaseShader<FRHIHullShader, SF_Hull>			FVulkanHullShader;
+typedef TVulkanBaseShader<FRHIDomainShader, SF_Domain>		FVulkanDomainShader;
+typedef TVulkanBaseShader<FRHIComputeShader, SF_Compute>	FVulkanComputeShader;
+typedef TVulkanBaseShader<FRHIGeometryShader, SF_Geometry>	FVulkanGeometryShader;
 
 class FVulkanBoundShaderState : public FRHIBoundShaderState
 {
@@ -157,15 +151,15 @@ public:
 	FORCEINLINE FVulkanDomainShader*   GetDomainShader() const { return (FVulkanDomainShader*)CacheLink.GetDomainShader(); }
 	FORCEINLINE FVulkanGeometryShader* GetGeometryShader() const { return (FVulkanGeometryShader*)CacheLink.GetGeometryShader(); }
 
-	const FVulkanShader* GetShader(EShaderFrequency Stage) const
+	const FVulkanShader* GetShader(DescriptorSet::EStage Stage) const
 	{
 		switch (Stage)
 		{
-		case SF_Vertex:		return GetVertexShader();
-		case SF_Hull:		return GetHullShader();
-		case SF_Domain:		return GetDomainShader();
-		case SF_Pixel:		return GetPixelShader();
-		case SF_Geometry:	return GetGeometryShader();
+		case DescriptorSet::Vertex:		return GetVertexShader();
+		//case DescriptorSet::Hull:		return GetHullShader();
+		//case DescriptorSet::Domain:		return GetDomainShader();
+		case DescriptorSet::Pixel:		return GetPixelShader();
+		case DescriptorSet::Geometry:	return GetGeometryShader();
 		default: break;
 		}
 		checkf(0, TEXT("Invalid Shader Frequency %d"), (int32)Stage);
diff --git a/Engine/Source/Runtime/VulkanRHI/Public/VulkanShaderResources.h b/Engine/Source/Runtime/VulkanRHI/Public/VulkanShaderResources.h
--- a/Engine/Source/Runtime/VulkanRHI/Public/VulkanShaderResources.h
+++ b/Engine/Source/Runtime/VulkanRHI/Public/VulkanShaderResources.h
@@ -9,51 +9,6 @@
 #include "BoundShaderStateCache.h"
 #include "CrossCompilerCommon.h"
 
-#if 0
-class FVulkanShaderVarying
-{
-public:
-	FVulkanShaderVarying():
-		Location(0),
-		Components(0)
-	{
-	}
-
-	uint16 Location;
-	TArray<ANSICHAR> Varying;
-	uint16 Components; // Scalars/Integers
-
-	friend bool operator==(const FVulkanShaderVarying& A, const FVulkanShaderVarying& B)
-	{
-		if (&A != &B)
-		{
-			return (A.Location == B.Location)
-				&& (A.Varying.Num() == B.Varying.Num())
-				&& (FMemory::Memcmp(A.Varying.GetData(), B.Varying.GetData(), A.Varying.Num() * sizeof(ANSICHAR)) == 0)
-				&& (A.Components == B.Components);
-		}
-		return true;
-	}
-
-	friend uint32 GetTypeHash(const FVulkanShaderVarying& Var)
-	{
-		uint32 Hash = GetTypeHash(Var.Location);
-		Hash ^= FCrc::MemCrc32(Var.Varying.GetData(), Var.Varying.Num() * sizeof(ANSICHAR));
-		return Hash;
-	}
-};
-
-inline FArchive& operator<<(FArchive& Ar, FVulkanShaderVarying& Var)
-{
-	check(!Ar.IsSaving() || Var.Components > 0);
-
-	Ar << Var.Varying;
-	Ar << Var.Location;
-	Ar << Var.Components;
-	return Ar;
-}
-#endif
-
 class FVulkanShaderSerializedBindings : public CrossCompiler::FShaderBindings
 {
 public:
@@ -153,8 +108,6 @@ struct FVulkanCodeHeader
 
 	// Number of copies per emulated buffer source index (to skip searching among UniformBuffersCopyInfo). Upper uint16 is the index, Lower uint16 is the count
 	TArray<uint32> NEWEmulatedUBCopyRanges;
-
-	FBaseShaderResourceTable ShaderResourceTable;
 };
 
 inline FArchive& operator<<(FArchive& Ar, FVulkanCodeHeader& Header)
@@ -234,3 +187,24 @@ static inline VkDescriptorType BindingToDescriptorType(EVulkanBindingType::EType
 
 	return VK_DESCRIPTOR_TYPE_MAX_ENUM;
 }
+
+static inline EVulkanBindingType::EType DescriptorTypeToBinding(VkDescriptorType Type, bool bUsePacked = false)
+{
+	switch (Type)
+	{
+	case VK_DESCRIPTOR_TYPE_UNIFORM_BUFFER:			return bUsePacked ? EVulkanBindingType::PackedUniformBuffer : EVulkanBindingType::UniformBuffer;
+	case VK_DESCRIPTOR_TYPE_COMBINED_IMAGE_SAMPLER:	return EVulkanBindingType::CombinedImageSampler;
+	case VK_DESCRIPTOR_TYPE_SAMPLER:				return EVulkanBindingType::Sampler;
+	case VK_DESCRIPTOR_TYPE_SAMPLED_IMAGE:			return EVulkanBindingType::Image;
+	case VK_DESCRIPTOR_TYPE_UNIFORM_TEXEL_BUFFER:	return EVulkanBindingType::UniformTexelBuffer;
+	case VK_DESCRIPTOR_TYPE_STORAGE_IMAGE:			return EVulkanBindingType::StorageImage;
+	case VK_DESCRIPTOR_TYPE_STORAGE_TEXEL_BUFFER:	return EVulkanBindingType::StorageTexelBuffer;
+	case VK_DESCRIPTOR_TYPE_STORAGE_BUFFER:			return EVulkanBindingType::StorageBuffer;
+	default:
+		check(0);
+		break;
+	}
+
+	return EVulkanBindingType::Count;
+}
+
diff --git a/Engine/Source/Runtime/VulkanRHI/Public/VulkanState.h b/Engine/Source/Runtime/VulkanRHI/Public/VulkanState.h
--- a/Engine/Source/Runtime/VulkanRHI/Public/VulkanState.h
+++ b/Engine/Source/Runtime/VulkanRHI/Public/VulkanState.h
@@ -28,8 +28,7 @@ public:
 
 	static void ResetCreateInfo(VkPipelineRasterizationStateCreateInfo& OutInfo)
 	{
-		FMemory::Memzero(OutInfo);
-		OutInfo.sType = VK_STRUCTURE_TYPE_PIPELINE_RASTERIZATION_STATE_CREATE_INFO;
+		ZeroVulkanStruct(OutInfo, VK_STRUCTURE_TYPE_PIPELINE_RASTERIZATION_STATE_CREATE_INFO);
 		OutInfo.frontFace = VK_FRONT_FACE_CLOCKWISE;
 		OutInfo.lineWidth = 1.0f;
 	}
