UE4.20 - Copy from DevRendering
- Fix negated index typo causing crash on RWBuffers
- Fix for Vulkan crash when pressing hammering on bixby button on S8
- Fix filtering for validation messages on 1.1.77.0
- Do not require vkstatic-1.lib to link Oculus plugin
#rnx

diff --git a/Engine/Plugins/Runtime/Oculus/OculusVR/Source/OculusHMD/OculusHMD.Build.cs b/Engine/Plugins/Runtime/Oculus/OculusVR/Source/OculusHMD/OculusHMD.Build.cs
--- a/Engine/Plugins/Runtime/Oculus/OculusVR/Source/OculusHMD/OculusHMD.Build.cs
+++ b/Engine/Plugins/Runtime/Oculus/OculusVR/Source/OculusHMD/OculusHMD.Build.cs
@@ -120,7 +120,6 @@ namespace UnrealBuildTool.Rules
 							}
 
 							PublicAdditionalLibraries.Add("vulkan-1.lib");
-							PublicAdditionalLibraries.Add("vkstatic.1.lib");
 							bUseThirdParty = false;
 						}
 					}
diff --git a/Engine/Shaders/Public/VulkanCommon.ush b/Engine/Shaders/Public/VulkanCommon.ush
--- a/Engine/Shaders/Public/VulkanCommon.ush
+++ b/Engine/Shaders/Public/VulkanCommon.ush
@@ -5,7 +5,7 @@
 =============================================================================*/
 
 // Update this GUID to improve shader recompilation for Vulkan only shaders
-// GUID = {792BDC0A-A4A8-402F-8752-3F9F642A03D5}
+// GUID = {792BDC0A-A4A8-402F-8752-3F9F652A03D5}
 // Glslang 1.0.65.1
 
 #pragma once
diff --git a/Engine/Source/Developer/VulkanShaderFormat/Private/VulkanShaderFormat.h b/Engine/Source/Developer/VulkanShaderFormat/Private/VulkanShaderFormat.h
--- a/Engine/Source/Developer/VulkanShaderFormat/Private/VulkanShaderFormat.h
+++ b/Engine/Source/Developer/VulkanShaderFormat/Private/VulkanShaderFormat.h
@@ -43,7 +43,7 @@ struct FSpirv
 
 		FEntry(const FString& InName, int32 InBinding)
 			: Name(InName)
-			, Binding(-InBinding)
+			, Binding(InBinding)
 		{
 		}
 
diff --git a/Engine/Source/Runtime/VulkanRHI/Private/VulkanCommands.cpp b/Engine/Source/Runtime/VulkanRHI/Private/VulkanCommands.cpp
--- a/Engine/Source/Runtime/VulkanRHI/Private/VulkanCommands.cpp
+++ b/Engine/Source/Runtime/VulkanRHI/Private/VulkanCommands.cpp
@@ -616,7 +616,7 @@ void FVulkanCommandListContext::RHISetShaderUniformBuffer(FComputeShaderRHIParam
 #if VULKAN_ENABLE_AGGRESSIVE_STATS
 	SCOPE_CYCLE_COUNTER(STAT_VulkanSetUniformBufferTime);
 #endif
-	FVulkanComputePipelineState& State = *PendingComputeState->CurrentState;
+	FVulkanComputePipelineDescriptorState& State = *PendingComputeState->CurrentState;
 
 	// Walk through all resources to set all appropriate states
 	FVulkanComputeShader* Shader = ResourceCast(ComputeShaderRHI);
diff --git a/Engine/Source/Runtime/VulkanRHI/Private/VulkanDebug.cpp b/Engine/Source/Runtime/VulkanRHI/Private/VulkanDebug.cpp
--- a/Engine/Source/Runtime/VulkanRHI/Private/VulkanDebug.cpp
+++ b/Engine/Source/Runtime/VulkanRHI/Private/VulkanDebug.cpp
@@ -142,10 +142,7 @@ static VkBool32 VKAPI_PTR DebugReportFunction(
 static VkBool32 DebugUtilsCallback(VkDebugUtilsMessageSeverityFlagBitsEXT MsgSeverity, VkDebugUtilsMessageTypeFlagsEXT MsgType,
 	const VkDebugUtilsMessengerCallbackDataEXT* CallbackData, void* UserData)
 {
-	if (CallbackData->pMessageIdName)
-	{
-	}
-	else
+	if (!CallbackData->pMessageIdName)
 	{
 		if (MsgType == 2 && CallbackData->messageIdNumber == 5)
 		{
@@ -191,11 +188,13 @@ static VkBool32 DebugUtilsCallback(VkDebugUtilsMessageSeverityFlagBitsEXT MsgSev
 		Severity = TEXT("Verbose");
 	}
 
+	uint32 MsgBucket = 0;
 	const TCHAR* Type = TEXT("");
 	if (MsgType & VK_DEBUG_UTILS_MESSAGE_TYPE_GENERAL_BIT_EXT)
 	{
 		ensure((MsgType & ~VK_DEBUG_UTILS_MESSAGE_TYPE_GENERAL_BIT_EXT) == 0);
 		Type = TEXT(" General");
+		MsgBucket = 1;
 	}
 	else if (MsgType & VK_DEBUG_UTILS_MESSAGE_TYPE_VALIDATION_BIT_EXT)
 	{
@@ -203,21 +202,24 @@ static VkBool32 DebugUtilsCallback(VkDebugUtilsMessageSeverityFlagBitsEXT MsgSev
 		{
 			ensure((MsgType & ~(VK_DEBUG_UTILS_MESSAGE_TYPE_VALIDATION_BIT_EXT | VK_DEBUG_UTILS_MESSAGE_TYPE_PERFORMANCE_BIT_EXT)) == 0);
 			Type = TEXT("Perf/Validation");
+			MsgBucket = 2;
 		}
 		else
 		{
 			ensure((MsgType & ~VK_DEBUG_UTILS_MESSAGE_TYPE_VALIDATION_BIT_EXT) == 0);
 			Type = TEXT("Validation");
+			MsgBucket = 3;
 		}
 	}
 	else if (MsgType & VK_DEBUG_UTILS_MESSAGE_TYPE_PERFORMANCE_BIT_EXT)
 	{
 		ensure((MsgType & ~VK_DEBUG_UTILS_MESSAGE_TYPE_PERFORMANCE_BIT_EXT) == 0);
 		Type = TEXT("Perf");
+		MsgBucket = 4;
 	}
 
-	static TMap<const char*, TSet<int32>> SeenCodes;
-	if (GCVarUniqueValidationMessages->GetInt() == 0 || !SeenCodes.FindOrAdd(CallbackData->pMessageIdName).Contains(CallbackData->messageIdNumber))
+	static TStaticArray<TSet<int32>, 10> SeenCodes;
+	if (GCVarUniqueValidationMessages->GetInt() == 0 || !SeenCodes[MsgBucket].Contains(CallbackData->messageIdNumber))
 	{
 		if (CallbackData->pMessageIdName)
 		{
@@ -229,7 +231,7 @@ static VkBool32 DebugUtilsCallback(VkDebugUtilsMessageSeverityFlagBitsEXT MsgSev
 		}
 		if (GCVarUniqueValidationMessages->GetInt() == 1)
 		{
-			SeenCodes.Find(CallbackData->pMessageIdName)->Add(CallbackData->messageIdNumber);
+			SeenCodes[MsgBucket].Add(CallbackData->messageIdNumber);
 		}
 	}
 
diff --git a/Engine/Source/Runtime/VulkanRHI/Private/VulkanDescriptorSets.h b/Engine/Source/Runtime/VulkanRHI/Private/VulkanDescriptorSets.h
--- a/Engine/Source/Runtime/VulkanRHI/Private/VulkanDescriptorSets.h
+++ b/Engine/Source/Runtime/VulkanRHI/Private/VulkanDescriptorSets.h
@@ -107,7 +107,7 @@ protected:
 #endif
 	void AddDescriptor(int32 DescriptorSetIndex, const VkDescriptorSetLayoutBinding& Descriptor, int32 BindingIndex);
 
-	friend class FVulkanPipelineStateCache;
+	friend class FVulkanPipelineStateCacheManager;
 };
 
 // The actual run-time descriptor set layouts
@@ -443,7 +443,7 @@ protected:
 
 	friend class FVulkanComputePipeline;
 	friend class FVulkanGfxPipeline;
-	friend class FVulkanPipelineStateCache;
+	friend class FVulkanPipelineStateCacheManager;
 };
 
 #if !VULKAN_USE_DESCRIPTOR_POOL_MANAGER
@@ -496,8 +496,8 @@ protected:
 
 	FOLDVulkanDescriptorSets* RequestDescriptorSets(FVulkanCommandListContext* Context, FVulkanCmdBuffer* CmdBuffer, const FVulkanLayout& Layout);
 
-	friend class FVulkanComputePipelineState;
-	friend class FVulkanGfxPipelineState;
+	friend class FVulkanComputePipelineDescriptorState;
+	friend class FVulkanGraphicsPipelineDescriptorState;
 };
 #endif
 
@@ -631,6 +631,6 @@ protected:
 	uint32 SetupDescriptorWrites(const FNEWVulkanShaderDescriptorInfo& Info, VkWriteDescriptorSet* InWriteDescriptors, VkDescriptorImageInfo* InImageInfo, 
 		VkDescriptorBufferInfo* InBufferInfo, uint8* InBindingToDynamicOffsetMap);
 
-	friend class FVulkanComputePipelineState;
-	friend class FVulkanGfxPipelineState;
+	friend class FVulkanComputePipelineDescriptorState;
+	friend class FVulkanGraphicsPipelineDescriptorState;
 };
diff --git a/Engine/Source/Runtime/VulkanRHI/Private/VulkanDevice.cpp b/Engine/Source/Runtime/VulkanRHI/Private/VulkanDevice.cpp
--- a/Engine/Source/Runtime/VulkanRHI/Private/VulkanDevice.cpp
+++ b/Engine/Source/Runtime/VulkanRHI/Private/VulkanDevice.cpp
@@ -734,7 +734,7 @@ void FVulkanDevice::InitGPU(int32 DeviceIndex)
 	DescriptorPoolsManager = new FVulkanDescriptorPoolsManager();
 	DescriptorPoolsManager->Init(this);
 #endif
-	PipelineStateCache = new FVulkanPipelineStateCache(this);
+	PipelineStateCache = new FVulkanPipelineStateCacheManager(this);
 
 	TArray<FString> CacheFilenames;
 	FString StagedCacheDirectory = FPaths::ProjectDir() / TEXT("Build") / TEXT("ShaderCaches") / FPlatformProperties::IniPlatformName();
@@ -856,13 +856,13 @@ void FVulkanDevice::Destroy()
 
 	delete TimestampQueryPool;
 #else
-	for (FOLDVulkanQueryPool* QueryPool : OcclusionQueryPools)
+	for (FVulkanQueryPool* QueryPool : OcclusionQueryPools)
 	{
 		QueryPool->Destroy();
 		delete QueryPool;
 	}
 
-	for (FOLDVulkanQueryPool* QueryPool : TimestampQueryPools)
+	for (FVulkanQueryPool* QueryPool : TimestampQueryPools)
 	{
 		QueryPool->Destroy();
 		delete QueryPool;
@@ -999,7 +999,7 @@ void FVulkanDevice::SubmitCommandsAndFlushGPU()
 	//#todo-rco: Process other contexts first!
 }
 
-void FVulkanDevice::NotifyDeletedGfxPipeline(class FVulkanGraphicsPipelineState* Pipeline)
+void FVulkanDevice::NotifyDeletedGfxPipeline(class FVulkanRHIGraphicsPipelineState* Pipeline)
 {
 	if (ComputeContext != ImmediateContext)
 	{
diff --git a/Engine/Source/Runtime/VulkanRHI/Private/VulkanDevice.h b/Engine/Source/Runtime/VulkanRHI/Private/VulkanDevice.h
--- a/Engine/Source/Runtime/VulkanRHI/Private/VulkanDevice.h
+++ b/Engine/Source/Runtime/VulkanRHI/Private/VulkanDevice.h
@@ -14,7 +14,7 @@ class FVulkanDescriptorPoolsManager;
 #endif
 class FVulkanCommandListContextImmediate;
 class FVulkanOcclusionQueryPool;
-class FOLDVulkanQueryPool;
+class FVulkanQueryPool;
 
 class FVulkanDevice
 {
@@ -228,15 +228,15 @@ public:
 	void SubmitCommandsAndFlushGPU();
 
 #if VULKAN_USE_NEW_QUERIES
-	FVulkanOcclusionQueryPool* PrepareOcclusionQueryPool(uint32 NumQueries);
+	FVulkanOcclusionQueryPool* AcquireOcclusionQueryPool(uint32 NumQueries);
 	FVulkanTimestampQueryPool* PrepareTimestampQueryPool(bool& bOutRequiresReset);
 #else
-	inline FOLDVulkanBufferedQueryPool& FindAvailableQueryPool(TArray<FOLDVulkanBufferedQueryPool*>& Pools, VkQueryType QueryType)
+	inline FVulkanBufferedQueryPool& FindAvailableQueryPool(TArray<FVulkanBufferedQueryPool*>& Pools, VkQueryType QueryType)
 	{
 		// First try to find An available one
 		for (int32 Index = 0; Index < Pools.Num(); ++Index)
 		{
-			FOLDVulkanBufferedQueryPool* Pool = Pools[Index];
+			FVulkanBufferedQueryPool* Pool = Pools[Index];
 			if (Pool->HasRoom())
 			{
 				return *Pool;
@@ -244,26 +244,26 @@ public:
 		}
 
 		// None found, so allocate new Pool
-		FOLDVulkanBufferedQueryPool* Pool = new FOLDVulkanBufferedQueryPool(this, QueryType == VK_QUERY_TYPE_OCCLUSION ? NUM_OCCLUSION_QUERIES_PER_POOL : NUM_TIMESTAMP_QUERIES_PER_POOL, QueryType);
+		FVulkanBufferedQueryPool* Pool = new FVulkanBufferedQueryPool(this, QueryType == VK_QUERY_TYPE_OCCLUSION ? NUM_OCCLUSION_QUERIES_PER_POOL : NUM_TIMESTAMP_QUERIES_PER_POOL, QueryType);
 		Pools.Add(Pool);
 		return *Pool;
 	}
-	inline FOLDVulkanBufferedQueryPool& FindAvailableOcclusionQueryPool()
+	inline FVulkanBufferedQueryPool& FindAvailableOcclusionQueryPool()
 	{
 		return FindAvailableQueryPool(OcclusionQueryPools, VK_QUERY_TYPE_OCCLUSION);
 	}
 
-	inline FOLDVulkanBufferedQueryPool& FindAvailableTimestampQueryPool()
+	inline FVulkanBufferedQueryPool& FindAvailableTimestampQueryPool()
 	{
 		return FindAvailableQueryPool(TimestampQueryPools, VK_QUERY_TYPE_TIMESTAMP);
 	}
 #endif
-	inline class FVulkanPipelineStateCache* GetPipelineStateCache()
+	inline class FVulkanPipelineStateCacheManager* GetPipelineStateCache()
 	{
 		return PipelineStateCache;
 	}
 
-	void NotifyDeletedGfxPipeline(class FVulkanGraphicsPipelineState* Pipeline);
+	void NotifyDeletedGfxPipeline(class FVulkanRHIGraphicsPipelineState* Pipeline);
 	void NotifyDeletedComputePipeline(class FVulkanComputePipeline* Pipeline);
 
 	FVulkanCommandListContext* AcquireDeferredContext();
@@ -344,8 +344,8 @@ private:
 	TArray<FVulkanOcclusionQueryPool*> OcclusionQueryPools;
 	FVulkanTimestampQueryPool* TimestampQueryPool = nullptr;
 #else
-	TArray<FOLDVulkanBufferedQueryPool*> OcclusionQueryPools;
-	TArray<FOLDVulkanBufferedQueryPool*> TimestampQueryPools;
+	TArray<FVulkanBufferedQueryPool*> OcclusionQueryPools;
+	TArray<FVulkanBufferedQueryPool*> TimestampQueryPools;
 #endif
 
 	FVulkanQueue* GfxQueue;
@@ -398,6 +398,6 @@ private:
 	friend class FVulkanCommandListContext;
 #endif
 
-	class FVulkanPipelineStateCache* PipelineStateCache;
+	class FVulkanPipelineStateCacheManager* PipelineStateCache;
 	friend class FVulkanDynamicRHI;
 };
diff --git a/Engine/Source/Runtime/VulkanRHI/Private/VulkanPendingState.cpp b/Engine/Source/Runtime/VulkanRHI/Private/VulkanPendingState.cpp
--- a/Engine/Source/Runtime/VulkanRHI/Private/VulkanPendingState.cpp
+++ b/Engine/Source/Runtime/VulkanRHI/Private/VulkanPendingState.cpp
@@ -401,7 +401,7 @@ FVulkanPendingComputeState::~FVulkanPendingComputeState()
 {
 	for (auto& Pair : PipelineStates)
 	{
-		FVulkanComputePipelineState* State = Pair.Value;
+		FVulkanComputePipelineDescriptorState* State = Pair.Value;
 		delete State;
 	}
 }
@@ -499,7 +499,7 @@ FVulkanPendingGfxState::~FVulkanPendingGfxState()
 {
 	for (auto& Pair : PipelineStates)
 	{
-		FVulkanGfxPipelineState* State = Pair.Value;
+		FVulkanGraphicsPipelineDescriptorState* State = Pair.Value;
 		delete State;
 	}
 }
diff --git a/Engine/Source/Runtime/VulkanRHI/Private/VulkanPendingState.h b/Engine/Source/Runtime/VulkanRHI/Private/VulkanPendingState.h
--- a/Engine/Source/Runtime/VulkanRHI/Private/VulkanPendingState.h
+++ b/Engine/Source/Runtime/VulkanRHI/Private/VulkanPendingState.h
@@ -29,7 +29,7 @@ public:
 		if (InComputePipeline != CurrentPipeline)
 		{
 			CurrentPipeline = InComputePipeline;
-			FVulkanComputePipelineState** Found = PipelineStates.Find(InComputePipeline);
+			FVulkanComputePipelineDescriptorState** Found = PipelineStates.Find(InComputePipeline);
 			if (Found)
 			{
 				CurrentState = *Found;
@@ -37,7 +37,7 @@ public:
 			}
 			else
 			{
-				CurrentState = new FVulkanComputePipelineState(Device, InComputePipeline);
+				CurrentState = new FVulkanComputePipelineDescriptorState(Device, InComputePipeline);
 				PipelineStates.Add(CurrentPipeline, CurrentState);
 			}
 
@@ -90,9 +90,9 @@ protected:
 	TArray<FVulkanUnorderedAccessView*> UAVListForAutoFlush;
 
 	FVulkanComputePipeline* CurrentPipeline;
-	FVulkanComputePipelineState* CurrentState;
+	FVulkanComputePipelineDescriptorState* CurrentState;
 
-	TMap<FVulkanComputePipeline*, FVulkanComputePipelineState*> PipelineStates;
+	TMap<FVulkanComputePipeline*, FVulkanComputePipelineDescriptorState*> PipelineStates;
 
 	FVulkanCommandListContext& Context;
 
@@ -213,7 +213,7 @@ public:
 
 	void PrepareForDraw(FVulkanCmdBuffer* CmdBuffer);
 
-	bool SetGfxPipeline(FVulkanGraphicsPipelineState* InGfxPipeline)
+	bool SetGfxPipeline(FVulkanRHIGraphicsPipelineState* InGfxPipeline)
 	{
 		if (InGfxPipeline != CurrentPipeline)
 		{
@@ -230,7 +230,7 @@ public:
 			);
 
 			CurrentPipeline = InGfxPipeline;
-			FVulkanGfxPipelineState** Found = PipelineStates.Find(InGfxPipeline);
+			FVulkanGraphicsPipelineDescriptorState** Found = PipelineStates.Find(InGfxPipeline);
 			if (Found)
 			{
 				CurrentState = *Found;
@@ -238,7 +238,7 @@ public:
 			}
 			else
 			{
-				CurrentState = new FVulkanGfxPipelineState(Device, InGfxPipeline, CurrentBSS);
+				CurrentState = new FVulkanGraphicsPipelineDescriptorState(Device, InGfxPipeline, CurrentBSS);
 				PipelineStates.Add(CurrentPipeline, CurrentState);
 			}
 
@@ -262,7 +262,7 @@ public:
 		}
 	}
 
-	void NotifyDeletedPipeline(FVulkanGraphicsPipelineState* Pipeline)
+	void NotifyDeletedPipeline(FVulkanRHIGraphicsPipelineState* Pipeline)
 	{
 		PipelineStates.Remove(Pipeline);
 	}
@@ -279,11 +279,11 @@ protected:
 
 	bool bNeedToClear;
 
-	FVulkanGraphicsPipelineState* CurrentPipeline;
-	FVulkanGfxPipelineState* CurrentState;
+	FVulkanRHIGraphicsPipelineState* CurrentPipeline;
+	FVulkanGraphicsPipelineDescriptorState* CurrentState;
 	FVulkanBoundShaderState* CurrentBSS;
 
-	TMap<FVulkanGraphicsPipelineState*, FVulkanGfxPipelineState*> PipelineStates;
+	TMap<FVulkanRHIGraphicsPipelineState*, FVulkanGraphicsPipelineDescriptorState*> PipelineStates;
 
 	struct FVertexStream
 	{
diff --git a/Engine/Source/Runtime/VulkanRHI/Private/VulkanPipeline.cpp b/Engine/Source/Runtime/VulkanRHI/Private/VulkanPipeline.cpp
--- a/Engine/Source/Runtime/VulkanRHI/Private/VulkanPipeline.cpp
+++ b/Engine/Source/Runtime/VulkanRHI/Private/VulkanPipeline.cpp
@@ -87,7 +87,7 @@ void FVulkanGfxPipeline::CreateRuntimeObjects(const FGraphicsPipelineStateInitia
 	bRuntimeObjectsValid = true;
 }
 
-FVulkanGraphicsPipelineState::~FVulkanGraphicsPipelineState()
+FVulkanRHIGraphicsPipelineState::~FVulkanRHIGraphicsPipelineState()
 {
 	if (Pipeline)
 	{
@@ -113,25 +113,25 @@ static FAutoConsoleVariableRef GEnablePipelineCacheCompressionCvar(
 );
 
 
-FVulkanPipelineStateCache::FGfxPipelineEntry::~FGfxPipelineEntry()
+FVulkanPipelineStateCacheManager::FGfxPipelineEntry::~FGfxPipelineEntry()
 {
 	check(!bLoaded);
 	check(!RenderPass);
 }
 
-FVulkanPipelineStateCache::FComputePipelineEntry::~FComputePipelineEntry()
+FVulkanPipelineStateCacheManager::FComputePipelineEntry::~FComputePipelineEntry()
 {
 	check(!bLoaded);
 }
 
-FVulkanPipelineStateCache::FVulkanPipelineStateCache(FVulkanDevice* InDevice)
+FVulkanPipelineStateCacheManager::FVulkanPipelineStateCacheManager(FVulkanDevice* InDevice)
 	: Device(InDevice)
 	, PipelineCache(VK_NULL_HANDLE)
 {
 }
 
 
-FVulkanPipelineStateCache::~FVulkanPipelineStateCache()
+FVulkanPipelineStateCacheManager::~FVulkanPipelineStateCacheManager()
 {
 	DestroyCache();
 
@@ -145,7 +145,7 @@ FVulkanPipelineStateCache::~FVulkanPipelineStateCache()
 	PipelineCache = VK_NULL_HANDLE;
 }
 
-void FVulkanPipelineStateCache::Load(const TArray<FString>& CacheFilenames)
+void FVulkanPipelineStateCacheManager::Load(const TArray<FString>& CacheFilenames)
 {
 	// Try to load device cache first
 	for (const FString& CacheFilename : CacheFilenames)
@@ -305,7 +305,7 @@ void FVulkanPipelineStateCache::Load(const TArray<FString>& CacheFilenames)
 	}
 }
 
-void FVulkanPipelineStateCache::DestroyPipeline(FVulkanGfxPipeline* Pipeline)
+void FVulkanPipelineStateCacheManager::DestroyPipeline(FVulkanGfxPipeline* Pipeline)
 {
 	ensure(0);
 /*
@@ -317,7 +317,7 @@ void FVulkanPipelineStateCache::DestroyPipeline(FVulkanGfxPipeline* Pipeline)
 	}*/
 }
 
-void FVulkanPipelineStateCache::InitAndLoad(const TArray<FString>& CacheFilenames)
+void FVulkanPipelineStateCacheManager::InitAndLoad(const TArray<FString>& CacheFilenames)
 {
 	if (GEnablePipelineCacheLoadCvar.GetValueOnAnyThread() == 0)
 	{
@@ -337,7 +337,7 @@ void FVulkanPipelineStateCache::InitAndLoad(const TArray<FString>& CacheFilename
 	}
 }
 
-void FVulkanPipelineStateCache::Save(const FString& CacheFilename)
+void FVulkanPipelineStateCacheManager::Save(const FString& CacheFilename)
 {
 	FScopeLock Lock(&InitializerToPipelineMapCS);
 
@@ -387,7 +387,7 @@ void FVulkanPipelineStateCache::Save(const FString& CacheFilename)
 	}
 }
 
-FVulkanGraphicsPipelineState* FVulkanPipelineStateCache::CreateAndAdd(const FGraphicsPipelineStateInitializer& PSOInitializer, uint32 PSOInitializerHash, FGfxPipelineEntry* GfxEntry)
+FVulkanRHIGraphicsPipelineState* FVulkanPipelineStateCacheManager::CreateAndAdd(const FGraphicsPipelineStateInitializer& PSOInitializer, uint32 PSOInitializerHash, FGfxPipelineEntry* GfxEntry)
 {
 	FVulkanGfxPipeline* Pipeline = new FVulkanGfxPipeline(Device);
 
@@ -408,7 +408,7 @@ FVulkanGraphicsPipelineState* FVulkanPipelineStateCache::CreateAndAdd(const FGra
 		UE_LOG(LogVulkanRHI, Verbose, TEXT("Hitchy gfx pipeline (%.3f ms)"), (float)(Delta * 1000.0));
 	}
 
-	FVulkanGraphicsPipelineState* PipelineState = new FVulkanGraphicsPipelineState(PSOInitializer, Pipeline);
+	FVulkanRHIGraphicsPipelineState* PipelineState = new FVulkanRHIGraphicsPipelineState(PSOInitializer, Pipeline);
 	PipelineState->AddRef();
 
 	{
@@ -419,7 +419,7 @@ FVulkanGraphicsPipelineState* FVulkanPipelineStateCache::CreateAndAdd(const FGra
 	return PipelineState;
 }
 
-FArchive& operator << (FArchive& Ar, FVulkanPipelineStateCache::FGfxPipelineEntry::FBlendAttachment& Attachment)
+FArchive& operator << (FArchive& Ar, FVulkanPipelineStateCacheManager::FGfxPipelineEntry::FBlendAttachment& Attachment)
 {
 	// Modify VERSION if serialization changes
 	Ar << Attachment.bBlend;
@@ -433,7 +433,7 @@ FArchive& operator << (FArchive& Ar, FVulkanPipelineStateCache::FGfxPipelineEntr
 	return Ar;
 }
 
-void FVulkanPipelineStateCache::FGfxPipelineEntry::FBlendAttachment::ReadFrom(const VkPipelineColorBlendAttachmentState& InState)
+void FVulkanPipelineStateCacheManager::FGfxPipelineEntry::FBlendAttachment::ReadFrom(const VkPipelineColorBlendAttachmentState& InState)
 {
 	bBlend =				InState.blendEnable != VK_FALSE;
 	ColorBlendOp =			(uint8)InState.colorBlendOp;
@@ -445,7 +445,7 @@ void FVulkanPipelineStateCache::FGfxPipelineEntry::FBlendAttachment::ReadFrom(co
 	ColorWriteMask =		(uint8)InState.colorWriteMask;
 }
 
-void FVulkanPipelineStateCache::FGfxPipelineEntry::FBlendAttachment::WriteInto(VkPipelineColorBlendAttachmentState& Out) const
+void FVulkanPipelineStateCacheManager::FGfxPipelineEntry::FBlendAttachment::WriteInto(VkPipelineColorBlendAttachmentState& Out) const
 {
 	Out.blendEnable =			bBlend ? VK_TRUE : VK_FALSE;
 	Out.colorBlendOp =			(VkBlendOp)ColorBlendOp;
@@ -458,7 +458,7 @@ void FVulkanPipelineStateCache::FGfxPipelineEntry::FBlendAttachment::WriteInto(V
 }
 
 
-void FVulkanPipelineStateCache::FDescriptorSetLayoutBinding::ReadFrom(const VkDescriptorSetLayoutBinding& InState)
+void FVulkanPipelineStateCacheManager::FDescriptorSetLayoutBinding::ReadFrom(const VkDescriptorSetLayoutBinding& InState)
 {
 	Binding =			InState.binding;
 	ensure(InState.descriptorCount == 1);
@@ -467,7 +467,7 @@ void FVulkanPipelineStateCache::FDescriptorSetLayoutBinding::ReadFrom(const VkDe
 	StageFlags =		InState.stageFlags;
 }
 
-void FVulkanPipelineStateCache::FDescriptorSetLayoutBinding::WriteInto(VkDescriptorSetLayoutBinding& Out) const
+void FVulkanPipelineStateCacheManager::FDescriptorSetLayoutBinding::WriteInto(VkDescriptorSetLayoutBinding& Out) const
 {
 	Out.binding = Binding;
 	//Out.descriptorCount = DescriptorCount;
@@ -475,7 +475,7 @@ void FVulkanPipelineStateCache::FDescriptorSetLayoutBinding::WriteInto(VkDescrip
 	Out.stageFlags = StageFlags;
 }
 
-FArchive& operator << (FArchive& Ar, FVulkanPipelineStateCache::FDescriptorSetLayoutBinding& Binding)
+FArchive& operator << (FArchive& Ar, FVulkanPipelineStateCacheManager::FDescriptorSetLayoutBinding& Binding)
 {
 	// Modify VERSION if serialization changes
 	Ar << Binding.Binding;
@@ -485,21 +485,21 @@ FArchive& operator << (FArchive& Ar, FVulkanPipelineStateCache::FDescriptorSetLa
 	return Ar;
 }
 
-void FVulkanPipelineStateCache::FGfxPipelineEntry::FVertexBinding::ReadFrom(const VkVertexInputBindingDescription& InState)
+void FVulkanPipelineStateCacheManager::FGfxPipelineEntry::FVertexBinding::ReadFrom(const VkVertexInputBindingDescription& InState)
 {
 	Binding =	InState.binding;
 	InputRate =	(uint16)InState.inputRate;
 	Stride =	InState.stride;
 }
 
-void FVulkanPipelineStateCache::FGfxPipelineEntry::FVertexBinding::WriteInto(VkVertexInputBindingDescription& Out) const
+void FVulkanPipelineStateCacheManager::FGfxPipelineEntry::FVertexBinding::WriteInto(VkVertexInputBindingDescription& Out) const
 {
 	Out.binding =	Binding;
 	Out.inputRate =	(VkVertexInputRate)InputRate;
 	Out.stride =	Stride;
 }
 
-FArchive& operator << (FArchive& Ar, FVulkanPipelineStateCache::FGfxPipelineEntry::FVertexBinding& Binding)
+FArchive& operator << (FArchive& Ar, FVulkanPipelineStateCacheManager::FGfxPipelineEntry::FVertexBinding& Binding)
 {
 	// Modify VERSION if serialization changes
 	Ar << Binding.Stride;
@@ -508,7 +508,7 @@ FArchive& operator << (FArchive& Ar, FVulkanPipelineStateCache::FGfxPipelineEntr
 	return Ar;
 }
 
-void FVulkanPipelineStateCache::FGfxPipelineEntry::FVertexAttribute::ReadFrom(const VkVertexInputAttributeDescription& InState)
+void FVulkanPipelineStateCacheManager::FGfxPipelineEntry::FVertexAttribute::ReadFrom(const VkVertexInputAttributeDescription& InState)
 {
 	Binding =	InState.binding;
 	Format =	(uint32)InState.format;
@@ -516,7 +516,7 @@ void FVulkanPipelineStateCache::FGfxPipelineEntry::FVertexAttribute::ReadFrom(co
 	Offset =	InState.offset;
 }
 
-void FVulkanPipelineStateCache::FGfxPipelineEntry::FVertexAttribute::WriteInto(VkVertexInputAttributeDescription& Out) const
+void FVulkanPipelineStateCacheManager::FGfxPipelineEntry::FVertexAttribute::WriteInto(VkVertexInputAttributeDescription& Out) const
 {
 	Out.binding =	Binding;
 	Out.format =	(VkFormat)Format;
@@ -524,7 +524,7 @@ void FVulkanPipelineStateCache::FGfxPipelineEntry::FVertexAttribute::WriteInto(V
 	Out.offset =	Offset;
 }
 
-FArchive& operator << (FArchive& Ar, FVulkanPipelineStateCache::FGfxPipelineEntry::FVertexAttribute& Attribute)
+FArchive& operator << (FArchive& Ar, FVulkanPipelineStateCacheManager::FGfxPipelineEntry::FVertexAttribute& Attribute)
 {
 	// Modify VERSION if serialization changes
 	Ar << Attribute.Location;
@@ -534,7 +534,7 @@ FArchive& operator << (FArchive& Ar, FVulkanPipelineStateCache::FGfxPipelineEntr
 	return Ar;
 }
 
-void FVulkanPipelineStateCache::FGfxPipelineEntry::FRasterizer::ReadFrom(const VkPipelineRasterizationStateCreateInfo& InState)
+void FVulkanPipelineStateCacheManager::FGfxPipelineEntry::FRasterizer::ReadFrom(const VkPipelineRasterizationStateCreateInfo& InState)
 {
 	PolygonMode =				InState.polygonMode;
 	CullMode =					InState.cullMode;
@@ -542,7 +542,7 @@ void FVulkanPipelineStateCache::FGfxPipelineEntry::FRasterizer::ReadFrom(const V
 	DepthBiasConstantFactor =	InState.depthBiasConstantFactor;
 }
 
-void FVulkanPipelineStateCache::FGfxPipelineEntry::FRasterizer::WriteInto(VkPipelineRasterizationStateCreateInfo& Out) const
+void FVulkanPipelineStateCacheManager::FGfxPipelineEntry::FRasterizer::WriteInto(VkPipelineRasterizationStateCreateInfo& Out) const
 {
 	Out.polygonMode =				(VkPolygonMode)PolygonMode;
 	Out.cullMode =					(VkCullModeFlags)CullMode;
@@ -554,7 +554,7 @@ void FVulkanPipelineStateCache::FGfxPipelineEntry::FRasterizer::WriteInto(VkPipe
 	Out.depthBiasConstantFactor =	DepthBiasConstantFactor;
 }
 
-FArchive& operator << (FArchive& Ar, FVulkanPipelineStateCache::FGfxPipelineEntry::FRasterizer& Rasterizer)
+FArchive& operator << (FArchive& Ar, FVulkanPipelineStateCacheManager::FGfxPipelineEntry::FRasterizer& Rasterizer)
 {
 	// Modify VERSION if serialization changes
 	Ar << Rasterizer.PolygonMode;
@@ -564,7 +564,7 @@ FArchive& operator << (FArchive& Ar, FVulkanPipelineStateCache::FGfxPipelineEntr
 	return Ar;
 }
 
-void FVulkanPipelineStateCache::FGfxPipelineEntry::FDepthStencil::ReadFrom(const VkPipelineDepthStencilStateCreateInfo& InState)
+void FVulkanPipelineStateCacheManager::FGfxPipelineEntry::FDepthStencil::ReadFrom(const VkPipelineDepthStencilStateCreateInfo& InState)
 {
 	DepthCompareOp =			(uint8)InState.depthCompareOp;
 	bDepthTestEnable =			InState.depthTestEnable != VK_FALSE;
@@ -587,7 +587,7 @@ void FVulkanPipelineStateCache::FGfxPipelineEntry::FDepthStencil::ReadFrom(const
 	BackReference =				InState.back.reference;
 }
 
-void FVulkanPipelineStateCache::FGfxPipelineEntry::FDepthStencil::WriteInto(VkPipelineDepthStencilStateCreateInfo& Out) const
+void FVulkanPipelineStateCacheManager::FGfxPipelineEntry::FDepthStencil::WriteInto(VkPipelineDepthStencilStateCreateInfo& Out) const
 {
 	Out.depthCompareOp =		(VkCompareOp)DepthCompareOp;
 	Out.depthTestEnable =		bDepthTestEnable;
@@ -610,7 +610,7 @@ void FVulkanPipelineStateCache::FGfxPipelineEntry::FDepthStencil::WriteInto(VkPi
 	Out.back.reference =		BackReference;
 }
 
-FArchive& operator << (FArchive& Ar, FVulkanPipelineStateCache::FGfxPipelineEntry::FDepthStencil& DepthStencil)
+FArchive& operator << (FArchive& Ar, FVulkanPipelineStateCacheManager::FGfxPipelineEntry::FDepthStencil& DepthStencil)
 {
 	// Modify VERSION if serialization changes
 	Ar << DepthStencil.DepthCompareOp;
@@ -635,19 +635,19 @@ FArchive& operator << (FArchive& Ar, FVulkanPipelineStateCache::FGfxPipelineEntr
 	return Ar;
 }
 
-void FVulkanPipelineStateCache::FGfxPipelineEntry::FRenderTargets::FAttachmentRef::ReadFrom(const VkAttachmentReference& InState)
+void FVulkanPipelineStateCacheManager::FGfxPipelineEntry::FRenderTargets::FAttachmentRef::ReadFrom(const VkAttachmentReference& InState)
 {
 	Attachment =	InState.attachment;
 	Layout =		(uint64)InState.layout;
 }
 
-void FVulkanPipelineStateCache::FGfxPipelineEntry::FRenderTargets::FAttachmentRef::WriteInto(VkAttachmentReference& Out) const
+void FVulkanPipelineStateCacheManager::FGfxPipelineEntry::FRenderTargets::FAttachmentRef::WriteInto(VkAttachmentReference& Out) const
 {
 	Out.attachment =	Attachment;
 	Out.layout =		(VkImageLayout)Layout;
 }
 
-FArchive& operator << (FArchive& Ar, FVulkanPipelineStateCache::FGfxPipelineEntry::FRenderTargets::FAttachmentRef& AttachmentRef)
+FArchive& operator << (FArchive& Ar, FVulkanPipelineStateCacheManager::FGfxPipelineEntry::FRenderTargets::FAttachmentRef& AttachmentRef)
 {
 	// Modify VERSION if serialization changes
 	Ar << AttachmentRef.Attachment;
@@ -655,7 +655,7 @@ FArchive& operator << (FArchive& Ar, FVulkanPipelineStateCache::FGfxPipelineEntr
 	return Ar;
 }
 
-void FVulkanPipelineStateCache::FGfxPipelineEntry::FRenderTargets::FAttachmentDesc::ReadFrom(const VkAttachmentDescription &InState)
+void FVulkanPipelineStateCacheManager::FGfxPipelineEntry::FRenderTargets::FAttachmentDesc::ReadFrom(const VkAttachmentDescription &InState)
 {
 	Format =			(uint32)InState.format;
 	Flags =				(uint8)InState.flags;
@@ -668,7 +668,7 @@ void FVulkanPipelineStateCache::FGfxPipelineEntry::FRenderTargets::FAttachmentDe
 	FinalLayout =		(uint64)InState.finalLayout;
 }
 
-void FVulkanPipelineStateCache::FGfxPipelineEntry::FRenderTargets::FAttachmentDesc::WriteInto(VkAttachmentDescription& Out) const
+void FVulkanPipelineStateCacheManager::FGfxPipelineEntry::FRenderTargets::FAttachmentDesc::WriteInto(VkAttachmentDescription& Out) const
 {
 	Out.format =			(VkFormat)Format;
 	Out.flags =				Flags;
@@ -681,7 +681,7 @@ void FVulkanPipelineStateCache::FGfxPipelineEntry::FRenderTargets::FAttachmentDe
 	Out.finalLayout =		(VkImageLayout)FinalLayout;
 }
 
-FArchive& operator << (FArchive& Ar, FVulkanPipelineStateCache::FGfxPipelineEntry::FRenderTargets::FAttachmentDesc& AttachmentDesc)
+FArchive& operator << (FArchive& Ar, FVulkanPipelineStateCacheManager::FGfxPipelineEntry::FRenderTargets::FAttachmentDesc& AttachmentDesc)
 {
 	// Modify VERSION if serialization changes
 	Ar << AttachmentDesc.Format;
@@ -697,7 +697,7 @@ FArchive& operator << (FArchive& Ar, FVulkanPipelineStateCache::FGfxPipelineEntr
 	return Ar;
 }
 
-void FVulkanPipelineStateCache::FGfxPipelineEntry::FRenderTargets::ReadFrom(const FVulkanRenderTargetLayout& RTLayout)
+void FVulkanPipelineStateCacheManager::FGfxPipelineEntry::FRenderTargets::ReadFrom(const FVulkanRenderTargetLayout& RTLayout)
 {
 	NumAttachments =			RTLayout.NumAttachmentDescriptions;
 	NumColorAttachments =		RTLayout.NumColorAttachments;
@@ -732,7 +732,7 @@ void FVulkanPipelineStateCache::FGfxPipelineEntry::FRenderTargets::ReadFrom(cons
 	}
 }
 
-void FVulkanPipelineStateCache::FGfxPipelineEntry::FRenderTargets::WriteInto(FVulkanRenderTargetLayout& Out) const
+void FVulkanPipelineStateCacheManager::FGfxPipelineEntry::FRenderTargets::WriteInto(FVulkanRenderTargetLayout& Out) const
 {
 	Out.NumAttachmentDescriptions =	NumAttachments;
 	Out.NumColorAttachments =		NumColorAttachments;
@@ -765,7 +765,7 @@ void FVulkanPipelineStateCache::FGfxPipelineEntry::FRenderTargets::WriteInto(FVu
 	}
 }
 
-FArchive& operator << (FArchive& Ar, FVulkanPipelineStateCache::FGfxPipelineEntry::FRenderTargets& RTs)
+FArchive& operator << (FArchive& Ar, FVulkanPipelineStateCacheManager::FGfxPipelineEntry::FRenderTargets& RTs)
 {
 	// Modify VERSION if serialization changes
 	Ar << RTs.NumAttachments;
@@ -786,7 +786,7 @@ FArchive& operator << (FArchive& Ar, FVulkanPipelineStateCache::FGfxPipelineEntr
 	return Ar;
 }
 
-FArchive& operator << (FArchive& Ar, FVulkanPipelineStateCache::FGfxPipelineEntry& Entry)
+FArchive& operator << (FArchive& Ar, FVulkanPipelineStateCacheManager::FGfxPipelineEntry& Entry)
 {
 	// Modify VERSION if serialization changes
 	Ar << Entry.VertexInputKey;
@@ -813,12 +813,12 @@ FArchive& operator << (FArchive& Ar, FVulkanPipelineStateCache::FGfxPipelineEntr
 	return Ar;
 }
 
-FArchive& operator << (FArchive& Ar, FVulkanPipelineStateCache::FGfxPipelineEntry* Entry)
+FArchive& operator << (FArchive& Ar, FVulkanPipelineStateCacheManager::FGfxPipelineEntry* Entry)
 {
 	return Ar << (*Entry);
 }
 
-FArchive& operator << (FArchive& Ar, FVulkanPipelineStateCache::FComputePipelineEntry& Entry)
+FArchive& operator << (FArchive& Ar, FVulkanPipelineStateCacheManager::FComputePipelineEntry& Entry)
 {
 	// Modify VERSION if serialization changes
 	Ar << Entry.ShaderHash;
@@ -828,12 +828,12 @@ FArchive& operator << (FArchive& Ar, FVulkanPipelineStateCache::FComputePipeline
 	return Ar;
 }
 
-FArchive& operator << (FArchive& Ar, FVulkanPipelineStateCache::FComputePipelineEntry* Entry)
+FArchive& operator << (FArchive& Ar, FVulkanPipelineStateCacheManager::FComputePipelineEntry* Entry)
 {
 	return Ar << (*Entry);
 }
 
-uint32 FVulkanPipelineStateCache::FGfxPipelineEntry::GetEntryHash(uint32 Crc /*= 0*/)
+uint32 FVulkanPipelineStateCacheManager::FGfxPipelineEntry::GetEntryHash(uint32 Crc /*= 0*/)
 {
 	TArray<uint8> MemFile;
 	FMemoryWriter Ar(MemFile);
@@ -841,7 +841,7 @@ uint32 FVulkanPipelineStateCache::FGfxPipelineEntry::GetEntryHash(uint32 Crc /*=
 	return FCrc::MemCrc32(MemFile.GetData(), MemFile.GetTypeSize() * MemFile.Num(), Crc);
 }
 
-void FVulkanPipelineStateCache::CreateGfxPipelineFromEntry(const FGfxPipelineEntry* GfxEntry, FVulkanGfxPipeline* Pipeline)
+void FVulkanPipelineStateCacheManager::CreateGfxPipelineFromEntry(const FGfxPipelineEntry* GfxEntry, FVulkanGfxPipeline* Pipeline)
 {
 	// Pipeline
 	VkGraphicsPipelineCreateInfo PipelineInfo;
@@ -966,7 +966,7 @@ void FVulkanPipelineStateCache::CreateGfxPipelineFromEntry(const FGfxPipelineEnt
 	Pipeline->Layout = GfxEntry->Layout;
 }
 
-void FVulkanPipelineStateCache::CreatGfxEntryRuntimeObjects(FGfxPipelineEntry* GfxEntry)
+void FVulkanPipelineStateCacheManager::CreatGfxEntryRuntimeObjects(FGfxPipelineEntry* GfxEntry)
 {
 	{
 		// Descriptor Set Layouts
@@ -1014,7 +1014,7 @@ void FVulkanPipelineStateCache::CreatGfxEntryRuntimeObjects(FGfxPipelineEntry* G
 }
 
 
-void FVulkanPipelineStateCache::DestroyCache()
+void FVulkanPipelineStateCacheManager::DestroyCache()
 {
 	VkDevice DeviceHandle = Device->GetInstanceHandle();
 
@@ -1022,7 +1022,7 @@ void FVulkanPipelineStateCache::DestroyCache()
 	{
 		for (auto Pair : InitializerToPipelineMap)
 		{
-			FVulkanGraphicsPipelineState* Pipeline = Pair.Value;
+			FVulkanRHIGraphicsPipelineState* Pipeline = Pair.Value;
 			//when DestroyCache is called as part of r.Vulkan.RebuildPipelineCache, a pipeline can still be referenced by FVulkanPendingGfxState
 			ensure(GIsRHIInitialized || (!GIsRHIInitialized && Pipeline->GetRefCount() == 1));
 			Pipeline->Release();
@@ -1084,7 +1084,7 @@ void FVulkanPipelineStateCache::DestroyCache()
 	}
 }
 
-void FVulkanPipelineStateCache::RebuildCache()
+void FVulkanPipelineStateCacheManager::RebuildCache()
 {
 	UE_LOG(LogVulkanRHI, Warning, TEXT("Rebuilding pipeline cache; ditching %d entries"), GfxPipelineEntries.Num() + ComputePipelineEntries.Num());
 
@@ -1095,7 +1095,7 @@ void FVulkanPipelineStateCache::RebuildCache()
 	DestroyCache();
 }
 
-FVulkanPipelineStateCache::FShaderHashes::FShaderHashes(const FGraphicsPipelineStateInitializer& PSOInitializer)
+FVulkanPipelineStateCacheManager::FShaderHashes::FShaderHashes(const FGraphicsPipelineStateInitializer& PSOInitializer)
 {
 	Stages[DescriptorSet::Vertex] = GetShaderHash<FRHIVertexShader, FVulkanVertexShader>(PSOInitializer.BoundShaderState.VertexShaderRHI);
 	Stages[DescriptorSet::Pixel] = GetShaderHash<FRHIPixelShader, FVulkanPixelShader>(PSOInitializer.BoundShaderState.PixelShaderRHI);
@@ -1105,13 +1105,13 @@ FVulkanPipelineStateCache::FShaderHashes::FShaderHashes(const FGraphicsPipelineS
 	Finalize();
 }
 
-FVulkanPipelineStateCache::FShaderHashes::FShaderHashes()
+FVulkanPipelineStateCacheManager::FShaderHashes::FShaderHashes()
 {
 	FMemory::Memzero(Stages);
 	Hash = 0;
 }
 
-inline FVulkanLayout* FVulkanPipelineStateCache::FindOrAddLayout(const FVulkanDescriptorSetsLayoutInfo& DescriptorSetLayoutInfo)
+inline FVulkanLayout* FVulkanPipelineStateCacheManager::FindOrAddLayout(const FVulkanDescriptorSetsLayoutInfo& DescriptorSetLayoutInfo)
 {
 	FScopeLock Lock(&LayoutMapCS);
 	if (FVulkanLayout** FoundLayout = LayoutMap.Find(DescriptorSetLayoutInfo))
@@ -1127,7 +1127,7 @@ inline FVulkanLayout* FVulkanPipelineStateCache::FindOrAddLayout(const FVulkanDe
 	return Layout;
 }
 
-FVulkanLayout* FVulkanPipelineStateCache::GetOrGenerateGfxLayout(const FGraphicsPipelineStateInitializer& PSOInitializer,
+FVulkanLayout* FVulkanPipelineStateCacheManager::GetOrGenerateGfxLayout(const FGraphicsPipelineStateInitializer& PSOInitializer,
 	FVulkanShader** OutShaders, FVulkanVertexInputStateInfo& OutVertexInputState)
 {
 	const FBoundShaderStateInput& BSI = PSOInitializer.BoundShaderState;
@@ -1192,7 +1192,7 @@ FVulkanLayout* FVulkanPipelineStateCache::GetOrGenerateGfxLayout(const FGraphics
 	return FindOrAddLayout(DescriptorSetLayoutInfo);
 }
 
-FVulkanPipelineStateCache::FGfxPipelineEntry* FVulkanPipelineStateCache::CreateGfxEntry(const FGraphicsPipelineStateInitializer& PSOInitializer)
+FVulkanPipelineStateCacheManager::FGfxPipelineEntry* FVulkanPipelineStateCacheManager::CreateGfxEntry(const FGraphicsPipelineStateInitializer& PSOInitializer)
 {
 	FGfxPipelineEntry* OutGfxEntry = new FGfxPipelineEntry();
 
@@ -1276,7 +1276,7 @@ FVulkanPipelineStateCache::FGfxPipelineEntry* FVulkanPipelineStateCache::CreateG
 }
 
 
-FVulkanGraphicsPipelineState* FVulkanPipelineStateCache::FindInLoadedLibrary(const FGraphicsPipelineStateInitializer& PSOInitializer, uint32 PSOInitializerHash, const FShaderHashes& ShaderHashes, FGfxPipelineEntry*& OutGfxEntry)
+FVulkanRHIGraphicsPipelineState* FVulkanPipelineStateCacheManager::FindInLoadedLibrary(const FGraphicsPipelineStateInitializer& PSOInitializer, uint32 PSOInitializerHash, const FShaderHashes& ShaderHashes, FGfxPipelineEntry*& OutGfxEntry)
 {
 	FScopeLock Lock(&ShaderHashToGfxEntriesMapCS);
 	OutGfxEntry = nullptr;
@@ -1297,7 +1297,7 @@ FVulkanGraphicsPipelineState* FVulkanPipelineStateCache::FindInLoadedLibrary(con
 		{
 			(*FoundPipeline)->CreateRuntimeObjects(PSOInitializer);
 		}
-		FVulkanGraphicsPipelineState* PipelineState = new FVulkanGraphicsPipelineState(PSOInitializer, *FoundPipeline);
+		FVulkanRHIGraphicsPipelineState* PipelineState = new FVulkanRHIGraphicsPipelineState(PSOInitializer, *FoundPipeline);
 		{
 			FScopeLock Lock2(&InitializerToPipelineMapCS);
 			InitializerToPipelineMap.Add(PSOInitializerHash, PipelineState);
@@ -1328,17 +1328,17 @@ FGraphicsPipelineStateRHIRef FVulkanDynamicRHI::RHICreateGraphicsPipelineState(c
 
 	// First try the hash based off runtime objects
 	uint32 PSOInitializerHash = 0;
-	FVulkanGraphicsPipelineState* Found = Device->PipelineStateCache->FindInRuntimeCache(PSOInitializer, PSOInitializerHash);
+	FVulkanRHIGraphicsPipelineState* Found = Device->PipelineStateCache->FindInRuntimeCache(PSOInitializer, PSOInitializerHash);
 	if (Found)
 	{
 		ensure(FMemory::Memcmp(&Found->PipelineStateInitializer, &PSOInitializer, sizeof(PSOInitializer)) == 0);
 		return Found;
 	}
 
-	FVulkanPipelineStateCache::FShaderHashes ShaderHashes(PSOInitializer);
+	FVulkanPipelineStateCacheManager::FShaderHashes ShaderHashes(PSOInitializer);
 
 	// Now try the loaded cache from disk
-	FVulkanPipelineStateCache::FGfxPipelineEntry* GfxEntry = nullptr;
+	FVulkanPipelineStateCacheManager::FGfxPipelineEntry* GfxEntry = nullptr;
 	Found = Device->PipelineStateCache->FindInLoadedLibrary(PSOInitializer, PSOInitializerHash, ShaderHashes, GfxEntry);
 	if (Found)
 	{
@@ -1351,11 +1351,11 @@ FGraphicsPipelineStateRHIRef FVulkanDynamicRHI::RHICreateGraphicsPipelineState(c
 	FVulkanRenderPass* RenderPass = Device->GetImmediateContext().PrepareRenderPassForPSOCreation(PSOInitializer);
 
 	// have we made a matching state object yet?
-	FVulkanGraphicsPipelineState* PipelineState = Device->GetPipelineStateCache()->CreateAndAdd(PSOInitializer, PSOInitializerHash, GfxEntry);
+	FVulkanRHIGraphicsPipelineState* PipelineState = Device->GetPipelineStateCache()->CreateAndAdd(PSOInitializer, PSOInitializerHash, GfxEntry);
 	return PipelineState;
 }
 
-FVulkanComputePipeline* FVulkanPipelineStateCache::GetOrCreateComputePipeline(FVulkanComputeShader* ComputeShader)
+FVulkanComputePipeline* FVulkanPipelineStateCacheManager::GetOrCreateComputePipeline(FVulkanComputeShader* ComputeShader)
 {
 	FScopeLock ScopeLock(&CreateComputePipelineCS);
 
@@ -1402,7 +1402,7 @@ FVulkanComputePipeline* FVulkanPipelineStateCache::GetOrCreateComputePipeline(FV
 	return ComputePipeline;
 }
 
-void FVulkanPipelineStateCache::FComputePipelineEntry::CalculateEntryHash()
+void FVulkanPipelineStateCacheManager::FComputePipelineEntry::CalculateEntryHash()
 {
 	TArray<uint8> MemFile;
 	FMemoryWriter Ar(MemFile);
@@ -1411,7 +1411,7 @@ void FVulkanPipelineStateCache::FComputePipelineEntry::CalculateEntryHash()
 	EntryHash = FCrc::MemCrc32(&ShaderHash, sizeof(ShaderHash), EntryHash);
 }
 
-FVulkanPipelineStateCache::FComputePipelineEntry* FVulkanPipelineStateCache::CreateComputeEntry(const FVulkanComputeShader* ComputeShader)
+FVulkanPipelineStateCacheManager::FComputePipelineEntry* FVulkanPipelineStateCacheManager::CreateComputeEntry(const FVulkanComputeShader* ComputeShader)
 {
 	FComputePipelineEntry* OutComputeEntry = new FComputePipelineEntry();
 
@@ -1443,7 +1443,7 @@ FVulkanPipelineStateCache::FComputePipelineEntry* FVulkanPipelineStateCache::Cre
 	return OutComputeEntry;
 }
 
-FVulkanComputePipeline* FVulkanPipelineStateCache::CreateComputePipelineFromEntry(const FComputePipelineEntry* ComputeEntry)
+FVulkanComputePipeline* FVulkanPipelineStateCacheManager::CreateComputePipelineFromEntry(const FComputePipelineEntry* ComputeEntry)
 {
 	FVulkanComputePipeline* Pipeline = new FVulkanComputePipeline(Device);
 
@@ -1462,7 +1462,7 @@ FVulkanComputePipeline* FVulkanPipelineStateCache::CreateComputePipelineFromEntr
 	return Pipeline;
 }
 
-void FVulkanPipelineStateCache::CreateComputeEntryRuntimeObjects(FComputePipelineEntry* ComputeEntry)
+void FVulkanPipelineStateCacheManager::CreateComputeEntryRuntimeObjects(FComputePipelineEntry* ComputeEntry)
 {
 	{
 		// Descriptor Set Layouts
@@ -1523,7 +1523,7 @@ inline void SerializeArray(FArchive& Ar, TArray<T*>& Array)
 	}
 }
 
-void FVulkanPipelineStateCache::FVulkanPipelineStateCacheFile::Save(FArchive& Ar)
+void FVulkanPipelineStateCacheManager::FVulkanPipelineStateCacheFile::Save(FArchive& Ar)
 {
 	check(ShaderCache);
 
@@ -1558,7 +1558,7 @@ void FVulkanPipelineStateCache::FVulkanPipelineStateCacheFile::Save(FArchive& Ar
 }
 
 
-bool FVulkanPipelineStateCache::FVulkanPipelineStateCacheFile::Load(FArchive& Ar, const TCHAR* Filename)
+bool FVulkanPipelineStateCacheManager::FVulkanPipelineStateCacheFile::Load(FArchive& Ar, const TCHAR* Filename)
 {
 	check(ShaderCache);
 
@@ -1614,7 +1614,7 @@ bool FVulkanPipelineStateCache::FVulkanPipelineStateCacheFile::Load(FArchive& Ar
 	return true;
 }
 
-bool FVulkanPipelineStateCache::FVulkanPipelineStateCacheFile::BinaryCacheMatches(FVulkanDevice* InDevice, const TArray<uint8>& DeviceCache)
+bool FVulkanPipelineStateCacheManager::FVulkanPipelineStateCacheFile::BinaryCacheMatches(FVulkanDevice* InDevice, const TArray<uint8>& DeviceCache)
 {
 	if (DeviceCache.Num() > 4)
 	{
diff --git a/Engine/Source/Runtime/VulkanRHI/Private/VulkanPipeline.h b/Engine/Source/Runtime/VulkanRHI/Private/VulkanPipeline.h
--- a/Engine/Source/Runtime/VulkanRHI/Private/VulkanPipeline.h
+++ b/Engine/Source/Runtime/VulkanRHI/Private/VulkanPipeline.h
@@ -32,10 +32,10 @@ public:
 protected:
 	FVulkanDevice* Device;
 	VkPipeline Pipeline;
-	FVulkanLayout* Layout; /*owned by FVulkanPipelineStateCache, do not delete yourself !*/
+	FVulkanLayout* Layout; /*owned by FVulkanPipelineStateCacheManager, do not delete yourself !*/
 
-	friend class FVulkanPipelineStateCache;
-	friend class FVulkanGraphicsPipelineState;
+	friend class FVulkanPipelineStateCacheManager;
+	friend class FVulkanRHIGraphicsPipelineState;
 };
 
 class FVulkanComputePipeline : public FVulkanPipeline, public FRHIComputePipelineState
@@ -62,7 +62,7 @@ public:
 protected:
 	FVulkanComputeShader* ComputeShader;
 
-	friend class FVulkanPipelineStateCache;
+	friend class FVulkanPipelineStateCacheManager;
 };
 
 class FVulkanGfxPipeline : public FVulkanPipeline, public FRHIResource
@@ -93,30 +93,30 @@ private:
 	FVulkanVertexInputStateInfo VertexInputState;
 };
 
-class FVulkanGraphicsPipelineState : public FRHIGraphicsPipelineState
+class FVulkanRHIGraphicsPipelineState : public FRHIGraphicsPipelineState
 {
 public:
-	FVulkanGraphicsPipelineState(const FGraphicsPipelineStateInitializer& Initializer, FVulkanGfxPipeline* InPipeline)
+	FVulkanRHIGraphicsPipelineState(const FGraphicsPipelineStateInitializer& Initializer, FVulkanGfxPipeline* InPipeline)
 		: PipelineStateInitializer(Initializer), Pipeline(InPipeline)
 	{}
 
-	~FVulkanGraphicsPipelineState();
+	~FVulkanRHIGraphicsPipelineState();
 
 	FGraphicsPipelineStateInitializer PipelineStateInitializer;
 	TRefCountPtr<FVulkanGfxPipeline> Pipeline;
 };
 
-class FVulkanPipelineStateCache
+class FVulkanPipelineStateCacheManager
 {
 public:
-	inline FVulkanGraphicsPipelineState* FindInRuntimeCache(const FGraphicsPipelineStateInitializer& Initializer, uint32& OutHash)
+	inline FVulkanRHIGraphicsPipelineState* FindInRuntimeCache(const FGraphicsPipelineStateInitializer& Initializer, uint32& OutHash)
 	{
 		//#todo-rco: Improve hash: Don't use BSS pointers; skip load/store actions?? etc
 		OutHash = FCrc::MemCrc32(&Initializer, sizeof(Initializer));
 		
 		{
 			FScopeLock Lock(&InitializerToPipelineMapCS);
-			FVulkanGraphicsPipelineState** Found = InitializerToPipelineMap.Find(OutHash);
+			FVulkanRHIGraphicsPipelineState** Found = InitializerToPipelineMap.Find(OutHash);
 			if (Found)
 			{
 				return *Found;
@@ -132,8 +132,8 @@ public:
 	void InitAndLoad(const TArray<FString>& CacheFilenames);
 	void Save(const FString& CacheFilename);
 
-	FVulkanPipelineStateCache(FVulkanDevice* InParent);
-	~FVulkanPipelineStateCache();
+	FVulkanPipelineStateCacheManager(FVulkanDevice* InParent);
+	~FVulkanPipelineStateCacheManager();
 
 	void RebuildCache();
 
@@ -540,7 +540,7 @@ private:
 	FVulkanDevice* Device;
 
 	// Key is a hash of the PSO, which is based off shader pointers
-	TMap<uint32, FVulkanGraphicsPipelineState*> InitializerToPipelineMap;
+	TMap<uint32, FVulkanRHIGraphicsPipelineState*> InitializerToPipelineMap;
 	FCriticalSection InitializerToPipelineMapCS;
 
 	TMap<FVulkanComputeShader*, FVulkanComputePipeline*> ComputeShaderToPipelineMap;
@@ -555,7 +555,7 @@ private:
 
 	FShaderUCodeCache ShaderCache;
 
-	FVulkanGraphicsPipelineState* CreateAndAdd(const FGraphicsPipelineStateInitializer& PSOInitializer, uint32 PSOInitializerHash, FGfxPipelineEntry* GfxEntry);
+	FVulkanRHIGraphicsPipelineState* CreateAndAdd(const FGraphicsPipelineStateInitializer& PSOInitializer, uint32 PSOInitializerHash, FGfxPipelineEntry* GfxEntry);
 	void CreateGfxPipelineFromEntry(const FGfxPipelineEntry* GfxEntry, FVulkanGfxPipeline* Pipeline);
 	FGfxPipelineEntry* CreateGfxEntry(const FGraphicsPipelineStateInitializer& PSOInitializer);
 	void CreatGfxEntryRuntimeObjects(FGfxPipelineEntry* GfxEntry);
@@ -605,7 +605,7 @@ private:
 	TMap<FVulkanDescriptorSetsLayoutInfo, FVulkanLayout*> LayoutMap;
 	FCriticalSection LayoutMapCS;
 
-	FVulkanGraphicsPipelineState* FindInLoadedLibrary(const FGraphicsPipelineStateInitializer& PSOInitializer, uint32 PSOInitializerHash, const FShaderHashes& ShaderHashes, FGfxPipelineEntry*& outGfxEntry);
+	FVulkanRHIGraphicsPipelineState* FindInLoadedLibrary(const FGraphicsPipelineStateInitializer& PSOInitializer, uint32 PSOInitializerHash, const FShaderHashes& ShaderHashes, FGfxPipelineEntry*& outGfxEntry);
 	
 	friend class FVulkanDynamicRHI;
 
@@ -649,5 +649,5 @@ struct TVulkanResourceTraits<class FRHIComputePipelineState>
 template<>
 struct TVulkanResourceTraits<FRHIGraphicsPipelineState>
 {
-	typedef FVulkanGraphicsPipelineState TConcreteType;
+	typedef FVulkanRHIGraphicsPipelineState TConcreteType;
 };
\ No newline at end of file
diff --git a/Engine/Source/Runtime/VulkanRHI/Private/VulkanPipelineState.cpp b/Engine/Source/Runtime/VulkanRHI/Private/VulkanPipelineState.cpp
--- a/Engine/Source/Runtime/VulkanRHI/Private/VulkanPipelineState.cpp
+++ b/Engine/Source/Runtime/VulkanRHI/Private/VulkanPipelineState.cpp
@@ -30,8 +30,8 @@ static TAutoConsoleVariable<int32> GAlwaysWriteDS(
 );
 
 
-FVulkanComputePipelineState::FVulkanComputePipelineState(FVulkanDevice* InDevice, FVulkanComputePipeline* InComputePipeline)
-	: FVulkanCommonPipelineState(InDevice)
+FVulkanComputePipelineDescriptorState::FVulkanComputePipelineDescriptorState(FVulkanDevice* InDevice, FVulkanComputePipeline* InComputePipeline)
+	: FVulkanCommonPipelineDescriptorState(InDevice)
 	, PackedUniformBuffersMask(0)
 	, PackedUniformBuffersDirty(0)
 	, ComputePipeline(InComputePipeline)
@@ -47,7 +47,7 @@ FVulkanComputePipelineState::FVulkanComputePipelineState(FVulkanDevice* InDevice
 	InComputePipeline->AddRef();
 }
 
-void FVulkanComputePipelineState::CreateDescriptorWriteInfos()
+void FVulkanComputePipelineDescriptorState::CreateDescriptorWriteInfos()
 {
 	check(DSWriteContainer.DescriptorWrites.Num() == 0);
 
@@ -84,7 +84,7 @@ void FVulkanComputePipelineState::CreateDescriptorWriteInfos()
 #endif
 }
 
-bool FVulkanComputePipelineState::UpdateDescriptorSets(FVulkanCommandListContext* CmdListContext, FVulkanCmdBuffer* CmdBuffer)
+bool FVulkanComputePipelineDescriptorState::UpdateDescriptorSets(FVulkanCommandListContext* CmdListContext, FVulkanCmdBuffer* CmdBuffer)
 {
 #if VULKAN_ENABLE_AGGRESSIVE_STATS
 	SCOPE_CYCLE_COUNTER(STAT_VulkanUpdateDescriptorSets);
@@ -159,8 +159,8 @@ bool FVulkanComputePipelineState::UpdateDescriptorSets(FVulkanCommandListContext
 }
 
 
-FVulkanGfxPipelineState::FVulkanGfxPipelineState(FVulkanDevice* InDevice, FVulkanGraphicsPipelineState* InGfxPipeline, FVulkanBoundShaderState* InBSS)
-	: FVulkanCommonPipelineState(InDevice)
+FVulkanGraphicsPipelineDescriptorState::FVulkanGraphicsPipelineDescriptorState(FVulkanDevice* InDevice, FVulkanRHIGraphicsPipelineState* InGfxPipeline, FVulkanBoundShaderState* InBSS)
+	: FVulkanCommonPipelineDescriptorState(InDevice)
 	, GfxPipeline(InGfxPipeline)
 	, BSS(InBSS)
 {
@@ -222,7 +222,7 @@ FVulkanGfxPipelineState::FVulkanGfxPipelineState(FVulkanDevice* InDevice, FVulka
 	BSS->AddRef();
 }
 
-void FVulkanGfxPipelineState::CreateDescriptorWriteInfos()
+void FVulkanGraphicsPipelineDescriptorState::CreateDescriptorWriteInfos()
 {
 	check(DSWriteContainer.DescriptorWrites.Num() == 0);
 
@@ -291,7 +291,7 @@ void FVulkanGfxPipelineState::CreateDescriptorWriteInfos()
 #endif
 }
 
-bool FVulkanGfxPipelineState::UpdateDescriptorSets(FVulkanCommandListContext* CmdListContext, FVulkanCmdBuffer* CmdBuffer)
+bool FVulkanGraphicsPipelineDescriptorState::UpdateDescriptorSets(FVulkanCommandListContext* CmdListContext, FVulkanCmdBuffer* CmdBuffer)
 {
 #if VULKAN_ENABLE_AGGRESSIVE_STATS
 	SCOPE_CYCLE_COUNTER(STAT_VulkanUpdateDescriptorSets);
@@ -414,7 +414,7 @@ bool FVulkanGfxPipelineState::UpdateDescriptorSets(FVulkanCommandListContext* Cm
 
 void FVulkanCommandListContext::RHISetGraphicsPipelineState(FGraphicsPipelineStateRHIParamRef GraphicsState)
 {
-	FVulkanGraphicsPipelineState* Pipeline = ResourceCast(GraphicsState);
+	FVulkanRHIGraphicsPipelineState* Pipeline = ResourceCast(GraphicsState);
 	
 	FVulkanCmdBuffer* CmdBuffer = CommandBufferManager->GetActiveCmdBuffer();
 	if (PendingGfxState->SetGfxPipeline(Pipeline) || !CmdBuffer->bHasPipeline)
diff --git a/Engine/Source/Runtime/VulkanRHI/Private/VulkanPipelineState.h b/Engine/Source/Runtime/VulkanRHI/Private/VulkanPipelineState.h
--- a/Engine/Source/Runtime/VulkanRHI/Private/VulkanPipelineState.h
+++ b/Engine/Source/Runtime/VulkanRHI/Private/VulkanPipelineState.h
@@ -18,10 +18,10 @@
 class FVulkanComputePipeline;
 
 // Common Pipeline state
-class FVulkanCommonPipelineState : public VulkanRHI::FDeviceChild
+class FVulkanCommonPipelineDescriptorState : public VulkanRHI::FDeviceChild
 {
 public:
-	FVulkanCommonPipelineState(FVulkanDevice* InDevice)
+	FVulkanCommonPipelineDescriptorState(FVulkanDevice* InDevice)
 		: VulkanRHI::FDeviceChild(InDevice)
 #if !VULKAN_USE_DESCRIPTOR_POOL_MANAGER
 		, DSRingBuffer(InDevice)
@@ -72,11 +72,11 @@ protected:
 };
 
 
-class FVulkanComputePipelineState : public FVulkanCommonPipelineState
+class FVulkanComputePipelineDescriptorState : public FVulkanCommonPipelineDescriptorState
 {
 public:
-	FVulkanComputePipelineState(FVulkanDevice* InDevice, FVulkanComputePipeline* InComputePipeline);
-	~FVulkanComputePipelineState()
+	FVulkanComputePipelineDescriptorState(FVulkanDevice* InDevice, FVulkanComputePipeline* InComputePipeline);
+	~FVulkanComputePipelineDescriptorState()
 	{
 		ComputePipeline->Release();
 	}
@@ -183,11 +183,11 @@ protected:
 	friend class FVulkanCommandListContext;
 };
 
-class FVulkanGfxPipelineState : public FVulkanCommonPipelineState
+class FVulkanGraphicsPipelineDescriptorState : public FVulkanCommonPipelineDescriptorState
 {
 public:
-	FVulkanGfxPipelineState(FVulkanDevice* InDevice, FVulkanGraphicsPipelineState* InGfxPipeline, FVulkanBoundShaderState* InBSS);
-	~FVulkanGfxPipelineState()
+	FVulkanGraphicsPipelineDescriptorState(FVulkanDevice* InDevice, FVulkanRHIGraphicsPipelineState* InGfxPipeline, FVulkanBoundShaderState* InBSS);
+	~FVulkanGraphicsPipelineDescriptorState()
 	{
 		GfxPipeline->Release();
 		BSS->Release();
@@ -347,7 +347,7 @@ protected:
 	uint64 UniformBuffersWithDataMask[DescriptorSet::NumGfxStages];
 	FVulkanDescriptorSetWriter DSWriter[DescriptorSet::EStage::NumGfxStages];
 
-	FVulkanGraphicsPipelineState* GfxPipeline;
+	FVulkanRHIGraphicsPipelineState* GfxPipeline;
 	FVulkanBoundShaderState* BSS;
 	int32 ID;
 
diff --git a/Engine/Source/Runtime/VulkanRHI/Private/VulkanQuery.cpp b/Engine/Source/Runtime/VulkanRHI/Private/VulkanQuery.cpp
--- a/Engine/Source/Runtime/VulkanRHI/Private/VulkanQuery.cpp
+++ b/Engine/Source/Runtime/VulkanRHI/Private/VulkanQuery.cpp
@@ -41,119 +41,12 @@
 	ECVF_RenderThreadSafe
 );
 
+FCriticalSection GQueryLock;
+
 #if VULKAN_USE_NEW_QUERIES
-FCriticalSection GOcclusionQueryCS;
 const uint32 GMinNumberOfQueriesInPool = 256;
 const uint32 GMaxLifetimeTimestampQueries = 1024;
 
-FVulkanRenderQuery::FVulkanRenderQuery(FVulkanDevice* Device, ERenderQueryType InQueryType)
-	: QueryType(InQueryType)
-{
-	INC_DWORD_STAT(STAT_VulkanNumQueries);
-}
 
-FVulkanRenderQuery::~FVulkanRenderQuery()
-{
-	DEC_DWORD_STAT(STAT_VulkanNumQueries);
-}
-
-inline void FVulkanRenderQuery::Begin(FVulkanCmdBuffer* InCmdBuffer)
-{
-	BeginCmdBuffer = InCmdBuffer;
-	check(State == EState::Reset);
-	if (QueryType == RQT_Occlusion)
-	{
-		check(Pool);
-		VulkanRHI::vkCmdBeginQuery(InCmdBuffer->GetHandle(), Pool->GetHandle(), QueryIndex, VK_QUERY_CONTROL_PRECISE_BIT);
-		State = EState::InBegin;
-		LastPoolReset = Pool->NumResets;
-	}
-	else
-	{
-		ensureMsgf(0, TEXT("Timer queries do NOT support Begin()!"));
-	}
-}
-
-inline void FVulkanRenderQuery::End(FVulkanCmdBuffer* InCmdBuffer)
-{
-	ensure(QueryType != RQT_Occlusion || BeginCmdBuffer == InCmdBuffer);
-	if (QueryType == RQT_Occlusion)
-	{
-		check(State == EState::InBegin);
-		check(Pool);
-		VulkanRHI::vkCmdEndQuery(InCmdBuffer->GetHandle(), Pool->GetHandle(), QueryIndex);
-		check(LastPoolReset == Pool->NumResets);
-	}
-	else
-	{
-		check(State == EState::Reset);
-		BeginCmdBuffer = InCmdBuffer;
-		VulkanRHI::vkCmdWriteTimestamp(InCmdBuffer->GetHandle(), VK_PIPELINE_STAGE_BOTTOM_OF_PIPE_BIT, Pool->GetHandle(), QueryIndex);
-		LastPoolReset = Pool->NumResets;
-	}
-	State = EState::InEnd;
-}
-
-bool FVulkanRenderQuery::GetResult(FVulkanDevice* Device, uint64& OutResult, bool bWait)
-{
-FScopeLock ScopeLock(&GOcclusionQueryCS);
-	if (State == EState::HasResults)
-	{
-		OutResult = Result;
-		return true;
-	}
-
-	check(Pool);
-	check(State == EState::InEnd);
-	if (QueryType == RQT_Occlusion)
-	{
-		check(IsInRenderingThread());
-		if (((FVulkanOcclusionQueryPool*)Pool)->GetResults(QueryIndex, bWait, Result))
-		{
-			check(LastPoolReset == Pool->NumResets);
-			State = EState::HasResults;
-			OutResult = Result;
-			return true;
-		}
-	}
-	else
-	{
-		check(IsInRenderingThread() || IsInRHIThread());
-		if (((FVulkanTimestampQueryPool*)Pool)->GetResults(QueryIndex, bWait, Result))
-		{
-			check(LastPoolReset == Pool->NumResets);
-			State = EState::HasResults;
-			OutResult = Result;
-			return true;
-		}
-	}
-	return false;
-}
-
-FVulkanQueryPool::FVulkanQueryPool(FVulkanDevice* InDevice, uint32 InMaxQueries, VkQueryType InQueryType)
-	: VulkanRHI::FDeviceChild(InDevice)
-	, QueryPool(VK_NULL_HANDLE)
-	, MaxQueries(InMaxQueries)
-	, QueryType(InQueryType)
-{
-	VkQueryPoolCreateInfo PoolCreateInfo;
-	FMemory::Memzero(PoolCreateInfo);
-	PoolCreateInfo.sType = VK_STRUCTURE_TYPE_QUERY_POOL_CREATE_INFO;
-	PoolCreateInfo.queryType = QueryType;
-	PoolCreateInfo.queryCount = MaxQueries;
-
-	VERIFYVULKANRESULT(VulkanRHI::vkCreateQueryPool(Device->GetInstanceHandle(), &PoolCreateInfo, nullptr, &QueryPool));
-
-	QueryOutput.AddUninitialized(MaxQueries);
-}
-
-FVulkanQueryPool::~FVulkanQueryPool()
-{
-	VulkanRHI::vkDestroyQueryPool(Device->GetInstanceHandle(), QueryPool, nullptr);
-	QueryPool = VK_NULL_HANDLE;
-	check(QueryPool == VK_NULL_HANDLE);
-}
-
-
 inline VkResult FVulkanQueryPool::InternalGetQueryPoolResults(uint32 FirstQuery, uint32 NumQueries, VkQueryResultFlags ExtraFlags)
 {
@@ -163,7 +56,7 @@
 
 inline int32 FVulkanQueryPool::AllocateQuery()
 {
-	FScopeLock ScopeLock(&GOcclusionQueryCS);
+	FScopeLock ScopeLock(&GQueryLock);
 	int32 UsedQuery = NumUsedQueries;
 	++NumUsedQueries;
 	checkf((uint32)UsedQuery < MaxQueries, TEXT("Internal error allocating query! UsedQuery=%d NumUsedQueries=%d MaxQueries=%d"), UsedQuery, NumUsedQueries, MaxQueries);
@@ -172,7 +65,7 @@
 
 bool FVulkanOcclusionQueryPool::GetAllResults(bool bWait)
 {
-	FScopeLock ScopeLock(&GOcclusionQueryCS);
+	FScopeLock ScopeLock(&GQueryLock);
 	check(!bHasResults);
 	check(CmdBuffer);
 	checkf(CmdBuffer->GetSubmittedFenceCounter() >= FenceCounter, TEXT("Tried to read a query that has not been submitted!"));
@@ -282,7 +175,7 @@
 {
 	check(!CmdBuffer || CmdBuffer->GetFenceSignaledCounter() > FenceCounter);
 	CmdBuffer = nullptr;
-	FenceCounter = UINT32_MAX;
+	FenceCounter = UINT64_MAX;
 	NumUsedQueries = 0;
 	bHasResults = false;
 	FVulkanQueryPool::ResetAll(InCmdBuffer);
@@ -359,102 +252,17 @@
 	return false;
 }
 
-void FVulkanCommandListContext::ReadAndCalculateGPUFrameTime()
-{
-	check(IsImmediate());
-
-	if (FrameTiming)
-	{
-		const uint64 Delta = FrameTiming->GetTiming(false);
-#if VULKAN_USE_NEW_QUERIES
-		const double SecondsPerCycle = FPlatformTime::GetSecondsPerCycle();
-		const double Frequency = double(FVulkanGPUTiming::GetTimingFrequency());
-		GGPUFrameTime = FMath::TruncToInt(double(Delta) / Frequency / SecondsPerCycle);
-#else
-		GGPUFrameTime = Delta ? (Delta / 1e6) / FPlatformTime::GetSecondsPerCycle() : 0;
-#endif
-	}
-	else
-	{
-		GGPUFrameTime = 0;
-	}
-
-	//static const auto CVar = IConsoleManager::Get().FindConsoleVariable(TEXT("r.Vulkan.ProfileCmdBuffers"));
-	//if (CVar->GetInt() != 0)
-	//{
-	//	const uint64 Delta = GetCommandBufferManager()->CalculateGPUTime();
-	//	GGPUFrameTime = Delta ? (Delta / 1e6) / FPlatformTime::GetSecondsPerCycle() : 0;
-	//}
-}
-
-
-FRenderQueryRHIRef FVulkanDynamicRHI::RHICreateRenderQuery(ERenderQueryType QueryType)
-{
-	FVulkanRenderQuery* Query = new FVulkanRenderQuery(Device, QueryType);
-	return Query;
-}
-
-bool FVulkanDynamicRHI::RHIGetRenderQueryResult(FRenderQueryRHIParamRef QueryRHI, uint64& OutResult, bool bWait)
-{
-	FScopeLock ScopeLock(&GOcclusionQueryCS);
-	check(IsInRenderingThread());
-
-	FVulkanRenderQuery* Query = ResourceCast(QueryRHI);
-	return Query->GetResult(Device, OutResult, bWait);
-}
-
-void FVulkanCommandListContext::RHIBeginRenderQuery(FRenderQueryRHIParamRef QueryRHI)
-{
-	FScopeLock ScopeLock(&GOcclusionQueryCS);
-	FVulkanRenderQuery* Query = ResourceCast(QueryRHI);
-	if (Query->QueryType == RQT_Occlusion)
-	{
-		FVulkanCmdBuffer* CmdBuffer = CommandBufferManager->GetActiveCmdBuffer();
-		const int32 QueryIndex = CurrentOcclusionQueryPool->AllocateQuery();
-		Query->Reset(CurrentOcclusionQueryPool, QueryIndex);
-		Query->Begin(CmdBuffer);
-	}
-	else
-	{
-		check(0);
-	}
-}
-
-void FVulkanCommandListContext::RHIEndRenderQuery(FRenderQueryRHIParamRef QueryRHI)
-{
-	FScopeLock ScopeLock(&GOcclusionQueryCS);
-	FVulkanRenderQuery* Query = ResourceCast(QueryRHI);
-
-	FVulkanCmdBuffer* CmdBuffer = CommandBufferManager->GetActiveCmdBuffer();
-	if (Query->QueryType == RQT_AbsoluteTime)
-	{
-		bool bRequiresReset = false;
-		FVulkanTimestampQueryPool* Pool = Device->PrepareTimestampQueryPool(bRequiresReset);
-		if (bRequiresReset)
-		{
-			FVulkanCmdBuffer* UploadCmdBuffer = CommandBufferManager->GetUploadCmdBuffer();
-			Pool->ResetAll(UploadCmdBuffer);
-			CommandBufferManager->SubmitUploadCmdBuffer(false);
-		}
-		const int32 QueryIndex = Pool->AllocateQuery();
-		Query->Reset(Pool, QueryIndex);
-	}
-
-	Query->End(CmdBuffer);
-}
-
 void FVulkanCommandListContext::BeginOcclusionQueryBatch(FVulkanCmdBuffer* CmdBuffer, uint32 NumQueriesInBatch)
 {
-	FScopeLock ScopeLock(&GOcclusionQueryCS);
+	FScopeLock ScopeLock(&GQueryLock);
 	ensure(IsImmediate());
-	checkf(!CurrentOcclusionQueryPool, TEXT("BeginOcclusionQueryBatch called without corresponding EndOcclusionQueryBatch!"));
-	CurrentOcclusionQueryPool = Device->PrepareOcclusionQueryPool(NumQueriesInBatch);
-
 	ensure(CmdBuffer->IsOutsideRenderPass());
+	checkf(!CurrentOcclusionQueryPool, TEXT("BeginOcclusionQueryBatch called without corresponding EndOcclusionQueryBatch!"));
+	CurrentOcclusionQueryPool = Device->AcquireOcclusionQueryPool(NumQueriesInBatch);
 	CurrentOcclusionQueryPool->Reset(CmdBuffer);
 }
 
-FVulkanOcclusionQueryPool* FVulkanDevice::PrepareOcclusionQueryPool(uint32 NumQueries)
+FVulkanOcclusionQueryPool* FVulkanDevice::AcquireOcclusionQueryPool(uint32 NumQueries)
 {
 	if (!OcclusionQueryPools[CurrentOcclusionQueryPool] || OcclusionQueryPools[CurrentOcclusionQueryPool]->GetMaxQueries() < NumQueries)
 	{
@@ -490,7 +298,7 @@
 
 void FVulkanCommandListContext::EndOcclusionQueryBatch(FVulkanCmdBuffer* CmdBuffer)
 {
-	FScopeLock ScopeLock(&GOcclusionQueryCS);
+	FScopeLock ScopeLock(&GQueryLock);
 	ensure(IsImmediate());
 	checkf(CurrentOcclusionQueryPool, TEXT("EndOcclusionQueryBatch called without corresponding BeginOcclusionQueryBatch!"));
 	CurrentOcclusionQueryPool->SetFence(CmdBuffer);
@@ -504,78 +312,156 @@
 		SafePointSubmit();
 	}
 }
+#endif
 
+FVulkanRenderQuery::FVulkanRenderQuery(ERenderQueryType InQueryType)
+#if VULKAN_USE_NEW_QUERIES
+	: QueryType(InQueryType)
 #else
-FCriticalSection GQueryLock;
-
-FOLDVulkanRenderQuery::FOLDVulkanRenderQuery(FVulkanDevice* Device, ERenderQueryType InQueryType)
 	: CurrentQueryIdx(0)
 	, QueryType(InQueryType)
 	, CurrentCmdBuffer(nullptr)
+#endif
 {
 	INC_DWORD_STAT(STAT_VulkanNumQueries);
+#if VULKAN_USE_NEW_QUERIES
+#else
 	for (int Index = 0; Index < NumQueries; ++Index)
 	{
 		QueryIndices[Index] = -1;
 		QueryPools[Index] = nullptr;
 	}
+#endif
 }
 
-FOLDVulkanRenderQuery::~FOLDVulkanRenderQuery()
+FVulkanRenderQuery::~FVulkanRenderQuery()
 {
 	DEC_DWORD_STAT(STAT_VulkanNumQueries);
+#if VULKAN_USE_NEW_QUERIES
+#else
 	for (int Index = 0; Index < NumQueries; ++Index)
 	{
 		if (QueryIndices[Index] != -1)
 		{
 			FScopeLock Lock(&GQueryLock);
-			((FOLDVulkanBufferedQueryPool*)QueryPools[Index])->ReleaseQuery(QueryIndices[Index]);
+			((FVulkanBufferedQueryPool*)QueryPools[Index])->ReleaseQuery(QueryIndices[Index]);
 		}
 	}
+#endif
 }
 
-inline void FOLDVulkanRenderQuery::Begin(FVulkanCmdBuffer* CmdBuffer)
+inline void FVulkanRenderQuery::Begin(FVulkanCmdBuffer* InCmdBuffer)
 {
-	CurrentCmdBuffer = CmdBuffer;
+#if VULKAN_USE_NEW_QUERIES
+	BeginCmdBuffer = InCmdBuffer;
+	check(State == EState::Reset);
+	if (QueryType == RQT_Occlusion)
+	{
+		check(Pool);
+		VulkanRHI::vkCmdBeginQuery(InCmdBuffer->GetHandle(), Pool->GetHandle(), QueryIndex, VK_QUERY_CONTROL_PRECISE_BIT);
+		State = EState::InBegin;
+		LastPoolReset = Pool->NumResets;
+	}
+	else
+	{
+		ensureMsgf(0, TEXT("Timer queries do NOT support Begin()!"));
+	}
+#else
+	CurrentCmdBuffer = InCmdBuffer;
 	ensure(GetActiveQueryIndex() != -1);
 	if (QueryType == RQT_Occlusion)
 	{
-		VulkanRHI::vkCmdBeginQuery(CmdBuffer->GetHandle(), GetActiveQueryPool()->GetHandle(), GetActiveQueryIndex(), VK_QUERY_CONTROL_PRECISE_BIT);
+		VulkanRHI::vkCmdBeginQuery(InCmdBuffer->GetHandle(), GetActiveQueryPool()->GetHandle(), GetActiveQueryIndex(), VK_QUERY_CONTROL_PRECISE_BIT);
 	}
 	else
 	{
 		ensure(0);
 	}
+#endif
 }
 
-inline void FOLDVulkanRenderQuery::End(FVulkanCmdBuffer* CmdBuffer)
+inline void FVulkanRenderQuery::End(FVulkanCmdBuffer* InCmdBuffer)
 {
-	ensure(QueryType != RQT_Occlusion || CurrentCmdBuffer == CmdBuffer);
+#if VULKAN_USE_NEW_QUERIES
+	ensure(QueryType != RQT_Occlusion || BeginCmdBuffer == InCmdBuffer);
+	if (QueryType == RQT_Occlusion)
+	{
+		check(State == EState::InBegin);
+		check(Pool);
+		VulkanRHI::vkCmdEndQuery(InCmdBuffer->GetHandle(), Pool->GetHandle(), QueryIndex);
+		check(LastPoolReset == Pool->NumResets);
+	}
+	else
+	{
+		check(State == EState::Reset);
+		BeginCmdBuffer = InCmdBuffer;
+		VulkanRHI::vkCmdWriteTimestamp(InCmdBuffer->GetHandle(), VK_PIPELINE_STAGE_BOTTOM_OF_PIPE_BIT, Pool->GetHandle(), QueryIndex);
+		LastPoolReset = Pool->NumResets;
+	}
+	State = EState::InEnd;
+#else
+	ensure(QueryType != RQT_Occlusion || CurrentCmdBuffer == InCmdBuffer);
 	ensure(GetActiveQueryIndex() != -1);
 
 	if (QueryType == RQT_Occlusion)
 	{
-		VulkanRHI::vkCmdEndQuery(CmdBuffer->GetHandle(), GetActiveQueryPool()->GetHandle(), GetActiveQueryIndex());
+		VulkanRHI::vkCmdEndQuery(InCmdBuffer->GetHandle(), GetActiveQueryPool()->GetHandle(), GetActiveQueryIndex());
 	}
 	else
 	{
-		VulkanRHI::vkCmdWriteTimestamp(CmdBuffer->GetHandle(), VK_PIPELINE_STAGE_BOTTOM_OF_PIPE_BIT, GetActiveQueryPool()->GetHandle(), GetActiveQueryIndex());
+		VulkanRHI::vkCmdWriteTimestamp(InCmdBuffer->GetHandle(), VK_PIPELINE_STAGE_BOTTOM_OF_PIPE_BIT, GetActiveQueryPool()->GetHandle(), GetActiveQueryIndex());
 	}
+#endif
 }
 
-bool FOLDVulkanRenderQuery::GetResult(FVulkanDevice* Device, uint64& Result, bool bWait)
+bool FVulkanRenderQuery::GetResult(FVulkanDevice* Device, uint64& OutResult, bool bWait)
 {
+#if VULKAN_USE_NEW_QUERIES
+	FScopeLock ScopeLock(&GQueryLock);
+	if (State == EState::HasResults)
+	{
+		OutResult = Result;
+		return true;
+	}
+
+	check(Pool);
+	check(State == EState::InEnd);
+	if (QueryType == RQT_Occlusion)
+	{
+		check(IsInRenderingThread());
+		if (((FVulkanOcclusionQueryPool*)Pool)->GetResults(QueryIndex, bWait, Result))
+		{
+			check(LastPoolReset == Pool->NumResets);
+			State = EState::HasResults;
+			OutResult = Result;
+			return true;
+		}
+	}
+	else
+	{
+		check(IsInRenderingThread() || IsInRHIThread());
+		if (((FVulkanTimestampQueryPool*)Pool)->GetResults(QueryIndex, bWait, Result))
+		{
+			check(LastPoolReset == Pool->NumResets);
+			State = EState::HasResults;
+			OutResult = Result;
+			return true;
+		}
+	}
+	return false;
+#else
 	if (GetActiveQueryIndex() != -1)
 	{
 		check(IsInRenderingThread() || IsInRHIThread());
 		FVulkanCommandListContext& Context = Device->GetImmediateContext();
-		FOLDVulkanBufferedQueryPool* Pool = (FOLDVulkanBufferedQueryPool*)GetActiveQueryPool();
-		return Pool->GetResults(Context, this, bWait, Result);
+		FVulkanBufferedQueryPool* Pool = (FVulkanBufferedQueryPool*)GetActiveQueryPool();
+		return Pool->GetResults(Context, this, bWait, OutResult);
 	}
 	return false;
+#endif
 }
 
-FOLDVulkanQueryPool::FOLDVulkanQueryPool(FVulkanDevice* InDevice, uint32 InMaxQueries, VkQueryType InQueryType)
+FVulkanQueryPool::FVulkanQueryPool(FVulkanDevice* InDevice, uint32 InMaxQueries, VkQueryType InQueryType)
 	: VulkanRHI::FDeviceChild(InDevice)
 	, QueryPool(VK_NULL_HANDLE)
 	, MaxQueries(InMaxQueries)
@@ -584,31 +470,41 @@
 	check(InDevice);
 
 	VkQueryPoolCreateInfo PoolCreateInfo;
-	FMemory::Memzero(PoolCreateInfo);
-	PoolCreateInfo.sType = VK_STRUCTURE_TYPE_QUERY_POOL_CREATE_INFO;
+	ZeroVulkanStruct(PoolCreateInfo, VK_STRUCTURE_TYPE_QUERY_POOL_CREATE_INFO);
 	PoolCreateInfo.queryType = QueryType;
 	PoolCreateInfo.queryCount = MaxQueries;
 
 	VERIFYVULKANRESULT(VulkanRHI::vkCreateQueryPool(Device->GetInstanceHandle(), &PoolCreateInfo, nullptr, &QueryPool));
+
+#if VULKAN_USE_NEW_QUERIES
+	QueryOutput.AddUninitialized(MaxQueries);
+#endif
 }
 
-FOLDVulkanQueryPool::~FOLDVulkanQueryPool()
+FVulkanQueryPool::~FVulkanQueryPool()
 {
+#if VULKAN_USE_NEW_QUERIES
+	VulkanRHI::vkDestroyQueryPool(Device->GetInstanceHandle(), QueryPool, nullptr);
+	QueryPool = VK_NULL_HANDLE;
+#else
 	check(QueryPool == VK_NULL_HANDLE);
+#endif
 }
 
-void FOLDVulkanQueryPool::Destroy()
+#if VULKAN_USE_NEW_QUERIES
+#else
+void FVulkanQueryPool::Destroy()
 {
 	VulkanRHI::vkDestroyQueryPool(Device->GetInstanceHandle(), QueryPool, nullptr);
 	QueryPool = VK_NULL_HANDLE;
 }
 
-void FOLDVulkanQueryPool::Reset(FVulkanCmdBuffer* InCmdBuffer)
+void FVulkanQueryPool::Reset(FVulkanCmdBuffer* InCmdBuffer)
 {
 	VulkanRHI::vkCmdResetQueryPool(InCmdBuffer->GetHandle(), QueryPool, 0, MaxQueries);
 }
 
-inline bool FOLDVulkanBufferedQueryPool::GetResults(FVulkanCommandListContext& Context, FOLDVulkanRenderQuery* Query, bool bWait, uint64& OutResult)
+inline bool FVulkanBufferedQueryPool::GetResults(FVulkanCommandListContext& Context, FVulkanRenderQuery* Query, bool bWait, uint64& OutResult)
 {
 	VkQueryResultFlags Flags = VK_QUERY_RESULT_64_BIT;
 #if VULKAN_USE_QUERY_WAIT
@@ -748,9 +644,32 @@
 
 	return true;
 }
+#endif
 
 void FVulkanCommandListContext::ReadAndCalculateGPUFrameTime()
 {
+#if VULKAN_USE_NEW_QUERIES
 	check(IsImmediate());
 
 	if (FrameTiming)
+	{
+		const uint64 Delta = FrameTiming->GetTiming(false);
+		const double SecondsPerCycle = FPlatformTime::GetSecondsPerCycle();
+		const double Frequency = double(FVulkanGPUTiming::GetTimingFrequency());
+		GGPUFrameTime = FMath::TruncToInt(double(Delta) / Frequency / SecondsPerCycle);
+	}
+	else
+	{
+		GGPUFrameTime = 0;
+	}
+
+	//static const auto CVar = IConsoleManager::Get().FindConsoleVariable(TEXT("r.Vulkan.ProfileCmdBuffers"));
+	//if (CVar->GetInt() != 0)
+	//{
+	//	const uint64 Delta = GetCommandBufferManager()->CalculateGPUTime();
+	//	GGPUFrameTime = Delta ? (Delta / 1e6) / FPlatformTime::GetSecondsPerCycle() : 0;
+	//}
+#else
+	check(IsImmediate());
+
+	if (FrameTiming)
@@ -769,27 +688,51 @@
 		const uint64 Delta = GetCommandBufferManager()->CalculateGPUTime();
 		GGPUFrameTime = Delta ? (Delta / 1e6) / FPlatformTime::GetSecondsPerCycle() : 0;
 	}
+#endif
 }
 
 
 FRenderQueryRHIRef FVulkanDynamicRHI::RHICreateRenderQuery(ERenderQueryType QueryType)
 {
-	FOLDVulkanRenderQuery* Query = new FOLDVulkanRenderQuery(Device, QueryType);
+	FVulkanRenderQuery* Query = new FVulkanRenderQuery(QueryType);
 	return Query;
 }
 
-bool FVulkanDynamicRHI::RHIGetRenderQueryResult(FRenderQueryRHIParamRef QueryRHI, uint64& OutNumPixels, bool bWait)
+bool FVulkanDynamicRHI::RHIGetRenderQueryResult(FRenderQueryRHIParamRef QueryRHI, uint64& OutResult, bool bWait)
 {
+#if VULKAN_USE_NEW_QUERIES
+	FScopeLock ScopeLock(&GQueryLock);
 	check(IsInRenderingThread());
 
-	FOLDVulkanRenderQuery* Query = ResourceCast(QueryRHI);
+	FVulkanRenderQuery* Query = ResourceCast(QueryRHI);
+	return Query->GetResult(Device, OutResult, bWait);
+#else
+	check(IsInRenderingThread());
+
+	FVulkanRenderQuery* Query = ResourceCast(QueryRHI);
 	check(Query);
-	return Query->GetResult(Device, OutNumPixels, bWait);
+	return Query->GetResult(Device, OutResult, bWait);
+#endif
 }
 
 void FVulkanCommandListContext::RHIBeginRenderQuery(FRenderQueryRHIParamRef QueryRHI)
 {
-	FOLDVulkanRenderQuery* Query = ResourceCast(QueryRHI);
+#if VULKAN_USE_NEW_QUERIES
+	FScopeLock ScopeLock(&GQueryLock);
+	FVulkanRenderQuery* Query = ResourceCast(QueryRHI);
 	if (Query->QueryType == RQT_Occlusion)
 	{
 		FVulkanCmdBuffer* CmdBuffer = CommandBufferManager->GetActiveCmdBuffer();
+		const int32 QueryIndex = CurrentOcclusionQueryPool->AllocateQuery();
+		Query->Reset(CurrentOcclusionQueryPool, QueryIndex);
+		Query->Begin(CmdBuffer);
+	}
+	else
+	{
+		check(0);
+	}
+#else
+	FVulkanRenderQuery* Query = ResourceCast(QueryRHI);
+	if (Query->QueryType == RQT_Occlusion)
+	{
+		FVulkanCmdBuffer* CmdBuffer = CommandBufferManager->GetActiveCmdBuffer();
@@ -800,14 +743,17 @@
 	{
 		check(0);
 	}
+#endif
 }
 
-void FVulkanCommandListContext::AdvanceQuery(FOLDVulkanRenderQuery* Query)
+#if VULKAN_USE_NEW_QUERIES
+#else
+void FVulkanCommandListContext::AdvanceQuery(FVulkanRenderQuery* Query)
 {
 	//reset prev query
 	if (Query->GetActiveQueryIndex() != -1)
 	{
-		FOLDVulkanBufferedQueryPool* OcclusionPool = (FOLDVulkanBufferedQueryPool*)Query->GetActiveQueryPool();
+		FVulkanBufferedQueryPool* OcclusionPool = (FVulkanBufferedQueryPool*)Query->GetActiveQueryPool();
 		CurrentOcclusionQueryData.AddToResetList(OcclusionPool, Query->GetActiveQueryIndex());
 	}
 
@@ -818,7 +764,7 @@
 	if (Query->GetActiveQueryIndex() == -1)
 	{
 		uint32 QueryIndex = 0;
-		FOLDVulkanBufferedQueryPool* Pool = nullptr;
+		FVulkanBufferedQueryPool* Pool = nullptr;
 
 		{
 			FScopeLock Lock(&GQueryLock);
@@ -842,10 +788,10 @@
 	}
 
 	// mark at begin
-	((FOLDVulkanBufferedQueryPool*)Query->GetActiveQueryPool())->MarkQueryAsStarted(Query->GetActiveQueryIndex());
+	((FVulkanBufferedQueryPool*)Query->GetActiveQueryPool())->MarkQueryAsStarted(Query->GetActiveQueryIndex());
 }
 
-void FVulkanCommandListContext::EndRenderQueryInternal(FVulkanCmdBuffer* CmdBuffer, FOLDVulkanRenderQuery* Query)
+void FVulkanCommandListContext::EndRenderQueryInternal(FVulkanCmdBuffer* CmdBuffer, FVulkanRenderQuery* Query)
 {
 	if (Query->QueryType == RQT_Occlusion)
 	{
@@ -867,13 +813,38 @@
 		Query->End(CmdBuffer);
 	}
 }
+#endif
 
 void FVulkanCommandListContext::RHIEndRenderQuery(FRenderQueryRHIParamRef QueryRHI)
 {
-	FOLDVulkanRenderQuery* Query = ResourceCast(QueryRHI);
-	return EndRenderQueryInternal(CommandBufferManager->GetActiveCmdBuffer(), Query);	
+#if VULKAN_USE_NEW_QUERIES
+	FScopeLock ScopeLock(&GQueryLock);
+	FVulkanRenderQuery* Query = ResourceCast(QueryRHI);
+
+	FVulkanCmdBuffer* CmdBuffer = CommandBufferManager->GetActiveCmdBuffer();
+	if (Query->QueryType == RQT_AbsoluteTime)
+	{
+		bool bRequiresReset = false;
+		FVulkanTimestampQueryPool* Pool = Device->PrepareTimestampQueryPool(bRequiresReset);
+		if (bRequiresReset)
+		{
+			FVulkanCmdBuffer* UploadCmdBuffer = CommandBufferManager->GetUploadCmdBuffer();
+			Pool->ResetAll(UploadCmdBuffer);
+			CommandBufferManager->SubmitUploadCmdBuffer(false);
+		}
+		const int32 QueryIndex = Pool->AllocateQuery();
+		Query->Reset(Pool, QueryIndex);
+	}
+
+	Query->End(CmdBuffer);
+#else
+	FVulkanRenderQuery* Query = ResourceCast(QueryRHI);
+	return EndRenderQueryInternal(CommandBufferManager->GetActiveCmdBuffer(), Query);
+#endif
 }
 
+#if VULKAN_USE_NEW_QUERIES
+#else
 void FVulkanCommandListContext::RHIBeginOcclusionQueryBatch(uint32 NumQueriesInBatch)
 {
 	ensure(IsImmediate());
@@ -882,7 +853,7 @@
 	ensure(CmdBuffer->IsInsideRenderPass());
 }
 
-static inline void ProcessByte(VkCommandBuffer InCmdBufferHandle, FOLDVulkanBufferedQueryPool* Pool, uint8 Bits, int32 BaseStartIndex)
+static inline void ProcessByte(VkCommandBuffer InCmdBufferHandle, FVulkanBufferedQueryPool* Pool, uint8 Bits, int32 BaseStartIndex)
 {
 	if (Bits)
 	{
@@ -910,7 +881,7 @@
 	}
 }
 
-static inline void Process16Bits(VkCommandBuffer InCmdBufferHandle, FOLDVulkanBufferedQueryPool* Pool, uint16 Bits, int32 BaseStartIndex)
+static inline void Process16Bits(VkCommandBuffer InCmdBufferHandle, FVulkanBufferedQueryPool* Pool, uint16 Bits, int32 BaseStartIndex)
 {
 	if (Bits)
 	{
@@ -928,7 +899,7 @@
 	}
 }
 
-static inline void Process32Bits(VkCommandBuffer InCmdBufferHandle, FOLDVulkanBufferedQueryPool* Pool, uint32 Bits, int32 BaseStartIndex)
+static inline void Process32Bits(VkCommandBuffer InCmdBufferHandle, FVulkanBufferedQueryPool* Pool, uint32 Bits, int32 BaseStartIndex)
 {
 	if (Bits)
 	{
@@ -955,7 +926,7 @@
 	for (auto& Pair : ResetList)
 	{
 		TArray<uint64>& ListPerPool = Pair.Value;
-		FOLDVulkanBufferedQueryPool* Pool = (FOLDVulkanBufferedQueryPool*)Pair.Key;
+		FVulkanBufferedQueryPool* Pool = (FVulkanBufferedQueryPool*)Pair.Key;
 
 		// Initial bit Index
 		int32 WordIndex = 0;
diff --git a/Engine/Source/Runtime/VulkanRHI/Private/VulkanRHIPrivate.h b/Engine/Source/Runtime/VulkanRHI/Private/VulkanRHIPrivate.h
--- a/Engine/Source/Runtime/VulkanRHI/Private/VulkanRHIPrivate.h
+++ b/Engine/Source/Runtime/VulkanRHI/Private/VulkanRHIPrivate.h
@@ -141,7 +141,7 @@ protected:
 		Extent.Extent3D.depth = 0;
 	}
 
-	friend class FVulkanPipelineStateCache;
+	friend class FVulkanPipelineStateCacheManager;
 
 private:
 	void CreateRenderPassHash();
@@ -262,7 +262,7 @@ public:
 
 private:
 	friend class FTransitionAndLayoutManager;
-	friend class FVulkanPipelineStateCache;
+	friend class FVulkanPipelineStateCacheManager;
 
 	FVulkanRenderPass(FVulkanDevice& Device, const FVulkanRenderTargetLayout& RTLayout);
 	~FVulkanRenderPass();
diff --git a/Engine/Source/Runtime/VulkanRHI/Private/VulkanRenderTarget.cpp b/Engine/Source/Runtime/VulkanRHI/Private/VulkanRenderTarget.cpp
--- a/Engine/Source/Runtime/VulkanRHI/Private/VulkanRenderTarget.cpp
+++ b/Engine/Source/Runtime/VulkanRHI/Private/VulkanRenderTarget.cpp
@@ -1842,16 +1842,20 @@ FVulkanRenderTargetLayout::FVulkanRenderTargetLayout(FVulkanDevice& InDevice, co
 				NumUsedClearValues = NumAttachmentDescriptions + 1;
 			}
 		}
-
-		if (CurrDesc.samples != VK_SAMPLE_COUNT_1_BIT)
+		if (CurrDesc.samples == VK_SAMPLE_COUNT_1_BIT)
 		{
-			// Can't resolve MSAA depth/stencil
-			ensure(GetStoreAction(GetDepthActions(RPInfo.DepthStencilRenderTarget.Action)) != ERenderTargetStoreAction::EMultisampleResolve);
-			ensure(GetStoreAction(GetStencilActions(RPInfo.DepthStencilRenderTarget.Action)) != ERenderTargetStoreAction::EMultisampleResolve);
-		}
 		CurrDesc.storeOp = RenderTargetStoreActionToVulkan(GetStoreAction(GetDepthActions(RPInfo.DepthStencilRenderTarget.Action)), true);
 		CurrDesc.stencilStoreOp = RenderTargetStoreActionToVulkan(GetStoreAction(GetStencilActions(RPInfo.DepthStencilRenderTarget.Action)), true);
-		
+		}
+		else
+		{
+			ensure(GetStoreAction(GetDepthActions(RPInfo.DepthStencilRenderTarget.Action)) == ERenderTargetStoreAction::ENoAction);
+			ensure(GetStoreAction(GetStencilActions(RPInfo.DepthStencilRenderTarget.Action)) == ERenderTargetStoreAction::ENoAction);
+			// Never want to store MSAA depth/stencil
+			CurrDesc.storeOp = VK_ATTACHMENT_STORE_OP_DONT_CARE;
+			CurrDesc.stencilStoreOp = VK_ATTACHMENT_STORE_OP_DONT_CARE;
+		}
+
 		DepthStencilLayout = VulkanRHI::GetDepthStencilLayout(RPInfo.DepthStencilRenderTarget.ExclusiveDepthStencil, InDevice);
 		CurrDesc.initialLayout = DepthStencilLayout;
 		CurrDesc.finalLayout = DepthStencilLayout;
diff --git a/Engine/Source/Runtime/VulkanRHI/Private/VulkanTexture.cpp b/Engine/Source/Runtime/VulkanRHI/Private/VulkanTexture.cpp
--- a/Engine/Source/Runtime/VulkanRHI/Private/VulkanTexture.cpp
+++ b/Engine/Source/Runtime/VulkanRHI/Private/VulkanTexture.cpp
@@ -1881,7 +1881,7 @@ static VkMemoryRequirements FindOrCalculateTexturePlatformSize(FVulkanDevice* De
 	VkMemoryRequirements* Found = nullptr;
 	{
 		FScopeLock Lock(&TextureSizesLock);
-		TextureSizes.Find(Hash);
+		Found = TextureSizes.Find(Hash);
 		if (Found)
 		{
 			return *Found;
diff --git a/Engine/Source/Runtime/VulkanRHI/Private/VulkanUtil.cpp b/Engine/Source/Runtime/VulkanRHI/Private/VulkanUtil.cpp
--- a/Engine/Source/Runtime/VulkanRHI/Private/VulkanUtil.cpp
+++ b/Engine/Source/Runtime/VulkanRHI/Private/VulkanUtil.cpp
@@ -49,9 +49,8 @@ void FVulkanGPUTiming::CalibrateTimers(FVulkanCommandListContext& InCmdContext)
 	// TODO: Implement VULKAN_USE_NEW_QUERIES version
 
 #else
-
 	FVulkanDevice* Device = InCmdContext.GetDevice();
-	FOLDVulkanRenderQuery* TimestampQuery = new FOLDVulkanRenderQuery(Device, RQT_AbsoluteTime);
+	FVulkanRenderQuery* TimestampQuery = new FVulkanRenderQuery(RQT_AbsoluteTime);
 
 	{
 		FVulkanCmdBuffer* CmdBuffer = InCmdContext.GetCommandBufferManager()->GetUploadCmdBuffer();
@@ -99,13 +98,8 @@ void FVulkanGPUTiming::Initialize()
 	{
 		for (int32 Index = 0; Index < MaxTimers; ++Index)
 		{
-#if VULKAN_USE_NEW_QUERIES
-			Timers[Index].Begin = new FVulkanRenderQuery(CmdContext->GetDevice(), RQT_AbsoluteTime);
-			Timers[Index].End = new FVulkanRenderQuery(CmdContext->GetDevice(), RQT_AbsoluteTime);
-#else
-			Timers[Index].Begin = new FOLDVulkanRenderQuery(CmdContext->GetDevice(), RQT_AbsoluteTime);
-			Timers[Index].End = new FOLDVulkanRenderQuery(CmdContext->GetDevice(), RQT_AbsoluteTime);
-#endif
+			Timers[Index].Begin = new FVulkanRenderQuery(RQT_AbsoluteTime);
+			Timers[Index].End = new FVulkanRenderQuery(RQT_AbsoluteTime);
 		}
 	}
 }
@@ -208,7 +202,7 @@ uint64 FVulkanGPUTiming::GetTiming(bool bGetCurrentResultsAndBlock)
 					int i = 0;
 					++i;
 				}
-				else if (Timers[TimerIndex].Begin->HasQueryBeenEmitted() && Timers[TimerIndex].End->HasQueryBeenEmitted())
+				else if (Timers[TimerIndex].Begin->HasQueryBeenEnded() && Timers[TimerIndex].End->HasQueryBeenEnded())
 #endif
 				{
 #if VULKAN_USE_NEW_QUERIES
@@ -241,7 +235,7 @@ uint64 FVulkanGPUTiming::GetTiming(bool bGetCurrentResultsAndBlock)
 		{
 			TimerIndex = CurrentTimerIndex;
 #if VULKAN_USE_NEW_QUERIES
-			check(Timers[TimerIndex].Begin->HasQueryBeenEmitted() && Timers[TimerIndex].End->HasQueryBeenEmitted());
+			check(Timers[TimerIndex].Begin->HasQueryBeenEnded() && Timers[TimerIndex].End->HasQueryBeenEnded());
 #endif
 
 			if (Timers[TimerIndex].Begin->GetResult(Device, BeginTime, true))
diff --git a/Engine/Source/Runtime/VulkanRHI/Private/VulkanViewport.cpp b/Engine/Source/Runtime/VulkanRHI/Private/VulkanViewport.cpp
--- a/Engine/Source/Runtime/VulkanRHI/Private/VulkanViewport.cpp
+++ b/Engine/Source/Runtime/VulkanRHI/Private/VulkanViewport.cpp
@@ -525,6 +525,8 @@ void FVulkanViewport::CreateSwapchain()
 				VulkanRHI::ImagePipelineBarrier(CmdBuffer->GetHandle(), Images[Index], EImageLayoutBarrier::TransferDest, EImageLayoutBarrier::ColorAttachment, Range);
 			}
 		}
+
+		Device->GetImmediateContext().GetCommandBufferManager()->SubmitUploadCmdBuffer();
 	}
 	else
 	{
diff --git a/Engine/Source/Runtime/VulkanRHI/Private/Windows/VulkanWindowsPlatform.h b/Engine/Source/Runtime/VulkanRHI/Private/Windows/VulkanWindowsPlatform.h
--- a/Engine/Source/Runtime/VulkanRHI/Private/Windows/VulkanWindowsPlatform.h
+++ b/Engine/Source/Runtime/VulkanRHI/Private/Windows/VulkanWindowsPlatform.h
@@ -8,7 +8,7 @@
 #define VK_USE_PLATFORM_WIN32_KHR				1
 #define VK_USE_PLATFORM_WIN32_KHX				1
 
-//#define VULKAN_USE_NEW_QUERIES					1
+#define VULKAN_USE_NEW_QUERIES					0
 
 #define	VULKAN_SHOULD_ENABLE_DRAW_MARKERS		(UE_BUILD_DEBUG || UE_BUILD_DEVELOPMENT)
 #define VULKAN_HAS_PHYSICAL_DEVICE_PROPERTIES2	1
diff --git a/Engine/Source/Runtime/VulkanRHI/Public/VulkanContext.h b/Engine/Source/Runtime/VulkanRHI/Public/VulkanContext.h
--- a/Engine/Source/Runtime/VulkanRHI/Public/VulkanContext.h
+++ b/Engine/Source/Runtime/VulkanRHI/Public/VulkanContext.h
@@ -332,7 +332,7 @@ public:
 	{
 		return Device;
 	}
-	void EndRenderQueryInternal(FVulkanCmdBuffer* CmdBuffer, FOLDVulkanRenderQuery* Query);
+	void EndRenderQueryInternal(FVulkanCmdBuffer* CmdBuffer, FVulkanRenderQuery* Query);
 
 	inline VkImageLayout FindLayout(VkImage Image)
 	{
@@ -446,7 +446,7 @@ protected:
 		{
 		}
 
-		void AddToResetList(FOLDVulkanQueryPool* Pool, int32 QueryIndex)
+		void AddToResetList(FVulkanQueryPool* Pool, int32 QueryIndex)
 		{
 			TArray<uint64>& ListPerPool = ResetList.FindOrAdd(Pool);
 			int32 Word = QueryIndex / 64;
@@ -469,10 +469,10 @@ protected:
 			}
 		}
 
-		TMap<FOLDVulkanQueryPool*, TArray<uint64>> ResetList;
+		TMap<FVulkanQueryPool*, TArray<uint64>> ResetList;
 	};
 	FOcclusionQueryData CurrentOcclusionQueryData;
-	void AdvanceQuery(FOLDVulkanRenderQuery* Query);
+	void AdvanceQuery(FVulkanRenderQuery* Query);
 #endif
 
 	// List of UAVs which need setting for pixel shaders. D3D treats UAVs like rendertargets so the RHI doesn't make SetUAV calls at the right time
diff --git a/Engine/Source/Runtime/VulkanRHI/Public/VulkanDynamicRHI.h b/Engine/Source/Runtime/VulkanRHI/Public/VulkanDynamicRHI.h
--- a/Engine/Source/Runtime/VulkanRHI/Public/VulkanDynamicRHI.h
+++ b/Engine/Source/Runtime/VulkanRHI/Public/VulkanDynamicRHI.h
@@ -157,7 +157,6 @@ public:
 	virtual void UnlockVertexBuffer_RenderThread(class FRHICommandListImmediate& RHICmdList, FVertexBufferRHIParamRef VertexBuffer) override final;
 #endif
 	virtual FTexture2DRHIRef AsyncReallocateTexture2D_RenderThread(class FRHICommandListImmediate& RHICmdList, FTexture2DRHIParamRef Texture2D, int32 NewMipCount, int32 NewSizeX, int32 NewSizeY, FThreadSafeCounter* RequestStatus) override final;
-	//virtual ETextureReallocationStatus FinalizeAsyncReallocateTexture2D_RenderThread(class FRHICommandListImmediate& RHICmdList, FTexture2DRHIParamRef Texture2D, bool bBlockUntilCompleted) override final;
 	//virtual ETextureReallocationStatus CancelAsyncReallocateTexture2D_RenderThread(class FRHICommandListImmediate& RHICmdList, FTexture2DRHIParamRef Texture2D, bool bBlockUntilCompleted) override final;
 	//virtual FIndexBufferRHIRef CreateIndexBuffer_RenderThread(class FRHICommandListImmediate& RHICmdList, uint32 Stride, uint32 Size, uint32 InUsage, FRHIResourceCreateInfo& CreateInfo) override final;
 #if 0
diff --git a/Engine/Source/Runtime/VulkanRHI/Public/VulkanResources.h b/Engine/Source/Runtime/VulkanRHI/Public/VulkanResources.h
--- a/Engine/Source/Runtime/VulkanRHI/Public/VulkanResources.h
+++ b/Engine/Source/Runtime/VulkanRHI/Public/VulkanResources.h
@@ -90,7 +90,7 @@
 	FVulkanDevice*		Device;
 
 	friend class FVulkanCommandListContext;
-	friend class FVulkanPipelineStateCache;
+	friend class FVulkanPipelineStateCacheManager;
 	friend class FVulkanComputeShaderState;
 	friend class FVulkanComputePipeline;
 };
@@ -636,14 +636,13 @@
 	}
 }
 
-#if VULKAN_USE_NEW_QUERIES
-extern FCriticalSection GOcclusionQueryCS;
-class FVulkanQueryPool : public VulkanRHI::FDeviceChild 
+class FVulkanQueryPool : public VulkanRHI::FDeviceChild
 {
 public:
 	FVulkanQueryPool(FVulkanDevice* InDevice, uint32 InMaxQueries, VkQueryType InQueryType);
 	virtual ~FVulkanQueryPool();
 
+#if VULKAN_USE_NEW_QUERIES
 	inline uint32 GetMaxQueries() const
 	{
 		return MaxQueries;
@@ -664,8 +663,21 @@
 	}
 
 	uint32 NumResets = 0;
+#else
+	virtual void Destroy();
 
+	void Reset(FVulkanCmdBuffer* InCmdBuffer);
+
+	inline VkQueryPool GetHandle() const
+	{
+		return QueryPool;
+	}
+
+	void SetFence(FVulkanCmdBuffer* InCmdBuffer);
+#endif
+
 protected:
+#if VULKAN_USE_NEW_QUERIES
 	VkQueryPool QueryPool;
 	int32 NumUsedQueries = 0;
 	const uint32 MaxQueries;
@@ -680,8 +692,23 @@
 	}
 
 	friend class FVulkanDevice;
+#else
+	VkQueryPool QueryPool;
+	volatile int32 NumUsedQueries = 0;
+	const uint32 MaxQueries;
+	const VkQueryType QueryType;
+
+	TArray<uint64> QueryOutput;
+	bool bHasResults = false;
+
+	FVulkanCmdBuffer* CmdBuffer = nullptr;
+	uint64 FenceCounter = UINT32_MAX;
+
+	friend class FVulkanDynamicRHI;
+#endif
 };
 
+#if VULKAN_USE_NEW_QUERIES
 class FVulkanOcclusionQueryPool : public FVulkanQueryPool
 {
 public:
@@ -751,91 +778,12 @@
 	friend class FVulkanGPUTiming;
 };
 
-class FVulkanRenderQuery : public FRHIRenderQuery
-{
-public:
-	FVulkanRenderQuery(FVulkanDevice* Device, ERenderQueryType InQueryType);
-	virtual ~FVulkanRenderQuery();
-
-	inline bool HasQueryBeenEmitted() const
-	{
-		return State == EState::InEnd;
-	}
-
-	uint32 LastPoolReset = 0;
-
-private:
-	int32 QueryIndex = INT32_MAX;
-
-	const ERenderQueryType QueryType;
-
-	FVulkanCmdBuffer* BeginCmdBuffer = nullptr;
-
-	friend class FVulkanDynamicRHI;
-	friend class FVulkanCommandListContext;
-	friend class FVulkanGPUTiming;
-
-	FVulkanQueryPool* Pool = nullptr;
-	enum class EState
-	{
-		Reset,
-		InBegin,
-		InEnd,
-		HasResults,
-	};
-	EState State = EState::Reset;
-	uint64 Result = 0;
-
-	void Reset(FVulkanQueryPool* InPool, int32 InQueryIndex)
-	{
-		QueryIndex = InQueryIndex;
-		Pool = InPool;
-		State = EState::Reset;
-	}
-	void Begin(FVulkanCmdBuffer* InCmdBuffer);
-	void End(FVulkanCmdBuffer* InCmdBuffer);
-
-	bool GetResult(FVulkanDevice* Device, uint64& OutResult, bool bWait);
-};
-
 #else
-class FOLDVulkanQueryPool : public VulkanRHI::FDeviceChild
+class FVulkanBufferedQueryPool : public FVulkanQueryPool
 {
 public:
-	FOLDVulkanQueryPool(FVulkanDevice* InDevice, uint32 InMaxQueries, VkQueryType InQueryType);
-	virtual ~FOLDVulkanQueryPool();
-
-	virtual void Destroy();
-
-	void Reset(FVulkanCmdBuffer* InCmdBuffer);
-
-	inline VkQueryPool GetHandle() const
-	{
-		return QueryPool;
-	}
-
-	void SetFence(FVulkanCmdBuffer* InCmdBuffer);
-
-protected:
-	VkQueryPool QueryPool;
-	volatile int32 NumUsedQueries = 0;
-	const uint32 MaxQueries;
-	const VkQueryType QueryType;
-
-	TArray<uint64> QueryOutput;
-	bool bHasResults = false;
-
-	FVulkanCmdBuffer* CmdBuffer = nullptr;
-	uint64 FenceCounter = UINT32_MAX;
-
-	friend class FVulkanDynamicRHI;
-};
-
-class FOLDVulkanBufferedQueryPool : public FOLDVulkanQueryPool
-{
-public:
-	FOLDVulkanBufferedQueryPool(FVulkanDevice* InDevice, uint32 InNumQueries, VkQueryType InQueryType)
-		: FOLDVulkanQueryPool(InDevice, InNumQueries, InQueryType)
+	FVulkanBufferedQueryPool(FVulkanDevice* InDevice, uint32 InNumQueries, VkQueryType InQueryType)
+		: FVulkanQueryPool(InDevice, InNumQueries, InQueryType)
 		, LastBeginIndex(0)
 	{
 		QueryOutput.SetNum(InNumQueries);
@@ -920,7 +868,7 @@
 		}
 	}
 
-	bool GetResults(class FVulkanCommandListContext& Context, class FOLDVulkanRenderQuery* Query, bool bWait, uint64& OutNumPixels);
+	bool GetResults(class FVulkanCommandListContext& Context, class FVulkanRenderQuery* Query, bool bWait, uint64& OutNumPixels);
 
 	bool HasRoom() const
 	{
@@ -946,14 +894,57 @@
 
 	friend class FVulkanCommandListContext;
 };
+#endif
 
-class FOLDVulkanRenderQuery : public FRHIRenderQuery
+class FVulkanRenderQuery : public FRHIRenderQuery
 {
 public:
-	FOLDVulkanRenderQuery(FVulkanDevice* Device, ERenderQueryType InQueryType);
-	virtual ~FOLDVulkanRenderQuery();
+	FVulkanRenderQuery(ERenderQueryType InQueryType);
+	virtual ~FVulkanRenderQuery();
 
+#if VULKAN_USE_NEW_QUERIES
+	inline bool HasQueryBeenEmitted() const
+	{
+		return State == EState::InEnd;
+	}
+
+	uint32 LastPoolReset = 0;
+#endif
+
 private:
+#if VULKAN_USE_NEW_QUERIES
+	int32 QueryIndex = INT32_MAX;
+
+	const ERenderQueryType QueryType;
+
+	FVulkanCmdBuffer* BeginCmdBuffer = nullptr;
+
+	friend class FVulkanDynamicRHI;
+	friend class FVulkanCommandListContext;
+	friend class FVulkanGPUTiming;
+
+	FVulkanQueryPool* Pool = nullptr;
+	enum class EState
+	{
+		Reset,
+		InBegin,
+		InEnd,
+		HasResults,
+	};
+	EState State = EState::Reset;
+	uint64 Result = 0;
+
+	void Reset(FVulkanQueryPool* InPool, int32 InQueryIndex)
+	{
+		QueryIndex = InQueryIndex;
+		Pool = InPool;
+		State = EState::Reset;
+	}
+	void Begin(FVulkanCmdBuffer* InCmdBuffer);
+	void End(FVulkanCmdBuffer* InCmdBuffer);
+
+	bool GetResult(FVulkanDevice* Device, uint64& OutResult, bool bWait);
+#else
 	enum
 	{
 		NumQueries = NUM_RENDER_BUFFERS
@@ -962,9 +953,9 @@
 	int32 CurrentQueryIdx;
 
 	// Actual index and pool filled in after RHIBeginOcclusionQueryBatch
-	FOLDVulkanQueryPool* QueryPools[NumQueries];
-	FOLDVulkanQueryPool* GetActiveQueryPool() const { return QueryPools[CurrentQueryIdx]; }
-	void SetActiveQueryPool(FOLDVulkanQueryPool* Pool) { QueryPools[CurrentQueryIdx] = Pool; }
+	FVulkanQueryPool* QueryPools[NumQueries];
+	FVulkanQueryPool* GetActiveQueryPool() const { return QueryPools[CurrentQueryIdx]; }
+	void SetActiveQueryPool(FVulkanQueryPool* Pool) { QueryPools[CurrentQueryIdx] = Pool; }
 
 	int32 QueryIndices[NumQueries];
 	int32 GetActiveQueryIndex() const { return QueryIndices[CurrentQueryIdx]; }
@@ -977,15 +968,15 @@
 	friend class FVulkanDynamicRHI;
 	friend class FVulkanCommandListContext;
 
-	friend class FOLDVulkanBufferedQueryPool;
+	friend class FVulkanBufferedQueryPool;
 	friend class FVulkanGPUTiming;
 
 	void Begin(FVulkanCmdBuffer* CmdBuffer);
 	void End(FVulkanCmdBuffer* CmdBuffer);
 
 	bool GetResult(FVulkanDevice* Device, uint64& Result, bool bWait);
-};
 #endif
+};
 
 struct FVulkanBufferView : public FRHIResource, public VulkanRHI::FDeviceChild
 {
@@ -1373,7 +1364,7 @@
 	VkVertexInputAttributeDescription Attributes[MaxVertexElementCount];
 
 	friend class FVulkanPendingGfxState;
-	friend class FVulkanPipelineStateCache;
+	friend class FVulkanPipelineStateCacheManager;
 };
 
 // This class holds the staging area for packed global uniform buffers for a given shader
@@ -1543,7 +1534,7 @@
 template<>
 struct TVulkanResourceTraits<FRHIRenderQuery>
 {
-	typedef FOLDVulkanRenderQuery TConcreteType;
+	typedef FVulkanRenderQuery TConcreteType;
 };
 #endif
 template<>
diff --git a/Engine/Source/Runtime/VulkanRHI/Public/VulkanUtil.h b/Engine/Source/Runtime/VulkanRHI/Public/VulkanUtil.h
--- a/Engine/Source/Runtime/VulkanRHI/Public/VulkanUtil.h
+++ b/Engine/Source/Runtime/VulkanRHI/Public/VulkanUtil.h
@@ -10,7 +10,7 @@
 
 class FVulkanCmdBuffer;
 class FVulkanRenderQuery;
-class FOLDVulkanRenderQuery;
+class FVulkanRenderQuery;
 class FVulkanCommandListContext;
 
 class FVulkanGPUTiming : public FGPUTiming
@@ -89,12 +89,9 @@ private:
 		uint64 BeginFenceCounter = 0;
 		FVulkanCmdBuffer* EndCmdBuffer = nullptr;
 		uint64 EndFenceCounter = 0;
+#endif
 		FVulkanRenderQuery* Begin;
 		FVulkanRenderQuery* End;
-#else
-		FOLDVulkanRenderQuery* Begin;
-		FOLDVulkanRenderQuery* End;
-#endif
 	};
 	FBeginEndPair Timers[MaxTimers];
 };
