Vulkan: redone fix overbright rendering with some drivers (lost after a merge)

diff --git a/Engine/Source/Runtime/VulkanRHI/Private/VulkanRenderTarget.cpp b/Engine/Source/Runtime/VulkanRHI/Private/VulkanRenderTarget.cpp
--- a/Engine/Source/Runtime/VulkanRHI/Private/VulkanRenderTarget.cpp
+++ b/Engine/Source/Runtime/VulkanRHI/Private/VulkanRenderTarget.cpp
@@ -1267,22 +1267,20 @@ void FVulkanCommandListContext::TransitionResources(const FPendingTransition& Pe
 
 				if ((AspectMask & VK_IMAGE_ASPECT_COLOR_BIT) != 0)
 				{
-					// When transitioning a full 3d/cube texture just assume it's undefined
-					if (SrcLayout == VK_IMAGE_LAYOUT_UNDEFINED || (Surface.GetNumMips() > 1 && Surface.GetNumberOfArrayLevels() > 1))
-					{
-						Barrier.SetTransition(BarrierIndex, EImageLayoutBarrier::Undefined, EImageLayoutBarrier::ColorAttachment);
-					}
-					else
+					if (SrcLayout != VK_IMAGE_LAYOUT_COLOR_ATTACHMENT_OPTIMAL)
 					{
 						Barrier.SetTransition(BarrierIndex, VulkanRHI::GetImageLayoutFromVulkanLayout(SrcLayout), VulkanRHI::EImageLayoutBarrier::ColorAttachment);
+						SrcLayout = VK_IMAGE_LAYOUT_COLOR_ATTACHMENT_OPTIMAL;
 					}
-					SrcLayout = VK_IMAGE_LAYOUT_COLOR_ATTACHMENT_OPTIMAL;
 				}
 				else
 				{
-					check(Surface.IsDepthOrStencilAspect());
-					Barrier.SetTransition(BarrierIndex, VulkanRHI::GetImageLayoutFromVulkanLayout(SrcLayout), VulkanRHI::EImageLayoutBarrier::DepthStencilAttachment);
-					SrcLayout = VK_IMAGE_LAYOUT_DEPTH_STENCIL_ATTACHMENT_OPTIMAL;
+					if (SrcLayout != VK_IMAGE_LAYOUT_COLOR_ATTACHMENT_OPTIMAL)
+					{
+						check(Surface.IsDepthOrStencilAspect());
+						Barrier.SetTransition(BarrierIndex, VulkanRHI::GetImageLayoutFromVulkanLayout(SrcLayout), VulkanRHI::EImageLayoutBarrier::DepthStencilAttachment);
+						SrcLayout = VK_IMAGE_LAYOUT_DEPTH_STENCIL_ATTACHMENT_OPTIMAL;
+					}
 				}
 			}
 
