Fix lumin compile error subject to circular include order.

#rnx

diff --git a/Engine/Source/Runtime/VulkanRHI/Private/Lumin/VulkanLuminPlatform.cpp b/Engine/Source/Runtime/VulkanRHI/Private/Lumin/VulkanLuminPlatform.cpp
--- a/Engine/Source/Runtime/VulkanRHI/Private/Lumin/VulkanLuminPlatform.cpp
+++ b/Engine/Source/Runtime/VulkanRHI/Private/Lumin/VulkanLuminPlatform.cpp
@@ -136,6 +136,14 @@ void FVulkanLuminPlatform::GetDeviceExtensions(TArray<const ANSICHAR*>& OutExten
 	OutExtensions.Add(VK_KHR_SAMPLER_YCBCR_CONVERSION_EXTENSION_NAME);
 }
 
+void FVulkanLuminPlatform::SetupFeatureLevels()
+{
+	GShaderPlatformForFeatureLevel[ERHIFeatureLevel::ES2] = SP_VULKAN_ES3_1_LUMIN;
+	GShaderPlatformForFeatureLevel[ERHIFeatureLevel::ES3_1] = SP_VULKAN_ES3_1_LUMIN;
+	GShaderPlatformForFeatureLevel[ERHIFeatureLevel::SM4] = SP_NumPlatforms;
+	GShaderPlatformForFeatureLevel[ERHIFeatureLevel::SM5] = PLATFORM_LUMINGL4 ? SP_VULKAN_SM5_LUMIN : SP_NumPlatforms;
+}
+
 bool FVulkanLuminPlatform::ForceEnableDebugMarkers()
 {
 	// Preventing VK_EXT_DEBUG_MARKER from being enabled on Lumin, because the device doesn't support it.
diff --git a/Engine/Source/Runtime/VulkanRHI/Private/Lumin/VulkanLuminPlatform.h b/Engine/Source/Runtime/VulkanRHI/Private/Lumin/VulkanLuminPlatform.h
--- a/Engine/Source/Runtime/VulkanRHI/Private/Lumin/VulkanLuminPlatform.h
+++ b/Engine/Source/Runtime/VulkanRHI/Private/Lumin/VulkanLuminPlatform.h
@@ -50,13 +50,7 @@ public:
 	static bool SupportsASTCTextureFormats() { return true; }
 	static bool SupportsQuerySurfaceProperties() { return false; }
 
-	static void SetupFeatureLevels()
-	{
-		GShaderPlatformForFeatureLevel[ERHIFeatureLevel::ES2] = SP_VULKAN_ES3_1_LUMIN;
-		GShaderPlatformForFeatureLevel[ERHIFeatureLevel::ES3_1] = SP_VULKAN_ES3_1_LUMIN;
-		GShaderPlatformForFeatureLevel[ERHIFeatureLevel::SM4] = SP_NumPlatforms;
-		GShaderPlatformForFeatureLevel[ERHIFeatureLevel::SM5] = PLATFORM_LUMINGL4 ? SP_VULKAN_SM5_LUMIN : SP_NumPlatforms;
-	}
+	static void SetupFeatureLevels();
 
 	static bool SupportsStandardSwapchain() { return false; }
 	static EPixelFormat GetPixelFormatForNonDefaultSwapchain() { return PF_R8G8B8A8; }
diff --git a/Engine/Source/Runtime/VulkanRHI/Private/VulkanGenericPlatform.cpp b/Engine/Source/Runtime/VulkanRHI/Private/VulkanGenericPlatform.cpp
new file mode 100644
--- /dev/null
+++ b/Engine/Source/Runtime/VulkanRHI/Private/VulkanGenericPlatform.cpp
@@ -0,0 +1,12 @@
+// Copyright 1998-2019 Epic Games, Inc. All Rights Reserved.
+
+#include "VulkanRHIPrivate.h"
+#include "VulkanGenericPlatform.h"
+
+void FVulkanGenericPlatform::SetupFeatureLevels()
+{
+	GShaderPlatformForFeatureLevel[ERHIFeatureLevel::ES2] = SP_VULKAN_PCES3_1;
+	GShaderPlatformForFeatureLevel[ERHIFeatureLevel::ES3_1] = SP_VULKAN_PCES3_1;
+	GShaderPlatformForFeatureLevel[ERHIFeatureLevel::SM4] = SP_VULKAN_SM4;
+	GShaderPlatformForFeatureLevel[ERHIFeatureLevel::SM5] = SP_VULKAN_SM5;
+}
diff --git a/Engine/Source/Runtime/VulkanRHI/Private/VulkanGenericPlatform.h b/Engine/Source/Runtime/VulkanRHI/Private/VulkanGenericPlatform.h
--- a/Engine/Source/Runtime/VulkanRHI/Private/VulkanGenericPlatform.h
+++ b/Engine/Source/Runtime/VulkanRHI/Private/VulkanGenericPlatform.h
@@ -35,13 +35,7 @@ public:
 	// most platforms can query the surface for the present mode, and size, etc
 	static bool SupportsQuerySurfaceProperties() { return true; }
 
-	static void SetupFeatureLevels()
-	{
-		GShaderPlatformForFeatureLevel[ERHIFeatureLevel::ES2] = SP_VULKAN_PCES3_1;
-		GShaderPlatformForFeatureLevel[ERHIFeatureLevel::ES3_1] = SP_VULKAN_PCES3_1;
-		GShaderPlatformForFeatureLevel[ERHIFeatureLevel::SM4] = SP_VULKAN_SM4;
-		GShaderPlatformForFeatureLevel[ERHIFeatureLevel::SM5] = SP_VULKAN_SM5;
-	}
+	static void SetupFeatureLevels();
 
 	static bool SupportsStandardSwapchain() { return true; }
 	static EPixelFormat GetPixelFormatForNonDefaultSwapchain()
