UE4.22 - Remove debug code
#rnx

diff --git a/Engine/Source/Runtime/UtilityShaders/Private/OneColorShader.cpp b/Engine/Source/Runtime/UtilityShaders/Private/OneColorShader.cpp
--- a/Engine/Source/Runtime/UtilityShaders/Private/OneColorShader.cpp
+++ b/Engine/Source/Runtime/UtilityShaders/Private/OneColorShader.cpp
@@ -3,9 +3,6 @@
 
 #include "OneColorShader.h"
 
-PRAGMA_DISABLE_OPTIMIZATION
-PRAGMA_DISABLE_INLINING
-
 BEGIN_GLOBAL_SHADER_PARAMETER_STRUCT(FClearShaderUB, )
 	SHADER_PARAMETER_ARRAY(FVector4, DrawColorMRT, [8] )
 END_GLOBAL_SHADER_PARAMETER_STRUCT()
diff --git a/Engine/Source/Runtime/VulkanRHI/Private/VulkanCommands.cpp b/Engine/Source/Runtime/VulkanRHI/Private/VulkanCommands.cpp
--- a/Engine/Source/Runtime/VulkanRHI/Private/VulkanCommands.cpp
+++ b/Engine/Source/Runtime/VulkanRHI/Private/VulkanCommands.cpp
@@ -10,9 +10,6 @@
 #include "EngineGlobals.h"
 #include "VulkanLLM.h"
 
-PRAGMA_DISABLE_OPTIMIZATION
-PRAGMA_DISABLE_INLINING
-
 static TAutoConsoleVariable<int32> GCVarSubmitOnDispatch(
 	TEXT("r.Vulkan.SubmitOnDispatch"),
 	0,
diff --git a/Engine/Source/Runtime/VulkanRHI/Private/VulkanRHI.cpp b/Engine/Source/Runtime/VulkanRHI/Private/VulkanRHI.cpp
--- a/Engine/Source/Runtime/VulkanRHI/Private/VulkanRHI.cpp
+++ b/Engine/Source/Runtime/VulkanRHI/Private/VulkanRHI.cpp
@@ -19,9 +19,6 @@
 #include "Misc/EngineVersion.h"
 #include "GlobalShader.h"
 
-PRAGMA_DISABLE_OPTIMIZATION
-PRAGMA_DISABLE_INLINING
-
 extern RHI_API bool GUseTexture3DBulkDataRHI;
 
 TAtomic<uint64> GVulkanBufferHandleIdCounter{ 0 };
@@ -51,7 +48,7 @@ static_assert(VK_API_VERSION >= UE_VK_API_VERSION, "Vulkan SDK is older than the
 
 TAutoConsoleVariable<int32> GRHIThreadCvar(
 	TEXT("r.Vulkan.RHIThread"),
-	0,
+	1,
 	TEXT("0 to only use Render Thread\n")
 	TEXT("1 to use ONE RHI Thread\n")
 	TEXT("2 to use multiple RHI Thread\n")
