fix for vulkan movie export not working. Resolve destination in non-msaa rendering was being ignored.

diff --git a/Engine/Source/Runtime/VulkanRHI/Private/VulkanRenderTarget.cpp b/Engine/Source/Runtime/VulkanRHI/Private/VulkanRenderTarget.cpp
--- a/Engine/Source/Runtime/VulkanRHI/Private/VulkanRenderTarget.cpp
+++ b/Engine/Source/Runtime/VulkanRHI/Private/VulkanRenderTarget.cpp
@@ -1743,7 +1743,24 @@ void FVulkanCommandListContext::RHIEndRenderPass()
 	{
 		TransitionAndLayoutManager.EndRealRenderPass(CmdBuffer);
 	}
-
+	if(!RenderPassInfo.bIsMSAA)
+	{
+		for (int32 Index = 0; Index < MaxSimultaneousRenderTargets; ++Index)
+		{
+			if (!RenderPassInfo.ColorRenderTargets[Index].RenderTarget)
+			{
+				break;
+			}
+			if (RenderPassInfo.ColorRenderTargets[Index].ResolveTarget)
+			{
+				RHICopyToResolveTarget(RenderPassInfo.ColorRenderTargets[Index].RenderTarget, RenderPassInfo.ColorRenderTargets[Index].ResolveTarget, RenderPassInfo.ResolveParameters);
+			}
+		}
+		if (RenderPassInfo.DepthStencilRenderTarget.DepthStencilTarget && RenderPassInfo.DepthStencilRenderTarget.ResolveTarget)
+		{
+			RHICopyToResolveTarget(RenderPassInfo.DepthStencilRenderTarget.DepthStencilTarget, RenderPassInfo.DepthStencilRenderTarget.ResolveTarget, RenderPassInfo.ResolveParameters);
+		}
+	}
 	RHIPopEvent();
 }
 
