UE4.22 - Workaround for Lumin media textures
#rnx

diff --git a/Engine/Plugins/Lumin/MagicLeap/Source/MagicLeapHelperVulkan/Private/MagicLeapHelperVulkan.cpp b/Engine/Plugins/Lumin/MagicLeap/Source/MagicLeapHelperVulkan/Private/MagicLeapHelperVulkan.cpp
--- a/Engine/Plugins/Lumin/MagicLeap/Source/MagicLeapHelperVulkan/Private/MagicLeapHelperVulkan.cpp
+++ b/Engine/Plugins/Lumin/MagicLeap/Source/MagicLeapHelperVulkan/Private/MagicLeapHelperVulkan.cpp
@@ -278,6 +278,13 @@ bool FMagicLeapHelperVulkan::GetMediaTexture(FTextureRHIRef& ResultTexture, FSam
 		SamplerResult = RHI->RHICreateSamplerState(SamplerStateInitializer, ConversionInitializer);
 	}
 
+	// Insert the RHI thread lock fence. This stops any parallel translate tasks running until the command above has completed on the RHI thread.
+	// There's an odd edge case where parallel rendering is trying to access the RHI's layout map and the command to add it hasn't completed;
+	// wait for the RHI thread while we investigate the root cause of this issue.
+	FRHICommandListImmediate& RHICmdList = GetImmediateCommandList_ForRenderCommand();
+	FGraphEventRef Fence = RHICmdList.RHIThreadFence(true);
+	FRHICommandListExecutor::WaitOnRHIThreadFence(Fence);
+
 	return true;
 #endif //PLATFORM_LUMIN
 	return false;
diff --git a/Engine/Source/Runtime/VulkanRHI/Private/VulkanTexture.cpp b/Engine/Source/Runtime/VulkanRHI/Private/VulkanTexture.cpp
--- a/Engine/Source/Runtime/VulkanRHI/Private/VulkanTexture.cpp
+++ b/Engine/Source/Runtime/VulkanRHI/Private/VulkanTexture.cpp
@@ -442,6 +442,27 @@ struct FRHICommandRegisterImageLayout final : public FRHICommand<FRHICommandRegi
 	}
 };
 
+static void InsertInitialImageLayout(FVulkanDevice& Device, VkImage InImage, VkImageLayout InLayout)
+{
+	FRHICommandListImmediate& RHICmdList = FRHICommandListExecutor::GetImmediateCommandList();
+	const bool bIsInRenderingThread = IsInRenderingThread();
+	if (!bIsInRenderingThread || (RHICmdList.Bypass() || !IsRunningRHIInSeparateThread()))
+	{
+		Device.GetImmediateContext().FindOrAddLayout(InImage, InLayout);
+	}
+	else
+	{
+		check(IsInRenderingThread());
+		new (RHICmdList.AllocCommand<FRHICommandRegisterImageLayout>()) FRHICommandRegisterImageLayout(InImage, InLayout);
+	}
+
+	if (bIsInRenderingThread)
+	{
+		// Insert the RHI thread lock fence. This stops any parallel translate tasks running until the command above has completed on the RHI thread.
+		RHICmdList.RHIThreadFence(true);
+	}
+}
+
 struct FRHICommandOnDestroyImage final : public FRHICommand<FRHICommandOnDestroyImage>
 {
 	VkImage Image;
@@ -1593,15 +1614,9 @@ FVulkanTextureBase::FVulkanTextureBase(FVulkanDevice& Device, VkImageViewType Re
 	if (!CreateInfo.BulkData)
 	{
 		FRHICommandList& RHICmdList = FRHICommandListExecutor::GetImmediateCommandList();
-		if (!IsInRenderingThread() || (RHICmdList.Bypass() || !IsRunningRHIInSeparateThread()))
-		{
-			Device.GetImmediateContext().FindOrAddLayout(Surface.Image, VK_IMAGE_LAYOUT_UNDEFINED);
-		}
-		else
-		{
-			check(IsInRenderingThread());
-			new (RHICmdList.AllocCommand<FRHICommandRegisterImageLayout>()) FRHICommandRegisterImageLayout(Surface.Image, VK_IMAGE_LAYOUT_UNDEFINED);
-		}
+
+		// No initial data, so undefined
+		InsertInitialImageLayout(Device, Surface.Image, VK_IMAGE_LAYOUT_UNDEFINED);
 		return;
 	}
 
@@ -1720,17 +1735,8 @@ FVulkanTextureBase::FVulkanTextureBase(FVulkanDevice& Device, VkImageViewType Re
 		PartialView->Create(Device, Surface.Image, Surface.ViewType, Surface.PartialAspectMask, Surface.PixelFormat, Surface.ViewFormat, 0, FMath::Max(NumMips, 1u), 0, FMath::Max(1u, SizeZ), ConversionInitializer);
 	}
 
-
-	FRHICommandList& RHICmdList = FRHICommandListExecutor::GetImmediateCommandList();
-	if (!IsInRenderingThread() || (RHICmdList.Bypass() || !IsRunningRHIInSeparateThread()))
-	{
-		Device.GetImmediateContext().FindOrAddLayout(InImage, VK_IMAGE_LAYOUT_SHADER_READ_ONLY_OPTIMAL);
-	}
-	else
-	{
-		check(IsInRenderingThread());
-		new (RHICmdList.AllocCommand<FRHICommandRegisterImageLayout>()) FRHICommandRegisterImageLayout(InImage, VK_IMAGE_LAYOUT_SHADER_READ_ONLY_OPTIMAL);
-	}
+	// Since this is provided from an external image, assume it's ready for read
+	InsertInitialImageLayout(Device, InImage, VK_IMAGE_LAYOUT_SHADER_READ_ONLY_OPTIMAL);
 }
 
 FVulkanTextureBase::~FVulkanTextureBase()
