Separate supported flags for images and buffers in Vulkan. This should fix VK crashes on Note9.

diff --git a/Engine/Source/Runtime/VulkanRHI/Private/VulkanDevice.cpp b/Engine/Source/Runtime/VulkanRHI/Private/VulkanDevice.cpp
--- a/Engine/Source/Runtime/VulkanRHI/Private/VulkanDevice.cpp
+++ b/Engine/Source/Runtime/VulkanRHI/Private/VulkanDevice.cpp
@@ -36,6 +36,9 @@ static TAutoConsoleVariable<int32> GCVarRobustBufferAccess(
 	ECVF_ReadOnly
 );
 
+// Mirror GPixelFormats with format information for buffers
+VkFormat GVulkanBufferFormat[PF_MAX];
+
 EDelayAcquireImageType GVulkanDelayAcquireImage = EDelayAcquireImageType::DelayAcquire;
 
 TAutoConsoleVariable<int32> CVarDelayAcquireBackBuffer(
@@ -425,7 +428,8 @@ void FVulkanDevice::SetupFormats()
 	{
 		GPixelFormats[Index].PlatformFormat = VK_FORMAT_UNDEFINED;
 		GPixelFormats[Index].Supported = false;
-
+		GVulkanBufferFormat[Index] = VK_FORMAT_UNDEFINED;
+		
 		// Set default component mapping
 		VkComponentMapping& ComponentMapping = PixelFormatComponentMapping[Index];
 		ComponentMapping.r = VK_COMPONENT_SWIZZLE_R;
@@ -441,7 +445,7 @@ void FVulkanDevice::SetupFormats()
 	MapFormatSupport(PF_G8, VK_FORMAT_R8_UNORM);
 	SetComponentMapping(PF_G8, VK_COMPONENT_SWIZZLE_R, VK_COMPONENT_SWIZZLE_ZERO, VK_COMPONENT_SWIZZLE_ZERO, VK_COMPONENT_SWIZZLE_ZERO);
 
-	MapFormatSupport(PF_G16, VK_FORMAT_R16_UNORM);
+	MapFormatSupportWithFallback(PF_G16, VK_FORMAT_R16_UNORM, {VK_FORMAT_R16_SFLOAT});
 	SetComponentMapping(PF_G16, VK_COMPONENT_SWIZZLE_R, VK_COMPONENT_SWIZZLE_ZERO, VK_COMPONENT_SWIZZLE_ZERO, VK_COMPONENT_SWIZZLE_ZERO);
 
 	MapFormatSupport(PF_FloatRGB, VK_FORMAT_B10G11R11_UFLOAT_PACK32);
@@ -450,19 +454,7 @@ void FVulkanDevice::SetupFormats()
 	MapFormatSupport(PF_FloatRGBA, VK_FORMAT_R16G16B16A16_SFLOAT, 8);
 	SetComponentMapping(PF_FloatRGBA, VK_COMPONENT_SWIZZLE_R, VK_COMPONENT_SWIZZLE_G, VK_COMPONENT_SWIZZLE_B, VK_COMPONENT_SWIZZLE_A);
 
-	MapFormatSupport(PF_DepthStencil, VK_FORMAT_D32_SFLOAT_S8_UINT);
-	if (!GPixelFormats[PF_DepthStencil].Supported)
-	{
-		MapFormatSupport(PF_DepthStencil, VK_FORMAT_D24_UNORM_S8_UINT);
-		if (!GPixelFormats[PF_DepthStencil].Supported)
-		{
-			MapFormatSupport(PF_DepthStencil, VK_FORMAT_D16_UNORM_S8_UINT);
-			if (!GPixelFormats[PF_DepthStencil].Supported)
-			{
-				UE_LOG(LogVulkanRHI, Error, TEXT("No stencil texture format supported!"));
-			}
-		}
-	}
+	MapFormatSupportWithFallback(PF_DepthStencil, VK_FORMAT_D32_SFLOAT_S8_UINT, {VK_FORMAT_D24_UNORM_S8_UINT, VK_FORMAT_D16_UNORM_S8_UINT});
 	SetComponentMapping(PF_DepthStencil, VK_COMPONENT_SWIZZLE_R, VK_COMPONENT_SWIZZLE_IDENTITY, VK_COMPONENT_SWIZZLE_IDENTITY, VK_COMPONENT_SWIZZLE_IDENTITY);
 
 	MapFormatSupport(PF_ShadowDepth, VK_FORMAT_D16_UNORM);
@@ -475,7 +467,7 @@ void FVulkanDevice::SetupFormats()
 	MapFormatSupport(PF_A32B32G32R32F, VK_FORMAT_R32G32B32A32_SFLOAT, 16);
 	SetComponentMapping(PF_A32B32G32R32F, VK_COMPONENT_SWIZZLE_R, VK_COMPONENT_SWIZZLE_G, VK_COMPONENT_SWIZZLE_B, VK_COMPONENT_SWIZZLE_A);
 
-	MapFormatSupport(PF_G16R16, VK_FORMAT_R16G16_UNORM);
+	MapFormatSupportWithFallback(PF_G16R16, VK_FORMAT_R16G16_UNORM, {VK_FORMAT_R16G16_SFLOAT});
 	SetComponentMapping(PF_G16R16, VK_COMPONENT_SWIZZLE_R, VK_COMPONENT_SWIZZLE_G, VK_COMPONENT_SWIZZLE_ZERO, VK_COMPONENT_SWIZZLE_ZERO);
 
 	MapFormatSupport(PF_G16R16F, VK_FORMAT_R16G16_SFLOAT);
@@ -499,28 +491,7 @@ void FVulkanDevice::SetupFormats()
 	MapFormatSupport(PF_R8_UINT, VK_FORMAT_R8_UINT);
 	SetComponentMapping(PF_R8_UINT, VK_COMPONENT_SWIZZLE_R, VK_COMPONENT_SWIZZLE_ZERO, VK_COMPONENT_SWIZZLE_ZERO, VK_COMPONENT_SWIZZLE_ZERO);
 
-	MapFormatSupport(PF_D24, VK_FORMAT_X8_D24_UNORM_PACK32);
-	if (!GPixelFormats[PF_D24].Supported)
-	{
-		MapFormatSupport(PF_D24, VK_FORMAT_D24_UNORM_S8_UINT);
-		if (!GPixelFormats[PF_D24].Supported)
-		{
-			MapFormatSupport(PF_D24, VK_FORMAT_D16_UNORM_S8_UINT);
-			if (!GPixelFormats[PF_D24].Supported)
-			{
-				MapFormatSupport(PF_D24, VK_FORMAT_D32_SFLOAT);
-				if (!GPixelFormats[PF_D24].Supported)
-				{
-					MapFormatSupport(PF_D24, VK_FORMAT_D32_SFLOAT_S8_UINT);
-					if (!GPixelFormats[PF_D24].Supported)
-					{
-						MapFormatSupport(PF_D24, VK_FORMAT_D16_UNORM);
-					}
-				}
-			}
-		}
-	}
-
+	MapFormatSupportWithFallback(PF_D24, VK_FORMAT_X8_D24_UNORM_PACK32, {VK_FORMAT_D24_UNORM_S8_UINT, VK_FORMAT_D16_UNORM_S8_UINT, VK_FORMAT_D32_SFLOAT, VK_FORMAT_D32_SFLOAT_S8_UINT});
 	SetComponentMapping(PF_D24, VK_COMPONENT_SWIZZLE_R, VK_COMPONENT_SWIZZLE_ZERO, VK_COMPONENT_SWIZZLE_ZERO, VK_COMPONENT_SWIZZLE_ZERO);
 
 	MapFormatSupport(PF_R16F, VK_FORMAT_R16_SFLOAT);
@@ -535,7 +506,8 @@ void FVulkanDevice::SetupFormats()
 	MapFormatSupport(PF_A2B10G10R10, VK_FORMAT_A2B10G10R10_UNORM_PACK32, 4);
 	SetComponentMapping(PF_A2B10G10R10, VK_COMPONENT_SWIZZLE_R, VK_COMPONENT_SWIZZLE_G, VK_COMPONENT_SWIZZLE_B, VK_COMPONENT_SWIZZLE_A);
 
-	MapFormatSupport(PF_A16B16G16R16, VK_FORMAT_R16G16B16A16_UNORM, 8);
+	MapFormatSupportWithFallback(PF_A16B16G16R16, VK_FORMAT_R16G16B16A16_UNORM, {VK_FORMAT_R16G16B16A16_SFLOAT});
+	GPixelFormats[PF_A16B16G16R16].BlockBytes = 8;
 	SetComponentMapping(PF_A16B16G16R16, VK_COMPONENT_SWIZZLE_R, VK_COMPONENT_SWIZZLE_G, VK_COMPONENT_SWIZZLE_B, VK_COMPONENT_SWIZZLE_A);
 
 	MapFormatSupport(PF_A8, VK_FORMAT_R8_UNORM);
@@ -565,10 +537,10 @@ void FVulkanDevice::SetupFormats()
 	MapFormatSupport(PF_R32G32B32A32_UINT, VK_FORMAT_R32G32B32A32_UINT);
 	SetComponentMapping(PF_R32G32B32A32_UINT, VK_COMPONENT_SWIZZLE_R, VK_COMPONENT_SWIZZLE_G, VK_COMPONENT_SWIZZLE_B, VK_COMPONENT_SWIZZLE_A);
 
-	MapFormatSupport(PF_R16G16B16A16_SNORM, VK_FORMAT_R16G16B16A16_SNORM);
+	MapFormatSupportWithFallback(PF_R16G16B16A16_SNORM, VK_FORMAT_R16G16B16A16_SNORM, {VK_FORMAT_R16G16B16A16_SFLOAT});
 	SetComponentMapping(PF_R16G16B16A16_SNORM, VK_COMPONENT_SWIZZLE_R, VK_COMPONENT_SWIZZLE_G, VK_COMPONENT_SWIZZLE_B, VK_COMPONENT_SWIZZLE_A);
 
-	MapFormatSupport(PF_R16G16B16A16_UNORM, VK_FORMAT_R16G16B16A16_UNORM);
+	MapFormatSupportWithFallback(PF_R16G16B16A16_UNORM, VK_FORMAT_R16G16B16A16_UNORM, {VK_FORMAT_R16G16B16A16_SFLOAT});
 	SetComponentMapping(PF_R16G16B16A16_UNORM, VK_COMPONENT_SWIZZLE_R, VK_COMPONENT_SWIZZLE_G, VK_COMPONENT_SWIZZLE_B, VK_COMPONENT_SWIZZLE_A);
 
 	MapFormatSupport(PF_R8G8, VK_FORMAT_R8G8_UNORM);
@@ -661,7 +633,7 @@ void FVulkanDevice::SetupFormats()
 	for (int32 Index = (int32)VET_None + 1; Index < VET_MAX; ++Index)
 	{
 		EVertexElementType UEType = (EVertexElementType)Index;
-		VkFormat VulkanFormat = UEToVkFormat(UEType);
+		VkFormat VulkanFormat = UEToVkBufferFormat(UEType);
 		if (!IsBufferFormatSupported(VulkanFormat))
 		{
 			UE_LOG(LogVulkanRHI, Warning, TEXT("EVertexFormat(%d) is not supported with Vk format %d"), (int32)UEType, (int32)VulkanFormat);
@@ -690,13 +662,40 @@ VkSamplerYcbcrConversion FVulkanDevice::CreateSamplerColorConversion(const VkSam
 
 void FVulkanDevice::MapFormatSupport(EPixelFormat UEFormat, VkFormat VulkanFormat)
 {
-	FPixelFormatInfo& FormatInfo = GPixelFormats[UEFormat];
-	FormatInfo.PlatformFormat = VulkanFormat;
-	FormatInfo.Supported = IsTextureFormatSupported(VulkanFormat);
+	MapFormatSupportWithFallback(UEFormat, VulkanFormat, TArrayView<VkFormat>());
+}
 
+void FVulkanDevice::MapFormatSupportWithFallback(EPixelFormat UEFormat, VkFormat VulkanFormat, TArrayView<const VkFormat> FallbackTextureFormats)
+{
+	VkFormat SupportedTextureFormat = IsTextureFormatSupported(VulkanFormat) ? VulkanFormat : VK_FORMAT_UNDEFINED;
+	VkFormat SupportedBufferFormat = IsBufferFormatSupported(VulkanFormat) ? VulkanFormat : VK_FORMAT_UNDEFINED;
+	
+	FPixelFormatInfo& FormatInfo = GPixelFormats[UEFormat];
+	// at this point we don't know if high level code will use this pixel format for buffers or textures
+	FormatInfo.Supported = (SupportedTextureFormat!= VK_FORMAT_UNDEFINED || SupportedBufferFormat!= VK_FORMAT_UNDEFINED);
+	FormatInfo.PlatformFormat = SupportedTextureFormat;
+	
+	GVulkanBufferFormat[UEFormat] = SupportedBufferFormat;
+		
+	if (SupportedTextureFormat == VK_FORMAT_UNDEFINED)
+	{
+		for (int32 Idx = 0; Idx < FallbackTextureFormats.Num(); ++Idx)
+		{
+			VkFormat FallbackTextureFormat = FallbackTextureFormats[Idx];
+			if (IsTextureFormatSupported(FallbackTextureFormat))
+			{
+				SupportedTextureFormat = FallbackTextureFormat;
+				FormatInfo.PlatformFormat = FallbackTextureFormat;
+				FormatInfo.Supported = true;
+								
+				UE_LOG(LogVulkanRHI, Display, TEXT("EPixelFormat(%d) (images) is not supported with Vk format %d, falling back to Vk format %d"), (int32)UEFormat, (int32)VulkanFormat, (int32)FallbackTextureFormat);
+			}
+		}
+	}
+			
 	if (!FormatInfo.Supported)
 	{
-		UE_LOG(LogVulkanRHI, Warning, TEXT("EPixelFormat(%d) is not supported with Vk format %d"), (int32)UEFormat, (int32)VulkanFormat);
+		UE_LOG(LogVulkanRHI, Error, TEXT("EPixelFormat(%d) is not supported with Vk format %d"), (int32)UEFormat, (int32)VulkanFormat);
 	}
 }
 
diff --git a/Engine/Source/Runtime/VulkanRHI/Private/VulkanDevice.h b/Engine/Source/Runtime/VulkanRHI/Private/VulkanDevice.h
--- a/Engine/Source/Runtime/VulkanRHI/Private/VulkanDevice.h
+++ b/Engine/Source/Runtime/VulkanRHI/Private/VulkanDevice.h
@@ -301,6 +301,7 @@ public:
 
 private:
 	void MapFormatSupport(EPixelFormat UEFormat, VkFormat VulkanFormat);
+	void MapFormatSupportWithFallback(EPixelFormat UEFormat, VkFormat VulkanFormat, TArrayView<const VkFormat> FallbackTextureFormats);
 	void MapFormatSupport(EPixelFormat UEFormat, VkFormat VulkanFormat, int32 BlockBytes);
 	void SetComponentMapping(EPixelFormat UEFormat, VkComponentSwizzle r, VkComponentSwizzle g, VkComponentSwizzle b, VkComponentSwizzle a);
 
diff --git a/Engine/Source/Runtime/VulkanRHI/Private/VulkanRHI.cpp b/Engine/Source/Runtime/VulkanRHI/Private/VulkanRHI.cpp
--- a/Engine/Source/Runtime/VulkanRHI/Private/VulkanRHI.cpp
+++ b/Engine/Source/Runtime/VulkanRHI/Private/VulkanRHI.cpp
@@ -1320,13 +1320,13 @@ void FVulkanBufferView::Create(FVulkanBuffer& Buffer, EPixelFormat Format, uint3
 	Offset = InOffset;
 	Size = InSize;
 	check(Format != PF_Unknown);
-	const FPixelFormatInfo& FormatInfo = GPixelFormats[Format];
-	check(FormatInfo.Supported);
+	VkFormat BufferFormat = GVulkanBufferFormat[Format];
+	check(BufferFormat != VK_FORMAT_UNDEFINED);
 
 	VkBufferViewCreateInfo ViewInfo;
 	ZeroVulkanStruct(ViewInfo, VK_STRUCTURE_TYPE_BUFFER_VIEW_CREATE_INFO);
 	ViewInfo.buffer = Buffer.GetBufferHandle();
-	ViewInfo.format = (VkFormat)FormatInfo.PlatformFormat;
+	ViewInfo.format = BufferFormat;
 	ViewInfo.offset = Offset;
 	ViewInfo.range = Size;
 	Flags = Buffer.GetFlags() & VK_BUFFER_USAGE_UNIFORM_TEXEL_BUFFER_BIT;
@@ -1345,9 +1345,9 @@ void FVulkanBufferView::Create(FVulkanBuffer& Buffer, EPixelFormat Format, uint3
 void FVulkanBufferView::Create(FVulkanResourceMultiBuffer* Buffer, EPixelFormat Format, uint32 InOffset, uint32 InSize)
 {
 	check(Format != PF_Unknown);
-	const FPixelFormatInfo& FormatInfo = GPixelFormats[Format];
-	check(FormatInfo.Supported);
-	Create((VkFormat)FormatInfo.PlatformFormat, Buffer, InOffset, InSize);
+	VkFormat BufferFormat = GVulkanBufferFormat[Format];
+	check(BufferFormat != VK_FORMAT_UNDEFINED);
+	Create(BufferFormat, Buffer, InOffset, InSize);
 }
 
 void FVulkanBufferView::Create(VkFormat Format, FVulkanResourceMultiBuffer* Buffer, uint32 InOffset, uint32 InSize)
diff --git a/Engine/Source/Runtime/VulkanRHI/Private/VulkanRHIPrivate.h b/Engine/Source/Runtime/VulkanRHI/Private/VulkanRHIPrivate.h
--- a/Engine/Source/Runtime/VulkanRHI/Private/VulkanRHIPrivate.h
+++ b/Engine/Source/Runtime/VulkanRHI/Private/VulkanRHIPrivate.h
@@ -615,7 +615,7 @@ static inline VkAttachmentStoreOp RenderTargetStoreActionToVulkan(ERenderTargetS
 	return OutStoreAction;
 }
 
-inline VkFormat UEToVkFormat(EPixelFormat UEFormat, const bool bIsSRGB)
+inline VkFormat UEToVkTextureFormat(EPixelFormat UEFormat, const bool bIsSRGB)
 {
 	VkFormat Format = (VkFormat)GPixelFormats[UEFormat].PlatformFormat;
 	if (bIsSRGB && GMaxRHIFeatureLevel > ERHIFeatureLevel::ES2)
@@ -661,7 +661,7 @@ inline VkFormat UEToVkFormat(EPixelFormat UEFormat, const bool bIsSRGB)
 	return Format;
 }
 
-static inline VkFormat UEToVkFormat(EVertexElementType Type)
+static inline VkFormat UEToVkBufferFormat(EVertexElementType Type)
 {
 	switch (Type)
 	{
diff --git a/Engine/Source/Runtime/VulkanRHI/Private/VulkanRenderTarget.cpp b/Engine/Source/Runtime/VulkanRHI/Private/VulkanRenderTarget.cpp
--- a/Engine/Source/Runtime/VulkanRHI/Private/VulkanRenderTarget.cpp
+++ b/Engine/Source/Runtime/VulkanRHI/Private/VulkanRenderTarget.cpp
@@ -1851,7 +1851,7 @@ FVulkanRenderTargetLayout::FVulkanRenderTargetLayout(FVulkanDevice& InDevice, co
 		
 			VkAttachmentDescription& CurrDesc = Desc[NumAttachmentDescriptions];
 			CurrDesc.samples = static_cast<VkSampleCountFlagBits>(NumSamples);
-			CurrDesc.format = UEToVkFormat(RTView.Texture->GetFormat(), (Texture->Surface.UEFlags & TexCreate_SRGB) == TexCreate_SRGB);
+			CurrDesc.format = UEToVkTextureFormat(RTView.Texture->GetFormat(), (Texture->Surface.UEFlags & TexCreate_SRGB) == TexCreate_SRGB);
 			CurrDesc.loadOp = RenderTargetLoadActionToVulkan(RTView.LoadAction);
 			bFoundClearOp = bFoundClearOp || (CurrDesc.loadOp == VK_ATTACHMENT_LOAD_OP_CLEAR);
 			CurrDesc.storeOp = RenderTargetStoreActionToVulkan(RTView.StoreAction);
@@ -1904,7 +1904,7 @@ FVulkanRenderTargetLayout::FVulkanRenderTargetLayout(FVulkanDevice& InDevice, co
 		NumSamples = Surface->GetNumSamples();
 
 		CurrDesc.samples = static_cast<VkSampleCountFlagBits>(NumSamples);
-		CurrDesc.format = UEToVkFormat(RTInfo.DepthStencilRenderTarget.Texture->GetFormat(), false);
+		CurrDesc.format = UEToVkTextureFormat(RTInfo.DepthStencilRenderTarget.Texture->GetFormat(), false);
 		CurrDesc.loadOp = RenderTargetLoadActionToVulkan(RTInfo.DepthStencilRenderTarget.DepthLoadAction);
 		CurrDesc.stencilLoadOp = RenderTargetLoadActionToVulkan(RTInfo.DepthStencilRenderTarget.StencilLoadAction);
 		bFoundClearOp = bFoundClearOp || (CurrDesc.loadOp == VK_ATTACHMENT_LOAD_OP_CLEAR || CurrDesc.stencilLoadOp == VK_ATTACHMENT_LOAD_OP_CLEAR);
@@ -2022,7 +2022,7 @@ FVulkanRenderTargetLayout::FVulkanRenderTargetLayout(FVulkanDevice& InDevice, co
 
 		VkAttachmentDescription& CurrDesc = Desc[NumAttachmentDescriptions];
 		CurrDesc.samples = static_cast<VkSampleCountFlagBits>(NumSamples);
-		CurrDesc.format = UEToVkFormat(ColorEntry.RenderTarget->GetFormat(), (Texture->Surface.UEFlags & TexCreate_SRGB) == TexCreate_SRGB);
+		CurrDesc.format = UEToVkTextureFormat(ColorEntry.RenderTarget->GetFormat(), (Texture->Surface.UEFlags & TexCreate_SRGB) == TexCreate_SRGB);
 		CurrDesc.loadOp = RenderTargetLoadActionToVulkan(GetLoadAction(ColorEntry.Action));
 		bFoundClearOp = bFoundClearOp || (CurrDesc.loadOp == VK_ATTACHMENT_LOAD_OP_CLEAR);
 		CurrDesc.storeOp = RenderTargetStoreActionToVulkan(GetStoreAction(ColorEntry.Action), true);
@@ -2068,7 +2068,7 @@ FVulkanRenderTargetLayout::FVulkanRenderTargetLayout(FVulkanDevice& InDevice, co
 		CurrDesc.samples = static_cast<VkSampleCountFlagBits>(RPInfo.DepthStencilRenderTarget.DepthStencilTarget->GetNumSamples());
 		ensure(!NumSamples || CurrDesc.samples == NumSamples);
 		NumSamples = CurrDesc.samples;
-		CurrDesc.format = UEToVkFormat(RPInfo.DepthStencilRenderTarget.DepthStencilTarget->GetFormat(), false);
+		CurrDesc.format = UEToVkTextureFormat(RPInfo.DepthStencilRenderTarget.DepthStencilTarget->GetFormat(), false);
 		CurrDesc.loadOp = RenderTargetLoadActionToVulkan(GetLoadAction(GetDepthActions(RPInfo.DepthStencilRenderTarget.Action)));
 		CurrDesc.stencilLoadOp = RenderTargetLoadActionToVulkan(GetLoadAction(GetStencilActions(RPInfo.DepthStencilRenderTarget.Action)));
 		bFoundClearOp = bFoundClearOp || (CurrDesc.loadOp == VK_ATTACHMENT_LOAD_OP_CLEAR || CurrDesc.stencilLoadOp == VK_ATTACHMENT_LOAD_OP_CLEAR);
@@ -2163,7 +2163,7 @@ FVulkanRenderTargetLayout::FVulkanRenderTargetLayout(const FGraphicsPipelineStat
 		{
 			VkAttachmentDescription& CurrDesc = Desc[NumAttachmentDescriptions];
 			CurrDesc.samples = static_cast<VkSampleCountFlagBits>(NumSamples);
-			CurrDesc.format = UEToVkFormat(UEFormat, (Initializer.RenderTargetFlags[Index] & TexCreate_SRGB) == TexCreate_SRGB);
+			CurrDesc.format = UEToVkTextureFormat(UEFormat, (Initializer.RenderTargetFlags[Index] & TexCreate_SRGB) == TexCreate_SRGB);
 			CurrDesc.loadOp = VK_ATTACHMENT_LOAD_OP_DONT_CARE;
 			CurrDesc.storeOp = VK_ATTACHMENT_STORE_OP_DONT_CARE;
 			CurrDesc.stencilLoadOp = VK_ATTACHMENT_LOAD_OP_DONT_CARE;
@@ -2204,7 +2204,7 @@ FVulkanRenderTargetLayout::FVulkanRenderTargetLayout(const FGraphicsPipelineStat
 		FMemory::Memzero(CurrDesc);
 
 		CurrDesc.samples = static_cast<VkSampleCountFlagBits>(NumSamples);
-		CurrDesc.format = UEToVkFormat(Initializer.DepthStencilTargetFormat, false);
+		CurrDesc.format = UEToVkTextureFormat(Initializer.DepthStencilTargetFormat, false);
 		CurrDesc.loadOp = RenderTargetLoadActionToVulkan(Initializer.DepthTargetLoadAction);
 		CurrDesc.stencilLoadOp = RenderTargetLoadActionToVulkan(Initializer.StencilTargetLoadAction);
 		if (CurrDesc.loadOp == VK_ATTACHMENT_LOAD_OP_CLEAR || CurrDesc.stencilLoadOp == VK_ATTACHMENT_LOAD_OP_CLEAR)
diff --git a/Engine/Source/Runtime/VulkanRHI/Private/VulkanSwapChain.cpp b/Engine/Source/Runtime/VulkanRHI/Private/VulkanSwapChain.cpp
--- a/Engine/Source/Runtime/VulkanRHI/Private/VulkanSwapChain.cpp
+++ b/Engine/Source/Runtime/VulkanRHI/Private/VulkanSwapChain.cpp
@@ -204,7 +204,7 @@ FVulkanSwapChain::FVulkanSwapChain(VkInstance InInstance, FVulkanDevice& InDevic
 		if (InOutPixelFormat == PF_Unknown)
 		{
 			UE_LOG(LogVulkanRHI, Warning, TEXT("Can't find a proper pixel format for the swapchain, trying to pick up the first available"));
-			VkFormat PlatformFormat = UEToVkFormat(InOutPixelFormat, false);
+			VkFormat PlatformFormat = UEToVkTextureFormat(InOutPixelFormat, false);
 			bool bSupported = false;
 			for (int32 Index = 0; Index < Formats.Num(); ++Index)
 			{
@@ -242,7 +242,7 @@ FVulkanSwapChain::FVulkanSwapChain(VkInstance InInstance, FVulkanDevice& InDevic
 		}
 	}
 
-	VkFormat PlatformFormat = UEToVkFormat(InOutPixelFormat, false);
+	VkFormat PlatformFormat = UEToVkTextureFormat(InOutPixelFormat, false);
 
 	Device.SetupPresentQueue(Surface);
 
diff --git a/Engine/Source/Runtime/VulkanRHI/Private/VulkanTexture.cpp b/Engine/Source/Runtime/VulkanRHI/Private/VulkanTexture.cpp
--- a/Engine/Source/Runtime/VulkanRHI/Private/VulkanTexture.cpp
+++ b/Engine/Source/Runtime/VulkanRHI/Private/VulkanTexture.cpp
@@ -190,8 +190,10 @@ VkImage FVulkanSurface::CreateImage(
 	bool bForceLinearTexture)
 {
 	const VkPhysicalDeviceProperties& DeviceProperties = InDevice.GetDeviceProperties();
+	const FPixelFormatInfo& FormatInfo = GPixelFormats[InFormat];
+	VkFormat TextureFormat = (VkFormat)FormatInfo.PlatformFormat;
 
-	checkf(GPixelFormats[InFormat].Supported, TEXT("Format %d"), (int32)InFormat);
+	checkf(TextureFormat != VK_FORMAT_UNDEFINED, TEXT("PixelFormat %d, is not supported for images"), (int32)InFormat);
 
 	VkImageCreateInfo TmpCreateInfo;
 	VkImageCreateInfo* ImageCreateInfoPtr = OutInfo ? OutInfo : &TmpCreateInfo;
@@ -226,7 +228,7 @@ VkImage FVulkanSurface::CreateImage(
 		break;
 	}
 
-	ImageCreateInfo.format = UEToVkFormat(InFormat, false);
+	ImageCreateInfo.format = UEToVkTextureFormat(InFormat, false);
 
 	checkf(ImageCreateInfo.format != VK_FORMAT_UNDEFINED, TEXT("Pixel Format %d not defined!"), (int32)InFormat);
 	if (OutStorageFormat)
@@ -236,7 +238,7 @@ VkImage FVulkanSurface::CreateImage(
 
 	if (OutViewFormat)
 	{
-		VkFormat ViewFormat = UEToVkFormat(InFormat, (UEFlags & TexCreate_SRGB) == TexCreate_SRGB);
+		VkFormat ViewFormat = UEToVkTextureFormat(InFormat, (UEFlags & TexCreate_SRGB) == TexCreate_SRGB);
 		*OutViewFormat = ViewFormat;
 		ImageCreateInfo.format = ViewFormat;
 	}
@@ -590,6 +592,8 @@ FVulkanSurface::FVulkanSurface(FVulkanDevice& InDevice, VkImageViewType Resource
 {
 	StorageFormat = (VkFormat)GPixelFormats[PixelFormat].PlatformFormat;
 	check((UEFlags & TexCreate_SRGB) == 0);
+	checkf(PixelFormat == PF_Unknown || StorageFormat != VK_FORMAT_UNDEFINED, TEXT("PixelFormat %d, is not supported for images"), (int32)PixelFormat);
+
 	ViewFormat = StorageFormat;
 	FullAspectMask = VulkanRHI::GetAspectMaskFromUEFormat(PixelFormat, true, true);
 	PartialAspectMask = VulkanRHI::GetAspectMaskFromUEFormat(PixelFormat, false, true);
@@ -1275,7 +1279,7 @@ void FVulkanDynamicRHI::InternalUpdateTexture2D(bool bFromRenderingThread, FText
 	const int32 BlockSizeY = GPixelFormats[PixelFormat].BlockSizeY;
 	const int32 BlockSizeZ = GPixelFormats[PixelFormat].BlockSizeZ;
 	const int32 BlockBytes = GPixelFormats[PixelFormat].BlockBytes;
-	VkFormat Format = UEToVkFormat(PixelFormat, false);
+	VkFormat Format = UEToVkTextureFormat(PixelFormat, false);
 
 	ensure(BlockSizeZ == 1);
 
@@ -1350,7 +1354,7 @@ void FVulkanDynamicRHI::InternalUpdateTexture3D(bool bFromRenderingThread, FText
 	const int32 BlockSizeY = GPixelFormats[PixelFormat].BlockSizeY;
 	const int32 BlockSizeZ = GPixelFormats[PixelFormat].BlockSizeZ;
 	const int32 BlockBytes = GPixelFormats[PixelFormat].BlockBytes;
-	const VkFormat Format = UEToVkFormat(PixelFormat, false);
+	const VkFormat Format = UEToVkTextureFormat(PixelFormat, false);
 
 	ensure(BlockSizeZ == 1);
 
@@ -1507,6 +1511,7 @@ VkImageView FVulkanTextureView::StaticCreate(FVulkanDevice& Device, VkImage InIm
 	{
 		ensure(ViewInfo.format == VK_FORMAT_UNDEFINED);
 		ViewInfo.format = (VkFormat)GPixelFormats[PF_DepthStencil].PlatformFormat;
+		ensure(ViewInfo.format != VK_FORMAT_UNDEFINED);
 		ViewInfo.subresourceRange.aspectMask = VK_IMAGE_ASPECT_STENCIL_BIT;
 	}
 
@@ -1579,8 +1584,8 @@ FVulkanTextureBase::FVulkanTextureBase(FVulkanDevice& Device, VkImageViewType Re
 	LLM_SCOPE_VULKAN(ELLMTagVulkan::VulkanTextures);
 	if (Surface.ViewFormat == VK_FORMAT_UNDEFINED)
 	{
-		Surface.StorageFormat = UEToVkFormat(InFormat, false);
-		Surface.ViewFormat = UEToVkFormat(InFormat, (UEFlags & TexCreate_SRGB) == TexCreate_SRGB);
+		Surface.StorageFormat = UEToVkTextureFormat(InFormat, false);
+		Surface.ViewFormat = UEToVkTextureFormat(InFormat, (UEFlags & TexCreate_SRGB) == TexCreate_SRGB);
 		checkf(Surface.StorageFormat != VK_FORMAT_UNDEFINED, TEXT("Pixel Format %d not defined!"), (int32)InFormat);
 	}
 
diff --git a/Engine/Source/Runtime/VulkanRHI/Private/VulkanUAV.cpp b/Engine/Source/Runtime/VulkanRHI/Private/VulkanUAV.cpp
--- a/Engine/Source/Runtime/VulkanRHI/Private/VulkanUAV.cpp
+++ b/Engine/Source/Runtime/VulkanRHI/Private/VulkanUAV.cpp
@@ -102,22 +102,22 @@ void FVulkanShaderResourceView::UpdateView()
 			{
 				FVulkanTexture2D* VTex2D = ResourceCast(Tex2D);
 				EPixelFormat OriginalFormat = Format;
-				TextureView.Create(*Device, VTex2D->Surface.Image, VK_IMAGE_VIEW_TYPE_2D, VTex2D->Surface.GetPartialAspectMask(), Format, UEToVkFormat(Format, false), MipLevel, NumMips, 0, 1);
+				TextureView.Create(*Device, VTex2D->Surface.Image, VK_IMAGE_VIEW_TYPE_2D, VTex2D->Surface.GetPartialAspectMask(), Format, UEToVkTextureFormat(Format, false), MipLevel, NumMips, 0, 1);
 			}
 			else if (FRHITextureCube* TexCube = SourceTexture->GetTextureCube())
 			{
 				FVulkanTextureCube* VTexCube = ResourceCast(TexCube);
-				TextureView.Create(*Device, VTexCube->Surface.Image, VK_IMAGE_VIEW_TYPE_CUBE, VTexCube->Surface.GetPartialAspectMask(), Format, UEToVkFormat(Format, false), MipLevel, NumMips, 0, 1);
+				TextureView.Create(*Device, VTexCube->Surface.Image, VK_IMAGE_VIEW_TYPE_CUBE, VTexCube->Surface.GetPartialAspectMask(), Format, UEToVkTextureFormat(Format, false), MipLevel, NumMips, 0, 1);
 			}
 			else if (FRHITexture3D* Tex3D = SourceTexture->GetTexture3D())
 			{
 				FVulkanTexture3D* VTex3D = ResourceCast(Tex3D);
-				TextureView.Create(*Device, VTex3D->Surface.Image, VK_IMAGE_VIEW_TYPE_3D, VTex3D->Surface.GetPartialAspectMask(), Format, UEToVkFormat(Format, false), MipLevel, NumMips, 0, 1);
+				TextureView.Create(*Device, VTex3D->Surface.Image, VK_IMAGE_VIEW_TYPE_3D, VTex3D->Surface.GetPartialAspectMask(), Format, UEToVkTextureFormat(Format, false), MipLevel, NumMips, 0, 1);
 			}
 			else if (FRHITexture2DArray* Tex2DArray = SourceTexture->GetTexture2DArray())
 			{
 				FVulkanTexture2DArray* VTex2DArray = ResourceCast(Tex2DArray);
-				TextureView.Create(*Device, VTex2DArray->Surface.Image, VK_IMAGE_VIEW_TYPE_2D_ARRAY, VTex2DArray->Surface.GetPartialAspectMask(), Format, UEToVkFormat(Format, false), MipLevel, NumMips, 0, VTex2DArray->GetSizeZ());
+				TextureView.Create(*Device, VTex2DArray->Surface.Image, VK_IMAGE_VIEW_TYPE_2D_ARRAY, VTex2DArray->Surface.GetPartialAspectMask(), Format, UEToVkTextureFormat(Format, false), MipLevel, NumMips, 0, VTex2DArray->GetSizeZ());
 			}
 			else
 			{
@@ -189,22 +189,22 @@ void FVulkanUnorderedAccessView::UpdateView()
 		if (FRHITexture2D* Tex2D = SourceTexture->GetTexture2D())
 		{
 			FVulkanTexture2D* VTex2D = ResourceCast(Tex2D);
-			TextureView.Create(*Device, VTex2D->Surface.Image, VK_IMAGE_VIEW_TYPE_2D, VTex2D->Surface.GetPartialAspectMask(), Format, UEToVkFormat(Format, false), MipLevel, 1, 0, 1);
+			TextureView.Create(*Device, VTex2D->Surface.Image, VK_IMAGE_VIEW_TYPE_2D, VTex2D->Surface.GetPartialAspectMask(), Format, UEToVkTextureFormat(Format, false), MipLevel, 1, 0, 1);
 		}
 		else if (FRHITextureCube* TexCube = SourceTexture->GetTextureCube())
 		{
 			FVulkanTextureCube* VTexCube = ResourceCast(TexCube);
-			TextureView.Create(*Device, VTexCube->Surface.Image, VK_IMAGE_VIEW_TYPE_CUBE, VTexCube->Surface.GetPartialAspectMask(), Format, UEToVkFormat(Format, false), MipLevel, 1, 0, 1);
+			TextureView.Create(*Device, VTexCube->Surface.Image, VK_IMAGE_VIEW_TYPE_CUBE, VTexCube->Surface.GetPartialAspectMask(), Format, UEToVkTextureFormat(Format, false), MipLevel, 1, 0, 1);
 		}
 		else if (FRHITexture3D* Tex3D = SourceTexture->GetTexture3D())
 		{
 			FVulkanTexture3D* VTex3D = ResourceCast(Tex3D);
-			TextureView.Create(*Device, VTex3D->Surface.Image, VK_IMAGE_VIEW_TYPE_3D, VTex3D->Surface.GetPartialAspectMask(), Format, UEToVkFormat(Format, false), MipLevel, 1, 0, VTex3D->GetSizeZ());
+			TextureView.Create(*Device, VTex3D->Surface.Image, VK_IMAGE_VIEW_TYPE_3D, VTex3D->Surface.GetPartialAspectMask(), Format, UEToVkTextureFormat(Format, false), MipLevel, 1, 0, VTex3D->GetSizeZ());
 		}
 		else if (FRHITexture2DArray* Tex2DArray = SourceTexture->GetTexture2DArray())
 		{
 			FVulkanTexture2DArray* VTex2DArray = ResourceCast(Tex2DArray);
-			TextureView.Create(*Device, VTex2DArray->Surface.Image, VK_IMAGE_VIEW_TYPE_2D_ARRAY, VTex2DArray->Surface.GetPartialAspectMask(), Format, UEToVkFormat(Format, false), MipLevel, 1, 0, VTex2DArray->GetSizeZ());
+			TextureView.Create(*Device, VTex2DArray->Surface.Image, VK_IMAGE_VIEW_TYPE_2D_ARRAY, VTex2DArray->Surface.GetPartialAspectMask(), Format, UEToVkTextureFormat(Format, false), MipLevel, 1, 0, VTex2DArray->GetSizeZ());
 		}
 		else
 		{
diff --git a/Engine/Source/Runtime/VulkanRHI/Private/VulkanVertexDeclaration.cpp b/Engine/Source/Runtime/VulkanRHI/Private/VulkanVertexDeclaration.cpp
--- a/Engine/Source/Runtime/VulkanRHI/Private/VulkanVertexDeclaration.cpp
+++ b/Engine/Source/Runtime/VulkanRHI/Private/VulkanVertexDeclaration.cpp
@@ -147,7 +147,7 @@ void FVulkanVertexInputStateInfo::Generate(FVulkanVertexDeclaration* VertexDecla
 
 		CurrAttribute.binding = StreamToBinding.FindChecked(CurrElement.StreamIndex);
 		CurrAttribute.location = CurrElement.AttributeIndex;
-		CurrAttribute.format = UEToVkFormat(CurrElement.Type);
+		CurrAttribute.format = UEToVkBufferFormat(CurrElement.Type);
 		CurrAttribute.offset = CurrElement.Offset;
 	}
 
diff --git a/Engine/Source/Runtime/VulkanRHI/Private/VulkanViewport.cpp b/Engine/Source/Runtime/VulkanRHI/Private/VulkanViewport.cpp
--- a/Engine/Source/Runtime/VulkanRHI/Private/VulkanViewport.cpp
+++ b/Engine/Source/Runtime/VulkanRHI/Private/VulkanViewport.cpp
@@ -619,7 +619,7 @@ void FVulkanViewport::CreateSwapchain()
 			FName Name = FName(*FString::Printf(TEXT("BackBuffer%d"), Index));
 			//BackBuffers[Index]->SetName(Name);
 
-			TextureViews[Index].Create(*Device, Images[Index], VK_IMAGE_VIEW_TYPE_2D, VK_IMAGE_ASPECT_COLOR_BIT, PixelFormat, UEToVkFormat(PixelFormat, false), 0, 1, 0, 1);
+			TextureViews[Index].Create(*Device, Images[Index], VK_IMAGE_VIEW_TYPE_2D, VK_IMAGE_ASPECT_COLOR_BIT, PixelFormat, UEToVkTextureFormat(PixelFormat, false), 0, 1, 0, 1);
 
 			// Clear the swapchain to avoid a validation warning, and transition to ColorAttachment
 			{
diff --git a/Engine/Source/Runtime/VulkanRHI/Public/VulkanResources.h b/Engine/Source/Runtime/VulkanRHI/Public/VulkanResources.h
--- a/Engine/Source/Runtime/VulkanRHI/Public/VulkanResources.h
+++ b/Engine/Source/Runtime/VulkanRHI/Public/VulkanResources.h
@@ -52,6 +52,9 @@ struct FSamplerYcbcrConversionInitializer
 	VkChromaLocation YOffset;
 };
 
+// Mirror GPixelFormats with format information for buffers
+extern VkFormat GVulkanBufferFormat[PF_MAX];
+
 /** This represents a vertex declaration that hasn't been combined with a specific shader to create a bound shader. */
 class FVulkanVertexDeclaration : public FRHIVertexDeclaration
 {
