Fix for vulkan hdr video export crashing.
Fix vulkan validation error when exporting videos. Window swap chain was smaller than rendertarget, and a copy was writing out of bounds on the swap chain.

diff --git a/Engine/Source/Runtime/VulkanRHI/Private/VulkanRenderTarget.cpp b/Engine/Source/Runtime/VulkanRHI/Private/VulkanRenderTarget.cpp
--- a/Engine/Source/Runtime/VulkanRHI/Private/VulkanRenderTarget.cpp
+++ b/Engine/Source/Runtime/VulkanRHI/Private/VulkanRenderTarget.cpp
@@ -1093,14 +1093,17 @@ void FVulkanDynamicRHI::RHIReadSurfaceFloatData(FTextureRHIParamRef TextureRHI,
 
 		StagingBuffer->InvalidateMappedMemory();
 
-		OutputData.SetNum(NumPixels);
+		uint32 OutWidth = InRect.Max.X - InRect.Min.X;
+		uint32 OutHeight= InRect.Max.Y - InRect.Min.Y;
+		OutputData.SetNum(OutWidth * OutHeight);
+		uint32 OutIndex = 0;
 		FFloat16Color* Dest = OutputData.GetData();
 		for (int32 Row = InRect.Min.Y; Row < InRect.Max.Y; ++Row)
 		{
 			FFloat16Color* Src = (FFloat16Color*)StagingBuffer->GetMappedPointer() + Row * (Surface.Width >> InMipIndex) + InRect.Min.X;
 			for (int32 Col = InRect.Min.X; Col < InRect.Max.X; ++Col)
 			{
-				*Dest++ = *Src++;
+				OutputData[OutIndex++] = *Src++;
 			}
 		}
 		InDevice->GetStagingManager().ReleaseBuffer(InCmdBuffer, StagingBuffer);
diff --git a/Engine/Source/Runtime/VulkanRHI/Private/VulkanSwapChain.cpp b/Engine/Source/Runtime/VulkanRHI/Private/VulkanSwapChain.cpp
--- a/Engine/Source/Runtime/VulkanRHI/Private/VulkanSwapChain.cpp
+++ b/Engine/Source/Runtime/VulkanRHI/Private/VulkanSwapChain.cpp
@@ -431,6 +431,9 @@ FVulkanSwapChain::FVulkanSwapChain(VkInstance InInstance, FVulkanDevice& InDevic
 
 	VERIFYVULKANRESULT_EXPANDED(VulkanRHI::vkCreateSwapchainKHR(Device.GetInstanceHandle(), &SwapChainInfo, VULKAN_CPU_ALLOCATOR, &SwapChain));
 
+	InternalWidth = FMath::Min(Width, SwapChainInfo.imageExtent.width);
+	InternalHeight = FMath::Min(Height, SwapChainInfo.imageExtent.height);
+
 	uint32 NumSwapChainImages;
 	VERIFYVULKANRESULT_EXPANDED(VulkanRHI::vkGetSwapchainImagesKHR(Device.GetInstanceHandle(), SwapChain, &NumSwapChainImages, nullptr));
 
diff --git a/Engine/Source/Runtime/VulkanRHI/Private/VulkanSwapChain.h b/Engine/Source/Runtime/VulkanRHI/Private/VulkanSwapChain.h
--- a/Engine/Source/Runtime/VulkanRHI/Private/VulkanSwapChain.h
+++ b/Engine/Source/Runtime/VulkanRHI/Private/VulkanSwapChain.h
@@ -43,6 +43,8 @@ protected:
 	int32 SemaphoreIndex;
 	uint32 NumPresentCalls;
 	uint32 NumAcquireCalls;
+	uint32 InternalWidth = 0;
+	uint32 InternalHeight = 0;
 
 	uint32 RTPacingSampleCount = 0;
 	double RTPacingPreviousFrameCPUTime = 0;
diff --git a/Engine/Source/Runtime/VulkanRHI/Private/VulkanViewport.cpp b/Engine/Source/Runtime/VulkanRHI/Private/VulkanViewport.cpp
--- a/Engine/Source/Runtime/VulkanRHI/Private/VulkanViewport.cpp
+++ b/Engine/Source/Runtime/VulkanRHI/Private/VulkanViewport.cpp
@@ -764,8 +764,8 @@ bool FVulkanViewport::Present(FVulkanCommandListContext* Context, FVulkanCmdBuff
 			SCOPE_CYCLE_COUNTER(STAT_VulkanAcquireBackBuffer);
 			GetNextImageIndex();
 
-			uint32 WindowSizeX = SizeX;
-			uint32 WindowSizeY = SizeY;
+			uint32 WindowSizeX = FMath::Min(SizeX, SwapChain->InternalWidth);
+			uint32 WindowSizeY = FMath::Min(SizeY, SwapChain->InternalHeight);
 
 			Context->RHIPushEvent(TEXT("CopyImageToBackBuffer"), FColor::Blue);
 			CopyImageToBackBuffer(CmdBuffer, true, RenderingBackBuffer->Surface.Image, BackBufferImages[AcquiredImageIndex], SizeX, SizeY, WindowSizeX, WindowSizeY);
