Fixed: Nearest filtering on Vulkan doesn't work on some Android devices.
Disable anisotropy for nearest filtering in VK

diff --git a/Engine/Source/Runtime/VulkanRHI/Private/VulkanState.cpp b/Engine/Source/Runtime/VulkanRHI/Private/VulkanState.cpp
--- a/Engine/Source/Runtime/VulkanRHI/Private/VulkanState.cpp
+++ b/Engine/Source/Runtime/VulkanRHI/Private/VulkanState.cpp
@@ -232,8 +232,13 @@ void FVulkanSamplerState::SetupSamplerCreateInfo(const FSamplerStateInitializerR
 	OutSamplerInfo.addressModeW = TranslateWrapMode(Initializer.AddressW);
 
 	OutSamplerInfo.mipLodBias = Initializer.MipBias;
-	OutSamplerInfo.maxAnisotropy = FMath::Clamp((float)ComputeAnisotropyRT(Initializer.MaxAnisotropy), 1.0f, InDevice.GetLimits().maxSamplerAnisotropy);
-	OutSamplerInfo.anisotropyEnable = OutSamplerInfo.maxAnisotropy > 1;
+	
+	OutSamplerInfo.maxAnisotropy = 1.0f;
+	if (Initializer.Filter != SF_Point)
+	{
+		OutSamplerInfo.maxAnisotropy = FMath::Clamp((float)ComputeAnisotropyRT(Initializer.MaxAnisotropy), 1.0f, InDevice.GetLimits().maxSamplerAnisotropy);
+	}
+	OutSamplerInfo.anisotropyEnable = OutSamplerInfo.maxAnisotropy > 1.0f;
 
 	OutSamplerInfo.compareEnable = Initializer.SamplerComparisonFunction != SCF_Never ? VK_TRUE : VK_FALSE;
 	OutSamplerInfo.compareOp = TranslateSamplerCompareFunction(Initializer.SamplerComparisonFunction);
