UE4.22 - Show entry point name for SPIRV shaders when debugging is enabled
#rnx

diff --git a/Engine/Source/Runtime/VulkanRHI/Private/VulkanShaders.cpp b/Engine/Source/Runtime/VulkanRHI/Private/VulkanShaders.cpp
--- a/Engine/Source/Runtime/VulkanRHI/Private/VulkanShaders.cpp
+++ b/Engine/Source/Runtime/VulkanRHI/Private/VulkanShaders.cpp
@@ -110,6 +110,13 @@ void FVulkanShader::Setup(const TArray<uint8>& InShaderHeaderAndCode, uint64 InS
 		checkSlow(CodeHeader.UniformBufferSpirvInfos.Num() == 0);
 	}
 	check(CodeHeader.GlobalSpirvInfos.Num() == CodeHeader.Globals.Num());
+
+#if VULKAN_ENABLE_SHADER_DEBUG_NAMES
+	// main_00000000_00000000
+	ANSICHAR EntryPoint[24];
+	GetEntryPoint(EntryPoint);
+	DebugEntryPoint = EntryPoint;
+#endif
 }
 
 VkShaderModule FVulkanShader::CreateHandle(const FVulkanLayout* Layout, uint32 LayoutHash)
diff --git a/Engine/Source/Runtime/VulkanRHI/Public/VulkanResources.h b/Engine/Source/Runtime/VulkanRHI/Public/VulkanResources.h
--- a/Engine/Source/Runtime/VulkanRHI/Public/VulkanResources.h
+++ b/Engine/Source/Runtime/VulkanRHI/Public/VulkanResources.h
@@ -125,6 +125,9 @@ public:
 	}
 
 protected:
+#if VULKAN_ENABLE_SHADER_DEBUG_NAMES
+	FString							DebugEntryPoint;
+#endif
 	uint64							ShaderKey;
 
 	/** External bindings for this shader. */
