Fix Vulkan running out of memory when doing highres screenshots.
This was caused by Delayresize staying at the highres resolution even after screenshot was taken, causing memory use to be 2x what it actually needed to be.
cleaned up some vulkan memory debug and fixed a minor leak in vulkanmemory.cpp

diff --git a/Engine/Source/Runtime/Renderer/Private/PostProcess/SceneRenderTargets.cpp b/Engine/Source/Runtime/Renderer/Private/PostProcess/SceneRenderTargets.cpp
--- a/Engine/Source/Runtime/Renderer/Private/PostProcess/SceneRenderTargets.cpp
+++ b/Engine/Source/Runtime/Renderer/Private/PostProcess/SceneRenderTargets.cpp
@@ -304,31 +304,44 @@ inline const TCHAR* GetSceneColorTargetName(EShadingPath ShadingPath)
 
 #if PREVENT_RENDERTARGET_SIZE_THRASHING
 
-static inline void UpdateHistoryFlags(uint8& Flags, bool bIsSceneCapture, bool bIsReflectionCapture)
+namespace ERenderTargetHistory
 {
-	Flags |= bIsSceneCapture ? 1 : 0;
-	Flags |= bIsReflectionCapture ? (1 << 1) : 0;
+	enum Type
+	{
+		RTH_SceneCapture		= 0x1,
+		RTH_ReflectionCapture	= 0x2,
+		RTH_HighresScreenshot	= 0x4,
+		RTH_MaskAll				= 0x7,
+	};
+}
+
+static inline void UpdateHistoryFlags(uint8& Flags, bool bIsSceneCapture, bool bIsReflectionCapture, bool bIsHighResScreenShot)
+{
+	Flags |= bIsSceneCapture ? ERenderTargetHistory::RTH_SceneCapture : 0;
+	Flags |= bIsReflectionCapture ? ERenderTargetHistory::RTH_ReflectionCapture : 0;
+	Flags |= bIsHighResScreenShot ? ERenderTargetHistory::RTH_HighresScreenshot : 0;
 }
 
 template <uint32 NumEntries>
-static bool AnyCaptureRenderedRecently(const uint8* HitoryFlags)
+static bool AnyCaptureRenderedRecently(const uint8* HistoryFlags, uint8 Mask)
 {
 	uint8 Result = 0;
 	for (uint32 Idx = 0; Idx < NumEntries; ++Idx)
 	{
-		Result |= HitoryFlags[Idx];
+		Result |= (HistoryFlags[Idx] & Mask);
 	}
 	return Result != 0;
 }
 
-#define UPDATE_HISTORY_FLAGS(Flags, bIsSceneCapture, bIsReflectionCapture) UpdateHistoryFlags(Flags, bIsSceneCapture, bIsReflectionCapture)
-#define ANY_CAPTURE_RENDERED_RECENTLY(HistoryFlags, NumEntries) AnyCaptureRenderedRecently<NumEntries>(HistoryFlags)
+#define UPDATE_HISTORY_FLAGS(Flags, bIsSceneCapture, bIsReflectionCapture, bIsHighResScreenShot) UpdateHistoryFlags(Flags, bIsSceneCapture, bIsReflectionCapture, bIsHighResScreenShot)
+#define ANY_CAPTURE_RENDERED_RECENTLY(HistoryFlags, NumEntries) AnyCaptureRenderedRecently<NumEntries>(HistoryFlags, ERenderTargetHistory::RTH_MaskAll)
+#define ANY_HIGHRES_CAPTURE_RENDERED_RECENTLY(HistoryFlags, NumEntries) AnyCaptureRenderedRecently<NumEntries>(HistoryFlags, ERenderTargetHistory::RTH_HighresScreenshot)
 
 #else
 
-#define UPDATE_HISTORY_FLAGS(Flags, bIsSceneCapture, bIsReflectionCapture)
+#define UPDATE_HISTORY_FLAGS(Flags, bIsSceneCapture, bIsReflectionCapture, bIsHighResScreenShot)
 #define ANY_CAPTURE_RENDERED_RECENTLY(HistoryFlags, NumEntries) (false)
-
+#define ANY_HIGHRES_CAPTURE_RENDERED_RECENTLY(HistoryFlags, NumEntries) (false)
 #endif
 
 FIntPoint FSceneRenderTargets::ComputeDesiredSize(const FSceneViewFamily& ViewFamily)
@@ -405,7 +418,8 @@ FIntPoint FSceneRenderTargets::ComputeDesiredSize(const FSceneViewFamily& ViewFa
 	// this allows The BufferSize to not grow below the SceneCapture requests (happen before scene rendering, in the same frame with a Grow request)
 	FIntPoint& LargestDesiredSizeThisFrame = LargestDesiredSizes[CurrentDesiredSizeIndex];
 	LargestDesiredSizeThisFrame = LargestDesiredSizeThisFrame.ComponentMax(DesiredBufferSize);
-	UPDATE_HISTORY_FLAGS(HistoryFlags[CurrentDesiredSizeIndex], bIsSceneCapture, bIsReflectionCapture);
+	bool bIsHighResScreenshot = GIsHighResScreenshot;
+	UPDATE_HISTORY_FLAGS(HistoryFlags[CurrentDesiredSizeIndex], bIsSceneCapture, bIsReflectionCapture, bIsHighResScreenshot);
 
 	// we want to shrink the buffer but as we can have multiple scenecaptures per frame we have to delay that a frame to get all size requests
 	// Don't save buffer size in history while making high-res screenshot.
@@ -441,6 +455,11 @@ FIntPoint FSceneRenderTargets::ComputeDesiredSize(const FSceneViewFamily& ViewFa
 			}
 		}
 	}
+	const bool bAnyHighresScreenshotRecently = ANY_HIGHRES_CAPTURE_RENDERED_RECENTLY(HistoryFlags, FrameSizeHistoryCount);
+	if(bAnyHighresScreenshotRecently != GIsHighResScreenshot)
+	{
+		bAllowDelayResize = false;
+	}
 
 	if(bAllowDelayResize)
 	{
diff --git a/Engine/Source/Runtime/VulkanRHI/Private/VulkanMemory.cpp b/Engine/Source/Runtime/VulkanRHI/Private/VulkanMemory.cpp
--- a/Engine/Source/Runtime/VulkanRHI/Private/VulkanMemory.cpp
+++ b/Engine/Source/Runtime/VulkanRHI/Private/VulkanMemory.cpp
@@ -15,6 +15,10 @@ uint32 GVulkanRHIDeletionFrameNumber = 0;
 const uint32 NUM_FRAMES_TO_WAIT_FOR_RESOURCE_DELETE = 2;
 
 
+#if UE_BUILD_DEBUG
+RENDERCORE_API	void DumpRenderTargetPoolMemory(FOutputDevice& OutputDevice);
+#endif
+
 #if VULKAN_MEMORY_TRACK_CALLSTACK
 static FCriticalSection GStackTraceMutex;
 static char GStackTrace[65536];
@@ -209,29 +213,20 @@ namespace VulkanRHI
 		check(!DedicatedAllocateInfo);
 #endif
 
-		FDeviceMemoryAllocation* NewAllocation = new FDeviceMemoryAllocation;
-		NewAllocation->DeviceHandle = DeviceHandle;
-		NewAllocation->Size = AllocationSize;
-		NewAllocation->MemoryTypeIndex = MemoryTypeIndex;
-		NewAllocation->bCanBeMapped = ((MemoryProperties.memoryTypes[MemoryTypeIndex].propertyFlags & VK_MEMORY_PROPERTY_HOST_VISIBLE_BIT) == VK_MEMORY_PROPERTY_HOST_VISIBLE_BIT);
-		NewAllocation->bIsCoherent = ((MemoryProperties.memoryTypes[MemoryTypeIndex].propertyFlags & VK_MEMORY_PROPERTY_HOST_COHERENT_BIT) == VK_MEMORY_PROPERTY_HOST_COHERENT_BIT);
-		NewAllocation->bIsCached = ((MemoryProperties.memoryTypes[MemoryTypeIndex].propertyFlags & VK_MEMORY_PROPERTY_HOST_CACHED_BIT) == VK_MEMORY_PROPERTY_HOST_CACHED_BIT);
-#if VULKAN_MEMORY_TRACK_FILE_LINE
-		NewAllocation->File = File;
-		NewAllocation->Line = Line;
-		static uint32 ID = 0;
-		NewAllocation->UID = ++ID;
-#endif
-#if VULKAN_MEMORY_TRACK_CALLSTACK
-		CaptureCallStack(NewAllocation->Callstack);
-#endif
-		VkResult Result = VulkanRHI::vkAllocateMemory(DeviceHandle, &Info, VULKAN_CPU_ALLOCATOR, &NewAllocation->Handle);
+		VkDeviceMemory Handle;
+		VkResult Result = VulkanRHI::vkAllocateMemory(DeviceHandle, &Info, VULKAN_CPU_ALLOCATOR, &Handle);
 		
-		//FPlatformMisc::LowLevelOutputDebugStringf(TEXT("(%u)(%x) VulkanRHI::vkAllocateMemory size=%u Type=%u"),
-		//	GFrameNumberRenderThread, this, AllocationSize, MemoryTypeIndex);
+		//FPlatformMisc::LowLevelOutputDebugStringf(TEXT("(%u)(%x) VulkanRHI::vkAllocateMemory size=%u(%6.3fmb) Type=%u\n"),
+		//	GFrameNumberRenderThread, this, AllocationSize, AllocationSize / (1024.f * 1024.f), MemoryTypeIndex);
 
 		if (Result == VK_ERROR_OUT_OF_DEVICE_MEMORY)
 		{
+#if UE_BUILD_DEBUG
+			DumpMemory();
+			GLog->PanicFlushThreadedLogs();
+			DumpRenderTargetPoolMemory(*GLog);
+			GLog->PanicFlushThreadedLogs();
+#endif
 			if (bCanFail)
 			{
 				UE_LOG(LogVulkanRHI, Warning, TEXT("Failed to allocate Device Memory, Requested=%.2fKb MemTypeIndex=%d"), (float)Info.allocationSize / 1024.0f, Info.memoryTypeIndex);
@@ -268,6 +263,25 @@ namespace VulkanRHI
 		{
 			VERIFYVULKANRESULT(Result);
 		}
+		FDeviceMemoryAllocation* NewAllocation = new FDeviceMemoryAllocation;
+		NewAllocation->DeviceHandle = DeviceHandle;
+		NewAllocation->Handle = Handle;
+		NewAllocation->Size = AllocationSize;
+		NewAllocation->MemoryTypeIndex = MemoryTypeIndex;
+		NewAllocation->bCanBeMapped = ((MemoryProperties.memoryTypes[MemoryTypeIndex].propertyFlags & VK_MEMORY_PROPERTY_HOST_VISIBLE_BIT) == VK_MEMORY_PROPERTY_HOST_VISIBLE_BIT);
+		NewAllocation->bIsCoherent = ((MemoryProperties.memoryTypes[MemoryTypeIndex].propertyFlags & VK_MEMORY_PROPERTY_HOST_COHERENT_BIT) == VK_MEMORY_PROPERTY_HOST_COHERENT_BIT);
+		NewAllocation->bIsCached = ((MemoryProperties.memoryTypes[MemoryTypeIndex].propertyFlags & VK_MEMORY_PROPERTY_HOST_CACHED_BIT) == VK_MEMORY_PROPERTY_HOST_CACHED_BIT);
+#if VULKAN_MEMORY_TRACK_FILE_LINE
+		NewAllocation->File = File;
+		NewAllocation->Line = Line;
+		static uint32 ID = 0;
+		NewAllocation->UID = ++ID;
+#endif
+#if VULKAN_MEMORY_TRACK_CALLSTACK
+		CaptureCallStack(NewAllocation->Callstack);
+#endif
+
+
 
 		++NumAllocations;
 		PeakNumAllocations = FMath::Max(NumAllocations, PeakNumAllocations);
@@ -333,14 +347,16 @@ namespace VulkanRHI
 			{
 				FDeviceMemoryAllocation* Allocation = HeapInfo.Allocations[SubIndex];
 #if VULKAN_MEMORY_TRACK_FILE_LINE
-				UE_LOG(LogVulkanRHI, Display, TEXT("\t\t%d Size %d Handle %p ID %d %s(%d)"), SubIndex, Allocation->Size, (void*)Allocation->Handle, Allocation->UID, ANSI_TO_TCHAR(Allocation->File), Allocation->Line);
+				UE_LOG(LogVulkanRHI, Display, TEXT("\t\t%d Size %.2f MB  %2.f Handle %p ID %d %s(%d)"), SubIndex, Allocation->Size / 1024.f / 1024.f, TotalSize / 1024.0f / 1024.0f, (void*)Allocation->Handle, Allocation->UID, ANSI_TO_TCHAR(Allocation->File), Allocation->Line);
 #else
-				UE_LOG(LogVulkanRHI, Display, TEXT("\t\t%d Size %d Handle %p"), SubIndex, Allocation->Size, (void*)Allocation->Handle);
+				UE_LOG(LogVulkanRHI, Display, TEXT("\t\t%d Size %.2f MB  %2.f Handle %p"), SubIndex, Allocation->Size / 1024.f / 1024.f, TotalSize / 1024.0f / 1024.0f, (void*)Allocation->Handle);
 #endif
 				TotalSize += Allocation->Size;
 			}
 			UE_LOG(LogVulkanRHI, Display, TEXT("\t\tTotal Allocated %.2f MB, Peak %.2f MB"), TotalSize / 1024.0f / 1024.0f, HeapInfo.PeakSize / 1024.0f / 1024.0f);
 		}
+		Device->GetResourceHeapManager().DumpMemory();
+		GLog->PanicFlushThreadedLogs();
 	}
 #endif
 
@@ -811,10 +827,15 @@ namespace VulkanRHI
 				SubAllocAllocatedMemory += UsedPages[Index]->MaxSize;
 				NumSuballocations += UsedPages[Index]->ResourceAllocations.Num();
 
-				UE_LOG(LogVulkanRHI, Display, TEXT("\t\t%d: ID %4d %4d suballocs, %4d free chunks (%d used/%d free/%d max) DeviceMemory %p"), Index, UsedPages[Index]->GetID(), UsedPages[Index]->ResourceAllocations.Num(), UsedPages[Index]->FreeList.Num(), UsedPages[Index]->UsedSize, UsedPages[Index]->MaxSize - UsedPages[Index]->UsedSize, UsedPages[Index]->MaxSize, (void*)UsedPages[Index]->DeviceMemoryAllocation->GetHandle());
+				UE_LOG(LogVulkanRHI, Display, TEXT("\t\t%d: ID %4d %4d suballocs, %4d free chunks (%6.2fmb used/%6.2fmb free/%6.2fmb max) DeviceMemory %p"), Index, UsedPages[Index]->GetID(), UsedPages[Index]->ResourceAllocations.Num(), 
+					UsedPages[Index]->FreeList.Num(), 
+					UsedPages[Index]->UsedSize / (1024.f*1024.f),
+					(UsedPages[Index]->MaxSize - UsedPages[Index]->UsedSize) / (1024.f*1024.f), 
+					UsedPages[Index]->MaxSize / (1024.f*1024.f),
+					(void*)UsedPages[Index]->DeviceMemoryAllocation->GetHandle());
 			}
 
-			UE_LOG(LogVulkanRHI, Display, TEXT("%d Suballocations for Used/Total: %d/%d = %.2f%%"), NumSuballocations, SubAllocUsedMemory, SubAllocAllocatedMemory, SubAllocAllocatedMemory > 0 ? 100.0f * (float)SubAllocUsedMemory / (float)SubAllocAllocatedMemory : 0.0f);
+			UE_LOG(LogVulkanRHI, Display, TEXT("%d Suballocations for Used/Total: %.2fmb/%.2fmb = %.2f%%"), NumSuballocations, SubAllocUsedMemory / (1024.f*1024.f), SubAllocAllocatedMemory / (1024.f*1024.f), SubAllocAllocatedMemory > 0 ? 100.0f * (float)SubAllocUsedMemory / (float)SubAllocAllocatedMemory : 0.0f);
 		};
 
 		DumpPages(UsedBufferPages, TEXT("Buffer"));
@@ -1361,11 +1382,11 @@ namespace VulkanRHI
 
 				if (PoolSizeIndex == (int32)EPoolSizes::SizesCount)
 				{
-					UE_LOG(LogVulkanRHI, Display, TEXT(" Large Alloc Used/Max %d/%d %.2f%%"), _UsedLargeTotal, _AllocLargeTotal, 100.0f * (float)_UsedLargeTotal / (float)_AllocLargeTotal);
+					UE_LOG(LogVulkanRHI, Display, TEXT(" Large Alloc Used/Max %d/%d %6.2f%%"), _UsedLargeTotal, _AllocLargeTotal, 100.0f * (float)_UsedLargeTotal / (float)_AllocLargeTotal);
 				}
 				else
 				{
-					UE_LOG(LogVulkanRHI, Display, TEXT(" Binned [%d] Alloc Used/Max %d/%d %.2f%%"), PoolSizes[PoolSizeIndex], _UsedBinnedTotal, _AllocBinnedTotal, 100.0f * (float)_UsedBinnedTotal / (float)_AllocBinnedTotal);
+					UE_LOG(LogVulkanRHI, Display, TEXT(" Binned [%d] Alloc Used/Max %d/%d %6.2f%%"), PoolSizes[PoolSizeIndex], _UsedBinnedTotal, _AllocBinnedTotal, 100.0f * (float)_UsedBinnedTotal / (float)_AllocBinnedTotal);
 				}
 			}
 		}
