Fixed: UI flicker on Mali devices with vulkan enabled

diff --git a/Engine/Source/Runtime/VulkanRHI/Private/VulkanRenderTarget.cpp b/Engine/Source/Runtime/VulkanRHI/Private/VulkanRenderTarget.cpp
--- a/Engine/Source/Runtime/VulkanRHI/Private/VulkanRenderTarget.cpp
+++ b/Engine/Source/Runtime/VulkanRHI/Private/VulkanRenderTarget.cpp
@@ -625,20 +625,23 @@ void FVulkanCommandListContext::RHICopyToResolveTarget(FTextureRHIParamRef Sourc
 		return;
 	}
 
-	if (SourceTextureRHI == DestTextureRHI)
-	{
-		FRHITexture2D* SourceTexture2D = SourceTextureRHI->GetTexture2D();
-		if (SourceTexture2D)
-		{
-			FVulkanTexture2D* VulkanSourceTexture2D  = (FVulkanTexture2D*)SourceTexture2D;
-			if (VulkanSourceTexture2D->GetBackBuffer() != nullptr)
-			{
-				// skip Backbuffer implicit transition to Readable, to avoid splitting Post->UI renderpass
-				// do explicit transition when need to read from Backbuffer
-				return;
-			}
-		}
-	}
+	
+	// Disabled as it causes flicker on MALI
+	// TODO: With explicit render-pases this does not help with fusing Post->UI renderpass anyway
+	//if (SourceTextureRHI == DestTextureRHI)
+	//{
+	//	FRHITexture2D* SourceTexture2D = SourceTextureRHI->GetTexture2D();
+	//	if (SourceTexture2D)
+	//	{
+	//		FVulkanTexture2D* VulkanSourceTexture2D  = (FVulkanTexture2D*)SourceTexture2D;
+	//		if (VulkanSourceTexture2D->GetBackBuffer() != nullptr)
+	//		{
+	//			// skip Backbuffer implicit transition to Readable, to avoid splitting Post->UI renderpass
+	//			// do explicit transition when need to read from Backbuffer
+	//			return;
+	//		}
+	//	}
+	//}
 
 	RHITransitionResources(EResourceTransitionAccess::EReadable, &SourceTextureRHI, 1);
 
