UE4.22.1 - Fix compilation fixes when VULKAN_ENABLE_DUMP_LAYER is enabled
#rnx

diff --git a/Engine/Source/Runtime/VulkanRHI/Private/VulkanDebug.cpp b/Engine/Source/Runtime/VulkanRHI/Private/VulkanDebug.cpp
--- a/Engine/Source/Runtime/VulkanRHI/Private/VulkanDebug.cpp
+++ b/Engine/Source/Runtime/VulkanRHI/Private/VulkanDebug.cpp
@@ -2009,8 +2009,10 @@ void FWrapLayer::UpdateDescriptorSets(VkResult Result, VkDevice Device, uint32 D
 						BreakOnTrackingImageView(DescriptorWrites[Index].pImageInfo->imageView);
 #endif
 #if VULKAN_ENABLE_DUMP_LAYER
+						auto* FoundImageInfo = GVulkanTrackingImageViews.Find(DescriptorWrites[Index].pImageInfo->imageView);
+						VkImage *Image = FoundImageInfo ? &FoundImageInfo->CreateInfo.image : nullptr;
 						DebugLog += FString::Printf(TEXT("%s\tpImageInfo[%d]: Sampler=0x%p, ImageView=0x%p(I:0x%p), imageLayout=%s\n"), Tabs, SubIndex,
-							DescriptorWrites[Index].pImageInfo->sampler, DescriptorWrites[Index].pImageInfo->imageView, FindTrackingImage(DescriptorWrites[Index].pImageInfo->imageView), *GetImageLayoutString(DescriptorWrites[Index].pImageInfo->imageLayout));
+							DescriptorWrites[Index].pImageInfo->sampler, DescriptorWrites[Index].pImageInfo->imageView, Image ? *Image : nullptr, *GetImageLayoutString(DescriptorWrites[Index].pImageInfo->imageLayout));
 #endif
 					}
 				}
@@ -2048,7 +2050,10 @@ void FWrapLayer::CreateFramebuffer(VkResult Result, VkDevice Device, const VkFra
 		DebugLog += FString::Printf(TEXT("%sVkFramebufferCreateInfo: Flags=%d, RenderPass=0x%p, NumAttachments=%d\n"), Tabs, CreateInfo->flags, CreateInfo->renderPass, CreateInfo->attachmentCount);
 		for (uint32 Index = 0; Index < CreateInfo->attachmentCount; ++Index)
 		{
-			DebugLog += FString::Printf(TEXT("%s\tAttachment[%d]: ImageView=0x%p(I:0x%p)\n"), Tabs, Index, CreateInfo->pAttachments[Index], FindTrackingImage(CreateInfo->pAttachments[Index]));
+			auto* FoundImageInfo = GVulkanTrackingImageViews.Find(CreateInfo->pAttachments[Index]);
+			VkImage *Image = FoundImageInfo ? &FoundImageInfo->CreateInfo.image : nullptr;
+
+			DebugLog += FString::Printf(TEXT("%s\tAttachment[%d]: ImageView=0x%p(I:0x%p)\n"), Tabs, Index, CreateInfo->pAttachments[Index], Image ? *Image : nullptr);
 		}
 		DebugLog += FString::Printf(TEXT("%s\twidth=%d, height=%d, layers=%d\n"), Tabs, CreateInfo->width, CreateInfo->height, CreateInfo->layers);
 #endif
@@ -3672,7 +3677,7 @@ void FWrapLayer::SetEvent(VkResult Result, VkDevice Device, VkEvent Event)
 	if (Result == VK_RESULT_MAX_ENUM)
 	{
 #if VULKAN_ENABLE_DUMP_LAYER
-		CmdPrintfBegin(CommandBuffer, FString::Printf(TEXT("vkCmdSetEvent(Event=0x%p, StageMask=0x%x)"), Event, StageMask));
+		DevicePrintfBegin(Device, FString::Printf(TEXT("vkCmdSetEvent(Event=0x%p)"), Event));
 #endif
 	}
 }
@@ -3682,7 +3687,7 @@ void FWrapLayer::ResetEvent(VkResult Result, VkDevice Device, VkEvent Event)
 	if (Result == VK_RESULT_MAX_ENUM)
 	{
 #if VULKAN_ENABLE_DUMP_LAYER
-		CmdPrintfBegin(CommandBuffer, FString::Printf(TEXT("vkCmdResetEvent(Event=0x%p, StageMask=0x%x)"), Event, StageMask));
+		DevicePrintfBegin(Device, FString::Printf(TEXT("vkCmdResetEvent(Event=0x%p)"), Event));
 #endif
 	}
 }
@@ -3692,7 +3697,7 @@ void FWrapLayer::GetEventStatus(VkResult Result, VkDevice Device, VkEvent Event)
 	if (Result == VK_RESULT_MAX_ENUM)
 	{
 #if VULKAN_ENABLE_DUMP_LAYER
-		DevicePrintfBegin(CommandBuffer, FString::Printf(TEXT("GetEventStatus(Event=0x%p)"), Event));
+		DevicePrintfBegin(Device, FString::Printf(TEXT("GetEventStatus(Event=0x%p)"), Event));
 #endif
 	}
 	else
