Move Vulkan headers from Public to Private
#rnx

diff --git a/Engine/Source/Runtime/VulkanRHI/Private/VulkanContext.h b/Engine/Source/Runtime/VulkanRHI/Private/VulkanContext.h
new file mode 100644
--- /dev/null
+++ b/Engine/Source/Runtime/VulkanRHI/Private/VulkanContext.h
@@ -0,0 +1,515 @@
+// Copyright 1998-2019 Epic Games, Inc. All Rights Reserved.
+
+/*=============================================================================
+	VulkanContext.h: Class to generate Vulkan command buffers from RHI CommandLists
+=============================================================================*/
+
+#pragma once 
+
+#include "VulkanResources.h"
+
+class FVulkanDevice;
+class FVulkanCommandBufferManager;
+class FVulkanPendingGfxState;
+class FVulkanPendingComputeState;
+class FVulkanQueue;
+class FVulkanOcclusionQueryPool;
+
+struct FInputAttachmentData;
+
+class FTransitionAndLayoutManagerData
+{
+public:
+	void TempCopy(const FTransitionAndLayoutManagerData& In)
+	{
+		Framebuffers = In.Framebuffers;
+		RenderPasses = In.RenderPasses;
+		Layouts = In.Layouts;
+	}
+
+protected:
+	TMap<uint32, FVulkanRenderPass*> RenderPasses;
+	struct FFramebufferList
+	{
+		TArray<FVulkanFramebuffer*> Framebuffer;
+	};
+	TMap<uint32, FFramebufferList*> Framebuffers;
+	TMap<VkImage, VkImageLayout> Layouts;
+};
+
+class FTransitionAndLayoutManager : public FTransitionAndLayoutManagerData
+{
+	using FFramebufferList = FTransitionAndLayoutManagerData::FFramebufferList;
+public:
+	FTransitionAndLayoutManager()
+		: CurrentRenderPass(nullptr)
+		, CurrentFramebuffer(nullptr)
+	{
+	}
+
+	void Destroy(FVulkanDevice& InDevice, FTransitionAndLayoutManager* Immediate);
+
+	FVulkanFramebuffer* GetOrCreateFramebuffer(FVulkanDevice& InDevice, const FRHISetRenderTargetsInfo& RenderTargetsInfo, const FVulkanRenderTargetLayout& RTLayout, FVulkanRenderPass* RenderPass);
+	FVulkanRenderPass* GetOrCreateRenderPass(FVulkanDevice& InDevice, const FVulkanRenderTargetLayout& RTLayout)
+	{
+		uint32 RenderPassHash = RTLayout.GetRenderPassFullHash();
+		FVulkanRenderPass** FoundRenderPass = nullptr;
+		{
+			FScopeLock Lock(&RenderPassesCS);
+			FoundRenderPass = RenderPasses.Find(RenderPassHash);
+		}
+		if (FoundRenderPass)
+		{
+			return *FoundRenderPass;
+		}
+
+		FVulkanRenderPass* RenderPass = new FVulkanRenderPass(InDevice, RTLayout);
+		{
+			FScopeLock Lock(&RenderPassesCS);
+			RenderPasses.Add(RenderPassHash, RenderPass);
+		}
+		return RenderPass;
+	}
+
+	void BeginEmulatedRenderPass(FVulkanCommandListContext& Context, FVulkanDevice& InDevice, FVulkanCmdBuffer* CmdBuffer, const FRHISetRenderTargetsInfo& RenderTargetsInfo, const FVulkanRenderTargetLayout& RTLayout, FVulkanRenderPass* RenderPass, FVulkanFramebuffer* Framebuffer);
+	void EndEmulatedRenderPass(FVulkanCmdBuffer* CmdBuffer);
+
+	void BeginRealRenderPass(FVulkanCommandListContext& Context, FVulkanDevice& InDevice, FVulkanCmdBuffer* CmdBuffer, const FRHIRenderPassInfo& RPInfo, const FVulkanRenderTargetLayout& RTLayout, FVulkanRenderPass* RenderPass, FVulkanFramebuffer* Framebuffer);
+	void EndRealRenderPass(FVulkanCmdBuffer* CmdBuffer);
+
+	struct FGenerateMipsInfo
+	{
+		int32 NumRenderTargets = 0;
+
+		bool bInsideGenerateMips;
+		bool bLastMip;
+		int32 CurrentSlice;
+		int32 CurrentMip;
+
+		struct
+		{
+			// Per face/slice array of mip layouts
+			TArray<TArray<VkImageLayout>> Layouts;
+			VkImage CurrentImage;
+		} Target[MaxSimultaneousRenderTargets];
+
+		FGenerateMipsInfo()
+		{
+			Reset();
+		}
+
+		void Reset()
+		{
+			NumRenderTargets = 0;
+			bInsideGenerateMips = false;
+			bLastMip = false;
+			CurrentSlice = -1;
+			CurrentMip = -1;
+			for (int32 Index = 0; Index < MaxSimultaneousRenderTargets; ++Index)
+			{
+				Target[Index].Layouts.Reset(0);
+				Target[Index].CurrentImage = VK_NULL_HANDLE;
+			}
+		}
+	} GenerateMipsInfo;
+
+	bool bInsideRealRenderPass = false;
+
+	FVulkanRenderPass* CurrentRenderPass;
+	FVulkanFramebuffer* CurrentFramebuffer;
+
+	FCriticalSection RenderPassesCS;
+
+	void NotifyDeletedRenderTarget(FVulkanDevice& InDevice, VkImage Image);
+
+	inline void NotifyDeletedImage(VkImage Image)
+	{
+		Layouts.Remove(Image);
+	}
+
+	VkImageLayout FindLayoutChecked(VkImage Image) const
+	{
+		return Layouts.FindChecked(Image);
+	}
+
+	VkImageLayout FindOrAddLayout(VkImage Image, VkImageLayout LayoutIfNotFound)
+	{
+		VkImageLayout* Found = Layouts.Find(Image);
+		if (Found)
+		{
+			return *Found;
+		}
+
+		Layouts.Add(Image, LayoutIfNotFound);
+		return LayoutIfNotFound;
+	}
+
+	VkImageLayout& FindOrAddLayoutRW(VkImage Image, VkImageLayout LayoutIfNotFound)
+	{
+		VkImageLayout* Found = Layouts.Find(Image);
+		if (Found)
+		{
+			return *Found;
+		}
+
+		return Layouts.Add(Image, LayoutIfNotFound);
+	}
+
+	void TransitionResource(FVulkanCmdBuffer* CmdBuffer, FVulkanSurface& Surface, VulkanRHI::EImageLayoutBarrier DestLayout);
+};
+
+class FVulkanCommandListContext : public IRHICommandContext
+{
+public:
+	FVulkanCommandListContext(FVulkanDynamicRHI* InRHI, FVulkanDevice* InDevice, FVulkanQueue* InQueue, FVulkanCommandListContext* InImmediate);
+	virtual ~FVulkanCommandListContext();
+
+	inline bool IsImmediate() const
+	{
+		return Immediate == nullptr;
+	}
+
+	virtual void RHISetStreamSource(uint32 StreamIndex, FVertexBufferRHIParamRef VertexBuffer, uint32 Offset) final override;
+	virtual void RHISetViewport(uint32 MinX, uint32 MinY, float MinZ, uint32 MaxX, uint32 MaxY, float MaxZ) final override;
+	virtual void RHISetScissorRect(bool bEnable, uint32 MinX, uint32 MinY, uint32 MaxX, uint32 MaxY) final override;
+	virtual void RHISetGraphicsPipelineState(FGraphicsPipelineStateRHIParamRef GraphicsState) final override;
+	virtual void RHISetShaderTexture(FVertexShaderRHIParamRef VertexShader, uint32 TextureIndex, FTextureRHIParamRef NewTexture) final override;
+	virtual void RHISetShaderTexture(FHullShaderRHIParamRef HullShader, uint32 TextureIndex, FTextureRHIParamRef NewTexture) final override;
+	virtual void RHISetShaderTexture(FDomainShaderRHIParamRef DomainShader, uint32 TextureIndex, FTextureRHIParamRef NewTexture) final override;
+	virtual void RHISetShaderTexture(FGeometryShaderRHIParamRef GeometryShader, uint32 TextureIndex, FTextureRHIParamRef NewTexture) final override;
+	virtual void RHISetShaderTexture(FPixelShaderRHIParamRef PixelShader, uint32 TextureIndex, FTextureRHIParamRef NewTexture) final override;
+	virtual void RHISetShaderTexture(FComputeShaderRHIParamRef PixelShader, uint32 TextureIndex, FTextureRHIParamRef NewTexture) final override;
+	virtual void RHISetShaderSampler(FComputeShaderRHIParamRef ComputeShader, uint32 SamplerIndex, FSamplerStateRHIParamRef NewState) final override;
+	virtual void RHISetShaderSampler(FVertexShaderRHIParamRef VertexShader, uint32 SamplerIndex, FSamplerStateRHIParamRef NewState) final override;
+	virtual void RHISetShaderSampler(FGeometryShaderRHIParamRef GeometryShader, uint32 SamplerIndex, FSamplerStateRHIParamRef NewState) final override;
+	virtual void RHISetShaderSampler(FDomainShaderRHIParamRef DomainShader, uint32 SamplerIndex, FSamplerStateRHIParamRef NewState) final override;
+	virtual void RHISetShaderSampler(FHullShaderRHIParamRef HullShader, uint32 SamplerIndex, FSamplerStateRHIParamRef NewState) final override;
+	virtual void RHISetShaderSampler(FPixelShaderRHIParamRef PixelShader, uint32 SamplerIndex, FSamplerStateRHIParamRef NewState) final override;
+	virtual void RHISetUAVParameter(FComputeShaderRHIParamRef ComputeShader, uint32 UAVIndex, FUnorderedAccessViewRHIParamRef UAV) final override;
+	virtual void RHISetUAVParameter(FComputeShaderRHIParamRef ComputeShader, uint32 UAVIndex, FUnorderedAccessViewRHIParamRef UAV, uint32 InitialCount) final override;
+	virtual void RHISetShaderResourceViewParameter(FPixelShaderRHIParamRef PixelShader, uint32 SamplerIndex, FShaderResourceViewRHIParamRef SRV) final override;
+	virtual void RHISetShaderResourceViewParameter(FVertexShaderRHIParamRef VertexShader, uint32 SamplerIndex, FShaderResourceViewRHIParamRef SRV) final override;
+	virtual void RHISetShaderResourceViewParameter(FComputeShaderRHIParamRef ComputeShader, uint32 SamplerIndex, FShaderResourceViewRHIParamRef SRV) final override;
+	virtual void RHISetShaderResourceViewParameter(FHullShaderRHIParamRef HullShader, uint32 SamplerIndex, FShaderResourceViewRHIParamRef SRV) final override;
+	virtual void RHISetShaderResourceViewParameter(FDomainShaderRHIParamRef DomainShader, uint32 SamplerIndex, FShaderResourceViewRHIParamRef SRV) final override;
+	virtual void RHISetShaderResourceViewParameter(FGeometryShaderRHIParamRef GeometryShader, uint32 SamplerIndex, FShaderResourceViewRHIParamRef SRV) final override;
+	virtual void RHISetShaderUniformBuffer(FVertexShaderRHIParamRef VertexShader, uint32 BufferIndex, FUniformBufferRHIParamRef Buffer) final override;
+	virtual void RHISetShaderUniformBuffer(FHullShaderRHIParamRef HullShader, uint32 BufferIndex, FUniformBufferRHIParamRef Buffer) final override;
+	virtual void RHISetShaderUniformBuffer(FDomainShaderRHIParamRef DomainShader, uint32 BufferIndex, FUniformBufferRHIParamRef Buffer) final override;
+	virtual void RHISetShaderUniformBuffer(FGeometryShaderRHIParamRef GeometryShader, uint32 BufferIndex, FUniformBufferRHIParamRef Buffer) final override;
+	virtual void RHISetShaderUniformBuffer(FPixelShaderRHIParamRef PixelShader, uint32 BufferIndex, FUniformBufferRHIParamRef Buffer) final override;
+	virtual void RHISetShaderUniformBuffer(FComputeShaderRHIParamRef ComputeShader, uint32 BufferIndex, FUniformBufferRHIParamRef Buffer) final override;
+	virtual void RHISetShaderParameter(FVertexShaderRHIParamRef VertexShader, uint32 BufferIndex, uint32 BaseIndex, uint32 NumBytes, const void* NewValue) final override;
+	virtual void RHISetShaderParameter(FHullShaderRHIParamRef HullShader, uint32 BufferIndex, uint32 BaseIndex, uint32 NumBytes, const void* NewValue) final override;
+	virtual void RHISetShaderParameter(FDomainShaderRHIParamRef DomainShader, uint32 BufferIndex, uint32 BaseIndex, uint32 NumBytes, const void* NewValue) final override;
+	virtual void RHISetShaderParameter(FGeometryShaderRHIParamRef GeometryShader, uint32 BufferIndex, uint32 BaseIndex, uint32 NumBytes, const void* NewValue) final override;
+	virtual void RHISetShaderParameter(FPixelShaderRHIParamRef PixelShader, uint32 BufferIndex, uint32 BaseIndex, uint32 NumBytes, const void* NewValue) final override;
+	virtual void RHISetShaderParameter(FComputeShaderRHIParamRef ComputeShader, uint32 BufferIndex, uint32 BaseIndex, uint32 NumBytes, const void* NewValue) final override;
+	virtual void RHISetStencilRef(uint32 StencilRef) final override;
+	virtual void RHISetRenderTargets(uint32 NumSimultaneousRenderTargets, const FRHIRenderTargetView* NewRenderTargets, const FRHIDepthRenderTargetView* NewDepthStencilTarget, uint32 NumUAVs, const FUnorderedAccessViewRHIParamRef* UAVs) final override;
+	virtual void RHISetRenderTargetsAndClear(const FRHISetRenderTargetsInfo& RenderTargetsInfo) final override;
+	virtual void RHIDrawPrimitive(uint32 BaseVertexIndex, uint32 NumPrimitives, uint32 NumInstances) final override;
+	virtual void RHIDrawPrimitiveIndirect(FVertexBufferRHIParamRef ArgumentBuffer, uint32 ArgumentOffset) final override;
+	virtual void RHIDrawIndexedIndirect(FIndexBufferRHIParamRef IndexBufferRHI, FStructuredBufferRHIParamRef ArgumentsBufferRHI, int32 DrawArgumentsIndex, uint32 NumInstances) final override;
+	virtual void RHIDrawIndexedPrimitive(FIndexBufferRHIParamRef IndexBuffer, int32 BaseVertexIndex, uint32 FirstInstance, uint32 NumVertices, uint32 StartIndex, uint32 NumPrimitives, uint32 NumInstances) final override;
+	virtual void RHIDrawIndexedPrimitiveIndirect(FIndexBufferRHIParamRef IndexBuffer, FVertexBufferRHIParamRef ArgumentBuffer, uint32 ArgumentOffset) final override;
+	virtual void RHIBeginDrawPrimitiveUP(uint32 NumPrimitives, uint32 NumVertices, uint32 VertexDataStride, void*& OutVertexData) final override;
+	virtual void RHIEndDrawPrimitiveUP() final override;
+	virtual void RHIBeginDrawIndexedPrimitiveUP(uint32 NumPrimitives, uint32 NumVertices, uint32 VertexDataStride, void*& OutVertexData, uint32 MinVertexIndex, uint32 NumIndices, uint32 IndexDataStride, void*& OutIndexData) final override;
+	virtual void RHIEndDrawIndexedPrimitiveUP() final override;
+	virtual void RHISetDepthBounds(float MinDepth, float MaxDepth) final override;
+	virtual void RHIPushEvent(const TCHAR* Name, FColor Color) final override;
+	virtual void RHIPopEvent() final override;
+
+	virtual void RHISetComputeShader(FComputeShaderRHIParamRef ComputeShader) final override;
+	virtual void RHISetComputePipelineState(FRHIComputePipelineState* ComputePipelineState) final override;
+	virtual void RHIWaitComputeFence(FComputeFenceRHIParamRef InFence) final override;
+	virtual void RHIDispatchComputeShader(uint32 ThreadGroupCountX, uint32 ThreadGroupCountY, uint32 ThreadGroupCountZ) final override;
+	virtual void RHIDispatchIndirectComputeShader(FVertexBufferRHIParamRef ArgumentBuffer, uint32 ArgumentOffset) final override;
+	virtual void RHIAutomaticCacheFlushAfterComputeShader(bool bEnable) final override;
+
+	virtual void RHIFlushComputeShaderCache() final override;
+	virtual void RHISetMultipleViewports(uint32 Count, const FViewportBounds* Data) final override;
+	virtual void RHIClearTinyUAV(FUnorderedAccessViewRHIParamRef UnorderedAccessViewRHI, const uint32* Values) final override;
+	virtual void RHICopyToResolveTarget(FTextureRHIParamRef SourceTexture, FTextureRHIParamRef DestTexture, const FResolveParams& ResolveParams) final override;
+	virtual void RHICopyTexture(FTextureRHIParamRef SourceTexture, FTextureRHIParamRef DestTexture, const FRHICopyTextureInfo& CopyInfo) final override;
+	virtual void RHITransitionResources(EResourceTransitionAccess TransitionType, FTextureRHIParamRef* InRenderTargets, int32 NumTextures) final override;
+	virtual void RHITransitionResources(EResourceTransitionAccess TransitionType, EResourceTransitionPipeline TransitionPipeline, FUnorderedAccessViewRHIParamRef* InUAVs, int32 NumUAVs, FComputeFenceRHIParamRef WriteComputeFence) final override;
+	virtual void RHICopyToStagingBuffer(FVertexBufferRHIParamRef SourceBuffer, FStagingBufferRHIParamRef DestinationStagingBuffer, uint32 Offset, uint32 NumBytes, FGPUFenceRHIParamRef Fence) final override;
+
+	// Render time measurement
+	virtual void RHIBeginRenderQuery(FRenderQueryRHIParamRef RenderQuery) final override;
+	virtual void RHIEndRenderQuery(FRenderQueryRHIParamRef RenderQuery) final override;
+
+	virtual void RHIUpdateTextureReference(FTextureReferenceRHIParamRef TextureRef, FTextureRHIParamRef NewTexture) final override;
+
+	virtual void RHISubmitCommandsHint() final override;
+
+	virtual void RHIBeginDrawingViewport(FViewportRHIParamRef Viewport, FTextureRHIParamRef RenderTargetRHI) final override;
+	virtual void RHIEndDrawingViewport(FViewportRHIParamRef Viewport, bool bPresent, bool bLockToVsync) final override;
+
+	virtual void RHIBeginFrame() final override;
+	virtual void RHIEndFrame() final override;
+
+	virtual void RHIBeginScene() final override;
+	virtual void RHIEndScene() final override;
+
+	virtual void RHIBeginRenderPass(const FRHIRenderPassInfo& InInfo, const TCHAR* InName) final override;
+	virtual void RHIEndRenderPass() final override;
+
+	inline FVulkanCommandBufferManager* GetCommandBufferManager()
+	{
+		return CommandBufferManager;
+	}
+
+	inline VulkanRHI::FTempFrameAllocationBuffer& GetTempFrameAllocationBuffer()
+	{
+		return TempFrameAllocationBuffer;
+	}
+
+	inline FVulkanPendingGfxState* GetPendingGfxState()
+	{
+		return PendingGfxState;
+	}
+
+	inline FVulkanPendingComputeState* GetPendingComputeState()
+	{
+		return PendingComputeState;
+	}
+
+	inline void NotifyDeletedRenderTarget(VkImage Image)
+	{
+		TransitionAndLayoutManager.NotifyDeletedRenderTarget(*Device, Image);
+	}
+
+	inline void NotifyDeletedImage(VkImage Image)
+	{
+		TransitionAndLayoutManager.NotifyDeletedImage(Image);
+	}
+
+	inline FVulkanRenderPass* GetCurrentRenderPass()
+	{
+		return TransitionAndLayoutManager.CurrentRenderPass;
+	}
+
+	inline uint64 GetFrameCounter() const
+	{
+		return FrameCounter;
+	}
+
+	inline FVulkanUniformBufferUploader* GetUniformBufferUploader()
+	{
+		return UniformBufferUploader;
+	}
+
+	inline FVulkanQueue* GetQueue()
+	{
+		return Queue;
+	}
+
+	void WriteBeginTimestamp(FVulkanCmdBuffer* CmdBuffer);
+	void WriteEndTimestamp(FVulkanCmdBuffer* CmdBuffer);
+
+	void ReadAndCalculateGPUFrameTime();
+	
+	inline FVulkanGPUProfiler& GetGPUProfiler()
+	{
+		return GpuProfiler;
+	}
+
+	inline FVulkanDevice* GetDevice() const
+	{
+		return Device;
+	}
+	void EndRenderQueryInternal(FVulkanCmdBuffer* CmdBuffer, FVulkanRenderQuery* Query);
+
+	inline VkImageLayout FindLayout(VkImage Image)
+	{
+		return TransitionAndLayoutManager.FindLayoutChecked(Image);
+	}
+
+	inline VkImageLayout GetLayoutForDescriptor(const FVulkanSurface& Surface) const
+	{
+#if PLATFORM_ANDROID && !PLATFORM_LUMIN && !PLATFORM_LUMINGL4
+		// Workaround clang bug; don't use IsDepthOrStencilAspect() directly
+		VkImageAspectFlags AspectFlags = Surface.GetFullAspectMask();
+		if ((AspectFlags & VK_IMAGE_ASPECT_DEPTH_BIT) != 0 || (AspectFlags & VK_IMAGE_ASPECT_STENCIL_BIT) != 0)
+		{
+			return VK_IMAGE_LAYOUT_DEPTH_STENCIL_READ_ONLY_OPTIMAL;
+		}
+#else
+		if (Surface.IsDepthOrStencilAspect())
+		{
+#if VULKAN_SUPPORTS_MAINTENANCE_LAYER2
+			// If the spec gets lenient, we could remove this search since then 
+			// Images in VK_IMAGE_LAYOUT_DEPTH_READ_ONLY_STENCIL_ATTACHMENT_OPTIMAL_KHR could be used with 
+			// descriptor write of VK_IMAGE_LAYOUT_DEPTH_STENCIL_READ_ONLY_OPTIMAL
+			if (Device->GetOptionalExtensions().HasKHRMaintenance2)
+			{
+				return TransitionAndLayoutManager.FindLayoutChecked(Surface.Image);
+			}
+#else
+			return VK_IMAGE_LAYOUT_DEPTH_STENCIL_READ_ONLY_OPTIMAL;
+#endif
+		}
+#endif
+
+		return TransitionAndLayoutManager.FindLayoutChecked(Surface.Image);
+	}
+
+	inline VkImageLayout FindOrAddLayout(VkImage Image, VkImageLayout NewLayout)
+	{
+		return TransitionAndLayoutManager.FindOrAddLayout(Image, NewLayout);
+	}
+
+	inline VkImageLayout& FindOrAddLayoutRW(VkImage Image, VkImageLayout NewLayout)
+	{
+		return TransitionAndLayoutManager.FindOrAddLayoutRW(Image, NewLayout);
+	}
+
+	void PrepareParallelFromBase(const FVulkanCommandListContext& BaseContext);
+
+	void* Hotfix;
+
+protected:
+	FVulkanDynamicRHI* RHI;
+	FVulkanCommandListContext* Immediate;
+	FVulkanDevice* Device;
+	FVulkanQueue* Queue;
+	bool bSubmitAtNextSafePoint;
+	bool bAutomaticFlushAfterComputeShader;
+	FVulkanUniformBufferUploader* UniformBufferUploader;
+
+	void BeginOcclusionQueryBatch(FVulkanCmdBuffer* CmdBuffer, uint32 NumQueriesInBatch);
+	void EndOcclusionQueryBatch(FVulkanCmdBuffer* CmdBuffer);
+
+	void SetShaderUniformBuffer(ShaderStage::EStage Stage, const FVulkanUniformBuffer* UniformBuffer, int32 ParameterIndex, const FVulkanShader* Shader);
+
+	struct
+	{
+		VulkanRHI::FTempFrameAllocationBuffer::FTempAllocInfo VertexAllocInfo;
+		uint32 NumVertices = 0;
+		uint32 VertexDataStride = 0;
+
+		VulkanRHI::FTempFrameAllocationBuffer::FTempAllocInfo IndexAllocInfo;
+		VkIndexType IndexType = VK_INDEX_TYPE_MAX_ENUM;
+		uint32 NumPrimitives = 0;
+		uint32 MinVertexIndex = 0;
+		uint32 IndexDataStride = 0;
+	} UserPrimitive;
+
+
+	VulkanRHI::FTempFrameAllocationBuffer TempFrameAllocationBuffer;
+
+	TArray<FString> EventStack;
+
+	FVulkanCommandBufferManager* CommandBufferManager;
+
+	FTransitionAndLayoutManager TransitionAndLayoutManager;
+
+
+	struct FPendingTransition
+	{
+		EResourceTransitionAccess TransitionType;
+
+		// Only one of a) Textures or b) UAVs is active at a time
+		TArray<FRHITexture*, TInlineAllocator<MaxSimultaneousRenderTargets + 1>> Textures;	// a
+
+		TArray<FRHIUnorderedAccessView*, TInlineAllocator<4>> UAVs;	// b
+		FRHIComputeFence* WriteComputeFenceRHI = nullptr;			// b
+		EResourceTransitionPipeline TransitionPipeline;				// b
+
+		bool GatherBarriers(FTransitionAndLayoutManager& TransitionAndLayoutManager, TArray<VkBufferMemoryBarrier>& OutBufferBarriers, 
+			TArray<VkImageMemoryBarrier>& OutImageBarriers) const;
+	};
+
+	void TransitionResources(const FPendingTransition& PendingTransition);
+	static void TransitionUAVResourcesTransferringOwnership(FVulkanCommandListContext& GfxContext, FVulkanCommandListContext& ComputeContext, 
+		EResourceTransitionPipeline Pipeline, const TArray<VkBufferMemoryBarrier>& BufferBarriers, const TArray<VkImageMemoryBarrier>& ImageBarriers);
+
+	FVulkanOcclusionQueryPool* CurrentOcclusionQueryPool = nullptr;
+
+	// List of UAVs which need setting for pixel shaders. D3D treats UAVs like rendertargets so the RHI doesn't make SetUAV calls at the right time
+	struct FPendingPixelUAV
+	{
+		FVulkanUnorderedAccessView* UAV;
+		uint32 BindIndex;
+	};
+	TArray<FPendingPixelUAV> PendingPixelUAVs;
+
+	FVulkanPendingGfxState* PendingGfxState;
+	FVulkanPendingComputeState* PendingComputeState;
+
+	void PrepareForCPURead();
+	void RequestSubmitCurrentCommands();
+
+	void InternalClearMRT(FVulkanCmdBuffer* CmdBuffer, bool bClearColor, int32 NumClearColors, const FLinearColor* ColorArray, bool bClearDepth, float Depth, bool bClearStencil, uint32 Stencil);
+
+public:
+	inline FTransitionAndLayoutManager& GetTransitionAndLayoutManager()
+	{
+		return TransitionAndLayoutManager;
+	}
+
+	FVulkanRenderPass* PrepareRenderPassForPSOCreation(const FGraphicsPipelineStateInitializer& Initializer, const TArray<FInputAttachmentData>& InputAttachmentData);
+	FVulkanRenderPass* PrepareRenderPassForPSOCreation(const FVulkanRenderTargetLayout& Initializer);
+
+private:
+	void RHIClearMRT(bool bClearColor, int32 NumClearColors, const FLinearColor* ColorArray, bool bClearDepth, float Depth, bool bClearStencil, uint32 Stencil);
+
+	inline bool SafePointSubmit()
+	{
+		if (bSubmitAtNextSafePoint)
+		{
+			InternalSubmitActiveCmdBuffer();
+			bSubmitAtNextSafePoint = false;
+			return true;
+		}
+
+		return false;
+	}
+
+	void InternalSubmitActiveCmdBuffer();
+	void FlushAfterComputeShader();
+
+	friend class FVulkanDevice;
+	friend class FVulkanDynamicRHI;
+
+	// Number of times EndFrame() has been called on this context
+	uint64 FrameCounter;
+
+	FVulkanGPUProfiler GpuProfiler;
+	FVulkanGPUTiming* FrameTiming;
+
+	friend struct FVulkanCommandContextContainer;
+};
+
+class FVulkanCommandListContextImmediate : public FVulkanCommandListContext
+{
+public:
+	FVulkanCommandListContextImmediate(FVulkanDynamicRHI* InRHI, FVulkanDevice* InDevice, FVulkanQueue* InQueue);
+};
+
+
+struct FVulkanCommandContextContainer : public IRHICommandContextContainer, public VulkanRHI::FDeviceChild
+{
+	FVulkanCommandListContext* CmdContext;
+
+	FVulkanCommandContextContainer(FVulkanDevice* InDevice);
+
+	virtual IRHICommandContext* GetContext() override final;
+	virtual void FinishContext() override final;
+	virtual void SubmitAndFreeContextContainer(int32 Index, int32 Num) override final;
+
+	/** Custom new/delete with recycling */
+	void* operator new(size_t Size);
+	void operator delete(void* RawMemory);
+
+private:
+	friend class FVulkanDevice;
+};
+
+inline FVulkanCommandListContextImmediate& FVulkanDevice::GetImmediateContext()
+{
+	return *ImmediateContext;
+}
diff --git a/Engine/Source/Runtime/VulkanRHI/Private/VulkanLLM.h b/Engine/Source/Runtime/VulkanRHI/Private/VulkanLLM.h
new file mode 100644
--- /dev/null
+++ b/Engine/Source/Runtime/VulkanRHI/Private/VulkanLLM.h
@@ -0,0 +1,50 @@
+// Copyright 1998-2019 Epic Games, Inc. All Rights Reserved.
+
+/*=============================================================================
+	VulkanLLM.h: Vulkan LLM definitions.
+=============================================================================*/
+
+
+#pragma once
+
+#include "HAL/LowLevelMemTracker.h"
+
+#if VULKAN_USE_LLM
+
+#define LLM_SCOPE_VULKAN(Tag)			LLM_SCOPE((ELLMTag)Tag)
+#define LLM_PLATFORM_SCOPE_VULKAN(Tag)	LLM_PLATFORM_SCOPE((ELLMTag)Tag)
+extern uint64 GVulkanLLMAllocationID;
+#define LLM_TRACK_VULKAN_HIGH_LEVEL_ALLOC(AllocObj, Size) { AllocObj->SetLLMTrackerID(0xDEAD | (++GVulkanLLMAllocationID << 16)); FLowLevelMemTracker::Get().OnLowLevelAlloc(ELLMTracker::Default, (void*)AllocObj->GetLLMTrackerID(), Size, (ELLMTag)ELLMTagVulkan::VulkanMisc); }
+#define LLM_TRACK_VULKAN_HIGH_LEVEL_FREE(AllocObj) { FLowLevelMemTracker::Get().OnLowLevelFree(ELLMTracker::Default, (void*)AllocObj->GetLLMTrackerID()); }
+
+#define LLM_TRACK_VULKAN_HIGH_LEVEL_ALLOCATION
+enum class ELLMTagVulkan : LLM_TAG_TYPE
+{
+	VulkanMisc = (LLM_TAG_TYPE)ELLMTag::PlatformTagStart,
+	VulkanUniformBuffers,
+	VulkanIndexBuffers,
+	VulkanVertexBuffers,
+	VulkanTextures,
+	VulkanShaders,	
+	VulkanFrameTemp,
+	VulkanStagingBuffers,
+	VulkanDriverMemoryCPU,
+	VulkanDriverMemoryGPU,
+	Count,
+};
+
+static_assert((int32)ELLMTagVulkan::Count <= (int32)ELLMTag::PlatformTagEnd, "too many ELLMTagVulkan tags");
+
+namespace VulkanLLM
+{
+	void Initialize();
+}
+
+#else
+
+#define LLM_SCOPE_VULKAN(...)
+#define LLM_PLATFORM_SCOPE_VULKAN(...)
+#define LLM_TRACK_VULKAN_HIGH_LEVEL_ALLOC(...)
+#define LLM_TRACK_VULKAN_HIGH_LEVEL_FREE(...)
+
+#endif		// #if VULKAN_USE_LLM
diff --git a/Engine/Source/Runtime/VulkanRHI/Private/VulkanPendingState.h b/Engine/Source/Runtime/VulkanRHI/Private/VulkanPendingState.h
--- a/Engine/Source/Runtime/VulkanRHI/Private/VulkanPendingState.h
+++ b/Engine/Source/Runtime/VulkanRHI/Private/VulkanPendingState.h
@@ -7,7 +7,13 @@
 #pragma once
 
 // Dependencies
-#include "VulkanRHI.h"
+#include "VulkanGlobals.h"
+#include "VulkanConfiguration.h"
+#include "VulkanState.h"
+#include "VulkanResources.h"
+#include "VulkanUtil.h"
+#include "VulkanViewport.h"
+#include "VulkanDynamicRHI.h"
 #include "VulkanPipeline.h"
 #include "VulkanGlobalUniformBuffer.h"
 #include "VulkanPipelineState.h"
diff --git a/Engine/Source/Runtime/VulkanRHI/Private/VulkanRHIPrivate.h b/Engine/Source/Runtime/VulkanRHI/Private/VulkanRHIPrivate.h
--- a/Engine/Source/Runtime/VulkanRHI/Private/VulkanRHIPrivate.h
+++ b/Engine/Source/Runtime/VulkanRHI/Private/VulkanRHIPrivate.h
@@ -35,7 +35,12 @@
 	#endif
 #endif
 
-#include "VulkanRHI.h"
+#include "VulkanGlobals.h"
+#include "VulkanState.h"
+#include "VulkanResources.h"
+#include "VulkanUtil.h"
+#include "VulkanViewport.h"
+#include "VulkanDynamicRHI.h"
 #include "VulkanGlobalUniformBuffer.h"
 #include "RHI.h"
 #include "VulkanDevice.h"
diff --git a/Engine/Source/Runtime/VulkanRHI/Private/VulkanResources.h b/Engine/Source/Runtime/VulkanRHI/Private/VulkanResources.h
new file mode 100644
--- /dev/null
+++ b/Engine/Source/Runtime/VulkanRHI/Private/VulkanResources.h
@@ -0,0 +1,1685 @@
+// Copyright 1998-2019 Epic Games, Inc. All Rights Reserved.
+
+/*=============================================================================
+	VulkanResources.h: Vulkan resource RHI definitions.
+=============================================================================*/
+
+#pragma once
+
+#include "VulkanConfiguration.h"
+#include "VulkanState.h"
+#include "VulkanUtil.h"
+#include "BoundShaderStateCache.h"
+#include "VulkanShaderResources.h"
+#include "VulkanState.h"
+#include "VulkanMemory.h"
+#include "Misc/ScopeRWLock.h"
+
+class FVulkanDevice;
+class FVulkanQueue;
+class FVulkanCmdBuffer;
+class FVulkanBuffer;
+class FVulkanBufferCPU;
+struct FVulkanTextureBase;
+class FVulkanTexture2D;
+struct FVulkanBufferView;
+class FVulkanResourceMultiBuffer;
+class FVulkanLayout;
+class FVulkanOcclusionQuery;
+
+namespace VulkanRHI
+{
+	class FDeviceMemoryAllocation;
+	class FOldResourceAllocation;
+	struct FPendingBufferLock;
+}
+
+enum
+{
+	NUM_OCCLUSION_QUERIES_PER_POOL = 4096,
+
+	NUM_TIMESTAMP_QUERIES_PER_POOL = 1024,
+};
+
+struct FSamplerYcbcrConversionInitializer
+{
+	VkFormat Format;
+	uint64 ExternalFormat;
+	VkComponentMapping Components;
+	VkSamplerYcbcrModelConversion Model;
+	VkSamplerYcbcrRange Range;
+	VkChromaLocation XOffset;
+	VkChromaLocation YOffset;
+};
+
+// Mirror GPixelFormats with format information for buffers
+extern VkFormat GVulkanBufferFormat[PF_MAX];
+
+/** This represents a vertex declaration that hasn't been combined with a specific shader to create a bound shader. */
+class FVulkanVertexDeclaration : public FRHIVertexDeclaration
+{
+public:
+	FVertexDeclarationElementList Elements;
+
+	FVulkanVertexDeclaration(const FVertexDeclarationElementList& InElements);
+
+	virtual bool GetInitializer(FVertexDeclarationElementList& Out) final override
+	{
+		Out = Elements;
+		return true;
+	}
+
+	static void EmptyCache();
+};
+
+
+class FVulkanShader : public IRefCountedObject
+{
+public:
+	FVulkanShader(FVulkanDevice* InDevice, EShaderFrequency InFrequency, VkShaderStageFlagBits InStageFlag)
+		: ShaderKey(0)
+		, StageFlag(InStageFlag)
+		, Frequency(InFrequency)
+		, Device(InDevice)
+	{
+	}
+
+	virtual ~FVulkanShader();
+
+	void PurgeShaderModules();
+
+	void Setup(const TArray<uint8>& InShaderHeaderAndCode, uint64 InShaderKey);
+
+	VkShaderModule GetOrCreateHandle(const FVulkanLayout* Layout, uint32 LayoutHash)
+	{
+		VkShaderModule* Found = ShaderModules.Find(LayoutHash);
+		if (Found)
+		{
+			return *Found;
+		}
+
+		return CreateHandle(Layout, LayoutHash);
+	}
+
+#if VULKAN_ENABLE_SHADER_DEBUG_NAMES
+	inline const FString& GetDebugName() const
+	{
+		return CodeHeader.DebugName;
+	}
+#endif
+
+	// Name should be pointing to "main_"
+	void GetEntryPoint(ANSICHAR* Name)
+	{
+		FCStringAnsi::Sprintf(Name, "main_%0.8x_%0.8x", Spirv.Num() * sizeof(uint32), CodeHeader.SpirvCRC);
+	}
+
+	FORCEINLINE const FVulkanShaderHeader& GetCodeHeader() const
+	{
+		return CodeHeader;
+	}
+
+	inline uint64 GetShaderKey() const
+	{
+		return ShaderKey;
+	}
+
+protected:
+#if VULKAN_ENABLE_SHADER_DEBUG_NAMES
+	FString							DebugEntryPoint;
+#endif
+	uint64							ShaderKey;
+
+	/** External bindings for this shader. */
+	FVulkanShaderHeader				CodeHeader;
+	TMap<uint32, VkShaderModule>	ShaderModules;
+	const VkShaderStageFlagBits		StageFlag;
+	EShaderFrequency				Frequency;
+
+	TArray<uint32>					Spirv;
+
+	FVulkanDevice*					Device;
+
+	VkShaderModule CreateHandle(const FVulkanLayout* Layout, uint32 LayoutHash);
+
+	friend class FVulkanCommandListContext;
+	friend class FVulkanPipelineStateCacheManager;
+	friend class FVulkanComputeShaderState;
+	friend class FVulkanComputePipeline;
+	friend class FVulkanShaderFactory;
+};
+
+/** This represents a vertex shader that hasn't been combined with a specific declaration to create a bound shader. */
+template<typename BaseResourceType, EShaderFrequency ShaderType, VkShaderStageFlagBits StageFlagBits>
+class TVulkanBaseShader : public BaseResourceType, public FVulkanShader
+{
+private:
+	TVulkanBaseShader(FVulkanDevice* InDevice) :
+		FVulkanShader(InDevice, ShaderType, StageFlagBits)
+	{
+	}
+	friend class FVulkanShaderFactory;
+public:
+	enum { StaticFrequency = ShaderType };
+
+	// IRefCountedObject interface.
+	virtual uint32 AddRef() const override final
+	{
+		return FRHIResource::AddRef();
+	}
+	virtual uint32 Release() const override final
+	{
+		return FRHIResource::Release();
+	}
+	virtual uint32 GetRefCount() const override final
+	{
+		return FRHIResource::GetRefCount();
+	}
+};
+
+typedef TVulkanBaseShader<FRHIVertexShader, SF_Vertex, VK_SHADER_STAGE_VERTEX_BIT>					FVulkanVertexShader;
+typedef TVulkanBaseShader<FRHIPixelShader, SF_Pixel, VK_SHADER_STAGE_FRAGMENT_BIT>					FVulkanPixelShader;
+typedef TVulkanBaseShader<FRHIHullShader, SF_Hull, VK_SHADER_STAGE_TESSELLATION_CONTROL_BIT>		FVulkanHullShader;
+typedef TVulkanBaseShader<FRHIDomainShader, SF_Domain, VK_SHADER_STAGE_TESSELLATION_EVALUATION_BIT>	FVulkanDomainShader;
+typedef TVulkanBaseShader<FRHIComputeShader, SF_Compute, VK_SHADER_STAGE_COMPUTE_BIT>				FVulkanComputeShader;
+typedef TVulkanBaseShader<FRHIGeometryShader, SF_Geometry, VK_SHADER_STAGE_GEOMETRY_BIT>			FVulkanGeometryShader;
+
+class FVulkanShaderFactory
+{
+public:
+	~FVulkanShaderFactory();
+	
+	template <typename ShaderType> 
+	ShaderType* CreateShader(const TArray<uint8>& Code, FVulkanDevice* Device);
+	
+	template <typename ShaderType> 
+	ShaderType* LookupShader(uint64 ShaderKey) const
+	{
+		if (ShaderKey)
+		{
+			FRWScopeLock ScopedLock(Lock, SLT_ReadOnly);
+			FVulkanShader* const * FoundShaderPtr = ShaderMap[ShaderType::StaticFrequency].Find(ShaderKey);
+			if (FoundShaderPtr)
+			{
+				return static_cast<ShaderType*>(*FoundShaderPtr);
+			}
+		}
+		return nullptr;
+	}
+
+	void LookupShaders(const uint64 InShaderKeys[ShaderStage::NumStages], FVulkanShader* OutShaders[ShaderStage::NumStages]) const;
+		
+	void OnDeleteShader(const FVulkanShader& Shader);
+
+private:
+	mutable FRWLock Lock;
+	TMap<uint64, FVulkanShader*> ShaderMap[SF_NumFrequencies];
+};
+
+class FVulkanBoundShaderState : public FRHIBoundShaderState
+{
+public:
+	FVulkanBoundShaderState(
+		FVertexDeclarationRHIParamRef InVertexDeclarationRHI,
+		FVertexShaderRHIParamRef InVertexShaderRHI,
+		FPixelShaderRHIParamRef InPixelShaderRHI,
+		FHullShaderRHIParamRef InHullShaderRHI,
+		FDomainShaderRHIParamRef InDomainShaderRHI,
+		FGeometryShaderRHIParamRef InGeometryShaderRHI
+	);
+
+	virtual ~FVulkanBoundShaderState();
+
+	FORCEINLINE FVulkanVertexShader*   GetVertexShader() const { return (FVulkanVertexShader*)CacheLink.GetVertexShader(); }
+	FORCEINLINE FVulkanPixelShader*    GetPixelShader() const { return (FVulkanPixelShader*)CacheLink.GetPixelShader(); }
+	FORCEINLINE FVulkanHullShader*     GetHullShader() const { return (FVulkanHullShader*)CacheLink.GetHullShader(); }
+	FORCEINLINE FVulkanDomainShader*   GetDomainShader() const { return (FVulkanDomainShader*)CacheLink.GetDomainShader(); }
+	FORCEINLINE FVulkanGeometryShader* GetGeometryShader() const { return (FVulkanGeometryShader*)CacheLink.GetGeometryShader(); }
+
+	const FVulkanShader* GetShader(ShaderStage::EStage Stage) const
+	{
+		switch (Stage)
+		{
+		case ShaderStage::Vertex:		return GetVertexShader();
+#if PLATFORM_SUPPORTS_TESSELLATION_SHADERS
+		case ShaderStage::Hull:			return GetHullShader();
+		case ShaderStage::Domain:		return GetDomainShader();
+#endif
+		case ShaderStage::Pixel:		return GetPixelShader();
+#if VULKAN_SUPPORTS_GEOMETRY_SHADERS
+		case ShaderStage::Geometry:	return GetGeometryShader();
+#endif
+		default: break;
+		}
+		checkf(0, TEXT("Invalid Shader Frequency %d"), (int32)Stage);
+		return nullptr;
+	}
+
+private:
+	FCachedBoundShaderStateLink_Threadsafe CacheLink;
+};
+
+/** Texture/RT wrapper. */
+class FVulkanSurface
+{
+public:
+
+	// Seperate method for creating image, this can be used to measure image size
+	// After VkImage is no longer needed, dont forget to destroy/release it 
+	static VkImage CreateImage(
+		FVulkanDevice& InDevice,
+		VkImageViewType ResourceType,
+		EPixelFormat InFormat,
+		uint32 SizeX, uint32 SizeY, uint32 SizeZ,
+		bool bArray, uint32 ArraySize,
+		uint32 NumMips,
+		uint32 NumSamples,
+		uint32 UEFlags,
+		VkMemoryRequirements& OutMemoryRequirements,
+		VkFormat* OutStorageFormat = nullptr,
+		VkFormat* OutViewFormat = nullptr,
+		VkImageCreateInfo* OutInfo = nullptr,
+		bool bForceLinearTexture = false);
+
+	FVulkanSurface(FVulkanDevice& Device, VkImageViewType ResourceType, EPixelFormat Format,
+					uint32 SizeX, uint32 SizeY, uint32 SizeZ, bool bArray, uint32 ArraySize,
+					uint32 NumMips, uint32 NumSamples, uint32 UEFlags, const FRHIResourceCreateInfo& CreateInfo);
+
+	// Constructor for externally owned Image
+	FVulkanSurface(FVulkanDevice& Device, VkImageViewType ResourceType, EPixelFormat Format,
+					uint32 SizeX, uint32 SizeY, uint32 SizeZ, uint32 NumMips, uint32 NumSamples,
+					VkImage InImage, uint32 UEFlags, const FRHIResourceCreateInfo& CreateInfo);
+
+	virtual ~FVulkanSurface();
+
+	void Destroy();
+
+#if 0
+	/**
+	 * Locks one of the texture's mip-maps.
+	 * @param ArrayIndex Index of the texture array/face in the form Index*6+Face
+	 * @return A pointer to the specified texture data.
+	 */
+	void* Lock(uint32 MipIndex, uint32 ArrayIndex, EResourceLockMode LockMode, uint32& DestStride);
+
+	/** Unlocks a previously locked mip-map.
+	 * @param ArrayIndex Index of the texture array/face in the form Index*6+Face
+	 */
+	void Unlock(uint32 MipIndex, uint32 ArrayIndex);
+#endif
+
+	/**
+	 * Returns how much memory is used by the surface
+	 */
+	uint32 GetMemorySize() const
+	{
+		return MemoryRequirements.size;
+	}
+
+	/**
+	 * Returns one of the texture's mip-maps stride.
+	 */
+	void GetMipStride(uint32 MipIndex, uint32& Stride);
+
+	/*
+	 * Returns the memory offset to the texture's mip-map.
+	 */
+	void GetMipOffset(uint32 MipIndex, uint32& Offset);
+
+	/**
+	* Returns how much memory a single mip uses.
+	*/
+	void GetMipSize(uint32 MipIndex, uint32& MipBytes);
+
+	inline VkImageViewType GetViewType() const { return ViewType; }
+
+	inline VkImageTiling GetTiling() const { return Tiling; }
+
+	inline uint32 GetNumMips() const { return NumMips; }
+
+	inline uint32 GetNumSamples() const { return NumSamples; }
+
+	inline uint32 GetNumberOfArrayLevels() const
+	{
+		switch (ViewType)
+		{
+		case VK_IMAGE_VIEW_TYPE_1D:
+		case VK_IMAGE_VIEW_TYPE_2D:
+		case VK_IMAGE_VIEW_TYPE_3D:
+			return 1;
+		case VK_IMAGE_VIEW_TYPE_2D_ARRAY:
+			return NumArrayLevels;
+		case VK_IMAGE_VIEW_TYPE_CUBE:
+			return 6;
+		case VK_IMAGE_VIEW_TYPE_CUBE_ARRAY:
+			return 6 * NumArrayLevels;
+		default:
+			UE_LOG(LogVulkanRHI, Error, TEXT("Invalid ViewType %d"), (uint32)ViewType);
+			return 1;
+		}
+	}
+
+	// Full includes Depth+Stencil
+	inline VkImageAspectFlags GetFullAspectMask() const
+	{
+		return FullAspectMask;
+	}
+
+	// Only Depth or Stencil
+	inline VkImageAspectFlags GetPartialAspectMask() const
+	{
+		return PartialAspectMask;
+	}
+
+	inline bool IsDepthOrStencilAspect() const
+	{
+		return (FullAspectMask & (VK_IMAGE_ASPECT_DEPTH_BIT | VK_IMAGE_ASPECT_STENCIL_BIT)) != 0;
+	}
+
+	inline bool IsImageOwner() const
+	{
+		return bIsImageOwner;
+	}
+
+	inline VkDeviceMemory GetAllocationHandle() const
+	{
+		if (ResourceAllocation.IsValid())
+		{
+			return ResourceAllocation->GetHandle();
+		}
+		else
+		{
+			return VK_NULL_HANDLE;
+		}
+	}
+
+	inline uint64 GetAllocationOffset() const
+	{
+		if (ResourceAllocation.IsValid())
+		{
+			return ResourceAllocation->GetOffset();
+		}
+		else
+		{
+			return 0;
+		}
+	}
+
+	FVulkanDevice* Device;
+
+	VkImage Image;
+	
+	// Removes SRGB if requested, used to upload data
+	VkFormat StorageFormat;
+	// Format for SRVs, render targets
+	VkFormat ViewFormat;
+	uint32 Width, Height, Depth;
+	// UE format
+	EPixelFormat PixelFormat;
+	uint32 UEFlags;
+	VkMemoryPropertyFlags MemProps;
+	VkMemoryRequirements MemoryRequirements;
+	uint32 NumArrayLevels;
+
+	static void InternalLockWrite(FVulkanCommandListContext& Context, FVulkanSurface* Surface, const VkImageSubresourceRange& SubresourceRange, const VkBufferImageCopy& Region, VulkanRHI::FStagingBuffer* StagingBuffer);
+
+private:
+
+	// Used to clear render-target objects on creation
+	void InitialClear(FVulkanCommandListContext& Context, const FClearValueBinding& ClearValueBinding, bool bTransitionToPresentable);
+	friend struct FRHICommandInitialClearTexture;
+
+private:
+	VkImageTiling Tiling;
+	VkImageViewType	ViewType;
+
+	bool bIsImageOwner;
+	TRefCountPtr<VulkanRHI::FOldResourceAllocation> ResourceAllocation;
+
+	uint32 NumMips;
+	uint32 NumSamples;
+
+	VkImageAspectFlags FullAspectMask;
+	VkImageAspectFlags PartialAspectMask;
+
+	friend struct FVulkanTextureBase;
+};
+
+
+struct FVulkanTextureView
+{
+	FVulkanTextureView()
+		: View(VK_NULL_HANDLE)
+		, Image(VK_NULL_HANDLE)
+		, ViewId(0)
+	{
+	}
+
+	void Create(FVulkanDevice& Device, VkImage InImage, VkImageViewType ViewType, VkImageAspectFlags AspectFlags, EPixelFormat UEFormat, VkFormat Format, uint32 FirstMip, uint32 NumMips, uint32 ArraySliceIndex, uint32 NumArraySlices, bool bUseIdentitySwizzle = false);
+	void Create(FVulkanDevice& Device, VkImage InImage, VkImageViewType ViewType, VkImageAspectFlags AspectFlags, EPixelFormat UEFormat, VkFormat Format, uint32 FirstMip, uint32 NumMips, uint32 ArraySliceIndex, uint32 NumArraySlices, FSamplerYcbcrConversionInitializer& ConversionInitializer, bool bUseIdentitySwizzle = false);
+	void Destroy(FVulkanDevice& Device);
+
+	VkImageView View;
+	VkImage Image;
+	uint32 ViewId;
+
+private:
+	static VkImageView StaticCreate(FVulkanDevice& Device, VkImage InImage, VkImageViewType ViewType, VkImageAspectFlags AspectFlags, EPixelFormat UEFormat, VkFormat Format, uint32 FirstMip, uint32 NumMips, uint32 ArraySliceIndex, uint32 NumArraySlices, bool bUseIdentitySwizzle, const FSamplerYcbcrConversionInitializer* ConversionInitializer);
+};
+
+/** The base class of resources that may be bound as shader resources. */
+class FVulkanBaseShaderResource : public IRefCountedObject
+{
+};
+
+struct FVulkanTextureBase : public FVulkanBaseShaderResource
+{
+	inline static FVulkanTextureBase* Cast(FRHITexture* Texture)
+	{
+		check(Texture);
+		FVulkanTextureBase* OutTexture = (FVulkanTextureBase*)Texture->GetTextureBaseRHI();
+		check(OutTexture);
+		return OutTexture;
+	}
+
+	FVulkanTextureBase(FVulkanDevice& Device, VkImageViewType ResourceType, EPixelFormat Format, uint32 SizeX, uint32 SizeY, uint32 SizeZ, bool bArray, uint32 ArraySize, uint32 NumMips, uint32 NumSamples, uint32 UEFlags, const FRHIResourceCreateInfo& CreateInfo);
+	FVulkanTextureBase(FVulkanDevice& Device, VkImageViewType ResourceType, EPixelFormat Format, uint32 SizeX, uint32 SizeY, uint32 SizeZ, uint32 NumMips, uint32 NumSamples, uint32 NumSamplesTileMem, VkImage InImage, VkDeviceMemory InMem, uint32 UEFlags, const FRHIResourceCreateInfo& CreateInfo = FRHIResourceCreateInfo());
+	FVulkanTextureBase(FVulkanDevice& Device, VkImageViewType ResourceType, EPixelFormat Format, uint32 SizeX, uint32 SizeY, uint32 SizeZ, uint32 NumMips, uint32 NumSamples, VkImage InImage, VkDeviceMemory InMem, FSamplerYcbcrConversionInitializer& ConversionInitializer, uint32 UEFlags, const FRHIResourceCreateInfo& CreateInfo = FRHIResourceCreateInfo());
+
+	virtual ~FVulkanTextureBase();
+
+	void AliasTextureResources(const FVulkanTextureBase* SrcTexture);
+
+	FVulkanSurface Surface;
+
+	// View with all mips/layers
+	FVulkanTextureView DefaultView;
+	// View with all mips/layers, but if it's a Depth/Stencil, only the Depth view
+	FVulkanTextureView* PartialView;
+
+#if VULKAN_USE_MSAA_RESOLVE_ATTACHMENTS
+	// Surface and view for MSAA render target, valid only when created with NumSamples > 1
+	FVulkanSurface* MSAASurface;
+	FVulkanTextureView MSAAView;
+#endif
+
+	bool bIsAliased;
+
+private:
+	void DestroyViews();
+};
+
+class FVulkanBackBuffer;
+class FVulkanTexture2D : public FRHITexture2D, public FVulkanTextureBase
+{
+public:
+	FVulkanTexture2D(FVulkanDevice& Device, EPixelFormat Format, uint32 SizeX, uint32 SizeY, uint32 NumMips, uint32 NumSamples, uint32 UEFlags, const FRHIResourceCreateInfo& CreateInfo);
+	FVulkanTexture2D(FVulkanDevice& Device, EPixelFormat Format, uint32 SizeX, uint32 SizeY, uint32 NumMips, uint32 NumSamples, uint32 NumSamplesTileMem, VkImage Image, uint32 UEFlags, const FRHIResourceCreateInfo& CreateInfo);
+	FVulkanTexture2D(FVulkanDevice& Device, EPixelFormat Format, uint32 SizeX, uint32 SizeY, uint32 NumMips, uint32 NumSamples, VkImage Image, struct FSamplerYcbcrConversionInitializer& ConversionInitializer, uint32 UEFlags, const FRHIResourceCreateInfo& CreateInfo);
+
+	virtual ~FVulkanTexture2D();
+
+	// IRefCountedObject interface.
+	virtual uint32 AddRef() const override final
+	{
+		return FRHIResource::AddRef();
+	}
+	virtual uint32 Release() const override final
+	{
+		return FRHIResource::Release();
+	}
+	virtual uint32 GetRefCount() const override final
+	{
+		return FRHIResource::GetRefCount();
+	}
+
+	virtual FVulkanBackBuffer* GetBackBuffer()
+	{
+		return nullptr;
+	}
+
+	virtual void* GetTextureBaseRHI() override final
+	{
+		FVulkanTextureBase* Base = static_cast<FVulkanTextureBase*>(this);
+		return Base;
+	}
+
+	virtual void* GetNativeResource() const
+	{
+		return (void*)Surface.Image;
+	}
+};
+
+class FVulkanBackBuffer : public FVulkanTexture2D
+{
+public:
+	FVulkanBackBuffer(FVulkanDevice& Device, EPixelFormat Format, uint32 SizeX, uint32 SizeY, uint32 UEFlags);
+	FVulkanBackBuffer(FVulkanDevice& Device, EPixelFormat Format, uint32 SizeX, uint32 SizeY, VkImage Image, uint32 UEFlags);
+	virtual ~FVulkanBackBuffer();
+
+	virtual FVulkanBackBuffer* GetBackBuffer() override final
+	{
+		return this;
+	}
+};
+
+class FVulkanBackBufferReference : public FRHITexture2D
+{
+public:
+	FVulkanBackBufferReference(EPixelFormat Format, uint32 SizeX, uint32 SizeY, uint32 UEFlags)
+		: FRHITexture2D(SizeX, SizeY, 1, 1, Format, UEFlags, FRHIResourceCreateInfo().ClearValueBinding)
+		, AcquiredBackBuffer(nullptr)
+	{
+	}
+
+	virtual FRHITexture2D* GetTexture2D() override 
+	{ 
+		return AcquiredBackBuffer; 
+	}
+	
+	virtual void* GetTextureBaseRHI() override
+	{
+		FVulkanTextureBase* Base = static_cast<FVulkanTextureBase*>(AcquiredBackBuffer);
+		return Base;
+	}
+
+	void SetBackBuffer(FVulkanBackBuffer* InAcquiredBackbuffer)
+	{
+		AcquiredBackBuffer = InAcquiredBackbuffer;
+	}
+
+private:
+	FVulkanBackBuffer* AcquiredBackBuffer;
+};
+
+class FVulkanTexture2DArray : public FRHITexture2DArray, public FVulkanTextureBase
+{
+public:
+	// Constructor, just calls base and Surface constructor
+	FVulkanTexture2DArray(FVulkanDevice& Device, EPixelFormat Format, uint32 SizeX, uint32 SizeY, uint32 ArraySize, uint32 NumMips, uint32 Flags, FResourceBulkDataInterface* BulkData, const FClearValueBinding& InClearValue);
+	FVulkanTexture2DArray(FVulkanDevice& Device, EPixelFormat Format, uint32 SizeX, uint32 SizeY, uint32 ArraySize, uint32 NumMips, VkImage Image, uint32 Flags, FResourceBulkDataInterface* BulkData, const FClearValueBinding& InClearValue);
+		
+	// IRefCountedObject interface.
+	virtual uint32 AddRef() const override final
+	{
+		return FRHIResource::AddRef();
+	}
+	virtual uint32 Release() const override final
+	{
+		return FRHIResource::Release();
+	}
+	virtual uint32 GetRefCount() const override final
+	{
+		return FRHIResource::GetRefCount();
+	}
+
+	virtual void* GetTextureBaseRHI() override final
+	{
+		return (FVulkanTextureBase*)this;
+	}
+
+	virtual void* GetNativeResource() const
+	{
+		return (void*)Surface.Image;
+	}
+};
+
+class FVulkanTexture3D : public FRHITexture3D, public FVulkanTextureBase
+{
+public:
+	// Constructor, just calls base and Surface constructor
+	FVulkanTexture3D(FVulkanDevice& Device, EPixelFormat Format, uint32 SizeX, uint32 SizeY, uint32 SizeZ, uint32 NumMips, uint32 Flags, FResourceBulkDataInterface* BulkData, const FClearValueBinding& InClearValue);
+	FVulkanTexture3D(FVulkanDevice& Device, EPixelFormat Format, uint32 SizeX, uint32 SizeY, uint32 SizeZ, uint32 NumMips, VkImage Image, uint32 Flags, FResourceBulkDataInterface* BulkData, const FClearValueBinding& InClearValue);
+	virtual ~FVulkanTexture3D();
+
+	// IRefCountedObject interface.
+	virtual uint32 AddRef() const override final
+	{
+		return FRHIResource::AddRef();
+	}
+	virtual uint32 Release() const override final
+	{
+		return FRHIResource::Release();
+	}
+	virtual uint32 GetRefCount() const override final
+	{
+		return FRHIResource::GetRefCount();
+	}
+
+	virtual void* GetTextureBaseRHI() override final
+	{
+		return (FVulkanTextureBase*)this;
+	}
+
+	virtual void* GetNativeResource() const
+	{
+		return (void*)Surface.Image;
+	}
+};
+
+class FVulkanTextureCube : public FRHITextureCube, public FVulkanTextureBase
+{
+public:
+	FVulkanTextureCube(FVulkanDevice& Device, EPixelFormat Format, uint32 Size, bool bArray, uint32 ArraySize, uint32 NumMips, uint32 Flags, FResourceBulkDataInterface* BulkData, const FClearValueBinding& InClearValue);
+	FVulkanTextureCube(FVulkanDevice& Device, EPixelFormat Format, uint32 Size, bool bArray, uint32 ArraySize, uint32 NumMips, VkImage Image, uint32 Flags, FResourceBulkDataInterface* BulkData, const FClearValueBinding& InClearValue);
+	virtual ~FVulkanTextureCube();
+
+	// IRefCountedObject interface.
+	virtual uint32 AddRef() const override final
+	{
+		return FRHIResource::AddRef();
+	}
+	virtual uint32 Release() const override final
+	{
+		return FRHIResource::Release();
+	}
+	virtual uint32 GetRefCount() const override final
+	{
+		return FRHIResource::GetRefCount();
+	}
+
+	virtual void* GetTextureBaseRHI() override final
+	{
+		return (FVulkanTextureBase*)this;
+	}
+
+	virtual void* GetNativeResource() const
+	{
+		return (void*)Surface.Image;
+	}
+};
+
+class FVulkanTextureReference : public FRHITextureReference, public FVulkanTextureBase
+{
+public:
+	explicit FVulkanTextureReference(FVulkanDevice& Device, FLastRenderTimeContainer* InLastRenderTime)
+	:	FRHITextureReference(InLastRenderTime)
+	,	FVulkanTextureBase(Device, VK_IMAGE_VIEW_TYPE_MAX_ENUM, PF_Unknown, 0, 0, 0, 1, 1, 1, VK_NULL_HANDLE, VK_NULL_HANDLE, 0)
+	{}
+
+	// IRefCountedObject interface.
+	virtual uint32 AddRef() const override final
+	{
+		return FRHIResource::AddRef();
+	}
+
+	virtual uint32 Release() const override final
+	{
+		return FRHIResource::Release();
+	}
+
+	virtual uint32 GetRefCount() const override final
+	{
+		return FRHIResource::GetRefCount();
+	}
+
+	virtual void* GetTextureBaseRHI() override final
+	{
+		return GetReferencedTexture()->GetTextureBaseRHI();
+	}
+
+	virtual void* GetNativeResource() const
+	{
+		return (void*)Surface.Image;
+	}
+
+	void SetReferencedTexture(FRHITexture* InTexture);
+};
+
+/** Given a pointer to a RHI texture that was created by the Vulkan RHI, returns a pointer to the FVulkanTextureBase it encapsulates. */
+inline FVulkanTextureBase* GetVulkanTextureFromRHITexture(FRHITexture* Texture)
+{
+	if (!Texture)
+	{
+		return NULL;
+	}
+	else if (FRHITexture2D* Tex2D = Texture->GetTexture2D())
+	{
+		return static_cast<FVulkanTexture2D*>(Tex2D);
+	}
+	else if (FRHITextureReference* TexRef = Texture->GetTextureReference())
+	{
+		return static_cast<FVulkanTextureReference*>(TexRef);
+	}
+	else if (FRHITexture2DArray* Tex2DArray = Texture->GetTexture2DArray())
+	{
+		return static_cast<FVulkanTexture2DArray*>(Tex2DArray);
+	}
+	else if (FRHITexture3D* Tex3D = Texture->GetTexture3D())
+	{
+		return static_cast<FVulkanTexture3D*>(Tex3D);
+	}
+	else if (FRHITextureCube* TexCube = Texture->GetTextureCube())
+	{
+		return static_cast<FVulkanTextureCube*>(TexCube);
+	}
+	else
+	{
+		UE_LOG(LogVulkanRHI, Fatal, TEXT("Unknown Vulkan RHI texture type"));
+		return NULL;
+	}
+}
+
+class FVulkanQueryPool : public VulkanRHI::FDeviceChild 
+{
+public:
+	FVulkanQueryPool(FVulkanDevice* InDevice, uint32 InMaxQueries, VkQueryType InQueryType);
+	virtual ~FVulkanQueryPool();
+
+	inline uint32 GetMaxQueries() const
+	{
+		return MaxQueries;
+	}
+
+	inline VkQueryPool GetHandle() const
+	{
+		return QueryPool;
+	}
+
+	inline uint64 GetResultValue(uint32 Index) const
+	{
+		return QueryOutput[Index];
+	}
+
+protected:
+	VkQueryPool QueryPool;
+	VkEvent ResetEvent;
+	const uint32 MaxQueries;
+	const VkQueryType QueryType;
+	TArray<uint64> QueryOutput;
+};
+
+class FVulkanOcclusionQueryPool : public FVulkanQueryPool
+{
+public:
+	FVulkanOcclusionQueryPool(FVulkanDevice* InDevice, uint32 InMaxQueries)
+		: FVulkanQueryPool(InDevice, InMaxQueries, VK_QUERY_TYPE_OCCLUSION)
+	{
+		AcquiredIndices.AddZeroed(Align(InMaxQueries, 64) / 64);
+		AllocatedQueries.AddZeroed(InMaxQueries);
+	}
+
+	inline uint32 AcquireIndex(FVulkanOcclusionQuery* Query)
+	{
+		check(NumUsedQueries < MaxQueries);
+		const uint32 Index = NumUsedQueries;
+		const uint32 Word = Index / 64;
+		const uint32 Bit = Index % 64;
+		const uint64 Mask = (uint64)1 << (uint64)Bit;
+		const uint64& WordValue = AcquiredIndices[Word];
+		AcquiredIndices[Word] = WordValue | Mask;
+		++NumUsedQueries;
+		ensure(AllocatedQueries[Index] == nullptr);
+		AllocatedQueries[Index] = Query;
+		return Index;
+	}
+
+	inline void ReleaseIndex(uint32 Index)
+	{
+		check(Index < NumUsedQueries);
+		const uint32 Word = Index / 64;
+		const uint32 Bit = Index % 64;
+		const uint64 Mask = (uint64)1 << (uint64)Bit;
+		const uint64& WordValue = AcquiredIndices[Word];
+		ensure((WordValue & Mask) == Mask);
+		AcquiredIndices[Word] = WordValue & (~Mask);
+		AllocatedQueries[Index] = nullptr;
+	}
+
+	inline void EndBatch(FVulkanCmdBuffer* InCmdBuffer)
+	{
+		ensure(State == EState::RHIT_PostBeginBatch);
+		State = EState::RHIT_PostEndBatch;
+		SetFence(InCmdBuffer);
+	}
+
+	bool CanBeReused();
+
+	inline bool TryGetResults(bool bWait)
+	{
+		if (State == RT_PostGetResults)
+		{
+			return true;
+		}
+
+		if (State == RHIT_PostEndBatch)
+		{
+			return InternalTryGetResults(bWait);
+		}
+
+		return false;
+	}
+
+	void Reset(FVulkanCmdBuffer* InCmdBuffer, uint32 InFrameNumber);
+
+	bool IsStalePool() const;
+
+	void FlushAllocatedQueries();
+
+	enum EState
+	{
+		Undefined,
+		RHIT_PostBeginBatch,
+		RHIT_PostEndBatch,
+		RT_PostGetResults,
+	};
+	EState State = Undefined;
+	
+	// frame number when pool was placed into free list
+	uint32 FreedFrameNumber = UINT32_MAX;
+protected:
+	uint32 NumUsedQueries = 0;
+	TArray<FVulkanOcclusionQuery*> AllocatedQueries;
+	TArray<uint64> AcquiredIndices;
+	bool InternalTryGetResults(bool bWait);
+	void SetFence(FVulkanCmdBuffer* InCmdBuffer);
+
+	FVulkanCmdBuffer* CmdBuffer = nullptr;
+	uint64 FenceCounter = UINT64_MAX;
+	uint32 FrameNumber = UINT32_MAX;
+};
+
+class FVulkanTimingQueryPool : public FVulkanQueryPool
+{
+public:
+	FVulkanTimingQueryPool(FVulkanDevice* InDevice, uint32 InBufferSize)
+		: FVulkanQueryPool(InDevice, InBufferSize * 2, VK_QUERY_TYPE_TIMESTAMP)
+		, BufferSize(InBufferSize)
+	{
+		TimestampListHandles.AddZeroed(InBufferSize * 2);
+	}
+
+	uint32 CurrentTimestamp = 0;
+	uint32 NumIssuedTimestamps = 0;
+	const uint32 BufferSize;
+
+	struct FCmdBufferFence
+	{
+		FVulkanCmdBuffer* CmdBuffer;
+		uint64 FenceCounter;
+	};
+	TArray<FCmdBufferFence> TimestampListHandles;
+
+	VulkanRHI::FStagingBuffer* ResultsBuffer = nullptr;
+};
+
+class FVulkanRenderQuery : public FRHIRenderQuery
+{
+public:
+	FVulkanRenderQuery(ERenderQueryType InType)
+		: QueryType(InType)
+	{
+	}
+
+	virtual ~FVulkanRenderQuery() {}
+
+	const ERenderQueryType QueryType;
+	uint64 Result = 0;
+
+	uint32 IndexInPool = UINT32_MAX;
+};
+
+class FVulkanOcclusionQuery : public FVulkanRenderQuery
+{
+public:
+	FVulkanOcclusionQuery();
+	virtual ~FVulkanOcclusionQuery();
+
+	enum class EState
+	{
+		Undefined,
+		RHI_PostBegin,
+		RHI_PostEnd,
+		RT_GotResults,
+		FlushedFromPoolHadResults,
+	};
+
+	FVulkanOcclusionQueryPool* Pool = nullptr;
+
+	void ReleaseFromPool();
+
+	EState State = EState::Undefined;
+};
+
+class FVulkanTimingQuery : public FVulkanRenderQuery
+{
+public:
+	FVulkanTimingQuery(FVulkanDevice* InDevice);
+	virtual ~FVulkanTimingQuery();
+
+	FVulkanTimingQueryPool Pool;
+};
+
+struct FVulkanBufferView : public FRHIResource, public VulkanRHI::FDeviceChild
+{
+	FVulkanBufferView(FVulkanDevice* InDevice)
+		: VulkanRHI::FDeviceChild(InDevice)
+		, View(VK_NULL_HANDLE)
+		, ViewId(0)
+		, Flags(0)
+		, Offset(0)
+		, Size(0)
+	{
+	}
+
+	virtual ~FVulkanBufferView()
+	{
+		Destroy();
+	}
+
+	void Create(FVulkanBuffer& Buffer, EPixelFormat Format, uint32 InOffset, uint32 InSize);
+	void Create(FVulkanResourceMultiBuffer* Buffer, EPixelFormat Format, uint32 InOffset, uint32 InSize);
+	void Create(VkFormat Format, FVulkanResourceMultiBuffer* Buffer, uint32 InOffset, uint32 InSize);
+	void Destroy();
+
+	VkBufferView View;
+	uint32 ViewId;
+	VkFlags Flags;
+	uint32 Offset;
+	uint32 Size;
+};
+
+class FVulkanBuffer : public FRHIResource
+{
+public:
+	FVulkanBuffer(FVulkanDevice& Device, uint32 InSize, VkFlags InUsage, VkMemoryPropertyFlags InMemPropertyFlags, bool bAllowMultiLock, const char* File, int32 Line);
+	virtual ~FVulkanBuffer();
+
+	inline VkBuffer GetBufferHandle() const { return Buf; }
+
+	inline uint32 GetSize() const { return Size; }
+
+	void* Lock(uint32 InSize, uint32 InOffset = 0);
+
+	void Unlock();
+
+	inline VkFlags GetFlags() const { return Usage; }
+
+private:
+	FVulkanDevice& Device;
+	VkBuffer Buf;
+	VulkanRHI::FDeviceMemoryAllocation* Allocation;
+	uint32 Size;
+	VkFlags Usage;
+
+	void* BufferPtr;	
+	VkMappedMemoryRange MappedRange;
+
+	bool bAllowMultiLock;
+	int32 LockStack;
+};
+
+struct FVulkanRingBuffer : public VulkanRHI::FDeviceChild
+{
+public:
+	FVulkanRingBuffer(FVulkanDevice* InDevice, uint64 TotalSize, VkFlags Usage, VkMemoryPropertyFlags MemPropertyFlags);
+	~FVulkanRingBuffer();
+
+	// Allocate some space in the ring buffer
+	inline uint64 AllocateMemory(uint64 Size, uint32 Alignment, FVulkanCmdBuffer* InCmdBuffer)
+	{
+		Alignment = FMath::Max(Alignment, MinAlignment);
+		uint64 AllocationOffset = Align<uint64>(BufferOffset, Alignment);
+		if (AllocationOffset + Size <= BufferSize)
+		{
+			BufferOffset = AllocationOffset + Size;
+			return AllocationOffset;
+		}
+
+		return WrapAroundAllocateMemory(Size, Alignment, InCmdBuffer);
+	}
+
+	inline VulkanRHI::FBufferAllocation* GetBufferAllocation() const
+	{
+		return BufferSuballocation->GetBufferAllocation();
+	}
+
+	inline uint32 GetBufferOffset() const
+	{
+		return BufferSuballocation->GetOffset();
+	}
+
+	inline VkBuffer GetHandle() const
+	{
+		return BufferSuballocation->GetHandle();
+	}
+
+	inline void* GetMappedPointer()
+	{
+		return BufferSuballocation->GetMappedPointer();
+	}
+
+protected:
+	uint64 BufferSize;
+	uint64 BufferOffset;
+	uint32 MinAlignment;
+	VulkanRHI::FBufferSuballocation* BufferSuballocation;
+
+	// Fence for wrapping around
+	FVulkanCmdBuffer* FenceCmdBuffer = nullptr;
+	uint64 FenceCounter = 0;
+
+	uint64 WrapAroundAllocateMemory(uint64 Size, uint32 Alignment, FVulkanCmdBuffer* InCmdBuffer);
+};
+
+struct FVulkanUniformBufferUploader : public VulkanRHI::FDeviceChild
+{
+public:
+	FVulkanUniformBufferUploader(FVulkanDevice* InDevice);
+	~FVulkanUniformBufferUploader();
+
+	uint8* GetCPUMappedPointer()
+	{
+		return (uint8*)CPUBuffer->GetMappedPointer();
+	}
+
+	uint64 AllocateMemory(uint64 Size, uint32 Alignment, FVulkanCmdBuffer* InCmdBuffer)
+	{
+		return CPUBuffer->AllocateMemory(Size, Alignment, InCmdBuffer);
+	}
+
+	VulkanRHI::FBufferAllocation* GetCPUBufferAllocation() const
+	{
+		return CPUBuffer->GetBufferAllocation();
+	}
+
+	VkBuffer GetCPUBufferHandle() const
+	{
+		return CPUBuffer->GetHandle();
+	}
+
+	inline uint32 GetCPUBufferOffset() const
+	{
+		return CPUBuffer->GetBufferOffset();
+	}
+
+protected:
+	FVulkanRingBuffer* CPUBuffer;
+	friend class FVulkanCommandListContext;
+};
+
+class FVulkanResourceMultiBuffer : public VulkanRHI::FDeviceChild
+{
+public:
+	FVulkanResourceMultiBuffer(FVulkanDevice* InDevice, VkBufferUsageFlags InBufferUsageFlags, uint32 InSize, uint32 InUEUsage, FRHIResourceCreateInfo& CreateInfo, class FRHICommandListImmediate* InRHICmdList = nullptr);
+	virtual ~FVulkanResourceMultiBuffer();
+
+	inline VulkanRHI::FBufferAllocation* GetBufferAllocation() const
+	{
+		return Current.BufferAllocation;
+	}
+
+	inline VkBuffer GetHandle() const
+	{
+		return Current.Handle;
+	}
+
+	inline bool IsDynamic() const
+	{
+		return NumBuffers > 1;
+	}
+
+	inline int32 GetDynamicIndex() const
+	{
+		return DynamicBufferIndex;
+	}
+
+	inline bool IsVolatile() const
+	{
+		return NumBuffers == 0;
+	}
+
+	inline uint32 GetVolatileLockCounter() const
+	{
+		check(IsVolatile());
+		return VolatileLockInfo.LockCounter;
+	}
+
+	inline int32 GetNumBuffers() const
+	{
+		return NumBuffers;
+	}
+
+	// Offset used for Binding a VkBuffer
+	inline uint32 GetOffset() const
+	{
+		return Current.Offset;
+	}
+
+	inline VkBufferUsageFlags GetBufferUsageFlags() const
+	{
+		return BufferUsageFlags;
+	}
+
+	void* Lock(bool bFromRenderingThread, EResourceLockMode LockMode, uint32 Size, uint32 Offset);
+	void Unlock(bool bFromRenderingThread);
+
+	void Swap(FVulkanResourceMultiBuffer& Other);
+
+protected:
+	uint32 UEUsage;
+	VkBufferUsageFlags BufferUsageFlags;
+	uint32 NumBuffers;
+	uint32 DynamicBufferIndex;
+
+	enum
+	{
+		NUM_BUFFERS = 3,
+	};
+
+	TRefCountPtr<VulkanRHI::FBufferSuballocation> Buffers[NUM_BUFFERS];
+	struct
+	{
+		VulkanRHI::FBufferSuballocation* SubAlloc = nullptr;
+		VulkanRHI::FBufferAllocation* BufferAllocation = nullptr;
+		VkBuffer Handle = VK_NULL_HANDLE;
+		uint64 Offset = 0;
+	} Current;
+	VulkanRHI::FTempFrameAllocationBuffer::FTempAllocInfo VolatileLockInfo;
+
+	static void InternalUnlock(FVulkanCommandListContext& Context, VulkanRHI::FPendingBufferLock& PendingLock, FVulkanResourceMultiBuffer* MultiBuffer, int32 InDynamicBufferIndex);
+
+	friend class FVulkanCommandListContext;
+	friend struct FRHICommandMultiBufferUnlock;
+};
+
+class FVulkanIndexBuffer : public FRHIIndexBuffer, public FVulkanResourceMultiBuffer
+{
+public:
+	FVulkanIndexBuffer(FVulkanDevice* InDevice, uint32 InStride, uint32 InSize, uint32 InUsage, FRHIResourceCreateInfo& CreateInfo, class FRHICommandListImmediate* InRHICmdList);
+
+	inline VkIndexType GetIndexType() const
+	{
+		return IndexType;
+	}
+
+	void Swap(FVulkanIndexBuffer& Other);
+
+private:
+	VkIndexType IndexType;
+};
+
+class FVulkanVertexBuffer : public FRHIVertexBuffer, public FVulkanResourceMultiBuffer
+{
+public:
+	FVulkanVertexBuffer(FVulkanDevice* InDevice, uint32 InSize, uint32 InUsage, FRHIResourceCreateInfo& CreateInfo, class FRHICommandListImmediate* InRHICmdList);
+
+	void Swap(FVulkanVertexBuffer& Other);
+};
+
+class FVulkanUniformBuffer : public FRHIUniformBuffer
+{
+public:
+	FVulkanUniformBuffer(const FRHIUniformBufferLayout& InLayout, const void* Contents, EUniformBufferUsage InUsage, EUniformBufferValidation Validation);
+
+	const TArray<TRefCountPtr<FRHIResource>>& GetResourceTable() const { return ResourceTable; }
+
+	void UpdateResourceTable(const FRHIUniformBufferLayout& InLayout, const void* Contents, int32 ResourceNum);
+	void UpdateResourceTable(FRHIResource** Resources, int32 ResourceNum);
+
+protected:
+	TArray<TRefCountPtr<FRHIResource>> ResourceTable;
+};
+
+class FVulkanEmulatedUniformBuffer : public FVulkanUniformBuffer
+{
+public:
+	FVulkanEmulatedUniformBuffer(const FRHIUniformBufferLayout& InLayout, const void* Contents, EUniformBufferUsage InUsage, EUniformBufferValidation Validation);
+
+	TArray<uint8> ConstantData;
+
+	void UpdateConstantData(const void* Contents, int32 ContentsSize);
+};
+
+class FVulkanRealUniformBuffer : public FVulkanUniformBuffer, public FVulkanResourceMultiBuffer
+{
+public:
+	FVulkanRealUniformBuffer(FVulkanDevice& Device, const FRHIUniformBufferLayout& InLayout, const void* Contents, EUniformBufferUsage InUsage, EUniformBufferValidation Validation);
+
+	void Update(const void* Contents, int32 ContentsSize);
+
+private:
+	TArray<TRefCountPtr<FRHIResource>> ResourceTable;
+};
+
+class FVulkanStructuredBuffer : public FRHIStructuredBuffer, public FVulkanResourceMultiBuffer
+{
+public:
+	FVulkanStructuredBuffer(FVulkanDevice* InDevice, uint32 Stride, uint32 Size, FRHIResourceCreateInfo& CreateInfo, uint32 InUsage);
+
+	~FVulkanStructuredBuffer();
+
+};
+
+
+
+class FVulkanUnorderedAccessView : public FRHIUnorderedAccessView, public VulkanRHI::FDeviceChild
+{
+public:
+	// the potential resources to refer to with the UAV object
+	TRefCountPtr<FVulkanStructuredBuffer> SourceStructuredBuffer;
+
+	FVulkanUnorderedAccessView(FVulkanDevice* Device)
+		: VulkanRHI::FDeviceChild(Device)
+		, MipLevel(0)
+		, BufferViewFormat(PF_Unknown)
+		, VolatileLockCounter(MAX_uint32)
+	{
+	}
+
+	~FVulkanUnorderedAccessView();
+
+	void UpdateView();
+
+	// The texture that this UAV come from
+	TRefCountPtr<FRHITexture> SourceTexture;
+	FVulkanTextureView TextureView;
+	uint32 MipLevel;
+
+	// The vertex buffer this UAV comes from (can be null)
+	TRefCountPtr<FVulkanVertexBuffer> SourceVertexBuffer;
+	TRefCountPtr<FVulkanIndexBuffer> SourceIndexBuffer;
+	TRefCountPtr<FVulkanBufferView> BufferView;
+	EPixelFormat BufferViewFormat;
+
+protected:
+	// Used to check on volatile buffers if a new BufferView is required
+	uint32 VolatileLockCounter;
+};
+
+
+class FVulkanShaderResourceView : public FRHIShaderResourceView, public VulkanRHI::FDeviceChild
+{
+public:
+	FVulkanShaderResourceView(FVulkanDevice* Device, FRHIResource* InRHIBuffer, FVulkanResourceMultiBuffer* InSourceBuffer, uint32 InSize, EPixelFormat InFormat);
+
+	FVulkanShaderResourceView(FVulkanDevice* Device, FRHITexture* InSourceTexture, uint32 InMipLevel, int32 InNumMips, EPixelFormat InFormat)
+		: VulkanRHI::FDeviceChild(Device)
+		, BufferViewFormat(InFormat)
+		, SourceTexture(InSourceTexture)
+		, SourceStructuredBuffer(nullptr)
+		, MipLevel(InMipLevel)
+		, NumMips(InNumMips)
+		, Size(0)
+		, SourceBuffer(nullptr)
+		, VolatileLockCounter(MAX_uint32)
+	{
+	}
+
+	FVulkanShaderResourceView(FVulkanDevice* Device, FVulkanStructuredBuffer* InStructuredBuffer)
+		: VulkanRHI::FDeviceChild(Device)
+		, BufferViewFormat(PF_Unknown)
+		, SourceTexture(nullptr)
+		, SourceStructuredBuffer(InStructuredBuffer)
+		, MipLevel(0)
+		, NumMips(0)
+		, Size(InStructuredBuffer->GetSize())
+		, SourceBuffer(nullptr)
+		, VolatileLockCounter(MAX_uint32)
+	{
+	}
+
+	void Clear();
+
+	void Rename(FRHIResource* InRHIBuffer, FVulkanResourceMultiBuffer* InSourceBuffer, uint32 InSize, EPixelFormat InFormat);
+
+	void UpdateView();
+
+	inline FVulkanBufferView* GetBufferView()
+	{
+		return BufferViews[BufferIndex];
+	}
+
+	EPixelFormat BufferViewFormat;
+
+	// The texture that this SRV come from
+	TRefCountPtr<FRHITexture> SourceTexture;
+	FVulkanTextureView TextureView;
+	FVulkanStructuredBuffer* SourceStructuredBuffer;
+	uint32 MipLevel;
+	uint32 NumMips;
+
+	~FVulkanShaderResourceView();
+
+	TArray<TRefCountPtr<FVulkanBufferView>> BufferViews;
+	uint32 BufferIndex = 0;
+	uint32 Size;
+	// The buffer this SRV comes from (can be null)
+	FVulkanResourceMultiBuffer* SourceBuffer;
+	// To keep a reference
+	TRefCountPtr<FRHIResource> SourceRHIBuffer;
+
+protected:
+	// Used to check on volatile buffers if a new BufferView is required
+	uint32 VolatileLockCounter;
+};
+
+class FVulkanComputeFence : public FRHIComputeFence, public VulkanRHI::FGPUEvent
+{
+	bool bWriteEvent = false;
+
+public:
+	FVulkanComputeFence(FVulkanDevice* InDevice, FName InName);
+	virtual ~FVulkanComputeFence();
+
+	void WriteCmd(VkCommandBuffer CmdBuffer, bool bInWriteEvent);
+	void WriteWaitEvent(VkCommandBuffer CmdBuffer);
+};
+
+
+class FVulkanVertexInputStateInfo
+{
+public:
+	FVulkanVertexInputStateInfo();
+
+	void Generate(FVulkanVertexDeclaration* VertexDeclaration, uint32 VertexHeaderInOutAttributeMask);
+
+	inline uint32 GetHash() const
+	{
+		check(Info.sType == VK_STRUCTURE_TYPE_PIPELINE_VERTEX_INPUT_STATE_CREATE_INFO);
+		return Hash;
+	}
+
+	inline const VkPipelineVertexInputStateCreateInfo& GetInfo() const
+	{
+		return Info;
+	}
+
+protected:
+	VkPipelineVertexInputStateCreateInfo Info;
+	uint32 Hash;
+
+	uint32 BindingsNum;
+	uint32 BindingsMask;
+
+	//#todo-rco: Remove these TMaps
+	TMap<uint32, uint32> BindingToStream;
+	TMap<uint32, uint32> StreamToBinding;
+	VkVertexInputBindingDescription Bindings[MaxVertexElementCount];
+
+	uint32 AttributesNum;
+	VkVertexInputAttributeDescription Attributes[MaxVertexElementCount];
+
+	friend class FVulkanPendingGfxState;
+	friend class FVulkanPipelineStateCacheManager;
+};
+
+// This class holds the staging area for packed global uniform buffers for a given shader
+class FPackedUniformBuffers
+{
+public:
+	// One buffer is a chunk of bytes
+	typedef TArray<uint8> FPackedBuffer;
+
+	void Init(const FVulkanShaderHeader& InCodeHeader, uint64& OutPackedUniformBufferStagingMask)
+	{
+		PackedUniformBuffers.AddDefaulted(InCodeHeader.PackedUBs.Num());
+		for (int32 Index = 0; Index < InCodeHeader.PackedUBs.Num(); ++Index)
+		{
+			PackedUniformBuffers[Index].AddUninitialized(InCodeHeader.PackedUBs[Index].SizeInBytes);
+		}
+
+		OutPackedUniformBufferStagingMask = ((uint64)1 << (uint64)InCodeHeader.PackedUBs.Num()) - 1;
+		EmulatedUBsCopyInfo = InCodeHeader.EmulatedUBsCopyInfo;
+		EmulatedUBsCopyRanges = InCodeHeader.EmulatedUBCopyRanges;
+	}
+
+	inline void SetPackedGlobalParameter(uint32 BufferIndex, uint32 ByteOffset, uint32 NumBytes, const void* RESTRICT NewValue, uint64& InOutPackedUniformBufferStagingDirty)
+	{
+		FPackedBuffer& StagingBuffer = PackedUniformBuffers[BufferIndex];
+		check(ByteOffset + NumBytes <= (uint32)StagingBuffer.Num());
+		check((NumBytes & 3) == 0 && (ByteOffset & 3) == 0);
+		uint32* RESTRICT RawDst = (uint32*)(StagingBuffer.GetData() + ByteOffset);
+		uint32* RESTRICT RawSrc = (uint32*)NewValue;
+		uint32* RESTRICT RawSrcEnd = RawSrc + (NumBytes >> 2);
+
+		bool bChanged = false;
+		do
+		{
+			bChanged |= CopyAndReturnNotEqual(*RawDst++, *RawSrc++);
+		}
+		while (RawSrc != RawSrcEnd);
+
+		InOutPackedUniformBufferStagingDirty = InOutPackedUniformBufferStagingDirty | ((uint64)(bChanged ? 1 : 0) << (uint64)BufferIndex);
+	}
+
+	// Copies a 'real' constant buffer into the packed globals uniform buffer (only the used ranges)
+	inline void SetEmulatedUniformBufferIntoPacked(uint32 BindPoint, const TArray<uint8>& ConstantData, uint64& NEWPackedUniformBufferStagingDirty)
+	{
+		// Emulated UBs. Assumes UniformBuffersCopyInfo table is sorted by CopyInfo.SourceUBIndex
+		if (BindPoint < (uint32)EmulatedUBsCopyRanges.Num())
+		{
+			uint32 Range = EmulatedUBsCopyRanges[BindPoint];
+			uint16 Start = (Range >> 16) & 0xffff;
+			uint16 Count = Range & 0xffff;
+			const uint8* RESTRICT SourceData = ConstantData.GetData();
+			for (int32 Index = Start; Index < Start + Count; ++Index)
+			{
+				const CrossCompiler::FUniformBufferCopyInfo& CopyInfo = EmulatedUBsCopyInfo[Index];
+				check(CopyInfo.SourceUBIndex == BindPoint);
+				FPackedBuffer& StagingBuffer = PackedUniformBuffers[(int32)CopyInfo.DestUBIndex];
+				//check(ByteOffset + NumBytes <= (uint32)StagingBuffer.Num());
+				bool bChanged = false;
+				uint32* RESTRICT RawDst = (uint32*)(StagingBuffer.GetData() + CopyInfo.DestOffsetInFloats * 4);
+				uint32* RESTRICT RawSrc = (uint32*)(SourceData + CopyInfo.SourceOffsetInFloats * 4);
+				uint32* RESTRICT RawSrcEnd = RawSrc + CopyInfo.SizeInFloats;
+				do
+				{
+					bChanged |= CopyAndReturnNotEqual(*RawDst++, *RawSrc++);
+				}
+				while (RawSrc != RawSrcEnd);
+				NEWPackedUniformBufferStagingDirty = NEWPackedUniformBufferStagingDirty | ((uint64)(bChanged ? 1 : 0) << (uint64)CopyInfo.DestUBIndex);
+			}
+		}
+	}
+
+	inline const FPackedBuffer& GetBuffer(int32 Index) const
+	{
+		return PackedUniformBuffers[Index];
+	}
+
+protected:
+	TArray<FPackedBuffer>									PackedUniformBuffers;
+
+	// Copies to Shader Code Header (shaders may be deleted when we use this object again)
+	TArray<CrossCompiler::FUniformBufferCopyInfo>			EmulatedUBsCopyInfo;
+	TArray<uint32>											EmulatedUBsCopyRanges;
+};
+
+class FVulkanStagingBuffer : public FRHIStagingBuffer
+{
+	friend class FVulkanCommandListContext;
+public:
+	FVulkanStagingBuffer()
+		: FRHIStagingBuffer()
+	{
+		check(!bIsLocked);
+	}
+
+	virtual ~FVulkanStagingBuffer();
+
+	virtual void* Lock(uint32 Offset, uint32 NumBytes) final override;
+	virtual void Unlock() final override;
+
+private:
+	VulkanRHI::FStagingBuffer* StagingBuffer = nullptr;
+	uint32 QueuedOffset = 0;
+	uint32 QueuedNumBytes = 0;
+	// The staging buffer was allocated from this device.
+	FVulkanDevice* Device;
+};
+
+class FVulkanGPUFence : public FRHIGPUFence
+{
+public:
+	FVulkanGPUFence(FName InName)
+		: FRHIGPUFence(InName)
+	{
+	}
+
+	virtual void Clear() final override;
+	virtual bool Poll() const final override;
+
+protected:
+	FVulkanCmdBuffer*	CmdBuffer = nullptr;
+	uint64				FenceSignaledCounter = 0;
+
+	friend class FVulkanCommandListContext;
+};
+
+template<class T>
+struct TVulkanResourceTraits
+{
+};
+template<>
+struct TVulkanResourceTraits<FRHIVertexDeclaration>
+{
+	typedef FVulkanVertexDeclaration TConcreteType;
+};
+template<>
+struct TVulkanResourceTraits<FRHIVertexShader>
+{
+	typedef FVulkanVertexShader TConcreteType;
+};
+template<>
+struct TVulkanResourceTraits<FRHIGeometryShader>
+{
+	typedef FVulkanGeometryShader TConcreteType;
+};
+template<>
+struct TVulkanResourceTraits<FRHIHullShader>
+{
+	typedef FVulkanHullShader TConcreteType;
+};
+template<>
+struct TVulkanResourceTraits<FRHIDomainShader>
+{
+	typedef FVulkanDomainShader TConcreteType;
+};
+template<>
+struct TVulkanResourceTraits<FRHIPixelShader>
+{
+	typedef FVulkanPixelShader TConcreteType;
+};
+template<>
+struct TVulkanResourceTraits<FRHIComputeShader>
+{
+	typedef FVulkanComputeShader TConcreteType;
+};
+template<>
+struct TVulkanResourceTraits<FRHITexture3D>
+{
+	typedef FVulkanTexture3D TConcreteType;
+};
+//template<>
+//struct TVulkanResourceTraits<FRHITexture>
+//{
+//	typedef FVulkanTexture TConcreteType;
+//};
+template<>
+struct TVulkanResourceTraits<FRHITexture2D>
+{
+	typedef FVulkanTexture2D TConcreteType;
+};
+template<>
+struct TVulkanResourceTraits<FRHITexture2DArray>
+{
+	typedef FVulkanTexture2DArray TConcreteType;
+};
+template<>
+struct TVulkanResourceTraits<FRHITextureCube>
+{
+	typedef FVulkanTextureCube TConcreteType;
+};
+template<>
+struct TVulkanResourceTraits<FRHIRenderQuery>
+{
+	typedef FVulkanRenderQuery TConcreteType;
+};
+template<>
+struct TVulkanResourceTraits<FRHIUniformBuffer>
+{
+	typedef FVulkanUniformBuffer TConcreteType;
+};
+template<>
+struct TVulkanResourceTraits<FRHIIndexBuffer>
+{
+	typedef FVulkanIndexBuffer TConcreteType;
+};
+template<>
+struct TVulkanResourceTraits<FRHIStructuredBuffer>
+{
+	typedef FVulkanStructuredBuffer TConcreteType;
+};
+template<>
+struct TVulkanResourceTraits<FRHIVertexBuffer>
+{
+	typedef FVulkanVertexBuffer TConcreteType;
+};
+template<>
+struct TVulkanResourceTraits<FRHIShaderResourceView>
+{
+	typedef FVulkanShaderResourceView TConcreteType;
+};
+template<>
+struct TVulkanResourceTraits<FRHIUnorderedAccessView>
+{
+	typedef FVulkanUnorderedAccessView TConcreteType;
+};
+
+template<>
+struct TVulkanResourceTraits<FRHISamplerState>
+{
+	typedef FVulkanSamplerState TConcreteType;
+};
+template<>
+struct TVulkanResourceTraits<FRHIRasterizerState>
+{
+	typedef FVulkanRasterizerState TConcreteType;
+};
+template<>
+struct TVulkanResourceTraits<FRHIDepthStencilState>
+{
+	typedef FVulkanDepthStencilState TConcreteType;
+};
+template<>
+struct TVulkanResourceTraits<FRHIBlendState>
+{
+	typedef FVulkanBlendState TConcreteType;
+};
+
+template<>
+struct TVulkanResourceTraits<FRHIComputeFence>
+{
+	typedef FVulkanComputeFence TConcreteType;
+};
+
+template<>
+struct TVulkanResourceTraits<FRHIBoundShaderState>
+{
+	typedef FVulkanBoundShaderState TConcreteType;
+};
+
+template<>
+struct TVulkanResourceTraits<FRHIStagingBuffer>
+{
+	typedef FVulkanStagingBuffer TConcreteType;
+};
+
+template<>
+struct TVulkanResourceTraits<FRHIGPUFence>
+{
+	typedef FVulkanGPUFence TConcreteType;
+};
+
+template<typename TRHIType>
+static FORCEINLINE typename TVulkanResourceTraits<TRHIType>::TConcreteType* ResourceCast(TRHIType* Resource)
+{
+	return static_cast<typename TVulkanResourceTraits<TRHIType>::TConcreteType*>(Resource);
+}
+
+template<typename TRHIType>
+static FORCEINLINE typename TVulkanResourceTraits<TRHIType>::TConcreteType* ResourceCast(const TRHIType* Resource)
+{
+	return static_cast<const typename TVulkanResourceTraits<TRHIType>::TConcreteType*>(Resource);
+}
diff --git a/Engine/Source/Runtime/VulkanRHI/Private/VulkanState.h b/Engine/Source/Runtime/VulkanRHI/Private/VulkanState.h
new file mode 100644
--- /dev/null
+++ b/Engine/Source/Runtime/VulkanRHI/Private/VulkanState.h
@@ -0,0 +1,86 @@
+// Copyright 1998-2019 Epic Games, Inc. All Rights Reserved.
+
+/*=============================================================================
+	VulkanState.h: Vulkan state definitions.
+=============================================================================*/
+
+#pragma once
+
+#include "CoreMinimal.h"
+#include "RHIResources.h"
+
+class FVulkanDevice;
+
+class FVulkanSamplerState : public FRHISamplerState
+{
+public:
+	FVulkanSamplerState(const VkSamplerCreateInfo& InInfo, FVulkanDevice& InDevice, const bool bInIsImmutable = false);
+	
+	virtual bool IsImmutable() const final override { return bIsImmutable; }
+
+	VkSampler Sampler;
+	uint32 SamplerId;
+
+	static void SetupSamplerCreateInfo(const FSamplerStateInitializerRHI& Initializer, FVulkanDevice& InDevice, VkSamplerCreateInfo& OutSamplerInfo);
+
+private:
+	bool bIsImmutable;
+};
+
+class FVulkanRasterizerState : public FRHIRasterizerState
+{
+public:
+	FVulkanRasterizerState(const FRasterizerStateInitializerRHI& InInitializer);
+
+	static void ResetCreateInfo(VkPipelineRasterizationStateCreateInfo& OutInfo)
+	{
+		ZeroVulkanStruct(OutInfo, VK_STRUCTURE_TYPE_PIPELINE_RASTERIZATION_STATE_CREATE_INFO);
+		OutInfo.frontFace = VK_FRONT_FACE_CLOCKWISE;
+		OutInfo.lineWidth = 1.0f;
+	}
+
+	virtual bool GetInitializer(FRasterizerStateInitializerRHI& Out) override final
+	{
+		Out = Initializer;
+		return true;
+	}
+
+	VkPipelineRasterizationStateCreateInfo	RasterizerState;
+	FRasterizerStateInitializerRHI			Initializer;
+};
+
+class FVulkanDepthStencilState : public FRHIDepthStencilState
+{
+public:
+	FVulkanDepthStencilState(const FDepthStencilStateInitializerRHI& InInitializer)
+	{
+		Initializer = InInitializer;
+	}
+
+	virtual bool GetInitializer(FDepthStencilStateInitializerRHI& Out) override final
+	{
+		Out = Initializer;
+		return true;
+	}
+
+	void SetupCreateInfo(const FGraphicsPipelineStateInitializer& GfxPSOInit, VkPipelineDepthStencilStateCreateInfo& OutDepthStencilState);
+
+	FDepthStencilStateInitializerRHI Initializer;
+};
+
+class FVulkanBlendState : public FRHIBlendState
+{
+public:
+	FVulkanBlendState(const FBlendStateInitializerRHI& InInitializer);
+
+	virtual bool GetInitializer(FBlendStateInitializerRHI& Out) override final
+	{
+		Out = Initializer;
+		return true;
+	}
+
+	// array the pipeline state can point right to
+	VkPipelineColorBlendAttachmentState BlendStates[MaxSimultaneousRenderTargets];
+
+	FBlendStateInitializerRHI Initializer;
+};
diff --git a/Engine/Source/Runtime/VulkanRHI/Private/VulkanUtil.cpp b/Engine/Source/Runtime/VulkanRHI/Private/VulkanUtil.cpp
--- a/Engine/Source/Runtime/VulkanRHI/Private/VulkanUtil.cpp
+++ b/Engine/Source/Runtime/VulkanRHI/Private/VulkanUtil.cpp
@@ -581,12 +581,6 @@ namespace VulkanRHIBridge
 }
 
 
-#include "ShaderParameterUtils.h"
-#include "RHIStaticStates.h"
-#include "OneColorShader.h"
-#include "VulkanRHI.h"
-#include "StaticBoundShaderState.h"
-
 namespace VulkanRHI
 {
 	VkBuffer CreateBuffer(FVulkanDevice* InDevice, VkDeviceSize Size, VkBufferUsageFlags BufferUsageFlags, VkMemoryRequirements& OutMemoryRequirements)
diff --git a/Engine/Source/Runtime/VulkanRHI/Private/VulkanViewport.h b/Engine/Source/Runtime/VulkanRHI/Private/VulkanViewport.h
new file mode 100644
--- /dev/null
+++ b/Engine/Source/Runtime/VulkanRHI/Private/VulkanViewport.h
@@ -0,0 +1,119 @@
+// Copyright 1998-2019 Epic Games, Inc. All Rights Reserved.
+
+/*=============================================================================
+	VulkanViewport.h: Vulkan viewport RHI definitions.
+=============================================================================*/
+
+#pragma once
+
+#include "VulkanResources.h"
+#include "HAL/CriticalSection.h"
+
+class FVulkanDynamicRHI;
+class FVulkanSwapChain;
+class FVulkanQueue;
+
+namespace VulkanRHI
+{
+	class FSemaphore;
+}
+
+class FVulkanViewport : public FRHIViewport, public VulkanRHI::FDeviceChild
+{
+public:
+	enum { NUM_BUFFERS = 3 };
+
+	FVulkanViewport(FVulkanDynamicRHI* InRHI, FVulkanDevice* InDevice, void* InWindowHandle, uint32 InSizeX,uint32 InSizeY,bool bInIsFullscreen, EPixelFormat InPreferredPixelFormat);
+	~FVulkanViewport();
+
+	FTexture2DRHIRef GetBackBuffer(FRHICommandList& RHICmdList);
+
+	void WaitForFrameEventCompletion();
+
+	void IssueFrameEvent();
+
+	inline FIntPoint GetSizeXY() const
+	{
+		return FIntPoint(SizeX, SizeY);
+	}
+
+	virtual void SetCustomPresent(FRHICustomPresent* InCustomPresent) override final
+	{
+		CustomPresent = InCustomPresent;
+	}
+
+	virtual FRHICustomPresent* GetCustomPresent() const override final
+	{
+		return CustomPresent;
+	}
+
+	virtual void Tick(float DeltaTime) override final;
+
+	void AdvanceBackBufferFrame();
+
+	bool Present(FVulkanCommandListContext* Context, FVulkanCmdBuffer* CmdBuffer, FVulkanQueue* Queue, FVulkanQueue* PresentQueue, bool bLockToVsync);
+
+	inline uint32 GetPresentCount() const
+	{
+		return PresentCount;
+	}
+
+protected:
+	VkImage BackBufferImages[NUM_BUFFERS];
+	VulkanRHI::FSemaphore* RenderingDoneSemaphores[NUM_BUFFERS];
+	FVulkanTextureView TextureViews[NUM_BUFFERS];
+	
+	TRefCountPtr<FVulkanBackBuffer> BackBuffers[NUM_BUFFERS];
+	TRefCountPtr<FVulkanBackBufferReference> RenderingBackBufferReference;
+
+	// 'Dummy' back buffer
+	TRefCountPtr<FVulkanBackBuffer> RenderingBackBuffer;
+	TRefCountPtr<FVulkanBackBuffer> RHIBackBuffer;
+
+	/** narrow-scoped section that locks access to back buffer during its recreation*/
+	FCriticalSection RecreatingSwapchain;
+
+	FVulkanDynamicRHI* RHI;
+	uint32 SizeX;
+	uint32 SizeY;
+	bool bIsFullscreen;
+	EPixelFormat PixelFormat;
+	int32 AcquiredImageIndex;
+	int32 PreAcquiredImageIndex;
+	FVulkanSwapChain* SwapChain;
+	void* WindowHandle;
+	uint32 PresentCount;
+
+	int8 LockToVsync;
+
+	// Just a pointer, not owned by this class
+	VulkanRHI::FSemaphore* AcquiredSemaphore;
+
+	FCustomPresentRHIRef CustomPresent;
+
+	FVulkanCmdBuffer* LastFrameCommandBuffer = nullptr;
+	uint64 LastFrameFenceCounter = 0;
+
+	void CreateSwapchain();
+	void PreAcquireSwapchainImage();
+	void GetNextImageIndex();
+	void AcquireImageIndex();
+	void AcquireBackBuffer(FRHICommandListBase& CmdList, FVulkanBackBufferReference* NewBackBuffer);
+
+	void RecreateSwapchain(void* NewNativeWindow, bool bForce = false);
+	void RecreateSwapchainFromRT(EPixelFormat PreferredPixelFormat);
+	void Resize(uint32 InSizeX, uint32 InSizeY, bool bIsFullscreen, EPixelFormat PreferredPixelFormat);
+
+	static int32 DoAcquireImageIndex(FVulkanViewport* Viewport);
+	bool DoCheckedSwapChainJob(TFunction<int32(FVulkanViewport*)> SwapChainJob);
+
+	friend class FVulkanDynamicRHI;
+	friend class FVulkanCommandListContext;
+	friend struct FRHICommandAcquireBackBuffer;
+};
+
+template<>
+struct TVulkanResourceTraits<FRHIViewport>
+{
+	typedef FVulkanViewport TConcreteType;
+};
diff --git a/Engine/Source/Runtime/VulkanRHI/Public/VulkanContext.h b/Engine/Source/Runtime/VulkanRHI/Public/VulkanContext.h
deleted file mode 100644
diff --git a/Engine/Source/Runtime/VulkanRHI/Public/VulkanLLM.h b/Engine/Source/Runtime/VulkanRHI/Public/VulkanLLM.h
deleted file mode 100644
diff --git a/Engine/Source/Runtime/VulkanRHI/Public/VulkanRHI.h b/Engine/Source/Runtime/VulkanRHI/Public/VulkanRHI.h
deleted file mode 100644
diff --git a/Engine/Source/Runtime/VulkanRHI/Public/VulkanResources.h b/Engine/Source/Runtime/VulkanRHI/Public/VulkanResources.h
deleted file mode 100644
diff --git a/Engine/Source/Runtime/VulkanRHI/Public/VulkanState.h b/Engine/Source/Runtime/VulkanRHI/Public/VulkanState.h
deleted file mode 100644
diff --git a/Engine/Source/Runtime/VulkanRHI/Public/VulkanViewport.h b/Engine/Source/Runtime/VulkanRHI/Public/VulkanViewport.h
deleted file mode 100644
