Linux: remove Vulkan workaround (UE-67825).

- The underlying issue is fixed by not allowing the CPU to run further than one frame ahead.

diff --git a/Engine/Source/Runtime/VulkanRHI/Private/Linux/VulkanLinuxPlatform.cpp b/Engine/Source/Runtime/VulkanRHI/Private/Linux/VulkanLinuxPlatform.cpp
--- a/Engine/Source/Runtime/VulkanRHI/Private/Linux/VulkanLinuxPlatform.cpp
+++ b/Engine/Source/Runtime/VulkanRHI/Private/Linux/VulkanLinuxPlatform.cpp
@@ -146,13 +146,6 @@ void FVulkanLinuxPlatform::GetDeviceExtensions(TArray<const ANSICHAR*>& OutExten
 	OutExtensions.Add(VK_KHR_GET_MEMORY_REQUIREMENTS_2_EXTENSION_NAME);
 	OutExtensions.Add(VK_KHR_DEDICATED_ALLOCATION_EXTENSION_NAME);
 #endif
-
-	if (IsRHIDeviceNVIDIA())
-	{
-		//#todo-rco: Temporary workaround for some buffers not updating
-		extern FAutoConsoleVariableRef CVarVulkanWaitForIdleOnSubmit;
-		CVarVulkanWaitForIdleOnSubmit->Set(1);
-	}
 }
 
 void FVulkanLinuxPlatform::CreateSurface(void* WindowHandle, VkInstance Instance, VkSurfaceKHR* OutSurface)
