Remove unused marker names & stack usage on Vulkan

diff --git a/Engine/Source/Runtime/VulkanRHI/Private/VulkanRHI.cpp b/Engine/Source/Runtime/VulkanRHI/Private/VulkanRHI.cpp
--- a/Engine/Source/Runtime/VulkanRHI/Private/VulkanRHI.cpp
+++ b/Engine/Source/Runtime/VulkanRHI/Private/VulkanRHI.cpp
@@ -843,9 +843,6 @@ void FVulkanCommandListContext::RHIEndFrame()
 
 void FVulkanCommandListContext::RHIPushEvent(const TCHAR* Name, FColor Color)
 {
-	FString EventName = Name;
-	EventStack.Add(Name);
-
 #if VULKAN_ENABLE_DRAW_MARKERS
 #if 0//VULKAN_SUPPORTS_DEBUG_UTILS
 	if (auto CmdBeginLabel = Device->GetCmdBeginDebugLabel())
@@ -928,9 +925,6 @@ void FVulkanCommandListContext::RHIPopEvent()
 
 		GpuProfiler.PopEvent();
 	}
-
-	check(EventStack.Num() > 0);
-	EventStack.Pop(false);
 }
 
 void FVulkanDynamicRHI::RHIGetSupportedResolution( uint32 &Width, uint32 &Height )
