Vulkan android frame pacing improvements (contributed by Samsung)

diff --git a/Engine/Source/Runtime/VulkanRHI/Private/VulkanSwapChain.cpp b/Engine/Source/Runtime/VulkanRHI/Private/VulkanSwapChain.cpp
--- a/Engine/Source/Runtime/VulkanRHI/Private/VulkanSwapChain.cpp
+++ b/Engine/Source/Runtime/VulkanRHI/Private/VulkanSwapChain.cpp
@@ -456,14 +456,11 @@ FVulkanSwapChain::FVulkanSwapChain(VkInstance InInstance, FVulkanDevice& InDevic
 	}
 
 #if VULKAN_SUPPORTS_GOOGLE_DISPLAY_TIMING
-	if (Device.GetOptionalExtensions().HasGoogleDisplayTiming)
+	if (GVulkanExtensionFramePacer && Device.GetOptionalExtensions().HasGoogleDisplayTiming)
 	{
 		GDTimingFramePacer = MakeUnique<FGDTimingFramePacer>(Device, SwapChain);
-		if (GVulkanExtensionFramePacer)
-		{
-			GVulkanCPURenderThreadFramePacer = 0;
-			GVulkanCPURHIFramePacer = 0;
-		}
+		GVulkanCPURenderThreadFramePacer = 0;
+		GVulkanCPURHIFramePacer = 0;
 	}
 #endif
 
@@ -503,6 +500,8 @@ int32 FVulkanSwapChain::AcquireImageIndex(VulkanRHI::FSemaphore** OutSemaphore)
 {
 	check(FVulkanPlatform::SupportsStandardSwapchain());
 
+	check(CurrentImageIndex == -1);
+
 	// Get the index of the next swapchain image we should render to.
 	// We'll wait with an "infinite" timeout, the function will block until an image is ready.
 	// The ImageAcquiredSemaphore[ImageAcquiredSemaphoreIndex] will get signaled when the image is ready (upon function return).
@@ -601,7 +600,6 @@ FGDTimingFramePacer::FGDTimingFramePacer(FVulkanDevice& InDevice, VkSwapchainKHR
 	}
 	HalfRefreshDuration = (RefreshDuration / 2);
 
-	FMemory::Memzero(CpuPresentTimeHistory);
 	FMemory::Memzero(PresentTime); 
 
 	ZeroVulkanStruct(PresentTimesInfo, VK_STRUCTURE_TYPE_PRESENT_TIMES_INFO_GOOGLE);
@@ -609,14 +607,17 @@ FGDTimingFramePacer::FGDTimingFramePacer(FVulkanDevice& InDevice, VkSwapchainKHR
 	PresentTimesInfo.pTimes = &PresentTime;
 }
 
+// Used as a safety measure to prevent scheduling too far ahead in case of an error
+static constexpr uint64 GMaxAheadSchedulingTimeNanosec = 500000000llu; // 0.5 sec.
+
 static uint64 TimeNanoseconds()
 {
-#if PLATFORM_ANDROID
+#if PLATFORM_ANDROID || PLATFORM_LINUX
 	struct timespec ts;
 	clock_gettime(CLOCK_MONOTONIC, &ts);
 	return ts.tv_sec*1000000000ull + ts.tv_nsec;
 #else
-	return (uint64)(FPlatformTime::Seconds()*1000000000.0);
+	#error VK_GOOGLE_display_timing requires TimeNanoseconds() implementation for this platform
 #endif
 }
 
@@ -628,36 +629,35 @@ void FGDTimingFramePacer::ScheduleNextFrame(uint32 InPresentID, int32 InSyncInte
 		return;
 	}
 
+	const uint64 CpuPresentTime = TimeNanoseconds();
+
+	PresentTime.presentID = InPresentID; // Still need to pass ID for proper history values
+
 	PollPastFrameInfo();
 	if (!LastKnownFrameInfo.bValid)
 	{
-		LastScheduledPresentTime = 0;
 		return;
 	}
 
-	const uint64 CpuPresentTime = TimeNanoseconds();
-	const int32 HistorySize = ARRAY_COUNT(CpuPresentTimeHistory);
-	const int32 HistoryIndex = InPresentID % HistorySize;
-	CpuPresentTimeHistory[HistoryIndex] = CpuPresentTime;
+	const uint64 CpuTargetPresentTimeMin = CalculateMinPresentTime(CpuPresentTime);
+	const uint64 CpuTargetPresentTimeMax = CalculateMaxPresentTime(CpuPresentTime);
+	const uint64 GpuTargetPresentTime = (PredictLastScheduledFramePresentTime(InPresentID) + SyncDuration);
 	
-	uint64 CpuTargetPresentTime = CalculateNearestPresentTime(CpuPresentTime);
-	uint64 GpuTargetPresentTime = CalculateNearestVsTime(LastKnownFrameInfo.ActualPresentTime, PredictLastScheduledFramePresentTime(InPresentID) + SyncDuration);
-	
-	uint64 TargetPresentTime = FMath::Max(CpuTargetPresentTime, GpuTargetPresentTime);
+	const uint64 TargetPresentTime = CalculateNearestVsTime(LastKnownFrameInfo.ActualPresentTime, FMath::Clamp(GpuTargetPresentTime, CpuTargetPresentTimeMin, CpuTargetPresentTimeMax));
 	LastScheduledPresentTime = TargetPresentTime;
 
-	PresentTime.presentID = InPresentID;
 	PresentTime.desiredPresentTime = (TargetPresentTime - HalfRefreshDuration);
 
 	if (GPrintVulkanVsyncDebug != 0)
 	{
-		double cpuP = CpuTargetPresentTime/1000000000.0;
+		double cpuPMin = CpuTargetPresentTimeMin/1000000000.0;
+		double cpuPMax = CpuTargetPresentTimeMax/1000000000.0;
 		double gpuP = GpuTargetPresentTime/1000000000.0;
 		double desP = PresentTime.desiredPresentTime/1000000000.0;
 		double lastP = LastKnownFrameInfo.ActualPresentTime/1000000000.0;
-		double cpuDelta = CpuToGpuPresentDelta/1000000000.0;
+		double cpuDelta = 0.0;
 		double cpuNow = CpuPresentTime/1000000000.0;
-		FPlatformMisc::LowLevelOutputDebugStringf(TEXT(" -- ID: %u, desired %.3f, pred-gpu %.3f, pred-cpu %.3f, last: %.3f, cpu-gpu-delta: %.3f, now-cpu %.3f"), PresentTime.presentID, desP, gpuP, cpuP, lastP, cpuDelta, cpuNow);
+		FPlatformMisc::LowLevelOutputDebugStringf(TEXT(" -- ID: %u, desired %.3f, pred-gpu %.3f, pred-cpu-min %.3f, pred-cpu-max %.3f, last: %.3f, cpu-gpu-delta: %.3f, now-cpu %.3f"), PresentTime.presentID, desP, gpuP, cpuPMin, cpuPMax, lastP, cpuDelta, cpuNow);
 	}
 }
 
@@ -667,14 +667,10 @@ void FGDTimingFramePacer::UpdateSyncDuration(int32 InSyncInterval)
 	{
 		return;
 	}
-	SyncInterval = InSyncInterval;
 
-	// reset cached history on sync interval changes
-	FMemory::Memzero(CpuPresentTimeHistory);
-	LastKnownFrameInfo.bValid = false;
-	LastScheduledPresentTime = 0;
-	
+	SyncInterval = InSyncInterval;
 	SyncDuration = ((1000000000llu * FMath::Clamp(SyncInterval, 0, 3) + 30) / 60);
+	
 	if (SyncDuration > 0)
 	{
 		SyncDuration = (FMath::Max((SyncDuration + HalfRefreshDuration) / RefreshDuration, 1llu) * RefreshDuration);
@@ -684,13 +680,25 @@ void FGDTimingFramePacer::UpdateSyncDuration(int32 InSyncInterval)
 uint64 FGDTimingFramePacer::PredictLastScheduledFramePresentTime(uint32 CurrentPresentID) const
 {
 	const uint32 PredictFrameCount = (CurrentPresentID - LastKnownFrameInfo.PresentID - 1);
-	return FMath::Max(LastScheduledPresentTime, LastKnownFrameInfo.ActualPresentTime + (SyncDuration * PredictFrameCount));
+	// Use RefreshDuration for predicted frames and not SyncDuration for most optimistic prediction of future frames after last known (possible hitchy) frame.
+	// Second parameter will be always >= than LastScheduledPresentTime if use SyncDuration.
+	// It is possible that GPU will recover after hitch without any changes to a normal schedule but pessimistic planning will prevent this from happening.
+	return FMath::Max(LastScheduledPresentTime, LastKnownFrameInfo.ActualPresentTime + (RefreshDuration * PredictFrameCount));
 }
 
-uint64 FGDTimingFramePacer::CalculateNearestPresentTime(uint64 CpuPresentTime) const
+uint64 FGDTimingFramePacer::CalculateMinPresentTime(uint64 CpuPresentTime) const
 {
-	const uint64 NearestGpuPresentTime = CpuPresentTime + CpuToGpuPresentDelta;
-	return CalculateNearestVsTime(LastKnownFrameInfo.ActualPresentTime, NearestGpuPresentTime - HalfRefreshDuration);
+	// Do not use delta on Android because already using CLOCK_MONOTONIC for CPU time which is also used in the extension.
+	// Using delta will mostly work fine but there were problems in other projects. If GPU load changes quickly because
+	// of the delta filter lag its value may be too high for current frame and cause pessimistic planning and stuttering.
+	// Need additional time for testing to improve filtering.
+	// Adding HalfRefreshDuration to produce round-up (ceil) in the final CalculateNearestVsTime()
+	return (CpuPresentTime + HalfRefreshDuration);
+}
+
+uint64 FGDTimingFramePacer::CalculateMaxPresentTime(uint64 CpuPresentTime) const
+{
+	return (CpuPresentTime + GMaxAheadSchedulingTimeNanosec);
 }
 
 uint64 FGDTimingFramePacer::CalculateNearestVsTime(uint64 ActualPresentTime, uint64 TargetTime) const
@@ -721,37 +729,19 @@ void FGDTimingFramePacer::PollPastFrameInfo()
 		Result = VulkanDynamicAPI::vkGetPastPresentationTimingGOOGLE(Device.GetInstanceHandle(), SwapChain, &Count, &PastPresentationTiming);
 		checkf(Result == VK_SUCCESS || Result == VK_INCOMPLETE, TEXT("vkGetPastPresentationTimingGOOGLE failed: %i"), Result);
 
+		// If desiredPresentTime was too large for some reason driver may ignore this value to prevent long wait
+		// Reset LastScheduledPresentTime in that case to be able to schedule on proper time
+		if (PastPresentationTiming.actualPresentTime < PastPresentationTiming.desiredPresentTime)
+		{
+			UE_LOG(LogVulkanRHI, Warning, TEXT("PastPresentationTiming actualPresentTime is less than desiredPresentTime! Resetting LastScheduledPresentTime..."));
+			LastScheduledPresentTime = 0;
+		}
+
 		LastKnownFrameInfo.PresentID = PastPresentationTiming.presentID;
 		LastKnownFrameInfo.ActualPresentTime = PastPresentationTiming.actualPresentTime;
 		LastKnownFrameInfo.bValid = true;
-
-		UpdateCpuToGpuPresentDelta(PastPresentationTiming);
 	}
 }
-
-void FGDTimingFramePacer::UpdateCpuToGpuPresentDelta(const VkPastPresentationTimingGOOGLE& PastPresentationTiming)
-{
-	const int32 HistorySize = ARRAY_COUNT(CpuPresentTimeHistory);
-	if ((PresentTime.presentID - PastPresentationTiming.presentID) >= HistorySize)
-	{
-		return;
-	}
-		
-	const int32 HistoryIndex = PastPresentationTiming.presentID % HistorySize;
-	const uint64 PastCpuPresentTime = CpuPresentTimeHistory[HistoryIndex];
-	if (PastCpuPresentTime == 0)
-	{
-		CpuToGpuPresentDelta = SyncDuration;
-		return;
-	}
-
-	// "presentMargin" may be negative despite being unsigned
-	const uint64 Delta = PastPresentationTiming.earliestPresentTime - (PastCpuPresentTime + (int64_t&)PastPresentationTiming.presentMargin);
-	const uint64 FilterParam = (CpuToGpuPresentDelta == 0) ? 0 : 10; // greater -> smoother
-	CpuToGpuPresentDelta = (CpuToGpuPresentDelta * FilterParam + Delta) / (FilterParam + 1);
-	// filter out bad frames, in general delta should be 2-4 sync durations
-	CpuToGpuPresentDelta = FMath::Min(CpuToGpuPresentDelta, SyncDuration*4);
-}
 #endif //VULKAN_SUPPORTS_GOOGLE_DISPLAY_TIMING
 
 void FVulkanSwapChain::RenderThreadPacing()
@@ -808,11 +798,7 @@ FVulkanSwapChain::EStatus FVulkanSwapChain::Present(FVulkanQueue* GfxQueue, FVul
 {
 	check(FVulkanPlatform::SupportsStandardSwapchain());
 
-	if (CurrentImageIndex == -1)
-	{
-		// Skip present silently if image has not been acquired
-		return EStatus::Healthy;
-	}
+	check(CurrentImageIndex != -1);
 
 	//ensure(GfxQueue == PresentQueue);
 
@@ -897,6 +883,8 @@ FVulkanSwapChain::EStatus FVulkanSwapChain::Present(FVulkanQueue* GfxQueue, FVul
 			GRenderThreadNumIdle[ERenderThreadIdleTypes::WaitingForGPUPresent]++;
 		}
 
+		CurrentImageIndex = -1;
+
 #if !UE_BUILD_SHIPPING
 		PresentResult = SimulateErrors(PresentResult);
 #endif
diff --git a/Engine/Source/Runtime/VulkanRHI/Private/VulkanSwapChain.h b/Engine/Source/Runtime/VulkanRHI/Private/VulkanSwapChain.h
--- a/Engine/Source/Runtime/VulkanRHI/Private/VulkanSwapChain.h
+++ b/Engine/Source/Runtime/VulkanRHI/Private/VulkanSwapChain.h
@@ -89,11 +89,10 @@ private:
 	void UpdateSyncDuration(int32 SyncInterval);
 
 	uint64 PredictLastScheduledFramePresentTime(uint32 CurrentPresentID) const;
-	uint64 CalculateNearestPresentTime(uint64 CpuPresentTime) const;
+	uint64 CalculateMinPresentTime(uint64 CpuPresentTime) const;
+	uint64 CalculateMaxPresentTime(uint64 CpuPresentTime) const;
 	uint64 CalculateNearestVsTime(uint64 ActualPresentTime, uint64 TargetTime) const;
-
 	void PollPastFrameInfo();
-	void UpdateCpuToGpuPresentDelta(const VkPastPresentationTimingGOOGLE& PastPresentationTiming);
 
 private:
 	struct FKnownFrameInfo
@@ -116,8 +115,5 @@ private:
 	uint64 LastScheduledPresentTime = 0;
 	uint64 SyncDuration = 0;
 	int32 SyncInterval = 0;
-
-	uint64 CpuPresentTimeHistory[10];
-	uint64 CpuToGpuPresentDelta = 0;
 };
 #endif //VULKAN_SUPPORTS_GOOGLE_DISPLAY_TIMING
\ No newline at end of file
diff --git a/Engine/Source/Runtime/VulkanRHI/Private/VulkanViewport.cpp b/Engine/Source/Runtime/VulkanRHI/Private/VulkanViewport.cpp
--- a/Engine/Source/Runtime/VulkanRHI/Private/VulkanViewport.cpp
+++ b/Engine/Source/Runtime/VulkanRHI/Private/VulkanViewport.cpp
@@ -822,6 +822,12 @@ bool FVulkanViewport::Present(FVulkanCommandListContext* Context, FVulkanCmdBuff
 		// Present the back buffer to the viewport window.
 		auto SwapChainJob = [Queue, PresentQueue](FVulkanViewport* Viewport)
 		{
+			// May happend if swapchain was recreated in DoCheckedSwapChainJob()
+			if (Viewport->AcquiredImageIndex == -1)
+			{
+				// Skip present silently if image has not been acquired
+				return (int32)FVulkanSwapChain::EStatus::Healthy;
+			}
 			return (int32)Viewport->SwapChain->Present(Queue, PresentQueue, Viewport->RenderingDoneSemaphores[Viewport->AcquiredImageIndex]);
 		};
 		if (FVulkanPlatform::SupportsStandardSwapchain() && !DoCheckedSwapChainJob(SwapChainJob))
