Copy of CL 5906283
Fixed tessellation shader cross-compilation for Vulkan backend.

diff --git a/Engine/Shaders/Private/BasePassCommon.ush b/Engine/Shaders/Private/BasePassCommon.ush
--- a/Engine/Shaders/Private/BasePassCommon.ush
+++ b/Engine/Shaders/Private/BasePassCommon.ush
@@ -106,7 +106,7 @@ struct FSharedBasePassInterpolants
 #endif
 };
 
-#if PC_D3D || GL4_PROFILE || SWITCH_PROFILE || SWITCH_PROFILE_FORWARD || METAL_TESSELLATION_PROFILE || VULKAN_ENABLE_TESSELLATION
+#if PC_D3D || GL4_PROFILE || SWITCH_PROFILE || SWITCH_PROFILE_FORWARD || METAL_TESSELLATION_PROFILE || VULKAN_PROFILE_SM5
 
 /** Interpolants passed from the vertex shader to the pixel shader. */
 struct FBasePassInterpolantsVSToPS : FSharedBasePassInterpolants
diff --git a/Engine/Shaders/Public/Platform.ush b/Engine/Shaders/Public/Platform.ush
--- a/Engine/Shaders/Public/Platform.ush
+++ b/Engine/Shaders/Public/Platform.ush
@@ -131,10 +131,6 @@
 #define	VULKAN_PROFILE_SM5 0
 #endif
 
-#ifndef VULKAN_ENABLE_TESSELLATION
-#define VULKAN_ENABLE_TESSELLATION	0
-#endif
-
 #ifndef IOS
 #define IOS 0
 #endif
diff --git a/Engine/Shaders/Public/Platform/Vulkan/VulkanCommon.ush b/Engine/Shaders/Public/Platform/Vulkan/VulkanCommon.ush
--- a/Engine/Shaders/Public/Platform/Vulkan/VulkanCommon.ush
+++ b/Engine/Shaders/Public/Platform/Vulkan/VulkanCommon.ush
@@ -5,7 +5,7 @@
 =============================================================================*/
 
 // Update this GUID to improve shader recompilation for Vulkan only shaders
-// GUID = {792BDC0A-A4A8-402F-8752-000000000042}
+// GUID = {98AC8FD8-9750-44D2-96BD-4880735DBC94}
 // Glslang 1.1.82.1
 
 #pragma once
diff --git a/Engine/Source/Developer/VulkanShaderFormat/Private/VulkanBackend.cpp b/Engine/Source/Developer/VulkanShaderFormat/Private/VulkanBackend.cpp
--- a/Engine/Source/Developer/VulkanShaderFormat/Private/VulkanBackend.cpp
+++ b/Engine/Source/Developer/VulkanShaderFormat/Private/VulkanBackend.cpp
@@ -4432,7 +4432,7 @@ static ir_rvalue* GenShaderInputSemantic(
 		Variable->interpolation = InputQualifier.Fields.InterpolationMode;
 		Variable->is_patch_constant = InputQualifier.Fields.bIsPatchConstant;
 
-		if (ParseState->bGenerateLayoutLocations && !InputQualifier.Fields.bIsPatchConstant)
+		if (ParseState->bGenerateLayoutLocations)
 		{
 			ConfigureInOutVariableLayout(Frequency, ParseState, Semantic, Variable, ir_var_in);
 		}
@@ -4464,7 +4464,7 @@ static ir_rvalue* GenShaderInputSemantic(
 		Variable->interpolation = InputQualifier.Fields.InterpolationMode;
 		Variable->is_patch_constant = InputQualifier.Fields.bIsPatchConstant;
 
-		if (ParseState->bGenerateLayoutLocations && !InputQualifier.Fields.bIsPatchConstant)
+		if (ParseState->bGenerateLayoutLocations)
 		{
 			ConfigureInOutVariableLayout(Frequency, ParseState, Semantic, Variable, ir_var_in);
 		}
@@ -4672,7 +4672,7 @@ static ir_rvalue* GenShaderOutputSemantic(
 	Variable->interpolation = OutputQualifier.Fields.InterpolationMode;
 	Variable->is_patch_constant = OutputQualifier.Fields.bIsPatchConstant;
 
-	if (ParseState->bGenerateLayoutLocations && !OutputQualifier.Fields.bIsPatchConstant)
+	if (ParseState->bGenerateLayoutLocations)
 	{
 		ConfigureInOutVariableLayout(Frequency, ParseState, Semantic, Variable, ir_var_out);
 	}
@@ -5785,17 +5785,22 @@ void FVulkanCodeBackend::GenShaderPatchConstantFunctionInputs(_mesa_glsl_parse_s
 		ir_dereference_record* lhs = assignment->lhs->as_dereference_record();
 		ir_rvalue* rhs = assignment->rhs;
 
-		if (!lhs)
+		if (!rhs)
 		{
 			continue;
 		}
 
-		if (!rhs)
+		// Check whether LHS is wrapped into an array.
+		// This might be the case on the OpenGL backend, but not necessarily on the Vulkan backend.
+		ir_dereference_array* lhs_array = nullptr;
+		if (lhs)
 		{
-			continue;
+			lhs_array = lhs->record->as_dereference_array();
+		}
+		else
+		{
+			lhs_array = assignment->lhs->as_dereference_array();
 		}
-
-		ir_dereference_array* lhs_array = lhs->record->as_dereference_array();
 
 		if (!lhs_array)
 		{
@@ -5820,8 +5825,6 @@ void FVulkanCodeBackend::GenShaderPatchConstantFunctionInputs(_mesa_glsl_parse_s
 			continue;
 		}
 
-		const char* OutArrayFieldName = lhs->field;
-
 		for (int OutputVertex = 0; OutputVertex < ParseState->tessellation.outputcontrolpoints; ++OutputVertex)
 		{
 			struct Helper
@@ -5869,18 +5872,34 @@ void FVulkanCodeBackend::GenShaderPatchConstantFunctionInputs(_mesa_glsl_parse_s
 			ir_rvalue* OutputPatchElement = rhs->clone(ParseState, 0);
 			Helper::ReplaceVariableDerefWithArrayDeref(OutputPatchElement, OutputPatchElementIndex);
 
-			PostCallInstructions.push_tail(
-				new (ParseState)ir_assignment(
-				OutputPatchElement,
-				new(ParseState)ir_dereference_record(
-				new(ParseState)ir_dereference_array(
-				OutputPatchArray->clone(ParseState, 0),
-				new(ParseState)ir_constant(OutputVertex)
-				),
-				OutArrayFieldName
-				)
-				)
+			if (lhs)
+			{
+				// Wrap LHS into a record again
+				PostCallInstructions.push_tail(
+					new (ParseState)ir_assignment(
+						OutputPatchElement,
+						new(ParseState)ir_dereference_record(
+							new(ParseState)ir_dereference_array(
+								OutputPatchArray->clone(ParseState, 0),
+								new(ParseState)ir_constant(OutputVertex)
+							),
+							lhs->field
+						)
+					)
 				);
+			}
+			else
+			{
+				PostCallInstructions.push_tail(
+					new (ParseState)ir_assignment(
+						OutputPatchElement,
+						new(ParseState)ir_dereference_array(
+							OutputPatchArray->clone(ParseState, 0),
+							new(ParseState)ir_constant(OutputVertex)
+						)
+					)
+				);
+			}
 		}
 	}
 }
diff --git a/Engine/Source/Runtime/RHI/Private/RHI.cpp b/Engine/Source/Runtime/RHI/Private/RHI.cpp
--- a/Engine/Source/Runtime/RHI/Private/RHI.cpp
+++ b/Engine/Source/Runtime/RHI/Private/RHI.cpp
@@ -42,14 +42,6 @@ static FAutoConsoleVariable CVarUseVulkanRealUBs(
 	ECVF_ReadOnly
 	);
 
-static FAutoConsoleVariable CVarVulkanEnableTessellation(
-	TEXT("r.Vulkan.EnableTessellation"),
-	0,
-	TEXT("0: Tessellation disabled[default]\n")
-	TEXT("1: Enable flat tessellation (experimental)"),
-	ECVF_ReadOnly
-);
-
 static TAutoConsoleVariable<int32> CVarDisableEngineAndAppRegistration(
 	TEXT("r.DisableEngineAndAppRegistration"),
 	0,
@@ -804,7 +796,7 @@ RHI_API bool RHISupportsTessellation(const EShaderPlatform Platform)
 {
 	if (IsFeatureLevelSupported(Platform, ERHIFeatureLevel::SM5))
 	{
-		return (Platform == SP_PCD3D_SM5) || (Platform == SP_XBOXONE_D3D12) || (Platform == SP_OPENGL_SM5) || (Platform == SP_OPENGL_ES31_EXT) || (Platform == SP_METAL_SM5) || (IsVulkanSM5Platform(Platform) && CVarVulkanEnableTessellation->GetInt() != 0);
+		return (Platform == SP_PCD3D_SM5) || (Platform == SP_XBOXONE_D3D12) || (Platform == SP_OPENGL_SM5) || (Platform == SP_OPENGL_ES31_EXT) || (Platform == SP_METAL_SM5) || (IsVulkanSM5Platform(Platform));
 	}
 	return false;
 }
diff --git a/Engine/Source/Runtime/VulkanRHI/Private/VulkanPipeline.cpp b/Engine/Source/Runtime/VulkanRHI/Private/VulkanPipeline.cpp
--- a/Engine/Source/Runtime/VulkanRHI/Private/VulkanPipeline.cpp
+++ b/Engine/Source/Runtime/VulkanRHI/Private/VulkanPipeline.cpp
@@ -1210,6 +1210,11 @@ void FVulkanPipelineStateCacheManager::CreateGfxPipelineFromEntry(FGfxPipelineEn
 		PipelineInfo.pTessellationState = &TessState;
 		check(InputAssembly.topology == VK_PRIMITIVE_TOPOLOGY_PATCH_LIST);
 		TessState.patchControlPoints = GfxEntry->ControlPoints;
+
+		// Workaround for translating HLSL tessellation shaders on Vulkan backend.
+		// Vertical flip of <gl_Position> alone is not sufficient here.
+		// Usual value for <frontFace> in UE4 is VK_FRONT_FACE_CLOCKWISE; use the opposite value to flip faces.
+		RasterizerState.frontFace = VK_FRONT_FACE_COUNTER_CLOCKWISE;
 	}
 
 	//#todo-rco: Fix me
@@ -1439,7 +1444,7 @@ FVulkanGfxLayout* FVulkanPipelineStateCacheManager::GetOrGenerateGfxLayout(const
 	return GfxLayout;
 }
 
-static inline VkPrimitiveTopology UEToVulkanTopologyType(EPrimitiveType PrimitiveType, bool bHasTessellation, uint16& OutControlPoints)
+static inline VkPrimitiveTopology UEToVulkanTopologyType(const FVulkanDevice* InDevice, EPrimitiveType PrimitiveType, bool bHasTessellation, uint16& OutControlPoints)
 {
 	if (bHasTessellation)
 	{
@@ -1450,10 +1455,42 @@ static inline VkPrimitiveTopology UEToVulkanTopologyType(EPrimitiveType Primitiv
 			OutControlPoints = 3;
 			return VK_PRIMITIVE_TOPOLOGY_PATCH_LIST;
 		case PT_1_ControlPointPatchList:
-			OutControlPoints = 1;
-			return VK_PRIMITIVE_TOPOLOGY_PATCH_LIST;
 		case PT_2_ControlPointPatchList:
-			OutControlPoints = 2;
+		case PT_3_ControlPointPatchList:
+		case PT_4_ControlPointPatchList:
+		case PT_5_ControlPointPatchList:
+		case PT_6_ControlPointPatchList:
+		case PT_7_ControlPointPatchList:
+		case PT_8_ControlPointPatchList:
+		case PT_9_ControlPointPatchList:
+		case PT_10_ControlPointPatchList:
+		case PT_12_ControlPointPatchList:
+		case PT_13_ControlPointPatchList:
+		case PT_14_ControlPointPatchList:
+		case PT_15_ControlPointPatchList:
+		case PT_16_ControlPointPatchList:
+		case PT_17_ControlPointPatchList:
+		case PT_18_ControlPointPatchList:
+		case PT_19_ControlPointPatchList:
+		case PT_20_ControlPointPatchList:
+		case PT_22_ControlPointPatchList:
+		case PT_23_ControlPointPatchList:
+		case PT_24_ControlPointPatchList:
+		case PT_25_ControlPointPatchList:
+		case PT_26_ControlPointPatchList:
+		case PT_27_ControlPointPatchList:
+		case PT_28_ControlPointPatchList:
+		case PT_29_ControlPointPatchList:
+		case PT_30_ControlPointPatchList:
+		case PT_31_ControlPointPatchList:
+		case PT_32_ControlPointPatchList:
+			OutControlPoints = (PrimitiveType - PT_1_ControlPointPatchList + 1);
+			checkf(
+				OutControlPoints <= InDevice->GetLimits().maxTessellationPatchSize,
+				TEXT("OutControlPoints (%d) exceeded limit of maximal patch size (%d)"),
+				OutControlPoints,
+				InDevice->GetLimits().maxTessellationPatchSize
+			);
 			return VK_PRIMITIVE_TOPOLOGY_PATCH_LIST;
 		default:
 			checkf(false, TEXT("Unsupported tessellation EPrimitiveType %d; probably missing a case in FStaticMeshSceneProxy::GetMeshElement()!"), (uint32)PrimitiveType);
@@ -1499,7 +1536,7 @@ FVulkanPipelineStateCacheManager::FGfxPipelineEntry* FVulkanPipelineStateCacheMa
 
 	OutGfxEntry->RasterizationSamples = OutGfxEntry->RenderPass->GetLayout().GetAttachmentDescriptions()[0].samples;
 	ensure(OutGfxEntry->RasterizationSamples == PSOInitializer.NumSamples);
-	OutGfxEntry->Topology = (uint32)UEToVulkanTopologyType(PSOInitializer.PrimitiveType, bHasTessellation, OutGfxEntry->ControlPoints);
+	OutGfxEntry->Topology = (uint32)UEToVulkanTopologyType(Device, PSOInitializer.PrimitiveType, bHasTessellation, OutGfxEntry->ControlPoints);
 
 	OutGfxEntry->ColorAttachmentStates.AddUninitialized(OutGfxEntry->RenderPass->GetLayout().GetNumColorAttachments());
 	for (int32 Index = 0; Index < OutGfxEntry->ColorAttachmentStates.Num(); ++Index)
