Make sure SDL is inited in Vulkan before attempting to use SDL functions

diff --git a/Engine/Source/Runtime/VulkanRHI/Private/Linux/VulkanLinuxPlatform.cpp b/Engine/Source/Runtime/VulkanRHI/Private/Linux/VulkanLinuxPlatform.cpp
--- a/Engine/Source/Runtime/VulkanRHI/Private/Linux/VulkanLinuxPlatform.cpp
+++ b/Engine/Source/Runtime/VulkanRHI/Private/Linux/VulkanLinuxPlatform.cpp
@@ -5,6 +5,7 @@
 #include <dlfcn.h>
 #include <SDL.h>
 #include <SDL_vulkan.h>
+#include "Linux/LinuxPlatformApplicationMisc.h"
 
 
 
@@ -122,10 +123,22 @@ void FVulkanLinuxPlatform::FreeVulkanLibrary()
 	bAttemptedLoad = false;
 }
 
-
+namespace
+{
+	void EnsureSDLIsInited()
+	{
+		if (!FLinuxPlatformApplicationMisc::InitSDL()) //   will not initialize more than once
+		{
+			FPlatformMisc::MessageBoxExt(EAppMsgType::Ok, TEXT("Vulkan InitSDL() failed, cannot initialize SDL."), TEXT("InitSDL Failed"));
+			UE_LOG(LogInit, Error, TEXT("Vulkan InitSDL() failed, cannot initialize SDL."));
+		}
+	}
+}
 
 void FVulkanLinuxPlatform::GetInstanceExtensions(TArray<const ANSICHAR*>& OutExtensions)
 {
+	EnsureSDLIsInited();
+
 	// we don't hardcode the extensions in Linux, we query SDL
 	static TArray<const ANSICHAR*> CachedLinuxExtensions;
 	if (CachedLinuxExtensions.Num() == 0)
@@ -150,6 +163,8 @@ void FVulkanLinuxPlatform::GetDeviceExtensions(TArray<const ANSICHAR*>& OutExten
 
 void FVulkanLinuxPlatform::CreateSurface(void* WindowHandle, VkInstance Instance, VkSurfaceKHR* OutSurface)
 {
+	EnsureSDLIsInited();
+
 	if (SDL_Vulkan_CreateSurface((SDL_Window*)WindowHandle, Instance, OutSurface) == SDL_FALSE)
 	{
 		UE_LOG(LogInit, Error, TEXT("Error initializing SDL Vulkan Surface: %s"), SDL_GetError());
