UE4.22 - Copy from 5996632
Fix only first device extension getting applied on Vulkan. Might help on some drivers by using the new memory extensions.

diff --git a/Engine/Source/Runtime/VulkanRHI/Private/VulkanLayers.cpp b/Engine/Source/Runtime/VulkanRHI/Private/VulkanLayers.cpp
--- a/Engine/Source/Runtime/VulkanRHI/Private/VulkanLayers.cpp
+++ b/Engine/Source/Runtime/VulkanRHI/Private/VulkanLayers.cpp
@@ -417,9 +417,9 @@ void FVulkanDynamicRHI::GetInstanceLayersAndExtensions(TArray<const ANSICHAR*>&
 	}
 #endif
 
-	TrimDuplicates(OutInstanceLayers);
 	if (OutInstanceLayers.Num() > 0)
 	{
+		TrimDuplicates(OutInstanceLayers);
 		UE_LOG(LogVulkanRHI, Display, TEXT("Using instance layers"));
 		for (const ANSICHAR* Layer : OutInstanceLayers)
 		{
@@ -431,9 +431,9 @@ void FVulkanDynamicRHI::GetInstanceLayersAndExtensions(TArray<const ANSICHAR*>&
 		UE_LOG(LogVulkanRHI, Display, TEXT("Not using instance layers"));
 	}
 
-	TrimDuplicates(OutInstanceExtensions);
 	if (OutInstanceExtensions.Num() > 0)
 	{
+		TrimDuplicates(OutInstanceExtensions);
 		UE_LOG(LogVulkanRHI, Display, TEXT("Using instance extensions"));
 		for (const ANSICHAR* Extension : OutInstanceExtensions)
 		{
@@ -618,7 +618,6 @@ void FVulkanDevice::GetDeviceExtensionsAndLayers(TArray<const ANSICHAR*>& OutDev
 		if (ListContains(AvailableExtensions, PlatformExtension))
 		{
 			OutDeviceExtensions.Add(PlatformExtension);
-			break;
 		}
 	}
 
@@ -644,6 +643,8 @@ void FVulkanDevice::GetDeviceExtensionsAndLayers(TArray<const ANSICHAR*>& OutDev
 
 	if (OutDeviceExtensions.Num() > 0)
 	{
+		TrimDuplicates(OutDeviceExtensions);
+
 		UE_LOG(LogVulkanRHI, Display, TEXT("Using device extensions"));
 		for (const ANSICHAR* Extension : OutDeviceExtensions)
 		{
@@ -653,6 +654,8 @@ void FVulkanDevice::GetDeviceExtensionsAndLayers(TArray<const ANSICHAR*>& OutDev
 
 	if (OutDeviceLayers.Num() > 0)
 	{
+		TrimDuplicates(OutDeviceLayers);
+
 		UE_LOG(LogVulkanRHI, Display, TEXT("Using device layers"));
 		for (const ANSICHAR* Layer : OutDeviceLayers)
 		{
