UE4.22 - Recommend Radeon drivers 19.4.1 or newer

diff --git a/Engine/Source/Runtime/VulkanRHI/Private/Windows/VulkanWindowsPlatform.cpp b/Engine/Source/Runtime/VulkanRHI/Private/Windows/VulkanWindowsPlatform.cpp
--- a/Engine/Source/Runtime/VulkanRHI/Private/Windows/VulkanWindowsPlatform.cpp
+++ b/Engine/Source/Runtime/VulkanRHI/Private/Windows/VulkanWindowsPlatform.cpp
@@ -235,7 +235,7 @@ void FVulkanWindowsPlatform::CheckDeviceDriver(uint32 DeviceIndex, const VkPhysi
 							if (MajorVersion < 18)
 							{
 								// Blacklist drivers older than 18.xx.xx drivers
-								FPlatformMisc::MessageBoxExt(EAppMsgType::Ok, TEXT("There are known issues with Vulkan with drivers older than the 18.xx.xx.xx\nfamily of Radeon drivers; the recommended version is 18.12.1.1 or anything more recent starting at 19.2.2: please try updating your driver to that version."), TEXT("Vulkan driver version"));
+								FPlatformMisc::MessageBoxExt(EAppMsgType::Ok, TEXT("There are known issues with older Vulkan Radeon drivers; the recommended version is 19.4.1; please try updating your driver to that version."), TEXT("Vulkan driver version"));
 								FPlatformMisc::RequestExitWithStatus(true, 1);
 							}
 							else if (WITH_EDITOR)
@@ -243,7 +243,7 @@ void FVulkanWindowsPlatform::CheckDeviceDriver(uint32 DeviceIndex, const VkPhysi
 								bool bBadVersion = false;
 								if (MajorVersion == 19)
 								{
-									if (MinorVersion < 2 || (MinorVersion == 2 && RevisionVersion <= 1))
+									if (MinorVersion < 4 || (MinorVersion == 4 && RevisionVersion <= 1))
 									{
 										bBadVersion = true;
 									}
@@ -258,8 +258,8 @@ void FVulkanWindowsPlatform::CheckDeviceDriver(uint32 DeviceIndex, const VkPhysi
 
 								if (bBadVersion)
 								{
-									// Blacklist drivers between 18.12.2 and 19.2.1, as they as it introduced an issue with Slate windows/Vulkan viewports on the editor
-									FPlatformMisc::MessageBoxExt(EAppMsgType::Ok, TEXT("There are known issues with Vulkan on the editor with the some \nRadeon drivers; the recommended version is up to 18.12.1.1 or anything more recent starting at 19.2.2: please try updating your driver to that version."), TEXT("Vulkan driver version"));
+									// Blacklist drivers between 18.12.2 and 19.2.1, as they as it introduced an issue with Slate windows/Vulkan viewports on the editor; 19.3.x also have crashes
+									FPlatformMisc::MessageBoxExt(EAppMsgType::Ok, TEXT("There are known issues with Vulkan on the editor with the some \nRadeon drivers; the recommended version is 19.4.1: please try updating your driver to that version."), TEXT("Vulkan driver version"));
 									FPlatformMisc::RequestExitWithStatus(true, 1);
 								}
 							}
