UE4.22 - Fix for null SRVs on RHIs < SM5

diff --git a/Engine/Source/Runtime/OpenGLDrv/Private/OpenGLUniformBuffer.cpp b/Engine/Source/Runtime/OpenGLDrv/Private/OpenGLUniformBuffer.cpp
--- a/Engine/Source/Runtime/OpenGLDrv/Private/OpenGLUniformBuffer.cpp
+++ b/Engine/Source/Runtime/OpenGLDrv/Private/OpenGLUniformBuffer.cpp
@@ -13,6 +13,11 @@
 #include "OpenGLDrvPrivate.h"
 #include "Misc/ScopeLock.h"
 
+#if UE_BUILD_DEBUG || UE_BUILD_DEVELOPMENT
+constexpr EUniformBufferValidation UniformBufferValidation = EUniformBufferValidation::ValidateResources;
+#else
+constexpr EUniformBufferValidation UniformBufferValidation = EUniformBufferValidation::None;
+#endif
 
 namespace OpenGLConsoleVariables
 {
@@ -737,13 +742,15 @@ void FOpenGLDynamicRHI::RHIUpdateUniformBuffer(FUniformBufferRHIParamRef Uniform
 		{
 			FRHIResource* Resource = *(FRHIResource**)((uint8*)Contents + Layout.Resources[ResourceIndex].MemberOffset);
 
-#if 0
-			checkf(Resource, TEXT("Invalid resource entry creating uniform buffer, %s.Resources[%u], ResourceType 0x%x."),
-				*Layout.GetDebugName().ToString(),
-				ResourceIndex,
-				Layout.Resources[ResourceIndex].MemberType);
-#endif
-			check(Resource);
+			if (!(GMaxRHIFeatureLevel <= ERHIFeatureLevel::ES3_1
+				&& (Layout.Resources[ResourceIndex].MemberType == UBMT_SRV || Layout.Resources[ResourceIndex].MemberType == UBMT_RDG_TEXTURE_SRV || Layout.Resources[ResourceIndex].MemberType == UBMT_RDG_BUFFER_SRV))
+				&& UniformBufferValidation == EUniformBufferValidation::ValidateResources)
+			{
+				checkf(Resource, TEXT("Invalid resource entry creating uniform buffer, %s.Resources[%u], ResourceType 0x%x."),
+					*Layout.GetDebugName().ToString(),
+					ResourceIndex,
+					Layout.Resources[ResourceIndex].MemberType);
+			}
 
 			UniformBuffer->ResourceTable[ResourceIndex] = Resource;
 		}
@@ -763,10 +770,15 @@ void FOpenGLDynamicRHI::RHIUpdateUniformBuffer(FUniformBufferRHIParamRef Uniform
 			{
 				FRHIResource* Resource = *(FRHIResource**)((uint8*)Contents + Layout.Resources[ResourceIndex].MemberOffset);
 
-				checkf(Resource, TEXT("Invalid resource entry creating uniform buffer, %s.Resources[%u], ResourceType 0x%x."),
-					*Layout.GetDebugName().ToString(),
-					ResourceIndex,
-					(uint8)Layout.Resources[ResourceIndex].MemberType);
+				if (!(GMaxRHIFeatureLevel <= ERHIFeatureLevel::ES3_1
+					&& (Layout.Resources[ResourceIndex].MemberType == UBMT_SRV || Layout.Resources[ResourceIndex].MemberType == UBMT_RDG_TEXTURE_SRV || Layout.Resources[ResourceIndex].MemberType == UBMT_RDG_BUFFER_SRV))
+					&& UniformBufferValidation == EUniformBufferValidation::ValidateResources)
+				{
+					checkf(Resource, TEXT("Invalid resource entry creating uniform buffer, %s.Resources[%u], ResourceType 0x%x."),
+						*Layout.GetDebugName().ToString(),
+						ResourceIndex,
+						(uint8)Layout.Resources[ResourceIndex].MemberType);
+				}
 
 				CmdListResources[ResourceIndex] = Resource;
 			}
diff --git a/Engine/Source/Runtime/VulkanRHI/Private/VulkanUniformBuffer.cpp b/Engine/Source/Runtime/VulkanRHI/Private/VulkanUniformBuffer.cpp
--- a/Engine/Source/Runtime/VulkanRHI/Private/VulkanUniformBuffer.cpp
+++ b/Engine/Source/Runtime/VulkanRHI/Private/VulkanUniformBuffer.cpp
@@ -17,6 +17,13 @@ enum
 #endif
 };
 
+#if UE_BUILD_DEBUG || UE_BUILD_DEVELOPMENT
+constexpr EUniformBufferValidation UniformBufferValidation = EUniformBufferValidation::ValidateResources;
+#else
+constexpr EUniformBufferValidation UniformBufferValidation = EUniformBufferValidation::None;
+#endif
+
+
 /*-----------------------------------------------------------------------------
 	Uniform buffer RHI object
 -----------------------------------------------------------------------------*/
@@ -63,7 +70,9 @@ FVulkanUniformBuffer::FVulkanUniformBuffer(const FRHIUniformBufferLayout& InLayo
 			FRHIResource* Resource = *(FRHIResource**)((uint8*)Contents + InLayout.Resources[Index].MemberOffset);
 
 			// Allow null SRV's in uniform buffers for feature levels that don't support SRV's in shaders
-			if (!(GMaxRHIFeatureLevel <= ERHIFeatureLevel::ES3_1 && InLayout.Resources[Index].MemberType == UBMT_SRV) && Validation == EUniformBufferValidation::ValidateResources)
+			if (!(GMaxRHIFeatureLevel <= ERHIFeatureLevel::ES3_1
+				&& (InLayout.Resources[Index].MemberType == UBMT_SRV || InLayout.Resources[Index].MemberType == UBMT_RDG_TEXTURE_SRV || InLayout.Resources[Index].MemberType == UBMT_RDG_BUFFER_SRV))
+				&& Validation == EUniformBufferValidation::ValidateResources)
 			{
 				checkf(Resource, TEXT("Invalid resource entry creating uniform buffer, %s.Resources[%u], ResourceType 0x%x."), *InLayout.GetDebugName().ToString(), Index, (uint8)InLayout.Resources[Index].MemberType);
 			}
@@ -75,16 +84,20 @@ FVulkanUniformBuffer::FVulkanUniformBuffer(const FRHIUniformBufferLayout& InLayo
 void FVulkanUniformBuffer::UpdateResourceTable(const FRHIUniformBufferLayout& InLayout, const void* Contents, int32 ResourceNum)
 {
 	check(ResourceTable.Num() == ResourceNum);
-
 	for (int32 ResourceIndex = 0; ResourceIndex < ResourceNum; ++ResourceIndex)
 	{
 		FRHIResource* Resource = *(FRHIResource**)((uint8*)Contents + InLayout.Resources[ResourceIndex].MemberOffset);
 
-		checkf(Resource, TEXT("Invalid resource entry creating uniform buffer, %s.Resources[%u], ResourceType 0x%x."),
-			*InLayout.GetDebugName().ToString(),
-			ResourceIndex,
-			(uint8)InLayout.Resources[ResourceIndex].MemberType);
-
+		// Allow null SRV's in uniform buffers for feature levels that don't support SRV's in shaders
+		if (!(GMaxRHIFeatureLevel <= ERHIFeatureLevel::ES3_1
+			&& (InLayout.Resources[ResourceIndex].MemberType == UBMT_SRV || InLayout.Resources[ResourceIndex].MemberType == UBMT_RDG_TEXTURE_SRV || InLayout.Resources[ResourceIndex].MemberType == UBMT_RDG_BUFFER_SRV))
+			&& UniformBufferValidation == EUniformBufferValidation::ValidateResources)
+		{
+			checkf(Resource, TEXT("Invalid resource entry creating uniform buffer, %s.Resources[%u], ResourceType 0x%x."),
+				*InLayout.GetDebugName().ToString(),
+				ResourceIndex,
+				(uint8)InLayout.Resources[ResourceIndex].MemberType);
+		}
 		ResourceTable[ResourceIndex] = Resource;
 	}
 }
diff --git a/Engine/Source/Runtime/Windows/D3D11RHI/Private/D3D11UniformBuffer.cpp b/Engine/Source/Runtime/Windows/D3D11RHI/Private/D3D11UniformBuffer.cpp
--- a/Engine/Source/Runtime/Windows/D3D11RHI/Private/D3D11UniformBuffer.cpp
+++ b/Engine/Source/Runtime/Windows/D3D11RHI/Private/D3D11UniformBuffer.cpp
@@ -7,6 +7,14 @@
 #include "D3D11RHIPrivate.h"
 #include "UniformBuffer.h"
 
+
+
+#if UE_BUILD_DEBUG || UE_BUILD_DEVELOPMENT
+constexpr EUniformBufferValidation UniformBufferValidation = EUniformBufferValidation::ValidateResources;
+#else
+constexpr EUniformBufferValidation UniformBufferValidation = EUniformBufferValidation::None;
+#endif
+
 /** Describes a uniform buffer in the free pool. */
 struct FPooledUniformBuffer
 {
@@ -295,11 +303,16 @@ void FD3D11DynamicRHI::RHIUpdateUniformBuffer(FUniformBufferRHIParamRef UniformB
 		{
 			FRHIResource* Resource = *(FRHIResource**)((uint8*)Contents + Layout.Resources[ResourceIndex].MemberOffset);
 
-			checkf(Resource, TEXT("Invalid resource entry creating uniform buffer, %s.Resources[%u], ResourceType 0x%x."),
-				*Layout.GetDebugName().ToString(),
-				ResourceIndex,
-				Layout.Resources[ResourceIndex].MemberType);
-
+			if (!(GMaxRHIFeatureLevel <= ERHIFeatureLevel::ES3_1
+				&& (Layout.Resources[ResourceIndex].MemberType == UBMT_SRV || Layout.Resources[ResourceIndex].MemberType == UBMT_RDG_TEXTURE_SRV || Layout.Resources[ResourceIndex].MemberType == UBMT_RDG_BUFFER_SRV))
+				&& UniformBufferValidation == EUniformBufferValidation::ValidateResources)
+			{
+				checkf(Resource, TEXT("Invalid resource entry creating uniform buffer, %s.Resources[%u], ResourceType 0x%x."),
+					*Layout.GetDebugName().ToString(),
+					ResourceIndex,
+					Layout.Resources[ResourceIndex].MemberType);
+			}
+	
 			UniformBuffer->ResourceTable[ResourceIndex] = Resource;
 		}
 	}
@@ -316,10 +329,15 @@ void FD3D11DynamicRHI::RHIUpdateUniformBuffer(FUniformBufferRHIParamRef UniformB
 			{
 				FRHIResource* Resource = *(FRHIResource**)((uint8*)Contents + Layout.Resources[ResourceIndex].MemberOffset);
 
-				checkf(Resource, TEXT("Invalid resource entry creating uniform buffer, %s.Resources[%u], ResourceType 0x%x."),
-					*Layout.GetDebugName().ToString(),
-					ResourceIndex,
-					Layout.Resources[ResourceIndex].MemberType);
+				if (!(GMaxRHIFeatureLevel <= ERHIFeatureLevel::ES3_1
+					&& (Layout.Resources[ResourceIndex].MemberType == UBMT_SRV || Layout.Resources[ResourceIndex].MemberType == UBMT_RDG_TEXTURE_SRV || Layout.Resources[ResourceIndex].MemberType == UBMT_RDG_BUFFER_SRV))
+					&& UniformBufferValidation == EUniformBufferValidation::ValidateResources)
+				{
+					checkf(Resource, TEXT("Invalid resource entry creating uniform buffer, %s.Resources[%u], ResourceType 0x%x."),
+						*Layout.GetDebugName().ToString(),
+						ResourceIndex,
+						Layout.Resources[ResourceIndex].MemberType);
+				}
 
 				CmdListResources[ResourceIndex] = Resource;
 			}
