UE4.22 - Recommend Radeon drivers 19.4.1 or newer (fix typo)
#rnx

diff --git a/Engine/Source/Runtime/VulkanRHI/Private/Windows/VulkanWindowsPlatform.cpp b/Engine/Source/Runtime/VulkanRHI/Private/Windows/VulkanWindowsPlatform.cpp
--- a/Engine/Source/Runtime/VulkanRHI/Private/Windows/VulkanWindowsPlatform.cpp
+++ b/Engine/Source/Runtime/VulkanRHI/Private/Windows/VulkanWindowsPlatform.cpp
@@ -243,7 +243,7 @@ void FVulkanWindowsPlatform::CheckDeviceDriver(uint32 DeviceIndex, const VkPhysi
 								bool bBadVersion = false;
 								if (MajorVersion == 19)
 								{
-									if (MinorVersion < 4 || (MinorVersion == 4 && RevisionVersion <= 1))
+									if (MinorVersion < 4 || (MinorVersion == 4 && RevisionVersion < 1))
 									{
 										bBadVersion = true;
 									}
