Fixed: Galaxy S10 crash during install

diff --git a/Engine/Source/Runtime/VulkanRHI/Private/VulkanPipeline.cpp b/Engine/Source/Runtime/VulkanRHI/Private/VulkanPipeline.cpp
--- a/Engine/Source/Runtime/VulkanRHI/Private/VulkanPipeline.cpp
+++ b/Engine/Source/Runtime/VulkanRHI/Private/VulkanPipeline.cpp
@@ -1082,6 +1082,14 @@ FGfxEntryKey FVulkanPipelineStateCacheManager::FGfxPipelineEntry::CreateKey() co
 
 void FVulkanPipelineStateCacheManager::CreateGfxPipelineFromEntry(FGfxPipelineEntry* GfxEntry, FVulkanShader* Shaders[ShaderStage::NumStages], FVulkanGfxPipeline* Pipeline)
 {
+	// It seems like android devices always require PS stage
+#if PLATFORM_ANDROID
+	if (Shaders[ShaderStage::Pixel] == nullptr)
+	{
+		Shaders[ShaderStage::Pixel] = ResourceCast(TShaderMapRef<FNULLPS>(GetGlobalShaderMap(GMaxRHIFeatureLevel))->GetPixelShader());
+	}
+#endif 
+	
 	if (!GfxEntry->bLoaded)
 	{
 		GfxEntry->GetOrCreateShaderModules(Shaders);
@@ -2038,11 +2046,6 @@ void GetVulkanShaders(const FBoundShaderStateInput& BSI, FVulkanShader* OutShade
 	{
 		OutShaders[ShaderStage::Pixel] = ResourceCast(BSI.PixelShaderRHI);
 	}
-	else if (GMaxRHIFeatureLevel <= ERHIFeatureLevel::ES3_1)
-	{
-		// Some mobile devices expect PS stage (S7 Adreno)
-		OutShaders[ShaderStage::Pixel] = ResourceCast(TShaderMapRef<FNULLPS>(GetGlobalShaderMap(GMaxRHIFeatureLevel))->GetPixelShader());
-	}
 
 #if PLATFORM_SUPPORTS_GEOMETRY_SHADERS
 	if (BSI.GeometryShaderRHI)
