Fix typo in Vulkan function name
#rnx

diff --git a/Engine/Source/Runtime/VulkanRHI/Private/Android/VulkanAndroidPlatform.cpp b/Engine/Source/Runtime/VulkanRHI/Private/Android/VulkanAndroidPlatform.cpp
--- a/Engine/Source/Runtime/VulkanRHI/Private/Android/VulkanAndroidPlatform.cpp
+++ b/Engine/Source/Runtime/VulkanRHI/Private/Android/VulkanAndroidPlatform.cpp
@@ -194,14 +194,14 @@ void FVulkanAndroidPlatform::OverridePlatformHandlers(bool bInit)
 	}
 }
 
-void FVulkanAndroidPlatform::BlockUntilWindowIsAwailable()
+void FVulkanAndroidPlatform::BlockUntilWindowIsAvailable()
 {
 	void* WindowHandle = FAndroidWindow::GetHardwareWindow();
 	if (WindowHandle == nullptr)
 	{
 		// Sleep if the hardware window isn't currently available.
 		// The Window may not exist if the activity is pausing/resuming, in which case we make this thread wait
-		FPlatformMisc::LowLevelOutputDebugString(TEXT("Waiting for Native window in FVulkanAndroidPlatform::BlockUntilWindowIsAwailable"));
+		FPlatformMisc::LowLevelOutputDebugString(TEXT("Waiting for Native window in FVulkanAndroidPlatform::BlockUntilWindowIsAvailable()"));
 		while (WindowHandle == nullptr)
 		{
 			FPlatformProcess::Sleep(0.001f);
diff --git a/Engine/Source/Runtime/VulkanRHI/Private/Android/VulkanAndroidPlatform.h b/Engine/Source/Runtime/VulkanRHI/Private/Android/VulkanAndroidPlatform.h
--- a/Engine/Source/Runtime/VulkanRHI/Private/Android/VulkanAndroidPlatform.h
+++ b/Engine/Source/Runtime/VulkanRHI/Private/Android/VulkanAndroidPlatform.h
@@ -96,8 +96,8 @@ public:
 
 	static bool RequiresWaitingForFrameCompletionEvent() { return false; }
 	
-	static void BlockUntilWindowIsAwailable();
-	
+	static void BlockUntilWindowIsAvailable();
+
 protected:
 	static void* VulkanLib;
 	static bool bAttemptedLoad;
diff --git a/Engine/Source/Runtime/VulkanRHI/Private/VulkanGenericPlatform.h b/Engine/Source/Runtime/VulkanRHI/Private/VulkanGenericPlatform.h
--- a/Engine/Source/Runtime/VulkanRHI/Private/VulkanGenericPlatform.h
+++ b/Engine/Source/Runtime/VulkanRHI/Private/VulkanGenericPlatform.h
@@ -89,7 +89,7 @@ public:
 	static bool SupportParallelRenderingTasks() { return true; }
 
 	// Allow platforms to add extension features to the DeviceInfo pNext chain
-	static void EnablePhysicalDeviceFeatureExtensions(VkDeviceCreateInfo& DeviceInfo) {};
+	static void EnablePhysicalDeviceFeatureExtensions(VkDeviceCreateInfo& DeviceInfo) {}
 
 	static bool RequiresSwapchainGeneralInitialLayout() { return false; }
 	
@@ -100,5 +100,5 @@ public:
 	static bool RequiresWaitingForFrameCompletionEvent() { return true; }
 
 	// Blocks until hardware window is available
-	static void BlockUntilWindowIsAwailable() {};
+	static void BlockUntilWindowIsAvailable() {}
 };
diff --git a/Engine/Source/Runtime/VulkanRHI/Private/VulkanViewport.cpp b/Engine/Source/Runtime/VulkanRHI/Private/VulkanViewport.cpp
--- a/Engine/Source/Runtime/VulkanRHI/Private/VulkanViewport.cpp
+++ b/Engine/Source/Runtime/VulkanRHI/Private/VulkanViewport.cpp
@@ -800,7 +800,7 @@ bool FVulkanViewport::Present(FVulkanCommandListContext* Context, FVulkanCmdBuff
 	}
 
 	// Do not present until hardware window is available. On Android window could be destroyed while RHIT executes commands
-	FVulkanPlatform::BlockUntilWindowIsAwailable();
+	FVulkanPlatform::BlockUntilWindowIsAvailable();
 
 	//Flush all commands
 	//check(0);
