Change define in Vulkan to an FVulkanPlatform call
#rnx

diff --git a/Engine/Source/Runtime/VulkanRHI/Private/Android/VulkanAndroidPlatform.h b/Engine/Source/Runtime/VulkanRHI/Private/Android/VulkanAndroidPlatform.h
--- a/Engine/Source/Runtime/VulkanRHI/Private/Android/VulkanAndroidPlatform.h
+++ b/Engine/Source/Runtime/VulkanRHI/Private/Android/VulkanAndroidPlatform.h
@@ -98,6 +98,9 @@ public:
 	
 	static void BlockUntilWindowIsAvailable();
 
+	// Does the platform allow a nullptr Pixelshader on the pipeline
+	static bool SupportsNullPixelShader() { return false; }
+
 protected:
 	static void* VulkanLib;
 	static bool bAttemptedLoad;
diff --git a/Engine/Source/Runtime/VulkanRHI/Private/VulkanGenericPlatform.h b/Engine/Source/Runtime/VulkanRHI/Private/VulkanGenericPlatform.h
--- a/Engine/Source/Runtime/VulkanRHI/Private/VulkanGenericPlatform.h
+++ b/Engine/Source/Runtime/VulkanRHI/Private/VulkanGenericPlatform.h
@@ -101,4 +101,7 @@ public:
 
 	// Blocks until hardware window is available
 	static void BlockUntilWindowIsAvailable() {}
+
+	// Does the platform allow a nullptr Pixelshader on the pipeline
+	static bool SupportsNullPixelShader() { return true; }
 };
diff --git a/Engine/Source/Runtime/VulkanRHI/Private/VulkanPipeline.cpp b/Engine/Source/Runtime/VulkanRHI/Private/VulkanPipeline.cpp
--- a/Engine/Source/Runtime/VulkanRHI/Private/VulkanPipeline.cpp
+++ b/Engine/Source/Runtime/VulkanRHI/Private/VulkanPipeline.cpp
@@ -1082,14 +1082,11 @@ FGfxEntryKey FVulkanPipelineStateCacheManager::FGfxPipelineEntry::CreateKey() co
 
 void FVulkanPipelineStateCacheManager::CreateGfxPipelineFromEntry(FGfxPipelineEntry* GfxEntry, FVulkanShader* Shaders[ShaderStage::NumStages], FVulkanGfxPipeline* Pipeline)
 {
-	// It seems like android devices always require PS stage
-#if PLATFORM_ANDROID
-	if (Shaders[ShaderStage::Pixel] == nullptr)
+	if (Shaders[ShaderStage::Pixel] == nullptr && !FVulkanPlatform::SupportsNullPixelShader())
 	{
 		Shaders[ShaderStage::Pixel] = ResourceCast(TShaderMapRef<FNULLPS>(GetGlobalShaderMap(GMaxRHIFeatureLevel))->GetPixelShader());
 	}
-#endif 
-	
+
 	if (!GfxEntry->bLoaded)
 	{
 		GfxEntry->GetOrCreateShaderModules(Shaders);
