UE4.22 - Fix for missing SF_AnisotropicLinear on Vulkan

diff --git a/Engine/Source/Runtime/VulkanRHI/Private/VulkanState.cpp b/Engine/Source/Runtime/VulkanRHI/Private/VulkanState.cpp
--- a/Engine/Source/Runtime/VulkanRHI/Private/VulkanState.cpp
+++ b/Engine/Source/Runtime/VulkanRHI/Private/VulkanState.cpp
@@ -16,7 +16,8 @@ inline VkSamplerMipmapMode TranslateMipFilterMode(ESamplerFilter InFilter)
 		case SF_Point:				return VK_SAMPLER_MIPMAP_MODE_NEAREST;
 		case SF_Bilinear:			return VK_SAMPLER_MIPMAP_MODE_NEAREST;
 		case SF_Trilinear:			return VK_SAMPLER_MIPMAP_MODE_LINEAR;
-		case SF_AnisotropicPoint:	return VK_SAMPLER_MIPMAP_MODE_LINEAR;
+		case SF_AnisotropicPoint:	return VK_SAMPLER_MIPMAP_MODE_POINT;
+		case SF_AnisotropicLinear:	return VK_SAMPLER_MIPMAP_MODE_LINEAR;
 		default:
 			break;
 	}
@@ -32,7 +33,8 @@ inline VkFilter TranslateMinMagFilterMode(ESamplerFilter InFilter)
 		case SF_Point:				return VK_FILTER_NEAREST;
 		case SF_Bilinear:			return VK_FILTER_LINEAR;
 		case SF_Trilinear:			return VK_FILTER_LINEAR;
-		case SF_AnisotropicPoint:	return VK_FILTER_LINEAR;
+		case SF_AnisotropicPoint:
+		case SF_AnisotropicLinear:	return VK_FILTER_LINEAR;
 		default:
 			break;
 	}
