Compile fix
#rnx

diff --git a/Engine/Source/Runtime/VulkanRHI/Private/VulkanState.cpp b/Engine/Source/Runtime/VulkanRHI/Private/VulkanState.cpp
--- a/Engine/Source/Runtime/VulkanRHI/Private/VulkanState.cpp
+++ b/Engine/Source/Runtime/VulkanRHI/Private/VulkanState.cpp
@@ -16,7 +16,7 @@ inline VkSamplerMipmapMode TranslateMipFilterMode(ESamplerFilter InFilter)
 		case SF_Point:				return VK_SAMPLER_MIPMAP_MODE_NEAREST;
 		case SF_Bilinear:			return VK_SAMPLER_MIPMAP_MODE_NEAREST;
 		case SF_Trilinear:			return VK_SAMPLER_MIPMAP_MODE_LINEAR;
-		case SF_AnisotropicPoint:	return VK_SAMPLER_MIPMAP_MODE_POINT;
+		case SF_AnisotropicPoint:	return VK_SAMPLER_MIPMAP_MODE_NEAREST;
 		case SF_AnisotropicLinear:	return VK_SAMPLER_MIPMAP_MODE_LINEAR;
 		default:
 			break;
