UE4.22 - Copy from 6429394
Fixed Vulkan validation messages not coming through

diff --git a/Engine/Source/Runtime/VulkanRHI/Private/VulkanDebug.cpp b/Engine/Source/Runtime/VulkanRHI/Private/VulkanDebug.cpp
--- a/Engine/Source/Runtime/VulkanRHI/Private/VulkanDebug.cpp
+++ b/Engine/Source/Runtime/VulkanRHI/Private/VulkanDebug.cpp
@@ -178,48 +178,6 @@ static VkBool32 DebugUtilsCallback(VkDebugUtilsMessageSeverityFlagBitsEXT MsgSev
 	const bool bError = (MsgSeverity & VK_DEBUG_UTILS_MESSAGE_SEVERITY_ERROR_BIT_EXT) != 0;
 	const bool bWarning = (MsgSeverity & VK_DEBUG_UTILS_MESSAGE_SEVERITY_WARNING_BIT_EXT) != 0;
 
-	if (CallbackData->pMessageIdName)
-	{
-		if (!FCStringAnsi::Strcmp(CallbackData->pMessageIdName, "UNASSIGNED-CoreValidation-Shader-OutputNotConsumed"))
-		{
-			// Warning: *** [Warning:Validation-1(UNASSIGNED-CoreValidation-Shader-OutputNotConsumed)] fragment shader writes to output location 0 with no matching attachment
-			return VK_FALSE;
-		}
-		else if (!FCStringAnsi::Strcmp(CallbackData->pMessageIdName, "VUID-VkSwapchainCreateInfoKHR-imageExtent-01274"))
-		{
-			// Warning: *** [Error:Validation341838324(VUID-VkSwapchainCreateInfoKHR-imageExtent-01274)] vkCreateSwapChainKHR() called with imageExtent = (8,8), which is outside the bounds returned by vkGetPhysicalDeviceSurfaceCapabilitiesKHR(): currentExtent = (0,0), minImageExtent = (0,0), maxImageExtent = (0,0).
-			return VK_FALSE;
-		}
-	}
-	else
-	{
-		if (MsgType == 2 && CallbackData->messageIdNumber == 5)
-		{
-			// SPIR-V module not valid: MemoryBarrier: Vulkan specification requires Memory Semantics to have one of the following bits set: Acquire, Release, AcquireRelease or SequentiallyConsistent
-			return VK_FALSE;
-		}
-		else if (MsgType == 2 && CallbackData->messageIdNumber == 2)
-		{
-			// fragment shader writes to output location 0 with no matching attachment
-			return VK_FALSE;
-		}
-		else if (MsgType == 2 && CallbackData->messageIdNumber == 3)
-		{
-			// Attachment 2 not written by fragment shader
-			return VK_FALSE;
-		}
-		else if (MsgType == 2 && CallbackData->messageIdNumber == 15)
-		{
-			// Cannot get query results on queryPool 0x9 with index 21 as data has not been collected for this index.
-			//return VK_FALSE;
-		}
-		else if (MsgType == 6 && CallbackData->messageIdNumber == 2)
-		{
-			// Vertex shader writes to output location 0.0 which is not consumed by fragment shader
-			return VK_FALSE;
-		}
-	}
-
 	const TCHAR* Severity = TEXT("");
 	if (bError)
 	{
@@ -290,20 +248,45 @@ static VkBool32 DebugUtilsCallback(VkDebugUtilsMessageSeverityFlagBitsEXT MsgSev
 		MsgBucket = EMsgBucket::Perf;
 	}
 
-	static TStaticArray<TSet<int32>, (int32)EMsgBucket::Count> SeenCodes;
-	if (GCVarUniqueValidationMessages->GetInt() == 0 || !SeenCodes[(int32)MsgBucket].Contains(CallbackData->messageIdNumber))
+	if (!CallbackData->pMessageIdName)
 	{
-		if (CallbackData->pMessageIdName)
-		{
-			VULKAN_REPORT_LOG(TEXT("*** [%s:%s%d(%s)] %s"), Severity, Type, CallbackData->messageIdNumber, ANSI_TO_TCHAR(CallbackData->pMessageIdName), ANSI_TO_TCHAR(CallbackData->pMessage));
-		}
-		else
-		{
-			VULKAN_REPORT_LOG(TEXT("*** [%s:%s%d] %s"), Severity, Type, CallbackData->messageIdNumber, ANSI_TO_TCHAR(CallbackData->pMessage));
-		}
-		if (GCVarUniqueValidationMessages->GetInt() == 1)
+		VULKAN_REPORT_LOG(TEXT("*** [%s:%s(NULL)%d] %s"), Severity, Type, CallbackData->messageIdNumber, ANSI_TO_TCHAR(CallbackData->pMessage));
+		ensure(0);
+		return VK_FALSE;
+	}
+
+	if (!FCStringAnsi::Strcmp(CallbackData->pMessageIdName, "UNASSIGNED-CoreValidation-Shader-OutputNotConsumed"))
+	{
+		// Warning: *** [Warning:Validation-1(UNASSIGNED-CoreValidation-Shader-OutputNotConsumed)] fragment shader writes to output location 0 with no matching attachment
+		return VK_FALSE;
+	}
+	else if (!FCStringAnsi::Strcmp(CallbackData->pMessageIdName, "VUID-VkSwapchainCreateInfoKHR-imageExtent-01274"))
+	{
+		// Warning: *** [Error:Validation341838324(VUID-VkSwapchainCreateInfoKHR-imageExtent-01274)] vkCreateSwapChainKHR() called with imageExtent = (8,8), which is outside the bounds returned by vkGetPhysicalDeviceSurfaceCapabilitiesKHR(): currentExtent = (0,0), minImageExtent = (0,0), maxImageExtent = (0,0).
+		return VK_FALSE;
+	}
+	else if (!FCStringAnsi::Strcmp(CallbackData->pMessageIdName, "UNASSIGNED-GPU-Assisted Validation Setup Error."))
+	{
+		// *** [Error:Validation(UNASSIGNED-GPU-Assisted Validation Setup Error.)] Unable to reserve descriptor binding slot on a device with only one slot.
+		return VK_FALSE;
+	}
+
+	static TStaticArray<TSet<uint32>, (int32)EMsgBucket::Count> SeenCodes;
+	bool bUniqueMessages = GCVarUniqueValidationMessages->GetInt() != 0;
+	bool bPrintMessage = true;
+	uint32 CRC = 0;
+	if (bUniqueMessages)
+	{
+		CRC = FCrc::StrCrc32(CallbackData->pMessageIdName);
+		bPrintMessage = !SeenCodes[(int32)MsgBucket].Contains(CRC);
+	}
+
+	if (bPrintMessage)
+	{
+		VULKAN_REPORT_LOG(TEXT("*** [%s:%s(%s)] %s"), Severity, Type, ANSI_TO_TCHAR(CallbackData->pMessageIdName), ANSI_TO_TCHAR(CallbackData->pMessage));
+		if (bUniqueMessages)
 		{
-			SeenCodes[(int32)MsgBucket].Add(CallbackData->messageIdNumber);
+			SeenCodes[(int32)MsgBucket].Add(CRC);
 		}
 	}
 
