Fixed: "Invalid resource entry creating uniform buffer" when running Vulkan on Mobile

diff --git a/Engine/Source/Runtime/VulkanRHI/Private/VulkanUniformBuffer.cpp b/Engine/Source/Runtime/VulkanRHI/Private/VulkanUniformBuffer.cpp
--- a/Engine/Source/Runtime/VulkanRHI/Private/VulkanUniformBuffer.cpp
+++ b/Engine/Source/Runtime/VulkanRHI/Private/VulkanUniformBuffer.cpp
@@ -23,6 +23,16 @@ constexpr EUniformBufferValidation UniformBufferValidation = EUniformBufferValid
 constexpr EUniformBufferValidation UniformBufferValidation = EUniformBufferValidation::None;
 #endif
 
+static void ValidateUniformBufferResource(const FRHIUniformBufferLayout& InLayout, int32 Index, FRHIResource* Resource, EUniformBufferValidation Validation)
+{
+	// Allow null SRV's in uniform buffers for feature levels that don't support SRV's in shaders
+	if (!(GMaxRHIFeatureLevel <= ERHIFeatureLevel::ES3_1
+		&& (InLayout.Resources[Index].MemberType == UBMT_SRV || InLayout.Resources[Index].MemberType == UBMT_RDG_TEXTURE_SRV || InLayout.Resources[Index].MemberType == UBMT_RDG_BUFFER_SRV))
+		&& Validation == EUniformBufferValidation::ValidateResources)
+	{
+		checkf(Resource, TEXT("Invalid resource entry creating uniform buffer, %s.Resources[%u], ResourceType 0x%x."), *InLayout.GetDebugName().ToString(), Index, (uint8)InLayout.Resources[Index].MemberType);
+	}
+}
 
 /*-----------------------------------------------------------------------------
 	Uniform buffer RHI object
@@ -68,14 +78,7 @@ FVulkanUniformBuffer::FVulkanUniformBuffer(const FRHIUniformBufferLayout& InLayo
 		for (uint32 Index = 0; Index < NumResources; Index++)
 		{
 			FRHIResource* Resource = *(FRHIResource**)((uint8*)Contents + InLayout.Resources[Index].MemberOffset);
-
-			// Allow null SRV's in uniform buffers for feature levels that don't support SRV's in shaders
-			if (!(GMaxRHIFeatureLevel <= ERHIFeatureLevel::ES3_1
-				&& (InLayout.Resources[Index].MemberType == UBMT_SRV || InLayout.Resources[Index].MemberType == UBMT_RDG_TEXTURE_SRV || InLayout.Resources[Index].MemberType == UBMT_RDG_BUFFER_SRV))
-				&& Validation == EUniformBufferValidation::ValidateResources)
-			{
-				checkf(Resource, TEXT("Invalid resource entry creating uniform buffer, %s.Resources[%u], ResourceType 0x%x."), *InLayout.GetDebugName().ToString(), Index, (uint8)InLayout.Resources[Index].MemberType);
-			}
+			ValidateUniformBufferResource(InLayout, (int32)Index, Resource, Validation);
 			ResourceTable[Index] = Resource;
 		}
 	}
@@ -87,17 +90,7 @@ void FVulkanUniformBuffer::UpdateResourceTable(const FRHIUniformBufferLayout& In
 	for (int32 ResourceIndex = 0; ResourceIndex < ResourceNum; ++ResourceIndex)
 	{
 		FRHIResource* Resource = *(FRHIResource**)((uint8*)Contents + InLayout.Resources[ResourceIndex].MemberOffset);
-
-		// Allow null SRV's in uniform buffers for feature levels that don't support SRV's in shaders
-		if (!(GMaxRHIFeatureLevel <= ERHIFeatureLevel::ES3_1
-			&& (InLayout.Resources[ResourceIndex].MemberType == UBMT_SRV || InLayout.Resources[ResourceIndex].MemberType == UBMT_RDG_TEXTURE_SRV || InLayout.Resources[ResourceIndex].MemberType == UBMT_RDG_BUFFER_SRV))
-			&& UniformBufferValidation == EUniformBufferValidation::ValidateResources)
-		{
-			checkf(Resource, TEXT("Invalid resource entry creating uniform buffer, %s.Resources[%u], ResourceType 0x%x."),
-				*InLayout.GetDebugName().ToString(),
-				ResourceIndex,
-				(uint8)InLayout.Resources[ResourceIndex].MemberType);
-		}
+		ValidateUniformBufferResource(InLayout, ResourceIndex, Resource, UniformBufferValidation);
 		ResourceTable[ResourceIndex] = Resource;
 	}
 }
@@ -230,12 +223,7 @@ inline void FVulkanDynamicRHI::UpdateUniformBuffer(FVulkanUniformBuffer* Uniform
 			for (int32 ResourceIndex = 0; ResourceIndex < NumResources; ++ResourceIndex)
 			{
 				FRHIResource* Resource = *(FRHIResource**)((uint8*)Contents + Layout.Resources[ResourceIndex].MemberOffset);
-
-				checkf(Resource, TEXT("Invalid resource entry creating uniform buffer, %s.Resources[%u], ResourceType 0x%x."),
-					*Layout.GetDebugName().ToString(),
-					ResourceIndex,
-					(uint8)Layout.Resources[ResourceIndex].MemberType);
-
+				ValidateUniformBufferResource(Layout, ResourceIndex, Resource, UniformBufferValidation);
 				CmdListResources[ResourceIndex] = Resource;
 			}
 		}
