Fixed: debug stats flickering on vulkan Mali devices
Insert a barrier between render-passes that use same attachements

diff --git a/Engine/Source/Runtime/VulkanRHI/Private/VulkanRenderTarget.cpp b/Engine/Source/Runtime/VulkanRHI/Private/VulkanRenderTarget.cpp
--- a/Engine/Source/Runtime/VulkanRHI/Private/VulkanRenderTarget.cpp
+++ b/Engine/Source/Runtime/VulkanRHI/Private/VulkanRenderTarget.cpp
@@ -312,6 +312,12 @@ void FTransitionAndLayoutManager::BeginRealRenderPass(FVulkanCommandListContext&
 			{
 				VulkanRHI::ImagePipelineBarrier(CmdBuffer->GetHandle(), Surface.Image, EImageLayoutBarrier::Undefined, EImageLayoutBarrier::ColorAttachment, SetupImageSubresourceRange());
 			}
+			else if (*Found == VK_IMAGE_LAYOUT_COLOR_ATTACHMENT_OPTIMAL && GetLoadAction(RPInfo.ColorRenderTargets[Index].Action) == ERenderTargetLoadAction::ELoad)
+			{
+				// make sure we have dependency between multiple render-passes that use the same attachment
+				// otherwise GPU can be execute them in any order
+				VulkanRHI::ImagePipelineBarrier(CmdBuffer->GetHandle(), Surface.Image, EImageLayoutBarrier::ColorAttachment, EImageLayoutBarrier::ColorAttachment, SetupImageSubresourceRange());
+			}
 			else
 			{
 				Context.RHITransitionResources(EResourceTransitionAccess::EWritable, &Texture, 1);
