Added: lazy vulkan backbuffer acquire (on write), enabled by default on Android

diff --git a/Engine/Source/Runtime/Engine/Private/UserInterface/Canvas.cpp b/Engine/Source/Runtime/Engine/Private/UserInterface/Canvas.cpp
--- a/Engine/Source/Runtime/Engine/Private/UserInterface/Canvas.cpp
+++ b/Engine/Source/Runtime/Engine/Private/UserInterface/Canvas.cpp
@@ -976,6 +976,7 @@ void FCanvas::Clear(const FLinearColor& ClearColor)
 				{
 					// do fast clear
 					FRHIRenderPassInfo RPInfo(CanvasRenderTarget->GetRenderTargetTexture(), ERenderTargetActions::Clear_Store);
+					TransitionRenderPassTargets(RHICmdList, RPInfo);
 					RHICmdList.BeginRenderPass(RPInfo, TEXT("ClearCanvas"));
 					RHICmdList.EndRenderPass();					
 				}
diff --git a/Engine/Source/Runtime/Renderer/Private/MobileShadingRenderer.cpp b/Engine/Source/Runtime/Renderer/Private/MobileShadingRenderer.cpp
--- a/Engine/Source/Runtime/Renderer/Private/MobileShadingRenderer.cpp
+++ b/Engine/Source/Runtime/Renderer/Private/MobileShadingRenderer.cpp
@@ -430,10 +430,14 @@ void FMobileSceneRenderer::Render(FRHICommandListImmediate& RHICmdList)
 			SceneColor = (View.bIsMobileMultiViewEnabled) ? SceneContext.MobileMultiViewSceneColor->GetRenderTargetItem().TargetableTexture : SceneContext.GetSceneColorSurface();
 			SceneColorResolve = GetMultiViewSceneColor(SceneContext);
 			ColorTargetAction = ERenderTargetActions::Clear_Resolve;
+			// Rendering to a backbuffer, make sure it's writeable
+			RHICmdList.TransitionResource(EResourceTransitionAccess::EWritable, SceneColorResolve);
 		}
 		else
 		{
 			SceneColor = GetMultiViewSceneColor(SceneContext);
+			// Rendering to a backbuffer, make sure it's writeable
+			RHICmdList.TransitionResource(EResourceTransitionAccess::EWritable, SceneColor);
 		}
 		SceneDepth = (View.bIsMobileMultiViewEnabled) ? SceneContext.MobileMultiViewSceneDepthZ->GetRenderTargetItem().TargetableTexture : static_cast<FTextureRHIRef>(SceneContext.GetSceneDepthSurface());
 	}
diff --git a/Engine/Source/Runtime/VulkanRHI/Private/VulkanDevice.cpp b/Engine/Source/Runtime/VulkanRHI/Private/VulkanDevice.cpp
--- a/Engine/Source/Runtime/VulkanRHI/Private/VulkanDevice.cpp
+++ b/Engine/Source/Runtime/VulkanRHI/Private/VulkanDevice.cpp
@@ -46,8 +46,8 @@ TAutoConsoleVariable<int32> CVarDelayAcquireBackBuffer(
 	1,
 	TEXT("Whether to delay acquiring the back buffer \n")
 	TEXT(" 0: acquire next image on frame start \n")
-	TEXT(" 1: acquire next image just before presenting, rendering is done to intermediate image which is then copied to real backbuffer (default) \n")
-	TEXT(" 2: acquire next image immediately after presenting current"),
+	TEXT(" 1: acquire next image just before presenting, rendering is done to intermediate image which is then copied to a real backbuffer (default) \n")
+	TEXT(" 2: acquire next image on first use"),
 	ECVF_ReadOnly
 );
 
@@ -59,7 +59,7 @@ static EDelayAcquireImageType DelayAcquireBackBuffer()
 	case 1:
 		return EDelayAcquireImageType::DelayAcquire;
 	case 2:
-		return EDelayAcquireImageType::PreAcquire;
+		return EDelayAcquireImageType::LazyAcquire;
 	}
 	return EDelayAcquireImageType::None;
 }
diff --git a/Engine/Source/Runtime/VulkanRHI/Private/VulkanRHI.cpp b/Engine/Source/Runtime/VulkanRHI/Private/VulkanRHI.cpp
--- a/Engine/Source/Runtime/VulkanRHI/Private/VulkanRHI.cpp
+++ b/Engine/Source/Runtime/VulkanRHI/Private/VulkanRHI.cpp
@@ -802,11 +802,6 @@ void FVulkanCommandListContext::RHIEndDrawingViewport(FViewportRHIParamRef Viewp
 		//#todo-rco: Check for r.FinishCurrentFrame
 	}
 
-	if (GVulkanDelayAcquireImage == EDelayAcquireImageType::PreAcquire)
-	{
-		RHI->DrawingViewport->PreAcquireSwapchainImage();
-	}
-
 	RHI->DrawingViewport = nullptr;
 
 	ReadAndCalculateGPUFrameTime();
diff --git a/Engine/Source/Runtime/VulkanRHI/Private/VulkanRenderTarget.cpp b/Engine/Source/Runtime/VulkanRHI/Private/VulkanRenderTarget.cpp
--- a/Engine/Source/Runtime/VulkanRHI/Private/VulkanRenderTarget.cpp
+++ b/Engine/Source/Runtime/VulkanRHI/Private/VulkanRenderTarget.cpp
@@ -249,6 +249,8 @@ void FTransitionAndLayoutManager::BeginRealRenderPass(FVulkanCommandListContext&
 		FTextureRHIParamRef Texture = RPInfo.ColorRenderTargets[Index].RenderTarget;
 		CA_ASSUME(Texture);
 		FVulkanSurface& Surface = FVulkanTextureBase::Cast(Texture)->Surface;
+		check(Surface.Image != VK_NULL_HANDLE);
+
 		VkImageLayout* Found = Layouts.Find(Surface.Image);
 		check(Found);
 
@@ -1284,9 +1286,13 @@ void FVulkanCommandListContext::RHITransitionResources(EResourceTransitionAccess
 		FPendingTransition PendingTransition;
 		for (int32 Index = 0; Index < NumTextures; ++Index)
 		{
-			if (InTextures[Index])
+			FRHITexture* RHITexture = InTextures[Index];
+			if (RHITexture)
 			{
-				PendingTransition.Textures.Add(InTextures[Index]);
+				PendingTransition.Textures.Add(RHITexture);
+				
+				FVulkanTextureBase* VulkanTexture = FVulkanTextureBase::Cast(RHITexture);
+				VulkanTexture->OnTransitionResource(*this, TransitionType);
 			}
 		}
 
diff --git a/Engine/Source/Runtime/VulkanRHI/Private/VulkanResources.h b/Engine/Source/Runtime/VulkanRHI/Private/VulkanResources.h
--- a/Engine/Source/Runtime/VulkanRHI/Private/VulkanResources.h
+++ b/Engine/Source/Runtime/VulkanRHI/Private/VulkanResources.h
@@ -499,11 +499,12 @@ struct FVulkanTextureBase : public FVulkanBaseShaderResource
 
 	bool bIsAliased;
 
+	virtual void OnTransitionResource(FVulkanCommandListContext& Context, EResourceTransitionAccess TransitionType) {};
+
 private:
 	void DestroyViews();
 };
 
-class FVulkanBackBuffer;
 class FVulkanTexture2D : public FRHITexture2D, public FVulkanTextureBase
 {
 public:
@@ -527,11 +528,6 @@ public:
 		return FRHIResource::GetRefCount();
 	}
 
-	virtual FVulkanBackBuffer* GetBackBuffer()
-	{
-		return nullptr;
-	}
-
 	virtual void* GetTextureBaseRHI() override final
 	{
 		FVulkanTextureBase* Base = static_cast<FVulkanTextureBase*>(this);
@@ -544,48 +540,6 @@ public:
 	}
 };
 
-class FVulkanBackBuffer : public FVulkanTexture2D
-{
-public:
-	FVulkanBackBuffer(FVulkanDevice& Device, EPixelFormat Format, uint32 SizeX, uint32 SizeY, uint32 UEFlags);
-	FVulkanBackBuffer(FVulkanDevice& Device, EPixelFormat Format, uint32 SizeX, uint32 SizeY, VkImage Image, uint32 UEFlags);
-	virtual ~FVulkanBackBuffer();
-
-	virtual FVulkanBackBuffer* GetBackBuffer() override final
-	{
-		return this;
-	}
-};
-
-class FVulkanBackBufferReference : public FRHITexture2D
-{
-public:
-	FVulkanBackBufferReference(EPixelFormat Format, uint32 SizeX, uint32 SizeY, uint32 UEFlags)
-		: FRHITexture2D(SizeX, SizeY, 1, 1, Format, UEFlags, FRHIResourceCreateInfo().ClearValueBinding)
-		, AcquiredBackBuffer(nullptr)
-	{
-	}
-
-	virtual FRHITexture2D* GetTexture2D() override 
-	{ 
-		return AcquiredBackBuffer; 
-	}
-	
-	virtual void* GetTextureBaseRHI() override
-	{
-		FVulkanTextureBase* Base = static_cast<FVulkanTextureBase*>(AcquiredBackBuffer);
-		return Base;
-	}
-
-	void SetBackBuffer(FVulkanBackBuffer* InAcquiredBackbuffer)
-	{
-		AcquiredBackBuffer = InAcquiredBackbuffer;
-	}
-
-private:
-	FVulkanBackBuffer* AcquiredBackBuffer;
-};
-
 class FVulkanTexture2DArray : public FRHITexture2DArray, public FVulkanTextureBase
 {
 public:
diff --git a/Engine/Source/Runtime/VulkanRHI/Private/VulkanTexture.cpp b/Engine/Source/Runtime/VulkanRHI/Private/VulkanTexture.cpp
--- a/Engine/Source/Runtime/VulkanRHI/Private/VulkanTexture.cpp
+++ b/Engine/Source/Runtime/VulkanRHI/Private/VulkanTexture.cpp
@@ -1805,30 +1805,6 @@ FVulkanTexture2D::~FVulkanTexture2D()
 	}
 }
 
-
-FVulkanBackBuffer::FVulkanBackBuffer(FVulkanDevice& Device, EPixelFormat Format, uint32 SizeX, uint32 SizeY, uint32 UEFlags)
-	: FVulkanTexture2D(Device, Format, SizeX, SizeY, 1, 1, UEFlags, FRHIResourceCreateInfo())
-{
-}
-
-FVulkanBackBuffer::FVulkanBackBuffer(FVulkanDevice& Device, EPixelFormat Format, uint32 SizeX, uint32 SizeY, VkImage Image, uint32 UEFlags)
-	: FVulkanTexture2D(Device, Format, SizeX, SizeY, 1, 1, Image, UEFlags, FRHIResourceCreateInfo())
-{
-}
-
-FVulkanBackBuffer::~FVulkanBackBuffer()
-{
-	if (Surface.IsImageOwner() == false)
-	{
-		// Clear flags so ~FVulkanTexture2D() doesn't try to re-destroy it
-		Surface.UEFlags = 0;
-		DefaultView.View = VK_NULL_HANDLE;
-		DefaultView.ViewId = 0;
-		Surface.Image = VK_NULL_HANDLE;
-	}
-}
-
-
 FVulkanTexture2DArray::FVulkanTexture2DArray(FVulkanDevice& Device, EPixelFormat Format, uint32 SizeX, uint32 SizeY, uint32 ArraySize, uint32 NumMips, uint32 NumSamples, uint32 Flags, FResourceBulkDataInterface* BulkData, const FClearValueBinding& InClearValue)
 	:	FRHITexture2DArray(SizeX, SizeY, ArraySize, NumMips, NumSamples, Format, Flags, InClearValue)
 	,	FVulkanTextureBase(Device, VK_IMAGE_VIEW_TYPE_2D_ARRAY, Format, SizeX, SizeY, 1, ArraySize, NumMips, NumSamples, Flags, BulkData)
diff --git a/Engine/Source/Runtime/VulkanRHI/Private/VulkanViewport.cpp b/Engine/Source/Runtime/VulkanRHI/Private/VulkanViewport.cpp
--- a/Engine/Source/Runtime/VulkanRHI/Private/VulkanViewport.cpp
+++ b/Engine/Source/Runtime/VulkanRHI/Private/VulkanViewport.cpp
@@ -12,37 +12,105 @@
 #include "HAL/PlatformAtomics.h"
 #include "Engine/RendererSettings.h"
 
-struct FRHICommandAcquireBackBuffer final : public FRHICommand<FRHICommandAcquireBackBuffer>
+struct FRHICommandProcessDeferredDeletionQueue final : public FRHICommand<FRHICommandProcessDeferredDeletionQueue>
 {
-	FVulkanViewport* Viewport;
-	FVulkanBackBufferReference* NewBackBufferReference;
-	FORCEINLINE_DEBUGGABLE FRHICommandAcquireBackBuffer(FVulkanViewport* InViewport, FVulkanBackBufferReference* InNewBackBufferReference)
-		: Viewport(InViewport)
-		, NewBackBufferReference(InNewBackBufferReference)
+	FVulkanDevice* Device;
+	FORCEINLINE_DEBUGGABLE FRHICommandProcessDeferredDeletionQueue(FVulkanDevice* InDevice)
+		: Device(InDevice)
 	{
 	}
 
 	void Execute(FRHICommandListBase& CmdList)
 	{
-		Viewport->AcquireBackBuffer(CmdList, NewBackBufferReference);
+		Device->GetDeferredDeletionQueue().ReleaseResources();
 	}
 };
 
 
-struct FRHICommandProcessDeferredDeletionQueue final : public FRHICommand<FRHICommandProcessDeferredDeletionQueue>
+FVulkanBackBuffer::FVulkanBackBuffer(FVulkanDevice& Device, FVulkanViewport* InViewport, EPixelFormat Format, uint32 SizeX, uint32 SizeY, uint32 UEFlags)
+	: FVulkanTexture2D(Device, Format, SizeX, SizeY, 1, 1, VK_NULL_HANDLE, UEFlags, FRHIResourceCreateInfo())
+	, Viewport(InViewport)
 {
-	FVulkanDevice* Device;
-	FORCEINLINE_DEBUGGABLE FRHICommandProcessDeferredDeletionQueue(FVulkanDevice* InDevice)
-		: Device(InDevice)
+}
+
+void FVulkanBackBuffer::ReleaseAcquiredImage()
+{
+	DefaultView.View = VK_NULL_HANDLE;
+	DefaultView.ViewId = 0;
+	Surface.Image = VK_NULL_HANDLE;
+}
+
+void FVulkanBackBuffer::ReleaseViewport()
+{
+	Viewport = nullptr;
+	ReleaseAcquiredImage();
+}
+
+void FVulkanBackBuffer::OnGetBackBufferImage(FRHICommandListImmediate& RHICmdList)
+{
+	check(Viewport);
+	if (GVulkanDelayAcquireImage == EDelayAcquireImageType::None)
 	{
+		FVulkanCommandListContext& Context = (FVulkanCommandListContext&)RHICmdList.GetContext();
+		AcquireBackBufferImage(Context);
 	}
+}
 
-	void Execute(FRHICommandListBase& CmdList)
+void FVulkanBackBuffer::OnAdvanceBackBufferFrame(FRHICommandListImmediate& RHICmdList)
+{
+	check(Viewport);
+	ReleaseAcquiredImage();
+}
+
+void FVulkanBackBuffer::OnTransitionResource(FVulkanCommandListContext& Context, EResourceTransitionAccess TransitionType)
+{
+	if (TransitionType == EResourceTransitionAccess::EWritable && GVulkanDelayAcquireImage == EDelayAcquireImageType::LazyAcquire)
 	{
-		Device->GetDeferredDeletionQueue().ReleaseResources();
+		AcquireBackBufferImage(Context);
 	}
-};
+}
+
+void FVulkanBackBuffer::AcquireBackBufferImage(FVulkanCommandListContext& Context)
+{
+	check(Viewport);
+	if (Surface.Image == VK_NULL_HANDLE)
+	{
+		check(Viewport->AcquiredImageIndex == -1);
+		
+		Viewport->AcquireImageIndex();
+		int32 AcquiredImageIndex = Viewport->AcquiredImageIndex;
+		FVulkanTextureView& ImageView = Viewport->TextureViews[AcquiredImageIndex];
+
+		Surface.Image = ImageView.Image;
+		DefaultView.View = ImageView.View;
+		DefaultView.ViewId = ImageView.ViewId;
+
+		// right after acquiring image is in undefined state
+		FTransitionAndLayoutManager& LayoutMgr = Context.GetTransitionAndLayoutManager();
+		VkImageLayout& CurrentLayout = LayoutMgr.FindOrAddLayoutRW(ImageView.Image, VK_IMAGE_LAYOUT_UNDEFINED);
+		CurrentLayout = VK_IMAGE_LAYOUT_UNDEFINED;
+		
+		FVulkanCommandBufferManager* CmdBufferManager = Context.GetCommandBufferManager();
+		FVulkanCmdBuffer* CmdBuffer = CmdBufferManager->GetActiveCmdBuffer();
+		if (CmdBuffer->IsInsideRenderPass())
+		{
+			// This could happen due to a SetRT(AndClear) call lingering around (so emulated needs to be ended); however REAL render passes should already have been ended!
+			checkf(!LayoutMgr.bInsideRealRenderPass, TEXT("Did not end Render Pass!"));
+			LayoutMgr.EndEmulatedRenderPass(CmdBuffer);
+		}
+			
+		// Add a dependency with the acquired semaphore
+		CmdBuffer->AddWaitSemaphore(VK_PIPELINE_STAGE_BOTTOM_OF_PIPE_BIT, Viewport->AcquiredSemaphore);
+	}
+}
 
+FVulkanBackBuffer::~FVulkanBackBuffer()
+{
+	check(Surface.IsImageOwner() == false);
+	// Clear flags so ~FVulkanTexture2D() doesn't try to re-destroy it
+	Surface.UEFlags = 0;
+	ReleaseAcquiredImage();
+}
 
 FVulkanViewport::FVulkanViewport(FVulkanDynamicRHI* InRHI, FVulkanDevice* InDevice, void* InWindowHandle, uint32 InSizeX, uint32 InSizeY, bool bInIsFullscreen, EPixelFormat InPreferredPixelFormat)
 	: VulkanRHI::FDeviceChild(InDevice)
@@ -52,7 +120,6 @@ FVulkanViewport::FVulkanViewport(FVulkanDynamicRHI* InRHI, FVulkanDevice* InDevi
 	, bIsFullscreen(bInIsFullscreen)
 	, PixelFormat(InPreferredPixelFormat)
 	, AcquiredImageIndex(-1)
-	, PreAcquiredImageIndex(-1)
 	, SwapChain(nullptr)
 	, WindowHandle(InWindowHandle)
 	, PresentCount(0)
@@ -81,8 +148,12 @@ FVulkanViewport::FVulkanViewport(FVulkanDynamicRHI* InRHI, FVulkanDevice* InDevi
 FVulkanViewport::~FVulkanViewport()
 {
 	RenderingBackBuffer = nullptr;
-	RenderingBackBufferReference = nullptr;
-	RHIBackBuffer = nullptr;
+	
+	if (RHIBackBuffer)
+	{
+		RHIBackBuffer->ReleaseViewport();
+		RHIBackBuffer = nullptr;
+	}
 	
 	if (FVulkanPlatform::SupportsStandardSwapchain())
 	{
@@ -90,15 +161,11 @@ FVulkanViewport::~FVulkanViewport()
 		{
 			RenderingDoneSemaphores[Index]->Release();
 
-			for (int32 i = 0; i < NUM_BUFFERS; ++i)
-			{
-				BackBuffers[i] = nullptr;
-			}
 			TextureViews[Index].Destroy(*Device);
 
 			// FIXME: race condition on TransitionAndLayoutManager, could this be called from RT while RHIT is active?
 			Device->NotifyDeletedImage(BackBufferImages[Index]);
-			Device->NotifyDeletedRenderTarget(BackBufferImages[Index]);
+			Device->NotifyDeletedRenderTarget(BackBufferImages[Index]); 
 			BackBufferImages[Index] = VK_NULL_HANDLE;
 		}
 
@@ -152,27 +219,6 @@ bool FVulkanViewport::DoCheckedSwapChainJob(TFunction<int32(FVulkanViewport*)> S
 	return Status >= 0;
 }
 
-void FVulkanViewport::PreAcquireSwapchainImage()
-{
-	check(PreAcquiredImageIndex == -1);
-	AcquireImageIndex();
-	PreAcquiredImageIndex = AcquiredImageIndex;
-}
-
-void FVulkanViewport::GetNextImageIndex()
-{
-	if (PreAcquiredImageIndex != -1)
-	{
-		check(PreAcquiredImageIndex == AcquiredImageIndex);
-		check(AcquiredImageIndex == SwapChain->CurrentImageIndex);
-		PreAcquiredImageIndex = -1;
-	}
-	else
-	{
-		AcquireImageIndex();
-	}
-}
-
 void FVulkanViewport::AcquireImageIndex()
 {
 	if (!DoCheckedSwapChainJob(DoAcquireImageIndex))
@@ -182,104 +228,41 @@ void FVulkanViewport::AcquireImageIndex()
 	check(AcquiredImageIndex != -1);
 }
 
-void FVulkanViewport::AcquireBackBuffer(FRHICommandListBase& CmdList, FVulkanBackBufferReference* NewBackBufferReference)
-{
-	if (FVulkanPlatform::SupportsStandardSwapchain())
-	{
-		check(NewBackBufferReference);
-
-		GetNextImageIndex();
-
-		TRefCountPtr<FVulkanBackBuffer> AcquriedBackbuffer = BackBuffers[AcquiredImageIndex];
-		NewBackBufferReference->SetBackBuffer(AcquriedBackbuffer);
-		RHIBackBuffer = AcquriedBackbuffer;
-	}
-	FVulkanCommandListContext& Context = (FVulkanCommandListContext&)CmdList.GetContext();
-
-	FVulkanCommandBufferManager* CmdBufferManager = Context.GetCommandBufferManager();
-	FVulkanCmdBuffer* CmdBuffer = CmdBufferManager->GetActiveCmdBuffer();
-	if (CmdBuffer->IsInsideRenderPass())
-	{
-		// This could happen due to a SetRT(AndClear) call lingering around (so emulated needs to be ended); however REAL render passes should already have been ended!
-		FTransitionAndLayoutManager& LayoutMgr = Context.GetTransitionAndLayoutManager();
-		checkf(!LayoutMgr.bInsideRealRenderPass, TEXT("Did not end Render Pass!"));
-		LayoutMgr.EndEmulatedRenderPass(CmdBuffer);
-	}
-
-	if (FVulkanPlatform::SupportsStandardSwapchain())
-	{
-		FTransitionAndLayoutManager& LayoutMgr = Context.GetTransitionAndLayoutManager();
-		VkImageLayout& CurrentLayout = LayoutMgr.FindOrAddLayoutRW(BackBufferImages[AcquiredImageIndex], VK_IMAGE_LAYOUT_UNDEFINED);
-		
-		VulkanRHI::ImagePipelineBarrier(CmdBuffer->GetHandle(), BackBufferImages[AcquiredImageIndex],
-			EImageLayoutBarrier::Undefined, EImageLayoutBarrier::ColorAttachment, VulkanRHI::SetupImageSubresourceRange());
-		if (FVulkanPlatform::RequiresSwapchainGeneralInitialLayout())
-		{
-			// Fix for artifacting on Mali on Android O: Take an extra roundtrip through COLOR_OPTIMAL -> GENERAL -> COLOR_OPTIMAL
-			VulkanRHI::ImagePipelineBarrier(CmdBuffer->GetHandle(), BackBufferImages[AcquiredImageIndex],
-				EImageLayoutBarrier::ColorAttachment, EImageLayoutBarrier::PixelGeneralRW, VulkanRHI::SetupImageSubresourceRange());
-			VulkanRHI::ImagePipelineBarrier(CmdBuffer->GetHandle(), BackBufferImages[AcquiredImageIndex],
-				EImageLayoutBarrier::PixelGeneralRW, EImageLayoutBarrier::ColorAttachment, VulkanRHI::SetupImageSubresourceRange());
-		}
-
-		CurrentLayout = VK_IMAGE_LAYOUT_COLOR_ATTACHMENT_OPTIMAL;
-	}
-
-	// Submit here so we can add a dependency with the acquired semaphore
-	CmdBuffer->End();
-	if (FVulkanPlatform::SupportsStandardSwapchain())
-	{
-		CmdBuffer->AddWaitSemaphore(VK_PIPELINE_STAGE_BOTTOM_OF_PIPE_BIT, AcquiredSemaphore);
-	}
-	Device->GetGraphicsQueue()->Submit(CmdBuffer);
-	CmdBufferManager->FreeUnusedCmdBuffers();
-	CmdBufferManager->PrepareForNewActiveCommandBuffer();
-}
-
-FTexture2DRHIRef FVulkanViewport::GetBackBuffer(FRHICommandList& RHICmdList)
+FTexture2DRHIRef FVulkanViewport::GetBackBuffer(FRHICommandListImmediate& RHICmdList)
 {
 	check(IsInRenderingThread());
 
 	// make sure we aren't in the middle of swapchain recreation (which can happen on e.g. RHI thread)
 	FScopeLock LockSwapchain(&RecreatingSwapchain);
 
-	if (!RenderingBackBuffer && FVulkanPlatform::SupportsStandardSwapchain())
+	if (FVulkanPlatform::SupportsStandardSwapchain() && GVulkanDelayAcquireImage != EDelayAcquireImageType::DelayAcquire)
 	{
-		check(GVulkanDelayAcquireImage != EDelayAcquireImageType::DelayAcquire);
-		
-		if (RenderingBackBufferReference.IsValid())
-		{
-			return RenderingBackBufferReference.GetReference();
-		}
-				
-		RenderingBackBufferReference = new FVulkanBackBufferReference(PixelFormat, SizeX, SizeY, TexCreate_Presentable | TexCreate_RenderTargetable);
-						
 		check(RHICmdList.IsImmediate());
-
-		if (RHICmdList.Bypass() || !IsRunningRHIInSeparateThread())
-		{
-			FRHICommandAcquireBackBuffer Cmd(this, RenderingBackBufferReference);
-			Cmd.Execute(RHICmdList);
-		}
-		else
+		check(RHIBackBuffer);
+		
+		RHICmdList.EnqueueLambda([this](FRHICommandListImmediate& CmdList)
 		{
-			ALLOC_COMMAND_CL(RHICmdList, FRHICommandAcquireBackBuffer)(this, RenderingBackBufferReference);
-		}
+			this->RHIBackBuffer->OnGetBackBufferImage(CmdList);
+		});
 
-		return RenderingBackBufferReference.GetReference();
+		return RHIBackBuffer.GetReference();
 	}
-
+	
 	return RenderingBackBuffer.GetReference();
 }
 
-void FVulkanViewport::AdvanceBackBufferFrame()
+void FVulkanViewport::AdvanceBackBufferFrame(FRHICommandListImmediate& RHICmdList)
 {
 	check(IsInRenderingThread());
 
 	if (FVulkanPlatform::SupportsStandardSwapchain() && GVulkanDelayAcquireImage != EDelayAcquireImageType::DelayAcquire)
 	{
-		RenderingBackBuffer = nullptr;
-		RenderingBackBufferReference = nullptr;
+		check(RHIBackBuffer);
+		
+		RHICmdList.EnqueueLambda([this](FRHICommandListImmediate& CmdList)
+		{
+			this->RHIBackBuffer->OnAdvanceBackBufferFrame(CmdList);
+		});
 	}
 }
 
@@ -343,6 +326,11 @@ FVulkanFramebuffer::FVulkanFramebuffer(FVulkanDevice& Device, const FRHISetRende
 		}
 
 		FVulkanTextureBase* Texture = FVulkanTextureBase::Cast(RHITexture);
+		// this could fire in case one of the textures is FVulkanBackBuffer and it has not acquired an image
+		// with EDelayAcquireImageType::LazyAcquire acquire happens when texture transition to Writeable state
+		// make sure you call TransitionResource(Writable, Tex) before using this texture as a render-target
+		check(Texture->Surface.Image != VK_NULL_HANDLE);
+
 		ColorRenderTargetImages[Index] = Texture->Surface.Image;
 		MipIndex = InRTInfo.ColorRenderTarget[Index].MipIndex;
 
@@ -522,16 +510,15 @@ void FVulkanViewport::RecreateSwapchain(void* NewNativeWindow, bool bForce)
 
 	FScopeLock LockSwapchain(&RecreatingSwapchain);
 	RenderingBackBuffer = nullptr;
-	RenderingBackBufferReference = nullptr;
-	RHIBackBuffer = nullptr;
+	
+	if (RHIBackBuffer)
+	{
+		RHIBackBuffer->ReleaseViewport();
+		RHIBackBuffer = nullptr;
+	}
 	
 	if (FVulkanPlatform::SupportsStandardSwapchain())
 	{
-		for (int32 i = 0; i < NUM_BUFFERS; ++i)
-		{
-			BackBuffers[i] = nullptr;
-		}
-		
 		for (int32 Index = 0; Index < NUM_BUFFERS; ++Index)
 		{
 			TextureViews[Index].Destroy(*Device);
@@ -592,16 +579,15 @@ void FVulkanViewport::RecreateSwapchainFromRT(EPixelFormat PreferredPixelFormat)
 	Device->WaitUntilIdle();
 
 	RenderingBackBuffer = nullptr;
-	RenderingBackBufferReference = nullptr;
-	RHIBackBuffer = nullptr;
+	
+	if (RHIBackBuffer)
+	{
+		RHIBackBuffer->ReleaseViewport();
+		RHIBackBuffer = nullptr;
+	}
 		
 	if (FVulkanPlatform::SupportsStandardSwapchain())
 	{
-		for (int32 i = 0; i < NUM_BUFFERS; ++i)
-		{
-			BackBuffers[i] = nullptr;
-		}
-		
 		for (int32 Index = 0; Index < NUM_BUFFERS; ++Index)
 		{
 			TextureViews[Index].Destroy(*Device);
@@ -650,10 +636,6 @@ void FVulkanViewport::CreateSwapchain()
 		for (int32 Index = 0; Index < Images.Num(); ++Index)
 		{
 			BackBufferImages[Index] = Images[Index];
-
-			FName Name = FName(*FString::Printf(TEXT("BackBuffer%d"), Index));
-			//BackBuffers[Index]->SetName(Name);
-
 			TextureViews[Index].Create(*Device, Images[Index], VK_IMAGE_VIEW_TYPE_2D, VK_IMAGE_ASPECT_COLOR_BIT, PixelFormat, UEToVkTextureFormat(PixelFormat, false), 0, 1, 0, 1);
 
 			// Clear the swapchain to avoid a validation warning, and transition to ColorAttachment
@@ -672,27 +654,18 @@ void FVulkanViewport::CreateSwapchain()
 				VulkanRHI::vkCmdClearColorImage(CmdBuffer->GetHandle(), Images[Index], VK_IMAGE_LAYOUT_TRANSFER_DST_OPTIMAL, &Color, 1, &Range);
 				VulkanRHI::ImagePipelineBarrier(CmdBuffer->GetHandle(), Images[Index], EImageLayoutBarrier::TransferDest, EImageLayoutBarrier::ColorAttachment, Range);
 			}
-		}
-
-		if (GVulkanDelayAcquireImage != EDelayAcquireImageType::DelayAcquire)
-		{
-			for (int32 i = 0; i < NUM_BUFFERS; ++i)
-			{
-				BackBuffers[i] = new FVulkanBackBuffer(*Device, PixelFormat, SizeX, SizeY, VK_NULL_HANDLE, TexCreate_Presentable | TexCreate_RenderTargetable);
-				BackBuffers[i]->Surface.Image = BackBufferImages[i];
-				BackBuffers[i]->DefaultView.View = TextureViews[i].View;
-				BackBuffers[i]->DefaultView.ViewId = TextureViews[i].ViewId;
 
 #if VULKAN_ENABLE_DRAW_MARKERS
-				if (Device->GetDebugMarkerSetObjectName())
-				{
-					VulkanRHI::SetDebugMarkerName(Device->GetDebugMarkerSetObjectName(), Device->GetInstanceHandle(), BackBufferImages[i], "RenderingBackBuffer");
-				}
-#endif
+			if (Device->GetDebugMarkerSetObjectName())
+			{
+				VulkanRHI::SetDebugMarkerName(Device->GetDebugMarkerSetObjectName(), Device->GetInstanceHandle(), BackBufferImages[Index], "RenderingBackBuffer");
 			}
+#endif
 		}
 		
 		Device->GetImmediateContext().GetCommandBufferManager()->SubmitUploadCmdBuffer();
+
+		RHIBackBuffer = new FVulkanBackBuffer(*Device, this, PixelFormat, SizeX, SizeY, TexCreate_RenderTargetable | TexCreate_ShaderResource);
 	}
 	else
 	{
@@ -701,7 +674,7 @@ void FVulkanViewport::CreateSwapchain()
 
 	if (!FVulkanPlatform::SupportsStandardSwapchain() || GVulkanDelayAcquireImage == EDelayAcquireImageType::DelayAcquire)
 	{
-		RenderingBackBuffer = new FVulkanBackBuffer(*Device, PixelFormat, SizeX, SizeY, TexCreate_RenderTargetable | TexCreate_ShaderResource);
+		RenderingBackBuffer = new FVulkanTexture2D(*Device, PixelFormat, SizeX, SizeY, 1, 1, TexCreate_RenderTargetable | TexCreate_ShaderResource, FRHIResourceCreateInfo());
 #if VULKAN_ENABLE_DRAW_MARKERS
 		if (Device->GetDebugMarkerSetObjectName())
 		{
@@ -711,7 +684,6 @@ void FVulkanViewport::CreateSwapchain()
 	}
 
 	AcquiredImageIndex = -1;
-	PreAcquiredImageIndex = -1;
 }
 
 inline static void CopyImageToBackBuffer(FVulkanCmdBuffer* CmdBuffer, bool bSourceReadOnly, VkImage SrcSurface, VkImage DstSurface, int32 SizeX, int32 SizeY, int32 WindowSizeX, int32 WindowSizeY)
@@ -797,7 +769,7 @@ bool FVulkanViewport::Present(FVulkanCommandListContext* Context, FVulkanCmdBuff
 		if (GVulkanDelayAcquireImage == EDelayAcquireImageType::DelayAcquire && RenderingBackBuffer)
 		{
 			SCOPE_CYCLE_COUNTER(STAT_VulkanAcquireBackBuffer);
-			GetNextImageIndex();
+			AcquireImageIndex();
 
 			uint32 WindowSizeX = FMath::Min(SizeX, SwapChain->InternalWidth);
 			uint32 WindowSizeY = FMath::Min(SizeY, SwapChain->InternalHeight);
@@ -809,9 +781,7 @@ bool FVulkanViewport::Present(FVulkanCommandListContext* Context, FVulkanCmdBuff
 		else
 		{
 			check(AcquiredImageIndex != -1);
-			check(PreAcquiredImageIndex == -1);
-
-			check(RHIBackBuffer == nullptr || RHIBackBuffer->Surface.Image == BackBufferImages[AcquiredImageIndex]);
+			check(RHIBackBuffer != nullptr && RHIBackBuffer->Surface.Image == BackBufferImages[AcquiredImageIndex]);
 
 			VkImageLayout& Layout = Context->GetTransitionAndLayoutManager().FindOrAddLayoutRW(BackBufferImages[AcquiredImageIndex], VK_IMAGE_LAYOUT_UNDEFINED);
 			VulkanRHI::ImagePipelineBarrier(CmdBuffer->GetHandle(), BackBufferImages[AcquiredImageIndex], VulkanRHI::GetImageLayoutFromVulkanLayout(Layout), EImageLayoutBarrier::Present, VulkanRHI::SetupImageSubresourceRange());
@@ -879,9 +849,6 @@ bool FVulkanViewport::Present(FVulkanCommandListContext* Context, FVulkanCmdBuff
 		{
 			CustomPresent->PostPresent();
 		}
-
-		// Release the back buffer
-		RHIBackBuffer = nullptr;
 	}
 
 	if (FVulkanPlatform::RequiresWaitingForFrameCompletionEvent() && !bHasCustomPresent)
@@ -1025,20 +992,19 @@ void FVulkanDynamicRHI::RHIAdvanceFrameForGetViewportBackBuffer(FViewportRHIPara
 	check(IsInRenderingThread());
 	check(ViewportRHI);
 	FVulkanViewport* Viewport = ResourceCast(ViewportRHI);
-	Viewport->AdvanceBackBufferFrame();
+	FRHICommandListImmediate& RHICmdList = FRHICommandListExecutor::GetImmediateCommandList();
+
+	Viewport->AdvanceBackBufferFrame(RHICmdList);
 
+	if (RHICmdList.Bypass() || !IsRunningRHIInSeparateThread())
 	{
-		FRHICommandList& RHICmdList = FRHICommandListExecutor::GetImmediateCommandList();
-		if (RHICmdList.Bypass() || !IsRunningRHIInSeparateThread())
-		{
-			FRHICommandProcessDeferredDeletionQueue Cmd(Device);
-			Cmd.Execute(RHICmdList);
-		}
-		else
-		{
-			check(IsInRenderingThread());
-			ALLOC_COMMAND_CL(RHICmdList, FRHICommandProcessDeferredDeletionQueue)(Device);
-		}
+		FRHICommandProcessDeferredDeletionQueue Cmd(Device);
+		Cmd.Execute(RHICmdList);
+	}
+	else
+	{
+		check(IsInRenderingThread());
+		ALLOC_COMMAND_CL(RHICmdList, FRHICommandProcessDeferredDeletionQueue)(Device);
 	}
 }
 
diff --git a/Engine/Source/Runtime/VulkanRHI/Private/VulkanViewport.h b/Engine/Source/Runtime/VulkanRHI/Private/VulkanViewport.h
--- a/Engine/Source/Runtime/VulkanRHI/Private/VulkanViewport.h
+++ b/Engine/Source/Runtime/VulkanRHI/Private/VulkanViewport.h
@@ -12,12 +12,35 @@
 class FVulkanDynamicRHI;
 class FVulkanSwapChain;
 class FVulkanQueue;
+class FVulkanViewport;
 
 namespace VulkanRHI
 {
 	class FSemaphore;
 }
 
+class FVulkanBackBuffer : public FVulkanTexture2D
+{
+public:
+	FVulkanBackBuffer(FVulkanDevice& Device, FVulkanViewport* InViewport, EPixelFormat Format, uint32 SizeX, uint32 SizeY, uint32 UEFlags);
+	virtual ~FVulkanBackBuffer();
+	
+	virtual void OnTransitionResource(FVulkanCommandListContext& Context, EResourceTransitionAccess TransitionType) override final;
+
+	void OnGetBackBufferImage(FRHICommandListImmediate& RHICmdList);
+	void OnAdvanceBackBufferFrame(FRHICommandListImmediate& RHICmdList);
+
+	void ReleaseViewport();
+
+private:
+	void AcquireBackBufferImage(FVulkanCommandListContext& Context);
+	void ReleaseAcquiredImage();
+
+private:
+	FVulkanViewport* Viewport;
+};
+
+
 class FVulkanViewport : public FRHIViewport, public VulkanRHI::FDeviceChild
 {
 public:
@@ -26,7 +49,8 @@ public:
 	FVulkanViewport(FVulkanDynamicRHI* InRHI, FVulkanDevice* InDevice, void* InWindowHandle, uint32 InSizeX,uint32 InSizeY,bool bInIsFullscreen, EPixelFormat InPreferredPixelFormat);
 	~FVulkanViewport();
 
-	FTexture2DRHIRef GetBackBuffer(FRHICommandList& RHICmdList);
+	FTexture2DRHIRef GetBackBuffer(FRHICommandListImmediate& RHICmdList);
+	void AdvanceBackBufferFrame(FRHICommandListImmediate& RHICmdList);
 
 	void WaitForFrameEventCompletion();
 
@@ -48,9 +72,7 @@ public:
 	}
 
 	virtual void Tick(float DeltaTime) override final;
-
-	void AdvanceBackBufferFrame();
-
+	
 	bool Present(FVulkanCommandListContext* Context, FVulkanCmdBuffer* CmdBuffer, FVulkanQueue* Queue, FVulkanQueue* PresentQueue, bool bLockToVsync);
 
 	inline uint32 GetPresentCount() const
@@ -62,14 +84,11 @@ protected:
 	VkImage BackBufferImages[NUM_BUFFERS];
 	VulkanRHI::FSemaphore* RenderingDoneSemaphores[NUM_BUFFERS];
 	FVulkanTextureView TextureViews[NUM_BUFFERS];
-	
-	TRefCountPtr<FVulkanBackBuffer> BackBuffers[NUM_BUFFERS];
-	TRefCountPtr<FVulkanBackBufferReference> RenderingBackBufferReference;
-
-	// 'Dummy' back buffer
-	TRefCountPtr<FVulkanBackBuffer> RenderingBackBuffer;
 	TRefCountPtr<FVulkanBackBuffer> RHIBackBuffer;
 
+	// 'Dummy' back buffer
+	TRefCountPtr<FVulkanTexture2D>	RenderingBackBuffer;
+	
 	/** narrow-scoped section that locks access to back buffer during its recreation*/
 	FCriticalSection RecreatingSwapchain;
 
@@ -79,7 +98,6 @@ protected:
 	bool bIsFullscreen;
 	EPixelFormat PixelFormat;
 	int32 AcquiredImageIndex;
-	int32 PreAcquiredImageIndex;
 	FVulkanSwapChain* SwapChain;
 	void* WindowHandle;
 	uint32 PresentCount;
@@ -95,10 +113,7 @@ protected:
 	uint64 LastFrameFenceCounter = 0;
 
 	void CreateSwapchain();
-	void PreAcquireSwapchainImage();
-	void GetNextImageIndex();
 	void AcquireImageIndex();
-	void AcquireBackBuffer(FRHICommandListBase& CmdList, FVulkanBackBufferReference* NewBackBuffer);
 
 	void RecreateSwapchain(void* NewNativeWindow, bool bForce = false);
 	void RecreateSwapchainFromRT(EPixelFormat PreferredPixelFormat);
@@ -110,6 +125,7 @@ protected:
 	friend class FVulkanDynamicRHI;
 	friend class FVulkanCommandListContext;
 	friend struct FRHICommandAcquireBackBuffer;
+	friend class FVulkanBackBuffer;
 };
 
 template<>
diff --git a/Engine/Source/Runtime/VulkanRHI/Public/VulkanMemory.h b/Engine/Source/Runtime/VulkanRHI/Public/VulkanMemory.h
--- a/Engine/Source/Runtime/VulkanRHI/Public/VulkanMemory.h
+++ b/Engine/Source/Runtime/VulkanRHI/Public/VulkanMemory.h
@@ -20,7 +20,7 @@ enum class EDelayAcquireImageType
 {
 	None,			// acquire next image on frame start
 	DelayAcquire,	// acquire next image just before presenting, rendering is done to intermediate image which is copied to real backbuffer
-	PreAcquire,		// acquire next image immediately after presenting current
+	LazyAcquire,	// acquire next image on first use
 };
 
 extern EDelayAcquireImageType GVulkanDelayAcquireImage;
