fixed: Static analysis warning related to VK backbuffer acquire changes

diff --git a/Engine/Source/Runtime/VulkanRHI/Private/VulkanViewport.cpp b/Engine/Source/Runtime/VulkanRHI/Private/VulkanViewport.cpp
--- a/Engine/Source/Runtime/VulkanRHI/Private/VulkanViewport.cpp
+++ b/Engine/Source/Runtime/VulkanRHI/Private/VulkanViewport.cpp
@@ -79,6 +79,8 @@ void FVulkanBackBuffer::AcquireBackBufferImage(FVulkanCommandListContext& Contex
 		
 		Viewport->AcquireImageIndex();
 		int32 AcquiredImageIndex = Viewport->AcquiredImageIndex;
+		check(AcquiredImageIndex >= 0 && AcquiredImageIndex < ARRAY_COUNT(Viewport->TextureViews));
+		
 		FVulkanTextureView& ImageView = Viewport->TextureViews[AcquiredImageIndex];
 
 		Surface.Image = ImageView.Image;
