UE4.23 - Deprecate FUniformBufferRHIParamRef, FIndexBufferRHIParamRef, FVertexBufferRHIParamRef, FStructuredBufferRHIParamRef, FTextureRHIParamRef, FTexture2DRHIParamRef, FTexture2DArrayRHIParamRef, F

diff --git a/Engine/Source/Runtime/VulkanRHI/Private/VulkanCommands.cpp b/Engine/Source/Runtime/VulkanRHI/Private/VulkanCommands.cpp
--- a/Engine/Source/Runtime/VulkanRHI/Private/VulkanCommands.cpp
+++ b/Engine/Source/Runtime/VulkanRHI/Private/VulkanCommands.cpp
@@ -48,7 +48,7 @@ static_assert(STRUCT_OFFSET(FRHIDispatchIndirectParameters, ThreadGroupCountY) =
 static_assert(STRUCT_OFFSET(FRHIDispatchIndirectParameters, ThreadGroupCountZ) == STRUCT_OFFSET(VkDispatchIndirectCommand, z), "FRHIDispatchIndirectParameters Z dimension is wrong.");
 
 
-void FVulkanCommandListContext::RHISetStreamSource(uint32 StreamIndex, FVertexBufferRHIParamRef VertexBufferRHI, uint32 Offset)
+void FVulkanCommandListContext::RHISetStreamSource(uint32 StreamIndex, FRHIVertexBuffer* VertexBufferRHI, uint32 Offset)
 {
 	FVulkanVertexBuffer* VertexBuffer = ResourceCast(VertexBufferRHI);
 	if (VertexBuffer != nullptr)
@@ -57,7 +57,7 @@ void FVulkanCommandListContext::RHISetStreamSource(uint32 StreamIndex, FVertexBu
 	}
 }
 
-void FVulkanDynamicRHI::RHISetStreamOutTargets(uint32 NumTargets, const FVertexBufferRHIParamRef* VertexBuffers, const uint32* Offsets)
+void FVulkanDynamicRHI::RHISetStreamOutTargets(uint32 NumTargets, FRHIVertexBuffer* const* VertexBuffers, const uint32* Offsets)
 {
 	VULKAN_SIGNAL_UNIMPLEMENTED();
 }
@@ -125,7 +125,7 @@ void FVulkanCommandListContext::RHIDispatchComputeShader(uint32 ThreadGroupCount
 	VulkanRHI::/*Debug*/HeavyWeightBarrier(CmdBuffer/*, 2*/);
 }
 
-void FVulkanCommandListContext::RHIDispatchIndirectComputeShader(FVertexBufferRHIParamRef ArgumentBufferRHI, uint32 ArgumentOffset) 
+void FVulkanCommandListContext::RHIDispatchIndirectComputeShader(FRHIVertexBuffer* ArgumentBufferRHI, uint32 ArgumentOffset)
 {
 	static_assert(sizeof(FRHIDispatchIndirectParameters) == sizeof(VkDispatchIndirectCommand), "Dispatch indirect doesn't match!");
 	FVulkanVertexBuffer* ArgumentBuffer = ResourceCast(ArgumentBufferRHI);
@@ -177,7 +177,7 @@ void FVulkanCommandListContext::RHISetUAVParameter(FRHIComputeShader* ComputeSha
 }
 
 
-void FVulkanCommandListContext::RHISetShaderTexture(FRHIVertexShader* VertexShaderRHI, uint32 TextureIndex, FTextureRHIParamRef NewTextureRHI)
+void FVulkanCommandListContext::RHISetShaderTexture(FRHIVertexShader* VertexShaderRHI, uint32 TextureIndex, FRHITexture* NewTextureRHI)
 {
 	check(PendingGfxState->GetCurrentShaderKey(ShaderStage::Vertex) == GetShaderKey(VertexShaderRHI));
 	FVulkanTextureBase* Texture = GetVulkanTextureFromRHITexture(NewTextureRHI);
@@ -186,7 +186,7 @@ void FVulkanCommandListContext::RHISetShaderTexture(FRHIVertexShader* VertexShad
 	NewTextureRHI->SetLastRenderTime((float)FPlatformTime::Seconds());
 }
 
-void FVulkanCommandListContext::RHISetShaderTexture(FRHIHullShader* HullShaderRHI, uint32 TextureIndex, FTextureRHIParamRef NewTextureRHI)
+void FVulkanCommandListContext::RHISetShaderTexture(FRHIHullShader* HullShaderRHI, uint32 TextureIndex, FRHITexture* NewTextureRHI)
 {
 #if PLATFORM_SUPPORTS_TESSELLATION_SHADERS
 	check(PendingGfxState->GetCurrentShaderKey(ShaderStage::Hull) == GetShaderKey(HullShaderRHI));
@@ -199,7 +199,7 @@ void FVulkanCommandListContext::RHISetShaderTexture(FRHIHullShader* HullShaderRH
 #endif
 }
 
-void FVulkanCommandListContext::RHISetShaderTexture(FRHIDomainShader* DomainShaderRHI, uint32 TextureIndex, FTextureRHIParamRef NewTextureRHI)
+void FVulkanCommandListContext::RHISetShaderTexture(FRHIDomainShader* DomainShaderRHI, uint32 TextureIndex, FRHITexture* NewTextureRHI)
 {
 #if PLATFORM_SUPPORTS_TESSELLATION_SHADERS
 	check(PendingGfxState->GetCurrentShaderKey(ShaderStage::Domain) == GetShaderKey(DomainShaderRHI));
@@ -212,7 +212,7 @@ void FVulkanCommandListContext::RHISetShaderTexture(FRHIDomainShader* DomainShad
 #endif
 }
 
-void FVulkanCommandListContext::RHISetShaderTexture(FRHIGeometryShader* GeometryShaderRHI, uint32 TextureIndex, FTextureRHIParamRef NewTextureRHI)
+void FVulkanCommandListContext::RHISetShaderTexture(FRHIGeometryShader* GeometryShaderRHI, uint32 TextureIndex, FRHITexture* NewTextureRHI)
 {
 #if VULKAN_SUPPORTS_GEOMETRY_SHADERS
 	check(PendingGfxState->GetCurrentShaderKey(ShaderStage::Geometry) == GetShaderKey(GeometryShaderRHI));
@@ -225,7 +225,7 @@ void FVulkanCommandListContext::RHISetShaderTexture(FRHIGeometryShader* Geometry
 #endif
 }
 
-void FVulkanCommandListContext::RHISetShaderTexture(FRHIPixelShader* PixelShaderRHI, uint32 TextureIndex, FTextureRHIParamRef NewTextureRHI)
+void FVulkanCommandListContext::RHISetShaderTexture(FRHIPixelShader* PixelShaderRHI, uint32 TextureIndex, FRHITexture* NewTextureRHI)
 {
 	check(PendingGfxState->GetCurrentShaderKey(ShaderStage::Pixel) == GetShaderKey(PixelShaderRHI));
 	FVulkanTextureBase* Texture = GetVulkanTextureFromRHITexture(NewTextureRHI);
@@ -234,7 +234,7 @@ void FVulkanCommandListContext::RHISetShaderTexture(FRHIPixelShader* PixelShader
 	NewTextureRHI->SetLastRenderTime((float)FPlatformTime::Seconds());
 }
 
-void FVulkanCommandListContext::RHISetShaderTexture(FRHIComputeShader* ComputeShaderRHI, uint32 TextureIndex, FTextureRHIParamRef NewTextureRHI)
+void FVulkanCommandListContext::RHISetShaderTexture(FRHIComputeShader* ComputeShaderRHI, uint32 TextureIndex, FRHITexture* NewTextureRHI)
 {
 	FVulkanComputeShader* ComputeShader = ResourceCast(ComputeShaderRHI);
 	check(PendingComputeState->GetCurrentShader() == ComputeShader);
@@ -524,7 +524,7 @@ inline void /*FVulkanCommandListContext::*/SetShaderUniformBufferResources(FVulk
 		{
 			const VkDescriptorType DescriptorType = DescriptorTypes[GlobalInfos[ResourceInfo.GlobalIndex].TypeIndex];
 			ensure(DescriptorType == VK_DESCRIPTOR_TYPE_SAMPLED_IMAGE || DescriptorType == VK_DESCRIPTOR_TYPE_COMBINED_IMAGE_SAMPLER);
-			FTextureRHIParamRef TexRef = (FTextureRHIParamRef)(ResourceArray[ResourceInfo.SourceUBResourceIndex].GetReference());
+			FRHITexture* TexRef = (FRHITexture*)(ResourceArray[ResourceInfo.SourceUBResourceIndex].GetReference());
 			if (TexRef)
 			{
 				const FVulkanTextureBase* BaseTexture = FVulkanTextureBase::Cast(TexRef);
@@ -647,13 +647,13 @@ inline void FVulkanCommandListContext::SetShaderUniformBuffer(ShaderStage::EStag
 	}
 }
 
-void FVulkanCommandListContext::RHISetShaderUniformBuffer(FRHIVertexShader* VertexShaderRHI, uint32 BufferIndex, FUniformBufferRHIParamRef BufferRHI)
+void FVulkanCommandListContext::RHISetShaderUniformBuffer(FRHIVertexShader* VertexShaderRHI, uint32 BufferIndex, FRHIUniformBuffer* BufferRHI)
 {
 	FVulkanUniformBuffer* UniformBuffer = ResourceCast(BufferRHI);
 	SetShaderUniformBuffer(ShaderStage::Vertex, UniformBuffer, BufferIndex, ResourceCast(VertexShaderRHI));
 }
 
-void FVulkanCommandListContext::RHISetShaderUniformBuffer(FRHIHullShader* HullShaderRHI, uint32 BufferIndex, FUniformBufferRHIParamRef BufferRHI)
+void FVulkanCommandListContext::RHISetShaderUniformBuffer(FRHIHullShader* HullShaderRHI, uint32 BufferIndex, FRHIUniformBuffer* BufferRHI)
 {
 #if PLATFORM_SUPPORTS_TESSELLATION_SHADERS
 	FVulkanUniformBuffer* UniformBuffer = ResourceCast(BufferRHI);
@@ -663,7 +663,7 @@ void FVulkanCommandListContext::RHISetShaderUniformBuffer(FRHIHullShader* HullSh
 #endif
 }
 
-void FVulkanCommandListContext::RHISetShaderUniformBuffer(FRHIDomainShader* DomainShaderRHI, uint32 BufferIndex, FUniformBufferRHIParamRef BufferRHI)
+void FVulkanCommandListContext::RHISetShaderUniformBuffer(FRHIDomainShader* DomainShaderRHI, uint32 BufferIndex, FRHIUniformBuffer* BufferRHI)
 {
 #if PLATFORM_SUPPORTS_TESSELLATION_SHADERS
 	FVulkanUniformBuffer* UniformBuffer = ResourceCast(BufferRHI);
@@ -673,7 +673,7 @@ void FVulkanCommandListContext::RHISetShaderUniformBuffer(FRHIDomainShader* Doma
 #endif
 }
 
-void FVulkanCommandListContext::RHISetShaderUniformBuffer(FRHIGeometryShader* GeometryShaderRHI, uint32 BufferIndex, FUniformBufferRHIParamRef BufferRHI)
+void FVulkanCommandListContext::RHISetShaderUniformBuffer(FRHIGeometryShader* GeometryShaderRHI, uint32 BufferIndex, FRHIUniformBuffer* BufferRHI)
 {
 #if VULKAN_SUPPORTS_GEOMETRY_SHADERS
 	FVulkanUniformBuffer* UniformBuffer = ResourceCast(BufferRHI);
@@ -683,13 +683,13 @@ void FVulkanCommandListContext::RHISetShaderUniformBuffer(FRHIGeometryShader* Ge
 #endif
 }
 
-void FVulkanCommandListContext::RHISetShaderUniformBuffer(FRHIPixelShader* PixelShaderRHI, uint32 BufferIndex, FUniformBufferRHIParamRef BufferRHI)
+void FVulkanCommandListContext::RHISetShaderUniformBuffer(FRHIPixelShader* PixelShaderRHI, uint32 BufferIndex, FRHIUniformBuffer* BufferRHI)
 {
 	FVulkanUniformBuffer* UniformBuffer = ResourceCast(BufferRHI);
 	SetShaderUniformBuffer(ShaderStage::Pixel, UniformBuffer, BufferIndex, ResourceCast(PixelShaderRHI));
 }
 
-void FVulkanCommandListContext::RHISetShaderUniformBuffer(FRHIComputeShader* ComputeShaderRHI, uint32 BufferIndex, FUniformBufferRHIParamRef BufferRHI)
+void FVulkanCommandListContext::RHISetShaderUniformBuffer(FRHIComputeShader* ComputeShaderRHI, uint32 BufferIndex, FRHIUniformBuffer* BufferRHI)
 {
 	FVulkanComputeShader* ComputeShader = ResourceCast(ComputeShaderRHI);
 	check(PendingComputeState->GetCurrentShader() == ComputeShader);
@@ -775,7 +775,7 @@ void FVulkanCommandListContext::RHIDrawPrimitive(uint32 BaseVertexIndex, uint32
 	}
 }
 
-void FVulkanCommandListContext::RHIDrawPrimitiveIndirect(FVertexBufferRHIParamRef ArgumentBufferRHI, uint32 ArgumentOffset)
+void FVulkanCommandListContext::RHIDrawPrimitiveIndirect(FRHIVertexBuffer* ArgumentBufferRHI, uint32 ArgumentOffset)
 {
 	static_assert(sizeof(FRHIDrawIndirectParameters) == sizeof(VkDrawIndirectCommand), "Draw indirect doesn't match!");
 
@@ -798,7 +798,7 @@ void FVulkanCommandListContext::RHIDrawPrimitiveIndirect(FVertexBufferRHIParamRe
 	}
 }
 
-void FVulkanCommandListContext::RHIDrawIndexedPrimitive(FIndexBufferRHIParamRef IndexBufferRHI, int32 BaseVertexIndex, uint32 FirstInstance,
+void FVulkanCommandListContext::RHIDrawIndexedPrimitive(FRHIIndexBuffer* IndexBufferRHI, int32 BaseVertexIndex, uint32 FirstInstance,
 	uint32 NumVertices, uint32 StartIndex, uint32 NumPrimitives, uint32 NumInstances)
 {
 #if VULKAN_ENABLE_AGGRESSIVE_STATS
@@ -823,7 +823,7 @@ void FVulkanCommandListContext::RHIDrawIndexedPrimitive(FIndexBufferRHIParamRef
 	}
 }
 
-void FVulkanCommandListContext::RHIDrawIndexedIndirect(FIndexBufferRHIParamRef IndexBufferRHI, FStructuredBufferRHIParamRef ArgumentsBufferRHI, int32 DrawArgumentsIndex, uint32 NumInstances)
+void FVulkanCommandListContext::RHIDrawIndexedIndirect(FRHIIndexBuffer* IndexBufferRHI, FRHIStructuredBuffer* ArgumentsBufferRHI, int32 DrawArgumentsIndex, uint32 NumInstances)
 {
 #if VULKAN_ENABLE_AGGRESSIVE_STATS
 	SCOPE_CYCLE_COUNTER(STAT_VulkanDrawCallTime);
@@ -845,7 +845,7 @@ void FVulkanCommandListContext::RHIDrawIndexedIndirect(FIndexBufferRHIParamRef I
 	}
 }
 
-void FVulkanCommandListContext::RHIDrawIndexedPrimitiveIndirect(FIndexBufferRHIParamRef IndexBufferRHI,FVertexBufferRHIParamRef ArgumentBufferRHI,uint32 ArgumentOffset)
+void FVulkanCommandListContext::RHIDrawIndexedPrimitiveIndirect(FRHIIndexBuffer* IndexBufferRHI, FRHIVertexBuffer* ArgumentBufferRHI,uint32 ArgumentOffset)
 {
 #if VULKAN_ENABLE_AGGRESSIVE_STATS
 	SCOPE_CYCLE_COUNTER(STAT_VulkanDrawCallTime);
@@ -1105,7 +1105,7 @@ void FVulkanCommandListContext::PrepareParallelFromBase(const FVulkanCommandList
 	TransitionAndLayoutManager.TempCopy(BaseContext.TransitionAndLayoutManager);
 }
 
-void FVulkanCommandListContext::RHICopyToStagingBuffer(FVertexBufferRHIParamRef SourceBufferRHI, FRHIStagingBuffer* StagingBufferRHI, uint32 Offset, uint32 NumBytes)
+void FVulkanCommandListContext::RHICopyToStagingBuffer(FRHIVertexBuffer* SourceBufferRHI, FRHIStagingBuffer* StagingBufferRHI, uint32 Offset, uint32 NumBytes)
 {
 	FVulkanCmdBuffer* CmdBuffer = CommandBufferManager->GetActiveCmdBuffer();
 	FVulkanVertexBuffer* VertexBuffer = ResourceCast(SourceBufferRHI);
diff --git a/Engine/Source/Runtime/VulkanRHI/Private/VulkanContext.h b/Engine/Source/Runtime/VulkanRHI/Private/VulkanContext.h
--- a/Engine/Source/Runtime/VulkanRHI/Private/VulkanContext.h
+++ b/Engine/Source/Runtime/VulkanRHI/Private/VulkanContext.h
@@ -169,16 +169,16 @@ public:
 		return Immediate == nullptr;
 	}
 
-	virtual void RHISetStreamSource(uint32 StreamIndex, FVertexBufferRHIParamRef VertexBuffer, uint32 Offset) final override;
+	virtual void RHISetStreamSource(uint32 StreamIndex, FRHIVertexBuffer* VertexBuffer, uint32 Offset) final override;
 	virtual void RHISetViewport(uint32 MinX, uint32 MinY, float MinZ, uint32 MaxX, uint32 MaxY, float MaxZ) final override;
 	virtual void RHISetScissorRect(bool bEnable, uint32 MinX, uint32 MinY, uint32 MaxX, uint32 MaxY) final override;
 	virtual void RHISetGraphicsPipelineState(FRHIGraphicsPipelineState* GraphicsState) final override;
-	virtual void RHISetShaderTexture(FRHIVertexShader* VertexShader, uint32 TextureIndex, FTextureRHIParamRef NewTexture) final override;
-	virtual void RHISetShaderTexture(FRHIHullShader* HullShader, uint32 TextureIndex, FTextureRHIParamRef NewTexture) final override;
-	virtual void RHISetShaderTexture(FRHIDomainShader* DomainShader, uint32 TextureIndex, FTextureRHIParamRef NewTexture) final override;
-	virtual void RHISetShaderTexture(FRHIGeometryShader* GeometryShader, uint32 TextureIndex, FTextureRHIParamRef NewTexture) final override;
-	virtual void RHISetShaderTexture(FRHIPixelShader* PixelShader, uint32 TextureIndex, FTextureRHIParamRef NewTexture) final override;
-	virtual void RHISetShaderTexture(FRHIComputeShader* PixelShader, uint32 TextureIndex, FTextureRHIParamRef NewTexture) final override;
+	virtual void RHISetShaderTexture(FRHIVertexShader* VertexShader, uint32 TextureIndex, FRHITexture* NewTexture) final override;
+	virtual void RHISetShaderTexture(FRHIHullShader* HullShader, uint32 TextureIndex, FRHITexture* NewTexture) final override;
+	virtual void RHISetShaderTexture(FRHIDomainShader* DomainShader, uint32 TextureIndex, FRHITexture* NewTexture) final override;
+	virtual void RHISetShaderTexture(FRHIGeometryShader* GeometryShader, uint32 TextureIndex, FRHITexture* NewTexture) final override;
+	virtual void RHISetShaderTexture(FRHIPixelShader* PixelShader, uint32 TextureIndex, FRHITexture* NewTexture) final override;
+	virtual void RHISetShaderTexture(FRHIComputeShader* PixelShader, uint32 TextureIndex, FRHITexture* NewTexture) final override;
 	virtual void RHISetShaderSampler(FRHIComputeShader* ComputeShader, uint32 SamplerIndex, FRHISamplerState* NewState) final override;
 	virtual void RHISetShaderSampler(FRHIVertexShader* VertexShader, uint32 SamplerIndex, FRHISamplerState* NewState) final override;
 	virtual void RHISetShaderSampler(FRHIGeometryShader* GeometryShader, uint32 SamplerIndex, FRHISamplerState* NewState) final override;
@@ -193,12 +193,12 @@ public:
 	virtual void RHISetShaderResourceViewParameter(FRHIHullShader* HullShader, uint32 SamplerIndex, FRHIShaderResourceView* SRV) final override;
 	virtual void RHISetShaderResourceViewParameter(FRHIDomainShader* DomainShader, uint32 SamplerIndex, FRHIShaderResourceView* SRV) final override;
 	virtual void RHISetShaderResourceViewParameter(FRHIGeometryShader* GeometryShader, uint32 SamplerIndex, FRHIShaderResourceView* SRV) final override;
-	virtual void RHISetShaderUniformBuffer(FRHIVertexShader* VertexShader, uint32 BufferIndex, FUniformBufferRHIParamRef Buffer) final override;
-	virtual void RHISetShaderUniformBuffer(FRHIHullShader* HullShader, uint32 BufferIndex, FUniformBufferRHIParamRef Buffer) final override;
-	virtual void RHISetShaderUniformBuffer(FRHIDomainShader* DomainShader, uint32 BufferIndex, FUniformBufferRHIParamRef Buffer) final override;
-	virtual void RHISetShaderUniformBuffer(FRHIGeometryShader* GeometryShader, uint32 BufferIndex, FUniformBufferRHIParamRef Buffer) final override;
-	virtual void RHISetShaderUniformBuffer(FRHIPixelShader* PixelShader, uint32 BufferIndex, FUniformBufferRHIParamRef Buffer) final override;
-	virtual void RHISetShaderUniformBuffer(FRHIComputeShader* ComputeShader, uint32 BufferIndex, FUniformBufferRHIParamRef Buffer) final override;
+	virtual void RHISetShaderUniformBuffer(FRHIVertexShader* VertexShader, uint32 BufferIndex, FRHIUniformBuffer* Buffer) final override;
+	virtual void RHISetShaderUniformBuffer(FRHIHullShader* HullShader, uint32 BufferIndex, FRHIUniformBuffer* Buffer) final override;
+	virtual void RHISetShaderUniformBuffer(FRHIDomainShader* DomainShader, uint32 BufferIndex, FRHIUniformBuffer* Buffer) final override;
+	virtual void RHISetShaderUniformBuffer(FRHIGeometryShader* GeometryShader, uint32 BufferIndex, FRHIUniformBuffer* Buffer) final override;
+	virtual void RHISetShaderUniformBuffer(FRHIPixelShader* PixelShader, uint32 BufferIndex, FRHIUniformBuffer* Buffer) final override;
+	virtual void RHISetShaderUniformBuffer(FRHIComputeShader* ComputeShader, uint32 BufferIndex, FRHIUniformBuffer* Buffer) final override;
 	virtual void RHISetShaderParameter(FRHIVertexShader* VertexShader, uint32 BufferIndex, uint32 BaseIndex, uint32 NumBytes, const void* NewValue) final override;
 	virtual void RHISetShaderParameter(FRHIHullShader* HullShader, uint32 BufferIndex, uint32 BaseIndex, uint32 NumBytes, const void* NewValue) final override;
 	virtual void RHISetShaderParameter(FRHIDomainShader* DomainShader, uint32 BufferIndex, uint32 BaseIndex, uint32 NumBytes, const void* NewValue) final override;
@@ -209,10 +209,10 @@ public:
 	virtual void RHISetRenderTargets(uint32 NumSimultaneousRenderTargets, const FRHIRenderTargetView* NewRenderTargets, const FRHIDepthRenderTargetView* NewDepthStencilTarget, uint32 NumUAVs, FRHIUnorderedAccessView* const* UAVs) final override;
 	virtual void RHISetRenderTargetsAndClear(const FRHISetRenderTargetsInfo& RenderTargetsInfo) final override;
 	virtual void RHIDrawPrimitive(uint32 BaseVertexIndex, uint32 NumPrimitives, uint32 NumInstances) final override;
-	virtual void RHIDrawPrimitiveIndirect(FVertexBufferRHIParamRef ArgumentBuffer, uint32 ArgumentOffset) final override;
-	virtual void RHIDrawIndexedIndirect(FIndexBufferRHIParamRef IndexBufferRHI, FStructuredBufferRHIParamRef ArgumentsBufferRHI, int32 DrawArgumentsIndex, uint32 NumInstances) final override;
-	virtual void RHIDrawIndexedPrimitive(FIndexBufferRHIParamRef IndexBuffer, int32 BaseVertexIndex, uint32 FirstInstance, uint32 NumVertices, uint32 StartIndex, uint32 NumPrimitives, uint32 NumInstances) final override;
-	virtual void RHIDrawIndexedPrimitiveIndirect(FIndexBufferRHIParamRef IndexBuffer, FVertexBufferRHIParamRef ArgumentBuffer, uint32 ArgumentOffset) final override;
+	virtual void RHIDrawPrimitiveIndirect(FRHIVertexBuffer* ArgumentBuffer, uint32 ArgumentOffset) final override;
+	virtual void RHIDrawIndexedIndirect(FRHIIndexBuffer* IndexBufferRHI, FRHIStructuredBuffer* ArgumentsBufferRHI, int32 DrawArgumentsIndex, uint32 NumInstances) final override;
+	virtual void RHIDrawIndexedPrimitive(FRHIIndexBuffer* IndexBuffer, int32 BaseVertexIndex, uint32 FirstInstance, uint32 NumVertices, uint32 StartIndex, uint32 NumPrimitives, uint32 NumInstances) final override;
+	virtual void RHIDrawIndexedPrimitiveIndirect(FRHIIndexBuffer* IndexBuffer, FRHIVertexBuffer* ArgumentBuffer, uint32 ArgumentOffset) final override;
 	virtual void RHISetDepthBounds(float MinDepth, float MaxDepth) final override;
 	virtual void RHIPushEvent(const TCHAR* Name, FColor Color) final override;
 	virtual void RHIPopEvent() final override;
@@ -221,28 +221,28 @@ public:
 	virtual void RHISetComputePipelineState(FRHIComputePipelineState* ComputePipelineState) final override;
 	virtual void RHIWaitComputeFence(FRHIComputeFence* InFence) final override;
 	virtual void RHIDispatchComputeShader(uint32 ThreadGroupCountX, uint32 ThreadGroupCountY, uint32 ThreadGroupCountZ) final override;
-	virtual void RHIDispatchIndirectComputeShader(FVertexBufferRHIParamRef ArgumentBuffer, uint32 ArgumentOffset) final override;
+	virtual void RHIDispatchIndirectComputeShader(FRHIVertexBuffer* ArgumentBuffer, uint32 ArgumentOffset) final override;
 	virtual void RHIAutomaticCacheFlushAfterComputeShader(bool bEnable) final override;
 
 	virtual void RHIFlushComputeShaderCache() final override;
 	virtual void RHISetMultipleViewports(uint32 Count, const FViewportBounds* Data) final override;
 	virtual void RHIClearTinyUAV(FRHIUnorderedAccessView* UnorderedAccessViewRHI, const uint32* Values) final override;
-	virtual void RHICopyToResolveTarget(FTextureRHIParamRef SourceTexture, FTextureRHIParamRef DestTexture, const FResolveParams& ResolveParams) final override;
-	virtual void RHICopyTexture(FTextureRHIParamRef SourceTexture, FTextureRHIParamRef DestTexture, const FRHICopyTextureInfo& CopyInfo) final override;
-	virtual void RHITransitionResources(EResourceTransitionAccess TransitionType, FTextureRHIParamRef* InRenderTargets, int32 NumTextures) final override;
+	virtual void RHICopyToResolveTarget(FRHITexture* SourceTexture, FRHITexture* DestTexture, const FResolveParams& ResolveParams) final override;
+	virtual void RHICopyTexture(FRHITexture* SourceTexture, FRHITexture* DestTexture, const FRHICopyTextureInfo& CopyInfo) final override;
+	virtual void RHITransitionResources(EResourceTransitionAccess TransitionType, FRHITexture** InRenderTargets, int32 NumTextures) final override;
 	virtual void RHITransitionResources(EResourceTransitionAccess TransitionType, EResourceTransitionPipeline TransitionPipeline, FRHIUnorderedAccessView** InUAVs, int32 NumUAVs, FRHIComputeFence* WriteComputeFence) final override;
-	virtual void RHICopyToStagingBuffer(FVertexBufferRHIParamRef SourceBuffer, FRHIStagingBuffer* DestinationStagingBuffer, uint32 Offset, uint32 NumBytes) final override;
+	virtual void RHICopyToStagingBuffer(FRHIVertexBuffer* SourceBuffer, FRHIStagingBuffer* DestinationStagingBuffer, uint32 Offset, uint32 NumBytes) final override;
 	virtual void RHIWriteGPUFence(FRHIGPUFence* Fence) final override;
 
 	// Render time measurement
 	virtual void RHIBeginRenderQuery(FRHIRenderQuery* RenderQuery) final override;
 	virtual void RHIEndRenderQuery(FRHIRenderQuery* RenderQuery) final override;
 
-	virtual void RHIUpdateTextureReference(FRHITextureReference* TextureRef, FTextureRHIParamRef NewTexture) final override;
+	virtual void RHIUpdateTextureReference(FRHITextureReference* TextureRef, FRHITexture* NewTexture) final override;
 
 	virtual void RHISubmitCommandsHint() final override;
 
-	virtual void RHIBeginDrawingViewport(FRHIViewport* Viewport, FTextureRHIParamRef RenderTargetRHI) final override;
+	virtual void RHIBeginDrawingViewport(FRHIViewport* Viewport, FRHITexture* RenderTargetRHI) final override;
 	virtual void RHIEndDrawingViewport(FRHIViewport* Viewport, bool bPresent, bool bLockToVsync) final override;
 
 	virtual void RHIBeginFrame() final override;
diff --git a/Engine/Source/Runtime/VulkanRHI/Private/VulkanIndexBuffer.cpp b/Engine/Source/Runtime/VulkanRHI/Private/VulkanIndexBuffer.cpp
--- a/Engine/Source/Runtime/VulkanRHI/Private/VulkanIndexBuffer.cpp
+++ b/Engine/Source/Runtime/VulkanRHI/Private/VulkanIndexBuffer.cpp
@@ -364,33 +364,33 @@ FIndexBufferRHIRef FVulkanDynamicRHI::RHICreateIndexBuffer(uint32 Stride, uint32
 	return new FVulkanIndexBuffer(Device, Stride, Size, InUsage, CreateInfo, nullptr);
 }
 
-void* FVulkanDynamicRHI::RHILockIndexBuffer(FIndexBufferRHIParamRef IndexBufferRHI, uint32 Offset, uint32 Size, EResourceLockMode LockMode)
+void* FVulkanDynamicRHI::RHILockIndexBuffer(FRHIIndexBuffer* IndexBufferRHI, uint32 Offset, uint32 Size, EResourceLockMode LockMode)
 {
 	FVulkanIndexBuffer* IndexBuffer = ResourceCast(IndexBufferRHI);
 	return IndexBuffer->Lock(false, LockMode, Size, Offset);
 }
 
 #if VULKAN_BUFFER_LOCK_THREADSAFE
-void* FVulkanDynamicRHI::LockIndexBuffer_RenderThread(class FRHICommandListImmediate& RHICmdList, FIndexBufferRHIParamRef IndexBufferRHI, uint32 Offset, uint32 SizeRHI, EResourceLockMode LockMode)
+void* FVulkanDynamicRHI::LockIndexBuffer_RenderThread(class FRHICommandListImmediate& RHICmdList, FRHIIndexBuffer* IndexBufferRHI, uint32 Offset, uint32 SizeRHI, EResourceLockMode LockMode)
 {
 	return this->RHILockIndexBuffer(IndexBufferRHI, Offset, SizeRHI, LockMode);
 }
 #endif
 
-void FVulkanDynamicRHI::RHIUnlockIndexBuffer(FIndexBufferRHIParamRef IndexBufferRHI)
+void FVulkanDynamicRHI::RHIUnlockIndexBuffer(FRHIIndexBuffer* IndexBufferRHI)
 {
 	FVulkanIndexBuffer* IndexBuffer = ResourceCast(IndexBufferRHI);
 	IndexBuffer->Unlock(false);
 }
 
 #if VULKAN_BUFFER_LOCK_THREADSAFE
-void FVulkanDynamicRHI::UnlockIndexBuffer_RenderThread(FRHICommandListImmediate& RHICmdList, FIndexBufferRHIParamRef IndexBufferRHI)
+void FVulkanDynamicRHI::UnlockIndexBuffer_RenderThread(FRHICommandListImmediate& RHICmdList, FRHIIndexBuffer* IndexBufferRHI)
 {
 	this->RHIUnlockIndexBuffer(IndexBufferRHI);
 }
 #endif
 
-void FVulkanDynamicRHI::RHITransferIndexBufferUnderlyingResource(FIndexBufferRHIParamRef DestIndexBuffer, FIndexBufferRHIParamRef SrcIndexBuffer)
+void FVulkanDynamicRHI::RHITransferIndexBufferUnderlyingResource(FRHIIndexBuffer* DestIndexBuffer, FRHIIndexBuffer* SrcIndexBuffer)
 {
 	check(DestIndexBuffer);
 	FVulkanIndexBuffer* Dest = ResourceCast(DestIndexBuffer);
diff --git a/Engine/Source/Runtime/VulkanRHI/Private/VulkanRHI.cpp b/Engine/Source/Runtime/VulkanRHI/Private/VulkanRHI.cpp
--- a/Engine/Source/Runtime/VulkanRHI/Private/VulkanRHI.cpp
+++ b/Engine/Source/Runtime/VulkanRHI/Private/VulkanRHI.cpp
@@ -804,7 +804,7 @@ void FVulkanCommandListContext::RHIEndScene()
 	//FRCLog::Printf(FString::Printf(TEXT("FVulkanCommandListContext::RHIEndScene()")));
 }
 
-void FVulkanCommandListContext::RHIBeginDrawingViewport(FRHIViewport* ViewportRHI, FTextureRHIParamRef RenderTargetRHI)
+void FVulkanCommandListContext::RHIBeginDrawingViewport(FRHIViewport* ViewportRHI, FRHITexture* RenderTargetRHI)
 {
 	//FRCLog::Printf(FString::Printf(TEXT("FVulkanCommandListContext::RHIBeginDrawingViewport\n")));
 	check(ViewportRHI);
@@ -1038,7 +1038,7 @@ FTextureCubeRHIRef FVulkanDynamicRHI::RHICreateTextureCubeFromResource(EPixelFor
 	return new FVulkanTextureCube(*Device, Format, Size, bArray, ArraySize, NumMips, Resource, Flags, nullptr, ClearValueBinding);
 }
 
-void FVulkanDynamicRHI::RHIAliasTextureResources(FTextureRHIParamRef DestTextureRHI, FTextureRHIParamRef SrcTextureRHI)
+void FVulkanDynamicRHI::RHIAliasTextureResources(FRHITexture* DestTextureRHI, FRHITexture* SrcTextureRHI)
 {
 	if (DestTextureRHI && SrcTextureRHI)
 	{
@@ -1052,7 +1052,7 @@ void FVulkanDynamicRHI::RHIAliasTextureResources(FTextureRHIParamRef DestTexture
 	}
 }
 
-void FVulkanDynamicRHI::RHICopySubTextureRegion(FTexture2DRHIParamRef SourceTexture, FTexture2DRHIParamRef DestinationTexture, FBox2D SourceBox, FBox2D DestinationBox)
+void FVulkanDynamicRHI::RHICopySubTextureRegion(FRHITexture2D* SourceTexture, FRHITexture2D* DestinationTexture, FBox2D SourceBox, FBox2D DestinationBox)
 {
 	FRHICopyTextureInfo CopyInfo;
 
@@ -1579,7 +1579,7 @@ void VulkanSetImageLayout(
 	VulkanRHI::vkCmdPipelineBarrier(CmdBuffer, SourceStages, DestStages, 0, 0, nullptr, 0, nullptr, 1, &ImageBarrier);
 }
 
-void VulkanResolveImage(VkCommandBuffer Cmd, FTextureRHIParamRef SourceTextureRHI, FTextureRHIParamRef DestTextureRHI)
+void VulkanResolveImage(VkCommandBuffer Cmd, FRHITexture* SourceTextureRHI, FRHITexture* DestTextureRHI)
 {
 	FVulkanTextureBase* Src = FVulkanTextureBase::Cast(SourceTextureRHI);
 	FVulkanTextureBase* Dst = FVulkanTextureBase::Cast(DestTextureRHI);
diff --git a/Engine/Source/Runtime/VulkanRHI/Private/VulkanRHIPrivate.h b/Engine/Source/Runtime/VulkanRHI/Private/VulkanRHIPrivate.h
--- a/Engine/Source/Runtime/VulkanRHI/Private/VulkanRHIPrivate.h
+++ b/Engine/Source/Runtime/VulkanRHI/Private/VulkanRHIPrivate.h
@@ -372,7 +372,7 @@ inline void VulkanSetImageLayoutSimple(VkCommandBuffer CmdBuffer, VkImage Image,
 	VulkanSetImageLayout(CmdBuffer, Image, OldLayout, NewLayout, SubresourceRange);
 }
 
-void VulkanResolveImage(VkCommandBuffer Cmd, FTextureRHIParamRef SourceTextureRHI, FTextureRHIParamRef DestTextureRHI);
+void VulkanResolveImage(VkCommandBuffer Cmd, FRHITexture* SourceTextureRHI, FRHITexture* DestTextureRHI);
 
 // Stats
 DECLARE_STATS_GROUP(TEXT("Vulkan RHI"), STATGROUP_VulkanRHI, STATCAT_Advanced);
diff --git a/Engine/Source/Runtime/VulkanRHI/Private/VulkanRenderTarget.cpp b/Engine/Source/Runtime/VulkanRHI/Private/VulkanRenderTarget.cpp
--- a/Engine/Source/Runtime/VulkanRHI/Private/VulkanRenderTarget.cpp
+++ b/Engine/Source/Runtime/VulkanRHI/Private/VulkanRenderTarget.cpp
@@ -141,7 +141,7 @@ void FTransitionAndLayoutManager::BeginEmulatedRenderPass(FVulkanCommandListCont
 	int32 Index = 0;
 	for (Index = 0; Index < RenderTargetsInfo.NumColorRenderTargets; ++Index)
 	{
-		FTextureRHIParamRef Texture = RenderTargetsInfo.ColorRenderTarget[Index].Texture;
+		FRHITexture* Texture = RenderTargetsInfo.ColorRenderTarget[Index].Texture;
 		if (Texture)
 		{
 			FVulkanSurface& Surface = FVulkanTextureBase::Cast(Texture)->Surface;
@@ -177,7 +177,7 @@ void FTransitionAndLayoutManager::BeginEmulatedRenderPass(FVulkanCommandListCont
 
 	if (RenderTargetsInfo.DepthStencilRenderTarget.Texture)
 	{
-		FTextureRHIParamRef DSTexture = RenderTargetsInfo.DepthStencilRenderTarget.Texture;
+		FRHITexture* DSTexture = RenderTargetsInfo.DepthStencilRenderTarget.Texture;
 		FVulkanSurface& Surface = FVulkanTextureBase::Cast(DSTexture)->Surface;
 		VkImageLayout& DSLayout = Layouts.FindOrAdd(Surface.Image);
 		FExclusiveDepthStencil RequestedDSAccess = RenderTargetsInfo.DepthStencilRenderTarget.GetDepthStencilAccess();
@@ -246,7 +246,7 @@ void FTransitionAndLayoutManager::BeginRealRenderPass(FVulkanCommandListContext&
 
 	for (Index = 0; Index < NumColorTargets; ++Index)
 	{
-		FTextureRHIParamRef Texture = RPInfo.ColorRenderTargets[Index].RenderTarget;
+		FRHITexture* Texture = RPInfo.ColorRenderTargets[Index].RenderTarget;
 		CA_ASSUME(Texture);
 		FVulkanSurface& Surface = FVulkanTextureBase::Cast(Texture)->Surface;
 		check(Surface.Image != VK_NULL_HANDLE);
@@ -343,7 +343,7 @@ void FTransitionAndLayoutManager::BeginRealRenderPass(FVulkanCommandListContext&
 		}
 	}
 
-	FTextureRHIParamRef DSTexture = RPInfo.DepthStencilRenderTarget.DepthStencilTarget;
+	FRHITexture* DSTexture = RPInfo.DepthStencilRenderTarget.DepthStencilTarget;
 	if (DSTexture)
 	{
 		FVulkanSurface& Surface = FVulkanTextureBase::Cast(DSTexture)->Surface;
@@ -483,7 +483,7 @@ void FVulkanCommandListContext::RHISetRenderTargets(uint32 NumSimultaneousRender
 	}
 	else
 	{
-		DepthView = FRHIDepthRenderTargetView(FTextureRHIParamRef(), ERenderTargetLoadAction::ENoAction, ERenderTargetStoreAction::ENoAction, ERenderTargetLoadAction::ENoAction, ERenderTargetStoreAction::ENoAction);
+		DepthView = FRHIDepthRenderTargetView(nullptr, ERenderTargetLoadAction::ENoAction, ERenderTargetStoreAction::ENoAction, ERenderTargetLoadAction::ENoAction, ERenderTargetStoreAction::ENoAction);
 	}
 
 	if (NumSimultaneousRenderTargets == 1 && (!NewRenderTargets || !NewRenderTargets->Texture))
@@ -620,7 +620,7 @@ void FVulkanCommandListContext::RHISetRenderTargetsAndClear(const FRHISetRenderT
 	}
 }
 
-void FVulkanCommandListContext::RHICopyToResolveTarget(FTextureRHIParamRef SourceTextureRHI, FTextureRHIParamRef DestTextureRHI, const FResolveParams& InResolveParams)
+void FVulkanCommandListContext::RHICopyToResolveTarget(FRHITexture* SourceTextureRHI, FRHITexture* DestTextureRHI, const FResolveParams& InResolveParams)
 {
 	//FRCLog::Printf(FString::Printf(TEXT("RHICopyToResolveTarget")));
 	if (!SourceTextureRHI || !DestTextureRHI)
@@ -751,7 +751,7 @@ void FVulkanCommandListContext::RHICopyToResolveTarget(FTextureRHIParamRef Sourc
 	}
 }
 
-void FVulkanDynamicRHI::RHIReadSurfaceData(FTextureRHIParamRef TextureRHI, FIntRect Rect, TArray<FColor>& OutData, FReadSurfaceDataFlags InFlags)
+void FVulkanDynamicRHI::RHIReadSurfaceData(FRHITexture* TextureRHI, FIntRect Rect, TArray<FColor>& OutData, FReadSurfaceDataFlags InFlags)
 {
 	FRHITexture2D* TextureRHI2D = TextureRHI->GetTexture2D();
 	check(TextureRHI2D);
@@ -935,7 +935,7 @@ void FVulkanDynamicRHI::RHIReadSurfaceData(FTextureRHIParamRef TextureRHI, FIntR
 	ImmediateContext.GetCommandBufferManager()->PrepareForNewActiveCommandBuffer();
 }
 
-void FVulkanDynamicRHI::RHIReadSurfaceData(FTextureRHIParamRef TextureRHI, FIntRect Rect, TArray<FLinearColor>& OutData, FReadSurfaceDataFlags InFlags)
+void FVulkanDynamicRHI::RHIReadSurfaceData(FRHITexture* TextureRHI, FIntRect Rect, TArray<FLinearColor>& OutData, FReadSurfaceDataFlags InFlags)
 {
 	TArray<FColor> FromColorData;
 	RHIReadSurfaceData(TextureRHI, Rect, FromColorData, InFlags);
@@ -945,7 +945,7 @@ void FVulkanDynamicRHI::RHIReadSurfaceData(FTextureRHIParamRef TextureRHI, FIntR
 	}
 }
 
-void FVulkanDynamicRHI::RHIMapStagingSurface(FTextureRHIParamRef TextureRHI,void*& OutData,int32& OutWidth,int32& OutHeight)
+void FVulkanDynamicRHI::RHIMapStagingSurface(FRHITexture* TextureRHI,void*& OutData,int32& OutWidth,int32& OutHeight)
 {
 	FRHITexture2D* TextureRHI2D = TextureRHI->GetTexture2D();
 	check(TextureRHI2D);
@@ -960,23 +960,23 @@ void FVulkanDynamicRHI::RHIMapStagingSurface(FTextureRHIParamRef TextureRHI,void
 }
 
 
-void FVulkanDynamicRHI::RHIMapStagingSurface_RenderThread(class FRHICommandListImmediate& RHICmdList, FTextureRHIParamRef Texture, void*& OutData, int32& OutWidth, int32& OutHeight)
+void FVulkanDynamicRHI::RHIMapStagingSurface_RenderThread(class FRHICommandListImmediate& RHICmdList, FRHITexture* Texture, void*& OutData, int32& OutWidth, int32& OutHeight)
 {
 	RHIMapStagingSurface(Texture, OutData, OutWidth, OutHeight);
 }
 
-void FVulkanDynamicRHI::RHIUnmapStagingSurface_RenderThread(class FRHICommandListImmediate& RHICmdList, FTextureRHIParamRef Texture)
+void FVulkanDynamicRHI::RHIUnmapStagingSurface_RenderThread(class FRHICommandListImmediate& RHICmdList, FRHITexture* Texture)
 {
 	RHIUnmapStagingSurface(Texture);
 }
 
 
 
-void FVulkanDynamicRHI::RHIUnmapStagingSurface(FTextureRHIParamRef TextureRHI)
+void FVulkanDynamicRHI::RHIUnmapStagingSurface(FRHITexture* TextureRHI)
 {
 }
 
-void FVulkanDynamicRHI::RHIReadSurfaceFloatData(FTextureRHIParamRef TextureRHI, FIntRect Rect, TArray<FFloat16Color>& OutData, ECubeFace CubeFace,int32 ArrayIndex,int32 MipIndex)
+void FVulkanDynamicRHI::RHIReadSurfaceFloatData(FRHITexture* TextureRHI, FIntRect Rect, TArray<FFloat16Color>& OutData, ECubeFace CubeFace,int32 ArrayIndex,int32 MipIndex)
 {
 	auto DoCopyFloat = [](FVulkanDevice* InDevice, FVulkanCmdBuffer* InCmdBuffer, const FVulkanSurface& Surface, uint32 InMipIndex, uint32 SrcBaseArrayLayer, FIntRect InRect, TArray<FFloat16Color>& OutputData)
 	{
@@ -1097,7 +1097,7 @@ void FVulkanDynamicRHI::RHIReadSurfaceFloatData(FTextureRHIParamRef TextureRHI,
 	}
 }
 
-void FVulkanDynamicRHI::RHIRead3DSurfaceFloatData(FTextureRHIParamRef TextureRHI,FIntRect InRect,FIntPoint ZMinMax,TArray<FFloat16Color>& OutData)
+void FVulkanDynamicRHI::RHIRead3DSurfaceFloatData(FRHITexture* TextureRHI,FIntRect InRect,FIntPoint ZMinMax,TArray<FFloat16Color>& OutData)
 {
 	FRHITexture3D* TextureRHI3D = TextureRHI->GetTexture3D();
 	check(TextureRHI3D);
@@ -1224,7 +1224,7 @@ void FVulkanCommandListContext::RHITransitionResources(EResourceTransitionAccess
 	}
 }
 
-void FVulkanCommandListContext::RHITransitionResources(EResourceTransitionAccess TransitionType, FTextureRHIParamRef* InTextures, int32 NumTextures)
+void FVulkanCommandListContext::RHITransitionResources(EResourceTransitionAccess TransitionType, FRHITexture** InTextures, int32 NumTextures)
 {
 	if (NumTextures > 0)
 	{
diff --git a/Engine/Source/Runtime/VulkanRHI/Private/VulkanStructuredBuffer.cpp b/Engine/Source/Runtime/VulkanRHI/Private/VulkanStructuredBuffer.cpp
--- a/Engine/Source/Runtime/VulkanRHI/Private/VulkanStructuredBuffer.cpp
+++ b/Engine/Source/Runtime/VulkanRHI/Private/VulkanStructuredBuffer.cpp
@@ -20,14 +20,14 @@ FStructuredBufferRHIRef FVulkanDynamicRHI::RHICreateStructuredBuffer(uint32 InSt
 	return new FVulkanStructuredBuffer(Device, InStride, InSize, CreateInfo, InUsage);
 }
 
-void* FVulkanDynamicRHI::RHILockStructuredBuffer(FStructuredBufferRHIParamRef StructuredBufferRHI,uint32 Offset,uint32 Size,EResourceLockMode LockMode)
+void* FVulkanDynamicRHI::RHILockStructuredBuffer(FRHIStructuredBuffer* StructuredBufferRHI,uint32 Offset,uint32 Size,EResourceLockMode LockMode)
 {
 	FVulkanStructuredBuffer* StructuredBuffer = ResourceCast(StructuredBufferRHI);
 
 	return StructuredBuffer->Lock(false, LockMode, Size, Offset);
 }
 
-void FVulkanDynamicRHI::RHIUnlockStructuredBuffer(FStructuredBufferRHIParamRef StructuredBufferRHI)
+void FVulkanDynamicRHI::RHIUnlockStructuredBuffer(FRHIStructuredBuffer* StructuredBufferRHI)
 {
 	FVulkanStructuredBuffer* StructuredBuffer = ResourceCast(StructuredBufferRHI);
 
diff --git a/Engine/Source/Runtime/VulkanRHI/Private/VulkanTexture.cpp b/Engine/Source/Runtime/VulkanRHI/Private/VulkanTexture.cpp
--- a/Engine/Source/Runtime/VulkanRHI/Private/VulkanTexture.cpp
+++ b/Engine/Source/Runtime/VulkanRHI/Private/VulkanTexture.cpp
@@ -921,7 +921,7 @@ bool FVulkanDynamicRHI::RHIGetTextureMemoryVisualizeData( FColor* /*TextureData*
 	return false;
 }
 
-uint32 FVulkanDynamicRHI::RHIComputeMemorySize(FTextureRHIParamRef TextureRHI)
+uint32 FVulkanDynamicRHI::RHIComputeMemorySize(FRHITexture* TextureRHI)
 {
 	if(!TextureRHI)
 	{
@@ -949,7 +949,7 @@ FTexture2DRHIRef FVulkanDynamicRHI::RHIAsyncCreateTexture2D(uint32 SizeX,uint32
 	return FTexture2DRHIRef();
 }
 
-void FVulkanDynamicRHI::RHICopySharedMips(FTexture2DRHIParamRef DestTexture2D,FTexture2DRHIParamRef SrcTexture2D)
+void FVulkanDynamicRHI::RHICopySharedMips(FRHITexture2D* DestTexture2D, FRHITexture2D* SrcTexture2D)
 {
 	VULKAN_SIGNAL_UNIMPLEMENTED();
 }
@@ -968,7 +968,7 @@ FTexture3DRHIRef FVulkanDynamicRHI::RHICreateTexture3D(uint32 SizeX, uint32 Size
 	return Tex3d;
 }
 
-void FVulkanDynamicRHI::RHIGetResourceInfo(FTextureRHIParamRef Ref, FRHIResourceInfo& OutInfo)
+void FVulkanDynamicRHI::RHIGetResourceInfo(FRHITexture* Ref, FRHIResourceInfo& OutInfo)
 {
 	FVulkanTextureBase* Base = (FVulkanTextureBase*)Ref->GetTextureBaseRHI();
 	OutInfo.VRamAllocation.AllocationSize = Base->Surface.GetMemorySize();
@@ -1083,7 +1083,7 @@ struct FRHICommandVulkanAsyncReallocateTexture2D final : public FRHICommand<FRHI
 	}
 };
 
-FTexture2DRHIRef FVulkanDynamicRHI::AsyncReallocateTexture2D_RenderThread(class FRHICommandListImmediate& RHICmdList, FTexture2DRHIParamRef OldTextureRHI, int32 NewMipCount, int32 NewSizeX, int32 NewSizeY, FThreadSafeCounter* RequestStatus)
+FTexture2DRHIRef FVulkanDynamicRHI::AsyncReallocateTexture2D_RenderThread(class FRHICommandListImmediate& RHICmdList, FRHITexture2D* OldTextureRHI, int32 NewMipCount, int32 NewSizeX, int32 NewSizeY, FThreadSafeCounter* RequestStatus)
 {
 	LLM_SCOPE_VULKAN(ELLMTagVulkan::VulkanTextures);
 	if (RHICmdList.Bypass())
@@ -1101,7 +1101,7 @@ FTexture2DRHIRef FVulkanDynamicRHI::AsyncReallocateTexture2D_RenderThread(class
 	return NewTexture;
 }
 
-FTexture2DRHIRef FVulkanDynamicRHI::RHIAsyncReallocateTexture2D(FTexture2DRHIParamRef OldTextureRHI, int32 NewMipCount, int32 NewSizeX, int32 NewSizeY, FThreadSafeCounter* RequestStatus)
+FTexture2DRHIRef FVulkanDynamicRHI::RHIAsyncReallocateTexture2D(FRHITexture2D* OldTextureRHI, int32 NewMipCount, int32 NewSizeX, int32 NewSizeY, FThreadSafeCounter* RequestStatus)
 {
 	LLM_SCOPE_VULKAN(ELLMTagVulkan::VulkanTextures);
 	FVulkanTexture2D* OldTexture = ResourceCast(OldTextureRHI);
@@ -1114,17 +1114,17 @@ FTexture2DRHIRef FVulkanDynamicRHI::RHIAsyncReallocateTexture2D(FTexture2DRHIPar
 	return NewTexture;
 }
 
-ETextureReallocationStatus FVulkanDynamicRHI::RHIFinalizeAsyncReallocateTexture2D(FTexture2DRHIParamRef Texture2D, bool bBlockUntilCompleted)
+ETextureReallocationStatus FVulkanDynamicRHI::RHIFinalizeAsyncReallocateTexture2D(FRHITexture2D* Texture2D, bool bBlockUntilCompleted)
 {
 	return TexRealloc_Succeeded;
 }
 
-ETextureReallocationStatus FVulkanDynamicRHI::RHICancelAsyncReallocateTexture2D(FTexture2DRHIParamRef Texture2D, bool bBlockUntilCompleted)
+ETextureReallocationStatus FVulkanDynamicRHI::RHICancelAsyncReallocateTexture2D(FRHITexture2D* Texture2D, bool bBlockUntilCompleted)
 {
 	return TexRealloc_Succeeded;
 }
 
-void* FVulkanDynamicRHI::RHILockTexture2D(FTexture2DRHIParamRef TextureRHI,uint32 MipIndex,EResourceLockMode LockMode,uint32& DestStride,bool bLockWithinMiptail)
+void* FVulkanDynamicRHI::RHILockTexture2D(FRHITexture2D* TextureRHI,uint32 MipIndex,EResourceLockMode LockMode,uint32& DestStride,bool bLockWithinMiptail)
 {
 	LLM_SCOPE_VULKAN(ELLMTagVulkan::VulkanTextures);
 	FVulkanTexture2D* Texture = ResourceCast(TextureRHI);
@@ -1150,7 +1150,7 @@ void* FVulkanDynamicRHI::RHILockTexture2D(FTexture2DRHIParamRef TextureRHI,uint3
 	return Data;
 }
 
-void FVulkanDynamicRHI::InternalUnlockTexture2D(bool bFromRenderingThread, FTexture2DRHIParamRef TextureRHI,uint32 MipIndex,bool bLockWithinMiptail)
+void FVulkanDynamicRHI::InternalUnlockTexture2D(bool bFromRenderingThread, FRHITexture2D* TextureRHI,uint32 MipIndex,bool bLockWithinMiptail)
 {
 	LLM_SCOPE_VULKAN(ELLMTagVulkan::VulkanTextures);
 	FVulkanTexture2D* Texture = ResourceCast(TextureRHI);
@@ -1204,7 +1204,7 @@ void FVulkanDynamicRHI::InternalUnlockTexture2D(bool bFromRenderingThread, FText
 	}
 }
 
-void* FVulkanDynamicRHI::RHILockTexture2DArray(FTexture2DArrayRHIParamRef TextureRHI,uint32 TextureIndex,uint32 MipIndex,EResourceLockMode LockMode,uint32& DestStride,bool bLockWithinMiptail)
+void* FVulkanDynamicRHI::RHILockTexture2DArray(FRHITexture2DArray* TextureRHI,uint32 TextureIndex,uint32 MipIndex,EResourceLockMode LockMode,uint32& DestStride,bool bLockWithinMiptail)
 {
 	LLM_SCOPE_VULKAN(ELLMTagVulkan::VulkanTextures);
 	FVulkanTexture2DArray* Texture = ResourceCast(TextureRHI);
@@ -1227,7 +1227,7 @@ void* FVulkanDynamicRHI::RHILockTexture2DArray(FTexture2DArrayRHIParamRef Textur
 	return Data;
 }
 
-void FVulkanDynamicRHI::RHIUnlockTexture2DArray(FTexture2DArrayRHIParamRef TextureRHI,uint32 TextureIndex,uint32 MipIndex,bool bLockWithinMiptail)
+void FVulkanDynamicRHI::RHIUnlockTexture2DArray(FRHITexture2DArray* TextureRHI,uint32 TextureIndex,uint32 MipIndex,bool bLockWithinMiptail)
 {
 	LLM_SCOPE_VULKAN(ELLMTagVulkan::VulkanTextures);
 	FVulkanTexture2DArray* Texture = ResourceCast(TextureRHI);
@@ -1280,7 +1280,7 @@ void FVulkanDynamicRHI::RHIUnlockTexture2DArray(FTexture2DArrayRHIParamRef Textu
 	}
 }
 
-void FVulkanDynamicRHI::InternalUpdateTexture2D(bool bFromRenderingThread, FTexture2DRHIParamRef TextureRHI, uint32 MipIndex, const struct FUpdateTextureRegion2D& UpdateRegion, uint32 SourceRowPitch, const uint8* SourceData)
+void FVulkanDynamicRHI::InternalUpdateTexture2D(bool bFromRenderingThread, FRHITexture2D* TextureRHI, uint32 MipIndex, const struct FUpdateTextureRegion2D& UpdateRegion, uint32 SourceRowPitch, const uint8* SourceData)
 {
 	LLM_SCOPE_VULKAN(ELLMTagVulkan::VulkanTextures);
 	FVulkanTexture2D* Texture = ResourceCast(TextureRHI);
@@ -1355,7 +1355,7 @@ void FVulkanDynamicRHI::InternalUpdateTexture2D(bool bFromRenderingThread, FText
 	}
 }
 
-void FVulkanDynamicRHI::InternalUpdateTexture3D(bool bFromRenderingThread, FTexture3DRHIParamRef TextureRHI, uint32 MipIndex, const FUpdateTextureRegion3D& UpdateRegion, uint32 SourceRowPitch, uint32 SourceDepthPitch, const uint8* SourceData)
+void FVulkanDynamicRHI::InternalUpdateTexture3D(bool bFromRenderingThread, FRHITexture3D* TextureRHI, uint32 MipIndex, const FUpdateTextureRegion3D& UpdateRegion, uint32 SourceRowPitch, uint32 SourceDepthPitch, const uint8* SourceData)
 {
 	LLM_SCOPE_VULKAN(ELLMTagVulkan::VulkanTextures);
 	FVulkanTexture3D* Texture = ResourceCast(TextureRHI);
@@ -1890,7 +1890,7 @@ FTextureCubeRHIRef FVulkanDynamicRHI::RHICreateTextureCubeArray(uint32 Size, uin
 	return new FVulkanTextureCube(*Device, (EPixelFormat)Format, Size, true, ArraySize, NumMips, Flags, CreateInfo.BulkData, CreateInfo.ClearValueBinding);
 }
 
-void* FVulkanDynamicRHI::RHILockTextureCubeFace(FTextureCubeRHIParamRef TextureCubeRHI, uint32 FaceIndex, uint32 ArrayIndex, uint32 MipIndex, EResourceLockMode LockMode, uint32& DestStride, bool bLockWithinMiptail)
+void* FVulkanDynamicRHI::RHILockTextureCubeFace(FRHITextureCube* TextureCubeRHI, uint32 FaceIndex, uint32 ArrayIndex, uint32 MipIndex, EResourceLockMode LockMode, uint32& DestStride, bool bLockWithinMiptail)
 {
 	LLM_SCOPE_VULKAN(ELLMTagVulkan::VulkanTextures);
 	FVulkanTextureCube* Texture = ResourceCast(TextureCubeRHI);
@@ -1913,7 +1913,7 @@ void* FVulkanDynamicRHI::RHILockTextureCubeFace(FTextureCubeRHIParamRef TextureC
 	return Data;
 }
 
-void FVulkanDynamicRHI::RHIUnlockTextureCubeFace(FTextureCubeRHIParamRef TextureCubeRHI,uint32 FaceIndex,uint32 ArrayIndex,uint32 MipIndex,bool bLockWithinMiptail)
+void FVulkanDynamicRHI::RHIUnlockTextureCubeFace(FRHITextureCube* TextureCubeRHI,uint32 FaceIndex,uint32 ArrayIndex,uint32 MipIndex,bool bLockWithinMiptail)
 {
 	LLM_SCOPE_VULKAN(ELLMTagVulkan::VulkanTextures);
 	FVulkanTextureCube* Texture = ResourceCast(TextureCubeRHI);
@@ -1966,7 +1966,7 @@ void FVulkanDynamicRHI::RHIUnlockTextureCubeFace(FTextureCubeRHIParamRef Texture
 	}
 }
 
-void FVulkanDynamicRHI::RHIBindDebugLabelName(FTextureRHIParamRef TextureRHI, const TCHAR* Name)
+void FVulkanDynamicRHI::RHIBindDebugLabelName(FRHITexture* TextureRHI, const TCHAR* Name)
 {
 #if VULKAN_ENABLE_IMAGE_TRACKING_LAYER
 	{
@@ -2030,12 +2030,12 @@ void FVulkanDynamicRHI::RHIBindDebugLabelName(FRHIUnorderedAccessView* Unordered
 }
 
 
-void FVulkanDynamicRHI::RHIVirtualTextureSetFirstMipInMemory(FTexture2DRHIParamRef TextureRHI, uint32 FirstMip)
+void FVulkanDynamicRHI::RHIVirtualTextureSetFirstMipInMemory(FRHITexture2D* TextureRHI, uint32 FirstMip)
 {
 	VULKAN_SIGNAL_UNIMPLEMENTED();
 }
 
-void FVulkanDynamicRHI::RHIVirtualTextureSetFirstMipVisible(FTexture2DRHIParamRef TextureRHI, uint32 FirstMip)
+void FVulkanDynamicRHI::RHIVirtualTextureSetFirstMipVisible(FRHITexture2D* TextureRHI, uint32 FirstMip)
 {
 	VULKAN_SIGNAL_UNIMPLEMENTED();
 }
@@ -2121,7 +2121,7 @@ FTextureReferenceRHIRef FVulkanDynamicRHI::RHICreateTextureReference(FLastRender
 	return new FVulkanTextureReference(*Device, LastRenderTime);
 }
 
-void FVulkanCommandListContext::RHIUpdateTextureReference(FRHITextureReference* TextureRef, FTextureRHIParamRef NewTexture)
+void FVulkanCommandListContext::RHIUpdateTextureReference(FRHITextureReference* TextureRef, FRHITexture* NewTexture)
 {
 	//#todo-rco: Implementation needs to be verified
 	FVulkanTextureReference* VulkanTextureRef = (FVulkanTextureReference*)TextureRef;
@@ -2131,7 +2131,7 @@ void FVulkanCommandListContext::RHIUpdateTextureReference(FRHITextureReference*
 	}
 }
 
-void FVulkanCommandListContext::RHICopyTexture(FTextureRHIParamRef SourceTexture, FTextureRHIParamRef DestTexture, const FRHICopyTextureInfo& CopyInfo)
+void FVulkanCommandListContext::RHICopyTexture(FRHITexture* SourceTexture, FRHITexture* DestTexture, const FRHICopyTextureInfo& CopyInfo)
 {
 	LLM_SCOPE_VULKAN(ELLMTagVulkan::VulkanTextures);
 	check(SourceTexture && DestTexture);
diff --git a/Engine/Source/Runtime/VulkanRHI/Private/VulkanUAV.cpp b/Engine/Source/Runtime/VulkanRHI/Private/VulkanUAV.cpp
--- a/Engine/Source/Runtime/VulkanRHI/Private/VulkanUAV.cpp
+++ b/Engine/Source/Runtime/VulkanRHI/Private/VulkanUAV.cpp
@@ -236,7 +236,7 @@ void FVulkanUnorderedAccessView::UpdateView()
 	}
 }
 
-FUnorderedAccessViewRHIRef FVulkanDynamicRHI::RHICreateUnorderedAccessView(FStructuredBufferRHIParamRef StructuredBufferRHI, bool bUseUAVCounter, bool bAppendBuffer)
+FUnorderedAccessViewRHIRef FVulkanDynamicRHI::RHICreateUnorderedAccessView(FRHIStructuredBuffer* StructuredBufferRHI, bool bUseUAVCounter, bool bAppendBuffer)
 {
 	FVulkanStructuredBuffer* StructuredBuffer = ResourceCast(StructuredBufferRHI);
 
@@ -249,7 +249,7 @@ FUnorderedAccessViewRHIRef FVulkanDynamicRHI::RHICreateUnorderedAccessView(FStru
 	return UAV;
 }
 
-FUnorderedAccessViewRHIRef FVulkanDynamicRHI::RHICreateUnorderedAccessView(FTextureRHIParamRef TextureRHI, uint32 MipLevel)
+FUnorderedAccessViewRHIRef FVulkanDynamicRHI::RHICreateUnorderedAccessView(FRHITexture* TextureRHI, uint32 MipLevel)
 {
 	FVulkanUnorderedAccessView* UAV = new FVulkanUnorderedAccessView(Device);
 	UAV->SourceTexture = TextureRHI;
@@ -257,7 +257,7 @@ FUnorderedAccessViewRHIRef FVulkanDynamicRHI::RHICreateUnorderedAccessView(FText
 	return UAV;
 }
 
-FUnorderedAccessViewRHIRef FVulkanDynamicRHI::RHICreateUnorderedAccessView(FVertexBufferRHIParamRef VertexBufferRHI, uint8 Format)
+FUnorderedAccessViewRHIRef FVulkanDynamicRHI::RHICreateUnorderedAccessView(FRHIVertexBuffer* VertexBufferRHI, uint8 Format)
 {
 	FVulkanVertexBuffer* VertexBuffer = ResourceCast(VertexBufferRHI);
 
@@ -269,7 +269,7 @@ FUnorderedAccessViewRHIRef FVulkanDynamicRHI::RHICreateUnorderedAccessView(FVert
 	return UAV;
 }
 
-FUnorderedAccessViewRHIRef FVulkanDynamicRHI::RHICreateUnorderedAccessView(FIndexBufferRHIParamRef IndexBufferRHI, uint8 Format)
+FUnorderedAccessViewRHIRef FVulkanDynamicRHI::RHICreateUnorderedAccessView(FRHIIndexBuffer* IndexBufferRHI, uint8 Format)
 {
 	FVulkanIndexBuffer* IndexBuffer = ResourceCast(IndexBufferRHI);
 
@@ -281,14 +281,14 @@ FUnorderedAccessViewRHIRef FVulkanDynamicRHI::RHICreateUnorderedAccessView(FInde
 	return UAV;
 }
 
-FShaderResourceViewRHIRef FVulkanDynamicRHI::RHICreateShaderResourceView(FStructuredBufferRHIParamRef StructuredBufferRHI)
+FShaderResourceViewRHIRef FVulkanDynamicRHI::RHICreateShaderResourceView(FRHIStructuredBuffer* StructuredBufferRHI)
 {
 	FVulkanStructuredBuffer* StructuredBuffer = ResourceCast(StructuredBufferRHI);
 	FVulkanShaderResourceView* SRV = new FVulkanShaderResourceView(Device, StructuredBuffer);
 	return SRV;
 }
 
-FShaderResourceViewRHIRef FVulkanDynamicRHI::RHICreateShaderResourceView(FVertexBufferRHIParamRef VertexBufferRHI, uint32 Stride, uint8 Format)
+FShaderResourceViewRHIRef FVulkanDynamicRHI::RHICreateShaderResourceView(FRHIVertexBuffer* VertexBufferRHI, uint32 Stride, uint8 Format)
 {	
 	if (!VertexBufferRHI)
 	{
@@ -298,13 +298,13 @@ FShaderResourceViewRHIRef FVulkanDynamicRHI::RHICreateShaderResourceView(FVertex
 	return new FVulkanShaderResourceView(Device, VertexBufferRHI, VertexBuffer, VertexBuffer->GetSize(), (EPixelFormat)Format);
 }
 
-FShaderResourceViewRHIRef FVulkanDynamicRHI::RHICreateShaderResourceView(FTextureRHIParamRef Texture, const FRHITextureSRVCreateInfo& CreateInfo)
+FShaderResourceViewRHIRef FVulkanDynamicRHI::RHICreateShaderResourceView(FRHITexture* Texture, const FRHITextureSRVCreateInfo& CreateInfo)
 {
 	FVulkanShaderResourceView* SRV = new FVulkanShaderResourceView(Device, Texture, CreateInfo);
 	return SRV;
 }
 
-FShaderResourceViewRHIRef FVulkanDynamicRHI::RHICreateShaderResourceView(FIndexBufferRHIParamRef IndexBufferRHI)
+FShaderResourceViewRHIRef FVulkanDynamicRHI::RHICreateShaderResourceView(FRHIIndexBuffer* IndexBufferRHI)
 {
 	if (!IndexBufferRHI)
 	{
@@ -317,7 +317,7 @@ FShaderResourceViewRHIRef FVulkanDynamicRHI::RHICreateShaderResourceView(FIndexB
 	return SRV;
 }
 
-void FVulkanDynamicRHI::RHIUpdateShaderResourceView(FRHIShaderResourceView* SRV, FVertexBufferRHIParamRef VertexBuffer, uint32 Stride, uint8 Format)
+void FVulkanDynamicRHI::RHIUpdateShaderResourceView(FRHIShaderResourceView* SRV, FRHIVertexBuffer* VertexBuffer, uint32 Stride, uint8 Format)
 {
 	FVulkanShaderResourceView* SRVVk = ResourceCast(SRV);
 	check(SRVVk && SRVVk->GetParent() == Device);
@@ -332,7 +332,7 @@ void FVulkanDynamicRHI::RHIUpdateShaderResourceView(FRHIShaderResourceView* SRV,
 	}
 }
 
-void FVulkanDynamicRHI::RHIUpdateShaderResourceView(FRHIShaderResourceView* SRV, FIndexBufferRHIParamRef IndexBuffer)
+void FVulkanDynamicRHI::RHIUpdateShaderResourceView(FRHIShaderResourceView* SRV, FRHIIndexBuffer* IndexBuffer)
 {
 	FVulkanShaderResourceView* SRVVk = ResourceCast(SRV);
 	check(SRVVk && SRVVk->GetParent() == Device);
diff --git a/Engine/Source/Runtime/VulkanRHI/Private/VulkanUniformBuffer.cpp b/Engine/Source/Runtime/VulkanRHI/Private/VulkanUniformBuffer.cpp
--- a/Engine/Source/Runtime/VulkanRHI/Private/VulkanUniformBuffer.cpp
+++ b/Engine/Source/Runtime/VulkanRHI/Private/VulkanUniformBuffer.cpp
@@ -260,7 +260,7 @@ inline void FVulkanDynamicRHI::UpdateUniformBuffer(FVulkanUniformBuffer* Uniform
 }
 
 
-void FVulkanDynamicRHI::RHIUpdateUniformBuffer(FUniformBufferRHIParamRef UniformBufferRHI, const void* Contents)
+void FVulkanDynamicRHI::RHIUpdateUniformBuffer(FRHIUniformBuffer* UniformBufferRHI, const void* Contents)
 {
 	static TConsoleVariableData<int32>* CVar = IConsoleManager::Get().FindTConsoleVariableDataInt(TEXT("r.Vulkan.UseRealUBs"));
 	const bool bHasRealUBs = FVulkanPlatform::UseRealUBsOptimization(CVar && CVar->GetValueOnAnyThread() > 0);
diff --git a/Engine/Source/Runtime/VulkanRHI/Private/VulkanVertexBuffer.cpp b/Engine/Source/Runtime/VulkanRHI/Private/VulkanVertexBuffer.cpp
--- a/Engine/Source/Runtime/VulkanRHI/Private/VulkanVertexBuffer.cpp
+++ b/Engine/Source/Runtime/VulkanRHI/Private/VulkanVertexBuffer.cpp
@@ -30,7 +30,7 @@ FVertexBufferRHIRef FVulkanDynamicRHI::RHICreateVertexBuffer(uint32 Size, uint32
 	return VertexBuffer;
 }
 
-void* FVulkanDynamicRHI::RHILockVertexBuffer(FVertexBufferRHIParamRef VertexBufferRHI, uint32 Offset, uint32 Size, EResourceLockMode LockMode)
+void* FVulkanDynamicRHI::RHILockVertexBuffer(FRHIVertexBuffer* VertexBufferRHI, uint32 Offset, uint32 Size, EResourceLockMode LockMode)
 {
 	LLM_SCOPE_VULKAN(ELLMTagVulkan::VulkanVertexBuffers);
 	FVulkanVertexBuffer* VertexBuffer = ResourceCast(VertexBufferRHI);
@@ -38,13 +38,13 @@ void* FVulkanDynamicRHI::RHILockVertexBuffer(FVertexBufferRHIParamRef VertexBuff
 }
 
 #if VULKAN_BUFFER_LOCK_THREADSAFE
-void* FVulkanDynamicRHI::LockVertexBuffer_RenderThread(class FRHICommandListImmediate& RHICmdList, FVertexBufferRHIParamRef VertexBufferRHI, uint32 Offset, uint32 SizeRHI, EResourceLockMode LockMode)
+void* FVulkanDynamicRHI::LockVertexBuffer_RenderThread(class FRHICommandListImmediate& RHICmdList, FRHIVertexBuffer* VertexBufferRHI, uint32 Offset, uint32 SizeRHI, EResourceLockMode LockMode)
 {
 	return this->RHILockVertexBuffer(VertexBufferRHI, Offset, SizeRHI, LockMode);
 }
 #endif
 
-void FVulkanDynamicRHI::RHIUnlockVertexBuffer(FVertexBufferRHIParamRef VertexBufferRHI)
+void FVulkanDynamicRHI::RHIUnlockVertexBuffer(FRHIVertexBuffer* VertexBufferRHI)
 {
 	LLM_SCOPE_VULKAN(ELLMTagVulkan::VulkanVertexBuffers);
 	FVulkanVertexBuffer* VertexBuffer = ResourceCast(VertexBufferRHI);
@@ -52,18 +52,18 @@ void FVulkanDynamicRHI::RHIUnlockVertexBuffer(FVertexBufferRHIParamRef VertexBuf
 }
 
 #if VULKAN_BUFFER_LOCK_THREADSAFE
-void FVulkanDynamicRHI::UnlockVertexBuffer_RenderThread(FRHICommandListImmediate& RHICmdList, FVertexBufferRHIParamRef VertexBufferRHI)
+void FVulkanDynamicRHI::UnlockVertexBuffer_RenderThread(FRHICommandListImmediate& RHICmdList, FRHIVertexBuffer* VertexBufferRHI)
 {
 	return this->RHIUnlockVertexBuffer(VertexBufferRHI);
 }
 #endif
 
-void FVulkanDynamicRHI::RHICopyVertexBuffer(FVertexBufferRHIParamRef SourceBufferRHI,FVertexBufferRHIParamRef DestBufferRHI)
+void FVulkanDynamicRHI::RHICopyVertexBuffer(FRHIVertexBuffer* SourceBufferRHI, FRHIVertexBuffer* DestBufferRHI)
 {
 	VULKAN_SIGNAL_UNIMPLEMENTED();
 }
 
-void FVulkanDynamicRHI::RHITransferVertexBufferUnderlyingResource(FVertexBufferRHIParamRef DestVertexBuffer, FVertexBufferRHIParamRef SrcVertexBuffer)
+void FVulkanDynamicRHI::RHITransferVertexBufferUnderlyingResource(FRHIVertexBuffer* DestVertexBuffer, FRHIVertexBuffer* SrcVertexBuffer)
 {
 	check(DestVertexBuffer);
 	FVulkanVertexBuffer* Dest = ResourceCast(DestVertexBuffer);
diff --git a/Engine/Source/Runtime/VulkanRHI/Public/VulkanDynamicRHI.h b/Engine/Source/Runtime/VulkanRHI/Public/VulkanDynamicRHI.h
--- a/Engine/Source/Runtime/VulkanRHI/Public/VulkanDynamicRHI.h
+++ b/Engine/Source/Runtime/VulkanRHI/Public/VulkanDynamicRHI.h
@@ -58,28 +58,28 @@ public:
 	virtual FBoundShaderStateRHIRef RHICreateBoundShaderState(FRHIVertexDeclaration* VertexDeclaration, FRHIVertexShader* VertexShader, FRHIHullShader* HullShader, FRHIDomainShader* DomainShader, FRHIPixelShader* PixelShader, FRHIGeometryShader* GeometryShader) final override;
 	virtual FGraphicsPipelineStateRHIRef RHICreateGraphicsPipelineState(const FGraphicsPipelineStateInitializer& Initializer) final override;
 	virtual FUniformBufferRHIRef RHICreateUniformBuffer(const void* Contents, const FRHIUniformBufferLayout& Layout, EUniformBufferUsage Usage, EUniformBufferValidation Validation) final override;
-	virtual void RHIUpdateUniformBuffer(FUniformBufferRHIParamRef UniformBufferRHI, const void* Contents) final override;
+	virtual void RHIUpdateUniformBuffer(FRHIUniformBuffer* UniformBufferRHI, const void* Contents) final override;
 	virtual FIndexBufferRHIRef RHICreateIndexBuffer(uint32 Stride, uint32 Size, uint32 InUsage, FRHIResourceCreateInfo& CreateInfo) final override;
-	virtual void* RHILockIndexBuffer(FIndexBufferRHIParamRef IndexBuffer, uint32 Offset, uint32 Size, EResourceLockMode LockMode) final override;
-	virtual void RHIUnlockIndexBuffer(FIndexBufferRHIParamRef IndexBuffer) final override;
-	virtual void RHITransferIndexBufferUnderlyingResource(FIndexBufferRHIParamRef DestIndexBuffer, FIndexBufferRHIParamRef SrcIndexBuffer) final override;
+	virtual void* RHILockIndexBuffer(FRHIIndexBuffer* IndexBuffer, uint32 Offset, uint32 Size, EResourceLockMode LockMode) final override;
+	virtual void RHIUnlockIndexBuffer(FRHIIndexBuffer* IndexBuffer) final override;
+	virtual void RHITransferIndexBufferUnderlyingResource(FRHIIndexBuffer* DestIndexBuffer, FRHIIndexBuffer* SrcIndexBuffer) final override;
 	virtual FVertexBufferRHIRef RHICreateVertexBuffer(uint32 Size, uint32 InUsage, FRHIResourceCreateInfo& CreateInfo) final override;
-	virtual void* RHILockVertexBuffer(FVertexBufferRHIParamRef VertexBuffer, uint32 Offset, uint32 SizeRHI, EResourceLockMode LockMode) final override;
-	virtual void RHIUnlockVertexBuffer(FVertexBufferRHIParamRef VertexBuffer) final override;
-	virtual void RHICopyVertexBuffer(FVertexBufferRHIParamRef SourceBuffer, FVertexBufferRHIParamRef DestBuffer) final override;
-	virtual void RHITransferVertexBufferUnderlyingResource(FVertexBufferRHIParamRef DestVertexBuffer, FVertexBufferRHIParamRef SrcVertexBuffer) final override;
+	virtual void* RHILockVertexBuffer(FRHIVertexBuffer* VertexBuffer, uint32 Offset, uint32 SizeRHI, EResourceLockMode LockMode) final override;
+	virtual void RHIUnlockVertexBuffer(FRHIVertexBuffer* VertexBuffer) final override;
+	virtual void RHICopyVertexBuffer(FRHIVertexBuffer* SourceBuffer, FRHIVertexBuffer* DestBuffer) final override;
+	virtual void RHITransferVertexBufferUnderlyingResource(FRHIVertexBuffer* DestVertexBuffer, FRHIVertexBuffer* SrcVertexBuffer) final override;
 	virtual FStructuredBufferRHIRef RHICreateStructuredBuffer(uint32 Stride, uint32 Size, uint32 InUsage, FRHIResourceCreateInfo& CreateInfo) final override;
-	virtual void* RHILockStructuredBuffer(FStructuredBufferRHIParamRef StructuredBuffer, uint32 Offset, uint32 SizeRHI, EResourceLockMode LockMode) final override;
-	virtual void RHIUnlockStructuredBuffer(FStructuredBufferRHIParamRef StructuredBuffer) final override;
-	virtual FUnorderedAccessViewRHIRef RHICreateUnorderedAccessView(FStructuredBufferRHIParamRef StructuredBuffer, bool bUseUAVCounter, bool bAppendBuffer) final override;
-	virtual FUnorderedAccessViewRHIRef RHICreateUnorderedAccessView(FTextureRHIParamRef Texture, uint32 MipLevel) final override;
-	virtual FUnorderedAccessViewRHIRef RHICreateUnorderedAccessView(FVertexBufferRHIParamRef VertexBuffer, uint8 Format) final override;
-	virtual FUnorderedAccessViewRHIRef RHICreateUnorderedAccessView(FIndexBufferRHIParamRef IndexBuffer, uint8 Format) final override;
-	virtual FShaderResourceViewRHIRef RHICreateShaderResourceView(FStructuredBufferRHIParamRef StructuredBuffer) final override;
-	virtual FShaderResourceViewRHIRef RHICreateShaderResourceView(FVertexBufferRHIParamRef VertexBuffer, uint32 Stride, uint8 Format) final override;
-	virtual FShaderResourceViewRHIRef RHICreateShaderResourceView(FIndexBufferRHIParamRef Buffer) final override;
-	virtual void RHIUpdateShaderResourceView(FRHIShaderResourceView* SRV, FVertexBufferRHIParamRef VertexBuffer, uint32 Stride, uint8 Format) final override;
-	virtual void RHIUpdateShaderResourceView(FRHIShaderResourceView* SRV, FIndexBufferRHIParamRef IndexBuffer) final override;
+	virtual void* RHILockStructuredBuffer(FRHIStructuredBuffer* StructuredBuffer, uint32 Offset, uint32 SizeRHI, EResourceLockMode LockMode) final override;
+	virtual void RHIUnlockStructuredBuffer(FRHIStructuredBuffer* StructuredBuffer) final override;
+	virtual FUnorderedAccessViewRHIRef RHICreateUnorderedAccessView(FRHIStructuredBuffer* StructuredBuffer, bool bUseUAVCounter, bool bAppendBuffer) final override;
+	virtual FUnorderedAccessViewRHIRef RHICreateUnorderedAccessView(FRHITexture* Texture, uint32 MipLevel) final override;
+	virtual FUnorderedAccessViewRHIRef RHICreateUnorderedAccessView(FRHIVertexBuffer* VertexBuffer, uint8 Format) final override;
+	virtual FUnorderedAccessViewRHIRef RHICreateUnorderedAccessView(FRHIIndexBuffer* IndexBuffer, uint8 Format) final override;
+	virtual FShaderResourceViewRHIRef RHICreateShaderResourceView(FRHIStructuredBuffer* StructuredBuffer) final override;
+	virtual FShaderResourceViewRHIRef RHICreateShaderResourceView(FRHIVertexBuffer* VertexBuffer, uint32 Stride, uint8 Format) final override;
+	virtual FShaderResourceViewRHIRef RHICreateShaderResourceView(FRHIIndexBuffer* Buffer) final override;
+	virtual void RHIUpdateShaderResourceView(FRHIShaderResourceView* SRV, FRHIVertexBuffer* VertexBuffer, uint32 Stride, uint8 Format) final override;
+	virtual void RHIUpdateShaderResourceView(FRHIShaderResourceView* SRV, FRHIIndexBuffer* IndexBuffer) final override;
 	virtual uint64 RHICalcTexture2DPlatformSize(uint32 SizeX, uint32 SizeY, uint8 Format, uint32 NumMips, uint32 NumSamples, uint32 Flags, uint32& OutAlign) final override;
 	virtual uint64 RHICalcTexture3DPlatformSize(uint32 SizeX, uint32 SizeY, uint32 SizeZ, uint8 Format, uint32 NumMips, uint32 Flags, uint32& OutAlign) final override;
 	virtual uint64 RHICalcTextureCubePlatformSize(uint32 Size, uint8 Format, uint32 NumMips, uint32 Flags, uint32& OutAlign) final override;
@@ -88,55 +88,55 @@ public:
 	virtual FTextureReferenceRHIRef RHICreateTextureReference(FLastRenderTimeContainer* LastRenderTime) final override;
 	virtual FTexture2DRHIRef RHICreateTexture2D(uint32 SizeX, uint32 SizeY, uint8 Format, uint32 NumMips, uint32 NumSamples, uint32 Flags, FRHIResourceCreateInfo& CreateInfo) final override;
 	virtual FTexture2DRHIRef RHIAsyncCreateTexture2D(uint32 SizeX, uint32 SizeY, uint8 Format, uint32 NumMips, uint32 Flags, void** InitialMipData, uint32 NumInitialMips) final override;
-	virtual void RHICopySharedMips(FTexture2DRHIParamRef DestTexture2D, FTexture2DRHIParamRef SrcTexture2D) final override;
+	virtual void RHICopySharedMips(FRHITexture2D* DestTexture2D, FRHITexture2D* SrcTexture2D) final override;
 	virtual FTexture2DArrayRHIRef RHICreateTexture2DArray(uint32 SizeX, uint32 SizeY, uint32 SizeZ, uint8 Format, uint32 NumMips, uint32 NumSamples, uint32 Flags, FRHIResourceCreateInfo& CreateInfo) final override;
 	virtual FTexture3DRHIRef RHICreateTexture3D(uint32 SizeX, uint32 SizeY, uint32 SizeZ, uint8 Format, uint32 NumMips, uint32 Flags, FRHIResourceCreateInfo& CreateInfo) final override;
-	virtual void RHIGetResourceInfo(FTextureRHIParamRef Ref, FRHIResourceInfo& OutInfo) final override;
-	virtual FShaderResourceViewRHIRef RHICreateShaderResourceView(FTextureRHIParamRef Texture, const FRHITextureSRVCreateInfo& CreateInfo) final override;
-	virtual uint32 RHIComputeMemorySize(FTextureRHIParamRef TextureRHI) final override;
-	virtual FTexture2DRHIRef RHIAsyncReallocateTexture2D(FTexture2DRHIParamRef Texture2D, int32 NewMipCount, int32 NewSizeX, int32 NewSizeY, FThreadSafeCounter* RequestStatus) final override;
-	virtual ETextureReallocationStatus RHIFinalizeAsyncReallocateTexture2D(FTexture2DRHIParamRef Texture2D, bool bBlockUntilCompleted) final override;
-	virtual ETextureReallocationStatus FinalizeAsyncReallocateTexture2D_RenderThread(class FRHICommandListImmediate& RHICmdList, FTexture2DRHIParamRef Texture2D, bool bBlockUntilCompleted) final override
+	virtual void RHIGetResourceInfo(FRHITexture* Ref, FRHIResourceInfo& OutInfo) final override;
+	virtual FShaderResourceViewRHIRef RHICreateShaderResourceView(FRHITexture* Texture, const FRHITextureSRVCreateInfo& CreateInfo) final override;
+	virtual uint32 RHIComputeMemorySize(FRHITexture* TextureRHI) final override;
+	virtual FTexture2DRHIRef RHIAsyncReallocateTexture2D(FRHITexture2D* Texture2D, int32 NewMipCount, int32 NewSizeX, int32 NewSizeY, FThreadSafeCounter* RequestStatus) final override;
+	virtual ETextureReallocationStatus RHIFinalizeAsyncReallocateTexture2D(FRHITexture2D* Texture2D, bool bBlockUntilCompleted) final override;
+	virtual ETextureReallocationStatus FinalizeAsyncReallocateTexture2D_RenderThread(class FRHICommandListImmediate& RHICmdList, FRHITexture2D* Texture2D, bool bBlockUntilCompleted) final override
 	{
 		return this->RHIFinalizeAsyncReallocateTexture2D(Texture2D, bBlockUntilCompleted);
 	}
 
-	virtual ETextureReallocationStatus RHICancelAsyncReallocateTexture2D(FTexture2DRHIParamRef Texture2D, bool bBlockUntilCompleted) final override;
-	virtual void* RHILockTexture2D(FTexture2DRHIParamRef Texture, uint32 MipIndex, EResourceLockMode LockMode, uint32& DestStride, bool bLockWithinMiptail) final override;
-	virtual void RHIUnlockTexture2D(FTexture2DRHIParamRef Texture, uint32 MipIndex, bool bLockWithinMiptail) final override
+	virtual ETextureReallocationStatus RHICancelAsyncReallocateTexture2D(FRHITexture2D* Texture2D, bool bBlockUntilCompleted) final override;
+	virtual void* RHILockTexture2D(FRHITexture2D* Texture, uint32 MipIndex, EResourceLockMode LockMode, uint32& DestStride, bool bLockWithinMiptail) final override;
+	virtual void RHIUnlockTexture2D(FRHITexture2D* Texture, uint32 MipIndex, bool bLockWithinMiptail) final override
 	{
 		InternalUnlockTexture2D(false, Texture, MipIndex, bLockWithinMiptail);
 	}
-	virtual void* RHILockTexture2DArray(FTexture2DArrayRHIParamRef Texture, uint32 TextureIndex, uint32 MipIndex, EResourceLockMode LockMode, uint32& DestStride, bool bLockWithinMiptail) final override;
-	virtual void RHIUnlockTexture2DArray(FTexture2DArrayRHIParamRef Texture, uint32 TextureIndex, uint32 MipIndex, bool bLockWithinMiptail) final override;
-	virtual void RHIUpdateTexture2D(FTexture2DRHIParamRef Texture, uint32 MipIndex, const struct FUpdateTextureRegion2D& UpdateRegion, uint32 SourcePitch, const uint8* SourceData) final override
+	virtual void* RHILockTexture2DArray(FRHITexture2DArray* Texture, uint32 TextureIndex, uint32 MipIndex, EResourceLockMode LockMode, uint32& DestStride, bool bLockWithinMiptail) final override;
+	virtual void RHIUnlockTexture2DArray(FRHITexture2DArray* Texture, uint32 TextureIndex, uint32 MipIndex, bool bLockWithinMiptail) final override;
+	virtual void RHIUpdateTexture2D(FRHITexture2D* Texture, uint32 MipIndex, const struct FUpdateTextureRegion2D& UpdateRegion, uint32 SourcePitch, const uint8* SourceData) final override
 	{
 		InternalUpdateTexture2D(false, Texture, MipIndex, UpdateRegion, SourcePitch, SourceData);
 	}
-	virtual void RHIUpdateTexture3D(FTexture3DRHIParamRef Texture, uint32 MipIndex, const struct FUpdateTextureRegion3D& UpdateRegion, uint32 SourceRowPitch, uint32 SourceDepthPitch, const uint8* SourceData) final override
+	virtual void RHIUpdateTexture3D(FRHITexture3D* Texture, uint32 MipIndex, const struct FUpdateTextureRegion3D& UpdateRegion, uint32 SourceRowPitch, uint32 SourceDepthPitch, const uint8* SourceData) final override
 	{
 		InternalUpdateTexture3D(false, Texture, MipIndex, UpdateRegion, SourceRowPitch, SourceDepthPitch, SourceData);
 	}
 	virtual FTextureCubeRHIRef RHICreateTextureCube(uint32 Size, uint8 Format, uint32 NumMips, uint32 Flags, FRHIResourceCreateInfo& CreateInfo) final override;
 	virtual FTextureCubeRHIRef RHICreateTextureCubeArray(uint32 Size, uint32 ArraySize, uint8 Format, uint32 NumMips, uint32 Flags, FRHIResourceCreateInfo& CreateInfo) final override;
-	virtual void* RHILockTextureCubeFace(FTextureCubeRHIParamRef Texture, uint32 FaceIndex, uint32 ArrayIndex, uint32 MipIndex, EResourceLockMode LockMode, uint32& DestStride, bool bLockWithinMiptail) final override;
-	virtual void RHIUnlockTextureCubeFace(FTextureCubeRHIParamRef Texture, uint32 FaceIndex, uint32 ArrayIndex, uint32 MipIndex, bool bLockWithinMiptail) final override;
-	virtual void RHIBindDebugLabelName(FTextureRHIParamRef Texture, const TCHAR* Name) final override;
+	virtual void* RHILockTextureCubeFace(FRHITextureCube* Texture, uint32 FaceIndex, uint32 ArrayIndex, uint32 MipIndex, EResourceLockMode LockMode, uint32& DestStride, bool bLockWithinMiptail) final override;
+	virtual void RHIUnlockTextureCubeFace(FRHITextureCube* Texture, uint32 FaceIndex, uint32 ArrayIndex, uint32 MipIndex, bool bLockWithinMiptail) final override;
+	virtual void RHIBindDebugLabelName(FRHITexture* Texture, const TCHAR* Name) final override;
 	virtual void RHIBindDebugLabelName(FRHIUnorderedAccessView* UnorderedAccessViewRHI, const TCHAR* Name) final override;
-	virtual void RHIReadSurfaceData(FTextureRHIParamRef Texture, FIntRect Rect, TArray<FColor>& OutData, FReadSurfaceDataFlags InFlags) final override;
-	virtual void RHIReadSurfaceData(FTextureRHIParamRef Texture, FIntRect Rect, TArray<FLinearColor>& OutData, FReadSurfaceDataFlags InFlags) final override;
-	virtual void RHIMapStagingSurface(FTextureRHIParamRef Texture, void*& OutData, int32& OutWidth, int32& OutHeight) final override;
-	virtual void RHIUnmapStagingSurface(FTextureRHIParamRef Texture) final override;
-	virtual void RHIMapStagingSurface_RenderThread(class FRHICommandListImmediate& RHICmdList, FTextureRHIParamRef Texture, void*& OutData, int32& OutWidth, int32& OutHeight);
-	virtual void RHIUnmapStagingSurface_RenderThread(class FRHICommandListImmediate& RHICmdList, FTextureRHIParamRef Texture);
+	virtual void RHIReadSurfaceData(FRHITexture* Texture, FIntRect Rect, TArray<FColor>& OutData, FReadSurfaceDataFlags InFlags) final override;
+	virtual void RHIReadSurfaceData(FRHITexture* Texture, FIntRect Rect, TArray<FLinearColor>& OutData, FReadSurfaceDataFlags InFlags) final override;
+	virtual void RHIMapStagingSurface(FRHITexture* Texture, void*& OutData, int32& OutWidth, int32& OutHeight) final override;
+	virtual void RHIUnmapStagingSurface(FRHITexture* Texture) final override;
+	virtual void RHIMapStagingSurface_RenderThread(class FRHICommandListImmediate& RHICmdList, FRHITexture* Texture, void*& OutData, int32& OutWidth, int32& OutHeight);
+	virtual void RHIUnmapStagingSurface_RenderThread(class FRHICommandListImmediate& RHICmdList, FRHITexture* Texture);
 
 
-	virtual void RHIReadSurfaceFloatData(FTextureRHIParamRef Texture, FIntRect Rect, TArray<FFloat16Color>& OutData, ECubeFace CubeFace, int32 ArrayIndex, int32 MipIndex) final override;
-	virtual void RHIRead3DSurfaceFloatData(FTextureRHIParamRef Texture, FIntRect Rect, FIntPoint ZMinMax, TArray<FFloat16Color>& OutData) final override;
+	virtual void RHIReadSurfaceFloatData(FRHITexture* Texture, FIntRect Rect, TArray<FFloat16Color>& OutData, ECubeFace CubeFace, int32 ArrayIndex, int32 MipIndex) final override;
+	virtual void RHIRead3DSurfaceFloatData(FRHITexture* Texture, FIntRect Rect, FIntPoint ZMinMax, TArray<FFloat16Color>& OutData) final override;
 	virtual FRenderQueryRHIRef RHICreateRenderQuery(ERenderQueryType QueryType) final override;
 	virtual bool RHIGetRenderQueryResult(FRHIRenderQuery* RenderQuery, uint64& OutResult, bool bWait) final override;
 	virtual FTexture2DRHIRef RHIGetViewportBackBuffer(FRHIViewport* Viewport) final override;
-	virtual void RHIAliasTextureResources(FTextureRHIParamRef DestTexture, FTextureRHIParamRef SrcTexture) final override;
+	virtual void RHIAliasTextureResources(FRHITexture* DestTexture, FRHITexture* SrcTexture) final override;
 	virtual void RHIAdvanceFrameForGetViewportBackBuffer(FRHIViewport* Viewport) final override;
 	virtual void RHIAcquireThreadOwnership() final override;
 	virtual void RHIReleaseThreadOwnership() final override;
@@ -146,7 +146,7 @@ public:
 	virtual void RHIResizeViewport(FRHIViewport* Viewport, uint32 SizeX, uint32 SizeY, bool bIsFullscreen) final override;
 	virtual void RHIResizeViewport(FRHIViewport* Viewport, uint32 SizeX, uint32 SizeY, bool bIsFullscreen, EPixelFormat PreferredPixelFormat) final override;
 	virtual void RHITick(float DeltaTime) final override;
-	virtual void RHISetStreamOutTargets(uint32 NumTargets, const FVertexBufferRHIParamRef* VertexBuffers, const uint32* Offsets) final override;
+	virtual void RHISetStreamOutTargets(uint32 NumTargets, FRHIVertexBuffer* const* VertexBuffers, const uint32* Offsets) final override;
 	virtual void RHIBlockUntilGPUIdle() final override;
 	virtual void RHISubmitCommandsAndFlushGPU() final override;
 	virtual void RHISuspendRendering() final override;
@@ -154,8 +154,8 @@ public:
 	virtual bool RHIIsRenderingSuspended() final override;
 	virtual bool RHIGetAvailableResolutions(FScreenResolutionArray& Resolutions, bool bIgnoreRefreshRate) final override;
 	virtual void RHIGetSupportedResolution(uint32& Width, uint32& Height) final override;
-	virtual void RHIVirtualTextureSetFirstMipInMemory(FTexture2DRHIParamRef Texture, uint32 FirstMip) final override;
-	virtual void RHIVirtualTextureSetFirstMipVisible(FTexture2DRHIParamRef Texture, uint32 FirstMip) final override;
+	virtual void RHIVirtualTextureSetFirstMipInMemory(FRHITexture2D* Texture, uint32 FirstMip) final override;
+	virtual void RHIVirtualTextureSetFirstMipVisible(FRHITexture2D* Texture, uint32 FirstMip) final override;
 	virtual void RHIExecuteCommandList(FRHICommandList* CmdList) final override;
 
 	virtual void* RHIGetNativeDevice() final override;
@@ -174,11 +174,11 @@ public:
 		return this->RHICreateVertexBuffer(Size, InUsage, CreateInfo);
 	}
 #if VULKAN_BUFFER_LOCK_THREADSAFE	
-	virtual void* LockVertexBuffer_RenderThread(class FRHICommandListImmediate& RHICmdList, FVertexBufferRHIParamRef VertexBuffer, uint32 Offset, uint32 SizeRHI, EResourceLockMode LockMode) override final;
-	virtual void UnlockVertexBuffer_RenderThread(class FRHICommandListImmediate& RHICmdList, FVertexBufferRHIParamRef VertexBuffer) override final;
+	virtual void* LockVertexBuffer_RenderThread(class FRHICommandListImmediate& RHICmdList, FRHIVertexBuffer* VertexBuffer, uint32 Offset, uint32 SizeRHI, EResourceLockMode LockMode) override final;
+	virtual void UnlockVertexBuffer_RenderThread(class FRHICommandListImmediate& RHICmdList, FRHIVertexBuffer* VertexBuffer) override final;
 #endif
-	virtual FTexture2DRHIRef AsyncReallocateTexture2D_RenderThread(class FRHICommandListImmediate& RHICmdList, FTexture2DRHIParamRef Texture2D, int32 NewMipCount, int32 NewSizeX, int32 NewSizeY, FThreadSafeCounter* RequestStatus) override final;
-	//virtual ETextureReallocationStatus CancelAsyncReallocateTexture2D_RenderThread(class FRHICommandListImmediate& RHICmdList, FTexture2DRHIParamRef Texture2D, bool bBlockUntilCompleted) override final;
+	virtual FTexture2DRHIRef AsyncReallocateTexture2D_RenderThread(class FRHICommandListImmediate& RHICmdList, FRHITexture2D* Texture2D, int32 NewMipCount, int32 NewSizeX, int32 NewSizeY, FThreadSafeCounter* RequestStatus) override final;
+	//virtual ETextureReallocationStatus CancelAsyncReallocateTexture2D_RenderThread(class FRHICommandListImmediate& RHICmdList, FRHITexture2D* Texture2D, bool bBlockUntilCompleted) override final;
 	
 	virtual FIndexBufferRHIRef CreateIndexBuffer_RenderThread(class FRHICommandListImmediate& RHICmdList, uint32 Stride, uint32 Size, uint32 InUsage, FRHIResourceCreateInfo& CreateInfo) override final
 	{
@@ -186,8 +186,8 @@ public:
 		return this->RHICreateIndexBuffer(Stride, Size, InUsage, CreateInfo);
 	}
 #if VULKAN_BUFFER_LOCK_THREADSAFE		
-	virtual void* LockIndexBuffer_RenderThread(class FRHICommandListImmediate& RHICmdList, FIndexBufferRHIParamRef IndexBuffer, uint32 Offset, uint32 SizeRHI, EResourceLockMode LockMode) override final;
-	virtual void UnlockIndexBuffer_RenderThread(class FRHICommandListImmediate& RHICmdList, FIndexBufferRHIParamRef IndexBuffer) override final;
+	virtual void* LockIndexBuffer_RenderThread(class FRHICommandListImmediate& RHICmdList, FRHIIndexBuffer* IndexBuffer, uint32 Offset, uint32 SizeRHI, EResourceLockMode LockMode) override final;
+	virtual void UnlockIndexBuffer_RenderThread(class FRHICommandListImmediate& RHICmdList, FRHIIndexBuffer* IndexBuffer) override final;
 #endif
 
 	virtual FVertexShaderRHIRef CreateVertexShader_RenderThread(class FRHICommandListImmediate& RHICmdList, const TArray<uint8>& Code) override final
@@ -225,19 +225,19 @@ public:
 		return RHICreateDomainShader(Code);
 	}
 	//virtual FDomainShaderRHIRef CreateDomainShader_RenderThread(class FRHICommandListImmediate& RHICmdList, FRHIShaderLibrary* Library, FSHAHash Hash) override final;
-	virtual void* LockTexture2D_RenderThread(class FRHICommandListImmediate& RHICmdList, FTexture2DRHIParamRef Texture, uint32 MipIndex, EResourceLockMode LockMode, uint32& DestStride, bool bLockWithinMiptail, bool bNeedsDefaultRHIFlush = true) override final
+	virtual void* LockTexture2D_RenderThread(class FRHICommandListImmediate& RHICmdList, FRHITexture2D* Texture, uint32 MipIndex, EResourceLockMode LockMode, uint32& DestStride, bool bLockWithinMiptail, bool bNeedsDefaultRHIFlush = true) override final
 	{
 		return RHILockTexture2D(Texture, MipIndex, LockMode, DestStride, bLockWithinMiptail);
 	}
-	virtual void UnlockTexture2D_RenderThread(class FRHICommandListImmediate& RHICmdList, FTexture2DRHIParamRef Texture, uint32 MipIndex, bool bLockWithinMiptail, bool bNeedsDefaultRHIFlush = true) override final
+	virtual void UnlockTexture2D_RenderThread(class FRHICommandListImmediate& RHICmdList, FRHITexture2D* Texture, uint32 MipIndex, bool bLockWithinMiptail, bool bNeedsDefaultRHIFlush = true) override final
 	{
 		InternalUnlockTexture2D(true, Texture, MipIndex, bLockWithinMiptail);
 	}
-	virtual void UpdateTexture2D_RenderThread(class FRHICommandListImmediate& RHICmdList, FTexture2DRHIParamRef Texture, uint32 MipIndex, const struct FUpdateTextureRegion2D& UpdateRegion, uint32 SourcePitch, const uint8* SourceData) override final
+	virtual void UpdateTexture2D_RenderThread(class FRHICommandListImmediate& RHICmdList, FRHITexture2D* Texture, uint32 MipIndex, const struct FUpdateTextureRegion2D& UpdateRegion, uint32 SourcePitch, const uint8* SourceData) override final
 	{
 		InternalUpdateTexture2D(true, Texture, MipIndex, UpdateRegion, SourcePitch, SourceData);
 	}
-	virtual void UpdateTexture3D_RenderThread(class FRHICommandListImmediate& RHICmdList, FTexture3DRHIParamRef Texture, uint32 MipIndex, const struct FUpdateTextureRegion3D& UpdateRegion, uint32 SourceRowPitch, uint32 SourceDepthPitch, const uint8* SourceData) override final
+	virtual void UpdateTexture3D_RenderThread(class FRHICommandListImmediate& RHICmdList, FRHITexture3D* Texture, uint32 MipIndex, const struct FUpdateTextureRegion3D& UpdateRegion, uint32 SourceRowPitch, uint32 SourceDepthPitch, const uint8* SourceData) override final
 	{
 		InternalUpdateTexture3D(true, Texture, MipIndex, UpdateRegion, SourceRowPitch, SourceDepthPitch, SourceData);
 	}
@@ -264,42 +264,42 @@ public:
 	{
 		return this->RHICreateTexture3D(SizeX, SizeY, SizeZ, Format, NumMips, Flags, CreateInfo);
 	}
-	virtual FUnorderedAccessViewRHIRef RHICreateUnorderedAccessView_RenderThread(class FRHICommandListImmediate& RHICmdList, FStructuredBufferRHIParamRef StructuredBuffer, bool bUseUAVCounter, bool bAppendBuffer) final override
+	virtual FUnorderedAccessViewRHIRef RHICreateUnorderedAccessView_RenderThread(class FRHICommandListImmediate& RHICmdList, FRHIStructuredBuffer* StructuredBuffer, bool bUseUAVCounter, bool bAppendBuffer) final override
 	{
 		// this-> is required to avoid calling the global version of this function
 		return this->RHICreateUnorderedAccessView(StructuredBuffer, bUseUAVCounter, bAppendBuffer);
 	}
 
-	virtual FUnorderedAccessViewRHIRef RHICreateUnorderedAccessView_RenderThread(class FRHICommandListImmediate& RHICmdList, FTextureRHIParamRef Texture, uint32 MipLevel) final override
+	virtual FUnorderedAccessViewRHIRef RHICreateUnorderedAccessView_RenderThread(class FRHICommandListImmediate& RHICmdList, FRHITexture* Texture, uint32 MipLevel) final override
 	{
 		// this-> is required to avoid calling the global version of this function
 		return this->RHICreateUnorderedAccessView(Texture, MipLevel);
 	}
 
-	virtual FUnorderedAccessViewRHIRef RHICreateUnorderedAccessView_RenderThread(class FRHICommandListImmediate& RHICmdList, FVertexBufferRHIParamRef VertexBuffer, uint8 Format) final override
+	virtual FUnorderedAccessViewRHIRef RHICreateUnorderedAccessView_RenderThread(class FRHICommandListImmediate& RHICmdList, FRHIVertexBuffer* VertexBuffer, uint8 Format) final override
 	{
 		// this-> is required to avoid calling the global version of this function
 		return this->RHICreateUnorderedAccessView(VertexBuffer, Format);
 	}
 
-	virtual FShaderResourceViewRHIRef CreateShaderResourceView_RenderThread(class FRHICommandListImmediate& RHICmdList, FVertexBufferRHIParamRef VertexBuffer, uint32 Stride, uint8 Format) final override
+	virtual FShaderResourceViewRHIRef CreateShaderResourceView_RenderThread(class FRHICommandListImmediate& RHICmdList, FRHIVertexBuffer* VertexBuffer, uint32 Stride, uint8 Format) final override
 	{
 		return this->RHICreateShaderResourceView(VertexBuffer, Stride, Format);
 	}
 
-	virtual FShaderResourceViewRHIRef RHICreateShaderResourceView_RenderThread(class FRHICommandListImmediate& RHICmdList, FTextureRHIParamRef Texture, const FRHITextureSRVCreateInfo& CreateInfo) final override
+	virtual FShaderResourceViewRHIRef RHICreateShaderResourceView_RenderThread(class FRHICommandListImmediate& RHICmdList, FRHITexture* Texture, const FRHITextureSRVCreateInfo& CreateInfo) final override
 	{
 		// this-> is required to avoid calling the global version of this function
 		return this->RHICreateShaderResourceView(Texture, CreateInfo);
 	}
 
-	virtual FShaderResourceViewRHIRef RHICreateShaderResourceView_RenderThread(class FRHICommandListImmediate& RHICmdList, FVertexBufferRHIParamRef VertexBuffer, uint32 Stride, uint8 Format) final override
+	virtual FShaderResourceViewRHIRef RHICreateShaderResourceView_RenderThread(class FRHICommandListImmediate& RHICmdList, FRHIVertexBuffer* VertexBuffer, uint32 Stride, uint8 Format) final override
 	{
 		// this-> is required to avoid calling the global version of this function
 		return this->RHICreateShaderResourceView(VertexBuffer, Stride, Format);
 	}
 
-	virtual FShaderResourceViewRHIRef RHICreateShaderResourceView_RenderThread(class FRHICommandListImmediate& RHICmdList, FStructuredBufferRHIParamRef StructuredBuffer) final override
+	virtual FShaderResourceViewRHIRef RHICreateShaderResourceView_RenderThread(class FRHICommandListImmediate& RHICmdList, FRHIStructuredBuffer* StructuredBuffer) final override
 	{
 		// this-> is required to avoid calling the global version of this function
 		return this->RHICreateShaderResourceView(StructuredBuffer);
@@ -320,12 +320,12 @@ public:
 		return this->RHICreateRenderQuery(QueryType);
 	}
 
-	//virtual void RHIAcquireTransientResource_RenderThread(FTextureRHIParamRef Texture) { }
-	//virtual void RHIDiscardTransientResource_RenderThread(FTextureRHIParamRef Texture) { }
-	//virtual void RHIAcquireTransientResource_RenderThread(FVertexBufferRHIParamRef Buffer) { }
-	//virtual void RHIDiscardTransientResource_RenderThread(FVertexBufferRHIParamRef Buffer) { }
-	//virtual void RHIAcquireTransientResource_RenderThread(FStructuredBufferRHIParamRef Buffer) { }
-	//virtual void RHIDiscardTransientResource_RenderThread(FStructuredBufferRHIParamRef Buffer) { }
+	//virtual void RHIAcquireTransientResource_RenderThread(FRHITexture* Texture) { }
+	//virtual void RHIDiscardTransientResource_RenderThread(FRHITexture* Texture) { }
+	//virtual void RHIAcquireTransientResource_RenderThread(FRHIVertexBuffer* Buffer) { }
+	//virtual void RHIDiscardTransientResource_RenderThread(FRHIVertexBuffer* Buffer) { }
+	//virtual void RHIAcquireTransientResource_RenderThread(FRHIStructuredBuffer*FRHIStructuredBuffer* Buffer) { }
+	//virtual void RHIDiscardTransientResource_RenderThread(FRHIStructuredBuffer* Buffer) { }
 
 	// FVulkanDynamicRHI interface
 	virtual FTexture2DRHIRef RHICreateTexture2DFromResource(EPixelFormat Format, uint32 SizeX, uint32 SizeY, uint32 NumMips, uint32 NumSamples, VkImage Resource, uint32 Flags);
@@ -333,7 +333,7 @@ public:
 	virtual FTexture2DArrayRHIRef RHICreateTexture2DArrayFromResource(EPixelFormat Format, uint32 SizeX, uint32 SizeY, uint32 ArraySize, uint32 NumMips, uint32 NumSamples, VkImage Resource, uint32 Flags);
 	virtual FTextureCubeRHIRef RHICreateTextureCubeFromResource(EPixelFormat Format, uint32 Size, bool bArray, uint32 ArraySize, uint32 NumMips, VkImage Resource, uint32 Flags);
 
-	virtual void RHICopySubTextureRegion(FTexture2DRHIParamRef SourceTexture, FTexture2DRHIParamRef DestinationTexture, FBox2D SourceBox, FBox2D DestinationBox) final override;
+	virtual void RHICopySubTextureRegion(FRHITexture2D* SourceTexture, FRHITexture2D* DestinationTexture, FBox2D SourceBox, FBox2D DestinationBox) final override;
 
 	void RHICalibrateTimers() override;
 
@@ -458,9 +458,9 @@ protected:
 
 	FCriticalSection LockBufferCS;
 
-	void InternalUnlockTexture2D(bool bFromRenderingThread, FTexture2DRHIParamRef Texture, uint32 MipIndex, bool bLockWithinMiptail);
-	void InternalUpdateTexture2D(bool bFromRenderingThread, FTexture2DRHIParamRef TextureRHI, uint32 MipIndex, const struct FUpdateTextureRegion2D& UpdateRegion, uint32 SourcePitch, const uint8* SourceData);
-	void InternalUpdateTexture3D(bool bFromRenderingThread, FTexture3DRHIParamRef TextureRHI, uint32 MipIndex, const struct FUpdateTextureRegion3D& UpdateRegion, uint32 SourceRowPitch, uint32 SourceDepthPitch, const uint8* SourceData);
+	void InternalUnlockTexture2D(bool bFromRenderingThread, FRHITexture2D* Texture, uint32 MipIndex, bool bLockWithinMiptail);
+	void InternalUpdateTexture2D(bool bFromRenderingThread, FRHITexture2D* TextureRHI, uint32 MipIndex, const struct FUpdateTextureRegion2D& UpdateRegion, uint32 SourcePitch, const uint8* SourceData);
+	void InternalUpdateTexture3D(bool bFromRenderingThread, FRHITexture3D* TextureRHI, uint32 MipIndex, const struct FUpdateTextureRegion3D& UpdateRegion, uint32 SourceRowPitch, uint32 SourceDepthPitch, const uint8* SourceData);
 
 	template <bool bRealUBs>
 	void UpdateUniformBuffer(FVulkanUniformBuffer* UniformBuffer, const void* Contents);
diff --git a/Engine/Source/Runtime/Windows/D3D11RHI/Private/D3D11Commands.cpp b/Engine/Source/Runtime/Windows/D3D11RHI/Private/D3D11Commands.cpp
--- a/Engine/Source/Runtime/Windows/D3D11RHI/Private/D3D11Commands.cpp
+++ b/Engine/Source/Runtime/Windows/D3D11RHI/Private/D3D11Commands.cpp
@@ -74,7 +74,7 @@ void FD3D11BaseShaderResource::SetDirty(bool bInDirty, uint32 CurrentFrame)
 
 #if !PLATFORM_HOLOLENS
 //MultiGPU
-void FD3D11DynamicRHI::RHIBeginUpdateMultiFrameResource(FTextureRHIParamRef RHITexture)
+void FD3D11DynamicRHI::RHIBeginUpdateMultiFrameResource(FRHITexture* RHITexture)
 {
 	if (!IsRHIDeviceNVIDIA() || GNumAlternateFrameRenderingGroups == 1) return;
 
@@ -98,7 +98,7 @@ void FD3D11DynamicRHI::RHIBeginUpdateMultiFrameResource(FTextureRHIParamRef RHIT
 	RHIPopEvent();
 }
 
-void FD3D11DynamicRHI::RHIEndUpdateMultiFrameResource(FTextureRHIParamRef RHITexture)
+void FD3D11DynamicRHI::RHIEndUpdateMultiFrameResource(FRHITexture* RHITexture)
 {
 	if (!IsRHIDeviceNVIDIA() || GNumAlternateFrameRenderingGroups == 1) return;
 
@@ -156,7 +156,7 @@ void FD3D11DynamicRHI::RHIEndUpdateMultiFrameResource(FRHIUnorderedAccessView* U
 #endif
 
 // Vertex state.
-void FD3D11DynamicRHI::RHISetStreamSource(uint32 StreamIndex, FVertexBufferRHIParamRef VertexBufferRHI, uint32 Offset)
+void FD3D11DynamicRHI::RHISetStreamSource(uint32 StreamIndex, FRHIVertexBuffer* VertexBufferRHI, uint32 Offset)
 {
 	FD3D11VertexBuffer* VertexBuffer = ResourceCast(VertexBufferRHI);
 
@@ -165,7 +165,7 @@ void FD3D11DynamicRHI::RHISetStreamSource(uint32 StreamIndex, FVertexBufferRHIPa
 	StateCache.SetStreamSource(D3DBuffer, StreamIndex, Offset);
 }
 
-void FD3D11DynamicRHI::RHISetStreamOutTargets(uint32 NumTargets, const FVertexBufferRHIParamRef* VertexBuffers, const uint32* Offsets)
+void FD3D11DynamicRHI::RHISetStreamOutTargets(uint32 NumTargets, FRHIVertexBuffer* const* VertexBuffers, const uint32* Offsets)
 {
 	ID3D11Buffer* D3DVertexBuffers[D3D11_SO_BUFFER_SLOT_COUNT] = {0};
 
@@ -217,7 +217,7 @@ void FD3D11DynamicRHI::RHIDispatchComputeShader(uint32 ThreadGroupCountX, uint32
 	StateCache.SetComputeShader(nullptr);
 }
 
-void FD3D11DynamicRHI::RHIDispatchIndirectComputeShader(FVertexBufferRHIParamRef ArgumentBufferRHI, uint32 ArgumentOffset) 
+void FD3D11DynamicRHI::RHIDispatchIndirectComputeShader(FRHIVertexBuffer* ArgumentBufferRHI, uint32 ArgumentOffset)
 { 
 	FRHIComputeShader* ComputeShaderRHI = GetCurrentComputeShader();
 	FD3D11ComputeShader* ComputeShader = ResourceCast(ComputeShaderRHI);
@@ -335,7 +335,7 @@ void FD3D11DynamicRHI::RHISetBoundShaderState(FRHIBoundShaderState* BoundShaderS
 	}
 }
 
-void FD3D11DynamicRHI::RHISetShaderTexture(FRHIVertexShader* VertexShaderRHI,uint32 TextureIndex,FTextureRHIParamRef NewTextureRHI)
+void FD3D11DynamicRHI::RHISetShaderTexture(FRHIVertexShader* VertexShaderRHI,uint32 TextureIndex, FRHITexture* NewTextureRHI)
 {
 	VALIDATE_BOUND_SHADER(VertexShaderRHI);
 
@@ -352,7 +352,7 @@ void FD3D11DynamicRHI::RHISetShaderTexture(FRHIVertexShader* VertexShaderRHI,uin
 	}
 }
 
-void FD3D11DynamicRHI::RHISetShaderTexture(FRHIHullShader* HullShaderRHI,uint32 TextureIndex,FTextureRHIParamRef NewTextureRHI)
+void FD3D11DynamicRHI::RHISetShaderTexture(FRHIHullShader* HullShaderRHI,uint32 TextureIndex, FRHITexture* NewTextureRHI)
 {
 	VALIDATE_BOUND_SHADER(HullShaderRHI);
 
@@ -369,7 +369,7 @@ void FD3D11DynamicRHI::RHISetShaderTexture(FRHIHullShader* HullShaderRHI,uint32
 	}
 }
 
-void FD3D11DynamicRHI::RHISetShaderTexture(FRHIDomainShader* DomainShaderRHI,uint32 TextureIndex,FTextureRHIParamRef NewTextureRHI)
+void FD3D11DynamicRHI::RHISetShaderTexture(FRHIDomainShader* DomainShaderRHI,uint32 TextureIndex, FRHITexture* NewTextureRHI)
 {
 	VALIDATE_BOUND_SHADER(DomainShaderRHI);
 
@@ -386,7 +386,7 @@ void FD3D11DynamicRHI::RHISetShaderTexture(FRHIDomainShader* DomainShaderRHI,uin
 	}
 }
 
-void FD3D11DynamicRHI::RHISetShaderTexture(FRHIGeometryShader* GeometryShaderRHI,uint32 TextureIndex,FTextureRHIParamRef NewTextureRHI)
+void FD3D11DynamicRHI::RHISetShaderTexture(FRHIGeometryShader* GeometryShaderRHI,uint32 TextureIndex, FRHITexture* NewTextureRHI)
 {
 	VALIDATE_BOUND_SHADER(GeometryShaderRHI);
 
@@ -403,7 +403,7 @@ void FD3D11DynamicRHI::RHISetShaderTexture(FRHIGeometryShader* GeometryShaderRHI
 	}
 }
 
-void FD3D11DynamicRHI::RHISetShaderTexture(FRHIPixelShader* PixelShaderRHI,uint32 TextureIndex,FTextureRHIParamRef NewTextureRHI)
+void FD3D11DynamicRHI::RHISetShaderTexture(FRHIPixelShader* PixelShaderRHI,uint32 TextureIndex, FRHITexture* NewTextureRHI)
 {
 	VALIDATE_BOUND_SHADER(PixelShaderRHI);
 
@@ -419,7 +419,7 @@ void FD3D11DynamicRHI::RHISetShaderTexture(FRHIPixelShader* PixelShaderRHI,uint3
 	}
 }
 
-void FD3D11DynamicRHI::RHISetShaderTexture(FRHIComputeShader* ComputeShaderRHI,uint32 TextureIndex,FTextureRHIParamRef NewTextureRHI)
+void FD3D11DynamicRHI::RHISetShaderTexture(FRHIComputeShader* ComputeShaderRHI,uint32 TextureIndex, FRHITexture* NewTextureRHI)
 {
 	//VALIDATE_BOUND_SHADER(ComputeShaderRHI);
 
@@ -657,7 +657,7 @@ void FD3D11DynamicRHI::RHISetShaderSampler(FRHIComputeShader* ComputeShaderRHI,u
 	StateCache.SetSamplerState<SF_Compute>(StateResource, SamplerIndex);
 }
 
-void FD3D11DynamicRHI::RHISetShaderUniformBuffer(FRHIVertexShader* VertexShader,uint32 BufferIndex,FUniformBufferRHIParamRef BufferRHI)
+void FD3D11DynamicRHI::RHISetShaderUniformBuffer(FRHIVertexShader* VertexShader,uint32 BufferIndex, FRHIUniformBuffer* BufferRHI)
 {
 	check(BufferRHI->GetLayout().GetHash());
 	VALIDATE_BOUND_SHADER(VertexShader);
@@ -671,7 +671,7 @@ void FD3D11DynamicRHI::RHISetShaderUniformBuffer(FRHIVertexShader* VertexShader,
 	DirtyUniformBuffers[SF_Vertex] |= (1 << BufferIndex);
 }
 
-void FD3D11DynamicRHI::RHISetShaderUniformBuffer(FRHIHullShader* HullShader,uint32 BufferIndex,FUniformBufferRHIParamRef BufferRHI)
+void FD3D11DynamicRHI::RHISetShaderUniformBuffer(FRHIHullShader* HullShader,uint32 BufferIndex, FRHIUniformBuffer* BufferRHI)
 {
 	check(BufferRHI->GetLayout().GetHash());
 	VALIDATE_BOUND_SHADER(HullShader);
@@ -685,7 +685,7 @@ void FD3D11DynamicRHI::RHISetShaderUniformBuffer(FRHIHullShader* HullShader,uint
 	DirtyUniformBuffers[SF_Hull] |= (1 << BufferIndex);
 }
 
-void FD3D11DynamicRHI::RHISetShaderUniformBuffer(FRHIDomainShader* DomainShader,uint32 BufferIndex,FUniformBufferRHIParamRef BufferRHI)
+void FD3D11DynamicRHI::RHISetShaderUniformBuffer(FRHIDomainShader* DomainShader,uint32 BufferIndex, FRHIUniformBuffer* BufferRHI)
 {
 	check(BufferRHI->GetLayout().GetHash());
 	VALIDATE_BOUND_SHADER(DomainShader);
@@ -699,7 +699,7 @@ void FD3D11DynamicRHI::RHISetShaderUniformBuffer(FRHIDomainShader* DomainShader,
 	DirtyUniformBuffers[SF_Domain] |= (1 << BufferIndex);
 }
 
-void FD3D11DynamicRHI::RHISetShaderUniformBuffer(FRHIGeometryShader* GeometryShader,uint32 BufferIndex,FUniformBufferRHIParamRef BufferRHI)
+void FD3D11DynamicRHI::RHISetShaderUniformBuffer(FRHIGeometryShader* GeometryShader,uint32 BufferIndex, FRHIUniformBuffer* BufferRHI)
 {
 	check(BufferRHI->GetLayout().GetHash());
 	VALIDATE_BOUND_SHADER(GeometryShader);
@@ -713,7 +713,7 @@ void FD3D11DynamicRHI::RHISetShaderUniformBuffer(FRHIGeometryShader* GeometrySha
 	DirtyUniformBuffers[SF_Geometry] |= (1 << BufferIndex);
 }
 
-void FD3D11DynamicRHI::RHISetShaderUniformBuffer(FRHIPixelShader* PixelShader,uint32 BufferIndex,FUniformBufferRHIParamRef BufferRHI)
+void FD3D11DynamicRHI::RHISetShaderUniformBuffer(FRHIPixelShader* PixelShader,uint32 BufferIndex, FRHIUniformBuffer* BufferRHI)
 {
 	check(BufferRHI->GetLayout().GetHash());
 	VALIDATE_BOUND_SHADER(PixelShader);
@@ -728,7 +728,7 @@ void FD3D11DynamicRHI::RHISetShaderUniformBuffer(FRHIPixelShader* PixelShader,ui
 	DirtyUniformBuffers[SF_Pixel] |= (1 << BufferIndex);
 }
 
-void FD3D11DynamicRHI::RHISetShaderUniformBuffer(FRHIComputeShader* ComputeShader,uint32 BufferIndex,FUniformBufferRHIParamRef BufferRHI)
+void FD3D11DynamicRHI::RHISetShaderUniformBuffer(FRHIComputeShader* ComputeShader,uint32 BufferIndex, FRHIUniformBuffer* BufferRHI)
 {
 	check(BufferRHI->GetLayout().GetHash());
 	//VALIDATE_BOUND_SHADER(ComputeShader);
@@ -1601,7 +1601,7 @@ void FD3D11DynamicRHI::RHIDrawPrimitive(uint32 BaseVertexIndex,uint32 NumPrimiti
 	}
 }
 
-void FD3D11DynamicRHI::RHIDrawPrimitiveIndirect(FVertexBufferRHIParamRef ArgumentBufferRHI,uint32 ArgumentOffset)
+void FD3D11DynamicRHI::RHIDrawPrimitiveIndirect(FRHIVertexBuffer* ArgumentBufferRHI,uint32 ArgumentOffset)
 {
 	FD3D11VertexBuffer* ArgumentBuffer = ResourceCast(ArgumentBufferRHI);
 
@@ -1616,7 +1616,7 @@ void FD3D11DynamicRHI::RHIDrawPrimitiveIndirect(FVertexBufferRHIParamRef Argumen
 	Direct3DDeviceIMContext->DrawInstancedIndirect(ArgumentBuffer->Resource,ArgumentOffset);
 }
 
-void FD3D11DynamicRHI::RHIDrawIndexedIndirect(FIndexBufferRHIParamRef IndexBufferRHI, FStructuredBufferRHIParamRef ArgumentsBufferRHI, int32 DrawArgumentsIndex, uint32 NumInstances)
+void FD3D11DynamicRHI::RHIDrawIndexedIndirect(FRHIIndexBuffer* IndexBufferRHI, FRHIStructuredBuffer* ArgumentsBufferRHI, int32 DrawArgumentsIndex, uint32 NumInstances)
 {
 	FD3D11IndexBuffer* IndexBuffer = ResourceCast(IndexBufferRHI);
 	FD3D11StructuredBuffer* ArgumentsBuffer = ResourceCast(ArgumentsBufferRHI);
@@ -1646,7 +1646,7 @@ void FD3D11DynamicRHI::RHIDrawIndexedIndirect(FIndexBufferRHIParamRef IndexBuffe
 	}
 }
 
-void FD3D11DynamicRHI::RHIDrawIndexedPrimitive(FIndexBufferRHIParamRef IndexBufferRHI,int32 BaseVertexIndex,uint32 FirstInstance,uint32 NumVertices,uint32 StartIndex,uint32 NumPrimitives,uint32 NumInstances)
+void FD3D11DynamicRHI::RHIDrawIndexedPrimitive(FRHIIndexBuffer* IndexBufferRHI,int32 BaseVertexIndex,uint32 FirstInstance,uint32 NumVertices,uint32 StartIndex,uint32 NumPrimitives,uint32 NumInstances)
 {
 	FD3D11IndexBuffer* IndexBuffer = ResourceCast(IndexBufferRHI);
 
@@ -1687,7 +1687,7 @@ void FD3D11DynamicRHI::RHIDrawIndexedPrimitive(FIndexBufferRHIParamRef IndexBuff
 	}
 }
 
-void FD3D11DynamicRHI::RHIDrawIndexedPrimitiveIndirect(FIndexBufferRHIParamRef IndexBufferRHI,FVertexBufferRHIParamRef ArgumentBufferRHI,uint32 ArgumentOffset)
+void FD3D11DynamicRHI::RHIDrawIndexedPrimitiveIndirect(FRHIIndexBuffer* IndexBufferRHI, FRHIVertexBuffer* ArgumentBufferRHI,uint32 ArgumentOffset)
 {
 	FD3D11IndexBuffer* IndexBuffer = ResourceCast(IndexBufferRHI);
 	FD3D11VertexBuffer* ArgumentBuffer = ResourceCast(ArgumentBufferRHI);
@@ -1871,7 +1871,7 @@ IRHICommandContextContainer* FD3D11DynamicRHI::RHIGetCommandContextContainer(int
 	return nullptr;
 }
 
-void FD3D11DynamicRHI::RHITransitionResources(EResourceTransitionAccess TransitionType, FTextureRHIParamRef* InTextures, int32 NumTextures)
+void FD3D11DynamicRHI::RHITransitionResources(EResourceTransitionAccess TransitionType, FRHITexture** InTextures, int32 NumTextures)
 {
 	static IConsoleVariable* CVarShowTransitions = IConsoleManager::Get().FindConsoleVariable(TEXT("r.ProfileGPU.ShowTransitions"));
 	bool bShowTransitionEvents = CVarShowTransitions->GetInt() != 0;
@@ -1879,7 +1879,7 @@ void FD3D11DynamicRHI::RHITransitionResources(EResourceTransitionAccess Transiti
 	SCOPED_RHI_CONDITIONAL_DRAW_EVENTF(*this, RHITransitionResources, bShowTransitionEvents, TEXT("TransitionTo: %s: %i Textures"), *FResourceTransitionUtility::ResourceTransitionAccessStrings[(int32)TransitionType], NumTextures);
 	for (int32 i = 0; i < NumTextures; ++i)
 	{				
-		FTextureRHIParamRef RenderTarget = InTextures[i];
+		FRHITexture* RenderTarget = InTextures[i];
 		if (RenderTarget)
 		{
 			SCOPED_RHI_CONDITIONAL_DRAW_EVENTF(*this, RHITransitionResourcesLoop, bShowTransitionEvents, TEXT("To:%i - %s"), i, *RenderTarget->GetName().ToString());
@@ -2052,7 +2052,7 @@ FStagingBufferRHIRef FD3D11DynamicRHI::RHICreateStagingBuffer()
 	return new FD3D11StagingBuffer();
 }
 
-void FD3D11DynamicRHI::RHICopyToStagingBuffer(FVertexBufferRHIParamRef SourceBufferRHI, FRHIStagingBuffer* StagingBufferRHI, uint32 Offset, uint32 NumBytes)
+void FD3D11DynamicRHI::RHICopyToStagingBuffer(FRHIVertexBuffer* SourceBufferRHI, FRHIStagingBuffer* StagingBufferRHI, uint32 Offset, uint32 NumBytes)
 {
 	FD3D11VertexBuffer* SourceBuffer = ResourceCast(SourceBufferRHI);
 	FD3D11StagingBuffer* StagingBuffer = ResourceCast(StagingBufferRHI);
