fixed: Resolve inconsistencies in VulkanPipeline

diff --git a/Engine/Source/Runtime/VulkanRHI/Private/VulkanPipeline.cpp b/Engine/Source/Runtime/VulkanRHI/Private/VulkanPipeline.cpp
--- a/Engine/Source/Runtime/VulkanRHI/Private/VulkanPipeline.cpp
+++ b/Engine/Source/Runtime/VulkanRHI/Private/VulkanPipeline.cpp
@@ -1669,12 +1669,15 @@ FVulkanRHIGraphicsPipelineState* FVulkanPipelineStateCacheManager::FindInRuntime
 
 	#if VULKAN_SUPPORTS_GEOMETRY_SHADERS
 			TempUInt64 = GetShaderKey(PSI.BoundShaderState.GeometryShaderRHI);
-	#endif
 			Ar << TempUInt64;
+	#endif
+	
+	#if PLATFORM_SUPPORTS_TESSELLATION_SHADERS
 			TempUInt64 = GetShaderKey(PSI.BoundShaderState.HullShaderRHI);
 			Ar << TempUInt64;
 			TempUInt64 = GetShaderKey(PSI.BoundShaderState.DomainShaderRHI);
 			Ar << TempUInt64;
+	#endif
 			Ar << ResourceCast(PSI.BoundShaderState.VertexDeclarationRHI)->Elements;
 			Ar << ResourceCast(PSI.RasterizerState)->Initializer;
 			Ar << ResourceCast(PSI.DepthStencilState)->Initializer;
