Fixed: Crash after backgrounding during 'Launching...' phase with Vulkan Android

diff --git a/Engine/Source/Runtime/Core/Private/Android/AndroidPlatformMisc.cpp b/Engine/Source/Runtime/Core/Private/Android/AndroidPlatformMisc.cpp
--- a/Engine/Source/Runtime/Core/Private/Android/AndroidPlatformMisc.cpp
+++ b/Engine/Source/Runtime/Core/Private/Android/AndroidPlatformMisc.cpp
@@ -2179,6 +2179,18 @@ void FAndroidMisc::SetOnReInitWindowCallback(FAndroidMisc::ReInitWindowCallbackT
 	OnReInitWindowCallback = InOnReInitWindowCallback;
 }
 
+static FAndroidMisc::ReleaseWindowCallbackType OnReleaseWindowCallback;
+
+FAndroidMisc::ReleaseWindowCallbackType FAndroidMisc::GetOnReleaseWindowCallback()
+{
+	return OnReleaseWindowCallback;
+}
+
+void FAndroidMisc::SetOnReleaseWindowCallback(FAndroidMisc::ReleaseWindowCallbackType InOnReleaseWindowCallback)
+{
+	OnReleaseWindowCallback = InOnReleaseWindowCallback;
+}
+
 static FAndroidMisc::OnPauseCallBackType OnPauseCallback;
 
 FAndroidMisc::OnPauseCallBackType FAndroidMisc::GetOnPauseCallback()
diff --git a/Engine/Source/Runtime/Core/Public/Android/AndroidPlatformMisc.h b/Engine/Source/Runtime/Core/Public/Android/AndroidPlatformMisc.h
--- a/Engine/Source/Runtime/Core/Public/Android/AndroidPlatformMisc.h
+++ b/Engine/Source/Runtime/Core/Public/Android/AndroidPlatformMisc.h
@@ -202,6 +202,9 @@ public:
 	typedef TFunction<void(void* NewNativeHandle)> ReInitWindowCallbackType;
 	static ReInitWindowCallbackType GetOnReInitWindowCallback();
 	static void SetOnReInitWindowCallback(ReInitWindowCallbackType InOnReInitWindowCallback);
+	typedef TFunction<void()> ReleaseWindowCallbackType;
+	static ReleaseWindowCallbackType GetOnReleaseWindowCallback();
+	static void SetOnReleaseWindowCallback(ReleaseWindowCallbackType InOnReleaseWindowCallback);
 	static FString GetOSVersion();
 	static bool GetOverrideResolution(int32 &ResX, int32& ResY) { return false; }
 	typedef TFunction<void()> OnPauseCallBackType;
diff --git a/Engine/Source/Runtime/OpenGLDrv/Private/Android/AndroidEGL.cpp b/Engine/Source/Runtime/OpenGLDrv/Private/Android/AndroidEGL.cpp
--- a/Engine/Source/Runtime/OpenGLDrv/Private/Android/AndroidEGL.cpp
+++ b/Engine/Source/Runtime/OpenGLDrv/Private/Android/AndroidEGL.cpp
@@ -1726,6 +1726,14 @@ void BlockOnLostWindowRenderCommand(TSharedPtr<FEvent, ESPMode::ThreadSafe> RTBl
 			RHICmdList.ImmediateFlush(EImmediateFlushType::FlushRHIThread);
 			UE_LOG(LogAndroid, Log, TEXT("RendererBlock DONE FlushRHIThread"));
 		}
+		
+		const auto& OnReleaseWindowCallback = FAndroidMisc::GetOnReleaseWindowCallback();
+		if (OnReleaseWindowCallback)
+		{
+			UE_LOG(LogAndroid, Log, TEXT("RendererBlock release window callback"));
+			OnReleaseWindowCallback();
+		}
+
 		RTBlockedTrigger->Trigger();
 
 		GAndroidWindowLock.Lock();
diff --git a/Engine/Source/Runtime/VulkanRHI/Private/Android/VulkanAndroidPlatform.cpp b/Engine/Source/Runtime/VulkanRHI/Private/Android/VulkanAndroidPlatform.cpp
--- a/Engine/Source/Runtime/VulkanRHI/Private/Android/VulkanAndroidPlatform.cpp
+++ b/Engine/Source/Runtime/VulkanRHI/Private/Android/VulkanAndroidPlatform.cpp
@@ -184,30 +184,16 @@ void FVulkanAndroidPlatform::OverridePlatformHandlers(bool bInit)
 		// Want to see the actual crash report on Android so unregister signal handlers
 		FPlatformMisc::SetCrashHandler((void(*)(const FGenericCrashContext& Context)) -1);
 		FPlatformMisc::SetOnReInitWindowCallback(FVulkanDynamicRHI::RecreateSwapChain);
+		FPlatformMisc::SetOnReleaseWindowCallback(FVulkanDynamicRHI::DestroySwapChain);
 		FPlatformMisc::SetOnPauseCallback(FVulkanDynamicRHI::SavePipelineCache);
 	}
 	else
 	{
 		FPlatformMisc::SetCrashHandler(nullptr);
 		FPlatformMisc::SetOnReInitWindowCallback(nullptr);
+		FPlatformMisc::SetOnReleaseWindowCallback(nullptr);
 		FPlatformMisc::SetOnPauseCallback(nullptr);
 	}
 }
 
-void FVulkanAndroidPlatform::BlockUntilWindowIsAvailable()
-{
-	void* WindowHandle = FAndroidWindow::GetHardwareWindow_EventThread();
-	if (WindowHandle == nullptr)
-	{
-		// Sleep if the hardware window isn't currently available.
-		// The Window may not exist if the activity is pausing/resuming, in which case we make this thread wait
-		FPlatformMisc::LowLevelOutputDebugString(TEXT("Waiting for Native window in FVulkanAndroidPlatform::BlockUntilWindowIsAvailable()"));
-		WindowHandle = FAndroidWindow::WaitForHardwareWindow();
-		if (WindowHandle == NULL)
-		{
-			FPlatformMisc::LowLevelOutputDebugString(TEXT("Aborted FVulkanAndroidPlatform::BlockUntilWindowIsAvailable, FAndroidWindow::WaitForHardwareWindow() returned null"));
-		}
-	}
-}
-
 #endif
diff --git a/Engine/Source/Runtime/VulkanRHI/Private/Android/VulkanAndroidPlatform.h b/Engine/Source/Runtime/VulkanRHI/Private/Android/VulkanAndroidPlatform.h
--- a/Engine/Source/Runtime/VulkanRHI/Private/Android/VulkanAndroidPlatform.h
+++ b/Engine/Source/Runtime/VulkanRHI/Private/Android/VulkanAndroidPlatform.h
@@ -96,8 +96,6 @@ public:
 
 	static bool RequiresWaitingForFrameCompletionEvent() { return false; }
 	
-	static void BlockUntilWindowIsAvailable();
-
 	// Does the platform allow a nullptr Pixelshader on the pipeline
 	static bool SupportsNullPixelShader() { return false; }
 
diff --git a/Engine/Source/Runtime/VulkanRHI/Private/VulkanGenericPlatform.h b/Engine/Source/Runtime/VulkanRHI/Private/VulkanGenericPlatform.h
--- a/Engine/Source/Runtime/VulkanRHI/Private/VulkanGenericPlatform.h
+++ b/Engine/Source/Runtime/VulkanRHI/Private/VulkanGenericPlatform.h
@@ -99,9 +99,6 @@ public:
 	// Ensure the last frame completed on the GPU
 	static bool RequiresWaitingForFrameCompletionEvent() { return true; }
 
-	// Blocks until hardware window is available
-	static void BlockUntilWindowIsAvailable() {}
-
 	// Does the platform allow a nullptr Pixelshader on the pipeline
 	static bool SupportsNullPixelShader() { return true; }
 
diff --git a/Engine/Source/Runtime/VulkanRHI/Private/VulkanRHI.cpp b/Engine/Source/Runtime/VulkanRHI/Private/VulkanRHI.cpp
--- a/Engine/Source/Runtime/VulkanRHI/Private/VulkanRHI.cpp
+++ b/Engine/Source/Runtime/VulkanRHI/Private/VulkanRHI.cpp
@@ -1723,6 +1723,32 @@ void FVulkanDynamicRHI::DumpMemory()
 }
 #endif
 
+void FVulkanDynamicRHI::DestroySwapChain()
+{
+	if (IsInGameThread())
+	{
+		FlushRenderingCommands();
+	}
+
+	FVulkanDynamicRHI* RHI = (FVulkanDynamicRHI*)GDynamicRHI;
+	TArray<FVulkanViewport*> Viewports = RHI->Viewports;
+	ENQUEUE_RENDER_COMMAND(VulkanDestroySwapChain)(
+		[Viewports](FRHICommandListImmediate& RHICmdList)
+	{
+		UE_LOG(LogVulkanRHI, Log, TEXT("Destroy swapchain ... "));
+		
+		for (auto& Viewport : Viewports)
+		{
+			Viewport->DestroySwapchain();
+		}
+	});
+
+	if (IsInGameThread())
+	{
+		FlushRenderingCommands();
+	}
+}
+
 void FVulkanDynamicRHI::RecreateSwapChain(void* NewNativeWindow)
 {
 	if (NewNativeWindow)
@@ -1736,6 +1762,8 @@ void FVulkanDynamicRHI::RecreateSwapChain(void* NewNativeWindow)
 		ENQUEUE_RENDER_COMMAND(VulkanRecreateSwapChain)(
 			[Viewports, NewNativeWindow](FRHICommandListImmediate& RHICmdList)
 		{
+			UE_LOG(LogVulkanRHI, Log, TEXT("Recreate swapchain ... "));
+			
 			for (auto& Viewport : Viewports)
 			{
 				Viewport->RecreateSwapchain(NewNativeWindow);
diff --git a/Engine/Source/Runtime/VulkanRHI/Private/VulkanViewport.cpp b/Engine/Source/Runtime/VulkanRHI/Private/VulkanViewport.cpp
--- a/Engine/Source/Runtime/VulkanRHI/Private/VulkanViewport.cpp
+++ b/Engine/Source/Runtime/VulkanRHI/Private/VulkanViewport.cpp
@@ -191,9 +191,6 @@
 
 	while (Status < 0 && AttemptsPending > 0)
 	{
-		// always force to recreate swapchain, on Android it will block until window is available
-		bool bForce = true;
-
 		if (Status == (int32)FVulkanSwapChain::EStatus::OutOfDate)
 		{
 			UE_LOG(LogVulkanRHI, Verbose, TEXT("Swapchain is out of date! Trying to recreate the swapchain."));
@@ -207,7 +204,7 @@
 			check(0);
 		}
 
-		RecreateSwapchain(WindowHandle, bForce);
+		RecreateSwapchain(WindowHandle);
 
 		// Swapchain creation pushes some commands - flush the command buffers now to begin with a fresh state
 		Device->SubmitCommandsAndFlushGPU();
@@ -502,42 +499,11 @@
 }
 
 // Tear down and recreate swapchain and related resources.
-void FVulkanViewport::RecreateSwapchain(void* NewNativeWindow, bool bForce)
+void FVulkanViewport::RecreateSwapchain(void* NewNativeWindow)
 {
-	if (WindowHandle == NewNativeWindow && !bForce)
-	{
-		// No action is required if handle has not changed.
-		return;
-	}
-
 	FScopeLock LockSwapchain(&RecreatingSwapchain);
-	RenderingBackBuffer = nullptr;
-	
-	if (RHIBackBuffer)
-	{
-		RHIBackBuffer->ReleaseViewport();
-		RHIBackBuffer = nullptr;
-	}
-	
-	if (FVulkanPlatform::SupportsStandardSwapchain())
-	{
-		for (int32 Index = 0; Index < NUM_BUFFERS; ++Index)
-		{
-			TextureViews[Index].Destroy(*Device);
-			Device->NotifyDeletedImage(BackBufferImages[Index]);
-			Device->NotifyDeletedRenderTarget(BackBufferImages[Index]);
-			BackBufferImages[Index] = VK_NULL_HANDLE;
-		}
 
-		Device->GetDeferredDeletionQueue().ReleaseResources(true);
-
-		SwapChain->Destroy();
-		delete SwapChain;
-		SwapChain = nullptr;
-
-		Device->GetDeferredDeletionQueue().ReleaseResources(true);
-	}
-
+	DestroySwapchain();
 	WindowHandle = NewNativeWindow;
 	CreateSwapchain();
 }
@@ -570,39 +536,10 @@
 void FVulkanViewport::RecreateSwapchainFromRT(EPixelFormat PreferredPixelFormat)
 {
 	check(IsInRenderingThread());
+
+	// TODO: should flush RHIT commands here?
 	
-	// Submit all command buffers here
-	Device->SubmitCommandsAndFlushGPU();
-
-	Device->WaitUntilIdle();
-
-	RenderingBackBuffer = nullptr;
-	
-	if (RHIBackBuffer)
-	{
-		RHIBackBuffer->ReleaseViewport();
-		RHIBackBuffer = nullptr;
-	}
-		
-	if (FVulkanPlatform::SupportsStandardSwapchain())
-	{
-		for (int32 Index = 0; Index < NUM_BUFFERS; ++Index)
-		{
-			TextureViews[Index].Destroy(*Device);
-			Device->NotifyDeletedImage(BackBufferImages[Index]);
-			Device->NotifyDeletedRenderTarget(BackBufferImages[Index]);
-			BackBufferImages[Index] = VK_NULL_HANDLE;
-		}
-		
-		Device->GetDeferredDeletionQueue().ReleaseResources(true);
-
-		SwapChain->Destroy();
-		delete SwapChain;
-		SwapChain = nullptr;
-
-		Device->GetDeferredDeletionQueue().ReleaseResources(true);
-	}
-
+	DestroySwapchain();
 	PixelFormat = PreferredPixelFormat;
 	CreateSwapchain();
 }
@@ -680,6 +617,42 @@
 	AcquiredImageIndex = -1;
 }
 
+void FVulkanViewport::DestroySwapchain()
+{
+	// Submit all command buffers here
+	Device->SubmitCommandsAndFlushGPU();
+	Device->WaitUntilIdle();
+	
+	RenderingBackBuffer = nullptr;
+	
+	if (RHIBackBuffer)
+	{
+		RHIBackBuffer->ReleaseViewport();
+		RHIBackBuffer = nullptr;
+	}
+		
+	if (FVulkanPlatform::SupportsStandardSwapchain())
+	{
+		for (int32 Index = 0; Index < NUM_BUFFERS; ++Index)
+		{
+			TextureViews[Index].Destroy(*Device);
+			Device->NotifyDeletedImage(BackBufferImages[Index]);
+			Device->NotifyDeletedRenderTarget(BackBufferImages[Index]);
+			BackBufferImages[Index] = VK_NULL_HANDLE;
+		}
+		
+		Device->GetDeferredDeletionQueue().ReleaseResources(true);
+
+		SwapChain->Destroy();
+		delete SwapChain;
+		SwapChain = nullptr;
+
+		Device->GetDeferredDeletionQueue().ReleaseResources(true);
+	}
+
+	AcquiredImageIndex = -1;
+}
+
 inline static void CopyImageToBackBuffer(FVulkanCmdBuffer* CmdBuffer, bool bSourceReadOnly, VkImage SrcSurface, VkImage DstSurface, int32 SizeX, int32 SizeY, int32 WindowSizeX, int32 WindowSizeY)
 {
 	VulkanRHI::FPendingBarrier Barriers;
@@ -798,9 +771,6 @@
 		Queue->Submit(CmdBuffer);
 	}
 
-	// Do not present until hardware window is available. On Android window could be destroyed while RHIT executes commands
-	FVulkanPlatform::BlockUntilWindowIsAvailable();
-
 	//Flush all commands
 	//check(0);
 
diff --git a/Engine/Source/Runtime/VulkanRHI/Private/VulkanViewport.h b/Engine/Source/Runtime/VulkanRHI/Private/VulkanViewport.h
--- a/Engine/Source/Runtime/VulkanRHI/Private/VulkanViewport.h
+++ b/Engine/Source/Runtime/VulkanRHI/Private/VulkanViewport.h
@@ -113,9 +113,10 @@ protected:
 	uint64 LastFrameFenceCounter = 0;
 
 	void CreateSwapchain();
+	void DestroySwapchain();
 	void AcquireImageIndex();
 
-	void RecreateSwapchain(void* NewNativeWindow, bool bForce = false);
+	void RecreateSwapchain(void* NewNativeWindow);
 	void RecreateSwapchainFromRT(EPixelFormat PreferredPixelFormat);
 	void Resize(uint32 InSizeX, uint32 InSizeY, bool bIsFullscreen, EPixelFormat PreferredPixelFormat);
 
diff --git a/Engine/Source/Runtime/VulkanRHI/Public/VulkanDynamicRHI.h b/Engine/Source/Runtime/VulkanRHI/Public/VulkanDynamicRHI.h
--- a/Engine/Source/Runtime/VulkanRHI/Public/VulkanDynamicRHI.h
+++ b/Engine/Source/Runtime/VulkanRHI/Public/VulkanDynamicRHI.h
@@ -350,6 +350,7 @@ public:
 		return InstanceLayers;
 	}
 
+	static void DestroySwapChain();
 	static void RecreateSwapChain(void* NewNativeWindow);
 
 	inline VkInstance GetInstance() const
