DR - Fix ensure when generating mips on Vulkan
- Also remove unused class
#rnx

diff --git a/Engine/Source/Runtime/RHI/Public/RHICommandList.h b/Engine/Source/Runtime/RHI/Public/RHICommandList.h
--- a/Engine/Source/Runtime/RHI/Public/RHICommandList.h
+++ b/Engine/Source/Runtime/RHI/Public/RHICommandList.h
@@ -1121,18 +1121,6 @@ struct FRHICommandEndComputePass final : public FRHICommand<FRHICommandEndComput
 	RHI_API void Execute(FRHICommandListBase& CmdList);
 };
 
-struct FRHICommandSetRenderTargetsAndClear final : public FRHICommand<FRHICommandSetRenderTargetsAndClear>
-{
-	FRHISetRenderTargetsInfo RenderTargetsInfo;
-
-	FORCEINLINE_DEBUGGABLE FRHICommandSetRenderTargetsAndClear(const FRHISetRenderTargetsInfo& InRenderTargetsInfo) :
-		RenderTargetsInfo(InRenderTargetsInfo)
-	{
-	}
-
-	RHI_API void Execute(FRHICommandListBase& CmdList);
-};
-
 struct FRHICommandBindClearMRTValues final : public FRHICommand<FRHICommandBindClearMRTValues>
 {
 	bool bClearColor;
diff --git a/Engine/Source/Runtime/RHI/Public/RHICommandListCommandExecutes.inl b/Engine/Source/Runtime/RHI/Public/RHICommandListCommandExecutes.inl
--- a/Engine/Source/Runtime/RHI/Public/RHICommandListCommandExecutes.inl
+++ b/Engine/Source/Runtime/RHI/Public/RHICommandListCommandExecutes.inl
@@ -292,12 +292,6 @@ void FRHICommandSetRenderTargets::Execute(FRHICommandListBase& CmdList)
 		UAVs);
 }
 
-void FRHICommandSetRenderTargetsAndClear::Execute(FRHICommandListBase& CmdList)
-{
-	RHISTAT(SetRenderTargetsAndClear);
-	INTERNAL_DECORATOR(RHISetRenderTargetsAndClear)(RenderTargetsInfo);
-}
-
 void FRHICommandBindClearMRTValues::Execute(FRHICommandListBase& CmdList)
 {
 	RHISTAT(BindClearMRTValues);
diff --git a/Engine/Source/Runtime/VulkanRHI/Private/VulkanRenderTarget.cpp b/Engine/Source/Runtime/VulkanRHI/Private/VulkanRenderTarget.cpp
--- a/Engine/Source/Runtime/VulkanRHI/Private/VulkanRenderTarget.cpp
+++ b/Engine/Source/Runtime/VulkanRHI/Private/VulkanRenderTarget.cpp
@@ -279,6 +279,7 @@ void FTransitionAndLayoutManager::BeginRealRenderPass(FVulkanCommandListContext&
 
 				if (*Found != VK_IMAGE_LAYOUT_SHADER_READ_ONLY_OPTIMAL)
 				{
+					// This is since the previous mip index is used as a source image for the descriptor, it needs to know it's in R/O state
 					*Found = VK_IMAGE_LAYOUT_SHADER_READ_ONLY_OPTIMAL;
 				}
 			}
@@ -291,9 +292,23 @@ void FTransitionAndLayoutManager::BeginRealRenderPass(FVulkanCommandListContext&
 			GenerateMipsInfo.CurrentSlice = SliceIndex;
 			GenerateMipsInfo.CurrentMip = RTMipIndex;
 			GenerateMipsInfo.bLastMip = (RTMipIndex == (NumMips - 1));
+
+			// Check that previous mip is read only
+			if (GenerateMipsInfo.Target[Index].Layouts[SliceIndex][RTMipIndex - 1] != VK_IMAGE_LAYOUT_SHADER_READ_ONLY_OPTIMAL)
+			{
+				// Transition to readable
+				int32 BarrierIndex = Barrier.AddImageBarrier(Surface.Image, VK_IMAGE_ASPECT_COLOR_BIT, 1);
+				VkImageSubresourceRange& Range = Barrier.GetSubresource(BarrierIndex);
+				Range.baseMipLevel = RTMipIndex - 1;
+				Range.baseArrayLayer = SliceIndex;
 #if !USING_CODE_ANALYSIS
-			ensure(GenerateMipsInfo.Target[Index].Layouts[SliceIndex][RTMipIndex - 1] == VK_IMAGE_LAYOUT_SHADER_READ_ONLY_OPTIMAL);
+				ensure(GenerateMipsInfo.Target[Index].Layouts[SliceIndex][RTMipIndex - 1] == VK_IMAGE_LAYOUT_COLOR_ATTACHMENT_OPTIMAL);
 #endif
+				Barrier.SetTransition(BarrierIndex, EImageLayoutBarrier::ColorAttachment, EImageLayoutBarrier::PixelShaderRead);
+				GenerateMipsInfo.Target[Index].Layouts[SliceIndex][RTMipIndex - 1] = VK_IMAGE_LAYOUT_SHADER_READ_ONLY_OPTIMAL;
+			}
+
+			// Check that current mip is write-only
 			if (GenerateMipsInfo.Target[Index].Layouts[SliceIndex][RTMipIndex] != VK_IMAGE_LAYOUT_COLOR_ATTACHMENT_OPTIMAL)
 			{
 				// Transition to writeable
