increase default wait timeout. on some hw, you might hit a 1 sec timeout so it needed to be bumped. (Suggested by [REDACTED_PERSON]/ARM)

diff --git a/Engine/Source/Runtime/VulkanRHI/Private/VulkanCommandBuffer.h b/Engine/Source/Runtime/VulkanRHI/Private/VulkanCommandBuffer.h
--- a/Engine/Source/Runtime/VulkanRHI/Private/VulkanCommandBuffer.h
+++ b/Engine/Source/Runtime/VulkanRHI/Private/VulkanCommandBuffer.h
@@ -306,7 +306,7 @@ public:
 
 	void SubmitActiveCmdBuffer(VulkanRHI::FSemaphore* SignalSemaphore = nullptr);
 
-	void WaitForCmdBuffer(FVulkanCmdBuffer* CmdBuffer, float TimeInSecondsToWait = 1.0f);
+	void WaitForCmdBuffer(FVulkanCmdBuffer* CmdBuffer, float TimeInSecondsToWait = 10.0f);
 
 	// Update the fences of all cmd buffers except SkipCmdBuffer
 	void RefreshFenceStatus(FVulkanCmdBuffer* SkipCmdBuffer = nullptr)
