VK: put the resolve image pointer inside of ColorResolveTargetImages instead of the color one (vk framebuffer matching)

diff --git a/Engine/Source/Runtime/VulkanRHI/Private/VulkanViewport.cpp b/Engine/Source/Runtime/VulkanRHI/Private/VulkanViewport.cpp
--- a/Engine/Source/Runtime/VulkanRHI/Private/VulkanViewport.cpp
+++ b/Engine/Source/Runtime/VulkanRHI/Private/VulkanViewport.cpp
@@ -361,8 +361,8 @@ FVulkanFramebuffer::FVulkanFramebuffer(FVulkanDevice& Device, const FRHISetRende
 		if (InRTInfo.bHasResolveAttachments)
 		{
 			FRHITexture* ResolveRHITexture = InRTInfo.ColorResolveRenderTarget[Index].Texture;
-			ColorResolveTargetImages[Index] = Texture->Surface.Image;
 			FVulkanTextureBase* ResolveTexture = FVulkanTextureBase::Cast(ResolveRHITexture);
+			ColorResolveTargetImages[Index] = ResolveTexture->Surface.Image;
 
 			//resolve attachments only supported for 2d/2d array textures
 			FVulkanTextureView ResolveRTView;
