Merging CL 4373138 from Dev-VR - Fix for media texture crash on Lumin

diff --git a/Engine/Source/Runtime/VulkanRHI/Private/VulkanTexture.cpp b/Engine/Source/Runtime/VulkanRHI/Private/VulkanTexture.cpp
--- a/Engine/Source/Runtime/VulkanRHI/Private/VulkanTexture.cpp
+++ b/Engine/Source/Runtime/VulkanRHI/Private/VulkanTexture.cpp
@@ -1643,6 +1643,8 @@ FVulkanTextureBase::FVulkanTextureBase(FVulkanDevice& Device, VkImageViewType Re
 		PartialView = new FVulkanTextureView;
 		PartialView->Create(Device, Surface.Image, Surface.ViewType, Surface.PartialAspectMask, Surface.PixelFormat, Surface.ViewFormat, 0, FMath::Max(NumMips, 1u), 0, FMath::Max(1u, SizeZ), ConversionInitializer);
 	}
+
+	Device.GetImmediateContext().FindOrAddLayout(InImage, VK_IMAGE_LAYOUT_SHADER_READ_ONLY_OPTIMAL);
 }
 
 FVulkanTextureBase::~FVulkanTextureBase()
