UE4.21 - Copy fix from 4415584 and 4401641
- Added support for -vulkanpresentmode=N where N is a VkPresentModeKHR
- This prevents the CPU from getting more than one frame ahead of the GPU. In addition to making the Vulkan RHI more consistent with the DX RHIs and reducing game input latency, this also fixes the issue where memory for volatile uniform buffers could be overwritten by the CPU before the GPU was done with it. This could happen when running in immediate present mode with frame pacing and smoothing disabled, as the CPU could get several frames ahead of the GPU.

diff --git a/Engine/Source/Runtime/VulkanRHI/Private/Android/VulkanAndroidPlatform.h b/Engine/Source/Runtime/VulkanRHI/Private/Android/VulkanAndroidPlatform.h
--- a/Engine/Source/Runtime/VulkanRHI/Private/Android/VulkanAndroidPlatform.h
+++ b/Engine/Source/Runtime/VulkanRHI/Private/Android/VulkanAndroidPlatform.h
@@ -85,6 +85,8 @@ public:
 	//#todo-rco: Detect Mali? Doing a clear on ColorAtt layout on empty cmd buffer causes issues
 	static bool RequiresSwapchainGeneralInitialLayout() { return true; }
 
+	static bool RequiresWaitingForFrameCompletionEvent() { return false; }
+
 protected:
 	static void* VulkanLib;
 	static bool bAttemptedLoad;
diff --git a/Engine/Source/Runtime/VulkanRHI/Private/Lumin/VulkanLuminPlatform.h b/Engine/Source/Runtime/VulkanRHI/Private/Lumin/VulkanLuminPlatform.h
--- a/Engine/Source/Runtime/VulkanRHI/Private/Lumin/VulkanLuminPlatform.h
+++ b/Engine/Source/Runtime/VulkanRHI/Private/Lumin/VulkanLuminPlatform.h
@@ -64,6 +64,8 @@ public:
 
 	static void EnablePhysicalDeviceFeatureExtensions(VkDeviceCreateInfo& DeviceInfo);
 
+	static bool RequiresWaitingForFrameCompletionEvent() { return false; }
+
 protected:
 	static void* VulkanLib;
 	static bool bAttemptedLoad;
diff --git a/Engine/Source/Runtime/VulkanRHI/Private/VulkanCommandBuffer.cpp b/Engine/Source/Runtime/VulkanRHI/Private/VulkanCommandBuffer.cpp
--- a/Engine/Source/Runtime/VulkanRHI/Private/VulkanCommandBuffer.cpp
+++ b/Engine/Source/Runtime/VulkanRHI/Private/VulkanCommandBuffer.cpp
@@ -28,10 +28,6 @@ static FAutoConsoleVariableRef CVarVulkanProfileCmdBuffers(
 
 #define CMD_BUFFER_TIME_TO_WAIT_BEFORE_DELETING		10
 
-#if UE_BUILD_DEBUG || UE_BUILD_DEVELOPMENT
-extern TAutoConsoleVariable<int32> CVarVulkanDebugBarrier;
-#endif
-
 const uint32 GNumberOfFramesBeforeDeletingDescriptorPool = 300;
 
 FVulkanCmdBuffer::FVulkanCmdBuffer(FVulkanDevice* InDevice, FVulkanCommandBufferPool* InCommandBufferPool, bool bInIsUploadOnly)
@@ -342,12 +338,7 @@ void FVulkanCommandBufferManager::SubmitUploadCmdBuffer(uint32 NumSignalSemaphor
 	{
 		check(UploadCmdBuffer->IsOutsideRenderPass());
 
-#if UE_BUILD_DEBUG || UE_BUILD_DEVELOPMENT
-		if (CVarVulkanDebugBarrier.GetValueOnAnyThread() & 4)
-		{
-			VulkanRHI::InsertHeavyWeightBarrier(UploadCmdBuffer->GetHandle());
-		}
-#endif
+		VulkanRHI::DebugHeavyWeightBarrier(UploadCmdBuffer->GetHandle(), 4);
 
 		UploadCmdBuffer->End();
 		Queue->Submit(UploadCmdBuffer, NumSignalSemaphores, SignalSemaphores);
@@ -370,12 +361,7 @@ void FVulkanCommandBufferManager::SubmitActiveCmdBuffer(VulkanRHI::FSemaphore* S
 			ActiveCmdBuffer->EndRenderPass();
 		}
 
-#if UE_BUILD_DEBUG || UE_BUILD_DEVELOPMENT
-		if (CVarVulkanDebugBarrier.GetValueOnAnyThread() & 8)
-		{
-			VulkanRHI::InsertHeavyWeightBarrier(ActiveCmdBuffer->GetHandle());
-		}
-#endif
+		VulkanRHI::DebugHeavyWeightBarrier(ActiveCmdBuffer->GetHandle(), 8);
 
 		ActiveCmdBuffer->End();
 		if (SignalSemaphore)
diff --git a/Engine/Source/Runtime/VulkanRHI/Private/VulkanCommands.cpp b/Engine/Source/Runtime/VulkanRHI/Private/VulkanCommands.cpp
--- a/Engine/Source/Runtime/VulkanRHI/Private/VulkanCommands.cpp
+++ b/Engine/Source/Runtime/VulkanRHI/Private/VulkanCommands.cpp
@@ -28,10 +28,6 @@ static FAutoConsoleVariableRef CVarVulkanSubmitAfterEveryEndRenderPass(
 	ECVF_Default
 );
 
-#if UE_BUILD_DEBUG || UE_BUILD_DEVELOPMENT
-extern TAutoConsoleVariable<int32> CVarVulkanDebugBarrier;
-#endif
-
 // make sure what the hardware expects matches what we give it for indirect arguments
 static_assert(sizeof(FRHIDrawIndirectParameters) == sizeof(VkDrawIndirectCommand), "FRHIDrawIndirectParameters size is wrong.");
 static_assert(STRUCT_OFFSET(FRHIDrawIndirectParameters, VertexCountPerInstance) == STRUCT_OFFSET(VkDrawIndirectCommand, vertexCount), "Wrong offset of FRHIDrawIndirectParameters::VertexCountPerInstance.");
@@ -130,12 +126,7 @@ void FVulkanCommandListContext::RHIDispatchComputeShader(uint32 ThreadGroupCount
 		GpuProfiler.RegisterGPUWork(1);
 	}
 
-#if UE_BUILD_DEBUG || UE_BUILD_DEVELOPMENT
-	if (CVarVulkanDebugBarrier.GetValueOnRenderThread() & 2)
-	{
-		VulkanRHI::InsertHeavyWeightBarrier(CmdBuffer);
-	}
-#endif
+	VulkanRHI::DebugHeavyWeightBarrier(CmdBuffer, 2);
 }
 
 void FVulkanCommandListContext::RHIDispatchIndirectComputeShader(FVertexBufferRHIParamRef ArgumentBufferRHI, uint32 ArgumentOffset) 
@@ -165,12 +156,7 @@ void FVulkanCommandListContext::RHIDispatchIndirectComputeShader(FVertexBufferRH
 		GpuProfiler.RegisterGPUWork(1);
 	}
 
-#if UE_BUILD_DEBUG || UE_BUILD_DEVELOPMENT
-	if (CVarVulkanDebugBarrier.GetValueOnRenderThread() & 2)
-	{
-		VulkanRHI::InsertHeavyWeightBarrier(CmdBuffer);
-	}
-#endif
+	VulkanRHI::DebugHeavyWeightBarrier(CmdBuffer, 2);
 }
 
 void FVulkanCommandListContext::RHISetBoundShaderState(FBoundShaderStateRHIParamRef BoundShaderStateRHI)
diff --git a/Engine/Source/Runtime/VulkanRHI/Private/VulkanDebug.cpp b/Engine/Source/Runtime/VulkanRHI/Private/VulkanDebug.cpp
--- a/Engine/Source/Runtime/VulkanRHI/Private/VulkanDebug.cpp
+++ b/Engine/Source/Runtime/VulkanRHI/Private/VulkanDebug.cpp
@@ -190,6 +190,11 @@ static VkBool32 DebugUtilsCallback(VkDebugUtilsMessageSeverityFlagBitsEXT MsgSev
 			// Warning: *** [Warning:Validation-1(UNASSIGNED-CoreValidation-Shader-OutputNotConsumed)] fragment shader writes to output location 0 with no matching attachment
 			return VK_FALSE;
 		}
+		else if (!FCStringAnsi::Strcmp(CallbackData->pMessageIdName, "VUID-VkSwapchainCreateInfoKHR-imageExtent-01274"))
+		{
+			// Warning: *** [Error:Validation341838324(VUID-VkSwapchainCreateInfoKHR-imageExtent-01274)] vkCreateSwapChainKHR() called with imageExtent = (8,8), which is outside the bounds returned by vkGetPhysicalDeviceSurfaceCapabilitiesKHR(): currentExtent = (0,0), minImageExtent = (0,0), maxImageExtent = (0,0).
+			return VK_FALSE;
+		}
 	}
 	else
 	{
@@ -309,8 +314,6 @@ static VkBool32 DebugUtilsCallback(VkDebugUtilsMessageSeverityFlagBitsEXT MsgSev
 
 	return VK_FALSE;
 }
-
-#undef VULKAN_REPORT_LOG
 #endif
 
 void FVulkanDynamicRHI::SetupDebugLayerCallback()
diff --git a/Engine/Source/Runtime/VulkanRHI/Private/VulkanDevice.h b/Engine/Source/Runtime/VulkanRHI/Private/VulkanDevice.h
--- a/Engine/Source/Runtime/VulkanRHI/Private/VulkanDevice.h
+++ b/Engine/Source/Runtime/VulkanRHI/Private/VulkanDevice.h
@@ -357,7 +357,9 @@ private:
 	FVulkanCommandListContextImmediate* ImmediateContext;
 	FVulkanCommandListContext* ComputeContext;
 	TArray<FVulkanCommandListContext*> CommandContexts;
+#if VULKAN_SUPPORTS_COLOR_CONVERSIONS
 	TMap<uint32, VkSamplerYcbcrConversion> SamplerColorConversionMap;
+#endif
 
 	void GetDeviceExtensionsAndLayers(TArray<const ANSICHAR*>& OutDeviceExtensions, TArray<const ANSICHAR*>& OutDeviceLayers, bool& bOutDebugMarkers);
 
diff --git a/Engine/Source/Runtime/VulkanRHI/Private/VulkanGenericPlatform.h b/Engine/Source/Runtime/VulkanRHI/Private/VulkanGenericPlatform.h
--- a/Engine/Source/Runtime/VulkanRHI/Private/VulkanGenericPlatform.h
+++ b/Engine/Source/Runtime/VulkanRHI/Private/VulkanGenericPlatform.h
@@ -97,4 +97,7 @@ public:
 	
 	// Allow platforms to add extension features to the PresentInfo pNext chain
 	static void EnablePresentInfoExtensions(VkPresentInfoKHR& PresentInfo) {}
+
+	// Ensure the last frame completed on the GPU
+	static bool RequiresWaitingForFrameCompletionEvent() { return true; }
 };
diff --git a/Engine/Source/Runtime/VulkanRHI/Private/VulkanIndexBuffer.cpp b/Engine/Source/Runtime/VulkanRHI/Private/VulkanIndexBuffer.cpp
--- a/Engine/Source/Runtime/VulkanRHI/Private/VulkanIndexBuffer.cpp
+++ b/Engine/Source/Runtime/VulkanRHI/Private/VulkanIndexBuffer.cpp
@@ -10,6 +10,7 @@
 #include "Containers/ResourceArray.h"
 #include "VulkanLLM.h"
 
+
 static TMap<FVulkanResourceMultiBuffer*, VulkanRHI::FPendingBufferLock> GPendingLockIBs;
 static FCriticalSection GPendingLockIBsMutex;
 
@@ -233,6 +234,8 @@ inline void FVulkanResourceMultiBuffer::InternalUnlock(FVulkanCommandListContext
 	ensure(Cmd->IsOutsideRenderPass());
 	VkCommandBuffer CmdBuffer = Cmd->GetHandle();
 
+	VulkanRHI::DebugHeavyWeightBarrier(CmdBuffer, 16);
+
 	VkBufferCopy Region;
 	FMemory::Memzero(Region);
 	Region.size = LockSize;
diff --git a/Engine/Source/Runtime/VulkanRHI/Private/VulkanRHIPrivate.h b/Engine/Source/Runtime/VulkanRHI/Private/VulkanRHIPrivate.h
--- a/Engine/Source/Runtime/VulkanRHI/Private/VulkanRHIPrivate.h
+++ b/Engine/Source/Runtime/VulkanRHI/Private/VulkanRHIPrivate.h
@@ -751,6 +751,10 @@ static inline VkPrimitiveTopology UEToVulkanType(EPrimitiveType PrimitiveType)
 	return VK_PRIMITIVE_TOPOLOGY_MAX_ENUM;
 }
 
+#if UE_BUILD_DEBUG || UE_BUILD_DEVELOPMENT
+extern TAutoConsoleVariable<int32> CVarVulkanDebugBarrier;
+#endif
+
 namespace VulkanRHI
 {
 	static inline FString GetPipelineCacheFilename()
@@ -817,41 +821,46 @@ namespace VulkanRHI
 		return VK_IMAGE_LAYOUT_DEPTH_STENCIL_ATTACHMENT_OPTIMAL;
 	}
 
-	inline void InsertHeavyWeightBarrier(VkCommandBuffer CmdBuffer)
-	{
-		VkMemoryBarrier Barrier;
-		ZeroVulkanStruct(Barrier, VK_STRUCTURE_TYPE_MEMORY_BARRIER);
-		Barrier.srcAccessMask = VK_ACCESS_INDIRECT_COMMAND_READ_BIT |
-			VK_ACCESS_INDEX_READ_BIT |
-			VK_ACCESS_VERTEX_ATTRIBUTE_READ_BIT |
-			VK_ACCESS_UNIFORM_READ_BIT |
-			VK_ACCESS_INPUT_ATTACHMENT_READ_BIT |
-			VK_ACCESS_SHADER_READ_BIT |
-			VK_ACCESS_SHADER_WRITE_BIT |
-			VK_ACCESS_COLOR_ATTACHMENT_READ_BIT |
-			VK_ACCESS_COLOR_ATTACHMENT_WRITE_BIT |
-			VK_ACCESS_DEPTH_STENCIL_ATTACHMENT_READ_BIT |
-			VK_ACCESS_DEPTH_STENCIL_ATTACHMENT_WRITE_BIT |
-			VK_ACCESS_TRANSFER_READ_BIT |
-			VK_ACCESS_TRANSFER_WRITE_BIT |
-			VK_ACCESS_HOST_READ_BIT |
-			VK_ACCESS_HOST_WRITE_BIT;
-		Barrier.dstAccessMask = VK_ACCESS_INDIRECT_COMMAND_READ_BIT |
-			VK_ACCESS_INDEX_READ_BIT |
-			VK_ACCESS_VERTEX_ATTRIBUTE_READ_BIT |
-			VK_ACCESS_UNIFORM_READ_BIT |
-			VK_ACCESS_INPUT_ATTACHMENT_READ_BIT |
-			VK_ACCESS_SHADER_READ_BIT |
-			VK_ACCESS_SHADER_WRITE_BIT |
-			VK_ACCESS_COLOR_ATTACHMENT_READ_BIT |
-			VK_ACCESS_COLOR_ATTACHMENT_WRITE_BIT |
-			VK_ACCESS_DEPTH_STENCIL_ATTACHMENT_READ_BIT |
-			VK_ACCESS_DEPTH_STENCIL_ATTACHMENT_WRITE_BIT |
-			VK_ACCESS_TRANSFER_READ_BIT |
-			VK_ACCESS_TRANSFER_WRITE_BIT |
-			VK_ACCESS_HOST_READ_BIT |
-			VK_ACCESS_HOST_WRITE_BIT;
-		VulkanRHI::vkCmdPipelineBarrier(CmdBuffer, VK_PIPELINE_STAGE_ALL_COMMANDS_BIT, VK_PIPELINE_STAGE_ALL_COMMANDS_BIT, 0, 1, &Barrier, 0, nullptr, 0, nullptr);
+	inline void DebugHeavyWeightBarrier(VkCommandBuffer CmdBuffer, int32 CVarConditionMask)
+	{
+#if UE_BUILD_DEBUG || UE_BUILD_DEVELOPMENT
+		if (CVarVulkanDebugBarrier.GetValueOnAnyThread() & CVarConditionMask)
+		{
+			VkMemoryBarrier Barrier;
+			ZeroVulkanStruct(Barrier, VK_STRUCTURE_TYPE_MEMORY_BARRIER);
+			Barrier.srcAccessMask = VK_ACCESS_INDIRECT_COMMAND_READ_BIT |
+				VK_ACCESS_INDEX_READ_BIT |
+				VK_ACCESS_VERTEX_ATTRIBUTE_READ_BIT |
+				VK_ACCESS_UNIFORM_READ_BIT |
+				VK_ACCESS_INPUT_ATTACHMENT_READ_BIT |
+				VK_ACCESS_SHADER_READ_BIT |
+				VK_ACCESS_SHADER_WRITE_BIT |
+				VK_ACCESS_COLOR_ATTACHMENT_READ_BIT |
+				VK_ACCESS_COLOR_ATTACHMENT_WRITE_BIT |
+				VK_ACCESS_DEPTH_STENCIL_ATTACHMENT_READ_BIT |
+				VK_ACCESS_DEPTH_STENCIL_ATTACHMENT_WRITE_BIT |
+				VK_ACCESS_TRANSFER_READ_BIT |
+				VK_ACCESS_TRANSFER_WRITE_BIT |
+				VK_ACCESS_HOST_READ_BIT |
+				VK_ACCESS_HOST_WRITE_BIT;
+			Barrier.dstAccessMask = VK_ACCESS_INDIRECT_COMMAND_READ_BIT |
+				VK_ACCESS_INDEX_READ_BIT |
+				VK_ACCESS_VERTEX_ATTRIBUTE_READ_BIT |
+				VK_ACCESS_UNIFORM_READ_BIT |
+				VK_ACCESS_INPUT_ATTACHMENT_READ_BIT |
+				VK_ACCESS_SHADER_READ_BIT |
+				VK_ACCESS_SHADER_WRITE_BIT |
+				VK_ACCESS_COLOR_ATTACHMENT_READ_BIT |
+				VK_ACCESS_COLOR_ATTACHMENT_WRITE_BIT |
+				VK_ACCESS_DEPTH_STENCIL_ATTACHMENT_READ_BIT |
+				VK_ACCESS_DEPTH_STENCIL_ATTACHMENT_WRITE_BIT |
+				VK_ACCESS_TRANSFER_READ_BIT |
+				VK_ACCESS_TRANSFER_WRITE_BIT |
+				VK_ACCESS_HOST_READ_BIT |
+				VK_ACCESS_HOST_WRITE_BIT;
+			VulkanRHI::vkCmdPipelineBarrier(CmdBuffer, VK_PIPELINE_STAGE_ALL_COMMANDS_BIT, VK_PIPELINE_STAGE_ALL_COMMANDS_BIT, 0, 1, &Barrier, 0, nullptr, 0, nullptr);
+		}
+#endif
 	}
 }
 
diff --git a/Engine/Source/Runtime/VulkanRHI/Private/VulkanRenderTarget.cpp b/Engine/Source/Runtime/VulkanRHI/Private/VulkanRenderTarget.cpp
--- a/Engine/Source/Runtime/VulkanRHI/Private/VulkanRenderTarget.cpp
+++ b/Engine/Source/Runtime/VulkanRHI/Private/VulkanRenderTarget.cpp
@@ -44,7 +44,8 @@ TAutoConsoleVariable<int32> CVarVulkanDebugBarrier(
 	TEXT(" 2: Enable heavy barriers after every dispatch\n")
 	TEXT(" 4: Enable heavy barriers after upload cmd buffers\n")
 	TEXT(" 8: Enable heavy barriers after active cmd buffers\n")
-	/*TEXT(" 8: Enable heavy barrier after swapchain\n")*/,
+	TEXT(" 16: Enable heavy buffer barrier after uploads\n")
+	TEXT(" 32: Enable heavy buffer barrier between acquiring back buffer and blitting into swapchain\n"),
 	ECVF_Default
 );
 #endif
@@ -222,12 +223,7 @@ void FTransitionAndLayoutManager::EndEmulatedRenderPass(FVulkanCmdBuffer* CmdBuf
 	CmdBuffer->EndRenderPass();
 	CurrentRenderPass = nullptr;
 
-#if UE_BUILD_DEBUG || UE_BUILD_DEVELOPMENT
-	if (CVarVulkanDebugBarrier.GetValueOnRenderThread() & 1)
-	{
-		VulkanRHI::InsertHeavyWeightBarrier(CmdBuffer->GetHandle());
-	}
-#endif
+	VulkanRHI::DebugHeavyWeightBarrier(CmdBuffer->GetHandle(), 1);
 }
 
 void FTransitionAndLayoutManager::BeginRealRenderPass(FVulkanCommandListContext& Context, FVulkanDevice& InDevice, FVulkanCmdBuffer* CmdBuffer, const FRHIRenderPassInfo& RPInfo, const FVulkanRenderTargetLayout& RTLayout, FVulkanRenderPass* RenderPass, FVulkanFramebuffer* Framebuffer)
@@ -431,12 +427,7 @@ void FTransitionAndLayoutManager::EndRealRenderPass(FVulkanCmdBuffer* CmdBuffer)
 	CurrentRenderPass = nullptr;
 	bInsideRealRenderPass = false;
 
-#if UE_BUILD_DEBUG || UE_BUILD_DEVELOPMENT
-	if (CVarVulkanDebugBarrier.GetValueOnRenderThread() & 1)
-	{
-		VulkanRHI::InsertHeavyWeightBarrier(CmdBuffer->GetHandle());
-	}
-#endif
+	VulkanRHI::DebugHeavyWeightBarrier(CmdBuffer->GetHandle(), 1);
 }
 
 void FTransitionAndLayoutManager::NotifyDeletedRenderTarget(FVulkanDevice& InDevice, VkImage Image)
@@ -1372,6 +1363,8 @@ void FVulkanCommandListContext::TransitionResources(const FPendingTransition& Pe
 	}
 	else
 	{
+		SCOPED_RHI_CONDITIONAL_DRAW_EVENTF(*this, RHITransitionResources, bShowTransitionEvents, TEXT("TransitionTo: %s: %i UAVs"), *FResourceTransitionUtility::ResourceTransitionAccessStrings[(int32)PendingTransition.TransitionType], PendingTransition.UAVs.Num());
+
 		const bool bIsRealAsyncComputeContext = Device->IsRealAsyncComputeContext(this);
 		ensure(IsImmediate() || bIsRealAsyncComputeContext);
 		check(PendingTransition.UAVs.Num() > 0);
diff --git a/Engine/Source/Runtime/VulkanRHI/Private/VulkanSwapChain.cpp b/Engine/Source/Runtime/VulkanRHI/Private/VulkanSwapChain.cpp
--- a/Engine/Source/Runtime/VulkanRHI/Private/VulkanSwapChain.cpp
+++ b/Engine/Source/Runtime/VulkanRHI/Private/VulkanSwapChain.cpp
@@ -253,6 +253,8 @@ FVulkanSwapChain::FVulkanSwapChain(VkInstance InInstance, FVulkanDevice& InDevic
 	VkPresentModeKHR PresentMode = VK_PRESENT_MODE_FIFO_KHR;
 	if (FVulkanPlatform::SupportsQuerySurfaceProperties())
 	{
+		static bool bFirstTimeLog = true;
+
 		uint32 NumFoundPresentModes = 0;
 		VERIFYVULKANRESULT(VulkanRHI::vkGetPhysicalDeviceSurfacePresentModesKHR(Device.GetPhysicalHandle(), Surface, &NumFoundPresentModes, nullptr));
 		check(NumFoundPresentModes > 0);
@@ -261,6 +263,8 @@ FVulkanSwapChain::FVulkanSwapChain(VkInstance InInstance, FVulkanDevice& InDevic
 		FoundPresentModes.AddZeroed(NumFoundPresentModes);
 		VERIFYVULKANRESULT(VulkanRHI::vkGetPhysicalDeviceSurfacePresentModesKHR(Device.GetPhysicalHandle(), Surface, &NumFoundPresentModes, FoundPresentModes.GetData()));
 
+		UE_CLOG(bFirstTimeLog, LogVulkanRHI, Display, TEXT("Found %d Surface present modes:"), NumFoundPresentModes);
+
 		bool bFoundPresentModeMailbox = false;
 		bool bFoundPresentModeImmediate = false;
 		bool bFoundPresentModeFIFO = false;
@@ -271,36 +275,86 @@ FVulkanSwapChain::FVulkanSwapChain(VkInstance InInstance, FVulkanDevice& InDevic
 			{
 			case VK_PRESENT_MODE_MAILBOX_KHR:
 				bFoundPresentModeMailbox = true;
+				UE_CLOG(bFirstTimeLog, LogVulkanRHI, Display, TEXT("- VK_PRESENT_MODE_MAILBOX_KHR (%d)"), (int32)VK_PRESENT_MODE_MAILBOX_KHR);
 				break;
 			case VK_PRESENT_MODE_IMMEDIATE_KHR:
 				bFoundPresentModeImmediate = true;
+				UE_CLOG(bFirstTimeLog, LogVulkanRHI, Display, TEXT("- VK_PRESENT_MODE_IMMEDIATE_KHR (%d)"), (int32)VK_PRESENT_MODE_IMMEDIATE_KHR);
 				break;
 			case VK_PRESENT_MODE_FIFO_KHR:
 				bFoundPresentModeFIFO = true;
+				UE_CLOG(bFirstTimeLog, LogVulkanRHI, Display, TEXT("- VK_PRESENT_MODE_FIFO_KHR (%d)"), (int32)VK_PRESENT_MODE_FIFO_KHR);
+				break;
+			default:
+				UE_CLOG(bFirstTimeLog, LogVulkanRHI, Display, TEXT("- VkPresentModeKHR %d"), (int32)FoundPresentModes[i]);
 				break;
 			}
 		}
 
-		// Until FVulkanViewport::Present honors SyncInterval, we need to disable vsync for the spectator window if using an HMD.
-		const bool bDisableVsyncForHMD = (FVulkanDynamicRHI::HMDVulkanExtensions.IsValid()) ? FVulkanDynamicRHI::HMDVulkanExtensions->ShouldDisableVulkanVSync() : false;
-
-		if (bFoundPresentModeMailbox && LockToVsync)
-		{
-			PresentMode = VK_PRESENT_MODE_MAILBOX_KHR;
-		}
-		else if(bFoundPresentModeImmediate && (bDisableVsyncForHMD || !LockToVsync))
-		{
-			PresentMode = VK_PRESENT_MODE_IMMEDIATE_KHR;
-		}
-		else if (bFoundPresentModeFIFO)
+		int32 RequestedPresentMode = -1;
+		if (FParse::Value(FCommandLine::Get(), TEXT("vulkanpresentmode="), RequestedPresentMode))
 		{
-			PresentMode = VK_PRESENT_MODE_FIFO_KHR;
+			bool bRequestSuccessful = false;
+			switch (RequestedPresentMode)
+			{
+			case VK_PRESENT_MODE_MAILBOX_KHR:
+				if (bFoundPresentModeMailbox)
+				{
+					PresentMode = VK_PRESENT_MODE_MAILBOX_KHR;
+					bRequestSuccessful = true;
+				}
+				break;
+			case VK_PRESENT_MODE_IMMEDIATE_KHR:
+				if (bFoundPresentModeImmediate)
+				{
+					PresentMode = VK_PRESENT_MODE_IMMEDIATE_KHR;
+					bRequestSuccessful = true;
+				}
+				break;
+			case VK_PRESENT_MODE_FIFO_KHR:
+				if (bFoundPresentModeFIFO)
+				{
+					PresentMode = VK_PRESENT_MODE_FIFO_KHR;
+					bRequestSuccessful = true;
+				}
+				break;
+			default:
+				break;
+			}
+
+			if (!bRequestSuccessful)
+			{
+				UE_CLOG(bFirstTimeLog, LogVulkanRHI, Warning, TEXT("Requested PresentMode (%d) is not handled or available, ignoring..."), RequestedPresentMode);
+				RequestedPresentMode = -1;
+			}
 		}
-		else
+
+		if (RequestedPresentMode == -1)
 		{
-			UE_LOG(LogVulkanRHI, Warning, TEXT("Couldn't find desired PresentMode! Using %d"), static_cast<int32>(FoundPresentModes[0]));
-			PresentMode = FoundPresentModes[0];
+			// Until FVulkanViewport::Present honors SyncInterval, we need to disable vsync for the spectator window if using an HMD.
+			const bool bDisableVsyncForHMD = (FVulkanDynamicRHI::HMDVulkanExtensions.IsValid()) ? FVulkanDynamicRHI::HMDVulkanExtensions->ShouldDisableVulkanVSync() : false;
+
+			if (bFoundPresentModeMailbox && LockToVsync)
+			{
+				PresentMode = VK_PRESENT_MODE_MAILBOX_KHR;
+			}
+			else if (bFoundPresentModeImmediate && (bDisableVsyncForHMD || !LockToVsync))
+			{
+				PresentMode = VK_PRESENT_MODE_IMMEDIATE_KHR;
+			}
+			else if (bFoundPresentModeFIFO)
+			{
+				PresentMode = VK_PRESENT_MODE_FIFO_KHR;
+			}
+			else
+			{
+				UE_LOG(LogVulkanRHI, Warning, TEXT("Couldn't find desired PresentMode! Using %d"), static_cast<int32>(FoundPresentModes[0]));
+				PresentMode = FoundPresentModes[0];
+			}
 		}
+
+		UE_CLOG(bFirstTimeLog, LogVulkanRHI, Display, TEXT("Selected VkPresentModeKHR mode %d"), PresentMode);
+		bFirstTimeLog = false;
 	}
 
 	// Check the surface properties and formats
diff --git a/Engine/Source/Runtime/VulkanRHI/Private/VulkanViewport.cpp b/Engine/Source/Runtime/VulkanRHI/Private/VulkanViewport.cpp
--- a/Engine/Source/Runtime/VulkanRHI/Private/VulkanViewport.cpp
+++ b/Engine/Source/Runtime/VulkanRHI/Private/VulkanViewport.cpp
@@ -19,6 +19,7 @@ FAutoConsoleVariable GCVarDelayAcquireBackBuffer(
 	ECVF_ReadOnly
 );
 
+
 struct FRHICommandAcquireBackBuffer final : public FRHICommand<FRHICommandAcquireBackBuffer>
 {
 	FVulkanViewport* Viewport;
@@ -244,9 +245,38 @@ void FVulkanViewport::AdvanceBackBufferFrame()
 	}
 }
 
-//void FVulkanViewport::WaitForFrameEventCompletion()
-//{
-//}
+void FVulkanViewport::WaitForFrameEventCompletion()
+{
+	if (FVulkanPlatform::RequiresWaitingForFrameCompletionEvent())
+	{
+		static FCriticalSection CS;
+		FScopeLock ScopeLock(&CS);
+		if (LastFrameCommandBuffer && LastFrameCommandBuffer->IsSubmitted())
+		{
+			// If last frame's fence hasn't been signaled already, wait for it here
+			if (LastFrameFenceCounter == LastFrameCommandBuffer->GetFenceSignaledCounter())
+			{
+				if (!GWaitForIdleOnSubmit)
+				{
+					// The wait has already happened if GWaitForIdleOnSubmit is set
+					LastFrameCommandBuffer->GetOwner()->GetMgr().WaitForCmdBuffer(LastFrameCommandBuffer);
+				}
+			}
+		}
+	}
+}
+
+void FVulkanViewport::IssueFrameEvent()
+{
+	if (FVulkanPlatform::RequiresWaitingForFrameCompletionEvent())
+	{
+		// The fence we need to wait on next frame is already there in the command buffer
+		// that was just submitted in this frame's Present. Just grab that command buffer's
+		// info to use next frame in WaitForFrameEventCompletion.
+		FVulkanQueue* Queue = Device->GetGraphicsQueue();
+		Queue->GetLastSubmittedInfo(LastFrameCommandBuffer, LastFrameFenceCounter);
+	}
+}
 
 
 FVulkanFramebuffer::FVulkanFramebuffer(FVulkanDevice& Device, const FRHISetRenderTargetsInfo& InRTInfo, const FVulkanRenderTargetLayout& RTLayout, const FVulkanRenderPass& RenderPass)
@@ -577,6 +607,8 @@ inline static void CopyImageToBackBuffer(FVulkanCmdBuffer* CmdBuffer, bool bSour
 	Barriers.SetTransition(DestIndex, EImageLayoutBarrier::Undefined, EImageLayoutBarrier::TransferDest);
 	Barriers.Execute(CmdBuffer);
 
+	VulkanRHI::DebugHeavyWeightBarrier(CmdBuffer->GetHandle(), 32);
+
 	VkImageCopy Region;
 	FMemory::Memzero(Region);
 	Region.extent.width = SizeX;
@@ -693,19 +725,12 @@ bool FVulkanViewport::Present(FVulkanCommandListContext* Context, FVulkanCmdBuff
 		RHIBackBuffer = nullptr;
 	}
 
-	//static const TConsoleVariableData<int32>* CFinishFrameVar = IConsoleManager::Get().FindTConsoleVariableDataInt(TEXT("r.FinishCurrentFrame"));
-	//if (!CFinishFrameVar->GetValueOnRenderThread())
-	//{
-	//	// Wait for the GPU to finish rendering the previous frame before finishing this frame.
-	//	WaitForFrameEventCompletion();
-	//	IssueFrameEvent();
-	//}
-	//else
-	//{
-	//	// Finish current frame immediately to reduce latency
-	//	IssueFrameEvent();
-	//	WaitForFrameEventCompletion();
-	//}
+	if (FVulkanPlatform::RequiresWaitingForFrameCompletionEvent() && !bHasCustomPresent)
+	{
+		// Wait for the GPU to finish rendering the previous frame before finishing this frame.
+		WaitForFrameEventCompletion();
+		IssueFrameEvent();
+	}
 
 	// If the input latency timer has been triggered, block until the GPU is completely
 	// finished displaying this frame and calculate the delta time.
diff --git a/Engine/Source/Runtime/VulkanRHI/Public/VulkanViewport.h b/Engine/Source/Runtime/VulkanRHI/Public/VulkanViewport.h
--- a/Engine/Source/Runtime/VulkanRHI/Public/VulkanViewport.h
+++ b/Engine/Source/Runtime/VulkanRHI/Public/VulkanViewport.h
@@ -27,10 +27,9 @@ public:
 
 	FVulkanTexture2D* GetBackBuffer(FRHICommandList& RHICmdList);
 
-	//void WaitForFrameEventCompletion();
+	void WaitForFrameEventCompletion();
 
-	//#todo-rco
-	//void IssueFrameEvent() {}
+	void IssueFrameEvent();
 
 	inline FIntPoint GetSizeXY() const
 	{
@@ -84,6 +83,9 @@ protected:
 
 	FCustomPresentRHIRef CustomPresent;
 
+	FVulkanCmdBuffer* LastFrameCommandBuffer = nullptr;
+	uint64 LastFrameFenceCounter = 0;
+
 	void CreateSwapchain();
 	void AcquireBackBuffer(FRHICommandListBase& CmdList, FVulkanBackBuffer* NewBackBuffer);
 
