UE4.21 - Vulkan - Add missing viewport resize function
- Separate num buffers per usage
#rnx

diff --git a/Engine/Source/Runtime/VulkanRHI/Private/VulkanIndexBuffer.cpp b/Engine/Source/Runtime/VulkanRHI/Private/VulkanIndexBuffer.cpp
--- a/Engine/Source/Runtime/VulkanRHI/Private/VulkanIndexBuffer.cpp
+++ b/Engine/Source/Runtime/VulkanRHI/Private/VulkanIndexBuffer.cpp
@@ -102,7 +102,7 @@ FVulkanResourceMultiBuffer::FVulkanResourceMultiBuffer(FVulkanDevice* InDevice,
 				BufferMemFlags |= (VK_MEMORY_PROPERTY_HOST_VISIBLE_BIT | VK_MEMORY_PROPERTY_HOST_COHERENT_BIT);
 			}
 
-			NumBuffers = bDynamic ? NUM_RENDER_BUFFERS : 1;
+			NumBuffers = bDynamic ? NUM_BUFFERS : 1;
 			check(NumBuffers <= ARRAY_COUNT(Buffers));
 
 			for (uint32 Index = 0; Index < NumBuffers; ++Index)
diff --git a/Engine/Source/Runtime/VulkanRHI/Private/VulkanMemory.cpp b/Engine/Source/Runtime/VulkanRHI/Private/VulkanMemory.cpp
--- a/Engine/Source/Runtime/VulkanRHI/Private/VulkanMemory.cpp
+++ b/Engine/Source/Runtime/VulkanRHI/Private/VulkanMemory.cpp
@@ -1936,7 +1936,7 @@ namespace VulkanRHI
 		: FDeviceChild(InDevice)
 		, BufferIndex(0)
 	{
-		for (int32 Index = 0; Index < NUM_RENDER_BUFFERS; ++Index)
+		for (int32 Index = 0; Index < NUM_BUFFERS; ++Index)
 		{
 			Entries[Index].InitBuffer(Device, ALLOCATION_SIZE);
 		}
@@ -1962,7 +1962,7 @@ namespace VulkanRHI
 
 	void FTempFrameAllocationBuffer::Destroy()
 	{
-		for (int32 Index = 0; Index < NUM_RENDER_BUFFERS; ++Index)
+		for (int32 Index = 0; Index < NUM_BUFFERS; ++Index)
 		{
 			Entries[Index].BufferSuballocation = nullptr;
 		}
@@ -2006,7 +2006,7 @@ namespace VulkanRHI
 	void FTempFrameAllocationBuffer::Reset()
 	{
 		FScopeLock ScopeLock(&CS);
-		BufferIndex = (BufferIndex + 1) % NUM_RENDER_BUFFERS;
+		BufferIndex = (BufferIndex + 1) % NUM_BUFFERS;
 		Entries[BufferIndex].Reset();
 	}
 
diff --git a/Engine/Source/Runtime/VulkanRHI/Private/VulkanSwapChain.cpp b/Engine/Source/Runtime/VulkanRHI/Private/VulkanSwapChain.cpp
--- a/Engine/Source/Runtime/VulkanRHI/Private/VulkanSwapChain.cpp
+++ b/Engine/Source/Runtime/VulkanRHI/Private/VulkanSwapChain.cpp
@@ -168,7 +168,7 @@ FVulkanSwapChain::FVulkanSwapChain(VkInstance InInstance, FVulkanDevice& InDevic
 					if (!GPixelFormatNotSupportedWarning.Contains(InOutPixelFormat))
 					{
 						GPixelFormatNotSupportedWarning.Add(InOutPixelFormat);
-						UE_LOG(LogVulkanRHI, Warning, TEXT("Requested PixelFormat %d not supported by this swapchain! Falling back to supported swapchain formats..."), (uint32)InOutPixelFormat);
+						UE_LOG(LogVulkanRHI, Warning, TEXT("Requested PixelFormat %d not supported by this swapchain! Falling back to supported swapchain format..."), (uint32)InOutPixelFormat);
 					}
 					InOutPixelFormat = PF_Unknown;
 				}
diff --git a/Engine/Source/Runtime/VulkanRHI/Private/VulkanViewport.cpp b/Engine/Source/Runtime/VulkanRHI/Private/VulkanViewport.cpp
--- a/Engine/Source/Runtime/VulkanRHI/Private/VulkanViewport.cpp
+++ b/Engine/Source/Runtime/VulkanRHI/Private/VulkanViewport.cpp
@@ -10,6 +10,7 @@
 #include "VulkanContext.h"
 #include "GlobalShader.h"
 #include "HAL/PlatformAtomics.h"
+#include "Engine/RendererSettings.h"
 
 //#todo-Lumin: Until we have LuminEngine.ini
 FAutoConsoleVariable GCVarDelayAcquireBackBuffer(
@@ -476,22 +477,22 @@ void FVulkanViewport::Tick(float DeltaTime)
 		ENQUEUE_RENDER_COMMAND(UpdateVsync)(
 			[this](FRHICommandListImmediate& RHICmdList)
 		{
-			RecreateSwapchainFromRT();
+			RecreateSwapchainFromRT(PixelFormat);
 		});
 		FlushRenderingCommands();
 	}
 }
 
-void FVulkanViewport::Resize(uint32 InSizeX, uint32 InSizeY, bool bInIsFullscreen)
+void FVulkanViewport::Resize(uint32 InSizeX, uint32 InSizeY, bool bInIsFullscreen, EPixelFormat PreferredPixelFormat)
 {
 	SizeX = InSizeX;
 	SizeY = InSizeY;
 	bIsFullscreen = bInIsFullscreen;
 
-	RecreateSwapchainFromRT();
+	RecreateSwapchainFromRT(PreferredPixelFormat);
 }
 
-void FVulkanViewport::RecreateSwapchainFromRT()
+void FVulkanViewport::RecreateSwapchainFromRT(EPixelFormat PreferredPixelFormat)
 {
 	check(IsInRenderingThread());
 	
@@ -525,6 +526,7 @@ void FVulkanViewport::RecreateSwapchainFromRT()
 		Device->GetDeferredDeletionQueue().ReleaseResources(true);
 	}
 
+	PixelFormat = PreferredPixelFormat;
 	CreateSwapchain();
 }
 
@@ -761,9 +763,42 @@ bool FVulkanViewport::Present(FVulkanCommandListContext* Context, FVulkanCmdBuff
 FViewportRHIRef FVulkanDynamicRHI::RHICreateViewport(void* WindowHandle, uint32 SizeX, uint32 SizeY, bool bIsFullscreen, EPixelFormat PreferredPixelFormat)
 {
 	check( IsInGameThread() );
+
+	// Use a default pixel format if none was specified	
+	if (PreferredPixelFormat == PF_Unknown)
+	{
+		static const auto* CVarDefaultBackBufferPixelFormat = IConsoleManager::Get().FindTConsoleVariableDataInt(TEXT("r.DefaultBackBufferPixelFormat"));
+		PreferredPixelFormat = EDefaultBackBufferPixelFormat::Convert2PixelFormat(EDefaultBackBufferPixelFormat::FromInt(CVarDefaultBackBufferPixelFormat->GetValueOnAnyThread()));
+	}
+
 	return new FVulkanViewport(this, Device, WindowHandle, SizeX, SizeY, bIsFullscreen, PreferredPixelFormat);
 }
 
+void FVulkanDynamicRHI::RHIResizeViewport(FViewportRHIParamRef ViewportRHI, uint32 SizeX, uint32 SizeY, bool bIsFullscreen, EPixelFormat PreferredPixelFormat)
+{
+	check(IsInGameThread());
+	FVulkanViewport* Viewport = ResourceCast(ViewportRHI);
+
+	// Use a default pixel format if none was specified	
+	if (PreferredPixelFormat == PF_Unknown)
+	{
+		static const auto* CVarDefaultBackBufferPixelFormat = IConsoleManager::Get().FindTConsoleVariableDataInt(TEXT("r.DefaultBackBufferPixelFormat"));
+		PreferredPixelFormat = EDefaultBackBufferPixelFormat::Convert2PixelFormat(EDefaultBackBufferPixelFormat::FromInt(CVarDefaultBackBufferPixelFormat->GetValueOnAnyThread()));
+	}
+
+	if (Viewport->GetSizeXY() != FIntPoint(SizeX, SizeY))
+	{
+		FlushRenderingCommands();
+
+		ENQUEUE_RENDER_COMMAND(ResizeViewport)(
+			[Viewport, SizeX, SizeY, bIsFullscreen, PreferredPixelFormat](FRHICommandListImmediate& RHICmdList)
+			{
+				Viewport->Resize(SizeX, SizeY, bIsFullscreen, PreferredPixelFormat);
+			});
+		FlushRenderingCommands();
+	}
+}
+
 void FVulkanDynamicRHI::RHIResizeViewport(FViewportRHIParamRef ViewportRHI, uint32 SizeX, uint32 SizeY, bool bIsFullscreen)
 {
 	check(IsInGameThread());
@@ -776,7 +811,7 @@ void FVulkanDynamicRHI::RHIResizeViewport(FViewportRHIParamRef ViewportRHI, uint
 		ENQUEUE_RENDER_COMMAND(ResizeViewport)(
 			[Viewport, SizeX, SizeY, bIsFullscreen](FRHICommandListImmediate& RHICmdList)
 			{
-				Viewport->Resize(SizeX, SizeY, bIsFullscreen);
+				Viewport->Resize(SizeX, SizeY, bIsFullscreen, PF_Unknown);
 			});
 		FlushRenderingCommands();
 	}
diff --git a/Engine/Source/Runtime/VulkanRHI/Public/VulkanCommon.h b/Engine/Source/Runtime/VulkanRHI/Public/VulkanCommon.h
--- a/Engine/Source/Runtime/VulkanRHI/Public/VulkanCommon.h
+++ b/Engine/Source/Runtime/VulkanRHI/Public/VulkanCommon.h
@@ -137,9 +137,3 @@ namespace EVulkanBindingType
 		return 0;
 	}
 }
-
-/** How many back buffers to cycle through */
-enum
-{
-	NUM_RENDER_BUFFERS = 3,
-};
diff --git a/Engine/Source/Runtime/VulkanRHI/Public/VulkanDynamicRHI.h b/Engine/Source/Runtime/VulkanRHI/Public/VulkanDynamicRHI.h
--- a/Engine/Source/Runtime/VulkanRHI/Public/VulkanDynamicRHI.h
+++ b/Engine/Source/Runtime/VulkanRHI/Public/VulkanDynamicRHI.h
@@ -135,6 +135,7 @@ public:
 	virtual uint32 RHIGetGPUFrameCycles() final override;
 	virtual FViewportRHIRef RHICreateViewport(void* WindowHandle, uint32 SizeX, uint32 SizeY, bool bIsFullscreen, EPixelFormat PreferredPixelFormat) final override;
 	virtual void RHIResizeViewport(FViewportRHIParamRef Viewport, uint32 SizeX, uint32 SizeY, bool bIsFullscreen) final override;
+	virtual void RHIResizeViewport(FViewportRHIParamRef Viewport, uint32 SizeX, uint32 SizeY, bool bIsFullscreen, EPixelFormat PreferredPixelFormat) final override;
 	virtual void RHITick(float DeltaTime) final override;
 	virtual void RHISetStreamOutTargets(uint32 NumTargets, const FVertexBufferRHIParamRef* VertexBuffers, const uint32* Offsets) final override;
 	virtual void RHIBlockUntilGPUIdle() final override;
diff --git a/Engine/Source/Runtime/VulkanRHI/Public/VulkanMemory.h b/Engine/Source/Runtime/VulkanRHI/Public/VulkanMemory.h
--- a/Engine/Source/Runtime/VulkanRHI/Public/VulkanMemory.h
+++ b/Engine/Source/Runtime/VulkanRHI/Public/VulkanMemory.h
@@ -1250,6 +1250,11 @@ namespace VulkanRHI
 	protected:
 		uint32 BufferIndex;
 
+		enum
+		{
+			NUM_BUFFERS = 3,
+		};
+
 		struct FFrameEntry
 		{
 			TRefCountPtr<FBufferSuballocation> BufferSuballocation;
@@ -1263,7 +1268,7 @@ namespace VulkanRHI
 			void Reset();
 			bool TryAlloc(uint32 InSize, uint32 InAlignment, FTempAllocInfo& OutInfo);
 		};
-		FFrameEntry Entries[NUM_RENDER_BUFFERS];
+		FFrameEntry Entries[NUM_BUFFERS];
 		FCriticalSection CS;
 
 		friend class FVulkanCommandListContext;
diff --git a/Engine/Source/Runtime/VulkanRHI/Public/VulkanResources.h b/Engine/Source/Runtime/VulkanRHI/Public/VulkanResources.h
--- a/Engine/Source/Runtime/VulkanRHI/Public/VulkanResources.h
+++ b/Engine/Source/Runtime/VulkanRHI/Public/VulkanResources.h
@@ -1152,7 +1152,13 @@ protected:
 	VkBufferUsageFlags BufferUsageFlags;
 	uint32 NumBuffers;
 	uint32 DynamicBufferIndex;
-	TRefCountPtr<VulkanRHI::FBufferSuballocation> Buffers[NUM_RENDER_BUFFERS];
+
+	enum
+	{
+		NUM_BUFFERS = 3,
+	};
+
+	TRefCountPtr<VulkanRHI::FBufferSuballocation> Buffers[NUM_BUFFERS];
 	struct
 	{
 		VulkanRHI::FBufferSuballocation* SubAlloc = nullptr;
diff --git a/Engine/Source/Runtime/VulkanRHI/Public/VulkanViewport.h b/Engine/Source/Runtime/VulkanRHI/Public/VulkanViewport.h
--- a/Engine/Source/Runtime/VulkanRHI/Public/VulkanViewport.h
+++ b/Engine/Source/Runtime/VulkanRHI/Public/VulkanViewport.h
@@ -90,8 +90,8 @@ protected:
 	void AcquireBackBuffer(FRHICommandListBase& CmdList, FVulkanBackBuffer* NewBackBuffer);
 
 	void RecreateSwapchain(void* NewNativeWindow, bool bForce = false);
-	void RecreateSwapchainFromRT();
-	void Resize(uint32 InSizeX, uint32 InSizeY, bool bIsFullscreen);
+	void RecreateSwapchainFromRT(EPixelFormat PreferredPixelFormat);
+	void Resize(uint32 InSizeX, uint32 InSizeY, bool bIsFullscreen, EPixelFormat PreferredPixelFormat);
 
 	static int32 DoAcquireImageIndex(FVulkanViewport* Viewport);
 	bool DoCheckedSwapChainJob(TFunction<int32(FVulkanViewport*)> SwapChainJob);
