UE4.21 - Android compile fix
#rnx

diff --git a/Engine/Source/Runtime/VulkanRHI/Private/VulkanPipeline.h b/Engine/Source/Runtime/VulkanRHI/Private/VulkanPipeline.h
--- a/Engine/Source/Runtime/VulkanRHI/Private/VulkanPipeline.h
+++ b/Engine/Source/Runtime/VulkanRHI/Private/VulkanPipeline.h
@@ -685,13 +685,17 @@ public:
 			FGfxPipelineEntry* GfxEntry;
 
 		public:
+			enum
+			{
+				NUM_FRAMES = 3,
+			};
 			FEvictedVkPipeline(FVulkanDevice* InDevice, FVulkanGfxPipeline* InPipeline, FGfxPipelineEntry* InGfxEntry)
 				: Device(InDevice), GfxPipeline(InPipeline), GfxEntry(InGfxEntry)
 			{
 				check(InGfxEntry);
 				check(InGfxEntry->bConsumed);
 
-				GfxPipeline->DeleteVkPipeline(InPipeline->RecentFrame + NUM_RENDER_BUFFERS < GFrameNumberRenderThread);
+				GfxPipeline->DeleteVkPipeline(InPipeline->RecentFrame + NUM_FRAMES < GFrameNumberRenderThread);
 			}
 
 			void RestoreVkPipeline();
