Linux: Use the size of our In Rect which is the size of our buffer rather then the size of the Texture which is not always the same size as the buffer

diff --git a/Engine/Source/Runtime/VulkanRHI/Private/VulkanRenderTarget.cpp b/Engine/Source/Runtime/VulkanRHI/Private/VulkanRenderTarget.cpp
--- a/Engine/Source/Runtime/VulkanRHI/Private/VulkanRenderTarget.cpp
+++ b/Engine/Source/Runtime/VulkanRHI/Private/VulkanRenderTarget.cpp
@@ -692,7 +692,7 @@ void FVulkanDynamicRHI::RHIReadSurfaceData(FTextureRHIParamRef TextureRHI, FIntR
 	FRHITexture2D* TextureRHI2D = TextureRHI->GetTexture2D();
 	check(TextureRHI2D);
 	FVulkanTexture2D* Texture2D = (FVulkanTexture2D*)TextureRHI2D;
-	uint32 NumPixels = TextureRHI2D->GetSizeX() * TextureRHI2D->GetSizeY();
+	uint32 NumPixels = (Rect.Max.X - Rect.Min.X) * (Rect.Max.Y - Rect.Min.Y);
 
 	if (GIgnoreCPUReads == 2)
 	{
