UE4.21 - General fixes from Vulkan
#rnx

diff --git a/Engine/Source/Runtime/VulkanRHI/Private/VulkanCommandBuffer.cpp b/Engine/Source/Runtime/VulkanRHI/Private/VulkanCommandBuffer.cpp
--- a/Engine/Source/Runtime/VulkanRHI/Private/VulkanCommandBuffer.cpp
+++ b/Engine/Source/Runtime/VulkanRHI/Private/VulkanCommandBuffer.cpp
@@ -479,30 +479,45 @@ FVulkanCmdBuffer* FVulkanCommandBufferManager::GetUploadCmdBuffer()
 struct FRHICommandFreeUnusedCmdBuffers final : public FRHICommand<FRHICommandFreeUnusedCmdBuffers>
 {
 	FVulkanCommandBufferPool* Pool;
+	FVulkanQueue* Queue;
 
-	FRHICommandFreeUnusedCmdBuffers(FVulkanCommandBufferPool* InPool)
+	FRHICommandFreeUnusedCmdBuffers(FVulkanCommandBufferPool* InPool, FVulkanQueue* InQueue)
 		: Pool(InPool)
+		, Queue(InQueue)
 	{
 	}
 
 	void Execute(FRHICommandListBase& CmdList)
 	{
-		Pool->FreeUnusedCmdBuffers();
+		Pool->FreeUnusedCmdBuffers(Queue);
 	}
 };
 #endif
 
 
-void FVulkanCommandBufferPool::FreeUnusedCmdBuffers()
+void FVulkanCommandBufferPool::FreeUnusedCmdBuffers(FVulkanQueue* InQueue)
 {
 #if VULKAN_DELETE_STALE_CMDBUFFERS
 	FScopeLock ScopeLock(&CS);
 	const double CurrentTime = FPlatformTime::Seconds();
+	
+	// In case Queue stores pointer to a cmdbuffer, do not delete it 
+	FVulkanCmdBuffer* LastSubmittedCmdBuffer = nullptr;
+	uint64 LastSubmittedFenceCounter = 0;
+	InQueue->GetLastSubmittedInfo(LastSubmittedCmdBuffer, LastSubmittedFenceCounter);
+
+	// Deferred deletion queue caches pointers to cmdbuffers
+	FDeferredDeletionQueue& DeferredDeletionQueue = Device->GetDeferredDeletionQueue();
+	
 	for (int32 Index = CmdBuffers.Num() - 1; Index >= 0; --Index)
 	{
 		FVulkanCmdBuffer* CmdBuffer = CmdBuffers[Index];
-		if (CmdBuffer->State == FVulkanCmdBuffer::EState::ReadyForBegin && CurrentTime - CmdBuffer->SubmittedTime > CMD_BUFFER_TIME_TO_WAIT_BEFORE_DELETING)
+		if (CmdBuffer != LastSubmittedCmdBuffer &&
+			CmdBuffer->State == FVulkanCmdBuffer::EState::ReadyForBegin && 
+			(CurrentTime - CmdBuffer->SubmittedTime) > CMD_BUFFER_TIME_TO_WAIT_BEFORE_DELETING)
 		{
+			DeferredDeletionQueue.OnCmdBufferDeleted(CmdBuffer);
+			
 			delete CmdBuffer;
 			CmdBuffers.RemoveAtSwap(Index, 1, false);
 		}
@@ -516,12 +531,12 @@ void FVulkanCommandBufferManager::FreeUnusedCmdBuffers()
 	FRHICommandList& RHICmdList = FRHICommandListExecutor::GetImmediateCommandList();
 	if (!IsInRenderingThread() || (RHICmdList.Bypass() || !IsRunningRHIInSeparateThread()))
 	{
-		Pool.FreeUnusedCmdBuffers();
+		Pool.FreeUnusedCmdBuffers(Queue);
 	}
 	else
 	{
 		check(IsInRenderingThread());
-		new (RHICmdList.AllocCommand<FRHICommandFreeUnusedCmdBuffers>()) FRHICommandFreeUnusedCmdBuffers(&Pool);
+		new (RHICmdList.AllocCommand<FRHICommandFreeUnusedCmdBuffers>()) FRHICommandFreeUnusedCmdBuffers(&Pool, Queue);
 	}
 #endif
 }
diff --git a/Engine/Source/Runtime/VulkanRHI/Private/VulkanCommandBuffer.h b/Engine/Source/Runtime/VulkanRHI/Private/VulkanCommandBuffer.h
--- a/Engine/Source/Runtime/VulkanRHI/Private/VulkanCommandBuffer.h
+++ b/Engine/Source/Runtime/VulkanRHI/Private/VulkanCommandBuffer.h
@@ -230,7 +230,7 @@ public:
 		return &CS;
 	}
 
-	void FreeUnusedCmdBuffers();
+	void FreeUnusedCmdBuffers(FVulkanQueue* Queue);
 
 	inline FVulkanCommandBufferManager& GetMgr()
 	{
diff --git a/Engine/Source/Runtime/VulkanRHI/Private/VulkanCommandWrappers.h b/Engine/Source/Runtime/VulkanRHI/Private/VulkanCommandWrappers.h
--- a/Engine/Source/Runtime/VulkanRHI/Private/VulkanCommandWrappers.h
+++ b/Engine/Source/Runtime/VulkanRHI/Private/VulkanCommandWrappers.h
@@ -165,6 +165,10 @@ struct FWrapLayer
 #if VULKAN_USE_CREATE_WIN32_SURFACE
 	static void CreateWin32SurfaceKHR(VkResult Result, VkInstance instance, const VkWin32SurfaceCreateInfoKHR* pCreateInfo, VkSurfaceKHR* pSurface);
 #endif
+#if VULKAN_SUPPORTS_COLOR_CONVERSIONS
+	static void CreateSamplerYcbcrConversionKHR(VkResult Result, VkDevice device, const VkSamplerYcbcrConversionCreateInfo* pCreateInfo, VkSamplerYcbcrConversion* YcbcrConversion) VULKAN_LAYER_BODY
+	static void DestroySamplerYcbcrConversionKHR(VkResult Result, VkDevice device, VkSamplerYcbcrConversion ycbcrConversion) VULKAN_LAYER_BODY
+#endif
 };
 
 #undef VULKAN_LAYER_BODY
@@ -1319,6 +1323,23 @@ namespace VulkanRHI
 	}
 #endif
 
+#if VULKAN_SUPPORTS_COLOR_CONVERSIONS
+	static FORCEINLINE_DEBUGGABLE VkResult vkCreateSamplerYcbcrConversionKHR(VkDevice Device, const VkSamplerYcbcrConversionCreateInfo* CreateInfo, const VkAllocationCallbacks* Allocator, VkSamplerYcbcrConversion* YcbcrConversion)
+	{
+		FWrapLayer::CreateSamplerYcbcrConversionKHR(VK_RESULT_MAX_ENUM, Device, CreateInfo, YcbcrConversion);
+		VkResult Result = VULKANAPINAMESPACE::vkCreateSamplerYcbcrConversionKHR(Device, CreateInfo, Allocator, YcbcrConversion);
+		FWrapLayer::CreateSamplerYcbcrConversionKHR(Result, Device, CreateInfo, YcbcrConversion);
+		return Result;
+	}
+
+	static FORCEINLINE_DEBUGGABLE void vkDestroySamplerYcbcrConversionKHR(VkDevice Device, VkSamplerYcbcrConversion YcbcrConversion, const VkAllocationCallbacks* Allocator)
+	{
+		FWrapLayer::DestroySamplerYcbcrConversionKHR(VK_RESULT_MAX_ENUM, Device, YcbcrConversion);
+		VULKANAPINAMESPACE::vkDestroySamplerYcbcrConversionKHR(Device, YcbcrConversion, Allocator);
+		FWrapLayer::DestroySamplerYcbcrConversionKHR(VK_SUCCESS, Device, YcbcrConversion);
+	}
+#endif
+
 #if VULKAN_ENABLE_IMAGE_TRACKING_LAYER
 	void BindDebugLabelName(VkImage Image, const TCHAR* Name);
 #endif
diff --git a/Engine/Source/Runtime/VulkanRHI/Private/VulkanCommands.cpp b/Engine/Source/Runtime/VulkanRHI/Private/VulkanCommands.cpp
--- a/Engine/Source/Runtime/VulkanRHI/Private/VulkanCommands.cpp
+++ b/Engine/Source/Runtime/VulkanRHI/Private/VulkanCommands.cpp
@@ -188,7 +188,7 @@ void FVulkanCommandListContext::RHISetUAVParameter(FComputeShaderRHIParamRef Com
 
 void FVulkanCommandListContext::RHISetShaderTexture(FVertexShaderRHIParamRef VertexShaderRHI, uint32 TextureIndex, FTextureRHIParamRef NewTextureRHI)
 {
-	check(PendingGfxState->GetCurrentShader(ShaderStage::Vertex) == ResourceCast(VertexShaderRHI));
+	check(PendingGfxState->GetCurrentShader(ShaderStage::Vertex)->GetShaderKey() == GetShaderKey(VertexShaderRHI));
 	FVulkanTextureBase* Texture = GetVulkanTextureFromRHITexture(NewTextureRHI);
 	VkImageLayout Layout = GetLayoutForDescriptor(Texture->Surface);
 	PendingGfxState->SetTextureForStage(ShaderStage::Vertex, TextureIndex, Texture, Layout);
@@ -198,7 +198,7 @@ void FVulkanCommandListContext::RHISetShaderTexture(FHullShaderRHIParamRef HullS
 {
 	ensureMsgf(0, TEXT("Tessellation not supported yet!"));
 /*
-	check(PendingGfxState->GetCurrentShader(ShaderStage::Hull) == ResourceCast(HullShaderRHI));
+	check(PendingGfxState->GetCurrentShader(ShaderStage::Hull)->GetId() == GetShaderId(HullShaderRHI));
 	FVulkanTextureBase* Texture = GetVulkanTextureFromRHITexture(NewTextureRHI);
 	VkImageLayout Layout = GetLayoutForDescriptor(Texture->Surface);
 	PendingGfxState->SetTexture(ShaderStage::Hull, TextureIndex, Texture, Layout);
@@ -209,7 +209,7 @@ void FVulkanCommandListContext::RHISetShaderTexture(FDomainShaderRHIParamRef Dom
 {
 	ensureMsgf(0, TEXT("Tessellation not supported yet!"));
 /*
-	check(PendingGfxState->GetCurrentShader(ShaderStage::Domain) == ResourceCast(DomainShaderRHI));
+	check(PendingGfxState->GetCurrentShader(ShaderStage::Domain)->GetId() == GetShaderId(DomainShaderRHI));
 	FVulkanTextureBase* Texture = GetVulkanTextureFromRHITexture(NewTextureRHI);
 	VkImageLayout Layout = GetLayoutForDescriptor(Texture->Surface);
 	PendingGfxState->SetTexture(ShaderStage::Domain, TextureIndex, Texture, Layout);
@@ -219,7 +219,7 @@ void FVulkanCommandListContext::RHISetShaderTexture(FDomainShaderRHIParamRef Dom
 void FVulkanCommandListContext::RHISetShaderTexture(FGeometryShaderRHIParamRef GeometryShaderRHI, uint32 TextureIndex, FTextureRHIParamRef NewTextureRHI)
 {
 #if VULKAN_SUPPORTS_GEOMETRY_SHADERS
-	check(PendingGfxState->GetCurrentShader(ShaderStage::Geometry) == ResourceCast(GeometryShaderRHI));
+	check(PendingGfxState->GetCurrentShader(ShaderStage::Geometry)->GetShaderKey() == GetShaderKey(GeometryShaderRHI));
 	FVulkanTextureBase* Texture = GetVulkanTextureFromRHITexture(NewTextureRHI);
 	VkImageLayout Layout = GetLayoutForDescriptor(Texture->Surface);
 	PendingGfxState->SetTextureForStage(ShaderStage::Geometry, TextureIndex, Texture, Layout);
@@ -230,7 +230,7 @@ void FVulkanCommandListContext::RHISetShaderTexture(FGeometryShaderRHIParamRef G
 
 void FVulkanCommandListContext::RHISetShaderTexture(FPixelShaderRHIParamRef PixelShaderRHI, uint32 TextureIndex, FTextureRHIParamRef NewTextureRHI)
 {
-	check(PendingGfxState->GetCurrentShader(ShaderStage::Pixel) == ResourceCast(PixelShaderRHI));
+	check(PendingGfxState->GetCurrentShader(ShaderStage::Pixel)->GetShaderKey() == GetShaderKey(PixelShaderRHI));
 	FVulkanTextureBase* Texture = GetVulkanTextureFromRHITexture(NewTextureRHI);
 	VkImageLayout Layout = GetLayoutForDescriptor(Texture->Surface);
 	PendingGfxState->SetTextureForStage(ShaderStage::Pixel, TextureIndex, Texture, Layout);
@@ -248,7 +248,7 @@ void FVulkanCommandListContext::RHISetShaderTexture(FComputeShaderRHIParamRef Co
 
 void FVulkanCommandListContext::RHISetShaderResourceViewParameter(FVertexShaderRHIParamRef VertexShaderRHI, uint32 TextureIndex, FShaderResourceViewRHIParamRef SRVRHI)
 {
-	check(PendingGfxState->GetCurrentShader(ShaderStage::Vertex) == ResourceCast(VertexShaderRHI));
+	check(PendingGfxState->GetCurrentShader(ShaderStage::Vertex)->GetShaderKey() == GetShaderKey(VertexShaderRHI));
 	FVulkanShaderResourceView* SRV = ResourceCast(SRVRHI);
 	PendingGfxState->SetSRVForStage(ShaderStage::Vertex, TextureIndex, SRV);
 }
@@ -257,7 +257,7 @@ void FVulkanCommandListContext::RHISetShaderResourceViewParameter(FHullShaderRHI
 {
 	ensureMsgf(0, TEXT("Tessellation not supported yet!"));
 /*
-	check(PendingGfxState->GetCurrentShader(ShaderStage::Hull) == ResourceCast(HullShaderRHI));
+	check(PendingGfxState->GetCurrentShader(ShaderStage::Hull)->GetId() == GetShaderId(HullShaderRHI));
 	FVulkanShaderResourceView* SRV = ResourceCast(SRVRHI);
 	PendingGfxState->SetSRV(ShaderStage::Hull, TextureIndex, SRV);
 */
@@ -267,7 +267,7 @@ void FVulkanCommandListContext::RHISetShaderResourceViewParameter(FDomainShaderR
 {
 	ensureMsgf(0, TEXT("Tessellation not supported yet!"));
 /*
-	check(PendingGfxState->GetCurrentShader(ShaderStage::Domain) == ResourceCast(DomainShaderRHI));
+	check(PendingGfxState->GetCurrentShader(ShaderStage::Domain)->GetId() == GetShaderId(DomainShaderRHI));
 	FVulkanShaderResourceView* SRV = ResourceCast(SRVRHI);
 	PendingGfxState->SetSRV(ShaderStage::Domain, TextureIndex, SRV);
 */
@@ -276,7 +276,7 @@ void FVulkanCommandListContext::RHISetShaderResourceViewParameter(FDomainShaderR
 void FVulkanCommandListContext::RHISetShaderResourceViewParameter(FGeometryShaderRHIParamRef GeometryShaderRHI,uint32 TextureIndex,FShaderResourceViewRHIParamRef SRVRHI)
 {
 #if VULKAN_SUPPORTS_GEOMETRY_SHADERS
-	check(PendingGfxState->GetCurrentShader(ShaderStage::Geometry) == ResourceCast(GeometryShaderRHI));
+	check(PendingGfxState->GetCurrentShader(ShaderStage::Geometry)->GetShaderKey() == GetShaderKey(GeometryShaderRHI));
 	FVulkanShaderResourceView* SRV = ResourceCast(SRVRHI);
 	PendingGfxState->SetSRVForStage(ShaderStage::Geometry, TextureIndex, SRV);
 #else
@@ -286,7 +286,7 @@ void FVulkanCommandListContext::RHISetShaderResourceViewParameter(FGeometryShade
 
 void FVulkanCommandListContext::RHISetShaderResourceViewParameter(FPixelShaderRHIParamRef PixelShaderRHI,uint32 TextureIndex,FShaderResourceViewRHIParamRef SRVRHI)
 {
-	check(PendingGfxState->GetCurrentShader(ShaderStage::Pixel) == ResourceCast(PixelShaderRHI));
+	check(PendingGfxState->GetCurrentShader(ShaderStage::Pixel)->GetShaderKey() == GetShaderKey(PixelShaderRHI));
 	FVulkanShaderResourceView* SRV = ResourceCast(SRVRHI);
 	PendingGfxState->SetSRVForStage(ShaderStage::Pixel, TextureIndex, SRV);
 }
@@ -301,7 +301,7 @@ void FVulkanCommandListContext::RHISetShaderResourceViewParameter(FComputeShader
 
 void FVulkanCommandListContext::RHISetShaderSampler(FVertexShaderRHIParamRef VertexShaderRHI, uint32 SamplerIndex, FSamplerStateRHIParamRef NewStateRHI)
 {
-	check(PendingGfxState->GetCurrentShader(ShaderStage::Vertex) == ResourceCast(VertexShaderRHI));
+	check(PendingGfxState->GetCurrentShader(ShaderStage::Vertex)->GetShaderKey() == GetShaderKey(VertexShaderRHI));
 	FVulkanSamplerState* Sampler = ResourceCast(NewStateRHI);
 	PendingGfxState->SetSamplerStateForStage(ShaderStage::Vertex, SamplerIndex, Sampler);
 }
@@ -310,7 +310,7 @@ void FVulkanCommandListContext::RHISetShaderSampler(FHullShaderRHIParamRef HullS
 {
 	ensureMsgf(0, TEXT("Tessellation not supported yet!"));
 /*
-	check(PendingGfxState->GetCurrentShader(ShaderStage::Hull) == ResourceCast(HullShaderRHI));
+	check(PendingGfxState->GetCurrentShader(ShaderStage::Hull)->GetId() == GetShaderId(HullShaderRHI));
 	FVulkanSamplerState* Sampler = ResourceCast(NewStateRHI);
 	PendingGfxState->SetSamplerState(ShaderStage::Hull, SamplerIndex, Sampler);
 */
@@ -320,7 +320,7 @@ void FVulkanCommandListContext::RHISetShaderSampler(FDomainShaderRHIParamRef Dom
 {
 	ensureMsgf(0, TEXT("Tessellation not supported yet!"));
 /*
-	check(PendingGfxState->GetCurrentShader(ShaderStage::Domain) == ResourceCast(DomainShaderRHI));
+	check(PendingGfxState->GetCurrentShader(ShaderStage::Domain)->GetId() == GetShaderId(DomainShaderRHI));
 	FVulkanSamplerState* Sampler = ResourceCast(NewStateRHI);
 	PendingGfxState->SetSamplerState(ShaderStage::Domain, SamplerIndex, Sampler);
 */
@@ -329,7 +329,7 @@ void FVulkanCommandListContext::RHISetShaderSampler(FDomainShaderRHIParamRef Dom
 void FVulkanCommandListContext::RHISetShaderSampler(FGeometryShaderRHIParamRef GeometryShaderRHI, uint32 SamplerIndex, FSamplerStateRHIParamRef NewStateRHI)
 {
 #if VULKAN_SUPPORTS_GEOMETRY_SHADERS
-	check(PendingGfxState->GetCurrentShader(ShaderStage::Geometry) == ResourceCast(GeometryShaderRHI));
+	check(PendingGfxState->GetCurrentShader(ShaderStage::Geometry)->GetShaderKey() == GetShaderKey(GeometryShaderRHI));
 	FVulkanSamplerState* Sampler = ResourceCast(NewStateRHI);
 	PendingGfxState->SetSamplerStateForStage(ShaderStage::Geometry, SamplerIndex, Sampler);
 #else
@@ -339,7 +339,7 @@ void FVulkanCommandListContext::RHISetShaderSampler(FGeometryShaderRHIParamRef G
 
 void FVulkanCommandListContext::RHISetShaderSampler(FPixelShaderRHIParamRef PixelShaderRHI, uint32 SamplerIndex, FSamplerStateRHIParamRef NewStateRHI)
 {
-	check(PendingGfxState->GetCurrentShader(ShaderStage::Pixel) == ResourceCast(PixelShaderRHI));
+	check(PendingGfxState->GetCurrentShader(ShaderStage::Pixel)->GetShaderKey() == GetShaderKey(PixelShaderRHI));
 	FVulkanSamplerState* Sampler = ResourceCast(NewStateRHI);
 	PendingGfxState->SetSamplerStateForStage(ShaderStage::Pixel, SamplerIndex, Sampler);
 }
@@ -355,7 +355,7 @@ void FVulkanCommandListContext::RHISetShaderSampler(FComputeShaderRHIParamRef Co
 
 void FVulkanCommandListContext::RHISetShaderParameter(FVertexShaderRHIParamRef VertexShaderRHI, uint32 BufferIndex, uint32 BaseIndex, uint32 NumBytes, const void* NewValue)
 {
-	check(PendingGfxState->GetCurrentShader(ShaderStage::Vertex) == ResourceCast(VertexShaderRHI));
+	check(PendingGfxState->GetCurrentShader(ShaderStage::Vertex)->GetShaderKey() == GetShaderKey(VertexShaderRHI));
 
 	PendingGfxState->SetPackedGlobalShaderParameter(ShaderStage::Vertex, BufferIndex, BaseIndex, NumBytes, NewValue);
 }
@@ -364,7 +364,7 @@ void FVulkanCommandListContext::RHISetShaderParameter(FHullShaderRHIParamRef Hul
 {
 	ensureMsgf(0, TEXT("Tessellation not supported yet!"));
 /*
-	check(PendingGfxState->GetCurrentShader(ShaderStage::Hull) == ResourceCast(HullShaderRHI));
+	check(PendingGfxState->GetCurrentShader(ShaderStage::Hull)->GetId() == GetShaderId(HullShaderRHI));
 
 	PendingGfxState->SetShaderParameter(ShaderStage::Hull, BufferIndex, BaseIndex, NumBytes, NewValue);
 */
@@ -374,7 +374,7 @@ void FVulkanCommandListContext::RHISetShaderParameter(FDomainShaderRHIParamRef D
 {
 	ensureMsgf(0, TEXT("Tessellation not supported yet!"));
 /*
-	check(PendingGfxState->GetCurrentShader(ShaderStage::Domain) == ResourceCast(DomainShaderRHI));
+	check(PendingGfxState->GetCurrentShader(ShaderStage::Domain)->GetId() == GetShaderId(DomainShaderRHI));
 
 	PendingGfxState->SetShaderParameter(ShaderStage::Domain, BufferIndex, BaseIndex, NumBytes, NewValue);
 */
@@ -383,7 +383,7 @@ void FVulkanCommandListContext::RHISetShaderParameter(FDomainShaderRHIParamRef D
 void FVulkanCommandListContext::RHISetShaderParameter(FGeometryShaderRHIParamRef GeometryShaderRHI, uint32 BufferIndex, uint32 BaseIndex, uint32 NumBytes, const void* NewValue)
 {
 #if VULKAN_SUPPORTS_GEOMETRY_SHADERS
-	check(PendingGfxState->GetCurrentShader(ShaderStage::Geometry) == ResourceCast(GeometryShaderRHI));
+	check(PendingGfxState->GetCurrentShader(ShaderStage::Geometry)->GetShaderKey() == GetShaderKey(GeometryShaderRHI));
 
 	PendingGfxState->SetPackedGlobalShaderParameter(ShaderStage::Geometry, BufferIndex, BaseIndex, NumBytes, NewValue);
 #else
@@ -393,7 +393,7 @@ void FVulkanCommandListContext::RHISetShaderParameter(FGeometryShaderRHIParamRef
 
 void FVulkanCommandListContext::RHISetShaderParameter(FPixelShaderRHIParamRef PixelShaderRHI, uint32 BufferIndex, uint32 BaseIndex, uint32 NumBytes, const void* NewValue)
 {
-	check(PendingGfxState->GetCurrentShader(ShaderStage::Pixel) == ResourceCast(PixelShaderRHI));
+	check(PendingGfxState->GetCurrentShader(ShaderStage::Pixel)->GetShaderKey() == GetShaderKey(PixelShaderRHI));
 
 	PendingGfxState->SetPackedGlobalShaderParameter(ShaderStage::Pixel, BufferIndex, BaseIndex, NumBytes, NewValue);
 }
@@ -594,7 +594,7 @@ inline void FVulkanCommandListContext::SetShaderUniformBuffer(ShaderStage::EStag
 #if VULKAN_ENABLE_AGGRESSIVE_STATS
 	SCOPE_CYCLE_COUNTER(STAT_VulkanSetUniformBufferTime);
 #endif
-	check(Shader == PendingGfxState->GetCurrentShader(Stage));
+	check(Shader->GetShaderKey() == PendingGfxState->GetCurrentShader(Stage)->GetShaderKey());
 
 	const FVulkanShaderHeader& CodeHeader = Shader->GetCodeHeader();
 	const bool bUseRealUBs = FVulkanPlatform::UseRealUBsOptimization(CodeHeader.bHasRealUBs);
diff --git a/Engine/Source/Runtime/VulkanRHI/Private/VulkanDebug.cpp b/Engine/Source/Runtime/VulkanRHI/Private/VulkanDebug.cpp
--- a/Engine/Source/Runtime/VulkanRHI/Private/VulkanDebug.cpp
+++ b/Engine/Source/Runtime/VulkanRHI/Private/VulkanDebug.cpp
@@ -3859,6 +3859,31 @@ void FWrapLayer::CreateWin32SurfaceKHR(VkResult Result, VkInstance Instance, con
 }
 #endif
 
+#if VULKAN_SUPPORTS_COLOR_CONVERSIONS
+void FWrapLayer::CreateSamplerYcbcrConversionKHR(VkResult Result, VkDevice Device, const VkSamplerYcbcrConversionCreateInfo* CreateInfo, VkSamplerYcbcrConversion* YcbcrConversion)
+{
+	if (Result == VK_RESULT_MAX_ENUM)
+	{
+		DevicePrintfBeginResult(Device, FString::Printf(TEXT("vkCreateSamplerYcbcrConversionKHR(CreateInfo=0x%p, YcbcrConversion=0x%p)[...]"), CreateInfo, YcbcrConversion));
+	}
+	else
+	{
+		PrintResultAndNamedHandle(Result, TEXT("SamplerYcbcrConversionKHR"), *YcbcrConversion);
+	}
+	FlushDebugWrapperLog();
+}
+
+void FWrapLayer::DestroySamplerYcbcrConversionKHR(VkResult Result, VkDevice Device, VkSamplerYcbcrConversion YcbcrConversion)
+{
+	if (Result == VK_RESULT_MAX_ENUM)
+	{
+#if VULKAN_ENABLE_DUMP_LAYER
+		DevicePrintfBegin(Device, FString::Printf(TEXT("vkDestroySamplerYcbcrConversionKHR(YcbcrConversion=0x%p)"), YcbcrConversion));
+#endif
+	}
+}
+#endif
+
 #if VULKAN_ENABLE_IMAGE_TRACKING_LAYER
 namespace VulkanRHI
 {
diff --git a/Engine/Source/Runtime/VulkanRHI/Private/VulkanDescriptorSets.h b/Engine/Source/Runtime/VulkanRHI/Private/VulkanDescriptorSets.h
--- a/Engine/Source/Runtime/VulkanRHI/Private/VulkanDescriptorSets.h
+++ b/Engine/Source/Runtime/VulkanRHI/Private/VulkanDescriptorSets.h
@@ -48,6 +48,7 @@ struct FDescriptorSetRemappingInfo
 		FRemappingInfo			Remapping;
 
 		bool					bHasConstantData = false;
+		const bool				bPadding = false;	// padding is need on memcmp/MemCrc to make sure mem align
 		//bool					bIsRedundant = false;
 	};
 
@@ -91,6 +92,58 @@ struct FDescriptorSetRemappingInfo
 
 	TArray<FInputAttachmentData>						InputAttachmentData;
 
+	inline bool operator==(const FDescriptorSetRemappingInfo& In) const
+	{
+		if (InputAttachmentData.Num() != In.InputAttachmentData.Num())
+		{
+			return false;
+		}
+
+		if (SetInfos.Num() != In.SetInfos.Num())
+		{
+			return false;
+		}
+
+		if (FMemory::Memcmp(InputAttachmentData.GetData(), In.InputAttachmentData.GetData(), sizeof(FInputAttachmentData) * InputAttachmentData.Num()))
+		{
+			return false;
+		}
+
+		for (int32 SetInfosIndex = 0; SetInfosIndex < SetInfos.Num(); ++SetInfosIndex)
+		{
+			int32 SetInfosNums = SetInfos[SetInfosIndex].Types.Num();
+			if (SetInfos[SetInfosIndex].NumBufferInfos != In.SetInfos[SetInfosIndex].NumBufferInfos ||
+				SetInfos[SetInfosIndex].NumImageInfos != In.SetInfos[SetInfosIndex].NumImageInfos ||
+				SetInfosNums != In.SetInfos[SetInfosIndex].Types.Num() ||
+				(SetInfosNums != 0 && FMemory::Memcmp(SetInfos[SetInfosIndex].Types.GetData(), In.SetInfos[SetInfosIndex].Types.GetData(), sizeof(VkDescriptorType) * SetInfosNums)))
+			{
+				return false;
+			}
+		}
+
+		for (uint32 StageInfosIndex = 0; StageInfosIndex < ShaderStage::NumStages; ++StageInfosIndex)
+		{
+			if (StageInfos[StageInfosIndex].PackedUBDescriptorSet != In.StageInfos[StageInfosIndex].PackedUBDescriptorSet ||
+				StageInfos[StageInfosIndex].Pad0 != In.StageInfos[StageInfosIndex].Pad0 ||
+				StageInfos[StageInfosIndex].Globals.Num() != In.StageInfos[StageInfosIndex].Globals.Num() ||
+				StageInfos[StageInfosIndex].PackedUBBindingIndices.Num() != In.StageInfos[StageInfosIndex].PackedUBBindingIndices.Num() ||
+				StageInfos[StageInfosIndex].UniformBuffers.Num() != In.StageInfos[StageInfosIndex].UniformBuffers.Num() ||
+				FMemory::Memcmp(StageInfos[StageInfosIndex].Globals.GetData(), In.StageInfos[StageInfosIndex].Globals.GetData(), sizeof(FRemappingInfo) * StageInfos[StageInfosIndex].Globals.Num()) ||
+				FMemory::Memcmp(StageInfos[StageInfosIndex].PackedUBBindingIndices.GetData(), In.StageInfos[StageInfosIndex].PackedUBBindingIndices.GetData(), sizeof(uint16) * StageInfos[StageInfosIndex].PackedUBBindingIndices.Num()) ||
+				FMemory::Memcmp(StageInfos[StageInfosIndex].UniformBuffers.GetData(), In.StageInfos[StageInfosIndex].UniformBuffers.GetData(), sizeof(FUBRemappingInfo) * StageInfos[StageInfosIndex].UniformBuffers.Num()))
+			{
+				return false;
+			}
+		}
+
+		return true;
+	}
+
+	inline bool operator != (const FDescriptorSetRemappingInfo& In) const
+	{
+		return !(*this == In);
+	}
+
 	inline bool IsEmpty() const
 	{
 		if (SetInfos.Num() == 0)
@@ -193,6 +246,15 @@ struct FDescriptorSetRemappingInfo
 		return NewBindingIndex;
 	}
 */
+
+
+	static_assert(
+		sizeof(FRemappingInfo) == (sizeof(FRemappingInfo::NewDescriptorSet) + sizeof(FRemappingInfo::NewBindingIndex)),
+		"FRemappingInfo should not have padding! structure is used for MemCrc32/Memcmp");
+
+	static_assert(
+		sizeof(FUBRemappingInfo) == (sizeof(FRemappingInfo) + sizeof(FUBRemappingInfo::bHasConstantData) + sizeof(FUBRemappingInfo::bPadding)),
+		"FUBRemappingInfo should not have padding! structure is used for MemCrc32/Memcmp");
 };
 
 // Information for the layout of descriptor sets; does not hold runtime objects
@@ -219,12 +281,13 @@ public:
 		return SetLayouts;
 	}
 
-	void AddBindingsForStage(VkShaderStageFlagBits StageFlags, ShaderStage::EStage DescSet, const FVulkanShaderHeader& CodeHeader, const FImmutableSamplerState* const ImmutableSamplerState = nullptr);
 	void ProcessBindingsForStage(VkShaderStageFlagBits StageFlags, ShaderStage::EStage DescSetStage, const FVulkanShaderHeader& CodeHeader, FUniformBufferGatherInfo& OutUBGatherInfo) const;
 
 	template<bool bIsCompute>
 	void FinalizeBindings(const FUniformBufferGatherInfo& UBGatherInfo, const TArrayView<const FSamplerStateRHIParamRef>& ImmutableSamplers);
 
+	void GenerateHash();
+
 	friend uint32 GetTypeHash(const FVulkanDescriptorSetsLayoutInfo& In)
 	{
 		return In.Hash;
@@ -250,12 +313,17 @@ public:
 				return false;
 			}
 
-			if (NumBindings != 0 && FMemory::Memcmp(&In.SetLayouts[Index].LayoutBindings[0], &SetLayouts[Index].LayoutBindings[0], NumBindings * sizeof(VkDescriptorSetLayoutBinding)))
+			if (NumBindings != 0 && FMemory::Memcmp(In.SetLayouts[Index].LayoutBindings.GetData(), SetLayouts[Index].LayoutBindings.GetData(), NumBindings * sizeof(VkDescriptorSetLayoutBinding)) != 0)
 			{
 				return false;
 			}
 		}
 
+		if (RemappingInfo != In.RemappingInfo)
+		{
+			return false;
+		}
+
 		return true;
 	}
 
@@ -642,7 +710,7 @@ public:
 		return RemappingInfo->InputAttachmentData;
 	}
 
-	void Initialize(const FDescriptorSetRemappingInfo& InRemappingInfo, FVulkanShader** Shaders);
+	void Initialize(const FDescriptorSetRemappingInfo& InRemappingInfo);
 
 protected:
 	// Cached data from FDescriptorSetRemappingInfo
diff --git a/Engine/Source/Runtime/VulkanRHI/Private/VulkanDevice.cpp b/Engine/Source/Runtime/VulkanRHI/Private/VulkanDevice.cpp
--- a/Engine/Source/Runtime/VulkanRHI/Private/VulkanDevice.cpp
+++ b/Engine/Source/Runtime/VulkanRHI/Private/VulkanDevice.cpp
@@ -629,7 +629,6 @@ void FVulkanDevice::SetupFormats()
 #if VULKAN_SUPPORTS_COLOR_CONVERSIONS
 VkSamplerYcbcrConversion FVulkanDevice::CreateSamplerColorConversion(const VkSamplerYcbcrConversionCreateInfo& CreateInfo)
 {
-
 	const uint32 CreateInfoHash = FCrc::MemCrc32(&CreateInfo, sizeof(CreateInfo));
 	VkSamplerYcbcrConversion* const FindResult = SamplerColorConversionMap.Find(CreateInfoHash);
 	if (FindResult != nullptr)
@@ -639,7 +638,7 @@ VkSamplerYcbcrConversion FVulkanDevice::CreateSamplerColorConversion(const VkSam
 	else
 	{
 		VkSamplerYcbcrConversion NewConversion;
-		VERIFYVULKANRESULT(vkCreateSamplerYcbcrConversionKHR(GetInstanceHandle(), &CreateInfo, nullptr, &NewConversion));
+		VERIFYVULKANRESULT(VulkanRHI::vkCreateSamplerYcbcrConversionKHR(GetInstanceHandle(), &CreateInfo, nullptr, &NewConversion));
 		SamplerColorConversionMap.Add(CreateInfoHash, NewConversion);
 		return NewConversion;
 	}
@@ -874,7 +873,7 @@ void FVulkanDevice::Destroy()
 #if VULKAN_SUPPORTS_COLOR_CONVERSIONS
 	for (const auto& Pair : SamplerColorConversionMap)
 	{
-		vkDestroySamplerYcbcrConversionKHR(GetInstanceHandle(), Pair.Value, nullptr);
+		VulkanRHI::vkDestroySamplerYcbcrConversionKHR(GetInstanceHandle(), Pair.Value, nullptr);
 	}
 	SamplerColorConversionMap.Reset();
 #endif
@@ -926,13 +925,13 @@ void FVulkanDevice::Destroy()
 
 	ResourceHeapManager.Deinit();
 
-	FRHIResource::FlushPendingDeletes();
-	DeferredDeletionQueue.Clear();
-
 	delete TransferQueue;
 	delete ComputeQueue;
 	delete GfxQueue;
 
+	FRHIResource::FlushPendingDeletes();
+	DeferredDeletionQueue.Clear();
+
 	FenceManager.Deinit();
 	MemoryManager.Deinit();
 
@@ -1052,7 +1051,7 @@ void FVulkanDevice::NotifyDeletedGfxPipeline(class FVulkanRHIGraphicsPipelineSta
 
 void FVulkanDevice::NotifyDeletedComputePipeline(class FVulkanComputePipeline* Pipeline)
 {
-	if (ComputeContext != ImmediateContext)
+	if (ComputeContext && ComputeContext != ImmediateContext)
 	{
 		ComputeContext->PendingComputeState->NotifyDeletedPipeline(Pipeline);
 	}
diff --git a/Engine/Source/Runtime/VulkanRHI/Private/VulkanDevice.h b/Engine/Source/Runtime/VulkanRHI/Private/VulkanDevice.h
--- a/Engine/Source/Runtime/VulkanRHI/Private/VulkanDevice.h
+++ b/Engine/Source/Runtime/VulkanRHI/Private/VulkanDevice.h
@@ -28,6 +28,7 @@ struct FOptionalVulkanDeviceExtensions
 	uint32 HasEXTValidationCache : 1;
 	uint32 HasAMDBufferMarker : 1;
 	uint32 HasGoogleDisplayTiming : 1;
+	uint32 HasYcbcrSampler : 1;
 };
 
 class FVulkanDevice
@@ -249,9 +250,6 @@ public:
 	void SubmitCommandsAndFlushGPU();
 
 	FVulkanOcclusionQueryPool* AcquireOcclusionQueryPool(uint32 NumQueries);
-/*
-	FVulkanTimestampQueryPool* PrepareTimestampQueryPool(bool& bOutRequiresReset);
-*/
 
 	inline class FVulkanPipelineStateCacheManager* GetPipelineStateCache()
 	{
@@ -329,9 +327,6 @@ private:
 
 	TArray<FVulkanOcclusionQueryPool*> UsedOcclusionQueryPools;
 	TArray<FVulkanOcclusionQueryPool*> FreeOcclusionQueryPools;
-/*
-	FVulkanTimestampQueryPool* TimestampQueryPool = nullptr;
-*/
 
 	uint64 TimestampValidBitsMask = 0;
 
diff --git a/Engine/Source/Runtime/VulkanRHI/Private/VulkanLayers.cpp b/Engine/Source/Runtime/VulkanRHI/Private/VulkanLayers.cpp
--- a/Engine/Source/Runtime/VulkanRHI/Private/VulkanLayers.cpp
+++ b/Engine/Source/Runtime/VulkanRHI/Private/VulkanLayers.cpp
@@ -691,4 +691,8 @@ void FVulkanDevice::ParseOptionalDeviceExtensions(const TArray<const ANSICHAR *>
 #if VULKAN_SUPPORTS_GOOGLE_DISPLAY_TIMING
 	OptionalDeviceExtensions.HasGoogleDisplayTiming = HasExtension(VK_GOOGLE_DISPLAY_TIMING_EXTENSION_NAME);
 #endif
+
+#if VULKAN_SUPPORTS_COLOR_CONVERSIONS
+	OptionalDeviceExtensions.HasYcbcrSampler = HasExtension(VK_KHR_SAMPLER_YCBCR_CONVERSION_EXTENSION_NAME) && HasExtension(VK_KHR_BIND_MEMORY_2_EXTENSION_NAME) && HasExtension(VK_KHR_GET_MEMORY_REQUIREMENTS_2_EXTENSION_NAME);
+#endif
 }
diff --git a/Engine/Source/Runtime/VulkanRHI/Private/VulkanMemory.cpp b/Engine/Source/Runtime/VulkanRHI/Private/VulkanMemory.cpp
--- a/Engine/Source/Runtime/VulkanRHI/Private/VulkanMemory.cpp
+++ b/Engine/Source/Runtime/VulkanRHI/Private/VulkanMemory.cpp
@@ -27,6 +27,17 @@ static void CaptureCallStack(FString& OutCallstack)
 }
 #endif
 
+#if (UE_BUILD_DEBUG || UE_BUILD_DEVELOPMENT)
+static int32 GForceCoherent = 0;
+static FAutoConsoleVariableRef CVarForceCoherentOperations(
+	TEXT("r.Vulkan.ForceCoherentOperations"),
+	GForceCoherent,
+	TEXT("1 forces memory invalidation and flushing of coherent memory\n"),
+	ECVF_ReadOnly
+);
+#else
+constexpr int32 GForceCoherent = 0;
+#endif
 
 namespace VulkanRHI
 {
@@ -364,7 +375,7 @@ namespace VulkanRHI
 
 	void FDeviceMemoryAllocation::FlushMappedMemory(VkDeviceSize InOffset, VkDeviceSize InSize)
 	{
-		if (!IsCoherent())
+		if (!IsCoherent() || GForceCoherent != 0)
 		{
 			check(IsMapped());
 			check(InOffset + InSize <= Size);
@@ -379,7 +390,7 @@ namespace VulkanRHI
 
 	void FDeviceMemoryAllocation::InvalidateMappedMemory(VkDeviceSize InOffset, VkDeviceSize InSize)
 	{
-		if (!IsCoherent())
+		if (!IsCoherent() || GForceCoherent != 0)
 		{
 			check(IsMapped());
 			check(InOffset + InSize <= Size);
@@ -1932,6 +1943,19 @@ namespace VulkanRHI
 		}
 	}
 
+	void FDeferredDeletionQueue::OnCmdBufferDeleted(FVulkanCmdBuffer* DeletedCmdBuffer)
+	{
+		FScopeLock ScopeLock(&CS);
+		for (int32 Index = 0; Index < Entries.Num(); ++Index)
+		{
+			FEntry& Entry = Entries[Index];
+			if (Entry.CmdBuffer == DeletedCmdBuffer)
+			{
+				Entry.CmdBuffer = nullptr;
+			}
+		}
+	}
+
 	FTempFrameAllocationBuffer::FTempFrameAllocationBuffer(FVulkanDevice* InDevice)
 		: FDeviceChild(InDevice)
 		, BufferIndex(0)
diff --git a/Engine/Source/Runtime/VulkanRHI/Private/VulkanPendingState.cpp b/Engine/Source/Runtime/VulkanRHI/Private/VulkanPendingState.cpp
--- a/Engine/Source/Runtime/VulkanRHI/Private/VulkanPendingState.cpp
+++ b/Engine/Source/Runtime/VulkanRHI/Private/VulkanPendingState.cpp
@@ -265,7 +265,6 @@ void FVulkanDescriptorPoolsManager::ReleasePoolSet(FVulkanDescriptorPoolSetConta
 {
 	PoolSet.Reset();
 	PoolSet.SetUsed(false);
-	PoolSet.Reset();
 }
 
 void FVulkanDescriptorPoolsManager::GC()
@@ -304,9 +303,9 @@ void FVulkanDescriptorPoolsManager::GC()
 
 FVulkanPendingComputeState::~FVulkanPendingComputeState()
 {
-	for (auto& Pair : PipelineStates)
+	for (auto It = PipelineStates.CreateIterator(); It; ++It)
 	{
-		FVulkanCommonPipelineDescriptorState* State = Pair.Value;
+		FVulkanCommonPipelineDescriptorState* State = It->Value;
 		delete State;
 	}
 }
@@ -474,7 +473,7 @@ void FVulkanPendingGfxState::PrepareForDraw(FVulkanCmdBuffer* CmdBuffer)
 #if VULKAN_ENABLE_SHADER_DEBUG_NAMES
 						UE_LOG(LogVulkanRHI, Warning, TEXT("Missing binding on location %d in '%s' vertex shader"),
 							CurrBinding.binding,
-							*CurrentPipeline->GetShader(SF_Vertex)->GetDebugName());
+							*GetCurrentShader(ShaderStage::Vertex)->GetDebugName());
 #else
 						UE_LOG(LogVulkanRHI, Warning, TEXT("Missing binding on location %d in vertex shader"), CurrBinding.binding);
 #endif
diff --git a/Engine/Source/Runtime/VulkanRHI/Private/VulkanPendingState.h b/Engine/Source/Runtime/VulkanRHI/Private/VulkanPendingState.h
--- a/Engine/Source/Runtime/VulkanRHI/Private/VulkanPendingState.h
+++ b/Engine/Source/Runtime/VulkanRHI/Private/VulkanPendingState.h
@@ -358,7 +358,7 @@ public:
 
 			CurrentState->Reset();
 
-			PrimitiveType = InGfxPipeline->PipelineStateInitializer.PrimitiveType;
+			PrimitiveType = InGfxPipeline->PrimitiveType;
 
 			return true;
 		}
diff --git a/Engine/Source/Runtime/VulkanRHI/Private/VulkanPipeline.cpp b/Engine/Source/Runtime/VulkanRHI/Private/VulkanPipeline.cpp
--- a/Engine/Source/Runtime/VulkanRHI/Private/VulkanPipeline.cpp
+++ b/Engine/Source/Runtime/VulkanRHI/Private/VulkanPipeline.cpp
@@ -16,10 +16,11 @@
 #include "VulkanContext.h"
 #include "GlobalShader.h"
 #include "VulkanLLM.h"
+#include "Misc/ScopeRWLock.h"
 
 static const double HitchTime = 1.0 / 1000.0;
 
-static FCriticalSection GShaderHashToGfxEntriesMapCS;
+static FCriticalSection EntryKeyToGfxPipelineMapCS;
 
 TAutoConsoleVariable<int32> CVarPipelineLRUCacheEvictBinaryPreloadScreen(
 	TEXT("r.Vulkan.PipelineLRUCacheEvictBinaryPreloadScreen"),
@@ -104,8 +105,11 @@ FVulkanPipeline::FVulkanPipeline(FVulkanDevice* InDevice)
 
 FVulkanPipeline::~FVulkanPipeline()
 {
-	Device->GetDeferredDeletionQueue().EnqueueResource(VulkanRHI::FDeferredDeletionQueue::EType::Pipeline, Pipeline);
-	Pipeline = VK_NULL_HANDLE;
+	if (Pipeline != VK_NULL_HANDLE)
+	{
+		Device->GetDeferredDeletionQueue().EnqueueResource(VulkanRHI::FDeferredDeletionQueue::EType::Pipeline, Pipeline);
+		Pipeline = VK_NULL_HANDLE;
+	}
 	/* we do NOT own Layout !*/
 }
 
@@ -121,13 +125,13 @@ FVulkanComputePipeline::~FVulkanComputePipeline()
 }
 
 #if VULKAN_ENABLE_LRU_CACHE
-FVulkanGfxPipeline::FVulkanGfxPipeline(FVulkanDevice* InDevice, uint32 InGfxEntryHash, uint32 InShaderHash)
+FVulkanGfxPipeline::FVulkanGfxPipeline(FVulkanDevice* InDevice, TGfxPipelineEntrySharedPtr InGfxPipelineEntry)
 	: FVulkanPipeline(InDevice)
-	, bRuntimeObjectsValid(false)
-	, GfxEntryHash(InGfxEntryHash)
+	, GfxPipelineEntry(MoveTemp(InGfxPipelineEntry))
 	, PipelineCacheSize(0)
 	, RecentFrame(0)
-	, ShaderHash(InShaderHash)
+	, bTrackedByLRU(false)
+	, bRuntimeObjectsValid(false)
 #else
 FVulkanGfxPipeline::FVulkanGfxPipeline(FVulkanDevice* InDevice)
 	: FVulkanPipeline(InDevice)
@@ -167,7 +171,7 @@ static TAutoConsoleVariable<int32> GEnablePipelineCacheLoadCvar(
 
 static TAutoConsoleVariable<int32> GPipelineCacheFromShaderPipelineCacheCvar(
 	TEXT("r.Vulkan.PipelineCacheFromShaderPipelineCache"),
-	1,
+	PLATFORM_ANDROID && !(PLATFORM_LUMIN || PLATFORM_LUMINGL4),
 	TEXT("0 look for a pipeline cache in the normal locations with the normal names.")
 	TEXT("1 tie the vulkan pipeline cache to the shader pipeline cache, use the PSOFC guid as part of the filename, etc."),
 	ECVF_ReadOnly
@@ -186,14 +190,13 @@ static FAutoConsoleVariableRef GEnablePipelineCacheCompressionCvar(
 FVulkanPipelineStateCacheManager::FGfxPipelineEntry::~FGfxPipelineEntry()
 {
 	check(!bLoaded);
-	check(!bConsumed);
 }
 
-void FVulkanPipelineStateCacheManager::FGfxPipelineEntry::CreateShaderModules()
+void FVulkanPipelineStateCacheManager::FGfxPipelineEntry::GetOrCreateShaderModules(FVulkanShader*const* Shaders)
 {
 	for (int32 Index = 0; Index < ShaderStage::NumStages; ++Index)
 	{
-		FVulkanShader* Shader = VulkanShaders[Index].GetReference();
+		FVulkanShader* Shader = Shaders[Index];
 		if (Shader)
 		{
 			ShaderModules[Index] = Shader->GetOrCreateHandle(Layout, Layout->GetDescriptorSetLayoutHash());
@@ -201,11 +204,13 @@ void FVulkanPipelineStateCacheManager::FGfxPipelineEntry::CreateShaderModules()
 	}
 }
 
-void FVulkanPipelineStateCacheManager::FGfxPipelineEntry::PurgeShaderModules()
+void FVulkanPipelineStateCacheManager::FGfxPipelineEntry::PurgeShaderModules(FVulkanShader*const* Shaders)
 {
+	check(!bLoaded);
+
 	for (int32 Index = 0; Index < ShaderStage::NumStages; ++Index)
 	{
-		FVulkanShader* Shader = VulkanShaders[Index].GetReference();
+		FVulkanShader* Shader = Shaders[Index];
 		if (Shader)
 		{
 			Shader->PurgeShaderModules();
@@ -214,11 +219,23 @@ void FVulkanPipelineStateCacheManager::FGfxPipelineEntry::PurgeShaderModules()
 	}
 }
 
-FVulkanPipelineStateCacheManager::FComputePipelineEntry::~FComputePipelineEntry()
+void FVulkanPipelineStateCacheManager::FGfxPipelineEntry::PurgeLoadedShaderModules(FVulkanDevice* InDevice)
 {
-	check(!bLoaded);
+	check(bLoaded);
+
+	for (int32 Index = 0; Index < ShaderStage::NumStages; ++Index)
+	{
+		if (ShaderModules[Index] != VK_NULL_HANDLE)
+		{
+			VulkanRHI::vkDestroyShaderModule(InDevice->GetInstanceHandle(), ShaderModules[Index], VULKAN_CPU_ALLOCATOR);
+			ShaderModules[Index] = VK_NULL_HANDLE;
+		}
+	}
+
+	bLoaded = false;
 }
 
+
 FVulkanPipelineStateCacheManager::FVulkanPipelineStateCacheManager(FVulkanDevice* InDevice)
 	: Device(InDevice)
 	, bEvictImmediately(false)
@@ -252,10 +269,6 @@ FVulkanPipelineStateCacheManager::~FVulkanPipelineStateCacheManager()
 		delete Pair.Value;
 	}
 
-#if VULKAN_ENABLE_LRU_CACHE
-	PipelineLRU.Empty();
-#endif
-
 	VulkanRHI::vkDestroyPipelineCache(Device->GetInstanceHandle(), PipelineCache, VULKAN_CPU_ALLOCATOR);
 	PipelineCache = VK_NULL_HANDLE;
 }
@@ -357,129 +370,6 @@ bool FVulkanPipelineStateCacheManager::Load(const TArray<FString>& CacheFilename
 	}
 #endif
 
-#if VULKAN_ENABLE_GENERIC_PIPELINE_CACHE_FILE
-	for (const FString& CacheFilename : CacheFilenames)
-	{
-		TArray<uint8> MemFile;
-		if (FFileHelper::LoadFileToArray(MemFile, *CacheFilename, FILEREAD_Silent))
-		{
-			FMemoryReader Ar(MemFile);
-
-			FVulkanPipelineStateCacheFile File;
-
-			FShaderUCodeCache::TDataMap FileShaderCacheData;
-			File.ShaderCache = &FileShaderCacheData;
-
-			bool Valid = File.Load(Ar, *CacheFilename);
-			if (!Valid)
-			{
-				UE_LOG(LogVulkanRHI, Warning, TEXT("Unable to load pipeline cache '%s'"), *CacheFilename);
-				continue;
-			}
-
-			// Create the binary cache if we haven't already
-			if (PipelineCache == VK_NULL_HANDLE)
-			{
-				VkPipelineCacheCreateInfo PipelineCacheInfo;
-				ZeroVulkanStruct(PipelineCacheInfo, VK_STRUCTURE_TYPE_PIPELINE_CACHE_CREATE_INFO);
-				VERIFYVULKANRESULT(VulkanRHI::vkCreatePipelineCache(Device->GetInstanceHandle(), &PipelineCacheInfo, VULKAN_CPU_ALLOCATOR, &PipelineCache));
-			}
-
-			// Not using TMap::Append to avoid copying duplicate microcode
-			for (auto& Pair : FileShaderCacheData)
-			{
-				if (!ShaderCache.Data.Find(Pair.Key))
-				{
-					ShaderCache.Data.Add(Pair.Key, Pair.Value);
-				}
-			}
-
-			double BeginTime = FPlatformTime::Seconds();
-			{
-				for (int32 Index = 0; Index < File.GfxPipelineEntries.Num(); ++Index)
-				{
-					FGfxPipelineEntry* GfxEntry = File.GfxPipelineEntries[Index];
-
-					FShaderHashes ShaderHashes;
-					for (int32 i = 0; i < ShaderStage::NumStages; ++i)
-					{
-						ShaderHashes.Stages[i] = GfxEntry->ShaderHashes[i];
-
-						GfxEntry->ShaderMicrocodes[i] = ShaderCache.Get(GfxEntry->ShaderHashes[i]);
-					}
-					ShaderHashes.Finalize();
-
-					uint32 EntryHash = GfxEntry->GetEntryHash();
-					if (GfxPipelineEntries.Find(EntryHash))
-					{
-						delete GfxEntry;
-					}
-					else
-					{
-						FHashToGfxPipelinesMap* Found = ShaderHashToGfxPipelineMap.Find(ShaderHashes);
-						if (!Found)
-						{
-							Found = &ShaderHashToGfxPipelineMap.Add(ShaderHashes);
-						}
-
-						CreatGfxEntryRuntimeObjects(GfxEntry);
-#if VULKAN_ENABLE_LRU_CACHE
-						FVulkanGfxPipeline* Pipeline = new FVulkanGfxPipeline(Device, EntryHash, ShaderHashes.Hash);
-#else
-						FVulkanGfxPipeline* Pipeline = new FVulkanGfxPipeline(Device);
-#endif
-						CreateGfxPipelineFromEntry(GfxEntry, Pipeline);
-
-						Found->Add(EntryHash, Pipeline);
-						GfxPipelineEntries.Add(EntryHash, GfxEntry);
-						GfxEntry->bConsumed = true;
-					}
-				}
-			}
-
-			for (int32 Index = 0; Index < File.ComputePipelineEntries.Num(); ++Index)
-			{
-				FComputePipelineEntry* ComputeEntry = File.ComputePipelineEntries[Index];
-				ComputeEntry->ShaderMicrocode = ShaderCache.Get(ComputeEntry->ShaderHash);
-				ComputeEntry->CalculateEntryHash();
-
-				if (ComputePipelineEntries.Find(ComputeEntry->EntryHash))
-				{
-					delete ComputeEntry;
-				}
-				else
-				{
-					CreateComputeEntryRuntimeObjects(ComputeEntry);
-
-					FVulkanComputePipeline* Pipeline = CreateComputePipelineFromEntry(ComputeEntry);
-					ComputeEntryHashToPipelineMap.Add(ComputeEntry->EntryHash, Pipeline);
-					ComputePipelineEntries.Add(ComputeEntry->EntryHash, ComputeEntry);
-					Pipeline->AddRef();
-				}
-			}
-
-			double EndTime = FPlatformTime::Seconds();
-			UE_LOG(LogVulkanRHI, Display, TEXT("Loaded pipeline cache in %.2f seconds"), (float)(EndTime - BeginTime));
-		}
-		else
-		{
-			UE_LOG(LogVulkanRHI, Warning, TEXT("Unable to load pipeline cache '%s'"), *CacheFilename);
-		}
-	}
-
-#endif
-
-#if VULKAN_ENABLE_GENERIC_PIPELINE_CACHE_FILE
-	if (ShaderCache.Data.Num())
-	{
-		UE_LOG(LogVulkanRHI, Display, TEXT("Pipeline cache: %d Gfx Pipelines, %d Compute Pipelines, %d Microcodes"), GfxPipelineEntries.Num(), ComputePipelineEntries.Num(), ShaderCache.Data.Num());
-	}
-	else
-	{
-		UE_LOG(LogVulkanRHI, Display, TEXT("Pipeline cache: No pipeline cache(s) loaded"));
-	}
-#endif
-
 	// Lazily create the cache in case the load failed
 	if (PipelineCache == VK_NULL_HANDLE)
 	{
@@ -491,18 +381,6 @@ bool FVulkanPipelineStateCacheManager::Load(const TArray<FString>& CacheFilename
 	return bResult;
 }
 
-void FVulkanPipelineStateCacheManager::DestroyPipeline(FVulkanGfxPipeline* Pipeline)
-{
-	ensure(0);
-/*
-	if (Pipeline->Release() == 0)
-	{
-		const FVulkanGfxPipelineStateKey* Key = KeyToGfxPipelineMap.FindKey(Pipeline);
-		check(Key);
-		KeyToGfxPipelineMap.Remove(*Key);
-	}*/
-}
-
 void FVulkanPipelineStateCacheManager::InitAndLoad(const TArray<FString>& CacheFilenames)
 {
 	if (GEnablePipelineCacheLoadCvar.GetValueOnAnyThread() == 0)
@@ -619,16 +497,6 @@ void FVulkanPipelineStateCacheManager::OnShaderPipelineCachePrecompilationComple
 	if (!bLinkedToPSOFCSucessfulLoaded)
 	{
 		Save(LinkedToPSOFCCacheFolderFilename, true);
-#if PLATFORM_ANDROID && USE_ANDROID_JNI
-		static const auto CVar = IConsoleManager::Get().FindTConsoleVariableDataInt(TEXT("r.ProgramBinaryCache.RestartAndroidAfterPrecompile"));
-		const bool bRestart = CVar->GetValueOnAnyThread() != 0;
-		if (bRestart)
-		{
-			UE_LOG(LogVulkanRHI, Log, TEXT("FVulkanPipelineStateCacheManager::OnShaderPipelineCachePrecompilationComplete, restarting app by request."));
-			extern void AndroidThunkCpp_RestartApplication();
-			AndroidThunkCpp_RestartApplication();
-		}
-#endif
 	}
 
 	// Want to ignore any subsequent Shader Pipeline Cache opening/closing, eg when loading modules
@@ -722,116 +590,57 @@ void FVulkanPipelineStateCacheManager::Save(const FString& CacheFilename, bool b
 		}
 	}
 #endif
-
-#if VULKAN_ENABLE_GENERIC_PIPELINE_CACHE_FILE
-	{
-		// Now the generic cache
-		TArray<uint8> MemFile;
-		FMemoryWriter Ar(MemFile);
-		FVulkanPipelineStateCacheFile File;
-
-		File.Header.Version = VERSION;
-		File.Header.SizeOfGfxEntry = (int32)sizeof(FGfxPipelineEntry);
-		File.Header.SizeOfComputeEntry = (int32)sizeof(FComputePipelineEntry);
-		File.Header.UncompressedSize = 0;
-
-		// Shader ucode cache
-		File.ShaderCache = &ShaderCache.Data;
-
-		// Then Gfx entries
-		GfxPipelineEntries.GenerateValueArray(File.GfxPipelineEntries);
-
-		// And Compute entries
-		ComputePipelineEntries.GenerateValueArray(File.ComputePipelineEntries);
-
-		File.Save(Ar);
-
-		if (FFileHelper::SaveArrayToFile(MemFile, *CacheFilename))
-		{
-			UE_LOG(LogVulkanRHI, Display, TEXT("Saved pipeline cache file '%s', %d Gfx Pipelines, %d Compute Pipelines, %d Microcodes, %d bytes"), *CacheFilename, GfxPipelineEntries.Num(), ComputePipelineEntries.Num(), ShaderCache.Data.Num(), MemFile.Num());
-		}
-	}
-
-#endif
-}
-
-#if VULKAN_ENABLE_LRU_CACHE
-FVulkanGfxPipeline* FVulkanPipelineStateCacheManager::CreateAsEvicted(FGfxPipelineEntry* GfxEntry, uint32 ShaderHash)
-{
-	uint32 EntryHash = GfxEntry->GetEntryHash();
-	FVulkanGfxPipeline* Pipeline = new FVulkanGfxPipeline(Device, EntryHash, ShaderHash);
-	{
-		FScopeLock ScopeLock(&GfxPipelineEntriesCS);
-		GfxPipelineEntries.Add(EntryHash, GfxEntry);
-		GfxEntry->bConsumed = true;
-	}
-
-	return Pipeline;
 }
-#endif
-
 
-
-#if VULKAN_ENABLE_LRU_CACHE
-FVulkanRHIGraphicsPipelineState* FVulkanPipelineStateCacheManager::CreateAndAdd(const FGraphicsPipelineStateInitializer& PSOInitializer, uint32 PSOInitializerHash, FGfxPipelineEntry* GfxEntry, uint32 ShaderHash)
-#else
-FVulkanRHIGraphicsPipelineState* FVulkanPipelineStateCacheManager::CreateAndAdd(const FGraphicsPipelineStateInitializer& PSOInitializer, uint32 PSOInitializerHash, FGfxPipelineEntry* GfxEntry)
-#endif
+FVulkanRHIGraphicsPipelineState* FVulkanPipelineStateCacheManager::CreateAndAdd(const FGraphicsPipelineStateInitializer& PSOInitializer, FGfxPSIKey PSIKey, TGfxPipelineEntrySharedPtr GfxEntry, FGfxEntryKey GfxEntryKey)
 {
 	check(GfxEntry);
 
-#if !VULKAN_PURGE_SHADER_MODULES
-	GfxEntry->CreateShaderModules();
-#endif
-
 #if VULKAN_ENABLE_LRU_CACHE
-	uint32 EntryHash = GfxEntry->GetEntryHash();
-	FVulkanGfxPipeline* Pipeline = new FVulkanGfxPipeline(Device, EntryHash, ShaderHash);
-	if (PipelineLRU.IsActive())
-	{
-		FScopeLock ScopeLock(&GfxPipelineEntriesCS);
-		GfxPipelineEntries.Add(EntryHash, GfxEntry);
-		GfxEntry->bConsumed = true;
-	}
+	FVulkanGfxPipeline* Pipeline = new FVulkanGfxPipeline(Device,
+		(PipelineLRU.IsActive() ? GfxEntry : TGfxPipelineEntrySharedPtr()));
 #else
 	FVulkanGfxPipeline* Pipeline = new FVulkanGfxPipeline(Device);
-#if VULKAN_ENABLE_GENERIC_PIPELINE_CACHE_FILE
-	{
-		FScopeLock ScopeLock(&GfxPipelineEntriesCS);
-		GfxPipelineEntries.Add(GfxEntry->GetEntryHash(), GfxEntry);
-		GfxEntry->bConsumed = true;
-	}
-#endif
 #endif
-	
 
-	// Create the pipeline
-	double BeginTime = FPlatformTime::Seconds();
-	CreateGfxPipelineFromEntry(GfxEntry, Pipeline);
+	Pipeline->CreateRuntimeObjects(PSOInitializer);
 
-	// Recover if we failed to create the pipeline.
-	if (!Pipeline->Pipeline)
+#if VULKAN_ENABLE_LRU_CACHE
+	if (!PSOInitializer.bFromPSOFileCache || !ShouldEvictImmediately())
+#endif
 	{
-		delete Pipeline;
-		return nullptr;
-	}
+		// Create the pipeline
+		double BeginTime = FPlatformTime::Seconds();
+		CreateGfxPipelineFromEntry(GfxEntry.Get(), &PSOInitializer.BoundShaderState, Pipeline);
 
-	Pipeline->CreateRuntimeObjects(PSOInitializer);
-	double EndTime = FPlatformTime::Seconds();
-	double Delta = EndTime - BeginTime;
-	if (Delta > HitchTime)
-	{
-		UE_LOG(LogVulkanRHI, Verbose, TEXT("Hitchy gfx pipeline (%.3f ms)"), (float)(Delta * 1000.0));
+		// Recover if we failed to create the pipeline.
+		if (!Pipeline->GetHandle())
+		{
+			delete Pipeline;
+			return nullptr;
+		}
+
+		double EndTime = FPlatformTime::Seconds();
+		double Delta = EndTime - BeginTime;
+		if (Delta > HitchTime)
+		{
+			UE_LOG(LogVulkanRHI, Verbose, TEXT("Hitchy gfx pipeline (%.3f ms)"), (float)(Delta * 1000.0));
+		}
 	}
 
-	FVulkanRHIGraphicsPipelineState* PipelineState = new FVulkanRHIGraphicsPipelineState(PSOInitializer, Pipeline);
+	FVulkanRHIGraphicsPipelineState* PipelineState = new FVulkanRHIGraphicsPipelineState(PSOInitializer.BoundShaderState, Pipeline, PSOInitializer.PrimitiveType);
 	PipelineState->AddRef();
 
 	{
 		FScopeLock Lock(&InitializerToPipelineMapCS);
-		InitializerToPipelineMap.Add(PSOInitializerHash, PipelineState);
+		InitializerToPipelineMap.Add(MoveTemp(PSIKey), PipelineState);
 	}
 
+	EntryKeyToGfxPipelineMap.Add(MoveTemp(GfxEntryKey), PipelineState->Pipeline);
+#if VULKAN_ENABLE_LRU_CACHE
+	PipelineLRU.Add(PipelineState->Pipeline);
+#endif
+
 	return PipelineState;
 }
 
@@ -1217,9 +1026,9 @@ FArchive& operator << (FArchive& Ar, FVulkanPipelineStateCacheManager::FGfxPipel
 
 	Ar << Entry.DepthStencil;
 
-	for (int32 Index = 0; Index < ARRAY_COUNT(Entry.ShaderMicrocodes); ++Index)
+	for (int32 Index = 0; Index < ARRAY_COUNT(Entry.ShaderHashes.Stages); ++Index)
 	{
-		Ar << Entry.ShaderHashes[Index];
+		Ar << Entry.ShaderHashes.Stages[Index];
 	}
 
 	Ar << Entry.RenderTargets;
@@ -1242,44 +1051,27 @@ FArchive& operator << (FArchive& Ar, FVulkanPipelineStateCacheManager::FPipeline
 }
 #endif
 
-FArchive& operator << (FArchive& Ar, FVulkanPipelineStateCacheManager::FComputePipelineEntry& Entry)
-{
-	// Modify VERSION if serialization changes
-	Ar << Entry.ShaderHash;
-
-	Ar << Entry.DescriptorSetLayoutBindings;
-
-	return Ar;
-}
 
-FArchive& operator << (FArchive& Ar, FVulkanPipelineStateCacheManager::FComputePipelineEntry* Entry)
+FGfxEntryKey FVulkanPipelineStateCacheManager::FGfxPipelineEntry::CreateKey() const
 {
-	return Ar << (*Entry);
+	FGfxEntryKey Result;
+	Result.Generate([this](FArchive& Ar)
+		{
+			Ar << const_cast<FGfxPipelineEntry&>(*this);
+		});
+	return Result;
 }
 
-uint32 FVulkanPipelineStateCacheManager::FGfxPipelineEntry::GetEntryHash(uint32 Crc /*= 0*/)
+void FVulkanPipelineStateCacheManager::CreateGfxPipelineFromEntry(FGfxPipelineEntry* GfxEntry, const FBoundShaderStateInput* BSI, FVulkanGfxPipeline* Pipeline)
 {
-	TArray<uint8> MemFile;
-	FMemoryWriter Ar(MemFile);
-	Ar << this;
-
-	uint32 Result = FCrc::MemCrc32(MemFile.GetData(), MemFile.GetTypeSize() * MemFile.Num(), Crc);
+	FVulkanShader* Shaders[ShaderStage::NumStages];
 
-#if VULKAN_SUPPORTS_COLOR_CONVERSIONS
-	for (uint32 Index = 0; Index < MaxImmutableSamplers; ++Index)
+	if (!GfxEntry->bLoaded)
 	{
-		Result = FCrc::MemCrc32(&ImmutableSamplers[Index], sizeof(SIZE_T), Result);
+		check(BSI);
+		GetVulkanShaders(*BSI, Shaders);
+		GfxEntry->GetOrCreateShaderModules(Shaders);
 	}
-#endif
-
-	return Result;
-}
-
-void FVulkanPipelineStateCacheManager::CreateGfxPipelineFromEntry(FGfxPipelineEntry* GfxEntry, FVulkanGfxPipeline* Pipeline)
-{
-#if VULKAN_PURGE_SHADER_MODULES
-	GfxEntry->CreateShaderModules();
-#endif
 
 	// Pipeline
 	VkGraphicsPipelineCreateInfo PipelineInfo;
@@ -1319,7 +1111,7 @@ void FVulkanPipelineStateCacheManager::CreateGfxPipelineFromEntry(FGfxPipelineEn
 	PipelineInfo.pStages = ShaderStages;
 	for (int32 ShaderStage = 0; ShaderStage < ShaderStage::NumStages; ++ShaderStage)
 	{
-		if (GfxEntry->ShaderMicrocodes[ShaderStage] == nullptr)
+		if (!GfxEntry->ShaderModules[ShaderStage])
 		{
 			continue;
 		}
@@ -1392,42 +1184,56 @@ void FVulkanPipelineStateCacheManager::CreateGfxPipelineFromEntry(FGfxPipelineEn
 	PipelineInfo.pDynamicState = &DynamicState;
 
 	//#todo-rco: Fix me
+	VkResult Result = VK_ERROR_INITIALIZATION_FAILED;
 	double BeginTime = FPlatformTime::Seconds();
 #if VULKAN_ENABLE_LRU_CACHE
-	FPipelineSize** Found = PipelineSizeList.Find(Pipeline->ShaderHash);
-	size_t PreSize = 0, AfterSize = 0;
-	if (Found)
-	{
-		Pipeline->PipelineCacheSize = (*Found)->PipelineSize;
-	}
-	else
+	if (PipelineLRU.IsActive())
 	{
-		VulkanRHI::vkGetPipelineCacheData(Device->GetInstanceHandle(), PipelineCache, &PreSize, nullptr);
-	}
+		const uint32 ShaderHash = GfxEntry->ShaderHashes.Hash;
+		FPipelineSize** Found = PipelineSizeList.Find(ShaderHash);
+		size_t PreSize = 0, AfterSize = 0;
+		if (Found)
+		{
+			Pipeline->PipelineCacheSize = (*Found)->PipelineSize;
+		}
+		else
+		{
+			VulkanRHI::vkGetPipelineCacheData(Device->GetInstanceHandle(), PipelineCache, &PreSize, nullptr);
+		}
 
-	VkResult Result = VulkanRHI::vkCreateGraphicsPipelines(Device->GetInstanceHandle(), PipelineCache, 1, &PipelineInfo, VULKAN_CPU_ALLOCATOR, &Pipeline->Pipeline);
+		Result = VulkanRHI::vkCreateGraphicsPipelines(Device->GetInstanceHandle(), PipelineCache, 1, &PipelineInfo, VULKAN_CPU_ALLOCATOR, &Pipeline->Pipeline);
 
-	if (!Found && Result == VK_SUCCESS)
-	{
-		VulkanRHI::vkGetPipelineCacheData(Device->GetInstanceHandle(), PipelineCache, &AfterSize, nullptr);
-		uint32 Diff = AfterSize - PreSize;
-		if (!Diff && PipelineLRU.IsActive())
+		if (!Found && Result == VK_SUCCESS)
 		{
-			UE_LOG(LogVulkanRHI, Warning, TEXT("Shader size was computed as zero, using 20k instead."));
-			Diff = 20 * 1024;
+			VulkanRHI::vkGetPipelineCacheData(Device->GetInstanceHandle(), PipelineCache, &AfterSize, nullptr);
+			uint32 Diff = AfterSize - PreSize;
+			if (!Diff)
+			{
+				UE_LOG(LogVulkanRHI, Warning, TEXT("Shader size was computed as zero, using 20k instead."));
+				Diff = 20 * 1024;
+			}
+			FPipelineSize* PipelineSize = new FPipelineSize();
+			PipelineSize->ShaderHash = ShaderHash;
+			PipelineSize->PipelineSize = Diff;
+			PipelineSizeList.Add(ShaderHash, PipelineSize);
+			Pipeline->PipelineCacheSize = Diff;
 		}
-		FPipelineSize* PipelineSize = new FPipelineSize();
-		PipelineSize->ShaderHash = Pipeline->ShaderHash;
-		PipelineSize->PipelineSize = Diff;
-		PipelineSizeList.Add(Pipeline->ShaderHash, PipelineSize);
-		Pipeline->PipelineCacheSize = Diff;
 	}
-#else
-	VkResult Result = VulkanRHI::vkCreateGraphicsPipelines(Device->GetInstanceHandle(), PipelineCache, 1, &PipelineInfo, VULKAN_CPU_ALLOCATOR, &Pipeline->Pipeline);
+	else
 #endif
+	{
+		Result = VulkanRHI::vkCreateGraphicsPipelines(Device->GetInstanceHandle(), PipelineCache, 1, &PipelineInfo, VULKAN_CPU_ALLOCATOR, &Pipeline->Pipeline);
+	}
 
 #if VULKAN_PURGE_SHADER_MODULES
-	GfxEntry->PurgeShaderModules();
+	if (GfxEntry->bLoaded)
+	{
+		GfxEntry->PurgeLoadedShaderModules(Device);
+	}
+	else
+	{
+		GfxEntry->PurgeShaderModules(Shaders);
+	}
 #endif
 
 	if (Result != VK_SUCCESS)
@@ -1449,67 +1255,13 @@ void FVulkanPipelineStateCacheManager::CreateGfxPipelineFromEntry(FGfxPipelineEn
 	Pipeline->Layout = GfxEntry->Layout;
 }
 
-void FVulkanPipelineStateCacheManager::CreatGfxEntryRuntimeObjects(FGfxPipelineEntry* GfxEntry)
-{
-	{
-		// Descriptor Set Layouts
-		check(!GfxEntry->Layout);
-
-		FVulkanDescriptorSetsLayoutInfo Info;
-		for (int32 SetIndex = 0; SetIndex < GfxEntry->DescriptorSetLayoutBindings.Num(); ++SetIndex)
-		{
-			for (int32 Index = 0; Index < GfxEntry->DescriptorSetLayoutBindings[SetIndex].Num(); ++Index)
-			{
-				VkDescriptorSetLayoutBinding Binding;
-				Binding.descriptorCount = 1;
-				Binding.pImmutableSamplers = nullptr;
-				GfxEntry->DescriptorSetLayoutBindings[SetIndex][Index].WriteInto(Binding);
-				Info.AddDescriptor(SetIndex, Binding);
-			}
-		}
-
-		ensure(0);
-/*
-		GfxEntry->Layout = FindOrAddLayout(Info, true);
-*/
-	}
-
-	{
-		// Shaders
-		for (int32 Index = 0; Index < ARRAY_COUNT(GfxEntry->ShaderMicrocodes); ++Index)
-		{
-			if (GfxEntry->ShaderMicrocodes[Index] != nullptr)
-			{
-				VkShaderModuleCreateInfo ModuleCreateInfo;
-				ZeroVulkanStruct(ModuleCreateInfo, VK_STRUCTURE_TYPE_SHADER_MODULE_CREATE_INFO);
-				ModuleCreateInfo.codeSize = GfxEntry->ShaderMicrocodes[Index]->Num();
-				ModuleCreateInfo.pCode = (uint32*)GfxEntry->ShaderMicrocodes[Index]->GetData();
-				VERIFYVULKANRESULT(VulkanRHI::vkCreateShaderModule(Device->GetInstanceHandle(), &ModuleCreateInfo, VULKAN_CPU_ALLOCATOR, &GfxEntry->ShaderModules[Index]));
-			}
-		}
-	}
-
-	{
-		ensure(0);
-/*
-		// Render Pass
-		FVulkanRenderTargetLayout RTLayout;
-		GfxEntry->RenderTargets.WriteInto(RTLayout);
-		GfxEntry->RenderPass = Device->GetImmediateContext().PrepareRenderPassForPSOCreation(RTLayout, GfxEntry->Layout->GetDescriptorSetsLayout().RemappingInfo.InputAttachmentData);
-*/
-	}
-
-	GfxEntry->bLoaded = true;
-}
-
-
 void FVulkanPipelineStateCacheManager::DestroyCache()
 {
 	VkDevice DeviceHandle = Device->GetInstanceHandle();
 
 	// Graphics
 	{
-		for (auto Pair : InitializerToPipelineMap)
+		for (auto& Pair : InitializerToPipelineMap)
 		{
 			FVulkanRHIGraphicsPipelineState* Pipeline = Pair.Value;
 			//when DestroyCache is called as part of r.Vulkan.RebuildPipelineCache, a pipeline can still be referenced by FVulkanPendingGfxState
@@ -1520,21 +1272,11 @@ void FVulkanPipelineStateCacheManager::DestroyCache()
 
 		for (auto& Pair : GfxPipelineEntries)
 		{
-			FGfxPipelineEntry* Entry = Pair.Value;
-			Entry->RenderPass = nullptr;
+			FGfxPipelineEntry* Entry = Pair.Value.Get();
 			if (Entry->bLoaded)
 			{
-				for (int32 Index = 0; Index < ARRAY_COUNT(Entry->ShaderModules); ++Index)
-				{
-					if (Entry->ShaderModules[Index] != VK_NULL_HANDLE)
-					{
-						VulkanRHI::vkDestroyShaderModule(DeviceHandle, Entry->ShaderModules[Index], VULKAN_CPU_ALLOCATOR);
-					}
-				}
-				Entry->bLoaded = false;
-				Entry->bConsumed = false;
+				Entry->PurgeLoadedShaderModules(Device);
 			}
-			delete Entry;
 		}
 		GfxPipelineEntries.Reset();
 
@@ -1545,43 +1287,17 @@ void FVulkanPipelineStateCacheManager::DestroyCache()
 			delete Entry;
 		}
 		PipelineSizeList.Reset();
-		PipelineLRU.Empty();
 #endif
 
 		// This map can simply be cleared as InitializerToPipelineMap already decreased the refcount of the pipeline objects	
 		{
-			FScopeLock Lock(&GShaderHashToGfxEntriesMapCS);
-			ShaderHashToGfxPipelineMap.Reset();
+			FScopeLock Lock(&EntryKeyToGfxPipelineMapCS);
+			EntryKeyToGfxPipelineMap.Reset();
 		}
 	}
 
-	// Compute
-	{
-		for (auto Pair : ComputeEntryHashToPipelineMap)
-		{
-			FVulkanComputePipeline* Pipeline = Pair.Value;
-			//when DestroyCache is called as part of r.Vulkan.RebuildPipelineCache, a pipeline can still be referenced by FVulkanPendingComputeState
-			ensure(GIsRHIInitialized || (!GIsRHIInitialized && Pipeline->GetRefCount() == 1));
-			Pipeline->Release();
-		}
-		ComputeEntryHashToPipelineMap.Reset();
-		ComputeShaderToPipelineMap.Reset();
-
-		for (auto& Pair : ComputePipelineEntries)
-		{
-			FComputePipelineEntry* Entry = Pair.Value;
-			if (Entry->bLoaded)
-			{
-				if (Entry->ShaderModule != VK_NULL_HANDLE)
-				{
-					VulkanRHI::vkDestroyShaderModule(DeviceHandle, Entry->ShaderModule, VULKAN_CPU_ALLOCATOR);
-				}
-				Entry->bLoaded = false;
-			}
-			delete Entry;
-		}
-		ComputePipelineEntries.Reset();
-	}
+	// Compute pipelines already deleted...
+	ComputePipelineEntries.Reset();
 }
 
 void FVulkanPipelineStateCacheManager::RebuildCache()
@@ -1613,11 +1329,10 @@ FVulkanPipelineStateCacheManager::FShaderHashes::FShaderHashes()
 	Hash = 0;
 }
 
-#if VULKAN_ENABLE_LRU_CACHE
-inline FVulkanLayout* FVulkanPipelineStateCacheManager::FindOrAddLayoutForGfx(uint32 PSOHash, const FVulkanDescriptorSetsLayoutInfo& DescriptorSetLayoutInfo, bool bGfxLayout)
+inline FVulkanLayout* FVulkanPipelineStateCacheManager::FindOrAddLayout(const FVulkanDescriptorSetsLayoutInfo& DescriptorSetLayoutInfo, bool bGfxLayout)
 {
 	FScopeLock Lock(&LayoutMapCS);
-	if (FVulkanLayout** FoundLayout = LayoutMapGfx.Find(PSOHash))
+	if (FVulkanLayout** FoundLayout = LayoutMap.Find(DescriptorSetLayoutInfo))
 	{
 		check(bGfxLayout == (*FoundLayout)->IsGfxLayout());
 		return *FoundLayout;
@@ -1633,55 +1348,62 @@ inline FVulkanLayout* FVulkanPipelineStateCacheManager::FindOrAddLayoutForGfx(ui
 	{
 		Layout = new FVulkanComputeLayout(Device);
 	}
-
+	
 	Layout->DescriptorSetLayout.CopyFrom(DescriptorSetLayoutInfo);
 	Layout->Compile();
 
-	LayoutMapGfx.Add(PSOHash, Layout);
+	LayoutMap.Add(Layout->DescriptorSetLayout, Layout);
 	return Layout;
 }
-#endif
 
-inline FVulkanLayout* FVulkanPipelineStateCacheManager::FindOrAddLayout(const FVulkanDescriptorSetsLayoutInfo& DescriptorSetLayoutInfo, bool bGfxLayout)
+void FVulkanPipelineStateCacheManager::GetVulkanShaders(const FBoundShaderStateInput& BSI, FVulkanShader* OutShaders[ShaderStage::NumStages])
 {
-	FScopeLock Lock(&LayoutMapCS);
-	if (FVulkanLayout** FoundLayout = LayoutMap.Find(DescriptorSetLayoutInfo))
-	{
-		check(bGfxLayout == (*FoundLayout)->IsGfxLayout());
-		return *FoundLayout;
+	FMemory::Memzero(OutShaders, ShaderStage::NumStages * sizeof(*OutShaders));
+
+	OutShaders[ShaderStage::Vertex] = ResourceCast(BSI.VertexShaderRHI);
+
+	if (BSI.PixelShaderRHI)
+	{
+		OutShaders[ShaderStage::Pixel] = ResourceCast(BSI.PixelShaderRHI);
+	}
+	else if (GMaxRHIFeatureLevel <= ERHIFeatureLevel::ES3_1)
+	{
+		// Some mobile devices expect PS stage (S7 Adreno)
+		OutShaders[ShaderStage::Pixel] = ResourceCast(TShaderMapRef<FNULLPS>(GetGlobalShaderMap(GMaxRHIFeatureLevel))->GetPixelShader());
 	}
 
-	FVulkanLayout* Layout = nullptr;
+	if (BSI.GeometryShaderRHI)
+	{
+#if VULKAN_SUPPORTS_GEOMETRY_SHADERS
+		OutShaders[ShaderStage::Geometry] = ResourceCast(BSI.GeometryShaderRHI);
+#else
+		ensureMsgf(0, TEXT("Geometry not supported!"));
+#endif
+	}
 
-	if (bGfxLayout)
+	if (BSI.HullShaderRHI)
 	{
-		Layout = new FVulkanGfxLayout(Device);
+		ensureMsgf(0, TEXT("Tessellation not supported yet!"));
+		/*
+		// Can't have Hull w/o Domain
+		check(BSI.DomainShaderRHI);
+		OutShaders[ShaderStage::Hull] = ResourceCast(BSI.HullShaderRHI);
+		OutShaders[ShaderStage::Domain] = ResourceCast(BSI.DomainShaderRHI);
+		*/
 	}
 	else
 	{
-		Layout = new FVulkanComputeLayout(Device);
+		// Can't have Domain w/o Hull
+		check(BSI.DomainShaderRHI == nullptr);
 	}
-	
-	Layout->DescriptorSetLayout.CopyFrom(DescriptorSetLayoutInfo);
-	Layout->Compile();
-
-	LayoutMap.Add(Layout->DescriptorSetLayout, Layout);
-	return Layout;
 }
 
-#if VULKAN_ENABLE_LRU_CACHE
 FVulkanGfxLayout* FVulkanPipelineStateCacheManager::GetOrGenerateGfxLayout(const FGraphicsPipelineStateInitializer& PSOInitializer,
-	FVulkanShader** OutShaders, FVulkanVertexInputStateInfo& OutVertexInputState, uint32 PSOHash)
-#else
-FVulkanGfxLayout* FVulkanPipelineStateCacheManager::GetOrGenerateGfxLayout(const FGraphicsPipelineStateInitializer& PSOInitializer,
-	FVulkanShader** OutShaders, FVulkanVertexInputStateInfo& OutVertexInputState)
-#endif
+	FVulkanShader*const* Shaders, FVulkanVertexInputStateInfo& OutVertexInputState)
 {
 	const FBoundShaderStateInput& BSI = PSOInitializer.BoundShaderState;
 
-	FVulkanVertexShader* VS = ResourceCast(BSI.VertexShaderRHI);
-	const FVulkanShaderHeader& VSHeader = VS->GetCodeHeader();
-	OutShaders[ShaderStage::Vertex] = VS;
+	const FVulkanShaderHeader& VSHeader = Shaders[ShaderStage::Vertex]->GetCodeHeader();
 	OutVertexInputState.Generate(ResourceCast(PSOInitializer.BoundShaderState.VertexDeclarationRHI), VSHeader.InOutMask);
 
 	FUniformBufferGatherInfo UBGatherInfo;
@@ -1690,98 +1412,55 @@ FVulkanGfxLayout* FVulkanPipelineStateCacheManager::GetOrGenerateGfxLayout(const
 	FVulkanDescriptorSetsLayoutInfo DescriptorSetLayoutInfo;
 	DescriptorSetLayoutInfo.ProcessBindingsForStage(VK_SHADER_STAGE_VERTEX_BIT, ShaderStage::Vertex, VSHeader, UBGatherInfo);
 
-	FVulkanPixelShader* PS = nullptr;
-	if (BSI.PixelShaderRHI)
-	{
-		PS = ResourceCast(BSI.PixelShaderRHI);
-	}
-	else if (GMaxRHIFeatureLevel <= ERHIFeatureLevel::ES3_1)
-	{
-		// Some mobile devices expect PS stage (S7 Adreno)
-		PS = ResourceCast(TShaderMapRef<FNULLPS>(GetGlobalShaderMap(GMaxRHIFeatureLevel))->GetPixelShader());
-	}
-
-	if (PS)
+	if (Shaders[ShaderStage::Pixel])
 	{
-		OutShaders[ShaderStage::Pixel] = PS;
-		const FVulkanShaderHeader& PSHeader = PS->GetCodeHeader();
+		const FVulkanShaderHeader& PSHeader = Shaders[ShaderStage::Pixel]->GetCodeHeader();
 		DescriptorSetLayoutInfo.ProcessBindingsForStage(VK_SHADER_STAGE_FRAGMENT_BIT, ShaderStage::Pixel, PSHeader, UBGatherInfo);
 	}
 
-	if (BSI.GeometryShaderRHI)
-	{
 #if VULKAN_SUPPORTS_GEOMETRY_SHADERS
-		FVulkanGeometryShader* GS = ResourceCast(BSI.GeometryShaderRHI);
-		OutShaders[ShaderStage::Geometry] = GS;
-		const FVulkanShaderHeader& GSHeader = GS->GetCodeHeader();
+	if (Shaders[ShaderStage::Geometry])
+	{
+		const FVulkanShaderHeader& GSHeader = Shaders[ShaderStage::Geometry]->GetCodeHeader();
 		DescriptorSetLayoutInfo.ProcessBindingsForStage(VK_SHADER_STAGE_GEOMETRY_BIT, ShaderStage::Geometry, GSHeader, UBGatherInfo);
-#else
-		ensureMsgf(0, TEXT("Geometry not supported!"));
-#endif
 	}
+#endif
 
-	if (BSI.HullShaderRHI)
+	/* We don't support tessellation on desktop currently
+	if (Shaders[ShaderStage::Hull])
 	{
-		ensureMsgf(0, TEXT("Tessellation not supported yet!"));
-		/*
-		// Can't have Hull w/o Domain
-		check(BSI.DomainShaderRHI);
-		FVulkanHullShader* HS = ResourceCast(BSI.HullShaderRHI);
-		FVulkanDomainShader* DS = ResourceCast(BSI.DomainShaderRHI);
-		Shaders[ShaderStage::Hull] = HS;
-		Shaders[ShaderStage::Domain] = DS;
-		const FVulkanCodeHeader& HSHeader = HS->GetCodeHeader();
-		const FVulkanCodeHeader& DSHeader = DS->GetCodeHeader();
+		const FVulkanCodeHeader& HSHeader = Shaders[ShaderStage::Hull]->GetCodeHeader();
+		const FVulkanCodeHeader& DSHeader = Shaders[ShaderStage::Domain]->GetCodeHeader();
 		DescriptorSetLayoutInfo.AddBindingsForStage(VK_SHADER_STAGE_TESSELLATION_CONTROL_BIT, ShaderStage::Hull, HSHeader);
 		DescriptorSetLayoutInfo.AddBindingsForStage(VK_SHADER_STAGE_TESSELLATION_EVALUATION_BIT, ShaderStage::Domain, DSHeader);
-		*/
-	}
-	else
-	{
-		// Can't have Domain w/o Hull
-		check(BSI.DomainShaderRHI == nullptr);
-	}
+	}*/
 
 	// Second pass
 	const int32 NumImmutableSamplers = PSOInitializer.ImmutableSamplerState.ImmutableSamplers.Num();
 	TArrayView<const FSamplerStateRHIParamRef> ImmutableSamplers(NumImmutableSamplers > 0 ? &PSOInitializer.ImmutableSamplerState.ImmutableSamplers[0] : nullptr, NumImmutableSamplers);
 	DescriptorSetLayoutInfo.FinalizeBindings<false>(UBGatherInfo, ImmutableSamplers);
 
-#if VULKAN_ENABLE_LRU_CACHE
-	FVulkanLayout* Layout = FindOrAddLayoutForGfx(PSOHash, DescriptorSetLayoutInfo, true);
+	FVulkanLayout* Layout = FindOrAddLayout(DescriptorSetLayoutInfo, true);
 	FVulkanGfxLayout* GfxLayout = (FVulkanGfxLayout*)Layout;
 	if (!GfxLayout->GfxPipelineDescriptorInfo.IsInitialized())
 	{
-		GfxLayout->GfxPipelineDescriptorInfo.Initialize(Layout->GetDescriptorSetsLayout().RemappingInfo, OutShaders);
+		GfxLayout->GfxPipelineDescriptorInfo.Initialize(Layout->GetDescriptorSetsLayout().RemappingInfo);
 	}
-#else
-	FVulkanLayout* Layout = FindOrAddLayout(DescriptorSetLayoutInfo, true);
-	FVulkanGfxLayout* GfxLayout = (FVulkanGfxLayout*)Layout;
-	check(!GfxLayout->GfxPipelineDescriptorInfo.IsInitialized());
-	GfxLayout->GfxPipelineDescriptorInfo.Initialize(Layout->GetDescriptorSetsLayout().RemappingInfo, OutShaders);
-#endif
+
 	return GfxLayout;
 }
 
-#if VULKAN_ENABLE_LRU_CACHE
-FVulkanPipelineStateCacheManager::FGfxPipelineEntry* FVulkanPipelineStateCacheManager::CreateGfxEntry(const FGraphicsPipelineStateInitializer& PSOInitializer, uint32 PSOHash)
-#else
 FVulkanPipelineStateCacheManager::FGfxPipelineEntry* FVulkanPipelineStateCacheManager::CreateGfxEntry(const FGraphicsPipelineStateInitializer& PSOInitializer)
-#endif
 {
 	FGfxPipelineEntry* OutGfxEntry = new FGfxPipelineEntry();
 
 	FVulkanShader* Shaders[ShaderStage::NumStages];
-	FMemory::Memzero(Shaders);
+	GetVulkanShaders(PSOInitializer.BoundShaderState, Shaders);
 
 	OutGfxEntry->RenderPass = Device->GetImmediateContext().PrepareRenderPassForPSOCreation(PSOInitializer, OutGfxEntry->Layout->GetDescriptorSetsLayout().RemappingInfo.InputAttachmentData);
 
 	FVulkanVertexInputStateInfo VertexInputState;
-#if VULKAN_ENABLE_LRU_CACHE
-	OutGfxEntry->Layout = GetOrGenerateGfxLayout(PSOInitializer, Shaders, VertexInputState, PSOHash);
-#else
 	OutGfxEntry->Layout = GetOrGenerateGfxLayout(PSOInitializer, Shaders, VertexInputState);
-#endif
 
 	OutGfxEntry->RasterizationSamples = OutGfxEntry->RenderPass->GetLayout().GetAttachmentDescriptions()[0].samples;
 	ensure(OutGfxEntry->RasterizationSamples == PSOInitializer.NumSamples);
@@ -1835,15 +1514,8 @@ FVulkanPipelineStateCacheManager::FGfxPipelineEntry* FVulkanPipelineStateCacheMa
 			check(Shader->Spirv.Num() != 0);
 
 			FSHAHash Hash = GetShaderHashForStage(PSOInitializer, (ShaderStage::EStage)Index);
-			OutGfxEntry->ShaderHashes[Index] = Hash;
-
-			OutGfxEntry->VulkanShaders[Index] = Shader;
+			OutGfxEntry->ShaderHashes.Stages[Index] = Hash;
 
-			OutGfxEntry->ShaderMicrocodes[Index] = ShaderCache.Get(Hash);
-			if (OutGfxEntry->ShaderMicrocodes[Index] == nullptr)
-			{
-				OutGfxEntry->ShaderMicrocodes[Index] = ShaderCache.Add(Hash, Shader);
-			}
 			++NumShaders;
 		}
 	}
@@ -1858,91 +1530,122 @@ FVulkanPipelineStateCacheManager::FGfxPipelineEntry* FVulkanPipelineStateCacheMa
 
 	OutGfxEntry->RenderTargets.ReadFrom(OutGfxEntry->RenderPass->GetLayout());
 
+	OutGfxEntry->ShaderHashes.Finalize();
+	
 	return OutGfxEntry;
 }
 
 
-FVulkanRHIGraphicsPipelineState* FVulkanPipelineStateCacheManager::FindInLoadedLibrary(const FGraphicsPipelineStateInitializer& PSOInitializer, uint32 PSOInitializerHash, const FShaderHashes& ShaderHashes, FGfxPipelineEntry*& OutGfxEntry, FVulkanPipelineStateCacheManager::FHashToGfxPipelinesMap*& OutHashToGfxPipelinesMap)
+FVulkanRHIGraphicsPipelineState* FVulkanPipelineStateCacheManager::FindInLoadedLibrary(const FGraphicsPipelineStateInitializer& PSOInitializer, FGfxPSIKey& PSIKey, TGfxPipelineEntrySharedPtr& OutGfxEntry, FGfxEntryKey& OutGfxEntryKey)
 {
 	OutGfxEntry = nullptr;
-	OutHashToGfxPipelinesMap = nullptr;
 
-	FHashToGfxPipelinesMap* FoundHashToGfxPipelinesMap = ShaderHashToGfxPipelineMap.Find(ShaderHashes);
-	if (!FoundHashToGfxPipelinesMap)
-	{
-		FoundHashToGfxPipelinesMap = &ShaderHashToGfxPipelineMap.Add(ShaderHashes);
-	}
-	
-#if VULKAN_ENABLE_LRU_CACHE
-	FGfxPipelineEntry* GfxEntry = CreateGfxEntry(PSOInitializer, PSOInitializerHash);
-#else
 	FGfxPipelineEntry* GfxEntry = CreateGfxEntry(PSOInitializer);
-#endif
-	uint32 EntryHash = GfxEntry->GetEntryHash();
+	FGfxEntryKey GfxEntryKey = GfxEntry->CreateKey();
 
-	FVulkanGfxPipeline** FoundPipeline = FoundHashToGfxPipelinesMap->Find(EntryHash);
+	FVulkanGfxPipeline** FoundPipeline = EntryKeyToGfxPipelineMap.Find(GfxEntryKey);
 	if (FoundPipeline)
 	{
 		if (!(*FoundPipeline)->IsRuntimeInitialized())
 		{
 			(*FoundPipeline)->CreateRuntimeObjects(PSOInitializer);
 		}
-		FVulkanRHIGraphicsPipelineState* PipelineState = new FVulkanRHIGraphicsPipelineState(PSOInitializer, *FoundPipeline);
+		FVulkanRHIGraphicsPipelineState* PipelineState = new FVulkanRHIGraphicsPipelineState(PSOInitializer.BoundShaderState, *FoundPipeline, PSOInitializer.PrimitiveType);
 		{
 			FScopeLock Lock2(&InitializerToPipelineMapCS);
-			InitializerToPipelineMap.Add(PSOInitializerHash, PipelineState);
+			InitializerToPipelineMap.Add(MoveTemp(PSIKey), PipelineState);
 		}
 		PipelineState->AddRef();
 		delete GfxEntry;
+#if VULKAN_ENABLE_LRU_CACHE
+		PipelineLRU.Touch(PipelineState);
+#endif
 		return PipelineState;
 	}
 
-	OutGfxEntry = GfxEntry;
-	OutHashToGfxPipelinesMap = FoundHashToGfxPipelinesMap;
+	OutGfxEntry = TGfxPipelineEntrySharedPtr(GfxEntry);
+	OutGfxEntryKey = MoveTemp(GfxEntryKey);
 	return nullptr;
 }
 
-inline FVulkanRHIGraphicsPipelineState* FVulkanPipelineStateCacheManager::FindInRuntimeCache(const FGraphicsPipelineStateInitializer& Initializer, uint32& OutHash)
+FVulkanRHIGraphicsPipelineState* FVulkanPipelineStateCacheManager::FindInRuntimeCache(const FGraphicsPipelineStateInitializer& Initializer, FGfxPSIKey& OutKey)
 {
-	FPSOHashable Hashable;
-	FMemory::Memzero(Hashable);
-	Hashable.ShaderIds[ShaderStage::Vertex] = GetShaderId(Initializer.BoundShaderState.VertexShaderRHI);
-	Hashable.ShaderIds[ShaderStage::Pixel] = GetShaderId(Initializer.BoundShaderState.PixelShaderRHI);
-#if VULKAN_SUPPORTS_GEOMETRY_SHADERS
-	Hashable.ShaderIds[ShaderStage::Geometry] = GetShaderId(Initializer.BoundShaderState.GeometryShaderRHI);
-	//Hashable.ShaderIds[ShaderStage::Hull] = GetShaderId(Initializer.BoundShaderState.HullShaderRHI);
-	//Hashable.ShaderIds[ShaderStage::Domain] = GetShaderId(Initializer.BoundShaderState.DomainShaderRHI);
-#endif
-	Initializer.BoundShaderState.VertexDeclarationRHI->GetInitializer(Hashable.VertexDeclaration);
-	Initializer.BlendState->GetInitializer(Hashable.BlendState);
-	Initializer.RasterizerState->GetInitializer(Hashable.RasterizerState);
-	Initializer.DepthStencilState->GetInitializer(Hashable.DepthStencilState);
-
-	Hashable.DepthStencilAccess = Initializer.DepthStencilAccess;
-	Hashable.DepthStencilTargetFormat = Initializer.DepthStencilTargetFormat;
-	Hashable.bDepthBounds = Initializer.bDepthBounds ? 1 : 0;
-	Hashable.PrimitiveType = Initializer.PrimitiveType;
-	Hashable.RenderTargetsEnabled = (uint8)Initializer.RenderTargetsEnabled;
-	Hashable.NumSamples = (uint8)Initializer.NumSamples;
+	OutKey.Generate([&Initializer](FArchive& Ar)
+		{
+			FGraphicsPipelineStateInitializer& PSI = const_cast<FGraphicsPipelineStateInitializer&>(Initializer);
+
+			uint64 TempUInt64 = 0;
+			int32 TempEnumValue = 0;
+
+			TempUInt64 = GetShaderKey(PSI.BoundShaderState.VertexShaderRHI);
+			Ar << TempUInt64;
+
+			TempUInt64 = GetShaderKey(PSI.BoundShaderState.PixelShaderRHI);
+			Ar << TempUInt64;
+
+	#if VULKAN_SUPPORTS_GEOMETRY_SHADERS
+			TempUInt64 = GetShaderKey(PSI.BoundShaderState.GeometryShaderRHI);
+			Ar << TempUInt64;
+			//TempUInt64 = GetShaderKey(PSI.BoundShaderState.HullShaderRHI);
+			Ar << TempUInt64;
+			//TempUInt64 = GetShaderKey(PSI.BoundShaderState.DomainShaderRHI);
+			Ar << TempUInt64;
+	#endif
+			Ar << ResourceCast(PSI.BoundShaderState.VertexDeclarationRHI)->Elements;
+			Ar << ResourceCast(PSI.RasterizerState)->Initializer;
+			Ar << ResourceCast(PSI.DepthStencilState)->Initializer;
+			{
+				FBlendStateInitializerRHI &BlendState = ResourceCast(PSI.BlendState)->Initializer;
+				for (uint32 Index = 0; Index < PSI.RenderTargetsEnabled; ++Index)
+				{
+					Ar << BlendState.RenderTargets[Index];
+					TempEnumValue = PSI.RenderTargetFormats[Index];
+					Ar << TempEnumValue;
+				}
+				Ar << BlendState.bUseIndependentRenderTargetBlendStates;
+			}
+
+			// TODO To avoid changing of global header. Add getter when submitting
+			Ar.Serialize(&PSI.DepthStencilAccess, sizeof(PSI.DepthStencilAccess));
+
+			TempEnumValue = PSI.DepthStencilTargetFormat;
+			Ar << TempEnumValue;
+
+			TempEnumValue = PSI.PrimitiveType;
+			Ar << TempEnumValue;
 
 #if VULKAN_SUPPORTS_COLOR_CONVERSIONS
-	for (int32 Index = 0; Index < Initializer.ImmutableSamplerState.ImmutableSamplers.Num(); ++Index)
-	{
-		FRHISamplerState* SamplerState = Initializer.ImmutableSamplerState.ImmutableSamplers[Index];
-		Hashable.ImmutableSamplers[Index] = SamplerState ? ResourceCast(SamplerState)->Sampler : VK_NULL_HANDLE;
-	}
+			for (uint32 Index = 0; Index < MaxImmutableSamplers; ++Index)
+			{
+				FRHISamplerState* SamplerState = PSI.ImmutableSamplerState.ImmutableSamplers[Index];
+				VkSampler Handle = SamplerState ? ResourceCast(SamplerState)->Sampler : VK_NULL_HANDLE;
+				TempUInt64 = (uint64)Handle;
+				Ar << TempUInt64;
+			}
 #endif
 
-	OutHash = FCrc::MemCrc32(&Hashable, sizeof(Hashable));
-	OutHash = FCrc::MemCrc32(&Initializer.RenderTargetFormats[0], sizeof(Initializer.RenderTargetFormats[0]) * Initializer.RenderTargetsEnabled, OutHash);
+			Ar << PSI.bDepthBounds;
+			Ar << PSI.RenderTargetsEnabled;
+			Ar << PSI.NumSamples;
+
+		},
+		256);
+
+	FVulkanRHIGraphicsPipelineState** Found = nullptr;
 
 	{
 		FScopeLock Lock(&InitializerToPipelineMapCS);
-		FVulkanRHIGraphicsPipelineState** Found = InitializerToPipelineMap.Find(OutHash);
-		if (Found)
-		{
-			return *Found;
-		}
+		Found = InitializerToPipelineMap.Find(OutKey);
+	}
+
+	if (Found)
+	{
+		// TODO: looks like with enabled shader caching (r.Vulkan.CacheShaders), PSO can end up with stale shader pointers
+		(*Found)->BSI = Initializer.BoundShaderState;
+#if VULKAN_ENABLE_LRU_CACHE
+		PipelineLRU.Touch(*Found);
+#endif
+		return *Found;
 	}
 
 	return nullptr;
@@ -1957,94 +1660,29 @@ FGraphicsPipelineStateRHIRef FVulkanDynamicRHI::RHICreateGraphicsPipelineState(c
 	LLM_SCOPE_VULKAN(ELLMTagVulkan::VulkanShaders);
 
 	// First try the hash based off runtime objects
-	uint32 PSOInitializerHash = 0;
-	FVulkanRHIGraphicsPipelineState* Found = Device->PipelineStateCache->FindInRuntimeCache(PSOInitializer, PSOInitializerHash);
+	FGfxPSIKey PSIKey;
+	FVulkanRHIGraphicsPipelineState* Found = Device->PipelineStateCache->FindInRuntimeCache(PSOInitializer, PSIKey);
 	if (Found)
 	{
-#if UE_BUILD_DEBUG || UE_BUILD_DEVELOPMENT
-		const FGraphicsPipelineStateInitializer& FoundInitializer = Found->PipelineStateInitializer;
-		ensure(ResourceCast(FoundInitializer.BoundShaderState.VertexDeclarationRHI) == ResourceCast(PSOInitializer.BoundShaderState.VertexDeclarationRHI) &&
-			GetShaderId(FoundInitializer.BoundShaderState.VertexShaderRHI) == GetShaderId(PSOInitializer.BoundShaderState.VertexShaderRHI) &&
-			GetShaderId(FoundInitializer.BoundShaderState.PixelShaderRHI) == GetShaderId(PSOInitializer.BoundShaderState.PixelShaderRHI));
-#if VULKAN_SUPPORTS_GEOMETRY_SHADERS
-		ensure(GetShaderId(FoundInitializer.BoundShaderState.GeometryShaderRHI) == GetShaderId(PSOInitializer.BoundShaderState.GeometryShaderRHI));
-/*
-		ensure(GetShaderId(FoundInitializer.BoundShaderState.HullShaderRHI) == GetShaderId(PSOInitializer.BoundShaderState.HullShaderRHI) &&
-		GetShaderId(FoundInitializer.BoundShaderState.DomainShaderRHI) == GetShaderId(PSOInitializer.BoundShaderState.DomainShaderRHI));
-*/
-#endif
-		ensure(FoundInitializer.BlendState == PSOInitializer.BlendState &&
-			FoundInitializer.RasterizerState == PSOInitializer.RasterizerState &&
-			FoundInitializer.DepthStencilState == PSOInitializer.DepthStencilState &&
-			FoundInitializer.bDepthBounds == PSOInitializer.bDepthBounds &&
-			FoundInitializer.PrimitiveType == PSOInitializer.PrimitiveType);
-		ensure(FoundInitializer.RenderTargetsEnabled == PSOInitializer.RenderTargetsEnabled &&
-			FMemory::Memcmp(&FoundInitializer.RenderTargetFormats[0], &PSOInitializer.RenderTargetFormats[0], FoundInitializer.RenderTargetsEnabled * sizeof(PSOInitializer.RenderTargetFormats[0])) == 0 &&
-			FoundInitializer.DepthStencilTargetFormat == PSOInitializer.DepthStencilTargetFormat &&
-			FoundInitializer.DepthStencilAccess == PSOInitializer.DepthStencilAccess &&
-			FoundInitializer.NumSamples == PSOInitializer.NumSamples);
-#endif
-#if VULKAN_ENABLE_LRU_CACHE
-		Device->GetPipelineStateCache()->PipelineLRU.Touch(Device, Found->Pipeline, Device->GetPipelineStateCache()->GfxPipelineEntries);
-#endif
 		return Found;
 	}
 
-	FVulkanPipelineStateCacheManager::FShaderHashes ShaderHashes(PSOInitializer);
-
 	// Now try the loaded cache from disk
-	FVulkanPipelineStateCacheManager::FGfxPipelineEntry* GfxEntry = nullptr;
-	FVulkanPipelineStateCacheManager::FHashToGfxPipelinesMap* HashToGfxPipelinesMap = nullptr;
+	FVulkanPipelineStateCacheManager::TGfxPipelineEntrySharedPtr GfxEntry;
+	FGfxEntryKey GfxEntryKey;
 
 	{
-		FScopeLock Lock(&GShaderHashToGfxEntriesMapCS);
+		FScopeLock Lock(&EntryKeyToGfxPipelineMapCS);
 
-		Found = Device->PipelineStateCache->FindInLoadedLibrary(PSOInitializer, PSOInitializerHash, ShaderHashes, GfxEntry, HashToGfxPipelinesMap);
+		Found = Device->PipelineStateCache->FindInLoadedLibrary(PSOInitializer, PSIKey, GfxEntry, GfxEntryKey);
 		if (Found)
 		{
-#if VULKAN_ENABLE_LRU_CACHE
-			Device->GetPipelineStateCache()->PipelineLRU.Touch(Device, Found->Pipeline, Device->GetPipelineStateCache()->GfxPipelineEntries);
-#endif
 			return Found;
 		}
 
 		UE_LOG(LogVulkanRHI, Verbose, TEXT("PSO not found in cache, compiling..."));
 
-		// Have we made a matching state object yet?
-		FVulkanRHIGraphicsPipelineState* PipelineState = nullptr;
-#if VULKAN_ENABLE_LRU_CACHE
-		if (PSOInitializer.bFromPSOFileCache && Device->GetPipelineStateCache()->ShouldEvictImmediately())
-		{
-			if (Device->GetPipelineStateCache()->PipelineLRU.IsEvicted(GfxEntry->GetEntryHash()))
-			{
-				UE_LOG(LogVulkanRHI, Warning, TEXT("Redundant attempt to create pipeline %u as evicted??"), GfxEntry->GetEntryHash());
-			}
-			else
-			{
-				FVulkanGfxPipeline* Pipeline = Device->GetPipelineStateCache()->CreateAsEvicted(GfxEntry, ShaderHashes.Hash);
-				Device->GetPipelineStateCache()->PipelineLRU.AddAsEvicted(Device, Pipeline, Device->GetPipelineStateCache()->GfxPipelineEntries, GfxEntry);
-			}
-		}
-		else
-		{
-			PipelineState = Device->GetPipelineStateCache()->CreateAndAdd(PSOInitializer, PSOInitializerHash, GfxEntry, ShaderHashes.Hash);
-			if (PipelineState)
-			{
-				Device->GetPipelineStateCache()->PipelineLRU.Add(Device, PipelineState->Pipeline, Device->GetPipelineStateCache()->GfxPipelineEntries);
-			}
-		}
-#else
-		PipelineState = Device->GetPipelineStateCache()->CreateAndAdd(PSOInitializer, PSOInitializerHash, GfxEntry);
-#endif
-		if (PipelineState)
-		{
-			HashToGfxPipelinesMap->Add(GfxEntry->GetEntryHash(), PipelineState->Pipeline);
-		}
-
-		if (!GfxEntry->bConsumed)
-		{
-			delete GfxEntry;
-		}
+		FVulkanRHIGraphicsPipelineState* PipelineState = Device->PipelineStateCache->CreateAndAdd(PSOInitializer, MoveTemp(PSIKey), MoveTemp(GfxEntry), MoveTemp(GfxEntryKey));
 
 		return PipelineState;
 	}
@@ -2052,34 +1690,21 @@ FGraphicsPipelineStateRHIRef FVulkanDynamicRHI::RHICreateGraphicsPipelineState(c
 
 FVulkanComputePipeline* FVulkanPipelineStateCacheManager::GetOrCreateComputePipeline(FVulkanComputeShader* ComputeShader)
 {
-	FScopeLock ScopeLock(&CreateComputePipelineCS);
-
-	// Fast path, try based on FVulkanComputeShader pointer
-	FVulkanComputePipeline** ComputePipelinePtr = ComputeShaderToPipelineMap.Find(ComputeShader);
-	if (ComputePipelinePtr)
+	check(ComputeShader);
+	const uint64 Key = ComputeShader->GetShaderKey();
 	{
-		return *ComputePipelinePtr;
-	}
-
-	// create entry based on shader
-	FComputePipelineEntry * ComputeEntry = CreateComputeEntry(ComputeShader);
-
-	// Find pipeline based on entry
-	ComputePipelinePtr = ComputeEntryHashToPipelineMap.Find(ComputeEntry->EntryHash);
-	if (ComputePipelinePtr)
-	{
-		if ((*ComputePipelinePtr)->ComputeShader == nullptr) //if loaded from disk, link it to actual shader (1 time initialize step)
+		FRWScopeLock ScopeLock(ComputePipelineLock, SLT_ReadOnly);
+		FVulkanComputePipeline** ComputePipelinePtr = ComputePipelineEntries.Find(Key);
+		if (ComputePipelinePtr)
 		{
-			(*ComputePipelinePtr)->ComputeShader = ComputeShader;
+			return *ComputePipelinePtr;
 		}
-		ComputeShaderToPipelineMap.Add(ComputeShader, *ComputePipelinePtr);
-		return *ComputePipelinePtr;
 	}
 
 	// create pipeline of entry + store entry
 	double BeginTime = FPlatformTime::Seconds();
 
-	FVulkanComputePipeline* ComputePipeline = CreateComputePipelineFromEntry(ComputeEntry);
+	FVulkanComputePipeline* ComputePipeline = CreateComputePipelineFromShader(ComputeShader);
 	ComputePipeline->ComputeShader = ComputeShader;
 
 	double EndTime = FPlatformTime::Seconds();
@@ -2089,122 +1714,49 @@ FVulkanComputePipeline* FVulkanPipelineStateCacheManager::GetOrCreateComputePipe
 		UE_LOG(LogVulkanRHI, Verbose, TEXT("Hitchy compute pipeline key CS (%.3f ms)"), (float)(Delta * 1000.0));
 	}
 
-	ComputePipeline->AddRef();
-	ComputeEntryHashToPipelineMap.Add(ComputeEntry->EntryHash, ComputePipeline);
-	ComputeShaderToPipelineMap.Add(ComputeShader, ComputePipeline);
-	ComputePipelineEntries.Add(ComputeEntry->EntryHash, ComputeEntry);
+	{
+		FRWScopeLock ScopeLock(ComputePipelineLock, SLT_Write);
+		ComputePipelineEntries.FindOrAdd(Key) = ComputePipeline;
+	}
 
 	return ComputePipeline;
 }
 
-void FVulkanPipelineStateCacheManager::FComputePipelineEntry::CalculateEntryHash()
+FVulkanComputePipeline* FVulkanPipelineStateCacheManager::CreateComputePipelineFromShader(FVulkanComputeShader* Shader)
 {
-	TArray<uint8> MemFile;
-	FMemoryWriter Ar(MemFile);
-	Ar << this;
-	EntryHash = FCrc::MemCrc32(MemFile.GetData(), MemFile.GetTypeSize() * MemFile.Num());
-	EntryHash = FCrc::MemCrc32(&ShaderHash, sizeof(ShaderHash), EntryHash);
-}
-
-FVulkanPipelineStateCacheManager::FComputePipelineEntry* FVulkanPipelineStateCacheManager::CreateComputeEntry(FVulkanComputeShader* ComputeShader)
-{
-	FComputePipelineEntry* OutComputeEntry = new FComputePipelineEntry();
-
-	OutComputeEntry->ShaderHash = ComputeShader->GetHash();
-	OutComputeEntry->ShaderMicrocode = ShaderCache.Get(ComputeShader->GetHash());
-	if (!OutComputeEntry->ShaderMicrocode)
-	{
-		OutComputeEntry->ShaderMicrocode = ShaderCache.Add(ComputeShader->GetHash(), ComputeShader);
-	}
+	FVulkanComputePipeline* Pipeline = new FVulkanComputePipeline(Device);
 
-	FUniformBufferGatherInfo UBGatherInfo;
 	FVulkanDescriptorSetsLayoutInfo DescriptorSetLayoutInfo;
-	DescriptorSetLayoutInfo.ProcessBindingsForStage(VK_SHADER_STAGE_COMPUTE_BIT, ShaderStage::Compute, ComputeShader->GetCodeHeader(), UBGatherInfo);
+	const FVulkanShaderHeader& CSHeader = Shader->GetCodeHeader();
+	FUniformBufferGatherInfo UBGatherInfo;
+	DescriptorSetLayoutInfo.ProcessBindingsForStage(VK_SHADER_STAGE_COMPUTE_BIT, ShaderStage::Compute, CSHeader, UBGatherInfo);
 	DescriptorSetLayoutInfo.FinalizeBindings<true>(UBGatherInfo, TArrayView<const FSamplerStateRHIParamRef>());
-	FVulkanLayout* Layout = FindOrAddLayout(DescriptorSetLayoutInfo, false);
-	check(!Layout->IsGfxLayout());
+	FVulkanLayout* Layout = FindOrAddLayout(DescriptorSetLayoutInfo, true);
 	FVulkanComputeLayout* ComputeLayout = (FVulkanComputeLayout*)Layout;
-	OutComputeEntry->Layout = ComputeLayout;
-	check(!ComputeLayout->ComputePipelineDescriptorInfo.IsInitialized());
-	ComputeLayout->ComputePipelineDescriptorInfo.Initialize(Layout->GetDescriptorSetsLayout().RemappingInfo, ComputeShader);
-
-	OutComputeEntry->ShaderModule = ComputeShader->GetOrCreateHandle(OutComputeEntry->Layout, OutComputeEntry->Layout->GetDescriptorSetLayoutHash());
-
-	const TArray<FVulkanDescriptorSetsLayout::FSetLayout>& Layouts = DescriptorSetLayoutInfo.GetLayouts();
-	OutComputeEntry->DescriptorSetLayoutBindings.AddDefaulted(Layouts.Num());
-	for (int32 Index = 0; Index < Layouts.Num(); ++Index)
+	if (!ComputeLayout->ComputePipelineDescriptorInfo.IsInitialized())
 	{
-		for (int32 SubIndex = 0; SubIndex < Layouts[Index].LayoutBindings.Num(); ++SubIndex)
-		{
-			FDescriptorSetLayoutBinding* Binding = new(OutComputeEntry->DescriptorSetLayoutBindings[Index]) FDescriptorSetLayoutBinding;
-			Binding->ReadFrom(Layouts[Index].LayoutBindings[SubIndex]);
-		}
+		ComputeLayout->ComputePipelineDescriptorInfo.Initialize(Layout->GetDescriptorSetsLayout().RemappingInfo, Shader);
 	}
 
-	OutComputeEntry->CalculateEntryHash();
-	return OutComputeEntry;
-}
-
-FVulkanComputePipeline* FVulkanPipelineStateCacheManager::CreateComputePipelineFromEntry(const FComputePipelineEntry* ComputeEntry)
-{
-	FVulkanComputePipeline* Pipeline = new FVulkanComputePipeline(Device);
+	VkShaderModule ShaderModule = Shader->GetOrCreateHandle(Layout, Layout->GetDescriptorSetLayoutHash());
 
 	VkComputePipelineCreateInfo PipelineInfo;
 	ZeroVulkanStruct(PipelineInfo, VK_STRUCTURE_TYPE_COMPUTE_PIPELINE_CREATE_INFO);
 	PipelineInfo.stage.sType = VK_STRUCTURE_TYPE_PIPELINE_SHADER_STAGE_CREATE_INFO;
 	PipelineInfo.stage.stage = VK_SHADER_STAGE_COMPUTE_BIT;
-	PipelineInfo.stage.module = ComputeEntry->ShaderModule;
+	PipelineInfo.stage.module = ShaderModule;
 	PipelineInfo.stage.pName = "main";
-	PipelineInfo.layout = ComputeEntry->Layout->GetPipelineLayout();
+	PipelineInfo.layout = ComputeLayout->GetPipelineLayout();
 		
 	VERIFYVULKANRESULT(VulkanRHI::vkCreateComputePipelines(Device->GetInstanceHandle(), PipelineCache, 1, &PipelineInfo, VULKAN_CPU_ALLOCATOR, &Pipeline->Pipeline));
 
-	Pipeline->Layout = ComputeEntry->Layout;
+	Pipeline->Layout = ComputeLayout;
 
 	INC_DWORD_STAT(STAT_VulkanNumPSOs);
 
 	return Pipeline;
 }
 
-void FVulkanPipelineStateCacheManager::CreateComputeEntryRuntimeObjects(FComputePipelineEntry* ComputeEntry)
-{
-	{
-		// Descriptor Set Layouts
-		check(!ComputeEntry->Layout);
-
-		FVulkanDescriptorSetsLayoutInfo Info;
-		for (int32 SetIndex = 0; SetIndex < ComputeEntry->DescriptorSetLayoutBindings.Num(); ++SetIndex)
-		{
-			for (int32 Index = 0; Index < ComputeEntry->DescriptorSetLayoutBindings[SetIndex].Num(); ++Index)
-			{
-				VkDescriptorSetLayoutBinding Binding;
-				Binding.descriptorCount = 1;
-				Binding.pImmutableSamplers = nullptr;
-				ComputeEntry->DescriptorSetLayoutBindings[SetIndex][Index].WriteInto(Binding);
-				Info.AddDescriptor(SetIndex, Binding);
-			}
-		}
-
-		FVulkanLayout* Layout = FindOrAddLayout(Info, false);
-		check(!Layout->IsGfxLayout());
-		ComputeEntry->Layout = (FVulkanComputeLayout*)Layout;
-	}
-
-	{
-		// Shader
-		if (ComputeEntry->ShaderMicrocode != nullptr)
-		{
-			VkShaderModuleCreateInfo ModuleCreateInfo;
-			ZeroVulkanStruct(ModuleCreateInfo, VK_STRUCTURE_TYPE_SHADER_MODULE_CREATE_INFO);
-			ModuleCreateInfo.codeSize = ComputeEntry->ShaderMicrocode->Num();
-			ModuleCreateInfo.pCode = (uint32*)ComputeEntry->ShaderMicrocode->GetData();
-			VERIFYVULKANRESULT(VulkanRHI::vkCreateShaderModule(Device->GetInstanceHandle(), &ModuleCreateInfo, VULKAN_CPU_ALLOCATOR, &ComputeEntry->ShaderModule));
-		}		
-	}
-
-	ComputeEntry->bLoaded = true;
-}
-
 template<typename T>
 inline void SerializeArray(FArchive& Ar, TArray<T*>& Array)
 {
@@ -2229,100 +1781,6 @@ inline void SerializeArray(FArchive& Ar, TArray<T*>& Array)
 	}
 }
 
-#if VULKAN_ENABLE_GENERIC_PIPELINE_CACHE_FILE
-void FVulkanPipelineStateCacheManager::FVulkanPipelineStateCacheFile::Save(FArchive& Ar)
-{
-	check(ShaderCache);
-
-	// Modify VERSION if serialization changes
-	TArray<uint8> DataBuffer;
-	FMemoryWriter DataAr(DataBuffer);
-
-	DataAr << *ShaderCache;
-	SerializeArray(DataAr, GfxPipelineEntries);
-	SerializeArray(DataAr, ComputePipelineEntries);
-
-	// compress the data buffer
-	TArray<uint8> CompressedDataBuffer = DataBuffer;
-	if (GEnablePipelineCacheCompression)
-	{
-		Header.UncompressedSize = DataBuffer.Num() * DataBuffer.GetTypeSize();
-		int32 CompressedSize = CompressedDataBuffer.Num();
-		if (FCompression::CompressMemory(CompressionFlags, CompressedDataBuffer.GetData(), CompressedSize,
-			DataBuffer.GetData(), Header.UncompressedSize))
-		{
-			CompressedDataBuffer.SetNum(CompressedSize);
-		}
-		CompressedDataBuffer.Shrink();
-	}
-
-	Ar << Header.Version;
-	Ar << Header.SizeOfGfxEntry;
-	Ar << Header.SizeOfComputeEntry;
-	Ar << Header.UncompressedSize;
-
-	Ar << CompressedDataBuffer;
-}
-
-
-bool FVulkanPipelineStateCacheManager::FVulkanPipelineStateCacheFile::Load(FArchive& Ar, const TCHAR* Filename)
-{
-	check(ShaderCache);
-
-	// Modify VERSION if serialization changes
-	Ar << Header.Version;
-	if (Header.Version != VERSION)
-	{
-		UE_LOG(LogVulkanRHI, Warning, TEXT("Unable to load shader cache due to mismatched Version %d != %d"), Header.Version, (int32)VERSION);
-		return false;
-	}
-
-	Ar << Header.SizeOfGfxEntry;
-	if (Header.SizeOfGfxEntry != (int32)(sizeof(FGfxPipelineEntry)))
-	{
-		UE_LOG(LogVulkanRHI, Warning, TEXT("Unable to load shader cache due to mismatched size of FGfxEntry %d != %d; forgot to bump up VERSION?"), Header.SizeOfGfxEntry, (int32)sizeof(FGfxPipelineEntry));
-		return false;
-	}
-
-	Ar << Header.SizeOfComputeEntry;
-	if (Header.SizeOfComputeEntry != (int32)(sizeof(FComputePipelineEntry)))
-	{
-		UE_LOG(LogVulkanRHI, Warning, TEXT("Unable to load shader cache due to mismatched size of FComputePipelineEntry %d != %d; forgot to bump up VERSION?"), Header.SizeOfComputeEntry, (int32)sizeof(FComputePipelineEntry));
-		return false;
-	}
-
-	Ar << Header.UncompressedSize;
-
-	TArray<uint8> CompressedDataBuffer;
-	Ar << CompressedDataBuffer;
-
-	TArray<uint8> UncompressedDataBuffer;
-	if (Header.UncompressedSize != 0)
-	{
-		const uint32 CompressedSize = CompressedDataBuffer.Num() * CompressedDataBuffer.GetTypeSize();
-		UncompressedDataBuffer.SetNum(Header.UncompressedSize);
-		if (!FCompression::UncompressMemory(CompressionFlags, UncompressedDataBuffer.GetData(), Header.UncompressedSize, CompressedDataBuffer.GetData(), CompressedSize))
-		{
-			UE_LOG(LogVulkanRHI, Error, TEXT("Failed to uncompress data for pipeline cache file %s!"), Filename);
-			return false;
-		}
-	}
-	else
-	{
-		UncompressedDataBuffer = CompressedDataBuffer;
-	}
-
-	FMemoryReader DataAr(UncompressedDataBuffer);
-	DataAr << *ShaderCache;
-
-	SerializeArray(DataAr, GfxPipelineEntries);
-
-	SerializeArray(DataAr, ComputePipelineEntries);
-
-	return true;
-}
-#endif
-
 bool FVulkanPipelineStateCacheManager::BinaryCacheMatches(FVulkanDevice* InDevice, const TArray<uint8>& DeviceCache)
 {
 	if (DeviceCache.Num() > 4)
@@ -2358,9 +1816,100 @@ bool FVulkanPipelineStateCacheManager::BinaryCacheMatches(FVulkanDevice* InDevic
 }
 
 #if VULKAN_ENABLE_LRU_CACHE
-void FVulkanPipelineStateCacheManager::FVKPipelineLRU::FEvictedVkPipeline::RestoreVkPipeline()
+void FVulkanPipelineStateCacheManager::FVKPipelineLRU::DeleteVkPipeline(FVulkanGfxPipeline* GfxPipeline)
 {
-	Device->GetPipelineStateCache()->CreateGfxPipelineFromEntry(GfxEntry, GfxPipeline);
+	GfxPipeline->DeleteVkPipeline(GfxPipeline->RecentFrame + NUM_BUFFERS < GFrameNumberRenderThread);
+}
+
+void FVulkanPipelineStateCacheManager::FVKPipelineLRU::EnsureVkPipelineAndAddToLRU(FVulkanRHIGraphicsPipelineState* Pipeline)
+{
+	FVulkanGfxPipeline* GfxPipeline = Pipeline->Pipeline;
+	FGfxPipelineEntry* GfxPipelineEntry = static_cast<FGfxPipelineEntry*>(GfxPipeline->GfxPipelineEntry.Get());
+	FVulkanPipelineStateCacheManager* PipelineStateCache = GfxPipeline->Device->GetPipelineStateCache();
+
+	{
+		// This CS is used in RHICreateGraphicsPipelineState() to guard pipeline creation from multiple threads
+		FScopeLock Lock(&EntryKeyToGfxPipelineMapCS);
+
+		if (!GfxPipeline->GetHandle())
+		{
+			PipelineStateCache->CreateGfxPipelineFromEntry(GfxPipelineEntry, &Pipeline->BSI, GfxPipeline);
+
+			// Failed to create VkPipeline again
+			check(GfxPipeline->GetHandle());
+
+			AddToLRU(GfxPipeline);
+		}
+	}
+}
+
+void FVulkanPipelineStateCacheManager::FVKPipelineLRU::AddToLRU(FVulkanGfxPipeline* GfxPipeline)
+{
+	FScopeLock Lock(&LRUCS);
+
+	while (LRUUsedPipelineSize + GfxPipeline->PipelineCacheSize >(uint32)CVarLRUMaxPipelineSize.GetValueOnAnyThread() || LRU.Num() == LRU.Max())
+	{
+		EvictFromLRU();
+	}
+
+	check(!LRU.Contains(GfxPipeline));
+	GfxPipeline->LRUNode = LRU.Add(GfxPipeline, GfxPipeline);
+	LRUUsedPipelineSize += GfxPipeline->PipelineCacheSize;
+}
+
+void FVulkanPipelineStateCacheManager::FVKPipelineLRU::EvictFromLRU()
+{
+	FVulkanGfxPipeline* LeastRecentPipeline = LRU.RemoveLeastRecent();
+	check(LeastRecentPipeline);
+	LeastRecentPipeline->LRUNode = FSetElementId();
+
+	LRUUsedPipelineSize -= LeastRecentPipeline->PipelineCacheSize;
+
+	DeleteVkPipeline(LeastRecentPipeline);
+}
+
+void FVulkanPipelineStateCacheManager::FVKPipelineLRU::Add(FVulkanGfxPipeline* GfxPipeline)
+{
+	if (!bUseLRU)
+	{
+		return;
+	}
+
+	check(!GfxPipeline->bTrackedByLRU);
+	GfxPipeline->bTrackedByLRU = true;
+
+	if (!GfxPipeline->GetHandle())
+	{
+		GfxPipeline->LRUNode = FSetElementId(); // Add as evicted
+		return;
+	}
+
+	AddToLRU(GfxPipeline);
+}
+
+void FVulkanPipelineStateCacheManager::FVKPipelineLRU::Touch(FVulkanRHIGraphicsPipelineState* Pipeline)
+{
+	if (!bUseLRU)
+	{
+		return;
+	}
+
+	FVulkanGfxPipeline* GfxPipeline = Pipeline->Pipeline;
+	check(GfxPipeline->bTrackedByLRU);
+
+	{
+		FScopeLock Lock(&LRUCS);
+		if (GfxPipeline->LRUNode.IsValidId())
+		{
+			check(GfxPipeline->GetHandle());
+			check(LRU.Contains(GfxPipeline));
+			LRU.MarkAsRecent(GfxPipeline->LRUNode);
+			GfxPipeline->RecentFrame = GFrameNumberRenderThread;
+			return;
+		}
+	}
+
+	EnsureVkPipelineAndAddToLRU(Pipeline);
 }
 
 void FVulkanPipelineStateCacheManager::FVulkanLRUCacheFile::Save(FArchive& Ar)
diff --git a/Engine/Source/Runtime/VulkanRHI/Private/VulkanPipeline.h b/Engine/Source/Runtime/VulkanRHI/Private/VulkanPipeline.h
--- a/Engine/Source/Runtime/VulkanRHI/Private/VulkanPipeline.h
+++ b/Engine/Source/Runtime/VulkanRHI/Private/VulkanPipeline.h
@@ -18,205 +18,23 @@ extern TAutoConsoleVariable<int32> CVarEnableLRU;
 #endif
 
 class FVulkanDevice;
+class FVulkanRHIGraphicsPipelineState;
+class FVulkanGfxPipeline;
 
-// Common pipeline class
-class FVulkanPipeline
-{
-public:
-	FVulkanPipeline(FVulkanDevice* InDevice);
-	virtual ~FVulkanPipeline();
-
-	inline VkPipeline GetHandle() const
-	{
-		return Pipeline;
-	}
-
-	inline const FVulkanLayout& GetLayout() const
-	{
-		return *Layout;
-	}
-#if VULKAN_ENABLE_LRU_CACHE
-	inline void DeleteVkPipeline(bool ImmediateDestroy)
-	{
-		if (Pipeline)
-		{
-			if (ImmediateDestroy)
-			{
-				VulkanRHI::vkDestroyPipeline(Device->GetInstanceHandle(), Pipeline, VULKAN_CPU_ALLOCATOR);
-}
-			else
-			{
-				Device->GetDeferredDeletionQueue().EnqueueResource(VulkanRHI::FDeferredDeletionQueue::EType::Pipeline, Pipeline);
-			}
-			Pipeline = VK_NULL_HANDLE;
-		}
-	}
-#endif
-protected:
-	FVulkanDevice* Device;
-	VkPipeline Pipeline;
-	FVulkanLayout* Layout; /*owned by FVulkanPipelineStateCacheManager, do not delete yourself !*/
-
-	friend class FVulkanPipelineStateCacheManager;
-	friend class FVulkanRHIGraphicsPipelineState;
-	friend class FVulkanComputePipelineDescriptorState;
-};
-
-class FVulkanComputePipeline : public FVulkanPipeline, public FRHIComputePipelineState
-{
-public:
-	FVulkanComputePipeline(FVulkanDevice* InDevice);
-	virtual ~FVulkanComputePipeline();
-
-	inline const FVulkanShaderHeader& GetShaderCodeHeader() const
-	{
-		return ComputeShader->GetCodeHeader();
-	}
-
-	inline const FVulkanComputeShader* GetShader() const
-	{
-		return ComputeShader;
-	}
-
-	inline void Bind(VkCommandBuffer CmdBuffer)
-	{
-		VulkanRHI::vkCmdBindPipeline(CmdBuffer, VK_PIPELINE_BIND_POINT_COMPUTE, Pipeline);
-	}
-
-	inline const FVulkanComputeLayout& GetComputeLayout() const
-	{
-		return *(FVulkanComputeLayout*)Layout;
-	}
-
-protected:
-	FVulkanComputeShader*	ComputeShader;
-
-	friend class FVulkanPipelineStateCacheManager;
-};
-
-class FVulkanGfxPipeline : public FVulkanPipeline, public FRHIResource
-{
-public:
-#if VULKAN_ENABLE_LRU_CACHE
-	FVulkanGfxPipeline(FVulkanDevice* InDevice, uint32 InGfxEntryHash, uint32 InShaderHash);
-#else
-	FVulkanGfxPipeline(FVulkanDevice* InDevice);
-#endif
-
-	inline void Bind(VkCommandBuffer CmdBuffer)
-	{
-#if VULKAN_ENABLE_LRU_CACHE
-		RecentFrame = GFrameNumberRenderThread;
-#endif
-		VulkanRHI::vkCmdBindPipeline(CmdBuffer, VK_PIPELINE_BIND_POINT_GRAPHICS, Pipeline);
-	}
-
-	inline bool IsRuntimeInitialized() const
-	{
-		return bRuntimeObjectsValid;
-	}
-
-	inline const FVulkanVertexInputStateInfo& GetVertexInputState() const
-	{
-		check(bRuntimeObjectsValid);
-		return VertexInputState;
-	}
-
-	void CreateRuntimeObjects(const FGraphicsPipelineStateInitializer& InPSOInitializer);
-
-	inline const FVulkanGfxLayout& GetGfxLayout() const
-	{
-		return *(FVulkanGfxLayout*)Layout;
-	}
-
-#if VULKAN_ENABLE_LRU_CACHE
-	uint32 GfxEntryHash;
-	uint32 PipelineCacheSize;
-	uint32 ShaderHash;
-	uint32 RecentFrame;
-	// ID to LRU (if used) allows quick access when updating LRU status.
-	FSetElementId LRUNode;
-#endif
-private:
-	FVulkanVertexInputStateInfo VertexInputState;
-	bool						bRuntimeObjectsValid;
-};
-
-class FVulkanRHIGraphicsPipelineState : public FRHIGraphicsPipelineState
-{
-public:
-	FVulkanRHIGraphicsPipelineState(const FGraphicsPipelineStateInitializer& Initializer, FVulkanGfxPipeline* InPipeline)
-		: Pipeline(InPipeline)
-		, PipelineStateInitializer(Initializer)
-	{
-		bHasInputAttachments = InPipeline->GetGfxLayout().GetDescriptorSetsLayout().HasInputAttachments();
-	}
-
-	~FVulkanRHIGraphicsPipelineState();
-
-	inline void Bind(VkCommandBuffer CmdBuffer)
-	{
-		Pipeline->Bind(CmdBuffer);
-	}
-
-	inline const FVulkanShader* GetShader(EShaderFrequency Frequency) const
-	{
-		auto& BSS = PipelineStateInitializer.BoundShaderState;
-
-		FVulkanShader* Shader = nullptr;
-
-		switch (Frequency)
-		{
-		case SF_Vertex:   Shader = ResourceCast(BSS.VertexShaderRHI); break;
-		case SF_Hull:     Shader = ResourceCast(BSS.HullShaderRHI); break;
-		case SF_Domain:   Shader = ResourceCast(BSS.DomainShaderRHI); break;
-		case SF_Pixel:    Shader = ResourceCast(BSS.PixelShaderRHI); break;
-		case SF_Geometry: Shader = ResourceCast(BSS.GeometryShaderRHI); break;
-		default:
-			check(0);
-			break;
-		}
-
-		return Shader;
-	}
-
-	TRefCountPtr<FVulkanGfxPipeline>	Pipeline;
-	bool								bHasInputAttachments;
-
-	FGraphicsPipelineStateInitializer	PipelineStateInitializer;
-};
+class FGfxPSIKey : public TDataKey<FGfxPSIKey> {};
+class FGfxEntryKey : public TDataKey<FGfxEntryKey> {};
 
 template <typename T>
-static inline uint64 GetShaderId(T* ShaderType)
+static inline uint64 GetShaderKey(T* ShaderType)
 {
 	auto* VulkanShader = ResourceCast(ShaderType);
-	return VulkanShader ? VulkanShader->GetId() : 0;
+	return VulkanShader ? VulkanShader->GetShaderKey() : 0;
 }
 
 class FVulkanPipelineStateCacheManager
 {
 public:
-	struct FPSOHashable
-	{
-		uint64								ShaderIds[ShaderStage::NumStages];
-		FVertexDeclarationElementList		VertexDeclaration;
-#if VULKAN_SUPPORTS_COLOR_CONVERSIONS
-		VkSampler							ImmutableSamplers[MaxImmutableSamplers];
-#endif
-		FBlendStateInitializerRHI			BlendState;
-		FRasterizerStateInitializerRHI		RasterizerState;
-		FDepthStencilStateInitializerRHI	DepthStencilState;
-		FExclusiveDepthStencil				DepthStencilAccess;
-		EPixelFormat						DepthStencilTargetFormat;
-		uint8								bDepthBounds;
-		TEnumAsByte<EPrimitiveType>			PrimitiveType;
-		uint8								RenderTargetsEnabled;
-		uint8								NumSamples;
-	};
-
-	FVulkanRHIGraphicsPipelineState* FindInRuntimeCache(const FGraphicsPipelineStateInitializer& Initializer, uint32& OutHash);
-
-	void DestroyPipeline(FVulkanGfxPipeline* Pipeline);
+	FVulkanRHIGraphicsPipelineState* FindInRuntimeCache(const FGraphicsPipelineStateInitializer& Initializer, FGfxPSIKey& OutKey);
 
 	// Array of potential cache locations; first entries have highest priority. Only one cache file is loaded. If unsuccessful, tries next entry in the array.
 	void InitAndLoad(const TArray<FString>& CacheFilenames);
@@ -227,14 +45,6 @@ public:
 
 	void RebuildCache();
 
-#if VULKAN_ENABLE_GENERIC_PIPELINE_CACHE_FILE
-	enum
-	{
-		// Bump every time serialization changes
-		VERSION = 23,
-	};
-#endif
-
 	struct FDescriptorSetLayoutBinding
 	{
 		uint32 Binding;
@@ -276,13 +86,45 @@ public:
 		}
 	};
 
+	struct FShaderHashes
+	{
+		uint32 Hash;
+		FSHAHash Stages[ShaderStage::NumStages];
+
+		FShaderHashes();
+		FShaderHashes(const FGraphicsPipelineStateInitializer& PSOInitializer);
+
+		friend inline uint32 GetTypeHash(const FShaderHashes& Hashes)
+		{
+			return Hashes.Hash;
+		}
+
+		inline void Finalize()
+		{
+			Hash = FCrc::MemCrc32(Stages, sizeof(Stages));
+		}
+
+		friend inline bool operator == (const FShaderHashes& A, const FShaderHashes& B)
+		{
+			for (int32 Index = 0; Index < ShaderStage::NumStages; ++Index)
+			{
+				if (A.Stages[Index] != B.Stages[Index])
+				{
+					return false;
+				}
+			}
+
+			return true;
+		}
+	};
+
 	// Actual information required to recreate a pipeline when saving/loading from disk
 	struct FGfxPipelineEntry
 	{
-		uint32 GetEntryHash(uint32 Crc = 0);
+		FGfxEntryKey CreateKey() const;
+
 		uint32 VertexInputKey;
 		bool bLoaded;
-		bool bConsumed; // Entry was added to cache and may not be deleted
 
 		uint32 RasterizationSamples;
 		uint32 Topology;
@@ -423,8 +265,7 @@ public:
 		};
 		FDepthStencil DepthStencil;
 
-		TArray<uint32>* ShaderMicrocodes[ShaderStage::NumStages];
-		FSHAHash ShaderHashes[ShaderStage::NumStages];
+		FShaderHashes ShaderHashes;
 #if VULKAN_SUPPORTS_COLOR_CONVERSIONS
 		SIZE_T ImmutableSamplers[MaxImmutableSamplers] = { 0 };
 #endif
@@ -508,7 +349,6 @@ public:
 		FGfxPipelineEntry()
 			: VertexInputKey(0)
 			, bLoaded(false)
-			, bConsumed(false)
 			, RasterizationSamples(0)
 			, Topology(0)
 			, RenderPass(nullptr)
@@ -517,26 +357,17 @@ public:
 			FMemory::Memzero(Rasterizer);
 			FMemory::Memzero(DepthStencil);
 			FMemory::Memzero(ShaderModules);
-			FMemory::Memzero(ShaderMicrocodes);
 		}
 
 		~FGfxPipelineEntry();
 
-		// Vulkan Runtime Data/Objects
-		// Used for not loaded entries in order to
-		// create ShaderModules on pipeline creation time.
-		// After pipeline is created shader modules will be purged.
-		TRefCountPtr<FVulkanShader> VulkanShaders[ShaderStage::NumStages];
-		// Holds shader modules that is used to create pipeline
-		// For loaded entries shader modules created when entry is loaded and destroyed when entry is destroyed
-		// For runtime entries shader modules temporary created when pipeline is created
-		// Outside of pipeline creation shader modules are empty
 		VkShaderModule ShaderModules[ShaderStage::NumStages];
 		const FVulkanRenderPass* RenderPass;
 		FVulkanGfxLayout* Layout;
 
-		void CreateShaderModules();
-		void PurgeShaderModules();
+		void GetOrCreateShaderModules(FVulkanShader*const* Shaders);
+		void PurgeShaderModules(FVulkanShader*const* Shaders);
+		void PurgeLoadedShaderModules(FVulkanDevice* InDevice);
 
 		bool operator==(const FGfxPipelineEntry& In) const
 		{
@@ -580,20 +411,9 @@ public:
 				return false;
 			}
 
-			for (int32 Index = 0; Index < ARRAY_COUNT(In.ShaderHashes); ++Index)
+			if (!(ShaderHashes == In.ShaderHashes))
 			{
-				if (ShaderHashes[Index] != In.ShaderHashes[Index])
-				{
-					return false;
-				}
-			}
-
-			for (int32 Index = 0; Index < ARRAY_COUNT(In.ShaderMicrocodes); ++Index)
-			{
-				if (ShaderMicrocodes[Index] != In.ShaderMicrocodes[Index])
-				{
-					return false;
-				}
+				return false;
 			}
 
 #if VULKAN_SUPPORTS_COLOR_CONVERSIONS
@@ -625,6 +445,8 @@ public:
 		}
 	};
 
+	using TGfxPipelineEntrySharedPtr = TSharedPtr<FGfxPipelineEntry, ESPMode::ThreadSafe>;
+
 #if VULKAN_ENABLE_LRU_CACHE
 	struct FPipelineSize
 	{
@@ -638,119 +460,37 @@ public:
 	};
 #endif
 
-	struct FComputePipelineEntry
-	{
-		uint32 EntryHash;
-		void CalculateEntryHash();
-
-		bool bLoaded;
-
-		TArray<uint32>* ShaderMicrocode;
-		FSHAHash ShaderHash;
-		TArray<TArray<FDescriptorSetLayoutBinding>> DescriptorSetLayoutBindings;
-
-		// Runtime objects
-		VkShaderModule ShaderModule;
-		FVulkanComputeLayout* Layout;
-
-		FComputePipelineEntry()
-			: EntryHash(0)
-			, bLoaded(false)
-			, ShaderMicrocode(nullptr)
-			, ShaderModule(VK_NULL_HANDLE)
-			, Layout(nullptr)
-		{}
-
-		~FComputePipelineEntry();
-	};
 	FVulkanComputePipeline* GetOrCreateComputePipeline(FVulkanComputeShader* ComputeShader);
 
 #if VULKAN_ENABLE_LRU_CACHE
+
 	bool ShouldEvictImmediately()
 	{
 		return bEvictImmediately && PipelineLRU.IsActive();
 	}
 
-	void CreateGfxPipelineFromEntry(FGfxPipelineEntry* GfxEntry, FVulkanGfxPipeline* Pipeline);
-	TMap<uint32, FGfxPipelineEntry*> GfxPipelineEntries;
-
 	TMap<uint32, FPipelineSize*> PipelineSizeList;	// key: Shader hash (FShaderHash), value: pipeline size
 
 	class FVKPipelineLRU
 	{
-		class FEvictedVkPipeline
-		{
-			FVulkanGfxPipeline* GfxPipeline;
-			FVulkanDevice* Device;
-			FGfxPipelineEntry* GfxEntry;
-
-		public:
-			enum
-			{
-				NUM_FRAMES = 3,
-			};
-			FEvictedVkPipeline(FVulkanDevice* InDevice, FVulkanGfxPipeline* InPipeline, FGfxPipelineEntry* InGfxEntry)
-				: Device(InDevice), GfxPipeline(InPipeline), GfxEntry(InGfxEntry)
-			{
-				check(InGfxEntry);
-				check(InGfxEntry->bConsumed);
-
-				GfxPipeline->DeleteVkPipeline(InPipeline->RecentFrame + NUM_FRAMES < GFrameNumberRenderThread);
-			}
-
-			void RestoreVkPipeline();
-
-			FVulkanGfxPipeline* GetGfxPipeline()
-			{
-				return GfxPipeline;
-			}
-		};
-		typedef TMap<uint32, FEvictedVkPipeline> FVkEvictedPipelineMap;	//key: GfxEntry hash
-		typedef TPsoLruCache<uint32, FVulkanGfxPipeline*> FVulkanPipelineLRUCache;	//key: GfxEntryHash
+		typedef TPsoLruCache<FVulkanGfxPipeline*, FVulkanGfxPipeline*> FVulkanPipelineLRUCache;
 
 		const int LRUCapacity = 2048;
 		int32 LRUUsedPipelineSize;
 
-		FVulkanGfxPipeline* FindEvictedAndUpdateLRU(FVulkanDevice* InDevice, uint32 EntryHash, const TMap<uint32, FGfxPipelineEntry*>& InGfxEntriesMap)
+		enum
 		{
-			FScopeLock Lock(&EvictCS);
-
-			FEvictedVkPipeline* FoundEvicted = EvictedPipelines.Find(EntryHash);
-			if (FoundEvicted)
-			{
-				FoundEvicted->RestoreVkPipeline();
-				FVulkanGfxPipeline* GfxPipeline = FoundEvicted->GetGfxPipeline();
-
-				Add(InDevice, GfxPipeline, InGfxEntriesMap);
-
-				EvictedPipelines.Remove(EntryHash);
-
-				return GfxPipeline;
-			}
+			NUM_BUFFERS = 3,
+		};
 
-			return nullptr;
-		}
+		static void DeleteVkPipeline(FVulkanGfxPipeline* GfxPipeline);
 
-		void EvictFromLRU(FVulkanDevice* InDevice, uint32 InEntryHash, const TMap<uint32, FGfxPipelineEntry*>& InGfxEntriesMap)
-		{
-			FVulkanGfxPipeline* LeastRecentPipeline = nullptr;
-			{
-				FScopeLock Lock(&LRUCS);
-				LeastRecentPipeline = LRU.RemoveLeastRecent();
-				check(LeastRecentPipeline);
-				LeastRecentPipeline->LRUNode = FSetElementId();
+		void EnsureVkPipelineAndAddToLRU(FVulkanRHIGraphicsPipelineState* Pipeline);
 
-				LRUUsedPipelineSize -= LeastRecentPipeline->PipelineCacheSize;
+		void AddToLRU(FVulkanGfxPipeline* GfxPipeline);
 
-				uint32 LeastGfxEntryHash = LeastRecentPipeline->GfxEntryHash;
-				check(!(LeastGfxEntryHash == InEntryHash));
-				check(!EvictedPipelines.Contains(LeastGfxEntryHash));
-				FGfxPipelineEntry*const* Found = InGfxEntriesMap.Find(LeastGfxEntryHash);
-				check(Found);
-
-				FEvictedVkPipeline& test = EvictedPipelines.Emplace(LeastGfxEntryHash, FEvictedVkPipeline(InDevice, LeastRecentPipeline, *Found));
-			}
-		}
+		// Use with external lock - LRUCS
+		void EvictFromLRU();
 
 	public:
 		FVKPipelineLRU() : LRUUsedPipelineSize(0), LRU(LRUCapacity)
@@ -763,97 +503,14 @@ public:
 			return bUseLRU;
 		}
 
-		void Add(FVulkanDevice* InDevice, FVulkanGfxPipeline* GfxPipeline, const TMap<uint32, FGfxPipelineEntry*>& InGfxEntriesMap)
-		{
-			if (!bUseLRU)
-			{
-				return;
-			}
-
-			uint32 EntryHash = GfxPipeline->GfxEntryHash;
-			check(!LRU.Contains(EntryHash));
-
-			while (LRUUsedPipelineSize + GfxPipeline->PipelineCacheSize > (uint32)CVarLRUMaxPipelineSize.GetValueOnAnyThread() || LRU.Num() == LRU.Max())
-			{
-				EvictFromLRU(InDevice, GfxPipeline->GfxEntryHash, InGfxEntriesMap);
-			}
-
-
-			{
-				FScopeLock Lock(&LRUCS);
-				GfxPipeline->LRUNode = LRU.Add(EntryHash, GfxPipeline);
-			}
-			LRUUsedPipelineSize += GfxPipeline->PipelineCacheSize;
-		}
-
-		bool IsEvicted(uint32 EntryHash)
-		{
-			return EvictedPipelines.Contains(EntryHash);
-		}
-
-		void AddAsEvicted(FVulkanDevice* InDevice, FVulkanGfxPipeline* GfxPipeline, const TMap<uint32, FGfxPipelineEntry*>& InGfxEntriesMap, FGfxPipelineEntry* GfxEntry)
-		{
-			check(bUseLRU);
-
-			uint32 EntryHash = GfxPipeline->GfxEntryHash;
-			check(!LRU.Contains(EntryHash));
-			check(!EvictedPipelines.Contains(EntryHash));
-
-			{
-				FScopeLock Lock(&LRUCS);
-
-				GfxPipeline->LRUNode = FSetElementId();
-
-				FEvictedVkPipeline& test = EvictedPipelines.Emplace(EntryHash, FEvictedVkPipeline(InDevice, GfxPipeline, GfxEntry));
-			}
-		}
-
-		FVulkanGfxPipeline* Find(FVulkanDevice* InDevice, uint32 EntryHash, const TMap<uint32, FGfxPipelineEntry*>& InGfxEntriesMap)
-		{
-			check(bUseLRU);
-
-			FVulkanGfxPipeline*const* Found = LRU.FindAndTouch(EntryHash);
-			if (Found)
-			{
-				check((*Found)->LRUNode.IsValidId());
-				(*Found)->RecentFrame = GFrameNumberRenderThread;
-				return *Found;
-			}
-
-			return FindEvictedAndUpdateLRU(InDevice, EntryHash, InGfxEntriesMap);
-		}
-
-		FORCEINLINE_DEBUGGABLE void Touch(FVulkanDevice* InDevice, FVulkanGfxPipeline* GfxPipeline, const TMap<uint32, FGfxPipelineEntry*>& InGfxEntriesMap)
-		{
-			if (!bUseLRU)
-			{
-				return;
-			}
-
-			if (GfxPipeline->LRUNode.IsValidId())
-			{
-				FScopeLock Lock(&LRUCS);
-				LRU.MarkAsRecent(GfxPipeline->LRUNode);
-				GfxPipeline->RecentFrame = GFrameNumberRenderThread;
-			}
-			else
-			{
-				FindEvictedAndUpdateLRU(InDevice, GfxPipeline->GfxEntryHash, InGfxEntriesMap);
-			}
-		}
-
-		void Empty()
-		{
-			EvictedPipelines.Empty();
-			LRU.Empty(LRUCapacity);
-		}
+		void Add(FVulkanGfxPipeline* GfxPipeline);
 
+		void Touch(FVulkanRHIGraphicsPipelineState* Pipeline);
+	
 	private:
 		bool bUseLRU;
 		FVulkanPipelineLRUCache LRU;
-		FVkEvictedPipelineMap EvictedPipelines;
-		FCriticalSection LRUCS;
-		FCriticalSection EvictCS;
+		FCriticalSection LRUCS; // For LRU, LRUUsedPipelineSize and GfxPipeline->LRUNode
 	};
 
 	FVKPipelineLRU PipelineLRU;
@@ -876,101 +533,44 @@ private:
 	void OnShaderPipelineCacheOpened(FString const& Name, EShaderPlatform Platform, uint32 Count, const FGuid& VersionGuid, FShaderPipelineCache::FShaderCachePrecompileContext& ShaderCachePrecompileContext);
 	void OnShaderPipelineCachePrecompilationComplete(uint32 Count, double Seconds, const FShaderPipelineCache::FShaderCachePrecompileContext& ShaderCachePrecompileContext);
 
-	// Key is a hash of the PSO, which is based off shader pointers
-	TMap<uint32, FVulkanRHIGraphicsPipelineState*> InitializerToPipelineMap;
+	TMap<FGfxPSIKey, FVulkanRHIGraphicsPipelineState*> InitializerToPipelineMap;
 	FCriticalSection InitializerToPipelineMapCS;
 
-	TMap<FVulkanComputeShader*, FVulkanComputePipeline*> ComputeShaderToPipelineMap;
-	TMap<uint32, FVulkanComputePipeline*> ComputeEntryHashToPipelineMap;
-
 	FCriticalSection GfxPipelineEntriesCS;
-#if !VULKAN_ENABLE_LRU_CACHE
-	TMap<uint32, FGfxPipelineEntry*> GfxPipelineEntries;
-#endif
-	FCriticalSection CreateComputePipelineCS;
-	TMap<uint32, FComputePipelineEntry*> ComputePipelineEntries;
+	TMap<FGfxEntryKey, TGfxPipelineEntrySharedPtr> GfxPipelineEntries;
+
+	FRWLock ComputePipelineLock;
+	TMap<uint64, FVulkanComputePipeline*> ComputePipelineEntries;
 
 	VkPipelineCache PipelineCache;
 
 	FShaderUCodeCache ShaderCache;
 
-#if VULKAN_ENABLE_LRU_CACHE
-	FVulkanGfxPipeline* CreateAsEvicted(FGfxPipelineEntry* GfxEntry, uint32 ShaderHash);
-
-	FVulkanRHIGraphicsPipelineState* CreateAndAdd(const FGraphicsPipelineStateInitializer& PSOInitializer, uint32 PSOInitializerHash, FGfxPipelineEntry* GfxEntry, uint32 ShaderHash);
-#else
-	FVulkanRHIGraphicsPipelineState* CreateAndAdd(const FGraphicsPipelineStateInitializer& PSOInitializer, uint32 PSOInitializerHash, FGfxPipelineEntry* GfxEntry);
-	void CreateGfxPipelineFromEntry(FGfxPipelineEntry* GfxEntry, FVulkanGfxPipeline* Pipeline);
-#endif
+	FVulkanRHIGraphicsPipelineState* CreateAndAdd(const FGraphicsPipelineStateInitializer& PSOInitializer, FGfxPSIKey PSIKey, TGfxPipelineEntrySharedPtr GfxEntry, FGfxEntryKey GfxEntryKey);
+	void CreateGfxPipelineFromEntry(FGfxPipelineEntry* GfxEntry, const FBoundShaderStateInput* BSI, FVulkanGfxPipeline* Pipeline);
 
-#if VULKAN_ENABLE_LRU_CACHE
-	FGfxPipelineEntry* CreateGfxEntry(const FGraphicsPipelineStateInitializer& PSOInitializer, uint32 PSOHash);
-#else
 	FGfxPipelineEntry* CreateGfxEntry(const FGraphicsPipelineStateInitializer& PSOInitializer);
-#endif
-	void CreatGfxEntryRuntimeObjects(FGfxPipelineEntry* GfxEntry);
+
 	bool Load(const TArray<FString>& CacheFilenames);
 	void DestroyCache();
 
-#if VULKAN_ENABLE_LRU_CACHE
-	FVulkanGfxLayout* GetOrGenerateGfxLayout(const FGraphicsPipelineStateInitializer& PSOInitializer, FVulkanShader** OutShaders, FVulkanVertexInputStateInfo& OutVertexInputState, uint32 PSOHash);
-#else
-	FVulkanGfxLayout* GetOrGenerateGfxLayout(const FGraphicsPipelineStateInitializer& PSOInitializer, FVulkanShader** OutShaders, FVulkanVertexInputStateInfo& OutVertexInputState);
-#endif
+	void GetVulkanShaders(const FBoundShaderStateInput& BSI, FVulkanShader* OutShaders[ShaderStage::NumStages]);
+	FVulkanGfxLayout* GetOrGenerateGfxLayout(const FGraphicsPipelineStateInitializer& PSOInitializer, FVulkanShader*const* Shaders, FVulkanVertexInputStateInfo& OutVertexInputState);
 
-
-	struct FShaderHashes
-	{
-		uint32 Hash;
-		FSHAHash Stages[ShaderStage::NumStages];
-
-		FShaderHashes();
-		FShaderHashes(const FGraphicsPipelineStateInitializer& PSOInitializer);
-
-		friend inline uint32 GetTypeHash(const FShaderHashes& Hashes)
-		{
-			return Hashes.Hash;
-		}
-
-		inline void Finalize()
-		{
-			Hash = FCrc::MemCrc32(Stages, sizeof(Stages));
-		}
-
-
-		friend inline bool operator == (const FShaderHashes& A, const FShaderHashes& B)
-		{
-			for (int32 Index = 0; Index < ShaderStage::NumStages; ++Index)
-			{
-				if (A.Stages[Index] != B.Stages[Index])
-				{
-					return false;
-				}
-			}
-
-			return true;
-		}
-	};
-
-	typedef TMap<uint32, FVulkanGfxPipeline*> FHashToGfxPipelinesMap;
-	TMap<FShaderHashes, FHashToGfxPipelinesMap> ShaderHashToGfxPipelineMap;
+	TMap<FGfxEntryKey, FVulkanGfxPipeline*> EntryKeyToGfxPipelineMap;
 
 	//@@
-	TMap<uint32, FVulkanLayout*> LayoutMapGfx;
 	TMap<FVulkanDescriptorSetsLayoutInfo, FVulkanLayout*> LayoutMap;
 	FCriticalSection LayoutMapCS;
 
-	FVulkanRHIGraphicsPipelineState* FindInLoadedLibrary(const FGraphicsPipelineStateInitializer& PSOInitializer, uint32 PSOInitializerHash, const FShaderHashes& ShaderHashes, FGfxPipelineEntry*& OutGfxEntry, FVulkanPipelineStateCacheManager::FHashToGfxPipelinesMap*& OutHashToGfxPipelinesMap);
+	FVulkanRHIGraphicsPipelineState* FindInLoadedLibrary(const FGraphicsPipelineStateInitializer& PSOInitializer, FGfxPSIKey& PSIKey, TGfxPipelineEntrySharedPtr& OutGfxEntry, FGfxEntryKey& OutGfxEntryKey);
 
 	friend class FVulkanDynamicRHI;
 
 	//@@
-	FVulkanLayout* FindOrAddLayoutForGfx(uint32 PSOHash, const FVulkanDescriptorSetsLayoutInfo& DescriptorSetLayoutInfo, bool bGfxLayout);
 	FVulkanLayout* FindOrAddLayout(const FVulkanDescriptorSetsLayoutInfo& DescriptorSetLayoutInfo, bool bGfxLayout);
 
-	FComputePipelineEntry* CreateComputeEntry(FVulkanComputeShader* ComputeShader);
-	FVulkanComputePipeline* CreateComputePipelineFromEntry(const FComputePipelineEntry* ComputeEntry);
-	void CreateComputeEntryRuntimeObjects(FComputePipelineEntry* GfxEntry);
+	FVulkanComputePipeline* CreateComputePipelineFromShader(FVulkanComputeShader* Shader);
 
 #if VULKAN_ENABLE_LRU_CACHE
 	struct FVulkanLRUCacheFile
@@ -992,30 +592,176 @@ private:
 	};
 #endif
 
-#if VULKAN_ENABLE_GENERIC_PIPELINE_CACHE_FILE
-	struct FVulkanPipelineStateCacheFile
+	static bool BinaryCacheMatches(FVulkanDevice* InDevice, const TArray<uint8>& DeviceCache);
+};
+
+// Common pipeline class
+class FVulkanPipeline
+{
+public:
+	FVulkanPipeline(FVulkanDevice* InDevice);
+	virtual ~FVulkanPipeline();
+
+	inline VkPipeline GetHandle() const
 	{
-		struct FFileHeader
+		return Pipeline;
+	}
+
+	inline const FVulkanLayout& GetLayout() const
+	{
+		return *Layout;
+	}
+#if VULKAN_ENABLE_LRU_CACHE
+	inline void DeleteVkPipeline(bool ImmediateDestroy)
+	{
+		if (Pipeline != VK_NULL_HANDLE)
 		{
-			int32 Version = -1;
-			int32 SizeOfGfxEntry = -1;
-			int32 SizeOfComputeEntry = -1;
-			int32 UncompressedSize = 0; // 0 == file is uncompressed
-		} Header;
+			if (ImmediateDestroy)
+			{
+				VulkanRHI::vkDestroyPipeline(Device->GetInstanceHandle(), Pipeline, VULKAN_CPU_ALLOCATOR);
+			}
+			else
+			{
+				Device->GetDeferredDeletionQueue().EnqueueResource(VulkanRHI::FDeferredDeletionQueue::EType::Pipeline, Pipeline);
+			}
+			Pipeline = VK_NULL_HANDLE;
+		}
+	}
+#endif
+protected:
+	FVulkanDevice* Device;
+	VkPipeline Pipeline;
+	FVulkanLayout* Layout; /*owned by FVulkanPipelineStateCacheManager, do not delete yourself !*/
 
-		FShaderUCodeCache::TDataMap* ShaderCache = nullptr;
+	friend class FVulkanPipelineStateCacheManager;
+	friend class FVulkanRHIGraphicsPipelineState;
+	friend class FVulkanComputePipelineDescriptorState;
+};
 
-		TArray<FGfxPipelineEntry*> GfxPipelineEntries;
-		TArray<FComputePipelineEntry*> ComputePipelineEntries;
+class FVulkanComputePipeline : public FVulkanPipeline, public FRHIComputePipelineState
+{
+public:
+	FVulkanComputePipeline(FVulkanDevice* InDevice);
+	virtual ~FVulkanComputePipeline();
 
-		void Save(FArchive& Ar);
-		bool Load(FArchive& Ar, const TCHAR* Filename);
+	inline const FVulkanShaderHeader& GetShaderCodeHeader() const
+	{
+		return ComputeShader->GetCodeHeader();
+	}
 
-		static const ECompressionFlags CompressionFlags = (ECompressionFlags)(COMPRESS_ZLIB | COMPRESS_BiasSpeed);
-	};
+	inline const FVulkanComputeShader* GetShader() const
+	{
+		return ComputeShader;
+	}
+
+	inline void Bind(VkCommandBuffer CmdBuffer)
+	{
+		VulkanRHI::vkCmdBindPipeline(CmdBuffer, VK_PIPELINE_BIND_POINT_COMPUTE, Pipeline);
+	}
+
+	inline const FVulkanComputeLayout& GetComputeLayout() const
+	{
+		return *(FVulkanComputeLayout*)Layout;
+	}
+
+protected:
+	FVulkanComputeShader*	ComputeShader;
+
+	friend class FVulkanPipelineStateCacheManager;
+};
+
+using TGfxPipelineEntrySharedPtr = TSharedPtr<FVulkanPipelineStateCacheManager::FGfxPipelineEntry, ESPMode::ThreadSafe>;
+
+class FVulkanGfxPipeline : public FVulkanPipeline, public FRHIResource
+{
+public:
+#if VULKAN_ENABLE_LRU_CACHE
+	FVulkanGfxPipeline(FVulkanDevice* InDevice, TGfxPipelineEntrySharedPtr InGfxPipelineEntry = {});
+#else
+	FVulkanGfxPipeline(FVulkanDevice* InDevice);
 #endif
 
-	static bool BinaryCacheMatches(FVulkanDevice* InDevice, const TArray<uint8>& DeviceCache);
+	inline void Bind(VkCommandBuffer CmdBuffer)
+	{
+#if VULKAN_ENABLE_LRU_CACHE
+		RecentFrame = GFrameNumberRenderThread;
+#endif
+		VulkanRHI::vkCmdBindPipeline(CmdBuffer, VK_PIPELINE_BIND_POINT_GRAPHICS, Pipeline);
+	}
+
+	inline bool IsRuntimeInitialized() const
+	{
+		return bRuntimeObjectsValid;
+	}
+
+	inline const FVulkanVertexInputStateInfo& GetVertexInputState() const
+	{
+		check(bRuntimeObjectsValid);
+		return VertexInputState;
+	}
+
+	void CreateRuntimeObjects(const FGraphicsPipelineStateInitializer& InPSOInitializer);
+
+	inline const FVulkanGfxLayout& GetGfxLayout() const
+	{
+		return *(FVulkanGfxLayout*)Layout;
+	}
+
+private:
+#if VULKAN_ENABLE_LRU_CACHE
+	TGfxPipelineEntrySharedPtr GfxPipelineEntry;
+	uint32 PipelineCacheSize;
+	uint32 RecentFrame;
+	// ID to LRU (if used) allows quick access when updating LRU status.
+	FSetElementId LRUNode;
+	bool bTrackedByLRU;
+	friend class FVulkanPipelineStateCacheManager;
+#endif
+	FVulkanVertexInputStateInfo VertexInputState;
+	bool						bRuntimeObjectsValid;
+};
+
+class FVulkanRHIGraphicsPipelineState : public FRHIGraphicsPipelineState
+{
+public:
+	FVulkanRHIGraphicsPipelineState(const FBoundShaderStateInput& InBSI, FVulkanGfxPipeline* InPipeline, EPrimitiveType InPrimitiveType)
+		: Pipeline(InPipeline)
+		, BSI(InBSI)
+		, PrimitiveType(InPrimitiveType)
+	{
+		bHasInputAttachments = InPipeline->GetGfxLayout().GetDescriptorSetsLayout().HasInputAttachments();
+	}
+
+	~FVulkanRHIGraphicsPipelineState();
+
+	inline void Bind(VkCommandBuffer CmdBuffer)
+	{
+		Pipeline->Bind(CmdBuffer);
+	}
+
+	inline const FVulkanShader* GetShader(EShaderFrequency Frequency) const
+	{
+		FVulkanShader* Shader = nullptr;
+
+		switch (Frequency)
+		{
+		case SF_Vertex:   Shader = ResourceCast(BSI.VertexShaderRHI); break;
+		case SF_Hull:     Shader = ResourceCast(BSI.HullShaderRHI); break;
+		case SF_Domain:   Shader = ResourceCast(BSI.DomainShaderRHI); break;
+		case SF_Pixel:    Shader = ResourceCast(BSI.PixelShaderRHI); break;
+		case SF_Geometry: Shader = ResourceCast(BSI.GeometryShaderRHI); break;
+		default:
+			check(0);
+			break;
+		}
+
+		return Shader;
+	}
+
+	TRefCountPtr<FVulkanGfxPipeline>	Pipeline;
+	FBoundShaderStateInput				BSI;
+	TEnumAsByte<EPrimitiveType>			PrimitiveType;
+	bool								bHasInputAttachments;
 };
 
 template<>
diff --git a/Engine/Source/Runtime/VulkanRHI/Private/VulkanPipelineState.cpp b/Engine/Source/Runtime/VulkanRHI/Private/VulkanPipelineState.cpp
--- a/Engine/Source/Runtime/VulkanRHI/Private/VulkanPipelineState.cpp
+++ b/Engine/Source/Runtime/VulkanRHI/Private/VulkanPipelineState.cpp
@@ -419,10 +419,9 @@ void FVulkanCommandListContext::RHISetGraphicsPipelineState(FGraphicsPipelineSta
 		return; // this happens when we immediately evict an cached PSO...the thing is never actually created
 	}
 	FVulkanPipelineStateCacheManager* PipelineStateCache = Device->GetPipelineStateCache();
-	PipelineStateCache->PipelineLRU.Touch(Device, Pipeline->Pipeline, PipelineStateCache->GfxPipelineEntries);
-#else
-	check(Pipeline);
+	PipelineStateCache->PipelineLRU.Touch(Pipeline);
 #endif
+	check(Pipeline);
 
 	FVulkanCmdBuffer* CmdBuffer = CommandBufferManager->GetActiveCmdBuffer();
 	if (PendingGfxState->SetGfxPipeline(Pipeline) || !CmdBuffer->bHasPipeline)
diff --git a/Engine/Source/Runtime/VulkanRHI/Private/VulkanRHI.cpp b/Engine/Source/Runtime/VulkanRHI/Private/VulkanRHI.cpp
--- a/Engine/Source/Runtime/VulkanRHI/Private/VulkanRHI.cpp
+++ b/Engine/Source/Runtime/VulkanRHI/Private/VulkanRHI.cpp
@@ -16,6 +16,7 @@
 #include "VulkanPipelineState.h"
 #include "Misc/FileHelper.h"
 #include "VulkanLLM.h"
+#include "Misc/EngineVersion.h"
 #include "GlobalShader.h"
 
 extern RHI_API bool GUseTexture3DBulkDataRHI;
@@ -187,6 +188,10 @@ void FVulkanDynamicRHI::Init()
 	LLM(VulkanLLM::Initialize());
 #endif
 
+
+	static const auto CVarStreamingTexturePoolSize = IConsoleManager::Get().FindTConsoleVariableDataInt(TEXT("r.Streaming.PoolSize"));
+	int32 StreamingPoolSizeValue = CVarStreamingTexturePoolSize->GetValueOnAnyThread();
+			
 	if (GPoolSizeVRAMPercentage > 0)
 	{
 		const uint64 TotalGPUMemory = Device->GetMemoryManager().GetTotalMemory(true);
@@ -201,6 +206,15 @@ void FVulkanDynamicRHI::Init()
 			GPoolSizeVRAMPercentage,
 			TotalGPUMemory / 1024 / 1024);
 	}
+	else if (StreamingPoolSizeValue > 0)
+	{
+		GTexturePoolSize = (int64)StreamingPoolSizeValue * 1024 * 1024;
+
+		const uint64 TotalGPUMemory = Device->GetMemoryManager().GetTotalMemory(true);
+		UE_LOG(LogRHI,Log,TEXT("Texture pool is %llu MB (of %llu MB total graphics mem)"),
+				GTexturePoolSize / 1024 / 1024,
+				TotalGPUMemory / 1024 / 1024);
+	}
 }
 
 void FVulkanDynamicRHI::PostInit()
@@ -300,12 +314,19 @@ void FVulkanDynamicRHI::CreateInstance()
 	auto* CVarDisableEngineAndAppRegistration = IConsoleManager::Get().FindTConsoleVariableDataInt(TEXT("r.DisableEngineAndAppRegistration"));
 	bool bDisableEngineRegistration = (CVarDisableEngineAndAppRegistration && CVarDisableEngineAndAppRegistration->GetValueOnAnyThread() != 0) ||
 		(CVarShaderDevelopmentMode && CVarShaderDevelopmentMode->GetValueOnAnyThread() != 0);
+
+	// EngineName will be of the form "UnrealEngine4.21", with the minor version ("21" in this example)
+	// updated with every quarterly release
+	FString EngineName = FApp::GetEpicProductIdentifier() + FEngineVersion::Current().ToString(EVersionComponent::Minor);
+	FTCHARToUTF8 EngineNameConverter(*EngineName);
+	FTCHARToUTF8 ProjectNameConverter(FApp::GetProjectName());
+
 	VkApplicationInfo AppInfo;
 	ZeroVulkanStruct(AppInfo, VK_STRUCTURE_TYPE_APPLICATION_INFO);
-	AppInfo.pApplicationName = bDisableEngineRegistration ? "" : "UE4";
-	//AppInfo.applicationVersion = 0;
-	AppInfo.pEngineName = bDisableEngineRegistration ? "" : "UE4";
-	AppInfo.engineVersion = 15;
+	AppInfo.pApplicationName = bDisableEngineRegistration ? nullptr : ProjectNameConverter.Get();
+	AppInfo.applicationVersion = 0;	// Do we want FApp::GetBuildVersion() ?
+	AppInfo.pEngineName = bDisableEngineRegistration ? nullptr : EngineNameConverter.Get();
+	AppInfo.engineVersion = FEngineVersion::Current().GetMinor();
 	AppInfo.apiVersion = UE_VK_API_VERSION;
 
 	VkInstanceCreateInfo InstInfo;
@@ -553,6 +574,7 @@ void FVulkanDynamicRHI::SelectAndInitDevice()
 		static_assert(sizeof(NvidiaVersion) == sizeof(Props.driverVersion), "Mismatched Nvidia pack driver version!");
 		NvidiaVersion.Packed = Props.driverVersion;
 		GRHIAdapterUserDriverVersion = FString::Printf(TEXT("%d.%d"), NvidiaVersion.Major, NvidiaVersion.Minor);
+		UE_LOG(LogVulkanRHI, Display, TEXT("Nvidia User Driver Version = %s"), *GRHIAdapterUserDriverVersion);
 
 		// Ignore GRHIAdapterInternalDriverVersion for now as the device name doesn't match
 	}
@@ -1087,6 +1109,33 @@ void FVulkanDescriptorSetsLayoutInfo::AddDescriptor(int32 DescriptorSetIndex, co
 	}
 }
 
+void FVulkanDescriptorSetsLayoutInfo::GenerateHash()
+{
+	const int32 LayoutCount = SetLayouts.Num();
+	Hash = FCrc::MemCrc32(&TypesUsageID, sizeof(uint32), LayoutCount);
+
+	for (int32 layoutIndex = 0; layoutIndex < LayoutCount; ++layoutIndex)
+	{
+		TArray<VkDescriptorSetLayoutBinding>& DescSetLayout = SetLayouts[layoutIndex].LayoutBindings;
+		Hash = FCrc::MemCrc32(DescSetLayout.GetData(), sizeof(VkDescriptorSetLayoutBinding) * DescSetLayout.Num(), Hash);
+	}
+
+	for (uint32 RemapingIndex = 0; RemapingIndex < ShaderStage::NumStages; ++RemapingIndex)
+	{
+		Hash = FCrc::MemCrc32(&RemappingInfo.StageInfos[RemapingIndex].PackedUBDescriptorSet, sizeof(uint16), Hash);
+		Hash = FCrc::MemCrc32(&RemappingInfo.StageInfos[RemapingIndex].Pad0, sizeof(uint16), Hash);
+
+		TArray<FDescriptorSetRemappingInfo::FRemappingInfo>& Globals = RemappingInfo.StageInfos[RemapingIndex].Globals;
+		Hash = FCrc::MemCrc32(Globals.GetData(), sizeof(FDescriptorSetRemappingInfo::FRemappingInfo) * Globals.Num(), Hash);
+
+		TArray<FDescriptorSetRemappingInfo::FUBRemappingInfo>& UniformBuffers = RemappingInfo.StageInfos[RemapingIndex].UniformBuffers;
+		Hash = FCrc::MemCrc32(UniformBuffers.GetData(), sizeof(FDescriptorSetRemappingInfo::FRemappingInfo) * UniformBuffers.Num(), Hash);
+
+		TArray<uint16>& PackedUBBindingIndices = RemappingInfo.StageInfos[RemapingIndex].PackedUBBindingIndices;
+		Hash = FCrc::MemCrc32(PackedUBBindingIndices.GetData(), sizeof(uint16) * PackedUBBindingIndices.Num(), Hash);
+	}
+}
+
 void FVulkanDescriptorSetsLayoutInfo::CompileTypesUsageID()
 {
 	static TMap<uint32, uint32> GTypesUsageHashMap;
diff --git a/Engine/Source/Runtime/VulkanRHI/Private/VulkanRenderTarget.cpp b/Engine/Source/Runtime/VulkanRHI/Private/VulkanRenderTarget.cpp
--- a/Engine/Source/Runtime/VulkanRHI/Private/VulkanRenderTarget.cpp
+++ b/Engine/Source/Runtime/VulkanRHI/Private/VulkanRenderTarget.cpp
@@ -514,6 +514,54 @@ void FVulkanCommandListContext::RHISetRenderTargets(uint32 NumSimultaneousRender
 	}
 }
 
+// Find out whether we can re-use current renderpass instead of starting new one
+static bool IsCompatibleRenderPass(FVulkanRenderPass* CurrentRenderPass, FVulkanRenderPass* NewRenderPass)
+{
+	if (CurrentRenderPass == nullptr)
+	{
+		return false;
+	}
+
+	const FVulkanRenderTargetLayout& CurrentLayout = CurrentRenderPass->GetLayout();
+	const FVulkanRenderTargetLayout& NewLayout = NewRenderPass->GetLayout();
+	
+	if (CurrentLayout.GetRenderPassCompatibleHash() != NewLayout.GetRenderPassCompatibleHash())
+	{
+		return false;
+	}
+	
+	int32 NumDesc = CurrentLayout.GetNumAttachmentDescriptions();
+	check(NumDesc == NewLayout.GetNumAttachmentDescriptions());
+
+	const VkAttachmentDescription* CurrentDescriptions = CurrentLayout.GetAttachmentDescriptions();
+	const VkAttachmentDescription* NewDescriptions = NewLayout.GetAttachmentDescriptions();
+	for (int32 i = 0; i < NumDesc; ++i)
+	{
+		const VkAttachmentDescription& CurrentDesc = CurrentDescriptions[i];
+		const VkAttachmentDescription& NewDesc = NewDescriptions[i];
+		
+		// New render-pass wants a clear target
+		if (NewDesc.loadOp == VK_ATTACHMENT_LOAD_OP_CLEAR || NewDesc.stencilLoadOp == VK_ATTACHMENT_LOAD_OP_CLEAR)
+		{
+			return false;
+		}
+
+		// New render-pass wants to store, while current does not
+		if ((NewDesc.storeOp == VK_ATTACHMENT_STORE_OP_STORE && CurrentDesc.storeOp == VK_ATTACHMENT_STORE_OP_DONT_CARE) ||
+			(NewDesc.stencilStoreOp == VK_ATTACHMENT_STORE_OP_STORE && CurrentDesc.stencilStoreOp == VK_ATTACHMENT_STORE_OP_DONT_CARE))
+		{
+			return false;
+		}
+
+		if (NewDesc.finalLayout != CurrentDesc.finalLayout)
+		{
+			return false;
+		}
+	}
+	
+	return true;
+}
+
 void FVulkanCommandListContext::RHISetRenderTargetsAndClear(const FRHISetRenderTargetsInfo& RenderTargetsInfo)
 {
 	FVulkanRenderTargetLayout RTLayout(*Device, RenderTargetsInfo);
@@ -529,7 +577,7 @@ void FVulkanCommandListContext::RHISetRenderTargetsAndClear(const FRHISetRenderT
 		Framebuffer = TransitionAndLayoutManager.GetOrCreateFramebuffer(*Device, RenderTargetsInfo, RTLayout, RenderPass);
 	}
 
-	if (Framebuffer == TransitionAndLayoutManager.CurrentFramebuffer && RenderPass == TransitionAndLayoutManager.CurrentRenderPass)
+	if (Framebuffer == TransitionAndLayoutManager.CurrentFramebuffer && IsCompatibleRenderPass(TransitionAndLayoutManager.CurrentRenderPass, RenderPass))
 	{
 		return;
 	}
@@ -576,6 +624,21 @@ void FVulkanCommandListContext::RHICopyToResolveTarget(FTextureRHIParamRef Sourc
 		return;
 	}
 
+	if (SourceTextureRHI == DestTextureRHI)
+	{
+		FRHITexture2D* SourceTexture2D = SourceTextureRHI->GetTexture2D();
+		if (SourceTexture2D)
+		{
+			FVulkanTexture2D* VulkanSourceTexture2D  = (FVulkanTexture2D*)SourceTexture2D;
+			if (VulkanSourceTexture2D->GetBackBuffer() != nullptr)
+			{
+				// skip Backbuffer implicit transition to Readable, to avoid splitting Post->UI renderpass
+				// do explicit transition when need to read from Backbuffer
+				return;
+			}
+		}
+	}
+
 	RHITransitionResources(EResourceTransitionAccess::EReadable, &SourceTextureRHI, 1);
 
 	auto CopyImage = [](FTransitionAndLayoutManager& InRenderPassState, FVulkanCmdBuffer* InCmdBuffer, FVulkanSurface& SrcSurface, FVulkanSurface& DstSurface, uint32 SrcNumLayers, uint32 DstNumLayers, const FResolveParams& ResolveParams)
@@ -1271,18 +1334,6 @@ void FVulkanCommandListContext::TransitionResources(const FPendingTransition& Pe
 		}
 		else if (PendingTransition.TransitionType == EResourceTransitionAccess::EWritable)
 		{
-			//#todo-rco: Until render passes come online, assume writable means end render pass
-			if (TransitionAndLayoutManager.CurrentRenderPass)
-			{
-				TransitionAndLayoutManager.EndEmulatedRenderPass(CmdBuffer);
-				if (GVulkanSubmitAfterEveryEndRenderPass)
-				{
-					CommandBufferManager->SubmitActiveCmdBuffer();
-					CommandBufferManager->PrepareForNewActiveCommandBuffer();
-					CmdBuffer = CommandBufferManager->GetActiveCmdBuffer();
-				}
-			}
-
 			if (bShowTransitionEvents)
 			{
 				for (int32 i = 0; i < PendingTransition.Textures.Num(); ++i)
@@ -1299,34 +1350,44 @@ void FVulkanCommandListContext::TransitionResources(const FPendingTransition& Pe
 				FVulkanSurface& Surface = FVulkanTextureBase::Cast(PendingTransition.Textures[Index])->Surface;
 
 				const VkImageAspectFlags AspectMask = Surface.GetFullAspectMask();
-				VkImageSubresourceRange SubresourceRange = {AspectMask, 0, Surface.GetNumMips(), 0, Surface.GetNumberOfArrayLevels()};
-
 				VkImageLayout& SrcLayout = TransitionAndLayoutManager.FindOrAddLayoutRW(Surface.Image, VK_IMAGE_LAYOUT_UNDEFINED);
-				EImageLayoutBarrier DestLayout = VulkanRHI::EImageLayoutBarrier::Undefined;
-
+				
+				VkImageLayout FinalLayout;
 				if ((AspectMask & VK_IMAGE_ASPECT_COLOR_BIT) != 0)
 				{
-					VkImageLayout FinalLayout = (Surface.UEFlags & TexCreate_RenderTargetable) ? VK_IMAGE_LAYOUT_COLOR_ATTACHMENT_OPTIMAL : VK_IMAGE_LAYOUT_GENERAL;
-					if (SrcLayout != FinalLayout)
-					{
-						//#todo-rco: Switch to pending barrier
-						VulkanSetImageLayout(CmdBuffer->GetHandle(), Surface.Image, SrcLayout, FinalLayout, SubresourceRange);
-						SrcLayout = FinalLayout;
-					}
+					FinalLayout = (Surface.UEFlags & TexCreate_RenderTargetable) ? VK_IMAGE_LAYOUT_COLOR_ATTACHMENT_OPTIMAL : VK_IMAGE_LAYOUT_GENERAL;
 				}
 				else
 				{
-					if (SrcLayout != VK_IMAGE_LAYOUT_DEPTH_STENCIL_ATTACHMENT_OPTIMAL)
+					check(Surface.IsDepthOrStencilAspect());
+					FinalLayout = VK_IMAGE_LAYOUT_DEPTH_STENCIL_ATTACHMENT_OPTIMAL;
+				}
+
+				if (SrcLayout != FinalLayout)
+				{
+					int32 BarrierIndex = Barrier.AddImageBarrier(Surface.Image, AspectMask, Surface.GetNumMips(), Surface.GetNumberOfArrayLevels());
+					//todo: TransitionAndLayoutManager should use EImageLayoutBarrier type?
+					Barrier.SetTransition(BarrierIndex, VulkanRHI::GetImageLayoutFromVulkanLayout(SrcLayout), VulkanRHI::GetImageLayoutFromVulkanLayout(FinalLayout));
+					SrcLayout = FinalLayout;
+				}
+			}
+
+			if (Barrier.NumImageBarriers() > 0)
+			{
+				//#todo-rco: Until render passes come online, assume writable means end render pass
+				if (TransitionAndLayoutManager.CurrentRenderPass)
+				{
+					TransitionAndLayoutManager.EndEmulatedRenderPass(CmdBuffer);
+					if (GVulkanSubmitAfterEveryEndRenderPass)
 					{
-						check(Surface.IsDepthOrStencilAspect());
-						//#todo-rco: Switch to pending barrier
-						VulkanSetImageLayout(CmdBuffer->GetHandle(), Surface.Image, SrcLayout, VK_IMAGE_LAYOUT_DEPTH_STENCIL_ATTACHMENT_OPTIMAL, SubresourceRange);
-						SrcLayout = VK_IMAGE_LAYOUT_DEPTH_STENCIL_ATTACHMENT_OPTIMAL;
+						CommandBufferManager->SubmitActiveCmdBuffer();
+						CommandBufferManager->PrepareForNewActiveCommandBuffer();
+						CmdBuffer = CommandBufferManager->GetActiveCmdBuffer();
 					}
 				}
-			}
 
-			Barrier.Execute(CmdBuffer);
+				Barrier.Execute(CmdBuffer);
+			}
 		}
 		else if (PendingTransition.TransitionType == EResourceTransitionAccess::ERWSubResBarrier)
 		{
diff --git a/Engine/Source/Runtime/VulkanRHI/Private/VulkanShaders.cpp b/Engine/Source/Runtime/VulkanRHI/Private/VulkanShaders.cpp
--- a/Engine/Source/Runtime/VulkanRHI/Private/VulkanShaders.cpp
+++ b/Engine/Source/Runtime/VulkanRHI/Private/VulkanShaders.cpp
@@ -50,12 +50,13 @@ FVulkanShaderFactory::~FVulkanShaderFactory()
 template <typename ShaderType> 
 ShaderType* FVulkanShaderFactory::CreateShader(const TArray<uint8>& Code, FVulkanDevice* Device)
 {
+	uint32 ShaderCodeLen = Code.Num();
+	uint32 ShaderCodeCRC = FCrc::MemCrc32(Code.GetData(), Code.Num());
+	uint64 ShaderKey = ((uint64)ShaderCodeLen | ((uint64)ShaderCodeCRC << 32));
+
 	ShaderType* RetShader = nullptr;
 	if (GCacheCreatedShaders)
 	{
-		FVulkanShaderFactory::FKey ShaderKey; 
-		ShaderKey.CodeLen = Code.Num();
-		ShaderKey.CodeCRC = FCrc::MemCrc32(Code.GetData(), Code.Num());
 		FVulkanShader** FoundShaderPtr = nullptr;
 		{
 			FRWScopeLock ScopedLock(Lock, SLT_ReadOnly);
@@ -68,9 +69,7 @@ ShaderType* FVulkanShaderFactory::CreateShader(const TArray<uint8>& Code, FVulka
 		else
 		{
 			RetShader = new ShaderType(Device);
-			RetShader->Setup(Code);
-			RetShader->ShaderCodeCRC = ShaderKey.CodeCRC;
-			RetShader->ShaderCodeLen = ShaderKey.CodeLen;
+			RetShader->Setup(Code, ShaderKey);
 			
 			FRWScopeLock ScopedLock(Lock, SLT_Write);
 			ShaderMap[ShaderType::StaticFrequency].Add(ShaderKey, RetShader);
@@ -79,7 +78,7 @@ ShaderType* FVulkanShaderFactory::CreateShader(const TArray<uint8>& Code, FVulka
 	else
 	{
 		RetShader = new ShaderType(Device);
-		RetShader->Setup(Code);
+		RetShader->Setup(Code, ShaderKey);
 	}
 
 	return RetShader;
@@ -90,22 +89,22 @@ void FVulkanShaderFactory::OnDeleteShader(const FVulkanShader& Shader)
 	if (GCacheCreatedShaders)
 	{
 		FRWScopeLock ScopedLock(Lock, SLT_Write);
-		FVulkanShaderFactory::FKey ShaderKey; 
-		ShaderKey.CodeLen = Shader.ShaderCodeLen;
-		ShaderKey.CodeCRC = Shader.ShaderCodeCRC;
+		uint64 ShaderKey = Shader.GetShaderKey(); 
 		ShaderMap[Shader.Frequency].Remove(ShaderKey);
 	}
 }
 
-void FVulkanShader::Setup(const TArray<uint8>& InShaderHeaderAndCode)
+void FVulkanShader::Setup(const TArray<uint8>& InShaderHeaderAndCode, uint64 InShaderKey)
 {
 	LLM_SCOPE_VULKAN(ELLMTagVulkan::VulkanShaders);
 	check(Device);
 
+	ShaderKey = InShaderKey;
+
 	FMemoryReader Ar(InShaderHeaderAndCode, true);
 
 	Ar << CodeHeader;
-	
+
 	Ar << Spirv;
 #if VULKAN_ENABLE_SHADER_DEBUG_NAMES
 	checkf(Spirv.Num() != 0, TEXT("Empty SPIR-V!%s"), *CodeHeader.DebugName);
@@ -122,9 +121,6 @@ void FVulkanShader::Setup(const TArray<uint8>& InShaderHeaderAndCode)
 		checkSlow(CodeHeader.UniformBufferSpirvInfos.Num() == 0);
 	}
 	check(CodeHeader.GlobalSpirvInfos.Num() == CodeHeader.Globals.Num());
-
-	static TAtomic<uint64> IdCounter{ 0 };
-	Id = ++IdCounter;
 }
 
 VkShaderModule FVulkanShader::CreateHandle(const FVulkanLayout* Layout, uint32 LayoutHash)
@@ -543,9 +539,11 @@ void FVulkanDescriptorSetsLayoutInfo::FinalizeBindings(const FUniformBufferGathe
 				}
 			}
 		}
-
 	}
 
+	CompileTypesUsageID();
+	GenerateHash();
+
 	// Validate no empty sets were made
 	for (int32 Index = 0; Index < RemappingInfo.SetInfos.Num(); ++Index)
 	{
@@ -584,7 +582,7 @@ void FVulkanComputePipelineDescriptorInfo::Initialize(const FDescriptorSetRemapp
 	bInitialized = true;
 }
 
-void FVulkanGfxPipelineDescriptorInfo::Initialize(const FDescriptorSetRemappingInfo& InRemappingInfo, FVulkanShader** Shaders)
+void FVulkanGfxPipelineDescriptorInfo::Initialize(const FDescriptorSetRemappingInfo& InRemappingInfo)
 {
 	check(!bInitialized);
 
diff --git a/Engine/Source/Runtime/VulkanRHI/Private/VulkanSwapChain.cpp b/Engine/Source/Runtime/VulkanRHI/Private/VulkanSwapChain.cpp
--- a/Engine/Source/Runtime/VulkanRHI/Private/VulkanSwapChain.cpp
+++ b/Engine/Source/Runtime/VulkanRHI/Private/VulkanSwapChain.cpp
@@ -334,13 +334,13 @@ FVulkanSwapChain::FVulkanSwapChain(VkInstance InInstance, FVulkanDevice& InDevic
 			// Until FVulkanViewport::Present honors SyncInterval, we need to disable vsync for the spectator window if using an HMD.
 			const bool bDisableVsyncForHMD = (FVulkanDynamicRHI::HMDVulkanExtensions.IsValid()) ? FVulkanDynamicRHI::HMDVulkanExtensions->ShouldDisableVulkanVSync() : false;
 
-			if (bFoundPresentModeMailbox && LockToVsync)
+			if (bFoundPresentModeImmediate && (bDisableVsyncForHMD || !LockToVsync))
 			{
-				PresentMode = VK_PRESENT_MODE_MAILBOX_KHR;
+				PresentMode = VK_PRESENT_MODE_IMMEDIATE_KHR;
 			}
-			else if (bFoundPresentModeImmediate && (bDisableVsyncForHMD || !LockToVsync))
+			else if (bFoundPresentModeMailbox)
 			{
-				PresentMode = VK_PRESENT_MODE_IMMEDIATE_KHR;
+				PresentMode = VK_PRESENT_MODE_MAILBOX_KHR;
 			}
 			else if (bFoundPresentModeFIFO)
 			{
@@ -582,10 +582,10 @@ int32 FVulkanSwapChain::AcquireImageIndex(VulkanRHI::FSemaphore** OutSemaphore)
 void FVulkanSwapChain::RenderThreadPacing()
 {
 	check(IsInRenderingThread());
-	const int32 SyncInterval = RHIGetSyncInterval();
+	const int32 SyncInterval = LockToVsync ? RHIGetSyncInterval() : 0;
 
 	//very naive CPU side frame pacer.
-	if (GVulkanCPURenderThreadFramePacer)
+	if (GVulkanCPURenderThreadFramePacer && SyncInterval > 0)
 	{
 		static double PreviousFrameCPUTime = 0;
 		static double SampledDeltaTimeMS = 0;
@@ -658,7 +658,7 @@ FVulkanSwapChain::EStatus FVulkanSwapChain::Present(FVulkanQueue* GfxQueue, FVul
 	Info.pSwapchains = &SwapChain;
 	Info.pImageIndices = (uint32*)&CurrentImageIndex;
 
-	const int32 SyncInterval = RHIGetSyncInterval();
+	const int32 SyncInterval = LockToVsync ? RHIGetSyncInterval() : 0;
 	ensureMsgf(SyncInterval <= 3 && SyncInterval >= 0, TEXT("Unsupported sync interval: %i"), SyncInterval);
 	FVulkanPlatform::EnablePresentInfoExtensions(Info);
 
@@ -848,7 +848,7 @@ FVulkanSwapChain::EStatus FVulkanSwapChain::Present(FVulkanQueue* GfxQueue, FVul
 #endif
 
 	//very naive CPU side frame pacer.
-	if (GVulkanCPURHIFramePacer)
+	if (GVulkanCPURHIFramePacer && SyncInterval > 0)
 	{
 		static double PreviousFrameCPUTime = 0;
 		double NowCPUTime = FPlatformTime::Seconds();
diff --git a/Engine/Source/Runtime/VulkanRHI/Public/VulkanConfiguration.h b/Engine/Source/Runtime/VulkanRHI/Public/VulkanConfiguration.h
--- a/Engine/Source/Runtime/VulkanRHI/Public/VulkanConfiguration.h
+++ b/Engine/Source/Runtime/VulkanRHI/Public/VulkanConfiguration.h
@@ -108,10 +108,6 @@
 	#define VULKAN_ENABLE_LRU_CACHE								0
 #endif
 
-#ifndef VULKAN_ENABLE_GENERIC_PIPELINE_CACHE_FILE
-	#define VULKAN_ENABLE_GENERIC_PIPELINE_CACHE_FILE			0
-#endif
-
 #ifdef VK_KHR_maintenance1
 	#define VULKAN_SUPPORTS_MAINTENANCE_LAYER1					1
 #else
diff --git a/Engine/Source/Runtime/VulkanRHI/Public/VulkanMemory.h b/Engine/Source/Runtime/VulkanRHI/Public/VulkanMemory.h
--- a/Engine/Source/Runtime/VulkanRHI/Public/VulkanMemory.h
+++ b/Engine/Source/Runtime/VulkanRHI/Public/VulkanMemory.h
@@ -1170,6 +1170,8 @@ namespace VulkanRHI
 		{
 			ReleaseResources(true);
 		}
+
+		void OnCmdBufferDeleted(FVulkanCmdBuffer* CmdBuffer);
 /*
 		class FVulkanAsyncDeletionWorker : public FVulkanDeviceChild, FNonAbandonableTask
 		{
@@ -1473,6 +1475,16 @@ namespace VulkanRHI
 		void InnerExecute(FVulkanCmdBuffer* CmdBuffer, bool bEnsure);
 
 	public:
+		inline int32 NumImageBarriers() const
+		{
+			return ImageBarriers.Num();
+		}
+
+		inline int32 NumBufferBarriers() const
+		{
+			return BufferBarriers.Num();
+		}
+
 		inline void ResetStages()
 		{
 			SourceStage = 0;
diff --git a/Engine/Source/Runtime/VulkanRHI/Public/VulkanResources.h b/Engine/Source/Runtime/VulkanRHI/Public/VulkanResources.h
--- a/Engine/Source/Runtime/VulkanRHI/Public/VulkanResources.h
+++ b/Engine/Source/Runtime/VulkanRHI/Public/VulkanResources.h
@@ -74,12 +74,10 @@ class FVulkanShader : public IRefCountedObject
 {
 public:
 	FVulkanShader(FVulkanDevice* InDevice, EShaderFrequency InFrequency, VkShaderStageFlagBits InStageFlag)
-		: Id(0)
+		: ShaderKey(0)
 		, StageFlag(InStageFlag)
 		, Frequency(InFrequency)
 		, Device(InDevice)
-		, ShaderCodeCRC(0)
-		, ShaderCodeLen(0)
 	{
 	}
 
@@ -87,7 +85,7 @@ public:
 
 	void PurgeShaderModules();
 
-	void Setup(const TArray<uint8>& InShaderHeaderAndCode);
+	void Setup(const TArray<uint8>& InShaderHeaderAndCode, uint64 InShaderKey);
 
 	VkShaderModule GetOrCreateHandle(const FVulkanLayout* Layout, uint32 LayoutHash)
 	{
@@ -112,13 +110,13 @@ public:
 		return CodeHeader;
 	}
 
-	inline uint64 GetId() const
+	inline uint64 GetShaderKey() const
 	{
-		return Id;
+		return ShaderKey;
 	}
 
 protected:
-	uint64							Id;
+	uint64							ShaderKey;
 
 	/** External bindings for this shader. */
 	FVulkanShaderHeader				CodeHeader;
@@ -130,10 +128,6 @@ protected:
 
 	FVulkanDevice*					Device;
 
-	// Cached shader key information
-	uint32							ShaderCodeCRC;
-	uint32							ShaderCodeLen;
-
 	VkShaderModule CreateHandle(const FVulkanLayout* Layout, uint32 LayoutHash);
 
 	friend class FVulkanCommandListContext;
@@ -189,24 +183,8 @@ public:
 	void OnDeleteShader(const FVulkanShader& Shader);
 
 private:
-	struct FKey
-	{
-		uint32 CodeLen;
-		uint32 CodeCRC;
-	};
-
-	friend bool operator ==(const FVulkanShaderFactory::FKey& A, const FVulkanShaderFactory::FKey& B)
-	{
-		return A.CodeLen == B.CodeLen && A.CodeCRC == B.CodeCRC;
-	}
-
-	friend uint32 GetTypeHash(const FVulkanShaderFactory::FKey &Key)
-	{
-		return HashCombine(GetTypeHash(Key.CodeLen), GetTypeHash(Key.CodeCRC));
-	}
-	
 	FRWLock Lock;
-	TMap<FKey, FVulkanShader*> ShaderMap[SF_NumFrequencies];
+	TMap<uint64, FVulkanShader*> ShaderMap[SF_NumFrequencies];
 };
 
 class FVulkanBoundShaderState : public FRHIBoundShaderState
@@ -378,7 +356,10 @@ public:
 		{
 			return ResourceAllocation->GetHandle();
 		}
-		return VK_NULL_HANDLE;
+		else
+		{
+			return VK_NULL_HANDLE;
+		}
 	}
 
 	inline uint64 GetAllocationOffset() const
@@ -1390,8 +1371,8 @@ public:
 		}
 
 		OutPackedUniformBufferStagingMask = ((uint64)1 << (uint64)InCodeHeader.PackedUBs.Num()) - 1;
-		EmulatedUBsCopyInfo = &InCodeHeader.EmulatedUBsCopyInfo;
-		EmulatedUBsCopyRanges = &InCodeHeader.EmulatedUBCopyRanges;
+		EmulatedUBsCopyInfo = InCodeHeader.EmulatedUBsCopyInfo;
+		EmulatedUBsCopyRanges = InCodeHeader.EmulatedUBCopyRanges;
 	}
 
 	inline void SetPackedGlobalParameter(uint32 BufferIndex, uint32 ByteOffset, uint32 NumBytes, const void* RESTRICT NewValue, uint64& InOutPackedUniformBufferStagingDirty)
@@ -1417,15 +1398,15 @@ public:
 	inline void SetEmulatedUniformBufferIntoPacked(uint32 BindPoint, const TArray<uint8>& ConstantData, uint64& NEWPackedUniformBufferStagingDirty)
 	{
 		// Emulated UBs. Assumes UniformBuffersCopyInfo table is sorted by CopyInfo.SourceUBIndex
-		if (BindPoint < (uint32)EmulatedUBsCopyRanges->Num())
+		if (BindPoint < (uint32)EmulatedUBsCopyRanges.Num())
 		{
-			uint32 Range = (*EmulatedUBsCopyRanges)[BindPoint];
+			uint32 Range = EmulatedUBsCopyRanges[BindPoint];
 			uint16 Start = (Range >> 16) & 0xffff;
 			uint16 Count = Range & 0xffff;
 			const uint8* RESTRICT SourceData = ConstantData.GetData();
 			for (int32 Index = Start; Index < Start + Count; ++Index)
 			{
-				const CrossCompiler::FUniformBufferCopyInfo& CopyInfo = (*EmulatedUBsCopyInfo)[Index];
+				const CrossCompiler::FUniformBufferCopyInfo& CopyInfo = EmulatedUBsCopyInfo[Index];
 				check(CopyInfo.SourceUBIndex == BindPoint);
 				FPackedBuffer& StagingBuffer = PackedUniformBuffers[(int32)CopyInfo.DestUBIndex];
 				//check(ByteOffset + NumBytes <= (uint32)StagingBuffer.Num());
@@ -1451,9 +1432,9 @@ public:
 protected:
 	TArray<FPackedBuffer>									PackedUniformBuffers;
 
-	// Pointers to Shader Code Header
-	const TArray<CrossCompiler::FUniformBufferCopyInfo>*	EmulatedUBsCopyInfo;
-	const TArray<uint32>*									EmulatedUBsCopyRanges;
+	// Copies to Shader Code Header (shaders may be deleted when we use this object again)
+	TArray<CrossCompiler::FUniformBufferCopyInfo>			EmulatedUBsCopyInfo;
+	TArray<uint32>											EmulatedUBsCopyRanges;
 };
 
 class FVulkanStagingBuffer : public FRHIStagingBuffer
diff --git a/Engine/Source/Runtime/VulkanRHI/Public/VulkanUtil.h b/Engine/Source/Runtime/VulkanRHI/Public/VulkanUtil.h
--- a/Engine/Source/Runtime/VulkanRHI/Public/VulkanUtil.h
+++ b/Engine/Source/Runtime/VulkanRHI/Public/VulkanUtil.h
@@ -8,6 +8,8 @@
 
 #include "GPUProfiler.h"
 
+#include "Serialization/MemoryWriter.h"
+
 class FVulkanCmdBuffer;
 class FVulkanRenderQuery;
 class FVulkanRenderQuery;
@@ -230,3 +232,130 @@ inline bool CopyAndReturnNotEqual(T& A, T B)
 	A = B;
 	return bOut;
 }
+
+template <int Version>
+class TDataKeyBase;
+
+template <>
+class TDataKeyBase<0>
+{
+protected:
+	template <class TDataWriter>
+	void UpdateData(TDataWriter&& WriteToData)
+	{
+		TArray<uint8> TempData;
+		WriteToData(TempData);
+	}
+
+	void CopyDataDeep(TDataKeyBase& Result) const {}
+	void CopyDataShallow(TDataKeyBase& Result) const {}
+	bool IsDataEquals(const TDataKeyBase& Other) const { return true; }
+
+protected:
+	uint32 Hash = 0;
+};
+
+template <>
+class TDataKeyBase<1>
+{
+protected:
+	template <class TDataWriter>
+	void UpdateData(TDataWriter&& WriteToData)
+	{
+		if (!DataStorage)
+		{
+			DataStorage = MakeUnique<TArray<uint8>>();
+			Data = DataStorage.Get();
+		}
+		WriteToData(*Data);
+	}
+
+	void CopyDataDeep(TDataKeyBase& Result) const
+	{
+		check(Data);
+		Result.DataStorage = MakeUnique<TArray<uint8>>(*Data);
+		Result.Data = Result.DataStorage.Get();
+	}
+
+	void CopyDataShallow(TDataKeyBase& Result) const
+	{
+		check(Data);
+		Result.Data = Data;
+	}
+
+	bool IsDataEquals(const TDataKeyBase& Other) const
+	{
+		check(Data && Other.Data);
+		check(Data->Num() == Other.Data->Num());
+		check(FMemory::Memcmp(Data->GetData(), Other.Data->GetData(), Data->Num()) == 0);
+		return true;
+	}
+
+protected:
+	uint32 Hash = 0;
+	TArray<uint8> *Data = nullptr;
+private:
+	TUniquePtr<TArray<uint8>> DataStorage;
+};
+
+template <>
+class TDataKeyBase<2> : public TDataKeyBase<1>
+{
+protected:
+	bool IsDataEquals(const TDataKeyBase& Other) const
+	{
+		check(Data && Other.Data);
+		return ((Data->Num() == Other.Data->Num()) &&
+			(FMemory::Memcmp(Data->GetData(), Other.Data->GetData(), Data->Num()) == 0));
+	}
+};
+
+template <class TDerived, bool AlwaysCompareData = false>
+class TDataKey : private TDataKeyBase<AlwaysCompareData ? 2 : (DO_CHECK != 0)>
+{
+public:
+	template <class TArchiveWriter>
+	void Generate(TArchiveWriter&& WriteToArchive, int32 DataReserve = 0)
+	{
+		this->UpdateData([&](TArray<uint8>& InData)
+		{
+			FMemoryWriter Ar(InData);
+
+			InData.Reset(DataReserve);
+			WriteToArchive(Ar);
+
+			this->Hash = FCrc::MemCrc32(InData.GetData(), InData.Num());
+		});
+	}
+
+	uint32 GetHash() const
+	{
+		return this->Hash;
+	}
+
+	TDerived CopyDeep() const
+	{
+		TDerived Result;
+		Result.Hash = this->Hash;
+		this->CopyDataDeep(Result);
+		return Result;
+	}
+
+	TDerived CopyShallow() const
+	{
+		TDerived Result;
+		Result.Hash = this->Hash;
+		this->CopyDataShallow(Result);
+		return Result;
+	}
+
+	friend uint32 GetTypeHash(const TDerived& Key)
+	{
+		return Key.Hash;
+	}
+
+	friend bool operator==(const TDerived& A, const TDerived& B)
+	{
+		return ((A.Hash == B.Hash) && A.IsDataEquals(B));
+	}
+};
