UE4.21 - Fix Vulkan crash
#rnx

diff --git a/Engine/Source/Runtime/VulkanRHI/Private/VulkanTexture.cpp b/Engine/Source/Runtime/VulkanRHI/Private/VulkanTexture.cpp
--- a/Engine/Source/Runtime/VulkanRHI/Private/VulkanTexture.cpp
+++ b/Engine/Source/Runtime/VulkanRHI/Private/VulkanTexture.cpp
@@ -1526,6 +1526,7 @@ FVulkanTextureBase::FVulkanTextureBase(FVulkanDevice& Device, VkImageViewType Re
 
 	if (!CreateInfo.BulkData)
 	{
+		Device.GetImmediateContext().FindOrAddLayout(Surface.Image, VK_IMAGE_LAYOUT_UNDEFINED);
 		return;
 	}
 
