UE4.21 - Static analysis fix
#rnx

diff --git a/Engine/Source/Runtime/VulkanRHI/Private/VulkanRenderTarget.cpp b/Engine/Source/Runtime/VulkanRHI/Private/VulkanRenderTarget.cpp
--- a/Engine/Source/Runtime/VulkanRHI/Private/VulkanRenderTarget.cpp
+++ b/Engine/Source/Runtime/VulkanRHI/Private/VulkanRenderTarget.cpp
@@ -577,7 +577,7 @@ void FVulkanCommandListContext::RHISetRenderTargetsAndClear(const FRHISetRenderT
 		Framebuffer = TransitionAndLayoutManager.GetOrCreateFramebuffer(*Device, RenderTargetsInfo, RTLayout, RenderPass);
 	}
 
-	if (Framebuffer == TransitionAndLayoutManager.CurrentFramebuffer && IsCompatibleRenderPass(TransitionAndLayoutManager.CurrentRenderPass, RenderPass))
+	if (Framebuffer == TransitionAndLayoutManager.CurrentFramebuffer && RenderPass != nullptr && IsCompatibleRenderPass(TransitionAndLayoutManager.CurrentRenderPass, RenderPass))
 	{
 		return;
 	}
