UE4.21 - Fix Vulkan ensure
#rnx

diff --git a/Engine/Source/Runtime/VulkanRHI/Private/VulkanTexture.cpp b/Engine/Source/Runtime/VulkanRHI/Private/VulkanTexture.cpp
--- a/Engine/Source/Runtime/VulkanRHI/Private/VulkanTexture.cpp
+++ b/Engine/Source/Runtime/VulkanRHI/Private/VulkanTexture.cpp
@@ -797,8 +797,7 @@ void FVulkanSurface::InitialClear(FVulkanCommandListContext& Context,const FClea
 	}
 
 	VkImageLayout FinalLayout = Barrier.GetDestLayout(BarrierIndex);
-	VkImageLayout ImageLayout = Context.FindOrAddLayout(Image, FinalLayout);
-	ensure(FinalLayout == ImageLayout);
+	Context.FindOrAddLayoutRW(Image, FinalLayout) = FinalLayout;
 }
 
 /*-----------------------------------------------------------------------------
