UE4.21 - Really an info msg, not an error
#rnx

diff --git a/Engine/Source/Runtime/VulkanRHI/Private/VulkanPipeline.cpp b/Engine/Source/Runtime/VulkanRHI/Private/VulkanPipeline.cpp
--- a/Engine/Source/Runtime/VulkanRHI/Private/VulkanPipeline.cpp
+++ b/Engine/Source/Runtime/VulkanRHI/Private/VulkanPipeline.cpp
@@ -322,7 +322,7 @@ bool FVulkanPipelineStateCacheManager::Load(const TArray<FString>& CacheFilename
 		}
 		else
 		{
-			UE_LOG(LogVulkanRHI, Error, TEXT("FVulkanPipelineStateCacheManager: Failed to load binary pipeline cache %s"), *BinaryCacheFilename);
+			UE_LOG(LogVulkanRHI, Display, TEXT("FVulkanPipelineStateCacheManager: Binary pipeline cache '%s' not found."), *BinaryCacheFilename);
 		}
 	}
 
