UE4.21 - Misc Vulkan fixes
- Change how command buffers get destroyed
- Fix render pass compatibility on desktop
- Change warning to display
- Fix some leaks when exiting
#rnx

diff --git a/Engine/Source/Runtime/VulkanRHI/Private/Android/VulkanAndroidPlatform.h b/Engine/Source/Runtime/VulkanRHI/Private/Android/VulkanAndroidPlatform.h
--- a/Engine/Source/Runtime/VulkanRHI/Private/Android/VulkanAndroidPlatform.h
+++ b/Engine/Source/Runtime/VulkanRHI/Private/Android/VulkanAndroidPlatform.h
@@ -17,6 +17,12 @@
 #define VULKAN_FREEPAGE_FOR_TYPE					1
 #define VULKAN_PURGE_SHADER_MODULES					1
 
+// Android's hashes currently work fine as the problematic cases are:
+//	VK_IMAGE_LAYOUT_DEPTH_READ_ONLY_STENCIL_ATTACHMENT_OPTIMAL = 1000117000,
+//	VK_IMAGE_LAYOUT_DEPTH_ATTACHMENT_STENCIL_READ_ONLY_OPTIMAL = 1000117001,
+#define VULKAN_USE_REAL_RENDERPASS_COMPATIBILITY	0
+
+
 #define ENUM_VK_ENTRYPOINTS_PLATFORM_BASE(EnumMacro)
 
 #define ENUM_VK_ENTRYPOINTS_PLATFORM_INSTANCE(EnumMacro) \
diff --git a/Engine/Source/Runtime/VulkanRHI/Private/VulkanCommandBuffer.cpp b/Engine/Source/Runtime/VulkanRHI/Private/VulkanCommandBuffer.cpp
--- a/Engine/Source/Runtime/VulkanRHI/Private/VulkanCommandBuffer.cpp
+++ b/Engine/Source/Runtime/VulkanRHI/Private/VulkanCommandBuffer.cpp
@@ -31,23 +31,36 @@ static FAutoConsoleVariableRef CVarVulkanProfileCmdBuffers(
 const uint32 GNumberOfFramesBeforeDeletingDescriptorPool = 300;
 
 FVulkanCmdBuffer::FVulkanCmdBuffer(FVulkanDevice* InDevice, FVulkanCommandBufferPool* InCommandBufferPool, bool bInIsUploadOnly)
-	: bNeedsDynamicStateSet(true)
-	, bHasPipeline(false)
-	, bHasViewport(false)
-	, bHasScissor(false)
-	, bHasStencilRef(false)
-	, CurrentStencilRef(0)
+	: CurrentStencilRef(0)
+	, State(EState::NotAllocated)
+	, bNeedsDynamicStateSet(1)
+	, bHasPipeline(0)
+	, bHasViewport(0)
+	, bHasScissor(0)
+	, bHasStencilRef(0)
+	, bIsUploadOnly(bInIsUploadOnly ? 1 : 0)
 	, Device(InDevice)
 	, CommandBufferHandle(VK_NULL_HANDLE)
-	, State(EState::ReadyForBegin)
 	, Fence(nullptr)
 	, FenceSignaledCounter(0)
 	, SubmittedFenceCounter(0)
-	, bIsUploadOnly(bInIsUploadOnly)
 	, CommandBufferPool(InCommandBufferPool)
 	, Timing(nullptr)
 	, LastValidTiming(0)
 {
+
+	{
+		FScopeLock ScopeLock(CommandBufferPool->GetCS());
+		AllocMemory();
+	}
+
+	Fence = Device->GetFenceManager().AllocateFence();
+}
+
+void FVulkanCmdBuffer::AllocMemory()
+{
+	// Assumes we are inside a lock for the pool
+	check(State == EState::NotAllocated);
 	FMemory::Memzero(CurrentViewport);
 	FMemory::Memzero(CurrentScissor);
 
@@ -57,12 +70,16 @@ FVulkanCmdBuffer::FVulkanCmdBuffer(FVulkanDevice* InDevice, FVulkanCommandBuffer
 	CreateCmdBufInfo.commandBufferCount = 1;
 	CreateCmdBufInfo.commandPool = CommandBufferPool->GetHandle();
 
-	{
-		FScopeLock ScopeLock(InCommandBufferPool->GetCS());
-		VERIFYVULKANRESULT(VulkanRHI::vkAllocateCommandBuffers(Device->GetInstanceHandle(), &CreateCmdBufInfo, &CommandBufferHandle));
-	}
+	VERIFYVULKANRESULT(VulkanRHI::vkAllocateCommandBuffers(Device->GetInstanceHandle(), &CreateCmdBufInfo, &CommandBufferHandle));
+
+	bNeedsDynamicStateSet = 1;
+	bHasPipeline = 0;
+	bHasViewport = 0;
+	bHasScissor = 0;
+	bHasStencilRef = 0;
+	State = EState::ReadyForBegin;
+
 	INC_DWORD_STAT(STAT_VulkanNumCmdBuffers);
-	Fence = Device->GetFenceManager().AllocateFence();
 }
 
 FVulkanCmdBuffer::~FVulkanCmdBuffer()
@@ -80,13 +97,10 @@ FVulkanCmdBuffer::~FVulkanCmdBuffer()
 		FenceManager.ReleaseFence(Fence);
 	}
 
+	if (State != EState::NotAllocated)
 	{
-		FScopeLock ScopeLock(CommandBufferPool->GetCS());
-		VulkanRHI::vkFreeCommandBuffers(Device->GetInstanceHandle(), CommandBufferPool->GetHandle(), 1, &CommandBufferHandle);
+		FreeMemory();
 	}
-	CommandBufferHandle = VK_NULL_HANDLE;
-
-	DEC_DWORD_STAT(STAT_VulkanNumCmdBuffers);
 
 	if (Timing)
 	{
@@ -94,6 +108,18 @@ FVulkanCmdBuffer::~FVulkanCmdBuffer()
 	}
 }
 
+void FVulkanCmdBuffer::FreeMemory()
+{
+	// Assumes we are inside a lock for the pool
+	check(State != EState::NotAllocated);
+	check(CommandBufferHandle != VK_NULL_HANDLE);
+	VulkanRHI::vkFreeCommandBuffers(Device->GetInstanceHandle(), CommandBufferPool->GetHandle(), 1, &CommandBufferHandle);
+	CommandBufferHandle = VK_NULL_HANDLE;
+
+	DEC_DWORD_STAT(STAT_VulkanNumCmdBuffers);
+	State = EState::NotAllocated;
+}
+
 void FVulkanCmdBuffer::BeginRenderPass(const FVulkanRenderTargetLayout& Layout, FVulkanRenderPass* RenderPass, FVulkanFramebuffer* Framebuffer, const VkClearValue* AttachmentClearValues)
 {
 	checkf(IsOutsideRenderPass(), TEXT("Can't BeginRP as already inside one! CmdBuffer 0x%p State=%d"), CommandBufferHandle, (int32)State);
@@ -282,6 +308,13 @@ FVulkanCommandBufferPool::~FVulkanCommandBufferPool()
 	for (int32 Index = 0; Index < CmdBuffers.Num(); ++Index)
 	{
 		FVulkanCmdBuffer* CmdBuffer = CmdBuffers[Index];
+		CmdBuffer->FreeMemory();
+		delete CmdBuffer;
+	}
+
+	for (int32 Index = 0; Index < FreeCmdBuffers.Num(); ++Index)
+	{
+		FVulkanCmdBuffer* CmdBuffer = FreeCmdBuffers[Index];
 		delete CmdBuffer;
 	}
 
@@ -380,7 +413,21 @@ void FVulkanCommandBufferManager::SubmitActiveCmdBuffer(VulkanRHI::FSemaphore* S
 
 FVulkanCmdBuffer* FVulkanCommandBufferPool::Create(bool bIsUploadOnly)
 {
-	check(Device);
+	// Assumes we are inside a lock for the pool
+	for (int32 Index = FreeCmdBuffers.Num() - 1; Index >= 0; --Index)
+	{
+		FVulkanCmdBuffer* CmdBuffer = FreeCmdBuffers[Index];
+#if VULKAN_USE_DIFFERENT_POOL_CMDBUFFERS
+		if (CmdBuffer->bIsUploadOnly == bIsUploadOnly)
+#endif
+		{
+			FreeCmdBuffers.RemoveAtSwap(Index);
+			CmdBuffer->AllocMemory();
+			CmdBuffers.Add(CmdBuffer);
+			return CmdBuffer;
+		}
+	}
+
 	FVulkanCmdBuffer* CmdBuffer = new FVulkanCmdBuffer(Device, this, bIsUploadOnly);
 	CmdBuffers.Add(CmdBuffer);
 	check(CmdBuffer);
@@ -518,8 +565,9 @@ void FVulkanCommandBufferPool::FreeUnusedCmdBuffers(FVulkanQueue* InQueue)
 		{
 			DeferredDeletionQueue.OnCmdBufferDeleted(CmdBuffer);
 			
-			delete CmdBuffer;
+			CmdBuffer->FreeMemory();
 			CmdBuffers.RemoveAtSwap(Index, 1, false);
+			FreeCmdBuffers.Add(CmdBuffer);
 		}
 	}
 #endif
diff --git a/Engine/Source/Runtime/VulkanRHI/Private/VulkanCommandBuffer.h b/Engine/Source/Runtime/VulkanRHI/Private/VulkanCommandBuffer.h
--- a/Engine/Source/Runtime/VulkanRHI/Private/VulkanCommandBuffer.h
+++ b/Engine/Source/Runtime/VulkanRHI/Private/VulkanCommandBuffer.h
@@ -128,30 +128,32 @@ public:
 	{
 		return (Timing != nullptr) && (FMath::Abs((int64)FenceSignaledCounter - (int64)LastValidTiming) < 3);
 	}
-	
+
 	void AddWaitSemaphore(VkPipelineStageFlags InWaitFlags, VulkanRHI::FSemaphore* InWaitSemaphore);
 
 	void Begin();
 	void End();
 
-	enum class EState
+	enum class EState : uint8
 	{
 		ReadyForBegin,
 		IsInsideBegin,
 		IsInsideRenderPass,
 		HasEnded,
 		Submitted,
+		NotAllocated,
 	};
 
-	bool bNeedsDynamicStateSet;
-	bool bHasPipeline;
-	bool bHasViewport;
-	bool bHasScissor;
-	bool bHasStencilRef;
-
 	VkViewport CurrentViewport;
 	VkRect2D CurrentScissor;
 	uint32 CurrentStencilRef;
+	EState State;
+	uint8 bNeedsDynamicStateSet	: 1;
+	uint8 bHasPipeline			: 1;
+	uint8 bHasViewport			: 1;
+	uint8 bHasScissor			: 1;
+	uint8 bHasStencilRef		: 1;
+	uint8 bIsUploadOnly			: 1;
 
 	// You never want to call Begin/EndRenderPass directly as it will mess up with the FTransitionAndLayoutManager
 	void BeginRenderPass(const FVulkanRenderTargetLayout& Layout, class FVulkanRenderPass* RenderPass, class FVulkanFramebuffer* Framebuffer, const VkClearValue* AttachmentClearValues);
@@ -170,7 +172,6 @@ public:
 private:
 	FVulkanDevice* Device;
 	VkCommandBuffer CommandBufferHandle;
-	EState State;
 	double SubmittedTime = 0.0f;
 
 	TArray<VkPipelineStageFlags> WaitFlags;
@@ -196,7 +197,6 @@ private:
 	void RefreshFenceStatus();
 	void InitializeTimings(FVulkanCommandListContext* InContext);
 
-	bool bIsUploadOnly;
 	FVulkanCommandBufferPool* CommandBufferPool;
 
 	FVulkanGPUTiming* Timing;
@@ -204,6 +204,9 @@ private:
 
 	void AcquirePoolSetContainer();
 
+	void AllocMemory();
+	void FreeMemory();
+
 public:
 	//#todo-rco: Hide this
 	TMap<uint32, class FVulkanTypedDescriptorPoolSet*> TypedDescriptorPoolSets;
@@ -241,6 +244,7 @@ private:
 	VkCommandPool Handle;
 
 	TArray<FVulkanCmdBuffer*> CmdBuffers;
+	TArray<FVulkanCmdBuffer*> FreeCmdBuffers;
 
 	FCriticalSection CS;
 	FVulkanDevice* Device;
diff --git a/Engine/Source/Runtime/VulkanRHI/Private/VulkanMemory.cpp b/Engine/Source/Runtime/VulkanRHI/Private/VulkanMemory.cpp
--- a/Engine/Source/Runtime/VulkanRHI/Private/VulkanMemory.cpp
+++ b/Engine/Source/Runtime/VulkanRHI/Private/VulkanMemory.cpp
@@ -725,7 +725,7 @@ namespace VulkanRHI
 			FScopeLock ScopeLock(&GOldResourceLock);
 
 #if VULKAN_FREEPAGE_FOR_TYPE
-			for (int32 Index = 1; Index < FreeBufferPages.Num(); ++Index)
+			for (int32 Index = (bImmediately ? 0 : 1); Index < FreeBufferPages.Num(); ++Index)
 			{
 				FOldResourceHeapPage* Page = FreeBufferPages[Index];
 				if (bImmediately || Page->FrameFreed + NUM_FRAMES_TO_WAIT_BEFORE_RELEASING_TO_OS < GFrameNumberRenderThread)
@@ -735,7 +735,7 @@ namespace VulkanRHI
 					break;
 				}
 			}
-			for (int32 Index = 1; Index < FreeImagePages.Num(); ++Index)
+			for (int32 Index = (bImmediately ? 0 : 1); Index < FreeImagePages.Num(); ++Index)
 			{
 				FOldResourceHeapPage* Page = FreeImagePages[Index];
 				if (bImmediately || Page->FrameFreed + NUM_FRAMES_TO_WAIT_BEFORE_RELEASING_TO_OS < GFrameNumberRenderThread)
@@ -751,7 +751,7 @@ namespace VulkanRHI
 			for (int32 Index = 0; Index < FreePages.Num(); ++Index)
 #else
 			// Leave a page not freed to avoid potential hitching
-			for (int32 Index = 1; Index < FreePages.Num(); ++Index)
+			for (int32 Index = (bImmediately ? 0 : 1); Index < FreePages.Num(); ++Index)
 #endif
 			{
 				FOldResourceHeapPage* Page = FreePages[Index];
@@ -806,7 +806,7 @@ namespace VulkanRHI
 				UE_LOG(LogVulkanRHI, Display, TEXT("\t\t%d: ID %4d %4d suballocs, %4d free chunks (%d used/%d free/%d max) DeviceMemory %p"), Index, UsedPages[Index]->GetID(), UsedPages[Index]->ResourceAllocations.Num(), UsedPages[Index]->FreeList.Num(), UsedPages[Index]->UsedSize, UsedPages[Index]->MaxSize - UsedPages[Index]->UsedSize, UsedPages[Index]->MaxSize, (void*)UsedPages[Index]->DeviceMemoryAllocation->GetHandle());
 			}
 
-			UE_LOG(LogVulkanRHI, Display, TEXT("%d Suballocations for Used/Total: %d/%d = %.2f%%"), NumSuballocations, SubAllocUsedMemory, SubAllocAllocatedMemory, 100.0f * (float)SubAllocUsedMemory / (float)SubAllocAllocatedMemory);
+			UE_LOG(LogVulkanRHI, Display, TEXT("%d Suballocations for Used/Total: %d/%d = %.2f%%"), NumSuballocations, SubAllocUsedMemory, SubAllocAllocatedMemory, SubAllocAllocatedMemory > 0 ? 100.0f * (float)SubAllocUsedMemory / (float)SubAllocAllocatedMemory : 0.0f);
 		};
 
 		DumpPages(UsedBufferPages, TEXT("Buffer"));
@@ -1352,8 +1352,8 @@ namespace VulkanRHI
 		}
 
 		UE_LOG(LogVulkanRHI, Display, TEXT("::Totals::"));
-		UE_LOG(LogVulkanRHI, Display, TEXT("Large Alloc Used/Max %d/%d %.2f%%"), UsedLargeTotal, AllocLargeTotal, 100.0f * (float)UsedLargeTotal / (float)AllocLargeTotal);
-		UE_LOG(LogVulkanRHI, Display, TEXT("Binned Alloc Used/Max %d/%d %.2f%%"), UsedBinnedTotal, AllocBinnedTotal, 100.0f * (float)UsedBinnedTotal / (float)AllocBinnedTotal);
+		UE_LOG(LogVulkanRHI, Display, TEXT("Large Alloc Used/Max %d/%d %.2f%%"), UsedLargeTotal, AllocLargeTotal, 100.0f * AllocLargeTotal > 0 ? (float)UsedLargeTotal / (float)AllocLargeTotal : 0.0f);
+		UE_LOG(LogVulkanRHI, Display, TEXT("Binned Alloc Used/Max %d/%d %.2f%%"), UsedBinnedTotal, AllocBinnedTotal, AllocBinnedTotal > 0 ? 100.0f * (float)UsedBinnedTotal / (float)AllocBinnedTotal : 0.0f);
 	}
 #endif
 
diff --git a/Engine/Source/Runtime/VulkanRHI/Private/VulkanRenderTarget.cpp b/Engine/Source/Runtime/VulkanRHI/Private/VulkanRenderTarget.cpp
--- a/Engine/Source/Runtime/VulkanRHI/Private/VulkanRenderTarget.cpp
+++ b/Engine/Source/Runtime/VulkanRHI/Private/VulkanRenderTarget.cpp
@@ -1632,6 +1632,12 @@ struct FRenderPassFullHashableStruct
 	// +1 for Depth, +1 for Stencil
 	TEnumAsByte<VkAttachmentLoadOp>		LoadOps[MaxSimultaneousRenderTargets + 2];
 	TEnumAsByte<VkAttachmentStoreOp>	StoreOps[MaxSimultaneousRenderTargets + 2];
+#if VULKAN_USE_REAL_RENDERPASS_COMPATIBILITY
+	// If the initial != final we need to add FinalLayout and potentially RefLayout
+	VkImageLayout						InitialLayout[MaxSimultaneousRenderTargets + 1];
+	//VkImageLayout						FinalLayout[MaxSimultaneousRenderTargets + 1];
+	//VkImageLayout						RefLayout[MaxSimultaneousRenderTargets + 1];
+#endif
 };
 
 
@@ -1694,8 +1700,11 @@ FVulkanRenderTargetLayout::FVulkanRenderTargetLayout(FVulkanDevice& InDevice, co
 			CurrDesc.storeOp = RenderTargetStoreActionToVulkan(RTView.StoreAction);
 			CurrDesc.stencilLoadOp = VK_ATTACHMENT_LOAD_OP_DONT_CARE;
 			CurrDesc.stencilStoreOp = VK_ATTACHMENT_STORE_OP_DONT_CARE;
+
+			// If the initial != final we need to change the FullHashInfo and use FinalLayout
 			CurrDesc.initialLayout = VK_IMAGE_LAYOUT_COLOR_ATTACHMENT_OPTIMAL;
 			CurrDesc.finalLayout = VK_IMAGE_LAYOUT_COLOR_ATTACHMENT_OPTIMAL;
+
 			ColorReferences[NumColorAttachments].attachment = NumAttachmentDescriptions;
 			ColorReferences[NumColorAttachments].layout = VK_IMAGE_LAYOUT_COLOR_ATTACHMENT_OPTIMAL;
 			if (CurrDesc.samples > VK_SAMPLE_COUNT_1_BIT)
@@ -1711,6 +1720,9 @@ FVulkanRenderTargetLayout::FVulkanRenderTargetLayout(FVulkanDevice& InDevice, co
 			CompatibleHashInfo.Formats[NumColorAttachments] = CurrDesc.format;
 			FullHashInfo.LoadOps[NumColorAttachments] = CurrDesc.loadOp;
 			FullHashInfo.StoreOps[NumColorAttachments] = CurrDesc.storeOp;
+#if VULKAN_USE_REAL_RENDERPASS_COMPATIBILITY
+			FullHashInfo.InitialLayout[NumColorAttachments] = VK_IMAGE_LAYOUT_COLOR_ATTACHMENT_OPTIMAL;
+#endif
 			++CompatibleHashInfo.NumAttachments;
 
 			++NumAttachmentDescriptions;
@@ -1752,8 +1764,11 @@ FVulkanRenderTargetLayout::FVulkanRenderTargetLayout(FVulkanDevice& InDevice, co
 		}
 
 		DepthStencilLayout = VulkanRHI::GetDepthStencilLayout(RTInfo.DepthStencilRenderTarget.GetDepthStencilAccess(), InDevice);
+
+		// If the initial != final we need to change the FullHashInfo and use FinalLayout
 		CurrDesc.initialLayout = DepthStencilLayout;
 		CurrDesc.finalLayout = DepthStencilLayout;
+
 		DepthStencilReference.attachment = NumAttachmentDescriptions;
 		DepthStencilReference.layout = DepthStencilLayout;
 
@@ -1761,6 +1776,9 @@ FVulkanRenderTargetLayout::FVulkanRenderTargetLayout(FVulkanDevice& InDevice, co
 		FullHashInfo.LoadOps[MaxSimultaneousRenderTargets + 1] = CurrDesc.stencilLoadOp;
 		FullHashInfo.StoreOps[MaxSimultaneousRenderTargets] = CurrDesc.storeOp;
 		FullHashInfo.StoreOps[MaxSimultaneousRenderTargets + 1] = CurrDesc.stencilStoreOp;
+#if VULKAN_USE_REAL_RENDERPASS_COMPATIBILITY
+		FullHashInfo.InitialLayout[MaxSimultaneousRenderTargets] = DepthStencilLayout;
+#endif
 		CompatibleHashInfo.Formats[MaxSimultaneousRenderTargets] = CurrDesc.format;
 
 		++NumAttachmentDescriptions;
@@ -1853,10 +1871,13 @@ FVulkanRenderTargetLayout::FVulkanRenderTargetLayout(FVulkanDevice& InDevice, co
 		CurrDesc.storeOp = RenderTargetStoreActionToVulkan(GetStoreAction(ColorEntry.Action), true);
 		CurrDesc.stencilLoadOp = VK_ATTACHMENT_LOAD_OP_DONT_CARE;
 		CurrDesc.stencilStoreOp = VK_ATTACHMENT_STORE_OP_DONT_CARE;
-		CurrDesc.initialLayout = /*RPInfo.bGeneratingMips ? VK_IMAGE_LAYOUT_GENERAL : */VK_IMAGE_LAYOUT_COLOR_ATTACHMENT_OPTIMAL;
-		CurrDesc.finalLayout = /*RPInfo.bGeneratingMips ? VK_IMAGE_LAYOUT_GENERAL : */VK_IMAGE_LAYOUT_COLOR_ATTACHMENT_OPTIMAL;
+
+		// If the initial != final we need to change the FullHashInfo and use FinalLayout
+		CurrDesc.initialLayout = VK_IMAGE_LAYOUT_COLOR_ATTACHMENT_OPTIMAL;
+		CurrDesc.finalLayout = VK_IMAGE_LAYOUT_COLOR_ATTACHMENT_OPTIMAL;
+
 		ColorReferences[NumColorAttachments].attachment = NumAttachmentDescriptions;
-		ColorReferences[NumColorAttachments].layout = /*RPInfo.bGeneratingMips ? VK_IMAGE_LAYOUT_GENERAL : */VK_IMAGE_LAYOUT_COLOR_ATTACHMENT_OPTIMAL;
+		ColorReferences[NumColorAttachments].layout = VK_IMAGE_LAYOUT_COLOR_ATTACHMENT_OPTIMAL;
 		if (CurrDesc.samples > VK_SAMPLE_COUNT_1_BIT)
 		{
 			Desc[NumAttachmentDescriptions + 1] = Desc[NumAttachmentDescriptions];
@@ -1869,6 +1890,9 @@ FVulkanRenderTargetLayout::FVulkanRenderTargetLayout(FVulkanDevice& InDevice, co
 
 		CompatibleHashInfo.Formats[NumColorAttachments] = CurrDesc.format;
 		FullHashInfo.LoadOps[NumColorAttachments] = CurrDesc.loadOp;
+#if VULKAN_USE_REAL_RENDERPASS_COMPATIBILITY
+		FullHashInfo.InitialLayout[NumColorAttachments] = VK_IMAGE_LAYOUT_COLOR_ATTACHMENT_OPTIMAL;
+#endif
 		FullHashInfo.StoreOps[NumColorAttachments] = CurrDesc.storeOp;
 		++CompatibleHashInfo.NumAttachments;
 
@@ -1902,6 +1926,7 @@ FVulkanRenderTargetLayout::FVulkanRenderTargetLayout(FVulkanDevice& InDevice, co
 		CurrDesc.stencilStoreOp = RenderTargetStoreActionToVulkan(GetStoreAction(GetStencilActions(RPInfo.DepthStencilRenderTarget.Action)), true);
 
 		DepthStencilLayout = VulkanRHI::GetDepthStencilLayout(RPInfo.DepthStencilRenderTarget.ExclusiveDepthStencil, InDevice);
+		// If the initial != final we need to change the FullHashInfo and use FinalLayout
 		CurrDesc.initialLayout = DepthStencilLayout;
 		CurrDesc.finalLayout = DepthStencilLayout;
 		DepthStencilReference.attachment = NumAttachmentDescriptions;
@@ -1911,6 +1936,9 @@ FVulkanRenderTargetLayout::FVulkanRenderTargetLayout(FVulkanDevice& InDevice, co
 		FullHashInfo.LoadOps[MaxSimultaneousRenderTargets + 1] = CurrDesc.stencilLoadOp;
 		FullHashInfo.StoreOps[MaxSimultaneousRenderTargets] = CurrDesc.storeOp;
 		FullHashInfo.StoreOps[MaxSimultaneousRenderTargets + 1] = CurrDesc.stencilStoreOp;
+#if VULKAN_USE_REAL_RENDERPASS_COMPATIBILITY
+		FullHashInfo.InitialLayout[MaxSimultaneousRenderTargets] = DepthStencilLayout;
+#endif
 		CompatibleHashInfo.Formats[MaxSimultaneousRenderTargets] = CurrDesc.format;
 
 		++NumAttachmentDescriptions;
@@ -1983,8 +2011,11 @@ FVulkanRenderTargetLayout::FVulkanRenderTargetLayout(const FGraphicsPipelineStat
 			CurrDesc.storeOp = VK_ATTACHMENT_STORE_OP_DONT_CARE;
 			CurrDesc.stencilLoadOp = VK_ATTACHMENT_LOAD_OP_DONT_CARE;
 			CurrDesc.stencilStoreOp = VK_ATTACHMENT_STORE_OP_DONT_CARE;
+
+			// If the initial != final we need to change the FullHashInfo and use FinalLayout
 			CurrDesc.initialLayout = VK_IMAGE_LAYOUT_COLOR_ATTACHMENT_OPTIMAL;
 			CurrDesc.finalLayout = VK_IMAGE_LAYOUT_COLOR_ATTACHMENT_OPTIMAL;
+
 			ColorReferences[NumColorAttachments].attachment = NumAttachmentDescriptions;
 			ColorReferences[NumColorAttachments].layout = VK_IMAGE_LAYOUT_COLOR_ATTACHMENT_OPTIMAL;
 			if (CurrDesc.samples > VK_SAMPLE_COUNT_1_BIT)
@@ -2000,6 +2031,9 @@ FVulkanRenderTargetLayout::FVulkanRenderTargetLayout(const FGraphicsPipelineStat
 			CompatibleHashInfo.Formats[NumColorAttachments] = CurrDesc.format;
 			FullHashInfo.LoadOps[NumColorAttachments] = CurrDesc.loadOp;
 			FullHashInfo.StoreOps[NumColorAttachments] = CurrDesc.storeOp;
+#if VULKAN_USE_REAL_RENDERPASS_COMPATIBILITY
+			FullHashInfo.InitialLayout[NumColorAttachments] = VK_IMAGE_LAYOUT_COLOR_ATTACHMENT_OPTIMAL;
+#endif
 			++CompatibleHashInfo.NumAttachments;
 
 			++NumAttachmentDescriptions;
@@ -2031,6 +2065,8 @@ FVulkanRenderTargetLayout::FVulkanRenderTargetLayout(const FGraphicsPipelineStat
 			CurrDesc.storeOp = VK_ATTACHMENT_STORE_OP_DONT_CARE;
 			CurrDesc.stencilStoreOp = VK_ATTACHMENT_STORE_OP_DONT_CARE;
 		}
+
+		// If the initial != final we need to change the FullHashInfo and use FinalLayout
 		CurrDesc.initialLayout = VK_IMAGE_LAYOUT_DEPTH_STENCIL_ATTACHMENT_OPTIMAL;
 		CurrDesc.finalLayout = VK_IMAGE_LAYOUT_DEPTH_STENCIL_ATTACHMENT_OPTIMAL;
 
@@ -2051,6 +2087,9 @@ FVulkanRenderTargetLayout::FVulkanRenderTargetLayout(const FGraphicsPipelineStat
 		FullHashInfo.LoadOps[MaxSimultaneousRenderTargets + 1] = CurrDesc.stencilLoadOp;
 		FullHashInfo.StoreOps[MaxSimultaneousRenderTargets] = CurrDesc.storeOp;
 		FullHashInfo.StoreOps[MaxSimultaneousRenderTargets + 1] = CurrDesc.stencilStoreOp;
+#if VULKAN_USE_REAL_RENDERPASS_COMPATIBILITY
+		FullHashInfo.InitialLayout[MaxSimultaneousRenderTargets] = VK_IMAGE_LAYOUT_DEPTH_STENCIL_ATTACHMENT_OPTIMAL;
+#endif
 		CompatibleHashInfo.Formats[MaxSimultaneousRenderTargets] = CurrDesc.format;
 
 		++NumAttachmentDescriptions;
diff --git a/Engine/Source/Runtime/VulkanRHI/Private/VulkanSwapChain.cpp b/Engine/Source/Runtime/VulkanRHI/Private/VulkanSwapChain.cpp
--- a/Engine/Source/Runtime/VulkanRHI/Private/VulkanSwapChain.cpp
+++ b/Engine/Source/Runtime/VulkanRHI/Private/VulkanSwapChain.cpp
@@ -168,7 +168,7 @@ FVulkanSwapChain::FVulkanSwapChain(VkInstance InInstance, FVulkanDevice& InDevic
 					if (!GPixelFormatNotSupportedWarning.Contains(InOutPixelFormat))
 					{
 						GPixelFormatNotSupportedWarning.Add(InOutPixelFormat);
-						UE_LOG(LogVulkanRHI, Warning, TEXT("Requested PixelFormat %d not supported by this swapchain! Falling back to supported swapchain format..."), (uint32)InOutPixelFormat);
+						UE_LOG(LogVulkanRHI, Display, TEXT("Requested PixelFormat %d not supported by this swapchain! Falling back to supported swapchain format..."), (uint32)InOutPixelFormat);
 					}
 					InOutPixelFormat = PF_Unknown;
 				}
diff --git a/Engine/Source/Runtime/VulkanRHI/Public/VulkanConfiguration.h b/Engine/Source/Runtime/VulkanRHI/Public/VulkanConfiguration.h
--- a/Engine/Source/Runtime/VulkanRHI/Public/VulkanConfiguration.h
+++ b/Engine/Source/Runtime/VulkanRHI/Public/VulkanConfiguration.h
@@ -146,6 +146,10 @@
 	#define VULKAN_SUPPORTS_GEOMETRY_SHADERS					1
 #endif
 
+#ifndef VULKAN_USE_REAL_RENDERPASS_COMPATIBILITY
+	#define VULKAN_USE_REAL_RENDERPASS_COMPATIBILITY			1
+#endif
+
 #ifndef VULKAN_USE_DIFFERENT_POOL_CMDBUFFERS
 	#define VULKAN_USE_DIFFERENT_POOL_CMDBUFFERS				1
 #endif
