UE4.21 - Remove bogus ensure

diff --git a/Engine/Source/Runtime/VulkanRHI/Private/VulkanTexture.cpp b/Engine/Source/Runtime/VulkanRHI/Private/VulkanTexture.cpp
--- a/Engine/Source/Runtime/VulkanRHI/Private/VulkanTexture.cpp
+++ b/Engine/Source/Runtime/VulkanRHI/Private/VulkanTexture.cpp
@@ -1224,9 +1224,6 @@ void FVulkanDynamicRHI::InternalUpdateTexture2D(bool bFromRenderingThread, FText
 	//SubresourceRange.baseArrayLayer = 0;
 	SubresourceRange.layerCount = 1;
 
-	ensure(UpdateRegion.SrcX == 0);
-	ensure(UpdateRegion.SrcY == 0);
-
 	uint8* RESTRICT DestData = (uint8*)Memory;
 	uint8* RESTRICT SourceRowData = (uint8*)SourceData;
 	for (uint32 Height = 0; Height < NumBlocksY; ++Height)
