UE4.21 - Added some missing Vulkan allocation tags
- Workaround for Linux on Nvidia cards not updating buffers
#rnx

diff --git a/Engine/Source/Runtime/VulkanRHI/Private/Linux/VulkanLinuxPlatform.cpp b/Engine/Source/Runtime/VulkanRHI/Private/Linux/VulkanLinuxPlatform.cpp
--- a/Engine/Source/Runtime/VulkanRHI/Private/Linux/VulkanLinuxPlatform.cpp
+++ b/Engine/Source/Runtime/VulkanRHI/Private/Linux/VulkanLinuxPlatform.cpp
@@ -129,6 +129,13 @@ void FVulkanLinuxPlatform::GetDeviceExtensions(TArray<const ANSICHAR*>& OutExten
 	OutExtensions.Add(VK_KHR_GET_MEMORY_REQUIREMENTS_2_EXTENSION_NAME);
 	OutExtensions.Add(VK_KHR_DEDICATED_ALLOCATION_EXTENSION_NAME);
 #endif
+
+	if (IsRHIDeviceNVIDIA())
+	{
+		//#todo-rco: Temporary workaround for some buffers not updating
+		extern FAutoConsoleVariableRef CVarVulkanWaitForIdleOnSubmit;
+		CVarVulkanWaitForIdleOnSubmit->Set(1);
+	}
 }
 
 void FVulkanLinuxPlatform::CreateSurface(void* WindowHandle, VkInstance Instance, VkSurfaceKHR* OutSurface)
diff --git a/Engine/Source/Runtime/VulkanRHI/Private/VulkanDevice.cpp b/Engine/Source/Runtime/VulkanRHI/Private/VulkanDevice.cpp
--- a/Engine/Source/Runtime/VulkanRHI/Private/VulkanDevice.cpp
+++ b/Engine/Source/Runtime/VulkanRHI/Private/VulkanDevice.cpp
@@ -638,7 +638,7 @@ VkSamplerYcbcrConversion FVulkanDevice::CreateSamplerColorConversion(const VkSam
 	else
 	{
 		VkSamplerYcbcrConversion NewConversion;
-		VERIFYVULKANRESULT(VulkanRHI::vkCreateSamplerYcbcrConversionKHR(GetInstanceHandle(), &CreateInfo, nullptr, &NewConversion));
+		VERIFYVULKANRESULT(VulkanRHI::vkCreateSamplerYcbcrConversionKHR(GetInstanceHandle(), &CreateInfo, VULKAN_CPU_ALLOCATOR, &NewConversion));
 		SamplerColorConversionMap.Add(CreateInfoHash, NewConversion);
 		return NewConversion;
 	}
@@ -873,7 +873,7 @@ void FVulkanDevice::Destroy()
 #if VULKAN_SUPPORTS_COLOR_CONVERSIONS
 	for (const auto& Pair : SamplerColorConversionMap)
 	{
-		VulkanRHI::vkDestroySamplerYcbcrConversionKHR(GetInstanceHandle(), Pair.Value, nullptr);
+		VulkanRHI::vkDestroySamplerYcbcrConversionKHR(GetInstanceHandle(), Pair.Value, VULKAN_CPU_ALLOCATOR);
 	}
 	SamplerColorConversionMap.Reset();
 #endif
diff --git a/Engine/Source/Runtime/VulkanRHI/Private/VulkanPipeline.cpp b/Engine/Source/Runtime/VulkanRHI/Private/VulkanPipeline.cpp
--- a/Engine/Source/Runtime/VulkanRHI/Private/VulkanPipeline.cpp
+++ b/Engine/Source/Runtime/VulkanRHI/Private/VulkanPipeline.cpp
@@ -548,10 +548,10 @@ void FVulkanPipelineStateCacheManager::Save(const FString& CacheFilename, bool b
 		{
 			UE_LOG(LogVulkanRHI, Warning, TEXT("Failed to get Vulkan pipeline cache data."));
 
-			VulkanRHI::vkDestroyPipelineCache(Device->GetInstanceHandle(), PipelineCache, nullptr);
+			VulkanRHI::vkDestroyPipelineCache(Device->GetInstanceHandle(), PipelineCache, VULKAN_CPU_ALLOCATOR);
 			VkPipelineCacheCreateInfo PipelineCacheInfo;
 			ZeroVulkanStruct(PipelineCacheInfo, VK_STRUCTURE_TYPE_PIPELINE_CACHE_CREATE_INFO);
-			VERIFYVULKANRESULT(VulkanRHI::vkCreatePipelineCache(Device->GetInstanceHandle(), &PipelineCacheInfo, nullptr, &PipelineCache));
+			VERIFYVULKANRESULT(VulkanRHI::vkCreatePipelineCache(Device->GetInstanceHandle(), &PipelineCacheInfo, VULKAN_CPU_ALLOCATOR, &PipelineCache));
 		}
 		else
 		{
diff --git a/Engine/Source/Runtime/VulkanRHI/Private/VulkanQuery.cpp b/Engine/Source/Runtime/VulkanRHI/Private/VulkanQuery.cpp
--- a/Engine/Source/Runtime/VulkanRHI/Private/VulkanQuery.cpp
+++ b/Engine/Source/Runtime/VulkanRHI/Private/VulkanQuery.cpp
@@ -120,7 +120,7 @@ FVulkanQueryPool::FVulkanQueryPool(FVulkanDevice* InDevice, uint32 InMaxQueries,
 	PoolCreateInfo.queryType = QueryType;
 	PoolCreateInfo.queryCount = MaxQueries;
 
-	VERIFYVULKANRESULT(VulkanRHI::vkCreateQueryPool(Device->GetInstanceHandle(), &PoolCreateInfo, nullptr, &QueryPool));
+	VERIFYVULKANRESULT(VulkanRHI::vkCreateQueryPool(Device->GetInstanceHandle(), &PoolCreateInfo, VULKAN_CPU_ALLOCATOR, &QueryPool));
 
 	QueryOutput.AddZeroed(MaxQueries);
 }
@@ -128,7 +128,7 @@ FVulkanQueryPool::FVulkanQueryPool(FVulkanDevice* InDevice, uint32 InMaxQueries,
 FVulkanQueryPool::~FVulkanQueryPool()
 {
 	DEC_DWORD_STAT(STAT_VulkanNumQueryPools);
-	VulkanRHI::vkDestroyQueryPool(Device->GetInstanceHandle(), QueryPool, nullptr);
+	VulkanRHI::vkDestroyQueryPool(Device->GetInstanceHandle(), QueryPool, VULKAN_CPU_ALLOCATOR);
 	QueryPool = VK_NULL_HANDLE;
 	check(QueryPool == VK_NULL_HANDLE);
 }
diff --git a/Engine/Source/Runtime/VulkanRHI/Private/VulkanQueue.cpp b/Engine/Source/Runtime/VulkanRHI/Private/VulkanQueue.cpp
--- a/Engine/Source/Runtime/VulkanRHI/Private/VulkanQueue.cpp
+++ b/Engine/Source/Runtime/VulkanRHI/Private/VulkanQueue.cpp
@@ -9,7 +9,7 @@
 #include "VulkanMemory.h"
 
 int32 GWaitForIdleOnSubmit = 0;
-static FAutoConsoleVariableRef CVarVulkanWaitForIdleOnSubmit(
+FAutoConsoleVariableRef CVarVulkanWaitForIdleOnSubmit(
 	TEXT("r.Vulkan.WaitForIdleOnSubmit"),
 	GWaitForIdleOnSubmit,
 	TEXT("Waits for the GPU to be idle on every submit. Useful for tracking GPU hangs.\n")
