UE4.21 - Fix for random crashes toggling fullscreen on Vulkan games on Linux
#rnx

diff --git a/Engine/Source/Runtime/VulkanRHI/Private/VulkanSwapChain.cpp b/Engine/Source/Runtime/VulkanRHI/Private/VulkanSwapChain.cpp
--- a/Engine/Source/Runtime/VulkanRHI/Private/VulkanSwapChain.cpp
+++ b/Engine/Source/Runtime/VulkanRHI/Private/VulkanSwapChain.cpp
@@ -475,6 +475,10 @@ void FVulkanSwapChain::Destroy()
 {
 	check(FVulkanPlatform::SupportsStandardSwapchain());
 
+	// We could be responding to an OUT_OF_DATE event and the GPU might not be done with swapchain image, so wait for idle.
+	// Alternatively could also check on the fence(s) for the image(s) from the swapchain but then timing out/waiting could become an issue.
+	Device.WaitUntilIdle();
+
 	VulkanRHI::vkDestroySwapchainKHR(Device.GetInstanceHandle(), SwapChain, VULKAN_CPU_ALLOCATOR);
 	SwapChain = VK_NULL_HANDLE;
 
