UE4.21 - Fix for Vulkan tracking layer (for debugging memory leaks)
#rnx

diff --git a/Engine/Source/Runtime/VulkanRHI/Private/VulkanCommandWrappers.h b/Engine/Source/Runtime/VulkanRHI/Private/VulkanCommandWrappers.h
--- a/Engine/Source/Runtime/VulkanRHI/Private/VulkanCommandWrappers.h
+++ b/Engine/Source/Runtime/VulkanRHI/Private/VulkanCommandWrappers.h
@@ -6,7 +6,7 @@
 
 #pragma once 
 
-#if VULKAN_ENABLE_CUSTOM_LAYER
+#if VULKAN_ENABLE_WRAP_LAYER
 	#define  VULKAN_LAYER_BODY		;
 	#define	VULKAN_EXTERN_EXPORT	VULKANRHI_API
 #else
diff --git a/Engine/Source/Runtime/VulkanRHI/Private/VulkanDebug.cpp b/Engine/Source/Runtime/VulkanRHI/Private/VulkanDebug.cpp
--- a/Engine/Source/Runtime/VulkanRHI/Private/VulkanDebug.cpp
+++ b/Engine/Source/Runtime/VulkanRHI/Private/VulkanDebug.cpp
@@ -412,8 +412,7 @@ void FVulkanDynamicRHI::RemoveDebugLayerCallback()
 }
 
 
-#if VULKAN_ENABLE_TRACKING_LAYER
-
+#if VULKAN_ENABLE_TRACKING_LAYER || VULKAN_ENABLE_DUMP_LAYER
 template <typename TResourceCreateInfoType>
 struct TTrackingResource
 {
@@ -441,7 +440,7 @@ static void CaptureCallStack(FString& OutCallstack, uint32 Delta)
 #endif
 #endif
 
-#if VULKAN_ENABLE_IMAGE_TRACKING_LAYER
+#if VULKAN_ENABLE_IMAGE_TRACKING_LAYER || VULKAN_ENABLE_DUMP_LAYER
 struct FTrackingImage
 {
 	struct
@@ -497,22 +496,49 @@ struct FTrackingImage
 };
 static TMap<VkImage, FTrackingImage> GVulkanTrackingImageLayouts;
 static TMap<VkImageView, TTrackingResource<VkImageViewCreateInfo>> GVulkanTrackingImageViews;
+static VkImage GBreakOnTrackImage = VK_NULL_HANDLE;
+
+static FORCEINLINE void BreakOnTrackingImage(VkImage InImage)
+{
+	if (GBreakOnTrackImage != VK_NULL_HANDLE)
+	{
+		ensureAlways(InImage != GBreakOnTrackImage);
+	}
+}
+
+static VkImage FindTrackingImage(VkImageView InView)
+{
+	const auto& Found = GVulkanTrackingImageViews.FindChecked(InView);
+	return Found.CreateInfo.image;
+}
+
+static FORCEINLINE void BreakOnTrackingImageView(VkImageView InView)
+{
+	BreakOnTrackingImage(FindTrackingImage(InView));
+}
 #endif
 
-#if VULKAN_ENABLE_BUFFER_TRACKING_LAYER
+#if VULKAN_ENABLE_BUFFER_TRACKING_LAYER || VULKAN_ENABLE_DUMP_LAYER
 static TMap<VkBuffer, TTrackingResource<VkBufferCreateInfo>> GVulkanTrackingBuffers;
 static TMap<VkBuffer, TArray<VkBufferView>> GVulkanTrackingBufferToBufferViews;
 static TMap<VkBufferView, TTrackingResource<VkBufferViewCreateInfo>> GVulkanTrackingBufferViews;
+
+static VkBuffer FindTrackingBuffer(VkBufferView InView)
+{
+	const auto& Found = GVulkanTrackingBufferViews.FindChecked(InView);
+	return Found.CreateInfo.buffer;
+}
 #endif
 
+static void ValidationFail()
+{
+	ensure(0);
+}
 
 #if VULKAN_ENABLE_DUMP_LAYER
 #include "Misc/OutputDeviceRedirector.h"
 namespace VulkanRHI
 {
-	void TrackImageViewAdd(VkImageView View, const VkImageViewCreateInfo* CreateInfo);
-	void TrackBufferViewAdd(VkBufferView View, const VkBufferViewCreateInfo* CreateInfo);
-
 	static FCriticalSection CS;
 	struct FMutexString
 	{
@@ -575,33 +601,6 @@ namespace VulkanRHI
 
 	static const TCHAR* Tabs = TEXT("\t\t\t\t\t\t\t\t\t");
 
-	static FCriticalSection GTrackLock;
-	static TMap<VkImageView, VkImageViewCreateInfo> GImageViewTracker;
-	static TMap<VkBufferView, VkBufferViewCreateInfo> GBufferViewTracker;
-	// Image/Layer/Mip
-	struct FImageLayout
-	{
-		void Init(uint32 NumLayers, uint32 NumMips)
-		{
-			ensure(Layouts.Num() == 0);
-			Layouts.Empty(NumLayers);
-
-			TArray<VkImageLayout> Temp;
-			Temp.AddUninitialized(NumMips);
-			for (uint32 Index = 0; Index < NumMips; ++Index)
-			{
-				Temp[Index] = VK_IMAGE_LAYOUT_UNDEFINED;
-			}
-
-			for (uint32 Index = 0; Index < NumLayers; ++Index)
-			{
-				Layouts.Add(Temp);
-			}
-		}
-
-		TArray<TArray<VkImageLayout>> Layouts;
-	};
-	static TMap<VkImage, FImageLayout> GOLDImageLayoutTracker;
 	struct FRenderPassInfo
 	{
 		TArray<VkAttachmentDescription> Descriptions;
@@ -615,12 +614,6 @@ namespace VulkanRHI
 	};
 	static TMap<VkFramebuffer, FFBInfo> GFramebufferInfo;
 
-	static void ValidationFail()
-	{
-		int i =0;
-		++i;
-	}
-
 	void FlushDebugWrapperLog()
 	{
 		if (DebugLog.Len() > 0)
@@ -1328,7 +1321,7 @@ namespace VulkanRHI
 }
 #endif
 
-#if VULKAN_ENABLE_CUSTOM_LAYER
+#if VULKAN_ENABLE_WRAP_LAYER
 void FWrapLayer::GetPhysicalDeviceMemoryProperties(VkResult Result, VkPhysicalDevice PhysicalDevice, VkPhysicalDeviceMemoryProperties* Properties)
 {
 	if (Result == VK_RESULT_MAX_ENUM)
@@ -1461,7 +1454,6 @@ void FWrapLayer::CreateBufferView(VkResult Result, VkDevice Device, const VkBuff
 	else
 	{
 #if VULKAN_ENABLE_DUMP_LAYER
-		TrackBufferViewAdd(*BufferView, CreateInfo);
 		PrintResultAndNamedHandle(Result, TEXT("BufferView"), *BufferView);
 		FlushDebugWrapperLog();
 #endif
@@ -1497,8 +1489,6 @@ void FWrapLayer::CreateImage(VkResult Result, VkDevice Device, const VkImageCrea
 	{
 #if VULKAN_ENABLE_DUMP_LAYER
 		PrintResultAndNamedHandle(Result, TEXT("Image"), *Image);
-		FImageLayout& Found = GOLDImageLayoutTracker.FindOrAdd(*Image);
-		Found.Init(FMath::Max(CreateInfo->arrayLayers, CreateInfo->extent.depth), CreateInfo->mipLevels);
 		FlushDebugWrapperLog();
 #endif
 
@@ -1525,12 +1515,12 @@ void FWrapLayer::DestroyImage(VkResult Result, VkDevice Device, VkImage Image)
 	{
 #if VULKAN_ENABLE_DUMP_LAYER
 		DevicePrintfBegin(Device, FString::Printf(TEXT("vkDestroyImage(Image=0x%p)"), Image));
-		GOLDImageLayoutTracker.Remove(Image);
 #endif
 	}
 	else
 	{
 #if VULKAN_ENABLE_IMAGE_TRACKING_LAYER
+		BreakOnTrackingImage(Image);
 		{
 			FScopeLock ScopeLock(&GTrackingCS);
 			int32 NumRemoved = GVulkanTrackingImageLayouts.Remove(Image);
@@ -1554,7 +1544,6 @@ void FWrapLayer::CreateImageView(VkResult Result, VkDevice Device, const VkImage
 	else
 	{
 #if VULKAN_ENABLE_DUMP_LAYER
-		TrackImageViewAdd(*ImageView, CreateInfo);
 		PrintResultAndNamedHandle(Result, TEXT("ImageView"), *ImageView);
 		FlushDebugWrapperLog();
 #endif
@@ -1567,6 +1556,7 @@ void FWrapLayer::CreateImageView(VkResult Result, VkDevice Device, const VkImage
 			CaptureCallStack(Found.CreateCallstack, 3);
 #endif
 		}
+		BreakOnTrackingImageView(*ImageView);
 #endif
 	}
 }
@@ -1837,6 +1827,7 @@ void FWrapLayer::PipelineBarrier(VkResult Result, VkCommandBuffer CommandBuffer,
 			FScopeLock ScopeLock(&GTrackingCS);
 			for (uint32 Index = 0; Index < ImageMemoryBarrierCount; ++Index)
 			{
+				BreakOnTrackingImage(ImageMemoryBarriers[Index].image);
 				FTrackingImage* TrackingImage = GVulkanTrackingImageLayouts.Find(ImageMemoryBarriers[Index].image);
 				check(TrackingImage);
 #if VULKAN_ENABLE_TRACKING_CALLSTACK
@@ -1954,24 +1945,30 @@ void FWrapLayer::UpdateDescriptorSets(VkResult Result, VkDevice Device, uint32 D
 	{
 #if VULKAN_ENABLE_DUMP_LAYER
 		DevicePrintfBegin(Device, FString::Printf(TEXT("vkUpdateDescriptorSets(NumWrites=%d, Writes=0x%p, NumCopies=%d, Copies=0x%p)"), DescriptorWriteCount, DescriptorWrites, DescriptorCopyCount, DescriptorCopies));
+#endif
 		for (uint32 Index = 0; Index < DescriptorWriteCount; ++Index)
 		{
+#if VULKAN_ENABLE_DUMP_LAYER
 			DebugLog += FString::Printf(TEXT("%sWrite[%d]: Set=0x%p Binding=%d DstArrayElem=%d NumDesc=%d DescType=%s "), Tabs, Index,
 				DescriptorWrites[Index].dstSet, DescriptorWrites[Index].dstBinding, DescriptorWrites[Index].dstArrayElement, DescriptorWrites[Index].descriptorCount, *GetDescriptorTypeString(DescriptorWrites[Index].descriptorType));
-
+#endif
 			switch (DescriptorWrites[Index].descriptorType)
 			{
 			case VK_DESCRIPTOR_TYPE_UNIFORM_BUFFER:
 			case VK_DESCRIPTOR_TYPE_STORAGE_BUFFER:
 			case VK_DESCRIPTOR_TYPE_UNIFORM_BUFFER_DYNAMIC:
 			case VK_DESCRIPTOR_TYPE_STORAGE_BUFFER_DYNAMIC:
+#if VULKAN_ENABLE_DUMP_LAYER
 				DebugLog += FString::Printf(TEXT("pBufferInfo=0x%p\n"), DescriptorWrites[Index].pBufferInfo);
+#endif
 				if (DescriptorWrites[Index].pBufferInfo)
 				{
 					for (uint32 SubIndex = 0; SubIndex < DescriptorWrites[Index].descriptorCount; ++SubIndex)
 					{
+#if VULKAN_ENABLE_DUMP_LAYER
 						DebugLog += FString::Printf(TEXT("%s\tpBufferInfo[%d]: buffer=0x%p, offset=%d, range=%d\n"), Tabs, SubIndex,
 							DescriptorWrites[Index].pBufferInfo->buffer, (int32)DescriptorWrites[Index].pBufferInfo->offset, (int32)DescriptorWrites[Index].pBufferInfo->range);
+#endif
 					}
 				}
 				else
@@ -1982,12 +1979,16 @@ void FWrapLayer::UpdateDescriptorSets(VkResult Result, VkDevice Device, uint32 D
 
 			case VK_DESCRIPTOR_TYPE_UNIFORM_TEXEL_BUFFER:
 			case VK_DESCRIPTOR_TYPE_STORAGE_TEXEL_BUFFER:
+#if VULKAN_ENABLE_DUMP_LAYER
 				DebugLog += FString::Printf(TEXT("pTexelBufferView=0x%p\n"), DescriptorWrites[Index].pTexelBufferView);
+#endif
 				if (DescriptorWrites[Index].pTexelBufferView)
 				{
 					for (uint32 SubIndex = 0; SubIndex < DescriptorWrites[Index].descriptorCount; ++SubIndex)
 					{
-						DebugLog += FString::Printf(TEXT("%s\tpTexelBufferView[%d]=0x%p(B:0x%p)\n"), Tabs, SubIndex, DescriptorWrites[Index].pTexelBufferView[SubIndex], GBufferViewTracker.FindChecked(DescriptorWrites[Index].pTexelBufferView[SubIndex]).buffer);
+#if VULKAN_ENABLE_DUMP_LAYER
+						DebugLog += FString::Printf(TEXT("%s\tpTexelBufferView[%d]=0x%p(B:0x%p)\n"), Tabs, SubIndex, DescriptorWrites[Index].pTexelBufferView[SubIndex], FindTrackingBuffer(DescriptorWrites[Index].pTexelBufferView[SubIndex]));
+#endif
 					}
 				}
 				else
@@ -2002,13 +2003,20 @@ void FWrapLayer::UpdateDescriptorSets(VkResult Result, VkDevice Device, uint32 D
 			case VK_DESCRIPTOR_TYPE_STORAGE_IMAGE:
 			case VK_DESCRIPTOR_TYPE_INPUT_ATTACHMENT:
 			default:
+#if VULKAN_ENABLE_DUMP_LAYER
 				DebugLog += FString::Printf(TEXT("pImageInfo=0x%p\n"), DescriptorWrites[Index].pImageInfo);
+#endif
 				if (DescriptorWrites[Index].pImageInfo)
 				{
 					for (uint32 SubIndex = 0; SubIndex < DescriptorWrites[Index].descriptorCount; ++SubIndex)
 					{
+#if VULKAN_ENABLE_IMAGE_TRACKING_LAYER
+						BreakOnTrackingImageView(DescriptorWrites[Index].pImageInfo->imageView);
+#endif
+#if VULKAN_ENABLE_DUMP_LAYER
 						DebugLog += FString::Printf(TEXT("%s\tpImageInfo[%d]: Sampler=0x%p, ImageView=0x%p(I:0x%p), imageLayout=%s\n"), Tabs, SubIndex,
-							DescriptorWrites[Index].pImageInfo->sampler, DescriptorWrites[Index].pImageInfo->imageView, GImageViewTracker.FindChecked(DescriptorWrites[Index].pImageInfo->imageView).image, *GetImageLayoutString(DescriptorWrites[Index].pImageInfo->imageLayout));
+							DescriptorWrites[Index].pImageInfo->sampler, DescriptorWrites[Index].pImageInfo->imageView, FindTrackingImage(DescriptorWrites[Index].pImageInfo->imageView), *GetImageLayoutString(DescriptorWrites[Index].pImageInfo->imageLayout));
+#endif
 					}
 				}
 				else
@@ -2018,6 +2026,7 @@ void FWrapLayer::UpdateDescriptorSets(VkResult Result, VkDevice Device, uint32 D
 				break;
 			}
 		}
+#if VULKAN_ENABLE_DUMP_LAYER
 		FlushDebugWrapperLog();
 #endif
 	}
@@ -2044,7 +2053,7 @@ void FWrapLayer::CreateFramebuffer(VkResult Result, VkDevice Device, const VkFra
 		DebugLog += FString::Printf(TEXT("%sVkFramebufferCreateInfo: Flags=%d, RenderPass=0x%p, NumAttachments=%d\n"), Tabs, CreateInfo->flags, CreateInfo->renderPass, CreateInfo->attachmentCount);
 		for (uint32 Index = 0; Index < CreateInfo->attachmentCount; ++Index)
 		{
-			DebugLog += FString::Printf(TEXT("%s\tAttachment[%d]: ImageView=0x%p(I:0x%p)\n"), Tabs, Index, CreateInfo->pAttachments[Index], GImageViewTracker.FindChecked(CreateInfo->pAttachments[Index]).image);
+			DebugLog += FString::Printf(TEXT("%s\tAttachment[%d]: ImageView=0x%p(I:0x%p)\n"), Tabs, Index, CreateInfo->pAttachments[Index], FindTrackingImage(CreateInfo->pAttachments[Index]));
 		}
 		DebugLog += FString::Printf(TEXT("%s\twidth=%d, height=%d, layers=%d\n"), Tabs, CreateInfo->width, CreateInfo->height, CreateInfo->layers);
 #endif
@@ -2491,11 +2500,11 @@ void FWrapLayer::BeginRenderPass(VkResult Result, VkCommandBuffer CommandBuffer,
 				for (uint32 Index = 0; Index < FoundFBInfo->Info.attachmentCount; ++Index)
 				{
 					VkImageView View = FoundFBInfo->Attachments[Index];
-					VkImageViewCreateInfo* FoundImageInfo = GImageViewTracker.Find(View);
-					ensure(FoundImageInfo);
+					auto* FoundImageInfo = GVulkanTrackingImageViews.Find(View);
+					// Can be null for swapchain images!
 					if (FoundImageInfo)
 					{
-						DebugLog += FString::Printf(TEXT("%s\t\tAttachment[%d]: ImageView=0x%p(I:0x%p)\n"), Tabs, Index, View, FoundImageInfo->image);
+						DebugLog += FString::Printf(TEXT("%s\t\tAttachment[%d]: ImageView=0x%p(I:0x%p)\n"), Tabs, Index, View, FoundImageInfo->CreateInfo.image);
 					}
 				}
 			}
@@ -2666,6 +2675,8 @@ void FWrapLayer::CopyImage(VkResult Result, VkCommandBuffer CommandBuffer, VkIma
 		CmdPrintfBegin(CommandBuffer, FString::Printf(TEXT("vkCmdCopyImage(SrcImage=0x%p, SrcImageLayout=%d, DstImage=0x%p, DstImageLayout=%d, RegionCount=%d, Regions=0x%p)[...]"), SrcImage, (int32)SrcImageLayout, DstImage, (int32)DstImageLayout, RegionCount, Regions));
 #endif
 #if VULKAN_ENABLE_IMAGE_TRACKING_LAYER
+		BreakOnTrackingImage(SrcImage);
+		BreakOnTrackingImage(DstImage);
 		FScopeLock ScopeLock(&GTrackingCS);
 		FTrackingImage* FoundSrc = GVulkanTrackingImageLayouts.Find(SrcImage);
 		FTrackingImage* FoundDest = GVulkanTrackingImageLayouts.Find(DstImage);
@@ -2702,6 +2713,7 @@ void FWrapLayer::CopyBufferToImage(VkResult Result, VkCommandBuffer CommandBuffe
 		FlushDebugWrapperLog();
 #endif
 #if VULKAN_ENABLE_IMAGE_TRACKING_LAYER
+		BreakOnTrackingImage(DstImage);
 		FScopeLock ScopeLock(&GTrackingCS);
 		FTrackingImage* FoundDest = GVulkanTrackingImageLayouts.Find(DstImage);
 		ensure(FoundDest);
@@ -2761,6 +2773,8 @@ void FWrapLayer::BlitImage(VkResult Result, VkCommandBuffer CommandBuffer, VkIma
 		FlushDebugWrapperLog();
 #endif
 #if VULKAN_ENABLE_IMAGE_TRACKING_LAYER
+		BreakOnTrackingImage(SrcImage);
+		BreakOnTrackingImage(DstImage);
 		FScopeLock ScopeLock(&GTrackingCS);
 		FTrackingImage* FoundSrc = GVulkanTrackingImageLayouts.Find(SrcImage);
 		FTrackingImage* FoundDest = GVulkanTrackingImageLayouts.Find(DstImage);
@@ -2824,8 +2838,6 @@ void FWrapLayer::GetSwapChainImagesKHR(VkResult Result, VkDevice Device, VkSwapc
 			for (uint32 Index = 0; Index < *SwapchainImageCount; ++Index)
 			{
 				DebugLog += FString::Printf(TEXT("%sImage[%d]=0x%p\n"), Tabs, Index, SwapchainImages[Index]);
-				FImageLayout& Layout = GOLDImageLayoutTracker.Add(SwapchainImages[Index]);
-				Layout.Init(1, 1);
 			}
 		}
 		else
@@ -2839,6 +2851,7 @@ void FWrapLayer::GetSwapChainImagesKHR(VkResult Result, VkDevice Device, VkSwapc
 			FScopeLock ScopeLock(&GTrackingCS);
 			for (uint32 Index = 0; Index < *SwapchainImageCount; ++Index)
 			{
+				BreakOnTrackingImage(SwapchainImages[Index]);
 				FTrackingImage& TrackingImage = GVulkanTrackingImageLayouts.FindOrAdd(SwapchainImages[Index]);
 				TrackingImage.Setup(1, 1, VK_IMAGE_LAYOUT_UNDEFINED, true);
 			}
@@ -3300,6 +3313,7 @@ void FWrapLayer::DestroyImageView(VkResult Result, VkDevice Device, VkImageView
 #if VULKAN_ENABLE_IMAGE_TRACKING_LAYER
 		{
 			FScopeLock ScopeLock(&GTrackingCS);
+			BreakOnTrackingImageView(ImageView);
 			int32 NumRemoved = GVulkanTrackingImageViews.Remove(ImageView);
 			ensure(NumRemoved > 0);
 		}
@@ -3862,6 +3876,7 @@ void FWrapLayer::CreateWin32SurfaceKHR(VkResult Result, VkInstance Instance, con
 #if VULKAN_SUPPORTS_COLOR_CONVERSIONS
 void FWrapLayer::CreateSamplerYcbcrConversionKHR(VkResult Result, VkDevice Device, const VkSamplerYcbcrConversionCreateInfo* CreateInfo, VkSamplerYcbcrConversion* YcbcrConversion)
 {
+#if VULKAN_ENABLE_DUMP_LAYER
 	if (Result == VK_RESULT_MAX_ENUM)
 	{
 		DevicePrintfBeginResult(Device, FString::Printf(TEXT("vkCreateSamplerYcbcrConversionKHR(CreateInfo=0x%p, YcbcrConversion=0x%p)[...]"), CreateInfo, YcbcrConversion));
@@ -3871,6 +3886,7 @@ void FWrapLayer::CreateSamplerYcbcrConversionKHR(VkResult Result, VkDevice Devic
 		PrintResultAndNamedHandle(Result, TEXT("SamplerYcbcrConversionKHR"), *YcbcrConversion);
 	}
 	FlushDebugWrapperLog();
+#endif
 }
 
 void FWrapLayer::DestroySamplerYcbcrConversionKHR(VkResult Result, VkDevice Device, VkSamplerYcbcrConversion YcbcrConversion)
@@ -3903,36 +3919,11 @@ namespace VulkanRHI
 }
 #endif
 
-#endif	// VULKAN_ENABLE_CUSTOM_LAYER
+#endif	// VULKAN_ENABLE_WRAP_LAYER
 
 #if VULKAN_ENABLE_DUMP_LAYER
 namespace VulkanRHI
 {
-	void TrackImageViewAdd(VkImageView View, const VkImageViewCreateInfo* CreateInfo)
-	{
-		FScopeLock Lock(&GTrackLock);
-		GImageViewTracker.Add(View, *CreateInfo);
-	}
-
-	void TrackImageViewRemove(VkImageView View)
-	{
-		FScopeLock Lock(&GTrackLock);
-		GImageViewTracker.Remove(View);
-	}
-
-	void TrackBufferViewAdd(VkBufferView View, const VkBufferViewCreateInfo* CreateInfo)
-	{
-		FScopeLock Lock(&GTrackLock);
-		GBufferViewTracker.Add(View, *CreateInfo);
-	}
-
-	void TrackBufferViewRemove(VkBufferView View)
-	{
-		FScopeLock Lock(&GTrackLock);
-		GBufferViewTracker.Remove(View);
-	}
-
-
 	static struct FGlobalDumpLog
 	{
 		~FGlobalDumpLog()
diff --git a/Engine/Source/Runtime/VulkanRHI/Public/VulkanConfiguration.h b/Engine/Source/Runtime/VulkanRHI/Public/VulkanConfiguration.h
--- a/Engine/Source/Runtime/VulkanRHI/Public/VulkanConfiguration.h
+++ b/Engine/Source/Runtime/VulkanRHI/Public/VulkanConfiguration.h
@@ -33,6 +33,7 @@
 #ifndef VULKAN_ENABLE_DUMP_LAYER
 	#define VULKAN_ENABLE_DUMP_LAYER							0
 #endif
+
 #define VULKAN_ENABLE_DRAW_MARKERS								VULKAN_SHOULD_ENABLE_DRAW_MARKERS
 
 #ifndef VULKAN_ENABLE_IMAGE_TRACKING_LAYER
@@ -44,7 +45,7 @@
 #endif
 
 #define VULKAN_ENABLE_TRACKING_LAYER							(VULKAN_ENABLE_BUFFER_TRACKING_LAYER || VULKAN_ENABLE_IMAGE_TRACKING_LAYER)
-#define VULKAN_ENABLE_CUSTOM_LAYER								(VULKAN_ENABLE_DUMP_LAYER || VULKAN_ENABLE_TRACKING_LAYER)
+#define VULKAN_ENABLE_WRAP_LAYER								(VULKAN_ENABLE_DUMP_LAYER || VULKAN_ENABLE_TRACKING_LAYER)
 
 #define VULKAN_HASH_POOLS_WITH_TYPES_USAGE_ID					1
 
@@ -71,7 +72,7 @@
 #endif
 
 #ifndef VULKAN_SHOULD_USE_COMMANDWRAPPERS
-	#define VULKAN_SHOULD_USE_COMMANDWRAPPERS					VULKAN_ENABLE_CUSTOM_LAYER
+	#define VULKAN_SHOULD_USE_COMMANDWRAPPERS					VULKAN_ENABLE_WRAP_LAYER
 #endif
 
 #ifndef VULKAN_COMMANDWRAPPERS_ENABLE							
