UE4.21 - Fix for mesh paint on Vulkan
#rnx

diff --git a/Engine/Source/Runtime/VulkanRHI/Private/VulkanRenderTarget.cpp b/Engine/Source/Runtime/VulkanRHI/Private/VulkanRenderTarget.cpp
--- a/Engine/Source/Runtime/VulkanRHI/Private/VulkanRenderTarget.cpp
+++ b/Engine/Source/Runtime/VulkanRHI/Private/VulkanRenderTarget.cpp
@@ -770,7 +770,7 @@ void FVulkanDynamicRHI::RHIReadSurfaceData(FTextureRHIParamRef TextureRHI, FIntR
 
 	FVulkanCmdBuffer* CmdBuffer = ImmediateContext.GetCommandBufferManager()->GetUploadCmdBuffer();
 
-	ensure(Texture2D->Surface.StorageFormat == VK_FORMAT_R8G8B8A8_UNORM || Texture2D->Surface.StorageFormat == VK_FORMAT_B8G8R8A8_UNORM || Texture2D->Surface.StorageFormat == VK_FORMAT_R16G16B16A16_SFLOAT || Texture2D->Surface.StorageFormat == VK_FORMAT_A2B10G10R10_UNORM_PACK32);
+	ensure(Texture2D->Surface.StorageFormat == VK_FORMAT_R8G8B8A8_UNORM || Texture2D->Surface.StorageFormat == VK_FORMAT_B8G8R8A8_UNORM || Texture2D->Surface.StorageFormat == VK_FORMAT_R16G16B16A16_SFLOAT || Texture2D->Surface.StorageFormat == VK_FORMAT_A2B10G10R10_UNORM_PACK32 || Texture2D->Surface.StorageFormat == VK_FORMAT_R16G16B16A16_UNORM);
 	const bool bIs8Bpp = (Texture2D->Surface.StorageFormat == VK_FORMAT_R16G16B16A16_SFLOAT);
 	const uint32 Size = NumPixels * sizeof(FColor) * (bIs8Bpp ? 2 : 1);
 	VulkanRHI::FStagingBuffer* StagingBuffer = Device->GetStagingManager().AcquireBuffer(Size, VK_BUFFER_USAGE_TRANSFER_DST_BIT, true);
@@ -892,6 +892,31 @@ void FVulkanDynamicRHI::RHIReadSurfaceData(FTextureRHIParamRef TextureRHI, FIntR
 			}
 		}
 	}
+	else if (Texture2D->Surface.StorageFormat == VK_FORMAT_R16G16B16A16_UNORM)
+	{
+		struct FRGBA16U
+		{
+			uint16 R;
+			uint16 G;
+			uint16 B;
+			uint16 A;
+		};
+		for (int32 Row = Rect.Min.Y; Row < Rect.Max.Y; ++Row)
+		{
+			FRGBA16U* Src = (FRGBA16U*)StagingBuffer->GetMappedPointer() + Row * TextureRHI2D->GetSizeX() + Rect.Min.X;
+			for (int32 Col = Rect.Min.X; Col < Rect.Max.X; ++Col)
+			{
+				*Dest = FLinearColor(
+					(float)Src->R / 65535.0f,
+					(float)Src->G / 65535.0f,
+					(float)Src->B / 65535.0f,
+					(float)Src->A / 65535.0f
+				).Quantize();
+				Dest++;
+				Src++;
+			}
+		}
+	}
 	else if (Texture2D->Surface.StorageFormat == VK_FORMAT_B8G8R8A8_UNORM)
 	{
 		FColor* Src = (FColor*)StagingBuffer->GetMappedPointer() + Rect.Min.Y * TextureRHI2D->GetSizeX() + Rect.Min.X;
