Linux: fix Vulkan crash when alt-tabbing in fullscreen (UE-64312).

- Safeguard the access to the backbuffer pointer. RHI thread can be in the midst of recreating the swapchain while RenderingThread calls GetBackBuffer().

diff --git a/Engine/Source/Runtime/VulkanRHI/Private/VulkanViewport.cpp b/Engine/Source/Runtime/VulkanRHI/Private/VulkanViewport.cpp
--- a/Engine/Source/Runtime/VulkanRHI/Private/VulkanViewport.cpp
+++ b/Engine/Source/Runtime/VulkanRHI/Private/VulkanViewport.cpp
@@ -20,7 +20,6 @@ FAutoConsoleVariable GCVarDelayAcquireBackBuffer(
 	ECVF_ReadOnly
 );
 
-
 struct FRHICommandAcquireBackBuffer final : public FRHICommand<FRHICommandAcquireBackBuffer>
 {
 	FVulkanViewport* Viewport;
@@ -209,6 +208,9 @@ FVulkanTexture2D* FVulkanViewport::GetBackBuffer(FRHICommandList& RHICmdList)
 {
 	check(IsInRenderingThread());
 
+	// make sure we aren't in the middle of swapchain recreation (which can happen on e.g. RHI thread)
+	FScopeLock LockSwapchain(&RecreatingSwapchain);
+
 	if (!RenderingBackBuffer && FVulkanPlatform::SupportsStandardSwapchain())
 	{
 		check(!DelayAcquireBackBuffer());
@@ -443,6 +445,7 @@ void FVulkanViewport::RecreateSwapchain(void* NewNativeWindow, bool bForce)
 		return;
 	}
 
+	FScopeLock LockSwapchain(&RecreatingSwapchain);
 	RenderingBackBuffer = nullptr;
 	RHIBackBuffer = nullptr;
 
diff --git a/Engine/Source/Runtime/VulkanRHI/Public/VulkanViewport.h b/Engine/Source/Runtime/VulkanRHI/Public/VulkanViewport.h
--- a/Engine/Source/Runtime/VulkanRHI/Public/VulkanViewport.h
+++ b/Engine/Source/Runtime/VulkanRHI/Public/VulkanViewport.h
@@ -7,6 +7,7 @@
 #pragma once
 
 #include "VulkanResources.h"
+#include "HAL/CriticalSection.h"
 
 class FVulkanDynamicRHI;
 class FVulkanSwapChain;
@@ -66,6 +67,9 @@ protected:
 	TRefCountPtr<FVulkanBackBuffer> RenderingBackBuffer;
 	TRefCountPtr<FVulkanBackBuffer> RHIBackBuffer;
 
+	/** narrow-scoped section that locks access to back buffer during its recreation*/
+	FCriticalSection RecreatingSwapchain;
+
 	FVulkanDynamicRHI* RHI;
 	uint32 SizeX;
 	uint32 SizeY;
