UE4.21 - Add missing time for textures on Vulkan
#rnx

diff --git a/Engine/Source/Runtime/VulkanRHI/Private/VulkanCommands.cpp b/Engine/Source/Runtime/VulkanRHI/Private/VulkanCommands.cpp
--- a/Engine/Source/Runtime/VulkanRHI/Private/VulkanCommands.cpp
+++ b/Engine/Source/Runtime/VulkanRHI/Private/VulkanCommands.cpp
@@ -192,6 +192,7 @@ void FVulkanCommandListContext::RHISetShaderTexture(FVertexShaderRHIParamRef Ver
 	FVulkanTextureBase* Texture = GetVulkanTextureFromRHITexture(NewTextureRHI);
 	VkImageLayout Layout = GetLayoutForDescriptor(Texture->Surface);
 	PendingGfxState->SetTextureForStage(ShaderStage::Vertex, TextureIndex, Texture, Layout);
+	NewTextureRHI->SetLastRenderTime((float)FPlatformTime::Seconds());
 }
 
 void FVulkanCommandListContext::RHISetShaderTexture(FHullShaderRHIParamRef HullShaderRHI, uint32 TextureIndex, FTextureRHIParamRef NewTextureRHI)
@@ -202,6 +203,7 @@ void FVulkanCommandListContext::RHISetShaderTexture(FHullShaderRHIParamRef HullS
 	FVulkanTextureBase* Texture = GetVulkanTextureFromRHITexture(NewTextureRHI);
 	VkImageLayout Layout = GetLayoutForDescriptor(Texture->Surface);
 	PendingGfxState->SetTexture(ShaderStage::Hull, TextureIndex, Texture, Layout);
+	NewTextureRHI->SetLastRenderTime((float)FPlatformTime::Seconds());
 */
 }
 
@@ -213,6 +215,7 @@ void FVulkanCommandListContext::RHISetShaderTexture(FDomainShaderRHIParamRef Dom
 	FVulkanTextureBase* Texture = GetVulkanTextureFromRHITexture(NewTextureRHI);
 	VkImageLayout Layout = GetLayoutForDescriptor(Texture->Surface);
 	PendingGfxState->SetTexture(ShaderStage::Domain, TextureIndex, Texture, Layout);
+	NewTextureRHI->SetLastRenderTime((float)FPlatformTime::Seconds());
 */
 }
 
@@ -223,6 +226,7 @@ void FVulkanCommandListContext::RHISetShaderTexture(FGeometryShaderRHIParamRef G
 	FVulkanTextureBase* Texture = GetVulkanTextureFromRHITexture(NewTextureRHI);
 	VkImageLayout Layout = GetLayoutForDescriptor(Texture->Surface);
 	PendingGfxState->SetTextureForStage(ShaderStage::Geometry, TextureIndex, Texture, Layout);
+	NewTextureRHI->SetLastRenderTime((float)FPlatformTime::Seconds());
 #else
 	ensureMsgf(0, TEXT("Geometry not supported!"));
 #endif
@@ -234,6 +238,7 @@ void FVulkanCommandListContext::RHISetShaderTexture(FPixelShaderRHIParamRef Pixe
 	FVulkanTextureBase* Texture = GetVulkanTextureFromRHITexture(NewTextureRHI);
 	VkImageLayout Layout = GetLayoutForDescriptor(Texture->Surface);
 	PendingGfxState->SetTextureForStage(ShaderStage::Pixel, TextureIndex, Texture, Layout);
+	NewTextureRHI->SetLastRenderTime((float)FPlatformTime::Seconds());
 }
 
 void FVulkanCommandListContext::RHISetShaderTexture(FComputeShaderRHIParamRef ComputeShaderRHI, uint32 TextureIndex, FTextureRHIParamRef NewTextureRHI)
@@ -244,6 +249,7 @@ void FVulkanCommandListContext::RHISetShaderTexture(FComputeShaderRHIParamRef Co
 	FVulkanTextureBase* VulkanTexture = GetVulkanTextureFromRHITexture(NewTextureRHI);
 	VkImageLayout Layout = GetLayoutForDescriptor(VulkanTexture->Surface);
 	PendingComputeState->SetTextureForStage(TextureIndex, VulkanTexture, Layout);
+	NewTextureRHI->SetLastRenderTime((float)FPlatformTime::Seconds());
 }
 
 void FVulkanCommandListContext::RHISetShaderResourceViewParameter(FVertexShaderRHIParamRef VertexShaderRHI, uint32 TextureIndex, FShaderResourceViewRHIParamRef SRVRHI)
@@ -484,7 +490,7 @@ template <typename TState>
 inline void /*FVulkanCommandListContext::*/SetShaderUniformBufferResources(FVulkanCommandListContext* Context, TState* State, const FVulkanShader* Shader, const TArray<FVulkanShaderHeader::FGlobalInfo>& GlobalInfos, const TArray<TEnumAsByte<VkDescriptorType>>& DescriptorTypes, const FVulkanShaderHeader::FUniformBufferInfo& HeaderUBInfo, const FVulkanUniformBuffer* UniformBuffer, const TArray<FDescriptorSetRemappingInfo::FRemappingInfo>& GlobalRemappingInfo)
 {
 	ensure(UniformBuffer->GetLayout().GetHash() == HeaderUBInfo.LayoutHash);
-	float CurrentTime = (float)FApp::GetCurrentTime();
+	float CurrentTime = (float)FPlatformTime::Seconds();
 	const TArray<TRefCountPtr<FRHIResource>>& ResourceArray = UniformBuffer->GetResourceTable();
 	for (int32 Index = 0; Index < HeaderUBInfo.ResourceEntries.Num(); ++Index)
 	{
@@ -569,7 +575,7 @@ inline void /*FVulkanCommandListContext::*/SetShaderUniformBufferResources(FVulk
 		GatherUniformBufferResources(ResourceBindingTable.ShaderResourceViewMap, ResourceBindingTable.ResourceTableBits, UniformBuffer, BindingIndex, SRVBindings);
 		if (CurrentTime == 0.0f)
 		{
-			CurrentTime = (float)FApp::GetCurrentTime();
+			CurrentTime = (float)FPlatformTime::Seconds();
 		}
 		for (int32 Index = 0; Index < SRVBindings.Num(); Index++)
 		{
