Linux: Vulkan full screen fix

diff --git a/Engine/Source/Runtime/VulkanRHI/Private/Linux/VulkanLinuxPlatform.cpp b/Engine/Source/Runtime/VulkanRHI/Private/Linux/VulkanLinuxPlatform.cpp
--- a/Engine/Source/Runtime/VulkanRHI/Private/Linux/VulkanLinuxPlatform.cpp
+++ b/Engine/Source/Runtime/VulkanRHI/Private/Linux/VulkanLinuxPlatform.cpp
@@ -146,3 +146,8 @@ void FVulkanLinuxPlatform::CreateSurface(void* WindowHandle, VkInstance Instance
 		check(0);
 	}
 }
+
+void FVulkanLinuxPlatform::UpdateWindowSize(void* WindowHandle, uint32& Width, uint32& Height)
+{
+	SDL_Vulkan_GetDrawableSize((SDL_Window*) WindowHandle, (int32*) &Width, (int32*) &Height);
+}
diff --git a/Engine/Source/Runtime/VulkanRHI/Private/Linux/VulkanLinuxPlatform.h b/Engine/Source/Runtime/VulkanRHI/Private/Linux/VulkanLinuxPlatform.h
--- a/Engine/Source/Runtime/VulkanRHI/Private/Linux/VulkanLinuxPlatform.h
+++ b/Engine/Source/Runtime/VulkanRHI/Private/Linux/VulkanLinuxPlatform.h
@@ -52,6 +52,8 @@ public:
 #endif
 	}
 
+	static void UpdateWindowSize(void* WindowHandle, uint32& Width, uint32& Height);
+
 protected:
 	static void* VulkanLib;
 	static bool bAttemptedLoad;
diff --git a/Engine/Source/Runtime/VulkanRHI/Private/VulkanGenericPlatform.h b/Engine/Source/Runtime/VulkanRHI/Private/VulkanGenericPlatform.h
--- a/Engine/Source/Runtime/VulkanRHI/Private/VulkanGenericPlatform.h
+++ b/Engine/Source/Runtime/VulkanRHI/Private/VulkanGenericPlatform.h
@@ -100,4 +100,7 @@ public:
 
 	// Ensure the last frame completed on the GPU
 	static bool RequiresWaitingForFrameCompletionEvent() { return true; }
+
+	// Some platforms imitate fullscreen mode by using bigger window
+	static void UpdateWindowSize(void* WindowHandle, uint32& Width, uint32& Height) {}
 };
diff --git a/Engine/Source/Runtime/VulkanRHI/Private/VulkanSwapChain.cpp b/Engine/Source/Runtime/VulkanRHI/Private/VulkanSwapChain.cpp
--- a/Engine/Source/Runtime/VulkanRHI/Private/VulkanSwapChain.cpp
+++ b/Engine/Source/Runtime/VulkanRHI/Private/VulkanSwapChain.cpp
@@ -386,6 +386,8 @@ FVulkanSwapChain::FVulkanSwapChain(VkInstance InInstance, FVulkanDevice& InDevic
 	uint32 SizeY = FVulkanPlatform::SupportsQuerySurfaceProperties() ? (SurfProperties.currentExtent.height == 0xFFFFFFFF ? Height : SurfProperties.currentExtent.height) : Height;
 	//FPlatformMisc::LowLevelOutputDebugStringf(TEXT("Create swapchain: %ux%u \n"), SizeX, SizeY);
 
+	// Some platforms cheat with the window size. Allow them to update.
+	FVulkanPlatform::UpdateWindowSize(WindowHandle, SizeX, SizeY);
 
 	VkSwapchainCreateInfoKHR SwapChainInfo;
 	ZeroVulkanStruct(SwapChainInfo, VK_STRUCTURE_TYPE_SWAPCHAIN_CREATE_INFO_KHR);
diff --git a/Engine/Source/Runtime/VulkanRHI/Private/VulkanViewport.cpp b/Engine/Source/Runtime/VulkanRHI/Private/VulkanViewport.cpp
--- a/Engine/Source/Runtime/VulkanRHI/Private/VulkanViewport.cpp
+++ b/Engine/Source/Runtime/VulkanRHI/Private/VulkanViewport.cpp
@@ -599,7 +599,7 @@ void FVulkanViewport::CreateSwapchain()
 	}
 }
 
-inline static void CopyImageToBackBuffer(FVulkanCmdBuffer* CmdBuffer, bool bSourceReadOnly, VkImage SrcSurface, VkImage DstSurface, int32 SizeX, int32 SizeY)
+inline static void CopyImageToBackBuffer(FVulkanCmdBuffer* CmdBuffer, bool bSourceReadOnly, VkImage SrcSurface, VkImage DstSurface, int32 SizeX, int32 SizeY, int32 WindowSizeX, int32 WindowSizeY)
 {
 	VulkanRHI::FPendingBarrier Barriers;
 	int32 SourceIndex = Barriers.AddImageBarrier(SrcSurface, VK_IMAGE_ASPECT_COLOR_BIT, 1);
@@ -614,23 +614,52 @@ inline static void CopyImageToBackBuffer(FVulkanCmdBuffer* CmdBuffer, bool bSour
 
 	VulkanRHI::DebugHeavyWeightBarrier(CmdBuffer->GetHandle(), 32);
 
-	VkImageCopy Region;
-	FMemory::Memzero(Region);
-	Region.extent.width = SizeX;
-	Region.extent.height = SizeY;
-	Region.extent.depth = 1;
-	Region.srcSubresource.aspectMask = VK_IMAGE_ASPECT_COLOR_BIT;
-	//Region.srcSubresource.baseArrayLayer = 0;
-	Region.srcSubresource.layerCount = 1;
-	//Region.srcSubresource.mipLevel = 0;
-	Region.dstSubresource.aspectMask = VK_IMAGE_ASPECT_COLOR_BIT;
-	//Region.dstSubresource.baseArrayLayer = 0;
-	Region.dstSubresource.layerCount = 1;
-	//Region.dstSubresource.mipLevel = 0;
-	VulkanRHI::vkCmdCopyImage(CmdBuffer->GetHandle(),
-		SrcSurface, VK_IMAGE_LAYOUT_TRANSFER_SRC_OPTIMAL,
-		DstSurface, VK_IMAGE_LAYOUT_TRANSFER_DST_OPTIMAL,
-		1, &Region);
+	if (SizeX != WindowSizeX || SizeY != WindowSizeY)
+	{
+		VkImageBlit Region;
+		FMemory::Memzero(Region);
+		Region.srcOffsets[0].x = 0;
+		Region.srcOffsets[0].y = 0;
+		Region.srcOffsets[0].z = 0;
+		Region.srcOffsets[1].x = SizeX;
+		Region.srcOffsets[1].y = SizeY;
+		Region.srcOffsets[1].z = 1;
+		Region.srcSubresource.aspectMask = VK_IMAGE_ASPECT_COLOR_BIT;
+		Region.srcSubresource.layerCount = 1;
+		Region.dstOffsets[0].x = 0;
+		Region.dstOffsets[0].y = 0;
+		Region.dstOffsets[0].z = 0;
+		Region.dstOffsets[1].x = WindowSizeX;
+		Region.dstOffsets[1].y = WindowSizeY;
+		Region.dstOffsets[1].z = 1;
+		Region.dstSubresource.aspectMask = VK_IMAGE_ASPECT_COLOR_BIT;
+		Region.dstSubresource.baseArrayLayer = 0;
+		Region.dstSubresource.layerCount = 1;
+		VulkanRHI::vkCmdBlitImage(CmdBuffer->GetHandle(),
+			SrcSurface, VK_IMAGE_LAYOUT_TRANSFER_SRC_OPTIMAL,
+			DstSurface, VK_IMAGE_LAYOUT_TRANSFER_DST_OPTIMAL,
+			1, &Region, VK_FILTER_LINEAR);
+	}
+	else
+	{
+		VkImageCopy Region;
+		FMemory::Memzero(Region);
+		Region.extent.width = SizeX;
+		Region.extent.height = SizeY;
+		Region.extent.depth = 1;
+		Region.srcSubresource.aspectMask = VK_IMAGE_ASPECT_COLOR_BIT;
+		//Region.srcSubresource.baseArrayLayer = 0;
+		Region.srcSubresource.layerCount = 1;
+		//Region.srcSubresource.mipLevel = 0;
+		Region.dstSubresource.aspectMask = VK_IMAGE_ASPECT_COLOR_BIT;
+		//Region.dstSubresource.baseArrayLayer = 0;
+		Region.dstSubresource.layerCount = 1;
+		//Region.dstSubresource.mipLevel = 0;
+		VulkanRHI::vkCmdCopyImage(CmdBuffer->GetHandle(),
+			SrcSurface, VK_IMAGE_LAYOUT_TRANSFER_SRC_OPTIMAL,
+			DstSurface, VK_IMAGE_LAYOUT_TRANSFER_DST_OPTIMAL,
+			1, &Region);
+	}
 
 	// Prepare for present
 	Barriers.ResetStages();
@@ -658,8 +687,12 @@ bool FVulkanViewport::Present(FVulkanCommandListContext* Context, FVulkanCmdBuff
 				UE_LOG(LogVulkanRHI, Fatal, TEXT("Swapchain acquire image index failed!"));
 			}
 
+			uint32 WindowSizeX = SizeX;
+			uint32 WindowSizeY = SizeY;
+			FVulkanPlatform::UpdateWindowSize(WindowHandle, WindowSizeX, WindowSizeY);
+
 			Context->RHIPushEvent(TEXT("CopyImageToBackBuffer"), FColor::Blue);
-			CopyImageToBackBuffer(CmdBuffer, true, RenderingBackBuffer->Surface.Image, BackBufferImages[AcquiredImageIndex], SizeX, SizeY);
+			CopyImageToBackBuffer(CmdBuffer, true, RenderingBackBuffer->Surface.Image, BackBufferImages[AcquiredImageIndex], SizeX, SizeY, WindowSizeX, WindowSizeY);
 			Context->RHIPopEvent();
 		}
 		else
