UE4.21 - Fix immutable samplers not getting in PSO hash
#rnx

diff --git a/Engine/Source/Runtime/VulkanRHI/Private/VulkanDescriptorSets.h b/Engine/Source/Runtime/VulkanRHI/Private/VulkanDescriptorSets.h
--- a/Engine/Source/Runtime/VulkanRHI/Private/VulkanDescriptorSets.h
+++ b/Engine/Source/Runtime/VulkanRHI/Private/VulkanDescriptorSets.h
@@ -286,7 +286,7 @@ public:
 	template<bool bIsCompute>
 	void FinalizeBindings(const FUniformBufferGatherInfo& UBGatherInfo, const TArrayView<const FSamplerStateRHIParamRef>& ImmutableSamplers);
 
-	void GenerateHash();
+	void GenerateHash(const TArrayView<const FSamplerStateRHIParamRef>& ImmutableSamplers);
 
 	friend uint32 GetTypeHash(const FVulkanDescriptorSetsLayoutInfo& In)
 	{
diff --git a/Engine/Source/Runtime/VulkanRHI/Private/VulkanPipeline.cpp b/Engine/Source/Runtime/VulkanRHI/Private/VulkanPipeline.cpp
--- a/Engine/Source/Runtime/VulkanRHI/Private/VulkanPipeline.cpp
+++ b/Engine/Source/Runtime/VulkanRHI/Private/VulkanPipeline.cpp
@@ -1033,6 +1033,13 @@ FArchive& operator << (FArchive& Ar, FVulkanPipelineStateCacheManager::FGfxPipel
 
 	Ar << Entry.RenderTargets;
 
+#if VULKAN_SUPPORTS_COLOR_CONVERSIONS
+	for (uint32 Index = 0; Index < MaxImmutableSamplers; ++Index)
+	{
+		Ar << Entry.ImmutableSamplers[Index];
+	}
+#endif
+
 	return Ar;
 }
 
diff --git a/Engine/Source/Runtime/VulkanRHI/Private/VulkanRHI.cpp b/Engine/Source/Runtime/VulkanRHI/Private/VulkanRHI.cpp
--- a/Engine/Source/Runtime/VulkanRHI/Private/VulkanRHI.cpp
+++ b/Engine/Source/Runtime/VulkanRHI/Private/VulkanRHI.cpp
@@ -1109,7 +1109,7 @@ void FVulkanDescriptorSetsLayoutInfo::AddDescriptor(int32 DescriptorSetIndex, co
 	}
 }
 
-void FVulkanDescriptorSetsLayoutInfo::GenerateHash()
+void FVulkanDescriptorSetsLayoutInfo::GenerateHash(const TArrayView<const FSamplerStateRHIParamRef>& InImmutableSamplers)
 {
 	const int32 LayoutCount = SetLayouts.Num();
 	Hash = FCrc::MemCrc32(&TypesUsageID, sizeof(uint32), LayoutCount);
@@ -1134,6 +1134,18 @@ void FVulkanDescriptorSetsLayoutInfo::GenerateHash()
 		TArray<uint16>& PackedUBBindingIndices = RemappingInfo.StageInfos[RemapingIndex].PackedUBBindingIndices;
 		Hash = FCrc::MemCrc32(PackedUBBindingIndices.GetData(), sizeof(uint16) * PackedUBBindingIndices.Num(), Hash);
 	}
+
+#if VULKAN_SUPPORTS_COLOR_CONVERSIONS
+	VkSampler ImmutableSamplers[MaxImmutableSamplers];
+	VkSampler* ImmutableSamplerPtr = ImmutableSamplers;
+	for (int32 Index = 0; Index < InImmutableSamplers.Num(); ++Index)
+	{
+		FRHISamplerState* SamplerState = InImmutableSamplers[Index];
+		*ImmutableSamplerPtr++ = SamplerState ? ResourceCast(SamplerState)->Sampler : VK_NULL_HANDLE;
+	}
+	FMemory::Memzero(ImmutableSamplerPtr, (MaxImmutableSamplers - InImmutableSamplers.Num()));
+	Hash = FCrc::MemCrc32(ImmutableSamplers, sizeof(VkSampler) * MaxImmutableSamplers, Hash);
+#endif
 }
 
 void FVulkanDescriptorSetsLayoutInfo::CompileTypesUsageID()
diff --git a/Engine/Source/Runtime/VulkanRHI/Private/VulkanShaders.cpp b/Engine/Source/Runtime/VulkanRHI/Private/VulkanShaders.cpp
--- a/Engine/Source/Runtime/VulkanRHI/Private/VulkanShaders.cpp
+++ b/Engine/Source/Runtime/VulkanRHI/Private/VulkanShaders.cpp
@@ -542,7 +542,7 @@ void FVulkanDescriptorSetsLayoutInfo::FinalizeBindings(const FUniformBufferGathe
 	}
 
 	CompileTypesUsageID();
-	GenerateHash();
+	GenerateHash(ImmutableSamplers);
 
 	// Validate no empty sets were made
 	for (int32 Index = 0; Index < RemappingInfo.SetInfos.Num(); ++Index)
