UE4.21 - Better error wording
#rnx

diff --git a/Engine/Source/Runtime/VulkanRHI/Private/VulkanRHI.cpp b/Engine/Source/Runtime/VulkanRHI/Private/VulkanRHI.cpp
--- a/Engine/Source/Runtime/VulkanRHI/Private/VulkanRHI.cpp
+++ b/Engine/Source/Runtime/VulkanRHI/Private/VulkanRHI.cpp
@@ -451,11 +451,12 @@ void FVulkanDynamicRHI::SelectAndInitDevice()
 {
 	uint32 GpuCount = 0;
 	VERIFYVULKANRESULT_EXPANDED(VulkanRHI::vkEnumeratePhysicalDevices(Instance, &GpuCount, nullptr));
-	check(GpuCount >= 1);
+	checkf(GpuCount >= 1, TEXT("No GPU(s)/Driver(s) that support Vulkan were found! Make sure your drivers are up to date and that you are not pending a reboot."));
 
 	TArray<VkPhysicalDevice> PhysicalDevices;
 	PhysicalDevices.AddZeroed(GpuCount);
 	VERIFYVULKANRESULT_EXPANDED(VulkanRHI::vkEnumeratePhysicalDevices(Instance, &GpuCount, PhysicalDevices.GetData()));
+	checkf(GpuCount >= 1, TEXT("Couldn't enumerate physical devices! Make sure your drivers are up to date and that you are not pending a reboot."));
 
 #if VULKAN_ENABLE_DESKTOP_HMD_SUPPORT
 	FVulkanDevice* HmdDevice = nullptr;
