UE4.21 - Workaround for some device lost issues
- Fix crash when GPU hangs
#rnx

diff --git a/Engine/Source/Runtime/VulkanRHI/Private/VulkanCommands.cpp b/Engine/Source/Runtime/VulkanRHI/Private/VulkanCommands.cpp
--- a/Engine/Source/Runtime/VulkanRHI/Private/VulkanCommands.cpp
+++ b/Engine/Source/Runtime/VulkanRHI/Private/VulkanCommands.cpp
@@ -126,7 +126,8 @@ void FVulkanCommandListContext::RHIDispatchComputeShader(uint32 ThreadGroupCount
 		GpuProfiler.RegisterGPUWork(1);
 	}
 
-	VulkanRHI::DebugHeavyWeightBarrier(CmdBuffer, 2);
+	//#todo-rco: Temp workaround
+	VulkanRHI::/*Debug*/HeavyWeightBarrier(CmdBuffer/*, 2*/);
 }
 
 void FVulkanCommandListContext::RHIDispatchIndirectComputeShader(FVertexBufferRHIParamRef ArgumentBufferRHI, uint32 ArgumentOffset) 
@@ -156,7 +157,8 @@ void FVulkanCommandListContext::RHIDispatchIndirectComputeShader(FVertexBufferRH
 		GpuProfiler.RegisterGPUWork(1);
 	}
 
-	VulkanRHI::DebugHeavyWeightBarrier(CmdBuffer, 2);
+	//#todo-rco: Temp workaround
+	VulkanRHI::/*Debug*/HeavyWeightBarrier(CmdBuffer/*, 2*/);
 }
 
 void FVulkanCommandListContext::RHISetBoundShaderState(FBoundShaderStateRHIParamRef BoundShaderStateRHI)
diff --git a/Engine/Source/Runtime/VulkanRHI/Private/VulkanRHIPrivate.h b/Engine/Source/Runtime/VulkanRHI/Private/VulkanRHIPrivate.h
--- a/Engine/Source/Runtime/VulkanRHI/Private/VulkanRHIPrivate.h
+++ b/Engine/Source/Runtime/VulkanRHI/Private/VulkanRHIPrivate.h
@@ -821,44 +821,49 @@ namespace VulkanRHI
 		return VK_IMAGE_LAYOUT_DEPTH_STENCIL_ATTACHMENT_OPTIMAL;
 	}
 
+	inline void HeavyWeightBarrier(VkCommandBuffer CmdBuffer)
+	{
+		VkMemoryBarrier Barrier;
+		ZeroVulkanStruct(Barrier, VK_STRUCTURE_TYPE_MEMORY_BARRIER);
+		Barrier.srcAccessMask = VK_ACCESS_INDIRECT_COMMAND_READ_BIT |
+			VK_ACCESS_INDEX_READ_BIT |
+			VK_ACCESS_VERTEX_ATTRIBUTE_READ_BIT |
+			VK_ACCESS_UNIFORM_READ_BIT |
+			VK_ACCESS_INPUT_ATTACHMENT_READ_BIT |
+			VK_ACCESS_SHADER_READ_BIT |
+			VK_ACCESS_SHADER_WRITE_BIT |
+			VK_ACCESS_COLOR_ATTACHMENT_READ_BIT |
+			VK_ACCESS_COLOR_ATTACHMENT_WRITE_BIT |
+			VK_ACCESS_DEPTH_STENCIL_ATTACHMENT_READ_BIT |
+			VK_ACCESS_DEPTH_STENCIL_ATTACHMENT_WRITE_BIT |
+			VK_ACCESS_TRANSFER_READ_BIT |
+			VK_ACCESS_TRANSFER_WRITE_BIT |
+			VK_ACCESS_HOST_READ_BIT |
+			VK_ACCESS_HOST_WRITE_BIT;
+		Barrier.dstAccessMask = VK_ACCESS_INDIRECT_COMMAND_READ_BIT |
+			VK_ACCESS_INDEX_READ_BIT |
+			VK_ACCESS_VERTEX_ATTRIBUTE_READ_BIT |
+			VK_ACCESS_UNIFORM_READ_BIT |
+			VK_ACCESS_INPUT_ATTACHMENT_READ_BIT |
+			VK_ACCESS_SHADER_READ_BIT |
+			VK_ACCESS_SHADER_WRITE_BIT |
+			VK_ACCESS_COLOR_ATTACHMENT_READ_BIT |
+			VK_ACCESS_COLOR_ATTACHMENT_WRITE_BIT |
+			VK_ACCESS_DEPTH_STENCIL_ATTACHMENT_READ_BIT |
+			VK_ACCESS_DEPTH_STENCIL_ATTACHMENT_WRITE_BIT |
+			VK_ACCESS_TRANSFER_READ_BIT |
+			VK_ACCESS_TRANSFER_WRITE_BIT |
+			VK_ACCESS_HOST_READ_BIT |
+			VK_ACCESS_HOST_WRITE_BIT;
+		VulkanRHI::vkCmdPipelineBarrier(CmdBuffer, VK_PIPELINE_STAGE_ALL_COMMANDS_BIT, VK_PIPELINE_STAGE_ALL_COMMANDS_BIT, 0, 1, &Barrier, 0, nullptr, 0, nullptr);
+	}
+
 	inline void DebugHeavyWeightBarrier(VkCommandBuffer CmdBuffer, int32 CVarConditionMask)
 	{
 #if UE_BUILD_DEBUG || UE_BUILD_DEVELOPMENT
 		if (CVarVulkanDebugBarrier.GetValueOnAnyThread() & CVarConditionMask)
 		{
-			VkMemoryBarrier Barrier;
-			ZeroVulkanStruct(Barrier, VK_STRUCTURE_TYPE_MEMORY_BARRIER);
-			Barrier.srcAccessMask = VK_ACCESS_INDIRECT_COMMAND_READ_BIT |
-				VK_ACCESS_INDEX_READ_BIT |
-				VK_ACCESS_VERTEX_ATTRIBUTE_READ_BIT |
-				VK_ACCESS_UNIFORM_READ_BIT |
-				VK_ACCESS_INPUT_ATTACHMENT_READ_BIT |
-				VK_ACCESS_SHADER_READ_BIT |
-				VK_ACCESS_SHADER_WRITE_BIT |
-				VK_ACCESS_COLOR_ATTACHMENT_READ_BIT |
-				VK_ACCESS_COLOR_ATTACHMENT_WRITE_BIT |
-				VK_ACCESS_DEPTH_STENCIL_ATTACHMENT_READ_BIT |
-				VK_ACCESS_DEPTH_STENCIL_ATTACHMENT_WRITE_BIT |
-				VK_ACCESS_TRANSFER_READ_BIT |
-				VK_ACCESS_TRANSFER_WRITE_BIT |
-				VK_ACCESS_HOST_READ_BIT |
-				VK_ACCESS_HOST_WRITE_BIT;
-			Barrier.dstAccessMask = VK_ACCESS_INDIRECT_COMMAND_READ_BIT |
-				VK_ACCESS_INDEX_READ_BIT |
-				VK_ACCESS_VERTEX_ATTRIBUTE_READ_BIT |
-				VK_ACCESS_UNIFORM_READ_BIT |
-				VK_ACCESS_INPUT_ATTACHMENT_READ_BIT |
-				VK_ACCESS_SHADER_READ_BIT |
-				VK_ACCESS_SHADER_WRITE_BIT |
-				VK_ACCESS_COLOR_ATTACHMENT_READ_BIT |
-				VK_ACCESS_COLOR_ATTACHMENT_WRITE_BIT |
-				VK_ACCESS_DEPTH_STENCIL_ATTACHMENT_READ_BIT |
-				VK_ACCESS_DEPTH_STENCIL_ATTACHMENT_WRITE_BIT |
-				VK_ACCESS_TRANSFER_READ_BIT |
-				VK_ACCESS_TRANSFER_WRITE_BIT |
-				VK_ACCESS_HOST_READ_BIT |
-				VK_ACCESS_HOST_WRITE_BIT;
-			VulkanRHI::vkCmdPipelineBarrier(CmdBuffer, VK_PIPELINE_STAGE_ALL_COMMANDS_BIT, VK_PIPELINE_STAGE_ALL_COMMANDS_BIT, 0, 1, &Barrier, 0, nullptr, 0, nullptr);
+			HeavyWeightBarrier(CmdBuffer);
 		}
 #endif
 	}
diff --git a/Engine/Source/Runtime/VulkanRHI/Private/VulkanUtil.cpp b/Engine/Source/Runtime/VulkanRHI/Private/VulkanUtil.cpp
--- a/Engine/Source/Runtime/VulkanRHI/Private/VulkanUtil.cpp
+++ b/Engine/Source/Runtime/VulkanRHI/Private/VulkanUtil.cpp
@@ -537,7 +537,6 @@ void FVulkanGPUProfiler::DumpCrashMarkers(void* BufferData)
 				const FString* Frame = CachedStrings.Find(Value);
 				UE_LOG(LogVulkanRHI, Error, TEXT("[VK_NV_device_diagnostic_checkpoints] %i: Stage 0x%x, %s (CRC 0x%x)"), Index, Data[Index].stage, Frame ? *(*Frame) : TEXT("<undefined>"), Value);
 			}
-			GLog->FlushThreadedLogs();
 			GLog->PanicFlushThreadedLogs();
 			GLog->Flush();
 		}
