DR - vk - Set GPU capture options for debugging/capture

diff --git a/Engine/Source/Runtime/VulkanRHI/Private/VulkanRHI.cpp b/Engine/Source/Runtime/VulkanRHI/Private/VulkanRHI.cpp
--- a/Engine/Source/Runtime/VulkanRHI/Private/VulkanRHI.cpp
+++ b/Engine/Source/Runtime/VulkanRHI/Private/VulkanRHI.cpp
@@ -391,6 +391,11 @@ void FVulkanDynamicRHI::CreateInstance()
 #if VULKAN_HAS_DEBUGGING_ENABLED
 	SetupDebugLayerCallback();
 #endif
+
+	if (!GRHISupportsRHIThread || GRenderDocFound)
+	{
+		EnableIdealGPUCaptureOptions(true);
+	}
 }
 
 //#todo-rco: Common RHI should handle this...
