DR - vk - Change how r.Vulkan.DebugBarrier works

diff --git a/Engine/Source/Runtime/VulkanRHI/Private/VulkanCommandBuffer.cpp b/Engine/Source/Runtime/VulkanRHI/Private/VulkanCommandBuffer.cpp
--- a/Engine/Source/Runtime/VulkanRHI/Private/VulkanCommandBuffer.cpp
+++ b/Engine/Source/Runtime/VulkanRHI/Private/VulkanCommandBuffer.cpp
@@ -28,6 +28,10 @@ static FAutoConsoleVariableRef CVarVulkanProfileCmdBuffers(
 
 #define CMD_BUFFER_TIME_TO_WAIT_BEFORE_DELETING		10
 
+#if UE_BUILD_DEBUG || UE_BUILD_DEVELOPMENT
+extern TAutoConsoleVariable<int32> CVarVulkanDebugBarrier;
+#endif
+
 const uint32 GNumberOfFramesBeforeDeletingDescriptorPool = 300;
 
 FVulkanCmdBuffer::FVulkanCmdBuffer(FVulkanDevice* InDevice, FVulkanCommandBufferPool* InCommandBufferPool, bool bInIsUploadOnly)
@@ -337,6 +341,14 @@ void FVulkanCommandBufferManager::SubmitUploadCmdBuffer(uint32 NumSignalSemaphor
 	if (!UploadCmdBuffer->IsSubmitted() && UploadCmdBuffer->HasBegun())
 	{
 		check(UploadCmdBuffer->IsOutsideRenderPass());
+
+#if UE_BUILD_DEBUG || UE_BUILD_DEVELOPMENT
+		if (CVarVulkanDebugBarrier.GetValueOnAnyThread() & 4)
+		{
+			VulkanRHI::InsertHeavyWeightBarrier(UploadCmdBuffer->GetHandle());
+		}
+#endif
+
 		UploadCmdBuffer->End();
 		Queue->Submit(UploadCmdBuffer, NumSignalSemaphores, SignalSemaphores);
 		UploadCmdBuffer->SubmittedTime = FPlatformTime::Seconds();
@@ -357,6 +369,14 @@ void FVulkanCommandBufferManager::SubmitActiveCmdBuffer(VulkanRHI::FSemaphore* S
 			UE_LOG(LogVulkanRHI, Warning, TEXT("Forcing EndRenderPass() for submission"));
 			ActiveCmdBuffer->EndRenderPass();
 		}
+
+#if UE_BUILD_DEBUG || UE_BUILD_DEVELOPMENT
+		if (CVarVulkanDebugBarrier.GetValueOnAnyThread() & 8)
+		{
+			VulkanRHI::InsertHeavyWeightBarrier(ActiveCmdBuffer->GetHandle());
+		}
+#endif
+
 		ActiveCmdBuffer->End();
 		if (SignalSemaphore)
 		{
diff --git a/Engine/Source/Runtime/VulkanRHI/Private/VulkanCommands.cpp b/Engine/Source/Runtime/VulkanRHI/Private/VulkanCommands.cpp
--- a/Engine/Source/Runtime/VulkanRHI/Private/VulkanCommands.cpp
+++ b/Engine/Source/Runtime/VulkanRHI/Private/VulkanCommands.cpp
@@ -131,7 +131,7 @@ void FVulkanCommandListContext::RHIDispatchComputeShader(uint32 ThreadGroupCount
 	}
 
 #if UE_BUILD_DEBUG || UE_BUILD_DEVELOPMENT
-	if (CVarVulkanDebugBarrier.GetValueOnRenderThread() > 1)
+	if (CVarVulkanDebugBarrier.GetValueOnRenderThread() & 2)
 	{
 		VulkanRHI::InsertHeavyWeightBarrier(CmdBuffer);
 	}
@@ -166,7 +166,7 @@ void FVulkanCommandListContext::RHIDispatchIndirectComputeShader(FVertexBufferRH
 	}
 
 #if UE_BUILD_DEBUG || UE_BUILD_DEVELOPMENT
-	if (CVarVulkanDebugBarrier.GetValueOnRenderThread() > 1)
+	if (CVarVulkanDebugBarrier.GetValueOnRenderThread() & 2)
 	{
 		VulkanRHI::InsertHeavyWeightBarrier(CmdBuffer);
 	}
diff --git a/Engine/Source/Runtime/VulkanRHI/Private/VulkanRenderTarget.cpp b/Engine/Source/Runtime/VulkanRHI/Private/VulkanRenderTarget.cpp
--- a/Engine/Source/Runtime/VulkanRHI/Private/VulkanRenderTarget.cpp
+++ b/Engine/Source/Runtime/VulkanRHI/Private/VulkanRenderTarget.cpp
@@ -38,10 +38,13 @@ static TMap<FVulkanTextureBase*, VulkanRHI::FStagingBuffer*> GPendingLockedStagi
 TAutoConsoleVariable<int32> CVarVulkanDebugBarrier(
 	TEXT("r.Vulkan.DebugBarrier"),
 	0,
-	TEXT("Forces a full barrier for debugging.\n")
+	TEXT("Forces a full barrier for debugging. This is a mask/bitfield (so add up the values)!\n")
 	TEXT(" 0: Don't (default)\n")
-	TEXT(" 1: Enable barriers after EndRenderPass()\n")
-	TEXT(" 2: Enable barriers after EndRenderPass() AND every dispatch\n"),
+	TEXT(" 1: Enable heavy barriers after EndRenderPass()\n")
+	TEXT(" 2: Enable heavy barriers after every dispatch\n")
+	TEXT(" 4: Enable heavy barriers after upload cmd buffers\n")
+	TEXT(" 8: Enable heavy barriers after active cmd buffers\n")
+	/*TEXT(" 8: Enable heavy barrier after swapchain\n")*/,
 	ECVF_Default
 );
 #endif
@@ -220,7 +223,7 @@ void FTransitionAndLayoutManager::EndEmulatedRenderPass(FVulkanCmdBuffer* CmdBuf
 	CurrentRenderPass = nullptr;
 
 #if UE_BUILD_DEBUG || UE_BUILD_DEVELOPMENT
-	if (CVarVulkanDebugBarrier.GetValueOnRenderThread() > 0)
+	if (CVarVulkanDebugBarrier.GetValueOnRenderThread() & 1)
 	{
 		VulkanRHI::InsertHeavyWeightBarrier(CmdBuffer->GetHandle());
 	}
@@ -429,7 +432,7 @@ void FTransitionAndLayoutManager::EndRealRenderPass(FVulkanCmdBuffer* CmdBuffer)
 	bInsideRealRenderPass = false;
 
 #if UE_BUILD_DEBUG || UE_BUILD_DEVELOPMENT
-	if (CVarVulkanDebugBarrier.GetValueOnRenderThread() > 0)
+	if (CVarVulkanDebugBarrier.GetValueOnRenderThread() & 1)
 	{
 		VulkanRHI::InsertHeavyWeightBarrier(CmdBuffer->GetHandle());
 	}
