DR - vk - Log present mode

diff --git a/Engine/Source/Runtime/VulkanRHI/Private/VulkanSwapChain.cpp b/Engine/Source/Runtime/VulkanRHI/Private/VulkanSwapChain.cpp
--- a/Engine/Source/Runtime/VulkanRHI/Private/VulkanSwapChain.cpp
+++ b/Engine/Source/Runtime/VulkanRHI/Private/VulkanSwapChain.cpp
@@ -1,4 +1,4 @@
-﻿// Copyright 1998-2018 Epic Games, Inc. All Rights Reserved.
+// Copyright 1998-2018 Epic Games, Inc. All Rights Reserved.
 
 /*=============================================================================
 	VulkanSwapChain.h: Vulkan viewport RHI definitions.
@@ -203,6 +203,8 @@ FVulkanSwapChain::FVulkanSwapChain(VkInstance InInstance, FVulkanDevice& InDevic
 	VkPresentModeKHR PresentMode = VK_PRESENT_MODE_FIFO_KHR;
 	if (FVulkanPlatform::SupportsQuerySurfaceProperties())
 	{
+		static bool bFirstTimeLog = true;
+
 		uint32 NumFoundPresentModes = 0;
 		VERIFYVULKANRESULT(VulkanRHI::vkGetPhysicalDeviceSurfacePresentModesKHR(Device.GetPhysicalHandle(), Surface, &NumFoundPresentModes, nullptr));
 		check(NumFoundPresentModes > 0);
@@ -211,6 +213,8 @@ FVulkanSwapChain::FVulkanSwapChain(VkInstance InInstance, FVulkanDevice& InDevic
 		FoundPresentModes.AddZeroed(NumFoundPresentModes);
 		VERIFYVULKANRESULT(VulkanRHI::vkGetPhysicalDeviceSurfacePresentModesKHR(Device.GetPhysicalHandle(), Surface, &NumFoundPresentModes, FoundPresentModes.GetData()));
 
+		UE_CLOG(bFirstTimeLog, LogVulkanRHI, Display, TEXT("Found %d Surface present modes:"), NumFoundPresentModes);
+
 		bool bFoundPresentModeMailbox = false;
 		bool bFoundPresentModeImmediate = false;
 		bool bFoundPresentModeFIFO = false;
@@ -221,12 +225,18 @@ FVulkanSwapChain::FVulkanSwapChain(VkInstance InInstance, FVulkanDevice& InDevic
 			{
 			case VK_PRESENT_MODE_MAILBOX_KHR:
 				bFoundPresentModeMailbox = true;
+				UE_CLOG(bFirstTimeLog, LogVulkanRHI, Display, TEXT("- VK_PRESENT_MODE_MAILBOX_KHR (%d)"), (int32)VK_PRESENT_MODE_MAILBOX_KHR);
 				break;
 			case VK_PRESENT_MODE_IMMEDIATE_KHR:
 				bFoundPresentModeImmediate = true;
+				UE_CLOG(bFirstTimeLog, LogVulkanRHI, Display, TEXT("- VK_PRESENT_MODE_IMMEDIATE_KHR (%d)"), (int32)VK_PRESENT_MODE_IMMEDIATE_KHR);
 				break;
 			case VK_PRESENT_MODE_FIFO_KHR:
 				bFoundPresentModeFIFO = true;
+				UE_CLOG(bFirstTimeLog, LogVulkanRHI, Display, TEXT("- VK_PRESENT_MODE_FIFO_KHR (%d)"), (int32)VK_PRESENT_MODE_FIFO_KHR);
+				break;
+			default:
+				UE_CLOG(bFirstTimeLog, LogVulkanRHI, Display, TEXT("- VkPresentModeKHR mode %d"), (int32)FoundPresentModes[i]);
 				break;
 			}
 		}
@@ -251,6 +261,9 @@ FVulkanSwapChain::FVulkanSwapChain(VkInstance InInstance, FVulkanDevice& InDevic
 			UE_LOG(LogVulkanRHI, Warning, TEXT("Couldn't find desired PresentMode! Using %d"), static_cast<int32>(FoundPresentModes[0]));
 			PresentMode = FoundPresentModes[0];
 		}
+
+		UE_CLOG(bFirstTimeLog, LogVulkanRHI, Display, TEXT("Selected VkPresentModeKHR mode %d"), PresentMode);
+		bFirstTimeLog = false;
 	}
 
 	// Check the surface properties and formats
