DR - vk - Copy some fixes from staging branch

diff --git a/Engine/Source/Runtime/VulkanRHI/Private/VulkanDebug.cpp b/Engine/Source/Runtime/VulkanRHI/Private/VulkanDebug.cpp
--- a/Engine/Source/Runtime/VulkanRHI/Private/VulkanDebug.cpp
+++ b/Engine/Source/Runtime/VulkanRHI/Private/VulkanDebug.cpp
@@ -21,11 +21,15 @@ DEFINE_LOG_CATEGORY(LogVulkanRHI);
 
 #if VULKAN_HAS_DEBUGGING_ENABLED
 
-#if PLATFORM_ANDROID || PLATFORM_LUMIN || PLATFORM_LUMINGL4
-	#define VULKAN_REPORT_LOG(Format, ...)			UE_LOG(LogVulkanRHI, Warning, Format, __VA_ARGS__)
+#if PLATFORM_ANDROID
+	#define VULKAN_REPORT_LOG(Format, ...)	UE_LOG(LogVulkanRHI, Warning, Format, __VA_ARGS__)
 #else
-	#define VULKAN_REPORT_LOG(Format, ...)			FPlatformMisc::LowLevelOutputDebugStringf(Format, __VA_ARGS__); FPlatformMisc::LowLevelOutputDebugString(TEXT("\n"))
-#endif
+	#define VULKAN_REPORT_LOG(Format, ...)	do { if (FPlatformMisc::IsDebuggerPresent()) \
+											{ \
+												FPlatformMisc::LowLevelOutputDebugStringf(Format, __VA_ARGS__); FPlatformMisc::LowLevelOutputDebugString(TEXT("\n")); \
+											} \
+											UE_LOG(LogVulkanRHI, Warning, Format, __VA_ARGS__); } while (0)
+#endif // PLATFORM_ANDROID
 
 extern TAutoConsoleVariable<int32> GValidationCvar;
 
@@ -179,7 +183,20 @@ static VkBool32 DebugUtilsCallback(VkDebugUtilsMessageSeverityFlagBitsEXT MsgSev
 	const bool bError = (MsgSeverity & VK_DEBUG_UTILS_MESSAGE_SEVERITY_ERROR_BIT_EXT) != 0;
 	const bool bWarning = (MsgSeverity & VK_DEBUG_UTILS_MESSAGE_SEVERITY_WARNING_BIT_EXT) != 0;
 
-	if (!CallbackData->pMessageIdName)
+	if (CallbackData->pMessageIdName)
+	{
+		if (!FCStringAnsi::Strcmp(CallbackData->pMessageIdName, "UNASSIGNED-CoreValidation-Shader-OutputNotConsumed"))
+		{
+			// Warning: *** [Warning:Validation-1(UNASSIGNED-CoreValidation-Shader-OutputNotConsumed)] fragment shader writes to output location 0 with no matching attachment
+			return VK_FALSE;
+		}
+		else if (!FCStringAnsi::Strcmp(CallbackData->pMessageIdName, "VUID-VkSwapchainCreateInfoKHR-imageExtent-01274"))
+		{
+			// Warning: *** [Error:Validation341838324(VUID-VkSwapchainCreateInfoKHR-imageExtent-01274)] vkCreateSwapChainKHR() called with imageExtent = (8,8), which is outside the bounds returned by vkGetPhysicalDeviceSurfaceCapabilitiesKHR(): currentExtent = (0,0), minImageExtent = (0,0), maxImageExtent = (0,0).
+			return VK_FALSE;
+		}
+	}
+	else
 	{
 		if (MsgType == 2 && CallbackData->messageIdNumber == 5)
 		{
@@ -233,6 +250,7 @@ static VkBool32 DebugUtilsCallback(VkDebugUtilsMessageSeverityFlagBitsEXT MsgSev
 	enum class EMsgBucket
 	{
 		General,
+		GeneralValidation,
 		PerfValidation,
 		Validation,
 		Perf,
@@ -242,9 +260,18 @@ static VkBool32 DebugUtilsCallback(VkDebugUtilsMessageSeverityFlagBitsEXT MsgSev
 	const TCHAR* Type = TEXT("");
 	if (MsgType & VK_DEBUG_UTILS_MESSAGE_TYPE_GENERAL_BIT_EXT)
 	{
-		ensure((MsgType & ~VK_DEBUG_UTILS_MESSAGE_TYPE_GENERAL_BIT_EXT) == 0);
-		Type = TEXT(" General");
-		MsgBucket = EMsgBucket::General;
+		if (MsgType & VK_DEBUG_UTILS_MESSAGE_TYPE_VALIDATION_BIT_EXT)
+		{
+			ensure((MsgType & ~(VK_DEBUG_UTILS_MESSAGE_TYPE_GENERAL_BIT_EXT | VK_DEBUG_UTILS_MESSAGE_TYPE_VALIDATION_BIT_EXT)) == 0);
+			Type = TEXT(" General/Validation");
+			MsgBucket = EMsgBucket::GeneralValidation;
+		}
+		else
+		{
+			ensure((MsgType & ~VK_DEBUG_UTILS_MESSAGE_TYPE_GENERAL_BIT_EXT) == 0);
+			Type = TEXT(" General");
+			MsgBucket = EMsgBucket::General;
+		}
 	}
 	else if (MsgType & VK_DEBUG_UTILS_MESSAGE_TYPE_VALIDATION_BIT_EXT)
 	{
@@ -362,11 +389,6 @@ void FVulkanDynamicRHI::SetupDebugLayerCallback()
 			UE_LOG(LogVulkanRHI, Warning, TEXT("GetProcAddr: Unable to find vkDbgCreateMsgCallback/vkGetInstanceProcAddr; debug reporting skipped!"));
 		}
 	}
-	else
-	{
-		UE_LOG(LogVulkanRHI, Warning, TEXT("Instance does not support 'VK_EXT_debug_report' extension; debug reporting skipped!"));
-		return;
-	}
 }
 
 void FVulkanDynamicRHI::RemoveDebugLayerCallback()
diff --git a/Engine/Source/Runtime/VulkanRHI/Private/VulkanGenericPlatform.h b/Engine/Source/Runtime/VulkanRHI/Private/VulkanGenericPlatform.h
--- a/Engine/Source/Runtime/VulkanRHI/Private/VulkanGenericPlatform.h
+++ b/Engine/Source/Runtime/VulkanRHI/Private/VulkanGenericPlatform.h
@@ -9,6 +9,8 @@
 class FVulkanGenericPlatform 
 {
 public:
+	static bool IsSupported() { return true; }
+
 	static bool LoadVulkanLibrary() { return true; }
 	static bool LoadVulkanInstanceFunctions(VkInstance inInstance) { return true; }
 	static void FreeVulkanLibrary() {}
@@ -92,4 +94,7 @@ public:
 	static void EnablePhysicalDeviceFeatureExtensions(VkDeviceCreateInfo& DeviceInfo) {};
 
 	static bool RequiresSwapchainGeneralInitialLayout() { return false; }
+	
+	// Allow platforms to add extension features to the PresentInfo pNext chain
+	static void EnablePresentInfoExtensions(VkPresentInfoKHR& PresentInfo) {}
 };
diff --git a/Engine/Source/Runtime/VulkanRHI/Private/VulkanIndexBuffer.cpp b/Engine/Source/Runtime/VulkanRHI/Private/VulkanIndexBuffer.cpp
--- a/Engine/Source/Runtime/VulkanRHI/Private/VulkanIndexBuffer.cpp
+++ b/Engine/Source/Runtime/VulkanRHI/Private/VulkanIndexBuffer.cpp
@@ -81,7 +81,16 @@ FVulkanResourceMultiBuffer::FVulkanResourceMultiBuffer(FVulkanDevice* InDevice,
 		BufferUsageFlags |= bUAV ? VK_BUFFER_USAGE_STORAGE_TEXEL_BUFFER_BIT : 0;
 		BufferUsageFlags |= bIndirect ? VK_BUFFER_USAGE_INDIRECT_BUFFER_BIT : 0;
 
-		if (!bVolatile)
+		if (bVolatile)
+		{
+			bool bRenderThread = IsInRenderingThread();
+
+			// Get a dummy buffer as sometimes the high-level misbehaves and tries to use SRVs off volatile buffers before filling them in...
+			void* Data = Lock(bRenderThread, RLM_WriteOnly, InSize, 0);
+			FMemory::Memzero(Data, InSize);
+			Unlock(bRenderThread);
+		}
+		else
 		{
 			VkDevice VulkanDevice = InDevice->GetInstanceHandle();
 
diff --git a/Engine/Source/Runtime/VulkanRHI/Private/VulkanRHI.cpp b/Engine/Source/Runtime/VulkanRHI/Private/VulkanRHI.cpp
--- a/Engine/Source/Runtime/VulkanRHI/Private/VulkanRHI.cpp
+++ b/Engine/Source/Runtime/VulkanRHI/Private/VulkanRHI.cpp
@@ -56,7 +56,7 @@ DEFINE_LOG_CATEGORY(LogVulkan)
 
 bool FVulkanDynamicRHIModule::IsSupported()
 {
-	return true;
+	return FVulkanPlatform::IsSupported();
 }
 
 FDynamicRHI* FVulkanDynamicRHIModule::CreateRHI(ERHIFeatureLevel::Type InRequestedFeatureLevel)
diff --git a/Engine/Source/Runtime/VulkanRHI/Private/VulkanSwapChain.cpp b/Engine/Source/Runtime/VulkanRHI/Private/VulkanSwapChain.cpp
--- a/Engine/Source/Runtime/VulkanRHI/Private/VulkanSwapChain.cpp
+++ b/Engine/Source/Runtime/VulkanRHI/Private/VulkanSwapChain.cpp
@@ -568,6 +568,7 @@ FVulkanSwapChain::EStatus FVulkanSwapChain::Present(FVulkanQueue* GfxQueue, FVul
 
 	const int32 SyncInterval = RHIGetSyncInterval();
 	ensureMsgf(SyncInterval <= 3 && SyncInterval >= 0, TEXT("Unsupported sync interval: %i"), SyncInterval);
+	FVulkanPlatform::EnablePresentInfoExtensions(Info);
 
 #if VULKAN_SUPPORTS_GOOGLE_DISPLAY_TIMING
 	// These are referenced by VkPresentInfoKHR, so they need to be at the same level in the stack
diff --git a/Engine/Source/Runtime/VulkanRHI/Public/VulkanContext.h b/Engine/Source/Runtime/VulkanRHI/Public/VulkanContext.h
--- a/Engine/Source/Runtime/VulkanRHI/Public/VulkanContext.h
+++ b/Engine/Source/Runtime/VulkanRHI/Public/VulkanContext.h
@@ -358,7 +358,7 @@ public:
 		}
 #endif
 
-		return VK_IMAGE_LAYOUT_SHADER_READ_ONLY_OPTIMAL;
+		return TransitionAndLayoutManager.FindLayoutChecked(Surface.Image);
 	}
 
 	inline VkImageLayout FindOrAddLayout(VkImage Image, VkImageLayout NewLayout)
