DR - vk - Missing marker

diff --git a/Engine/Source/Runtime/VulkanRHI/Private/VulkanRenderTarget.cpp b/Engine/Source/Runtime/VulkanRHI/Private/VulkanRenderTarget.cpp
--- a/Engine/Source/Runtime/VulkanRHI/Private/VulkanRenderTarget.cpp
+++ b/Engine/Source/Runtime/VulkanRHI/Private/VulkanRenderTarget.cpp
@@ -1373,6 +1373,8 @@ void FVulkanCommandListContext::TransitionResources(const FPendingTransition& Pe
 	}
 	else
 	{
+		SCOPED_RHI_CONDITIONAL_DRAW_EVENTF(*this, RHITransitionResources, bShowTransitionEvents, TEXT("TransitionTo: %s: %i UAVs"), *FResourceTransitionUtility::ResourceTransitionAccessStrings[(int32)PendingTransition.TransitionType], PendingTransition.UAVs.Num());
+
 		const bool bIsRealAsyncComputeContext = Device->IsRealAsyncComputeContext(this);
 		ensure(IsImmediate() || bIsRealAsyncComputeContext);
 		check(PendingTransition.UAVs.Num() > 0);
