DR - vk - Heavy barrier for debugging uploads

diff --git a/Engine/Source/Runtime/VulkanRHI/Private/VulkanIndexBuffer.cpp b/Engine/Source/Runtime/VulkanRHI/Private/VulkanIndexBuffer.cpp
--- a/Engine/Source/Runtime/VulkanRHI/Private/VulkanIndexBuffer.cpp
+++ b/Engine/Source/Runtime/VulkanRHI/Private/VulkanIndexBuffer.cpp
@@ -10,6 +10,12 @@
 #include "Containers/ResourceArray.h"
 #include "VulkanLLM.h"
 
+
+#if UE_BUILD_DEBUG || UE_BUILD_DEVELOPMENT
+extern TAutoConsoleVariable<int32> CVarVulkanDebugBarrier;
+#endif
+
+
 static TMap<FVulkanResourceMultiBuffer*, VulkanRHI::FPendingBufferLock> GPendingLockIBs;
 static FCriticalSection GPendingLockIBsMutex;
 
@@ -233,6 +239,13 @@ inline void FVulkanResourceMultiBuffer::InternalUnlock(FVulkanCommandListContext
 	ensure(Cmd->IsOutsideRenderPass());
 	VkCommandBuffer CmdBuffer = Cmd->GetHandle();
 
+#if UE_BUILD_DEBUG || UE_BUILD_DEVELOPMENT
+	if (CVarVulkanDebugBarrier.GetValueOnRenderThread() & 16)
+	{
+		VulkanRHI::InsertHeavyWeightBarrier(CmdBuffer);
+	}
+#endif
+
 	VkBufferCopy Region;
 	FMemory::Memzero(Region);
 	Region.size = LockSize;
diff --git a/Engine/Source/Runtime/VulkanRHI/Private/VulkanRenderTarget.cpp b/Engine/Source/Runtime/VulkanRHI/Private/VulkanRenderTarget.cpp
--- a/Engine/Source/Runtime/VulkanRHI/Private/VulkanRenderTarget.cpp
+++ b/Engine/Source/Runtime/VulkanRHI/Private/VulkanRenderTarget.cpp
@@ -44,7 +44,7 @@ TAutoConsoleVariable<int32> CVarVulkanDebugBarrier(
 	TEXT(" 2: Enable heavy barriers after every dispatch\n")
 	TEXT(" 4: Enable heavy barriers after upload cmd buffers\n")
 	TEXT(" 8: Enable heavy barriers after active cmd buffers\n")
-	/*TEXT(" 8: Enable heavy barrier after swapchain\n")*/,
+	TEXT(" 16: Enable heavy buffer barrier after uploads\n"),
 	ECVF_Default
 );
 #endif
