DR - vk - Rename for clarity
#rnx

diff --git a/Engine/Source/Runtime/VulkanRHI/Private/VulkanCommandWrappers.h b/Engine/Source/Runtime/VulkanRHI/Private/VulkanCommandWrappers.h
--- a/Engine/Source/Runtime/VulkanRHI/Private/VulkanCommandWrappers.h
+++ b/Engine/Source/Runtime/VulkanRHI/Private/VulkanCommandWrappers.h
@@ -6,7 +6,7 @@
 
 #pragma once 
 
-#if VULKAN_ENABLE_CUSTOM_LAYER
+#if VULKAN_ENABLE_WRAP_LAYER
 	#define  VULKAN_LAYER_BODY		;
 	#define	VULKAN_EXTERN_EXPORT	VULKANRHI_API
 #else
diff --git a/Engine/Source/Runtime/VulkanRHI/Private/VulkanDebug.cpp b/Engine/Source/Runtime/VulkanRHI/Private/VulkanDebug.cpp
--- a/Engine/Source/Runtime/VulkanRHI/Private/VulkanDebug.cpp
+++ b/Engine/Source/Runtime/VulkanRHI/Private/VulkanDebug.cpp
@@ -1329,7 +1329,7 @@ namespace VulkanRHI
 }
 #endif
 
-#if VULKAN_ENABLE_CUSTOM_LAYER
+#if VULKAN_ENABLE_WRAP_LAYER
 void FWrapLayer::GetPhysicalDeviceMemoryProperties(VkResult Result, VkPhysicalDevice PhysicalDevice, VkPhysicalDeviceMemoryProperties* Properties)
 {
 	if (Result == VK_RESULT_MAX_ENUM)
@@ -3879,7 +3879,7 @@ namespace VulkanRHI
 }
 #endif
 
-#endif	// VULKAN_ENABLE_CUSTOM_LAYER
+#endif	// VULKAN_ENABLE_WRAP_LAYER
 
 #if VULKAN_ENABLE_DUMP_LAYER
 namespace VulkanRHI
diff --git a/Engine/Source/Runtime/VulkanRHI/Public/VulkanConfiguration.h b/Engine/Source/Runtime/VulkanRHI/Public/VulkanConfiguration.h
--- a/Engine/Source/Runtime/VulkanRHI/Public/VulkanConfiguration.h
+++ b/Engine/Source/Runtime/VulkanRHI/Public/VulkanConfiguration.h
@@ -44,7 +44,7 @@
 #endif
 
 #define VULKAN_ENABLE_TRACKING_LAYER							(VULKAN_ENABLE_BUFFER_TRACKING_LAYER || VULKAN_ENABLE_IMAGE_TRACKING_LAYER)
-#define VULKAN_ENABLE_CUSTOM_LAYER								(VULKAN_ENABLE_DUMP_LAYER || VULKAN_ENABLE_TRACKING_LAYER)
+#define VULKAN_ENABLE_WRAP_LAYER								(VULKAN_ENABLE_DUMP_LAYER || VULKAN_ENABLE_TRACKING_LAYER)
 
 #define VULKAN_HASH_POOLS_WITH_TYPES_USAGE_ID					1
 
@@ -67,7 +67,7 @@
 #endif
 
 #ifndef VULKAN_SHOULD_USE_COMMANDWRAPPERS
-	#define VULKAN_SHOULD_USE_COMMANDWRAPPERS					VULKAN_ENABLE_CUSTOM_LAYER
+	#define VULKAN_SHOULD_USE_COMMANDWRAPPERS					VULKAN_ENABLE_WRAP_LAYER
 #endif
 
 #ifndef VULKAN_COMMANDWRAPPERS_ENABLE							
