DR - vk - Added support for -vulkanpresentmode=N where N is a VkPresentModeKHR
- Cleaned up a little the debug barrier code and added one for swapchain copy

diff --git a/Engine/Source/Runtime/VulkanRHI/Private/VulkanCommandBuffer.cpp b/Engine/Source/Runtime/VulkanRHI/Private/VulkanCommandBuffer.cpp
--- a/Engine/Source/Runtime/VulkanRHI/Private/VulkanCommandBuffer.cpp
+++ b/Engine/Source/Runtime/VulkanRHI/Private/VulkanCommandBuffer.cpp
@@ -28,10 +28,6 @@ static FAutoConsoleVariableRef CVarVulkanProfileCmdBuffers(
 
 #define CMD_BUFFER_TIME_TO_WAIT_BEFORE_DELETING		10
 
-#if UE_BUILD_DEBUG || UE_BUILD_DEVELOPMENT
-extern TAutoConsoleVariable<int32> CVarVulkanDebugBarrier;
-#endif
-
 const uint32 GNumberOfFramesBeforeDeletingDescriptorPool = 300;
 
 FVulkanCmdBuffer::FVulkanCmdBuffer(FVulkanDevice* InDevice, FVulkanCommandBufferPool* InCommandBufferPool, bool bInIsUploadOnly)
@@ -343,10 +339,7 @@ void FVulkanCommandBufferManager::SubmitUploadCmdBuffer(uint32 NumSignalSemaphor
 		check(UploadCmdBuffer->IsOutsideRenderPass());
 
 #if UE_BUILD_DEBUG || UE_BUILD_DEVELOPMENT
-		if (CVarVulkanDebugBarrier.GetValueOnAnyThread() & 4)
-		{
-			VulkanRHI::InsertHeavyWeightBarrier(UploadCmdBuffer->GetHandle());
-		}
+		VulkanRHI::DebugHeavyWeightBarrier(UploadCmdBuffer->GetHandle(), 4);
 #endif
 
 		UploadCmdBuffer->End();
@@ -370,12 +363,7 @@ void FVulkanCommandBufferManager::SubmitActiveCmdBuffer(VulkanRHI::FSemaphore* S
 			ActiveCmdBuffer->EndRenderPass();
 		}
 
-#if UE_BUILD_DEBUG || UE_BUILD_DEVELOPMENT
-		if (CVarVulkanDebugBarrier.GetValueOnAnyThread() & 8)
-		{
-			VulkanRHI::InsertHeavyWeightBarrier(ActiveCmdBuffer->GetHandle());
-		}
-#endif
+		VulkanRHI::DebugHeavyWeightBarrier(ActiveCmdBuffer->GetHandle(), 8);
 
 		ActiveCmdBuffer->End();
 		if (SignalSemaphore)
diff --git a/Engine/Source/Runtime/VulkanRHI/Private/VulkanCommands.cpp b/Engine/Source/Runtime/VulkanRHI/Private/VulkanCommands.cpp
--- a/Engine/Source/Runtime/VulkanRHI/Private/VulkanCommands.cpp
+++ b/Engine/Source/Runtime/VulkanRHI/Private/VulkanCommands.cpp
@@ -28,10 +28,6 @@ static FAutoConsoleVariableRef CVarVulkanSubmitAfterEveryEndRenderPass(
 	ECVF_Default
 );
 
-#if UE_BUILD_DEBUG || UE_BUILD_DEVELOPMENT
-extern TAutoConsoleVariable<int32> CVarVulkanDebugBarrier;
-#endif
-
 // make sure what the hardware expects matches what we give it for indirect arguments
 static_assert(sizeof(FRHIDrawIndirectParameters) == sizeof(VkDrawIndirectCommand), "FRHIDrawIndirectParameters size is wrong.");
 static_assert(STRUCT_OFFSET(FRHIDrawIndirectParameters, VertexCountPerInstance) == STRUCT_OFFSET(VkDrawIndirectCommand, vertexCount), "Wrong offset of FRHIDrawIndirectParameters::VertexCountPerInstance.");
@@ -130,12 +126,7 @@ void FVulkanCommandListContext::RHIDispatchComputeShader(uint32 ThreadGroupCount
 		GpuProfiler.RegisterGPUWork(1);
 	}
 
-#if UE_BUILD_DEBUG || UE_BUILD_DEVELOPMENT
-	if (CVarVulkanDebugBarrier.GetValueOnRenderThread() & 2)
-	{
-		VulkanRHI::InsertHeavyWeightBarrier(CmdBuffer);
-	}
-#endif
+	VulkanRHI::DebugHeavyWeightBarrier(CmdBuffer, 2);
 }
 
 void FVulkanCommandListContext::RHIDispatchIndirectComputeShader(FVertexBufferRHIParamRef ArgumentBufferRHI, uint32 ArgumentOffset) 
@@ -165,12 +156,7 @@ void FVulkanCommandListContext::RHIDispatchIndirectComputeShader(FVertexBufferRH
 		GpuProfiler.RegisterGPUWork(1);
 	}
 
-#if UE_BUILD_DEBUG || UE_BUILD_DEVELOPMENT
-	if (CVarVulkanDebugBarrier.GetValueOnRenderThread() & 2)
-	{
-		VulkanRHI::InsertHeavyWeightBarrier(CmdBuffer);
-	}
-#endif
+	VulkanRHI::DebugHeavyWeightBarrier(CmdBuffer, 2);
 }
 
 void FVulkanCommandListContext::RHISetBoundShaderState(FBoundShaderStateRHIParamRef BoundShaderStateRHI)
diff --git a/Engine/Source/Runtime/VulkanRHI/Private/VulkanIndexBuffer.cpp b/Engine/Source/Runtime/VulkanRHI/Private/VulkanIndexBuffer.cpp
--- a/Engine/Source/Runtime/VulkanRHI/Private/VulkanIndexBuffer.cpp
+++ b/Engine/Source/Runtime/VulkanRHI/Private/VulkanIndexBuffer.cpp
@@ -11,11 +11,6 @@
 #include "VulkanLLM.h"
 
 
-#if UE_BUILD_DEBUG || UE_BUILD_DEVELOPMENT
-extern TAutoConsoleVariable<int32> CVarVulkanDebugBarrier;
-#endif
-
-
 static TMap<FVulkanResourceMultiBuffer*, VulkanRHI::FPendingBufferLock> GPendingLockIBs;
 static FCriticalSection GPendingLockIBsMutex;
 
@@ -239,12 +234,7 @@ inline void FVulkanResourceMultiBuffer::InternalUnlock(FVulkanCommandListContext
 	ensure(Cmd->IsOutsideRenderPass());
 	VkCommandBuffer CmdBuffer = Cmd->GetHandle();
 
-#if UE_BUILD_DEBUG || UE_BUILD_DEVELOPMENT
-	if (CVarVulkanDebugBarrier.GetValueOnRenderThread() & 16)
-	{
-		VulkanRHI::InsertHeavyWeightBarrier(CmdBuffer);
-	}
-#endif
+	VulkanRHI::DebugHeavyWeightBarrier(CmdBuffer, 16);
 
 	VkBufferCopy Region;
 	FMemory::Memzero(Region);
diff --git a/Engine/Source/Runtime/VulkanRHI/Private/VulkanRHIPrivate.h b/Engine/Source/Runtime/VulkanRHI/Private/VulkanRHIPrivate.h
--- a/Engine/Source/Runtime/VulkanRHI/Private/VulkanRHIPrivate.h
+++ b/Engine/Source/Runtime/VulkanRHI/Private/VulkanRHIPrivate.h
@@ -751,6 +751,10 @@ static inline VkPrimitiveTopology UEToVulkanType(EPrimitiveType PrimitiveType)
 	return VK_PRIMITIVE_TOPOLOGY_MAX_ENUM;
 }
 
+#if UE_BUILD_DEBUG || UE_BUILD_DEVELOPMENT
+extern TAutoConsoleVariable<int32> CVarVulkanDebugBarrier;
+#endif
+
 namespace VulkanRHI
 {
 	static inline FString GetPipelineCacheFilename()
@@ -817,41 +821,46 @@ namespace VulkanRHI
 		return VK_IMAGE_LAYOUT_DEPTH_STENCIL_ATTACHMENT_OPTIMAL;
 	}
 
-	inline void InsertHeavyWeightBarrier(VkCommandBuffer CmdBuffer)
-	{
-		VkMemoryBarrier Barrier;
-		ZeroVulkanStruct(Barrier, VK_STRUCTURE_TYPE_MEMORY_BARRIER);
-		Barrier.srcAccessMask = VK_ACCESS_INDIRECT_COMMAND_READ_BIT |
-			VK_ACCESS_INDEX_READ_BIT |
-			VK_ACCESS_VERTEX_ATTRIBUTE_READ_BIT |
-			VK_ACCESS_UNIFORM_READ_BIT |
-			VK_ACCESS_INPUT_ATTACHMENT_READ_BIT |
-			VK_ACCESS_SHADER_READ_BIT |
-			VK_ACCESS_SHADER_WRITE_BIT |
-			VK_ACCESS_COLOR_ATTACHMENT_READ_BIT |
-			VK_ACCESS_COLOR_ATTACHMENT_WRITE_BIT |
-			VK_ACCESS_DEPTH_STENCIL_ATTACHMENT_READ_BIT |
-			VK_ACCESS_DEPTH_STENCIL_ATTACHMENT_WRITE_BIT |
-			VK_ACCESS_TRANSFER_READ_BIT |
-			VK_ACCESS_TRANSFER_WRITE_BIT |
-			VK_ACCESS_HOST_READ_BIT |
-			VK_ACCESS_HOST_WRITE_BIT;
-		Barrier.dstAccessMask = VK_ACCESS_INDIRECT_COMMAND_READ_BIT |
-			VK_ACCESS_INDEX_READ_BIT |
-			VK_ACCESS_VERTEX_ATTRIBUTE_READ_BIT |
-			VK_ACCESS_UNIFORM_READ_BIT |
-			VK_ACCESS_INPUT_ATTACHMENT_READ_BIT |
-			VK_ACCESS_SHADER_READ_BIT |
-			VK_ACCESS_SHADER_WRITE_BIT |
-			VK_ACCESS_COLOR_ATTACHMENT_READ_BIT |
-			VK_ACCESS_COLOR_ATTACHMENT_WRITE_BIT |
-			VK_ACCESS_DEPTH_STENCIL_ATTACHMENT_READ_BIT |
-			VK_ACCESS_DEPTH_STENCIL_ATTACHMENT_WRITE_BIT |
-			VK_ACCESS_TRANSFER_READ_BIT |
-			VK_ACCESS_TRANSFER_WRITE_BIT |
-			VK_ACCESS_HOST_READ_BIT |
-			VK_ACCESS_HOST_WRITE_BIT;
-		VulkanRHI::vkCmdPipelineBarrier(CmdBuffer, VK_PIPELINE_STAGE_ALL_COMMANDS_BIT, VK_PIPELINE_STAGE_ALL_COMMANDS_BIT, 0, 1, &Barrier, 0, nullptr, 0, nullptr);
+	inline void DebugHeavyWeightBarrier(VkCommandBuffer CmdBuffer, int32 CVarConditionMask)
+	{
+#if UE_BUILD_DEBUG || UE_BUILD_DEVELOPMENT
+		if (CVarVulkanDebugBarrier.GetValueOnAnyThread() & CVarConditionMask)
+		{
+			VkMemoryBarrier Barrier;
+			ZeroVulkanStruct(Barrier, VK_STRUCTURE_TYPE_MEMORY_BARRIER);
+			Barrier.srcAccessMask = VK_ACCESS_INDIRECT_COMMAND_READ_BIT |
+				VK_ACCESS_INDEX_READ_BIT |
+				VK_ACCESS_VERTEX_ATTRIBUTE_READ_BIT |
+				VK_ACCESS_UNIFORM_READ_BIT |
+				VK_ACCESS_INPUT_ATTACHMENT_READ_BIT |
+				VK_ACCESS_SHADER_READ_BIT |
+				VK_ACCESS_SHADER_WRITE_BIT |
+				VK_ACCESS_COLOR_ATTACHMENT_READ_BIT |
+				VK_ACCESS_COLOR_ATTACHMENT_WRITE_BIT |
+				VK_ACCESS_DEPTH_STENCIL_ATTACHMENT_READ_BIT |
+				VK_ACCESS_DEPTH_STENCIL_ATTACHMENT_WRITE_BIT |
+				VK_ACCESS_TRANSFER_READ_BIT |
+				VK_ACCESS_TRANSFER_WRITE_BIT |
+				VK_ACCESS_HOST_READ_BIT |
+				VK_ACCESS_HOST_WRITE_BIT;
+			Barrier.dstAccessMask = VK_ACCESS_INDIRECT_COMMAND_READ_BIT |
+				VK_ACCESS_INDEX_READ_BIT |
+				VK_ACCESS_VERTEX_ATTRIBUTE_READ_BIT |
+				VK_ACCESS_UNIFORM_READ_BIT |
+				VK_ACCESS_INPUT_ATTACHMENT_READ_BIT |
+				VK_ACCESS_SHADER_READ_BIT |
+				VK_ACCESS_SHADER_WRITE_BIT |
+				VK_ACCESS_COLOR_ATTACHMENT_READ_BIT |
+				VK_ACCESS_COLOR_ATTACHMENT_WRITE_BIT |
+				VK_ACCESS_DEPTH_STENCIL_ATTACHMENT_READ_BIT |
+				VK_ACCESS_DEPTH_STENCIL_ATTACHMENT_WRITE_BIT |
+				VK_ACCESS_TRANSFER_READ_BIT |
+				VK_ACCESS_TRANSFER_WRITE_BIT |
+				VK_ACCESS_HOST_READ_BIT |
+				VK_ACCESS_HOST_WRITE_BIT;
+			VulkanRHI::vkCmdPipelineBarrier(CmdBuffer, VK_PIPELINE_STAGE_ALL_COMMANDS_BIT, VK_PIPELINE_STAGE_ALL_COMMANDS_BIT, 0, 1, &Barrier, 0, nullptr, 0, nullptr);
+		}
+#endif
 	}
 }
 
diff --git a/Engine/Source/Runtime/VulkanRHI/Private/VulkanRenderTarget.cpp b/Engine/Source/Runtime/VulkanRHI/Private/VulkanRenderTarget.cpp
--- a/Engine/Source/Runtime/VulkanRHI/Private/VulkanRenderTarget.cpp
+++ b/Engine/Source/Runtime/VulkanRHI/Private/VulkanRenderTarget.cpp
@@ -44,7 +44,8 @@ TAutoConsoleVariable<int32> CVarVulkanDebugBarrier(
 	TEXT(" 2: Enable heavy barriers after every dispatch\n")
 	TEXT(" 4: Enable heavy barriers after upload cmd buffers\n")
 	TEXT(" 8: Enable heavy barriers after active cmd buffers\n")
-	TEXT(" 16: Enable heavy buffer barrier after uploads\n"),
+	TEXT(" 16: Enable heavy buffer barrier after uploads\n")
+	TEXT(" 32: Enable heavy buffer barrier between acquiring back buffer and blitting into swapchain\n"),
 	ECVF_Default
 );
 #endif
@@ -222,12 +223,7 @@ void FTransitionAndLayoutManager::EndEmulatedRenderPass(FVulkanCmdBuffer* CmdBuf
 	CmdBuffer->EndRenderPass();
 	CurrentRenderPass = nullptr;
 
-#if UE_BUILD_DEBUG || UE_BUILD_DEVELOPMENT
-	if (CVarVulkanDebugBarrier.GetValueOnRenderThread() & 1)
-	{
-		VulkanRHI::InsertHeavyWeightBarrier(CmdBuffer->GetHandle());
-	}
-#endif
+	VulkanRHI::DebugHeavyWeightBarrier(CmdBuffer->GetHandle(), 1);
 }
 
 void FTransitionAndLayoutManager::BeginRealRenderPass(FVulkanCommandListContext& Context, FVulkanDevice& InDevice, FVulkanCmdBuffer* CmdBuffer, const FRHIRenderPassInfo& RPInfo, const FVulkanRenderTargetLayout& RTLayout, FVulkanRenderPass* RenderPass, FVulkanFramebuffer* Framebuffer)
@@ -431,12 +427,7 @@ void FTransitionAndLayoutManager::EndRealRenderPass(FVulkanCmdBuffer* CmdBuffer)
 	CurrentRenderPass = nullptr;
 	bInsideRealRenderPass = false;
 
-#if UE_BUILD_DEBUG || UE_BUILD_DEVELOPMENT
-	if (CVarVulkanDebugBarrier.GetValueOnRenderThread() & 1)
-	{
-		VulkanRHI::InsertHeavyWeightBarrier(CmdBuffer->GetHandle());
-	}
-#endif
+	VulkanRHI::DebugHeavyWeightBarrier(CmdBuffer->GetHandle(), 1);
 }
 
 void FTransitionAndLayoutManager::NotifyDeletedRenderTarget(FVulkanDevice& InDevice, VkImage Image)
diff --git a/Engine/Source/Runtime/VulkanRHI/Private/VulkanSwapChain.cpp b/Engine/Source/Runtime/VulkanRHI/Private/VulkanSwapChain.cpp
--- a/Engine/Source/Runtime/VulkanRHI/Private/VulkanSwapChain.cpp
+++ b/Engine/Source/Runtime/VulkanRHI/Private/VulkanSwapChain.cpp
@@ -236,30 +236,71 @@ FVulkanSwapChain::FVulkanSwapChain(VkInstance InInstance, FVulkanDevice& InDevic
 				UE_CLOG(bFirstTimeLog, LogVulkanRHI, Display, TEXT("- VK_PRESENT_MODE_FIFO_KHR (%d)"), (int32)VK_PRESENT_MODE_FIFO_KHR);
 				break;
 			default:
-				UE_CLOG(bFirstTimeLog, LogVulkanRHI, Display, TEXT("- VkPresentModeKHR mode %d"), (int32)FoundPresentModes[i]);
+				UE_CLOG(bFirstTimeLog, LogVulkanRHI, Display, TEXT("- VkPresentModeKHR %d"), (int32)FoundPresentModes[i]);
 				break;
 			}
 		}
 
-		// Until FVulkanViewport::Present honors SyncInterval, we need to disable vsync for the spectator window if using an HMD.
-		const bool bDisableVsyncForHMD = (FVulkanDynamicRHI::HMDVulkanExtensions.IsValid()) ? FVulkanDynamicRHI::HMDVulkanExtensions->ShouldDisableVulkanVSync() : false;
-
-		if (bFoundPresentModeMailbox)
-		{
-			PresentMode = VK_PRESENT_MODE_MAILBOX_KHR;
-		}
-		else if (bFoundPresentModeImmediate && bDisableVsyncForHMD)
-		{
-			PresentMode = VK_PRESENT_MODE_IMMEDIATE_KHR;
-		}
-		else if (bFoundPresentModeFIFO)
+		int32 RequestedPresentMode = -1;
+		if (FParse::Value(FCommandLine::Get(), TEXT("vulkanpresentmode="), RequestedPresentMode))
 		{
-			PresentMode = VK_PRESENT_MODE_FIFO_KHR;
+			bool bRequestSuccessful = false;
+			switch (RequestedPresentMode)
+			{
+			case VK_PRESENT_MODE_MAILBOX_KHR:
+				if (bFoundPresentModeMailbox)
+				{
+					PresentMode = VK_PRESENT_MODE_MAILBOX_KHR;
+					bRequestSuccessful = true;
+				}
+				break;
+			case VK_PRESENT_MODE_IMMEDIATE_KHR:
+				if (bFoundPresentModeImmediate)
+				{
+					PresentMode = VK_PRESENT_MODE_IMMEDIATE_KHR;
+					bRequestSuccessful = true;
+				}
+				break;
+			case VK_PRESENT_MODE_FIFO_KHR:
+				if (bFoundPresentModeFIFO)
+				{
+					PresentMode = VK_PRESENT_MODE_FIFO_KHR;
+					bRequestSuccessful = true;
+				}
+				break;
+			default:
+				break;
+			}
+
+			if (!bRequestSuccessful)
+			{
+				UE_LOG(LogVulkanRHI, Warning, TEXT("Requested PresentMode (%d) is not handled or available, ignoring..."), RequestedPresentMode);
+				RequestedPresentMode = -1;
+			}
 		}
-		else
+
+		if (RequestedPresentMode == -1)
 		{
-			UE_LOG(LogVulkanRHI, Warning, TEXT("Couldn't find desired PresentMode! Using %d"), static_cast<int32>(FoundPresentModes[0]));
-			PresentMode = FoundPresentModes[0];
+			// Until FVulkanViewport::Present honors SyncInterval, we need to disable vsync for the spectator window if using an HMD.
+			const bool bDisableVsyncForHMD = (FVulkanDynamicRHI::HMDVulkanExtensions.IsValid()) ? FVulkanDynamicRHI::HMDVulkanExtensions->ShouldDisableVulkanVSync() : false;
+
+			if (bFoundPresentModeMailbox)
+			{
+				PresentMode = VK_PRESENT_MODE_MAILBOX_KHR;
+			}
+			else if (bFoundPresentModeImmediate && bDisableVsyncForHMD)
+			{
+				PresentMode = VK_PRESENT_MODE_IMMEDIATE_KHR;
+			}
+			else if (bFoundPresentModeFIFO)
+			{
+				PresentMode = VK_PRESENT_MODE_FIFO_KHR;
+			}
+			else
+			{
+				UE_LOG(LogVulkanRHI, Warning, TEXT("Couldn't find desired PresentMode! Using %d"), static_cast<int32>(FoundPresentModes[0]));
+				PresentMode = FoundPresentModes[0];
+			}
 		}
 
 		UE_CLOG(bFirstTimeLog, LogVulkanRHI, Display, TEXT("Selected VkPresentModeKHR mode %d"), PresentMode);
diff --git a/Engine/Source/Runtime/VulkanRHI/Private/VulkanViewport.cpp b/Engine/Source/Runtime/VulkanRHI/Private/VulkanViewport.cpp
--- a/Engine/Source/Runtime/VulkanRHI/Private/VulkanViewport.cpp
+++ b/Engine/Source/Runtime/VulkanRHI/Private/VulkanViewport.cpp
@@ -18,6 +18,7 @@ FAutoConsoleVariable GCVarDelayAcquireBackBuffer(
 	ECVF_ReadOnly
 );
 
+
 struct FRHICommandAcquireBackBuffer final : public FRHICommand<FRHICommandAcquireBackBuffer>
 {
 	FVulkanViewport* Viewport;
@@ -550,6 +551,8 @@ inline static void CopyImageToBackBuffer(FVulkanCmdBuffer* CmdBuffer, bool bSour
 	Barriers.SetTransition(DestIndex, EImageLayoutBarrier::Undefined, EImageLayoutBarrier::TransferDest);
 	Barriers.Execute(CmdBuffer);
 
+	VulkanRHI::DebugHeavyWeightBarrier(CmdBuffer->GetHandle(), 32);
+
 	VkImageCopy Region;
 	FMemory::Memzero(Region);
 	Region.extent.width = SizeX;
