DR - vk - Don't spam
#rnx

diff --git a/Engine/Source/Runtime/VulkanRHI/Private/VulkanSwapChain.cpp b/Engine/Source/Runtime/VulkanRHI/Private/VulkanSwapChain.cpp
--- a/Engine/Source/Runtime/VulkanRHI/Private/VulkanSwapChain.cpp
+++ b/Engine/Source/Runtime/VulkanRHI/Private/VulkanSwapChain.cpp
@@ -274,7 +274,7 @@ FVulkanSwapChain::FVulkanSwapChain(VkInstance InInstance, FVulkanDevice& InDevic
 
 			if (!bRequestSuccessful)
 			{
-				UE_LOG(LogVulkanRHI, Warning, TEXT("Requested PresentMode (%d) is not handled or available, ignoring..."), RequestedPresentMode);
+				UE_CLOG(bFirstTimeLog, LogVulkanRHI, Warning, TEXT("Requested PresentMode (%d) is not handled or available, ignoring..."), RequestedPresentMode);
 				RequestedPresentMode = -1;
 			}
 		}
