DR - vk - Remove generic pipeline cache

diff --git a/Engine/Source/Runtime/VulkanRHI/Private/VulkanPipeline.cpp b/Engine/Source/Runtime/VulkanRHI/Private/VulkanPipeline.cpp
--- a/Engine/Source/Runtime/VulkanRHI/Private/VulkanPipeline.cpp
+++ b/Engine/Source/Runtime/VulkanRHI/Private/VulkanPipeline.cpp
@@ -357,129 +357,6 @@ bool FVulkanPipelineStateCacheManager::Load(const TArray<FString>& CacheFilename
 	}
 #endif
 
-#if VULKAN_ENABLE_GENERIC_PIPELINE_CACHE_FILE
-	for (const FString& CacheFilename : CacheFilenames)
-	{
-		TArray<uint8> MemFile;
-		if (FFileHelper::LoadFileToArray(MemFile, *CacheFilename, FILEREAD_Silent))
-		{
-			FMemoryReader Ar(MemFile);
-
-			FVulkanPipelineStateCacheFile File;
-
-			FShaderUCodeCache::TDataMap FileShaderCacheData;
-			File.ShaderCache = &FileShaderCacheData;
-
-			bool Valid = File.Load(Ar, *CacheFilename);
-			if (!Valid)
-			{
-				UE_LOG(LogVulkanRHI, Warning, TEXT("Unable to load pipeline cache '%s'"), *CacheFilename);
-				continue;
-			}
-
-			// Create the binary cache if we haven't already
-			if (PipelineCache == VK_NULL_HANDLE)
-			{
-				VkPipelineCacheCreateInfo PipelineCacheInfo;
-				ZeroVulkanStruct(PipelineCacheInfo, VK_STRUCTURE_TYPE_PIPELINE_CACHE_CREATE_INFO);
-				VERIFYVULKANRESULT(VulkanRHI::vkCreatePipelineCache(Device->GetInstanceHandle(), &PipelineCacheInfo, VULKAN_CPU_ALLOCATOR, &PipelineCache));
-			}
-
-			// Not using TMap::Append to avoid copying duplicate microcode
-			for (auto& Pair : FileShaderCacheData)
-			{
-				if (!ShaderCache.Data.Find(Pair.Key))
-				{
-					ShaderCache.Data.Add(Pair.Key, Pair.Value);
-				}
-			}
-
-			double BeginTime = FPlatformTime::Seconds();
-			{
-				for (int32 Index = 0; Index < File.GfxPipelineEntries.Num(); ++Index)
-				{
-					FGfxPipelineEntry* GfxEntry = File.GfxPipelineEntries[Index];
-
-					FShaderHashes ShaderHashes;
-					for (int32 i = 0; i < ShaderStage::NumStages; ++i)
-					{
-						ShaderHashes.Stages[i] = GfxEntry->ShaderHashes[i];
-
-						GfxEntry->ShaderMicrocodes[i] = ShaderCache.Get(GfxEntry->ShaderHashes[i]);
-					}
-					ShaderHashes.Finalize();
-
-					uint32 EntryHash = GfxEntry->GetEntryHash();
-					if (GfxPipelineEntries.Find(EntryHash))
-					{
-						delete GfxEntry;
-					}
-					else
-					{
-						FHashToGfxPipelinesMap* Found = ShaderHashToGfxPipelineMap.Find(ShaderHashes);
-						if (!Found)
-						{
-							Found = &ShaderHashToGfxPipelineMap.Add(ShaderHashes);
-						}
-
-						CreatGfxEntryRuntimeObjects(GfxEntry);
-#if VULKAN_ENABLE_LRU_CACHE
-						FVulkanGfxPipeline* Pipeline = new FVulkanGfxPipeline(Device, EntryHash, ShaderHashes.Hash);
-#else
-						FVulkanGfxPipeline* Pipeline = new FVulkanGfxPipeline(Device);
-#endif
-						CreateGfxPipelineFromEntry(GfxEntry, Pipeline);
-
-						Found->Add(EntryHash, Pipeline);
-						GfxPipelineEntries.Add(EntryHash, GfxEntry);
-						GfxEntry->bConsumed = true;
-					}
-				}
-			}
-
-			for (int32 Index = 0; Index < File.ComputePipelineEntries.Num(); ++Index)
-			{
-				FComputePipelineEntry* ComputeEntry = File.ComputePipelineEntries[Index];
-				ComputeEntry->ShaderMicrocode = ShaderCache.Get(ComputeEntry->ShaderHash);
-				ComputeEntry->CalculateEntryHash();
-
-				if (ComputePipelineEntries.Find(ComputeEntry->EntryHash))
-				{
-					delete ComputeEntry;
-				}
-				else
-				{
-					CreateComputeEntryRuntimeObjects(ComputeEntry);
-
-					FVulkanComputePipeline* Pipeline = CreateComputePipelineFromEntry(ComputeEntry);
-					ComputeEntryHashToPipelineMap.Add(ComputeEntry->EntryHash, Pipeline);
-					ComputePipelineEntries.Add(ComputeEntry->EntryHash, ComputeEntry);
-					Pipeline->AddRef();
-				}
-			}
-
-			double EndTime = FPlatformTime::Seconds();
-			UE_LOG(LogVulkanRHI, Display, TEXT("Loaded pipeline cache in %.2f seconds"), (float)(EndTime - BeginTime));
-		}
-		else
-		{
-			UE_LOG(LogVulkanRHI, Warning, TEXT("Unable to load pipeline cache '%s'"), *CacheFilename);
-		}
-	}
-
-#endif
-
-#if VULKAN_ENABLE_GENERIC_PIPELINE_CACHE_FILE
-	if (ShaderCache.Data.Num())
-	{
-		UE_LOG(LogVulkanRHI, Display, TEXT("Pipeline cache: %d Gfx Pipelines, %d Compute Pipelines, %d Microcodes"), GfxPipelineEntries.Num(), ComputePipelineEntries.Num(), ShaderCache.Data.Num());
-	}
-	else
-	{
-		UE_LOG(LogVulkanRHI, Display, TEXT("Pipeline cache: No pipeline cache(s) loaded"));
-	}
-#endif
-
 	// Lazily create the cache in case the load failed
 	if (PipelineCache == VK_NULL_HANDLE)
 	{
@@ -722,37 +599,6 @@ void FVulkanPipelineStateCacheManager::Save(const FString& CacheFilename, bool b
 		}
 	}
 #endif
-
-#if VULKAN_ENABLE_GENERIC_PIPELINE_CACHE_FILE
-	{
-		// Now the generic cache
-		TArray<uint8> MemFile;
-		FMemoryWriter Ar(MemFile);
-		FVulkanPipelineStateCacheFile File;
-
-		File.Header.Version = VERSION;
-		File.Header.SizeOfGfxEntry = (int32)sizeof(FGfxPipelineEntry);
-		File.Header.SizeOfComputeEntry = (int32)sizeof(FComputePipelineEntry);
-		File.Header.UncompressedSize = 0;
-
-		// Shader ucode cache
-		File.ShaderCache = &ShaderCache.Data;
-
-		// Then Gfx entries
-		GfxPipelineEntries.GenerateValueArray(File.GfxPipelineEntries);
-
-		// And Compute entries
-		ComputePipelineEntries.GenerateValueArray(File.ComputePipelineEntries);
-
-		File.Save(Ar);
-
-		if (FFileHelper::SaveArrayToFile(MemFile, *CacheFilename))
-		{
-			UE_LOG(LogVulkanRHI, Display, TEXT("Saved pipeline cache file '%s', %d Gfx Pipelines, %d Compute Pipelines, %d Microcodes, %d bytes"), *CacheFilename, GfxPipelineEntries.Num(), ComputePipelineEntries.Num(), ShaderCache.Data.Num(), MemFile.Num());
-		}
-	}
-
-#endif
 }
 
 #if VULKAN_ENABLE_LRU_CACHE
@@ -795,13 +641,6 @@ FVulkanRHIGraphicsPipelineState* FVulkanPipelineStateCacheManager::CreateAndAdd(
 	}
 #else
 	FVulkanGfxPipeline* Pipeline = new FVulkanGfxPipeline(Device);
-#if VULKAN_ENABLE_GENERIC_PIPELINE_CACHE_FILE
-	{
-		FScopeLock ScopeLock(&GfxPipelineEntriesCS);
-		GfxPipelineEntries.Add(GfxEntry->GetEntryHash(), GfxEntry);
-		GfxEntry->bConsumed = true;
-	}
-#endif
 #endif
 	
 
@@ -2229,100 +2068,6 @@ inline void SerializeArray(FArchive& Ar, TArray<T*>& Array)
 	}
 }
 
-#if VULKAN_ENABLE_GENERIC_PIPELINE_CACHE_FILE
-void FVulkanPipelineStateCacheManager::FVulkanPipelineStateCacheFile::Save(FArchive& Ar)
-{
-	check(ShaderCache);
-
-	// Modify VERSION if serialization changes
-	TArray<uint8> DataBuffer;
-	FMemoryWriter DataAr(DataBuffer);
-
-	DataAr << *ShaderCache;
-	SerializeArray(DataAr, GfxPipelineEntries);
-	SerializeArray(DataAr, ComputePipelineEntries);
-
-	// compress the data buffer
-	TArray<uint8> CompressedDataBuffer = DataBuffer;
-	if (GEnablePipelineCacheCompression)
-	{
-		Header.UncompressedSize = DataBuffer.Num() * DataBuffer.GetTypeSize();
-		int32 CompressedSize = CompressedDataBuffer.Num();
-		if (FCompression::CompressMemory(CompressionFlags, CompressedDataBuffer.GetData(), CompressedSize,
-			DataBuffer.GetData(), Header.UncompressedSize))
-		{
-			CompressedDataBuffer.SetNum(CompressedSize);
-		}
-		CompressedDataBuffer.Shrink();
-	}
-
-	Ar << Header.Version;
-	Ar << Header.SizeOfGfxEntry;
-	Ar << Header.SizeOfComputeEntry;
-	Ar << Header.UncompressedSize;
-
-	Ar << CompressedDataBuffer;
-}
-
-
-bool FVulkanPipelineStateCacheManager::FVulkanPipelineStateCacheFile::Load(FArchive& Ar, const TCHAR* Filename)
-{
-	check(ShaderCache);
-
-	// Modify VERSION if serialization changes
-	Ar << Header.Version;
-	if (Header.Version != VERSION)
-	{
-		UE_LOG(LogVulkanRHI, Warning, TEXT("Unable to load shader cache due to mismatched Version %d != %d"), Header.Version, (int32)VERSION);
-		return false;
-	}
-
-	Ar << Header.SizeOfGfxEntry;
-	if (Header.SizeOfGfxEntry != (int32)(sizeof(FGfxPipelineEntry)))
-	{
-		UE_LOG(LogVulkanRHI, Warning, TEXT("Unable to load shader cache due to mismatched size of FGfxEntry %d != %d; forgot to bump up VERSION?"), Header.SizeOfGfxEntry, (int32)sizeof(FGfxPipelineEntry));
-		return false;
-	}
-
-	Ar << Header.SizeOfComputeEntry;
-	if (Header.SizeOfComputeEntry != (int32)(sizeof(FComputePipelineEntry)))
-	{
-		UE_LOG(LogVulkanRHI, Warning, TEXT("Unable to load shader cache due to mismatched size of FComputePipelineEntry %d != %d; forgot to bump up VERSION?"), Header.SizeOfComputeEntry, (int32)sizeof(FComputePipelineEntry));
-		return false;
-	}
-
-	Ar << Header.UncompressedSize;
-
-	TArray<uint8> CompressedDataBuffer;
-	Ar << CompressedDataBuffer;
-
-	TArray<uint8> UncompressedDataBuffer;
-	if (Header.UncompressedSize != 0)
-	{
-		const uint32 CompressedSize = CompressedDataBuffer.Num() * CompressedDataBuffer.GetTypeSize();
-		UncompressedDataBuffer.SetNum(Header.UncompressedSize);
-		if (!FCompression::UncompressMemory(CompressionFlags, UncompressedDataBuffer.GetData(), Header.UncompressedSize, CompressedDataBuffer.GetData(), CompressedSize))
-		{
-			UE_LOG(LogVulkanRHI, Error, TEXT("Failed to uncompress data for pipeline cache file %s!"), Filename);
-			return false;
-		}
-	}
-	else
-	{
-		UncompressedDataBuffer = CompressedDataBuffer;
-	}
-
-	FMemoryReader DataAr(UncompressedDataBuffer);
-	DataAr << *ShaderCache;
-
-	SerializeArray(DataAr, GfxPipelineEntries);
-
-	SerializeArray(DataAr, ComputePipelineEntries);
-
-	return true;
-}
-#endif
-
 bool FVulkanPipelineStateCacheManager::BinaryCacheMatches(FVulkanDevice* InDevice, const TArray<uint8>& DeviceCache)
 {
 	if (DeviceCache.Num() > 4)
diff --git a/Engine/Source/Runtime/VulkanRHI/Private/VulkanPipeline.h b/Engine/Source/Runtime/VulkanRHI/Private/VulkanPipeline.h
--- a/Engine/Source/Runtime/VulkanRHI/Private/VulkanPipeline.h
+++ b/Engine/Source/Runtime/VulkanRHI/Private/VulkanPipeline.h
@@ -227,14 +227,6 @@ public:
 
 	void RebuildCache();
 
-#if VULKAN_ENABLE_GENERIC_PIPELINE_CACHE_FILE
-	enum
-	{
-		// Bump every time serialization changes
-		VERSION = 23,
-	};
-#endif
-
 	struct FDescriptorSetLayoutBinding
 	{
 		uint32 Binding;
@@ -988,29 +980,6 @@ private:
 	};
 #endif
 
-#if VULKAN_ENABLE_GENERIC_PIPELINE_CACHE_FILE
-	struct FVulkanPipelineStateCacheFile
-	{
-		struct FFileHeader
-		{
-			int32 Version = -1;
-			int32 SizeOfGfxEntry = -1;
-			int32 SizeOfComputeEntry = -1;
-			int32 UncompressedSize = 0; // 0 == file is uncompressed
-		} Header;
-
-		FShaderUCodeCache::TDataMap* ShaderCache = nullptr;
-
-		TArray<FGfxPipelineEntry*> GfxPipelineEntries;
-		TArray<FComputePipelineEntry*> ComputePipelineEntries;
-
-		void Save(FArchive& Ar);
-		bool Load(FArchive& Ar, const TCHAR* Filename);
-
-		static const ECompressionFlags CompressionFlags = (ECompressionFlags)(COMPRESS_ZLIB | COMPRESS_BiasSpeed);
-	};
-#endif
-
 	static bool BinaryCacheMatches(FVulkanDevice* InDevice, const TArray<uint8>& DeviceCache);
 };
 
diff --git a/Engine/Source/Runtime/VulkanRHI/Public/VulkanConfiguration.h b/Engine/Source/Runtime/VulkanRHI/Public/VulkanConfiguration.h
--- a/Engine/Source/Runtime/VulkanRHI/Public/VulkanConfiguration.h
+++ b/Engine/Source/Runtime/VulkanRHI/Public/VulkanConfiguration.h
@@ -108,10 +108,6 @@
 	#define VULKAN_ENABLE_LRU_CACHE								0
 #endif
 
-#ifndef VULKAN_ENABLE_GENERIC_PIPELINE_CACHE_FILE
-	#define VULKAN_ENABLE_GENERIC_PIPELINE_CACHE_FILE			0
-#endif
-
 #ifdef VK_KHR_maintenance1
 	#define VULKAN_SUPPORTS_MAINTENANCE_LAYER1					1
 #else
