DR - vk - Support and fixes for color conversion samplers

diff --git a/Engine/Source/Runtime/VulkanRHI/Private/Lumin/VulkanLuminPlatform.cpp b/Engine/Source/Runtime/VulkanRHI/Private/Lumin/VulkanLuminPlatform.cpp
--- a/Engine/Source/Runtime/VulkanRHI/Private/Lumin/VulkanLuminPlatform.cpp
+++ b/Engine/Source/Runtime/VulkanRHI/Private/Lumin/VulkanLuminPlatform.cpp
@@ -130,6 +130,7 @@ void FVulkanLuminPlatform::GetDeviceExtensions(TArray<const ANSICHAR*>& OutExten
 	{
 		OutExtensions.Add("VK_LAYER_NV_vgd");
 	}
+	OutExtensions.Add(VK_KHR_SAMPLER_YCBCR_CONVERSION_EXTENSION_NAME);
 }
 
 bool FVulkanLuminPlatform::ForceEnableDebugMarkers()
diff --git a/Engine/Source/Runtime/VulkanRHI/Private/VulkanCommandBuffer.cpp b/Engine/Source/Runtime/VulkanRHI/Private/VulkanCommandBuffer.cpp
--- a/Engine/Source/Runtime/VulkanRHI/Private/VulkanCommandBuffer.cpp
+++ b/Engine/Source/Runtime/VulkanRHI/Private/VulkanCommandBuffer.cpp
@@ -28,10 +28,6 @@ static FAutoConsoleVariableRef CVarVulkanProfileCmdBuffers(
 
 #define CMD_BUFFER_TIME_TO_WAIT_BEFORE_DELETING		10
 
-#if UE_BUILD_DEBUG || UE_BUILD_DEVELOPMENT
-extern TAutoConsoleVariable<int32> CVarVulkanDebugBarrier;
-#endif
-
 const uint32 GNumberOfFramesBeforeDeletingDescriptorPool = 300;
 
 FVulkanCmdBuffer::FVulkanCmdBuffer(FVulkanDevice* InDevice, FVulkanCommandBufferPool* InCommandBufferPool, bool bInIsUploadOnly)
diff --git a/Engine/Source/Runtime/VulkanRHI/Private/VulkanCommandWrappers.h b/Engine/Source/Runtime/VulkanRHI/Private/VulkanCommandWrappers.h
--- a/Engine/Source/Runtime/VulkanRHI/Private/VulkanCommandWrappers.h
+++ b/Engine/Source/Runtime/VulkanRHI/Private/VulkanCommandWrappers.h
@@ -165,6 +165,10 @@ struct FWrapLayer
 #if VULKAN_USE_CREATE_WIN32_SURFACE
 	static void CreateWin32SurfaceKHR(VkResult Result, VkInstance instance, const VkWin32SurfaceCreateInfoKHR* pCreateInfo, VkSurfaceKHR* pSurface);
 #endif
+#if VULKAN_SUPPORTS_COLOR_CONVERSIONS
+	static void CreateSamplerYcbcrConversionKHR(VkResult Result, VkDevice device, const VkSamplerYcbcrConversionCreateInfo* pCreateInfo, VkSamplerYcbcrConversion* YcbcrConversion) VULKAN_LAYER_BODY
+	static void DestroySamplerYcbcrConversionKHR(VkResult Result, VkDevice device, VkSamplerYcbcrConversion ycbcrConversion) VULKAN_LAYER_BODY
+#endif
 };
 
 #undef VULKAN_LAYER_BODY
@@ -1319,6 +1323,23 @@ namespace VulkanRHI
 	}
 #endif
 
+#if VULKAN_SUPPORTS_COLOR_CONVERSIONS
+	static FORCEINLINE_DEBUGGABLE VkResult vkCreateSamplerYcbcrConversionKHR(VkDevice Device, const VkSamplerYcbcrConversionCreateInfo* CreateInfo, const VkAllocationCallbacks* Allocator, VkSamplerYcbcrConversion* YcbcrConversion)
+	{
+		FWrapLayer::CreateSamplerYcbcrConversionKHR(VK_RESULT_MAX_ENUM, Device, CreateInfo, YcbcrConversion);
+		VkResult Result = VULKANAPINAMESPACE::vkCreateSamplerYcbcrConversionKHR(Device, CreateInfo, Allocator, YcbcrConversion);
+		FWrapLayer::CreateSamplerYcbcrConversionKHR(Result, Device, CreateInfo, YcbcrConversion);
+		return Result;
+	}
+
+	static FORCEINLINE_DEBUGGABLE void vkDestroySamplerYcbcrConversionKHR(VkDevice Device, VkSamplerYcbcrConversion YcbcrConversion, const VkAllocationCallbacks* Allocator)
+	{
+		FWrapLayer::DestroySamplerYcbcrConversionKHR(VK_RESULT_MAX_ENUM, Device, YcbcrConversion);
+		VULKANAPINAMESPACE::vkDestroySamplerYcbcrConversionKHR(Device, YcbcrConversion, Allocator);
+		FWrapLayer::DestroySamplerYcbcrConversionKHR(VK_SUCCESS, Device, YcbcrConversion);
+	}
+#endif
+
 #if VULKAN_ENABLE_IMAGE_TRACKING_LAYER
 	void BindDebugLabelName(VkImage Image, const TCHAR* Name);
 #endif
diff --git a/Engine/Source/Runtime/VulkanRHI/Private/VulkanDebug.cpp b/Engine/Source/Runtime/VulkanRHI/Private/VulkanDebug.cpp
--- a/Engine/Source/Runtime/VulkanRHI/Private/VulkanDebug.cpp
+++ b/Engine/Source/Runtime/VulkanRHI/Private/VulkanDebug.cpp
@@ -314,8 +314,6 @@ static VkBool32 DebugUtilsCallback(VkDebugUtilsMessageSeverityFlagBitsEXT MsgSev
 
 	return VK_FALSE;
 }
-
-#undef VULKAN_REPORT_LOG
 #endif
 
 void FVulkanDynamicRHI::SetupDebugLayerCallback()
@@ -3861,6 +3859,31 @@ void FWrapLayer::CreateWin32SurfaceKHR(VkResult Result, VkInstance Instance, con
 }
 #endif
 
+#if VULKAN_SUPPORTS_COLOR_CONVERSIONS
+void FWrapLayer::CreateSamplerYcbcrConversionKHR(VkResult Result, VkDevice Device, const VkSamplerYcbcrConversionCreateInfo* CreateInfo, VkSamplerYcbcrConversion* YcbcrConversion)
+{
+	if (Result == VK_RESULT_MAX_ENUM)
+	{
+		DevicePrintfBeginResult(Device, FString::Printf(TEXT("vkCreateSamplerYcbcrConversionKHR(CreateInfo=0x%p, YcbcrConversion=0x%p)[...]"), CreateInfo, YcbcrConversion));
+	}
+	else
+	{
+		PrintResultAndNamedHandle(Result, TEXT("SamplerYcbcrConversionKHR"), *YcbcrConversion);
+	}
+	FlushDebugWrapperLog();
+}
+
+void FWrapLayer::DestroySamplerYcbcrConversionKHR(VkResult Result, VkDevice Device, VkSamplerYcbcrConversion YcbcrConversion)
+{
+	if (Result == VK_RESULT_MAX_ENUM)
+	{
+#if VULKAN_ENABLE_DUMP_LAYER
+		DevicePrintfBegin(Device, FString::Printf(TEXT("vkDestroySamplerYcbcrConversionKHR(YcbcrConversion=0x%p)"), YcbcrConversion));
+#endif
+	}
+}
+#endif
+
 #if VULKAN_ENABLE_IMAGE_TRACKING_LAYER
 namespace VulkanRHI
 {
diff --git a/Engine/Source/Runtime/VulkanRHI/Private/VulkanDevice.cpp b/Engine/Source/Runtime/VulkanRHI/Private/VulkanDevice.cpp
--- a/Engine/Source/Runtime/VulkanRHI/Private/VulkanDevice.cpp
+++ b/Engine/Source/Runtime/VulkanRHI/Private/VulkanDevice.cpp
@@ -629,7 +629,6 @@ void FVulkanDevice::SetupFormats()
 #if VULKAN_SUPPORTS_COLOR_CONVERSIONS
 VkSamplerYcbcrConversion FVulkanDevice::CreateSamplerColorConversion(const VkSamplerYcbcrConversionCreateInfo& CreateInfo)
 {
-
 	const uint32 CreateInfoHash = FCrc::MemCrc32(&CreateInfo, sizeof(CreateInfo));
 	VkSamplerYcbcrConversion* const FindResult = SamplerColorConversionMap.Find(CreateInfoHash);
 	if (FindResult != nullptr)
@@ -639,7 +638,7 @@ VkSamplerYcbcrConversion FVulkanDevice::CreateSamplerColorConversion(const VkSam
 	else
 	{
 		VkSamplerYcbcrConversion NewConversion;
-		VERIFYVULKANRESULT(vkCreateSamplerYcbcrConversionKHR(GetInstanceHandle(), &CreateInfo, nullptr, &NewConversion));
+		VERIFYVULKANRESULT(VulkanRHI::vkCreateSamplerYcbcrConversionKHR(GetInstanceHandle(), &CreateInfo, nullptr, &NewConversion));
 		SamplerColorConversionMap.Add(CreateInfoHash, NewConversion);
 		return NewConversion;
 	}
@@ -874,7 +873,7 @@ void FVulkanDevice::Destroy()
 #if VULKAN_SUPPORTS_COLOR_CONVERSIONS
 	for (const auto& Pair : SamplerColorConversionMap)
 	{
-		vkDestroySamplerYcbcrConversionKHR(GetInstanceHandle(), Pair.Value, nullptr);
+		VulkanRHI::vkDestroySamplerYcbcrConversionKHR(GetInstanceHandle(), Pair.Value, nullptr);
 	}
 	SamplerColorConversionMap.Reset();
 #endif
diff --git a/Engine/Source/Runtime/VulkanRHI/Private/VulkanDevice.h b/Engine/Source/Runtime/VulkanRHI/Private/VulkanDevice.h
--- a/Engine/Source/Runtime/VulkanRHI/Private/VulkanDevice.h
+++ b/Engine/Source/Runtime/VulkanRHI/Private/VulkanDevice.h
@@ -28,6 +28,7 @@ struct FOptionalVulkanDeviceExtensions
 	uint32 HasEXTValidationCache : 1;
 	uint32 HasAMDBufferMarker : 1;
 	uint32 HasGoogleDisplayTiming : 1;
+	uint32 HasYcbcrSampler : 1;
 };
 
 class FVulkanDevice
@@ -357,7 +358,9 @@ private:
 	FVulkanCommandListContextImmediate* ImmediateContext;
 	FVulkanCommandListContext* ComputeContext;
 	TArray<FVulkanCommandListContext*> CommandContexts;
+#if VULKAN_SUPPORTS_COLOR_CONVERSIONS
 	TMap<uint32, VkSamplerYcbcrConversion> SamplerColorConversionMap;
+#endif
 
 	void GetDeviceExtensionsAndLayers(TArray<const ANSICHAR*>& OutDeviceExtensions, TArray<const ANSICHAR*>& OutDeviceLayers, bool& bOutDebugMarkers);
 
diff --git a/Engine/Source/Runtime/VulkanRHI/Private/VulkanLayers.cpp b/Engine/Source/Runtime/VulkanRHI/Private/VulkanLayers.cpp
--- a/Engine/Source/Runtime/VulkanRHI/Private/VulkanLayers.cpp
+++ b/Engine/Source/Runtime/VulkanRHI/Private/VulkanLayers.cpp
@@ -691,4 +691,8 @@ void FVulkanDevice::ParseOptionalDeviceExtensions(const TArray<const ANSICHAR *>
 #if VULKAN_SUPPORTS_GOOGLE_DISPLAY_TIMING
 	OptionalDeviceExtensions.HasGoogleDisplayTiming = HasExtension(VK_GOOGLE_DISPLAY_TIMING_EXTENSION_NAME);
 #endif
+
+#if VULKAN_SUPPORTS_COLOR_CONVERSIONS
+	OptionalDeviceExtensions.HasYcbcrSampler = HasExtension(VK_KHR_SAMPLER_YCBCR_CONVERSION_EXTENSION_NAME);
+#endif
 }
diff --git a/Engine/Source/Runtime/VulkanRHI/Private/VulkanLoader.h b/Engine/Source/Runtime/VulkanRHI/Private/VulkanLoader.h
--- a/Engine/Source/Runtime/VulkanRHI/Private/VulkanLoader.h
+++ b/Engine/Source/Runtime/VulkanRHI/Private/VulkanLoader.h
@@ -170,9 +170,6 @@
 // List all optional Vulkan entry points used by Unreal that need to be loaded manually
 #define ENUM_VK_ENTRYPOINTS_OPTIONAL_INSTANCE(EnumMacro) \
 	EnumMacro(PFN_vkCreateSharedSwapchainsKHR, vkCreateSharedSwapchainsKHR)
-	//EnumMacro(PFN_vkCreateDebugReportCallbackEXT, vkCreateDebugReportCallbackEXT)
-	//EnumMacro(PFN_vkDestroyDebugReportCallbackEXT, vkDestroyDebugReportCallbackEXT)
-	//EnumMacro(PFN_vkCreateDisplayPlaneSurfaceKHR, vkCreateDisplayPlaneSurfaceKHR)
 
 // List of all Vulkan entry points
 #define ENUM_VK_ENTRYPOINTS_ALL(EnumMacro) \
diff --git a/Engine/Source/Runtime/VulkanRHI/Private/Windows/VulkanWindowsPlatform.cpp b/Engine/Source/Runtime/VulkanRHI/Private/Windows/VulkanWindowsPlatform.cpp
--- a/Engine/Source/Runtime/VulkanRHI/Private/Windows/VulkanWindowsPlatform.cpp
+++ b/Engine/Source/Runtime/VulkanRHI/Private/Windows/VulkanWindowsPlatform.cpp
@@ -123,6 +123,10 @@ void FVulkanWindowsPlatform::GetDeviceExtensions(TArray<const ANSICHAR*>& OutExt
 	{
 		OutExtensions.Add(VK_AMD_BUFFER_MARKER_EXTENSION_NAME);
 	}
+
+#if VULKAN_SUPPORTS_COLOR_CONVERSIONS
+	OutExtensions.Add(VK_KHR_SAMPLER_YCBCR_CONVERSION_EXTENSION_NAME);
+#endif
 }
 
 void FVulkanWindowsPlatform::CreateSurface(void* WindowHandle, VkInstance Instance, VkSurfaceKHR* OutSurface)
diff --git a/Engine/Source/Runtime/VulkanRHI/Private/Windows/VulkanWindowsPlatform.h b/Engine/Source/Runtime/VulkanRHI/Private/Windows/VulkanWindowsPlatform.h
--- a/Engine/Source/Runtime/VulkanRHI/Private/Windows/VulkanWindowsPlatform.h
+++ b/Engine/Source/Runtime/VulkanRHI/Private/Windows/VulkanWindowsPlatform.h
@@ -15,7 +15,7 @@
 #define VULKAN_DYNAMICALLYLOADED				1
 #define VULKAN_ENABLE_DESKTOP_HMD_SUPPORT		1
 #define VULKAN_SIGNAL_UNIMPLEMENTED()			checkf(false, TEXT("Unimplemented vulkan functionality: %s"), TEXT(__FUNCTION__))
-
+#define VULKAN_SUPPORTS_COLOR_CONVERSIONS		1
 #define	VULKAN_SUPPORTS_DEDICATED_ALLOCATION	0
 
 // 32-bit windows has warnings on custom mem mgr callbacks
@@ -30,7 +30,9 @@
 	EnumMacro(PFN_vkCmdWriteBufferMarkerAMD, vkCmdWriteBufferMarkerAMD) \
 	EnumMacro(PFN_vkGetBufferMemoryRequirements2KHR , vkGetBufferMemoryRequirements2KHR)
 
-#define ENUM_VK_ENTRYPOINTS_OPTIONAL_PLATFORM_INSTANCE(EnumMacro)
+#define ENUM_VK_ENTRYPOINTS_OPTIONAL_PLATFORM_INSTANCE(EnumMacro) \
+	EnumMacro(PFN_vkCreateSamplerYcbcrConversionKHR, vkCreateSamplerYcbcrConversionKHR) \
+	EnumMacro(PFN_vkDestroySamplerYcbcrConversionKHR, vkDestroySamplerYcbcrConversionKHR)
 
 #include "../VulkanLoader.h"
 
