DR - vk - Wait for end of frame completion event on some platforms
Copy from 4413413
This prevents the CPU from getting more than one frame ahead of the GPU. In addition to making the Vulkan RHI more consistent with the DX RHIs and reducing game input latency, this also fixes the issue where memory for volatile uniform buffers could be overwritten by the CPU before the GPU was done with it. This could happen when running in immediate present mode with frame pacing and smoothing disabled, as the CPU could get several frames ahead of the GPU.

diff --git a/Engine/Source/Runtime/VulkanRHI/Private/Android/VulkanAndroidPlatform.h b/Engine/Source/Runtime/VulkanRHI/Private/Android/VulkanAndroidPlatform.h
--- a/Engine/Source/Runtime/VulkanRHI/Private/Android/VulkanAndroidPlatform.h
+++ b/Engine/Source/Runtime/VulkanRHI/Private/Android/VulkanAndroidPlatform.h
@@ -85,6 +85,8 @@ public:
 	//#todo-rco: Detect Mali? Doing a clear on ColorAtt layout on empty cmd buffer causes issues
 	static bool RequiresSwapchainGeneralInitialLayout() { return true; }
 
+	static bool RequiresWaitingForFrameCompletionEvent() { return false; }
+
 protected:
 	static void* VulkanLib;
 	static bool bAttemptedLoad;
diff --git a/Engine/Source/Runtime/VulkanRHI/Private/Lumin/VulkanLuminPlatform.h b/Engine/Source/Runtime/VulkanRHI/Private/Lumin/VulkanLuminPlatform.h
--- a/Engine/Source/Runtime/VulkanRHI/Private/Lumin/VulkanLuminPlatform.h
+++ b/Engine/Source/Runtime/VulkanRHI/Private/Lumin/VulkanLuminPlatform.h
@@ -64,6 +64,8 @@ public:
 
 	static void EnablePhysicalDeviceFeatureExtensions(VkDeviceCreateInfo& DeviceInfo);
 
+	static bool RequiresWaitingForFrameCompletionEvent() { return false; }
+
 protected:
 	static void* VulkanLib;
 	static bool bAttemptedLoad;
diff --git a/Engine/Source/Runtime/VulkanRHI/Private/VulkanGenericPlatform.h b/Engine/Source/Runtime/VulkanRHI/Private/VulkanGenericPlatform.h
--- a/Engine/Source/Runtime/VulkanRHI/Private/VulkanGenericPlatform.h
+++ b/Engine/Source/Runtime/VulkanRHI/Private/VulkanGenericPlatform.h
@@ -97,4 +97,7 @@ public:
 	
 	// Allow platforms to add extension features to the PresentInfo pNext chain
 	static void EnablePresentInfoExtensions(VkPresentInfoKHR& PresentInfo) {}
+
+	// Ensure the last frame completed on the GPU
+	static bool RequiresWaitingForFrameCompletionEvent() { return true; }
 };
diff --git a/Engine/Source/Runtime/VulkanRHI/Private/VulkanViewport.cpp b/Engine/Source/Runtime/VulkanRHI/Private/VulkanViewport.cpp
--- a/Engine/Source/Runtime/VulkanRHI/Private/VulkanViewport.cpp
+++ b/Engine/Source/Runtime/VulkanRHI/Private/VulkanViewport.cpp
@@ -245,9 +245,38 @@ void FVulkanViewport::AdvanceBackBufferFrame()
 	}
 }
 
-//void FVulkanViewport::WaitForFrameEventCompletion()
-//{
-//}
+void FVulkanViewport::WaitForFrameEventCompletion()
+{
+	if (FVulkanPlatform::RequiresWaitingForFrameCompletionEvent())
+	{
+		static FCriticalSection CS;
+		FScopeLock ScopeLock(&CS);
+		if (LastFrameCommandBuffer && LastFrameCommandBuffer->IsSubmitted())
+		{
+			// If last frame's fence hasn't been signaled already, wait for it here
+			if (LastFrameFenceCounter == LastFrameCommandBuffer->GetFenceSignaledCounter())
+			{
+				if (!GWaitForIdleOnSubmit)
+				{
+					// The wait has already happened if GWaitForIdleOnSubmit is set
+					LastFrameCommandBuffer->GetOwner()->GetMgr().WaitForCmdBuffer(LastFrameCommandBuffer);
+				}
+			}
+		}
+	}
+}
+
+void FVulkanViewport::IssueFrameEvent()
+{
+	if (FVulkanPlatform::RequiresWaitingForFrameCompletionEvent())
+	{
+		// The fence we need to wait on next frame is already there in the command buffer
+		// that was just submitted in this frame's Present. Just grab that command buffer's
+		// info to use next frame in WaitForFrameEventCompletion.
+		FVulkanQueue* Queue = Device->GetGraphicsQueue();
+		Queue->GetLastSubmittedInfo(LastFrameCommandBuffer, LastFrameFenceCounter);
+	}
+}
 
 
 FVulkanFramebuffer::FVulkanFramebuffer(FVulkanDevice& Device, const FRHISetRenderTargetsInfo& InRTInfo, const FVulkanRenderTargetLayout& RTLayout, const FVulkanRenderPass& RenderPass)
@@ -696,19 +725,12 @@ bool FVulkanViewport::Present(FVulkanCommandListContext* Context, FVulkanCmdBuff
 		RHIBackBuffer = nullptr;
 	}
 
-	//static const TConsoleVariableData<int32>* CFinishFrameVar = IConsoleManager::Get().FindTConsoleVariableDataInt(TEXT("r.FinishCurrentFrame"));
-	//if (!CFinishFrameVar->GetValueOnRenderThread())
-	//{
-	//	// Wait for the GPU to finish rendering the previous frame before finishing this frame.
-	//	WaitForFrameEventCompletion();
-	//	IssueFrameEvent();
-	//}
-	//else
-	//{
-	//	// Finish current frame immediately to reduce latency
-	//	IssueFrameEvent();
-	//	WaitForFrameEventCompletion();
-	//}
+	if (FVulkanPlatform::RequiresWaitingForFrameCompletionEvent() && !bHasCustomPresent)
+	{
+		// Wait for the GPU to finish rendering the previous frame before finishing this frame.
+		WaitForFrameEventCompletion();
+		IssueFrameEvent();
+	}
 
 	// If the input latency timer has been triggered, block until the GPU is completely
 	// finished displaying this frame and calculate the delta time.
diff --git a/Engine/Source/Runtime/VulkanRHI/Public/VulkanViewport.h b/Engine/Source/Runtime/VulkanRHI/Public/VulkanViewport.h
--- a/Engine/Source/Runtime/VulkanRHI/Public/VulkanViewport.h
+++ b/Engine/Source/Runtime/VulkanRHI/Public/VulkanViewport.h
@@ -27,10 +27,9 @@ public:
 
 	FVulkanTexture2D* GetBackBuffer(FRHICommandList& RHICmdList);
 
-	//void WaitForFrameEventCompletion();
+	void WaitForFrameEventCompletion();
 
-	//#todo-rco
-	//void IssueFrameEvent() {}
+	void IssueFrameEvent();
 
 	inline FIntPoint GetSizeXY() const
 	{
@@ -84,6 +83,9 @@ protected:
 
 	FCustomPresentRHIRef CustomPresent;
 
+	FVulkanCmdBuffer* LastFrameCommandBuffer = nullptr;
+	uint64 LastFrameFenceCounter = 0;
+
 	void CreateSwapchain();
 	void AcquireBackBuffer(FRHICommandListBase& CmdList, FVulkanBackBuffer* NewBackBuffer);
 
