DR - vk - r.Vulkan.ForceCoherentOperations to force memory flush and invalidate
#rnx

diff --git a/Engine/Source/Runtime/VulkanRHI/Private/VulkanMemory.cpp b/Engine/Source/Runtime/VulkanRHI/Private/VulkanMemory.cpp
--- a/Engine/Source/Runtime/VulkanRHI/Private/VulkanMemory.cpp
+++ b/Engine/Source/Runtime/VulkanRHI/Private/VulkanMemory.cpp
@@ -27,6 +27,17 @@ static void CaptureCallStack(FString& OutCallstack)
 }
 #endif
 
+#if (UE_BUILD_DEBUG || UE_BUILD_DEVELOPMENT)
+static int32 GForceCoherent = 0;
+static FAutoConsoleVariableRef CVarForceCoherentOperations(
+	TEXT("r.Vulkan.ForceCoherentOperations"),
+	GForceCoherent,
+	TEXT("1 forces memory invalidation and flushing of coherent memory\n"),
+	ECVF_ReadOnly
+);
+#else
+constexpr int32 GForceCoherent = 0;
+#endif
 
 namespace VulkanRHI
 {
@@ -364,7 +375,7 @@ namespace VulkanRHI
 
 	void FDeviceMemoryAllocation::FlushMappedMemory(VkDeviceSize InOffset, VkDeviceSize InSize)
 	{
-		if (!IsCoherent())
+		if (!IsCoherent() || GForceCoherent != 0)
 		{
 			check(IsMapped());
 			check(InOffset + InSize <= Size);
@@ -379,7 +390,7 @@ namespace VulkanRHI
 
 	void FDeviceMemoryAllocation::InvalidateMappedMemory(VkDeviceSize InOffset, VkDeviceSize InSize)
 	{
-		if (!IsCoherent())
+		if (!IsCoherent() || GForceCoherent != 0)
 		{
 			check(IsMapped());
 			check(InOffset + InSize <= Size);
