DR - vk - Temporarily disable rhi thread on Linux NV cards
#rnx

diff --git a/Engine/Source/Runtime/VulkanRHI/Private/Linux/VulkanLinuxPlatform.h b/Engine/Source/Runtime/VulkanRHI/Private/Linux/VulkanLinuxPlatform.h
--- a/Engine/Source/Runtime/VulkanRHI/Private/Linux/VulkanLinuxPlatform.h
+++ b/Engine/Source/Runtime/VulkanRHI/Private/Linux/VulkanLinuxPlatform.h
@@ -52,6 +52,12 @@ public:
 #endif
 	}
 
+	static bool SupportsRHIThread()
+	{
+		//#todo-rco: Workaround for suspected driver bug
+		return !IsRHIDeviceNVIDIA();
+	}
+
 protected:
 	static void* VulkanLib;
 	static bool bAttemptedLoad;
diff --git a/Engine/Source/Runtime/VulkanRHI/Private/VulkanGenericPlatform.h b/Engine/Source/Runtime/VulkanRHI/Private/VulkanGenericPlatform.h
--- a/Engine/Source/Runtime/VulkanRHI/Private/VulkanGenericPlatform.h
+++ b/Engine/Source/Runtime/VulkanRHI/Private/VulkanGenericPlatform.h
@@ -100,4 +100,6 @@ public:
 
 	// Ensure the last frame completed on the GPU
 	static bool RequiresWaitingForFrameCompletionEvent() { return true; }
+
+	static bool SupportsRHIThread() { return true; }
 };
diff --git a/Engine/Source/Runtime/VulkanRHI/Private/VulkanRHI.cpp b/Engine/Source/Runtime/VulkanRHI/Private/VulkanRHI.cpp
--- a/Engine/Source/Runtime/VulkanRHI/Private/VulkanRHI.cpp
+++ b/Engine/Source/Runtime/VulkanRHI/Private/VulkanRHI.cpp
@@ -594,14 +594,17 @@ void FVulkanDynamicRHI::InitInstance()
 		GRHISupportsTextureStreaming = true;
 		GSupportsTimestampRenderQueries = FVulkanPlatform::SupportsTimestampRenderQueries();
 		GRHISupportsGPUFence = true;
+
+		bool bSupportsRHIThread = FVulkanPlatform::SupportsRHIThread();
 #if VULKAN_ENABLE_DUMP_LAYER
 		// Disable RHI thread by default if the dump layer is enabled
 		GRHISupportsRHIThread = false;
 		GRHISupportsParallelRHIExecute = false;
 #else
-		GRHISupportsRHIThread = GRHIThreadCvar->GetInt() != 0;
+		GRHISupportsRHIThread = GRHIThreadCvar->GetInt() != 0 && bSupportsRHIThread;
 		GRHISupportsParallelRHIExecute = GRHIThreadCvar->GetInt() > 1;
 #endif
+
 		// Some platforms might only have CPU for an RHI thread, but not for parallel tasks
 		GSupportsParallelRenderingTasksWithSeparateRHIThread = GRHISupportsRHIThread ? FVulkanPlatform::SupportParallelRenderingTasks() : false;
 
