DR - vk - Disable the frame pacing when vsync is disabled
- Vulkan app registration

diff --git a/Engine/Source/Runtime/VulkanRHI/Private/VulkanRHI.cpp b/Engine/Source/Runtime/VulkanRHI/Private/VulkanRHI.cpp
--- a/Engine/Source/Runtime/VulkanRHI/Private/VulkanRHI.cpp
+++ b/Engine/Source/Runtime/VulkanRHI/Private/VulkanRHI.cpp
@@ -16,6 +16,7 @@
 #include "VulkanPipelineState.h"
 #include "Misc/FileHelper.h"
 #include "VulkanLLM.h"
+#include "Misc/EngineVersion.h"
 #include "GlobalShader.h"
 
 extern RHI_API bool GUseTexture3DBulkDataRHI;
@@ -300,12 +301,19 @@ void FVulkanDynamicRHI::CreateInstance()
 	auto* CVarDisableEngineAndAppRegistration = IConsoleManager::Get().FindTConsoleVariableDataInt(TEXT("r.DisableEngineAndAppRegistration"));
 	bool bDisableEngineRegistration = (CVarDisableEngineAndAppRegistration && CVarDisableEngineAndAppRegistration->GetValueOnAnyThread() != 0) ||
 		(CVarShaderDevelopmentMode && CVarShaderDevelopmentMode->GetValueOnAnyThread() != 0);
+
+	// EngineName will be of the form "UnrealEngine4.21", with the minor version ("21" in this example)
+	// updated with every quarterly release
+	FString EngineName = FApp::GetEpicProductIdentifier() + FEngineVersion::Current().ToString(EVersionComponent::Minor);
+	FTCHARToUTF8 EngineNameConverter(*EngineName);
+	FTCHARToUTF8 ProjectNameConverter(FApp::GetProjectName());
+
 	VkApplicationInfo AppInfo;
 	ZeroVulkanStruct(AppInfo, VK_STRUCTURE_TYPE_APPLICATION_INFO);
-	AppInfo.pApplicationName = bDisableEngineRegistration ? "" : "UE4";
-	//AppInfo.applicationVersion = 0;
-	AppInfo.pEngineName = bDisableEngineRegistration ? "" : "UE4";
-	AppInfo.engineVersion = 15;
+	AppInfo.pApplicationName = bDisableEngineRegistration ? nullptr : ProjectNameConverter.Get();
+	AppInfo.applicationVersion = 0;	// Do we want FApp::GetBuildVersion() ?
+	AppInfo.pEngineName = bDisableEngineRegistration ? nullptr : EngineNameConverter.Get();
+	AppInfo.engineVersion = FEngineVersion::Current().GetMinor();
 	AppInfo.apiVersion = UE_VK_API_VERSION;
 
 	VkInstanceCreateInfo InstInfo;
diff --git a/Engine/Source/Runtime/VulkanRHI/Private/VulkanSwapChain.cpp b/Engine/Source/Runtime/VulkanRHI/Private/VulkanSwapChain.cpp
--- a/Engine/Source/Runtime/VulkanRHI/Private/VulkanSwapChain.cpp
+++ b/Engine/Source/Runtime/VulkanRHI/Private/VulkanSwapChain.cpp
@@ -334,13 +334,13 @@ FVulkanSwapChain::FVulkanSwapChain(VkInstance InInstance, FVulkanDevice& InDevic
 			// Until FVulkanViewport::Present honors SyncInterval, we need to disable vsync for the spectator window if using an HMD.
 			const bool bDisableVsyncForHMD = (FVulkanDynamicRHI::HMDVulkanExtensions.IsValid()) ? FVulkanDynamicRHI::HMDVulkanExtensions->ShouldDisableVulkanVSync() : false;
 
-			if (bFoundPresentModeMailbox && LockToVsync)
+			if (bFoundPresentModeImmediate && (bDisableVsyncForHMD || !LockToVsync))
 			{
-				PresentMode = VK_PRESENT_MODE_MAILBOX_KHR;
+				PresentMode = VK_PRESENT_MODE_IMMEDIATE_KHR;
 			}
-			else if (bFoundPresentModeImmediate && (bDisableVsyncForHMD || !LockToVsync))
+			else if (bFoundPresentModeMailbox)
 			{
-				PresentMode = VK_PRESENT_MODE_IMMEDIATE_KHR;
+				PresentMode = VK_PRESENT_MODE_MAILBOX_KHR;
 			}
 			else if (bFoundPresentModeFIFO)
 			{
@@ -582,10 +582,10 @@ int32 FVulkanSwapChain::AcquireImageIndex(VulkanRHI::FSemaphore** OutSemaphore)
 void FVulkanSwapChain::RenderThreadPacing()
 {
 	check(IsInRenderingThread());
-	const int32 SyncInterval = RHIGetSyncInterval();
+	const int32 SyncInterval = LockToVsync ? RHIGetSyncInterval() : 0;
 
 	//very naive CPU side frame pacer.
-	if (GVulkanCPURenderThreadFramePacer)
+	if (GVulkanCPURenderThreadFramePacer && SyncInterval > 0)
 	{
 		static double PreviousFrameCPUTime = 0;
 		static double SampledDeltaTimeMS = 0;
@@ -658,7 +658,7 @@ FVulkanSwapChain::EStatus FVulkanSwapChain::Present(FVulkanQueue* GfxQueue, FVul
 	Info.pSwapchains = &SwapChain;
 	Info.pImageIndices = (uint32*)&CurrentImageIndex;
 
-	const int32 SyncInterval = RHIGetSyncInterval();
+	const int32 SyncInterval = LockToVsync ? RHIGetSyncInterval() : 0;
 	ensureMsgf(SyncInterval <= 3 && SyncInterval >= 0, TEXT("Unsupported sync interval: %i"), SyncInterval);
 	FVulkanPlatform::EnablePresentInfoExtensions(Info);
 
@@ -848,7 +848,7 @@ FVulkanSwapChain::EStatus FVulkanSwapChain::Present(FVulkanQueue* GfxQueue, FVul
 #endif
 
 	//very naive CPU side frame pacer.
-	if (GVulkanCPURHIFramePacer)
+	if (GVulkanCPURHIFramePacer && SyncInterval > 0)
 	{
 		static double PreviousFrameCPUTime = 0;
 		double NowCPUTime = FPlatformTime::Seconds();
