DR - vk - Fix validation warning

diff --git a/Engine/Source/Runtime/VulkanRHI/Private/Lumin/VulkanLuminPlatform.cpp b/Engine/Source/Runtime/VulkanRHI/Private/Lumin/VulkanLuminPlatform.cpp
--- a/Engine/Source/Runtime/VulkanRHI/Private/Lumin/VulkanLuminPlatform.cpp
+++ b/Engine/Source/Runtime/VulkanRHI/Private/Lumin/VulkanLuminPlatform.cpp
@@ -130,6 +130,9 @@ void FVulkanLuminPlatform::GetDeviceExtensions(TArray<const ANSICHAR*>& OutExten
 	{
 		OutExtensions.Add("VK_LAYER_NV_vgd");
 	}
+	// YCbCr requires BindMem2 and GetMemReqs2
+	OutExtensions.Add(VK_KHR_BIND_MEMORY_2_EXTENSION_NAME);
+	OutExtensions.Add(VK_KHR_GET_MEMORY_REQUIREMENTS_2_EXTENSION_NAME);
 	OutExtensions.Add(VK_KHR_SAMPLER_YCBCR_CONVERSION_EXTENSION_NAME);
 }
 
diff --git a/Engine/Source/Runtime/VulkanRHI/Private/VulkanLayers.cpp b/Engine/Source/Runtime/VulkanRHI/Private/VulkanLayers.cpp
--- a/Engine/Source/Runtime/VulkanRHI/Private/VulkanLayers.cpp
+++ b/Engine/Source/Runtime/VulkanRHI/Private/VulkanLayers.cpp
@@ -693,6 +693,6 @@ void FVulkanDevice::ParseOptionalDeviceExtensions(const TArray<const ANSICHAR *>
 #endif
 
 #if VULKAN_SUPPORTS_COLOR_CONVERSIONS
-	OptionalDeviceExtensions.HasYcbcrSampler = HasExtension(VK_KHR_SAMPLER_YCBCR_CONVERSION_EXTENSION_NAME);
+	OptionalDeviceExtensions.HasYcbcrSampler = HasExtension(VK_KHR_SAMPLER_YCBCR_CONVERSION_EXTENSION_NAME) && HasExtension(VK_KHR_BIND_MEMORY_2_EXTENSION_NAME) && HasExtension(VK_KHR_GET_MEMORY_REQUIREMENTS_2_EXTENSION_NAME);
 #endif
 }
diff --git a/Engine/Source/Runtime/VulkanRHI/Private/Windows/VulkanWindowsPlatform.cpp b/Engine/Source/Runtime/VulkanRHI/Private/Windows/VulkanWindowsPlatform.cpp
--- a/Engine/Source/Runtime/VulkanRHI/Private/Windows/VulkanWindowsPlatform.cpp
+++ b/Engine/Source/Runtime/VulkanRHI/Private/Windows/VulkanWindowsPlatform.cpp
@@ -125,6 +125,9 @@ void FVulkanWindowsPlatform::GetDeviceExtensions(TArray<const ANSICHAR*>& OutExt
 	}
 
 #if VULKAN_SUPPORTS_COLOR_CONVERSIONS
+	// YCbCr requires BindMem2 and GetMemReqs2
+	OutExtensions.Add(VK_KHR_BIND_MEMORY_2_EXTENSION_NAME);
+	OutExtensions.Add(VK_KHR_GET_MEMORY_REQUIREMENTS_2_EXTENSION_NAME);
 	OutExtensions.Add(VK_KHR_SAMPLER_YCBCR_CONVERSION_EXTENSION_NAME);
 #endif
 }
