DR - vk - Add back supports  fence (remove by mistake)
#rnx

diff --git a/Engine/Source/Runtime/VulkanRHI/Private/VulkanRHI.cpp b/Engine/Source/Runtime/VulkanRHI/Private/VulkanRHI.cpp
--- a/Engine/Source/Runtime/VulkanRHI/Private/VulkanRHI.cpp
+++ b/Engine/Source/Runtime/VulkanRHI/Private/VulkanRHI.cpp
@@ -601,6 +601,7 @@ void FVulkanDynamicRHI::InitInstance()
 		GSupportsRenderTargetFormat_PF_G8 = false;	// #todo-rco
 		GRHISupportsTextureStreaming = true;
 		GSupportsTimestampRenderQueries = FVulkanPlatform::SupportsTimestampRenderQueries();
+		GRHISupportsGPUFence = true;
 #if VULKAN_ENABLE_DUMP_LAYER
 		// Disable RHI thread by default if the dump layer is enabled
 		GRHISupportsRHIThread = false;
