DR - initial VulkanTimestampQuery implemetation: Exposing RHI query pools which can be extended further. They are only used by vk timestamp queries for now because the latent timers for translucency were hoarding the queries without reading them in editor.

diff --git a/Engine/Source/Runtime/Engine/Private/SceneUtils.cpp b/Engine/Source/Runtime/Engine/Private/SceneUtils.cpp
--- a/Engine/Source/Runtime/Engine/Private/SceneUtils.cpp
+++ b/Engine/Source/Runtime/Engine/Private/SceneUtils.cpp
@@ -175,7 +175,7 @@ public:
 	static const uint64 InvalidQueryResult = 0xFFFFFFFFFFFFFFFFull;
 
 public:
-	FRealtimeGPUProfilerEvent(const FName& InName, const FName& InStatName, FRenderQueryPool* RenderQueryPool)
+	FRealtimeGPUProfilerEvent(const FName& InName, const FName& InStatName, FRenderQueryPoolRHIParamRef RenderQueryPool, uint32& QueryCount)
 		: StartResultMicroseconds(InvalidQueryResult)
 		, EndResultMicroseconds(InvalidQueryResult)
 		, FrameNumber(-1)
@@ -189,10 +189,11 @@ public:
 		Name = InName;
 
 		const int MaxGPUQueries = CVarGPUStatsMaxQueriesPerFrame.GetValueOnRenderThread();
-		if ( MaxGPUQueries == -1 || RenderQueryPool->GetAllocatedQueryCount() < MaxGPUQueries )
+		if ( MaxGPUQueries == -1 || QueryCount < uint32(MaxGPUQueries) )
 		{
 			StartQuery = RenderQueryPool->AllocateQuery();
 			EndQuery = RenderQueryPool->AllocateQuery();
+			QueryCount += 2;
 		}
 	}
 
@@ -201,7 +202,7 @@ public:
 		return IsValidRef(StartQuery); 
 	}
 
-	void ReleaseQueries(FRenderQueryPool* RenderQueryPool, FRHICommandListImmediate* RHICmdListPtr)
+	void ReleaseQueries(FRenderQueryPoolRHIParamRef RenderQueryPool, FRHICommandListImmediate* RHICmdListPtr, uint32& QueryCount)
 	{
 		if ( HasQueriesAllocated() )
 		{
@@ -221,6 +222,7 @@ public:
 			}
 			RenderQueryPool->ReleaseQuery(StartQuery);
 			RenderQueryPool->ReleaseQuery(EndQuery);
+			QueryCount -= 2;
 		}
 	}
 
@@ -348,9 +350,11 @@ Container for a single frame's GPU stats
 -----------------------------------------------------------------------------*/
 class FRealtimeGPUProfilerFrame
 {
+	uint32& QueryCount;
 public:
-	FRealtimeGPUProfilerFrame(FRenderQueryPool* InRenderQueryPool)
-		: FrameNumber(-1)
+	FRealtimeGPUProfilerFrame(FRenderQueryPoolRHIRef InRenderQueryPool, uint32& InQueryCount)
+		: QueryCount(InQueryCount)
+		, FrameNumber(-1)
 		, RenderQueryPool(InRenderQueryPool)
 	{}
 
@@ -362,7 +366,7 @@ public:
 	void PushEvent(FRHICommandListImmediate& RHICmdList, const FName& Name, const FName& StatName)
 	{
 		// TODO: this should really use a pool / free list
-		FRealtimeGPUProfilerEvent* Event = new FRealtimeGPUProfilerEvent(Name, StatName, RenderQueryPool);
+		FRealtimeGPUProfilerEvent* Event = new FRealtimeGPUProfilerEvent(Name, StatName, RenderQueryPool, QueryCount);
 		const int32 EventIndex = GpuProfilerEvents.Num();
 
 		GpuProfilerEvents.Add(Event);
@@ -397,7 +401,7 @@ public:
 		{
 			if (GpuProfilerEvents[Index])
 			{
-				GpuProfilerEvents[Index]->ReleaseQueries(RenderQueryPool, RHICommandListPtr);
+				GpuProfilerEvents[Index]->ReleaseQueries(RenderQueryPool, RHICommandListPtr, QueryCount);
 				delete GpuProfilerEvents[Index];
 			}
 		}
@@ -564,7 +568,7 @@ private:
 	TArray<FGPUEventTimeAggregate> EventAggregates;
 
 	uint32 FrameNumber;
-	FRenderQueryPool* RenderQueryPool;
+	FRenderQueryPoolRHIRef RenderQueryPool;
 };
 
 /*-----------------------------------------------------------------------------
@@ -588,10 +592,11 @@ FRealtimeGPUProfiler::FRealtimeGPUProfiler()
 	, bStatGatheringPaused(false)
 	, bInBeginEndBlock(false)
 {
-	RenderQueryPool = new FRenderQueryPool(RQT_AbsoluteTime);
+	const int MaxGPUQueries = CVarGPUStatsMaxQueriesPerFrame.GetValueOnRenderThread();
+	RenderQueryPool = RHICreateRenderQueryPool(RQT_AbsoluteTime, (MaxGPUQueries > 0) ? MaxGPUQueries * 2 : 10240);
 	for (int Index = 0; Index < NumGPUProfilerBufferedFrames; Index++)
 	{
-		Frames.Add(new FRealtimeGPUProfilerFrame(RenderQueryPool));
+		Frames.Add(new FRealtimeGPUProfilerFrame(RenderQueryPool, QueryCount));
 	}
 }
 
@@ -602,8 +607,7 @@ void FRealtimeGPUProfiler::Release()
 		delete Frames[Index];
 	}
 	Frames.Empty();
-	delete RenderQueryPool;
-	RenderQueryPool = nullptr;
+	RenderQueryPool.SafeRelease();
 }
 
 void FRealtimeGPUProfiler::BeginFrame(FRHICommandListImmediate& RHICmdList)
diff --git a/Engine/Source/Runtime/Engine/Public/SceneUtils.h b/Engine/Source/Runtime/Engine/Public/SceneUtils.h
--- a/Engine/Source/Runtime/Engine/Public/SceneUtils.h
+++ b/Engine/Source/Runtime/Engine/Public/SceneUtils.h
@@ -227,7 +227,8 @@ private:
 	int32 WriteBufferIndex;
 	int32 ReadBufferIndex;
 	uint32 WriteFrameNumber;
-	FRenderQueryPool* RenderQueryPool;
+	uint32 QueryCount = 0;
+	FRenderQueryPoolRHIRef RenderQueryPool;
 	bool bStatGatheringPaused;
 	bool bInBeginEndBlock;
 };
diff --git a/Engine/Source/Runtime/RHI/Private/DynamicRHI.cpp b/Engine/Source/Runtime/RHI/Private/DynamicRHI.cpp
--- a/Engine/Source/Runtime/RHI/Private/DynamicRHI.cpp
+++ b/Engine/Source/Runtime/RHI/Private/DynamicRHI.cpp
@@ -314,3 +314,26 @@ void FDynamicRHI::EnableIdealGPUCaptureOptions(bool bEnabled)
 		RHICmdBypassVar->Set(bShouldRHICmdBypass ? 1 : 0, ECVF_SetByConsole);		
 	}	
 }
+
+TRefCountPtr<FRHIRenderQuery> FDefaultRHIRenderQueryPool::AllocateQuery()
+{
+	if (Queries.Num() > 0)
+	{
+		return Queries.Pop();
+	}
+	else
+	{
+		ensure(++AllocatedQueries <= NumQueries);
+		return DynamicRHI->RHICreateRenderQuery(QueryType);
+	}
+}
+
+void FDefaultRHIRenderQueryPool::ReleaseQuery(TRefCountPtr<FRHIRenderQuery> &Query)
+{
+	Queries.Push(Query);
+}
+
+FRenderQueryPoolRHIRef RHICreateRenderQueryPool(ERenderQueryType QueryType, uint32 NumQueries)
+{
+	return GDynamicRHI->RHICreateRenderQueryPool(QueryType, NumQueries);
+}
\ No newline at end of file
diff --git a/Engine/Source/Runtime/RHI/Public/DynamicRHI.h b/Engine/Source/Runtime/RHI/Public/DynamicRHI.h
--- a/Engine/Source/Runtime/RHI/Public/DynamicRHI.h
+++ b/Engine/Source/Runtime/RHI/Public/DynamicRHI.h
@@ -75,6 +75,28 @@ struct FRHIFlipDetails
 	{}
 };
 
+class FDynamicRHI;
+
+class RHI_API FDefaultRHIRenderQueryPool final : public FRHIRenderQueryPool
+{
+public:
+	FDefaultRHIRenderQueryPool(ERenderQueryType InQueryType, FDynamicRHI* InDynamicRHI, uint32 InNumQueries)
+		: DynamicRHI(InDynamicRHI)
+		, QueryType(InQueryType)
+		, NumQueries(InNumQueries)
+	{ }
+
+private:
+	virtual TRefCountPtr<FRHIRenderQuery> AllocateQuery() override;
+	virtual void ReleaseQuery(TRefCountPtr<FRHIRenderQuery> &Query) override;
+
+	FDynamicRHI* DynamicRHI = nullptr;
+	ERenderQueryType QueryType;
+	uint32 NumQueries = 0;
+	uint32 AllocatedQueries = 0;
+	TArray<TRefCountPtr<FRHIRenderQuery>> Queries;
+};
+
 /** The interface which is implemented by the dynamically bound RHI. */
 class RHI_API FDynamicRHI
 {
@@ -198,6 +220,17 @@ public:
 		return nullptr;
 	}
 	
+	/**
+	* Creates a pool for querys like timers or occlusion queries.
+	* @param QueryType The ype of the queries provided by this pool like RQT_Occlusion or RQT_AbsoluteTime.
+	* @return the Querypool.
+	*/
+	// FlushType: Must be Thread-Safe.
+	virtual FRenderQueryPoolRHIRef RHICreateRenderQueryPool(ERenderQueryType QueryType, uint32 NumQueries)
+	{
+		return new FDefaultRHIRenderQueryPool(QueryType, this, NumQueries);
+	}
+
 	/**
 	* Creates a compute fence.  Compute fences are named GPU fences which can be written to once before resetting.
 	* A command to write the fence must be enqueued before any commands to wait on them.  This is enforced on the CPU to avoid GPU hangs.
@@ -1236,6 +1269,8 @@ FORCEINLINE class IRHICommandContextContainer* RHIGetCommandContextContainer(int
 	return GDynamicRHI->RHIGetCommandContextContainer(Index, Num, GPUMask);
 }
 
+RHI_API FRenderQueryPoolRHIRef RHICreateRenderQueryPool(ERenderQueryType QueryType, uint32 NumQueries);
+
 /**
 * Defragment the texture pool.
 */
diff --git a/Engine/Source/Runtime/RHI/Public/RHIResources.h b/Engine/Source/Runtime/RHI/Public/RHIResources.h
--- a/Engine/Source/Runtime/RHI/Public/RHIResources.h
+++ b/Engine/Source/Runtime/RHI/Public/RHIResources.h
@@ -831,6 +831,13 @@ private:
 
 class FRHIRenderQuery : public FRHIResource {};
 
+class FRHIRenderQueryPool : public FRHIResource 
+{
+public:
+	virtual TRefCountPtr<FRHIRenderQuery> AllocateQuery() = 0;
+	virtual void ReleaseQuery(TRefCountPtr<FRHIRenderQuery> &Query) = 0;
+};
+
 class FRHIComputeFence : public FRHIResource
 {
 public:
@@ -1002,6 +1009,9 @@ typedef TRefCountPtr<FRHITextureReference> FTextureReferenceRHIRef;
 typedef FRHIRenderQuery*              FRenderQueryRHIParamRef;
 typedef TRefCountPtr<FRHIRenderQuery> FRenderQueryRHIRef;
 
+typedef FRHIRenderQueryPool*              FRenderQueryPoolRHIParamRef;
+typedef TRefCountPtr<FRHIRenderQueryPool> FRenderQueryPoolRHIRef;
+
 typedef FRHIGPUFence*				FGPUFenceRHIParamRef;
 typedef TRefCountPtr<FRHIGPUFence>	FGPUFenceRHIRef;
 
diff --git a/Engine/Source/Runtime/Renderer/Private/GPUBenchmark.cpp b/Engine/Source/Runtime/Renderer/Private/GPUBenchmark.cpp
--- a/Engine/Source/Runtime/Renderer/Private/GPUBenchmark.cpp
+++ b/Engine/Source/Runtime/Renderer/Private/GPUBenchmark.cpp
@@ -411,8 +411,6 @@ void RendererGPUBenchmark(FRHICommandListImmediate& RHICmdList, FSynthBenchmarkR
 {
 	check(IsInRenderingThread());
 
-	FRenderQueryPool TimerQueryPool(RQT_AbsoluteTime);
-
 	bool bValidGPUTimer = (FGPUTiming::GetTimingFrequency() / (1000 * 1000)) != 0;
 
 	if(!bValidGPUTimer)
@@ -497,13 +495,14 @@ void RendererGPUBenchmark(FRHICommandListImmediate& RHICmdList, FSynthBenchmarkR
 		uint32 DestRTIndex = 0;
 
 		const uint32 TimerSampleCount = IterationCount * MethodCount + 1;
+		FRenderQueryPoolRHIRef TimerQueryPool = RHICreateRenderQueryPool(RQT_AbsoluteTime, TimerSampleCount * 2);
 
 		static FRenderQueryRHIRef TimerQueries[TimerSampleCount];
 		static float LocalWorkScale[IterationCount];
 
 		for(uint32  i = 0; i < TimerSampleCount; ++i)
 		{
-			TimerQueries[i] = TimerQueryPool.AllocateQuery();
+			TimerQueries[i] = TimerQueryPool->AllocateQuery();
 		}
 
 		const bool bSupportsTimerQueries = (TimerQueries[0] != NULL);
@@ -617,7 +616,7 @@ void RendererGPUBenchmark(FRHICommandListImmediate& RHICmdList, FSynthBenchmarkR
 			// flushes the RHI thread to make sure all RHICmdList.EndRenderQuery() commands got executed.
 			RHICmdList.ImmediateFlush(EImmediateFlushType::FlushRHIThread);
 			RHICmdList.GetRenderQueryResult(TimerQueries[0], OldAbsTime, true);
-			TimerQueryPool.ReleaseQuery(TimerQueries[0]);
+			TimerQueryPool->ReleaseQuery(TimerQueries[0]);
 
 			for(uint32 Iteration = 0; Iteration < IterationCount; ++Iteration)
 			{
@@ -629,7 +628,7 @@ void RendererGPUBenchmark(FRHICommandListImmediate& RHICmdList, FSynthBenchmarkR
 
 					uint64 AbsTime;
 					RHICmdList.GetRenderQueryResult(TimerQueries[QueryIndex], AbsTime, true);
-					TimerQueryPool.ReleaseQuery(TimerQueries[QueryIndex]);
+					TimerQueryPool->ReleaseQuery(TimerQueries[QueryIndex]);
 
 					uint64 RelTime = FMath::Max(AbsTime - OldAbsTime, 1ull);
 
diff --git a/Engine/Source/Runtime/Renderer/Private/RendererScene.cpp b/Engine/Source/Runtime/Renderer/Private/RendererScene.cpp
--- a/Engine/Source/Runtime/Renderer/Private/RendererScene.cpp
+++ b/Engine/Source/Runtime/Renderer/Private/RendererScene.cpp
@@ -94,6 +94,9 @@ struct FSpeedTreeWindComputation
 /** Default constructor. */
 FSceneViewState::FSceneViewState()
 	: OcclusionQueryPool(RQT_Occlusion)
+	, TimerQueryPool(RHICreateRenderQueryPool(RQT_AbsoluteTime, FLatentGPUTimer::NumBufferedFrames * 2 * 2 * 2))
+	, TranslucencyTimer(TimerQueryPool)
+	, SeparateTranslucencyTimer(TimerQueryPool)
 {
 	UniqueID = FSceneViewState_UniqueID.Increment();
 	OcclusionFrameCounter = 0;
@@ -3552,12 +3555,18 @@ bool FMotionBlurInfoData::GetPrimitiveMotionBlurInfo(const FPrimitiveSceneInfo*
 
 //////////////////////////////////////////////////////////////////////////
 
-FLatentGPUTimer::FLatentGPUTimer(int32 InAvgSamples)
-: AvgSamples(InAvgSamples)
+FLatentGPUTimer::FLatentGPUTimer(FRenderQueryPoolRHIRef& InTimerQueryPool, int32 InAvgSamples)
+: TimerQueryPool(InTimerQueryPool)
+, AvgSamples(InAvgSamples)
 , TotalTime(0.0f)
 , SampleIndex(0)
 , QueryIndex(0)
 {
+	for (int i = 0; i < NumBufferedFrames; i++)
+	{
+		QueriesInFlight[i] = false;
+	}
+
 	TimeSamples.AddZeroed(AvgSamples);
 }
 
@@ -3570,43 +3579,57 @@ bool FLatentGPUTimer::Tick(FRHICommandListImmediate& RHICmdList)
 
 	QueryIndex = (QueryIndex + 1) % NumBufferedFrames;
 
-	if (StartQueries[QueryIndex] && EndQueries[QueryIndex])
+	for (int IndexOffset = 0; IndexOffset < NumBufferedFrames - 1; IndexOffset++)
 	{
-		if (IsRunningRHIInSeparateThread())
+		int CollectedIndex = (QueryIndex + IndexOffset) % NumBufferedFrames;
+		if (!QueriesInFlight[CollectedIndex])
 		{
-			// Block until the RHI thread has processed the previous query commands, if necessary
-			// Stat disabled since we buffer 2 frames minimum, it won't actually block
-			//SCOPE_CYCLE_COUNTER(STAT_TranslucencyTimestampQueryFence_Wait);
-			int32 BlockFrame = NumBufferedFrames - 1;
-			FRHICommandListExecutor::WaitOnRHIThreadFence(QuerySubmittedFences[BlockFrame]);
-			QuerySubmittedFences[BlockFrame] = nullptr;
+			continue;
 		}
 
-		uint64 StartMicroseconds;
-		uint64 EndMicroseconds;
-		bool bStartSuccess;
-		bool bEndSuccess;
-
+		if (StartQueries[CollectedIndex] && EndQueries[CollectedIndex])
 		{
-			// Block on the GPU until we have the timestamp query results, if necessary
-			// Stat disabled since we buffer 2 frames minimum, it won't actually block
-			//SCOPE_CYCLE_COUNTER(STAT_TranslucencyTimestampQuery_Wait);
-			bStartSuccess = RHICmdList.GetRenderQueryResult(StartQueries[QueryIndex], StartMicroseconds, true);
-			bEndSuccess = RHICmdList.GetRenderQueryResult(EndQueries[QueryIndex], EndMicroseconds, true);
-		}
+			if (IsRunningRHIInSeparateThread())
+			{
+				// Block until the RHI thread has processed the previous query commands, if necessary
+				// Stat disabled since we buffer 2 frames minimum, it won't actually block
+				//SCOPE_CYCLE_COUNTER(STAT_TranslucencyTimestampQueryFence_Wait);
+				int32 BlockFrame = NumBufferedFrames - 1;
+				FRHICommandListExecutor::WaitOnRHIThreadFence(QuerySubmittedFences[BlockFrame]);
+				QuerySubmittedFences[BlockFrame] = nullptr;
+			}
 
-		TotalTime -= TimeSamples[SampleIndex];
-		float LastFrameTranslucencyDurationMS = TimeSamples[SampleIndex];
-		if (bStartSuccess && bEndSuccess)
-		{
-			LastFrameTranslucencyDurationMS = (EndMicroseconds - StartMicroseconds) / 1000.0f;
-		}
+			uint64 StartMicroseconds;
+			uint64 EndMicroseconds;
+			bool bStartSuccess;
+			bool bEndSuccess;
 
-		TimeSamples[SampleIndex] = LastFrameTranslucencyDurationMS;
-		TotalTime += LastFrameTranslucencyDurationMS;
-		SampleIndex = (SampleIndex + 1) % AvgSamples;
+			{
+				// Block on the GPU until we have the timestamp query results, if necessary
+				// Stat disabled since we buffer 2 frames minimum, it won't actually block
+				//SCOPE_CYCLE_COUNTER(STAT_TranslucencyTimestampQuery_Wait);
+				bStartSuccess = RHICmdList.GetRenderQueryResult(StartQueries[CollectedIndex], StartMicroseconds, false);
+				bEndSuccess = RHICmdList.GetRenderQueryResult(EndQueries[CollectedIndex], EndMicroseconds, false);
+				if (!bStartSuccess || !bEndSuccess)
+				{
+					return false;
+				}
+				QueriesInFlight[CollectedIndex] = false;
+			}
 
-		return bStartSuccess && bEndSuccess;
+			TotalTime -= TimeSamples[SampleIndex];
+			float LastFrameTranslucencyDurationMS = TimeSamples[SampleIndex];
+			if (bStartSuccess && bEndSuccess)
+			{
+				LastFrameTranslucencyDurationMS = (EndMicroseconds - StartMicroseconds) / 1000.0f;
+			}
+
+			TimeSamples[SampleIndex] = LastFrameTranslucencyDurationMS;
+			TotalTime += LastFrameTranslucencyDurationMS;
+			SampleIndex = (SampleIndex + 1) % AvgSamples;
+
+			return bStartSuccess && bEndSuccess;
+		}
 	}
 
 	return false;
@@ -3620,10 +3643,11 @@ void FLatentGPUTimer::Begin(FRHICommandListImmediate& RHICmdList)
 	}
 	
 	if (!StartQueries[QueryIndex])
-	{
-		StartQueries[QueryIndex] = RHICmdList.CreateRenderQuery(RQT_AbsoluteTime);
+	{		
+		StartQueries[QueryIndex] = TimerQueryPool->AllocateQuery();
 	}
 
+	ensure(QueriesInFlight[QueryIndex] == false);
 	RHICmdList.EndRenderQuery(StartQueries[QueryIndex]);
 }
 
@@ -3636,10 +3660,12 @@ void FLatentGPUTimer::End(FRHICommandListImmediate& RHICmdList)
 	
 	if (!EndQueries[QueryIndex])
 	{
-		EndQueries[QueryIndex] = RHICmdList.CreateRenderQuery(RQT_AbsoluteTime);
+		EndQueries[QueryIndex] = TimerQueryPool->AllocateQuery();
 	}
 
+	ensure(QueriesInFlight[QueryIndex] == false);
 	RHICmdList.EndRenderQuery(EndQueries[QueryIndex]);
+	QueriesInFlight[QueryIndex] = true;
 	// Hint to the RHI to submit commands up to this point to the GPU if possible.  Can help avoid CPU stalls next frame waiting
 	// for these query results on some platforms.
 	RHICmdList.SubmitCommandsHint();
@@ -3661,8 +3687,8 @@ void FLatentGPUTimer::Release()
 {
 	for (int32 i = 0; i < NumBufferedFrames; ++i)
 	{
-		StartQueries[i].SafeRelease();
-		EndQueries[i].SafeRelease();
+		if(StartQueries[i].IsValid()) StartQueries[i].SafeRelease();
+		if(EndQueries[i].IsValid()) EndQueries[i].SafeRelease();
 	}
 }
 
diff --git a/Engine/Source/Runtime/Renderer/Private/ScenePrivate.h b/Engine/Source/Runtime/Renderer/Private/ScenePrivate.h
--- a/Engine/Source/Runtime/Renderer/Private/ScenePrivate.h
+++ b/Engine/Source/Runtime/Renderer/Private/ScenePrivate.h
@@ -495,11 +495,15 @@ Buffers multiple frames to avoid waiting on the GPU so times are a little lagged
 */
 class FLatentGPUTimer
 {
-	static const int32 NumBufferedFrames = FOcclusionQueryHelpers::MaxBufferedOcclusionFrames + 1;
-
+	FRenderQueryPoolRHIRef& TimerQueryPool;
 public:
+	static const int32 NumBufferedFrames = FOcclusionQueryHelpers::MaxBufferedOcclusionFrames + 1;
 
-	FLatentGPUTimer(int32 InAvgSamples = 30);
+	FLatentGPUTimer(FRenderQueryPoolRHIRef& InTimerQueryPool, int32 InAvgSamples = 30);
+	~FLatentGPUTimer()
+	{
+		Release();
+	}
 
 	void Release();
 
@@ -526,6 +530,7 @@ private:
 	int32 SampleIndex;
 
 	int32 QueryIndex;
+	bool QueriesInFlight[NumBufferedFrames];
 	FRenderQueryRHIRef StartQueries[NumBufferedFrames];
 	FRenderQueryRHIRef EndQueries[NumBufferedFrames];
 	FGraphEventRef QuerySubmittedFences[NumBufferedFrames];
@@ -898,6 +903,7 @@ public:
 	FRWBuffer CapsuleTileIntersectionCountsBuffer;
 
 	/** Timestamp queries around separate translucency, used for auto-downsampling. */
+	FRenderQueryPoolRHIRef TimerQueryPool;
 	FLatentGPUTimer TranslucencyTimer;
 	FLatentGPUTimer SeparateTranslucencyTimer;
 
diff --git a/Engine/Source/Runtime/UtilityShaders/Private/LongGPUTask.cpp b/Engine/Source/Runtime/UtilityShaders/Private/LongGPUTask.cpp
--- a/Engine/Source/Runtime/UtilityShaders/Private/LongGPUTask.cpp
+++ b/Engine/Source/Runtime/UtilityShaders/Private/LongGPUTask.cpp
@@ -13,6 +13,7 @@ int32 NumMeasuredIterationsToAchieve100ms = 0;
 
 const int32 NumIterationsForMeasurement = 5;
 
+static bool HasRun = false;
 FRenderQueryRHIRef TimeQueryStart;
 FRenderQueryRHIRef TimeQueryEnd;
 
@@ -78,10 +79,8 @@ void IssueScalableLongGPUTask(FRHICommandListImmediate& RHICmdList, int32 NumIte
 
 void MeasureLongGPUTaskExecutionTime(FRHICommandListImmediate& RHICmdList)
 {
-	if (TimeQueryStart != nullptr && TimeQueryEnd != nullptr)
-	{
-		return;
-	}
+	if (HasRun) { return; }
+
 	check(TimeQueryStart == nullptr && TimeQueryEnd == nullptr);
 
 	TimeQueryStart = RHICmdList.CreateRenderQuery(RQT_AbsoluteTime);
@@ -94,7 +93,10 @@ void MeasureLongGPUTaskExecutionTime(FRHICommandListImmediate& RHICmdList)
 		IssueScalableLongGPUTask(RHICmdList, NumIterationsForMeasurement);
 
 		RHICmdList.EndRenderQuery(TimeQueryEnd);
+		HasRun = true;
 	}
+
+	ClearLongGPUTaskQueries();
 }
 
 void ClearLongGPUTaskQueries()
diff --git a/Engine/Source/Runtime/VulkanRHI/Private/VulkanCommandBuffer.cpp b/Engine/Source/Runtime/VulkanRHI/Private/VulkanCommandBuffer.cpp
--- a/Engine/Source/Runtime/VulkanRHI/Private/VulkanCommandBuffer.cpp
+++ b/Engine/Source/Runtime/VulkanRHI/Private/VulkanCommandBuffer.cpp
@@ -48,6 +48,11 @@ FVulkanCmdBuffer::FVulkanCmdBuffer(FVulkanDevice* InDevice, FVulkanCommandBuffer
 	, Timing(nullptr)
 	, LastValidTiming(0)
 {
+	if (FVulkanPlatform::SupportsTimestampRenderQueries())
+	{
+		TimestampQueryPool = new FVulkanTimestampQueryPool(InDevice, NUM_TIMESTAMP_QUERIES_PER_POOL);
+	}
+
 	FMemory::Memzero(CurrentViewport);
 	FMemory::Memzero(CurrentScissor);
 
diff --git a/Engine/Source/Runtime/VulkanRHI/Private/VulkanCommandBuffer.h b/Engine/Source/Runtime/VulkanRHI/Private/VulkanCommandBuffer.h
--- a/Engine/Source/Runtime/VulkanRHI/Private/VulkanCommandBuffer.h
+++ b/Engine/Source/Runtime/VulkanRHI/Private/VulkanCommandBuffer.h
@@ -167,8 +167,14 @@ public:
 
 	bool AcquirePoolSetAndDescriptorsIfNeeded(const class FVulkanDescriptorSetsLayout& Layout, bool bNeedDescriptors, VkDescriptorSet* OutDescriptors);
 
+	TRefCountPtr<FVulkanTimestampQueryPool> PrepareTimestampQueryPool()
+	{
+		return TimestampQueryPool;
+	}
+
 private:
 	FVulkanDevice* Device;
+	TRefCountPtr<FVulkanTimestampQueryPool> TimestampQueryPool = nullptr;
 	VkCommandBuffer CommandBufferHandle;
 	EState State;
 	double SubmittedTime = 0.0f;
diff --git a/Engine/Source/Runtime/VulkanRHI/Private/VulkanDevice.h b/Engine/Source/Runtime/VulkanRHI/Private/VulkanDevice.h
--- a/Engine/Source/Runtime/VulkanRHI/Private/VulkanDevice.h
+++ b/Engine/Source/Runtime/VulkanRHI/Private/VulkanDevice.h
@@ -250,9 +250,6 @@ public:
 	void SubmitCommandsAndFlushGPU();
 
 	FVulkanOcclusionQueryPool* AcquireOcclusionQueryPool(uint32 NumQueries);
-/*
-	FVulkanTimestampQueryPool* PrepareTimestampQueryPool(bool& bOutRequiresReset);
-*/
 
 	inline class FVulkanPipelineStateCacheManager* GetPipelineStateCache()
 	{
@@ -330,9 +327,6 @@ private:
 
 	TArray<FVulkanOcclusionQueryPool*> UsedOcclusionQueryPools;
 	TArray<FVulkanOcclusionQueryPool*> FreeOcclusionQueryPools;
-/*
-	FVulkanTimestampQueryPool* TimestampQueryPool = nullptr;
-*/
 
 	uint64 TimestampValidBitsMask = 0;
 
diff --git a/Engine/Source/Runtime/VulkanRHI/Private/VulkanQuery.cpp b/Engine/Source/Runtime/VulkanRHI/Private/VulkanQuery.cpp
--- a/Engine/Source/Runtime/VulkanRHI/Private/VulkanQuery.cpp
+++ b/Engine/Source/Runtime/VulkanRHI/Private/VulkanQuery.cpp
@@ -9,6 +9,7 @@
 #include "VulkanContext.h"
 #include "VulkanCommandBuffer.h"
 #include "EngineGlobals.h"
+#include "Core.h"
 
 
 TAutoConsoleVariable<int32> GSubmitOcclusionBatchCmdBufferCVar(
@@ -385,16 +386,77 @@ void FVulkanOcclusionQueryPool::Reset(FVulkanCmdBuffer* InCmdBuffer, uint32 InFr
 
 	State = EState::RHIT_PostBeginBatch;
 }
-/*
-bool FVulkanTimestampQueryPool::GetResults(uint32 QueryIndex, / *uint32 Word, uint32 Bit, * /bool bWait, uint64& OutResults)
+
+bool FVulkanTimingQuery::InternalGetResult(uint64& OutTime, bool bWait)
 {
-	//check((HasResultsMask[Word] & Bit) == 0);
+	ensure(State == EState::Written);
+	if (Pool->TryGetResults(this, IndexInPool, bWait, Result))
+	{
+		State = EState::HasResult;
+		OutTime = Result;
+		return true;
+	}
 
-	VkResult Result = InternalGetQueryPoolResults(QueryIndex, 1, 0/ *bWait ? 0 : VK_QUERY_RESULT_PARTIAL_BIT* /);
-	if (Result == VK_SUCCESS)
+	return false;
+}
+
+VkResult FVulkanTimestampQueryPool::InternalGetQueryPoolResults(FVulkanTimingQuery* Query, uint32 QueryIndex, uint64& OutResults)
+{
+	uint32 PendingQueries = (QueueTail <= QueueHead) ? (QueueHead - QueueTail) : (MaxQueries - QueueTail) + QueueHead;
+	uint32 FetchingQueries = (QueueTail <= QueryIndex) ? (QueryIndex - QueueTail) + 1 : (MaxQueries - QueueTail) + QueryIndex + 1;
+
+	if (PendingQueries < FetchingQueries) //if we try to fetch more than are pending that we have the result already
 	{
 		OutResults = QueryOutput[QueryIndex];
-		//HasResultsMask[Word] |= Bit;
+		ensure(QueryAllocation[QueryIndex] == Query);
+		QueryAllocation[QueryIndex] = nullptr;
+		return VK_SUCCESS;
+	}
+	else if (QueryIndex >= QueueTail)
+	{
+		uint32 NumQueries = QueryIndex - QueueTail + 1;
+		VkResult Result = VulkanRHI::vkGetQueryPoolResults(Device->GetInstanceHandle(), QueryPool, QueueTail, NumQueries, sizeof(uint64) * NumQueries, &QueryOutput[QueueTail], sizeof(uint64), VK_QUERY_RESULT_64_BIT);
+		if (Result == VK_SUCCESS)
+		{
+			check(IsInRenderingThread());
+			check(Query->Pool == this);
+			FRHICommandListExecutor::GetImmediateCommandList().EnqueueLambda([NumQueries, QueueTail = QueueTail, QueryPool = Query->Pool]
+			(FRHICommandListImmediate& RHICommandList)
+			{
+				FVulkanCmdBuffer* CmdBuffer = static_cast<FVulkanCommandListContext&>(RHICommandList.GetContext()).GetCommandBufferManager()->GetActiveCmdBuffer();
+				VulkanRHI::vkCmdResetQueryPool(CmdBuffer->GetHandle(), QueryPool->QueryPool, QueueTail, NumQueries);
+			});
+			OutResults = QueryOutput[QueryIndex];
+			ensure(QueryAllocation[QueryIndex] == Query);
+			QueryAllocation[QueryIndex] = nullptr;
+			QueueTail = QueryIndex + 1;
+		}
+		return Result;
+	}
+	else
+	{
+		FVulkanTimingQuery* SavedQuery = QueryAllocation[MaxQueries - 1];
+		VkResult FirstHalf = InternalGetQueryPoolResults(SavedQuery, MaxQueries - 1, OutResults);
+		QueryAllocation[MaxQueries - 1] = SavedQuery;
+		if (FirstHalf == VK_SUCCESS)
+		{
+			QueueTail = 0;
+			VkResult SecondHalf = InternalGetQueryPoolResults(Query, QueryIndex, OutResults);
+			if (SecondHalf == VK_SUCCESS)
+			{
+				QueueTail = QueryIndex + 1;
+			}
+			return SecondHalf;
+		}
+		return FirstHalf;
+	}
+}
+
+bool FVulkanTimestampQueryPool::TryGetResults(FVulkanTimingQuery* Query, uint32 QueryIndex, bool bWait, uint64& OutResults)
+{
+	VkResult Result = InternalGetQueryPoolResults(Query, QueryIndex, OutResults);
+	if (Result == VK_SUCCESS)
+	{
 		return true;
 	}
 	else if (Result == VK_NOT_READY)
@@ -418,7 +480,7 @@ bool FVulkanTimestampQueryPool::GetResults(uint32 QueryIndex, / *uint32 Word, ui
 					FTaskGraphInterface::Get().ProcessThreadUntilIdle(RenderThread_Local);
 				}
 
-				Result = InternalGetQueryPoolResults(QueryIndex, 1, 0);
+				Result = InternalGetQueryPoolResults(Query, QueryIndex, OutResults);
 				if (Result == VK_SUCCESS)
 				{
 					bSuccess = true;
@@ -444,8 +506,6 @@ bool FVulkanTimestampQueryPool::GetResults(uint32 QueryIndex, / *uint32 Word, ui
 				}
 			}
 
-			OutResults = QueryOutput[QueryIndex];
-			//HasResultsMask[Word] |= Bit;
 			return true;
 		}
 	}
@@ -456,7 +516,7 @@ bool FVulkanTimestampQueryPool::GetResults(uint32 QueryIndex, / *uint32 Word, ui
 
 	return false;
 }
-*/
+
 
 void FVulkanCommandListContext::BeginOcclusionQueryBatch(FVulkanCmdBuffer* CmdBuffer, uint32 NumQueriesInBatch)
 {
@@ -553,27 +613,6 @@ FVulkanOcclusionQueryPool* FVulkanDevice::AcquireOcclusionQueryPool(uint32 NumQu
 	return Pool;
 }
 
-/*
-FVulkanTimestampQueryPool* FVulkanDevice::PrepareTimestampQueryPool(bool& bOutRequiresReset)
-{
-	bOutRequiresReset = false;
-	if (!TimestampQueryPool)
-	{
-		//#todo-rco: Create this earlier
-		TimestampQueryPool = new FVulkanTimestampQueryPool(this, GMaxLifetimeTimestampQueries);
-		bOutRequiresReset = true;
-	}
-
-	// Wrap around the queries
-	if (TimestampQueryPool->GetNumAllocatedQueries() >= TimestampQueryPool->GetMaxQueries())
-	{
-		//#todo-rco: Check overlap and grow if necessary
-		TimestampQueryPool->NumUsedQueries = 0;
-	}
-
-	return TimestampQueryPool;
-}
-*/
 void FVulkanCommandListContext::EndOcclusionQueryBatch(FVulkanCmdBuffer* CmdBuffer)
 {
 	ensure(IsImmediate());
@@ -628,22 +667,15 @@ void FVulkanOcclusionQuery::ReleaseFromPool()
 
 void FVulkanCommandListContext::ReadAndCalculateGPUFrameTime()
 {
-	/*
 	check(IsImmediate());
 
 	if (FrameTiming)
 	{
-	const uint64 Delta = FrameTiming->GetTiming(false);
-	#if VULKAN_USE_NEW_QUERIES
-	const double SecondsPerCycle = FPlatformTime::GetSecondsPerCycle();
-	const double Frequency = double(FVulkanGPUTiming::GetTimingFrequency());
-	GGPUFrameTime = FMath::TruncToInt(double(Delta) / Frequency / SecondsPerCycle);
-	#else
-	GGPUFrameTime = Delta ? (Delta / 1e6) / FPlatformTime::GetSecondsPerCycle() : 0;
-	#endif
+		const uint64 Delta = FrameTiming->GetTiming(false);
+		const double Frequency = double(FVulkanGPUTiming::GetTimingFrequency());
+		GGPUFrameTime = FMath::TruncToInt(double(Delta) / Frequency);
 	}
 	else
-	*/
 	{
 		GGPUFrameTime = 0;
 	}
@@ -664,7 +696,14 @@ FRenderQueryRHIRef FVulkanDynamicRHI::RHICreateRenderQuery(ERenderQueryType Quer
 		FVulkanRenderQuery* Query = new FVulkanOcclusionQuery();
 		return Query;
 	}
+	else if (QueryType == RQT_AbsoluteTime)
+	{
+		FVulkanTimingQuery* Query = new FVulkanTimingQuery();
+		return Query;
+	}
 
+	// Dummy!
+	ensure(0);
 	FVulkanRenderQuery* Query = new FVulkanRenderQuery(QueryType);
 	return Query;
 }
@@ -698,6 +737,14 @@ bool FVulkanDynamicRHI::RHIGetRenderQueryResult(FRenderQueryRHIParamRef QueryRHI
 			return true;
 		}
 	}
+	else if (BaseQuery->QueryType == RQT_AbsoluteTime)
+	{
+		FVulkanTimingQuery* Query = static_cast<FVulkanTimingQuery*>(BaseQuery);
+		uint64 TimingResult;
+		bool QueryResult = Query->GetResult(TimingResult, false);
+		OutNumPixels = TimingResult / 1000; //RHIGetRenderQueryResult assumes microseconds where vk provided ns
+		return QueryResult;
+	}
 	/*
 	FScopeLock ScopeLock(&GOcclusionQueryCS);
 	check(IsInRenderingThread());
@@ -775,26 +822,33 @@ void FVulkanCommandListContext::RHIEndRenderQuery(FRenderQueryRHIParamRef QueryR
 		FVulkanCmdBuffer* CmdBuffer = CommandBufferManager->GetActiveCmdBuffer();
 		VulkanRHI::vkCmdEndQuery(CmdBuffer->GetHandle(), CurrentOcclusionQueryPool->GetHandle(), Query->IndexInPool);
 	}
-	/*
-	FScopeLock ScopeLock(&GOcclusionQueryCS);
-
-	FVulkanCmdBuffer* CmdBuffer = CommandBufferManager->GetActiveCmdBuffer();
-	if (Query->QueryType == RQT_AbsoluteTime)
-	{
-	bool bRequiresReset = false;
-	FVulkanTimestampQueryPool* Pool = Device->PrepareTimestampQueryPool(bRequiresReset);
-	if (bRequiresReset)
+	else if (BaseQuery->QueryType == RQT_AbsoluteTime)
 	{
-	FVulkanCmdBuffer* UploadCmdBuffer = CommandBufferManager->GetUploadCmdBuffer();
-	Pool->ResetAll(UploadCmdBuffer);
-	CommandBufferManager->SubmitUploadCmdBuffer(false);
-	}
-	const int32 QueryIndex = Pool->AllocateQuery();
-	Query->Reset(Pool, QueryIndex);
+		FVulkanTimingQuery* Query = static_cast<FVulkanTimingQuery*>(BaseQuery);
+		FVulkanCmdBuffer* CmdBuffer = CommandBufferManager->GetActiveCmdBuffer();
+		Query->Pool = Query->IsPooledAtHighLevel ? Query->Pool : CmdBuffer->PrepareTimestampQueryPool();
+		Query->IndexInPool = Query->Pool->AcquireIndex(Query);
+		Query->State = FVulkanTimingQuery::Written;
+
+		VulkanRHI::vkCmdWriteTimestamp(CmdBuffer->GetHandle(), VK_PIPELINE_STAGE_BOTTOM_OF_PIPE_BIT, Query->Pool->GetHandle(), Query->IndexInPool);
+#if UE_BUILD_DEBUG
+		FPlatformStackWalk::CaptureStackBackTrace(Query->StackFrames, 16);
+#endif
 	}
+}
 
-	Query->End(CmdBuffer);
-	*/
+FRenderQueryRHIRef FVulkanTimestampQueryPool::AllocateQuery()
+{
+	FVulkanTimingQuery* Query = new FVulkanTimingQuery();
+	Query->State = FVulkanTimingQuery::Unused;
+	Query->IsPooledAtHighLevel = true;
+	Query->Pool = this;
+	return Query;
+}
+
+void FVulkanTimestampQueryPool::ReleaseQuery(TRefCountPtr<FRHIRenderQuery> &Query)
+{
+	Query.SafeRelease();
 }
 
 void FVulkanCommandListContext::WriteBeginTimestamp(FVulkanCmdBuffer* CmdBuffer)
@@ -804,5 +858,5 @@ void FVulkanCommandListContext::WriteBeginTimestamp(FVulkanCmdBuffer* CmdBuffer)
 
 void FVulkanCommandListContext::WriteEndTimestamp(FVulkanCmdBuffer* CmdBuffer)
 {
-	FrameTiming->EndTiming();
+	FrameTiming->EndTiming(CmdBuffer);
 }
diff --git a/Engine/Source/Runtime/VulkanRHI/Private/VulkanRHI.cpp b/Engine/Source/Runtime/VulkanRHI/Private/VulkanRHI.cpp
--- a/Engine/Source/Runtime/VulkanRHI/Private/VulkanRHI.cpp
+++ b/Engine/Source/Runtime/VulkanRHI/Private/VulkanRHI.cpp
@@ -113,9 +113,6 @@ FVulkanCommandListContext::FVulkanCommandListContext(FVulkanDynamicRHI* InRHI, F
 	CommandBufferManager = new FVulkanCommandBufferManager(InDevice, this);
 	if (IsImmediate())
 	{
-		// Insert the Begin frame timestamp query. On EndDrawingViewport() we'll insert the End and immediately after a new Begin()
-		WriteBeginTimestamp(CommandBufferManager->GetActiveCmdBuffer());
-
 		// Flush the cmd buffer immediately to ensure a valid
 		// 'Last submitted' cmd buffer exists at frame 0.
 		CommandBufferManager->SubmitActiveCmdBuffer();
@@ -702,6 +699,8 @@ void FVulkanCommandListContext::RHIBeginFrame()
 	check(IsImmediate());
 	RHIPrivateBeginFrame();
 
+	FVulkanCmdBuffer* CmdBuffer = CommandBufferManager->HasPendingUploadCmdBuffer() ? CommandBufferManager->GetUploadCmdBuffer() : CommandBufferManager->GetActiveCmdBuffer();
+
 	extern uint32 GVulkanRHIDeletionFrameNumber;
 	++GVulkanRHIDeletionFrameNumber;
 
@@ -725,6 +724,10 @@ void FVulkanCommandListContext::RHIBeginDrawingViewport(FViewportRHIParamRef Vie
 	check(ViewportRHI);
 	FVulkanViewport* Viewport = ResourceCast(ViewportRHI);
 	RHI->DrawingViewport = Viewport;
+
+	check(IsImmediate());
+	// Insert the Begin frame timestamp query. On EndDrawingViewport() we'll insert the End and immediately after a new Begin()
+	WriteBeginTimestamp(CommandBufferManager->GetActiveCmdBuffer());
 }
 
 void FVulkanCommandListContext::RHIEndDrawingViewport(FViewportRHIParamRef ViewportRHI, bool bPresent, bool bLockToVsync)
diff --git a/Engine/Source/Runtime/VulkanRHI/Private/VulkanUAV.cpp b/Engine/Source/Runtime/VulkanRHI/Private/VulkanUAV.cpp
--- a/Engine/Source/Runtime/VulkanRHI/Private/VulkanUAV.cpp
+++ b/Engine/Source/Runtime/VulkanRHI/Private/VulkanUAV.cpp
@@ -343,6 +343,18 @@ FComputeFenceRHIRef FVulkanDynamicRHI::RHICreateComputeFence(const FName& Name)
 	return new FVulkanComputeFence(Device, Name);
 }
 
+FRenderQueryPoolRHIRef FVulkanDynamicRHI::RHICreateRenderQueryPool(ERenderQueryType QueryType, uint32 NumQueries)
+{
+	if (QueryType == RQT_AbsoluteTime)
+	{
+		return new FVulkanTimestampQueryPool(Device, NumQueries);
+	}
+	else
+	{
+		return new FDefaultRHIRenderQueryPool(QueryType, this, NumQueries);
+	}
+}
+
 
 bool FVulkanGPUFence::Poll() const
 {
diff --git a/Engine/Source/Runtime/VulkanRHI/Private/VulkanUtil.cpp b/Engine/Source/Runtime/VulkanRHI/Private/VulkanUtil.cpp
--- a/Engine/Source/Runtime/VulkanRHI/Private/VulkanUtil.cpp
+++ b/Engine/Source/Runtime/VulkanRHI/Private/VulkanUtil.cpp
@@ -36,7 +36,8 @@ void FVulkanGPUTiming::PlatformStaticInitialize(void* UserData)
 			UE_LOG(LogVulkanRHI, Warning, TEXT("Timestamps not supported on Device"));
 			return;
 		}
-		GTimingFrequency = (uint64)((1.0f / Limits.timestampPeriod) * 1000.0f * 1000.0f * 1000.0f);
+		GTimingFrequency = (uint64)((1000.0 * 1000.0 * 1000.0) / Limits.timestampPeriod);
+		GIsSupported = true;
 	}
 }
 
@@ -150,12 +151,12 @@ void FVulkanGPUTiming::Initialize()
 	bIsTiming = false;
 
 	// Now initialize the queries for this timing object.
-	if ( GIsSupported )
+	if (GIsSupported)
 	{
 		for (int32 Index = 0; Index < MaxTimers; ++Index)
 		{
-			Timers[Index].Begin = new FVulkanRenderQuery(RQT_AbsoluteTime);
-			Timers[Index].End = new FVulkanRenderQuery(RQT_AbsoluteTime);
+			Timers[Index].Begin = new FVulkanTimingQuery();
+			Timers[Index].End = new FVulkanTimingQuery();
 		}
 	}
 }
@@ -180,11 +181,10 @@ void FVulkanGPUTiming::Release()
 void FVulkanGPUTiming::StartTiming(FVulkanCmdBuffer* CmdBuffer)
 {
 	// Issue a timestamp query for the 'start' time.
-	if ( GIsSupported && !bIsTiming )
+	if (GIsSupported && !bIsTiming)
 	{
 		CurrentTimerIndex = (CurrentTimerIndex + 1) % MaxTimers;
 		NumActiveTimers = FMath::Min(NumActiveTimers + 1, (int32)MaxTimers);
-/*
 		if (CmdBuffer == nullptr)
 		{
 			CmdBuffer = CmdContext->GetCommandBufferManager()->GetActiveCmdBuffer();
@@ -192,7 +192,6 @@ void FVulkanGPUTiming::StartTiming(FVulkanCmdBuffer* CmdBuffer)
 		Timers[CurrentTimerIndex].BeginCmdBuffer = CmdBuffer;
 		Timers[CurrentTimerIndex].BeginFenceCounter = CmdBuffer->GetFenceSignaledCounter();
 		CmdContext->RHIEndRenderQuery(Timers[CurrentTimerIndex].Begin);
-*/
 
 		bIsTiming = true;
 	}
@@ -207,7 +206,6 @@ void FVulkanGPUTiming::EndTiming(FVulkanCmdBuffer* CmdBuffer)
 	// Issue a timestamp query for the 'end' time.
 	if (GIsSupported && bIsTiming)
 	{
-/*
 		if (CmdBuffer == nullptr)
 		{
 			CmdBuffer = CmdContext->GetCommandBufferManager()->GetActiveCmdBuffer();
@@ -215,7 +213,8 @@ void FVulkanGPUTiming::EndTiming(FVulkanCmdBuffer* CmdBuffer)
 		Timers[CurrentTimerIndex].EndCmdBuffer = CmdBuffer;
 		Timers[CurrentTimerIndex].EndFenceCounter = CmdBuffer->GetFenceSignaledCounter();
 		CmdContext->RHIEndRenderQuery(Timers[CurrentTimerIndex].End);
-*/
+		Timers[CurrentTimerIndex].IssuedTimer = true;
+
 		bIsTiming = false;
 		bEndTimestampIssued = true;
 	}
@@ -231,97 +230,76 @@ uint64 FVulkanGPUTiming::GetTiming(bool bGetCurrentResultsAndBlock)
 {
 	if (GIsSupported)
 	{
-/*
-		uint64 BeginTime, EndTime;
-		int32 TimerIndex = CurrentTimerIndex;
 		if (!bGetCurrentResultsAndBlock)
 		{
-			// Go backwards through the list
 			for (int32 Index = 1; Index < NumActiveTimers; ++Index)
 			{
-#if VULKAN_USE_NEW_QUERIES
-				check(Timers[TimerIndex].BeginCmdBuffer->GetSubmittedFenceCounter() >= Timers[TimerIndex].BeginFenceCounter);
-				check(Timers[TimerIndex].EndCmdBuffer->GetSubmittedFenceCounter() >= Timers[TimerIndex].EndFenceCounter);
-				if (Timers[TimerIndex].EndCmdBuffer->GetFenceSignaledCounter() <= Timers[TimerIndex].EndFenceCounter)
+				int32 OldestTimer = (CurrentTimerIndex + Index) % MaxTimers;
+				if (Timers[OldestTimer].IssuedTimer == false)
 				{
-					// Nothing here...
-					int i = 0;
-					++i;
+					continue;
 				}
-				else if (Timers[TimerIndex].Begin->HasQueryBeenEnded() && Timers[TimerIndex].End->HasQueryBeenEnded())
-#endif
+
+				if (Timers[OldestTimer].BeginCmdBuffer && Timers[OldestTimer].EndCmdBuffer)
 				{
-#if VULKAN_USE_NEW_QUERIES
-					check(Timers[TimerIndex].BeginCmdBuffer->GetFenceSignaledCounter() > Timers[TimerIndex].BeginFenceCounter);
-#endif
-					check(Device == CmdContext->GetDevice());
-					if (Timers[TimerIndex].Begin->GetResult(Device, BeginTime, false))
+					if (Timers[OldestTimer].BeginCmdBuffer->GetSubmittedFenceCounter() < Timers[OldestTimer].BeginFenceCounter)
 					{
-						if (Timers[TimerIndex].End->GetResult(Device, EndTime, false))
+						continue;
+					}
+
+					if (Timers[OldestTimer].EndCmdBuffer->GetFenceSignaledCounter() <= Timers[OldestTimer].EndFenceCounter)
+					{
+						// End timestamp query not done yet
+						continue;
+					}
+					else if (Timers[OldestTimer].End)
+					{
+						check(Timers[OldestTimer].BeginCmdBuffer->GetFenceSignaledCounter() > Timers[OldestTimer].BeginFenceCounter);
+						check(Device == CmdContext->GetDevice());
+						uint64 BeginTime = 0;
+						if (Timers[OldestTimer].Begin->GetResult(BeginTime, bGetCurrentResultsAndBlock))
 						{
-							if (BeginTime < EndTime)
+							uint64 EndTime = 0;
+							if (Timers[OldestTimer].End->GetResult(EndTime, bGetCurrentResultsAndBlock))
 							{
-								return (EndTime - BeginTime);
+								if (BeginTime < EndTime)
+								{
+									Timers[OldestTimer].IssuedTimer = false;
+									return (EndTime - BeginTime);
+								}
 							}
 						}
-						else
-						{
-							int i = 0;
-							++i;
-						}
 					}
 				}
-
-				// Go back
-				TimerIndex = (TimerIndex + MaxTimers - 1) % MaxTimers;
 			}
 		}
-
-		if (bGetCurrentResultsAndBlock)
+		else
 		{
-			TimerIndex = CurrentTimerIndex;
-#if VULKAN_USE_NEW_QUERIES
-			check(Timers[TimerIndex].Begin->HasQueryBeenEnded() && Timers[TimerIndex].End->HasQueryBeenEnded());
-#endif
-
-			if (Timers[TimerIndex].Begin->GetResult(Device, BeginTime, true))
+			if (Timers[CurrentTimerIndex].BeginCmdBuffer && Timers[CurrentTimerIndex].EndCmdBuffer)
 			{
-				if (Timers[TimerIndex].End->GetResult(Device, EndTime, true))
-				{
-#if VULKAN_USE_NEW_QUERIES
-					check(BeginTime < EndTime);
-					return (EndTime - BeginTime);
-#else
-					if (BeginTime < EndTime)
-					{
-						return (EndTime - BeginTime);
-					}
-#endif
-				}
-				else
+				checkf(Timers[CurrentTimerIndex].BeginCmdBuffer->GetSubmittedFenceCounter() >= Timers[CurrentTimerIndex].BeginFenceCounter, TEXT("Begin timestamp in a cmd buffer not submitted yet!"));
+				checkf(Timers[CurrentTimerIndex].EndCmdBuffer->GetSubmittedFenceCounter() >= Timers[CurrentTimerIndex].EndFenceCounter, TEXT("End timestamp in a cmd buffer not submitted yet!"));
+				while ((Timers[CurrentTimerIndex].EndCmdBuffer->GetFenceSignaledCounter() <= Timers[CurrentTimerIndex].EndFenceCounter))
 				{
-					checkf(0, TEXT("Could not wait for End timer query result!"));
+					VulkanRHI::vkDeviceWaitIdle(Device->GetInstanceHandle());
+					Timers[CurrentTimerIndex].EndCmdBuffer->GetOwner()->RefreshFenceStatus();
 				}
-			}
-			else
-			{
-				checkf(0, TEXT("Could not wait for Begin timer query result!"));
+
+				uint64 BeginTime = 0;
+				bool bValid = Timers[CurrentTimerIndex].Begin->GetResult(BeginTime, bGetCurrentResultsAndBlock);
+				ensure(bValid);
+				uint64 EndTime = 0;
+				bValid = Timers[CurrentTimerIndex].End->GetResult(EndTime, bGetCurrentResultsAndBlock);
+				ensure(bValid);
+				ensure(BeginTime <= EndTime);
+				return (EndTime - BeginTime);
 			}
 		}
-		*/
 	}
 
 	return 0;
 }
 
-#if VULKAN_USE_NEW_QUERIES
-#else
-static double ConvertTiming(uint64 Delta)
-{
-	return Delta / 1e6;
-}
-#endif
-
 /** Start this frame of per tracking */
 void FVulkanEventNodeFrame::StartFrame()
 {
@@ -342,14 +320,7 @@ float FVulkanEventNodeFrame::GetRootTimingResults()
 	{
 		const uint64 GPUTiming = RootEventTiming.GetTiming(true);
 
-#if VULKAN_USE_NEW_QUERIES
-		// In milliseconds
-/*
 		RootResult = (double)GPUTiming / (double)RootEventTiming.GetTimingFrequency();
-*/
-#else
-		RootResult = ConvertTiming(GPUTiming);
-#endif
 	}
 
 	return (float)RootResult;
@@ -362,14 +333,7 @@ float FVulkanEventNode::GetTiming()
 	if (Timing.IsSupported())
 	{
 		const uint64 GPUTiming = Timing.GetTiming(true);
-#if VULKAN_USE_NEW_QUERIES
-		// In milliseconds
-/*
 		Result = (double)GPUTiming / (double)Timing.GetTimingFrequency();
-*/
-#else
-		Result = ConvertTiming(GPUTiming);
-#endif
 	}
 
 	return Result;
diff --git a/Engine/Source/Runtime/VulkanRHI/Public/VulkanConfiguration.h b/Engine/Source/Runtime/VulkanRHI/Public/VulkanConfiguration.h
--- a/Engine/Source/Runtime/VulkanRHI/Public/VulkanConfiguration.h
+++ b/Engine/Source/Runtime/VulkanRHI/Public/VulkanConfiguration.h
@@ -78,10 +78,6 @@
 	#define VULKAN_COMMANDWRAPPERS_ENABLE						VULKAN_SHOULD_USE_COMMANDWRAPPERS
 #endif
 
-#ifndef VULKAN_USE_QUERY_WAIT
-	#define VULKAN_USE_QUERY_WAIT								0
-#endif
-
 #ifndef VULKAN_USE_IMAGE_ACQUIRE_FENCES
 	#define VULKAN_USE_IMAGE_ACQUIRE_FENCES						1
 #endif
diff --git a/Engine/Source/Runtime/VulkanRHI/Public/VulkanDynamicRHI.h b/Engine/Source/Runtime/VulkanRHI/Public/VulkanDynamicRHI.h
--- a/Engine/Source/Runtime/VulkanRHI/Public/VulkanDynamicRHI.h
+++ b/Engine/Source/Runtime/VulkanRHI/Public/VulkanDynamicRHI.h
@@ -50,6 +50,7 @@ public:
 	virtual FGeometryShaderRHIRef RHICreateGeometryShaderWithStreamOutput(const TArray<uint8>& Code, const FStreamOutElementList& ElementList, uint32 NumStrides, const uint32* Strides, int32 RasterizedStream) final override;
 	virtual FComputeShaderRHIRef RHICreateComputeShader(const TArray<uint8>& Code) final override;
 	virtual FComputeFenceRHIRef RHICreateComputeFence(const FName& Name) final override;
+	virtual FRenderQueryPoolRHIRef RHICreateRenderQueryPool(ERenderQueryType QueryType, uint32 NumQueries) final override;
 	virtual FGPUFenceRHIRef RHICreateGPUFence(const FName &Name) final override;
 	virtual FStagingBufferRHIRef RHICreateStagingBuffer(FVertexBufferRHIParamRef BackingBuffer) final override;
 	virtual FBoundShaderStateRHIRef RHICreateBoundShaderState(FVertexDeclarationRHIParamRef VertexDeclaration, FVertexShaderRHIParamRef VertexShader, FHullShaderRHIParamRef HullShader, FDomainShaderRHIParamRef DomainShader, FPixelShaderRHIParamRef PixelShader, FGeometryShaderRHIParamRef GeometryShader) final override;
diff --git a/Engine/Source/Runtime/VulkanRHI/Public/VulkanResources.h b/Engine/Source/Runtime/VulkanRHI/Public/VulkanResources.h
--- a/Engine/Source/Runtime/VulkanRHI/Public/VulkanResources.h
+++ b/Engine/Source/Runtime/VulkanRHI/Public/VulkanResources.h
@@ -31,14 +31,13 @@ namespace VulkanRHI
 	class FDeviceMemoryAllocation;
 	class FOldResourceAllocation;
 	struct FPendingBufferLock;
-	class FStagingBuffer;
 }
 
 enum
 {
 	NUM_OCCLUSION_QUERIES_PER_POOL = 4096,
 
-	NUM_TIMESTAMP_QUERIES_PER_POOL = 1024,
+	NUM_TIMESTAMP_QUERIES_PER_POOL = 10240,
 };
 
 struct FSamplerYcbcrConversionInitializer
@@ -740,7 +739,6 @@ public:
 
 protected:
 	VkQueryPool QueryPool;
-	uint32 NumUsedQueries = 0;
 	const uint32 MaxQueries;
 	const VkQueryType QueryType;
 	TArray<uint64> QueryOutput;
@@ -823,6 +821,7 @@ public:
 	EState State = Undefined;
 
 protected:
+	uint32 NumUsedQueries = 0;
 	TArray<FVulkanOcclusionQuery*> AllocatedQueries;
 	TArray<uint64> AcquiredIndices;
 	bool InternalTryGetResults(bool bWait);
@@ -833,29 +832,37 @@ protected:
 	uint32 FrameNumber = UINT32_MAX;
 };
 
-/*
-class FVulkanTimestampQueryPool : public FVulkanQueryPool
+class FVulkanTimestampQueryPool final : public FRHIRenderQueryPool, public FVulkanQueryPool
 {
+	FRenderQueryRHIRef AllocateQuery() override;
+	virtual void ReleaseQuery(TRefCountPtr<FRHIRenderQuery> &Query) override;
+
+	uint32 QueueHead = 0;
+	uint32 QueueTail = 0;
+	TArray<class FVulkanTimingQuery*> QueryAllocation;
+
 public:
 	FVulkanTimestampQueryPool(FVulkanDevice* InDevice, uint32 InMaxQueries)
 		: FVulkanQueryPool(InDevice, InMaxQueries, VK_QUERY_TYPE_TIMESTAMP)
 	{
-		//const uint32 ElementSize = sizeof(decltype(HasResultsMask)::ElementType);
-		//HasResultsMask.AddZeroed((InMaxQueries + (ElementSize - 1)) / ElementSize);
+		QueryAllocation.AddDefaulted(InMaxQueries);
 	}
-		
-	bool GetResults(uint32 QueryIndex, bool bWait, uint64& OutResults);
 
-protected:
-	//TArray<uint64> HasResultsMask;
-	//uint32 LastPresentAllocation = 0;
+	uint32 AcquireIndex(class FVulkanTimingQuery* Query)
+	{
+		uint32 QueueIndex = QueueHead++;
+		QueueHead = (QueueHead < MaxQueries) ? QueueHead : 0;
+		ensure(QueueHead != QueueTail);
+		QueryAllocation[QueueIndex] = Query;
+		return QueueIndex;
+	}
 
-	//bool GetResults(uint32 QueryIndex, uint32 Word, uint32 Bit, bool bWait, uint64& OutResults);
+	bool TryGetResults(class FVulkanTimingQuery* Query, uint32 QueryIndex, bool bWait, uint64& OutResult);
 
-	friend class FVulkanDevice;
-	friend class FVulkanGPUTiming;
+private:
+	VkResult InternalGetQueryPoolResults(class FVulkanTimingQuery* Query, uint32 QueryIndex, uint64& OutResults);
 };
-*/
+
 class FVulkanRenderQuery : public FRHIRenderQuery
 {
 public:
@@ -888,60 +895,63 @@ public:
 	};
 
 	FVulkanOcclusionQueryPool* Pool = nullptr;
-	uint64 Result = 0;
 
 	void ReleaseFromPool();
 
 	EState State = EState::Undefined;
 };
-/*
-class FVulkanRenderQuery : public FRHIRenderQuery
+
+class FVulkanTimingQuery : public FVulkanRenderQuery
 {
 public:
-	FVulkanRenderQuery(FVulkanDevice* Device, ERenderQueryType InQueryType);
-	virtual ~FVulkanRenderQuery();
-
-	inline bool HasQueryBeenEmitted() const
+	FVulkanTimingQuery()
+		: FVulkanRenderQuery(RQT_AbsoluteTime)
 	{
-		return State == EState::InEnd;
 	}
 
-	uint32 LastPoolReset = 0;
-
-private:
-	int32 QueryIndex = INT32_MAX;
-
-	const ERenderQueryType QueryType;
-
-	FVulkanCmdBuffer* BeginCmdBuffer = nullptr;
-
-	friend class FVulkanDynamicRHI;
-	friend class FVulkanCommandListContext;
-	friend class FVulkanGPUTiming;
+	virtual ~FVulkanTimingQuery()
+	{
+	}
 
-	FVulkanQueryPool* Pool = nullptr;
-	enum class EState
+	enum EState : uint8
 	{
-		Reset,
-		InBegin,
-		InEnd,
-		HasResults,
+		Unused,
+		Written,
+		HasResult,
 	};
-	EState State = EState::Reset;
-	uint64 Result = 0;
 
-	void Reset(FVulkanQueryPool* InPool, int32 InQueryIndex)
+	bool GetResult(uint64& OutTime, bool bWait)
 	{
-		QueryIndex = InQueryIndex;
-		Pool = InPool;
-		State = EState::Reset;
+		if (State == EState::HasResult)
+		{
+			OutTime = Result;
+			return true;
+		}
+		else if (State != EState::Unused)
+		{
+			return InternalGetResult(OutTime, bWait);
+		}
+		else
+		{
+			ensure(State == EState::Unused);
+			return false;
+		}
 	}
-	void Begin(FVulkanCmdBuffer* InCmdBuffer);
-	void End(FVulkanCmdBuffer* InCmdBuffer);
 
-	bool GetResult(FVulkanDevice* Device, uint64& OutResult, bool bWait);
+#if UE_BUILD_DEBUG
+	uint64 StackFrames[16];
+#endif
+
+protected:
+	TRefCountPtr<FVulkanTimestampQueryPool> Pool;
+	EState State = EState::Unused;
+	bool IsPooledAtHighLevel = false;
+
+	bool InternalGetResult(uint64& OutTime, bool bWait);
+
+	friend class FVulkanTimestampQueryPool;
+	friend class FVulkanCommandListContext;
 };
-*/
 /*
 class FVulkanRenderQuery : public FRHIRenderQuery
 {
diff --git a/Engine/Source/Runtime/VulkanRHI/Public/VulkanUtil.h b/Engine/Source/Runtime/VulkanRHI/Public/VulkanUtil.h
--- a/Engine/Source/Runtime/VulkanRHI/Public/VulkanUtil.h
+++ b/Engine/Source/Runtime/VulkanRHI/Public/VulkanUtil.h
@@ -9,8 +9,7 @@
 #include "GPUProfiler.h"
 
 class FVulkanCmdBuffer;
-class FVulkanRenderQuery;
-class FVulkanRenderQuery;
+class FVulkanTimingQuery;
 class FVulkanCommandListContext;
 
 class FVulkanGPUTiming : public FGPUTiming
@@ -88,10 +87,11 @@ private:
 		uint64 BeginFenceCounter = 0;
 		FVulkanCmdBuffer* EndCmdBuffer = nullptr;
 		uint64 EndFenceCounter = 0;
-		FVulkanRenderQuery* Begin;
-		FVulkanRenderQuery* End;
+		FVulkanTimingQuery* Begin = nullptr;
+		FVulkanTimingQuery* End = nullptr;
+		bool IssuedTimer = false;
 	};
-	FBeginEndPair Timers[MaxTimers];
+	FBeginEndPair Timers[MaxTimers] = {};
 };
 
 /** A single perf event node, which tracks information about a appBeginDrawEvent/appEndDrawEvent range. */
