DR - vk - Dump human readable nvidia driver version
#rnx

diff --git a/Engine/Source/Runtime/VulkanRHI/Private/VulkanRHI.cpp b/Engine/Source/Runtime/VulkanRHI/Private/VulkanRHI.cpp
--- a/Engine/Source/Runtime/VulkanRHI/Private/VulkanRHI.cpp
+++ b/Engine/Source/Runtime/VulkanRHI/Private/VulkanRHI.cpp
@@ -558,6 +558,7 @@ void FVulkanDynamicRHI::SelectAndInitDevice()
 		static_assert(sizeof(NvidiaVersion) == sizeof(Props.driverVersion), "Mismatched Nvidia pack driver version!");
 		NvidiaVersion.Packed = Props.driverVersion;
 		GRHIAdapterUserDriverVersion = FString::Printf(TEXT("%d.%d"), NvidiaVersion.Major, NvidiaVersion.Minor);
+		UE_LOG(LogVulkanRHI, Display, TEXT("Nvidia User Driver Version = %s"), *GRHIAdapterUserDriverVersion);
 
 		// Ignore GRHIAdapterInternalDriverVersion for now as the device name doesn't match
 	}
