DR - vk - Minor fixes

diff --git a/Engine/Source/Runtime/VulkanRHI/Private/VulkanQuery.cpp b/Engine/Source/Runtime/VulkanRHI/Private/VulkanQuery.cpp
--- a/Engine/Source/Runtime/VulkanRHI/Private/VulkanQuery.cpp
+++ b/Engine/Source/Runtime/VulkanRHI/Private/VulkanQuery.cpp
@@ -9,8 +9,10 @@
 #include "VulkanContext.h"
 #include "VulkanCommandBuffer.h"
 #include "EngineGlobals.h"
-#include "Core.h"
 
+#if VULKAN_QUERY_CALLSTACK
+#include "HAL/PlatformStackwalk.h"
+#endif
 
 TAutoConsoleVariable<int32> GSubmitOcclusionBatchCmdBufferCVar(
 	TEXT("r.Vulkan.SubmitOcclusionBatchCmdBuffer"),
@@ -418,7 +420,6 @@ VkResult FVulkanTimestampQueryPool::InternalGetQueryPoolResults(FVulkanTimingQue
 		VkResult Result = VulkanRHI::vkGetQueryPoolResults(Device->GetInstanceHandle(), QueryPool, QueueTail, NumQueries, sizeof(uint64) * NumQueries, &QueryOutput[QueueTail], sizeof(uint64), VK_QUERY_RESULT_64_BIT);
 		if (Result == VK_SUCCESS)
 		{
-			check(IsInRenderingThread());
 			check(Query->Pool == this);
 			FRHICommandListExecutor::GetImmediateCommandList().EnqueueLambda([NumQueries, QueueTail = QueueTail, QueryPool = Query->Pool]
 			(FRHICommandListImmediate& RHICommandList)
@@ -831,7 +832,7 @@ void FVulkanCommandListContext::RHIEndRenderQuery(FRenderQueryRHIParamRef QueryR
 		Query->State = FVulkanTimingQuery::Written;
 
 		VulkanRHI::vkCmdWriteTimestamp(CmdBuffer->GetHandle(), VK_PIPELINE_STAGE_BOTTOM_OF_PIPE_BIT, Query->Pool->GetHandle(), Query->IndexInPool);
-#if UE_BUILD_DEBUG
+#if VULKAN_QUERY_CALLSTACK
 		FPlatformStackWalk::CaptureStackBackTrace(Query->StackFrames, 16);
 #endif
 	}
diff --git a/Engine/Source/Runtime/VulkanRHI/Public/VulkanConfiguration.h b/Engine/Source/Runtime/VulkanRHI/Public/VulkanConfiguration.h
--- a/Engine/Source/Runtime/VulkanRHI/Public/VulkanConfiguration.h
+++ b/Engine/Source/Runtime/VulkanRHI/Public/VulkanConfiguration.h
@@ -92,6 +92,10 @@
 
 #define VULKAN_REUSE_FENCES										1
 
+#ifndef VULKAN_QUERY_CALLSTACK
+	#define VULKAN_QUERY_CALLSTACK								0
+#endif
+
 #ifndef VULKAN_ENABLE_DESKTOP_HMD_SUPPORT
 	#define VULKAN_ENABLE_DESKTOP_HMD_SUPPORT					0
 #endif
diff --git a/Engine/Source/Runtime/VulkanRHI/Public/VulkanResources.h b/Engine/Source/Runtime/VulkanRHI/Public/VulkanResources.h
--- a/Engine/Source/Runtime/VulkanRHI/Public/VulkanResources.h
+++ b/Engine/Source/Runtime/VulkanRHI/Public/VulkanResources.h
@@ -938,7 +938,7 @@ public:
 		}
 	}
 
-#if UE_BUILD_DEBUG
+#if VULKAN_QUERY_CALLSTACK
 	uint64 StackFrames[16];
 #endif
 
