DR - Fix Queries coming from the RHI thread

diff --git a/Engine/Source/Runtime/VulkanRHI/Private/VulkanQuery.cpp b/Engine/Source/Runtime/VulkanRHI/Private/VulkanQuery.cpp
--- a/Engine/Source/Runtime/VulkanRHI/Private/VulkanQuery.cpp
+++ b/Engine/Source/Runtime/VulkanRHI/Private/VulkanQuery.cpp
@@ -420,13 +420,21 @@ VkResult FVulkanTimestampQueryPool::InternalGetQueryPoolResults(FVulkanTimingQue
 		VkResult Result = VulkanRHI::vkGetQueryPoolResults(Device->GetInstanceHandle(), QueryPool, QueueTail, NumQueries, sizeof(uint64) * NumQueries, &QueryOutput[QueueTail], sizeof(uint64), VK_QUERY_RESULT_64_BIT);
 		if (Result == VK_SUCCESS)
 		{
-			check(Query->Pool == this);
-			FRHICommandListExecutor::GetImmediateCommandList().EnqueueLambda([NumQueries, QueueTail = QueueTail, QueryPool = Query->Pool]
-			(FRHICommandListImmediate& RHICommandList)
+			if (!IsInRenderingThread() || !IsRunningRHIInSeparateThread())
 			{
-				FVulkanCmdBuffer* CmdBuffer = static_cast<FVulkanCommandListContext&>(RHICommandList.GetContext()).GetCommandBufferManager()->GetActiveCmdBuffer();
-				VulkanRHI::vkCmdResetQueryPool(CmdBuffer->GetHandle(), QueryPool->QueryPool, QueueTail, NumQueries);
-			});
+				FVulkanCmdBuffer* CmdBuffer = Device->GetImmediateContext().GetCommandBufferManager()->GetActiveCmdBuffer();
+				VulkanRHI::vkCmdResetQueryPool(CmdBuffer->GetHandle(), QueryPool, QueueTail, NumQueries);
+			}
+			else
+			{
+				check(Query->Pool == this);
+				FRHICommandListExecutor::GetImmediateCommandList().EnqueueLambda([NumQueries, QueueTail = QueueTail, QueryPool = Query->Pool]
+				(FRHICommandListImmediate& RHICommandList)
+				{
+					FVulkanCmdBuffer* CmdBuffer = static_cast<FVulkanCommandListContext&>(RHICommandList.GetContext()).GetCommandBufferManager()->GetActiveCmdBuffer();
+					VulkanRHI::vkCmdResetQueryPool(CmdBuffer->GetHandle(), QueryPool->QueryPool, QueueTail, NumQueries);
+				});
+			}
 			OutResults = QueryOutput[QueryIndex];
 			ensure(QueryAllocation[QueryIndex] == Query);
 			QueryAllocation[QueryIndex] = nullptr;
