DR - vk - Partial copy from 4404878
Better test if current render-pass is compatible with new one so it can be reused
Batch targets transtion to writeable and skip ending currrent render-pass if there are no actual transitions
#rnx

diff --git a/Engine/Source/Runtime/VulkanRHI/Private/VulkanRenderTarget.cpp b/Engine/Source/Runtime/VulkanRHI/Private/VulkanRenderTarget.cpp
--- a/Engine/Source/Runtime/VulkanRHI/Private/VulkanRenderTarget.cpp
+++ b/Engine/Source/Runtime/VulkanRHI/Private/VulkanRenderTarget.cpp
@@ -514,6 +514,54 @@ void FVulkanCommandListContext::RHISetRenderTargets(uint32 NumSimultaneousRender
 	}
 }
 
+// Find out whether we can re-use current renderpass instead of starting new one
+static bool IsCompatibleRenderPass(FVulkanRenderPass* CurrentRenderPass, FVulkanRenderPass* NewRenderPass)
+{
+	if (CurrentRenderPass == nullptr)
+	{
+		return false;
+	}
+
+	const FVulkanRenderTargetLayout& CurrentLayout = CurrentRenderPass->GetLayout();
+	const FVulkanRenderTargetLayout& NewLayout = NewRenderPass->GetLayout();
+	
+	if (CurrentLayout.GetRenderPassCompatibleHash() != NewLayout.GetRenderPassCompatibleHash())
+	{
+		return false;
+	}
+	
+	int32 NumDesc = CurrentLayout.GetNumAttachmentDescriptions();
+	check(NumDesc == NewLayout.GetNumAttachmentDescriptions());
+
+	const VkAttachmentDescription* CurrentDescriptions = CurrentLayout.GetAttachmentDescriptions();
+	const VkAttachmentDescription* NewDescriptions = NewLayout.GetAttachmentDescriptions();
+	for (int32 i = 0; i < NumDesc; ++i)
+	{
+		const VkAttachmentDescription& CurrentDesc = CurrentDescriptions[i];
+		const VkAttachmentDescription& NewDesc = NewDescriptions[i];
+		
+		// New render-pass wants a clear target
+		if (NewDesc.loadOp == VK_ATTACHMENT_LOAD_OP_CLEAR || NewDesc.stencilLoadOp == VK_ATTACHMENT_LOAD_OP_CLEAR)
+		{
+			return false;
+		}
+
+		// New render-pass wants to store, while current does not
+		if ((NewDesc.storeOp == VK_ATTACHMENT_STORE_OP_STORE && CurrentDesc.storeOp == VK_ATTACHMENT_STORE_OP_DONT_CARE) ||
+			(NewDesc.stencilStoreOp == VK_ATTACHMENT_STORE_OP_STORE && CurrentDesc.stencilStoreOp == VK_ATTACHMENT_STORE_OP_DONT_CARE))
+		{
+			return false;
+		}
+
+		if (NewDesc.finalLayout != CurrentDesc.finalLayout)
+		{
+			return false;
+		}
+	}
+	
+	return true;
+}
+
 void FVulkanCommandListContext::RHISetRenderTargetsAndClear(const FRHISetRenderTargetsInfo& RenderTargetsInfo)
 {
 	FVulkanRenderTargetLayout RTLayout(*Device, RenderTargetsInfo);
@@ -529,7 +577,7 @@ void FVulkanCommandListContext::RHISetRenderTargetsAndClear(const FRHISetRenderT
 		Framebuffer = TransitionAndLayoutManager.GetOrCreateFramebuffer(*Device, RenderTargetsInfo, RTLayout, RenderPass);
 	}
 
-	if (Framebuffer == TransitionAndLayoutManager.CurrentFramebuffer && RenderPass == TransitionAndLayoutManager.CurrentRenderPass)
+	if (Framebuffer == TransitionAndLayoutManager.CurrentFramebuffer && IsCompatibleRenderPass(TransitionAndLayoutManager.CurrentRenderPass, RenderPass))
 	{
 		return;
 	}
@@ -576,6 +624,21 @@ void FVulkanCommandListContext::RHICopyToResolveTarget(FTextureRHIParamRef Sourc
 		return;
 	}
 
+	if (SourceTextureRHI == DestTextureRHI)
+	{
+		FRHITexture2D* SourceTexture2D = SourceTextureRHI->GetTexture2D();
+		if (SourceTexture2D)
+		{
+			FVulkanTexture2D* VulkanSourceTexture2D  = (FVulkanTexture2D*)SourceTexture2D;
+			if (VulkanSourceTexture2D->GetBackBuffer() != nullptr)
+			{
+				// skip Backbuffer implicit transition to Readable, to avoid splitting Post->UI renderpass
+				// do explicit transition when need to read from Backbuffer
+				return;
+			}
+		}
+	}
+
 	RHITransitionResources(EResourceTransitionAccess::EReadable, &SourceTextureRHI, 1);
 
 	auto CopyImage = [](FTransitionAndLayoutManager& InRenderPassState, FVulkanCmdBuffer* InCmdBuffer, FVulkanSurface& SrcSurface, FVulkanSurface& DstSurface, uint32 SrcNumLayers, uint32 DstNumLayers, const FResolveParams& ResolveParams)
@@ -1271,18 +1334,6 @@ void FVulkanCommandListContext::TransitionResources(const FPendingTransition& Pe
 		}
 		else if (PendingTransition.TransitionType == EResourceTransitionAccess::EWritable)
 		{
-			//#todo-rco: Until render passes come online, assume writable means end render pass
-			if (TransitionAndLayoutManager.CurrentRenderPass)
-			{
-				TransitionAndLayoutManager.EndEmulatedRenderPass(CmdBuffer);
-				if (GVulkanSubmitAfterEveryEndRenderPass)
-				{
-					CommandBufferManager->SubmitActiveCmdBuffer();
-					CommandBufferManager->PrepareForNewActiveCommandBuffer();
-					CmdBuffer = CommandBufferManager->GetActiveCmdBuffer();
-				}
-			}
-
 			if (bShowTransitionEvents)
 			{
 				for (int32 i = 0; i < PendingTransition.Textures.Num(); ++i)
@@ -1299,34 +1350,44 @@ void FVulkanCommandListContext::TransitionResources(const FPendingTransition& Pe
 				FVulkanSurface& Surface = FVulkanTextureBase::Cast(PendingTransition.Textures[Index])->Surface;
 
 				const VkImageAspectFlags AspectMask = Surface.GetFullAspectMask();
-				VkImageSubresourceRange SubresourceRange = {AspectMask, 0, Surface.GetNumMips(), 0, Surface.GetNumberOfArrayLevels()};
-
 				VkImageLayout& SrcLayout = TransitionAndLayoutManager.FindOrAddLayoutRW(Surface.Image, VK_IMAGE_LAYOUT_UNDEFINED);
-				EImageLayoutBarrier DestLayout = VulkanRHI::EImageLayoutBarrier::Undefined;
-
+				
+				VkImageLayout FinalLayout;
 				if ((AspectMask & VK_IMAGE_ASPECT_COLOR_BIT) != 0)
 				{
-					VkImageLayout FinalLayout = (Surface.UEFlags & TexCreate_RenderTargetable) ? VK_IMAGE_LAYOUT_COLOR_ATTACHMENT_OPTIMAL : VK_IMAGE_LAYOUT_GENERAL;
-					if (SrcLayout != FinalLayout)
-					{
-						//#todo-rco: Switch to pending barrier
-						VulkanSetImageLayout(CmdBuffer->GetHandle(), Surface.Image, SrcLayout, FinalLayout, SubresourceRange);
-						SrcLayout = FinalLayout;
-					}
+					FinalLayout = (Surface.UEFlags & TexCreate_RenderTargetable) ? VK_IMAGE_LAYOUT_COLOR_ATTACHMENT_OPTIMAL : VK_IMAGE_LAYOUT_GENERAL;
 				}
 				else
 				{
-					if (SrcLayout != VK_IMAGE_LAYOUT_DEPTH_STENCIL_ATTACHMENT_OPTIMAL)
+					check(Surface.IsDepthOrStencilAspect());
+					FinalLayout = VK_IMAGE_LAYOUT_DEPTH_STENCIL_ATTACHMENT_OPTIMAL;
+				}
+
+				if (SrcLayout != FinalLayout)
+				{
+					int32 BarrierIndex = Barrier.AddImageBarrier(Surface.Image, AspectMask, Surface.GetNumMips(), Surface.GetNumberOfArrayLevels());
+					//todo: TransitionAndLayoutManager should use EImageLayoutBarrier type?
+					Barrier.SetTransition(BarrierIndex, VulkanRHI::GetImageLayoutFromVulkanLayout(SrcLayout), VulkanRHI::GetImageLayoutFromVulkanLayout(FinalLayout));
+					SrcLayout = FinalLayout;
+				}
+			}
+
+			if (Barrier.NumImageBarriers() > 0)
+			{
+				//#todo-rco: Until render passes come online, assume writable means end render pass
+				if (TransitionAndLayoutManager.CurrentRenderPass)
+				{
+					TransitionAndLayoutManager.EndEmulatedRenderPass(CmdBuffer);
+					if (GVulkanSubmitAfterEveryEndRenderPass)
 					{
-						check(Surface.IsDepthOrStencilAspect());
-						//#todo-rco: Switch to pending barrier
-						VulkanSetImageLayout(CmdBuffer->GetHandle(), Surface.Image, SrcLayout, VK_IMAGE_LAYOUT_DEPTH_STENCIL_ATTACHMENT_OPTIMAL, SubresourceRange);
-						SrcLayout = VK_IMAGE_LAYOUT_DEPTH_STENCIL_ATTACHMENT_OPTIMAL;
+						CommandBufferManager->SubmitActiveCmdBuffer();
+						CommandBufferManager->PrepareForNewActiveCommandBuffer();
+						CmdBuffer = CommandBufferManager->GetActiveCmdBuffer();
 					}
 				}
-			}
 
-			Barrier.Execute(CmdBuffer);
+				Barrier.Execute(CmdBuffer);
+			}
 		}
 		else if (PendingTransition.TransitionType == EResourceTransitionAccess::ERWSubResBarrier)
 		{
diff --git a/Engine/Source/Runtime/VulkanRHI/Public/VulkanMemory.h b/Engine/Source/Runtime/VulkanRHI/Public/VulkanMemory.h
--- a/Engine/Source/Runtime/VulkanRHI/Public/VulkanMemory.h
+++ b/Engine/Source/Runtime/VulkanRHI/Public/VulkanMemory.h
@@ -1473,6 +1473,16 @@ namespace VulkanRHI
 		void InnerExecute(FVulkanCmdBuffer* CmdBuffer, bool bEnsure);
 
 	public:
+		inline int32 NumImageBarriers() const
+		{
+			return ImageBarriers.Num();
+		}
+
+		inline int32 NumBufferBarriers() const
+		{
+			return BufferBarriers.Num();
+		}
+
 		inline void ResetStages()
 		{
 			SourceStage = 0;
