DR - Set GGPUFrameTime for vulkan and stat unit

diff --git a/Engine/Source/Runtime/VulkanRHI/Private/VulkanQuery.cpp b/Engine/Source/Runtime/VulkanRHI/Private/VulkanQuery.cpp
--- a/Engine/Source/Runtime/VulkanRHI/Private/VulkanQuery.cpp
+++ b/Engine/Source/Runtime/VulkanRHI/Private/VulkanQuery.cpp
@@ -682,19 +682,13 @@ void FVulkanCommandListContext::ReadAndCalculateGPUFrameTime()
 	{
 		const uint64 Delta = FrameTiming->GetTiming(false);
 		const double Frequency = double(FVulkanGPUTiming::GetTimingFrequency());
-		GGPUFrameTime = FMath::TruncToInt(double(Delta) / Frequency);
+		const double Cycles = FPlatformTime::GetSecondsPerCycle();
+		GGPUFrameTime = FMath::TruncToInt(double(Delta) / Frequency / Cycles);
 	}
 	else
 	{
 		GGPUFrameTime = 0;
 	}
-
-	//static const auto CVar = IConsoleManager::Get().FindConsoleVariable(TEXT("r.Vulkan.ProfileCmdBuffers"));
-	//if (CVar->GetInt() != 0)
-	//{
-	//	const uint64 Delta = GetCommandBufferManager()->CalculateGPUTime();
-	//	GGPUFrameTime = Delta ? (Delta / 1e6) / FPlatformTime::GetSecondsPerCycle() : 0;
-	//}
 }
 
 
