DR - vk - Copy from 4435239, 4443564 and 4366386
- Prevent double delete of VkPipeline
- Allow Vulkan to use the r.Streaming.PoolSize size cvar
- Fixed: Android crash with FVulkanCmdBuffer::GetFenceSignaledCounter() const
- Fix cached pointers in DeferredDeletionQueue when we deleting CmdBuffers
#rnx

diff --git a/Engine/Source/Runtime/VulkanRHI/Private/VulkanCommandBuffer.cpp b/Engine/Source/Runtime/VulkanRHI/Private/VulkanCommandBuffer.cpp
--- a/Engine/Source/Runtime/VulkanRHI/Private/VulkanCommandBuffer.cpp
+++ b/Engine/Source/Runtime/VulkanRHI/Private/VulkanCommandBuffer.cpp
@@ -484,30 +484,45 @@ FVulkanCmdBuffer* FVulkanCommandBufferManager::GetUploadCmdBuffer()
 struct FRHICommandFreeUnusedCmdBuffers final : public FRHICommand<FRHICommandFreeUnusedCmdBuffers>
 {
 	FVulkanCommandBufferPool* Pool;
+	FVulkanQueue* Queue;
 
-	FRHICommandFreeUnusedCmdBuffers(FVulkanCommandBufferPool* InPool)
+	FRHICommandFreeUnusedCmdBuffers(FVulkanCommandBufferPool* InPool, FVulkanQueue* InQueue)
 		: Pool(InPool)
+		, Queue(InQueue)
 	{
 	}
 
 	void Execute(FRHICommandListBase& CmdList)
 	{
-		Pool->FreeUnusedCmdBuffers();
+		Pool->FreeUnusedCmdBuffers(Queue);
 	}
 };
 #endif
 
 
-void FVulkanCommandBufferPool::FreeUnusedCmdBuffers()
+void FVulkanCommandBufferPool::FreeUnusedCmdBuffers(FVulkanQueue* InQueue)
 {
 #if VULKAN_DELETE_STALE_CMDBUFFERS
 	FScopeLock ScopeLock(&CS);
 	const double CurrentTime = FPlatformTime::Seconds();
+	
+	// In case Queue stores pointer to a cmdbuffer, do not delete it 
+	FVulkanCmdBuffer* LastSubmittedCmdBuffer = nullptr;
+	uint64 LastSubmittedFenceCounter = 0;
+	InQueue->GetLastSubmittedInfo(LastSubmittedCmdBuffer, LastSubmittedFenceCounter);
+
+	// Deferred deletion queue caches pointers to cmdbuffers
+	FDeferredDeletionQueue& DeferredDeletionQueue = Device->GetDeferredDeletionQueue();
+	
 	for (int32 Index = CmdBuffers.Num() - 1; Index >= 0; --Index)
 	{
 		FVulkanCmdBuffer* CmdBuffer = CmdBuffers[Index];
-		if (CmdBuffer->State == FVulkanCmdBuffer::EState::ReadyForBegin && CurrentTime - CmdBuffer->SubmittedTime > CMD_BUFFER_TIME_TO_WAIT_BEFORE_DELETING)
+		if (CmdBuffer != LastSubmittedCmdBuffer &&
+			CmdBuffer->State == FVulkanCmdBuffer::EState::ReadyForBegin && 
+			(CurrentTime - CmdBuffer->SubmittedTime) > CMD_BUFFER_TIME_TO_WAIT_BEFORE_DELETING)
 		{
+			DeferredDeletionQueue.OnCmdBufferDeleted(CmdBuffer);
+			
 			delete CmdBuffer;
 			CmdBuffers.RemoveAtSwap(Index, 1, false);
 		}
@@ -521,12 +536,12 @@ void FVulkanCommandBufferManager::FreeUnusedCmdBuffers()
 	FRHICommandList& RHICmdList = FRHICommandListExecutor::GetImmediateCommandList();
 	if (!IsInRenderingThread() || (RHICmdList.Bypass() || !IsRunningRHIInSeparateThread()))
 	{
-		Pool.FreeUnusedCmdBuffers();
+		Pool.FreeUnusedCmdBuffers(Queue);
 	}
 	else
 	{
 		check(IsInRenderingThread());
-		new (RHICmdList.AllocCommand<FRHICommandFreeUnusedCmdBuffers>()) FRHICommandFreeUnusedCmdBuffers(&Pool);
+		new (RHICmdList.AllocCommand<FRHICommandFreeUnusedCmdBuffers>()) FRHICommandFreeUnusedCmdBuffers(&Pool, Queue);
 	}
 #endif
 }
diff --git a/Engine/Source/Runtime/VulkanRHI/Private/VulkanCommandBuffer.h b/Engine/Source/Runtime/VulkanRHI/Private/VulkanCommandBuffer.h
--- a/Engine/Source/Runtime/VulkanRHI/Private/VulkanCommandBuffer.h
+++ b/Engine/Source/Runtime/VulkanRHI/Private/VulkanCommandBuffer.h
@@ -236,7 +236,7 @@ public:
 		return &CS;
 	}
 
-	void FreeUnusedCmdBuffers();
+	void FreeUnusedCmdBuffers(FVulkanQueue* Queue);
 
 	inline FVulkanCommandBufferManager& GetMgr()
 	{
diff --git a/Engine/Source/Runtime/VulkanRHI/Private/VulkanMemory.cpp b/Engine/Source/Runtime/VulkanRHI/Private/VulkanMemory.cpp
--- a/Engine/Source/Runtime/VulkanRHI/Private/VulkanMemory.cpp
+++ b/Engine/Source/Runtime/VulkanRHI/Private/VulkanMemory.cpp
@@ -1943,6 +1943,19 @@ namespace VulkanRHI
 		}
 	}
 
+	void FDeferredDeletionQueue::OnCmdBufferDeleted(FVulkanCmdBuffer* DeletedCmdBuffer)
+	{
+		FScopeLock ScopeLock(&CS);
+		for (int32 Index = 0; Index < Entries.Num(); ++Index)
+		{
+			FEntry& Entry = Entries[Index];
+			if (Entry.CmdBuffer == DeletedCmdBuffer)
+			{
+				Entry.CmdBuffer = nullptr;
+			}
+		}
+	}
+
 	FTempFrameAllocationBuffer::FTempFrameAllocationBuffer(FVulkanDevice* InDevice)
 		: FDeviceChild(InDevice)
 		, BufferIndex(0)
diff --git a/Engine/Source/Runtime/VulkanRHI/Private/VulkanPipeline.cpp b/Engine/Source/Runtime/VulkanRHI/Private/VulkanPipeline.cpp
--- a/Engine/Source/Runtime/VulkanRHI/Private/VulkanPipeline.cpp
+++ b/Engine/Source/Runtime/VulkanRHI/Private/VulkanPipeline.cpp
@@ -104,8 +104,11 @@ FVulkanPipeline::FVulkanPipeline(FVulkanDevice* InDevice)
 
 FVulkanPipeline::~FVulkanPipeline()
 {
-	Device->GetDeferredDeletionQueue().EnqueueResource(VulkanRHI::FDeferredDeletionQueue::EType::Pipeline, Pipeline);
-	Pipeline = VK_NULL_HANDLE;
+	if (Pipeline != VK_NULL_HANDLE)
+	{
+		Device->GetDeferredDeletionQueue().EnqueueResource(VulkanRHI::FDeferredDeletionQueue::EType::Pipeline, Pipeline);
+		Pipeline = VK_NULL_HANDLE;
+	}
 	/* we do NOT own Layout !*/
 }
 
@@ -167,7 +170,7 @@ static TAutoConsoleVariable<int32> GEnablePipelineCacheLoadCvar(
 
 static TAutoConsoleVariable<int32> GPipelineCacheFromShaderPipelineCacheCvar(
 	TEXT("r.Vulkan.PipelineCacheFromShaderPipelineCache"),
-	1,
+	PLATFORM_ANDROID && !(PLATFORM_LUMIN || PLATFORM_LUMINGL4),
 	TEXT("0 look for a pipeline cache in the normal locations with the normal names.")
 	TEXT("1 tie the vulkan pipeline cache to the shader pipeline cache, use the PSOFC guid as part of the filename, etc."),
 	ECVF_ReadOnly
@@ -496,7 +499,7 @@ void FVulkanPipelineStateCacheManager::OnShaderPipelineCachePrecompilationComple
 	if (!bLinkedToPSOFCSucessfulLoaded)
 	{
 		Save(LinkedToPSOFCCacheFolderFilename, true);
-#if PLATFORM_ANDROID && USE_ANDROID_JNI
+#if PLATFORM_ANDROID
 		static const auto CVar = IConsoleManager::Get().FindTConsoleVariableDataInt(TEXT("r.ProgramBinaryCache.RestartAndroidAfterPrecompile"));
 		const bool bRestart = CVar->GetValueOnAnyThread() != 0;
 		if (bRestart)
diff --git a/Engine/Source/Runtime/VulkanRHI/Private/VulkanRHI.cpp b/Engine/Source/Runtime/VulkanRHI/Private/VulkanRHI.cpp
--- a/Engine/Source/Runtime/VulkanRHI/Private/VulkanRHI.cpp
+++ b/Engine/Source/Runtime/VulkanRHI/Private/VulkanRHI.cpp
@@ -185,6 +185,10 @@ void FVulkanDynamicRHI::Init()
 	LLM(VulkanLLM::Initialize());
 #endif
 
+
+	static const auto CVarStreamingTexturePoolSize = IConsoleManager::Get().FindTConsoleVariableDataInt(TEXT("r.Streaming.PoolSize"));
+	int32 StreamingPoolSizeValue = CVarStreamingTexturePoolSize->GetValueOnAnyThread();
+			
 	if (GPoolSizeVRAMPercentage > 0)
 	{
 		const uint64 TotalGPUMemory = Device->GetMemoryManager().GetTotalMemory(true);
@@ -199,6 +203,15 @@ void FVulkanDynamicRHI::Init()
 			GPoolSizeVRAMPercentage,
 			TotalGPUMemory / 1024 / 1024);
 	}
+	else if (StreamingPoolSizeValue > 0)
+	{
+		GTexturePoolSize = (int64)StreamingPoolSizeValue * 1024 * 1024;
+
+		const uint64 TotalGPUMemory = Device->GetMemoryManager().GetTotalMemory(true);
+		UE_LOG(LogRHI,Log,TEXT("Texture pool is %llu MB (of %llu MB total graphics mem)"),
+				GTexturePoolSize / 1024 / 1024,
+				TotalGPUMemory / 1024 / 1024);
+	}
 }
 
 void FVulkanDynamicRHI::PostInit()
diff --git a/Engine/Source/Runtime/VulkanRHI/Public/VulkanMemory.h b/Engine/Source/Runtime/VulkanRHI/Public/VulkanMemory.h
--- a/Engine/Source/Runtime/VulkanRHI/Public/VulkanMemory.h
+++ b/Engine/Source/Runtime/VulkanRHI/Public/VulkanMemory.h
@@ -1170,6 +1170,8 @@ namespace VulkanRHI
 		{
 			ReleaseResources(true);
 		}
+
+		void OnCmdBufferDeleted(FVulkanCmdBuffer* CmdBuffer);
 /*
 		class FVulkanAsyncDeletionWorker : public FVulkanDeviceChild, FNonAbandonableTask
 		{
