DR - vk - Android compile fix
#rnx

diff --git a/Engine/Source/Runtime/VulkanRHI/Private/VulkanPipeline.cpp b/Engine/Source/Runtime/VulkanRHI/Private/VulkanPipeline.cpp
--- a/Engine/Source/Runtime/VulkanRHI/Private/VulkanPipeline.cpp
+++ b/Engine/Source/Runtime/VulkanRHI/Private/VulkanPipeline.cpp
@@ -497,16 +497,6 @@ void FVulkanPipelineStateCacheManager::OnShaderPipelineCachePrecompilationComple
 	if (!bLinkedToPSOFCSucessfulLoaded)
 	{
 		Save(LinkedToPSOFCCacheFolderFilename, true);
-#if PLATFORM_ANDROID
-		static const auto CVar = IConsoleManager::Get().FindTConsoleVariableDataInt(TEXT("r.ProgramBinaryCache.RestartAndroidAfterPrecompile"));
-		const bool bRestart = CVar->GetValueOnAnyThread() != 0;
-		if (bRestart)
-		{
-			UE_LOG(LogVulkanRHI, Log, TEXT("FVulkanPipelineStateCacheManager::OnShaderPipelineCachePrecompilationComplete, restarting app by request."));
-			extern void AndroidThunkCpp_RestartApplication();
-			AndroidThunkCpp_RestartApplication();
-		}
-#endif
 	}
 
 	// Want to ignore any subsequent Shader Pipeline Cache opening/closing, eg when loading modules
