DR - vk - Android compile fix
#rnx

diff --git a/Engine/Source/Runtime/VulkanRHI/Private/VulkanPipeline.cpp b/Engine/Source/Runtime/VulkanRHI/Private/VulkanPipeline.cpp
--- a/Engine/Source/Runtime/VulkanRHI/Private/VulkanPipeline.cpp
+++ b/Engine/Source/Runtime/VulkanRHI/Private/VulkanPipeline.cpp
@@ -1818,7 +1818,7 @@ bool FVulkanPipelineStateCacheManager::BinaryCacheMatches(FVulkanDevice* InDevic
 #if VULKAN_ENABLE_LRU_CACHE
 void FVulkanPipelineStateCacheManager::FVKPipelineLRU::DeleteVkPipeline(FVulkanGfxPipeline* GfxPipeline)
 {
-	GfxPipeline->DeleteVkPipeline(GfxPipeline->RecentFrame + NUM_RENDER_BUFFERS < GFrameNumberRenderThread);
+	GfxPipeline->DeleteVkPipeline(GfxPipeline->RecentFrame + NUM_BUFFERS < GFrameNumberRenderThread);
 }
 
 void FVulkanPipelineStateCacheManager::FVKPipelineLRU::EnsureVkPipelineAndAddToLRU(FVulkanRHIGraphicsPipelineState* Pipeline)
diff --git a/Engine/Source/Runtime/VulkanRHI/Private/VulkanPipeline.h b/Engine/Source/Runtime/VulkanRHI/Private/VulkanPipeline.h
--- a/Engine/Source/Runtime/VulkanRHI/Private/VulkanPipeline.h
+++ b/Engine/Source/Runtime/VulkanRHI/Private/VulkanPipeline.h
@@ -478,6 +478,11 @@ public:
 		const int LRUCapacity = 2048;
 		int32 LRUUsedPipelineSize;
 
+		enum
+		{
+			NUM_BUFFERS = 3,
+		};
+
 		static void DeleteVkPipeline(FVulkanGfxPipeline* GfxPipeline);
 
 		void EnsureVkPipelineAndAddToLRU(FVulkanRHIGraphicsPipelineState* Pipeline);
