DR - vk - Prep for cmd buffer fixes
#rnx

diff --git a/Engine/Source/Runtime/VulkanRHI/Private/VulkanCommandBuffer.cpp b/Engine/Source/Runtime/VulkanRHI/Private/VulkanCommandBuffer.cpp
--- a/Engine/Source/Runtime/VulkanRHI/Private/VulkanCommandBuffer.cpp
+++ b/Engine/Source/Runtime/VulkanRHI/Private/VulkanCommandBuffer.cpp
@@ -31,19 +31,19 @@ static FAutoConsoleVariableRef CVarVulkanProfileCmdBuffers(
 const uint32 GNumberOfFramesBeforeDeletingDescriptorPool = 300;
 
 FVulkanCmdBuffer::FVulkanCmdBuffer(FVulkanDevice* InDevice, FVulkanCommandBufferPool* InCommandBufferPool, bool bInIsUploadOnly)
-	: bNeedsDynamicStateSet(true)
-	, bHasPipeline(false)
-	, bHasViewport(false)
-	, bHasScissor(false)
-	, bHasStencilRef(false)
-	, CurrentStencilRef(0)
+	: CurrentStencilRef(0)
+	, State(EState::ReadyForBegin)
+	, bNeedsDynamicStateSet(1)
+	, bHasPipeline(0)
+	, bHasViewport(0)
+	, bHasScissor(0)
+	, bHasStencilRef(0)
+	, bIsUploadOnly(bInIsUploadOnly ? 1 : 0)
 	, Device(InDevice)
 	, CommandBufferHandle(VK_NULL_HANDLE)
-	, State(EState::ReadyForBegin)
 	, Fence(nullptr)
 	, FenceSignaledCounter(0)
 	, SubmittedFenceCounter(0)
-	, bIsUploadOnly(bInIsUploadOnly)
 	, CommandBufferPool(InCommandBufferPool)
 	, Timing(nullptr)
 	, LastValidTiming(0)
diff --git a/Engine/Source/Runtime/VulkanRHI/Private/VulkanCommandBuffer.h b/Engine/Source/Runtime/VulkanRHI/Private/VulkanCommandBuffer.h
--- a/Engine/Source/Runtime/VulkanRHI/Private/VulkanCommandBuffer.h
+++ b/Engine/Source/Runtime/VulkanRHI/Private/VulkanCommandBuffer.h
@@ -134,7 +134,7 @@ public:
 	void Begin();
 	void End();
 
-	enum class EState
+	enum class EState : uint8
 	{
 		ReadyForBegin,
 		IsInsideBegin,
@@ -143,15 +143,16 @@ public:
 		Submitted,
 	};
 
-	bool bNeedsDynamicStateSet;
-	bool bHasPipeline;
-	bool bHasViewport;
-	bool bHasScissor;
-	bool bHasStencilRef;
-
 	VkViewport CurrentViewport;
 	VkRect2D CurrentScissor;
 	uint32 CurrentStencilRef;
+	EState State;
+	uint8 bNeedsDynamicStateSet	: 1;
+	uint8 bHasPipeline			: 1;
+	uint8 bHasViewport			: 1;
+	uint8 bHasScissor			: 1;
+	uint8 bHasStencilRef		: 1;
+	uint8 bIsUploadOnly			: 1;
 
 	// You never want to call Begin/EndRenderPass directly as it will mess up with the FTransitionAndLayoutManager
 	void BeginRenderPass(const FVulkanRenderTargetLayout& Layout, class FVulkanRenderPass* RenderPass, class FVulkanFramebuffer* Framebuffer, const VkClearValue* AttachmentClearValues);
@@ -176,7 +177,6 @@ private:
 	FVulkanDevice* Device;
 	/*TRefCountPtr<*/FVulkanTimestampQueryPool* /*>*/ TimestampQueryPool = nullptr;
 	VkCommandBuffer CommandBufferHandle;
-	EState State;
 	double SubmittedTime = 0.0f;
 
 	TArray<VkPipelineStageFlags> WaitFlags;
@@ -202,7 +202,6 @@ private:
 	void RefreshFenceStatus();
 	void InitializeTimings(FVulkanCommandListContext* InContext);
 
-	bool bIsUploadOnly;
 	FVulkanCommandBufferPool* CommandBufferPool;
 
 	FVulkanGPUTiming* Timing;
