DR - vk - Fix for cmd buffer deletion; only delete the inner Vulkan object, not the outer object

diff --git a/Engine/Source/Runtime/VulkanRHI/Private/VulkanCommandBuffer.cpp b/Engine/Source/Runtime/VulkanRHI/Private/VulkanCommandBuffer.cpp
--- a/Engine/Source/Runtime/VulkanRHI/Private/VulkanCommandBuffer.cpp
+++ b/Engine/Source/Runtime/VulkanRHI/Private/VulkanCommandBuffer.cpp
@@ -32,7 +32,7 @@ const uint32 GNumberOfFramesBeforeDeletingDescriptorPool = 300;
 
 FVulkanCmdBuffer::FVulkanCmdBuffer(FVulkanDevice* InDevice, FVulkanCommandBufferPool* InCommandBufferPool, bool bInIsUploadOnly)
 	: CurrentStencilRef(0)
-	, State(EState::ReadyForBegin)
+	, State(EState::NotAllocated)
 	, bNeedsDynamicStateSet(1)
 	, bHasPipeline(0)
 	, bHasViewport(0)
@@ -53,6 +53,18 @@ FVulkanCmdBuffer::FVulkanCmdBuffer(FVulkanDevice* InDevice, FVulkanCommandBuffer
 		TimestampQueryPool = new FVulkanTimestampQueryPool(InDevice, NUM_TIMESTAMP_QUERIES_PER_POOL);
 	}
 
+	{
+		FScopeLock ScopeLock(CommandBufferPool->GetCS());
+		AllocMemory();
+	}
+
+	Fence = Device->GetFenceManager().AllocateFence();
+}
+
+void FVulkanCmdBuffer::AllocMemory()
+{
+	// Assumes we are inside a lock for the pool
+	check(State == EState::NotAllocated);
 	FMemory::Memzero(CurrentViewport);
 	FMemory::Memzero(CurrentScissor);
 
@@ -62,12 +74,16 @@ FVulkanCmdBuffer::FVulkanCmdBuffer(FVulkanDevice* InDevice, FVulkanCommandBuffer
 	CreateCmdBufInfo.commandBufferCount = 1;
 	CreateCmdBufInfo.commandPool = CommandBufferPool->GetHandle();
 
-	{
-		FScopeLock ScopeLock(InCommandBufferPool->GetCS());
-		VERIFYVULKANRESULT(VulkanRHI::vkAllocateCommandBuffers(Device->GetInstanceHandle(), &CreateCmdBufInfo, &CommandBufferHandle));
-	}
+	VERIFYVULKANRESULT(VulkanRHI::vkAllocateCommandBuffers(Device->GetInstanceHandle(), &CreateCmdBufInfo, &CommandBufferHandle));
+
+	bNeedsDynamicStateSet = 1;
+	bHasPipeline = 0;
+	bHasViewport = 0;
+	bHasScissor = 0;
+	bHasStencilRef = 0;
+	State = EState::ReadyForBegin;
+
 	INC_DWORD_STAT(STAT_VulkanNumCmdBuffers);
-	Fence = Device->GetFenceManager().AllocateFence();
 }
 
 FVulkanCmdBuffer::~FVulkanCmdBuffer()
@@ -85,13 +101,10 @@ FVulkanCmdBuffer::~FVulkanCmdBuffer()
 		FenceManager.ReleaseFence(Fence);
 	}
 
+	if (State != EState::NotAllocated)
 	{
-		FScopeLock ScopeLock(CommandBufferPool->GetCS());
-		VulkanRHI::vkFreeCommandBuffers(Device->GetInstanceHandle(), CommandBufferPool->GetHandle(), 1, &CommandBufferHandle);
+		FreeMemory();
 	}
-	CommandBufferHandle = VK_NULL_HANDLE;
-
-	DEC_DWORD_STAT(STAT_VulkanNumCmdBuffers);
 
 	if (Timing)
 	{
@@ -99,6 +112,18 @@ FVulkanCmdBuffer::~FVulkanCmdBuffer()
 	}
 }
 
+void FVulkanCmdBuffer::FreeMemory()
+{
+	// Assumes we are inside a lock for the pool
+	check(State != EState::NotAllocated);
+	check(CommandBufferHandle != VK_NULL_HANDLE);
+	VulkanRHI::vkFreeCommandBuffers(Device->GetInstanceHandle(), CommandBufferPool->GetHandle(), 1, &CommandBufferHandle);
+	CommandBufferHandle = VK_NULL_HANDLE;
+
+	DEC_DWORD_STAT(STAT_VulkanNumCmdBuffers);
+	State = EState::NotAllocated;
+}
+
 void FVulkanCmdBuffer::BeginRenderPass(const FVulkanRenderTargetLayout& Layout, FVulkanRenderPass* RenderPass, FVulkanFramebuffer* Framebuffer, const VkClearValue* AttachmentClearValues)
 {
 	checkf(IsOutsideRenderPass(), TEXT("Can't BeginRP as already inside one! CmdBuffer 0x%p State=%d"), CommandBufferHandle, (int32)State);
@@ -287,6 +312,13 @@ FVulkanCommandBufferPool::~FVulkanCommandBufferPool()
 	for (int32 Index = 0; Index < CmdBuffers.Num(); ++Index)
 	{
 		FVulkanCmdBuffer* CmdBuffer = CmdBuffers[Index];
+		CmdBuffer->FreeMemory();
+		delete CmdBuffer;
+	}
+
+	for (int32 Index = 0; Index < FreeCmdBuffers.Num(); ++Index)
+	{
+		FVulkanCmdBuffer* CmdBuffer = FreeCmdBuffers[Index];
 		delete CmdBuffer;
 	}
 
@@ -385,7 +417,21 @@ void FVulkanCommandBufferManager::SubmitActiveCmdBuffer(VulkanRHI::FSemaphore* S
 
 FVulkanCmdBuffer* FVulkanCommandBufferPool::Create(bool bIsUploadOnly)
 {
-	check(Device);
+	// Assumes we are inside a lock for the pool
+	for (int32 Index = FreeCmdBuffers.Num() - 1; Index >= 0; --Index)
+	{
+		FVulkanCmdBuffer* CmdBuffer = FreeCmdBuffers[Index];
+#if VULKAN_USE_DIFFERENT_POOL_CMDBUFFERS
+		if (CmdBuffer->bIsUploadOnly == bIsUploadOnly)
+#endif
+		{
+			FreeCmdBuffers.RemoveAtSwap(Index);
+			CmdBuffer->AllocMemory();
+			CmdBuffers.Add(CmdBuffer);
+			return CmdBuffer;
+		}
+	}
+
 	FVulkanCmdBuffer* CmdBuffer = new FVulkanCmdBuffer(Device, this, bIsUploadOnly);
 	CmdBuffers.Add(CmdBuffer);
 	check(CmdBuffer);
@@ -523,8 +569,9 @@ void FVulkanCommandBufferPool::FreeUnusedCmdBuffers(FVulkanQueue* InQueue)
 		{
 			DeferredDeletionQueue.OnCmdBufferDeleted(CmdBuffer);
 			
-			delete CmdBuffer;
+			CmdBuffer->FreeMemory();
 			CmdBuffers.RemoveAtSwap(Index, 1, false);
+			FreeCmdBuffers.Add(CmdBuffer);
 		}
 	}
 #endif
diff --git a/Engine/Source/Runtime/VulkanRHI/Private/VulkanCommandBuffer.h b/Engine/Source/Runtime/VulkanRHI/Private/VulkanCommandBuffer.h
--- a/Engine/Source/Runtime/VulkanRHI/Private/VulkanCommandBuffer.h
+++ b/Engine/Source/Runtime/VulkanRHI/Private/VulkanCommandBuffer.h
@@ -128,7 +128,7 @@ public:
 	{
 		return (Timing != nullptr) && (FMath::Abs((int64)FenceSignaledCounter - (int64)LastValidTiming) < 3);
 	}
-	
+
 	void AddWaitSemaphore(VkPipelineStageFlags InWaitFlags, VulkanRHI::FSemaphore* InWaitSemaphore);
 
 	void Begin();
@@ -141,6 +141,7 @@ public:
 		IsInsideRenderPass,
 		HasEnded,
 		Submitted,
+		NotAllocated,
 	};
 
 	VkViewport CurrentViewport;
@@ -209,6 +210,9 @@ private:
 
 	void AcquirePoolSetContainer();
 
+	void AllocMemory();
+	void FreeMemory();
+
 public:
 	//#todo-rco: Hide this
 	TMap<uint32, class FVulkanTypedDescriptorPoolSet*> TypedDescriptorPoolSets;
@@ -246,6 +250,7 @@ private:
 	VkCommandPool Handle;
 
 	TArray<FVulkanCmdBuffer*> CmdBuffers;
+	TArray<FVulkanCmdBuffer*> FreeCmdBuffers;
 
 	FCriticalSection CS;
 	FVulkanDevice* Device;
