DR - vk - Fix renderpass compatibility
#rnx

diff --git a/Engine/Source/Runtime/VulkanRHI/Private/Android/VulkanAndroidPlatform.h b/Engine/Source/Runtime/VulkanRHI/Private/Android/VulkanAndroidPlatform.h
--- a/Engine/Source/Runtime/VulkanRHI/Private/Android/VulkanAndroidPlatform.h
+++ b/Engine/Source/Runtime/VulkanRHI/Private/Android/VulkanAndroidPlatform.h
@@ -17,6 +17,12 @@
 #define VULKAN_FREEPAGE_FOR_TYPE					1
 #define VULKAN_PURGE_SHADER_MODULES					1
 
+// Android's hashes currently work fine as the problematic cases are:
+//	VK_IMAGE_LAYOUT_DEPTH_READ_ONLY_STENCIL_ATTACHMENT_OPTIMAL = 1000117000,
+//	VK_IMAGE_LAYOUT_DEPTH_ATTACHMENT_STENCIL_READ_ONLY_OPTIMAL = 1000117001,
+#define VULKAN_USE_REAL_RENDERPASS_COMPATIBILITY	0
+
+
 #define ENUM_VK_ENTRYPOINTS_PLATFORM_BASE(EnumMacro)
 
 #define ENUM_VK_ENTRYPOINTS_PLATFORM_INSTANCE(EnumMacro) \
diff --git a/Engine/Source/Runtime/VulkanRHI/Private/VulkanRenderTarget.cpp b/Engine/Source/Runtime/VulkanRHI/Private/VulkanRenderTarget.cpp
--- a/Engine/Source/Runtime/VulkanRHI/Private/VulkanRenderTarget.cpp
+++ b/Engine/Source/Runtime/VulkanRHI/Private/VulkanRenderTarget.cpp
@@ -1632,6 +1632,12 @@ struct FRenderPassFullHashableStruct
 	// +1 for Depth, +1 for Stencil
 	TEnumAsByte<VkAttachmentLoadOp>		LoadOps[MaxSimultaneousRenderTargets + 2];
 	TEnumAsByte<VkAttachmentStoreOp>	StoreOps[MaxSimultaneousRenderTargets + 2];
+#if VULKAN_USE_REAL_RENDERPASS_COMPATIBILITY
+	// If the initial != final we need to add FinalLayout and potentially RefLayout
+	VkImageLayout						InitialLayout[MaxSimultaneousRenderTargets + 1];
+	//VkImageLayout						FinalLayout[MaxSimultaneousRenderTargets + 1];
+	//VkImageLayout						RefLayout[MaxSimultaneousRenderTargets + 1];
+#endif
 };
 
 
@@ -1694,8 +1700,11 @@ FVulkanRenderTargetLayout::FVulkanRenderTargetLayout(FVulkanDevice& InDevice, co
 			CurrDesc.storeOp = RenderTargetStoreActionToVulkan(RTView.StoreAction);
 			CurrDesc.stencilLoadOp = VK_ATTACHMENT_LOAD_OP_DONT_CARE;
 			CurrDesc.stencilStoreOp = VK_ATTACHMENT_STORE_OP_DONT_CARE;
+
+			// If the initial != final we need to change the FullHashInfo and use FinalLayout
 			CurrDesc.initialLayout = VK_IMAGE_LAYOUT_COLOR_ATTACHMENT_OPTIMAL;
 			CurrDesc.finalLayout = VK_IMAGE_LAYOUT_COLOR_ATTACHMENT_OPTIMAL;
+
 			ColorReferences[NumColorAttachments].attachment = NumAttachmentDescriptions;
 			ColorReferences[NumColorAttachments].layout = VK_IMAGE_LAYOUT_COLOR_ATTACHMENT_OPTIMAL;
 			if (CurrDesc.samples > VK_SAMPLE_COUNT_1_BIT)
@@ -1711,6 +1720,7 @@ FVulkanRenderTargetLayout::FVulkanRenderTargetLayout(FVulkanDevice& InDevice, co
 			CompatibleHashInfo.Formats[NumColorAttachments] = CurrDesc.format;
 			FullHashInfo.LoadOps[NumColorAttachments] = CurrDesc.loadOp;
 			FullHashInfo.StoreOps[NumColorAttachments] = CurrDesc.storeOp;
+			FullHashInfo.InitialLayout[NumColorAttachments] = VK_IMAGE_LAYOUT_COLOR_ATTACHMENT_OPTIMAL;
 			++CompatibleHashInfo.NumAttachments;
 
 			++NumAttachmentDescriptions;
@@ -1752,8 +1762,11 @@ FVulkanRenderTargetLayout::FVulkanRenderTargetLayout(FVulkanDevice& InDevice, co
 		}
 
 		DepthStencilLayout = VulkanRHI::GetDepthStencilLayout(RTInfo.DepthStencilRenderTarget.GetDepthStencilAccess(), InDevice);
+
+		// If the initial != final we need to change the FullHashInfo and use FinalLayout
 		CurrDesc.initialLayout = DepthStencilLayout;
 		CurrDesc.finalLayout = DepthStencilLayout;
+
 		DepthStencilReference.attachment = NumAttachmentDescriptions;
 		DepthStencilReference.layout = DepthStencilLayout;
 
@@ -1761,6 +1774,7 @@ FVulkanRenderTargetLayout::FVulkanRenderTargetLayout(FVulkanDevice& InDevice, co
 		FullHashInfo.LoadOps[MaxSimultaneousRenderTargets + 1] = CurrDesc.stencilLoadOp;
 		FullHashInfo.StoreOps[MaxSimultaneousRenderTargets] = CurrDesc.storeOp;
 		FullHashInfo.StoreOps[MaxSimultaneousRenderTargets + 1] = CurrDesc.stencilStoreOp;
+		FullHashInfo.InitialLayout[MaxSimultaneousRenderTargets] = DepthStencilLayout;
 		CompatibleHashInfo.Formats[MaxSimultaneousRenderTargets] = CurrDesc.format;
 
 		++NumAttachmentDescriptions;
@@ -1853,10 +1867,13 @@ FVulkanRenderTargetLayout::FVulkanRenderTargetLayout(FVulkanDevice& InDevice, co
 		CurrDesc.storeOp = RenderTargetStoreActionToVulkan(GetStoreAction(ColorEntry.Action), true);
 		CurrDesc.stencilLoadOp = VK_ATTACHMENT_LOAD_OP_DONT_CARE;
 		CurrDesc.stencilStoreOp = VK_ATTACHMENT_STORE_OP_DONT_CARE;
-		CurrDesc.initialLayout = /*RPInfo.bGeneratingMips ? VK_IMAGE_LAYOUT_GENERAL : */VK_IMAGE_LAYOUT_COLOR_ATTACHMENT_OPTIMAL;
-		CurrDesc.finalLayout = /*RPInfo.bGeneratingMips ? VK_IMAGE_LAYOUT_GENERAL : */VK_IMAGE_LAYOUT_COLOR_ATTACHMENT_OPTIMAL;
+
+		// If the initial != final we need to change the FullHashInfo and use FinalLayout
+		CurrDesc.initialLayout = VK_IMAGE_LAYOUT_COLOR_ATTACHMENT_OPTIMAL;
+		CurrDesc.finalLayout = VK_IMAGE_LAYOUT_COLOR_ATTACHMENT_OPTIMAL;
+
 		ColorReferences[NumColorAttachments].attachment = NumAttachmentDescriptions;
-		ColorReferences[NumColorAttachments].layout = /*RPInfo.bGeneratingMips ? VK_IMAGE_LAYOUT_GENERAL : */VK_IMAGE_LAYOUT_COLOR_ATTACHMENT_OPTIMAL;
+		ColorReferences[NumColorAttachments].layout = VK_IMAGE_LAYOUT_COLOR_ATTACHMENT_OPTIMAL;
 		if (CurrDesc.samples > VK_SAMPLE_COUNT_1_BIT)
 		{
 			Desc[NumAttachmentDescriptions + 1] = Desc[NumAttachmentDescriptions];
@@ -1869,6 +1886,7 @@ FVulkanRenderTargetLayout::FVulkanRenderTargetLayout(FVulkanDevice& InDevice, co
 
 		CompatibleHashInfo.Formats[NumColorAttachments] = CurrDesc.format;
 		FullHashInfo.LoadOps[NumColorAttachments] = CurrDesc.loadOp;
+		FullHashInfo.InitialLayout[NumColorAttachments] = VK_IMAGE_LAYOUT_COLOR_ATTACHMENT_OPTIMAL;
 		FullHashInfo.StoreOps[NumColorAttachments] = CurrDesc.storeOp;
 		++CompatibleHashInfo.NumAttachments;
 
@@ -1902,6 +1920,7 @@ FVulkanRenderTargetLayout::FVulkanRenderTargetLayout(FVulkanDevice& InDevice, co
 		CurrDesc.stencilStoreOp = RenderTargetStoreActionToVulkan(GetStoreAction(GetStencilActions(RPInfo.DepthStencilRenderTarget.Action)), true);
 
 		DepthStencilLayout = VulkanRHI::GetDepthStencilLayout(RPInfo.DepthStencilRenderTarget.ExclusiveDepthStencil, InDevice);
+		// If the initial != final we need to change the FullHashInfo and use FinalLayout
 		CurrDesc.initialLayout = DepthStencilLayout;
 		CurrDesc.finalLayout = DepthStencilLayout;
 		DepthStencilReference.attachment = NumAttachmentDescriptions;
@@ -1911,6 +1930,7 @@ FVulkanRenderTargetLayout::FVulkanRenderTargetLayout(FVulkanDevice& InDevice, co
 		FullHashInfo.LoadOps[MaxSimultaneousRenderTargets + 1] = CurrDesc.stencilLoadOp;
 		FullHashInfo.StoreOps[MaxSimultaneousRenderTargets] = CurrDesc.storeOp;
 		FullHashInfo.StoreOps[MaxSimultaneousRenderTargets + 1] = CurrDesc.stencilStoreOp;
+		FullHashInfo.InitialLayout[MaxSimultaneousRenderTargets] = DepthStencilLayout;
 		CompatibleHashInfo.Formats[MaxSimultaneousRenderTargets] = CurrDesc.format;
 
 		++NumAttachmentDescriptions;
@@ -1983,8 +2003,11 @@ FVulkanRenderTargetLayout::FVulkanRenderTargetLayout(const FGraphicsPipelineStat
 			CurrDesc.storeOp = VK_ATTACHMENT_STORE_OP_DONT_CARE;
 			CurrDesc.stencilLoadOp = VK_ATTACHMENT_LOAD_OP_DONT_CARE;
 			CurrDesc.stencilStoreOp = VK_ATTACHMENT_STORE_OP_DONT_CARE;
+
+			// If the initial != final we need to change the FullHashInfo and use FinalLayout
 			CurrDesc.initialLayout = VK_IMAGE_LAYOUT_COLOR_ATTACHMENT_OPTIMAL;
 			CurrDesc.finalLayout = VK_IMAGE_LAYOUT_COLOR_ATTACHMENT_OPTIMAL;
+
 			ColorReferences[NumColorAttachments].attachment = NumAttachmentDescriptions;
 			ColorReferences[NumColorAttachments].layout = VK_IMAGE_LAYOUT_COLOR_ATTACHMENT_OPTIMAL;
 			if (CurrDesc.samples > VK_SAMPLE_COUNT_1_BIT)
@@ -2000,6 +2023,7 @@ FVulkanRenderTargetLayout::FVulkanRenderTargetLayout(const FGraphicsPipelineStat
 			CompatibleHashInfo.Formats[NumColorAttachments] = CurrDesc.format;
 			FullHashInfo.LoadOps[NumColorAttachments] = CurrDesc.loadOp;
 			FullHashInfo.StoreOps[NumColorAttachments] = CurrDesc.storeOp;
+			FullHashInfo.InitialLayout[NumColorAttachments] = VK_IMAGE_LAYOUT_COLOR_ATTACHMENT_OPTIMAL;
 			++CompatibleHashInfo.NumAttachments;
 
 			++NumAttachmentDescriptions;
@@ -2031,6 +2055,8 @@ FVulkanRenderTargetLayout::FVulkanRenderTargetLayout(const FGraphicsPipelineStat
 			CurrDesc.storeOp = VK_ATTACHMENT_STORE_OP_DONT_CARE;
 			CurrDesc.stencilStoreOp = VK_ATTACHMENT_STORE_OP_DONT_CARE;
 		}
+
+		// If the initial != final we need to change the FullHashInfo and use FinalLayout
 		CurrDesc.initialLayout = VK_IMAGE_LAYOUT_DEPTH_STENCIL_ATTACHMENT_OPTIMAL;
 		CurrDesc.finalLayout = VK_IMAGE_LAYOUT_DEPTH_STENCIL_ATTACHMENT_OPTIMAL;
 
@@ -2051,6 +2077,7 @@ FVulkanRenderTargetLayout::FVulkanRenderTargetLayout(const FGraphicsPipelineStat
 		FullHashInfo.LoadOps[MaxSimultaneousRenderTargets + 1] = CurrDesc.stencilLoadOp;
 		FullHashInfo.StoreOps[MaxSimultaneousRenderTargets] = CurrDesc.storeOp;
 		FullHashInfo.StoreOps[MaxSimultaneousRenderTargets + 1] = CurrDesc.stencilStoreOp;
+		FullHashInfo.InitialLayout[MaxSimultaneousRenderTargets] = VK_IMAGE_LAYOUT_DEPTH_STENCIL_ATTACHMENT_OPTIMAL;
 		CompatibleHashInfo.Formats[MaxSimultaneousRenderTargets] = CurrDesc.format;
 
 		++NumAttachmentDescriptions;
diff --git a/Engine/Source/Runtime/VulkanRHI/Public/VulkanConfiguration.h b/Engine/Source/Runtime/VulkanRHI/Public/VulkanConfiguration.h
--- a/Engine/Source/Runtime/VulkanRHI/Public/VulkanConfiguration.h
+++ b/Engine/Source/Runtime/VulkanRHI/Public/VulkanConfiguration.h
@@ -146,6 +146,10 @@
 	#define VULKAN_SUPPORTS_GEOMETRY_SHADERS					1
 #endif
 
+#ifndef VULKAN_USE_REAL_RENDERPASS_COMPATIBILITY
+	#define VULKAN_USE_REAL_RENDERPASS_COMPATIBILITY			1
+#endif
+
 #ifndef VULKAN_USE_DIFFERENT_POOL_CMDBUFFERS
 	#define VULKAN_USE_DIFFERENT_POOL_CMDBUFFERS				1
 #endif
