DR - vk - Fixes for tracking layer

diff --git a/Engine/Source/Runtime/VulkanRHI/Private/VulkanDebug.cpp b/Engine/Source/Runtime/VulkanRHI/Private/VulkanDebug.cpp
--- a/Engine/Source/Runtime/VulkanRHI/Private/VulkanDebug.cpp
+++ b/Engine/Source/Runtime/VulkanRHI/Private/VulkanDebug.cpp
@@ -407,8 +407,7 @@ void FVulkanDynamicRHI::RemoveDebugLayerCallback()
 }
 
 
-#if VULKAN_ENABLE_TRACKING_LAYER
-
+#if VULKAN_ENABLE_TRACKING_LAYER || VULKAN_ENABLE_DUMP_LAYER
 template <typename TResourceCreateInfoType>
 struct TTrackingResource
 {
@@ -436,7 +435,7 @@ static void CaptureCallStack(FString& OutCallstack, uint32 Delta)
 #endif
 #endif
 
-#if VULKAN_ENABLE_IMAGE_TRACKING_LAYER
+#if VULKAN_ENABLE_IMAGE_TRACKING_LAYER || VULKAN_ENABLE_DUMP_LAYER
 struct FTrackingImage
 {
 	struct
@@ -492,22 +491,49 @@ struct FTrackingImage
 };
 static TMap<VkImage, FTrackingImage> GVulkanTrackingImageLayouts;
 static TMap<VkImageView, TTrackingResource<VkImageViewCreateInfo>> GVulkanTrackingImageViews;
+static VkImage GBreakOnTrackImage = VK_NULL_HANDLE;
+
+static inline void BreakOnTrackingImage(VkImage InImage)
+{
+	if (GBreakOnTrackImage != VK_NULL_HANDLE)
+	{
+		ensure(InImage != GBreakOnTrackImage);
+	}
+}
+
+static VkImage FindTrackingImage(VkImageView InView)
+{
+	const auto& Found = GVulkanTrackingImageViews.FindChecked(InView);
+	return Found.CreateInfo.image;
+}
+
+static inline void BreakOnTrackingImage(VkImageView InView)
+{
+	BreakOnTrackingImage(FindTrackingImage(InView));
+}
 #endif
 
-#if VULKAN_ENABLE_BUFFER_TRACKING_LAYER
+#if VULKAN_ENABLE_BUFFER_TRACKING_LAYER || VULKAN_ENABLE_DUMP_LAYER
 static TMap<VkBuffer, TTrackingResource<VkBufferCreateInfo>> GVulkanTrackingBuffers;
 static TMap<VkBuffer, TArray<VkBufferView>> GVulkanTrackingBufferToBufferViews;
 static TMap<VkBufferView, TTrackingResource<VkBufferViewCreateInfo>> GVulkanTrackingBufferViews;
+
+static VkBuffer FindTrackingBuffer(VkBufferView InView)
+{
+	const auto& Found = GVulkanTrackingBufferViews.FindChecked(InView);
+	return Found.CreateInfo.buffer;
+}
 #endif
 
+static void ValidationFail()
+{
+	ensure(0);
+}
 
 #if VULKAN_ENABLE_DUMP_LAYER
 #include "Misc/OutputDeviceRedirector.h"
 namespace VulkanRHI
 {
-	void TrackImageViewAdd(VkImageView View, const VkImageViewCreateInfo* CreateInfo);
-	void TrackBufferViewAdd(VkBufferView View, const VkBufferViewCreateInfo* CreateInfo);
-
 	static FCriticalSection CS;
 	struct FMutexString
 	{
@@ -570,9 +596,6 @@ namespace VulkanRHI
 
 	static const TCHAR* Tabs = TEXT("\t\t\t\t\t\t\t\t\t");
 
-	static FCriticalSection GTrackLock;
-	static TMap<VkImageView, VkImageViewCreateInfo> GImageViewTracker;
-	static TMap<VkBufferView, VkBufferViewCreateInfo> GBufferViewTracker;
 	// Image/Layer/Mip
 	struct FImageLayout
 	{
@@ -610,12 +633,6 @@ namespace VulkanRHI
 	};
 	static TMap<VkFramebuffer, FFBInfo> GFramebufferInfo;
 
-	static void ValidationFail()
-	{
-		int i =0;
-		++i;
-	}
-
 	void FlushDebugWrapperLog()
 	{
 		if (DebugLog.Len() > 0)
@@ -1456,7 +1473,6 @@ void FWrapLayer::CreateBufferView(VkResult Result, VkDevice Device, const VkBuff
 	else
 	{
 #if VULKAN_ENABLE_DUMP_LAYER
-		TrackBufferViewAdd(*BufferView, CreateInfo);
 		PrintResultAndNamedHandle(Result, TEXT("BufferView"), *BufferView);
 		FlushDebugWrapperLog();
 #endif
@@ -1526,6 +1542,7 @@ void FWrapLayer::DestroyImage(VkResult Result, VkDevice Device, VkImage Image)
 	else
 	{
 #if VULKAN_ENABLE_IMAGE_TRACKING_LAYER
+		BreakOnTrackingImage(Image);
 		{
 			FScopeLock ScopeLock(&GTrackingCS);
 			int32 NumRemoved = GVulkanTrackingImageLayouts.Remove(Image);
@@ -1549,7 +1566,6 @@ void FWrapLayer::CreateImageView(VkResult Result, VkDevice Device, const VkImage
 	else
 	{
 #if VULKAN_ENABLE_DUMP_LAYER
-		TrackImageViewAdd(*ImageView, CreateInfo);
 		PrintResultAndNamedHandle(Result, TEXT("ImageView"), *ImageView);
 		FlushDebugWrapperLog();
 #endif
@@ -1562,6 +1578,7 @@ void FWrapLayer::CreateImageView(VkResult Result, VkDevice Device, const VkImage
 			CaptureCallStack(Found.CreateCallstack, 3);
 #endif
 		}
+		BreakOnTrackingImage(*ImageView);
 #endif
 	}
 }
@@ -1832,6 +1849,7 @@ void FWrapLayer::PipelineBarrier(VkResult Result, VkCommandBuffer CommandBuffer,
 			FScopeLock ScopeLock(&GTrackingCS);
 			for (uint32 Index = 0; Index < ImageMemoryBarrierCount; ++Index)
 			{
+				BreakOnTrackingImage(ImageMemoryBarriers[Index].image);
 				FTrackingImage* TrackingImage = GVulkanTrackingImageLayouts.Find(ImageMemoryBarriers[Index].image);
 				check(TrackingImage);
 #if VULKAN_ENABLE_TRACKING_CALLSTACK
@@ -1949,24 +1967,30 @@ void FWrapLayer::UpdateDescriptorSets(VkResult Result, VkDevice Device, uint32 D
 	{
 #if VULKAN_ENABLE_DUMP_LAYER
 		DevicePrintfBegin(Device, FString::Printf(TEXT("vkUpdateDescriptorSets(NumWrites=%d, Writes=0x%p, NumCopies=%d, Copies=0x%p)"), DescriptorWriteCount, DescriptorWrites, DescriptorCopyCount, DescriptorCopies));
+#endif
 		for (uint32 Index = 0; Index < DescriptorWriteCount; ++Index)
 		{
+#if VULKAN_ENABLE_DUMP_LAYER
 			DebugLog += FString::Printf(TEXT("%sWrite[%d]: Set=0x%p Binding=%d DstArrayElem=%d NumDesc=%d DescType=%s "), Tabs, Index,
 				DescriptorWrites[Index].dstSet, DescriptorWrites[Index].dstBinding, DescriptorWrites[Index].dstArrayElement, DescriptorWrites[Index].descriptorCount, *GetDescriptorTypeString(DescriptorWrites[Index].descriptorType));
-
+#endif
 			switch (DescriptorWrites[Index].descriptorType)
 			{
 			case VK_DESCRIPTOR_TYPE_UNIFORM_BUFFER:
 			case VK_DESCRIPTOR_TYPE_STORAGE_BUFFER:
 			case VK_DESCRIPTOR_TYPE_UNIFORM_BUFFER_DYNAMIC:
 			case VK_DESCRIPTOR_TYPE_STORAGE_BUFFER_DYNAMIC:
+#if VULKAN_ENABLE_DUMP_LAYER
 				DebugLog += FString::Printf(TEXT("pBufferInfo=0x%p\n"), DescriptorWrites[Index].pBufferInfo);
+#endif
 				if (DescriptorWrites[Index].pBufferInfo)
 				{
 					for (uint32 SubIndex = 0; SubIndex < DescriptorWrites[Index].descriptorCount; ++SubIndex)
 					{
+#if VULKAN_ENABLE_DUMP_LAYER
 						DebugLog += FString::Printf(TEXT("%s\tpBufferInfo[%d]: buffer=0x%p, offset=%d, range=%d\n"), Tabs, SubIndex,
 							DescriptorWrites[Index].pBufferInfo->buffer, (int32)DescriptorWrites[Index].pBufferInfo->offset, (int32)DescriptorWrites[Index].pBufferInfo->range);
+#endif
 					}
 				}
 				else
@@ -1977,12 +2001,16 @@ void FWrapLayer::UpdateDescriptorSets(VkResult Result, VkDevice Device, uint32 D
 
 			case VK_DESCRIPTOR_TYPE_UNIFORM_TEXEL_BUFFER:
 			case VK_DESCRIPTOR_TYPE_STORAGE_TEXEL_BUFFER:
+#if VULKAN_ENABLE_DUMP_LAYER
 				DebugLog += FString::Printf(TEXT("pTexelBufferView=0x%p\n"), DescriptorWrites[Index].pTexelBufferView);
+#endif
 				if (DescriptorWrites[Index].pTexelBufferView)
 				{
 					for (uint32 SubIndex = 0; SubIndex < DescriptorWrites[Index].descriptorCount; ++SubIndex)
 					{
-						DebugLog += FString::Printf(TEXT("%s\tpTexelBufferView[%d]=0x%p(B:0x%p)\n"), Tabs, SubIndex, DescriptorWrites[Index].pTexelBufferView[SubIndex], GBufferViewTracker.FindChecked(DescriptorWrites[Index].pTexelBufferView[SubIndex]).buffer);
+#if VULKAN_ENABLE_DUMP_LAYER
+						DebugLog += FString::Printf(TEXT("%s\tpTexelBufferView[%d]=0x%p(B:0x%p)\n"), Tabs, SubIndex, DescriptorWrites[Index].pTexelBufferView[SubIndex], FindTrackingBuffer(DescriptorWrites[Index].pTexelBufferView[SubIndex]));
+#endif
 					}
 				}
 				else
@@ -1997,13 +2025,20 @@ void FWrapLayer::UpdateDescriptorSets(VkResult Result, VkDevice Device, uint32 D
 			case VK_DESCRIPTOR_TYPE_STORAGE_IMAGE:
 			case VK_DESCRIPTOR_TYPE_INPUT_ATTACHMENT:
 			default:
+#if VULKAN_ENABLE_DUMP_LAYER
 				DebugLog += FString::Printf(TEXT("pImageInfo=0x%p\n"), DescriptorWrites[Index].pImageInfo);
+#endif
 				if (DescriptorWrites[Index].pImageInfo)
 				{
 					for (uint32 SubIndex = 0; SubIndex < DescriptorWrites[Index].descriptorCount; ++SubIndex)
 					{
+#if VULKAN_ENABLE_IMAGE_TRACKING_LAYER
+						BreakOnTrackingImage(DescriptorWrites[Index].pImageInfo->imageView);
+#endif
+#if VULKAN_ENABLE_DUMP_LAYER
 						DebugLog += FString::Printf(TEXT("%s\tpImageInfo[%d]: Sampler=0x%p, ImageView=0x%p(I:0x%p), imageLayout=%s\n"), Tabs, SubIndex,
-							DescriptorWrites[Index].pImageInfo->sampler, DescriptorWrites[Index].pImageInfo->imageView, GImageViewTracker.FindChecked(DescriptorWrites[Index].pImageInfo->imageView).image, *GetImageLayoutString(DescriptorWrites[Index].pImageInfo->imageLayout));
+							DescriptorWrites[Index].pImageInfo->sampler, DescriptorWrites[Index].pImageInfo->imageView, FindTrackingImage(DescriptorWrites[Index].pImageInfo->imageView), *GetImageLayoutString(DescriptorWrites[Index].pImageInfo->imageLayout));
+#endif
 					}
 				}
 				else
@@ -2013,6 +2048,7 @@ void FWrapLayer::UpdateDescriptorSets(VkResult Result, VkDevice Device, uint32 D
 				break;
 			}
 		}
+#if VULKAN_ENABLE_DUMP_LAYER
 		FlushDebugWrapperLog();
 #endif
 	}
@@ -2039,7 +2075,7 @@ void FWrapLayer::CreateFramebuffer(VkResult Result, VkDevice Device, const VkFra
 		DebugLog += FString::Printf(TEXT("%sVkFramebufferCreateInfo: Flags=%d, RenderPass=0x%p, NumAttachments=%d\n"), Tabs, CreateInfo->flags, CreateInfo->renderPass, CreateInfo->attachmentCount);
 		for (uint32 Index = 0; Index < CreateInfo->attachmentCount; ++Index)
 		{
-			DebugLog += FString::Printf(TEXT("%s\tAttachment[%d]: ImageView=0x%p(I:0x%p)\n"), Tabs, Index, CreateInfo->pAttachments[Index], GImageViewTracker.FindChecked(CreateInfo->pAttachments[Index]).image);
+			DebugLog += FString::Printf(TEXT("%s\tAttachment[%d]: ImageView=0x%p(I:0x%p)\n"), Tabs, Index, CreateInfo->pAttachments[Index], FindTrackingImage(CreateInfo->pAttachments[Index]));
 		}
 		DebugLog += FString::Printf(TEXT("%s\twidth=%d, height=%d, layers=%d\n"), Tabs, CreateInfo->width, CreateInfo->height, CreateInfo->layers);
 #endif
@@ -2486,11 +2522,11 @@ void FWrapLayer::BeginRenderPass(VkResult Result, VkCommandBuffer CommandBuffer,
 				for (uint32 Index = 0; Index < FoundFBInfo->Info.attachmentCount; ++Index)
 				{
 					VkImageView View = FoundFBInfo->Attachments[Index];
-					VkImageViewCreateInfo* FoundImageInfo = GImageViewTracker.Find(View);
-					ensure(FoundImageInfo);
+					auto* FoundImageInfo = GVulkanTrackingImageViews.Find(View);
+					// Can be null for swapchain images!
 					if (FoundImageInfo)
 					{
-						DebugLog += FString::Printf(TEXT("%s\t\tAttachment[%d]: ImageView=0x%p(I:0x%p)\n"), Tabs, Index, View, FoundImageInfo->image);
+						DebugLog += FString::Printf(TEXT("%s\t\tAttachment[%d]: ImageView=0x%p(I:0x%p)\n"), Tabs, Index, View, FoundImageInfo->CreateInfo.image);
 					}
 				}
 			}
@@ -2661,6 +2697,8 @@ void FWrapLayer::CopyImage(VkResult Result, VkCommandBuffer CommandBuffer, VkIma
 		CmdPrintfBegin(CommandBuffer, FString::Printf(TEXT("vkCmdCopyImage(SrcImage=0x%p, SrcImageLayout=%d, DstImage=0x%p, DstImageLayout=%d, RegionCount=%d, Regions=0x%p)[...]"), SrcImage, (int32)SrcImageLayout, DstImage, (int32)DstImageLayout, RegionCount, Regions));
 #endif
 #if VULKAN_ENABLE_IMAGE_TRACKING_LAYER
+		BreakOnTrackingImage(SrcImage);
+		BreakOnTrackingImage(DstImage);
 		FScopeLock ScopeLock(&GTrackingCS);
 		FTrackingImage* FoundSrc = GVulkanTrackingImageLayouts.Find(SrcImage);
 		FTrackingImage* FoundDest = GVulkanTrackingImageLayouts.Find(DstImage);
@@ -2697,6 +2735,7 @@ void FWrapLayer::CopyBufferToImage(VkResult Result, VkCommandBuffer CommandBuffe
 		FlushDebugWrapperLog();
 #endif
 #if VULKAN_ENABLE_IMAGE_TRACKING_LAYER
+		BreakOnTrackingImage(DstImage);
 		FScopeLock ScopeLock(&GTrackingCS);
 		FTrackingImage* FoundDest = GVulkanTrackingImageLayouts.Find(DstImage);
 		ensure(FoundDest);
@@ -2756,6 +2795,8 @@ void FWrapLayer::BlitImage(VkResult Result, VkCommandBuffer CommandBuffer, VkIma
 		FlushDebugWrapperLog();
 #endif
 #if VULKAN_ENABLE_IMAGE_TRACKING_LAYER
+		BreakOnTrackingImage(SrcImage);
+		BreakOnTrackingImage(DstImage);
 		FScopeLock ScopeLock(&GTrackingCS);
 		FTrackingImage* FoundSrc = GVulkanTrackingImageLayouts.Find(SrcImage);
 		FTrackingImage* FoundDest = GVulkanTrackingImageLayouts.Find(DstImage);
@@ -2834,6 +2875,7 @@ void FWrapLayer::GetSwapChainImagesKHR(VkResult Result, VkDevice Device, VkSwapc
 			FScopeLock ScopeLock(&GTrackingCS);
 			for (uint32 Index = 0; Index < *SwapchainImageCount; ++Index)
 			{
+				BreakOnTrackingImage(SwapchainImages[Index]);
 				FTrackingImage& TrackingImage = GVulkanTrackingImageLayouts.FindOrAdd(SwapchainImages[Index]);
 				TrackingImage.Setup(1, 1, VK_IMAGE_LAYOUT_UNDEFINED, true);
 			}
@@ -3295,6 +3337,7 @@ void FWrapLayer::DestroyImageView(VkResult Result, VkDevice Device, VkImageView
 #if VULKAN_ENABLE_IMAGE_TRACKING_LAYER
 		{
 			FScopeLock ScopeLock(&GTrackingCS);
+			BreakOnTrackingImage(ImageView);
 			int32 NumRemoved = GVulkanTrackingImageViews.Remove(ImageView);
 			ensure(NumRemoved > 0);
 		}
@@ -3857,6 +3900,7 @@ void FWrapLayer::CreateWin32SurfaceKHR(VkResult Result, VkInstance Instance, con
 #if VULKAN_SUPPORTS_COLOR_CONVERSIONS
 void FWrapLayer::CreateSamplerYcbcrConversionKHR(VkResult Result, VkDevice Device, const VkSamplerYcbcrConversionCreateInfo* CreateInfo, VkSamplerYcbcrConversion* YcbcrConversion)
 {
+#if VULKAN_ENABLE_DUMP_LAYER
 	if (Result == VK_RESULT_MAX_ENUM)
 	{
 		DevicePrintfBeginResult(Device, FString::Printf(TEXT("vkCreateSamplerYcbcrConversionKHR(CreateInfo=0x%p, YcbcrConversion=0x%p)[...]"), CreateInfo, YcbcrConversion));
@@ -3866,6 +3910,7 @@ void FWrapLayer::CreateSamplerYcbcrConversionKHR(VkResult Result, VkDevice Devic
 		PrintResultAndNamedHandle(Result, TEXT("SamplerYcbcrConversionKHR"), *YcbcrConversion);
 	}
 	FlushDebugWrapperLog();
+#endif
 }
 
 void FWrapLayer::DestroySamplerYcbcrConversionKHR(VkResult Result, VkDevice Device, VkSamplerYcbcrConversion YcbcrConversion)
@@ -3903,31 +3948,6 @@ namespace VulkanRHI
 #if VULKAN_ENABLE_DUMP_LAYER
 namespace VulkanRHI
 {
-	void TrackImageViewAdd(VkImageView View, const VkImageViewCreateInfo* CreateInfo)
-	{
-		FScopeLock Lock(&GTrackLock);
-		GImageViewTracker.Add(View, *CreateInfo);
-	}
-
-	void TrackImageViewRemove(VkImageView View)
-	{
-		FScopeLock Lock(&GTrackLock);
-		GImageViewTracker.Remove(View);
-	}
-
-	void TrackBufferViewAdd(VkBufferView View, const VkBufferViewCreateInfo* CreateInfo)
-	{
-		FScopeLock Lock(&GTrackLock);
-		GBufferViewTracker.Add(View, *CreateInfo);
-	}
-
-	void TrackBufferViewRemove(VkBufferView View)
-	{
-		FScopeLock Lock(&GTrackLock);
-		GBufferViewTracker.Remove(View);
-	}
-
-
 	static struct FGlobalDumpLog
 	{
 		~FGlobalDumpLog()
