DR - vk - More tracking fixes
#rnx

diff --git a/Engine/Source/Runtime/VulkanRHI/Private/VulkanDebug.cpp b/Engine/Source/Runtime/VulkanRHI/Private/VulkanDebug.cpp
--- a/Engine/Source/Runtime/VulkanRHI/Private/VulkanDebug.cpp
+++ b/Engine/Source/Runtime/VulkanRHI/Private/VulkanDebug.cpp
@@ -497,7 +497,7 @@ static inline void BreakOnTrackingImage(VkImage InImage)
 {
 	if (GBreakOnTrackImage != VK_NULL_HANDLE)
 	{
-		ensure(InImage != GBreakOnTrackImage);
+		ensureAlways(InImage != GBreakOnTrackImage);
 	}
 }
 
@@ -596,30 +596,6 @@ namespace VulkanRHI
 
 	static const TCHAR* Tabs = TEXT("\t\t\t\t\t\t\t\t\t");
 
-	// Image/Layer/Mip
-	struct FImageLayout
-	{
-		void Init(uint32 NumLayers, uint32 NumMips)
-		{
-			ensure(Layouts.Num() == 0);
-			Layouts.Empty(NumLayers);
-
-			TArray<VkImageLayout> Temp;
-			Temp.AddUninitialized(NumMips);
-			for (uint32 Index = 0; Index < NumMips; ++Index)
-			{
-				Temp[Index] = VK_IMAGE_LAYOUT_UNDEFINED;
-			}
-
-			for (uint32 Index = 0; Index < NumLayers; ++Index)
-			{
-				Layouts.Add(Temp);
-			}
-		}
-
-		TArray<TArray<VkImageLayout>> Layouts;
-	};
-	static TMap<VkImage, FImageLayout> GOLDImageLayoutTracker;
 	struct FRenderPassInfo
 	{
 		TArray<VkAttachmentDescription> Descriptions;
@@ -1508,8 +1484,6 @@ void FWrapLayer::CreateImage(VkResult Result, VkDevice Device, const VkImageCrea
 	{
 #if VULKAN_ENABLE_DUMP_LAYER
 		PrintResultAndNamedHandle(Result, TEXT("Image"), *Image);
-		FImageLayout& Found = GOLDImageLayoutTracker.FindOrAdd(*Image);
-		Found.Init(FMath::Max(CreateInfo->arrayLayers, CreateInfo->extent.depth), CreateInfo->mipLevels);
 		FlushDebugWrapperLog();
 #endif
 
@@ -1536,7 +1510,6 @@ void FWrapLayer::DestroyImage(VkResult Result, VkDevice Device, VkImage Image)
 	{
 #if VULKAN_ENABLE_DUMP_LAYER
 		DevicePrintfBegin(Device, FString::Printf(TEXT("vkDestroyImage(Image=0x%p)"), Image));
-		GOLDImageLayoutTracker.Remove(Image);
 #endif
 	}
 	else
@@ -2860,8 +2833,6 @@ void FWrapLayer::GetSwapChainImagesKHR(VkResult Result, VkDevice Device, VkSwapc
 			for (uint32 Index = 0; Index < *SwapchainImageCount; ++Index)
 			{
 				DebugLog += FString::Printf(TEXT("%sImage[%d]=0x%p\n"), Tabs, Index, SwapchainImages[Index]);
-				FImageLayout& Layout = GOLDImageLayoutTracker.Add(SwapchainImages[Index]);
-				Layout.Init(1, 1);
 			}
 		}
 		else
