DR - vk - Change warning to display

diff --git a/Engine/Source/Runtime/VulkanRHI/Private/VulkanSwapChain.cpp b/Engine/Source/Runtime/VulkanRHI/Private/VulkanSwapChain.cpp
--- a/Engine/Source/Runtime/VulkanRHI/Private/VulkanSwapChain.cpp
+++ b/Engine/Source/Runtime/VulkanRHI/Private/VulkanSwapChain.cpp
@@ -168,7 +168,7 @@ FVulkanSwapChain::FVulkanSwapChain(VkInstance InInstance, FVulkanDevice& InDevic
 					if (!GPixelFormatNotSupportedWarning.Contains(InOutPixelFormat))
 					{
 						GPixelFormatNotSupportedWarning.Add(InOutPixelFormat);
-						UE_LOG(LogVulkanRHI, Warning, TEXT("Requested PixelFormat %d not supported by this swapchain! Falling back to supported swapchain format..."), (uint32)InOutPixelFormat);
+						UE_LOG(LogVulkanRHI, Display, TEXT("Requested PixelFormat %d not supported by this swapchain! Falling back to supported swapchain format..."), (uint32)InOutPixelFormat);
 					}
 					InOutPixelFormat = PF_Unknown;
 				}
