DR - vk - Copy from 4462867, 4473127
- Change warning to display
- Fix some leaks when exiting
- Fix missing layout
#rnx

diff --git a/Engine/Source/Runtime/VulkanRHI/Private/VulkanDebug.cpp b/Engine/Source/Runtime/VulkanRHI/Private/VulkanDebug.cpp
--- a/Engine/Source/Runtime/VulkanRHI/Private/VulkanDebug.cpp
+++ b/Engine/Source/Runtime/VulkanRHI/Private/VulkanDebug.cpp
@@ -493,7 +493,7 @@ static TMap<VkImage, FTrackingImage> GVulkanTrackingImageLayouts;
 static TMap<VkImageView, TTrackingResource<VkImageViewCreateInfo>> GVulkanTrackingImageViews;
 static VkImage GBreakOnTrackImage = VK_NULL_HANDLE;
 
-static inline void BreakOnTrackingImage(VkImage InImage)
+static FORCEINLINE void BreakOnTrackingImage(VkImage InImage)
 {
 	if (GBreakOnTrackImage != VK_NULL_HANDLE)
 	{
@@ -507,7 +507,7 @@ static VkImage FindTrackingImage(VkImageView InView)
 	return Found.CreateInfo.image;
 }
 
-static inline void BreakOnTrackingImage(VkImageView InView)
+static FORCEINLINE void BreakOnTrackingImageView(VkImageView InView)
 {
 	BreakOnTrackingImage(FindTrackingImage(InView));
 }
@@ -1551,7 +1551,7 @@ void FWrapLayer::CreateImageView(VkResult Result, VkDevice Device, const VkImage
 			CaptureCallStack(Found.CreateCallstack, 3);
 #endif
 		}
-		BreakOnTrackingImage(*ImageView);
+		BreakOnTrackingImageView(*ImageView);
 #endif
 	}
 }
@@ -2006,7 +2006,7 @@ void FWrapLayer::UpdateDescriptorSets(VkResult Result, VkDevice Device, uint32 D
 					for (uint32 SubIndex = 0; SubIndex < DescriptorWrites[Index].descriptorCount; ++SubIndex)
 					{
 #if VULKAN_ENABLE_IMAGE_TRACKING_LAYER
-						BreakOnTrackingImage(DescriptorWrites[Index].pImageInfo->imageView);
+						BreakOnTrackingImageView(DescriptorWrites[Index].pImageInfo->imageView);
 #endif
 #if VULKAN_ENABLE_DUMP_LAYER
 						DebugLog += FString::Printf(TEXT("%s\tpImageInfo[%d]: Sampler=0x%p, ImageView=0x%p(I:0x%p), imageLayout=%s\n"), Tabs, SubIndex,
@@ -3308,7 +3308,7 @@ void FWrapLayer::DestroyImageView(VkResult Result, VkDevice Device, VkImageView
 #if VULKAN_ENABLE_IMAGE_TRACKING_LAYER
 		{
 			FScopeLock ScopeLock(&GTrackingCS);
-			BreakOnTrackingImage(ImageView);
+			BreakOnTrackingImageView(ImageView);
 			int32 NumRemoved = GVulkanTrackingImageViews.Remove(ImageView);
 			ensure(NumRemoved > 0);
 		}
diff --git a/Engine/Source/Runtime/VulkanRHI/Private/VulkanMemory.cpp b/Engine/Source/Runtime/VulkanRHI/Private/VulkanMemory.cpp
--- a/Engine/Source/Runtime/VulkanRHI/Private/VulkanMemory.cpp
+++ b/Engine/Source/Runtime/VulkanRHI/Private/VulkanMemory.cpp
@@ -725,7 +725,7 @@ namespace VulkanRHI
 			FScopeLock ScopeLock(&GOldResourceLock);
 
 #if VULKAN_FREEPAGE_FOR_TYPE
-			for (int32 Index = 1; Index < FreeBufferPages.Num(); ++Index)
+			for (int32 Index = (bImmediately ? 0 : 1); Index < FreeBufferPages.Num(); ++Index)
 			{
 				FOldResourceHeapPage* Page = FreeBufferPages[Index];
 				if (bImmediately || Page->FrameFreed + NUM_FRAMES_TO_WAIT_BEFORE_RELEASING_TO_OS < GFrameNumberRenderThread)
@@ -735,7 +735,7 @@ namespace VulkanRHI
 					break;
 				}
 			}
-			for (int32 Index = 1; Index < FreeImagePages.Num(); ++Index)
+			for (int32 Index = (bImmediately ? 0 : 1); Index < FreeImagePages.Num(); ++Index)
 			{
 				FOldResourceHeapPage* Page = FreeImagePages[Index];
 				if (bImmediately || Page->FrameFreed + NUM_FRAMES_TO_WAIT_BEFORE_RELEASING_TO_OS < GFrameNumberRenderThread)
@@ -751,7 +751,7 @@ namespace VulkanRHI
 			for (int32 Index = 0; Index < FreePages.Num(); ++Index)
 #else
 			// Leave a page not freed to avoid potential hitching
-			for (int32 Index = 1; Index < FreePages.Num(); ++Index)
+			for (int32 Index = (bImmediately ? 0 : 1); Index < FreePages.Num(); ++Index)
 #endif
 			{
 				FOldResourceHeapPage* Page = FreePages[Index];
@@ -806,7 +806,7 @@ namespace VulkanRHI
 				UE_LOG(LogVulkanRHI, Display, TEXT("\t\t%d: ID %4d %4d suballocs, %4d free chunks (%d used/%d free/%d max) DeviceMemory %p"), Index, UsedPages[Index]->GetID(), UsedPages[Index]->ResourceAllocations.Num(), UsedPages[Index]->FreeList.Num(), UsedPages[Index]->UsedSize, UsedPages[Index]->MaxSize - UsedPages[Index]->UsedSize, UsedPages[Index]->MaxSize, (void*)UsedPages[Index]->DeviceMemoryAllocation->GetHandle());
 			}
 
-			UE_LOG(LogVulkanRHI, Display, TEXT("%d Suballocations for Used/Total: %d/%d = %.2f%%"), NumSuballocations, SubAllocUsedMemory, SubAllocAllocatedMemory, 100.0f * (float)SubAllocUsedMemory / (float)SubAllocAllocatedMemory);
+			UE_LOG(LogVulkanRHI, Display, TEXT("%d Suballocations for Used/Total: %d/%d = %.2f%%"), NumSuballocations, SubAllocUsedMemory, SubAllocAllocatedMemory, SubAllocAllocatedMemory > 0 ? 100.0f * (float)SubAllocUsedMemory / (float)SubAllocAllocatedMemory : 0.0f);
 		};
 
 		DumpPages(UsedBufferPages, TEXT("Buffer"));
@@ -1352,8 +1352,8 @@ namespace VulkanRHI
 		}
 
 		UE_LOG(LogVulkanRHI, Display, TEXT("::Totals::"));
-		UE_LOG(LogVulkanRHI, Display, TEXT("Large Alloc Used/Max %d/%d %.2f%%"), UsedLargeTotal, AllocLargeTotal, 100.0f * (float)UsedLargeTotal / (float)AllocLargeTotal);
-		UE_LOG(LogVulkanRHI, Display, TEXT("Binned Alloc Used/Max %d/%d %.2f%%"), UsedBinnedTotal, AllocBinnedTotal, 100.0f * (float)UsedBinnedTotal / (float)AllocBinnedTotal);
+		UE_LOG(LogVulkanRHI, Display, TEXT("Large Alloc Used/Max %d/%d %.2f%%"), UsedLargeTotal, AllocLargeTotal, 100.0f * AllocLargeTotal > 0 ? (float)UsedLargeTotal / (float)AllocLargeTotal : 0.0f);
+		UE_LOG(LogVulkanRHI, Display, TEXT("Binned Alloc Used/Max %d/%d %.2f%%"), UsedBinnedTotal, AllocBinnedTotal, AllocBinnedTotal > 0 ? 100.0f * (float)UsedBinnedTotal / (float)AllocBinnedTotal : 0.0f);
 	}
 #endif
 
diff --git a/Engine/Source/Runtime/VulkanRHI/Private/VulkanPipeline.cpp b/Engine/Source/Runtime/VulkanRHI/Private/VulkanPipeline.cpp
--- a/Engine/Source/Runtime/VulkanRHI/Private/VulkanPipeline.cpp
+++ b/Engine/Source/Runtime/VulkanRHI/Private/VulkanPipeline.cpp
@@ -322,7 +322,7 @@ bool FVulkanPipelineStateCacheManager::Load(const TArray<FString>& CacheFilename
 		}
 		else
 		{
-			UE_LOG(LogVulkanRHI, Error, TEXT("FVulkanPipelineStateCacheManager: Failed to load binary pipeline cache %s"), *BinaryCacheFilename);
+			UE_LOG(LogVulkanRHI, Display, TEXT("FVulkanPipelineStateCacheManager: Binary pipeline cache '%s' not found."), *BinaryCacheFilename);
 		}
 	}
 
diff --git a/Engine/Source/Runtime/VulkanRHI/Private/VulkanRHI.cpp b/Engine/Source/Runtime/VulkanRHI/Private/VulkanRHI.cpp
--- a/Engine/Source/Runtime/VulkanRHI/Private/VulkanRHI.cpp
+++ b/Engine/Source/Runtime/VulkanRHI/Private/VulkanRHI.cpp
@@ -713,8 +713,6 @@ void FVulkanCommandListContext::RHIBeginFrame()
 	check(IsImmediate());
 	RHIPrivateBeginFrame();
 
-	FVulkanCmdBuffer* CmdBuffer = CommandBufferManager->HasPendingUploadCmdBuffer() ? CommandBufferManager->GetUploadCmdBuffer() : CommandBufferManager->GetActiveCmdBuffer();
-
 	extern uint32 GVulkanRHIDeletionFrameNumber;
 	++GVulkanRHIDeletionFrameNumber;
 
diff --git a/Engine/Source/Runtime/VulkanRHI/Private/VulkanRenderTarget.cpp b/Engine/Source/Runtime/VulkanRHI/Private/VulkanRenderTarget.cpp
--- a/Engine/Source/Runtime/VulkanRHI/Private/VulkanRenderTarget.cpp
+++ b/Engine/Source/Runtime/VulkanRHI/Private/VulkanRenderTarget.cpp
@@ -577,7 +577,7 @@ void FVulkanCommandListContext::RHISetRenderTargetsAndClear(const FRHISetRenderT
 		Framebuffer = TransitionAndLayoutManager.GetOrCreateFramebuffer(*Device, RenderTargetsInfo, RTLayout, RenderPass);
 	}
 
-	if (Framebuffer == TransitionAndLayoutManager.CurrentFramebuffer && IsCompatibleRenderPass(TransitionAndLayoutManager.CurrentRenderPass, RenderPass))
+	if (Framebuffer == TransitionAndLayoutManager.CurrentFramebuffer && RenderPass != nullptr && IsCompatibleRenderPass(TransitionAndLayoutManager.CurrentRenderPass, RenderPass))
 	{
 		return;
 	}
@@ -1720,7 +1720,9 @@ FVulkanRenderTargetLayout::FVulkanRenderTargetLayout(FVulkanDevice& InDevice, co
 			CompatibleHashInfo.Formats[NumColorAttachments] = CurrDesc.format;
 			FullHashInfo.LoadOps[NumColorAttachments] = CurrDesc.loadOp;
 			FullHashInfo.StoreOps[NumColorAttachments] = CurrDesc.storeOp;
+#if VULKAN_USE_REAL_RENDERPASS_COMPATIBILITY
 			FullHashInfo.InitialLayout[NumColorAttachments] = VK_IMAGE_LAYOUT_COLOR_ATTACHMENT_OPTIMAL;
+#endif
 			++CompatibleHashInfo.NumAttachments;
 
 			++NumAttachmentDescriptions;
@@ -1774,7 +1776,9 @@ FVulkanRenderTargetLayout::FVulkanRenderTargetLayout(FVulkanDevice& InDevice, co
 		FullHashInfo.LoadOps[MaxSimultaneousRenderTargets + 1] = CurrDesc.stencilLoadOp;
 		FullHashInfo.StoreOps[MaxSimultaneousRenderTargets] = CurrDesc.storeOp;
 		FullHashInfo.StoreOps[MaxSimultaneousRenderTargets + 1] = CurrDesc.stencilStoreOp;
+#if VULKAN_USE_REAL_RENDERPASS_COMPATIBILITY
 		FullHashInfo.InitialLayout[MaxSimultaneousRenderTargets] = DepthStencilLayout;
+#endif
 		CompatibleHashInfo.Formats[MaxSimultaneousRenderTargets] = CurrDesc.format;
 
 		++NumAttachmentDescriptions;
@@ -1886,7 +1890,9 @@ FVulkanRenderTargetLayout::FVulkanRenderTargetLayout(FVulkanDevice& InDevice, co
 
 		CompatibleHashInfo.Formats[NumColorAttachments] = CurrDesc.format;
 		FullHashInfo.LoadOps[NumColorAttachments] = CurrDesc.loadOp;
+#if VULKAN_USE_REAL_RENDERPASS_COMPATIBILITY
 		FullHashInfo.InitialLayout[NumColorAttachments] = VK_IMAGE_LAYOUT_COLOR_ATTACHMENT_OPTIMAL;
+#endif
 		FullHashInfo.StoreOps[NumColorAttachments] = CurrDesc.storeOp;
 		++CompatibleHashInfo.NumAttachments;
 
@@ -1930,7 +1936,9 @@ FVulkanRenderTargetLayout::FVulkanRenderTargetLayout(FVulkanDevice& InDevice, co
 		FullHashInfo.LoadOps[MaxSimultaneousRenderTargets + 1] = CurrDesc.stencilLoadOp;
 		FullHashInfo.StoreOps[MaxSimultaneousRenderTargets] = CurrDesc.storeOp;
 		FullHashInfo.StoreOps[MaxSimultaneousRenderTargets + 1] = CurrDesc.stencilStoreOp;
+#if VULKAN_USE_REAL_RENDERPASS_COMPATIBILITY
 		FullHashInfo.InitialLayout[MaxSimultaneousRenderTargets] = DepthStencilLayout;
+#endif
 		CompatibleHashInfo.Formats[MaxSimultaneousRenderTargets] = CurrDesc.format;
 
 		++NumAttachmentDescriptions;
@@ -2023,7 +2031,9 @@ FVulkanRenderTargetLayout::FVulkanRenderTargetLayout(const FGraphicsPipelineStat
 			CompatibleHashInfo.Formats[NumColorAttachments] = CurrDesc.format;
 			FullHashInfo.LoadOps[NumColorAttachments] = CurrDesc.loadOp;
 			FullHashInfo.StoreOps[NumColorAttachments] = CurrDesc.storeOp;
+#if VULKAN_USE_REAL_RENDERPASS_COMPATIBILITY
 			FullHashInfo.InitialLayout[NumColorAttachments] = VK_IMAGE_LAYOUT_COLOR_ATTACHMENT_OPTIMAL;
+#endif
 			++CompatibleHashInfo.NumAttachments;
 
 			++NumAttachmentDescriptions;
@@ -2077,7 +2087,9 @@ FVulkanRenderTargetLayout::FVulkanRenderTargetLayout(const FGraphicsPipelineStat
 		FullHashInfo.LoadOps[MaxSimultaneousRenderTargets + 1] = CurrDesc.stencilLoadOp;
 		FullHashInfo.StoreOps[MaxSimultaneousRenderTargets] = CurrDesc.storeOp;
 		FullHashInfo.StoreOps[MaxSimultaneousRenderTargets + 1] = CurrDesc.stencilStoreOp;
+#if VULKAN_USE_REAL_RENDERPASS_COMPATIBILITY
 		FullHashInfo.InitialLayout[MaxSimultaneousRenderTargets] = VK_IMAGE_LAYOUT_DEPTH_STENCIL_ATTACHMENT_OPTIMAL;
+#endif
 		CompatibleHashInfo.Formats[MaxSimultaneousRenderTargets] = CurrDesc.format;
 
 		++NumAttachmentDescriptions;
diff --git a/Engine/Source/Runtime/VulkanRHI/Private/VulkanTexture.cpp b/Engine/Source/Runtime/VulkanRHI/Private/VulkanTexture.cpp
--- a/Engine/Source/Runtime/VulkanRHI/Private/VulkanTexture.cpp
+++ b/Engine/Source/Runtime/VulkanRHI/Private/VulkanTexture.cpp
@@ -797,8 +797,7 @@ void FVulkanSurface::InitialClear(FVulkanCommandListContext& Context,const FClea
 	}
 
 	VkImageLayout FinalLayout = Barrier.GetDestLayout(BarrierIndex);
-	VkImageLayout ImageLayout = Context.FindOrAddLayout(Image, FinalLayout);
-	ensure(FinalLayout == ImageLayout);
+	Context.FindOrAddLayoutRW(Image, FinalLayout) = FinalLayout;
 }
 
 /*-----------------------------------------------------------------------------
@@ -1526,6 +1525,7 @@ FVulkanTextureBase::FVulkanTextureBase(FVulkanDevice& Device, VkImageViewType Re
 
 	if (!CreateInfo.BulkData)
 	{
+		Device.GetImmediateContext().FindOrAddLayout(Surface.Image, VK_IMAGE_LAYOUT_UNDEFINED);
 		return;
 	}
 
