DR - vk - Support for VK_NV_device_diagnostic_checkpoints

diff --git a/Engine/Source/Runtime/VulkanRHI/Private/Linux/VulkanLinuxPlatform.cpp b/Engine/Source/Runtime/VulkanRHI/Private/Linux/VulkanLinuxPlatform.cpp
--- a/Engine/Source/Runtime/VulkanRHI/Private/Linux/VulkanLinuxPlatform.cpp
+++ b/Engine/Source/Runtime/VulkanRHI/Private/Linux/VulkanLinuxPlatform.cpp
@@ -139,3 +139,29 @@ void FVulkanLinuxPlatform::CreateSurface(void* WindowHandle, VkInstance Instance
 		check(0);
 	}
 }
+
+void FVulkanLinuxPlatform::WriteCrashMarker(const FOptionalVulkanDeviceExtensions& OptionalExtensions, VkCommandBuffer CmdBuffer, VkBuffer DestBuffer, const TArrayView<uint32>& Entries, bool bAdding)
+{
+	ensure(Entries.Num() <= GMaxCrashBufferEntries);
+
+	if (OptionalExtensions.HasAMDBufferMarker)
+	{
+		// AMD API only allows updating one entry at a time. Assume buffer has entry 0 as num entries
+		VulkanDynamicAPI::vkCmdWriteBufferMarkerAMD(CmdBuffer, VK_PIPELINE_STAGE_BOTTOM_OF_PIPE_BIT, DestBuffer, 0, Entries.Num());
+		if (bAdding)
+		{
+			int32 LastIndex = Entries.Num() - 1;
+			// +1 size as entries start at index 1
+			VulkanDynamicAPI::vkCmdWriteBufferMarkerAMD(CmdBuffer, VK_PIPELINE_STAGE_BOTTOM_OF_PIPE_BIT, DestBuffer, (1 + LastIndex) * sizeof(uint32), Entries[LastIndex]);
+		}
+	}
+	else if (OptionalExtensions.HasNVDiagnosticCheckpoints)
+	{
+		if (bAdding)
+		{
+			int32 LastIndex = Entries.Num() - 1;
+			uint32 Value = Entries[LastIndex];
+			VulkanDynamicAPI::vkCmdSetCheckpointNV(CmdBuffer, (void*)(size_t)Value);
+		}
+	}
+}
diff --git a/Engine/Source/Runtime/VulkanRHI/Private/Linux/VulkanLinuxPlatform.h b/Engine/Source/Runtime/VulkanRHI/Private/Linux/VulkanLinuxPlatform.h
--- a/Engine/Source/Runtime/VulkanRHI/Private/Linux/VulkanLinuxPlatform.h
+++ b/Engine/Source/Runtime/VulkanRHI/Private/Linux/VulkanLinuxPlatform.h
@@ -10,8 +10,10 @@
 #define VULKAN_SHOULD_DEBUG_IN_DEVELOPMENT			1
 #define VULKAN_SHOULD_ENABLE_DRAW_MARKERS			(UE_BUILD_DEBUG || UE_BUILD_DEVELOPMENT)
 #define VULKAN_SIGNAL_UNIMPLEMENTED()				checkf(false, TEXT("Unimplemented vulkan functionality: %s"), __PRETTY_FUNCTION__)
+#define	VULKAN_SUPPORTS_DEDICATED_ALLOCATION		0
+#define VULKAN_SUPPORTS_AMD_BUFFER_MARKER			1
+#define VULKAN_SUPPORTS_NV_DIAGNOSTIC_CHECKPOINT	1
 
-#define	VULKAN_SUPPORTS_DEDICATED_ALLOCATION				0
 
 #define ENUM_VK_ENTRYPOINTS_PLATFORM_BASE(EnumMacro) \
 	EnumMacro(PFN_vkGetPhysicalDeviceProperties2KHR, vkGetPhysicalDeviceProperties2KHR) \
@@ -52,6 +54,8 @@ public:
 #endif
 	}
 
+	static void WriteCrashMarker(const FOptionalVulkanDeviceExtensions& OptionalExtensions, VkCommandBuffer CmdBuffer, VkBuffer DestBuffer, const TArrayView<uint32>& Entries, bool bAdding);
+
 protected:
 	static void* VulkanLib;
 	static bool bAttemptedLoad;
diff --git a/Engine/Source/Runtime/VulkanRHI/Private/Lumin/VulkanLuminPlatform.h b/Engine/Source/Runtime/VulkanRHI/Private/Lumin/VulkanLuminPlatform.h
--- a/Engine/Source/Runtime/VulkanRHI/Private/Lumin/VulkanLuminPlatform.h
+++ b/Engine/Source/Runtime/VulkanRHI/Private/Lumin/VulkanLuminPlatform.h
@@ -11,6 +11,7 @@
 #define VULKAN_USE_IMAGE_ACQUIRE_FENCES				0
 #define VULKAN_SUPPORTS_COLOR_CONVERSIONS			1
 #define VULKAN_SUPPORTS_GEOMETRY_SHADERS			1
+#define VULKAN_SUPPORTS_NV_DIAGNOSTIC_CHECKPOINT	1
 
 
 #define ENUM_VK_ENTRYPOINTS_PLATFORM_BASE(EnumMacro)
diff --git a/Engine/Source/Runtime/VulkanRHI/Private/VulkanDevice.cpp b/Engine/Source/Runtime/VulkanRHI/Private/VulkanDevice.cpp
--- a/Engine/Source/Runtime/VulkanRHI/Private/VulkanDevice.cpp
+++ b/Engine/Source/Runtime/VulkanRHI/Private/VulkanDevice.cpp
@@ -754,27 +754,39 @@ void FVulkanDevice::InitGPU(int32 DeviceIndex)
 
 	StagingManager.Init(this);
 
-#if VULKAN_SUPPORTS_AMD_BUFFER_MARKER
-	if (GGPUCrashDebuggingEnabled && OptionalDeviceExtensions.HasAMDBufferMarker)
+#if VULKAN_SUPPORTS_GPU_CRASH_DUMPS
+	if (GGPUCrashDebuggingEnabled)
 	{
-		VkBufferCreateInfo CreateInfo;
-		ZeroVulkanStruct(CreateInfo, VK_STRUCTURE_TYPE_BUFFER_CREATE_INFO);
-		CreateInfo.size = GMaxCrashBufferEntries * sizeof(uint32_t);
-		CreateInfo.usage = VK_BUFFER_USAGE_TRANSFER_DST_BIT;
-		VERIFYVULKANRESULT(VulkanRHI::vkCreateBuffer(Device, &CreateInfo, VULKAN_CPU_ALLOCATOR, &CrashMarker.Buffer));
-
-		VkMemoryRequirements MemReq;
-		FMemory::Memzero(MemReq);
-		VulkanRHI::vkGetBufferMemoryRequirements(Device, CrashMarker.Buffer, &MemReq);
-
-		CrashMarker.Allocation = MemoryManager.Alloc(false, CreateInfo.size, MemReq.memoryTypeBits, VK_MEMORY_PROPERTY_HOST_VISIBLE_BIT |
-			VK_MEMORY_PROPERTY_HOST_COHERENT_BIT, nullptr, __FILE__, __LINE__);
-
-		uint32* Entry = (uint32*)CrashMarker.Allocation->Map(VK_WHOLE_SIZE, 0);
-		check(Entry);
-		// Start with 0 entries
-		*Entry = 0;
-		VERIFYVULKANRESULT(VulkanRHI::vkBindBufferMemory(Device, CrashMarker.Buffer, CrashMarker.Allocation->GetHandle(), 0));
+		if (OptionalDeviceExtensions.HasAMDBufferMarker)
+		{
+			VkBufferCreateInfo CreateInfo;
+			ZeroVulkanStruct(CreateInfo, VK_STRUCTURE_TYPE_BUFFER_CREATE_INFO);
+			CreateInfo.size = GMaxCrashBufferEntries * sizeof(uint32_t);
+			CreateInfo.usage = VK_BUFFER_USAGE_TRANSFER_DST_BIT;
+			VERIFYVULKANRESULT(VulkanRHI::vkCreateBuffer(Device, &CreateInfo, VULKAN_CPU_ALLOCATOR, &CrashMarker.Buffer));
+
+			VkMemoryRequirements MemReq;
+			FMemory::Memzero(MemReq);
+			VulkanRHI::vkGetBufferMemoryRequirements(Device, CrashMarker.Buffer, &MemReq);
+
+			CrashMarker.Allocation = MemoryManager.Alloc(false, CreateInfo.size, MemReq.memoryTypeBits, VK_MEMORY_PROPERTY_HOST_VISIBLE_BIT |
+				VK_MEMORY_PROPERTY_HOST_COHERENT_BIT, nullptr, __FILE__, __LINE__);
+
+			uint32* Entry = (uint32*)CrashMarker.Allocation->Map(VK_WHOLE_SIZE, 0);
+			check(Entry);
+			// Start with 0 entries
+			*Entry = 0;
+			VERIFYVULKANRESULT(VulkanRHI::vkBindBufferMemory(Device, CrashMarker.Buffer, CrashMarker.Allocation->GetHandle(), 0));
+		}
+		else if (OptionalDeviceExtensions.HasNVDiagnosticCheckpoints)
+		{
+			CrashMarker.Allocation = MemoryManager.Alloc(false, GMaxCrashBufferEntries * sizeof(uint32_t), UINT32_MAX, VK_MEMORY_PROPERTY_HOST_VISIBLE_BIT |
+				VK_MEMORY_PROPERTY_HOST_COHERENT_BIT, nullptr, __FILE__, __LINE__);
+			uint32* Entry = (uint32*)CrashMarker.Allocation->Map(VK_WHOLE_SIZE, 0);
+			check(Entry);
+			// Start with 0 entries
+			*Entry = 0;
+		}
 	}
 #endif
 
@@ -912,16 +924,24 @@ void FVulkanDevice::Destroy()
 	StagingManager.Deinit();
 
 
-#if VULKAN_SUPPORTS_AMD_BUFFER_MARKER
-	if (GGPUCrashDebuggingEnabled && OptionalDeviceExtensions.HasAMDBufferMarker)
+	if (GGPUCrashDebuggingEnabled)
 	{
-		CrashMarker.Allocation->Unmap();
-		VulkanRHI::vkDestroyBuffer(Device, CrashMarker.Buffer, VULKAN_CPU_ALLOCATOR);
-		CrashMarker.Buffer = VK_NULL_HANDLE;
+#if VULKAN_SUPPORTS_AMD_BUFFER_MARKER
+		if (OptionalDeviceExtensions.HasAMDBufferMarker)
+		{
+			CrashMarker.Allocation->Unmap();
+			VulkanRHI::vkDestroyBuffer(Device, CrashMarker.Buffer, VULKAN_CPU_ALLOCATOR);
+			CrashMarker.Buffer = VK_NULL_HANDLE;
 
-		MemoryManager.Free(CrashMarker.Allocation);
-	}
+			MemoryManager.Free(CrashMarker.Allocation);
+		}
 #endif
+#if VULKAN_SUPPORTS_NV_DIAGNOSTIC_CHECKPOINT
+		if (OptionalDeviceExtensions.HasNVDiagnosticCheckpoints)
+		{
+		}
+#endif
+	}
 
 	ResourceHeapManager.Deinit();
 
diff --git a/Engine/Source/Runtime/VulkanRHI/Private/VulkanDevice.h b/Engine/Source/Runtime/VulkanRHI/Private/VulkanDevice.h
--- a/Engine/Source/Runtime/VulkanRHI/Private/VulkanDevice.h
+++ b/Engine/Source/Runtime/VulkanRHI/Private/VulkanDevice.h
@@ -27,8 +27,14 @@ struct FOptionalVulkanDeviceExtensions
 	uint32 HasKHRDedicatedAllocation : 1;
 	uint32 HasEXTValidationCache : 1;
 	uint32 HasAMDBufferMarker : 1;
+	uint32 HasNVDiagnosticCheckpoints : 1;
 	uint32 HasGoogleDisplayTiming : 1;
 	uint32 HasYcbcrSampler : 1;
+
+	inline bool HasGPUCrashDumpExtensions() const
+	{
+		return HasAMDBufferMarker || HasNVDiagnosticCheckpoints;
+	}
 };
 
 class FVulkanDevice
@@ -267,7 +273,7 @@ public:
 		return OptionalDeviceExtensions;
 	}
 
-#if VULKAN_SUPPORTS_AMD_BUFFER_MARKER
+#if VULKAN_SUPPORTS_GPU_CRASH_DUMPS
 	VkBuffer GetCrashMarkerBuffer() const
 	{
 		return CrashMarker.Buffer;
@@ -337,7 +343,7 @@ private:
 	bool bAsyncComputeQueue = false;
 	bool bPresentOnComputeQueue = false;
 
-#if VULKAN_SUPPORTS_AMD_BUFFER_MARKER
+#if VULKAN_SUPPORTS_GPU_CRASH_DUMPS
 	struct
 	{
 		VkBuffer Buffer = VK_NULL_HANDLE;
diff --git a/Engine/Source/Runtime/VulkanRHI/Private/VulkanGenericPlatform.h b/Engine/Source/Runtime/VulkanRHI/Private/VulkanGenericPlatform.h
--- a/Engine/Source/Runtime/VulkanRHI/Private/VulkanGenericPlatform.h
+++ b/Engine/Source/Runtime/VulkanRHI/Private/VulkanGenericPlatform.h
@@ -5,6 +5,8 @@
 #include "PixelFormat.h"
 #include "Containers/ArrayView.h"
 
+struct FOptionalVulkanDeviceExtensions;
+
 // the platform interface, and empty implementations for platforms that don't need em
 class FVulkanGenericPlatform 
 {
@@ -67,7 +69,7 @@ public:
 
 	static bool RegisterGPUWork() { return true; }
 
-	static void WriteBufferMarkerAMD(VkCommandBuffer CmdBuffer, VkBuffer DestBuffer, const TArrayView<uint32>& Entries, bool bAdding) {}
+	static void WriteCrashMarker(const FOptionalVulkanDeviceExtensions& OptionalExtensions, VkCommandBuffer CmdBuffer, VkBuffer DestBuffer, const TArrayView<uint32>& Entries, bool bAdding) {}
 
 	// Allow the platform code to restrict the device features
 	static void RestrictEnabledPhysicalDeviceFeatures(VkPhysicalDeviceFeatures& InOutFeaturesToEnable)
diff --git a/Engine/Source/Runtime/VulkanRHI/Private/VulkanLayers.cpp b/Engine/Source/Runtime/VulkanRHI/Private/VulkanLayers.cpp
--- a/Engine/Source/Runtime/VulkanRHI/Private/VulkanLayers.cpp
+++ b/Engine/Source/Runtime/VulkanRHI/Private/VulkanLayers.cpp
@@ -688,6 +688,10 @@ void FVulkanDevice::ParseOptionalDeviceExtensions(const TArray<const ANSICHAR *>
 	OptionalDeviceExtensions.HasAMDBufferMarker = HasExtension(VK_AMD_BUFFER_MARKER_EXTENSION_NAME);
 #endif
 
+#if VULKAN_SUPPORTS_NV_DIAGNOSTIC_CHECKPOINT
+	OptionalDeviceExtensions.HasNVDiagnosticCheckpoints = HasExtension(VK_NV_DEVICE_DIAGNOSTIC_CHECKPOINTS_EXTENSION_NAME);
+#endif
+
 #if VULKAN_SUPPORTS_GOOGLE_DISPLAY_TIMING
 	OptionalDeviceExtensions.HasGoogleDisplayTiming = HasExtension(VK_GOOGLE_DISPLAY_TIMING_EXTENSION_NAME);
 #endif
diff --git a/Engine/Source/Runtime/VulkanRHI/Private/VulkanRHI.cpp b/Engine/Source/Runtime/VulkanRHI/Private/VulkanRHI.cpp
--- a/Engine/Source/Runtime/VulkanRHI/Private/VulkanRHI.cpp
+++ b/Engine/Source/Runtime/VulkanRHI/Private/VulkanRHI.cpp
@@ -835,7 +835,7 @@ void FVulkanCommandListContext::RHIPushEvent(const TCHAR* Name, FColor Color)
 	}
 #endif
 
-#if VULKAN_SUPPORTS_AMD_BUFFER_MARKER
+#if VULKAN_SUPPORTS_GPU_CRASH_DUMPS
 	if (GpuProfiler.bTrackingGPUCrashData)
 	{
 		GpuProfiler.PushMarkerForCrash(GetCommandBufferManager()->GetActiveCmdBuffer()->GetHandle(), Device->GetCrashMarkerBuffer(), Name);
@@ -869,7 +869,7 @@ void FVulkanCommandListContext::RHIPopEvent()
 	}
 #endif
 
-#if VULKAN_SUPPORTS_AMD_BUFFER_MARKER
+#if VULKAN_SUPPORTS_GPU_CRASH_DUMPS
 	if (GpuProfiler.bTrackingGPUCrashData)
 	{
 		GpuProfiler.PopMarkerForCrash(GetCommandBufferManager()->GetActiveCmdBuffer()->GetHandle(), Device->GetCrashMarkerBuffer());
diff --git a/Engine/Source/Runtime/VulkanRHI/Private/VulkanUtil.cpp b/Engine/Source/Runtime/VulkanRHI/Private/VulkanUtil.cpp
--- a/Engine/Source/Runtime/VulkanRHI/Private/VulkanUtil.cpp
+++ b/Engine/Source/Runtime/VulkanRHI/Private/VulkanUtil.cpp
@@ -342,8 +342,8 @@ float FVulkanEventNode::GetTiming()
 
 void FVulkanGPUProfiler::BeginFrame()
 {
-#if VULKAN_SUPPORTS_AMD_BUFFER_MARKER
-	if (GGPUCrashDebuggingEnabled && Device->GetOptionalExtensions().HasAMDBufferMarker)
+#if VULKAN_SUPPORTS_GPU_CRASH_DUMPS
+	if (GGPUCrashDebuggingEnabled && Device->GetOptionalExtensions().HasGPUCrashDumpExtensions())
 	{
 		static auto* CrashCollectionEnableCvar = IConsoleManager::Get().FindTConsoleVariableDataInt(TEXT("r.gpucrash.collectionenable"));
 		static auto* CrashCollectionDataDepth = IConsoleManager::Get().FindTConsoleVariableDataInt(TEXT("r.gpucrash.datadepth"));
@@ -455,7 +455,7 @@ void FVulkanGPUProfiler::EndFrame()
 	}
 }
 
-#if VULKAN_SUPPORTS_AMD_BUFFER_MARKER
+#if VULKAN_SUPPORTS_GPU_CRASH_DUMPS
 void FVulkanGPUProfiler::PushMarkerForCrash(VkCommandBuffer CmdBuffer, VkBuffer DestBuffer, const TCHAR* Name)
 {
 	uint32 CRC = 0;
@@ -480,7 +480,7 @@ void FVulkanGPUProfiler::PushMarkerForCrash(VkCommandBuffer CmdBuffer, VkBuffer
 	}
 
 	PushPopStack.Push(CRC);
-	FVulkanPlatform::WriteBufferMarkerAMD(CmdBuffer, DestBuffer, TArrayView<uint32>(PushPopStack), true);
+	FVulkanPlatform::WriteCrashMarker(Device->GetOptionalExtensions(), CmdBuffer, DestBuffer, TArrayView<uint32>(PushPopStack), true);
 }
 
 void FVulkanGPUProfiler::PopMarkerForCrash(VkCommandBuffer CmdBuffer, VkBuffer DestBuffer)
@@ -488,19 +488,39 @@ void FVulkanGPUProfiler::PopMarkerForCrash(VkCommandBuffer CmdBuffer, VkBuffer D
 	if (PushPopStack.Num() > 0)
 	{
 		PushPopStack.Pop(false);
-		FVulkanPlatform::WriteBufferMarkerAMD(CmdBuffer, DestBuffer, TArrayView<uint32>(PushPopStack), false);
+		FVulkanPlatform::WriteCrashMarker(Device->GetOptionalExtensions(), CmdBuffer, DestBuffer, TArrayView<uint32>(PushPopStack), false);
 	}
 }
 
 void FVulkanGPUProfiler::DumpCrashMarkers(void* BufferData)
 {
-	uint32* Entries = (uint32*)BufferData;
-	uint32 NumCRCs = *Entries++;
-	for (uint32 Index = 0; Index < NumCRCs; ++Index)
+	if (Device->GetOptionalExtensions().HasAMDBufferMarker)
 	{
-		const FString* Frame = CachedStrings.Find(*Entries);
-		UE_LOG(LogVulkanRHI, Error, TEXT("[VK_AMD_buffer_info] %i: %s (CRC 0x%x)"), Index, Frame ? *(*Frame) : TEXT("<undefined>"), *Entries);
-		++Entries;
+		uint32* Entries = (uint32*)BufferData;
+		uint32 NumCRCs = *Entries++;
+		for (uint32 Index = 0; Index < NumCRCs; ++Index)
+		{
+			const FString* Frame = CachedStrings.Find(*Entries);
+			UE_LOG(LogVulkanRHI, Error, TEXT("[VK_AMD_buffer_info] %i: %s (CRC 0x%x)"), Index, Frame ? *(*Frame) : TEXT("<undefined>"), *Entries);
+			++Entries;
+		}
+	}
+	else if (Device->GetOptionalExtensions().HasNVDiagnosticCheckpoints)
+	{
+		TArray<VkCheckpointDataNV> Data;
+		uint32 Num = 0;
+		VkQueue QueueHandle = Device->GetGraphicsQueue()->GetHandle();
+		VulkanDynamicAPI::vkGetQueueCheckpointDataNV(QueueHandle, &Num, nullptr);
+		Data.AddUninitialized(Num);
+		VulkanDynamicAPI::vkGetQueueCheckpointDataNV(QueueHandle, &Num, &Data[0]);
+		check(Num == Data.Num());
+		for (uint32 Index = 0; Index < Num; ++Index)
+		{
+			check(Data[Index].sType == VK_STRUCTURE_TYPE_CHECKPOINT_DATA_NV);
+			uint32 Value = (uint32)(size_t)Data[Index].pCheckpointMarker;
+			const FString* Frame = CachedStrings.Find(Value);
+			UE_LOG(LogVulkanRHI, Error, TEXT("[VK_NV_device_diagnostic_checkpoints] %i: Stage 0x%x, %s (CRC 0x%x)"), Index, Data[Index].stage, Frame ? *(*Frame) : TEXT("<undefined>"), Value);
+		}
 	}
 }
 #endif
@@ -613,12 +633,12 @@ namespace VulkanRHI
 		UE_LOG(LogVulkanRHI, Error, TEXT("%s failed, VkResult=%d\n at %s:%u \n with error %s"),
 			ANSI_TO_TCHAR(VkFunction), (int32)Result, ANSI_TO_TCHAR(Filename), Line, *ErrorString);
 
-#if VULKAN_SUPPORTS_AMD_BUFFER_MARKER
+#if VULKAN_SUPPORTS_GPU_CRASH_DUMPS
 		if (GIsGPUCrashed && GGPUCrashDebuggingEnabled)
 		{
 			FVulkanDynamicRHI* RHI = (FVulkanDynamicRHI*)GDynamicRHI;
 			FVulkanDevice* Device = RHI->GetDevice();
-			if (Device->GetOptionalExtensions().HasAMDBufferMarker)
+			if (Device->GetOptionalExtensions().HasGPUCrashDumpExtensions())
 			{
 				Device->GetImmediateContext().GetGPUProfiler().DumpCrashMarkers(Device->GetCrashMarkerMappedPointer());
 			}
diff --git a/Engine/Source/Runtime/VulkanRHI/Private/Windows/VulkanWindowsPlatform.cpp b/Engine/Source/Runtime/VulkanRHI/Private/Windows/VulkanWindowsPlatform.cpp
--- a/Engine/Source/Runtime/VulkanRHI/Private/Windows/VulkanWindowsPlatform.cpp
+++ b/Engine/Source/Runtime/VulkanRHI/Private/Windows/VulkanWindowsPlatform.cpp
@@ -2,6 +2,7 @@
 
 #include "VulkanWindowsPlatform.h"
 #include "../VulkanRHIPrivate.h"
+#include "../VulkanDevice.h"
 
 #include "Windows/AllowWindowsPlatformTypes.h"
 static HMODULE GVulkanDLLModule = nullptr;
@@ -121,7 +122,14 @@ void FVulkanWindowsPlatform::GetDeviceExtensions(TArray<const ANSICHAR*>& OutExt
 #endif
 	if (GGPUCrashDebuggingEnabled)
 	{
-		OutExtensions.Add(VK_AMD_BUFFER_MARKER_EXTENSION_NAME);
+		if (IsRHIDeviceAMD())
+		{
+			OutExtensions.Add(VK_AMD_BUFFER_MARKER_EXTENSION_NAME);
+		}
+		if (IsRHIDeviceNVIDIA())
+		{
+			OutExtensions.Add(VK_NV_DEVICE_DIAGNOSTIC_CHECKPOINTS_EXTENSION_NAME);
+		}
 	}
 
 #if VULKAN_SUPPORTS_COLOR_CONVERSIONS
@@ -148,15 +156,28 @@ bool FVulkanWindowsPlatform::SupportsDeviceLocalHostVisibleWithNoPenalty()
 }
 
 
-void FVulkanWindowsPlatform::WriteBufferMarkerAMD(VkCommandBuffer CmdBuffer, VkBuffer DestBuffer, const TArrayView<uint32>& Entries, bool bAdding)
+void FVulkanWindowsPlatform::WriteCrashMarker(const FOptionalVulkanDeviceExtensions& OptionalExtensions, VkCommandBuffer CmdBuffer, VkBuffer DestBuffer, const TArrayView<uint32>& Entries, bool bAdding)
 {
 	ensure(Entries.Num() <= GMaxCrashBufferEntries);
-	// AMD API only allows updating one entry at a time. Assume buffer has entry 0 as num entries
-	VulkanDynamicAPI::vkCmdWriteBufferMarkerAMD(CmdBuffer, VK_PIPELINE_STAGE_BOTTOM_OF_PIPE_BIT, DestBuffer, 0, Entries.Num());
-	if (bAdding)
+
+	if (OptionalExtensions.HasAMDBufferMarker)
 	{
-		int32 LastIndex = Entries.Num() - 1;
-		// +1 size as entries start at index 1
-		VulkanDynamicAPI::vkCmdWriteBufferMarkerAMD(CmdBuffer, VK_PIPELINE_STAGE_BOTTOM_OF_PIPE_BIT, DestBuffer, (1 + LastIndex) * sizeof(uint32), Entries[LastIndex]);
+		// AMD API only allows updating one entry at a time. Assume buffer has entry 0 as num entries
+		VulkanDynamicAPI::vkCmdWriteBufferMarkerAMD(CmdBuffer, VK_PIPELINE_STAGE_BOTTOM_OF_PIPE_BIT, DestBuffer, 0, Entries.Num());
+		if (bAdding)
+		{
+			int32 LastIndex = Entries.Num() - 1;
+			// +1 size as entries start at index 1
+			VulkanDynamicAPI::vkCmdWriteBufferMarkerAMD(CmdBuffer, VK_PIPELINE_STAGE_BOTTOM_OF_PIPE_BIT, DestBuffer, (1 + LastIndex) * sizeof(uint32), Entries[LastIndex]);
+		}
+	}
+	else if (OptionalExtensions.HasNVDiagnosticCheckpoints)
+	{
+		if (bAdding)
+		{
+			int32 LastIndex = Entries.Num() - 1;
+			uint32 Value = Entries[LastIndex];
+			VulkanDynamicAPI::vkCmdSetCheckpointNV(CmdBuffer, (void*)(size_t)Value);
+		}
 	}
 }
diff --git a/Engine/Source/Runtime/VulkanRHI/Private/Windows/VulkanWindowsPlatform.h b/Engine/Source/Runtime/VulkanRHI/Private/Windows/VulkanWindowsPlatform.h
--- a/Engine/Source/Runtime/VulkanRHI/Private/Windows/VulkanWindowsPlatform.h
+++ b/Engine/Source/Runtime/VulkanRHI/Private/Windows/VulkanWindowsPlatform.h
@@ -17,6 +17,9 @@
 #define VULKAN_SIGNAL_UNIMPLEMENTED()			checkf(false, TEXT("Unimplemented vulkan functionality: %s"), TEXT(__FUNCTION__))
 #define VULKAN_SUPPORTS_COLOR_CONVERSIONS		1
 #define	VULKAN_SUPPORTS_DEDICATED_ALLOCATION	0
+#define VULKAN_SUPPORTS_AMD_BUFFER_MARKER			1
+#define VULKAN_SUPPORTS_NV_DIAGNOSTIC_CHECKPOINT	1
+
 
 // 32-bit windows has warnings on custom mem mgr callbacks
 #define VULKAN_SHOULD_USE_LLM					(UE_BUILD_DEBUG || UE_BUILD_DEVELOPMENT) && !PLATFORM_32BITS
@@ -28,6 +31,8 @@
 	EnumMacro(PFN_vkGetPhysicalDeviceProperties2KHR, vkGetPhysicalDeviceProperties2KHR) \
 	EnumMacro(PFN_vkGetImageMemoryRequirements2KHR , vkGetImageMemoryRequirements2KHR) \
 	EnumMacro(PFN_vkCmdWriteBufferMarkerAMD, vkCmdWriteBufferMarkerAMD) \
+	EnumMacro(PFN_vkCmdSetCheckpointNV, vkCmdSetCheckpointNV) \
+	EnumMacro(PFN_vkGetQueueCheckpointDataNV, vkGetQueueCheckpointDataNV) \
 	EnumMacro(PFN_vkGetBufferMemoryRequirements2KHR , vkGetBufferMemoryRequirements2KHR)
 
 #define ENUM_VK_ENTRYPOINTS_OPTIONAL_PLATFORM_INSTANCE(EnumMacro) \
@@ -53,7 +58,7 @@ public:
 
 	static bool SupportsDeviceLocalHostVisibleWithNoPenalty();
 
-	static void WriteBufferMarkerAMD(VkCommandBuffer CmdBuffer, VkBuffer DestBuffer, const TArrayView<uint32>& Entries, bool bAdding);
+	static void WriteCrashMarker(const FOptionalVulkanDeviceExtensions& OptionalExtensions, VkCommandBuffer CmdBuffer, VkBuffer DestBuffer, const TArrayView<uint32>& Entries, bool bAdding);
 
 	// Some platforms only support real or non-real UBs, so this function can optimize it out
 	static bool UseRealUBsOptimization(bool bCodeHeaderUseRealUBs)
diff --git a/Engine/Source/Runtime/VulkanRHI/Public/VulkanConfiguration.h b/Engine/Source/Runtime/VulkanRHI/Public/VulkanConfiguration.h
--- a/Engine/Source/Runtime/VulkanRHI/Public/VulkanConfiguration.h
+++ b/Engine/Source/Runtime/VulkanRHI/Public/VulkanConfiguration.h
@@ -163,12 +163,16 @@
 	#define VULKAN_SUPPORTS_COLOR_CONVERSIONS					0
 #endif
 
-#ifdef VK_AMD_buffer_marker
-	#define VULKAN_SUPPORTS_AMD_BUFFER_MARKER					1
-#else
+#ifndef VULKAN_SUPPORTS_AMD_BUFFER_MARKER
 	#define VULKAN_SUPPORTS_AMD_BUFFER_MARKER					0
 #endif
 
+#ifndef VULKAN_SUPPORTS_NV_DIAGNOSTIC_CHECKPOINT
+	#define VULKAN_SUPPORTS_NV_DIAGNOSTIC_CHECKPOINT			0
+#endif
+
+#define VULKAN_SUPPORTS_GPU_CRASH_DUMPS							(VULKAN_SUPPORTS_AMD_BUFFER_MARKER || VULKAN_SUPPORTS_NV_DIAGNOSTIC_CHECKPOINT)
+
 #ifdef VK_EXT_debug_utils
 	#define VULKAN_SUPPORTS_DEBUG_UTILS							1
 #else
diff --git a/Engine/Source/Runtime/VulkanRHI/Public/VulkanUtil.h b/Engine/Source/Runtime/VulkanRHI/Public/VulkanUtil.h
--- a/Engine/Source/Runtime/VulkanRHI/Public/VulkanUtil.h
+++ b/Engine/Source/Runtime/VulkanRHI/Public/VulkanUtil.h
@@ -195,7 +195,7 @@ struct FVulkanGPUProfiler : public FGPUProfiler
 	FVulkanDevice* Device;
 	FVulkanCommandListContext* CmdContext;
 
-#if VULKAN_SUPPORTS_AMD_BUFFER_MARKER
+#if VULKAN_SUPPORTS_GPU_CRASH_DUMPS
 	void PushMarkerForCrash(VkCommandBuffer CmdBuffer, VkBuffer DestBuffer, const TCHAR* Name);
 	void PopMarkerForCrash(VkCommandBuffer CmdBuffer, VkBuffer DestBuffer);
 	void DumpCrashMarkers(void* BufferData);
