DR - vk - Compile fixes
#rnx

diff --git a/Engine/Source/Runtime/VulkanRHI/Private/Linux/VulkanLinuxPlatform.h b/Engine/Source/Runtime/VulkanRHI/Private/Linux/VulkanLinuxPlatform.h
--- a/Engine/Source/Runtime/VulkanRHI/Private/Linux/VulkanLinuxPlatform.h
+++ b/Engine/Source/Runtime/VulkanRHI/Private/Linux/VulkanLinuxPlatform.h
@@ -16,6 +16,9 @@
 
 
 #define ENUM_VK_ENTRYPOINTS_PLATFORM_BASE(EnumMacro) \
+	EnumMacro(PFN_vkCmdWriteBufferMarkerAMD, vkCmdWriteBufferMarkerAMD) \
+	EnumMacro(PFN_vkCmdSetCheckpointNV, vkCmdSetCheckpointNV) \
+	EnumMacro(PFN_vkGetQueueCheckpointDataNV, vkGetQueueCheckpointDataNV) \
 	EnumMacro(PFN_vkGetPhysicalDeviceProperties2KHR, vkGetPhysicalDeviceProperties2KHR) \
 	EnumMacro(PFN_vkGetImageMemoryRequirements2KHR , vkGetImageMemoryRequirements2KHR) \
 	EnumMacro(PFN_vkGetBufferMemoryRequirements2KHR , vkGetBufferMemoryRequirements2KHR)
diff --git a/Engine/Source/Runtime/VulkanRHI/Private/Lumin/VulkanLuminPlatform.h b/Engine/Source/Runtime/VulkanRHI/Private/Lumin/VulkanLuminPlatform.h
--- a/Engine/Source/Runtime/VulkanRHI/Private/Lumin/VulkanLuminPlatform.h
+++ b/Engine/Source/Runtime/VulkanRHI/Private/Lumin/VulkanLuminPlatform.h
@@ -17,6 +17,8 @@
 #define ENUM_VK_ENTRYPOINTS_PLATFORM_BASE(EnumMacro)
 
 #define ENUM_VK_ENTRYPOINTS_PLATFORM_INSTANCE(EnumMacro) \
+	EnumMacro(PFN_vkCmdSetCheckpointNV, vkCmdSetCheckpointNV) \
+	EnumMacro(PFN_vkGetQueueCheckpointDataNV, vkGetQueueCheckpointDataNV) \
 	EnumMacro(PFN_vkCreateSamplerYcbcrConversionKHR, vkCreateSamplerYcbcrConversionKHR) \
 	EnumMacro(PFN_vkDestroySamplerYcbcrConversionKHR, vkDestroySamplerYcbcrConversionKHR)
 
