DR - vk - Static analysis fix
#rnx

diff --git a/Engine/Source/Runtime/VulkanRHI/Public/VulkanResources.h b/Engine/Source/Runtime/VulkanRHI/Public/VulkanResources.h
--- a/Engine/Source/Runtime/VulkanRHI/Public/VulkanResources.h
+++ b/Engine/Source/Runtime/VulkanRHI/Public/VulkanResources.h
@@ -917,7 +917,6 @@ public:
 		}
 		else
 		{
-			ensure(State == EState::Unused);
 			return false;
 		}
 	}
