DR - vk - Compile fix
#rnx

diff --git a/Engine/Source/Runtime/VulkanRHI/Private/VulkanUtil.cpp b/Engine/Source/Runtime/VulkanRHI/Private/VulkanUtil.cpp
--- a/Engine/Source/Runtime/VulkanRHI/Private/VulkanUtil.cpp
+++ b/Engine/Source/Runtime/VulkanRHI/Private/VulkanUtil.cpp
@@ -494,6 +494,7 @@ void FVulkanGPUProfiler::PopMarkerForCrash(VkCommandBuffer CmdBuffer, VkBuffer D
 
 void FVulkanGPUProfiler::DumpCrashMarkers(void* BufferData)
 {
+#if VULKAN_SUPPORTS_AMD_BUFFER_MARKER
 	if (Device->GetOptionalExtensions().HasAMDBufferMarker)
 	{
 		uint32* Entries = (uint32*)BufferData;
@@ -505,22 +506,28 @@ void FVulkanGPUProfiler::DumpCrashMarkers(void* BufferData)
 			++Entries;
 		}
 	}
-	else if (Device->GetOptionalExtensions().HasNVDiagnosticCheckpoints)
+	else
+#endif
 	{
-		TArray<VkCheckpointDataNV> Data;
-		uint32 Num = 0;
-		VkQueue QueueHandle = Device->GetGraphicsQueue()->GetHandle();
-		VulkanDynamicAPI::vkGetQueueCheckpointDataNV(QueueHandle, &Num, nullptr);
-		Data.AddUninitialized(Num);
-		VulkanDynamicAPI::vkGetQueueCheckpointDataNV(QueueHandle, &Num, &Data[0]);
-		check(Num == Data.Num());
-		for (uint32 Index = 0; Index < Num; ++Index)
+#if VULKAN_SUPPORTS_NV_DIAGNOSTIC_CHECKPOINT
+		if (Device->GetOptionalExtensions().HasNVDiagnosticCheckpoints)
 		{
-			check(Data[Index].sType == VK_STRUCTURE_TYPE_CHECKPOINT_DATA_NV);
-			uint32 Value = (uint32)(size_t)Data[Index].pCheckpointMarker;
-			const FString* Frame = CachedStrings.Find(Value);
-			UE_LOG(LogVulkanRHI, Error, TEXT("[VK_NV_device_diagnostic_checkpoints] %i: Stage 0x%x, %s (CRC 0x%x)"), Index, Data[Index].stage, Frame ? *(*Frame) : TEXT("<undefined>"), Value);
+			TArray<VkCheckpointDataNV> Data;
+			uint32 Num = 0;
+			VkQueue QueueHandle = Device->GetGraphicsQueue()->GetHandle();
+			VulkanDynamicAPI::vkGetQueueCheckpointDataNV(QueueHandle, &Num, nullptr);
+			Data.AddUninitialized(Num);
+			VulkanDynamicAPI::vkGetQueueCheckpointDataNV(QueueHandle, &Num, &Data[0]);
+			check(Num == Data.Num());
+			for (uint32 Index = 0; Index < Num; ++Index)
+			{
+				check(Data[Index].sType == VK_STRUCTURE_TYPE_CHECKPOINT_DATA_NV);
+				uint32 Value = (uint32)(size_t)Data[Index].pCheckpointMarker;
+				const FString* Frame = CachedStrings.Find(Value);
+				UE_LOG(LogVulkanRHI, Error, TEXT("[VK_NV_device_diagnostic_checkpoints] %i: Stage 0x%x, %s (CRC 0x%x)"), Index, Data[Index].stage, Frame ? *(*Frame) : TEXT("<undefined>"), Value);
+			}
 		}
+#endif
 	}
 }
 #endif
