DR - Copy from 4476587, 4488158, 4488857
- Remove bogus ensure
- Added some missing Vulkan allocation tags
- Workaround for Linux on Nvidia cards not updating buffers
- Fix for random crashes toggling fullscreen on Vulkan games on Linux

diff --git a/Engine/Source/Runtime/VulkanRHI/Private/Linux/VulkanLinuxPlatform.cpp b/Engine/Source/Runtime/VulkanRHI/Private/Linux/VulkanLinuxPlatform.cpp
--- a/Engine/Source/Runtime/VulkanRHI/Private/Linux/VulkanLinuxPlatform.cpp
+++ b/Engine/Source/Runtime/VulkanRHI/Private/Linux/VulkanLinuxPlatform.cpp
@@ -129,6 +129,13 @@ void FVulkanLinuxPlatform::GetDeviceExtensions(TArray<const ANSICHAR*>& OutExten
 	OutExtensions.Add(VK_KHR_GET_MEMORY_REQUIREMENTS_2_EXTENSION_NAME);
 	OutExtensions.Add(VK_KHR_DEDICATED_ALLOCATION_EXTENSION_NAME);
 #endif
+
+	if (IsRHIDeviceNVIDIA())
+	{
+		//#todo-rco: Temporary workaround for some buffers not updating
+		extern FAutoConsoleVariableRef CVarVulkanWaitForIdleOnSubmit;
+		CVarVulkanWaitForIdleOnSubmit->Set(1);
+	}
 }
 
 void FVulkanLinuxPlatform::CreateSurface(void* WindowHandle, VkInstance Instance, VkSurfaceKHR* OutSurface)
diff --git a/Engine/Source/Runtime/VulkanRHI/Private/VulkanDevice.cpp b/Engine/Source/Runtime/VulkanRHI/Private/VulkanDevice.cpp
--- a/Engine/Source/Runtime/VulkanRHI/Private/VulkanDevice.cpp
+++ b/Engine/Source/Runtime/VulkanRHI/Private/VulkanDevice.cpp
@@ -638,7 +638,7 @@ VkSamplerYcbcrConversion FVulkanDevice::CreateSamplerColorConversion(const VkSam
 	else
 	{
 		VkSamplerYcbcrConversion NewConversion;
-		VERIFYVULKANRESULT(VulkanRHI::vkCreateSamplerYcbcrConversionKHR(GetInstanceHandle(), &CreateInfo, nullptr, &NewConversion));
+		VERIFYVULKANRESULT(VulkanRHI::vkCreateSamplerYcbcrConversionKHR(GetInstanceHandle(), &CreateInfo, VULKAN_CPU_ALLOCATOR, &NewConversion));
 		SamplerColorConversionMap.Add(CreateInfoHash, NewConversion);
 		return NewConversion;
 	}
@@ -885,7 +885,7 @@ void FVulkanDevice::Destroy()
 #if VULKAN_SUPPORTS_COLOR_CONVERSIONS
 	for (const auto& Pair : SamplerColorConversionMap)
 	{
-		VulkanRHI::vkDestroySamplerYcbcrConversionKHR(GetInstanceHandle(), Pair.Value, nullptr);
+		VulkanRHI::vkDestroySamplerYcbcrConversionKHR(GetInstanceHandle(), Pair.Value, VULKAN_CPU_ALLOCATOR);
 	}
 	SamplerColorConversionMap.Reset();
 #endif
@@ -939,6 +939,8 @@ void FVulkanDevice::Destroy()
 #if VULKAN_SUPPORTS_NV_DIAGNOSTIC_CHECKPOINT
 		if (OptionalDeviceExtensions.HasNVDiagnosticCheckpoints)
 		{
+			CrashMarker.Allocation->Unmap();
+			MemoryManager.Free(CrashMarker.Allocation);
 		}
 #endif
 	}
diff --git a/Engine/Source/Runtime/VulkanRHI/Private/VulkanPipeline.cpp b/Engine/Source/Runtime/VulkanRHI/Private/VulkanPipeline.cpp
--- a/Engine/Source/Runtime/VulkanRHI/Private/VulkanPipeline.cpp
+++ b/Engine/Source/Runtime/VulkanRHI/Private/VulkanPipeline.cpp
@@ -548,10 +548,10 @@ void FVulkanPipelineStateCacheManager::Save(const FString& CacheFilename, bool b
 		{
 			UE_LOG(LogVulkanRHI, Warning, TEXT("Failed to get Vulkan pipeline cache data."));
 
-			VulkanRHI::vkDestroyPipelineCache(Device->GetInstanceHandle(), PipelineCache, nullptr);
+			VulkanRHI::vkDestroyPipelineCache(Device->GetInstanceHandle(), PipelineCache, VULKAN_CPU_ALLOCATOR);
 			VkPipelineCacheCreateInfo PipelineCacheInfo;
 			ZeroVulkanStruct(PipelineCacheInfo, VK_STRUCTURE_TYPE_PIPELINE_CACHE_CREATE_INFO);
-			VERIFYVULKANRESULT(VulkanRHI::vkCreatePipelineCache(Device->GetInstanceHandle(), &PipelineCacheInfo, nullptr, &PipelineCache));
+			VERIFYVULKANRESULT(VulkanRHI::vkCreatePipelineCache(Device->GetInstanceHandle(), &PipelineCacheInfo, VULKAN_CPU_ALLOCATOR, &PipelineCache));
 		}
 		else
 		{
diff --git a/Engine/Source/Runtime/VulkanRHI/Private/VulkanQuery.cpp b/Engine/Source/Runtime/VulkanRHI/Private/VulkanQuery.cpp
--- a/Engine/Source/Runtime/VulkanRHI/Private/VulkanQuery.cpp
+++ b/Engine/Source/Runtime/VulkanRHI/Private/VulkanQuery.cpp
@@ -123,7 +123,7 @@ FVulkanQueryPool::FVulkanQueryPool(FVulkanDevice* InDevice, uint32 InMaxQueries,
 	PoolCreateInfo.queryType = QueryType;
 	PoolCreateInfo.queryCount = MaxQueries;
 
-	VERIFYVULKANRESULT(VulkanRHI::vkCreateQueryPool(Device->GetInstanceHandle(), &PoolCreateInfo, nullptr, &QueryPool));
+	VERIFYVULKANRESULT(VulkanRHI::vkCreateQueryPool(Device->GetInstanceHandle(), &PoolCreateInfo, VULKAN_CPU_ALLOCATOR, &QueryPool));
 
 	QueryOutput.AddZeroed(MaxQueries);
 }
@@ -131,7 +131,7 @@ FVulkanQueryPool::FVulkanQueryPool(FVulkanDevice* InDevice, uint32 InMaxQueries,
 FVulkanQueryPool::~FVulkanQueryPool()
 {
 	DEC_DWORD_STAT(STAT_VulkanNumQueryPools);
-	VulkanRHI::vkDestroyQueryPool(Device->GetInstanceHandle(), QueryPool, nullptr);
+	VulkanRHI::vkDestroyQueryPool(Device->GetInstanceHandle(), QueryPool, VULKAN_CPU_ALLOCATOR);
 	QueryPool = VK_NULL_HANDLE;
 	check(QueryPool == VK_NULL_HANDLE);
 }
diff --git a/Engine/Source/Runtime/VulkanRHI/Private/VulkanQueue.cpp b/Engine/Source/Runtime/VulkanRHI/Private/VulkanQueue.cpp
--- a/Engine/Source/Runtime/VulkanRHI/Private/VulkanQueue.cpp
+++ b/Engine/Source/Runtime/VulkanRHI/Private/VulkanQueue.cpp
@@ -9,7 +9,7 @@
 #include "VulkanMemory.h"
 
 int32 GWaitForIdleOnSubmit = 0;
-static FAutoConsoleVariableRef CVarVulkanWaitForIdleOnSubmit(
+FAutoConsoleVariableRef CVarVulkanWaitForIdleOnSubmit(
 	TEXT("r.Vulkan.WaitForIdleOnSubmit"),
 	GWaitForIdleOnSubmit,
 	TEXT("Waits for the GPU to be idle on every submit. Useful for tracking GPU hangs.\n")
diff --git a/Engine/Source/Runtime/VulkanRHI/Private/VulkanSwapChain.cpp b/Engine/Source/Runtime/VulkanRHI/Private/VulkanSwapChain.cpp
--- a/Engine/Source/Runtime/VulkanRHI/Private/VulkanSwapChain.cpp
+++ b/Engine/Source/Runtime/VulkanRHI/Private/VulkanSwapChain.cpp
@@ -475,6 +475,10 @@ void FVulkanSwapChain::Destroy()
 {
 	check(FVulkanPlatform::SupportsStandardSwapchain());
 
+	// We could be responding to an OUT_OF_DATE event and the GPU might not be done with swapchain image, so wait for idle.
+	// Alternatively could also check on the fence(s) for the image(s) from the swapchain but then timing out/waiting could become an issue.
+	Device.WaitUntilIdle();
+
 	VulkanRHI::vkDestroySwapchainKHR(Device.GetInstanceHandle(), SwapChain, VULKAN_CPU_ALLOCATOR);
 	SwapChain = VK_NULL_HANDLE;
 
diff --git a/Engine/Source/Runtime/VulkanRHI/Private/VulkanTexture.cpp b/Engine/Source/Runtime/VulkanRHI/Private/VulkanTexture.cpp
--- a/Engine/Source/Runtime/VulkanRHI/Private/VulkanTexture.cpp
+++ b/Engine/Source/Runtime/VulkanRHI/Private/VulkanTexture.cpp
@@ -1224,9 +1224,6 @@ void FVulkanDynamicRHI::InternalUpdateTexture2D(bool bFromRenderingThread, FText
 	//SubresourceRange.baseArrayLayer = 0;
 	SubresourceRange.layerCount = 1;
 
-	ensure(UpdateRegion.SrcX == 0);
-	ensure(UpdateRegion.SrcY == 0);
-
 	uint8* RESTRICT DestData = (uint8*)Memory;
 	uint8* RESTRICT SourceRowData = (uint8*)SourceData;
 	for (uint32 Height = 0; Height < NumBlocksY; ++Height)
