DR - vk - Flush log on device lost
#rnx

diff --git a/Engine/Source/Runtime/VulkanRHI/Private/VulkanDevice.cpp b/Engine/Source/Runtime/VulkanRHI/Private/VulkanDevice.cpp
--- a/Engine/Source/Runtime/VulkanRHI/Private/VulkanDevice.cpp
+++ b/Engine/Source/Runtime/VulkanRHI/Private/VulkanDevice.cpp
@@ -718,7 +718,7 @@ bool FVulkanDevice::QueryGPU(int32 DeviceIndex)
 		return Info;
 	};
 
-	UE_LOG(LogVulkanRHI, Display, TEXT("Initializing Device %d: %s"), DeviceIndex, ANSI_TO_TCHAR(GpuProps.deviceName));
+	UE_LOG(LogVulkanRHI, Display, TEXT("Device %d: %s"), DeviceIndex, ANSI_TO_TCHAR(GpuProps.deviceName));
 	UE_LOG(LogVulkanRHI, Display, TEXT("- API 0x%x Driver 0x%x VendorId 0x%x"), GpuProps.apiVersion, GpuProps.driverVersion, GpuProps.vendorID);
 	UE_LOG(LogVulkanRHI, Display, TEXT("- DeviceID 0x%x Type %s"), GpuProps.deviceID, *GetDeviceTypeString());
 	UE_LOG(LogVulkanRHI, Display, TEXT("- Max Descriptor Sets Bound %d Timestamps %d"), GpuProps.limits.maxBoundDescriptorSets, GpuProps.limits.timestampComputeAndGraphics);
diff --git a/Engine/Source/Runtime/VulkanRHI/Private/VulkanUtil.cpp b/Engine/Source/Runtime/VulkanRHI/Private/VulkanUtil.cpp
--- a/Engine/Source/Runtime/VulkanRHI/Private/VulkanUtil.cpp
+++ b/Engine/Source/Runtime/VulkanRHI/Private/VulkanUtil.cpp
@@ -9,7 +9,8 @@
 #include "VulkanPendingState.h"
 #include "VulkanContext.h"
 #include "VulkanMemory.h"
-#include "PipelineStateCache.h"
+#include "Misc/OutputDeviceRedirector.h"
+
 
 extern CORE_API bool GIsGPUCrashed;
 
@@ -499,6 +500,8 @@ void FVulkanGPUProfiler::DumpCrashMarkers(void* BufferData)
 				const FString* Frame = CachedStrings.Find(Value);
 				UE_LOG(LogVulkanRHI, Error, TEXT("[VK_NV_device_diagnostic_checkpoints] %i: Stage 0x%x, %s (CRC 0x%x)"), Index, Data[Index].stage, Frame ? *(*Frame) : TEXT("<undefined>"), Value);
 			}
+			GLog->PanicFlushThreadedLogs();
+			GLog->Flush();
 		}
 #endif
 	}
