DR - vk - Fix for integrated gpus
#rnx

diff --git a/Engine/Source/Runtime/VulkanRHI/Private/VulkanRHI.cpp b/Engine/Source/Runtime/VulkanRHI/Private/VulkanRHI.cpp
--- a/Engine/Source/Runtime/VulkanRHI/Private/VulkanRHI.cpp
+++ b/Engine/Source/Runtime/VulkanRHI/Private/VulkanRHI.cpp
@@ -458,12 +458,13 @@ void FVulkanDynamicRHI::SelectAndInitDevice()
 	FVulkanDevice* HmdDevice = nullptr;
 	uint32 HmdDeviceIndex = 0;
 #endif
-	struct FDiscreteDevice
+	struct FDeviceInfo
 	{
 		FVulkanDevice* Device;
 		uint32 DeviceIndex;
 	};
-	TArray<FDiscreteDevice> DiscreteDevices;
+	TArray<FDeviceInfo> DiscreteDevices;
+	TArray<FDeviceInfo> IntegratedDevices;
 
 #if VULKAN_ENABLE_DESKTOP_HMD_SUPPORT
 	// Allow HMD to override which graphics adapter is chosen, so we pick the adapter where the HMD is connected
@@ -491,6 +492,10 @@ void FVulkanDynamicRHI::SelectAndInitDevice()
 		{
 			DiscreteDevices.Add({NewDevice, Index});
 		}
+		else
+		{
+			IntegratedDevices.Add({NewDevice, Index});
+		}
 	}
 
 	uint32 DeviceIndex = -1;
@@ -503,6 +508,9 @@ void FVulkanDynamicRHI::SelectAndInitDevice()
 	}
 #endif
 
+	// Add all integrated to the end of the list
+	DiscreteDevices.Append(IntegratedDevices);
+
 	if (DeviceIndex == -1)
 	{
 		if (DiscreteDevices.Num() > 0)
