DR - Copy from 4496258 and 4496977
- Fix for mesh paint
- Fix for Linux crash
#rnx

diff --git a/Engine/Source/Runtime/VulkanRHI/Private/VulkanRHI.cpp b/Engine/Source/Runtime/VulkanRHI/Private/VulkanRHI.cpp
--- a/Engine/Source/Runtime/VulkanRHI/Private/VulkanRHI.cpp
+++ b/Engine/Source/Runtime/VulkanRHI/Private/VulkanRHI.cpp
@@ -515,12 +515,13 @@ void FVulkanDynamicRHI::SelectAndInitDevice()
 	{
 		if (DiscreteDevices.Num() > 0)
 		{
-			if (DiscreteDevices.Num() > 1)
+			int32 PreferredVendor = PreferAdapterVendor();
+			if (DiscreteDevices.Num() > 1 && PreferredVendor != -1)
 			{
 				// Check for preferred
 				for (int32 Index = 0; Index < DiscreteDevices.Num(); ++Index)
 				{
-					if (DiscreteDevices[Index].Device->GpuProps.vendorID == PreferAdapterVendor())
+					if (DiscreteDevices[Index].Device->GpuProps.vendorID == PreferredVendor)
 					{
 						DeviceIndex = DiscreteDevices[Index].DeviceIndex;
 						Device = DiscreteDevices[Index].Device;
@@ -537,13 +538,11 @@ void FVulkanDynamicRHI::SelectAndInitDevice()
 		}
 		else
 		{
-			Device = Devices[0];
+			checkf(0, TEXT("No devices found!"));
 			DeviceIndex = 0;
 		}
 	}
 
-	check(Device);
-
 	const VkPhysicalDeviceProperties& Props = Device->GetDeviceProperties();
 	GRHIVendorId = Props.vendorID;
 	GRHIAdapterName = ANSI_TO_TCHAR(Props.deviceName);
diff --git a/Engine/Source/Runtime/VulkanRHI/Private/VulkanRenderTarget.cpp b/Engine/Source/Runtime/VulkanRHI/Private/VulkanRenderTarget.cpp
--- a/Engine/Source/Runtime/VulkanRHI/Private/VulkanRenderTarget.cpp
+++ b/Engine/Source/Runtime/VulkanRHI/Private/VulkanRenderTarget.cpp
@@ -770,7 +770,7 @@ void FVulkanDynamicRHI::RHIReadSurfaceData(FTextureRHIParamRef TextureRHI, FIntR
 
 	FVulkanCmdBuffer* CmdBuffer = ImmediateContext.GetCommandBufferManager()->GetUploadCmdBuffer();
 
-	ensure(Texture2D->Surface.StorageFormat == VK_FORMAT_R8G8B8A8_UNORM || Texture2D->Surface.StorageFormat == VK_FORMAT_B8G8R8A8_UNORM || Texture2D->Surface.StorageFormat == VK_FORMAT_R16G16B16A16_SFLOAT || Texture2D->Surface.StorageFormat == VK_FORMAT_A2B10G10R10_UNORM_PACK32);
+	ensure(Texture2D->Surface.StorageFormat == VK_FORMAT_R8G8B8A8_UNORM || Texture2D->Surface.StorageFormat == VK_FORMAT_B8G8R8A8_UNORM || Texture2D->Surface.StorageFormat == VK_FORMAT_R16G16B16A16_SFLOAT || Texture2D->Surface.StorageFormat == VK_FORMAT_A2B10G10R10_UNORM_PACK32 || Texture2D->Surface.StorageFormat == VK_FORMAT_R16G16B16A16_UNORM);
 	const bool bIs8Bpp = (Texture2D->Surface.StorageFormat == VK_FORMAT_R16G16B16A16_SFLOAT);
 	const uint32 Size = NumPixels * sizeof(FColor) * (bIs8Bpp ? 2 : 1);
 	VulkanRHI::FStagingBuffer* StagingBuffer = Device->GetStagingManager().AcquireBuffer(Size, VK_BUFFER_USAGE_TRANSFER_DST_BIT, true);
@@ -892,6 +892,31 @@ void FVulkanDynamicRHI::RHIReadSurfaceData(FTextureRHIParamRef TextureRHI, FIntR
 			}
 		}
 	}
+	else if (Texture2D->Surface.StorageFormat == VK_FORMAT_R16G16B16A16_UNORM)
+	{
+		struct FRGBA16U
+		{
+			uint16 R;
+			uint16 G;
+			uint16 B;
+			uint16 A;
+		};
+		for (int32 Row = Rect.Min.Y; Row < Rect.Max.Y; ++Row)
+		{
+			FRGBA16U* Src = (FRGBA16U*)StagingBuffer->GetMappedPointer() + Row * TextureRHI2D->GetSizeX() + Rect.Min.X;
+			for (int32 Col = Rect.Min.X; Col < Rect.Max.X; ++Col)
+			{
+				*Dest = FLinearColor(
+					(float)Src->R / 65535.0f,
+					(float)Src->G / 65535.0f,
+					(float)Src->B / 65535.0f,
+					(float)Src->A / 65535.0f
+				).Quantize();
+				Dest++;
+				Src++;
+			}
+		}
+	}
 	else if (Texture2D->Surface.StorageFormat == VK_FORMAT_B8G8R8A8_UNORM)
 	{
 		FColor* Src = (FColor*)StagingBuffer->GetMappedPointer() + Rect.Min.Y * TextureRHI2D->GetSizeX() + Rect.Min.X;
diff --git a/Engine/Source/Runtime/VulkanRHI/Private/VulkanViewport.cpp b/Engine/Source/Runtime/VulkanRHI/Private/VulkanViewport.cpp
--- a/Engine/Source/Runtime/VulkanRHI/Private/VulkanViewport.cpp
+++ b/Engine/Source/Runtime/VulkanRHI/Private/VulkanViewport.cpp
@@ -20,7 +20,6 @@ FAutoConsoleVariable GCVarDelayAcquireBackBuffer(
 	ECVF_ReadOnly
 );
 
-
 struct FRHICommandAcquireBackBuffer final : public FRHICommand<FRHICommandAcquireBackBuffer>
 {
 	FVulkanViewport* Viewport;
@@ -209,6 +208,9 @@ FVulkanTexture2D* FVulkanViewport::GetBackBuffer(FRHICommandList& RHICmdList)
 {
 	check(IsInRenderingThread());
 
+	// make sure we aren't in the middle of swapchain recreation (which can happen on e.g. RHI thread)
+	FScopeLock LockSwapchain(&RecreatingSwapchain);
+
 	if (!RenderingBackBuffer && FVulkanPlatform::SupportsStandardSwapchain())
 	{
 		check(!DelayAcquireBackBuffer());
@@ -443,6 +445,7 @@ void FVulkanViewport::RecreateSwapchain(void* NewNativeWindow, bool bForce)
 		return;
 	}
 
+	FScopeLock LockSwapchain(&RecreatingSwapchain);
 	RenderingBackBuffer = nullptr;
 	RHIBackBuffer = nullptr;
 
diff --git a/Engine/Source/Runtime/VulkanRHI/Public/VulkanViewport.h b/Engine/Source/Runtime/VulkanRHI/Public/VulkanViewport.h
--- a/Engine/Source/Runtime/VulkanRHI/Public/VulkanViewport.h
+++ b/Engine/Source/Runtime/VulkanRHI/Public/VulkanViewport.h
@@ -66,6 +66,9 @@ protected:
 	TRefCountPtr<FVulkanBackBuffer> RenderingBackBuffer;
 	TRefCountPtr<FVulkanBackBuffer> RHIBackBuffer;
 
+	/** narrow-scoped section that locks access to back buffer during its recreation*/
+	FCriticalSection RecreatingSwapchain;
+
 	FVulkanDynamicRHI* RHI;
 	uint32 SizeX;
 	uint32 SizeY;
