DR - vk - Copy from 4504673, 4504440, 4500360, 4501574 & 4496977
- Better error msg
- Fix for immutable samplers not getting in PSO hash
- Missing time for textures
- Fix for fullscreen
- Fix crash alt-tabbing: Safeguard the access to the backbuffer pointer. RHI thread can be in the midst of recreating the swapchain while RenderingThread calls GetBackBuffer().

diff --git a/Engine/Source/Runtime/VulkanRHI/Private/VulkanCommands.cpp b/Engine/Source/Runtime/VulkanRHI/Private/VulkanCommands.cpp
--- a/Engine/Source/Runtime/VulkanRHI/Private/VulkanCommands.cpp
+++ b/Engine/Source/Runtime/VulkanRHI/Private/VulkanCommands.cpp
@@ -192,6 +192,7 @@ void FVulkanCommandListContext::RHISetShaderTexture(FVertexShaderRHIParamRef Ver
 	FVulkanTextureBase* Texture = GetVulkanTextureFromRHITexture(NewTextureRHI);
 	VkImageLayout Layout = GetLayoutForDescriptor(Texture->Surface);
 	PendingGfxState->SetTextureForStage(ShaderStage::Vertex, TextureIndex, Texture, Layout);
+	NewTextureRHI->SetLastRenderTime((float)FPlatformTime::Seconds());
 }
 
 void FVulkanCommandListContext::RHISetShaderTexture(FHullShaderRHIParamRef HullShaderRHI, uint32 TextureIndex, FTextureRHIParamRef NewTextureRHI)
@@ -202,6 +203,7 @@ void FVulkanCommandListContext::RHISetShaderTexture(FHullShaderRHIParamRef HullS
 	FVulkanTextureBase* Texture = GetVulkanTextureFromRHITexture(NewTextureRHI);
 	VkImageLayout Layout = GetLayoutForDescriptor(Texture->Surface);
 	PendingGfxState->SetTexture(ShaderStage::Hull, TextureIndex, Texture, Layout);
+	NewTextureRHI->SetLastRenderTime((float)FPlatformTime::Seconds());
 */
 }
 
@@ -213,6 +215,7 @@ void FVulkanCommandListContext::RHISetShaderTexture(FDomainShaderRHIParamRef Dom
 	FVulkanTextureBase* Texture = GetVulkanTextureFromRHITexture(NewTextureRHI);
 	VkImageLayout Layout = GetLayoutForDescriptor(Texture->Surface);
 	PendingGfxState->SetTexture(ShaderStage::Domain, TextureIndex, Texture, Layout);
+	NewTextureRHI->SetLastRenderTime((float)FPlatformTime::Seconds());
 */
 }
 
@@ -223,6 +226,7 @@ void FVulkanCommandListContext::RHISetShaderTexture(FGeometryShaderRHIParamRef G
 	FVulkanTextureBase* Texture = GetVulkanTextureFromRHITexture(NewTextureRHI);
 	VkImageLayout Layout = GetLayoutForDescriptor(Texture->Surface);
 	PendingGfxState->SetTextureForStage(ShaderStage::Geometry, TextureIndex, Texture, Layout);
+	NewTextureRHI->SetLastRenderTime((float)FPlatformTime::Seconds());
 #else
 	ensureMsgf(0, TEXT("Geometry not supported!"));
 #endif
@@ -234,6 +238,7 @@ void FVulkanCommandListContext::RHISetShaderTexture(FPixelShaderRHIParamRef Pixe
 	FVulkanTextureBase* Texture = GetVulkanTextureFromRHITexture(NewTextureRHI);
 	VkImageLayout Layout = GetLayoutForDescriptor(Texture->Surface);
 	PendingGfxState->SetTextureForStage(ShaderStage::Pixel, TextureIndex, Texture, Layout);
+	NewTextureRHI->SetLastRenderTime((float)FPlatformTime::Seconds());
 }
 
 void FVulkanCommandListContext::RHISetShaderTexture(FComputeShaderRHIParamRef ComputeShaderRHI, uint32 TextureIndex, FTextureRHIParamRef NewTextureRHI)
@@ -244,6 +249,7 @@ void FVulkanCommandListContext::RHISetShaderTexture(FComputeShaderRHIParamRef Co
 	FVulkanTextureBase* VulkanTexture = GetVulkanTextureFromRHITexture(NewTextureRHI);
 	VkImageLayout Layout = GetLayoutForDescriptor(VulkanTexture->Surface);
 	PendingComputeState->SetTextureForStage(TextureIndex, VulkanTexture, Layout);
+	NewTextureRHI->SetLastRenderTime((float)FPlatformTime::Seconds());
 }
 
 void FVulkanCommandListContext::RHISetShaderResourceViewParameter(FVertexShaderRHIParamRef VertexShaderRHI, uint32 TextureIndex, FShaderResourceViewRHIParamRef SRVRHI)
@@ -484,7 +490,7 @@ template <typename TState>
 inline void /*FVulkanCommandListContext::*/SetShaderUniformBufferResources(FVulkanCommandListContext* Context, TState* State, const FVulkanShader* Shader, const TArray<FVulkanShaderHeader::FGlobalInfo>& GlobalInfos, const TArray<TEnumAsByte<VkDescriptorType>>& DescriptorTypes, const FVulkanShaderHeader::FUniformBufferInfo& HeaderUBInfo, const FVulkanUniformBuffer* UniformBuffer, const TArray<FDescriptorSetRemappingInfo::FRemappingInfo>& GlobalRemappingInfo)
 {
 	ensure(UniformBuffer->GetLayout().GetHash() == HeaderUBInfo.LayoutHash);
-	float CurrentTime = (float)FApp::GetCurrentTime();
+	float CurrentTime = (float)FPlatformTime::Seconds();
 	const TArray<TRefCountPtr<FRHIResource>>& ResourceArray = UniformBuffer->GetResourceTable();
 	for (int32 Index = 0; Index < HeaderUBInfo.ResourceEntries.Num(); ++Index)
 	{
@@ -569,7 +575,7 @@ inline void /*FVulkanCommandListContext::*/SetShaderUniformBufferResources(FVulk
 		GatherUniformBufferResources(ResourceBindingTable.ShaderResourceViewMap, ResourceBindingTable.ResourceTableBits, UniformBuffer, BindingIndex, SRVBindings);
 		if (CurrentTime == 0.0f)
 		{
-			CurrentTime = (float)FApp::GetCurrentTime();
+			CurrentTime = (float)FPlatformTime::Seconds();
 		}
 		for (int32 Index = 0; Index < SRVBindings.Num(); Index++)
 		{
diff --git a/Engine/Source/Runtime/VulkanRHI/Private/VulkanDescriptorSets.h b/Engine/Source/Runtime/VulkanRHI/Private/VulkanDescriptorSets.h
--- a/Engine/Source/Runtime/VulkanRHI/Private/VulkanDescriptorSets.h
+++ b/Engine/Source/Runtime/VulkanRHI/Private/VulkanDescriptorSets.h
@@ -286,7 +286,7 @@ public:
 	template<bool bIsCompute>
 	void FinalizeBindings(const FUniformBufferGatherInfo& UBGatherInfo, const TArrayView<const FSamplerStateRHIParamRef>& ImmutableSamplers);
 
-	void GenerateHash();
+	void GenerateHash(const TArrayView<const FSamplerStateRHIParamRef>& ImmutableSamplers);
 
 	friend uint32 GetTypeHash(const FVulkanDescriptorSetsLayoutInfo& In)
 	{
diff --git a/Engine/Source/Runtime/VulkanRHI/Private/VulkanPipeline.cpp b/Engine/Source/Runtime/VulkanRHI/Private/VulkanPipeline.cpp
--- a/Engine/Source/Runtime/VulkanRHI/Private/VulkanPipeline.cpp
+++ b/Engine/Source/Runtime/VulkanRHI/Private/VulkanPipeline.cpp
@@ -1033,6 +1033,13 @@ FArchive& operator << (FArchive& Ar, FVulkanPipelineStateCacheManager::FGfxPipel
 
 	Ar << Entry.RenderTargets;
 
+#if VULKAN_SUPPORTS_COLOR_CONVERSIONS
+	for (uint32 Index = 0; Index < MaxImmutableSamplers; ++Index)
+	{
+		Ar << Entry.ImmutableSamplers[Index];
+	}
+#endif
+
 	return Ar;
 }
 
diff --git a/Engine/Source/Runtime/VulkanRHI/Private/VulkanRHI.cpp b/Engine/Source/Runtime/VulkanRHI/Private/VulkanRHI.cpp
--- a/Engine/Source/Runtime/VulkanRHI/Private/VulkanRHI.cpp
+++ b/Engine/Source/Runtime/VulkanRHI/Private/VulkanRHI.cpp
@@ -448,11 +448,12 @@ void FVulkanDynamicRHI::SelectAndInitDevice()
 {
 	uint32 GpuCount = 0;
 	VERIFYVULKANRESULT_EXPANDED(VulkanRHI::vkEnumeratePhysicalDevices(Instance, &GpuCount, nullptr));
-	check(GpuCount >= 1);
+	checkf(GpuCount >= 1, TEXT("No GPU(s)/Driver(s) that support Vulkan were found! Make sure your drivers are up to date and that you are not pending a reboot."));
 
 	TArray<VkPhysicalDevice> PhysicalDevices;
 	PhysicalDevices.AddZeroed(GpuCount);
 	VERIFYVULKANRESULT_EXPANDED(VulkanRHI::vkEnumeratePhysicalDevices(Instance, &GpuCount, PhysicalDevices.GetData()));
+	checkf(GpuCount >= 1, TEXT("Couldn't enumerate physical devices! Make sure your drivers are up to date and that you are not pending a reboot."));
 
 #if VULKAN_ENABLE_DESKTOP_HMD_SUPPORT
 	FVulkanDevice* HmdDevice = nullptr;
@@ -1117,7 +1118,7 @@ void FVulkanDescriptorSetsLayoutInfo::AddDescriptor(int32 DescriptorSetIndex, co
 	}
 }
 
-void FVulkanDescriptorSetsLayoutInfo::GenerateHash()
+void FVulkanDescriptorSetsLayoutInfo::GenerateHash(const TArrayView<const FSamplerStateRHIParamRef>& InImmutableSamplers)
 {
 	const int32 LayoutCount = SetLayouts.Num();
 	Hash = FCrc::MemCrc32(&TypesUsageID, sizeof(uint32), LayoutCount);
@@ -1142,6 +1143,18 @@ void FVulkanDescriptorSetsLayoutInfo::GenerateHash()
 		TArray<uint16>& PackedUBBindingIndices = RemappingInfo.StageInfos[RemapingIndex].PackedUBBindingIndices;
 		Hash = FCrc::MemCrc32(PackedUBBindingIndices.GetData(), sizeof(uint16) * PackedUBBindingIndices.Num(), Hash);
 	}
+
+#if VULKAN_SUPPORTS_COLOR_CONVERSIONS
+	VkSampler ImmutableSamplers[MaxImmutableSamplers];
+	VkSampler* ImmutableSamplerPtr = ImmutableSamplers;
+	for (int32 Index = 0; Index < InImmutableSamplers.Num(); ++Index)
+	{
+		FRHISamplerState* SamplerState = InImmutableSamplers[Index];
+		*ImmutableSamplerPtr++ = SamplerState ? ResourceCast(SamplerState)->Sampler : VK_NULL_HANDLE;
+	}
+	FMemory::Memzero(ImmutableSamplerPtr, (MaxImmutableSamplers - InImmutableSamplers.Num()));
+	Hash = FCrc::MemCrc32(ImmutableSamplers, sizeof(VkSampler) * MaxImmutableSamplers, Hash);
+#endif
 }
 
 void FVulkanDescriptorSetsLayoutInfo::CompileTypesUsageID()
diff --git a/Engine/Source/Runtime/VulkanRHI/Private/VulkanShaders.cpp b/Engine/Source/Runtime/VulkanRHI/Private/VulkanShaders.cpp
--- a/Engine/Source/Runtime/VulkanRHI/Private/VulkanShaders.cpp
+++ b/Engine/Source/Runtime/VulkanRHI/Private/VulkanShaders.cpp
@@ -542,7 +542,7 @@ void FVulkanDescriptorSetsLayoutInfo::FinalizeBindings(const FUniformBufferGathe
 	}
 
 	CompileTypesUsageID();
-	GenerateHash();
+	GenerateHash(ImmutableSamplers);
 
 	// Validate no empty sets were made
 	for (int32 Index = 0; Index < RemappingInfo.SetInfos.Num(); ++Index)
diff --git a/Engine/Source/Runtime/VulkanRHI/Private/VulkanSwapChain.cpp b/Engine/Source/Runtime/VulkanRHI/Private/VulkanSwapChain.cpp
--- a/Engine/Source/Runtime/VulkanRHI/Private/VulkanSwapChain.cpp
+++ b/Engine/Source/Runtime/VulkanRHI/Private/VulkanSwapChain.cpp
@@ -386,6 +386,8 @@ FVulkanSwapChain::FVulkanSwapChain(VkInstance InInstance, FVulkanDevice& InDevic
 	uint32 SizeY = FVulkanPlatform::SupportsQuerySurfaceProperties() ? (SurfProperties.currentExtent.height == 0xFFFFFFFF ? Height : SurfProperties.currentExtent.height) : Height;
 	//FPlatformMisc::LowLevelOutputDebugStringf(TEXT("Create swapchain: %ux%u \n"), SizeX, SizeY);
 
+	// Some platforms cheat with the window size. Allow them to update.
+	FVulkanPlatform::UpdateWindowSize(WindowHandle, SizeX, SizeY);
 
 	VkSwapchainCreateInfoKHR SwapChainInfo;
 	ZeroVulkanStruct(SwapChainInfo, VK_STRUCTURE_TYPE_SWAPCHAIN_CREATE_INFO_KHR);
diff --git a/Engine/Source/Runtime/VulkanRHI/Private/VulkanViewport.cpp b/Engine/Source/Runtime/VulkanRHI/Private/VulkanViewport.cpp
--- a/Engine/Source/Runtime/VulkanRHI/Private/VulkanViewport.cpp
+++ b/Engine/Source/Runtime/VulkanRHI/Private/VulkanViewport.cpp
@@ -599,7 +599,7 @@ void FVulkanViewport::CreateSwapchain()
 	}
 }
 
-inline static void CopyImageToBackBuffer(FVulkanCmdBuffer* CmdBuffer, bool bSourceReadOnly, VkImage SrcSurface, VkImage DstSurface, int32 SizeX, int32 SizeY)
+inline static void CopyImageToBackBuffer(FVulkanCmdBuffer* CmdBuffer, bool bSourceReadOnly, VkImage SrcSurface, VkImage DstSurface, int32 SizeX, int32 SizeY, int32 WindowSizeX, int32 WindowSizeY)
 {
 	VulkanRHI::FPendingBarrier Barriers;
 	int32 SourceIndex = Barriers.AddImageBarrier(SrcSurface, VK_IMAGE_ASPECT_COLOR_BIT, 1);
@@ -614,23 +614,52 @@ inline static void CopyImageToBackBuffer(FVulkanCmdBuffer* CmdBuffer, bool bSour
 
 	VulkanRHI::DebugHeavyWeightBarrier(CmdBuffer->GetHandle(), 32);
 
-	VkImageCopy Region;
-	FMemory::Memzero(Region);
-	Region.extent.width = SizeX;
-	Region.extent.height = SizeY;
-	Region.extent.depth = 1;
-	Region.srcSubresource.aspectMask = VK_IMAGE_ASPECT_COLOR_BIT;
-	//Region.srcSubresource.baseArrayLayer = 0;
-	Region.srcSubresource.layerCount = 1;
-	//Region.srcSubresource.mipLevel = 0;
-	Region.dstSubresource.aspectMask = VK_IMAGE_ASPECT_COLOR_BIT;
-	//Region.dstSubresource.baseArrayLayer = 0;
-	Region.dstSubresource.layerCount = 1;
-	//Region.dstSubresource.mipLevel = 0;
-	VulkanRHI::vkCmdCopyImage(CmdBuffer->GetHandle(),
-		SrcSurface, VK_IMAGE_LAYOUT_TRANSFER_SRC_OPTIMAL,
-		DstSurface, VK_IMAGE_LAYOUT_TRANSFER_DST_OPTIMAL,
-		1, &Region);
+	if (SizeX != WindowSizeX || SizeY != WindowSizeY)
+	{
+		VkImageBlit Region;
+		FMemory::Memzero(Region);
+		Region.srcOffsets[0].x = 0;
+		Region.srcOffsets[0].y = 0;
+		Region.srcOffsets[0].z = 0;
+		Region.srcOffsets[1].x = SizeX;
+		Region.srcOffsets[1].y = SizeY;
+		Region.srcOffsets[1].z = 1;
+		Region.srcSubresource.aspectMask = VK_IMAGE_ASPECT_COLOR_BIT;
+		Region.srcSubresource.layerCount = 1;
+		Region.dstOffsets[0].x = 0;
+		Region.dstOffsets[0].y = 0;
+		Region.dstOffsets[0].z = 0;
+		Region.dstOffsets[1].x = WindowSizeX;
+		Region.dstOffsets[1].y = WindowSizeY;
+		Region.dstOffsets[1].z = 1;
+		Region.dstSubresource.aspectMask = VK_IMAGE_ASPECT_COLOR_BIT;
+		Region.dstSubresource.baseArrayLayer = 0;
+		Region.dstSubresource.layerCount = 1;
+		VulkanRHI::vkCmdBlitImage(CmdBuffer->GetHandle(),
+			SrcSurface, VK_IMAGE_LAYOUT_TRANSFER_SRC_OPTIMAL,
+			DstSurface, VK_IMAGE_LAYOUT_TRANSFER_DST_OPTIMAL,
+			1, &Region, VK_FILTER_LINEAR);
+	}
+	else
+	{
+		VkImageCopy Region;
+		FMemory::Memzero(Region);
+		Region.extent.width = SizeX;
+		Region.extent.height = SizeY;
+		Region.extent.depth = 1;
+		Region.srcSubresource.aspectMask = VK_IMAGE_ASPECT_COLOR_BIT;
+		//Region.srcSubresource.baseArrayLayer = 0;
+		Region.srcSubresource.layerCount = 1;
+		//Region.srcSubresource.mipLevel = 0;
+		Region.dstSubresource.aspectMask = VK_IMAGE_ASPECT_COLOR_BIT;
+		//Region.dstSubresource.baseArrayLayer = 0;
+		Region.dstSubresource.layerCount = 1;
+		//Region.dstSubresource.mipLevel = 0;
+		VulkanRHI::vkCmdCopyImage(CmdBuffer->GetHandle(),
+			SrcSurface, VK_IMAGE_LAYOUT_TRANSFER_SRC_OPTIMAL,
+			DstSurface, VK_IMAGE_LAYOUT_TRANSFER_DST_OPTIMAL,
+			1, &Region);
+	}
 
 	// Prepare for present
 	Barriers.ResetStages();
@@ -658,8 +687,12 @@ bool FVulkanViewport::Present(FVulkanCommandListContext* Context, FVulkanCmdBuff
 				UE_LOG(LogVulkanRHI, Fatal, TEXT("Swapchain acquire image index failed!"));
 			}
 
+			uint32 WindowSizeX = SizeX;
+			uint32 WindowSizeY = SizeY;
+			FVulkanPlatform::UpdateWindowSize(WindowHandle, WindowSizeX, WindowSizeY);
+
 			Context->RHIPushEvent(TEXT("CopyImageToBackBuffer"), FColor::Blue);
-			CopyImageToBackBuffer(CmdBuffer, true, RenderingBackBuffer->Surface.Image, BackBufferImages[AcquiredImageIndex], SizeX, SizeY);
+			CopyImageToBackBuffer(CmdBuffer, true, RenderingBackBuffer->Surface.Image, BackBufferImages[AcquiredImageIndex], SizeX, SizeY, WindowSizeX, WindowSizeY);
 			Context->RHIPopEvent();
 		}
 		else
