DR - vk - missed files
#rnx

diff --git a/Engine/Source/Runtime/VulkanRHI/Private/Linux/VulkanLinuxPlatform.cpp b/Engine/Source/Runtime/VulkanRHI/Private/Linux/VulkanLinuxPlatform.cpp
--- a/Engine/Source/Runtime/VulkanRHI/Private/Linux/VulkanLinuxPlatform.cpp
+++ b/Engine/Source/Runtime/VulkanRHI/Private/Linux/VulkanLinuxPlatform.cpp
@@ -172,3 +172,8 @@ void FVulkanLinuxPlatform::WriteCrashMarker(const FOptionalVulkanDeviceExtension
 		}
 	}
 }
+
+void FVulkanLinuxPlatform::UpdateWindowSize(void* WindowHandle, uint32& Width, uint32& Height)
+{
+	SDL_Vulkan_GetDrawableSize((SDL_Window*) WindowHandle, (int32*) &Width, (int32*) &Height);
+}
diff --git a/Engine/Source/Runtime/VulkanRHI/Private/Linux/VulkanLinuxPlatform.h b/Engine/Source/Runtime/VulkanRHI/Private/Linux/VulkanLinuxPlatform.h
--- a/Engine/Source/Runtime/VulkanRHI/Private/Linux/VulkanLinuxPlatform.h
+++ b/Engine/Source/Runtime/VulkanRHI/Private/Linux/VulkanLinuxPlatform.h
@@ -57,6 +57,8 @@ public:
 #endif
 	}
 
+	static void UpdateWindowSize(void* WindowHandle, uint32& Width, uint32& Height);
+
 	static void WriteCrashMarker(const FOptionalVulkanDeviceExtensions& OptionalExtensions, VkCommandBuffer CmdBuffer, VkBuffer DestBuffer, const TArrayView<uint32>& Entries, bool bAdding);
 
 protected:
diff --git a/Engine/Source/Runtime/VulkanRHI/Private/VulkanGenericPlatform.h b/Engine/Source/Runtime/VulkanRHI/Private/VulkanGenericPlatform.h
--- a/Engine/Source/Runtime/VulkanRHI/Private/VulkanGenericPlatform.h
+++ b/Engine/Source/Runtime/VulkanRHI/Private/VulkanGenericPlatform.h
@@ -102,4 +102,7 @@ public:
 
 	// Ensure the last frame completed on the GPU
 	static bool RequiresWaitingForFrameCompletionEvent() { return true; }
+
+	// Some platforms imitate fullscreen mode by using bigger window
+	static void UpdateWindowSize(void* WindowHandle, uint32& Width, uint32& Height) {}
 };
diff --git a/Engine/Source/Runtime/VulkanRHI/Private/VulkanUtil.cpp b/Engine/Source/Runtime/VulkanRHI/Private/VulkanUtil.cpp
--- a/Engine/Source/Runtime/VulkanRHI/Private/VulkanUtil.cpp
+++ b/Engine/Source/Runtime/VulkanRHI/Private/VulkanUtil.cpp
@@ -500,6 +500,7 @@ void FVulkanGPUProfiler::DumpCrashMarkers(void* BufferData)
 				const FString* Frame = CachedStrings.Find(Value);
 				UE_LOG(LogVulkanRHI, Error, TEXT("[VK_NV_device_diagnostic_checkpoints] %i: Stage 0x%x, %s (CRC 0x%x)"), Index, Data[Index].stage, Frame ? *(*Frame) : TEXT("<undefined>"), Value);
 			}
+			GLog->FlushThreadedLogs();
 			GLog->PanicFlushThreadedLogs();
 			GLog->Flush();
 		}
diff --git a/Engine/Source/Runtime/VulkanRHI/Public/VulkanViewport.h b/Engine/Source/Runtime/VulkanRHI/Public/VulkanViewport.h
--- a/Engine/Source/Runtime/VulkanRHI/Public/VulkanViewport.h
+++ b/Engine/Source/Runtime/VulkanRHI/Public/VulkanViewport.h
@@ -7,6 +7,7 @@
 #pragma once
 
 #include "VulkanResources.h"
+#include "HAL/CriticalSection.h"
 
 class FVulkanDynamicRHI;
 class FVulkanSwapChain;
