DR - vk - Win32 compile fix
#rnx

diff --git a/Engine/Source/Runtime/VulkanRHI/Private/VulkanPipeline.cpp b/Engine/Source/Runtime/VulkanRHI/Private/VulkanPipeline.cpp
--- a/Engine/Source/Runtime/VulkanRHI/Private/VulkanPipeline.cpp
+++ b/Engine/Source/Runtime/VulkanRHI/Private/VulkanPipeline.cpp
@@ -1036,7 +1036,9 @@ FArchive& operator << (FArchive& Ar, FVulkanPipelineStateCacheManager::FGfxPipel
 #if VULKAN_SUPPORTS_COLOR_CONVERSIONS
 	for (uint32 Index = 0; Index < MaxImmutableSamplers; ++Index)
 	{
-		Ar << Entry.ImmutableSamplers[Index];
+		uint64 Sampler = (uint64)Entry.ImmutableSamplers[Index];
+		Ar << Sampler;
+		Entry.ImmutableSamplers[Index] = (SIZE_T)Sampler;
 	}
 #endif
 
