Fixed: Crash in FVulkanViewport::RecreateSwapchain when app is backgrounded during startrup

diff --git a/Engine/Source/Runtime/VulkanRHI/Private/Android/VulkanAndroidPlatform.cpp b/Engine/Source/Runtime/VulkanRHI/Private/Android/VulkanAndroidPlatform.cpp
--- a/Engine/Source/Runtime/VulkanRHI/Private/Android/VulkanAndroidPlatform.cpp
+++ b/Engine/Source/Runtime/VulkanRHI/Private/Android/VulkanAndroidPlatform.cpp
@@ -188,4 +188,20 @@ void FVulkanAndroidPlatform::OverridePlatformHandlers(bool bInit)
 	}
 }
 
+void FVulkanAndroidPlatform::BlockUntilWindowIsAwailable()
+{
+	void* WindowHandle = FAndroidWindow::GetHardwareWindow();
+	if (WindowHandle == nullptr)
+	{
+		// Sleep if the hardware window isn't currently available.
+		// The Window may not exist if the activity is pausing/resuming, in which case we make this thread wait
+		FPlatformMisc::LowLevelOutputDebugString(TEXT("Waiting for Native window in FVulkanAndroidPlatform::BlockUntilWindowIsAwailable"));
+		while (WindowHandle == nullptr)
+		{
+			FPlatformProcess::Sleep(0.001f);
+			WindowHandle = FAndroidWindow::GetHardwareWindow();
+		}
+	}
+}
+
 #endif
diff --git a/Engine/Source/Runtime/VulkanRHI/Private/Android/VulkanAndroidPlatform.h b/Engine/Source/Runtime/VulkanRHI/Private/Android/VulkanAndroidPlatform.h
--- a/Engine/Source/Runtime/VulkanRHI/Private/Android/VulkanAndroidPlatform.h
+++ b/Engine/Source/Runtime/VulkanRHI/Private/Android/VulkanAndroidPlatform.h
@@ -2,6 +2,8 @@
 
 #pragma once
 
+#include "RHI.h"
+
 #define VK_USE_PLATFORM_ANDROID_KHR					1
 
 #define VULKAN_HAS_PHYSICAL_DEVICE_PROPERTIES2		0
@@ -93,7 +95,9 @@ public:
 	static bool RequiresSwapchainGeneralInitialLayout() { return true; }
 
 	static bool RequiresWaitingForFrameCompletionEvent() { return false; }
-
+	
+	static void BlockUntilWindowIsAwailable();
+	
 protected:
 	static void* VulkanLib;
 	static bool bAttemptedLoad;
diff --git a/Engine/Source/Runtime/VulkanRHI/Private/VulkanGenericPlatform.h b/Engine/Source/Runtime/VulkanRHI/Private/VulkanGenericPlatform.h
--- a/Engine/Source/Runtime/VulkanRHI/Private/VulkanGenericPlatform.h
+++ b/Engine/Source/Runtime/VulkanRHI/Private/VulkanGenericPlatform.h
@@ -96,4 +96,7 @@ public:
 
 	// Ensure the last frame completed on the GPU
 	static bool RequiresWaitingForFrameCompletionEvent() { return true; }
+
+	// Blocks until hardware window is available
+	static void BlockUntilWindowIsAwailable() {};
 };
diff --git a/Engine/Source/Runtime/VulkanRHI/Private/VulkanViewport.cpp b/Engine/Source/Runtime/VulkanRHI/Private/VulkanViewport.cpp
--- a/Engine/Source/Runtime/VulkanRHI/Private/VulkanViewport.cpp
+++ b/Engine/Source/Runtime/VulkanRHI/Private/VulkanViewport.cpp
@@ -793,6 +793,9 @@ bool FVulkanViewport::Present(FVulkanCommandListContext* Context, FVulkanCmdBuff
 		Queue->Submit(CmdBuffer);
 	}
 
+	// Do not present until hardware window is available. On Android window could be destroyed while RHIT executes commands
+	FVulkanPlatform::BlockUntilWindowIsAwailable();
+
 	//Flush all commands
 	//check(0);
 
