Implement RHICopyTexture for D3D11 that support TextureArray2D, Texture2D, TextureCube and Texture3D. On other platform, it will fallback to using RHIResolveToTarget

diff --git a/Engine/Source/Runtime/EmptyRHI/Public/EmptyRHI.h b/Engine/Source/Runtime/EmptyRHI/Public/EmptyRHI.h
--- a/Engine/Source/Runtime/EmptyRHI/Public/EmptyRHI.h
+++ b/Engine/Source/Runtime/EmptyRHI/Public/EmptyRHI.h
@@ -138,6 +138,7 @@ public:
 	virtual void RHISetMultipleViewports(uint32 Count, const FViewportBounds* Data) final override;
 	virtual void RHIClearTinyUAV(FUnorderedAccessViewRHIParamRef UnorderedAccessViewRHI, const uint32* Values) final override;
 	virtual void RHICopyToResolveTarget(FTextureRHIParamRef SourceTexture, FTextureRHIParamRef DestTexture, const FResolveParams& ResolveParams) final override;
+	virtual void RHICopyTexture(FTextureRHIParamRef SourceTexture, FTextureRHIParamRef DestTexture, const FRHICopyTextureInfo& CopyInfo) final override;
 	virtual void RHIBeginRenderQuery(FRenderQueryRHIParamRef RenderQuery) final override;
 	virtual void RHIEndRenderQuery(FRenderQueryRHIParamRef RenderQuery) final override;
 	virtual void RHISubmitCommandsHint() final override;
diff --git a/Engine/Source/Runtime/NullDrv/Public/NullRHI.h b/Engine/Source/Runtime/NullDrv/Public/NullRHI.h
--- a/Engine/Source/Runtime/NullDrv/Public/NullRHI.h
+++ b/Engine/Source/Runtime/NullDrv/Public/NullRHI.h
@@ -383,6 +383,11 @@ public:
 
 	}
 
+	virtual void RHICopyTexture(FTextureRHIParamRef SourceTexture, FTextureRHIParamRef DestTexture, const FRHICopyTextureInfo& CopyInfo) final override
+	{
+
+	}
+
 	virtual void RHIBindDebugLabelName(FTextureRHIParamRef Texture, const TCHAR* Name) final override
 	{
 
diff --git a/Engine/Source/Runtime/RHI/Public/RHI.h b/Engine/Source/Runtime/RHI/Public/RHI.h
--- a/Engine/Source/Runtime/RHI/Public/RHI.h
+++ b/Engine/Source/Runtime/RHI/Public/RHI.h
@@ -1257,36 +1257,20 @@ struct FResolveParams
 
 struct FRHICopyTextureInfo
 {
-	// Number of texels to copy. Z Must be always be > 0.
-	FIntVector Size = {0, 0, 1};
-
-	/** The mip index to copy in both source and dest. */
-	int32 MipIndex = 0;
-
-	/** Array index or cube face to resolve in the source. For Cubemap arrays this would be ArraySlice * 6 + FaceIndex. */
-	int32 SourceArraySlice = 0;
-	/** Array index or cube face to resolve in the dest. */
-	int32 DestArraySlice = 0;
-	// How many slices or faces to copy.
-	int32 NumArraySlices = 1;
-
-	// 2D copy
-	FRHICopyTextureInfo(int32 InWidth, int32 InHeight)
-		: Size(InWidth, InHeight, 1)
-	{
-	}
+	// Number of texels to copy. By default it will copy the whole resource if no size is specified.
+	FIntVector Size = FIntVector::ZeroValue;
 
-	FRHICopyTextureInfo(const FIntVector& InSize)
-		: Size(InSize)
-	{
-	}
+	// Position of the copy from the source texture/to destination texture
+	FIntVector SourcePosition = FIntVector::ZeroValue;
+	FIntVector DestPosition = FIntVector::ZeroValue;
 
-	void AdvanceMip()
-	{
-		++MipIndex;
-		Size.X = FMath::Max(Size.X / 2, 1);
-		Size.Y = FMath::Max(Size.Y / 2, 1);
-	}
+	uint32 SourceSliceIndex = 0;
+	uint32 DestSliceIndex = 0;
+	uint32 NumSlices = 1;
+
+	// Mips to copy and destination mips
+	uint32 SourceMipIndex = 0;
+	uint32 DestMipIndex = 0;
 };
 
 enum class EResourceTransitionAccess
diff --git a/Engine/Source/Runtime/RHI/Public/RHICommandList.h b/Engine/Source/Runtime/RHI/Public/RHICommandList.h
--- a/Engine/Source/Runtime/RHI/Public/RHICommandList.h
+++ b/Engine/Source/Runtime/RHI/Public/RHICommandList.h
@@ -1370,8 +1370,8 @@ struct FRHICommandCopyTexture final : public FRHICommand<FRHICommandCopyTexture>
 	{
 		ensure(SourceTexture);
 		ensure(DestTexture);
-		ensure(SourceTexture->GetTexture2D() || SourceTexture->GetTexture3D() || SourceTexture->GetTextureCube());
-		ensure(DestTexture->GetTexture2D() || DestTexture->GetTexture3D() || DestTexture->GetTextureCube());
+		ensure(SourceTexture->GetTexture2D() || SourceTexture->GetTexture2DArray() || SourceTexture->GetTexture3D() || SourceTexture->GetTextureCube());
+		ensure(DestTexture->GetTexture2D() || DestTexture->GetTexture2DArray() || DestTexture->GetTexture3D() || DestTexture->GetTextureCube());
 	}
 	RHI_API void Execute(FRHICommandListBase& CmdList);
 };
diff --git a/Engine/Source/Runtime/RHI/Public/RHIContext.h b/Engine/Source/Runtime/RHI/Public/RHIContext.h
--- a/Engine/Source/Runtime/RHI/Public/RHIContext.h
+++ b/Engine/Source/Runtime/RHI/Public/RHIContext.h
@@ -626,18 +626,18 @@ public:
 	virtual void RHICopyTexture(FTextureRHIParamRef SourceTexture, FTextureRHIParamRef DestTexture, const FRHICopyTextureInfo& CopyInfo)
 	{
 		const bool bIsCube = SourceTexture->GetTextureCube() != nullptr;
-		const bool bAllCubeFaces = bIsCube && (CopyInfo.NumArraySlices % 6) == 0;
-		const int32 NumArraySlices = bAllCubeFaces ? CopyInfo.NumArraySlices / 6 : CopyInfo.NumArraySlices;
+		const bool bAllCubeFaces = bIsCube && (CopyInfo.NumSlices % 6) == 0;
+		const int32 NumArraySlices = bAllCubeFaces ? CopyInfo.NumSlices / 6 : CopyInfo.NumSlices;
 		const int32 NumFaces = bAllCubeFaces ? 6 : 1;
 		for (int32 ArrayIndex = 0; ArrayIndex < NumArraySlices; ++ArrayIndex)
 		{
-			int32 SourceArrayIndex = CopyInfo.SourceArraySlice + ArrayIndex;
-			int32 DestArrayIndex = CopyInfo.DestArraySlice + ArrayIndex;
+			int32 SourceArrayIndex = CopyInfo.SourceSliceIndex + ArrayIndex;
+			int32 DestArrayIndex = CopyInfo.DestSliceIndex + ArrayIndex;
 			for (int32 FaceIndex = 0; FaceIndex < NumFaces; ++FaceIndex)
 			{
 				FResolveParams ResolveParams(FResolveRect(),
 					bIsCube ? (ECubeFace)FaceIndex : CubeFace_PosX,
-					CopyInfo.MipIndex,
+					CopyInfo.SourceMipIndex,
 					SourceArrayIndex,
 					DestArrayIndex
 				);
diff --git a/Engine/Source/Runtime/Renderer/Private/MobileReflectionEnvironmentCapture.cpp b/Engine/Source/Runtime/Renderer/Private/MobileReflectionEnvironmentCapture.cpp
--- a/Engine/Source/Runtime/Renderer/Private/MobileReflectionEnvironmentCapture.cpp
+++ b/Engine/Source/Runtime/Renderer/Private/MobileReflectionEnvironmentCapture.cpp
@@ -191,8 +191,9 @@ namespace MobileReflectionEnvironmentCapture
 
 			RHICmdList.TransitionResource(EResourceTransitionAccess::EWritable, ProcessedTexture->TextureRHI);
 
-			FRHICopyTextureInfo CopyInfo(ProcessedTexture->GetSizeX(), ProcessedTexture->GetSizeY());
-			CopyInfo.NumArraySlices = 6;
+			FRHICopyTextureInfo CopyInfo;
+			CopyInfo.Size = FIntVector(ProcessedTexture->GetSizeX(), ProcessedTexture->GetSizeY(), 1);
+			CopyInfo.NumSlices = 6;
 
 			// GPU copy back to the skylight's texture, which is not a render target
 			for (int32 MipIndex = 0; MipIndex < NumMips; MipIndex++)
@@ -202,7 +203,11 @@ namespace MobileReflectionEnvironmentCapture
 				FSceneRenderTargetItem& EffectiveSource = GetEffectiveRenderTarget(SceneContext, false, bUseHQFiltering ? 0 : MipIndex);
 				RHICmdList.TransitionResource(EResourceTransitionAccess::EReadable, EffectiveSource.ShaderResourceTexture);
 				RHICmdList.CopyTexture(EffectiveSource.ShaderResourceTexture, ProcessedTexture->TextureRHI, CopyInfo);
-				CopyInfo.AdvanceMip();
+
+				++CopyInfo.SourceMipIndex;
+				++CopyInfo.DestMipIndex;
+				CopyInfo.Size.X = FMath::Max(1, CopyInfo.Size.X / 2);
+				CopyInfo.Size.Y = FMath::Max(1, CopyInfo.Size.Y / 2);
 			}
 		}
 	}
diff --git a/Engine/Source/Runtime/Renderer/Private/ReflectionEnvironmentCapture.cpp b/Engine/Source/Runtime/Renderer/Private/ReflectionEnvironmentCapture.cpp
--- a/Engine/Source/Runtime/Renderer/Private/ReflectionEnvironmentCapture.cpp
+++ b/Engine/Source/Runtime/Renderer/Private/ReflectionEnvironmentCapture.cpp
@@ -1570,12 +1570,18 @@ void CopyToSkyTexture(FRHICommandList& RHICmdList, FScene* Scene, FTexture* Proc
 		FSceneRenderTargetItem& FilteredCube = FSceneRenderTargets::Get(RHICmdList).ReflectionColorScratchCubemap[1]->GetRenderTargetItem();
 
 		// GPU copy back to the skylight's texture, which is not a render target
-		FRHICopyTextureInfo CopyInfo(FilteredCube.ShaderResourceTexture->GetSizeXYZ());
-		CopyInfo.NumArraySlices = 6;
+		FRHICopyTextureInfo CopyInfo;
+		CopyInfo.Size = FilteredCube.ShaderResourceTexture->GetSizeXYZ();
+		CopyInfo.NumSlices = 6;
+
 		for (int32 MipIndex = 0; MipIndex < NumMips; MipIndex++)
 		{
 			RHICmdList.CopyTexture(FilteredCube.ShaderResourceTexture, ProcessedTexture->TextureRHI, CopyInfo);
-			CopyInfo.AdvanceMip();
+
+			++CopyInfo.SourceMipIndex;
+			++CopyInfo.DestMipIndex;
+			CopyInfo.Size.X = FMath::Max(1, CopyInfo.Size.X / 2);
+			CopyInfo.Size.Y = FMath::Max(1, CopyInfo.Size.Y / 2);
 		}
 
 		RHICmdList.TransitionResource(EResourceTransitionAccess::EReadable, ProcessedTexture->TextureRHI);
diff --git a/Engine/Source/Runtime/VulkanRHI/Private/VulkanTexture.cpp b/Engine/Source/Runtime/VulkanRHI/Private/VulkanTexture.cpp
--- a/Engine/Source/Runtime/VulkanRHI/Private/VulkanTexture.cpp
+++ b/Engine/Source/Runtime/VulkanRHI/Private/VulkanTexture.cpp
@@ -2189,16 +2189,16 @@ void FVulkanCommandListContext::RHICopyTexture(FTextureRHIParamRef SourceTexture
 	int32 DestBarrierIndex = Barrier.AddImageBarrier(DstSurface.Image, DstSurface.GetFullAspectMask(), 1);
 	{
 		VkImageSubresourceRange& Range = Barrier.GetSubresource(SourceBarrierIndex);
-		Range.baseMipLevel = CopyInfo.MipIndex;
-		Range.baseArrayLayer = CopyInfo.SourceArraySlice;
-		Range.layerCount = CopyInfo.NumArraySlices;
+		Range.baseMipLevel = CopyInfo.SourceMipIndex;
+		Range.baseArrayLayer = CopyInfo.SourceSliceIndex;
+		Range.layerCount = CopyInfo.NumSlices;
 		Barrier.SetTransition(SourceBarrierIndex, VulkanRHI::GetImageLayoutFromVulkanLayout(SrcLayout), EImageLayoutBarrier::TransferSource);
 	}
 	{
 		VkImageSubresourceRange& Range = Barrier.GetSubresource(DestBarrierIndex);
-		Range.baseMipLevel = CopyInfo.MipIndex;
-		Range.baseArrayLayer = CopyInfo.DestArraySlice;
-		Range.layerCount = CopyInfo.NumArraySlices;
+		Range.baseMipLevel = CopyInfo.DestMipIndex;
+		Range.baseArrayLayer = CopyInfo.DestSliceIndex;
+		Range.layerCount = CopyInfo.NumSlices;
 		Barrier.SetTransition(DestBarrierIndex, EImageLayoutBarrier::Undefined, EImageLayoutBarrier::TransferDest);
 	}
 
@@ -2207,17 +2207,17 @@ void FVulkanCommandListContext::RHICopyTexture(FTextureRHIParamRef SourceTexture
 	VkImageCopy Region;
 	FMemory::Memzero(Region);
 	ensure(SrcSurface.Width == DstSurface.Width && SrcSurface.Height == DstSurface.Height);
-	Region.extent.width = FMath::Max(1u, SrcSurface.Width>> CopyInfo.MipIndex);
-	Region.extent.height = FMath::Max(1u, SrcSurface.Height >> CopyInfo.MipIndex);
+	Region.extent.width = FMath::Max(1u, SrcSurface.Width>> CopyInfo.SourceMipIndex);
+	Region.extent.height = FMath::Max(1u, SrcSurface.Height >> CopyInfo.SourceMipIndex);
 	Region.extent.depth = 1;
 	Region.srcSubresource.aspectMask = SrcSurface.GetFullAspectMask();
-	Region.srcSubresource.baseArrayLayer = CopyInfo.SourceArraySlice;
-	Region.srcSubresource.layerCount = CopyInfo.NumArraySlices;
-	Region.srcSubresource.mipLevel = CopyInfo.MipIndex;
+	Region.srcSubresource.baseArrayLayer = CopyInfo.SourceSliceIndex;
+	Region.srcSubresource.layerCount = CopyInfo.NumSlices;
+	Region.srcSubresource.mipLevel = CopyInfo.SourceMipIndex;
 	Region.dstSubresource.aspectMask = DstSurface.GetFullAspectMask();
-	Region.dstSubresource.baseArrayLayer = CopyInfo.DestArraySlice;
-	Region.dstSubresource.layerCount = CopyInfo.NumArraySlices;
-	Region.dstSubresource.mipLevel = CopyInfo.MipIndex;
+	Region.dstSubresource.baseArrayLayer = CopyInfo.DestSliceIndex;
+	Region.dstSubresource.layerCount = CopyInfo.NumSlices;
+	Region.dstSubresource.mipLevel = CopyInfo.DestMipIndex;
 	VulkanRHI::vkCmdCopyImage(CmdBuffer,
 		SrcSurface.Image, VK_IMAGE_LAYOUT_TRANSFER_SRC_OPTIMAL,
 		DstSurface.Image, VK_IMAGE_LAYOUT_TRANSFER_DST_OPTIMAL,
diff --git a/Engine/Source/Runtime/Windows/D3D11RHI/Private/D3D11RHIPrivate.h b/Engine/Source/Runtime/Windows/D3D11RHI/Private/D3D11RHIPrivate.h
--- a/Engine/Source/Runtime/Windows/D3D11RHI/Private/D3D11RHIPrivate.h
+++ b/Engine/Source/Runtime/Windows/D3D11RHI/Private/D3D11RHIPrivate.h
@@ -490,6 +490,7 @@ public:
 	virtual void RHISetMultipleViewports(uint32 Count, const FViewportBounds* Data) final override;
 	virtual void RHIClearTinyUAV(FUnorderedAccessViewRHIParamRef UnorderedAccessViewRHI, const uint32* Values) final override;
 	virtual void RHICopyToResolveTarget(FTextureRHIParamRef SourceTexture, FTextureRHIParamRef DestTexture, const FResolveParams& ResolveParams) final override;
+	virtual void RHICopyTexture(FTextureRHIParamRef SourceTexture, FTextureRHIParamRef DestTexture, const FRHICopyTextureInfo& CopyInfo) final override;
 	virtual void RHITransitionResources(EResourceTransitionAccess TransitionType, FTextureRHIParamRef* InTextures, int32 NumTextures) final override;
 	virtual void RHITransitionResources(EResourceTransitionAccess TransitionType, EResourceTransitionPipeline TransitionPipeline, FUnorderedAccessViewRHIParamRef* InUAVs, int32 NumUAVs, FComputeFenceRHIParamRef WriteFence) final override;
 	virtual void RHIBeginRenderQuery(FRenderQueryRHIParamRef RenderQuery) final override;
diff --git a/Engine/Source/Runtime/Windows/D3D11RHI/Private/D3D11Texture.cpp b/Engine/Source/Runtime/Windows/D3D11RHI/Private/D3D11Texture.cpp
--- a/Engine/Source/Runtime/Windows/D3D11RHI/Private/D3D11Texture.cpp
+++ b/Engine/Source/Runtime/Windows/D3D11RHI/Private/D3D11Texture.cpp
@@ -2061,3 +2061,51 @@ void FD3D11DynamicRHI::RHIAliasTextureResources(FTextureRHIParamRef DestTextureR
 		DestTexture->AliasResources(SrcTexture);
 	}
 }
+
+void FD3D11DynamicRHI::RHICopyTexture(FTextureRHIParamRef SourceTextureRHI, FTextureRHIParamRef DestTextureRHI, const FRHICopyTextureInfo& CopyInfo)
+{
+	if (!SourceTextureRHI || !DestTextureRHI || SourceTextureRHI == DestTextureRHI)
+	{
+		// no need to do anything (silently ignored)
+		return;
+	}
+
+	RHITransitionResources(EResourceTransitionAccess::EReadable, &SourceTextureRHI, 1);
+
+	FRHICommandList_RecursiveHazardous RHICmdList(this);	
+
+	FD3D11TextureBase* SourceTexture = GetD3D11TextureFromRHITexture(SourceTextureRHI);
+	FD3D11TextureBase* DestTexture = GetD3D11TextureFromRHITexture(DestTextureRHI);
+
+	check(SourceTexture && DestTexture);
+
+	GPUProfilingData.RegisterGPUWork();
+
+	if (CopyInfo.Size != FIntVector::ZeroValue)
+	{
+		D3D11_BOX SrcBox;
+		SrcBox.left = CopyInfo.SourcePosition.X;
+		SrcBox.top = CopyInfo.SourcePosition.Y;
+		SrcBox.front = CopyInfo.SourcePosition.Z;
+		SrcBox.right = CopyInfo.SourcePosition.X + CopyInfo.Size.X;
+		SrcBox.bottom = CopyInfo.SourcePosition.Y + CopyInfo.Size.Y;
+		SrcBox.back = CopyInfo.SourcePosition.Z + CopyInfo.Size.Z;
+
+		for (uint32 i = 0; i < CopyInfo.NumSlices; ++i)
+		{
+			uint32 SourceSliceIndex = CopyInfo.SourceSliceIndex + i;
+			uint32 DestSliceIndex = CopyInfo.DestSliceIndex + i;
+
+			const uint32 SourceSubresource = D3D11CalcSubresource(CopyInfo.SourceMipIndex, SourceSliceIndex, SourceTextureRHI->GetNumMips());
+			const uint32 DestSubresource = D3D11CalcSubresource(CopyInfo.DestMipIndex, DestSliceIndex, DestTextureRHI->GetNumMips());
+
+			Direct3DDeviceIMContext->CopySubresourceRegion(DestTexture->GetResource(), DestSubresource, CopyInfo.DestPosition.X, CopyInfo.DestPosition.Y, CopyInfo.DestPosition.Z, SourceTexture->GetResource(), SourceSubresource, &SrcBox);
+		}
+	}
+	else
+	{
+		// Make sure the params are all by default when using this case
+		ensure(CopyInfo.SourceSliceIndex == 0 && CopyInfo.DestSliceIndex == 0 && CopyInfo.SourcePosition == FIntVector::ZeroValue && CopyInfo.DestPosition == FIntVector::ZeroValue && CopyInfo.SourceMipIndex == 0 && CopyInfo.DestMipIndex == 0);
+		Direct3DDeviceIMContext->CopyResource(DestTexture->GetResource(), SourceTexture->GetResource());
+	}
+}
\ No newline at end of file
