#DEV-VR - Fix for crash on Lumin Desktop Vulkan (rendering still has issues)

diff --git a/Engine/Source/Runtime/VulkanRHI/Private/VulkanViewport.cpp b/Engine/Source/Runtime/VulkanRHI/Private/VulkanViewport.cpp
--- a/Engine/Source/Runtime/VulkanRHI/Private/VulkanViewport.cpp
+++ b/Engine/Source/Runtime/VulkanRHI/Private/VulkanViewport.cpp
@@ -179,7 +179,10 @@ void FVulkanViewport::AcquireBackBuffer(FRHICommandListBase& CmdList, FVulkanBac
 
 	// Submit here so we can add a dependency with the acquired semaphore
 	CmdBuffer->End();
-	CmdBuffer->AddWaitSemaphore(VK_PIPELINE_STAGE_BOTTOM_OF_PIPE_BIT, AcquiredSemaphore);
+	if (FVulkanPlatform::SupportsStandardSwapchain())
+	{
+		CmdBuffer->AddWaitSemaphore(VK_PIPELINE_STAGE_BOTTOM_OF_PIPE_BIT, AcquiredSemaphore);
+	}
 	Device->GetGraphicsQueue()->Submit(CmdBuffer);
 	CmdBufferManager->PrepareForNewActiveCommandBuffer();
 }
@@ -615,12 +618,12 @@ bool FVulkanViewport::Present(FVulkanCommandListContext* Context, FVulkanCmdBuff
 
 	CmdBuffer->End();
 
-	if (DelayAcquireBackBuffer())
-	{
-		CmdBuffer->AddWaitSemaphore(VK_PIPELINE_STAGE_BOTTOM_OF_PIPE_BIT, AcquiredSemaphore);
-	}
 	if (FVulkanPlatform::SupportsStandardSwapchain())
 	{
+		if (DelayAcquireBackBuffer())
+		{
+			CmdBuffer->AddWaitSemaphore(VK_PIPELINE_STAGE_BOTTOM_OF_PIPE_BIT, AcquiredSemaphore);
+		}
 		Queue->Submit(CmdBuffer, RenderingDoneSemaphores[AcquiredImageIndex]->GetHandle());
 	}
 	else
