VR - Fix markers for Lumin

diff --git a/Engine/Source/Runtime/VulkanRHI/Private/Lumin/VulkanLuminPlatform.cpp b/Engine/Source/Runtime/VulkanRHI/Private/Lumin/VulkanLuminPlatform.cpp
--- a/Engine/Source/Runtime/VulkanRHI/Private/Lumin/VulkanLuminPlatform.cpp
+++ b/Engine/Source/Runtime/VulkanRHI/Private/Lumin/VulkanLuminPlatform.cpp
@@ -10,6 +10,7 @@ ENUM_VK_ENTRYPOINTS_ALL(DEFINE_VK_ENTRYPOINTS)
 
 #define CHECK_VK_ENTRYPOINTS(Type,Func) if (VulkanDynamicAPI::Func == NULL) { bFoundAllEntryPoints = false; UE_LOG(LogRHI, Warning, TEXT("Failed to find entry point for %s"), TEXT(#Func)); }
 
+static bool GFoundTegraGfxDebugger = false;
 
 void* FVulkanLuminPlatform::VulkanLib = nullptr;
 bool FVulkanLuminPlatform::bAttemptedLoad = false;
@@ -95,13 +96,39 @@ void FVulkanLuminPlatform::CreateSurface(void* WindowHandle, VkInstance Instance
 	//VERIFYVULKANRESULT(vkCreateLuminSurfaceKHR(Instance, &SurfaceCreateInfo, nullptr, OutSurface));
 }
 
+void FVulkanLuminPlatform::NotifyFoundInstanceLayersAndExtensions(const TArray<FString>& Layers, const TArray<FString>& Extensions)
+{
+	if (Extensions.Find(TEXT("VK_LAYER_NV_vgd")))
+	{
+		GFoundTegraGfxDebugger = true;
+	}
+}
 
+void FVulkanLuminPlatform::NotifyFoundDeviceLayersAndExtensions(VkPhysicalDevice PhysicalDevice, const TArray<FString>& Layers, const TArray<FString>& Extensions)
+{
+	if (Extensions.Find(TEXT("VK_LAYER_NV_vgd")))
+	{
+		GFoundTegraGfxDebugger = true;
+	}
+}
 
 void FVulkanLuminPlatform::GetInstanceExtensions(TArray<const ANSICHAR*>& OutExtensions)
 {
+	if (GFoundTegraGfxDebugger)
+	{
+		OutExtensions.Add("VK_LAYER_NV_vgd");
+	}
 }
 
 void FVulkanLuminPlatform::GetDeviceExtensions(TArray<const ANSICHAR*>& OutExtensions)
 {
+	if (GFoundTegraGfxDebugger)
+	{
+		OutExtensions.Add("VK_LAYER_NV_vgd");
+	}
 }
 
+bool FVulkanLuminPlatform::ForceEnableDebugMarkers()
+{
+	return GFoundTegraGfxDebugger;
+}
diff --git a/Engine/Source/Runtime/VulkanRHI/Private/Lumin/VulkanLuminPlatform.h b/Engine/Source/Runtime/VulkanRHI/Private/Lumin/VulkanLuminPlatform.h
--- a/Engine/Source/Runtime/VulkanRHI/Private/Lumin/VulkanLuminPlatform.h
+++ b/Engine/Source/Runtime/VulkanRHI/Private/Lumin/VulkanLuminPlatform.h
@@ -32,6 +32,9 @@ public:
 	static bool LoadVulkanInstanceFunctions(VkInstance inInstance);
 	static void FreeVulkanLibrary();
 
+	static void NotifyFoundInstanceLayersAndExtensions(const TArray<FString>& Layers, const TArray<FString>& Extensions);
+	static void NotifyFoundDeviceLayersAndExtensions(VkPhysicalDevice PhysicalDevice, const TArray<FString>& Layers, const TArray<FString>& Extensions);
+
 	static void GetInstanceExtensions(TArray<const ANSICHAR*>& OutExtensions);
 	static void GetDeviceExtensions(TArray<const ANSICHAR*>& OutExtensions);
 
@@ -44,7 +47,7 @@ public:
 	static bool SupportsStandardSwapchain() { return false; }
 	static EPixelFormat GetPixelFormatForNonDefaultSwapchain() { return PF_R8G8B8A8; }
 
-	static bool SupportsMarkersWithoutExtension() { return true; }
+	static bool ForceEnableDebugMarkers();
 
 protected:
 	static void* VulkanLib;
diff --git a/Engine/Source/Runtime/VulkanRHI/Private/VulkanDevice.cpp b/Engine/Source/Runtime/VulkanRHI/Private/VulkanDevice.cpp
--- a/Engine/Source/Runtime/VulkanRHI/Private/VulkanDevice.cpp
+++ b/Engine/Source/Runtime/VulkanRHI/Private/VulkanDevice.cpp
@@ -235,11 +235,17 @@ void FVulkanDevice::CreateDevice()
 	TransferQueue = new FVulkanQueue(this, TransferQueueFamilyIndex);
 
 #if VULKAN_ENABLE_DRAW_MARKERS
-	if (bDebugMarkersFound || FVulkanPlatform::SupportsMarkersWithoutExtension())
+	if (bDebugMarkersFound || FVulkanPlatform::ForceEnableDebugMarkers())
 	{
 		DebugMarkers.CmdBegin = (PFN_vkCmdDebugMarkerBeginEXT)(void*)VulkanRHI::vkGetDeviceProcAddr(Device, "vkCmdDebugMarkerBeginEXT");
 		DebugMarkers.CmdEnd = (PFN_vkCmdDebugMarkerEndEXT)(void*)VulkanRHI::vkGetDeviceProcAddr(Device, "vkCmdDebugMarkerEndEXT");
 		DebugMarkers.CmdSetObjectName = (PFN_vkDebugMarkerSetObjectNameEXT)(void*)VulkanRHI::vkGetDeviceProcAddr(Device, "vkDebugMarkerSetObjectNameEXT");
+
+		if (DebugMarkers.CmdBegin && DebugMarkers.CmdEnd && DebugMarkers.CmdSetObjectName)
+		{
+			bDebugMarkersFound = true;
+		}
+
 		if (!DebugMarkers.CmdBegin || !DebugMarkers.CmdEnd || !DebugMarkers.CmdSetObjectName)
 		{
 			UE_LOG(LogVulkanRHI, Warning, TEXT("Extension found, but entry points for vkCmdDebugMarker(Begin|End)EXT NOT found!"));
diff --git a/Engine/Source/Runtime/VulkanRHI/Private/VulkanGenericPlatform.h b/Engine/Source/Runtime/VulkanRHI/Private/VulkanGenericPlatform.h
--- a/Engine/Source/Runtime/VulkanRHI/Private/VulkanGenericPlatform.h
+++ b/Engine/Source/Runtime/VulkanRHI/Private/VulkanGenericPlatform.h
@@ -13,7 +13,11 @@ public:
 	static bool LoadVulkanInstanceFunctions(VkInstance inInstance) { return true; }
 	static void FreeVulkanLibrary() {}
 
-	// Array of extensions for the platform - required to implement!
+	// Called after querying all the available extensions and layers
+	static void NotifyFoundInstanceLayersAndExtensions(const TArray<FString>& Layers, const TArray<FString>& Extensions) {}
+	static void NotifyFoundDeviceLayersAndExtensions(VkPhysicalDevice PhysicalDevice, const TArray<FString>& Layers, const TArray<FString>& Extensions) {}
+
+	// Array of required extensions for the platform (Required!)
 	static void GetInstanceExtensions(TArray<const ANSICHAR*>& OutExtensions);
 	static void GetDeviceExtensions(TArray<const ANSICHAR*>& OutExtensions);
 
@@ -43,7 +47,7 @@ public:
 	// Some platforms have issues with the access flags for the Present layout
 	static bool RequiresPresentLayoutFix() { return false; }
 
-	static bool SupportsMarkersWithoutExtension() { return false; }
+	static bool ForceEnableDebugMarkers() { return false; }
 
 	static bool SupportsDeviceLocalHostVisibleWithNoPenalty() { return false; }
 
diff --git a/Engine/Source/Runtime/VulkanRHI/Private/VulkanLayers.cpp b/Engine/Source/Runtime/VulkanRHI/Private/VulkanLayers.cpp
--- a/Engine/Source/Runtime/VulkanRHI/Private/VulkanLayers.cpp
+++ b/Engine/Source/Runtime/VulkanRHI/Private/VulkanLayers.cpp
@@ -248,6 +248,8 @@ void FVulkanDynamicRHI::GetInstanceLayersAndExtensions(TArray<const ANSICHAR*>&
 		UE_LOG(LogVulkanRHI, Display, TEXT("- Found instance extension %s"), *Name);
 	}
 
+	FVulkanPlatform::NotifyFoundInstanceLayersAndExtensions(FoundUniqueLayers, FoundUniqueExtensions);
+
 	if (FParse::Param(FCommandLine::Get(), TEXT("vktrace")))
 	{
 		const char* VkTraceName = "VK_LAYER_LUNARG_vktrace";
@@ -417,6 +419,8 @@ void FVulkanDevice::GetDeviceExtensionsAndLayers(TArray<const ANSICHAR*>& OutDev
 		UE_LOG(LogVulkanRHI, Display, TEXT("- Found device extension %s"), *Name);
 	}
 
+	FVulkanPlatform::NotifyFoundDeviceLayersAndExtensions(Gpu, FoundUniqueLayers, FoundUniqueExtensions);
+
 	TArray<FString> UniqueUsedDeviceExtensions;
 	auto AddDeviceLayers = [&](const char* LayerName)
 	{
@@ -537,7 +541,7 @@ void FVulkanDevice::GetDeviceExtensionsAndLayers(TArray<const ANSICHAR*>& OutDev
 	}
 
 #if VULKAN_ENABLE_DRAW_MARKERS
-	if (!bOutDebugMarkers && ListContains(AvailableExtensions, VK_EXT_DEBUG_MARKER_EXTENSION_NAME) && (GRenderDocFound || FVulkanPlatform::SupportsMarkersWithoutExtension()))
+	if (!bOutDebugMarkers && ListContains(AvailableExtensions, VK_EXT_DEBUG_MARKER_EXTENSION_NAME) && (GRenderDocFound || FVulkanPlatform::ForceEnableDebugMarkers()))
 	{
 		OutDeviceExtensions.Add(VK_EXT_DEBUG_MARKER_EXTENSION_NAME);
 		bOutDebugMarkers = true;
