Rolando originally hacked the RHI for Lumin with the ForceEnableDebugMarkers() function; return false there instead of if PLATFORM_LUMIN.

diff --git a/Engine/Source/Runtime/VulkanRHI/Private/Lumin/VulkanLuminPlatform.cpp b/Engine/Source/Runtime/VulkanRHI/Private/Lumin/VulkanLuminPlatform.cpp
--- a/Engine/Source/Runtime/VulkanRHI/Private/Lumin/VulkanLuminPlatform.cpp
+++ b/Engine/Source/Runtime/VulkanRHI/Private/Lumin/VulkanLuminPlatform.cpp
@@ -130,5 +130,6 @@ void FVulkanLuminPlatform::GetDeviceExtensions(TArray<const ANSICHAR*>& OutExten
 
 bool FVulkanLuminPlatform::ForceEnableDebugMarkers()
 {
-	return GFoundTegraGfxDebugger;
+	// Preventing VK_EXT_DEBUG_MARKER from being enabled on Lumin, because the device doesn't support it.
+	return false; //GFoundTegraGfxDebugger;
 }
