#DEV-VR - Fixes for Android platform with new Lumin Vulkan Color Conversion Functions

diff --git a/Engine/Source/Runtime/VulkanRHI/Private/VulkanLoader.h b/Engine/Source/Runtime/VulkanRHI/Private/VulkanLoader.h
--- a/Engine/Source/Runtime/VulkanRHI/Private/VulkanLoader.h
+++ b/Engine/Source/Runtime/VulkanRHI/Private/VulkanLoader.h
@@ -5,6 +5,7 @@
 #define VK_NO_PROTOTYPES
 #include <vulkan.h>
 
+#if WITH_VULKAN_COLOR_CONVERSIONS
 // List all instance Vulkan entry points used by Unreal that need to be loaded manually
 #define ENUM_VK_ENTRYPOINTS_INSTANCE(EnumMacro) \
 	EnumMacro(PFN_vkDestroyInstance, vkDestroyInstance) \
@@ -151,6 +152,153 @@
 	EnumMacro(PFN_vkQueuePresentKHR, vkQueuePresentKHR) \
 	EnumMacro(PFN_vkCreateSamplerYcbcrConversionKHR, vkCreateSamplerYcbcrConversionKHR) \
 	EnumMacro(PFN_vkDestroySamplerYcbcrConversionKHR, vkDestroySamplerYcbcrConversionKHR)
+#else
+
+// List all instance Vulkan entry points used by Unreal that need to be loaded manually
+#define ENUM_VK_ENTRYPOINTS_INSTANCE(EnumMacro) \
+	EnumMacro(PFN_vkDestroyInstance, vkDestroyInstance) \
+	EnumMacro(PFN_vkEnumeratePhysicalDevices, vkEnumeratePhysicalDevices) \
+	EnumMacro(PFN_vkGetPhysicalDeviceFeatures, vkGetPhysicalDeviceFeatures) \
+	EnumMacro(PFN_vkGetPhysicalDeviceFormatProperties, vkGetPhysicalDeviceFormatProperties) \
+	EnumMacro(PFN_vkGetPhysicalDeviceImageFormatProperties, vkGetPhysicalDeviceImageFormatProperties) \
+	EnumMacro(PFN_vkGetPhysicalDeviceProperties, vkGetPhysicalDeviceProperties) \
+	EnumMacro(PFN_vkGetPhysicalDeviceQueueFamilyProperties, vkGetPhysicalDeviceQueueFamilyProperties) \
+	EnumMacro(PFN_vkGetPhysicalDeviceMemoryProperties, vkGetPhysicalDeviceMemoryProperties) \
+	EnumMacro(PFN_vkCreateDevice, vkCreateDevice) \
+	EnumMacro(PFN_vkDestroyDevice, vkDestroyDevice) \
+	EnumMacro(PFN_vkEnumerateDeviceExtensionProperties, vkEnumerateDeviceExtensionProperties) \
+	EnumMacro(PFN_vkEnumerateDeviceLayerProperties, vkEnumerateDeviceLayerProperties) \
+	EnumMacro(PFN_vkGetDeviceQueue, vkGetDeviceQueue) \
+	EnumMacro(PFN_vkQueueSubmit, vkQueueSubmit) \
+	EnumMacro(PFN_vkQueueWaitIdle, vkQueueWaitIdle) \
+	EnumMacro(PFN_vkDeviceWaitIdle, vkDeviceWaitIdle) \
+	EnumMacro(PFN_vkAllocateMemory, vkAllocateMemory) \
+	EnumMacro(PFN_vkFreeMemory, vkFreeMemory) \
+	EnumMacro(PFN_vkMapMemory, vkMapMemory) \
+	EnumMacro(PFN_vkUnmapMemory, vkUnmapMemory) \
+	EnumMacro(PFN_vkFlushMappedMemoryRanges, vkFlushMappedMemoryRanges) \
+	EnumMacro(PFN_vkInvalidateMappedMemoryRanges, vkInvalidateMappedMemoryRanges) \
+	EnumMacro(PFN_vkGetDeviceMemoryCommitment, vkGetDeviceMemoryCommitment) \
+	EnumMacro(PFN_vkBindBufferMemory, vkBindBufferMemory) \
+	EnumMacro(PFN_vkBindImageMemory, vkBindImageMemory) \
+	EnumMacro(PFN_vkGetBufferMemoryRequirements, vkGetBufferMemoryRequirements) \
+	EnumMacro(PFN_vkGetImageMemoryRequirements, vkGetImageMemoryRequirements) \
+	EnumMacro(PFN_vkGetImageSparseMemoryRequirements, vkGetImageSparseMemoryRequirements) \
+	EnumMacro(PFN_vkGetPhysicalDeviceSparseImageFormatProperties, vkGetPhysicalDeviceSparseImageFormatProperties) \
+	EnumMacro(PFN_vkQueueBindSparse, vkQueueBindSparse) \
+	EnumMacro(PFN_vkCreateFence, vkCreateFence) \
+	EnumMacro(PFN_vkDestroyFence, vkDestroyFence) \
+	EnumMacro(PFN_vkResetFences, vkResetFences) \
+	EnumMacro(PFN_vkGetFenceStatus, vkGetFenceStatus) \
+	EnumMacro(PFN_vkWaitForFences, vkWaitForFences) \
+	EnumMacro(PFN_vkCreateSemaphore, vkCreateSemaphore) \
+	EnumMacro(PFN_vkDestroySemaphore, vkDestroySemaphore) \
+	EnumMacro(PFN_vkCreateEvent, vkCreateEvent) \
+	EnumMacro(PFN_vkDestroyEvent, vkDestroyEvent) \
+	EnumMacro(PFN_vkGetEventStatus, vkGetEventStatus) \
+	EnumMacro(PFN_vkSetEvent, vkSetEvent) \
+	EnumMacro(PFN_vkResetEvent, vkResetEvent) \
+	EnumMacro(PFN_vkCreateQueryPool, vkCreateQueryPool) \
+	EnumMacro(PFN_vkDestroyQueryPool, vkDestroyQueryPool) \
+	EnumMacro(PFN_vkGetQueryPoolResults, vkGetQueryPoolResults) \
+	EnumMacro(PFN_vkCreateBuffer, vkCreateBuffer) \
+	EnumMacro(PFN_vkDestroyBuffer, vkDestroyBuffer) \
+	EnumMacro(PFN_vkCreateBufferView, vkCreateBufferView) \
+	EnumMacro(PFN_vkDestroyBufferView, vkDestroyBufferView) \
+	EnumMacro(PFN_vkCreateImage, vkCreateImage) \
+	EnumMacro(PFN_vkDestroyImage, vkDestroyImage) \
+	EnumMacro(PFN_vkGetImageSubresourceLayout, vkGetImageSubresourceLayout) \
+	EnumMacro(PFN_vkCreateImageView, vkCreateImageView) \
+	EnumMacro(PFN_vkDestroyImageView, vkDestroyImageView) \
+	EnumMacro(PFN_vkCreateShaderModule, vkCreateShaderModule) \
+	EnumMacro(PFN_vkDestroyShaderModule, vkDestroyShaderModule) \
+	EnumMacro(PFN_vkCreatePipelineCache, vkCreatePipelineCache) \
+	EnumMacro(PFN_vkDestroyPipelineCache, vkDestroyPipelineCache) \
+	EnumMacro(PFN_vkGetPipelineCacheData, vkGetPipelineCacheData) \
+	EnumMacro(PFN_vkMergePipelineCaches, vkMergePipelineCaches) \
+	EnumMacro(PFN_vkCreateGraphicsPipelines, vkCreateGraphicsPipelines) \
+	EnumMacro(PFN_vkCreateComputePipelines, vkCreateComputePipelines) \
+	EnumMacro(PFN_vkDestroyPipeline, vkDestroyPipeline) \
+	EnumMacro(PFN_vkCreatePipelineLayout, vkCreatePipelineLayout) \
+	EnumMacro(PFN_vkDestroyPipelineLayout, vkDestroyPipelineLayout) \
+	EnumMacro(PFN_vkCreateSampler, vkCreateSampler) \
+	EnumMacro(PFN_vkDestroySampler, vkDestroySampler) \
+	EnumMacro(PFN_vkCreateDescriptorSetLayout, vkCreateDescriptorSetLayout) \
+	EnumMacro(PFN_vkDestroyDescriptorSetLayout, vkDestroyDescriptorSetLayout) \
+	EnumMacro(PFN_vkCreateDescriptorPool, vkCreateDescriptorPool) \
+	EnumMacro(PFN_vkDestroyDescriptorPool, vkDestroyDescriptorPool) \
+	EnumMacro(PFN_vkResetDescriptorPool, vkResetDescriptorPool) \
+	EnumMacro(PFN_vkAllocateDescriptorSets, vkAllocateDescriptorSets) \
+	EnumMacro(PFN_vkFreeDescriptorSets, vkFreeDescriptorSets) \
+	EnumMacro(PFN_vkUpdateDescriptorSets, vkUpdateDescriptorSets) \
+	EnumMacro(PFN_vkCreateFramebuffer, vkCreateFramebuffer) \
+	EnumMacro(PFN_vkDestroyFramebuffer, vkDestroyFramebuffer) \
+	EnumMacro(PFN_vkCreateRenderPass, vkCreateRenderPass) \
+	EnumMacro(PFN_vkDestroyRenderPass, vkDestroyRenderPass) \
+	EnumMacro(PFN_vkGetRenderAreaGranularity, vkGetRenderAreaGranularity) \
+	EnumMacro(PFN_vkCreateCommandPool, vkCreateCommandPool) \
+	EnumMacro(PFN_vkDestroyCommandPool, vkDestroyCommandPool) \
+	EnumMacro(PFN_vkResetCommandPool, vkResetCommandPool) \
+	EnumMacro(PFN_vkAllocateCommandBuffers, vkAllocateCommandBuffers) \
+	EnumMacro(PFN_vkFreeCommandBuffers, vkFreeCommandBuffers) \
+	EnumMacro(PFN_vkBeginCommandBuffer, vkBeginCommandBuffer) \
+	EnumMacro(PFN_vkEndCommandBuffer, vkEndCommandBuffer) \
+	EnumMacro(PFN_vkResetCommandBuffer, vkResetCommandBuffer) \
+	EnumMacro(PFN_vkCmdBindPipeline, vkCmdBindPipeline) \
+	EnumMacro(PFN_vkCmdSetViewport, vkCmdSetViewport) \
+	EnumMacro(PFN_vkCmdSetScissor, vkCmdSetScissor) \
+	EnumMacro(PFN_vkCmdSetLineWidth, vkCmdSetLineWidth) \
+	EnumMacro(PFN_vkCmdSetDepthBias, vkCmdSetDepthBias) \
+	EnumMacro(PFN_vkCmdSetBlendConstants, vkCmdSetBlendConstants) \
+	EnumMacro(PFN_vkCmdSetDepthBounds, vkCmdSetDepthBounds) \
+	EnumMacro(PFN_vkCmdSetStencilCompareMask, vkCmdSetStencilCompareMask) \
+	EnumMacro(PFN_vkCmdSetStencilWriteMask, vkCmdSetStencilWriteMask) \
+	EnumMacro(PFN_vkCmdSetStencilReference, vkCmdSetStencilReference) \
+	EnumMacro(PFN_vkCmdBindDescriptorSets, vkCmdBindDescriptorSets) \
+	EnumMacro(PFN_vkCmdBindIndexBuffer, vkCmdBindIndexBuffer) \
+	EnumMacro(PFN_vkCmdBindVertexBuffers, vkCmdBindVertexBuffers) \
+	EnumMacro(PFN_vkCmdDraw, vkCmdDraw) \
+	EnumMacro(PFN_vkCmdDrawIndexed, vkCmdDrawIndexed) \
+	EnumMacro(PFN_vkCmdDrawIndirect, vkCmdDrawIndirect) \
+	EnumMacro(PFN_vkCmdDrawIndexedIndirect, vkCmdDrawIndexedIndirect) \
+	EnumMacro(PFN_vkCmdDispatch, vkCmdDispatch) \
+	EnumMacro(PFN_vkCmdDispatchIndirect, vkCmdDispatchIndirect) \
+	EnumMacro(PFN_vkCmdCopyBuffer, vkCmdCopyBuffer) \
+	EnumMacro(PFN_vkCmdCopyImage, vkCmdCopyImage) \
+	EnumMacro(PFN_vkCmdBlitImage, vkCmdBlitImage) \
+	EnumMacro(PFN_vkCmdCopyBufferToImage, vkCmdCopyBufferToImage) \
+	EnumMacro(PFN_vkCmdCopyImageToBuffer, vkCmdCopyImageToBuffer) \
+	EnumMacro(PFN_vkCmdUpdateBuffer, vkCmdUpdateBuffer) \
+	EnumMacro(PFN_vkCmdFillBuffer, vkCmdFillBuffer) \
+	EnumMacro(PFN_vkCmdClearColorImage, vkCmdClearColorImage) \
+	EnumMacro(PFN_vkCmdClearDepthStencilImage, vkCmdClearDepthStencilImage) \
+	EnumMacro(PFN_vkCmdClearAttachments, vkCmdClearAttachments) \
+	EnumMacro(PFN_vkCmdResolveImage, vkCmdResolveImage) \
+	EnumMacro(PFN_vkCmdSetEvent, vkCmdSetEvent) \
+	EnumMacro(PFN_vkCmdResetEvent, vkCmdResetEvent) \
+	EnumMacro(PFN_vkCmdWaitEvents, vkCmdWaitEvents) \
+	EnumMacro(PFN_vkCmdPipelineBarrier, vkCmdPipelineBarrier) \
+	EnumMacro(PFN_vkCmdBeginQuery, vkCmdBeginQuery) \
+	EnumMacro(PFN_vkCmdEndQuery, vkCmdEndQuery) \
+	EnumMacro(PFN_vkCmdResetQueryPool, vkCmdResetQueryPool) \
+	EnumMacro(PFN_vkCmdWriteTimestamp, vkCmdWriteTimestamp) \
+	EnumMacro(PFN_vkCmdCopyQueryPoolResults, vkCmdCopyQueryPoolResults) \
+	EnumMacro(PFN_vkCmdPushConstants, vkCmdPushConstants) \
+	EnumMacro(PFN_vkCmdBeginRenderPass, vkCmdBeginRenderPass) \
+	EnumMacro(PFN_vkCmdNextSubpass, vkCmdNextSubpass) \
+	EnumMacro(PFN_vkCmdEndRenderPass, vkCmdEndRenderPass) \
+	EnumMacro(PFN_vkCmdExecuteCommands, vkCmdExecuteCommands) \
+	EnumMacro(PFN_vkDestroySurfaceKHR, vkDestroySurfaceKHR) \
+	EnumMacro(PFN_vkGetPhysicalDeviceSurfaceSupportKHR, vkGetPhysicalDeviceSurfaceSupportKHR) \
+	EnumMacro(PFN_vkGetPhysicalDeviceSurfaceCapabilitiesKHR, vkGetPhysicalDeviceSurfaceCapabilitiesKHR) \
+	EnumMacro(PFN_vkGetPhysicalDeviceSurfaceFormatsKHR, vkGetPhysicalDeviceSurfaceFormatsKHR) \
+	EnumMacro(PFN_vkGetPhysicalDeviceSurfacePresentModesKHR, vkGetPhysicalDeviceSurfacePresentModesKHR) \
+	EnumMacro(PFN_vkCreateSwapchainKHR, vkCreateSwapchainKHR) \
+	EnumMacro(PFN_vkDestroySwapchainKHR, vkDestroySwapchainKHR) \
+	EnumMacro(PFN_vkGetSwapchainImagesKHR, vkGetSwapchainImagesKHR) \
+	EnumMacro(PFN_vkAcquireNextImageKHR, vkAcquireNextImageKHR) \
+	EnumMacro(PFN_vkQueuePresentKHR, vkQueuePresentKHR)
+#endif
 
 // List all base Vulkan entry points used by Unreal that need to be loaded manually
 #define ENUM_VK_ENTRYPOINTS_BASE(EnumMacro) \
diff --git a/Engine/Source/Runtime/VulkanRHI/Private/VulkanRHI.cpp b/Engine/Source/Runtime/VulkanRHI/Private/VulkanRHI.cpp
--- a/Engine/Source/Runtime/VulkanRHI/Private/VulkanRHI.cpp
+++ b/Engine/Source/Runtime/VulkanRHI/Private/VulkanRHI.cpp
@@ -895,10 +895,12 @@ FTexture2DRHIRef FVulkanDynamicRHI::RHICreateTexture2DFromResource(EPixelFormat
 	return new FVulkanTexture2D(*Device, Format, SizeX, SizeY, NumMips, NumSamples, NumSamplesTileMem, Resource, Flags, FRHIResourceCreateInfo());
 }
 
+#if WITH_VULKAN_COLOR_CONVERSIONS
 FTexture2DRHIRef FVulkanDynamicRHI::RHICreateTexture2DFromResource(EPixelFormat Format, uint32 SizeX, uint32 SizeY, uint32 NumMips, uint32 NumSamples, VkImage Resource, FSamplerYcbcrConversionInitializer& ConversionInitializer, uint32 Flags)
 {
 	return new FVulkanTexture2D(*Device, Format, SizeX, SizeY, NumMips, NumSamples, Resource, ConversionInitializer, Flags, FRHIResourceCreateInfo());
 }
+#endif
 
 FTexture2DArrayRHIRef FVulkanDynamicRHI::RHICreateTexture2DArrayFromResource(EPixelFormat Format, uint32 SizeX, uint32 SizeY, uint32 ArraySize, uint32 NumMips, VkImage Resource, uint32 Flags)
 {
diff --git a/Engine/Source/Runtime/VulkanRHI/Private/VulkanState.cpp b/Engine/Source/Runtime/VulkanRHI/Private/VulkanState.cpp
--- a/Engine/Source/Runtime/VulkanRHI/Private/VulkanState.cpp
+++ b/Engine/Source/Runtime/VulkanRHI/Private/VulkanState.cpp
@@ -368,6 +368,7 @@ FSamplerStateRHIRef FVulkanDynamicRHI::RHICreateSamplerState(const FSamplerState
 	}
 }
 
+#if WITH_VULKAN_COLOR_CONVERSIONS
 FSamplerStateRHIRef FVulkanDynamicRHI::RHICreateSamplerState(
 	const FSamplerStateInitializerRHI& Initializer, 
 	const FSamplerYcbcrConversionInitializer& ConversionInitializer)
@@ -403,6 +404,8 @@ FSamplerStateRHIRef FVulkanDynamicRHI::RHICreateSamplerState(
 
 	return new FVulkanSamplerState(SamplerInfo, *Device, true);
 }
+#endif
+
 
 FRasterizerStateRHIRef FVulkanDynamicRHI::RHICreateRasterizerState(const FRasterizerStateInitializerRHI& Initializer)
 {
diff --git a/Engine/Source/Runtime/VulkanRHI/Private/VulkanTexture.cpp b/Engine/Source/Runtime/VulkanRHI/Private/VulkanTexture.cpp
--- a/Engine/Source/Runtime/VulkanRHI/Private/VulkanTexture.cpp
+++ b/Engine/Source/Runtime/VulkanRHI/Private/VulkanTexture.cpp
@@ -1299,7 +1299,11 @@ void FVulkanDynamicRHI::InternalUpdateTexture3D(bool bFromRenderingThread, FText
 }
 
 
-VkImageView FVulkanTextureView::StaticCreate(FVulkanDevice& Device, VkImage Image, VkImageViewType ViewType, VkImageAspectFlags AspectFlags, EPixelFormat UEFormat, VkFormat Format, uint32 FirstMip, uint32 NumMips, uint32 ArraySliceIndex, uint32 NumArraySlices, bool bUseIdentitySwizzle, const FSamplerYcbcrConversionInitializer* ConversionInitializer)
+VkImageView FVulkanTextureView::StaticCreate(FVulkanDevice& Device, VkImage Image, VkImageViewType ViewType, VkImageAspectFlags AspectFlags, EPixelFormat UEFormat, VkFormat Format, uint32 FirstMip, uint32 NumMips, uint32 ArraySliceIndex, uint32 NumArraySlices, bool bUseIdentitySwizzle
+#if WITH_VULKAN_COLOR_CONVERSIONS
+	, const FSamplerYcbcrConversionInitializer* ConversionInitializer
+#endif
+)
 {
 	VkImageView View = VK_NULL_HANDLE;
 
@@ -1319,6 +1323,7 @@ VkImageView FVulkanTextureView::StaticCreate(FVulkanDevice& Device, VkImage Imag
 		ViewInfo.components = Device.GetFormatComponentMapping(UEFormat);
 	}
 
+#if WITH_VULKAN_COLOR_CONVERSIONS
 	if (ConversionInitializer != nullptr)
 	{
 		VkSamplerYcbcrConversionCreateInfo ConversionCreateInfo;
@@ -1347,6 +1352,7 @@ VkImageView FVulkanTextureView::StaticCreate(FVulkanDevice& Device, VkImage Imag
 		ConversionInfo.sType = VK_STRUCTURE_TYPE_SAMPLER_YCBCR_CONVERSION_INFO;
 		ViewInfo.pNext = &ConversionInfo;
 	}
+#endif
 
 	ViewInfo.subresourceRange.aspectMask = AspectFlags;
 	ViewInfo.subresourceRange.baseMipLevel = FirstMip;
@@ -1413,11 +1419,13 @@ void FVulkanTextureView::Create(FVulkanDevice& Device, VkImage InImage, VkImageV
 */
 }
 
+#if WITH_VULKAN_COLOR_CONVERSIONS
 void FVulkanTextureView::Create(FVulkanDevice& Device, VkImage InImage, VkImageViewType ViewType, VkImageAspectFlags AspectFlags, EPixelFormat UEFormat, VkFormat Format, uint32 FirstMip, uint32 NumMips, uint32 ArraySliceIndex, uint32 NumArraySlices, FSamplerYcbcrConversionInitializer& ConversionInitializer)
 {
 	View = StaticCreate(Device, InImage, ViewType, AspectFlags, UEFormat, Format, FirstMip, NumMips, ArraySliceIndex, NumArraySlices, false, &ConversionInitializer);
 	Image = InImage;
 }
+#endif
 
 void FVulkanTextureView::Destroy(FVulkanDevice& Device)
 {
@@ -1563,6 +1571,7 @@ FVulkanTextureBase::FVulkanTextureBase(FVulkanDevice& Device, VkImageViewType Re
 	}
 }
 
+#if WITH_VULKAN_COLOR_CONVERSIONS
 FVulkanTextureBase::FVulkanTextureBase(FVulkanDevice& Device, VkImageViewType ResourceType, EPixelFormat Format, uint32 SizeX, uint32 SizeY, uint32 SizeZ, uint32 NumMips, uint32 NumSamples, VkImage InImage, VkDeviceMemory InMem, FSamplerYcbcrConversionInitializer& ConversionInitializer, uint32 UEFlags, const FRHIResourceCreateInfo& CreateInfo)
 	: Surface(Device, ResourceType, Format, SizeX, SizeY, SizeZ, NumMips, NumSamples, InImage, UEFlags, CreateInfo)
 	, PartialView(nullptr)
@@ -1595,6 +1604,7 @@ FVulkanTextureBase::FVulkanTextureBase(FVulkanDevice& Device, VkImageViewType Re
 		PartialView->Create(Device, Surface.Image, Surface.ViewType, Surface.PartialAspectMask, Surface.PixelFormat, Surface.ViewFormat, 0, FMath::Max(NumMips, 1u), 0, FMath::Max(1u, SizeZ), ConversionInitializer);
 	}
 }
+#endif
 
 FVulkanTextureBase::~FVulkanTextureBase()
 {
@@ -1677,11 +1687,13 @@ FVulkanTexture2D::FVulkanTexture2D(FVulkanDevice& Device, EPixelFormat Format, u
 {
 }
 
+#if WITH_VULKAN_COLOR_CONVERSIONS
 FVulkanTexture2D::FVulkanTexture2D(FVulkanDevice& Device, EPixelFormat Format, uint32 SizeX, uint32 SizeY, uint32 NumMips, uint32 NumSamples, VkImage Image, FSamplerYcbcrConversionInitializer& ConversionInitializer, uint32 UEFlags, const FRHIResourceCreateInfo& CreateInfo)
 	: FRHITexture2D(SizeX, SizeY, NumMips, NumSamples, Format, UEFlags, CreateInfo.ClearValueBinding)
 	, FVulkanTextureBase(Device, VK_IMAGE_VIEW_TYPE_2D, Format, SizeX, SizeY, 1, NumMips, NumSamples, Image, VK_NULL_HANDLE, ConversionInitializer, UEFlags)
 {
 }
+#endif
 
 FVulkanTexture2D::~FVulkanTexture2D()
 {
diff --git a/Engine/Source/Runtime/VulkanRHI/Public/VulkanDynamicRHI.h b/Engine/Source/Runtime/VulkanRHI/Public/VulkanDynamicRHI.h
--- a/Engine/Source/Runtime/VulkanRHI/Public/VulkanDynamicRHI.h
+++ b/Engine/Source/Runtime/VulkanRHI/Public/VulkanDynamicRHI.h
@@ -34,7 +34,9 @@ public:
 	void InitInstance();
 
 	virtual FSamplerStateRHIRef RHICreateSamplerState(const FSamplerStateInitializerRHI& Initializer) final override;
+#if WITH_VULKAN_COLOR_CONVERSIONS
 	virtual FSamplerStateRHIRef RHICreateSamplerState(const FSamplerStateInitializerRHI& Initializer, const FSamplerYcbcrConversionInitializer& ConversionInitializer);
+#endif
 	virtual FRasterizerStateRHIRef RHICreateRasterizerState(const FRasterizerStateInitializerRHI& Initializer) final override;
 	virtual FDepthStencilStateRHIRef RHICreateDepthStencilState(const FDepthStencilStateInitializerRHI& Initializer) final override;
 	virtual FBlendStateRHIRef RHICreateBlendState(const FBlendStateInitializerRHI& Initializer) final override;
@@ -332,7 +334,9 @@ public:
 
 	// FVulkanDynamicRHI interface
 	virtual FTexture2DRHIRef RHICreateTexture2DFromResource(EPixelFormat Format, uint32 SizeX, uint32 SizeY, uint32 NumMips, uint32 NumSamples, uint32 NumSamplesTileMem, VkImage Resource, uint32 Flags);
+#if WITH_VULKAN_COLOR_CONVERSIONS
 	virtual FTexture2DRHIRef RHICreateTexture2DFromResource(EPixelFormat Format, uint32 SizeX, uint32 SizeY, uint32 NumMips, uint32 NumSamples, VkImage Resource, FSamplerYcbcrConversionInitializer& ConversionInitializer, uint32 Flags);
+#endif
 	virtual FTexture2DArrayRHIRef RHICreateTexture2DArrayFromResource(EPixelFormat Format, uint32 SizeX, uint32 SizeY, uint32 ArraySize, uint32 NumMips, VkImage Resource, uint32 Flags);
 	virtual FTextureCubeRHIRef RHICreateTextureCubeFromResource(EPixelFormat Format, uint32 Size, bool bArray, uint32 ArraySize, uint32 NumMips, VkImage Resource, uint32 Flags);
 	virtual void RHIAliasTextureResources(FTextureRHIParamRef DestTexture, FTextureRHIParamRef SrcTexture);
diff --git a/Engine/Source/Runtime/VulkanRHI/Public/VulkanResources.h b/Engine/Source/Runtime/VulkanRHI/Public/VulkanResources.h
--- a/Engine/Source/Runtime/VulkanRHI/Public/VulkanResources.h
+++ b/Engine/Source/Runtime/VulkanRHI/Public/VulkanResources.h
@@ -38,6 +38,7 @@ enum
 	NUM_TIMESTAMP_QUERIES_PER_POOL = 1024,
 };
 
+#if WITH_VULKAN_COLOR_CONVERSIONS
 struct FSamplerYcbcrConversionInitializer
 {
 	VkFormat Format;
@@ -48,6 +49,13 @@ struct FSamplerYcbcrConversionInitializer
 	VkChromaLocation XOffset;
 	VkChromaLocation YOffset;
 };
+//#else
+//place holder to allow similar function signatures
+//struct FSamplerYcbcrConversionInitializer
+//{
+//	//uint32 StubParam;
+//};
+#endif
 
 /** This represents a vertex declaration that hasn't been combined with a specific shader to create a bound shader. */
 class FVulkanVertexDeclaration : public FRHIVertexDeclaration
@@ -378,10 +386,16 @@ struct FVulkanTextureView
 	{
 	}
 
-	static VkImageView StaticCreate(FVulkanDevice& Device, VkImage InImage, VkImageViewType ViewType, VkImageAspectFlags AspectFlags, EPixelFormat UEFormat, VkFormat Format, uint32 FirstMip, uint32 NumMips, uint32 ArraySliceIndex, uint32 NumArraySlices, bool bUseIdentitySwizzle = false, const FSamplerYcbcrConversionInitializer* ConversionInitializer = nullptr);
+	static VkImageView StaticCreate(FVulkanDevice& Device, VkImage InImage, VkImageViewType ViewType, VkImageAspectFlags AspectFlags, EPixelFormat UEFormat, VkFormat Format, uint32 FirstMip, uint32 NumMips, uint32 ArraySliceIndex, uint32 NumArraySlices, bool bUseIdentitySwizzle = false
+#if WITH_VULKAN_COLOR_CONVERSIONS
+		, const FSamplerYcbcrConversionInitializer* ConversionInitializer = nullptr
+#endif
+	);
 
 	void Create(FVulkanDevice& Device, VkImage InImage, VkImageViewType ViewType, VkImageAspectFlags AspectFlags, EPixelFormat UEFormat, VkFormat Format, uint32 FirstMip, uint32 NumMips, uint32 ArraySliceIndex, uint32 NumArraySlices);
+#if WITH_VULKAN_COLOR_CONVERSIONS
 	void Create(FVulkanDevice& Device, VkImage InImage, VkImageViewType ViewType, VkImageAspectFlags AspectFlags, EPixelFormat UEFormat, VkFormat Format, uint32 FirstMip, uint32 NumMips, uint32 ArraySliceIndex, uint32 NumArraySlices, FSamplerYcbcrConversionInitializer& ConversionInitializer);
+#endif
 	void Destroy(FVulkanDevice& Device);
 
 	VkImageView View;
@@ -406,7 +420,9 @@ struct FVulkanTextureBase : public FVulkanBaseShaderResource
 
 	FVulkanTextureBase(FVulkanDevice& Device, VkImageViewType ResourceType, EPixelFormat Format, uint32 SizeX, uint32 SizeY, uint32 SizeZ, bool bArray, uint32 ArraySize, uint32 NumMips, uint32 NumSamples, uint32 UEFlags, const FRHIResourceCreateInfo& CreateInfo);
 	FVulkanTextureBase(FVulkanDevice& Device, VkImageViewType ResourceType, EPixelFormat Format, uint32 SizeX, uint32 SizeY, uint32 SizeZ, uint32 NumMips, uint32 NumSamples, uint32 NumSamplesTileMem, VkImage InImage, VkDeviceMemory InMem, uint32 UEFlags, const FRHIResourceCreateInfo& CreateInfo = FRHIResourceCreateInfo());
+#if WITH_VULKAN_COLOR_CONVERSIONS
 	FVulkanTextureBase(FVulkanDevice& Device, VkImageViewType ResourceType, EPixelFormat Format, uint32 SizeX, uint32 SizeY, uint32 SizeZ, uint32 NumMips, uint32 NumSamples, VkImage InImage, VkDeviceMemory InMem, FSamplerYcbcrConversionInitializer& ConversionInitializer, uint32 UEFlags, const FRHIResourceCreateInfo& CreateInfo = FRHIResourceCreateInfo());
+#endif
 	virtual ~FVulkanTextureBase();
 
 	VkImageView CreateRenderTargetView(uint32 MipIndex, uint32 NumMips, uint32 ArraySliceIndex, uint32 NumArraySlices);
@@ -437,7 +453,9 @@ class FVulkanTexture2D : public FRHITexture2D, public FVulkanTextureBase
 public:
 	FVulkanTexture2D(FVulkanDevice& Device, EPixelFormat Format, uint32 SizeX, uint32 SizeY, uint32 NumMips, uint32 NumSamples, uint32 UEFlags, const FRHIResourceCreateInfo& CreateInfo);
 	FVulkanTexture2D(FVulkanDevice& Device, EPixelFormat Format, uint32 SizeX, uint32 SizeY, uint32 NumMips, uint32 NumSamples, uint32 NumSamplesTileMem, VkImage Image, uint32 UEFlags, const FRHIResourceCreateInfo& CreateInfo);
+#if WITH_VULKAN_COLOR_CONVERSIONS
 	FVulkanTexture2D(FVulkanDevice& Device, EPixelFormat Format, uint32 SizeX, uint32 SizeY, uint32 NumMips, uint32 NumSamples, VkImage Image, struct FSamplerYcbcrConversionInitializer& ConversionInitializer, uint32 UEFlags, const FRHIResourceCreateInfo& CreateInfo);
+#endif
 	virtual ~FVulkanTexture2D();
 
 	// IRefCountedObject interface.
diff --git a/Engine/Source/Runtime/VulkanRHI/VulkanRHI.Build.cs b/Engine/Source/Runtime/VulkanRHI/VulkanRHI.Build.cs
--- a/Engine/Source/Runtime/VulkanRHI/VulkanRHI.Build.cs
+++ b/Engine/Source/Runtime/VulkanRHI/VulkanRHI.Build.cs
@@ -34,6 +34,8 @@ public class VulkanRHI : ModuleRules
 			}
 		);
 
+		bool bWithVulkanColorConversion = false;
+
 		if (Target.Platform == UnrealTargetPlatform.Win32 || Target.Platform == UnrealTargetPlatform.Win64)
 		{
 			string VulkanSDKPath = Environment.GetEnvironmentVariable("VULKAN_SDK");
@@ -118,6 +120,8 @@ public class VulkanRHI : ModuleRules
 			else if (Target.Platform == UnrealTargetPlatform.Lumin)
 			{
 				PrivateIncludePaths.Add(Target.UEThirdPartySourceDirectory + "Vulkan/Include/vulkan");
+				bWithVulkanColorConversion = true;
+				Log.TraceInformation("Including Vulkan Color Conversions");
 			}
 			else if (!String.IsNullOrEmpty(VulkanSDKPath))
 			{
@@ -146,6 +150,7 @@ public class VulkanRHI : ModuleRules
 		{
 			PrecompileForTargets = PrecompileTargetsType.None;
 		}
+		PublicDefinitions.Add("WITH_VULKAN_COLOR_CONVERSIONS=" + (bWithVulkanColorConversion ? "1" : "0"));
 	}
 
 	static int GetVersionFromString(string Text)
