Making sure bHaveVulkan is true for Lumin

diff --git a/Engine/Source/Runtime/VulkanRHI/VulkanRHI.Build.cs b/Engine/Source/Runtime/VulkanRHI/VulkanRHI.Build.cs
--- a/Engine/Source/Runtime/VulkanRHI/VulkanRHI.Build.cs
+++ b/Engine/Source/Runtime/VulkanRHI/VulkanRHI.Build.cs
@@ -120,6 +120,7 @@ public class VulkanRHI : ModuleRules
 			else if (Target.Platform == UnrealTargetPlatform.Lumin)
 			{
 				PrivateIncludePaths.Add(Target.UEThirdPartySourceDirectory + "Vulkan/Include/vulkan");
+				bHaveVulkan = true;
 				bWithVulkanColorConversion = true;
 				Log.TraceInformation("Including Vulkan Color Conversions");
 			}
