Moving color conversion vk entry points to lumin platform to avoid failing to find them in android drivers.

diff --git a/Engine/Source/Runtime/VulkanRHI/Private/Lumin/VulkanLuminPlatform.h b/Engine/Source/Runtime/VulkanRHI/Private/Lumin/VulkanLuminPlatform.h
--- a/Engine/Source/Runtime/VulkanRHI/Private/Lumin/VulkanLuminPlatform.h
+++ b/Engine/Source/Runtime/VulkanRHI/Private/Lumin/VulkanLuminPlatform.h
@@ -12,11 +12,14 @@
 #define VULKAN_USE_PER_PIPELINE_DESCRIPTOR_POOLS	0
 #define VULKAN_USE_DESCRIPTOR_POOL_MANAGER			0
 #define VULKAN_USE_IMAGE_ACQUIRE_FENCES				0
+#define VULKAN_SUPPORTS_COLOR_CONVERSIONS			1
 
 
 #define ENUM_VK_ENTRYPOINTS_PLATFORM_BASE(EnumMacro)
 
-#define ENUM_VK_ENTRYPOINTS_PLATFORM_INSTANCE(EnumMacro)
+#define ENUM_VK_ENTRYPOINTS_PLATFORM_INSTANCE(EnumMacro) \
+	EnumMacro(PFN_vkCreateSamplerYcbcrConversionKHR, vkCreateSamplerYcbcrConversionKHR) \
+	EnumMacro(PFN_vkDestroySamplerYcbcrConversionKHR, vkDestroySamplerYcbcrConversionKHR)
 
 #include "../VulkanLoader.h"
 
diff --git a/Engine/Source/Runtime/VulkanRHI/Private/VulkanDevice.cpp b/Engine/Source/Runtime/VulkanRHI/Private/VulkanDevice.cpp
--- a/Engine/Source/Runtime/VulkanRHI/Private/VulkanDevice.cpp
+++ b/Engine/Source/Runtime/VulkanRHI/Private/VulkanDevice.cpp
@@ -601,8 +601,10 @@ void FVulkanDevice::SetupFormats()
 	}
 }
 
+#if VULKAN_SUPPORTS_COLOR_CONVERSIONS
 VkSamplerYcbcrConversion FVulkanDevice::CreateSamplerColorConversion(const VkSamplerYcbcrConversionCreateInfo& CreateInfo)
 {
+
 	const uint32 CreateInfoHash = FCrc::MemCrc32(&CreateInfo, sizeof(CreateInfo));
 	VkSamplerYcbcrConversion* const FindResult = SamplerColorConversionMap.Find(CreateInfoHash);
 	if (FindResult != nullptr)
@@ -617,6 +619,7 @@ VkSamplerYcbcrConversion FVulkanDevice::CreateSamplerColorConversion(const VkSam
 		return NewConversion;
 	}
 }
+#endif
 
 void FVulkanDevice::MapFormatSupport(EPixelFormat UEFormat, VkFormat VulkanFormat)
 {
@@ -852,11 +855,13 @@ void FVulkanDevice::Destroy()
 	delete DefaultImage;
 	DefaultImage = nullptr;
 
+#if VULKAN_SUPPORTS_COLOR_CONVERSIONS
 	for (const auto& Pair : SamplerColorConversionMap)
 	{
 		vkDestroySamplerYcbcrConversionKHR(GetInstanceHandle(), Pair.Value, nullptr);
 	}
 	SamplerColorConversionMap.Reset();
+#endif
 
 	for (int32 Index = CommandContexts.Num() - 1; Index >= 0; --Index)
 	{
diff --git a/Engine/Source/Runtime/VulkanRHI/Private/VulkanDevice.h b/Engine/Source/Runtime/VulkanRHI/Private/VulkanDevice.h
--- a/Engine/Source/Runtime/VulkanRHI/Private/VulkanDevice.h
+++ b/Engine/Source/Runtime/VulkanRHI/Private/VulkanDevice.h
@@ -300,7 +300,9 @@ public:
 
 	void SetupPresentQueue(VkSurfaceKHR Surface);
 
+#if VULKAN_SUPPORTS_COLOR_CONVERSIONS
 	VkSamplerYcbcrConversion CreateSamplerColorConversion(const VkSamplerYcbcrConversionCreateInfo& CreateInfo);
+#endif
 
 private:
 	void MapFormatSupport(EPixelFormat UEFormat, VkFormat VulkanFormat);
diff --git a/Engine/Source/Runtime/VulkanRHI/Private/VulkanLoader.h b/Engine/Source/Runtime/VulkanRHI/Private/VulkanLoader.h
--- a/Engine/Source/Runtime/VulkanRHI/Private/VulkanLoader.h
+++ b/Engine/Source/Runtime/VulkanRHI/Private/VulkanLoader.h
@@ -148,9 +148,7 @@
 	EnumMacro(PFN_vkDestroySwapchainKHR, vkDestroySwapchainKHR) \
 	EnumMacro(PFN_vkGetSwapchainImagesKHR, vkGetSwapchainImagesKHR) \
 	EnumMacro(PFN_vkAcquireNextImageKHR, vkAcquireNextImageKHR) \
-	EnumMacro(PFN_vkQueuePresentKHR, vkQueuePresentKHR) \
-	EnumMacro(PFN_vkCreateSamplerYcbcrConversionKHR, vkCreateSamplerYcbcrConversionKHR) \
-	EnumMacro(PFN_vkDestroySamplerYcbcrConversionKHR, vkDestroySamplerYcbcrConversionKHR)
+	EnumMacro(PFN_vkQueuePresentKHR, vkQueuePresentKHR)
 
 // List all base Vulkan entry points used by Unreal that need to be loaded manually
 #define ENUM_VK_ENTRYPOINTS_BASE(EnumMacro) \
diff --git a/Engine/Source/Runtime/VulkanRHI/Private/VulkanState.cpp b/Engine/Source/Runtime/VulkanRHI/Private/VulkanState.cpp
--- a/Engine/Source/Runtime/VulkanRHI/Private/VulkanState.cpp
+++ b/Engine/Source/Runtime/VulkanRHI/Private/VulkanState.cpp
@@ -366,6 +366,7 @@ FSamplerStateRHIRef FVulkanDynamicRHI::RHICreateSamplerState(const FSamplerState
 	}
 }
 
+#if VULKAN_SUPPORTS_COLOR_CONVERSIONS
 FSamplerStateRHIRef FVulkanDynamicRHI::RHICreateSamplerState(
 	const FSamplerStateInitializerRHI& Initializer, 
 	const FSamplerYcbcrConversionInitializer& ConversionInitializer)
@@ -400,7 +401,7 @@ FSamplerStateRHIRef FVulkanDynamicRHI::RHICreateSamplerState(
 
 	return new FVulkanSamplerState(SamplerInfo, *Device, true);
 }
-
+#endif
 
 FRasterizerStateRHIRef FVulkanDynamicRHI::RHICreateRasterizerState(const FRasterizerStateInitializerRHI& Initializer)
 {
diff --git a/Engine/Source/Runtime/VulkanRHI/Private/VulkanTexture.cpp b/Engine/Source/Runtime/VulkanRHI/Private/VulkanTexture.cpp
--- a/Engine/Source/Runtime/VulkanRHI/Private/VulkanTexture.cpp
+++ b/Engine/Source/Runtime/VulkanRHI/Private/VulkanTexture.cpp
@@ -1297,7 +1297,6 @@ void FVulkanDynamicRHI::InternalUpdateTexture3D(bool bFromRenderingThread, FText
 	}
 }
 
-
 VkImageView FVulkanTextureView::StaticCreate(FVulkanDevice& Device, VkImage Image, VkImageViewType ViewType, VkImageAspectFlags AspectFlags, EPixelFormat UEFormat, VkFormat Format, uint32 FirstMip, uint32 NumMips, uint32 ArraySliceIndex, uint32 NumArraySlices, bool bUseIdentitySwizzle, const FSamplerYcbcrConversionInitializer* ConversionInitializer)
 {
 	VkImageView View = VK_NULL_HANDLE;
@@ -1316,6 +1315,7 @@ VkImageView FVulkanTextureView::StaticCreate(FVulkanDevice& Device, VkImage Imag
 		ViewInfo.components = Device.GetFormatComponentMapping(UEFormat);
 	}
 
+#if VULKAN_SUPPORTS_COLOR_CONVERSIONS
 	if (ConversionInitializer != nullptr)
 	{
 		VkSamplerYcbcrConversionCreateInfo ConversionCreateInfo;
@@ -1343,6 +1343,7 @@ VkImageView FVulkanTextureView::StaticCreate(FVulkanDevice& Device, VkImage Imag
 		ConversionInfo.sType = VK_STRUCTURE_TYPE_SAMPLER_YCBCR_CONVERSION_INFO;
 		ViewInfo.pNext = &ConversionInfo;
 	}
+#endif
 
 	ViewInfo.subresourceRange.aspectMask = AspectFlags;
 	ViewInfo.subresourceRange.baseMipLevel = FirstMip;
diff --git a/Engine/Source/Runtime/VulkanRHI/Public/VulkanConfiguration.h b/Engine/Source/Runtime/VulkanRHI/Public/VulkanConfiguration.h
--- a/Engine/Source/Runtime/VulkanRHI/Public/VulkanConfiguration.h
+++ b/Engine/Source/Runtime/VulkanRHI/Public/VulkanConfiguration.h
@@ -128,6 +128,11 @@
 	#define VULKAN_SUPPORTS_DEBUG_UTILS							0
 #endif
 
+#ifndef VULKAN_SUPPORTS_COLOR_CONVERSIONS
+	#define VULKAN_SUPPORTS_COLOR_CONVERSIONS					0
+#endif
+
+
 #define VULKAN_ENABLE_TRACKING_LAYER							(VULKAN_ENABLE_BUFFER_TRACKING_LAYER || VULKAN_ENABLE_IMAGE_TRACKING_LAYER)
 #define VULKAN_ENABLE_CUSTOM_LAYER								(VULKAN_ENABLE_DUMP_LAYER || VULKAN_ENABLE_TRACKING_LAYER)
 
diff --git a/Engine/Source/Runtime/VulkanRHI/Public/VulkanDynamicRHI.h b/Engine/Source/Runtime/VulkanRHI/Public/VulkanDynamicRHI.h
--- a/Engine/Source/Runtime/VulkanRHI/Public/VulkanDynamicRHI.h
+++ b/Engine/Source/Runtime/VulkanRHI/Public/VulkanDynamicRHI.h
@@ -34,7 +34,9 @@ public:
 	void InitInstance();
 
 	virtual FSamplerStateRHIRef RHICreateSamplerState(const FSamplerStateInitializerRHI& Initializer) final override;
+#if VULKAN_SUPPORTS_COLOR_CONVERSIONS
 	virtual FSamplerStateRHIRef RHICreateSamplerState(const FSamplerStateInitializerRHI& Initializer, const FSamplerYcbcrConversionInitializer& ConversionInitializer);
+#endif
 	virtual FRasterizerStateRHIRef RHICreateRasterizerState(const FRasterizerStateInitializerRHI& Initializer) final override;
 	virtual FDepthStencilStateRHIRef RHICreateDepthStencilState(const FDepthStencilStateInitializerRHI& Initializer) final override;
 	virtual FBlendStateRHIRef RHICreateBlendState(const FBlendStateInitializerRHI& Initializer) final override;
