Marker API override for ML VK.
The current ML driver unofficially supports debug markers, but will report the extension unsupported if queried directly.
Once they add official support, we need to add the extension name back into the list.

diff --git a/Engine/Source/Runtime/VulkanRHI/Private/VulkanLayers.cpp b/Engine/Source/Runtime/VulkanRHI/Private/VulkanLayers.cpp
--- a/Engine/Source/Runtime/VulkanRHI/Private/VulkanLayers.cpp
+++ b/Engine/Source/Runtime/VulkanRHI/Private/VulkanLayers.cpp
@@ -587,10 +587,12 @@ void FVulkanDevice::GetDeviceExtensionsAndLayers(TArray<const ANSICHAR*>& OutDev
 	}
 
 #if VULKAN_ENABLE_DRAW_MARKERS
-	if (!bOutDebugMarkers &&
+	if (!bOutDebugMarkers && 
 		((GValidationCvar.GetValueOnAnyThread() == 0 && ListContains(AvailableExtensions, VK_EXT_DEBUG_MARKER_EXTENSION_NAME)) || FVulkanPlatform::ForceEnableDebugMarkers()))
 	{
-		OutDeviceExtensions.Add(VK_EXT_DEBUG_MARKER_EXTENSION_NAME);
+		// Nvidia unofficially supports this extension, but will return false if we try to load it explicitly.
+		// We will avoid adding it for now and assume it works. Revisit when they update their drivers.
+		//OutDeviceExtensions.Add(VK_EXT_DEBUG_MARKER_EXTENSION_NAME);
 		bOutDebugMarkers = true;
 	}
 #endif
