We need to add the debug marker extension for all platforms but Lumin.

diff --git a/Engine/Source/Runtime/VulkanRHI/Private/VulkanLayers.cpp b/Engine/Source/Runtime/VulkanRHI/Private/VulkanLayers.cpp
--- a/Engine/Source/Runtime/VulkanRHI/Private/VulkanLayers.cpp
+++ b/Engine/Source/Runtime/VulkanRHI/Private/VulkanLayers.cpp
@@ -590,9 +590,10 @@ void FVulkanDevice::GetDeviceExtensionsAndLayers(TArray<const ANSICHAR*>& OutDev
 	if (!bOutDebugMarkers && 
 		((GValidationCvar.GetValueOnAnyThread() == 0 && ListContains(AvailableExtensions, VK_EXT_DEBUG_MARKER_EXTENSION_NAME)) || FVulkanPlatform::ForceEnableDebugMarkers()))
 	{
-		// Nvidia unofficially supports this extension, but will return false if we try to load it explicitly.
-		// We will avoid adding it for now and assume it works. Revisit when they update their drivers.
-		//OutDeviceExtensions.Add(VK_EXT_DEBUG_MARKER_EXTENSION_NAME);
+		// HACK: Lumin Nvidia driver unofficially supports this extension, but will return false if we try to load it explicitly.
+#if !PLATFORM_LUMIN
+		OutDeviceExtensions.Add(VK_EXT_DEBUG_MARKER_EXTENSION_NAME);
+#endif
 		bOutDebugMarkers = true;
 	}
 #endif
