Fix android vulkan crash when backbuffer is used for sampling

diff --git a/Engine/Source/Runtime/VulkanRHI/Private/VulkanViewport.cpp b/Engine/Source/Runtime/VulkanRHI/Private/VulkanViewport.cpp
--- a/Engine/Source/Runtime/VulkanRHI/Private/VulkanViewport.cpp
+++ b/Engine/Source/Runtime/VulkanRHI/Private/VulkanViewport.cpp
@@ -206,6 +206,9 @@ void FVulkanViewport::AcquireBackBuffer(FRHICommandListBase& CmdList, FVulkanBac
 
 	if (FVulkanPlatform::SupportsStandardSwapchain())
 	{
+		FTransitionAndLayoutManager& LayoutMgr = Context.GetTransitionAndLayoutManager();
+		VkImageLayout& CurrentLayout = LayoutMgr.FindOrAddLayoutRW(BackBufferImages[AcquiredImageIndex], VK_IMAGE_LAYOUT_UNDEFINED);
+		
 		VulkanRHI::ImagePipelineBarrier(CmdBuffer->GetHandle(), BackBufferImages[AcquiredImageIndex],
 			EImageLayoutBarrier::Undefined, EImageLayoutBarrier::ColorAttachment, VulkanRHI::SetupImageSubresourceRange());
 		if (FVulkanPlatform::RequiresSwapchainGeneralInitialLayout())
@@ -216,6 +219,8 @@ void FVulkanViewport::AcquireBackBuffer(FRHICommandListBase& CmdList, FVulkanBac
 			VulkanRHI::ImagePipelineBarrier(CmdBuffer->GetHandle(), BackBufferImages[AcquiredImageIndex],
 				EImageLayoutBarrier::PixelGeneralRW, EImageLayoutBarrier::ColorAttachment, VulkanRHI::SetupImageSubresourceRange());
 		}
+
+		CurrentLayout = VK_IMAGE_LAYOUT_COLOR_ATTACHMENT_OPTIMAL;
 	}
 
 	// Submit here so we can add a dependency with the acquired semaphore
diff --git a/Engine/Source/Runtime/VulkanRHI/Public/VulkanResources.h b/Engine/Source/Runtime/VulkanRHI/Public/VulkanResources.h
--- a/Engine/Source/Runtime/VulkanRHI/Public/VulkanResources.h
+++ b/Engine/Source/Runtime/VulkanRHI/Public/VulkanResources.h
@@ -719,25 +719,25 @@ inline FVulkanTextureBase* GetVulkanTextureFromRHITexture(FRHITexture* Texture)
 	{
 		return NULL;
 	}
-	else if (Texture->GetTexture2D())
+	else if (FRHITexture2D* Tex2D = Texture->GetTexture2D())
 	{
-		return static_cast<FVulkanTexture2D*>(Texture);
+		return static_cast<FVulkanTexture2D*>(Tex2D);
 	}
-	else if (Texture->GetTextureReference())
+	else if (FRHITextureReference* TexRef = Texture->GetTextureReference())
 	{
-		return static_cast<FVulkanTextureReference*>(Texture);
+		return static_cast<FVulkanTextureReference*>(TexRef);
 	}
-	else if (Texture->GetTexture2DArray())
+	else if (FRHITexture2DArray* Tex2DArray = Texture->GetTexture2DArray())
 	{
-		return static_cast<FVulkanTexture2DArray*>(Texture);
+		return static_cast<FVulkanTexture2DArray*>(Tex2DArray);
 	}
-	else if (Texture->GetTexture3D())
+	else if (FRHITexture3D* Tex3D = Texture->GetTexture3D())
 	{
-		return static_cast<FVulkanTexture3D*>(Texture);
+		return static_cast<FVulkanTexture3D*>(Tex3D);
 	}
-	else if (Texture->GetTextureCube())
+	else if (FRHITextureCube* TexCube = Texture->GetTextureCube())
 	{
-		return static_cast<FVulkanTextureCube*>(Texture);
+		return static_cast<FVulkanTextureCube*>(TexCube);
 	}
 	else
 	{
