Vulkan: try using events to block occlusion queries.

Original change 4701484.

diff --git a/Engine/Source/Runtime/VulkanRHI/Private/VulkanQuery.cpp b/Engine/Source/Runtime/VulkanRHI/Private/VulkanQuery.cpp
--- a/Engine/Source/Runtime/VulkanRHI/Private/VulkanQuery.cpp
+++ b/Engine/Source/Runtime/VulkanRHI/Private/VulkanQuery.cpp
@@ -123,6 +123,7 @@ FScopeLock ScopeLock(&GOcclusionQueryCS);
 FVulkanQueryPool::FVulkanQueryPool(FVulkanDevice* InDevice, uint32 InMaxQueries, VkQueryType InQueryType)
 	: VulkanRHI::FDeviceChild(InDevice)
 	, QueryPool(VK_NULL_HANDLE)
+	, ResetEvent(VK_NULL_HANDLE)
 	, MaxQueries(InMaxQueries)
 	, QueryType(InQueryType)
 {
@@ -134,6 +135,10 @@ FVulkanQueryPool::FVulkanQueryPool(FVulkanDevice* InDevice, uint32 InMaxQueries,
 
 	VERIFYVULKANRESULT(VulkanRHI::vkCreateQueryPool(Device->GetInstanceHandle(), &PoolCreateInfo, VULKAN_CPU_ALLOCATOR, &QueryPool));
 
+	VkEventCreateInfo EventCreateInfo;
+	ZeroVulkanStruct(EventCreateInfo, VK_STRUCTURE_TYPE_EVENT_CREATE_INFO);
+	VERIFYVULKANRESULT(VulkanRHI::vkCreateEvent(Device->GetInstanceHandle(), &EventCreateInfo, VULKAN_CPU_ALLOCATOR, &ResetEvent));
+	
 	QueryOutput.AddZeroed(MaxQueries);
 }
 
@@ -141,8 +146,10 @@ FVulkanQueryPool::~FVulkanQueryPool()
 {
 	DEC_DWORD_STAT(STAT_VulkanNumQueryPools);
 	VulkanRHI::vkDestroyQueryPool(Device->GetInstanceHandle(), QueryPool, VULKAN_CPU_ALLOCATOR);
+	VulkanRHI::vkDestroyEvent(Device->GetInstanceHandle(), ResetEvent, VULKAN_CPU_ALLOCATOR);
+	
 	QueryPool = VK_NULL_HANDLE;
-	check(QueryPool == VK_NULL_HANDLE);
+	ResetEvent = VK_NULL_HANDLE;
 }
 
 /*
@@ -183,22 +190,15 @@ bool FVulkanOcclusionQueryPool::InternalTryGetResults(bool bWait)
 	check(CmdBuffer);
 	check(State == EState::RHIT_PostEndBatch);
 
-	if (CmdBuffer->GetFenceSignaledCounter() <= FenceCounter)
+	VkResult Result = VK_NOT_READY;
+	if (VulkanRHI::vkGetEventStatus(Device->GetInstanceHandle(), ResetEvent) == VK_EVENT_SET)
 	{
-		if (FenceCounter == CmdBuffer->GetSubmittedFenceCounter())
+		Result = VulkanRHI::vkGetQueryPoolResults(Device->GetInstanceHandle(), QueryPool, 0, NumUsedQueries, NumUsedQueries * sizeof(uint64), QueryOutput.GetData(), sizeof(uint64), VK_QUERY_RESULT_64_BIT);
+				if (Result == VK_SUCCESS)
 		{
-			//UE_LOG(LogVulkanRHI, Error, TEXT("Tried to read a query that has not been submitted!"));
-			//ensure(0);
+			State = EState::RT_PostGetResults;
+			return true;
 		}
-
-		return false;
-	}
-
-	VkResult Result = VulkanRHI::vkGetQueryPoolResults(Device->GetInstanceHandle(), QueryPool, 0, NumUsedQueries, NumUsedQueries * sizeof(uint64), QueryOutput.GetData(), sizeof(uint64), VK_QUERY_RESULT_64_BIT);
-	if (Result == VK_SUCCESS)
-	{
-		State = EState::RT_PostGetResults;
-		return true;
 	}
 
 	if (Result == VK_NOT_READY)
@@ -225,7 +225,11 @@ bool FVulkanOcclusionQueryPool::InternalTryGetResults(bool bWait)
 					FTaskGraphInterface::Get().ProcessThreadUntilIdle(RenderThread_Local);
 				}
 
-				Result = VulkanRHI::vkGetQueryPoolResults(Device->GetInstanceHandle(), QueryPool, 0, NumUsedQueries, NumUsedQueries * sizeof(uint64), QueryOutput.GetData(), sizeof(uint64), VK_QUERY_RESULT_64_BIT);
+				if (VulkanRHI::vkGetEventStatus(Device->GetInstanceHandle(), ResetEvent) == VK_EVENT_SET)
+				{
+					Result = VulkanRHI::vkGetQueryPoolResults(Device->GetInstanceHandle(), QueryPool, 0, NumUsedQueries, NumUsedQueries * sizeof(uint64), QueryOutput.GetData(), sizeof(uint64), VK_QUERY_RESULT_64_BIT);
+				}
+
 				if (Result == VK_SUCCESS)
 				{
 					bSuccess = true;
@@ -392,7 +396,20 @@ void FVulkanOcclusionQueryPool::Reset(FVulkanCmdBuffer* InCmdBuffer, uint32 InFr
 /*
 	bHasResults = false;
 */
+	VulkanRHI::vkResetEvent(Device->GetInstanceHandle(), ResetEvent);
 	VulkanRHI::vkCmdResetQueryPool(InCmdBuffer->GetHandle(), QueryPool, 0, MaxQueries);
+	
+	// workaround for apparent cache-flush bug in AMD driver implementation of vkCmdResetQueryPool
+	if (IsRHIDeviceAMD())
+	{
+		VkMemoryBarrier barrier;
+		ZeroVulkanStruct(barrier, VK_STRUCTURE_TYPE_MEMORY_BARRIER);
+		barrier.srcAccessMask = VK_ACCESS_SHADER_WRITE_BIT;
+		barrier.dstAccessMask = VK_ACCESS_SHADER_READ_BIT;
+		VulkanRHI::vkCmdPipelineBarrier(InCmdBuffer->GetHandle(), VK_PIPELINE_STAGE_COMPUTE_SHADER_BIT, VK_PIPELINE_STAGE_BOTTOM_OF_PIPE_BIT, 0, 1, &barrier, 0, 0, 0, 0);
+	}
+	
+	VulkanRHI::vkCmdSetEvent(InCmdBuffer->GetHandle(), ResetEvent, VK_PIPELINE_STAGE_BOTTOM_OF_PIPE_BIT);
 	//FVulkanQueryPool::ResetAll(InCmdBuffer);
 
 	State = EState::RHIT_PostBeginBatch;
@@ -844,7 +861,7 @@ void FVulkanCommandListContext::RHIBeginRenderQuery(FRenderQueryRHIParamRef Quer
 		Query->Pool = CurrentOcclusionQueryPool;
 		Query->IndexInPool = IndexInPool;
 		FVulkanCmdBuffer* CmdBuffer = CommandBufferManager->GetActiveCmdBuffer();
-		VulkanRHI::vkCmdBeginQuery(CmdBuffer->GetHandle(), CurrentOcclusionQueryPool->GetHandle(), IndexInPool, VK_QUERY_RESULT_64_BIT);
+		VulkanRHI::vkCmdBeginQuery(CmdBuffer->GetHandle(), CurrentOcclusionQueryPool->GetHandle(), IndexInPool, VK_QUERY_CONTROL_PRECISE_BIT);
 	}
 	else if (BaseQuery->QueryType == RQT_AbsoluteTime)
 	{
diff --git a/Engine/Source/Runtime/VulkanRHI/Public/VulkanResources.h b/Engine/Source/Runtime/VulkanRHI/Public/VulkanResources.h
--- a/Engine/Source/Runtime/VulkanRHI/Public/VulkanResources.h
+++ b/Engine/Source/Runtime/VulkanRHI/Public/VulkanResources.h
@@ -769,6 +769,7 @@ public:
 
 protected:
 	VkQueryPool QueryPool;
+	VkEvent ResetEvent;
 	const uint32 MaxQueries;
 	const VkQueryType QueryType;
 	TArray<uint64> QueryOutput;
