Copy from 4846626
- Fix NV crash diagnostic
#rnx

diff --git a/Engine/Source/Runtime/VulkanRHI/Private/VulkanUtil.cpp b/Engine/Source/Runtime/VulkanRHI/Private/VulkanUtil.cpp
--- a/Engine/Source/Runtime/VulkanRHI/Private/VulkanUtil.cpp
+++ b/Engine/Source/Runtime/VulkanRHI/Private/VulkanUtil.cpp
@@ -550,22 +550,32 @@ void FVulkanGPUProfiler::DumpCrashMarkers(void* BufferData)
 #if VULKAN_SUPPORTS_NV_DIAGNOSTIC_CHECKPOINT
 		if (Device->GetOptionalExtensions().HasNVDiagnosticCheckpoints)
 		{
-			TArray<VkCheckpointDataNV> Data;
+			struct FCheckpointDataNV : public VkCheckpointDataNV
+			{
+				FCheckpointDataNV()
+				{
+					ZeroVulkanStruct(*this, VK_STRUCTURE_TYPE_CHECKPOINT_DATA_NV);
+				}
+			};
+			TArray<FCheckpointDataNV> Data;
 			uint32 Num = 0;
 			VkQueue QueueHandle = Device->GetGraphicsQueue()->GetHandle();
 			VulkanDynamicAPI::vkGetQueueCheckpointDataNV(QueueHandle, &Num, nullptr);
-			Data.AddUninitialized(Num);
-			VulkanDynamicAPI::vkGetQueueCheckpointDataNV(QueueHandle, &Num, &Data[0]);
-			check(Num == Data.Num());
-			for (uint32 Index = 0; Index < Num; ++Index)
+			if (Num > 0)
 			{
-				check(Data[Index].sType == VK_STRUCTURE_TYPE_CHECKPOINT_DATA_NV);
-				uint32 Value = (uint32)(size_t)Data[Index].pCheckpointMarker;
-				const FString* Frame = CachedStrings.Find(Value);
-				UE_LOG(LogVulkanRHI, Error, TEXT("[VK_NV_device_diagnostic_checkpoints] %i: Stage 0x%x, %s (CRC 0x%x)"), Index, Data[Index].stage, Frame ? *(*Frame) : TEXT("<undefined>"), Value);
+				Data.AddDefaulted(Num);
+				VulkanDynamicAPI::vkGetQueueCheckpointDataNV(QueueHandle, &Num, &Data[0]);
+				check(Num == Data.Num());
+				for (uint32 Index = 0; Index < Num; ++Index)
+				{
+					check(Data[Index].sType == VK_STRUCTURE_TYPE_CHECKPOINT_DATA_NV);
+					uint32 Value = (uint32)(size_t)Data[Index].pCheckpointMarker;
+					const FString* Frame = CachedStrings.Find(Value);
+					UE_LOG(LogVulkanRHI, Error, TEXT("[VK_NV_device_diagnostic_checkpoints] %i: Stage 0x%x, %s (CRC 0x%x)"), Index, Data[Index].stage, Frame ? *(*Frame) : TEXT("<undefined>"), Value);
+				}
+				GLog->PanicFlushThreadedLogs();
+				GLog->Flush();
 			}
-			GLog->PanicFlushThreadedLogs();
-			GLog->Flush();
 		}
 #endif
 	}
