Android compile fix
#rnx

diff --git a/Engine/Source/Runtime/VulkanRHI/Private/VulkanCommandWrappers.h b/Engine/Source/Runtime/VulkanRHI/Private/VulkanCommandWrappers.h
--- a/Engine/Source/Runtime/VulkanRHI/Private/VulkanCommandWrappers.h
+++ b/Engine/Source/Runtime/VulkanRHI/Private/VulkanCommandWrappers.h
@@ -151,6 +151,7 @@ struct FWrapLayer
 	static void FillBuffer(VkResult Result, VkCommandBuffer CommandBuffer, VkBuffer DstBuffer, VkDeviceSize DstOffset, VkDeviceSize Size, uint32 Data) VULKAN_LAYER_BODY
 	static void SetEvent(VkResult Result, VkCommandBuffer CommandBuffer, VkEvent Event, VkPipelineStageFlags StageMask) VULKAN_LAYER_BODY
 	static void ResetEvent(VkResult Result, VkCommandBuffer CommandBuffer, VkEvent Event, VkPipelineStageFlags StageMask) VULKAN_LAYER_BODY
+	static void GetEventStatus(VkResult Result, VkDevice Device, VkEvent Event) VULKAN_LAYER_BODY
 	static void CopyQueryPoolResults(VkResult Result, VkCommandBuffer CommandBuffer, VkQueryPool QueryPool, uint32 FirstQuery, uint32 QueryCount, VkBuffer DstBuffer, VkDeviceSize DstOffset, VkDeviceSize Stride, VkQueryResultFlags Flags) VULKAN_LAYER_BODY
 	static void GetDeviceProcAddr(VkResult Result, VkDevice Device, const char* Name, PFN_vkVoidFunction VoidFunction) VULKAN_LAYER_BODY
 	VULKAN_EXTERN_EXPORT static void EnumerateDeviceExtensionProperties(VkResult Result, VkPhysicalDevice PhysicalDevice, const char* LayerName, uint32* PropertyCount, VkExtensionProperties* Properties) VULKAN_LAYER_BODY
@@ -527,19 +528,31 @@ namespace VulkanRHI
 		VULKANAPINAMESPACE::vkDestroyEvent(Device, Event, Allocator);
 		FWrapLayer::DestroyEvent(VK_SUCCESS, Device, Event);
 	}
-#if 0
-	static FORCEINLINE_DEBUGGABLE VkResult  vkGetEventStatus(
-		VkDevice                                    Device,
-		VkEvent                                     event);
 
-	static FORCEINLINE_DEBUGGABLE VkResult  vkSetEvent(
-		VkDevice                                    Device,
-		VkEvent                                     event);
+	static FORCEINLINE_DEBUGGABLE VkResult  vkGetEventStatus(VkDevice Device, VkEvent Event)
+	{
+		FWrapLayer::GetEventStatus(VK_RESULT_MAX_ENUM, Device, Event);
+		VkResult Result = VULKANAPINAMESPACE::vkGetEventStatus(Device, Event);
+		FWrapLayer::GetEventStatus(Result, Device, Event);
+		return Result;
+	}
+
+	static FORCEINLINE_DEBUGGABLE VkResult  vkSetEvent(VkDevice Device, VkEvent Event)
+	{
+		FWrapLayer::SetEvent(VK_RESULT_MAX_ENUM, Device, Event);
+		VkResult Result = VULKANAPINAMESPACE::vkSetEvent(Device, Event);
+		FWrapLayer::SetEvent(Result, Device, Event);
+		return Result;
+	}
+
+	static FORCEINLINE_DEBUGGABLE VkResult  vkResetEvent(VkDevice Device, VkEvent Event)
+	{
+		FWrapLayer::ResetEvent(VK_RESULT_MAX_ENUM, Device, Event);
+		VkResult Result = VULKANAPINAMESPACE::vkResetEvent(Device, Event);
+		FWrapLayer::ResetEvent(Result, Device, Event);
+		return Result;
+	}
 
-	static FORCEINLINE_DEBUGGABLE VkResult  vkResetEvent(
-		VkDevice                                    Device,
-		VkEvent                                     event);
-#endif
 	static FORCEINLINE_DEBUGGABLE VkResult  vkCreateQueryPool(VkDevice Device, const VkQueryPoolCreateInfo* CreateInfo, const VkAllocationCallbacks* Allocator, VkQueryPool* QueryPool)
 	{
 		FWrapLayer::CreateQueryPool(VK_RESULT_MAX_ENUM, Device, CreateInfo, QueryPool);
diff --git a/Engine/Source/Runtime/VulkanRHI/Private/VulkanDebug.cpp b/Engine/Source/Runtime/VulkanRHI/Private/VulkanDebug.cpp
--- a/Engine/Source/Runtime/VulkanRHI/Private/VulkanDebug.cpp
+++ b/Engine/Source/Runtime/VulkanRHI/Private/VulkanDebug.cpp
@@ -3647,6 +3647,22 @@ void FWrapLayer::SetEvent(VkResult Result, VkCommandBuffer CommandBuffer, VkEven
 	}
 }
 
+void FWrapLayer::GetEventStatus(VkResult Result, VkDevice Device, VkEvent Event)
+{
+	if (Result == VK_RESULT_MAX_ENUM)
+	{
+#if VULKAN_ENABLE_DUMP_LAYER
+		DevicePrintfBegin(CommandBuffer, FString::Printf(TEXT("GetEventStatus(Event=0x%p)"), Event));
+#endif
+	}
+	else
+	{
+#if VULKAN_ENABLE_DUMP_LAYER
+		PrintResult(Result);
+#endif
+	}
+}
+
 void FWrapLayer::ResetEvent(VkResult Result, VkCommandBuffer CommandBuffer, VkEvent Event, VkPipelineStageFlags StageMask)
 {
 	if (Result == VK_RESULT_MAX_ENUM)
