Android compile fix (for real)
#rnx

diff --git a/Engine/Source/Runtime/VulkanRHI/Private/VulkanCommandWrappers.h b/Engine/Source/Runtime/VulkanRHI/Private/VulkanCommandWrappers.h
--- a/Engine/Source/Runtime/VulkanRHI/Private/VulkanCommandWrappers.h
+++ b/Engine/Source/Runtime/VulkanRHI/Private/VulkanCommandWrappers.h
@@ -151,6 +151,8 @@ struct FWrapLayer
 	static void FillBuffer(VkResult Result, VkCommandBuffer CommandBuffer, VkBuffer DstBuffer, VkDeviceSize DstOffset, VkDeviceSize Size, uint32 Data) VULKAN_LAYER_BODY
 	static void SetEvent(VkResult Result, VkCommandBuffer CommandBuffer, VkEvent Event, VkPipelineStageFlags StageMask) VULKAN_LAYER_BODY
 	static void ResetEvent(VkResult Result, VkCommandBuffer CommandBuffer, VkEvent Event, VkPipelineStageFlags StageMask) VULKAN_LAYER_BODY
+	static void SetEvent(VkResult Result, VkDevice Device, VkEvent Event) VULKAN_LAYER_BODY
+	static void ResetEvent(VkResult Result, VkDevice Device, VkEvent Event) VULKAN_LAYER_BODY
 	static void GetEventStatus(VkResult Result, VkDevice Device, VkEvent Event) VULKAN_LAYER_BODY
 	static void CopyQueryPoolResults(VkResult Result, VkCommandBuffer CommandBuffer, VkQueryPool QueryPool, uint32 FirstQuery, uint32 QueryCount, VkBuffer DstBuffer, VkDeviceSize DstOffset, VkDeviceSize Stride, VkQueryResultFlags Flags) VULKAN_LAYER_BODY
 	static void GetDeviceProcAddr(VkResult Result, VkDevice Device, const char* Name, PFN_vkVoidFunction VoidFunction) VULKAN_LAYER_BODY
diff --git a/Engine/Source/Runtime/VulkanRHI/Private/VulkanDebug.cpp b/Engine/Source/Runtime/VulkanRHI/Private/VulkanDebug.cpp
--- a/Engine/Source/Runtime/VulkanRHI/Private/VulkanDebug.cpp
+++ b/Engine/Source/Runtime/VulkanRHI/Private/VulkanDebug.cpp
@@ -3647,23 +3647,27 @@ void FWrapLayer::SetEvent(VkResult Result, VkCommandBuffer CommandBuffer, VkEven
 	}
 }
 
-void FWrapLayer::GetEventStatus(VkResult Result, VkDevice Device, VkEvent Event)
+void FWrapLayer::ResetEvent(VkResult Result, VkCommandBuffer CommandBuffer, VkEvent Event, VkPipelineStageFlags StageMask)
 {
 	if (Result == VK_RESULT_MAX_ENUM)
 	{
 #if VULKAN_ENABLE_DUMP_LAYER
-		DevicePrintfBegin(CommandBuffer, FString::Printf(TEXT("GetEventStatus(Event=0x%p)"), Event));
+		CmdPrintfBegin(CommandBuffer, FString::Printf(TEXT("vkCmdResetEvent(Event=0x%p, StageMask=0x%x)"), Event, StageMask));
 #endif
 	}
-	else
+}
+
+void FWrapLayer::SetEvent(VkResult Result, VkDevice Device, VkEvent Event)
+{
+	if (Result == VK_RESULT_MAX_ENUM)
 	{
 #if VULKAN_ENABLE_DUMP_LAYER
-		PrintResult(Result);
+		CmdPrintfBegin(CommandBuffer, FString::Printf(TEXT("vkCmdSetEvent(Event=0x%p, StageMask=0x%x)"), Event, StageMask));
 #endif
 	}
 }
 
-void FWrapLayer::ResetEvent(VkResult Result, VkCommandBuffer CommandBuffer, VkEvent Event, VkPipelineStageFlags StageMask)
+void FWrapLayer::ResetEvent(VkResult Result, VkDevice Device, VkEvent Event)
 {
 	if (Result == VK_RESULT_MAX_ENUM)
 	{
@@ -3673,6 +3677,22 @@ void FWrapLayer::ResetEvent(VkResult Result, VkCommandBuffer CommandBuffer, VkEv
 	}
 }
 
+void FWrapLayer::GetEventStatus(VkResult Result, VkDevice Device, VkEvent Event)
+{
+	if (Result == VK_RESULT_MAX_ENUM)
+	{
+#if VULKAN_ENABLE_DUMP_LAYER
+		DevicePrintfBegin(CommandBuffer, FString::Printf(TEXT("GetEventStatus(Event=0x%p)"), Event));
+#endif
+	}
+	else
+	{
+#if VULKAN_ENABLE_DUMP_LAYER
+		PrintResult(Result);
+#endif
+	}
+}
+
 void FWrapLayer::CopyQueryPoolResults(VkResult Result, VkCommandBuffer CommandBuffer, VkQueryPool QueryPool, uint32 FirstQuery, uint32 QueryCount, VkBuffer DstBuffer, VkDeviceSize DstOffset, VkDeviceSize Stride, VkQueryResultFlags Flags)
 {
 	if (Result == VK_RESULT_MAX_ENUM)
