VK and PS4 drop UAV binds on the floor if you use Renderpasses (and also SetRenderTargetsAndClear).

diff --git a/Engine/Source/Runtime/RHI/Private/RHI.cpp b/Engine/Source/Runtime/RHI/Private/RHI.cpp
--- a/Engine/Source/Runtime/RHI/Private/RHI.cpp
+++ b/Engine/Source/Runtime/RHI/Private/RHI.cpp
@@ -891,8 +891,9 @@ void FRHIRenderPassInfo::ConvertToRenderTargetsInfo(FRHISetRenderTargetsInfo& Ou
 	if (NumUAVs > 0)
 	{
 		check(UAVIndex != -1);
-		check(UAVIndex >= OutRTInfo.NumColorRenderTargets);
-		OutRTInfo.NumColorRenderTargets = UAVIndex;
+		int32 StartingUAVIndex = FMath::Max(UAVIndex, OutRTInfo.NumColorRenderTargets);
+		check((StartingUAVIndex+NumUAVs) <= MaxSimultaneousUAVs);
+		OutRTInfo.NumColorRenderTargets = StartingUAVIndex;
 		for (int32 Index = 0; Index < NumUAVs; ++Index)
 		{
 			OutRTInfo.UnorderedAccessView[Index] = UAVs[Index];
diff --git a/Engine/Source/Runtime/VulkanRHI/Private/VulkanRenderTarget.cpp b/Engine/Source/Runtime/VulkanRHI/Private/VulkanRenderTarget.cpp
--- a/Engine/Source/Runtime/VulkanRHI/Private/VulkanRenderTarget.cpp
+++ b/Engine/Source/Runtime/VulkanRHI/Private/VulkanRenderTarget.cpp
@@ -498,20 +498,17 @@ void FVulkanCommandListContext::RHISetRenderTargets(uint32 NumSimultaneousRender
 	}
 
 	FRHISetRenderTargetsInfo RenderTargetsInfo(NumSimultaneousRenderTargets, NewRenderTargets, DepthView);
-	RHISetRenderTargetsAndClear(RenderTargetsInfo);
 
-	// Yuck - Bind pending pixel shader UAVs from SetRenderTargets
+	if (NumUAVs)
 	{
-		PendingPixelUAVs.Reset();
-		for (uint32 UAVIndex = 0; UAVIndex < NumUAVs; ++UAVIndex)
+		RenderTargetsInfo.NumUAVs = NumUAVs;
+		for (uint32 Index = 0; Index < NumUAVs; Index++)
 		{
-			FVulkanUnorderedAccessView* UAV = ResourceCast(UAVs[UAVIndex]);
-			if (UAV)
-			{
-				PendingPixelUAVs.Add({UAV, UAVIndex});
-			}
+			RenderTargetsInfo.UnorderedAccessView[Index] = UAVs[Index];
 		}
 	}
+
+	RHISetRenderTargetsAndClear(RenderTargetsInfo);
 }
 
 // Find out whether we can re-use current renderpass instead of starting new one
@@ -613,6 +610,20 @@ void FVulkanCommandListContext::RHISetRenderTargetsAndClear(const FRHISetRenderT
 			ensureMsgf(0, TEXT("RenderPass not started! Bad combination of values? Depth %p #Color %d Color0 %p"), (void*)RenderTargetsInfo.DepthStencilRenderTarget.Texture, RenderTargetsInfo.NumColorRenderTargets, (void*)RenderTargetsInfo.ColorRenderTarget[0].Texture);
 		}
 	}
+
+	// Yuck - Bind pending pixel shader UAVs from SetRenderTargets
+	{
+		PendingPixelUAVs.Reset();
+		uint32 NumUAVs = RenderTargetsInfo.NumUAVs;
+		for (uint32 UAVIndex = 0; UAVIndex < NumUAVs; ++UAVIndex)
+		{
+			FVulkanUnorderedAccessView* UAV = ResourceCast(RenderTargetsInfo.UnorderedAccessView[UAVIndex].GetReference());
+			if (UAV)
+			{
+				PendingPixelUAVs.Add({ UAV, UAVIndex });
+			}
+		}
+	}
 }
 
 void FVulkanCommandListContext::RHICopyToResolveTarget(FTextureRHIParamRef SourceTextureRHI, FTextureRHIParamRef DestTextureRHI, const FResolveParams& InResolveParams)
@@ -1607,11 +1618,27 @@ void FVulkanCommandListContext::RHIBeginRenderPass(const FRHIRenderPassInfo& InI
 	{
 		BeginOcclusionQueryBatch(CmdBuffer, InInfo.NumOcclusionQueries);
 	}
+
 	FVulkanRenderTargetLayout RTLayout(*Device, InInfo);
 	check(RTLayout.GetExtent2D().width != 0 && RTLayout.GetExtent2D().height != 0);
 	FVulkanRenderPass* RenderPass = TransitionAndLayoutManager.GetOrCreateRenderPass(*Device, RTLayout);
 	FRHISetRenderTargetsInfo RTInfo;
 	InInfo.ConvertToRenderTargetsInfo(RTInfo);
+
+	// yuck. bind UAVs.
+	{
+		PendingPixelUAVs.Reset();
+		uint32 NumUAVs = InInfo.NumUAVs;
+		for (uint32 UAVIndex = 0; UAVIndex < NumUAVs; ++UAVIndex)
+		{
+			FVulkanUnorderedAccessView* UAV = ResourceCast(InInfo.UAVs[UAVIndex].GetReference());
+			if (UAV)
+			{
+				PendingPixelUAVs.Add({ UAV, UAVIndex });
+			}
+		}
+	}
+
 	FVulkanFramebuffer* Framebuffer = TransitionAndLayoutManager.GetOrCreateFramebuffer(*Device, RTInfo, RTLayout, RenderPass);
 	checkf(RenderPass != nullptr && Framebuffer != nullptr, TEXT("RenderPass not started! Bad combination of values? Depth %p #Color %d Color0 %p"), (void*)InInfo.DepthStencilRenderTarget.DepthStencilTarget, InInfo.GetNumColorRenderTargets(), (void*)InInfo.ColorRenderTargets[0].RenderTarget);
 	TransitionAndLayoutManager.BeginRealRenderPass(*this, *Device, CmdBuffer, InInfo, RTLayout, RenderPass, Framebuffer);
