Update RHIs to safer command allocation

diff --git a/Engine/Source/Runtime/D3D12RHI/Private/D3D12Adapter.cpp b/Engine/Source/Runtime/D3D12RHI/Private/D3D12Adapter.cpp
--- a/Engine/Source/Runtime/D3D12RHI/Private/D3D12Adapter.cpp
+++ b/Engine/Source/Runtime/D3D12RHI/Private/D3D12Adapter.cpp
@@ -556,7 +556,7 @@ void FD3D12Adapter::SignalFrameFence_RenderThread(FRHICommandListImmediate& RHIC
 	}
 	else
 	{
-		new (RHICmdList.AllocCommand<FRHICommandSignalFrameFence>()) FRHICommandSignalFrameFence(ED3D12CommandQueueType::Default, FrameFence, PreviousFence);
+		ALLOC_COMMAND_CL(RHICmdList, FRHICommandSignalFrameFence)(ED3D12CommandQueueType::Default, FrameFence, PreviousFence);
 	}
 }
 
diff --git a/Engine/Source/Runtime/D3D12RHI/Private/D3D12Buffer.cpp b/Engine/Source/Runtime/D3D12RHI/Private/D3D12Buffer.cpp
--- a/Engine/Source/Runtime/D3D12RHI/Private/D3D12Buffer.cpp
+++ b/Engine/Source/Runtime/D3D12RHI/Private/D3D12Buffer.cpp
@@ -224,7 +224,7 @@ BufferType* FD3D12Adapter::CreateRHIBuffer(FRHICommandListImmediate* RHICmdList,
 			
 			if (RHICmdList && !RHICmdList->Bypass())
 			{
-				new (RHICmdList->AllocCommand<FD3D12RHICommandInitializeBuffer>()) FD3D12RHICommandInitializeBuffer(BufferOut, SrcResourceLoc, Size);
+				ALLOC_COMMAND_CL(*RHICmdList, FD3D12RHICommandInitializeBuffer)(BufferOut, SrcResourceLoc, Size);
 			}
 			else
 			{
@@ -271,7 +271,7 @@ void* FD3D12DynamicRHI::LockBuffer(FRHICommandListImmediate* RHICmdList, BufferT
 			// If on the RenderThread, queue up a command on the RHIThread to rename this buffer at the correct time
 			if (ShouldDeferBufferLockOperation(RHICmdList))
 			{
-				FRHICommandRenameUploadBuffer<BufferType>* Command = new (RHICmdList->AllocCommand<FRHICommandRenameUploadBuffer<BufferType>>()) FRHICommandRenameUploadBuffer<BufferType>(Buffer, Device);
+				FRHICommandRenameUploadBuffer<BufferType>* Command = ALLOC_COMMAND_CL(*RHICmdList, FRHICommandRenameUploadBuffer<BufferType>)(Buffer, Device);
 
 				Data = Adapter.GetUploadHeapAllocator(Device->GetGPUIndex()).AllocUploadResource(Buffer->GetSize(), Buffer->BufferAlignment, Command->NewLocation);
 				RHICmdList->RHIThreadFence(true);
@@ -392,13 +392,13 @@ void FD3D12DynamicRHI::UnlockBuffer(FRHICommandListImmediate* RHICmdList, Buffer
 				{
 					if (GNumExplicitGPUsForRendering == 1)
 					{
-						new (RHICmdList->AllocCommand<FRHICommandUpdateBuffer>()) FRHICommandUpdateBuffer(&CurrentBuffer->ResourceLocation, LockedData.ResourceLocation, LockedData.LockedOffset, LockedData.LockedPitch);
+						ALLOC_COMMAND_CL(*RHICmdList, FRHICommandUpdateBuffer)(&CurrentBuffer->ResourceLocation, LockedData.ResourceLocation, LockedData.LockedOffset, LockedData.LockedPitch);
 					}
 					else // The resource location must be copied because the constructor in FRHICommandUpdateBuffer transfers ownership and clears it.
 					{
 						FD3D12ResourceLocation NodeResourceLocation(LockedData.ResourceLocation.GetParentDevice());
 						FD3D12ResourceLocation::ReferenceNode(NodeResourceLocation.GetParentDevice(), NodeResourceLocation, LockedData.ResourceLocation);
-						new (RHICmdList->AllocCommand<FRHICommandUpdateBuffer>()) FRHICommandUpdateBuffer(&CurrentBuffer->ResourceLocation, NodeResourceLocation, LockedData.LockedOffset, LockedData.LockedPitch);
+						ALLOC_COMMAND_CL(*RHICmdList, FRHICommandUpdateBuffer)(&CurrentBuffer->ResourceLocation, NodeResourceLocation, LockedData.LockedOffset, LockedData.LockedPitch);
 					}
 				}
 				else
diff --git a/Engine/Source/Runtime/D3D12RHI/Private/D3D12Texture.cpp b/Engine/Source/Runtime/D3D12RHI/Private/D3D12Texture.cpp
--- a/Engine/Source/Runtime/D3D12RHI/Private/D3D12Texture.cpp
+++ b/Engine/Source/Runtime/D3D12RHI/Private/D3D12Texture.cpp
@@ -1702,7 +1702,7 @@ void FD3D12TextureBase::InitializeTextureData(FRHICommandListImmediate* RHICmdLi
 
 	if (ShouldDeferCmdListOperation(RHICmdList))
 	{
-		new (RHICmdList->AllocCommand<FD3D12RHICommandInitializeTexture>()) FD3D12RHICommandInitializeTexture(this, SrcResourceLoc, SizeX, SizeY, SizeZ, NumSlices, NumMips, Format, DestinationState);
+		ALLOC_COMMAND_CL(*RHICmdList, FD3D12RHICommandInitializeTexture)(this, SrcResourceLoc, SizeX, SizeY, SizeZ, NumSlices, NumMips, Format, DestinationState);
 	}
 	else
 	{
@@ -1771,7 +1771,7 @@ void TD3D12Texture2D<RHIResourceType>::UnlockInternal(class FRHICommandListImmed
 			// If we are on the render thread, queue up the copy on the RHIThread so it happens at the correct time.
 			if (ShouldDeferCmdListOperation(RHICmdList))
 			{
-				new (RHICmdList->AllocCommand<FRHICommandUpdateTexture>()) FRHICommandUpdateTexture(this, DestCopyLocation, 0, 0, 0, UploadLocation, PlacedTexture2D);
+				ALLOC_COMMAND_CL(*RHICmdList, FRHICommandUpdateTexture)(this, DestCopyLocation, 0, 0, 0, UploadLocation, PlacedTexture2D);
 			}
 			else
 			{
@@ -1845,7 +1845,7 @@ void TD3D12Texture2D<RHIResourceType>::UpdateTexture2D(class FRHICommandListImme
 		// If we are on the render thread, queue up the copy on the RHIThread so it happens at the correct time.
 		if (ShouldDeferCmdListOperation(RHICmdList))
 		{
-			new (RHICmdList->AllocCommand<FRHICommandUpdateTexture>()) FRHICommandUpdateTexture(this, DestCopyLocation, UpdateRegion.DestX, UpdateRegion.DestY, 0, UploadHeapResourceLocation, PlacedTexture2D);
+			ALLOC_COMMAND_CL(*RHICmdList, FRHICommandUpdateTexture)(this, DestCopyLocation, UpdateRegion.DestX, UpdateRegion.DestY, 0, UploadHeapResourceLocation, PlacedTexture2D);
 		}
 		else
 		{
@@ -2149,7 +2149,7 @@ void FD3D12DynamicRHI::EndUpdateTexture3D_Internal(FUpdateTexture3DData& UpdateD
 		}
 		else
 		{
-			new (RHICmdList.AllocCommand<FD3D12RHICmdEndUpdateTexture3D>()) FD3D12RHICmdEndUpdateTexture3D(UpdateData);
+			ALLOC_COMMAND_CL(RHICmdList, FD3D12RHICmdEndUpdateTexture3D)(UpdateData);
 		}
 	}
 }
diff --git a/Engine/Source/Runtime/D3D12RHI/Private/D3D12View.cpp b/Engine/Source/Runtime/D3D12RHI/Private/D3D12View.cpp
--- a/Engine/Source/Runtime/D3D12RHI/Private/D3D12View.cpp
+++ b/Engine/Source/Runtime/D3D12RHI/Private/D3D12View.cpp
@@ -257,7 +257,7 @@ FShaderResourceViewRHIRef FD3D12DynamicRHI::RHICreateShaderResourceView(FVertexB
 		{
 			// We have to defer the SRV initialization to the RHI thread if the buffer is dynamic (and RHI threading is enabled), as dynamic buffers can be renamed.
 			// Also insert an RHI thread fence to prevent parallel translate tasks running until this command has completed.
-			new (RHICmdList.AllocCommand<FD3D12InitializeVertexBufferSRVRHICommand>()) FD3D12InitializeVertexBufferSRVRHICommand(VertexBuffer, ShaderResourceView, Stride, Format);
+			ALLOC_COMMAND_CL(RHICmdList, FD3D12InitializeVertexBufferSRVRHICommand)(VertexBuffer, ShaderResourceView, Stride, Format);
 			RHICmdList.RHIThreadFence(true);
 		}
 		else
diff --git a/Engine/Source/Runtime/OpenGLDrv/Private/OpenGLQuery.cpp b/Engine/Source/Runtime/OpenGLDrv/Private/OpenGLQuery.cpp
--- a/Engine/Source/Runtime/OpenGLDrv/Private/OpenGLQuery.cpp
+++ b/Engine/Source/Runtime/OpenGLDrv/Private/OpenGLQuery.cpp
@@ -525,9 +525,9 @@ bool FOpenGLDynamicRHI::RHIGetRenderQueryResult(FRenderQueryRHIParamRef QueryRHI
 				if (IsRunningRHIInDedicatedThread())
 				{
 					// send a command that will wait, so if the RHIT runs out of work, it just blocks and waits for the GPU
-					new (RHICmdList.AllocCommand<FRHICommandGLCommand>()) FRHICommandGLCommand([=]() {GetRenderQueryResult_OnThisThread(ResourceCast(QueryRHI), true); });
+					ALLOC_COMMAND_CL(RHICmdList, FRHICommandGLCommand)([=]() {GetRenderQueryResult_OnThisThread(ResourceCast(QueryRHI), true); });
 					FGraphEventRef Done = RHICmdList.RHIThreadFence(false);
-					new (RHICmdList.AllocCommand<FRHICommandGLCommand>()) FRHICommandGLCommand([=]() {GBatcher.Flush(*this, QueryRHI); });
+					ALLOC_COMMAND_CL(RHICmdList, FRHICommandGLCommand)([=]() {GBatcher.Flush(*this, QueryRHI); });
 					RHICmdList.ImmediateFlush(EImmediateFlushType::DispatchToRHIThread);
 					while (!Done->IsComplete())
 					{
@@ -545,9 +545,9 @@ bool FOpenGLDynamicRHI::RHIGetRenderQueryResult(FRenderQueryRHIParamRef QueryRHI
 				}
 				else
 				{
-					new (RHICmdList.AllocCommand<FRHICommandGLCommand>()) FRHICommandGLCommand([=]() {GetRenderQueryResult_OnThisThread(ResourceCast(QueryRHI), true); });
+					ALLOC_COMMAND_CL(RHICmdList, FRHICommandGLCommand)([=]() {GetRenderQueryResult_OnThisThread(ResourceCast(QueryRHI), true); });
 					FGraphEventRef Done = RHICmdList.RHIThreadFence(false);
-					new (RHICmdList.AllocCommand<FRHICommandGLCommand>()) FRHICommandGLCommand([=]() {GBatcher.Flush(*this, QueryRHI); });
+					ALLOC_COMMAND_CL(RHICmdList, FRHICommandGLCommand)([=]() {GBatcher.Flush(*this, QueryRHI); });
 					RHICmdList.ImmediateFlush(EImmediateFlushType::DispatchToRHIThread);
 					FRHICommandListExecutor::WaitOnRHIThreadFence(Done);
 				}
@@ -555,7 +555,7 @@ bool FOpenGLDynamicRHI::RHIGetRenderQueryResult(FRenderQueryRHIParamRef QueryRHI
 			}
 			else
 			{
-				new (RHICmdList.AllocCommand<FRHICommandGLCommand>()) FRHICommandGLCommand([=]() {GetRenderQueryResult_OnThisThread(ResourceCast(QueryRHI), false); GBatcher.Flush(*this, QueryRHI);  });
+				ALLOC_COMMAND_CL(RHICmdList, FRHICommandGLCommand)([=]() {GetRenderQueryResult_OnThisThread(ResourceCast(QueryRHI), false); GBatcher.Flush(*this, QueryRHI);  });
 			}
 		}	
 	}
@@ -600,7 +600,7 @@ FOpenGLRenderQuery::FOpenGLRenderQuery(ERenderQueryType InQueryType)
 	else
 	{
 		CreationFence.Reset();
-		new (RHICmdList.AllocCommand<FRHICommandGLCommand>()) FRHICommandGLCommand([=]() {AcquireResource(); CreationFence.WriteAssertFence(); });
+		ALLOC_COMMAND_CL(RHICmdList, FRHICommandGLCommand)([=]() {AcquireResource(); CreationFence.WriteAssertFence(); });
 		CreationFence.SetRHIThreadFence();
 	}
 }
@@ -624,7 +624,7 @@ FOpenGLRenderQuery::~FOpenGLRenderQuery()
 		else
 		{
 			CreationFence.WaitFence();
-			new (RHICmdList.AllocCommand<FRHICommandGLCommand>()) FRHICommandGLCommand([Resource = Resource, ResourceContext = ResourceContext]() {VERIFY_GL_SCOPE(); ReleaseResource(Resource, ResourceContext); });
+			ALLOC_COMMAND_CL(RHICmdList, FRHICommandGLCommand)([Resource = Resource, ResourceContext = ResourceContext]() {VERIFY_GL_SCOPE(); ReleaseResource(Resource, ResourceContext); });
 		}
 	}
 }
diff --git a/Engine/Source/Runtime/OpenGLDrv/Private/OpenGLTexture.cpp b/Engine/Source/Runtime/OpenGLDrv/Private/OpenGLTexture.cpp
--- a/Engine/Source/Runtime/OpenGLDrv/Private/OpenGLTexture.cpp
+++ b/Engine/Source/Runtime/OpenGLDrv/Private/OpenGLTexture.cpp
@@ -2248,7 +2248,7 @@ FTexture2DRHIRef FOpenGLDynamicRHI::AsyncReallocateTexture2D_RenderThread(class
 		FOpenGLTexture2D* Texture2D = FOpenGLDynamicRHI::ResourceCast(NewTexture2DRHI.GetReference());
 		Texture2D->CreationFence.Reset();
 
-		new (RHICmdList.AllocCommand<FRHICommandGLCommand>()) FRHICommandGLCommand([=]() 
+		ALLOC_COMMAND_CL(RHICmdList, FRHICommandGLCommand)([=]() 
 		{
 			GLCopyAsyncTexture2D(this, NewTexture2DRHI, NewSizeX, NewSizeY, Texture2DRHI, RequestStatus); 
 			Texture2D->CreationFence.WriteAssertFence();
@@ -2694,7 +2694,7 @@ void FOpenGLDynamicRHI::UnlockTexture2D_RenderThread(class FRHICommandListImmedi
 			FMemory::Free(Params.Buffer);
 			this->RHIUnlockTexture2D(Texture, MipIndex, bLockWithinMiptail);
 		};
-		new (RHICmdList.AllocCommand<FRHICommandGLCommand>()) FRHICommandGLCommand(GLCommand);
+		ALLOC_COMMAND_CL(RHICmdList, FRHICommandGLCommand)(GLCommand);
 	}
 }
 
@@ -2748,7 +2748,7 @@ void FOpenGLDynamicRHI::RHIUnlockTextureCubeFace_RenderThread(class FRHICommandL
 			FMemory::Free(Params.Buffer);
 			this->RHIUnlockTextureCubeFace(Texture, FaceIndex, ArrayIndex, MipIndex, bLockWithinMiptail);
 		};
-		new (RHICmdList.AllocCommand<FRHICommandGLCommand>()) FRHICommandGLCommand(GLCommand);
+		ALLOC_COMMAND_CL(RHICmdList, FRHICommandGLCommand)(GLCommand);
 	}
 }
 
diff --git a/Engine/Source/Runtime/OpenGLDrv/Private/OpenGLUniformBuffer.cpp b/Engine/Source/Runtime/OpenGLDrv/Private/OpenGLUniformBuffer.cpp
--- a/Engine/Source/Runtime/OpenGLDrv/Private/OpenGLUniformBuffer.cpp
+++ b/Engine/Source/Runtime/OpenGLDrv/Private/OpenGLUniformBuffer.cpp
@@ -483,7 +483,7 @@ void CopyDataToUniformBuffer(const bool bCanRunOnThisThread, FOpenGLUniformBuffe
 			void* ConstantBufferCopy = RHICmdList.Alloc(ContentSize, 16);
 			FMemory::Memcpy(ConstantBufferCopy, Contents, ContentSize);
 			
-			new (RHICmdList.AllocCommand<FRHICommandGLCommand>()) FRHICommandGLCommand(
+			ALLOC_COMMAND_CL(RHICmdList, FRHICommandGLCommand)(
 				[=]() 
 				{
 					VERIFY_GL_SCOPE();
@@ -568,7 +568,7 @@ static FUniformBufferRHIRef CreateUniformBuffer(const void* Contents, const FRHI
 		{
 			NewUniformBuffer->AccessFence.Reset();
 			// Queue GL resource creation.
-			new (RHICmdList.AllocCommand<FRHICommandGLCommand>()) FRHICommandGLCommand([=]() {GLCreationFunc(); NewUniformBuffer->AccessFence.WriteAssertFence(); });
+			ALLOC_COMMAND_CL(RHICmdList, FRHICommandGLCommand)([=]() {GLCreationFunc(); NewUniformBuffer->AccessFence.WriteAssertFence(); });
 			NewUniformBuffer->AccessFence.SetRHIThreadFence();
 
 			// flush for the UBO case
@@ -661,7 +661,7 @@ FUniformBufferRHIRef FOpenGLDynamicRHI::RHICreateUniformBuffer(const void* Conte
 					}
 					else
 					{
-						new (RHICmdList.AllocCommand<FRHICommandGLCommand>()) FRHICommandGLCommand(CacheGLUniformBuffer);
+						ALLOC_COMMAND_CL(RHICmdList, FRHICommandGLCommand)(CacheGLUniformBuffer);
 					}
 				}
 			}
diff --git a/Engine/Source/Runtime/OpenGLDrv/Private/OpenGLUtil.cpp b/Engine/Source/Runtime/OpenGLDrv/Private/OpenGLUtil.cpp
--- a/Engine/Source/Runtime/OpenGLDrv/Private/OpenGLUtil.cpp
+++ b/Engine/Source/Runtime/OpenGLDrv/Private/OpenGLUtil.cpp
@@ -142,7 +142,7 @@ void RunOnGLRenderContextThread(TFunction<void(void)> GLFunc, bool bWaitForCompl
 	}
 	else
 	{
-		new (RHICmdList.AllocCommand<FRHICommandGLCommand>()) FRHICommandGLCommand(MoveTemp(GLFunc));
+		ALLOC_COMMAND_CL(RHICmdList, FRHICommandGLCommand)(MoveTemp(GLFunc));
 		if (bWaitForCompletion)
 		{
 			RHITHREAD_GLTRACE_BLOCKING;
diff --git a/Engine/Source/Runtime/OpenGLDrv/Public/OpenGLDrv.h b/Engine/Source/Runtime/OpenGLDrv/Public/OpenGLDrv.h
--- a/Engine/Source/Runtime/OpenGLDrv/Public/OpenGLDrv.h
+++ b/Engine/Source/Runtime/OpenGLDrv/Public/OpenGLDrv.h
@@ -653,7 +653,7 @@ public:
 		else\
 		{\
 			x ReturnValue = (x)0;\
-			new (RHICmdList.AllocCommand<FRHICommandGLCommand>()) FRHICommandGLCommand([&ReturnValue, GLCommand = MoveTemp(GLCommand)]() { ReturnValue = GLCommand(); }); \
+			ALLOC_COMMAND_CL(RHICmdList, FRHICommandGLCommand)([&ReturnValue, GLCommand = MoveTemp(GLCommand)]() { ReturnValue = GLCommand(); }); \
 			RHITHREAD_GLTRACE_BLOCKING;\
 			RHICmdList.ImmediateFlush(EImmediateFlushType::FlushRHIThread);\
 			return ReturnValue;\
@@ -667,7 +667,7 @@ public:
 		}\
 		else\
 		{\
-			new (RHICmdList.AllocCommand<FRHICommandGLCommand>()) FRHICommandGLCommand([&ReturnValue, GLCommand = MoveTemp(GLCommand)]() { ReturnValue = GLCommand(); }); \
+			ALLOC_COMMAND_CL(RHICmdList, FRHICommandGLCommand)([&ReturnValue, GLCommand = MoveTemp(GLCommand)]() { ReturnValue = GLCommand(); }); \
 			RHITHREAD_GLTRACE_BLOCKING;\
 			RHICmdList.ImmediateFlush(EImmediateFlushType::FlushRHIThread);\
 		}\
@@ -680,7 +680,7 @@ public:
 		}\
 		else\
 		{\
-			new (RHICmdList.AllocCommand<FRHICommandGLCommand>()) FRHICommandGLCommand( GLCommand ); \
+			ALLOC_COMMAND_CL(RHICmdList, FRHICommandGLCommand)( GLCommand ); \
 			RHITHREAD_GLTRACE_BLOCKING;\
 			RHICmdList.ImmediateFlush(EImmediateFlushType::FlushRHIThread);\
 		}\
diff --git a/Engine/Source/Runtime/OpenGLDrv/Public/OpenGLResources.h b/Engine/Source/Runtime/OpenGLDrv/Public/OpenGLResources.h
--- a/Engine/Source/Runtime/OpenGLDrv/Public/OpenGLResources.h
+++ b/Engine/Source/Runtime/OpenGLDrv/Public/OpenGLResources.h
@@ -290,7 +290,7 @@ public:
 		else
 		{
 			CreationFence.Reset();
-			new (RHICmdList.AllocCommand<FRHICommandGLCommand>()) FRHICommandGLCommand([this, CreateFunc = MoveTemp(CreateFunc)]()
+			ALLOC_COMMAND_CL(RHICmdList, FRHICommandGLCommand)([this, CreateFunc = MoveTemp(CreateFunc)]()
 			{
 				GLResourceObject = CreateFunc(this);
 				GLResourceObject->AddRef();
@@ -452,7 +452,7 @@ public:
 					FMemory::Memcpy(BuffData, InData, RealSize);
 				}
 				TransitionFence.Reset();
-				new (RHICmdList.AllocCommand<FRHICommandGLCommand>()) FRHICommandGLCommand([=]() 
+				ALLOC_COMMAND_CL(RHICmdList, FRHICommandGLCommand)([=]() 
 				{
 					CreateGLBuffer(BuffData, ResourceToUse, ResourceSize); 
 					TransitionFence.WriteAssertFence();
diff --git a/Engine/Source/Runtime/VulkanRHI/Private/VulkanCommandBuffer.cpp b/Engine/Source/Runtime/VulkanRHI/Private/VulkanCommandBuffer.cpp
--- a/Engine/Source/Runtime/VulkanRHI/Private/VulkanCommandBuffer.cpp
+++ b/Engine/Source/Runtime/VulkanRHI/Private/VulkanCommandBuffer.cpp
@@ -588,7 +588,7 @@ void FVulkanCommandBufferManager::FreeUnusedCmdBuffers()
 	else
 	{
 		check(IsInRenderingThread());
-		new (RHICmdList.AllocCommand<FRHICommandFreeUnusedCmdBuffers>()) FRHICommandFreeUnusedCmdBuffers(&Pool, Queue);
+		ALLOC_COMMAND_CL(RHICmdList, FRHICommandFreeUnusedCmdBuffers)(&Pool, Queue);
 	}
 #endif
 }
diff --git a/Engine/Source/Runtime/VulkanRHI/Private/VulkanIndexBuffer.cpp b/Engine/Source/Runtime/VulkanRHI/Private/VulkanIndexBuffer.cpp
--- a/Engine/Source/Runtime/VulkanRHI/Private/VulkanIndexBuffer.cpp
+++ b/Engine/Source/Runtime/VulkanRHI/Private/VulkanIndexBuffer.cpp
@@ -315,7 +315,7 @@ void FVulkanResourceMultiBuffer::Unlock(bool bFromRenderingThread)
 			else
 			{
 				check(IsInRenderingThread());
-				new (RHICmdList.AllocCommand<FRHICommandMultiBufferUnlock>()) FRHICommandMultiBufferUnlock(Device, PendingLock, this, DynamicBufferIndex);
+				ALLOC_COMMAND_CL(RHICmdList, FRHICommandMultiBufferUnlock)(Device, PendingLock, this, DynamicBufferIndex);
 			}
 		}
 		else
diff --git a/Engine/Source/Runtime/VulkanRHI/Private/VulkanTexture.cpp b/Engine/Source/Runtime/VulkanRHI/Private/VulkanTexture.cpp
--- a/Engine/Source/Runtime/VulkanRHI/Private/VulkanTexture.cpp
+++ b/Engine/Source/Runtime/VulkanRHI/Private/VulkanTexture.cpp
@@ -491,7 +491,7 @@ FVulkanSurface::FVulkanSurface(FVulkanDevice& InDevice, VkImageViewType Resource
 		else
 		{
 			check(IsInRenderingThread());
-			new (RHICmdList.AllocCommand<FRHICommandInitialClearTexture>()) FRHICommandInitialClearTexture(this, CreateInfo.ClearValueBinding, bTransitionToPresentable);
+			ALLOC_COMMAND_CL(RHICmdList, FRHICommandInitialClearTexture)(this, CreateInfo.ClearValueBinding, bTransitionToPresentable);
 		}
 	}
 }
@@ -1002,7 +1002,7 @@ FTexture2DRHIRef FVulkanDynamicRHI::AsyncReallocateTexture2D_RenderThread(class
 	FRHIResourceCreateInfo CreateInfo;
 	FVulkanTexture2D* NewTexture = new FVulkanTexture2D(*Device, OldTexture->GetFormat(), NewSizeX, NewSizeY, NewMipCount, OldTexture->GetNumSamples(), OldTexture->GetFlags(), CreateInfo);
 
-	new (RHICmdList.AllocCommand<FRHICommandVulkanAsyncReallocateTexture2D>()) FRHICommandVulkanAsyncReallocateTexture2D(Device->GetImmediateContext(), OldTexture, NewTexture, NewMipCount, NewSizeX, NewSizeY, RequestStatus);
+	ALLOC_COMMAND_CL(RHICmdList, FRHICommandVulkanAsyncReallocateTexture2D)(Device->GetImmediateContext(), OldTexture, NewTexture, NewMipCount, NewSizeX, NewSizeY, RequestStatus);
 
 	return NewTexture;
 }
@@ -1105,7 +1105,7 @@ void FVulkanDynamicRHI::InternalUnlockTexture2D(bool bFromRenderingThread, FText
 	else
 	{
 		check(IsInRenderingThread());
-		new (RHICmdList.AllocCommand<FRHICommandLockWriteTexture>()) FRHICommandLockWriteTexture(&Texture->Surface, SubresourceRange, Region, StagingBuffer);
+		ALLOC_COMMAND_CL(RHICmdList, FRHICommandLockWriteTexture)(&Texture->Surface, SubresourceRange, Region, StagingBuffer);
 	}
 }
 
@@ -1181,7 +1181,7 @@ void FVulkanDynamicRHI::RHIUnlockTexture2DArray(FTexture2DArrayRHIParamRef Textu
 	else
 	{
 		check(IsInRenderingThread());
-		new (RHICmdList.AllocCommand<FRHICommandLockWriteTexture>()) FRHICommandLockWriteTexture(&Texture->Surface, SubresourceRange, Region, StagingBuffer);
+		ALLOC_COMMAND_CL(RHICmdList, FRHICommandLockWriteTexture)(&Texture->Surface, SubresourceRange, Region, StagingBuffer);
 	}
 }
 
@@ -1256,7 +1256,7 @@ void FVulkanDynamicRHI::InternalUpdateTexture2D(bool bFromRenderingThread, FText
 	else
 	{
 		check(IsInRenderingThread());
-		new (RHICmdList.AllocCommand<FRHICommandLockWriteTexture>()) FRHICommandLockWriteTexture(&Texture->Surface, SubresourceRange, Region, StagingBuffer);
+		ALLOC_COMMAND_CL(RHICmdList, FRHICommandLockWriteTexture)(&Texture->Surface, SubresourceRange, Region, StagingBuffer);
 	}
 }
 
@@ -1338,7 +1338,7 @@ void FVulkanDynamicRHI::InternalUpdateTexture3D(bool bFromRenderingThread, FText
 	else
 	{
 		check(IsInRenderingThread());
-		new (RHICmdList.AllocCommand<FRHICommandLockWriteTexture>()) FRHICommandLockWriteTexture(&Texture->Surface, SubresourceRange, Region, StagingBuffer);
+		ALLOC_COMMAND_CL(RHICmdList, FRHICommandLockWriteTexture)(&Texture->Surface, SubresourceRange, Region, StagingBuffer);
 	}
 }
 
@@ -1568,7 +1568,7 @@ FVulkanTextureBase::FVulkanTextureBase(FVulkanDevice& Device, VkImageViewType Re
 	else
 	{
 		check(IsInRenderingThread());
-		new (RHICmdList.AllocCommand<FRHICommandLockWriteTexture>()) FRHICommandLockWriteTexture(&Surface, SubresourceRange, Region, StagingBuffer);
+		ALLOC_COMMAND_CL(RHICmdList, FRHICommandLockWriteTexture)(&Surface, SubresourceRange, Region, StagingBuffer);
 	}
 }
 
@@ -1906,7 +1906,7 @@ void FVulkanDynamicRHI::RHIUnlockTextureCubeFace(FTextureCubeRHIParamRef Texture
 	else
 	{
 		check(IsInRenderingThread());
-		new (RHICmdList.AllocCommand<FRHICommandLockWriteTexture>()) FRHICommandLockWriteTexture(&Texture->Surface, SubresourceRange, Region, StagingBuffer);
+		ALLOC_COMMAND_CL(RHICmdList, FRHICommandLockWriteTexture)(&Texture->Surface, SubresourceRange, Region, StagingBuffer);
 	}
 }
 
diff --git a/Engine/Source/Runtime/VulkanRHI/Private/VulkanViewport.cpp b/Engine/Source/Runtime/VulkanRHI/Private/VulkanViewport.cpp
--- a/Engine/Source/Runtime/VulkanRHI/Private/VulkanViewport.cpp
+++ b/Engine/Source/Runtime/VulkanRHI/Private/VulkanViewport.cpp
@@ -231,7 +231,7 @@ FVulkanTexture2D* FVulkanViewport::GetBackBuffer(FRHICommandList& RHICmdList)
 		}
 		else
 		{
-			new (RHICmdList.AllocCommand<FRHICommandAcquireBackBuffer>()) FRHICommandAcquireBackBuffer(this, RenderingBackBuffer);
+			ALLOC_COMMAND_CL(RHICmdList, FRHICommandAcquireBackBuffer)(this, RenderingBackBuffer);
 		}
 	}
 
@@ -914,7 +914,7 @@ void FVulkanDynamicRHI::RHIAdvanceFrameForGetViewportBackBuffer(FViewportRHIPara
 		else
 		{
 			check(IsInRenderingThread());
-			new (RHICmdList.AllocCommand<FRHICommandProcessDeferredDeletionQueue>()) FRHICommandProcessDeferredDeletionQueue(Device);
+			ALLOC_COMMAND_CL(RHICmdList, FRHICommandProcessDeferredDeletionQueue)(Device);
 		}
 	}
 }
