Reset mip generation state when we open a new renderpass.

diff --git a/Engine/Source/Runtime/VulkanRHI/Private/VulkanRenderTarget.cpp b/Engine/Source/Runtime/VulkanRHI/Private/VulkanRenderTarget.cpp
--- a/Engine/Source/Runtime/VulkanRHI/Private/VulkanRenderTarget.cpp
+++ b/Engine/Source/Runtime/VulkanRHI/Private/VulkanRenderTarget.cpp
@@ -319,7 +319,6 @@ void FTransitionAndLayoutManager::BeginRealRenderPass(FVulkanCommandListContext&
 		}
 		else
 		{
-			GenerateMipsInfo.Reset();
 			if (*Found == VK_IMAGE_LAYOUT_UNDEFINED)
 			{
 				VulkanRHI::ImagePipelineBarrier(CmdBuffer->GetHandle(), Surface.Image, EImageLayoutBarrier::Undefined, EImageLayoutBarrier::ColorAttachment, SetupImageSubresourceRange());
@@ -1639,6 +1638,8 @@ void FVulkanCommandListContext::RHIBeginRenderPass(const FRHIRenderPassInfo& InI
 		}
 	}
 
+	TransitionAndLayoutManager.GenerateMipsInfo.Reset();
+
 	FVulkanFramebuffer* Framebuffer = TransitionAndLayoutManager.GetOrCreateFramebuffer(*Device, RTInfo, RTLayout, RenderPass);
 	checkf(RenderPass != nullptr && Framebuffer != nullptr, TEXT("RenderPass not started! Bad combination of values? Depth %p #Color %d Color0 %p"), (void*)InInfo.DepthStencilRenderTarget.DepthStencilTarget, InInfo.GetNumColorRenderTargets(), (void*)InInfo.ColorRenderTargets[0].RenderTarget);
 	TransitionAndLayoutManager.BeginRealRenderPass(*this, *Device, CmdBuffer, InInfo, RTLayout, RenderPass, Framebuffer);
