DR - vk - Mirror Clamp extension not required anymore

diff --git a/Engine/Source/Runtime/VulkanRHI/Private/VulkanDevice.h b/Engine/Source/Runtime/VulkanRHI/Private/VulkanDevice.h
--- a/Engine/Source/Runtime/VulkanRHI/Private/VulkanDevice.h
+++ b/Engine/Source/Runtime/VulkanRHI/Private/VulkanDevice.h
@@ -21,7 +21,7 @@ struct FOptionalVulkanDeviceExtensions
 {
 	uint32 HasKHRMaintenance1 : 1;
 	uint32 HasKHRMaintenance2 : 1;
-	uint32 HasMirrorClampToEdge : 1;
+	//uint32 HasMirrorClampToEdge : 1;
 	uint32 HasKHRExternalMemoryCapabilities : 1;
 	uint32 HasKHRGetPhysicalDeviceProperties2 : 1;
 	uint32 HasKHRDedicatedAllocation : 1;
diff --git a/Engine/Source/Runtime/VulkanRHI/Private/VulkanLayers.cpp b/Engine/Source/Runtime/VulkanRHI/Private/VulkanLayers.cpp
--- a/Engine/Source/Runtime/VulkanRHI/Private/VulkanLayers.cpp
+++ b/Engine/Source/Runtime/VulkanRHI/Private/VulkanLayers.cpp
@@ -107,7 +107,7 @@ static const ANSICHAR* GDeviceExtensions[] =
 #if VULKAN_SUPPORTS_VALIDATION_CACHE
 	VK_EXT_VALIDATION_CACHE_EXTENSION_NAME,
 #endif
-	VK_KHR_SAMPLER_MIRROR_CLAMP_TO_EDGE_EXTENSION_NAME,
+	//VK_KHR_SAMPLER_MIRROR_CLAMP_TO_EDGE_EXTENSION_NAME,
 	nullptr
 };
 
@@ -669,7 +669,7 @@ void FVulkanDevice::ParseOptionalDeviceExtensions(const TArray<const ANSICHAR *>
 #if VULKAN_SUPPORTS_MAINTENANCE_LAYER2
 	OptionalDeviceExtensions.HasKHRMaintenance2 = HasExtension(VK_KHR_MAINTENANCE2_EXTENSION_NAME);
 #endif
-	OptionalDeviceExtensions.HasMirrorClampToEdge = HasExtension(VK_KHR_SAMPLER_MIRROR_CLAMP_TO_EDGE_EXTENSION_NAME);
+	//OptionalDeviceExtensions.HasMirrorClampToEdge = HasExtension(VK_KHR_SAMPLER_MIRROR_CLAMP_TO_EDGE_EXTENSION_NAME);
 
 #if VULKAN_SUPPORTS_DEDICATED_ALLOCATION
 	OptionalDeviceExtensions.HasKHRDedicatedAllocation = HasExtension(VK_KHR_GET_MEMORY_REQUIREMENTS_2_EXTENSION_NAME) && HasExtension(VK_KHR_DEDICATED_ALLOCATION_EXTENSION_NAME);
diff --git a/Engine/Source/Runtime/VulkanRHI/Private/VulkanState.cpp b/Engine/Source/Runtime/VulkanRHI/Private/VulkanState.cpp
--- a/Engine/Source/Runtime/VulkanRHI/Private/VulkanState.cpp
+++ b/Engine/Source/Runtime/VulkanRHI/Private/VulkanState.cpp
@@ -63,7 +63,7 @@ inline VkFilter TranslateMinFilterMode(ESamplerFilter InFilter)
 	return OutFilter;
 }
 
-inline VkSamplerAddressMode TranslateWrapMode(ESamplerAddressMode InAddressMode, const bool bSupportsMirrorClampToEdge)
+inline VkSamplerAddressMode TranslateWrapMode(ESamplerAddressMode InAddressMode)
 {
 	VkSamplerAddressMode OutAddressMode = VK_SAMPLER_ADDRESS_MODE_MAX_ENUM;
 
@@ -227,10 +227,9 @@ void FVulkanSamplerState::SetupSamplerCreateInfo(const FSamplerStateInitializerR
 	OutSamplerInfo.magFilter = TranslateMagFilterMode(Initializer.Filter);
 	OutSamplerInfo.minFilter = TranslateMinFilterMode(Initializer.Filter);
 	OutSamplerInfo.mipmapMode = TranslateMipFilterMode(Initializer.Filter);
-	const bool bSupportsMirrorClampToEdge = InDevice.GetOptionalExtensions().HasMirrorClampToEdge;
-	OutSamplerInfo.addressModeU = TranslateWrapMode(Initializer.AddressU, bSupportsMirrorClampToEdge);
-	OutSamplerInfo.addressModeV = TranslateWrapMode(Initializer.AddressV, bSupportsMirrorClampToEdge);
-	OutSamplerInfo.addressModeW = TranslateWrapMode(Initializer.AddressW, bSupportsMirrorClampToEdge);
+	OutSamplerInfo.addressModeU = TranslateWrapMode(Initializer.AddressU);
+	OutSamplerInfo.addressModeV = TranslateWrapMode(Initializer.AddressV);
+	OutSamplerInfo.addressModeW = TranslateWrapMode(Initializer.AddressW);
 
 	OutSamplerInfo.mipLodBias = Initializer.MipBias;
 	OutSamplerInfo.maxAnisotropy = FMath::Clamp((float)ComputeAnisotropyRT(Initializer.MaxAnisotropy), 1.0f, InDevice.GetLimits().maxSamplerAnisotropy);
