DR - vk - Copy from 4580765
- Fix for integer overflow which causes OOM on some AMD Vulkan drivers

diff --git a/Engine/Source/Runtime/VulkanRHI/Private/VulkanPendingState.cpp b/Engine/Source/Runtime/VulkanRHI/Private/VulkanPendingState.cpp
--- a/Engine/Source/Runtime/VulkanRHI/Private/VulkanPendingState.cpp
+++ b/Engine/Source/Runtime/VulkanRHI/Private/VulkanPendingState.cpp
@@ -124,7 +124,8 @@ FVulkanDescriptorPool* FVulkanTypedDescriptorPoolSet::PushNewPool()
 {
 	// Max number of descriptor sets layout allocations
 	const uint32 MaxSetsAllocationsBase = 32;
-	const uint32 MaxSetsAllocations = FMath::Min((MaxSetsAllocationsBase << PoolsCount), 128u);
+	// Allow max 128 setS per pool (32 << 2)
+	const uint32 MaxSetsAllocations = MaxSetsAllocationsBase << FMath::Min(PoolsCount, 2u);
 
 	auto* NewPool = new FVulkanDescriptorPool(Device, Layout, MaxSetsAllocations);
 
