DR - vk - Copy from 4580830
- Better message when user tried to run Vulkan on unsupported systems

diff --git a/Engine/Source/Runtime/VulkanRHI/Private/VulkanRHI.cpp b/Engine/Source/Runtime/VulkanRHI/Private/VulkanRHI.cpp
--- a/Engine/Source/Runtime/VulkanRHI/Private/VulkanRHI.cpp
+++ b/Engine/Source/Runtime/VulkanRHI/Private/VulkanRHI.cpp
@@ -447,7 +447,13 @@ static inline int32 PreferAdapterVendor()
 void FVulkanDynamicRHI::SelectAndInitDevice()
 {
 	uint32 GpuCount = 0;
-	VERIFYVULKANRESULT_EXPANDED(VulkanRHI::vkEnumeratePhysicalDevices(Instance, &GpuCount, nullptr));
+	VkResult Result = VulkanRHI::vkEnumeratePhysicalDevices(Instance, &GpuCount, nullptr);
+	if (Result == VK_ERROR_INITIALIZATION_FAILED)
+	{
+		FPlatformMisc::MessageBoxExt(EAppMsgType::Ok, TEXT("Cannot find a compatible Vulkan device or driver. Try updating your video driver to a more recent version and make sure your video card supports Vulkan.\n\n"), TEXT("Vulkan device not available"));
+		FPlatformMisc::RequestExitWithStatus(true, 1);
+	}
+	VERIFYVULKANRESULT_EXPANDED(Result);
 	checkf(GpuCount >= 1, TEXT("No GPU(s)/Driver(s) that support Vulkan were found! Make sure your drivers are up to date and that you are not pending a reboot."));
 
 	TArray<VkPhysicalDevice> PhysicalDevices;
