DR - vk - Adapted 4573903
Potential fix for layout map crash in Vulkan

diff --git a/Engine/Source/Runtime/VulkanRHI/Public/VulkanContext.h b/Engine/Source/Runtime/VulkanRHI/Public/VulkanContext.h
--- a/Engine/Source/Runtime/VulkanRHI/Public/VulkanContext.h
+++ b/Engine/Source/Runtime/VulkanRHI/Public/VulkanContext.h
@@ -22,6 +22,7 @@ class FTransitionAndLayoutManagerData
 public:
 	void TempCopy(const FTransitionAndLayoutManagerData& In)
 	{
+		FScopeLock Lock(&LayoutLock);
 		Framebuffers = In.Framebuffers;
 		RenderPasses = In.RenderPasses;
 		Layouts = In.Layouts;
@@ -35,6 +36,8 @@ protected:
 	};
 	TMap<uint32, FFramebufferList*> Framebuffers;
 	TMap<VkImage, VkImageLayout> Layouts;
+
+	mutable FCriticalSection LayoutLock;
 };
 
 class FTransitionAndLayoutManager : public FTransitionAndLayoutManagerData
@@ -124,16 +127,19 @@ public:
 
 	inline void NotifyDeletedImage(VkImage Image)
 	{
+		FScopeLock Lock(&LayoutLock);
 		Layouts.Remove(Image);
 	}
 
 	VkImageLayout FindLayoutChecked(VkImage Image) const
 	{
+		FScopeLock Lock(&LayoutLock);
 		return Layouts.FindChecked(Image);
 	}
 
 	VkImageLayout FindOrAddLayout(VkImage Image, VkImageLayout LayoutIfNotFound)
 	{
+		FScopeLock Lock(&LayoutLock);
 		VkImageLayout* Found = Layouts.Find(Image);
 		if (Found)
 		{
@@ -146,6 +152,7 @@ public:
 
 	VkImageLayout& FindOrAddLayoutRW(VkImage Image, VkImageLayout LayoutIfNotFound)
 	{
+		FScopeLock Lock(&LayoutLock);
 		VkImageLayout* Found = Layouts.Find(Image);
 		if (Found)
 		{
