DR - Copy from 4580354
Re-enable gpu timings on Vulkan
Timestamp queries/Stat GPU are still broken

diff --git a/Engine/Source/Runtime/VulkanRHI/Private/VulkanCommandBuffer.cpp b/Engine/Source/Runtime/VulkanRHI/Private/VulkanCommandBuffer.cpp
--- a/Engine/Source/Runtime/VulkanRHI/Private/VulkanCommandBuffer.cpp
+++ b/Engine/Source/Runtime/VulkanRHI/Private/VulkanCommandBuffer.cpp
@@ -48,10 +48,6 @@ FVulkanCmdBuffer::FVulkanCmdBuffer(FVulkanDevice* InDevice, FVulkanCommandBuffer
 	, Timing(nullptr)
 	, LastValidTiming(0)
 {
-	if (FVulkanPlatform::SupportsTimestampRenderQueries())
-	{
-		TimestampQueryPool = new FVulkanTimestampQueryPool(InDevice, NUM_TIMESTAMP_QUERIES_PER_POOL);
-	}
 
 	{
 		FScopeLock ScopeLock(CommandBufferPool->GetCS());
@@ -109,6 +105,8 @@ FVulkanCmdBuffer::~FVulkanCmdBuffer()
 	if (Timing)
 	{
 		Timing->Release();
+		delete Timing;
+		Timing = nullptr;
 	}
 }
 
diff --git a/Engine/Source/Runtime/VulkanRHI/Private/VulkanQuery.cpp b/Engine/Source/Runtime/VulkanRHI/Private/VulkanQuery.cpp
--- a/Engine/Source/Runtime/VulkanRHI/Private/VulkanQuery.cpp
+++ b/Engine/Source/Runtime/VulkanRHI/Private/VulkanQuery.cpp
@@ -678,12 +678,12 @@ void FVulkanCommandListContext::ReadAndCalculateGPUFrameTime()
 {
 	check(IsImmediate());
 
-	if (FrameTiming)
+	if (FVulkanPlatform::SupportsTimestampRenderQueries() && FrameTiming)
 	{
 		const uint64 Delta = FrameTiming->GetTiming(false);
+		const double SecondsPerCycle = FPlatformTime::GetSecondsPerCycle();
 		const double Frequency = double(FVulkanGPUTiming::GetTimingFrequency());
-		const double Cycles = FPlatformTime::GetSecondsPerCycle();
-		GGPUFrameTime = FMath::TruncToInt(double(Delta) / Frequency / Cycles);
+		GGPUFrameTime = FMath::TruncToInt(double(Delta) / Frequency / SecondsPerCycle);
 	}
 	else
 	{
@@ -701,12 +701,16 @@ FRenderQueryRHIRef FVulkanDynamicRHI::RHICreateRenderQuery(ERenderQueryType Quer
 	}
 	else if (QueryType == RQT_AbsoluteTime)
 	{
+/*
 		FVulkanTimingQuery* Query = new FVulkanTimingQuery();
 		return Query;
+*/
+	}
+	else
+	{
+		// Dummy!
+		ensureMsgf(0, TEXT("Unknown QueryType %d"), QueryType);
 	}
-
-	// Dummy!
-	ensure(0);
 	FVulkanRenderQuery* Query = new FVulkanRenderQuery(QueryType);
 	return Query;
 }
@@ -742,11 +746,13 @@ bool FVulkanDynamicRHI::RHIGetRenderQueryResult(FRenderQueryRHIParamRef QueryRHI
 	}
 	else if (BaseQuery->QueryType == RQT_AbsoluteTime)
 	{
+/*
 		FVulkanTimingQuery* Query = static_cast<FVulkanTimingQuery*>(BaseQuery);
 		uint64 TimingResult;
 		bool QueryResult = Query->GetResult(TimingResult, false);
 		OutNumPixels = TimingResult / 1000; //RHIGetRenderQueryResult assumes microseconds where vk provided ns
 		return QueryResult;
+*/
 	}
 	/*
 	FScopeLock ScopeLock(&GOcclusionQueryCS);
@@ -792,20 +798,6 @@ void FVulkanCommandListContext::RHIBeginRenderQuery(FRenderQueryRHIParamRef Quer
 	{
 		ensureMsgf(0, TEXT("Timing queries should NOT call RHIBeginRenderQuery()!"));
 	}
-	/*
-	FScopeLock ScopeLock(&GOcclusionQueryCS);
-	if (Query->QueryType == RQT_Occlusion)
-	{
-	FVulkanCmdBuffer* CmdBuffer = CommandBufferManager->GetActiveCmdBuffer();
-	const int32 QueryIndex = CurrentOcclusionQueryPool->AllocateQuery();
-	Query->Reset(CurrentOcclusionQueryPool, QueryIndex);
-	Query->Begin(CmdBuffer);
-	}
-	else
-	{
-	check(0);
-	}
-	*/
 }
 
 void FVulkanCommandListContext::RHIEndRenderQuery(FRenderQueryRHIParamRef QueryRHI)
@@ -827,6 +819,7 @@ void FVulkanCommandListContext::RHIEndRenderQuery(FRenderQueryRHIParamRef QueryR
 	}
 	else if (BaseQuery->QueryType == RQT_AbsoluteTime)
 	{
+/*
 		FVulkanTimingQuery* Query = static_cast<FVulkanTimingQuery*>(BaseQuery);
 		FVulkanCmdBuffer* CmdBuffer = CommandBufferManager->GetActiveCmdBuffer();
 		Query->Pool = Query->IsPooledAtHighLevel ? Query->Pool : CmdBuffer->PrepareTimestampQueryPool();
@@ -837,6 +830,7 @@ void FVulkanCommandListContext::RHIEndRenderQuery(FRenderQueryRHIParamRef QueryR
 #if VULKAN_QUERY_CALLSTACK
 		FPlatformStackWalk::CaptureStackBackTrace(Query->StackFrames, 16);
 #endif
+*/
 	}
 }
 
diff --git a/Engine/Source/Runtime/VulkanRHI/Private/VulkanRHI.cpp b/Engine/Source/Runtime/VulkanRHI/Private/VulkanRHI.cpp
--- a/Engine/Source/Runtime/VulkanRHI/Private/VulkanRHI.cpp
+++ b/Engine/Source/Runtime/VulkanRHI/Private/VulkanRHI.cpp
@@ -113,6 +113,9 @@ FVulkanCommandListContext::FVulkanCommandListContext(FVulkanDynamicRHI* InRHI, F
 	CommandBufferManager = new FVulkanCommandBufferManager(InDevice, this);
 	if (IsImmediate())
 	{
+		// Insert the Begin frame timestamp query. On EndDrawingViewport() we'll insert the End and immediately after a new Begin()
+		WriteBeginTimestamp(CommandBufferManager->GetActiveCmdBuffer());
+
 		// Flush the cmd buffer immediately to ensure a valid
 		// 'Last submitted' cmd buffer exists at frame 0.
 		CommandBufferManager->SubmitActiveCmdBuffer();
@@ -128,6 +131,13 @@ FVulkanCommandListContext::FVulkanCommandListContext(FVulkanDynamicRHI* InRHI, F
 
 FVulkanCommandListContext::~FVulkanCommandListContext()
 {
+	if (FVulkanPlatform::SupportsTimestampRenderQueries())
+	{
+		FrameTiming->Release();
+		delete FrameTiming;
+		FrameTiming = nullptr;
+	}
+
 	check(CommandBufferManager != nullptr);
 	delete CommandBufferManager;
 	CommandBufferManager = nullptr;
@@ -749,10 +759,6 @@ void FVulkanCommandListContext::RHIBeginDrawingViewport(FViewportRHIParamRef Vie
 	check(ViewportRHI);
 	FVulkanViewport* Viewport = ResourceCast(ViewportRHI);
 	RHI->DrawingViewport = Viewport;
-
-	check(IsImmediate());
-	// Insert the Begin frame timestamp query. On EndDrawingViewport() we'll insert the End and immediately after a new Begin()
-	WriteBeginTimestamp(CommandBufferManager->GetActiveCmdBuffer());
 }
 
 void FVulkanCommandListContext::RHIEndDrawingViewport(FViewportRHIParamRef ViewportRHI, bool bPresent, bool bLockToVsync)
diff --git a/Engine/Source/Runtime/VulkanRHI/Private/VulkanUtil.cpp b/Engine/Source/Runtime/VulkanRHI/Private/VulkanUtil.cpp
--- a/Engine/Source/Runtime/VulkanRHI/Private/VulkanUtil.cpp
+++ b/Engine/Source/Runtime/VulkanRHI/Private/VulkanUtil.cpp
@@ -14,6 +14,35 @@
 
 extern CORE_API bool GIsGPUCrashed;
 
+//#todo-rco: Horrible work-around to avoid breaking binary compatibility; move to header!
+struct FVulkanTimingQueryPool : public FVulkanQueryPool
+{
+	FVulkanTimingQueryPool(FVulkanDevice* InDevice, uint32 InBufferSize)
+		: FVulkanQueryPool(InDevice, InBufferSize * 2, VK_QUERY_TYPE_TIMESTAMP)
+		, BufferSize(InBufferSize)
+	{
+		TimestampListHandles.AddZeroed(InBufferSize * 2);
+	}
+
+	uint32 CurrentTimestamp = 0;
+	uint32 NumIssuedTimestamps = 0;
+	const uint32 BufferSize;
+
+	struct FCmdBufferFence
+	{
+		FVulkanCmdBuffer* CmdBuffer;
+		uint64 FenceCounter;
+	};
+	TArray<FCmdBufferFence> TimestampListHandles;
+
+	VulkanRHI::FStagingBuffer* ResultsBuffer = nullptr;
+};
+
+
+//#todo-rco: Horrible work-around to avoid breaking binary compatibility
+static TMap<FVulkanGPUTiming*, FVulkanTimingQueryPool*> GTimingQueryPools;
+
+
 static FString		EventDeepString(TEXT("EventTooDeep"));
 static const uint32	EventDeepCRC = FCrc::StrCrc32<TCHAR>(*EventDeepString);
 
@@ -124,14 +153,11 @@ void FVulkanGPUTiming::Initialize()
 
 	bIsTiming = false;
 
-	// Now initialize the queries for this timing object.
-	if (GIsSupported)
+	if (FVulkanPlatform::SupportsTimestampRenderQueries() && GIsSupported)
 	{
-		for (int32 Index = 0; Index < MaxTimers; ++Index)
-		{
-			Timers[Index].Begin = new FVulkanTimingQuery();
-			Timers[Index].End = new FVulkanTimingQuery();
-		}
+		FVulkanTimingQueryPool* Pool = new FVulkanTimingQueryPool(Device, 8);
+		Pool->ResultsBuffer = Device->GetStagingManager().AcquireBuffer(8 * sizeof(uint64), VK_BUFFER_USAGE_TRANSFER_DST_BIT, true);
+		GTimingQueryPools.Add(this, Pool);
 	}
 }
 
@@ -140,13 +166,15 @@ void FVulkanGPUTiming::Initialize()
  */
 void FVulkanGPUTiming::Release()
 {
-	for (int32 Index = 0; Index < MaxTimers; ++Index)
+	if (FVulkanPlatform::SupportsTimestampRenderQueries())
 	{
-		delete Timers[Index].Begin;
-		delete Timers[Index].End;
+		FVulkanTimingQueryPool* FoundPool = nullptr;
+		if (GTimingQueryPools.RemoveAndCopyValue(this, FoundPool) && FoundPool)
+		{
+			Device->GetStagingManager().ReleaseBuffer(nullptr, FoundPool->ResultsBuffer);
+			delete FoundPool;
+		}
 	}
-
-	FMemory::Memzero(Timers);
 }
 
 /**
@@ -157,16 +185,16 @@ void FVulkanGPUTiming::StartTiming(FVulkanCmdBuffer* CmdBuffer)
 	// Issue a timestamp query for the 'start' time.
 	if (GIsSupported && !bIsTiming)
 	{
-		CurrentTimerIndex = (CurrentTimerIndex + 1) % MaxTimers;
-		NumActiveTimers = FMath::Min(NumActiveTimers + 1, (int32)MaxTimers);
 		if (CmdBuffer == nullptr)
 		{
 			CmdBuffer = CmdContext->GetCommandBufferManager()->GetActiveCmdBuffer();
 		}
-		Timers[CurrentTimerIndex].BeginCmdBuffer = CmdBuffer;
-		Timers[CurrentTimerIndex].BeginFenceCounter = CmdBuffer->GetFenceSignaledCounter();
-		CmdContext->RHIEndRenderQuery(Timers[CurrentTimerIndex].Begin);
-
+		FVulkanTimingQueryPool* Pool = GTimingQueryPools.FindChecked(this);
+		Pool->CurrentTimestamp = (Pool->CurrentTimestamp + 1) % Pool->BufferSize;
+		const uint32 QueryStartIndex = Pool->CurrentTimestamp * 2;
+		VulkanRHI::vkCmdWriteTimestamp(CmdBuffer->GetHandle(), VK_PIPELINE_STAGE_TOP_OF_PIPE_BIT, Pool->GetHandle(), QueryStartIndex);
+		Pool->TimestampListHandles[QueryStartIndex].CmdBuffer = CmdBuffer;
+		Pool->TimestampListHandles[QueryStartIndex].FenceCounter = CmdBuffer->GetFenceSignaledCounter();
 		bIsTiming = true;
 	}
 }
@@ -184,10 +212,18 @@ void FVulkanGPUTiming::EndTiming(FVulkanCmdBuffer* CmdBuffer)
 		{
 			CmdBuffer = CmdContext->GetCommandBufferManager()->GetActiveCmdBuffer();
 		}
-		Timers[CurrentTimerIndex].EndCmdBuffer = CmdBuffer;
-		Timers[CurrentTimerIndex].EndFenceCounter = CmdBuffer->GetFenceSignaledCounter();
-		CmdContext->RHIEndRenderQuery(Timers[CurrentTimerIndex].End);
-		Timers[CurrentTimerIndex].IssuedTimer = true;
+		FVulkanTimingQueryPool* Pool = GTimingQueryPools.FindChecked(this);
+		check(Pool->CurrentTimestamp < Pool->BufferSize);
+		const uint32 QueryStartIndex = Pool->CurrentTimestamp * 2;
+		const uint32 QueryEndIndex = Pool->CurrentTimestamp * 2 + 1;
+		check(QueryEndIndex == QueryStartIndex + 1);	// Make sure they're adjacent indices.
+		VulkanRHI::vkCmdWriteTimestamp(CmdBuffer->GetHandle(), VK_PIPELINE_STAGE_TOP_OF_PIPE_BIT, Pool->GetHandle(), QueryEndIndex);
+		//check(CmdBuffer->IsOutsideRenderPass());
+		VulkanRHI::vkCmdCopyQueryPoolResults(CmdBuffer->GetHandle(), Pool->GetHandle(), QueryStartIndex, 2, Pool->ResultsBuffer->GetHandle(), sizeof(uint64) * QueryStartIndex, sizeof(uint64), VK_QUERY_RESULT_64_BIT);
+		VulkanRHI::vkCmdResetQueryPool(CmdBuffer->GetHandle(), Pool->GetHandle(), QueryStartIndex, 2);
+		Pool->TimestampListHandles[QueryEndIndex].CmdBuffer = CmdBuffer;
+		Pool->TimestampListHandles[QueryEndIndex].FenceCounter = CmdBuffer->GetFenceSignaledCounter();
+		Pool->NumIssuedTimestamps = FMath::Min<uint32>(Pool->NumIssuedTimestamps + 1, Pool->BufferSize);
 
 		bIsTiming = false;
 		bEndTimestampIssued = true;
@@ -204,69 +240,82 @@ uint64 FVulkanGPUTiming::GetTiming(bool bGetCurrentResultsAndBlock)
 {
 	if (GIsSupported)
 	{
+		FVulkanTimingQueryPool* Pool = GTimingQueryPools.FindChecked(this);
+		check(Pool->CurrentTimestamp < Pool->BufferSize);
+		uint64 StartTime, EndTime;
+		int32 TimestampIndex = Pool->CurrentTimestamp;
 		if (!bGetCurrentResultsAndBlock)
 		{
-			for (int32 Index = 1; Index < NumActiveTimers; ++Index)
+			// Quickly check the most recent measurements to see if any of them has been resolved.  Do not flush these queries.
+			for (uint32 IssueIndex = 1; IssueIndex < Pool->NumIssuedTimestamps; ++IssueIndex)
 			{
-				int32 OldestTimer = (CurrentTimerIndex + Index) % MaxTimers;
-				if (Timers[OldestTimer].IssuedTimer == false)
+				const uint32 QueryStartIndex = TimestampIndex * 2;
+				const uint32 QueryEndIndex = TimestampIndex * 2 + 1;
+				const FVulkanTimingQueryPool::FCmdBufferFence& StartQuerySyncPoint = Pool->TimestampListHandles[QueryStartIndex];
+				const FVulkanTimingQueryPool::FCmdBufferFence& EndQuerySyncPoint = Pool->TimestampListHandles[QueryEndIndex];
+				if (EndQuerySyncPoint.FenceCounter < EndQuerySyncPoint.CmdBuffer->GetFenceSignaledCounter() && StartQuerySyncPoint.FenceCounter < StartQuerySyncPoint.CmdBuffer->GetFenceSignaledCounter())
 				{
-					continue;
-				}
+					uint64* Data = (uint64*)Pool->ResultsBuffer->GetMappedPointer();
+					StartTime = Data[QueryStartIndex];
+					EndTime = Data[QueryEndIndex];
 
-				if (Timers[OldestTimer].BeginCmdBuffer && Timers[OldestTimer].EndCmdBuffer)
-				{
-					if (Timers[OldestTimer].BeginCmdBuffer->GetSubmittedFenceCounter() < Timers[OldestTimer].BeginFenceCounter)
+					if (EndTime > StartTime)
 					{
-						continue;
-					}
-
-					if (Timers[OldestTimer].EndCmdBuffer->GetFenceSignaledCounter() <= Timers[OldestTimer].EndFenceCounter)
-					{
-						// End timestamp query not done yet
-						continue;
-					}
-					else if (Timers[OldestTimer].End)
-					{
-						check(Timers[OldestTimer].BeginCmdBuffer->GetFenceSignaledCounter() > Timers[OldestTimer].BeginFenceCounter);
-						check(Device == CmdContext->GetDevice());
-						uint64 BeginTime = 0;
-						if (Timers[OldestTimer].Begin->GetResult(BeginTime, bGetCurrentResultsAndBlock))
-						{
-							uint64 EndTime = 0;
-							if (Timers[OldestTimer].End->GetResult(EndTime, bGetCurrentResultsAndBlock))
-							{
-								if (BeginTime < EndTime)
-								{
-									Timers[OldestTimer].IssuedTimer = false;
-									return (EndTime - BeginTime);
-								}
-							}
-						}
+						return EndTime - StartTime;
 					}
 				}
+
+				TimestampIndex = (TimestampIndex + Pool->BufferSize - 1) % Pool->BufferSize;
 			}
 		}
-		else
+
+		if (Pool->NumIssuedTimestamps > 0 || bGetCurrentResultsAndBlock)
 		{
-			if (Timers[CurrentTimerIndex].BeginCmdBuffer && Timers[CurrentTimerIndex].EndCmdBuffer)
+			// None of the (NumIssuedTimestamps - 1) measurements were ready yet,
+			// so check the oldest measurement more thoroughly.
+			// This really only happens if occlusion and frame sync event queries are disabled, otherwise those will block until the GPU catches up to 1 frame behind
+
+			const bool bBlocking = (Pool->NumIssuedTimestamps == Pool->BufferSize) || bGetCurrentResultsAndBlock;
+			const uint32 IdleStart = FPlatformTime::Cycles();
+
+			SCOPE_CYCLE_COUNTER(STAT_RenderQueryResultTime);
+
+			const uint32 QueryStartIndex = TimestampIndex * 2;
+			const uint32 QueryEndIndex = TimestampIndex * 2 + 1;
+
+			if (bBlocking)
 			{
-				checkf(Timers[CurrentTimerIndex].BeginCmdBuffer->GetSubmittedFenceCounter() >= Timers[CurrentTimerIndex].BeginFenceCounter, TEXT("Begin timestamp in a cmd buffer not submitted yet!"));
-				checkf(Timers[CurrentTimerIndex].EndCmdBuffer->GetSubmittedFenceCounter() >= Timers[CurrentTimerIndex].EndFenceCounter, TEXT("End timestamp in a cmd buffer not submitted yet!"));
-				while ((Timers[CurrentTimerIndex].EndCmdBuffer->GetFenceSignaledCounter() <= Timers[CurrentTimerIndex].EndFenceCounter))
+				const FVulkanTimingQueryPool::FCmdBufferFence& StartQuerySyncPoint = Pool->TimestampListHandles[QueryStartIndex];
+				const FVulkanTimingQueryPool::FCmdBufferFence& EndQuerySyncPoint = Pool->TimestampListHandles[QueryEndIndex];
+				bool bWaitForStart = StartQuerySyncPoint.FenceCounter == StartQuerySyncPoint.CmdBuffer->GetFenceSignaledCounter();
+				bool bWaitForEnd = EndQuerySyncPoint.FenceCounter == EndQuerySyncPoint.CmdBuffer->GetFenceSignaledCounter();
+				if (bWaitForEnd || bWaitForStart)
+				{
+					// Need to submit the open command lists.
+					Device->SubmitCommandsAndFlushGPU();
+				}
+
+				// CPU wait for query results to be ready.
+				if (bWaitForStart && StartQuerySyncPoint.FenceCounter == StartQuerySyncPoint.CmdBuffer->GetFenceSignaledCounter())
+				{
+					CmdContext->GetCommandBufferManager()->WaitForCmdBuffer(StartQuerySyncPoint.CmdBuffer);
+				}
+				if (bWaitForEnd && EndQuerySyncPoint.FenceCounter == EndQuerySyncPoint.CmdBuffer->GetFenceSignaledCounter())
 				{
-					VulkanRHI::vkDeviceWaitIdle(Device->GetInstanceHandle());
-					Timers[CurrentTimerIndex].EndCmdBuffer->GetOwner()->RefreshFenceStatus();
+					CmdContext->GetCommandBufferManager()->WaitForCmdBuffer(EndQuerySyncPoint.CmdBuffer);
 				}
+			}
+
+			GRenderThreadIdle[ERenderThreadIdleTypes::WaitingForGPUQuery] += FPlatformTime::Cycles() - IdleStart;
+			GRenderThreadNumIdle[ERenderThreadIdleTypes::WaitingForGPUQuery]++;
 
-				uint64 BeginTime = 0;
-				bool bValid = Timers[CurrentTimerIndex].Begin->GetResult(BeginTime, bGetCurrentResultsAndBlock);
-				ensure(bValid);
-				uint64 EndTime = 0;
-				bValid = Timers[CurrentTimerIndex].End->GetResult(EndTime, bGetCurrentResultsAndBlock);
-				ensure(bValid);
-				ensure(BeginTime <= EndTime);
-				return (EndTime - BeginTime);
+			uint64* Data = (uint64*)Pool->ResultsBuffer->GetMappedPointer();
+			StartTime = Data[QueryStartIndex];
+			EndTime = Data[QueryEndIndex];
+
+			if (EndTime > StartTime)
+			{
+				return EndTime - StartTime;
 			}
 		}
 	}
@@ -294,6 +343,7 @@ float FVulkanEventNodeFrame::GetRootTimingResults()
 	{
 		const uint64 GPUTiming = RootEventTiming.GetTiming(true);
 
+		// In milliseconds
 		RootResult = (double)GPUTiming / (double)RootEventTiming.GetTimingFrequency();
 	}
 
@@ -307,6 +357,7 @@ float FVulkanEventNode::GetTiming()
 	if (Timing.IsSupported())
 	{
 		const uint64 GPUTiming = Timing.GetTiming(true);
+		// In milliseconds
 		Result = (double)GPUTiming / (double)Timing.GetTimingFrequency();
 	}
 
diff --git a/Engine/Source/Runtime/VulkanRHI/Public/VulkanUtil.h b/Engine/Source/Runtime/VulkanRHI/Public/VulkanUtil.h
--- a/Engine/Source/Runtime/VulkanRHI/Public/VulkanUtil.h
+++ b/Engine/Source/Runtime/VulkanRHI/Public/VulkanUtil.h
@@ -23,7 +23,6 @@ public:
 		, bEndTimestampIssued(false)
 		, CmdContext(InCmd)
 	{
-		FMemory::Memzero(Timers);
 	}
 
 	/**
@@ -77,23 +76,6 @@ private:
 	bool bEndTimestampIssued;
 
 	FVulkanCommandListContext* CmdContext;
-	enum
-	{
-		MaxTimers = 8,
-	};
-	int32 CurrentTimerIndex = 0;
-	int32 NumActiveTimers = 0;
-	struct FBeginEndPair
-	{
-		FVulkanCmdBuffer* BeginCmdBuffer = nullptr;
-		uint64 BeginFenceCounter = 0;
-		FVulkanCmdBuffer* EndCmdBuffer = nullptr;
-		uint64 EndFenceCounter = 0;
-		FVulkanTimingQuery* Begin = nullptr;
-		FVulkanTimingQuery* End = nullptr;
-		bool IssuedTimer = false;
-	};
-	FBeginEndPair Timers[MaxTimers] = {};
 };
 
 /** A single perf event node, which tracks information about a appBeginDrawEvent/appEndDrawEvent range. */
