DR - Copy from 4592154
- Support for -gpucrashdebugging on Vulkan
- Warn if trying to dump gpu crashes but extension is not available

diff --git a/Engine/Source/Runtime/VulkanRHI/Private/VulkanDescriptorSets.h b/Engine/Source/Runtime/VulkanRHI/Private/VulkanDescriptorSets.h
--- a/Engine/Source/Runtime/VulkanRHI/Private/VulkanDescriptorSets.h
+++ b/Engine/Source/Runtime/VulkanRHI/Private/VulkanDescriptorSets.h
@@ -641,7 +641,7 @@ public:
 		return bInitialized;
 	}
 
-	void Initialize(const FDescriptorSetRemappingInfo& InRemappingInfo);
+	void Initialize(const FDescriptorSetRemappingInfo& InRemappingInfo, FVulkanShader* ComputeShader);
 
 protected:
 	// Cached data from FDescriptorSetRemappingInfo
diff --git a/Engine/Source/Runtime/VulkanRHI/Private/VulkanLayers.cpp b/Engine/Source/Runtime/VulkanRHI/Private/VulkanLayers.cpp
--- a/Engine/Source/Runtime/VulkanRHI/Private/VulkanLayers.cpp
+++ b/Engine/Source/Runtime/VulkanRHI/Private/VulkanLayers.cpp
@@ -684,14 +684,28 @@ void FVulkanDevice::ParseOptionalDeviceExtensions(const TArray<const ANSICHAR *>
 	OptionalDeviceExtensions.HasEXTValidationCache = HasExtension(VK_EXT_VALIDATION_CACHE_EXTENSION_NAME);
 #endif
 
+	bool bHasAnyCrashExtension = false;
 #if VULKAN_SUPPORTS_AMD_BUFFER_MARKER
-	OptionalDeviceExtensions.HasAMDBufferMarker = HasExtension(VK_AMD_BUFFER_MARKER_EXTENSION_NAME);
+	if (GGPUCrashDebuggingEnabled)
+	{
+		OptionalDeviceExtensions.HasAMDBufferMarker = HasExtension(VK_AMD_BUFFER_MARKER_EXTENSION_NAME);
+		bHasAnyCrashExtension = bHasAnyCrashExtension || !OptionalDeviceExtensions.HasAMDBufferMarker;
+	}
 #endif
 
 #if VULKAN_SUPPORTS_NV_DIAGNOSTIC_CHECKPOINT
-	OptionalDeviceExtensions.HasNVDiagnosticCheckpoints = HasExtension(VK_NV_DEVICE_DIAGNOSTIC_CHECKPOINTS_EXTENSION_NAME);
+	if (GGPUCrashDebuggingEnabled)
+	{
+		OptionalDeviceExtensions.HasNVDiagnosticCheckpoints = HasExtension(VK_NV_DEVICE_DIAGNOSTIC_CHECKPOINTS_EXTENSION_NAME);
+		bHasAnyCrashExtension = bHasAnyCrashExtension || !OptionalDeviceExtensions.HasNVDiagnosticCheckpoints;
+	}
 #endif
 
+	if (GGPUCrashDebuggingEnabled && !bHasAnyCrashExtension)
+	{
+		UE_LOG(LogVulkanRHI, Warning, TEXT("Tried to enable GPU crash debugging but no extension found!"));
+	}
+
 #if VULKAN_SUPPORTS_GOOGLE_DISPLAY_TIMING
 	OptionalDeviceExtensions.HasGoogleDisplayTiming = HasExtension(VK_GOOGLE_DISPLAY_TIMING_EXTENSION_NAME);
 #endif
diff --git a/Engine/Source/Runtime/VulkanRHI/Private/VulkanRHI.cpp b/Engine/Source/Runtime/VulkanRHI/Private/VulkanRHI.cpp
--- a/Engine/Source/Runtime/VulkanRHI/Private/VulkanRHI.cpp
+++ b/Engine/Source/Runtime/VulkanRHI/Private/VulkanRHI.cpp
@@ -186,7 +186,7 @@ void FVulkanDynamicRHI::Init()
 
 	{
 		IConsoleVariable* GPUCrashDebuggingCVar = IConsoleManager::Get().FindConsoleVariable(TEXT("r.GPUCrashDebugging"));
-		GGPUCrashDebuggingEnabled = GPUCrashDebuggingCVar && GPUCrashDebuggingCVar->GetInt() != 0;
+		GGPUCrashDebuggingEnabled = (GPUCrashDebuggingCVar && GPUCrashDebuggingCVar->GetInt() != 0) || FParse::Param(FCommandLine::Get(), TEXT("gpucrashdebugging"));
 	}
 
 	InitInstance();
diff --git a/Engine/Source/Runtime/VulkanRHI/Private/VulkanShaders.cpp b/Engine/Source/Runtime/VulkanRHI/Private/VulkanShaders.cpp
--- a/Engine/Source/Runtime/VulkanRHI/Private/VulkanShaders.cpp
+++ b/Engine/Source/Runtime/VulkanRHI/Private/VulkanShaders.cpp
@@ -555,9 +555,10 @@ void FVulkanDescriptorSetsLayoutInfo::FinalizeBindings(const FUniformBufferGathe
 	check(!bConsolidateAllIntoOneSet || RemappingInfo.SetInfos.Num() == 1);
 }
 
-void FVulkanComputePipelineDescriptorInfo::Initialize(const FDescriptorSetRemappingInfo& InRemappingInfo)
+void FVulkanComputePipelineDescriptorInfo::Initialize(const FDescriptorSetRemappingInfo& InRemappingInfo, FVulkanShader* ComputeShader)
 {
 	check(!bInitialized);
+	check(ComputeShader);
 
 	RemappingGlobalInfos = InRemappingInfo.StageInfos[0].Globals.GetData();
 	RemappingUBInfos = InRemappingInfo.StageInfos[0].UniformBuffers.GetData();
diff --git a/Engine/Source/Runtime/VulkanRHI/Public/VulkanCommon.h b/Engine/Source/Runtime/VulkanRHI/Public/VulkanCommon.h
--- a/Engine/Source/Runtime/VulkanRHI/Public/VulkanCommon.h
+++ b/Engine/Source/Runtime/VulkanRHI/Public/VulkanCommon.h
@@ -14,7 +14,7 @@
 
 // This defines controls shader generation (so will cause a format rebuild)
 // be careful wrt cooker/target platform not matching define-wise!!!
-#define VULKAN_ENABLE_SHADER_DEBUG_NAMES		0
+#define VULKAN_ENABLE_SHADER_DEBUG_NAMES		1
 
 namespace ShaderStage
 {
