DR - Missed file

diff --git a/Engine/Source/Runtime/VulkanRHI/Private/VulkanDescriptorSets.h b/Engine/Source/Runtime/VulkanRHI/Private/VulkanDescriptorSets.h
--- a/Engine/Source/Runtime/VulkanRHI/Private/VulkanDescriptorSets.h
+++ b/Engine/Source/Runtime/VulkanRHI/Private/VulkanDescriptorSets.h
@@ -641,7 +641,7 @@ public:
 		return bInitialized;
 	}
 
-	void Initialize(const FDescriptorSetRemappingInfo& InRemappingInfo, FVulkanShader* ComputeShader);
+	void Initialize(const FDescriptorSetRemappingInfo& InRemappingInfo);
 
 protected:
 	// Cached data from FDescriptorSetRemappingInfo
