DR - another missed file
#rnx

diff --git a/Engine/Source/Runtime/VulkanRHI/Private/VulkanShaders.cpp b/Engine/Source/Runtime/VulkanRHI/Private/VulkanShaders.cpp
--- a/Engine/Source/Runtime/VulkanRHI/Private/VulkanShaders.cpp
+++ b/Engine/Source/Runtime/VulkanRHI/Private/VulkanShaders.cpp
@@ -555,10 +555,9 @@ void FVulkanDescriptorSetsLayoutInfo::FinalizeBindings(const FUniformBufferGathe
 	check(!bConsolidateAllIntoOneSet || RemappingInfo.SetInfos.Num() == 1);
 }
 
-void FVulkanComputePipelineDescriptorInfo::Initialize(const FDescriptorSetRemappingInfo& InRemappingInfo, FVulkanShader* ComputeShader)
+void FVulkanComputePipelineDescriptorInfo::Initialize(const FDescriptorSetRemappingInfo& InRemappingInfo)
 {
 	check(!bInitialized);
-	check(ComputeShader);
 
 	RemappingGlobalInfos = InRemappingInfo.StageInfos[0].Globals.GetData();
 	RemappingUBInfos = InRemappingInfo.StageInfos[0].UniformBuffers.GetData();
