Linux: Fix blackscreen on window when leaving fullscreen

Was fixed in 4.21 and submit back to Main but was missed when merging back into Dev-Rendering

diff --git a/Engine/Source/Runtime/VulkanRHI/Private/Linux/VulkanLinuxPlatform.cpp b/Engine/Source/Runtime/VulkanRHI/Private/Linux/VulkanLinuxPlatform.cpp
--- a/Engine/Source/Runtime/VulkanRHI/Private/Linux/VulkanLinuxPlatform.cpp
+++ b/Engine/Source/Runtime/VulkanRHI/Private/Linux/VulkanLinuxPlatform.cpp
@@ -187,8 +187,3 @@ void FVulkanLinuxPlatform::WriteCrashMarker(const FOptionalVulkanDeviceExtension
 		}
 	}
 }
-
-void FVulkanLinuxPlatform::UpdateWindowSize(void* WindowHandle, uint32& Width, uint32& Height)
-{
-	SDL_Vulkan_GetDrawableSize((SDL_Window*) WindowHandle, (int32*) &Width, (int32*) &Height);
-}
diff --git a/Engine/Source/Runtime/VulkanRHI/Private/Linux/VulkanLinuxPlatform.h b/Engine/Source/Runtime/VulkanRHI/Private/Linux/VulkanLinuxPlatform.h
--- a/Engine/Source/Runtime/VulkanRHI/Private/Linux/VulkanLinuxPlatform.h
+++ b/Engine/Source/Runtime/VulkanRHI/Private/Linux/VulkanLinuxPlatform.h
@@ -57,8 +57,6 @@ public:
 #endif
 	}
 
-	static void UpdateWindowSize(void* WindowHandle, uint32& Width, uint32& Height);
-
 	static void WriteCrashMarker(const FOptionalVulkanDeviceExtensions& OptionalExtensions, VkCommandBuffer CmdBuffer, VkBuffer DestBuffer, const TArrayView<uint32>& Entries, bool bAdding);
 
 protected:
diff --git a/Engine/Source/Runtime/VulkanRHI/Private/VulkanGenericPlatform.h b/Engine/Source/Runtime/VulkanRHI/Private/VulkanGenericPlatform.h
--- a/Engine/Source/Runtime/VulkanRHI/Private/VulkanGenericPlatform.h
+++ b/Engine/Source/Runtime/VulkanRHI/Private/VulkanGenericPlatform.h
@@ -102,7 +102,4 @@ public:
 
 	// Ensure the last frame completed on the GPU
 	static bool RequiresWaitingForFrameCompletionEvent() { return true; }
-
-	// Some platforms imitate fullscreen mode by using bigger window
-	static void UpdateWindowSize(void* WindowHandle, uint32& Width, uint32& Height) {}
 };
diff --git a/Engine/Source/Runtime/VulkanRHI/Private/VulkanSwapChain.cpp b/Engine/Source/Runtime/VulkanRHI/Private/VulkanSwapChain.cpp
--- a/Engine/Source/Runtime/VulkanRHI/Private/VulkanSwapChain.cpp
+++ b/Engine/Source/Runtime/VulkanRHI/Private/VulkanSwapChain.cpp
@@ -386,9 +386,6 @@ FVulkanSwapChain::FVulkanSwapChain(VkInstance InInstance, FVulkanDevice& InDevic
 	uint32 SizeY = FVulkanPlatform::SupportsQuerySurfaceProperties() ? (SurfProperties.currentExtent.height == 0xFFFFFFFF ? Height : SurfProperties.currentExtent.height) : Height;
 	//FPlatformMisc::LowLevelOutputDebugStringf(TEXT("Create swapchain: %ux%u \n"), SizeX, SizeY);
 
-	// Some platforms cheat with the window size. Allow them to update.
-	FVulkanPlatform::UpdateWindowSize(WindowHandle, SizeX, SizeY);
-
 	VkSwapchainCreateInfoKHR SwapChainInfo;
 	ZeroVulkanStruct(SwapChainInfo, VK_STRUCTURE_TYPE_SWAPCHAIN_CREATE_INFO_KHR);
 	SwapChainInfo.surface = Surface;
diff --git a/Engine/Source/Runtime/VulkanRHI/Private/VulkanViewport.cpp b/Engine/Source/Runtime/VulkanRHI/Private/VulkanViewport.cpp
--- a/Engine/Source/Runtime/VulkanRHI/Private/VulkanViewport.cpp
+++ b/Engine/Source/Runtime/VulkanRHI/Private/VulkanViewport.cpp
@@ -689,7 +689,6 @@ bool FVulkanViewport::Present(FVulkanCommandListContext* Context, FVulkanCmdBuff
 
 			uint32 WindowSizeX = SizeX;
 			uint32 WindowSizeY = SizeY;
-			FVulkanPlatform::UpdateWindowSize(WindowHandle, WindowSizeX, WindowSizeY);
 
 			Context->RHIPushEvent(TEXT("CopyImageToBackBuffer"), FColor::Blue);
 			CopyImageToBackBuffer(CmdBuffer, true, RenderingBackBuffer->Surface.Image, BackBufferImages[AcquiredImageIndex], SizeX, SizeY, WindowSizeX, WindowSizeY);
