DR - Copy hotfixes from 4.21.1
#rnx

diff --git a/Engine/Source/Runtime/RHI/Private/PipelineStateCache.cpp b/Engine/Source/Runtime/RHI/Private/PipelineStateCache.cpp
--- a/Engine/Source/Runtime/RHI/Private/PipelineStateCache.cpp
+++ b/Engine/Source/Runtime/RHI/Private/PipelineStateCache.cpp
@@ -484,9 +484,66 @@ private:
 
 };
 
+class FGraphicsPipelineStateKey
+{
+public:
+	explicit FGraphicsPipelineStateKey(const FGraphicsPipelineStateInitializer& InInitializer)
+	{
+		FMemory::Memcpy(Initializer, InInitializer);
+		
+		if (InInitializer.BoundShaderState.VertexShaderRHI)
+		{
+			VertexShaderHash = InInitializer.BoundShaderState.VertexShaderRHI->GetHash();
+		}
+		if (InInitializer.BoundShaderState.PixelShaderRHI)
+		{
+			PixelShaderHash = InInitializer.BoundShaderState.PixelShaderRHI->GetHash();
+		}
+		if (InInitializer.BoundShaderState.GeometryShaderRHI)
+		{
+			GeometryShaderHash = InInitializer.BoundShaderState.GeometryShaderRHI->GetHash();
+		}
+		if (InInitializer.BoundShaderState.HullShaderRHI)
+		{
+			HullShaderHash = InInitializer.BoundShaderState.HullShaderRHI->GetHash();
+		}
+		if (InInitializer.BoundShaderState.DomainShaderRHI)
+		{
+			DomainShaderHash = InInitializer.BoundShaderState.DomainShaderRHI->GetHash();
+		}
+	}
+	
+	bool operator==(const FGraphicsPipelineStateKey& rhs) const
+	{
+		if (Initializer == rhs.Initializer && 
+			VertexShaderHash == rhs.VertexShaderHash &&
+			PixelShaderHash == rhs.PixelShaderHash &&
+			GeometryShaderHash == rhs.GeometryShaderHash &&
+			HullShaderHash == rhs.HullShaderHash &&
+			DomainShaderHash == rhs.DomainShaderHash)
+		{
+			return true;
+		}
+
+		return false;
+	}
+
+	FGraphicsPipelineStateInitializer Initializer;
+	FSHAHash VertexShaderHash;
+	FSHAHash PixelShaderHash;
+	FSHAHash GeometryShaderHash;
+	FSHAHash HullShaderHash;
+	FSHAHash DomainShaderHash;
+};
+
+static inline uint32 GetTypeHash(const FGraphicsPipelineStateKey& PipelineStateKey)
+{
+	return GetTypeHash(PipelineStateKey.Initializer);
+}
+
 // Typed caches for compute and graphics
 typedef TDiscardableKeyValueCache< FRHIComputeShader*, FComputePipelineState*> FComputePipelineCache;
-typedef TSharedPipelineStateCache<FGraphicsPipelineStateInitializer, FGraphicsPipelineState*> FGraphicsPipelineCache;
+typedef TSharedPipelineStateCache<FGraphicsPipelineStateKey, FGraphicsPipelineState*> FGraphicsPipelineCache;
 
 // These are the actual caches for both pipelines
 FComputePipelineCache GComputePipelineCache;
@@ -821,7 +878,8 @@ FGraphicsPipelineState* PipelineStateCache::GetAndOrCreateGraphicsPipelineState(
 
 	FGraphicsPipelineState* OutCachedState = nullptr;
 
-	bool bWasFound = GGraphicsPipelineCache.Find(*Initializer, OutCachedState);
+	FGraphicsPipelineStateKey PipelineStateKey(*Initializer);
+	bool bWasFound = GGraphicsPipelineCache.Find(PipelineStateKey, OutCachedState);
 
 	if (bWasFound == false)
 	{
@@ -847,7 +905,7 @@ FGraphicsPipelineState* PipelineStateCache::GetAndOrCreateGraphicsPipelineState(
 		}
 
 		// GGraphicsPipelineCache.Add(*Initializer, OutCachedState, LockFlags);
-		GGraphicsPipelineCache.Add(*Initializer, OutCachedState);
+		GGraphicsPipelineCache.Add(PipelineStateKey, OutCachedState);
 	}
 	else
 	{
diff --git a/Engine/Source/Runtime/RHI/Private/RHICommandList.cpp b/Engine/Source/Runtime/RHI/Private/RHICommandList.cpp
--- a/Engine/Source/Runtime/RHI/Private/RHICommandList.cpp
+++ b/Engine/Source/Runtime/RHI/Private/RHICommandList.cpp
@@ -2303,6 +2303,7 @@ FRHIShaderLibraryRef FDynamicRHI::RHICreateShaderLibrary_RenderThread(class FRHI
 
 FTextureReferenceRHIRef FDynamicRHI::RHICreateTextureReference_RenderThread(class FRHICommandListImmediate& RHICmdList, FLastRenderTimeContainer* LastRenderTime)
 {
+	// Notice no StallRHIThread needed! This function is safe on all RHIs
 	return GDynamicRHI->RHICreateTextureReference(LastRenderTime);
 }
 
diff --git a/Engine/Source/Runtime/VulkanRHI/Private/VulkanCommands.cpp b/Engine/Source/Runtime/VulkanRHI/Private/VulkanCommands.cpp
--- a/Engine/Source/Runtime/VulkanRHI/Private/VulkanCommands.cpp
+++ b/Engine/Source/Runtime/VulkanRHI/Private/VulkanCommands.cpp
@@ -190,7 +190,7 @@ void FVulkanCommandListContext::RHISetUAVParameter(FComputeShaderRHIParamRef Com
 
 void FVulkanCommandListContext::RHISetShaderTexture(FVertexShaderRHIParamRef VertexShaderRHI, uint32 TextureIndex, FTextureRHIParamRef NewTextureRHI)
 {
-	check(PendingGfxState->GetCurrentShader(ShaderStage::Vertex)->GetShaderKey() == GetShaderKey(VertexShaderRHI));
+	check(PendingGfxState->GetCurrentShaderKey(ShaderStage::Vertex) == GetShaderKey(VertexShaderRHI));
 	FVulkanTextureBase* Texture = GetVulkanTextureFromRHITexture(NewTextureRHI);
 	VkImageLayout Layout = GetLayoutForDescriptor(Texture->Surface);
 	PendingGfxState->SetTextureForStage(ShaderStage::Vertex, TextureIndex, Texture, Layout);
@@ -201,7 +201,7 @@ void FVulkanCommandListContext::RHISetShaderTexture(FHullShaderRHIParamRef HullS
 {
 	ensureMsgf(0, TEXT("Tessellation not supported yet!"));
 /*
-	check(PendingGfxState->GetCurrentShader(ShaderStage::Hull)->GetId() == GetShaderId(HullShaderRHI));
+	check(PendingGfxState->GetCurrentShaderKey(ShaderStage::Hull) == GetShaderKey(HullShaderRHI));
 	FVulkanTextureBase* Texture = GetVulkanTextureFromRHITexture(NewTextureRHI);
 	VkImageLayout Layout = GetLayoutForDescriptor(Texture->Surface);
 	PendingGfxState->SetTexture(ShaderStage::Hull, TextureIndex, Texture, Layout);
@@ -213,7 +213,7 @@ void FVulkanCommandListContext::RHISetShaderTexture(FDomainShaderRHIParamRef Dom
 {
 	ensureMsgf(0, TEXT("Tessellation not supported yet!"));
 /*
-	check(PendingGfxState->GetCurrentShader(ShaderStage::Domain)->GetId() == GetShaderId(DomainShaderRHI));
+	check(PendingGfxState->GetCurrentShaderKey(ShaderStage::Domain) == GetShaderKey(DomainShaderRHI));
 	FVulkanTextureBase* Texture = GetVulkanTextureFromRHITexture(NewTextureRHI);
 	VkImageLayout Layout = GetLayoutForDescriptor(Texture->Surface);
 	PendingGfxState->SetTexture(ShaderStage::Domain, TextureIndex, Texture, Layout);
@@ -224,7 +224,7 @@ void FVulkanCommandListContext::RHISetShaderTexture(FDomainShaderRHIParamRef Dom
 void FVulkanCommandListContext::RHISetShaderTexture(FGeometryShaderRHIParamRef GeometryShaderRHI, uint32 TextureIndex, FTextureRHIParamRef NewTextureRHI)
 {
 #if VULKAN_SUPPORTS_GEOMETRY_SHADERS
-	check(PendingGfxState->GetCurrentShader(ShaderStage::Geometry)->GetShaderKey() == GetShaderKey(GeometryShaderRHI));
+	check(PendingGfxState->GetCurrentShaderKey(ShaderStage::Geometry) == GetShaderKey(GeometryShaderRHI));
 	FVulkanTextureBase* Texture = GetVulkanTextureFromRHITexture(NewTextureRHI);
 	VkImageLayout Layout = GetLayoutForDescriptor(Texture->Surface);
 	PendingGfxState->SetTextureForStage(ShaderStage::Geometry, TextureIndex, Texture, Layout);
@@ -236,7 +236,7 @@ void FVulkanCommandListContext::RHISetShaderTexture(FGeometryShaderRHIParamRef G
 
 void FVulkanCommandListContext::RHISetShaderTexture(FPixelShaderRHIParamRef PixelShaderRHI, uint32 TextureIndex, FTextureRHIParamRef NewTextureRHI)
 {
-	check(PendingGfxState->GetCurrentShader(ShaderStage::Pixel)->GetShaderKey() == GetShaderKey(PixelShaderRHI));
+	check(PendingGfxState->GetCurrentShaderKey(ShaderStage::Pixel) == GetShaderKey(PixelShaderRHI));
 	FVulkanTextureBase* Texture = GetVulkanTextureFromRHITexture(NewTextureRHI);
 	VkImageLayout Layout = GetLayoutForDescriptor(Texture->Surface);
 	PendingGfxState->SetTextureForStage(ShaderStage::Pixel, TextureIndex, Texture, Layout);
@@ -256,7 +256,7 @@ void FVulkanCommandListContext::RHISetShaderTexture(FComputeShaderRHIParamRef Co
 
 void FVulkanCommandListContext::RHISetShaderResourceViewParameter(FVertexShaderRHIParamRef VertexShaderRHI, uint32 TextureIndex, FShaderResourceViewRHIParamRef SRVRHI)
 {
-	check(PendingGfxState->GetCurrentShader(ShaderStage::Vertex)->GetShaderKey() == GetShaderKey(VertexShaderRHI));
+	check(PendingGfxState->GetCurrentShaderKey(ShaderStage::Vertex) == GetShaderKey(VertexShaderRHI));
 	FVulkanShaderResourceView* SRV = ResourceCast(SRVRHI);
 	PendingGfxState->SetSRVForStage(ShaderStage::Vertex, TextureIndex, SRV);
 }
@@ -265,7 +265,7 @@ void FVulkanCommandListContext::RHISetShaderResourceViewParameter(FHullShaderRHI
 {
 	ensureMsgf(0, TEXT("Tessellation not supported yet!"));
 /*
-	check(PendingGfxState->GetCurrentShader(ShaderStage::Hull)->GetId() == GetShaderId(HullShaderRHI));
+	check(PendingGfxState->GetCurrentShaderKey(ShaderStage::Hull) == GetShaderKey(HullShaderRHI));
 	FVulkanShaderResourceView* SRV = ResourceCast(SRVRHI);
 	PendingGfxState->SetSRV(ShaderStage::Hull, TextureIndex, SRV);
 */
@@ -275,7 +275,7 @@ void FVulkanCommandListContext::RHISetShaderResourceViewParameter(FDomainShaderR
 {
 	ensureMsgf(0, TEXT("Tessellation not supported yet!"));
 /*
-	check(PendingGfxState->GetCurrentShader(ShaderStage::Domain)->GetId() == GetShaderId(DomainShaderRHI));
+	check(PendingGfxState->GetCurrentShaderKey(ShaderStage::Domain) == GetShaderKey(DomainShaderRHI));
 	FVulkanShaderResourceView* SRV = ResourceCast(SRVRHI);
 	PendingGfxState->SetSRV(ShaderStage::Domain, TextureIndex, SRV);
 */
@@ -284,7 +284,7 @@ void FVulkanCommandListContext::RHISetShaderResourceViewParameter(FDomainShaderR
 void FVulkanCommandListContext::RHISetShaderResourceViewParameter(FGeometryShaderRHIParamRef GeometryShaderRHI,uint32 TextureIndex,FShaderResourceViewRHIParamRef SRVRHI)
 {
 #if VULKAN_SUPPORTS_GEOMETRY_SHADERS
-	check(PendingGfxState->GetCurrentShader(ShaderStage::Geometry)->GetShaderKey() == GetShaderKey(GeometryShaderRHI));
+	check(PendingGfxState->GetCurrentShaderKey(ShaderStage::Geometry) == GetShaderKey(GeometryShaderRHI));
 	FVulkanShaderResourceView* SRV = ResourceCast(SRVRHI);
 	PendingGfxState->SetSRVForStage(ShaderStage::Geometry, TextureIndex, SRV);
 #else
@@ -294,7 +294,7 @@ void FVulkanCommandListContext::RHISetShaderResourceViewParameter(FGeometryShade
 
 void FVulkanCommandListContext::RHISetShaderResourceViewParameter(FPixelShaderRHIParamRef PixelShaderRHI,uint32 TextureIndex,FShaderResourceViewRHIParamRef SRVRHI)
 {
-	check(PendingGfxState->GetCurrentShader(ShaderStage::Pixel)->GetShaderKey() == GetShaderKey(PixelShaderRHI));
+	check(PendingGfxState->GetCurrentShaderKey(ShaderStage::Pixel) == GetShaderKey(PixelShaderRHI));
 	FVulkanShaderResourceView* SRV = ResourceCast(SRVRHI);
 	PendingGfxState->SetSRVForStage(ShaderStage::Pixel, TextureIndex, SRV);
 }
@@ -309,7 +309,7 @@ void FVulkanCommandListContext::RHISetShaderResourceViewParameter(FComputeShader
 
 void FVulkanCommandListContext::RHISetShaderSampler(FVertexShaderRHIParamRef VertexShaderRHI, uint32 SamplerIndex, FSamplerStateRHIParamRef NewStateRHI)
 {
-	check(PendingGfxState->GetCurrentShader(ShaderStage::Vertex)->GetShaderKey() == GetShaderKey(VertexShaderRHI));
+	check(PendingGfxState->GetCurrentShaderKey(ShaderStage::Vertex) == GetShaderKey(VertexShaderRHI));
 	FVulkanSamplerState* Sampler = ResourceCast(NewStateRHI);
 	PendingGfxState->SetSamplerStateForStage(ShaderStage::Vertex, SamplerIndex, Sampler);
 }
@@ -318,7 +318,7 @@ void FVulkanCommandListContext::RHISetShaderSampler(FHullShaderRHIParamRef HullS
 {
 	ensureMsgf(0, TEXT("Tessellation not supported yet!"));
 /*
-	check(PendingGfxState->GetCurrentShader(ShaderStage::Hull)->GetId() == GetShaderId(HullShaderRHI));
+	check(PendingGfxState->GetCurrentShaderKey(ShaderStage::Hull) == GetShaderKey(HullShaderRHI));
 	FVulkanSamplerState* Sampler = ResourceCast(NewStateRHI);
 	PendingGfxState->SetSamplerState(ShaderStage::Hull, SamplerIndex, Sampler);
 */
@@ -328,7 +328,7 @@ void FVulkanCommandListContext::RHISetShaderSampler(FDomainShaderRHIParamRef Dom
 {
 	ensureMsgf(0, TEXT("Tessellation not supported yet!"));
 /*
-	check(PendingGfxState->GetCurrentShader(ShaderStage::Domain)->GetId() == GetShaderId(DomainShaderRHI));
+	check(PendingGfxState->GetCurrentShaderKey(ShaderStage::Domain) == GetShaderKey(DomainShaderRHI));
 	FVulkanSamplerState* Sampler = ResourceCast(NewStateRHI);
 	PendingGfxState->SetSamplerState(ShaderStage::Domain, SamplerIndex, Sampler);
 */
@@ -337,7 +337,7 @@ void FVulkanCommandListContext::RHISetShaderSampler(FDomainShaderRHIParamRef Dom
 void FVulkanCommandListContext::RHISetShaderSampler(FGeometryShaderRHIParamRef GeometryShaderRHI, uint32 SamplerIndex, FSamplerStateRHIParamRef NewStateRHI)
 {
 #if VULKAN_SUPPORTS_GEOMETRY_SHADERS
-	check(PendingGfxState->GetCurrentShader(ShaderStage::Geometry)->GetShaderKey() == GetShaderKey(GeometryShaderRHI));
+	check(PendingGfxState->GetCurrentShaderKey(ShaderStage::Geometry) == GetShaderKey(GeometryShaderRHI));
 	FVulkanSamplerState* Sampler = ResourceCast(NewStateRHI);
 	PendingGfxState->SetSamplerStateForStage(ShaderStage::Geometry, SamplerIndex, Sampler);
 #else
@@ -347,7 +347,7 @@ void FVulkanCommandListContext::RHISetShaderSampler(FGeometryShaderRHIParamRef G
 
 void FVulkanCommandListContext::RHISetShaderSampler(FPixelShaderRHIParamRef PixelShaderRHI, uint32 SamplerIndex, FSamplerStateRHIParamRef NewStateRHI)
 {
-	check(PendingGfxState->GetCurrentShader(ShaderStage::Pixel)->GetShaderKey() == GetShaderKey(PixelShaderRHI));
+	check(PendingGfxState->GetCurrentShaderKey(ShaderStage::Pixel) == GetShaderKey(PixelShaderRHI));
 	FVulkanSamplerState* Sampler = ResourceCast(NewStateRHI);
 	PendingGfxState->SetSamplerStateForStage(ShaderStage::Pixel, SamplerIndex, Sampler);
 }
@@ -363,7 +363,7 @@ void FVulkanCommandListContext::RHISetShaderSampler(FComputeShaderRHIParamRef Co
 
 void FVulkanCommandListContext::RHISetShaderParameter(FVertexShaderRHIParamRef VertexShaderRHI, uint32 BufferIndex, uint32 BaseIndex, uint32 NumBytes, const void* NewValue)
 {
-	check(PendingGfxState->GetCurrentShader(ShaderStage::Vertex)->GetShaderKey() == GetShaderKey(VertexShaderRHI));
+	check(PendingGfxState->GetCurrentShaderKey(ShaderStage::Vertex) == GetShaderKey(VertexShaderRHI));
 
 	PendingGfxState->SetPackedGlobalShaderParameter(ShaderStage::Vertex, BufferIndex, BaseIndex, NumBytes, NewValue);
 }
@@ -372,7 +372,7 @@ void FVulkanCommandListContext::RHISetShaderParameter(FHullShaderRHIParamRef Hul
 {
 	ensureMsgf(0, TEXT("Tessellation not supported yet!"));
 /*
-	check(PendingGfxState->GetCurrentShader(ShaderStage::Hull)->GetId() == GetShaderId(HullShaderRHI));
+	check(PendingGfxState->GetCurrentShaderKey(ShaderStage::Hull) == GetShaderKey(HullShaderRHI));
 
 	PendingGfxState->SetShaderParameter(ShaderStage::Hull, BufferIndex, BaseIndex, NumBytes, NewValue);
 */
@@ -382,7 +382,7 @@ void FVulkanCommandListContext::RHISetShaderParameter(FDomainShaderRHIParamRef D
 {
 	ensureMsgf(0, TEXT("Tessellation not supported yet!"));
 /*
-	check(PendingGfxState->GetCurrentShader(ShaderStage::Domain)->GetId() == GetShaderId(DomainShaderRHI));
+	check(PendingGfxState->GetCurrentShaderKey(ShaderStage::Domain) == GetShaderKey(DomainShaderRHI));
 
 	PendingGfxState->SetShaderParameter(ShaderStage::Domain, BufferIndex, BaseIndex, NumBytes, NewValue);
 */
@@ -391,7 +391,7 @@ void FVulkanCommandListContext::RHISetShaderParameter(FDomainShaderRHIParamRef D
 void FVulkanCommandListContext::RHISetShaderParameter(FGeometryShaderRHIParamRef GeometryShaderRHI, uint32 BufferIndex, uint32 BaseIndex, uint32 NumBytes, const void* NewValue)
 {
 #if VULKAN_SUPPORTS_GEOMETRY_SHADERS
-	check(PendingGfxState->GetCurrentShader(ShaderStage::Geometry)->GetShaderKey() == GetShaderKey(GeometryShaderRHI));
+	check(PendingGfxState->GetCurrentShaderKey(ShaderStage::Geometry) == GetShaderKey(GeometryShaderRHI));
 
 	PendingGfxState->SetPackedGlobalShaderParameter(ShaderStage::Geometry, BufferIndex, BaseIndex, NumBytes, NewValue);
 #else
@@ -401,7 +401,7 @@ void FVulkanCommandListContext::RHISetShaderParameter(FGeometryShaderRHIParamRef
 
 void FVulkanCommandListContext::RHISetShaderParameter(FPixelShaderRHIParamRef PixelShaderRHI, uint32 BufferIndex, uint32 BaseIndex, uint32 NumBytes, const void* NewValue)
 {
-	check(PendingGfxState->GetCurrentShader(ShaderStage::Pixel)->GetShaderKey() == GetShaderKey(PixelShaderRHI));
+	check(PendingGfxState->GetCurrentShaderKey(ShaderStage::Pixel) == GetShaderKey(PixelShaderRHI));
 
 	PendingGfxState->SetPackedGlobalShaderParameter(ShaderStage::Pixel, BufferIndex, BaseIndex, NumBytes, NewValue);
 }
@@ -602,7 +602,7 @@ inline void FVulkanCommandListContext::SetShaderUniformBuffer(ShaderStage::EStag
 #if VULKAN_ENABLE_AGGRESSIVE_STATS
 	SCOPE_CYCLE_COUNTER(STAT_VulkanSetUniformBufferTime);
 #endif
-	check(Shader->GetShaderKey() == PendingGfxState->GetCurrentShader(Stage)->GetShaderKey());
+	check(Shader->GetShaderKey() == PendingGfxState->GetCurrentShaderKey(Stage));
 
 	const FVulkanShaderHeader& CodeHeader = Shader->GetCodeHeader();
 	const bool bUseRealUBs = FVulkanPlatform::UseRealUBsOptimization(CodeHeader.bHasRealUBs);
diff --git a/Engine/Source/Runtime/VulkanRHI/Private/VulkanPendingState.cpp b/Engine/Source/Runtime/VulkanRHI/Private/VulkanPendingState.cpp
--- a/Engine/Source/Runtime/VulkanRHI/Private/VulkanPendingState.cpp
+++ b/Engine/Source/Runtime/VulkanRHI/Private/VulkanPendingState.cpp
@@ -472,9 +472,12 @@ void FVulkanPendingGfxState::PrepareForDraw(FVulkanCmdBuffer* CmdBuffer)
 					if (VertexInputStateInfo.Attributes[AttributeIndex].binding == CurrBinding.binding)
 					{
 #if VULKAN_ENABLE_SHADER_DEBUG_NAMES
+						extern FVulkanShader* FVulkanShaderFactory_LookupShader(FVulkanShaderFactory& ShaderFactory, EShaderFrequency ShaderFrequency, uint64 ShaderKey);
+						uint64 VertexShaderKey = GetCurrentShaderKey(ShaderStage::Vertex);
+						FVulkanShader* VertexShader = FVulkanShaderFactory_LookupShader(Device->GetShaderFactory(), SF_Vertex, VertexShaderKey);
 						UE_LOG(LogVulkanRHI, Warning, TEXT("Missing binding on location %d in '%s' vertex shader"),
 							CurrBinding.binding,
-							*GetCurrentShader(ShaderStage::Vertex)->GetDebugName());
+							VertexShader ? *VertexShader->GetDebugName() : TEXT("Null"));
 #else
 						UE_LOG(LogVulkanRHI, Warning, TEXT("Missing binding on location %d in vertex shader"), CurrBinding.binding);
 #endif
diff --git a/Engine/Source/Runtime/VulkanRHI/Private/VulkanPendingState.h b/Engine/Source/Runtime/VulkanRHI/Private/VulkanPendingState.h
--- a/Engine/Source/Runtime/VulkanRHI/Private/VulkanPendingState.h
+++ b/Engine/Source/Runtime/VulkanRHI/Private/VulkanPendingState.h
@@ -181,14 +181,14 @@ public:
 		//FMemory::Memzero(PendingStreams);
 	}
 
-	const FVulkanShader* GetCurrentShader(EShaderFrequency Frequency) const
+	const uint64 GetCurrentShaderKey(EShaderFrequency Frequency) const
 	{
-		return (CurrentPipeline ? CurrentPipeline->GetShader(Frequency) : nullptr);
+		return (CurrentPipeline ? CurrentPipeline->GetShaderKey(Frequency) : 0);
 	}
 
-	const FVulkanShader* GetCurrentShader(ShaderStage::EStage Stage) const
+	const uint64 GetCurrentShaderKey(ShaderStage::EStage Stage) const
 	{
-		return GetCurrentShader(ShaderStage::GetFrequencyForGfxStage(Stage));
+		return GetCurrentShaderKey(ShaderStage::GetFrequencyForGfxStage(Stage));
 	}
 
 	void SetViewport(uint32 MinX, uint32 MinY, float MinZ, uint32 MaxX, uint32 MaxY, float MaxZ)
diff --git a/Engine/Source/Runtime/VulkanRHI/Private/VulkanPipeline.cpp b/Engine/Source/Runtime/VulkanRHI/Private/VulkanPipeline.cpp
--- a/Engine/Source/Runtime/VulkanRHI/Private/VulkanPipeline.cpp
+++ b/Engine/Source/Runtime/VulkanRHI/Private/VulkanPipeline.cpp
@@ -611,7 +611,9 @@ FVulkanRHIGraphicsPipelineState* FVulkanPipelineStateCacheManager::CreateAndAdd(
 	{
 		// Create the pipeline
 		double BeginTime = FPlatformTime::Seconds();
-		CreateGfxPipelineFromEntry(GfxEntry.Get(), &PSOInitializer.BoundShaderState, Pipeline);
+		FVulkanShader* VulkanShaders[ShaderStage::NumStages];
+		GetVulkanShaders(PSOInitializer.BoundShaderState, VulkanShaders);
+		CreateGfxPipelineFromEntry(GfxEntry.Get(), VulkanShaders, Pipeline);
 
 		// Recover if we failed to create the pipeline.
 		if (!Pipeline->GetHandle())
@@ -1071,14 +1073,10 @@ FGfxEntryKey FVulkanPipelineStateCacheManager::FGfxPipelineEntry::CreateKey() co
 	return Result;
 }
 
-void FVulkanPipelineStateCacheManager::CreateGfxPipelineFromEntry(FGfxPipelineEntry* GfxEntry, const FBoundShaderStateInput* BSI, FVulkanGfxPipeline* Pipeline)
+void FVulkanPipelineStateCacheManager::CreateGfxPipelineFromEntry(FGfxPipelineEntry* GfxEntry, FVulkanShader* Shaders[ShaderStage::NumStages], FVulkanGfxPipeline* Pipeline)
 {
-	FVulkanShader* Shaders[ShaderStage::NumStages];
-
 	if (!GfxEntry->bLoaded)
 	{
-		check(BSI);
-		GetVulkanShaders(*BSI, Shaders);
 		GfxEntry->GetOrCreateShaderModules(Shaders);
 	}
 
@@ -1365,48 +1363,6 @@ inline FVulkanLayout* FVulkanPipelineStateCacheManager::FindOrAddLayout(const FV
 	return Layout;
 }
 
-void FVulkanPipelineStateCacheManager::GetVulkanShaders(const FBoundShaderStateInput& BSI, FVulkanShader* OutShaders[ShaderStage::NumStages])
-{
-	FMemory::Memzero(OutShaders, ShaderStage::NumStages * sizeof(*OutShaders));
-
-	OutShaders[ShaderStage::Vertex] = ResourceCast(BSI.VertexShaderRHI);
-
-	if (BSI.PixelShaderRHI)
-	{
-		OutShaders[ShaderStage::Pixel] = ResourceCast(BSI.PixelShaderRHI);
-	}
-	else if (GMaxRHIFeatureLevel <= ERHIFeatureLevel::ES3_1)
-	{
-		// Some mobile devices expect PS stage (S7 Adreno)
-		OutShaders[ShaderStage::Pixel] = ResourceCast(TShaderMapRef<FNULLPS>(GetGlobalShaderMap(GMaxRHIFeatureLevel))->GetPixelShader());
-	}
-
-	if (BSI.GeometryShaderRHI)
-	{
-#if VULKAN_SUPPORTS_GEOMETRY_SHADERS
-		OutShaders[ShaderStage::Geometry] = ResourceCast(BSI.GeometryShaderRHI);
-#else
-		ensureMsgf(0, TEXT("Geometry not supported!"));
-#endif
-	}
-
-	if (BSI.HullShaderRHI)
-	{
-		ensureMsgf(0, TEXT("Tessellation not supported yet!"));
-		/*
-		// Can't have Hull w/o Domain
-		check(BSI.DomainShaderRHI);
-		OutShaders[ShaderStage::Hull] = ResourceCast(BSI.HullShaderRHI);
-		OutShaders[ShaderStage::Domain] = ResourceCast(BSI.DomainShaderRHI);
-		*/
-	}
-	else
-	{
-		// Can't have Domain w/o Hull
-		check(BSI.DomainShaderRHI == nullptr);
-	}
-}
-
 FVulkanGfxLayout* FVulkanPipelineStateCacheManager::GetOrGenerateGfxLayout(const FGraphicsPipelineStateInitializer& PSOInitializer,
 	FVulkanShader*const* Shaders, FVulkanVertexInputStateInfo& OutVertexInputState)
 {
@@ -1649,8 +1605,6 @@ FVulkanRHIGraphicsPipelineState* FVulkanPipelineStateCacheManager::FindInRuntime
 
 	if (Found)
 	{
-		// TODO: looks like with enabled shader caching (r.Vulkan.CacheShaders), PSO can end up with stale shader pointers
-		(*Found)->BSI = Initializer.BoundShaderState;
 #if VULKAN_ENABLE_LRU_CACHE
 		PipelineLRU.Touch(*Found);
 #endif
@@ -1842,7 +1796,9 @@ void FVulkanPipelineStateCacheManager::FVKPipelineLRU::EnsureVkPipelineAndAddToL
 
 		if (!GfxPipeline->GetHandle())
 		{
-			PipelineStateCache->CreateGfxPipelineFromEntry(GfxPipelineEntry, &Pipeline->BSI, GfxPipeline);
+			FVulkanShader* VulkanShaders[ShaderStage::NumStages];
+			GetVulkanShaders(GfxPipeline->Device, *Pipeline, VulkanShaders);
+			PipelineStateCache->CreateGfxPipelineFromEntry(GfxPipelineEntry, VulkanShaders, GfxPipeline);
 
 			// Failed to create VkPipeline again
 			check(GfxPipeline->GetHandle());
@@ -1951,3 +1907,56 @@ bool FVulkanPipelineStateCacheManager::FVulkanLRUCacheFile::Load(FArchive& Ar)
 	return true;
 }
 #endif
+
+
+void GetVulkanShaders(const FBoundShaderStateInput& BSI, FVulkanShader* OutShaders[ShaderStage::NumStages])
+{
+	FMemory::Memzero(OutShaders, ShaderStage::NumStages * sizeof(*OutShaders));
+
+	OutShaders[ShaderStage::Vertex] = ResourceCast(BSI.VertexShaderRHI);
+
+	if (BSI.PixelShaderRHI)
+	{
+		OutShaders[ShaderStage::Pixel] = ResourceCast(BSI.PixelShaderRHI);
+	}
+	else if (GMaxRHIFeatureLevel <= ERHIFeatureLevel::ES3_1)
+	{
+		// Some mobile devices expect PS stage (S7 Adreno)
+		OutShaders[ShaderStage::Pixel] = ResourceCast(TShaderMapRef<FNULLPS>(GetGlobalShaderMap(GMaxRHIFeatureLevel))->GetPixelShader());
+	}
+
+	if (BSI.GeometryShaderRHI)
+	{
+#if VULKAN_SUPPORTS_GEOMETRY_SHADERS
+		OutShaders[ShaderStage::Geometry] = ResourceCast(BSI.GeometryShaderRHI);
+#else
+		ensureMsgf(0, TEXT("Geometry not supported!"));
+#endif
+	}
+
+	if (BSI.HullShaderRHI)
+	{
+		ensureMsgf(0, TEXT("Tessellation not supported yet!"));
+		/*
+		// Can't have Hull w/o Domain
+		check(BSI.DomainShaderRHI);
+		OutShaders[ShaderStage::Hull] = ResourceCast(BSI.HullShaderRHI);
+		OutShaders[ShaderStage::Domain] = ResourceCast(BSI.DomainShaderRHI);
+		*/
+	}
+	else
+	{
+		// Can't have Domain w/o Hull
+		check(BSI.DomainShaderRHI == nullptr);
+	}
+}
+
+extern void FVulkanShaderFactory_LookupShaders(FVulkanShaderFactory& ShaderFactory, const uint64 InShaderKeys[ShaderStage::NumStages], FVulkanShader* OutShaders[ShaderStage::NumStages]);
+
+void GetVulkanShaders(FVulkanDevice* Device, const FVulkanRHIGraphicsPipelineState& GfxPipelineState, FVulkanShader* OutShaders[ShaderStage::NumStages])
+{
+	FMemory::Memzero(OutShaders, ShaderStage::NumStages * sizeof(*OutShaders));
+	
+	//Device->GetShaderFactory().LookupShaders(GfxPipelineState.ShaderKeys, OutShaders);
+	FVulkanShaderFactory_LookupShaders(Device->GetShaderFactory(), GfxPipelineState.ShaderKeys, OutShaders);
+}
diff --git a/Engine/Source/Runtime/VulkanRHI/Private/VulkanPipeline.h b/Engine/Source/Runtime/VulkanRHI/Private/VulkanPipeline.h
--- a/Engine/Source/Runtime/VulkanRHI/Private/VulkanPipeline.h
+++ b/Engine/Source/Runtime/VulkanRHI/Private/VulkanPipeline.h
@@ -31,6 +31,28 @@ static inline uint64 GetShaderKey(T* ShaderType)
 	return VulkanShader ? VulkanShader->GetShaderKey() : 0;
 }
 
+inline uint64 GetShaderKeyForGfxStage(const FBoundShaderStateInput& BSI, ShaderStage::EStage Stage)
+{
+	switch (Stage)
+	{
+	case ShaderStage::Vertex:
+		return GetShaderKey(BSI.VertexShaderRHI);
+	case ShaderStage::Pixel:
+		return GetShaderKey(BSI.PixelShaderRHI);
+#if VULKAN_SUPPORTS_GEOMETRY_SHADERS
+	case ShaderStage::Geometry:
+		return GetShaderKey(BSI.GeometryShaderRHI);
+#endif
+	default:
+		check(0);
+	}
+	
+	return 0;
+}
+
+void GetVulkanShaders(const FBoundShaderStateInput& BSI, FVulkanShader* OutShaders[ShaderStage::NumStages]);
+void GetVulkanShaders(FVulkanDevice* Device, const FVulkanRHIGraphicsPipelineState& GfxPipelineState, FVulkanShader* OutShaders[ShaderStage::NumStages]);
+
 class FVulkanPipelineStateCacheManager
 {
 public:
@@ -547,14 +569,13 @@ private:
 	FShaderUCodeCache ShaderCache;
 
 	FVulkanRHIGraphicsPipelineState* CreateAndAdd(const FGraphicsPipelineStateInitializer& PSOInitializer, FGfxPSIKey PSIKey, TGfxPipelineEntrySharedPtr GfxEntry, FGfxEntryKey GfxEntryKey);
-	void CreateGfxPipelineFromEntry(FGfxPipelineEntry* GfxEntry, const FBoundShaderStateInput* BSI, FVulkanGfxPipeline* Pipeline);
+	void CreateGfxPipelineFromEntry(FGfxPipelineEntry* GfxEntry, FVulkanShader* Shaders[ShaderStage::NumStages], FVulkanGfxPipeline* Pipeline);
 
 	FGfxPipelineEntry* CreateGfxEntry(const FGraphicsPipelineStateInitializer& PSOInitializer);
 
 	bool Load(const TArray<FString>& CacheFilenames);
 	void DestroyCache();
 
-	void GetVulkanShaders(const FBoundShaderStateInput& BSI, FVulkanShader* OutShaders[ShaderStage::NumStages]);
 	FVulkanGfxLayout* GetOrGenerateGfxLayout(const FGraphicsPipelineStateInitializer& PSOInitializer, FVulkanShader*const* Shaders, FVulkanVertexInputStateInfo& OutVertexInputState);
 
 	TMap<FGfxEntryKey, FVulkanGfxPipeline*> EntryKeyToGfxPipelineMap;
@@ -726,9 +747,13 @@ class FVulkanRHIGraphicsPipelineState : public FRHIGraphicsPipelineState
 public:
 	FVulkanRHIGraphicsPipelineState(const FBoundShaderStateInput& InBSI, FVulkanGfxPipeline* InPipeline, EPrimitiveType InPrimitiveType)
 		: Pipeline(InPipeline)
-		, BSI(InBSI)
 		, PrimitiveType(InPrimitiveType)
 	{
+		for (int32 StageIdx = 0; StageIdx < ShaderStage::NumStages; ++StageIdx)
+		{
+			ShaderKeys[StageIdx] = GetShaderKeyForGfxStage(InBSI, (ShaderStage::EStage)StageIdx);
+		}
+		
 		bHasInputAttachments = InPipeline->GetGfxLayout().GetDescriptorSetsLayout().HasInputAttachments();
 	}
 
@@ -739,27 +764,14 @@ public:
 		Pipeline->Bind(CmdBuffer);
 	}
 
-	inline const FVulkanShader* GetShader(EShaderFrequency Frequency) const
+	inline const uint64 GetShaderKey(EShaderFrequency Frequency) const
 	{
-		FVulkanShader* Shader = nullptr;
-
-		switch (Frequency)
-		{
-		case SF_Vertex:   Shader = ResourceCast(BSI.VertexShaderRHI); break;
-		case SF_Hull:     Shader = ResourceCast(BSI.HullShaderRHI); break;
-		case SF_Domain:   Shader = ResourceCast(BSI.DomainShaderRHI); break;
-		case SF_Pixel:    Shader = ResourceCast(BSI.PixelShaderRHI); break;
-		case SF_Geometry: Shader = ResourceCast(BSI.GeometryShaderRHI); break;
-		default:
-			check(0);
-			break;
-		}
-
-		return Shader;
+		ShaderStage::EStage Stage = ShaderStage::GetStageForFrequency(Frequency);
+		return ShaderKeys[Stage];
 	}
 
 	TRefCountPtr<FVulkanGfxPipeline>	Pipeline;
-	FBoundShaderStateInput				BSI;
+	uint64								ShaderKeys[ShaderStage::NumStages];
 	TEnumAsByte<EPrimitiveType>			PrimitiveType;
 	bool								bHasInputAttachments;
 };
diff --git a/Engine/Source/Runtime/VulkanRHI/Private/VulkanPipelineState.cpp b/Engine/Source/Runtime/VulkanRHI/Private/VulkanPipelineState.cpp
--- a/Engine/Source/Runtime/VulkanRHI/Private/VulkanPipelineState.cpp
+++ b/Engine/Source/Runtime/VulkanRHI/Private/VulkanPipelineState.cpp
@@ -182,6 +182,7 @@ bool FVulkanComputePipelineDescriptorState::InternalUpdateDescriptorSets(FVulkan
 	return true;
 }
 
+extern FVulkanShader* FVulkanShaderFactory_LookupShader(FVulkanShaderFactory& ShaderFactory, EShaderFrequency ShaderFrequency, uint64 ShaderKey);
 
 FVulkanGraphicsPipelineDescriptorState::FVulkanGraphicsPipelineDescriptorState(FVulkanDevice* InDevice, FVulkanRHIGraphicsPipelineState* InGfxPipeline)
 	: FVulkanCommonPipelineDescriptorState(InDevice)
@@ -200,38 +201,55 @@ FVulkanGraphicsPipelineDescriptorState::FVulkanGraphicsPipelineDescriptorState(F
 
 	UsedSetsMask = PipelineDescriptorInfo->HasDescriptorsInSetMask;
 
-	PackedUniformBuffers[ShaderStage::Vertex].Init(InGfxPipeline->GetShader(SF_Vertex)->GetCodeHeader(), PackedUniformBuffersMask[ShaderStage::Vertex]);
+	FVulkanShaderFactory& ShaderFactory = Device->GetShaderFactory();
+	
+	//const FVulkanVertexShader* VertexShader = ShaderFactory.LookupShader<FVulkanVertexShader>(InGfxPipeline->GetShaderKey(SF_Vertex));
+	FVulkanShader* VertexShader = FVulkanShaderFactory_LookupShader(ShaderFactory, SF_Vertex, InGfxPipeline->GetShaderKey(SF_Vertex));
+	check(VertexShader);
+
+	PackedUniformBuffers[ShaderStage::Vertex].Init(VertexShader->GetCodeHeader(), PackedUniformBuffersMask[ShaderStage::Vertex]);
 	UsedStagesMask |= 1 << ShaderStage::Vertex;
 	HasDescriptorsPerStageMask |= DescriptorSetsLayout->RemappingInfo.StageInfos[ShaderStage::Vertex].IsEmpty() ? 0 : (1 << ShaderStage::Vertex);
 	UsedPackedUBStagesMask |= DescriptorSetsLayout->RemappingInfo.StageInfos[ShaderStage::Vertex].PackedUBBindingIndices.Num() > 0 ? (1 << ShaderStage::Vertex) : 0;
 
-	if (InGfxPipeline->GetShader(SF_Pixel))
+	uint64 PixelShaderKey = InGfxPipeline->GetShaderKey(SF_Pixel);
+	if (PixelShaderKey)
 	{
-		PackedUniformBuffers[ShaderStage::Pixel].Init(InGfxPipeline->GetShader(SF_Pixel)->GetCodeHeader(), PackedUniformBuffersMask[ShaderStage::Pixel]);
+		//const FVulkanPixelShader* PixelShader = ShaderFactory.LookupShader<FVulkanPixelShader>(PixelShaderKey);
+		FVulkanShader* PixelShader = FVulkanShaderFactory_LookupShader(ShaderFactory, SF_Pixel, PixelShaderKey);
+		check(PixelShader);
+
+		PackedUniformBuffers[ShaderStage::Pixel].Init(PixelShader->GetCodeHeader(), PackedUniformBuffersMask[ShaderStage::Pixel]);
 		UsedStagesMask |= 1 << ShaderStage::Pixel;
 		HasDescriptorsPerStageMask |= DescriptorSetsLayout->RemappingInfo.StageInfos[ShaderStage::Pixel].IsEmpty() ? 0 : (1 << ShaderStage::Pixel);
 		UsedPackedUBStagesMask |= DescriptorSetsLayout->RemappingInfo.StageInfos[ShaderStage::Pixel].PackedUBBindingIndices.Num() > 0 ? (1 << ShaderStage::Pixel) : 0;
 	}
-	if (InGfxPipeline->GetShader(SF_Geometry))
-	{
+
 #if VULKAN_SUPPORTS_GEOMETRY_SHADERS
-		PackedUniformBuffers[ShaderStage::Geometry].Init(InGfxPipeline->GetShader(SF_Geometry)->GetCodeHeader(), PackedUniformBuffersMask[ShaderStage::Geometry]);
+	uint64 GeometryShaderKey = InGfxPipeline->GetShaderKey(SF_Geometry);
+	if (GeometryShaderKey)
+	{
+		//const FVulkanGeometryShader* GeometryShader = ShaderFactory.LookupShader<FVulkanGeometryShader>(GeometryShaderKey);
+		FVulkanShader* GeometryShader = FVulkanShaderFactory_LookupShader(ShaderFactory, SF_Geometry, GeometryShaderKey);
+		check(GeometryShader);
+
+		PackedUniformBuffers[ShaderStage::Geometry].Init(GeometryShader->GetCodeHeader(), PackedUniformBuffersMask[ShaderStage::Geometry]);
 		UsedStagesMask |= 1 << ShaderStage::Geometry;
 		HasDescriptorsPerStageMask |= DescriptorSetsLayout->RemappingInfo.StageInfos[ShaderStage::Geometry].IsEmpty() ? 0 : (1 << ShaderStage::Geometry);
 		UsedPackedUBStagesMask |= DescriptorSetsLayout->RemappingInfo.StageInfos[ShaderStage::Geometry].PackedUBBindingIndices.Num() > 0 ? (1 << ShaderStage::Geometry) : 0;
-#else
-		ensureMsgf(0, TEXT("Geometry not supported!"));
-#endif
 	}
-	if (InGfxPipeline->GetShader(SF_Hull))
+#endif
+
+/*
+	uint64 HullShaderKey = InGfxPipeline->GetShaderKey(SF_Hull);
+	if (HullShaderKey)
 	{
 		ensureMsgf(0, TEXT("Tessellation not supported yet!"));
-/*
 		PackedUniformBuffers[ShaderStage::Domain].Init(CodeHeaderPerStage[ShaderStage::Domain], PackedUniformBuffersMask[ShaderStage::Domain], UniformBuffersWithDataMask[ShaderStage::Domain], ResourcesDirtyMask[ShaderStage::Domain]);
 		PackedUniformBuffers[ShaderStage::Hull].Init(CodeHeaderPerStage[ShaderStage::Hull], PackedUniformBuffersMask[ShaderStage::Hull], UniformBuffersWithDataMask[ShaderStage::Hull], ResourcesDirtyMask[ShaderStage::Domain]);
 		UsedStagesMask |= (1 << ShaderStage::Hull) | (1 << ShaderStage::Domain);
-*/
 	}
+*/
 
 	//checkf(PipelineDescriptorInfo->DescriptorSets.Num() <= DescriptorSet::NumGfxStages, TEXT("Update ResourcesDirty & ResourcesDirtyMask!"));
 	for (int32 Index = 0; Index < PipelineDescriptorInfo->RemappingInfo->SetInfos.Num(); ++Index)
@@ -421,6 +439,7 @@ void FVulkanCommandListContext::RHISetGraphicsPipelineState(FGraphicsPipelineSta
 	FVulkanPipelineStateCacheManager* PipelineStateCache = Device->GetPipelineStateCache();
 	PipelineStateCache->PipelineLRU.Touch(Pipeline);
 #endif
+	check(Pipeline);
 
 	FVulkanCmdBuffer* CmdBuffer = CommandBufferManager->GetActiveCmdBuffer();
 	if (PendingGfxState->SetGfxPipeline(Pipeline) || !CmdBuffer->bHasPipeline)
diff --git a/Engine/Source/Runtime/VulkanRHI/Private/VulkanShaders.cpp b/Engine/Source/Runtime/VulkanRHI/Private/VulkanShaders.cpp
--- a/Engine/Source/Runtime/VulkanRHI/Private/VulkanShaders.cpp
+++ b/Engine/Source/Runtime/VulkanRHI/Private/VulkanShaders.cpp
@@ -31,13 +31,51 @@ static TAutoConsoleVariable<int32> GDescriptorSetLayoutMode(
 	ECVF_ReadOnly | ECVF_RenderThreadSafe
 );
 
-int32 GCacheCreatedShaders = 1;
-FAutoConsoleVariableRef CVarCacheCreatedShaders(
-	TEXT("r.Vulkan.CacheShaders"),
-	GCacheCreatedShaders,
-	TEXT("Whether to cache created shaders to avoid shader duplication\n"),
-	ECVF_ReadOnly | ECVF_RenderThreadSafe
-	);
+// Vulkan hotfix for 4.21
+struct FVulkanShaderFactory_Hotfix421
+{
+	FRWLock Lock;
+	TMap<uint64, FVulkanShader*> ShaderMap[SF_NumFrequencies];
+};
+
+// This should be member of FVulkanShaderFactory
+FVulkanShader* FVulkanShaderFactory_LookupShader(FVulkanShaderFactory& ShaderFactory, EShaderFrequency ShaderFrequency, uint64 ShaderKey)
+{
+	if (ShaderKey)
+	{
+		FVulkanShaderFactory_Hotfix421& ShaderFactoryHotfix = reinterpret_cast<FVulkanShaderFactory_Hotfix421&>(ShaderFactory);
+		
+		FRWScopeLock ScopedLock(ShaderFactoryHotfix.Lock, SLT_ReadOnly);
+		FVulkanShader* const * FoundShaderPtr = ShaderFactoryHotfix.ShaderMap[ShaderFrequency].Find(ShaderKey);
+		if (FoundShaderPtr)
+		{
+			return *FoundShaderPtr;
+		}
+	}
+	return nullptr;
+}
+
+// This should be member of FVulkanShaderFactory
+void FVulkanShaderFactory_LookupShaders(FVulkanShaderFactory& ShaderFactory, const uint64 InShaderKeys[ShaderStage::NumStages], FVulkanShader* OutShaders[ShaderStage::NumStages])
+{
+	FVulkanShaderFactory_Hotfix421& ShaderFactoryHotfix = reinterpret_cast<FVulkanShaderFactory_Hotfix421&>(ShaderFactory);
+
+	FRWScopeLock ScopedLock(ShaderFactoryHotfix.Lock, SLT_ReadOnly);
+	
+	for (int32 Idx = 0; Idx < ShaderStage::NumStages; ++Idx)
+	{
+		uint64 ShaderKey = InShaderKeys[Idx];
+		if (ShaderKey)
+		{
+			EShaderFrequency ShaderFrequency = ShaderStage::GetFrequencyForGfxStage((ShaderStage::EStage)Idx);
+			FVulkanShader* const * FoundShaderPtr = ShaderFactoryHotfix.ShaderMap[ShaderFrequency].Find(ShaderKey);
+			if (FoundShaderPtr)
+			{
+				OutShaders[Idx] = *FoundShaderPtr;
+			}
+		}
+	}
+}
 
 FVulkanShaderFactory::~FVulkanShaderFactory()
 {
@@ -54,44 +92,25 @@ ShaderType* FVulkanShaderFactory::CreateShader(const TArray<uint8>& Code, FVulka
 	uint32 ShaderCodeCRC = FCrc::MemCrc32(Code.GetData(), Code.Num());
 	uint64 ShaderKey = ((uint64)ShaderCodeLen | ((uint64)ShaderCodeCRC << 32));
 
-	ShaderType* RetShader = nullptr;
-	if (GCacheCreatedShaders)
-	{
-		FVulkanShader** FoundShaderPtr = nullptr;
-		{
-			FRWScopeLock ScopedLock(Lock, SLT_ReadOnly);
-			FoundShaderPtr = ShaderMap[ShaderType::StaticFrequency].Find(ShaderKey);
-		}
-		if (FoundShaderPtr)
-		{
-			RetShader = static_cast<ShaderType*>(*FoundShaderPtr);
-		}
-		else
-		{
-			RetShader = new ShaderType(Device);
-			RetShader->Setup(Code, ShaderKey);
-			
-			FRWScopeLock ScopedLock(Lock, SLT_Write);
-			ShaderMap[ShaderType::StaticFrequency].Add(ShaderKey, RetShader);
-		}
-	}
-	else
+	EShaderFrequency ShaderFrequency = (EShaderFrequency)ShaderType::StaticFrequency;
+
+	ShaderType* RetShader = static_cast<ShaderType*>(FVulkanShaderFactory_LookupShader(*this, ShaderFrequency, ShaderKey));
+	if (RetShader == nullptr)
 	{
 		RetShader = new ShaderType(Device);
 		RetShader->Setup(Code, ShaderKey);
+			
+		FRWScopeLock ScopedLock(Lock, SLT_Write);
+		ShaderMap[ShaderFrequency].Add(ShaderKey, RetShader);
 	}
-
 	return RetShader;
 }
-	
+
 void FVulkanShaderFactory::OnDeleteShader(const FVulkanShader& Shader)
 {
-	if (GCacheCreatedShaders)
-	{
-		FRWScopeLock ScopedLock(Lock, SLT_Write);
-		uint64 ShaderKey = Shader.GetShaderKey(); 
-		ShaderMap[Shader.Frequency].Remove(ShaderKey);
-	}
+	FRWScopeLock ScopedLock(Lock, SLT_Write);
+	uint64 ShaderKey = Shader.GetShaderKey(); 
+	ShaderMap[Shader.Frequency].Remove(ShaderKey);
 }
 
 void FVulkanShader::Setup(const TArray<uint8>& InShaderHeaderAndCode, uint64 InShaderKey)
@@ -539,7 +558,6 @@ void FVulkanDescriptorSetsLayoutInfo::FinalizeBindings(const FUniformBufferGathe
 				}
 			}
 		}
-
 	}
 
 	CompileTypesUsageID();
diff --git a/Engine/Source/Runtime/VulkanRHI/Private/VulkanUtil.cpp b/Engine/Source/Runtime/VulkanRHI/Private/VulkanUtil.cpp
--- a/Engine/Source/Runtime/VulkanRHI/Private/VulkanUtil.cpp
+++ b/Engine/Source/Runtime/VulkanRHI/Private/VulkanUtil.cpp
@@ -551,7 +551,6 @@ void FVulkanGPUProfiler::DumpCrashMarkers(void* BufferData)
 				const FString* Frame = CachedStrings.Find(Value);
 				UE_LOG(LogVulkanRHI, Error, TEXT("[VK_NV_device_diagnostic_checkpoints] %i: Stage 0x%x, %s (CRC 0x%x)"), Index, Data[Index].stage, Frame ? *(*Frame) : TEXT("<undefined>"), Value);
 			}
-			GLog->FlushThreadedLogs();
 			GLog->PanicFlushThreadedLogs();
 			GLog->Flush();
 		}
