DR - Copy from 4634420
- Fix atmospheric fog/read 3d surface float on Vulkan
#rnx

diff --git a/Engine/Source/Runtime/VulkanRHI/Private/VulkanRenderTarget.cpp b/Engine/Source/Runtime/VulkanRHI/Private/VulkanRenderTarget.cpp
--- a/Engine/Source/Runtime/VulkanRHI/Private/VulkanRenderTarget.cpp
+++ b/Engine/Source/Runtime/VulkanRHI/Private/VulkanRenderTarget.cpp
@@ -769,6 +769,7 @@ void FVulkanDynamicRHI::RHIReadSurfaceData(FTextureRHIParamRef TextureRHI, FIntR
 
 	if (GIgnoreCPUReads == 2)
 	{
+		// Debug: Fill with CPU
 		OutData.Empty(0);
 		OutData.AddZeroed(NumPixels);
 		return;
@@ -1090,14 +1091,6 @@ void FVulkanDynamicRHI::RHIReadSurfaceFloatData(FTextureRHIParamRef TextureRHI,
 		InDevice->WaitUntilIdle();
 
 		StagingBuffer->InvalidateMappedMemory();
-/*
-		VkMappedMemoryRange MappedRange;
-		ZeroVulkanStruct(MappedRange, VK_STRUCTURE_TYPE_MAPPED_MEMORY_RANGE);
-		MappedRange.memory = StagingBuffer->GetDeviceMemoryHandle();
-		MappedRange.offset = StagingBuffer->GetAllocationOffset();
-		MappedRange.size = Size;
-		VulkanRHI::vkInvalidateMappedMemoryRanges(InDevice->GetInstanceHandle(), 1, &MappedRange);
-*/
 
 		OutputData.SetNum(NumPixels);
 		FFloat16Color* Dest = OutputData.GetData();
@@ -1114,7 +1107,7 @@ void FVulkanDynamicRHI::RHIReadSurfaceFloatData(FTextureRHIParamRef TextureRHI,
 
 	if (GIgnoreCPUReads == 2)
 	{
-		// FIll with CPU
+		// Debug: Fill with CPU
 		uint32 NumPixels = 0;
 		if (TextureRHI->GetTextureCube())
 		{
@@ -1157,12 +1150,106 @@ void FVulkanDynamicRHI::RHIReadSurfaceFloatData(FTextureRHIParamRef TextureRHI,
 
 void FVulkanDynamicRHI::RHIRead3DSurfaceFloatData(FTextureRHIParamRef TextureRHI,FIntRect InRect,FIntPoint ZMinMax,TArray<FFloat16Color>& OutData)
 {
+	FRHITexture3D* TextureRHI3D = TextureRHI->GetTexture3D();
+	check(TextureRHI3D);
+	FVulkanTexture3D* Texture3D = (FVulkanTexture3D*)TextureRHI3D;
+	FVulkanSurface& Surface = Texture3D->Surface;
+
+	uint32 SizeX = InRect.Width();
+	uint32 SizeY = InRect.Height();
+	uint32 SizeZ = ZMinMax.Y - ZMinMax.X;
+	uint32 NumPixels = SizeX * SizeY * SizeZ;
+	const uint32 Size = NumPixels * sizeof(FFloat16Color);
+
+	// Allocate the output buffer.
+	OutData.Reserve(Size);
+	if (GIgnoreCPUReads == 2)
+	{
+		OutData.AddZeroed(Size);
+
+		// Debug: Fill with CPU
+		return;
+	}
+
 	Device->PrepareForCPURead();
+	FVulkanCmdBuffer* CmdBuffer = Device->GetImmediateContext().GetCommandBufferManager()->GetUploadCmdBuffer();
 
-	VULKAN_SIGNAL_UNIMPLEMENTED();
+	ensure(Surface.StorageFormat == VK_FORMAT_R16G16B16A16_SFLOAT);
 
-	Device->GetImmediateContext().GetCommandBufferManager()->PrepareForNewActiveCommandBuffer();
+	VulkanRHI::FStagingBuffer* StagingBuffer = Device->GetStagingManager().AcquireBuffer(Size, VK_BUFFER_USAGE_TRANSFER_DST_BIT, true);
+	if (GIgnoreCPUReads == 0)
+	{
+		VkBufferImageCopy CopyRegion;
+		FMemory::Memzero(CopyRegion);
+		//Region.bufferOffset = 0;
+		CopyRegion.bufferRowLength = Surface.Width;
+		CopyRegion.bufferImageHeight = Surface.Height;
+		CopyRegion.imageSubresource.aspectMask = Surface.GetFullAspectMask();
+		//CopyRegion.imageSubresource.mipLevel = 0;
+		//CopyRegion.imageSubresource.baseArrayLayer = 0;
+		CopyRegion.imageSubresource.layerCount = 1;
+		CopyRegion.imageOffset.x = InRect.Min.X;
+		CopyRegion.imageOffset.y = InRect.Min.Y;
+		CopyRegion.imageOffset.z = ZMinMax.X;
+		CopyRegion.imageExtent.width = SizeX;
+		CopyRegion.imageExtent.height = SizeY;
+		CopyRegion.imageExtent.depth = SizeZ;
+
+		//#todo-rco: Multithreaded!
+		VkImageLayout& CurrentLayout = Device->GetImmediateContext().FindOrAddLayoutRW(Surface.Image, VK_IMAGE_LAYOUT_UNDEFINED);
+		bool bHadLayout = (CurrentLayout != VK_IMAGE_LAYOUT_UNDEFINED);
+		if (CurrentLayout != VK_IMAGE_LAYOUT_TRANSFER_SRC_OPTIMAL)
+		{
+			VulkanSetImageLayoutSimple(CmdBuffer->GetHandle(), Surface.Image, CurrentLayout, VK_IMAGE_LAYOUT_TRANSFER_SRC_OPTIMAL);
+		}
 
+		VulkanRHI::vkCmdCopyImageToBuffer(CmdBuffer->GetHandle(), Surface.Image, VK_IMAGE_LAYOUT_TRANSFER_SRC_OPTIMAL, StagingBuffer->GetHandle(), 1, &CopyRegion);
+
+		if (bHadLayout && CurrentLayout != VK_IMAGE_LAYOUT_TRANSFER_SRC_OPTIMAL)
+		{
+			VulkanSetImageLayoutSimple(CmdBuffer->GetHandle(), Surface.Image, VK_IMAGE_LAYOUT_TRANSFER_SRC_OPTIMAL, CurrentLayout);
+		}
+		else
+		{
+			CurrentLayout = VK_IMAGE_LAYOUT_TRANSFER_SRC_OPTIMAL;
+		}
+	}
+	else
+	{
+		VulkanRHI::vkCmdFillBuffer(CmdBuffer->GetHandle(), StagingBuffer->GetHandle(), 0, Size, (FFloat16(1.0).Encoded << 16) + FFloat16(1.0).Encoded);
+	}
+
+	VkBufferMemoryBarrier Barrier;
+	// the staging buffer size may be bigger then the size due to alignment, etc. but it must not be smaller!
+	ensure(StagingBuffer->GetSize() >= Size);
+	//#todo-rco: Change offset if reusing a buffer suballocation
+	VulkanRHI::SetupAndZeroBufferBarrier(Barrier, VK_ACCESS_TRANSFER_WRITE_BIT, VK_ACCESS_HOST_READ_BIT, StagingBuffer->GetHandle(), 0/*StagingBuffer->GetOffset()*/, StagingBuffer->GetSize());
+	VulkanRHI::vkCmdPipelineBarrier(CmdBuffer->GetHandle(), VK_PIPELINE_STAGE_TRANSFER_BIT, VK_PIPELINE_STAGE_HOST_BIT, 0, 0, nullptr, 1, &Barrier, 0, nullptr);
+
+	// Force upload
+	Device->GetImmediateContext().GetCommandBufferManager()->SubmitUploadCmdBuffer();
+	Device->WaitUntilIdle();
+
+	StagingBuffer->InvalidateMappedMemory();
+
+	OutData.SetNum(NumPixels);
+	FFloat16Color* Dest = OutData.GetData();
+	for (int32 Layer = ZMinMax.X; Layer < ZMinMax.Y; ++Layer)
+	{
+		for (int32 Row = InRect.Min.Y; Row < InRect.Max.Y; ++Row)
+		{
+			FFloat16Color* Src = (FFloat16Color*)StagingBuffer->GetMappedPointer() + Layer * SizeX * SizeY + Row * Surface.Width + InRect.Min.X;
+			for (int32 Col = InRect.Min.X; Col < InRect.Max.X; ++Col)
+			{
+				*Dest++ = *Src++;
+			}
+		}
+	}
+	FFloat16Color* End = OutData.GetData() + OutData.Num();
+	checkf(Dest <= End, TEXT("Memory overwrite! Calculated total size %d: SizeX %d SizeY %d SizeZ %d; InRect(%d, %d, %d, %d) InZ(%d, %d)"),
+		Size, SizeX, SizeY, SizeZ, InRect.Min.X, InRect.Min.Y, InRect.Max.X, InRect.Max.Y, ZMinMax.X, ZMinMax.Y);
+	Device->GetStagingManager().ReleaseBuffer(CmdBuffer, StagingBuffer);
+	Device->GetImmediateContext().GetCommandBufferManager()->PrepareForNewActiveCommandBuffer();
 }
 
 void FVulkanCommandListContext::RHITransitionResources(EResourceTransitionAccess TransitionType, EResourceTransitionPipeline TransitionPipeline, FUnorderedAccessViewRHIParamRef* InUAVs, int32 NumUAVs, FComputeFenceRHIParamRef WriteComputeFenceRHI)
