DR - vk- Copy from 4640101
PVS warning silence
#rnx

diff --git a/Engine/Source/Runtime/VulkanRHI/Private/VulkanQuery.cpp b/Engine/Source/Runtime/VulkanRHI/Private/VulkanQuery.cpp
--- a/Engine/Source/Runtime/VulkanRHI/Private/VulkanQuery.cpp
+++ b/Engine/Source/Runtime/VulkanRHI/Private/VulkanQuery.cpp
@@ -645,7 +645,7 @@ void FVulkanDevice::ReleaseUnusedOcclusionQueryPools()
 	for (int32 Index = FreeOcclusionQueryPools.Num() - 1; Index >= 0; --Index)
 	{
 		FVulkanOcclusionQueryPool* Pool = FreeOcclusionQueryPools[Index];
-		if (ReleaseFrame > Pool->FreedFrameNumber)
+		if (ReleaseFrame > Pool->FreedFrameNumber) //-V547
 		{
 			delete Pool;
 			FreeOcclusionQueryPools.RemoveAt(Index);
