DR - vk - Fix # timestamps
#rnx

diff --git a/Engine/Source/Runtime/VulkanRHI/Public/VulkanResources.h b/Engine/Source/Runtime/VulkanRHI/Public/VulkanResources.h
--- a/Engine/Source/Runtime/VulkanRHI/Public/VulkanResources.h
+++ b/Engine/Source/Runtime/VulkanRHI/Public/VulkanResources.h
@@ -38,7 +38,7 @@ enum
 {
 	NUM_OCCLUSION_QUERIES_PER_POOL = 4096,
 
-	NUM_TIMESTAMP_QUERIES_PER_POOL = 10240,
+	NUM_TIMESTAMP_QUERIES_PER_POOL = 1024,
 };
 
 struct FSamplerYcbcrConversionInitializer
