DR - vk - Remove unused extern
#rnx

diff --git a/Engine/Source/Runtime/VulkanRHI/Private/VulkanPipelineState.cpp b/Engine/Source/Runtime/VulkanRHI/Private/VulkanPipelineState.cpp
--- a/Engine/Source/Runtime/VulkanRHI/Private/VulkanPipelineState.cpp
+++ b/Engine/Source/Runtime/VulkanRHI/Private/VulkanPipelineState.cpp
@@ -195,7 +195,6 @@ bool FVulkanComputePipelineDescriptorState::InternalUpdateDescriptorSets(FVulkan
 	return true;
 }
 
-extern FVulkanShader* FVulkanShaderFactory_LookupShader(FVulkanShaderFactory& ShaderFactory, EShaderFrequency ShaderFrequency, uint64 ShaderKey);
 
 FVulkanGraphicsPipelineDescriptorState::FVulkanGraphicsPipelineDescriptorState(FVulkanDevice* InDevice, FVulkanRHIGraphicsPipelineState* InGfxPipeline)
 	: FVulkanCommonPipelineDescriptorState(InDevice)
@@ -215,7 +214,6 @@ FVulkanGraphicsPipelineDescriptorState::FVulkanGraphicsPipelineDescriptorState(F
 
 	const FVulkanVertexShader* VertexShader = ShaderFactory.LookupShader<FVulkanVertexShader>(InGfxPipeline->GetShaderKey(SF_Vertex));
 	check(VertexShader);
-
 	PackedUniformBuffers[ShaderStage::Vertex].Init(VertexShader->GetCodeHeader(), PackedUniformBuffersMask[ShaderStage::Vertex]);
 
 	uint64 PixelShaderKey = InGfxPipeline->GetShaderKey(SF_Pixel);
