DR - vk - Fix for validation warning on Niagara
#rnx

diff --git a/Engine/Source/Runtime/VulkanRHI/Private/VulkanIndexBuffer.cpp b/Engine/Source/Runtime/VulkanRHI/Private/VulkanIndexBuffer.cpp
--- a/Engine/Source/Runtime/VulkanRHI/Private/VulkanIndexBuffer.cpp
+++ b/Engine/Source/Runtime/VulkanRHI/Private/VulkanIndexBuffer.cpp
@@ -76,11 +76,13 @@ FVulkanResourceMultiBuffer::FVulkanResourceMultiBuffer(FVulkanDevice* InDevice,
 		const bool bIsUniformBuffer = (InBufferUsageFlags & VK_BUFFER_USAGE_UNIFORM_BUFFER_BIT) != 0;
 		const bool bUAV = (InUEUsage & BUF_UnorderedAccess) != 0;
 		const bool bIndirect = (InUEUsage & BUF_DrawIndirect) == BUF_DrawIndirect;
+		const bool bCPUReadable = (UEUsage & BUF_KeepCPUAccessible) != 0;
 
 		BufferUsageFlags |= bVolatile ? 0 : VK_BUFFER_USAGE_TRANSFER_DST_BIT;
 		BufferUsageFlags |= (bShaderResource && !bIsUniformBuffer) ? VK_BUFFER_USAGE_UNIFORM_TEXEL_BUFFER_BIT : 0;
 		BufferUsageFlags |= bUAV ? VK_BUFFER_USAGE_STORAGE_TEXEL_BUFFER_BIT : 0;
 		BufferUsageFlags |= bIndirect ? VK_BUFFER_USAGE_INDIRECT_BUFFER_BIT : 0;
+		BufferUsageFlags |= bCPUReadable ? VK_BUFFER_USAGE_TRANSFER_SRC_BIT : 0;
 
 		if (bVolatile)
 		{
