DR - vk - Undo 4571656 (fix for render passes on mip generation)

diff --git a/Engine/Source/Runtime/VulkanRHI/Private/VulkanRenderTarget.cpp b/Engine/Source/Runtime/VulkanRHI/Private/VulkanRenderTarget.cpp
--- a/Engine/Source/Runtime/VulkanRHI/Private/VulkanRenderTarget.cpp
+++ b/Engine/Source/Runtime/VulkanRHI/Private/VulkanRenderTarget.cpp
@@ -420,6 +420,8 @@ void FTransitionAndLayoutManager::EndRealRenderPass(FVulkanCmdBuffer* CmdBuffer)
 				GenerateMipsInfo.Target[Index].Layouts[GenerateMipsInfo.CurrentSlice][GenerateMipsInfo.CurrentMip] = VK_IMAGE_LAYOUT_SHADER_READ_ONLY_OPTIMAL;
 			}
 			Barrier.Execute(CmdBuffer);
+
+			GenerateMipsInfo.Reset();
 		}
 	}
 
@@ -1725,8 +1727,6 @@ void FVulkanCommandListContext::RHIBeginRenderPass(const FRHIRenderPassInfo& InI
 		}
 	}
 
-	TransitionAndLayoutManager.GenerateMipsInfo.Reset();
-
 	FVulkanFramebuffer* Framebuffer = TransitionAndLayoutManager.GetOrCreateFramebuffer(*Device, RTInfo, RTLayout, RenderPass);
 	checkf(RenderPass != nullptr && Framebuffer != nullptr, TEXT("RenderPass not started! Bad combination of values? Depth %p #Color %d Color0 %p"), (void*)InInfo.DepthStencilRenderTarget.DepthStencilTarget, InInfo.GetNumColorRenderTargets(), (void*)InInfo.ColorRenderTargets[0].RenderTarget);
 	TransitionAndLayoutManager.BeginRealRenderPass(*this, *Device, CmdBuffer, InInfo, RTLayout, RenderPass, Framebuffer);
