DR - vk - During mip gen, transition to readable at end render pass for correctness
#rnx

diff --git a/Engine/Source/Runtime/VulkanRHI/Private/VulkanRenderTarget.cpp b/Engine/Source/Runtime/VulkanRHI/Private/VulkanRenderTarget.cpp
--- a/Engine/Source/Runtime/VulkanRHI/Private/VulkanRenderTarget.cpp
+++ b/Engine/Source/Runtime/VulkanRHI/Private/VulkanRenderTarget.cpp
@@ -289,20 +289,9 @@ void FTransitionAndLayoutManager::BeginRealRenderPass(FVulkanCommandListContext&
 			GenerateMipsInfo.CurrentSlice = SliceIndex;
 			GenerateMipsInfo.CurrentMip = RTMipIndex;
 			GenerateMipsInfo.bLastMip = (RTMipIndex == (NumMips - 1));
-			if (GenerateMipsInfo.Target[Index].Layouts[SliceIndex][RTMipIndex - 1] != VK_IMAGE_LAYOUT_SHADER_READ_ONLY_OPTIMAL)
-			{
-				// Transition to readable
-				int32 BarrierIndex = Barrier.AddImageBarrier(Surface.Image, VK_IMAGE_ASPECT_COLOR_BIT, 1);
-				VkImageSubresourceRange& Range = Barrier.GetSubresource(BarrierIndex);
-				Range.baseMipLevel = RTMipIndex - 1;
-				Range.baseArrayLayer = SliceIndex;
 #if !USING_CODE_ANALYSIS
-				ensure(GenerateMipsInfo.Target[Index].Layouts[SliceIndex][RTMipIndex - 1] == VK_IMAGE_LAYOUT_COLOR_ATTACHMENT_OPTIMAL);
+			ensure(GenerateMipsInfo.Target[Index].Layouts[SliceIndex][RTMipIndex - 1] == VK_IMAGE_LAYOUT_SHADER_READ_ONLY_OPTIMAL);
 #endif
-				Barrier.SetTransition(BarrierIndex, EImageLayoutBarrier::ColorAttachment, EImageLayoutBarrier::PixelShaderRead);
-				GenerateMipsInfo.Target[Index].Layouts[SliceIndex][RTMipIndex - 1] = VK_IMAGE_LAYOUT_SHADER_READ_ONLY_OPTIMAL;
-			}
-
 			if (GenerateMipsInfo.Target[Index].Layouts[SliceIndex][RTMipIndex] != VK_IMAGE_LAYOUT_COLOR_ATTACHMENT_OPTIMAL)
 			{
 				// Transition to writeable
@@ -403,24 +392,24 @@ void FTransitionAndLayoutManager::EndRealRenderPass(FVulkanCmdBuffer* CmdBuffer)
 
 	if (GenerateMipsInfo.bInsideGenerateMips)
 	{
-		if (GenerateMipsInfo.bLastMip)
+		FPendingBarrier Barrier;
+		for (int32 Index = 0; Index < GenerateMipsInfo.NumRenderTargets; ++Index)
 		{
-			FPendingBarrier Barrier;
-			for (int32 Index = 0; Index < GenerateMipsInfo.NumRenderTargets; ++Index)
-			{
-				ensure(GenerateMipsInfo.Target[Index].Layouts[GenerateMipsInfo.CurrentSlice][GenerateMipsInfo.CurrentMip] == VK_IMAGE_LAYOUT_COLOR_ATTACHMENT_OPTIMAL);
+			ensure(GenerateMipsInfo.Target[Index].Layouts[GenerateMipsInfo.CurrentSlice][GenerateMipsInfo.CurrentMip] == VK_IMAGE_LAYOUT_COLOR_ATTACHMENT_OPTIMAL);
 
-				// Transition to readable
-				int32 BarrierIndex = Barrier.AddImageBarrier(GenerateMipsInfo.Target[Index].CurrentImage, VK_IMAGE_ASPECT_COLOR_BIT, 1);
-				VkImageSubresourceRange& Range = Barrier.GetSubresource(BarrierIndex);
-				Range.baseMipLevel = GenerateMipsInfo.CurrentMip;
-				Range.baseArrayLayer = GenerateMipsInfo.CurrentSlice;
-				Barrier.SetTransition(BarrierIndex, EImageLayoutBarrier::ColorAttachment, EImageLayoutBarrier::PixelShaderRead);
-				// This could really be ignored...
-				GenerateMipsInfo.Target[Index].Layouts[GenerateMipsInfo.CurrentSlice][GenerateMipsInfo.CurrentMip] = VK_IMAGE_LAYOUT_SHADER_READ_ONLY_OPTIMAL;
-			}
-			Barrier.Execute(CmdBuffer);
+			// Transition to readable
+			int32 BarrierIndex = Barrier.AddImageBarrier(GenerateMipsInfo.Target[Index].CurrentImage, VK_IMAGE_ASPECT_COLOR_BIT, 1);
+			VkImageSubresourceRange& Range = Barrier.GetSubresource(BarrierIndex);
+			Range.baseMipLevel = GenerateMipsInfo.CurrentMip;
+			Range.baseArrayLayer = GenerateMipsInfo.CurrentSlice;
+			Barrier.SetTransition(BarrierIndex, EImageLayoutBarrier::ColorAttachment, EImageLayoutBarrier::PixelShaderRead);
+			// This could really be ignored...
+			GenerateMipsInfo.Target[Index].Layouts[GenerateMipsInfo.CurrentSlice][GenerateMipsInfo.CurrentMip] = VK_IMAGE_LAYOUT_SHADER_READ_ONLY_OPTIMAL;
+		}
+		Barrier.Execute(CmdBuffer);
 
+		if (GenerateMipsInfo.bLastMip)
+		{
 			GenerateMipsInfo.Reset();
 		}
 	}
