DR - vk - Another fix for validation warning on Niagara
#rnx

diff --git a/Engine/Source/Runtime/VulkanRHI/Private/VulkanIndexBuffer.cpp b/Engine/Source/Runtime/VulkanRHI/Private/VulkanIndexBuffer.cpp
--- a/Engine/Source/Runtime/VulkanRHI/Private/VulkanIndexBuffer.cpp
+++ b/Engine/Source/Runtime/VulkanRHI/Private/VulkanIndexBuffer.cpp
@@ -82,7 +82,7 @@ FVulkanResourceMultiBuffer::FVulkanResourceMultiBuffer(FVulkanDevice* InDevice,
 		BufferUsageFlags |= (bShaderResource && !bIsUniformBuffer) ? VK_BUFFER_USAGE_UNIFORM_TEXEL_BUFFER_BIT : 0;
 		BufferUsageFlags |= bUAV ? VK_BUFFER_USAGE_STORAGE_TEXEL_BUFFER_BIT : 0;
 		BufferUsageFlags |= bIndirect ? VK_BUFFER_USAGE_INDIRECT_BUFFER_BIT : 0;
-		BufferUsageFlags |= bCPUReadable ? VK_BUFFER_USAGE_TRANSFER_SRC_BIT : 0;
+		BufferUsageFlags |= bCPUReadable ? (VK_BUFFER_USAGE_TRANSFER_SRC_BIT | VK_BUFFER_USAGE_TRANSFER_DST_BIT) : 0;
 
 		if (bVolatile)
 		{
