DR - Fix android compilation with no PCH
#rnx

diff --git a/Engine/Source/Runtime/VulkanRHI/Private/Android/VulkanAndroidPlatform.h b/Engine/Source/Runtime/VulkanRHI/Private/Android/VulkanAndroidPlatform.h
--- a/Engine/Source/Runtime/VulkanRHI/Private/Android/VulkanAndroidPlatform.h
+++ b/Engine/Source/Runtime/VulkanRHI/Private/Android/VulkanAndroidPlatform.h
@@ -2,6 +2,8 @@
 
 #pragma once
 
+#include "RHI.h"
+
 #define VK_USE_PLATFORM_ANDROID_KHR					1
 
 #define VULKAN_HAS_PHYSICAL_DEVICE_PROPERTIES2		0
