DR - vk - Fix assert when there are no crash markers
- OOM is now Fatal
#rnx

diff --git a/Engine/Source/Runtime/VulkanRHI/Private/VulkanMemory.cpp b/Engine/Source/Runtime/VulkanRHI/Private/VulkanMemory.cpp
--- a/Engine/Source/Runtime/VulkanRHI/Private/VulkanMemory.cpp
+++ b/Engine/Source/Runtime/VulkanRHI/Private/VulkanMemory.cpp
@@ -237,11 +237,15 @@ namespace VulkanRHI
 				UE_LOG(LogVulkanRHI, Warning, TEXT("Failed to allocate Device Memory, Requested=%fKb MemTypeIndex=%d"), (float)Info.allocationSize / 1024.0f, Info.memoryTypeIndex);
 				return nullptr;
 			}
-			UE_LOG(LogVulkanRHI, Error, TEXT("Out of Device Memory, Requested=%fKb MemTypeIndex=%d"), (float)Info.allocationSize / 1024.0f, Info.memoryTypeIndex);
+			auto Callback = [=]()
+			{
 #if UE_BUILD_DEBUG || UE_BUILD_DEVELOPMENT
-			DumpMemory();
-			GLog->PanicFlushThreadedLogs();
+				DumpMemory();
+				GLog->PanicFlushThreadedLogs();
 #endif
+				return FString::Printf(TEXT("Out of Device Memory, Requested=%fKb MemTypeIndex=%d"), (float)Info.allocationSize / 1024.0f, Info.memoryTypeIndex);
+			};
+			UE_LOG(LogVulkanRHI, Fatal, TEXT("%s"), *(Callback()));
 		}
 		else if (Result == VK_ERROR_OUT_OF_HOST_MEMORY)
 		{
@@ -250,11 +254,15 @@ namespace VulkanRHI
 				UE_LOG(LogVulkanRHI, Warning, TEXT("Failed to allocate Host Memory, Requested=%fKb MemTypeIndex=%d"), (float)Info.allocationSize / 1024.0f, Info.memoryTypeIndex);
 				return nullptr;
 			}
-			UE_LOG(LogVulkanRHI, Error, TEXT("Out of Host Memory, Requested=%fKb MemTypeIndex=%d"), (float)Info.allocationSize / 1024.0f, Info.memoryTypeIndex);
+			auto Callback = [=]()
+			{
 #if UE_BUILD_DEBUG || UE_BUILD_DEVELOPMENT
-			DumpMemory();
-			GLog->PanicFlushThreadedLogs();
+				DumpMemory();
+				GLog->PanicFlushThreadedLogs();
 #endif
+				return FString::Printf(TEXT("Out of Host Memory, Requested=%fKb MemTypeIndex=%d"), (float)Info.allocationSize / 1024.0f, Info.memoryTypeIndex);
+			};
+			UE_LOG(LogVulkanRHI, Error, TEXT("%s"), *(Callback()));
 		}
 		else
 		{
diff --git a/Engine/Source/Runtime/VulkanRHI/Private/VulkanUtil.cpp b/Engine/Source/Runtime/VulkanRHI/Private/VulkanUtil.cpp
--- a/Engine/Source/Runtime/VulkanRHI/Private/VulkanUtil.cpp
+++ b/Engine/Source/Runtime/VulkanRHI/Private/VulkanUtil.cpp
@@ -541,18 +541,21 @@ void FVulkanGPUProfiler::DumpCrashMarkers(void* BufferData)
 			uint32 Num = 0;
 			VkQueue QueueHandle = Device->GetGraphicsQueue()->GetHandle();
 			VulkanDynamicAPI::vkGetQueueCheckpointDataNV(QueueHandle, &Num, nullptr);
-			Data.AddUninitialized(Num);
-			VulkanDynamicAPI::vkGetQueueCheckpointDataNV(QueueHandle, &Num, &Data[0]);
-			check(Num == Data.Num());
-			for (uint32 Index = 0; Index < Num; ++Index)
+			if (Num > 0)
 			{
-				check(Data[Index].sType == VK_STRUCTURE_TYPE_CHECKPOINT_DATA_NV);
-				uint32 Value = (uint32)(size_t)Data[Index].pCheckpointMarker;
-				const FString* Frame = CachedStrings.Find(Value);
-				UE_LOG(LogVulkanRHI, Error, TEXT("[VK_NV_device_diagnostic_checkpoints] %i: Stage 0x%x, %s (CRC 0x%x)"), Index, Data[Index].stage, Frame ? *(*Frame) : TEXT("<undefined>"), Value);
+				Data.AddUninitialized(Num);
+				VulkanDynamicAPI::vkGetQueueCheckpointDataNV(QueueHandle, &Num, &Data[0]);
+				check(Num == Data.Num());
+				for (uint32 Index = 0; Index < Num; ++Index)
+				{
+					check(Data[Index].sType == VK_STRUCTURE_TYPE_CHECKPOINT_DATA_NV);
+					uint32 Value = (uint32)(size_t)Data[Index].pCheckpointMarker;
+					const FString* Frame = CachedStrings.Find(Value);
+					UE_LOG(LogVulkanRHI, Error, TEXT("[VK_NV_device_diagnostic_checkpoints] %i: Stage 0x%x, %s (CRC 0x%x)"), Index, Data[Index].stage, Frame ? *(*Frame) : TEXT("<undefined>"), Value);
+				}
+				GLog->PanicFlushThreadedLogs();
+				GLog->Flush();
 			}
-			GLog->PanicFlushThreadedLogs();
-			GLog->Flush();
 		}
 #endif
 	}
