DR - vk - Copy from 4667288
- More information for Vulkan issues and better error message.
#rnx

diff --git a/Engine/Source/Runtime/VulkanRHI/Private/VulkanDevice.cpp b/Engine/Source/Runtime/VulkanRHI/Private/VulkanDevice.cpp
--- a/Engine/Source/Runtime/VulkanRHI/Private/VulkanDevice.cpp
+++ b/Engine/Source/Runtime/VulkanRHI/Private/VulkanDevice.cpp
@@ -281,7 +281,13 @@ void FVulkanDevice::CreateDevice()
 	FVulkanPlatform::EnablePhysicalDeviceFeatureExtensions(DeviceInfo);
 
 	// Create the device
-	VERIFYVULKANRESULT(VulkanRHI::vkCreateDevice(Gpu, &DeviceInfo, VULKAN_CPU_ALLOCATOR, &Device));
+	VkResult Result = VulkanRHI::vkCreateDevice(Gpu, &DeviceInfo, VULKAN_CPU_ALLOCATOR, &Device);
+	if (Result == VK_ERROR_INITIALIZATION_FAILED)
+	{
+		FPlatformMisc::MessageBoxExt(EAppMsgType::Ok, TEXT("Cannot create a Vulkan device. Try updating your video driver to a more recent version.\n"), TEXT("Vulkan device creation failed"));
+		FPlatformMisc::RequestExitWithStatus(true, 1);
+	}
+	VERIFYVULKANRESULT_EXPANDED(Result);
 
 	// Create Graphics Queue, here we submit command buffers for execution
 	GfxQueue = new FVulkanQueue(this, GfxQueueFamilyIndex);
@@ -718,7 +724,7 @@ bool FVulkanDevice::QueryGPU(int32 DeviceIndex)
 	};
 
 	UE_LOG(LogVulkanRHI, Display, TEXT("Device %d: %s"), DeviceIndex, ANSI_TO_TCHAR(GpuProps.deviceName));
-	UE_LOG(LogVulkanRHI, Display, TEXT("- API 0x%x Driver 0x%x VendorId 0x%x"), GpuProps.apiVersion, GpuProps.driverVersion, GpuProps.vendorID);
+	UE_LOG(LogVulkanRHI, Display, TEXT("- API %d.%d.%d(0x%x) Driver 0x%x VendorId 0x%x"), VK_VERSION_MAJOR(GpuProps.apiVersion), VK_VERSION_MINOR(GpuProps.apiVersion), VK_VERSION_PATCH(GpuProps.apiVersion), GpuProps.apiVersion, GpuProps.driverVersion, GpuProps.vendorID);
 	UE_LOG(LogVulkanRHI, Display, TEXT("- DeviceID 0x%x Type %s"), GpuProps.deviceID, *GetDeviceTypeString());
 	UE_LOG(LogVulkanRHI, Display, TEXT("- Max Descriptor Sets Bound %d Timestamps %d"), GpuProps.limits.maxBoundDescriptorSets, GpuProps.limits.timestampComputeAndGraphics);
 
