DR - vk - Added missing UAV & SRV for texture arrays

diff --git a/Engine/Source/Runtime/VulkanRHI/Private/VulkanUAV.cpp b/Engine/Source/Runtime/VulkanRHI/Private/VulkanUAV.cpp
--- a/Engine/Source/Runtime/VulkanRHI/Private/VulkanUAV.cpp
+++ b/Engine/Source/Runtime/VulkanRHI/Private/VulkanUAV.cpp
@@ -80,6 +80,16 @@ void FVulkanShaderResourceView::UpdateView()
 				FVulkanTextureCube* VTexCube = ResourceCast(TexCube);
 				TextureView.Create(*Device, VTexCube->Surface.Image, VK_IMAGE_VIEW_TYPE_CUBE, VTexCube->Surface.GetPartialAspectMask(), Format, UEToVkFormat(Format, false), MipLevel, NumMips, 0, 1);
 			}
+			else if (FRHITexture3D* Tex3D = SourceTexture->GetTexture3D())
+			{
+				FVulkanTexture3D* VTex3D = ResourceCast(Tex3D);
+				TextureView.Create(*Device, VTex3D->Surface.Image, VK_IMAGE_VIEW_TYPE_3D, VTex3D->Surface.GetPartialAspectMask(), Format, UEToVkFormat(Format, false), MipLevel, NumMips, 0, 1);
+			}
+			else if (FRHITexture2DArray* Tex2DArray = SourceTexture->GetTexture2DArray())
+			{
+				FVulkanTexture2DArray* VTex2DArray = ResourceCast(Tex2DArray);
+				TextureView.Create(*Device, VTex2DArray->Surface.Image, VK_IMAGE_VIEW_TYPE_2D_ARRAY, VTex2DArray->Surface.GetPartialAspectMask(), Format, UEToVkFormat(Format, false), MipLevel, NumMips, 0, VTex2DArray->GetSizeZ());
+			}
 			else
 			{
 				ensure(0);
@@ -162,6 +172,11 @@ void FVulkanUnorderedAccessView::UpdateView()
 			FVulkanTexture3D* VTex3D = ResourceCast(Tex3D);
 			TextureView.Create(*Device, VTex3D->Surface.Image, VK_IMAGE_VIEW_TYPE_3D, VTex3D->Surface.GetPartialAspectMask(), Format, UEToVkFormat(Format, false), MipLevel, 1, 0, VTex3D->GetSizeZ());
 		}
+		else if (FRHITexture2DArray* Tex2DArray = SourceTexture->GetTexture2DArray())
+		{
+			FVulkanTexture2DArray* VTex2DArray = ResourceCast(Tex2DArray);
+			TextureView.Create(*Device, VTex2DArray->Surface.Image, VK_IMAGE_VIEW_TYPE_2D_ARRAY, VTex2DArray->Surface.GetPartialAspectMask(), Format, UEToVkFormat(Format, false), MipLevel, 1, 0, VTex2DArray->GetSizeZ());
+		}
 		else
 		{
 			ensure(0);
