Linux finally compiles

diff --git a/Engine/Plugins/FX/Niagara/Source/Niagara/Private/NiagaraRendererSprites.cpp b/Engine/Plugins/FX/Niagara/Source/Niagara/Private/NiagaraRendererSprites.cpp
--- a/Engine/Plugins/FX/Niagara/Source/Niagara/Private/NiagaraRendererSprites.cpp
+++ b/Engine/Plugins/FX/Niagara/Source/Niagara/Private/NiagaraRendererSprites.cpp
@@ -288,7 +288,7 @@ void NiagaraRendererSprites::GetDynamicMeshElements(const TArray<const FSceneVie
 				VFLooseParams.NiagaraFloatDataStride = CollectorResources.VertexFactory.GetFloatDataStride();
 				VFLooseParams.ParticleAlignmentMode = CollectorResources.VertexFactory.GetAlignmentMode();
 				VFLooseParams.ParticleFacingMode = CollectorResources.VertexFactory.GetFacingMode();
-				VFLooseParams.SortedIndices = CollectorResources.VertexFactory.GetSortedIndicesSRV() ? CollectorResources.VertexFactory.GetSortedIndicesSRV() : GFNiagaraNullSortedIndicesVertexBuffer.VertexBufferSRV.GetReference();
+				VFLooseParams.SortedIndices = CollectorResources.VertexFactory.GetSortedIndicesSRV() ? CollectorResources.VertexFactory.GetSortedIndicesSRV().GetReference() : GFNiagaraNullSortedIndicesVertexBuffer.VertexBufferSRV.GetReference();
 				VFLooseParams.SortedIndicesOffset = CollectorResources.VertexFactory.GetSortedIndicesOffset();
 
 				CollectorResources.VertexFactory.LooseParameterUniformBuffer = FNiagaraSpriteVFLooseParametersRef::CreateUniformBufferImmediate(VFLooseParams, UniformBuffer_SingleFrame);
diff --git a/Engine/Plugins/FX/Niagara/Source/NiagaraVertexFactories/Public/NiagaraSpriteVertexFactory.h b/Engine/Plugins/FX/Niagara/Source/NiagaraVertexFactories/Public/NiagaraSpriteVertexFactory.h
--- a/Engine/Plugins/FX/Niagara/Source/NiagaraVertexFactories/Public/NiagaraSpriteVertexFactory.h
+++ b/Engine/Plugins/FX/Niagara/Source/NiagaraVertexFactories/Public/NiagaraSpriteVertexFactory.h
@@ -112,19 +112,20 @@ public:
 	/** Default constructor. */
 	FNiagaraSpriteVertexFactory(ENiagaraVertexFactoryType InType, ERHIFeatureLevel::Type InFeatureLevel )
 		: FNiagaraVertexFactoryBase(InType, InFeatureLevel),
+		LooseParameterUniformBuffer(nullptr),
 		NumVertsInInstanceBuffer(0),
 		NumCutoutVerticesPerFrame(0),
 		CutoutGeometrySRV(nullptr),
 		AlignmentMode(0),
 		FacingMode(0),
 		FloatDataOffset(0),
-		FloatDataStride(0),
-		LooseParameterUniformBuffer(nullptr),
+		FloatDataStride(0),		
 		SortedIndicesOffset(0)
 	{}
 
 	FNiagaraSpriteVertexFactory()
 		: FNiagaraVertexFactoryBase(NVFT_MAX, ERHIFeatureLevel::Num),
+		LooseParameterUniformBuffer(nullptr),
 		NumVertsInInstanceBuffer(0),
 		NumCutoutVerticesPerFrame(0),
 		CutoutGeometrySRV(nullptr),
@@ -132,8 +133,8 @@ public:
 		FacingMode(0),
 		FloatDataOffset(0),
 		FloatDataStride(0),
-		SortedIndicesOffset(0),
-		LooseParameterUniformBuffer(nullptr)
+		SortedIndicesOffset(0)
+		
 	{}
 
 	// FRenderResource interface.
diff --git a/Engine/Source/Runtime/OpenGLDrv/Private/OpenGLUniformBuffer.cpp b/Engine/Source/Runtime/OpenGLDrv/Private/OpenGLUniformBuffer.cpp
--- a/Engine/Source/Runtime/OpenGLDrv/Private/OpenGLUniformBuffer.cpp
+++ b/Engine/Source/Runtime/OpenGLDrv/Private/OpenGLUniformBuffer.cpp
@@ -768,7 +768,7 @@ void FOpenGLDynamicRHI::RHIUpdateUniformBuffer(FUniformBufferRHIParamRef Uniform
 				checkf(Resource, TEXT("Invalid resource entry creating uniform buffer, %s.Resources[%u], ResourceType 0x%x."),
 					*Layout.GetDebugName().ToString(),
 					ResourceIndex,
-					Layout.Resources[ResourceIndex].MemberType);
+					(uint8)Layout.Resources[ResourceIndex].MemberType);
 
 				CmdListResources[ResourceIndex] = Resource;
 			}
diff --git a/Engine/Source/Runtime/Renderer/Private/DistanceFieldAmbientOcclusion.cpp b/Engine/Source/Runtime/Renderer/Private/DistanceFieldAmbientOcclusion.cpp
--- a/Engine/Source/Runtime/Renderer/Private/DistanceFieldAmbientOcclusion.cpp
+++ b/Engine/Source/Runtime/Renderer/Private/DistanceFieldAmbientOcclusion.cpp
@@ -18,7 +18,7 @@
 #include "RendererModule.h"
 #include "PipelineStateCache.h"
 #include "VisualizeTexture.h"
-#include "Raytracing/RaytracingOptions.h"
+#include "RayTracing/RaytracingOptions.h"
 
 int32 GDistanceFieldAO = 1;
 FAutoConsoleVariableRef CVarDistanceFieldAO(
diff --git a/Engine/Source/Runtime/Renderer/Private/ShaderBaseClasses.cpp b/Engine/Source/Runtime/Renderer/Private/ShaderBaseClasses.cpp
--- a/Engine/Source/Runtime/Renderer/Private/ShaderBaseClasses.cpp
+++ b/Engine/Source/Runtime/Renderer/Private/ShaderBaseClasses.cpp
@@ -120,7 +120,7 @@ void FMaterialShader::VerifyExpressionAndShaderMaps(const FMaterialRenderProxy*
 			FString ResourcesString;
 			for (int32 Index = 0; Index < Layout.Resources.Num(); ++Index)
 			{
-				ResourcesString += FString::Printf(TEXT("%d "), Layout.Resources[Index].MemberType);
+				ResourcesString += FString::Printf(TEXT("%d "), (uint8)Layout.Resources[Index].MemberType);
 			}
 			UE_LOG(LogShaders, Warning, TEXT("Layout CB Size %d %d Resources: %s"), Layout.ConstantBufferSize, Layout.Resources.Num(), *ResourcesString);
 		};
diff --git a/Engine/Source/Runtime/VulkanRHI/Private/VulkanUniformBuffer.cpp b/Engine/Source/Runtime/VulkanRHI/Private/VulkanUniformBuffer.cpp
--- a/Engine/Source/Runtime/VulkanRHI/Private/VulkanUniformBuffer.cpp
+++ b/Engine/Source/Runtime/VulkanRHI/Private/VulkanUniformBuffer.cpp
@@ -100,7 +100,7 @@ FVulkanUniformBuffer::FVulkanUniformBuffer(const FRHIUniformBufferLayout& InLayo
 			// Allow null SRV's in uniform buffers for feature levels that don't support SRV's in shaders
 			if (!(GMaxRHIFeatureLevel <= ERHIFeatureLevel::ES3_1 && InLayout.Resources[Index].MemberType == UBMT_SRV) && Validation == EUniformBufferValidation::ValidateResources)
 			{
-				checkf(Resource, TEXT("Invalid resource entry creating uniform buffer, %s.Resources[%u], ResourceType 0x%x."), *InLayout.GetDebugName().ToString(), Index, InLayout.Resources[Index].MemberType);
+				checkf(Resource, TEXT("Invalid resource entry creating uniform buffer, %s.Resources[%u], ResourceType 0x%x."), *InLayout.GetDebugName().ToString(), Index, (uint8)InLayout.Resources[Index].MemberType);
 			}
 			ResourceTable[Index] = Resource;
 		}
@@ -118,7 +118,7 @@ void FVulkanUniformBuffer::UpdateResourceTable(const FRHIUniformBufferLayout& In
 		checkf(Resource, TEXT("Invalid resource entry creating uniform buffer, %s.Resources[%u], ResourceType 0x%x."),
 			*InLayout.GetDebugName().ToString(),
 			ResourceIndex,
-			InLayout.Resources[ResourceIndex].MemberType);
+			(uint8)InLayout.Resources[ResourceIndex].MemberType);
 
 		ResourceTable[ResourceIndex] = Resource;
 	}
@@ -223,7 +223,7 @@ void FVulkanDynamicRHI::RHIUpdateUniformBuffer(FUniformBufferRHIParamRef Uniform
 				checkf(Resource, TEXT("Invalid resource entry creating uniform buffer, %s.Resources[%u], ResourceType 0x%x."),
 					*Layout.GetDebugName().ToString(),
 					ResourceIndex,
-					Layout.Resources[ResourceIndex].MemberType);
+					(uint8)Layout.Resources[ResourceIndex].MemberType);
 
 				CmdListResources[ResourceIndex] = Resource;
 			}
