FNS - Fix Linux Vulkan warning
#rnx

diff --git a/Engine/Source/Runtime/VulkanRHI/Private/Linux/VulkanLinuxPlatform.h b/Engine/Source/Runtime/VulkanRHI/Private/Linux/VulkanLinuxPlatform.h
--- a/Engine/Source/Runtime/VulkanRHI/Private/Linux/VulkanLinuxPlatform.h
+++ b/Engine/Source/Runtime/VulkanRHI/Private/Linux/VulkanLinuxPlatform.h
@@ -15,17 +15,17 @@
 #define VULKAN_SUPPORTS_NV_DIAGNOSTIC_CHECKPOINT	1
 
 
-#define ENUM_VK_ENTRYPOINTS_PLATFORM_BASE(EnumMacro) \
-	EnumMacro(PFN_vkCmdWriteBufferMarkerAMD, vkCmdWriteBufferMarkerAMD) \
-	EnumMacro(PFN_vkCmdSetCheckpointNV, vkCmdSetCheckpointNV) \
-	EnumMacro(PFN_vkGetQueueCheckpointDataNV, vkGetQueueCheckpointDataNV) \
-	EnumMacro(PFN_vkGetPhysicalDeviceProperties2KHR, vkGetPhysicalDeviceProperties2KHR) \
-	EnumMacro(PFN_vkGetImageMemoryRequirements2KHR , vkGetImageMemoryRequirements2KHR) \
-	EnumMacro(PFN_vkGetBufferMemoryRequirements2KHR , vkGetBufferMemoryRequirements2KHR)
+#define ENUM_VK_ENTRYPOINTS_PLATFORM_BASE(EnumMacro)
 
 #define ENUM_VK_ENTRYPOINTS_PLATFORM_INSTANCE(EnumMacro)
 
-#define ENUM_VK_ENTRYPOINTS_OPTIONAL_PLATFORM_INSTANCE(EnumMacro)
+#define ENUM_VK_ENTRYPOINTS_OPTIONAL_PLATFORM_INSTANCE(EnumMacro) \
+    EnumMacro(PFN_vkCmdWriteBufferMarkerAMD, vkCmdWriteBufferMarkerAMD) \
+    EnumMacro(PFN_vkCmdSetCheckpointNV, vkCmdSetCheckpointNV) \
+    EnumMacro(PFN_vkGetQueueCheckpointDataNV, vkGetQueueCheckpointDataNV) \
+    EnumMacro(PFN_vkGetPhysicalDeviceProperties2KHR, vkGetPhysicalDeviceProperties2KHR) \
+    EnumMacro(PFN_vkGetImageMemoryRequirements2KHR , vkGetImageMemoryRequirements2KHR) \
+    EnumMacro(PFN_vkGetBufferMemoryRequirements2KHR , vkGetBufferMemoryRequirements2KHR)
 
 #include "../VulkanLoader.h"
 
