DR - vk - Minor info added
#rnx

diff --git a/Engine/Source/Runtime/VulkanRHI/Private/VulkanSwapChain.cpp b/Engine/Source/Runtime/VulkanRHI/Private/VulkanSwapChain.cpp
--- a/Engine/Source/Runtime/VulkanRHI/Private/VulkanSwapChain.cpp
+++ b/Engine/Source/Runtime/VulkanRHI/Private/VulkanSwapChain.cpp
@@ -248,6 +248,7 @@ FVulkanSwapChain::FVulkanSwapChain(VkInstance InInstance, FVulkanDevice& InDevic
 	VkPresentModeKHR PresentMode = VK_PRESENT_MODE_FIFO_KHR;
 	if (FVulkanPlatform::SupportsQuerySurfaceProperties())
 	{
+		// Only dump the present modes the very first time they are queried
 		static bool bFirstTimeLog = true;
 
 		uint32 NumFoundPresentModes = 0;
@@ -280,6 +281,9 @@ FVulkanSwapChain::FVulkanSwapChain(VkInstance InInstance, FVulkanDevice& InDevic
 				bFoundPresentModeFIFO = true;
 				UE_CLOG(bFirstTimeLog, LogVulkanRHI, Display, TEXT("- VK_PRESENT_MODE_FIFO_KHR (%d)"), (int32)VK_PRESENT_MODE_FIFO_KHR);
 				break;
+			case VK_PRESENT_MODE_FIFO_RELAXED_KHR:
+				UE_CLOG(bFirstTimeLog, LogVulkanRHI, Display, TEXT("- VK_PRESENT_MODE_FIFO_RELAXED_KHR (%d)"), (int32)VK_PRESENT_MODE_FIFO_RELAXED_KHR);
+				break;
 			default:
 				UE_CLOG(bFirstTimeLog, LogVulkanRHI, Display, TEXT("- VkPresentModeKHR %d"), (int32)FoundPresentModes[i]);
 				break;
