DR - Support for RHIReadSurfaceData for float targets on Vulkan and OpenGL; added a default (slow) implementation as well
PR #5457: Adding readback of float32 render targets to FLinearColor arrays. (Contributed by [REDACTED_PERSON])

diff --git a/Engine/Source/Runtime/OpenGLDrv/Private/OpenGLRenderTarget.cpp b/Engine/Source/Runtime/OpenGLDrv/Private/OpenGLRenderTarget.cpp
--- a/Engine/Source/Runtime/OpenGLDrv/Private/OpenGLRenderTarget.cpp
+++ b/Engine/Source/Runtime/OpenGLDrv/Private/OpenGLRenderTarget.cpp
@@ -894,6 +894,47 @@ void FOpenGLDynamicRHI::RHIReadSurfaceData(FTextureRHIParamRef TextureRHI,FIntRe
 	RHITHREAD_GLCOMMAND_EPILOGUE();
 }
 
+void FOpenGLDynamicRHI::RHIReadSurfaceData(FTextureRHIParamRef TextureRHI, FIntRect Rect, TArray<FLinearColor>& OutData, FReadSurfaceDataFlags InFlags)
+{
+	VERIFY_GL_SCOPE();
+
+	// Verify requirements, but don't crash
+	if (!ensure(FOpenGL::SupportsFloatReadSurface()) || !ensure(TextureRHI) || !ensure(TextureRHI->GetFormat() == PF_A32B32G32R32F))
+	{
+		return;
+	}
+
+	FOpenGLTextureBase* Texture = GetOpenGLTextureFromRHITexture(TextureRHI);
+	if (!ensure(Texture))
+	{
+		return;
+	}
+
+	// Get framebuffer for texture
+	uint32 MipmapLevel = 0;
+	GLuint SourceFramebuffer = GetOpenGLFramebuffer(1, &Texture, NULL, &MipmapLevel, NULL);
+
+	uint32 SizeX = Rect.Width();
+	uint32 SizeY = Rect.Height();
+
+	// Initialize output
+	OutData.Empty(SizeX * SizeY);
+	OutData.AddUninitialized(SizeX * SizeY);
+
+	// Bind the framebuffer
+	// @TODO: Do we need to worry about multisampling?
+	glBindFramebuffer(UGL_READ_FRAMEBUFFER, SourceFramebuffer);
+	FOpenGL::ReadBuffer(SourceFramebuffer == 0 ? GL_BACK : GL_COLOR_ATTACHMENT0);
+
+	// Read the float data from the buffer directly into the output data
+	// @TODO: Do we need to support BGRA?
+	glPixelStorei(GL_PACK_ALIGNMENT, 1);
+	glReadPixels(Rect.Min.X, Rect.Min.Y, SizeX, SizeY, GL_RGBA, GL_FLOAT, OutData.GetData());
+	glPixelStorei(GL_PACK_ALIGNMENT, 4);
+
+	GetContextStateForCurrentContext().Framebuffer = (GLuint)-1;
+}
+
 void FOpenGLDynamicRHI::RHIMapStagingSurface(FTextureRHIParamRef TextureRHI,void*& OutData,int32& OutWidth,int32& OutHeight)
 {
 	FRHICommandListImmediate& RHICmdList = FRHICommandListExecutor::GetImmediateCommandList();
diff --git a/Engine/Source/Runtime/OpenGLDrv/Public/OpenGLDrv.h b/Engine/Source/Runtime/OpenGLDrv/Public/OpenGLDrv.h
--- a/Engine/Source/Runtime/OpenGLDrv/Public/OpenGLDrv.h
+++ b/Engine/Source/Runtime/OpenGLDrv/Public/OpenGLDrv.h
@@ -416,6 +416,7 @@ public:
 	virtual void RHIUnlockTextureCubeFace(FTextureCubeRHIParamRef Texture, uint32 FaceIndex, uint32 ArrayIndex, uint32 MipIndex, bool bLockWithinMiptail) final override;
 	virtual void RHIBindDebugLabelName(FTextureRHIParamRef Texture, const TCHAR* Name) final override;
 	virtual void RHIReadSurfaceData(FTextureRHIParamRef Texture,FIntRect Rect,TArray<FColor>& OutData,FReadSurfaceDataFlags InFlags) final override;
+	virtual void RHIReadSurfaceData(FTextureRHIParamRef Texture, FIntRect Rect, TArray<FLinearColor>& OutData, FReadSurfaceDataFlags InFlags) final override;
 	virtual void RHIMapStagingSurface(FTextureRHIParamRef Texture,void*& OutData,int32& OutWidth,int32& OutHeight) final override;
 	virtual void RHIUnmapStagingSurface(FTextureRHIParamRef Texture) final override;
 	virtual void RHIReadSurfaceFloatData(FTextureRHIParamRef Texture,FIntRect Rect,TArray<FFloat16Color>& OutData,ECubeFace CubeFace,int32 ArrayIndex,int32 MipIndex) final override;
diff --git a/Engine/Source/Runtime/RHI/Public/DynamicRHI.h b/Engine/Source/Runtime/RHI/Public/DynamicRHI.h
--- a/Engine/Source/Runtime/RHI/Public/DynamicRHI.h
+++ b/Engine/Source/Runtime/RHI/Public/DynamicRHI.h
@@ -802,8 +802,17 @@ public:
 	// FlushType: Flush Immediate (seems wrong)
 	virtual void RHIReadSurfaceData(FTextureRHIParamRef Texture, FIntRect Rect, TArray<FColor>& OutData, FReadSurfaceDataFlags InFlags) = 0;
 
-	// FlushType: Flush Immediate (seems wrong)
-	virtual void RHIReadSurfaceData(FTextureRHIParamRef Texture, FIntRect Rect, TArray<FLinearColor>& OutData, FReadSurfaceDataFlags InFlags) {}
+	// Default fallback; will not work for non-8-bit surfaces and it's extremely slow.
+	virtual void RHIReadSurfaceData(FTextureRHIParamRef Texture, FIntRect Rect, TArray<FLinearColor>& OutData, FReadSurfaceDataFlags InFlags)
+	{
+		TArray<FColor> TempData;
+		RHIReadSurfaceData(Texture, Rect, TempData, InFlags);
+		OutData.SetNumUninitialized(TempData.Num());
+		for (int32 Index = 0; Index < TempData.Num(); ++Index)
+		{
+			OutData[Index] = TempData[Index].ReinterpretAsLinear();
+		}
+	}
 
 	/** Watch out for OutData to be 0 (can happen on DXGI_ERROR_DEVICE_REMOVED), don't call RHIUnmapStagingSurface in that case. */
 	// FlushType: Flush Immediate (seems wrong)
diff --git a/Engine/Source/Runtime/VulkanRHI/Private/VulkanRenderTarget.cpp b/Engine/Source/Runtime/VulkanRHI/Private/VulkanRenderTarget.cpp
--- a/Engine/Source/Runtime/VulkanRHI/Private/VulkanRenderTarget.cpp
+++ b/Engine/Source/Runtime/VulkanRHI/Private/VulkanRenderTarget.cpp
@@ -751,48 +751,36 @@ void FVulkanCommandListContext::RHICopyToResolveTarget(FTextureRHIParamRef Sourc
 	}
 }
 
-void FVulkanDynamicRHI::RHIReadSurfaceData(FTextureRHIParamRef TextureRHI, FIntRect Rect, TArray<FColor>& OutData, FReadSurfaceDataFlags InFlags)
-{
-	FRHITexture2D* TextureRHI2D = TextureRHI->GetTexture2D();
-	check(TextureRHI2D);
-	FVulkanTexture2D* Texture2D = (FVulkanTexture2D*)TextureRHI2D;
-	uint32 NumPixels = (Rect.Max.X - Rect.Min.X) * (Rect.Max.Y - Rect.Min.Y);
-
-	if (GIgnoreCPUReads == 2)
-	{
-		// Debug: Fill with CPU
-		OutData.Empty(0);
-		OutData.AddZeroed(NumPixels);
-		return;
-	}
 
+static VulkanRHI::FStagingBuffer* ReadSurfaceDataCommon(FVulkanDevice* Device, FVulkanTexture2D* Texture2D, uint32 NumPixels, FVulkanCmdBuffer* CmdBuffer)
+{
 	Device->PrepareForCPURead();
 
 	FVulkanCommandListContext& ImmediateContext = Device->GetImmediateContext();
 
-	FVulkanCmdBuffer* CmdBuffer = ImmediateContext.GetCommandBufferManager()->GetUploadCmdBuffer();
-
 	ensure(Texture2D->Surface.StorageFormat == VK_FORMAT_R8G8B8A8_UNORM || Texture2D->Surface.StorageFormat == VK_FORMAT_B8G8R8A8_UNORM || Texture2D->Surface.StorageFormat == VK_FORMAT_R16G16B16A16_SFLOAT || Texture2D->Surface.StorageFormat == VK_FORMAT_A2B10G10R10_UNORM_PACK32 || Texture2D->Surface.StorageFormat == VK_FORMAT_R16G16B16A16_UNORM);
 	const bool bIs8Bpp = (Texture2D->Surface.StorageFormat == VK_FORMAT_R16G16B16A16_SFLOAT);
 	const uint32 Size = NumPixels * sizeof(FColor) * (bIs8Bpp ? 2 : 1);
+
 	VulkanRHI::FStagingBuffer* StagingBuffer = Device->GetStagingManager().AcquireBuffer(Size, VK_BUFFER_USAGE_TRANSFER_DST_BIT, true);
+
 	if (GIgnoreCPUReads == 0)
 	{
 		VkBufferImageCopy CopyRegion;
 		FMemory::Memzero(CopyRegion);
 		//Region.bufferOffset = 0;
-		CopyRegion.bufferRowLength = TextureRHI2D->GetSizeX();
-		CopyRegion.bufferImageHeight = TextureRHI2D->GetSizeY();
+		CopyRegion.bufferRowLength = Texture2D->GetSizeX();
+		CopyRegion.bufferImageHeight = Texture2D->GetSizeY();
 		CopyRegion.imageSubresource.aspectMask = Texture2D->Surface.GetFullAspectMask();
 		//Region.imageSubresource.mipLevel = 0;
 		//Region.imageSubresource.baseArrayLayer = 0;
 		CopyRegion.imageSubresource.layerCount = 1;
-		CopyRegion.imageExtent.width = TextureRHI2D->GetSizeX();
-		CopyRegion.imageExtent.height = TextureRHI2D->GetSizeY();
+		CopyRegion.imageExtent.width = Texture2D->GetSizeX();
+		CopyRegion.imageExtent.height = Texture2D->GetSizeY();
 		CopyRegion.imageExtent.depth = 1;
 
 		//#todo-rco: Multithreaded!
-		VkImageLayout& CurrentLayout = Device->GetImmediateContext().FindOrAddLayoutRW(Texture2D->Surface.Image, VK_IMAGE_LAYOUT_UNDEFINED);
+		VkImageLayout& CurrentLayout = ImmediateContext.FindOrAddLayoutRW(Texture2D->Surface.Image, VK_IMAGE_LAYOUT_UNDEFINED);
 		bool bHadLayout = (CurrentLayout != VK_IMAGE_LAYOUT_UNDEFINED);
 		if (CurrentLayout != VK_IMAGE_LAYOUT_TRANSFER_SRC_OPTIMAL)
 		{
@@ -824,16 +812,37 @@ void FVulkanDynamicRHI::RHIReadSurfaceData(FTextureRHIParamRef TextureRHI, FIntR
 	ImmediateContext.GetCommandBufferManager()->SubmitUploadCmdBuffer();
 	Device->WaitUntilIdle();
 
-/*
-	VkMappedMemoryRange MappedRange;
-	ZeroVulkanStruct(MappedRange, VK_STRUCTURE_TYPE_MAPPED_MEMORY_RANGE);
-	MappedRange.memory = StagingBuffer->GetDeviceMemoryHandle();
-	MappedRange.offset = StagingBuffer->GetAllocationOffset();
-	MappedRange.size = Size;
-	VulkanRHI::vkInvalidateMappedMemoryRanges(Device->GetInstanceHandle(), 1, &MappedRange);
-*/
+	/*
+		VkMappedMemoryRange MappedRange;
+		ZeroVulkanStruct(MappedRange, VK_STRUCTURE_TYPE_MAPPED_MEMORY_RANGE);
+		MappedRange.memory = StagingBuffer->GetDeviceMemoryHandle();
+		MappedRange.offset = StagingBuffer->GetAllocationOffset();
+		MappedRange.size = Size;
+		VulkanRHI::vkInvalidateMappedMemoryRanges(Device->GetInstanceHandle(), 1, &MappedRange);
+	*/
 	StagingBuffer->InvalidateMappedMemory();
 
+	return StagingBuffer;
+}
+
+void FVulkanDynamicRHI::RHIReadSurfaceData(FTextureRHIParamRef TextureRHI, FIntRect Rect, TArray<FColor>& OutData, FReadSurfaceDataFlags InFlags)
+{
+	FRHITexture2D* TextureRHI2D = TextureRHI->GetTexture2D();
+	check(TextureRHI2D);
+	FVulkanTexture2D* Texture2D = (FVulkanTexture2D*)TextureRHI2D;
+	uint32 NumPixels = (Rect.Max.X - Rect.Min.X) * (Rect.Max.Y - Rect.Min.Y);
+
+	if (GIgnoreCPUReads == 2)
+	{
+		// Debug: Fill with CPU
+		OutData.Empty(0);
+		OutData.AddZeroed(NumPixels);
+		return;
+	}
+
+	FVulkanCmdBuffer* CmdBuffer = Device->GetImmediateContext().GetCommandBufferManager()->GetUploadCmdBuffer();
+	VulkanRHI::FStagingBuffer* StagingBuffer = ReadSurfaceDataCommon(Device, Texture2D, NumPixels, CmdBuffer);
+
 	OutData.SetNum(NumPixels);
 	FColor* Dest = OutData.GetData();
 
@@ -932,7 +941,131 @@ void FVulkanDynamicRHI::RHIReadSurfaceData(FTextureRHIParamRef TextureRHI, FIntR
 	}
 
 	Device->GetStagingManager().ReleaseBuffer(CmdBuffer, StagingBuffer);
-	ImmediateContext.GetCommandBufferManager()->PrepareForNewActiveCommandBuffer();
+	Device->GetImmediateContext().GetCommandBufferManager()->PrepareForNewActiveCommandBuffer();
+}
+
+void FVulkanDynamicRHI::RHIReadSurfaceData(FTextureRHIParamRef TextureRHI, FIntRect Rect, TArray<FLinearColor>& OutData, FReadSurfaceDataFlags InFlags)
+{
+	FRHITexture2D* TextureRHI2D = TextureRHI->GetTexture2D();
+	check(TextureRHI2D);
+	FVulkanTexture2D* Texture2D = (FVulkanTexture2D*)TextureRHI2D;
+	uint32 NumPixels = (Rect.Max.X - Rect.Min.X) * (Rect.Max.Y - Rect.Min.Y);
+
+	if (GIgnoreCPUReads == 2)
+	{
+		// Debug: Fill with CPU
+		OutData.Empty(0);
+		OutData.AddZeroed(NumPixels);
+		return;
+	}
+
+	FVulkanCmdBuffer* CmdBuffer = Device->GetImmediateContext().GetCommandBufferManager()->GetUploadCmdBuffer();
+	VulkanRHI::FStagingBuffer* StagingBuffer = ReadSurfaceDataCommon(Device, Texture2D, NumPixels, CmdBuffer);
+
+	OutData.SetNum(NumPixels);
+	FLinearColor* Dest = OutData.GetData();
+
+	if (Texture2D->Surface.StorageFormat == VK_FORMAT_R16G16B16A16_SFLOAT)
+	{
+		for (int32 Row = Rect.Min.Y; Row < Rect.Max.Y; ++Row)
+		{
+			FFloat16Color* Src = (FFloat16Color*)StagingBuffer->GetMappedPointer() + Row * TextureRHI2D->GetSizeX() + Rect.Min.X;
+			for (int32 Col = Rect.Min.X; Col < Rect.Max.X; ++Col)
+			{
+				Dest->R = Src->R.GetFloat();
+				Dest->G = Src->G.GetFloat();
+				Dest->B = Src->B.GetFloat();
+				Dest->A = Src->A.GetFloat();
+				Dest++;
+				Src++;
+			}
+		}
+	}
+	else if (Texture2D->Surface.StorageFormat == VK_FORMAT_A2B10G10R10_UNORM_PACK32)
+	{
+		struct FR10G10B10A2
+		{
+			uint32 R : 10;
+			uint32 G : 10;
+			uint32 B : 10;
+			uint32 A : 2;
+		};
+		for (int32 Row = Rect.Min.Y; Row < Rect.Max.Y; ++Row)
+		{
+			FR10G10B10A2* Src = (FR10G10B10A2*)StagingBuffer->GetMappedPointer() + Row * TextureRHI2D->GetSizeX() + Rect.Min.X;
+			for (int32 Col = Rect.Min.X; Col < Rect.Max.X; ++Col)
+			{
+				*Dest = FLinearColor(
+					(float)Src->R / 1023.0f,
+					(float)Src->G / 1023.0f,
+					(float)Src->B / 1023.0f,
+					(float)Src->A / 3.0f
+				);
+				++Dest;
+				++Src;
+			}
+		}
+	}
+	else if (Texture2D->Surface.StorageFormat == VK_FORMAT_R8G8B8A8_UNORM)
+	{
+		for (int32 Row = Rect.Min.Y; Row < Rect.Max.Y; ++Row)
+		{
+			FColor* Src = (FColor*)StagingBuffer->GetMappedPointer() + Row * TextureRHI2D->GetSizeX() + Rect.Min.X;
+			for (int32 Col = Rect.Min.X; Col < Rect.Max.X; ++Col)
+			{
+				Dest->R = (float)Src->B / 255.0f;
+				Dest->G = (float)Src->G / 255.0f;
+				Dest->B = (float)Src->R / 255.0f;
+				Dest->A = (float)Src->A / 255.0f;
+				Dest++;
+				Src++;
+			}
+		}
+	}
+	else if (Texture2D->Surface.StorageFormat == VK_FORMAT_R16G16B16A16_UNORM)
+	{
+		struct FRGBA16U
+		{
+			uint16 R;
+			uint16 G;
+			uint16 B;
+			uint16 A;
+		};
+		for (int32 Row = Rect.Min.Y; Row < Rect.Max.Y; ++Row)
+		{
+			FRGBA16U* Src = (FRGBA16U*)StagingBuffer->GetMappedPointer() + Row * TextureRHI2D->GetSizeX() + Rect.Min.X;
+			for (int32 Col = Rect.Min.X; Col < Rect.Max.X; ++Col)
+			{
+				*Dest = FLinearColor(
+					(float)Src->R / 65535.0f,
+					(float)Src->G / 65535.0f,
+					(float)Src->B / 65535.0f,
+					(float)Src->A / 65535.0f
+				);
+				Dest++;
+				Src++;
+			}
+		}
+	}
+	else if (Texture2D->Surface.StorageFormat == VK_FORMAT_B8G8R8A8_UNORM)
+	{
+		for (int32 Row = Rect.Min.Y; Row < Rect.Max.Y; ++Row)
+		{
+			FColor* Src = (FColor*)StagingBuffer->GetMappedPointer() + Row * TextureRHI2D->GetSizeX() + Rect.Min.X;
+			for (int32 Col = Rect.Min.X; Col < Rect.Max.X; ++Col)
+			{
+				Dest->R = (float)Src->R / 255.0f;
+				Dest->G = (float)Src->G / 255.0f;
+				Dest->B = (float)Src->B / 255.0f;
+				Dest->A = (float)Src->A / 255.0f;
+				Dest++;
+				Src++;
+			}
+		}
+	}
+
+	Device->GetStagingManager().ReleaseBuffer(CmdBuffer, StagingBuffer);
+	Device->GetImmediateContext().GetCommandBufferManager()->PrepareForNewActiveCommandBuffer();
 }
 
 void FVulkanDynamicRHI::RHIMapStagingSurface(FTextureRHIParamRef TextureRHI,void*& OutData,int32& OutWidth,int32& OutHeight)
diff --git a/Engine/Source/Runtime/VulkanRHI/Public/VulkanDynamicRHI.h b/Engine/Source/Runtime/VulkanRHI/Public/VulkanDynamicRHI.h
--- a/Engine/Source/Runtime/VulkanRHI/Public/VulkanDynamicRHI.h
+++ b/Engine/Source/Runtime/VulkanRHI/Public/VulkanDynamicRHI.h
@@ -125,6 +125,7 @@ public:
 	virtual void RHIBindDebugLabelName(FTextureRHIParamRef Texture, const TCHAR* Name) final override;
 	virtual void RHIBindDebugLabelName(FUnorderedAccessViewRHIParamRef UnorderedAccessViewRHI, const TCHAR* Name) final override;
 	virtual void RHIReadSurfaceData(FTextureRHIParamRef Texture, FIntRect Rect, TArray<FColor>& OutData, FReadSurfaceDataFlags InFlags) final override;
+	virtual void RHIReadSurfaceData(FTextureRHIParamRef Texture, FIntRect Rect, TArray<FLinearColor>& OutData, FReadSurfaceDataFlags InFlags) final override;
 	virtual void RHIMapStagingSurface(FTextureRHIParamRef Texture, void*& OutData, int32& OutWidth, int32& OutHeight) final override;
 	virtual void RHIUnmapStagingSurface(FTextureRHIParamRef Texture) final override;
 	virtual void RHIReadSurfaceFloatData(FTextureRHIParamRef Texture, FIntRect Rect, TArray<FFloat16Color>& OutData, ECubeFace CubeFace, int32 ArrayIndex, int32 MipIndex) final override;
