If there are multiple windows open and vulkanrhi has multiiple swapchain objects the statics make frame pacing get very confused.

diff --git a/Engine/Source/Runtime/VulkanRHI/Private/VulkanSwapChain.cpp b/Engine/Source/Runtime/VulkanRHI/Private/VulkanSwapChain.cpp
--- a/Engine/Source/Runtime/VulkanRHI/Private/VulkanSwapChain.cpp
+++ b/Engine/Source/Runtime/VulkanRHI/Private/VulkanSwapChain.cpp
@@ -121,6 +121,8 @@ FVulkanSwapChain::FVulkanSwapChain(VkInstance InInstance, FVulkanDevice& InDevic
 {
 	check(FVulkanPlatform::SupportsStandardSwapchain());
 
+	NextPresentTargetTime = (FPlatformTime::Seconds() - GStartTime);
+
 	// let the platform create the surface
 	FVulkanPlatform::CreateSurface(WindowHandle, Instance, &Surface);
 
@@ -755,25 +757,21 @@ void FVulkanSwapChain::RenderThreadPacing()
 	//very naive CPU side frame pacer.
 	if (GVulkanCPURenderThreadFramePacer && SyncInterval > 0)
 	{
-		static double PreviousFrameCPUTime = 0;
-		static double SampledDeltaTimeMS = 0;
-		static uint32 SampleCount = 0;
-
 		double NowCPUTime = FPlatformTime::Seconds();
-		double DeltaCPUPresentTimeMS = (NowCPUTime - PreviousFrameCPUTime) * 1000.0;
+		double DeltaCPUPresentTimeMS = (NowCPUTime - RTPacingPreviousFrameCPUTime) * 1000.0;
 
 
 		double TargetIntervalWithEpsilonMS = (double)SyncInterval * (1.0 / 60.0) * 1000.0;
 		const double IntervalThresholdMS = TargetIntervalWithEpsilonMS * 0.1;
 
-		SampledDeltaTimeMS += DeltaCPUPresentTimeMS; SampleCount++;
+		RTPacingSampledDeltaTimeMS += DeltaCPUPresentTimeMS; RTPacingSampleCount++;
 
-		double SampledDeltaMS = (SampledDeltaTimeMS / (double)SampleCount) + IntervalThresholdMS;
+		double SampledDeltaMS = (RTPacingSampledDeltaTimeMS / (double)RTPacingSampleCount) + IntervalThresholdMS;
 
-		if (SampleCount > 1000)
+		if (RTPacingSampleCount > 1000)
 		{
-			SampledDeltaTimeMS = SampledDeltaMS;
-			SampleCount = 1;
+			RTPacingSampledDeltaTimeMS = SampledDeltaMS;
+			RTPacingSampleCount = 1;
 		}
 
 		if (SampledDeltaMS < (TargetIntervalWithEpsilonMS))
@@ -797,7 +795,7 @@ void FVulkanSwapChain::RenderThreadPacing()
 				UE_LOG(LogVulkanRHI, Log, TEXT("CPU RT delta: %f"), DeltaCPUPresentTimeMS);
 			}
 		}
-		PreviousFrameCPUTime = NowCPUTime;
+		RTPacingPreviousFrameCPUTime = NowCPUTime;
 	}
 }
 
@@ -843,7 +841,6 @@ FVulkanSwapChain::EStatus FVulkanSwapChain::Present(FVulkanQueue* GfxQueue, FVul
 	if (GVulkanCPURHIFramePacer && SyncInterval > 0)
 	{
 		const double NowCPUTime = (FPlatformTime::Seconds() - GStartTime);
-		static double NextPresentTargetTime = NowCPUTime;
 
 		const double TimeToSleep = (NextPresentTargetTime - NowCPUTime);
 		const double TargetIntervalWithEpsilon = (double)SyncInterval * (1.0 / 60.0);
diff --git a/Engine/Source/Runtime/VulkanRHI/Private/VulkanSwapChain.h b/Engine/Source/Runtime/VulkanRHI/Private/VulkanSwapChain.h
--- a/Engine/Source/Runtime/VulkanRHI/Private/VulkanSwapChain.h
+++ b/Engine/Source/Runtime/VulkanRHI/Private/VulkanSwapChain.h
@@ -43,6 +43,13 @@ protected:
 	int32 SemaphoreIndex;
 	uint32 NumPresentCalls;
 	uint32 NumAcquireCalls;
+
+	uint32 RTPacingSampleCount = 0;
+	double RTPacingPreviousFrameCPUTime = 0;
+	double RTPacingSampledDeltaTimeMS = 0;
+	
+	double NextPresentTargetTime = 0;
+
 	VkInstance Instance;
 	TArray<VulkanRHI::FSemaphore*> ImageAcquiredSemaphore;
 #if VULKAN_USE_IMAGE_ACQUIRE_FENCES
