DR - Use Vulkan 1.1 on Windows to allow GPU-based validation
#rnx

diff --git a/Engine/Source/Runtime/VulkanRHI/Private/Windows/VulkanWindowsPlatform.h b/Engine/Source/Runtime/VulkanRHI/Private/Windows/VulkanWindowsPlatform.h
--- a/Engine/Source/Runtime/VulkanRHI/Private/Windows/VulkanWindowsPlatform.h
+++ b/Engine/Source/Runtime/VulkanRHI/Private/Windows/VulkanWindowsPlatform.h
@@ -6,20 +6,22 @@
 #include "Windows/WindowsHWrapper.h"
 #include "RHI.h"
 
-#define VK_USE_PLATFORM_WIN32_KHR				1
-#define VK_USE_PLATFORM_WIN32_KHX				1
-
-#define	VULKAN_SHOULD_ENABLE_DRAW_MARKERS		(UE_BUILD_DEBUG || UE_BUILD_DEVELOPMENT)
-#define VULKAN_HAS_PHYSICAL_DEVICE_PROPERTIES2	1
-#define VULKAN_USE_CREATE_WIN32_SURFACE			1
-#define VULKAN_DYNAMICALLYLOADED				1
-#define VULKAN_ENABLE_DESKTOP_HMD_SUPPORT		1
-#define VULKAN_SIGNAL_UNIMPLEMENTED()			checkf(false, TEXT("Unimplemented vulkan functionality: %s"), TEXT(__FUNCTION__))
-#define VULKAN_SUPPORTS_COLOR_CONVERSIONS		1
-#define	VULKAN_SUPPORTS_DEDICATED_ALLOCATION	0
+#define VK_USE_PLATFORM_WIN32_KHR					1
+#define VK_USE_PLATFORM_WIN32_KHX					1
+
+#define	VULKAN_SHOULD_ENABLE_DRAW_MARKERS			(UE_BUILD_DEBUG || UE_BUILD_DEVELOPMENT)
+#define VULKAN_HAS_PHYSICAL_DEVICE_PROPERTIES2		1
+#define VULKAN_USE_CREATE_WIN32_SURFACE				1
+#define VULKAN_DYNAMICALLYLOADED					1
+#define VULKAN_ENABLE_DESKTOP_HMD_SUPPORT			1
+#define VULKAN_SIGNAL_UNIMPLEMENTED()				checkf(false, TEXT("Unimplemented vulkan functionality: %s"), TEXT(__FUNCTION__))
+#define VULKAN_SUPPORTS_COLOR_CONVERSIONS			1
+#define	VULKAN_SUPPORTS_DEDICATED_ALLOCATION		0
 #define VULKAN_SUPPORTS_AMD_BUFFER_MARKER			1
 #define VULKAN_SUPPORTS_NV_DIAGNOSTIC_CHECKPOINT	1
 
+#define	UE_VK_API_VERSION							VK_API_VERSION_1_1
+
 
 // 32-bit windows has warnings on custom mem mgr callbacks
 #define VULKAN_SHOULD_USE_LLM					(UE_BUILD_DEBUG || UE_BUILD_DEVELOPMENT) && !PLATFORM_32BITS
