DR - Blacklist any non-18.xx.xx AMD drivers on Vulkan

diff --git a/Engine/Source/Runtime/VulkanRHI/Private/VulkanGenericPlatform.h b/Engine/Source/Runtime/VulkanRHI/Private/VulkanGenericPlatform.h
--- a/Engine/Source/Runtime/VulkanRHI/Private/VulkanGenericPlatform.h
+++ b/Engine/Source/Runtime/VulkanRHI/Private/VulkanGenericPlatform.h
@@ -12,6 +12,7 @@ class FVulkanGenericPlatform
 {
 public:
 	static bool IsSupported() { return true; }
+	static void CheckDeviceDriver() {}
 
 	static bool LoadVulkanLibrary() { return true; }
 	static bool LoadVulkanInstanceFunctions(VkInstance inInstance) { return true; }
diff --git a/Engine/Source/Runtime/VulkanRHI/Private/VulkanRHI.cpp b/Engine/Source/Runtime/VulkanRHI/Private/VulkanRHI.cpp
--- a/Engine/Source/Runtime/VulkanRHI/Private/VulkanRHI.cpp
+++ b/Engine/Source/Runtime/VulkanRHI/Private/VulkanRHI.cpp
@@ -570,6 +570,8 @@ void FVulkanDynamicRHI::SelectAndInitDevice()
 	GRHIVendorId = Props.vendorID;
 	GRHIAdapterName = ANSI_TO_TCHAR(Props.deviceName);
 
+	FVulkanPlatform::CheckDeviceDriver(DeviceIndex);
+
 	Device->InitGPU(DeviceIndex);
 
 	if (PLATFORM_ANDROID && !PLATFORM_LUMIN && !PLATFORM_LUMINGL4)
diff --git a/Engine/Source/Runtime/VulkanRHI/Private/Windows/VulkanWindowsPlatform.cpp b/Engine/Source/Runtime/VulkanRHI/Private/Windows/VulkanWindowsPlatform.cpp
--- a/Engine/Source/Runtime/VulkanRHI/Private/Windows/VulkanWindowsPlatform.cpp
+++ b/Engine/Source/Runtime/VulkanRHI/Private/Windows/VulkanWindowsPlatform.cpp
@@ -3,6 +3,7 @@
 #include "VulkanWindowsPlatform.h"
 #include "../VulkanRHIPrivate.h"
 #include "../VulkanDevice.h"
+#include "amd_ags.h"
 
 #include "Windows/AllowWindowsPlatformTypes.h"
 static HMODULE GVulkanDLLModule = nullptr;
@@ -181,3 +182,47 @@ void FVulkanWindowsPlatform::WriteCrashMarker(const FOptionalVulkanDeviceExtensi
 		}
 	}
 }
+
+void FVulkanWindowsPlatform::CheckDeviceDriver(uint32 DeviceIndex)
+{
+	const bool bAllowVendorDevice = !FParse::Param(FCommandLine::Get(), TEXT("novendordevice"));
+	if (IsRHIDeviceAMD() && bAllowVendorDevice)
+	{
+		AGSGPUInfo AmdGpuInfo;
+		AGSContext* AmdAgsContext = nullptr;
+		if (agsInit(&AmdAgsContext, nullptr, &AmdGpuInfo) == AGS_SUCCESS)
+		{
+			if (DeviceIndex < (uint32)AmdGpuInfo.numDevices)
+			{
+				auto& DeviceInfo = AmdGpuInfo.devices[DeviceIndex];
+				bool bIsPreGCN = DeviceInfo.architectureVersion == AGSDeviceInfo::ArchitectureVersion_GCN;
+				if (DeviceInfo.architectureVersion == AGSDeviceInfo::ArchitectureVersion_PreGCN || bIsPreGCN)
+				{
+					int32 MajorVersion = FCStringAnsi::Atoi(AmdGpuInfo.radeonSoftwareVersion);
+					if (MajorVersion > 0)
+					{
+						if (MajorVersion < 18)
+						{
+							// Blacklist drivers older than 18.xx.xx drivers
+							FPlatformMisc::MessageBoxExt(EAppMsgType::Ok, TEXT("There are known issues with Vulkan with drivers older than the 18.xx.xx.xx\nfamily of Radeon drivers; the recommended version is 18.12.1.1: please try\nupdating the driver to that version."), TEXT("Vulkan driver version"));
+							FPlatformMisc::RequestExitWithStatus(true, 1);
+						}
+						else if (MajorVersion == 19)
+						{
+							if (WITH_EDITOR)
+							{
+								// Blacklist 19.xx.xx drivers as it introduced an issue with Slate windows/Vulkan viewports on the editor
+								FPlatformMisc::MessageBoxExt(EAppMsgType::Ok, TEXT("There are known issues with Vulkan on the editor with the 19.xx.xx.xx\nfamily of Radeon drivers; the recommended version is 18.12.1.1: please\n try updating the driver to that version."), TEXT("Vulkan driver version"));
+									FPlatformMisc::RequestExitWithStatus(true, 1);
+							}
+						}
+					}
+
+					GRHIDeviceIsAMDPreGCNArchitecture = GRHIDeviceIsAMDPreGCNArchitecture || bIsPreGCN;
+				}
+			}
+
+			agsDeInit(AmdAgsContext);
+		}
+	}
+}
diff --git a/Engine/Source/Runtime/VulkanRHI/Private/Windows/VulkanWindowsPlatform.h b/Engine/Source/Runtime/VulkanRHI/Private/Windows/VulkanWindowsPlatform.h
--- a/Engine/Source/Runtime/VulkanRHI/Private/Windows/VulkanWindowsPlatform.h
+++ b/Engine/Source/Runtime/VulkanRHI/Private/Windows/VulkanWindowsPlatform.h
@@ -49,6 +49,8 @@
 class FVulkanWindowsPlatform : public FVulkanGenericPlatform
 {
 public:
+	static void CheckDeviceDriver(uint32 DeviceIndex);
+
 	static bool LoadVulkanLibrary();
 	static bool LoadVulkanInstanceFunctions(VkInstance inInstance);
 	static void FreeVulkanLibrary();
diff --git a/Engine/Source/Runtime/VulkanRHI/VulkanRHI.Build.cs b/Engine/Source/Runtime/VulkanRHI/VulkanRHI.Build.cs
--- a/Engine/Source/Runtime/VulkanRHI/VulkanRHI.Build.cs
+++ b/Engine/Source/Runtime/VulkanRHI/VulkanRHI.Build.cs
@@ -14,6 +14,7 @@ public class VulkanRHI : ModuleRules
 		if (Target.Platform == UnrealTargetPlatform.Win32 || Target.Platform == UnrealTargetPlatform.Win64)
 		{
 			PrivateIncludePaths.Add("Runtime/VulkanRHI/Private/Windows");
+			AddEngineThirdPartyPrivateStaticDependencies(Target, "AMD_AGS");
 		}
 		else
 		{
