DR - Update blacklisted driver versions on Vulkan AMD Windows.
[REDACTED_PHONE].xx are blacklisted.

diff --git a/Engine/Source/Runtime/VulkanRHI/Private/Windows/VulkanWindowsPlatform.cpp b/Engine/Source/Runtime/VulkanRHI/Private/Windows/VulkanWindowsPlatform.cpp
--- a/Engine/Source/Runtime/VulkanRHI/Private/Windows/VulkanWindowsPlatform.cpp
+++ b/Engine/Source/Runtime/VulkanRHI/Private/Windows/VulkanWindowsPlatform.cpp
@@ -192,33 +192,81 @@ void FVulkanWindowsPlatform::CheckDeviceDriver(uint32 DeviceIndex)
 		AGSContext* AmdAgsContext = nullptr;
 		if (agsInit(&AmdAgsContext, nullptr, &AmdGpuInfo) == AGS_SUCCESS)
 		{
-			if (DeviceIndex < (uint32)AmdGpuInfo.numDevices)
+			const char* Version = AmdGpuInfo.radeonSoftwareVersion;
+			if (DeviceIndex < (uint32)AmdGpuInfo.numDevices && Version && *Version)
 			{
 				auto& DeviceInfo = AmdGpuInfo.devices[DeviceIndex];
 				bool bIsPreGCN = DeviceInfo.architectureVersion == AGSDeviceInfo::ArchitectureVersion_GCN;
 				if (DeviceInfo.architectureVersion == AGSDeviceInfo::ArchitectureVersion_PreGCN || bIsPreGCN)
 				{
-					int32 MajorVersion = FCStringAnsi::Atoi(AmdGpuInfo.radeonSoftwareVersion);
-					if (MajorVersion > 0)
+					// "Major.Minor.Revision"
+					do
 					{
-						if (MajorVersion < 18)
+						int32 MajorVersion = FCStringAnsi::Atoi(Version);
+						while (*Version >= '0' && *Version <= '9')
 						{
-							// Blacklist drivers older than 18.xx.xx drivers
-							FPlatformMisc::MessageBoxExt(EAppMsgType::Ok, TEXT("There are known issues with Vulkan with drivers older than the 18.xx.xx.xx\nfamily of Radeon drivers; the recommended version is 18.12.1.1: please try\nupdating the driver to that version."), TEXT("Vulkan driver version"));
-							FPlatformMisc::RequestExitWithStatus(true, 1);
+							++Version;
 						}
-						else if (MajorVersion == 19)
+
+						if (*Version != '.')
+						{
+							break;
+						}
+						++Version;
+
+						int32 MinorVersion = FCStringAnsi::Atoi(Version);
+						while (*Version >= '0' && *Version <= '9')
+						{
+							++Version;
+						}
+
+						if (*Version != '.')
 						{
-							if (WITH_EDITOR)
+							break;
+						}
+						++Version;
+
+						int32 RevisionVersion = FCStringAnsi::Atoi(Version);
+
+						if (MajorVersion > 0)
+						{
+							if (MajorVersion < 18)
+							{
+								// Blacklist drivers older than 18.xx.xx drivers
+								FPlatformMisc::MessageBoxExt(EAppMsgType::Ok, TEXT("There are known issues with Vulkan with drivers older than the 18.xx.xx.xx\nfamily of Radeon drivers; the recommended version is 18.12.1.1 or anything more recent than 19.10.xx: please try updating your driver to that version."), TEXT("Vulkan driver version"));
+								FPlatformMisc::RequestExitWithStatus(true, 1);
+							}
+							else if (WITH_EDITOR)
 							{
-								// Blacklist 19.xx.xx drivers as it introduced an issue with Slate windows/Vulkan viewports on the editor
-								FPlatformMisc::MessageBoxExt(EAppMsgType::Ok, TEXT("There are known issues with Vulkan on the editor with the 19.xx.xx.xx\nfamily of Radeon drivers; the recommended version is 18.12.1.1: please\n try updating the driver to that version."), TEXT("Vulkan driver version"));
+								bool bBadVersion = false;
+								if (MajorVersion == 19)
+								{
+									if (MinorVersion < 10)
+									{
+										bBadVersion = true;
+									}
+								}
+								else if (MajorVersion == 18)
+								{
+									if (MinorVersion > 12 || (MinorVersion == 12 && RevisionVersion >= 2))
+									{
+										bBadVersion = true;
+									}
+								}
+
+								if (bBadVersion)
+								{
+									// Blacklist 19.xx.xx drivers as it introduced an issue with Slate windows/Vulkan viewports on the editor
+									FPlatformMisc::MessageBoxExt(EAppMsgType::Ok, TEXT("There are known issues with Vulkan on the editor with the some \nRadeon drivers; the recommended version is up to 18.12.1.1 or anything more recent than 19.10.xx: please try updating your driver to that version."), TEXT("Vulkan driver version"));
 									FPlatformMisc::RequestExitWithStatus(true, 1);
+								}
 							}
 						}
 					}
+					while (0);
 
 					GRHIDeviceIsAMDPreGCNArchitecture = GRHIDeviceIsAMDPreGCNArchitecture || bIsPreGCN;
+					UE_LOG(LogVulkanRHI, Display, TEXT("AMD User Driver Version = %s"), ANSI_TO_TCHAR(AmdGpuInfo.radeonSoftwareVersion));
 				}
 			}
 
