DR - Compile fix
#rnx

diff --git a/Engine/Source/Runtime/VulkanRHI/Private/VulkanGenericPlatform.h b/Engine/Source/Runtime/VulkanRHI/Private/VulkanGenericPlatform.h
--- a/Engine/Source/Runtime/VulkanRHI/Private/VulkanGenericPlatform.h
+++ b/Engine/Source/Runtime/VulkanRHI/Private/VulkanGenericPlatform.h
@@ -12,7 +12,7 @@ class FVulkanGenericPlatform
 {
 public:
 	static bool IsSupported() { return true; }
-	static void CheckDeviceDriver() {}
+	static void CheckDeviceDriver(uint32 DeviceIndex) {}
 
 	static bool LoadVulkanLibrary() { return true; }
 	static bool LoadVulkanInstanceFunctions(VkInstance inInstance) { return true; }
