DR - Don't hard crash when OOM on Vulkan, show a log error
#rnx

diff --git a/Engine/Source/Runtime/VulkanRHI/Private/VulkanMemory.cpp b/Engine/Source/Runtime/VulkanRHI/Private/VulkanMemory.cpp
--- a/Engine/Source/Runtime/VulkanRHI/Private/VulkanMemory.cpp
+++ b/Engine/Source/Runtime/VulkanRHI/Private/VulkanMemory.cpp
@@ -895,10 +895,14 @@ namespace VulkanRHI
 		uint32 AllocationSize = FMath::Max(Size, TargetDefaultPageSize);
 #endif
 		FDeviceMemoryAllocation* DeviceMemoryAllocation = Owner->GetParent()->GetMemoryManager().Alloc(true, AllocationSize, MemoryTypeIndex, nullptr, File, Line);
-		if (!DeviceMemoryAllocation && Size < AllocationSize)
+		if (!DeviceMemoryAllocation)
 		{
 			// Retry with a smaller size
 			DeviceMemoryAllocation = Owner->GetParent()->GetMemoryManager().Alloc(false, Size, MemoryTypeIndex, nullptr, File, Line);
+			if (!DeviceMemoryAllocation)
+			{
+				UE_LOG(LogVulkanRHI, Fatal, TEXT("Out of memory on Vulkan; MemoryTypeIndex=%d, AllocSize=%0.3fMB"), MemoryTypeIndex, (float)AllocationSize / 1048576.0f);
+			}
 		}
 		++PageIDCounter;
 		FOldResourceHeapPage* NewPage = new FOldResourceHeapPage(this, DeviceMemoryAllocation, PageIDCounter);
