DR - vk - Fix NV crash diagnostic
#rnx

diff --git a/Engine/Source/Runtime/VulkanRHI/Private/VulkanUtil.cpp b/Engine/Source/Runtime/VulkanRHI/Private/VulkanUtil.cpp
--- a/Engine/Source/Runtime/VulkanRHI/Private/VulkanUtil.cpp
+++ b/Engine/Source/Runtime/VulkanRHI/Private/VulkanUtil.cpp
@@ -550,13 +550,20 @@ void FVulkanGPUProfiler::DumpCrashMarkers(void* BufferData)
 #if VULKAN_SUPPORTS_NV_DIAGNOSTIC_CHECKPOINT
 		if (Device->GetOptionalExtensions().HasNVDiagnosticCheckpoints)
 		{
-			TArray<VkCheckpointDataNV> Data;
+			struct FCheckpointDataNV : public VkCheckpointDataNV
+			{
+				FCheckpointDataNV()
+				{
+					ZeroVulkanStruct(*this, VK_STRUCTURE_TYPE_CHECKPOINT_DATA_NV);
+				}
+			};
+			TArray<FCheckpointDataNV> Data;
 			uint32 Num = 0;
 			VkQueue QueueHandle = Device->GetGraphicsQueue()->GetHandle();
 			VulkanDynamicAPI::vkGetQueueCheckpointDataNV(QueueHandle, &Num, nullptr);
 			if (Num > 0)
 			{
-				Data.AddUninitialized(Num);
+				Data.AddDefaulted(Num);
 				VulkanDynamicAPI::vkGetQueueCheckpointDataNV(QueueHandle, &Num, &Data[0]);
 				check(Num == Data.Num());
 				for (uint32 Index = 0; Index < Num; ++Index)
