UE4.22 - Fix Android Vulkan validation layers
#rnx

diff --git a/Engine/Source/Runtime/VulkanRHI/Private/VulkanLayers.cpp b/Engine/Source/Runtime/VulkanRHI/Private/VulkanLayers.cpp
--- a/Engine/Source/Runtime/VulkanRHI/Private/VulkanLayers.cpp
+++ b/Engine/Source/Runtime/VulkanRHI/Private/VulkanLayers.cpp
@@ -15,6 +15,7 @@ bool GRenderDocFound = false;
 #include "Runtime/HeadMountedDisplay/Public/IHeadMountedDisplayModule.h"
 #endif
 
+#if VULKAN_HAS_DEBUGGING_ENABLED
 TAutoConsoleVariable<int32> GValidationCvar(
 	TEXT("r.Vulkan.EnableValidation"),
 	0,
@@ -27,56 +28,30 @@ TAutoConsoleVariable<int32> GValidationCvar(
 	ECVF_ReadOnly | ECVF_RenderThreadSafe
 );
 
-#if VULKAN_HAS_DEBUGGING_ENABLED
+static TAutoConsoleVariable<int32> GStandardValidationCvar(
+	TEXT("r.Vulkan.StandardValidation"),
+	1,
+	TEXT("1 to use VK_LAYER_LUNARG_standard_validation (default) if available\n")
+	TEXT("0 to use individual layers"),
+	ECVF_ReadOnly | ECVF_RenderThreadSafe
+);
 
-	#if VULKAN_ENABLE_DRAW_MARKERS
-		#define RENDERDOC_LAYER_NAME	"VK_LAYER_RENDERDOC_Capture"
-	#endif
+#if VULKAN_ENABLE_DRAW_MARKERS
+	#define RENDERDOC_LAYER_NAME		"VK_LAYER_RENDERDOC_Capture"
+#endif
 
-#define VULKAN_ENABLE_STANDARD_VALIDATION 1
+#define STANDARD_VALIDATION_LAYER_NAME	"VK_LAYER_LUNARG_standard_validation"
 
-static const ANSICHAR* GValidationLayersInstance[] =
+static const ANSICHAR* GIndividualValidationLayers[] =
 {
-#if VULKAN_ENABLE_STANDARD_VALIDATION
-	"VK_LAYER_LUNARG_standard_validation",
-#elif PLATFORM_ANDROID
-	"VK_LAYER_GOOGLE_threading",
-	"VK_LAYER_LUNARG_parameter_validation",
-	"VK_LAYER_LUNARG_object_tracker",
-	"VK_LAYER_LUNARG_core_validation",
-	"VK_LAYER_GOOGLE_unique_objects",
-#else
 	"VK_LAYER_GOOGLE_threading",
 	"VK_LAYER_LUNARG_parameter_validation",
 	"VK_LAYER_LUNARG_object_tracker",
 	"VK_LAYER_LUNARG_core_validation",
 	"VK_LAYER_GOOGLE_unique_objects",
-#endif
 	nullptr
 };
 
-static const ANSICHAR* GValidationLayersDevice[] =
-{
-#if VULKAN_ENABLE_STANDARD_VALIDATION
-	"VK_LAYER_LUNARG_standard_validation",
-#elif PLATFORM_ANDROID
-	"VK_LAYER_GOOGLE_threading",
-	"VK_LAYER_LUNARG_parameter_validation",
-	"VK_LAYER_LUNARG_object_tracker",
-	"VK_LAYER_LUNARG_core_validation",
-	"VK_LAYER_GOOGLE_unique_objects",
-#else
-	"VK_LAYER_GOOGLE_threading",
-	"VK_LAYER_LUNARG_parameter_validation",
-	"VK_LAYER_LUNARG_object_tracker",
-	"VK_LAYER_LUNARG_image",
-	"VK_LAYER_LUNARG_core_validation",
-	"VK_LAYER_LUNARG_swapchain",
-	"VK_LAYER_GOOGLE_unique_objects",
-	"VK_LAYER_LUNARG_core_validation",
-#endif
-	nullptr
-};
 #endif // VULKAN_HAS_DEBUGGING_ENABLED
 
 // Instance Extensions to enable for all platforms
@@ -342,18 +317,35 @@ void FVulkanDynamicRHI::GetInstanceLayersAndExtensions(TArray<const ANSICHAR*>&
 
 	if (!bVkTrace && VulkanValidationOption > 0)
 	{
-		// Verify that all requested debugging device-layers are available
-		for (uint32 LayerIndex = 0; GValidationLayersInstance[LayerIndex] != nullptr; ++LayerIndex)
+		bool bStandardAvailable = false;
+		if (GStandardValidationCvar.GetValueOnAnyThread() != 0)
 		{
-			const ANSICHAR* CurrValidationLayer = GValidationLayersInstance[LayerIndex];
-			bool bValidationFound = FindLayerInList(GlobalLayerExtensions, CurrValidationLayer);
-			if (bValidationFound)
+			bStandardAvailable = FindLayerInList(GlobalLayerExtensions, STANDARD_VALIDATION_LAYER_NAME);
+			if (bStandardAvailable)
 			{
-				OutInstanceLayers.Add(CurrValidationLayer);
+				OutInstanceLayers.Add(STANDARD_VALIDATION_LAYER_NAME);
 			}
 			else
 			{
-				UE_LOG(LogVulkanRHI, Warning, TEXT("Unable to find Vulkan instance validation layer '%s'"), ANSI_TO_TCHAR(CurrValidationLayer));
+				UE_LOG(LogVulkanRHI, Warning, TEXT("Unable to find Vulkan instance validation layer %s; trying individual layers..."), TEXT(STANDARD_VALIDATION_LAYER_NAME));
+			}
+		}
+
+		if (!bStandardAvailable)
+		{
+			// Verify that all requested debugging device-layers are available
+			for (uint32 LayerIndex = 0; GIndividualValidationLayers[LayerIndex] != nullptr; ++LayerIndex)
+			{
+				const ANSICHAR* CurrValidationLayer = GIndividualValidationLayers[LayerIndex];
+				bool bValidationFound = FindLayerInList(GlobalLayerExtensions, CurrValidationLayer);
+				if (bValidationFound)
+				{
+					OutInstanceLayers.Add(CurrValidationLayer);
+				}
+				else
+				{
+					UE_LOG(LogVulkanRHI, Warning, TEXT("Unable to find Vulkan instance validation layer '%s'"), ANSI_TO_TCHAR(CurrValidationLayer));
+				}
 			}
 		}
 	}
@@ -410,6 +402,7 @@ void FVulkanDynamicRHI::GetInstanceLayersAndExtensions(TArray<const ANSICHAR*>&
 		OutInstanceExtensions.Add(VK_EXT_DEBUG_UTILS_EXTENSION_NAME);
 	}
 #endif
+#if VULKAN_HAS_DEBUGGING_ENABLED
 	if (!bVkTrace && !bOutDebugUtils && GValidationCvar.GetValueOnAnyThread() > 0)
 	{
 		if (FindLayerExtensionInList(GlobalLayerExtensions, VK_EXT_DEBUG_REPORT_EXTENSION_NAME))
@@ -417,6 +410,7 @@ void FVulkanDynamicRHI::GetInstanceLayersAndExtensions(TArray<const ANSICHAR*>&
 			OutInstanceExtensions.Add(VK_EXT_DEBUG_REPORT_EXTENSION_NAME);
 		}
 	}
+#endif
 
 	TrimDuplicates(OutInstanceLayers);
 	if (OutInstanceLayers.Num() > 0)
@@ -525,23 +519,36 @@ void FVulkanDevice::GetDeviceExtensionsAndLayers(TArray<const ANSICHAR*>& OutDev
 	// Verify that all requested debugging device-layers are available. Skip validation layers under RenderDoc
 	if (!GRenderDocFound && GValidationCvar.GetValueOnAnyThread() > 0)
 	{
-		for (uint32 LayerIndex = 0; GValidationLayersDevice[LayerIndex] != nullptr; ++LayerIndex)
+		bool bStandardAvailable = false;
+		if (GStandardValidationCvar.GetValueOnAnyThread() != 0)
 		{
-			bool bValidationFound = false;
-			const ANSICHAR* CurrValidationLayer = GValidationLayersDevice[LayerIndex];
-			for (int32 Index = 1; Index < DeviceLayerExtensions.Num(); ++Index)
+			bStandardAvailable = FindLayerInList(DeviceLayerExtensions, STANDARD_VALIDATION_LAYER_NAME);
+			if (bStandardAvailable)
 			{
-				if (!FCStringAnsi::Strcmp(DeviceLayerExtensions[Index].LayerProps.layerName, CurrValidationLayer))
-				{
-					bValidationFound = true;
-					OutDeviceLayers.Add(CurrValidationLayer);
-					break;
-				}
+				OutDeviceLayers.Add(STANDARD_VALIDATION_LAYER_NAME);
 			}
+		}
 
-			if (!bValidationFound)
+		if (!bStandardAvailable)
+		{
+			for (uint32 LayerIndex = 0; GIndividualValidationLayers[LayerIndex] != nullptr; ++LayerIndex)
 			{
-				UE_LOG(LogVulkanRHI, Warning, TEXT("Unable to find Vulkan device validation layer '%s'"), ANSI_TO_TCHAR(CurrValidationLayer));
+				bool bValidationFound = false;
+				const ANSICHAR* CurrValidationLayer = GIndividualValidationLayers[LayerIndex];
+				for (int32 Index = 1; Index < DeviceLayerExtensions.Num(); ++Index)
+				{
+					if (!FCStringAnsi::Strcmp(DeviceLayerExtensions[Index].LayerProps.layerName, CurrValidationLayer))
+					{
+						bValidationFound = true;
+						OutDeviceLayers.Add(CurrValidationLayer);
+						break;
+					}
+				}
+
+				if (!bValidationFound)
+				{
+					UE_LOG(LogVulkanRHI, Warning, TEXT("Unable to find Vulkan device validation layer '%s'"), ANSI_TO_TCHAR(CurrValidationLayer));
+				}
 			}
 		}
 	}
@@ -618,7 +625,7 @@ void FVulkanDevice::GetDeviceExtensionsAndLayers(TArray<const ANSICHAR*>& OutDev
 		}
 	}
 
-#if VULKAN_ENABLE_DRAW_MARKERS
+#if VULKAN_ENABLE_DRAW_MARKERS && VULKAN_HAS_DEBUGGING_ENABLED
 	if (!bOutDebugMarkers &&
 		(((GRenderDocFound || GValidationCvar.GetValueOnAnyThread() == 0) && ListContains(AvailableExtensions, VK_EXT_DEBUG_MARKER_EXTENSION_NAME)) || FVulkanPlatform::ForceEnableDebugMarkers()))
 	{
diff --git a/Engine/Source/Runtime/VulkanRHI/Private/VulkanSwapChain.cpp b/Engine/Source/Runtime/VulkanRHI/Private/VulkanSwapChain.cpp
--- a/Engine/Source/Runtime/VulkanRHI/Private/VulkanSwapChain.cpp
+++ b/Engine/Source/Runtime/VulkanRHI/Private/VulkanSwapChain.cpp
@@ -555,6 +555,7 @@ int32 FVulkanSwapChain::AcquireImageIndex(VulkanRHI::FSemaphore** OutSemaphore)
 	++NumAcquireCalls;
 	*OutSemaphore = ImageAcquiredSemaphore[SemaphoreIndex];
 
+#if VULKAN_HAS_DEBUGGING_ENABLED
 	if (Result == VK_ERROR_VALIDATION_FAILED_EXT)
 	{
 		extern TAutoConsoleVariable<int32> GValidationCvar;
@@ -564,6 +565,7 @@ int32 FVulkanSwapChain::AcquireImageIndex(VulkanRHI::FSemaphore** OutSemaphore)
 		}
 	}
 	else
+#endif
 	{
 		checkf(Result == VK_SUCCESS || Result == VK_SUBOPTIMAL_KHR, TEXT("vkAcquireNextImageKHR failed Result = %d"), int32(Result));
 	}
