UE4.22 - Vulkan async read API
- Fix staging buffer not marked for CPU read
- Use optimized read number of mips code

diff --git a/Engine/Source/Runtime/VulkanRHI/Private/VulkanCommands.cpp b/Engine/Source/Runtime/VulkanRHI/Private/VulkanCommands.cpp
--- a/Engine/Source/Runtime/VulkanRHI/Private/VulkanCommands.cpp
+++ b/Engine/Source/Runtime/VulkanRHI/Private/VulkanCommands.cpp
@@ -1191,7 +1191,7 @@ void FVulkanCommandListContext::RHICopyToStagingBuffer(FVertexBufferRHIParamRef
 			Device->GetStagingManager().ReleaseBuffer(nullptr, StagingBuffer->StagingBuffer);
 		}
 
-		VulkanRHI::FStagingBuffer* ReadbackStagingBuffer = Device->GetStagingManager().AcquireBuffer(NumBytes);
+		VulkanRHI::FStagingBuffer* ReadbackStagingBuffer = Device->GetStagingManager().AcquireBuffer(NumBytes, VK_BUFFER_USAGE_TRANSFER_DST_BIT, true);
 		StagingBuffer->StagingBuffer = ReadbackStagingBuffer;
 		StagingBuffer->Device = Device;
 	}
diff --git a/Engine/Source/Runtime/VulkanRHI/Private/VulkanRHI.cpp b/Engine/Source/Runtime/VulkanRHI/Private/VulkanRHI.cpp
--- a/Engine/Source/Runtime/VulkanRHI/Private/VulkanRHI.cpp
+++ b/Engine/Source/Runtime/VulkanRHI/Private/VulkanRHI.cpp
@@ -663,6 +663,8 @@ void FVulkanDynamicRHI::InitInstance()
 		// Indicate that the RHI needs to use the engine's deferred deletion queue.
 		GRHINeedsExtraDeletionLatency = true;
 
+		GRHISupportsCopyToTextureMultipleMips = true;
+
 		GMaxShadowDepthBufferSizeX =  FPlatformMath::Min<int32>(Props.limits.maxImageDimension2D, GMaxShadowDepthBufferSizeX);
 		GMaxShadowDepthBufferSizeY =  FPlatformMath::Min<int32>(Props.limits.maxImageDimension2D, GMaxShadowDepthBufferSizeY);
 		GMaxTextureDimensions = Props.limits.maxImageDimension2D;
diff --git a/Engine/Source/Runtime/VulkanRHI/Private/VulkanTexture.cpp b/Engine/Source/Runtime/VulkanRHI/Private/VulkanTexture.cpp
--- a/Engine/Source/Runtime/VulkanRHI/Private/VulkanTexture.cpp
+++ b/Engine/Source/Runtime/VulkanRHI/Private/VulkanTexture.cpp
@@ -2190,6 +2190,7 @@ void FVulkanCommandListContext::RHICopyTexture(FTextureRHIParamRef SourceTexture
 	{
 		VkImageSubresourceRange& Range = Barrier.GetSubresource(SourceBarrierIndex);
 		Range.baseMipLevel = CopyInfo.SourceMipIndex;
+		Range.levelCount = CopyInfo.NumMips;
 		Range.baseArrayLayer = CopyInfo.SourceSliceIndex;
 		Range.layerCount = CopyInfo.NumSlices;
 		Barrier.SetTransition(SourceBarrierIndex, VulkanRHI::GetImageLayoutFromVulkanLayout(SrcLayout), EImageLayoutBarrier::TransferSource);
@@ -2197,6 +2198,7 @@ void FVulkanCommandListContext::RHICopyTexture(FTextureRHIParamRef SourceTexture
 	{
 		VkImageSubresourceRange& Range = Barrier.GetSubresource(DestBarrierIndex);
 		Range.baseMipLevel = CopyInfo.DestMipIndex;
+		Range.levelCount = CopyInfo.NumMips;
 		Range.baseArrayLayer = CopyInfo.DestSliceIndex;
 		Range.layerCount = CopyInfo.NumSlices;
 		Barrier.SetTransition(DestBarrierIndex, EImageLayoutBarrier::Undefined, EImageLayoutBarrier::TransferDest);
@@ -2218,10 +2220,18 @@ void FVulkanCommandListContext::RHICopyTexture(FTextureRHIParamRef SourceTexture
 	Region.dstSubresource.baseArrayLayer = CopyInfo.DestSliceIndex;
 	Region.dstSubresource.layerCount = CopyInfo.NumSlices;
 	Region.dstSubresource.mipLevel = CopyInfo.DestMipIndex;
-	VulkanRHI::vkCmdCopyImage(CmdBuffer,
-		SrcSurface.Image, VK_IMAGE_LAYOUT_TRANSFER_SRC_OPTIMAL,
-		DstSurface.Image, VK_IMAGE_LAYOUT_TRANSFER_DST_OPTIMAL,
-		1, &Region);
+
+	for (uint32 Index = 0; Index < CopyInfo.NumMips; ++Index)
+	{
+		VulkanRHI::vkCmdCopyImage(CmdBuffer,
+			SrcSurface.Image, VK_IMAGE_LAYOUT_TRANSFER_SRC_OPTIMAL,
+			DstSurface.Image, VK_IMAGE_LAYOUT_TRANSFER_DST_OPTIMAL,
+			1, &Region);
+		Region.extent.width = FMath::Max(1u, Region.extent.width / 2);
+		Region.extent.height = FMath::Max(1u, Region.extent.height / 2);
+		++Region.srcSubresource.mipLevel;
+		++Region.dstSubresource.mipLevel;
+	}
 
 	Barrier.ResetStages();
 	Barrier.SetTransition(SourceBarrierIndex, EImageLayoutBarrier::TransferSource, VulkanRHI::GetImageLayoutFromVulkanLayout(SrcLayout));
