Prep for Vulkan define changes
#rnx

diff --git a/Engine/Source/Runtime/VulkanRHI/Private/Lumin/VulkanLuminPlatform.h b/Engine/Source/Runtime/VulkanRHI/Private/Lumin/VulkanLuminPlatform.h
--- a/Engine/Source/Runtime/VulkanRHI/Private/Lumin/VulkanLuminPlatform.h
+++ b/Engine/Source/Runtime/VulkanRHI/Private/Lumin/VulkanLuminPlatform.h
@@ -10,7 +10,6 @@
 #define VULKAN_SHOULD_ENABLE_DRAW_MARKERS			(UE_BUILD_DEBUG || UE_BUILD_DEVELOPMENT)
 #define VULKAN_USE_IMAGE_ACQUIRE_FENCES				0
 #define VULKAN_SUPPORTS_COLOR_CONVERSIONS			1
-#define VULKAN_SUPPORTS_GEOMETRY_SHADERS			1
 #define VULKAN_SUPPORTS_NV_DIAGNOSTIC_CHECKPOINT	1
 
 
diff --git a/Engine/Source/Runtime/VulkanRHI/Public/VulkanCommon.h b/Engine/Source/Runtime/VulkanRHI/Public/VulkanCommon.h
--- a/Engine/Source/Runtime/VulkanRHI/Public/VulkanCommon.h
+++ b/Engine/Source/Runtime/VulkanRHI/Public/VulkanCommon.h
@@ -9,7 +9,7 @@
 #include "RHIDefinitions.h"
 
 #ifndef VULKAN_SUPPORTS_GEOMETRY_SHADERS
-	#define VULKAN_SUPPORTS_GEOMETRY_SHADERS					!(PLATFORM_ANDROID) || PLATFORM_LUMIN || PLATFORM_LUMINGL4
+	#define VULKAN_SUPPORTS_GEOMETRY_SHADERS					(!(PLATFORM_ANDROID) || PLATFORM_LUMIN || PLATFORM_LUMINGL4)
 #endif
 
 // This defines controls shader generation (so will cause a format rebuild)
diff --git a/Engine/Source/Runtime/VulkanRHI/Public/VulkanConfiguration.h b/Engine/Source/Runtime/VulkanRHI/Public/VulkanConfiguration.h
--- a/Engine/Source/Runtime/VulkanRHI/Public/VulkanConfiguration.h
+++ b/Engine/Source/Runtime/VulkanRHI/Public/VulkanConfiguration.h
@@ -143,10 +143,6 @@
 	#define VULKAN_USE_CREATE_WIN32_SURFACE						0
 #endif
 
-#ifndef VULKAN_SUPPORTS_GEOMETRY_SHADERS
-	#define VULKAN_SUPPORTS_GEOMETRY_SHADERS					1
-#endif
-
 #ifndef VULKAN_USE_REAL_RENDERPASS_COMPATIBILITY
 	#define VULKAN_USE_REAL_RENDERPASS_COMPATIBILITY			1
 #endif
