UE4.22 - Reduce AMD blacklisted drivers range from 18.2.2 to 19.2.1

diff --git a/Engine/Source/Runtime/VulkanRHI/Private/VulkanSwapChain.cpp b/Engine/Source/Runtime/VulkanRHI/Private/VulkanSwapChain.cpp
--- a/Engine/Source/Runtime/VulkanRHI/Private/VulkanSwapChain.cpp
+++ b/Engine/Source/Runtime/VulkanRHI/Private/VulkanSwapChain.cpp
@@ -251,7 +251,7 @@ FVulkanSwapChain::FVulkanSwapChain(VkInstance InInstance, FVulkanDevice& InDevic
 	if (FVulkanPlatform::SupportsQuerySurfaceProperties())
 	{
 		// Only dump the present modes the very first time they are queried
-		static bool bFirstTimeLog = true;
+		static bool bFirstTimeLog = !!VULKAN_HAS_DEBUGGING_ENABLED;
 
 		uint32 NumFoundPresentModes = 0;
 		VERIFYVULKANRESULT(VulkanRHI::vkGetPhysicalDeviceSurfacePresentModesKHR(Device.GetPhysicalHandle(), Surface, &NumFoundPresentModes, nullptr));
diff --git a/Engine/Source/Runtime/VulkanRHI/Private/Windows/VulkanWindowsPlatform.cpp b/Engine/Source/Runtime/VulkanRHI/Private/Windows/VulkanWindowsPlatform.cpp
--- a/Engine/Source/Runtime/VulkanRHI/Private/Windows/VulkanWindowsPlatform.cpp
+++ b/Engine/Source/Runtime/VulkanRHI/Private/Windows/VulkanWindowsPlatform.cpp
@@ -196,8 +196,8 @@ void FVulkanWindowsPlatform::CheckDeviceDriver(uint32 DeviceIndex)
 			if (DeviceIndex < (uint32)AmdGpuInfo.numDevices && Version && *Version)
 			{
 				auto& DeviceInfo = AmdGpuInfo.devices[DeviceIndex];
-				bool bIsPreGCN = DeviceInfo.architectureVersion == AGSDeviceInfo::ArchitectureVersion_GCN;
-				if (DeviceInfo.architectureVersion == AGSDeviceInfo::ArchitectureVersion_PreGCN || bIsPreGCN)
+				bool bIsPreGCN = DeviceInfo.architectureVersion == AGSDeviceInfo::ArchitectureVersion_PreGCN;
+				if (DeviceInfo.architectureVersion == AGSDeviceInfo::ArchitectureVersion_GCN || bIsPreGCN)
 				{
 					// "Major.Minor.Revision"
 					do
@@ -233,7 +233,7 @@ void FVulkanWindowsPlatform::CheckDeviceDriver(uint32 DeviceIndex)
 							if (MajorVersion < 18)
 							{
 								// Blacklist drivers older than 18.xx.xx drivers
-								FPlatformMisc::MessageBoxExt(EAppMsgType::Ok, TEXT("There are known issues with Vulkan with drivers older than the 18.xx.xx.xx\nfamily of Radeon drivers; the recommended version is 18.12.1.1 or anything more recent than 19.10.xx: please try updating your driver to that version."), TEXT("Vulkan driver version"));
+								FPlatformMisc::MessageBoxExt(EAppMsgType::Ok, TEXT("There are known issues with Vulkan with drivers older than the 18.xx.xx.xx\nfamily of Radeon drivers; the recommended version is 18.12.1.1 or anything more recent starting at 19.2.2: please try updating your driver to that version."), TEXT("Vulkan driver version"));
 								FPlatformMisc::RequestExitWithStatus(true, 1);
 							}
 							else if (WITH_EDITOR)
@@ -241,7 +241,7 @@ void FVulkanWindowsPlatform::CheckDeviceDriver(uint32 DeviceIndex)
 								bool bBadVersion = false;
 								if (MajorVersion == 19)
 								{
-									if (MinorVersion < 10)
+									if (MinorVersion < 2 || (MinorVersion == 2 && RevisionVersion <= 1))
 									{
 										bBadVersion = true;
 									}
@@ -256,8 +256,8 @@ void FVulkanWindowsPlatform::CheckDeviceDriver(uint32 DeviceIndex)
 
 								if (bBadVersion)
 								{
-									// Blacklist 19.xx.xx drivers as it introduced an issue with Slate windows/Vulkan viewports on the editor
-									FPlatformMisc::MessageBoxExt(EAppMsgType::Ok, TEXT("There are known issues with Vulkan on the editor with the some \nRadeon drivers; the recommended version is up to 18.12.1.1 or anything more recent than 19.10.xx: please try updating your driver to that version."), TEXT("Vulkan driver version"));
+									// Blacklist drivers between 18.12.2 and 19.2.1, as they as it introduced an issue with Slate windows/Vulkan viewports on the editor
+									FPlatformMisc::MessageBoxExt(EAppMsgType::Ok, TEXT("There are known issues with Vulkan on the editor with the some \nRadeon drivers; the recommended version is up to 18.12.1.1 or anything more recent starting at 19.2.2: please try updating your driver to that version."), TEXT("Vulkan driver version"));
 									FPlatformMisc::RequestExitWithStatus(true, 1);
 								}
 							}
@@ -266,6 +266,10 @@ void FVulkanWindowsPlatform::CheckDeviceDriver(uint32 DeviceIndex)
 					while (0);
 
 					GRHIDeviceIsAMDPreGCNArchitecture = GRHIDeviceIsAMDPreGCNArchitecture || bIsPreGCN;
+					if (GRHIDeviceIsAMDPreGCNArchitecture)
+					{
+						UE_LOG(LogVulkanRHI, Log, TEXT("AMD Pre GCN architecture detected, some driver workarounds will be in place"));
+					}
 					UE_LOG(LogVulkanRHI, Display, TEXT("AMD User Driver Version = %s"), ANSI_TO_TCHAR(AmdGpuInfo.radeonSoftwareVersion));
 				}
 			}
