UE4.22 - Log missing Vulkan vertex formatsa
#rnx

diff --git a/Engine/Source/Runtime/VulkanRHI/Private/VulkanDevice.cpp b/Engine/Source/Runtime/VulkanRHI/Private/VulkanDevice.cpp
--- a/Engine/Source/Runtime/VulkanRHI/Private/VulkanDevice.cpp
+++ b/Engine/Source/Runtime/VulkanRHI/Private/VulkanDevice.cpp
@@ -655,6 +655,18 @@ void FVulkanDevice::SetupFormats()
 			SetComponentMapping(PF_ETC2_RGBA, VK_COMPONENT_SWIZZLE_R, VK_COMPONENT_SWIZZLE_G, VK_COMPONENT_SWIZZLE_B, VK_COMPONENT_SWIZZLE_A);
 		}
 	}
+
+	// Verify available Vertex Formats
+	static_assert(VET_None == 0, "Change loop below to skip VET_None");
+	for (int32 Index = (int32)VET_None + 1; Index < VET_MAX; ++Index)
+	{
+		EVertexElementType UEType = (EVertexElementType)Index;
+		VkFormat VulkanFormat = UEToVkFormat(UEType);
+		if (!IsFormatSupported(VulkanFormat))
+		{
+			UE_LOG(LogVulkanRHI, Warning, TEXT("EVertexFormat(%d) is not supported with Vk format %d"), (int32)UEType, (int32)VulkanFormat);
+		}
+	}
 }
 
 #if VULKAN_SUPPORTS_COLOR_CONVERSIONS
