UE4.22 - Separate Vulkan texture formats from buffer formats
#rnx

diff --git a/Engine/Source/Runtime/VulkanRHI/Private/VulkanDevice.cpp b/Engine/Source/Runtime/VulkanRHI/Private/VulkanDevice.cpp
--- a/Engine/Source/Runtime/VulkanRHI/Private/VulkanDevice.cpp
+++ b/Engine/Source/Runtime/VulkanRHI/Private/VulkanDevice.cpp
@@ -662,7 +662,7 @@ void FVulkanDevice::SetupFormats()
 	{
 		EVertexElementType UEType = (EVertexElementType)Index;
 		VkFormat VulkanFormat = UEToVkFormat(UEType);
-		if (!IsFormatSupported(VulkanFormat))
+		if (!IsBufferFormatSupported(VulkanFormat))
 		{
 			UE_LOG(LogVulkanRHI, Warning, TEXT("EVertexFormat(%d) is not supported with Vk format %d"), (int32)UEType, (int32)VulkanFormat);
 		}
@@ -692,7 +692,7 @@ void FVulkanDevice::MapFormatSupport(EPixelFormat UEFormat, VkFormat VulkanForma
 {
 	FPixelFormatInfo& FormatInfo = GPixelFormats[UEFormat];
 	FormatInfo.PlatformFormat = VulkanFormat;
-	FormatInfo.Supported = IsFormatSupported(VulkanFormat);
+	FormatInfo.Supported = IsTextureFormatSupported(VulkanFormat);
 
 	if (!FormatInfo.Supported)
 	{
@@ -1023,11 +1023,11 @@ void FVulkanDevice::WaitUntilIdle()
 	GetImmediateContext().GetCommandBufferManager()->RefreshFenceStatus();
 }
 
-bool FVulkanDevice::IsFormatSupported(VkFormat Format) const
+bool FVulkanDevice::IsTextureFormatSupported(VkFormat Format) const
 {
 	auto ArePropertiesSupported = [](const VkFormatProperties& Prop) -> bool
 	{
-		return (Prop.bufferFeatures != 0) || (Prop.linearTilingFeatures != 0) || (Prop.optimalTilingFeatures != 0);
+		return (Prop.linearTilingFeatures != 0) || (Prop.optimalTilingFeatures != 0);
 	};
 
 	if (Format >= 0 && Format < VK_FORMAT_RANGE_SIZE)
@@ -1051,6 +1051,34 @@ bool FVulkanDevice::IsFormatSupported(VkFormat Format) const
 	return ArePropertiesSupported(NewProperties);
 }
 
+bool FVulkanDevice::IsBufferFormatSupported(VkFormat Format) const
+{
+	auto ArePropertiesSupported = [](const VkFormatProperties& Prop) -> bool
+	{
+		return (Prop.bufferFeatures != 0);
+	};
+
+	if (Format >= 0 && Format < VK_FORMAT_RANGE_SIZE)
+	{
+		const VkFormatProperties& Prop = FormatProperties[Format];
+		return Prop.bufferFeatures != 0;
+	}
+
+	// Check for extension formats
+	const VkFormatProperties* FoundProperties = ExtensionFormatProperties.Find(Format);
+	if (FoundProperties)
+	{
+		return FoundProperties->bufferFeatures != 0;
+	}
+
+	// Add it for faster caching next time
+	VkFormatProperties& NewProperties = ExtensionFormatProperties.Add(Format);
+	FMemory::Memzero(NewProperties);
+	VulkanRHI::vkGetPhysicalDeviceFormatProperties(Gpu, Format, &NewProperties);
+
+	return NewProperties.bufferFeatures != 0;
+}
+
 const VkComponentMapping& FVulkanDevice::GetFormatComponentMapping(EPixelFormat UEFormat) const
 {
 	if (UEFormat == PF_X24_G8)
diff --git a/Engine/Source/Runtime/VulkanRHI/Private/VulkanDevice.h b/Engine/Source/Runtime/VulkanRHI/Private/VulkanDevice.h
--- a/Engine/Source/Runtime/VulkanRHI/Private/VulkanDevice.h
+++ b/Engine/Source/Runtime/VulkanRHI/Private/VulkanDevice.h
@@ -143,7 +143,8 @@ public:
 		return TimestampValidBitsMask;
 	}
 
-	bool IsFormatSupported(VkFormat Format) const;
+	bool IsTextureFormatSupported(VkFormat Format) const;
+	bool IsBufferFormatSupported(VkFormat Format) const;
 
 	const VkComponentMapping& GetFormatComponentMapping(EPixelFormat UEFormat) const;
 
