Move surface entry points to a separate macro so they can be skipped for platforms that don't support them.
Remove several missing Vulkan entrypoints on Lumin due to missing instance extensions.

diff --git a/Engine/Source/Runtime/VulkanRHI/Private/Android/VulkanAndroidPlatform.cpp b/Engine/Source/Runtime/VulkanRHI/Private/Android/VulkanAndroidPlatform.cpp
--- a/Engine/Source/Runtime/VulkanRHI/Private/Android/VulkanAndroidPlatform.cpp
+++ b/Engine/Source/Runtime/VulkanRHI/Private/Android/VulkanAndroidPlatform.cpp
@@ -73,6 +73,9 @@ bool FVulkanAndroidPlatform::LoadVulkanInstanceFunctions(VkInstance inInstance)
 	ENUM_VK_ENTRYPOINTS_INSTANCE(GETINSTANCE_VK_ENTRYPOINTS);
 	ENUM_VK_ENTRYPOINTS_INSTANCE(CHECK_VK_ENTRYPOINTS);
 
+	ENUM_VK_ENTRYPOINTS_SURFACE_INSTANCE(GETINSTANCE_VK_ENTRYPOINTS);
+	ENUM_VK_ENTRYPOINTS_SURFACE_INSTANCE(CHECK_VK_ENTRYPOINTS);
+
 	ENUM_VK_ENTRYPOINTS_PLATFORM_INSTANCE(GETINSTANCE_VK_ENTRYPOINTS);
 	ENUM_VK_ENTRYPOINTS_PLATFORM_INSTANCE(CHECK_VK_ENTRYPOINTS);
 
diff --git a/Engine/Source/Runtime/VulkanRHI/Private/Linux/VulkanLinuxPlatform.cpp b/Engine/Source/Runtime/VulkanRHI/Private/Linux/VulkanLinuxPlatform.cpp
--- a/Engine/Source/Runtime/VulkanRHI/Private/Linux/VulkanLinuxPlatform.cpp
+++ b/Engine/Source/Runtime/VulkanRHI/Private/Linux/VulkanLinuxPlatform.cpp
@@ -88,6 +88,8 @@ bool FVulkanLinuxPlatform::LoadVulkanInstanceFunctions(VkInstance inInstance)
 #define GETINSTANCE_VK_ENTRYPOINTS(Type, Func) VulkanDynamicAPI::Func = (Type)VulkanDynamicAPI::vkGetInstanceProcAddr(inInstance, #Func);
 	ENUM_VK_ENTRYPOINTS_INSTANCE(GETINSTANCE_VK_ENTRYPOINTS);
 	ENUM_VK_ENTRYPOINTS_INSTANCE(CHECK_VK_ENTRYPOINTS);
+	ENUM_VK_ENTRYPOINTS_SURFACE_INSTANCE(GETINSTANCE_VK_ENTRYPOINTS);
+	ENUM_VK_ENTRYPOINTS_SURFACE_INSTANCE(CHECK_VK_ENTRYPOINTS);
 
 	if (!bFoundAllEntryPoints)
 	{
diff --git a/Engine/Source/Runtime/VulkanRHI/Private/Lumin/VulkanLuminPlatform.cpp b/Engine/Source/Runtime/VulkanRHI/Private/Lumin/VulkanLuminPlatform.cpp
--- a/Engine/Source/Runtime/VulkanRHI/Private/Lumin/VulkanLuminPlatform.cpp
+++ b/Engine/Source/Runtime/VulkanRHI/Private/Lumin/VulkanLuminPlatform.cpp
@@ -45,11 +45,10 @@ bool FVulkanLuminPlatform::LoadVulkanLibrary()
 		return false;
 	}
 
-	ENUM_VK_ENTRYPOINTS_OPTIONAL_BASE(GET_VK_ENTRYPOINTS);
-	//ENUM_VK_ENTRYPOINTS_OPTIONAL_BASE(CHECK_VK_ENTRYPOINTS);
-
 	ENUM_VK_ENTRYPOINTS_PLATFORM_BASE(GET_VK_ENTRYPOINTS);
-	//ENUM_VK_ENTRYPOINTS_PLATFORM_BASE(CHECK_VK_ENTRYPOINTS);
+#if UE_BUILD_DEBUG
+	ENUM_VK_ENTRYPOINTS_PLATFORM_BASE(CHECK_VK_ENTRYPOINTS);
+#endif
 
 #undef GET_VK_ENTRYPOINTS
 	return true;
@@ -64,9 +63,6 @@ bool FVulkanLuminPlatform::LoadVulkanInstanceFunctions(VkInstance inInstance)
 	ENUM_VK_ENTRYPOINTS_INSTANCE(GETINSTANCE_VK_ENTRYPOINTS);
 	ENUM_VK_ENTRYPOINTS_INSTANCE(CHECK_VK_ENTRYPOINTS);
 
-	ENUM_VK_ENTRYPOINTS_OPTIONAL_INSTANCE(GETINSTANCE_VK_ENTRYPOINTS);
-	ENUM_VK_ENTRYPOINTS_OPTIONAL_INSTANCE(CHECK_VK_ENTRYPOINTS);
-
 	ENUM_VK_ENTRYPOINTS_PLATFORM_INSTANCE(GETINSTANCE_VK_ENTRYPOINTS);
 	ENUM_VK_ENTRYPOINTS_PLATFORM_INSTANCE(CHECK_VK_ENTRYPOINTS);
 
diff --git a/Engine/Source/Runtime/VulkanRHI/Private/Lumin/VulkanLuminPlatform.h b/Engine/Source/Runtime/VulkanRHI/Private/Lumin/VulkanLuminPlatform.h
--- a/Engine/Source/Runtime/VulkanRHI/Private/Lumin/VulkanLuminPlatform.h
+++ b/Engine/Source/Runtime/VulkanRHI/Private/Lumin/VulkanLuminPlatform.h
@@ -10,14 +10,12 @@
 #define VULKAN_SHOULD_ENABLE_DRAW_MARKERS			(UE_BUILD_DEBUG || UE_BUILD_DEVELOPMENT)
 #define VULKAN_USE_IMAGE_ACQUIRE_FENCES				0
 #define VULKAN_SUPPORTS_COLOR_CONVERSIONS			1
-#define VULKAN_SUPPORTS_NV_DIAGNOSTIC_CHECKPOINT	1
+#define VULKAN_SUPPORTS_NV_DIAGNOSTIC_CHECKPOINT	0
 
 
 #define ENUM_VK_ENTRYPOINTS_PLATFORM_BASE(EnumMacro)
 
 #define ENUM_VK_ENTRYPOINTS_PLATFORM_INSTANCE(EnumMacro) \
-	EnumMacro(PFN_vkCmdSetCheckpointNV, vkCmdSetCheckpointNV) \
-	EnumMacro(PFN_vkGetQueueCheckpointDataNV, vkGetQueueCheckpointDataNV) \
 	EnumMacro(PFN_vkCreateSamplerYcbcrConversionKHR, vkCreateSamplerYcbcrConversionKHR) \
 	EnumMacro(PFN_vkDestroySamplerYcbcrConversionKHR, vkDestroySamplerYcbcrConversionKHR)
 
diff --git a/Engine/Source/Runtime/VulkanRHI/Private/VulkanLoader.h b/Engine/Source/Runtime/VulkanRHI/Private/VulkanLoader.h
--- a/Engine/Source/Runtime/VulkanRHI/Private/VulkanLoader.h
+++ b/Engine/Source/Runtime/VulkanRHI/Private/VulkanLoader.h
@@ -139,17 +139,20 @@
 	EnumMacro(PFN_vkCmdNextSubpass, vkCmdNextSubpass) \
 	EnumMacro(PFN_vkCmdEndRenderPass, vkCmdEndRenderPass) \
 	EnumMacro(PFN_vkCmdExecuteCommands, vkCmdExecuteCommands) \
-	EnumMacro(PFN_vkDestroySurfaceKHR, vkDestroySurfaceKHR) \
-	EnumMacro(PFN_vkGetPhysicalDeviceSurfaceSupportKHR, vkGetPhysicalDeviceSurfaceSupportKHR) \
-	EnumMacro(PFN_vkGetPhysicalDeviceSurfaceCapabilitiesKHR, vkGetPhysicalDeviceSurfaceCapabilitiesKHR) \
-	EnumMacro(PFN_vkGetPhysicalDeviceSurfaceFormatsKHR, vkGetPhysicalDeviceSurfaceFormatsKHR) \
-	EnumMacro(PFN_vkGetPhysicalDeviceSurfacePresentModesKHR, vkGetPhysicalDeviceSurfacePresentModesKHR) \
 	EnumMacro(PFN_vkCreateSwapchainKHR, vkCreateSwapchainKHR) \
 	EnumMacro(PFN_vkDestroySwapchainKHR, vkDestroySwapchainKHR) \
 	EnumMacro(PFN_vkGetSwapchainImagesKHR, vkGetSwapchainImagesKHR) \
 	EnumMacro(PFN_vkAcquireNextImageKHR, vkAcquireNextImageKHR) \
 	EnumMacro(PFN_vkQueuePresentKHR, vkQueuePresentKHR)
 
+// List all surface Vulkan entry points used by Unreal that need to be loaded manually
+#define ENUM_VK_ENTRYPOINTS_SURFACE_INSTANCE(EnumMacro) \
+	EnumMacro(PFN_vkDestroySurfaceKHR, vkDestroySurfaceKHR) \
+	EnumMacro(PFN_vkGetPhysicalDeviceSurfaceSupportKHR, vkGetPhysicalDeviceSurfaceSupportKHR) \
+	EnumMacro(PFN_vkGetPhysicalDeviceSurfaceCapabilitiesKHR, vkGetPhysicalDeviceSurfaceCapabilitiesKHR) \
+	EnumMacro(PFN_vkGetPhysicalDeviceSurfaceFormatsKHR, vkGetPhysicalDeviceSurfaceFormatsKHR) \
+	EnumMacro(PFN_vkGetPhysicalDeviceSurfacePresentModesKHR, vkGetPhysicalDeviceSurfacePresentModesKHR)
+
 // List all base Vulkan entry points used by Unreal that need to be loaded manually
 #define ENUM_VK_ENTRYPOINTS_BASE(EnumMacro) \
 	EnumMacro(PFN_vkCreateInstance, vkCreateInstance) \
@@ -176,6 +179,7 @@
 	ENUM_VK_ENTRYPOINTS_BASE(EnumMacro) \
 	ENUM_VK_ENTRYPOINTS_OPTIONAL_BASE(EnumMacro) \
 	ENUM_VK_ENTRYPOINTS_INSTANCE(EnumMacro) \
+	ENUM_VK_ENTRYPOINTS_SURFACE_INSTANCE(EnumMacro) \
 	ENUM_VK_ENTRYPOINTS_OPTIONAL_INSTANCE(EnumMacro) \
 	ENUM_VK_ENTRYPOINTS_PLATFORM_BASE(EnumMacro) \
 	ENUM_VK_ENTRYPOINTS_OPTIONAL_PLATFORM_INSTANCE(EnumMacro) \
diff --git a/Engine/Source/Runtime/VulkanRHI/Private/Windows/VulkanWindowsPlatform.cpp b/Engine/Source/Runtime/VulkanRHI/Private/Windows/VulkanWindowsPlatform.cpp
--- a/Engine/Source/Runtime/VulkanRHI/Private/Windows/VulkanWindowsPlatform.cpp
+++ b/Engine/Source/Runtime/VulkanRHI/Private/Windows/VulkanWindowsPlatform.cpp
@@ -73,6 +73,8 @@ bool FVulkanWindowsPlatform::LoadVulkanInstanceFunctions(VkInstance inInstance)
 #define GETINSTANCE_VK_ENTRYPOINTS(Type, Func) VulkanDynamicAPI::Func = (Type)VulkanDynamicAPI::vkGetInstanceProcAddr(inInstance, #Func);
 	ENUM_VK_ENTRYPOINTS_INSTANCE(GETINSTANCE_VK_ENTRYPOINTS);
 	ENUM_VK_ENTRYPOINTS_INSTANCE(CHECK_VK_ENTRYPOINTS);
+	ENUM_VK_ENTRYPOINTS_SURFACE_INSTANCE(GETINSTANCE_VK_ENTRYPOINTS);
+	ENUM_VK_ENTRYPOINTS_SURFACE_INSTANCE(CHECK_VK_ENTRYPOINTS);
 	if (!bFoundAllEntryPoints)
 	{
 		FreeVulkanLibrary();
