Vulkan: fix validation issues when applying the workaround.

diff --git a/Engine/Source/Runtime/VulkanRHI/Private/VulkanQuery.cpp b/Engine/Source/Runtime/VulkanRHI/Private/VulkanQuery.cpp
--- a/Engine/Source/Runtime/VulkanRHI/Private/VulkanQuery.cpp
+++ b/Engine/Source/Runtime/VulkanRHI/Private/VulkanQuery.cpp
@@ -404,9 +404,9 @@ void FVulkanOcclusionQueryPool::Reset(FVulkanCmdBuffer* InCmdBuffer, uint32 InFr
 	{
 		VkMemoryBarrier barrier;
 		ZeroVulkanStruct(barrier, VK_STRUCTURE_TYPE_MEMORY_BARRIER);
-		barrier.srcAccessMask = VK_ACCESS_SHADER_WRITE_BIT;
-		barrier.dstAccessMask = VK_ACCESS_SHADER_READ_BIT;
-		VulkanRHI::vkCmdPipelineBarrier(InCmdBuffer->GetHandle(), VK_PIPELINE_STAGE_COMPUTE_SHADER_BIT, VK_PIPELINE_STAGE_BOTTOM_OF_PIPE_BIT, 0, 1, &barrier, 0, 0, 0, 0);
+		barrier.srcAccessMask = VK_ACCESS_TRANSFER_WRITE_BIT;
+		barrier.dstAccessMask = VK_ACCESS_HOST_READ_BIT;
+		VulkanRHI::vkCmdPipelineBarrier(InCmdBuffer->GetHandle(), VK_PIPELINE_STAGE_TRANSFER_BIT, VK_PIPELINE_STAGE_HOST_BIT, 0, 1, &barrier, 0, 0, 0, 0);
 	}
 	
 	VulkanRHI::vkCmdSetEvent(InCmdBuffer->GetHandle(), ResetEvent, VK_PIPELINE_STAGE_BOTTOM_OF_PIPE_BIT);
