Fix for FVulkanDynamicRHI::RHIReadSurfaceFloatData disregarding ArrayIndex. Fixes readback of textures on vulkan, which is used when cooking reflection captures.

diff --git a/Engine/Source/Runtime/VulkanRHI/Private/VulkanRenderTarget.cpp b/Engine/Source/Runtime/VulkanRHI/Private/VulkanRenderTarget.cpp
--- a/Engine/Source/Runtime/VulkanRHI/Private/VulkanRenderTarget.cpp
+++ b/Engine/Source/Runtime/VulkanRHI/Private/VulkanRenderTarget.cpp
@@ -1136,14 +1136,14 @@ void FVulkanDynamicRHI::RHIReadSurfaceFloatData(FTextureRHIParamRef TextureRHI,
 		{
 			FRHITextureCube* TextureRHICube = TextureRHI->GetTextureCube();
 			FVulkanTextureCube* TextureCube = (FVulkanTextureCube*)TextureRHICube;
-			DoCopyFloat(Device, CmdBuffer, TextureCube->Surface, MipIndex, CubeFace, Rect, OutData);
+			DoCopyFloat(Device, CmdBuffer, TextureCube->Surface, MipIndex, CubeFace + 6 * ArrayIndex, Rect, OutData);
 		}
 		else
 		{
 			FRHITexture2D* TextureRHI2D = TextureRHI->GetTexture2D();
 			check(TextureRHI2D);
 			FVulkanTexture2D* Texture2D = (FVulkanTexture2D*)TextureRHI2D;
-			DoCopyFloat(Device, CmdBuffer, Texture2D->Surface, MipIndex, 0, Rect, OutData);
+			DoCopyFloat(Device, CmdBuffer, Texture2D->Surface, MipIndex, ArrayIndex, Rect, OutData);
 		}
 		Device->GetImmediateContext().GetCommandBufferManager()->PrepareForNewActiveCommandBuffer();
 	}
