UE4.22 - Tessellation on Vulkan (experimental, disabled)
#rnx

diff --git a/Engine/Shaders/Private/BasePassCommon.ush b/Engine/Shaders/Private/BasePassCommon.ush
--- a/Engine/Shaders/Private/BasePassCommon.ush
+++ b/Engine/Shaders/Private/BasePassCommon.ush
@@ -104,7 +104,7 @@ struct FSharedBasePassInterpolants
 #endif
 };
 
-#if PC_D3D || GL4_PROFILE || SWITCH_PROFILE || SWITCH_PROFILE_FORWARD || METAL_TESSELLATION_PROFILE
+#if PC_D3D || GL4_PROFILE || SWITCH_PROFILE || SWITCH_PROFILE_FORWARD || METAL_TESSELLATION_PROFILE || VULKAN_ENABLE_TESSELLATION
 
 /** Interpolants passed from the vertex shader to the pixel shader. */
 struct FBasePassInterpolantsVSToPS : FSharedBasePassInterpolants
diff --git a/Engine/Shaders/Public/Platform.ush b/Engine/Shaders/Public/Platform.ush
--- a/Engine/Shaders/Public/Platform.ush
+++ b/Engine/Shaders/Public/Platform.ush
@@ -131,6 +131,10 @@
 #define	VULKAN_PROFILE_SM5 0
 #endif
 
+#ifndef VULKAN_ENABLE_TESSELLATION
+#define VULKAN_ENABLE_TESSELLATION	0
+#endif
+
 #ifndef IOS
 #define IOS 0
 #endif
diff --git a/Engine/Shaders/Public/Platform/Vulkan/VulkanCommon.ush b/Engine/Shaders/Public/Platform/Vulkan/VulkanCommon.ush
--- a/Engine/Shaders/Public/Platform/Vulkan/VulkanCommon.ush
+++ b/Engine/Shaders/Public/Platform/Vulkan/VulkanCommon.ush
@@ -5,7 +5,7 @@
 =============================================================================*/
 
 // Update this GUID to improve shader recompilation for Vulkan only shaders
-// GUID = {792BDC0A-A4A8-402F-8752-000000000041}
+// GUID = {792BDC0A-A4A8-402F-8752-000000000042}
 // Glslang 1.1.82.1
 
 #pragma once
diff --git a/Engine/Source/Developer/ShaderFormatOpenGL/Private/GlslBackend.cpp b/Engine/Source/Developer/ShaderFormatOpenGL/Private/GlslBackend.cpp
--- a/Engine/Source/Developer/ShaderFormatOpenGL/Private/GlslBackend.cpp
+++ b/Engine/Source/Developer/ShaderFormatOpenGL/Private/GlslBackend.cpp
@@ -4709,9 +4709,7 @@ void GenShaderOutputForVariable(
 	FSemanticQualifier OutputQualifier,
 	ir_dereference* OutputVariableDeref,
 	exec_list* DeclInstructions,
-	exec_list* PostCallInstructions,
-	int SemanticArraySize,
-	int SemanticArrayIndex
+	exec_list* PostCallInstructions
 	)
 {
 	const glsl_type* OutputType = OutputVariableDeref->type;
@@ -4771,9 +4769,7 @@ void GenShaderOutputForVariable(
 					Qualifier,
 					FieldDeref,
 					DeclInstructions,
-					PostCallInstructions,
-					SemanticArraySize,
-					SemanticArrayIndex
+					PostCallInstructions
 					);
 			}
 			else
@@ -4812,9 +4808,7 @@ void GenShaderOutputForVariable(
 					OutputQualifier,
 					ArrayDeref,
 					DeclInstructions,
-					PostCallInstructions,
-					SemanticArraySize,
-					SemanticArrayIndex
+					PostCallInstructions
 					);
 			}
 		}
@@ -4930,9 +4924,7 @@ static ir_dereference_variable* GenShaderOutput(
 		OutputQualifier,
 		TempVariableDeref,
 		DeclInstructions,
-		PostCallInstructions,
-		0,
-		0
+		PostCallInstructions
 		);
 	return TempVariableDeref;
 }
@@ -4975,9 +4967,7 @@ static void GenerateAppendFunctionBody(
 		OutputQualifier,
 		TempVariableDeref,
 		DeclInstructions,
-		&sig->body,
-		0,
-		0
+		&sig->body
 		);
 
 	// If the output structure type contains a SV_RenderTargetArrayIndex semantic, add a custom user output semantic.
@@ -5021,9 +5011,7 @@ static void GenerateAppendFunctionBody(
 			Qualifier,
 			new(ParseState)ir_dereference_record(var, new_output_type->fields.structure[indexof_RenderTargetArrayIndex].name),
 			DeclInstructions,
-			&sig->body,
-			0,
-			0
+			&sig->body
 			);
 	}
 
diff --git a/Engine/Source/Developer/VulkanShaderFormat/Private/VulkanBackend.cpp b/Engine/Source/Developer/VulkanShaderFormat/Private/VulkanBackend.cpp
--- a/Engine/Source/Developer/VulkanShaderFormat/Private/VulkanBackend.cpp
+++ b/Engine/Source/Developer/VulkanShaderFormat/Private/VulkanBackend.cpp
@@ -4653,16 +4653,21 @@ static ir_rvalue* GenShaderOutputSemantic(
 
 	*DestVariableType = Type;
 
-	// This code-section replacces "layout(location=0) out struct { vec4 Data; } out_TEXCOORD0;" pattern to
+	if (Frequency == HSF_HullShader && !OutputQualifier.Fields.bIsPatchConstant)
+	{
+		Type = glsl_type::get_array_instance(Type, ParseState->tessellation.outputcontrolpoints);
+	}
+
+	// This code-section replaces "layout(location=0) out struct { vec4 Data; } out_TEXCOORD0;" pattern to
 	// "layout(location=0) out vec4 out_TEXCOORD0;".
 
 	// Regular attribute
 	Variable = new(ParseState)ir_variable(
 		Type,
-		ralloc_asprintf(ParseState, "%s_%s", "out", Semantic),
+		ralloc_asprintf(ParseState, "out_%s", Semantic),
 		ir_var_out
 		);
-	Variable->read_only = true;
+	//Variable->read_only = true;
 	Variable->centroid = OutputQualifier.Fields.bCentroid;
 	Variable->interpolation = OutputQualifier.Fields.InterpolationMode;
 	Variable->is_patch_constant = OutputQualifier.Fields.bIsPatchConstant;
@@ -4675,7 +4680,13 @@ static ir_rvalue* GenShaderOutputSemantic(
 	DeclInstructions->push_tail(Variable);
 	ParseState->symbols->add_variable(Variable);
 
-	ir_dereference_variable* VariableDeref = new(ParseState)ir_dereference_variable(Variable);
+	ir_dereference* VariableDeref = new(ParseState)ir_dereference_variable(Variable);
+
+	if (Frequency == HSF_HullShader && !OutputQualifier.Fields.bIsPatchConstant)
+	{
+		VariableDeref = new(ParseState)ir_dereference_array(VariableDeref, new(ParseState)ir_dereference_variable(ParseState->symbols->get_variable("gl_InvocationID")));
+	}
+
 	return VariableDeref;
 }
 
@@ -4997,9 +5008,7 @@ static void GenShaderOutputForVariable(
 	FSemanticQualifier OutputQualifier,
 	ir_dereference* OutputVariableDeref,
 	exec_list* DeclInstructions,
-	exec_list* PostCallInstructions,
-	int SemanticArraySize,
-	int SemanticArrayIndex
+	exec_list* PostCallInstructions
 	)
 {
 	const glsl_type* OutputType = OutputVariableDeref->type;
@@ -5059,9 +5068,7 @@ static void GenShaderOutputForVariable(
 					Qualifier,
 					FieldDeref,
 					DeclInstructions,
-					PostCallInstructions,
-					SemanticArraySize,
-					SemanticArrayIndex
+					PostCallInstructions
 					);
 			}
 			else
@@ -5100,9 +5107,7 @@ static void GenShaderOutputForVariable(
 					OutputQualifier,
 					ArrayDeref,
 					DeclInstructions,
-					PostCallInstructions,
-					SemanticArraySize,
-					SemanticArrayIndex
+					PostCallInstructions
 					);
 			}
 		}
@@ -5189,7 +5194,7 @@ static void GenShaderOutputForVariable(
 * @param OutputQualifier - Qualifiers applied to the semantic.
 * @param OutputType - Value type.
 * @param DeclInstructions - IR to which declarations may be added.
-* @param PreCallInstructions - IR to which isntructions may be added before the
+* @param PreCallInstructions - IR to which instructions may be added before the
 entry point is called.
 * @param PostCallInstructions - IR to which instructions may be added after the
 *                               entry point returns.
@@ -5220,10 +5225,7 @@ static ir_dereference_variable* GenShaderOutput(
 		OutputQualifier,
 		TempVariableDeref,
 		DeclInstructions,
-		PostCallInstructions,
-		0,
-		0
-		);
+		PostCallInstructions);
 	return TempVariableDeref;
 }
 
@@ -5265,10 +5267,7 @@ static void GenerateAppendFunctionBody(
 		OutputQualifier,
 		TempVariableDeref,
 		DeclInstructions,
-		&sig->body,
-		0,
-		0
-		);
+		&sig->body);
 
 	// If the output structure type contains a SV_RenderTargetArrayIndex semantic, add a custom user output semantic.
 	// It's used to pass layer index to pixel shader, as GLSL 1.50 doesn't allow pixel shader to read from gl_Layer.
diff --git a/Engine/Source/Developer/VulkanShaderFormat/Private/VulkanShaderCompiler.cpp b/Engine/Source/Developer/VulkanShaderFormat/Private/VulkanShaderCompiler.cpp
--- a/Engine/Source/Developer/VulkanShaderFormat/Private/VulkanShaderCompiler.cpp
+++ b/Engine/Source/Developer/VulkanShaderFormat/Private/VulkanShaderCompiler.cpp
@@ -2018,6 +2018,7 @@ void DoCompileVulkanShader(const FShaderCompilerInput& Input, FShaderCompilerOut
 
 	TArray<ANSICHAR> GeneratedGlslSource;
 	FVulkanBindingTable BindingTable(CompilerInfo.Frequency);
+	bool bSuccess = false;
 	if (CallHlslcc(PreprocessedShaderSource, BindingTable, CompilerInfo, EntryPointName, HlslCompilerTarget, Output, GeneratedGlslSource))
 	{
 		//#todo-rco: Once it's all cleaned up...
@@ -2030,7 +2031,7 @@ void DoCompileVulkanShader(const FShaderCompilerInput& Input, FShaderCompilerOut
 			// For debugging: if you hit an error from Glslang/Spirv, use the SourceNoHeader for line numbers
 			auto* SourceWithHeader = GeneratedGlslSource.GetData();
 			char* SourceNoHeader = strstr(SourceWithHeader, "#version");
-			bool bSuccess = CompileUsingInternal(CompilerInfo, BindingTable, GeneratedGlslSource, EntryPointName, Output, bHasRealUBs);
+			bSuccess = CompileUsingInternal(CompilerInfo, BindingTable, GeneratedGlslSource, EntryPointName, Output, bHasRealUBs);
 			if (bDirectCompile)
 			{
 				FPlatformMisc::LowLevelOutputDebugStringf(TEXT("Success: %d\n%s\n"), bSuccess, ANSI_TO_TCHAR(SourceWithHeader));
@@ -2044,5 +2045,6 @@ void DoCompileVulkanShader(const FShaderCompilerInput& Input, FShaderCompilerOut
 		{
 			FPlatformMisc::LowLevelOutputDebugStringf(TEXT("%s\n"), *Error.GetErrorString());
 		}
+		ensure(bSuccess);
 	}
 }
diff --git a/Engine/Source/Runtime/Engine/Private/ShaderCompiler/ShaderCompiler.cpp b/Engine/Source/Runtime/Engine/Private/ShaderCompiler/ShaderCompiler.cpp
--- a/Engine/Source/Runtime/Engine/Private/ShaderCompiler/ShaderCompiler.cpp
+++ b/Engine/Source/Runtime/Engine/Private/ShaderCompiler/ShaderCompiler.cpp
@@ -3104,7 +3104,16 @@ void GlobalBeginCompileShader(
 		{
 			Input.Environment.CompilerFlags.Add(CFLAG_ForceRemoveUnusedInterpolators);
 		}
-	} 
+	}
+
+	if (IsVulkanPlatform((EShaderPlatform)Target.Platform))
+	{
+		const auto* CVar = IConsoleManager::Get().FindConsoleVariable(TEXT("r.Vulkan.EnableTessellation"));
+		if (CVar && CVar->GetInt() != 0)
+		{
+			Input.Environment.SetDefine(TEXT("VULKAN_ENABLE_TESSELLATION"), 1);
+		}
+	}
 	
 	if (IsMetalPlatform((EShaderPlatform)Target.Platform))
 	{
diff --git a/Engine/Source/Runtime/RHI/Private/RHI.cpp b/Engine/Source/Runtime/RHI/Private/RHI.cpp
--- a/Engine/Source/Runtime/RHI/Private/RHI.cpp
+++ b/Engine/Source/Runtime/RHI/Private/RHI.cpp
@@ -42,6 +42,14 @@ static FAutoConsoleVariable CVarUseVulkanRealUBs(
 	ECVF_ReadOnly
 	);
 
+static FAutoConsoleVariable CVarVulkanEnableTessellation(
+	TEXT("r.Vulkan.EnableTessellation"),
+	0,
+	TEXT("0: Tessellation disabled[default]\n")
+	TEXT("1: Enable flat tessellation (experimental)"),
+	ECVF_ReadOnly
+);
+
 static TAutoConsoleVariable<int32> CVarDisableEngineAndAppRegistration(
 	TEXT("r.DisableEngineAndAppRegistration"),
 	0,
@@ -796,7 +804,7 @@ RHI_API bool RHISupportsTessellation(const EShaderPlatform Platform)
 {
 	if (IsFeatureLevelSupported(Platform, ERHIFeatureLevel::SM5))
 	{
-		return (Platform == SP_PCD3D_SM5) || (Platform == SP_XBOXONE_D3D12) || (Platform == SP_OPENGL_SM5) || (Platform == SP_OPENGL_ES31_EXT) || (Platform == SP_METAL_SM5) /* || (IsVulkanSM5Platform(Platform)*/;
+		return (Platform == SP_PCD3D_SM5) || (Platform == SP_XBOXONE_D3D12) || (Platform == SP_OPENGL_SM5) || (Platform == SP_OPENGL_ES31_EXT) || (Platform == SP_METAL_SM5) || (IsVulkanSM5Platform(Platform) && CVarVulkanEnableTessellation->GetInt() != 0);
 	}
 	return false;
 }
diff --git a/Engine/Source/Runtime/VulkanRHI/Private/VulkanCommands.cpp b/Engine/Source/Runtime/VulkanRHI/Private/VulkanCommands.cpp
--- a/Engine/Source/Runtime/VulkanRHI/Private/VulkanCommands.cpp
+++ b/Engine/Source/Runtime/VulkanRHI/Private/VulkanCommands.cpp
@@ -188,26 +188,20 @@ void FVulkanCommandListContext::RHISetShaderTexture(FVertexShaderRHIParamRef Ver
 
 void FVulkanCommandListContext::RHISetShaderTexture(FHullShaderRHIParamRef HullShaderRHI, uint32 TextureIndex, FTextureRHIParamRef NewTextureRHI)
 {
-	ensureMsgf(0, TEXT("Tessellation not supported yet!"));
-/*
 	check(PendingGfxState->GetCurrentShaderKey(ShaderStage::Hull) == GetShaderKey(HullShaderRHI));
 	FVulkanTextureBase* Texture = GetVulkanTextureFromRHITexture(NewTextureRHI);
 	VkImageLayout Layout = GetLayoutForDescriptor(Texture->Surface);
-	PendingGfxState->SetTexture(ShaderStage::Hull, TextureIndex, Texture, Layout);
+	PendingGfxState->SetTextureForStage(ShaderStage::Hull, TextureIndex, Texture, Layout);
 	NewTextureRHI->SetLastRenderTime((float)FPlatformTime::Seconds());
-*/
 }
 
 void FVulkanCommandListContext::RHISetShaderTexture(FDomainShaderRHIParamRef DomainShaderRHI, uint32 TextureIndex, FTextureRHIParamRef NewTextureRHI)
 {
-	ensureMsgf(0, TEXT("Tessellation not supported yet!"));
-/*
 	check(PendingGfxState->GetCurrentShaderKey(ShaderStage::Domain) == GetShaderKey(DomainShaderRHI));
 	FVulkanTextureBase* Texture = GetVulkanTextureFromRHITexture(NewTextureRHI);
 	VkImageLayout Layout = GetLayoutForDescriptor(Texture->Surface);
-	PendingGfxState->SetTexture(ShaderStage::Domain, TextureIndex, Texture, Layout);
+	PendingGfxState->SetTextureForStage(ShaderStage::Domain, TextureIndex, Texture, Layout);
 	NewTextureRHI->SetLastRenderTime((float)FPlatformTime::Seconds());
-*/
 }
 
 void FVulkanCommandListContext::RHISetShaderTexture(FGeometryShaderRHIParamRef GeometryShaderRHI, uint32 TextureIndex, FTextureRHIParamRef NewTextureRHI)
@@ -252,22 +246,16 @@ void FVulkanCommandListContext::RHISetShaderResourceViewParameter(FVertexShaderR
 
 void FVulkanCommandListContext::RHISetShaderResourceViewParameter(FHullShaderRHIParamRef HullShaderRHI,uint32 TextureIndex,FShaderResourceViewRHIParamRef SRVRHI)
 {
-	ensureMsgf(0, TEXT("Tessellation not supported yet!"));
-/*
 	check(PendingGfxState->GetCurrentShaderKey(ShaderStage::Hull) == GetShaderKey(HullShaderRHI));
 	FVulkanShaderResourceView* SRV = ResourceCast(SRVRHI);
-	PendingGfxState->SetSRV(ShaderStage::Hull, TextureIndex, SRV);
-*/
+	PendingGfxState->SetSRVForStage(ShaderStage::Hull, TextureIndex, SRV);
 }
 
 void FVulkanCommandListContext::RHISetShaderResourceViewParameter(FDomainShaderRHIParamRef DomainShaderRHI,uint32 TextureIndex,FShaderResourceViewRHIParamRef SRVRHI)
 {
-	ensureMsgf(0, TEXT("Tessellation not supported yet!"));
-/*
 	check(PendingGfxState->GetCurrentShaderKey(ShaderStage::Domain) == GetShaderKey(DomainShaderRHI));
 	FVulkanShaderResourceView* SRV = ResourceCast(SRVRHI);
-	PendingGfxState->SetSRV(ShaderStage::Domain, TextureIndex, SRV);
-*/
+	PendingGfxState->SetSRVForStage(ShaderStage::Domain, TextureIndex, SRV);
 }
 
 void FVulkanCommandListContext::RHISetShaderResourceViewParameter(FGeometryShaderRHIParamRef GeometryShaderRHI,uint32 TextureIndex,FShaderResourceViewRHIParamRef SRVRHI)
@@ -305,22 +293,16 @@ void FVulkanCommandListContext::RHISetShaderSampler(FVertexShaderRHIParamRef Ver
 
 void FVulkanCommandListContext::RHISetShaderSampler(FHullShaderRHIParamRef HullShaderRHI, uint32 SamplerIndex, FSamplerStateRHIParamRef NewStateRHI)
 {
-	ensureMsgf(0, TEXT("Tessellation not supported yet!"));
-/*
 	check(PendingGfxState->GetCurrentShaderKey(ShaderStage::Hull) == GetShaderKey(HullShaderRHI));
 	FVulkanSamplerState* Sampler = ResourceCast(NewStateRHI);
-	PendingGfxState->SetSamplerState(ShaderStage::Hull, SamplerIndex, Sampler);
-*/
+	PendingGfxState->SetSamplerStateForStage(ShaderStage::Hull, SamplerIndex, Sampler);
 }
 
 void FVulkanCommandListContext::RHISetShaderSampler(FDomainShaderRHIParamRef DomainShaderRHI, uint32 SamplerIndex, FSamplerStateRHIParamRef NewStateRHI)
 {
-	ensureMsgf(0, TEXT("Tessellation not supported yet!"));
-/*
 	check(PendingGfxState->GetCurrentShaderKey(ShaderStage::Domain) == GetShaderKey(DomainShaderRHI));
 	FVulkanSamplerState* Sampler = ResourceCast(NewStateRHI);
-	PendingGfxState->SetSamplerState(ShaderStage::Domain, SamplerIndex, Sampler);
-*/
+	PendingGfxState->SetSamplerStateForStage(ShaderStage::Domain, SamplerIndex, Sampler);
 }
 
 void FVulkanCommandListContext::RHISetShaderSampler(FGeometryShaderRHIParamRef GeometryShaderRHI, uint32 SamplerIndex, FSamplerStateRHIParamRef NewStateRHI)
@@ -359,22 +341,16 @@ void FVulkanCommandListContext::RHISetShaderParameter(FVertexShaderRHIParamRef V
 
 void FVulkanCommandListContext::RHISetShaderParameter(FHullShaderRHIParamRef HullShaderRHI, uint32 BufferIndex, uint32 BaseIndex, uint32 NumBytes, const void* NewValue)
 {
-	ensureMsgf(0, TEXT("Tessellation not supported yet!"));
-/*
 	check(PendingGfxState->GetCurrentShaderKey(ShaderStage::Hull) == GetShaderKey(HullShaderRHI));
 
-	PendingGfxState->SetShaderParameter(ShaderStage::Hull, BufferIndex, BaseIndex, NumBytes, NewValue);
-*/
+	PendingGfxState->SetPackedGlobalShaderParameter(ShaderStage::Hull, BufferIndex, BaseIndex, NumBytes, NewValue);
 }
 
 void FVulkanCommandListContext::RHISetShaderParameter(FDomainShaderRHIParamRef DomainShaderRHI, uint32 BufferIndex, uint32 BaseIndex, uint32 NumBytes, const void* NewValue)
 {
-	ensureMsgf(0, TEXT("Tessellation not supported yet!"));
-/*
 	check(PendingGfxState->GetCurrentShaderKey(ShaderStage::Domain) == GetShaderKey(DomainShaderRHI));
 
-	PendingGfxState->SetShaderParameter(ShaderStage::Domain, BufferIndex, BaseIndex, NumBytes, NewValue);
-*/
+	PendingGfxState->SetPackedGlobalShaderParameter(ShaderStage::Domain, BufferIndex, BaseIndex, NumBytes, NewValue);
 }
 
 void FVulkanCommandListContext::RHISetShaderParameter(FGeometryShaderRHIParamRef GeometryShaderRHI, uint32 BufferIndex, uint32 BaseIndex, uint32 NumBytes, const void* NewValue)
@@ -646,20 +622,14 @@ void FVulkanCommandListContext::RHISetShaderUniformBuffer(FVertexShaderRHIParamR
 
 void FVulkanCommandListContext::RHISetShaderUniformBuffer(FHullShaderRHIParamRef HullShaderRHI, uint32 BufferIndex, FUniformBufferRHIParamRef BufferRHI)
 {
-	ensureMsgf(0, TEXT("Tessellation not supported yet!"));
-/*
 	FVulkanUniformBuffer* UniformBuffer = ResourceCast(BufferRHI);
 	SetShaderUniformBuffer(ShaderStage::Hull, UniformBuffer, BufferIndex, ResourceCast(HullShaderRHI));
-*/
 }
 
 void FVulkanCommandListContext::RHISetShaderUniformBuffer(FDomainShaderRHIParamRef DomainShaderRHI, uint32 BufferIndex, FUniformBufferRHIParamRef BufferRHI)
 {
-	ensureMsgf(0, TEXT("Tessellation not supported yet!"));
-/*
 	FVulkanUniformBuffer* UniformBuffer = ResourceCast(BufferRHI);
 	SetShaderUniformBuffer(ShaderStage::Domain, UniformBuffer, BufferIndex, ResourceCast(DomainShaderRHI));
-*/
 }
 
 void FVulkanCommandListContext::RHISetShaderUniformBuffer(FGeometryShaderRHIParamRef GeometryShaderRHI, uint32 BufferIndex, FUniformBufferRHIParamRef BufferRHI)
diff --git a/Engine/Source/Runtime/VulkanRHI/Private/VulkanPipeline.cpp b/Engine/Source/Runtime/VulkanRHI/Private/VulkanPipeline.cpp
--- a/Engine/Source/Runtime/VulkanRHI/Private/VulkanPipeline.cpp
+++ b/Engine/Source/Runtime/VulkanRHI/Private/VulkanPipeline.cpp
@@ -87,8 +87,8 @@ static inline FSHAHash GetShaderHashForStage(const FGraphicsPipelineStateInitial
 #if VULKAN_SUPPORTS_GEOMETRY_SHADERS
 	case ShaderStage::Geometry:		return GetShaderHash<FRHIGeometryShader, FVulkanGeometryShader>(PSOInitializer.BoundShaderState.GeometryShaderRHI);
 #endif
-	//case ShaderStage::Hull:		return GetShaderHash<FRHIHullShader, FVulkanHullShader>(PSOInitializer.BoundShaderState.HullShaderRHI);
-	//case ShaderStage::Domain:		return GetShaderHash<FRHIDomainShader, FVulkanDomainShader>(PSOInitializer.BoundShaderState.DomainShaderRHI);
+	case ShaderStage::Hull:			return GetShaderHash<FRHIHullShader, FVulkanHullShader>(PSOInitializer.BoundShaderState.HullShaderRHI);
+	case ShaderStage::Domain:		return GetShaderHash<FRHIDomainShader, FVulkanDomainShader>(PSOInitializer.BoundShaderState.DomainShaderRHI);
 	default:			check(0);	break;
 	}
 
@@ -1020,6 +1020,7 @@ FArchive& operator << (FArchive& Ar, FVulkanPipelineStateCacheManager::FGfxPipel
 	// Modify VERSION if serialization changes
 	Ar << Entry.VertexInputKey;
 	Ar << Entry.RasterizationSamples;
+	Ar << Entry.ControlPoints;
 	Ar << Entry.Topology;
 
 	Ar << Entry.ColorAttachmentStates;
@@ -1114,7 +1115,7 @@ void FVulkanPipelineStateCacheManager::CreateGfxPipelineFromEntry(FGfxPipelineEn
 	// Multisample
 	VkPipelineMultisampleStateCreateInfo MSInfo;
 	ZeroVulkanStruct(MSInfo, VK_STRUCTURE_TYPE_PIPELINE_MULTISAMPLE_STATE_CREATE_INFO);
-	MSInfo.rasterizationSamples = (VkSampleCountFlagBits)FMath::Max(1u, GfxEntry->RasterizationSamples);
+	MSInfo.rasterizationSamples = (VkSampleCountFlagBits)FMath::Max<uint16>(1u, GfxEntry->RasterizationSamples);
 
 	VkPipelineShaderStageCreateInfo ShaderStages[ShaderStage::NumStages];
 	FMemory::Memzero(ShaderStages);
@@ -1122,6 +1123,7 @@ void FVulkanPipelineStateCacheManager::CreateGfxPipelineFromEntry(FGfxPipelineEn
 	PipelineInfo.pStages = ShaderStages;
 	// main_00000000_00000000
 	ANSICHAR EntryPoints[ShaderStage::NumStages][24];
+	bool bHasTessellation = false;
 	for (int32 ShaderStage = 0; ShaderStage < ShaderStage::NumStages; ++ShaderStage)
 	{
 		if (!GfxEntry->ShaderModules[ShaderStage])
@@ -1131,7 +1133,9 @@ void FVulkanPipelineStateCacheManager::CreateGfxPipelineFromEntry(FGfxPipelineEn
 		const ShaderStage::EStage CurrStage = (ShaderStage::EStage)ShaderStage;
 
 		ShaderStages[PipelineInfo.stageCount].sType = VK_STRUCTURE_TYPE_PIPELINE_SHADER_STAGE_CREATE_INFO;
-		ShaderStages[PipelineInfo.stageCount].stage = UEFrequencyToVKStageBit(ShaderStage::GetFrequencyForGfxStage(CurrStage));
+		VkShaderStageFlagBits Stage = UEFrequencyToVKStageBit(ShaderStage::GetFrequencyForGfxStage(CurrStage));
+		ShaderStages[PipelineInfo.stageCount].stage = Stage;
+		bHasTessellation = bHasTessellation || ((Stage & (VK_SHADER_STAGE_TESSELLATION_CONTROL_BIT | VK_SHADER_STAGE_TESSELLATION_EVALUATION_BIT)) != 0);
 		ShaderStages[PipelineInfo.stageCount].module = GfxEntry->ShaderModules[CurrStage];
 		Shaders[ShaderStage]->GetEntryPoint(EntryPoints[PipelineInfo.stageCount]);
 		ShaderStages[PipelineInfo.stageCount].pName = EntryPoints[PipelineInfo.stageCount];
@@ -1197,6 +1201,15 @@ void FVulkanPipelineStateCacheManager::CreateGfxPipelineFromEntry(FGfxPipelineEn
 
 	PipelineInfo.pDynamicState = &DynamicState;
 
+	VkPipelineTessellationStateCreateInfo TessState;
+	if (bHasTessellation)
+	{
+		ZeroVulkanStruct(TessState, VK_STRUCTURE_TYPE_PIPELINE_TESSELLATION_STATE_CREATE_INFO);
+		PipelineInfo.pTessellationState = &TessState;
+		check(InputAssembly.topology == VK_PRIMITIVE_TOPOLOGY_PATCH_LIST);
+		TessState.patchControlPoints = GfxEntry->ControlPoints;
+	}
+
 	//#todo-rco: Fix me
 	VkResult Result = VK_ERROR_INITIALIZATION_FAILED;
 	double BeginTime = FPlatformTime::Seconds();
@@ -1332,8 +1345,8 @@ FVulkanPipelineStateCacheManager::FShaderHashes::FShaderHashes(const FGraphicsPi
 #if VULKAN_SUPPORTS_GEOMETRY_SHADERS
 	Stages[ShaderStage::Geometry] = GetShaderHash<FRHIGeometryShader, FVulkanGeometryShader>(PSOInitializer.BoundShaderState.GeometryShaderRHI);
 #endif
-	//Stages[ShaderStage::Hull] = GetShaderHash<FRHIHullShader, FVulkanHullShader>(PSOInitializer.BoundShaderState.HullShaderRHI);
-	//Stages[ShaderStage::Domain] = GetShaderHash<FRHIDomainShader, FVulkanDomainShader>(PSOInitializer.BoundShaderState.DomainShaderRHI);
+	Stages[ShaderStage::Hull] = GetShaderHash<FRHIHullShader, FVulkanHullShader>(PSOInitializer.BoundShaderState.HullShaderRHI);
+	Stages[ShaderStage::Domain] = GetShaderHash<FRHIDomainShader, FVulkanDomainShader>(PSOInitializer.BoundShaderState.DomainShaderRHI);
 	Finalize();
 }
 
@@ -1398,14 +1411,13 @@ FVulkanGfxLayout* FVulkanPipelineStateCacheManager::GetOrGenerateGfxLayout(const
 	}
 #endif
 
-	/* We don't support tessellation on desktop currently
 	if (Shaders[ShaderStage::Hull])
 	{
-		const FVulkanCodeHeader& HSHeader = Shaders[ShaderStage::Hull]->GetCodeHeader();
-		const FVulkanCodeHeader& DSHeader = Shaders[ShaderStage::Domain]->GetCodeHeader();
-		DescriptorSetLayoutInfo.AddBindingsForStage(VK_SHADER_STAGE_TESSELLATION_CONTROL_BIT, ShaderStage::Hull, HSHeader);
-		DescriptorSetLayoutInfo.AddBindingsForStage(VK_SHADER_STAGE_TESSELLATION_EVALUATION_BIT, ShaderStage::Domain, DSHeader);
-	}*/
+		const FVulkanShaderHeader& HSHeader = Shaders[ShaderStage::Hull]->GetCodeHeader();
+		const FVulkanShaderHeader& DSHeader = Shaders[ShaderStage::Domain]->GetCodeHeader();
+		DescriptorSetLayoutInfo.ProcessBindingsForStage(VK_SHADER_STAGE_TESSELLATION_CONTROL_BIT, ShaderStage::Hull, HSHeader, UBGatherInfo);
+		DescriptorSetLayoutInfo.ProcessBindingsForStage(VK_SHADER_STAGE_TESSELLATION_EVALUATION_BIT, ShaderStage::Domain, DSHeader, UBGatherInfo);
+	}
 
 	// Second pass
 	const int32 NumImmutableSamplers = PSOInitializer.ImmutableSamplerState.ImmutableSamplers.Num();
@@ -1422,6 +1434,50 @@ FVulkanGfxLayout* FVulkanPipelineStateCacheManager::GetOrGenerateGfxLayout(const
 	return GfxLayout;
 }
 
+static inline VkPrimitiveTopology UEToVulkanTopologyType(EPrimitiveType PrimitiveType, bool bHasTessellation, uint16& OutControlPoints)
+{
+	if (bHasTessellation)
+	{
+		switch (PrimitiveType)
+		{
+		case PT_TriangleList:
+			// This is the case for tessellation without AEN or other buffers, so just flip to 3 CPs
+			OutControlPoints = 3;
+			return VK_PRIMITIVE_TOPOLOGY_PATCH_LIST;
+		case PT_1_ControlPointPatchList:
+			OutControlPoints = 1;
+			return VK_PRIMITIVE_TOPOLOGY_PATCH_LIST;
+		case PT_2_ControlPointPatchList:
+			OutControlPoints = 2;
+			return VK_PRIMITIVE_TOPOLOGY_PATCH_LIST;
+		default:
+			checkf(false, TEXT("Unsupported tessellation EPrimitiveType %d; probably missing a case in FStaticMeshSceneProxy::GetMeshElement()!"), (uint32)PrimitiveType);
+			break;
+		}
+		OutControlPoints = 0;
+	}
+	else
+	{
+		OutControlPoints = 0;
+		switch (PrimitiveType)
+		{
+		case PT_PointList:
+			return VK_PRIMITIVE_TOPOLOGY_POINT_LIST;
+		case PT_LineList:
+			return VK_PRIMITIVE_TOPOLOGY_LINE_LIST;
+		case PT_TriangleList:
+			return VK_PRIMITIVE_TOPOLOGY_TRIANGLE_LIST;
+		case PT_TriangleStrip:
+			return VK_PRIMITIVE_TOPOLOGY_TRIANGLE_STRIP;
+		default:
+			checkf(false, TEXT("Unsupported EPrimitiveType %d"), (uint32)PrimitiveType);
+			break;
+		}
+	}
+
+	return VK_PRIMITIVE_TOPOLOGY_MAX_ENUM;
+}
+
 FVulkanPipelineStateCacheManager::FGfxPipelineEntry* FVulkanPipelineStateCacheManager::CreateGfxEntry(const FGraphicsPipelineStateInitializer& PSOInitializer)
 {
 	FGfxPipelineEntry* OutGfxEntry = new FGfxPipelineEntry();
@@ -1434,9 +1490,11 @@ FVulkanPipelineStateCacheManager::FGfxPipelineEntry* FVulkanPipelineStateCacheMa
 	FVulkanVertexInputStateInfo VertexInputState;
 	OutGfxEntry->Layout = GetOrGenerateGfxLayout(PSOInitializer, Shaders, VertexInputState);
 
+	const bool bHasTessellation = (PSOInitializer.BoundShaderState.DomainShaderRHI != nullptr);
+
 	OutGfxEntry->RasterizationSamples = OutGfxEntry->RenderPass->GetLayout().GetAttachmentDescriptions()[0].samples;
 	ensure(OutGfxEntry->RasterizationSamples == PSOInitializer.NumSamples);
-	OutGfxEntry->Topology = (uint32)UEToVulkanType(PSOInitializer.PrimitiveType);
+	OutGfxEntry->Topology = (uint32)UEToVulkanTopologyType(PSOInitializer.PrimitiveType, bHasTessellation, OutGfxEntry->ControlPoints);
 
 	OutGfxEntry->ColorAttachmentStates.AddUninitialized(OutGfxEntry->RenderPass->GetLayout().GetNumColorAttachments());
 	for (int32 Index = 0; Index < OutGfxEntry->ColorAttachmentStates.Num(); ++Index)
@@ -1559,9 +1617,9 @@ FVulkanRHIGraphicsPipelineState* FVulkanPipelineStateCacheManager::FindInRuntime
 			TempUInt64 = GetShaderKey(PSI.BoundShaderState.GeometryShaderRHI);
 	#endif
 			Ar << TempUInt64;
-			//TempUInt64 = GetShaderKey(PSI.BoundShaderState.HullShaderRHI);
+			TempUInt64 = GetShaderKey(PSI.BoundShaderState.HullShaderRHI);
 			Ar << TempUInt64;
-			//TempUInt64 = GetShaderKey(PSI.BoundShaderState.DomainShaderRHI);
+			TempUInt64 = GetShaderKey(PSI.BoundShaderState.DomainShaderRHI);
 			Ar << TempUInt64;
 			Ar << ResourceCast(PSI.BoundShaderState.VertexDeclarationRHI)->Elements;
 			Ar << ResourceCast(PSI.RasterizerState)->Initializer;
@@ -1949,13 +2007,10 @@ void GetVulkanShaders(const FBoundShaderStateInput& BSI, FVulkanShader* OutShade
 #if PLATFORM_SUPPORTS_TESSELLATION_SHADERS
 	if (BSI.HullShaderRHI)
 	{
-		ensureMsgf(0, TEXT("Tessellation not supported yet!"));
-		/*
 		// Can't have Hull w/o Domain
 		check(BSI.DomainShaderRHI);
 		OutShaders[ShaderStage::Hull] = ResourceCast(BSI.HullShaderRHI);
 		OutShaders[ShaderStage::Domain] = ResourceCast(BSI.DomainShaderRHI);
-		*/
 	}
 	else
 	{
diff --git a/Engine/Source/Runtime/VulkanRHI/Private/VulkanPipeline.h b/Engine/Source/Runtime/VulkanRHI/Private/VulkanPipeline.h
--- a/Engine/Source/Runtime/VulkanRHI/Private/VulkanPipeline.h
+++ b/Engine/Source/Runtime/VulkanRHI/Private/VulkanPipeline.h
@@ -43,6 +43,10 @@ inline uint64 GetShaderKeyForGfxStage(const FBoundShaderStateInput& BSI, ShaderS
 	case ShaderStage::Geometry:
 		return GetShaderKey(BSI.GeometryShaderRHI);
 #endif
+	case ShaderStage::Hull:
+		return GetShaderKey(BSI.HullShaderRHI);
+	case ShaderStage::Domain:
+		return GetShaderKey(BSI.DomainShaderRHI);
 	default:
 		check(0);
 	}
@@ -148,7 +152,8 @@ public:
 		uint32 VertexInputKey;
 		bool bLoaded;
 
-		uint32 RasterizationSamples;
+		uint16 RasterizationSamples;
+		uint16 ControlPoints;
 		uint32 Topology;
 		struct FBlendAttachment
 		{
diff --git a/Engine/Source/Runtime/VulkanRHI/Private/VulkanPipelineState.cpp b/Engine/Source/Runtime/VulkanRHI/Private/VulkanPipelineState.cpp
--- a/Engine/Source/Runtime/VulkanRHI/Private/VulkanPipelineState.cpp
+++ b/Engine/Source/Runtime/VulkanRHI/Private/VulkanPipelineState.cpp
@@ -236,15 +236,18 @@ FVulkanGraphicsPipelineDescriptorState::FVulkanGraphicsPipelineDescriptorState(F
 	}
 #endif
 
-/*
 	uint64 HullShaderKey = InGfxPipeline->GetShaderKey(SF_Hull);
 	if (HullShaderKey)
 	{
-		ensureMsgf(0, TEXT("Tessellation not supported yet!"));
-		PackedUniformBuffers[ShaderStage::Domain].Init(CodeHeaderPerStage[ShaderStage::Domain], PackedUniformBuffersMask[ShaderStage::Domain], UniformBuffersWithDataMask[ShaderStage::Domain], ResourcesDirtyMask[ShaderStage::Domain]);
-		PackedUniformBuffers[ShaderStage::Hull].Init(CodeHeaderPerStage[ShaderStage::Hull], PackedUniformBuffersMask[ShaderStage::Hull], UniformBuffersWithDataMask[ShaderStage::Hull], ResourcesDirtyMask[ShaderStage::Domain]);
+		const FVulkanHullShader* HullShader = ShaderFactory.LookupShader<FVulkanHullShader>(HullShaderKey);
+		PackedUniformBuffers[ShaderStage::Hull].Init(HullShader->GetCodeHeader(), PackedUniformBuffersMask[ShaderStage::Hull]);
+	}
+	uint64 DomainShaderKey = InGfxPipeline->GetShaderKey(SF_Domain);
+	if (DomainShaderKey)
+	{
+		const FVulkanDomainShader* DomainShader = ShaderFactory.LookupShader<FVulkanDomainShader>(DomainShaderKey);
+		PackedUniformBuffers[ShaderStage::Domain].Init(DomainShader->GetCodeHeader(), PackedUniformBuffersMask[ShaderStage::Domain]);
 	}
-*/
 
 	CreateDescriptorWriteInfos();
 
diff --git a/Engine/Source/Runtime/VulkanRHI/Private/VulkanRHIPrivate.h b/Engine/Source/Runtime/VulkanRHI/Private/VulkanRHIPrivate.h
--- a/Engine/Source/Runtime/VulkanRHI/Private/VulkanRHIPrivate.h
+++ b/Engine/Source/Runtime/VulkanRHI/Private/VulkanRHIPrivate.h
@@ -77,8 +77,8 @@ inline EShaderFrequency VkStageBitToUEFrequency(VkShaderStageFlagBits FlagBits)
 	switch (FlagBits)
 	{
 	case VK_SHADER_STAGE_VERTEX_BIT:					return SF_Vertex;
-	//case VK_SHADER_STAGE_TESSELLATION_CONTROL_BIT:		return SF_Hull;
-	//case VK_SHADER_STAGE_TESSELLATION_EVALUATION_BIT:		return SF_Domain;
+	case VK_SHADER_STAGE_TESSELLATION_CONTROL_BIT:		return SF_Hull;
+	case VK_SHADER_STAGE_TESSELLATION_EVALUATION_BIT:	return SF_Domain;
 	case VK_SHADER_STAGE_FRAGMENT_BIT:					return SF_Pixel;
 	case VK_SHADER_STAGE_GEOMETRY_BIT:					return SF_Geometry;
 	case VK_SHADER_STAGE_COMPUTE_BIT:					return SF_Compute;
@@ -713,22 +713,6 @@ static inline VkFormat UEToVkBufferFormat(EVertexElementType Type)
 	return VK_FORMAT_UNDEFINED;
 }
 
-static inline VkPrimitiveTopology UEToVulkanType(EPrimitiveType PrimitiveType)
-{
-	switch (PrimitiveType)
-	{
-	case PT_PointList:			return VK_PRIMITIVE_TOPOLOGY_POINT_LIST;
-	case PT_LineList:			return VK_PRIMITIVE_TOPOLOGY_LINE_LIST;
-	case PT_TriangleList:		return VK_PRIMITIVE_TOPOLOGY_TRIANGLE_LIST;
-	case PT_TriangleStrip:		return VK_PRIMITIVE_TOPOLOGY_TRIANGLE_STRIP;
-	default:
-		break;
-	}
-
-	checkf(false, TEXT("Unsupported primitive type"));
-	return VK_PRIMITIVE_TOPOLOGY_MAX_ENUM;
-}
-
 #if UE_BUILD_DEBUG || UE_BUILD_DEVELOPMENT
 extern TAutoConsoleVariable<int32> CVarVulkanDebugBarrier;
 #endif
diff --git a/Engine/Source/Runtime/VulkanRHI/Public/VulkanCommon.h b/Engine/Source/Runtime/VulkanRHI/Public/VulkanCommon.h
--- a/Engine/Source/Runtime/VulkanRHI/Public/VulkanCommon.h
+++ b/Engine/Source/Runtime/VulkanRHI/Public/VulkanCommon.h
@@ -26,12 +26,12 @@ namespace ShaderStage
 		Pixel			= 1,
 
 #if VULKAN_SUPPORTS_GEOMETRY_SHADERS
-		// We don't support tessellation on desktop currently
-		//Hull			= 3,
-		//Domain		= 4,
 		Geometry		= 2,
-
-		NumStages		= 3,
+#if PLATFORM_SUPPORTS_TESSELLATION_SHADERS
+		Hull			= 3,
+		Domain			= 4,
+#endif
+		NumStages,
 
 		MaxNumSets		= 8,
 #else
@@ -51,8 +51,10 @@ namespace ShaderStage
 		switch (Stage)
 		{
 		case SF_Vertex:		return Vertex;
-		//case SF_Hull:		return Hull;
-		//case SF_Domain:		return Domain;
+#if PLATFORM_SUPPORTS_TESSELLATION_SHADERS
+		case SF_Hull:		return Hull;
+		case SF_Domain:		return Domain;
+#endif
 		case SF_Pixel:		return Pixel;
 #if VULKAN_SUPPORTS_GEOMETRY_SHADERS
 		case SF_Geometry:	return Geometry;
@@ -71,8 +73,10 @@ namespace ShaderStage
 		switch (Stage)
 		{
 		case EStage::Vertex:	return SF_Vertex;
-		//case EStage::Hull:		return SF_Hull;
-		//case EStage::Domain:	return SF_Domain;
+#if PLATFORM_SUPPORTS_TESSELLATION_SHADERS
+		case EStage::Hull:		return SF_Hull;
+		case EStage::Domain:	return SF_Domain;
+#endif
 		case EStage::Pixel:		return SF_Pixel;
 #if VULKAN_SUPPORTS_GEOMETRY_SHADERS
 		case EStage::Geometry:	return SF_Geometry;
diff --git a/Engine/Source/Runtime/VulkanRHI/Public/VulkanResources.h b/Engine/Source/Runtime/VulkanRHI/Public/VulkanResources.h
--- a/Engine/Source/Runtime/VulkanRHI/Public/VulkanResources.h
+++ b/Engine/Source/Runtime/VulkanRHI/Public/VulkanResources.h
@@ -241,8 +241,8 @@ public:
 		switch (Stage)
 		{
 		case ShaderStage::Vertex:		return GetVertexShader();
-		//case ShaderStage::Hull:		return GetHullShader();
-		//case ShaderStage::Domain:		return GetDomainShader();
+		case ShaderStage::Hull:			return GetHullShader();
+		case ShaderStage::Domain:		return GetDomainShader();
 		case ShaderStage::Pixel:		return GetPixelShader();
 #if VULKAN_SUPPORTS_GEOMETRY_SHADERS
 		case ShaderStage::Geometry:	return GetGeometryShader();
