UE4 - Android compile fixes
#rnx

diff --git a/Engine/Source/Runtime/VulkanRHI/Private/VulkanCommands.cpp b/Engine/Source/Runtime/VulkanRHI/Private/VulkanCommands.cpp
--- a/Engine/Source/Runtime/VulkanRHI/Private/VulkanCommands.cpp
+++ b/Engine/Source/Runtime/VulkanRHI/Private/VulkanCommands.cpp
@@ -188,20 +188,28 @@ void FVulkanCommandListContext::RHISetShaderTexture(FVertexShaderRHIParamRef Ver
 
 void FVulkanCommandListContext::RHISetShaderTexture(FHullShaderRHIParamRef HullShaderRHI, uint32 TextureIndex, FTextureRHIParamRef NewTextureRHI)
 {
+#if PLATFORM_SUPPORTS_TESSELLATION_SHADERS
 	check(PendingGfxState->GetCurrentShaderKey(ShaderStage::Hull) == GetShaderKey(HullShaderRHI));
 	FVulkanTextureBase* Texture = GetVulkanTextureFromRHITexture(NewTextureRHI);
 	VkImageLayout Layout = GetLayoutForDescriptor(Texture->Surface);
 	PendingGfxState->SetTextureForStage(ShaderStage::Hull, TextureIndex, Texture, Layout);
 	NewTextureRHI->SetLastRenderTime((float)FPlatformTime::Seconds());
+#else
+	ensureMsgf(0, TEXT("Tessellation not supported on this platform!"));
+#endif
 }
 
 void FVulkanCommandListContext::RHISetShaderTexture(FDomainShaderRHIParamRef DomainShaderRHI, uint32 TextureIndex, FTextureRHIParamRef NewTextureRHI)
 {
+#if PLATFORM_SUPPORTS_TESSELLATION_SHADERS
 	check(PendingGfxState->GetCurrentShaderKey(ShaderStage::Domain) == GetShaderKey(DomainShaderRHI));
 	FVulkanTextureBase* Texture = GetVulkanTextureFromRHITexture(NewTextureRHI);
 	VkImageLayout Layout = GetLayoutForDescriptor(Texture->Surface);
 	PendingGfxState->SetTextureForStage(ShaderStage::Domain, TextureIndex, Texture, Layout);
 	NewTextureRHI->SetLastRenderTime((float)FPlatformTime::Seconds());
+#else
+	ensureMsgf(0, TEXT("Tessellation not supported on this platform!"));
+#endif
 }
 
 void FVulkanCommandListContext::RHISetShaderTexture(FGeometryShaderRHIParamRef GeometryShaderRHI, uint32 TextureIndex, FTextureRHIParamRef NewTextureRHI)
@@ -246,16 +254,24 @@ void FVulkanCommandListContext::RHISetShaderResourceViewParameter(FVertexShaderR
 
 void FVulkanCommandListContext::RHISetShaderResourceViewParameter(FHullShaderRHIParamRef HullShaderRHI,uint32 TextureIndex,FShaderResourceViewRHIParamRef SRVRHI)
 {
+#if PLATFORM_SUPPORTS_TESSELLATION_SHADERS
 	check(PendingGfxState->GetCurrentShaderKey(ShaderStage::Hull) == GetShaderKey(HullShaderRHI));
 	FVulkanShaderResourceView* SRV = ResourceCast(SRVRHI);
 	PendingGfxState->SetSRVForStage(ShaderStage::Hull, TextureIndex, SRV);
+#else
+	ensureMsgf(0, TEXT("Tessellation not supported on this platform!"));
+#endif
 }
 
 void FVulkanCommandListContext::RHISetShaderResourceViewParameter(FDomainShaderRHIParamRef DomainShaderRHI,uint32 TextureIndex,FShaderResourceViewRHIParamRef SRVRHI)
 {
+#if PLATFORM_SUPPORTS_TESSELLATION_SHADERS
 	check(PendingGfxState->GetCurrentShaderKey(ShaderStage::Domain) == GetShaderKey(DomainShaderRHI));
 	FVulkanShaderResourceView* SRV = ResourceCast(SRVRHI);
 	PendingGfxState->SetSRVForStage(ShaderStage::Domain, TextureIndex, SRV);
+#else
+	ensureMsgf(0, TEXT("Tessellation not supported on this platform!"));
+#endif
 }
 
 void FVulkanCommandListContext::RHISetShaderResourceViewParameter(FGeometryShaderRHIParamRef GeometryShaderRHI,uint32 TextureIndex,FShaderResourceViewRHIParamRef SRVRHI)
@@ -293,16 +309,24 @@ void FVulkanCommandListContext::RHISetShaderSampler(FVertexShaderRHIParamRef Ver
 
 void FVulkanCommandListContext::RHISetShaderSampler(FHullShaderRHIParamRef HullShaderRHI, uint32 SamplerIndex, FSamplerStateRHIParamRef NewStateRHI)
 {
+#if PLATFORM_SUPPORTS_TESSELLATION_SHADERS
 	check(PendingGfxState->GetCurrentShaderKey(ShaderStage::Hull) == GetShaderKey(HullShaderRHI));
 	FVulkanSamplerState* Sampler = ResourceCast(NewStateRHI);
 	PendingGfxState->SetSamplerStateForStage(ShaderStage::Hull, SamplerIndex, Sampler);
+#else
+	ensureMsgf(0, TEXT("Tessellation not supported on this platform!"));
+#endif
 }
 
 void FVulkanCommandListContext::RHISetShaderSampler(FDomainShaderRHIParamRef DomainShaderRHI, uint32 SamplerIndex, FSamplerStateRHIParamRef NewStateRHI)
 {
+#if PLATFORM_SUPPORTS_TESSELLATION_SHADERS
 	check(PendingGfxState->GetCurrentShaderKey(ShaderStage::Domain) == GetShaderKey(DomainShaderRHI));
 	FVulkanSamplerState* Sampler = ResourceCast(NewStateRHI);
 	PendingGfxState->SetSamplerStateForStage(ShaderStage::Domain, SamplerIndex, Sampler);
+#else
+	ensureMsgf(0, TEXT("Tessellation not supported on this platform!"));
+#endif
 }
 
 void FVulkanCommandListContext::RHISetShaderSampler(FGeometryShaderRHIParamRef GeometryShaderRHI, uint32 SamplerIndex, FSamplerStateRHIParamRef NewStateRHI)
@@ -341,16 +365,24 @@ void FVulkanCommandListContext::RHISetShaderParameter(FVertexShaderRHIParamRef V
 
 void FVulkanCommandListContext::RHISetShaderParameter(FHullShaderRHIParamRef HullShaderRHI, uint32 BufferIndex, uint32 BaseIndex, uint32 NumBytes, const void* NewValue)
 {
+#if PLATFORM_SUPPORTS_TESSELLATION_SHADERS
 	check(PendingGfxState->GetCurrentShaderKey(ShaderStage::Hull) == GetShaderKey(HullShaderRHI));
 
 	PendingGfxState->SetPackedGlobalShaderParameter(ShaderStage::Hull, BufferIndex, BaseIndex, NumBytes, NewValue);
+#else
+	ensureMsgf(0, TEXT("Tessellation not supported on this platform!"));
+#endif
 }
 
 void FVulkanCommandListContext::RHISetShaderParameter(FDomainShaderRHIParamRef DomainShaderRHI, uint32 BufferIndex, uint32 BaseIndex, uint32 NumBytes, const void* NewValue)
 {
+#if PLATFORM_SUPPORTS_TESSELLATION_SHADERS
 	check(PendingGfxState->GetCurrentShaderKey(ShaderStage::Domain) == GetShaderKey(DomainShaderRHI));
 
 	PendingGfxState->SetPackedGlobalShaderParameter(ShaderStage::Domain, BufferIndex, BaseIndex, NumBytes, NewValue);
+#else
+	ensureMsgf(0, TEXT("Tessellation not supported on this platform!"));
+#endif
 }
 
 void FVulkanCommandListContext::RHISetShaderParameter(FGeometryShaderRHIParamRef GeometryShaderRHI, uint32 BufferIndex, uint32 BaseIndex, uint32 NumBytes, const void* NewValue)
@@ -622,14 +654,22 @@ void FVulkanCommandListContext::RHISetShaderUniformBuffer(FVertexShaderRHIParamR
 
 void FVulkanCommandListContext::RHISetShaderUniformBuffer(FHullShaderRHIParamRef HullShaderRHI, uint32 BufferIndex, FUniformBufferRHIParamRef BufferRHI)
 {
+#if PLATFORM_SUPPORTS_TESSELLATION_SHADERS
 	FVulkanUniformBuffer* UniformBuffer = ResourceCast(BufferRHI);
 	SetShaderUniformBuffer(ShaderStage::Hull, UniformBuffer, BufferIndex, ResourceCast(HullShaderRHI));
+#else
+	ensureMsgf(0, TEXT("Tessellation not supported on this platform!"));
+#endif
 }
 
 void FVulkanCommandListContext::RHISetShaderUniformBuffer(FDomainShaderRHIParamRef DomainShaderRHI, uint32 BufferIndex, FUniformBufferRHIParamRef BufferRHI)
 {
+#if PLATFORM_SUPPORTS_TESSELLATION_SHADERS
 	FVulkanUniformBuffer* UniformBuffer = ResourceCast(BufferRHI);
 	SetShaderUniformBuffer(ShaderStage::Domain, UniformBuffer, BufferIndex, ResourceCast(DomainShaderRHI));
+#else
+	ensureMsgf(0, TEXT("Tessellation not supported on this platform!"));
+#endif
 }
 
 void FVulkanCommandListContext::RHISetShaderUniformBuffer(FGeometryShaderRHIParamRef GeometryShaderRHI, uint32 BufferIndex, FUniformBufferRHIParamRef BufferRHI)
diff --git a/Engine/Source/Runtime/VulkanRHI/Private/VulkanPipeline.cpp b/Engine/Source/Runtime/VulkanRHI/Private/VulkanPipeline.cpp
--- a/Engine/Source/Runtime/VulkanRHI/Private/VulkanPipeline.cpp
+++ b/Engine/Source/Runtime/VulkanRHI/Private/VulkanPipeline.cpp
@@ -87,8 +87,10 @@ static inline FSHAHash GetShaderHashForStage(const FGraphicsPipelineStateInitial
 #if VULKAN_SUPPORTS_GEOMETRY_SHADERS
 	case ShaderStage::Geometry:		return GetShaderHash<FRHIGeometryShader, FVulkanGeometryShader>(PSOInitializer.BoundShaderState.GeometryShaderRHI);
 #endif
+#if PLATFORM_SUPPORTS_TESSELLATION_SHADERS
 	case ShaderStage::Hull:			return GetShaderHash<FRHIHullShader, FVulkanHullShader>(PSOInitializer.BoundShaderState.HullShaderRHI);
 	case ShaderStage::Domain:		return GetShaderHash<FRHIDomainShader, FVulkanDomainShader>(PSOInitializer.BoundShaderState.DomainShaderRHI);
+#endif
 	default:			check(0);	break;
 	}
 
@@ -1345,8 +1347,10 @@ FVulkanPipelineStateCacheManager::FShaderHashes::FShaderHashes(const FGraphicsPi
 #if VULKAN_SUPPORTS_GEOMETRY_SHADERS
 	Stages[ShaderStage::Geometry] = GetShaderHash<FRHIGeometryShader, FVulkanGeometryShader>(PSOInitializer.BoundShaderState.GeometryShaderRHI);
 #endif
+#if PLATFORM_SUPPORTS_TESSELLATION_SHADERS
 	Stages[ShaderStage::Hull] = GetShaderHash<FRHIHullShader, FVulkanHullShader>(PSOInitializer.BoundShaderState.HullShaderRHI);
 	Stages[ShaderStage::Domain] = GetShaderHash<FRHIDomainShader, FVulkanDomainShader>(PSOInitializer.BoundShaderState.DomainShaderRHI);
+#endif
 	Finalize();
 }
 
@@ -1411,6 +1415,7 @@ FVulkanGfxLayout* FVulkanPipelineStateCacheManager::GetOrGenerateGfxLayout(const
 	}
 #endif
 
+#if PLATFORM_SUPPORTS_TESSELLATION_SHADERS
 	if (Shaders[ShaderStage::Hull])
 	{
 		const FVulkanShaderHeader& HSHeader = Shaders[ShaderStage::Hull]->GetCodeHeader();
@@ -1418,7 +1423,7 @@ FVulkanGfxLayout* FVulkanPipelineStateCacheManager::GetOrGenerateGfxLayout(const
 		DescriptorSetLayoutInfo.ProcessBindingsForStage(VK_SHADER_STAGE_TESSELLATION_CONTROL_BIT, ShaderStage::Hull, HSHeader, UBGatherInfo);
 		DescriptorSetLayoutInfo.ProcessBindingsForStage(VK_SHADER_STAGE_TESSELLATION_EVALUATION_BIT, ShaderStage::Domain, DSHeader, UBGatherInfo);
 	}
-
+#endif
 	// Second pass
 	const int32 NumImmutableSamplers = PSOInitializer.ImmutableSamplerState.ImmutableSamplers.Num();
 	TArrayView<const FSamplerStateRHIParamRef> ImmutableSamplers(NumImmutableSamplers > 0 ? &PSOInitializer.ImmutableSamplerState.ImmutableSamplers[0] : nullptr, NumImmutableSamplers);
diff --git a/Engine/Source/Runtime/VulkanRHI/Private/VulkanPipeline.h b/Engine/Source/Runtime/VulkanRHI/Private/VulkanPipeline.h
--- a/Engine/Source/Runtime/VulkanRHI/Private/VulkanPipeline.h
+++ b/Engine/Source/Runtime/VulkanRHI/Private/VulkanPipeline.h
@@ -43,10 +43,12 @@ inline uint64 GetShaderKeyForGfxStage(const FBoundShaderStateInput& BSI, ShaderS
 	case ShaderStage::Geometry:
 		return GetShaderKey(BSI.GeometryShaderRHI);
 #endif
+#if PLATFORM_SUPPORTS_TESSELLATION_SHADERS
 	case ShaderStage::Hull:
 		return GetShaderKey(BSI.HullShaderRHI);
 	case ShaderStage::Domain:
 		return GetShaderKey(BSI.DomainShaderRHI);
+#endif
 	default:
 		check(0);
 	}
diff --git a/Engine/Source/Runtime/VulkanRHI/Private/VulkanPipelineState.cpp b/Engine/Source/Runtime/VulkanRHI/Private/VulkanPipelineState.cpp
--- a/Engine/Source/Runtime/VulkanRHI/Private/VulkanPipelineState.cpp
+++ b/Engine/Source/Runtime/VulkanRHI/Private/VulkanPipelineState.cpp
@@ -236,6 +236,7 @@ FVulkanGraphicsPipelineDescriptorState::FVulkanGraphicsPipelineDescriptorState(F
 	}
 #endif
 
+#if PLATFORM_SUPPORTS_TESSELLATION_SHADERS
 	uint64 HullShaderKey = InGfxPipeline->GetShaderKey(SF_Hull);
 	if (HullShaderKey)
 	{
@@ -248,7 +249,7 @@ FVulkanGraphicsPipelineDescriptorState::FVulkanGraphicsPipelineDescriptorState(F
 		const FVulkanDomainShader* DomainShader = ShaderFactory.LookupShader<FVulkanDomainShader>(DomainShaderKey);
 		PackedUniformBuffers[ShaderStage::Domain].Init(DomainShader->GetCodeHeader(), PackedUniformBuffersMask[ShaderStage::Domain]);
 	}
-
+#endif
 	CreateDescriptorWriteInfos();
 
 	//UE_LOG(LogVulkanRHI, Warning, TEXT("GfxPSOState %p For PSO %p Writes:%d"), this, InGfxPipeline, DSWriteContainer.DescriptorWrites.Num());
diff --git a/Engine/Source/Runtime/VulkanRHI/Public/VulkanResources.h b/Engine/Source/Runtime/VulkanRHI/Public/VulkanResources.h
--- a/Engine/Source/Runtime/VulkanRHI/Public/VulkanResources.h
+++ b/Engine/Source/Runtime/VulkanRHI/Public/VulkanResources.h
@@ -241,8 +241,10 @@ public:
 		switch (Stage)
 		{
 		case ShaderStage::Vertex:		return GetVertexShader();
+#if PLATFORM_SUPPORTS_TESSELLATION_SHADERS
 		case ShaderStage::Hull:			return GetHullShader();
 		case ShaderStage::Domain:		return GetDomainShader();
+#endif
 		case ShaderStage::Pixel:		return GetPixelShader();
 #if VULKAN_SUPPORTS_GEOMETRY_SHADERS
 		case ShaderStage::Geometry:	return GetGeometryShader();
