UE4.22 - Add better information for bad states
#rnx

diff --git a/Engine/Source/Runtime/VulkanRHI/Private/VulkanState.cpp b/Engine/Source/Runtime/VulkanRHI/Private/VulkanState.cpp
--- a/Engine/Source/Runtime/VulkanRHI/Private/VulkanState.cpp
+++ b/Engine/Source/Runtime/VulkanRHI/Private/VulkanState.cpp
@@ -9,160 +9,127 @@
 
 static FCriticalSection GSamplerHashLock;
 
-inline VkSamplerMipmapMode TranslateMipFilterMode(ESamplerFilter Filter)
+inline VkSamplerMipmapMode TranslateMipFilterMode(ESamplerFilter InFilter)
 {
-	VkSamplerMipmapMode OutFilter = VK_SAMPLER_MIPMAP_MODE_MAX_ENUM;
-
-	switch (Filter)
-	{
-		case SF_Point:				OutFilter = VK_SAMPLER_MIPMAP_MODE_NEAREST;	break;
-		case SF_Bilinear:			OutFilter = VK_SAMPLER_MIPMAP_MODE_NEAREST;	break;
-		case SF_Trilinear:			OutFilter = VK_SAMPLER_MIPMAP_MODE_LINEAR;	break;
-		case SF_AnisotropicPoint:	OutFilter = VK_SAMPLER_MIPMAP_MODE_LINEAR;	break;
-		default:																break;
-	}
-
-	// Check for missing translation
-	check(OutFilter != VK_SAMPLER_MIPMAP_MODE_MAX_ENUM);
-	return OutFilter;
-}
-
-inline VkFilter TranslateMagFilterMode(ESamplerFilter InFilter)
-{
-	VkFilter OutFilter = VK_FILTER_MAX_ENUM;
-
 	switch (InFilter)
 	{
-		case SF_Point:				OutFilter = VK_FILTER_NEAREST;	break;
-		case SF_Bilinear:			OutFilter = VK_FILTER_LINEAR;	break;
-		case SF_Trilinear:			OutFilter = VK_FILTER_LINEAR;	break;
-		case SF_AnisotropicPoint:	OutFilter = VK_FILTER_LINEAR;	break;
-		default:													break;
+		case SF_Point:				return VK_SAMPLER_MIPMAP_MODE_NEAREST;
+		case SF_Bilinear:			return VK_SAMPLER_MIPMAP_MODE_NEAREST;
+		case SF_Trilinear:			return VK_SAMPLER_MIPMAP_MODE_LINEAR;
+		case SF_AnisotropicPoint:	return VK_SAMPLER_MIPMAP_MODE_LINEAR;
+		default:
+			break;
 	}
 
-	// Check for missing translation
-	check(OutFilter != VK_FILTER_MAX_ENUM);
-	return OutFilter;
+	checkf(0, TEXT("Unknown Mip ESamplerFilter %d"), (uint32)InFilter);
+	return VK_SAMPLER_MIPMAP_MODE_MAX_ENUM;
 }
 
-inline VkFilter TranslateMinFilterMode(ESamplerFilter InFilter)
+inline VkFilter TranslateMinMagFilterMode(ESamplerFilter InFilter)
 {
-	VkFilter OutFilter = VK_FILTER_MAX_ENUM;
-
 	switch (InFilter)
 	{
-		case SF_Point:				OutFilter = VK_FILTER_NEAREST;	break;
-		case SF_Bilinear:			OutFilter = VK_FILTER_LINEAR;	break;
-		case SF_Trilinear:			OutFilter = VK_FILTER_LINEAR;	break;
-		case SF_AnisotropicPoint:	OutFilter = VK_FILTER_LINEAR;	break;
-		default:													break;
+		case SF_Point:				return VK_FILTER_NEAREST;
+		case SF_Bilinear:			return VK_FILTER_LINEAR;
+		case SF_Trilinear:			return VK_FILTER_LINEAR;
+		case SF_AnisotropicPoint:	return VK_FILTER_LINEAR;
+		default:
+			break;
 	}
 
-	// Check for missing translation
-	check(OutFilter != VK_FILTER_MAX_ENUM);
-	return OutFilter;
+	checkf(0, TEXT("Unknown ESamplerFilter %d"), (uint32)InFilter);
+	return VK_FILTER_MAX_ENUM;
 }
 
 inline VkSamplerAddressMode TranslateWrapMode(ESamplerAddressMode InAddressMode)
 {
-	VkSamplerAddressMode OutAddressMode = VK_SAMPLER_ADDRESS_MODE_MAX_ENUM;
-
 	switch (InAddressMode)
 	{
-		case AM_Wrap:		OutAddressMode = VK_SAMPLER_ADDRESS_MODE_REPEAT;			break;
-		case AM_Clamp:		OutAddressMode = VK_SAMPLER_ADDRESS_MODE_CLAMP_TO_EDGE;		break;
-		case AM_Mirror:		OutAddressMode = VK_SAMPLER_ADDRESS_MODE_MIRRORED_REPEAT;	break;
-		case AM_Border:		OutAddressMode = VK_SAMPLER_ADDRESS_MODE_CLAMP_TO_BORDER;	break;
+		case AM_Wrap:		return VK_SAMPLER_ADDRESS_MODE_REPEAT;
+		case AM_Clamp:		return VK_SAMPLER_ADDRESS_MODE_CLAMP_TO_EDGE;
+		case AM_Mirror:		return VK_SAMPLER_ADDRESS_MODE_MIRRORED_REPEAT;
+		case AM_Border:		return VK_SAMPLER_ADDRESS_MODE_CLAMP_TO_BORDER;
 		default:
 			break;
 	}
 
-	// Check for missing translation
-	check(OutAddressMode != VK_SAMPLER_ADDRESS_MODE_MAX_ENUM);
-	return OutAddressMode;
+	checkf(0, TEXT("Unknown Wrap ESamplerAddressMode %d"), (uint32)InAddressMode);
+	return VK_SAMPLER_ADDRESS_MODE_MAX_ENUM;
 }
 
 inline VkCompareOp TranslateSamplerCompareFunction(ESamplerCompareFunction InSamplerComparisonFunction)
 {
-	VkCompareOp OutSamplerComparisonFunction = VK_COMPARE_OP_MAX_ENUM;
-
 	switch (InSamplerComparisonFunction)
 	{
-		case SCF_Less:	OutSamplerComparisonFunction = VK_COMPARE_OP_LESS;	break;
-		case SCF_Never:	OutSamplerComparisonFunction = VK_COMPARE_OP_NEVER;	break;
-		default:															break;
+		case SCF_Less:	return VK_COMPARE_OP_LESS;
+		case SCF_Never:	return VK_COMPARE_OP_NEVER;
+		default:
+			break;
 	};
 
-	// Check for missing translation
-	check(OutSamplerComparisonFunction != VK_COMPARE_OP_MAX_ENUM);
-	return OutSamplerComparisonFunction;
+	checkf(0, TEXT("Unknown ESamplerCompareFunction %d"), (uint32)InSamplerComparisonFunction);
+	return VK_COMPARE_OP_MAX_ENUM;
 }
 
 static inline VkBlendOp BlendOpToVulkan(EBlendOperation InOp)
 {
-	VkBlendOp OutOp = VK_BLEND_OP_MAX_ENUM;
-
 	switch (InOp)
 	{
-		case BO_Add:				OutOp = VK_BLEND_OP_ADD;				break;
-		case BO_Subtract:			OutOp = VK_BLEND_OP_SUBTRACT;			break;
-		case BO_Min:				OutOp = VK_BLEND_OP_MIN;				break;
-		case BO_Max:				OutOp = VK_BLEND_OP_MAX;				break;
-		case BO_ReverseSubtract:	OutOp = VK_BLEND_OP_REVERSE_SUBTRACT;	break;
-		default:															break;
+		case BO_Add:				return VK_BLEND_OP_ADD;
+		case BO_Subtract:			return VK_BLEND_OP_SUBTRACT;
+		case BO_Min:				return VK_BLEND_OP_MIN;
+		case BO_Max:				return VK_BLEND_OP_MAX;
+		case BO_ReverseSubtract:	return VK_BLEND_OP_REVERSE_SUBTRACT;
+		default:
+			break;
 	}
 
-	// Check for missing translation
-	check(OutOp != VK_BLEND_OP_MAX_ENUM);
-	return OutOp;
+	checkf(0, TEXT("Unknown EBlendOperation %d"), (uint32)InOp);
+	return VK_BLEND_OP_MAX_ENUM;
 }
 
 static inline VkBlendFactor BlendFactorToVulkan(EBlendFactor InFactor)
 {
-	VkBlendFactor BlendMode = VK_BLEND_FACTOR_MAX_ENUM;
-
 	switch (InFactor)
 	{
-		case BF_Zero:						BlendMode = VK_BLEND_FACTOR_ZERO;						break;
-		case BF_One:						BlendMode = VK_BLEND_FACTOR_ONE;						break;
-		case BF_SourceColor:				BlendMode = VK_BLEND_FACTOR_SRC_COLOR;					break;
-		case BF_InverseSourceColor:			BlendMode = VK_BLEND_FACTOR_ONE_MINUS_SRC_COLOR;		break;
-		case BF_SourceAlpha:				BlendMode = VK_BLEND_FACTOR_SRC_ALPHA;					break;
-		case BF_InverseSourceAlpha:			BlendMode = VK_BLEND_FACTOR_ONE_MINUS_SRC_ALPHA;		break;
-		case BF_DestAlpha:					BlendMode = VK_BLEND_FACTOR_DST_ALPHA;					break;
-		case BF_InverseDestAlpha:			BlendMode = VK_BLEND_FACTOR_ONE_MINUS_DST_ALPHA;		break;
-		case BF_DestColor:					BlendMode = VK_BLEND_FACTOR_DST_COLOR;					break;
-		case BF_InverseDestColor:			BlendMode = VK_BLEND_FACTOR_ONE_MINUS_DST_COLOR;		break;
-		case BF_ConstantBlendFactor:		BlendMode = VK_BLEND_FACTOR_CONSTANT_COLOR;				break;
-		case BF_InverseConstantBlendFactor:	BlendMode = VK_BLEND_FACTOR_ONE_MINUS_CONSTANT_COLOR;	break;
-		default:																					break;
+		case BF_Zero:						return VK_BLEND_FACTOR_ZERO;
+		case BF_One:						return VK_BLEND_FACTOR_ONE;
+		case BF_SourceColor:				return VK_BLEND_FACTOR_SRC_COLOR;
+		case BF_InverseSourceColor:			return VK_BLEND_FACTOR_ONE_MINUS_SRC_COLOR;
+		case BF_SourceAlpha:				return VK_BLEND_FACTOR_SRC_ALPHA;
+		case BF_InverseSourceAlpha:			return VK_BLEND_FACTOR_ONE_MINUS_SRC_ALPHA;
+		case BF_DestAlpha:					return VK_BLEND_FACTOR_DST_ALPHA;
+		case BF_InverseDestAlpha:			return VK_BLEND_FACTOR_ONE_MINUS_DST_ALPHA;
+		case BF_DestColor:					return VK_BLEND_FACTOR_DST_COLOR;
+		case BF_InverseDestColor:			return VK_BLEND_FACTOR_ONE_MINUS_DST_COLOR;
+		case BF_ConstantBlendFactor:		return VK_BLEND_FACTOR_CONSTANT_COLOR;
+		case BF_InverseConstantBlendFactor:	return VK_BLEND_FACTOR_ONE_MINUS_CONSTANT_COLOR;
+		default:
+			break;
 	}
 
-	// Check for missing translation
-	check(BlendMode != VK_BLEND_FACTOR_MAX_ENUM);
-	return BlendMode;
+	checkf(0, TEXT("Unknown EBlendFactor %d"), (uint32)InFactor);
+	return VK_BLEND_FACTOR_MAX_ENUM;
 }
 
 static inline VkCompareOp CompareOpToVulkan(ECompareFunction InOp)
 {
-	VkCompareOp OutOp = VK_COMPARE_OP_MAX_ENUM;
-
 	switch (InOp)
 	{
-		case CF_Less:			OutOp = VK_COMPARE_OP_LESS;				break;
-		case CF_LessEqual:		OutOp = VK_COMPARE_OP_LESS_OR_EQUAL;	break;
-		case CF_Greater:		OutOp = VK_COMPARE_OP_GREATER;			break;
-		case CF_GreaterEqual:	OutOp = VK_COMPARE_OP_GREATER_OR_EQUAL;	break;
-		case CF_Equal:			OutOp = VK_COMPARE_OP_EQUAL;			break;
-		case CF_NotEqual:		OutOp = VK_COMPARE_OP_NOT_EQUAL;		break;
-		case CF_Never:			OutOp = VK_COMPARE_OP_NEVER;			break;
-		case CF_Always:			OutOp = VK_COMPARE_OP_ALWAYS;			break;
-		default:														break;
+		case CF_Less:			return VK_COMPARE_OP_LESS;
+		case CF_LessEqual:		return VK_COMPARE_OP_LESS_OR_EQUAL;
+		case CF_Greater:		return VK_COMPARE_OP_GREATER;
+		case CF_GreaterEqual:	return VK_COMPARE_OP_GREATER_OR_EQUAL;
+		case CF_Equal:			return VK_COMPARE_OP_EQUAL;
+		case CF_NotEqual:		return VK_COMPARE_OP_NOT_EQUAL;
+		case CF_Never:			return VK_COMPARE_OP_NEVER;
+		case CF_Always:			return VK_COMPARE_OP_ALWAYS;
+		default:
+			break;
 	}
 
-	// Check for missing translation
-	check(OutOp != VK_COMPARE_OP_MAX_ENUM);
-	return OutOp;
+	checkf(0, TEXT("Unknown ECompareFunction %d"), (uint32)InOp);
+	return VK_COMPARE_OP_MAX_ENUM;
 }
 
 static inline VkStencilOp StencilOpToVulkan(EStencilOp InOp)
@@ -171,52 +138,49 @@ static inline VkStencilOp StencilOpToVulkan(EStencilOp InOp)
 
 	switch (InOp)
 	{
-		case SO_Keep:					OutOp = VK_STENCIL_OP_KEEP;					break;
-		case SO_Zero:					OutOp = VK_STENCIL_OP_ZERO;					break;
-		case SO_Replace:				OutOp = VK_STENCIL_OP_REPLACE;				break;
-		case SO_SaturatedIncrement:		OutOp = VK_STENCIL_OP_INCREMENT_AND_CLAMP;	break;
-		case SO_SaturatedDecrement:		OutOp = VK_STENCIL_OP_DECREMENT_AND_CLAMP;	break;
-		case SO_Invert:					OutOp = VK_STENCIL_OP_INVERT;				break;
-		case SO_Increment:				OutOp = VK_STENCIL_OP_INCREMENT_AND_WRAP;	break;
-		case SO_Decrement:				OutOp = VK_STENCIL_OP_DECREMENT_AND_WRAP;	break;
-		default:																	break;
+		case SO_Keep:					return VK_STENCIL_OP_KEEP;
+		case SO_Zero:					return VK_STENCIL_OP_ZERO;
+		case SO_Replace:				return VK_STENCIL_OP_REPLACE;
+		case SO_SaturatedIncrement:		return VK_STENCIL_OP_INCREMENT_AND_CLAMP;
+		case SO_SaturatedDecrement:		return VK_STENCIL_OP_DECREMENT_AND_CLAMP;
+		case SO_Invert:					return VK_STENCIL_OP_INVERT;
+		case SO_Increment:				return VK_STENCIL_OP_INCREMENT_AND_WRAP;
+		case SO_Decrement:				return VK_STENCIL_OP_DECREMENT_AND_WRAP;
+		default:
+			break;
 	}
 
-	check(OutOp != VK_STENCIL_OP_MAX_ENUM);
-	return OutOp;
+	checkf(0, TEXT("Unknown EStencilOp %d"), (uint32)InOp);
+	return VK_STENCIL_OP_MAX_ENUM;
 }
 
 static inline VkPolygonMode RasterizerFillModeToVulkan(ERasterizerFillMode InFillMode)
 {
-	VkPolygonMode OutFillMode = VK_POLYGON_MODE_MAX_ENUM;
-
 	switch (InFillMode)
 	{
-		case FM_Point:			OutFillMode = VK_POLYGON_MODE_POINT;	break;
-		case FM_Wireframe:		OutFillMode = VK_POLYGON_MODE_LINE;		break;
-		case FM_Solid:			OutFillMode = VK_POLYGON_MODE_FILL;		break;
-		default:														break;
+		case FM_Point:			return VK_POLYGON_MODE_POINT;
+		case FM_Wireframe:		return VK_POLYGON_MODE_LINE;
+		case FM_Solid:			return VK_POLYGON_MODE_FILL;
+		default:
+			break;
 	}
 
-	// Check for missing translation
-	check(OutFillMode != VK_POLYGON_MODE_MAX_ENUM);
-	return OutFillMode;
+	checkf(0, TEXT("Unknown ERasterizerFillMode %d"), (uint32)InFillMode);
+	return VK_POLYGON_MODE_MAX_ENUM;
 }
 
 static inline VkCullModeFlags RasterizerCullModeToVulkan(ERasterizerCullMode InCullMode)
 {
-	VkCullModeFlags outCullMode = VK_CULL_MODE_NONE;
-
 	switch (InCullMode)
 	{
-		case CM_None:	outCullMode = VK_CULL_MODE_NONE;		break;
-		case CM_CW:		outCullMode = VK_CULL_MODE_FRONT_BIT;	break;
-		case CM_CCW:	outCullMode = VK_CULL_MODE_BACK_BIT;	break;
-			// Check for missing translation
-		default:		check(false);							break;
+		case CM_None:	return VK_CULL_MODE_NONE;
+		case CM_CW:		return VK_CULL_MODE_FRONT_BIT;
+		case CM_CCW:	return VK_CULL_MODE_BACK_BIT;
+		default:		break;
 	}
 
-	return outCullMode;
+	checkf(0, TEXT("Unknown ERasterizerCullMode %d"), (uint32)InCullMode);
+	return VK_CULL_MODE_NONE;
 }
 
 
@@ -224,8 +188,8 @@ void FVulkanSamplerState::SetupSamplerCreateInfo(const FSamplerStateInitializerR
 {
 	ZeroVulkanStruct(OutSamplerInfo, VK_STRUCTURE_TYPE_SAMPLER_CREATE_INFO);
 
-	OutSamplerInfo.magFilter = TranslateMagFilterMode(Initializer.Filter);
-	OutSamplerInfo.minFilter = TranslateMinFilterMode(Initializer.Filter);
+	OutSamplerInfo.magFilter = TranslateMinMagFilterMode(Initializer.Filter);
+	OutSamplerInfo.minFilter = TranslateMinMagFilterMode(Initializer.Filter);
 	OutSamplerInfo.mipmapMode = TranslateMipFilterMode(Initializer.Filter);
 	OutSamplerInfo.addressModeU = TranslateWrapMode(Initializer.AddressU);
 	OutSamplerInfo.addressModeV = TranslateWrapMode(Initializer.AddressV);
