Vulkan: do not enable draw events unnecessarily.

diff --git a/Engine/Source/Runtime/VulkanRHI/Private/VulkanRHI.cpp b/Engine/Source/Runtime/VulkanRHI/Private/VulkanRHI.cpp
--- a/Engine/Source/Runtime/VulkanRHI/Private/VulkanRHI.cpp
+++ b/Engine/Source/Runtime/VulkanRHI/Private/VulkanRHI.cpp
@@ -432,7 +432,7 @@ void FVulkanDynamicRHI::CreateInstance()
 #if VULKAN_HAS_DEBUGGING_ENABLED
 	SetupDebugLayerCallback();
 
-	if (!GRHISupportsRHIThread || GRenderDocFound)
+	if (GRenderDocFound)
 	{
 		EnableIdealGPUCaptureOptions(true);
 	}
