Linux: minor Vulkan optimziation.

- reduce calls to check the CVar in the editor
- for cooked games, don't check the CVar at all.

diff --git a/Engine/Source/Runtime/VulkanRHI/Private/Linux/VulkanLinuxPlatform.h b/Engine/Source/Runtime/VulkanRHI/Private/Linux/VulkanLinuxPlatform.h
--- a/Engine/Source/Runtime/VulkanRHI/Private/Linux/VulkanLinuxPlatform.h
+++ b/Engine/Source/Runtime/VulkanRHI/Private/Linux/VulkanLinuxPlatform.h
@@ -49,8 +49,20 @@ public:
 	// Some platforms only support real or non-real UBs, so this function can optimize it out
 	static bool UseRealUBsOptimization(bool bCodeHeaderUseRealUBs)
 	{
-		static auto* CVar = IConsoleManager::Get().FindTConsoleVariableDataInt(TEXT("r.Vulkan.UseRealUBs"));
-		return (CVar && CVar->GetValueOnAnyThread() == 0) ? false : bCodeHeaderUseRealUBs;
+		// cooked builds will return the bool unchanged - relying on the compiler to optimize out the editor code path
+		if (UE_EDITOR)
+		{
+			static bool bAlwaysUseRealUBs([]()
+			{
+				static auto* CVar = IConsoleManager::Get().FindTConsoleVariableDataInt(TEXT("r.Vulkan.UseRealUBs"));
+				return (CVar && CVar->GetValueOnAnyThread() == 0);
+			});
+			return bAlwaysUseRealUBs ? false : bCodeHeaderUseRealUBs;
+		}
+		else
+		{
+			return bCodeHeaderUseRealUBs;
+		}
 	}
 
 	static void WriteCrashMarker(const FOptionalVulkanDeviceExtensions& OptionalExtensions, VkCommandBuffer CmdBuffer, VkBuffer DestBuffer, const TArrayView<uint32>& Entries, bool bAdding);
