UE4.22 - Prep for hotfix changes
#rnx

diff --git a/Engine/Source/Runtime/VulkanRHI/Private/VulkanCommandBuffer.h b/Engine/Source/Runtime/VulkanRHI/Private/VulkanCommandBuffer.h
--- a/Engine/Source/Runtime/VulkanRHI/Private/VulkanCommandBuffer.h
+++ b/Engine/Source/Runtime/VulkanRHI/Private/VulkanCommandBuffer.h
@@ -63,6 +63,11 @@ public:
 		return State == EState::Submitted;
 	}
 
+	inline bool IsAllocated() const
+	{
+		return State != EState::NotAllocated;
+	}
+
 	inline VkCommandBuffer GetHandle()
 	{
 		return CommandBufferHandle;
diff --git a/Engine/Source/Runtime/VulkanRHI/Private/VulkanDescriptorSets.h b/Engine/Source/Runtime/VulkanRHI/Private/VulkanDescriptorSets.h
--- a/Engine/Source/Runtime/VulkanRHI/Private/VulkanDescriptorSets.h
+++ b/Engine/Source/Runtime/VulkanRHI/Private/VulkanDescriptorSets.h
@@ -969,12 +969,11 @@ public:
 		bool bChanged = false;
 		if (UseVulkanDescriptorCache())
 		{
-			
-			FVulkanHashableDescriptorInfo& HeshableInfo = HashableDescriptorInfos[DescriptorIndex];
+			FVulkanHashableDescriptorInfo& HashableInfo = HashableDescriptorInfos[DescriptorIndex];
 			check(Sampler.SamplerId > 0);
-			if (HeshableInfo.Image.SamplerId != Sampler.SamplerId)
+			if (HashableInfo.Image.SamplerId != Sampler.SamplerId)
 			{
-				HeshableInfo.Image.SamplerId = Sampler.SamplerId;
+				HashableInfo.Image.SamplerId = Sampler.SamplerId;
 				ImageInfo->sampler = Sampler.Sampler;
 				bChanged = true;
 			}
@@ -1050,23 +1049,23 @@ protected:
 		bool bChanged = false;
 		if (UseVulkanDescriptorCache())
 		{
-			FVulkanHashableDescriptorInfo& HeshableInfo = HashableDescriptorInfos[DescriptorIndex];
+			FVulkanHashableDescriptorInfo& HashableInfo = HashableDescriptorInfos[DescriptorIndex];
 			check(BufferAllocation.GetHandleId() > 0);
-			if (HeshableInfo.Buffer.Id != BufferAllocation.GetHandleId())
+			if (HashableInfo.Buffer.Id != BufferAllocation.GetHandleId())
 			{
-				HeshableInfo.Buffer.Id = BufferAllocation.GetHandleId();
+				HashableInfo.Buffer.Id = BufferAllocation.GetHandleId();
 				BufferInfo->buffer = BufferAllocation.GetHandle();
 				bChanged = true;
 			}
-			if (HeshableInfo.Buffer.Offset != static_cast<uint32>(Offset))
+			if (HashableInfo.Buffer.Offset != static_cast<uint32>(Offset))
 			{
-				HeshableInfo.Buffer.Offset = static_cast<uint32>(Offset);
+				HashableInfo.Buffer.Offset = static_cast<uint32>(Offset);
 				BufferInfo->offset = Offset;
 				bChanged = true;
 			}
-			if (HeshableInfo.Buffer.Range != static_cast<uint32>(Range))
+			if (HashableInfo.Buffer.Range != static_cast<uint32>(Range))
 			{
-				HeshableInfo.Buffer.Range = static_cast<uint32>(Range);
+				HashableInfo.Buffer.Range = static_cast<uint32>(Range);
 				BufferInfo->range = Range;
 				bChanged = true;
 			}
@@ -1105,17 +1104,17 @@ protected:
 		bool bChanged = false;
 		if (UseVulkanDescriptorCache())
 		{
-			FVulkanHashableDescriptorInfo& HeshableInfo = HashableDescriptorInfos[DescriptorIndex];
+			FVulkanHashableDescriptorInfo& HashableInfo = HashableDescriptorInfos[DescriptorIndex];
 			check(TextureView.ViewId > 0);
-			if (HeshableInfo.Image.ImageViewId != TextureView.ViewId)
+			if (HashableInfo.Image.ImageViewId != TextureView.ViewId)
 			{
-				HeshableInfo.Image.ImageViewId = TextureView.ViewId;
+				HashableInfo.Image.ImageViewId = TextureView.ViewId;
 				ImageInfo->imageView = TextureView.View;
 				bChanged = true;
 			}
-			if (HeshableInfo.Image.ImageLayout != static_cast<uint32>(Layout))
+			if (HashableInfo.Image.ImageLayout != static_cast<uint32>(Layout))
 			{
-				HeshableInfo.Image.ImageLayout = static_cast<uint32>(Layout);
+				HashableInfo.Image.ImageLayout = static_cast<uint32>(Layout);
 				ImageInfo->imageLayout = Layout;
 				bChanged = true;
 			}
@@ -1141,11 +1140,11 @@ protected:
 		if (UseVulkanDescriptorCache())
 		{
 			bool bChanged = false;
-			FVulkanHashableDescriptorInfo& HeshableInfo = HashableDescriptorInfos[DescriptorIndex];
+			FVulkanHashableDescriptorInfo& HashableInfo = HashableDescriptorInfos[DescriptorIndex];
 			check(View->ViewId > 0);
-			if (HeshableInfo.BufferView.Id != View->ViewId)
+			if (HashableInfo.BufferView.Id != View->ViewId)
 			{
-				HeshableInfo.BufferView.Id = View->ViewId;
+				HashableInfo.BufferView.Id = View->ViewId;
 				bChanged = true;
 			}
 			bIsKeyDirty |= bChanged;
diff --git a/Engine/Source/Runtime/VulkanRHI/Private/VulkanDevice.h b/Engine/Source/Runtime/VulkanRHI/Private/VulkanDevice.h
--- a/Engine/Source/Runtime/VulkanRHI/Private/VulkanDevice.h
+++ b/Engine/Source/Runtime/VulkanRHI/Private/VulkanDevice.h
@@ -173,6 +173,11 @@ public:
 		return MemoryManager;
 	}
 
+	inline const VkPhysicalDeviceMemoryProperties& GetDeviceMemoryProperties() const
+	{
+		return MemoryManager.GetMemoryProperties();
+	}
+
 	inline VulkanRHI::FResourceHeapManager& GetResourceHeapManager()
 	{
 		return ResourceHeapManager;
@@ -299,6 +304,8 @@ public:
 	VkSamplerYcbcrConversion CreateSamplerColorConversion(const VkSamplerYcbcrConversionCreateInfo& CreateInfo);
 #endif
 
+	void*	Hotfix;
+
 private:
 	void MapFormatSupport(EPixelFormat UEFormat, VkFormat VulkanFormat);
 	void MapFormatSupportWithFallback(EPixelFormat UEFormat, VkFormat VulkanFormat, TArrayView<const VkFormat> FallbackTextureFormats);
diff --git a/Engine/Source/Runtime/VulkanRHI/Private/VulkanMemory.cpp b/Engine/Source/Runtime/VulkanRHI/Private/VulkanMemory.cpp
--- a/Engine/Source/Runtime/VulkanRHI/Private/VulkanMemory.cpp
+++ b/Engine/Source/Runtime/VulkanRHI/Private/VulkanMemory.cpp
@@ -234,7 +234,7 @@ namespace VulkanRHI
 		{
 			if (bCanFail)
 			{
-				UE_LOG(LogVulkanRHI, Warning, TEXT("Failed to allocate Device Memory, Requested=%fKb MemTypeIndex=%d"), (float)Info.allocationSize / 1024.0f, Info.memoryTypeIndex);
+				UE_LOG(LogVulkanRHI, Warning, TEXT("Failed to allocate Device Memory, Requested=%.2fKb MemTypeIndex=%d"), (float)Info.allocationSize / 1024.0f, Info.memoryTypeIndex);
 				return nullptr;
 			}
 			auto Callback = [=]()
@@ -243,7 +243,7 @@ namespace VulkanRHI
 				DumpMemory();
 				GLog->PanicFlushThreadedLogs();
 #endif
-				return FString::Printf(TEXT("Out of Device Memory, Requested=%fKb MemTypeIndex=%d"), (float)Info.allocationSize / 1024.0f, Info.memoryTypeIndex);
+				return FString::Printf(TEXT("Out of Device Memory, Requested=%.2fKb MemTypeIndex=%d"), (float)Info.allocationSize / 1024.0f, Info.memoryTypeIndex);
 			};
 			UE_LOG(LogVulkanRHI, Fatal, TEXT("%s"), *(Callback()));
 		}
@@ -251,7 +251,7 @@ namespace VulkanRHI
 		{
 			if (bCanFail)
 			{
-				UE_LOG(LogVulkanRHI, Warning, TEXT("Failed to allocate Host Memory, Requested=%fKb MemTypeIndex=%d"), (float)Info.allocationSize / 1024.0f, Info.memoryTypeIndex);
+				UE_LOG(LogVulkanRHI, Warning, TEXT("Failed to allocate Host Memory, Requested=%.2fKb MemTypeIndex=%d"), (float)Info.allocationSize / 1024.0f, Info.memoryTypeIndex);
 				return nullptr;
 			}
 			auto Callback = [=]()
@@ -260,7 +260,7 @@ namespace VulkanRHI
 				DumpMemory();
 				GLog->PanicFlushThreadedLogs();
 #endif
-				return FString::Printf(TEXT("Out of Host Memory, Requested=%fKb MemTypeIndex=%d"), (float)Info.allocationSize / 1024.0f, Info.memoryTypeIndex);
+				return FString::Printf(TEXT("Out of Host Memory, Requested=%.2fKb MemTypeIndex=%d"), (float)Info.allocationSize / 1024.0f, Info.memoryTypeIndex);
 			};
 			UE_LOG(LogVulkanRHI, Error, TEXT("%s"), *(Callback()));
 		}
@@ -1530,9 +1530,9 @@ namespace VulkanRHI
 			for (int32 Index = 0; Index < FreeStagingBuffers.Num(); ++Index)
 			{
 				FFreeEntry& FreeBuffer = FreeStagingBuffers[Index];
-				if (FreeBuffer.Buffer->GetSize() == Size && FreeBuffer.Buffer->bCPURead == bCPURead)
+				if (FreeBuffer.StagingBuffer->GetSize() == Size && FreeBuffer.StagingBuffer->bCPURead == bCPURead)
 				{
-					FStagingBuffer* Buffer = FreeBuffer.Buffer;
+					FStagingBuffer* Buffer = FreeBuffer.StagingBuffer;
 					FreeStagingBuffers.RemoveAtSwap(Index, 1, false);
 					UsedStagingBuffers.Add(Buffer);
 					return Buffer;
@@ -1620,6 +1620,7 @@ namespace VulkanRHI
 		{
 			FPendingItemsPerCmdBuffer* ItemsForCmdBuffer = FindOrAdd(CmdBuffer);
 			FPendingItemsPerCmdBuffer::FPendingItems* ItemsForFence = ItemsForCmdBuffer->FindOrAddItemsForFence(CmdBuffer->GetFenceSignaledCounterA());
+			check(StagingBuffer);
 			ItemsForFence->Resources.Add(StagingBuffer);
 		}
 		else
@@ -1661,7 +1662,7 @@ namespace VulkanRHI
 		for (int32 Index = 0; Index < FreeStagingBuffers.Num(); ++Index)
 		{
 			FFreeEntry& Entry = FreeStagingBuffers[Index];
-			UE_LOG(LogVulkanRHI, Display, TEXT("%6d %p %p"), Index, (void*)Entry.Buffer->GetHandle(), (void*)Entry.Buffer->ResourceAllocation->GetHandle());
+			UE_LOG(LogVulkanRHI, Display, TEXT("%6d %p %p"), Index, (void*)Entry.StagingBuffer->GetHandle(), (void*)Entry.StagingBuffer->ResourceAllocation->GetHandle());
 		}
 	}
 #endif
@@ -1679,6 +1680,7 @@ namespace VulkanRHI
 				{
 					for (int32 ResourceIndex = 0; ResourceIndex < PendingItems.Resources.Num(); ++ResourceIndex)
 					{
+						check(PendingItems.Resources[ResourceIndex]);
 						FreeStagingBuffers.Add({PendingItems.Resources[ResourceIndex], GFrameNumberRenderThread});
 					}
 
@@ -1700,9 +1702,9 @@ namespace VulkanRHI
 				FFreeEntry& Entry = FreeStagingBuffers[Index];
 				if (bImmediately || Entry.FrameNumber + NUM_FRAMES_TO_WAIT_BEFORE_RELEASING_TO_OS < GFrameNumberRenderThread)
 				{
-					UsedMemory -= Entry.Buffer->GetSize();
-					Entry.Buffer->Destroy(Device);
-					delete Entry.Buffer;
+					UsedMemory -= Entry.StagingBuffer->GetSize();
+					Entry.StagingBuffer->Destroy(Device);
+					delete Entry.StagingBuffer;
 					FreeStagingBuffers.RemoveAtSwap(Index, 1, false);
 				}
 			}
@@ -1910,7 +1912,7 @@ namespace VulkanRHI
 				{ 
 					return InEntry.Handle == Entry.Handle; 
 				});
-			checkf(ExistingEntry == nullptr, TEXT("Attempt to double-delete resource, Type: %d, Handle: %llu"), (int32)Type, Handle);
+			checkf(ExistingEntry == nullptr, TEXT("Attempt to double-delete resource, FDeferredDeletionQueue::EType: %d, Handle: %llu"), (int32)Type, Handle);
 #endif
 
 			Entries.Add(Entry);
diff --git a/Engine/Source/Runtime/VulkanRHI/Public/VulkanContext.h b/Engine/Source/Runtime/VulkanRHI/Public/VulkanContext.h
--- a/Engine/Source/Runtime/VulkanRHI/Public/VulkanContext.h
+++ b/Engine/Source/Runtime/VulkanRHI/Public/VulkanContext.h
@@ -369,6 +369,8 @@ public:
 
 	void PrepareParallelFromBase(const FVulkanCommandListContext& BaseContext);
 
+	void* Hotfix;
+
 protected:
 	FVulkanDynamicRHI* RHI;
 	FVulkanCommandListContext* Immediate;
diff --git a/Engine/Source/Runtime/VulkanRHI/Public/VulkanMemory.h b/Engine/Source/Runtime/VulkanRHI/Public/VulkanMemory.h
--- a/Engine/Source/Runtime/VulkanRHI/Public/VulkanMemory.h
+++ b/Engine/Source/Runtime/VulkanRHI/Public/VulkanMemory.h
@@ -834,6 +834,8 @@ namespace VulkanRHI
 		void DumpMemory();
 #endif
 
+		void* Hotfix;
+
 	protected:
 		FDeviceMemoryManager* DeviceMemoryManager;
 		TArray<FOldResourceHeap*> ResourceTypeHeaps;
@@ -906,10 +908,7 @@ namespace VulkanRHI
 
 		TArray<FBufferAllocation*> UsedBufferAllocations[(int32)EPoolSizes::SizesCount + 1];
 		TArray<FBufferAllocation*> FreeBufferAllocations[(int32)EPoolSizes::SizesCount + 1];
-#if 0
-		TArray<FImageAllocation*> UsedImageAllocations;
-		TArray<FImageAllocation*> FreeImageAllocations;
-#endif
+
 		void ReleaseFreedResources(bool bImmediately);
 		void DestroyResourceAllocations();
 	};
@@ -935,11 +934,6 @@ namespace VulkanRHI
 			return ResourceAllocation->GetMappedPointer();
 		}
 
-		inline uint32 GetAllocationOffset() const
-		{
-			return ResourceAllocation->GetOffset();
-		}
-
 		inline uint32 GetSize() const
 		{
 			return BufferSize;
@@ -1023,7 +1017,7 @@ namespace VulkanRHI
 		TArray<FPendingItemsPerCmdBuffer> PendingFreeStagingBuffers;
 		struct FFreeEntry
 		{
-			FStagingBuffer* Buffer;
+			FStagingBuffer* StagingBuffer;
 			uint32 FrameNumber;
 		};
 		TArray<FFreeEntry> FreeStagingBuffers;
@@ -1587,7 +1581,6 @@ namespace VulkanRHI
 		FVulkanDevice& Device;
 		VkSemaphore SemaphoreHandle;
 	};
-
 }
 
 
