UE4.22 - Potential fix for multithreaded access
#rnx

diff --git a/Engine/Source/Runtime/VulkanRHI/Private/VulkanRHI.cpp b/Engine/Source/Runtime/VulkanRHI/Private/VulkanRHI.cpp
--- a/Engine/Source/Runtime/VulkanRHI/Private/VulkanRHI.cpp
+++ b/Engine/Source/Runtime/VulkanRHI/Private/VulkanRHI.cpp
@@ -1190,8 +1190,11 @@ void FVulkanDescriptorSetsLayoutInfo::GenerateHash(const TArrayView<const FSampl
 #endif
 }
 
+static FCriticalSection GTypesUsageCS;
 void FVulkanDescriptorSetsLayoutInfo::CompileTypesUsageID()
 {
+	FScopeLock ScopeLock(&GTypesUsageCS);
+
 	static TMap<uint32, uint32> GTypesUsageHashMap;
 	static uint32 GUniqueID = 1;
 
