UE4.22 - Add -vulkandebug as an option to match -d3debug and -opengldebug
#rnx

diff --git a/Engine/Source/Runtime/VulkanRHI/Private/VulkanLayers.cpp b/Engine/Source/Runtime/VulkanRHI/Private/VulkanLayers.cpp
--- a/Engine/Source/Runtime/VulkanRHI/Private/VulkanLayers.cpp
+++ b/Engine/Source/Runtime/VulkanRHI/Private/VulkanLayers.cpp
@@ -310,7 +310,12 @@ void FVulkanDynamicRHI::GetInstanceLayersAndExtensions(TArray<const ANSICHAR*>&
 #endif	// VULKAN_ENABLE_API_DUMP
 
 	int32 VulkanValidationOption = GValidationCvar.GetValueOnAnyThread();
-	if (FParse::Value(FCommandLine::Get(), TEXT("vulkanvalidation="), VulkanValidationOption))
+	if (FParse::Param(FCommandLine::Get(), TEXT("vulkandebug")))
+	{
+		// Match D3D and GL
+		GValidationCvar->Set(2, ECVF_SetByCommandline);
+	}
+	else if (FParse::Value(FCommandLine::Get(), TEXT("vulkanvalidation="), VulkanValidationOption))
 	{
 		GValidationCvar->Set(VulkanValidationOption, ECVF_SetByCommandline);
 	}
