UE4.22 - Prep for timestamp changes
#rnx

diff --git a/Engine/Source/Runtime/VulkanRHI/Private/VulkanCommandBuffer.h b/Engine/Source/Runtime/VulkanRHI/Private/VulkanCommandBuffer.h
--- a/Engine/Source/Runtime/VulkanRHI/Private/VulkanCommandBuffer.h
+++ b/Engine/Source/Runtime/VulkanRHI/Private/VulkanCommandBuffer.h
@@ -174,14 +174,8 @@ public:
 
 	bool AcquirePoolSetAndDescriptorsIfNeeded(const class FVulkanDescriptorSetsLayout& Layout, bool bNeedDescriptors, VkDescriptorSet* OutDescriptors);
 
-	/*TRefCountPtr<*/FVulkanTimestampQueryPool* /*>*/ PrepareTimestampQueryPool()
-	{
-		return TimestampQueryPool;
-	}
-
 private:
 	FVulkanDevice* Device;
-	/*TRefCountPtr<*/FVulkanTimestampQueryPool* /*>*/ TimestampQueryPool = nullptr;
 	VkCommandBuffer CommandBufferHandle;
 	double SubmittedTime = 0.0f;
 
diff --git a/Engine/Source/Runtime/VulkanRHI/Private/VulkanQuery.cpp b/Engine/Source/Runtime/VulkanRHI/Private/VulkanQuery.cpp
--- a/Engine/Source/Runtime/VulkanRHI/Private/VulkanQuery.cpp
+++ b/Engine/Source/Runtime/VulkanRHI/Private/VulkanQuery.cpp
@@ -32,94 +32,7 @@ constexpr int32 NUM_FRAMES_TO_WAIT_REUSE_POOL = 10;
 constexpr uint32 NUM_FRAMES_TO_WAIT_RELEASE_POOL = MAX_uint32; // never release
 #endif
 
-/*
-FCriticalSection GOcclusionQueryCS;
-const uint32 GMaxLifetimeTimestampQueries = 1024;
-
-FVulkanRenderQuery::FVulkanRenderQuery(FVulkanDevice* Device, ERenderQueryType InQueryType)
-	: QueryType(InQueryType)
-{
-	INC_DWORD_STAT(STAT_VulkanNumQueries);
-}
-
-FVulkanRenderQuery::~FVulkanRenderQuery()
-{
-	DEC_DWORD_STAT(STAT_VulkanNumQueries);
-}
-
-inline void FVulkanRenderQuery::Begin(FVulkanCmdBuffer* InCmdBuffer)
-{
-	BeginCmdBuffer = InCmdBuffer;
-	check(State == EState::Reset);
-	if (QueryType == RQT_Occlusion)
-	{
-		check(Pool);
-		VulkanRHI::vkCmdBeginQuery(InCmdBuffer->GetHandle(), Pool->GetHandle(), QueryIndex, VK_QUERY_CONTROL_PRECISE_BIT);
-		State = EState::InBegin;
-		LastPoolReset = Pool->NumResets;
-	}
-	else
-	{
-		ensureMsgf(0, TEXT("Timer queries do NOT support Begin()!"));
-	}
-}
-
-inline void FVulkanRenderQuery::End(FVulkanCmdBuffer* InCmdBuffer)
-{
-	ensure(QueryType != RQT_Occlusion || BeginCmdBuffer == InCmdBuffer);
-	if (QueryType == RQT_Occlusion)
-	{
-		check(State == EState::InBegin);
-		check(Pool);
-		VulkanRHI::vkCmdEndQuery(InCmdBuffer->GetHandle(), Pool->GetHandle(), QueryIndex);
-		check(LastPoolReset == Pool->NumResets);
-	}
-	else
-	{
-		check(State == EState::Reset);
-		BeginCmdBuffer = InCmdBuffer;
-		VulkanRHI::vkCmdWriteTimestamp(InCmdBuffer->GetHandle(), VK_PIPELINE_STAGE_BOTTOM_OF_PIPE_BIT, Pool->GetHandle(), QueryIndex);
-		LastPoolReset = Pool->NumResets;
-	}
-	State = EState::InEnd;
-}
-
-bool FVulkanRenderQuery::GetResult(FVulkanDevice* Device, uint64& OutResult, bool bWait)
-{
-FScopeLock ScopeLock(&GOcclusionQueryCS);
-	if (State == EState::HasResults)
-	{
-		OutResult = Result;
-		return true;
-	}
 
-	check(Pool);
-	check(State == EState::InEnd);
-	if (QueryType == RQT_Occlusion)
-	{
-		check(IsInRenderingThread());
-		if (((FVulkanOcclusionQueryPool*)Pool)->GetResults(QueryIndex, bWait, Result))
-		{
-			check(LastPoolReset == Pool->NumResets);
-			State = EState::HasResults;
-			OutResult = Result;
-			return true;
-		}
-	}
-	else
-	{
-		check(IsInRenderingThread() || IsInRHIThread());
-		if (((FVulkanTimestampQueryPool*)Pool)->GetResults(QueryIndex, bWait, Result))
-		{
-			check(LastPoolReset == Pool->NumResets);
-			State = EState::HasResults;
-			OutResult = Result;
-			return true;
-		}
-	}
-	return false;
-}
-*/
 FVulkanQueryPool::FVulkanQueryPool(FVulkanDevice* InDevice, uint32 InMaxQueries, VkQueryType InQueryType)
 	: VulkanRHI::FDeviceChild(InDevice)
 	, QueryPool(VK_NULL_HANDLE)
@@ -152,22 +65,6 @@ FVulkanQueryPool::~FVulkanQueryPool()
 	ResetEvent = VK_NULL_HANDLE;
 }
 
-/*
-inline VkResult FVulkanQueryPool::InternalGetQueryPoolResults(uint32 FirstQuery, uint32 NumQueries, VkQueryResultFlags ExtraFlags)
-{
-	VkResult Result = VulkanRHI::vkGetQueryPoolResults(Device->GetInstanceHandle(), QueryPool, FirstQuery, NumQueries, NumQueries * sizeof(uint64), QueryOutput.GetData() + FirstQuery, sizeof(uint64), VK_QUERY_RESULT_64_BIT | ExtraFlags);
-	return Result;
-}
-
-inline int32 FVulkanQueryPool::AllocateQuery()
-{
-	FScopeLock ScopeLock(&GOcclusionQueryCS);
-	int32 UsedQuery = NumUsedQueries;
-	++NumUsedQueries;
-	checkf((uint32)UsedQuery < MaxQueries, TEXT("Internal error allocating query! UsedQuery=%d NumUsedQueries=%d MaxQueries=%d"), UsedQuery, NumUsedQueries, MaxQueries);
-	return UsedQuery;
-}
-*/
 
 bool FVulkanOcclusionQueryPool::CanBeReused()
 {
@@ -275,104 +172,10 @@ bool FVulkanOcclusionQueryPool::InternalTryGetResults(bool bWait)
 	{
 		VERIFYVULKANRESULT(Result);
 	}
-/*
-	FScopeLock ScopeLock(&GOcclusionQueryCS);
-	check(!bHasResults);
-	VkResult Result = VK_NOT_READY;
-	bool bFencePassed = false;
-	if (CmdBuffer->GetFenceSignaledCounter() > FenceCounter)
-	{
-		bFencePassed = true;
-		Result = InternalGetQueryPoolResults(0);
-		if (Result == VK_SUCCESS)
-		{
-			bHasResults = true;
-			return true;
-		}
-	}
-
-	if (Result == VK_NOT_READY)
-	{
-		if (bWait)
-		{
-			uint32 IdleStart = FPlatformTime::Cycles();
-
-			SCOPE_CYCLE_COUNTER(STAT_VulkanWaitQuery);
-
-			// We'll do manual wait
-			double StartTime = FPlatformTime::Seconds();
-
-			ENamedThreads::Type RenderThread_Local = ENamedThreads::GetRenderThread_Local();
-			bool bSuccess = false;
-			int32 NumLoops = 0;
-			while (!bSuccess)
-			{
-				FPlatformProcess::SleepNoStats(0);
-
-				// pump RHIThread to make sure these queries have actually been submitted to the GPU.
-				if (IsInActualRenderingThread())
-				{
-					FTaskGraphInterface::Get().ProcessThreadUntilIdle(RenderThread_Local);
-				}
-
-				Result = InternalGetQueryPoolResults(0);
-				if (Result == VK_SUCCESS)
-				{
-					bSuccess = true;
-					break;
-				}
-				else if (Result == VK_NOT_READY)
-				{
-					bSuccess = false;
-				}
-				else
-				{
-					bSuccess = false;
-					VERIFYVULKANRESULT(Result);
-				}
-
-				// timer queries are used for Benchmarks which can stall a bit more
-				const double TimeoutValue = (QueryType == VK_QUERY_TYPE_TIMESTAMP) ? 2.0 : 0.5;
-				// look for gpu stuck/crashed
-				if ((FPlatformTime::Seconds() - StartTime) > TimeoutValue)
-				{
-					if (QueryType == VK_QUERY_TYPE_OCCLUSION)
-					{
-						UE_LOG(LogRHI, Log, TEXT("Timed out while waiting for GPU to catch up on occlusion results. (%.1f s)"), TimeoutValue);
-					}
-					else
-					{
-						UE_LOG(LogRHI, Log, TEXT("Timed out while waiting for GPU to catch up on occlusion/timer results. (%.1f s)"), TimeoutValue);
-					}
-					return false;
-				}
-
-				++NumLoops;
-			}
-
-			GRenderThreadIdle[ERenderThreadIdleTypes::WaitingForGPUQuery] += FPlatformTime::Cycles() - IdleStart;
-			GRenderThreadNumIdle[ERenderThreadIdleTypes::WaitingForGPUQuery]++;
 
-			bHasResults = true;
-			return true;
-		}
-	}
-	else
-	{
-		VERIFYVULKANRESULT(Result);
-	}
-	*/
 	return false;
 }
 
-/*
-void FVulkanQueryPool::ResetAll(FVulkanCmdBuffer* InCmdBuffer)
-{
-	VulkanRHI::vkCmdResetQueryPool(InCmdBuffer->GetHandle(), QueryPool, 0, MaxQueries);
-	++NumResets;
-}
-*/
-
 void FVulkanOcclusionQueryPool::SetFence(FVulkanCmdBuffer* InCmdBuffer)
 {
 	check(InCmdBuffer);
@@ -415,144 +218,6 @@ void FVulkanOcclusionQueryPool::Reset(FVulkanCmdBuffer* InCmdBuffer, uint32 InFr
 	State = EState::RHIT_PostBeginBatch;
 }
 
-bool FVulkanTimingQuery::InternalGetResult(uint64& OutTime, bool bWait)
-{
-	ensure(State == EState::Written);
-	if (Pool->TryGetResults(this, IndexInPool, bWait, Result))
-	{
-		State = EState::HasResult;
-		OutTime = Result;
-		return true;
-	}
-
-	return false;
-}
-
-VkResult FVulkanTimestampQueryPool::InternalGetQueryPoolResults(FVulkanTimingQuery* Query, uint32 QueryIndex, uint64& OutResults)
-{
-	uint32 PendingQueries = (QueueTail <= QueueHead) ? (QueueHead - QueueTail) : (MaxQueries - QueueTail) + QueueHead;
-	uint32 FetchingQueries = (QueueTail <= QueryIndex) ? (QueryIndex - QueueTail) + 1 : (MaxQueries - QueueTail) + QueryIndex + 1;
-
-	if (PendingQueries < FetchingQueries) //if we try to fetch more than are pending that we have the result already
-	{
-		OutResults = QueryOutput[QueryIndex];
-		ensure(QueryAllocation[QueryIndex] == Query);
-		QueryAllocation[QueryIndex] = nullptr;
-		return VK_SUCCESS;
-	}
-	else if (QueryIndex >= QueueTail)
-	{
-		uint32 NumQueries = QueryIndex - QueueTail + 1;
-		VkResult Result = VulkanRHI::vkGetQueryPoolResults(Device->GetInstanceHandle(), QueryPool, QueueTail, NumQueries, sizeof(uint64) * NumQueries, &QueryOutput[QueueTail], sizeof(uint64), VK_QUERY_RESULT_64_BIT);
-		if (Result == VK_SUCCESS)
-		{
-			if (!IsInRenderingThread() || !IsRunningRHIInSeparateThread())
-			{
-				FVulkanCmdBuffer* CmdBuffer = Device->GetImmediateContext().GetCommandBufferManager()->GetActiveCmdBuffer();
-				VulkanRHI::vkCmdResetQueryPool(CmdBuffer->GetHandle(), QueryPool, QueueTail, NumQueries);
-			}
-			else
-			{
-				check(Query->Pool == this);
-				FRHICommandListExecutor::GetImmediateCommandList().EnqueueLambda([NumQueries, QueueTail = QueueTail, QueryPool = Query->Pool]
-				(FRHICommandListImmediate& RHICommandList)
-				{
-					FVulkanCmdBuffer* CmdBuffer = static_cast<FVulkanCommandListContext&>(RHICommandList.GetContext()).GetCommandBufferManager()->GetActiveCmdBuffer();
-					VulkanRHI::vkCmdResetQueryPool(CmdBuffer->GetHandle(), QueryPool->QueryPool, QueueTail, NumQueries);
-				});
-			}
-			OutResults = QueryOutput[QueryIndex];
-			ensure(QueryAllocation[QueryIndex] == Query);
-			QueryAllocation[QueryIndex] = nullptr;
-			QueueTail = QueryIndex + 1;
-		}
-		return Result;
-	}
-	else
-	{
-		FVulkanTimingQuery* SavedQuery = QueryAllocation[MaxQueries - 1];
-		VkResult FirstHalf = InternalGetQueryPoolResults(SavedQuery, MaxQueries - 1, OutResults);
-		QueryAllocation[MaxQueries - 1] = SavedQuery;
-		if (FirstHalf == VK_SUCCESS)
-		{
-			QueueTail = 0;
-			VkResult SecondHalf = InternalGetQueryPoolResults(Query, QueryIndex, OutResults);
-			if (SecondHalf == VK_SUCCESS)
-			{
-				QueueTail = QueryIndex + 1;
-			}
-			return SecondHalf;
-		}
-		return FirstHalf;
-	}
-}
-
-bool FVulkanTimestampQueryPool::TryGetResults(FVulkanTimingQuery* Query, uint32 QueryIndex, bool bWait, uint64& OutResults)
-{
-	VkResult Result = InternalGetQueryPoolResults(Query, QueryIndex, OutResults);
-	if (Result == VK_SUCCESS)
-	{
-		return true;
-	}
-	else if (Result == VK_NOT_READY)
-	{
-		if (bWait)
-		{
-			SCOPE_CYCLE_COUNTER(STAT_VulkanWaitQuery);
-
-			// We'll do manual wait
-			double StartTime = FPlatformTime::Seconds();
-
-			ENamedThreads::Type RenderThread_Local = ENamedThreads::GetRenderThread_Local();
-			bool bSuccess = false;
-			while (!bSuccess)
-			{
-				FPlatformProcess::SleepNoStats(0);
-
-				// pump RHIThread to make sure these queries have actually been submitted to the GPU.
-				if (IsInActualRenderingThread())
-				{
-					FTaskGraphInterface::Get().ProcessThreadUntilIdle(RenderThread_Local);
-				}
-
-				Result = InternalGetQueryPoolResults(Query, QueryIndex, OutResults);
-				if (Result == VK_SUCCESS)
-				{
-					bSuccess = true;
-					break;
-				}
-				else if (Result == VK_NOT_READY)
-				{
-					bSuccess = false;
-				}
-				else
-				{
-					bSuccess = false;
-					VERIFYVULKANRESULT(Result);
-				}
-
-				// timer queries are used for Benchmarks which can stall a bit more
-				const double TimeoutValue = (QueryType == VK_QUERY_TYPE_TIMESTAMP) ? 2.0 : 0.5;
-				// look for gpu stuck/crashed
-				if ((FPlatformTime::Seconds() - StartTime) > TimeoutValue)
-				{
-					UE_LOG(LogRHI, Log, TEXT("Timed out while waiting for GPU to catch up on timer results. (%.1f s)"), TimeoutValue);
-					return false;
-				}
-			}
-
-			return true;
-		}
-	}
-	else
-	{
-		VERIFYVULKANRESULT(Result);
-	}
-
-	return false;
-}
-
-
 void FVulkanCommandListContext::BeginOcclusionQueryBatch(FVulkanCmdBuffer* CmdBuffer, uint32 NumQueriesInBatch)
 {
 	ensure(IsImmediate());
@@ -670,27 +335,6 @@ void FVulkanDevice::ReleaseUnusedOcclusionQueryPools()
 	}
 }
 
-/*
-FVulkanTimestampQueryPool* FVulkanDevice::PrepareTimestampQueryPool(bool& bOutRequiresReset)
-{
-	bOutRequiresReset = false;
-	if (!TimestampQueryPool)
-	{
-		//#todo-rco: Create this earlier
-		TimestampQueryPool = new FVulkanTimestampQueryPool(this, GMaxLifetimeTimestampQueries);
-		bOutRequiresReset = true;
-	}
-
-	// Wrap around the queries
-	if (TimestampQueryPool->GetNumAllocatedQueries() >= TimestampQueryPool->GetMaxQueries())
-	{
-		//#todo-rco: Check overlap and grow if necessary
-		TimestampQueryPool->NumUsedQueries = 0;
-	}
-
-	return TimestampQueryPool;
-}
-*/
 void FVulkanCommandListContext::EndOcclusionQueryBatch(FVulkanCmdBuffer* CmdBuffer)
 {
 	ensure(IsImmediate());
@@ -823,12 +467,7 @@ bool FVulkanDynamicRHI::RHIGetRenderQueryResult(FRenderQueryRHIParamRef QueryRHI
 		return QueryResult;
 */
 	}
-	/*
-	FScopeLock ScopeLock(&GOcclusionQueryCS);
-	check(IsInRenderingThread());
 
-	return Query->GetResult(Device, OutResult, bWait);
-	*/
 	return false;
 }
 
@@ -903,22 +542,6 @@ void FVulkanCommandListContext::RHIEndRenderQuery(FRenderQueryRHIParamRef QueryR
 	}
 }
 
-/*
-FRenderQueryRHIRef FVulkanTimestampQueryPool::AllocateQuery()
-{
-	FVulkanTimingQuery* Query = new FVulkanTimingQuery();
-	Query->State = FVulkanTimingQuery::Unused;
-	Query->IsPooledAtHighLevel = true;
-	Query->Pool = this;
-	return Query;
-}
-
-void FVulkanTimestampQueryPool::ReleaseQuery(TRefCountPtr<FRHIRenderQuery> &Query)
-{
-	Query.SafeRelease();
-}
-*/
-
 void FVulkanCommandListContext::WriteBeginTimestamp(FVulkanCmdBuffer* CmdBuffer)
 {
 	FrameTiming->StartTiming(CmdBuffer);
diff --git a/Engine/Source/Runtime/VulkanRHI/Private/VulkanUAV.cpp b/Engine/Source/Runtime/VulkanRHI/Private/VulkanUAV.cpp
--- a/Engine/Source/Runtime/VulkanRHI/Private/VulkanUAV.cpp
+++ b/Engine/Source/Runtime/VulkanRHI/Private/VulkanUAV.cpp
@@ -410,20 +410,6 @@ void FVulkanGPUFence::Clear()
 	CmdBuffer = nullptr;
 	FenceSignaledCounter = MAX_uint64;
 }
-/*
-FRenderQueryPoolRHIRef FVulkanDynamicRHI::RHICreateRenderQueryPool(ERenderQueryType QueryType, uint32 NumQueries)
-{
-	if (QueryType == RQT_AbsoluteTime)
-	{
-		return new FVulkanTimestampQueryPool(Device, NumQueries);
-	}
-	else
-	{
-		return new FDefaultRHIRenderQueryPool(QueryType, this, NumQueries);
-	}
-}
-*/
-
 
 bool FVulkanGPUFence::Poll() const
 {
diff --git a/Engine/Source/Runtime/VulkanRHI/Private/VulkanUtil.cpp b/Engine/Source/Runtime/VulkanRHI/Private/VulkanUtil.cpp
--- a/Engine/Source/Runtime/VulkanRHI/Private/VulkanUtil.cpp
+++ b/Engine/Source/Runtime/VulkanRHI/Private/VulkanUtil.cpp
@@ -14,34 +14,6 @@
 
 extern CORE_API bool GIsGPUCrashed;
 
-//#todo-rco: Horrible work-around to avoid breaking binary compatibility; move to header!
-struct FVulkanTimingQueryPool : public FVulkanQueryPool
-{
-	FVulkanTimingQueryPool(FVulkanDevice* InDevice, uint32 InBufferSize)
-		: FVulkanQueryPool(InDevice, InBufferSize * 2, VK_QUERY_TYPE_TIMESTAMP)
-		, BufferSize(InBufferSize)
-	{
-		TimestampListHandles.AddZeroed(InBufferSize * 2);
-	}
-
-	uint32 CurrentTimestamp = 0;
-	uint32 NumIssuedTimestamps = 0;
-	const uint32 BufferSize;
-
-	struct FCmdBufferFence
-	{
-		FVulkanCmdBuffer* CmdBuffer;
-		uint64 FenceCounter;
-	};
-	TArray<FCmdBufferFence> TimestampListHandles;
-
-	VulkanRHI::FStagingBuffer* ResultsBuffer = nullptr;
-};
-
-
-//#todo-rco: Horrible work-around to avoid breaking binary compatibility
-static TMap<FVulkanGPUTiming*, FVulkanTimingQueryPool*> GTimingQueryPools;
-
 
 static FString		EventDeepString(TEXT("EventTooDeep"));
 static const uint32	EventDeepCRC = FCrc::StrCrc32<TCHAR>(*EventDeepString);
@@ -157,6 +129,11 @@ void FVulkanDynamicRHI::RHIUnlockStagingBuffer(FStagingBufferRHIParamRef Staging
 	StagingBuffer->Unlock();
 }
 
+FVulkanGPUTiming::~FVulkanGPUTiming()
+{
+	check(!Pool);
+}
+
 /**
  * Initializes all Vulkan resources and if necessary, the static variables.
  */
@@ -168,9 +145,9 @@ void FVulkanGPUTiming::Initialize()
 
 	if (FVulkanPlatform::SupportsTimestampRenderQueries() && GIsSupported)
 	{
-		FVulkanTimingQueryPool* Pool = new FVulkanTimingQueryPool(Device, 8);
+		check(!Pool);
+		Pool = new FVulkanTimingQueryPool(Device, 8);
 		Pool->ResultsBuffer = Device->GetStagingManager().AcquireBuffer(8 * sizeof(uint64), VK_BUFFER_USAGE_TRANSFER_DST_BIT, true);
-		GTimingQueryPools.Add(this, Pool);
 	}
 }
 
@@ -181,12 +158,10 @@ void FVulkanGPUTiming::Release()
 {
 	if (FVulkanPlatform::SupportsTimestampRenderQueries())
 	{
-		FVulkanTimingQueryPool* FoundPool = nullptr;
-		if (GTimingQueryPools.RemoveAndCopyValue(this, FoundPool) && FoundPool)
-		{
-			Device->GetStagingManager().ReleaseBuffer(nullptr, FoundPool->ResultsBuffer);
-			delete FoundPool;
-		}
+		check(Pool);
+		Device->GetStagingManager().ReleaseBuffer(nullptr, Pool->ResultsBuffer);
+		delete Pool;
+		Pool = nullptr;
 	}
 }
 
@@ -202,7 +177,6 @@ void FVulkanGPUTiming::StartTiming(FVulkanCmdBuffer* CmdBuffer)
 		{
 			CmdBuffer = CmdContext->GetCommandBufferManager()->GetActiveCmdBuffer();
 		}
-		FVulkanTimingQueryPool* Pool = GTimingQueryPools.FindChecked(this);
 		Pool->CurrentTimestamp = (Pool->CurrentTimestamp + 1) % Pool->BufferSize;
 		const uint32 QueryStartIndex = Pool->CurrentTimestamp * 2;
 		VulkanRHI::vkCmdWriteTimestamp(CmdBuffer->GetHandle(), VK_PIPELINE_STAGE_TOP_OF_PIPE_BIT, Pool->GetHandle(), QueryStartIndex);
@@ -225,7 +199,6 @@ void FVulkanGPUTiming::EndTiming(FVulkanCmdBuffer* CmdBuffer)
 		{
 			CmdBuffer = CmdContext->GetCommandBufferManager()->GetActiveCmdBuffer();
 		}
-		FVulkanTimingQueryPool* Pool = GTimingQueryPools.FindChecked(this);
 		check(Pool->CurrentTimestamp < Pool->BufferSize);
 		const uint32 QueryStartIndex = Pool->CurrentTimestamp * 2;
 		const uint32 QueryEndIndex = Pool->CurrentTimestamp * 2 + 1;
@@ -253,7 +226,6 @@ uint64 FVulkanGPUTiming::GetTiming(bool bGetCurrentResultsAndBlock)
 {
 	if (GIsSupported)
 	{
-		FVulkanTimingQueryPool* Pool = GTimingQueryPools.FindChecked(this);
 		check(Pool->CurrentTimestamp < Pool->BufferSize);
 		uint64 StartTime, EndTime;
 		int32 TimestampIndex = Pool->CurrentTimestamp;
diff --git a/Engine/Source/Runtime/VulkanRHI/Public/VulkanResources.h b/Engine/Source/Runtime/VulkanRHI/Public/VulkanResources.h
--- a/Engine/Source/Runtime/VulkanRHI/Public/VulkanResources.h
+++ b/Engine/Source/Runtime/VulkanRHI/Public/VulkanResources.h
@@ -879,38 +879,28 @@ protected:
 	uint32 FrameNumber = UINT32_MAX;
 };
 
-
-class FVulkanTimestampQueryPool final : public FVulkanQueryPool
+class FVulkanTimingQueryPool : public FVulkanQueryPool
 {
-/*
-	FRenderQueryRHIRef AllocateQuery() override;
-	virtual void ReleaseQuery(TRefCountPtr<FRHIRenderQuery> &Query) override;
-*/
-
-	uint32 QueueHead = 0;
-	uint32 QueueTail = 0;
-	TArray<class FVulkanTimingQuery*> QueryAllocation;
-
 public:
-	FVulkanTimestampQueryPool(FVulkanDevice* InDevice, uint32 InMaxQueries)
-		: FVulkanQueryPool(InDevice, InMaxQueries, VK_QUERY_TYPE_TIMESTAMP)
+	FVulkanTimingQueryPool(FVulkanDevice* InDevice, uint32 InBufferSize)
+		: FVulkanQueryPool(InDevice, InBufferSize * 2, VK_QUERY_TYPE_TIMESTAMP)
+		, BufferSize(InBufferSize)
 	{
-		QueryAllocation.AddDefaulted(InMaxQueries);
+		TimestampListHandles.AddZeroed(InBufferSize * 2);
 	}
 
-	uint32 AcquireIndex(class FVulkanTimingQuery* Query)
-	{
-		uint32 QueueIndex = QueueHead++;
-		QueueHead = (QueueHead < MaxQueries) ? QueueHead : 0;
-		ensure(QueueHead != QueueTail);
-		QueryAllocation[QueueIndex] = Query;
-		return QueueIndex;
-	}
+	uint32 CurrentTimestamp = 0;
+	uint32 NumIssuedTimestamps = 0;
+	const uint32 BufferSize;
 
-	bool TryGetResults(class FVulkanTimingQuery* Query, uint32 QueryIndex, bool bWait, uint64& OutResult);
+	struct FCmdBufferFence
+	{
+		FVulkanCmdBuffer* CmdBuffer;
+		uint64 FenceCounter;
+	};
+	TArray<FCmdBufferFence> TimestampListHandles;
 
-private:
-	VkResult InternalGetQueryPoolResults(class FVulkanTimingQuery* Query, uint32 QueryIndex, uint64& OutResults);
+	VulkanRHI::FStagingBuffer* ResultsBuffer = nullptr;
 };
 
 class FVulkanRenderQuery : public FRHIRenderQuery
@@ -951,74 +941,6 @@ public:
 	EState State = EState::Undefined;
 };
 
-class FVulkanTimingQuery : public FVulkanRenderQuery
-{
-public:
-	FVulkanTimingQuery()
-		: FVulkanRenderQuery(RQT_AbsoluteTime)
-	{
-	}
-
-	virtual ~FVulkanTimingQuery()
-	{
-	}
-
-	enum EState : uint8
-	{
-		Unused,
-		Written,
-		HasResult,
-	};
-
-	bool GetResult(uint64& OutTime, bool bWait)
-	{
-		if (State == EState::HasResult)
-		{
-			OutTime = Result;
-			return true;
-		}
-		else if (State != EState::Unused)
-		{
-			return InternalGetResult(OutTime, bWait);
-		}
-		else
-		{
-			return false;
-		}
-	}
-
-#if VULKAN_QUERY_CALLSTACK
-	uint64 StackFrames[16];
-#endif
-
-protected:
-	/*TRefCountPtr<*/FVulkanTimestampQueryPool* /*>*/ Pool;
-	EState State = EState::Unused;
-	bool IsPooledAtHighLevel = false;
-
-	bool InternalGetResult(uint64& OutTime, bool bWait);
-
-	friend class FVulkanTimestampQueryPool;
-	friend class FVulkanCommandListContext;
-};
-/*
-class FVulkanRenderQuery : public FRHIRenderQuery
-{
-public:
-	FVulkanRenderQuery(FVulkanDevice* Device, ERenderQueryType InQueryType);
-	virtual ~FVulkanRenderQuery();
-
-	inline bool HasQueryBeenEmitted() const
-	{
-		return State == EState::InEnd;
-	}
-
-	uint32 LastPoolReset = 0;
-
-private:
-};
-*/
-
 struct FVulkanBufferView : public FRHIResource, public VulkanRHI::FDeviceChild
 {
 	FVulkanBufferView(FVulkanDevice* InDevice)
diff --git a/Engine/Source/Runtime/VulkanRHI/Public/VulkanUtil.h b/Engine/Source/Runtime/VulkanRHI/Public/VulkanUtil.h
--- a/Engine/Source/Runtime/VulkanRHI/Public/VulkanUtil.h
+++ b/Engine/Source/Runtime/VulkanRHI/Public/VulkanUtil.h
@@ -11,8 +11,8 @@
 #include "Serialization/MemoryWriter.h"
 
 class FVulkanCmdBuffer;
-class FVulkanTimingQuery;
 class FVulkanCommandListContext;
+class FVulkanTimingQueryPool;
 
 class FVulkanGPUTiming : public FGPUTiming
 {
@@ -25,6 +25,8 @@ public:
 	{
 	}
 
+	~FVulkanGPUTiming();
+
 	/**
 	 * Start a GPU timing measurement.
 	 */
@@ -76,6 +78,7 @@ private:
 	bool bEndTimestampIssued;
 
 	FVulkanCommandListContext* CmdContext;
+	FVulkanTimingQueryPool* Pool = nullptr;
 };
 
 /** A single perf event node, which tracks information about a appBeginDrawEvent/appEndDrawEvent range. */
