UE4.22 - Stat GPU works again in Vulkan
#rnx

diff --git a/Engine/Source/Runtime/VulkanRHI/Private/VulkanDevice.cpp b/Engine/Source/Runtime/VulkanRHI/Private/VulkanDevice.cpp
--- a/Engine/Source/Runtime/VulkanRHI/Private/VulkanDevice.cpp
+++ b/Engine/Source/Runtime/VulkanRHI/Private/VulkanDevice.cpp
@@ -968,15 +968,11 @@ void FVulkanDevice::Destroy()
 		delete Pool;
 	}
 	FreeOcclusionQueryPools.SetNum(0, false);
-/*
-	delete TimestampQueryPool;
-*/
 
 	delete PipelineStateCache;
 	PipelineStateCache = nullptr;
 	StagingManager.Deinit();
 
-
 	if (GGPUCrashDebuggingEnabled)
 	{
 #if VULKAN_SUPPORTS_AMD_BUFFER_MARKER
diff --git a/Engine/Source/Runtime/VulkanRHI/Private/VulkanQuery.cpp b/Engine/Source/Runtime/VulkanRHI/Private/VulkanQuery.cpp
--- a/Engine/Source/Runtime/VulkanRHI/Private/VulkanQuery.cpp
+++ b/Engine/Source/Runtime/VulkanRHI/Private/VulkanQuery.cpp
@@ -414,10 +414,8 @@ FRenderQueryRHIRef FVulkanDynamicRHI::RHICreateRenderQuery(ERenderQueryType Quer
 	}
 	else if (QueryType == RQT_AbsoluteTime)
 	{
-/*
-		FVulkanTimingQuery* Query = new FVulkanTimingQuery();
+		FVulkanTimingQuery* Query = new FVulkanTimingQuery(Device);
 		return Query;
-*/
 	}
 	else
 	{
@@ -430,6 +428,12 @@ FRenderQueryRHIRef FVulkanDynamicRHI::RHICreateRenderQuery(ERenderQueryType Quer
 
 bool FVulkanDynamicRHI::RHIGetRenderQueryResult(FRenderQueryRHIParamRef QueryRHI, uint64& OutNumPixels, bool bWait)
 {
+	auto ToMicroseconds = [](uint64 Timestamp)
+	{
+		const double Frequency = double(FVulkanGPUTiming::GetTimingFrequency());
+		uint64 Microseconds = (uint64)((double(Timestamp) / Frequency) * 1000.0 * 1000.0);
+		return Microseconds;
+	};
 	check(IsInRenderingThread());
 	FVulkanRenderQuery* BaseQuery = ResourceCast(QueryRHI);
 	if (BaseQuery->QueryType == RQT_Occlusion)
@@ -459,13 +463,63 @@ bool FVulkanDynamicRHI::RHIGetRenderQueryResult(FRenderQueryRHIParamRef QueryRHI
 	}
 	else if (BaseQuery->QueryType == RQT_AbsoluteTime)
 	{
-/*
 		FVulkanTimingQuery* Query = static_cast<FVulkanTimingQuery*>(BaseQuery);
-		uint64 TimingResult;
-		bool QueryResult = Query->GetResult(TimingResult, false);
-		OutNumPixels = TimingResult / 1000; //RHIGetRenderQueryResult assumes microseconds where vk provided ns
-		return QueryResult;
-*/
+		check(Query->Pool.CurrentTimestamp < Query->Pool.BufferSize);
+		int32 TimestampIndex = Query->Pool.CurrentTimestamp;
+		if (!bWait)
+		{
+			// Quickly check the most recent measurements to see if any of them has been resolved.  Do not flush these queries.
+			for (uint32 IssueIndex = 1; IssueIndex < Query->Pool.NumIssuedTimestamps; ++IssueIndex)
+			{
+				const FVulkanTimingQueryPool::FCmdBufferFence& StartQuerySyncPoint = Query->Pool.TimestampListHandles[TimestampIndex];
+				if (StartQuerySyncPoint.FenceCounter < StartQuerySyncPoint.CmdBuffer->GetFenceSignaledCounter())
+				{
+					Query->Pool.ResultsBuffer->InvalidateMappedMemory();
+					uint64* Data = (uint64*)Query->Pool.ResultsBuffer->GetMappedPointer();
+					OutNumPixels = ToMicroseconds(Data[TimestampIndex]);
+					return true;
+				}
+
+				TimestampIndex = (TimestampIndex + Query->Pool.BufferSize - 1) % Query->Pool.BufferSize;
+			}
+		}
+
+		if (Query->Pool.NumIssuedTimestamps > 0 || bWait)
+		{
+			// None of the (NumIssuedTimestamps - 1) measurements were ready yet,
+			// so check the oldest measurement more thoroughly.
+			// This really only happens if occlusion and frame sync event queries are disabled, otherwise those will block until the GPU catches up to 1 frame behind
+
+			const bool bBlocking = (Query->Pool.NumIssuedTimestamps == Query->Pool.BufferSize) || bWait;
+			const uint32 IdleStart = FPlatformTime::Cycles();
+
+			SCOPE_CYCLE_COUNTER(STAT_RenderQueryResultTime);
+
+			if (bBlocking)
+			{
+				const FVulkanTimingQueryPool::FCmdBufferFence& StartQuerySyncPoint = Query->Pool.TimestampListHandles[TimestampIndex];
+				bool bWaitForStart = StartQuerySyncPoint.FenceCounter == StartQuerySyncPoint.CmdBuffer->GetFenceSignaledCounter();
+				if (bWaitForStart)
+				{
+					// Need to submit the open command lists.
+					Device->SubmitCommandsAndFlushGPU();
+				}
+
+				// CPU wait for query results to be ready.
+				if (bWaitForStart && StartQuerySyncPoint.FenceCounter == StartQuerySyncPoint.CmdBuffer->GetFenceSignaledCounter())
+				{
+					Device->GetImmediateContext().GetCommandBufferManager()->WaitForCmdBuffer(StartQuerySyncPoint.CmdBuffer);
+				}
+			}
+
+			Query->Pool.ResultsBuffer->InvalidateMappedMemory();
+			GRenderThreadIdle[ERenderThreadIdleTypes::WaitingForGPUQuery] += FPlatformTime::Cycles() - IdleStart;
+			GRenderThreadNumIdle[ERenderThreadIdleTypes::WaitingForGPUQuery]++;
+
+			uint64* Data = (uint64*)Query->Pool.ResultsBuffer->GetMappedPointer();
+			OutNumPixels = ToMicroseconds(Data[TimestampIndex]);
+			return true;
+		}
 	}
 
 	return false;
@@ -527,18 +581,17 @@ void FVulkanCommandListContext::RHIEndRenderQuery(FRenderQueryRHIParamRef QueryR
 	}
 	else if (BaseQuery->QueryType == RQT_AbsoluteTime)
 	{
-/*
 		FVulkanTimingQuery* Query = static_cast<FVulkanTimingQuery*>(BaseQuery);
+		Query->Pool.CurrentTimestamp = (Query->Pool.CurrentTimestamp + 1) % Query->Pool.BufferSize;
+		const uint32 QueryEndIndex = Query->Pool.CurrentTimestamp;
 		FVulkanCmdBuffer* CmdBuffer = CommandBufferManager->GetActiveCmdBuffer();
-		Query->Pool = Query->IsPooledAtHighLevel ? Query->Pool : CmdBuffer->PrepareTimestampQueryPool();
-		Query->IndexInPool = Query->Pool->AcquireIndex(Query);
-		Query->State = FVulkanTimingQuery::Written;
-
-		VulkanRHI::vkCmdWriteTimestamp(CmdBuffer->GetHandle(), VK_PIPELINE_STAGE_BOTTOM_OF_PIPE_BIT, Query->Pool->GetHandle(), Query->IndexInPool);
-#if VULKAN_QUERY_CALLSTACK
-		FPlatformStackWalk::CaptureStackBackTrace(Query->StackFrames, 16);
-#endif
-*/
+		VulkanRHI::vkCmdWriteTimestamp(CmdBuffer->GetHandle(), VK_PIPELINE_STAGE_TOP_OF_PIPE_BIT, Query->Pool.GetHandle(), QueryEndIndex);
+		//check(CmdBuffer->IsOutsideRenderPass());
+		VulkanRHI::vkCmdCopyQueryPoolResults(CmdBuffer->GetHandle(), Query->Pool.GetHandle(), QueryEndIndex, 1, Query->Pool.ResultsBuffer->GetHandle(), sizeof(uint64) * QueryEndIndex, sizeof(uint64), VK_QUERY_RESULT_64_BIT);
+		VulkanRHI::vkCmdResetQueryPool(CmdBuffer->GetHandle(), Query->Pool.GetHandle(), QueryEndIndex, 1);
+		Query->Pool.TimestampListHandles[QueryEndIndex].CmdBuffer = CmdBuffer;
+		Query->Pool.TimestampListHandles[QueryEndIndex].FenceCounter = CmdBuffer->GetFenceSignaledCounter();
+		Query->Pool.NumIssuedTimestamps = FMath::Min<uint32>(Query->Pool.NumIssuedTimestamps + 1, Query->Pool.BufferSize);
 	}
 }
 
@@ -551,3 +604,16 @@ void FVulkanCommandListContext::WriteEndTimestamp(FVulkanCmdBuffer* CmdBuffer)
 {
 	FrameTiming->EndTiming(CmdBuffer);
 }
+
+FVulkanTimingQuery::FVulkanTimingQuery(FVulkanDevice* InDevice)
+	: FVulkanRenderQuery(RQT_AbsoluteTime)
+	, Pool(InDevice, 4)
+{
+	Pool.ResultsBuffer = InDevice->GetStagingManager().AcquireBuffer(Pool.BufferSize * sizeof(uint64), VK_BUFFER_USAGE_TRANSFER_DST_BIT, true);
+}
+
+
+FVulkanTimingQuery::~FVulkanTimingQuery()
+{
+	Pool.GetParent()->GetStagingManager().ReleaseBuffer(nullptr, Pool.ResultsBuffer);
+}
diff --git a/Engine/Source/Runtime/VulkanRHI/Public/VulkanResources.h b/Engine/Source/Runtime/VulkanRHI/Public/VulkanResources.h
--- a/Engine/Source/Runtime/VulkanRHI/Public/VulkanResources.h
+++ b/Engine/Source/Runtime/VulkanRHI/Public/VulkanResources.h
@@ -941,6 +941,15 @@ public:
 	EState State = EState::Undefined;
 };
 
+class FVulkanTimingQuery : public FVulkanRenderQuery
+{
+public:
+	FVulkanTimingQuery(FVulkanDevice* InDevice);
+	virtual ~FVulkanTimingQuery();
+
+	FVulkanTimingQueryPool Pool;
+};
+
 struct FVulkanBufferView : public FRHIResource, public VulkanRHI::FDeviceChild
 {
 	FVulkanBufferView(FVulkanDevice* InDevice)
diff --git a/Engine/Source/Runtime/VulkanRHI/Public/VulkanUtil.h b/Engine/Source/Runtime/VulkanRHI/Public/VulkanUtil.h
--- a/Engine/Source/Runtime/VulkanRHI/Public/VulkanUtil.h
+++ b/Engine/Source/Runtime/VulkanRHI/Public/VulkanUtil.h
@@ -59,7 +59,6 @@ public:
 	bool IsComplete() const
 	{
 		check(bEndTimestampIssued);
-		//return *((uint64*)EndTimestamp.GetPointer()) != 0;
 		return true;
 	}
 
