UE4.22 - Disable rhithread and enable rhibypass on NV 20xx to workaround crash bug on Vulkan

diff --git a/Engine/Source/Runtime/VulkanRHI/Private/VulkanGenericPlatform.h b/Engine/Source/Runtime/VulkanRHI/Private/VulkanGenericPlatform.h
--- a/Engine/Source/Runtime/VulkanRHI/Private/VulkanGenericPlatform.h
+++ b/Engine/Source/Runtime/VulkanRHI/Private/VulkanGenericPlatform.h
@@ -13,7 +13,7 @@ class FVulkanGenericPlatform
 {
 public:
 	static bool IsSupported() { return true; }
-	static void CheckDeviceDriver(uint32 DeviceIndex) {}
+	static void CheckDeviceDriver(uint32 DeviceIndex, const VkPhysicalDeviceProperties& Props) {}
 
 	static bool LoadVulkanLibrary() { return true; }
 	static bool LoadVulkanInstanceFunctions(VkInstance inInstance) { return true; }
diff --git a/Engine/Source/Runtime/VulkanRHI/Private/VulkanRHI.cpp b/Engine/Source/Runtime/VulkanRHI/Private/VulkanRHI.cpp
--- a/Engine/Source/Runtime/VulkanRHI/Private/VulkanRHI.cpp
+++ b/Engine/Source/Runtime/VulkanRHI/Private/VulkanRHI.cpp
@@ -570,7 +570,7 @@ void FVulkanDynamicRHI::SelectAndInitDevice()
 	GRHIVendorId = Props.vendorID;
 	GRHIAdapterName = ANSI_TO_TCHAR(Props.deviceName);
 
-	FVulkanPlatform::CheckDeviceDriver(DeviceIndex);
+	FVulkanPlatform::CheckDeviceDriver(DeviceIndex, Props);
 
 	Device->InitGPU(DeviceIndex);
 
diff --git a/Engine/Source/Runtime/VulkanRHI/Private/Windows/VulkanWindowsPlatform.cpp b/Engine/Source/Runtime/VulkanRHI/Private/Windows/VulkanWindowsPlatform.cpp
--- a/Engine/Source/Runtime/VulkanRHI/Private/Windows/VulkanWindowsPlatform.cpp
+++ b/Engine/Source/Runtime/VulkanRHI/Private/Windows/VulkanWindowsPlatform.cpp
@@ -185,7 +185,7 @@ void FVulkanWindowsPlatform::WriteCrashMarker(const FOptionalVulkanDeviceExtensi
 	}
 }
 
-void FVulkanWindowsPlatform::CheckDeviceDriver(uint32 DeviceIndex)
+void FVulkanWindowsPlatform::CheckDeviceDriver(uint32 DeviceIndex, const VkPhysicalDeviceProperties& Props)
 {
 	const bool bAllowVendorDevice = !FParse::Param(FCommandLine::Get(), TEXT("novendordevice"));
 	if (IsRHIDeviceAMD() && bAllowVendorDevice)
@@ -279,4 +279,16 @@ void FVulkanWindowsPlatform::CheckDeviceDriver(uint32 DeviceIndex)
 			agsDeInit(AmdAgsContext);
 		}
 	}
+	else if (IsRHIDeviceNVIDIA())
+	{
+		// Workaround a crash on 20xx family
+		UE_LOG(LogVulkanRHI, Warning, TEXT("Nvidia 20xx family of GPUs have a known crash. Compatibility mode (slow!) will now be enabled"));
+		if (GRHIAdapterName.Contains(TEXT("RTX 20")))
+		{
+			extern TAutoConsoleVariable<int32> GRHIThreadCvar;
+			GRHIThreadCvar->SetWithCurrentPriority(0);
+			IConsoleVariable* BypassVar = IConsoleManager::Get().FindConsoleVariable(TEXT("r.RHICmdBypass"));
+			BypassVar->SetWithCurrentPriority(1);
+		}
+	}
 }
diff --git a/Engine/Source/Runtime/VulkanRHI/Private/Windows/VulkanWindowsPlatform.h b/Engine/Source/Runtime/VulkanRHI/Private/Windows/VulkanWindowsPlatform.h
--- a/Engine/Source/Runtime/VulkanRHI/Private/Windows/VulkanWindowsPlatform.h
+++ b/Engine/Source/Runtime/VulkanRHI/Private/Windows/VulkanWindowsPlatform.h
@@ -49,7 +49,7 @@
 class FVulkanWindowsPlatform : public FVulkanGenericPlatform
 {
 public:
-	static void CheckDeviceDriver(uint32 DeviceIndex);
+	static void CheckDeviceDriver(uint32 DeviceIndex, const VkPhysicalDeviceProperties& Props);
 
 	static bool LoadVulkanLibrary();
 	static bool LoadVulkanInstanceFunctions(VkInstance inInstance);
