UE4.22 - Pre-req for re-enabling real uniform buffers on Vulkan
#rnx

diff --git a/Engine/Source/Runtime/VulkanRHI/Private/VulkanCommands.cpp b/Engine/Source/Runtime/VulkanRHI/Private/VulkanCommands.cpp
--- a/Engine/Source/Runtime/VulkanRHI/Private/VulkanCommands.cpp
+++ b/Engine/Source/Runtime/VulkanRHI/Private/VulkanCommands.cpp
@@ -631,7 +631,8 @@ inline void FVulkanCommandListContext::SetShaderUniformBuffer(ShaderStage::EStag
 		}
 		else
 		{
-			PendingGfxState->SetUniformBufferConstantData(Stage, BufferIndex, UniformBuffer->ConstantData);
+			const FVulkanEmulatedUniformBuffer* EmulatedUniformBuffer = static_cast<const FVulkanEmulatedUniformBuffer*>(UniformBuffer);
+			PendingGfxState->SetUniformBufferConstantData(Stage, BufferIndex, EmulatedUniformBuffer->ConstantData);
 		}
 	}
 
@@ -734,7 +735,8 @@ void FVulkanCommandListContext::RHISetShaderUniformBuffer(FComputeShaderRHIParam
 		}
 		else
 		{
-			State.SetUniformBufferConstantData(BufferIndex, UniformBuffer->ConstantData);
+			const FVulkanEmulatedUniformBuffer* EmulatedUniformBuffer = static_cast<const FVulkanEmulatedUniformBuffer*>(UniformBuffer);
+			State.SetUniformBufferConstantData(BufferIndex, EmulatedUniformBuffer->ConstantData);
 		}
 	}
 
diff --git a/Engine/Source/Runtime/VulkanRHI/Private/VulkanUniformBuffer.cpp b/Engine/Source/Runtime/VulkanRHI/Private/VulkanUniformBuffer.cpp
--- a/Engine/Source/Runtime/VulkanRHI/Private/VulkanUniformBuffer.cpp
+++ b/Engine/Source/Runtime/VulkanRHI/Private/VulkanUniformBuffer.cpp
@@ -17,30 +17,6 @@ enum
 #endif
 };
 
-#if 0
-const int NUM_SAFE_FRAMES = 5;
-static TArray<TRefCountPtr<FVulkanBuffer>> GUBPool[NUM_SAFE_FRAMES];
-
-static TRefCountPtr<FVulkanBuffer> AllocateBufferFromPool(FVulkanDevice& Device, uint32 ConstantBufferSize, EUniformBufferUsage Usage)
-{
-	if (Usage == UniformBuffer_MultiFrame)
-	{
-		return new FVulkanBuffer(Device, ConstantBufferSize, VK_BUFFER_USAGE_UNIFORM_BUFFER_BIT, VK_MEMORY_PROPERTY_HOST_VISIBLE_BIT | VK_MEMORY_PROPERTY_HOST_COHERENT_BIT, false, __FILE__, __LINE__);
-	}
-
-	int32 BufferIndex = GFrameNumberRenderThread % NUM_SAFE_FRAMES;
-
-	GUBPool[BufferIndex].Add(new FVulkanBuffer(Device, ConstantBufferSize, VK_BUFFER_USAGE_UNIFORM_BUFFER_BIT, VK_MEMORY_PROPERTY_HOST_VISIBLE_BIT | VK_MEMORY_PROPERTY_HOST_COHERENT_BIT, false, __FILE__, __LINE__));
-	return GUBPool[BufferIndex].Last();
-}
-
-void CleanupUniformBufferPool()
-{
-	int32 BufferIndex = (GFrameNumberRenderThread  + 1) % NUM_SAFE_FRAMES;
-	GUBPool[BufferIndex].Reset(0);
-}
-#endif
-
 /*-----------------------------------------------------------------------------
 	Uniform buffer RHI object
 -----------------------------------------------------------------------------*/
@@ -63,7 +39,7 @@ static inline EBufferUsageFlags UniformBufferToBufferUsage(EUniformBufferUsage U
 	}
 }
 
-FVulkanUniformBuffer::FVulkanUniformBuffer(const FRHIUniformBufferLayout& InLayout, const void* Contents, EUniformBufferUsage InUsage, EUniformBufferValidation Validation, bool bCopyIntoConstantData)
+FVulkanUniformBuffer::FVulkanUniformBuffer(const FRHIUniformBufferLayout& InLayout, const void* Contents, EUniformBufferUsage InUsage, EUniformBufferValidation Validation)
 	: FRHIUniformBuffer(InLayout)
 {
 #if VULKAN_ENABLE_AGGRESSIVE_STATS
@@ -75,18 +51,7 @@ FVulkanUniformBuffer::FVulkanUniformBuffer(const FRHIUniformBufferLayout& InLayo
 	//	- Meaning, there is always a uniform buffer with a size specified larged than 0 bytes
 	check(InLayout.Resources.Num() > 0 || InLayout.ConstantBufferSize > 0);
 
-	// Contents might be null but size > 0 as the data doesn't need a CPU copy (it lives inside the 'real' VkBuffer)
-	if (bCopyIntoConstantData && InLayout.ConstantBufferSize)
-	{
-		// Create uniform buffer, which is stored on the CPU, the buffer is uploaded to a correct GPU buffer in UpdateDescriptorSets()
-		ConstantData.AddUninitialized(InLayout.ConstantBufferSize);
-		if (Contents)
-		{
-			FMemory::Memcpy(ConstantData.GetData(), Contents, InLayout.ConstantBufferSize);
-		}
-	}
-
-	// Parse Sampler and Texture resources, if necessary
+	// Setup resource table
 	const uint32 NumResources = InLayout.Resources.Num();
 	if (NumResources > 0)
 	{
@@ -134,19 +99,41 @@ void FVulkanUniformBuffer::UpdateResourceTable(FRHIResource** Resources, int32 R
 	}
 }
 
-void FVulkanUniformBuffer::Update(const void* Contents, int32 ContentsSize)
+
+FVulkanEmulatedUniformBuffer::FVulkanEmulatedUniformBuffer(const FRHIUniformBufferLayout& InLayout, const void* Contents, EUniformBufferUsage InUsage, EUniformBufferValidation Validation)
+	: FVulkanUniformBuffer(InLayout, Contents, InUsage, Validation)
 {
-	check(ConstantData.Num() * sizeof(ConstantData[0]) == ContentsSize);
+#if VULKAN_ENABLE_AGGRESSIVE_STATS
+	SCOPE_CYCLE_COUNTER(STAT_VulkanUniformBufferCreateTime);
+#endif
 
+	// Contents might be null but size > 0 as the data doesn't need a CPU copy
+	if (InLayout.ConstantBufferSize)
+	{
+		// Create uniform buffer, which is stored on the CPU, the buffer is uploaded to a correct GPU buffer in UpdateDescriptorSets()
+		ConstantData.AddUninitialized(InLayout.ConstantBufferSize);
+		if (Contents)
+		{
+			FMemory::Memcpy(ConstantData.GetData(), Contents, InLayout.ConstantBufferSize);
+		}
+	}
+
+	// Ancestor's constructor will set up the Resource table, so nothing else to do here
+}
+
+void FVulkanEmulatedUniformBuffer::UpdateConstantData(const void* Contents, int32 ContentsSize)
+{
+	checkSlow(ConstantData.Num() * sizeof(ConstantData[0]) == ContentsSize);
 	if (ContentsSize > 0)
 	{
 		FMemory::Memcpy(ConstantData.GetData(), Contents, ContentsSize);
 	}
 }
 
-FVulkanRealUniformBuffer::FVulkanRealUniformBuffer(FVulkanDevice& Device, const FRHIUniformBufferLayout& InLayout, const void* Contents, EUniformBufferUsage Usage, EUniformBufferValidation Validation)
-	: FVulkanUniformBuffer(InLayout, Contents, Usage, Validation, false)
-	, FVulkanResourceMultiBuffer(&Device, VK_BUFFER_USAGE_UNIFORM_BUFFER_BIT, InLayout.ConstantBufferSize, UniformBufferToBufferUsage(Usage), GEmptyCreateInfo)
+
+FVulkanRealUniformBuffer::FVulkanRealUniformBuffer(FVulkanDevice& Device, const FRHIUniformBufferLayout& InLayout, const void* Contents, EUniformBufferUsage InUsage, EUniformBufferValidation Validation)
+	: FVulkanUniformBuffer(InLayout, Contents, InUsage, Validation)
+	, FVulkanResourceMultiBuffer(&Device, VK_BUFFER_USAGE_UNIFORM_BUFFER_BIT, InLayout.ConstantBufferSize, UniformBufferToBufferUsage(InUsage), GEmptyCreateInfo)
 {
 #if VULKAN_ENABLE_AGGRESSIVE_STATS
 	SCOPE_CYCLE_COUNTER(STAT_VulkanUniformBufferCreateTime);
@@ -160,6 +147,8 @@ FVulkanRealUniformBuffer::FVulkanRealUniformBuffer(FVulkanDevice& Device, const
 		FMemory::Memcpy(Data, Contents, InLayout.ConstantBufferSize);
 		Unlock(bRT);
 	}
+
+	// Ancestor's constructor will set up the Resource table, so nothing else to do here
 }
 
 void FVulkanRealUniformBuffer::Update(const void* Contents, int32 ContentsSize)
@@ -187,24 +176,33 @@ FUniformBufferRHIRef FVulkanDynamicRHI::RHICreateUniformBuffer(const void* Conte
 	else
 	{
 		// Parts of the buffer are later on copied for each shader stage into the packed uniform buffer
-		return new FVulkanUniformBuffer(Layout, Contents, Usage, Validation, true);
+		return new FVulkanEmulatedUniformBuffer(Layout, Contents, Usage, Validation);
 	}
 }
 
-void FVulkanDynamicRHI::RHIUpdateUniformBuffer(FUniformBufferRHIParamRef UniformBufferRHI, const void* Contents)
+template <bool bRealUBs>
+inline void FVulkanDynamicRHI::UpdateUniformBuffer(FVulkanUniformBuffer* UniformBuffer, const void* Contents)
 {
-	FVulkanUniformBuffer* UniformBuffer = ResourceCast(UniformBufferRHI);
-
-	const FRHIUniformBufferLayout& Layout = UniformBufferRHI->GetLayout();
+	const FRHIUniformBufferLayout& Layout = UniformBuffer->GetLayout();
 
 	const int32 ConstantBufferSize = Layout.ConstantBufferSize;
 	const int32 NumResources = Layout.Resources.Num();
 
 	FRHICommandListImmediate& RHICmdList = FRHICommandListExecutor::GetImmediateCommandList();
 
+	FVulkanRealUniformBuffer* RealUniformBuffer = bRealUBs ? (FVulkanRealUniformBuffer*)UniformBuffer : nullptr;
+	FVulkanEmulatedUniformBuffer* EmulatedUniformBuffer = bRealUBs ? nullptr : (FVulkanEmulatedUniformBuffer*)UniformBuffer;
+
 	if (RHICmdList.Bypass())
 	{
-		UniformBuffer->Update(Contents, ConstantBufferSize);
+		if (bRealUBs)
+		{
+			RealUniformBuffer->Update(Contents, ConstantBufferSize);
+		}
+		else
+		{
+			EmulatedUniformBuffer->UpdateConstantData(Contents, ConstantBufferSize);
+		}
 		UniformBuffer->UpdateResourceTable(Layout, Contents, NumResources);
 	}
 	else
@@ -236,15 +234,39 @@ void FVulkanDynamicRHI::RHIUpdateUniformBuffer(FUniformBufferRHIParamRef Uniform
 			FMemory::Memcpy(CmdListConstantBufferData, Contents, ConstantBufferSize);
 		}
 
-		RHICmdList.EnqueueLambda([UniformBuffer, CmdListResources, NumResources, CmdListConstantBufferData, ConstantBufferSize](FRHICommandList&)
+		RHICmdList.EnqueueLambda([RealUniformBuffer, EmulatedUniformBuffer, CmdListResources, NumResources, CmdListConstantBufferData, ConstantBufferSize](FRHICommandList&)
 		{
-			UniformBuffer->Update(CmdListConstantBufferData, ConstantBufferSize);
-			UniformBuffer->UpdateResourceTable(CmdListResources, NumResources);
+			if (bRealUBs)
+			{
+				RealUniformBuffer->Update(CmdListConstantBufferData, ConstantBufferSize);
+				RealUniformBuffer->UpdateResourceTable(CmdListResources, NumResources);
+			}
+			else
+			{
+				EmulatedUniformBuffer->UpdateConstantData(CmdListConstantBufferData, ConstantBufferSize);
+				EmulatedUniformBuffer->UpdateResourceTable(CmdListResources, NumResources);
+			}
 		});
 		RHICmdList.RHIThreadFence(true);
 	}
 }
 
+
+void FVulkanDynamicRHI::RHIUpdateUniformBuffer(FUniformBufferRHIParamRef UniformBufferRHI, const void* Contents)
+{
+	static TConsoleVariableData<int32>* CVar = IConsoleManager::Get().FindTConsoleVariableDataInt(TEXT("r.Vulkan.UseRealUBs"));
+	const bool bHasRealUBs = FVulkanPlatform::UseRealUBsOptimization(CVar && CVar->GetValueOnAnyThread() > 0);
+	FVulkanUniformBuffer* UniformBuffer = ResourceCast(UniformBufferRHI);
+	if (bHasRealUBs)
+	{
+		UpdateUniformBuffer<true>(UniformBuffer, Contents);
+	}
+	else
+	{
+		UpdateUniformBuffer<false>(UniformBuffer, Contents);
+	}
+}
+
 FVulkanUniformBufferUploader::FVulkanUniformBufferUploader(FVulkanDevice* InDevice)
 	: VulkanRHI::FDeviceChild(InDevice)
 	, CPUBuffer(nullptr)
diff --git a/Engine/Source/Runtime/VulkanRHI/Public/VulkanDynamicRHI.h b/Engine/Source/Runtime/VulkanRHI/Public/VulkanDynamicRHI.h
--- a/Engine/Source/Runtime/VulkanRHI/Public/VulkanDynamicRHI.h
+++ b/Engine/Source/Runtime/VulkanRHI/Public/VulkanDynamicRHI.h
@@ -480,6 +480,9 @@ protected:
 	void InternalUpdateTexture2D(bool bFromRenderingThread, FTexture2DRHIParamRef TextureRHI, uint32 MipIndex, const struct FUpdateTextureRegion2D& UpdateRegion, uint32 SourcePitch, const uint8* SourceData);
 	void InternalUpdateTexture3D(bool bFromRenderingThread, FTexture3DRHIParamRef TextureRHI, uint32 MipIndex, const struct FUpdateTextureRegion3D& UpdateRegion, uint32 SourceRowPitch, uint32 SourceDepthPitch, const uint8* SourceData);
 
+	template <bool bRealUBs>
+	void UpdateUniformBuffer(FVulkanUniformBuffer* UniformBuffer, const void* Contents);
+
 public:
 	static TSharedPtr< IHeadMountedDisplayVulkanExtensions, ESPMode::ThreadSafe > HMDVulkanExtensions;
 };
diff --git a/Engine/Source/Runtime/VulkanRHI/Public/VulkanResources.h b/Engine/Source/Runtime/VulkanRHI/Public/VulkanResources.h
--- a/Engine/Source/Runtime/VulkanRHI/Public/VulkanResources.h
+++ b/Engine/Source/Runtime/VulkanRHI/Public/VulkanResources.h
@@ -1204,28 +1204,33 @@ public:
 class FVulkanUniformBuffer : public FRHIUniformBuffer
 {
 public:
-	FVulkanUniformBuffer(const FRHIUniformBufferLayout& InLayout, const void* Contents, EUniformBufferUsage InUsage, EUniformBufferValidation Validation, bool bCopyIntoConstantData);
-
-	TArray<uint8> ConstantData;
+	FVulkanUniformBuffer(const FRHIUniformBufferLayout& InLayout, const void* Contents, EUniformBufferUsage InUsage, EUniformBufferValidation Validation);
 
 	const TArray<TRefCountPtr<FRHIResource>>& GetResourceTable() const { return ResourceTable; }
 
 	void UpdateResourceTable(const FRHIUniformBufferLayout& InLayout, const void* Contents, int32 ResourceNum);
 	void UpdateResourceTable(FRHIResource** Resources, int32 ResourceNum);
 
-	virtual void Update(const void* Contents, int32 ContentsSize);
+protected:
+	TArray<TRefCountPtr<FRHIResource>> ResourceTable;
+};
 
+class FVulkanEmulatedUniformBuffer : public FVulkanUniformBuffer
+{
+public:
+	FVulkanEmulatedUniformBuffer(const FRHIUniformBufferLayout& InLayout, const void* Contents, EUniformBufferUsage InUsage, EUniformBufferValidation Validation);
 
-private:
-	TArray<TRefCountPtr<FRHIResource>> ResourceTable;
+	TArray<uint8> ConstantData;
+
+	void UpdateConstantData(const void* Contents, int32 ContentsSize);
 };
 
 class FVulkanRealUniformBuffer : public FVulkanUniformBuffer, public FVulkanResourceMultiBuffer
 {
 public:
-	FVulkanRealUniformBuffer(FVulkanDevice& Device, const FRHIUniformBufferLayout& InLayout, const void* Contents, EUniformBufferUsage Usage, EUniformBufferValidation Validation);
+	FVulkanRealUniformBuffer(FVulkanDevice& Device, const FRHIUniformBufferLayout& InLayout, const void* Contents, EUniformBufferUsage InUsage, EUniformBufferValidation Validation);
 
-	virtual void Update(const void* Contents, int32 ContentsSize);
+	void Update(const void* Contents, int32 ContentsSize);
 
 private:
 	TArray<TRefCountPtr<FRHIResource>> ResourceTable;
