workaround for UE-71464.
must run with -rtx20xxmipworkaround to enable

diff --git a/Engine/Source/Runtime/VulkanRHI/Private/VulkanTexture.cpp b/Engine/Source/Runtime/VulkanRHI/Private/VulkanTexture.cpp
--- a/Engine/Source/Runtime/VulkanRHI/Private/VulkanTexture.cpp
+++ b/Engine/Source/Runtime/VulkanRHI/Private/VulkanTexture.cpp
@@ -1480,6 +1480,28 @@ VkImageView FVulkanTextureView::StaticCreate(FVulkanDevice& Device, VkImage InIm
 	ViewInfo.subresourceRange.baseMipLevel = FirstMip;
 	ensure(NumMips != 0xFFFFFFFF);
 	ViewInfo.subresourceRange.levelCount = NumMips;
+
+	auto CheckUseNvidiaWorkaround = []() -> bool
+	{
+		if (IsRHIDeviceNVIDIA())
+		{
+			if(FParse::Param(FCommandLine::Get(), TEXT("rtx20xxmipworkaround")))
+			{
+				// Workaround for 20xx family not copying last mips correctly, so instead the view is created without the last 1x1 and 2x2 mips
+				if (GRHIAdapterName.Contains(TEXT("RTX 20")))
+				{
+					return true;
+				}
+			}
+		}
+		return false;
+	};
+	static bool NvidiaWorkaround = CheckUseNvidiaWorkaround();
+	if(NvidiaWorkaround && Format >= VK_FORMAT_BC1_RGB_UNORM_BLOCK && Format <= VK_FORMAT_BC7_SRGB_BLOCK && NumMips > 1)
+	{
+		ViewInfo.subresourceRange.levelCount = FMath::Max(1, int32(NumMips) - 2);
+	}
+
 	ensure(ArraySliceIndex != 0xFFFFFFFF);
 	ViewInfo.subresourceRange.baseArrayLayer = ArraySliceIndex;
 	ensure(NumArraySlices != 0xFFFFFFFF);
