Fixed case where vulkan PSO could be used in incompatible render-pass

diff --git a/Engine/Source/Runtime/VulkanRHI/Private/VulkanPipeline.cpp b/Engine/Source/Runtime/VulkanRHI/Private/VulkanPipeline.cpp
--- a/Engine/Source/Runtime/VulkanRHI/Private/VulkanPipeline.cpp
+++ b/Engine/Source/Runtime/VulkanRHI/Private/VulkanPipeline.cpp
@@ -1634,7 +1634,9 @@ FVulkanRHIGraphicsPipelineState* FVulkanPipelineStateCacheManager::FindInRuntime
 				for (uint32 Index = 0; Index < PSI.RenderTargetsEnabled; ++Index)
 				{
 					Ar << BlendState.RenderTargets[Index];
-					TempEnumValue = PSI.RenderTargetFormats[Index];
+					EPixelFormat PixelFormat = (EPixelFormat)PSI.RenderTargetFormats[Index];
+					bool bSRGB = (PSI.RenderTargetFlags[Index] & TexCreate_SRGB) == TexCreate_SRGB;
+					TempEnumValue = UEToVkFormat(PixelFormat, bSRGB);
 					Ar << TempEnumValue;
 				}
 				Ar << BlendState.bUseIndependentRenderTargetBlendStates;
