Build Fix

diff --git a/Engine/Source/Runtime/VulkanRHI/Private/VulkanPipeline.cpp b/Engine/Source/Runtime/VulkanRHI/Private/VulkanPipeline.cpp
--- a/Engine/Source/Runtime/VulkanRHI/Private/VulkanPipeline.cpp
+++ b/Engine/Source/Runtime/VulkanRHI/Private/VulkanPipeline.cpp
@@ -1636,7 +1636,7 @@ FVulkanRHIGraphicsPipelineState* FVulkanPipelineStateCacheManager::FindInRuntime
 					Ar << BlendState.RenderTargets[Index];
 					EPixelFormat PixelFormat = (EPixelFormat)PSI.RenderTargetFormats[Index];
 					bool bSRGB = (PSI.RenderTargetFlags[Index] & TexCreate_SRGB) == TexCreate_SRGB;
-					TempEnumValue = UEToVkFormat(PixelFormat, bSRGB);
+					TempEnumValue = UEToVkTextureFormat(PixelFormat, bSRGB);
 					Ar << TempEnumValue;
 				}
 				Ar << BlendState.bUseIndependentRenderTargetBlendStates;
