Fixed: SRV of volatile buffer was not properly updated in VulkanRHI

diff --git a/Engine/Source/Runtime/VulkanRHI/Private/VulkanResources.h b/Engine/Source/Runtime/VulkanRHI/Private/VulkanResources.h
--- a/Engine/Source/Runtime/VulkanRHI/Private/VulkanResources.h
+++ b/Engine/Source/Runtime/VulkanRHI/Private/VulkanResources.h
@@ -1302,6 +1302,7 @@ public:
 		, NumMips(InNumMips)
 		, Size(0)
 		, SourceBuffer(nullptr)
+		, VolatileBufferHandle(VK_NULL_HANDLE)
 		, VolatileLockCounter(MAX_uint32)
 	{
 	}
@@ -1315,6 +1316,7 @@ public:
 		, NumMips(0)
 		, Size(InStructuredBuffer->GetSize())
 		, SourceBuffer(nullptr)
+		, VolatileBufferHandle(VK_NULL_HANDLE)
 		, VolatileLockCounter(MAX_uint32)
 	{
 	}
@@ -1351,6 +1353,7 @@ public:
 
 protected:
 	// Used to check on volatile buffers if a new BufferView is required
+	VkBuffer VolatileBufferHandle;
 	uint32 VolatileLockCounter;
 };
 
diff --git a/Engine/Source/Runtime/VulkanRHI/Private/VulkanUAV.cpp b/Engine/Source/Runtime/VulkanRHI/Private/VulkanUAV.cpp
--- a/Engine/Source/Runtime/VulkanRHI/Private/VulkanUAV.cpp
+++ b/Engine/Source/Runtime/VulkanRHI/Private/VulkanUAV.cpp
@@ -13,6 +13,7 @@ FVulkanShaderResourceView::FVulkanShaderResourceView(FVulkanDevice* Device, FRHI
 	, Size(InSize)
 	, SourceBuffer(InSourceBuffer)
 	, SourceRHIBuffer(InRHIBuffer)
+	, VolatileBufferHandle(VK_NULL_HANDLE)
 	, VolatileLockCounter(MAX_uint32)
 {
 	int32 NumBuffers = SourceBuffer->IsVolatile() ? 1 : SourceBuffer->GetNumBuffers();
@@ -38,6 +39,9 @@ void FVulkanShaderResourceView::Clear()
 	}
 	SourceTexture = nullptr;
 	Device = nullptr;
+
+	VolatileBufferHandle = VK_NULL_HANDLE;
+	VolatileLockCounter = MAX_uint32;
 }
 
 void FVulkanShaderResourceView::Rename(FRHIResource* InRHIBuffer, FVulkanResourceMultiBuffer* InSourceBuffer, uint32 InSize, EPixelFormat InFormat)
@@ -55,6 +59,7 @@ void FVulkanShaderResourceView::Rename(FRHIResource* InRHIBuffer, FVulkanResourc
 	Size = InSize;
 	SourceBuffer = InSourceBuffer;
 	SourceRHIBuffer = InRHIBuffer;
+	VolatileBufferHandle = VK_NULL_HANDLE;	
 	VolatileLockCounter = MAX_uint32;
 }
 
@@ -69,14 +74,19 @@ void FVulkanShaderResourceView::UpdateView()
 	{
 		if (SourceBuffer->IsVolatile() && VolatileLockCounter != SourceBuffer->GetVolatileLockCounter())
 		{
-			bool bNeedNewView = true;
+			VkBuffer SourceVolatileBufferHandle = SourceBuffer->GetHandle();
+
 			// We might end up with the same BufferView, so do not recreate in that case
-			if (!BufferViews[0] || BufferViews[0]->Offset != SourceBuffer->GetOffset() || BufferViews[0]->Size != Size)
+			if (!BufferViews[0] 
+				|| BufferViews[0]->Offset != SourceBuffer->GetOffset() 
+				|| BufferViews[0]->Size != Size
+				|| VolatileBufferHandle != SourceVolatileBufferHandle)
 			{
 				BufferViews[0] = nullptr;
 			}
 
 			VolatileLockCounter = SourceBuffer->GetVolatileLockCounter();
+			VolatileBufferHandle = SourceVolatileBufferHandle;
 		}
 		else if (SourceBuffer->IsDynamic())
 		{
