Fixed: UI blinks on MALI based Android devices running vulkan

diff --git a/Engine/Source/Runtime/VulkanRHI/Private/VulkanRenderTarget.cpp b/Engine/Source/Runtime/VulkanRHI/Private/VulkanRenderTarget.cpp
--- a/Engine/Source/Runtime/VulkanRHI/Private/VulkanRenderTarget.cpp
+++ b/Engine/Source/Runtime/VulkanRHI/Private/VulkanRenderTarget.cpp
@@ -631,24 +631,6 @@ void FVulkanCommandListContext::RHICopyToResolveTarget(FTextureRHIParamRef Sourc
 		return;
 	}
 
-	
-	// Disabled as it causes flicker on MALI
-	// TODO: With explicit render-pases this does not help with fusing Post->UI renderpass anyway
-	//if (SourceTextureRHI == DestTextureRHI)
-	//{
-	//	FRHITexture2D* SourceTexture2D = SourceTextureRHI->GetTexture2D();
-	//	if (SourceTexture2D)
-	//	{
-	//		FVulkanTexture2D* VulkanSourceTexture2D  = (FVulkanTexture2D*)SourceTexture2D;
-	//		if (VulkanSourceTexture2D->GetBackBuffer() != nullptr)
-	//		{
-	//			// skip Backbuffer implicit transition to Readable, to avoid splitting Post->UI renderpass
-	//			// do explicit transition when need to read from Backbuffer
-	//			return;
-	//		}
-	//	}
-	//}
-
 	RHITransitionResources(EResourceTransitionAccess::EReadable, &SourceTextureRHI, 1);
 
 	auto CopyImage = [](FTransitionAndLayoutManager& InRenderPassState, FVulkanCmdBuffer* InCmdBuffer, FVulkanSurface& SrcSurface, FVulkanSurface& DstSurface, uint32 SrcNumLayers, uint32 DstNumLayers, const FResolveParams& ResolveParams)
