Fix crash due to swapchain going out of date when acquiring the backbuffer.

diff --git a/Engine/Source/Runtime/VulkanRHI/Private/VulkanViewport.cpp b/Engine/Source/Runtime/VulkanRHI/Private/VulkanViewport.cpp
--- a/Engine/Source/Runtime/VulkanRHI/Private/VulkanViewport.cpp
+++ b/Engine/Source/Runtime/VulkanRHI/Private/VulkanViewport.cpp
@@ -78,6 +78,12 @@ void FVulkanBackBuffer::AcquireBackBufferImage(FVulkanCommandListContext& Contex
 		check(Viewport->AcquiredImageIndex == -1);
 		
 		Viewport->AcquireImageIndex();
+		// If swapchain got invalidated (OUT_OF_DATE etc) in the above call, we may end up not having a valid viewport pointer at this point. Abort the whole thing.
+		if (Viewport == nullptr)
+		{
+			return;
+		}
+
 		int32 AcquiredImageIndex = Viewport->AcquiredImageIndex;
 		check(AcquiredImageIndex >= 0 && AcquiredImageIndex < Viewport->TextureViews.Num());
 
