Allow for better overriding of Vulkan PSO functions

diff --git a/Engine/Source/Runtime/VulkanRHI/Private/VulkanDevice.cpp b/Engine/Source/Runtime/VulkanRHI/Private/VulkanDevice.cpp
--- a/Engine/Source/Runtime/VulkanRHI/Private/VulkanDevice.cpp
+++ b/Engine/Source/Runtime/VulkanRHI/Private/VulkanDevice.cpp
@@ -847,17 +847,7 @@ void FVulkanDevice::InitGPU(int32 DeviceIndex)
 
 	PipelineStateCache = new FVulkanPipelineStateCacheManager(this);
 
-	TArray<FString> CacheFilenames;
-	FString StagedCacheDirectory = FPaths::ProjectDir() / TEXT("Build") / TEXT("ShaderCaches") / FPlatformProperties::IniPlatformName();
-
-	// look for any staged caches
-	TArray<FString> StagedCaches;
-	IFileManager::Get().FindFiles(StagedCaches, *StagedCacheDirectory, TEXT("cache"));
-	// FindFiles returns the filenames without directory, so prepend the stage directory
-	for (const FString& Filename : StagedCaches)
-	{
-		CacheFilenames.Add(StagedCacheDirectory / Filename);
-	}
+	TArray<FString> CacheFilenames = FVulkanPlatform::GetPSOCacheFilenames();
 
 	// always look in the saved directory (for the cache from previous run that wasn't moved over to stage directory)
 	CacheFilenames.Add(VulkanRHI::GetPipelineCacheFilename());
diff --git a/Engine/Source/Runtime/VulkanRHI/Private/VulkanGenericPlatform.cpp b/Engine/Source/Runtime/VulkanRHI/Private/VulkanGenericPlatform.cpp
--- a/Engine/Source/Runtime/VulkanRHI/Private/VulkanGenericPlatform.cpp
+++ b/Engine/Source/Runtime/VulkanRHI/Private/VulkanGenericPlatform.cpp
@@ -10,3 +10,66 @@ void FVulkanGenericPlatform::SetupFeatureLevels()
 	GShaderPlatformForFeatureLevel[ERHIFeatureLevel::SM4] = SP_VULKAN_SM4;
 	GShaderPlatformForFeatureLevel[ERHIFeatureLevel::SM5] = SP_VULKAN_SM5;
 }
+
+bool FVulkanGenericPlatform::PSOBinaryCacheMatches(FVulkanDevice* Device, const TArray<uint8>& DeviceCache)
+{
+	if (DeviceCache.Num() > 4)
+	{
+		uint32* Data = (uint32*)DeviceCache.GetData();
+		uint32 HeaderSize = *Data++;
+		// 16 is HeaderSize + HeaderVersion
+		if (HeaderSize == 16 + VK_UUID_SIZE)
+		{
+			uint32 HeaderVersion = *Data++;
+			if (HeaderVersion == VK_PIPELINE_CACHE_HEADER_VERSION_ONE)
+			{
+				uint32 VendorID = *Data++;
+				const VkPhysicalDeviceProperties& DeviceProperties = Device->GetDeviceProperties();
+				if (VendorID == DeviceProperties.vendorID)
+				{
+					uint32 DeviceID = *Data++;
+					if (DeviceID == DeviceProperties.deviceID)
+					{
+						uint8* Uuid = (uint8*)Data;
+						if (FMemory::Memcmp(DeviceProperties.pipelineCacheUUID, Uuid, VK_UUID_SIZE) == 0)
+						{
+							// This particular binary cache matches this device
+							return true;
+						}
+					}
+				}
+			}
+		}
+	}
+
+	return false;
+}
+
+FString FVulkanGenericPlatform::CreatePSOBinaryCacheFilename(FVulkanDevice* Device, FString CacheFilename)
+{
+	const VkPhysicalDeviceProperties& DeviceProperties = Device->GetDeviceProperties();
+	FString BinaryCacheAppendage = FString::Printf(TEXT(".%x.%x"), DeviceProperties.vendorID, DeviceProperties.deviceID);
+	if (!CacheFilename.EndsWith(BinaryCacheAppendage))
+	{
+		CacheFilename += BinaryCacheAppendage;
+	}
+
+	return CacheFilename;
+}
+
+TArray<FString> FVulkanGenericPlatform::GetPSOCacheFilenames()
+{
+	TArray<FString> CacheFilenames;
+	FString StagedCacheDirectory = FPaths::ProjectDir() / TEXT("Build") / TEXT("ShaderCaches") / FPlatformProperties::IniPlatformName();
+
+	// look for any staged caches
+	TArray<FString> StagedCaches;
+	IFileManager::Get().FindFiles(StagedCaches, *StagedCacheDirectory, TEXT("cache"));
+	// FindFiles returns the filenames without directory, so prepend the stage directory
+	for (const FString& Filename : StagedCaches)
+	{
+		CacheFilenames.Add(StagedCacheDirectory / Filename);
+	}
+
+	return CacheFilenames;
+}
diff --git a/Engine/Source/Runtime/VulkanRHI/Private/VulkanGenericPlatform.h b/Engine/Source/Runtime/VulkanRHI/Private/VulkanGenericPlatform.h
--- a/Engine/Source/Runtime/VulkanRHI/Private/VulkanGenericPlatform.h
+++ b/Engine/Source/Runtime/VulkanRHI/Private/VulkanGenericPlatform.h
@@ -7,6 +7,7 @@
 #include "RHI.h"	// for GShaderPlatformForFeatureLevel and its friends
 
 struct FOptionalVulkanDeviceExtensions;
+class FVulkanDevice;
 
 // the platform interface, and empty implementations for platforms that don't need em
 class FVulkanGenericPlatform 
@@ -104,4 +105,13 @@ public:
 
 	// Does the platform require resolve attachments in its MSAA renderpasses
 	static bool RequiresRenderPassResolveAttachments() { return false; }
+
+	// Checks if the PSO cache matches the expected vulkan device properties
+	static bool PSOBinaryCacheMatches(FVulkanDevice* Device, const TArray<uint8>& DeviceCache);
+
+	// Will create the correct format from a generic pso filename
+	static FString CreatePSOBinaryCacheFilename(FVulkanDevice* Device, FString CacheFilename);
+
+	// Gathers a list of pso cache filenames to attempt to load
+	static TArray<FString> GetPSOCacheFilenames();
 };
diff --git a/Engine/Source/Runtime/VulkanRHI/Private/VulkanPipeline.cpp b/Engine/Source/Runtime/VulkanRHI/Private/VulkanPipeline.cpp
--- a/Engine/Source/Runtime/VulkanRHI/Private/VulkanPipeline.cpp
+++ b/Engine/Source/Runtime/VulkanRHI/Private/VulkanPipeline.cpp
@@ -285,18 +285,13 @@ bool FVulkanPipelineStateCacheManager::Load(const TArray<FString>& CacheFilename
 	// Try to load device cache first
 	for (const FString& CacheFilename : CacheFilenames)
 	{
-		const VkPhysicalDeviceProperties& DeviceProperties = Device->GetDeviceProperties();
 		double BeginTime = FPlatformTime::Seconds();
-		FString BinaryCacheAppendage = FString::Printf(TEXT(".%x.%x"), DeviceProperties.vendorID, DeviceProperties.deviceID);
-		FString BinaryCacheFilename = CacheFilename;
-		if (!CacheFilename.EndsWith(BinaryCacheAppendage))
-		{
-			BinaryCacheFilename += BinaryCacheAppendage;
-		}
+		FString BinaryCacheFilename = FVulkanPlatform::CreatePSOBinaryCacheFilename(Device, CacheFilename);
+
 		TArray<uint8> DeviceCache;
 		if (FFileHelper::LoadFileToArray(DeviceCache, *BinaryCacheFilename, FILEREAD_Silent))
 		{
-			if (BinaryCacheMatches(Device, DeviceCache))
+			if (FVulkanPlatform::PSOBinaryCacheMatches(Device, DeviceCache))
 			{
 				VkPipelineCacheCreateInfo PipelineCacheInfo;
 				ZeroVulkanStruct(PipelineCacheInfo, VK_STRUCTURE_TYPE_PIPELINE_CACHE_CREATE_INFO);
@@ -337,14 +332,8 @@ bool FVulkanPipelineStateCacheManager::Load(const TArray<FString>& CacheFilename
 	{
 		for (const FString& CacheFilename : CacheFilenames)
 		{
-			const VkPhysicalDeviceProperties& DeviceProperties = Device->GetDeviceProperties();
 			double BeginTime = FPlatformTime::Seconds();
-			FString BinaryCacheAppendage = FString::Printf(TEXT(".%x.%x"), DeviceProperties.vendorID, DeviceProperties.deviceID);
-			FString LruCacheFilename = CacheFilename;
-			if (!CacheFilename.EndsWith(BinaryCacheAppendage))
-			{
-				LruCacheFilename += BinaryCacheAppendage;
-			}
+			FString LruCacheFilename = FVulkanPlatform::CreatePSOBinaryCacheFilename(Device, CacheFilename);
 			LruCacheFilename += TEXT(".lru");
 			LruCacheFilename.ReplaceInline(TEXT("TempScanVulkanPSO_"), TEXT("VulkanPSO_"));  //lru files do not use the rename trick...but are still protected against corruption indirectly
 
@@ -533,13 +522,7 @@ void FVulkanPipelineStateCacheManager::Save(const FString& CacheFilename, bool b
 		VkResult Result = VulkanRHI::vkGetPipelineCacheData(Device->GetInstanceHandle(), PipelineCache, &Size, DeviceCache.GetData());
 		if (Result == VK_SUCCESS)
 		{
-			const VkPhysicalDeviceProperties& DeviceProperties = Device->GetDeviceProperties();
-			FString BinaryCacheAppendage = FString::Printf(TEXT(".%x.%x"), DeviceProperties.vendorID, DeviceProperties.deviceID);
-			FString BinaryCacheFilename = CacheFilename;
-			if (!BinaryCacheFilename.EndsWith(BinaryCacheAppendage))
-			{
-				BinaryCacheFilename += BinaryCacheAppendage;
-			}
+			FString BinaryCacheFilename = FVulkanPlatform::CreatePSOBinaryCacheFilename(Device, CacheFilename);
 
 			if (FFileHelper::SaveArrayToFile(DeviceCache, *BinaryCacheFilename))
 			{
@@ -577,13 +560,7 @@ void FVulkanPipelineStateCacheManager::Save(const FString& CacheFilename, bool b
 		PipelineSizeList.GenerateValueArray(File.PipelineSizes);
 		File.Save(Ar);
 
-		const VkPhysicalDeviceProperties& DeviceProperties = Device->GetDeviceProperties();
-		FString BinaryCacheAppendage = FString::Printf(TEXT(".%x.%x"), DeviceProperties.vendorID, DeviceProperties.deviceID);
-		FString LruCacheFilename = CacheFilename;
-		if (!CacheFilename.EndsWith(BinaryCacheAppendage))
-		{
-			LruCacheFilename += BinaryCacheAppendage;
-		}
+		FString LruCacheFilename = FVulkanPlatform::CreatePSOBinaryCacheFilename(Device, CacheFilename);
 		LruCacheFilename += TEXT(".lru");
 
 		if (FFileHelper::SaveArrayToFile(MemFile, *LruCacheFilename))
@@ -1871,40 +1848,6 @@ inline void SerializeArray(FArchive& Ar, TArray<T*>& Array)
 	}
 }
 
-bool FVulkanPipelineStateCacheManager::BinaryCacheMatches(FVulkanDevice* InDevice, const TArray<uint8>& DeviceCache)
-{
-	if (DeviceCache.Num() > 4)
-	{
-		uint32* Data = (uint32*)DeviceCache.GetData();
-		uint32 HeaderSize = *Data++;
-		// 16 is HeaderSize + HeaderVersion
-		if (HeaderSize == 16 + VK_UUID_SIZE)
-		{
-			uint32 HeaderVersion = *Data++;
-			if (HeaderVersion == VK_PIPELINE_CACHE_HEADER_VERSION_ONE)
-			{
-				uint32 VendorID = *Data++;
-				const VkPhysicalDeviceProperties& DeviceProperties = InDevice->GetDeviceProperties();
-				if (VendorID == DeviceProperties.vendorID)
-				{
-					uint32 DeviceID = *Data++;
-					if (DeviceID == DeviceProperties.deviceID)
-					{
-						uint8* Uuid = (uint8*)Data;
-						if (FMemory::Memcmp(DeviceProperties.pipelineCacheUUID, Uuid, VK_UUID_SIZE) == 0)
-						{
-							// This particular binary cache matches this device
-							return true;
-						}
-					}
-				}
-			}
-		}
-	}
-
-	return false;
-}
-
 #if VULKAN_ENABLE_LRU_CACHE
 void FVulkanPipelineStateCacheManager::FVKPipelineLRU::DeleteVkPipeline(FVulkanGfxPipeline* GfxPipeline)
 {
diff --git a/Engine/Source/Runtime/VulkanRHI/Private/VulkanPipeline.h b/Engine/Source/Runtime/VulkanRHI/Private/VulkanPipeline.h
--- a/Engine/Source/Runtime/VulkanRHI/Private/VulkanPipeline.h
+++ b/Engine/Source/Runtime/VulkanRHI/Private/VulkanPipeline.h
@@ -627,8 +627,6 @@ private:
 		bool Load(FArchive& Ar);
 	};
 #endif
-
-	static bool BinaryCacheMatches(FVulkanDevice* InDevice, const TArray<uint8>& DeviceCache);
 };
 
 // Common pipeline class
