Fix build failure

diff --git a/Engine/Source/Runtime/VulkanRHI/Private/VulkanGenericPlatform.cpp b/Engine/Source/Runtime/VulkanRHI/Private/VulkanGenericPlatform.cpp
--- a/Engine/Source/Runtime/VulkanRHI/Private/VulkanGenericPlatform.cpp
+++ b/Engine/Source/Runtime/VulkanRHI/Private/VulkanGenericPlatform.cpp
@@ -2,6 +2,7 @@
 
 #include "VulkanRHIPrivate.h"
 #include "VulkanGenericPlatform.h"
+#include "HAL/FileManager.h"
 
 void FVulkanGenericPlatform::SetupFeatureLevels()
 {
