Vulkan: replaced PresentInfoExtensions hook with a Present hook.

diff --git a/Engine/Source/Runtime/VulkanRHI/Private/VulkanGenericPlatform.cpp b/Engine/Source/Runtime/VulkanRHI/Private/VulkanGenericPlatform.cpp
--- a/Engine/Source/Runtime/VulkanRHI/Private/VulkanGenericPlatform.cpp
+++ b/Engine/Source/Runtime/VulkanRHI/Private/VulkanGenericPlatform.cpp
@@ -74,3 +74,8 @@ TArray<FString> FVulkanGenericPlatform::GetPSOCacheFilenames()
 
 	return CacheFilenames;
 }
+
+VkResult FVulkanGenericPlatform::Present(VkQueue Queue, VkPresentInfoKHR& PresentInfo)
+{
+	return VulkanRHI::vkQueuePresentKHR(Queue, &PresentInfo);
+}
diff --git a/Engine/Source/Runtime/VulkanRHI/Private/VulkanGenericPlatform.h b/Engine/Source/Runtime/VulkanRHI/Private/VulkanGenericPlatform.h
--- a/Engine/Source/Runtime/VulkanRHI/Private/VulkanGenericPlatform.h
+++ b/Engine/Source/Runtime/VulkanRHI/Private/VulkanGenericPlatform.h
@@ -94,8 +94,8 @@ public:
 
 	static bool RequiresSwapchainGeneralInitialLayout() { return false; }
 	
-	// Allow platforms to add extension features to the PresentInfo pNext chain
-	static void EnablePresentInfoExtensions(VkPresentInfoKHR& PresentInfo) {}
+	// Allow platforms to do extra work on present
+	static VkResult Present(VkQueue Queue, VkPresentInfoKHR& PresentInfo);
 
 	// Ensure the last frame completed on the GPU
 	static bool RequiresWaitingForFrameCompletionEvent() { return true; }
diff --git a/Engine/Source/Runtime/VulkanRHI/Private/VulkanSwapChain.cpp b/Engine/Source/Runtime/VulkanRHI/Private/VulkanSwapChain.cpp
--- a/Engine/Source/Runtime/VulkanRHI/Private/VulkanSwapChain.cpp
+++ b/Engine/Source/Runtime/VulkanRHI/Private/VulkanSwapChain.cpp
@@ -428,6 +428,7 @@ FVulkanSwapChain::FVulkanSwapChain(VkInstance InInstance, FVulkanDevice& InDevic
 
 	//ensure(SwapChainInfo.imageExtent.width >= SurfProperties.minImageExtent.width && SwapChainInfo.imageExtent.width <= SurfProperties.maxImageExtent.width);
 	//ensure(SwapChainInfo.imageExtent.height >= SurfProperties.minImageExtent.height && SwapChainInfo.imageExtent.height <= SurfProperties.maxImageExtent.height);
+	UE_LOG(LogVulkanRHI, Verbose, TEXT("Creating new VK swapchain with format %d, color space %d"), static_cast<uint32>(SwapChainInfo.imageFormat), static_cast<uint32>(SwapChainInfo.imageColorSpace));
 
 	VERIFYVULKANRESULT_EXPANDED(VulkanRHI::vkCreateSwapchainKHR(Device.GetInstanceHandle(), &SwapChainInfo, VULKAN_CPU_ALLOCATOR, &SwapChain));
 
@@ -817,7 +818,6 @@ FVulkanSwapChain::EStatus FVulkanSwapChain::Present(FVulkanQueue* GfxQueue, FVul
 
 	const int32 SyncInterval = LockToVsync ? RHIGetSyncInterval() : 0;
 	ensureMsgf(SyncInterval <= 3 && SyncInterval >= 0, TEXT("Unsupported sync interval: %i"), SyncInterval);
-	FVulkanPlatform::EnablePresentInfoExtensions(Info);
 
 #if VULKAN_SUPPORTS_GOOGLE_DISPLAY_TIMING
 	if (GVulkanExtensionFramePacer && Device.GetOptionalExtensions().HasGoogleDisplayTiming)
@@ -871,7 +871,7 @@ FVulkanSwapChain::EStatus FVulkanSwapChain::Present(FVulkanQueue* GfxQueue, FVul
 	{
 		SCOPE_CYCLE_COUNTER(STAT_VulkanQueuePresent);
 		uint32 IdleStart = FPlatformTime::Cycles();
-		VkResult PresentResult = VulkanRHI::vkQueuePresentKHR(PresentQueue->GetHandle(), &Info);
+		VkResult PresentResult = FVulkanPlatform::Present(PresentQueue->GetHandle(), Info);
 		uint32 ThisCycles = FPlatformTime::Cycles() - IdleStart;
 		if (IsInRHIThread())
 		{
