Fixed the static analysis (UE-77570).

diff --git a/Engine/Source/Editor/UnrealEd/Private/PlayLevel.cpp b/Engine/Source/Editor/UnrealEd/Private/PlayLevel.cpp
--- a/Engine/Source/Editor/UnrealEd/Private/PlayLevel.cpp
+++ b/Engine/Source/Editor/UnrealEd/Private/PlayLevel.cpp
@@ -530,7 +530,7 @@ void UEditorEngine::EndPlayMap()
 	}
 
 	//ensure stereo rendering is disabled in case we need to re-enable next PIE run.
-	if (GEngine && GEngine->StereoRenderingDevice)
+	if (GEngine->StereoRenderingDevice)
 	{
 		GEngine->StereoRenderingDevice->EnableStereo(false);
 	}
diff --git a/Engine/Source/Runtime/VulkanRHI/Private/VulkanViewport.cpp b/Engine/Source/Runtime/VulkanRHI/Private/VulkanViewport.cpp
--- a/Engine/Source/Runtime/VulkanRHI/Private/VulkanViewport.cpp
+++ b/Engine/Source/Runtime/VulkanRHI/Private/VulkanViewport.cpp
@@ -75,9 +75,9 @@ void FVulkanBackBuffer::AcquireBackBufferImage(FVulkanCommandListContext& Contex
 	check(Viewport);
 	if (Surface.Image == VK_NULL_HANDLE)
 	{
-		check(Viewport->AcquiredImageIndex == -1);
+		check(Viewport->AcquiredImageIndex == -1); //-V595
 		
-		Viewport->AcquireImageIndex();
+		Viewport->AcquireImageIndex(); //-V595
 		// If swapchain got invalidated (OUT_OF_DATE etc) in the above call, we may end up not having a valid viewport pointer at this point. Abort the whole thing.
 		if (Viewport == nullptr)
 		{
