Fix a race condition where after a JoinFreeBlock and before ReleaseBuffer has time to lock a mutex another suballocation is allocated in the BufferAllocation about to be released failing a check when attempting to JoinFreeBlocks again

Move the first JoinFreeBlocks inside the mutex

diff --git a/Engine/Source/Runtime/VulkanRHI/Private/VulkanMemory.cpp b/Engine/Source/Runtime/VulkanRHI/Private/VulkanMemory.cpp
--- a/Engine/Source/Runtime/VulkanRHI/Private/VulkanMemory.cpp
+++ b/Engine/Source/Runtime/VulkanRHI/Private/VulkanMemory.cpp
@@ -1361,10 +1361,14 @@ namespace VulkanRHI
 	void FResourceHeapManager::ReleaseBuffer(FBufferAllocation* BufferAllocation)
 	{
 		FScopeLock ScopeLock(&GResourceHeapLock);
-		check(BufferAllocation->JoinFreeBlocks());
-		UsedBufferAllocations[BufferAllocation->PoolSizeIndex].RemoveSingleSwap(BufferAllocation, false);
-		BufferAllocation->FrameFreed = GFrameNumberRenderThread;
-		FreeBufferAllocations[BufferAllocation->PoolSizeIndex].Add(BufferAllocation);
+
+		if (BufferAllocation->JoinFreeBlocks())
+		{
+			check(BufferAllocation->JoinFreeBlocks());
+			UsedBufferAllocations[BufferAllocation->PoolSizeIndex].RemoveSingleSwap(BufferAllocation, false);
+			BufferAllocation->FrameFreed = GFrameNumberRenderThread;
+			FreeBufferAllocations[BufferAllocation->PoolSizeIndex].Add(BufferAllocation);
+		}
 	}
 #if VULKAN_SUPPORTS_DEDICATED_ALLOCATION
 	FOldResourceAllocation* FResourceHeapManager::AllocateDedicatedImageMemory(VkImage Image, const VkMemoryRequirements& MemoryReqs, VkMemoryPropertyFlags MemoryPropertyFlags, const char* File, uint32 Line)
@@ -1594,11 +1598,7 @@ namespace VulkanRHI
 			check(UsedSize >= 0);
 		}
 
-
-		if (JoinFreeBlocks())
-		{
-			Owner->ReleaseBuffer(this);
-		}
+		Owner->ReleaseBuffer(this);
 	}
 
 	void FBufferAllocation::Destroy(FVulkanDevice* Device)
