fix vulkan gfx2gfx transition

diff --git a/Engine/Source/Runtime/VulkanRHI/Private/VulkanRenderTarget.cpp b/Engine/Source/Runtime/VulkanRHI/Private/VulkanRenderTarget.cpp
--- a/Engine/Source/Runtime/VulkanRHI/Private/VulkanRenderTarget.cpp
+++ b/Engine/Source/Runtime/VulkanRHI/Private/VulkanRenderTarget.cpp
@@ -1537,6 +1537,10 @@ void FVulkanCommandListContext::TransitionResources(const FPendingTransition& Pe
 				VkPipelineStageFlags SourceStage = VK_PIPELINE_STAGE_ALL_COMMANDS_BIT, DestStage = VK_PIPELINE_STAGE_ALL_COMMANDS_BIT;
 				switch (PendingTransition.TransitionPipeline)
 				{
+				case EResourceTransitionPipeline::EGfxToGfx:
+					SourceStage = VK_PIPELINE_STAGE_ALL_GRAPHICS_BIT;
+					DestStage = VK_PIPELINE_STAGE_ALL_GRAPHICS_BIT;
+					break;
 				case EResourceTransitionPipeline::EGfxToCompute:
 					SourceStage = VK_PIPELINE_STAGE_ALL_GRAPHICS_BIT;
 					DestStage = VK_PIPELINE_STAGE_COMPUTE_SHADER_BIT;
