Vulkan: make dynres support per platform.

diff --git a/Engine/Source/Runtime/VulkanRHI/Private/VulkanGenericPlatform.h b/Engine/Source/Runtime/VulkanRHI/Private/VulkanGenericPlatform.h
--- a/Engine/Source/Runtime/VulkanRHI/Private/VulkanGenericPlatform.h
+++ b/Engine/Source/Runtime/VulkanRHI/Private/VulkanGenericPlatform.h
@@ -89,6 +89,9 @@ public:
 
 	static bool SupportParallelRenderingTasks() { return true; }
 
+	/** The status quo is false, so the default is chosen to not change it. As platforms opt in it may be better to flip the default. */
+	static bool SupportsDynamicResolution() { return false; }
+
 	// Allow platforms to add extension features to the DeviceInfo pNext chain
 	static void EnablePhysicalDeviceFeatureExtensions(VkDeviceCreateInfo& DeviceInfo) {}
 
diff --git a/Engine/Source/Runtime/VulkanRHI/Private/VulkanRHI.cpp b/Engine/Source/Runtime/VulkanRHI/Private/VulkanRHI.cpp
--- a/Engine/Source/Runtime/VulkanRHI/Private/VulkanRHI.cpp
+++ b/Engine/Source/Runtime/VulkanRHI/Private/VulkanRHI.cpp
@@ -683,6 +683,7 @@ void FVulkanDynamicRHI::InitInstance()
 
 		// Initialize the RHI capabilities.
 		GRHISupportsFirstInstance = true;
+		GRHISupportsDynamicResolution = FVulkanPlatform::SupportsDynamicResolution();
 		GSupportsDepthBoundsTest = Device->GetPhysicalFeatures().depthBounds != 0;
 		GSupportsRenderTargetFormat_PF_G8 = false;	// #todo-rco
 		GRHISupportsTextureStreaming = true;
