Vulkan: a better way to set applicationVersion.

diff --git a/Engine/Source/Runtime/VulkanRHI/Private/VulkanRHI.cpp b/Engine/Source/Runtime/VulkanRHI/Private/VulkanRHI.cpp
--- a/Engine/Source/Runtime/VulkanRHI/Private/VulkanRHI.cpp
+++ b/Engine/Source/Runtime/VulkanRHI/Private/VulkanRHI.cpp
@@ -5,6 +5,7 @@
 =============================================================================*/
 
 #include "VulkanRHIPrivate.h"
+#include "BuildSettings.h"
 #include "HardwareInfo.h"
 #include "VulkanShaderResources.h"
 #include "VulkanResources.h"
@@ -369,7 +370,7 @@ void FVulkanDynamicRHI::CreateInstance()
 	VkApplicationInfo AppInfo;
 	ZeroVulkanStruct(AppInfo, VK_STRUCTURE_TYPE_APPLICATION_INFO);
 	AppInfo.pApplicationName = bDisableEngineRegistration ? nullptr : ProjectNameConverter.Get();
-	AppInfo.applicationVersion = 0;	// Do we want FApp::GetBuildVersion() ?
+	AppInfo.applicationVersion = static_cast<uint32>(BuildSettings::GetCurrentChangelist()) | (BuildSettings::IsLicenseeVersion() ? 0x80000000 : 0);
 	AppInfo.pEngineName = bDisableEngineRegistration ? nullptr : EngineNameConverter.Get();
 	AppInfo.engineVersion = FEngineVersion::Current().GetMinor();
 	AppInfo.apiVersion = UE_VK_API_VERSION;
diff --git a/Engine/Source/Runtime/VulkanRHI/VulkanRHI.Build.cs b/Engine/Source/Runtime/VulkanRHI/VulkanRHI.Build.cs
--- a/Engine/Source/Runtime/VulkanRHI/VulkanRHI.Build.cs
+++ b/Engine/Source/Runtime/VulkanRHI/VulkanRHI.Build.cs
@@ -31,7 +31,8 @@ public class VulkanRHI : ModuleRules
 				"RenderCore", 
 				"UtilityShaders",
 				"HeadMountedDisplay",
-                "PreLoadScreen"
+                "PreLoadScreen",
+				"BuildSettings"
             }
         );
 
