Vulkan: handle VK_ERROR_OUT_OF_DATE in AcquireImageIndex for DelayAcquire (UE-77590).

diff --git a/Engine/Source/Runtime/VulkanRHI/Private/VulkanViewport.cpp b/Engine/Source/Runtime/VulkanRHI/Private/VulkanViewport.cpp
--- a/Engine/Source/Runtime/VulkanRHI/Private/VulkanViewport.cpp
+++ b/Engine/Source/Runtime/VulkanRHI/Private/VulkanViewport.cpp
@@ -233,6 +233,18 @@ void FVulkanViewport::AcquireImageIndex()
 	check(AcquiredImageIndex != -1);
 }
 
+bool FVulkanViewport::TryAcquireImageIndex()
+{
+	int NewImageIndex = DoAcquireImageIndex(this);
+	if (NewImageIndex != -1)
+	{
+		AcquiredImageIndex = NewImageIndex;
+		return true;
+	}
+
+	return false;
+}
+
 FTexture2DRHIRef FVulkanViewport::GetBackBuffer(FRHICommandListImmediate& RHICmdList)
 {
 	check(IsInRenderingThread());
@@ -733,6 +745,7 @@ inline static void CopyImageToBackBuffer(FVulkanCmdBuffer* CmdBuffer, bool bSour
 bool FVulkanViewport::Present(FVulkanCommandListContext* Context, FVulkanCmdBuffer* CmdBuffer, FVulkanQueue* Queue, FVulkanQueue* PresentQueue, bool bLockToVsync)
 {
 	FPlatformAtomics::AtomicStore(&LockToVsync, bLockToVsync ? 1 : 0);
+	bool bFailedToDelayAcquireBackbuffer = false;
 
 	//Transition back buffer to presentable and submit that command
 	check(CmdBuffer->IsOutsideRenderPass());
@@ -742,14 +755,20 @@ bool FVulkanViewport::Present(FVulkanCommandListContext* Context, FVulkanCmdBuff
 		if (GVulkanDelayAcquireImage == EDelayAcquireImageType::DelayAcquire && RenderingBackBuffer)
 		{
 			SCOPE_CYCLE_COUNTER(STAT_VulkanAcquireBackBuffer);
-			AcquireImageIndex();
-
-			uint32 WindowSizeX = FMath::Min(SizeX, SwapChain->InternalWidth);
-			uint32 WindowSizeY = FMath::Min(SizeY, SwapChain->InternalHeight);
+			// swapchain can go out of date, do not crash at this point
+			if (LIKELY(TryAcquireImageIndex()))
+			{
+				uint32 WindowSizeX = FMath::Min(SizeX, SwapChain->InternalWidth);
+				uint32 WindowSizeY = FMath::Min(SizeY, SwapChain->InternalHeight);
 
-			Context->RHIPushEvent(TEXT("CopyImageToBackBuffer"), FColor::Blue);
-			CopyImageToBackBuffer(CmdBuffer, true, RenderingBackBuffer->Surface.Image, BackBufferImages[AcquiredImageIndex], SizeX, SizeY, WindowSizeX, WindowSizeY);
-			Context->RHIPopEvent();
+				Context->RHIPushEvent(TEXT("CopyImageToBackBuffer"), FColor::Blue);
+				CopyImageToBackBuffer(CmdBuffer, true, RenderingBackBuffer->Surface.Image, BackBufferImages[AcquiredImageIndex], SizeX, SizeY, WindowSizeX, WindowSizeY);
+				Context->RHIPopEvent();
+			}
+			else
+			{
+				bFailedToDelayAcquireBackbuffer = true;
+			}
 		}
 		else
 		{
@@ -766,11 +785,32 @@ bool FVulkanViewport::Present(FVulkanCommandListContext* Context, FVulkanCmdBuff
 
 	if (FVulkanPlatform::SupportsStandardSwapchain())
 	{
-		if (GVulkanDelayAcquireImage == EDelayAcquireImageType::DelayAcquire)
+		if (LIKELY(!bFailedToDelayAcquireBackbuffer))
 		{
-			CmdBuffer->AddWaitSemaphore(VK_PIPELINE_STAGE_BOTTOM_OF_PIPE_BIT, AcquiredSemaphore);
+			if (GVulkanDelayAcquireImage == EDelayAcquireImageType::DelayAcquire)
+			{
+				CmdBuffer->AddWaitSemaphore(VK_PIPELINE_STAGE_BOTTOM_OF_PIPE_BIT, AcquiredSemaphore);
+			}
+			Queue->Submit(CmdBuffer, RenderingDoneSemaphores[AcquiredImageIndex]->GetHandle());
+		}
+		else
+		{
+			// failing to do the delayacquire can only happen if we were in this mode to begin with
+			check(GVulkanDelayAcquireImage == EDelayAcquireImageType::DelayAcquire);
+
+			UE_LOG(LogVulkanRHI, Log, TEXT("AcquireNextImage() failed due to the outdated swapchain, not even attempting to present."));
+
+			// cannot just throw out this command buffer (needs to be submitted or other checks fail)
+			Queue->Submit(CmdBuffer);
+			RecreateSwapchain(WindowHandle, true);
+
+			// Swapchain creation pushes some commands - flush the command buffers now to begin with a fresh state
+			Device->SubmitCommandsAndFlushGPU();
+			Device->WaitUntilIdle();
+
+			// early exit
+			return (int32)FVulkanSwapChain::EStatus::Healthy;
 		}
-		Queue->Submit(CmdBuffer, RenderingDoneSemaphores[AcquiredImageIndex]->GetHandle());
 	}
 	else
 	{
diff --git a/Engine/Source/Runtime/VulkanRHI/Private/VulkanViewport.h b/Engine/Source/Runtime/VulkanRHI/Private/VulkanViewport.h
--- a/Engine/Source/Runtime/VulkanRHI/Private/VulkanViewport.h
+++ b/Engine/Source/Runtime/VulkanRHI/Private/VulkanViewport.h
@@ -115,6 +115,7 @@ protected:
 	void CreateSwapchain();
 	void DestroySwapchain();
 	void AcquireImageIndex();
+	bool TryAcquireImageIndex();
 
 	void RecreateSwapchain(void* NewNativeWindow);
 	void RecreateSwapchainFromRT(EPixelFormat PreferredPixelFormat);
