Vulkan: fix build breakage (extra parameter to RecreateSwapchain).

diff --git a/Engine/Source/Runtime/VulkanRHI/Private/VulkanViewport.cpp b/Engine/Source/Runtime/VulkanRHI/Private/VulkanViewport.cpp
--- a/Engine/Source/Runtime/VulkanRHI/Private/VulkanViewport.cpp
+++ b/Engine/Source/Runtime/VulkanRHI/Private/VulkanViewport.cpp
@@ -802,7 +802,7 @@ bool FVulkanViewport::Present(FVulkanCommandListContext* Context, FVulkanCmdBuff
 
 			// cannot just throw out this command buffer (needs to be submitted or other checks fail)
 			Queue->Submit(CmdBuffer);
-			RecreateSwapchain(WindowHandle, true);
+			RecreateSwapchain(WindowHandle);
 
 			// Swapchain creation pushes some commands - flush the command buffers now to begin with a fresh state
 			Device->SubmitCommandsAndFlushGPU();
