Fixed: Shimmering main menu when packaging Match3 for Android Vulkan

diff --git a/Engine/Config/BaseDeviceProfiles.ini b/Engine/Config/BaseDeviceProfiles.ini
--- a/Engine/Config/BaseDeviceProfiles.ini
+++ b/Engine/Config/BaseDeviceProfiles.ini
@@ -442,7 +442,7 @@ BaseProfileName=
 +CVars=r.RefractionQuality=0
 +CVars=r.ShadowQuality=2
 +CVars=slate.AbsoluteIndices=1
-+CVars=r.Vulkan.DelayAcquireBackBuffer=0
++CVars=r.Vulkan.DelayAcquireBackBuffer=2
 +CVars=r.Vulkan.RobustBufferAccess=1
 +CVars=r.Vulkan.DescriptorSetLayoutMode=2
 
diff --git a/Engine/Source/Runtime/VulkanRHI/Private/VulkanViewport.cpp b/Engine/Source/Runtime/VulkanRHI/Private/VulkanViewport.cpp
--- a/Engine/Source/Runtime/VulkanRHI/Private/VulkanViewport.cpp
+++ b/Engine/Source/Runtime/VulkanRHI/Private/VulkanViewport.cpp
@@ -107,8 +107,8 @@ void FVulkanBackBuffer::AcquireBackBufferImage(FVulkanCommandListContext& Contex
 			LayoutMgr.EndEmulatedRenderPass(CmdBuffer);
 		}
 			
-		// Add a dependency with the acquired semaphore
-		CmdBuffer->AddWaitSemaphore(VK_PIPELINE_STAGE_BOTTOM_OF_PIPE_BIT, Viewport->AcquiredSemaphore);
+		// Wait for semaphore signal before writing to backbuffer image
+		CmdBuffer->AddWaitSemaphore(VK_PIPELINE_STAGE_COLOR_ATTACHMENT_OUTPUT_BIT, Viewport->AcquiredSemaphore);
 	}
 }
 
