Vulkan - More info on invalid transition

diff --git a/Engine/Source/Runtime/VulkanRHI/Private/VulkanRenderTarget.cpp b/Engine/Source/Runtime/VulkanRHI/Private/VulkanRenderTarget.cpp
--- a/Engine/Source/Runtime/VulkanRHI/Private/VulkanRenderTarget.cpp
+++ b/Engine/Source/Runtime/VulkanRHI/Private/VulkanRenderTarget.cpp
@@ -1554,7 +1554,7 @@ void FVulkanCommandListContext::TransitionResources(const FPendingTransition& Pe
 					DestStage = VK_PIPELINE_STAGE_COMPUTE_SHADER_BIT | VK_PIPELINE_STAGE_ALL_GRAPHICS_BIT;
 					break;
 				default:
-					ensure(0);
+					ensureMsgf(0, TEXT("Unknown transition pipeline %d"), (int32)PendingTransition.TransitionPipeline);
 					break;
 				}
 
