Can't lock RLM_ReadOnly resource in vulkan.

diff --git a/Engine/Source/Runtime/Engine/Private/Particles/ParticleGpuSimulation.cpp b/Engine/Source/Runtime/Engine/Private/Particles/ParticleGpuSimulation.cpp
--- a/Engine/Source/Runtime/Engine/Private/Particles/ParticleGpuSimulation.cpp
+++ b/Engine/Source/Runtime/Engine/Private/Particles/ParticleGpuSimulation.cpp
@@ -2391,7 +2391,7 @@ static FBox ComputeParticleBounds(
 		FRHIResourceCreateInfo CreateInfo;
 		FVertexBufferRHIRef BoundsVertexBufferRHI = RHICreateVertexBuffer(
 			BufferSize,
-			BUF_Static | BUF_UnorderedAccess,
+			BUF_Static | BUF_UnorderedAccess | BUF_KeepCPUAccessible,
 			CreateInfo);
 		FUnorderedAccessViewRHIRef BoundsVertexBufferUAV = RHICreateUnorderedAccessView(
 			BoundsVertexBufferRHI,
diff --git a/Engine/Source/Runtime/VulkanRHI/Private/VulkanIndexBuffer.cpp b/Engine/Source/Runtime/VulkanRHI/Private/VulkanIndexBuffer.cpp
--- a/Engine/Source/Runtime/VulkanRHI/Private/VulkanIndexBuffer.cpp
+++ b/Engine/Source/Runtime/VulkanRHI/Private/VulkanIndexBuffer.cpp
@@ -166,7 +166,7 @@ void* FVulkanResourceMultiBuffer::Lock(bool bFromRenderingThread, EResourceLockM
 		check(NumBuffers == 0);
 		if (LockMode == RLM_ReadOnly)
 		{
-			ensure(0);
+			checkf(0, TEXT("Volatile buffers can't be locked for read."));
 		}
 		else
 		{
@@ -184,7 +184,62 @@ void* FVulkanResourceMultiBuffer::Lock(bool bFromRenderingThread, EResourceLockM
 
 		if (LockMode == RLM_ReadOnly)
 		{
-			ensure(0);
+			const bool bUnifiedMem = Device->HasUnifiedMemory();
+			if (bUnifiedMem)
+			{
+				Data = (uint8*)Buffers[DynamicBufferIndex]->GetMappedPointer() + Offset;
+			}
+			else 
+			{
+				Device->PrepareForCPURead();
+				FVulkanCommandListContext& ImmediateContext = Device->GetImmediateContext();
+				FVulkanCmdBuffer* CmdBuffer = ImmediateContext.GetCommandBufferManager()->GetUploadCmdBuffer();
+				
+				// Make sure any previous tasks have finished on the source buffer.
+				VkBufferMemoryBarrier Barrier;
+				VulkanRHI::SetupAndZeroBufferBarrier(Barrier, VK_ACCESS_MEMORY_WRITE_BIT, VK_ACCESS_MEMORY_READ_BIT, Buffers[DynamicBufferIndex]->GetHandle(), Offset, Size);
+				VulkanRHI::vkCmdPipelineBarrier(CmdBuffer->GetHandle(), VK_PIPELINE_STAGE_TRANSFER_BIT, VK_PIPELINE_STAGE_TRANSFER_BIT, 0, 0, 0, 1, &Barrier, 0, nullptr);
+
+				// Create a staging buffer we can use to copy data from device to cpu.
+				VulkanRHI::FStagingBuffer* StagingBuffer = Device->GetStagingManager().AcquireBuffer(Size, VK_BUFFER_USAGE_TRANSFER_DST_BIT, true);
+
+				// Fill the staging buffer with the data on the device.
+				VkBufferCopy Regions;
+				Regions.size = Size;
+				Regions.srcOffset = Offset;
+				Regions.dstOffset = 0;
+
+				VulkanRHI::vkCmdCopyBuffer(CmdBuffer->GetHandle(), Buffers[DynamicBufferIndex]->GetHandle(), StagingBuffer->GetHandle(), 1, &Regions);
+
+				// Setup barrier.
+				VulkanRHI::SetupAndZeroBufferBarrier(Barrier, VK_ACCESS_MEMORY_WRITE_BIT, VK_ACCESS_HOST_READ_BIT, StagingBuffer->GetHandle(), 0, Size);
+				VulkanRHI::vkCmdPipelineBarrier(CmdBuffer->GetHandle(), VK_PIPELINE_STAGE_TRANSFER_BIT, VK_PIPELINE_STAGE_HOST_BIT, 0, 0, 0, 1, &Barrier, 0, nullptr);
+				
+				// Force upload.
+				ImmediateContext.GetCommandBufferManager()->SubmitUploadCmdBuffer();
+				Device->WaitUntilIdle();
+
+				// Flush.
+				StagingBuffer->FlushMappedMemory();
+
+				// Get mapped pointer. 
+				Data = StagingBuffer->GetMappedPointer();
+
+				// Release temp staging buffer during unlock.
+				VulkanRHI::FPendingBufferLock PendingLock;
+				PendingLock.Offset = 0;
+				PendingLock.Size = Size;
+				PendingLock.LockMode = LockMode;
+				PendingLock.StagingBuffer = StagingBuffer;
+
+				{
+					FScopeLock ScopeLock(&GPendingLockIBsMutex);
+					check(!GPendingLockIBs.Contains(this));
+					GPendingLockIBs.Add(this, PendingLock);
+				}
+
+				ImmediateContext.GetCommandBufferManager()->PrepareForNewActiveCommandBuffer();
+			}
 		}
 		else
 		{
@@ -325,10 +380,10 @@ void FVulkanResourceMultiBuffer::Unlock(bool bFromRenderingThread)
 				ALLOC_COMMAND_CL(RHICmdList, FRHICommandMultiBufferUnlock)(Device, PendingLock, this, DynamicBufferIndex);
 			}
 		}
-		else
+		else if(PendingLock.LockMode == RLM_ReadOnly)
 		{
-			// Not implemented
-			ensure(0);
+			// Just remove the staging buffer here.
+			Device->GetStagingManager().ReleaseBuffer(0, PendingLock.StagingBuffer);
 		}
 	}
 }
