If we use a 2nd subpass for transluency, don't resolve in the first subpass

diff --git a/Engine/Source/Runtime/VulkanRHI/Private/VulkanRHI.cpp b/Engine/Source/Runtime/VulkanRHI/Private/VulkanRHI.cpp
--- a/Engine/Source/Runtime/VulkanRHI/Private/VulkanRHI.cpp
+++ b/Engine/Source/Runtime/VulkanRHI/Private/VulkanRHI.cpp
@@ -1451,6 +1451,8 @@ static VkRenderPass CreateRenderPass(FVulkanDevice& InDevice, const FVulkanRende
 
 	VkSubpassDescription SubpassDescriptions[2];
 	VkSubpassDependency SubpassDependencies[2];
+
+	const bool bHasDepthReadSubpass = RTLayout.GetSubpassHint() == ESubpassHint::DepthReadSubpass;
 		
 	// main sub-pass
 	{
@@ -1460,13 +1462,13 @@ static VkRenderPass CreateRenderPass(FVulkanDevice& InDevice, const FVulkanRende
 		SubpassDesc.pipelineBindPoint = VK_PIPELINE_BIND_POINT_GRAPHICS;
 		SubpassDesc.colorAttachmentCount = RTLayout.GetNumColorAttachments();
 		SubpassDesc.pColorAttachments = RTLayout.GetColorAttachmentReferences();
-		SubpassDesc.pResolveAttachments = RTLayout.GetResolveAttachmentReferences();
+		SubpassDesc.pResolveAttachments = bHasDepthReadSubpass ? nullptr : RTLayout.GetResolveAttachmentReferences();
 		SubpassDesc.pDepthStencilAttachment = RTLayout.GetDepthStencilAttachmentReference();
 	}
 
 	// depth read sub-pass
 	VkAttachmentReference InputAttachments[MaxSimultaneousRenderTargets + 1];
-	if (RTLayout.GetSubpassHint() == ESubpassHint::DepthReadSubpass)
+	if (bHasDepthReadSubpass)
 	{
 		VkSubpassDescription& SubpassDesc = SubpassDescriptions[NumSubpasses++];
 		FMemory::Memzero(&SubpassDesc, sizeof(VkSubpassDescription));
