Vulkan: reducing the number of times PendingFrees are iterated.

diff --git a/Engine/Source/Runtime/VulkanRHI/Private/VulkanUniformBuffer.cpp b/Engine/Source/Runtime/VulkanRHI/Private/VulkanUniformBuffer.cpp
--- a/Engine/Source/Runtime/VulkanRHI/Private/VulkanUniformBuffer.cpp
+++ b/Engine/Source/Runtime/VulkanRHI/Private/VulkanUniformBuffer.cpp
@@ -330,14 +330,56 @@ namespace VulkanRHI
 
 	void FResourceHeapManager::ProcessPendingUBFreesNoLock(bool bForce)
 	{
-		for (int32 Index = UBAllocations.PendingFree.Num() - 1; Index >= 0; --Index)
+		// this keeps an frame number of the first frame when we can expect to delete things, updated in the loop if any pending allocations are left
+		static uint32 GFrameNumberRenderThread_WhenWeCanDelete = 0;
+
+		if (UNLIKELY(bForce))
 		{
-			FUBPendingFree& Alloc = UBAllocations.PendingFree[Index];
-			if (GFrameNumberRenderThread > Alloc.Frame + VulkanRHI::NUM_FRAMES_TO_WAIT_BEFORE_RELEASING_TO_OS || bForce)
+			int32 NumAlloc = UBAllocations.PendingFree.Num();
+			for (int32 Index = 0; Index < NumAlloc; ++Index)
 			{
+				FUBPendingFree& Alloc = UBAllocations.PendingFree[Index];
 				delete Alloc.Allocation;
-				UBAllocations.PendingFree.RemoveAtSwap(Index, 1, false);
 			}
+			UBAllocations.PendingFree.Empty();
+
+			// invalidate the value
+			GFrameNumberRenderThread_WhenWeCanDelete = 0;
+		}
+		else
+		{
+			if (LIKELY(GFrameNumberRenderThread < GFrameNumberRenderThread_WhenWeCanDelete))
+			{
+				// too early
+				return;
+			}
+
+			// making use of the fact that we always add to the end of the array, so allocations are sorted by frame ascending
+			int32 OldestFrameToKeep = GFrameNumberRenderThread - VulkanRHI::NUM_FRAMES_TO_WAIT_BEFORE_RELEASING_TO_OS;
+			int32 NumAlloc = UBAllocations.PendingFree.Num();
+			int32 Index = 0;
+			for (; Index < NumAlloc; ++Index)
+			{
+				FUBPendingFree& Alloc = UBAllocations.PendingFree[Index];
+				if (LIKELY(Alloc.Frame < OldestFrameToKeep))
+				{
+					delete Alloc.Allocation;
+				}
+				else
+				{
+					// calculate when we will be able to delete the oldest allocation
+					GFrameNumberRenderThread_WhenWeCanDelete = Alloc.Frame + VulkanRHI::NUM_FRAMES_TO_WAIT_BEFORE_RELEASING_TO_OS + 1;
+					break;
+				}
+			}
+
+			int32 ElementsLeft = NumAlloc - Index;
+			if (ElementsLeft > 0 && ElementsLeft != NumAlloc)
+			{
+				// FUBPendingFree is POD because it is stored in a TArray
+				FMemory::Memmove(UBAllocations.PendingFree.GetData(), UBAllocations.PendingFree.GetData() + Index, ElementsLeft * sizeof(FUBPendingFree));
+			}
+			UBAllocations.PendingFree.SetNum(NumAlloc - Index, false);
 		}
 	}
 
