Vulkan: keep free ranges sorted and merged on the fly (UE-78697)

- Good performance improvements on samples like Infiltrator.

diff --git a/Engine/Source/Runtime/VulkanRHI/Private/VulkanMemory.cpp b/Engine/Source/Runtime/VulkanRHI/Private/VulkanMemory.cpp
--- a/Engine/Source/Runtime/VulkanRHI/Private/VulkanMemory.cpp
+++ b/Engine/Source/Runtime/VulkanRHI/Private/VulkanMemory.cpp
@@ -475,12 +475,24 @@ namespace VulkanRHI
 		}
 	}
 
+#define UE_VK_MEMORY_KEEP_FREELIST_SORTED					1
+#define UE_VK_MEMORY_JOIN_FREELIST_ON_THE_FLY				(UE_VK_MEMORY_KEEP_FREELIST_SORTED && 1)
+// debugging
+#define UE_VK_MEMORY_KEEP_FREELIST_SORTED_CATCHBUGS			0
+
 	void FRange::JoinConsecutiveRanges(TArray<FRange>& Ranges)
 	{
 		if (Ranges.Num() > 1)
 		{
+#if !UE_VK_MEMORY_KEEP_FREELIST_SORTED
 			Ranges.Sort();
+#else
+	#if UE_VK_MEMORY_KEEP_FREELIST_SORTED_CATCHBUGS
+			SanityCheck(Ranges);
+	#endif
+#endif
 
+#if !UE_VK_MEMORY_JOIN_FREELIST_ON_THE_FLY
 			for (int32 Index = Ranges.Num() - 1; Index > 0; --Index)
 			{
 				FRange& Current = Ranges[Index];
@@ -491,10 +503,172 @@ namespace VulkanRHI
 					Ranges.RemoveAt(Index, 1, false);
 				}
 			}
+#endif
+		}
+	}
+
+	int32 FRange::InsertAndTryToMerge(TArray<FRange>& Ranges, const FRange& Item, int32 ProposedIndex)
+	{
+#if !UE_VK_MEMORY_JOIN_FREELIST_ON_THE_FLY
+		int32 Ret = Ranges.Insert(Item, ProposedIndex);
+#else
+		// there are four cases here
+		// 1) nothing can be merged (distinct ranges)		 XXXX YYY ZZZZZ  =>   XXXX YYY ZZZZZ
+		// 2) new range can be merged with the previous one: XXXXYYY  ZZZZZ  =>   XXXXXXX  ZZZZZ
+		// 3) new range can be merged with the next one:     XXXX  YYYZZZZZ  =>   XXXX  ZZZZZZZZ
+		// 4) new range perfectly fills the gap:             XXXXYYYYYZZZZZ  =>   XXXXXXXXXXXXXX
+
+		// note: we can have a case where we're inserting at the beginning of the array (no previous element), but we won't have a case
+		// where we're inserting at the end (no next element) - AppendAndTryToMerge() should be called instead
+		checkf(Item.Offset < Ranges[ProposedIndex].Offset, TEXT("FRange::InsertAndTryToMerge() was called to append an element - internal logic error, FRange::AppendAndTryToMerge() should have been called instead."))
+		int32 Ret = ProposedIndex;
+		if (UNLIKELY(ProposedIndex == 0))
+		{
+			// only cases 1 and 3 apply
+			FRange& NextRange = Ranges[Ret];
+
+			if (UNLIKELY(NextRange.Offset == Item.Offset + Item.Size))
+			{
+				NextRange.Offset = Item.Offset;
+				NextRange.Size += Item.Size;
+			}
+			else
+			{
+				Ret = Ranges.Insert(Item, ProposedIndex);
+			}
+		}
+		else
+		{
+			// all cases apply
+			FRange& NextRange = Ranges[ProposedIndex];
+			FRange& PrevRange = Ranges[ProposedIndex - 1];
+
+			// see if we can merge with previous
+			if (UNLIKELY(PrevRange.Offset + PrevRange.Size == Item.Offset))
+			{
+				// case 2, can still end up being case 4
+				PrevRange.Size += Item.Size;
+
+				if (UNLIKELY(PrevRange.Offset + PrevRange.Size == NextRange.Offset))
+				{
+					// case 4
+					PrevRange.Size += NextRange.Size;
+					Ranges.RemoveAt(ProposedIndex);
+					Ret = ProposedIndex - 1;
+				}
+			}
+			else if (UNLIKELY(Item.Offset + Item.Size == NextRange.Offset))
+			{
+				// case 3
+				NextRange.Offset = Item.Offset;
+				NextRange.Size += Item.Size;
+			}
+			else
+			{
+				// case 1 - the new range is disjoint with both
+				Ret = Ranges.Insert(Item, ProposedIndex);	// this can invalidate NextRange/PrevRange references, don't touch them after this
+			}
+		}
+#endif
+
+#if UE_VK_MEMORY_KEEP_FREELIST_SORTED_CATCHBUGS
+		SanityCheck(Ranges);
+#endif
+		return Ret;
+	}
+
+	int32 FRange::AppendAndTryToMerge(TArray<FRange>& Ranges, const FRange& Item)
+	{
+#if !UE_VK_MEMORY_JOIN_FREELIST_ON_THE_FLY
+		int32 Ret = Ranges.Add(Item);
+#else
+		int32 Ret = Ranges.Num() - 1;
+		// we only get here when we have an element in front of us
+		checkf(Ret >= 0, TEXT("FRange::AppendAndTryToMerge() was called on an empty array."));
+		FRange& PrevRange = Ranges[Ret];
+		if (UNLIKELY(PrevRange.Offset + PrevRange.Size == Item.Offset))
+		{
+			PrevRange.Size += Item.Size;
+		}
+		else
+		{
+			Ret = Ranges.Add(Item);
+		}
+#endif
+
+#if UE_VK_MEMORY_KEEP_FREELIST_SORTED_CATCHBUGS
+		SanityCheck(Ranges);
+#endif
+		return Ret;
+	}
+
+	void FRange::AllocateFromEntry(TArray<FRange>& Ranges, int32 Index, uint32 SizeToAllocate)
+	{
+		FRange& Entry = Ranges[Index];
+		if (SizeToAllocate < Entry.Size)
+		{
+			// Modify current free entry in-place.
+			Entry.Size -= SizeToAllocate;
+			Entry.Offset += SizeToAllocate;
+		}
+		else
+		{
+			// Remove this free entry.
+			Ranges.RemoveAt(Index, 1, false);
+#if UE_VK_MEMORY_KEEP_FREELIST_SORTED_CATCHBUGS
+			SanityCheck(Ranges);
+#endif
+		}
+	}
+
+	void FRange::SanityCheck(TArray<FRange>& Ranges)
+	{
+		if (UE_VK_MEMORY_KEEP_FREELIST_SORTED_CATCHBUGS)	// keeping the check code visible to the compiler
+		{
+			int32 Num = Ranges.Num();
+			if (Num > 1)
+			{
+				for (int32 ChkIndex = 0; ChkIndex < Num - 1; ++ChkIndex)
+				{
+					checkf(Ranges[ChkIndex].Offset < Ranges[ChkIndex + 1].Offset, TEXT("Array is not sorted!"));
+					// if we're joining on the fly, then there cannot be any adjoining ranges, so use < instead of <=
+#if UE_VK_MEMORY_JOIN_FREELIST_ON_THE_FLY
+					checkf(Ranges[ChkIndex].Offset + Ranges[ChkIndex].Size < Ranges[ChkIndex + 1].Offset, TEXT("Ranges are overlapping or adjoining!"));
+#else
+					checkf(Ranges[ChkIndex].Offset + Ranges[ChkIndex].Size <= Ranges[ChkIndex + 1].Offset, TEXT("Ranges are overlapping!"));
+#endif
+				}
+			}
 		}
 	}
 
 
+	int32 FRange::Add(TArray<FRange>& Ranges, const FRange & Item)
+	{
+#if UE_VK_MEMORY_KEEP_FREELIST_SORTED
+		// find the right place to add
+		int32 NumRanges = Ranges.Num();
+		if (LIKELY(NumRanges <= 0))
+		{
+			return Ranges.Add(Item);
+		}
+
+		FRange* Data = Ranges.GetData();
+		for (int32 Index = 0; Index < NumRanges; ++Index)
+		{
+			if (UNLIKELY(Data[Index].Offset > Item.Offset))
+			{
+				return InsertAndTryToMerge(Ranges, Item, Index);
+			}
+		}
+
+		// if we got this far and still haven't inserted, we're a new element
+		return AppendAndTryToMerge(Ranges, Item);
+#else
+		return Ranges.Add(Item);
+#endif
+	}
+
 	FOldResourceAllocation::FOldResourceAllocation(FOldResourceHeapPage* InOwner, FDeviceMemoryAllocation* InDeviceMemoryAllocation,
 		uint32 InRequestedSize, uint32 InAlignedOffset,
 		uint32 InAllocationSize, uint32 InAllocationOffset, const char* InFile, uint32 InLine)
@@ -559,7 +733,7 @@ namespace VulkanRHI
 		FRange FullRange;
 		FullRange.Offset = 0;
 		FullRange.Size = MaxSize;
-		FreeList.Add(FullRange);
+		FRange::Add(FreeList, FullRange);
 	}
 
 	FOldResourceHeapPage::~FOldResourceHeapPage()
@@ -580,17 +754,7 @@ namespace VulkanRHI
 			uint32 AllocatedSize = AlignmentAdjustment + Size;
 			if (AllocatedSize <= Entry.Size)
 			{
-				if (AllocatedSize < Entry.Size)
-				{
-					// Modify current free entry in-place
-					Entry.Size -= AllocatedSize;
-					Entry.Offset += AllocatedSize;
-				}
-				else
-				{
-					// Remove this free entry
-					FreeList.RemoveAtSwap(Index, 1, false);
-				}
+				FRange::AllocateFromEntry(FreeList, Index, AllocatedSize);
 
 				UsedSize += AllocatedSize;
 
@@ -620,7 +784,7 @@ namespace VulkanRHI
 			NewFree.Offset = Allocation->AllocationOffset;
 			NewFree.Size = Allocation->AllocationSize;
 
-			FreeList.Add(NewFree);
+			FRange::Add(FreeList, NewFree);
 		}
 
 		UsedSize -= Allocation->AllocationSize;
@@ -635,7 +799,9 @@ namespace VulkanRHI
 	bool FOldResourceHeapPage::JoinFreeBlocks()
 	{
 		FScopeLock ScopeLock(&GOldResourcePageLock);
+#if !UE_VK_MEMORY_JOIN_FREELIST_ON_THE_FLY
 		FRange::JoinConsecutiveRanges(FreeList);
+#endif
 
 		if (FreeList.Num() == 1)
 		{
@@ -1510,7 +1676,9 @@ namespace VulkanRHI
 	bool FSubresourceAllocator::JoinFreeBlocks()
 	{
 		FScopeLock ScopeLock(&CS);
+#if !UE_VK_MEMORY_JOIN_FREELIST_ON_THE_FLY
 		FRange::JoinConsecutiveRanges(FreeList);
+#endif
 
 		if (FreeList.Num() == 1)
 		{
@@ -1538,17 +1706,7 @@ namespace VulkanRHI
 			uint32 AllocatedSize = AlignmentAdjustment + InSize;
 			if (AllocatedSize <= Entry.Size)
 			{
-				if (AllocatedSize < Entry.Size)
-				{
-					// Modify current free entry in-place
-					Entry.Size -= AllocatedSize;
-					Entry.Offset += AllocatedSize;
-				}
-				else
-				{
-					// Remove this free entry
-					FreeList.RemoveAtSwap(Index, 1, false);
-				}
+				FRange::AllocateFromEntry(FreeList, Index, AllocatedSize);
 
 				UsedSize += AllocatedSize;
 
@@ -1592,7 +1750,7 @@ namespace VulkanRHI
 			check(NewFree.Offset <= GetMaxSize());
 			check(NewFree.Offset + NewFree.Size <= GetMaxSize());
 
-			FreeList.Add(NewFree);
+			FRange::Add(FreeList, NewFree);
 
 			UsedSize -= Suballocation->AllocationSize;
 			check(UsedSize >= 0);
diff --git a/Engine/Source/Runtime/VulkanRHI/Public/VulkanMemory.h b/Engine/Source/Runtime/VulkanRHI/Public/VulkanMemory.h
--- a/Engine/Source/Runtime/VulkanRHI/Public/VulkanMemory.h
+++ b/Engine/Source/Runtime/VulkanRHI/Public/VulkanMemory.h
@@ -426,6 +426,21 @@ namespace VulkanRHI
 		}
 
 		static void JoinConsecutiveRanges(TArray<FRange>& Ranges);
+
+		/** Tries to insert the item so it has index ProposedIndex, but may end up merging it with neighbors */
+		static int32 InsertAndTryToMerge(TArray<FRange>& Ranges, const FRange& Item, int32 ProposedIndex);
+
+		/** Tries to append the item to the end but may end up merging it with the neighbor */
+		static int32 AppendAndTryToMerge(TArray<FRange>& Ranges, const FRange& Item);
+
+		/** Attempts to allocate from an entry - can remove it if it was used up*/
+		static void AllocateFromEntry(TArray<FRange>& Ranges, int32 Index, uint32 SizeToAllocate);
+
+		/** Sanity checks an array of ranges */
+		static void SanityCheck(TArray<FRange>& Ranges);
+
+		/** Adds to the array while maintaing the sort. */
+		static int32 Add(TArray<FRange>& Ranges, const FRange& Item);
 	};
 
 	// One device allocation that is shared amongst different resources
