-update uniform buffers inplace instead of renaming, when outside of renderpasses.
-Added stats for uniformbuffer updates, and updates with rename.
-Added a counter for Frame Temp memory

diff --git a/Engine/Source/Runtime/VulkanRHI/Private/VulkanCommandBuffer.h b/Engine/Source/Runtime/VulkanRHI/Private/VulkanCommandBuffer.h
--- a/Engine/Source/Runtime/VulkanRHI/Private/VulkanCommandBuffer.h
+++ b/Engine/Source/Runtime/VulkanRHI/Private/VulkanCommandBuffer.h
@@ -290,6 +290,12 @@ public:
 		return ActiveCmdBuffer;
 	}
 
+	inline FVulkanCmdBuffer* GetActiveCmdBufferDirect()
+	{
+		return ActiveCmdBuffer;
+	}
+
+
 	inline bool HasPendingUploadCmdBuffer() const
 	{
 		return UploadCmdBuffer != nullptr;
diff --git a/Engine/Source/Runtime/VulkanRHI/Private/VulkanMemory.cpp b/Engine/Source/Runtime/VulkanRHI/Private/VulkanMemory.cpp
--- a/Engine/Source/Runtime/VulkanRHI/Private/VulkanMemory.cpp
+++ b/Engine/Source/Runtime/VulkanRHI/Private/VulkanMemory.cpp
@@ -2276,6 +2276,7 @@ namespace VulkanRHI
 	{
 		for (int32 Index = 0; Index < NUM_BUFFERS; ++Index)
 		{
+			INC_MEMORY_STAT_BY(STAT_VulkanTempFrameAllocationBuffer, ALLOCATION_SIZE);
 			Entries[Index].InitBuffer(Device, ALLOCATION_SIZE);
 		}
 	}
@@ -2291,7 +2292,7 @@ namespace VulkanRHI
 		Size = InSize;
 		PeakUsed = 0;
 		BufferSuballocation = InDevice->GetResourceHeapManager().AllocateBuffer(InSize,
-			VK_BUFFER_USAGE_VERTEX_BUFFER_BIT | VK_BUFFER_USAGE_INDEX_BUFFER_BIT | VK_BUFFER_USAGE_UNIFORM_BUFFER_BIT | VK_BUFFER_USAGE_UNIFORM_TEXEL_BUFFER_BIT,
+			VK_BUFFER_USAGE_TRANSFER_SRC_BIT | VK_BUFFER_USAGE_VERTEX_BUFFER_BIT | VK_BUFFER_USAGE_INDEX_BUFFER_BIT | VK_BUFFER_USAGE_UNIFORM_BUFFER_BIT | VK_BUFFER_USAGE_UNIFORM_TEXEL_BUFFER_BIT,
 			VK_MEMORY_PROPERTY_HOST_VISIBLE_BIT | VK_MEMORY_PROPERTY_HOST_COHERENT_BIT,
 			__FILE__, __LINE__);
 		MappedData = (uint8*)BufferSuballocation->GetMappedPointer();
@@ -2333,6 +2334,8 @@ namespace VulkanRHI
 		
 		// Couldn't fit in the current buffers; allocate a new bigger one and schedule the current one for deletion
 		uint32 NewSize = Align(ALLOCATION_SIZE + InSize + InAlignment, ALLOCATION_SIZE);
+		DEC_MEMORY_STAT_BY(STAT_VulkanTempFrameAllocationBuffer, Entries[BufferIndex].BufferSuballocation->GetSize());
+		INC_MEMORY_STAT_BY(STAT_VulkanTempFrameAllocationBuffer, NewSize);
 		Entries[BufferIndex].PendingDeletionList.Add(Entries[BufferIndex].BufferSuballocation);
 		Entries[BufferIndex].InitBuffer(Device, NewSize);
 		if (!Entries[BufferIndex].TryAlloc(InSize, InAlignment, OutInfo))
diff --git a/Engine/Source/Runtime/VulkanRHI/Private/VulkanRHIPrivate.h b/Engine/Source/Runtime/VulkanRHI/Private/VulkanRHIPrivate.h
--- a/Engine/Source/Runtime/VulkanRHI/Private/VulkanRHIPrivate.h
+++ b/Engine/Source/Runtime/VulkanRHI/Private/VulkanRHIPrivate.h
@@ -392,6 +392,7 @@ DECLARE_DWORD_ACCUMULATOR_STAT_EXTERN(TEXT("Num Frame Buffers"), STAT_VulkanNumF
 DECLARE_DWORD_ACCUMULATOR_STAT_EXTERN(TEXT("Num Buffer Views"), STAT_VulkanNumBufferViews, STATGROUP_VulkanRHI, );
 DECLARE_DWORD_ACCUMULATOR_STAT_EXTERN(TEXT("Num Image Views"), STAT_VulkanNumImageViews, STATGROUP_VulkanRHI, );
 DECLARE_DWORD_ACCUMULATOR_STAT_EXTERN(TEXT("Num Physical Mem Allocations"), STAT_VulkanNumPhysicalMemAllocations, STATGROUP_VulkanRHI, );
+DECLARE_MEMORY_STAT_EXTERN(TEXT("Frame Temp Memory"), STAT_VulkanTempFrameAllocationBuffer, STATGROUP_VulkanRHI, );
 DECLARE_DWORD_COUNTER_STAT_EXTERN(TEXT("Dynamic VB Size"), STAT_VulkanDynamicVBSize, STATGROUP_VulkanRHI, );
 DECLARE_DWORD_COUNTER_STAT_EXTERN(TEXT("Dynamic IB Size"), STAT_VulkanDynamicIBSize, STATGROUP_VulkanRHI, );
 DECLARE_CYCLE_STAT_EXTERN(TEXT("Dynamic VB Lock/Unlock time"), STAT_VulkanDynamicVBLockTime, STATGROUP_VulkanRHI, );
@@ -415,6 +416,8 @@ DECLARE_CYCLE_STAT_EXTERN(TEXT("Acquire Backbuffer"), STAT_VulkanAcquireBackBuff
 DECLARE_CYCLE_STAT_EXTERN(TEXT("Staging Buffer Mgmt"), STAT_VulkanStagingBuffer, STATGROUP_VulkanRHI, );
 DECLARE_CYCLE_STAT_EXTERN(TEXT("VkCreateDescriptorPool"), STAT_VulkanVkCreateDescriptorPool, STATGROUP_VulkanRHI, );
 DECLARE_DWORD_ACCUMULATOR_STAT_EXTERN(TEXT("Num Created DescSet Pools"), STAT_VulkanNumDescPools, STATGROUP_VulkanRHI, );
+DECLARE_CYCLE_STAT_EXTERN(TEXT("UpdateUniformBuffers"), STAT_VulkanUpdateUniformBuffers, STATGROUP_VulkanRHI, );
+DECLARE_CYCLE_STAT_EXTERN(TEXT("UpdateUniformBuffersRename"), STAT_VulkanUpdateUniformBuffersRename, STATGROUP_VulkanRHI, );
 #if VULKAN_ENABLE_AGGRESSIVE_STATS
 DECLARE_CYCLE_STAT_EXTERN(TEXT("Update DescriptorSets"), STAT_VulkanUpdateDescriptorSets, STATGROUP_VulkanRHI, );
 DECLARE_DWORD_COUNTER_STAT_EXTERN(TEXT("Num Desc Sets Updated"), STAT_VulkanNumDescSets, STATGROUP_VulkanRHI, );
diff --git a/Engine/Source/Runtime/VulkanRHI/Private/VulkanUniformBuffer.cpp b/Engine/Source/Runtime/VulkanRHI/Private/VulkanUniformBuffer.cpp
--- a/Engine/Source/Runtime/VulkanRHI/Private/VulkanUniformBuffer.cpp
+++ b/Engine/Source/Runtime/VulkanRHI/Private/VulkanUniformBuffer.cpp
@@ -8,6 +8,16 @@
 #include "VulkanContext.h"
 #include "VulkanLLM.h"
 
+static int32 GVulkanAllowUniformUpload = 1;
+static FAutoConsoleVariableRef CVarVulkanSubmitOnCopyToResolve(
+	TEXT("r.Vulkan.AllowUniformUpload"),
+	GVulkanAllowUniformUpload,
+	TEXT("Allow Uniform Buffer uploads outside of renderpasses\n")
+	TEXT(" 0: Disabled, buffers are always reallocated\n")
+	TEXT(" 1: Enabled, buffers are uploaded outside renderpasses"),
+	ECVF_Default
+);
+
 enum
 {
 #if PLATFORM_DESKTOP
@@ -184,6 +194,7 @@ FUniformBufferRHIRef FVulkanDynamicRHI::RHICreateUniformBuffer(const void* Conte
 template <bool bRealUBs>
 inline void FVulkanDynamicRHI::UpdateUniformBuffer(FVulkanUniformBuffer* UniformBuffer, const void* Contents)
 {
+	SCOPE_CYCLE_COUNTER(STAT_VulkanUpdateUniformBuffers);
 	const FRHIUniformBufferLayout& Layout = UniformBuffer->GetLayout();
 
 	const int32 ConstantBufferSize = Layout.ConstantBufferSize;
@@ -195,15 +206,34 @@ inline void FVulkanDynamicRHI::UpdateUniformBuffer(FVulkanUniformBuffer* Uniform
 	FVulkanEmulatedUniformBuffer* EmulatedUniformBuffer = bRealUBs ? nullptr : (FVulkanEmulatedUniformBuffer*)UniformBuffer;
 
 	FBufferSuballocation* NewUBAlloc = nullptr;
-	if (bRealUBs)
+	bool bIsInRenderPass = RHICmdList.IsInsideRenderPass();
+	bool bUseUpload = GVulkanAllowUniformUpload && !bIsInRenderPass; //inside renderpasses, a rename is enforced.
+
+	if (bRealUBs && !bUseUpload)
 	{
 		NewUBAlloc = nullptr;
 		if (ConstantBufferSize > 0)
 		{
+			SCOPE_CYCLE_COUNTER(STAT_VulkanUpdateUniformBuffersRename);
 			NewUBAlloc = Device->GetResourceHeapManager().AllocUniformBuffer(ConstantBufferSize, Contents);
 		}
 	}
 
+	auto UpdateUniformBufferHelper = [](FVulkanCommandListContext& Context, FVulkanRealUniformBuffer* UniformBuffer, int32 DataSize, const void* Data)
+	{
+		FVulkanCmdBuffer* CmdBuffer = Context.GetCommandBufferManager()->GetActiveCmdBufferDirect();
+		ensure(CmdBuffer->IsOutsideRenderPass());
+		VulkanRHI::FTempFrameAllocationBuffer::FTempAllocInfo LockInfo;
+		Context.GetTempFrameAllocationBuffer().Alloc(DataSize, 16, LockInfo);
+		FMemory::Memcpy(LockInfo.Data, Data, DataSize);
+		VkBufferCopy Region;
+		Region.size = DataSize;
+		Region.srcOffset = LockInfo.GetBindOffset();
+		Region.dstOffset = UniformBuffer->GetOffset();
+		VkBuffer UBBuffer = UniformBuffer->GetBufferAllocation()->GetHandle();
+		VulkanRHI::vkCmdCopyBuffer(CmdBuffer->GetHandle(), LockInfo.GetHandle(), UBBuffer, 1, &Region);
+	};
+
 	bool bRHIBypass = RHICmdList.Bypass();
 	if (bRHIBypass)
 	{
@@ -211,8 +241,16 @@ inline void FVulkanDynamicRHI::UpdateUniformBuffer(FVulkanUniformBuffer* Uniform
 		{
 			if (ConstantBufferSize > 0)
 			{
-				FBufferSuballocation* PrevAlloc = RealUniformBuffer->UpdateUBAllocation(NewUBAlloc);
-				Device->GetResourceHeapManager().ReleaseUniformBuffer(PrevAlloc);
+				if(bUseUpload)
+				{			
+					FVulkanCommandListContext& Context = (FVulkanCommandListContext&)*Device->ImmediateContext;
+					UpdateUniformBufferHelper(Context, RealUniformBuffer, ConstantBufferSize, Contents);
+				}
+				else
+				{
+					FBufferSuballocation* PrevAlloc = RealUniformBuffer->UpdateUBAllocation(NewUBAlloc);
+					Device->GetResourceHeapManager().ReleaseUniformBuffer(PrevAlloc);
+				}
 			}
 		}
 		else
@@ -238,12 +276,27 @@ inline void FVulkanDynamicRHI::UpdateUniformBuffer(FVulkanUniformBuffer* Uniform
 
 		if (bRealUBs)
 		{
-			RHICmdList.EnqueueLambda([RealUniformBuffer, NewUBAlloc, CmdListResources, NumResources, ConstantBufferSize](FRHICommandList&)
+			if(bUseUpload)
 			{
-				FBufferSuballocation* PrevAlloc = RealUniformBuffer->UpdateUBAllocation(NewUBAlloc);
-				RealUniformBuffer->Device->GetResourceHeapManager().ReleaseUniformBuffer(PrevAlloc);
-				RealUniformBuffer->UpdateResourceTable(CmdListResources, NumResources);
-			});
+				void* CmdListConstantBufferData = RHICmdList.Alloc(ConstantBufferSize, 16);
+				FMemory::Memcpy(CmdListConstantBufferData, Contents, ConstantBufferSize);
+
+				RHICmdList.EnqueueLambda([UpdateUniformBufferHelper, RealUniformBuffer, CmdListResources, NumResources, ConstantBufferSize, CmdListConstantBufferData](FRHICommandList& CmdList)
+				{
+					FVulkanCommandListContext& Context = (FVulkanCommandListContext&)CmdList.GetContext();
+					UpdateUniformBufferHelper(Context, RealUniformBuffer, ConstantBufferSize, CmdListConstantBufferData);
+					RealUniformBuffer->UpdateResourceTable(CmdListResources, NumResources);
+				});
+			}
+			else
+			{
+				RHICmdList.EnqueueLambda([RealUniformBuffer, NewUBAlloc, CmdListResources, NumResources](FRHICommandList& CmdList)
+				{
+					FBufferSuballocation* PrevAlloc = RealUniformBuffer->UpdateUBAllocation(NewUBAlloc);
+					RealUniformBuffer->Device->GetResourceHeapManager().ReleaseUniformBuffer(PrevAlloc);
+					RealUniformBuffer->UpdateResourceTable(CmdListResources, NumResources);
+				});
+			}
 		}
 		else
 		{
@@ -307,7 +360,7 @@ namespace VulkanRHI
 {
 	VulkanRHI::FBufferSuballocation* FResourceHeapManager::AllocUniformBuffer(uint32 Size, const void* Contents)
 	{
-		VulkanRHI::FBufferSuballocation* OutAlloc = Device->GetResourceHeapManager().AllocateBuffer(Size, VK_BUFFER_USAGE_UNIFORM_BUFFER_BIT, VK_MEMORY_PROPERTY_HOST_VISIBLE_BIT | VK_MEMORY_PROPERTY_HOST_COHERENT_BIT, __FILE__, __LINE__);
+		VulkanRHI::FBufferSuballocation* OutAlloc = Device->GetResourceHeapManager().AllocateBuffer(Size, VK_BUFFER_USAGE_UNIFORM_BUFFER_BIT | VK_BUFFER_USAGE_TRANSFER_DST_BIT, VK_MEMORY_PROPERTY_HOST_VISIBLE_BIT | VK_MEMORY_PROPERTY_HOST_COHERENT_BIT, __FILE__, __LINE__);
 		FMemory::Memcpy(OutAlloc->GetMappedPointer(), Contents, Size);
 		OutAlloc->Flush();
 
diff --git a/Engine/Source/Runtime/VulkanRHI/Private/VulkanUtil.cpp b/Engine/Source/Runtime/VulkanRHI/Private/VulkanUtil.cpp
--- a/Engine/Source/Runtime/VulkanRHI/Private/VulkanUtil.cpp
+++ b/Engine/Source/Runtime/VulkanRHI/Private/VulkanUtil.cpp
@@ -706,6 +706,7 @@ DEFINE_STAT(STAT_VulkanNumFrameBuffers);
 DEFINE_STAT(STAT_VulkanNumBufferViews);
 DEFINE_STAT(STAT_VulkanNumImageViews);
 DEFINE_STAT(STAT_VulkanNumPhysicalMemAllocations);
+DEFINE_STAT(STAT_VulkanTempFrameAllocationBuffer);
 DEFINE_STAT(STAT_VulkanDynamicVBSize);
 DEFINE_STAT(STAT_VulkanDynamicIBSize);
 DEFINE_STAT(STAT_VulkanDynamicVBLockTime);
@@ -729,6 +730,8 @@ DEFINE_STAT(STAT_VulkanAcquireBackBuffer);
 DEFINE_STAT(STAT_VulkanStagingBuffer);
 DEFINE_STAT(STAT_VulkanVkCreateDescriptorPool);
 DEFINE_STAT(STAT_VulkanNumDescPools);
+DEFINE_STAT(STAT_VulkanUpdateUniformBuffers);
+DEFINE_STAT(STAT_VulkanUpdateUniformBuffersRename);
 #if VULKAN_ENABLE_AGGRESSIVE_STATS
 DEFINE_STAT(STAT_VulkanUpdateDescriptorSets);
 DEFINE_STAT(STAT_VulkanNumUpdateDescriptors);
