Correction for cl 8112323. Accidentally turned on feature that should be off by default.

diff --git a/Engine/Source/Runtime/VulkanRHI/Private/VulkanUniformBuffer.cpp b/Engine/Source/Runtime/VulkanRHI/Private/VulkanUniformBuffer.cpp
--- a/Engine/Source/Runtime/VulkanRHI/Private/VulkanUniformBuffer.cpp
+++ b/Engine/Source/Runtime/VulkanRHI/Private/VulkanUniformBuffer.cpp
@@ -8,8 +8,8 @@
 #include "VulkanContext.h"
 #include "VulkanLLM.h"
 
-static int32 GVulkanAllowUniformUpload = 1;
-static FAutoConsoleVariableRef CVarVulkanSubmitOnCopyToResolve(
+static int32 GVulkanAllowUniformUpload = 0;
+static FAutoConsoleVariableRef CVarVulkanAllowUniformUpload(
 	TEXT("r.Vulkan.AllowUniformUpload"),
 	GVulkanAllowUniformUpload,
 	TEXT("Allow Uniform Buffer uploads outside of renderpasses\n")
