Attempted fix for Android compile warning

diff --git a/Engine/Source/Runtime/VulkanRHI/Private/VulkanUniformBuffer.cpp b/Engine/Source/Runtime/VulkanRHI/Private/VulkanUniformBuffer.cpp
--- a/Engine/Source/Runtime/VulkanRHI/Private/VulkanUniformBuffer.cpp
+++ b/Engine/Source/Runtime/VulkanRHI/Private/VulkanUniformBuffer.cpp
@@ -219,7 +219,7 @@ inline void FVulkanDynamicRHI::UpdateUniformBuffer(FVulkanUniformBuffer* Uniform
 		}
 	}
 
-	auto UpdateUniformBufferHelper = [](FVulkanCommandListContext& Context, FVulkanRealUniformBuffer* UniformBuffer, int32 DataSize, const void* Data)
+	auto UpdateUniformBufferHelper = [](FVulkanCommandListContext& Context, FVulkanRealUniformBuffer* VulkanUniformBuffer, int32 DataSize, const void* Data)
 	{
 		FVulkanCmdBuffer* CmdBuffer = Context.GetCommandBufferManager()->GetActiveCmdBufferDirect();
 		ensure(CmdBuffer->IsOutsideRenderPass());
@@ -229,8 +229,8 @@ inline void FVulkanDynamicRHI::UpdateUniformBuffer(FVulkanUniformBuffer* Uniform
 		VkBufferCopy Region;
 		Region.size = DataSize;
 		Region.srcOffset = LockInfo.GetBindOffset();
-		Region.dstOffset = UniformBuffer->GetOffset();
-		VkBuffer UBBuffer = UniformBuffer->GetBufferAllocation()->GetHandle();
+		Region.dstOffset = VulkanUniformBuffer->GetOffset();
+		VkBuffer UBBuffer = VulkanUniformBuffer->GetBufferAllocation()->GetHandle();
 		VulkanRHI::vkCmdCopyBuffer(CmdBuffer->GetHandle(), LockInfo.GetHandle(), UBBuffer, 1, &Region);
 	};
 
