Prefer D24_S8 for scene depth on Vulkan Android

diff --git a/Engine/Config/Android/AndroidEngine.ini b/Engine/Config/Android/AndroidEngine.ini
--- a/Engine/Config/Android/AndroidEngine.ini
+++ b/Engine/Config/Android/AndroidEngine.ini
@@ -74,5 +74,6 @@ r.VSync=1
 r.ParallelShadows=0
 ; high priority thread tasks can be downgraded to normal priority if a task is created with EPowerSavingEligibility::Eligible.
 TaskGraph.EnablePowerSavingThreadPriorityReduction=1
-
+; prefer D24_S8 for depth/stencil
+r.Vulkan.Depth24Bit=1
 ;GalaxyS6Edge=577
diff --git a/Engine/Source/Runtime/VulkanRHI/Private/VulkanDevice.cpp b/Engine/Source/Runtime/VulkanRHI/Private/VulkanDevice.cpp
--- a/Engine/Source/Runtime/VulkanRHI/Private/VulkanDevice.cpp
+++ b/Engine/Source/Runtime/VulkanRHI/Private/VulkanDevice.cpp
@@ -36,6 +36,14 @@ static TAutoConsoleVariable<int32> GCVarRobustBufferAccess(
 	ECVF_ReadOnly
 );
 
+static TAutoConsoleVariable<int32> CVarVulkanUseD24(
+	TEXT("r.Vulkan.Depth24Bit"),
+	0,
+	TEXT("0: Use 32-bit float depth buffer (default)\n1: Use 24-bit fixed point depth buffer\n"),
+	ECVF_ReadOnly
+);
+
+
 // Mirror GPixelFormats with format information for buffers
 VkFormat GVulkanBufferFormat[PF_MAX];
 
@@ -454,7 +462,16 @@ void FVulkanDevice::SetupFormats()
 	MapFormatSupport(PF_FloatRGBA, VK_FORMAT_R16G16B16A16_SFLOAT, 8);
 	SetComponentMapping(PF_FloatRGBA, VK_COMPONENT_SWIZZLE_R, VK_COMPONENT_SWIZZLE_G, VK_COMPONENT_SWIZZLE_B, VK_COMPONENT_SWIZZLE_A);
 
-	MapFormatSupportWithFallback(PF_DepthStencil, VK_FORMAT_D32_SFLOAT_S8_UINT, {VK_FORMAT_D24_UNORM_S8_UINT, VK_FORMAT_D16_UNORM_S8_UINT});
+	if (CVarVulkanUseD24.GetValueOnAnyThread() != 0)
+	{
+		// prefer VK_FORMAT_D24_UNORM_S8_UINT
+		MapFormatSupportWithFallback(PF_DepthStencil, VK_FORMAT_D24_UNORM_S8_UINT, {VK_FORMAT_D32_SFLOAT_S8_UINT, VK_FORMAT_D16_UNORM_S8_UINT});
+	}
+	else
+	{
+		// prefer VK_FORMAT_D32_SFLOAT_S8_UINT
+		MapFormatSupportWithFallback(PF_DepthStencil, VK_FORMAT_D32_SFLOAT_S8_UINT, {VK_FORMAT_D24_UNORM_S8_UINT, VK_FORMAT_D16_UNORM_S8_UINT});
+	}
 	SetComponentMapping(PF_DepthStencil, VK_COMPONENT_SWIZZLE_R, VK_COMPONENT_SWIZZLE_IDENTITY, VK_COMPONENT_SWIZZLE_IDENTITY, VK_COMPONENT_SWIZZLE_IDENTITY);
 
 	MapFormatSupport(PF_ShadowDepth, VK_FORMAT_D16_UNORM);
