Disable MSVC warning C4471 to include amd_ags.h with -Strict build.

#rnx

diff --git a/Engine/Source/Runtime/VulkanRHI/Private/Windows/VulkanWindowsPlatform.cpp b/Engine/Source/Runtime/VulkanRHI/Private/Windows/VulkanWindowsPlatform.cpp
--- a/Engine/Source/Runtime/VulkanRHI/Private/Windows/VulkanWindowsPlatform.cpp
+++ b/Engine/Source/Runtime/VulkanRHI/Private/Windows/VulkanWindowsPlatform.cpp
@@ -3,7 +3,12 @@
 #include "VulkanWindowsPlatform.h"
 #include "../VulkanRHIPrivate.h"
 #include "../VulkanDevice.h"
+
+// Disable warning about forward declared enumeration without a type, since the D3D specific enums are not used in this translation unit
+#pragma warning(push)
+#pragma warning(disable : 4471)
 #include "amd_ags.h"
+#pragma warning(pop)
 
 #include "Windows/AllowWindowsPlatformTypes.h"
 static HMODULE GVulkanDLLModule = nullptr;
