Fix crash on Android when returning from background with Vulkan.
#rnx

diff --git a/Engine/Source/Runtime/VulkanRHI/Private/VulkanViewport.cpp b/Engine/Source/Runtime/VulkanRHI/Private/VulkanViewport.cpp
--- a/Engine/Source/Runtime/VulkanRHI/Private/VulkanViewport.cpp
+++ b/Engine/Source/Runtime/VulkanRHI/Private/VulkanViewport.cpp
@@ -649,7 +649,7 @@ void FVulkanViewport::DestroySwapchain()
 		RHIBackBuffer = nullptr;
 	}
 		
-	if (FVulkanPlatform::SupportsStandardSwapchain())
+	if (FVulkanPlatform::SupportsStandardSwapchain() && SwapChain)
 	{
 		for (int32 Index = 0; Index < NUM_BUFFERS; ++Index)
 		{
