Avoid disabling VirtualTexture for VulkanRHI in the PostInit function as it causes a mis-match in DDC string key

diff --git a/Engine/Source/Runtime/VulkanRHI/Private/VulkanRHI.cpp b/Engine/Source/Runtime/VulkanRHI/Private/VulkanRHI.cpp
--- a/Engine/Source/Runtime/VulkanRHI/Private/VulkanRHI.cpp
+++ b/Engine/Source/Runtime/VulkanRHI/Private/VulkanRHI.cpp
@@ -251,22 +251,6 @@ void FVulkanDynamicRHI::PostInit()
 {
 	//work around layering violation
 	TShaderMapRef<FNULLPS>(GetGlobalShaderMap(GMaxRHIFeatureLevel))->GetPixelShader();
-
-	//#todo-rco: Remove this!
-	{
-		// Disable virtual texturing on Vulkan
-		auto* CVarVT = IConsoleManager::Get().FindConsoleVariable(TEXT("r.VirtualTextures"));
-		if (CVarVT)
-		{
-			CVarVT->Set(0);
-		}
-
-		auto* CVarVTL = IConsoleManager::Get().FindConsoleVariable(TEXT("r.VirtualTexturedLightmaps"));
-		if (CVarVTL)
-		{
-			CVarVTL->Set(0);
-		}
-	}
 }
 
 void FVulkanDynamicRHI::Shutdown()
