Multiple vulkan validation layer fixes:

- FVulkanPipelineStateCacheManager: Fixed an issue that'd leak PSOs when two threads are requesting the same PSO.
- FRenderingCompositePass::ExtractRDGTextureForOutput: Transition to read when extracting.
- InitialClear now leaves rendertargets in READ state
- DoAsyncReallocateTexture2D Now transisitions it source texture back: fixes a validation issue when source texture is used after being used to update the new texture.
- If Not using depth, stay in VK_IMAGE_LAYOUT_DEPTH_READ_ONLY_STENCIL_ATTACHMENT_OPTIMAL_KHR
- If the pixelshader is not writing a color target, set the colorwritemask to 0
- PostProcessBloomSetup.cpp: Remove unneeded resolve target parameter.
- Fixed an issue in ~FVulkanPendingGfxState, where the map would be modified when iterating over it.
- changed InternalUnlockTexture2D & RHIUnlockTextureCubeFace to -not- round up to the block size as it was giving validation problems.
- Added a missing transition to Readable, when copying from cached shadowmaps

diff --git a/Engine/Source/Runtime/Renderer/Private/PostProcess/PostProcessBloomSetup.cpp b/Engine/Source/Runtime/Renderer/Private/PostProcess/PostProcessBloomSetup.cpp
--- a/Engine/Source/Runtime/Renderer/Private/PostProcess/PostProcessBloomSetup.cpp
+++ b/Engine/Source/Runtime/Renderer/Private/PostProcess/PostProcessBloomSetup.cpp
@@ -286,7 +286,7 @@ void FRCPassPostProcessBloomSetup::Process(FRenderingCompositePassContext& Conte
 		WaitForInputPassComputeFences(Context.RHICmdList);
 
 		// Set the view family's render target/viewport.
-		FRHIRenderPassInfo RPInfo(DestRenderTarget.TargetableTexture, ERenderTargetActions::Clear_Store, DestRenderTarget.ShaderResourceTexture);
+		FRHIRenderPassInfo RPInfo(DestRenderTarget.TargetableTexture, ERenderTargetActions::Clear_Store);
 		Context.RHICmdList.BeginRenderPass(RPInfo, TEXT("Bloom"));
 		{
 			Context.SetViewportAndCallRHI(0, 0, 0.0f, DestSize.X, DestSize.Y, 1.0f);
diff --git a/Engine/Source/Runtime/Renderer/Private/ShadowDepthRendering.cpp b/Engine/Source/Runtime/Renderer/Private/ShadowDepthRendering.cpp
--- a/Engine/Source/Runtime/Renderer/Private/ShadowDepthRendering.cpp
+++ b/Engine/Source/Runtime/Renderer/Private/ShadowDepthRendering.cpp
@@ -1092,6 +1092,18 @@ void FProjectedShadowInfo::SetupShadowUniformBuffers(FRHICommandListImmediate& R
 	}
 }
 
+void FProjectedShadowInfo::TransitionCachedShadowmap(FRHICommandListImmediate& RHICmdList, FScene* Scene)
+{
+	if (CacheMode == SDCM_MovablePrimitivesOnly)
+	{
+		const FCachedShadowMapData& CachedShadowMapData = Scene->CachedShadowMaps.FindChecked(GetLightSceneInfo().Id);
+		if (CachedShadowMapData.bCachedShadowMapHasPrimitives && CachedShadowMapData.ShadowMap.IsValid())
+		{
+			RHICmdList.TransitionResource(EResourceTransitionAccess::EReadable, CachedShadowMapData.ShadowMap.DepthTarget->GetRenderTargetItem().ShaderResourceTexture);
+		}
+	}
+}
+
 
 void FProjectedShadowInfo::RenderDepthInner(FRHICommandListImmediate& RHICmdList, FSceneRenderer* SceneRenderer, FBeginShadowRenderPassFunction BeginShadowRenderPass, bool bDoParallelDispatch)
 {
@@ -1400,6 +1412,7 @@ void FSceneRenderer::RenderShadowDepthMapAtlases(FRHICommandListImmediate& RHICm
 						*LightNameWithLevel);
 				}
 				ProjectedShadowInfo->SetupShadowUniformBuffers(RHICmdList, Scene);
+				ProjectedShadowInfo->TransitionCachedShadowmap(RHICmdList, Scene);
 				ProjectedShadowInfo->RenderDepth(RHICmdList, this, BeginShadowRenderPass, true);
 			}
 		}
@@ -1436,6 +1449,7 @@ void FSceneRenderer::RenderShadowDepthMapAtlases(FRHICommandListImmediate& RHICm
 				}
 
 				ProjectedShadowInfo->SetupShadowUniformBuffers(RHICmdList, Scene);
+				ProjectedShadowInfo->TransitionCachedShadowmap(RHICmdList, Scene);
 				BeginShadowRenderPass(RHICmdList, ShadowIndex == 0);
 				ProjectedShadowInfo->RenderDepth(RHICmdList, this, BeginShadowRenderPass, false);
 				RHICmdList.EndRenderPass();
diff --git a/Engine/Source/Runtime/Renderer/Private/ShadowRendering.h b/Engine/Source/Runtime/Renderer/Private/ShadowRendering.h
--- a/Engine/Source/Runtime/Renderer/Private/ShadowRendering.h
+++ b/Engine/Source/Runtime/Renderer/Private/ShadowRendering.h
@@ -513,6 +513,11 @@ public:
 	*/
 	void SetupShadowUniformBuffers(FRHICommandListImmediate& RHICmdList, FScene* Scene);
 
+	/**
+	* Ensure Cached Shadowmap is in EReadable state
+	*/
+	void TransitionCachedShadowmap(FRHICommandListImmediate& RHICmdList, FScene* Scene);
+
 
 private:
 	// 0 if Setup...() wasn't called yet
diff --git a/Engine/Source/Runtime/VulkanRHI/Private/VulkanContext.h b/Engine/Source/Runtime/VulkanRHI/Private/VulkanContext.h
--- a/Engine/Source/Runtime/VulkanRHI/Private/VulkanContext.h
+++ b/Engine/Source/Runtime/VulkanRHI/Private/VulkanContext.h
@@ -132,29 +132,10 @@ public:
 		return Layouts.FindChecked(Image);
 	}
 
-	VkImageLayout FindOrAddLayout(VkImage Image, VkImageLayout LayoutIfNotFound)
-	{
-		VkImageLayout* Found = Layouts.Find(Image);
-		if (Found)
-		{
-			return *Found;
-		}
-
-		Layouts.Add(Image, LayoutIfNotFound);
-		return LayoutIfNotFound;
-	}
 
-	VkImageLayout& FindOrAddLayoutRW(VkImage Image, VkImageLayout LayoutIfNotFound)
-	{
-		VkImageLayout* Found = Layouts.Find(Image);
-		if (Found)
-		{
-			return *Found;
-		}
-
-		return Layouts.Add(Image, LayoutIfNotFound);
-	}
 
+	VkImageLayout FindOrAddLayout(VkImage Image, VkImageLayout LayoutIfNotFound);
+	VULKANRHI_API VkImageLayout& FindOrAddLayoutRW(VkImage Image, VkImageLayout LayoutIfNotFound);
 	void TransitionResource(FVulkanCmdBuffer* CmdBuffer, FVulkanSurface& Surface, VulkanRHI::EImageLayoutBarrier DestLayout);
 };
 
diff --git a/Engine/Source/Runtime/VulkanRHI/Private/VulkanPendingState.cpp b/Engine/Source/Runtime/VulkanRHI/Private/VulkanPendingState.cpp
--- a/Engine/Source/Runtime/VulkanRHI/Private/VulkanPendingState.cpp
+++ b/Engine/Source/Runtime/VulkanRHI/Private/VulkanPendingState.cpp
@@ -402,9 +402,15 @@ void FVulkanPendingComputeState::PrepareForDispatch(FVulkanCmdBuffer* InCmdBuffe
 
 FVulkanPendingGfxState::~FVulkanPendingGfxState()
 {
+	TArray<FVulkanGraphicsPipelineDescriptorState*> DescriptorStates;
+
 	for (auto& Pair : PipelineStates)
 	{
 		FVulkanGraphicsPipelineDescriptorState* State = Pair.Value;
+		DescriptorStates.Push(State);
+	}
+	for(FVulkanGraphicsPipelineDescriptorState* State : DescriptorStates)
+	{
 		delete State;
 	}
 }
diff --git a/Engine/Source/Runtime/VulkanRHI/Private/VulkanPipeline.cpp b/Engine/Source/Runtime/VulkanRHI/Private/VulkanPipeline.cpp
--- a/Engine/Source/Runtime/VulkanRHI/Private/VulkanPipeline.cpp
+++ b/Engine/Source/Runtime/VulkanRHI/Private/VulkanPipeline.cpp
@@ -18,6 +18,13 @@
 #include "VulkanLLM.h"
 #include "Misc/ScopeRWLock.h"
 
+
+#if !UE_BUILD_SHIPPING
+static TAtomic<uint64> SGraphicsRHICount;
+static TAtomic<uint64> SPipelineCount;
+static TAtomic<uint64> SPipelineGfxCount;
+#endif
+
 static const double HitchTime = 1.0 / 1000.0;
 
 static FCriticalSection EntryKeyToGfxPipelineMapCS;
@@ -103,10 +110,16 @@ FVulkanPipeline::FVulkanPipeline(FVulkanDevice* InDevice)
 	, Pipeline(VK_NULL_HANDLE)
 	, Layout(nullptr)
 {
+#if !UE_BUILD_SHIPPING
+	SPipelineCount++;
+#endif
 }
 
 FVulkanPipeline::~FVulkanPipeline()
 {
+#if !UE_BUILD_SHIPPING
+	SPipelineCount--;
+#endif
 	if (Pipeline != VK_NULL_HANDLE)
 	{
 		Device->GetDeferredDeletionQueue().EnqueueResource(VulkanRHI::FDeferredDeletionQueue::EType::Pipeline, Pipeline);
@@ -140,6 +153,16 @@ FVulkanGfxPipeline::FVulkanGfxPipeline(FVulkanDevice* InDevice)
 	, bRuntimeObjectsValid(false)
 #endif
 {
+#if !UE_BUILD_SHIPPING
+	SPipelineGfxCount++;
+#endif
+}
+
+FVulkanGfxPipeline::~FVulkanGfxPipeline()
+{
+#if !UE_BUILD_SHIPPING
+	SPipelineGfxCount--;
+#endif
 }
 
 void FVulkanGfxPipeline::CreateRuntimeObjects(const FGraphicsPipelineStateInitializer& InPSOInitializer)
@@ -153,9 +176,27 @@ void FVulkanGfxPipeline::CreateRuntimeObjects(const FGraphicsPipelineStateInitia
 	VertexInputState.Generate(ResourceCast(InPSOInitializer.BoundShaderState.VertexDeclarationRHI), VSHeader.InOutMask);
 	bRuntimeObjectsValid = true;
 }
+FVulkanRHIGraphicsPipelineState::FVulkanRHIGraphicsPipelineState(const FBoundShaderStateInput& InBSI, FVulkanGfxPipeline* InPipeline, EPrimitiveType InPrimitiveType)
+	: Pipeline(InPipeline)
+	, PrimitiveType(InPrimitiveType)
+{
+
+	for (int32 StageIdx = 0; StageIdx < ShaderStage::NumStages; ++StageIdx)
+	{
+		ShaderKeys[StageIdx] = GetShaderKeyForGfxStage(InBSI, (ShaderStage::EStage)StageIdx);
+	}
+
+	bHasInputAttachments = InPipeline->GetGfxLayout().GetDescriptorSetsLayout().HasInputAttachments();
+#if !UE_BUILD_SHIPPING
+	SGraphicsRHICount++;
+#endif
+}
 
 FVulkanRHIGraphicsPipelineState::~FVulkanRHIGraphicsPipelineState()
 {
+#if !UE_BUILD_SHIPPING
+	SGraphicsRHICount--;
+#endif
 	if (Pipeline)
 	{
 		Pipeline->Device->NotifyDeletedGfxPipeline(this);
@@ -613,12 +654,20 @@ FVulkanRHIGraphicsPipelineState* FVulkanPipelineStateCacheManager::CreateAndAdd(
 		}
 	}
 
-	FVulkanRHIGraphicsPipelineState* PipelineState = new FVulkanRHIGraphicsPipelineState(PSOInitializer.BoundShaderState, Pipeline, PSOInitializer.PrimitiveType);
-	PipelineState->AddRef();
-
+	FVulkanRHIGraphicsPipelineState* PipelineState = nullptr;
 	{
 		FScopeLock Lock(&InitializerToPipelineMapCS);
-		InitializerToPipelineMap.Add(MoveTemp(PSIKey), PipelineState);
+		FVulkanRHIGraphicsPipelineState** Existing = InitializerToPipelineMap.Find(PSIKey);
+		if (Existing)
+		{
+			PipelineState = *Existing;
+		}
+		else
+		{
+			PipelineState = new FVulkanRHIGraphicsPipelineState(PSOInitializer.BoundShaderState, Pipeline, PSOInitializer.PrimitiveType);
+			PipelineState->AddRef();
+			InitializerToPipelineMap.Add(MoveTemp(PSIKey), PipelineState);
+		}
 	}
 
 	EntryKeyToGfxPipelineMap.Add(MoveTemp(GfxEntryKey), PipelineState->Pipeline);
@@ -1080,9 +1129,20 @@ void FVulkanPipelineStateCacheManager::CreateGfxPipelineFromEntry(FGfxPipelineEn
 	CBInfo.attachmentCount = GfxEntry->ColorAttachmentStates.Num();
 	VkPipelineColorBlendAttachmentState BlendStates[MaxSimultaneousRenderTargets];
 	FMemory::Memzero(BlendStates);
+	uint32 ColorWriteMask = 0xffffffff;
+	if(Shaders[ShaderStage::Pixel])
+	{
+		ColorWriteMask = Shaders[ShaderStage::Pixel]->CodeHeader.InOutMask;
+	}
 	for (int32 Index = 0; Index < GfxEntry->ColorAttachmentStates.Num(); ++Index)
 	{
 		GfxEntry->ColorAttachmentStates[Index].WriteInto(BlendStates[Index]);
+		
+		if(0 == (ColorWriteMask & 1)) //clear write mask of rendertargets not written by pixelshader.
+		{
+			BlendStates[Index].colorWriteMask = 0;
+		}
+		ColorWriteMask >>= 1;
 	}
 	CBInfo.pAttachments = BlendStates;
 	CBInfo.blendConstants[0] = 1.0f;
@@ -1611,12 +1671,21 @@ FVulkanRHIGraphicsPipelineState* FVulkanPipelineStateCacheManager::FindInLoadedL
 		{
 			(*FoundPipeline)->CreateRuntimeObjects(PSOInitializer);
 		}
-		FVulkanRHIGraphicsPipelineState* PipelineState = new FVulkanRHIGraphicsPipelineState(PSOInitializer.BoundShaderState, *FoundPipeline, PSOInitializer.PrimitiveType);
+		FVulkanRHIGraphicsPipelineState* PipelineState = nullptr;
 		{
 			FScopeLock Lock2(&InitializerToPipelineMapCS);
-			InitializerToPipelineMap.Add(MoveTemp(PSIKey), PipelineState);
-		}
-		PipelineState->AddRef();
+			FVulkanRHIGraphicsPipelineState** Existing = InitializerToPipelineMap.Find(PSIKey);
+			if (Existing)
+			{
+				PipelineState = *Existing;
+			}
+			else
+			{
+				PipelineState = new FVulkanRHIGraphicsPipelineState(PSOInitializer.BoundShaderState, *FoundPipeline, PSOInitializer.PrimitiveType);
+				PipelineState->AddRef();
+				InitializerToPipelineMap.Add(MoveTemp(PSIKey), PipelineState);
+			}
+		}		
 		delete GfxEntry;
 #if VULKAN_ENABLE_LRU_CACHE
 		PipelineLRU.Touch(PipelineState);
diff --git a/Engine/Source/Runtime/VulkanRHI/Private/VulkanPipeline.h b/Engine/Source/Runtime/VulkanRHI/Private/VulkanPipeline.h
--- a/Engine/Source/Runtime/VulkanRHI/Private/VulkanPipeline.h
+++ b/Engine/Source/Runtime/VulkanRHI/Private/VulkanPipeline.h
@@ -714,7 +714,7 @@ public:
 #else
 	FVulkanGfxPipeline(FVulkanDevice* InDevice);
 #endif
-
+	virtual ~FVulkanGfxPipeline();
 	inline void Bind(VkCommandBuffer CmdBuffer)
 	{
 #if VULKAN_ENABLE_LRU_CACHE
@@ -758,18 +758,7 @@ private:
 class FVulkanRHIGraphicsPipelineState : public FRHIGraphicsPipelineState
 {
 public:
-	FVulkanRHIGraphicsPipelineState(const FBoundShaderStateInput& InBSI, FVulkanGfxPipeline* InPipeline, EPrimitiveType InPrimitiveType)
-		: Pipeline(InPipeline)
-		, PrimitiveType(InPrimitiveType)
-	{
-		for (int32 StageIdx = 0; StageIdx < ShaderStage::NumStages; ++StageIdx)
-		{
-			ShaderKeys[StageIdx] = GetShaderKeyForGfxStage(InBSI, (ShaderStage::EStage)StageIdx);
-		}
-		
-		bHasInputAttachments = InPipeline->GetGfxLayout().GetDescriptorSetsLayout().HasInputAttachments();
-	}
-
+	FVulkanRHIGraphicsPipelineState(const FBoundShaderStateInput& InBSI, FVulkanGfxPipeline* InPipeline, EPrimitiveType InPrimitiveType);
 	~FVulkanRHIGraphicsPipelineState();
 
 	inline void Bind(VkCommandBuffer CmdBuffer)
diff --git a/Engine/Source/Runtime/VulkanRHI/Private/VulkanRHIPrivate.h b/Engine/Source/Runtime/VulkanRHI/Private/VulkanRHIPrivate.h
--- a/Engine/Source/Runtime/VulkanRHI/Private/VulkanRHIPrivate.h
+++ b/Engine/Source/Runtime/VulkanRHI/Private/VulkanRHIPrivate.h
@@ -813,6 +813,10 @@ namespace VulkanRHI
 		{
 			return VK_IMAGE_LAYOUT_DEPTH_READ_ONLY_STENCIL_ATTACHMENT_OPTIMAL_KHR;
 		}
+		if(!RequestedDSAccess.IsUsingDepth() && RequestedDSAccess.IsUsingStencil())
+		{
+			return VK_IMAGE_LAYOUT_DEPTH_READ_ONLY_STENCIL_ATTACHMENT_OPTIMAL_KHR;
+		}
 #endif
 
 		ensure(RequestedDSAccess.IsDepthWrite() || RequestedDSAccess.IsStencilWrite());
diff --git a/Engine/Source/Runtime/VulkanRHI/Private/VulkanRenderTarget.cpp b/Engine/Source/Runtime/VulkanRHI/Private/VulkanRenderTarget.cpp
--- a/Engine/Source/Runtime/VulkanRHI/Private/VulkanRenderTarget.cpp
+++ b/Engine/Source/Runtime/VulkanRHI/Private/VulkanRenderTarget.cpp
@@ -469,6 +469,28 @@ void FTransitionAndLayoutManager::NotifyDeletedRenderTarget(FVulkanDevice& InDev
 	}
 }
 
+VkImageLayout FTransitionAndLayoutManager::FindOrAddLayout(VkImage Image, VkImageLayout LayoutIfNotFound)
+{
+	VkImageLayout* Found = Layouts.Find(Image);
+	if (Found)
+	{
+		return *Found;
+	}
+
+	Layouts.Add(Image, LayoutIfNotFound);
+	return LayoutIfNotFound;
+}
+
+VkImageLayout& FTransitionAndLayoutManager::FindOrAddLayoutRW(VkImage Image, VkImageLayout LayoutIfNotFound)
+{
+	VkImageLayout* Found = Layouts.Find(Image);
+	if (Found)
+	{
+		return *Found;
+	}
+	return Layouts.Add(Image, LayoutIfNotFound);
+}
+
 void FTransitionAndLayoutManager::TransitionResource(FVulkanCmdBuffer* CmdBuffer, FVulkanSurface& Surface, VulkanRHI::EImageLayoutBarrier DestLayout)
 {
 	VkImageLayout* FoundLayout = Layouts.Find(Surface.Image);
diff --git a/Engine/Source/Runtime/VulkanRHI/Private/VulkanTexture.cpp b/Engine/Source/Runtime/VulkanRHI/Private/VulkanTexture.cpp
--- a/Engine/Source/Runtime/VulkanRHI/Private/VulkanTexture.cpp
+++ b/Engine/Source/Runtime/VulkanRHI/Private/VulkanTexture.cpp
@@ -863,12 +863,14 @@ void FVulkanSurface::InitialClear(FVulkanCommandListContext& Context,const FClea
 		Color.float32[2] = ClearValueBinding.Value.Color[2];
 		Color.float32[3] = ClearValueBinding.Value.Color[3];
 
+
+
 		// Clear
 		VulkanRHI::vkCmdClearColorImage(CmdBuffer->GetHandle(), Image, VK_IMAGE_LAYOUT_TRANSFER_DST_OPTIMAL, &Color, 1, &Barrier.GetSubresource(BarrierIndex));
 
 		// Transfer to Present or Color
 		Barrier.ResetStages();
-		Barrier.SetTransition(BarrierIndex, EImageLayoutBarrier::TransferDest, bTransitionToPresentable ? EImageLayoutBarrier::Present : EImageLayoutBarrier::ColorAttachment);
+		Barrier.SetTransition(BarrierIndex, EImageLayoutBarrier::TransferDest, bTransitionToPresentable ? EImageLayoutBarrier::Present : EImageLayoutBarrier::PixelShaderRead);
 		Barrier.Execute(CmdBuffer);
 	}
 	else
@@ -885,7 +887,7 @@ void FVulkanSurface::InitialClear(FVulkanCommandListContext& Context,const FClea
 
 		// General -> DepthStencil
 		Barrier.ResetStages();
-		Barrier.SetTransition(BarrierIndex, EImageLayoutBarrier::TransferDest, EImageLayoutBarrier::DepthStencilAttachment);
+		Barrier.SetTransition(BarrierIndex, EImageLayoutBarrier::TransferDest, EImageLayoutBarrier::PixelDepthStencilRead);
 		Barrier.Execute(CmdBuffer);
 	}
 
@@ -1014,6 +1016,9 @@ static void DoAsyncReallocateTexture2D(FVulkanCommandListContext& Context, FVulk
 		Region.extent.depth = 1;
 	}
 
+
+	VkImageLayout OldTextureLayout = VK_IMAGE_LAYOUT_UNDEFINED;
+
 	{
 		// Pre-copy barriers
 		FPendingBarrier Barrier;
@@ -1023,7 +1028,7 @@ static void DoAsyncReallocateTexture2D(FVulkanCommandListContext& Context, FVulk
 			Barrier.SetTransition(BarrierIndex, EImageLayoutBarrier::Undefined, VulkanRHI::EImageLayoutBarrier::TransferDest);
 		}
 
-		VkImageLayout OldTextureLayout = Context.GetTransitionAndLayoutManager().FindOrAddLayout(OldTexture->Surface.Image, VK_IMAGE_LAYOUT_UNDEFINED);
+		OldTextureLayout = Context.GetTransitionAndLayoutManager().FindOrAddLayout(OldTexture->Surface.Image, VK_IMAGE_LAYOUT_UNDEFINED);
 		ensure(OldTextureLayout != VK_IMAGE_LAYOUT_UNDEFINED);
 		if (OldTextureLayout != VK_IMAGE_LAYOUT_TRANSFER_SRC_OPTIMAL)
 		{
@@ -1034,7 +1039,6 @@ static void DoAsyncReallocateTexture2D(FVulkanCommandListContext& Context, FVulk
 
 		Barrier.Execute(CmdBuffer);
 	}
-	Context.GetTransitionAndLayoutManager().FindOrAddLayoutRW(OldTexture->Surface.Image, VK_IMAGE_LAYOUT_UNDEFINED) = VK_IMAGE_LAYOUT_TRANSFER_SRC_OPTIMAL;
 
 	VulkanRHI::vkCmdCopyImage(StagingCommandBuffer, OldTexture->Surface.Image, VK_IMAGE_LAYOUT_TRANSFER_SRC_OPTIMAL, NewTexture->Surface.Image, VK_IMAGE_LAYOUT_TRANSFER_DST_OPTIMAL, NumSharedMips, Regions);
 
@@ -1044,6 +1048,14 @@ static void DoAsyncReallocateTexture2D(FVulkanCommandListContext& Context, FVulk
 		int32 BarrierIndex = Barrier.AddImageBarrier(NewTexture->Surface.Image, VK_IMAGE_ASPECT_COLOR_BIT, NumSharedMips);
 		Barrier.GetSubresource(BarrierIndex).baseMipLevel = DestFirstMip;
 		Barrier.SetTransition(BarrierIndex, EImageLayoutBarrier::TransferDest, VulkanRHI::EImageLayoutBarrier::PixelShaderRead);
+
+		if (OldTextureLayout != VK_IMAGE_LAYOUT_TRANSFER_SRC_OPTIMAL)
+		{
+			BarrierIndex = Barrier.AddImageBarrier(OldTexture->Surface.Image, VK_IMAGE_ASPECT_COLOR_BIT, NumSharedMips);
+			Barrier.GetSubresource(BarrierIndex).baseMipLevel = SourceFirstMip;
+			Barrier.SetTransition(BarrierIndex, VulkanRHI::EImageLayoutBarrier::TransferSource, VulkanRHI::GetImageLayoutFromVulkanLayout(OldTextureLayout));
+		}
+
 		Barrier.Execute(CmdBuffer);
 	}
 	Context.GetTransitionAndLayoutManager().FindOrAddLayoutRW(NewTexture->Surface.Image, VK_IMAGE_LAYOUT_UNDEFINED) = VK_IMAGE_LAYOUT_SHADER_READ_ONLY_OPTIMAL;
@@ -1166,8 +1178,11 @@ void FVulkanDynamicRHI::InternalUnlockTexture2D(bool bFromRenderingThread, FRHIT
 	}
 
 	EPixelFormat Format = Texture->Surface.PixelFormat;
-	uint32 MipWidth = FMath::Max<uint32>(Texture->Surface.Width >> MipIndex, GPixelFormats[Format].BlockSizeX);
-	uint32 MipHeight = FMath::Max<uint32>(Texture->Surface.Height >> MipIndex, GPixelFormats[Format].BlockSizeY);
+	uint32 MipWidth = FMath::Max<uint32>(Texture->Surface.Width >> MipIndex, 0);
+	uint32 MipHeight = FMath::Max<uint32>(Texture->Surface.Height >> MipIndex, 0);
+	ensure(!(MipHeight == 0 && MipWidth == 0));
+	MipWidth = FMath::Max<uint32>(MipWidth, 1);
+	MipHeight = FMath::Max<uint32>(MipHeight, 1);
 	uint32 LayerCount = Texture->Surface.GetNumberOfArrayLevels();
 
 	VkImageSubresourceRange SubresourceRange;
@@ -1929,8 +1944,11 @@ void FVulkanDynamicRHI::RHIUnlockTextureCubeFace(FRHITextureCube* TextureCubeRHI
 	}
 
 	EPixelFormat Format = Texture->Surface.PixelFormat;
-	uint32 MipWidth = FMath::Max<uint32>(Texture->Surface.Width >> MipIndex, GPixelFormats[Format].BlockSizeX);
-	uint32 MipHeight = FMath::Max<uint32>(Texture->Surface.Height >> MipIndex, GPixelFormats[Format].BlockSizeY);
+	uint32 MipWidth = FMath::Max<uint32>(Texture->Surface.Width >> MipIndex, 0);
+	uint32 MipHeight = FMath::Max<uint32>(Texture->Surface.Height >> MipIndex, 0);
+	ensure(!(MipHeight == 0 && MipWidth == 0));
+	MipWidth = FMath::Max<uint32>(MipWidth, 1);
+	MipHeight = FMath::Max<uint32>(MipHeight, 1);
 
 	VkImageSubresourceRange SubresourceRange;
 	FMemory::Memzero(SubresourceRange);
