fix for uniformbuffer updates not waiting properly on previous reads when using r.Vulkan.AllowUniformUpload=1
This fixes a problem with different levels of the CSM looking messed up.

diff --git a/Engine/Source/Runtime/VulkanRHI/Private/VulkanContext.h b/Engine/Source/Runtime/VulkanRHI/Private/VulkanContext.h
--- a/Engine/Source/Runtime/VulkanRHI/Private/VulkanContext.h
+++ b/Engine/Source/Runtime/VulkanRHI/Private/VulkanContext.h
@@ -357,6 +357,7 @@ protected:
 	FVulkanQueue* Queue;
 	bool bSubmitAtNextSafePoint;
 	bool bAutomaticFlushAfterComputeShader;
+	bool bUniformBufferUploadRenderPassDirty = true;
 	FVulkanUniformBufferUploader* UniformBufferUploader;
 
 	void BeginOcclusionQueryBatch(FVulkanCmdBuffer* CmdBuffer, uint32 NumQueriesInBatch);
diff --git a/Engine/Source/Runtime/VulkanRHI/Private/VulkanRenderTarget.cpp b/Engine/Source/Runtime/VulkanRHI/Private/VulkanRenderTarget.cpp
--- a/Engine/Source/Runtime/VulkanRHI/Private/VulkanRenderTarget.cpp
+++ b/Engine/Source/Runtime/VulkanRHI/Private/VulkanRenderTarget.cpp
@@ -1755,6 +1755,7 @@ void FVulkanCommandListContext::RHIEndRenderPass()
 		}
 	}
 	RHIPopEvent();
+	bUniformBufferUploadRenderPassDirty = true;
 }
 
 void FVulkanCommandListContext::RHINextSubpass()
diff --git a/Engine/Source/Runtime/VulkanRHI/Private/VulkanUniformBuffer.cpp b/Engine/Source/Runtime/VulkanRHI/Private/VulkanUniformBuffer.cpp
--- a/Engine/Source/Runtime/VulkanRHI/Private/VulkanUniformBuffer.cpp
+++ b/Engine/Source/Runtime/VulkanRHI/Private/VulkanUniformBuffer.cpp
@@ -231,6 +231,15 @@ inline void FVulkanDynamicRHI::UpdateUniformBuffer(FVulkanUniformBuffer* Uniform
 		Region.srcOffset = LockInfo.GetBindOffset();
 		Region.dstOffset = VulkanUniformBuffer->GetOffset();
 		VkBuffer UBBuffer = VulkanUniformBuffer->GetBufferAllocation()->GetHandle();
+		if(Context.bUniformBufferUploadRenderPassDirty)
+		{
+			Context.bUniformBufferUploadRenderPassDirty = false;
+			VkMemoryBarrier Barrier;
+			ZeroVulkanStruct(Barrier, VK_STRUCTURE_TYPE_MEMORY_BARRIER);
+			Barrier.srcAccessMask = VK_ACCESS_UNIFORM_READ_BIT; 
+			Barrier.dstAccessMask = VK_ACCESS_TRANSFER_WRITE_BIT;
+			VulkanRHI::vkCmdPipelineBarrier(CmdBuffer->GetHandle(), VK_PIPELINE_STAGE_ALL_COMMANDS_BIT, VK_PIPELINE_STAGE_ALL_COMMANDS_BIT, 0, 1, &Barrier, 0, nullptr, 0, nullptr);
+		}
 		VulkanRHI::vkCmdCopyBuffer(CmdBuffer->GetHandle(), LockInfo.GetHandle(), UBBuffer, 1, &Region);
 	};
 
