Pass in previous swapchain when recreating. stadia validation issue.

diff --git a/Engine/Source/Runtime/VulkanRHI/Private/VulkanSwapChain.cpp b/Engine/Source/Runtime/VulkanRHI/Private/VulkanSwapChain.cpp
--- a/Engine/Source/Runtime/VulkanRHI/Private/VulkanSwapChain.cpp
+++ b/Engine/Source/Runtime/VulkanRHI/Private/VulkanSwapChain.cpp
@@ -11,6 +11,15 @@
 #include "IHeadMountedDisplayModule.h"
 #include "IHeadMountedDisplayVulkanExtensions.h"
 
+
+int32 GVulkanKeepSwapChain = 1;
+static FAutoConsoleVariableRef CVarVulkanKeepSwapChain(
+	TEXT("r.Vulkan.KeepSwapChain"),
+	GVulkanKeepSwapChain,
+	TEXT("Whether to keep old swap chain to pass through when creating the next one"),
+	ECVF_RenderThreadSafe
+);
+
 int32 GShouldCpuWaitForFence = 1;
 static FAutoConsoleVariableRef CVarCpuWaitForFence(
 	TEXT("r.Vulkan.CpuWaitForFence"),
@@ -108,7 +117,7 @@ extern TAutoConsoleVariable<int32> GAllowPresentOnComputeQueue;
 static TSet<EPixelFormat> GPixelFormatNotSupportedWarning;
 
 FVulkanSwapChain::FVulkanSwapChain(VkInstance InInstance, FVulkanDevice& InDevice, void* WindowHandle, EPixelFormat& InOutPixelFormat, uint32 Width, uint32 Height,
-	uint32* InOutDesiredNumBackBuffers, TArray<VkImage>& OutImages, int8 InLockToVsync)
+	uint32* InOutDesiredNumBackBuffers, TArray<VkImage>& OutImages, int8 InLockToVsync, FVulkanSwapChainRecreateInfo* RecreateInfo)
 	: SwapChain(VK_NULL_HANDLE)
 	, Device(InDevice)
 	, Surface(VK_NULL_HANDLE)
@@ -123,8 +132,17 @@ FVulkanSwapChain::FVulkanSwapChain(VkInstance InInstance, FVulkanDevice& InDevic
 
 	NextPresentTargetTime = (FPlatformTime::Seconds() - GStartTime);
 
-	// let the platform create the surface
-	FVulkanPlatform::CreateSurface(WindowHandle, Instance, &Surface);
+	if(RecreateInfo != nullptr && RecreateInfo->SwapChain != VK_NULL_HANDLE)
+	{
+		check(RecreateInfo->Surface != VK_NULL_HANDLE);
+		Surface = RecreateInfo->Surface;
+		RecreateInfo->Surface = VK_NULL_HANDLE;
+	}
+	else
+	{
+		// let the platform create the surface
+		FVulkanPlatform::CreateSurface(WindowHandle, Instance, &Surface);
+	}
 
 	// Find Pixel format for presentable images
 	VkSurfaceFormatKHR CurrFormat;
@@ -385,7 +403,6 @@ FVulkanSwapChain::FVulkanSwapChain(VkInstance InInstance, FVulkanDevice& InDevic
 
 	uint32 SizeX = FVulkanPlatform::SupportsQuerySurfaceProperties() ? (SurfProperties.currentExtent.width == 0xFFFFFFFF ? Width : SurfProperties.currentExtent.width) : Width;
 	uint32 SizeY = FVulkanPlatform::SupportsQuerySurfaceProperties() ? (SurfProperties.currentExtent.height == 0xFFFFFFFF ? Height : SurfProperties.currentExtent.height) : Height;
-	//FPlatformMisc::LowLevelOutputDebugStringf(TEXT("Create swapchain: %ux%u \n"), SizeX, SizeY);
 
 	VkSwapchainCreateInfoKHR SwapChainInfo;
 	ZeroVulkanStruct(SwapChainInfo, VK_STRUCTURE_TYPE_SWAPCHAIN_CREATE_INFO_KHR);
@@ -405,6 +422,12 @@ FVulkanSwapChain::FVulkanSwapChain(VkInstance InInstance, FVulkanDevice& InDevic
 	SwapChainInfo.imageSharingMode = VK_SHARING_MODE_EXCLUSIVE;
 	SwapChainInfo.presentMode = PresentMode;
 	SwapChainInfo.oldSwapchain = VK_NULL_HANDLE;
+	if(RecreateInfo != nullptr)
+	{
+		SwapChainInfo.oldSwapchain = RecreateInfo->SwapChain;
+	}
+	
+
 	SwapChainInfo.clipped = VK_TRUE;
 	SwapChainInfo.compositeAlpha = CompositeAlpha;
 
@@ -432,6 +455,20 @@ FVulkanSwapChain::FVulkanSwapChain(VkInstance InInstance, FVulkanDevice& InDevic
 
 	VERIFYVULKANRESULT_EXPANDED(VulkanRHI::vkCreateSwapchainKHR(Device.GetInstanceHandle(), &SwapChainInfo, VULKAN_CPU_ALLOCATOR, &SwapChain));
 
+	if (RecreateInfo != nullptr)
+	{
+		if (RecreateInfo->SwapChain != VK_NULL_HANDLE)
+		{
+			VulkanRHI::vkDestroySwapchainKHR(Device.GetInstanceHandle(), RecreateInfo->SwapChain, VULKAN_CPU_ALLOCATOR);
+			RecreateInfo->SwapChain = VK_NULL_HANDLE;
+		}
+		if (RecreateInfo->Surface != VK_NULL_HANDLE)
+		{
+			VulkanRHI::vkDestroySurfaceKHR(Instance, RecreateInfo->Surface, VULKAN_CPU_ALLOCATOR);
+			RecreateInfo->Surface = VK_NULL_HANDLE;
+		}
+	}
+
 	InternalWidth = FMath::Min(Width, SwapChainInfo.imageExtent.width);
 	InternalHeight = FMath::Min(Height, SwapChainInfo.imageExtent.height);
 
@@ -468,7 +505,7 @@ FVulkanSwapChain::FVulkanSwapChain(VkInstance InInstance, FVulkanDevice& InDevic
 	PresentID = 0;
 }
 
-void FVulkanSwapChain::Destroy()
+void FVulkanSwapChain::Destroy(FVulkanSwapChainRecreateInfo* RecreateInfo)
 {
 	check(FVulkanPlatform::SupportsStandardSwapchain());
 
@@ -476,7 +513,16 @@ void FVulkanSwapChain::Destroy()
 	// Alternatively could also check on the fence(s) for the image(s) from the swapchain but then timing out/waiting could become an issue.
 	Device.WaitUntilIdle();
 
-	VulkanRHI::vkDestroySwapchainKHR(Device.GetInstanceHandle(), SwapChain, VULKAN_CPU_ALLOCATOR);
+	bool bRecreate = RecreateInfo && GVulkanKeepSwapChain;
+	if (bRecreate)
+	{
+		RecreateInfo->SwapChain = SwapChain;
+		RecreateInfo->Surface = Surface;
+	}
+	else
+	{
+		VulkanRHI::vkDestroySwapchainKHR(Device.GetInstanceHandle(), SwapChain, VULKAN_CPU_ALLOCATOR);
+	}
 	SwapChain = VK_NULL_HANDLE;
 
 #if VULKAN_USE_IMAGE_ACQUIRE_FENCES
@@ -493,7 +539,10 @@ void FVulkanSwapChain::Destroy()
 		ImageAcquiredSemaphore[BufferIndex]->Release();
 	}
 
-	VulkanRHI::vkDestroySurfaceKHR(Instance, Surface, VULKAN_CPU_ALLOCATOR);
+	if(!bRecreate)
+	{
+		VulkanRHI::vkDestroySurfaceKHR(Instance, Surface, VULKAN_CPU_ALLOCATOR);
+	}
 	Surface = VK_NULL_HANDLE;
 }
 
diff --git a/Engine/Source/Runtime/VulkanRHI/Private/VulkanSwapChain.h b/Engine/Source/Runtime/VulkanRHI/Private/VulkanSwapChain.h
--- a/Engine/Source/Runtime/VulkanRHI/Private/VulkanSwapChain.h
+++ b/Engine/Source/Runtime/VulkanRHI/Private/VulkanSwapChain.h
@@ -13,13 +13,20 @@ namespace VulkanRHI
 
 class FVulkanQueue;
 
+struct FVulkanSwapChainRecreateInfo
+{
+	VkSwapchainKHR SwapChain;
+	VkSurfaceKHR Surface;
+};
+
+
 class FVulkanSwapChain
 {
 public:
 	FVulkanSwapChain(VkInstance InInstance, FVulkanDevice& InDevice, void* WindowHandle, EPixelFormat& InOutPixelFormat, uint32 Width, uint32 Height,
-		uint32* InOutDesiredNumBackBuffers, TArray<VkImage>& OutImages, int8 bLockToVsync);
+		uint32* InOutDesiredNumBackBuffers, TArray<VkImage>& OutImages, int8 bLockToVsync, FVulkanSwapChainRecreateInfo* RecreateInfo);
 
-	void Destroy();
+	void Destroy(FVulkanSwapChainRecreateInfo* RecreateInfo);
 
 	// Has to be negative as we use this also on other callbacks as the acquired image index
 	enum class EStatus
diff --git a/Engine/Source/Runtime/VulkanRHI/Private/VulkanViewport.cpp b/Engine/Source/Runtime/VulkanRHI/Private/VulkanViewport.cpp
--- a/Engine/Source/Runtime/VulkanRHI/Private/VulkanViewport.cpp
+++ b/Engine/Source/Runtime/VulkanRHI/Private/VulkanViewport.cpp
@@ -141,7 +141,7 @@ FVulkanViewport::FVulkanViewport(FVulkanDynamicRHI* InRHI, FVulkanDevice* InDevi
 	// Make sure Instance is created
 	RHI->InitInstance();
 
-	CreateSwapchain();
+	CreateSwapchain(nullptr);
 
 	if (FVulkanPlatform::SupportsStandardSwapchain())
 	{
@@ -177,7 +177,7 @@ FVulkanViewport::~FVulkanViewport()
 			BackBufferImages[Index] = VK_NULL_HANDLE;
 		}
 
-		SwapChain->Destroy();
+		SwapChain->Destroy(nullptr);
 		delete SwapChain;
 		SwapChain = nullptr;
 	}
@@ -521,9 +521,12 @@ void FVulkanViewport::RecreateSwapchain(void* NewNativeWindow)
 {
 	FScopeLock LockSwapchain(&RecreatingSwapchain);
 
-	DestroySwapchain();
+	FVulkanSwapChainRecreateInfo RecreateInfo = { VK_NULL_HANDLE, VK_NULL_HANDLE };
+	DestroySwapchain(&RecreateInfo);
 	WindowHandle = NewNativeWindow;
-	CreateSwapchain();
+	CreateSwapchain(&RecreateInfo);
+	check(RecreateInfo.Surface == VK_NULL_HANDLE);
+	check(RecreateInfo.SwapChain == VK_NULL_HANDLE);
 }
 
 void FVulkanViewport::Tick(float DeltaTime)
@@ -557,12 +560,15 @@ void FVulkanViewport::RecreateSwapchainFromRT(EPixelFormat PreferredPixelFormat)
 
 	// TODO: should flush RHIT commands here?
 	
-	DestroySwapchain();
+	FVulkanSwapChainRecreateInfo RecreateInfo = { VK_NULL_HANDLE, VK_NULL_HANDLE};
+	DestroySwapchain(&RecreateInfo);
 	PixelFormat = PreferredPixelFormat;
-	CreateSwapchain();
+	CreateSwapchain(&RecreateInfo);
+	check(RecreateInfo.Surface == VK_NULL_HANDLE);
+	check(RecreateInfo.SwapChain == VK_NULL_HANDLE);
 }
 
-void FVulkanViewport::CreateSwapchain()
+void FVulkanViewport::CreateSwapchain(FVulkanSwapChainRecreateInfo* RecreateInfo)
 {
 	if (FVulkanPlatform::SupportsStandardSwapchain())
 	{
@@ -574,7 +580,8 @@ void FVulkanViewport::CreateSwapchain()
 			PixelFormat, SizeX, SizeY,
 			&DesiredNumBackBuffers,
 			Images,
-			LockToVsync
+			LockToVsync,
+			RecreateInfo
 		);
 
 		checkf(Images.Num() == NUM_BUFFERS, TEXT("Actual Num: %i"), Images.Num());
@@ -619,6 +626,19 @@ void FVulkanViewport::CreateSwapchain()
 	else
 	{
 		PixelFormat = FVulkanPlatform::GetPixelFormatForNonDefaultSwapchain();
+		if (RecreateInfo != nullptr)
+		{
+			if(RecreateInfo->SwapChain)
+			{
+				VulkanRHI::vkDestroySwapchainKHR(Device->GetInstanceHandle(), RecreateInfo->SwapChain, VULKAN_CPU_ALLOCATOR);
+				RecreateInfo->SwapChain = VK_NULL_HANDLE;
+			}
+			if (RecreateInfo->Surface)
+			{
+				VulkanRHI::vkDestroySurfaceKHR(RHI->Instance, RecreateInfo->Surface, VULKAN_CPU_ALLOCATOR);
+				RecreateInfo->Surface = VK_NULL_HANDLE;
+			}
+		}
 	}
 
 	if (!FVulkanPlatform::SupportsStandardSwapchain() || GVulkanDelayAcquireImage == EDelayAcquireImageType::DelayAcquire)
@@ -635,7 +655,7 @@ void FVulkanViewport::CreateSwapchain()
 	AcquiredImageIndex = -1;
 }
 
-void FVulkanViewport::DestroySwapchain()
+void FVulkanViewport::DestroySwapchain(FVulkanSwapChainRecreateInfo* RecreateInfo)
 {
 	// Submit all command buffers here
 	Device->SubmitCommandsAndFlushGPU();
@@ -661,7 +681,7 @@ void FVulkanViewport::DestroySwapchain()
 		
 		Device->GetDeferredDeletionQueue().ReleaseResources(true);
 
-		SwapChain->Destroy();
+		SwapChain->Destroy(RecreateInfo);
 		delete SwapChain;
 		SwapChain = nullptr;
 
diff --git a/Engine/Source/Runtime/VulkanRHI/Private/VulkanViewport.h b/Engine/Source/Runtime/VulkanRHI/Private/VulkanViewport.h
--- a/Engine/Source/Runtime/VulkanRHI/Private/VulkanViewport.h
+++ b/Engine/Source/Runtime/VulkanRHI/Private/VulkanViewport.h
@@ -112,8 +112,8 @@ protected:
 	FVulkanCmdBuffer* LastFrameCommandBuffer = nullptr;
 	uint64 LastFrameFenceCounter = 0;
 
-	void CreateSwapchain();
-	void DestroySwapchain();
+	void CreateSwapchain(struct FVulkanSwapChainRecreateInfo* RecreateInfo);
+	void DestroySwapchain(struct FVulkanSwapChainRecreateInfo* RecreateInfo);
 	void AcquireImageIndex();
 	bool TryAcquireImageIndex();
 
