Fix compile error in light of CL# 8741159
#rnx

diff --git a/Engine/Source/Runtime/VulkanRHI/Private/VulkanRHI.cpp b/Engine/Source/Runtime/VulkanRHI/Private/VulkanRHI.cpp
--- a/Engine/Source/Runtime/VulkanRHI/Private/VulkanRHI.cpp
+++ b/Engine/Source/Runtime/VulkanRHI/Private/VulkanRHI.cpp
@@ -1718,9 +1718,9 @@ void FVulkanDynamicRHI::DestroySwapChain()
 	{
 		UE_LOG(LogVulkanRHI, Log, TEXT("Destroy swapchain ... "));
 		
-		for (auto& Viewport : Viewports)
+		for (FVulkanViewport* Viewport : Viewports)
 		{
-			Viewport->DestroySwapchain();
+			Viewport->DestroySwapchain(nullptr);
 		}
 	});
 
@@ -1745,7 +1745,7 @@ void FVulkanDynamicRHI::RecreateSwapChain(void* NewNativeWindow)
 		{
 			UE_LOG(LogVulkanRHI, Log, TEXT("Recreate swapchain ... "));
 			
-			for (auto& Viewport : Viewports)
+			for (FVulkanViewport* Viewport : Viewports)
 			{
 				Viewport->RecreateSwapchain(NewNativeWindow);
 			}
