fix two issues with automatic mip gen on vulkan

diff --git a/Engine/Source/Runtime/VulkanRHI/Private/VulkanRenderTarget.cpp b/Engine/Source/Runtime/VulkanRHI/Private/VulkanRenderTarget.cpp
--- a/Engine/Source/Runtime/VulkanRHI/Private/VulkanRenderTarget.cpp
+++ b/Engine/Source/Runtime/VulkanRHI/Private/VulkanRenderTarget.cpp
@@ -1249,6 +1249,7 @@ bool FVulkanCommandListContext::FPendingTransition::GatherBarriers(FTransitionAn
 			SrcAccess = VK_ACCESS_SHADER_WRITE_BIT;
 			DestAccess = VK_ACCESS_SHADER_READ_BIT;
 			break;
+		case EResourceTransitionAccess::ERWSubResBarrier: //not optimal, but will have to do for now
 		case EResourceTransitionAccess::ERWBarrier:
 			SrcAccess = VK_ACCESS_SHADER_READ_BIT | VK_ACCESS_SHADER_WRITE_BIT;
 			DestAccess = VK_ACCESS_SHADER_READ_BIT | VK_ACCESS_SHADER_WRITE_BIT;
diff --git a/Engine/Source/Runtime/VulkanRHI/Private/VulkanTexture.cpp b/Engine/Source/Runtime/VulkanRHI/Private/VulkanTexture.cpp
--- a/Engine/Source/Runtime/VulkanRHI/Private/VulkanTexture.cpp
+++ b/Engine/Source/Runtime/VulkanRHI/Private/VulkanTexture.cpp
@@ -246,7 +246,10 @@ VkImage FVulkanSurface::CreateImage(
 	{
 		VkFormat ViewFormat = UEToVkTextureFormat(InFormat, (UEFlags & TexCreate_SRGB) == TexCreate_SRGB);
 		*OutViewFormat = ViewFormat;
-		ImageCreateInfo.format = ViewFormat;
+		if((UEFlags & TexCreate_UAV) == 0)
+		{
+			ImageCreateInfo.format = ViewFormat;
+		}
 	}
 
 	ImageCreateInfo.extent.width = SizeX;
