Remove duplicate case
#rnx

diff --git a/Engine/Source/Runtime/VulkanRHI/Private/VulkanRenderTarget.cpp b/Engine/Source/Runtime/VulkanRHI/Private/VulkanRenderTarget.cpp
--- a/Engine/Source/Runtime/VulkanRHI/Private/VulkanRenderTarget.cpp
+++ b/Engine/Source/Runtime/VulkanRHI/Private/VulkanRenderTarget.cpp
@@ -1255,7 +1255,6 @@ bool FVulkanCommandListContext::FPendingTransition::GatherBarriers(FTransitionAn
 			DestAccess = VK_ACCESS_SHADER_READ_BIT | VK_ACCESS_SHADER_WRITE_BIT;
 			break;
 		case EResourceTransitionAccess::ERWNoBarrier:
-		case EResourceTransitionAccess::ERWSubResBarrier:
 			//#todo-rco: Skip for now
 			continue;
 		default:
