Accomodate drivers that require more than 3 swapchain images (UE-79943).

- Some drivers set minCount to be more than 3 in surface  properties, see e.g. https://bugs.freedesktop.org/show_bug.cgi?id=111522

diff --git a/Engine/Source/Runtime/VulkanRHI/Private/VulkanSwapChain.cpp b/Engine/Source/Runtime/VulkanRHI/Private/VulkanSwapChain.cpp
--- a/Engine/Source/Runtime/VulkanRHI/Private/VulkanSwapChain.cpp
+++ b/Engine/Source/Runtime/VulkanRHI/Private/VulkanSwapChain.cpp
@@ -451,7 +451,14 @@ FVulkanSwapChain::FVulkanSwapChain(VkInstance InInstance, FVulkanDevice& InDevic
 
 	//ensure(SwapChainInfo.imageExtent.width >= SurfProperties.minImageExtent.width && SwapChainInfo.imageExtent.width <= SurfProperties.maxImageExtent.width);
 	//ensure(SwapChainInfo.imageExtent.height >= SurfProperties.minImageExtent.height && SwapChainInfo.imageExtent.height <= SurfProperties.maxImageExtent.height);
-	UE_LOG(LogVulkanRHI, Verbose, TEXT("Creating new VK swapchain with format %d, color space %d"), static_cast<uint32>(SwapChainInfo.imageFormat), static_cast<uint32>(SwapChainInfo.imageColorSpace));
+	static bool bPrintSwapchainCreationInfo = true;
+	if (bPrintSwapchainCreationInfo)
+	{
+		UE_LOG(LogVulkanRHI, Log, TEXT("Creating new VK swapchain with format %d, color space %d, num images %d"), static_cast<uint32>(SwapChainInfo.imageFormat), static_cast<uint32>(SwapChainInfo.imageColorSpace), static_cast<uint32>(SwapChainInfo.minImageCount));
+#if WITH_EDITOR
+		bPrintSwapchainCreationInfo = false;
+#endif
+	}
 
 	VERIFYVULKANRESULT_EXPANDED(VulkanRHI::vkCreateSwapchainKHR(Device.GetInstanceHandle(), &SwapChainInfo, VULKAN_CPU_ALLOCATOR, &SwapChain));
 
diff --git a/Engine/Source/Runtime/VulkanRHI/Private/VulkanViewport.cpp b/Engine/Source/Runtime/VulkanRHI/Private/VulkanViewport.cpp
--- a/Engine/Source/Runtime/VulkanRHI/Private/VulkanViewport.cpp
+++ b/Engine/Source/Runtime/VulkanRHI/Private/VulkanViewport.cpp
@@ -145,7 +145,7 @@ FVulkanViewport::FVulkanViewport(FVulkanDynamicRHI* InRHI, FVulkanDevice* InDevi
 
 	if (FVulkanPlatform::SupportsStandardSwapchain())
 	{
-		for (int32 Index = 0; Index < NUM_BUFFERS; ++Index)
+		for (int32 Index = 0, NumBuffers = RenderingDoneSemaphores.Num(); Index < NumBuffers; ++Index)
 		{
 			RenderingDoneSemaphores[Index] = new VulkanRHI::FSemaphore(*InDevice);
 			RenderingDoneSemaphores[Index]->AddRef();
@@ -165,7 +165,7 @@ FVulkanViewport::~FVulkanViewport()
 	
 	if (FVulkanPlatform::SupportsStandardSwapchain())
 	{
-		for (int32 Index = 0; Index < NUM_BUFFERS; ++Index)
+		for (int32 Index = 0, NumBuffers = RenderingDoneSemaphores.Num(); Index < NumBuffers; ++Index)
 		{
 			RenderingDoneSemaphores[Index]->Release();
 
@@ -584,7 +584,10 @@ void FVulkanViewport::CreateSwapchain(FVulkanSwapChainRecreateInfo* RecreateInfo
 			RecreateInfo
 		);
 
-		checkf(Images.Num() == NUM_BUFFERS, TEXT("Actual Num: %i"), Images.Num());
+		checkf(Images.Num() >= NUM_BUFFERS, TEXT("We wanted at least %i images, actual Num: %i"), NUM_BUFFERS, Images.Num());
+		BackBufferImages.SetNum(Images.Num());
+		RenderingDoneSemaphores.SetNum(Images.Num());
+		TextureViews.SetNum(Images.Num());
 
 		FVulkanCmdBuffer* CmdBuffer = Device->GetImmediateContext().GetCommandBufferManager()->GetUploadCmdBuffer();
 		ensure(CmdBuffer->IsOutsideRenderPass());
@@ -671,7 +674,7 @@ void FVulkanViewport::DestroySwapchain(FVulkanSwapChainRecreateInfo* RecreateInf
 		
 	if (FVulkanPlatform::SupportsStandardSwapchain() && SwapChain)
 	{
-		for (int32 Index = 0; Index < NUM_BUFFERS; ++Index)
+		for (int32 Index = 0, NumBuffers = BackBufferImages.Num(); Index < NumBuffers; ++Index)
 		{
 			TextureViews[Index].Destroy(*Device);
 			Device->NotifyDeletedImage(BackBufferImages[Index]);
diff --git a/Engine/Source/Runtime/VulkanRHI/Private/VulkanViewport.h b/Engine/Source/Runtime/VulkanRHI/Private/VulkanViewport.h
--- a/Engine/Source/Runtime/VulkanRHI/Private/VulkanViewport.h
+++ b/Engine/Source/Runtime/VulkanRHI/Private/VulkanViewport.h
@@ -81,9 +81,10 @@ public:
 	}
 
 protected:
-	TStaticArray<VkImage, NUM_BUFFERS> BackBufferImages;
-	TStaticArray<VulkanRHI::FSemaphore*, NUM_BUFFERS> RenderingDoneSemaphores;
-	TStaticArray<FVulkanTextureView, NUM_BUFFERS> TextureViews;
+	// NUM_BUFFERS don't have to match exactly as the driver can require a minimum number larger than NUM_BUFFERS. Provide some slack
+	TArray<VkImage, TInlineAllocator<NUM_BUFFERS*2>> BackBufferImages;
+	TArray<VulkanRHI::FSemaphore*, TInlineAllocator<NUM_BUFFERS*2>> RenderingDoneSemaphores;
+	TArray<FVulkanTextureView, TInlineAllocator<NUM_BUFFERS*2>> TextureViews;
 	TRefCountPtr<FVulkanBackBuffer> RHIBackBuffer;
 
 	// 'Dummy' back buffer
