Vulkan: fix capturing images of odd sizes with frame grabber (UE-81539).

- Should also fix UE-79632 (virtual camera display).

diff --git a/Engine/Source/Runtime/VulkanRHI/Private/VulkanRenderTarget.cpp b/Engine/Source/Runtime/VulkanRHI/Private/VulkanRenderTarget.cpp
--- a/Engine/Source/Runtime/VulkanRHI/Private/VulkanRenderTarget.cpp
+++ b/Engine/Source/Runtime/VulkanRHI/Private/VulkanRenderTarget.cpp
@@ -932,8 +932,25 @@ void FVulkanDynamicRHI::RHIMapStagingSurface(FRHITexture* TextureRHI,void*& OutD
 	check(TextureRHI2D);
 	FVulkanTexture2D* Texture2D = ResourceCast(TextureRHI2D);
 
+	int32 Pitch = Texture2D->GetSizeX();
+	if (ensureMsgf(Texture2D->Surface.GetTiling() == VK_IMAGE_TILING_LINEAR, TEXT("RHIMapStagingSurface() called with the texture is non-linear tiling %d, the result will likely be garbled."), Texture2D->Surface.GetTiling()))
+	{
+		// Pitch can be only retrieved from linear textures.
+		VkImageSubresource ImageSubResource;
+		FMemory::Memzero(ImageSubResource);
+
+		ImageSubResource.aspectMask = Texture2D->Surface.GetFullAspectMask();
+		ImageSubResource.mipLevel = 0;
+		ImageSubResource.arrayLayer = 0;
+
+		VkSubresourceLayout SubResourceLayout;
+		VulkanRHI::vkGetImageSubresourceLayout(Device->GetInstanceHandle(), Texture2D->Surface.Image, &ImageSubResource, &SubResourceLayout);
+
+		int32 BytesPerPixel = GetNumBitsPerPixel(Texture2D->Surface.StorageFormat) / 8;
+		Pitch = SubResourceLayout.rowPitch / BytesPerPixel;
+	}
 
-	OutWidth = Texture2D->GetSizeX();
+	OutWidth = Pitch;
 	OutHeight = Texture2D->GetSizeY();
 
 	OutData = Texture2D->Surface.GetMappedPointer();
