Vulkan: cosmetic errata for RHIMapStagingSurface()

- Better ensure message with more info.

diff --git a/Engine/Source/Runtime/VulkanRHI/Private/VulkanRenderTarget.cpp b/Engine/Source/Runtime/VulkanRHI/Private/VulkanRenderTarget.cpp
--- a/Engine/Source/Runtime/VulkanRHI/Private/VulkanRenderTarget.cpp
+++ b/Engine/Source/Runtime/VulkanRHI/Private/VulkanRenderTarget.cpp
@@ -933,7 +933,7 @@ void FVulkanDynamicRHI::RHIMapStagingSurface(FRHITexture* TextureRHI,void*& OutD
 	FVulkanTexture2D* Texture2D = ResourceCast(TextureRHI2D);
 
 	int32 Pitch = Texture2D->GetSizeX();
-	if (ensureMsgf(Texture2D->Surface.GetTiling() == VK_IMAGE_TILING_LINEAR, TEXT("RHIMapStagingSurface() called with the texture is non-linear tiling %d, the result will likely be garbled."), Texture2D->Surface.GetTiling()))
+	if (ensureMsgf(Texture2D->Surface.GetTiling() == VK_IMAGE_TILING_LINEAR, TEXT("RHIMapStagingSurface() called with a %d x %d texture in non-linear tiling %d, the result will likely be garbled."), static_cast<int32>(Texture2D->GetSizeX()), static_cast<int32>(Texture2D->GetSizeY()), static_cast<int32>(Texture2D->Surface.GetTiling())))
 	{
 		// Pitch can be only retrieved from linear textures.
 		VkImageSubresource ImageSubResource;
