Vulkan mobile multiview support

diff --git a/Engine/Plugins/Runtime/Oculus/OculusVR/Source/OculusHMD/Private/OculusHMD.cpp b/Engine/Plugins/Runtime/Oculus/OculusVR/Source/OculusHMD/Private/OculusHMD.cpp
--- a/Engine/Plugins/Runtime/Oculus/OculusVR/Source/OculusHMD/Private/OculusHMD.cpp
+++ b/Engine/Plugins/Runtime/Oculus/OculusVR/Source/OculusHMD/Private/OculusHMD.cpp
@@ -1274,8 +1274,7 @@ namespace OculusHMD
 
 			if (TextureSetProxy.IsValid())
 			{
-				OutTargetableTexture = TextureSetProxy->GetTexture2D();
-				OutShaderResourceTexture = TextureSetProxy->GetTexture2D();
+				OutTargetableTexture = OutShaderResourceTexture = TextureSetProxy->GetTexture2DArray() ? TextureSetProxy->GetTexture2DArray() : TextureSetProxy->GetTexture2D();
 				bNeedReAllocateViewportRenderTarget = false;
 				return true;
 			}
diff --git a/Engine/Plugins/Runtime/Oculus/OculusVR/Source/OculusHMD/Private/OculusHMD_TextureSetProxy.h b/Engine/Plugins/Runtime/Oculus/OculusVR/Source/OculusHMD/Private/OculusHMD_TextureSetProxy.h
--- a/Engine/Plugins/Runtime/Oculus/OculusVR/Source/OculusHMD/Private/OculusHMD_TextureSetProxy.h
+++ b/Engine/Plugins/Runtime/Oculus/OculusVR/Source/OculusHMD/Private/OculusHMD_TextureSetProxy.h
@@ -23,6 +23,7 @@ public:
 
 	FRHITexture* GetTexture() const { return RHITexture.GetReference(); }
 	FRHITexture2D* GetTexture2D() const { return RHITexture->GetTexture2D(); }
+	FRHITexture2DArray* GetTexture2DArray() const { return RHITexture->GetTexture2DArray(); }
 	FRHITextureCube* GetTextureCube() const { return RHITexture->GetTextureCube(); }
 	uint32 GetSwapChainLength() const { return (uint32) RHITextureSwapChain.Num(); }
 
diff --git a/Engine/Shaders/Private/MobileBasePassVertexShader.usf b/Engine/Shaders/Private/MobileBasePassVertexShader.usf
--- a/Engine/Shaders/Private/MobileBasePassVertexShader.usf
+++ b/Engine/Shaders/Private/MobileBasePassVertexShader.usf
@@ -41,9 +41,14 @@ void Main(
 	)
 {
 #if MOBILE_MULTI_VIEW
+#if (COMPILER_GLSL_ES2 || COMPILER_GLSL_ES3_1 || COMPILER_GLSL_ES3_1_EXT)
 	const int MultiViewId = int(ViewId);
 	ResolvedView = ResolveView(uint(MultiViewId));
 	Output.BasePassInterpolants.MultiViewId = float(MultiViewId);
+#else
+	ResolvedView = ResolveView(ViewId);
+	Output.BasePassInterpolants.MultiViewId = float(ViewId);
+#endif
 #else
 	ResolvedView = ResolveView();
 #endif
diff --git a/Engine/Source/Developer/VulkanShaderFormat/Private/VulkanBackend.cpp b/Engine/Source/Developer/VulkanShaderFormat/Private/VulkanBackend.cpp
--- a/Engine/Source/Developer/VulkanShaderFormat/Private/VulkanBackend.cpp
+++ b/Engine/Source/Developer/VulkanShaderFormat/Private/VulkanBackend.cpp
@@ -3345,7 +3345,7 @@ class FGenerateVulkanVisitor : public ir_visitor
 #endif		
 	}
 
-	void print_extensions(_mesa_glsl_parse_state* state, bool bUsesES31Extensions)
+	void print_extensions(_mesa_glsl_parse_state* state, bool bUsesES31Extensions, bool bShouldEmitMultiview)
 	{
 		if (bUsesES2TextureLODExtension)
 		{
@@ -3382,6 +3382,11 @@ class FGenerateVulkanVisitor : public ir_visitor
 			}
 
 		}
+
+		if (bShouldEmitMultiview && (state->target == vertex_shader || state->target == fragment_shader))
+		{
+			ralloc_asprintf_append(buffer, "#extension GL_EXT_multiview : enable\n");
+		}
 	}
 
 public:
@@ -3513,9 +3518,11 @@ public:
 		print_layout(state);
 		buffer = 0;
 
+		const bool bShouldEmitMultiview = strstr(signature, "gl_ViewIndex") != nullptr;
+
 		char* Extensions = ralloc_asprintf(mem_ctx, "");
 		buffer = &Extensions;
-		print_extensions(state, state->language_version == 310);
+		print_extensions(state, state->language_version == 310, bShouldEmitMultiview);
 		if (state->bSeparateShaderObjects && !state->bGenerateES)
 		{
 			switch (state->target)
@@ -4070,6 +4077,7 @@ struct FSystemValue
 /** Vertex shader system values. */
 static FSystemValue VertexSystemValueTable[] =
 {
+	{ "SV_ViewID", glsl_type::int_type, "gl_ViewIndex", ir_var_in, false, false, false, false },
 	{ "SV_VertexID", glsl_type::int_type, "gl_VertexIndex", ir_var_in, false, false, false, false },
 	{ "SV_InstanceID", glsl_type::int_type, "gl_InstanceIndex", ir_var_in, false, false, false, false },
 	{ "SV_Position", glsl_type::vec4_type, "gl_Position", ir_var_out, false, false, true, false },
@@ -4079,6 +4087,7 @@ static FSystemValue VertexSystemValueTable[] =
 /** Pixel shader system values. */
 static FSystemValue PixelSystemValueTable[] =
 {
+	{ "SV_ViewID", glsl_type::int_type, "gl_ViewIndex", ir_var_in, false, false, false, false },
 	{ "SV_Depth", glsl_type::float_type, "gl_FragDepth", ir_var_out, false, false, false, false },
 	{ "SV_Position", glsl_type::vec4_type, "gl_FragCoord", ir_var_in, true, false, false, false },
 	{ "SV_IsFrontFace", glsl_type::bool_type, "gl_FrontFacing", ir_var_in, false, false, true, false },
@@ -4094,6 +4103,7 @@ static FSystemValue PixelSystemValueTable[] =
 /** Geometry shader system values. */
 static FSystemValue GeometrySystemValueTable[] =
 {
+	{ "SV_ViewID", glsl_type::int_type, "gl_ViewIndex", ir_var_in, false, false, false, false },
 	{ "SV_VertexID", glsl_type::int_type, "gl_VertexID", ir_var_in, false, false, false, false },
 	{ "SV_InstanceID", glsl_type::int_type, "gl_InstanceID", ir_var_in, false, false, false, false },
 	{ "SV_Position", glsl_type::vec4_type, "gl_Position", ir_var_in, false, true, true, false },
@@ -4108,6 +4118,7 @@ static FSystemValue GeometrySystemValueTable[] =
 /** Hull shader system values. */
 static FSystemValue HullSystemValueTable[] =
 {
+	{ "SV_ViewID", glsl_type::int_type, "gl_ViewIndex", ir_var_in, false, false, false, false },
 	{ "SV_OutputControlPointID", glsl_type::int_type, "gl_InvocationID", ir_var_in, false, false, false, false },
 	{ NULL, NULL, NULL, ir_var_auto, false, false, false, false }
 };
@@ -4115,7 +4126,7 @@ static FSystemValue HullSystemValueTable[] =
 /** Domain shader system values. */
 static FSystemValue DomainSystemValueTable[] =
 {
-
+	{ "SV_ViewID", glsl_type::int_type, "gl_ViewIndex", ir_var_in, false, false, false, false },
 	// TODO : SV_DomainLocation has types float2 or float3 depending on the input topology
 	{ "SV_Position", glsl_type::vec4_type, "gl_Position", ir_var_in, false, true, true, false },
 	{ "SV_Position", glsl_type::vec4_type, "gl_Position", ir_var_out, false, false, true, false },
@@ -4351,6 +4362,18 @@ static ir_rvalue* GenShaderInputSemantic(
 		}
 	}
 
+	if (Variable == NULL && (Frequency == HSF_VertexShader || Frequency == HSF_PixelShader))
+	{
+		if (FCStringAnsi::Stricmp(Semantic, "SV_ViewId") == 0)
+		{
+			Variable = new(ParseState)ir_variable(
+				Type,
+				ralloc_asprintf(ParseState, "gl_ViewIndex"),
+				ir_var_in
+			);
+		}
+	}
+
 	if (Variable)
 	{
 		// Up to this point, variables aren't contained in structs
diff --git a/Engine/Source/Developer/VulkanShaderFormat/Private/VulkanShaderCompiler.cpp b/Engine/Source/Developer/VulkanShaderFormat/Private/VulkanShaderCompiler.cpp
--- a/Engine/Source/Developer/VulkanShaderFormat/Private/VulkanShaderCompiler.cpp
+++ b/Engine/Source/Developer/VulkanShaderFormat/Private/VulkanShaderCompiler.cpp
@@ -1864,6 +1864,7 @@ void DoCompileVulkanShader(const FShaderCompilerInput& Input, FShaderCompilerOut
 	const bool bHasRealUBs = HasRealUBs(Version);
 	const bool bIsSM5 = (Version == EVulkanShaderVersion::SM5 || Version == EVulkanShaderVersion::SM5_NOUB);
 	const bool bIsSM4 = (Version == EVulkanShaderVersion::SM4 || Version == EVulkanShaderVersion::SM4_NOUB);
+	const bool bIsMobile = (Version == EVulkanShaderVersion::ES3_1 || Version == EVulkanShaderVersion::ES3_1_ANDROID || Version == EVulkanShaderVersion::ES3_1_NOUB || Version == EVulkanShaderVersion::ES3_1_ANDROID_NOUB);
 
 	const EHlslShaderFrequency FrequencyTable[] =
 	{
@@ -1892,7 +1893,7 @@ void DoCompileVulkanShader(const FShaderCompilerInput& Input, FShaderCompilerOut
 	EHlslCompileTarget HlslCompilerTargetES = HCT_FeatureLevelES3_1Ext;
 	AdditionalDefines.SetDefine(TEXT("COMPILER_HLSLCC"), 1);
 	AdditionalDefines.SetDefine(TEXT("COMPILER_VULKAN"), 1);
-	if (Version == EVulkanShaderVersion::ES3_1 || Version == EVulkanShaderVersion::ES3_1_ANDROID || Version == EVulkanShaderVersion::ES3_1_NOUB || Version == EVulkanShaderVersion::ES3_1_ANDROID_NOUB)
+	if(bIsMobile)
 	{
 		HlslCompilerTarget = HCT_FeatureLevelES3_1Ext;
 		HlslCompilerTargetES = HCT_FeatureLevelES3_1Ext;
@@ -1944,8 +1945,11 @@ void DoCompileVulkanShader(const FShaderCompilerInput& Input, FShaderCompilerOut
 			// The preprocessing stage will add any relevant errors.
 			return;
 		}
+	}
 
-		// Disable instanced stereo until supported for Vulkan
+	if (!PreprocessedShaderSource.Contains(TEXT("SV_ViewID")))
+	{
+		// Disable instanced stereo if not requested by the shader
 		StripInstancedStereo(PreprocessedShaderSource);
 	}
 
diff --git a/Engine/Source/Runtime/Engine/Private/SceneView.cpp b/Engine/Source/Runtime/Engine/Private/SceneView.cpp
--- a/Engine/Source/Runtime/Engine/Private/SceneView.cpp
+++ b/Engine/Source/Runtime/Engine/Private/SceneView.cpp
@@ -736,7 +736,7 @@ FSceneView::FSceneView(const FSceneViewInitOptions& InitOptions)
 	bShouldBindInstancedViewUB = bIsInstancedStereoEnabled || bIsMobileMultiViewEnabled;
 
 	// If the device doesn't support mobile multi-view, disable it.
-	bIsMobileMultiViewEnabled = bIsMobileMultiViewEnabled && GSupportsMobileMultiView;
+	bIsMobileMultiViewEnabled = bIsMobileMultiViewEnabled && GSupportsMobileMultiView && StereoPass != eSSP_FULL;
 
 	SetupAntiAliasingMethod();
 
diff --git a/Engine/Source/Runtime/Engine/Private/ShaderCompiler/ShaderCompiler.cpp b/Engine/Source/Runtime/Engine/Private/ShaderCompiler/ShaderCompiler.cpp
--- a/Engine/Source/Runtime/Engine/Private/ShaderCompiler/ShaderCompiler.cpp
+++ b/Engine/Source/Runtime/Engine/Private/ShaderCompiler/ShaderCompiler.cpp
@@ -3025,8 +3025,8 @@ void GlobalBeginCompileShader(
 		Input.Environment.SetDefine(TEXT("INSTANCED_STEREO"), bIsInstancedStereo);
 		Input.Environment.SetDefine(TEXT("MULTI_VIEW"), bIsInstancedStereo && bIsMultiViewCVar && RHISupportsMultiView(ShaderPlatform));
 
-		const bool bIsAndroidGLES = RHISupportsMobileMultiView(ShaderPlatform);
-		Input.Environment.SetDefine(TEXT("MOBILE_MULTI_VIEW"), bIsMobileMultiViewCVar && bIsAndroidGLES);
+		const bool bSupportsMobileMultiview = RHISupportsMobileMultiView(ShaderPlatform);
+		Input.Environment.SetDefine(TEXT("MOBILE_MULTI_VIEW"), bIsMobileMultiViewCVar && bSupportsMobileMultiview);
 
 		// Throw a warning if we are silently disabling ISR due to missing platform support.
 		if (bIsInstancedStereoCVar && !bIsInstancedStereo && !GShaderCompilingManager->AreWarningsSuppressed(ShaderPlatform))
diff --git a/Engine/Source/Runtime/RHI/Public/RHICommandList.h b/Engine/Source/Runtime/RHI/Public/RHICommandList.h
--- a/Engine/Source/Runtime/RHI/Public/RHICommandList.h
+++ b/Engine/Source/Runtime/RHI/Public/RHICommandList.h
@@ -1380,8 +1380,8 @@ struct FRHICommandCopyToResolveTarget final : public FRHICommand<FRHICommandCopy
 	{
 		ensure(SourceTexture);
 		ensure(DestTexture);
-		ensure(SourceTexture->GetTexture2D() || SourceTexture->GetTexture3D() || SourceTexture->GetTextureCube());
-		ensure(DestTexture->GetTexture2D() || DestTexture->GetTexture3D() || DestTexture->GetTextureCube());
+		ensure(SourceTexture->GetTexture2D() || SourceTexture->GetTexture3D() || SourceTexture->GetTextureCube() || SourceTexture->GetTexture2DArray());
+		ensure(DestTexture->GetTexture2D() || DestTexture->GetTexture3D() || DestTexture->GetTextureCube() || DestTexture->GetTexture2DArray());
 	}
 	RHI_API void Execute(FRHICommandListBase& CmdList);
 };
@@ -2400,6 +2400,8 @@ public:
 			{
 				GraphicsPSOInit.RenderTargetFormats[i] = PSOContext.CachedRenderTargets[i].Texture->GetFormat();
 				GraphicsPSOInit.RenderTargetFlags[i] = PSOContext.CachedRenderTargets[i].Texture->GetFlags();
+				const FRHITexture2DArray* TextureArray = PSOContext.CachedRenderTargets[i].Texture->GetTexture2DArray();
+				GraphicsPSOInit.bMultiView = TextureArray && TextureArray->GetSizeZ() > 1;
 			}
 			else
 			{
@@ -2416,6 +2418,8 @@ public:
 		{
 			GraphicsPSOInit.DepthStencilTargetFormat = PSOContext.CachedDepthStencilTarget.Texture->GetFormat();
 			GraphicsPSOInit.DepthStencilTargetFlag = PSOContext.CachedDepthStencilTarget.Texture->GetFlags();
+			const FRHITexture2DArray* TextureArray = PSOContext.CachedDepthStencilTarget.Texture->GetTexture2DArray();
+			GraphicsPSOInit.bMultiView = TextureArray && TextureArray->GetSizeZ() > 1;
 		}
 		else
 		{
diff --git a/Engine/Source/Runtime/RHI/Public/RHIDefinitions.h b/Engine/Source/Runtime/RHI/Public/RHIDefinitions.h
--- a/Engine/Source/Runtime/RHI/Public/RHIDefinitions.h
+++ b/Engine/Source/Runtime/RHI/Public/RHIDefinitions.h
@@ -1098,7 +1098,7 @@ inline bool RHIHasTiledGPU(const EShaderPlatform Platform)
 
 inline bool RHISupportsMobileMultiView(const EShaderPlatform Platform)
 {
-	return (Platform == EShaderPlatform::SP_OPENGL_ES3_1_ANDROID || Platform == EShaderPlatform::SP_OPENGL_ES2_ANDROID);
+	return (Platform == EShaderPlatform::SP_OPENGL_ES3_1_ANDROID || Platform == EShaderPlatform::SP_OPENGL_ES2_ANDROID || IsVulkanMobilePlatform(Platform));
 }
 
 inline bool RHISupportsDrawIndirect(const EShaderPlatform Platform)
diff --git a/Engine/Source/Runtime/RHI/Public/RHIResources.h b/Engine/Source/Runtime/RHI/Public/RHIResources.h
--- a/Engine/Source/Runtime/RHI/Public/RHIResources.h
+++ b/Engine/Source/Runtime/RHI/Public/RHIResources.h
@@ -742,7 +742,7 @@ public:
 		return FIntPoint(SizeX, SizeY);
 	}
 
-	virtual FIntVector GetSizeXYZ() const final override
+	virtual FIntVector GetSizeXYZ() const override
 	{
 		return FIntVector(SizeX, SizeY, 1);
 	}
@@ -753,39 +753,31 @@ private:
 	uint32 SizeY;
 };
 
-class RHI_API FRHITexture2DArray : public FRHITexture
+class RHI_API FRHITexture2DArray : public FRHITexture2D
 {
 public:
 	
 	/** Initialization constructor. */
 	FRHITexture2DArray(uint32 InSizeX,uint32 InSizeY,uint32 InSizeZ,uint32 InNumMips,EPixelFormat InFormat,uint32 InFlags, const FClearValueBinding& InClearValue)
-	: FRHITexture(InNumMips,1,InFormat,InFlags,NULL, InClearValue)
-	, SizeX(InSizeX)
-	, SizeY(InSizeY)
+	: FRHITexture2D(InSizeX, InSizeY, InNumMips,1,InFormat,InFlags, InClearValue)
 	, SizeZ(InSizeZ)
 	{}
 	
 	// Dynamic cast methods.
 	virtual FRHITexture2DArray* GetTexture2DArray() { return this; }
-	
-	/** @return The width of the textures in the array. */
-	uint32 GetSizeX() const { return SizeX; }
-	
-	/** @return The height of the texture in the array. */
-	uint32 GetSizeY() const { return SizeY; }
+
+	virtual FRHITexture2D* GetTexture2D() { return NULL; }
 
 	/** @return The number of textures in the array. */
 	uint32 GetSizeZ() const { return SizeZ; }
 
 	virtual FIntVector GetSizeXYZ() const final override
 	{
-		return FIntVector(SizeX, SizeY, SizeZ);
+		return FIntVector(GetSizeX(), GetSizeY(), SizeZ);
 	}
 
 private:
 
-	uint32 SizeX;
-	uint32 SizeY;
 	uint32 SizeZ;
 };
 
@@ -1811,6 +1803,7 @@ public:
 			DepthStencilState != rhs.DepthStencilState ||
 			ImmutableSamplerState != rhs.ImmutableSamplerState ||
 			bDepthBounds != rhs.bDepthBounds ||
+			bMultiView != rhs.bMultiView ||
 			PrimitiveType != rhs.PrimitiveType) 
 		{
 			return false;
@@ -1861,6 +1854,7 @@ public:
 		COMPARE_FIELD(RasterizerState)
 		COMPARE_FIELD(DepthStencilState)
 		COMPARE_FIELD(bDepthBounds)
+		COMPARE_FIELD(bMultiView)
 		COMPARE_FIELD(PrimitiveType)
 		COMPARE_FIELD_END;
 
@@ -1885,6 +1879,7 @@ public:
 		COMPARE_FIELD(RasterizerState)
 		COMPARE_FIELD(DepthStencilState)
 		COMPARE_FIELD(bDepthBounds)
+		COMPARE_FIELD(bMultiView)
 		COMPARE_FIELD(PrimitiveType)
 		COMPARE_FIELD_END;
 
@@ -1907,7 +1902,8 @@ public:
 	// as it is sometimes hashed and compared as raw bytes. Explicit padding is therefore required between
 	// all data members and at the end of the structure.
 	bool							bDepthBounds = false;
-	uint8							Padding[3] = {};
+	bool							bMultiView = false;
+	uint8							Padding[2] = {};
 
 	EPrimitiveType					PrimitiveType;
 };
@@ -2396,6 +2392,9 @@ struct FRHIRenderPassInfo
 	// Some RHIs need to know if this render pass is going to be reading and writing to the same texture in the case of generating mip maps for partial resource transitions
 	bool bGeneratingMips = false;
 
+	// If this render pass should be multiview
+	bool bMultiviewPass = false;
+
 	// Hint for some RHI's that renderpass will have specific sub-passes 
 	ESubpassHint SubpassHint = ESubpassHint::None;
 
diff --git a/Engine/Source/Runtime/Renderer/Private/MobileShadingRenderer.cpp b/Engine/Source/Runtime/Renderer/Private/MobileShadingRenderer.cpp
--- a/Engine/Source/Runtime/Renderer/Private/MobileShadingRenderer.cpp
+++ b/Engine/Source/Runtime/Renderer/Private/MobileShadingRenderer.cpp
@@ -407,6 +407,7 @@ void FMobileSceneRenderer::Render(FRHICommandListImmediate& RHICmdList)
 	SceneColorRenderPassInfo.SubpassHint = ESubpassHint::DepthReadSubpass;
 	SceneColorRenderPassInfo.NumOcclusionQueries = ComputeNumOcclusionQueriesToBatch();
 	SceneColorRenderPassInfo.bOcclusionQueries = SceneColorRenderPassInfo.NumOcclusionQueries != 0;
+	SceneColorRenderPassInfo.bMultiviewPass = View.bIsMobileMultiViewEnabled;
 	RHICmdList.BeginRenderPass(SceneColorRenderPassInfo, TEXT("SceneColorRendering"));
 
 	if (GIsEditor && !View.bIsSceneCapture)
diff --git a/Engine/Source/Runtime/VulkanRHI/Private/VulkanDevice.cpp b/Engine/Source/Runtime/VulkanRHI/Private/VulkanDevice.cpp
--- a/Engine/Source/Runtime/VulkanRHI/Private/VulkanDevice.cpp
+++ b/Engine/Source/Runtime/VulkanRHI/Private/VulkanDevice.cpp
@@ -897,7 +897,7 @@ void FVulkanDevice::InitGPU(int32 DeviceIndex)
 		DefaultSampler = ResourceCast(RHICreateSamplerState(Default).GetReference());
 
 		FRHIResourceCreateInfo CreateInfo;
-		DefaultImage = new FVulkanSurface(*this, VK_IMAGE_VIEW_TYPE_2D, PF_B8G8R8A8, 1, 1, 1, false, 0, 1, 1, TexCreate_RenderTargetable | TexCreate_ShaderResource, CreateInfo);
+		DefaultImage = new FVulkanSurface(*this, VK_IMAGE_VIEW_TYPE_2D, PF_B8G8R8A8, 1, 1, 1, 1, 1, 1, TexCreate_RenderTargetable | TexCreate_ShaderResource, CreateInfo);
 		DefaultTextureView.Create(*this, DefaultImage->Image, VK_IMAGE_VIEW_TYPE_2D, DefaultImage->GetFullAspectMask(), PF_B8G8R8A8, VK_FORMAT_B8G8R8A8_UNORM, 0, 1, 0, 1);
 	}
 }
diff --git a/Engine/Source/Runtime/VulkanRHI/Private/VulkanRHI.cpp b/Engine/Source/Runtime/VulkanRHI/Private/VulkanRHI.cpp
--- a/Engine/Source/Runtime/VulkanRHI/Private/VulkanRHI.cpp
+++ b/Engine/Source/Runtime/VulkanRHI/Private/VulkanRHI.cpp
@@ -1450,6 +1450,34 @@ static VkRenderPass CreateRenderPass(FVulkanDevice& InDevice, const FVulkanRende
 	CreateInfo.pSubpasses = SubpassDescriptions;
 	CreateInfo.dependencyCount = NumDependencies;
 	CreateInfo.pDependencies = SubpassDependencies;
+
+	/*
+	Bit mask that specifies which view rendering is broadcast to
+	0011 = Broadcast to first and second view (layer)
+	*/
+	const uint32_t ViewMask = 0b00000011;
+
+	/*
+	Bit mask that specifices correlation between views
+	An implementation may use this for optimizations (concurrent render)
+	*/
+	const uint32_t CorrelationMask = 0b00000011;
+
+	VkRenderPassMultiviewCreateInfo MultiviewInfo;
+	if (RTLayout.GetIsMultiView())
+	{
+		FMemory::Memzero(MultiviewInfo);
+		MultiviewInfo.sType = VK_STRUCTURE_TYPE_RENDER_PASS_MULTIVIEW_CREATE_INFO;
+		MultiviewInfo.pNext = nullptr;
+		MultiviewInfo.subpassCount = 1;
+		MultiviewInfo.pViewMasks = &ViewMask;
+		MultiviewInfo.dependencyCount = 0;
+		MultiviewInfo.pViewOffsets = nullptr;
+		MultiviewInfo.correlationMaskCount = 1;
+		MultiviewInfo.pCorrelationMasks = &CorrelationMask;
+
+		CreateInfo.pNext = &MultiviewInfo;
+	}
 	
 	VkRenderPass RenderPassHandle;
 	VERIFYVULKANRESULT_EXPANDED(VulkanRHI::vkCreateRenderPass(InDevice.GetInstanceHandle(), &CreateInfo, VULKAN_CPU_ALLOCATOR, &RenderPassHandle));
diff --git a/Engine/Source/Runtime/VulkanRHI/Private/VulkanRHIPrivate.h b/Engine/Source/Runtime/VulkanRHI/Private/VulkanRHIPrivate.h
--- a/Engine/Source/Runtime/VulkanRHI/Private/VulkanRHIPrivate.h
+++ b/Engine/Source/Runtime/VulkanRHI/Private/VulkanRHIPrivate.h
@@ -120,10 +120,8 @@ public:
 	inline bool GetHasResolveAttachments() const { return bHasResolveAttachments != 0; }
 	inline uint32 GetNumAttachmentDescriptions() const { return NumAttachmentDescriptions; }
 	inline uint32 GetNumSamples() const { return NumSamples; }
-	inline uint32 GetNumUsedClearValues() const
-	{
-		return NumUsedClearValues;
-	}
+	inline uint32 GetNumUsedClearValues() const { return NumUsedClearValues; }
+	inline bool GetIsMultiView() const { return bIsMultiView != 0; }
 
 	inline const VkAttachmentReference* GetColorAttachmentReferences() const { return NumColorAttachments > 0 ? ColorReferences : nullptr; }
 	inline const VkAttachmentReference* GetResolveAttachmentReferences() const { return bHasResolveAttachments ? ResolveReferences : nullptr; }
@@ -147,6 +145,7 @@ protected:
 	uint8 NumSamples;
 	uint8 NumUsedClearValues;
 	ESubpassHint SubpassHint = ESubpassHint::None;
+	uint8 bIsMultiView;
 
 	// Hash for a compatible RenderPass
 	uint32 RenderPassCompatibleHash = 0;
@@ -173,6 +172,7 @@ protected:
 		Extent.Extent3D.width = 0;
 		Extent.Extent3D.height = 0;
 		Extent.Extent3D.depth = 0;
+		bIsMultiView = 0;
 	}
 
 	bool bCalculatedHash = false;
diff --git a/Engine/Source/Runtime/VulkanRHI/Private/VulkanRenderTarget.cpp b/Engine/Source/Runtime/VulkanRHI/Private/VulkanRenderTarget.cpp
--- a/Engine/Source/Runtime/VulkanRHI/Private/VulkanRenderTarget.cpp
+++ b/Engine/Source/Runtime/VulkanRHI/Private/VulkanRenderTarget.cpp
@@ -705,9 +705,11 @@ void FVulkanCommandListContext::RHICopyToResolveTarget(FTextureRHIParamRef Sourc
 	};
 
 	FRHITexture2D* SourceTexture2D = SourceTextureRHI->GetTexture2D();
+	FRHITexture2DArray* SourceTexture2DArray = SourceTextureRHI->GetTexture2DArray();
 	FRHITexture3D* SourceTexture3D = SourceTextureRHI->GetTexture3D();
 	FRHITextureCube* SourceTextureCube = SourceTextureRHI->GetTextureCube();
 	FRHITexture2D* DestTexture2D = DestTextureRHI->GetTexture2D();
+	FRHITexture2DArray* DestTexture2DArray = DestTextureRHI->GetTexture2DArray();
 	FRHITexture3D* DestTexture3D = DestTextureRHI->GetTexture3D();
 	FRHITextureCube* DestTextureCube = DestTextureRHI->GetTextureCube();
 	FVulkanCmdBuffer* CmdBuffer = CommandBufferManager->GetActiveCmdBuffer();
@@ -748,6 +750,15 @@ void FVulkanCommandListContext::RHICopyToResolveTarget(FTextureRHIParamRef Sourc
 			CopyImage(TransitionAndLayoutManager, CmdBuffer, VulkanSourceTexture->Surface, VulkanDestTexture->Surface, 1, 1, InResolveParams);
 		}
 	}
+	else if (SourceTexture2DArray && DestTexture2DArray)
+	{
+		FVulkanTexture2DArray* VulkanSourceTexture = (FVulkanTexture2DArray*)SourceTexture2DArray;
+		FVulkanTexture2DArray* VulkanDestTexture = (FVulkanTexture2DArray*)DestTexture2DArray;
+		if (VulkanSourceTexture->Surface.Image != VulkanDestTexture->Surface.Image)
+		{
+			CopyImage(TransitionAndLayoutManager, CmdBuffer, VulkanSourceTexture->Surface, VulkanDestTexture->Surface, VulkanDestTexture->GetSizeZ(), VulkanSourceTexture->GetSizeZ(), InResolveParams);
+		}
+	} 
 	else 
 	{
 		checkf(false, TEXT("Using unsupported Resolve combination"));
@@ -1787,6 +1798,7 @@ struct FRenderPassCompatibleHashableStruct
 	}
 
 	uint8						NumAttachments;
+	uint8						bIsMultiview;
 	uint8						NumSamples;
 	uint8						SubpassHint;
 	// +1 for DepthStencil
@@ -1820,6 +1832,7 @@ FVulkanRenderTargetLayout::FVulkanRenderTargetLayout(FVulkanDevice& InDevice, co
 	, bHasResolveAttachments(false)
 	, NumSamples(0)
 	, NumUsedClearValues(0)
+	, bIsMultiView(0)
 {
 	FMemory::Memzero(ColorReferences);
 	FMemory::Memzero(DepthStencilReference);
@@ -1978,7 +1991,7 @@ FVulkanRenderTargetLayout::FVulkanRenderTargetLayout(FVulkanDevice& InDevice, co
 			bSetExtent = true;
 			Extent.Extent3D.width = Texture->Surface.Width;
 			Extent.Extent3D.height = Texture->Surface.Height;
-			Extent.Extent3D.depth = 1;
+			Extent.Extent3D.depth = Texture->Surface.NumArrayLevels;
 		}
 	}
 
@@ -1986,6 +1999,7 @@ FVulkanRenderTargetLayout::FVulkanRenderTargetLayout(FVulkanDevice& InDevice, co
 	CompatibleHashInfo.SubpassHint = 0;
 
 	CompatibleHashInfo.NumSamples = NumSamples;
+	CompatibleHashInfo.bIsMultiview = bIsMultiView;
 
 	RenderPassCompatibleHash = FCrc::MemCrc32(&CompatibleHashInfo, sizeof(CompatibleHashInfo));
 	RenderPassFullHash = FCrc::MemCrc32(&FullHashInfo, sizeof(FullHashInfo), RenderPassCompatibleHash);
@@ -2001,6 +2015,7 @@ FVulkanRenderTargetLayout::FVulkanRenderTargetLayout(FVulkanDevice& InDevice, co
 	, bHasResolveAttachments(false)
 	, NumSamples(0)
 	, NumUsedClearValues(0)
+	, bIsMultiView(RPInfo.bMultiviewPass)
 {
 	FMemory::Memzero(ColorReferences);
 	FMemory::Memzero(DepthStencilReference);
@@ -2014,6 +2029,8 @@ FVulkanRenderTargetLayout::FVulkanRenderTargetLayout(FVulkanDevice& InDevice, co
 
 	bool bSetExtent = false;
 	bool bFoundClearOp = false;
+	bool bMultiviewRenderTargets = false;
+
 	int32 NumColorRenderTargets = RPInfo.GetNumColorRenderTargets();
 	for (int32 Index = 0; Index < NumColorRenderTargets; ++Index)
 	{
@@ -2038,6 +2055,9 @@ FVulkanRenderTargetLayout::FVulkanRenderTargetLayout(FVulkanDevice& InDevice, co
 		ensure(!NumSamples || NumSamples == ColorEntry.RenderTarget->GetNumSamples());
 		NumSamples = ColorEntry.RenderTarget->GetNumSamples();
 
+		ensure( !bMultiviewRenderTargets || Texture->Surface.NumArrayLevels > 1 );
+		bMultiviewRenderTargets = Texture->Surface.NumArrayLevels > 1;
+
 		VkAttachmentDescription& CurrDesc = Desc[NumAttachmentDescriptions];
 		CurrDesc.samples = static_cast<VkSampleCountFlagBits>(NumSamples);
 		CurrDesc.format = UEToVkTextureFormat(ColorEntry.RenderTarget->GetFormat(), (Texture->Surface.UEFlags & TexCreate_SRGB) == TexCreate_SRGB);
@@ -2130,6 +2150,9 @@ FVulkanRenderTargetLayout::FVulkanRenderTargetLayout(FVulkanDevice& InDevice, co
 
 		bHasDepthStencil = true;
 
+		ensure(!bMultiviewRenderTargets || Texture->Surface.NumArrayLevels > 1);
+		bMultiviewRenderTargets = Texture->Surface.NumArrayLevels > 1;
+
 		if (bSetExtent)
 		{
 			// Depth can be greater or equal to color
@@ -2141,7 +2164,7 @@ FVulkanRenderTargetLayout::FVulkanRenderTargetLayout(FVulkanDevice& InDevice, co
 			bSetExtent = true;
 			Extent.Extent3D.width = Texture->Surface.Width;
 			Extent.Extent3D.height = Texture->Surface.Height;
-			Extent.Extent3D.depth = 1;
+			Extent.Extent3D.depth = Texture->Surface.NumArrayLevels;
 		}
 	}
 
@@ -2149,6 +2172,12 @@ FVulkanRenderTargetLayout::FVulkanRenderTargetLayout(FVulkanDevice& InDevice, co
 	CompatibleHashInfo.SubpassHint = (uint8)RPInfo.SubpassHint;
 
 	CompatibleHashInfo.NumSamples = NumSamples;
+	CompatibleHashInfo.bIsMultiview = bIsMultiView;
+
+	if (bIsMultiView && !bMultiviewRenderTargets)
+	{
+		UE_LOG(LogVulkan, Error, TEXT("Non multiview textures on a multiview layout!"));
+	}
 
 	RenderPassCompatibleHash = FCrc::MemCrc32(&CompatibleHashInfo, sizeof(CompatibleHashInfo));
 	RenderPassFullHash = FCrc::MemCrc32(&FullHashInfo, sizeof(FullHashInfo), RenderPassCompatibleHash);
@@ -2163,6 +2192,7 @@ FVulkanRenderTargetLayout::FVulkanRenderTargetLayout(const FGraphicsPipelineStat
 	, bHasResolveAttachments(false)
 	, NumSamples(0)
 	, NumUsedClearValues(0)
+	, bIsMultiView(0)
 {
 	FMemory::Memzero(ColorReferences);
 	FMemory::Memzero(DepthStencilReference);
@@ -2176,6 +2206,7 @@ FVulkanRenderTargetLayout::FVulkanRenderTargetLayout(const FGraphicsPipelineStat
 
 	bool bSetExtent = false;
 	bool bFoundClearOp = false;
+	bIsMultiView = Initializer.bMultiView;
 	NumSamples = Initializer.NumSamples;
 	for (uint32 Index = 0; Index < Initializer.RenderTargetsEnabled; ++Index)
 	{
@@ -2278,6 +2309,7 @@ FVulkanRenderTargetLayout::FVulkanRenderTargetLayout(const FGraphicsPipelineStat
 	CompatibleHashInfo.SubpassHint = (uint8)Initializer.SubpassHint;
 
 	CompatibleHashInfo.NumSamples = NumSamples;
+	CompatibleHashInfo.bIsMultiview = bIsMultiView;
 
 	RenderPassCompatibleHash = FCrc::MemCrc32(&CompatibleHashInfo, sizeof(CompatibleHashInfo));
 	RenderPassFullHash = FCrc::MemCrc32(&FullHashInfo, sizeof(FullHashInfo), RenderPassCompatibleHash);
diff --git a/Engine/Source/Runtime/VulkanRHI/Private/VulkanResources.h b/Engine/Source/Runtime/VulkanRHI/Private/VulkanResources.h
--- a/Engine/Source/Runtime/VulkanRHI/Private/VulkanResources.h
+++ b/Engine/Source/Runtime/VulkanRHI/Private/VulkanResources.h
@@ -271,7 +271,7 @@ public:
 		VkImageViewType ResourceType,
 		EPixelFormat InFormat,
 		uint32 SizeX, uint32 SizeY, uint32 SizeZ,
-		bool bArray, uint32 ArraySize,
+		uint32 ArraySize,
 		uint32 NumMips,
 		uint32 NumSamples,
 		uint32 UEFlags,
@@ -282,12 +282,12 @@ public:
 		bool bForceLinearTexture = false);
 
 	FVulkanSurface(FVulkanDevice& Device, VkImageViewType ResourceType, EPixelFormat Format,
-					uint32 SizeX, uint32 SizeY, uint32 SizeZ, bool bArray, uint32 ArraySize,
+					uint32 SizeX, uint32 SizeY, uint32 SizeZ, uint32 ArraySize,
 					uint32 NumMips, uint32 NumSamples, uint32 UEFlags, const FRHIResourceCreateInfo& CreateInfo);
 
 	// Constructor for externally owned Image
 	FVulkanSurface(FVulkanDevice& Device, VkImageViewType ResourceType, EPixelFormat Format,
-					uint32 SizeX, uint32 SizeY, uint32 SizeZ, uint32 NumMips, uint32 NumSamples,
+					uint32 SizeX, uint32 SizeY, uint32 SizeZ, uint32 ArraySize, uint32 NumMips, uint32 NumSamples,
 					VkImage InImage, uint32 UEFlags, const FRHIResourceCreateInfo& CreateInfo);
 
 	virtual ~FVulkanSurface();
@@ -482,9 +482,9 @@ struct FVulkanTextureBase : public FVulkanBaseShaderResource
 		return OutTexture;
 	}
 
-	FVulkanTextureBase(FVulkanDevice& Device, VkImageViewType ResourceType, EPixelFormat Format, uint32 SizeX, uint32 SizeY, uint32 SizeZ, bool bArray, uint32 ArraySize, uint32 NumMips, uint32 NumSamples, uint32 UEFlags, const FRHIResourceCreateInfo& CreateInfo);
-	FVulkanTextureBase(FVulkanDevice& Device, VkImageViewType ResourceType, EPixelFormat Format, uint32 SizeX, uint32 SizeY, uint32 SizeZ, uint32 NumMips, uint32 NumSamples, uint32 NumSamplesTileMem, VkImage InImage, VkDeviceMemory InMem, uint32 UEFlags, const FRHIResourceCreateInfo& CreateInfo = FRHIResourceCreateInfo());
-	FVulkanTextureBase(FVulkanDevice& Device, VkImageViewType ResourceType, EPixelFormat Format, uint32 SizeX, uint32 SizeY, uint32 SizeZ, uint32 NumMips, uint32 NumSamples, VkImage InImage, VkDeviceMemory InMem, FSamplerYcbcrConversionInitializer& ConversionInitializer, uint32 UEFlags, const FRHIResourceCreateInfo& CreateInfo = FRHIResourceCreateInfo());
+	FVulkanTextureBase(FVulkanDevice& Device, VkImageViewType ResourceType, EPixelFormat Format, uint32 SizeX, uint32 SizeY, uint32 SizeZ, uint32 ArraySize, uint32 NumMips, uint32 NumSamples, uint32 UEFlags, const FRHIResourceCreateInfo& CreateInfo);
+	FVulkanTextureBase(FVulkanDevice& Device, VkImageViewType ResourceType, EPixelFormat Format, uint32 SizeX, uint32 SizeY, uint32 SizeZ, uint32 ArraySize, uint32 NumMips, uint32 NumSamples, uint32 NumSamplesTileMem, VkImage InImage, VkDeviceMemory InMem, uint32 UEFlags, const FRHIResourceCreateInfo& CreateInfo = FRHIResourceCreateInfo());
+	FVulkanTextureBase(FVulkanDevice& Device, VkImageViewType ResourceType, EPixelFormat Format, uint32 SizeX, uint32 SizeY, uint32 SizeZ, uint32 ArraySize, uint32 NumMips, uint32 NumSamples, VkImage InImage, VkDeviceMemory InMem, FSamplerYcbcrConversionInitializer& ConversionInitializer, uint32 UEFlags, const FRHIResourceCreateInfo& CreateInfo = FRHIResourceCreateInfo());
 
 	virtual ~FVulkanTextureBase();
 
@@ -694,7 +694,7 @@ class FVulkanTextureReference : public FRHITextureReference, public FVulkanTextu
 public:
 	explicit FVulkanTextureReference(FVulkanDevice& Device, FLastRenderTimeContainer* InLastRenderTime)
 	:	FRHITextureReference(InLastRenderTime)
-	,	FVulkanTextureBase(Device, VK_IMAGE_VIEW_TYPE_MAX_ENUM, PF_Unknown, 0, 0, 0, 1, 1, 1, VK_NULL_HANDLE, VK_NULL_HANDLE, 0)
+	,	FVulkanTextureBase(Device, VK_IMAGE_VIEW_TYPE_MAX_ENUM, PF_Unknown, 0, 0, 0, 1, 1, 1, 1, VK_NULL_HANDLE, VK_NULL_HANDLE, 0)
 	{}
 
 	// IRefCountedObject interface.
diff --git a/Engine/Source/Runtime/VulkanRHI/Private/VulkanTexture.cpp b/Engine/Source/Runtime/VulkanRHI/Private/VulkanTexture.cpp
--- a/Engine/Source/Runtime/VulkanRHI/Private/VulkanTexture.cpp
+++ b/Engine/Source/Runtime/VulkanRHI/Private/VulkanTexture.cpp
@@ -179,7 +179,7 @@ VkImage FVulkanSurface::CreateImage(
 	VkImageViewType ResourceType,
 	EPixelFormat InFormat,
 	uint32 SizeX, uint32 SizeY, uint32 SizeZ,
-	bool bArray, uint32 ArraySize,
+	uint32 ArraySize,
 	uint32 NumMips,
 	uint32 NumSamples,
 	uint32 UEFlags,
@@ -248,7 +248,7 @@ VkImage FVulkanSurface::CreateImage(
 	ImageCreateInfo.extent.depth = ResourceType == VK_IMAGE_VIEW_TYPE_3D ? SizeZ : 1;
 	ImageCreateInfo.mipLevels = NumMips;
 	uint32 LayerCount = (ResourceType == VK_IMAGE_VIEW_TYPE_CUBE || ResourceType == VK_IMAGE_VIEW_TYPE_CUBE_ARRAY) ? 6 : 1;
-	ImageCreateInfo.arrayLayers = (bArray ? ArraySize : 1) * LayerCount;
+	ImageCreateInfo.arrayLayers = ArraySize * LayerCount;
 	check(ImageCreateInfo.arrayLayers <= DeviceProperties.limits.maxImageArrayLayers);
 
 	ImageCreateInfo.flags = (ResourceType == VK_IMAGE_VIEW_TYPE_CUBE || ResourceType == VK_IMAGE_VIEW_TYPE_CUBE_ARRAY) ? VK_IMAGE_CREATE_CUBE_COMPATIBLE_BIT : 0;
@@ -484,7 +484,7 @@ struct FRHICommandOnDestroyImage final : public FRHICommand<FRHICommandOnDestroy
 
 
 FVulkanSurface::FVulkanSurface(FVulkanDevice& InDevice, VkImageViewType ResourceType, EPixelFormat InFormat,
-								uint32 SizeX, uint32 SizeY, uint32 SizeZ, bool bArray, uint32 ArraySize, uint32 InNumMips,
+								uint32 SizeX, uint32 SizeY, uint32 SizeZ, uint32 ArraySize, uint32 InNumMips,
 								uint32 InNumSamples, uint32 InUEFlags, const FRHIResourceCreateInfo& CreateInfo)
 	: Device(&InDevice)
 	, Image(VK_NULL_HANDLE)
@@ -508,12 +508,12 @@ FVulkanSurface::FVulkanSurface(FVulkanDevice& InDevice, VkImageViewType Resource
 	VkImageCreateInfo ImageCreateInfo;	// Zeroed inside CreateImage
 	Image = FVulkanSurface::CreateImage(InDevice, ResourceType,
 		InFormat, SizeX, SizeY, SizeZ,
-		bArray, ArraySize, NumMips, NumSamples, UEFlags, MemoryRequirements,
+		ArraySize, NumMips, NumSamples, UEFlags, MemoryRequirements,
 		&StorageFormat, &ViewFormat,
 		&ImageCreateInfo);
 
 	uint32 LayerCount = (ResourceType == VK_IMAGE_VIEW_TYPE_CUBE || ResourceType == VK_IMAGE_VIEW_TYPE_CUBE_ARRAY) ? 6 : 1;
-	NumArrayLevels = (bArray ? ArraySize : 1) * LayerCount;
+	NumArrayLevels = ArraySize * LayerCount;
 
 	FullAspectMask = VulkanRHI::GetAspectMaskFromUEFormat(PixelFormat, true, true);
 	PartialAspectMask = VulkanRHI::GetAspectMaskFromUEFormat(PixelFormat, false, true);
@@ -569,7 +569,7 @@ FVulkanSurface::FVulkanSurface(FVulkanDevice& InDevice, VkImageViewType Resource
 
 // This is usually used for the framebuffer image
 FVulkanSurface::FVulkanSurface(FVulkanDevice& InDevice, VkImageViewType ResourceType, EPixelFormat InFormat,
-								uint32 SizeX, uint32 SizeY, uint32 SizeZ, uint32 InNumMips, uint32 InNumSamples,
+								uint32 SizeX, uint32 SizeY, uint32 SizeZ, uint32 ArraySize, uint32 InNumMips, uint32 InNumSamples,
 								VkImage InImage, uint32 InUEFlags, const FRHIResourceCreateInfo& CreateInfo)
 	: Device(&InDevice)
 	, Image(InImage)
@@ -581,7 +581,7 @@ FVulkanSurface::FVulkanSurface(FVulkanDevice& InDevice, VkImageViewType Resource
 	, PixelFormat(InFormat)
 	, UEFlags(InUEFlags)
 	, MemProps(VK_MEMORY_PROPERTY_DEVICE_LOCAL_BIT)
-	, NumArrayLevels(0)
+	, NumArrayLevels(ArraySize)
 	, Tiling(VK_IMAGE_TILING_MAX_ENUM)	// Can be expanded to a per-platform definition
 	, ViewType(ResourceType)
 	, bIsImageOwner(false)
@@ -1158,6 +1158,7 @@ void FVulkanDynamicRHI::InternalUnlockTexture2D(bool bFromRenderingThread, FText
 	EPixelFormat Format = Texture->Surface.PixelFormat;
 	uint32 MipWidth = FMath::Max<uint32>(Texture->Surface.Width >> MipIndex, GPixelFormats[Format].BlockSizeX);
 	uint32 MipHeight = FMath::Max<uint32>(Texture->Surface.Height >> MipIndex, GPixelFormats[Format].BlockSizeY);
+	uint32 LayerCount = Texture->Surface.GetNumberOfArrayLevels();
 
 	VkImageSubresourceRange SubresourceRange;
 	FMemory::Memzero(SubresourceRange);
@@ -1176,7 +1177,7 @@ void FVulkanDynamicRHI::InternalUnlockTexture2D(bool bFromRenderingThread, FText
 	Region.imageSubresource.aspectMask = VK_IMAGE_ASPECT_COLOR_BIT;
 	Region.imageSubresource.mipLevel = MipIndex;
 	//Region.imageSubresource.baseArrayLayer = 0;
-	Region.imageSubresource.layerCount = 1;
+	Region.imageSubresource.layerCount = LayerCount;
 	Region.imageExtent.width = MipWidth;
 	Region.imageExtent.height = MipHeight;
 	Region.imageExtent.depth = 1;
@@ -1592,18 +1593,19 @@ void FVulkanTextureView::Destroy(FVulkanDevice& Device)
 	}
 }
 
-FVulkanTextureBase::FVulkanTextureBase(FVulkanDevice& Device, VkImageViewType ResourceType, EPixelFormat InFormat, uint32 SizeX, uint32 SizeY, uint32 SizeZ, bool bArray, uint32 ArraySize, uint32 NumMips, uint32 NumSamples, uint32 UEFlags, const FRHIResourceCreateInfo& CreateInfo)
+FVulkanTextureBase::FVulkanTextureBase(FVulkanDevice& Device, VkImageViewType ResourceType, EPixelFormat InFormat, uint32 SizeX, uint32 SizeY, uint32 SizeZ, uint32 ArraySize, uint32 NumMips, uint32 NumSamples, uint32 UEFlags, const FRHIResourceCreateInfo& CreateInfo)
 	#if !VULKAN_USE_MSAA_RESOLVE_ATTACHMENTS
-	: Surface(Device, ResourceType, InFormat, SizeX, SizeY, SizeZ, bArray, ArraySize, NumMips, NumSamples, UEFlags, CreateInfo)
+	: Surface(Device, ResourceType, InFormat, SizeX, SizeY, SizeZ, ArraySize, NumMips, NumSamples, UEFlags, CreateInfo)
 	, PartialView(nullptr)
 	#else
-	: Surface(Device, ResourceType, InFormat, SizeX, SizeY, SizeZ, bArray, ArraySize, NumMips, (UEFlags & TexCreate_DepthStencilTargetable) ? NumSamples : 1, UEFlags, CreateInfo)
+	: Surface(Device, ResourceType, InFormat, SizeX, SizeY, SizeZ, ArraySize, NumMips, (UEFlags & TexCreate_DepthStencilTargetable) ? NumSamples : 1, UEFlags, CreateInfo)
 	, PartialView(nullptr)
 	, MSAASurface(nullptr)
 	#endif
 	, bIsAliased(false)
 {
 	LLM_SCOPE_VULKAN(ELLMTagVulkan::VulkanTextures);
+	const bool bArray = ResourceType == VK_IMAGE_VIEW_TYPE_1D_ARRAY || ResourceType == VK_IMAGE_VIEW_TYPE_2D_ARRAY || ResourceType == VK_IMAGE_VIEW_TYPE_CUBE_ARRAY;
 	if (Surface.ViewFormat == VK_FORMAT_UNDEFINED)
 	{
 		Surface.StorageFormat = UEToVkTextureFormat(InFormat, false);
@@ -1620,7 +1622,7 @@ FVulkanTextureBase::FVulkanTextureBase(FVulkanDevice& Device, VkImageViewType Re
 	// Create MSAA surface. The surface above is the resolve target
 	if (NumSamples > 1 && (UEFlags & TexCreate_RenderTargetable) && !(UEFlags & TexCreate_DepthStencilTargetable))
 	{
-		MSAASurface = new FVulkanSurface(Device, ResourceType, InFormat, SizeX, SizeY, SizeZ, /*bArray=*/ false, 1, NumMips, NumSamples, UEFlags, CreateInfo);
+		MSAASurface = new FVulkanSurface(Device, ResourceType, InFormat, SizeX, SizeY, SizeZ, 1, NumMips, NumSamples, UEFlags, CreateInfo);
 		if (ResourceType != VK_IMAGE_VIEW_TYPE_MAX_ENUM)
 		{
 			MSAAView.Create(Device, MSAASurface->Image, ResourceType, MSAASurface->GetFullAspectMask(), MSAASurface->PixelFormat, MSAASurface->ViewFormat, 0, FMath::Max(NumMips, 1u), 0, bArray ? FMath::Max(1u, ArraySize) : FMath::Max(1u, SizeZ));
@@ -1693,8 +1695,8 @@ FVulkanTextureBase::FVulkanTextureBase(FVulkanDevice& Device, VkImageViewType Re
 	}
 }
 
-FVulkanTextureBase::FVulkanTextureBase(FVulkanDevice& Device, VkImageViewType ResourceType, EPixelFormat Format, uint32 SizeX, uint32 SizeY, uint32 SizeZ, uint32 InNumMips, uint32 InNumSamples, uint32 InNumSamplesTileMem, VkImage InImage, VkDeviceMemory InMem, uint32 UEFlags, const FRHIResourceCreateInfo& CreateInfo)
-	: Surface(Device, ResourceType, Format, SizeX, SizeY, SizeZ, InNumMips, InNumSamples, InImage, UEFlags, CreateInfo)
+FVulkanTextureBase::FVulkanTextureBase(FVulkanDevice& Device, VkImageViewType ResourceType, EPixelFormat Format, uint32 SizeX, uint32 SizeY, uint32 SizeZ, uint32 ArraySize, uint32 InNumMips, uint32 InNumSamples, uint32 InNumSamplesTileMem, VkImage InImage, VkDeviceMemory InMem, uint32 UEFlags, const FRHIResourceCreateInfo& CreateInfo)
+	: Surface(Device, ResourceType, Format, SizeX, SizeY, SizeZ, ArraySize, InNumMips, InNumSamples, InImage, UEFlags, CreateInfo)
 	, PartialView(nullptr)
 	#if VULKAN_USE_MSAA_RESOLVE_ATTACHMENTS
 	, MSAASurface(nullptr)
@@ -1702,16 +1704,18 @@ FVulkanTextureBase::FVulkanTextureBase(FVulkanDevice& Device, VkImageViewType Re
 	, bIsAliased(false)
 {
 	check(InMem == VK_NULL_HANDLE);
+	const bool bArray = ResourceType == VK_IMAGE_VIEW_TYPE_1D_ARRAY || ResourceType == VK_IMAGE_VIEW_TYPE_2D_ARRAY || ResourceType == VK_IMAGE_VIEW_TYPE_CUBE_ARRAY;
+
 	if (ResourceType != VK_IMAGE_VIEW_TYPE_MAX_ENUM && Surface.Image != VK_NULL_HANDLE)
 	{
-		DefaultView.Create(Device, Surface.Image, ResourceType, Surface.GetFullAspectMask(), Format, Surface.ViewFormat, 0, FMath::Max(Surface.NumMips, 1u), 0, 1u);
+		DefaultView.Create(Device, Surface.Image, ResourceType, Surface.GetFullAspectMask(), Format, Surface.ViewFormat, 0, FMath::Max(Surface.NumMips, 1u), 0, bArray ? FMath::Max(1u, ArraySize) : FMath::Max(1u, SizeZ));
 	}
 
 #if VULKAN_USE_MSAA_RESOLVE_ATTACHMENTS
 	// Create MSAA surface. The surface above is the resolve target
 	if (InNumSamples == 1 && InNumSamplesTileMem > 1 && (UEFlags & TexCreate_RenderTargetable) && !(UEFlags & TexCreate_DepthStencilTargetable))
 	{
-		MSAASurface = new FVulkanSurface(Device, ResourceType, Format, SizeX, SizeY, SizeZ, /*bArray=*/ false, 1, InNumMips, InNumSamplesTileMem, UEFlags, CreateInfo);
+		MSAASurface = new FVulkanSurface(Device, ResourceType, Format, SizeX, SizeY, SizeZ, 1, InNumMips, InNumSamplesTileMem, UEFlags, CreateInfo);
 		if (ResourceType != VK_IMAGE_VIEW_TYPE_MAX_ENUM)
 		{
 			MSAAView.Create(Device, MSAASurface->Image, ResourceType, MSAASurface->GetFullAspectMask(), MSAASurface->PixelFormat, MSAASurface->ViewFormat, 0, FMath::Max(InNumMips, 1u), 0, FMath::Max(1u, SizeZ));
@@ -1726,12 +1730,12 @@ FVulkanTextureBase::FVulkanTextureBase(FVulkanDevice& Device, VkImageViewType Re
 	else
 	{
 		PartialView = new FVulkanTextureView;
-		PartialView->Create(Device, Surface.Image, Surface.ViewType, Surface.PartialAspectMask, Surface.PixelFormat, Surface.ViewFormat, 0, FMath::Max(InNumMips, 1u), 0, FMath::Max(1u, SizeZ));
+		PartialView->Create(Device, Surface.Image, Surface.ViewType, Surface.PartialAspectMask, Surface.PixelFormat, Surface.ViewFormat, 0, FMath::Max(InNumMips, 1u), 0, bArray ? FMath::Max(1u, ArraySize) : FMath::Max(1u, SizeZ));
 	}
 }
 
-FVulkanTextureBase::FVulkanTextureBase(FVulkanDevice& Device, VkImageViewType ResourceType, EPixelFormat Format, uint32 SizeX, uint32 SizeY, uint32 SizeZ, uint32 NumMips, uint32 NumSamples, VkImage InImage, VkDeviceMemory InMem, FSamplerYcbcrConversionInitializer& ConversionInitializer, uint32 UEFlags, const FRHIResourceCreateInfo& CreateInfo)
-	: Surface(Device, ResourceType, Format, SizeX, SizeY, SizeZ, NumMips, NumSamples, InImage, UEFlags, CreateInfo)
+FVulkanTextureBase::FVulkanTextureBase(FVulkanDevice& Device, VkImageViewType ResourceType, EPixelFormat Format, uint32 SizeX, uint32 SizeY, uint32 SizeZ, uint32 ArraySize, uint32 NumMips, uint32 NumSamples, VkImage InImage, VkDeviceMemory InMem, FSamplerYcbcrConversionInitializer& ConversionInitializer, uint32 UEFlags, const FRHIResourceCreateInfo& CreateInfo)
+	: Surface(Device, ResourceType, Format, SizeX, SizeY, SizeZ, ArraySize, NumMips, NumSamples, InImage, UEFlags, CreateInfo)
 	, PartialView(nullptr)
 #if VULKAN_USE_MSAA_RESOLVE_ATTACHMENTS
 	, MSAASurface(nullptr)
@@ -1739,13 +1743,14 @@ FVulkanTextureBase::FVulkanTextureBase(FVulkanDevice& Device, VkImageViewType Re
 	, bIsAliased(false)
 {
 	check(InMem == VK_NULL_HANDLE);
+	const bool bArray = ResourceType == VK_IMAGE_VIEW_TYPE_1D_ARRAY || ResourceType == VK_IMAGE_VIEW_TYPE_2D_ARRAY || ResourceType == VK_IMAGE_VIEW_TYPE_CUBE_ARRAY;
 
 	Surface.ViewFormat = ConversionInitializer.Format;
 	Surface.StorageFormat = ConversionInitializer.Format;
 
 	if (ResourceType != VK_IMAGE_VIEW_TYPE_MAX_ENUM && Surface.Image != VK_NULL_HANDLE)
 	{
-		DefaultView.Create(Device, Surface.Image, ResourceType, Surface.GetFullAspectMask(), Format, Surface.ViewFormat, 0, FMath::Max(Surface.NumMips, 1u), 0, 1u, ConversionInitializer);
+		DefaultView.Create(Device, Surface.Image, ResourceType, Surface.GetFullAspectMask(), Format, Surface.ViewFormat, 0, FMath::Max(Surface.NumMips, 1u), 0, bArray ? FMath::Max(1u, ArraySize) : FMath::Max(1u, SizeZ), ConversionInitializer);
 	}
 
 	// No MSAA support
@@ -1759,7 +1764,7 @@ FVulkanTextureBase::FVulkanTextureBase(FVulkanDevice& Device, VkImageViewType Re
 	else
 	{
 		PartialView = new FVulkanTextureView;
-		PartialView->Create(Device, Surface.Image, Surface.ViewType, Surface.PartialAspectMask, Surface.PixelFormat, Surface.ViewFormat, 0, FMath::Max(NumMips, 1u), 0, FMath::Max(1u, SizeZ), ConversionInitializer);
+		PartialView->Create(Device, Surface.Image, Surface.ViewType, Surface.PartialAspectMask, Surface.PixelFormat, Surface.ViewFormat, 0, FMath::Max(NumMips, 1u), 0, bArray ? FMath::Max(1u, ArraySize) : FMath::Max(1u, SizeZ), ConversionInitializer);
 	}
 
 	// Since this is provided from an external image, assume it's ready for read
@@ -1835,19 +1840,19 @@ void FVulkanTextureBase::DestroyViews()
 
 FVulkanTexture2D::FVulkanTexture2D(FVulkanDevice& Device, EPixelFormat InFormat, uint32 SizeX, uint32 SizeY, uint32 NumMips, uint32 NumSamples, uint32 UEFlags, const FRHIResourceCreateInfo& CreateInfo)
 :	FRHITexture2D(SizeX, SizeY, FMath::Max(NumMips, 1u), NumSamples, InFormat, UEFlags, CreateInfo.ClearValueBinding)
-,	FVulkanTextureBase(Device, VK_IMAGE_VIEW_TYPE_2D, InFormat, SizeX, SizeY, 1, /*bArray=*/ false, 1, FMath::Max(NumMips, 1u), NumSamples, UEFlags, CreateInfo)
+,	FVulkanTextureBase(Device, VK_IMAGE_VIEW_TYPE_2D, InFormat, SizeX, SizeY, 1, 1, FMath::Max(NumMips, 1u), NumSamples, UEFlags, CreateInfo)
 {
 }
 
 FVulkanTexture2D::FVulkanTexture2D(FVulkanDevice& Device, EPixelFormat Format, uint32 SizeX, uint32 SizeY, uint32 NumMips, uint32 NumSamples, uint32 NumSamplesTileMem, VkImage Image, uint32 UEFlags, const FRHIResourceCreateInfo& CreateInfo)
 :	FRHITexture2D(SizeX, SizeY, NumMips, NumSamples, Format, UEFlags, CreateInfo.ClearValueBinding)
-,	FVulkanTextureBase(Device, VK_IMAGE_VIEW_TYPE_2D, Format, SizeX, SizeY, 1, NumMips, NumSamples, NumSamplesTileMem, Image, VK_NULL_HANDLE, UEFlags)
+,	FVulkanTextureBase(Device, VK_IMAGE_VIEW_TYPE_2D, Format, SizeX, SizeY, 1, 1, NumMips, NumSamples, NumSamplesTileMem, Image, VK_NULL_HANDLE, UEFlags)
 {
 }
 
 FVulkanTexture2D::FVulkanTexture2D(FVulkanDevice& Device, EPixelFormat Format, uint32 SizeX, uint32 SizeY, uint32 NumMips, uint32 NumSamples, VkImage Image, FSamplerYcbcrConversionInitializer& ConversionInitializer, uint32 UEFlags, const FRHIResourceCreateInfo& CreateInfo)
 	: FRHITexture2D(SizeX, SizeY, NumMips, NumSamples, Format, UEFlags, CreateInfo.ClearValueBinding)
-	, FVulkanTextureBase(Device, VK_IMAGE_VIEW_TYPE_2D, Format, SizeX, SizeY, 1, NumMips, NumSamples, Image, VK_NULL_HANDLE, ConversionInitializer, UEFlags)
+	, FVulkanTextureBase(Device, VK_IMAGE_VIEW_TYPE_2D, Format, SizeX, SizeY, 1, 1, NumMips, NumSamples, Image, VK_NULL_HANDLE, ConversionInitializer, UEFlags)
 {
 }
 
@@ -1886,13 +1891,13 @@ FVulkanBackBuffer::~FVulkanBackBuffer()
 
 FVulkanTexture2DArray::FVulkanTexture2DArray(FVulkanDevice& Device, EPixelFormat Format, uint32 SizeX, uint32 SizeY, uint32 ArraySize, uint32 NumMips, uint32 Flags, FResourceBulkDataInterface* BulkData, const FClearValueBinding& InClearValue)
 	:	FRHITexture2DArray(SizeX, SizeY, ArraySize, NumMips, Format, Flags, InClearValue)
-	,	FVulkanTextureBase(Device, VK_IMAGE_VIEW_TYPE_2D_ARRAY, Format, SizeX, SizeY, 1, /*bArray=*/ true, ArraySize, NumMips, /*NumSamples=*/ 1, Flags, BulkData)
+	,	FVulkanTextureBase(Device, VK_IMAGE_VIEW_TYPE_2D_ARRAY, Format, SizeX, SizeY, 1, ArraySize, NumMips, /*NumSamples=*/ 1, Flags, BulkData)
 {
 }
 
 FVulkanTexture2DArray::FVulkanTexture2DArray(FVulkanDevice& Device, EPixelFormat Format, uint32 SizeX, uint32 SizeY, uint32 ArraySize, uint32 NumMips, VkImage Image, uint32 Flags, FResourceBulkDataInterface* BulkData, const FClearValueBinding& InClearValue)
 	:	FRHITexture2DArray(SizeX, SizeY, ArraySize, NumMips, Format, Flags, InClearValue)
-	,	FVulkanTextureBase(Device, VK_IMAGE_VIEW_TYPE_2D_ARRAY, Format, SizeX, SizeY, 1, NumMips, /*NumSamples=*/ 1, /*NumSamplesTileMem=*/ 1, Image, VK_NULL_HANDLE, Flags, BulkData)
+	,	FVulkanTextureBase(Device, VK_IMAGE_VIEW_TYPE_2D_ARRAY, Format, SizeX, SizeY, 1, ArraySize, NumMips, /*NumSamples=*/ 1, /*NumSamplesTileMem=*/ 1, Image, VK_NULL_HANDLE, Flags, BulkData)
 {
 }
 
@@ -1906,14 +1911,14 @@ void FVulkanTextureReference::SetReferencedTexture(FRHITexture* InTexture)
 FVulkanTextureCube::FVulkanTextureCube(FVulkanDevice& Device, EPixelFormat Format, uint32 Size, bool bArray, uint32 ArraySize, uint32 NumMips, uint32 Flags, FResourceBulkDataInterface* BulkData, const FClearValueBinding& InClearValue)
 :	 FRHITextureCube(Size, NumMips, Format, Flags, InClearValue)
 	//#todo-rco: Array/slices count
-,	FVulkanTextureBase(Device, bArray ? VK_IMAGE_VIEW_TYPE_CUBE_ARRAY : VK_IMAGE_VIEW_TYPE_CUBE, Format, Size, Size, 1, bArray, ArraySize, NumMips, /*NumSamples=*/ 1, Flags, BulkData)
+,	FVulkanTextureBase(Device, bArray ? VK_IMAGE_VIEW_TYPE_CUBE_ARRAY : VK_IMAGE_VIEW_TYPE_CUBE, Format, Size, Size, 1, ArraySize, NumMips, /*NumSamples=*/ 1, Flags, BulkData)
 {
 }
 
 FVulkanTextureCube::FVulkanTextureCube(FVulkanDevice& Device, EPixelFormat Format, uint32 Size, bool bArray, uint32 ArraySize, uint32 NumMips, VkImage Image, uint32 Flags, FResourceBulkDataInterface* BulkData, const FClearValueBinding& InClearValue)
 	:	 FRHITextureCube(Size, NumMips, Format, Flags, InClearValue)
 	//#todo-rco: Array/slices count
-	,	FVulkanTextureBase(Device, bArray ? VK_IMAGE_VIEW_TYPE_CUBE_ARRAY : VK_IMAGE_VIEW_TYPE_CUBE, Format, Size, Size, 1, NumMips, /*NumSamples=*/ 1, /*NumSamplesTileMem=*/ 1, Image, VK_NULL_HANDLE, Flags, BulkData)
+	,	FVulkanTextureBase(Device, bArray ? VK_IMAGE_VIEW_TYPE_CUBE_ARRAY : VK_IMAGE_VIEW_TYPE_CUBE, Format, Size, Size, 1, ArraySize, NumMips, /*NumSamples=*/ 1, /*NumSamplesTileMem=*/ 1, Image, VK_NULL_HANDLE, Flags, BulkData)
 {
 }
 
@@ -1928,7 +1933,7 @@ FVulkanTextureCube::~FVulkanTextureCube()
 
 FVulkanTexture3D::FVulkanTexture3D(FVulkanDevice& Device, EPixelFormat Format, uint32 SizeX, uint32 SizeY, uint32 SizeZ, uint32 NumMips, uint32 Flags, FResourceBulkDataInterface* BulkData, const FClearValueBinding& InClearValue)
 	: FRHITexture3D(SizeX, SizeY, SizeZ, NumMips, Format, Flags, InClearValue)
-	, FVulkanTextureBase(Device, VK_IMAGE_VIEW_TYPE_3D, Format, SizeX, SizeY, SizeZ, /*bArray=*/ false, 1, NumMips, /*NumSamples=*/ 1, Flags, BulkData)
+	, FVulkanTextureBase(Device, VK_IMAGE_VIEW_TYPE_3D, Format, SizeX, SizeY, SizeZ, 1, NumMips, /*NumSamples=*/ 1, Flags, BulkData)
 {
 }
 
@@ -2145,7 +2150,7 @@ static VkMemoryRequirements FindOrCalculateTexturePlatformSize(FVulkanDevice* De
 
 	// Create temporary image to measure the memory requirements
 	VkImage TmpImage = FVulkanSurface::CreateImage(*Device, ViewType,
-		PixelFormat, SizeX, SizeY, SizeZ, false, 0, NumMips, NumSamples,
+		PixelFormat, SizeX, SizeY, SizeZ, 1, NumMips, NumSamples,
 		Flags, MemReq, &InternalStorageFormat, &InternalViewFormat, &CreateInfo, false);
 
 	VulkanRHI::vkDestroyImage(Device->GetInstanceHandle(), TmpImage, VULKAN_CPU_ALLOCATOR);
diff --git a/Engine/Source/Runtime/VulkanRHI/Private/VulkanViewport.cpp b/Engine/Source/Runtime/VulkanRHI/Private/VulkanViewport.cpp
--- a/Engine/Source/Runtime/VulkanRHI/Private/VulkanViewport.cpp
+++ b/Engine/Source/Runtime/VulkanRHI/Private/VulkanViewport.cpp
@@ -345,9 +345,9 @@ FVulkanFramebuffer::FVulkanFramebuffer(FVulkanDevice& Device, const FRHISetRende
 		MipIndex = InRTInfo.ColorRenderTarget[Index].MipIndex;
 
 		FVulkanTextureView RTView;
-		if (Texture->Surface.GetViewType() == VK_IMAGE_VIEW_TYPE_2D)
+		if (Texture->Surface.GetViewType() == VK_IMAGE_VIEW_TYPE_2D || Texture->Surface.GetViewType() == VK_IMAGE_VIEW_TYPE_2D_ARRAY)
 		{
-			RTView.Create(*Texture->Surface.Device, Texture->Surface.Image, VK_IMAGE_VIEW_TYPE_2D, Texture->Surface.GetFullAspectMask(), Texture->Surface.PixelFormat, Texture->Surface.ViewFormat, MipIndex, 1, FMath::Max(0, (int32)InRTInfo.ColorRenderTarget[Index].ArraySliceIndex), 1, true);
+			RTView.Create(*Texture->Surface.Device, Texture->Surface.Image, Texture->Surface.GetViewType(), Texture->Surface.GetFullAspectMask(), Texture->Surface.PixelFormat, Texture->Surface.ViewFormat, MipIndex, 1, FMath::Max(0, (int32)InRTInfo.ColorRenderTarget[Index].ArraySliceIndex), Texture->Surface.GetNumberOfArrayLevels(), true);
 		}
 		else if (Texture->Surface.GetViewType() == VK_IMAGE_VIEW_TYPE_CUBE)
 		{
@@ -383,7 +383,7 @@ FVulkanFramebuffer::FVulkanFramebuffer(FVulkanDevice& Device, const FRHISetRende
 		check(Texture->PartialView);
 		PartialDepthTextureView = *Texture->PartialView;
 
-		ensure(Texture->Surface.GetViewType() == VK_IMAGE_VIEW_TYPE_2D || Texture->Surface.GetViewType() == VK_IMAGE_VIEW_TYPE_CUBE);
+		ensure(Texture->Surface.GetViewType() == VK_IMAGE_VIEW_TYPE_2D || Texture->Surface.GetViewType() == VK_IMAGE_VIEW_TYPE_2D_ARRAY || Texture->Surface.GetViewType() == VK_IMAGE_VIEW_TYPE_CUBE);
 		if (NumColorAttachments == 0 && Texture->Surface.GetViewType() == VK_IMAGE_VIEW_TYPE_CUBE)
 		{
 			FVulkanTextureView RTView;
