Vulkan MSAA resolve attachment support for mobile

diff --git a/Engine/Plugins/Runtime/Oculus/OculusVR/Source/OculusHMD/Private/OculusHMD_CustomPresent_Vulkan.cpp b/Engine/Plugins/Runtime/Oculus/OculusVR/Source/OculusHMD/Private/OculusHMD_CustomPresent_Vulkan.cpp
--- a/Engine/Plugins/Runtime/Oculus/OculusVR/Source/OculusHMD/Private/OculusHMD_CustomPresent_Vulkan.cpp
+++ b/Engine/Plugins/Runtime/Oculus/OculusVR/Source/OculusHMD/Private/OculusHMD_CustomPresent_Vulkan.cpp
@@ -126,7 +126,7 @@ FTextureRHIRef FVulkanCustomPresent::CreateTexture_RenderThread(uint32 InSizeX,
 		return DynamicRHI->RHICreateTexture2DFromResource(InFormat, InSizeX, InSizeY, InNumMips, InNumSamples, (VkImage) InTexture, InTexCreateFlags).GetReference();
 
 	case RRT_Texture2DArray:
-		return DynamicRHI->RHICreateTexture2DArrayFromResource(InFormat, InSizeX, InSizeY, 2, InNumMips, (VkImage) InTexture, InTexCreateFlags).GetReference();
+		return DynamicRHI->RHICreateTexture2DArrayFromResource(InFormat, InSizeX, InSizeY, 2, InNumMips, InNumSamples, (VkImage) InTexture, InTexCreateFlags).GetReference();
 
 	case RRT_TextureCube:
 		return DynamicRHI->RHICreateTextureCubeFromResource(InFormat, InSizeX, false, 1, InNumMips, (VkImage) InTexture, InTexCreateFlags).GetReference();
diff --git a/Engine/Source/Runtime/Apple/MetalRHI/Private/MetalTexture.cpp b/Engine/Source/Runtime/Apple/MetalRHI/Private/MetalTexture.cpp
--- a/Engine/Source/Runtime/Apple/MetalRHI/Private/MetalTexture.cpp
+++ b/Engine/Source/Runtime/Apple/MetalRHI/Private/MetalTexture.cpp
@@ -1768,7 +1768,7 @@ void FMetalDynamicRHI::RHICopySharedMips(FTexture2DRHIParamRef DestTexture2D,FTe
 	NOT_SUPPORTED("RHICopySharedMips");
 }
 
-FTexture2DArrayRHIRef FMetalDynamicRHI::RHICreateTexture2DArray(uint32 SizeX, uint32 SizeY, uint32 SizeZ, uint8 Format, uint32 NumMips, uint32 Flags, FRHIResourceCreateInfo& CreateInfo)
+FTexture2DArrayRHIRef FMetalDynamicRHI::RHICreateTexture2DArray(uint32 SizeX, uint32 SizeY, uint32 SizeZ, uint8 Format, uint32 NumMips, uint32 NumSamples, uint32 Flags, FRHIResourceCreateInfo& CreateInfo)
 {
 	@autoreleasepool {
 		return new FMetalTexture2DArray((EPixelFormat)Format, SizeX, SizeY, SizeZ, NumMips, Flags, CreateInfo.BulkData, CreateInfo.ClearValueBinding);
@@ -2391,13 +2391,13 @@ FTexture2DRHIRef FMetalDynamicRHI::RHICreateTexture2D_RenderThread(class FRHICom
 	}
 }
 
-FTexture2DArrayRHIRef FMetalDynamicRHI::RHICreateTexture2DArray_RenderThread(class FRHICommandListImmediate& RHICmdList, uint32 SizeX, uint32 SizeY, uint32 SizeZ, uint8 Format, uint32 NumMips, uint32 Flags, FRHIResourceCreateInfo& CreateInfo)
+FTexture2DArrayRHIRef FMetalDynamicRHI::RHICreateTexture2DArray_RenderThread(class FRHICommandListImmediate& RHICmdList, uint32 SizeX, uint32 SizeY, uint32 SizeZ, uint8 Format, uint32 NumMips, uint32 NumSamples, uint32 Flags, FRHIResourceCreateInfo& CreateInfo)
 {
 	@autoreleasepool {
 		FRHIResourceCreateInfo Info = CreateInfo;
 		bool const bIOSurfaceData = (CreateInfo.BulkData && CreateInfo.BulkData->GetResourceType() != FResourceBulkDataInterface::EBulkDataType::Default);
 		Info.BulkData = bIOSurfaceData ? CreateInfo.BulkData : nullptr;
-		FTexture2DArrayRHIRef Result = GDynamicRHI->RHICreateTexture2DArray(SizeX, SizeY, SizeZ, Format, NumMips, Flags, Info);
+		FTexture2DArrayRHIRef Result = GDynamicRHI->RHICreateTexture2DArray(SizeX, SizeY, SizeZ, Format, NumMips, NumSamples, Flags, Info);
 		if (CreateInfo.BulkData)
 		{
 			// upload existing bulkdata
diff --git a/Engine/Source/Runtime/Apple/MetalRHI/Public/MetalRHI.h b/Engine/Source/Runtime/Apple/MetalRHI/Public/MetalRHI.h
--- a/Engine/Source/Runtime/Apple/MetalRHI/Public/MetalRHI.h
+++ b/Engine/Source/Runtime/Apple/MetalRHI/Public/MetalRHI.h
@@ -94,7 +94,7 @@ public:
 	virtual FTexture2DRHIRef RHICreateTextureExternal2D(uint32 SizeX, uint32 SizeY, uint8 Format, uint32 NumMips, uint32 NumSamples, uint32 Flags, FRHIResourceCreateInfo& Info) final override;
 	virtual FTexture2DRHIRef RHIAsyncCreateTexture2D(uint32 SizeX, uint32 SizeY, uint8 Format, uint32 NumMips, uint32 Flags, void** InitialMipData, uint32 NumInitialMips) final override;
 	virtual void RHICopySharedMips(FTexture2DRHIParamRef DestTexture2D, FTexture2DRHIParamRef SrcTexture2D) final override;
-	virtual FTexture2DArrayRHIRef RHICreateTexture2DArray(uint32 SizeX, uint32 SizeY, uint32 SizeZ, uint8 Format, uint32 NumMips, uint32 Flags, FRHIResourceCreateInfo& CreateInfo) final override;
+	virtual FTexture2DArrayRHIRef RHICreateTexture2DArray(uint32 SizeX, uint32 SizeY, uint32 SizeZ, uint8 Format, uint32 NumMips, uint32 NumSamples, uint32 Flags, FRHIResourceCreateInfo& CreateInfo) final override;
 	virtual FTexture3DRHIRef RHICreateTexture3D(uint32 SizeX, uint32 SizeY, uint32 SizeZ, uint8 Format, uint32 NumMips, uint32 Flags, FRHIResourceCreateInfo& CreateInfo) final override;
 	virtual void RHIGetResourceInfo(FTextureRHIParamRef Ref, FRHIResourceInfo& OutInfo) final override;
 	virtual FShaderResourceViewRHIRef RHICreateShaderResourceView(FTexture2DRHIParamRef Texture2DRHI, uint8 MipLevel) final override;
@@ -165,7 +165,7 @@ public:
 	virtual void UpdateTexture3D_RenderThread(class FRHICommandListImmediate& RHICmdList, FTexture3DRHIParamRef Texture, uint32 MipIndex, const struct FUpdateTextureRegion3D& UpdateRegion, uint32 SourceRowPitch, uint32 SourceDepthPitch, const uint8* SourceData) final override;
 	
 	virtual FTexture2DRHIRef RHICreateTexture2D_RenderThread(class FRHICommandListImmediate& RHICmdList, uint32 SizeX, uint32 SizeY, uint8 Format, uint32 NumMips, uint32 NumSamples, uint32 Flags, FRHIResourceCreateInfo& CreateInfo) final override;
-	virtual FTexture2DArrayRHIRef RHICreateTexture2DArray_RenderThread(class FRHICommandListImmediate& RHICmdList, uint32 SizeX, uint32 SizeY, uint32 SizeZ, uint8 Format, uint32 NumMips, uint32 Flags, FRHIResourceCreateInfo& CreateInfo) final override;
+	virtual FTexture2DArrayRHIRef RHICreateTexture2DArray_RenderThread(class FRHICommandListImmediate& RHICmdList, uint32 SizeX, uint32 SizeY, uint32 SizeZ, uint8 Format, uint32 NumMips, uint32 NumSamples, uint32 Flags, FRHIResourceCreateInfo& CreateInfo) final override;
 	virtual FTexture3DRHIRef RHICreateTexture3D_RenderThread(class FRHICommandListImmediate& RHICmdList, uint32 SizeX, uint32 SizeY, uint32 SizeZ, uint8 Format, uint32 NumMips, uint32 Flags, FRHIResourceCreateInfo& CreateInfo) final override;
 	virtual FTextureCubeRHIRef RHICreateTextureCube_RenderThread(class FRHICommandListImmediate& RHICmdList, uint32 Size, uint8 Format, uint32 NumMips, uint32 Flags, FRHIResourceCreateInfo& CreateInfo) final override;
 	virtual FTextureCubeRHIRef RHICreateTextureCubeArray_RenderThread(class FRHICommandListImmediate& RHICmdList, uint32 Size, uint32 ArraySize, uint8 Format, uint32 NumMips, uint32 Flags, FRHIResourceCreateInfo& CreateInfo) final override;
diff --git a/Engine/Source/Runtime/Apple/MetalRHI/Public/MetalResources.h b/Engine/Source/Runtime/Apple/MetalRHI/Public/MetalResources.h
--- a/Engine/Source/Runtime/Apple/MetalRHI/Public/MetalResources.h
+++ b/Engine/Source/Runtime/Apple/MetalRHI/Public/MetalResources.h
@@ -585,7 +585,7 @@ public:
 
 	// Constructor, just calls base and Surface constructor
 	FMetalTexture2DArray(EPixelFormat Format, uint32 SizeX, uint32 SizeY, uint32 ArraySize, uint32 NumMips, uint32 Flags, FResourceBulkDataInterface* BulkData, const FClearValueBinding& InClearValue)
-		: FRHITexture2DArray(SizeX, SizeY, ArraySize, NumMips, Format, Flags, InClearValue)
+		: FRHITexture2DArray(SizeX, SizeY, ArraySize, NumMips, 1, Format, Flags, InClearValue)
 		, Surface(RRT_Texture2DArray, Format, SizeX, SizeY, 1, /*NumSamples=*/1, /*bArray=*/ true, ArraySize, NumMips, Flags, BulkData)
 	{
 	}
diff --git a/Engine/Source/Runtime/D3D12RHI/Private/D3D12RHIPrivate.h b/Engine/Source/Runtime/D3D12RHI/Private/D3D12RHIPrivate.h
--- a/Engine/Source/Runtime/D3D12RHI/Private/D3D12RHIPrivate.h
+++ b/Engine/Source/Runtime/D3D12RHI/Private/D3D12RHIPrivate.h
@@ -312,7 +312,7 @@ public:
 	virtual void RHICopySharedMips(FTexture2DRHIParamRef DestTexture2D, FTexture2DRHIParamRef SrcTexture2D) final override;
 	void RHIMultiGPULockstep(FRHIGPUMask GPUMask);
 	virtual void RHITransferTexture(FTexture2DRHIParamRef Texture, FIntRect Rect, uint32 SrcGPUIndex, uint32 DestGPUIndex, bool PullData) final override;
-	virtual FTexture2DArrayRHIRef RHICreateTexture2DArray(uint32 SizeX, uint32 SizeY, uint32 SizeZ, uint8 Format, uint32 NumMips, uint32 Flags, FRHIResourceCreateInfo& CreateInfo) override;
+	virtual FTexture2DArrayRHIRef RHICreateTexture2DArray(uint32 SizeX, uint32 SizeY, uint32 SizeZ, uint8 Format, uint32 NumMips, uint32 NumSamples, uint32 Flags, FRHIResourceCreateInfo& CreateInfo) override;
 	virtual FTexture3DRHIRef RHICreateTexture3D(uint32 SizeX, uint32 SizeY, uint32 SizeZ, uint8 Format, uint32 NumMips, uint32 Flags, FRHIResourceCreateInfo& CreateInfo) override;
 	virtual void RHIGetResourceInfo(FTextureRHIParamRef Ref, FRHIResourceInfo& OutInfo) override;
 	virtual FShaderResourceViewRHIRef RHICreateShaderResourceView(FTexture2DRHIParamRef Texture2DRHI, uint8 MipLevel) final override;
@@ -440,7 +440,7 @@ public:
 	virtual void UnlockTexture2D_RenderThread(class FRHICommandListImmediate& RHICmdList, FTexture2DRHIParamRef Texture, uint32 MipIndex, bool bLockWithinMiptail, bool bNeedsDefaultRHIFlush = true);
 	// 
 	virtual FTexture2DRHIRef RHICreateTexture2D_RenderThread(class FRHICommandListImmediate& RHICmdList, uint32 SizeX, uint32 SizeY, uint8 Format, uint32 NumMips, uint32 NumSamples, uint32 Flags, FRHIResourceCreateInfo& CreateInfo);
-	virtual FTexture2DArrayRHIRef RHICreateTexture2DArray_RenderThread(class FRHICommandListImmediate& RHICmdList, uint32 SizeX, uint32 SizeY, uint32 SizeZ, uint8 Format, uint32 NumMips, uint32 Flags, FRHIResourceCreateInfo& CreateInfo);
+	virtual FTexture2DArrayRHIRef RHICreateTexture2DArray_RenderThread(class FRHICommandListImmediate& RHICmdList, uint32 SizeX, uint32 SizeY, uint32 SizeZ, uint8 Format, uint32 NumMips, uint32 NumSamples, uint32 Flags, FRHIResourceCreateInfo& CreateInfo);
 	virtual FTexture3DRHIRef RHICreateTexture3D_RenderThread(class FRHICommandListImmediate& RHICmdList, uint32 SizeX, uint32 SizeY, uint32 SizeZ, uint8 Format, uint32 NumMips, uint32 Flags, FRHIResourceCreateInfo& CreateInfo);
 	virtual FUnorderedAccessViewRHIRef RHICreateUnorderedAccessView_RenderThread(class FRHICommandListImmediate& RHICmdList, FStructuredBufferRHIParamRef StructuredBuffer, bool bUseUAVCounter, bool bAppendBuffer);
 	virtual FUnorderedAccessViewRHIRef RHICreateUnorderedAccessView_RenderThread(class FRHICommandListImmediate& RHICmdList, FTextureRHIParamRef Texture, uint32 MipLevel);
diff --git a/Engine/Source/Runtime/D3D12RHI/Private/D3D12Texture.cpp b/Engine/Source/Runtime/D3D12RHI/Private/D3D12Texture.cpp
--- a/Engine/Source/Runtime/D3D12RHI/Private/D3D12Texture.cpp
+++ b/Engine/Source/Runtime/D3D12RHI/Private/D3D12Texture.cpp
@@ -1332,16 +1332,16 @@ void FD3D12DynamicRHI::RHICopySharedMips(FTexture2DRHIParamRef DestTexture2DRHI,
 	}
 }
 
-FTexture2DArrayRHIRef FD3D12DynamicRHI::RHICreateTexture2DArray_RenderThread(class FRHICommandListImmediate& RHICmdList, uint32 SizeX, uint32 SizeY, uint32 SizeZ, uint8 Format, uint32 NumMips, uint32 Flags, FRHIResourceCreateInfo& CreateInfo)
+FTexture2DArrayRHIRef FD3D12DynamicRHI::RHICreateTexture2DArray_RenderThread(class FRHICommandListImmediate& RHICmdList, uint32 SizeX, uint32 SizeY, uint32 SizeZ, uint8 Format, uint32 NumMips, uint32 NumSamples, uint32 Flags, FRHIResourceCreateInfo& CreateInfo)
 {
 	check(SizeZ >= 1);
-	return CreateD3D12Texture2D<FD3D12BaseTexture2DArray>(&RHICmdList, SizeX, SizeY, SizeZ, true, false, (EPixelFormat) Format, NumMips, 1, Flags, CreateInfo);
+	return CreateD3D12Texture2D<FD3D12BaseTexture2DArray>(&RHICmdList, SizeX, SizeY, SizeZ, true, false, (EPixelFormat) Format, NumMips, NumSamples, Flags, CreateInfo);
 }
 
-FTexture2DArrayRHIRef FD3D12DynamicRHI::RHICreateTexture2DArray(uint32 SizeX, uint32 SizeY, uint32 SizeZ, uint8 Format, uint32 NumMips, uint32 Flags, FRHIResourceCreateInfo& CreateInfo)
+FTexture2DArrayRHIRef FD3D12DynamicRHI::RHICreateTexture2DArray(uint32 SizeX, uint32 SizeY, uint32 SizeZ, uint8 Format, uint32 NumMips, uint32 NumSamples, uint32 Flags, FRHIResourceCreateInfo& CreateInfo)
 {
 	check(SizeZ >= 1);
-	return CreateD3D12Texture2D<FD3D12BaseTexture2DArray>(nullptr, SizeX, SizeY, SizeZ, true, false, (EPixelFormat) Format, NumMips, 1, Flags, CreateInfo);
+	return CreateD3D12Texture2D<FD3D12BaseTexture2DArray>(nullptr, SizeX, SizeY, SizeZ, true, false, (EPixelFormat) Format, NumMips, NumSamples, Flags, CreateInfo);
 }
 
 FTexture3DRHIRef FD3D12DynamicRHI::RHICreateTexture3D_RenderThread(class FRHICommandListImmediate& RHICmdList, uint32 SizeX, uint32 SizeY, uint32 SizeZ, uint8 Format, uint32 NumMips, uint32 Flags, FRHIResourceCreateInfo& CreateInfo)
diff --git a/Engine/Source/Runtime/D3D12RHI/Private/D3D12Texture.h b/Engine/Source/Runtime/D3D12RHI/Private/D3D12Texture.h
--- a/Engine/Source/Runtime/D3D12RHI/Private/D3D12Texture.h
+++ b/Engine/Source/Runtime/D3D12RHI/Private/D3D12Texture.h
@@ -409,7 +409,7 @@ class FD3D12BaseTexture2DArray : public FRHITexture2DArray, public FD3D12FastCle
 {
 public:
 	FD3D12BaseTexture2DArray(uint32 InSizeX, uint32 InSizeY, uint32 InSizeZ, uint32 InNumMips, uint32 InNumSamples, EPixelFormat InFormat, uint32 InFlags, const FClearValueBinding& InClearValue)
-		: FRHITexture2DArray(InSizeX, InSizeY, InSizeZ, InNumMips, InFormat, InFlags, InClearValue)
+		: FRHITexture2DArray(InSizeX, InSizeY, InSizeZ, InNumMips, InNumSamples, InFormat, InFlags, InClearValue)
 	{
 		check(InNumSamples == 1);
 	}
diff --git a/Engine/Source/Runtime/EmptyRHI/Private/EmptyTexture.cpp b/Engine/Source/Runtime/EmptyRHI/Private/EmptyTexture.cpp
--- a/Engine/Source/Runtime/EmptyRHI/Private/EmptyTexture.cpp
+++ b/Engine/Source/Runtime/EmptyRHI/Private/EmptyTexture.cpp
@@ -128,7 +128,7 @@ void FEmptyDynamicRHI::RHICopySharedMips(FTexture2DRHIParamRef DestTexture2D,FTe
 
 }
 
-FTexture2DArrayRHIRef FEmptyDynamicRHI::RHICreateTexture2DArray(uint32 SizeX, uint32 SizeY, uint32 SizeZ, uint8 Format, uint32 NumMips, uint32 Flags, FRHIResourceCreateInfo& CreateInfo)
+FTexture2DArrayRHIRef FEmptyDynamicRHI::RHICreateTexture2DArray(uint32 SizeX, uint32 SizeY, uint32 SizeZ, uint8 Format, uint32 NumMips, uint32 NumSamples, uint32 Flags, FRHIResourceCreateInfo& CreateInfo)
 {
 	return new FEmptyTexture2DArray((EPixelFormat)Format, SizeX, SizeY, SizeZ, NumMips, Flags, CreateInfo.BulkData);
 }
diff --git a/Engine/Source/Runtime/EmptyRHI/Public/EmptyRHI.h b/Engine/Source/Runtime/EmptyRHI/Public/EmptyRHI.h
--- a/Engine/Source/Runtime/EmptyRHI/Public/EmptyRHI.h
+++ b/Engine/Source/Runtime/EmptyRHI/Public/EmptyRHI.h
@@ -79,7 +79,7 @@ public:
 	virtual FTexture2DRHIRef RHICreateTexture2D(uint32 SizeX, uint32 SizeY, uint8 Format, uint32 NumMips, uint32 NumSamples, uint32 Flags, FRHIResourceCreateInfo& CreateInfo) final override;
 	virtual FTexture2DRHIRef RHIAsyncCreateTexture2D(uint32 SizeX, uint32 SizeY, uint8 Format, uint32 NumMips, uint32 Flags, void** InitialMipData, uint32 NumInitialMips) final override;
 	virtual void RHICopySharedMips(FTexture2DRHIParamRef DestTexture2D, FTexture2DRHIParamRef SrcTexture2D) final override;
-	virtual FTexture2DArrayRHIRef RHICreateTexture2DArray(uint32 SizeX, uint32 SizeY, uint32 SizeZ, uint8 Format, uint32 NumMips, uint32 Flags, FRHIResourceCreateInfo& CreateInfo) final override;
+	virtual FTexture2DArrayRHIRef RHICreateTexture2DArray(uint32 SizeX, uint32 SizeY, uint32 SizeZ, uint8 Format, uint32 NumMips, uint32 NumSamples, uint32 Flags, FRHIResourceCreateInfo& CreateInfo) final override;
 	virtual FTexture3DRHIRef RHICreateTexture3D(uint32 SizeX, uint32 SizeY, uint32 SizeZ, uint8 Format, uint32 NumMips, uint32 Flags, FRHIResourceCreateInfo& CreateInfo) final override;
 	virtual void RHIGetResourceInfo(FTextureRHIParamRef Ref, FRHIResourceInfo& OutInfo) final override;
 	virtual FShaderResourceViewRHIRef RHICreateShaderResourceView(FTexture2DRHIParamRef Texture2DRHI, uint8 MipLevel) final override;
diff --git a/Engine/Source/Runtime/EmptyRHI/Public/EmptyResources.h b/Engine/Source/Runtime/EmptyRHI/Public/EmptyResources.h
--- a/Engine/Source/Runtime/EmptyRHI/Public/EmptyResources.h
+++ b/Engine/Source/Runtime/EmptyRHI/Public/EmptyResources.h
@@ -115,7 +115,7 @@ public:
 
 	// Constructor, just calls base and Surface constructor
 	FEmptyTexture2DArray(EPixelFormat Format, uint32 SizeX, uint32 SizeY, uint32 ArraySize, uint32 NumMips, uint32 Flags, FResourceBulkDataInterface* BulkData)
-		: FRHITexture2DArray(SizeX, SizeY, ArraySize, NumMips, Format, Flags)
+		: FRHITexture2DArray(SizeX, SizeY, ArraySize, NumMips, 1, Format, Flags)
 		, Surface(RRT_Texture2DArray, Format, SizeX, SizeY, 1, /*bArray=*/ true, ArraySize, NumMips, Flags, BulkData)
 	{
 	}
diff --git a/Engine/Source/Runtime/Engine/Private/Texture2D.cpp b/Engine/Source/Runtime/Engine/Private/Texture2D.cpp
--- a/Engine/Source/Runtime/Engine/Private/Texture2D.cpp
+++ b/Engine/Source/Runtime/Engine/Private/Texture2D.cpp
@@ -1861,7 +1861,7 @@ void FTexture2DArrayResource::InitRHI()
 	// Create the RHI texture.
 	const uint32 TexCreateFlags = (bSRGB ? TexCreate_SRGB : 0) | TexCreate_OfflineProcessed;
 	FRHIResourceCreateInfo CreateInfo;
-	FTexture2DArrayRHIRef TextureArray = RHICreateTexture2DArray(SizeX, SizeY, GetNumValidTextures(), Format, NumMips, TexCreateFlags, CreateInfo);
+	FTexture2DArrayRHIRef TextureArray = RHICreateTexture2DArray(SizeX, SizeY, GetNumValidTextures(), Format, NumMips, 1, TexCreateFlags, CreateInfo);
 	TextureRHI = TextureArray;
 
 	// Read the mip-levels into the RHI texture.
diff --git a/Engine/Source/Runtime/Landscape/Private/LandscapeEditLayers.cpp b/Engine/Source/Runtime/Landscape/Private/LandscapeEditLayers.cpp
--- a/Engine/Source/Runtime/Landscape/Private/LandscapeEditLayers.cpp
+++ b/Engine/Source/Runtime/Landscape/Private/LandscapeEditLayers.cpp
@@ -593,7 +593,7 @@ public:
 			Flags |= TexCreate_UAV;
 		}
 
-		TextureRHI = RHICreateTexture2DArray(SizeX, SizeY, SizeZ, Format, NumMips, Flags, CreateInfo);
+		TextureRHI = RHICreateTexture2DArray(SizeX, SizeY, SizeZ, Format, NumMips, 1, Flags, CreateInfo);
 
 		if (CreateUAV)
 		{
diff --git a/Engine/Source/Runtime/NullDrv/Public/NullRHI.h b/Engine/Source/Runtime/NullDrv/Public/NullRHI.h
--- a/Engine/Source/Runtime/NullDrv/Public/NullRHI.h
+++ b/Engine/Source/Runtime/NullDrv/Public/NullRHI.h
@@ -293,9 +293,9 @@ public:
 	virtual void RHICopySharedMips(FTexture2DRHIParamRef DestTexture2D, FTexture2DRHIParamRef SrcTexture2D) final override
 	{
 	}
-	virtual FTexture2DArrayRHIRef RHICreateTexture2DArray(uint32 SizeX, uint32 SizeY, uint32 SizeZ, uint8 Format, uint32 NumMips, uint32 Flags, FRHIResourceCreateInfo& CreateInfo) final override
+	virtual FTexture2DArrayRHIRef RHICreateTexture2DArray(uint32 SizeX, uint32 SizeY, uint32 SizeZ, uint8 Format, uint32 NumMips, uint32 NumSamples, uint32 Flags, FRHIResourceCreateInfo& CreateInfo) final override
 	{ 
-		return new FRHITexture2DArray(SizeX,SizeY,SizeZ,NumMips,(EPixelFormat)Format,Flags, CreateInfo.ClearValueBinding); 
+		return new FRHITexture2DArray(SizeX,SizeY,SizeZ,NumMips,NumSamples,(EPixelFormat)Format,Flags, CreateInfo.ClearValueBinding); 
 	}
 
 	virtual FTexture3DRHIRef RHICreateTexture3D(uint32 SizeX, uint32 SizeY, uint32 SizeZ, uint8 Format, uint32 NumMips, uint32 Flags, FRHIResourceCreateInfo& CreateInfo) final override
diff --git a/Engine/Source/Runtime/OpenGLDrv/Private/OpenGLTexture.cpp b/Engine/Source/Runtime/OpenGLDrv/Private/OpenGLTexture.cpp
--- a/Engine/Source/Runtime/OpenGLDrv/Private/OpenGLTexture.cpp
+++ b/Engine/Source/Runtime/OpenGLDrv/Private/OpenGLTexture.cpp
@@ -1602,7 +1602,7 @@ void FOpenGLDynamicRHI::RHICopySharedMips(FTexture2DRHIParamRef DestTexture2D,FT
 	check(0);
 }
 
-FTexture2DArrayRHIRef FOpenGLDynamicRHI::RHICreateTexture2DArray(uint32 SizeX,uint32 SizeY,uint32 SizeZ,uint8 Format,uint32 NumMips,uint32 Flags, FRHIResourceCreateInfo& Info)
+FTexture2DArrayRHIRef FOpenGLDynamicRHI::RHICreateTexture2DArray(uint32 SizeX,uint32 SizeY,uint32 SizeZ,uint8 Format,uint32 NumMips,uint32 NumSamples,uint32 Flags, FRHIResourceCreateInfo& Info)
 {
 	VERIFY_GL_SCOPE();
 
diff --git a/Engine/Source/Runtime/OpenGLDrv/Public/OpenGLDrv.h b/Engine/Source/Runtime/OpenGLDrv/Public/OpenGLDrv.h
--- a/Engine/Source/Runtime/OpenGLDrv/Public/OpenGLDrv.h
+++ b/Engine/Source/Runtime/OpenGLDrv/Public/OpenGLDrv.h
@@ -392,7 +392,7 @@ public:
 	virtual FTexture2DRHIRef RHICreateTextureExternal2D(uint32 SizeX, uint32 SizeY, uint8 Format, uint32 NumMips, uint32 NumSamples, uint32 Flags, FRHIResourceCreateInfo& CreateInfo) final override;
 	virtual FTexture2DRHIRef RHIAsyncCreateTexture2D(uint32 SizeX, uint32 SizeY, uint8 Format, uint32 NumMips, uint32 Flags, void** InitialMipData, uint32 NumInitialMips) final override;
 	virtual void RHICopySharedMips(FTexture2DRHIParamRef DestTexture2D, FTexture2DRHIParamRef SrcTexture2D) final override;
-	virtual FTexture2DArrayRHIRef RHICreateTexture2DArray(uint32 SizeX, uint32 SizeY, uint32 SizeZ, uint8 Format, uint32 NumMips, uint32 Flags, FRHIResourceCreateInfo& CreateInfo) final override;
+	virtual FTexture2DArrayRHIRef RHICreateTexture2DArray(uint32 SizeX, uint32 SizeY, uint32 SizeZ, uint8 Format, uint32 NumMips, uint32 NumSamples, uint32 Flags, FRHIResourceCreateInfo& CreateInfo) final override;
 	virtual FTexture3DRHIRef RHICreateTexture3D(uint32 SizeX, uint32 SizeY, uint32 SizeZ, uint8 Format, uint32 NumMips, uint32 Flags, FRHIResourceCreateInfo& CreateInfo) final override;
 	virtual void RHIGetResourceInfo(FTextureRHIParamRef Ref, FRHIResourceInfo& OutInfo) final override;
 	virtual FShaderResourceViewRHIRef RHICreateShaderResourceView(FTexture2DRHIParamRef Texture2DRHI, uint8 MipLevel) final override;
@@ -829,10 +829,10 @@ public:
 		RHITHREAD_GLCOMMAND_EPILOGUE_RETURN(FTexture2DRHIRef);
 	}
 
-	virtual FTexture2DArrayRHIRef RHICreateTexture2DArray_RenderThread(class FRHICommandListImmediate& RHICmdList, uint32 SizeX, uint32 SizeY, uint32 SizeZ, uint8 Format, uint32 NumMips, uint32 Flags, FRHIResourceCreateInfo& CreateInfo) final override
+	virtual FTexture2DArrayRHIRef RHICreateTexture2DArray_RenderThread(class FRHICommandListImmediate& RHICmdList, uint32 SizeX, uint32 SizeY, uint32 SizeZ, uint8 Format, uint32 NumMips, uint32 NumSamples, uint32 Flags, FRHIResourceCreateInfo& CreateInfo) final override
 	{
 		RHITHREAD_GLCOMMAND_PROLOGUE();
-		return this->RHICreateTexture2DArray(SizeX, SizeY, SizeZ, Format, NumMips, Flags, CreateInfo);
+		return this->RHICreateTexture2DArray(SizeX, SizeY, SizeZ, Format, NumMips, NumSamples, Flags, CreateInfo);
 		RHITHREAD_GLCOMMAND_EPILOGUE_RETURN(FTexture2DArrayRHIRef);
 	}
 
diff --git a/Engine/Source/Runtime/OpenGLDrv/Public/OpenGLResources.h b/Engine/Source/Runtime/OpenGLDrv/Public/OpenGLResources.h
--- a/Engine/Source/Runtime/OpenGLDrv/Public/OpenGLResources.h
+++ b/Engine/Source/Runtime/OpenGLDrv/Public/OpenGLResources.h
@@ -1588,7 +1588,7 @@ class FOpenGLBaseTexture2DArray : public FRHITexture2DArray
 {
 public:
 	FOpenGLBaseTexture2DArray(uint32 InSizeX, uint32 InSizeY, uint32 InSizeZ, uint32 InNumMips, uint32 InNumSamples, uint32 InNumSamplesTileMem, uint32 InArraySize, EPixelFormat InFormat, uint32 InFlags, const FClearValueBinding& InClearValue)
-	: FRHITexture2DArray(InSizeX,InSizeY,InSizeZ,InNumMips,InFormat,InFlags, InClearValue)
+	: FRHITexture2DArray(InSizeX,InSizeY,InSizeZ,InNumMips,InNumSamples,InFormat,InFlags, InClearValue)
 	{
 		check(InNumSamples == 1);	// OpenGL supports multisampled texture arrays, but they're currently not implemented in OpenGLDrv.
 		check(InNumSamplesTileMem == 1);
diff --git a/Engine/Source/Runtime/RHI/Private/RHI.cpp b/Engine/Source/Runtime/RHI/Private/RHI.cpp
--- a/Engine/Source/Runtime/RHI/Private/RHI.cpp
+++ b/Engine/Source/Runtime/RHI/Private/RHI.cpp
@@ -866,6 +866,14 @@ void FRHIRenderPassInfo::ConvertToRenderTargetsInfo(FRHISetRenderTargetsInfo& Ou
 		++OutRTInfo.NumColorRenderTargets;
 
 		OutRTInfo.bClearColor |= (LoadAction == ERenderTargetLoadAction::EClear);
+
+		ensure(!OutRTInfo.bHasResolveAttachments || ColorRenderTargets[Index].ResolveTarget);
+		if (ColorRenderTargets[Index].ResolveTarget)
+		{
+			OutRTInfo.bHasResolveAttachments = true;
+			OutRTInfo.ColorResolveRenderTarget[Index] = OutRTInfo.ColorRenderTarget[Index];
+			OutRTInfo.ColorResolveRenderTarget[Index].Texture = ColorRenderTargets[Index].ResolveTarget;
+		}
 	}
 
 	ERenderTargetActions DepthActions = GetDepthActions(DepthStencilRenderTarget.Action);
diff --git a/Engine/Source/Runtime/RHI/Private/RHICommandList.cpp b/Engine/Source/Runtime/RHI/Private/RHICommandList.cpp
--- a/Engine/Source/Runtime/RHI/Private/RHICommandList.cpp
+++ b/Engine/Source/Runtime/RHI/Private/RHICommandList.cpp
@@ -2334,11 +2334,11 @@ FTexture2DRHIRef FDynamicRHI::RHICreateTextureExternal2D_RenderThread(class FRHI
 	return GDynamicRHI->RHICreateTextureExternal2D(SizeX, SizeY, Format, NumMips, NumSamples, Flags, CreateInfo);
 }
 
-FTexture2DArrayRHIRef FDynamicRHI::RHICreateTexture2DArray_RenderThread(class FRHICommandListImmediate& RHICmdList, uint32 SizeX, uint32 SizeY, uint32 SizeZ, uint8 Format, uint32 NumMips, uint32 Flags, FRHIResourceCreateInfo& CreateInfo)
+FTexture2DArrayRHIRef FDynamicRHI::RHICreateTexture2DArray_RenderThread(class FRHICommandListImmediate& RHICmdList, uint32 SizeX, uint32 SizeY, uint32 SizeZ, uint8 Format, uint32 NumMips, uint32 NumSamples, uint32 Flags, FRHIResourceCreateInfo& CreateInfo)
 {
 	CSV_SCOPED_TIMING_STAT(RHITStalls, RHICreateTexture2DArray);
 	FScopedRHIThreadStaller StallRHIThread(RHICmdList);
-	return GDynamicRHI->RHICreateTexture2DArray(SizeX, SizeY, SizeZ, Format, NumMips, Flags, CreateInfo);
+	return GDynamicRHI->RHICreateTexture2DArray(SizeX, SizeY, SizeZ, Format, NumMips, NumSamples, Flags, CreateInfo);
 }
 
 FTexture3DRHIRef FDynamicRHI::RHICreateTexture3D_RenderThread(class FRHICommandListImmediate& RHICmdList, uint32 SizeX, uint32 SizeY, uint32 SizeZ, uint8 Format, uint32 NumMips, uint32 Flags, FRHIResourceCreateInfo& CreateInfo)
diff --git a/Engine/Source/Runtime/RHI/Public/DynamicRHI.h b/Engine/Source/Runtime/RHI/Public/DynamicRHI.h
--- a/Engine/Source/Runtime/RHI/Public/DynamicRHI.h
+++ b/Engine/Source/Runtime/RHI/Public/DynamicRHI.h
@@ -603,7 +603,24 @@ public:
 	* @param Flags - ETextureCreateFlags creation flags
 	*/
 	// FlushType: Wait RHI Thread
-	virtual FTexture2DArrayRHIRef RHICreateTexture2DArray(uint32 SizeX, uint32 SizeY, uint32 SizeZ, uint8 Format, uint32 NumMips, uint32 Flags, FRHIResourceCreateInfo& CreateInfo) = 0;
+	UE_DEPRECATED(4.23, "RHICreateTexture2DArray now takes NumSamples")
+	virtual FTexture2DArrayRHIRef RHICreateTexture2DArray(uint32 SizeX, uint32 SizeY, uint32 SizeZ, uint8 Format, uint32 NumMips, uint32 Flags, FRHIResourceCreateInfo& CreateInfo)
+	{
+		return RHICreateTexture2DArray(SizeX, SizeY, SizeZ, Format, NumMips, 1, Flags, CreateInfo);
+	}
+
+	/**
+	* Creates a Array RHI texture resource
+	* @param SizeX - width of the texture to create
+	* @param SizeY - height of the texture to create
+	* @param SizeZ - depth of the texture to create
+	* @param Format - EPixelFormat texture format
+	* @param NumMips - number of mips to generate or 0 for full mip pyramid
+	* @param NumSamples - number of MSAA samples, usually 1
+	* @param Flags - ETextureCreateFlags creation flags
+	*/
+	// FlushType: Wait RHI Thread
+	virtual FTexture2DArrayRHIRef RHICreateTexture2DArray(uint32 SizeX, uint32 SizeY, uint32 SizeZ, uint8 Format, uint32 NumMips, uint32 NumSamples, uint32 Flags, FRHIResourceCreateInfo& CreateInfo) = 0;
 
 	/**
 	* Creates a 3d RHI texture resource
@@ -1090,7 +1107,14 @@ public:
 	virtual FTextureReferenceRHIRef RHICreateTextureReference_RenderThread(class FRHICommandListImmediate& RHICmdList, FLastRenderTimeContainer* LastRenderTime);
 	virtual FTexture2DRHIRef RHICreateTexture2D_RenderThread(class FRHICommandListImmediate& RHICmdList, uint32 SizeX, uint32 SizeY, uint8 Format, uint32 NumMips, uint32 NumSamples, uint32 Flags, FRHIResourceCreateInfo& CreateInfo);
 	virtual FTexture2DRHIRef RHICreateTextureExternal2D_RenderThread(class FRHICommandListImmediate& RHICmdList, uint32 SizeX, uint32 SizeY, uint8 Format, uint32 NumMips, uint32 NumSamples, uint32 Flags, FRHIResourceCreateInfo& CreateInfo);
-	virtual FTexture2DArrayRHIRef RHICreateTexture2DArray_RenderThread(class FRHICommandListImmediate& RHICmdList, uint32 SizeX, uint32 SizeY, uint32 SizeZ, uint8 Format, uint32 NumMips, uint32 Flags, FRHIResourceCreateInfo& CreateInfo);
+
+	virtual FTexture2DArrayRHIRef RHICreateTexture2DArray_RenderThread(class FRHICommandListImmediate& RHICmdList, uint32 SizeX, uint32 SizeY, uint32 SizeZ, uint8 Format, uint32 NumMips, uint32 NumSamples, uint32 Flags, FRHIResourceCreateInfo& CreateInfo);
+	UE_DEPRECATED(4.23, "RHICreateTexture2DArray_RenderThread now takes NumSamples")
+	virtual FTexture2DArrayRHIRef RHICreateTexture2DArray_RenderThread(class FRHICommandListImmediate& RHICmdList, uint32 SizeX, uint32 SizeY, uint32 SizeZ, uint8 Format, uint32 NumMips, uint32 Flags, FRHIResourceCreateInfo& CreateInfo)
+	{
+		return RHICreateTexture2DArray_RenderThread(RHICmdList, SizeX, SizeY, SizeZ, Format, NumMips, 1, Flags, CreateInfo);
+	}
+
 	virtual FTexture3DRHIRef RHICreateTexture3D_RenderThread(class FRHICommandListImmediate& RHICmdList, uint32 SizeX, uint32 SizeY, uint32 SizeZ, uint8 Format, uint32 NumMips, uint32 Flags, FRHIResourceCreateInfo& CreateInfo);
 	virtual FUnorderedAccessViewRHIRef RHICreateUnorderedAccessView_RenderThread(class FRHICommandListImmediate& RHICmdList, FStructuredBufferRHIParamRef StructuredBuffer, bool bUseUAVCounter, bool bAppendBuffer);
 	virtual FUnorderedAccessViewRHIRef RHICreateUnorderedAccessView_RenderThread(class FRHICommandListImmediate& RHICmdList, FTextureRHIParamRef Texture, uint32 MipLevel);
diff --git a/Engine/Source/Runtime/RHI/Public/RHICommandList.h b/Engine/Source/Runtime/RHI/Public/RHICommandList.h
--- a/Engine/Source/Runtime/RHI/Public/RHICommandList.h
+++ b/Engine/Source/Runtime/RHI/Public/RHICommandList.h
@@ -3949,10 +3949,16 @@ public:
 		return GDynamicRHI->RHITransferTexture(Texture, Rect, SrcGPUIndex, DestGPUIndex, PullData);
 	}
 
-	FORCEINLINE FTexture2DArrayRHIRef CreateTexture2DArray(uint32 SizeX, uint32 SizeY, uint32 SizeZ, uint8 Format, uint32 NumMips, uint32 Flags, FRHIResourceCreateInfo& CreateInfo)
+	FORCEINLINE FTexture2DArrayRHIRef CreateTexture2DArray(uint32 SizeX, uint32 SizeY, uint32 SizeZ, uint8 Format, uint32 NumMips, uint32 NumSamples, uint32 Flags, FRHIResourceCreateInfo& CreateInfo)
 	{
 		LLM_SCOPE((Flags & (TexCreate_RenderTargetable | TexCreate_DepthStencilTargetable)) != 0 ? ELLMTag::RenderTargets : ELLMTag::Textures);
-		return GDynamicRHI->RHICreateTexture2DArray_RenderThread(*this, SizeX, SizeY, SizeZ, Format, NumMips, Flags, CreateInfo);
+		return GDynamicRHI->RHICreateTexture2DArray_RenderThread(*this, SizeX, SizeY, SizeZ, Format, NumMips, NumSamples, Flags, CreateInfo);
+	}
+
+	UE_DEPRECATED(4.23, "CreateTexture2DArray now takes NumSamples")
+	FORCEINLINE FTexture2DArrayRHIRef CreateTexture2DArray(uint32 SizeX, uint32 SizeY, uint32 SizeZ, uint8 Format, uint32 NumMips, uint32 Flags, FRHIResourceCreateInfo& CreateInfo)
+	{
+		return CreateTexture2DArray(SizeX, SizeY, SizeZ, Format, NumMips, 1, Flags, CreateInfo);
 	}
 	
 	FORCEINLINE FTexture3DRHIRef CreateTexture3D(uint32 SizeX, uint32 SizeY, uint32 SizeZ, uint8 Format, uint32 NumMips, uint32 Flags, FRHIResourceCreateInfo& CreateInfo)
@@ -4807,9 +4813,15 @@ FORCEINLINE void RHITransferTexture(FTexture2DRHIParamRef Texture, FIntRect Rect
 	return FRHICommandListExecutor::GetImmediateCommandList().TransferTexture(Texture, Rect, SrcGPUIndex, DestGPUIndex, PullData);
 }
 
+FORCEINLINE FTexture2DArrayRHIRef RHICreateTexture2DArray(uint32 SizeX, uint32 SizeY, uint32 SizeZ, uint8 Format, uint32 NumMips, uint32 NumSamples, uint32 Flags, FRHIResourceCreateInfo& CreateInfo)
+{
+	return FRHICommandListExecutor::GetImmediateCommandList().CreateTexture2DArray(SizeX, SizeY, SizeZ, Format, NumMips, NumSamples, Flags, CreateInfo);
+}
+
+UE_DEPRECATED(4.23, "RHICreateTexture2DArray now takes NumSamples")
 FORCEINLINE FTexture2DArrayRHIRef RHICreateTexture2DArray(uint32 SizeX, uint32 SizeY, uint32 SizeZ, uint8 Format, uint32 NumMips, uint32 Flags, FRHIResourceCreateInfo& CreateInfo)
 {
-	return FRHICommandListExecutor::GetImmediateCommandList().CreateTexture2DArray(SizeX, SizeY, SizeZ, Format, NumMips, Flags, CreateInfo);
+	return RHICreateTexture2DArray(SizeX, SizeY, SizeZ, Format, NumMips, 1, Flags, CreateInfo);
 }
 
 FORCEINLINE FTexture3DRHIRef RHICreateTexture3D(uint32 SizeX, uint32 SizeY, uint32 SizeZ, uint8 Format, uint32 NumMips, uint32 Flags, FRHIResourceCreateInfo& CreateInfo)
diff --git a/Engine/Source/Runtime/RHI/Public/RHIResources.h b/Engine/Source/Runtime/RHI/Public/RHIResources.h
--- a/Engine/Source/Runtime/RHI/Public/RHIResources.h
+++ b/Engine/Source/Runtime/RHI/Public/RHIResources.h
@@ -751,8 +751,8 @@ class RHI_API FRHITexture2DArray : public FRHITexture2D
 public:
 	
 	/** Initialization constructor. */
-	FRHITexture2DArray(uint32 InSizeX,uint32 InSizeY,uint32 InSizeZ,uint32 InNumMips,EPixelFormat InFormat,uint32 InFlags, const FClearValueBinding& InClearValue)
-	: FRHITexture2D(InSizeX, InSizeY, InNumMips,1,InFormat,InFlags, InClearValue)
+	FRHITexture2DArray(uint32 InSizeX,uint32 InSizeY,uint32 InSizeZ,uint32 InNumMips,uint32 NumSamples, EPixelFormat InFormat,uint32 InFlags, const FClearValueBinding& InClearValue)
+	: FRHITexture2D(InSizeX, InSizeY, InNumMips,NumSamples,InFormat,InFlags, InClearValue)
 	, SizeZ(InSizeZ)
 	{}
 	
@@ -1500,6 +1500,10 @@ public:
 	int32 NumColorRenderTargets;
 	bool bClearColor;
 
+	// Color Render Targets Info
+	FRHIRenderTargetView ColorResolveRenderTarget[MaxSimultaneousRenderTargets];	
+	bool bHasResolveAttachments;
+
 	// Depth/Stencil Render Target Info
 	FRHIDepthRenderTargetView DepthStencilRenderTarget;	
 	bool bClearDepth;
@@ -1511,7 +1515,8 @@ public:
 
 	FRHISetRenderTargetsInfo() :
 		NumColorRenderTargets(0),
-		bClearColor(false),		
+		bClearColor(false),
+		bHasResolveAttachments(false),
 		bClearDepth(false),
 		bClearStencil(false),
 		NumUAVs(0)
@@ -1520,6 +1525,7 @@ public:
 	FRHISetRenderTargetsInfo(int32 InNumColorRenderTargets, const FRHIRenderTargetView* InColorRenderTargets, const FRHIDepthRenderTargetView& InDepthStencilRenderTarget) :
 		NumColorRenderTargets(InNumColorRenderTargets),
 		bClearColor(InNumColorRenderTargets > 0 && InColorRenderTargets[0].LoadAction == ERenderTargetLoadAction::EClear),
+		bHasResolveAttachments(false),
 		DepthStencilRenderTarget(InDepthStencilRenderTarget),		
 		bClearDepth(InDepthStencilRenderTarget.Texture && InDepthStencilRenderTarget.DepthLoadAction == ERenderTargetLoadAction::EClear),
 		bClearStencil(InDepthStencilRenderTarget.Texture && InDepthStencilRenderTarget.StencilLoadAction == ERenderTargetLoadAction::EClear),
@@ -1552,7 +1558,7 @@ public:
 		struct FHashableStruct
 		{
 			// Depth goes in the last slot
-			FRHITexture* Texture[MaxSimultaneousRenderTargets + 1];
+			FRHITexture* Texture[MaxSimultaneousRenderTargets*2 + 1];
 			uint32 MipIndex[MaxSimultaneousRenderTargets];
 			uint32 ArraySliceIndex[MaxSimultaneousRenderTargets];
 			ERenderTargetLoadAction LoadAction[MaxSimultaneousRenderTargets];
@@ -1567,6 +1573,7 @@ public:
 			bool bClearDepth;
 			bool bClearStencil;
 			bool bClearColor;
+			bool bHasResolveAttachments;
 			FRHIUnorderedAccessView* UnorderedAccessView[MaxSimultaneousUAVs];
 
 			void Set(const FRHISetRenderTargetsInfo& RTInfo)
@@ -1575,6 +1582,7 @@ public:
 				for (int32 Index = 0; Index < RTInfo.NumColorRenderTargets; ++Index)
 				{
 					Texture[Index] = RTInfo.ColorRenderTarget[Index].Texture;
+					Texture[MaxSimultaneousRenderTargets+Index] = RTInfo.ColorResolveRenderTarget[Index].Texture;
 					MipIndex[Index] = RTInfo.ColorRenderTarget[Index].MipIndex;
 					ArraySliceIndex[Index] = RTInfo.ColorRenderTarget[Index].ArraySliceIndex;
 					LoadAction[Index] = RTInfo.ColorRenderTarget[Index].LoadAction;
@@ -1591,6 +1599,7 @@ public:
 				bClearDepth = RTInfo.bClearDepth;
 				bClearStencil = RTInfo.bClearStencil;
 				bClearColor = RTInfo.bClearColor;
+				bHasResolveAttachments = RTInfo.bHasResolveAttachments;
 
 				for (int32 Index = 0; Index < MaxSimultaneousUAVs; ++Index)
 				{
diff --git a/Engine/Source/Runtime/RHI/Public/RHIUtilities.h b/Engine/Source/Runtime/RHI/Public/RHIUtilities.h
--- a/Engine/Source/Runtime/RHI/Public/RHIUtilities.h
+++ b/Engine/Source/Runtime/RHI/Public/RHIUtilities.h
@@ -645,9 +645,28 @@ inline void RHICreateTargetableShaderResource2DArray(
 
 	// Ensure that the targetable texture is either render or depth-stencil targetable.
 	check(TargetableTextureFlags & (TexCreate_RenderTargetable | TexCreate_DepthStencilTargetable));
+	bool bForceSeparateTargetAndShaderResource = false;
+	if (NumSamples > 1)
+	{
+		bForceSeparateTargetAndShaderResource = RHISupportsSeparateMSAAAndResolveTextures(GMaxRHIShaderPlatform);
+	}
 
-	// Create a single texture that has both TargetableTextureFlags and TexCreate_ShaderResource set.
-	OutTargetableTexture = OutShaderResourceTexture = RHICreateTexture2DArray(SizeX, SizeY, SizeZ, Format, NumMips, Flags | TargetableTextureFlags | TexCreate_ShaderResource, CreateInfo);
+	if (!bForceSeparateTargetAndShaderResource)
+	{
+		// Create a single texture that has both TargetableTextureFlags and TexCreate_ShaderResource set.
+		OutTargetableTexture = OutShaderResourceTexture = RHICreateTexture2DArray(SizeX, SizeY, SizeZ, Format, NumMips, NumSamples, Flags | TargetableTextureFlags | TexCreate_ShaderResource, CreateInfo);
+	}
+	else
+	{
+		uint32 ResolveTargetableTextureFlags = TexCreate_ResolveTargetable;
+		if (TargetableTextureFlags & TexCreate_DepthStencilTargetable)
+		{
+			ResolveTargetableTextureFlags |= TexCreate_DepthStencilResolveTarget;
+		}
+		// Create a texture that has TargetableTextureFlags set, and a second texture that has TexCreate_ResolveTargetable and TexCreate_ShaderResource set.
+		OutTargetableTexture = RHICreateTexture2DArray(SizeX, SizeY, SizeZ, Format, NumMips, NumSamples, Flags | TargetableTextureFlags, CreateInfo);
+		OutShaderResourceTexture = RHICreateTexture2DArray(SizeX, SizeY, SizeZ, Format, NumMips, 1,  Flags | ResolveTargetableTextureFlags | TexCreate_ShaderResource, CreateInfo);
+	}
 }
 
 /**
diff --git a/Engine/Source/Runtime/RenderCore/Private/RenderUtils.cpp b/Engine/Source/Runtime/RenderCore/Private/RenderUtils.cpp
--- a/Engine/Source/Runtime/RenderCore/Private/RenderUtils.cpp
+++ b/Engine/Source/Runtime/RenderCore/Private/RenderUtils.cpp
@@ -343,7 +343,7 @@ public:
 			FBlackVolumeTextureResourceBulkDataInterface BlackTextureBulkData;
 			FRHIResourceCreateInfo CreateInfo(&BlackTextureBulkData);
 			CreateInfo.DebugName = TEXT("BlackArrayTexture");
-			FTexture2DArrayRHIRef TextureArray = RHICreateTexture2DArray(1, 1, 1, PF_B8G8R8A8, 1, TexCreate_ShaderResource, CreateInfo);
+			FTexture2DArrayRHIRef TextureArray = RHICreateTexture2DArray(1, 1, 1, PF_B8G8R8A8, 1, 1, TexCreate_ShaderResource, CreateInfo);
 			TextureRHI = TextureArray;
 
 			// Create the sampler state RHI resource.
diff --git a/Engine/Source/Runtime/Renderer/Private/MobileShadingRenderer.cpp b/Engine/Source/Runtime/Renderer/Private/MobileShadingRenderer.cpp
--- a/Engine/Source/Runtime/Renderer/Private/MobileShadingRenderer.cpp
+++ b/Engine/Source/Runtime/Renderer/Private/MobileShadingRenderer.cpp
@@ -405,18 +405,25 @@ void FMobileSceneRenderer::Render(FRHICommandListImmediate& RHICmdList)
 	FTextureRHIParamRef SceneDepth = nullptr;
 	ERenderTargetActions ColorTargetAction = ERenderTargetActions::Clear_Store;
 	EDepthStencilTargetActions DepthTargetAction = EDepthStencilTargetActions::ClearDepthStencil_DontStoreDepthStencil;
-	bool bMobileMSAA = false;
+	bool bMobileMSAA = SceneContext.GetSceneColorSurface()->GetNumSamples() > 1;
 	
 	if (bGammaSpace && !bRenderToSceneColor)
 	{
-		SceneColor = GetMultiViewSceneColor(SceneContext);
-		bMobileMSAA = SceneColor->GetNumSamples() > 1;
-		SceneDepth = (View.bIsMobileMultiViewEnabled) ? SceneContext.MobileMultiViewSceneDepthZ->GetRenderTargetItem().TargetableTexture : static_cast<FTextureRHIRef>(SceneContext.GetSceneDepthTexture());
+		if (bMobileMSAA)
+		{
+			SceneColor = (View.bIsMobileMultiViewEnabled) ? SceneContext.MobileMultiViewSceneColor->GetRenderTargetItem().TargetableTexture : SceneContext.GetSceneColorSurface();
+			SceneColorResolve = GetMultiViewSceneColor(SceneContext);
+			ColorTargetAction = ERenderTargetActions::Clear_Resolve;
+		}
+		else
+		{
+			SceneColor = GetMultiViewSceneColor(SceneContext);
+		}
+		SceneDepth = (View.bIsMobileMultiViewEnabled) ? SceneContext.MobileMultiViewSceneDepthZ->GetRenderTargetItem().TargetableTexture : static_cast<FTextureRHIRef>(SceneContext.GetSceneDepthSurface());
 	}
 	else
 	{
 		SceneColor = SceneContext.GetSceneColorSurface();
-		bMobileMSAA = SceneColor->GetNumSamples() > 1;
 		SceneColorResolve = bMobileMSAA ? SceneContext.GetSceneColorTexture() : nullptr;
 		ColorTargetAction = bMobileMSAA ? ERenderTargetActions::Clear_Resolve : ERenderTargetActions::Clear_Store;
 		SceneDepth = SceneContext.GetSceneDepthSurface();
diff --git a/Engine/Source/Runtime/Renderer/Private/PostProcess/SceneRenderTargets.cpp b/Engine/Source/Runtime/Renderer/Private/PostProcess/SceneRenderTargets.cpp
--- a/Engine/Source/Runtime/Renderer/Private/PostProcess/SceneRenderTargets.cpp
+++ b/Engine/Source/Runtime/Renderer/Private/PostProcess/SceneRenderTargets.cpp
@@ -1051,7 +1051,7 @@ void FSceneRenderTargets::AllocMobileMultiViewDepth(FRHICommandList& RHICmdList,
 		const FIntPoint MultiViewBufferSize(BufferSize.X / ScaleFactor, BufferSize.Y);
 
 		// Using the result of GetDepthFormat() without stencil due to packed depth-stencil not working in array frame buffers.
-		FPooledRenderTargetDesc Desc(FPooledRenderTargetDesc::Create2DDesc(MultiViewBufferSize, PF_D24, DefaultDepthClear, TexCreate_None, TexCreate_DepthStencilTargetable, false));
+		FPooledRenderTargetDesc Desc(FPooledRenderTargetDesc::Create2DDesc(MultiViewBufferSize, PF_D24, DefaultDepthClear, TexCreate_None, TexCreate_DepthStencilTargetable | TexCreate_ShaderResource | TexCreate_InputAttachmentRead, false));
 		Desc.Flags |= TexCreate_FastVRAM;
 		Desc.NumSamples = GetNumSceneColorMSAASamples(CurrentFeatureLevel);
 		Desc.ArraySize = 2;
@@ -1990,11 +1990,8 @@ void FSceneRenderTargets::AllocateMobileRenderTargets(FRHICommandList& RHICmdLis
 	if (bIsUsingMobileMultiView)
 	{
 		const int32 ScaleFactor = (bIsMobileMultiViewDirectEnabled) ? 1 : 2;
-		if (!bIsMobileMultiViewDirectEnabled)
-		{
-			AllocMobileMultiViewSceneColor(RHICmdList, ScaleFactor);
-		}
 
+		AllocMobileMultiViewSceneColor(RHICmdList, ScaleFactor);
 		AllocMobileMultiViewDepth(RHICmdList, ScaleFactor);
 	}
 #endif
diff --git a/Engine/Source/Runtime/SlateRHIRenderer/Private/SlateRHIRenderer.cpp b/Engine/Source/Runtime/SlateRHIRenderer/Private/SlateRHIRenderer.cpp
--- a/Engine/Source/Runtime/SlateRHIRenderer/Private/SlateRHIRenderer.cpp
+++ b/Engine/Source/Runtime/SlateRHIRenderer/Private/SlateRHIRenderer.cpp
@@ -657,6 +657,7 @@ void FSlateRHIRenderer::DrawWindow_RenderThread(FRHICommandListImmediate& RHICmd
 	{
 		SCOPED_DRAW_EVENTF(RHICmdList, SlateUI, TEXT("SlateUI Title = %s"), DrawCommandParams.WindowTitle.IsEmpty() ? TEXT("<none>") : *DrawCommandParams.WindowTitle);
 		SCOPED_GPU_STAT(RHICmdList, SlateUI);
+		SCOPED_NAMED_EVENT_TEXT("Slate::DrawWindow_RenderThread", FColor::Magenta);
 
 		// Should only be called by the rendering thread
 		check(IsInRenderingThread());
@@ -813,10 +814,10 @@ void FSlateRHIRenderer::DrawWindow_RenderThread(FRHICommandListImmediate& RHICmd
 					}
 				}
 	#endif
-				RHICmdList.BeginRenderPass(RPInfo, TEXT("SlateBatches"));
 				{
 					if (BatchData.GetRenderBatches().Num() > 0)
 					{
+						RHICmdList.BeginRenderPass(RPInfo, TEXT("SlateBatches"));
 						SCOPE_CYCLE_COUNTER(STAT_SlateRTDrawBatches);
 
 						FSlateBackBuffer BackBufferTarget(BackBuffer, FIntPoint(ViewportWidth, ViewportHeight));
diff --git a/Engine/Source/Runtime/VulkanRHI/Private/Android/VulkanAndroidPlatform.h b/Engine/Source/Runtime/VulkanRHI/Private/Android/VulkanAndroidPlatform.h
--- a/Engine/Source/Runtime/VulkanRHI/Private/Android/VulkanAndroidPlatform.h
+++ b/Engine/Source/Runtime/VulkanRHI/Private/Android/VulkanAndroidPlatform.h
@@ -101,6 +101,8 @@ public:
 	// Does the platform allow a nullptr Pixelshader on the pipeline
 	static bool SupportsNullPixelShader() { return false; }
 
+	static bool RequiresRenderPassResolveAttachments() { return true; }
+
 protected:
 	static void* VulkanLib;
 	static bool bAttemptedLoad;
diff --git a/Engine/Source/Runtime/VulkanRHI/Private/VulkanGenericPlatform.h b/Engine/Source/Runtime/VulkanRHI/Private/VulkanGenericPlatform.h
--- a/Engine/Source/Runtime/VulkanRHI/Private/VulkanGenericPlatform.h
+++ b/Engine/Source/Runtime/VulkanRHI/Private/VulkanGenericPlatform.h
@@ -104,4 +104,7 @@ public:
 
 	// Does the platform allow a nullptr Pixelshader on the pipeline
 	static bool SupportsNullPixelShader() { return true; }
+
+	// Does the platform require resolve attachments in its MSAA renderpasses
+	static bool RequiresRenderPassResolveAttachments() { return false; }
 };
diff --git a/Engine/Source/Runtime/VulkanRHI/Private/VulkanRHI.cpp b/Engine/Source/Runtime/VulkanRHI/Private/VulkanRHI.cpp
--- a/Engine/Source/Runtime/VulkanRHI/Private/VulkanRHI.cpp
+++ b/Engine/Source/Runtime/VulkanRHI/Private/VulkanRHI.cpp
@@ -88,9 +88,6 @@ FDynamicRHI* FVulkanDynamicRHIModule::CreateRHI(ERHIFeatureLevel::Type InRequest
 		GMaxRHIShaderPlatform = (PLATFORM_LUMINGL4 || PLATFORM_LUMIN) ? SP_VULKAN_SM5_LUMIN : SP_VULKAN_SM5;
 	}
 
-	// VULKAN_USE_MSAA_RESOLVE_ATTACHMENTS=0 requires separate MSAA and resolve textures
-	check(RHISupportsSeparateMSAAAndResolveTextures(GMaxRHIShaderPlatform) == (!VULKAN_USE_MSAA_RESOLVE_ATTACHMENTS));
-
 	return new FVulkanDynamicRHI();
 }
 
@@ -990,10 +987,10 @@ FTexture2DRHIRef FVulkanDynamicRHI::RHICreateTexture2DFromResource(EPixelFormat
 	return new FVulkanTexture2D(*Device, Format, SizeX, SizeY, NumMips, NumSamples, Resource, ConversionInitializer, Flags, ResourceCreateInfo);
 }
 
-FTexture2DArrayRHIRef FVulkanDynamicRHI::RHICreateTexture2DArrayFromResource(EPixelFormat Format, uint32 SizeX, uint32 SizeY, uint32 ArraySize, uint32 NumMips, VkImage Resource, uint32 Flags)
+FTexture2DArrayRHIRef FVulkanDynamicRHI::RHICreateTexture2DArrayFromResource(EPixelFormat Format, uint32 SizeX, uint32 SizeY, uint32 ArraySize, uint32 NumMips, uint32 NumSamples, VkImage Resource, uint32 Flags)
 {
 	const FClearValueBinding ClearValueBinding(IsDepthOrStencilFormat(Format) ? FClearValueBinding::DepthZero : FClearValueBinding::Transparent);
-	return new FVulkanTexture2DArray(*Device, Format, SizeX, SizeY, ArraySize, NumMips, Resource, Flags, nullptr, ClearValueBinding);
+	return new FVulkanTexture2DArray(*Device, Format, SizeX, SizeY, ArraySize, NumMips, NumSamples, Resource, Flags, nullptr, ClearValueBinding);
 }
 
 FTextureCubeRHIRef FVulkanDynamicRHI::RHICreateTextureCubeFromResource(EPixelFormat Format, uint32 Size, bool bArray, uint32 ArraySize, uint32 NumMips, VkImage Resource, uint32 Flags)
diff --git a/Engine/Source/Runtime/VulkanRHI/Private/VulkanRHIPrivate.h b/Engine/Source/Runtime/VulkanRHI/Private/VulkanRHIPrivate.h
--- a/Engine/Source/Runtime/VulkanRHI/Private/VulkanRHIPrivate.h
+++ b/Engine/Source/Runtime/VulkanRHI/Private/VulkanRHIPrivate.h
@@ -256,6 +256,7 @@ private:
 	// Save image off for comparison, in case it gets aliased.
 	uint32 NumColorAttachments;
 	VkImage ColorRenderTargetImages[MaxSimultaneousRenderTargets];
+	VkImage ColorResolveTargetImages[MaxSimultaneousRenderTargets];
 	VkImage DepthStencilRenderTargetImage;
 
 	// Predefined set of barriers, when executes ensuring all writes are finished
diff --git a/Engine/Source/Runtime/VulkanRHI/Private/VulkanRenderTarget.cpp b/Engine/Source/Runtime/VulkanRHI/Private/VulkanRenderTarget.cpp
--- a/Engine/Source/Runtime/VulkanRHI/Private/VulkanRenderTarget.cpp
+++ b/Engine/Source/Runtime/VulkanRHI/Private/VulkanRenderTarget.cpp
@@ -361,10 +361,6 @@ void FTransitionAndLayoutManager::BeginRealRenderPass(FVulkanCommandListContext&
 			ClearValues[ClearValueIndex].depthStencil.depth = Depth;
 			ClearValues[ClearValueIndex].depthStencil.stencil = Stencil;
 			++ClearValueIndex;
-			if (Surface.GetNumSamples() > 1)
-			{
-				++ClearValueIndex;
-			}
 		}
 	}
 
@@ -1888,17 +1884,19 @@ FVulkanRenderTargetLayout::FVulkanRenderTargetLayout(FVulkanDevice& InDevice, co
 
 			ColorReferences[NumColorAttachments].attachment = NumAttachmentDescriptions;
 			ColorReferences[NumColorAttachments].layout = VK_IMAGE_LAYOUT_COLOR_ATTACHMENT_OPTIMAL;
-#if VULKAN_USE_MSAA_RESOLVE_ATTACHMENTS
-			if (CurrDesc.samples > VK_SAMPLE_COUNT_1_BIT)
+
+			const bool bHasValidResolveAttachment = RTInfo.bHasResolveAttachments && RTInfo.ColorResolveRenderTarget[Index].Texture;
+			ensure(!(CurrDesc.samples > VK_SAMPLE_COUNT_1_BIT && FVulkanPlatform::RequiresRenderPassResolveAttachments()) || bHasValidResolveAttachment);
+			if (CurrDesc.samples > VK_SAMPLE_COUNT_1_BIT && bHasValidResolveAttachment && FVulkanPlatform::RequiresRenderPassResolveAttachments())
 			{
 				Desc[NumAttachmentDescriptions + 1] = Desc[NumAttachmentDescriptions];
 				Desc[NumAttachmentDescriptions + 1].samples = VK_SAMPLE_COUNT_1_BIT;
+				Desc[NumAttachmentDescriptions + 1].storeOp = VK_ATTACHMENT_STORE_OP_STORE;
 				ResolveReferences[NumColorAttachments].attachment = NumAttachmentDescriptions + 1;
-				ResolveReferences[NumColorAttachments].layout = VK_IMAGE_LAYOUT_GENERAL;
+				ResolveReferences[NumColorAttachments].layout = VK_IMAGE_LAYOUT_COLOR_ATTACHMENT_OPTIMAL;
 				++NumAttachmentDescriptions;
 				bHasResolveAttachments = true;
 			}
-#endif
 
 			CompatibleHashInfo.Formats[NumColorAttachments] = CurrDesc.format;
 			FullHashInfo.LoadOps[NumColorAttachments] = CurrDesc.loadOp;
@@ -2067,17 +2065,18 @@ FVulkanRenderTargetLayout::FVulkanRenderTargetLayout(FVulkanDevice& InDevice, co
 
 		ColorReferences[NumColorAttachments].attachment = NumAttachmentDescriptions;
 		ColorReferences[NumColorAttachments].layout = VK_IMAGE_LAYOUT_COLOR_ATTACHMENT_OPTIMAL;
-#if VULKAN_USE_MSAA_RESOLVE_ATTACHMENTS
-		if (CurrDesc.samples > VK_SAMPLE_COUNT_1_BIT)
+
+		ensure(!(CurrDesc.samples > VK_SAMPLE_COUNT_1_BIT && FVulkanPlatform::RequiresRenderPassResolveAttachments()) || ColorEntry.ResolveTarget);
+		if (CurrDesc.samples > VK_SAMPLE_COUNT_1_BIT && ColorEntry.ResolveTarget && FVulkanPlatform::RequiresRenderPassResolveAttachments())
 		{
 			Desc[NumAttachmentDescriptions + 1] = Desc[NumAttachmentDescriptions];
 			Desc[NumAttachmentDescriptions + 1].samples = VK_SAMPLE_COUNT_1_BIT;
+			Desc[NumAttachmentDescriptions + 1].storeOp = VK_ATTACHMENT_STORE_OP_STORE;
 			ResolveReferences[NumColorAttachments].attachment = NumAttachmentDescriptions + 1;
 			ResolveReferences[NumColorAttachments].layout = VK_IMAGE_LAYOUT_GENERAL;
 			++NumAttachmentDescriptions;
 			bHasResolveAttachments = true;
 		}
-#endif
 
 		CompatibleHashInfo.Formats[NumColorAttachments] = CurrDesc.format;
 		FullHashInfo.LoadOps[NumColorAttachments] = CurrDesc.loadOp;
@@ -2223,17 +2222,17 @@ FVulkanRenderTargetLayout::FVulkanRenderTargetLayout(const FGraphicsPipelineStat
 
 			ColorReferences[NumColorAttachments].attachment = NumAttachmentDescriptions;
 			ColorReferences[NumColorAttachments].layout = VK_IMAGE_LAYOUT_COLOR_ATTACHMENT_OPTIMAL;
-#if VULKAN_USE_MSAA_RESOLVE_ATTACHMENTS
-			if (CurrDesc.samples > VK_SAMPLE_COUNT_1_BIT)
+
+			if (CurrDesc.samples > VK_SAMPLE_COUNT_1_BIT && FVulkanPlatform::RequiresRenderPassResolveAttachments())
 			{
 				Desc[NumAttachmentDescriptions + 1] = Desc[NumAttachmentDescriptions];
 				Desc[NumAttachmentDescriptions + 1].samples = VK_SAMPLE_COUNT_1_BIT;
+				Desc[NumAttachmentDescriptions + 1].storeOp = VK_ATTACHMENT_STORE_OP_STORE;
 				ResolveReferences[NumColorAttachments].attachment = NumAttachmentDescriptions + 1;
 				ResolveReferences[NumColorAttachments].layout = VK_IMAGE_LAYOUT_GENERAL;
 				++NumAttachmentDescriptions;
 				bHasResolveAttachments = true;
 			}
-#endif
 
 			CompatibleHashInfo.Formats[NumColorAttachments] = CurrDesc.format;
 			FullHashInfo.LoadOps[NumColorAttachments] = CurrDesc.loadOp;
diff --git a/Engine/Source/Runtime/VulkanRHI/Private/VulkanResources.h b/Engine/Source/Runtime/VulkanRHI/Private/VulkanResources.h
--- a/Engine/Source/Runtime/VulkanRHI/Private/VulkanResources.h
+++ b/Engine/Source/Runtime/VulkanRHI/Private/VulkanResources.h
@@ -590,8 +590,8 @@ class FVulkanTexture2DArray : public FRHITexture2DArray, public FVulkanTextureBa
 {
 public:
 	// Constructor, just calls base and Surface constructor
-	FVulkanTexture2DArray(FVulkanDevice& Device, EPixelFormat Format, uint32 SizeX, uint32 SizeY, uint32 ArraySize, uint32 NumMips, uint32 Flags, FResourceBulkDataInterface* BulkData, const FClearValueBinding& InClearValue);
-	FVulkanTexture2DArray(FVulkanDevice& Device, EPixelFormat Format, uint32 SizeX, uint32 SizeY, uint32 ArraySize, uint32 NumMips, VkImage Image, uint32 Flags, FResourceBulkDataInterface* BulkData, const FClearValueBinding& InClearValue);
+	FVulkanTexture2DArray(FVulkanDevice& Device, EPixelFormat Format, uint32 SizeX, uint32 SizeY, uint32 ArraySize, uint32 NumMips, uint32 NumSamples, uint32 Flags, FResourceBulkDataInterface* BulkData, const FClearValueBinding& InClearValue);
+	FVulkanTexture2DArray(FVulkanDevice& Device, EPixelFormat Format, uint32 SizeX, uint32 SizeY, uint32 ArraySize, uint32 NumMips, uint32 NumSamples, VkImage Image, uint32 Flags, FResourceBulkDataInterface* BulkData, const FClearValueBinding& InClearValue);
 		
 	// IRefCountedObject interface.
 	virtual uint32 AddRef() const override final
diff --git a/Engine/Source/Runtime/VulkanRHI/Private/VulkanTexture.cpp b/Engine/Source/Runtime/VulkanRHI/Private/VulkanTexture.cpp
--- a/Engine/Source/Runtime/VulkanRHI/Private/VulkanTexture.cpp
+++ b/Engine/Source/Runtime/VulkanRHI/Private/VulkanTexture.cpp
@@ -944,10 +944,10 @@ void FVulkanDynamicRHI::RHICopySharedMips(FTexture2DRHIParamRef DestTexture2D,FT
 	VULKAN_SIGNAL_UNIMPLEMENTED();
 }
 
-FTexture2DArrayRHIRef FVulkanDynamicRHI::RHICreateTexture2DArray(uint32 SizeX, uint32 SizeY, uint32 SizeZ, uint8 Format, uint32 NumMips, uint32 Flags, FRHIResourceCreateInfo& CreateInfo)
+FTexture2DArrayRHIRef FVulkanDynamicRHI::RHICreateTexture2DArray(uint32 SizeX, uint32 SizeY, uint32 SizeZ, uint8 Format, uint32 NumMips, uint32 NumSamples, uint32 Flags, FRHIResourceCreateInfo& CreateInfo)
 {
 	LLM_SCOPE_VULKAN(ELLMTagVulkan::VulkanTextures);
-	return new FVulkanTexture2DArray(*Device, (EPixelFormat)Format, SizeX, SizeY, SizeZ, NumMips, Flags, CreateInfo.BulkData, CreateInfo.ClearValueBinding);
+	return new FVulkanTexture2DArray(*Device, (EPixelFormat)Format, SizeX, SizeY, SizeZ, NumMips, NumSamples, Flags, CreateInfo.BulkData, CreateInfo.ClearValueBinding);
 }
 
 FTexture3DRHIRef FVulkanDynamicRHI::RHICreateTexture3D(uint32 SizeX, uint32 SizeY, uint32 SizeZ, uint8 Format, uint32 NumMips, uint32 Flags, FRHIResourceCreateInfo& CreateInfo)
@@ -1830,15 +1830,15 @@ FVulkanBackBuffer::~FVulkanBackBuffer()
 }
 
 
-FVulkanTexture2DArray::FVulkanTexture2DArray(FVulkanDevice& Device, EPixelFormat Format, uint32 SizeX, uint32 SizeY, uint32 ArraySize, uint32 NumMips, uint32 Flags, FResourceBulkDataInterface* BulkData, const FClearValueBinding& InClearValue)
-	:	FRHITexture2DArray(SizeX, SizeY, ArraySize, NumMips, Format, Flags, InClearValue)
-	,	FVulkanTextureBase(Device, VK_IMAGE_VIEW_TYPE_2D_ARRAY, Format, SizeX, SizeY, 1, ArraySize, NumMips, /*NumSamples=*/ 1, Flags, BulkData)
+FVulkanTexture2DArray::FVulkanTexture2DArray(FVulkanDevice& Device, EPixelFormat Format, uint32 SizeX, uint32 SizeY, uint32 ArraySize, uint32 NumMips, uint32 NumSamples, uint32 Flags, FResourceBulkDataInterface* BulkData, const FClearValueBinding& InClearValue)
+	:	FRHITexture2DArray(SizeX, SizeY, ArraySize, NumMips, NumSamples, Format, Flags, InClearValue)
+	,	FVulkanTextureBase(Device, VK_IMAGE_VIEW_TYPE_2D_ARRAY, Format, SizeX, SizeY, 1, ArraySize, NumMips, NumSamples, Flags, BulkData)
 {
 }
 
-FVulkanTexture2DArray::FVulkanTexture2DArray(FVulkanDevice& Device, EPixelFormat Format, uint32 SizeX, uint32 SizeY, uint32 ArraySize, uint32 NumMips, VkImage Image, uint32 Flags, FResourceBulkDataInterface* BulkData, const FClearValueBinding& InClearValue)
-	:	FRHITexture2DArray(SizeX, SizeY, ArraySize, NumMips, Format, Flags, InClearValue)
-	,	FVulkanTextureBase(Device, VK_IMAGE_VIEW_TYPE_2D_ARRAY, Format, SizeX, SizeY, 1, ArraySize, NumMips, /*NumSamples=*/ 1, Image, VK_NULL_HANDLE, Flags, BulkData)
+FVulkanTexture2DArray::FVulkanTexture2DArray(FVulkanDevice& Device, EPixelFormat Format, uint32 SizeX, uint32 SizeY, uint32 ArraySize, uint32 NumMips, uint32 NumSamples, VkImage Image, uint32 Flags, FResourceBulkDataInterface* BulkData, const FClearValueBinding& InClearValue)
+	:	FRHITexture2DArray(SizeX, SizeY, ArraySize, NumMips, NumSamples, Format, Flags, InClearValue)
+	,	FVulkanTextureBase(Device, VK_IMAGE_VIEW_TYPE_2D_ARRAY, Format, SizeX, SizeY, 1, ArraySize, NumMips, NumSamples, Image, VK_NULL_HANDLE, Flags, BulkData)
 {
 }
 
diff --git a/Engine/Source/Runtime/VulkanRHI/Private/VulkanViewport.cpp b/Engine/Source/Runtime/VulkanRHI/Private/VulkanViewport.cpp
--- a/Engine/Source/Runtime/VulkanRHI/Private/VulkanViewport.cpp
+++ b/Engine/Source/Runtime/VulkanRHI/Private/VulkanViewport.cpp
@@ -323,6 +323,7 @@ FVulkanFramebuffer::FVulkanFramebuffer(FVulkanDevice& Device, const FRHISetRende
 	, DepthStencilRenderTargetImage(VK_NULL_HANDLE)
 {
 	FMemory::Memzero(ColorRenderTargetImages);
+	FMemory::Memzero(ColorResolveTargetImages);
 		
 	AttachmentTextureViews.Empty(RTLayout.GetNumAttachmentDescriptions());
 	uint32 MipIndex = 0;
@@ -368,6 +369,24 @@ FVulkanFramebuffer::FVulkanFramebuffer(FVulkanDevice& Device, const FRHISetRende
 		AttachmentViewsToDelete.Add(RTView.View);
 
 		++NumColorAttachments;
+
+		if (InRTInfo.bHasResolveAttachments)
+		{
+			FRHITexture* ResolveRHITexture = InRTInfo.ColorResolveRenderTarget[Index].Texture;
+			ColorResolveTargetImages[Index] = Texture->Surface.Image;
+			FVulkanTextureBase* ResolveTexture = FVulkanTextureBase::Cast(ResolveRHITexture);
+
+			//resolve attachments only supported for 2d/2d array textures
+			FVulkanTextureView ResolveRTView;
+			if (ResolveTexture->Surface.GetViewType() == VK_IMAGE_VIEW_TYPE_2D || ResolveTexture->Surface.GetViewType() == VK_IMAGE_VIEW_TYPE_2D_ARRAY)
+			{
+				ResolveRTView.Create(*ResolveTexture->Surface.Device, ResolveTexture->Surface.Image, ResolveTexture->Surface.GetViewType(), ResolveTexture->Surface.GetFullAspectMask(), ResolveTexture->Surface.PixelFormat, ResolveTexture->Surface.ViewFormat, 
+					MipIndex, 1, FMath::Max(0, (int32)InRTInfo.ColorRenderTarget[Index].ArraySliceIndex), ResolveTexture->Surface.GetNumberOfArrayLevels(), true);
+			}
+
+			AttachmentTextureViews.Add(ResolveRTView);
+			AttachmentViewsToDelete.Add(ResolveRTView.View);
+		}
 	}
 
 	if (RTLayout.GetHasDepthStencil())
@@ -461,6 +480,20 @@ bool FVulkanFramebuffer::Matches(const FRHISetRenderTargetsInfo& InRTInfo) const
 	int32 AttachementIndex = 0;
 	for (int32 Index = 0; Index < InRTInfo.NumColorRenderTargets; ++Index)
 	{
+		if (InRTInfo.bHasResolveAttachments)
+		{
+			const FRHIRenderTargetView& R = InRTInfo.ColorResolveRenderTarget[Index];
+			if (R.Texture)
+			{
+				VkImage AImage = ColorResolveTargetImages[AttachementIndex];
+				VkImage BImage = ((FVulkanTextureBase*)R.Texture->GetTextureBaseRHI())->Surface.Image;
+				if (AImage != BImage)
+				{
+					return false;
+				}
+			}
+		}
+
 		const FRHIRenderTargetView& B = InRTInfo.ColorRenderTarget[Index];
 		if (B.Texture)
 		{
diff --git a/Engine/Source/Runtime/VulkanRHI/Private/Windows/VulkanWindowsPlatform.h b/Engine/Source/Runtime/VulkanRHI/Private/Windows/VulkanWindowsPlatform.h
--- a/Engine/Source/Runtime/VulkanRHI/Private/Windows/VulkanWindowsPlatform.h
+++ b/Engine/Source/Runtime/VulkanRHI/Private/Windows/VulkanWindowsPlatform.h
@@ -70,6 +70,8 @@ public:
 		static auto* CVar = IConsoleManager::Get().FindTConsoleVariableDataInt(TEXT("r.Vulkan.UseRealUBs"));
 		return (CVar && CVar->GetValueOnAnyThread() == 0) ? false : bCodeHeaderUseRealUBs;
 	}
+
+	static bool RequiresRenderPassResolveAttachments() { return GMaxRHIFeatureLevel <= ERHIFeatureLevel::ES3_1; }
 };
 
 typedef FVulkanWindowsPlatform FVulkanPlatform;
diff --git a/Engine/Source/Runtime/VulkanRHI/Public/VulkanConfiguration.h b/Engine/Source/Runtime/VulkanRHI/Public/VulkanConfiguration.h
--- a/Engine/Source/Runtime/VulkanRHI/Public/VulkanConfiguration.h
+++ b/Engine/Source/Runtime/VulkanRHI/Public/VulkanConfiguration.h
@@ -87,8 +87,6 @@
 	#define VULKAN_HAS_PHYSICAL_DEVICE_PROPERTIES2				0
 #endif
 
-#define VULKAN_USE_MSAA_RESOLVE_ATTACHMENTS						0
-
 #define VULKAN_ENABLE_AGGRESSIVE_STATS							0
 
 #define VULKAN_REUSE_FENCES										1
diff --git a/Engine/Source/Runtime/VulkanRHI/Public/VulkanDynamicRHI.h b/Engine/Source/Runtime/VulkanRHI/Public/VulkanDynamicRHI.h
--- a/Engine/Source/Runtime/VulkanRHI/Public/VulkanDynamicRHI.h
+++ b/Engine/Source/Runtime/VulkanRHI/Public/VulkanDynamicRHI.h
@@ -89,7 +89,7 @@ public:
 	virtual FTexture2DRHIRef RHICreateTexture2D(uint32 SizeX, uint32 SizeY, uint8 Format, uint32 NumMips, uint32 NumSamples, uint32 Flags, FRHIResourceCreateInfo& CreateInfo) final override;
 	virtual FTexture2DRHIRef RHIAsyncCreateTexture2D(uint32 SizeX, uint32 SizeY, uint8 Format, uint32 NumMips, uint32 Flags, void** InitialMipData, uint32 NumInitialMips) final override;
 	virtual void RHICopySharedMips(FTexture2DRHIParamRef DestTexture2D, FTexture2DRHIParamRef SrcTexture2D) final override;
-	virtual FTexture2DArrayRHIRef RHICreateTexture2DArray(uint32 SizeX, uint32 SizeY, uint32 SizeZ, uint8 Format, uint32 NumMips, uint32 Flags, FRHIResourceCreateInfo& CreateInfo) final override;
+	virtual FTexture2DArrayRHIRef RHICreateTexture2DArray(uint32 SizeX, uint32 SizeY, uint32 SizeZ, uint8 Format, uint32 NumMips, uint32 NumSamples, uint32 Flags, FRHIResourceCreateInfo& CreateInfo) final override;
 	virtual FTexture3DRHIRef RHICreateTexture3D(uint32 SizeX, uint32 SizeY, uint32 SizeZ, uint8 Format, uint32 NumMips, uint32 Flags, FRHIResourceCreateInfo& CreateInfo) final override;
 	virtual void RHIGetResourceInfo(FTextureRHIParamRef Ref, FRHIResourceInfo& OutInfo) final override;
 	virtual FShaderResourceViewRHIRef RHICreateShaderResourceView(FTexture2DRHIParamRef Texture2DRHI, uint8 MipLevel) final override;
@@ -257,9 +257,9 @@ public:
 	{
 		return RHICreateTexture2D(SizeX, SizeY, Format, NumMips, NumSamples, Flags, CreateInfo);
 	}
-	virtual FTexture2DArrayRHIRef RHICreateTexture2DArray_RenderThread(class FRHICommandListImmediate& RHICmdList, uint32 SizeX, uint32 SizeY, uint32 SizeZ, uint8 Format, uint32 NumMips, uint32 Flags, FRHIResourceCreateInfo& CreateInfo) override final
+	virtual FTexture2DArrayRHIRef RHICreateTexture2DArray_RenderThread(class FRHICommandListImmediate& RHICmdList, uint32 SizeX, uint32 SizeY, uint32 SizeZ, uint8 Format, uint32 NumMips, uint32 NumSamples, uint32 Flags, FRHIResourceCreateInfo& CreateInfo) override final
 	{
-		return this->RHICreateTexture2DArray(SizeX, SizeY, SizeZ, Format, NumMips, Flags, CreateInfo);
+		return this->RHICreateTexture2DArray(SizeX, SizeY, SizeZ, Format, NumMips, NumSamples, Flags, CreateInfo);
 	}
 	virtual FTexture3DRHIRef RHICreateTexture3D_RenderThread(class FRHICommandListImmediate& RHICmdList, uint32 SizeX, uint32 SizeY, uint32 SizeZ, uint8 Format, uint32 NumMips, uint32 Flags, FRHIResourceCreateInfo& CreateInfo) override final
 	{
@@ -355,7 +355,7 @@ public:
 	// FVulkanDynamicRHI interface
 	virtual FTexture2DRHIRef RHICreateTexture2DFromResource(EPixelFormat Format, uint32 SizeX, uint32 SizeY, uint32 NumMips, uint32 NumSamples, VkImage Resource, uint32 Flags);
 	virtual FTexture2DRHIRef RHICreateTexture2DFromResource(EPixelFormat Format, uint32 SizeX, uint32 SizeY, uint32 NumMips, uint32 NumSamples, VkImage Resource, FSamplerYcbcrConversionInitializer& ConversionInitializer, uint32 Flags);
-	virtual FTexture2DArrayRHIRef RHICreateTexture2DArrayFromResource(EPixelFormat Format, uint32 SizeX, uint32 SizeY, uint32 ArraySize, uint32 NumMips, VkImage Resource, uint32 Flags);
+	virtual FTexture2DArrayRHIRef RHICreateTexture2DArrayFromResource(EPixelFormat Format, uint32 SizeX, uint32 SizeY, uint32 ArraySize, uint32 NumMips, uint32 NumSamples, VkImage Resource, uint32 Flags);
 	virtual FTextureCubeRHIRef RHICreateTextureCubeFromResource(EPixelFormat Format, uint32 Size, bool bArray, uint32 ArraySize, uint32 NumMips, VkImage Resource, uint32 Flags);
 
 	void RHICalibrateTimers() override;
diff --git a/Engine/Source/Runtime/Windows/D3D11RHI/Private/D3D11RHIPrivate.h b/Engine/Source/Runtime/Windows/D3D11RHI/Private/D3D11RHIPrivate.h
--- a/Engine/Source/Runtime/Windows/D3D11RHI/Private/D3D11RHIPrivate.h
+++ b/Engine/Source/Runtime/Windows/D3D11RHI/Private/D3D11RHIPrivate.h
@@ -441,7 +441,7 @@ public:
 	virtual FTexture2DRHIRef RHICreateTexture2D(uint32 SizeX, uint32 SizeY, uint8 Format, uint32 NumMips, uint32 NumSamples, uint32 Flags, FRHIResourceCreateInfo& CreateInfo) final override;
 	virtual FTexture2DRHIRef RHIAsyncCreateTexture2D(uint32 SizeX, uint32 SizeY, uint8 Format, uint32 NumMips, uint32 Flags, void** InitialMipData, uint32 NumInitialMips) final override;
 	virtual void RHICopySharedMips(FTexture2DRHIParamRef DestTexture2D, FTexture2DRHIParamRef SrcTexture2D) final override;
-	virtual FTexture2DArrayRHIRef RHICreateTexture2DArray(uint32 SizeX, uint32 SizeY, uint32 SizeZ, uint8 Format, uint32 NumMips, uint32 Flags, FRHIResourceCreateInfo& CreateInfo) final override;
+	virtual FTexture2DArrayRHIRef RHICreateTexture2DArray(uint32 SizeX, uint32 SizeY, uint32 SizeZ, uint8 Format, uint32 NumMips, uint32 NumSamples, uint32 Flags, FRHIResourceCreateInfo& CreateInfo) final override;
 	virtual FTexture3DRHIRef RHICreateTexture3D(uint32 SizeX, uint32 SizeY, uint32 SizeZ, uint8 Format, uint32 NumMips, uint32 Flags, FRHIResourceCreateInfo& CreateInfo) final override;
 	virtual void RHIGetResourceInfo(FTextureRHIParamRef Ref, FRHIResourceInfo& OutInfo) final override;
 	virtual FShaderResourceViewRHIRef RHICreateShaderResourceView(FTexture2DRHIParamRef Texture2DRHI, uint8 MipLevel) final override;
@@ -629,7 +629,7 @@ public:
 	virtual void EndUpdateTexture3D_RenderThread(class FRHICommandListImmediate& RHICmdList, FUpdateTexture3DData& UpdateData) final override;
 	virtual void UpdateTexture3D_RenderThread(class FRHICommandListImmediate& RHICmdList, FTexture3DRHIParamRef Texture, uint32 MipIndex, const struct FUpdateTextureRegion3D& UpdateRegion, uint32 SourceRowPitch, uint32 SourceDepthPitch, const uint8* SourceData) final override;
 	virtual FTexture2DRHIRef RHICreateTexture2D_RenderThread(class FRHICommandListImmediate& RHICmdList, uint32 SizeX, uint32 SizeY, uint8 Format, uint32 NumMips, uint32 NumSamples, uint32 Flags, FRHIResourceCreateInfo& CreateInfo) final override;
-	virtual FTexture2DArrayRHIRef RHICreateTexture2DArray_RenderThread(class FRHICommandListImmediate& RHICmdList, uint32 SizeX, uint32 SizeY, uint32 SizeZ, uint8 Format, uint32 NumMips, uint32 Flags, FRHIResourceCreateInfo& CreateInfo) final override;
+	virtual FTexture2DArrayRHIRef RHICreateTexture2DArray_RenderThread(class FRHICommandListImmediate& RHICmdList, uint32 SizeX, uint32 SizeY, uint32 SizeZ, uint8 Format, uint32 NumMips, uint32 NumSamples, uint32 Flags, FRHIResourceCreateInfo& CreateInfo) final override;
 	virtual FTexture3DRHIRef RHICreateTexture3D_RenderThread(class FRHICommandListImmediate& RHICmdList, uint32 SizeX, uint32 SizeY, uint32 SizeZ, uint8 Format, uint32 NumMips, uint32 Flags, FRHIResourceCreateInfo& CreateInfo) final override;
 	virtual FUnorderedAccessViewRHIRef RHICreateUnorderedAccessView_RenderThread(class FRHICommandListImmediate& RHICmdList, FStructuredBufferRHIParamRef StructuredBuffer, bool bUseUAVCounter, bool bAppendBuffer) final override;
 	virtual FUnorderedAccessViewRHIRef RHICreateUnorderedAccessView_RenderThread(class FRHICommandListImmediate& RHICmdList, FTextureRHIParamRef Texture, uint32 MipLevel) final override;
diff --git a/Engine/Source/Runtime/Windows/D3D11RHI/Private/D3D11Texture.cpp b/Engine/Source/Runtime/Windows/D3D11RHI/Private/D3D11Texture.cpp
--- a/Engine/Source/Runtime/Windows/D3D11RHI/Private/D3D11Texture.cpp
+++ b/Engine/Source/Runtime/Windows/D3D11RHI/Private/D3D11Texture.cpp
@@ -1257,10 +1257,10 @@ void FD3D11DynamicRHI::RHICopySharedMips(FTexture2DRHIParamRef DestTexture2DRHI,
 	}
 }
 
-FTexture2DArrayRHIRef FD3D11DynamicRHI::RHICreateTexture2DArray(uint32 SizeX,uint32 SizeY,uint32 SizeZ,uint8 Format,uint32 NumMips,uint32 Flags,FRHIResourceCreateInfo& CreateInfo)
+FTexture2DArrayRHIRef FD3D11DynamicRHI::RHICreateTexture2DArray(uint32 SizeX,uint32 SizeY,uint32 SizeZ,uint8 Format,uint32 NumMips,uint32 NumSamples,uint32 Flags,FRHIResourceCreateInfo& CreateInfo)
 {
 	check(SizeZ >= 1);
-	return CreateD3D11Texture2D<FD3D11BaseTexture2DArray>(SizeX,SizeY,SizeZ,true,false,Format,NumMips,1,Flags,CreateInfo);
+	return CreateD3D11Texture2D<FD3D11BaseTexture2DArray>(SizeX,SizeY,SizeZ,true,false,Format,NumMips,NumSamples,Flags,CreateInfo);
 }
 
 FTexture2DArrayRHIRef FD3D11DynamicRHI::RHICreateTexture2DArray_RenderThread(
@@ -1270,10 +1270,11 @@ FTexture2DArrayRHIRef FD3D11DynamicRHI::RHICreateTexture2DArray_RenderThread(
 	uint32 SizeZ,
 	uint8 Format,
 	uint32 NumMips,
+	uint32 NumSamples,
 	uint32 Flags,
 	FRHIResourceCreateInfo& CreateInfo)
 {
-	return RHICreateTexture2DArray(SizeX, SizeY, SizeZ, Format, NumMips, Flags, CreateInfo);
+	return RHICreateTexture2DArray(SizeX, SizeY, SizeZ, Format, NumMips, NumSamples, Flags, CreateInfo);
 }
 
 FTexture3DRHIRef FD3D11DynamicRHI::RHICreateTexture3D(uint32 SizeX,uint32 SizeY,uint32 SizeZ,uint8 Format,uint32 NumMips,uint32 Flags,FRHIResourceCreateInfo& CreateInfo)
diff --git a/Engine/Source/Runtime/Windows/D3D11RHI/Public/D3D11Resources.h b/Engine/Source/Runtime/Windows/D3D11RHI/Public/D3D11Resources.h
--- a/Engine/Source/Runtime/Windows/D3D11RHI/Public/D3D11Resources.h
+++ b/Engine/Source/Runtime/Windows/D3D11RHI/Public/D3D11Resources.h
@@ -543,7 +543,7 @@ class FD3D11BaseTexture2DArray : public FRHITexture2DArray
 {
 public:
 	FD3D11BaseTexture2DArray(uint32 InSizeX, uint32 InSizeY, uint32 InSizeZ, uint32 InNumMips, uint32 InNumSamples, EPixelFormat InFormat, uint32 InFlags, const FClearValueBinding& InClearValue)
-	: FRHITexture2DArray(InSizeX,InSizeY,InSizeZ,InNumMips,InFormat,InFlags,InClearValue)
+	: FRHITexture2DArray(InSizeX,InSizeY,InSizeZ,InNumMips,InNumSamples, InFormat,InFlags,InClearValue)
 	{ check(InNumSamples == 1); }
 };
 
