Set the multi-view mask for all subpasses

diff --git a/Engine/Source/Runtime/VulkanRHI/Private/VulkanRHI.cpp b/Engine/Source/Runtime/VulkanRHI/Private/VulkanRHI.cpp
--- a/Engine/Source/Runtime/VulkanRHI/Private/VulkanRHI.cpp
+++ b/Engine/Source/Runtime/VulkanRHI/Private/VulkanRHI.cpp
@@ -1452,7 +1452,7 @@ static VkRenderPass CreateRenderPass(FVulkanDevice& InDevice, const FVulkanRende
 	Bit mask that specifies which view rendering is broadcast to
 	0011 = Broadcast to first and second view (layer)
 	*/
-	const uint32_t ViewMask = 0b00000011;
+	const uint32_t ViewMask[2] = { 0b00000011, 0b00000011 };
 
 	/*
 	Bit mask that specifices correlation between views
@@ -1466,8 +1466,8 @@ static VkRenderPass CreateRenderPass(FVulkanDevice& InDevice, const FVulkanRende
 		FMemory::Memzero(MultiviewInfo);
 		MultiviewInfo.sType = VK_STRUCTURE_TYPE_RENDER_PASS_MULTIVIEW_CREATE_INFO;
 		MultiviewInfo.pNext = nullptr;
-		MultiviewInfo.subpassCount = 1;
-		MultiviewInfo.pViewMasks = &ViewMask;
+		MultiviewInfo.subpassCount = NumSubpasses;
+		MultiviewInfo.pViewMasks = ViewMask;
 		MultiviewInfo.dependencyCount = 0;
 		MultiviewInfo.pViewOffsets = nullptr;
 		MultiviewInfo.correlationMaskCount = 1;
