DR - RHI validation layer for Vulkan
#rnx

diff --git a/Engine/Source/Runtime/VulkanRHI/Private/VulkanDevice.cpp b/Engine/Source/Runtime/VulkanRHI/Private/VulkanDevice.cpp
--- a/Engine/Source/Runtime/VulkanRHI/Private/VulkanDevice.cpp
+++ b/Engine/Source/Runtime/VulkanRHI/Private/VulkanDevice.cpp
@@ -356,7 +356,7 @@ void FVulkanDevice::CreateDevice()
 
 #if VULKAN_ENABLE_DRAW_MARKERS
 #if 0//VULKAN_SUPPORTS_DEBUG_UTILS
-	FVulkanDynamicRHI* RHI = (FVulkanDynamicRHI*)GDynamicRHI;
+	FVulkanDynamicRHI* RHI = GVulkanRHI;
 	if (RHI->SupportsDebugUtilsExt() && GRenderDocFound)
 	{
 		DebugMarkers.CmdBeginDebugLabel = (PFN_vkCmdBeginDebugUtilsLabelEXT)(void*)VulkanRHI::vkGetInstanceProcAddr(RHI->GetInstance(), "vkCmdBeginDebugUtilsLabelEXT");
@@ -862,11 +862,11 @@ void FVulkanDevice::InitGPU(int32 DeviceIndex)
 	// always look in the saved directory (for the cache from previous run that wasn't moved over to stage directory)
 	CacheFilenames.Add(VulkanRHI::GetPipelineCacheFilename());
 
-	ImmediateContext = new FVulkanCommandListContextImmediate((FVulkanDynamicRHI*)GDynamicRHI, this, GfxQueue);
+	ImmediateContext = new FVulkanCommandListContextImmediate(GVulkanRHI, this, GfxQueue);
 
 	if (GfxQueue->GetFamilyIndex() != ComputeQueue->GetFamilyIndex() && GRHIAllowAsyncComputeCvar.GetValueOnAnyThread() != 0)
 	{
-		ComputeContext = new FVulkanCommandListContextImmediate((FVulkanDynamicRHI*)GDynamicRHI, this, ComputeQueue);
+		ComputeContext = new FVulkanCommandListContextImmediate(GVulkanRHI, this, ComputeQueue);
 		GEnableAsyncCompute = true;
 	}
 	else
@@ -880,7 +880,7 @@ void FVulkanDevice::InitGPU(int32 DeviceIndex)
 		int32 Num = FTaskGraphInterface::Get().GetNumWorkerThreads();
 		for (int32 Index = 0; Index < Num; Index++)
 		{
-			FVulkanCommandListContext* CmdContext = new FVulkanCommandListContext((FVulkanDynamicRHI*)GDynamicRHI, this, GfxQueue, ImmediateContext);
+			FVulkanCommandListContext* CmdContext = new FVulkanCommandListContext(GVulkanRHI, this, GfxQueue, ImmediateContext);
 			CommandContexts.Add(CmdContext);
 		}
 	}
@@ -1171,7 +1171,7 @@ FVulkanCommandListContext* FVulkanDevice::AcquireDeferredContext()
 	FScopeLock Lock(&GContextCS);
 	if (CommandContexts.Num() == 0)
 	{
-		return new FVulkanCommandListContext((FVulkanDynamicRHI*)GDynamicRHI, this, GfxQueue, ImmediateContext);
+		return new FVulkanCommandListContext(GVulkanRHI, this, GfxQueue, ImmediateContext);
 	}
 	return CommandContexts.Pop(false);
 }
diff --git a/Engine/Source/Runtime/VulkanRHI/Private/VulkanIndexBuffer.cpp b/Engine/Source/Runtime/VulkanRHI/Private/VulkanIndexBuffer.cpp
--- a/Engine/Source/Runtime/VulkanRHI/Private/VulkanIndexBuffer.cpp
+++ b/Engine/Source/Runtime/VulkanRHI/Private/VulkanIndexBuffer.cpp
@@ -271,7 +271,7 @@ struct FRHICommandMultiBufferUnlock final : public FRHICommand<FRHICommandMultiB
 
 	void Execute(FRHICommandListBase& CmdList)
 	{
-		FVulkanResourceMultiBuffer::InternalUnlock((FVulkanCommandListContext&)CmdList.GetContext(), PendingLock, MultiBuffer, DynamicBufferIndex);
+		FVulkanResourceMultiBuffer::InternalUnlock(VulkanRHI::GetVulkanContext(CmdList.GetContext()), PendingLock, MultiBuffer, DynamicBufferIndex);
 	}
 };
 
diff --git a/Engine/Source/Runtime/VulkanRHI/Private/VulkanRHI.cpp b/Engine/Source/Runtime/VulkanRHI/Private/VulkanRHI.cpp
--- a/Engine/Source/Runtime/VulkanRHI/Private/VulkanRHI.cpp
+++ b/Engine/Source/Runtime/VulkanRHI/Private/VulkanRHI.cpp
@@ -18,6 +18,10 @@
 #include "VulkanLLM.h"
 #include "Misc/EngineVersion.h"
 #include "GlobalShader.h"
+#include "ValidationRHI.h"
+
+
+FValidationRHI*		GValidationRHI = nullptr;
 
 extern RHI_API bool GUseTexture3DBulkDataRHI;
 
@@ -91,7 +95,18 @@ FDynamicRHI* FVulkanDynamicRHIModule::CreateRHI(ERHIFeatureLevel::Type InRequest
 	// VULKAN_USE_MSAA_RESOLVE_ATTACHMENTS=0 requires separate MSAA and resolve textures
 	check(RHISupportsSeparateMSAAAndResolveTextures(GMaxRHIShaderPlatform) == (!VULKAN_USE_MSAA_RESOLVE_ATTACHMENTS));
 
-	return new FVulkanDynamicRHI();
+	GVulkanRHI = new FVulkanDynamicRHI();
+
+	if (FParse::Param(FCommandLine::Get(), TEXT("RHIValidation")))
+	{
+		GValidationRHI = new FValidationRHI(GVulkanRHI);
+	}
+	else
+	{
+		check(!GValidationRHI);
+	}
+
+	return GValidationRHI ? (FDynamicRHI*)GValidationRHI : (FDynamicRHI*)GVulkanRHI;
 }
 
 IMPLEMENT_MODULE(FVulkanDynamicRHIModule, VulkanRHI);
@@ -683,7 +698,7 @@ void FVulkanDynamicRHI::InitInstance()
 
 		GUseTexture3DBulkDataRHI = true;
 
-		GDynamicRHI = this;
+		GDynamicRHI = GValidationRHI ? (FDynamicRHI*)GValidationRHI : (FDynamicRHI*)this;
 
 		// Notify all initialized FRenderResources that there's a valid RHI device to create their RHI resources for now.
 		for (TLinkedList<FRenderResource*>::TIterator ResourceIt(FRenderResource::GetResourceList()); ResourceIt; ResourceIt.Next())
@@ -1541,24 +1556,21 @@ void FVulkanDynamicRHI::SavePipelineCache()
 {
 	FString CacheFile = GetPipelineCacheFilename();
 
-	FVulkanDynamicRHI* RHI = (FVulkanDynamicRHI*)GDynamicRHI;
-	RHI->Device->PipelineStateCache->Save(CacheFile);
+	GVulkanRHI->Device->PipelineStateCache->Save(CacheFile);
 }
 
 void FVulkanDynamicRHI::RebuildPipelineCache()
 {
-	FVulkanDynamicRHI* RHI = (FVulkanDynamicRHI*)GDynamicRHI;
-	RHI->Device->PipelineStateCache->RebuildCache();
+	GVulkanRHI->Device->PipelineStateCache->RebuildCache();
 }
 
 #if VULKAN_SUPPORTS_VALIDATION_CACHE
 void FVulkanDynamicRHI::SaveValidationCache()
 {
-	FVulkanDynamicRHI* RHI = (FVulkanDynamicRHI*)GDynamicRHI;
-	VkValidationCacheEXT ValidationCache = RHI->Device->GetValidationCache();
+	VkValidationCacheEXT ValidationCache = GVulkanRHI->Device->GetValidationCache();
 	if (ValidationCache != VK_NULL_HANDLE)
 	{
-		VkDevice Device = RHI->Device->GetInstanceHandle();
+		VkDevice Device = GVulkanRHI->Device->GetInstanceHandle();
 		PFN_vkGetValidationCacheDataEXT vkGetValidationCacheData = (PFN_vkGetValidationCacheDataEXT)(void*)VulkanRHI::vkGetDeviceProcAddr(Device, "vkGetValidationCacheDataEXT");
 		check(vkGetValidationCacheData);
 		size_t CacheSize = 0;
@@ -1595,10 +1607,9 @@ void FVulkanDynamicRHI::SaveValidationCache()
 #if UE_BUILD_DEBUG || UE_BUILD_DEVELOPMENT
 void FVulkanDynamicRHI::DumpMemory()
 {
-	FVulkanDynamicRHI* RHI = (FVulkanDynamicRHI*)GDynamicRHI;
-	RHI->Device->GetMemoryManager().DumpMemory();
-	RHI->Device->GetResourceHeapManager().DumpMemory();
-	RHI->Device->GetStagingManager().DumpMemory();
+	GVulkanRHI->Device->GetMemoryManager().DumpMemory();
+	GVulkanRHI->Device->GetResourceHeapManager().DumpMemory();
+	GVulkanRHI->Device->GetStagingManager().DumpMemory();
 }
 #endif
 
@@ -1607,8 +1618,7 @@ void FVulkanDynamicRHI::RecreateSwapChain(void* NewNativeWindow)
 	if (NewNativeWindow)
 	{
 		FlushRenderingCommands();
-		FVulkanDynamicRHI* RHI = (FVulkanDynamicRHI*)GDynamicRHI;
-		TArray<FVulkanViewport*> Viewports = RHI->Viewports;
+		TArray<FVulkanViewport*> Viewports = GVulkanRHI->Viewports;
 		ENQUEUE_RENDER_COMMAND(VulkanRecreateSwapChain)(
 			[Viewports, NewNativeWindow](FRHICommandListImmediate& RHICmdList)
 			{
diff --git a/Engine/Source/Runtime/VulkanRHI/Private/VulkanRHIPrivate.h b/Engine/Source/Runtime/VulkanRHI/Private/VulkanRHIPrivate.h
--- a/Engine/Source/Runtime/VulkanRHI/Private/VulkanRHIPrivate.h
+++ b/Engine/Source/Runtime/VulkanRHI/Private/VulkanRHIPrivate.h
@@ -53,6 +53,8 @@ class FVulkanGfxPipeline;
 class FVulkanRenderPass;
 class FVulkanCommandBufferManager;
 struct FInputAttachmentData;
+class FValidationContext;
+
 
 inline VkShaderStageFlagBits UEFrequencyToVKStageBit(EShaderFrequency InStage)
 {
@@ -90,6 +92,7 @@ inline EShaderFrequency VkStageBitToUEFrequency(VkShaderStageFlagBits FlagBits)
 	return SF_NumFrequencies;
 }
 
+
 class FVulkanRenderTargetLayout
 {
 public:
@@ -288,6 +291,9 @@ private:
 	FVulkanDevice&				Device;
 };
 
+extern class FValidationRHI*	GValidationRHI;
+
+
 namespace VulkanRHI
 {
 	inline void SetupImageBarrierOLD(VkImageMemoryBarrier& Barrier, const FVulkanSurface& Surface, VkAccessFlags SrcMask, VkImageLayout SrcLayout, VkAccessFlags DstMask, VkImageLayout DstLayout, uint32 NumLayers = 1)
@@ -829,6 +835,18 @@ namespace VulkanRHI
 		}
 #endif
 	}
+
+	FVulkanCommandListContext& GetVulkanContext(FValidationContext& CmdContext);
+
+	inline FVulkanCommandListContext& GetVulkanContext(IRHICommandContext& CmdContext)
+	{
+		if (GValidationRHI)
+		{
+			return GetVulkanContext((FValidationContext&)CmdContext);
+		}
+
+		return (FVulkanCommandListContext&)CmdContext;
+	}
 }
 
 extern int32 GVulkanSubmitAfterEveryEndRenderPass;
@@ -840,3 +858,5 @@ extern bool GRenderDocFound;
 #endif
 
 const int GMaxCrashBufferEntries = 2048;
+
+extern class FVulkanDynamicRHI*	GVulkanRHI;
diff --git a/Engine/Source/Runtime/VulkanRHI/Private/VulkanTexture.cpp b/Engine/Source/Runtime/VulkanRHI/Private/VulkanTexture.cpp
--- a/Engine/Source/Runtime/VulkanRHI/Private/VulkanTexture.cpp
+++ b/Engine/Source/Runtime/VulkanRHI/Private/VulkanTexture.cpp
@@ -170,7 +170,7 @@ struct FRHICommandLockWriteTexture final : public FRHICommand<FRHICommandLockWri
 
 	void Execute(FRHICommandListBase& RHICmdList)
 	{
-		FVulkanSurface::InternalLockWrite((FVulkanCommandListContext&)RHICmdList.GetContext(), Surface, SubresourceRange, Region, StagingBuffer);
+		FVulkanSurface::InternalLockWrite(VulkanRHI::GetVulkanContext(RHICmdList.GetContext()), Surface, SubresourceRange, Region, StagingBuffer);
 	}
 };
 
@@ -423,7 +423,7 @@ struct FRHICommandInitialClearTexture final : public FRHICommand<FRHICommandInit
 
 	void Execute(FRHICommandListBase& CmdList)
 	{
-		Surface->InitialClear((FVulkanCommandListContext&)CmdList.GetContext(), ClearValueBinding, bTransitionToPresentable);
+		Surface->InitialClear(VulkanRHI::GetVulkanContext(CmdList.GetContext()), ClearValueBinding, bTransitionToPresentable);
 	}
 };
 
@@ -440,7 +440,8 @@ struct FRHICommandRegisterImageLayout final : public FRHICommand<FRHICommandRegi
 
 	void Execute(FRHICommandListBase& RHICmdList)
 	{
-		((FVulkanCommandListContext&)RHICmdList.GetContext()).FindOrAddLayout(Image, ImageLayout);
+		auto& Context = RHICmdList.GetContext();
+		VulkanRHI::GetVulkanContext(Context).FindOrAddLayout(Image, ImageLayout);
 	}
 };
 
diff --git a/Engine/Source/Runtime/VulkanRHI/Private/VulkanUtil.cpp b/Engine/Source/Runtime/VulkanRHI/Private/VulkanUtil.cpp
--- a/Engine/Source/Runtime/VulkanRHI/Private/VulkanUtil.cpp
+++ b/Engine/Source/Runtime/VulkanRHI/Private/VulkanUtil.cpp
@@ -10,7 +10,9 @@
 #include "VulkanContext.h"
 #include "VulkanMemory.h"
 #include "Misc/OutputDeviceRedirector.h"
+#include "ValidationContext.h"
 
+FVulkanDynamicRHI*	GVulkanRHI = nullptr;
 
 extern CORE_API bool GIsGPUCrashed;
 
@@ -585,10 +587,21 @@ namespace VulkanRHIBridge
 #include "RHIStaticStates.h"
 #include "OneColorShader.h"
 #include "VulkanRHI.h"
+#include "ValidationRHI.h"
 #include "StaticBoundShaderState.h"
 
 namespace VulkanRHI
 {
+	FVulkanCommandListContext& GetVulkanContext(class FValidationContext& CmdContext)
+	{
+		if (GValidationRHI && GValidationRHI->Context == &CmdContext)
+		{
+			return (FVulkanCommandListContext&)*CmdContext.RHIContext;
+		}
+
+		return (FVulkanCommandListContext&)CmdContext;
+	}
+
 	VkBuffer CreateBuffer(FVulkanDevice* InDevice, VkDeviceSize Size, VkBufferUsageFlags BufferUsageFlags, VkMemoryRequirements& OutMemoryRequirements)
 	{
 		VkDevice Device = InDevice->GetInstanceHandle();
@@ -665,8 +678,7 @@ namespace VulkanRHI
 #if VULKAN_SUPPORTS_GPU_CRASH_DUMPS
 		if (GIsGPUCrashed && GGPUCrashDebuggingEnabled)
 		{
-			FVulkanDynamicRHI* RHI = (FVulkanDynamicRHI*)GDynamicRHI;
-			FVulkanDevice* Device = RHI->GetDevice();
+			FVulkanDevice* Device = GVulkanRHI->GetDevice();
 			if (Device->GetOptionalExtensions().HasGPUCrashDumpExtensions())
 			{
 				Device->GetImmediateContext().GetGPUProfiler().DumpCrashMarkers(Device->GetCrashMarkerMappedPointer());
diff --git a/Engine/Source/Runtime/VulkanRHI/Private/VulkanViewport.cpp b/Engine/Source/Runtime/VulkanRHI/Private/VulkanViewport.cpp
--- a/Engine/Source/Runtime/VulkanRHI/Private/VulkanViewport.cpp
+++ b/Engine/Source/Runtime/VulkanRHI/Private/VulkanViewport.cpp
@@ -870,7 +870,7 @@ bool FVulkanViewport::Present(FVulkanCommandListContext* Context, FVulkanCmdBuff
 	AcquiredImageIndex = -1;
 
 	++PresentCount;
-	++((FVulkanDynamicRHI*)GDynamicRHI)->TotalPresentCount;
+	++GVulkanRHI->TotalPresentCount;
 
 	return bResult;
 }
