Vulkan:
-Allocations Larger than a certain threshold are now given full size allocations and not suballocated
-Macros to fake lower memory budget
-support for EXT_memory_priority

diff --git a/Engine/Source/Runtime/VulkanRHI/Private/VulkanDevice.cpp b/Engine/Source/Runtime/VulkanRHI/Private/VulkanDevice.cpp
--- a/Engine/Source/Runtime/VulkanRHI/Private/VulkanDevice.cpp
+++ b/Engine/Source/Runtime/VulkanRHI/Private/VulkanDevice.cpp
@@ -815,7 +815,7 @@ void FVulkanDevice::InitGPU(int32 DeviceIndex)
 			VulkanRHI::vkGetBufferMemoryRequirements(Device, CrashMarker.Buffer, &MemReq);
 
 			CrashMarker.Allocation = MemoryManager.Alloc(false, CreateInfo.size, MemReq.memoryTypeBits, VK_MEMORY_PROPERTY_HOST_VISIBLE_BIT |
-				VK_MEMORY_PROPERTY_HOST_COHERENT_BIT, nullptr, __FILE__, __LINE__);
+				VK_MEMORY_PROPERTY_HOST_COHERENT_BIT, nullptr, VULKAN_MEMORY_MEDIUM_PRIORITY, __FILE__, __LINE__);
 
 			uint32* Entry = (uint32*)CrashMarker.Allocation->Map(VK_WHOLE_SIZE, 0);
 			check(Entry);
@@ -826,7 +826,7 @@ void FVulkanDevice::InitGPU(int32 DeviceIndex)
 		else if (OptionalDeviceExtensions.HasNVDiagnosticCheckpoints)
 		{
 			CrashMarker.Allocation = MemoryManager.Alloc(false, GMaxCrashBufferEntries * sizeof(uint32_t), UINT32_MAX, VK_MEMORY_PROPERTY_HOST_VISIBLE_BIT |
-				VK_MEMORY_PROPERTY_HOST_COHERENT_BIT, nullptr, __FILE__, __LINE__);
+				VK_MEMORY_PROPERTY_HOST_COHERENT_BIT, nullptr, VULKAN_MEMORY_MEDIUM_PRIORITY, __FILE__, __LINE__);
 			uint32* Entry = (uint32*)CrashMarker.Allocation->Map(VK_WHOLE_SIZE, 0);
 			check(Entry);
 			// Start with 0 entries
diff --git a/Engine/Source/Runtime/VulkanRHI/Private/VulkanDevice.h b/Engine/Source/Runtime/VulkanRHI/Private/VulkanDevice.h
--- a/Engine/Source/Runtime/VulkanRHI/Private/VulkanDevice.h
+++ b/Engine/Source/Runtime/VulkanRHI/Private/VulkanDevice.h
@@ -31,6 +31,7 @@ struct FOptionalVulkanDeviceExtensions
 	uint32 HasNVDiagnosticCheckpoints : 1;
 	uint32 HasGoogleDisplayTiming : 1;
 	uint32 HasYcbcrSampler : 1;
+	uint32 HasMemoryPriority : 1;
 
 	inline bool HasGPUCrashDumpExtensions() const
 	{
diff --git a/Engine/Source/Runtime/VulkanRHI/Private/VulkanLayers.cpp b/Engine/Source/Runtime/VulkanRHI/Private/VulkanLayers.cpp
--- a/Engine/Source/Runtime/VulkanRHI/Private/VulkanLayers.cpp
+++ b/Engine/Source/Runtime/VulkanRHI/Private/VulkanLayers.cpp
@@ -82,6 +82,10 @@ static const ANSICHAR* GDeviceExtensions[] =
 #if VULKAN_SUPPORTS_VALIDATION_CACHE
 	VK_EXT_VALIDATION_CACHE_EXTENSION_NAME,
 #endif
+#if VULKAN_SUPPORTS_MEMORY_PRIORITY
+	VK_EXT_MEMORY_PRIORITY_EXTENSION_NAME,
+#endif
+
 	//VK_KHR_SAMPLER_MIRROR_CLAMP_TO_EDGE_EXTENSION_NAME,
 	nullptr
 };
@@ -725,4 +729,14 @@ void FVulkanDevice::ParseOptionalDeviceExtensions(const TArray<const ANSICHAR *>
 #if VULKAN_SUPPORTS_COLOR_CONVERSIONS
 	OptionalDeviceExtensions.HasYcbcrSampler = HasExtension(VK_KHR_SAMPLER_YCBCR_CONVERSION_EXTENSION_NAME) && HasExtension(VK_KHR_BIND_MEMORY_2_EXTENSION_NAME) && HasExtension(VK_KHR_GET_MEMORY_REQUIREMENTS_2_EXTENSION_NAME);
 #endif
+
+#if VULKAN_SUPPORTS_MEMORY_PRIORITY
+	OptionalDeviceExtensions.HasMemoryPriority = HasExtension(VK_EXT_MEMORY_PRIORITY_EXTENSION_NAME);
+	if (FParse::Param(FCommandLine::Get(), TEXT("disablememorypriority")))
+	{
+		OptionalDeviceExtensions.HasMemoryPriority = 0;
+	}
+#else
+	OptionalDeviceExtensions.HasMemoryPriority = 0;
+#endif
 }
diff --git a/Engine/Source/Runtime/VulkanRHI/Private/VulkanMemory.cpp b/Engine/Source/Runtime/VulkanRHI/Private/VulkanMemory.cpp
--- a/Engine/Source/Runtime/VulkanRHI/Private/VulkanMemory.cpp
+++ b/Engine/Source/Runtime/VulkanRHI/Private/VulkanMemory.cpp
@@ -14,6 +14,9 @@
 uint32 GVulkanRHIDeletionFrameNumber = 0;
 const uint32 NUM_FRAMES_TO_WAIT_FOR_RESOURCE_DELETE = 2;
 
+#define VULKAN_MAX_SUB_ALLOCATION (64llu << 20llu) // set to 0 to disable
+#define VULKAN_FAKE_MEMORY_LIMIT 0llu /// set to # of GB to fake out of memory when hitting limit.
+
 
 #if UE_BUILD_DEBUG
 RENDERCORE_API	void DumpRenderTargetPoolMemory(FOutputDevice& OutputDevice);
@@ -74,6 +77,9 @@ namespace VulkanRHI
 	static FCriticalSection GFenceLock;
 	static FCriticalSection GResourceHeapLock;
 
+#if VULKAN_FAKE_MEMORY_LIMIT
+	static uint64 GDeviceMemAllocated = 0;
+#endif
 	FDeviceMemoryManager::FDeviceMemoryManager() :
 		DeviceHandle(VK_NULL_HANDLE),
 		bHasUnifiedMemory(false),
@@ -195,7 +201,7 @@ namespace VulkanRHI
 		return false;
 	}
 
-	FDeviceMemoryAllocation* FDeviceMemoryManager::Alloc(bool bCanFail, VkDeviceSize AllocationSize, uint32 MemoryTypeIndex, void* DedicatedAllocateInfo, const char* File, uint32 Line)
+	FDeviceMemoryAllocation* FDeviceMemoryManager::Alloc(bool bCanFail, VkDeviceSize AllocationSize, uint32 MemoryTypeIndex, void* DedicatedAllocateInfo, float Priority, const char* File, uint32 Line)
 	{
 		FScopeLock Lock(&GDeviceMemLock);
 
@@ -207,14 +213,41 @@ namespace VulkanRHI
 		Info.allocationSize = AllocationSize;
 		Info.memoryTypeIndex = MemoryTypeIndex;
 
+		
+#if VULKAN_SUPPORTS_MEMORY_PRIORITY
+		VkMemoryPriorityAllocateInfoEXT Prio;
+		ZeroVulkanStruct(Prio, VK_STRUCTURE_TYPE_MEMORY_PRIORITY_ALLOCATE_INFO_EXT);
+		Prio.priority = Priority;
+		if (Device->GetOptionalExtensions().HasMemoryPriority)
+		{
+			Info.pNext = &Prio;
+		}
+#endif
+
 #if VULKAN_SUPPORTS_DEDICATED_ALLOCATION
-		Info.pNext = DedicatedAllocateInfo;
+	((VkMemoryDedicatedAllocateInfoKHR*)DedicatedAllocateInfo)->pNext = Info.pNext;
+	Info.pNext = DedicatedAllocateInfo;
 #else
-		check(!DedicatedAllocateInfo);
+	check(!DedicatedAllocateInfo);
 #endif
 
 		VkDeviceMemory Handle;
-		VkResult Result = VulkanRHI::vkAllocateMemory(DeviceHandle, &Info, VULKAN_CPU_ALLOCATOR, &Handle);
+		VkResult Result;
+#if VULKAN_FAKE_MEMORY_LIMIT
+		GDeviceMemAllocated += AllocationSize;
+		if( (VULKAN_FAKE_MEMORY_LIMIT << 30llu)  < (int64)GDeviceMemAllocated)
+		{
+			Handle = 0;
+			Result = VK_ERROR_OUT_OF_DEVICE_MEMORY;
+			FPlatformMisc::LowLevelOutputDebugStringf(TEXT("DEV MEM %6.2f/%6.2f\n"), GDeviceMemAllocated / (1024.f *1024.f), VULKAN_FAKE_MEMORY_LIMIT *1024.f);
+			GDeviceMemAllocated -= AllocationSize;
+
+		}
+		else
+#endif
+		{
+			Result = VulkanRHI::vkAllocateMemory(DeviceHandle, &Info, VULKAN_CPU_ALLOCATOR, &Handle);
+		}
 		
 		//FPlatformMisc::LowLevelOutputDebugStringf(TEXT("(%u)(%x) VulkanRHI::vkAllocateMemory size=%u(%6.3fmb) Type=%u\n"),
 		//	GFrameNumberRenderThread, this, AllocationSize, AllocationSize / (1024.f * 1024.f), MemoryTypeIndex);
@@ -314,6 +347,10 @@ namespace VulkanRHI
 		check(Allocation);
 		check(Allocation->Handle != VK_NULL_HANDLE);
 		check(!Allocation->bFreedBySystem);
+#if VULKAN_FAKE_MEMORY_LIMIT
+		GDeviceMemAllocated -= Allocation->Size;
+#endif
+
 		VulkanRHI::vkFreeMemory(DeviceHandle, Allocation->Handle, VULKAN_CPU_ALLOCATOR);
 
 #if VULKAN_USE_LLM
@@ -858,68 +895,79 @@ namespace VulkanRHI
 #else
 		uint32 TargetDefaultPageSize = DefaultPageSize;
 #endif
-		if (Size < TargetDefaultPageSize)
+		uint32 AllocationSize;
+		
 		{
-			// Check Used pages to see if we can fit this in
-			for (int32 Index = 0; Index < UsedPages.Num(); ++Index)
+			if (Size < TargetDefaultPageSize)
 			{
-				FOldResourceHeapPage* Page = UsedPages[Index];
-				if (Page->DeviceMemoryAllocation->IsMapped() == bMapAllocation)
+				// Check Used pages to see if we can fit this in
+				for (int32 Index = 0; Index < UsedPages.Num(); ++Index)
 				{
-					FOldResourceAllocation* ResourceAllocation = Page->TryAllocate(Size, Alignment, File, Line);
-					if (ResourceAllocation)
+					FOldResourceHeapPage* Page = UsedPages[Index];
+					if (Page->DeviceMemoryAllocation->IsMapped() == bMapAllocation)
 					{
-						//FPlatformMisc::LowLevelOutputDebugStringf(TEXT("JYP (%u)(%x)(%s) Found Heap in UsedPage size=%u max=%u"), 
-						//	GFrameNumberRenderThread, this, Type == EType::Image ? TEXT("Image") : TEXT("Buffer"), Size, TargetDefaultPageSize);
-						return ResourceAllocation;
+						FOldResourceAllocation* ResourceAllocation = Page->TryAllocate(Size, Alignment, File, Line);
+						if (ResourceAllocation)
+						{
+							//FPlatformMisc::LowLevelOutputDebugStringf(TEXT("JYP (%u)(%x)(%s) Found Heap in UsedPage size=%u max=%u"), 
+							//	GFrameNumberRenderThread, this, Type == EType::Image ? TEXT("Image") : TEXT("Buffer"), Size, TargetDefaultPageSize);
+							return ResourceAllocation;
+						}
 					}
 				}
 			}
-		}
 
 #if VULKAN_FREEPAGE_FOR_TYPE
-		TArray<FOldResourceHeapPage*>& FreePages = (Type == EType::Image) ? FreeImagePages : FreeBufferPages;
-		for (int32 Index = 0; Index < FreePages.Num(); ++Index)
-		{
-			FOldResourceHeapPage* Page = FreePages[Index];
-			if (Page->DeviceMemoryAllocation->IsMapped() == bMapAllocation)
+			TArray<FOldResourceHeapPage*>& FreePages = (Type == EType::Image) ? FreeImagePages : FreeBufferPages;
+			for (int32 Index = 0; Index < FreePages.Num(); ++Index)
 			{
-				FOldResourceAllocation* ResourceAllocation = Page->TryAllocate(Size, Alignment, File, Line);
-				if (ResourceAllocation)
+				FOldResourceHeapPage* Page = FreePages[Index];
+				if (Page->DeviceMemoryAllocation->IsMapped() == bMapAllocation)
 				{
-					//FPlatformMisc::LowLevelOutputDebugStringf(TEXT("JYP (%u)(%x)(%s) Found Heap in FreePages size=%u max=%u"),
-					//	GFrameNumberRenderThread, this, Type == EType::Image ? TEXT("Image") : TEXT("Buffer"), Size, TargetDefaultPageSize);
-					FreePages.RemoveSingleSwap(Page, false);
-					UsedPages.Add(Page);
-					return ResourceAllocation;
+					FOldResourceAllocation* ResourceAllocation = Page->TryAllocate(Size, Alignment, File, Line);
+					if (ResourceAllocation)
+					{
+						//FPlatformMisc::LowLevelOutputDebugStringf(TEXT("JYP (%u)(%x)(%s) Found Heap in FreePages size=%u max=%u"),
+						//	GFrameNumberRenderThread, this, Type == EType::Image ? TEXT("Image") : TEXT("Buffer"), Size, TargetDefaultPageSize);
+						FreePages.RemoveSingleSwap(Page, false);
+						UsedPages.Add(Page);
+						return ResourceAllocation;
+					}
 				}
 			}
-		}
 #else
-		for (int32 Index = 0; Index < FreePages.Num(); ++Index)
-		{
-			FOldResourceHeapPage* Page = FreePages[Index];
-			if (Page->DeviceMemoryAllocation->IsMapped() == bMapAllocation)
+			for (int32 Index = 0; Index < FreePages.Num(); ++Index)
 			{
-				FOldResourceAllocation* ResourceAllocation = Page->TryAllocate(Size, Alignment, File, Line);
-				if (ResourceAllocation)
+				FOldResourceHeapPage* Page = FreePages[Index];
+				if (Page->DeviceMemoryAllocation->IsMapped() == bMapAllocation)
 				{
-					//FPlatformMisc::LowLevelOutputDebugStringf(TEXT("JYP (%u)(%x)(%s) Found Heap in FreePages size=%u max=%u"),
-					//	GFrameNumberRenderThread, this, Type == EType::Image ? TEXT("Image") : TEXT("Buffer"), Size, TargetDefaultPageSize);
-					FreePages.RemoveSingleSwap(Page, false);
-					UsedPages.Add(Page);
-					return ResourceAllocation;
+					FOldResourceAllocation* ResourceAllocation = Page->TryAllocate(Size, Alignment, File, Line);
+					if (ResourceAllocation)
+					{
+						//FPlatformMisc::LowLevelOutputDebugStringf(TEXT("JYP (%u)(%x)(%s) Found Heap in FreePages size=%u max=%u"),
+						//	GFrameNumberRenderThread, this, Type == EType::Image ? TEXT("Image") : TEXT("Buffer"), Size, TargetDefaultPageSize);
+						FreePages.RemoveSingleSwap(Page, false);
+						UsedPages.Add(Page);
+						return ResourceAllocation;
+					}
 				}
 			}
-		}
 #endif
-		uint32 AllocationSize = FMath::Max(Size, TargetDefaultPageSize);
+			if (0 != VULKAN_MAX_SUB_ALLOCATION && Size >= VULKAN_MAX_SUB_ALLOCATION)
+			{
+				AllocationSize = Size;
+			}
+			else
+			{
+				AllocationSize = FMath::Max(Size, TargetDefaultPageSize);
+			}
+		}
 #endif
-		FDeviceMemoryAllocation* DeviceMemoryAllocation = Owner->GetParent()->GetMemoryManager().Alloc(true, AllocationSize, MemoryTypeIndex, nullptr, File, Line);
-		if (!DeviceMemoryAllocation)
+		FDeviceMemoryAllocation* DeviceMemoryAllocation = Owner->GetParent()->GetMemoryManager().Alloc(true, AllocationSize, MemoryTypeIndex, nullptr, VULKAN_MEMORY_HIGHEST_PRIORITY, File, Line);
+		if (!DeviceMemoryAllocation && Size != AllocationSize)
 		{
 			// Retry with a smaller size
-			DeviceMemoryAllocation = Owner->GetParent()->GetMemoryManager().Alloc(false, Size, MemoryTypeIndex, nullptr, File, Line);
+			DeviceMemoryAllocation = Owner->GetParent()->GetMemoryManager().Alloc(false, Size, MemoryTypeIndex, nullptr, VULKAN_MEMORY_HIGHEST_PRIORITY, File, Line);
 			if (!DeviceMemoryAllocation)
 			{
 				UE_LOG(LogVulkanRHI, Fatal, TEXT("Out of memory on Vulkan; MemoryTypeIndex=%d, AllocSize=%0.3fMB"), MemoryTypeIndex, (float)AllocationSize / 1048576.0f);
@@ -965,7 +1013,7 @@ namespace VulkanRHI
 		VkMemoryDedicatedAllocateInfoKHR DedicatedAllocInfo;
 		ZeroVulkanStruct(DedicatedAllocInfo, VK_STRUCTURE_TYPE_MEMORY_DEDICATED_ALLOCATE_INFO_KHR);
 		DedicatedAllocInfo.image = Image;
-		FDeviceMemoryAllocation* DeviceMemoryAllocation = Owner->GetParent()->GetMemoryManager().Alloc(false, AllocationSize, MemoryTypeIndex, &DedicatedAllocInfo, File, Line);
+		FDeviceMemoryAllocation* DeviceMemoryAllocation = Owner->GetParent()->GetMemoryManager().Alloc(false, AllocationSize, MemoryTypeIndex, &DedicatedAllocInfo, VULKAN_MEMORY_HIGHEST_PRIORITY, File, Line);
 
 		++PageIDCounter;
 		FOldResourceHeapPage* NewPage = new FOldResourceHeapPage(this, DeviceMemoryAllocation, PageIDCounter);
@@ -1037,6 +1085,10 @@ namespace VulkanRHI
 			{
 				int32 HeapIndex = MemoryProperties.memoryTypes[TypeIndices[Index]].heapIndex;
 				VkDeviceSize HeapSize = MemoryProperties.memoryHeaps[HeapIndex].size;
+				if (VULKAN_FAKE_MEMORY_LIMIT)
+				{
+					HeapSize = FMath::Min<VkDeviceSize>(VULKAN_FAKE_MEMORY_LIMIT << 30llu, HeapSize);
+				}
 				VkDeviceSize PageSize = FMath::Min<VkDeviceSize>(HeapSize / 8, GPU_ONLY_HEAP_PAGE_SIZE);
 #if PLATFORM_ANDROID && !PLATFORM_LUMIN && !PLATFORM_LUMINGL4
 				PageSize = FMath::Min<VkDeviceSize>(PageSize, ANDROID_MAX_HEAP_PAGE_SIZE);
@@ -1185,9 +1237,25 @@ namespace VulkanRHI
 			Alignment = (uint32)Limits.minUniformBufferOffsetAlignment;
 
 			// reduce the number of unique buffer types by allowing all buffer types in non-texel/storage buffers
+			// @FYI ROLANDO: I want to comment out this line as this effectively prevents us from lowering the priority on vertex/index buffers
 			BufferUsageFlags |= VK_BUFFER_USAGE_UNIFORM_BUFFER_BIT | VK_BUFFER_USAGE_INDEX_BUFFER_BIT | VK_BUFFER_USAGE_VERTEX_BUFFER_BIT | VK_BUFFER_USAGE_INDIRECT_BUFFER_BIT;
 		}
 
+		float Priority = VULKAN_MEMORY_MEDIUM_PRIORITY;
+
+		if(BufferUsageFlags & (VK_BUFFER_USAGE_STORAGE_TEXEL_BUFFER_BIT| VK_BUFFER_USAGE_INDIRECT_BUFFER_BIT))
+		{
+			Priority = VULKAN_MEMORY_HIGHEST_PRIORITY;
+		}
+		if(BufferUsageFlags & (VK_BUFFER_USAGE_INDEX_BUFFER_BIT| VK_BUFFER_USAGE_VERTEX_BUFFER_BIT))
+		{
+			Priority = VULKAN_MEMORY_LOW_PRIORITY;
+		}
+		else if(BufferUsageFlags & (VK_BUFFER_USAGE_UNIFORM_BUFFER_BIT))
+		{
+			Priority = VULKAN_MEMORY_HIGHER_PRIORITY;
+		}
+
 //		const bool bIsUniformBuffer = ((BufferUsageFlags & VK_BUFFER_USAGE_UNIFORM_BUFFER_BIT) != 0);
 
 		FScopeLock ScopeLock(&GResourceHeapLock);
@@ -1246,7 +1314,7 @@ namespace VulkanRHI
 		uint32 MemoryTypeIndex;
 		VERIFYVULKANRESULT(Device->GetMemoryManager().GetMemoryTypeFromProperties(MemReqs.memoryTypeBits, MemoryPropertyFlags, &MemoryTypeIndex));
 
-		FDeviceMemoryAllocation* DeviceMemoryAllocation = Device->GetMemoryManager().Alloc(false, MemReqs.size, MemoryTypeIndex, nullptr, File, Line);
+		FDeviceMemoryAllocation* DeviceMemoryAllocation = Device->GetMemoryManager().Alloc(false, MemReqs.size, MemoryTypeIndex, nullptr, Priority, File, Line);
 		VERIFYVULKANRESULT(VulkanRHI::vkBindBufferMemory(Device->GetInstanceHandle(), Buffer, DeviceMemoryAllocation->GetHandle(), 0));
 		if (DeviceMemoryAllocation->CanBeMapped())
 		{
diff --git a/Engine/Source/Runtime/VulkanRHI/Private/VulkanRHI.cpp b/Engine/Source/Runtime/VulkanRHI/Private/VulkanRHI.cpp
--- a/Engine/Source/Runtime/VulkanRHI/Private/VulkanRHI.cpp
+++ b/Engine/Source/Runtime/VulkanRHI/Private/VulkanRHI.cpp
@@ -1056,7 +1056,7 @@ FVulkanBuffer::FVulkanBuffer(FVulkanDevice& InDevice, uint32 InSize, VkFlags InU
 	VkMemoryRequirements MemoryRequirements;
 	VulkanRHI::vkGetBufferMemoryRequirements(Device.GetInstanceHandle(), Buf, &MemoryRequirements);
 
-	Allocation = InDevice.GetMemoryManager().Alloc(false, MemoryRequirements.size, MemoryRequirements.memoryTypeBits, InMemPropertyFlags, nullptr, File ? File : __FILE__, Line ? Line : __LINE__);
+	Allocation = InDevice.GetMemoryManager().Alloc(false, MemoryRequirements.size, MemoryRequirements.memoryTypeBits, InMemPropertyFlags, nullptr, VULKAN_MEMORY_MEDIUM_PRIORITY, File ? File : __FILE__, Line ? Line : __LINE__);
 	check(Allocation);
 	VERIFYVULKANRESULT_EXPANDED(VulkanRHI::vkBindBufferMemory(Device.GetInstanceHandle(), Buf, Allocation->GetHandle(), 0));
 }
diff --git a/Engine/Source/Runtime/VulkanRHI/Public/VulkanConfiguration.h b/Engine/Source/Runtime/VulkanRHI/Public/VulkanConfiguration.h
--- a/Engine/Source/Runtime/VulkanRHI/Public/VulkanConfiguration.h
+++ b/Engine/Source/Runtime/VulkanRHI/Public/VulkanConfiguration.h
@@ -175,6 +175,13 @@
 	#define VULKAN_SUPPORTS_DEBUG_UTILS							0
 #endif
 
+#ifdef VK_EXT_memory_priority
+#define VULKAN_SUPPORTS_MEMORY_PRIORITY							1
+#else
+#define VULKAN_SUPPORTS_MEMORY_PRIORITY							0
+#endif
+
+
 
 DECLARE_LOG_CATEGORY_EXTERN(LogVulkanRHI, Log, All);
 
diff --git a/Engine/Source/Runtime/VulkanRHI/Public/VulkanMemory.h b/Engine/Source/Runtime/VulkanRHI/Public/VulkanMemory.h
--- a/Engine/Source/Runtime/VulkanRHI/Public/VulkanMemory.h
+++ b/Engine/Source/Runtime/VulkanRHI/Public/VulkanMemory.h
@@ -13,6 +13,12 @@
 #define VULKAN_MEMORY_TRACK_CALLSTACK	0
 
 
+#define VULKAN_MEMORY_LOW_PRIORITY 0.f
+#define VULKAN_MEMORY_MEDIUM_PRIORITY 0.5f
+#define VULKAN_MEMORY_HIGHER_PRIORITY 0.75f
+#define VULKAN_MEMORY_HIGHEST_PRIORITY 1.f
+
+
 class FVulkanQueue;
 class FVulkanCmdBuffer;
 
@@ -267,13 +273,13 @@ namespace VulkanRHI
 
 
 		// bCanFail means an allocation failing is not a fatal error, just returns nullptr
-		FDeviceMemoryAllocation* Alloc(bool bCanFail, VkDeviceSize AllocationSize, uint32 MemoryTypeIndex, void* DedicatedAllocateInfo, const char* File, uint32 Line);
+		FDeviceMemoryAllocation* Alloc(bool bCanFail, VkDeviceSize AllocationSize, uint32 MemoryTypeIndex, void* DedicatedAllocateInfo, float Priority, const char* File, uint32 Line);
 
-		inline FDeviceMemoryAllocation* Alloc(bool bCanFail, VkDeviceSize AllocationSize, uint32 MemoryTypeBits, VkMemoryPropertyFlags MemoryPropertyFlags, void* DedicatedAllocateInfo, const char* File, uint32 Line)
+		inline FDeviceMemoryAllocation* Alloc(bool bCanFail, VkDeviceSize AllocationSize, uint32 MemoryTypeBits, VkMemoryPropertyFlags MemoryPropertyFlags, void* DedicatedAllocateInfo, float Priority, const char* File, uint32 Line)
 		{
 			uint32 MemoryTypeIndex = ~0;
 			VERIFYVULKANRESULT(this->GetMemoryTypeFromProperties(MemoryTypeBits, MemoryPropertyFlags, &MemoryTypeIndex));
-			return Alloc(bCanFail, AllocationSize, MemoryTypeIndex, DedicatedAllocateInfo, File, Line);
+			return Alloc(bCanFail, AllocationSize, MemoryTypeIndex, DedicatedAllocateInfo, Priority, File, Line);
 		}
 
 		// Sets the Allocation to nullptr
