DR - Dump memory on OOM
#rnx

diff --git a/Engine/Source/Runtime/VulkanRHI/Private/VulkanUtil.cpp b/Engine/Source/Runtime/VulkanRHI/Private/VulkanUtil.cpp
--- a/Engine/Source/Runtime/VulkanRHI/Private/VulkanUtil.cpp
+++ b/Engine/Source/Runtime/VulkanRHI/Private/VulkanUtil.cpp
@@ -629,6 +629,7 @@ namespace VulkanRHI
 	 */
 	void VerifyVulkanResult(VkResult Result, const ANSICHAR* VkFunction, const ANSICHAR* Filename, uint32 Line)
 	{
+		bool bDumpMemory = false;
 		FString ErrorString;
 		switch (Result)
 		{
@@ -638,8 +639,8 @@ namespace VulkanRHI
 		VKERRORCASE(VK_EVENT_SET); break;
 		VKERRORCASE(VK_EVENT_RESET); break;
 		VKERRORCASE(VK_INCOMPLETE); break;
-		VKERRORCASE(VK_ERROR_OUT_OF_HOST_MEMORY); break;
-		VKERRORCASE(VK_ERROR_OUT_OF_DEVICE_MEMORY); break;
+		VKERRORCASE(VK_ERROR_OUT_OF_HOST_MEMORY); bDumpMemory = true; break;
+		VKERRORCASE(VK_ERROR_OUT_OF_DEVICE_MEMORY); bDumpMemory = true; break;
 		VKERRORCASE(VK_ERROR_INITIALIZATION_FAILED); break;
 		VKERRORCASE(VK_ERROR_DEVICE_LOST); GIsGPUCrashed = true; break;
 		VKERRORCASE(VK_ERROR_MEMORY_MAP_FAILED); break;
@@ -687,6 +688,13 @@ namespace VulkanRHI
 		}
 #endif
 
+#if UE_BUILD_DEBUG || UE_BUILD_DEVELOPMENT
+		if (bDumpMemory)
+		{
+			GVulkanRHI->DumpMemory();
+		}
+#endif
+
 		UE_LOG(LogVulkanRHI, Fatal, TEXT("%s failed, VkResult=%d\n at %s:%u \n with error %s"),
 			ANSI_TO_TCHAR(VkFunction), (int32)Result, ANSI_TO_TCHAR(Filename), Line, *ErrorString);
 	}
