attempted fix for static analysis

diff --git a/Engine/Source/Runtime/VulkanRHI/Private/VulkanMemory.cpp b/Engine/Source/Runtime/VulkanRHI/Private/VulkanMemory.cpp
--- a/Engine/Source/Runtime/VulkanRHI/Private/VulkanMemory.cpp
+++ b/Engine/Source/Runtime/VulkanRHI/Private/VulkanMemory.cpp
@@ -953,7 +953,8 @@ namespace VulkanRHI
 				}
 			}
 #endif
-			if (0 != VULKAN_MAX_SUB_ALLOCATION && Size >= VULKAN_MAX_SUB_ALLOCATION)
+			constexpr bool bUseMaxSubAllocation = 0 != VULKAN_MAX_SUB_ALLOCATION;
+			if (bUseMaxSubAllocation && Size >= VULKAN_MAX_SUB_ALLOCATION)
 			{
 				AllocationSize = Size;
 			}
