DR - Static analysis fix
#rnx

diff --git a/Engine/Source/Runtime/VulkanRHI/Private/VulkanMemory.cpp b/Engine/Source/Runtime/VulkanRHI/Private/VulkanMemory.cpp
--- a/Engine/Source/Runtime/VulkanRHI/Private/VulkanMemory.cpp
+++ b/Engine/Source/Runtime/VulkanRHI/Private/VulkanMemory.cpp
@@ -953,7 +953,12 @@ namespace VulkanRHI
 				}
 			}
 #endif
-			constexpr bool bUseMaxSubAllocation = 0 != VULKAN_MAX_SUB_ALLOCATION;
+
+#if VULKAN_MAX_SUB_ALLOCATION > 0
+			constexpr bool bUseMaxSubAllocation = true;
+#else
+			constexpr bool bUseMaxSubAllocation = false;
+#endif
 			if (bUseMaxSubAllocation && Size >= VULKAN_MAX_SUB_ALLOCATION)
 			{
 				AllocationSize = Size;
