DR - Support for VK_LAYER_KHRONOS_validation with r.Vulkan.StandardValidation=2
#rnx

diff --git a/Engine/Source/Runtime/VulkanRHI/Private/VulkanLayers.cpp b/Engine/Source/Runtime/VulkanRHI/Private/VulkanLayers.cpp
--- a/Engine/Source/Runtime/VulkanRHI/Private/VulkanLayers.cpp
+++ b/Engine/Source/Runtime/VulkanRHI/Private/VulkanLayers.cpp
@@ -31,16 +31,18 @@ TAutoConsoleVariable<int32> GValidationCvar(
 static TAutoConsoleVariable<int32> GStandardValidationCvar(
 	TEXT("r.Vulkan.StandardValidation"),
 	1,
-	TEXT("1 to use VK_LAYER_LUNARG_standard_validation (default) if available\n")
-	TEXT("0 to use individual layers"),
+	TEXT("2 to use VK_LAYER_KHRONOS_validation if available\n")
+	TEXT("1 to use VK_LAYER_LUNARG_standard_validation (default) if available, or \n")
+	TEXT("0 to use individual validation layers"),
 	ECVF_ReadOnly | ECVF_RenderThreadSafe
 );
 
 #if VULKAN_ENABLE_DRAW_MARKERS
-	#define RENDERDOC_LAYER_NAME		"VK_LAYER_RENDERDOC_Capture"
+	#define RENDERDOC_LAYER_NAME				"VK_LAYER_RENDERDOC_Capture"
 #endif
 
-#define STANDARD_VALIDATION_LAYER_NAME	"VK_LAYER_LUNARG_standard_validation"
+#define KHRONOS_STANDARD_VALIDATION_LAYER_NAME	"VK_LAYER_KHRONOS_validation"
+#define STANDARD_VALIDATION_LAYER_NAME			"VK_LAYER_LUNARG_standard_validation"
 
 static const ANSICHAR* GIndividualValidationLayers[] =
 {
@@ -329,14 +331,29 @@ void FVulkanDynamicRHI::GetInstanceLayersAndExtensions(TArray<const ANSICHAR*>&
 		bool bStandardAvailable = false;
 		if (GStandardValidationCvar.GetValueOnAnyThread() != 0)
 		{
-			bStandardAvailable = FindLayerInList(GlobalLayerExtensions, STANDARD_VALIDATION_LAYER_NAME);
-			if (bStandardAvailable)
+			if (GStandardValidationCvar.GetValueOnAnyThread() == 2)
 			{
-				OutInstanceLayers.Add(STANDARD_VALIDATION_LAYER_NAME);
+				bStandardAvailable = FindLayerInList(GlobalLayerExtensions, KHRONOS_STANDARD_VALIDATION_LAYER_NAME);
+				if (bStandardAvailable)
+				{
+					OutInstanceLayers.Add(KHRONOS_STANDARD_VALIDATION_LAYER_NAME);
+				}
+				else
+				{
+					UE_LOG(LogVulkanRHI, Warning, TEXT("Unable to find Vulkan instance validation layer %s; trying individual layers..."), TEXT(STANDARD_VALIDATION_LAYER_NAME));
+				}
 			}
 			else
 			{
-				UE_LOG(LogVulkanRHI, Warning, TEXT("Unable to find Vulkan instance validation layer %s; trying individual layers..."), TEXT(STANDARD_VALIDATION_LAYER_NAME));
+				bStandardAvailable = FindLayerInList(GlobalLayerExtensions, STANDARD_VALIDATION_LAYER_NAME);
+				if (bStandardAvailable)
+				{
+					OutInstanceLayers.Add(STANDARD_VALIDATION_LAYER_NAME);
+				}
+				else
+				{
+					UE_LOG(LogVulkanRHI, Warning, TEXT("Unable to find Vulkan instance validation layer %s; trying individual layers..."), TEXT(STANDARD_VALIDATION_LAYER_NAME));
+				}
 			}
 		}
 
