DR - Fix some Vulkan dedicated allocation bugs
#rnx

diff --git a/Engine/Source/Runtime/VulkanRHI/Private/VulkanMemory.cpp b/Engine/Source/Runtime/VulkanRHI/Private/VulkanMemory.cpp
--- a/Engine/Source/Runtime/VulkanRHI/Private/VulkanMemory.cpp
+++ b/Engine/Source/Runtime/VulkanRHI/Private/VulkanMemory.cpp
@@ -225,10 +225,13 @@ namespace VulkanRHI
 #endif
 
 #if VULKAN_SUPPORTS_DEDICATED_ALLOCATION
-	((VkMemoryDedicatedAllocateInfoKHR*)DedicatedAllocateInfo)->pNext = Info.pNext;
-	Info.pNext = DedicatedAllocateInfo;
+		if (DedicatedAllocateInfo)
+		{
+			((VkMemoryDedicatedAllocateInfoKHR*)DedicatedAllocateInfo)->pNext = Info.pNext;
+			Info.pNext = DedicatedAllocateInfo;
+		}
 #else
-	check(!DedicatedAllocateInfo);
+		check(!DedicatedAllocateInfo);
 #endif
 
 		VkDeviceMemory Handle;
@@ -748,8 +751,8 @@ namespace VulkanRHI
 		else
 		{
 #if VULKAN_SUPPORTS_DEDICATED_ALLOCATION
-			int32 Removed = UsedDedicatedImagePages.RemoveSingleSwap(InPage, false);
-			check(Removed > 0);
+			int32 RemovedDedicated = UsedDedicatedImagePages.RemoveSingleSwap(InPage, false);
+			check(RemovedDedicated > 0);
 #else
 			//checkf(0, TEXT("Page not found in Pool!"));
 #endif
diff --git a/Engine/Source/Runtime/VulkanRHI/Private/VulkanTexture.cpp b/Engine/Source/Runtime/VulkanRHI/Private/VulkanTexture.cpp
--- a/Engine/Source/Runtime/VulkanRHI/Private/VulkanTexture.cpp
+++ b/Engine/Source/Runtime/VulkanRHI/Private/VulkanTexture.cpp
@@ -533,7 +533,7 @@ FVulkanSurface::FVulkanSurface(FVulkanDevice& InDevice, VkImageViewType Resource
 
 #if VULKAN_SUPPORTS_DEDICATED_ALLOCATION
 	// Per https://developer.nvidia.com/what%E2%80%99s-your-vulkan-memory-type
-	VkDeviceSize SizeToBeConsideredForDedicated = 16 * 1024 * 1024;
+	VkDeviceSize SizeToBeConsideredForDedicated = 12 * 1024 * 1024;
 	if ((bRenderTarget || MemoryRequirements.size >= SizeToBeConsideredForDedicated) && InDevice.GetOptionalExtensions().HasKHRDedicatedAllocation)
 	{
 		ResourceAllocation = InDevice.GetResourceHeapManager().AllocateDedicatedImageMemory(Image, MemoryRequirements, VK_MEMORY_PROPERTY_DEVICE_LOCAL_BIT, __FILE__, __LINE__);
