Implemented RHICopySubTextureRegion interface function for Vulkan backend.

diff --git a/Engine/Source/Runtime/VulkanRHI/Private/VulkanRHI.cpp b/Engine/Source/Runtime/VulkanRHI/Private/VulkanRHI.cpp
--- a/Engine/Source/Runtime/VulkanRHI/Private/VulkanRHI.cpp
+++ b/Engine/Source/Runtime/VulkanRHI/Private/VulkanRHI.cpp
@@ -1030,6 +1030,21 @@ void FVulkanDynamicRHI::RHIAliasTextureResources(FTextureRHIParamRef DestTexture
 	}
 }
 
+void FVulkanDynamicRHI::RHICopySubTextureRegion(FTexture2DRHIParamRef SourceTexture, FTexture2DRHIParamRef DestinationTexture, FBox2D SourceBox, FBox2D DestinationBox)
+{
+	FRHICopyTextureInfo CopyInfo;
+
+	CopyInfo.Size.X = (int32)(SourceBox.Max.X - SourceBox.Min.X);
+	CopyInfo.Size.Y = (int32)(SourceBox.Max.Y - SourceBox.Min.Y);
+
+	CopyInfo.SourcePosition.X = (int32)(SourceBox.Min.X);
+	CopyInfo.SourcePosition.Y = (int32)(SourceBox.Min.Y);
+	CopyInfo.DestPosition.X = (int32)(DestinationBox.Min.X);
+	CopyInfo.DestPosition.Y = (int32)(DestinationBox.Min.Y);
+
+	RHIGetDefaultContext()->RHICopyTexture(SourceTexture, DestinationTexture, CopyInfo);
+}
+
 
 FVulkanBuffer::FVulkanBuffer(FVulkanDevice& InDevice, uint32 InSize, VkFlags InUsage, VkMemoryPropertyFlags InMemPropertyFlags, bool bInAllowMultiLock, const char* File, int32 Line) :
 	Device(InDevice),
diff --git a/Engine/Source/Runtime/VulkanRHI/Private/VulkanTexture.cpp b/Engine/Source/Runtime/VulkanRHI/Private/VulkanTexture.cpp
--- a/Engine/Source/Runtime/VulkanRHI/Private/VulkanTexture.cpp
+++ b/Engine/Source/Runtime/VulkanRHI/Private/VulkanTexture.cpp
@@ -2242,18 +2242,32 @@ void FVulkanCommandListContext::RHICopyTexture(FTextureRHIParamRef SourceTexture
 
 	VkImageCopy Region;
 	FMemory::Memzero(Region);
-	ensure(SrcSurface.Width == DstSurface.Width && SrcSurface.Height == DstSurface.Height);
-	Region.extent.width = FMath::Max(1u, SrcSurface.Width>> CopyInfo.SourceMipIndex);
-	Region.extent.height = FMath::Max(1u, SrcSurface.Height >> CopyInfo.SourceMipIndex);
+	if (CopyInfo.Size.X == 0 || CopyInfo.Size.Y == 0)
+	{
+		// Copy whole texture when zero vector is specified for region size
+		ensure(SrcSurface.Width == DstSurface.Width && SrcSurface.Height == DstSurface.Height);
+		Region.extent.width = FMath::Max(1u, SrcSurface.Width >> CopyInfo.SourceMipIndex);
+		Region.extent.height = FMath::Max(1u, SrcSurface.Height >> CopyInfo.SourceMipIndex);
+	}
+	else
+	{
+		ensure(CopyInfo.Size.X > 0 && (uint32)CopyInfo.Size.X <= DstSurface.Width && CopyInfo.Size.Y > 0 && (uint32)CopyInfo.Size.Y <= DstSurface.Height);
+		Region.extent.width = CopyInfo.Size.X;
+		Region.extent.height = CopyInfo.Size.Y;
+	}
 	Region.extent.depth = 1;
 	Region.srcSubresource.aspectMask = SrcSurface.GetFullAspectMask();
 	Region.srcSubresource.baseArrayLayer = CopyInfo.SourceSliceIndex;
 	Region.srcSubresource.layerCount = CopyInfo.NumSlices;
 	Region.srcSubresource.mipLevel = CopyInfo.SourceMipIndex;
+	Region.srcOffset.x = CopyInfo.SourcePosition.X;
+	Region.srcOffset.y = CopyInfo.SourcePosition.Y;
 	Region.dstSubresource.aspectMask = DstSurface.GetFullAspectMask();
 	Region.dstSubresource.baseArrayLayer = CopyInfo.DestSliceIndex;
 	Region.dstSubresource.layerCount = CopyInfo.NumSlices;
 	Region.dstSubresource.mipLevel = CopyInfo.DestMipIndex;
+	Region.dstOffset.x = CopyInfo.DestPosition.X;
+	Region.dstOffset.y = CopyInfo.DestPosition.Y;
 
 	for (uint32 Index = 0; Index < CopyInfo.NumMips; ++Index)
 	{
diff --git a/Engine/Source/Runtime/VulkanRHI/Public/VulkanDynamicRHI.h b/Engine/Source/Runtime/VulkanRHI/Public/VulkanDynamicRHI.h
--- a/Engine/Source/Runtime/VulkanRHI/Public/VulkanDynamicRHI.h
+++ b/Engine/Source/Runtime/VulkanRHI/Public/VulkanDynamicRHI.h
@@ -358,6 +358,8 @@ public:
 	virtual FTextureCubeRHIRef RHICreateTextureCubeFromResource(EPixelFormat Format, uint32 Size, bool bArray, uint32 ArraySize, uint32 NumMips, VkImage Resource, uint32 Flags);
 	virtual void RHIAliasTextureResources(FTextureRHIParamRef DestTexture, FTextureRHIParamRef SrcTexture);
 
+	virtual void RHICopySubTextureRegion(FTexture2DRHIParamRef SourceTexture, FTexture2DRHIParamRef DestinationTexture, FBox2D SourceBox, FBox2D DestinationBox) final override;
+
 	void RHICalibrateTimers() override;
 
 	// Historical number of times we've presented any and all viewports
