fix "real" uniform buffers.
Original change by rolando.caloca.

diff --git a/Engine/Source/Runtime/VulkanRHI/Private/VulkanMemory.cpp b/Engine/Source/Runtime/VulkanRHI/Private/VulkanMemory.cpp
--- a/Engine/Source/Runtime/VulkanRHI/Private/VulkanMemory.cpp
+++ b/Engine/Source/Runtime/VulkanRHI/Private/VulkanMemory.cpp
@@ -1145,6 +1145,10 @@ namespace VulkanRHI
 
 	void FResourceHeapManager::Deinit()
 	{
+		{
+			ProcessPendingUBFreesNoLock(true);
+			check(UBAllocations.PendingFree.Num() == 0);
+		}
 		DestroyResourceAllocations();
 
 		for (int32 Index = 0; Index < ResourceTypeHeaps.Num(); ++Index)
@@ -1281,7 +1285,7 @@ namespace VulkanRHI
 			if ((BufferAllocation->BufferUsageFlags & BufferUsageFlags) == BufferUsageFlags &&
 				(BufferAllocation->MemoryPropertyFlags & MemoryPropertyFlags) == MemoryPropertyFlags)
 			{
-				FBufferSuballocation* Suballocation = (FBufferSuballocation*)BufferAllocation->TryAllocateNoLocking(Size, Alignment, File, Line);
+				FBufferSuballocation* Suballocation = (FBufferSuballocation*)BufferAllocation->TryAllocate(Size, Alignment, File, Line);
 				if (Suballocation)
 				{
 					return Suballocation;
@@ -1295,7 +1299,7 @@ namespace VulkanRHI
 			if ((BufferAllocation->BufferUsageFlags & BufferUsageFlags) == BufferUsageFlags &&
 				(BufferAllocation->MemoryPropertyFlags & MemoryPropertyFlags) == MemoryPropertyFlags)
 			{
-				FBufferSuballocation* Suballocation = (FBufferSuballocation*)BufferAllocation->TryAllocateNoLocking(Size, Alignment, File, Line);
+				FBufferSuballocation* Suballocation = (FBufferSuballocation*)BufferAllocation->TryAllocate(Size, Alignment, File, Line);
 				if (Suballocation)
 				{
 					FreeBufferAllocations[PoolSize].RemoveAtSwap(Index, 1, false);
@@ -1340,7 +1344,7 @@ namespace VulkanRHI
 			MemoryPropertyFlags, MemReqs.alignment, Buffer, BufferId, BufferUsageFlags, PoolSize);
 		UsedBufferAllocations[PoolSize].Add(BufferAllocation);
 
-		return (FBufferSuballocation*)BufferAllocation->TryAllocateNoLocking(Size, Alignment, File, Line);
+		return (FBufferSuballocation*)BufferAllocation->TryAllocate(Size, Alignment, File, Line);
 	}
 
 	void FResourceHeapManager::ReleaseBuffer(FBufferAllocation* BufferAllocation)
@@ -1481,6 +1485,11 @@ namespace VulkanRHI
 	}
 
 
+	void FBufferSuballocation::Flush()
+	{
+		Owner->Flush(AlignedOffset, AllocationSize);
+	}
+
 	FCriticalSection FSubresourceAllocator::CS;
 
 	bool FSubresourceAllocator::JoinFreeBlocks()
@@ -1501,8 +1510,9 @@ namespace VulkanRHI
 		return false;
 	}
 	
-	FResourceSuballocation* FSubresourceAllocator::TryAllocateNoLocking(uint32 InSize, uint32 InAlignment, const char* File, uint32 Line)
+	FResourceSuballocation* FSubresourceAllocator::TryAllocate(uint32 InSize, uint32 InAlignment, const char* File, uint32 Line)
 	{
+		FScopeLock ScopeLock(&CS);
 		InAlignment = FMath::Max(InAlignment, Alignment);
 		for (int32 Index = 0; Index < FreeList.Num(); ++Index)
 		{
@@ -1547,6 +1557,12 @@ namespace VulkanRHI
 		return nullptr;
 	}
 
+	void FSubresourceAllocator::Flush(VkDeviceSize Offset, VkDeviceSize AllocationSize)
+	{
+		MemoryAllocation->FlushMappedMemory(Offset, AllocationSize);
+	}
+
+
 	void FBufferAllocation::Release(FBufferSuballocation* Suballocation)
 	{
 		{
@@ -1558,12 +1574,15 @@ namespace VulkanRHI
 			FRange NewFree;
 			NewFree.Offset = Suballocation->AllocationOffset;
 			NewFree.Size = Suballocation->AllocationSize;
+			check(NewFree.Offset <= GetMaxSize());
+			check(NewFree.Offset + NewFree.Size <= GetMaxSize());
 
 			FreeList.Add(NewFree);
+
+			UsedSize -= Suballocation->AllocationSize;
+			check(UsedSize >= 0);
 		}
 
-		UsedSize -= Suballocation->AllocationSize;
-		check(UsedSize >= 0);
 
 		if (JoinFreeBlocks())
 		{
diff --git a/Engine/Source/Runtime/VulkanRHI/Private/VulkanResources.h b/Engine/Source/Runtime/VulkanRHI/Private/VulkanResources.h
--- a/Engine/Source/Runtime/VulkanRHI/Private/VulkanResources.h
+++ b/Engine/Source/Runtime/VulkanRHI/Private/VulkanResources.h
@@ -1231,15 +1231,34 @@ public:
 	void UpdateConstantData(const void* Contents, int32 ContentsSize);
 };
 
-class FVulkanRealUniformBuffer : public FVulkanUniformBuffer, public FVulkanResourceMultiBuffer
+class FVulkanRealUniformBuffer : public FVulkanUniformBuffer
 {
 public:
+	FVulkanDevice* Device;
 	FVulkanRealUniformBuffer(FVulkanDevice& Device, const FRHIUniformBufferLayout& InLayout, const void* Contents, EUniformBufferUsage InUsage, EUniformBufferValidation Validation);
+	virtual ~FVulkanRealUniformBuffer();
 
-	void Update(const void* Contents, int32 ContentsSize);
+	VulkanRHI::FBufferAllocation* GetBufferAllocation() const
+	{
+		checkSlow(UBAllocation);
+		return UBAllocation->GetBufferAllocation();
+	}
 
-private:
-	TArray<TRefCountPtr<FRHIResource>> ResourceTable;
+	inline uint32 GetOffset() const
+	{
+		return UBAllocation->GetOffset();
+	}
+
+	// Returns previous allocation
+	inline VulkanRHI::FBufferSuballocation* UpdateUBAllocation(VulkanRHI::FBufferSuballocation* NewAlloc)
+	{
+		checkSlow(UBAllocation);
+		VulkanRHI::FBufferSuballocation* PrevBufferSuballoc = UBAllocation;
+		UBAllocation = NewAlloc;
+
+		return PrevBufferSuballoc;
+	}
+	VulkanRHI::FBufferSuballocation* UBAllocation = nullptr;
 };
 
 class FVulkanStructuredBuffer : public FRHIStructuredBuffer, public FVulkanResourceMultiBuffer
diff --git a/Engine/Source/Runtime/VulkanRHI/Private/VulkanUniformBuffer.cpp b/Engine/Source/Runtime/VulkanRHI/Private/VulkanUniformBuffer.cpp
--- a/Engine/Source/Runtime/VulkanRHI/Private/VulkanUniformBuffer.cpp
+++ b/Engine/Source/Runtime/VulkanRHI/Private/VulkanUniformBuffer.cpp
@@ -146,33 +146,28 @@ void FVulkanEmulatedUniformBuffer::UpdateConstantData(const void* Contents, int3
 
 FVulkanRealUniformBuffer::FVulkanRealUniformBuffer(FVulkanDevice& Device, const FRHIUniformBufferLayout& InLayout, const void* Contents, EUniformBufferUsage InUsage, EUniformBufferValidation Validation)
 	: FVulkanUniformBuffer(InLayout, Contents, InUsage, Validation)
-	, FVulkanResourceMultiBuffer(&Device, VK_BUFFER_USAGE_UNIFORM_BUFFER_BIT, InLayout.ConstantBufferSize, UniformBufferToBufferUsage(InUsage), GEmptyCreateInfo)
+	, Device(&Device)
 {
 #if VULKAN_ENABLE_AGGRESSIVE_STATS
 	SCOPE_CYCLE_COUNTER(STAT_VulkanUniformBufferCreateTime);
 #endif
 
-	if (InLayout.ConstantBufferSize)
+	if (InLayout.ConstantBufferSize > 0)
 	{
-		//#todo-rco: Optimize
-		const bool bRT = IsInRenderingThread();
-		void* Data = Lock(bRT, RLM_WriteOnly, InLayout.ConstantBufferSize, 0);
-		FMemory::Memcpy(Data, Contents, InLayout.ConstantBufferSize);
-		Unlock(bRT);
+		VulkanRHI::FResourceHeapManager& ResourceMgr = Device.GetResourceHeapManager();
+
+		// Set it directly as there is no previous one
+		UBAllocation = ResourceMgr.AllocUniformBuffer(InLayout.ConstantBufferSize, Contents);
 	}
 
 	// Ancestor's constructor will set up the Resource table, so nothing else to do here
 }
 
-void FVulkanRealUniformBuffer::Update(const void* Contents, int32 ContentsSize)
+FVulkanRealUniformBuffer::~FVulkanRealUniformBuffer()
 {
-	if (ContentsSize > 0)
+	if (UBAllocation)
 	{
-		FVulkanCmdBuffer* Cmd = Device->GetImmediateContext().GetCommandBufferManager()->GetActiveCmdBuffer();
-
-		const VkCommandBuffer CmdBuffer = Cmd->GetHandle();
-
-		VulkanRHI::vkCmdUpdateBuffer(CmdBuffer, GetHandle(), GetOffset(), ContentsSize, Contents);
+		Device->GetResourceHeapManager().ReleaseUniformBuffer(UBAllocation);
 	}
 }
 
@@ -206,11 +201,26 @@ inline void FVulkanDynamicRHI::UpdateUniformBuffer(FVulkanUniformBuffer* Uniform
 	FVulkanRealUniformBuffer* RealUniformBuffer = bRealUBs ? (FVulkanRealUniformBuffer*)UniformBuffer : nullptr;
 	FVulkanEmulatedUniformBuffer* EmulatedUniformBuffer = bRealUBs ? nullptr : (FVulkanEmulatedUniformBuffer*)UniformBuffer;
 
-	if (RHICmdList.Bypass())
+	FBufferSuballocation* NewUBAlloc = nullptr;
+	if (bRealUBs)
+	{
+		NewUBAlloc = nullptr;
+		if (ConstantBufferSize > 0)
+		{
+			NewUBAlloc = Device->GetResourceHeapManager().AllocUniformBuffer(ConstantBufferSize, Contents);
+		}
+	}
+
+	bool bRHIBypass = RHICmdList.Bypass();
+	if (bRHIBypass)
 	{
 		if (bRealUBs)
 		{
-			RealUniformBuffer->Update(Contents, ConstantBufferSize);
+			if (ConstantBufferSize > 0)
+			{
+				FBufferSuballocation* PrevAlloc = RealUniformBuffer->UpdateUBAllocation(NewUBAlloc);
+				Device->GetResourceHeapManager().ReleaseUniformBuffer(PrevAlloc);
+			}
 		}
 		else
 		{
@@ -221,8 +231,6 @@ inline void FVulkanDynamicRHI::UpdateUniformBuffer(FVulkanUniformBuffer* Uniform
 	else
 	{
 		FRHIResource** CmdListResources = nullptr;
-		void* CmdListConstantBufferData = nullptr;
-
 		if (NumResources > 0)
 		{
 			CmdListResources = (FRHIResource**)RHICmdList.Alloc(sizeof(FRHIResource*) * NumResources, alignof(FRHIResource*));
@@ -230,7 +238,6 @@ inline void FVulkanDynamicRHI::UpdateUniformBuffer(FVulkanUniformBuffer* Uniform
 			for (int32 ResourceIndex = 0; ResourceIndex < NumResources; ++ResourceIndex)
 			{
 				FRHIResource* Resource = *(FRHIResource**)((uint8*)Contents + Layout.Resources[ResourceIndex].MemberOffset);
-
 				checkf(Resource, TEXT("Invalid resource entry creating uniform buffer, %s.Resources[%u], ResourceType 0x%x."),
 					*Layout.GetDebugName().ToString(),
 					ResourceIndex,
@@ -240,26 +247,25 @@ inline void FVulkanDynamicRHI::UpdateUniformBuffer(FVulkanUniformBuffer* Uniform
 			}
 		}
 
-		if (ConstantBufferSize > 0)
-		{
-			// Can be optimized creating a new Vulkan buffer here instead of extra memcpy, but would require refactoring entire Vulkan uniform buffer code.
-			CmdListConstantBufferData = (void*)RHICmdList.Alloc(ConstantBufferSize, 256);
-			FMemory::Memcpy(CmdListConstantBufferData, Contents, ConstantBufferSize);
-		}
-
-		RHICmdList.EnqueueLambda([RealUniformBuffer, EmulatedUniformBuffer, CmdListResources, NumResources, CmdListConstantBufferData, ConstantBufferSize](FRHICommandList&)
+		if (bRealUBs)
 		{
-			if (bRealUBs)
+			RHICmdList.EnqueueLambda([RealUniformBuffer, NewUBAlloc, CmdListResources, NumResources, ConstantBufferSize](FRHICommandList&)
 			{
-				RealUniformBuffer->Update(CmdListConstantBufferData, ConstantBufferSize);
+				FBufferSuballocation* PrevAlloc = RealUniformBuffer->UpdateUBAllocation(NewUBAlloc);
+				RealUniformBuffer->Device->GetResourceHeapManager().ReleaseUniformBuffer(PrevAlloc);
 				RealUniformBuffer->UpdateResourceTable(CmdListResources, NumResources);
-			}
-			else
+			});
+		}
+		else
+		{
+			void* CmdListConstantBufferData = RHICmdList.Alloc(ConstantBufferSize, 16);
+			FMemory::Memcpy(CmdListConstantBufferData, Contents, ConstantBufferSize);
+			RHICmdList.EnqueueLambda([EmulatedUniformBuffer, CmdListResources, NumResources, CmdListConstantBufferData, ConstantBufferSize](FRHICommandList&)
 			{
 				EmulatedUniformBuffer->UpdateConstantData(CmdListConstantBufferData, ConstantBufferSize);
 				EmulatedUniformBuffer->UpdateResourceTable(CmdListResources, NumResources);
-			}
-		});
+			});
+		}
 		RHICmdList.RHIThreadFence(true);
 	}
 }
@@ -306,3 +312,49 @@ FVulkanUniformBufferUploader::~FVulkanUniformBufferUploader()
 {
 	delete CPUBuffer;
 }
+
+
+namespace VulkanRHI
+{
+	VulkanRHI::FBufferSuballocation* FResourceHeapManager::AllocUniformBuffer(uint32 Size, const void* Contents)
+	{
+		VulkanRHI::FBufferSuballocation* OutAlloc = Device->GetResourceHeapManager().AllocateBuffer(Size, VK_BUFFER_USAGE_UNIFORM_BUFFER_BIT, VK_MEMORY_PROPERTY_HOST_VISIBLE_BIT | VK_MEMORY_PROPERTY_HOST_COHERENT_BIT, __FILE__, __LINE__);
+		FMemory::Memcpy(OutAlloc->GetMappedPointer(), Contents, Size);
+		OutAlloc->Flush();
+
+		return OutAlloc;
+	}
+
+	void FResourceHeapManager::ReleaseUniformBuffer(VulkanRHI::FBufferSuballocation* UBAlloc)
+	{
+		checkSlow(UBAlloc);
+
+		FScopeLock ScopeLock(&UBAllocations.CS);
+		ProcessPendingUBFreesNoLock(false);
+		FUBPendingFree Pending;
+		Pending.Frame = GFrameNumberRenderThread;
+		Pending.Allocation = UBAlloc;
+		UBAllocations.PendingFree.Add(Pending);
+
+		UBAllocations.Peak = FMath::Max(UBAllocations.Peak, (uint32)UBAllocations.PendingFree.Num());
+	}
+
+	void FResourceHeapManager::ProcessPendingUBFreesNoLock(bool bForce)
+	{
+		for (int32 Index = UBAllocations.PendingFree.Num() - 1; Index >= 0; --Index)
+		{
+			FUBPendingFree& Alloc = UBAllocations.PendingFree[Index];
+			if (GFrameNumberRenderThread > Alloc.Frame + VulkanRHI::NUM_FRAMES_TO_WAIT_BEFORE_RELEASING_TO_OS || bForce)
+			{
+				delete Alloc.Allocation;
+				UBAllocations.PendingFree.RemoveAtSwap(Index, 1, false);
+			}
+		}
+	}
+
+	void FResourceHeapManager::ProcessPendingUBFrees(bool bForce)
+	{
+		FScopeLock ScopeLock(&UBAllocations.CS);
+		ProcessPendingUBFreesNoLock(bForce);
+	}
+}
diff --git a/Engine/Source/Runtime/VulkanRHI/Public/VulkanMemory.h b/Engine/Source/Runtime/VulkanRHI/Public/VulkanMemory.h
--- a/Engine/Source/Runtime/VulkanRHI/Public/VulkanMemory.h
+++ b/Engine/Source/Runtime/VulkanRHI/Public/VulkanMemory.h
@@ -33,6 +33,7 @@ extern EDelayAcquireImageType GVulkanDelayAcquireImage;
 
 namespace VulkanRHI
 {
+	uint32 GetMaxSize(class FBufferAllocation* S);
 	class FFenceManager;
 
 	extern int32 GVulkanUseBufferBinning;
@@ -531,6 +532,10 @@ namespace VulkanRHI
 			, Owner(InOwner)
 			, Handle(InHandle)
 		{
+			uint32 S = GetMaxSize(InOwner);
+			check(InAlignedOffset <= S);
+			check(InAllocationOffset + InAllocationSize <= S);
+			check(InAlignedOffset + InRequestedSize <= S);
 		}
 
 		virtual ~FBufferSuballocation();
@@ -548,6 +553,8 @@ namespace VulkanRHI
 		// Returns the pointer to the mapped data for this SubAllocation, not the full buffer!
 		void* GetMappedPointer();
 
+		void Flush();
+
 	protected:
 		friend class FBufferAllocation;
 		FBufferAllocation* Owner;
@@ -581,13 +588,7 @@ namespace VulkanRHI
 		virtual FResourceSuballocation* CreateSubAllocation(uint32 Size, uint32 AlignedOffset, uint32 AllocatedSize, uint32 AllocatedOffset) = 0;
 		virtual void Destroy(FVulkanDevice* Device) = 0;
 
-		FResourceSuballocation* TryAllocateNoLocking(uint32 InSize, uint32 InAlignment, const char* File, uint32 Line);
-
-		inline FResourceSuballocation* TryAllocateLocking(uint32 InSize, uint32 InAlignment, const char* File, uint32 Line)
-		{
-			FScopeLock ScopeLock(&CS);
-			return TryAllocateNoLocking(InSize, InAlignment, File, Line);
-		}
+		FResourceSuballocation* TryAllocate(uint32 InSize, uint32 InAlignment, const char* File, uint32 Line);
 
 		inline uint32 GetAlignment() const
 		{
@@ -599,6 +600,8 @@ namespace VulkanRHI
 			return MemoryAllocation->GetMappedPointer();
 		}
 
+		void Flush(VkDeviceSize Offset, VkDeviceSize AllocationSize);
+		uint32 GetMaxSize(){ return MaxSize; }
 	protected:
 		FResourceHeapManager* Owner;
 		uint32 MemoryTypeIndex;
@@ -744,6 +747,10 @@ namespace VulkanRHI
 
 		friend class FResourceHeapManager;
 	};
+	inline uint32 GetMaxSize(FBufferAllocation* S)
+	{
+		return S->GetMaxSize();
+	}
 
 	// Manages heaps and their interactions
 	class FResourceHeapManager : public FDeviceChild
@@ -839,6 +846,9 @@ namespace VulkanRHI
 #if UE_BUILD_DEBUG || UE_BUILD_DEVELOPMENT
 		void DumpMemory();
 #endif
+		
+		FBufferSuballocation* AllocUniformBuffer(uint32 Size, const void* Contents);
+		void ReleaseUniformBuffer(FBufferSuballocation* UBAlloc);
 
 		void* Hotfix;
 
@@ -917,6 +927,21 @@ namespace VulkanRHI
 
 		void ReleaseFreedResources(bool bImmediately);
 		void DestroyResourceAllocations();
+		struct FUBPendingFree
+		{
+			FBufferSuballocation* Allocation = nullptr;
+			uint64 Frame = 0;
+		};
+
+		struct
+		{
+			FCriticalSection CS;
+			TArray<FUBPendingFree> PendingFree;
+			uint32 Peak = 0;
+		} UBAllocations;
+
+		void ProcessPendingUBFreesNoLock(bool bForce);
+		void ProcessPendingUBFrees(bool bForce);
 	};
 
 	class FStagingBuffer : public FRefCount
