change vulkan on desktop/quail back to us "real" uniform buffers.

diff --git a/Engine/Config/BaseDeviceProfiles.ini b/Engine/Config/BaseDeviceProfiles.ini
--- a/Engine/Config/BaseDeviceProfiles.ini
+++ b/Engine/Config/BaseDeviceProfiles.ini
@@ -124,7 +124,6 @@ TextureLODGroups=(Group=TEXTUREGROUP_World,MinLODSize=1,MaxLODSize=8192,LODBias=
 [Windows DeviceProfile]
 DeviceType=Windows
 BaseProfileName=
-+CVars=r.Vulkan.UseRealUBs=0		; Required until we fix HLR fall-out
 ; +CVars=r.Vulkan.EnableValidation=0
 
 [WindowsNoEditor DeviceProfile]
@@ -646,7 +645,6 @@ BaseProfileName=
 +CVars=r.LightShaftQuality=0
 +CVars=r.RefractionQuality=0
 +CVars=r.ShadowQuality=2
-+CVars=r.Vulkan.UseRealUBs=0		; Required until we fix HLR fall-out
 
 [LuminClient DeviceProfile]
 DeviceType=Lumin
diff --git a/Engine/Source/Runtime/RHI/Private/RHI.cpp b/Engine/Source/Runtime/RHI/Private/RHI.cpp
--- a/Engine/Source/Runtime/RHI/Private/RHI.cpp
+++ b/Engine/Source/Runtime/RHI/Private/RHI.cpp
@@ -36,7 +36,7 @@ DEFINE_STAT(STAT_PixelBufferMemory);
 
 static FAutoConsoleVariable CVarUseVulkanRealUBs(
 	TEXT("r.Vulkan.UseRealUBs"),
-	0,
+	1,
 	TEXT("0: Emulate uniform buffers on Vulkan SM4/SM5 (debugging ONLY)\n")
 	TEXT("1: Use real uniform buffers [default]"),
 	ECVF_ReadOnly
diff --git a/Engine/Source/Runtime/VulkanRHI/Public/VulkanMemory.h b/Engine/Source/Runtime/VulkanRHI/Public/VulkanMemory.h
--- a/Engine/Source/Runtime/VulkanRHI/Public/VulkanMemory.h
+++ b/Engine/Source/Runtime/VulkanRHI/Public/VulkanMemory.h
@@ -33,7 +33,7 @@ extern EDelayAcquireImageType GVulkanDelayAcquireImage;
 
 namespace VulkanRHI
 {
-	uint32 GetMaxSize(class FBufferAllocation* S);
+	static uint32 GetMaxSize(class FBufferAllocation* Allocation);
 	class FFenceManager;
 
 	extern int32 GVulkanUseBufferBinning;
@@ -532,10 +532,10 @@ namespace VulkanRHI
 			, Owner(InOwner)
 			, Handle(InHandle)
 		{
-			uint32 S = GetMaxSize(InOwner);
-			check(InAlignedOffset <= S);
-			check(InAllocationOffset + InAllocationSize <= S);
-			check(InAlignedOffset + InRequestedSize <= S);
+			uint32 Size = GetMaxSize(InOwner);
+			check(InAlignedOffset <= Size);
+			check(InAllocationOffset + InAllocationSize <= Size);
+			check(InAlignedOffset + InRequestedSize <= Size);
 		}
 
 		virtual ~FBufferSuballocation();
@@ -747,9 +747,9 @@ namespace VulkanRHI
 
 		friend class FResourceHeapManager;
 	};
-	inline uint32 GetMaxSize(FBufferAllocation* S)
+	static inline uint32 GetMaxSize(FBufferAllocation* Allocation)
 	{
-		return S->GetMaxSize();
+		return Allocation->GetMaxSize();
 	}
 
 	// Manages heaps and their interactions
