DR - Fix for Vulkan copies
#rnx

diff --git a/Engine/Source/Runtime/VulkanRHI/Private/VulkanTexture.cpp b/Engine/Source/Runtime/VulkanRHI/Private/VulkanTexture.cpp
--- a/Engine/Source/Runtime/VulkanRHI/Private/VulkanTexture.cpp
+++ b/Engine/Source/Runtime/VulkanRHI/Private/VulkanTexture.cpp
@@ -2242,10 +2242,10 @@ void FVulkanCommandListContext::RHICopyTexture(FTextureRHIParamRef SourceTexture
 
 	VkImageCopy Region;
 	FMemory::Memzero(Region);
-	if (CopyInfo.Size.X == 0 || CopyInfo.Size.Y == 0)
+	if (CopyInfo.Size == FIntVector::ZeroValue)
 	{
 		// Copy whole texture when zero vector is specified for region size
-		ensure(SrcSurface.Width == DstSurface.Width && SrcSurface.Height == DstSurface.Height);
+		ensure(SrcSurface.Width <= DstSurface.Width && SrcSurface.Height <= DstSurface.Height);
 		Region.extent.width = FMath::Max(1u, SrcSurface.Width >> CopyInfo.SourceMipIndex);
 		Region.extent.height = FMath::Max(1u, SrcSurface.Height >> CopyInfo.SourceMipIndex);
 	}
