DR - Fix for Lumin mobile asserting on Vulkan
#rnx

diff --git a/Engine/Source/Runtime/VulkanRHI/Private/Lumin/VulkanLuminPlatform.cpp b/Engine/Source/Runtime/VulkanRHI/Private/Lumin/VulkanLuminPlatform.cpp
--- a/Engine/Source/Runtime/VulkanRHI/Private/Lumin/VulkanLuminPlatform.cpp
+++ b/Engine/Source/Runtime/VulkanRHI/Private/Lumin/VulkanLuminPlatform.cpp
@@ -153,3 +153,8 @@ void FVulkanLuminPlatform::EnablePhysicalDeviceFeatureExtensions(VkDeviceCreateI
 	SamplerConversion.samplerYcbcrConversion = VK_TRUE;
 	DeviceInfo.pNext = &SamplerConversion;
 }
+
+bool FVulkanLuminPlatform::RequiresMobileRenderer()
+{
+	return !FLuminPlatformMisc::ShouldUseDesktopVulkan();
+}
diff --git a/Engine/Source/Runtime/VulkanRHI/Private/Lumin/VulkanLuminPlatform.h b/Engine/Source/Runtime/VulkanRHI/Private/Lumin/VulkanLuminPlatform.h
--- a/Engine/Source/Runtime/VulkanRHI/Private/Lumin/VulkanLuminPlatform.h
+++ b/Engine/Source/Runtime/VulkanRHI/Private/Lumin/VulkanLuminPlatform.h
@@ -47,6 +47,8 @@ public:
 	static bool SupportsASTCTextureFormats() { return true; }
 	static bool SupportsQuerySurfaceProperties() { return false; }
 
+	static bool RequiresMobileRenderer();
+
 	static void SetupFeatureLevels();
 
 	static bool SupportsStandardSwapchain() { return false; }
