Two vulkan crash-fixes

diff --git a/Engine/Source/Runtime/VulkanRHI/Private/VulkanUAV.cpp b/Engine/Source/Runtime/VulkanRHI/Private/VulkanUAV.cpp
--- a/Engine/Source/Runtime/VulkanRHI/Private/VulkanUAV.cpp
+++ b/Engine/Source/Runtime/VulkanRHI/Private/VulkanUAV.cpp
@@ -15,8 +15,12 @@ FVulkanShaderResourceView::FVulkanShaderResourceView(FVulkanDevice* Device, FRHI
 	, SourceRHIBuffer(InRHIBuffer)
 	, VolatileLockCounter(MAX_uint32)
 {
-	int32 NumBuffers = SourceBuffer->IsVolatile() ? 1 : SourceBuffer->GetNumBuffers();
-	BufferViews.AddZeroed(NumBuffers);
+	check(Device);
+	if(SourceBuffer)
+	{
+		int32 NumBuffers = SourceBuffer->IsVolatile() ? 1 : SourceBuffer->GetNumBuffers();
+		BufferViews.AddZeroed(NumBuffers);
+	}
 	check(BufferViewFormat != PF_Unknown);
 }
 
@@ -24,6 +28,7 @@ FVulkanShaderResourceView::FVulkanShaderResourceView(FVulkanDevice* Device, FRHI
 FVulkanShaderResourceView::~FVulkanShaderResourceView()
 {
 	Clear();
+	Device = nullptr;
 }
 
 void FVulkanShaderResourceView::Clear()
@@ -37,7 +42,6 @@ void FVulkanShaderResourceView::Clear()
 		TextureView.Destroy(*Device);
 	}
 	SourceTexture = nullptr;
-	Device = nullptr;
 }
 
 void FVulkanShaderResourceView::Rename(FRHIResource* InRHIBuffer, FVulkanResourceMultiBuffer* InSourceBuffer, uint32 InSize, EPixelFormat InFormat)
