DR - Capture callstack to debug buffer tracking
#rnx

diff --git a/Engine/Source/Runtime/VulkanRHI/Private/VulkanMemory.cpp b/Engine/Source/Runtime/VulkanRHI/Private/VulkanMemory.cpp
--- a/Engine/Source/Runtime/VulkanRHI/Private/VulkanMemory.cpp
+++ b/Engine/Source/Runtime/VulkanRHI/Private/VulkanMemory.cpp
@@ -25,11 +25,11 @@ RENDERCORE_API	void DumpRenderTargetPoolMemory(FOutputDevice& OutputDevice);
 #if VULKAN_MEMORY_TRACK_CALLSTACK
 static FCriticalSection GStackTraceMutex;
 static char GStackTrace[65536];
-static void CaptureCallStack(FString& OutCallstack)
+static void CaptureCallStack(FString& OutCallstack, int32 Delta = 3)
 {
 	FScopeLock ScopeLock(&GStackTraceMutex);
 	GStackTrace[0] = 0;
-	FPlatformStackWalk::StackWalkAndDump(GStackTrace, 65535, 3);
+	FPlatformStackWalk::StackWalkAndDump(GStackTrace, 65535, Delta);
 	OutCallstack = ANSI_TO_TCHAR(GStackTrace);
 }
 #endif
@@ -313,12 +313,10 @@ namespace VulkanRHI
 		static uint32 ID = 0;
 		NewAllocation->UID = ++ID;
 #endif
+
 #if VULKAN_MEMORY_TRACK_CALLSTACK
 		CaptureCallStack(NewAllocation->Callstack);
 #endif
-
-
-
 		++NumAllocations;
 		PeakNumAllocations = FMath::Max(NumAllocations, PeakNumAllocations);
 #if !VULKAN_SINGLE_ALLOCATION_PER_RESOURCE
@@ -1652,6 +1650,9 @@ namespace VulkanRHI
 					FStagingBuffer* Buffer = FreeBuffer.StagingBuffer;
 					FreeStagingBuffers.RemoveAtSwap(Index, 1, false);
 					UsedStagingBuffers.Add(Buffer);
+#if VULKAN_MEMORY_TRACK_CALLSTACK
+					CaptureCallStack(Buffer->Callstack, 1);
+#endif
 					return Buffer;
 				}
 			}
@@ -1691,6 +1692,10 @@ namespace VulkanRHI
 			UsedMemory += StagingBuffer->GetSize();
 			PeakUsedMemory = FMath::Max(UsedMemory, PeakUsedMemory);
 		}
+
+#if VULKAN_MEMORY_TRACK_CALLSTACK
+		CaptureCallStack(StagingBuffer->Callstack, 1);
+#endif
 		return StagingBuffer;
 	}
 
diff --git a/Engine/Source/Runtime/VulkanRHI/Public/VulkanMemory.h b/Engine/Source/Runtime/VulkanRHI/Public/VulkanMemory.h
--- a/Engine/Source/Runtime/VulkanRHI/Public/VulkanMemory.h
+++ b/Engine/Source/Runtime/VulkanRHI/Public/VulkanMemory.h
@@ -985,6 +985,10 @@ namespace VulkanRHI
 			ResourceAllocation->InvalidateMappedMemory();
 		}
 
+#if VULKAN_MEMORY_TRACK_CALLSTACK
+		FString Callstack;
+#endif
+
 	protected:
 		TRefCountPtr<FOldResourceAllocation> ResourceAllocation;
 		VkBuffer Buffer;
