DR - Slightly cleaner allocator dump
#rnx

diff --git a/Engine/Source/Runtime/VulkanRHI/Private/VulkanMemory.cpp b/Engine/Source/Runtime/VulkanRHI/Private/VulkanMemory.cpp
--- a/Engine/Source/Runtime/VulkanRHI/Private/VulkanMemory.cpp
+++ b/Engine/Source/Runtime/VulkanRHI/Private/VulkanMemory.cpp
@@ -381,13 +381,23 @@ namespace VulkanRHI
 			FHeapInfo& HeapInfo = HeapInfos[Index];
 			UE_LOG(LogVulkanRHI, Display, TEXT("\tHeap %d, %d allocations"), Index, HeapInfo.Allocations.Num());
 			uint64 TotalSize = 0;
+
+			if (HeapInfo.Allocations.Num() > 0)
+			{
+#if VULKAN_MEMORY_TRACK_FILE_LINE
+				UE_LOG(LogVulkanRHI, Display, TEXT("\t\tAlloc AllocSize(MB) TotalSize(MB)    Handle  UID  File(Line)"));
+#else
+				UE_LOG(LogVulkanRHI, Display, TEXT("\t\tAlloc AllocSize(MB) TotalSize(MB)    Handle"));
+#endif
+			}
+
 			for (int32 SubIndex = 0; SubIndex < HeapInfo.Allocations.Num(); ++SubIndex)
 			{
 				FDeviceMemoryAllocation* Allocation = HeapInfo.Allocations[SubIndex];
 #if VULKAN_MEMORY_TRACK_FILE_LINE
-				UE_LOG(LogVulkanRHI, Display, TEXT("\t\t%d Size %.2f MB  %2.f Handle %p ID %d %s(%d)"), SubIndex, Allocation->Size / 1024.f / 1024.f, TotalSize / 1024.0f / 1024.0f, (void*)Allocation->Handle, Allocation->UID, ANSI_TO_TCHAR(Allocation->File), Allocation->Line);
+				UE_LOG(LogVulkanRHI, Display, TEXT("\t\t%5d %13.3f %13.3f %p %4d %s(%d)"), SubIndex, Allocation->Size / 1024.f / 1024.f, TotalSize / 1024.0f / 1024.0f, (void*)Allocation->Handle, Allocation->UID, ANSI_TO_TCHAR(Allocation->File), Allocation->Line);
 #else
-				UE_LOG(LogVulkanRHI, Display, TEXT("\t\t%d Size %.2f MB  %2.f Handle %p"), SubIndex, Allocation->Size / 1024.f / 1024.f, TotalSize / 1024.0f / 1024.0f, (void*)Allocation->Handle);
+				UE_LOG(LogVulkanRHI, Display, TEXT("\t\t%5d %13.3f %13.3f %p"), SubIndex, Allocation->Size / 1024.f / 1024.f, TotalSize / 1024.0f / 1024.0f, (void*)Allocation->Handle);
 #endif
 				TotalSize += Allocation->Size;
 			}
