DR - Added some more missing stats dumping Vulkan memory
#rnx

diff --git a/Engine/Source/Runtime/VulkanRHI/Private/VulkanMemory.cpp b/Engine/Source/Runtime/VulkanRHI/Private/VulkanMemory.cpp
--- a/Engine/Source/Runtime/VulkanRHI/Private/VulkanMemory.cpp
+++ b/Engine/Source/Runtime/VulkanRHI/Private/VulkanMemory.cpp
@@ -1770,7 +1770,7 @@ namespace VulkanRHI
 		for (int32 Index = 0; Index < UsedStagingBuffers.Num(); ++Index)
 		{
 			FStagingBuffer* Buffer = UsedStagingBuffers[Index];
-			UE_LOG(LogVulkanRHI, Display, TEXT("%6d %p %p"), Index, (void*)Buffer->GetHandle(), (void*)Buffer->ResourceAllocation->GetHandle());
+			UE_LOG(LogVulkanRHI, Display, TEXT("%6d %p %p %6d"), Index, (void*)Buffer->GetHandle(), (void*)Buffer->ResourceAllocation->GetHandle(), Buffer->BufferSize);
 		}
 
 		UE_LOG(LogVulkanRHI, Display, TEXT("Pending CmdBuffer   Fence   BufferHandle ResourceAllocation"));
@@ -1790,11 +1790,11 @@ namespace VulkanRHI
 			}
 		}
 
-		UE_LOG(LogVulkanRHI, Display, TEXT("Free   BufferHandle ResourceAllocation"));
+		UE_LOG(LogVulkanRHI, Display, TEXT("Free   BufferHandle ResourceAllocation Size"));
 		for (int32 Index = 0; Index < FreeStagingBuffers.Num(); ++Index)
 		{
 			FFreeEntry& Entry = FreeStagingBuffers[Index];
-			UE_LOG(LogVulkanRHI, Display, TEXT("%6d %p %p"), Index, (void*)Entry.StagingBuffer->GetHandle(), (void*)Entry.StagingBuffer->ResourceAllocation->GetHandle());
+			UE_LOG(LogVulkanRHI, Display, TEXT("%6d %p %p %6d"), Index, (void*)Entry.StagingBuffer->GetHandle(), (void*)Entry.StagingBuffer->ResourceAllocation->GetHandle(), Entry.StagingBuffer->BufferSize);
 		}
 	}
 #endif
