DR - Decouple Vulkan UBs from VBs & IBs mem priorities
#rnx

diff --git a/Engine/Source/Runtime/VulkanRHI/Private/VulkanMemory.cpp b/Engine/Source/Runtime/VulkanRHI/Private/VulkanMemory.cpp
--- a/Engine/Source/Runtime/VulkanRHI/Private/VulkanMemory.cpp
+++ b/Engine/Source/Runtime/VulkanRHI/Private/VulkanMemory.cpp
@@ -1247,39 +1247,40 @@ namespace VulkanRHI
 		const VkPhysicalDeviceLimits& Limits = Device->GetLimits();
 		uint32 Alignment = 1;
 
-		bool bIsStorageOrTexel = (BufferUsageFlags & (VK_BUFFER_USAGE_UNIFORM_TEXEL_BUFFER_BIT | VK_BUFFER_USAGE_STORAGE_TEXEL_BUFFER_BIT | VK_BUFFER_USAGE_STORAGE_BUFFER_BIT)) != 0;
-		if (bIsStorageOrTexel)
+		float Priority = VULKAN_MEMORY_MEDIUM_PRIORITY;
+
+		bool bIsTexelBuffer = (BufferUsageFlags & (VK_BUFFER_USAGE_UNIFORM_TEXEL_BUFFER_BIT | VK_BUFFER_USAGE_STORAGE_TEXEL_BUFFER_BIT)) != 0;
+		bool bIsStorageBuffer = (BufferUsageFlags & VK_BUFFER_USAGE_STORAGE_BUFFER_BIT) != 0;
+		if (bIsTexelBuffer)
 		{
-			Alignment = FMath::Max(Alignment, ((BufferUsageFlags & (VK_BUFFER_USAGE_UNIFORM_TEXEL_BUFFER_BIT | VK_BUFFER_USAGE_STORAGE_TEXEL_BUFFER_BIT)) != 0) ? (uint32)Limits.minTexelBufferOffsetAlignment : 1u);
-			Alignment = FMath::Max(Alignment, ((BufferUsageFlags & VK_BUFFER_USAGE_STORAGE_BUFFER_BIT) != 0) ? (uint32)Limits.minStorageBufferOffsetAlignment : 1u);
+			Alignment = FMath::Max(Alignment, (uint32)Limits.minTexelBufferOffsetAlignment);
+		}
+		else if (bIsStorageBuffer)
+		{
+			Alignment = FMath::Max(Alignment, (uint32)Limits.minStorageBufferOffsetAlignment);
 		}
 		else
 		{
-			Alignment = (uint32)Limits.minUniformBufferOffsetAlignment;
+			bool bIsVertexOrIndexBuffer = (BufferUsageFlags & (VK_BUFFER_USAGE_VERTEX_BUFFER_BIT | VK_BUFFER_USAGE_INDEX_BUFFER_BIT)) != 0;
+			if (bIsVertexOrIndexBuffer)
+			{
+				// No alignment restrictions on Vertex or Index buffers, can live on CPU mem
+				Priority = VULKAN_MEMORY_LOW_PRIORITY;
+			}
+			else
+			{
+				// Uniform buffer
+				ensure((BufferUsageFlags & VK_BUFFER_USAGE_UNIFORM_BUFFER_BIT) == VK_BUFFER_USAGE_UNIFORM_BUFFER_BIT);
+				Alignment = FMath::Max(Alignment, (uint32)Limits.minUniformBufferOffsetAlignment);
 
-			// reduce the number of unique buffer types by allowing all buffer types in non-texel/storage buffers
-			// @FYI ROLANDO: I want to comment out this line as this effectively prevents us from lowering the priority on vertex/index buffers
-			BufferUsageFlags |= VK_BUFFER_USAGE_UNIFORM_BUFFER_BIT | VK_BUFFER_USAGE_INDEX_BUFFER_BIT | VK_BUFFER_USAGE_VERTEX_BUFFER_BIT | VK_BUFFER_USAGE_INDIRECT_BUFFER_BIT;
+				Priority = VULKAN_MEMORY_HIGHER_PRIORITY;
+			}
 		}
 
-		float Priority = VULKAN_MEMORY_MEDIUM_PRIORITY;
-
-		if(BufferUsageFlags & (VK_BUFFER_USAGE_STORAGE_TEXEL_BUFFER_BIT| VK_BUFFER_USAGE_INDIRECT_BUFFER_BIT))
+		if (BufferUsageFlags & (VK_BUFFER_USAGE_STORAGE_TEXEL_BUFFER_BIT | VK_BUFFER_USAGE_INDIRECT_BUFFER_BIT))
 		{
 			Priority = VULKAN_MEMORY_HIGHEST_PRIORITY;
 		}
-		if(BufferUsageFlags & (VK_BUFFER_USAGE_INDEX_BUFFER_BIT| VK_BUFFER_USAGE_VERTEX_BUFFER_BIT))
-		{
-			Priority = VULKAN_MEMORY_LOW_PRIORITY;
-		}
-		else if(BufferUsageFlags & (VK_BUFFER_USAGE_UNIFORM_BUFFER_BIT))
-		{
-			Priority = VULKAN_MEMORY_HIGHER_PRIORITY;
-		}
-
-//		const bool bIsUniformBuffer = ((BufferUsageFlags & VK_BUFFER_USAGE_UNIFORM_BUFFER_BIT) != 0);
-
-		FScopeLock ScopeLock(&GResourceHeapLock);
 
 		int32 PoolSize = (int32)GetPoolTypeForAlloc(Size, Alignment);
 		if (PoolSize != (int32)EPoolSizes::SizesCount)
@@ -1287,6 +1288,8 @@ namespace VulkanRHI
 			Size = PoolSizes[PoolSize];
 		}
 
+		FScopeLock ScopeLock(&GResourceHeapLock);
+
 		for (int32 Index = 0; Index < UsedBufferAllocations[PoolSize].Num(); ++Index)
 		{
 			FBufferAllocation* BufferAllocation = UsedBufferAllocations[PoolSize][Index];
@@ -1766,14 +1769,14 @@ namespace VulkanRHI
 	void FStagingManager::DumpMemory()
 	{
 		UE_LOG(LogVulkanRHI, Display, TEXT("StagingManager %d Used %d Pending Free %d Free"), UsedStagingBuffers.Num(), PendingFreeStagingBuffers.Num(), FreeStagingBuffers.Num());
-		UE_LOG(LogVulkanRHI, Display, TEXT("Used   BufferHandle ResourceAllocation"));
+		UE_LOG(LogVulkanRHI, Display, TEXT("Used   BufferHandle      ResourceAllocation Size"));
 		for (int32 Index = 0; Index < UsedStagingBuffers.Num(); ++Index)
 		{
 			FStagingBuffer* Buffer = UsedStagingBuffers[Index];
 			UE_LOG(LogVulkanRHI, Display, TEXT("%6d %p %p %6d"), Index, (void*)Buffer->GetHandle(), (void*)Buffer->ResourceAllocation->GetHandle(), Buffer->BufferSize);
 		}
 
-		UE_LOG(LogVulkanRHI, Display, TEXT("Pending CmdBuffer   Fence   BufferHandle ResourceAllocation"));
+		UE_LOG(LogVulkanRHI, Display, TEXT("Pending CmdBuffer   Fence   BufferHandle    ResourceAllocation Size"));
 		for (int32 Index = 0; Index < PendingFreeStagingBuffers.Num(); ++Index)
 		{
 			FPendingItemsPerCmdBuffer& ItemPerCmdBuffer = PendingFreeStagingBuffers[Index];
@@ -1785,7 +1788,7 @@ namespace VulkanRHI
 				for (int32 BufferIndex = 0; BufferIndex < ItemsPerFence.Resources.Num(); ++BufferIndex)
 				{
 					FStagingBuffer* Buffer = ItemsPerFence.Resources[BufferIndex];
-					UE_LOG(LogVulkanRHI, Display, TEXT("                   %p %p"), (void*)Buffer->GetHandle(), (void*)Buffer->ResourceAllocation->GetHandle());
+					UE_LOG(LogVulkanRHI, Display, TEXT("                   %p %p %6d"), (void*)Buffer->GetHandle(), (void*)Buffer->ResourceAllocation->GetHandle(), Buffer->BufferSize);
 				}
 			}
 		}
