DR - Disable GPU validation warning (seems to be related to multiple GPUs on one machine)
#rnx

diff --git a/Engine/Source/Runtime/VulkanRHI/Private/VulkanDebug.cpp b/Engine/Source/Runtime/VulkanRHI/Private/VulkanDebug.cpp
--- a/Engine/Source/Runtime/VulkanRHI/Private/VulkanDebug.cpp
+++ b/Engine/Source/Runtime/VulkanRHI/Private/VulkanDebug.cpp
@@ -194,6 +194,11 @@ static VkBool32 DebugUtilsCallback(VkDebugUtilsMessageSeverityFlagBitsEXT MsgSev
 		// Warning: *** [Error:Validation341838324(VUID-VkSwapchainCreateInfoKHR-imageExtent-01274)] vkCreateSwapChainKHR() called with imageExtent = (8,8), which is outside the bounds returned by vkGetPhysicalDeviceSurfaceCapabilitiesKHR(): currentExtent = (0,0), minImageExtent = (0,0), maxImageExtent = (0,0).
 		return VK_FALSE;
 	}
+	else if (!FCStringAnsi::Strcmp(CallbackData->pMessageIdName, "UNASSIGNED-GPU-Assisted Validation Setup Error."))
+	{
+		// *** [Error:Validation(UNASSIGNED-GPU-Assisted Validation Setup Error.)] Unable to reserve descriptor binding slot on a device with only one slot.
+		return VK_FALSE;
+	}
 
 	const TCHAR* Severity = TEXT("");
 	if (bError)
