DR - Show number if Vulkan validation message is missing
#rnx

diff --git a/Engine/Source/Runtime/VulkanRHI/Private/VulkanDebug.cpp b/Engine/Source/Runtime/VulkanRHI/Private/VulkanDebug.cpp
--- a/Engine/Source/Runtime/VulkanRHI/Private/VulkanDebug.cpp
+++ b/Engine/Source/Runtime/VulkanRHI/Private/VulkanDebug.cpp
@@ -178,9 +178,10 @@ static VkBool32 DebugUtilsCallback(VkDebugUtilsMessageSeverityFlagBitsEXT MsgSev
 	const bool bError = (MsgSeverity & VK_DEBUG_UTILS_MESSAGE_SEVERITY_ERROR_BIT_EXT) != 0;
 	const bool bWarning = (MsgSeverity & VK_DEBUG_UTILS_MESSAGE_SEVERITY_WARNING_BIT_EXT) != 0;
 
-	ensure(CallbackData->pMessageIdName);
 	if (!CallbackData->pMessageIdName)
 	{
+		VULKAN_REPORT_LOG(TEXT("*** [%s:%s(NULL)%d] %s"), Severity, Type, CallbackData->messageIdNumber, ANSI_TO_TCHAR(CallbackData->pMessage));
+		ensure(0);
 		return VK_FALSE;
 	}
 
@@ -282,14 +283,7 @@ static VkBool32 DebugUtilsCallback(VkDebugUtilsMessageSeverityFlagBitsEXT MsgSev
 
 	if (bPrintMessage)
 	{
-		if (CallbackData->pMessageIdName)
-		{
-			VULKAN_REPORT_LOG(TEXT("*** [%s:%s(%s)] %s"), Severity, Type, ANSI_TO_TCHAR(CallbackData->pMessageIdName), ANSI_TO_TCHAR(CallbackData->pMessage));
-		}
-		else
-		{
-			VULKAN_REPORT_LOG(TEXT("*** [%s:%s] %s"), Severity, Type, ANSI_TO_TCHAR(CallbackData->pMessage));
-		}
+		VULKAN_REPORT_LOG(TEXT("*** [%s:%s(%s)] %s"), Severity, Type, ANSI_TO_TCHAR(CallbackData->pMessageIdName), ANSI_TO_TCHAR(CallbackData->pMessage));
 		if (bUniqueMessages)
 		{
 			SeenCodes[(int32)MsgBucket].Add(CRC);
