DR - Fix for vulkan dump crash
#rnx

diff --git a/Engine/Source/Runtime/VulkanRHI/Private/VulkanDebug.cpp b/Engine/Source/Runtime/VulkanRHI/Private/VulkanDebug.cpp
--- a/Engine/Source/Runtime/VulkanRHI/Private/VulkanDebug.cpp
+++ b/Engine/Source/Runtime/VulkanRHI/Private/VulkanDebug.cpp
@@ -1965,7 +1965,12 @@ void FWrapLayer::UpdateDescriptorSets(VkResult Result, VkDevice Device, uint32 D
 					for (uint32 SubIndex = 0; SubIndex < DescriptorWrites[Index].descriptorCount; ++SubIndex)
 					{
 #if VULKAN_ENABLE_DUMP_LAYER
+
+#if VULKAN_ENABLE_BUFFER_TRACKING_LAYER
 						DebugLog += FString::Printf(TEXT("%s\tpTexelBufferView[%d]=0x%p(B:0x%p)\n"), Tabs, SubIndex, DescriptorWrites[Index].pTexelBufferView[SubIndex], FindTrackingBuffer(DescriptorWrites[Index].pTexelBufferView[SubIndex]));
+#else
+						DebugLog += FString::Printf(TEXT("%s\tpTexelBufferView[%d]=0x%p\n"), Tabs, SubIndex, DescriptorWrites[Index].pTexelBufferView[SubIndex]);
+#endif
 #endif
 					}
 				}
