- RHIMapStagingBuffer/RHIUnmapStagingBuffer now calls RHIMapStagingBuffer_RenderThread/RHIMapStagingBuffer_RenderThread. Default implementation works as before. This allows implementations to not block on the RHI thread
- Staging textures are now allocated in cpu accessible memory vulkan, so we dont need to block when reading from it
- Removed the option to not stall the RHI thread as a parameter(introduced on _VT). This will cause race conditions on platforms that don't support this.

diff --git a/Engine/Source/Runtime/RHI/Private/RHICommandList.cpp b/Engine/Source/Runtime/RHI/Private/RHICommandList.cpp
--- a/Engine/Source/Runtime/RHI/Private/RHICommandList.cpp
+++ b/Engine/Source/Runtime/RHI/Private/RHICommandList.cpp
@@ -2444,6 +2444,25 @@ void FDynamicRHI::RHIUnlockTextureCubeFace_RenderThread(class FRHICommandListImm
 	GDynamicRHI->RHIUnlockTextureCubeFace(Texture, FaceIndex, ArrayIndex, MipIndex, bLockWithinMiptail);
 }
 
+
+void FDynamicRHI::RHIMapStagingSurface_RenderThread(class FRHICommandListImmediate& RHICmdList, FTextureRHIParamRef Texture, void*& OutData, int32& OutWidth, int32& OutHeight)
+{
+	{
+		QUICK_SCOPE_CYCLE_COUNTER(STAT_RHIMETHOD_MapStagingSurface_Flush);
+		RHICmdList.ImmediateFlush(EImmediateFlushType::FlushRHIThread);
+	}
+	GDynamicRHI->RHIMapStagingSurface(Texture, OutData, OutWidth, OutHeight);
+}
+
+void FDynamicRHI::RHIUnmapStagingSurface_RenderThread(class FRHICommandListImmediate& RHICmdList, FTextureRHIParamRef Texture)
+{
+	{
+		QUICK_SCOPE_CYCLE_COUNTER(STAT_RHIMETHOD_MapStagingSurface_Flush);
+		RHICmdList.ImmediateFlush(EImmediateFlushType::FlushRHIThread);
+	}
+	GDynamicRHI->RHIUnmapStagingSurface(Texture);
+}
+
 void FDynamicRHI::RHIReadSurfaceFloatData_RenderThread(class FRHICommandListImmediate& RHICmdList, FTextureRHIParamRef Texture, FIntRect Rect, TArray<FFloat16Color>& OutData, ECubeFace CubeFace, int32 ArrayIndex, int32 MipIndex)
 {
 	QUICK_SCOPE_CYCLE_COUNTER(STAT_RHIMETHOD_ReadSurfaceFloatData_Flush);
diff --git a/Engine/Source/Runtime/RHI/Public/DynamicRHI.h b/Engine/Source/Runtime/RHI/Public/DynamicRHI.h
--- a/Engine/Source/Runtime/RHI/Public/DynamicRHI.h
+++ b/Engine/Source/Runtime/RHI/Public/DynamicRHI.h
@@ -1145,6 +1145,9 @@ public:
 	virtual void RHIAcquireTransientResource_RenderThread(FStructuredBufferRHIParamRef Buffer) { }
 	virtual void RHIDiscardTransientResource_RenderThread(FStructuredBufferRHIParamRef Buffer) { }
 
+
+	virtual void RHIMapStagingSurface_RenderThread(class FRHICommandListImmediate& RHICmdList, FTextureRHIParamRef Texture, void*& OutData, int32& OutWidth, int32& OutHeight);
+	virtual void RHIUnmapStagingSurface_RenderThread(class FRHICommandListImmediate& RHICmdList, FTextureRHIParamRef Texture);
 	virtual void RHIReadSurfaceFloatData_RenderThread(class FRHICommandListImmediate& RHICmdList, FTextureRHIParamRef Texture, FIntRect Rect, TArray<FFloat16Color>& OutData, ECubeFace CubeFace, int32 ArrayIndex, int32 MipIndex);
 
 	//Utilities
diff --git a/Engine/Source/Runtime/RHI/Public/RHICommandList.h b/Engine/Source/Runtime/RHI/Public/RHICommandList.h
--- a/Engine/Source/Runtime/RHI/Public/RHICommandList.h
+++ b/Engine/Source/Runtime/RHI/Public/RHICommandList.h
@@ -4029,24 +4029,14 @@ public:
 		GDynamicRHI->RHIReadSurfaceData(Texture, Rect, OutData, InFlags);
 	}
 	
-	FORCEINLINE void MapStagingSurface(FTextureRHIParamRef Texture,void*& OutData,int32& OutWidth,int32& OutHeight, bool bFlushRHIThread = true)
+	FORCEINLINE void MapStagingSurface(FTextureRHIParamRef Texture,void*& OutData,int32& OutWidth,int32& OutHeight)
 	{
-		if (bFlushRHIThread)
-		{
-			QUICK_SCOPE_CYCLE_COUNTER(STAT_RHIMETHOD_MapStagingSurface_Flush);
-			ImmediateFlush(EImmediateFlushType::FlushRHIThread);
-		}
-		GDynamicRHI->RHIMapStagingSurface(Texture,OutData,OutWidth,OutHeight);
+		GDynamicRHI->RHIMapStagingSurface_RenderThread(*this, Texture,OutData,OutWidth,OutHeight);
 	}
 	
-	FORCEINLINE void UnmapStagingSurface(FTextureRHIParamRef Texture, bool bFlushRHIThread = true)
+	FORCEINLINE void UnmapStagingSurface(FTextureRHIParamRef Texture)
 	{
-		if (bFlushRHIThread)
-		{
-			QUICK_SCOPE_CYCLE_COUNTER(STAT_RHIMETHOD_UnmapStagingSurface_Flush);
-			ImmediateFlush(EImmediateFlushType::FlushRHIThread);
-		}
-		GDynamicRHI->RHIUnmapStagingSurface(Texture);
+		GDynamicRHI->RHIUnmapStagingSurface_RenderThread(*this, Texture);
 	}
 	
 	FORCEINLINE void ReadSurfaceFloatData(FTextureRHIParamRef Texture,FIntRect Rect,TArray<FFloat16Color>& OutData,ECubeFace CubeFace,int32 ArrayIndex,int32 MipIndex)
diff --git a/Engine/Source/Runtime/Renderer/Private/VT/VirtualTextureFeedback.cpp b/Engine/Source/Runtime/Renderer/Private/VT/VirtualTextureFeedback.cpp
--- a/Engine/Source/Runtime/Renderer/Private/VT/VirtualTextureFeedback.cpp
+++ b/Engine/Source/Runtime/Renderer/Private/VT/VirtualTextureFeedback.cpp
@@ -110,9 +110,7 @@ bool FVirtualTextureFeedback::Map( FRHICommandListImmediate& RHICmdList, MapResu
 		OutResult.Rect = FeedbackEntryCPU.Rect;
 
 		int32 LockHeight = 0;
-		// We can avoid flushing RHI thread here, since we wait for GPU fence to complete
-		const bool bFlushRHIThread = false;
-		RHICmdList.MapStagingSurface(FeedbackEntryCPU.TextureCPU->GetRenderTargetItem().ShaderResourceTexture, *(void**)&OutResult.Buffer, OutResult.Pitch, LockHeight, bFlushRHIThread);
+		RHICmdList.MapStagingSurface(FeedbackEntryCPU.TextureCPU->GetRenderTargetItem().ShaderResourceTexture, *(void**)&OutResult.Buffer, OutResult.Pitch, LockHeight);
 
 		check(PendingTargetCount > 0u);
 		--PendingTargetCount;
@@ -128,7 +126,6 @@ void FVirtualTextureFeedback::Unmap( FRHICommandListImmediate& RHICmdList, int32
 {
 	check(FeedbackTextureCPU[MapHandle].TextureCPU.IsValid());
 
-	const bool bFlushRHIThread = false;
-	RHICmdList.UnmapStagingSurface(FeedbackTextureCPU[MapHandle].TextureCPU->GetRenderTargetItem().ShaderResourceTexture, bFlushRHIThread);
+	RHICmdList.UnmapStagingSurface(FeedbackTextureCPU[MapHandle].TextureCPU->GetRenderTargetItem().ShaderResourceTexture);
 	GRenderTargetPool.FreeUnusedResource(FeedbackTextureCPU[MapHandle].TextureCPU);
 }
diff --git a/Engine/Source/Runtime/VulkanRHI/Private/VulkanRenderTarget.cpp b/Engine/Source/Runtime/VulkanRHI/Private/VulkanRenderTarget.cpp
--- a/Engine/Source/Runtime/VulkanRHI/Private/VulkanRenderTarget.cpp
+++ b/Engine/Source/Runtime/VulkanRHI/Private/VulkanRenderTarget.cpp
@@ -954,84 +954,29 @@ void FVulkanDynamicRHI::RHIMapStagingSurface(FTextureRHIParamRef TextureRHI,void
 	check(TextureRHI2D);
 	FVulkanTexture2D* Texture2D = ResourceCast(TextureRHI2D);
 
-	VulkanRHI::FStagingBuffer** StagingBufferPtr = nullptr;
-	{
-		FScopeLock Lock(&GStagingMapLock);
-		StagingBufferPtr = &GPendingLockedStagingBuffers.FindOrAdd(Texture2D);
-		checkf(!*StagingBufferPtr, TEXT("Can't map the same texture twice!"));
-	}
 
 	OutWidth = Texture2D->GetSizeX();
 	OutHeight = Texture2D->GetSizeY();
 
-	uint32 BufferSize = OutWidth * OutHeight * VulkanRHI::GetNumBitsPerPixel(Texture2D->Surface.ViewFormat) / 8;
-	VulkanRHI::FStagingBuffer* StagingBuffer = Device->GetStagingManager().AcquireBuffer(BufferSize, VK_BUFFER_USAGE_TRANSFER_DST_BIT, true);
-	*StagingBufferPtr = StagingBuffer;
-
-	Device->PrepareForCPURead(); //make sure the results are ready 
-	FVulkanCmdBuffer* CmdBuffer = Device->GetImmediateContext().GetCommandBufferManager()->GetUploadCmdBuffer();
-
-	// Transition texture to source copy layout
-	VkImageLayout& CurrentLayout = Device->GetImmediateContext().TransitionAndLayoutManager.FindOrAddLayoutRW(Texture2D->Surface.Image, VK_IMAGE_LAYOUT_UNDEFINED);
-	bool bHadLayout = (CurrentLayout != VK_IMAGE_LAYOUT_UNDEFINED);
-	if (CurrentLayout != VK_IMAGE_LAYOUT_TRANSFER_SRC_OPTIMAL)
-	{
-		VulkanSetImageLayoutSimple(CmdBuffer->GetHandle(), Texture2D->Surface.Image, CurrentLayout, VK_IMAGE_LAYOUT_TRANSFER_SRC_OPTIMAL);
-	}
-
-	VkBufferImageCopy CopyRegion;
-	FMemory::Memzero(CopyRegion);
-	//Region.bufferOffset = 0;
-	CopyRegion.bufferRowLength = OutWidth;
-	CopyRegion.bufferImageHeight = OutHeight;
-	CopyRegion.imageSubresource.aspectMask = Texture2D->Surface.GetFullAspectMask();
-	//CopyRegion.imageSubresource.mipLevel = InMipIndex;
-	//CopyRegion.imageSubresource.baseArrayLayer = SrcBaseArrayLayer;
-	CopyRegion.imageSubresource.layerCount = 1;
-	CopyRegion.imageExtent.width = OutWidth;
-	CopyRegion.imageExtent.height = OutHeight;
-	CopyRegion.imageExtent.depth = 1;
-
-	VulkanRHI::vkCmdCopyImageToBuffer(CmdBuffer->GetHandle(), Texture2D->Surface.Image, VK_IMAGE_LAYOUT_TRANSFER_SRC_OPTIMAL, StagingBuffer->GetHandle(), 1, &CopyRegion);
-	// Transition back to original layout
-	if (bHadLayout && CurrentLayout != VK_IMAGE_LAYOUT_TRANSFER_SRC_OPTIMAL)
-	{
-		VulkanSetImageLayoutSimple(CmdBuffer->GetHandle(), Texture2D->Surface.Image, VK_IMAGE_LAYOUT_TRANSFER_SRC_OPTIMAL, CurrentLayout);
-	}
-	else
-	{
-		CurrentLayout = VK_IMAGE_LAYOUT_TRANSFER_SRC_OPTIMAL;
-	}
-
-	VkBufferMemoryBarrier Barrier;
-	ensure(StagingBuffer->GetSize() >= BufferSize);
-	//#todo-rco: Change offset if reusing a buffer suballocation
-	VulkanRHI::SetupAndZeroBufferBarrier(Barrier, VK_ACCESS_TRANSFER_WRITE_BIT, VK_ACCESS_MEMORY_READ_BIT, StagingBuffer->GetHandle(), 0/*StagingBuffer->GetOffset()*/, BufferSize);
-	VulkanRHI::vkCmdPipelineBarrier(CmdBuffer->GetHandle(), VK_PIPELINE_STAGE_TRANSFER_BIT, VK_PIPELINE_STAGE_HOST_BIT, 0, 0, nullptr, 1, &Barrier, 0, nullptr);
+	OutData = Texture2D->Surface.GetMappedPointer();
+	Texture2D->Surface.InvalidateMappedMemory();
+}
 
-	Device->GetImmediateContext().GetCommandBufferManager()->SubmitUploadCmdBuffer();
 
-	OutData = StagingBuffer->GetMappedPointer();
-	StagingBuffer->InvalidateMappedMemory();
+void FVulkanDynamicRHI::RHIMapStagingSurface_RenderThread(class FRHICommandListImmediate& RHICmdList, FTextureRHIParamRef Texture, void*& OutData, int32& OutWidth, int32& OutHeight)
+{
+	RHIMapStagingSurface(Texture, OutData, OutWidth, OutHeight);
 }
 
-void FVulkanDynamicRHI::RHIUnmapStagingSurface(FTextureRHIParamRef TextureRHI)
+void FVulkanDynamicRHI::RHIUnmapStagingSurface_RenderThread(class FRHICommandListImmediate& RHICmdList, FTextureRHIParamRef Texture)
 {
-	FRHITexture2D* TextureRHI2D = TextureRHI->GetTexture2D();
-	check(TextureRHI2D);
-	FVulkanTexture2D* Texture2D = ResourceCast(TextureRHI2D);
+	RHIUnmapStagingSurface(Texture);
+}
 
-	VulkanRHI::FStagingBuffer* StagingBuffer = nullptr;
-	{
-		FScopeLock Lock(&GStagingMapLock);
-		bool bFound = GPendingLockedStagingBuffers.RemoveAndCopyValue(Texture2D, StagingBuffer);
-		checkf(bFound, TEXT("Texture was not mapped!"));
-	}
 
-	ensure(!Device->GetImmediateContext().GetCommandBufferManager()->HasPendingUploadCmdBuffer());
 
-	Device->GetImmediateContext().GetCommandBufferManager()->PrepareForNewActiveCommandBuffer();
-	Device->GetStagingManager().ReleaseBuffer(nullptr, StagingBuffer);
+void FVulkanDynamicRHI::RHIUnmapStagingSurface(FTextureRHIParamRef TextureRHI)
+{
 }
 
 void FVulkanDynamicRHI::RHIReadSurfaceFloatData(FTextureRHIParamRef TextureRHI, FIntRect Rect, TArray<FFloat16Color>& OutData, ECubeFace CubeFace,int32 ArrayIndex,int32 MipIndex)
diff --git a/Engine/Source/Runtime/VulkanRHI/Private/VulkanResources.h b/Engine/Source/Runtime/VulkanRHI/Private/VulkanResources.h
--- a/Engine/Source/Runtime/VulkanRHI/Private/VulkanResources.h
+++ b/Engine/Source/Runtime/VulkanRHI/Private/VulkanResources.h
@@ -293,6 +293,8 @@ public:
 	virtual ~FVulkanSurface();
 
 	void Destroy();
+	void InvalidateMappedMemory(){ ResourceAllocation->InvalidateMappedMemory();}
+	void* GetMappedPointer() { return ResourceAllocation->GetMappedPointer(); }
 
 #if 0
 	/**
diff --git a/Engine/Source/Runtime/VulkanRHI/Private/VulkanTexture.cpp b/Engine/Source/Runtime/VulkanRHI/Private/VulkanTexture.cpp
--- a/Engine/Source/Runtime/VulkanRHI/Private/VulkanTexture.cpp
+++ b/Engine/Source/Runtime/VulkanRHI/Private/VulkanTexture.cpp
@@ -193,6 +193,12 @@ VkImage FVulkanSurface::CreateImage(
 	const FPixelFormatInfo& FormatInfo = GPixelFormats[InFormat];
 	VkFormat TextureFormat = (VkFormat)FormatInfo.PlatformFormat;
 
+
+	if(UEFlags & TexCreate_CPUReadback)
+	{
+		bForceLinearTexture = true;
+	}
+
 	checkf(TextureFormat != VK_FORMAT_UNDEFINED, TEXT("PixelFormat %d, is not supported for images"), (int32)InFormat);
 
 	VkImageCreateInfo TmpCreateInfo;
@@ -531,17 +537,20 @@ FVulkanSurface::FVulkanSurface(FVulkanDevice& InDevice, VkImageViewType Resource
 	const bool bCPUReadback = (UEFlags & TexCreate_CPUReadback) != 0;
 	const bool bDynamic = (UEFlags & TexCreate_Dynamic) != 0;
 
+
+	VkMemoryPropertyFlags MemoryFlags = bCPUReadback ? VK_MEMORY_PROPERTY_HOST_VISIBLE_BIT : VK_MEMORY_PROPERTY_DEVICE_LOCAL_BIT;
+
 #if VULKAN_SUPPORTS_DEDICATED_ALLOCATION
 	// Per https://developer.nvidia.com/what%E2%80%99s-your-vulkan-memory-type
 	VkDeviceSize SizeToBeConsideredForDedicated = 12 * 1024 * 1024;
 	if ((bRenderTarget || MemoryRequirements.size >= SizeToBeConsideredForDedicated) && InDevice.GetOptionalExtensions().HasKHRDedicatedAllocation)
 	{
-		ResourceAllocation = InDevice.GetResourceHeapManager().AllocateDedicatedImageMemory(Image, MemoryRequirements, VK_MEMORY_PROPERTY_DEVICE_LOCAL_BIT, __FILE__, __LINE__);
+		ResourceAllocation = InDevice.GetResourceHeapManager().AllocateDedicatedImageMemory(Image, MemoryRequirements, MemoryFlags, __FILE__, __LINE__);
 	}
 	else
 #endif
 	{
-		ResourceAllocation = InDevice.GetResourceHeapManager().AllocateImageMemory(MemoryRequirements, VK_MEMORY_PROPERTY_DEVICE_LOCAL_BIT, __FILE__, __LINE__);
+		ResourceAllocation = InDevice.GetResourceHeapManager().AllocateImageMemory(MemoryRequirements, MemoryFlags, __FILE__, __LINE__);
 	}
 	ResourceAllocation->BindImage(Device, Image);
 
diff --git a/Engine/Source/Runtime/VulkanRHI/Public/VulkanDynamicRHI.h b/Engine/Source/Runtime/VulkanRHI/Public/VulkanDynamicRHI.h
--- a/Engine/Source/Runtime/VulkanRHI/Public/VulkanDynamicRHI.h
+++ b/Engine/Source/Runtime/VulkanRHI/Public/VulkanDynamicRHI.h
@@ -128,6 +128,10 @@ public:
 	virtual void RHIReadSurfaceData(FTextureRHIParamRef Texture, FIntRect Rect, TArray<FLinearColor>& OutData, FReadSurfaceDataFlags InFlags) final override;
 	virtual void RHIMapStagingSurface(FTextureRHIParamRef Texture, void*& OutData, int32& OutWidth, int32& OutHeight) final override;
 	virtual void RHIUnmapStagingSurface(FTextureRHIParamRef Texture) final override;
+	virtual void RHIMapStagingSurface_RenderThread(class FRHICommandListImmediate& RHICmdList, FTextureRHIParamRef Texture, void*& OutData, int32& OutWidth, int32& OutHeight);
+	virtual void RHIUnmapStagingSurface_RenderThread(class FRHICommandListImmediate& RHICmdList, FTextureRHIParamRef Texture);
+
+
 	virtual void RHIReadSurfaceFloatData(FTextureRHIParamRef Texture, FIntRect Rect, TArray<FFloat16Color>& OutData, ECubeFace CubeFace, int32 ArrayIndex, int32 MipIndex) final override;
 	virtual void RHIRead3DSurfaceFloatData(FTextureRHIParamRef Texture, FIntRect Rect, FIntPoint ZMinMax, TArray<FFloat16Color>& OutData) final override;
 	virtual FRenderQueryRHIRef RHICreateRenderQuery(ERenderQueryType QueryType) final override;
