DR - Consolidate r.GraphicsAdapter cvar. You can also now use -graphicsadapter=N for D3D11, D3D12 and Vulkan
This deprecates the old D3D12.GraphicsAdapter cvar

diff --git a/Engine/Source/Runtime/D3D12RHI/Private/Windows/WindowsD3D12Device.cpp b/Engine/Source/Runtime/D3D12RHI/Private/Windows/WindowsD3D12Device.cpp
--- a/Engine/Source/Runtime/D3D12RHI/Private/Windows/WindowsD3D12Device.cpp
+++ b/Engine/Source/Runtime/D3D12RHI/Private/Windows/WindowsD3D12Device.cpp
@@ -26,16 +26,6 @@ extern bool D3D12RHI_ShouldCreateWithWarp();
 extern bool D3D12RHI_ShouldAllowAsyncResourceCreation();
 extern bool D3D12RHI_ShouldForceCompatibility();
 
-static TAutoConsoleVariable<int32> CVarGraphicsAdapter(
-	TEXT("D3D12.GraphicsAdapter"),
-	-1,
-	TEXT("User request to pick a specific graphics adapter (e.g. when using an integrated graphics card with a discrete one)\n")
-	TEXT(" -2: Take the first one that fulfills the criteria\n")
-	TEXT(" -1: Favor discrete because they are usually faster (default)\n")
-	TEXT("  0: Adapter #0\n")
-	TEXT("  1: Adapter #1, ..."),
-	ECVF_RenderThreadSafe);
-
 #if NV_AFTERMATH
 // Disabled by default since introduces stalls between render and driver threads
 int32 GDX12NVAfterMathEnabled = 0;
@@ -289,7 +279,10 @@ void FD3D12DynamicRHIModule::FindAdapter()
 
 	// Allow HMD to override which graphics adapter is chosen, so we pick the adapter where the HMD is connected
 	uint64 HmdGraphicsAdapterLuid = IHeadMountedDisplayModule::IsAvailable() ? IHeadMountedDisplayModule::Get().GetGraphicsAdapterLuid() : 0;
-	int32 CVarExplicitAdapterValue = HmdGraphicsAdapterLuid == 0 ? CVarGraphicsAdapter.GetValueOnGameThread() : -2;
+	// Non-static as it is used only a few times
+	auto* CVarGraphicsAdapter = IConsoleManager::Get().FindTConsoleVariableDataInt(TEXT("r.GraphicsAdapter"));
+	int32 CVarExplicitAdapterValue = HmdGraphicsAdapterLuid == 0 ? (CVarGraphicsAdapter ? CVarGraphicsAdapter->GetValueOnGameThread() : -1) : -2;
+	FParse::Value(FCommandLine::Get(), TEXT("graphicsadapter="), CVarExplicitAdapterValue);
 
 	const bool bFavorNonIntegrated = CVarExplicitAdapterValue == -1;
 
diff --git a/Engine/Source/Runtime/RHI/Private/RHI.cpp b/Engine/Source/Runtime/RHI/Private/RHI.cpp
--- a/Engine/Source/Runtime/RHI/Private/RHI.cpp
+++ b/Engine/Source/Runtime/RHI/Private/RHI.cpp
@@ -49,6 +49,18 @@ static TAutoConsoleVariable<int32> CVarDisableEngineAndAppRegistration(
 	TEXT("Changes will only take effect in new game/editor instances - can't be changed at runtime.\n"),
 	ECVF_Default);
 
+static TAutoConsoleVariable<int32> CVarGraphicsAdapter(
+	TEXT("r.GraphicsAdapter"),
+	-1,
+	TEXT("User request to pick a specific graphics adapter (e.g. when using a integrated graphics card with a discrete one)\n")
+	TEXT("For Windows D3D, unless a specific adapter is chosen we reject Microsoft adapters because we don't want the software emulation.\n")
+	TEXT("This takes precedence over -prefer{AMD|NVidia|Intel} when the value is >= 0.\n")
+	TEXT(" -2: Take the first one that fulfills the criteria\n")
+	TEXT(" -1: Favour non integrated because there are usually faster (default)\n")
+	TEXT("  0: Adapter #0\n")
+	TEXT("  1: Adapter #1, ..."),
+	ECVF_ReadOnly | ECVF_RenderThreadSafe);
+
 
 const FString FResourceTransitionUtility::ResourceTransitionAccessStrings[(int32)EResourceTransitionAccess::EMaxAccess + 1] =
 {
diff --git a/Engine/Source/Runtime/VulkanRHI/Private/VulkanRHI.cpp b/Engine/Source/Runtime/VulkanRHI/Private/VulkanRHI.cpp
--- a/Engine/Source/Runtime/VulkanRHI/Private/VulkanRHI.cpp
+++ b/Engine/Source/Runtime/VulkanRHI/Private/VulkanRHI.cpp
@@ -445,7 +445,7 @@ void FVulkanDynamicRHI::CreateInstance()
 	else if (Result != VK_SUCCESS)
 	{
 		FPlatformMisc::MessageBoxExt(EAppMsgType::Ok, TEXT(
-			"Vulkan failed to create instace (apiVersion=0x%x)\n\nDo you have a compatible Vulkan "
+			"Vulkan failed to create instance (apiVersion=0x%x)\n\nDo you have a compatible Vulkan "
 			 "driver (ICD) installed?\nPlease look at "
 			 "the Getting Started guide for additional information."), TEXT("No Vulkan driver found!"));
 		FPlatformMisc::RequestExitWithStatus(true, 1);
@@ -521,6 +521,8 @@ void FVulkanDynamicRHI::SelectAndInitDevice()
 	TArray<FDeviceInfo> DiscreteDevices;
 	TArray<FDeviceInfo> IntegratedDevices;
 
+	TArray<FDeviceInfo> OriginalOrderedDevices;
+
 #if VULKAN_ENABLE_DESKTOP_HMD_SUPPORT
 	// Allow HMD to override which graphics adapter is chosen, so we pick the adapter where the HMD is connected
 	uint64 HmdGraphicsAdapterLuid  = IHeadMountedDisplayModule::IsAvailable() ? IHeadMountedDisplayModule::Get().GetGraphicsAdapterLuid() : 0;
@@ -551,6 +553,8 @@ void FVulkanDynamicRHI::SelectAndInitDevice()
 		{
 			IntegratedDevices.Add({NewDevice, Index});
 		}
+
+		OriginalOrderedDevices.Add({NewDevice, Index});
 	}
 
 	uint32 DeviceIndex = -1;
@@ -566,36 +570,61 @@ void FVulkanDynamicRHI::SelectAndInitDevice()
 	// Add all integrated to the end of the list
 	DiscreteDevices.Append(IntegratedDevices);
 
+	// Non-static as it is used only a few times
+	auto* CVarGraphicsAdapter = IConsoleManager::Get().FindTConsoleVariableDataInt(TEXT("r.GraphicsAdapter"));
+	int32 CVarExplicitAdapterValue = CVarGraphicsAdapter ? CVarGraphicsAdapter->GetValueOnAnyThread() : -1;
+	FParse::Value(FCommandLine::Get(), TEXT("graphicsadapter="), CVarExplicitAdapterValue);
+
+	// If HMD didn't choose one...
 	if (DeviceIndex == -1)
 	{
-		if (DiscreteDevices.Num() > 0)
+		if (CVarExplicitAdapterValue >= (int32)GpuCount)
+		{
+			UE_LOG(LogVulkanRHI, Warning, TEXT("Tried to use r.GraphicsAdapter=%d, but only %d Adapter(s) found. Falling back to first device..."), CVarExplicitAdapterValue, GpuCount);
+			CVarExplicitAdapterValue = 0;
+		}
+		
+		if (CVarExplicitAdapterValue >= 0)
+		{
+			DeviceIndex = OriginalOrderedDevices[CVarExplicitAdapterValue].DeviceIndex;
+			Device = OriginalOrderedDevices[CVarExplicitAdapterValue].Device;
+		}
+		else
 		{
-			int32 PreferredVendor = PreferAdapterVendor();
-			if (DiscreteDevices.Num() > 1 && PreferredVendor != -1)
+			if (CVarExplicitAdapterValue == -2)
 			{
-				// Check for preferred
-				for (int32 Index = 0; Index < DiscreteDevices.Num(); ++Index)
+				DeviceIndex = OriginalOrderedDevices[0].DeviceIndex;
+				Device = OriginalOrderedDevices[0].Device;
+			}
+			else if (DiscreteDevices.Num() > 0 && CVarExplicitAdapterValue == -1)
+			{
+				int32 PreferredVendor = PreferAdapterVendor();
+				if (DiscreteDevices.Num() > 1 && PreferredVendor != -1)
 				{
-					if (DiscreteDevices[Index].Device->GpuProps.vendorID == PreferredVendor)
+					// Check for preferred
+					for (int32 Index = 0; Index < DiscreteDevices.Num(); ++Index)
 					{
-						DeviceIndex = DiscreteDevices[Index].DeviceIndex;
-						Device = DiscreteDevices[Index].Device;
-						break;
+						if (DiscreteDevices[Index].Device->GpuProps.vendorID == PreferredVendor)
+						{
+							DeviceIndex = DiscreteDevices[Index].DeviceIndex;
+							Device = DiscreteDevices[Index].Device;
+							break;
+						}
 					}
 				}
-			}
 
-			if (DeviceIndex == -1)
+				if (DeviceIndex == -1)
+				{
+					Device = DiscreteDevices[0].Device;
+					DeviceIndex = DiscreteDevices[0].DeviceIndex;
+				}
+			}
+			else
 			{
-				Device = DiscreteDevices[0].Device;
-				DeviceIndex = DiscreteDevices[0].DeviceIndex;
+				checkf(0, TEXT("No devices found!"));
+				DeviceIndex = 0;
 			}
 		}
-		else
-		{
-			checkf(0, TEXT("No devices found!"));
-			DeviceIndex = 0;
-		}
 	}
 
 	const VkPhysicalDeviceProperties& Props = Device->GetDeviceProperties();
diff --git a/Engine/Source/Runtime/Windows/D3D11RHI/Private/Windows/WindowsD3D11Device.cpp b/Engine/Source/Runtime/Windows/D3D11RHI/Private/Windows/WindowsD3D11Device.cpp
--- a/Engine/Source/Runtime/Windows/D3D11RHI/Private/Windows/WindowsD3D11Device.cpp
+++ b/Engine/Source/Runtime/Windows/D3D11RHI/Private/Windows/WindowsD3D11Device.cpp
@@ -35,7 +35,7 @@ FD3D11DynamicRHI*	GD3D11RHI = nullptr;
 extern bool D3D11RHI_ShouldCreateWithD3DDebug();
 extern bool D3D11RHI_ShouldAllowAsyncResourceCreation();
 
-int D3D11RHI_PreferAdaperVendor()
+static int D3D11RHI_PreferAdapterVendor()
 {
 	if (FParse::Param(FCommandLine::Get(), TEXT("preferAMD")))
 	{
@@ -55,7 +55,7 @@ int D3D11RHI_PreferAdaperVendor()
 	return -1;
 }
 
-bool D3D11RHI_AllowSoftwareFallback()
+static bool D3D11RHI_AllowSoftwareFallback()
 {
 	if (FParse::Param(FCommandLine::Get(), TEXT("AllowSoftwareRendering")))
 	{
@@ -73,17 +73,6 @@ struct AmdAgsInfo
 };
 static AmdAgsInfo AmdInfo;
 
-static TAutoConsoleVariable<int32> CVarGraphicsAdapter(
-	TEXT("r.GraphicsAdapter"),
-	-1,
-	TEXT("User request to pick a specific graphics adapter (e.g. when using a integrated graphics card with a discrete one)\n")
-	TEXT("At the moment this only works on Direct3D 11. Unless a specific adapter is chosen we reject Microsoft adapters because we don't want the software emulation.\n")
-	TEXT(" -2: Take the first one that fulfills the criteria\n")
-	TEXT(" -1: Favour non integrated because there are usually faster (default)\n")
-	TEXT("  0: Adapter #0\n")
-	TEXT("  1: Adapter #1, ..."),
-	ECVF_RenderThreadSafe);
-
 static TAutoConsoleVariable<int32> CVarForceAMDToSM4(
 	TEXT("r.ForceAMDToSM4"),
 	0,
@@ -741,7 +730,10 @@ void FD3D11DynamicRHIModule::FindAdapter()
 
 	// Allow HMD to override which graphics adapter is chosen, so we pick the adapter where the HMD is connected
 	uint64 HmdGraphicsAdapterLuid  = IHeadMountedDisplayModule::IsAvailable() ? IHeadMountedDisplayModule::Get().GetGraphicsAdapterLuid() : 0;
-	int32 CVarExplicitAdapterValue = HmdGraphicsAdapterLuid == 0 ? CVarGraphicsAdapter.GetValueOnGameThread() : -2;
+	// Non-static as it is used only a few times
+	auto* CVarGraphicsAdapter = IConsoleManager::Get().FindTConsoleVariableDataInt(TEXT("r.GraphicsAdapter"));
+	int32 CVarExplicitAdapterValue = HmdGraphicsAdapterLuid == 0 ? (CVarGraphicsAdapter ? CVarGraphicsAdapter->GetValueOnGameThread() : -1) : -2;
+	FParse::Value(FCommandLine::Get(), TEXT("graphicsadapter="), CVarExplicitAdapterValue);
 
 	const bool bFavorNonIntegrated = CVarExplicitAdapterValue == -1;
 
@@ -759,7 +751,7 @@ void FD3D11DynamicRHIModule::FindAdapter()
 
 	UE_LOG(LogD3D11RHI, Log, TEXT("D3D11 adapters:"));
 
-	int PreferredVendor = D3D11RHI_PreferAdaperVendor();
+	int PreferredVendor = D3D11RHI_PreferAdapterVendor();
 	bool bAllowSoftwareFallback = D3D11RHI_AllowSoftwareFallback();
 
 	// Enumerate the DXGIFactory's adapters.
