DR - Remove NVidia 20xx hacks on drivers 430.yy or newer for Vulkan

diff --git a/Engine/Source/Runtime/VulkanRHI/Private/VulkanRHI.cpp b/Engine/Source/Runtime/VulkanRHI/Private/VulkanRHI.cpp
--- a/Engine/Source/Runtime/VulkanRHI/Private/VulkanRHI.cpp
+++ b/Engine/Source/Runtime/VulkanRHI/Private/VulkanRHI.cpp
@@ -642,24 +642,6 @@ void FVulkanDynamicRHI::SelectAndInitDevice()
 	}
 	else if (IsRHIDeviceNVIDIA())
 	{
-		union UNvidiaDriverVersion
-		{
-			struct
-			{
-#if PLATFORM_LITTLE_ENDIAN
-				uint32 Tertiary		: 6;
-				uint32 Secondary	: 8;
-				uint32 Minor		: 8;
-				uint32 Major		: 10;
-#else
-				uint32 Major		: 10;
-				uint32 Minor		: 8;
-				uint32 Secondary	: 8;
-				uint32 Tertiary		: 6;
-#endif
-			};
-			uint32 Packed;
-		};
 		UNvidiaDriverVersion NvidiaVersion;
 		static_assert(sizeof(NvidiaVersion) == sizeof(Props.driverVersion), "Mismatched Nvidia pack driver version!");
 		NvidiaVersion.Packed = Props.driverVersion;
diff --git a/Engine/Source/Runtime/VulkanRHI/Private/VulkanRHIPrivate.h b/Engine/Source/Runtime/VulkanRHI/Private/VulkanRHIPrivate.h
--- a/Engine/Source/Runtime/VulkanRHI/Private/VulkanRHIPrivate.h
+++ b/Engine/Source/Runtime/VulkanRHI/Private/VulkanRHIPrivate.h
@@ -297,6 +297,25 @@ private:
 	FVulkanDevice&				Device;
 };
 
+union UNvidiaDriverVersion
+{
+	struct
+	{
+#if PLATFORM_LITTLE_ENDIAN
+		uint32 Tertiary		: 6;
+		uint32 Secondary	: 8;
+		uint32 Minor		: 8;
+		uint32 Major		: 10;
+#else
+		uint32 Major		: 10;
+		uint32 Minor		: 8;
+		uint32 Secondary	: 8;
+		uint32 Tertiary		: 6;
+#endif
+	};
+	uint32 Packed;
+};
+
 
 namespace VulkanRHI
 {
diff --git a/Engine/Source/Runtime/VulkanRHI/Private/VulkanTexture.cpp b/Engine/Source/Runtime/VulkanRHI/Private/VulkanTexture.cpp
--- a/Engine/Source/Runtime/VulkanRHI/Private/VulkanTexture.cpp
+++ b/Engine/Source/Runtime/VulkanRHI/Private/VulkanTexture.cpp
@@ -1491,14 +1491,18 @@ VkImageView FVulkanTextureView::StaticCreate(FVulkanDevice& Device, VkImage InIm
 	ensure(NumMips != 0xFFFFFFFF);
 	ViewInfo.subresourceRange.levelCount = NumMips;
 
-	auto CheckUseNvidiaWorkaround = []() -> bool
+	auto CheckUseNvidiaWorkaround = [&Device]() -> bool
 	{
 		if (IsRHIDeviceNVIDIA())
 		{
-			if(FParse::Param(FCommandLine::Get(), TEXT("rtx20xxmipworkaround")))
+			// Workaround for 20xx family not copying last mips correctly, so instead the view is created without the last 1x1 and 2x2 mips
+			if (GRHIAdapterName.Contains(TEXT("RTX 20")))
 			{
-				// Workaround for 20xx family not copying last mips correctly, so instead the view is created without the last 1x1 and 2x2 mips
-				if (GRHIAdapterName.Contains(TEXT("RTX 20")))
+				UNvidiaDriverVersion NvidiaVersion;
+				const VkPhysicalDeviceProperties& Props = Device.GetDeviceProperties();
+				static_assert(sizeof(NvidiaVersion) == sizeof(Props.driverVersion), "Mismatched Nvidia pack driver version!");
+				NvidiaVersion.Packed = Props.driverVersion;
+				if (NvidiaVersion.Major < 430)
 				{
 					return true;
 				}
@@ -1506,8 +1510,8 @@ VkImageView FVulkanTextureView::StaticCreate(FVulkanDevice& Device, VkImage InIm
 		}
 		return false;
 	};
-	static bool NvidiaWorkaround = CheckUseNvidiaWorkaround();
-	if(NvidiaWorkaround && Format >= VK_FORMAT_BC1_RGB_UNORM_BLOCK && Format <= VK_FORMAT_BC7_SRGB_BLOCK && NumMips > 1)
+	static bool bNvidiaWorkaround = CheckUseNvidiaWorkaround();
+	if (bNvidiaWorkaround && Format >= VK_FORMAT_BC1_RGB_UNORM_BLOCK && Format <= VK_FORMAT_BC7_SRGB_BLOCK && NumMips > 1)
 	{
 		ViewInfo.subresourceRange.levelCount = FMath::Max(1, int32(NumMips) - 2);
 	}
diff --git a/Engine/Source/Runtime/VulkanRHI/Private/Windows/VulkanWindowsPlatform.cpp b/Engine/Source/Runtime/VulkanRHI/Private/Windows/VulkanWindowsPlatform.cpp
--- a/Engine/Source/Runtime/VulkanRHI/Private/Windows/VulkanWindowsPlatform.cpp
+++ b/Engine/Source/Runtime/VulkanRHI/Private/Windows/VulkanWindowsPlatform.cpp
@@ -281,14 +281,22 @@ void FVulkanWindowsPlatform::CheckDeviceDriver(uint32 DeviceIndex, const VkPhysi
 	}
 	else if (IsRHIDeviceNVIDIA())
 	{
-		// Workaround a crash on 20xx family
-		UE_LOG(LogVulkanRHI, Warning, TEXT("Nvidia 20xx family of GPUs have a known crash. Compatibility mode (slow!) will now be enabled"));
 		if (GRHIAdapterName.Contains(TEXT("RTX 20")))
 		{
-			extern TAutoConsoleVariable<int32> GRHIThreadCvar;
-			GRHIThreadCvar->SetWithCurrentPriority(0);
-			IConsoleVariable* BypassVar = IConsoleManager::Get().FindConsoleVariable(TEXT("r.RHICmdBypass"));
-			BypassVar->SetWithCurrentPriority(1);
+			UNvidiaDriverVersion NvidiaVersion;
+			static_assert(sizeof(NvidiaVersion) == sizeof(Props.driverVersion), "Mismatched Nvidia pack driver version!");
+			NvidiaVersion.Packed = Props.driverVersion;
+
+			if (NvidiaVersion.Major < 430)
+			{
+				// Workaround a crash on 20xx family
+				UE_LOG(LogVulkanRHI, Warning, TEXT("Nvidia 20xx family of GPUs have a known crash on drivers < 430. Compatibility mode (slow!) will now be enabled"));
+
+				extern TAutoConsoleVariable<int32> GRHIThreadCvar;
+				GRHIThreadCvar->SetWithCurrentPriority(0);
+				IConsoleVariable* BypassVar = IConsoleManager::Get().FindConsoleVariable(TEXT("r.RHICmdBypass"));
+				BypassVar->SetWithCurrentPriority(1);
+			}
 		}
 	}
 }
