Dont resolve timing queries untill command buffer is submitted.
Fixes a validation issue with calling functions that are not allowed inside render passes.

diff --git a/Engine/Source/Runtime/VulkanRHI/Private/VulkanCommandBuffer.cpp b/Engine/Source/Runtime/VulkanRHI/Private/VulkanCommandBuffer.cpp
--- a/Engine/Source/Runtime/VulkanRHI/Private/VulkanCommandBuffer.cpp
+++ b/Engine/Source/Runtime/VulkanRHI/Private/VulkanCommandBuffer.cpp
@@ -148,10 +148,34 @@ void FVulkanCmdBuffer::BeginRenderPass(const FVulkanRenderTargetLayout& Layout,
 	}
 }
 
+
+void FVulkanCmdBuffer::AddPendingTimestampQuery(uint64 Index, uint64 Count, VkQueryPool PoolHandle, VkBuffer BufferHandle)
+{
+	PendingQuery PQ;
+	PQ.Index = Index;
+	PQ.Count = Count;
+	PQ.PoolHandle = PoolHandle;
+	PQ.BufferHandle = BufferHandle;
+	PendingQueries.Add(PQ);
+}
+
+
 void FVulkanCmdBuffer::End()
 {
 	checkf(IsOutsideRenderPass(), TEXT("Can't End as we're inside a render pass! CmdBuffer 0x%p State=%d"), CommandBufferHandle, (int32)State);
 
+
+	for (PendingQuery& Query : PendingQueries)
+	{
+		uint64 Index = Query.Index;
+		VkBuffer BufferHandle = Query.BufferHandle;
+		VkQueryPool PoolHandle = Query.PoolHandle;
+		VulkanRHI::vkCmdCopyQueryPoolResults(GetHandle(), PoolHandle, Index, Query.Count, BufferHandle, sizeof(uint64) * Index, sizeof(uint64), VK_QUERY_RESULT_64_BIT);
+		VulkanRHI::vkCmdResetQueryPool(GetHandle(), PoolHandle, Index, Query.Count);
+	}
+
+	PendingQueries.Reset();
+
 	if (GVulkanProfileCmdBuffers)
 	{
 		if (Timing)
diff --git a/Engine/Source/Runtime/VulkanRHI/Private/VulkanCommandBuffer.h b/Engine/Source/Runtime/VulkanRHI/Private/VulkanCommandBuffer.h
--- a/Engine/Source/Runtime/VulkanRHI/Private/VulkanCommandBuffer.h
+++ b/Engine/Source/Runtime/VulkanRHI/Private/VulkanCommandBuffer.h
@@ -174,6 +174,16 @@ public:
 
 	bool AcquirePoolSetAndDescriptorsIfNeeded(const class FVulkanDescriptorSetsLayout& Layout, bool bNeedDescriptors, VkDescriptorSet* OutDescriptors);
 
+
+	struct PendingQuery
+	{
+		uint64 Index;
+		uint64 Count;
+		VkBuffer BufferHandle;
+		VkQueryPool PoolHandle;
+	};
+	void AddPendingTimestampQuery(uint64 Index, uint64 Count, VkQueryPool PoolHandle, VkBuffer BufferHandle);
+
 private:
 	FVulkanDevice* Device;
 	VkCommandBuffer CommandBufferHandle;
@@ -182,6 +192,7 @@ private:
 	TArray<VkPipelineStageFlags> WaitFlags;
 	TArray<VulkanRHI::FSemaphore*> WaitSemaphores;
 	TArray<VulkanRHI::FSemaphore*> SubmittedWaitSemaphores;
+	TArray<PendingQuery> PendingQueries;
 
 	void MarkSemaphoresAsSubmitted()
 	{
diff --git a/Engine/Source/Runtime/VulkanRHI/Private/VulkanQuery.cpp b/Engine/Source/Runtime/VulkanRHI/Private/VulkanQuery.cpp
--- a/Engine/Source/Runtime/VulkanRHI/Private/VulkanQuery.cpp
+++ b/Engine/Source/Runtime/VulkanRHI/Private/VulkanQuery.cpp
@@ -588,9 +588,7 @@ void FVulkanCommandListContext::RHIEndRenderQuery(FRenderQueryRHIParamRef QueryR
 		const uint32 QueryEndIndex = Query->Pool.CurrentTimestamp;
 		FVulkanCmdBuffer* CmdBuffer = CommandBufferManager->GetActiveCmdBuffer();
 		VulkanRHI::vkCmdWriteTimestamp(CmdBuffer->GetHandle(), VK_PIPELINE_STAGE_TOP_OF_PIPE_BIT, Query->Pool.GetHandle(), QueryEndIndex);
-		//check(CmdBuffer->IsOutsideRenderPass());
-		VulkanRHI::vkCmdCopyQueryPoolResults(CmdBuffer->GetHandle(), Query->Pool.GetHandle(), QueryEndIndex, 1, Query->Pool.ResultsBuffer->GetHandle(), sizeof(uint64) * QueryEndIndex, sizeof(uint64), VK_QUERY_RESULT_64_BIT);
-		VulkanRHI::vkCmdResetQueryPool(CmdBuffer->GetHandle(), Query->Pool.GetHandle(), QueryEndIndex, 1);
+		CmdBuffer->AddPendingTimestampQuery(QueryEndIndex, 1, Query->Pool.GetHandle(), Query->Pool.ResultsBuffer->GetHandle());
 		Query->Pool.TimestampListHandles[QueryEndIndex].CmdBuffer = CmdBuffer;
 		Query->Pool.TimestampListHandles[QueryEndIndex].FenceCounter = CmdBuffer->GetFenceSignaledCounter();
 		Query->Pool.NumIssuedTimestamps = FMath::Min<uint32>(Query->Pool.NumIssuedTimestamps + 1, Query->Pool.BufferSize);
diff --git a/Engine/Source/Runtime/VulkanRHI/Private/VulkanUtil.cpp b/Engine/Source/Runtime/VulkanRHI/Private/VulkanUtil.cpp
--- a/Engine/Source/Runtime/VulkanRHI/Private/VulkanUtil.cpp
+++ b/Engine/Source/Runtime/VulkanRHI/Private/VulkanUtil.cpp
@@ -206,9 +206,7 @@ void FVulkanGPUTiming::EndTiming(FVulkanCmdBuffer* CmdBuffer)
 		const uint32 QueryEndIndex = Pool->CurrentTimestamp * 2 + 1;
 		check(QueryEndIndex == QueryStartIndex + 1);	// Make sure they're adjacent indices.
 		VulkanRHI::vkCmdWriteTimestamp(CmdBuffer->GetHandle(), VK_PIPELINE_STAGE_TOP_OF_PIPE_BIT, Pool->GetHandle(), QueryEndIndex);
-		//check(CmdBuffer->IsOutsideRenderPass());
-		VulkanRHI::vkCmdCopyQueryPoolResults(CmdBuffer->GetHandle(), Pool->GetHandle(), QueryStartIndex, 2, Pool->ResultsBuffer->GetHandle(), sizeof(uint64) * QueryStartIndex, sizeof(uint64), VK_QUERY_RESULT_64_BIT);
-		VulkanRHI::vkCmdResetQueryPool(CmdBuffer->GetHandle(), Pool->GetHandle(), QueryStartIndex, 2);
+		CmdBuffer->AddPendingTimestampQuery(QueryStartIndex, 2, Pool->GetHandle(), Pool->ResultsBuffer->GetHandle());
 		Pool->TimestampListHandles[QueryEndIndex].CmdBuffer = CmdBuffer;
 		Pool->TimestampListHandles[QueryEndIndex].FenceCounter = CmdBuffer->GetFenceSignaledCounter();
 		Pool->NumIssuedTimestamps = FMath::Min<uint32>(Pool->NumIssuedTimestamps + 1, Pool->BufferSize);
