DR - Clang compile fix
#rnx

diff --git a/Engine/Source/Runtime/VulkanRHI/Private/VulkanPipeline.cpp b/Engine/Source/Runtime/VulkanRHI/Private/VulkanPipeline.cpp
--- a/Engine/Source/Runtime/VulkanRHI/Private/VulkanPipeline.cpp
+++ b/Engine/Source/Runtime/VulkanRHI/Private/VulkanPipeline.cpp
@@ -1436,7 +1436,7 @@ FVulkanGfxLayout* FVulkanPipelineStateCacheManager::GetOrGenerateGfxLayout(const
 #endif
 	// Second pass
 	const int32 NumImmutableSamplers = PSOInitializer.ImmutableSamplerState.ImmutableSamplers.Num();
-	TArrayView<FRHISamplerState*> ImmutableSamplers(NumImmutableSamplers > 0 ? &(FRHISamplerState*)PSOInitializer.ImmutableSamplerState.ImmutableSamplers[0] : nullptr, NumImmutableSamplers);
+	TArrayView<FRHISamplerState*> ImmutableSamplers(NumImmutableSamplers > 0 ? &(FRHISamplerState*&)PSOInitializer.ImmutableSamplerState.ImmutableSamplers[0] : nullptr, NumImmutableSamplers);
 	DescriptorSetLayoutInfo.FinalizeBindings<false>(UBGatherInfo, ImmutableSamplers);
 
 	FVulkanLayout* Layout = FindOrAddLayout(DescriptorSetLayoutInfo, true);
