DR - Fix uninitalized variables
#rnx

diff --git a/Engine/Source/Runtime/VulkanRHI/Private/VulkanResources.h b/Engine/Source/Runtime/VulkanRHI/Private/VulkanResources.h
--- a/Engine/Source/Runtime/VulkanRHI/Private/VulkanResources.h
+++ b/Engine/Source/Runtime/VulkanRHI/Private/VulkanResources.h
@@ -1326,25 +1326,17 @@ public:
 		, NumArraySlices(InCreateInfo.NumArraySlices)
 		, Size(0)
 		, SourceBuffer(nullptr)
-		, VolatileBufferHandle(VK_NULL_HANDLE)
-		, VolatileLockCounter(MAX_uint32)
 	{
 	}
 
 	FVulkanShaderResourceView(FVulkanDevice* Device, FVulkanStructuredBuffer* InStructuredBuffer)
 		: VulkanRHI::FDeviceChild(Device)
 		, BufferViewFormat(PF_Unknown)
-		, SRGBOverride(SRGBO_Default)
 		, SourceTexture(nullptr)
 		, SourceStructuredBuffer(InStructuredBuffer)
-		, MipLevel(0)
 		, NumMips(0)
-		, FirstArraySlice(0)
-		, NumArraySlices(0)
 		, Size(InStructuredBuffer->GetSize())
 		, SourceBuffer(nullptr)
-		, VolatileBufferHandle(VK_NULL_HANDLE)
-		, VolatileLockCounter(MAX_uint32)
 	{
 	}
 
@@ -1360,16 +1352,16 @@ public:
 	}
 
 	EPixelFormat BufferViewFormat;
-	ERHITextureSRVOverrideSRGBType SRGBOverride;
+	ERHITextureSRVOverrideSRGBType SRGBOverride = SRGBO_Default;
 
 	// The texture that this SRV come from
 	TRefCountPtr<FRHITexture> SourceTexture;
 	FVulkanTextureView TextureView;
 	FVulkanStructuredBuffer* SourceStructuredBuffer;
-	uint32 MipLevel;
-	uint32 NumMips;
-	uint32 FirstArraySlice;
-	uint32 NumArraySlices;
+	uint32 MipLevel = 0;
+	uint32 NumMips = MAX_uint32;
+	uint32 FirstArraySlice = 0;
+	uint32 NumArraySlices = 0;
 
 	~FVulkanShaderResourceView();
 
@@ -1383,8 +1375,8 @@ public:
 
 protected:
 	// Used to check on volatile buffers if a new BufferView is required
-	VkBuffer VolatileBufferHandle;
-	uint32 VolatileLockCounter;
+	VkBuffer VolatileBufferHandle = VK_NULL_HANDLE;
+	uint32 VolatileLockCounter = MAX_uint32;
 };
 
 class FVulkanComputeFence : public FRHIComputeFence, public VulkanRHI::FGPUEvent
diff --git a/Engine/Source/Runtime/VulkanRHI/Private/VulkanUAV.cpp b/Engine/Source/Runtime/VulkanRHI/Private/VulkanUAV.cpp
--- a/Engine/Source/Runtime/VulkanRHI/Private/VulkanUAV.cpp
+++ b/Engine/Source/Runtime/VulkanRHI/Private/VulkanUAV.cpp
@@ -6,16 +6,11 @@
 FVulkanShaderResourceView::FVulkanShaderResourceView(FVulkanDevice* Device, FRHIResource* InRHIBuffer, FVulkanResourceMultiBuffer* InSourceBuffer, uint32 InSize, EPixelFormat InFormat)
 	: VulkanRHI::FDeviceChild(Device)
 	, BufferViewFormat(InFormat)
-	, SRGBOverride(SRGBO_Default)
 	, SourceTexture(nullptr)
 	, SourceStructuredBuffer(nullptr)
-	, MipLevel(0)
-	, NumMips(-1)
 	, Size(InSize)
 	, SourceBuffer(InSourceBuffer)
 	, SourceRHIBuffer(InRHIBuffer)
-	, VolatileBufferHandle(VK_NULL_HANDLE)
-	, VolatileLockCounter(MAX_uint32)
 {
 	check(Device);
 	if(SourceBuffer)
