Refactored and extended custom Vulkan validation layer.

#rnx

diff --git a/Engine/Source/Runtime/VulkanRHI/Private/VulkanCommandWrappers.h b/Engine/Source/Runtime/VulkanRHI/Private/VulkanCommandWrappers.h
--- a/Engine/Source/Runtime/VulkanRHI/Private/VulkanCommandWrappers.h
+++ b/Engine/Source/Runtime/VulkanRHI/Private/VulkanCommandWrappers.h
@@ -19,146 +19,156 @@ struct FWrapLayer
 	// Pass in VK_RESULT_MAX_ENUM for Prolog calls; for Epilog use a different value or the actual Result if available
 
 	static void CreateInstance(VkResult Result, const VkInstanceCreateInfo* CreateInfo, VkInstance* Instance) VULKAN_LAYER_BODY
-	static void EnumeratePhysicalDevices(VkResult Result, VkInstance Instance, uint32* PhysicalDeviceCount, VkPhysicalDevice* PhysicalDevices) VULKAN_LAYER_BODY
 	static void DestroyInstance(VkResult Result, VkInstance Instance) VULKAN_LAYER_BODY
+	static void EnumeratePhysicalDevices(VkResult Result, VkInstance Instance, uint32* PhysicalDeviceCount, VkPhysicalDevice* PhysicalDevices) VULKAN_LAYER_BODY
+	static void GetPhysicalDeviceFeatures(VkResult Result, VkPhysicalDevice PhysicalDevice, VkPhysicalDeviceFeatures* Features) VULKAN_LAYER_BODY
+	static void GetPhysicalDeviceFormatProperties(VkResult Result, VkPhysicalDevice PhysicalDevice, VkFormat Format, VkFormatProperties* FormatProperties) VULKAN_LAYER_BODY
+	static void GetPhysicalDeviceImageFormatProperties(VkResult Result, VkPhysicalDevice PhysicalDevice, VkFormat Format, VkImageType Type, VkImageTiling Tiling, VkImageUsageFlags Usage, VkImageCreateFlags Flags, VkImageFormatProperties* pImageFormatProperties) VULKAN_LAYER_BODY
+	static void GetPhysicalDeviceProperties(VkResult Result, VkPhysicalDevice PhysicalDevice, VkPhysicalDeviceProperties* FormatProperties) VULKAN_LAYER_BODY
+	static void GetPhysicalDeviceQueueFamilyProperties(VkResult Result, VkPhysicalDevice PhysicalDevice, uint32* QueueFamilyPropertyCount, VkQueueFamilyProperties* QueueFamilyProperties) VULKAN_LAYER_BODY
+	static void GetPhysicalDeviceMemoryProperties(VkResult Result, VkPhysicalDevice PhysicalDevice, VkPhysicalDeviceMemoryProperties* Properties) VULKAN_LAYER_BODY
 	static void GetInstanceProcAddr(VkResult Result, VkInstance Instance, const char* Name, PFN_vkVoidFunction VoidFunction) VULKAN_LAYER_BODY
+	static void GetDeviceProcAddr(VkResult Result, VkDevice Device, const char* Name, PFN_vkVoidFunction VoidFunction) VULKAN_LAYER_BODY
+	static void CreateDevice(VkResult Result, VkPhysicalDevice PhysicalDevice, const VkDeviceCreateInfo* CreateInfo, VkDevice* Device) VULKAN_LAYER_BODY
+	static void DestroyDevice(VkResult Result, VkDevice Device) VULKAN_LAYER_BODY
 	VULKAN_EXTERN_EXPORT static void EnumerateInstanceExtensionProperties(VkResult Result, const char* LayerName, uint32* PropertyCount, VkExtensionProperties* Properties) VULKAN_LAYER_BODY
+	VULKAN_EXTERN_EXPORT static void EnumerateDeviceExtensionProperties(VkResult Result, VkPhysicalDevice PhysicalDevice, const char* LayerName, uint32* PropertyCount, VkExtensionProperties* Properties) VULKAN_LAYER_BODY
 	static void EnumerateInstanceLayerProperties(VkResult Result, uint32* PropertyCount, VkLayerProperties* Properties) VULKAN_LAYER_BODY
+	static void EnumerateDeviceLayerProperties(VkResult Result, VkPhysicalDevice PhysicalDevice, uint32* PropertyCount, VkLayerProperties* Properties) VULKAN_LAYER_BODY
+	static void GetDeviceQueue(VkResult Result, VkDevice Device, uint32 QueueFamilyIndex, uint32 QueueIndex, VkQueue* Queue) VULKAN_LAYER_BODY
+	static void QueueSubmit(VkResult Result, VkQueue Queue, uint32 SubmitCount, const VkSubmitInfo* Submits, VkFence Fence) VULKAN_LAYER_BODY
+	static void QueueWaitIdle(VkResult Result, VkQueue Queue) VULKAN_LAYER_BODY
+	static void DeviceWaitIdle(VkResult Result, VkDevice Device) VULKAN_LAYER_BODY
 	static void AllocateMemory(VkResult Result, VkDevice Device, const VkMemoryAllocateInfo* AllocateInfo, VkDeviceMemory* Memory) VULKAN_LAYER_BODY
+	static void FreeMemory(VkResult Result, VkDevice Device, VkDeviceMemory Memory) VULKAN_LAYER_BODY
+	static void MapMemory(VkResult Result, VkDevice Device, VkDeviceMemory Memory, VkDeviceSize Offset, VkDeviceSize Size, VkMemoryMapFlags Flags, void** Data) VULKAN_LAYER_BODY
+	static void UnmapMemory(VkResult Result, VkDevice Device, VkDeviceMemory Memory) VULKAN_LAYER_BODY
 	static void FlushMappedMemoryRanges(VkResult Result, VkDevice Device, uint32 MemoryRangeCount, const VkMappedMemoryRange* pMemoryRanges) VULKAN_LAYER_BODY
 	static void InvalidateMappedMemoryRanges(VkResult Result, VkDevice Device, uint32 MemoryRangeCount, const VkMappedMemoryRange* pMemoryRanges) VULKAN_LAYER_BODY
+	static void GetDeviceMemoryCommitment(VkResult Result, VkDevice Device, VkDeviceMemory Memory, VkDeviceSize* pCommittedMemoryInBytes);
+	static void BindBufferMemory(VkResult Result, VkDevice Device, VkBuffer Buffer, VkDeviceMemory Memory, VkDeviceSize MemoryOffset) VULKAN_LAYER_BODY
+	VULKAN_EXTERN_EXPORT static void BindImageMemory(VkResult Result, VkDevice Device, VkImage Image, VkDeviceMemory Memory, VkDeviceSize MemoryOffset) VULKAN_LAYER_BODY
 	static void GetBufferMemoryRequirements(VkResult Result, VkDevice Device, VkBuffer Buffer, VkMemoryRequirements* MemoryRequirements) VULKAN_LAYER_BODY
 	static void GetImageMemoryRequirements(VkResult Result, VkDevice Device, VkImage Image, VkMemoryRequirements* MemoryRequirements) VULKAN_LAYER_BODY
+	static void GetImageSparseMemoryRequirements(VkResult Result, VkDevice Device, VkImage Image, uint32_t* pSparseMemoryRequirementCount, VkSparseImageMemoryRequirements* pSparseMemoryRequirements) VULKAN_LAYER_BODY
+	static void GetPhysicalDeviceSparseImageFormatProperties(VkResult Result, VkPhysicalDevice PhysicalDevice, VkFormat Format, VkImageType Type, VkSampleCountFlagBits Samples, VkImageUsageFlags Usage, VkImageTiling Tiling, uint32_t* pPropertyCount, VkSparseImageFormatProperties* pProperties) VULKAN_LAYER_BODY
+	static void QueueBindSparse(VkResult Result, VkQueue Queue, uint32_t BindInfoCount, const VkBindSparseInfo* pBindInfo, VkFence Fence) VULKAN_LAYER_BODY
 	static void CreateFence(VkResult Result, VkDevice Device, const VkFenceCreateInfo* CreateInfo, VkFence* Fence) VULKAN_LAYER_BODY
+	static void DestroyFence(VkResult Result, VkDevice Device, VkFence Fence) VULKAN_LAYER_BODY
 	static void ResetFences(VkResult Result, VkDevice Device, uint32 FenceCount, const VkFence* Fences) VULKAN_LAYER_BODY
+	static void GetFenceStatus(VkResult Result, VkDevice Device, VkFence Fence) VULKAN_LAYER_BODY
 	static void WaitForFences(VkResult Result, VkDevice Device, uint32 FenceCount, const VkFence* Fences, VkBool32 bWaitAll, uint64_t Timeout) VULKAN_LAYER_BODY
 	static void CreateSemaphore(VkResult Result, VkDevice Device, const VkSemaphoreCreateInfo* CreateInfo, VkSemaphore* Semaphore) VULKAN_LAYER_BODY
+	static void DestroySemaphore(VkResult Result, VkDevice Device, VkSemaphore Semaphore) VULKAN_LAYER_BODY
+	static void CreateEvent(VkResult Result, VkDevice Device, const VkEventCreateInfo* CreateInfo, VkEvent* Event) VULKAN_LAYER_BODY
+	static void DestroyEvent(VkResult Result, VkDevice Device, VkEvent Event) VULKAN_LAYER_BODY
+	static void GetEventStatus(VkResult Result, VkDevice Device, VkEvent Event) VULKAN_LAYER_BODY
+	static void SetEvent(VkResult Result, VkDevice Device, VkEvent Event) VULKAN_LAYER_BODY
+	static void ResetEvent(VkResult Result, VkDevice Device, VkEvent Event) VULKAN_LAYER_BODY
 	static void CreateQueryPool(VkResult Result, VkDevice Device, const VkQueryPoolCreateInfo* CreateInfo, VkQueryPool* QueryPool) VULKAN_LAYER_BODY
+	static void DestroyQueryPool(VkResult Result, VkDevice Device, VkQueryPool QueryPool) VULKAN_LAYER_BODY
+	static void GetQueryPoolResults(VkResult Result, VkDevice Device, VkQueryPool QueryPool, uint32 FirstQuery, uint32 QueryCount, size_t DataSize, void* Data, VkDeviceSize Stride, VkQueryResultFlags Flags) VULKAN_LAYER_BODY
 	static void CreateBuffer(VkResult Result, VkDevice Device, const VkBufferCreateInfo* CreateInfo, VkBuffer* Buffer) VULKAN_LAYER_BODY
+	static void DestroyBuffer(VkResult Result, VkDevice Device, VkBuffer Buffer) VULKAN_LAYER_BODY
 	static void CreateBufferView(VkResult Result, VkDevice Device, const VkBufferViewCreateInfo* CreateInfo, VkBufferView* BufferView) VULKAN_LAYER_BODY
+	static void DestroyBufferView(VkResult Result, VkDevice Device, VkBufferView BufferView) VULKAN_LAYER_BODY
 	VULKAN_EXTERN_EXPORT static void CreateImage(VkResult Result, VkDevice Device, const VkImageCreateInfo* CreateInfo, VkImage* Image) VULKAN_LAYER_BODY
 	static void DestroyImage(VkResult Result, VkDevice Device, VkImage Image) VULKAN_LAYER_BODY
 	static void GetImageSubresourceLayout(VkResult Result, VkDevice Device, VkImage Image, const VkImageSubresource* Subresource, VkSubresourceLayout* Layout) VULKAN_LAYER_BODY
 	static void CreateImageView(VkResult Result, VkDevice Device, const VkImageViewCreateInfo* CreateInfo, VkImageView* ImageView) VULKAN_LAYER_BODY
+	static void DestroyImageView(VkResult Result, VkDevice Device, VkImageView ImageView) VULKAN_LAYER_BODY
 	static void CreateShaderModule(VkResult Result, VkDevice Device, const VkShaderModuleCreateInfo* CreateInfo, VkShaderModule* ShaderModule) VULKAN_LAYER_BODY
+	static void DestroyShaderModule(VkResult Result, VkDevice Device, VkShaderModule ShaderModule) VULKAN_LAYER_BODY
 	static void CreatePipelineCache(VkResult Result, VkDevice Device, const VkPipelineCacheCreateInfo* CreateInfo, VkPipelineCache* PipelineCache) VULKAN_LAYER_BODY
+	static void DestroyPipelineCache(VkResult Result, VkDevice Device, VkPipelineCache PipelineCache) VULKAN_LAYER_BODY
+	static void GetPipelineCacheData(VkResult Result, VkDevice Device, VkPipelineCache PipelineCache, size_t* DataSize, void* Data) VULKAN_LAYER_BODY
+	static void MergePipelineCaches(VkResult Result, VkDevice Device, VkPipelineCache DestCache, uint32 SourceCacheCount, const VkPipelineCache* SrcCaches) VULKAN_LAYER_BODY
 	static void CreateGraphicsPipelines(VkResult Result, VkDevice Device, VkPipelineCache PipelineCache, uint32 CreateInfoCount, const VkGraphicsPipelineCreateInfo* CreateInfos, VkPipeline* Pipelines) VULKAN_LAYER_BODY
+	static void CreateComputePipelines(VkResult Result, VkDevice Device, VkPipelineCache PipelineCache, uint32 CreateInfoCount, const VkComputePipelineCreateInfo* CreateInfos, VkPipeline* Pipelines) VULKAN_LAYER_BODY
+	static void DestroyPipeline(VkResult Result, VkDevice Device, VkPipeline Pipeline) VULKAN_LAYER_BODY
 	static void CreatePipelineLayout(VkResult Result, VkDevice Device, const VkPipelineLayoutCreateInfo* CreateInfo, VkPipelineLayout* PipelineLayout) VULKAN_LAYER_BODY
+	static void DestroyPipelineLayout(VkResult Result, VkDevice Device, VkPipelineLayout PipelineLayout) VULKAN_LAYER_BODY
 	static void CreateSampler(VkResult Result, VkDevice Device, const VkSamplerCreateInfo* CreateInfo, VkSampler* Sampler) VULKAN_LAYER_BODY
+	static void DestroySampler(VkResult Result, VkDevice Device, VkSampler Sampler) VULKAN_LAYER_BODY
 	static void CreateDescriptorSetLayout(VkResult Result, VkDevice Device, const VkDescriptorSetLayoutCreateInfo* CreateInfo, VkDescriptorSetLayout* SetLayout) VULKAN_LAYER_BODY
+	static void DestroyDescriptorSetLayout(VkResult Result, VkDevice Device, VkDescriptorSetLayout DescriptorSetLayout) VULKAN_LAYER_BODY
 	static void CreateDescriptorPool(VkResult Result, VkDevice Device, const VkDescriptorPoolCreateInfo* CreateInfo, VkDescriptorPool* DescriptorPool) VULKAN_LAYER_BODY
+	static void DestroyDescriptorPool(VkResult Result, VkDevice Device, VkDescriptorPool DescriptorPool) VULKAN_LAYER_BODY
+	static void ResetDescriptorPool(VkResult Result, VkDevice Device, VkDescriptorPool DescriptorPool, VkDescriptorPoolResetFlags Flags) VULKAN_LAYER_BODY
 	static void AllocateDescriptorSets(VkResult Result, VkDevice Device, const VkDescriptorSetAllocateInfo* AllocateInfo, VkDescriptorSet* DescriptorSets) VULKAN_LAYER_BODY
 	static void FreeDescriptorSets(VkResult Result, VkDevice Device, VkDescriptorPool DescriptorPool, uint32 DescriptorSetCount, const VkDescriptorSet* DescriptorSets) VULKAN_LAYER_BODY
 	static void UpdateDescriptorSets(VkResult Result, VkDevice Device, uint32 DescriptorWriteCount, const VkWriteDescriptorSet* DescriptorWrites, uint32 DescriptorCopyCount, const VkCopyDescriptorSet* DescriptorCopies) VULKAN_LAYER_BODY
 	static void CreateFramebuffer(VkResult Result, VkDevice Device, const VkFramebufferCreateInfo* CreateInfo, VkFramebuffer* Framebuffer) VULKAN_LAYER_BODY
+	static void DestroyFramebuffer(VkResult Result, VkDevice Device, VkFramebuffer Framebuffer) VULKAN_LAYER_BODY
 	static void CreateRenderPass(VkResult Result, VkDevice Device, const VkRenderPassCreateInfo* CreateInfo, VkRenderPass* RenderPass) VULKAN_LAYER_BODY
+	static void DestroyRenderPass(VkResult Result, VkDevice Device, VkRenderPass RenderPass) VULKAN_LAYER_BODY
+	static void GetRenderAreaGranularity(VkResult Result, VkDevice Device, VkRenderPass RenderPass, VkExtent2D* pGranularity) VULKAN_LAYER_BODY
 	static void CreateCommandPool(VkResult Result, VkDevice Device, const VkCommandPoolCreateInfo* CreateInfo, VkCommandPool* CommandPool) VULKAN_LAYER_BODY
+	static void DestroyCommandPool(VkResult Result, VkDevice Device, VkCommandPool CommandPool) VULKAN_LAYER_BODY
+	static void ResetCommandPool(VkResult Result, VkDevice Device, VkCommandPool CommandPool, VkCommandPoolResetFlags Flags) VULKAN_LAYER_BODY
+	static void AllocateCommandBuffers(VkResult Result, VkDevice Device, const VkCommandBufferAllocateInfo* AllocateInfo, VkCommandBuffer* CommandBuffers) VULKAN_LAYER_BODY
+	static void FreeCommandBuffers(VkResult Result, VkDevice Device, VkCommandPool CommandPool, uint32 CommandBufferCount, const VkCommandBuffer* CommandBuffers) VULKAN_LAYER_BODY
 	static void BeginCommandBuffer(VkResult Result, VkCommandBuffer CommandBuffer, const VkCommandBufferBeginInfo* BeginInfo) VULKAN_LAYER_BODY
 	static void EndCommandBuffer(VkResult Result, VkCommandBuffer CommandBuffer) VULKAN_LAYER_BODY
-	static void BindDescriptorSets(VkResult Result, VkCommandBuffer CommandBuffer, VkPipelineBindPoint PipelineBindPoint, VkPipelineLayout Layout, uint32 FirstSet, uint32 DescriptorSetCount, const VkDescriptorSet* DescriptorSets, uint32 DynamicOffsetCount, const uint32* DynamicOffsets) VULKAN_LAYER_BODY
-	static void BindVertexBuffers(VkResult Result, VkCommandBuffer CommandBuffer, uint32 FirstBinding, uint32 BindingCount, const VkBuffer* Buffers, const VkDeviceSize* pOffsets) VULKAN_LAYER_BODY
-	static void BindIndexBuffer(VkResult Result, VkCommandBuffer CommandBuffer, VkBuffer IndexBuffer, VkDeviceSize Offset, VkIndexType IndexType) VULKAN_LAYER_BODY
-	static void CopyBuffer(VkResult Result, VkCommandBuffer CommandBuffer, VkBuffer SrcBuffer, VkBuffer DstBuffer, uint32 RegionCount, const VkBufferCopy* Regions) VULKAN_LAYER_BODY
-	static void CopyBufferToImage(VkResult Result, VkCommandBuffer CommandBuffer, VkBuffer SrcBuffer, VkImage DstImage, VkImageLayout DstImageLayout, uint32 RegionCount, const VkBufferImageCopy* Regions) VULKAN_LAYER_BODY
-	VULKAN_EXTERN_EXPORT static void BlitImage(VkResult Result, VkCommandBuffer CommandBuffer, VkImage SrcImage, VkImageLayout SrcImageLayout, VkImage DstImage, VkImageLayout DstImageLayout, uint32 RegionCount, const VkImageBlit* Regions, VkFilter Filter) VULKAN_LAYER_BODY
-	static void CopyImageToBuffer(VkResult Result, VkCommandBuffer CommandBuffer, VkImage SrcImage, VkImageLayout SrcImageLayout, VkBuffer DstBuffer, uint32 RegionCount, const VkBufferImageCopy* Regions) VULKAN_LAYER_BODY
-	VULKAN_EXTERN_EXPORT static void ClearColorImage(VkResult Result, VkCommandBuffer CommandBuffer, VkImage Image, VkImageLayout ImageLayout, const VkClearColorValue* ColorValue, uint32 RangeCount, const VkImageSubresourceRange* Ranges) VULKAN_LAYER_BODY
-	static void ClearDepthStencilImage(VkResult Result, VkCommandBuffer CommandBuffer, VkImage Image, VkImageLayout ImageLayout, const VkClearDepthStencilValue* DepthStencil, uint32 RangeCount, const VkImageSubresourceRange* Ranges) VULKAN_LAYER_BODY
-	static void ClearAttachments(VkResult Result, VkCommandBuffer CommandBuffer, uint32 AttachmentCount, const VkClearAttachment* Attachments, uint32 RectCount, const VkClearRect* Rects) VULKAN_LAYER_BODY
-	static void ResolveImage(VkResult Result, VkCommandBuffer CommandBuffer, VkImage SrcImage, VkImageLayout SrcImageLayout, VkImage DstImage, VkImageLayout DstImageLayout, uint32 RegionCount, const VkImageResolve* Regions) VULKAN_LAYER_BODY
-	static void WaitEvents(VkResult Result, VkCommandBuffer CommandBuffer, uint32 EventCount, const VkEvent* Events, VkPipelineStageFlags SrcStageMask, VkPipelineStageFlags DstStageMask, uint32 MemoryBarrierCount,
-			const VkMemoryBarrier* pMemoryBarriers, uint32 BufferMemoryBarrierCount, const VkBufferMemoryBarrier* pBufferMemoryBarriers, uint32 ImageMemoryBarrierCount, const VkImageMemoryBarrier* pImageMemoryBarriers) VULKAN_LAYER_BODY
-	static void PipelineBarrier(VkResult Result, VkCommandBuffer CommandBuffer, VkPipelineStageFlags SrcStageMask, VkPipelineStageFlags DstStageMask, VkDependencyFlags DependencyFlags, uint32 MemoryBarrierCount, const VkMemoryBarrier* MemoryBarriers, uint32 BufferMemoryBarrierCount, const VkBufferMemoryBarrier* BufferMemoryBarriers, uint32 ImageMemoryBarrierCount, const VkImageMemoryBarrier* ImageMemoryBarriers) VULKAN_LAYER_BODY
-	static void BeginRenderPass(VkResult Result, VkCommandBuffer CommandBuffer, const VkRenderPassBeginInfo* RenderPassBegin, VkSubpassContents Contents) VULKAN_LAYER_BODY
-	static void EndRenderPass(VkResult Result, VkCommandBuffer CommandBuffer) VULKAN_LAYER_BODY
-	static void NextSubpass(VkResult Result, VkCommandBuffer CommandBuffer, VkSubpassContents Contents) VULKAN_LAYER_BODY
+	static void ResetCommandBuffer(VkResult Result, VkCommandBuffer CommandBuffer, VkCommandBufferResetFlags Flags) VULKAN_LAYER_BODY
+
+	static void CmdBindPipeline(VkResult Result, VkCommandBuffer CommandBuffer, VkPipelineBindPoint PipelineBindPoint, VkPipeline Pipeline) VULKAN_LAYER_BODY
+	static void CmdSetViewport(VkResult Result, VkCommandBuffer CommandBuffer, uint32 FirstViewport, uint32 ViewportCount, const VkViewport* Viewports) VULKAN_LAYER_BODY
+	static void CmdSetScissor(VkResult Result, VkCommandBuffer CommandBuffer, uint32 FirstScissor, uint32 ScissorCount, const VkRect2D* Scissors) VULKAN_LAYER_BODY
+	static void CmdSetLineWidth(VkResult Result, VkCommandBuffer CommandBuffer, float LineWidth) VULKAN_LAYER_BODY
+	static void CmdSetDepthBias(VkResult Result, VkCommandBuffer CommandBuffer, float DepthBiasConstantFactor, float DepthBiasClamp, float DepthBiasSlopeFactor) VULKAN_LAYER_BODY
+	static void CmdSetBlendConstants(VkResult Result, VkCommandBuffer CommandBuffer, const float BlendConstants[4]) VULKAN_LAYER_BODY
+	static void CmdSetDepthBounds(VkResult Result, VkCommandBuffer CommandBuffer, float MinDepthBounds, float MaxDepthBounds) VULKAN_LAYER_BODY
+	static void CmdSetStencilCompareMask(VkResult Result, VkCommandBuffer CommandBuffer, VkStencilFaceFlags FaceMask, uint32 CompareMask) VULKAN_LAYER_BODY
+	static void CmdSetStencilWriteMask(VkResult Result, VkCommandBuffer CommandBuffer, VkStencilFaceFlags FaceMask, uint32 WriteMask) VULKAN_LAYER_BODY
+	static void CmdSetStencilReference(VkResult Result, VkCommandBuffer CommandBuffer, VkStencilFaceFlags FaceMask, uint32 Reference) VULKAN_LAYER_BODY
+	static void CmdBindDescriptorSets(VkResult Result, VkCommandBuffer CommandBuffer, VkPipelineBindPoint PipelineBindPoint, VkPipelineLayout Layout, uint32 FirstSet, uint32 DescriptorSetCount, const VkDescriptorSet* DescriptorSets, uint32 DynamicOffsetCount, const uint32* DynamicOffsets) VULKAN_LAYER_BODY
+	static void CmdBindIndexBuffer(VkResult Result, VkCommandBuffer CommandBuffer, VkBuffer IndexBuffer, VkDeviceSize Offset, VkIndexType IndexType) VULKAN_LAYER_BODY
+	static void CmdBindVertexBuffers(VkResult Result, VkCommandBuffer CommandBuffer, uint32 FirstBinding, uint32 BindingCount, const VkBuffer* Buffers, const VkDeviceSize* pOffsets) VULKAN_LAYER_BODY
+	static void CmdDraw(VkResult Result, VkCommandBuffer CommandBuffer, uint32 VertexCount, uint32 InstanceCount, uint32 FirstVertex, uint32 FirstInstance) VULKAN_LAYER_BODY
+	static void CmdDrawIndexed(VkResult Result, VkCommandBuffer CommandBuffer, uint32 IndexCount, uint32 InstanceCount, uint32 FirstIndex, int32_t VertexOffset, uint32 FirstInstance) VULKAN_LAYER_BODY
+	static void CmdDrawIndirect(VkResult Result, VkCommandBuffer CommandBuffer, VkBuffer Buffer, VkDeviceSize Offset, uint32 DrawCount, uint32 Stride) VULKAN_LAYER_BODY
+	static void CmdDrawIndexedIndirect(VkResult Result, VkCommandBuffer CommandBuffer, VkBuffer Buffer, VkDeviceSize Offset, uint32 DrawCount, uint32 Stride) VULKAN_LAYER_BODY
+	static void CmdDispatch(VkResult Result, VkCommandBuffer CommandBuffer, uint32 X, uint32 Y, uint32 Z) VULKAN_LAYER_BODY
+	static void CmdDispatchIndirect(VkResult Result, VkCommandBuffer CommandBuffer, VkBuffer Buffer, VkDeviceSize Offset) VULKAN_LAYER_BODY
+	static void CmdCopyBuffer(VkResult Result, VkCommandBuffer CommandBuffer, VkBuffer SrcBuffer, VkBuffer DstBuffer, uint32 RegionCount, const VkBufferCopy* Regions) VULKAN_LAYER_BODY
+	static void CmdCopyImage(VkResult Result, VkCommandBuffer CommandBuffer, VkImage SrcImage, VkImageLayout SrcImageLayout, VkImage DstImage, VkImageLayout DstImageLayout, uint32 RegionCount, const VkImageCopy* Regions) VULKAN_LAYER_BODY
+	VULKAN_EXTERN_EXPORT static void CmdBlitImage(VkResult Result, VkCommandBuffer CommandBuffer, VkImage SrcImage, VkImageLayout SrcImageLayout, VkImage DstImage, VkImageLayout DstImageLayout, uint32 RegionCount, const VkImageBlit* Regions, VkFilter Filter) VULKAN_LAYER_BODY
+	static void CmdCopyBufferToImage(VkResult Result, VkCommandBuffer CommandBuffer, VkBuffer SrcBuffer, VkImage DstImage, VkImageLayout DstImageLayout, uint32 RegionCount, const VkBufferImageCopy* Regions) VULKAN_LAYER_BODY
+	static void CmdCopyImageToBuffer(VkResult Result, VkCommandBuffer CommandBuffer, VkImage SrcImage, VkImageLayout SrcImageLayout, VkBuffer DstBuffer, uint32 RegionCount, const VkBufferImageCopy* Regions) VULKAN_LAYER_BODY
+	static void CmdUpdateBuffer(VkResult Result, VkCommandBuffer CommandBuffer, VkBuffer DstBuffer, VkDeviceSize DstOffset, VkDeviceSize DataSize, const void* pData) VULKAN_LAYER_BODY
+	static void CmdFillBuffer(VkResult Result, VkCommandBuffer CommandBuffer, VkBuffer DstBuffer, VkDeviceSize DstOffset, VkDeviceSize Size, uint32 Data) VULKAN_LAYER_BODY
+	VULKAN_EXTERN_EXPORT static void CmdClearColorImage(VkResult Result, VkCommandBuffer CommandBuffer, VkImage Image, VkImageLayout ImageLayout, const VkClearColorValue* ColorValue, uint32 RangeCount, const VkImageSubresourceRange* Ranges) VULKAN_LAYER_BODY
+	static void CmdClearDepthStencilImage(VkResult Result, VkCommandBuffer CommandBuffer, VkImage Image, VkImageLayout ImageLayout, const VkClearDepthStencilValue* DepthStencil, uint32 RangeCount, const VkImageSubresourceRange* Ranges) VULKAN_LAYER_BODY
+	static void CmdClearAttachments(VkResult Result, VkCommandBuffer CommandBuffer, uint32 AttachmentCount, const VkClearAttachment* Attachments, uint32 RectCount, const VkClearRect* Rects) VULKAN_LAYER_BODY
+	static void CmdResolveImage(VkResult Result, VkCommandBuffer CommandBuffer, VkImage SrcImage, VkImageLayout SrcImageLayout, VkImage DstImage, VkImageLayout DstImageLayout, uint32 RegionCount, const VkImageResolve* Regions) VULKAN_LAYER_BODY
+	static void CmdSetEvent(VkResult Result, VkCommandBuffer CommandBuffer, VkEvent Event, VkPipelineStageFlags StageMask) VULKAN_LAYER_BODY
+	static void CmdResetEvent(VkResult Result, VkCommandBuffer CommandBuffer, VkEvent Event, VkPipelineStageFlags StageMask) VULKAN_LAYER_BODY
+	static void CmdWaitEvents(VkResult Result, VkCommandBuffer CommandBuffer, uint32 EventCount, const VkEvent* Events, VkPipelineStageFlags SrcStageMask, VkPipelineStageFlags DstStageMask, uint32 MemoryBarrierCount, const VkMemoryBarrier* pMemoryBarriers, uint32 BufferMemoryBarrierCount, const VkBufferMemoryBarrier* pBufferMemoryBarriers, uint32 ImageMemoryBarrierCount, const VkImageMemoryBarrier* pImageMemoryBarriers) VULKAN_LAYER_BODY
+	static void CmdPipelineBarrier(VkResult Result, VkCommandBuffer CommandBuffer, VkPipelineStageFlags SrcStageMask, VkPipelineStageFlags DstStageMask, VkDependencyFlags DependencyFlags, uint32 MemoryBarrierCount, const VkMemoryBarrier* MemoryBarriers, uint32 BufferMemoryBarrierCount, const VkBufferMemoryBarrier* BufferMemoryBarriers, uint32 ImageMemoryBarrierCount, const VkImageMemoryBarrier* ImageMemoryBarriers) VULKAN_LAYER_BODY
+	static void CmdBeginQuery(VkResult Result, VkCommandBuffer CommandBuffer, VkQueryPool QueryPool, uint32 Query, VkQueryControlFlags Flags) VULKAN_LAYER_BODY
+	static void CmdEndQuery(VkResult Result, VkCommandBuffer CommandBuffer, VkQueryPool QueryPool, uint32 Query) VULKAN_LAYER_BODY
+	static void CmdResetQueryPool(VkResult Result, VkCommandBuffer CommandBuffer, VkQueryPool QueryPool, uint32 FirstQuery, uint32 QueryCount) VULKAN_LAYER_BODY
+	static void CmdWriteTimestamp(VkResult Result, VkCommandBuffer CommandBuffer, VkPipelineStageFlagBits PipelineStage, VkQueryPool QueryPool, uint32 Query) VULKAN_LAYER_BODY
+	static void CmdCopyQueryPoolResults(VkResult Result, VkCommandBuffer CommandBuffer, VkQueryPool QueryPool, uint32 FirstQuery, uint32 QueryCount, VkBuffer DstBuffer, VkDeviceSize DstOffset, VkDeviceSize Stride, VkQueryResultFlags Flags) VULKAN_LAYER_BODY
+	static void CmdPushConstants(VkResult Result, VkCommandBuffer CommandBuffer, VkPipelineLayout Layout, VkShaderStageFlags StageFlags, uint32_t Offset, uint32_t Size, const void* pValues) VULKAN_LAYER_BODY
+	static void CmdBeginRenderPass(VkResult Result, VkCommandBuffer CommandBuffer, const VkRenderPassBeginInfo* RenderPassBegin, VkSubpassContents Contents) VULKAN_LAYER_BODY
+	static void CmdNextSubpass(VkResult Result, VkCommandBuffer CommandBuffer, VkSubpassContents Contents) VULKAN_LAYER_BODY
+	static void CmdEndRenderPass(VkResult Result, VkCommandBuffer CommandBuffer) VULKAN_LAYER_BODY
+	static void CmdExecuteCommands(VkResult Result, VkCommandBuffer CommandBuffer, uint32 CommandBufferCount, const VkCommandBuffer* pCommandBuffers) VULKAN_LAYER_BODY
+
 	static void QueuePresent(VkResult Result, VkQueue Queue, const VkPresentInfoKHR* PresentInfo) VULKAN_LAYER_BODY
-	static void GetPhysicalDeviceMemoryProperties(VkResult Result, VkPhysicalDevice PhysicalDevice, VkPhysicalDeviceMemoryProperties* Properties) VULKAN_LAYER_BODY
-	static void CreateDevice(VkResult Result, VkPhysicalDevice PhysicalDevice, const VkDeviceCreateInfo* CreateInfo, VkDevice* Device) VULKAN_LAYER_BODY
-	static void QueueSubmit(VkResult Result, VkQueue Queue, uint32 SubmitCount, const VkSubmitInfo* Submits, VkFence Fence) VULKAN_LAYER_BODY
-	static void GetPhysicalDeviceFeatures(VkResult Result, VkPhysicalDevice PhysicalDevice, VkPhysicalDeviceFeatures* Features) VULKAN_LAYER_BODY
 	static void GetSwapChainImagesKHR(VkResult Result, VkDevice Device, VkSwapchainKHR Swapchain, uint32_t* SwapchainImageCount, VkImage* SwapchainImages) VULKAN_LAYER_BODY
-	static void GetDeviceQueue(VkResult Result, VkDevice Device, uint32 QueueFamilyIndex, uint32 QueueIndex, VkQueue* Queue) VULKAN_LAYER_BODY
-	static void DeviceWaitIdle(VkResult Result, VkDevice Device) VULKAN_LAYER_BODY
-	static void MapMemory(VkResult Result, VkDevice Device, VkDeviceMemory Memory, VkDeviceSize Offset, VkDeviceSize Size, VkMemoryMapFlags Flags, void** Data) VULKAN_LAYER_BODY
-	static void UnmapMemory(VkResult Result, VkDevice Device, VkDeviceMemory Memory) VULKAN_LAYER_BODY
-	static void BindBufferMemory(VkResult Result, VkDevice Device, VkBuffer Buffer, VkDeviceMemory Memory, VkDeviceSize MemoryOffset) VULKAN_LAYER_BODY
-	VULKAN_EXTERN_EXPORT static void BindImageMemory(VkResult Result, VkDevice Device, VkImage Image, VkDeviceMemory Memory, VkDeviceSize MemoryOffset) VULKAN_LAYER_BODY
-	static void GetFenceStatus(VkResult Result, VkDevice Device, VkFence Fence) VULKAN_LAYER_BODY
-	static void GetQueryPoolResults(VkResult Result, VkDevice Device, VkQueryPool QueryPool, uint32 FirstQuery, uint32 QueryCount, size_t DataSize, void* Data, VkDeviceSize Stride, VkQueryResultFlags Flags) VULKAN_LAYER_BODY
-	static void CreateComputePipelines(VkResult Result, VkDevice Device, VkPipelineCache PipelineCache, uint32 CreateInfoCount, const VkComputePipelineCreateInfo* CreateInfos, VkPipeline* Pipelines) VULKAN_LAYER_BODY
-	static void AllocateCommandBuffers(VkResult Result, VkDevice Device, const VkCommandBufferAllocateInfo* AllocateInfo, VkCommandBuffer* CommandBuffers) VULKAN_LAYER_BODY
 	static void CreateSwapchainKHR(VkResult Result, VkDevice Device, const VkSwapchainCreateInfoKHR* CreateInfo, VkSwapchainKHR* Swapchain) VULKAN_LAYER_BODY
 	static void AcquireNextImageKHR(VkResult Result, VkDevice Device, VkSwapchainKHR Swapchain, uint64_t Timeout, VkSemaphore Semaphore, VkFence Fence, uint32_t* ImageIndex) VULKAN_LAYER_BODY
-	static void FreeMemory(VkResult Result, VkDevice Device, VkDeviceMemory Memory) VULKAN_LAYER_BODY
-	static void DestroyFence(VkResult Result, VkDevice Device, VkFence Fence) VULKAN_LAYER_BODY
-	static void DestroySemaphore(VkResult Result, VkDevice Device, VkSemaphore Semaphore) VULKAN_LAYER_BODY
-	static void CreateEvent(VkResult Result, VkDevice Device, const VkEventCreateInfo* CreateInfo, VkEvent* Event) VULKAN_LAYER_BODY
-	static void DestroyEvent(VkResult Result, VkDevice Device, VkEvent Event) VULKAN_LAYER_BODY
-	static void DestroyBuffer(VkResult Result, VkDevice Device, VkBuffer Buffer) VULKAN_LAYER_BODY
-	static void DestroyBufferView(VkResult Result, VkDevice Device, VkBufferView BufferView) VULKAN_LAYER_BODY
-	static void DestroyImageView(VkResult Result, VkDevice Device, VkImageView ImageView) VULKAN_LAYER_BODY
-	static void DestroyQueryPool(VkResult Result, VkDevice Device, VkQueryPool QueryPool) VULKAN_LAYER_BODY
-	static void DestroyPipeline(VkResult Result, VkDevice Device, VkPipeline Pipeline) VULKAN_LAYER_BODY
-	static void DestroyPipelineLayout(VkResult Result, VkDevice Device, VkPipelineLayout PipelineLayout) VULKAN_LAYER_BODY
-	static void DestroyPipelineCache(VkResult Result, VkDevice Device, VkPipelineCache PipelineCache) VULKAN_LAYER_BODY
-	static void DestroySampler(VkResult Result, VkDevice Device, VkSampler Sampler) VULKAN_LAYER_BODY
-	static void DestroyShaderModule(VkResult Result, VkDevice Device, VkShaderModule ShaderModule) VULKAN_LAYER_BODY
-	static void DestroyFramebuffer(VkResult Result, VkDevice Device, VkFramebuffer Framebuffer) VULKAN_LAYER_BODY
-	static void DestroyRenderPass(VkResult Result, VkDevice Device, VkRenderPass RenderPass) VULKAN_LAYER_BODY
-	static void DestroyCommandPool(VkResult Result, VkDevice Device, VkCommandPool CommandPool) VULKAN_LAYER_BODY
-	static void DestroyDevice(VkResult Result, VkDevice Device) VULKAN_LAYER_BODY
 	static void DestroySurfaceKHR(VkResult Result, VkInstance Instance, VkSurfaceKHR SurfaceKHR) VULKAN_LAYER_BODY
-	static void DestroyDescriptorSetLayout(VkResult Result, VkDevice Device, VkDescriptorSetLayout DescriptorSetLayout) VULKAN_LAYER_BODY
-	static void DestroyDescriptorPool(VkResult Result, VkDevice Device, VkDescriptorPool DescriptorPool) VULKAN_LAYER_BODY
-	static void ResetDescriptorPool(VkResult Result, VkDevice Device, VkDescriptorPool DescriptorPool, VkDescriptorPoolResetFlags Flags) VULKAN_LAYER_BODY
 	static void DestroySwapchainKHR(VkResult Result, VkDevice Device, VkSwapchainKHR Swapchain) VULKAN_LAYER_BODY
-	static void FreeCommandBuffers(VkResult Result, VkDevice Device, VkCommandPool CommandPool, uint32 CommandBufferCount, const VkCommandBuffer* CommandBuffers) VULKAN_LAYER_BODY
-	static void GetPipelineCacheData(VkResult Result, VkDevice Device, VkPipelineCache PipelineCache, size_t* DataSize, void* Data) VULKAN_LAYER_BODY
-	static void MergePipelineCaches(VkResult Result, VkDevice Device, VkPipelineCache DestCache, uint32 SourceCacheCount, const VkPipelineCache* SrcCaches) VULKAN_LAYER_BODY
-	static void GetPhysicalDeviceQueueFamilyProperties(VkResult Result, VkPhysicalDevice PhysicalDevice, uint32* QueueFamilyPropertyCount, VkQueueFamilyProperties* QueueFamilyProperties) VULKAN_LAYER_BODY
-	static void QueueWaitIdle(VkResult Result, VkQueue Queue) VULKAN_LAYER_BODY
-	static void Draw(VkResult Result, VkCommandBuffer CommandBuffer, uint32 VertexCount, uint32 InstanceCount, uint32 FirstVertex, uint32 FirstInstance) VULKAN_LAYER_BODY
-	static void DrawIndexed(VkResult Result, VkCommandBuffer CommandBuffer, uint32 IndexCount, uint32 InstanceCount, uint32 FirstIndex, int32_t VertexOffset, uint32 FirstInstance) VULKAN_LAYER_BODY
-	static void DrawIndirect(VkResult Result, VkCommandBuffer CommandBuffer, VkBuffer Buffer, VkDeviceSize Offset, uint32 DrawCount, uint32 Stride) VULKAN_LAYER_BODY
-	static void DrawIndexedIndirect(VkResult Result, VkCommandBuffer CommandBuffer, VkBuffer Buffer, VkDeviceSize Offset, uint32 DrawCount, uint32 Stride) VULKAN_LAYER_BODY
-	static void Dispatch(VkResult Result, VkCommandBuffer CommandBuffer, uint32 X, uint32 Y, uint32 Z) VULKAN_LAYER_BODY
-	static void DispatchIndirect(VkResult Result, VkCommandBuffer CommandBuffer, VkBuffer Buffer, VkDeviceSize Offset) VULKAN_LAYER_BODY
-	static void CopyImage(VkResult Result, VkCommandBuffer CommandBuffer, VkImage SrcImage, VkImageLayout SrcImageLayout, VkImage DstImage, VkImageLayout DstImageLayout, uint32 RegionCount, const VkImageCopy* Regions) VULKAN_LAYER_BODY
-	static void BeginQuery(VkResult Result, VkCommandBuffer CommandBuffer, VkQueryPool QueryPool, uint32 Query, VkQueryControlFlags Flags) VULKAN_LAYER_BODY
-	static void EndQuery(VkResult Result, VkCommandBuffer CommandBuffer, VkQueryPool QueryPool, uint32 Query) VULKAN_LAYER_BODY
 #if VULKAN_SUPPORTS_DEDICATED_ALLOCATION
 	static void GetImageMemoryRequirements2KHR(VkResult Result, VkDevice Device, const VkImageMemoryRequirementsInfo2KHR* Info, VkMemoryRequirements2KHR* MemoryRequirements) VULKAN_LAYER_BODY
 #endif
-	static void ResetQueryPool(VkResult Result, VkCommandBuffer CommandBuffer, VkQueryPool QueryPool, uint32 FirstQuery, uint32 QueryCount) VULKAN_LAYER_BODY
-	static void WriteTimestamp(VkResult Result, VkCommandBuffer CommandBuffer, VkPipelineStageFlagBits PipelineStage, VkQueryPool QueryPool, uint32 Query) VULKAN_LAYER_BODY
-	static void BindPipeline(VkResult Result, VkCommandBuffer CommandBuffer, VkPipelineBindPoint PipelineBindPoint, VkPipeline Pipeline) VULKAN_LAYER_BODY
-	static void ResetCommandBuffer(VkResult Result, VkCommandBuffer CommandBuffer, VkCommandBufferResetFlags Flags) VULKAN_LAYER_BODY
-	static void SetViewport(VkResult Result, VkCommandBuffer CommandBuffer, uint32 FirstViewport, uint32 ViewportCount, const VkViewport* Viewports) VULKAN_LAYER_BODY
-	static void GetPhysicalDeviceFormatProperties(VkResult Result, VkPhysicalDevice PhysicalDevice, VkFormat Format, VkFormatProperties* FormatProperties) VULKAN_LAYER_BODY
-	static void GetPhysicalDeviceProperties(VkResult Result, VkPhysicalDevice PhysicalDevice, VkPhysicalDeviceProperties* FormatProperties) VULKAN_LAYER_BODY
-	static void SetScissor(VkResult Result, VkCommandBuffer CommandBuffer, uint32 FirstScissor, uint32 ScissorCount, const VkRect2D* Scissors) VULKAN_LAYER_BODY
-	static void SetLineWidth(VkResult Result, VkCommandBuffer CommandBuffer, float LineWidth) VULKAN_LAYER_BODY
 #if VULKAN_HAS_PHYSICAL_DEVICE_PROPERTIES2
 	static void GetPhysicalDeviceProperties2KHR(VkResult Result, VkPhysicalDevice PhysicalDevice, VkPhysicalDeviceProperties2KHR* Properties) VULKAN_LAYER_BODY
 #endif
-	static void SetDepthBias(VkResult Result, VkCommandBuffer CommandBuffer, float DepthBiasConstantFactor, float DepthBiasClamp, float DepthBiasSlopeFactor) VULKAN_LAYER_BODY
-	static void SetBlendConstants(VkResult Result, VkCommandBuffer CommandBuffer, const float BlendConstants[4]) VULKAN_LAYER_BODY
-	static void SetDepthBounds(VkResult Result, VkCommandBuffer CommandBuffer, float MinDepthBounds, float MaxDepthBounds) VULKAN_LAYER_BODY
-	static void SetStencilCompareMask(VkResult Result, VkCommandBuffer CommandBuffer, VkStencilFaceFlags FaceMask, uint32 CompareMask) VULKAN_LAYER_BODY
-	static void SetStencilWriteMask(VkResult Result, VkCommandBuffer CommandBuffer, VkStencilFaceFlags FaceMask, uint32 WriteMask) VULKAN_LAYER_BODY
-	static void SetStencilReference(VkResult Result, VkCommandBuffer CommandBuffer, VkStencilFaceFlags FaceMask, uint32 Reference) VULKAN_LAYER_BODY
-	static void UpdateBuffer(VkResult Result, VkCommandBuffer CommandBuffer, VkBuffer DstBuffer, VkDeviceSize DstOffset, VkDeviceSize DataSize, const void* pData) VULKAN_LAYER_BODY
-	static void FillBuffer(VkResult Result, VkCommandBuffer CommandBuffer, VkBuffer DstBuffer, VkDeviceSize DstOffset, VkDeviceSize Size, uint32 Data) VULKAN_LAYER_BODY
-	static void SetEvent(VkResult Result, VkCommandBuffer CommandBuffer, VkEvent Event, VkPipelineStageFlags StageMask) VULKAN_LAYER_BODY
-	static void ResetEvent(VkResult Result, VkCommandBuffer CommandBuffer, VkEvent Event, VkPipelineStageFlags StageMask) VULKAN_LAYER_BODY
-	static void SetEvent(VkResult Result, VkDevice Device, VkEvent Event) VULKAN_LAYER_BODY
-	static void ResetEvent(VkResult Result, VkDevice Device, VkEvent Event) VULKAN_LAYER_BODY
-	static void GetEventStatus(VkResult Result, VkDevice Device, VkEvent Event) VULKAN_LAYER_BODY
-	static void CopyQueryPoolResults(VkResult Result, VkCommandBuffer CommandBuffer, VkQueryPool QueryPool, uint32 FirstQuery, uint32 QueryCount, VkBuffer DstBuffer, VkDeviceSize DstOffset, VkDeviceSize Stride, VkQueryResultFlags Flags) VULKAN_LAYER_BODY
-	static void GetDeviceProcAddr(VkResult Result, VkDevice Device, const char* Name, PFN_vkVoidFunction VoidFunction) VULKAN_LAYER_BODY
-	VULKAN_EXTERN_EXPORT static void EnumerateDeviceExtensionProperties(VkResult Result, VkPhysicalDevice PhysicalDevice, const char* LayerName, uint32* PropertyCount, VkExtensionProperties* Properties) VULKAN_LAYER_BODY
-	static void EnumerateDeviceLayerProperties(VkResult Result, VkPhysicalDevice PhysicalDevice, uint32* PropertyCount, VkLayerProperties* Properties) VULKAN_LAYER_BODY
 	static void GetPhysicalDeviceSurfaceCapabilitiesKHR(VkResult Result, VkPhysicalDevice PhysicalDevice, VkSurfaceKHR Surface, VkSurfaceCapabilitiesKHR* SurfaceCapabilities) VULKAN_LAYER_BODY
 	static void GetPhysicalDeviceSurfaceFormatsKHR(VkResult Result, VkPhysicalDevice PhysicalDevice, VkSurfaceKHR Surface, uint32_t* SurfaceFormatCountPtr, VkSurfaceFormatKHR* SurfaceFormats) VULKAN_LAYER_BODY
 	static void GetPhysicalDeviceSurfaceSupportKHR(VkResult Result, VkPhysicalDevice PhysicalDevice, uint32_t QueueFamilyIndex, VkSurfaceKHR Surface, VkBool32* SupportedPtr) VULKAN_LAYER_BODY
@@ -222,16 +232,14 @@ namespace VulkanRHI
 		FWrapLayer::GetPhysicalDeviceFormatProperties(VK_SUCCESS, PhysicalDevice, Format, FormatProperties);
 	}
 
-#if 0
-	static FORCEINLINE_DEBUGGABLE VkResult  vkGetPhysicalDeviceImageFormatProperties(
-		VkPhysicalDevice                            physicalDevice,
-		VkFormat                                    format,
-		VkImageType                                 type,
-		VkImageTiling                               tiling,
-		VkImageUsageFlags                           usage,
-		VkImageCreateFlags                          flags,
-		VkImageFormatProperties*                    pImageFormatProperties);
-#endif
+	static FORCEINLINE_DEBUGGABLE VkResult  vkGetPhysicalDeviceImageFormatProperties(VkPhysicalDevice PhysicalDevice, VkFormat Format, VkImageType Type, VkImageTiling Tiling, VkImageUsageFlags Usage, VkImageCreateFlags Flags, VkImageFormatProperties* pImageFormatProperties)
+	{
+		FWrapLayer::GetPhysicalDeviceImageFormatProperties(VK_RESULT_MAX_ENUM, PhysicalDevice, Format, Type, Tiling, Usage, Flags, pImageFormatProperties);
+		VkResult Result = VULKANAPINAMESPACE::vkGetPhysicalDeviceImageFormatProperties(PhysicalDevice, Format, Type, Tiling, Usage, Flags, pImageFormatProperties);
+		FWrapLayer::GetPhysicalDeviceImageFormatProperties(Result, PhysicalDevice, Format, Type, Tiling, Usage, Flags, pImageFormatProperties);
+		return Result;
+	}
+
 	static FORCEINLINE_DEBUGGABLE void  vkGetPhysicalDeviceProperties(VkPhysicalDevice PhysicalDevice, VkPhysicalDeviceProperties* Properties)
 	{
 		FWrapLayer::GetPhysicalDeviceProperties(VK_RESULT_MAX_ENUM, PhysicalDevice, Properties);
@@ -403,12 +411,14 @@ namespace VulkanRHI
 		FWrapLayer::InvalidateMappedMemoryRanges(Result, Device, MemoryRangeCount, MemoryRanges);
 		return Result;
 	}
-#if 0
-	static FORCEINLINE_DEBUGGABLE void  vkGetDeviceMemoryCommitment(
-		VkDevice                                    Device,
-		VkDeviceMemory                              Memory,
-		VkDeviceSize*                               pCommittedMemoryInBytes);
-#endif
+
+	static FORCEINLINE_DEBUGGABLE void  vkGetDeviceMemoryCommitment(VkDevice Device, VkDeviceMemory Memory, VkDeviceSize* pCommittedMemoryInBytes)
+	{
+		FWrapLayer::GetDeviceMemoryCommitment(VK_RESULT_MAX_ENUM, Device, Memory, pCommittedMemoryInBytes);
+		VULKANAPINAMESPACE::vkGetDeviceMemoryCommitment(Device, Memory, pCommittedMemoryInBytes);
+		FWrapLayer::GetDeviceMemoryCommitment(VK_SUCCESS, Device, Memory, pCommittedMemoryInBytes);
+	}
+
 	static FORCEINLINE_DEBUGGABLE VkResult  vkBindBufferMemory(VkDevice Device, VkBuffer Buffer, VkDeviceMemory Memory, VkDeviceSize MemoryOffset)
 	{
 		FWrapLayer::BindBufferMemory(VK_RESULT_MAX_ENUM, Device, Buffer, Memory, MemoryOffset);
@@ -439,29 +449,28 @@ namespace VulkanRHI
 		FWrapLayer::GetImageMemoryRequirements(VK_SUCCESS, Device, Image, MemoryRequirements);
 	}
 
-#if 0
-	static FORCEINLINE_DEBUGGABLE void  vkGetImageSparseMemoryRequirements(
-		VkDevice                                    Device,
-		VkImage                                     Image,
-		uint32*                                   pSparseMemoryRequirementCount,
-		VkSparseImageMemoryRequirements*            pSparseMemoryRequirements);
-
-	static FORCEINLINE_DEBUGGABLE void  vkGetPhysicalDeviceSparseImageFormatProperties(
-		VkPhysicalDevice                            physicalDevice,
-		VkFormat                                    format,
-		VkImageType                                 type,
-		VkSampleCountFlagBits                       samples,
-		VkImageUsageFlags                           usage,
-		VkImageTiling                               tiling,
-		uint32*                                   pPropertyCount,
-		VkSparseImageFormatProperties*              pProperties);
-
-	static FORCEINLINE_DEBUGGABLE VkResult  vkQueueBindSparse(
-		VkQueue                                     queue,
-		uint32                                    bindInfoCount,
-		const VkBindSparseInfo*                     pBindInfo,
-		VkFence                                     Fence);
-#endif
+	static FORCEINLINE_DEBUGGABLE void  vkGetImageSparseMemoryRequirements(VkDevice Device, VkImage Image, uint32* pSparseMemoryRequirementCount, VkSparseImageMemoryRequirements* pSparseMemoryRequirements)
+	{
+		FWrapLayer::GetImageSparseMemoryRequirements(VK_RESULT_MAX_ENUM, Device, Image, pSparseMemoryRequirementCount, pSparseMemoryRequirements);
+		VULKANAPINAMESPACE::vkGetImageSparseMemoryRequirements(Device, Image, pSparseMemoryRequirementCount, pSparseMemoryRequirements);
+		FWrapLayer::GetImageSparseMemoryRequirements(VK_SUCCESS, Device, Image, pSparseMemoryRequirementCount, pSparseMemoryRequirements);
+	}
+
+	static FORCEINLINE_DEBUGGABLE void  vkGetPhysicalDeviceSparseImageFormatProperties(VkPhysicalDevice PhysicalDevice, VkFormat Format, VkImageType Type, VkSampleCountFlagBits Samples, VkImageUsageFlags Usage, VkImageTiling Tiling, uint32* pPropertyCount, VkSparseImageFormatProperties* pProperties)
+	{
+		FWrapLayer::GetPhysicalDeviceSparseImageFormatProperties(VK_RESULT_MAX_ENUM, PhysicalDevice, Format, Type, Samples, Usage, Tiling, pPropertyCount, pProperties);
+		VULKANAPINAMESPACE::vkGetPhysicalDeviceSparseImageFormatProperties(PhysicalDevice, Format, Type, Samples, Usage, Tiling, pPropertyCount, pProperties);
+		FWrapLayer::GetPhysicalDeviceSparseImageFormatProperties(VK_SUCCESS, PhysicalDevice, Format, Type, Samples, Usage, Tiling, pPropertyCount, pProperties);
+	}
+
+	static FORCEINLINE_DEBUGGABLE VkResult  vkQueueBindSparse(VkQueue Queue, uint32 BindInfoCount, const VkBindSparseInfo* pBindInfo, VkFence Fence)
+	{
+		FWrapLayer::QueueBindSparse(VK_RESULT_MAX_ENUM, Queue, BindInfoCount, pBindInfo, Fence);
+		VkResult Result = vkQueueBindSparse(Queue, BindInfoCount, pBindInfo, Fence);
+		FWrapLayer::QueueBindSparse(Result, Queue, BindInfoCount, pBindInfo, Fence);
+		return Result;
+	}
+
 	static FORCEINLINE_DEBUGGABLE VkResult  vkCreateFence(VkDevice Device, const VkFenceCreateInfo* CreateInfo, const VkAllocationCallbacks* Allocator, VkFence* Fence)
 	{
 		FWrapLayer::CreateFence(VK_RESULT_MAX_ENUM, Device, CreateInfo, Fence);
@@ -836,12 +845,14 @@ namespace VulkanRHI
 		VULKANAPINAMESPACE::vkDestroyRenderPass(Device, RenderPass, Allocator);
 		FWrapLayer::DestroyRenderPass(VK_SUCCESS, Device, RenderPass);
 	}
-#if 0
-	static FORCEINLINE_DEBUGGABLE void  vkGetRenderAreaGranularity(
-		VkDevice                                    Device,
-		VkRenderPass                                renderPass,
-		VkExtent2D*                                 pGranularity);
-#endif
+
+	static FORCEINLINE_DEBUGGABLE void  vkGetRenderAreaGranularity(VkDevice Device, VkRenderPass RenderPass, VkExtent2D* pGranularity)
+	{
+		FWrapLayer::GetRenderAreaGranularity(VK_RESULT_MAX_ENUM, Device, RenderPass, pGranularity);
+		VULKANAPINAMESPACE::vkGetRenderAreaGranularity(Device, RenderPass, pGranularity);
+		FWrapLayer::GetRenderAreaGranularity(VK_SUCCESS, Device, RenderPass, pGranularity);
+	}
+
 	static FORCEINLINE_DEBUGGABLE VkResult  vkCreateCommandPool(VkDevice Device, const VkCommandPoolCreateInfo* CreateInfo, const VkAllocationCallbacks* Allocator, VkCommandPool* CommandPool)
 	{
 		FWrapLayer::CreateCommandPool(VK_RESULT_MAX_ENUM, Device, CreateInfo, CommandPool);
@@ -857,12 +868,14 @@ namespace VulkanRHI
 		FWrapLayer::DestroyCommandPool(VK_SUCCESS, Device, CommandPool);
 	}
 
-#if 0
-	static FORCEINLINE_DEBUGGABLE VkResult  vkResetCommandPool(
-		VkDevice                                    Device,
-		VkCommandPool                               commandPool,
-		VkCommandPoolResetFlags                     flags);
-#endif
+	static FORCEINLINE_DEBUGGABLE VkResult  vkResetCommandPool(VkDevice Device, VkCommandPool CommandPool, VkCommandPoolResetFlags Flags)
+	{
+		FWrapLayer::ResetCommandPool(VK_RESULT_MAX_ENUM, Device, CommandPool, Flags);
+		VkResult Result = VULKANAPINAMESPACE::vkResetCommandPool(Device, CommandPool, Flags);
+		FWrapLayer::ResetCommandPool(Result, Device, CommandPool, Flags);
+		return Result;
+	}
+
 	static FORCEINLINE_DEBUGGABLE VkResult  vkAllocateCommandBuffers(VkDevice Device, const VkCommandBufferAllocateInfo* AllocateInfo, VkCommandBuffer* CommandBuffers)
 	{
 		FWrapLayer::AllocateCommandBuffers(VK_RESULT_MAX_ENUM, Device, AllocateInfo, CommandBuffers);
@@ -904,205 +917,205 @@ namespace VulkanRHI
 
 	static FORCEINLINE_DEBUGGABLE void  vkCmdBindPipeline(VkCommandBuffer CommandBuffer, VkPipelineBindPoint PipelineBindPoint, VkPipeline Pipeline)
 	{
-		FWrapLayer::BindPipeline(VK_RESULT_MAX_ENUM, CommandBuffer, PipelineBindPoint, Pipeline);
+		FWrapLayer::CmdBindPipeline(VK_RESULT_MAX_ENUM, CommandBuffer, PipelineBindPoint, Pipeline);
 		VULKANAPINAMESPACE::vkCmdBindPipeline(CommandBuffer, PipelineBindPoint, Pipeline);
-		FWrapLayer::BindPipeline(VK_SUCCESS, CommandBuffer, PipelineBindPoint, Pipeline);
+		FWrapLayer::CmdBindPipeline(VK_SUCCESS, CommandBuffer, PipelineBindPoint, Pipeline);
 	}
 
 	static FORCEINLINE_DEBUGGABLE void  vkCmdSetViewport(VkCommandBuffer CommandBuffer, uint32 FirstViewport, uint32 ViewportCount, const VkViewport* Viewports)
 	{
-		FWrapLayer::SetViewport(VK_RESULT_MAX_ENUM, CommandBuffer, FirstViewport, ViewportCount, Viewports);
+		FWrapLayer::CmdSetViewport(VK_RESULT_MAX_ENUM, CommandBuffer, FirstViewport, ViewportCount, Viewports);
 		VULKANAPINAMESPACE::vkCmdSetViewport(CommandBuffer, FirstViewport, ViewportCount, Viewports);
-		FWrapLayer::SetViewport(VK_SUCCESS, CommandBuffer, FirstViewport, ViewportCount, Viewports);
+		FWrapLayer::CmdSetViewport(VK_SUCCESS, CommandBuffer, FirstViewport, ViewportCount, Viewports);
 	}
 
 	static FORCEINLINE_DEBUGGABLE void  vkCmdSetScissor(VkCommandBuffer CommandBuffer, uint32 FirstScissor, uint32 ScissorCount, const VkRect2D* Scissors)
 	{
-		FWrapLayer::SetScissor(VK_RESULT_MAX_ENUM, CommandBuffer, FirstScissor, ScissorCount, Scissors);
+		FWrapLayer::CmdSetScissor(VK_RESULT_MAX_ENUM, CommandBuffer, FirstScissor, ScissorCount, Scissors);
 		VULKANAPINAMESPACE::vkCmdSetScissor(CommandBuffer, FirstScissor, ScissorCount, Scissors);
-		FWrapLayer::SetScissor(VK_SUCCESS, CommandBuffer, FirstScissor, ScissorCount, Scissors);
+		FWrapLayer::CmdSetScissor(VK_SUCCESS, CommandBuffer, FirstScissor, ScissorCount, Scissors);
 	}
 
 	static FORCEINLINE_DEBUGGABLE void  vkCmdSetLineWidth(VkCommandBuffer CommandBuffer, float LineWidth)
 	{
-		FWrapLayer::SetLineWidth(VK_RESULT_MAX_ENUM, CommandBuffer, LineWidth);
+		FWrapLayer::CmdSetLineWidth(VK_RESULT_MAX_ENUM, CommandBuffer, LineWidth);
 		VULKANAPINAMESPACE::vkCmdSetLineWidth(CommandBuffer, LineWidth);
-		FWrapLayer::SetLineWidth(VK_SUCCESS, CommandBuffer, LineWidth);
+		FWrapLayer::CmdSetLineWidth(VK_SUCCESS, CommandBuffer, LineWidth);
 	}
 
 	static FORCEINLINE_DEBUGGABLE void  vkCmdSetDepthBias(VkCommandBuffer CommandBuffer, float DepthBiasConstantFactor, float DepthBiasClamp, float DepthBiasSlopeFactor)
 	{
-		FWrapLayer::SetDepthBias(VK_RESULT_MAX_ENUM, CommandBuffer, DepthBiasConstantFactor, DepthBiasClamp, DepthBiasSlopeFactor);
+		FWrapLayer::CmdSetDepthBias(VK_RESULT_MAX_ENUM, CommandBuffer, DepthBiasConstantFactor, DepthBiasClamp, DepthBiasSlopeFactor);
 		VULKANAPINAMESPACE::vkCmdSetDepthBias(CommandBuffer, DepthBiasConstantFactor, DepthBiasClamp, DepthBiasSlopeFactor);
-		FWrapLayer::SetDepthBias(VK_SUCCESS, CommandBuffer, DepthBiasConstantFactor, DepthBiasClamp, DepthBiasSlopeFactor);
+		FWrapLayer::CmdSetDepthBias(VK_SUCCESS, CommandBuffer, DepthBiasConstantFactor, DepthBiasClamp, DepthBiasSlopeFactor);
 	}
 
 	static FORCEINLINE_DEBUGGABLE void  vkCmdSetBlendConstants(VkCommandBuffer CommandBuffer, const float BlendConstants[4])
 	{
-		FWrapLayer::SetBlendConstants(VK_RESULT_MAX_ENUM, CommandBuffer, BlendConstants);
+		FWrapLayer::CmdSetBlendConstants(VK_RESULT_MAX_ENUM, CommandBuffer, BlendConstants);
 		VULKANAPINAMESPACE::vkCmdSetBlendConstants(CommandBuffer, BlendConstants);
-		FWrapLayer::SetBlendConstants(VK_SUCCESS, CommandBuffer, BlendConstants);
+		FWrapLayer::CmdSetBlendConstants(VK_SUCCESS, CommandBuffer, BlendConstants);
 	}
 
 	static FORCEINLINE_DEBUGGABLE void  vkCmdSetDepthBounds(VkCommandBuffer CommandBuffer, float MinDepthBounds, float MaxDepthBounds)
 	{
-		FWrapLayer::SetDepthBounds(VK_RESULT_MAX_ENUM, CommandBuffer, MinDepthBounds, MaxDepthBounds);
+		FWrapLayer::CmdSetDepthBounds(VK_RESULT_MAX_ENUM, CommandBuffer, MinDepthBounds, MaxDepthBounds);
 		VULKANAPINAMESPACE::vkCmdSetDepthBounds(CommandBuffer, MinDepthBounds, MaxDepthBounds);
-		FWrapLayer::SetDepthBounds(VK_SUCCESS, CommandBuffer, MinDepthBounds, MaxDepthBounds);
+		FWrapLayer::CmdSetDepthBounds(VK_SUCCESS, CommandBuffer, MinDepthBounds, MaxDepthBounds);
 	}
 
 	static FORCEINLINE_DEBUGGABLE void  vkCmdSetStencilCompareMask(VkCommandBuffer CommandBuffer, VkStencilFaceFlags FaceMask, uint32 CompareMask)
 	{
-		FWrapLayer::SetStencilCompareMask(VK_RESULT_MAX_ENUM, CommandBuffer, FaceMask, CompareMask);
+		FWrapLayer::CmdSetStencilCompareMask(VK_RESULT_MAX_ENUM, CommandBuffer, FaceMask, CompareMask);
 		VULKANAPINAMESPACE::vkCmdSetStencilCompareMask(CommandBuffer, FaceMask, CompareMask);
-		FWrapLayer::SetStencilCompareMask(VK_SUCCESS, CommandBuffer, FaceMask, CompareMask);
+		FWrapLayer::CmdSetStencilCompareMask(VK_SUCCESS, CommandBuffer, FaceMask, CompareMask);
 	}
 
 	static FORCEINLINE_DEBUGGABLE void  vkCmdSetStencilWriteMask(VkCommandBuffer CommandBuffer, VkStencilFaceFlags FaceMask, uint32 WriteMask)
 	{
-		FWrapLayer::SetStencilWriteMask(VK_RESULT_MAX_ENUM, CommandBuffer, FaceMask, WriteMask);
+		FWrapLayer::CmdSetStencilWriteMask(VK_RESULT_MAX_ENUM, CommandBuffer, FaceMask, WriteMask);
 		VULKANAPINAMESPACE::vkCmdSetStencilWriteMask(CommandBuffer, FaceMask, WriteMask);
-		FWrapLayer::SetStencilWriteMask(VK_SUCCESS, CommandBuffer, FaceMask, WriteMask);
+		FWrapLayer::CmdSetStencilWriteMask(VK_SUCCESS, CommandBuffer, FaceMask, WriteMask);
 	}
 
 	static FORCEINLINE_DEBUGGABLE void  vkCmdSetStencilReference(VkCommandBuffer CommandBuffer, VkStencilFaceFlags FaceMask, uint32 Reference)
 	{
-		FWrapLayer::SetStencilReference(VK_RESULT_MAX_ENUM, CommandBuffer, FaceMask, Reference);
+		FWrapLayer::CmdSetStencilReference(VK_RESULT_MAX_ENUM, CommandBuffer, FaceMask, Reference);
 		VULKANAPINAMESPACE::vkCmdSetStencilReference(CommandBuffer, FaceMask, Reference);
-		FWrapLayer::SetStencilReference(VK_SUCCESS, CommandBuffer, FaceMask, Reference);
+		FWrapLayer::CmdSetStencilReference(VK_SUCCESS, CommandBuffer, FaceMask, Reference);
 	}
 
 	static FORCEINLINE_DEBUGGABLE void  vkCmdBindDescriptorSets(VkCommandBuffer CommandBuffer, VkPipelineBindPoint PipelineBindPoint, VkPipelineLayout Layout, uint32 FirstSet, uint32 DescriptorSetCount, const VkDescriptorSet* DescriptorSets, uint32 DynamicOffsetCount, const uint32* DynamicOffsets)
 	{
-		FWrapLayer::BindDescriptorSets(VK_RESULT_MAX_ENUM, CommandBuffer, PipelineBindPoint, Layout, FirstSet, DescriptorSetCount, DescriptorSets, DynamicOffsetCount, DynamicOffsets);
+		FWrapLayer::CmdBindDescriptorSets(VK_RESULT_MAX_ENUM, CommandBuffer, PipelineBindPoint, Layout, FirstSet, DescriptorSetCount, DescriptorSets, DynamicOffsetCount, DynamicOffsets);
 		VULKANAPINAMESPACE::vkCmdBindDescriptorSets(CommandBuffer, PipelineBindPoint, Layout, FirstSet, DescriptorSetCount, DescriptorSets, DynamicOffsetCount, DynamicOffsets);
-		FWrapLayer::BindDescriptorSets(VK_SUCCESS, CommandBuffer, PipelineBindPoint, Layout, FirstSet, DescriptorSetCount, DescriptorSets, DynamicOffsetCount, DynamicOffsets);
+		FWrapLayer::CmdBindDescriptorSets(VK_SUCCESS, CommandBuffer, PipelineBindPoint, Layout, FirstSet, DescriptorSetCount, DescriptorSets, DynamicOffsetCount, DynamicOffsets);
 	}
 
 	static FORCEINLINE_DEBUGGABLE void  vkCmdBindIndexBuffer(VkCommandBuffer CommandBuffer, VkBuffer Buffer, VkDeviceSize Offset, VkIndexType IndexType)
 	{
-		FWrapLayer::BindIndexBuffer(VK_RESULT_MAX_ENUM, CommandBuffer, Buffer, Offset, IndexType);
+		FWrapLayer::CmdBindIndexBuffer(VK_RESULT_MAX_ENUM, CommandBuffer, Buffer, Offset, IndexType);
 		VULKANAPINAMESPACE::vkCmdBindIndexBuffer(CommandBuffer, Buffer, Offset, IndexType);
-		FWrapLayer::BindIndexBuffer(VK_SUCCESS, CommandBuffer, Buffer, Offset, IndexType);
+		FWrapLayer::CmdBindIndexBuffer(VK_SUCCESS, CommandBuffer, Buffer, Offset, IndexType);
 	}
 
 	static FORCEINLINE_DEBUGGABLE void  vkCmdBindVertexBuffers(VkCommandBuffer CommandBuffer, uint32 FirstBinding, uint32 BindingCount, const VkBuffer* Buffers, const VkDeviceSize* Offsets)
 	{
-		FWrapLayer::BindVertexBuffers(VK_RESULT_MAX_ENUM, CommandBuffer, FirstBinding, BindingCount, Buffers, Offsets);
+		FWrapLayer::CmdBindVertexBuffers(VK_RESULT_MAX_ENUM, CommandBuffer, FirstBinding, BindingCount, Buffers, Offsets);
 		VULKANAPINAMESPACE::vkCmdBindVertexBuffers(CommandBuffer, FirstBinding, BindingCount, Buffers, Offsets);
-		FWrapLayer::BindVertexBuffers(VK_SUCCESS, CommandBuffer, FirstBinding, BindingCount, Buffers, Offsets);
+		FWrapLayer::CmdBindVertexBuffers(VK_SUCCESS, CommandBuffer, FirstBinding, BindingCount, Buffers, Offsets);
 	}
 
 	static FORCEINLINE_DEBUGGABLE void  vkCmdDraw(VkCommandBuffer CommandBuffer, uint32 VertexCount, uint32 InstanceCount, uint32 FirstVertex, uint32 FirstInstance)
 	{
-		FWrapLayer::Draw(VK_RESULT_MAX_ENUM, CommandBuffer, VertexCount, InstanceCount, FirstVertex, FirstInstance);
+		FWrapLayer::CmdDraw(VK_RESULT_MAX_ENUM, CommandBuffer, VertexCount, InstanceCount, FirstVertex, FirstInstance);
 		VULKANAPINAMESPACE::vkCmdDraw(CommandBuffer, VertexCount, InstanceCount, FirstVertex, FirstInstance);
-		FWrapLayer::Draw(VK_SUCCESS, CommandBuffer, VertexCount, InstanceCount, FirstVertex, FirstInstance);
+		FWrapLayer::CmdDraw(VK_SUCCESS, CommandBuffer, VertexCount, InstanceCount, FirstVertex, FirstInstance);
 	}
 
 	static FORCEINLINE_DEBUGGABLE void  vkCmdDrawIndexed(VkCommandBuffer CommandBuffer, uint32 IndexCount, uint32 InstanceCount, uint32 FirstIndex, int32_t VertexOffset, uint32 FirstInstance)
 	{
-		FWrapLayer::DrawIndexed(VK_RESULT_MAX_ENUM, CommandBuffer, IndexCount, InstanceCount, FirstIndex, VertexOffset, FirstInstance);
+		FWrapLayer::CmdDrawIndexed(VK_RESULT_MAX_ENUM, CommandBuffer, IndexCount, InstanceCount, FirstIndex, VertexOffset, FirstInstance);
 		VULKANAPINAMESPACE::vkCmdDrawIndexed(CommandBuffer, IndexCount, InstanceCount, FirstIndex, VertexOffset, FirstInstance);
-		FWrapLayer::DrawIndexed(VK_SUCCESS, CommandBuffer, IndexCount, InstanceCount, FirstIndex, VertexOffset, FirstInstance);
+		FWrapLayer::CmdDrawIndexed(VK_SUCCESS, CommandBuffer, IndexCount, InstanceCount, FirstIndex, VertexOffset, FirstInstance);
 	}
 
 	static FORCEINLINE_DEBUGGABLE void  vkCmdDrawIndirect(VkCommandBuffer CommandBuffer, VkBuffer Buffer, VkDeviceSize Offset, uint32 DrawCount, uint32 Stride)
 	{
-		FWrapLayer::DrawIndirect(VK_RESULT_MAX_ENUM, CommandBuffer, Buffer, Offset, DrawCount, Stride);
+		FWrapLayer::CmdDrawIndirect(VK_RESULT_MAX_ENUM, CommandBuffer, Buffer, Offset, DrawCount, Stride);
 		VULKANAPINAMESPACE::vkCmdDrawIndirect(CommandBuffer, Buffer, Offset, DrawCount, Stride);
-		FWrapLayer::DrawIndirect(VK_SUCCESS, CommandBuffer, Buffer, Offset, DrawCount, Stride);
+		FWrapLayer::CmdDrawIndirect(VK_SUCCESS, CommandBuffer, Buffer, Offset, DrawCount, Stride);
 	}
 
 	static FORCEINLINE_DEBUGGABLE void  vkCmdDrawIndexedIndirect(VkCommandBuffer CommandBuffer, VkBuffer Buffer, VkDeviceSize Offset, uint32 DrawCount, uint32 Stride)
 	{
-		FWrapLayer::DrawIndexedIndirect(VK_RESULT_MAX_ENUM, CommandBuffer, Buffer, Offset, DrawCount, Stride);
+		FWrapLayer::CmdDrawIndexedIndirect(VK_RESULT_MAX_ENUM, CommandBuffer, Buffer, Offset, DrawCount, Stride);
 		VULKANAPINAMESPACE::vkCmdDrawIndexedIndirect(CommandBuffer, Buffer, Offset, DrawCount, Stride);
-		FWrapLayer::DrawIndexedIndirect(VK_SUCCESS, CommandBuffer, Buffer, Offset, DrawCount, Stride);
+		FWrapLayer::CmdDrawIndexedIndirect(VK_SUCCESS, CommandBuffer, Buffer, Offset, DrawCount, Stride);
 	}
 
 	static FORCEINLINE_DEBUGGABLE void  vkCmdDispatch(VkCommandBuffer CommandBuffer, uint32 X, uint32 Y, uint32 Z)
 	{
-		FWrapLayer::Dispatch(VK_RESULT_MAX_ENUM, CommandBuffer, X, Y, Z);
+		FWrapLayer::CmdDispatch(VK_RESULT_MAX_ENUM, CommandBuffer, X, Y, Z);
 		VULKANAPINAMESPACE::vkCmdDispatch(CommandBuffer, X, Y, Z);
-		FWrapLayer::Dispatch(VK_SUCCESS, CommandBuffer, X, Y, Z);
+		FWrapLayer::CmdDispatch(VK_SUCCESS, CommandBuffer, X, Y, Z);
 	}
 
 	static FORCEINLINE_DEBUGGABLE void  vkCmdDispatchIndirect(VkCommandBuffer CommandBuffer, VkBuffer Buffer, VkDeviceSize Offset)
 	{
-		FWrapLayer::DispatchIndirect(VK_RESULT_MAX_ENUM, CommandBuffer, Buffer, Offset);
+		FWrapLayer::CmdDispatchIndirect(VK_RESULT_MAX_ENUM, CommandBuffer, Buffer, Offset);
 		VULKANAPINAMESPACE::vkCmdDispatchIndirect(CommandBuffer, Buffer, Offset);
-		FWrapLayer::DispatchIndirect(VK_SUCCESS, CommandBuffer, Buffer, Offset);
+		FWrapLayer::CmdDispatchIndirect(VK_SUCCESS, CommandBuffer, Buffer, Offset);
 	}
 
 	static FORCEINLINE_DEBUGGABLE void  vkCmdCopyBuffer(VkCommandBuffer CommandBuffer, VkBuffer SrcBuffer, VkBuffer DstBuffer, uint32 RegionCount, const VkBufferCopy* Regions)
 	{
-		FWrapLayer::CopyBuffer(VK_RESULT_MAX_ENUM, CommandBuffer, SrcBuffer, DstBuffer, RegionCount, Regions);
+		FWrapLayer::CmdCopyBuffer(VK_RESULT_MAX_ENUM, CommandBuffer, SrcBuffer, DstBuffer, RegionCount, Regions);
 		VULKANAPINAMESPACE::vkCmdCopyBuffer(CommandBuffer, SrcBuffer, DstBuffer, RegionCount, Regions);
-		FWrapLayer::CopyBuffer(VK_SUCCESS, CommandBuffer, SrcBuffer, DstBuffer, RegionCount, Regions);
+		FWrapLayer::CmdCopyBuffer(VK_SUCCESS, CommandBuffer, SrcBuffer, DstBuffer, RegionCount, Regions);
 	}
 
 	static FORCEINLINE_DEBUGGABLE void  vkCmdCopyImage(VkCommandBuffer CommandBuffer, VkImage SrcImage, VkImageLayout SrcImageLayout, VkImage DstImage, VkImageLayout DstImageLayout, uint32 RegionCount, const VkImageCopy* Regions)
 	{
-		FWrapLayer::CopyImage(VK_RESULT_MAX_ENUM, CommandBuffer, SrcImage, SrcImageLayout, DstImage, DstImageLayout, RegionCount, Regions);
+		FWrapLayer::CmdCopyImage(VK_RESULT_MAX_ENUM, CommandBuffer, SrcImage, SrcImageLayout, DstImage, DstImageLayout, RegionCount, Regions);
 		VULKANAPINAMESPACE::vkCmdCopyImage(CommandBuffer, SrcImage, SrcImageLayout, DstImage, DstImageLayout, RegionCount, Regions);
-		FWrapLayer::CopyImage(VK_SUCCESS, CommandBuffer, SrcImage, SrcImageLayout, DstImage, DstImageLayout, RegionCount, Regions);
+		FWrapLayer::CmdCopyImage(VK_SUCCESS, CommandBuffer, SrcImage, SrcImageLayout, DstImage, DstImageLayout, RegionCount, Regions);
 	}
 
 	static FORCEINLINE_DEBUGGABLE void  vkCmdBlitImage(VkCommandBuffer CommandBuffer, VkImage SrcImage, VkImageLayout SrcImageLayout, VkImage DstImage, VkImageLayout DstImageLayout, uint32 RegionCount, const VkImageBlit* Regions, VkFilter Filter)
 	{
-		FWrapLayer::BlitImage(VK_RESULT_MAX_ENUM, CommandBuffer, SrcImage, SrcImageLayout, DstImage, DstImageLayout, RegionCount, Regions, Filter);
+		FWrapLayer::CmdBlitImage(VK_RESULT_MAX_ENUM, CommandBuffer, SrcImage, SrcImageLayout, DstImage, DstImageLayout, RegionCount, Regions, Filter);
 		VULKANAPINAMESPACE::vkCmdBlitImage(CommandBuffer, SrcImage, SrcImageLayout, DstImage, DstImageLayout, RegionCount, Regions, Filter);
-		FWrapLayer::BlitImage(VK_SUCCESS, CommandBuffer, SrcImage, SrcImageLayout, DstImage, DstImageLayout, RegionCount, Regions, Filter);
+		FWrapLayer::CmdBlitImage(VK_SUCCESS, CommandBuffer, SrcImage, SrcImageLayout, DstImage, DstImageLayout, RegionCount, Regions, Filter);
 	}
 
 	static FORCEINLINE_DEBUGGABLE void  vkCmdCopyBufferToImage(VkCommandBuffer CommandBuffer, VkBuffer SrcBuffer, VkImage DstImage, VkImageLayout DstImageLayout, uint32 RegionCount, const VkBufferImageCopy* Regions)
 	{
-		FWrapLayer::CopyBufferToImage(VK_RESULT_MAX_ENUM, CommandBuffer, SrcBuffer, DstImage, DstImageLayout, RegionCount, Regions);
+		FWrapLayer::CmdCopyBufferToImage(VK_RESULT_MAX_ENUM, CommandBuffer, SrcBuffer, DstImage, DstImageLayout, RegionCount, Regions);
 		VULKANAPINAMESPACE::vkCmdCopyBufferToImage(CommandBuffer, SrcBuffer, DstImage, DstImageLayout, RegionCount, Regions);
-		FWrapLayer::CopyBufferToImage(VK_SUCCESS, CommandBuffer, SrcBuffer, DstImage, DstImageLayout, RegionCount, Regions);
+		FWrapLayer::CmdCopyBufferToImage(VK_SUCCESS, CommandBuffer, SrcBuffer, DstImage, DstImageLayout, RegionCount, Regions);
 	}
 
 	static FORCEINLINE_DEBUGGABLE void  vkCmdCopyImageToBuffer(VkCommandBuffer CommandBuffer, VkImage SrcImage, VkImageLayout SrcImageLayout, VkBuffer DstBuffer, uint32 RegionCount, const VkBufferImageCopy* Regions)
 	{
-		FWrapLayer::CopyImageToBuffer(VK_RESULT_MAX_ENUM, CommandBuffer, SrcImage, SrcImageLayout, DstBuffer, RegionCount, Regions);
+		FWrapLayer::CmdCopyImageToBuffer(VK_RESULT_MAX_ENUM, CommandBuffer, SrcImage, SrcImageLayout, DstBuffer, RegionCount, Regions);
 		VULKANAPINAMESPACE::vkCmdCopyImageToBuffer(CommandBuffer, SrcImage, SrcImageLayout, DstBuffer, RegionCount, Regions);
-		FWrapLayer::CopyImageToBuffer(VK_SUCCESS, CommandBuffer, SrcImage, SrcImageLayout, DstBuffer, RegionCount, Regions);
+		FWrapLayer::CmdCopyImageToBuffer(VK_SUCCESS, CommandBuffer, SrcImage, SrcImageLayout, DstBuffer, RegionCount, Regions);
 	}
 
 	static FORCEINLINE_DEBUGGABLE void  vkCmdUpdateBuffer(VkCommandBuffer CommandBuffer, VkBuffer DstBuffer, VkDeviceSize DstOffset, VkDeviceSize DataSize, const void* pData)
 	{
-		FWrapLayer::UpdateBuffer(VK_RESULT_MAX_ENUM, CommandBuffer, DstBuffer, DstOffset, DataSize, pData);
+		FWrapLayer::CmdUpdateBuffer(VK_RESULT_MAX_ENUM, CommandBuffer, DstBuffer, DstOffset, DataSize, pData);
 		VULKANAPINAMESPACE::vkCmdUpdateBuffer(CommandBuffer, DstBuffer, DstOffset, DataSize, pData);
-		FWrapLayer::UpdateBuffer(VK_SUCCESS, CommandBuffer, DstBuffer, DstOffset, DataSize, pData);
+		FWrapLayer::CmdUpdateBuffer(VK_SUCCESS, CommandBuffer, DstBuffer, DstOffset, DataSize, pData);
 	}
 
 	static FORCEINLINE_DEBUGGABLE void  vkCmdFillBuffer(VkCommandBuffer CommandBuffer, VkBuffer DstBuffer, VkDeviceSize DstOffset, VkDeviceSize Size, uint32 Data)
 	{
-		FWrapLayer::FillBuffer(VK_RESULT_MAX_ENUM, CommandBuffer, DstBuffer, DstOffset, Size, Data);
+		FWrapLayer::CmdFillBuffer(VK_RESULT_MAX_ENUM, CommandBuffer, DstBuffer, DstOffset, Size, Data);
 		VULKANAPINAMESPACE::vkCmdFillBuffer(CommandBuffer, DstBuffer, DstOffset, Size, Data);
-		FWrapLayer::FillBuffer(VK_SUCCESS, CommandBuffer, DstBuffer, DstOffset, Size, Data);
+		FWrapLayer::CmdFillBuffer(VK_SUCCESS, CommandBuffer, DstBuffer, DstOffset, Size, Data);
 	}
 
 	static FORCEINLINE_DEBUGGABLE void  vkCmdClearColorImage(VkCommandBuffer CommandBuffer, VkImage Image, VkImageLayout ImageLayout, const VkClearColorValue* Color, uint32 RangeCount, const VkImageSubresourceRange* Ranges)
 	{
-		FWrapLayer::ClearColorImage(VK_RESULT_MAX_ENUM, CommandBuffer, Image, ImageLayout, Color, RangeCount, Ranges);
+		FWrapLayer::CmdClearColorImage(VK_RESULT_MAX_ENUM, CommandBuffer, Image, ImageLayout, Color, RangeCount, Ranges);
 		VULKANAPINAMESPACE::vkCmdClearColorImage(CommandBuffer, Image, ImageLayout, Color, RangeCount, Ranges);
-		FWrapLayer::ClearColorImage(VK_SUCCESS, CommandBuffer, Image, ImageLayout, Color, RangeCount, Ranges);
+		FWrapLayer::CmdClearColorImage(VK_SUCCESS, CommandBuffer, Image, ImageLayout, Color, RangeCount, Ranges);
 	}
 
 	static FORCEINLINE_DEBUGGABLE void  vkCmdClearDepthStencilImage(VkCommandBuffer CommandBuffer, VkImage Image, VkImageLayout ImageLayout, const VkClearDepthStencilValue* DepthStencil, uint32 RangeCount, const VkImageSubresourceRange* Ranges)
 	{
-		FWrapLayer::ClearDepthStencilImage(VK_RESULT_MAX_ENUM, CommandBuffer, Image, ImageLayout, DepthStencil, RangeCount, Ranges);
+		FWrapLayer::CmdClearDepthStencilImage(VK_RESULT_MAX_ENUM, CommandBuffer, Image, ImageLayout, DepthStencil, RangeCount, Ranges);
 		VULKANAPINAMESPACE::vkCmdClearDepthStencilImage(CommandBuffer, Image, ImageLayout, DepthStencil, RangeCount, Ranges);
-		FWrapLayer::ClearDepthStencilImage(VK_SUCCESS, CommandBuffer, Image, ImageLayout, DepthStencil, RangeCount, Ranges);
+		FWrapLayer::CmdClearDepthStencilImage(VK_SUCCESS, CommandBuffer, Image, ImageLayout, DepthStencil, RangeCount, Ranges);
 	}
 
 	static FORCEINLINE_DEBUGGABLE void  vkCmdClearAttachments(VkCommandBuffer CommandBuffer, uint32 AttachmentCount, const VkClearAttachment* Attachments, uint32 RectCount, const VkClearRect* Rects)
 	{
-		FWrapLayer::ClearAttachments(VK_RESULT_MAX_ENUM, CommandBuffer, AttachmentCount, Attachments, RectCount, Rects);
+		FWrapLayer::CmdClearAttachments(VK_RESULT_MAX_ENUM, CommandBuffer, AttachmentCount, Attachments, RectCount, Rects);
 		VULKANAPINAMESPACE::vkCmdClearAttachments(CommandBuffer, AttachmentCount, Attachments, RectCount, Rects);
-		FWrapLayer::ClearAttachments(VK_SUCCESS, CommandBuffer, AttachmentCount, Attachments, RectCount, Rects);
+		FWrapLayer::CmdClearAttachments(VK_SUCCESS, CommandBuffer, AttachmentCount, Attachments, RectCount, Rects);
 	}
 
 	static FORCEINLINE_DEBUGGABLE void  vkCmdResolveImage(
@@ -1111,23 +1124,23 @@ namespace VulkanRHI
 		VkImage DstImage, VkImageLayout DstImageLayout,
 		uint32 RegionCount, const VkImageResolve* Regions)
 	{
-		FWrapLayer::ResolveImage(VK_RESULT_MAX_ENUM, CommandBuffer, SrcImage, SrcImageLayout, DstImage, DstImageLayout, RegionCount, Regions);
+		FWrapLayer::CmdResolveImage(VK_RESULT_MAX_ENUM, CommandBuffer, SrcImage, SrcImageLayout, DstImage, DstImageLayout, RegionCount, Regions);
 		VULKANAPINAMESPACE::vkCmdResolveImage(CommandBuffer, SrcImage, SrcImageLayout, DstImage, DstImageLayout, RegionCount, Regions);
-		FWrapLayer::ResolveImage(VK_SUCCESS, CommandBuffer, SrcImage, SrcImageLayout, DstImage, DstImageLayout, RegionCount, Regions);
+		FWrapLayer::CmdResolveImage(VK_SUCCESS, CommandBuffer, SrcImage, SrcImageLayout, DstImage, DstImageLayout, RegionCount, Regions);
 	}
 
 	static FORCEINLINE_DEBUGGABLE void  vkCmdSetEvent(VkCommandBuffer CommandBuffer, VkEvent Event, VkPipelineStageFlags StageMask)
 	{
-		FWrapLayer::SetEvent(VK_RESULT_MAX_ENUM, CommandBuffer, Event, StageMask);
+		FWrapLayer::CmdSetEvent(VK_RESULT_MAX_ENUM, CommandBuffer, Event, StageMask);
 		VULKANAPINAMESPACE::vkCmdSetEvent(CommandBuffer, Event, StageMask);
-		FWrapLayer::SetEvent(VK_SUCCESS, CommandBuffer, Event, StageMask);
+		FWrapLayer::CmdSetEvent(VK_SUCCESS, CommandBuffer, Event, StageMask);
 	}
 
 	static FORCEINLINE_DEBUGGABLE void  vkCmdResetEvent(VkCommandBuffer CommandBuffer, VkEvent Event, VkPipelineStageFlags StageMask)
 	{
-		FWrapLayer::ResetEvent(VK_RESULT_MAX_ENUM, CommandBuffer, Event, StageMask);
+		FWrapLayer::CmdResetEvent(VK_RESULT_MAX_ENUM, CommandBuffer, Event, StageMask);
 		VULKANAPINAMESPACE::vkCmdResetEvent(CommandBuffer, Event, StageMask);
-		FWrapLayer::ResetEvent(VK_SUCCESS, CommandBuffer, Event, StageMask);
+		FWrapLayer::CmdResetEvent(VK_SUCCESS, CommandBuffer, Event, StageMask);
 	}
 
 	static FORCEINLINE_DEBUGGABLE void  vkCmdWaitEvents(VkCommandBuffer CommandBuffer, uint32 EventCount, const VkEvent* Events,
@@ -1136,11 +1149,11 @@ namespace VulkanRHI
 		uint32 BufferMemoryBarrierCount, const VkBufferMemoryBarrier* pBufferMemoryBarriers,
 		uint32 ImageMemoryBarrierCount, const VkImageMemoryBarrier* pImageMemoryBarriers)
 	{
-		FWrapLayer::WaitEvents(VK_RESULT_MAX_ENUM, CommandBuffer, EventCount, Events, SrcStageMask, DstStageMask,
+		FWrapLayer::CmdWaitEvents(VK_RESULT_MAX_ENUM, CommandBuffer, EventCount, Events, SrcStageMask, DstStageMask,
 			MemoryBarrierCount, pMemoryBarriers, BufferMemoryBarrierCount, pBufferMemoryBarriers, ImageMemoryBarrierCount, pImageMemoryBarriers);
 		VULKANAPINAMESPACE::vkCmdWaitEvents(CommandBuffer, EventCount, Events, SrcStageMask, DstStageMask, MemoryBarrierCount, pMemoryBarriers,
 			BufferMemoryBarrierCount, pBufferMemoryBarriers, ImageMemoryBarrierCount, pImageMemoryBarriers);
-		FWrapLayer::WaitEvents(VK_SUCCESS, CommandBuffer, EventCount, Events, SrcStageMask, DstStageMask,
+		FWrapLayer::CmdWaitEvents(VK_SUCCESS, CommandBuffer, EventCount, Events, SrcStageMask, DstStageMask,
 			MemoryBarrierCount, pMemoryBarriers, BufferMemoryBarrierCount, pBufferMemoryBarriers, ImageMemoryBarrierCount, pImageMemoryBarriers);
 	}
 
@@ -1150,85 +1163,81 @@ namespace VulkanRHI
 		uint32 BufferMemoryBarrierCount, const VkBufferMemoryBarrier* BufferMemoryBarriers,
 		uint32 ImageMemoryBarrierCount, const VkImageMemoryBarrier* ImageMemoryBarriers)
 	{
-		FWrapLayer::PipelineBarrier(VK_RESULT_MAX_ENUM, CommandBuffer, SrcStageMask, DstStageMask, DependencyFlags, MemoryBarrierCount, MemoryBarriers, BufferMemoryBarrierCount, BufferMemoryBarriers, ImageMemoryBarrierCount, ImageMemoryBarriers);
+		FWrapLayer::CmdPipelineBarrier(VK_RESULT_MAX_ENUM, CommandBuffer, SrcStageMask, DstStageMask, DependencyFlags, MemoryBarrierCount, MemoryBarriers, BufferMemoryBarrierCount, BufferMemoryBarriers, ImageMemoryBarrierCount, ImageMemoryBarriers);
 		VULKANAPINAMESPACE::vkCmdPipelineBarrier(CommandBuffer, SrcStageMask, DstStageMask, DependencyFlags, MemoryBarrierCount, MemoryBarriers, BufferMemoryBarrierCount, BufferMemoryBarriers, ImageMemoryBarrierCount, ImageMemoryBarriers);
-		FWrapLayer::PipelineBarrier(VK_SUCCESS, CommandBuffer, SrcStageMask, DstStageMask, DependencyFlags, MemoryBarrierCount, MemoryBarriers, BufferMemoryBarrierCount, BufferMemoryBarriers, ImageMemoryBarrierCount, ImageMemoryBarriers);
+		FWrapLayer::CmdPipelineBarrier(VK_SUCCESS, CommandBuffer, SrcStageMask, DstStageMask, DependencyFlags, MemoryBarrierCount, MemoryBarriers, BufferMemoryBarrierCount, BufferMemoryBarriers, ImageMemoryBarrierCount, ImageMemoryBarriers);
 	}
 
 	static FORCEINLINE_DEBUGGABLE void  vkCmdBeginQuery(VkCommandBuffer CommandBuffer, VkQueryPool QueryPool, uint32 Query, VkQueryControlFlags Flags)
 	{
-		FWrapLayer::BeginQuery(VK_RESULT_MAX_ENUM, CommandBuffer, QueryPool, Query, Flags);
+		FWrapLayer::CmdBeginQuery(VK_RESULT_MAX_ENUM, CommandBuffer, QueryPool, Query, Flags);
 		VULKANAPINAMESPACE::vkCmdBeginQuery(CommandBuffer, QueryPool, Query, Flags);
-		FWrapLayer::BeginQuery(VK_SUCCESS, CommandBuffer, QueryPool, Query, Flags);
+		FWrapLayer::CmdBeginQuery(VK_SUCCESS, CommandBuffer, QueryPool, Query, Flags);
 	}
 
 	static FORCEINLINE_DEBUGGABLE void  vkCmdEndQuery(VkCommandBuffer CommandBuffer, VkQueryPool QueryPool, uint32 Query)
 	{
-		FWrapLayer::EndQuery(VK_RESULT_MAX_ENUM, CommandBuffer, QueryPool, Query);
+		FWrapLayer::CmdEndQuery(VK_RESULT_MAX_ENUM, CommandBuffer, QueryPool, Query);
 		VULKANAPINAMESPACE::vkCmdEndQuery(CommandBuffer, QueryPool, Query);
-		FWrapLayer::EndQuery(VK_SUCCESS, CommandBuffer, QueryPool, Query);
+		FWrapLayer::CmdEndQuery(VK_SUCCESS, CommandBuffer, QueryPool, Query);
 	}
 
 	static FORCEINLINE_DEBUGGABLE void  vkCmdResetQueryPool(VkCommandBuffer CommandBuffer, VkQueryPool QueryPool, uint32 FirstQuery, uint32 QueryCount)
 	{
-		FWrapLayer::ResetQueryPool(VK_RESULT_MAX_ENUM, CommandBuffer, QueryPool, FirstQuery, QueryCount);
+		FWrapLayer::CmdResetQueryPool(VK_RESULT_MAX_ENUM, CommandBuffer, QueryPool, FirstQuery, QueryCount);
 		VULKANAPINAMESPACE::vkCmdResetQueryPool(CommandBuffer, QueryPool, FirstQuery, QueryCount);
-		FWrapLayer::ResetQueryPool(VK_SUCCESS, CommandBuffer, QueryPool, FirstQuery, QueryCount);
+		FWrapLayer::CmdResetQueryPool(VK_SUCCESS, CommandBuffer, QueryPool, FirstQuery, QueryCount);
 	}
 
 	static FORCEINLINE_DEBUGGABLE void  vkCmdWriteTimestamp(VkCommandBuffer CommandBuffer, VkPipelineStageFlagBits PipelineStage, VkQueryPool QueryPool, uint32 Query)
 	{
-		FWrapLayer::WriteTimestamp(VK_RESULT_MAX_ENUM, CommandBuffer, PipelineStage, QueryPool, Query);
+		FWrapLayer::CmdWriteTimestamp(VK_RESULT_MAX_ENUM, CommandBuffer, PipelineStage, QueryPool, Query);
 		VULKANAPINAMESPACE::vkCmdWriteTimestamp(CommandBuffer, PipelineStage, QueryPool, Query);
-		FWrapLayer::WriteTimestamp(VK_SUCCESS, CommandBuffer, PipelineStage, QueryPool, Query);
+		FWrapLayer::CmdWriteTimestamp(VK_SUCCESS, CommandBuffer, PipelineStage, QueryPool, Query);
 	}
 
 	static FORCEINLINE_DEBUGGABLE void  vkCmdCopyQueryPoolResults(VkCommandBuffer CommandBuffer, VkQueryPool QueryPool, uint32 FirstQuery, uint32 QueryCount,
 		VkBuffer DstBuffer, VkDeviceSize DstOffset, VkDeviceSize Stride, VkQueryResultFlags Flags)
 	{
-		FWrapLayer::CopyQueryPoolResults(VK_RESULT_MAX_ENUM, CommandBuffer, QueryPool, FirstQuery, QueryCount, DstBuffer, DstOffset, Stride, Flags);
+		FWrapLayer::CmdCopyQueryPoolResults(VK_RESULT_MAX_ENUM, CommandBuffer, QueryPool, FirstQuery, QueryCount, DstBuffer, DstOffset, Stride, Flags);
 		VULKANAPINAMESPACE::vkCmdCopyQueryPoolResults(CommandBuffer, QueryPool, FirstQuery, QueryCount, DstBuffer, DstOffset, Stride, Flags);
-		FWrapLayer::CopyQueryPoolResults(VK_SUCCESS, CommandBuffer, QueryPool, FirstQuery, QueryCount, DstBuffer, DstOffset, Stride, Flags);
-	}
-#if 0
-	static FORCEINLINE_DEBUGGABLE void  vkCmdPushConstants(
-		VkCommandBuffer                             commandBuffer,
-		VkPipelineLayout                            layout,
-		VkShaderStageFlags                          stageFlags,
-		uint32                                    offset,
-		uint32                                    size,
-		const void*                                 pValues);
-#endif
-	static FORCEINLINE_DEBUGGABLE void  vkCmdBeginRenderPass(
-		VkCommandBuffer CommandBuffer,
-		const VkRenderPassBeginInfo* RenderPassBegin,
-		VkSubpassContents Contents)
+		FWrapLayer::CmdCopyQueryPoolResults(VK_SUCCESS, CommandBuffer, QueryPool, FirstQuery, QueryCount, DstBuffer, DstOffset, Stride, Flags);
+	}
+
+	static FORCEINLINE_DEBUGGABLE void  vkCmdPushConstants(VkCommandBuffer CommandBuffer, VkPipelineLayout Layout, VkShaderStageFlags StageFlags, uint32 Offset, uint32 Size, const void* pValues)
+	{
+		FWrapLayer::CmdPushConstants(VK_RESULT_MAX_ENUM, CommandBuffer, Layout, StageFlags, Offset, Size, pValues);
+		VULKANAPINAMESPACE::vkCmdPushConstants(CommandBuffer, Layout, StageFlags, Offset, Size, pValues);
+		FWrapLayer::CmdPushConstants(VK_SUCCESS, CommandBuffer, Layout, StageFlags, Offset, Size, pValues);
+	}
+
+	static FORCEINLINE_DEBUGGABLE void  vkCmdBeginRenderPass(VkCommandBuffer CommandBuffer, const VkRenderPassBeginInfo* RenderPassBegin, VkSubpassContents Contents)
 	{
-		FWrapLayer::BeginRenderPass(VK_RESULT_MAX_ENUM, CommandBuffer, RenderPassBegin, Contents);
+		FWrapLayer::CmdBeginRenderPass(VK_RESULT_MAX_ENUM, CommandBuffer, RenderPassBegin, Contents);
 		VULKANAPINAMESPACE::vkCmdBeginRenderPass(CommandBuffer, RenderPassBegin, Contents);
-		FWrapLayer::BeginRenderPass(VK_SUCCESS, CommandBuffer, RenderPassBegin, Contents);
+		FWrapLayer::CmdBeginRenderPass(VK_SUCCESS, CommandBuffer, RenderPassBegin, Contents);
 	}
 
 	static FORCEINLINE_DEBUGGABLE void  vkCmdNextSubpass(VkCommandBuffer CommandBuffer, VkSubpassContents Contents)
 	{
-		FWrapLayer::NextSubpass(VK_RESULT_MAX_ENUM, CommandBuffer, Contents);
+		FWrapLayer::CmdNextSubpass(VK_RESULT_MAX_ENUM, CommandBuffer, Contents);
 		VULKANAPINAMESPACE::vkCmdNextSubpass(CommandBuffer, Contents);
-		FWrapLayer::NextSubpass(VK_SUCCESS, CommandBuffer, Contents);
+		FWrapLayer::CmdNextSubpass(VK_SUCCESS, CommandBuffer, Contents);
 	}
 
 	static FORCEINLINE_DEBUGGABLE void  vkCmdEndRenderPass(VkCommandBuffer CommandBuffer)
 	{
-		FWrapLayer::EndRenderPass(VK_RESULT_MAX_ENUM, CommandBuffer);
+		FWrapLayer::CmdEndRenderPass(VK_RESULT_MAX_ENUM, CommandBuffer);
 		VULKANAPINAMESPACE::vkCmdEndRenderPass(CommandBuffer);
-		FWrapLayer::EndRenderPass(VK_SUCCESS, CommandBuffer);
+		FWrapLayer::CmdEndRenderPass(VK_SUCCESS, CommandBuffer);
 	}
 
-#if 0
-	static FORCEINLINE_DEBUGGABLE void  vkCmdExecuteCommands(
-		VkCommandBuffer                             commandBuffer,
-		uint32                                    commandBufferCount,
-		const VkCommandBuffer*                      pCommandBuffers);
-#endif
+	static FORCEINLINE_DEBUGGABLE void  vkCmdExecuteCommands(VkCommandBuffer CommandBuffer, uint32 CommandBufferCount, const VkCommandBuffer* pCommandBuffers)
+	{
+		FWrapLayer::CmdExecuteCommands(VK_RESULT_MAX_ENUM, CommandBuffer, CommandBufferCount, pCommandBuffers);
+		VULKANAPINAMESPACE::vkCmdExecuteCommands(CommandBuffer, CommandBufferCount, pCommandBuffers);
+		FWrapLayer::CmdExecuteCommands(VK_SUCCESS, CommandBuffer, CommandBufferCount, pCommandBuffers);
+	}
 
 	static FORCEINLINE_DEBUGGABLE VkResult vkCreateSwapchainKHR(VkDevice Device, const VkSwapchainCreateInfoKHR* CreateInfo, const VkAllocationCallbacks* Allocator, VkSwapchainKHR* Swapchain)
 	{
diff --git a/Engine/Source/Runtime/VulkanRHI/Private/VulkanDebug.cpp b/Engine/Source/Runtime/VulkanRHI/Private/VulkanDebug.cpp
--- a/Engine/Source/Runtime/VulkanRHI/Private/VulkanDebug.cpp
+++ b/Engine/Source/Runtime/VulkanRHI/Private/VulkanDebug.cpp
@@ -877,7 +877,7 @@ namespace VulkanRHI
 		return FString::Printf(TEXT("Unknown VkResult %d"), (int32)Result);
 	}
 
-	FString GetImageTilingString(VkImageTiling Tiling)
+	static FString GetVkImageTilingString(VkImageTiling Tiling)
 	{
 		switch (Tiling)
 		{
@@ -893,7 +893,7 @@ namespace VulkanRHI
 		return FString::Printf(TEXT("Unknown VkImageTiling %d"), (int32)Tiling);
 	}
 
-	static FString GetImageLayoutString(VkImageLayout Layout)
+	static FString GetVkImageLayoutString(VkImageLayout Layout)
 	{
 		switch (Layout)
 		{
@@ -917,7 +917,7 @@ namespace VulkanRHI
 		return FString::Printf(TEXT("Unknown VkImageLayout %d"), (int32)Layout);
 	}
 
-	static FString GetImageViewTypeString(VkImageViewType Type)
+	static FString GetVkImageViewTypeString(VkImageViewType Type)
 	{
 		switch (Type)
 		{
@@ -938,7 +938,7 @@ namespace VulkanRHI
 		return FString::Printf(TEXT("Unknown VkImageViewType %d"), (int32)Type);
 	}
 
-	static FString GetImageTypeString(VkImageType Type)
+	static FString GetVkImageTypeString(VkImageType Type)
 	{
 		switch (Type)
 		{
@@ -955,7 +955,7 @@ namespace VulkanRHI
 		return FString::Printf(TEXT("Unknown VkImageType %d"), (int32)Type);
 	}
 
-	static FString GetDescriptorTypeString(VkDescriptorType Type)
+	static FString GetVkDescriptorTypeString(VkDescriptorType Type)
 	{
 		switch (Type)
 		{
@@ -980,7 +980,7 @@ namespace VulkanRHI
 		return FString::Printf(TEXT("Unknown VkDescriptorType %d"), (int32)Type);
 	}
 
-	static FString GetStencilOpString(VkStencilOp Op)
+	static FString GetVkStencilOpString(VkStencilOp Op)
 	{
 		switch (Op)
 		{
@@ -1002,7 +1002,7 @@ namespace VulkanRHI
 		return FString::Printf(TEXT("Unknown VkStencilOp %d"), (int32)Op);
 	}
 
-	static FString GetCompareOpString(VkCompareOp Op)
+	static FString GetVkCompareOpString(VkCompareOp Op)
 	{
 		switch (Op)
 		{
@@ -1024,7 +1024,7 @@ namespace VulkanRHI
 		return FString::Printf(TEXT("Unknown VkStencilOp %d"), (int32)Op);
 	}
 
-	static FString GetComponentMappingString(const VkComponentMapping& Mapping)
+	static FString GetVkComponentMappingString(const VkComponentMapping& Mapping)
 	{
 		auto GetSwizzle = [](VkComponentSwizzle Swizzle) -> const TCHAR*
 			{
@@ -1345,9 +1345,23 @@ void FWrapLayer::AllocateMemory(VkResult Result, VkDevice Device, const VkMemory
 }
 
 #if VULKAN_ENABLE_DUMP_LAYER
-static void DumpMemoryRequirements(VkMemoryRequirements* MemoryRequirements)
+static void DumpErrorInvalidPointer(const FString& PointerTypename)
 {
-	DebugLog += FString::Printf(TEXT(" -> Size=%d Align=%d MemTypeBits=0x%x\n"), (uint32)MemoryRequirements->size, (uint32)MemoryRequirements->alignment, MemoryRequirements->memoryTypeBits);
+	DebugLog += FString::Printf(TEXT(" -> ERROR: invalid pointer to %s\n"), *PointerTypename);
+}
+#endif
+
+#if VULKAN_ENABLE_DUMP_LAYER
+static void DumpMemoryRequirements(VkMemoryRequirements* pMemRequirements)
+{
+	if (pMemRequirements)
+	{
+		DebugLog += FString::Printf(TEXT(" -> Size=%llu Align=%llu MemTypeBits=0x%x\n"), pMemRequirements->size, pMemRequirements->alignment, pMemRequirements->memoryTypeBits);
+	}
+	else
+	{
+		DumpErrorInvalidPointer(TEXT("VkMemoryRequirements"));
+	}
 	FlushDebugWrapperLog();
 }
 #endif
@@ -1456,10 +1470,10 @@ void FWrapLayer::CreateImage(VkResult Result, VkDevice Device, const VkImageCrea
 	{
 #if VULKAN_ENABLE_DUMP_LAYER
 		DevicePrintfBegin(Device, FString::Printf(TEXT("vkCreateImage(Info=0x%p, OutImage=0x%p)"), CreateInfo, Image));
-		DebugLog += FString::Printf(TEXT("%sVkImageCreateInfo: Flags=%d, ImageType=%s, Format=%s, MipLevels=%d, ArrayLayers=%d, Samples=%s\n"), Tabs, CreateInfo->flags, *GetImageTypeString(CreateInfo->imageType),
+		DebugLog += FString::Printf(TEXT("%sVkImageCreateInfo: Flags=%d, ImageType=%s, Format=%s, MipLevels=%d, ArrayLayers=%d, Samples=%s\n"), Tabs, CreateInfo->flags, *GetVkImageTypeString(CreateInfo->imageType),
 			*GetVkFormatString(CreateInfo->format), CreateInfo->mipLevels, CreateInfo->arrayLayers, *GetSampleCountString(CreateInfo->samples));
 		DebugLog += FString::Printf(TEXT("%s\tExtent=(%s) Tiling=%s, Usage=%s, Initial=%s\n"), Tabs, *GetExtentString(CreateInfo->extent),
-			*GetImageTilingString(CreateInfo->tiling), *GetImageUsageString(CreateInfo->usage), *GetImageLayoutString(CreateInfo->initialLayout));
+			*GetVkImageTilingString(CreateInfo->tiling), *GetImageUsageString(CreateInfo->usage), *GetVkImageLayoutString(CreateInfo->initialLayout));
 		FlushDebugWrapperLog();
 #endif
 	}
@@ -1515,7 +1529,7 @@ void FWrapLayer::CreateImageView(VkResult Result, VkDevice Device, const VkImage
 #if VULKAN_ENABLE_DUMP_LAYER
 		DevicePrintfBegin(Device, FString::Printf(TEXT("vkCreateImageView(Info=0x%p, OutImageView=0x%p)"), CreateInfo, ImageView));
 		DebugLog += FString::Printf(TEXT("%sVkImageViewCreateInfo: Flags=%d, Image=0x%p, ViewType=%s, Format=%s, Components=%s\n"), Tabs, CreateInfo->flags, CreateInfo->image,
-			*GetImageViewTypeString(CreateInfo->viewType), *GetVkFormatString(CreateInfo->format), *GetComponentMappingString(CreateInfo->components));
+			*GetVkImageViewTypeString(CreateInfo->viewType), *GetVkFormatString(CreateInfo->format), *GetVkComponentMappingString(CreateInfo->components));
 		DebugLog += FString::Printf(TEXT("%s\tSubresourceRange=(%s)"), Tabs, *GetImageSubResourceRangeString(CreateInfo->subresourceRange));
 #endif
 	}
@@ -1675,13 +1689,13 @@ void FWrapLayer::FlushMappedMemoryRanges(VkResult Result, VkDevice Device, uint3
 }
 
 
-void FWrapLayer::ResolveImage(VkResult Result, VkCommandBuffer CommandBuffer, VkImage SrcImage, VkImageLayout SrcImageLayout, VkImage DstImage, VkImageLayout DstImageLayout, uint32 RegionCount, const VkImageResolve* Regions)
+void FWrapLayer::CmdResolveImage(VkResult Result, VkCommandBuffer CommandBuffer, VkImage SrcImage, VkImageLayout SrcImageLayout, VkImage DstImage, VkImageLayout DstImageLayout, uint32 RegionCount, const VkImageResolve* Regions)
 {
 	if (Result == VK_RESULT_MAX_ENUM)
 	{
 #if VULKAN_ENABLE_DUMP_LAYER
 		CmdPrintfBegin(CommandBuffer, FString::Printf(TEXT("vkCmdResolveImage(SrcImage=0x%p, SrcImageLayout=%s, DestImage=0x%p, DestImageLayout=%s, NumRegions=%d, Regions=0x%p)[...]"),
-			CommandBuffer, SrcImage, *GetImageLayoutString(SrcImageLayout), DstImage, *GetImageLayoutString(DstImageLayout), RegionCount, Regions));
+			CommandBuffer, SrcImage, *GetVkImageLayoutString(SrcImageLayout), DstImage, *GetVkImageLayoutString(DstImageLayout), RegionCount, Regions));
 		for (uint32 Index = 0; Index < RegionCount; ++Index)
 		{
 			DebugLog += Tabs;
@@ -1777,19 +1791,16 @@ void FWrapLayer::EnumeratePhysicalDevices(VkResult Result, VkInstance Instance,
 #if VULKAN_ENABLE_DUMP_LAYER
 static void DumpImageMemoryBarriers(uint32 ImageMemoryBarrierCount, const VkImageMemoryBarrier* ImageMemoryBarriers)
 {
-	if (ImageMemoryBarrierCount)
+	for (uint32 Index = 0; Index < ImageMemoryBarrierCount; ++Index)
 	{
-		for (uint32 Index = 0; Index < ImageMemoryBarrierCount; ++Index)
-		{
-			DebugLog += FString::Printf(TEXT("%s\tImageBarrier[%d]: srcAccess=%s, oldLayout=%s, srcQueueFamilyIndex=%d\n"), Tabs, Index, *GetAccessFlagString(ImageMemoryBarriers[Index].srcAccessMask), *GetImageLayoutString(ImageMemoryBarriers[Index].oldLayout), ImageMemoryBarriers[Index].srcQueueFamilyIndex);
-			DebugLog += FString::Printf(TEXT("%s\t\tdstAccess=%s, newLayout=%s, dstQueueFamilyIndex=%d\n"), Tabs, *GetAccessFlagString(ImageMemoryBarriers[Index].dstAccessMask), *GetImageLayoutString(ImageMemoryBarriers[Index].newLayout), ImageMemoryBarriers[Index].dstQueueFamilyIndex);
-			DebugLog += FString::Printf(TEXT("%s\t\tImage=0x%p, subresourceRange=(%s)\n"), Tabs, ImageMemoryBarriers[Index].image, *GetImageSubResourceRangeString(ImageMemoryBarriers[Index].subresourceRange));
-		}
+		DebugLog += FString::Printf(TEXT("%s\tImageBarrier[%d]: srcAccess=%s, oldLayout=%s, srcQueueFamilyIndex=%d\n"), Tabs, Index, *GetAccessFlagString(ImageMemoryBarriers[Index].srcAccessMask), *GetVkImageLayoutString(ImageMemoryBarriers[Index].oldLayout), ImageMemoryBarriers[Index].srcQueueFamilyIndex);
+		DebugLog += FString::Printf(TEXT("%s\t\tdstAccess=%s, newLayout=%s, dstQueueFamilyIndex=%d\n"), Tabs, *GetAccessFlagString(ImageMemoryBarriers[Index].dstAccessMask), *GetVkImageLayoutString(ImageMemoryBarriers[Index].newLayout), ImageMemoryBarriers[Index].dstQueueFamilyIndex);
+		DebugLog += FString::Printf(TEXT("%s\t\tImage=0x%p, subresourceRange=(%s)\n"), Tabs, ImageMemoryBarriers[Index].image, *GetImageSubResourceRangeString(ImageMemoryBarriers[Index].subresourceRange));
 	}
 }
 #endif
 
-void FWrapLayer::PipelineBarrier(VkResult Result, VkCommandBuffer CommandBuffer, VkPipelineStageFlags SrcStageMask, VkPipelineStageFlags DstStageMask, VkDependencyFlags DependencyFlags, uint32 MemoryBarrierCount, const VkMemoryBarrier* MemoryBarriers, uint32 BufferMemoryBarrierCount, const VkBufferMemoryBarrier* BufferMemoryBarriers, uint32 ImageMemoryBarrierCount, const VkImageMemoryBarrier* ImageMemoryBarriers)
+void FWrapLayer::CmdPipelineBarrier(VkResult Result, VkCommandBuffer CommandBuffer, VkPipelineStageFlags SrcStageMask, VkPipelineStageFlags DstStageMask, VkDependencyFlags DependencyFlags, uint32 MemoryBarrierCount, const VkMemoryBarrier* MemoryBarriers, uint32 BufferMemoryBarrierCount, const VkBufferMemoryBarrier* BufferMemoryBarriers, uint32 ImageMemoryBarrierCount, const VkImageMemoryBarrier* ImageMemoryBarriers)
 {
 	if (Result == VK_RESULT_MAX_ENUM)
 	{
@@ -1836,7 +1847,7 @@ void FWrapLayer::PipelineBarrier(VkResult Result, VkCommandBuffer CommandBuffer,
 	}
 }
 
-void FWrapLayer::WaitEvents(VkResult Result, VkCommandBuffer CommandBuffer, uint32 EventCount, const VkEvent* Events, VkPipelineStageFlags SrcStageMask, VkPipelineStageFlags DstStageMask, uint32 MemoryBarrierCount,
+void FWrapLayer::CmdWaitEvents(VkResult Result, VkCommandBuffer CommandBuffer, uint32 EventCount, const VkEvent* Events, VkPipelineStageFlags SrcStageMask, VkPipelineStageFlags DstStageMask, uint32 MemoryBarrierCount,
 		const VkMemoryBarrier* MemoryBarriers, uint32 BufferMemoryBarrierCount, const VkBufferMemoryBarrier* BufferMemoryBarriers, uint32 ImageMemoryBarrierCount, const VkImageMemoryBarrier* ImageMemoryBarriers)
 {
 	if (Result == VK_RESULT_MAX_ENUM)
@@ -1854,7 +1865,7 @@ void FWrapLayer::WaitEvents(VkResult Result, VkCommandBuffer CommandBuffer, uint
 	}
 }
 
-void FWrapLayer::BindDescriptorSets(VkResult Result, VkCommandBuffer CommandBuffer, VkPipelineBindPoint PipelineBindPoint, VkPipelineLayout Layout, uint32 FirstSet, uint32 DescriptorSetCount, const VkDescriptorSet* DescriptorSets, uint32 DynamicOffsetCount, const uint32* DynamicOffsets)
+void FWrapLayer::CmdBindDescriptorSets(VkResult Result, VkCommandBuffer CommandBuffer, VkPipelineBindPoint PipelineBindPoint, VkPipelineLayout Layout, uint32 FirstSet, uint32 DescriptorSetCount, const VkDescriptorSet* DescriptorSets, uint32 DynamicOffsetCount, const uint32* DynamicOffsets)
 {
 	if (Result == VK_RESULT_MAX_ENUM)
 	{
@@ -1883,7 +1894,7 @@ void FWrapLayer::CreateDescriptorSetLayout(VkResult Result, VkDevice Device, con
 		for (uint32 Index = 0; Index < CreateInfo->bindingCount; ++Index)
 		{
 			DebugLog += FString::Printf(TEXT("%s\tBinding[%d]= binding=%d DescType=%s NumDesc=%d StageFlags=%x\n"), Tabs, Index,
-				CreateInfo->pBindings[Index].binding, *GetDescriptorTypeString(CreateInfo->pBindings[Index].descriptorType), CreateInfo->pBindings[Index].descriptorCount, (uint32)CreateInfo->pBindings[Index].stageFlags);
+				CreateInfo->pBindings[Index].binding, *GetVkDescriptorTypeString(CreateInfo->pBindings[Index].descriptorType), CreateInfo->pBindings[Index].descriptorCount, (uint32)CreateInfo->pBindings[Index].stageFlags);
 		}
 #endif
 	}
@@ -1928,7 +1939,7 @@ void FWrapLayer::UpdateDescriptorSets(VkResult Result, VkDevice Device, uint32 D
 		{
 #if VULKAN_ENABLE_DUMP_LAYER
 			DebugLog += FString::Printf(TEXT("%sWrite[%d]: Set=0x%p Binding=%d DstArrayElem=%d NumDesc=%d DescType=%s "), Tabs, Index,
-				DescriptorWrites[Index].dstSet, DescriptorWrites[Index].dstBinding, DescriptorWrites[Index].dstArrayElement, DescriptorWrites[Index].descriptorCount, *GetDescriptorTypeString(DescriptorWrites[Index].descriptorType));
+				DescriptorWrites[Index].dstSet, DescriptorWrites[Index].dstBinding, DescriptorWrites[Index].dstArrayElement, DescriptorWrites[Index].descriptorCount, *GetVkDescriptorTypeString(DescriptorWrites[Index].descriptorType));
 #endif
 			switch (DescriptorWrites[Index].descriptorType)
 			{
@@ -2000,7 +2011,7 @@ void FWrapLayer::UpdateDescriptorSets(VkResult Result, VkDevice Device, uint32 D
 						auto* FoundImageInfo = GVulkanTrackingImageViews.Find(DescriptorWrites[Index].pImageInfo->imageView);
 						VkImage *Image = FoundImageInfo ? &FoundImageInfo->CreateInfo.image : nullptr;
 						DebugLog += FString::Printf(TEXT("%s\tpImageInfo[%d]: Sampler=0x%p, ImageView=0x%p(I:0x%p), imageLayout=%s\n"), Tabs, SubIndex,
-							DescriptorWrites[Index].pImageInfo->sampler, DescriptorWrites[Index].pImageInfo->imageView, Image ? *Image : nullptr, *GetImageLayoutString(DescriptorWrites[Index].pImageInfo->imageLayout));
+							DescriptorWrites[Index].pImageInfo->sampler, DescriptorWrites[Index].pImageInfo->imageView, Image ? *Image : nullptr, *GetVkImageLayoutString(DescriptorWrites[Index].pImageInfo->imageLayout));
 #endif
 					}
 				}
@@ -2099,7 +2110,7 @@ void FWrapLayer::CreateRenderPass(VkResult Result, VkDevice Device, const VkRend
 				(Desc.flags == VK_ATTACHMENT_DESCRIPTION_MAY_ALIAS_BIT ? TEXT("MAY_ALIAS") : TEXT("0")),
 				*GetVkFormatString(Desc.format), *GetSampleCountString(Desc.samples), *GetLoadOpString(Desc.loadOp), *GetStoreOpString(Desc.storeOp));
 			DebugLog += FString::Printf(TEXT("%s\t\t\tLoadStencil=%s, StoreStencil=%s, Initial=%s, Final=%s\n"), Tabs,
-				*GetLoadOpString(Desc.stencilLoadOp), *GetStoreOpString(Desc.stencilStoreOp), *VulkanRHI::GetImageLayoutString(Desc.initialLayout), *VulkanRHI::GetImageLayoutString(Desc.finalLayout));
+				*GetLoadOpString(Desc.stencilLoadOp), *GetStoreOpString(Desc.stencilStoreOp), *VulkanRHI::GetVkImageLayoutString(Desc.initialLayout), *VulkanRHI::GetVkImageLayoutString(Desc.finalLayout));
 		}
 
 		for (uint32 Index = 0; Index < CreateInfo->subpassCount; ++Index)
@@ -2112,16 +2123,16 @@ void FWrapLayer::CreateRenderPass(VkResult Result, VkDevice Device, const VkRend
 			for (uint32 SubIndex = 0; SubIndex < Desc.inputAttachmentCount; ++SubIndex)
 			{
 				DebugLog += FString::Printf(TEXT("%s\t\t\tInputAttach[%d]: Attach=%d, Layout=%s\n"), Tabs, Index,
-					Desc.pInputAttachments[SubIndex].attachment, *GetImageLayoutString(Desc.pInputAttachments[SubIndex].layout));
+					Desc.pInputAttachments[SubIndex].attachment, *GetVkImageLayoutString(Desc.pInputAttachments[SubIndex].layout));
 			}
 			for (uint32 SubIndex = 0; SubIndex < Desc.colorAttachmentCount; ++SubIndex)
 			{
 				DebugLog += FString::Printf(TEXT("%s\t\t\tColorAttach[%d]: Attach=%d, Layout=%s\n"), Tabs, Index,
-					Desc.pColorAttachments[SubIndex].attachment, *GetImageLayoutString(Desc.pColorAttachments[SubIndex].layout));
+					Desc.pColorAttachments[SubIndex].attachment, *GetVkImageLayoutString(Desc.pColorAttachments[SubIndex].layout));
 			}
 			if (Desc.pDepthStencilAttachment)
 			{
-				DebugLog += FString::Printf(TEXT("%s\t\t\tDSAttach: Attach=%d, Layout=%s\n"), Tabs, Desc.pDepthStencilAttachment->attachment, *GetImageLayoutString(Desc.pDepthStencilAttachment->layout));
+				DebugLog += FString::Printf(TEXT("%s\t\t\tDSAttach: Attach=%d, Layout=%s\n"), Tabs, Desc.pDepthStencilAttachment->attachment, *GetVkImageLayoutString(Desc.pDepthStencilAttachment->layout));
 			}
 			/*
 			typedef struct VkSubpassDescription {
@@ -2376,6 +2387,43 @@ void FWrapLayer::GetPhysicalDeviceFormatProperties(VkResult Result, VkPhysicalDe
 	}
 }
 
+#if VULKAN_ENABLE_DUMP_LAYER
+static void DumpImageFormatProperties(const VkImageFormatProperties* pProps)
+{
+	if (pProps)
+	{
+		DebugLog += FString::Printf(
+			TEXT(" -> MaxExtent={ %u, %u, %u } MaxMipLevels=%u MaxArrayLayers=%u SampleCounts=0x%X MaxResourceSize=%llu\n"),
+			pProps->maxExtent.width, pProps->maxExtent.height, pProps->maxMipLevels, pProps->maxArrayLayers, pProps->sampleCounts, pProps->maxResourceSize
+		);
+	}
+	else
+	{
+		DumpErrorInvalidPointer(TEXT("VkImageFormatProperties"));
+	}
+	FlushDebugWrapperLog();
+}
+#endif
+
+void FWrapLayer::GetPhysicalDeviceImageFormatProperties(VkResult Result, VkPhysicalDevice PhysicalDevice, VkFormat Format, VkImageType Type, VkImageTiling Tiling, VkImageUsageFlags Usage, VkImageCreateFlags Flags, VkImageFormatProperties* pImageFormatProperties)
+{
+	if (Result == VK_RESULT_MAX_ENUM)
+	{
+#if VULKAN_ENABLE_DUMP_LAYER
+		PrintfBeginResult(FString::Printf(
+			TEXT("vkGetPhysicalDeviceImageFormatProperties(PhysicalDevice=0x%p, Format=%s, Type=%d, Tiling=%s, Usage=%d, Flags=%d, pImageFormatProperties=0x%p)[...]"),
+			PhysicalDevice, *GetVkFormatString(Format), (int32)Type, *GetVkImageTilingString(Tiling), (int32)Usage, (int32)Flags, pImageFormatProperties
+		));
+#endif
+	}
+	else
+	{
+#if VULKAN_ENABLE_DUMP_LAYER
+		DumpImageFormatProperties(pImageFormatProperties);
+#endif
+	}
+}
+
 void FWrapLayer::GetPhysicalDeviceProperties(VkResult Result, VkPhysicalDevice PhysicalDevice, VkPhysicalDeviceProperties* FormatProperties)
 {
 	if (Result == VK_RESULT_MAX_ENUM)
@@ -2418,7 +2466,7 @@ void FWrapLayer::EndCommandBuffer(VkResult Result, VkCommandBuffer CommandBuffer
 	}
 }
 
-void FWrapLayer::ResetQueryPool(VkResult Result, VkCommandBuffer CommandBuffer, VkQueryPool QueryPool, uint32 FirstQuery, uint32 QueryCount)
+void FWrapLayer::CmdResetQueryPool(VkResult Result, VkCommandBuffer CommandBuffer, VkQueryPool QueryPool, uint32 FirstQuery, uint32 QueryCount)
 {
 	if (Result == VK_RESULT_MAX_ENUM)
 	{
@@ -2428,7 +2476,7 @@ void FWrapLayer::ResetQueryPool(VkResult Result, VkCommandBuffer CommandBuffer,
 	}
 }
 
-void FWrapLayer::WriteTimestamp(VkResult Result, VkCommandBuffer CommandBuffer, VkPipelineStageFlagBits PipelineStage, VkQueryPool QueryPool, uint32 Query)
+void FWrapLayer::CmdWriteTimestamp(VkResult Result, VkCommandBuffer CommandBuffer, VkPipelineStageFlagBits PipelineStage, VkQueryPool QueryPool, uint32 Query)
 {
 	if (Result == VK_RESULT_MAX_ENUM)
 	{
@@ -2438,7 +2486,7 @@ void FWrapLayer::WriteTimestamp(VkResult Result, VkCommandBuffer CommandBuffer,
 	}
 }
 
-void FWrapLayer::BindPipeline(VkResult Result, VkCommandBuffer CommandBuffer, VkPipelineBindPoint PipelineBindPoint, VkPipeline Pipeline)
+void FWrapLayer::CmdBindPipeline(VkResult Result, VkCommandBuffer CommandBuffer, VkPipelineBindPoint PipelineBindPoint, VkPipeline Pipeline)
 {
 	if (Result == VK_RESULT_MAX_ENUM)
 	{
@@ -2446,13 +2494,19 @@ void FWrapLayer::BindPipeline(VkResult Result, VkCommandBuffer CommandBuffer, Vk
 		CmdPrintfBegin(CommandBuffer, FString::Printf(TEXT("vkCmdBindPipeline(BindPoint=%d, Pipeline=0x%p)[...]"), (int32)PipelineBindPoint, Pipeline));
 #endif
 	}
-	else
+}
+
+void FWrapLayer::CmdPushConstants(VkResult Result, VkCommandBuffer CommandBuffer, VkPipelineLayout Layout, VkShaderStageFlags StageFlags, uint32_t Offset, uint32_t Size, const void* pValues)
+{
+	if (Result == VK_RESULT_MAX_ENUM)
 	{
+#if VULKAN_ENABLE_DUMP_LAYER
+		CmdPrintfBegin(CommandBuffer, FString::Printf(TEXT("vkCmdPushConstants(Layout=0x%p, StageFlags=0x%08X, Offset=%u, Size=%u, pValues=0x%p)"), Layout, StageFlags, Offset, Size, pValues));
+#endif
 	}
 }
 
-
-void FWrapLayer::BeginRenderPass(VkResult Result, VkCommandBuffer CommandBuffer, const VkRenderPassBeginInfo* RenderPassBegin, VkSubpassContents Contents)
+void FWrapLayer::CmdBeginRenderPass(VkResult Result, VkCommandBuffer CommandBuffer, const VkRenderPassBeginInfo* RenderPassBegin, VkSubpassContents Contents)
 {
 	if (Result == VK_RESULT_MAX_ENUM)
 	{
@@ -2503,7 +2557,7 @@ void FWrapLayer::BeginRenderPass(VkResult Result, VkCommandBuffer CommandBuffer,
 	}
 }
 
-void FWrapLayer::EndRenderPass(VkResult Result, VkCommandBuffer CommandBuffer)
+void FWrapLayer::CmdEndRenderPass(VkResult Result, VkCommandBuffer CommandBuffer)
 {
 	if (Result == VK_RESULT_MAX_ENUM)
 	{
@@ -2513,7 +2567,7 @@ void FWrapLayer::EndRenderPass(VkResult Result, VkCommandBuffer CommandBuffer)
 	}
 }
 
-void FWrapLayer::NextSubpass(VkResult Result, VkCommandBuffer CommandBuffer, VkSubpassContents Contents)
+void FWrapLayer::CmdNextSubpass(VkResult Result, VkCommandBuffer CommandBuffer, VkSubpassContents Contents)
 {
 	if (Result == VK_RESULT_MAX_ENUM)
 	{
@@ -2523,7 +2577,17 @@ void FWrapLayer::NextSubpass(VkResult Result, VkCommandBuffer CommandBuffer, VkS
 	}
 }
 
-void FWrapLayer::BeginQuery(VkResult Result, VkCommandBuffer CommandBuffer, VkQueryPool QueryPool, uint32 Query, VkQueryControlFlags Flags)
+void FWrapLayer::CmdExecuteCommands(VkResult Result, VkCommandBuffer CommandBuffer, uint32 CommandBufferCount, const VkCommandBuffer* pCommandBuffers)
+{
+	if (Result == VK_RESULT_MAX_ENUM)
+	{
+#if VULKAN_ENABLE_DUMP_LAYER
+		CmdPrintfBegin(CommandBuffer, FString::Printf(TEXT("vkCmdExecuteCommands(CommandBufferCount=%u, pCommandBuffers=0x%p)"), CommandBufferCount, pCommandBuffers));
+#endif
+	}
+}
+
+void FWrapLayer::CmdBeginQuery(VkResult Result, VkCommandBuffer CommandBuffer, VkQueryPool QueryPool, uint32 Query, VkQueryControlFlags Flags)
 {
 	if (Result == VK_RESULT_MAX_ENUM)
 	{
@@ -2532,7 +2596,7 @@ void FWrapLayer::BeginQuery(VkResult Result, VkCommandBuffer CommandBuffer, VkQu
 #endif
 	}
 }
-void FWrapLayer::EndQuery(VkResult Result, VkCommandBuffer CommandBuffer, VkQueryPool QueryPool, uint32 Query)
+void FWrapLayer::CmdEndQuery(VkResult Result, VkCommandBuffer CommandBuffer, VkQueryPool QueryPool, uint32 Query)
 {
 	if (Result == VK_RESULT_MAX_ENUM)
 	{
@@ -2543,7 +2607,7 @@ void FWrapLayer::EndQuery(VkResult Result, VkCommandBuffer CommandBuffer, VkQuer
 }
 
 
-void FWrapLayer::BindVertexBuffers(VkResult Result, VkCommandBuffer CommandBuffer, uint32 FirstBinding, uint32 BindingCount, const VkBuffer* Buffers, const VkDeviceSize* Offsets)
+void FWrapLayer::CmdBindVertexBuffers(VkResult Result, VkCommandBuffer CommandBuffer, uint32 FirstBinding, uint32 BindingCount, const VkBuffer* Buffers, const VkDeviceSize* Offsets)
 {
 	if (Result == VK_RESULT_MAX_ENUM)
 	{
@@ -2554,7 +2618,7 @@ void FWrapLayer::BindVertexBuffers(VkResult Result, VkCommandBuffer CommandBuffe
 	}
 }
 
-void FWrapLayer::BindIndexBuffer(VkResult Result, VkCommandBuffer CommandBuffer, VkBuffer IndexBuffer, VkDeviceSize Offset, VkIndexType IndexType)
+void FWrapLayer::CmdBindIndexBuffer(VkResult Result, VkCommandBuffer CommandBuffer, VkBuffer IndexBuffer, VkDeviceSize Offset, VkIndexType IndexType)
 {
 	if (Result == VK_RESULT_MAX_ENUM)
 	{
@@ -2565,7 +2629,7 @@ void FWrapLayer::BindIndexBuffer(VkResult Result, VkCommandBuffer CommandBuffer,
 	}
 }
 
-void FWrapLayer::SetViewport(VkResult Result, VkCommandBuffer CommandBuffer, uint32 FirstViewport, uint32 ViewportCount, const VkViewport* Viewports)
+void FWrapLayer::CmdSetViewport(VkResult Result, VkCommandBuffer CommandBuffer, uint32 FirstViewport, uint32 ViewportCount, const VkViewport* Viewports)
 {
 	if (Result == VK_RESULT_MAX_ENUM)
 	{
@@ -2575,7 +2639,7 @@ void FWrapLayer::SetViewport(VkResult Result, VkCommandBuffer CommandBuffer, uin
 	}
 }
 
-void FWrapLayer::SetScissor(VkResult Result, VkCommandBuffer CommandBuffer, uint32 FirstScissor, uint32 ScissorCount, const VkRect2D* Scissors)
+void FWrapLayer::CmdSetScissor(VkResult Result, VkCommandBuffer CommandBuffer, uint32 FirstScissor, uint32 ScissorCount, const VkRect2D* Scissors)
 {
 	if (Result == VK_RESULT_MAX_ENUM)
 	{
@@ -2585,7 +2649,7 @@ void FWrapLayer::SetScissor(VkResult Result, VkCommandBuffer CommandBuffer, uint
 	}
 }
 
-void FWrapLayer::SetLineWidth(VkResult Result, VkCommandBuffer CommandBuffer, float LineWidth)
+void FWrapLayer::CmdSetLineWidth(VkResult Result, VkCommandBuffer CommandBuffer, float LineWidth)
 {
 	if (Result == VK_RESULT_MAX_ENUM)
 	{
@@ -2595,7 +2659,7 @@ void FWrapLayer::SetLineWidth(VkResult Result, VkCommandBuffer CommandBuffer, fl
 	}
 }
 
-void FWrapLayer::Draw(VkResult Result, VkCommandBuffer CommandBuffer, uint32 VertexCount, uint32 InstanceCount, uint32 FirstVertex, uint32 FirstInstance)
+void FWrapLayer::CmdDraw(VkResult Result, VkCommandBuffer CommandBuffer, uint32 VertexCount, uint32 InstanceCount, uint32 FirstVertex, uint32 FirstInstance)
 {
 	if (Result == VK_RESULT_MAX_ENUM)
 	{
@@ -2605,7 +2669,7 @@ void FWrapLayer::Draw(VkResult Result, VkCommandBuffer CommandBuffer, uint32 Ver
 	}
 }
 
-void FWrapLayer::DrawIndexed(VkResult Result, VkCommandBuffer CommandBuffer, uint32 IndexCount, uint32 InstanceCount, uint32 FirstIndex, int32_t VertexOffset, uint32 FirstInstance)
+void FWrapLayer::CmdDrawIndexed(VkResult Result, VkCommandBuffer CommandBuffer, uint32 IndexCount, uint32 InstanceCount, uint32 FirstIndex, int32_t VertexOffset, uint32 FirstInstance)
 {
 	if (Result == VK_RESULT_MAX_ENUM)
 	{
@@ -2615,7 +2679,7 @@ void FWrapLayer::DrawIndexed(VkResult Result, VkCommandBuffer CommandBuffer, uin
 	}
 }
 
-void FWrapLayer::DrawIndirect(VkResult Result, VkCommandBuffer CommandBuffer, VkBuffer Buffer, VkDeviceSize Offset, uint32 DrawCount, uint32 Stride)
+void FWrapLayer::CmdDrawIndirect(VkResult Result, VkCommandBuffer CommandBuffer, VkBuffer Buffer, VkDeviceSize Offset, uint32 DrawCount, uint32 Stride)
 {
 	if (Result == VK_RESULT_MAX_ENUM)
 	{
@@ -2625,7 +2689,7 @@ void FWrapLayer::DrawIndirect(VkResult Result, VkCommandBuffer CommandBuffer, Vk
 	}
 }
 
-void FWrapLayer::DrawIndexedIndirect(VkResult Result, VkCommandBuffer CommandBuffer, VkBuffer Buffer, VkDeviceSize Offset, uint32 DrawCount, uint32 Stride)
+void FWrapLayer::CmdDrawIndexedIndirect(VkResult Result, VkCommandBuffer CommandBuffer, VkBuffer Buffer, VkDeviceSize Offset, uint32 DrawCount, uint32 Stride)
 {
 	if (Result == VK_RESULT_MAX_ENUM)
 	{
@@ -2635,7 +2699,7 @@ void FWrapLayer::DrawIndexedIndirect(VkResult Result, VkCommandBuffer CommandBuf
 	}
 }
 
-void FWrapLayer::Dispatch(VkResult Result, VkCommandBuffer CommandBuffer, uint32 X, uint32 Y, uint32 Z)
+void FWrapLayer::CmdDispatch(VkResult Result, VkCommandBuffer CommandBuffer, uint32 X, uint32 Y, uint32 Z)
 {
 	if (Result == VK_RESULT_MAX_ENUM)
 	{
@@ -2645,7 +2709,7 @@ void FWrapLayer::Dispatch(VkResult Result, VkCommandBuffer CommandBuffer, uint32
 	}
 }
 
-void FWrapLayer::DispatchIndirect(VkResult Result, VkCommandBuffer CommandBuffer, VkBuffer Buffer, VkDeviceSize Offset)
+void FWrapLayer::CmdDispatchIndirect(VkResult Result, VkCommandBuffer CommandBuffer, VkBuffer Buffer, VkDeviceSize Offset)
 {
 	if (Result == VK_RESULT_MAX_ENUM)
 	{
@@ -2655,7 +2719,7 @@ void FWrapLayer::DispatchIndirect(VkResult Result, VkCommandBuffer CommandBuffer
 	}
 }
 
-void FWrapLayer::CopyImage(VkResult Result, VkCommandBuffer CommandBuffer, VkImage SrcImage, VkImageLayout SrcImageLayout, VkImage DstImage, VkImageLayout DstImageLayout, uint32 RegionCount, const VkImageCopy* Regions)
+void FWrapLayer::CmdCopyImage(VkResult Result, VkCommandBuffer CommandBuffer, VkImage SrcImage, VkImageLayout SrcImageLayout, VkImage DstImage, VkImageLayout DstImageLayout, uint32 RegionCount, const VkImageCopy* Regions)
 {
 	if (Result == VK_RESULT_MAX_ENUM)
 	{
@@ -2687,13 +2751,13 @@ void FWrapLayer::CopyImage(VkResult Result, VkCommandBuffer CommandBuffer, VkIma
 
 
 
-void FWrapLayer::CopyBufferToImage(VkResult Result, VkCommandBuffer CommandBuffer, VkBuffer SrcBuffer, VkImage DstImage, VkImageLayout DstImageLayout, uint32 RegionCount, const VkBufferImageCopy* Regions)
+void FWrapLayer::CmdCopyBufferToImage(VkResult Result, VkCommandBuffer CommandBuffer, VkBuffer SrcBuffer, VkImage DstImage, VkImageLayout DstImageLayout, uint32 RegionCount, const VkBufferImageCopy* Regions)
 {
 	if (Result == VK_RESULT_MAX_ENUM)
 	{
 #if VULKAN_ENABLE_DUMP_LAYER
 		CmdPrintfBegin(CommandBuffer, FString::Printf(TEXT("vkCmdCopyBufferToImage(SrcBuffer=0x%p, DstImage=0x%p, DstImageLayout=%s, NumRegions=%d, Regions=0x%p)"),
-			SrcBuffer, DstImage, *GetImageLayoutString(DstImageLayout), RegionCount, Regions));
+			SrcBuffer, DstImage, *GetVkImageLayoutString(DstImageLayout), RegionCount, Regions));
 		for (uint32 Index = 0; Index < RegionCount; ++Index)
 		{
 			DebugLog += FString::Printf(TEXT("%sRegion[%d]: %s\n"), Tabs, Index, *GetBufferImageCopyString(Regions[Index]));
@@ -2719,13 +2783,13 @@ void FWrapLayer::CopyBufferToImage(VkResult Result, VkCommandBuffer CommandBuffe
 	}
 }
 
-void FWrapLayer::CopyImageToBuffer(VkResult Result, VkCommandBuffer CommandBuffer, VkImage SrcImage, VkImageLayout SrcImageLayout, VkBuffer DstBuffer, uint32 RegionCount, const VkBufferImageCopy* Regions)
+void FWrapLayer::CmdCopyImageToBuffer(VkResult Result, VkCommandBuffer CommandBuffer, VkImage SrcImage, VkImageLayout SrcImageLayout, VkBuffer DstBuffer, uint32 RegionCount, const VkBufferImageCopy* Regions)
 {
 	if (Result == VK_RESULT_MAX_ENUM)
 	{
 #if VULKAN_ENABLE_DUMP_LAYER
 		CmdPrintfBegin(CommandBuffer, FString::Printf(TEXT("vkCmdCopyImageToBuffer(SrcImage=0x%p, SrcImageLayout=%s, SrcBuffer=0x%p, NumRegions=%d, Regions=0x%p)"),
-			SrcImage, *GetImageLayoutString(SrcImageLayout), DstBuffer, RegionCount, Regions));
+			SrcImage, *GetVkImageLayoutString(SrcImageLayout), DstBuffer, RegionCount, Regions));
 		for (uint32 Index = 0; Index < RegionCount; ++Index)
 		{
 			DebugLog += FString::Printf(TEXT("%sRegion[%d]: %s\n"), Tabs, Index, *GetBufferImageCopyString(Regions[Index]));
@@ -2735,7 +2799,7 @@ void FWrapLayer::CopyImageToBuffer(VkResult Result, VkCommandBuffer CommandBuffe
 	}
 }
 
-void FWrapLayer::CopyBuffer(VkResult Result, VkCommandBuffer CommandBuffer, VkBuffer SrcBuffer, VkBuffer DstBuffer, uint32 RegionCount, const VkBufferCopy* Regions)
+void FWrapLayer::CmdCopyBuffer(VkResult Result, VkCommandBuffer CommandBuffer, VkBuffer SrcBuffer, VkBuffer DstBuffer, uint32 RegionCount, const VkBufferCopy* Regions)
 {
 	if (Result == VK_RESULT_MAX_ENUM)
 	{
@@ -2751,7 +2815,7 @@ void FWrapLayer::CopyBuffer(VkResult Result, VkCommandBuffer CommandBuffer, VkBu
 	}
 }
 
-void FWrapLayer::BlitImage(VkResult Result, VkCommandBuffer CommandBuffer, VkImage SrcImage, VkImageLayout SrcImageLayout, VkImage DstImage, VkImageLayout DstImageLayout, uint32 RegionCount, const VkImageBlit* Regions, VkFilter Filter)
+void FWrapLayer::CmdBlitImage(VkResult Result, VkCommandBuffer CommandBuffer, VkImage SrcImage, VkImageLayout SrcImageLayout, VkImage DstImage, VkImageLayout DstImageLayout, uint32 RegionCount, const VkImageBlit* Regions, VkFilter Filter)
 {
 	if (Result == VK_RESULT_MAX_ENUM)
 	{
@@ -2784,9 +2848,16 @@ void FWrapLayer::BlitImage(VkResult Result, VkCommandBuffer CommandBuffer, VkIma
 }
 
 #if VULKAN_ENABLE_DUMP_LAYER
-static void DumpImageSubresourceLayout(VkSubresourceLayout* Layout)
+static void DumpImageSubresourceLayout(VkSubresourceLayout* pLayout)
 {
-	DebugLog += FString::Printf(TEXT("VkSubresourceLayout: [...]\n"));
+	if (pLayout)
+	{
+		DebugLog += FString::Printf(TEXT("VkSubresourceLayout: [...]\n"));
+	}
+	else
+	{
+		DumpErrorInvalidPointer(TEXT("VkSubresourceLayout"));
+	}
 	FlushDebugWrapperLog();
 }
 #endif
@@ -2848,7 +2919,7 @@ void FWrapLayer::GetSwapChainImagesKHR(VkResult Result, VkDevice Device, VkSwapc
 	}
 }
 
-void FWrapLayer::ClearAttachments(VkResult Result, VkCommandBuffer CommandBuffer, uint32 AttachmentCount, const VkClearAttachment* Attachments, uint32 RectCount, const VkClearRect* Rects)
+void FWrapLayer::CmdClearAttachments(VkResult Result, VkCommandBuffer CommandBuffer, uint32 AttachmentCount, const VkClearAttachment* Attachments, uint32 RectCount, const VkClearRect* Rects)
 {
 	if (Result == VK_RESULT_MAX_ENUM)
 	{
@@ -2870,12 +2941,12 @@ void FWrapLayer::ClearAttachments(VkResult Result, VkCommandBuffer CommandBuffer
 	}
 }
 
-void FWrapLayer::ClearColorImage(VkResult Result, VkCommandBuffer CommandBuffer, VkImage Image, VkImageLayout ImageLayout, const VkClearColorValue* Color, uint32 RangeCount, const VkImageSubresourceRange* Ranges)
+void FWrapLayer::CmdClearColorImage(VkResult Result, VkCommandBuffer CommandBuffer, VkImage Image, VkImageLayout ImageLayout, const VkClearColorValue* Color, uint32 RangeCount, const VkImageSubresourceRange* Ranges)
 {
 	if (Result == VK_RESULT_MAX_ENUM)
 	{
 #if VULKAN_ENABLE_DUMP_LAYER
-		CmdPrintfBegin(CommandBuffer, FString::Printf(TEXT("vkCmdClearColorImage(Image=0x%p, ImageLayout=%s, Color=%s, RangeCount=%d, Ranges=0x%p)"), Image, *GetImageLayoutString(ImageLayout), *GetClearColorValueString(*Color), RangeCount, Ranges));
+		CmdPrintfBegin(CommandBuffer, FString::Printf(TEXT("vkCmdClearColorImage(Image=0x%p, ImageLayout=%s, Color=%s, RangeCount=%d, Ranges=0x%p)"), Image, *GetVkImageLayoutString(ImageLayout), *GetClearColorValueString(*Color), RangeCount, Ranges));
 		for (uint32 Index = 0; Index < RangeCount; ++Index)
 		{
 			DebugLog += FString::Printf(TEXT("%sRange[%d]= %s\n"), Tabs, Index, *GetImageSubResourceRangeString(Ranges[Index]));
@@ -2885,12 +2956,12 @@ void FWrapLayer::ClearColorImage(VkResult Result, VkCommandBuffer CommandBuffer,
 	}
 }
 
-void FWrapLayer::ClearDepthStencilImage(VkResult Result, VkCommandBuffer CommandBuffer, VkImage Image, VkImageLayout ImageLayout, const VkClearDepthStencilValue* DepthStencil, uint32 RangeCount, const VkImageSubresourceRange* Ranges)
+void FWrapLayer::CmdClearDepthStencilImage(VkResult Result, VkCommandBuffer CommandBuffer, VkImage Image, VkImageLayout ImageLayout, const VkClearDepthStencilValue* DepthStencil, uint32 RangeCount, const VkImageSubresourceRange* Ranges)
 {
 	if (Result == VK_RESULT_MAX_ENUM)
 	{
 #if VULKAN_ENABLE_DUMP_LAYER
-		CmdPrintfBegin(CommandBuffer, FString::Printf(TEXT("vkCmdClearDepthStencilImage(Image=0x%p, ImageLayout=%s, DepthStencil=%s, RangeCount=%d, Ranges=0x%p)"), Image, *GetImageLayoutString(ImageLayout), *GetClearDepthStencilValueString(*DepthStencil), RangeCount, Ranges));
+		CmdPrintfBegin(CommandBuffer, FString::Printf(TEXT("vkCmdClearDepthStencilImage(Image=0x%p, ImageLayout=%s, DepthStencil=%s, RangeCount=%d, Ranges=0x%p)"), Image, *GetVkImageLayoutString(ImageLayout), *GetClearDepthStencilValueString(*DepthStencil), RangeCount, Ranges));
 		for (uint32 Index = 0; Index < RangeCount; ++Index)
 		{
 			DebugLog += FString::Printf(TEXT("%sRange[%d]= %s\n"), Tabs, Index, *GetImageSubResourceRangeString(Ranges[Index]));
@@ -2945,11 +3016,11 @@ void FWrapLayer::CreateGraphicsPipelines(VkResult Result, VkDevice Device, VkPip
 			const VkGraphicsPipelineCreateInfo& CreateInfo = CreateInfos[Index];
 			DebugLog += FString::Printf(TEXT("%s%d: Flags=%d Stages=%d Layout=0x%p RenderPass=0x%p Subpass=%d\n"), Tabs, Index,
 				CreateInfo.flags, CreateInfo.stageCount, (void*)CreateInfo.layout, (void*)CreateInfo.renderPass, CreateInfo.subpass);
-			DebugLog += FString::Printf(TEXT("%s\tDepth Test %d Write %d %s Bounds %d (min %f max %f) Stencil %d\n"), Tabs, CreateInfo.pDepthStencilState->depthTestEnable, CreateInfo.pDepthStencilState->depthWriteEnable, *GetCompareOpString(CreateInfo.pDepthStencilState->depthCompareOp), CreateInfo.pDepthStencilState->depthBoundsTestEnable, CreateInfo.pDepthStencilState->minDepthBounds, CreateInfo.pDepthStencilState->maxDepthBounds, CreateInfo.pDepthStencilState->stencilTestEnable);
+			DebugLog += FString::Printf(TEXT("%s\tDepth Test %d Write %d %s Bounds %d (min %f max %f) Stencil %d\n"), Tabs, CreateInfo.pDepthStencilState->depthTestEnable, CreateInfo.pDepthStencilState->depthWriteEnable, *GetVkCompareOpString(CreateInfo.pDepthStencilState->depthCompareOp), CreateInfo.pDepthStencilState->depthBoundsTestEnable, CreateInfo.pDepthStencilState->minDepthBounds, CreateInfo.pDepthStencilState->maxDepthBounds, CreateInfo.pDepthStencilState->stencilTestEnable);
 
 			auto PrintStencilOp = [](const VkStencilOpState& State)
 			{
-				return FString::Printf(TEXT("Fail %s Pass %s DepthFail %s Compare %s CompareMask 0x%x WriteMask 0x%x Ref 0x%0x"), *GetStencilOpString(State.failOp), *GetStencilOpString(State.passOp), *GetStencilOpString(State.depthFailOp), *GetCompareOpString(State.compareOp), State.compareMask, State.writeMask, State.reference);
+				return FString::Printf(TEXT("Fail %s Pass %s DepthFail %s Compare %s CompareMask 0x%x WriteMask 0x%x Ref 0x%0x"), *GetVkStencilOpString(State.failOp), *GetVkStencilOpString(State.passOp), *GetVkStencilOpString(State.depthFailOp), *GetVkCompareOpString(State.compareOp), State.compareMask, State.writeMask, State.reference);
 			};
 
 			DebugLog += FString::Printf(TEXT("%s\t\tFront: %s\n"), Tabs, *PrintStencilOp(CreateInfo.pDepthStencilState->front));
@@ -3043,12 +3114,41 @@ void FWrapLayer::UnmapMemory(VkResult Result, VkDevice Device, VkDeviceMemory Me
 	}
 }
 
+#if VULKAN_ENABLE_DUMP_LAYER
+static void DumpCommittedMemoryInBytes(VkDeviceSize* pCommittedMemoryInBytes)
+{
+	if (pCommittedMemoryInBytes)
+	{
+		DebugLog += FString::Printf(TEXT(" -> Size=%llu\n"), *pCommittedMemoryInBytes);
+	}
+	else
+	{
+		DumpErrorInvalidPointer(TEXT("VkDeviceSize"));
+	}
+	FlushDebugWrapperLog();
+}
+#endif
+
+void FWrapLayer::GetDeviceMemoryCommitment(VkResult Result, VkDevice Device, VkDeviceMemory Memory, VkDeviceSize* pCommittedMemoryInBytes)
+{
+#if VULKAN_ENABLE_DUMP_LAYER
+	if (Result == VK_RESULT_MAX_ENUM)
+	{
+		DevicePrintfBeginResult(Device, FString::Printf(TEXT("vkGetDeviceMemoryCommitment(Memory=0x%p, pCommittedMemoryInBytes=0x%p)"), Memory, pCommittedMemoryInBytes));
+	}
+	else
+	{
+		DumpCommittedMemoryInBytes(pCommittedMemoryInBytes);
+	}
+#endif
+}
+
 void FWrapLayer::BindBufferMemory(VkResult Result, VkDevice Device, VkBuffer Buffer, VkDeviceMemory Memory, VkDeviceSize MemoryOffset)
 {
 	if (Result == VK_RESULT_MAX_ENUM)
 	{
 #if VULKAN_ENABLE_DUMP_LAYER
-		DevicePrintfBeginResult(Device, FString::Printf(TEXT("vkBindBufferMemory(Buffer=0x%p, DevMem=0x%p, MemOff=%d)\n"), Buffer, Memory, (uint32)MemoryOffset));
+		DevicePrintfBeginResult(Device, FString::Printf(TEXT("vkBindBufferMemory(Buffer=0x%p, Memory=0x%p, MemoryOffset=%llu)\n"), Buffer, Memory, MemoryOffset));
 #endif
 	}
 	else
@@ -3124,6 +3224,22 @@ void FWrapLayer::CreateComputePipelines(VkResult Result, VkDevice Device, VkPipe
 	}
 }
 
+void FWrapLayer::ResetCommandPool(VkResult Result, VkDevice Device, VkCommandPool CommandPool, VkCommandPoolResetFlags Flags)
+{
+	if (Result == VK_RESULT_MAX_ENUM)
+	{
+#if VULKAN_ENABLE_DUMP_LAYER
+		DevicePrintfBeginResult(Device, FString::Printf(TEXT("vkResetCommandPool(CommandPool=0x%p, Flags=0x%08X)"), CommandPool, Flags));
+#endif
+	}
+	else
+	{
+#if VULKAN_ENABLE_DUMP_LAYER
+		PrintResult(Result);
+#endif
+	}
+}
+
 void FWrapLayer::AllocateCommandBuffers(VkResult Result, VkDevice Device, const VkCommandBufferAllocateInfo* AllocateInfo, VkCommandBuffer* CommandBuffers)
 {
 	if (Result == VK_RESULT_MAX_ENUM)
@@ -3467,6 +3583,37 @@ void FWrapLayer::DestroyRenderPass(VkResult Result, VkDevice Device, VkRenderPas
 	}
 }
 
+#if VULKAN_ENABLE_DUMP_LAYER
+static void DumpExtent2D(const VkExtent2D* pExtent)
+{
+	if (pExtent)
+	{
+		DebugLog += FString::Printf(TEXT(" -> Width=%u Height=%u\n"), pExtent->width, pExtent->height);
+	}
+	else
+	{
+		DumpErrorInvalidPointer(TEXT("VkExtent2D"));
+	}
+	FlushDebugWrapperLog();
+}
+#endif
+
+void FWrapLayer::GetRenderAreaGranularity(VkResult Result, VkDevice Device, VkRenderPass RenderPass, VkExtent2D* pGranularity)
+{
+	if (Result == VK_RESULT_MAX_ENUM)
+	{
+#if VULKAN_ENABLE_DUMP_LAYER
+		DevicePrintfBeginResult(Device, FString::Printf(TEXT("vkGetRenderAreaGranularity(RenderPass=0x%p, pGranularity=0x%p)"), RenderPass, pGranularity));
+#endif
+	}
+	else
+	{
+#if VULKAN_ENABLE_DUMP_LAYER
+		DumpExtent2D(pGranularity);
+#endif
+	}
+}
+
 void FWrapLayer::DestroyCommandPool(VkResult Result, VkDevice Device, VkCommandPool CommandPool)
 {
 	if (Result == VK_RESULT_MAX_ENUM)
@@ -3547,20 +3694,59 @@ void FWrapLayer::GetImageMemoryRequirements2KHR(VkResult Result, VkDevice Device
 }
 #endif	// VULKAN_SUPPORTS_DEDICATED_ALLOCATION
 
+void FWrapLayer::GetImageSparseMemoryRequirements(VkResult Result, VkDevice Device, VkImage Image, uint32_t* pSparseMemoryRequirementCount, VkSparseImageMemoryRequirements* pSparseMemoryRequirements)
+{
+	if (Result == VK_RESULT_MAX_ENUM)
+	{
+#if VULKAN_ENABLE_DUMP_LAYER
+		DevicePrintfBegin(Device, FString::Printf(TEXT("vkGetImageSparseMemoryRequirements(Device=0x%p, Image=0x%p, pSparseMemoryRequirementCount=0x%p, pSparseMemoryRequirements=0x%p)[...]"), Device, Image, pSparseMemoryRequirementCount, pSparseMemoryRequirements));
+#endif
+	}
+}
+
+void FWrapLayer::GetPhysicalDeviceSparseImageFormatProperties(VkResult Result, VkPhysicalDevice PhysicalDevice, VkFormat Format, VkImageType Type, VkSampleCountFlagBits Samples, VkImageUsageFlags Usage, VkImageTiling Tiling, uint32_t* pPropertyCount, VkSparseImageFormatProperties* pProperties)
+{
+	if (Result == VK_RESULT_MAX_ENUM)
+	{
+#if VULKAN_ENABLE_DUMP_LAYER
+		PrintfBegin(FString::Printf(
+			TEXT("vkGetPhysicalDeviceSparseImageFormatProperties(PhysicalDevice=0x%p, Format=%s, Type=%d, Samples=0x%08X, Usage=%d, Tiling=%d, pPropertyCount=%u, pProperties=0x%p)[...]"),
+			PhysicalDevice, *GetVkFormatString(Format), (int32)Type, (int32)Samples, (int32)Usage, *GetVkImageTilingString(Tiling), (pPropertyCount != nullptr ? *pPropertyCount : 0u), pProperties
+		));
+#endif
+	}
+}
+
+void FWrapLayer::QueueBindSparse(VkResult Result, VkQueue Queue, uint32_t BindInfoCount, const VkBindSparseInfo* pBindInfo, VkFence Fence)
+{
+	if (Result == VK_RESULT_MAX_ENUM)
+	{
+#if VULKAN_ENABLE_DUMP_LAYER
+		PrintfBeginResult(FString::Printf(TEXT("vkQueueBindSparse(Queue=0x%p, BindInfoCount=%u, pBindInfo=0x%p, Fence=0x%p)"), Queue, BindInfoCount, pBindInfo, Fence));
+#endif
+	}
+	else
+	{
+#if VULKAN_ENABLE_DUMP_LAYER
+		PrintResult(Result);
+#endif
+	}
+}
+
 #if VULKAN_HAS_PHYSICAL_DEVICE_PROPERTIES2
 void FWrapLayer::GetPhysicalDeviceProperties2KHR(VkResult Result, VkPhysicalDevice PhysicalDevice, VkPhysicalDeviceProperties2KHR* Properties)
 {
 	if (Result == VK_RESULT_MAX_ENUM)
 	{
 #if VULKAN_ENABLE_DUMP_LAYER
-			PrintfBegin(FString::Printf(TEXT("vkGetPhysicalDeviceProperties2KHR(PhysicalDevice=0x%p, Properties=0x%p)[...]"), PhysicalDevice, Properties));
+		PrintfBegin(FString::Printf(TEXT("vkGetPhysicalDeviceProperties2KHR(PhysicalDevice=0x%p, Properties=0x%p)[...]"), PhysicalDevice, Properties));
 #endif
 	}
 }
 #endif
 
 
-void FWrapLayer::SetDepthBias(VkResult Result, VkCommandBuffer CommandBuffer, float DepthBiasConstantFactor, float DepthBiasClamp, float DepthBiasSlopeFactor)
+void FWrapLayer::CmdSetDepthBias(VkResult Result, VkCommandBuffer CommandBuffer, float DepthBiasConstantFactor, float DepthBiasClamp, float DepthBiasSlopeFactor)
 {
 	if (Result == VK_RESULT_MAX_ENUM)
 	{
@@ -3570,7 +3756,7 @@ void FWrapLayer::SetDepthBias(VkResult Result, VkCommandBuffer CommandBuffer, fl
 	}
 }
 
-void FWrapLayer::SetBlendConstants(VkResult Result, VkCommandBuffer CommandBuffer, const float BlendConstants[4])
+void FWrapLayer::CmdSetBlendConstants(VkResult Result, VkCommandBuffer CommandBuffer, const float BlendConstants[4])
 {
 	if (Result == VK_RESULT_MAX_ENUM)
 	{
@@ -3580,7 +3766,7 @@ void FWrapLayer::SetBlendConstants(VkResult Result, VkCommandBuffer CommandBuffe
 	}
 }
 
-void FWrapLayer::SetDepthBounds(VkResult Result, VkCommandBuffer CommandBuffer, float MinDepthBounds, float MaxDepthBounds)
+void FWrapLayer::CmdSetDepthBounds(VkResult Result, VkCommandBuffer CommandBuffer, float MinDepthBounds, float MaxDepthBounds)
 {
 	if (Result == VK_RESULT_MAX_ENUM)
 	{
@@ -3590,7 +3776,7 @@ void FWrapLayer::SetDepthBounds(VkResult Result, VkCommandBuffer CommandBuffer,
 	}
 }
 
-void FWrapLayer::SetStencilCompareMask(VkResult Result, VkCommandBuffer CommandBuffer, VkStencilFaceFlags FaceMask, uint32 CompareMask)
+void FWrapLayer::CmdSetStencilCompareMask(VkResult Result, VkCommandBuffer CommandBuffer, VkStencilFaceFlags FaceMask, uint32 CompareMask)
 {
 	if (Result == VK_RESULT_MAX_ENUM)
 	{
@@ -3600,7 +3786,7 @@ void FWrapLayer::SetStencilCompareMask(VkResult Result, VkCommandBuffer CommandB
 	}
 }
 
-void FWrapLayer::SetStencilWriteMask(VkResult Result, VkCommandBuffer CommandBuffer, VkStencilFaceFlags FaceMask, uint32 WriteMask)
+void FWrapLayer::CmdSetStencilWriteMask(VkResult Result, VkCommandBuffer CommandBuffer, VkStencilFaceFlags FaceMask, uint32 WriteMask)
 {
 	if (Result == VK_RESULT_MAX_ENUM)
 	{
@@ -3610,7 +3796,7 @@ void FWrapLayer::SetStencilWriteMask(VkResult Result, VkCommandBuffer CommandBuf
 	}
 }
 
-void FWrapLayer::SetStencilReference(VkResult Result, VkCommandBuffer CommandBuffer, VkStencilFaceFlags FaceMask, uint32 Reference)
+void FWrapLayer::CmdSetStencilReference(VkResult Result, VkCommandBuffer CommandBuffer, VkStencilFaceFlags FaceMask, uint32 Reference)
 {
 	if (Result == VK_RESULT_MAX_ENUM)
 	{
@@ -3620,7 +3806,7 @@ void FWrapLayer::SetStencilReference(VkResult Result, VkCommandBuffer CommandBuf
 	}
 }
 
-void FWrapLayer::UpdateBuffer(VkResult Result, VkCommandBuffer CommandBuffer, VkBuffer DstBuffer, VkDeviceSize DstOffset, VkDeviceSize DataSize, const void* pData)
+void FWrapLayer::CmdUpdateBuffer(VkResult Result, VkCommandBuffer CommandBuffer, VkBuffer DstBuffer, VkDeviceSize DstOffset, VkDeviceSize DataSize, const void* pData)
 {
 	if (Result == VK_RESULT_MAX_ENUM)
 	{
@@ -3630,7 +3816,7 @@ void FWrapLayer::UpdateBuffer(VkResult Result, VkCommandBuffer CommandBuffer, Vk
 	}
 }
 
-void FWrapLayer::FillBuffer(VkResult Result, VkCommandBuffer CommandBuffer, VkBuffer DstBuffer, VkDeviceSize DstOffset, VkDeviceSize Size, uint32 Data)
+void FWrapLayer::CmdFillBuffer(VkResult Result, VkCommandBuffer CommandBuffer, VkBuffer DstBuffer, VkDeviceSize DstOffset, VkDeviceSize Size, uint32 Data)
 {
 	if (Result == VK_RESULT_MAX_ENUM)
 	{
@@ -3640,7 +3826,7 @@ void FWrapLayer::FillBuffer(VkResult Result, VkCommandBuffer CommandBuffer, VkBu
 	}
 }
 
-void FWrapLayer::SetEvent(VkResult Result, VkCommandBuffer CommandBuffer, VkEvent Event, VkPipelineStageFlags StageMask)
+void FWrapLayer::CmdSetEvent(VkResult Result, VkCommandBuffer CommandBuffer, VkEvent Event, VkPipelineStageFlags StageMask)
 {
 	if (Result == VK_RESULT_MAX_ENUM)
 	{
@@ -3650,7 +3836,7 @@ void FWrapLayer::SetEvent(VkResult Result, VkCommandBuffer CommandBuffer, VkEven
 	}
 }
 
-void FWrapLayer::ResetEvent(VkResult Result, VkCommandBuffer CommandBuffer, VkEvent Event, VkPipelineStageFlags StageMask)
+void FWrapLayer::CmdResetEvent(VkResult Result, VkCommandBuffer CommandBuffer, VkEvent Event, VkPipelineStageFlags StageMask)
 {
 	if (Result == VK_RESULT_MAX_ENUM)
 	{
@@ -3696,7 +3882,7 @@ void FWrapLayer::GetEventStatus(VkResult Result, VkDevice Device, VkEvent Event)
 	}
 }
 
-void FWrapLayer::CopyQueryPoolResults(VkResult Result, VkCommandBuffer CommandBuffer, VkQueryPool QueryPool, uint32 FirstQuery, uint32 QueryCount, VkBuffer DstBuffer, VkDeviceSize DstOffset, VkDeviceSize Stride, VkQueryResultFlags Flags)
+void FWrapLayer::CmdCopyQueryPoolResults(VkResult Result, VkCommandBuffer CommandBuffer, VkQueryPool QueryPool, uint32 FirstQuery, uint32 QueryCount, VkBuffer DstBuffer, VkDeviceSize DstOffset, VkDeviceSize Stride, VkQueryResultFlags Flags)
 {
 	if (Result == VK_RESULT_MAX_ENUM)
 	{
