DR - Fix recursion problem
#rnx

diff --git a/Engine/Source/Runtime/VulkanRHI/Private/VulkanCommandWrappers.h b/Engine/Source/Runtime/VulkanRHI/Private/VulkanCommandWrappers.h
--- a/Engine/Source/Runtime/VulkanRHI/Private/VulkanCommandWrappers.h
+++ b/Engine/Source/Runtime/VulkanRHI/Private/VulkanCommandWrappers.h
@@ -466,7 +466,7 @@ namespace VulkanRHI
 	static FORCEINLINE_DEBUGGABLE VkResult  vkQueueBindSparse(VkQueue Queue, uint32 BindInfoCount, const VkBindSparseInfo* pBindInfo, VkFence Fence)
 	{
 		FWrapLayer::QueueBindSparse(VK_RESULT_MAX_ENUM, Queue, BindInfoCount, pBindInfo, Fence);
-		VkResult Result = vkQueueBindSparse(Queue, BindInfoCount, pBindInfo, Fence);
+		VkResult Result = VULKANAPINAMESPACE::vkQueueBindSparse(Queue, BindInfoCount, pBindInfo, Fence);
 		FWrapLayer::QueueBindSparse(Result, Queue, BindInfoCount, pBindInfo, Fence);
 		return Result;
 	}
