DR - Process Vulkan validation before loading the dlls
#rnx

diff --git a/Engine/Source/Runtime/VulkanRHI/Private/VulkanLayers.cpp b/Engine/Source/Runtime/VulkanRHI/Private/VulkanLayers.cpp
--- a/Engine/Source/Runtime/VulkanRHI/Private/VulkanLayers.cpp
+++ b/Engine/Source/Runtime/VulkanRHI/Private/VulkanLayers.cpp
@@ -30,10 +30,10 @@ TAutoConsoleVariable<int32> GValidationCvar(
 
 static TAutoConsoleVariable<int32> GStandardValidationCvar(
 	TEXT("r.Vulkan.StandardValidation"),
-	1,
-	TEXT("2 to use VK_LAYER_KHRONOS_validation if available\n")
-	TEXT("1 to use VK_LAYER_LUNARG_standard_validation (default) if available, or \n")
-	TEXT("0 to use individual validation layers"),
+	2,
+	TEXT("2 to use VK_LAYER_KHRONOS_validation (default) if available\n")
+	TEXT("1 to use VK_LAYER_LUNARG_standard_validation if available, or \n")
+	TEXT("0 to use individual validation layers (deprecated)"),
 	ECVF_ReadOnly | ECVF_RenderThreadSafe
 );
 
@@ -315,17 +315,8 @@ void FVulkanDynamicRHI::GetInstanceLayersAndExtensions(TArray<const ANSICHAR*>&
 	}
 #endif	// VULKAN_ENABLE_API_DUMP
 
-	int32 VulkanValidationOption = GValidationCvar.GetValueOnAnyThread();
-	if (FParse::Param(FCommandLine::Get(), TEXT("vulkandebug")))
-	{
-		// Match D3D and GL
-		GValidationCvar->Set(2, ECVF_SetByCommandline);
-	}
-	else if (FParse::Value(FCommandLine::Get(), TEXT("vulkanvalidation="), VulkanValidationOption))
-	{
-		GValidationCvar->Set(VulkanValidationOption, ECVF_SetByCommandline);
-	}
-
+	// At this point the CVar holds the final value
+	const int32 VulkanValidationOption = GValidationCvar.GetValueOnAnyThread();
 	if (!bVkTrace && VulkanValidationOption > 0)
 	{
 		bool bStandardAvailable = false;
@@ -377,7 +368,7 @@ void FVulkanDynamicRHI::GetInstanceLayersAndExtensions(TArray<const ANSICHAR*>&
 	}
 
 #if VULKAN_SUPPORTS_DEBUG_UTILS
-	if (!bVkTrace && GValidationCvar.GetValueOnAnyThread() > 0)
+	if (!bVkTrace && VulkanValidationOption > 0)
 	{
 		const char* FoundDebugUtilsLayer = nullptr;
 		bOutDebugUtils = FindLayerExtensionInList(GlobalLayerExtensions, VK_EXT_DEBUG_UTILS_EXTENSION_NAME, FoundDebugUtilsLayer);
@@ -429,7 +420,7 @@ void FVulkanDynamicRHI::GetInstanceLayersAndExtensions(TArray<const ANSICHAR*>&
 	}
 #endif
 #if VULKAN_HAS_DEBUGGING_ENABLED
-	if (!bVkTrace && !bOutDebugUtils && GValidationCvar.GetValueOnAnyThread() > 0)
+	if (!bVkTrace && !bOutDebugUtils && VulkanValidationOption > 0)
 	{
 		if (FindLayerExtensionInList(GlobalLayerExtensions, VK_EXT_DEBUG_REPORT_EXTENSION_NAME))
 		{
@@ -543,7 +534,8 @@ void FVulkanDevice::GetDeviceExtensionsAndLayers(TArray<const ANSICHAR*>& OutDev
 	#endif
 
 	// Verify that all requested debugging device-layers are available. Skip validation layers under RenderDoc
-	if (!GRenderDocFound && GValidationCvar.GetValueOnAnyThread() > 0)
+	const int32 VulkanValidationOption = GValidationCvar.GetValueOnAnyThread();
+	if (!GRenderDocFound && VulkanValidationOption > 0)
 	{
 		bool bStandardAvailable = false;
 		if (GStandardValidationCvar.GetValueOnAnyThread() != 0)
@@ -652,7 +644,7 @@ void FVulkanDevice::GetDeviceExtensionsAndLayers(TArray<const ANSICHAR*>& OutDev
 
 #if VULKAN_ENABLE_DRAW_MARKERS && VULKAN_HAS_DEBUGGING_ENABLED
 	if (!bOutDebugMarkers &&
-		(((GRenderDocFound || GValidationCvar.GetValueOnAnyThread() == 0) && ListContains(AvailableExtensions, VK_EXT_DEBUG_MARKER_EXTENSION_NAME)) || FVulkanPlatform::ForceEnableDebugMarkers()))
+		(((GRenderDocFound || VulkanValidationOption == 0) && ListContains(AvailableExtensions, VK_EXT_DEBUG_MARKER_EXTENSION_NAME)) || FVulkanPlatform::ForceEnableDebugMarkers()))
 	{
 		// HACK: Lumin Nvidia driver unofficially supports this extension, but will return false if we try to load it explicitly.
 #if !PLATFORM_LUMIN
@@ -760,3 +752,19 @@ void FVulkanDevice::ParseOptionalDeviceExtensions(const TArray<const ANSICHAR *>
 	OptionalDeviceExtensions.HasMemoryPriority = 0;
 #endif
 }
+
+void FVulkanDynamicRHI::SetupValidationRequests()
+{
+	int32 VulkanValidationOption = GValidationCvar.GetValueOnAnyThread();
+
+	// Command line overrides Cvar
+	if (FParse::Param(FCommandLine::Get(), TEXT("vulkandebug")))
+	{
+		// Match D3D and GL
+		GValidationCvar->Set(2, ECVF_SetByCommandline);
+	}
+	else if (FParse::Value(FCommandLine::Get(), TEXT("vulkanvalidation="), VulkanValidationOption))
+	{
+		GValidationCvar->Set(VulkanValidationOption, ECVF_SetByCommandline);
+	}
+}
diff --git a/Engine/Source/Runtime/VulkanRHI/Private/VulkanRHI.cpp b/Engine/Source/Runtime/VulkanRHI/Private/VulkanRHI.cpp
--- a/Engine/Source/Runtime/VulkanRHI/Private/VulkanRHI.cpp
+++ b/Engine/Source/Runtime/VulkanRHI/Private/VulkanRHI.cpp
@@ -196,6 +196,9 @@ FVulkanDynamicRHI::FVulkanDynamicRHI()
 
 void FVulkanDynamicRHI::Init()
 {
+	// Setup the validation requests ready before we load dlls
+	SetupValidationRequests();
+
 	if (!FVulkanPlatform::LoadVulkanLibrary())
 	{
 #if PLATFORM_LINUX
diff --git a/Engine/Source/Runtime/VulkanRHI/Public/VulkanDynamicRHI.h b/Engine/Source/Runtime/VulkanRHI/Public/VulkanDynamicRHI.h
--- a/Engine/Source/Runtime/VulkanRHI/Public/VulkanDynamicRHI.h
+++ b/Engine/Source/Runtime/VulkanRHI/Public/VulkanDynamicRHI.h
@@ -464,6 +464,8 @@ protected:
 	template <bool bRealUBs>
 	void UpdateUniformBuffer(FVulkanUniformBuffer* UniformBuffer, const void* Contents);
 
+	static void SetupValidationRequests();
+
 public:
 	static TSharedPtr< IHeadMountedDisplayVulkanExtensions, ESPMode::ThreadSafe > HMDVulkanExtensions;
 };
