DR - Deprecate FVertexShaderRHIParamRef, FHullShaderRHIParamRef, FDomainShaderRHIParamRef, FPixelShaderRHIParamRef, FGeometryShaderRHIParamRef and FComputeShaderRHIParamRef

diff --git a/Engine/Source/Runtime/VulkanRHI/Private/VulkanCommands.cpp b/Engine/Source/Runtime/VulkanRHI/Private/VulkanCommands.cpp
--- a/Engine/Source/Runtime/VulkanRHI/Private/VulkanCommands.cpp
+++ b/Engine/Source/Runtime/VulkanRHI/Private/VulkanCommands.cpp
@@ -62,7 +62,7 @@ void FVulkanDynamicRHI::RHISetStreamOutTargets(uint32 NumTargets, const FVertexB
 	VULKAN_SIGNAL_UNIMPLEMENTED();
 }
 
-void FVulkanCommandListContext::RHISetComputeShader(FComputeShaderRHIParamRef ComputeShaderRHI)
+void FVulkanCommandListContext::RHISetComputeShader(FRHIComputeShader* ComputeShaderRHI)
 {
 	FVulkanComputeShader* ComputeShader = ResourceCast(ComputeShaderRHI);
 	FVulkanComputePipeline* ComputePipeline = Device->GetPipelineStateCache()->GetOrCreateComputePipeline(ComputeShader);
@@ -156,7 +156,7 @@ void FVulkanCommandListContext::RHIDispatchIndirectComputeShader(FVertexBufferRH
 	VulkanRHI::/*Debug*/HeavyWeightBarrier(CmdBuffer/*, 2*/);
 }
 
-void FVulkanCommandListContext::RHISetUAVParameter(FComputeShaderRHIParamRef ComputeShaderRHI, uint32 UAVIndex, FRHIUnorderedAccessView* UAVRHI)
+void FVulkanCommandListContext::RHISetUAVParameter(FRHIComputeShader* ComputeShaderRHI, uint32 UAVIndex, FRHIUnorderedAccessView* UAVRHI)
 {
 	check(PendingComputeState->GetCurrentShader() == ResourceCast(ComputeShaderRHI));
 
@@ -168,7 +168,7 @@ void FVulkanCommandListContext::RHISetUAVParameter(FComputeShaderRHIParamRef Com
 	}
 }
 
-void FVulkanCommandListContext::RHISetUAVParameter(FComputeShaderRHIParamRef ComputeShaderRHI,uint32 UAVIndex, FRHIUnorderedAccessView* UAVRHI, uint32 InitialCount)
+void FVulkanCommandListContext::RHISetUAVParameter(FRHIComputeShader* ComputeShaderRHI,uint32 UAVIndex, FRHIUnorderedAccessView* UAVRHI, uint32 InitialCount)
 {
 	check(PendingComputeState->GetCurrentShader() == ResourceCast(ComputeShaderRHI));
 
@@ -177,7 +177,7 @@ void FVulkanCommandListContext::RHISetUAVParameter(FComputeShaderRHIParamRef Com
 }
 
 
-void FVulkanCommandListContext::RHISetShaderTexture(FVertexShaderRHIParamRef VertexShaderRHI, uint32 TextureIndex, FTextureRHIParamRef NewTextureRHI)
+void FVulkanCommandListContext::RHISetShaderTexture(FRHIVertexShader* VertexShaderRHI, uint32 TextureIndex, FTextureRHIParamRef NewTextureRHI)
 {
 	check(PendingGfxState->GetCurrentShaderKey(ShaderStage::Vertex) == GetShaderKey(VertexShaderRHI));
 	FVulkanTextureBase* Texture = GetVulkanTextureFromRHITexture(NewTextureRHI);
@@ -186,7 +186,7 @@ void FVulkanCommandListContext::RHISetShaderTexture(FVertexShaderRHIParamRef Ver
 	NewTextureRHI->SetLastRenderTime((float)FPlatformTime::Seconds());
 }
 
-void FVulkanCommandListContext::RHISetShaderTexture(FHullShaderRHIParamRef HullShaderRHI, uint32 TextureIndex, FTextureRHIParamRef NewTextureRHI)
+void FVulkanCommandListContext::RHISetShaderTexture(FRHIHullShader* HullShaderRHI, uint32 TextureIndex, FTextureRHIParamRef NewTextureRHI)
 {
 #if PLATFORM_SUPPORTS_TESSELLATION_SHADERS
 	check(PendingGfxState->GetCurrentShaderKey(ShaderStage::Hull) == GetShaderKey(HullShaderRHI));
@@ -199,7 +199,7 @@ void FVulkanCommandListContext::RHISetShaderTexture(FHullShaderRHIParamRef HullS
 #endif
 }
 
-void FVulkanCommandListContext::RHISetShaderTexture(FDomainShaderRHIParamRef DomainShaderRHI, uint32 TextureIndex, FTextureRHIParamRef NewTextureRHI)
+void FVulkanCommandListContext::RHISetShaderTexture(FRHIDomainShader* DomainShaderRHI, uint32 TextureIndex, FTextureRHIParamRef NewTextureRHI)
 {
 #if PLATFORM_SUPPORTS_TESSELLATION_SHADERS
 	check(PendingGfxState->GetCurrentShaderKey(ShaderStage::Domain) == GetShaderKey(DomainShaderRHI));
@@ -212,7 +212,7 @@ void FVulkanCommandListContext::RHISetShaderTexture(FDomainShaderRHIParamRef Dom
 #endif
 }
 
-void FVulkanCommandListContext::RHISetShaderTexture(FGeometryShaderRHIParamRef GeometryShaderRHI, uint32 TextureIndex, FTextureRHIParamRef NewTextureRHI)
+void FVulkanCommandListContext::RHISetShaderTexture(FRHIGeometryShader* GeometryShaderRHI, uint32 TextureIndex, FTextureRHIParamRef NewTextureRHI)
 {
 #if VULKAN_SUPPORTS_GEOMETRY_SHADERS
 	check(PendingGfxState->GetCurrentShaderKey(ShaderStage::Geometry) == GetShaderKey(GeometryShaderRHI));
@@ -225,7 +225,7 @@ void FVulkanCommandListContext::RHISetShaderTexture(FGeometryShaderRHIParamRef G
 #endif
 }
 
-void FVulkanCommandListContext::RHISetShaderTexture(FPixelShaderRHIParamRef PixelShaderRHI, uint32 TextureIndex, FTextureRHIParamRef NewTextureRHI)
+void FVulkanCommandListContext::RHISetShaderTexture(FRHIPixelShader* PixelShaderRHI, uint32 TextureIndex, FTextureRHIParamRef NewTextureRHI)
 {
 	check(PendingGfxState->GetCurrentShaderKey(ShaderStage::Pixel) == GetShaderKey(PixelShaderRHI));
 	FVulkanTextureBase* Texture = GetVulkanTextureFromRHITexture(NewTextureRHI);
@@ -234,7 +234,7 @@ void FVulkanCommandListContext::RHISetShaderTexture(FPixelShaderRHIParamRef Pixe
 	NewTextureRHI->SetLastRenderTime((float)FPlatformTime::Seconds());
 }
 
-void FVulkanCommandListContext::RHISetShaderTexture(FComputeShaderRHIParamRef ComputeShaderRHI, uint32 TextureIndex, FTextureRHIParamRef NewTextureRHI)
+void FVulkanCommandListContext::RHISetShaderTexture(FRHIComputeShader* ComputeShaderRHI, uint32 TextureIndex, FTextureRHIParamRef NewTextureRHI)
 {
 	FVulkanComputeShader* ComputeShader = ResourceCast(ComputeShaderRHI);
 	check(PendingComputeState->GetCurrentShader() == ComputeShader);
@@ -245,14 +245,14 @@ void FVulkanCommandListContext::RHISetShaderTexture(FComputeShaderRHIParamRef Co
 	NewTextureRHI->SetLastRenderTime((float)FPlatformTime::Seconds());
 }
 
-void FVulkanCommandListContext::RHISetShaderResourceViewParameter(FVertexShaderRHIParamRef VertexShaderRHI, uint32 TextureIndex, FRHIShaderResourceView* SRVRHI)
+void FVulkanCommandListContext::RHISetShaderResourceViewParameter(FRHIVertexShader* VertexShaderRHI, uint32 TextureIndex, FRHIShaderResourceView* SRVRHI)
 {
 	check(PendingGfxState->GetCurrentShaderKey(ShaderStage::Vertex) == GetShaderKey(VertexShaderRHI));
 	FVulkanShaderResourceView* SRV = ResourceCast(SRVRHI);
 	PendingGfxState->SetSRVForStage(ShaderStage::Vertex, TextureIndex, SRV);
 }
 
-void FVulkanCommandListContext::RHISetShaderResourceViewParameter(FHullShaderRHIParamRef HullShaderRHI,uint32 TextureIndex, FRHIShaderResourceView* SRVRHI)
+void FVulkanCommandListContext::RHISetShaderResourceViewParameter(FRHIHullShader* HullShaderRHI,uint32 TextureIndex, FRHIShaderResourceView* SRVRHI)
 {
 #if PLATFORM_SUPPORTS_TESSELLATION_SHADERS
 	check(PendingGfxState->GetCurrentShaderKey(ShaderStage::Hull) == GetShaderKey(HullShaderRHI));
@@ -263,7 +263,7 @@ void FVulkanCommandListContext::RHISetShaderResourceViewParameter(FHullShaderRHI
 #endif
 }
 
-void FVulkanCommandListContext::RHISetShaderResourceViewParameter(FDomainShaderRHIParamRef DomainShaderRHI,uint32 TextureIndex, FRHIShaderResourceView* SRVRHI)
+void FVulkanCommandListContext::RHISetShaderResourceViewParameter(FRHIDomainShader* DomainShaderRHI,uint32 TextureIndex, FRHIShaderResourceView* SRVRHI)
 {
 #if PLATFORM_SUPPORTS_TESSELLATION_SHADERS
 	check(PendingGfxState->GetCurrentShaderKey(ShaderStage::Domain) == GetShaderKey(DomainShaderRHI));
@@ -274,7 +274,7 @@ void FVulkanCommandListContext::RHISetShaderResourceViewParameter(FDomainShaderR
 #endif
 }
 
-void FVulkanCommandListContext::RHISetShaderResourceViewParameter(FGeometryShaderRHIParamRef GeometryShaderRHI,uint32 TextureIndex, FRHIShaderResourceView* SRVRHI)
+void FVulkanCommandListContext::RHISetShaderResourceViewParameter(FRHIGeometryShader* GeometryShaderRHI,uint32 TextureIndex, FRHIShaderResourceView* SRVRHI)
 {
 #if VULKAN_SUPPORTS_GEOMETRY_SHADERS
 	check(PendingGfxState->GetCurrentShaderKey(ShaderStage::Geometry) == GetShaderKey(GeometryShaderRHI));
@@ -285,14 +285,14 @@ void FVulkanCommandListContext::RHISetShaderResourceViewParameter(FGeometryShade
 #endif
 }
 
-void FVulkanCommandListContext::RHISetShaderResourceViewParameter(FPixelShaderRHIParamRef PixelShaderRHI,uint32 TextureIndex, FRHIShaderResourceView* SRVRHI)
+void FVulkanCommandListContext::RHISetShaderResourceViewParameter(FRHIPixelShader* PixelShaderRHI,uint32 TextureIndex, FRHIShaderResourceView* SRVRHI)
 {
 	check(PendingGfxState->GetCurrentShaderKey(ShaderStage::Pixel) == GetShaderKey(PixelShaderRHI));
 	FVulkanShaderResourceView* SRV = ResourceCast(SRVRHI);
 	PendingGfxState->SetSRVForStage(ShaderStage::Pixel, TextureIndex, SRV);
 }
 
-void FVulkanCommandListContext::RHISetShaderResourceViewParameter(FComputeShaderRHIParamRef ComputeShaderRHI,uint32 TextureIndex, FRHIShaderResourceView* SRVRHI)
+void FVulkanCommandListContext::RHISetShaderResourceViewParameter(FRHIComputeShader* ComputeShaderRHI,uint32 TextureIndex, FRHIShaderResourceView* SRVRHI)
 {
 	check(PendingComputeState->GetCurrentShader() == ResourceCast(ComputeShaderRHI));
 
@@ -300,14 +300,14 @@ void FVulkanCommandListContext::RHISetShaderResourceViewParameter(FComputeShader
 	PendingComputeState->SetSRVForStage(TextureIndex, SRV);
 }
 
-void FVulkanCommandListContext::RHISetShaderSampler(FVertexShaderRHIParamRef VertexShaderRHI, uint32 SamplerIndex, FRHISamplerState* NewStateRHI)
+void FVulkanCommandListContext::RHISetShaderSampler(FRHIVertexShader* VertexShaderRHI, uint32 SamplerIndex, FRHISamplerState* NewStateRHI)
 {
 	check(PendingGfxState->GetCurrentShaderKey(ShaderStage::Vertex) == GetShaderKey(VertexShaderRHI));
 	FVulkanSamplerState* Sampler = ResourceCast(NewStateRHI);
 	PendingGfxState->SetSamplerStateForStage(ShaderStage::Vertex, SamplerIndex, Sampler);
 }
 
-void FVulkanCommandListContext::RHISetShaderSampler(FHullShaderRHIParamRef HullShaderRHI, uint32 SamplerIndex, FRHISamplerState* NewStateRHI)
+void FVulkanCommandListContext::RHISetShaderSampler(FRHIHullShader* HullShaderRHI, uint32 SamplerIndex, FRHISamplerState* NewStateRHI)
 {
 #if PLATFORM_SUPPORTS_TESSELLATION_SHADERS
 	check(PendingGfxState->GetCurrentShaderKey(ShaderStage::Hull) == GetShaderKey(HullShaderRHI));
@@ -318,7 +318,7 @@ void FVulkanCommandListContext::RHISetShaderSampler(FHullShaderRHIParamRef HullS
 #endif
 }
 
-void FVulkanCommandListContext::RHISetShaderSampler(FDomainShaderRHIParamRef DomainShaderRHI, uint32 SamplerIndex, FRHISamplerState* NewStateRHI)
+void FVulkanCommandListContext::RHISetShaderSampler(FRHIDomainShader* DomainShaderRHI, uint32 SamplerIndex, FRHISamplerState* NewStateRHI)
 {
 #if PLATFORM_SUPPORTS_TESSELLATION_SHADERS
 	check(PendingGfxState->GetCurrentShaderKey(ShaderStage::Domain) == GetShaderKey(DomainShaderRHI));
@@ -329,7 +329,7 @@ void FVulkanCommandListContext::RHISetShaderSampler(FDomainShaderRHIParamRef Dom
 #endif
 }
 
-void FVulkanCommandListContext::RHISetShaderSampler(FGeometryShaderRHIParamRef GeometryShaderRHI, uint32 SamplerIndex, FRHISamplerState* NewStateRHI)
+void FVulkanCommandListContext::RHISetShaderSampler(FRHIGeometryShader* GeometryShaderRHI, uint32 SamplerIndex, FRHISamplerState* NewStateRHI)
 {
 #if VULKAN_SUPPORTS_GEOMETRY_SHADERS
 	check(PendingGfxState->GetCurrentShaderKey(ShaderStage::Geometry) == GetShaderKey(GeometryShaderRHI));
@@ -340,14 +340,14 @@ void FVulkanCommandListContext::RHISetShaderSampler(FGeometryShaderRHIParamRef G
 #endif
 }
 
-void FVulkanCommandListContext::RHISetShaderSampler(FPixelShaderRHIParamRef PixelShaderRHI, uint32 SamplerIndex, FRHISamplerState* NewStateRHI)
+void FVulkanCommandListContext::RHISetShaderSampler(FRHIPixelShader* PixelShaderRHI, uint32 SamplerIndex, FRHISamplerState* NewStateRHI)
 {
 	check(PendingGfxState->GetCurrentShaderKey(ShaderStage::Pixel) == GetShaderKey(PixelShaderRHI));
 	FVulkanSamplerState* Sampler = ResourceCast(NewStateRHI);
 	PendingGfxState->SetSamplerStateForStage(ShaderStage::Pixel, SamplerIndex, Sampler);
 }
 
-void FVulkanCommandListContext::RHISetShaderSampler(FComputeShaderRHIParamRef ComputeShaderRHI, uint32 SamplerIndex, FRHISamplerState* NewStateRHI)
+void FVulkanCommandListContext::RHISetShaderSampler(FRHIComputeShader* ComputeShaderRHI, uint32 SamplerIndex, FRHISamplerState* NewStateRHI)
 {
 	FVulkanComputeShader* ComputeShader = ResourceCast(ComputeShaderRHI);
 	check(PendingComputeState->GetCurrentShader() == ComputeShader);
@@ -356,14 +356,14 @@ void FVulkanCommandListContext::RHISetShaderSampler(FComputeShaderRHIParamRef Co
 	PendingComputeState->SetSamplerStateForStage(SamplerIndex, Sampler);
 }
 
-void FVulkanCommandListContext::RHISetShaderParameter(FVertexShaderRHIParamRef VertexShaderRHI, uint32 BufferIndex, uint32 BaseIndex, uint32 NumBytes, const void* NewValue)
+void FVulkanCommandListContext::RHISetShaderParameter(FRHIVertexShader* VertexShaderRHI, uint32 BufferIndex, uint32 BaseIndex, uint32 NumBytes, const void* NewValue)
 {
 	check(PendingGfxState->GetCurrentShaderKey(ShaderStage::Vertex) == GetShaderKey(VertexShaderRHI));
 
 	PendingGfxState->SetPackedGlobalShaderParameter(ShaderStage::Vertex, BufferIndex, BaseIndex, NumBytes, NewValue);
 }
 
-void FVulkanCommandListContext::RHISetShaderParameter(FHullShaderRHIParamRef HullShaderRHI, uint32 BufferIndex, uint32 BaseIndex, uint32 NumBytes, const void* NewValue)
+void FVulkanCommandListContext::RHISetShaderParameter(FRHIHullShader* HullShaderRHI, uint32 BufferIndex, uint32 BaseIndex, uint32 NumBytes, const void* NewValue)
 {
 #if PLATFORM_SUPPORTS_TESSELLATION_SHADERS
 	check(PendingGfxState->GetCurrentShaderKey(ShaderStage::Hull) == GetShaderKey(HullShaderRHI));
@@ -374,7 +374,7 @@ void FVulkanCommandListContext::RHISetShaderParameter(FHullShaderRHIParamRef Hul
 #endif
 }
 
-void FVulkanCommandListContext::RHISetShaderParameter(FDomainShaderRHIParamRef DomainShaderRHI, uint32 BufferIndex, uint32 BaseIndex, uint32 NumBytes, const void* NewValue)
+void FVulkanCommandListContext::RHISetShaderParameter(FRHIDomainShader* DomainShaderRHI, uint32 BufferIndex, uint32 BaseIndex, uint32 NumBytes, const void* NewValue)
 {
 #if PLATFORM_SUPPORTS_TESSELLATION_SHADERS
 	check(PendingGfxState->GetCurrentShaderKey(ShaderStage::Domain) == GetShaderKey(DomainShaderRHI));
@@ -385,7 +385,7 @@ void FVulkanCommandListContext::RHISetShaderParameter(FDomainShaderRHIParamRef D
 #endif
 }
 
-void FVulkanCommandListContext::RHISetShaderParameter(FGeometryShaderRHIParamRef GeometryShaderRHI, uint32 BufferIndex, uint32 BaseIndex, uint32 NumBytes, const void* NewValue)
+void FVulkanCommandListContext::RHISetShaderParameter(FRHIGeometryShader* GeometryShaderRHI, uint32 BufferIndex, uint32 BaseIndex, uint32 NumBytes, const void* NewValue)
 {
 #if VULKAN_SUPPORTS_GEOMETRY_SHADERS
 	check(PendingGfxState->GetCurrentShaderKey(ShaderStage::Geometry) == GetShaderKey(GeometryShaderRHI));
@@ -396,14 +396,14 @@ void FVulkanCommandListContext::RHISetShaderParameter(FGeometryShaderRHIParamRef
 #endif
 }
 
-void FVulkanCommandListContext::RHISetShaderParameter(FPixelShaderRHIParamRef PixelShaderRHI, uint32 BufferIndex, uint32 BaseIndex, uint32 NumBytes, const void* NewValue)
+void FVulkanCommandListContext::RHISetShaderParameter(FRHIPixelShader* PixelShaderRHI, uint32 BufferIndex, uint32 BaseIndex, uint32 NumBytes, const void* NewValue)
 {
 	check(PendingGfxState->GetCurrentShaderKey(ShaderStage::Pixel) == GetShaderKey(PixelShaderRHI));
 
 	PendingGfxState->SetPackedGlobalShaderParameter(ShaderStage::Pixel, BufferIndex, BaseIndex, NumBytes, NewValue);
 }
 
-void FVulkanCommandListContext::RHISetShaderParameter(FComputeShaderRHIParamRef ComputeShaderRHI,uint32 BufferIndex, uint32 BaseIndex, uint32 NumBytes, const void* NewValue)
+void FVulkanCommandListContext::RHISetShaderParameter(FRHIComputeShader* ComputeShaderRHI,uint32 BufferIndex, uint32 BaseIndex, uint32 NumBytes, const void* NewValue)
 {
 	FVulkanComputeShader* ComputeShader = ResourceCast(ComputeShaderRHI);
 	check(PendingComputeState->GetCurrentShader() == ComputeShader);
@@ -647,13 +647,13 @@ inline void FVulkanCommandListContext::SetShaderUniformBuffer(ShaderStage::EStag
 	}
 }
 
-void FVulkanCommandListContext::RHISetShaderUniformBuffer(FVertexShaderRHIParamRef VertexShaderRHI, uint32 BufferIndex, FUniformBufferRHIParamRef BufferRHI)
+void FVulkanCommandListContext::RHISetShaderUniformBuffer(FRHIVertexShader* VertexShaderRHI, uint32 BufferIndex, FUniformBufferRHIParamRef BufferRHI)
 {
 	FVulkanUniformBuffer* UniformBuffer = ResourceCast(BufferRHI);
 	SetShaderUniformBuffer(ShaderStage::Vertex, UniformBuffer, BufferIndex, ResourceCast(VertexShaderRHI));
 }
 
-void FVulkanCommandListContext::RHISetShaderUniformBuffer(FHullShaderRHIParamRef HullShaderRHI, uint32 BufferIndex, FUniformBufferRHIParamRef BufferRHI)
+void FVulkanCommandListContext::RHISetShaderUniformBuffer(FRHIHullShader* HullShaderRHI, uint32 BufferIndex, FUniformBufferRHIParamRef BufferRHI)
 {
 #if PLATFORM_SUPPORTS_TESSELLATION_SHADERS
 	FVulkanUniformBuffer* UniformBuffer = ResourceCast(BufferRHI);
@@ -663,7 +663,7 @@ void FVulkanCommandListContext::RHISetShaderUniformBuffer(FHullShaderRHIParamRef
 #endif
 }
 
-void FVulkanCommandListContext::RHISetShaderUniformBuffer(FDomainShaderRHIParamRef DomainShaderRHI, uint32 BufferIndex, FUniformBufferRHIParamRef BufferRHI)
+void FVulkanCommandListContext::RHISetShaderUniformBuffer(FRHIDomainShader* DomainShaderRHI, uint32 BufferIndex, FUniformBufferRHIParamRef BufferRHI)
 {
 #if PLATFORM_SUPPORTS_TESSELLATION_SHADERS
 	FVulkanUniformBuffer* UniformBuffer = ResourceCast(BufferRHI);
@@ -673,7 +673,7 @@ void FVulkanCommandListContext::RHISetShaderUniformBuffer(FDomainShaderRHIParamR
 #endif
 }
 
-void FVulkanCommandListContext::RHISetShaderUniformBuffer(FGeometryShaderRHIParamRef GeometryShaderRHI, uint32 BufferIndex, FUniformBufferRHIParamRef BufferRHI)
+void FVulkanCommandListContext::RHISetShaderUniformBuffer(FRHIGeometryShader* GeometryShaderRHI, uint32 BufferIndex, FUniformBufferRHIParamRef BufferRHI)
 {
 #if VULKAN_SUPPORTS_GEOMETRY_SHADERS
 	FVulkanUniformBuffer* UniformBuffer = ResourceCast(BufferRHI);
@@ -683,13 +683,13 @@ void FVulkanCommandListContext::RHISetShaderUniformBuffer(FGeometryShaderRHIPara
 #endif
 }
 
-void FVulkanCommandListContext::RHISetShaderUniformBuffer(FPixelShaderRHIParamRef PixelShaderRHI, uint32 BufferIndex, FUniformBufferRHIParamRef BufferRHI)
+void FVulkanCommandListContext::RHISetShaderUniformBuffer(FRHIPixelShader* PixelShaderRHI, uint32 BufferIndex, FUniformBufferRHIParamRef BufferRHI)
 {
 	FVulkanUniformBuffer* UniformBuffer = ResourceCast(BufferRHI);
 	SetShaderUniformBuffer(ShaderStage::Pixel, UniformBuffer, BufferIndex, ResourceCast(PixelShaderRHI));
 }
 
-void FVulkanCommandListContext::RHISetShaderUniformBuffer(FComputeShaderRHIParamRef ComputeShaderRHI, uint32 BufferIndex, FUniformBufferRHIParamRef BufferRHI)
+void FVulkanCommandListContext::RHISetShaderUniformBuffer(FRHIComputeShader* ComputeShaderRHI, uint32 BufferIndex, FUniformBufferRHIParamRef BufferRHI)
 {
 	FVulkanComputeShader* ComputeShader = ResourceCast(ComputeShaderRHI);
 	check(PendingComputeState->GetCurrentShader() == ComputeShader);
diff --git a/Engine/Source/Runtime/VulkanRHI/Private/VulkanContext.h b/Engine/Source/Runtime/VulkanRHI/Private/VulkanContext.h
--- a/Engine/Source/Runtime/VulkanRHI/Private/VulkanContext.h
+++ b/Engine/Source/Runtime/VulkanRHI/Private/VulkanContext.h
@@ -173,38 +173,38 @@ public:
 	virtual void RHISetViewport(uint32 MinX, uint32 MinY, float MinZ, uint32 MaxX, uint32 MaxY, float MaxZ) final override;
 	virtual void RHISetScissorRect(bool bEnable, uint32 MinX, uint32 MinY, uint32 MaxX, uint32 MaxY) final override;
 	virtual void RHISetGraphicsPipelineState(FRHIGraphicsPipelineState* GraphicsState) final override;
-	virtual void RHISetShaderTexture(FVertexShaderRHIParamRef VertexShader, uint32 TextureIndex, FTextureRHIParamRef NewTexture) final override;
-	virtual void RHISetShaderTexture(FHullShaderRHIParamRef HullShader, uint32 TextureIndex, FTextureRHIParamRef NewTexture) final override;
-	virtual void RHISetShaderTexture(FDomainShaderRHIParamRef DomainShader, uint32 TextureIndex, FTextureRHIParamRef NewTexture) final override;
-	virtual void RHISetShaderTexture(FGeometryShaderRHIParamRef GeometryShader, uint32 TextureIndex, FTextureRHIParamRef NewTexture) final override;
-	virtual void RHISetShaderTexture(FPixelShaderRHIParamRef PixelShader, uint32 TextureIndex, FTextureRHIParamRef NewTexture) final override;
-	virtual void RHISetShaderTexture(FComputeShaderRHIParamRef PixelShader, uint32 TextureIndex, FTextureRHIParamRef NewTexture) final override;
-	virtual void RHISetShaderSampler(FComputeShaderRHIParamRef ComputeShader, uint32 SamplerIndex, FRHISamplerState* NewState) final override;
-	virtual void RHISetShaderSampler(FVertexShaderRHIParamRef VertexShader, uint32 SamplerIndex, FRHISamplerState* NewState) final override;
-	virtual void RHISetShaderSampler(FGeometryShaderRHIParamRef GeometryShader, uint32 SamplerIndex, FRHISamplerState* NewState) final override;
-	virtual void RHISetShaderSampler(FDomainShaderRHIParamRef DomainShader, uint32 SamplerIndex, FRHISamplerState* NewState) final override;
-	virtual void RHISetShaderSampler(FHullShaderRHIParamRef HullShader, uint32 SamplerIndex, FRHISamplerState* NewState) final override;
-	virtual void RHISetShaderSampler(FPixelShaderRHIParamRef PixelShader, uint32 SamplerIndex, FRHISamplerState* NewState) final override;
-	virtual void RHISetUAVParameter(FComputeShaderRHIParamRef ComputeShader, uint32 UAVIndex, FRHIUnorderedAccessView* UAV) final override;
-	virtual void RHISetUAVParameter(FComputeShaderRHIParamRef ComputeShader, uint32 UAVIndex, FRHIUnorderedAccessView* UAV, uint32 InitialCount) final override;
-	virtual void RHISetShaderResourceViewParameter(FPixelShaderRHIParamRef PixelShader, uint32 SamplerIndex, FRHIShaderResourceView* SRV) final override;
-	virtual void RHISetShaderResourceViewParameter(FVertexShaderRHIParamRef VertexShader, uint32 SamplerIndex, FRHIShaderResourceView* SRV) final override;
-	virtual void RHISetShaderResourceViewParameter(FComputeShaderRHIParamRef ComputeShader, uint32 SamplerIndex, FRHIShaderResourceView* SRV) final override;
-	virtual void RHISetShaderResourceViewParameter(FHullShaderRHIParamRef HullShader, uint32 SamplerIndex, FRHIShaderResourceView* SRV) final override;
-	virtual void RHISetShaderResourceViewParameter(FDomainShaderRHIParamRef DomainShader, uint32 SamplerIndex, FRHIShaderResourceView* SRV) final override;
-	virtual void RHISetShaderResourceViewParameter(FGeometryShaderRHIParamRef GeometryShader, uint32 SamplerIndex, FRHIShaderResourceView* SRV) final override;
-	virtual void RHISetShaderUniformBuffer(FVertexShaderRHIParamRef VertexShader, uint32 BufferIndex, FUniformBufferRHIParamRef Buffer) final override;
-	virtual void RHISetShaderUniformBuffer(FHullShaderRHIParamRef HullShader, uint32 BufferIndex, FUniformBufferRHIParamRef Buffer) final override;
-	virtual void RHISetShaderUniformBuffer(FDomainShaderRHIParamRef DomainShader, uint32 BufferIndex, FUniformBufferRHIParamRef Buffer) final override;
-	virtual void RHISetShaderUniformBuffer(FGeometryShaderRHIParamRef GeometryShader, uint32 BufferIndex, FUniformBufferRHIParamRef Buffer) final override;
-	virtual void RHISetShaderUniformBuffer(FPixelShaderRHIParamRef PixelShader, uint32 BufferIndex, FUniformBufferRHIParamRef Buffer) final override;
-	virtual void RHISetShaderUniformBuffer(FComputeShaderRHIParamRef ComputeShader, uint32 BufferIndex, FUniformBufferRHIParamRef Buffer) final override;
-	virtual void RHISetShaderParameter(FVertexShaderRHIParamRef VertexShader, uint32 BufferIndex, uint32 BaseIndex, uint32 NumBytes, const void* NewValue) final override;
-	virtual void RHISetShaderParameter(FHullShaderRHIParamRef HullShader, uint32 BufferIndex, uint32 BaseIndex, uint32 NumBytes, const void* NewValue) final override;
-	virtual void RHISetShaderParameter(FDomainShaderRHIParamRef DomainShader, uint32 BufferIndex, uint32 BaseIndex, uint32 NumBytes, const void* NewValue) final override;
-	virtual void RHISetShaderParameter(FGeometryShaderRHIParamRef GeometryShader, uint32 BufferIndex, uint32 BaseIndex, uint32 NumBytes, const void* NewValue) final override;
-	virtual void RHISetShaderParameter(FPixelShaderRHIParamRef PixelShader, uint32 BufferIndex, uint32 BaseIndex, uint32 NumBytes, const void* NewValue) final override;
-	virtual void RHISetShaderParameter(FComputeShaderRHIParamRef ComputeShader, uint32 BufferIndex, uint32 BaseIndex, uint32 NumBytes, const void* NewValue) final override;
+	virtual void RHISetShaderTexture(FRHIVertexShader* VertexShader, uint32 TextureIndex, FTextureRHIParamRef NewTexture) final override;
+	virtual void RHISetShaderTexture(FRHIHullShader* HullShader, uint32 TextureIndex, FTextureRHIParamRef NewTexture) final override;
+	virtual void RHISetShaderTexture(FRHIDomainShader* DomainShader, uint32 TextureIndex, FTextureRHIParamRef NewTexture) final override;
+	virtual void RHISetShaderTexture(FRHIGeometryShader* GeometryShader, uint32 TextureIndex, FTextureRHIParamRef NewTexture) final override;
+	virtual void RHISetShaderTexture(FRHIPixelShader* PixelShader, uint32 TextureIndex, FTextureRHIParamRef NewTexture) final override;
+	virtual void RHISetShaderTexture(FRHIComputeShader* PixelShader, uint32 TextureIndex, FTextureRHIParamRef NewTexture) final override;
+	virtual void RHISetShaderSampler(FRHIComputeShader* ComputeShader, uint32 SamplerIndex, FRHISamplerState* NewState) final override;
+	virtual void RHISetShaderSampler(FRHIVertexShader* VertexShader, uint32 SamplerIndex, FRHISamplerState* NewState) final override;
+	virtual void RHISetShaderSampler(FRHIGeometryShader* GeometryShader, uint32 SamplerIndex, FRHISamplerState* NewState) final override;
+	virtual void RHISetShaderSampler(FRHIDomainShader* DomainShader, uint32 SamplerIndex, FRHISamplerState* NewState) final override;
+	virtual void RHISetShaderSampler(FRHIHullShader* HullShader, uint32 SamplerIndex, FRHISamplerState* NewState) final override;
+	virtual void RHISetShaderSampler(FRHIPixelShader* PixelShader, uint32 SamplerIndex, FRHISamplerState* NewState) final override;
+	virtual void RHISetUAVParameter(FRHIComputeShader* ComputeShader, uint32 UAVIndex, FRHIUnorderedAccessView* UAV) final override;
+	virtual void RHISetUAVParameter(FRHIComputeShader* ComputeShader, uint32 UAVIndex, FRHIUnorderedAccessView* UAV, uint32 InitialCount) final override;
+	virtual void RHISetShaderResourceViewParameter(FRHIPixelShader* PixelShader, uint32 SamplerIndex, FRHIShaderResourceView* SRV) final override;
+	virtual void RHISetShaderResourceViewParameter(FRHIVertexShader* VertexShader, uint32 SamplerIndex, FRHIShaderResourceView* SRV) final override;
+	virtual void RHISetShaderResourceViewParameter(FRHIComputeShader* ComputeShader, uint32 SamplerIndex, FRHIShaderResourceView* SRV) final override;
+	virtual void RHISetShaderResourceViewParameter(FRHIHullShader* HullShader, uint32 SamplerIndex, FRHIShaderResourceView* SRV) final override;
+	virtual void RHISetShaderResourceViewParameter(FRHIDomainShader* DomainShader, uint32 SamplerIndex, FRHIShaderResourceView* SRV) final override;
+	virtual void RHISetShaderResourceViewParameter(FRHIGeometryShader* GeometryShader, uint32 SamplerIndex, FRHIShaderResourceView* SRV) final override;
+	virtual void RHISetShaderUniformBuffer(FRHIVertexShader* VertexShader, uint32 BufferIndex, FUniformBufferRHIParamRef Buffer) final override;
+	virtual void RHISetShaderUniformBuffer(FRHIHullShader* HullShader, uint32 BufferIndex, FUniformBufferRHIParamRef Buffer) final override;
+	virtual void RHISetShaderUniformBuffer(FRHIDomainShader* DomainShader, uint32 BufferIndex, FUniformBufferRHIParamRef Buffer) final override;
+	virtual void RHISetShaderUniformBuffer(FRHIGeometryShader* GeometryShader, uint32 BufferIndex, FUniformBufferRHIParamRef Buffer) final override;
+	virtual void RHISetShaderUniformBuffer(FRHIPixelShader* PixelShader, uint32 BufferIndex, FUniformBufferRHIParamRef Buffer) final override;
+	virtual void RHISetShaderUniformBuffer(FRHIComputeShader* ComputeShader, uint32 BufferIndex, FUniformBufferRHIParamRef Buffer) final override;
+	virtual void RHISetShaderParameter(FRHIVertexShader* VertexShader, uint32 BufferIndex, uint32 BaseIndex, uint32 NumBytes, const void* NewValue) final override;
+	virtual void RHISetShaderParameter(FRHIHullShader* HullShader, uint32 BufferIndex, uint32 BaseIndex, uint32 NumBytes, const void* NewValue) final override;
+	virtual void RHISetShaderParameter(FRHIDomainShader* DomainShader, uint32 BufferIndex, uint32 BaseIndex, uint32 NumBytes, const void* NewValue) final override;
+	virtual void RHISetShaderParameter(FRHIGeometryShader* GeometryShader, uint32 BufferIndex, uint32 BaseIndex, uint32 NumBytes, const void* NewValue) final override;
+	virtual void RHISetShaderParameter(FRHIPixelShader* PixelShader, uint32 BufferIndex, uint32 BaseIndex, uint32 NumBytes, const void* NewValue) final override;
+	virtual void RHISetShaderParameter(FRHIComputeShader* ComputeShader, uint32 BufferIndex, uint32 BaseIndex, uint32 NumBytes, const void* NewValue) final override;
 	virtual void RHISetStencilRef(uint32 StencilRef) final override;
 	virtual void RHISetRenderTargets(uint32 NumSimultaneousRenderTargets, const FRHIRenderTargetView* NewRenderTargets, const FRHIDepthRenderTargetView* NewDepthStencilTarget, uint32 NumUAVs, FRHIUnorderedAccessView* const* UAVs) final override;
 	virtual void RHISetRenderTargetsAndClear(const FRHISetRenderTargetsInfo& RenderTargetsInfo) final override;
@@ -217,7 +217,7 @@ public:
 	virtual void RHIPushEvent(const TCHAR* Name, FColor Color) final override;
 	virtual void RHIPopEvent() final override;
 
-	virtual void RHISetComputeShader(FComputeShaderRHIParamRef ComputeShader) final override;
+	virtual void RHISetComputeShader(FRHIComputeShader* ComputeShader) final override;
 	virtual void RHISetComputePipelineState(FRHIComputePipelineState* ComputePipelineState) final override;
 	virtual void RHIWaitComputeFence(FRHIComputeFence* InFence) final override;
 	virtual void RHIDispatchComputeShader(uint32 ThreadGroupCountX, uint32 ThreadGroupCountY, uint32 ThreadGroupCountZ) final override;
diff --git a/Engine/Source/Runtime/VulkanRHI/Private/VulkanResources.h b/Engine/Source/Runtime/VulkanRHI/Private/VulkanResources.h
--- a/Engine/Source/Runtime/VulkanRHI/Private/VulkanResources.h
+++ b/Engine/Source/Runtime/VulkanRHI/Private/VulkanResources.h
@@ -221,11 +221,11 @@ class FVulkanBoundShaderState : public FRHIBoundShaderState
 public:
 	FVulkanBoundShaderState(
 		FRHIVertexDeclaration* InVertexDeclarationRHI,
-		FVertexShaderRHIParamRef InVertexShaderRHI,
-		FPixelShaderRHIParamRef InPixelShaderRHI,
-		FHullShaderRHIParamRef InHullShaderRHI,
-		FDomainShaderRHIParamRef InDomainShaderRHI,
-		FGeometryShaderRHIParamRef InGeometryShaderRHI
+		FRHIVertexShader* InVertexShaderRHI,
+		FRHIPixelShader* InPixelShaderRHI,
+		FRHIHullShader* InHullShaderRHI,
+		FRHIDomainShader* InDomainShaderRHI,
+		FRHIGeometryShader* InGeometryShaderRHI
 	);
 
 	virtual ~FVulkanBoundShaderState();
diff --git a/Engine/Source/Runtime/VulkanRHI/Private/VulkanShaders.cpp b/Engine/Source/Runtime/VulkanRHI/Private/VulkanShaders.cpp
--- a/Engine/Source/Runtime/VulkanRHI/Private/VulkanShaders.cpp
+++ b/Engine/Source/Runtime/VulkanRHI/Private/VulkanShaders.cpp
@@ -638,9 +638,9 @@ void FVulkanGfxPipelineDescriptorInfo::Initialize(const FDescriptorSetRemappingI
 }
 
 
-FVulkanBoundShaderState::FVulkanBoundShaderState(FRHIVertexDeclaration* InVertexDeclarationRHI, FVertexShaderRHIParamRef InVertexShaderRHI,
-	FPixelShaderRHIParamRef InPixelShaderRHI, FHullShaderRHIParamRef InHullShaderRHI,
-	FDomainShaderRHIParamRef InDomainShaderRHI, FGeometryShaderRHIParamRef InGeometryShaderRHI)
+FVulkanBoundShaderState::FVulkanBoundShaderState(FRHIVertexDeclaration* InVertexDeclarationRHI, FRHIVertexShader* InVertexShaderRHI,
+	FRHIPixelShader* InPixelShaderRHI, FRHIHullShader* InHullShaderRHI,
+	FRHIDomainShader* InDomainShaderRHI, FRHIGeometryShader* InGeometryShaderRHI)
 	: CacheLink(InVertexDeclarationRHI, InVertexShaderRHI, InPixelShaderRHI, InHullShaderRHI, InDomainShaderRHI, InGeometryShaderRHI, this)
 {
 	CacheLink.AddToCache();
@@ -653,11 +653,11 @@ FVulkanBoundShaderState::~FVulkanBoundShaderState()
 
 FBoundShaderStateRHIRef FVulkanDynamicRHI::RHICreateBoundShaderState(
 	FRHIVertexDeclaration* VertexDeclarationRHI,
-	FVertexShaderRHIParamRef VertexShaderRHI, 
-	FHullShaderRHIParamRef HullShaderRHI, 
-	FDomainShaderRHIParamRef DomainShaderRHI, 
-	FPixelShaderRHIParamRef PixelShaderRHI,
-	FGeometryShaderRHIParamRef GeometryShaderRHI
+	FRHIVertexShader* VertexShaderRHI,
+	FRHIHullShader* HullShaderRHI,
+	FRHIDomainShader* DomainShaderRHI,
+	FRHIPixelShader* PixelShaderRHI,
+	FRHIGeometryShader* GeometryShaderRHI
 	)
 {
 	LLM_SCOPE_VULKAN(ELLMTagVulkan::VulkanShaders);
diff --git a/Engine/Source/Runtime/Windows/D3D11RHI/Private/D3D11Commands.cpp b/Engine/Source/Runtime/Windows/D3D11RHI/Private/D3D11Commands.cpp
--- a/Engine/Source/Runtime/Windows/D3D11RHI/Private/D3D11Commands.cpp
+++ b/Engine/Source/Runtime/Windows/D3D11RHI/Private/D3D11Commands.cpp
@@ -23,7 +23,7 @@
 #include "Windows/HideWindowsPlatformTypes.h"
 #endif
 
-#define DECLARE_ISBOUNDSHADER(ShaderType) inline void ValidateBoundShader(FD3D11StateCache& InStateCache, F##ShaderType##RHIParamRef ShaderType##RHI) \
+#define DECLARE_ISBOUNDSHADER(ShaderType) inline void ValidateBoundShader(FD3D11StateCache& InStateCache, FRHI##ShaderType* ShaderType##RHI) \
 { \
 	ID3D11##ShaderType* CachedShader; \
 	InStateCache.Get##ShaderType(&CachedShader); \
@@ -185,7 +185,7 @@ void FD3D11DynamicRHI::RHISetRasterizerState(FRHIRasterizerState* NewStateRHI)
 	StateCache.SetRasterizerState(NewState->Resource);
 }
 
-void FD3D11DynamicRHI::RHISetComputeShader(FComputeShaderRHIParamRef ComputeShaderRHI)
+void FD3D11DynamicRHI::RHISetComputeShader(FRHIComputeShader* ComputeShaderRHI)
 {
 	FD3D11ComputeShader* ComputeShader = ResourceCast(ComputeShaderRHI);
 	SetCurrentComputeShader(ComputeShaderRHI);
@@ -198,7 +198,7 @@ void FD3D11DynamicRHI::RHISetComputeShader(FComputeShaderRHIParamRef ComputeShad
 
 void FD3D11DynamicRHI::RHIDispatchComputeShader(uint32 ThreadGroupCountX, uint32 ThreadGroupCountY, uint32 ThreadGroupCountZ) 
 { 
-	FComputeShaderRHIParamRef ComputeShaderRHI = GetCurrentComputeShader();
+	FRHIComputeShader* ComputeShaderRHI = GetCurrentComputeShader();
 	FD3D11ComputeShader* ComputeShader = ResourceCast(ComputeShaderRHI);
 
 	StateCache.SetComputeShader(ComputeShader->Resource);
@@ -217,7 +217,7 @@ void FD3D11DynamicRHI::RHIDispatchComputeShader(uint32 ThreadGroupCountX, uint32
 
 void FD3D11DynamicRHI::RHIDispatchIndirectComputeShader(FVertexBufferRHIParamRef ArgumentBufferRHI, uint32 ArgumentOffset) 
 { 
-	FComputeShaderRHIParamRef ComputeShaderRHI = GetCurrentComputeShader();
+	FRHIComputeShader* ComputeShaderRHI = GetCurrentComputeShader();
 	FD3D11ComputeShader* ComputeShader = ResourceCast(ComputeShaderRHI);
 	FD3D11VertexBuffer* ArgumentBuffer = ResourceCast(ArgumentBufferRHI);
 
@@ -333,7 +333,7 @@ void FD3D11DynamicRHI::RHISetBoundShaderState(FRHIBoundShaderState* BoundShaderS
 	}
 }
 
-void FD3D11DynamicRHI::RHISetShaderTexture(FVertexShaderRHIParamRef VertexShaderRHI,uint32 TextureIndex,FTextureRHIParamRef NewTextureRHI)
+void FD3D11DynamicRHI::RHISetShaderTexture(FRHIVertexShader* VertexShaderRHI,uint32 TextureIndex,FTextureRHIParamRef NewTextureRHI)
 {
 	VALIDATE_BOUND_SHADER(VertexShaderRHI);
 
@@ -350,7 +350,7 @@ void FD3D11DynamicRHI::RHISetShaderTexture(FVertexShaderRHIParamRef VertexShader
 	}
 }
 
-void FD3D11DynamicRHI::RHISetShaderTexture(FHullShaderRHIParamRef HullShaderRHI,uint32 TextureIndex,FTextureRHIParamRef NewTextureRHI)
+void FD3D11DynamicRHI::RHISetShaderTexture(FRHIHullShader* HullShaderRHI,uint32 TextureIndex,FTextureRHIParamRef NewTextureRHI)
 {
 	VALIDATE_BOUND_SHADER(HullShaderRHI);
 
@@ -367,7 +367,7 @@ void FD3D11DynamicRHI::RHISetShaderTexture(FHullShaderRHIParamRef HullShaderRHI,
 	}
 }
 
-void FD3D11DynamicRHI::RHISetShaderTexture(FDomainShaderRHIParamRef DomainShaderRHI,uint32 TextureIndex,FTextureRHIParamRef NewTextureRHI)
+void FD3D11DynamicRHI::RHISetShaderTexture(FRHIDomainShader* DomainShaderRHI,uint32 TextureIndex,FTextureRHIParamRef NewTextureRHI)
 {
 	VALIDATE_BOUND_SHADER(DomainShaderRHI);
 
@@ -384,7 +384,7 @@ void FD3D11DynamicRHI::RHISetShaderTexture(FDomainShaderRHIParamRef DomainShader
 	}
 }
 
-void FD3D11DynamicRHI::RHISetShaderTexture(FGeometryShaderRHIParamRef GeometryShaderRHI,uint32 TextureIndex,FTextureRHIParamRef NewTextureRHI)
+void FD3D11DynamicRHI::RHISetShaderTexture(FRHIGeometryShader* GeometryShaderRHI,uint32 TextureIndex,FTextureRHIParamRef NewTextureRHI)
 {
 	VALIDATE_BOUND_SHADER(GeometryShaderRHI);
 
@@ -401,9 +401,8 @@ void FD3D11DynamicRHI::RHISetShaderTexture(FGeometryShaderRHIParamRef GeometrySh
 	}
 }
 
-void FD3D11DynamicRHI::RHISetShaderTexture(FPixelShaderRHIParamRef PixelShaderRHI,uint32 TextureIndex,FTextureRHIParamRef NewTextureRHI)
+void FD3D11DynamicRHI::RHISetShaderTexture(FRHIPixelShader* PixelShaderRHI,uint32 TextureIndex,FTextureRHIParamRef NewTextureRHI)
 {
-
 	VALIDATE_BOUND_SHADER(PixelShaderRHI);
 
 	FD3D11TextureBase* NewTexture = GetD3D11TextureFromRHITexture(NewTextureRHI);
@@ -418,7 +417,7 @@ void FD3D11DynamicRHI::RHISetShaderTexture(FPixelShaderRHIParamRef PixelShaderRH
 	}
 }
 
-void FD3D11DynamicRHI::RHISetShaderTexture(FComputeShaderRHIParamRef ComputeShaderRHI,uint32 TextureIndex,FTextureRHIParamRef NewTextureRHI)
+void FD3D11DynamicRHI::RHISetShaderTexture(FRHIComputeShader* ComputeShaderRHI,uint32 TextureIndex,FTextureRHIParamRef NewTextureRHI)
 {
 	//VALIDATE_BOUND_SHADER(ComputeShaderRHI);
 
@@ -435,7 +434,7 @@ void FD3D11DynamicRHI::RHISetShaderTexture(FComputeShaderRHIParamRef ComputeShad
 	}
 }
 
-void FD3D11DynamicRHI::RHISetUAVParameter(FComputeShaderRHIParamRef ComputeShaderRHI,uint32 UAVIndex, FRHIUnorderedAccessView* UAVRHI)
+void FD3D11DynamicRHI::RHISetUAVParameter(FRHIComputeShader* ComputeShaderRHI,uint32 UAVIndex, FRHIUnorderedAccessView* UAVRHI)
 {
 	//VALIDATE_BOUND_SHADER(ComputeShaderRHI);
 
@@ -460,7 +459,7 @@ void FD3D11DynamicRHI::RHISetUAVParameter(FComputeShaderRHIParamRef ComputeShade
 	Direct3DDeviceIMContext->CSSetUnorderedAccessViews(UAVIndex,1,&D3D11UAV, &InitialCount );
 }
 
-void FD3D11DynamicRHI::RHISetUAVParameter(FComputeShaderRHIParamRef ComputeShaderRHI,uint32 UAVIndex, FRHIUnorderedAccessView* UAVRHI, uint32 InitialCount )
+void FD3D11DynamicRHI::RHISetUAVParameter(FRHIComputeShader* ComputeShaderRHI,uint32 UAVIndex, FRHIUnorderedAccessView* UAVRHI, uint32 InitialCount )
 {
 	//VALIDATE_BOUND_SHADER(ComputeShaderRHI);
 
@@ -483,7 +482,7 @@ void FD3D11DynamicRHI::RHISetUAVParameter(FComputeShaderRHIParamRef ComputeShade
 	Direct3DDeviceIMContext->CSSetUnorderedAccessViews(UAVIndex,1,&D3D11UAV, &InitialCount );
 }
 
-void FD3D11DynamicRHI::RHISetShaderResourceViewParameter(FPixelShaderRHIParamRef PixelShaderRHI,uint32 TextureIndex, FRHIShaderResourceView* SRVRHI)
+void FD3D11DynamicRHI::RHISetShaderResourceViewParameter(FRHIPixelShader* PixelShaderRHI,uint32 TextureIndex, FRHIShaderResourceView* SRVRHI)
 {
 	VALIDATE_BOUND_SHADER(PixelShaderRHI);
 
@@ -501,7 +500,7 @@ void FD3D11DynamicRHI::RHISetShaderResourceViewParameter(FPixelShaderRHIParamRef
 	SetShaderResourceView<SF_Pixel>(Resource, D3D11SRV, TextureIndex, NAME_None);
 }
 
-void FD3D11DynamicRHI::RHISetShaderResourceViewParameter(FVertexShaderRHIParamRef VertexShaderRHI,uint32 TextureIndex, FRHIShaderResourceView* SRVRHI)
+void FD3D11DynamicRHI::RHISetShaderResourceViewParameter(FRHIVertexShader* VertexShaderRHI,uint32 TextureIndex, FRHIShaderResourceView* SRVRHI)
 {
 	VALIDATE_BOUND_SHADER(VertexShaderRHI);
 
@@ -519,7 +518,7 @@ void FD3D11DynamicRHI::RHISetShaderResourceViewParameter(FVertexShaderRHIParamRe
 	SetShaderResourceView<SF_Vertex>(Resource, D3D11SRV, TextureIndex, NAME_None);
 }
 
-void FD3D11DynamicRHI::RHISetShaderResourceViewParameter(FComputeShaderRHIParamRef ComputeShaderRHI,uint32 TextureIndex, FRHIShaderResourceView* SRVRHI)
+void FD3D11DynamicRHI::RHISetShaderResourceViewParameter(FRHIComputeShader* ComputeShaderRHI,uint32 TextureIndex, FRHIShaderResourceView* SRVRHI)
 {
 	//VALIDATE_BOUND_SHADER(ComputeShaderRHI);
 
@@ -537,7 +536,7 @@ void FD3D11DynamicRHI::RHISetShaderResourceViewParameter(FComputeShaderRHIParamR
 	SetShaderResourceView<SF_Compute>(Resource, D3D11SRV, TextureIndex, NAME_None);
 }
 
-void FD3D11DynamicRHI::RHISetShaderResourceViewParameter(FHullShaderRHIParamRef HullShaderRHI,uint32 TextureIndex, FRHIShaderResourceView* SRVRHI)
+void FD3D11DynamicRHI::RHISetShaderResourceViewParameter(FRHIHullShader* HullShaderRHI,uint32 TextureIndex, FRHIShaderResourceView* SRVRHI)
 {
 	VALIDATE_BOUND_SHADER(HullShaderRHI);
 
@@ -555,7 +554,7 @@ void FD3D11DynamicRHI::RHISetShaderResourceViewParameter(FHullShaderRHIParamRef
 	SetShaderResourceView<SF_Hull>(Resource, D3D11SRV, TextureIndex, NAME_None);
 }
 
-void FD3D11DynamicRHI::RHISetShaderResourceViewParameter(FDomainShaderRHIParamRef DomainShaderRHI,uint32 TextureIndex, FRHIShaderResourceView* SRVRHI)
+void FD3D11DynamicRHI::RHISetShaderResourceViewParameter(FRHIDomainShader* DomainShaderRHI,uint32 TextureIndex, FRHIShaderResourceView* SRVRHI)
 {
 	VALIDATE_BOUND_SHADER(DomainShaderRHI);
 
@@ -573,7 +572,7 @@ void FD3D11DynamicRHI::RHISetShaderResourceViewParameter(FDomainShaderRHIParamRe
 	SetShaderResourceView<SF_Domain>(Resource, D3D11SRV, TextureIndex, NAME_None);
 }
 
-void FD3D11DynamicRHI::RHISetShaderResourceViewParameter(FGeometryShaderRHIParamRef GeometryShaderRHI,uint32 TextureIndex, FRHIShaderResourceView* SRVRHI)
+void FD3D11DynamicRHI::RHISetShaderResourceViewParameter(FRHIGeometryShader* GeometryShaderRHI,uint32 TextureIndex, FRHIShaderResourceView* SRVRHI)
 {
 	VALIDATE_BOUND_SHADER(GeometryShaderRHI);
 
@@ -591,7 +590,7 @@ void FD3D11DynamicRHI::RHISetShaderResourceViewParameter(FGeometryShaderRHIParam
 	SetShaderResourceView<SF_Geometry>(Resource, D3D11SRV, TextureIndex, NAME_None);
 }
 
-void FD3D11DynamicRHI::RHISetShaderSampler(FVertexShaderRHIParamRef VertexShaderRHI,uint32 SamplerIndex, FRHISamplerState* NewStateRHI)
+void FD3D11DynamicRHI::RHISetShaderSampler(FRHIVertexShader* VertexShaderRHI,uint32 SamplerIndex, FRHISamplerState* NewStateRHI)
 {
 	VALIDATE_BOUND_SHADER(VertexShaderRHI);
 
@@ -602,7 +601,7 @@ void FD3D11DynamicRHI::RHISetShaderSampler(FVertexShaderRHIParamRef VertexShader
 	StateCache.SetSamplerState<SF_Vertex>(StateResource, SamplerIndex);
 }
 
-void FD3D11DynamicRHI::RHISetShaderSampler(FHullShaderRHIParamRef HullShaderRHI,uint32 SamplerIndex, FRHISamplerState* NewStateRHI)
+void FD3D11DynamicRHI::RHISetShaderSampler(FRHIHullShader* HullShaderRHI,uint32 SamplerIndex, FRHISamplerState* NewStateRHI)
 {
 	VALIDATE_BOUND_SHADER(HullShaderRHI);
 
@@ -613,7 +612,7 @@ void FD3D11DynamicRHI::RHISetShaderSampler(FHullShaderRHIParamRef HullShaderRHI,
 	StateCache.SetSamplerState<SF_Hull>(StateResource, SamplerIndex);
 }
 
-void FD3D11DynamicRHI::RHISetShaderSampler(FDomainShaderRHIParamRef DomainShaderRHI,uint32 SamplerIndex, FRHISamplerState* NewStateRHI)
+void FD3D11DynamicRHI::RHISetShaderSampler(FRHIDomainShader* DomainShaderRHI,uint32 SamplerIndex, FRHISamplerState* NewStateRHI)
 {
 	VALIDATE_BOUND_SHADER(DomainShaderRHI);
 
@@ -624,7 +623,7 @@ void FD3D11DynamicRHI::RHISetShaderSampler(FDomainShaderRHIParamRef DomainShader
 	StateCache.SetSamplerState<SF_Domain>(StateResource, SamplerIndex);
 }
 
-void FD3D11DynamicRHI::RHISetShaderSampler(FGeometryShaderRHIParamRef GeometryShaderRHI,uint32 SamplerIndex, FRHISamplerState* NewStateRHI)
+void FD3D11DynamicRHI::RHISetShaderSampler(FRHIGeometryShader* GeometryShaderRHI,uint32 SamplerIndex, FRHISamplerState* NewStateRHI)
 {
 	VALIDATE_BOUND_SHADER(GeometryShaderRHI);
 
@@ -635,7 +634,7 @@ void FD3D11DynamicRHI::RHISetShaderSampler(FGeometryShaderRHIParamRef GeometrySh
 	StateCache.SetSamplerState<SF_Geometry>(StateResource, SamplerIndex);
 }
 
-void FD3D11DynamicRHI::RHISetShaderSampler(FPixelShaderRHIParamRef PixelShaderRHI,uint32 SamplerIndex, FRHISamplerState* NewStateRHI)
+void FD3D11DynamicRHI::RHISetShaderSampler(FRHIPixelShader* PixelShaderRHI,uint32 SamplerIndex, FRHISamplerState* NewStateRHI)
 {
 	VALIDATE_BOUND_SHADER(PixelShaderRHI);
 
@@ -646,7 +645,7 @@ void FD3D11DynamicRHI::RHISetShaderSampler(FPixelShaderRHIParamRef PixelShaderRH
 	StateCache.SetSamplerState<SF_Pixel>(StateResource, SamplerIndex);
 }
 
-void FD3D11DynamicRHI::RHISetShaderSampler(FComputeShaderRHIParamRef ComputeShaderRHI,uint32 SamplerIndex, FRHISamplerState* NewStateRHI)
+void FD3D11DynamicRHI::RHISetShaderSampler(FRHIComputeShader* ComputeShaderRHI,uint32 SamplerIndex, FRHISamplerState* NewStateRHI)
 {
 	//VALIDATE_BOUND_SHADER(ComputeShaderRHI);
 	FD3D11ComputeShader* ComputeShader = ResourceCast(ComputeShaderRHI);
@@ -656,7 +655,7 @@ void FD3D11DynamicRHI::RHISetShaderSampler(FComputeShaderRHIParamRef ComputeShad
 	StateCache.SetSamplerState<SF_Compute>(StateResource, SamplerIndex);
 }
 
-void FD3D11DynamicRHI::RHISetShaderUniformBuffer(FVertexShaderRHIParamRef VertexShader,uint32 BufferIndex,FUniformBufferRHIParamRef BufferRHI)
+void FD3D11DynamicRHI::RHISetShaderUniformBuffer(FRHIVertexShader* VertexShader,uint32 BufferIndex,FUniformBufferRHIParamRef BufferRHI)
 {
 	check(BufferRHI->GetLayout().GetHash());
 	VALIDATE_BOUND_SHADER(VertexShader);
@@ -670,7 +669,7 @@ void FD3D11DynamicRHI::RHISetShaderUniformBuffer(FVertexShaderRHIParamRef Vertex
 	DirtyUniformBuffers[SF_Vertex] |= (1 << BufferIndex);
 }
 
-void FD3D11DynamicRHI::RHISetShaderUniformBuffer(FHullShaderRHIParamRef HullShader,uint32 BufferIndex,FUniformBufferRHIParamRef BufferRHI)
+void FD3D11DynamicRHI::RHISetShaderUniformBuffer(FRHIHullShader* HullShader,uint32 BufferIndex,FUniformBufferRHIParamRef BufferRHI)
 {
 	check(BufferRHI->GetLayout().GetHash());
 	VALIDATE_BOUND_SHADER(HullShader);
@@ -684,7 +683,7 @@ void FD3D11DynamicRHI::RHISetShaderUniformBuffer(FHullShaderRHIParamRef HullShad
 	DirtyUniformBuffers[SF_Hull] |= (1 << BufferIndex);
 }
 
-void FD3D11DynamicRHI::RHISetShaderUniformBuffer(FDomainShaderRHIParamRef DomainShader,uint32 BufferIndex,FUniformBufferRHIParamRef BufferRHI)
+void FD3D11DynamicRHI::RHISetShaderUniformBuffer(FRHIDomainShader* DomainShader,uint32 BufferIndex,FUniformBufferRHIParamRef BufferRHI)
 {
 	check(BufferRHI->GetLayout().GetHash());
 	VALIDATE_BOUND_SHADER(DomainShader);
@@ -698,7 +697,7 @@ void FD3D11DynamicRHI::RHISetShaderUniformBuffer(FDomainShaderRHIParamRef Domain
 	DirtyUniformBuffers[SF_Domain] |= (1 << BufferIndex);
 }
 
-void FD3D11DynamicRHI::RHISetShaderUniformBuffer(FGeometryShaderRHIParamRef GeometryShader,uint32 BufferIndex,FUniformBufferRHIParamRef BufferRHI)
+void FD3D11DynamicRHI::RHISetShaderUniformBuffer(FRHIGeometryShader* GeometryShader,uint32 BufferIndex,FUniformBufferRHIParamRef BufferRHI)
 {
 	check(BufferRHI->GetLayout().GetHash());
 	VALIDATE_BOUND_SHADER(GeometryShader);
@@ -712,7 +711,7 @@ void FD3D11DynamicRHI::RHISetShaderUniformBuffer(FGeometryShaderRHIParamRef Geom
 	DirtyUniformBuffers[SF_Geometry] |= (1 << BufferIndex);
 }
 
-void FD3D11DynamicRHI::RHISetShaderUniformBuffer(FPixelShaderRHIParamRef PixelShader,uint32 BufferIndex,FUniformBufferRHIParamRef BufferRHI)
+void FD3D11DynamicRHI::RHISetShaderUniformBuffer(FRHIPixelShader* PixelShader,uint32 BufferIndex,FUniformBufferRHIParamRef BufferRHI)
 {
 	check(BufferRHI->GetLayout().GetHash());
 	VALIDATE_BOUND_SHADER(PixelShader);
@@ -727,7 +726,7 @@ void FD3D11DynamicRHI::RHISetShaderUniformBuffer(FPixelShaderRHIParamRef PixelSh
 	DirtyUniformBuffers[SF_Pixel] |= (1 << BufferIndex);
 }
 
-void FD3D11DynamicRHI::RHISetShaderUniformBuffer(FComputeShaderRHIParamRef ComputeShader,uint32 BufferIndex,FUniformBufferRHIParamRef BufferRHI)
+void FD3D11DynamicRHI::RHISetShaderUniformBuffer(FRHIComputeShader* ComputeShader,uint32 BufferIndex,FUniformBufferRHIParamRef BufferRHI)
 {
 	check(BufferRHI->GetLayout().GetHash());
 	//VALIDATE_BOUND_SHADER(ComputeShader);
@@ -741,42 +740,42 @@ void FD3D11DynamicRHI::RHISetShaderUniformBuffer(FComputeShaderRHIParamRef Compu
 	DirtyUniformBuffers[SF_Compute] |= (1 << BufferIndex);
 }
 
-void FD3D11DynamicRHI::RHISetShaderParameter(FHullShaderRHIParamRef HullShaderRHI,uint32 BufferIndex,uint32 BaseIndex,uint32 NumBytes,const void* NewValue)
+void FD3D11DynamicRHI::RHISetShaderParameter(FRHIHullShader* HullShaderRHI,uint32 BufferIndex,uint32 BaseIndex,uint32 NumBytes,const void* NewValue)
 {
 	VALIDATE_BOUND_SHADER(HullShaderRHI);
 	checkSlow(HSConstantBuffers[BufferIndex]);
 	HSConstantBuffers[BufferIndex]->UpdateConstant((const uint8*)NewValue,BaseIndex,NumBytes);
 }
 
-void FD3D11DynamicRHI::RHISetShaderParameter(FDomainShaderRHIParamRef DomainShaderRHI,uint32 BufferIndex,uint32 BaseIndex,uint32 NumBytes,const void* NewValue)
+void FD3D11DynamicRHI::RHISetShaderParameter(FRHIDomainShader* DomainShaderRHI,uint32 BufferIndex,uint32 BaseIndex,uint32 NumBytes,const void* NewValue)
 {
 	VALIDATE_BOUND_SHADER(DomainShaderRHI);
 	checkSlow(DSConstantBuffers[BufferIndex]);
 	DSConstantBuffers[BufferIndex]->UpdateConstant((const uint8*)NewValue,BaseIndex,NumBytes);
 }
 
-void FD3D11DynamicRHI::RHISetShaderParameter(FVertexShaderRHIParamRef VertexShaderRHI,uint32 BufferIndex,uint32 BaseIndex,uint32 NumBytes,const void* NewValue)
+void FD3D11DynamicRHI::RHISetShaderParameter(FRHIVertexShader* VertexShaderRHI,uint32 BufferIndex,uint32 BaseIndex,uint32 NumBytes,const void* NewValue)
 {
 	VALIDATE_BOUND_SHADER(VertexShaderRHI);
 	checkSlow(VSConstantBuffers[BufferIndex]);
 	VSConstantBuffers[BufferIndex]->UpdateConstant((const uint8*)NewValue,BaseIndex,NumBytes);
 }
 
-void FD3D11DynamicRHI::RHISetShaderParameter(FPixelShaderRHIParamRef PixelShaderRHI,uint32 BufferIndex,uint32 BaseIndex,uint32 NumBytes,const void* NewValue)
+void FD3D11DynamicRHI::RHISetShaderParameter(FRHIPixelShader* PixelShaderRHI,uint32 BufferIndex,uint32 BaseIndex,uint32 NumBytes,const void* NewValue)
 {
 	VALIDATE_BOUND_SHADER(PixelShaderRHI);
 	checkSlow(PSConstantBuffers[BufferIndex]);
 	PSConstantBuffers[BufferIndex]->UpdateConstant((const uint8*)NewValue,BaseIndex,NumBytes);
 }
 
-void FD3D11DynamicRHI::RHISetShaderParameter(FGeometryShaderRHIParamRef GeometryShaderRHI,uint32 BufferIndex,uint32 BaseIndex,uint32 NumBytes,const void* NewValue)
+void FD3D11DynamicRHI::RHISetShaderParameter(FRHIGeometryShader* GeometryShaderRHI,uint32 BufferIndex,uint32 BaseIndex,uint32 NumBytes,const void* NewValue)
 {
 	VALIDATE_BOUND_SHADER(GeometryShaderRHI);
 	checkSlow(GSConstantBuffers[BufferIndex]);
 	GSConstantBuffers[BufferIndex]->UpdateConstant((const uint8*)NewValue,BaseIndex,NumBytes);
 }
 
-void FD3D11DynamicRHI::RHISetShaderParameter(FComputeShaderRHIParamRef ComputeShaderRHI,uint32 BufferIndex,uint32 BaseIndex,uint32 NumBytes,const void* NewValue)
+void FD3D11DynamicRHI::RHISetShaderParameter(FRHIComputeShader* ComputeShaderRHI,uint32 BufferIndex,uint32 BaseIndex,uint32 NumBytes,const void* NewValue)
 {
 	//VALIDATE_BOUND_SHADER(ComputeShaderRHI);
 	checkSlow(CSConstantBuffers[BufferIndex]);
diff --git a/Engine/Source/Runtime/Windows/D3D11RHI/Private/D3D11RHIPrivate.h b/Engine/Source/Runtime/Windows/D3D11RHI/Private/D3D11RHIPrivate.h
--- a/Engine/Source/Runtime/Windows/D3D11RHI/Private/D3D11RHIPrivate.h
+++ b/Engine/Source/Runtime/Windows/D3D11RHI/Private/D3D11RHIPrivate.h
@@ -409,7 +409,7 @@ public:
 	virtual void RHIWriteGPUFence(FRHIGPUFence* Fence) final override;
 	virtual void* RHILockStagingBuffer(FRHIStagingBuffer* StagingBuffer, uint32 Offset, uint32 SizeRHI) final override;
     virtual void RHIUnlockStagingBuffer(FRHIStagingBuffer* StagingBuffer) final override;
-	virtual FBoundShaderStateRHIRef RHICreateBoundShaderState(FRHIVertexDeclaration* VertexDeclaration, FVertexShaderRHIParamRef VertexShader, FHullShaderRHIParamRef HullShader, FDomainShaderRHIParamRef DomainShader, FPixelShaderRHIParamRef PixelShader, FGeometryShaderRHIParamRef GeometryShader) final override;
+	virtual FBoundShaderStateRHIRef RHICreateBoundShaderState(FRHIVertexDeclaration* VertexDeclaration, FRHIVertexShader* VertexShader, FRHIHullShader* HullShader, FRHIDomainShader* DomainShader, FRHIPixelShader* PixelShader, FRHIGeometryShader* GeometryShader) final override;
 	virtual FUniformBufferRHIRef RHICreateUniformBuffer(const void* Contents, const FRHIUniformBufferLayout& Layout, EUniformBufferUsage Usage, EUniformBufferValidation Validation) final override;
 	virtual void RHIUpdateUniformBuffer(FUniformBufferRHIParamRef UniformBufferRHI, const void* Contents) final override;
 	virtual FIndexBufferRHIRef RHICreateIndexBuffer(uint32 Stride, uint32 Size, uint32 InUsage, FRHIResourceCreateInfo& CreateInfo) final override;
@@ -494,7 +494,7 @@ public:
 	virtual class IRHICommandContext* RHIGetDefaultContext() final override;
 	virtual class IRHICommandContextContainer* RHIGetCommandContextContainer(int32 Index, int32 Num) final override;
 
-	virtual void RHISetComputeShader(FComputeShaderRHIParamRef ComputeShader) final override;
+	virtual void RHISetComputeShader(FRHIComputeShader* ComputeShader) final override;
 	virtual void RHIDispatchComputeShader(uint32 ThreadGroupCountX, uint32 ThreadGroupCountY, uint32 ThreadGroupCountZ) final override;
 	virtual void RHIDispatchIndirectComputeShader(FVertexBufferRHIParamRef ArgumentBuffer, uint32 ArgumentOffset) final override;
 	virtual void RHIAutomaticCacheFlushAfterComputeShader(bool bEnable) final override;
@@ -529,38 +529,38 @@ public:
 		PrimitiveType = FallbackGraphicsState->Initializer.PrimitiveType;
 	}
 
-	virtual void RHISetShaderTexture(FVertexShaderRHIParamRef VertexShader, uint32 TextureIndex, FTextureRHIParamRef NewTexture) final override;
-	virtual void RHISetShaderTexture(FHullShaderRHIParamRef HullShader, uint32 TextureIndex, FTextureRHIParamRef NewTexture) final override;
-	virtual void RHISetShaderTexture(FDomainShaderRHIParamRef DomainShader, uint32 TextureIndex, FTextureRHIParamRef NewTexture) final override;
-	virtual void RHISetShaderTexture(FGeometryShaderRHIParamRef GeometryShader, uint32 TextureIndex, FTextureRHIParamRef NewTexture) final override;
-	virtual void RHISetShaderTexture(FPixelShaderRHIParamRef PixelShader, uint32 TextureIndex, FTextureRHIParamRef NewTexture) final override;
-	virtual void RHISetShaderTexture(FComputeShaderRHIParamRef PixelShader, uint32 TextureIndex, FTextureRHIParamRef NewTexture) final override;
-	virtual void RHISetShaderSampler(FComputeShaderRHIParamRef ComputeShader, uint32 SamplerIndex, FRHISamplerState* NewState) final override;
-	virtual void RHISetShaderSampler(FVertexShaderRHIParamRef VertexShader, uint32 SamplerIndex, FRHISamplerState* NewState) final override;
-	virtual void RHISetShaderSampler(FGeometryShaderRHIParamRef GeometryShader, uint32 SamplerIndex, FRHISamplerState* NewState) final override;
-	virtual void RHISetShaderSampler(FDomainShaderRHIParamRef DomainShader, uint32 SamplerIndex, FRHISamplerState* NewState) final override;
-	virtual void RHISetShaderSampler(FHullShaderRHIParamRef HullShader, uint32 SamplerIndex, FRHISamplerState* NewState) final override;
-	virtual void RHISetShaderSampler(FPixelShaderRHIParamRef PixelShader, uint32 SamplerIndex, FRHISamplerState* NewState) final override;
-	virtual void RHISetUAVParameter(FComputeShaderRHIParamRef ComputeShader, uint32 UAVIndex, FRHIUnorderedAccessView* UAV) final override;
-	virtual void RHISetUAVParameter(FComputeShaderRHIParamRef ComputeShader, uint32 UAVIndex, FRHIUnorderedAccessView* UAV, uint32 InitialCount) final override;
-	virtual void RHISetShaderResourceViewParameter(FPixelShaderRHIParamRef PixelShader, uint32 SamplerIndex, FRHIShaderResourceView* SRV) final override;
-	virtual void RHISetShaderResourceViewParameter(FVertexShaderRHIParamRef VertexShader, uint32 SamplerIndex, FRHIShaderResourceView* SRV) final override;
-	virtual void RHISetShaderResourceViewParameter(FComputeShaderRHIParamRef ComputeShader, uint32 SamplerIndex, FRHIShaderResourceView* SRV) final override;
-	virtual void RHISetShaderResourceViewParameter(FHullShaderRHIParamRef HullShader, uint32 SamplerIndex, FRHIShaderResourceView* SRV) final override;
-	virtual void RHISetShaderResourceViewParameter(FDomainShaderRHIParamRef DomainShader, uint32 SamplerIndex, FRHIShaderResourceView* SRV) final override;
-	virtual void RHISetShaderResourceViewParameter(FGeometryShaderRHIParamRef GeometryShader, uint32 SamplerIndex, FRHIShaderResourceView* SRV) final override;
-	virtual void RHISetShaderUniformBuffer(FVertexShaderRHIParamRef VertexShader, uint32 BufferIndex, FUniformBufferRHIParamRef Buffer) final override;
-	virtual void RHISetShaderUniformBuffer(FHullShaderRHIParamRef HullShader, uint32 BufferIndex, FUniformBufferRHIParamRef Buffer) final override;
-	virtual void RHISetShaderUniformBuffer(FDomainShaderRHIParamRef DomainShader, uint32 BufferIndex, FUniformBufferRHIParamRef Buffer) final override;
-	virtual void RHISetShaderUniformBuffer(FGeometryShaderRHIParamRef GeometryShader, uint32 BufferIndex, FUniformBufferRHIParamRef Buffer) final override;
-	virtual void RHISetShaderUniformBuffer(FPixelShaderRHIParamRef PixelShader, uint32 BufferIndex, FUniformBufferRHIParamRef Buffer) final override;
-	virtual void RHISetShaderUniformBuffer(FComputeShaderRHIParamRef ComputeShader, uint32 BufferIndex, FUniformBufferRHIParamRef Buffer) final override;
-	virtual void RHISetShaderParameter(FVertexShaderRHIParamRef VertexShader, uint32 BufferIndex, uint32 BaseIndex, uint32 NumBytes, const void* NewValue) final override;
-	virtual void RHISetShaderParameter(FPixelShaderRHIParamRef PixelShader, uint32 BufferIndex, uint32 BaseIndex, uint32 NumBytes, const void* NewValue) final override;
-	virtual void RHISetShaderParameter(FHullShaderRHIParamRef HullShader, uint32 BufferIndex, uint32 BaseIndex, uint32 NumBytes, const void* NewValue) final override;
-	virtual void RHISetShaderParameter(FDomainShaderRHIParamRef DomainShader, uint32 BufferIndex, uint32 BaseIndex, uint32 NumBytes, const void* NewValue) final override;
-	virtual void RHISetShaderParameter(FGeometryShaderRHIParamRef GeometryShader, uint32 BufferIndex, uint32 BaseIndex, uint32 NumBytes, const void* NewValue) final override;
-	virtual void RHISetShaderParameter(FComputeShaderRHIParamRef ComputeShader, uint32 BufferIndex, uint32 BaseIndex, uint32 NumBytes, const void* NewValue) final override;
+	virtual void RHISetShaderTexture(FRHIVertexShader* VertexShader, uint32 TextureIndex, FTextureRHIParamRef NewTexture) final override;
+	virtual void RHISetShaderTexture(FRHIHullShader* HullShader, uint32 TextureIndex, FTextureRHIParamRef NewTexture) final override;
+	virtual void RHISetShaderTexture(FRHIDomainShader* DomainShader, uint32 TextureIndex, FTextureRHIParamRef NewTexture) final override;
+	virtual void RHISetShaderTexture(FRHIGeometryShader* GeometryShader, uint32 TextureIndex, FTextureRHIParamRef NewTexture) final override;
+	virtual void RHISetShaderTexture(FRHIPixelShader* PixelShader, uint32 TextureIndex, FTextureRHIParamRef NewTexture) final override;
+	virtual void RHISetShaderTexture(FRHIComputeShader* PixelShader, uint32 TextureIndex, FTextureRHIParamRef NewTexture) final override;
+	virtual void RHISetShaderSampler(FRHIComputeShader* ComputeShader, uint32 SamplerIndex, FRHISamplerState* NewState) final override;
+	virtual void RHISetShaderSampler(FRHIVertexShader* VertexShader, uint32 SamplerIndex, FRHISamplerState* NewState) final override;
+	virtual void RHISetShaderSampler(FRHIGeometryShader* GeometryShader, uint32 SamplerIndex, FRHISamplerState* NewState) final override;
+	virtual void RHISetShaderSampler(FRHIDomainShader* DomainShader, uint32 SamplerIndex, FRHISamplerState* NewState) final override;
+	virtual void RHISetShaderSampler(FRHIHullShader* HullShader, uint32 SamplerIndex, FRHISamplerState* NewState) final override;
+	virtual void RHISetShaderSampler(FRHIPixelShader* PixelShader, uint32 SamplerIndex, FRHISamplerState* NewState) final override;
+	virtual void RHISetUAVParameter(FRHIComputeShader* ComputeShader, uint32 UAVIndex, FRHIUnorderedAccessView* UAV) final override;
+	virtual void RHISetUAVParameter(FRHIComputeShader* ComputeShader, uint32 UAVIndex, FRHIUnorderedAccessView* UAV, uint32 InitialCount) final override;
+	virtual void RHISetShaderResourceViewParameter(FRHIPixelShader* PixelShader, uint32 SamplerIndex, FRHIShaderResourceView* SRV) final override;
+	virtual void RHISetShaderResourceViewParameter(FRHIVertexShader* VertexShader, uint32 SamplerIndex, FRHIShaderResourceView* SRV) final override;
+	virtual void RHISetShaderResourceViewParameter(FRHIComputeShader* ComputeShader, uint32 SamplerIndex, FRHIShaderResourceView* SRV) final override;
+	virtual void RHISetShaderResourceViewParameter(FRHIHullShader* HullShader, uint32 SamplerIndex, FRHIShaderResourceView* SRV) final override;
+	virtual void RHISetShaderResourceViewParameter(FRHIDomainShader* DomainShader, uint32 SamplerIndex, FRHIShaderResourceView* SRV) final override;
+	virtual void RHISetShaderResourceViewParameter(FRHIGeometryShader* GeometryShader, uint32 SamplerIndex, FRHIShaderResourceView* SRV) final override;
+	virtual void RHISetShaderUniformBuffer(FRHIVertexShader* VertexShader, uint32 BufferIndex, FUniformBufferRHIParamRef Buffer) final override;
+	virtual void RHISetShaderUniformBuffer(FRHIHullShader* HullShader, uint32 BufferIndex, FUniformBufferRHIParamRef Buffer) final override;
+	virtual void RHISetShaderUniformBuffer(FRHIDomainShader* DomainShader, uint32 BufferIndex, FUniformBufferRHIParamRef Buffer) final override;
+	virtual void RHISetShaderUniformBuffer(FRHIGeometryShader* GeometryShader, uint32 BufferIndex, FUniformBufferRHIParamRef Buffer) final override;
+	virtual void RHISetShaderUniformBuffer(FRHIPixelShader* PixelShader, uint32 BufferIndex, FUniformBufferRHIParamRef Buffer) final override;
+	virtual void RHISetShaderUniformBuffer(FRHIComputeShader* ComputeShader, uint32 BufferIndex, FUniformBufferRHIParamRef Buffer) final override;
+	virtual void RHISetShaderParameter(FRHIVertexShader* VertexShader, uint32 BufferIndex, uint32 BaseIndex, uint32 NumBytes, const void* NewValue) final override;
+	virtual void RHISetShaderParameter(FRHIPixelShader* PixelShader, uint32 BufferIndex, uint32 BaseIndex, uint32 NumBytes, const void* NewValue) final override;
+	virtual void RHISetShaderParameter(FRHIHullShader* HullShader, uint32 BufferIndex, uint32 BaseIndex, uint32 NumBytes, const void* NewValue) final override;
+	virtual void RHISetShaderParameter(FRHIDomainShader* DomainShader, uint32 BufferIndex, uint32 BaseIndex, uint32 NumBytes, const void* NewValue) final override;
+	virtual void RHISetShaderParameter(FRHIGeometryShader* GeometryShader, uint32 BufferIndex, uint32 BaseIndex, uint32 NumBytes, const void* NewValue) final override;
+	virtual void RHISetShaderParameter(FRHIComputeShader* ComputeShader, uint32 BufferIndex, uint32 BaseIndex, uint32 NumBytes, const void* NewValue) final override;
 	virtual void RHISetDepthStencilState(FRHIDepthStencilState* NewState, uint32 StencilRef) final override;
 	virtual void RHISetStencilRef(uint32 StencilRef) final override;
 	virtual void RHISetBlendState(FRHIBlendState* NewState, const FLinearColor& BlendFactor) final override;
@@ -730,7 +730,7 @@ private:
 	void TrackResourceBoundAsVB(FD3D11BaseShaderResource* Resource, int32 StreamIndex);
 	void TrackResourceBoundAsIB(FD3D11BaseShaderResource* Resource);
 
-	void SetCurrentComputeShader(FComputeShaderRHIParamRef ComputeShader)
+	void SetCurrentComputeShader(FRHIComputeShader* ComputeShader)
 	{
 		CurrentComputeShader = ComputeShader;
 	}
diff --git a/Engine/Source/Runtime/Windows/D3D11RHI/Private/D3D11Shaders.cpp b/Engine/Source/Runtime/Windows/D3D11RHI/Private/D3D11Shaders.cpp
--- a/Engine/Source/Runtime/Windows/D3D11RHI/Private/D3D11Shaders.cpp
+++ b/Engine/Source/Runtime/Windows/D3D11RHI/Private/D3D11Shaders.cpp
@@ -266,11 +266,11 @@ void FD3D11DynamicRHI::RHISetMultipleViewports(uint32 Count, const FViewportBoun
 
 FD3D11BoundShaderState::FD3D11BoundShaderState(
 	FRHIVertexDeclaration* InVertexDeclarationRHI,
-	FVertexShaderRHIParamRef InVertexShaderRHI,
-	FPixelShaderRHIParamRef InPixelShaderRHI,
-	FHullShaderRHIParamRef InHullShaderRHI,
-	FDomainShaderRHIParamRef InDomainShaderRHI,
-	FGeometryShaderRHIParamRef InGeometryShaderRHI,
+	FRHIVertexShader* InVertexShaderRHI,
+	FRHIPixelShader* InPixelShaderRHI,
+	FRHIHullShader* InHullShaderRHI,
+	FRHIDomainShader* InDomainShaderRHI,
+	FRHIGeometryShader* InGeometryShaderRHI,
 	ID3D11Device* Direct3DDevice
 	):
 	CacheLink(InVertexDeclarationRHI,InVertexShaderRHI,InPixelShaderRHI,InHullShaderRHI,InDomainShaderRHI,InGeometryShaderRHI,this)
@@ -345,11 +345,11 @@ FD3D11BoundShaderState::~FD3D11BoundShaderState()
 */
 FBoundShaderStateRHIRef FD3D11DynamicRHI::RHICreateBoundShaderState(
 	FRHIVertexDeclaration* VertexDeclarationRHI,
-	FVertexShaderRHIParamRef VertexShaderRHI, 
-	FHullShaderRHIParamRef HullShaderRHI, 
-	FDomainShaderRHIParamRef DomainShaderRHI, 
-	FPixelShaderRHIParamRef PixelShaderRHI,
-	FGeometryShaderRHIParamRef GeometryShaderRHI
+	FRHIVertexShader* VertexShaderRHI,
+	FRHIHullShader* HullShaderRHI,
+	FRHIDomainShader* DomainShaderRHI,
+	FRHIPixelShader* PixelShaderRHI,
+	FRHIGeometryShader* GeometryShaderRHI
 	)
 {
 	check(IsInRenderingThread() || IsInRHIThread());
diff --git a/Engine/Source/Runtime/Windows/D3D11RHI/Public/D3D11Resources.h b/Engine/Source/Runtime/Windows/D3D11RHI/Public/D3D11Resources.h
--- a/Engine/Source/Runtime/Windows/D3D11RHI/Public/D3D11Resources.h
+++ b/Engine/Source/Runtime/Windows/D3D11RHI/Public/D3D11Resources.h
@@ -129,11 +129,11 @@ public:
 	/** Initialization constructor. */
 	FD3D11BoundShaderState(
 		FRHIVertexDeclaration* InVertexDeclarationRHI,
-		FVertexShaderRHIParamRef InVertexShaderRHI,
-		FPixelShaderRHIParamRef InPixelShaderRHI,
-		FHullShaderRHIParamRef InHullShaderRHI,
-		FDomainShaderRHIParamRef InDomainShaderRHI,
-		FGeometryShaderRHIParamRef InGeometryShaderRHI,
+		FRHIVertexShader* InVertexShaderRHI,
+		FRHIPixelShader* InPixelShaderRHI,
+		FRHIHullShader* InHullShaderRHI,
+		FRHIDomainShader* InDomainShaderRHI,
+		FRHIGeometryShader* InGeometryShaderRHI,
 		ID3D11Device* Direct3DDevice
 		);
 
