DR - Non-editor compile fix
#rnx

diff --git a/Engine/Source/Runtime/VulkanRHI/Private/VulkanLayers.cpp b/Engine/Source/Runtime/VulkanRHI/Private/VulkanLayers.cpp
--- a/Engine/Source/Runtime/VulkanRHI/Private/VulkanLayers.cpp
+++ b/Engine/Source/Runtime/VulkanRHI/Private/VulkanLayers.cpp
@@ -316,6 +316,7 @@ void FVulkanDynamicRHI::GetInstanceLayersAndExtensions(TArray<const ANSICHAR*>&
 #endif	// VULKAN_ENABLE_API_DUMP
 
 	// At this point the CVar holds the final value
+#if VULKAN_HAS_DEBUGGING_ENABLED
 	const int32 VulkanValidationOption = GValidationCvar.GetValueOnAnyThread();
 	if (!bVkTrace && VulkanValidationOption > 0)
 	{
@@ -366,6 +367,7 @@ void FVulkanDynamicRHI::GetInstanceLayersAndExtensions(TArray<const ANSICHAR*>&
 			}
 		}
 	}
+#endif
 
 #if VULKAN_SUPPORTS_DEBUG_UTILS
 	if (!bVkTrace && VulkanValidationOption > 0)
