Linux: fix crashes in the editor (UE-74652, UE-74271).

diff --git a/Engine/Source/Runtime/VulkanRHI/Private/Linux/VulkanLinuxPlatform.h b/Engine/Source/Runtime/VulkanRHI/Private/Linux/VulkanLinuxPlatform.h
--- a/Engine/Source/Runtime/VulkanRHI/Private/Linux/VulkanLinuxPlatform.h
+++ b/Engine/Source/Runtime/VulkanRHI/Private/Linux/VulkanLinuxPlatform.h
@@ -56,7 +56,7 @@ public:
 			{
 				static auto* CVar = IConsoleManager::Get().FindTConsoleVariableDataInt(TEXT("r.Vulkan.UseRealUBs"));
 				return (CVar && CVar->GetValueOnAnyThread() == 0);
-			});
+			}());
 			return bAlwaysUseRealUBs ? false : bCodeHeaderUseRealUBs;
 		}
 		else
