DR - Change message to avoid false positive
#rnx

diff --git a/Engine/Source/Runtime/VulkanRHI/Private/VulkanPendingState.cpp b/Engine/Source/Runtime/VulkanRHI/Private/VulkanPendingState.cpp
--- a/Engine/Source/Runtime/VulkanRHI/Private/VulkanPendingState.cpp
+++ b/Engine/Source/Runtime/VulkanRHI/Private/VulkanPendingState.cpp
@@ -788,7 +788,8 @@ void FVulkanDescriptorSetCache::AddCachedPool()
 			CachedPools.EmplaceAt(0, MoveTemp(FreePool));
 			return;
 		}
-		UE_LOG(LogVulkanRHI, Display, TEXT("FVulkanDescriptorSetCache::AddCachedPool() MaxDescriptorSets Error: %f. Tolerance: [%f..%f]."),
+		// Don't write 'error' as it confuses reporting; it's a perf warning more than an actual error
+		UE_LOG(LogVulkanRHI, Display, TEXT("FVulkanDescriptorSetCache::AddCachedPool() MaxDescriptorSets Delta/Err: %f. Tolerance: [%f..%f]."),
 			static_cast<double>(Error), static_cast<double>(MinErrorTolerance), static_cast<double>(MaxErrorTolerance));
 		FreePool.Reset();
 	}
