resource type name fix

diff --git a/Engine/Source/Runtime/VulkanRHI/Private/VulkanRenderTarget.cpp b/Engine/Source/Runtime/VulkanRHI/Private/VulkanRenderTarget.cpp
--- a/Engine/Source/Runtime/VulkanRHI/Private/VulkanRenderTarget.cpp
+++ b/Engine/Source/Runtime/VulkanRHI/Private/VulkanRenderTarget.cpp
@@ -872,7 +872,7 @@ void FVulkanDynamicRHI::RHIReadSurfaceData(FRHITexture* TextureRHI, FIntRect Rec
 		uint32 PixelByteSize = 4u;
 		uint8* In = (uint8*)StagingBuffer->GetMappedPointer() + (Rect.Min.Y * TextureRHI2D->GetSizeX() + Rect.Min.X) * PixelByteSize;
 		uint32 SrcPitch = TextureRHI2D->GetSizeX() * PixelByteSize;
-		ConvertRawR8G8B8A8DataToFColor(TextureRHI2D->GetSizeX(), TextureRHI2D->GetSizeY(), In, SrcPitch, Dest);
+		ConvertRawB8G8R8A8DataToFColor(TextureRHI2D->GetSizeX(), TextureRHI2D->GetSizeY(), In, SrcPitch, Dest);
 	}
 
 	Device->GetStagingManager().ReleaseBuffer(CmdBuffer, StagingBuffer);
