DR - Fix ensure

diff --git a/Engine/Source/Runtime/VulkanRHI/Private/VulkanRenderTarget.cpp b/Engine/Source/Runtime/VulkanRHI/Private/VulkanRenderTarget.cpp
--- a/Engine/Source/Runtime/VulkanRHI/Private/VulkanRenderTarget.cpp
+++ b/Engine/Source/Runtime/VulkanRHI/Private/VulkanRenderTarget.cpp
@@ -1232,6 +1232,7 @@ bool FVulkanCommandListContext::FPendingTransition::GatherBarriers(FTransitionAn
 			DestAccess = VK_ACCESS_SHADER_READ_BIT | VK_ACCESS_SHADER_WRITE_BIT;
 			break;
 		case EResourceTransitionAccess::ERWNoBarrier:
+		case EResourceTransitionAccess::ERWSubResBarrier:
 			//#todo-rco: Skip for now
 			continue;
 		default:
