DR - Remove dummy testings
Change log to verbose

diff --git a/Engine/Source/Runtime/RenderCore/Private/DynamicBufferAllocator.cpp b/Engine/Source/Runtime/RenderCore/Private/DynamicBufferAllocator.cpp
--- a/Engine/Source/Runtime/RenderCore/Private/DynamicBufferAllocator.cpp
+++ b/Engine/Source/Runtime/RenderCore/Private/DynamicBufferAllocator.cpp
@@ -74,7 +74,7 @@ void FGlobalDynamicReadBuffer::Cleanup()
 {
 	if (FloatBufferPool)
 	{
-		UE_LOG(LogRendererCore, Log, TEXT("FGlobalDynamicReadBuffer::Cleanup()"));
+		UE_LOG(LogRendererCore, Verbose, TEXT("FGlobalDynamicReadBuffer::Cleanup()"));
 		delete FloatBufferPool;
 		FloatBufferPool = nullptr;
 	}
diff --git a/Engine/Source/Runtime/VulkanRHI/Private/Android/VulkanAndroidPlatform.h b/Engine/Source/Runtime/VulkanRHI/Private/Android/VulkanAndroidPlatform.h
--- a/Engine/Source/Runtime/VulkanRHI/Private/Android/VulkanAndroidPlatform.h
+++ b/Engine/Source/Runtime/VulkanRHI/Private/Android/VulkanAndroidPlatform.h
@@ -69,7 +69,6 @@ public:
 	static bool SupportsStandardSwapchain();
 	static EPixelFormat GetPixelFormatForNonDefaultSwapchain();
 
-	static bool SupportsDepthFetchDuringDepthTest() { return true; }
 	static bool SupportsTimestampRenderQueries() { return false; }
 
 	static bool RequiresMobileRenderer() { return true; }
diff --git a/Engine/Source/Runtime/VulkanRHI/Private/VulkanGenericPlatform.h b/Engine/Source/Runtime/VulkanRHI/Private/VulkanGenericPlatform.h
--- a/Engine/Source/Runtime/VulkanRHI/Private/VulkanGenericPlatform.h
+++ b/Engine/Source/Runtime/VulkanRHI/Private/VulkanGenericPlatform.h
@@ -47,7 +47,6 @@ public:
 		return PF_Unknown;
 	}
 
-	static bool SupportsDepthFetchDuringDepthTest() { return true; }
 	static bool SupportsTimestampRenderQueries() { return true; }
 
 	static bool RequiresMobileRenderer() { return false; }
diff --git a/Engine/Source/Runtime/VulkanRHI/Private/VulkanRHI.cpp b/Engine/Source/Runtime/VulkanRHI/Private/VulkanRHI.cpp
--- a/Engine/Source/Runtime/VulkanRHI/Private/VulkanRHI.cpp
+++ b/Engine/Source/Runtime/VulkanRHI/Private/VulkanRHI.cpp
@@ -720,7 +720,7 @@ void FVulkanDynamicRHI::InitInstance()
 		GRHISupportsBaseVertexIndex = true;
 		GSupportsSeparateRenderTargetBlendState = true;
 
-		GSupportsDepthFetchDuringDepthTest = FVulkanPlatform::SupportsDepthFetchDuringDepthTest();
+		GSupportsDepthFetchDuringDepthTest = true;
 
 		FVulkanPlatform::SetupFeatureLevels();
 
