Fixed read access violation in FPackedUniformBuffers::SetPackedGlobalParameter with "NumBytes==0" (occurs with "QAGame TM-ShaderModles -vulkan").

#rnx

diff --git a/Engine/Source/Runtime/VulkanRHI/Private/VulkanResources.h b/Engine/Source/Runtime/VulkanRHI/Private/VulkanResources.h
--- a/Engine/Source/Runtime/VulkanRHI/Private/VulkanResources.h
+++ b/Engine/Source/Runtime/VulkanRHI/Private/VulkanResources.h
@@ -1407,11 +1407,10 @@ public:
 		uint32* RESTRICT RawSrcEnd = RawSrc + (NumBytes >> 2);
 
 		bool bChanged = false;
-		do
+		while (RawSrc != RawSrcEnd)
 		{
 			bChanged |= CopyAndReturnNotEqual(*RawDst++, *RawSrc++);
 		}
-		while (RawSrc != RawSrcEnd);
 
 		InOutPackedUniformBufferStagingDirty = InOutPackedUniformBufferStagingDirty | ((uint64)(bChanged ? 1 : 0) << (uint64)BufferIndex);
 	}
