DRS - Fix Android compile
#rnx

diff --git a/Engine/Source/Runtime/VulkanRHI/Private/Android/VulkanAndroidPlatform.h b/Engine/Source/Runtime/VulkanRHI/Private/Android/VulkanAndroidPlatform.h
--- a/Engine/Source/Runtime/VulkanRHI/Private/Android/VulkanAndroidPlatform.h
+++ b/Engine/Source/Runtime/VulkanRHI/Private/Android/VulkanAndroidPlatform.h
@@ -17,6 +17,8 @@
 #define VULKAN_SUPPORTS_GOOGLE_DISPLAY_TIMING		1
 #define VULKAN_FREEPAGE_FOR_TYPE					1
 #define VULKAN_PURGE_SHADER_MODULES					0
+#define VULKAN_SUPPORTS_DEDICATED_ALLOCATION		0
+#define VULKAN_SUPPORTS_PHYSICAL_DEVICE_PROPERTIES2	0
 
 
 // Android's hashes currently work fine as the problematic cases are:
diff --git a/Engine/Source/Runtime/VulkanRHI/Private/VulkanCommandWrappers.h b/Engine/Source/Runtime/VulkanRHI/Private/VulkanCommandWrappers.h
--- a/Engine/Source/Runtime/VulkanRHI/Private/VulkanCommandWrappers.h
+++ b/Engine/Source/Runtime/VulkanRHI/Private/VulkanCommandWrappers.h
@@ -251,10 +251,7 @@ namespace VulkanRHI
 	static FORCEINLINE_DEBUGGABLE void  vkGetPhysicalDeviceProperties2KHR(VkPhysicalDevice PhysicalDevice, VkPhysicalDeviceProperties2KHR* Properties)
 	{
 		FWrapLayer::GetPhysicalDeviceProperties2KHR(VK_RESULT_MAX_ENUM, PhysicalDevice, Properties);
-		if (VULKANAPINAMESPACE::vkGetPhysicalDeviceProperties2KHR != nullptr)
-		{
-			VULKANAPINAMESPACE::vkGetPhysicalDeviceProperties2KHR(PhysicalDevice, Properties);
-		}
+		VULKANAPINAMESPACE::vkGetPhysicalDeviceProperties2KHR(PhysicalDevice, Properties);
 		FWrapLayer::GetPhysicalDeviceProperties2KHR(VK_SUCCESS, PhysicalDevice, Properties);
 	}
 #endif
diff --git a/Engine/Source/Runtime/VulkanRHI/Public/VulkanConfiguration.h b/Engine/Source/Runtime/VulkanRHI/Public/VulkanConfiguration.h
--- a/Engine/Source/Runtime/VulkanRHI/Public/VulkanConfiguration.h
+++ b/Engine/Source/Runtime/VulkanRHI/Public/VulkanConfiguration.h
@@ -121,10 +121,12 @@
 	#define VULKAN_SUPPORTS_VALIDATION_CACHE					0
 #endif
 
-#ifdef VK_KHR_dedicated_allocation
-	#define VULKAN_SUPPORTS_DEDICATED_ALLOCATION				1
-#else
-	#define VULKAN_SUPPORTS_DEDICATED_ALLOCATION				0
+#ifndef VULKAN_SUPPORTS_DEDICATED_ALLOCATION
+	#ifdef VK_KHR_dedicated_allocation
+		#define VULKAN_SUPPORTS_DEDICATED_ALLOCATION			1
+	#else
+		#define VULKAN_SUPPORTS_DEDICATED_ALLOCATION			0
+	#endif
 #endif
 
 #ifndef VULKAN_SUPPORTS_GOOGLE_DISPLAY_TIMING
@@ -165,34 +167,44 @@
 
 #define VULKAN_SUPPORTS_GPU_CRASH_DUMPS							(VULKAN_SUPPORTS_AMD_BUFFER_MARKER || VULKAN_SUPPORTS_NV_DIAGNOSTIC_CHECKPOINT)
 
-#ifdef VK_EXT_debug_utils
-	#define VULKAN_SUPPORTS_DEBUG_UTILS							1
-#else
-	#define VULKAN_SUPPORTS_DEBUG_UTILS							0
+#ifndef VULKAN_SUPPORTS_DEBUG_UTILS
+	#ifdef VK_EXT_debug_utils
+		#define VULKAN_SUPPORTS_DEBUG_UTILS							1
+	#else
+		#define VULKAN_SUPPORTS_DEBUG_UTILS							0
+	#endif
 #endif
 
-#ifdef VK_EXT_memory_priority
-	#define VULKAN_SUPPORTS_MEMORY_PRIORITY						1
-#else
-	#define VULKAN_SUPPORTS_MEMORY_PRIORITY						0
+#ifndef VULKAN_SUPPORTS_MEMORY_PRIORITY
+	#ifdef VK_EXT_memory_priority
+		#define VULKAN_SUPPORTS_MEMORY_PRIORITY						1
+	#else
+		#define VULKAN_SUPPORTS_MEMORY_PRIORITY						0
+	#endif
 #endif
 
-#ifdef VK_KHR_get_physical_device_properties2
-	#define VULKAN_SUPPORTS_PHYSICAL_DEVICE_PROPERTIES2			1
-#else
-	#define VULKAN_SUPPORTS_PHYSICAL_DEVICE_PROPERTIES2			0
+#ifndef VULKAN_SUPPORTS_PHYSICAL_DEVICE_PROPERTIES2
+	#ifdef VK_KHR_get_physical_device_properties2
+		#define VULKAN_SUPPORTS_PHYSICAL_DEVICE_PROPERTIES2			1
+	#else
+		#define VULKAN_SUPPORTS_PHYSICAL_DEVICE_PROPERTIES2			0
+	#endif
 #endif
 
-#ifdef VK_KHR_external_memory_capabilities
-	#define VULKAN_SUPPORTS_EXTERNAL_MEMORY							(VULKAN_SUPPORTS_PHYSICAL_DEVICE_PROPERTIES2)	// Requirement
-#else
-	#define VULKAN_SUPPORTS_EXTERNAL_MEMORY							0
+#ifndef VULKAN_SUPPORTS_EXTERNAL_MEMORY
+	#ifdef VK_KHR_external_memory_capabilities
+		#define VULKAN_SUPPORTS_EXTERNAL_MEMORY							(VULKAN_SUPPORTS_PHYSICAL_DEVICE_PROPERTIES2)	// Requirement
+	#else
+		#define VULKAN_SUPPORTS_EXTERNAL_MEMORY							0
+	#endif
 #endif
 
-#ifdef VK_KHR_driver_properties
-	#define VULKAN_SUPPORTS_DRIVER_PROPERTIES						1
-#else
-	#define VULKAN_SUPPORTS_DRIVER_PROPERTIES						0
+#ifndef VULKAN_SUPPORTS_DRIVER_PROPERTIES
+	#ifdef VK_KHR_driver_properties
+		#define VULKAN_SUPPORTS_DRIVER_PROPERTIES						1
+	#else
+		#define VULKAN_SUPPORTS_DRIVER_PROPERTIES						0
+	#endif
 #endif
 
 DECLARE_LOG_CATEGORY_EXTERN(LogVulkanRHI, Log, All);
