DRS - Compile fix
#rnx

diff --git a/Engine/Source/Runtime/VulkanRHI/Private/Lumin/VulkanLuminPlatform.h b/Engine/Source/Runtime/VulkanRHI/Private/Lumin/VulkanLuminPlatform.h
--- a/Engine/Source/Runtime/VulkanRHI/Private/Lumin/VulkanLuminPlatform.h
+++ b/Engine/Source/Runtime/VulkanRHI/Private/Lumin/VulkanLuminPlatform.h
@@ -10,6 +10,7 @@
 #define VULKAN_USE_IMAGE_ACQUIRE_FENCES				0
 #define VULKAN_SUPPORTS_COLOR_CONVERSIONS			1
 #define VULKAN_SUPPORTS_NV_DIAGNOSTIC_CHECKPOINT	0
+#define VULKAN_SUPPORTS_PHYSICAL_DEVICE_PROPERTIES2	0
 
 
 #define ENUM_VK_ENTRYPOINTS_PLATFORM_BASE(EnumMacro)
