DRS - Static analysis fix
#rnx

diff --git a/Engine/Source/Runtime/VulkanRHI/Private/VulkanRHI.cpp b/Engine/Source/Runtime/VulkanRHI/Private/VulkanRHI.cpp
--- a/Engine/Source/Runtime/VulkanRHI/Private/VulkanRHI.cpp
+++ b/Engine/Source/Runtime/VulkanRHI/Private/VulkanRHI.cpp
@@ -502,10 +502,8 @@ void FVulkanDynamicRHI::SelectAndInitDevice()
 	VERIFYVULKANRESULT_EXPANDED(VulkanRHI::vkEnumeratePhysicalDevices(Instance, &GpuCount, PhysicalDevices.GetData()));
 	checkf(GpuCount >= 1, TEXT("Couldn't enumerate physical devices! Make sure your drivers are up to date and that you are not pending a reboot."));
 
-#if VULKAN_ENABLE_DESKTOP_HMD_SUPPORT
 	FVulkanDevice* HmdDevice = nullptr;
 	uint32 HmdDeviceIndex = 0;
-#endif
 	struct FDeviceInfo
 	{
 		FVulkanDevice* Device;
@@ -551,14 +549,11 @@ void FVulkanDynamicRHI::SelectAndInitDevice()
 	}
 
 	uint32 DeviceIndex = -1;
-
-#if VULKAN_ENABLE_DESKTOP_HMD_SUPPORT
 	if (HmdDevice)
 	{
 		Device = HmdDevice;
 		DeviceIndex = HmdDeviceIndex;
 	}
-#endif
 
 	// Add all integrated to the end of the list
 	DiscreteDevices.Append(IntegratedDevices);
