DRS - Fix static analysis
#rnx

diff --git a/Engine/Source/Runtime/VulkanRHI/Private/VulkanRHI.cpp b/Engine/Source/Runtime/VulkanRHI/Private/VulkanRHI.cpp
--- a/Engine/Source/Runtime/VulkanRHI/Private/VulkanRHI.cpp
+++ b/Engine/Source/Runtime/VulkanRHI/Private/VulkanRHI.cpp
@@ -549,11 +549,13 @@ void FVulkanDynamicRHI::SelectAndInitDevice()
 	}
 
 	uint32 DeviceIndex = -1;
+#if VULKAN_ENABLE_DESKTOP_HMD_SUPPORT
 	if (HmdDevice)
 	{
 		Device = HmdDevice;
 		DeviceIndex = HmdDeviceIndex;
 	}
+#endif
 
 	// Add all integrated to the end of the list
 	DiscreteDevices.Append(IntegratedDevices);
@@ -563,8 +565,8 @@ void FVulkanDynamicRHI::SelectAndInitDevice()
 	int32 CVarExplicitAdapterValue = CVarGraphicsAdapter ? CVarGraphicsAdapter->GetValueOnAnyThread() : -1;
 	FParse::Value(FCommandLine::Get(), TEXT("graphicsadapter="), CVarExplicitAdapterValue);
 
-	// If HMD didn't choose one...
-	if (DeviceIndex == -1)
+	// If HMD didn't choose one... (disable static analysis that DeviceIndex is always -1)
+	if (DeviceIndex == -1)	//-V547
 	{
 		if (CVarExplicitAdapterValue >= (int32)GpuCount)
 		{
