DRS - Unblock integration; movie textures with immutable samplers not working on ES3.1
#rnx

diff --git a/Engine/Source/Runtime/VulkanRHI/Private/Lumin/VulkanLuminPlatform.cpp b/Engine/Source/Runtime/VulkanRHI/Private/Lumin/VulkanLuminPlatform.cpp
--- a/Engine/Source/Runtime/VulkanRHI/Private/Lumin/VulkanLuminPlatform.cpp
+++ b/Engine/Source/Runtime/VulkanRHI/Private/Lumin/VulkanLuminPlatform.cpp
@@ -69,6 +69,10 @@ bool FVulkanLuminPlatform::LoadVulkanInstanceFunctions(VkInstance inInstance)
 	ENUM_VK_ENTRYPOINTS_OPTIONAL_PLATFORM_INSTANCE(GETINSTANCE_VK_ENTRYPOINTS);
 	ENUM_VK_ENTRYPOINTS_OPTIONAL_PLATFORM_INSTANCE(CHECK_VK_ENTRYPOINTS);
 
+
+	//#todo-rco: Media textures are not working properly, quick workaround
+	GValidationCvar->Set(2, ECVF_SetByCommandline);
+
 	return bFoundAllEntryPoints;
 }
 
