Oculus rendering changes integration from partner-oculus-staging

diff --git a/Engine/Source/Runtime/RHI/Private/RHI.cpp b/Engine/Source/Runtime/RHI/Private/RHI.cpp
--- a/Engine/Source/Runtime/RHI/Private/RHI.cpp
+++ b/Engine/Source/Runtime/RHI/Private/RHI.cpp
@@ -945,6 +945,8 @@ void FRHIRenderPassInfo::ConvertToRenderTargetsInfo(FRHISetRenderTargetsInfo& Ou
 	OutRTInfo.bClearDepth = (DepthLoadAction == ERenderTargetLoadAction::EClear);
 	OutRTInfo.bClearStencil = (StencilLoadAction == ERenderTargetLoadAction::EClear);
 
+	OutRTInfo.FoveationTexture = FoveationTexture;
+
 	if (NumUAVs > 0)
 	{
 		check(UAVIndex != -1);
diff --git a/Engine/Source/Runtime/RHI/Public/RHIResources.h b/Engine/Source/Runtime/RHI/Public/RHIResources.h
--- a/Engine/Source/Runtime/RHI/Public/RHIResources.h
+++ b/Engine/Source/Runtime/RHI/Public/RHIResources.h
@@ -1611,6 +1611,8 @@ public:
 	bool bClearDepth;
 	bool bClearStencil;
 
+	FRHITexture* FoveationTexture;
+
 	// UAVs info.
 	FUnorderedAccessViewRHIRef UnorderedAccessView[MaxSimultaneousUAVs];
 	int32 NumUAVs;
@@ -1621,6 +1623,7 @@ public:
 		bHasResolveAttachments(false),
 		bClearDepth(false),
 		bClearStencil(false),
+		FoveationTexture(nullptr),
 		NumUAVs(0)
 	{}
 
@@ -1631,6 +1634,7 @@ public:
 		DepthStencilRenderTarget(InDepthStencilRenderTarget),		
 		bClearDepth(InDepthStencilRenderTarget.Texture && InDepthStencilRenderTarget.DepthLoadAction == ERenderTargetLoadAction::EClear),
 		bClearStencil(InDepthStencilRenderTarget.Texture && InDepthStencilRenderTarget.StencilLoadAction == ERenderTargetLoadAction::EClear),
+		FoveationTexture(nullptr),
 		NumUAVs(0)
 	{
 		check(InNumColorRenderTargets <= 0 || InColorRenderTargets);
@@ -1659,8 +1663,8 @@ public:
 		// Need a separate struct so we can memzero/remove dependencies on reference counts
 		struct FHashableStruct
 		{
-			// Depth goes in the last slot
-			FRHITexture* Texture[MaxSimultaneousRenderTargets*2 + 1];
+			// *2 for color and resolves, depth goes in the second-to-last slot, fixed foveation goes in the last slot
+			FRHITexture* Texture[MaxSimultaneousRenderTargets*2 + 2];
 			uint32 MipIndex[MaxSimultaneousRenderTargets];
 			uint32 ArraySliceIndex[MaxSimultaneousRenderTargets];
 			ERenderTargetLoadAction LoadAction[MaxSimultaneousRenderTargets];
@@ -1692,6 +1696,7 @@ public:
 				}
 
 				Texture[MaxSimultaneousRenderTargets] = RTInfo.DepthStencilRenderTarget.Texture;
+				Texture[MaxSimultaneousRenderTargets + 1] = RTInfo.FoveationTexture;
 				DepthLoadAction = RTInfo.DepthStencilRenderTarget.DepthLoadAction;
 				DepthStoreAction = RTInfo.DepthStencilRenderTarget.DepthStoreAction;
 				StencilLoadAction = RTInfo.DepthStencilRenderTarget.StencilLoadAction;
@@ -2518,6 +2523,10 @@ struct FRHIRenderPassInfo
 
 	FResolveParams ResolveParameters;
 
+	// Some RHIs can use a texture to control the sampling and/or shading resolution of different areas 
+	// (@todo: This implementation is specific to fixed foveated rendering, and will be updated or replaced as a more general pathway comes online)
+	FTextureRHIRef FoveationTexture = nullptr;
+
 	// Some RHIs require a hint that occlusion queries will be used in this render pass
 	uint32 NumOcclusionQueries = 0;
 	bool bOcclusionQueries = false;
@@ -2715,6 +2724,27 @@ struct FRHIRenderPassInfo
 		FMemory::Memzero(&ColorRenderTargets[1], sizeof(FColorEntry) * (MaxSimultaneousRenderTargets - 1));
 	}
 
+	// Color and depth with resolve and optional sample density
+	explicit FRHIRenderPassInfo(FRHITexture* ColorRT, ERenderTargetActions ColorAction, FRHITexture* ResolveColorRT,
+		FRHITexture* DepthRT, EDepthStencilTargetActions DepthActions, FRHITexture* ResolveDepthRT, FRHITexture* InFoveationTexture,
+		FExclusiveDepthStencil InEDS = FExclusiveDepthStencil::DepthWrite_StencilWrite)
+	{
+		check(ColorRT);
+		ColorRenderTargets[0].RenderTarget = ColorRT;
+		ColorRenderTargets[0].ResolveTarget = ResolveColorRT;
+		ColorRenderTargets[0].ArraySlice = -1;
+		ColorRenderTargets[0].MipIndex = 0;
+		ColorRenderTargets[0].Action = ColorAction;
+		bIsMSAA = ColorRT->GetNumSamples() > 1;
+		check(DepthRT);
+		DepthStencilRenderTarget.DepthStencilTarget = DepthRT;
+		DepthStencilRenderTarget.ResolveTarget = ResolveDepthRT;
+		DepthStencilRenderTarget.Action = DepthActions;
+		DepthStencilRenderTarget.ExclusiveDepthStencil = InEDS;
+		FoveationTexture = InFoveationTexture;
+		FMemory::Memzero(&ColorRenderTargets[1], sizeof(FColorEntry) * (MaxSimultaneousRenderTargets - 1));
+	}
+
 	explicit FRHIRenderPassInfo(int32 InNumUAVs, FRHIUnorderedAccessView** InUAVs)
 	{
 		if (InNumUAVs > MaxSimultaneousUAVs)
diff --git a/Engine/Source/Runtime/Renderer/Private/MobileSceneCaptureRendering.cpp b/Engine/Source/Runtime/Renderer/Private/MobileSceneCaptureRendering.cpp
--- a/Engine/Source/Runtime/Renderer/Private/MobileSceneCaptureRendering.cpp
+++ b/Engine/Source/Runtime/Renderer/Private/MobileSceneCaptureRendering.cpp
@@ -138,6 +138,7 @@ IMPLEMENT_SHADER_TYPE(template<>, FMobileSceneCaptureCopyPS##SCENETYPE, TEXT("/E
 IMPLEMENT_SHADER_TYPE(template<>, FMobileSceneCaptureCopyPS_Mosaic##SCENETYPE, TEXT("/Engine/Private/MobileSceneCapture.usf"), TEXT("MainCopyPS"), SF_Pixel);
 
 IMPLEMENT_MOBILE_SCENE_CAPTURECOPY(SCS_SceneColorHDR);
+IMPLEMENT_MOBILE_SCENE_CAPTURECOPY(SCS_FinalColorLDR);
 IMPLEMENT_MOBILE_SCENE_CAPTURECOPY(SCS_SceneColorHDRNoAlpha);
 IMPLEMENT_MOBILE_SCENE_CAPTURECOPY(SCS_SceneColorSceneDepth);
 IMPLEMENT_MOBILE_SCENE_CAPTURECOPY(SCS_SceneDepth);
@@ -171,6 +172,7 @@ static FShader* SetCaptureToTargetShaders(FRHICommandListImmediate& RHICmdList,
 		case SCS_SceneColorHDR:
 			return SetCaptureToTargetShaders<bDemosaic, SCS_SceneColorHDR>(RHICmdList, GraphicsPSOInit, View, SourceTexSize, SourceTextureRHI);
 		case SCS_FinalColorLDR:
+			return SetCaptureToTargetShaders<bDemosaic, SCS_FinalColorLDR>(RHICmdList, GraphicsPSOInit, View, SourceTexSize, SourceTextureRHI);
 		case SCS_SceneColorHDRNoAlpha:
 			return SetCaptureToTargetShaders<bDemosaic, SCS_SceneColorHDRNoAlpha>(RHICmdList, GraphicsPSOInit, View, SourceTexSize, SourceTextureRHI);
 		case SCS_SceneColorSceneDepth:
diff --git a/Engine/Source/Runtime/Renderer/Private/MobileShadingRenderer.cpp b/Engine/Source/Runtime/Renderer/Private/MobileShadingRenderer.cpp
--- a/Engine/Source/Runtime/Renderer/Private/MobileShadingRenderer.cpp
+++ b/Engine/Source/Runtime/Renderer/Private/MobileShadingRenderer.cpp
@@ -493,13 +493,21 @@ void FMobileSceneRenderer::Render(FRHICommandListImmediate& RHICmdList)
 		}
 	}
 
+	FRHITexture* FoveationTexture = nullptr;
+	
+	if (SceneContext.IsFoveationTextureAllocated())
+	{
+		FoveationTexture = SceneContext.GetFoveationTexture();
+	}
+
 	FRHIRenderPassInfo SceneColorRenderPassInfo(
 		SceneColor,
 		ColorTargetAction,
 		SceneColorResolve,
 		SceneDepth,
-		DepthTargetAction, 
+		DepthTargetAction,
 		nullptr, // we never resolve scene depth on mobile
+		FoveationTexture,
 		FExclusiveDepthStencil::DepthWrite_StencilWrite
 	);
 	SceneColorRenderPassInfo.SubpassHint = ESubpassHint::DepthReadSubpass;
@@ -590,6 +598,7 @@ void FMobileSceneRenderer::Render(FRHICommandListImmediate& RHICmdList)
 			SceneDepth,
 			DepthTargetAction, 
 			nullptr,
+			FoveationTexture,
 			ExclusiveDepthStencil
 		);
 		TranslucentRenderPassInfo.NumOcclusionQueries = 0;
diff --git a/Engine/Source/Runtime/Renderer/Private/PostProcess/SceneRenderTargets.cpp b/Engine/Source/Runtime/Renderer/Private/PostProcess/SceneRenderTargets.cpp
--- a/Engine/Source/Runtime/Renderer/Private/PostProcess/SceneRenderTargets.cpp
+++ b/Engine/Source/Runtime/Renderer/Private/PostProcess/SceneRenderTargets.cpp
@@ -260,6 +260,7 @@ FSceneRenderTargets::FSceneRenderTargets(const FViewInfo& View, const FSceneRend
 	, EditorPrimitivesDepth(GRenderTargetPool.MakeSnapshot(SnapshotSource.EditorPrimitivesDepth))
 	, SeparateTranslucencyRT(SnapshotSource.SeparateTranslucencyRT)
 	, DownsampledTranslucencyDepthRT(SnapshotSource.DownsampledTranslucencyDepthRT)
+	, FoveationTexture(GRenderTargetPool.MakeSnapshot(SnapshotSource.FoveationTexture))
 	, bScreenSpaceAOIsValid(SnapshotSource.bScreenSpaceAOIsValid)
 	, bCustomDepthIsValid(SnapshotSource.bCustomDepthIsValid)
 	, GBufferRefCount(SnapshotSource.GBufferRefCount)
@@ -295,6 +296,7 @@ FSceneRenderTargets::FSceneRenderTargets(const FViewInfo& View, const FSceneRend
 	, QuadOverdrawIndex(SnapshotSource.QuadOverdrawIndex)
 	, bHMDAllocatedDepthTarget(SnapshotSource.bHMDAllocatedDepthTarget)
 	, bKeepDepthContent(SnapshotSource.bKeepDepthContent)
+	, bAllocatedFoveationTexture(SnapshotSource.bAllocatedFoveationTexture)
 {
 	FMemory::Memcpy(LargestDesiredSizes, SnapshotSource.LargestDesiredSizes);
 #if PREVENT_RENDERTARGET_SIZE_THRASHING
@@ -1080,6 +1082,8 @@ void FSceneRenderTargets::AllocMobileMultiViewSceneColor(FRHICommandList& RHICmd
 		MobileMultiViewSceneColor.SafeRelease();
 	}
 
+	bool MobileHWsRGB = IsMobileColorsRGB();
+
 	if (!MobileMultiViewSceneColor)
 	{
 		const EPixelFormat SceneColorBufferFormat = GetSceneColorFormat();
@@ -1089,6 +1093,7 @@ void FSceneRenderTargets::AllocMobileMultiViewSceneColor(FRHICommandList& RHICmd
 		Desc.NumSamples = GetNumSceneColorMSAASamples(CurrentFeatureLevel);
 		Desc.ArraySize = 2;
 		Desc.bIsArray = true;
+		Desc.Flags |= MobileHWsRGB ? TexCreate_SRGB : 0;
 		GRenderTargetPool.FindFreeElement(RHICmdList, Desc, MobileMultiViewSceneColor, TEXT("MobileMultiViewSceneColor"));
 	}
 	check(MobileMultiViewSceneColor);
@@ -2034,6 +2039,7 @@ void FSceneRenderTargets::AllocateMobileRenderTargets(FRHICommandList& RHICmdLis
 	// on ES2 we don't do on demand allocation of SceneColor yet (in non ES2 it's released in the Tonemapper Process())
 	AllocSceneColor(RHICmdList);
 	AllocateCommonDepthTargets(RHICmdList);
+	AllocateFoveationTexture(RHICmdList);
 
 #if PLATFORM_ANDROID
 	static const auto MobileMultiViewCVar = IConsoleManager::Get().FindTConsoleVariableDataInt(TEXT("vr.MobileMultiView"));
@@ -2260,6 +2266,30 @@ void FSceneRenderTargets::AllocateCommonDepthTargets(FRHICommandList& RHICmdList
 	}
 }
 
+void FSceneRenderTargets::AllocateFoveationTexture(FRHICommandList& RHICmdList)
+{
+	const bool bStereo = GEngine->StereoRenderingDevice.IsValid() && GEngine->StereoRenderingDevice->IsStereoEnabled();
+	IStereoRenderTargetManager* const StereoRenderTargetManager = bStereo ? GEngine->StereoRenderingDevice->GetRenderTargetManager() : nullptr;
+	
+	FTexture2DRHIRef Texture;
+	FIntPoint TextureSize;
+
+	// Allocate variable resolution texture for VR foveation if supported
+	if (StereoRenderTargetManager && StereoRenderTargetManager->NeedReAllocateFoveationTexture(FoveationTexture))
+	{
+		FoveationTexture.SafeRelease();
+	}
+	bAllocatedFoveationTexture = StereoRenderTargetManager && StereoRenderTargetManager->AllocateFoveationTexture(0, BufferSize.X, BufferSize.Y, PF_R8G8, 0, TexCreate_None, TexCreate_None, Texture, TextureSize);
+	if (bAllocatedFoveationTexture)
+	{
+		FPooledRenderTargetDesc Desc(FPooledRenderTargetDesc::Create2DDesc(TextureSize, PF_R8G8, FClearValueBinding::White, TexCreate_None, TexCreate_None, false));
+		GRenderTargetPool.FindFreeElement(RHICmdList, Desc, FoveationTexture, TEXT("FixedFoveation"));
+		const uint32 OldElementSize = FoveationTexture->ComputeMemorySize();
+		FoveationTexture->GetRenderTargetItem().ShaderResourceTexture = FoveationTexture->GetRenderTargetItem().TargetableTexture = Texture;
+		GRenderTargetPool.UpdateElementSize(FoveationTexture, OldElementSize);
+	}
+}
+
 void FSceneRenderTargets::AllocateScreenShadowMask(FRHICommandList& RHICmdList, TRefCountPtr<IPooledRenderTarget>& ScreenShadowMaskTexture)
 {
 	FPooledRenderTargetDesc Desc(FPooledRenderTargetDesc::Create2DDesc(GetBufferSizeXY(), PF_B8G8R8A8, FClearValueBinding::White, TexCreate_None, TexCreate_RenderTargetable, false));
@@ -2632,7 +2662,7 @@ EPixelFormat FSceneRenderTargets::GetSceneColorFormat(ERHIFeatureLevel::Type InF
 
 void FSceneRenderTargets::AllocateRenderTargets(FRHICommandListImmediate& RHICmdList, const int32 NumViews)
 {
-	if (BufferSize.X > 0 && BufferSize.Y > 0 && (!AreShadingPathRenderTargetsAllocated(GetSceneColorFormatType()) || !AreRenderTargetClearsValid(GetSceneColorFormatType())))
+	if (BufferSize.X > 0 && BufferSize.Y > 0 && IsAllocateRenderTargetsRequired())
 	{
 		if ((EShadingPath)CurrentShadingPath == EShadingPath::Mobile)
 		{
@@ -2726,6 +2756,8 @@ void FSceneRenderTargets::ReleaseAllTargets()
 	EditorPrimitivesColor.SafeRelease();
 	EditorPrimitivesDepth.SafeRelease();
 
+	FoveationTexture.SafeRelease();
+
 	VirtualTextureFeedback.ReleaseResources();
 }
 
@@ -2991,6 +3023,22 @@ bool FSceneRenderTargets::AreShadingPathRenderTargetsAllocated(ESceneColorFormat
 	}
 }
 
+bool FSceneRenderTargets::IsAllocateRenderTargetsRequired() const
+{
+	bool bAllocateRequired = false;
+
+	const bool bStereo = GEngine->StereoRenderingDevice.IsValid() && GEngine->StereoRenderingDevice->IsStereoEnabled();
+	IStereoRenderTargetManager* const StereoRenderTargetManager = bStereo ? GEngine->StereoRenderingDevice->GetRenderTargetManager() : nullptr;
+
+	// HMD controlled foveation textures may be destroyed externally and need a new allocation
+	if (StereoRenderTargetManager && StereoRenderTargetManager->NeedReAllocateFoveationTexture(FoveationTexture))
+	{
+		bAllocateRequired = true;
+	}
+	
+	return bAllocateRequired || !AreShadingPathRenderTargetsAllocated(GetSceneColorFormatType()) || !AreRenderTargetClearsValid(GetSceneColorFormatType());
+}
+
 /*-----------------------------------------------------------------------------
 FSceneTextureShaderParameters
 -----------------------------------------------------------------------------*/
diff --git a/Engine/Source/Runtime/Renderer/Private/PostProcess/SceneRenderTargets.h b/Engine/Source/Runtime/Renderer/Private/PostProcess/SceneRenderTargets.h
--- a/Engine/Source/Runtime/Renderer/Private/PostProcess/SceneRenderTargets.h
+++ b/Engine/Source/Runtime/Renderer/Private/PostProcess/SceneRenderTargets.h
@@ -186,7 +186,8 @@ protected:
 		DefaultDepthClear(FClearValueBinding::DepthFar),
 		QuadOverdrawIndex(INDEX_NONE),
 		bHMDAllocatedDepthTarget(false),
-		bKeepDepthContent(true)
+		bKeepDepthContent(true),
+		bAllocatedFoveationTexture(false)
 		{
 			FMemory::Memset(LargestDesiredSizes, 0);
 #if PREVENT_RENDERTARGET_SIZE_THRASHING
@@ -373,6 +374,7 @@ public:
 	const FTexture2DRHIRef& GetGBufferDTexture() const { return (const FTexture2DRHIRef&)GBufferD->GetRenderTargetItem().ShaderResourceTexture; }
 	const FTexture2DRHIRef& GetGBufferETexture() const { return (const FTexture2DRHIRef&)GBufferE->GetRenderTargetItem().ShaderResourceTexture; }
 	const FTexture2DRHIRef& GetGBufferVelocityTexture() const { return (const FTexture2DRHIRef&)SceneVelocity->GetRenderTargetItem().ShaderResourceTexture; }
+	const FTexture2DRHIRef& GetFoveationTexture() const { return (const FTexture2DRHIRef&)FoveationTexture->GetRenderTargetItem().ShaderResourceTexture; }
 
 	const FTextureRHIRef& GetLightAttenuationTexture() const
 	{
@@ -501,6 +503,7 @@ public:
 	void ReleaseSceneColor();
 	
 	ERHIFeatureLevel::Type GetCurrentFeatureLevel() const { return CurrentFeatureLevel; }
+	bool IsFoveationTextureAllocated() const { return bAllocatedFoveationTexture;  }
 
 private: // Get...() methods instead of direct access
 
@@ -584,6 +587,9 @@ public:
 	/** Downsampled depth used when rendering translucency in smaller resolution. */
 	TRefCountPtr<IPooledRenderTarget> DownsampledTranslucencyDepthRT;
 
+	/** Texture to control variable resolution rendering */
+	TRefCountPtr<IPooledRenderTarget> FoveationTexture;
+
 	// todo: free ScreenSpaceAO so pool can reuse
 	bool bScreenSpaceAOIsValid;
 
@@ -639,6 +645,9 @@ private:
 	/** Allocates common depth render targets that are used by both mobile and deferred rendering paths */
 	void AllocateCommonDepthTargets(FRHICommandList& RHICmdList);
 
+	/** Allocates a texture for controlling variable resolution rendering. */
+	void AllocateFoveationTexture(FRHICommandList& RHICmdList);
+
 	/** Determine the appropriate render target dimensions. */
 	FIntPoint ComputeDesiredSize(const FSceneViewFamily& ViewFamily);
 
@@ -664,6 +673,9 @@ private:
 	/** Determine if the default clear values for color and depth match the allocated scene render targets. Mobile only. */
 	bool AreRenderTargetClearsValid(ESceneColorFormatType InSceneColorFormatType) const;
 
+	/** Determine if an allocate is required for the render targets. */
+	bool IsAllocateRenderTargetsRequired() const;
+
 	/** Determine whether the render targets for any shading path have been allocated */
 	bool AreAnyShadingPathRenderTargetsAllocated() const 
 	{ 
@@ -769,6 +781,9 @@ private:
 	/** True if the contents of the depth buffer must be kept for post-processing. When this is false, the depth buffer can be allocated as memoryless on mobile platforms which support it. */
 	bool bKeepDepthContent;
 
+	/** True if the a variable resolution texture is allocated to control sampling or shading rate */
+	bool bAllocatedFoveationTexture;
+
 	/** CAUTION: When adding new data, make sure you copy it in the snapshot constructor! **/
 
 };
diff --git a/Engine/Source/Runtime/VulkanRHI/Private/Android/VulkanAndroidPlatform.cpp b/Engine/Source/Runtime/VulkanRHI/Private/Android/VulkanAndroidPlatform.cpp
--- a/Engine/Source/Runtime/VulkanRHI/Private/Android/VulkanAndroidPlatform.cpp
+++ b/Engine/Source/Runtime/VulkanRHI/Private/Android/VulkanAndroidPlatform.cpp
@@ -151,6 +151,7 @@ void FVulkanAndroidPlatform::GetDeviceExtensions(EGpuVendorId VendorId, TArray<c
 	OutExtensions.Add(VK_KHR_SURFACE_EXTENSION_NAME);
 	OutExtensions.Add(VK_KHR_ANDROID_SURFACE_EXTENSION_NAME);
 	OutExtensions.Add(VK_GOOGLE_DISPLAY_TIMING_EXTENSION_NAME);
+	OutExtensions.Add(VK_EXT_FRAGMENT_DENSITY_MAP_EXTENSION_NAME);
 }
 
 bool FVulkanAndroidPlatform::SupportsStandardSwapchain()
@@ -177,6 +178,12 @@ EPixelFormat FVulkanAndroidPlatform::GetPixelFormatForNonDefaultSwapchain()
 	}
 }
 
+bool FVulkanAndroidPlatform::SupportsTimestampRenderQueries()
+{
+	// standalone devices have newer drivers where timestamp render queries work.
+	return FPlatformMisc::IsStandaloneStereoOnlyDevice();
+}
+
 void FVulkanAndroidPlatform::OverridePlatformHandlers(bool bInit)
 {
 	if (bInit)
diff --git a/Engine/Source/Runtime/VulkanRHI/Private/Android/VulkanAndroidPlatform.h b/Engine/Source/Runtime/VulkanRHI/Private/Android/VulkanAndroidPlatform.h
--- a/Engine/Source/Runtime/VulkanRHI/Private/Android/VulkanAndroidPlatform.h
+++ b/Engine/Source/Runtime/VulkanRHI/Private/Android/VulkanAndroidPlatform.h
@@ -70,7 +70,7 @@ public:
 	static bool SupportsStandardSwapchain();
 	static EPixelFormat GetPixelFormatForNonDefaultSwapchain();
 
-	static bool SupportsTimestampRenderQueries() { return false; }
+	static bool SupportsTimestampRenderQueries();
 
 	static bool RequiresMobileRenderer() { return true; }
 	static void OverridePlatformHandlers(bool bInit);
diff --git a/Engine/Source/Runtime/VulkanRHI/Private/VulkanCommands.cpp b/Engine/Source/Runtime/VulkanRHI/Private/VulkanCommands.cpp
--- a/Engine/Source/Runtime/VulkanRHI/Private/VulkanCommands.cpp
+++ b/Engine/Source/Runtime/VulkanRHI/Private/VulkanCommands.cpp
@@ -1114,7 +1114,7 @@ void FVulkanCommandListContext::RHICopyToStagingBuffer(FRHIVertexBuffer* SourceB
 			Device->GetStagingManager().ReleaseBuffer(nullptr, StagingBuffer->StagingBuffer);
 		}
 
-		VulkanRHI::FStagingBuffer* ReadbackStagingBuffer = Device->GetStagingManager().AcquireBuffer(NumBytes, VK_BUFFER_USAGE_TRANSFER_DST_BIT, true);
+		VulkanRHI::FStagingBuffer* ReadbackStagingBuffer = Device->GetStagingManager().AcquireBuffer(NumBytes, VK_BUFFER_USAGE_TRANSFER_DST_BIT, VK_MEMORY_PROPERTY_HOST_CACHED_BIT);
 		StagingBuffer->StagingBuffer = ReadbackStagingBuffer;
 		StagingBuffer->Device = Device;
 	}
diff --git a/Engine/Source/Runtime/VulkanRHI/Private/VulkanDevice.h b/Engine/Source/Runtime/VulkanRHI/Private/VulkanDevice.h
--- a/Engine/Source/Runtime/VulkanRHI/Private/VulkanDevice.h
+++ b/Engine/Source/Runtime/VulkanRHI/Private/VulkanDevice.h
@@ -35,6 +35,7 @@ struct FOptionalVulkanDeviceExtensions
 			uint32 HasYcbcrSampler : 1;
 			uint32 HasMemoryPriority : 1;
 			uint32 HasDriverProperties : 1;
+			uint32 HasEXTFragmentDensityMap : 1;
 		};
 		uint32 Packed;
 	};
diff --git a/Engine/Source/Runtime/VulkanRHI/Private/VulkanIndexBuffer.cpp b/Engine/Source/Runtime/VulkanRHI/Private/VulkanIndexBuffer.cpp
--- a/Engine/Source/Runtime/VulkanRHI/Private/VulkanIndexBuffer.cpp
+++ b/Engine/Source/Runtime/VulkanRHI/Private/VulkanIndexBuffer.cpp
@@ -201,7 +201,7 @@ void* FVulkanResourceMultiBuffer::Lock(bool bFromRenderingThread, EResourceLockM
 				VulkanRHI::vkCmdPipelineBarrier(CmdBuffer->GetHandle(), VK_PIPELINE_STAGE_TRANSFER_BIT, VK_PIPELINE_STAGE_TRANSFER_BIT, 0, 0, 0, 1, &Barrier, 0, nullptr);
 
 				// Create a staging buffer we can use to copy data from device to cpu.
-				VulkanRHI::FStagingBuffer* StagingBuffer = Device->GetStagingManager().AcquireBuffer(Size, VK_BUFFER_USAGE_TRANSFER_DST_BIT, true);
+				VulkanRHI::FStagingBuffer* StagingBuffer = Device->GetStagingManager().AcquireBuffer(Size, VK_BUFFER_USAGE_TRANSFER_DST_BIT, VK_MEMORY_PROPERTY_HOST_CACHED_BIT);
 
 				// Fill the staging buffer with the data on the device.
 				VkBufferCopy Regions;
diff --git a/Engine/Source/Runtime/VulkanRHI/Private/VulkanLayers.cpp b/Engine/Source/Runtime/VulkanRHI/Private/VulkanLayers.cpp
--- a/Engine/Source/Runtime/VulkanRHI/Private/VulkanLayers.cpp
+++ b/Engine/Source/Runtime/VulkanRHI/Private/VulkanLayers.cpp
@@ -750,6 +750,8 @@ void FOptionalVulkanDeviceExtensions::Setup(const TArray<const ANSICHAR*>& Devic
 #if VULKAN_SUPPORTS_DRIVER_PROPERTIES
 	HasDriverProperties = HasExtension(DeviceExtensions, VK_KHR_DRIVER_PROPERTIES_EXTENSION_NAME);
 #endif
+
+	HasEXTFragmentDensityMap = HasExtension(DeviceExtensions, VK_EXT_FRAGMENT_DENSITY_MAP_EXTENSION_NAME);
 }
 
 void FVulkanDynamicRHI::SetupValidationRequests()
diff --git a/Engine/Source/Runtime/VulkanRHI/Private/VulkanMemory.cpp b/Engine/Source/Runtime/VulkanRHI/Private/VulkanMemory.cpp
--- a/Engine/Source/Runtime/VulkanRHI/Private/VulkanMemory.cpp
+++ b/Engine/Source/Runtime/VulkanRHI/Private/VulkanMemory.cpp
@@ -1796,13 +1796,13 @@ namespace VulkanRHI
 		check(FreeStagingBuffers.Num() == 0);
 	}
 
-	FStagingBuffer* FStagingManager::AcquireBuffer(uint32 Size, VkBufferUsageFlags InUsageFlags, bool bCPURead)
+	FStagingBuffer* FStagingManager::AcquireBuffer(uint32 Size, VkBufferUsageFlags InUsageFlags, VkMemoryPropertyFlagBits InMemoryReadFlags)
 	{
 #if VULKAN_ENABLE_AGGRESSIVE_STATS
 		SCOPE_CYCLE_COUNTER(STAT_VulkanStagingBuffer);
 #endif
 		LLM_SCOPE_VULKAN(ELLMTagVulkan::VulkanStagingBuffers);
-		if (bCPURead)
+		if (InMemoryReadFlags == VK_MEMORY_PROPERTY_HOST_CACHED_BIT)
 		{
 			uint64 NonCoherentAtomSize = (uint64)Device->GetLimits().nonCoherentAtomSize;
 			Size = AlignArbitrary(Size, NonCoherentAtomSize);
@@ -1814,7 +1814,7 @@ namespace VulkanRHI
 			for (int32 Index = 0; Index < FreeStagingBuffers.Num(); ++Index)
 			{
 				FFreeEntry& FreeBuffer = FreeStagingBuffers[Index];
-				if (FreeBuffer.StagingBuffer->GetSize() == Size && FreeBuffer.StagingBuffer->bCPURead == bCPURead)
+				if (FreeBuffer.StagingBuffer->GetSize() == Size && FreeBuffer.StagingBuffer->MemoryReadFlags == InMemoryReadFlags)
 				{
 					FStagingBuffer* Buffer = FreeBuffer.StagingBuffer;
 					FreeStagingBuffers.RemoveAtSwap(Index, 1, false);
@@ -1844,14 +1844,16 @@ namespace VulkanRHI
 
 		// Set minimum alignment to 16 bytes, as some buffers are used with CPU SIMD instructions
 		MemReqs.alignment = FMath::Max<VkDeviceSize>(16, MemReqs.alignment);
-		if (bCPURead)
+		if (InMemoryReadFlags == VK_MEMORY_PROPERTY_HOST_CACHED_BIT)
 		{
 			uint64 NonCoherentAtomSize = (uint64)Device->GetLimits().nonCoherentAtomSize;
 			MemReqs.alignment = AlignArbitrary(MemReqs.alignment, NonCoherentAtomSize);
 		}
 
-		StagingBuffer->ResourceAllocation = Device->GetResourceHeapManager().AllocateBufferMemory(MemReqs, VK_MEMORY_PROPERTY_HOST_VISIBLE_BIT | (bCPURead ? VK_MEMORY_PROPERTY_HOST_CACHED_BIT : VK_MEMORY_PROPERTY_HOST_COHERENT_BIT), __FILE__, __LINE__);
-		StagingBuffer->bCPURead = bCPURead;
+		VkMemoryPropertyFlags readTypeFlags = InMemoryReadFlags;
+
+		StagingBuffer->ResourceAllocation = Device->GetResourceHeapManager().AllocateBufferMemory(MemReqs, VK_MEMORY_PROPERTY_HOST_VISIBLE_BIT | readTypeFlags, __FILE__, __LINE__);
+		StagingBuffer->MemoryReadFlags = InMemoryReadFlags;
 		StagingBuffer->BufferSize = Size;
 		StagingBuffer->ResourceAllocation->BindBuffer(Device, StagingBuffer->Buffer);
 
diff --git a/Engine/Source/Runtime/VulkanRHI/Private/VulkanPipeline.cpp b/Engine/Source/Runtime/VulkanRHI/Private/VulkanPipeline.cpp
--- a/Engine/Source/Runtime/VulkanRHI/Private/VulkanPipeline.cpp
+++ b/Engine/Source/Runtime/VulkanRHI/Private/VulkanPipeline.cpp
@@ -958,14 +958,15 @@ FArchive& operator << (FArchive& Ar, FVulkanPipelineStateCacheManager::FGfxPipel
 
 void FVulkanPipelineStateCacheManager::FGfxPipelineEntry::FRenderTargets::ReadFrom(const FVulkanRenderTargetLayout& RTLayout)
 {
-	NumAttachments =			RTLayout.NumAttachmentDescriptions;
-	NumColorAttachments =		RTLayout.NumColorAttachments;
+	NumAttachments =				RTLayout.NumAttachmentDescriptions;
+	NumColorAttachments =			RTLayout.NumColorAttachments;
 
-	bHasDepthStencil =			RTLayout.bHasDepthStencil != 0;
-	bHasResolveAttachments =	RTLayout.bHasResolveAttachments != 0;
-	NumUsedClearValues =		RTLayout.NumUsedClearValues;
+	bHasDepthStencil =				RTLayout.bHasDepthStencil != 0;
+	bHasResolveAttachments =		RTLayout.bHasResolveAttachments != 0;
+	bHasFragmentDensityAttachment =	RTLayout.bHasFragmentDensityAttachment != 0;
+	NumUsedClearValues =			RTLayout.NumUsedClearValues;
 
-	RenderPassCompatibleHash =	RTLayout.GetRenderPassCompatibleHash();
+	RenderPassCompatibleHash =		RTLayout.GetRenderPassCompatibleHash();
 
 	Extent3D.X = RTLayout.Extent.Extent3D.width;
 	Extent3D.Y = RTLayout.Extent.Extent3D.height;
@@ -982,6 +983,7 @@ void FVulkanPipelineStateCacheManager::FGfxPipelineEntry::FRenderTargets::ReadFr
 	CopyAttachmentRefs(ColorAttachments, RTLayout.ColorReferences, ARRAY_COUNT(RTLayout.ColorReferences));
 	CopyAttachmentRefs(ResolveAttachments, RTLayout.ResolveReferences, ARRAY_COUNT(RTLayout.ResolveReferences));
 	DepthStencil.ReadFrom(RTLayout.DepthStencilReference);
+	FragmentDensity.ReadFrom(RTLayout.FragmentDensityReference);
 
 	Descriptions.AddZeroed(ARRAY_COUNT(RTLayout.Desc));
 	for (int32 Index = 0; Index < ARRAY_COUNT(RTLayout.Desc); ++Index)
@@ -992,19 +994,20 @@ void FVulkanPipelineStateCacheManager::FGfxPipelineEntry::FRenderTargets::ReadFr
 
 void FVulkanPipelineStateCacheManager::FGfxPipelineEntry::FRenderTargets::WriteInto(FVulkanRenderTargetLayout& Out) const
 {
-	Out.NumAttachmentDescriptions =	NumAttachments;
-	Out.NumColorAttachments =		NumColorAttachments;
+	Out.NumAttachmentDescriptions =		NumAttachments;
+	Out.NumColorAttachments =			NumColorAttachments;
 
-	Out.bHasDepthStencil =			bHasDepthStencil;
-	Out.bHasResolveAttachments =	bHasResolveAttachments;
-	Out.NumUsedClearValues =		NumUsedClearValues;
+	Out.bHasDepthStencil =				bHasDepthStencil;
+	Out.bHasResolveAttachments =		bHasResolveAttachments;
+	Out.bHasFragmentDensityAttachment =	bHasFragmentDensityAttachment;
+	Out.NumUsedClearValues =			NumUsedClearValues;
 
 	ensure(0);
-	Out.RenderPassCompatibleHash =	RenderPassCompatibleHash;
+	Out.RenderPassCompatibleHash =		RenderPassCompatibleHash;
 
-	Out.Extent.Extent3D.width =		Extent3D.X;
-	Out.Extent.Extent3D.height =	Extent3D.Y;
-	Out.Extent.Extent3D.depth =		Extent3D.Z;
+	Out.Extent.Extent3D.width =			Extent3D.X;
+	Out.Extent.Extent3D.height =		Extent3D.Y;
+	Out.Extent.Extent3D.depth =			Extent3D.Z;
 
 	auto CopyAttachmentRefs = [&](const TArray<FGfxPipelineEntry::FRenderTargets::FAttachmentRef>& Source, VkAttachmentReference* Dest, uint32 Count)
 	{
@@ -1016,6 +1019,7 @@ void FVulkanPipelineStateCacheManager::FGfxPipelineEntry::FRenderTargets::WriteI
 	CopyAttachmentRefs(ColorAttachments, Out.ColorReferences, ARRAY_COUNT(Out.ColorReferences));
 	CopyAttachmentRefs(ResolveAttachments, Out.ResolveReferences, ARRAY_COUNT(Out.ResolveReferences));
 	DepthStencil.WriteInto(Out.DepthStencilReference);
+	FragmentDensity.WriteInto(Out.FragmentDensityReference);
 
 	for (int32 Index = 0; Index < ARRAY_COUNT(Out.Desc); ++Index)
 	{
@@ -1032,6 +1036,7 @@ FArchive& operator << (FArchive& Ar, FVulkanPipelineStateCacheManager::FGfxPipel
 	Ar << RTs.ColorAttachments;
 	Ar << RTs.ResolveAttachments;
 	Ar << RTs.DepthStencil;
+	Ar << RTs.FragmentDensity;
 
 	Ar << RTs.Descriptions;
 
@@ -1736,6 +1741,7 @@ FVulkanRHIGraphicsPipelineState* FVulkanPipelineStateCacheManager::FindInRuntime
 					bool bSRGB = (PSI.RenderTargetFlags[Index] & TexCreate_SRGB) == TexCreate_SRGB;
 					TempEnumValue = UEToVkTextureFormat(PixelFormat, bSRGB);
 					Ar << TempEnumValue;
+					Ar << PSI.RenderTargetFlags[Index];
 				}
 				Ar << BlendState.bUseIndependentRenderTargetBlendStates;
 			}
diff --git a/Engine/Source/Runtime/VulkanRHI/Private/VulkanPipeline.h b/Engine/Source/Runtime/VulkanRHI/Private/VulkanPipeline.h
--- a/Engine/Source/Runtime/VulkanRHI/Private/VulkanPipeline.h
+++ b/Engine/Source/Runtime/VulkanRHI/Private/VulkanPipeline.h
@@ -316,6 +316,7 @@ public:
 			TArray<FAttachmentRef> ColorAttachments;
 			TArray<FAttachmentRef> ResolveAttachments;
 			FAttachmentRef DepthStencil;
+			FAttachmentRef FragmentDensity;
 
 			struct FAttachmentDesc
 			{
@@ -351,6 +352,7 @@ public:
 			uint8 NumColorAttachments;
 			uint8 bHasDepthStencil;
 			uint8 bHasResolveAttachments;
+			uint8 bHasFragmentDensityAttachment;
 			uint8 NumUsedClearValues;
 			uint32 RenderPassCompatibleHash;
 			FVector Extent3D;
@@ -363,11 +365,13 @@ public:
 				return ColorAttachments == In.ColorAttachments &&
 					ResolveAttachments == In.ResolveAttachments &&
 					DepthStencil == In.DepthStencil &&
+					FragmentDensity == In.FragmentDensity &&
 					Descriptions == In.Descriptions &&
 					NumAttachments == In.NumAttachments &&
 					NumColorAttachments == In.NumColorAttachments &&
 					bHasDepthStencil == In.bHasDepthStencil &&
 					bHasResolveAttachments == In.bHasResolveAttachments &&
+					bHasFragmentDensityAttachment == In.bHasFragmentDensityAttachment &&
 					NumUsedClearValues == In.NumUsedClearValues &&
 					RenderPassCompatibleHash == In.RenderPassCompatibleHash &&
 					Extent3D == In.Extent3D;
diff --git a/Engine/Source/Runtime/VulkanRHI/Private/VulkanQuery.cpp b/Engine/Source/Runtime/VulkanRHI/Private/VulkanQuery.cpp
--- a/Engine/Source/Runtime/VulkanRHI/Private/VulkanQuery.cpp
+++ b/Engine/Source/Runtime/VulkanRHI/Private/VulkanQuery.cpp
@@ -609,7 +609,7 @@ FVulkanTimingQuery::FVulkanTimingQuery(FVulkanDevice* InDevice)
 	: FVulkanRenderQuery(RQT_AbsoluteTime)
 	, Pool(InDevice, 4)
 {
-	Pool.ResultsBuffer = InDevice->GetStagingManager().AcquireBuffer(Pool.BufferSize * sizeof(uint64), VK_BUFFER_USAGE_TRANSFER_DST_BIT, true);
+	Pool.ResultsBuffer = InDevice->GetStagingManager().AcquireBuffer(Pool.BufferSize * sizeof(uint64), VK_BUFFER_USAGE_TRANSFER_DST_BIT, VK_MEMORY_PROPERTY_HOST_CACHED_BIT);
 }
 
 
diff --git a/Engine/Source/Runtime/VulkanRHI/Private/VulkanRHI.cpp b/Engine/Source/Runtime/VulkanRHI/Private/VulkanRHI.cpp
--- a/Engine/Source/Runtime/VulkanRHI/Private/VulkanRHI.cpp
+++ b/Engine/Source/Runtime/VulkanRHI/Private/VulkanRHI.cpp
@@ -1534,6 +1534,18 @@ static VkRenderPass CreateRenderPass(FVulkanDevice& InDevice, const FVulkanRende
 		CreateInfo.pNext = &MultiviewInfo;
 	}
 	
+	VkRenderPassFragmentDensityMapCreateInfoEXT FragDensityCreateInfo;
+	if (InDevice.GetOptionalExtensions().HasEXTFragmentDensityMap && RTLayout.GetHasFragmentDensityAttachment())
+	{
+		ZeroVulkanStruct(FragDensityCreateInfo, VK_STRUCTURE_TYPE_RENDER_PASS_FRAGMENT_DENSITY_MAP_CREATE_INFO_EXT);
+		FragDensityCreateInfo.fragmentDensityMapAttachment = *RTLayout.GetFragmentDensityAttachmentReference();
+
+		// Chain fragment density info onto create info and the rest of the pNexts
+		// onto the fragment density info
+		FragDensityCreateInfo.pNext = CreateInfo.pNext;
+		CreateInfo.pNext = &FragDensityCreateInfo;
+	}
+
 	VkRenderPass RenderPassHandle;
 	VERIFYVULKANRESULT_EXPANDED(VulkanRHI::vkCreateRenderPass(InDevice.GetInstanceHandle(), &CreateInfo, VULKAN_CPU_ALLOCATOR, &RenderPassHandle));
 	return RenderPassHandle;
diff --git a/Engine/Source/Runtime/VulkanRHI/Private/VulkanRHIPrivate.h b/Engine/Source/Runtime/VulkanRHI/Private/VulkanRHIPrivate.h
--- a/Engine/Source/Runtime/VulkanRHI/Private/VulkanRHIPrivate.h
+++ b/Engine/Source/Runtime/VulkanRHI/Private/VulkanRHIPrivate.h
@@ -122,6 +122,7 @@ public:
 	inline uint32 GetNumColorAttachments() const { return NumColorAttachments; }
 	inline bool GetHasDepthStencil() const { return bHasDepthStencil != 0; }
 	inline bool GetHasResolveAttachments() const { return bHasResolveAttachments != 0; }
+	inline bool GetHasFragmentDensityAttachment() const { return bHasFragmentDensityAttachment != 0; }
 	inline uint32 GetNumAttachmentDescriptions() const { return NumAttachmentDescriptions; }
 	inline uint32 GetNumSamples() const { return NumSamples; }
 	inline uint32 GetNumUsedClearValues() const { return NumUsedClearValues; }
@@ -130,22 +131,26 @@ public:
 	inline const VkAttachmentReference* GetColorAttachmentReferences() const { return NumColorAttachments > 0 ? ColorReferences : nullptr; }
 	inline const VkAttachmentReference* GetResolveAttachmentReferences() const { return bHasResolveAttachments ? ResolveReferences : nullptr; }
 	inline const VkAttachmentReference* GetDepthStencilAttachmentReference() const { return bHasDepthStencil ? &DepthStencilReference : nullptr; }
+	inline const VkAttachmentReference* GetFragmentDensityAttachmentReference() const { return bHasFragmentDensityAttachment ? &FragmentDensityReference : nullptr; }
 
 	inline const ESubpassHint GetSubpassHint() const { return SubpassHint; }
 
 protected:
 	VkAttachmentReference ColorReferences[MaxSimultaneousRenderTargets];
 	VkAttachmentReference DepthStencilReference;
+	VkAttachmentReference FragmentDensityReference;
 	VkAttachmentReference ResolveReferences[MaxSimultaneousRenderTargets];
 	VkAttachmentReference InputAttachments[MaxSimultaneousRenderTargets + 1];
 
-	VkAttachmentDescription Desc[MaxSimultaneousRenderTargets * 2 + 1];
+	// Depth goes in the "+1" slot and the Fixed Foveation texture goes in the "+2" slot.
+	VkAttachmentDescription Desc[MaxSimultaneousRenderTargets * 2 + 2];
 
 	uint8 NumAttachmentDescriptions;
 	uint8 NumColorAttachments;
 	uint8 NumInputAttachments = 0;
 	uint8 bHasDepthStencil;
 	uint8 bHasResolveAttachments;
+	uint8 bHasFragmentDensityAttachment;
 	uint8 NumSamples;
 	uint8 NumUsedClearValues;
 	ESubpassHint SubpassHint = ESubpassHint::None;
@@ -169,6 +174,7 @@ protected:
 	{
 		FMemory::Memzero(ColorReferences);
 		FMemory::Memzero(DepthStencilReference);
+		FMemory::Memzero(FragmentDensityReference);
 		FMemory::Memzero(ResolveReferences);
 		FMemory::Memzero(InputAttachments);
 		FMemory::Memzero(Desc);
@@ -176,6 +182,7 @@ protected:
 		NumColorAttachments = 0;
 		bHasDepthStencil = 0;
 		bHasResolveAttachments = 0;
+		bHasFragmentDensityAttachment = 0;
 		Extent.Extent3D.width = 0;
 		Extent.Extent3D.height = 0;
 		Extent.Extent3D.depth = 0;
@@ -262,6 +269,7 @@ private:
 	VkImage ColorRenderTargetImages[MaxSimultaneousRenderTargets];
 	VkImage ColorResolveTargetImages[MaxSimultaneousRenderTargets];
 	VkImage DepthStencilRenderTargetImage;
+	VkImage FragmentDensityImage;
 
 	// Predefined set of barriers, when executes ensuring all writes are finished
 	TArray<VkImageMemoryBarrier> WriteBarriers;
diff --git a/Engine/Source/Runtime/VulkanRHI/Private/VulkanRenderTarget.cpp b/Engine/Source/Runtime/VulkanRHI/Private/VulkanRenderTarget.cpp
--- a/Engine/Source/Runtime/VulkanRHI/Private/VulkanRenderTarget.cpp
+++ b/Engine/Source/Runtime/VulkanRHI/Private/VulkanRenderTarget.cpp
@@ -814,7 +814,7 @@ void FVulkanDynamicRHI::RHIReadSurfaceData(FRHITexture* TextureRHI, FIntRect Rec
 	ensure(Texture2D->Surface.StorageFormat == VK_FORMAT_R8G8B8A8_UNORM || Texture2D->Surface.StorageFormat == VK_FORMAT_B8G8R8A8_UNORM || Texture2D->Surface.StorageFormat == VK_FORMAT_R16G16B16A16_SFLOAT || Texture2D->Surface.StorageFormat == VK_FORMAT_A2B10G10R10_UNORM_PACK32 || Texture2D->Surface.StorageFormat == VK_FORMAT_R16G16B16A16_UNORM);
 	const bool bIs8Bpp = (Texture2D->Surface.StorageFormat == VK_FORMAT_R16G16B16A16_SFLOAT);
 	const uint32 Size = NumPixels * sizeof(FColor) * (bIs8Bpp ? 2 : 1);
-	VulkanRHI::FStagingBuffer* StagingBuffer = Device->GetStagingManager().AcquireBuffer(Size, VK_BUFFER_USAGE_TRANSFER_DST_BIT, true);
+	VulkanRHI::FStagingBuffer* StagingBuffer = Device->GetStagingManager().AcquireBuffer(Size, VK_BUFFER_USAGE_TRANSFER_DST_BIT, VK_MEMORY_PROPERTY_HOST_CACHED_BIT);
 	if (GIgnoreCPUReads == 0)
 	{
 		VkBufferImageCopy CopyRegion;
@@ -965,7 +965,7 @@ void FVulkanDynamicRHI::RHIReadSurfaceFloatData(FRHITexture* TextureRHI, FIntRec
 
 		uint32 NumPixels = (Surface.Width >> InMipIndex) * (Surface.Height >> InMipIndex);
 		const uint32 Size = NumPixels * sizeof(FFloat16Color);
-		VulkanRHI::FStagingBuffer* StagingBuffer = InDevice->GetStagingManager().AcquireBuffer(Size, VK_BUFFER_USAGE_TRANSFER_DST_BIT, true);
+		VulkanRHI::FStagingBuffer* StagingBuffer = InDevice->GetStagingManager().AcquireBuffer(Size, VK_BUFFER_USAGE_TRANSFER_DST_BIT, VK_MEMORY_PROPERTY_HOST_CACHED_BIT);
 
 		if (GIgnoreCPUReads == 0)
 		{
@@ -1106,7 +1106,7 @@ void FVulkanDynamicRHI::RHIRead3DSurfaceFloatData(FRHITexture* TextureRHI,FIntRe
 
 	ensure(Surface.StorageFormat == VK_FORMAT_R16G16B16A16_SFLOAT);
 
-	VulkanRHI::FStagingBuffer* StagingBuffer = Device->GetStagingManager().AcquireBuffer(Size, VK_BUFFER_USAGE_TRANSFER_DST_BIT, true);
+	VulkanRHI::FStagingBuffer* StagingBuffer = Device->GetStagingManager().AcquireBuffer(Size, VK_BUFFER_USAGE_TRANSFER_DST_BIT, VK_MEMORY_PROPERTY_HOST_CACHED_BIT);
 	if (GIgnoreCPUReads == 0)
 	{
 		VkBufferImageCopy CopyRegion;
@@ -1723,8 +1723,8 @@ struct FRenderPassCompatibleHashableStruct
 	uint8						bIsMultiview;
 	uint8						NumSamples;
 	uint8						SubpassHint;
-	// +1 for DepthStencil
-	VkFormat					Formats[MaxSimultaneousRenderTargets + 1];
+	// +1 for DepthStencil, +1 for Fragment Density
+	VkFormat					Formats[MaxSimultaneousRenderTargets + 2];
 };
 
 // Need a separate struct so we can memzero/remove dependencies on reference counts
@@ -1735,14 +1735,14 @@ struct FRenderPassFullHashableStruct
 		FMemory::Memzero(*this);
 	}
 
-	// +1 for Depth, +1 for Stencil
-	TEnumAsByte<VkAttachmentLoadOp>		LoadOps[MaxSimultaneousRenderTargets + 2];
-	TEnumAsByte<VkAttachmentStoreOp>	StoreOps[MaxSimultaneousRenderTargets + 2];
+	// +1 for Depth, +1 for Stencil, +1 for Fragment Density
+	TEnumAsByte<VkAttachmentLoadOp>		LoadOps[MaxSimultaneousRenderTargets + 3];
+	TEnumAsByte<VkAttachmentStoreOp>	StoreOps[MaxSimultaneousRenderTargets + 3];
 #if VULKAN_USE_REAL_RENDERPASS_COMPATIBILITY
 	// If the initial != final we need to add FinalLayout and potentially RefLayout
-	VkImageLayout						InitialLayout[MaxSimultaneousRenderTargets + 1];
-	//VkImageLayout						FinalLayout[MaxSimultaneousRenderTargets + 1];
-	//VkImageLayout						RefLayout[MaxSimultaneousRenderTargets + 1];
+	VkImageLayout						InitialLayout[MaxSimultaneousRenderTargets + 2];
+	//VkImageLayout						FinalLayout[MaxSimultaneousRenderTargets + 2];
+	//VkImageLayout						RefLayout[MaxSimultaneousRenderTargets + 2];
 #endif
 };
 
@@ -1752,12 +1752,14 @@ FVulkanRenderTargetLayout::FVulkanRenderTargetLayout(FVulkanDevice& InDevice, co
 	, NumColorAttachments(0)
 	, bHasDepthStencil(false)
 	, bHasResolveAttachments(false)
+	, bHasFragmentDensityAttachment(false)
 	, NumSamples(0)
 	, NumUsedClearValues(0)
 	, bIsMultiView(0)
 {
 	FMemory::Memzero(ColorReferences);
 	FMemory::Memzero(DepthStencilReference);
+	FMemory::Memzero(FragmentDensityReference);
 	FMemory::Memzero(ResolveReferences);
 	FMemory::Memzero(InputAttachments);
 	FMemory::Memzero(Desc);
@@ -1924,6 +1926,40 @@ FVulkanRenderTargetLayout::FVulkanRenderTargetLayout(FVulkanDevice& InDevice, co
 		}
 	}
 
+	if (InDevice.GetOptionalExtensions().HasEXTFragmentDensityMap && RTInfo.FoveationTexture)
+	{
+		FVulkanTextureBase* Texture = FVulkanTextureBase::Cast(RTInfo.FoveationTexture);
+		check(Texture);
+
+		VkAttachmentDescription& CurrDesc = Desc[NumAttachmentDescriptions];
+		FMemory::Memzero(CurrDesc);
+
+		const VkImageLayout FragmentDensityLayout = VK_IMAGE_LAYOUT_FRAGMENT_DENSITY_MAP_OPTIMAL_EXT;
+
+		CurrDesc.flags = 0;
+		CurrDesc.format = UEToVkTextureFormat(RTInfo.FoveationTexture->GetFormat(), false);
+		CurrDesc.samples = static_cast<VkSampleCountFlagBits>(RTInfo.FoveationTexture->GetNumSamples());
+		CurrDesc.loadOp = VK_ATTACHMENT_LOAD_OP_DONT_CARE;
+		CurrDesc.storeOp = VK_ATTACHMENT_STORE_OP_DONT_CARE;
+		CurrDesc.stencilLoadOp = VK_ATTACHMENT_LOAD_OP_DONT_CARE;
+		CurrDesc.stencilStoreOp = VK_ATTACHMENT_STORE_OP_DONT_CARE;
+		CurrDesc.initialLayout = FragmentDensityLayout;
+		CurrDesc.finalLayout = FragmentDensityLayout;
+
+		FragmentDensityReference.attachment = NumAttachmentDescriptions;
+		FragmentDensityReference.layout = FragmentDensityLayout;
+
+		FullHashInfo.LoadOps[MaxSimultaneousRenderTargets + 2] = CurrDesc.stencilLoadOp;
+		FullHashInfo.StoreOps[MaxSimultaneousRenderTargets + 2] = CurrDesc.stencilStoreOp;
+#if VULKAN_USE_REAL_RENDERPASS_COMPATIBILITY
+		FullHashInfo.InitialLayout[MaxSimultaneousRenderTargets + 1] = FragmentDensityLayout;
+#endif
+		CompatibleHashInfo.Formats[MaxSimultaneousRenderTargets + 1] = CurrDesc.format;
+
+		++NumAttachmentDescriptions;
+		bHasFragmentDensityAttachment = true;
+	}
+
 	SubpassHint = ESubpassHint::None;
 	CompatibleHashInfo.SubpassHint = 0;
 
@@ -1942,12 +1978,14 @@ FVulkanRenderTargetLayout::FVulkanRenderTargetLayout(FVulkanDevice& InDevice, co
 	, NumColorAttachments(0)
 	, bHasDepthStencil(false)
 	, bHasResolveAttachments(false)
+	, bHasFragmentDensityAttachment(false)
 	, NumSamples(0)
 	, NumUsedClearValues(0)
 	, bIsMultiView(RPInfo.bMultiviewPass)
 {
 	FMemory::Memzero(ColorReferences);
 	FMemory::Memzero(DepthStencilReference);
+	FMemory::Memzero(FragmentDensityReference);
 	FMemory::Memzero(ResolveReferences);
 	FMemory::Memzero(InputAttachments);
 	FMemory::Memzero(Desc);
@@ -2111,6 +2149,40 @@ FVulkanRenderTargetLayout::FVulkanRenderTargetLayout(FVulkanDevice& InDevice, co
 		}
 	}
 
+	if (InDevice.GetOptionalExtensions().HasEXTFragmentDensityMap && RPInfo.FoveationTexture)
+	{
+		FVulkanTextureBase* Texture = FVulkanTextureBase::Cast(RPInfo.FoveationTexture);
+		check(Texture);
+
+		VkAttachmentDescription& CurrDesc = Desc[NumAttachmentDescriptions];
+		FMemory::Memzero(CurrDesc);
+
+		const VkImageLayout FragmentDensityLayout = VK_IMAGE_LAYOUT_FRAGMENT_DENSITY_MAP_OPTIMAL_EXT;
+
+		CurrDesc.flags = 0;
+		CurrDesc.format = UEToVkTextureFormat(RPInfo.FoveationTexture->GetFormat(), false);
+		CurrDesc.samples = static_cast<VkSampleCountFlagBits>(RPInfo.FoveationTexture->GetNumSamples());
+		CurrDesc.loadOp = VK_ATTACHMENT_LOAD_OP_DONT_CARE;
+		CurrDesc.storeOp = VK_ATTACHMENT_STORE_OP_DONT_CARE;
+		CurrDesc.stencilLoadOp = VK_ATTACHMENT_LOAD_OP_DONT_CARE;
+		CurrDesc.stencilStoreOp = VK_ATTACHMENT_STORE_OP_DONT_CARE;
+		CurrDesc.initialLayout = FragmentDensityLayout;
+		CurrDesc.finalLayout = FragmentDensityLayout;
+
+		FragmentDensityReference.attachment = NumAttachmentDescriptions;
+		FragmentDensityReference.layout = FragmentDensityLayout;
+
+		FullHashInfo.LoadOps[MaxSimultaneousRenderTargets + 2] = CurrDesc.stencilLoadOp;
+		FullHashInfo.StoreOps[MaxSimultaneousRenderTargets + 2] = CurrDesc.stencilStoreOp;
+#if VULKAN_USE_REAL_RENDERPASS_COMPATIBILITY
+		FullHashInfo.InitialLayout[MaxSimultaneousRenderTargets + 1] = FragmentDensityLayout;
+#endif
+		CompatibleHashInfo.Formats[MaxSimultaneousRenderTargets + 1] = CurrDesc.format;
+
+		++NumAttachmentDescriptions;
+		bHasFragmentDensityAttachment = true;
+	}
+
 	SubpassHint = RPInfo.SubpassHint;
 	CompatibleHashInfo.SubpassHint = (uint8)RPInfo.SubpassHint;
 
@@ -2133,12 +2205,14 @@ FVulkanRenderTargetLayout::FVulkanRenderTargetLayout(const FGraphicsPipelineStat
 	, NumColorAttachments(0)
 	, bHasDepthStencil(false)
 	, bHasResolveAttachments(false)
+	, bHasFragmentDensityAttachment(false)
 	, NumSamples(0)
 	, NumUsedClearValues(0)
 	, bIsMultiView(0)
 {
 	FMemory::Memzero(ColorReferences);
 	FMemory::Memzero(DepthStencilReference);
+	FMemory::Memzero(FragmentDensityReference);
 	FMemory::Memzero(ResolveReferences);
 	FMemory::Memzero(InputAttachments);
 	FMemory::Memzero(Desc);
diff --git a/Engine/Source/Runtime/VulkanRHI/Private/VulkanTexture.cpp b/Engine/Source/Runtime/VulkanRHI/Private/VulkanTexture.cpp
--- a/Engine/Source/Runtime/VulkanRHI/Private/VulkanTexture.cpp
+++ b/Engine/Source/Runtime/VulkanRHI/Private/VulkanTexture.cpp
@@ -616,11 +616,11 @@ FVulkanSurface::FVulkanSurface(FVulkanDevice& InDevice, VkImageViewType Resource
 	, FullAspectMask(0)
 	, PartialAspectMask(0)
 {
-	StorageFormat = (VkFormat)GPixelFormats[PixelFormat].PlatformFormat;
-	check((UEFlags & TexCreate_SRGB) == 0);
+	StorageFormat = UEToVkTextureFormat(PixelFormat, false);
+
 	checkf(PixelFormat == PF_Unknown || StorageFormat != VK_FORMAT_UNDEFINED, TEXT("PixelFormat %d, is not supported for images"), (int32)PixelFormat);
 
-	ViewFormat = StorageFormat;
+	ViewFormat = UEToVkTextureFormat(PixelFormat, (UEFlags & TexCreate_SRGB) == TexCreate_SRGB);
 	FullAspectMask = VulkanRHI::GetAspectMaskFromUEFormat(PixelFormat, true, true);
 	PartialAspectMask = VulkanRHI::GetAspectMaskFromUEFormat(PixelFormat, false, true);
 
diff --git a/Engine/Source/Runtime/VulkanRHI/Private/VulkanUtil.cpp b/Engine/Source/Runtime/VulkanRHI/Private/VulkanUtil.cpp
--- a/Engine/Source/Runtime/VulkanRHI/Private/VulkanUtil.cpp
+++ b/Engine/Source/Runtime/VulkanRHI/Private/VulkanUtil.cpp
@@ -149,7 +149,7 @@ void FVulkanGPUTiming::Initialize()
 	{
 		check(!Pool);
 		Pool = new FVulkanTimingQueryPool(Device, 8);
-		Pool->ResultsBuffer = Device->GetStagingManager().AcquireBuffer(8 * sizeof(uint64), VK_BUFFER_USAGE_TRANSFER_DST_BIT, true);
+		Pool->ResultsBuffer = Device->GetStagingManager().AcquireBuffer(8 * sizeof(uint64), VK_BUFFER_USAGE_TRANSFER_DST_BIT, VK_MEMORY_PROPERTY_HOST_CACHED_BIT);
 	}
 }
 
diff --git a/Engine/Source/Runtime/VulkanRHI/Private/VulkanViewport.cpp b/Engine/Source/Runtime/VulkanRHI/Private/VulkanViewport.cpp
--- a/Engine/Source/Runtime/VulkanRHI/Private/VulkanViewport.cpp
+++ b/Engine/Source/Runtime/VulkanRHI/Private/VulkanViewport.cpp
@@ -322,6 +322,7 @@ FVulkanFramebuffer::FVulkanFramebuffer(FVulkanDevice& Device, const FRHISetRende
 	, NumColorRenderTargets(InRTInfo.NumColorRenderTargets)
 	, NumColorAttachments(0)
 	, DepthStencilRenderTargetImage(VK_NULL_HANDLE)
+	, FragmentDensityImage(VK_NULL_HANDLE)
 {
 	FMemory::Memzero(ColorRenderTargetImages);
 	FMemory::Memzero(ColorResolveTargetImages);
@@ -418,6 +419,20 @@ FVulkanFramebuffer::FVulkanFramebuffer(FVulkanDevice& Device, const FRHISetRende
 		}
 	}
 
+	if (RTLayout.GetHasFragmentDensityAttachment() && Device.GetOptionalExtensions().HasEXTFragmentDensityMap)
+	{
+		FVulkanTextureBase* Texture = FVulkanTextureBase::Cast(InRTInfo.FoveationTexture);
+		FragmentDensityImage = Texture->Surface.Image;
+
+		ensure(Texture->Surface.GetViewType() == VK_IMAGE_VIEW_TYPE_2D || Texture->Surface.GetViewType() == VK_IMAGE_VIEW_TYPE_2D_ARRAY);
+
+		FVulkanTextureView RTView;
+		RTView.Create(*Texture->Surface.Device, Texture->Surface.Image, Texture->Surface.GetViewType(), Texture->Surface.GetFullAspectMask(), Texture->Surface.PixelFormat, Texture->Surface.ViewFormat, MipIndex, 1, 0, Texture->Surface.GetNumberOfArrayLevels(), true);
+
+		AttachmentTextureViews.Add(RTView);
+		AttachmentViewsToDelete.Add(RTView.View);
+	}
+
 	TArray<VkImageView> AttachmentViews;
 	AttachmentViews.Empty(AttachmentTextureViews.Num());
 	for (auto& TextureView : AttachmentTextureViews)
@@ -483,6 +498,19 @@ bool FVulkanFramebuffer::Matches(const FRHISetRenderTargetsInfo& InRTInfo) const
 		}
 	}
 
+	{
+		FRHITexture* Texture = InRTInfo.FoveationTexture;
+		if (Texture)
+		{
+			VkImage AImage = FragmentDensityImage;
+			VkImage BImage = ((FVulkanTextureBase*)Texture->GetTextureBaseRHI())->Surface.Image;
+			if (AImage != BImage)
+			{
+				return false;
+			}
+		}
+	}
+
 	int32 AttachementIndex = 0;
 	for (int32 Index = 0; Index < InRTInfo.NumColorRenderTargets; ++Index)
 	{
diff --git a/Engine/Source/Runtime/VulkanRHI/Public/VulkanMemory.h b/Engine/Source/Runtime/VulkanRHI/Public/VulkanMemory.h
--- a/Engine/Source/Runtime/VulkanRHI/Public/VulkanMemory.h
+++ b/Engine/Source/Runtime/VulkanRHI/Public/VulkanMemory.h
@@ -978,7 +978,7 @@ namespace VulkanRHI
 		FStagingBuffer()
 			: ResourceAllocation(nullptr)
 			, Buffer(VK_NULL_HANDLE)
-			, bCPURead(false)
+			, MemoryReadFlags(VK_MEMORY_PROPERTY_HOST_COHERENT_BIT)
 			, BufferSize(0)
 		{
 		}
@@ -1020,7 +1020,7 @@ namespace VulkanRHI
 	protected:
 		TRefCountPtr<FOldResourceAllocation> ResourceAllocation;
 		VkBuffer Buffer;
-		bool bCPURead;
+		VkMemoryPropertyFlagBits MemoryReadFlags;
 		uint32 BufferSize;
 
 		// Owner maintains lifetime
@@ -1049,7 +1049,7 @@ namespace VulkanRHI
 
 		void Deinit();
 
-		FStagingBuffer* AcquireBuffer(uint32 Size, VkBufferUsageFlags InUsageFlags = VK_BUFFER_USAGE_TRANSFER_SRC_BIT, bool bCPURead = false);
+		FStagingBuffer* AcquireBuffer(uint32 Size, VkBufferUsageFlags InUsageFlags = VK_BUFFER_USAGE_TRANSFER_SRC_BIT, VkMemoryPropertyFlagBits InMemoryReadFlags = VK_MEMORY_PROPERTY_HOST_COHERENT_BIT);
 
 		// Sets pointer to nullptr
 		void ReleaseBuffer(FVulkanCmdBuffer* CmdBuffer, FStagingBuffer*& StagingBuffer);
