Wrap for validation cvar usage

diff --git a/Engine/Source/Runtime/VulkanRHI/Private/Lumin/VulkanLuminPlatform.cpp b/Engine/Source/Runtime/VulkanRHI/Private/Lumin/VulkanLuminPlatform.cpp
--- a/Engine/Source/Runtime/VulkanRHI/Private/Lumin/VulkanLuminPlatform.cpp
+++ b/Engine/Source/Runtime/VulkanRHI/Private/Lumin/VulkanLuminPlatform.cpp
@@ -70,8 +70,10 @@ bool FVulkanLuminPlatform::LoadVulkanInstanceFunctions(VkInstance inInstance)
 	ENUM_VK_ENTRYPOINTS_OPTIONAL_PLATFORM_INSTANCE(CHECK_VK_ENTRYPOINTS);
 
 
+#if VULKAN_HAS_DEBUGGING_ENABLED
 	//#todo-rco: Media textures are not working properly, quick workaround
 	GValidationCvar->Set(2, ECVF_SetByCommandline);
+#endif
 
 	return bFoundAllEntryPoints;
 }
