Be a bit more verbose when failing on a validation error

diff --git a/Engine/Source/Runtime/VulkanRHI/Private/VulkanUtil.cpp b/Engine/Source/Runtime/VulkanRHI/Private/VulkanUtil.cpp
--- a/Engine/Source/Runtime/VulkanRHI/Private/VulkanUtil.cpp
+++ b/Engine/Source/Runtime/VulkanRHI/Private/VulkanUtil.cpp
@@ -740,6 +740,17 @@ namespace VulkanRHI
 			break;
 		}
 
+#if VULKAN_HAS_DEBUGGING_ENABLED
+		if (Result == VK_ERROR_VALIDATION_FAILED_EXT)
+		{
+			extern TAutoConsoleVariable<int32> GValidationCvar;
+			if (GValidationCvar.GetValueOnRenderThread() == 0)
+			{
+				UE_LOG(LogVulkanRHI, Fatal, TEXT("Failed with Validation error. Try running with r.Vulkan.EnableValidation=1 to get information from the driver"));
+			}
+		}
+#endif
+
 		UE_LOG(LogVulkanRHI, Error, TEXT("%s failed, VkResult=%d\n at %s:%u \n with error %s"),
 			ANSI_TO_TCHAR(VkFunction), (int32)Result, ANSI_TO_TCHAR(Filename), Line, *ErrorString);
 
