Fix compiler error with extern an already defiend var

diff --git a/Engine/Source/Runtime/VulkanRHI/Private/VulkanUtil.cpp b/Engine/Source/Runtime/VulkanRHI/Private/VulkanUtil.cpp
--- a/Engine/Source/Runtime/VulkanRHI/Private/VulkanUtil.cpp
+++ b/Engine/Source/Runtime/VulkanRHI/Private/VulkanUtil.cpp
@@ -743,7 +743,6 @@ namespace VulkanRHI
 #if VULKAN_HAS_DEBUGGING_ENABLED
 		if (Result == VK_ERROR_VALIDATION_FAILED_EXT)
 		{
-			extern TAutoConsoleVariable<int32> GValidationCvar;
 			if (GValidationCvar.GetValueOnRenderThread() == 0)
 			{
 				UE_LOG(LogVulkanRHI, Fatal, TEXT("Failed with Validation error. Try running with r.Vulkan.EnableValidation=1 to get information from the driver"));
