Integrate changes to Vulkan RHI from Oculus. Removes the need for CreateImage, turns it into GenerateCreateInfo and lets callers do whatever they want with the createinfo directly.

diff --git a/Engine/Source/Runtime/VulkanRHI/Private/VulkanResources.h b/Engine/Source/Runtime/VulkanRHI/Private/VulkanResources.h
--- a/Engine/Source/Runtime/VulkanRHI/Private/VulkanResources.h
+++ b/Engine/Source/Runtime/VulkanRHI/Private/VulkanResources.h
@@ -264,9 +264,8 @@ class FVulkanSurface
 {
 public:
 
-	// Seperate method for creating image, this can be used to measure image size
-	// After VkImage is no longer needed, dont forget to destroy/release it 
-	static VkImage CreateImage(
+	// Seperate method for creating VkImageCreateInfo
+	static VkImageCreateInfo GenerateImageCreateInfo(
 		FVulkanDevice& InDevice,
 		VkImageViewType ResourceType,
 		EPixelFormat InFormat,
@@ -275,10 +274,8 @@ public:
 		uint32 NumMips,
 		uint32 NumSamples,
 		uint32 UEFlags,
-		VkMemoryRequirements& OutMemoryRequirements,
 		VkFormat* OutStorageFormat = nullptr,
 		VkFormat* OutViewFormat = nullptr,
-		VkImageCreateInfo* OutInfo = nullptr,
 		bool bForceLinearTexture = false);
 
 	FVulkanSurface(FVulkanDevice& Device, VkImageViewType ResourceType, EPixelFormat Format,
diff --git a/Engine/Source/Runtime/VulkanRHI/Private/VulkanTexture.cpp b/Engine/Source/Runtime/VulkanRHI/Private/VulkanTexture.cpp
--- a/Engine/Source/Runtime/VulkanRHI/Private/VulkanTexture.cpp
+++ b/Engine/Source/Runtime/VulkanRHI/Private/VulkanTexture.cpp
@@ -174,7 +174,7 @@ struct FRHICommandLockWriteTexture final : public FRHICommand<FRHICommandLockWri
 	}
 };
 
-VkImage FVulkanSurface::CreateImage(
+VkImageCreateInfo FVulkanSurface::GenerateImageCreateInfo(
 	FVulkanDevice& InDevice,
 	VkImageViewType ResourceType,
 	EPixelFormat InFormat,
@@ -183,17 +183,14 @@ VkImage FVulkanSurface::CreateImage(
 	uint32 NumMips,
 	uint32 NumSamples,
 	uint32 UEFlags,
-	VkMemoryRequirements& OutMemoryRequirements,
 	VkFormat* OutStorageFormat,
 	VkFormat* OutViewFormat,
-	VkImageCreateInfo* OutInfo,
 	bool bForceLinearTexture)
 {
 	const VkPhysicalDeviceProperties& DeviceProperties = InDevice.GetDeviceProperties();
 	const FPixelFormatInfo& FormatInfo = GPixelFormats[InFormat];
 	VkFormat TextureFormat = (VkFormat)FormatInfo.PlatformFormat;
 
-
 	if(UEFlags & TexCreate_CPUReadback)
 	{
 		bForceLinearTexture = true;
@@ -201,9 +198,7 @@ VkImage FVulkanSurface::CreateImage(
 
 	checkf(TextureFormat != VK_FORMAT_UNDEFINED, TEXT("PixelFormat %d, is not supported for images"), (int32)InFormat);
 
-	VkImageCreateInfo TmpCreateInfo;
-	VkImageCreateInfo* ImageCreateInfoPtr = OutInfo ? OutInfo : &TmpCreateInfo;
-	VkImageCreateInfo& ImageCreateInfo = *ImageCreateInfoPtr;
+	VkImageCreateInfo ImageCreateInfo;
 	ZeroVulkanStruct(ImageCreateInfo, VK_STRUCTURE_TYPE_IMAGE_CREATE_INFO);
 
 	switch(ResourceType)
@@ -234,22 +229,19 @@ VkImage FVulkanSurface::CreateImage(
 		break;
 	}
 
-	ImageCreateInfo.format = UEToVkTextureFormat(InFormat, false);
+	VkFormat srgbFormat = UEToVkTextureFormat(InFormat, (UEFlags & TexCreate_SRGB) == TexCreate_SRGB);
+	VkFormat nonSrgbFormat = UEToVkTextureFormat(InFormat, false);
+
+	ImageCreateInfo.format = ((UEFlags & TexCreate_UAV) == 0) ? srgbFormat : nonSrgbFormat; 
 
 	checkf(ImageCreateInfo.format != VK_FORMAT_UNDEFINED, TEXT("Pixel Format %d not defined!"), (int32)InFormat);
-	if (OutStorageFormat)
+	if (OutViewFormat)
 	{
-		*OutStorageFormat = ImageCreateInfo.format;
+		*OutViewFormat = srgbFormat;
 	}
-
-	if (OutViewFormat)
+	if (OutStorageFormat)
 	{
-		VkFormat ViewFormat = UEToVkTextureFormat(InFormat, (UEFlags & TexCreate_SRGB) == TexCreate_SRGB);
-		*OutViewFormat = ViewFormat;
-		if((UEFlags & TexCreate_UAV) == 0)
-		{
-			ImageCreateInfo.format = ViewFormat;
-		}
+		*OutStorageFormat = nonSrgbFormat;
 	}
 
 	ImageCreateInfo.extent.width = SizeX;
@@ -308,7 +300,7 @@ VkImage FVulkanSurface::CreateImage(
 		ImageCreateInfo.usage |= VK_IMAGE_USAGE_COLOR_ATTACHMENT_BIT;
 		ImageCreateInfo.tiling = VK_IMAGE_TILING_OPTIMAL;
 	}
-	
+
 	if (UEFlags & TexCreate_UAV)
 	{
 		ImageCreateInfo.usage |= VK_IMAGE_USAGE_STORAGE_BIT;
@@ -353,7 +345,7 @@ VkImage FVulkanSurface::CreateImage(
 		break;
 	}
 
-//#todo-rco: Verify flags work on newer Android drivers
+	//#todo-rco: Verify flags work on newer Android drivers
 #if !PLATFORM_ANDROID
 	if (ImageCreateInfo.tiling == VK_IMAGE_TILING_LINEAR)
 	{
@@ -410,17 +402,9 @@ VkImage FVulkanSurface::CreateImage(
 		}
 	}
 #endif
-
-	VkImage Image;
-	VERIFYVULKANRESULT(VulkanRHI::vkCreateImage(InDevice.GetInstanceHandle(), &ImageCreateInfo, VULKAN_CPU_ALLOCATOR, &Image));
-
-	// Fetch image size
-	VulkanRHI::vkGetImageMemoryRequirements(InDevice.GetInstanceHandle(), Image, &OutMemoryRequirements);
-
-	return Image;
+	return ImageCreateInfo;
 }
 
-
 struct FRHICommandInitialClearTexture final : public FRHICommand<FRHICommandInitialClearTexture>
 {
 	FVulkanSurface* Surface;
@@ -519,12 +503,16 @@ FVulkanSurface::FVulkanSurface(FVulkanDevice& InDevice, VkImageViewType Resource
 	, FullAspectMask(0)
 	, PartialAspectMask(0)
 {
-	VkImageCreateInfo ImageCreateInfo;	// Zeroed inside CreateImage
-	Image = FVulkanSurface::CreateImage(InDevice, ResourceType,
+	VkImageCreateInfo ImageCreateInfo = FVulkanSurface::GenerateImageCreateInfo(
+		InDevice, ResourceType,
 		InFormat, SizeX, SizeY, SizeZ,
-		ArraySize, NumMips, NumSamples, UEFlags, MemoryRequirements,
-		&StorageFormat, &ViewFormat,
-		&ImageCreateInfo);
+		ArraySize, NumMips, NumSamples, UEFlags,
+		&StorageFormat, &ViewFormat);
+
+	VERIFYVULKANRESULT(VulkanRHI::vkCreateImage(InDevice.GetInstanceHandle(), &ImageCreateInfo, VULKAN_CPU_ALLOCATOR, &Image));
+
+	// Fetch image size
+	VulkanRHI::vkGetImageMemoryRequirements(InDevice.GetInstanceHandle(), Image, &MemoryRequirements);
 
 	uint32 LayerCount = (ResourceType == VK_IMAGE_VIEW_TYPE_CUBE || ResourceType == VK_IMAGE_VIEW_TYPE_CUBE_ARRAY) ? 6 : 1;
 	NumArrayLevels = ArraySize * LayerCount;
@@ -635,6 +623,15 @@ FVulkanSurface::FVulkanSurface(FVulkanDevice& InDevice, VkImageViewType Resource
 	
 	if (Image != VK_NULL_HANDLE)
 	{
+#if VULKAN_ENABLE_WRAP_LAYER	
+		VkImageCreateInfo ImageCreateInfo = FVulkanSurface::GenerateImageCreateInfo(
+			InDevice, ResourceType,
+			InFormat, SizeX, SizeY, SizeZ,
+			ArraySize, NumMips, NumSamples, UEFlags,
+			&StorageFormat, &ViewFormat);
+		FWrapLayer::CreateImage(VK_SUCCESS, InDevice.GetInstanceHandle(), &ImageCreateInfo, &Image);
+#endif
+
 		if (UEFlags & (TexCreate_RenderTargetable | TexCreate_DepthStencilTargetable))
 		{
 			const bool bTransitionToPresentable = ((UEFlags & TexCreate_Presentable) == TexCreate_Presentable);
@@ -2144,16 +2141,17 @@ static VkMemoryRequirements FindOrCalculateTexturePlatformSize(FVulkanDevice* De
 		}
 	}
 
-	VkFormat InternalStorageFormat, InternalViewFormat;
-	VkImageCreateInfo CreateInfo;
-	VkMemoryRequirements MemReq;
 	EPixelFormat PixelFormat = (EPixelFormat)Format;
+	VkMemoryRequirements MemReq;
 
 	// Create temporary image to measure the memory requirements
-	VkImage TmpImage = FVulkanSurface::CreateImage(*Device, ViewType,
+	VkImageCreateInfo TmpCreateInfo = FVulkanSurface::GenerateImageCreateInfo(*Device, ViewType,
 		PixelFormat, SizeX, SizeY, SizeZ, 1, NumMips, NumSamples,
-		Flags, MemReq, &InternalStorageFormat, &InternalViewFormat, &CreateInfo, false);
+		Flags, nullptr, nullptr, false);
 
+	VkImage TmpImage;
+	VERIFYVULKANRESULT(VulkanRHI::vkCreateImage(Device->GetInstanceHandle(), &TmpCreateInfo, VULKAN_CPU_ALLOCATOR, &TmpImage));
+	VulkanRHI::vkGetImageMemoryRequirements(Device->GetInstanceHandle(), TmpImage, &MemReq);
 	VulkanRHI::vkDestroyImage(Device->GetInstanceHandle(), TmpImage, VULKAN_CPU_ALLOCATOR);
 
 	{
