Fix CIS warnings (UE-82156)

diff --git a/Engine/Source/Runtime/Core/Core.Build.cs b/Engine/Source/Runtime/Core/Core.Build.cs
--- a/Engine/Source/Runtime/Core/Core.Build.cs
+++ b/Engine/Source/Runtime/Core/Core.Build.cs
@@ -129,10 +129,6 @@ public class Core : ModuleRules
 		}
 		else if (Target.IsInPlatformGroup(UnrealPlatformGroup.Unix))
         {
-			string PlatformName = Target.IsInPlatformGroup(UnrealPlatformGroup.Linux) ?
-				"Linux" : Target.Platform.ToString();
-
-			PublicIncludePaths.Add(string.Format("Runtime/Core/Public/{0}", PlatformName));
 			AddEngineThirdPartyPrivateStaticDependencies(Target,
 				"zlib",
 				"jemalloc"
diff --git a/Engine/Source/Runtime/Linux/AudioMixerSDL/AudioMixerSDL.Build.cs b/Engine/Source/Runtime/Linux/AudioMixerSDL/AudioMixerSDL.Build.cs
--- a/Engine/Source/Runtime/Linux/AudioMixerSDL/AudioMixerSDL.Build.cs
+++ b/Engine/Source/Runtime/Linux/AudioMixerSDL/AudioMixerSDL.Build.cs
@@ -18,7 +18,10 @@ public class AudioMixerSDL : ModuleRules
 			PlatformName = "Linux";
 		}
 
-		PrivateIncludePaths.Add("Runtime/Linux/AudioMixerSDL/Private/" + PlatformName);
+		if (Target.IsInPlatformGroup(UnrealPlatformGroup.Linux))
+		{
+			PrivateIncludePaths.Add("Runtime/Linux/AudioMixerSDL/Private/" + PlatformName);
+		}
 
 		PrivateDependencyModuleNames.AddRange(
 			new string[] {
diff --git a/Engine/Source/Runtime/VulkanRHI/VulkanRHI.Build.cs b/Engine/Source/Runtime/VulkanRHI/VulkanRHI.Build.cs
--- a/Engine/Source/Runtime/VulkanRHI/VulkanRHI.Build.cs
+++ b/Engine/Source/Runtime/VulkanRHI/VulkanRHI.Build.cs
@@ -16,9 +16,12 @@ public class VulkanRHI : ModuleRules
 			PrivateIncludePaths.Add("Runtime/VulkanRHI/Private/Windows");
 			AddEngineThirdPartyPrivateStaticDependencies(Target, "AMD_AGS");
 		}
-		else if (Target.IsInPlatformGroup(UnrealPlatformGroup.Linux))
+		else if (Target.IsInPlatformGroup(UnrealPlatformGroup.Unix))
 		{
-			PrivateIncludePaths.Add("Runtime/VulkanRHI/Private/Linux");
+			if (Target.IsInPlatformGroup(UnrealPlatformGroup.Linux))
+			{
+				PrivateIncludePaths.Add("Runtime/VulkanRHI/Private/Linux");
+			}
 		}
 		else
 		{
