Disabled VK_KHR_dedicated_allocation by default due to video memory leaking

diff --git a/Engine/Source/Runtime/VulkanRHI/Public/VulkanConfiguration.h b/Engine/Source/Runtime/VulkanRHI/Public/VulkanConfiguration.h
--- a/Engine/Source/Runtime/VulkanRHI/Public/VulkanConfiguration.h
+++ b/Engine/Source/Runtime/VulkanRHI/Public/VulkanConfiguration.h
@@ -123,7 +123,8 @@
 
 #ifndef VULKAN_SUPPORTS_DEDICATED_ALLOCATION
 	#ifdef VK_KHR_dedicated_allocation
-		#define VULKAN_SUPPORTS_DEDICATED_ALLOCATION			1
+		// Disable this for now as it is causing a large memory leak
+		#define VULKAN_SUPPORTS_DEDICATED_ALLOCATION			0
 	#else
 		#define VULKAN_SUPPORTS_DEDICATED_ALLOCATION			0
 	#endif
