Fix spacing on Linux for Vulkan Memory Dump

diff --git a/Engine/Source/Runtime/VulkanRHI/Private/VulkanMemory.cpp b/Engine/Source/Runtime/VulkanRHI/Private/VulkanMemory.cpp
--- a/Engine/Source/Runtime/VulkanRHI/Private/VulkanMemory.cpp
+++ b/Engine/Source/Runtime/VulkanRHI/Private/VulkanMemory.cpp
@@ -1618,11 +1618,11 @@ namespace VulkanRHI
 				uint64 _UsedLargeTotal = 0;
 				uint64 _AllocLargeTotal = 0;
 
-				UE_LOG(LogVulkanRHI, Display, TEXT("Index  BufferHandle   DeviceMemoryHandle MemFlags BufferFlags #Suballocs #FreeChunks UsedSize/MaxSize"));
+				UE_LOG(LogVulkanRHI, Display, TEXT("Index  BufferHandle       DeviceMemoryHandle MemFlags BufferFlags #Suballocs #FreeChunks UsedSize/MaxSize"));
 				for (int32 Index = 0; Index < UsedAllocations.Num(); ++Index)
 				{
 					FBufferAllocation* BA = UsedAllocations[Index];
-					UE_LOG(LogVulkanRHI, Display, TEXT("%6d %p %p 0x%06x 0x%08x %6d   %6d    %d/%d"), Index, (void*)BA->Buffer, (void*)BA->MemoryAllocation->GetHandle(), BA->MemoryPropertyFlags, BA->BufferUsageFlags, BA->Suballocations.Num(), BA->FreeList.Num(), BA->UsedSize, BA->MaxSize);
+					UE_LOG(LogVulkanRHI, Display, TEXT("%6d 0x%016llx 0x%016llx 0x%06x 0x%08x %6d   %6d        %d/%d"), Index, (void*)BA->Buffer, (void*)BA->MemoryAllocation->GetHandle(), BA->MemoryPropertyFlags, BA->BufferUsageFlags, BA->Suballocations.Num(), BA->FreeList.Num(), BA->UsedSize, BA->MaxSize);
 
 					if (PoolSizeIndex == (int32)EPoolSizes::SizesCount)
 					{
@@ -1927,11 +1927,11 @@ namespace VulkanRHI
 	void FStagingManager::DumpMemory()
 	{
 		UE_LOG(LogVulkanRHI, Display, TEXT("StagingManager %d Used %d Pending Free %d Free"), UsedStagingBuffers.Num(), PendingFreeStagingBuffers.Num(), FreeStagingBuffers.Num());
-		UE_LOG(LogVulkanRHI, Display, TEXT("Used   BufferHandle      ResourceAllocation Size"));
+		UE_LOG(LogVulkanRHI, Display, TEXT("Used   BufferHandle       ResourceAllocation Size"));
 		for (int32 Index = 0; Index < UsedStagingBuffers.Num(); ++Index)
 		{
 			FStagingBuffer* Buffer = UsedStagingBuffers[Index];
-			UE_LOG(LogVulkanRHI, Display, TEXT("%6d %p %p %6d"), Index, (void*)Buffer->GetHandle(), (void*)Buffer->ResourceAllocation->GetHandle(), Buffer->BufferSize);
+			UE_LOG(LogVulkanRHI, Display, TEXT("%6d 0x%016llx 0x%016llx %6d"), Index, (void*)Buffer->GetHandle(), (void*)Buffer->ResourceAllocation->GetHandle(), Buffer->BufferSize);
 		}
 
 		UE_LOG(LogVulkanRHI, Display, TEXT("Pending CmdBuffer   Fence   BufferHandle    ResourceAllocation Size"));
@@ -1946,16 +1946,16 @@ namespace VulkanRHI
 				for (int32 BufferIndex = 0; BufferIndex < ItemsPerFence.Resources.Num(); ++BufferIndex)
 				{
 					FStagingBuffer* Buffer = ItemsPerFence.Resources[BufferIndex];
-					UE_LOG(LogVulkanRHI, Display, TEXT("                   %p %p %6d"), (void*)Buffer->GetHandle(), (void*)Buffer->ResourceAllocation->GetHandle(), Buffer->BufferSize);
+					UE_LOG(LogVulkanRHI, Display, TEXT("                   0x%016llx 0x%016llx %6d"), (void*)Buffer->GetHandle(), (void*)Buffer->ResourceAllocation->GetHandle(), Buffer->BufferSize);
 				}
 			}
 		}
 
-		UE_LOG(LogVulkanRHI, Display, TEXT("Free   BufferHandle ResourceAllocation Size"));
+		UE_LOG(LogVulkanRHI, Display, TEXT("Free   BufferHandle     ResourceAllocation Size"));
 		for (int32 Index = 0; Index < FreeStagingBuffers.Num(); ++Index)
 		{
 			FFreeEntry& Entry = FreeStagingBuffers[Index];
-			UE_LOG(LogVulkanRHI, Display, TEXT("%6d %p %p %6d"), Index, (void*)Entry.StagingBuffer->GetHandle(), (void*)Entry.StagingBuffer->ResourceAllocation->GetHandle(), Entry.StagingBuffer->BufferSize);
+			UE_LOG(LogVulkanRHI, Display, TEXT("%6d 0x%016llx 0x%016llx %6d"), Index, (void*)Entry.StagingBuffer->GetHandle(), (void*)Entry.StagingBuffer->ResourceAllocation->GetHandle(), Entry.StagingBuffer->BufferSize);
 		}
 	}
 #endif
