Vulkan Graphics PSO cache rewrite/cleanup.
TLDR: Multithreaded creation of PSOs, which should hitches and makes the "Optimizing shaders" faster.
-Only uses 1 map&lock
-Only has one type
-Use UE4 shaderkeys instead of bytecode shaderhash
-Enables parrallel pso creation
-added VulkanPSO stat group
-LRU is now compiled in - but disabled - on all vulkan platforms.

diff --git a/Engine/Source/Runtime/VulkanRHI/Private/VulkanDevice.h b/Engine/Source/Runtime/VulkanRHI/Private/VulkanDevice.h
--- a/Engine/Source/Runtime/VulkanRHI/Private/VulkanDevice.h
+++ b/Engine/Source/Runtime/VulkanRHI/Private/VulkanDevice.h
@@ -443,4 +443,5 @@ private:
 
 	class FVulkanPipelineStateCacheManager* PipelineStateCache;
 	friend class FVulkanDynamicRHI;
+	friend class FVulkanRHIGraphicsPipelineState;
 };
diff --git a/Engine/Source/Runtime/VulkanRHI/Private/VulkanPendingState.cpp b/Engine/Source/Runtime/VulkanRHI/Private/VulkanPendingState.cpp
--- a/Engine/Source/Runtime/VulkanRHI/Private/VulkanPendingState.cpp
+++ b/Engine/Source/Runtime/VulkanRHI/Private/VulkanPendingState.cpp
@@ -445,7 +445,7 @@ void FVulkanPendingGfxState::PrepareForDraw(FVulkanCmdBuffer* CmdBuffer)
 		SCOPE_CYCLE_COUNTER(STAT_VulkanBindVertexStreamsTime);
 #endif
 		// Its possible to have no vertex buffers
-		const FVulkanVertexInputStateInfo& VertexInputStateInfo = CurrentPipeline->Pipeline->GetVertexInputState();
+		const FVulkanVertexInputStateInfo& VertexInputStateInfo = CurrentPipeline->GetVertexInputState();
 		if (VertexInputStateInfo.AttributesNum == 0)
 		{
 			// However, we need to verify that there are also no bindings
@@ -923,3 +923,33 @@ float FVulkanDescriptorSetCache::FCachedPool::CalcAllocRatio() const
 	}
 	return AllocRatio;
 }
+bool FVulkanPendingGfxState::SetGfxPipeline(FVulkanRHIGraphicsPipelineState* InGfxPipeline, bool bForceReset)
+{
+	bool bChanged = bForceReset;
+
+	if (InGfxPipeline != CurrentPipeline)
+	{
+		CurrentPipeline = InGfxPipeline;
+		FVulkanGraphicsPipelineDescriptorState** Found = PipelineStates.Find(InGfxPipeline);
+		if (Found)
+		{
+			CurrentState = *Found;
+			check(CurrentState->GfxPipeline == InGfxPipeline);
+		}
+		else
+		{
+			CurrentState = new FVulkanGraphicsPipelineDescriptorState(Device, InGfxPipeline);
+			PipelineStates.Add(CurrentPipeline, CurrentState);
+		}
+
+		PrimitiveType = InGfxPipeline->PrimitiveType;
+		bChanged = true;
+	}
+
+	if (bChanged || bForceReset)
+	{
+		CurrentState->Reset();
+	}
+
+	return bChanged;
+}
diff --git a/Engine/Source/Runtime/VulkanRHI/Private/VulkanPendingState.h b/Engine/Source/Runtime/VulkanRHI/Private/VulkanPendingState.h
--- a/Engine/Source/Runtime/VulkanRHI/Private/VulkanPendingState.h
+++ b/Engine/Source/Runtime/VulkanRHI/Private/VulkanPendingState.h
@@ -341,36 +341,7 @@ public:
 
 	void PrepareForDraw(FVulkanCmdBuffer* CmdBuffer);
 
-	bool SetGfxPipeline(FVulkanRHIGraphicsPipelineState* InGfxPipeline, bool bForceReset)
-	{
-		bool bChanged = bForceReset;
-
-		if (InGfxPipeline != CurrentPipeline)
-		{
-			CurrentPipeline = InGfxPipeline;
-			FVulkanGraphicsPipelineDescriptorState** Found = PipelineStates.Find(InGfxPipeline);
-			if (Found)
-			{
-				CurrentState = *Found;
-				check(CurrentState->GfxPipeline == InGfxPipeline);
-			}
-			else
-			{
-				CurrentState = new FVulkanGraphicsPipelineDescriptorState(Device, InGfxPipeline);
-				PipelineStates.Add(CurrentPipeline, CurrentState);
-			}
-
-			PrimitiveType = InGfxPipeline->PrimitiveType;
-			bChanged = true;
-		}
-
-		if (bChanged || bForceReset)
-		{
-			CurrentState->Reset();
-		}
-
-		return bChanged;
-	}
+	bool SetGfxPipeline(FVulkanRHIGraphicsPipelineState* InGfxPipeline, bool bForceReset);
 
 	inline void UpdateDynamicStates(FVulkanCmdBuffer* Cmd)
 	{
diff --git a/Engine/Source/Runtime/VulkanRHI/Private/VulkanPipeline.cpp b/Engine/Source/Runtime/VulkanRHI/Private/VulkanPipeline.cpp
--- a/Engine/Source/Runtime/VulkanRHI/Private/VulkanPipeline.cpp
+++ b/Engine/Source/Runtime/VulkanRHI/Private/VulkanPipeline.cpp
@@ -18,6 +18,31 @@
 #include "VulkanLLM.h"
 #include "Misc/ScopeRWLock.h"
 
+#define LRU_DEBUG 0
+#if !UE_BUILD_SHIPPING
+#define LRUPRINT(...) FPlatformMisc::LowLevelOutputDebugStringf(__VA_ARGS__)
+#if LRU_DEBUG
+#define LRUPRINT_DEBUG(...) FPlatformMisc::LowLevelOutputDebugStringf(__VA_ARGS__)
+#endif
+#else
+#define LRUPRINT(...) do{}while(0)
+#endif
+
+#ifndef LRUPRINT_DEBUG
+#define LRUPRINT_DEBUG(...) do{}while(0)
+#endif
+
+
+
+#if PLATFORM_ANDROID
+#define LRU_MAX_PIPELINE_SIZE 10
+#define LRU_PIPELINE_CAPACITY 2048
+#else
+#define LRU_MAX_PIPELINE_SIZE 512 //needs to be super high to work on pc.
+#define LRU_PIPELINE_CAPACITY 8192
+#endif
+
+
 
 #if !UE_BUILD_SHIPPING
 static TAtomic<uint64> SGraphicsRHICount;
@@ -27,7 +52,16 @@ static TAtomic<uint64> SPipelineGfxCount;
 
 static const double HitchTime = 1.0 / 1000.0;
 
-static FCriticalSection EntryKeyToGfxPipelineMapCS;
+static FCriticalSection FVulkanShaderHandleCS;
+
+
+TAutoConsoleVariable<int32> CVarPipelineDebugForceEvictImmediately(
+	TEXT("r.Vulkan.PipelineDebugForceEvictImmediately"),
+	0,
+	TEXT("1: Force all created PSOs to be evicted immediately. Only for debugging"),
+	ECVF_RenderThreadSafe | ECVF_ReadOnly);
+
+
 
 TAutoConsoleVariable<int32> CVarPipelineLRUCacheEvictBinaryPreloadScreen(
 	TEXT("r.Vulkan.PipelineLRUCacheEvictBinaryPreloadScreen"),
@@ -35,7 +69,6 @@ TAutoConsoleVariable<int32> CVarPipelineLRUCacheEvictBinaryPreloadScreen(
 	TEXT("1: Use a preload screen while loading preevicted PSOs ala r.Vulkan.PipelineLRUCacheEvictBinary"),
 	ECVF_RenderThreadSafe);
 
-#if VULKAN_ENABLE_LRU_CACHE
 TAutoConsoleVariable<int32> CVarEnableLRU(
 	TEXT("r.Vulkan.EnablePipelineLRUCache"),
 	0,
@@ -52,23 +85,47 @@ TAutoConsoleVariable<int32> CVarPipelineLRUCacheEvictBinary(
 	ECVF_RenderThreadSafe | ECVF_ReadOnly);
 
 
-
 TAutoConsoleVariable<int32> CVarLRUMaxPipelineSize(
 	TEXT("r.Vulkan.PipelineLRUSize"),
-	10 * 1024 * 1024,
+	LRU_MAX_PIPELINE_SIZE * 1024 * 1024,
 	TEXT("Maximum size of shader memory ."),
 	ECVF_RenderThreadSafe);
 
-static bool IsUsePipelineLRU()
-{
-	static int32 bUse = -1;
-	if (bUse == -1)
-	{
-		bUse = CVarEnableLRU.GetValueOnAnyThread();
-	}
-	return bUse == 1;
-}
-#endif
+TAutoConsoleVariable<int32> CVarLRUPipelineCapacity(
+	TEXT("r.Vulkan.PipelineLRUCapactiy"),
+	LRU_PIPELINE_CAPACITY,
+	TEXT("Maximum no. of PSOs in LRU."),
+	ECVF_RenderThreadSafe| ECVF_ReadOnly);
+
+
+
+
+
+static TAutoConsoleVariable<int32> GEnablePipelineCacheLoadCvar(
+	TEXT("r.Vulkan.PipelineCacheLoad"),
+	1,
+	TEXT("0 to disable loading the pipeline cache")
+	TEXT("1 to enable using pipeline cache")
+);
+
+static TAutoConsoleVariable<int32> GPipelineCacheFromShaderPipelineCacheCvar(
+	TEXT("r.Vulkan.PipelineCacheFromShaderPipelineCache"),
+	PLATFORM_ANDROID && !(PLATFORM_LUMIN || PLATFORM_LUMINGL4),
+	TEXT("0 look for a pipeline cache in the normal locations with the normal names.")
+	TEXT("1 tie the vulkan pipeline cache to the shader pipeline cache, use the PSOFC guid as part of the filename, etc."),
+	ECVF_ReadOnly
+);
+
+
+static int32 GEnablePipelineCacheCompression = 1;
+static FAutoConsoleVariableRef GEnablePipelineCacheCompressionCvar(
+	TEXT("r.Vulkan.PipelineCacheCompression"),
+	GEnablePipelineCacheCompression,
+	TEXT("Enable/disable compression on the Vulkan pipeline cache disk file\n"),
+	ECVF_Default | ECVF_RenderThreadSafe
+);
+
+
 
 template <typename TRHIType, typename TVulkanType>
 static inline FSHAHash GetShaderHash(TRHIType* RHIShader)
@@ -119,7 +176,7 @@ FVulkanPipeline::~FVulkanPipeline()
 {
 #if !UE_BUILD_SHIPPING
 	SPipelineCount--;
-#endif
+#endif	
 	if (Pipeline != VK_NULL_HANDLE)
 	{
 		Device->GetDeferredDeletionQueue().EnqueueResource(VulkanRHI::FDeferredDeletionQueue::EType::Pipeline, Pipeline);
@@ -132,111 +189,32 @@ FVulkanComputePipeline::FVulkanComputePipeline(FVulkanDevice* InDevice)
 	: FVulkanPipeline(InDevice)
 	, ComputeShader(nullptr)
 {
+	INC_DWORD_STAT(STAT_VulkanNumComputePSOs);
 }
 
 FVulkanComputePipeline::~FVulkanComputePipeline()
 {
 	Device->NotifyDeletedComputePipeline(this);
+	DEC_DWORD_STAT(STAT_VulkanNumComputePSOs);
 }
 
-#if VULKAN_ENABLE_LRU_CACHE
-FVulkanGfxPipeline::FVulkanGfxPipeline(FVulkanDevice* InDevice, TGfxPipelineEntrySharedPtr InGfxPipelineEntry)
-	: FVulkanPipeline(InDevice)
-	, GfxPipelineEntry(MoveTemp(InGfxPipelineEntry))
-	, PipelineCacheSize(0)
-	, RecentFrame(0)
-	, bTrackedByLRU(false)
-	, bRuntimeObjectsValid(false)
-#else
-FVulkanGfxPipeline::FVulkanGfxPipeline(FVulkanDevice* InDevice)
-	: FVulkanPipeline(InDevice)
-	, bRuntimeObjectsValid(false)
-#endif
-{
-#if !UE_BUILD_SHIPPING
-	SPipelineGfxCount++;
-#endif
-}
-
-FVulkanGfxPipeline::~FVulkanGfxPipeline()
-{
-#if !UE_BUILD_SHIPPING
-	SPipelineGfxCount--;
-#endif
-}
-
-void FVulkanGfxPipeline::CreateRuntimeObjects(const FGraphicsPipelineStateInitializer& InPSOInitializer)
-{
-	const FBoundShaderStateInput& BSI = InPSOInitializer.BoundShaderState;
-	
-	check(BSI.VertexShaderRHI);
-	FVulkanVertexShader* VS = ResourceCast(BSI.VertexShaderRHI);
-	const FVulkanShaderHeader& VSHeader = VS->GetCodeHeader();
-
-	VertexInputState.Generate(ResourceCast(InPSOInitializer.BoundShaderState.VertexDeclarationRHI), VSHeader.InOutMask);
-	bRuntimeObjectsValid = true;
-}
-FVulkanRHIGraphicsPipelineState::FVulkanRHIGraphicsPipelineState(const FBoundShaderStateInput& InBSI, FVulkanGfxPipeline* InPipeline, EPrimitiveType InPrimitiveType)
-	: Pipeline(InPipeline)
-	, PrimitiveType(InPrimitiveType)
-{
-
-	for (int32 StageIdx = 0; StageIdx < ShaderStage::NumStages; ++StageIdx)
-	{
-		ShaderKeys[StageIdx] = GetShaderKeyForGfxStage(InBSI, (ShaderStage::EStage)StageIdx);
-	}
-
-	bHasInputAttachments = InPipeline->GetGfxLayout().GetDescriptorSetsLayout().HasInputAttachments();
-#if !UE_BUILD_SHIPPING
-	SGraphicsRHICount++;
-#endif
-}
 
 FVulkanRHIGraphicsPipelineState::~FVulkanRHIGraphicsPipelineState()
 {
 #if !UE_BUILD_SHIPPING
 	SGraphicsRHICount--;
-#endif
-	if (Pipeline)
+#endif	
+	DEC_DWORD_STAT(STAT_VulkanNumGraphicsPSOs);
+	Device->PipelineStateCache->NotifyDeletedGraphicsPSO(this);
+	if (bShaderModulesLoaded)
 	{
-		Pipeline->Device->NotifyDeletedGfxPipeline(this);
-		Pipeline = nullptr;
+		PurgeLoadedShaderModules(Device);
 	}
 }
 
-
-static TAutoConsoleVariable<int32> GEnablePipelineCacheLoadCvar(
-	TEXT("r.Vulkan.PipelineCacheLoad"),
-	1,
-	TEXT("0 to disable loading the pipeline cache")
-	TEXT("1 to enable using pipeline cache")
-	);
-
-static TAutoConsoleVariable<int32> GPipelineCacheFromShaderPipelineCacheCvar(
-	TEXT("r.Vulkan.PipelineCacheFromShaderPipelineCache"),
-	PLATFORM_ANDROID && !PLATFORM_LUMIN,
-	TEXT("0 look for a pipeline cache in the normal locations with the normal names.")
-	TEXT("1 tie the vulkan pipeline cache to the shader pipeline cache, use the PSOFC guid as part of the filename, etc."),
-	ECVF_ReadOnly
-);
-
-
-static int32 GEnablePipelineCacheCompression = 1;
-static FAutoConsoleVariableRef GEnablePipelineCacheCompressionCvar(
-	TEXT("r.Vulkan.PipelineCacheCompression"),
-	GEnablePipelineCacheCompression,
-	TEXT("Enable/disable compression on the Vulkan pipeline cache disk file\n"),
-	ECVF_Default | ECVF_RenderThreadSafe
-);
-
-
-FVulkanPipelineStateCacheManager::FGfxPipelineEntry::~FGfxPipelineEntry()
-{
-	check(!bLoaded);
-}
-
-void FVulkanPipelineStateCacheManager::FGfxPipelineEntry::GetOrCreateShaderModules(FVulkanShader*const* Shaders)
+void FVulkanRHIGraphicsPipelineState::GetOrCreateShaderModules(FVulkanShader*const* Shaders)
 {
+	FScopeLock Lock(&FVulkanShaderHandleCS);
 	for (int32 Index = 0; Index < ShaderStage::NumStages; ++Index)
 	{
 		FVulkanShader* Shader = Shaders[Index];
@@ -247,9 +225,9 @@ void FVulkanPipelineStateCacheManager::FGfxPipelineEntry::GetOrCreateShaderModul
 	}
 }
 
-void FVulkanPipelineStateCacheManager::FGfxPipelineEntry::PurgeShaderModules(FVulkanShader*const* Shaders)
+void FVulkanRHIGraphicsPipelineState::PurgeShaderModules(FVulkanShader*const* Shaders)
 {
-	check(!bLoaded);
+	check(!bShaderModulesLoaded);
 
 	for (int32 Index = 0; Index < ShaderStage::NumStages; ++Index)
 	{
@@ -262,9 +240,9 @@ void FVulkanPipelineStateCacheManager::FGfxPipelineEntry::PurgeShaderModules(FVu
 	}
 }
 
-void FVulkanPipelineStateCacheManager::FGfxPipelineEntry::PurgeLoadedShaderModules(FVulkanDevice* InDevice)
+void FVulkanRHIGraphicsPipelineState::PurgeLoadedShaderModules(FVulkanDevice* InDevice)
 {
-	check(bLoaded);
+	check(bShaderModulesLoaded);
 
 	for (int32 Index = 0; Index < ShaderStage::NumStages; ++Index)
 	{
@@ -275,7 +253,7 @@ void FVulkanPipelineStateCacheManager::FGfxPipelineEntry::PurgeLoadedShaderModul
 		}
 	}
 
-	bLoaded = false;
+	bShaderModulesLoaded = false;
 }
 
 
@@ -286,6 +264,8 @@ FVulkanPipelineStateCacheManager::FVulkanPipelineStateCacheManager(FVulkanDevice
 	, bLinkedToPSOFCSucessfulLoaded(false)
 	, PipelineCache(VK_NULL_HANDLE)
 {
+	bUseLRU = (int32)CVarEnableLRU.GetValueOnAnyThread() != 0;
+	LRUUsedPipelineMax = CVarLRUPipelineCapacity.GetValueOnAnyThread();
 }
 
 
@@ -368,8 +348,7 @@ bool FVulkanPipelineStateCacheManager::Load(const TArray<FString>& CacheFilename
 		}
 	}
 
-#if VULKAN_ENABLE_LRU_CACHE
-	if (IsUsePipelineLRU())
+	if(CVarEnableLRU.GetValueOnAnyThread() != 0)
 	{
 		for (const FString& CacheFilename : CacheFilenames)
 		{
@@ -393,7 +372,7 @@ bool FVulkanPipelineStateCacheManager::Load(const TArray<FString>& CacheFilename
 
 				for (int32 Index = 0; Index < File.PipelineSizes.Num(); ++Index)
 				{
-					PipelineSizeList.Add(File.PipelineSizes[Index]->ShaderHash, File.PipelineSizes[Index]);
+					LRU2SizeList.Add(File.PipelineSizes[Index].ShaderHash, File.PipelineSizes[Index]);
 				}
 				UE_LOG(LogVulkanRHI, Display, TEXT("Loaded %d LRU size entries for '%s'"), File.PipelineSizes.Num(), *LruCacheFilename);
 			}
@@ -404,7 +383,6 @@ bool FVulkanPipelineStateCacheManager::Load(const TArray<FString>& CacheFilename
 			}
 		}
 	}
-#endif
 
 	// Lazily create the cache in case the load failed
 	if (PipelineCache == VK_NULL_HANDLE)
@@ -497,13 +475,10 @@ void FVulkanPipelineStateCacheManager::OnShaderPipelineCacheOpened(FString const
 			bLinkedToPSOFCSucessfulLoaded = true;
 			PlatformFile.MoveFile(*LinkedToPSOFCCacheFolderFilename, *TempName);
 
-#if VULKAN_ENABLE_LRU_CACHE
 			if (CVarPipelineLRUCacheEvictBinary.GetValueOnAnyThread())
 			{
 				bEvictImmediately = true;
 			}
-#endif
-
 		}
 	}
 	else
@@ -542,7 +517,6 @@ void FVulkanPipelineStateCacheManager::OnShaderPipelineCachePrecompilationComple
 	OnShaderPipelineCachePrecompilationCompleteDelegate.Reset();
 }
 
-
 void FVulkanPipelineStateCacheManager::Save(const FString& CacheFilename, bool bFromPSOFC)
 {
 	if (bLinkedToPSOFC && !bFromPSOFC)
@@ -550,7 +524,10 @@ void FVulkanPipelineStateCacheManager::Save(const FString& CacheFilename, bool b
 		UE_LOG(LogVulkanRHI, Log, TEXT("FVulkanPipelineStateCacheManager: skipped saving because we only save if the PSOFC based one failed to load."));
 		return;
 	}
-	FScopeLock Lock(&InitializerToPipelineMapCS);
+	FScopeLock Lock1(&GraphicsPSOLockedCS);
+	FScopeLock Lock2(&LRUCS);
+
+
 
 	// First save Device Cache
 	size_t Size = 0;
@@ -589,16 +566,15 @@ void FVulkanPipelineStateCacheManager::Save(const FString& CacheFilename, bool b
 		}
 	}
 
-#if VULKAN_ENABLE_LRU_CACHE
-	if (IsUsePipelineLRU())
+	if (CVarEnableLRU.GetValueOnAnyThread() != 0)
 	{
 		// LRU cache file
 		TArray<uint8> MemFile;
 		FMemoryWriter Ar(MemFile);
 		FVulkanLRUCacheFile File;
 		File.Header.Version = FVulkanLRUCacheFile::LRU_CACHE_VERSION;
-		File.Header.SizeOfPipelineSizes = (int32)sizeof(FPipelineSize);
-		PipelineSizeList.GenerateValueArray(File.PipelineSizes);
+		File.Header.SizeOfPipelineSizes = (int32)sizeof(FVulkanPipelineSize);
+		LRU2SizeList.GenerateValueArray(File.PipelineSizes);
 		File.Save(Ar);
 
 		FString LruCacheFilename = FVulkanPlatform::CreatePSOBinaryCacheFilename(Device, CacheFilename);
@@ -606,79 +582,16 @@ void FVulkanPipelineStateCacheManager::Save(const FString& CacheFilename, bool b
 
 		if (FFileHelper::SaveArrayToFile(MemFile, *LruCacheFilename))
 		{
-			UE_LOG(LogVulkanRHI, Display, TEXT("FVulkanPipelineStateCacheManager: Saved pipeline lru pipeline cache file '%s', %d hashes, %d bytes"), *LruCacheFilename, PipelineSizeList.Num(), MemFile.Num());
-		}
-		else
-		{
-			UE_LOG(LogVulkanRHI, Error, TEXT("FVulkanPipelineStateCacheManager: Failed to save pipeline lru pipeline cache file '%s', %d hashes, %d bytes"), *LruCacheFilename, PipelineSizeList.Num(), MemFile.Num());
-		}
-	}
-#endif
-}
-
-FVulkanRHIGraphicsPipelineState* FVulkanPipelineStateCacheManager::CreateAndAdd(const FGraphicsPipelineStateInitializer& PSOInitializer, FGfxPSIKey PSIKey, TGfxPipelineEntrySharedPtr GfxEntry, FGfxEntryKey GfxEntryKey)
-{
-	check(GfxEntry);
-
-#if VULKAN_ENABLE_LRU_CACHE
-	FVulkanGfxPipeline* Pipeline = new FVulkanGfxPipeline(Device,
-		(PipelineLRU.IsActive() ? GfxEntry : TGfxPipelineEntrySharedPtr()));
-#else
-	FVulkanGfxPipeline* Pipeline = new FVulkanGfxPipeline(Device);
-#endif
-
-	Pipeline->CreateRuntimeObjects(PSOInitializer);
-
-#if VULKAN_ENABLE_LRU_CACHE
-	if (!PSOInitializer.bFromPSOFileCache || !ShouldEvictImmediately())
-#endif
-	{
-		// Create the pipeline
-		double BeginTime = FPlatformTime::Seconds();
-		FVulkanShader* VulkanShaders[ShaderStage::NumStages];
-		GetVulkanShaders(PSOInitializer.BoundShaderState, VulkanShaders);
-		CreateGfxPipelineFromEntry(GfxEntry.Get(), VulkanShaders, Pipeline);
-	
-		// Recover if we failed to create the pipeline.
-		if (!Pipeline->GetHandle())
-		{
-			delete Pipeline;
-			return nullptr;
-		}
-
-		double EndTime = FPlatformTime::Seconds();
-		double Delta = EndTime - BeginTime;
-		if (Delta > HitchTime)
-		{
-			UE_LOG(LogVulkanRHI, Verbose, TEXT("Hitchy gfx pipeline (%.3f ms)"), (float)(Delta * 1000.0));
-		}
-	}
-
-	FVulkanRHIGraphicsPipelineState* PipelineState = nullptr;
-	{
-		FScopeLock Lock(&InitializerToPipelineMapCS);
-		FVulkanRHIGraphicsPipelineState** Existing = InitializerToPipelineMap.Find(PSIKey);
-		if (Existing)
-		{
-			PipelineState = *Existing;
+			UE_LOG(LogVulkanRHI, Display, TEXT("FVulkanPipelineStateCacheManager: Saved pipeline lru pipeline cache file '%s', %d hashes, %d bytes"), *LruCacheFilename, LRU2SizeList.Num(), MemFile.Num());
 		}
 		else
 		{
-			PipelineState = new FVulkanRHIGraphicsPipelineState(PSOInitializer.BoundShaderState, Pipeline, PSOInitializer.PrimitiveType);
-			PipelineState->AddRef();
-			InitializerToPipelineMap.Add(MoveTemp(PSIKey), PipelineState);
+			UE_LOG(LogVulkanRHI, Error, TEXT("FVulkanPipelineStateCacheManager: Failed to save pipeline lru pipeline cache file '%s', %d hashes, %d bytes"), *LruCacheFilename, LRU2SizeList.Num(), MemFile.Num());
 		}
 	}
-
-	EntryKeyToGfxPipelineMap.Add(MoveTemp(GfxEntryKey), PipelineState->Pipeline);
-#if VULKAN_ENABLE_LRU_CACHE
-	PipelineLRU.Add(PipelineState->Pipeline);
-#endif
-
-	return PipelineState;
 }
 
-FArchive& operator << (FArchive& Ar, FVulkanPipelineStateCacheManager::FGfxPipelineEntry::FBlendAttachment& Attachment)
+FArchive& operator << (FArchive& Ar, FGfxPipelineDesc::FBlendAttachment& Attachment)
 {
 	// Modify VERSION if serialization changes
 	Ar << Attachment.bBlend;
@@ -692,7 +605,7 @@ FArchive& operator << (FArchive& Ar, FVulkanPipelineStateCacheManager::FGfxPipel
 	return Ar;
 }
 
-void FVulkanPipelineStateCacheManager::FGfxPipelineEntry::FBlendAttachment::ReadFrom(const VkPipelineColorBlendAttachmentState& InState)
+void FGfxPipelineDesc::FBlendAttachment::ReadFrom(const VkPipelineColorBlendAttachmentState& InState)
 {
 	bBlend =				InState.blendEnable != VK_FALSE;
 	ColorBlendOp =			(uint8)InState.colorBlendOp;
@@ -704,7 +617,7 @@ void FVulkanPipelineStateCacheManager::FGfxPipelineEntry::FBlendAttachment::Read
 	ColorWriteMask =		(uint8)InState.colorWriteMask;
 }
 
-void FVulkanPipelineStateCacheManager::FGfxPipelineEntry::FBlendAttachment::WriteInto(VkPipelineColorBlendAttachmentState& Out) const
+void FGfxPipelineDesc::FBlendAttachment::WriteInto(VkPipelineColorBlendAttachmentState& Out) const
 {
 	Out.blendEnable =			bBlend ? VK_TRUE : VK_FALSE;
 	Out.colorBlendOp =			(VkBlendOp)ColorBlendOp;
@@ -717,7 +630,7 @@ void FVulkanPipelineStateCacheManager::FGfxPipelineEntry::FBlendAttachment::Writ
 }
 
 
-void FVulkanPipelineStateCacheManager::FDescriptorSetLayoutBinding::ReadFrom(const VkDescriptorSetLayoutBinding& InState)
+void FDescriptorSetLayoutBinding::ReadFrom(const VkDescriptorSetLayoutBinding& InState)
 {
 	Binding =			InState.binding;
 	ensure(InState.descriptorCount == 1);
@@ -726,7 +639,7 @@ void FVulkanPipelineStateCacheManager::FDescriptorSetLayoutBinding::ReadFrom(con
 	StageFlags =		InState.stageFlags;
 }
 
-void FVulkanPipelineStateCacheManager::FDescriptorSetLayoutBinding::WriteInto(VkDescriptorSetLayoutBinding& Out) const
+void FDescriptorSetLayoutBinding::WriteInto(VkDescriptorSetLayoutBinding& Out) const
 {
 	Out.binding = Binding;
 	//Out.descriptorCount = DescriptorCount;
@@ -734,7 +647,7 @@ void FVulkanPipelineStateCacheManager::FDescriptorSetLayoutBinding::WriteInto(Vk
 	Out.stageFlags = StageFlags;
 }
 
-FArchive& operator << (FArchive& Ar, FVulkanPipelineStateCacheManager::FDescriptorSetLayoutBinding& Binding)
+FArchive& operator << (FArchive& Ar, FDescriptorSetLayoutBinding& Binding)
 {
 	// Modify VERSION if serialization changes
 	Ar << Binding.Binding;
@@ -744,21 +657,21 @@ FArchive& operator << (FArchive& Ar, FVulkanPipelineStateCacheManager::FDescript
 	return Ar;
 }
 
-void FVulkanPipelineStateCacheManager::FGfxPipelineEntry::FVertexBinding::ReadFrom(const VkVertexInputBindingDescription& InState)
+void FGfxPipelineDesc::FVertexBinding::ReadFrom(const VkVertexInputBindingDescription& InState)
 {
 	Binding =	InState.binding;
 	InputRate =	(uint16)InState.inputRate;
 	Stride =	InState.stride;
 }
 
-void FVulkanPipelineStateCacheManager::FGfxPipelineEntry::FVertexBinding::WriteInto(VkVertexInputBindingDescription& Out) const
+void FGfxPipelineDesc::FVertexBinding::WriteInto(VkVertexInputBindingDescription& Out) const
 {
 	Out.binding =	Binding;
 	Out.inputRate =	(VkVertexInputRate)InputRate;
 	Out.stride =	Stride;
 }
 
-FArchive& operator << (FArchive& Ar, FVulkanPipelineStateCacheManager::FGfxPipelineEntry::FVertexBinding& Binding)
+FArchive& operator << (FArchive& Ar, FGfxPipelineDesc::FVertexBinding& Binding)
 {
 	// Modify VERSION if serialization changes
 	Ar << Binding.Stride;
@@ -767,7 +680,7 @@ FArchive& operator << (FArchive& Ar, FVulkanPipelineStateCacheManager::FGfxPipel
 	return Ar;
 }
 
-void FVulkanPipelineStateCacheManager::FGfxPipelineEntry::FVertexAttribute::ReadFrom(const VkVertexInputAttributeDescription& InState)
+void FGfxPipelineDesc::FVertexAttribute::ReadFrom(const VkVertexInputAttributeDescription& InState)
 {
 	Binding =	InState.binding;
 	Format =	(uint32)InState.format;
@@ -775,7 +688,7 @@ void FVulkanPipelineStateCacheManager::FGfxPipelineEntry::FVertexAttribute::Read
 	Offset =	InState.offset;
 }
 
-void FVulkanPipelineStateCacheManager::FGfxPipelineEntry::FVertexAttribute::WriteInto(VkVertexInputAttributeDescription& Out) const
+void FGfxPipelineDesc::FVertexAttribute::WriteInto(VkVertexInputAttributeDescription& Out) const
 {
 	Out.binding =	Binding;
 	Out.format =	(VkFormat)Format;
@@ -783,7 +696,7 @@ void FVulkanPipelineStateCacheManager::FGfxPipelineEntry::FVertexAttribute::Writ
 	Out.offset =	Offset;
 }
 
-FArchive& operator << (FArchive& Ar, FVulkanPipelineStateCacheManager::FGfxPipelineEntry::FVertexAttribute& Attribute)
+FArchive& operator << (FArchive& Ar, FGfxPipelineDesc::FVertexAttribute& Attribute)
 {
 	// Modify VERSION if serialization changes
 	Ar << Attribute.Location;
@@ -793,7 +706,7 @@ FArchive& operator << (FArchive& Ar, FVulkanPipelineStateCacheManager::FGfxPipel
 	return Ar;
 }
 
-void FVulkanPipelineStateCacheManager::FGfxPipelineEntry::FRasterizer::ReadFrom(const VkPipelineRasterizationStateCreateInfo& InState)
+void FGfxPipelineDesc::FRasterizer::ReadFrom(const VkPipelineRasterizationStateCreateInfo& InState)
 {
 	PolygonMode =				InState.polygonMode;
 	CullMode =					InState.cullMode;
@@ -801,7 +714,7 @@ void FVulkanPipelineStateCacheManager::FGfxPipelineEntry::FRasterizer::ReadFrom(
 	DepthBiasConstantFactor =	InState.depthBiasConstantFactor;
 }
 
-void FVulkanPipelineStateCacheManager::FGfxPipelineEntry::FRasterizer::WriteInto(VkPipelineRasterizationStateCreateInfo& Out) const
+void FGfxPipelineDesc::FRasterizer::WriteInto(VkPipelineRasterizationStateCreateInfo& Out) const
 {
 	Out.polygonMode =				(VkPolygonMode)PolygonMode;
 	Out.cullMode =					(VkCullModeFlags)CullMode;
@@ -813,7 +726,7 @@ void FVulkanPipelineStateCacheManager::FGfxPipelineEntry::FRasterizer::WriteInto
 	Out.depthBiasConstantFactor =	DepthBiasConstantFactor;
 }
 
-FArchive& operator << (FArchive& Ar, FVulkanPipelineStateCacheManager::FGfxPipelineEntry::FRasterizer& Rasterizer)
+FArchive& operator << (FArchive& Ar, FGfxPipelineDesc::FRasterizer& Rasterizer)
 {
 	// Modify VERSION if serialization changes
 	Ar << Rasterizer.PolygonMode;
@@ -823,7 +736,7 @@ FArchive& operator << (FArchive& Ar, FVulkanPipelineStateCacheManager::FGfxPipel
 	return Ar;
 }
 
-void FVulkanPipelineStateCacheManager::FGfxPipelineEntry::FDepthStencil::ReadFrom(const VkPipelineDepthStencilStateCreateInfo& InState)
+void FGfxPipelineDesc::FDepthStencil::ReadFrom(const VkPipelineDepthStencilStateCreateInfo& InState)
 {
 	DepthCompareOp =			(uint8)InState.depthCompareOp;
 	bDepthTestEnable =			InState.depthTestEnable != VK_FALSE;
@@ -846,7 +759,7 @@ void FVulkanPipelineStateCacheManager::FGfxPipelineEntry::FDepthStencil::ReadFro
 	BackReference =				InState.back.reference;
 }
 
-void FVulkanPipelineStateCacheManager::FGfxPipelineEntry::FDepthStencil::WriteInto(VkPipelineDepthStencilStateCreateInfo& Out) const
+void FGfxPipelineDesc::FDepthStencil::WriteInto(VkPipelineDepthStencilStateCreateInfo& Out) const
 {
 	Out.depthCompareOp =		(VkCompareOp)DepthCompareOp;
 	Out.depthTestEnable =		bDepthTestEnable;
@@ -869,7 +782,7 @@ void FVulkanPipelineStateCacheManager::FGfxPipelineEntry::FDepthStencil::WriteIn
 	Out.back.reference =		BackReference;
 }
 
-FArchive& operator << (FArchive& Ar, FVulkanPipelineStateCacheManager::FGfxPipelineEntry::FDepthStencil& DepthStencil)
+FArchive& operator << (FArchive& Ar, FGfxPipelineDesc::FDepthStencil& DepthStencil)
 {
 	// Modify VERSION if serialization changes
 	Ar << DepthStencil.DepthCompareOp;
@@ -894,19 +807,19 @@ FArchive& operator << (FArchive& Ar, FVulkanPipelineStateCacheManager::FGfxPipel
 	return Ar;
 }
 
-void FVulkanPipelineStateCacheManager::FGfxPipelineEntry::FRenderTargets::FAttachmentRef::ReadFrom(const VkAttachmentReference& InState)
+void FGfxPipelineDesc::FRenderTargets::FAttachmentRef::ReadFrom(const VkAttachmentReference& InState)
 {
 	Attachment =	InState.attachment;
 	Layout =		(uint64)InState.layout;
 }
 
-void FVulkanPipelineStateCacheManager::FGfxPipelineEntry::FRenderTargets::FAttachmentRef::WriteInto(VkAttachmentReference& Out) const
+void FGfxPipelineDesc::FRenderTargets::FAttachmentRef::WriteInto(VkAttachmentReference& Out) const
 {
 	Out.attachment =	Attachment;
 	Out.layout =		(VkImageLayout)Layout;
 }
 
-FArchive& operator << (FArchive& Ar, FVulkanPipelineStateCacheManager::FGfxPipelineEntry::FRenderTargets::FAttachmentRef& AttachmentRef)
+FArchive& operator << (FArchive& Ar, FGfxPipelineDesc::FRenderTargets::FAttachmentRef& AttachmentRef)
 {
 	// Modify VERSION if serialization changes
 	Ar << AttachmentRef.Attachment;
@@ -914,7 +827,7 @@ FArchive& operator << (FArchive& Ar, FVulkanPipelineStateCacheManager::FGfxPipel
 	return Ar;
 }
 
-void FVulkanPipelineStateCacheManager::FGfxPipelineEntry::FRenderTargets::FAttachmentDesc::ReadFrom(const VkAttachmentDescription &InState)
+void FGfxPipelineDesc::FRenderTargets::FAttachmentDesc::ReadFrom(const VkAttachmentDescription &InState)
 {
 	Format =			(uint32)InState.format;
 	Flags =				(uint8)InState.flags;
@@ -927,7 +840,7 @@ void FVulkanPipelineStateCacheManager::FGfxPipelineEntry::FRenderTargets::FAttac
 	FinalLayout =		(uint64)InState.finalLayout;
 }
 
-void FVulkanPipelineStateCacheManager::FGfxPipelineEntry::FRenderTargets::FAttachmentDesc::WriteInto(VkAttachmentDescription& Out) const
+void FGfxPipelineDesc::FRenderTargets::FAttachmentDesc::WriteInto(VkAttachmentDescription& Out) const
 {
 	Out.format =			(VkFormat)Format;
 	Out.flags =				Flags;
@@ -940,7 +853,7 @@ void FVulkanPipelineStateCacheManager::FGfxPipelineEntry::FRenderTargets::FAttac
 	Out.finalLayout =		(VkImageLayout)FinalLayout;
 }
 
-FArchive& operator << (FArchive& Ar, FVulkanPipelineStateCacheManager::FGfxPipelineEntry::FRenderTargets::FAttachmentDesc& AttachmentDesc)
+FArchive& operator << (FArchive& Ar, FGfxPipelineDesc::FRenderTargets::FAttachmentDesc& AttachmentDesc)
 {
 	// Modify VERSION if serialization changes
 	Ar << AttachmentDesc.Format;
@@ -956,27 +869,27 @@ FArchive& operator << (FArchive& Ar, FVulkanPipelineStateCacheManager::FGfxPipel
 	return Ar;
 }
 
-void FVulkanPipelineStateCacheManager::FGfxPipelineEntry::FRenderTargets::ReadFrom(const FVulkanRenderTargetLayout& RTLayout)
+void FGfxPipelineDesc::FRenderTargets::ReadFrom(const FVulkanRenderTargetLayout& RTLayout)
 {
-	NumAttachments =				RTLayout.NumAttachmentDescriptions;
-	NumColorAttachments =			RTLayout.NumColorAttachments;
+	NumAttachments =			RTLayout.NumAttachmentDescriptions;
+	NumColorAttachments =		RTLayout.NumColorAttachments;
 
-	bHasDepthStencil =				RTLayout.bHasDepthStencil != 0;
-	bHasResolveAttachments =		RTLayout.bHasResolveAttachments != 0;
+	bHasDepthStencil =			RTLayout.bHasDepthStencil != 0;
+	bHasResolveAttachments =	RTLayout.bHasResolveAttachments != 0;
 	bHasFragmentDensityAttachment =	RTLayout.bHasFragmentDensityAttachment != 0;
-	NumUsedClearValues =			RTLayout.NumUsedClearValues;
+	NumUsedClearValues =		RTLayout.NumUsedClearValues;
 
-	RenderPassCompatibleHash =		RTLayout.GetRenderPassCompatibleHash();
+	RenderPassCompatibleHash =	RTLayout.GetRenderPassCompatibleHash();
 
 	Extent3D.X = RTLayout.Extent.Extent3D.width;
 	Extent3D.Y = RTLayout.Extent.Extent3D.height;
 	Extent3D.Z = RTLayout.Extent.Extent3D.depth;
 
-	auto CopyAttachmentRefs = [&](TArray<FGfxPipelineEntry::FRenderTargets::FAttachmentRef>& Dest, const VkAttachmentReference* Source, uint32 Count)
+	auto CopyAttachmentRefs = [&](TArray<FGfxPipelineDesc::FRenderTargets::FAttachmentRef>& Dest, const VkAttachmentReference* Source, uint32 Count)
 	{
 		for (uint32 Index = 0; Index < Count; ++Index)
 		{
-			FGfxPipelineEntry::FRenderTargets::FAttachmentRef* New = new(Dest) FGfxPipelineEntry::FRenderTargets::FAttachmentRef;
+			FGfxPipelineDesc::FRenderTargets::FAttachmentRef* New = new(Dest) FGfxPipelineDesc::FRenderTargets::FAttachmentRef;
 			New->ReadFrom(Source[Index]);
 		}
 	};
@@ -992,24 +905,24 @@ void FVulkanPipelineStateCacheManager::FGfxPipelineEntry::FRenderTargets::ReadFr
 	}
 }
 
-void FVulkanPipelineStateCacheManager::FGfxPipelineEntry::FRenderTargets::WriteInto(FVulkanRenderTargetLayout& Out) const
+void FGfxPipelineDesc::FRenderTargets::WriteInto(FVulkanRenderTargetLayout& Out) const
 {
-	Out.NumAttachmentDescriptions =		NumAttachments;
-	Out.NumColorAttachments =			NumColorAttachments;
+	Out.NumAttachmentDescriptions =	NumAttachments;
+	Out.NumColorAttachments =		NumColorAttachments;
 
-	Out.bHasDepthStencil =				bHasDepthStencil;
-	Out.bHasResolveAttachments =		bHasResolveAttachments;
+	Out.bHasDepthStencil =			bHasDepthStencil;
+	Out.bHasResolveAttachments =	bHasResolveAttachments;
 	Out.bHasFragmentDensityAttachment =	bHasFragmentDensityAttachment;
-	Out.NumUsedClearValues =			NumUsedClearValues;
+	Out.NumUsedClearValues =		NumUsedClearValues;
 
 	ensure(0);
-	Out.RenderPassCompatibleHash =		RenderPassCompatibleHash;
+	Out.RenderPassCompatibleHash =	RenderPassCompatibleHash;
 
-	Out.Extent.Extent3D.width =			Extent3D.X;
-	Out.Extent.Extent3D.height =		Extent3D.Y;
-	Out.Extent.Extent3D.depth =			Extent3D.Z;
+	Out.Extent.Extent3D.width =		Extent3D.X;
+	Out.Extent.Extent3D.height =	Extent3D.Y;
+	Out.Extent.Extent3D.depth =		Extent3D.Z;
 
-	auto CopyAttachmentRefs = [&](const TArray<FGfxPipelineEntry::FRenderTargets::FAttachmentRef>& Source, VkAttachmentReference* Dest, uint32 Count)
+	auto CopyAttachmentRefs = [&](const TArray<FGfxPipelineDesc::FRenderTargets::FAttachmentRef>& Source, VkAttachmentReference* Dest, uint32 Count)
 	{
 		for (uint32 Index = 0; Index < Count; ++Index, ++Dest)
 		{
@@ -1021,13 +934,14 @@ void FVulkanPipelineStateCacheManager::FGfxPipelineEntry::FRenderTargets::WriteI
 	DepthStencil.WriteInto(Out.DepthStencilReference);
 	FragmentDensity.WriteInto(Out.FragmentDensityReference);
 
+
 	for (int32 Index = 0; Index < UE_ARRAY_COUNT(Out.Desc); ++Index)
 	{
 		Descriptions[Index].WriteInto(Out.Desc[Index]);
 	}
 }
 
-FArchive& operator << (FArchive& Ar, FVulkanPipelineStateCacheManager::FGfxPipelineEntry::FRenderTargets& RTs)
+FArchive& operator << (FArchive& Ar, FGfxPipelineDesc::FRenderTargets& RTs)
 {
 	// Modify VERSION if serialization changes
 	Ar << RTs.NumAttachments;
@@ -1048,7 +962,7 @@ FArchive& operator << (FArchive& Ar, FVulkanPipelineStateCacheManager::FGfxPipel
 	return Ar;
 }
 
-FArchive& operator << (FArchive& Ar, FVulkanPipelineStateCacheManager::FGfxPipelineEntry& Entry)
+FArchive& operator << (FArchive& Ar, FGfxPipelineDesc& Entry)
 {
 	// Modify VERSION if serialization changes
 	Ar << Entry.VertexInputKey;
@@ -1066,11 +980,17 @@ FArchive& operator << (FArchive& Ar, FVulkanPipelineStateCacheManager::FGfxPipel
 
 	Ar << Entry.DepthStencil;
 
+#if VULKAN_USE_SHADERKEYS
+	for (uint64& ShaderKey : Entry.ShaderKeys)
+	{
+		Ar << ShaderKey;
+	}
+#else
 	for (int32 Index = 0; Index < UE_ARRAY_COUNT(Entry.ShaderHashes.Stages); ++Index)
 	{
 		Ar << Entry.ShaderHashes.Stages[Index];
 	}
-
+#endif
 	Ar << Entry.RenderTargets;
 
 #if VULKAN_SUPPORTS_COLOR_CONVERSIONS
@@ -1085,48 +1005,49 @@ FArchive& operator << (FArchive& Ar, FVulkanPipelineStateCacheManager::FGfxPipel
 	return Ar;
 }
 
-FArchive& operator << (FArchive& Ar, FVulkanPipelineStateCacheManager::FGfxPipelineEntry* Entry)
+FArchive& operator << (FArchive& Ar, FGfxPipelineDesc* Entry)
 {
 	return Ar << (*Entry);
 }
 
-#if VULKAN_ENABLE_LRU_CACHE	
-FArchive& operator << (FArchive& Ar, FVulkanPipelineStateCacheManager::FPipelineSize& PS)
+FArchive& operator << (FArchive& Ar, FVulkanPipelineSize& PS)
 {
 	Ar << PS.ShaderHash;
 	Ar << PS.PipelineSize;
 
 	return Ar;
 }
-#endif
 
 
-FGfxEntryKey FVulkanPipelineStateCacheManager::FGfxPipelineEntry::CreateKey() const
+FVulkanPSOKey FGfxPipelineDesc::CreateKey2() const
 {
-	FGfxEntryKey Result;
+	FVulkanPSOKey Result;
 	Result.GenerateFromArchive([this](FArchive& Ar)
-		{
-			Ar << const_cast<FGfxPipelineEntry&>(*this);
-		});
+	{
+		Ar << const_cast<FGfxPipelineDesc&>(*this);
+	});
 	return Result;
 }
 
-void FVulkanPipelineStateCacheManager::CreateGfxPipelineFromEntry(FGfxPipelineEntry* GfxEntry, FVulkanShader* Shaders[ShaderStage::NumStages], FVulkanGfxPipeline* Pipeline)
+
+
+bool FVulkanPipelineStateCacheManager::CreateGfxPipelineFromEntry(FVulkanRHIGraphicsPipelineState* PSO, FVulkanShader* Shaders[ShaderStage::NumStages], VkPipeline* Pipeline)
 {
+	FGfxPipelineDesc* GfxEntry = &PSO->Desc;
 	if (Shaders[ShaderStage::Pixel] == nullptr && !FVulkanPlatform::SupportsNullPixelShader())
 	{
 		Shaders[ShaderStage::Pixel] = ResourceCast(TShaderMapRef<FNULLPS>(GetGlobalShaderMap(GMaxRHIFeatureLevel))->GetPixelShader());
 	}
 
-	if (!GfxEntry->bLoaded)
+	if (!PSO->bShaderModulesLoaded)
 	{
-		GfxEntry->GetOrCreateShaderModules(Shaders);
+		PSO->GetOrCreateShaderModules(Shaders);
 	}
-	
+
 	// Pipeline
 	VkGraphicsPipelineCreateInfo PipelineInfo;
 	ZeroVulkanStruct(PipelineInfo, VK_STRUCTURE_TYPE_GRAPHICS_PIPELINE_CREATE_INFO);
-	PipelineInfo.layout = GfxEntry->Layout->GetPipelineLayout();
+	PipelineInfo.layout = PSO->Layout->GetPipelineLayout();
 
 	// Color Blend
 	VkPipelineColorBlendStateCreateInfo CBInfo;
@@ -1147,7 +1068,7 @@ void FVulkanPipelineStateCacheManager::CreateGfxPipelineFromEntry(FGfxPipelineEn
 		{
 			BlendStates[Index].colorWriteMask = 0;
 		}
-		ColorWriteMask >>= 1;
+		ColorWriteMask >>= 1;		
 	}
 	CBInfo.pAttachments = BlendStates;
 	CBInfo.blendConstants[0] = 1.0f;
@@ -1175,7 +1096,7 @@ void FVulkanPipelineStateCacheManager::CreateGfxPipelineFromEntry(FGfxPipelineEn
 	bool bHasTessellation = false;
 	for (int32 ShaderStage = 0; ShaderStage < ShaderStage::NumStages; ++ShaderStage)
 	{
-		if (!GfxEntry->ShaderModules[ShaderStage])
+		if (!PSO->ShaderModules[ShaderStage])
 		{
 			continue;
 		}
@@ -1185,7 +1106,7 @@ void FVulkanPipelineStateCacheManager::CreateGfxPipelineFromEntry(FGfxPipelineEn
 		VkShaderStageFlagBits Stage = UEFrequencyToVKStageBit(ShaderStage::GetFrequencyForGfxStage(CurrStage));
 		ShaderStages[PipelineInfo.stageCount].stage = Stage;
 		bHasTessellation = bHasTessellation || ((Stage & (VK_SHADER_STAGE_TESSELLATION_CONTROL_BIT | VK_SHADER_STAGE_TESSELLATION_EVALUATION_BIT)) != 0);
-		ShaderStages[PipelineInfo.stageCount].module = GfxEntry->ShaderModules[CurrStage];
+		ShaderStages[PipelineInfo.stageCount].module = PSO->ShaderModules[CurrStage];
 		Shaders[ShaderStage]->GetEntryPoint(EntryPoints[PipelineInfo.stageCount]);
 		ShaderStages[PipelineInfo.stageCount].pName = EntryPoints[PipelineInfo.stageCount];
 		PipelineInfo.stageCount++;
@@ -1197,7 +1118,7 @@ void FVulkanPipelineStateCacheManager::CreateGfxPipelineFromEntry(FGfxPipelineEn
 	VkPipelineVertexInputStateCreateInfo VBInfo;
 	ZeroVulkanStruct(VBInfo, VK_STRUCTURE_TYPE_PIPELINE_VERTEX_INPUT_STATE_CREATE_INFO);
 	TArray<VkVertexInputBindingDescription> VBBindings;
-	for (const FGfxPipelineEntry::FVertexBinding& SourceBinding : GfxEntry->VertexBindings)
+	for (const FGfxPipelineDesc::FVertexBinding& SourceBinding : GfxEntry->VertexBindings)
 	{
 		VkVertexInputBindingDescription* Binding = new(VBBindings) VkVertexInputBindingDescription;
 		SourceBinding.WriteInto(*Binding);
@@ -1205,7 +1126,7 @@ void FVulkanPipelineStateCacheManager::CreateGfxPipelineFromEntry(FGfxPipelineEn
 	VBInfo.vertexBindingDescriptionCount = VBBindings.Num();
 	VBInfo.pVertexBindingDescriptions = VBBindings.GetData();
 	TArray<VkVertexInputAttributeDescription> VBAttributes;
-	for (const FGfxPipelineEntry::FVertexAttribute& SourceAttr : GfxEntry->VertexAttributes)
+	for (const FGfxPipelineDesc::FVertexAttribute& SourceAttr : GfxEntry->VertexAttributes)
 	{
 		VkVertexInputAttributeDescription* Attr = new(VBAttributes) VkVertexInputAttributeDescription;
 		SourceAttr.WriteInto(*Attr);
@@ -1218,7 +1139,7 @@ void FVulkanPipelineStateCacheManager::CreateGfxPipelineFromEntry(FGfxPipelineEn
 	PipelineInfo.pMultisampleState = &MSInfo;
 	PipelineInfo.pViewportState = &VPInfo;
 
-	PipelineInfo.renderPass = GfxEntry->RenderPass->GetHandle();
+	PipelineInfo.renderPass = PSO->RenderPass->GetHandle();
 	PipelineInfo.subpass = GfxEntry->SubpassIndex;
 
 	VkPipelineInputAssemblyStateCreateInfo InputAssembly;
@@ -1264,25 +1185,36 @@ void FVulkanPipelineStateCacheManager::CreateGfxPipelineFromEntry(FGfxPipelineEn
 		RasterizerState.frontFace = VK_FRONT_FACE_COUNTER_CLOCKWISE;
 	}
 
-	//#todo-rco: Fix me
 	VkResult Result = VK_ERROR_INITIALIZATION_FAILED;
 	double BeginTime = FPlatformTime::Seconds();
-#if VULKAN_ENABLE_LRU_CACHE
-	if (PipelineLRU.IsActive())
+	if(bUseLRU)
 	{
-		const uint32 ShaderHash = GfxEntry->ShaderHashes.Hash;
-		FPipelineSize** Found = PipelineSizeList.Find(ShaderHash);
+#if VULKAN_USE_SHADERKEYS
+		const uint64 ShaderHash = GfxEntry->ShaderKeyShared;
+#else
+		const uint64 ShaderHash = GfxEntry->ShaderHashes.Hash;
+#endif
+		FVulkanPipelineSize* Found;
+		{
+			FScopeLock Lock(&LRUCS); 
+			Found = LRU2SizeList.Find(ShaderHash);
+		}
 		size_t PreSize = 0, AfterSize = 0;
+		uint32 FoundSize = 0;
 		if (Found)
 		{
-			Pipeline->PipelineCacheSize = (*Found)->PipelineSize;
+			FoundSize = Found->PipelineSize;
 		}
 		else
 		{
 			VulkanRHI::vkGetPipelineCacheData(Device->GetInstanceHandle(), PipelineCache, &PreSize, nullptr);
 		}
 
-		Result = VulkanRHI::vkCreateGraphicsPipelines(Device->GetInstanceHandle(), PipelineCache, 1, &PipelineInfo, VULKAN_CPU_ALLOCATOR, &Pipeline->Pipeline);
+
+		{
+			SCOPE_CYCLE_COUNTER(STAT_VulkanPSOVulkanCreationTime);
+			Result = VulkanRHI::vkCreateGraphicsPipelines(Device->GetInstanceHandle(), PipelineCache, 1, &PipelineInfo, VULKAN_CPU_ALLOCATOR, Pipeline);
+		}
 
 		if (!Found && Result == VK_SUCCESS)
 		{
@@ -1293,35 +1225,31 @@ void FVulkanPipelineStateCacheManager::CreateGfxPipelineFromEntry(FGfxPipelineEn
 				UE_LOG(LogVulkanRHI, Warning, TEXT("Shader size was computed as zero, using 20k instead."));
 				Diff = 20 * 1024;
 			}
-			FPipelineSize* PipelineSize = new FPipelineSize();
-			PipelineSize->ShaderHash = ShaderHash;
-			PipelineSize->PipelineSize = Diff;
-			PipelineSizeList.Add(ShaderHash, PipelineSize);
-			Pipeline->PipelineCacheSize = Diff;
+			FVulkanPipelineSize PipelineSize;
+			PipelineSize.ShaderHash = ShaderHash;
+			PipelineSize.PipelineSize = Diff;
+			{
+				FScopeLock Lock(&LRUCS);
+				LRU2SizeList.Add(ShaderHash, PipelineSize);
+			}
+			FoundSize = Diff;
+		}
+		if(Result == VK_SUCCESS)
+		{
+			PSO->PipelineCacheSize = FoundSize;
 		}
 	}
 	else
-#endif
 	{
-		Result = VulkanRHI::vkCreateGraphicsPipelines(Device->GetInstanceHandle(), PipelineCache, 1, &PipelineInfo, VULKAN_CPU_ALLOCATOR, &Pipeline->Pipeline);
+		SCOPE_CYCLE_COUNTER(STAT_VulkanPSOVulkanCreationTime);
+		Result = VulkanRHI::vkCreateGraphicsPipelines(Device->GetInstanceHandle(), PipelineCache, 1, &PipelineInfo, VULKAN_CPU_ALLOCATOR, Pipeline);
 	}
 
-#if VULKAN_PURGE_SHADER_MODULES
-	if (GfxEntry->bLoaded)
-	{
-		GfxEntry->PurgeLoadedShaderModules(Device);
-	}
-	else
-	{
-		GfxEntry->PurgeShaderModules(Shaders);
-	}
-#endif
 
 	if (Result != VK_SUCCESS)
 	{
 		UE_LOG(LogVulkanRHI, Error, TEXT("Failed to create graphics pipeline."));
-		Pipeline->Pipeline = VK_NULL_HANDLE;
-		return;
+		return false;
 	}
 
 	double EndTime = FPlatformTime::Seconds();
@@ -1332,59 +1260,34 @@ void FVulkanPipelineStateCacheManager::CreateGfxPipelineFromEntry(FGfxPipelineEn
 	}
 
 	INC_DWORD_STAT(STAT_VulkanNumPSOs);
-
-	Pipeline->Layout = GfxEntry->Layout;
+	return true;
 }
 
+
 void FVulkanPipelineStateCacheManager::DestroyCache()
 {
 	VkDevice DeviceHandle = Device->GetInstanceHandle();
 
-	// Graphics
+	FScopeLock Lock1(&GraphicsPSOLockedCS);
+	int idx = 0;
+	for (auto& Pair : GraphicsPSOLockedMap)
 	{
-		for (auto& Pair : InitializerToPipelineMap)
-		{
-			FVulkanRHIGraphicsPipelineState* Pipeline = Pair.Value;
-			//when DestroyCache is called as part of r.Vulkan.RebuildPipelineCache, a pipeline can still be referenced by FVulkanPendingGfxState
-			ensure(GIsRHIInitialized || (!GIsRHIInitialized && Pipeline->GetRefCount() == 1));
-			Pipeline->Release();
-		}
-		InitializerToPipelineMap.Reset();
-
-		for (auto& Pair : GfxPipelineEntries)
-		{
-			FGfxPipelineEntry* Entry = Pair.Value.Get();
-			if (Entry->bLoaded)
-			{
-				Entry->PurgeLoadedShaderModules(Device);
-			}
-		}
-		GfxPipelineEntries.Reset();
+		FVulkanRHIGraphicsPipelineState* Pipeline = Pair.Value;
+		LRUPRINT_DEBUG(TEXT("LIVE PSO %05d  REF %d %p :: %p\n"), idx++, Pipeline->GetRefCount(), Pipeline, Pipeline->LRUNode);
+	}
+	LRU2SizeList.Reset();
 
-#if VULKAN_ENABLE_LRU_CACHE
-		for (auto& Pair : PipelineSizeList)
-		{
-			FPipelineSize* Entry = Pair.Value;
-			delete Entry;
-		}
-		PipelineSizeList.Reset();
+#if LRU_DEBUG
+	LRUDump();
 #endif
 
-		// This map can simply be cleared as InitializerToPipelineMap already decreased the refcount of the pipeline objects	
-		{
-			FScopeLock Lock(&EntryKeyToGfxPipelineMapCS);
-			EntryKeyToGfxPipelineMap.Reset();
-		}
-	}
-
 	// Compute pipelines already deleted...
 	ComputePipelineEntries.Reset();
+	check(GraphicsPSOLockedMap.Num() == 0);
 }
 
 void FVulkanPipelineStateCacheManager::RebuildCache()
 {
-	UE_LOG(LogVulkanRHI, Warning, TEXT("Rebuilding pipeline cache; ditching %d entries"), GfxPipelineEntries.Num() + ComputePipelineEntries.Num());
-
 	if (IsInGameThread())
 	{
 		FlushRenderingCommands();
@@ -1392,7 +1295,7 @@ void FVulkanPipelineStateCacheManager::RebuildCache()
 	DestroyCache();
 }
 
-FVulkanPipelineStateCacheManager::FShaderHashes::FShaderHashes(const FGraphicsPipelineStateInitializer& PSOInitializer)
+FVulkanShaderHashes::FVulkanShaderHashes(const FGraphicsPipelineStateInitializer& PSOInitializer)
 {
 	Stages[ShaderStage::Vertex] = GetShaderHash<FRHIVertexShader, FVulkanVertexShader>(PSOInitializer.BoundShaderState.VertexShaderRHI);
 	Stages[ShaderStage::Pixel] = GetShaderHash<FRHIPixelShader, FVulkanPixelShader>(PSOInitializer.BoundShaderState.PixelShaderRHI);
@@ -1406,13 +1309,13 @@ FVulkanPipelineStateCacheManager::FShaderHashes::FShaderHashes(const FGraphicsPi
 	Finalize();
 }
 
-FVulkanPipelineStateCacheManager::FShaderHashes::FShaderHashes()
+FVulkanShaderHashes::FVulkanShaderHashes()
 {
 	FMemory::Memzero(Stages);
 	Hash = 0;
 }
 
-inline FVulkanLayout* FVulkanPipelineStateCacheManager::FindOrAddLayout(const FVulkanDescriptorSetsLayoutInfo& DescriptorSetLayoutInfo, bool bGfxLayout)
+FVulkanLayout* FVulkanPipelineStateCacheManager::FindOrAddLayout(const FVulkanDescriptorSetsLayoutInfo& DescriptorSetLayoutInfo, bool bGfxLayout)
 {
 	FScopeLock Lock(&LayoutMapCS);
 	if (FVulkanLayout** FoundLayout = LayoutMap.Find(DescriptorSetLayoutInfo))
@@ -1422,10 +1325,12 @@ inline FVulkanLayout* FVulkanPipelineStateCacheManager::FindOrAddLayout(const FV
 	}
 
 	FVulkanLayout* Layout = nullptr;
+	FVulkanGfxLayout* GfxLayout = nullptr;
 
 	if (bGfxLayout)
 	{
-		Layout = new FVulkanGfxLayout(Device);
+		GfxLayout = new FVulkanGfxLayout(Device);
+		Layout = GfxLayout;
 	}
 	else
 	{
@@ -1435,60 +1340,14 @@ inline FVulkanLayout* FVulkanPipelineStateCacheManager::FindOrAddLayout(const FV
 	Layout->DescriptorSetLayout.CopyFrom(DescriptorSetLayoutInfo);
 	Layout->Compile(DSetLayoutMap);
 
-	LayoutMap.Add(DescriptorSetLayoutInfo, Layout);
-	return Layout;
-}
-
-FVulkanGfxLayout* FVulkanPipelineStateCacheManager::GetOrGenerateGfxLayout(const FGraphicsPipelineStateInitializer& PSOInitializer,
-	FVulkanShader*const* Shaders, FVulkanVertexInputStateInfo& OutVertexInputState)
-{
-	const FBoundShaderStateInput& BSI = PSOInitializer.BoundShaderState;
-
-	const FVulkanShaderHeader& VSHeader = Shaders[ShaderStage::Vertex]->GetCodeHeader();
-	OutVertexInputState.Generate(ResourceCast(PSOInitializer.BoundShaderState.VertexDeclarationRHI), VSHeader.InOutMask);
-
-	FUniformBufferGatherInfo UBGatherInfo;
-
-	// First pass to gather uniform buffer info
-	FVulkanDescriptorSetsLayoutInfo DescriptorSetLayoutInfo;
-	DescriptorSetLayoutInfo.ProcessBindingsForStage(VK_SHADER_STAGE_VERTEX_BIT, ShaderStage::Vertex, VSHeader, UBGatherInfo);
-
-	if (Shaders[ShaderStage::Pixel])
+	if(GfxLayout)
 	{
-		const FVulkanShaderHeader& PSHeader = Shaders[ShaderStage::Pixel]->GetCodeHeader();
-		DescriptorSetLayoutInfo.ProcessBindingsForStage(VK_SHADER_STAGE_FRAGMENT_BIT, ShaderStage::Pixel, PSHeader, UBGatherInfo);
+		GfxLayout->GfxPipelineDescriptorInfo.Initialize(GfxLayout->GetDescriptorSetsLayout().RemappingInfo);
 	}
 
-#if VULKAN_SUPPORTS_GEOMETRY_SHADERS
-	if (Shaders[ShaderStage::Geometry])
-	{
-		const FVulkanShaderHeader& GSHeader = Shaders[ShaderStage::Geometry]->GetCodeHeader();
-		DescriptorSetLayoutInfo.ProcessBindingsForStage(VK_SHADER_STAGE_GEOMETRY_BIT, ShaderStage::Geometry, GSHeader, UBGatherInfo);
-	}
-#endif
-
-#if PLATFORM_SUPPORTS_TESSELLATION_SHADERS
-	if (Shaders[ShaderStage::Hull])
-	{
-		const FVulkanShaderHeader& HSHeader = Shaders[ShaderStage::Hull]->GetCodeHeader();
-		const FVulkanShaderHeader& DSHeader = Shaders[ShaderStage::Domain]->GetCodeHeader();
-		DescriptorSetLayoutInfo.ProcessBindingsForStage(VK_SHADER_STAGE_TESSELLATION_CONTROL_BIT, ShaderStage::Hull, HSHeader, UBGatherInfo);
-		DescriptorSetLayoutInfo.ProcessBindingsForStage(VK_SHADER_STAGE_TESSELLATION_EVALUATION_BIT, ShaderStage::Domain, DSHeader, UBGatherInfo);
-	}
-#endif
-	// Second pass
-	const int32 NumImmutableSamplers = PSOInitializer.ImmutableSamplerState.ImmutableSamplers.Num();
-	TArrayView<FRHISamplerState*> ImmutableSamplers(NumImmutableSamplers > 0 ? &(FRHISamplerState*&)PSOInitializer.ImmutableSamplerState.ImmutableSamplers[0] : nullptr, NumImmutableSamplers);
-	DescriptorSetLayoutInfo.FinalizeBindings<false>(UBGatherInfo, ImmutableSamplers);
-
-	FVulkanLayout* Layout = FindOrAddLayout(DescriptorSetLayoutInfo, true);
-	FVulkanGfxLayout* GfxLayout = (FVulkanGfxLayout*)Layout;
-	if (!GfxLayout->GfxPipelineDescriptorInfo.IsInitialized())
-	{
-		GfxLayout->GfxPipelineDescriptorInfo.Initialize(Layout->GetDescriptorSetsLayout().RemappingInfo);
-	}
 
-	return GfxLayout;
+	LayoutMap.Add(DescriptorSetLayoutInfo, Layout);
+	return Layout;
 }
 
 static inline VkPrimitiveTopology UEToVulkanTopologyType(const FVulkanDevice* InDevice, EPrimitiveType PrimitiveType, bool bHasTessellation, uint16& OutControlPoints)
@@ -1567,32 +1426,74 @@ static inline VkPrimitiveTopology UEToVulkanTopologyType(const FVulkanDevice* In
 	return VK_PRIMITIVE_TOPOLOGY_MAX_ENUM;
 }
 
-FVulkanPipelineStateCacheManager::FGfxPipelineEntry* FVulkanPipelineStateCacheManager::CreateGfxEntry(const FGraphicsPipelineStateInitializer& PSOInitializer)
+void FVulkanPipelineStateCacheManager::CreateGfxEntry(const FGraphicsPipelineStateInitializer& PSOInitializer, FVulkanDescriptorSetsLayoutInfo& DescriptorSetLayoutInfo, FGfxPipelineDesc* Desc)
 {
-	FGfxPipelineEntry* OutGfxEntry = new FGfxPipelineEntry();
+	FGfxPipelineDesc* OutGfxEntry = Desc;
 
 	FVulkanShader* Shaders[ShaderStage::NumStages];
 	GetVulkanShaders(PSOInitializer.BoundShaderState, Shaders);
 
 	FVulkanVertexInputStateInfo VertexInputState;
-	OutGfxEntry->Layout = GetOrGenerateGfxLayout(PSOInitializer, Shaders, VertexInputState);
-	OutGfxEntry->RenderPass = Device->GetImmediateContext().PrepareRenderPassForPSOCreation(PSOInitializer);
 	
-	// check that any depth fetch is actually using depth read sub-pass
-	if (OutGfxEntry->Layout->UsesInputAttachment(FVulkanShaderHeader::EAttachmentType::Depth))
 	{
-		check(PSOInitializer.SubpassHint == ESubpassHint::DepthReadSubpass);
-		check(PSOInitializer.SubpassIndex == 1);
+
+			const FBoundShaderStateInput& BSI = PSOInitializer.BoundShaderState;
+
+			const FVulkanShaderHeader& VSHeader = Shaders[ShaderStage::Vertex]->GetCodeHeader();
+			VertexInputState.Generate(ResourceCast(PSOInitializer.BoundShaderState.VertexDeclarationRHI), VSHeader.InOutMask);
+
+			FUniformBufferGatherInfo UBGatherInfo;
+
+			DescriptorSetLayoutInfo.ProcessBindingsForStage(VK_SHADER_STAGE_VERTEX_BIT, ShaderStage::Vertex, VSHeader, UBGatherInfo);
+
+			if (Shaders[ShaderStage::Pixel])
+			{
+				const FVulkanShaderHeader& PSHeader = Shaders[ShaderStage::Pixel]->GetCodeHeader();
+				DescriptorSetLayoutInfo.ProcessBindingsForStage(VK_SHADER_STAGE_FRAGMENT_BIT, ShaderStage::Pixel, PSHeader, UBGatherInfo);
+			}
+
+#if VULKAN_SUPPORTS_GEOMETRY_SHADERS
+			if (Shaders[ShaderStage::Geometry])
+			{
+				const FVulkanShaderHeader& GSHeader = Shaders[ShaderStage::Geometry]->GetCodeHeader();
+				DescriptorSetLayoutInfo.ProcessBindingsForStage(VK_SHADER_STAGE_GEOMETRY_BIT, ShaderStage::Geometry, GSHeader, UBGatherInfo);
+			}
+#endif
+
+#if PLATFORM_SUPPORTS_TESSELLATION_SHADERS
+			if (Shaders[ShaderStage::Hull])
+			{
+				const FVulkanShaderHeader& HSHeader = Shaders[ShaderStage::Hull]->GetCodeHeader();
+				const FVulkanShaderHeader& DSHeader = Shaders[ShaderStage::Domain]->GetCodeHeader();
+				DescriptorSetLayoutInfo.ProcessBindingsForStage(VK_SHADER_STAGE_TESSELLATION_CONTROL_BIT, ShaderStage::Hull, HSHeader, UBGatherInfo);
+				DescriptorSetLayoutInfo.ProcessBindingsForStage(VK_SHADER_STAGE_TESSELLATION_EVALUATION_BIT, ShaderStage::Domain, DSHeader, UBGatherInfo);
+			}
+#endif
+			// Second pass
+			const int32 NumImmutableSamplers = PSOInitializer.ImmutableSamplerState.ImmutableSamplers.Num();
+			TArrayView<FRHISamplerState*> ImmutableSamplers(NumImmutableSamplers > 0 ? &(FRHISamplerState*&)PSOInitializer.ImmutableSamplerState.ImmutableSamplers[0] : nullptr, NumImmutableSamplers);
+			DescriptorSetLayoutInfo.FinalizeBindings<false>(UBGatherInfo, ImmutableSamplers);
+	}
+
+	FDescriptorSetRemappingInfo& RemappingInfo = DescriptorSetLayoutInfo.RemappingInfo;
+
+	for(const FInputAttachmentData& Input : RemappingInfo.InputAttachmentData)
+	{
+		// check that any depth fetch is actually using depth read sub-pass
+		if(Input.Type == FVulkanShaderHeader::EAttachmentType::Depth)
+		{
+			check(PSOInitializer.SubpassHint == ESubpassHint::DepthReadSubpass);
+			check(PSOInitializer.SubpassIndex == 1);
+		}
 	}
 	OutGfxEntry->SubpassIndex = PSOInitializer.SubpassIndex;
 
 	const bool bHasTessellation = (PSOInitializer.BoundShaderState.DomainShaderRHI != nullptr);
 
-	OutGfxEntry->RasterizationSamples = OutGfxEntry->RenderPass->GetLayout().GetAttachmentDescriptions()[0].samples;
-	ensure(OutGfxEntry->RasterizationSamples == PSOInitializer.NumSamples);
+	OutGfxEntry->RasterizationSamples = PSOInitializer.NumSamples;
 	OutGfxEntry->Topology = (uint32)UEToVulkanTopologyType(Device, PSOInitializer.PrimitiveType, bHasTessellation, OutGfxEntry->ControlPoints);
-
-	OutGfxEntry->ColorAttachmentStates.AddUninitialized(OutGfxEntry->RenderPass->GetLayout().GetNumColorAttachments());
+	uint32 NumRenderTargets = PSOInitializer.ComputeNumValidRenderTargets();
+	OutGfxEntry->ColorAttachmentStates.AddUninitialized(NumRenderTargets);
 	for (int32 Index = 0; Index < OutGfxEntry->ColorAttachmentStates.Num(); ++Index)
 	{
 		OutGfxEntry->ColorAttachmentStates[Index].ReadFrom(ResourceCast(PSOInitializer.BlendState)->BlendStates[Index]);
@@ -1613,7 +1514,7 @@ FVulkanPipelineStateCacheManager::FGfxPipelineEntry* FVulkanPipelineStateCacheMa
 		}
 	}
 
-	const TArray<FVulkanDescriptorSetsLayout::FSetLayout>& Layouts = OutGfxEntry->Layout->GetDescriptorSetsLayout().GetLayouts();
+	const TArray<FVulkanDescriptorSetsLayout::FSetLayout>& Layouts = DescriptorSetLayoutInfo.GetLayouts();
 	OutGfxEntry->DescriptorSetLayoutBindings.AddDefaulted(Layouts.Num());
 	for (int32 Index = 0; Index < Layouts.Num(); ++Index)
 	{
@@ -1632,200 +1533,263 @@ FVulkanPipelineStateCacheManager::FGfxPipelineEntry* FVulkanPipelineStateCacheMa
 	}
 
 	int32 NumShaders = 0;
+#if VULKAN_USE_SHADERKEYS
+	uint64 SharedKey = 0;
+	uint64 Primes[] = {
+		6843488303525203279llu,
+		3095754086865563867llu,
+		8242695776924673527llu,
+		7556751872809527943llu,
+		8278265491465149053llu,
+	};
+	check(sizeof(Primes) / sizeof(Primes[0]) >= ShaderStage::NumStages);
 	for (int32 Index = 0; Index < ShaderStage::NumStages; ++Index)
 	{
 		FVulkanShader* Shader = Shaders[Index];
+		uint64 Key = 0;
 		if (Shader)
 		{
-			check(Shader->Spirv.Num() != 0);
-
-			FSHAHash Hash = GetShaderHashForStage(PSOInitializer, (ShaderStage::EStage)Index);
-			OutGfxEntry->ShaderHashes.Stages[Index] = Hash;
-
+			Key = Shader->GetShaderKey();
 			++NumShaders;
 		}
+		OutGfxEntry->ShaderKeys[Index] = Key;
+		SharedKey += Key * Primes[Index];
 	}
-	check(NumShaders > 0);
-
-#if VULKAN_SUPPORTS_COLOR_CONVERSIONS
+	OutGfxEntry->ShaderKeyShared = SharedKey;
+#else
+	for (int32 Index = 0; Index < ShaderStage::NumStages; ++Index)
+	{
+		FVulkanShader* Shader = Shaders[Index];
+		if (Shader)
+		{
+			check(Shader->Spirv.Num() != 0);
+
+			FSHAHash Hash = GetShaderHashForStage(PSOInitializer, (ShaderStage::EStage)Index);
+			OutGfxEntry->ShaderHashes.Stages[Index] = Hash;
+
+			++NumShaders;
+		}
+	}
+	OutGfxEntry->ShaderHashes.Finalize();
+#endif
+	check(NumShaders > 0);
+
+#if VULKAN_SUPPORTS_COLOR_CONVERSIONS
 	for (uint32 Index = 0; Index < MaxImmutableSamplers; ++Index)
 	{
 		OutGfxEntry->ImmutableSamplers[Index] = reinterpret_cast<SIZE_T>(PSOInitializer.ImmutableSamplerState.ImmutableSamplers[Index]);
 	}
 #endif
 
-	OutGfxEntry->RenderTargets.ReadFrom(OutGfxEntry->RenderPass->GetLayout());
 
-	OutGfxEntry->ShaderHashes.Finalize();
-	
-	return OutGfxEntry;
+	FVulkanRenderTargetLayout RTLayout(PSOInitializer);
+	OutGfxEntry->RenderTargets.ReadFrom(RTLayout);
 }
 
 
-FVulkanRHIGraphicsPipelineState* FVulkanPipelineStateCacheManager::FindInLoadedLibrary(const FGraphicsPipelineStateInitializer& PSOInitializer, FGfxPSIKey& PSIKey, TGfxPipelineEntrySharedPtr& OutGfxEntry, FGfxEntryKey& OutGfxEntryKey)
+
+
+
+FVulkanRHIGraphicsPipelineState::FVulkanRHIGraphicsPipelineState(FVulkanDevice* Device, const FGraphicsPipelineStateInitializer& PSOInitializer_, FGfxPipelineDesc& Desc, FVulkanPSOKey* VulkanKey)
+	: bIsRegistered(false)
+	, PrimitiveType(PSOInitializer_.PrimitiveType)
+	, VulkanPipeline(0)
+	, Device(Device)
+	, Desc(Desc)
+	, VulkanKey(VulkanKey->CopyDeep())
 {
-	OutGfxEntry = nullptr;
-	
-	FGfxPipelineEntry* GfxEntry = CreateGfxEntry(PSOInitializer);
-	FGfxEntryKey GfxEntryKey = GfxEntry->CreateKey();
+#if !UE_BUILD_SHIPPING
+	SGraphicsRHICount++;
+#endif
+#if VULKAN_PSO_CACHE_DEBUG
+	PixelShaderRHI = PSOInitializer_.BoundShaderState.PixelShaderRHI;
+	VertexShaderRHI = PSOInitializer_.BoundShaderState.VertexShaderRHI;
+	VertexDeclarationRHI = PSOInitializer_.BoundShaderState.VertexDeclarationRHI;
 
-	FVulkanGfxPipeline** FoundPipeline = EntryKeyToGfxPipelineMap.Find(GfxEntryKey);
-	if (FoundPipeline)
+#if PLATFORM_SUPPORTS_TESSELLATION_SHADERS
+	DomainShaderRHI = PSOInitializer_.BoundShaderState.DomainShaderRHI;
+	HullShaderRHI = PSOInitializer_.BoundShaderState.HullShaderRHI;
+#endif 
+
+#if PLATFORM_SUPPORTS_GEOMETRY_SHADERS
+	GeometryShaderRHI = PSOInitializer_.BoundShaderState.GeometryShaderRHI;
+#endif
+
+	PSOInitializer = PSOInitializer_;
+#endif
+	FMemory::Memset(ShaderModules, 0, sizeof(ShaderModules));
+	INC_DWORD_STAT(STAT_VulkanNumGraphicsPSOs);
+	INC_DWORD_STAT_BY(STAT_VulkanPSOKeyMemory, this->VulkanKey.GetDataRef().Num());
+}
+
+void FVulkanPipelineStateCacheManager::NotifyDeletedGraphicsPSO(FRHIGraphicsPipelineState* PSO)
+{
+	FVulkanRHIGraphicsPipelineState* VkPSO = (FVulkanRHIGraphicsPipelineState*)PSO;
+	Device->NotifyDeletedGfxPipeline(VkPSO);
+	FVulkanPSOKey& Key = VkPSO->VulkanKey;
+	DEC_DWORD_STAT_BY(STAT_VulkanPSOKeyMemory, Key.GetDataRef().Num());
+	if(VkPSO->bIsRegistered)
 	{
-		if (!(*FoundPipeline)->IsRuntimeInitialized())
+		FScopeLock Lock(&GraphicsPSOLockedCS);
+		FVulkanRHIGraphicsPipelineState** Contained = GraphicsPSOLockedMap.Find(Key);
+		check(Contained && *Contained == PSO);
+		GraphicsPSOLockedMap.Remove(Key);
+		VkPSO->bIsRegistered = false;
+		if(bUseLRU)
 		{
-			(*FoundPipeline)->CreateRuntimeObjects(PSOInitializer);
+			LRURemove(*Contained);
+			check((*Contained)->LRUNode == 0);
 		}
-		FVulkanRHIGraphicsPipelineState* PipelineState = nullptr;
+		else
 		{
-			FScopeLock Lock2(&InitializerToPipelineMapCS);
-			FVulkanRHIGraphicsPipelineState** Existing = InitializerToPipelineMap.Find(PSIKey);
-			if (Existing)
-			{
-				PipelineState = *Existing;
-			}
-			else
-			{
-				PipelineState = new FVulkanRHIGraphicsPipelineState(PSOInitializer.BoundShaderState, *FoundPipeline, PSOInitializer.PrimitiveType);
-				PipelineState->AddRef();
-				InitializerToPipelineMap.Add(MoveTemp(PSIKey), PipelineState);
-			}
-		}		
-		delete GfxEntry;
-#if VULKAN_ENABLE_LRU_CACHE
-		PipelineLRU.Touch(PipelineState);
-#endif
-		return PipelineState;
+			(*Contained)->DeleteVkPipeline(true);
+			check(VkPSO->GetVulkanPipeline() == 0 );
+		}
+		
+	}
+	else
+	{
+		FScopeLock Lock(&GraphicsPSOLockedCS);
+		FVulkanRHIGraphicsPipelineState** Contained = GraphicsPSOLockedMap.Find(Key);
+		if (Contained && *Contained == VkPSO)
+		{
+			check(0);
+		}
 	}
-
-	OutGfxEntry = TGfxPipelineEntrySharedPtr(GfxEntry);
-	OutGfxEntryKey = MoveTemp(GfxEntryKey);
-	return nullptr;
 }
 
-FVulkanRHIGraphicsPipelineState* FVulkanPipelineStateCacheManager::FindInRuntimeCache(const FGraphicsPipelineStateInitializer& Initializer, FGfxPSIKey& OutKey)
-{
-	OutKey.GenerateFromArchive([&Initializer](FArchive& Ar)
-		{
-			FGraphicsPipelineStateInitializer& PSI = const_cast<FGraphicsPipelineStateInitializer&>(Initializer);
 
-			uint64 TempUInt64 = 0;
-			int32 TempEnumValue = 0;
+FVulkanRHIGraphicsPipelineState* FVulkanPipelineStateCacheManager::RHICreateGraphicsPipelineState(const FGraphicsPipelineStateInitializer& Initializer)
+{
+	QUICK_SCOPE_CYCLE_COUNTER(STAT_Vulkan_RHICreateGraphicsPipelineState_NEW);
+	FVulkanPSOKey Key;
+	FGfxPipelineDesc Desc;
+	FVulkanDescriptorSetsLayoutInfo DescriptorSetLayoutInfo;
+	{
 
-			TempUInt64 = GetShaderKey(PSI.BoundShaderState.VertexShaderRHI);
-			Ar << TempUInt64;
+		SCOPE_CYCLE_COUNTER(STAT_VulkanPSOHeaderInitTime);
+		CreateGfxEntry(Initializer, DescriptorSetLayoutInfo, &Desc);
+		Key = Desc.CreateKey2();
+	}
 
-			TempUInt64 = GetShaderKey(PSI.BoundShaderState.PixelShaderRHI);
-			Ar << TempUInt64;
 
-	#if VULKAN_SUPPORTS_GEOMETRY_SHADERS
-			TempUInt64 = GetShaderKey(PSI.BoundShaderState.GeometryShaderRHI);
-			Ar << TempUInt64;
-	#endif
-	
-	#if PLATFORM_SUPPORTS_TESSELLATION_SHADERS
-			TempUInt64 = GetShaderKey(PSI.BoundShaderState.HullShaderRHI);
-			Ar << TempUInt64;
-			TempUInt64 = GetShaderKey(PSI.BoundShaderState.DomainShaderRHI);
-			Ar << TempUInt64;
-	#endif
-			Ar << ResourceCast(PSI.BoundShaderState.VertexDeclarationRHI)->Elements;
-			Ar << ResourceCast(PSI.RasterizerState)->Initializer;
-			Ar << ResourceCast(PSI.DepthStencilState)->Initializer;
+	FVulkanRHIGraphicsPipelineState* NewPSO = 0;
+	{
+		SCOPE_CYCLE_COUNTER(STAT_VulkanPSOLookupTime);
+		FScopeLock Lock(&GraphicsPSOLockedCS);
+		{
+			FVulkanRHIGraphicsPipelineState** PSO = GraphicsPSOLockedMap.Find(Key);
+			if(PSO)
 			{
-				FBlendStateInitializerRHI &BlendState = ResourceCast(PSI.BlendState)->Initializer;
-				for (uint32 Index = 0; Index < PSI.RenderTargetsEnabled; ++Index)
-				{
-					Ar << BlendState.RenderTargets[Index];
-					EPixelFormat PixelFormat = (EPixelFormat)PSI.RenderTargetFormats[Index];
-					bool bSRGB = (PSI.RenderTargetFlags[Index] & TexCreate_SRGB) == TexCreate_SRGB;
-					TempEnumValue = UEToVkTextureFormat(PixelFormat, bSRGB);
-					Ar << TempEnumValue;
-					Ar << PSI.RenderTargetFlags[Index];
-				}
-				Ar << BlendState.bUseIndependentRenderTargetBlendStates;
+				check(*PSO);
+				LRUTouch(*PSO);
+				return *PSO;
 			}
+		}
+	}
 
-			// TODO To avoid changing of global header. Add getter when submitting
-			Ar.Serialize(&PSI.DepthStencilAccess, sizeof(PSI.DepthStencilAccess));
 
-			TempEnumValue = PSI.DepthStencilTargetFormat;
-			Ar << TempEnumValue;
 
-			TempEnumValue = PSI.PrimitiveType;
-			Ar << TempEnumValue;
+	{
+		SCOPE_CYCLE_COUNTER(STAT_VulkanPSOCreationTime);
+		NewPSO = new FVulkanRHIGraphicsPipelineState(Device, Initializer, Desc, &Key);
+		{
 
-#if VULKAN_SUPPORTS_COLOR_CONVERSIONS
-			for (uint32 Index = 0; Index < MaxImmutableSamplers; ++Index)
+			FVulkanLayout* Layout = FindOrAddLayout(DescriptorSetLayoutInfo, true);
+			FVulkanGfxLayout* GfxLayout = (FVulkanGfxLayout*)Layout;
+			check(GfxLayout->GfxPipelineDescriptorInfo.IsInitialized());
+			NewPSO->Layout = GfxLayout;
+		}
+		NewPSO->RenderPass = Device->GetImmediateContext().PrepareRenderPassForPSOCreation(Initializer);
+		{
+			const FBoundShaderStateInput& BSI = Initializer.BoundShaderState;
+			for (int32 StageIdx = 0; StageIdx < ShaderStage::NumStages; ++StageIdx)
 			{
-				FRHISamplerState* SamplerState = PSI.ImmutableSamplerState.ImmutableSamplers[Index];
-				VkSampler Handle = SamplerState ? ResourceCast(SamplerState)->Sampler : VK_NULL_HANDLE;
-				TempUInt64 = (uint64)Handle;
-				Ar << TempUInt64;
+				NewPSO->ShaderKeys[StageIdx] = GetShaderKeyForGfxStage(BSI, (ShaderStage::EStage)StageIdx);
 			}
-#endif
 
-			Ar << PSI.bDepthBounds;
-			Ar << PSI.RenderTargetsEnabled;
-			Ar << PSI.NumSamples;
+			check(BSI.VertexShaderRHI);
+			FVulkanVertexShader* VS = ResourceCast(BSI.VertexShaderRHI);
+			const FVulkanShaderHeader& VSHeader = VS->GetCodeHeader();
+			NewPSO->VertexInputState.Generate(ResourceCast(Initializer.BoundShaderState.VertexDeclarationRHI), VSHeader.InOutMask);
 
-		},
-		256);
+			if((!Initializer.bFromPSOFileCache || !LRUEvictImmediately()) 
+	#if !UE_BUILD_SHIPPING
+				&& 0 == CVarPipelineDebugForceEvictImmediately.GetValueOnAnyThread()
+	#endif
+				)
+			{
 
-	FVulkanRHIGraphicsPipelineState** Found = nullptr;
+				// Create the pipeline
+				double BeginTime = FPlatformTime::Seconds();
+				FVulkanShader* VulkanShaders[ShaderStage::NumStages];
+				GetVulkanShaders(Initializer.BoundShaderState, VulkanShaders);
 
-	{
-		FScopeLock Lock(&InitializerToPipelineMapCS);
-		Found = InitializerToPipelineMap.Find(OutKey);
-	}
+				for (int32 StageIdx = 0; StageIdx < ShaderStage::NumStages; ++StageIdx)
+				{
+					uint64 key = GetShaderKeyForGfxStage(BSI, (ShaderStage::EStage)StageIdx);
+					check(key == NewPSO->ShaderKeys[StageIdx]);
+				}
 
-	if (Found)
-	{
-#if VULKAN_ENABLE_LRU_CACHE
-		PipelineLRU.Touch(*Found);
-#endif
-		return *Found;
-	}
+			
+				QUICK_SCOPE_CYCLE_COUNTER(STAT_Vulkan_RHICreateGraphicsPipelineState_CREATE_PART0);
 
-	return nullptr;
+				if(!CreateGfxPipelineFromEntry(NewPSO, VulkanShaders, &NewPSO->VulkanPipeline))
+				{
+					delete NewPSO;
+					return nullptr;
+				}
+				// Recover if we failed to create the pipeline.
+				double EndTime = FPlatformTime::Seconds();
+				double Delta = EndTime - BeginTime;
+				if (Delta > HitchTime)
+				{
+					UE_LOG(LogVulkanRHI, Verbose, TEXT("Hitchy gfx pipeline (%.3f ms)"), (float)(Delta * 1000.0));
+				}
+			}
+			FScopeLock Lock(&GraphicsPSOLockedCS); 
+			FVulkanRHIGraphicsPipelineState** MapPSO = GraphicsPSOLockedMap.Find(Key);
+			if(MapPSO)//another thread could end up creating it.
+			{
+				delete NewPSO;
+				NewPSO = *MapPSO;
+			}
+			else
+			{
+				GraphicsPSOLockedMap.Add(MoveTemp(Key), NewPSO);
+				if (bUseLRU)
+				{
+					FScopeLock LockRU(&LRUCS);
+					NewPSO->bIsRegistered = true;
+					LRUTrim(NewPSO->PipelineCacheSize);
+					LRUAdd(NewPSO);
+				}
+				else
+				{
+					NewPSO->bIsRegistered = true;
+				}
+
+
+			}
+		}
+	}
+	return NewPSO;
 }
 
+
 FGraphicsPipelineStateRHIRef FVulkanDynamicRHI::RHICreateGraphicsPipelineState(const FGraphicsPipelineStateInitializer& PSOInitializer)
 {
 #if VULKAN_ENABLE_AGGRESSIVE_STATS
 	SCOPE_CYCLE_COUNTER(STAT_VulkanGetOrCreatePipeline);
 #endif
-
+	QUICK_SCOPE_CYCLE_COUNTER(STAT_Vulkan_RHICreateGraphicsPipelineState);
 	LLM_SCOPE_VULKAN(ELLMTagVulkan::VulkanShaders);
 
-	// First try the hash based off runtime objects
-	FGfxPSIKey PSIKey;
-	FVulkanRHIGraphicsPipelineState* Found = Device->PipelineStateCache->FindInRuntimeCache(PSOInitializer, PSIKey);
-	if (Found)
-	{
-		return Found;
-	}
-
-	// Now try the loaded cache from disk
-	FVulkanPipelineStateCacheManager::TGfxPipelineEntrySharedPtr GfxEntry;
-	FGfxEntryKey GfxEntryKey;
-
-	{
-		FScopeLock Lock(&EntryKeyToGfxPipelineMapCS);
-
-		Found = Device->PipelineStateCache->FindInLoadedLibrary(PSOInitializer, PSIKey, GfxEntry, GfxEntryKey);
-		if (Found)
-		{
-			return Found;
-		}
+	return Device->PipelineStateCache->RHICreateGraphicsPipelineState(PSOInitializer);
+}
 
-		UE_LOG(LogVulkanRHI, Verbose, TEXT("PSO not found in cache, compiling..."));
 
-		FVulkanRHIGraphicsPipelineState* PipelineState = Device->PipelineStateCache->CreateAndAdd(PSOInitializer, MoveTemp(PSIKey), MoveTemp(GfxEntry), MoveTemp(GfxEntryKey));
 
-		return PipelineState;
-	}
-}
 
 FVulkanComputePipeline* FVulkanPipelineStateCacheManager::GetOrCreateComputePipeline(FVulkanComputeShader* ComputeShader)
 {
@@ -1855,7 +1819,10 @@ FVulkanComputePipeline* FVulkanPipelineStateCacheManager::GetOrCreateComputePipe
 
 	{
 		FRWScopeLock ScopeLock(ComputePipelineLock, SLT_Write);
-		ComputePipelineEntries.FindOrAdd(Key) = ComputePipeline;
+		if(0 == ComputePipelineEntries.Find(Key))
+		{
+			ComputePipelineEntries.FindOrAdd(Key) = ComputePipeline;
+		}
 	}
 
 	return ComputePipeline;
@@ -1900,7 +1867,7 @@ FVulkanComputePipeline* FVulkanPipelineStateCacheManager::CreateComputePipelineF
 }
 
 template<typename T>
-inline void SerializeArray(FArchive& Ar, TArray<T*>& Array)
+static void SerializeArray(FArchive& Ar, TArray<T>& Array)
 {
 	int32 Num = Array.Num();
 	Ar << Num;
@@ -1909,118 +1876,18 @@ inline void SerializeArray(FArchive& Ar, TArray<T*>& Array)
 		Array.SetNum(Num);
 		for (int32 Index = 0; Index < Num; ++Index)
 		{
-			T* Entry = new T;
-			Array[Index] = Entry;
-			Ar << *Entry;
+			Ar << Array[Index];
 		}
 	}
 	else
 	{
 		for (int32 Index = 0; Index < Num; ++Index)
 		{
-			Ar << *(Array[Index]);
-		}
-	}
-}
-
-#if VULKAN_ENABLE_LRU_CACHE
-void FVulkanPipelineStateCacheManager::FVKPipelineLRU::DeleteVkPipeline(FVulkanGfxPipeline* GfxPipeline)
-{
-	GfxPipeline->DeleteVkPipeline(GfxPipeline->RecentFrame + NUM_BUFFERS < GFrameNumberRenderThread);
-}
-
-void FVulkanPipelineStateCacheManager::FVKPipelineLRU::EnsureVkPipelineAndAddToLRU(FVulkanRHIGraphicsPipelineState* Pipeline)
-{
-	FVulkanGfxPipeline* GfxPipeline = Pipeline->Pipeline;
-	FGfxPipelineEntry* GfxPipelineEntry = static_cast<FGfxPipelineEntry*>(GfxPipeline->GfxPipelineEntry.Get());
-	FVulkanPipelineStateCacheManager* PipelineStateCache = GfxPipeline->Device->GetPipelineStateCache();
-
-	{
-		// This CS is used in RHICreateGraphicsPipelineState() to guard pipeline creation from multiple threads
-		FScopeLock Lock(&EntryKeyToGfxPipelineMapCS);
-
-		if (!GfxPipeline->GetHandle())
-		{
-			FVulkanShader* VulkanShaders[ShaderStage::NumStages];
-			GetVulkanShaders(GfxPipeline->Device, *Pipeline, VulkanShaders);
-			PipelineStateCache->CreateGfxPipelineFromEntry(GfxPipelineEntry, VulkanShaders, GfxPipeline);
-
-			// Failed to create VkPipeline again
-			check(GfxPipeline->GetHandle());
-
-			AddToLRU(GfxPipeline);
+			Ar << Array[Index];
 		}
 	}
 }
 
-void FVulkanPipelineStateCacheManager::FVKPipelineLRU::AddToLRU(FVulkanGfxPipeline* GfxPipeline)
-{
-	FScopeLock Lock(&LRUCS);
-
-	while (LRUUsedPipelineSize + GfxPipeline->PipelineCacheSize >(uint32)CVarLRUMaxPipelineSize.GetValueOnAnyThread() || LRU.Num() == LRU.Max())
-	{
-		EvictFromLRU();
-	}
-
-	check(!LRU.Contains(GfxPipeline));
-	GfxPipeline->LRUNode = LRU.Add(GfxPipeline, GfxPipeline);
-	LRUUsedPipelineSize += GfxPipeline->PipelineCacheSize;
-}
-
-void FVulkanPipelineStateCacheManager::FVKPipelineLRU::EvictFromLRU()
-{
-	FVulkanGfxPipeline* LeastRecentPipeline = LRU.RemoveLeastRecent();
-	check(LeastRecentPipeline);
-	LeastRecentPipeline->LRUNode = FSetElementId();
-
-	LRUUsedPipelineSize -= LeastRecentPipeline->PipelineCacheSize;
-
-	DeleteVkPipeline(LeastRecentPipeline);
-}
-
-void FVulkanPipelineStateCacheManager::FVKPipelineLRU::Add(FVulkanGfxPipeline* GfxPipeline)
-{
-	if (!bUseLRU)
-	{
-		return;
-	}
-
-	check(!GfxPipeline->bTrackedByLRU);
-	GfxPipeline->bTrackedByLRU = true;
-
-	if (!GfxPipeline->GetHandle())
-	{
-		GfxPipeline->LRUNode = FSetElementId(); // Add as evicted
-		return;
-	}
-
-	AddToLRU(GfxPipeline);
-}
-
-void FVulkanPipelineStateCacheManager::FVKPipelineLRU::Touch(FVulkanRHIGraphicsPipelineState* Pipeline)
-{
-	if (!bUseLRU)
-	{
-		return;
-	}
-
-	FVulkanGfxPipeline* GfxPipeline = Pipeline->Pipeline;
-	check(GfxPipeline->bTrackedByLRU);
-
-	{
-		FScopeLock Lock(&LRUCS);
-		if (GfxPipeline->LRUNode.IsValidId())
-		{
-			check(GfxPipeline->GetHandle());
-			check(LRU.Contains(GfxPipeline));
-			LRU.MarkAsRecent(GfxPipeline->LRUNode);
-			GfxPipeline->RecentFrame = GFrameNumberRenderThread;
-			return;
-		}
-	}
-
-	EnsureVkPipelineAndAddToLRU(Pipeline);
-}
 
 void FVulkanPipelineStateCacheManager::FVulkanLRUCacheFile::Save(FArchive& Ar)
 {
@@ -2042,16 +1909,16 @@ bool FVulkanPipelineStateCacheManager::FVulkanLRUCacheFile::Load(FArchive& Ar)
 	}
 
 	Ar << Header.SizeOfPipelineSizes;
-	if (Header.SizeOfPipelineSizes != (int32)(sizeof(FPipelineSize)))
+	if (Header.SizeOfPipelineSizes != (int32)(sizeof(FVulkanPipelineSize)))
 	{
-		UE_LOG(LogVulkanRHI, Warning, TEXT("Unable to load lru pipeline cache due to mismatched size of FPipelineSize %d != %d; forgot to bump up LRU_CACHE_VERSION?"), Header.SizeOfPipelineSizes, (int32)sizeof(FPipelineSize));
+		UE_LOG(LogVulkanRHI, Warning, TEXT("Unable to load lru pipeline cache due to mismatched size of FVulkanPipelineSize %d != %d; forgot to bump up LRU_CACHE_VERSION?"), Header.SizeOfPipelineSizes, (int32)sizeof(FVulkanPipelineSize));
 		return false;
 	}
 
 	SerializeArray(Ar, PipelineSizes);
 	return true;
 }
-#endif
+
 
 
 void GetVulkanShaders(const FBoundShaderStateInput& BSI, FVulkanShader* OutShaders[ShaderStage::NumStages])
@@ -2097,3 +1964,233 @@ void GetVulkanShaders(FVulkanDevice* Device, const FVulkanRHIGraphicsPipelineSta
 	FMemory::Memzero(OutShaders, ShaderStage::NumStages * sizeof(*OutShaders));
 	Device->GetShaderFactory().LookupShaders(GfxPipelineState.ShaderKeys, OutShaders);
 }
+
+
+
+void FVulkanPipelineStateCacheManager::LRUDump()
+{
+#if !UE_BUILD_SHIPPING
+	uint32 tid = FPlatformTLS::GetCurrentThreadId();
+	LRUPRINT(TEXT("//***** LRU DUMP *****\\\\\n"));
+	FVulkanRHIGraphicsPipelineStateLRUNode* Node= LRU.GetHead();
+	uint32_t Size = 0;
+	uint32_t Index = 0;
+	while(Node)
+	{
+		FVulkanRHIGraphicsPipelineState* PSO = Node->GetValue();
+		Size += PSO->PipelineCacheSize;
+		LRUPRINT(TEXT("\t%08x PSO %p :: %d  :: %06d \\ %06d\n"), tid, PSO, PSO->LRUFrame, PSO->PipelineCacheSize, Size);
+		Node = Node->GetNextNode();
+		Index++;
+	}
+	LRUPRINT(TEXT("\\\\***** LRU DUMP *****//\n"));
+#endif
+}
+
+
+
+bool FVulkanPipelineStateCacheManager::LRUEvictImmediately()
+{
+	return bEvictImmediately && CVarEnableLRU.GetValueOnAnyThread() != 0;
+}
+
+
+void FVulkanPipelineStateCacheManager::LRUTrim(uint32 nSpaceNeeded)
+{
+	if(!bUseLRU)
+	{
+		return;
+	}
+	uint32 tid = FPlatformTLS::GetCurrentThreadId();
+	uint32 MaxSize = (uint32)CVarLRUMaxPipelineSize.GetValueOnAnyThread();
+	while (LRUUsedPipelineSize + nSpaceNeeded > MaxSize || LRUUsedPipelineCount > LRUUsedPipelineMax)
+	{
+		LRUPRINT_DEBUG(TEXT("%d EVICTING %d + %d > %d || %d > %d\n"), tid, LRUUsedPipelineSize , nSpaceNeeded, MaxSize ,LRUUsedPipelineCount ,LRUUsedPipelineMax);
+		LRUEvictOne();
+	}
+}
+
+void FVulkanPipelineStateCacheManager::LRUDebugEvictAll()
+{
+	check(bUseLRU);
+	FScopeLock Lock(&LRUCS);
+	int Count = 0;
+	while(LRUEvictOne(true))
+		Count++;
+
+	LRUPRINT_DEBUG(TEXT("Evicted %d\n"), Count);
+}
+
+void FVulkanPipelineStateCacheManager::LRUAdd(FVulkanRHIGraphicsPipelineState* PSO)
+{
+	if(!bUseLRU)
+	{
+		return;
+	}
+
+	FScopeLock Lock(&LRUCS);
+	check(PSO->LRUNode == 0);
+	check(PSO->GetVulkanPipeline());
+	uint32 MaxSize = (uint32)CVarLRUMaxPipelineSize.GetValueOnAnyThread();
+	uint32 PSOSize = PSO->PipelineCacheSize;
+
+	LRUUsedPipelineSize += PSOSize;
+	LRUUsedPipelineCount += 1;
+
+	SET_DWORD_STAT(STAT_VulkanNumPSOLRUSize, LRUUsedPipelineSize);
+	SET_DWORD_STAT(STAT_VulkanNumPSOLRU, LRUUsedPipelineCount);
+
+	check(LRUUsedPipelineSize <= MaxSize); //should always be trimmed before.
+	LRU.AddHead(PSO);
+	PSO->LRUNode = LRU.GetHead();
+	PSO->LRUFrame = GFrameNumberRenderThread;
+	LRUPRINT_DEBUG(TEXT("LRUADD %p .. Frame %d :: %d    VKPSO %08x\n"), PSO, PSO->LRUFrame, GFrameNumberRenderThread, PSO->GetVulkanPipeline());
+
+}
+
+void FVulkanPipelineStateCacheManager::LRUTouch(FVulkanRHIGraphicsPipelineState* PSO)
+{
+	if(!bUseLRU)
+	{
+		return;
+	}
+	FScopeLock Lock(&LRUCS);
+	check((PSO->GetVulkanPipeline() == 0) == (PSO->LRUNode == 0));
+	
+	if(PSO->LRUNode)
+	{
+		check(PSO->GetVulkanPipeline());
+		if(PSO->LRUNode != LRU.GetHead())
+		{
+			LRU.RemoveNode(PSO->LRUNode, false);
+			LRU.AddHead(PSO->LRUNode);
+		}
+		PSO->LRUFrame = GFrameNumberRenderThread;
+	}
+	else
+	{
+		PSO->LRUFrame = GFrameNumberRenderThread;
+		if(!PSO->GetVulkanPipeline())
+		{
+			// Create the pipeline
+			double BeginTime = FPlatformTime::Seconds();
+			FVulkanShader* VulkanShaders[ShaderStage::NumStages];
+
+			GetVulkanShaders(Device, *PSO, VulkanShaders);
+
+			QUICK_SCOPE_CYCLE_COUNTER(STAT_Vulkan_RHICreateGraphicsPipelineState_CREATE_PART0);
+
+			if (!CreateGfxPipelineFromEntry(PSO, VulkanShaders, &PSO->VulkanPipeline))
+			{
+				check(0);
+			}
+			double EndTime = FPlatformTime::Seconds();
+			double Delta = EndTime - BeginTime;
+			if (Delta > HitchTime)
+			{
+				UE_LOG(LogVulkanRHI, Verbose, TEXT("Hitchy gfx pipeline (%.3f ms)"), (float)(Delta * 1000.0));
+			}
+
+			if(bUseLRU)
+			{
+				LRUTrim(PSO->PipelineCacheSize);
+				LRUAdd(PSO);
+			}
+		}
+		else
+		{
+			check(PSO->LRUNode);
+		}
+	}
+}
+
+void FVulkanRHIGraphicsPipelineState::DeleteVkPipeline(bool bImmediate)
+{
+	if (VulkanPipeline != VK_NULL_HANDLE)
+	{
+		if (bImmediate)
+		{
+			VulkanRHI::vkDestroyPipeline(Device->GetInstanceHandle(), VulkanPipeline, VULKAN_CPU_ALLOCATOR);
+		}
+		else
+		{
+			Device->GetDeferredDeletionQueue().EnqueueResource(VulkanRHI::FDeferredDeletionQueue::EType::Pipeline, VulkanPipeline);
+		}
+		VulkanPipeline = VK_NULL_HANDLE;
+	}
+
+	Device->PipelineStateCache->LRUCheckNotInside(this);
+}
+void FVulkanPipelineStateCacheManager::LRUCheckNotInside(FVulkanRHIGraphicsPipelineState* PSO)
+{
+	FScopeLock Lock(&LRUCS);
+
+
+	FVulkanRHIGraphicsPipelineStateLRUNode* Node = LRU.GetHead();
+	uint32_t Size = 0;
+	uint32_t Index = 0;
+	while (Node)
+	{
+		FVulkanRHIGraphicsPipelineState* foo = Node->GetValue();
+		if (foo == PSO)
+		{
+			check(0 == foo->LRUNode);
+
+		}
+		check(foo != PSO);
+		Node = Node->GetNextNode();
+	}
+	check(0 == PSO->LRUNode);
+}
+
+void FVulkanPipelineStateCacheManager::LRURemove(FVulkanRHIGraphicsPipelineState* PSO)
+{
+	check(bUseLRU);
+	if (PSO->LRUNode != 0)
+	{
+		bool bImmediate = PSO->LRUFrame + 3 < GFrameNumberRenderThread;
+		LRU.RemoveNode(PSO->LRUNode);
+		PSO->LRUNode = 0;
+
+		LRUUsedPipelineSize -= PSO->PipelineCacheSize;
+		LRUUsedPipelineCount--;
+
+		PSO->DeleteVkPipeline(bImmediate);
+		SET_DWORD_STAT(STAT_VulkanNumPSOLRUSize, LRUUsedPipelineSize);
+		SET_DWORD_STAT(STAT_VulkanNumPSOLRU, LRUUsedPipelineCount);
+	}
+	else
+	{
+		check(0 == PSO->GetVulkanPipeline());
+	}
+}
+
+bool FVulkanPipelineStateCacheManager::LRUEvictOne(bool bOnlyOld)
+{
+	check(bUseLRU);
+	uint32 tid = FPlatformTLS::GetCurrentThreadId();
+	FVulkanRHIGraphicsPipelineStateLRUNode* Node = LRU.GetTail();
+	check(Node != 0);
+	TRefCountPtr<FVulkanRHIGraphicsPipelineState> PSO = Node->GetValue();
+
+	bool bImmediate = PSO->LRUFrame + 3 < GFrameNumberRenderThread;
+	if(bOnlyOld && !bImmediate)
+	{
+		return false;
+	}
+	check(PSO->LRUFrame != GFrameNumberRenderThread);
+
+	LRURemove(PSO);
+	return true;
+}
+
+void FVulkanPipelineStateCacheManager::LRURemoveAll()
+{
+	if (!bUseLRU)
+	{
+		return;
+	}
+	check(0);
+}
+
+
diff --git a/Engine/Source/Runtime/VulkanRHI/Private/VulkanPipeline.h b/Engine/Source/Runtime/VulkanRHI/Private/VulkanPipeline.h
--- a/Engine/Source/Runtime/VulkanRHI/Private/VulkanPipeline.h
+++ b/Engine/Source/Runtime/VulkanRHI/Private/VulkanPipeline.h
@@ -11,18 +11,18 @@
 #include "VulkanDescriptorSets.h"
 #include "ShaderPipelineCache.h"
 
-#if VULKAN_ENABLE_LRU_CACHE
-#include "PsoLruCache.h"
-extern TAutoConsoleVariable<int32> CVarLRUMaxPipelineSize;
-extern TAutoConsoleVariable<int32> CVarEnableLRU;
-#endif
 
+#define VULKAN_PSO_CACHE_DEBUG 0
+#define VULKAN_USE_SHADERKEYS 1
 class FVulkanDevice;
 class FVulkanRHIGraphicsPipelineState;
-class FVulkanGfxPipeline;
 
-class FGfxPSIKey : public TDataKey<FGfxPSIKey> {};
-class FGfxEntryKey : public TDataKey<FGfxEntryKey> {};
+
+using FVulkanRHIGraphicsPipelineStateLRU = TDoubleLinkedList<FVulkanRHIGraphicsPipelineState*>;
+using FVulkanRHIGraphicsPipelineStateLRUNode = FVulkanRHIGraphicsPipelineStateLRU::TDoubleLinkedListNode;
+
+
+class FVulkanPSOKey : public TDataKey<FVulkanPSOKey, true> {};
 
 template <typename T>
 static inline uint64 GetShaderKey(T* ShaderType)
@@ -59,508 +59,463 @@ inline uint64 GetShaderKeyForGfxStage(const FBoundShaderStateInput& BSI, ShaderS
 void GetVulkanShaders(const FBoundShaderStateInput& BSI, FVulkanShader* OutShaders[ShaderStage::NumStages]);
 void GetVulkanShaders(FVulkanDevice* Device, const FVulkanRHIGraphicsPipelineState& GfxPipelineState, FVulkanShader* OutShaders[ShaderStage::NumStages]);
 
-class FVulkanPipelineStateCacheManager
-{
-public:
-	FVulkanRHIGraphicsPipelineState* FindInRuntimeCache(const FGraphicsPipelineStateInitializer& Initializer, FGfxPSIKey& OutKey);
 
-	// Array of potential cache locations; first entries have highest priority. Only one cache file is loaded. If unsuccessful, tries next entry in the array.
-	void InitAndLoad(const TArray<FString>& CacheFilenames);
-	void Save(const FString& CacheFilename, bool bFromPSOFC = false);
 
-	FVulkanPipelineStateCacheManager(FVulkanDevice* InParent);
-	~FVulkanPipelineStateCacheManager();
+struct FVulkanShaderHashes
+{
+	uint32 Hash;
+	FSHAHash Stages[ShaderStage::NumStages];
 
-	void RebuildCache();
+	FVulkanShaderHashes();
+	FVulkanShaderHashes(const FGraphicsPipelineStateInitializer& PSOInitializer);
 
-	struct FDescriptorSetLayoutBinding
+	friend inline uint32 GetTypeHash(const FVulkanShaderHashes& Hashes)
 	{
-		uint32 Binding;
-		//uint16 DescriptorCount;
-		uint8 DescriptorType;
-		uint8 StageFlags;
+		return Hashes.Hash;
+	}
 
-		void ReadFrom(const VkDescriptorSetLayoutBinding& InState);
-		void WriteInto(VkDescriptorSetLayoutBinding& OutState) const;
+	inline void Finalize()
+	{
+		Hash = FCrc::MemCrc32(Stages, sizeof(Stages));
+	}
 
-		bool operator==(const FDescriptorSetLayoutBinding& In) const
+	friend inline bool operator == (const FVulkanShaderHashes& A, const FVulkanShaderHashes& B)
+	{
+		for (int32 Index = 0; Index < ShaderStage::NumStages; ++Index)
 		{
-			return  Binding == In.Binding &&
-				//DescriptorCount == In.DescriptorCount &&
-				DescriptorType == In.DescriptorType &&
-				StageFlags == In.StageFlags;
+			if (A.Stages[Index] != B.Stages[Index])
+			{
+				return false;
+			}
 		}
-	};
 
-	// Shader microcode is shared between pipeline entries so keep a cache around to prevent duplicated storage
-	struct FShaderUCodeCache
+		return true;
+	}
+};
+
+struct FVulkanPipelineSize
+{
+	uint64 ShaderHash;
+	uint32 PipelineSize;
+
+	FVulkanPipelineSize()
+		:ShaderHash(0), PipelineSize(0)
 	{
-		using TDataMap = TMap<FSHAHash, TArray<uint32>>;
-		TDataMap Data;
+	}
+};
 
-		TArray<uint32>* Add(const FSHAHash& Hash, const FVulkanShader* Shader)
-		{
-			check(Shader->Spirv.Num() != 0);
+struct FDescriptorSetLayoutBinding
+{
+	uint32 Binding;
+	uint8 DescriptorType;
+	uint8 StageFlags;
 
-			TArray<uint32>& Code = Data.Add(Hash);
-			Code = Shader->Spirv;
+	void ReadFrom(const VkDescriptorSetLayoutBinding& InState);
+	void WriteInto(VkDescriptorSetLayoutBinding& OutState) const;
 
-			return &Data[Hash];
-		}
+	bool operator==(const FDescriptorSetLayoutBinding& In) const
+	{
+		return  Binding == In.Binding &&
+			DescriptorType == In.DescriptorType &&
+			StageFlags == In.StageFlags;
+	}
+};
 
-		TArray<uint32>* Get(const FSHAHash& Hash)
+
+
+
+struct FGfxPipelineDesc
+{
+	FVulkanPSOKey CreateKey2() const;
+
+	uint32 VertexInputKey;
+	uint16 RasterizationSamples;
+	uint16 ControlPoints;
+	uint32 Topology;
+	struct FBlendAttachment
+	{
+		bool bBlend;
+		uint8 ColorBlendOp;
+		uint8 SrcColorBlendFactor;
+		uint8 DstColorBlendFactor;
+		uint8 AlphaBlendOp;
+		uint8 SrcAlphaBlendFactor;
+		uint8 DstAlphaBlendFactor;
+		uint8 ColorWriteMask;
+
+		void ReadFrom(const VkPipelineColorBlendAttachmentState& InState);
+		void WriteInto(VkPipelineColorBlendAttachmentState& OutState) const;
+
+		bool operator==(const FBlendAttachment& In) const
 		{
-			return Data.Find(Hash);
+			return bBlend == In.bBlend &&
+				ColorBlendOp == In.ColorBlendOp &&
+				SrcColorBlendFactor == In.SrcColorBlendFactor &&
+				DstColorBlendFactor == In.DstColorBlendFactor &&
+				AlphaBlendOp == In.AlphaBlendOp &&
+				SrcAlphaBlendFactor == In.SrcAlphaBlendFactor &&
+				DstAlphaBlendFactor == In.DstAlphaBlendFactor &&
+				ColorWriteMask == In.ColorWriteMask;
 		}
 	};
+	TArray<FBlendAttachment> ColorAttachmentStates;
 
-	struct FShaderHashes
+	TArray<TArray<FDescriptorSetLayoutBinding>> DescriptorSetLayoutBindings;
+
+	struct FVertexBinding
 	{
-		uint32 Hash;
-		FSHAHash Stages[ShaderStage::NumStages];
+		uint32 Stride;
+		uint16 Binding;
+		uint16 InputRate;
 
-		FShaderHashes();
-		FShaderHashes(const FGraphicsPipelineStateInitializer& PSOInitializer);
+		void ReadFrom(const VkVertexInputBindingDescription& InState);
+		void WriteInto(VkVertexInputBindingDescription& OutState) const;
 
-		friend inline uint32 GetTypeHash(const FShaderHashes& Hashes)
+		bool operator==(const FVertexBinding& In) const
 		{
-			return Hashes.Hash;
+			return Stride == In.Stride &&
+				Binding == In.Binding &&
+				InputRate == In.InputRate;
 		}
+	};
+	TArray<FVertexBinding> VertexBindings;
+	struct FVertexAttribute
+	{
+		uint32 Location;
+		uint32 Binding;
+		uint32 Format;
+		uint32 Offset;
 
-		inline void Finalize()
-		{
-			Hash = FCrc::MemCrc32(Stages, sizeof(Stages));
-		}
+		void ReadFrom(const VkVertexInputAttributeDescription& InState);
+		void WriteInto(VkVertexInputAttributeDescription& OutState) const;
 
-		friend inline bool operator == (const FShaderHashes& A, const FShaderHashes& B)
+		bool operator==(const FVertexAttribute& In) const
 		{
-			for (int32 Index = 0; Index < ShaderStage::NumStages; ++Index)
-			{
-				if (A.Stages[Index] != B.Stages[Index])
-				{
-					return false;
-				}
-			}
-
-			return true;
+			return Location == In.Location &&
+				Binding == In.Binding &&
+				Format == In.Format &&
+				Offset == In.Offset;
 		}
 	};
+	TArray<FVertexAttribute> VertexAttributes;
 
-	// Actual information required to recreate a pipeline when saving/loading from disk
-	struct FGfxPipelineEntry
+	struct FRasterizer
 	{
-		FGfxEntryKey CreateKey() const;
-
-		uint32 VertexInputKey;
-		bool bLoaded;
-
-		uint16 RasterizationSamples;
-		uint16 ControlPoints;
-		uint32 Topology;
-		struct FBlendAttachment
-		{
-			bool bBlend;
-			uint8 ColorBlendOp;
-			uint8 SrcColorBlendFactor;
-			uint8 DstColorBlendFactor;
-			uint8 AlphaBlendOp;
-			uint8 SrcAlphaBlendFactor;
-			uint8 DstAlphaBlendFactor;
-			uint8 ColorWriteMask;
-
-			void ReadFrom(const VkPipelineColorBlendAttachmentState& InState);
-			void WriteInto(VkPipelineColorBlendAttachmentState& OutState) const;
-
-			bool operator==(const FBlendAttachment& In) const
-			{
-				return bBlend == In.bBlend &&
-					ColorBlendOp == In.ColorBlendOp &&
-					SrcColorBlendFactor == In.SrcColorBlendFactor &&
-					DstColorBlendFactor == In.DstColorBlendFactor &&
-					AlphaBlendOp == In.AlphaBlendOp &&
-					SrcAlphaBlendFactor == In.SrcAlphaBlendFactor &&
-					DstAlphaBlendFactor == In.DstAlphaBlendFactor &&
-					ColorWriteMask == In.ColorWriteMask;
-			}
-		};
-		TArray<FBlendAttachment> ColorAttachmentStates;
+		uint8 PolygonMode;
+		uint8 CullMode;
+		float DepthBiasSlopeScale;
+		float DepthBiasConstantFactor;
 
-		TArray<TArray<FDescriptorSetLayoutBinding>> DescriptorSetLayoutBindings;
+		void ReadFrom(const VkPipelineRasterizationStateCreateInfo& InState);
+		void WriteInto(VkPipelineRasterizationStateCreateInfo& OutState) const;
 
-		struct FVertexBinding
+		bool operator==(const FRasterizer& In) const
 		{
-			uint32 Stride;
-			uint16 Binding;
-			uint16 InputRate;
-
-			void ReadFrom(const VkVertexInputBindingDescription& InState);
-			void WriteInto(VkVertexInputBindingDescription& OutState) const;
+			return PolygonMode == In.PolygonMode &&
+				CullMode == In.CullMode &&
+				DepthBiasSlopeScale == In.DepthBiasSlopeScale &&
+				DepthBiasConstantFactor == In.DepthBiasConstantFactor;
+		}
+	};
+	FRasterizer Rasterizer;
 
-			bool operator==(const FVertexBinding& In) const
-			{
-				return Stride == In.Stride &&
-					Binding == In.Binding &&
-					InputRate == In.InputRate;
-			}
-		};
-		TArray<FVertexBinding> VertexBindings;
-		struct FVertexAttribute
+	struct FDepthStencil
+	{
+		uint8 DepthCompareOp;
+		bool bDepthTestEnable;
+		bool bDepthWriteEnable;
+		bool bStencilTestEnable;
+		bool bDepthBoundsTestEnable;
+		uint8 FrontFailOp;
+		uint8 FrontPassOp;
+		uint8 FrontDepthFailOp;
+		uint8 FrontCompareOp;
+		uint32 FrontCompareMask;
+		uint32 FrontWriteMask;
+		uint32 FrontReference;
+		uint8 BackFailOp;
+		uint8 BackPassOp;
+		uint8 BackDepthFailOp;
+		uint8 BackCompareOp;
+		uint32 BackCompareMask;
+		uint32 BackWriteMask;
+		uint32 BackReference;
+
+		void ReadFrom(const VkPipelineDepthStencilStateCreateInfo& InState);
+		void WriteInto(VkPipelineDepthStencilStateCreateInfo& OutState) const;
+
+		bool operator==(const FDepthStencil& In) const
 		{
-			uint32 Location;
-			uint32 Binding;
-			uint32 Format;
-			uint32 Offset;
+			return DepthCompareOp == In.DepthCompareOp &&
+				bDepthTestEnable == In.bDepthTestEnable &&
+				bDepthWriteEnable == In.bDepthWriteEnable &&
+				bDepthBoundsTestEnable == In.bDepthBoundsTestEnable &&
+				bStencilTestEnable == In.bStencilTestEnable &&
+				FrontFailOp == In.FrontFailOp &&
+				FrontPassOp == In.FrontPassOp &&
+				FrontDepthFailOp == In.FrontDepthFailOp &&
+				FrontCompareOp == In.FrontCompareOp &&
+				FrontCompareMask == In.FrontCompareMask &&
+				FrontWriteMask == In.FrontWriteMask &&
+				FrontReference == In.FrontReference &&
+				BackFailOp == In.BackFailOp &&
+				BackPassOp == In.BackPassOp &&
+				BackDepthFailOp == In.BackDepthFailOp &&
+				BackCompareOp == In.BackCompareOp &&
+				BackCompareMask == In.BackCompareMask &&
+				BackWriteMask == In.BackWriteMask &&
+				BackReference == In.BackReference;
+		}
+	};
+	FDepthStencil DepthStencil;
 
-			void ReadFrom(const VkVertexInputAttributeDescription& InState);
-			void WriteInto(VkVertexInputAttributeDescription& OutState) const;
+#if VULKAN_USE_SHADERKEYS
+	uint64 ShaderKeys[ShaderStage::NumStages];
+	uint64 ShaderKeyShared;
+#else
+	FVulkanShaderHashes ShaderHashes;
+#endif
 
-			bool operator==(const FVertexAttribute& In) const
-			{
-				return Location == In.Location &&
-					Binding == In.Binding &&
-					Format == In.Format &&
-					Offset == In.Offset;
-			}
-		};
-		TArray<FVertexAttribute> VertexAttributes;
+#if VULKAN_SUPPORTS_COLOR_CONVERSIONS
+	SIZE_T ImmutableSamplers[MaxImmutableSamplers] = { 0 };
+#endif
 
-		struct FRasterizer
+	struct FRenderTargets
+	{
+		struct FAttachmentRef
 		{
-			uint8 PolygonMode;
-			uint8 CullMode;
-			float DepthBiasSlopeScale;
-			float DepthBiasConstantFactor;
-
-			void ReadFrom(const VkPipelineRasterizationStateCreateInfo& InState);
-			void WriteInto(VkPipelineRasterizationStateCreateInfo& OutState) const;
+			uint32 Attachment;
+			uint64 Layout;
 
-			bool operator==(const FRasterizer& In) const
+			void ReadFrom(const VkAttachmentReference& InState);
+			void WriteInto(VkAttachmentReference& OutState) const;
+			bool operator == (const FAttachmentRef& In) const
 			{
-				return PolygonMode == In.PolygonMode &&
-					CullMode == In.CullMode &&
-					DepthBiasSlopeScale == In.DepthBiasSlopeScale &&
-					DepthBiasConstantFactor == In.DepthBiasConstantFactor;
+				return Attachment == In.Attachment && Layout == In.Layout;
 			}
 		};
-		FRasterizer Rasterizer;
+		TArray<FAttachmentRef> ColorAttachments;
+		TArray<FAttachmentRef> ResolveAttachments;
+		FAttachmentRef DepthStencil;
+		FAttachmentRef FragmentDensity;
 
-		struct FDepthStencil
+		struct FAttachmentDesc
 		{
-			uint8 DepthCompareOp;
-			bool bDepthTestEnable;
-			bool bDepthWriteEnable;
-			bool bStencilTestEnable;
-			bool bDepthBoundsTestEnable;
-			uint8 FrontFailOp;
-			uint8 FrontPassOp;
-			uint8 FrontDepthFailOp;
-			uint8 FrontCompareOp;
-			uint32 FrontCompareMask;
-			uint32 FrontWriteMask;
-			uint32 FrontReference;
-			uint8 BackFailOp;
-			uint8 BackPassOp;
-			uint8 BackDepthFailOp;
-			uint8 BackCompareOp;
-			uint32 BackCompareMask;
-			uint32 BackWriteMask;
-			uint32 BackReference;
-
-			void ReadFrom(const VkPipelineDepthStencilStateCreateInfo& InState);
-			void WriteInto(VkPipelineDepthStencilStateCreateInfo& OutState) const;
-
-			bool operator==(const FDepthStencil& In) const
+			uint32 Format;
+			uint8 Flags;
+			uint8 Samples;
+			uint8 LoadOp;
+			uint8 StoreOp;
+			uint8 StencilLoadOp;
+			uint8 StencilStoreOp;
+			uint64 InitialLayout;
+			uint64 FinalLayout;
+
+			bool operator==(const FAttachmentDesc& In) const
 			{
-				return DepthCompareOp == In.DepthCompareOp &&
-					bDepthTestEnable == In.bDepthTestEnable &&
-					bDepthWriteEnable == In.bDepthWriteEnable &&
-					bDepthBoundsTestEnable == In.bDepthBoundsTestEnable &&
-					bStencilTestEnable == In.bStencilTestEnable &&
-					FrontFailOp == In.FrontFailOp &&
-					FrontPassOp == In.FrontPassOp &&
-					FrontDepthFailOp == In.FrontDepthFailOp &&
-					FrontCompareOp == In.FrontCompareOp &&
-					FrontCompareMask == In.FrontCompareMask &&
-					FrontWriteMask == In.FrontWriteMask &&
-					FrontReference == In.FrontReference &&
-					BackFailOp == In.BackFailOp &&
-					BackPassOp == In.BackPassOp &&
-					BackDepthFailOp == In.BackDepthFailOp &&
-					BackCompareOp == In.BackCompareOp &&
-					BackCompareMask == In.BackCompareMask &&
-					BackWriteMask == In.BackWriteMask &&
-					BackReference == In.BackReference;
+				return Format == In.Format &&
+					Flags == In.Flags &&
+					Samples == In.Samples &&
+					LoadOp == In.LoadOp &&
+					StoreOp == In.StoreOp &&
+					StencilLoadOp == In.StencilLoadOp &&
+					StencilStoreOp == In.StencilStoreOp &&
+					InitialLayout == In.InitialLayout &&
+					FinalLayout == In.FinalLayout;
 			}
-		};
-		FDepthStencil DepthStencil;
-
-		FShaderHashes ShaderHashes;
-#if VULKAN_SUPPORTS_COLOR_CONVERSIONS
-		SIZE_T ImmutableSamplers[MaxImmutableSamplers] = { 0 };
-#endif
 
-		struct FRenderTargets
-		{
-			struct FAttachmentRef
-			{
-				uint32 Attachment;
-				uint64 Layout;
-
-				void ReadFrom(const VkAttachmentReference& InState);
-				void WriteInto(VkAttachmentReference& OutState) const;
-				bool operator == (const FAttachmentRef& In) const
-				{
-					return Attachment == In.Attachment && Layout == In.Layout;
-				}
-			};
-			TArray<FAttachmentRef> ColorAttachments;
-			TArray<FAttachmentRef> ResolveAttachments;
-			FAttachmentRef DepthStencil;
-			FAttachmentRef FragmentDensity;
-
-			struct FAttachmentDesc
-			{
-				uint32 Format;
-				uint8 Flags;
-				uint8 Samples;
-				uint8 LoadOp;
-				uint8 StoreOp;
-				uint8 StencilLoadOp;
-				uint8 StencilStoreOp;
-				uint64 InitialLayout;
-				uint64 FinalLayout;
-
-				bool operator==(const FAttachmentDesc& In) const
-				{
-					return Format == In.Format &&
-						Flags == In.Flags &&
-						Samples == In.Samples &&
-						LoadOp == In.LoadOp &&
-						StoreOp == In.StoreOp &&
-						StencilLoadOp == In.StencilLoadOp &&
-						StencilStoreOp == In.StencilStoreOp &&
-						InitialLayout == In.InitialLayout &&
-						FinalLayout == In.FinalLayout;
-				}
-
-				void ReadFrom(const VkAttachmentDescription &InState);
-				void WriteInto(VkAttachmentDescription& OutState) const;
-			};
-			TArray<FAttachmentDesc> Descriptions;
-
-			uint8 NumAttachments;
-			uint8 NumColorAttachments;
-			uint8 bHasDepthStencil;
-			uint8 bHasResolveAttachments;
-			uint8 bHasFragmentDensityAttachment;
-			uint8 NumUsedClearValues;
-			uint32 RenderPassCompatibleHash;
-			FVector Extent3D;
-
-			void ReadFrom(const FVulkanRenderTargetLayout &InState);
-			void WriteInto(FVulkanRenderTargetLayout& OutState) const;
-
-			bool operator==(const FRenderTargets& In) const
-			{
-				return ColorAttachments == In.ColorAttachments &&
-					ResolveAttachments == In.ResolveAttachments &&
-					DepthStencil == In.DepthStencil &&
-					FragmentDensity == In.FragmentDensity &&
-					Descriptions == In.Descriptions &&
-					NumAttachments == In.NumAttachments &&
-					NumColorAttachments == In.NumColorAttachments &&
-					bHasDepthStencil == In.bHasDepthStencil &&
-					bHasResolveAttachments == In.bHasResolveAttachments &&
-					bHasFragmentDensityAttachment == In.bHasFragmentDensityAttachment &&
-					NumUsedClearValues == In.NumUsedClearValues &&
-					RenderPassCompatibleHash == In.RenderPassCompatibleHash &&
-					Extent3D == In.Extent3D;
-			}
+			void ReadFrom(const VkAttachmentDescription &InState);
+			void WriteInto(VkAttachmentDescription& OutState) const;
 		};
-		FRenderTargets RenderTargets;
-
-		FGfxPipelineEntry()
-			: VertexInputKey(0)
-			, bLoaded(false)
-			, RasterizationSamples(0)
-			, Topology(0)
-			, RenderPass(nullptr)
-			, Layout(nullptr)
-			, SubpassIndex(0)
-		{
-			FMemory::Memzero(Rasterizer);
-			FMemory::Memzero(DepthStencil);
-			FMemory::Memzero(ShaderModules);
-		}
-
-		~FGfxPipelineEntry();
+		TArray<FAttachmentDesc> Descriptions;
 
-		VkShaderModule ShaderModules[ShaderStage::NumStages];
-		const FVulkanRenderPass* RenderPass;
-		FVulkanGfxLayout* Layout;
-		uint8 SubpassIndex;
+		uint8 NumAttachments;
+		uint8 NumColorAttachments;
+		uint8 bHasDepthStencil;
+		uint8 bHasResolveAttachments;
+		uint8 bHasFragmentDensityAttachment;
+		uint8 NumUsedClearValues;
+		uint32 RenderPassCompatibleHash;
+		FVector Extent3D;
 
-		void GetOrCreateShaderModules(FVulkanShader*const* Shaders);
-		void PurgeShaderModules(FVulkanShader*const* Shaders);
-		void PurgeLoadedShaderModules(FVulkanDevice* InDevice);
+		void ReadFrom(const FVulkanRenderTargetLayout &InState);
+		void WriteInto(FVulkanRenderTargetLayout& OutState) const;
 
-		bool operator==(const FGfxPipelineEntry& In) const
+		bool operator==(const FRenderTargets& In) const
 		{
-			if (VertexInputKey != In.VertexInputKey)
-			{
-				return false;
-			}
-
-			if( bLoaded != In.bLoaded)
-			{
-				return false;
-			}
+			return ColorAttachments == In.ColorAttachments &&
+				ResolveAttachments == In.ResolveAttachments &&
+				DepthStencil == In.DepthStencil &&
+				FragmentDensity == In.FragmentDensity &&
+				Descriptions == In.Descriptions &&
+				NumAttachments == In.NumAttachments &&
+				NumColorAttachments == In.NumColorAttachments &&
+				bHasDepthStencil == In.bHasDepthStencil &&
+				bHasResolveAttachments == In.bHasResolveAttachments &&
+				bHasFragmentDensityAttachment == In.bHasFragmentDensityAttachment &&
+				NumUsedClearValues == In.NumUsedClearValues &&
+				RenderPassCompatibleHash == In.RenderPassCompatibleHash &&
+				Extent3D == In.Extent3D;
+		}
+	};
+	FRenderTargets RenderTargets;
 
-			if (RasterizationSamples != In.RasterizationSamples)
-			{
-				return false;
-			}
+	uint8 SubpassIndex;
 
-			if (Topology != In.Topology)
-			{
-				return false;
-			}
+	bool operator==(const FGfxPipelineDesc& In) const
+	{
+		if (VertexInputKey != In.VertexInputKey)
+		{
+			return false;
+		}
 
-			if (ColorAttachmentStates != In.ColorAttachmentStates)
-			{
-				return false;
-			}
+		if (RasterizationSamples != In.RasterizationSamples)
+		{
+			return false;
+		}
 
-			if (DescriptorSetLayoutBindings != In.DescriptorSetLayoutBindings)
-			{
-				return false;
-			}
+		if (Topology != In.Topology)
+		{
+			return false;
+		}
 
-			if (!(Rasterizer == In.Rasterizer))
-			{
-				return false;
-			}
+		if (ColorAttachmentStates != In.ColorAttachmentStates)
+		{
+			return false;
+		}
 
-			if (!(DepthStencil == In.DepthStencil))
-			{
-				return false;
-			}
+		if (DescriptorSetLayoutBindings != In.DescriptorSetLayoutBindings)
+		{
+			return false;
+		}
 
-			if (!(SubpassIndex == In.SubpassIndex))
-			{
-				return false;
-			}
+		if (!(Rasterizer == In.Rasterizer))
+		{
+			return false;
+		}
 
-			if (!(ShaderHashes == In.ShaderHashes))
-			{
-				return false;
-			}
+		if (!(DepthStencil == In.DepthStencil))
+		{
+			return false;
+		}
 
-#if VULKAN_SUPPORTS_COLOR_CONVERSIONS
-			for (uint32 Index = 0; Index < MaxImmutableSamplers; ++Index)
-			{
-				if (ImmutableSamplers[Index] != In.ImmutableSamplers[Index])
-				{
-					return false;
-				}
-			}
+		if (!(SubpassIndex == In.SubpassIndex))
+		{
+			return false;
+		}
+#if 0 == VULKAN_USE_SHADERKEYS
+		if (!(ShaderHashes == In.ShaderHashes))
+		{
+			return false;
+		}
+#else
+		if (0 != FMemory::Memcmp(ShaderKeys, In.ShaderKeys, sizeof(ShaderKeys)))
+		{
+			return false;
+		}
 #endif
 
-			if (!(RenderTargets == In.RenderTargets))
-			{
-				return false;
-			}
-
-			if (VertexBindings != In.VertexBindings)
-			{
-				return false;
-			}
-
-			if (VertexAttributes != In.VertexAttributes)
+#if VULKAN_SUPPORTS_COLOR_CONVERSIONS
+		for (uint32 Index = 0; Index < MaxImmutableSamplers; ++Index)
+		{
+			if (ImmutableSamplers[Index] != In.ImmutableSamplers[Index])
 			{
 				return false;
 			}
-
-			return true;
 		}
-	};
-
-	using TGfxPipelineEntrySharedPtr = TSharedPtr<FGfxPipelineEntry, ESPMode::ThreadSafe>;
-
-#if VULKAN_ENABLE_LRU_CACHE
-	struct FPipelineSize
-	{
-		uint32 ShaderHash;
-		uint32 PipelineSize;
+#endif
 
-		FPipelineSize()
-			:ShaderHash(0), PipelineSize(0)
+		if (!(RenderTargets == In.RenderTargets))
 		{
+			return false;
 		}
-	};
-#endif
 
-	FVulkanComputePipeline* GetOrCreateComputePipeline(FVulkanComputeShader* ComputeShader);
+		if (VertexBindings != In.VertexBindings)
+		{
+			return false;
+		}
 
-#if VULKAN_ENABLE_LRU_CACHE
+		if (VertexAttributes != In.VertexAttributes)
+		{
+			return false;
+		}
 
-	bool ShouldEvictImmediately()
-	{
-		return bEvictImmediately && PipelineLRU.IsActive();
+		return true;
 	}
+};
 
-	TMap<uint32, FPipelineSize*> PipelineSizeList;	// key: Shader hash (FShaderHash), value: pipeline size
 
-	class FVKPipelineLRU
-	{
-		typedef TPsoLruCache<FVulkanGfxPipeline*, FVulkanGfxPipeline*> FVulkanPipelineLRUCache;
 
-		const int LRUCapacity = 2048;
-		int32 LRUUsedPipelineSize;
 
-		enum
-		{
-			NUM_BUFFERS = 3,
-		};
+class FVulkanPipelineStateCacheManager
+{
+public:
+
+	// Array of potential cache locations; first entries have highest priority. Only one cache file is loaded. If unsuccessful, tries next entry in the array.
+	void InitAndLoad(const TArray<FString>& CacheFilenames);
+	void Save(const FString& CacheFilename, bool bFromPSOFC = false);
 
-		static void DeleteVkPipeline(FVulkanGfxPipeline* GfxPipeline);
+	FVulkanPipelineStateCacheManager(FVulkanDevice* InParent);
+	~FVulkanPipelineStateCacheManager();
 
-		void EnsureVkPipelineAndAddToLRU(FVulkanRHIGraphicsPipelineState* Pipeline);
+	void RebuildCache();
 
-		void AddToLRU(FVulkanGfxPipeline* GfxPipeline);
 
-		// Use with external lock - LRUCS
-		void EvictFromLRU();
+	// Shader microcode is shared between pipeline entries so keep a cache around to prevent duplicated storage
+	struct FShaderUCodeCache
+	{
+		using TDataMap = TMap<FSHAHash, TArray<uint32>>;
+		TDataMap Data;
 
-	public:
-		FVKPipelineLRU() : LRUUsedPipelineSize(0), LRU(LRUCapacity)
+		TArray<uint32>* Add(const FSHAHash& Hash, const FVulkanShader* Shader)
 		{
-			bUseLRU = (int32)CVarEnableLRU.GetValueOnAnyThread() == 1;
+			check(Shader->Spirv.Num() != 0);
+
+			TArray<uint32>& Code = Data.Add(Hash);
+			Code = Shader->Spirv;
+
+			return &Data[Hash];
 		}
 
-		bool IsActive()
+		TArray<uint32>* Get(const FSHAHash& Hash)
 		{
-			return bUseLRU;
+			return Data.Find(Hash);
 		}
+	};
+	FVulkanComputePipeline* GetOrCreateComputePipeline(FVulkanComputeShader* ComputeShader);
 
-		void Add(FVulkanGfxPipeline* GfxPipeline);
+private:
+	/** Delegate handlers to track the ShaderPipelineCache precompile. */
+	void OnShaderPipelineCacheOpened(FString const& Name, EShaderPlatform Platform, uint32 Count, const FGuid& VersionGuid, FShaderPipelineCache::FShaderCachePrecompileContext& ShaderCachePrecompileContext);
+	void OnShaderPipelineCachePrecompilationComplete(uint32 Count, double Seconds, const FShaderPipelineCache::FShaderCachePrecompileContext& ShaderCachePrecompileContext);
 
-		void Touch(FVulkanRHIGraphicsPipelineState* Pipeline);
+	void CreateGfxEntry(const FGraphicsPipelineStateInitializer& PSOInitializer, FVulkanDescriptorSetsLayoutInfo& DescriptorSetLayoutInfo, FGfxPipelineDesc* Desc);
+	bool Load(const TArray<FString>& CacheFilenames);
+	void DestroyCache();
 
-	private:
-		bool bUseLRU;
-		FVulkanPipelineLRUCache LRU;
-		FCriticalSection LRUCS; // For LRU, LRUUsedPipelineSize and GfxPipeline->LRUNode
-	};
+	FVulkanRHIGraphicsPipelineState* RHICreateGraphicsPipelineState(const FGraphicsPipelineStateInitializer& Initializer);
+	void NotifyDeletedGraphicsPSO(FRHIGraphicsPipelineState* PSO);
+	bool CreateGfxPipelineFromEntry(FVulkanRHIGraphicsPipelineState* PSO, FVulkanShader* Shaders[ShaderStage::NumStages], VkPipeline* Pipeline);
 
-	FVKPipelineLRU PipelineLRU;
 
-#endif
-private:
-	FVulkanDevice* Device;
+	FVulkanLayout* FindOrAddLayout(const FVulkanDescriptorSetsLayoutInfo& DescriptorSetLayoutInfo, bool bGfxLayout);
+	FVulkanComputePipeline* CreateComputePipelineFromShader(FVulkanComputeShader* Shader);
+
+	/** LRU Related functions */
+	bool LRUEvictImmediately();
+	void LRUTrim(uint32 nSpaceNeeded);
+	void LRUAdd(FVulkanRHIGraphicsPipelineState* PSO);
+	void LRUTouch(FVulkanRHIGraphicsPipelineState* PSO);
+	bool LRUEvictOne(bool bOnlyOld = false);
+	void LRURemoveAll();
+	void LRUDump();
+	void LRUDebugEvictAll(); //evict all that are safe to evict without stalling..
+	void LRURemove(FVulkanRHIGraphicsPipelineState* PSO);
+	void LRUCheckNotInside(FVulkanRHIGraphicsPipelineState* PSO);
 
-	bool bEvictImmediately;
 
+	FVulkanDevice* Device;
+	bool bEvictImmediately;
 	// if true, we will link to the PSOFC, loading later, when we have that guid and only if the guid matches, saving only if there is no match, and only saving after the PSOFC is done.
 	bool bLinkedToPSOFC;
 	bool bLinkedToPSOFCSucessfulLoaded;
@@ -569,55 +524,37 @@ private:
 	FDelegateHandle OnShaderPipelineCacheOpenedDelegate;
 	FDelegateHandle OnShaderPipelineCachePrecompilationCompleteDelegate;
 
-	/** Delegate handlers to track the ShaderPipelineCache precompile. */
-	void OnShaderPipelineCacheOpened(FString const& Name, EShaderPlatform Platform, uint32 Count, const FGuid& VersionGuid, FShaderPipelineCache::FShaderCachePrecompileContext& ShaderCachePrecompileContext);
-	void OnShaderPipelineCachePrecompilationComplete(uint32 Count, double Seconds, const FShaderPipelineCache::FShaderCachePrecompileContext& ShaderCachePrecompileContext);
-
-	TMap<FGfxPSIKey, FVulkanRHIGraphicsPipelineState*> InitializerToPipelineMap;
-	FCriticalSection InitializerToPipelineMapCS;
-
-	FCriticalSection GfxPipelineEntriesCS;
-	TMap<FGfxEntryKey, TGfxPipelineEntrySharedPtr> GfxPipelineEntries;
 
 	FRWLock ComputePipelineLock;
 	TMap<uint64, FVulkanComputePipeline*> ComputePipelineEntries;
 
 	VkPipelineCache PipelineCache;
-
 	FShaderUCodeCache ShaderCache;
-
-	FVulkanRHIGraphicsPipelineState* CreateAndAdd(const FGraphicsPipelineStateInitializer& PSOInitializer, FGfxPSIKey PSIKey, TGfxPipelineEntrySharedPtr GfxEntry, FGfxEntryKey GfxEntryKey);
-	void CreateGfxPipelineFromEntry(FGfxPipelineEntry* GfxEntry, FVulkanShader* Shaders[ShaderStage::NumStages], FVulkanGfxPipeline* Pipeline);
-
-	FGfxPipelineEntry* CreateGfxEntry(const FGraphicsPipelineStateInitializer& PSOInitializer);
-
-	bool Load(const TArray<FString>& CacheFilenames);
-	void DestroyCache();
-
-	FVulkanGfxLayout* GetOrGenerateGfxLayout(const FGraphicsPipelineStateInitializer& PSOInitializer, FVulkanShader*const* Shaders, FVulkanVertexInputStateInfo& OutVertexInputState);
-
-	TMap<FGfxEntryKey, FVulkanGfxPipeline*> EntryKeyToGfxPipelineMap;
-
-	//@@
+	FCriticalSection LayoutMapCS;
 	TMap<FVulkanDescriptorSetsLayoutInfo, FVulkanLayout*> LayoutMap;
 	FVulkanDescriptorSetLayoutMap DSetLayoutMap;
-	FCriticalSection LayoutMapCS;
-
-	FVulkanRHIGraphicsPipelineState* FindInLoadedLibrary(const FGraphicsPipelineStateInitializer& PSOInitializer, FGfxPSIKey& PSIKey, TGfxPipelineEntrySharedPtr& OutGfxEntry, FGfxEntryKey& OutGfxEntryKey);
 
-	friend class FVulkanDynamicRHI;
 
-	//@@
-	FVulkanLayout* FindOrAddLayout(const FVulkanDescriptorSetsLayoutInfo& DescriptorSetLayoutInfo, bool bGfxLayout);
+	FCriticalSection GraphicsPSOLockedCS;
+	TMap<FVulkanPSOKey, FVulkanRHIGraphicsPipelineState*> GraphicsPSOLockedMap;
 
-	FVulkanComputePipeline* CreateComputePipelineFromShader(FVulkanComputeShader* Shader);
+		
+	FCriticalSection LRUCS;
+	FVulkanRHIGraphicsPipelineStateLRU LRU;
+	uint32 LRUUsedPipelineSize = 0;
+	uint32 LRUUsedPipelineCount = 0;
+	uint32 LRUUsedPipelineMax = 0;
+	TMap<uint32, FVulkanPipelineSize> LRU2SizeList;	// key: Shader hash (FShaderHash), value: pipeline size
+	bool bUseLRU = true;
+	friend class FVulkanDynamicRHI;
+	friend class FVulkanCommandListContext;
+	friend class FVulkanRHIGraphicsPipelineState;
 
-#if VULKAN_ENABLE_LRU_CACHE
 	struct FVulkanLRUCacheFile
 	{
 		enum
 		{
-			LRU_CACHE_VERSION = 1,
+			LRU_CACHE_VERSION = 2,
 		};
 		struct FFileHeader
 		{
@@ -625,12 +562,11 @@ private:
 			int32 SizeOfPipelineSizes = -1;
 		} Header;
 
-		TArray<FPipelineSize*> PipelineSizes;
+		TArray<FVulkanPipelineSize> PipelineSizes;
 
 		void Save(FArchive& Ar);
 		bool Load(FArchive& Ar);
 	};
-#endif
 };
 
 // Common pipeline class
@@ -708,66 +644,30 @@ protected:
 	friend class FVulkanPipelineStateCacheManager;
 };
 
-using TGfxPipelineEntrySharedPtr = TSharedPtr<FVulkanPipelineStateCacheManager::FGfxPipelineEntry, ESPMode::ThreadSafe>;
-
-class FVulkanGfxPipeline : public FVulkanPipeline, public FRHIResource
+class FVulkanRHIGraphicsPipelineState : public FRHIGraphicsPipelineState
 {
 public:
-#if VULKAN_ENABLE_LRU_CACHE
-	FVulkanGfxPipeline(FVulkanDevice* InDevice, TGfxPipelineEntrySharedPtr InGfxPipelineEntry = {});
-#else
-	FVulkanGfxPipeline(FVulkanDevice* InDevice);
-#endif
-	virtual ~FVulkanGfxPipeline();
-	inline void Bind(VkCommandBuffer CmdBuffer)
-	{
-#if VULKAN_ENABLE_LRU_CACHE
-		RecentFrame = GFrameNumberRenderThread;
-#endif
-		VulkanRHI::vkCmdBindPipeline(CmdBuffer, VK_PIPELINE_BIND_POINT_GRAPHICS, Pipeline);
-	}
-
-	inline bool IsRuntimeInitialized() const
-	{
-		return bRuntimeObjectsValid;
-	}
+	FVulkanRHIGraphicsPipelineState(FVulkanDevice* Device, const FGraphicsPipelineStateInitializer& PSOInitializer, FGfxPipelineDesc& Desc, FVulkanPSOKey* Key);
+	~FVulkanRHIGraphicsPipelineState();
 
 	inline const FVulkanVertexInputStateInfo& GetVertexInputState() const
 	{
-		check(bRuntimeObjectsValid);
 		return VertexInputState;
 	}
 
-	void CreateRuntimeObjects(const FGraphicsPipelineStateInitializer& InPSOInitializer);
+	inline const FVulkanLayout& GetLayout() const
+	{
+		return *Layout;
+	}
 
 	inline const FVulkanGfxLayout& GetGfxLayout() const
 	{
-		return *(FVulkanGfxLayout*)Layout;
+		return *(FVulkanGfxLayout*)&GetLayout();
 	}
 
-private:
-#if VULKAN_ENABLE_LRU_CACHE
-	TGfxPipelineEntrySharedPtr GfxPipelineEntry;
-	uint32 PipelineCacheSize;
-	uint32 RecentFrame;
-	// ID to LRU (if used) allows quick access when updating LRU status.
-	FSetElementId LRUNode;
-	bool bTrackedByLRU;
-	friend class FVulkanPipelineStateCacheManager;
-#endif
-	FVulkanVertexInputStateInfo VertexInputState;
-	bool						bRuntimeObjectsValid;
-};
-
-class FVulkanRHIGraphicsPipelineState : public FRHIGraphicsPipelineState
-{
-public:
-	FVulkanRHIGraphicsPipelineState(const FBoundShaderStateInput& InBSI, FVulkanGfxPipeline* InPipeline, EPrimitiveType InPrimitiveType);
-	~FVulkanRHIGraphicsPipelineState();
-
 	inline void Bind(VkCommandBuffer CmdBuffer)
 	{
-		Pipeline->Bind(CmdBuffer);
+		VulkanRHI::vkCmdBindPipeline(CmdBuffer, VK_PIPELINE_BIND_POINT_GRAPHICS, VulkanPipeline);
 	}
 
 	inline const uint64 GetShaderKey(EShaderFrequency Frequency) const
@@ -776,10 +676,58 @@ public:
 		return ShaderKeys[Stage];
 	}
 
-	TRefCountPtr<FVulkanGfxPipeline>	Pipeline;
+	inline VkPipeline GetVulkanPipeline() const
+	{
+		return VulkanPipeline;
+	}
+
+	void DeleteVkPipeline(bool bImmediate);
+	void GetOrCreateShaderModules(FVulkanShader*const* Shaders);
+	void PurgeShaderModules(FVulkanShader*const* Shaders);
+	void PurgeLoadedShaderModules(FVulkanDevice* InDevice);
+
+
+
+	bool								bHasInputAttachments;
+	bool								bIsRegistered;
+	bool								bShaderModulesLoaded = false;
+
 	uint64								ShaderKeys[ShaderStage::NumStages];
 	TEnumAsByte<EPrimitiveType>			PrimitiveType;
-	bool								bHasInputAttachments;
+
+
+	VkPipeline VulkanPipeline;
+	FVulkanVertexInputStateInfo VertexInputState;
+	FVulkanLayout* Layout;
+	FVulkanDevice* Device;
+	FGfxPipelineDesc Desc;
+	FVulkanShader* VulkanShaders[ShaderStage::NumStages];
+	VkShaderModule ShaderModules[ShaderStage::NumStages];
+	const FVulkanRenderPass* RenderPass;
+
+
+	FVulkanRHIGraphicsPipelineStateLRUNode* LRUNode = nullptr;
+	uint32 LRUFrame;
+	uint32 PipelineCacheSize;
+	FVulkanPSOKey							VulkanKey;
+
+
+#if VULKAN_PSO_CACHE_DEBUG
+	FPixelShaderRHIRef					PixelShaderRHI;
+	FVertexShaderRHIRef					VertexShaderRHI;
+	FVertexDeclarationRHIRef			VertexDeclarationRHI;
+
+#if PLATFORM_SUPPORTS_TESSELLATION_SHADERS
+	FDomainShaderRHIRef					DomainShaderRHI;
+	FHullShaderRHIRef					HullShaderRHI;
+#endif 
+#if PLATFORM_SUPPORTS_GEOMETRY_SHADERS
+	FGeometryShaderRHIRef				GeometryShaderRHI;
+#endif
+
+	FGraphicsPipelineStateInitializer		PSOInitializer;
+#endif
+
 };
 
 template<>
diff --git a/Engine/Source/Runtime/VulkanRHI/Private/VulkanPipelineState.cpp b/Engine/Source/Runtime/VulkanRHI/Private/VulkanPipelineState.cpp
--- a/Engine/Source/Runtime/VulkanRHI/Private/VulkanPipelineState.cpp
+++ b/Engine/Source/Runtime/VulkanRHI/Private/VulkanPipelineState.cpp
@@ -205,56 +205,60 @@ FVulkanGraphicsPipelineDescriptorState::FVulkanGraphicsPipelineDescriptorState(F
 	FMemory::Memzero(PackedUniformBuffersDirty);
 
 	check(InGfxPipeline);
-	check(InGfxPipeline->Pipeline);
-	DescriptorSetsLayout = &InGfxPipeline->Pipeline->GetLayout().GetDescriptorSetsLayout();
-	PipelineDescriptorInfo = &InGfxPipeline->Pipeline->GetGfxLayout().GetGfxPipelineDescriptorInfo();
-	
-	UsedSetsMask = PipelineDescriptorInfo->HasDescriptorsInSetMask;
-	const FVulkanShaderFactory& ShaderFactory = Device->GetShaderFactory();
+	{
+		
+		check(InGfxPipeline->Layout);
+		DescriptorSetsLayout = &InGfxPipeline->Layout->GetDescriptorSetsLayout();
+		FVulkanGfxLayout& GfxLayout  = *(FVulkanGfxLayout*)InGfxPipeline->Layout;
+		PipelineDescriptorInfo = &GfxLayout.GetGfxPipelineDescriptorInfo();
 
-	const FVulkanVertexShader* VertexShader = ShaderFactory.LookupShader<FVulkanVertexShader>(InGfxPipeline->GetShaderKey(SF_Vertex));
-	check(VertexShader);
-	PackedUniformBuffers[ShaderStage::Vertex].Init(VertexShader->GetCodeHeader(), PackedUniformBuffersMask[ShaderStage::Vertex]);
+		UsedSetsMask = PipelineDescriptorInfo->HasDescriptorsInSetMask;
+		const FVulkanShaderFactory& ShaderFactory = Device->GetShaderFactory();
 
-	uint64 PixelShaderKey = InGfxPipeline->GetShaderKey(SF_Pixel);
-	if (PixelShaderKey)
-	{
-		const FVulkanPixelShader* PixelShader = ShaderFactory.LookupShader<FVulkanPixelShader>(PixelShaderKey);
-		check(PixelShader);
+		const FVulkanVertexShader* VertexShader = ShaderFactory.LookupShader<FVulkanVertexShader>(InGfxPipeline->GetShaderKey(SF_Vertex));
+		check(VertexShader);
+		PackedUniformBuffers[ShaderStage::Vertex].Init(VertexShader->GetCodeHeader(), PackedUniformBuffersMask[ShaderStage::Vertex]);
 
-		PackedUniformBuffers[ShaderStage::Pixel].Init(PixelShader->GetCodeHeader(), PackedUniformBuffersMask[ShaderStage::Pixel]);
-	}
+		uint64 PixelShaderKey = InGfxPipeline->GetShaderKey(SF_Pixel);
+		if (PixelShaderKey)
+		{
+			const FVulkanPixelShader* PixelShader = ShaderFactory.LookupShader<FVulkanPixelShader>(PixelShaderKey);
+			check(PixelShader);
+
+			PackedUniformBuffers[ShaderStage::Pixel].Init(PixelShader->GetCodeHeader(), PackedUniformBuffersMask[ShaderStage::Pixel]);
+		}
 
 #if VULKAN_SUPPORTS_GEOMETRY_SHADERS
-	uint64 GeometryShaderKey = InGfxPipeline->GetShaderKey(SF_Geometry);
-	if (GeometryShaderKey)
-	{
-		const FVulkanGeometryShader* GeometryShader = ShaderFactory.LookupShader<FVulkanGeometryShader>(GeometryShaderKey);
-		check(GeometryShader);
+		uint64 GeometryShaderKey = InGfxPipeline->GetShaderKey(SF_Geometry);
+		if (GeometryShaderKey)
+		{
+			const FVulkanGeometryShader* GeometryShader = ShaderFactory.LookupShader<FVulkanGeometryShader>(GeometryShaderKey);
+			check(GeometryShader);
 
-		PackedUniformBuffers[ShaderStage::Geometry].Init(GeometryShader->GetCodeHeader(), PackedUniformBuffersMask[ShaderStage::Geometry]);
-	}
+			PackedUniformBuffers[ShaderStage::Geometry].Init(GeometryShader->GetCodeHeader(), PackedUniformBuffersMask[ShaderStage::Geometry]);
+		}
 #endif
 
 #if PLATFORM_SUPPORTS_TESSELLATION_SHADERS
-	uint64 HullShaderKey = InGfxPipeline->GetShaderKey(SF_Hull);
-	if (HullShaderKey)
-	{
-		const FVulkanHullShader* HullShader = ShaderFactory.LookupShader<FVulkanHullShader>(HullShaderKey);
-		PackedUniformBuffers[ShaderStage::Hull].Init(HullShader->GetCodeHeader(), PackedUniformBuffersMask[ShaderStage::Hull]);
-	}
-	uint64 DomainShaderKey = InGfxPipeline->GetShaderKey(SF_Domain);
-	if (DomainShaderKey)
-	{
-		const FVulkanDomainShader* DomainShader = ShaderFactory.LookupShader<FVulkanDomainShader>(DomainShaderKey);
-		PackedUniformBuffers[ShaderStage::Domain].Init(DomainShader->GetCodeHeader(), PackedUniformBuffersMask[ShaderStage::Domain]);
-	}
+		uint64 HullShaderKey = InGfxPipeline->GetShaderKey(SF_Hull);
+		if (HullShaderKey)
+		{
+			const FVulkanHullShader* HullShader = ShaderFactory.LookupShader<FVulkanHullShader>(HullShaderKey);
+			PackedUniformBuffers[ShaderStage::Hull].Init(HullShader->GetCodeHeader(), PackedUniformBuffersMask[ShaderStage::Hull]);
+		}
+		uint64 DomainShaderKey = InGfxPipeline->GetShaderKey(SF_Domain);
+		if (DomainShaderKey)
+		{
+			const FVulkanDomainShader* DomainShader = ShaderFactory.LookupShader<FVulkanDomainShader>(DomainShaderKey);
+			PackedUniformBuffers[ShaderStage::Domain].Init(DomainShader->GetCodeHeader(), PackedUniformBuffersMask[ShaderStage::Domain]);
+		}
 #endif
-	CreateDescriptorWriteInfos();
+		CreateDescriptorWriteInfos();
 
-	//UE_LOG(LogVulkanRHI, Warning, TEXT("GfxPSOState %p For PSO %p Writes:%d"), this, InGfxPipeline, DSWriteContainer.DescriptorWrites.Num());
+		//UE_LOG(LogVulkanRHI, Warning, TEXT("GfxPSOState %p For PSO %p Writes:%d"), this, InGfxPipeline, DSWriteContainer.DescriptorWrites.Num());
 
-	InGfxPipeline->AddRef();
+		InGfxPipeline->AddRef();
+	}
 }
 
 template<bool bUseDynamicGlobalUBs>
@@ -342,14 +346,8 @@ void FVulkanCommandListContext::RHISetGraphicsPipelineState(FRHIGraphicsPipeline
 {
 	FVulkanRHIGraphicsPipelineState* Pipeline = ResourceCast(GraphicsState);
 	
-#if VULKAN_ENABLE_LRU_CACHE
-	if (!Pipeline)
-	{
-		return; // this happens when we immediately evict an cached PSO...the thing is never actually created
-	}
 	FVulkanPipelineStateCacheManager* PipelineStateCache = Device->GetPipelineStateCache();
-	PipelineStateCache->PipelineLRU.Touch(Pipeline);
-#endif
+	PipelineStateCache->LRUTouch(Pipeline);
 
 	FVulkanCmdBuffer* CmdBuffer = CommandBufferManager->GetActiveCmdBuffer();
 	bool bForceResetPipeline = !CmdBuffer->bHasPipeline;
diff --git a/Engine/Source/Runtime/VulkanRHI/Private/VulkanPipelineState.h b/Engine/Source/Runtime/VulkanRHI/Private/VulkanPipelineState.h
--- a/Engine/Source/Runtime/VulkanRHI/Private/VulkanPipelineState.h
+++ b/Engine/Source/Runtime/VulkanRHI/Private/VulkanPipelineState.h
@@ -256,7 +256,7 @@ public:
 
 	inline void BindDescriptorSets(VkCommandBuffer CmdBuffer)
 	{
-		Bind(CmdBuffer, GfxPipeline->Pipeline->GetLayout().GetPipelineLayout(), VK_PIPELINE_BIND_POINT_GRAPHICS);
+		Bind(CmdBuffer, GfxPipeline->GetLayout().GetPipelineLayout(), VK_PIPELINE_BIND_POINT_GRAPHICS);
 	}
 
 	void Reset()
diff --git a/Engine/Source/Runtime/VulkanRHI/Private/VulkanRHI.cpp b/Engine/Source/Runtime/VulkanRHI/Private/VulkanRHI.cpp
--- a/Engine/Source/Runtime/VulkanRHI/Private/VulkanRHI.cpp
+++ b/Engine/Source/Runtime/VulkanRHI/Private/VulkanRHI.cpp
@@ -323,6 +323,8 @@ void FVulkanDynamicRHI::Shutdown()
 
 #if UE_BUILD_DEBUG || UE_BUILD_DEVELOPMENT
 	IConsoleManager::Get().UnregisterConsoleObject(DumpMemoryCmd);
+	IConsoleManager::Get().UnregisterConsoleObject(DumpLRUCmd);
+	IConsoleManager::Get().UnregisterConsoleObject(TrimLRUCmd);
 #endif
 
 	FVulkanPlatform::FreeVulkanLibrary();
@@ -757,6 +759,19 @@ void FVulkanDynamicRHI::InitInstance()
 			FConsoleCommandDelegate::CreateStatic(DumpMemory),
 			ECVF_Default
 			);
+		DumpLRUCmd = IConsoleManager::Get().RegisterConsoleCommand(
+			TEXT("r.Vulkan.DumpPSOLRU"),
+			TEXT("Dumps Vulkan PSO LRU."),
+			FConsoleCommandDelegate::CreateStatic(DumpLRU),
+			ECVF_Default
+		);
+		TrimLRUCmd = IConsoleManager::Get().RegisterConsoleCommand(
+			TEXT("r.Vulkan.TrimPSOLRU"),
+			TEXT("Trim Vulkan PSO LRU."),
+			FConsoleCommandDelegate::CreateStatic(TrimLRU),
+			ECVF_Default
+		);
+
 #endif
 	}
 }
@@ -1735,6 +1750,18 @@ void FVulkanDynamicRHI::DumpMemory()
 	GVulkanRHI->Device->GetResourceHeapManager().DumpMemory();
 	GVulkanRHI->Device->GetStagingManager().DumpMemory();
 }
+void FVulkanDynamicRHI::DumpLRU()
+{
+	FVulkanDynamicRHI* RHI = (FVulkanDynamicRHI*)GDynamicRHI;
+	RHI->Device->PipelineStateCache->LRUDump();
+}
+void FVulkanDynamicRHI::TrimLRU()
+{
+	FVulkanDynamicRHI* RHI = (FVulkanDynamicRHI*)GDynamicRHI;
+	RHI->Device->PipelineStateCache->LRUDebugEvictAll();
+}
+
+
 #endif
 
 void FVulkanDynamicRHI::DestroySwapChain()
diff --git a/Engine/Source/Runtime/VulkanRHI/Private/VulkanRHIPrivate.h b/Engine/Source/Runtime/VulkanRHI/Private/VulkanRHIPrivate.h
--- a/Engine/Source/Runtime/VulkanRHI/Private/VulkanRHIPrivate.h
+++ b/Engine/Source/Runtime/VulkanRHI/Private/VulkanRHIPrivate.h
@@ -193,6 +193,7 @@ protected:
 	void CalculateRenderPassHashes(const FRHISetRenderTargetsInfo& RTInfo);
 
 	friend class FVulkanPipelineStateCacheManager;
+	friend struct FGfxPipelineDesc;
 };
 
 class FVulkanFramebuffer
@@ -382,6 +383,20 @@ inline void VulkanSetImageLayoutSimple(VkCommandBuffer CmdBuffer, VkImage Image,
 
 void VulkanResolveImage(VkCommandBuffer Cmd, FRHITexture* SourceTextureRHI, FRHITexture* DestTextureRHI);
 
+DECLARE_STATS_GROUP(TEXT("Vulkan PSO"), STATGROUP_VulkanPSO, STATCAT_Advanced);
+DECLARE_DWORD_ACCUMULATOR_STAT_EXTERN(TEXT("PSO LRU Elements"), STAT_VulkanNumPSOLRU, STATGROUP_VulkanPSO, );
+DECLARE_DWORD_ACCUMULATOR_STAT_EXTERN(TEXT("PSO LRU Size"), STAT_VulkanNumPSOLRUSize, STATGROUP_VulkanPSO, );
+DECLARE_DWORD_ACCUMULATOR_STAT_EXTERN(TEXT("Num PSOs"), STAT_VulkanNumPSOs, STATGROUP_VulkanPSO, );
+DECLARE_DWORD_ACCUMULATOR_STAT_EXTERN(TEXT("Num Graphics PSOs"), STAT_VulkanNumGraphicsPSOs, STATGROUP_VulkanPSO, );
+DECLARE_DWORD_ACCUMULATOR_STAT_EXTERN(TEXT("Num Compute  PSOs"), STAT_VulkanNumComputePSOs, STATGROUP_VulkanPSO, );
+DECLARE_DWORD_ACCUMULATOR_STAT_EXTERN(TEXT("VulkanPSOKey Memory"), STAT_VulkanPSOKeyMemory, STATGROUP_VulkanPSO, );
+
+DECLARE_CYCLE_STAT_EXTERN(TEXT("PSO HeaderInit time"), STAT_VulkanPSOHeaderInitTime, STATGROUP_VulkanPSO, );
+DECLARE_CYCLE_STAT_EXTERN(TEXT("PSO Lookup time"), STAT_VulkanPSOLookupTime, STATGROUP_VulkanPSO, );
+DECLARE_CYCLE_STAT_EXTERN(TEXT("PSO Creation time"), STAT_VulkanPSOCreationTime, STATGROUP_VulkanPSO, );
+DECLARE_CYCLE_STAT_EXTERN(TEXT("PSO Vulkan Creation time"), STAT_VulkanPSOVulkanCreationTime, STATGROUP_VulkanPSO, );
+
+
 // Stats
 DECLARE_STATS_GROUP(TEXT("Vulkan RHI"), STATGROUP_VulkanRHI, STATCAT_Advanced);
 //DECLARE_STATS_GROUP(TEXT("Vulkan RHI Verbose"), STATGROUP_VulkanRHIVERBOSE, STATCAT_Advanced);
@@ -394,7 +409,6 @@ DECLARE_CYCLE_STAT_EXTERN(TEXT("Get Or Create Pipeline"), STAT_VulkanGetOrCreate
 DECLARE_CYCLE_STAT_EXTERN(TEXT("Get DescriptorSet"), STAT_VulkanGetDescriptorSet, STATGROUP_VulkanRHI, );
 DECLARE_CYCLE_STAT_EXTERN(TEXT("Pipeline Bind"), STAT_VulkanPipelineBind, STATGROUP_VulkanRHI, );
 DECLARE_DWORD_ACCUMULATOR_STAT_EXTERN(TEXT("Num Cmd Buffers"), STAT_VulkanNumCmdBuffers, STATGROUP_VulkanRHI, );
-DECLARE_DWORD_ACCUMULATOR_STAT_EXTERN(TEXT("Num PSOs"), STAT_VulkanNumPSOs, STATGROUP_VulkanRHI, );
 DECLARE_DWORD_ACCUMULATOR_STAT_EXTERN(TEXT("Num Render Passes"), STAT_VulkanNumRenderPasses, STATGROUP_VulkanRHI, );
 DECLARE_DWORD_ACCUMULATOR_STAT_EXTERN(TEXT("Num Frame Buffers"), STAT_VulkanNumFrameBuffers, STATGROUP_VulkanRHI, );
 DECLARE_DWORD_ACCUMULATOR_STAT_EXTERN(TEXT("Num Buffer Views"), STAT_VulkanNumBufferViews, STATGROUP_VulkanRHI, );
diff --git a/Engine/Source/Runtime/VulkanRHI/Private/VulkanResources.h b/Engine/Source/Runtime/VulkanRHI/Private/VulkanResources.h
--- a/Engine/Source/Runtime/VulkanRHI/Private/VulkanResources.h
+++ b/Engine/Source/Runtime/VulkanRHI/Private/VulkanResources.h
@@ -1355,6 +1355,7 @@ class FVulkanVertexInputStateInfo
 {
 public:
 	FVulkanVertexInputStateInfo();
+	~FVulkanVertexInputStateInfo();
 
 	void Generate(FVulkanVertexDeclaration* VertexDeclaration, uint32 VertexHeaderInOutAttributeMask);
 
@@ -1369,6 +1370,8 @@ public:
 		return Info;
 	}
 
+	bool operator ==(const FVulkanVertexInputStateInfo& Other);
+
 protected:
 	VkPipelineVertexInputStateCreateInfo Info;
 	uint32 Hash;
diff --git a/Engine/Source/Runtime/VulkanRHI/Private/VulkanUtil.cpp b/Engine/Source/Runtime/VulkanRHI/Private/VulkanUtil.cpp
--- a/Engine/Source/Runtime/VulkanRHI/Private/VulkanUtil.cpp
+++ b/Engine/Source/Runtime/VulkanRHI/Private/VulkanUtil.cpp
@@ -773,6 +773,19 @@ namespace VulkanRHI
 	}
 }
 
+
+DEFINE_STAT(STAT_VulkanNumPSOs);
+DEFINE_STAT(STAT_VulkanNumGraphicsPSOs);
+DEFINE_STAT(STAT_VulkanNumPSOLRU);
+DEFINE_STAT(STAT_VulkanNumPSOLRUSize);
+DEFINE_STAT(STAT_VulkanPSOLookupTime);
+DEFINE_STAT(STAT_VulkanPSOCreationTime);
+DEFINE_STAT(STAT_VulkanPSOHeaderInitTime);
+DEFINE_STAT(STAT_VulkanPSOVulkanCreationTime);
+DEFINE_STAT(STAT_VulkanNumComputePSOs);
+DEFINE_STAT(STAT_VulkanPSOKeyMemory);
+
+
 DEFINE_STAT(STAT_VulkanDrawCallTime);
 DEFINE_STAT(STAT_VulkanDispatchCallTime);
 DEFINE_STAT(STAT_VulkanDrawCallPrepareTime);
@@ -782,7 +795,6 @@ DEFINE_STAT(STAT_VulkanGetOrCreatePipeline);
 DEFINE_STAT(STAT_VulkanGetDescriptorSet);
 DEFINE_STAT(STAT_VulkanPipelineBind);
 DEFINE_STAT(STAT_VulkanNumCmdBuffers);
-DEFINE_STAT(STAT_VulkanNumPSOs);
 DEFINE_STAT(STAT_VulkanNumRenderPasses);
 DEFINE_STAT(STAT_VulkanNumFrameBuffers);
 DEFINE_STAT(STAT_VulkanNumBufferViews);
diff --git a/Engine/Source/Runtime/VulkanRHI/Private/VulkanVertexDeclaration.cpp b/Engine/Source/Runtime/VulkanRHI/Private/VulkanVertexDeclaration.cpp
--- a/Engine/Source/Runtime/VulkanRHI/Private/VulkanVertexDeclaration.cpp
+++ b/Engine/Source/Runtime/VulkanRHI/Private/VulkanVertexDeclaration.cpp
@@ -60,6 +60,9 @@ FVertexDeclarationRHIRef FVulkanDynamicRHI::RHICreateVertexDeclaration(const FVe
 	return *VertexDeclarationRefPtr;
 }
 
+FVulkanVertexInputStateInfo::~FVulkanVertexInputStateInfo()
+{
+}
 FVulkanVertexInputStateInfo::FVulkanVertexInputStateInfo()
 	: Hash(0)
 	, BindingsNum(0)
@@ -71,6 +74,22 @@ FVulkanVertexInputStateInfo::FVulkanVertexInputStateInfo()
 	FMemory::Memzero(Bindings);
 }
 
+bool FVulkanVertexInputStateInfo::operator ==(const FVulkanVertexInputStateInfo& Other)
+{
+	if(AttributesNum != Other.AttributesNum)
+	{
+		return false;
+	}
+	for(uint32 i = 0; i < AttributesNum; ++i)
+	{
+		if(0 != FMemory::Memcmp(&Attributes[i], &Other.Attributes[i], sizeof(Attributes[i])))
+		{
+			return false;
+		}
+	}
+	return true;
+}
+
 void FVulkanVertexInputStateInfo::Generate(FVulkanVertexDeclaration* VertexDeclaration, uint32 VertexHeaderInOutAttributeMask)
 {
 	// GenerateVertexInputState is expected to be called only once!
diff --git a/Engine/Source/Runtime/VulkanRHI/Public/VulkanDynamicRHI.h b/Engine/Source/Runtime/VulkanRHI/Public/VulkanDynamicRHI.h
--- a/Engine/Source/Runtime/VulkanRHI/Public/VulkanDynamicRHI.h
+++ b/Engine/Source/Runtime/VulkanRHI/Public/VulkanDynamicRHI.h
@@ -448,8 +448,12 @@ protected:
 
 #if UE_BUILD_DEBUG || UE_BUILD_DEVELOPMENT
 	IConsoleObject* DumpMemoryCmd = nullptr;
+	IConsoleObject* DumpLRUCmd = nullptr;
+	IConsoleObject* TrimLRUCmd = nullptr;
 public:
 	static void DumpMemory();
+	static void DumpLRU();
+	static void TrimLRU();
 #endif
 
 protected:
diff --git a/Engine/Source/Runtime/VulkanRHI/Public/VulkanUtil.h b/Engine/Source/Runtime/VulkanRHI/Public/VulkanUtil.h
--- a/Engine/Source/Runtime/VulkanRHI/Public/VulkanUtil.h
+++ b/Engine/Source/Runtime/VulkanRHI/Public/VulkanUtil.h
@@ -280,6 +280,11 @@ protected:
 		check(FMemory::Memcmp(Data->GetData(), Other.Data->GetData(), Data->Num()) == 0);
 		return true;
 	}
+public:
+	TArray<uint8>& GetDataRef()
+	{
+		return *Data;
+	}
 
 private:
 	void EnsureDataStorage()
@@ -311,7 +316,7 @@ protected:
 };
 
 template <class Derived, bool AlwaysCompareData = false>
-class TDataKey : private TDataKeyBase<AlwaysCompareData ? 2 : (DO_CHECK != 0)>
+class TDataKey : public TDataKeyBase<AlwaysCompareData ? 2 : (DO_CHECK != 0)>
 {
 public:
 	template <class ArchiveWriter>
