Do intiial reset of VkQueryPools. fixes one validation issue.
Because this needs a command buffer, it suddenly gets rather involved.

diff --git a/Engine/Source/Runtime/VulkanRHI/Private/VulkanCommandBuffer.cpp b/Engine/Source/Runtime/VulkanRHI/Private/VulkanCommandBuffer.cpp
--- a/Engine/Source/Runtime/VulkanRHI/Private/VulkanCommandBuffer.cpp
+++ b/Engine/Source/Runtime/VulkanRHI/Private/VulkanCommandBuffer.cpp
@@ -389,6 +389,49 @@ void FVulkanCommandBufferManager::WaitForCmdBuffer(FVulkanCmdBuffer* CmdBuffer,
 	CmdBuffer->RefreshFenceStatus();
 }
 
+void FVulkanCommandBufferManager::AddQueryPoolForReset(VkQueryPool QueryPool, uint32 Size)
+{
+	FScopeLock ScopeLock(&Pool.CS);
+	FQueryPoolReset QPR = { QueryPool, Size};
+	PoolResets.Add(QPR);
+}
+void FVulkanCommandBufferManager::FlushResetQueryPools()
+{
+	FScopeLock ScopeLock(&Pool.CS);
+	if(PoolResets.Num())
+	{
+		FVulkanCmdBuffer* PoolResetCommandBuffer = nullptr;
+		for (int32 Index = 0; Index < Pool.CmdBuffers.Num(); ++Index)
+		{
+			FVulkanCmdBuffer* CmdBuffer = Pool.CmdBuffers[Index];
+			CmdBuffer->RefreshFenceStatus();
+	#if VULKAN_USE_DIFFERENT_POOL_CMDBUFFERS
+			if (!CmdBuffer->bIsUploadOnly)
+	#endif
+			{
+				if (CmdBuffer->State == FVulkanCmdBuffer::EState::ReadyForBegin)
+				{
+					PoolResetCommandBuffer = CmdBuffer;
+				}
+			}
+		}
+		if (!PoolResetCommandBuffer)
+		{
+			PoolResetCommandBuffer = Pool.Create(false);
+		}
+		PoolResetCommandBuffer->Begin();
+
+		VkCommandBuffer CommandBuffer = PoolResetCommandBuffer->GetHandle();
+		for (FQueryPoolReset& QPR : PoolResets)
+		{
+			VulkanRHI::vkCmdResetQueryPool(CommandBuffer, QPR.Pool, 0, QPR.Size);
+		}
+		PoolResetCommandBuffer->End();
+		Queue->Submit(PoolResetCommandBuffer);
+		check(PoolResetCommandBuffer->State == FVulkanCmdBuffer::EState::Submitted);
+		PoolResets.Empty();
+	}
+}
 
 void FVulkanCommandBufferManager::SubmitUploadCmdBuffer(uint32 NumSignalSemaphores, VkSemaphore* SignalSemaphores)
 {
@@ -412,6 +455,7 @@ void FVulkanCommandBufferManager::SubmitUploadCmdBuffer(uint32 NumSignalSemaphor
 void FVulkanCommandBufferManager::SubmitActiveCmdBuffer(VulkanRHI::FSemaphore* SignalSemaphore)
 {
 	FScopeLock ScopeLock(&Pool.CS);
+	FlushResetQueryPools();
 	check(!UploadCmdBuffer);
 	check(ActiveCmdBuffer);
 	if (!ActiveCmdBuffer->IsSubmitted() && ActiveCmdBuffer->HasBegun())
diff --git a/Engine/Source/Runtime/VulkanRHI/Private/VulkanCommandBuffer.h b/Engine/Source/Runtime/VulkanRHI/Private/VulkanCommandBuffer.h
--- a/Engine/Source/Runtime/VulkanRHI/Private/VulkanCommandBuffer.h
+++ b/Engine/Source/Runtime/VulkanRHI/Private/VulkanCommandBuffer.h
@@ -314,6 +314,10 @@ public:
 
 	void WaitForCmdBuffer(FVulkanCmdBuffer* CmdBuffer, float TimeInSecondsToWait = 10.0f);
 
+
+	void AddQueryPoolForReset(VkQueryPool Pool, uint32 Size);
+	void FlushResetQueryPools();
+
 	// Update the fences of all cmd buffers except SkipCmdBuffer
 	void RefreshFenceStatus(FVulkanCmdBuffer* SkipCmdBuffer = nullptr)
 	{
@@ -332,9 +336,16 @@ public:
 	void FreeUnusedCmdBuffers();
 
 private:
+	struct FQueryPoolReset
+	{
+		VkQueryPool Pool; 
+		uint32 Size;
+	};
+	
 	FVulkanDevice* Device;
 	FVulkanCommandBufferPool Pool;
 	FVulkanQueue* Queue;
 	FVulkanCmdBuffer* ActiveCmdBuffer;
 	FVulkanCmdBuffer* UploadCmdBuffer;
+	TArray<FQueryPoolReset> PoolResets;
 };
diff --git a/Engine/Source/Runtime/VulkanRHI/Private/VulkanQuery.cpp b/Engine/Source/Runtime/VulkanRHI/Private/VulkanQuery.cpp
--- a/Engine/Source/Runtime/VulkanRHI/Private/VulkanQuery.cpp
+++ b/Engine/Source/Runtime/VulkanRHI/Private/VulkanQuery.cpp
@@ -32,6 +32,7 @@ constexpr int32 NUM_FRAMES_TO_WAIT_REUSE_POOL = 10;
 constexpr uint32 NUM_FRAMES_TO_WAIT_RELEASE_POOL = MAX_uint32; // never release
 #endif
 
+extern FVulkanCommandBufferManager* GVulkanCommandBufferManager;
 
 FVulkanQueryPool::FVulkanQueryPool(FVulkanDevice* InDevice, uint32 InMaxQueries, VkQueryType InQueryType)
 	: VulkanRHI::FDeviceChild(InDevice)
@@ -48,6 +49,8 @@ FVulkanQueryPool::FVulkanQueryPool(FVulkanDevice* InDevice, uint32 InMaxQueries,
 
 	VERIFYVULKANRESULT(VulkanRHI::vkCreateQueryPool(Device->GetInstanceHandle(), &PoolCreateInfo, VULKAN_CPU_ALLOCATOR, &QueryPool));
 
+	GVulkanCommandBufferManager->AddQueryPoolForReset(QueryPool, InMaxQueries);
+	
 	VkEventCreateInfo EventCreateInfo;
 	ZeroVulkanStruct(EventCreateInfo, VK_STRUCTURE_TYPE_EVENT_CREATE_INFO);
 	VERIFYVULKANRESULT(VulkanRHI::vkCreateEvent(Device->GetInstanceHandle(), &EventCreateInfo, VULKAN_CPU_ALLOCATOR, &ResetEvent));
diff --git a/Engine/Source/Runtime/VulkanRHI/Private/VulkanRHI.cpp b/Engine/Source/Runtime/VulkanRHI/Private/VulkanRHI.cpp
--- a/Engine/Source/Runtime/VulkanRHI/Private/VulkanRHI.cpp
+++ b/Engine/Source/Runtime/VulkanRHI/Private/VulkanRHI.cpp
@@ -30,6 +30,9 @@ TAtomic<uint64> GVulkanImageViewHandleIdCounter{ 0 };
 TAtomic<uint64> GVulkanSamplerHandleIdCounter{ 0 };
 TAtomic<uint64> GVulkanDSetLayoutHandleIdCounter{ 0 };
 
+
+FVulkanCommandBufferManager* GVulkanCommandBufferManager = nullptr;
+
 #if VULKAN_ENABLE_DESKTOP_HMD_SUPPORT
 #include "Runtime/HeadMountedDisplay/Public/IHeadMountedDisplayModule.h"
 #endif
@@ -122,10 +125,11 @@ FVulkanCommandListContext::FVulkanCommandListContext(FVulkanDynamicRHI* InRHI, F
 	, GpuProfiler(this, InDevice)
 {
 	FrameTiming = new FVulkanGPUTiming(this, InDevice);
-	FrameTiming->Initialize();
 
 	// Create CommandBufferManager, contain all active buffers
 	CommandBufferManager = new FVulkanCommandBufferManager(InDevice, this);
+	GVulkanCommandBufferManager = CommandBufferManager;
+	FrameTiming->Initialize();
 	if (IsImmediate())
 	{
 		// Insert the Begin frame timestamp query. On EndDrawingViewport() we'll insert the End and immediately after a new Begin()
@@ -156,6 +160,7 @@ FVulkanCommandListContext::~FVulkanCommandListContext()
 	check(CommandBufferManager != nullptr);
 	delete CommandBufferManager;
 	CommandBufferManager = nullptr;
+	GVulkanCommandBufferManager = nullptr;
 
 	TransitionAndLayoutManager.Destroy(*Device, Immediate ? &TransitionAndLayoutManager : nullptr);
 
diff --git a/Engine/Source/Runtime/VulkanRHI/Private/VulkanViewport.cpp b/Engine/Source/Runtime/VulkanRHI/Private/VulkanViewport.cpp
--- a/Engine/Source/Runtime/VulkanRHI/Private/VulkanViewport.cpp
+++ b/Engine/Source/Runtime/VulkanRHI/Private/VulkanViewport.cpp
@@ -12,6 +12,8 @@
 #include "HAL/PlatformAtomics.h"
 #include "Engine/RendererSettings.h"
 
+extern FVulkanCommandBufferManager* GVulkanCommandBufferManager;
+
 struct FRHICommandProcessDeferredDeletionQueue final : public FRHICommand<FRHICommandProcessDeferredDeletionQueue>
 {
 	FVulkanDevice* Device;
@@ -833,7 +835,7 @@ bool FVulkanViewport::Present(FVulkanCommandListContext* Context, FVulkanCmdBuff
 	}
 
 	CmdBuffer->End();
-
+	GVulkanCommandBufferManager->FlushResetQueryPools();
 	if (FVulkanPlatform::SupportsStandardSwapchain())
 	{
 		if (LIKELY(!bFailedToDelayAcquireBackbuffer))
