Lots of vulkan validation work:
-Debug Names for images for helping track down validation errors
-When the reflection environment rounds up to 2 cubemaps, explicitly clear it to avoid a validation error
-Shadowmask texture is now properly transitioned to read.
-Add Missing transtion for LightScatteringHistoryTexture
-Add Ensure when updating descriptorset, to catch validation errors when texture is set
-Fixed an issue in transition resources, where it would use the incorrect no. of subresources
-Add VK_QUERY_RESULT_WAIT_BIT to ensure Query results are actually ready when copied back to result buffer

diff --git a/Engine/Source/Runtime/VulkanRHI/Private/VulkanDevice.cpp b/Engine/Source/Runtime/VulkanRHI/Private/VulkanDevice.cpp
--- a/Engine/Source/Runtime/VulkanRHI/Private/VulkanDevice.cpp
+++ b/Engine/Source/Runtime/VulkanRHI/Private/VulkanDevice.cpp
@@ -418,8 +418,7 @@ void FVulkanDevice::SetupDrawMarkers()
 			bDebugMarkersFound = true;
 		}
 	}
-	FVulkanDynamicRHI* RHI = GVulkanRHI;
-	if (RHI->SupportsDebugUtilsExt())
+	if(GVulkanRHI->SupportsDebugUtilsExt())
 	{
 		DebugMarkers.SetDebugName = (PFN_vkSetDebugUtilsObjectNameEXT)(void*)VulkanRHI::vkGetInstanceProcAddr(RHI->GetInstance(), "vkSetDebugUtilsObjectNameEXT");
 	}
