Fix compile errors when VULKAN_ENABLE_DRAW_MARKERS is false
#rnx

diff --git a/Engine/Source/Runtime/VulkanRHI/Private/VulkanDevice.cpp b/Engine/Source/Runtime/VulkanRHI/Private/VulkanDevice.cpp
--- a/Engine/Source/Runtime/VulkanRHI/Private/VulkanDevice.cpp
+++ b/Engine/Source/Runtime/VulkanRHI/Private/VulkanDevice.cpp
@@ -1260,6 +1260,7 @@ void FVulkanDevice::ReleaseDeferredContext(FVulkanCommandListContext* InContext)
 
 void FVulkanDevice::VulkanSetObjectName(VkObjectType Type, uint64_t Handle, const TCHAR* Name)
 {
+#if VULKAN_ENABLE_DRAW_MARKERS
 	if(DebugMarkers.SetDebugName)
 	{
 		FTCHARToUTF8 Converter(Name);
@@ -1270,4 +1271,5 @@ void FVulkanDevice::VulkanSetObjectName(VkObjectType Type, uint64_t Handle, cons
 		Info.pObjectName = Converter.Get();
 		DebugMarkers.SetDebugName(Device, &Info);
 	}
+#endif // VULKAN_ENABLE_DRAW_MARKERS
 }
\ No newline at end of file
