UE4R - Fix Vulkan validation error
#rnx

diff --git a/Engine/Source/Runtime/VulkanRHI/Private/VulkanMemory.cpp b/Engine/Source/Runtime/VulkanRHI/Private/VulkanMemory.cpp
--- a/Engine/Source/Runtime/VulkanRHI/Private/VulkanMemory.cpp
+++ b/Engine/Source/Runtime/VulkanRHI/Private/VulkanMemory.cpp
@@ -1808,6 +1808,12 @@ namespace VulkanRHI
 			Size = AlignArbitrary(Size, NonCoherentAtomSize);
 		}
 
+		// Add both source and dest flags
+		if ((InUsageFlags & (VK_BUFFER_USAGE_TRANSFER_SRC_BIT | VK_BUFFER_USAGE_TRANSFER_DST_BIT)) != 0)
+		{
+			InUsageFlags |= (VK_BUFFER_USAGE_TRANSFER_SRC_BIT | VK_BUFFER_USAGE_TRANSFER_DST_BIT);
+		}
+
 		//#todo-rco: Better locking!
 		{
 			FScopeLock Lock(&GStagingLock);
