bHasInputAttachments was never initialized

diff --git a/Engine/Source/Runtime/VulkanRHI/Private/VulkanPipeline.cpp b/Engine/Source/Runtime/VulkanRHI/Private/VulkanPipeline.cpp
--- a/Engine/Source/Runtime/VulkanRHI/Private/VulkanPipeline.cpp
+++ b/Engine/Source/Runtime/VulkanRHI/Private/VulkanPipeline.cpp
@@ -1699,6 +1699,7 @@ FVulkanRHIGraphicsPipelineState* FVulkanPipelineStateCacheManager::RHICreateGrap
 			FVulkanGfxLayout* GfxLayout = (FVulkanGfxLayout*)Layout;
 			check(GfxLayout->GfxPipelineDescriptorInfo.IsInitialized());
 			NewPSO->Layout = GfxLayout;
+			NewPSO->bHasInputAttachments = GfxLayout->GetDescriptorSetsLayout().HasInputAttachments();
 		}
 		NewPSO->RenderPass = Device->GetImmediateContext().PrepareRenderPassForPSOCreation(Initializer);
 		{
diff --git a/Engine/Source/Runtime/VulkanRHI/Private/VulkanPipeline.h b/Engine/Source/Runtime/VulkanRHI/Private/VulkanPipeline.h
--- a/Engine/Source/Runtime/VulkanRHI/Private/VulkanPipeline.h
+++ b/Engine/Source/Runtime/VulkanRHI/Private/VulkanPipeline.h
@@ -688,7 +688,7 @@ public:
 
 
 
-	bool								bHasInputAttachments;
+	bool								bHasInputAttachments = false;
 	bool								bIsRegistered;
 	bool								bShaderModulesLoaded = false;
 
