Implement support for ReadOnly Depth/Writable Stencil in vulkan.
Fixes an ensure on startup & Vulkan validation errors.

diff --git a/Engine/Source/Runtime/VulkanRHI/Private/VulkanContext.h b/Engine/Source/Runtime/VulkanRHI/Private/VulkanContext.h
--- a/Engine/Source/Runtime/VulkanRHI/Private/VulkanContext.h
+++ b/Engine/Source/Runtime/VulkanRHI/Private/VulkanContext.h
@@ -212,6 +212,7 @@ public:
 	virtual void RHICopyTexture(FRHITexture* SourceTexture, FRHITexture* DestTexture, const FRHICopyTextureInfo& CopyInfo) final override;
 	virtual void RHITransitionResources(EResourceTransitionAccess TransitionType, FRHITexture** InRenderTargets, int32 NumTextures) final override;
 	virtual void RHITransitionResources(EResourceTransitionAccess TransitionType, EResourceTransitionPipeline TransitionPipeline, FRHIUnorderedAccessView** InUAVs, int32 NumUAVs, FRHIComputeFence* WriteComputeFence) final override;
+	virtual void RHITransitionResources(FExclusiveDepthStencil DepthStencilMode, FRHITexture* DepthTexture) final override;
 	virtual void RHICopyToStagingBuffer(FRHIVertexBuffer* SourceBuffer, FRHIStagingBuffer* DestinationStagingBuffer, uint32 Offset, uint32 NumBytes) final override;
 	virtual void RHIWriteGPUFence(FRHIGPUFence* Fence) final override;
 
diff --git a/Engine/Source/Runtime/VulkanRHI/Private/VulkanRenderTarget.cpp b/Engine/Source/Runtime/VulkanRHI/Private/VulkanRenderTarget.cpp
--- a/Engine/Source/Runtime/VulkanRHI/Private/VulkanRenderTarget.cpp
+++ b/Engine/Source/Runtime/VulkanRHI/Private/VulkanRenderTarget.cpp
@@ -1325,6 +1325,57 @@ bool FVulkanCommandListContext::FPendingTransition::GatherBarriers(FTransitionAn
 }
 
 
+
+void FVulkanCommandListContext::RHITransitionResources(FExclusiveDepthStencil DepthStencilMode, FRHITexture* DepthTexture)
+{
+	static IConsoleVariable* CVarShowTransitions = IConsoleManager::Get().FindConsoleVariable(TEXT("r.ProfileGPU.ShowTransitions"));
+	bool bShowTransitionEvents = CVarShowTransitions->GetInt() != 0;
+	FVulkanCmdBuffer* CmdBuffer = CommandBufferManager->GetActiveCmdBuffer();
+	check(CmdBuffer->HasBegun());
+	check(!TransitionAndLayoutManager.CurrentRenderPass);
+
+	if (bShowTransitionEvents)
+	{
+		SCOPED_RHI_DRAW_EVENTF(*this, RHITransitionResourcesLoop, TEXT("To:%s"), *DepthTexture->GetName().ToString());
+	}
+
+	FVulkanTextureBase* VulkanTexture = FVulkanTextureBase::Cast(DepthTexture);
+
+	VulkanRHI::FPendingBarrier Barrier;
+	VkImageLayout& SrcLayout = TransitionAndLayoutManager.FindOrAddLayoutRW(VulkanTexture->Surface.Image, VK_IMAGE_LAYOUT_UNDEFINED);
+	check(VulkanTexture->Surface.IsDepthOrStencilAspect())
+	VkImageLayout DstLayout = SrcLayout;
+	
+	if(DepthStencilMode.IsDepthWrite())
+	{
+		if(DepthStencilMode.IsStencilWrite())
+		{
+			DstLayout = VK_IMAGE_LAYOUT_DEPTH_STENCIL_ATTACHMENT_OPTIMAL;
+		}
+		else
+		{
+			DstLayout = VK_IMAGE_LAYOUT_DEPTH_ATTACHMENT_STENCIL_READ_ONLY_OPTIMAL;
+		}
+	}
+	else
+	{
+		if (DepthStencilMode.IsStencilWrite())
+		{
+			DstLayout = VK_IMAGE_LAYOUT_DEPTH_READ_ONLY_STENCIL_ATTACHMENT_OPTIMAL;
+		}
+		else
+		{
+			DstLayout  = VK_IMAGE_LAYOUT_DEPTH_STENCIL_READ_ONLY_OPTIMAL;
+		}
+	}
+		
+	int32 BarrierIndex = Barrier.AddImageBarrier(VulkanTexture->Surface.Image, VulkanTexture->Surface.GetFullAspectMask(), VulkanTexture->Surface.GetNumMips(), VulkanTexture->Surface.NumArrayLevels);
+	Barrier.SetTransition(BarrierIndex, VulkanRHI::GetImageLayoutFromVulkanLayout(SrcLayout), VulkanRHI::GetImageLayoutFromVulkanLayout(DstLayout));
+	SrcLayout = DstLayout;
+	Barrier.Execute(CmdBuffer, false);
+}
+
+
 void FVulkanCommandListContext::TransitionResources(const FPendingTransition& PendingTransition)
 {
 	static IConsoleVariable* CVarShowTransitions = IConsoleManager::Get().FindConsoleVariable(TEXT("r.ProfileGPU.ShowTransitions"));
