fix for vulkan debug names accidentally always being enabled.

diff --git a/Engine/Source/Runtime/VulkanRHI/Private/VulkanDevice.h b/Engine/Source/Runtime/VulkanRHI/Private/VulkanDevice.h
--- a/Engine/Source/Runtime/VulkanRHI/Private/VulkanDevice.h
+++ b/Engine/Source/Runtime/VulkanRHI/Private/VulkanDevice.h
@@ -18,9 +18,9 @@ class FVulkanOcclusionQueryPool;
 class FOLDVulkanQueryPool;
 #endif
 
-#define VULKAN_USE_DEBUG_NAMES 0
+#define VULKAN_USE_DEBUG_NAMES 1
 
-#ifdef VULKAN_USE_DEBUG_NAMES
+#if VULKAN_USE_DEBUG_NAMES
 #define VULKAN_SET_DEBUG_NAME(Device, Type, Handle, Format, ...) Device.VulkanSetObjectName(Type, (uint64)Handle, *FString::Printf(Format, __VA_ARGS__))
 #else
 #define VULKAN_SET_DEBUG_NAME(Device, Type, Handle, Format, ...) do{}while(0)
