Remove Android Vulkan crash handler override.
#rnx

diff --git a/Engine/Source/Runtime/VulkanRHI/Private/Android/VulkanAndroidPlatform.cpp b/Engine/Source/Runtime/VulkanRHI/Private/Android/VulkanAndroidPlatform.cpp
--- a/Engine/Source/Runtime/VulkanRHI/Private/Android/VulkanAndroidPlatform.cpp
+++ b/Engine/Source/Runtime/VulkanRHI/Private/Android/VulkanAndroidPlatform.cpp
@@ -188,15 +188,12 @@ void FVulkanAndroidPlatform::OverridePlatformHandlers(bool bInit)
 {
 	if (bInit)
 	{
-		// Want to see the actual crash report on Android so unregister signal handlers
-		FPlatformMisc::SetCrashHandler((void(*)(const FGenericCrashContext& Context)) -1);
 		FPlatformMisc::SetOnReInitWindowCallback(FVulkanDynamicRHI::RecreateSwapChain);
 		FPlatformMisc::SetOnReleaseWindowCallback(FVulkanDynamicRHI::DestroySwapChain);
 		FPlatformMisc::SetOnPauseCallback(FVulkanDynamicRHI::SavePipelineCache);
 	}
 	else
 	{
-		FPlatformMisc::SetCrashHandler(nullptr);
 		FPlatformMisc::SetOnReInitWindowCallback(nullptr);
 		FPlatformMisc::SetOnReleaseWindowCallback(nullptr);
 		FPlatformMisc::SetOnPauseCallback(nullptr);
