F - Some IWYU prep fixes for the Vulkan RHI
- Remove VulkanGlobals.h file
- VulkanLoader.h is now included in the generic platform
#rnx

diff --git a/Engine/Source/Runtime/VulkanRHI/Private/Android/VulkanAndroidPlatform.h b/Engine/Source/Runtime/VulkanRHI/Private/Android/VulkanAndroidPlatform.h
--- a/Engine/Source/Runtime/VulkanRHI/Private/Android/VulkanAndroidPlatform.h
+++ b/Engine/Source/Runtime/VulkanRHI/Private/Android/VulkanAndroidPlatform.h
@@ -36,8 +36,6 @@
 	EnumMacro(PFN_vkGetRefreshCycleDurationGOOGLE, vkGetRefreshCycleDurationGOOGLE) \
 	EnumMacro(PFN_vkGetPastPresentationTimingGOOGLE, vkGetPastPresentationTimingGOOGLE)
 
-#include "../VulkanLoader.h"
-
 // and now, include the GenericPlatform class
 #include "../VulkanGenericPlatform.h"
 
diff --git a/Engine/Source/Runtime/VulkanRHI/Private/Linux/VulkanLinuxPlatform.h b/Engine/Source/Runtime/VulkanRHI/Private/Linux/VulkanLinuxPlatform.h
--- a/Engine/Source/Runtime/VulkanRHI/Private/Linux/VulkanLinuxPlatform.h
+++ b/Engine/Source/Runtime/VulkanRHI/Private/Linux/VulkanLinuxPlatform.h
@@ -25,8 +25,6 @@
     EnumMacro(PFN_vkGetImageMemoryRequirements2KHR , vkGetImageMemoryRequirements2KHR) \
     EnumMacro(PFN_vkGetBufferMemoryRequirements2KHR , vkGetBufferMemoryRequirements2KHR)
 
-#include "../VulkanLoader.h"
-
 // and now, include the GenericPlatform class
 #include "../VulkanGenericPlatform.h"
 
diff --git a/Engine/Source/Runtime/VulkanRHI/Private/Lumin/VulkanLuminPlatform.h b/Engine/Source/Runtime/VulkanRHI/Private/Lumin/VulkanLuminPlatform.h
--- a/Engine/Source/Runtime/VulkanRHI/Private/Lumin/VulkanLuminPlatform.h
+++ b/Engine/Source/Runtime/VulkanRHI/Private/Lumin/VulkanLuminPlatform.h
@@ -23,8 +23,6 @@
 
 #define ENUM_VK_ENTRYPOINTS_OPTIONAL_PLATFORM_INSTANCE(EnumMacro)
 
-#include "../VulkanLoader.h"
-
 // and now, include the GenericPlatform class
 #include "../VulkanGenericPlatform.h"
 
diff --git a/Engine/Source/Runtime/VulkanRHI/Private/VulkanGenericPlatform.h b/Engine/Source/Runtime/VulkanRHI/Private/VulkanGenericPlatform.h
--- a/Engine/Source/Runtime/VulkanRHI/Private/VulkanGenericPlatform.h
+++ b/Engine/Source/Runtime/VulkanRHI/Private/VulkanGenericPlatform.h
@@ -6,6 +6,8 @@
 #include "Containers/ArrayView.h"
 #include "RHI.h"	// for GShaderPlatformForFeatureLevel and its friends
 
+#include "../VulkanLoader.h"
+
 struct FOptionalVulkanDeviceExtensions;
 class FVulkanDevice;
 
diff --git a/Engine/Source/Runtime/VulkanRHI/Private/VulkanPendingState.h b/Engine/Source/Runtime/VulkanRHI/Private/VulkanPendingState.h
--- a/Engine/Source/Runtime/VulkanRHI/Private/VulkanPendingState.h
+++ b/Engine/Source/Runtime/VulkanRHI/Private/VulkanPendingState.h
@@ -7,7 +7,6 @@
 #pragma once
 
 // Dependencies
-#include "VulkanGlobals.h"
 #include "VulkanConfiguration.h"
 #include "VulkanState.h"
 #include "VulkanResources.h"
diff --git a/Engine/Source/Runtime/VulkanRHI/Private/VulkanRHI.cpp b/Engine/Source/Runtime/VulkanRHI/Private/VulkanRHI.cpp
--- a/Engine/Source/Runtime/VulkanRHI/Private/VulkanRHI.cpp
+++ b/Engine/Source/Runtime/VulkanRHI/Private/VulkanRHI.cpp
@@ -21,6 +21,7 @@
 #include "GlobalShader.h"
 #include "RHIValidation.h"
 
+static_assert(sizeof(VkStructureType) == sizeof(int32), "ZeroVulkanStruct() assumes VkStructureType is int32!");
 
 extern RHI_API bool GUseTexture3DBulkDataRHI;
 
diff --git a/Engine/Source/Runtime/VulkanRHI/Private/VulkanRHIPrivate.h b/Engine/Source/Runtime/VulkanRHI/Private/VulkanRHIPrivate.h
--- a/Engine/Source/Runtime/VulkanRHI/Private/VulkanRHIPrivate.h
+++ b/Engine/Source/Runtime/VulkanRHI/Private/VulkanRHIPrivate.h
@@ -36,7 +36,6 @@
 	#endif
 #endif
 
-#include "VulkanGlobals.h"
 #include "VulkanState.h"
 #include "VulkanResources.h"
 #include "VulkanUtil.h"
@@ -906,6 +905,11 @@ namespace VulkanRHI
 	}
 }
 
+inline bool UseVulkanDescriptorCache()
+{
+	return (PLATFORM_ANDROID && !PLATFORM_LUMIN)|| GMaxRHIFeatureLevel <= ERHIFeatureLevel::ES3_1;
+}
+
 extern int32 GVulkanSubmitAfterEveryEndRenderPass;
 extern int32 GWaitForIdleOnSubmit;
 extern bool GGPUCrashDebuggingEnabled;
@@ -917,3 +921,9 @@ extern bool GRenderDocFound;
 const int GMaxCrashBufferEntries = 2048;
 
 extern class FVulkanDynamicRHI*	GVulkanRHI;
+
+extern TAtomic<uint64> GVulkanBufferHandleIdCounter;
+extern TAtomic<uint64> GVulkanBufferViewHandleIdCounter;
+extern TAtomic<uint64> GVulkanImageViewHandleIdCounter;
+extern TAtomic<uint64> GVulkanSamplerHandleIdCounter;
+extern TAtomic<uint64> GVulkanDSetLayoutHandleIdCounter;
diff --git a/Engine/Source/Runtime/VulkanRHI/Private/Windows/VulkanWindowsPlatform.h b/Engine/Source/Runtime/VulkanRHI/Private/Windows/VulkanWindowsPlatform.h
--- a/Engine/Source/Runtime/VulkanRHI/Private/Windows/VulkanWindowsPlatform.h
+++ b/Engine/Source/Runtime/VulkanRHI/Private/Windows/VulkanWindowsPlatform.h
@@ -39,8 +39,6 @@
 	EnumMacro(PFN_vkCreateSamplerYcbcrConversionKHR, vkCreateSamplerYcbcrConversionKHR) \
 	EnumMacro(PFN_vkDestroySamplerYcbcrConversionKHR, vkDestroySamplerYcbcrConversionKHR)
 
-#include "../VulkanLoader.h"
-
 // and now, include the GenericPlatform class
 #include "../VulkanGenericPlatform.h"
 
diff --git a/Engine/Source/Runtime/VulkanRHI/Public/VulkanCommon.h b/Engine/Source/Runtime/VulkanRHI/Public/VulkanCommon.h
--- a/Engine/Source/Runtime/VulkanRHI/Public/VulkanCommon.h
+++ b/Engine/Source/Runtime/VulkanRHI/Public/VulkanCommon.h
@@ -141,3 +141,16 @@ namespace EVulkanBindingType
 		return 0;
 	}
 }
+
+DECLARE_LOG_CATEGORY_EXTERN(LogVulkan, Display, All);
+
+template< class T >
+static FORCEINLINE void ZeroVulkanStruct(T& Struct, int32 VkStructureType)
+{
+	static_assert(!TIsPointer<T>::Value, "Don't use a pointer!");
+	static_assert(STRUCT_OFFSET(T, sType) == 0, "Assumes sType is the first member in the Vulkan type!");
+	static_assert(sizeof(T::sType) == sizeof(int32), "Assumed sType is compatible with int32!");
+	// Horrible way to coerce the compiler to not have to know what T::sType is so we can have this header not have to include vulkan.h
+	(int32&)Struct.sType = VkStructureType;
+	FMemory::Memzero(((uint8*)&Struct) + sizeof(VkStructureType), sizeof(T) - sizeof(VkStructureType));
+}
diff --git a/Engine/Source/Runtime/VulkanRHI/Public/VulkanDynamicRHI.h b/Engine/Source/Runtime/VulkanRHI/Public/VulkanDynamicRHI.h
--- a/Engine/Source/Runtime/VulkanRHI/Public/VulkanDynamicRHI.h
+++ b/Engine/Source/Runtime/VulkanRHI/Public/VulkanDynamicRHI.h
@@ -6,8 +6,6 @@
 
 #pragma once 
 
-#include "VulkanGlobals.h"
-
 class FVulkanTexture2D;
 class FVulkanFramebuffer;
 class FVulkanDevice;
diff --git a/Engine/Source/Runtime/VulkanRHI/Public/VulkanGlobals.h b/Engine/Source/Runtime/VulkanRHI/Public/VulkanGlobals.h
deleted file mode 100644
