F - Move Vulkan profiler out to its own file for IWYU prep
#rnx

diff --git a/Engine/Source/Runtime/VulkanRHI/Private/VulkanCommandBuffer.h b/Engine/Source/Runtime/VulkanRHI/Private/VulkanCommandBuffer.h
--- a/Engine/Source/Runtime/VulkanRHI/Private/VulkanCommandBuffer.h
+++ b/Engine/Source/Runtime/VulkanRHI/Private/VulkanCommandBuffer.h
@@ -15,6 +15,7 @@ class FVulkanCommandBufferManager;
 class FVulkanRenderTargetLayout;
 class FVulkanQueue;
 class FVulkanDescriptorPoolSetContainer;
+class FVulkanGPUTiming;
 
 namespace VulkanRHI
 {
diff --git a/Engine/Source/Runtime/VulkanRHI/Private/VulkanContext.h b/Engine/Source/Runtime/VulkanRHI/Private/VulkanContext.h
--- a/Engine/Source/Runtime/VulkanRHI/Private/VulkanContext.h
+++ b/Engine/Source/Runtime/VulkanRHI/Private/VulkanContext.h
@@ -7,6 +7,7 @@
 #pragma once 
 
 #include "VulkanResources.h"
+#include "VulkanGPUProfiler.h"
 
 class FVulkanDevice;
 class FVulkanCommandBufferManager;
diff --git a/Engine/Source/Runtime/VulkanRHI/Private/VulkanGPUProfiler.h b/Engine/Source/Runtime/VulkanRHI/Private/VulkanGPUProfiler.h
new file mode 100644
--- /dev/null
+++ b/Engine/Source/Runtime/VulkanRHI/Private/VulkanGPUProfiler.h
@@ -0,0 +1,191 @@
+// Copyright 1998-2019 Epic Games, Inc. All Rights Reserved.
+
+/*=============================================================================
+	VulkanGPUProfiler.h: Vulkan Utility definitions.
+=============================================================================*/
+
+#pragma once
+
+#include "GPUProfiler.h"
+
+#include "Serialization/MemoryWriter.h"
+
+class FVulkanCmdBuffer;
+class FVulkanCommandListContext;
+class FVulkanTimingQueryPool;
+
+class FVulkanGPUTiming : public FGPUTiming
+{
+public:
+	FVulkanGPUTiming(FVulkanCommandListContext* InCmd, FVulkanDevice* InDevice)
+		: Device(InDevice)
+		, bIsTiming(false)
+		, bEndTimestampIssued(false)
+		, CmdContext(InCmd)
+	{
+	}
+
+	~FVulkanGPUTiming();
+
+	/**
+	 * Start a GPU timing measurement.
+	 */
+	void StartTiming(FVulkanCmdBuffer* CmdBuffer = nullptr);
+
+	/**
+	 * End a GPU timing measurement.
+	 * The timing for this particular measurement will be resolved at a later time by the GPU.
+	 */
+	void EndTiming(FVulkanCmdBuffer* CmdBuffer = nullptr);
+
+	/**
+	 * Retrieves the most recently resolved timing measurement.
+	 * The unit is the same as for FPlatformTime::Cycles(). Returns 0 if there are no resolved measurements.
+	 *
+	 * @return	Value of the most recently resolved timing, or 0 if no measurements have been resolved by the GPU yet.
+	 */
+	uint64 GetTiming(bool bGetCurrentResultsAndBlock = false);
+
+	/**
+	 * Initializes all Vulkan resources.
+	 */
+	void Initialize();
+
+	/**
+	 * Releases all Vulkan resources.
+	 */
+	void Release();
+
+	bool IsComplete() const
+	{
+		check(bEndTimestampIssued);
+		return true;
+	}
+
+	static void CalibrateTimers(FVulkanCommandListContext& InCmdContext);
+
+private:
+	/**
+	 * Initializes the static variables, if necessary.
+	 */
+	static void PlatformStaticInitialize(void* UserData);
+
+	FVulkanDevice* Device;
+
+	/** Whether we are currently timing the GPU: between StartTiming() and EndTiming(). */
+	bool bIsTiming;
+	bool bEndTimestampIssued;
+
+	FVulkanCommandListContext* CmdContext;
+	FVulkanTimingQueryPool* Pool = nullptr;
+};
+
+/** A single perf event node, which tracks information about a appBeginDrawEvent/appEndDrawEvent range. */
+class FVulkanEventNode : public FGPUProfilerEventNode
+{
+public:
+	FVulkanEventNode(const TCHAR* InName, FGPUProfilerEventNode* InParent, FVulkanCommandListContext* InCmd, FVulkanDevice* InDevice) :
+		FGPUProfilerEventNode(InName, InParent),
+		Timing(InCmd, InDevice)
+	{
+		// Initialize Buffered timestamp queries 
+		Timing.Initialize();
+	}
+
+	virtual ~FVulkanEventNode()
+	{
+		Timing.Release(); 
+	}
+
+	/** 
+	 * Returns the time in ms that the GPU spent in this draw event.  
+	 * This blocks the CPU if necessary, so can cause hitching.
+	 */
+	virtual float GetTiming() override final;
+
+
+	virtual void StartTiming() override final
+	{
+		Timing.StartTiming();
+	}
+
+	virtual void StopTiming() override final
+	{
+		Timing.EndTiming();
+	}
+
+	FVulkanGPUTiming Timing;
+};
+
+/** An entire frame of perf event nodes, including ancillary timers. */
+class FVulkanEventNodeFrame : public FGPUProfilerEventNodeFrame
+{
+public:
+
+	FVulkanEventNodeFrame(FVulkanCommandListContext* InCmd, FVulkanDevice* InDevice)
+		: RootEventTiming(InCmd, InDevice)
+	{
+		RootEventTiming.Initialize();
+	}
+
+	~FVulkanEventNodeFrame()
+	{
+		RootEventTiming.Release();
+	}
+
+	/** Start this frame of per tracking */
+	virtual void StartFrame() override final;
+
+	/** End this frame of per tracking, but do not block yet */
+	virtual void EndFrame() override final;
+
+	/** Calculates root timing base frequency (if needed by this RHI) */
+	virtual float GetRootTimingResults() override final;
+
+	virtual bool PlatformDisablesVSync() const { return true; }
+
+	/** Timer tracking inclusive time spent in the root nodes. */
+	FVulkanGPUTiming RootEventTiming;
+};
+
+/** 
+ * Encapsulates GPU profiling logic and data. 
+ * There's only one global instance of this struct so it should only contain global data, nothing specific to a frame.
+ */
+struct FVulkanGPUProfiler : public FGPUProfiler
+{
+	/** GPU hitch profile histories */
+	TIndirectArray<FVulkanEventNodeFrame> GPUHitchEventNodeFrames;
+
+	FVulkanGPUProfiler(FVulkanCommandListContext* InCmd, FVulkanDevice* InDevice)
+		: bCommandlistSubmitted(false)
+		, Device(InDevice)
+		, CmdContext(InCmd)
+	{
+	}
+
+	virtual FGPUProfilerEventNode* CreateEventNode(const TCHAR* InName, FGPUProfilerEventNode* InParent) override final
+	{
+		FVulkanEventNode* EventNode = new FVulkanEventNode(InName, InParent, CmdContext, Device);
+		return EventNode;
+	}
+
+	void BeginFrame();
+
+	void EndFrameBeforeSubmit();
+	void EndFrame();
+
+	bool bCommandlistSubmitted;
+	FVulkanDevice* Device;
+	FVulkanCommandListContext* CmdContext;
+
+#if VULKAN_SUPPORTS_GPU_CRASH_DUMPS
+	void PushMarkerForCrash(VkCommandBuffer CmdBuffer, VkBuffer DestBuffer, const TCHAR* Name);
+	void PopMarkerForCrash(VkCommandBuffer CmdBuffer, VkBuffer DestBuffer);
+	void DumpCrashMarkers(void* BufferData);
+#endif
+
+	// For crash/marker tracking
+	TMap<uint32, FString> CachedStrings;
+	TArray<uint32> PushPopStack;
+};
diff --git a/Engine/Source/Runtime/VulkanRHI/Public/VulkanMemory.h b/Engine/Source/Runtime/VulkanRHI/Public/VulkanMemory.h
--- a/Engine/Source/Runtime/VulkanRHI/Public/VulkanMemory.h
+++ b/Engine/Source/Runtime/VulkanRHI/Public/VulkanMemory.h
@@ -18,7 +18,7 @@
 #define VULKAN_MEMORY_HIGHER_PRIORITY 0.75f
 #define VULKAN_MEMORY_HIGHEST_PRIORITY 1.f
 
-
+class FVulkanCommandListContext;
 class FVulkanQueue;
 class FVulkanCmdBuffer;
 
diff --git a/Engine/Source/Runtime/VulkanRHI/Public/VulkanUtil.h b/Engine/Source/Runtime/VulkanRHI/Public/VulkanUtil.h
--- a/Engine/Source/Runtime/VulkanRHI/Public/VulkanUtil.h
+++ b/Engine/Source/Runtime/VulkanRHI/Public/VulkanUtil.h
@@ -6,6 +6,7 @@
 
 #pragma once
 
+#if 0
 #include "GPUProfiler.h"
 
 #include "Serialization/MemoryWriter.h"
@@ -157,9 +158,12 @@ struct FVulkanGPUProfiler : public FGPUProfiler
 	/** GPU hitch profile histories */
 	TIndirectArray<FVulkanEventNodeFrame> GPUHitchEventNodeFrames;
 
-	FVulkanGPUProfiler(FVulkanCommandListContext* InCmd, FVulkanDevice* InDevice);
-
-	virtual ~FVulkanGPUProfiler();
+	FVulkanGPUProfiler(FVulkanCommandListContext* InCmd, FVulkanDevice* InDevice)
+		: bCommandlistSubmitted(false)
+		, Device(InDevice)
+		, CmdContext(InCmd)
+	{
+	}
 
 	virtual FGPUProfilerEventNode* CreateEventNode(const TCHAR* InName, FGPUProfilerEventNode* InParent) override final
 	{
@@ -185,11 +189,8 @@ struct FVulkanGPUProfiler : public FGPUProfiler
 	// For crash/marker tracking
 	TMap<uint32, FString> CachedStrings;
 	TArray<uint32> PushPopStack;
-
-	FVulkanTimingQueryPool* LocalTracePointsQueryPool;
-	TArray<uint64> CrashMarkers;
-	bool bBeginFrame;
 };
+#endif
 
 namespace VulkanRHI
 {
