F - Compile non-unity fix
#rnx

diff --git a/Engine/Source/Runtime/VulkanRHI/Private/VulkanGPUProfiler.h b/Engine/Source/Runtime/VulkanRHI/Private/VulkanGPUProfiler.h
--- a/Engine/Source/Runtime/VulkanRHI/Private/VulkanGPUProfiler.h
+++ b/Engine/Source/Runtime/VulkanRHI/Private/VulkanGPUProfiler.h
@@ -8,8 +8,6 @@
 
 #include "GPUProfiler.h"
 
-#include "Serialization/MemoryWriter.h"
-
 class FVulkanCmdBuffer;
 class FVulkanCommandListContext;
 class FVulkanTimingQueryPool;
diff --git a/Engine/Source/Runtime/VulkanRHI/Public/VulkanUtil.h b/Engine/Source/Runtime/VulkanRHI/Public/VulkanUtil.h
--- a/Engine/Source/Runtime/VulkanRHI/Public/VulkanUtil.h
+++ b/Engine/Source/Runtime/VulkanRHI/Public/VulkanUtil.h
@@ -6,6 +6,8 @@
 
 #pragma once
 
+#include "Serialization/MemoryWriter.h"
+
 namespace VulkanRHI
 {
 	/**
