Fix incorrect path for include
#rnx

diff --git a/Engine/Source/Runtime/VulkanRHI/Private/VulkanGenericPlatform.h b/Engine/Source/Runtime/VulkanRHI/Private/VulkanGenericPlatform.h
--- a/Engine/Source/Runtime/VulkanRHI/Private/VulkanGenericPlatform.h
+++ b/Engine/Source/Runtime/VulkanRHI/Private/VulkanGenericPlatform.h
@@ -6,7 +6,7 @@
 #include "Containers/ArrayView.h"
 #include "RHI.h"	// for GShaderPlatformForFeatureLevel and its friends
 
-#include "../VulkanLoader.h"
+#include "VulkanLoader.h"
 
 struct FOptionalVulkanDeviceExtensions;
 class FVulkanDevice;
