fix for volatile buffers being created too large, when the volatile allocation is smaller than the original allocation

diff --git a/Engine/Source/Runtime/VulkanRHI/Private/VulkanMemory.cpp b/Engine/Source/Runtime/VulkanRHI/Private/VulkanMemory.cpp
--- a/Engine/Source/Runtime/VulkanRHI/Private/VulkanMemory.cpp
+++ b/Engine/Source/Runtime/VulkanRHI/Private/VulkanMemory.cpp
@@ -2321,6 +2321,7 @@ namespace VulkanRHI
 			OutInfo.Data = AlignedData;
 			OutInfo.BufferSuballocation = BufferSuballocation;
 			OutInfo.CurrentOffset = (uint32)(AlignedData - MappedData);
+			OutInfo.Size = InSize;
 			CurrentData = AlignedData + InSize;
 			PeakUsed = FMath::Max(PeakUsed, (uint32)(CurrentData - MappedData));
 			return true;
diff --git a/Engine/Source/Runtime/VulkanRHI/Private/VulkanResources.h b/Engine/Source/Runtime/VulkanRHI/Private/VulkanResources.h
--- a/Engine/Source/Runtime/VulkanRHI/Private/VulkanResources.h
+++ b/Engine/Source/Runtime/VulkanRHI/Private/VulkanResources.h
@@ -1095,6 +1095,11 @@ public:
 		check(IsVolatile());
 		return VolatileLockInfo.LockCounter;
 	}
+	inline uint32 GetVolatileLockSize() const
+	{
+		check(IsVolatile());
+		return VolatileLockInfo.Size;
+	}
 
 	inline int32 GetNumBuffers() const
 	{
diff --git a/Engine/Source/Runtime/VulkanRHI/Private/VulkanUAV.cpp b/Engine/Source/Runtime/VulkanRHI/Private/VulkanUAV.cpp
--- a/Engine/Source/Runtime/VulkanRHI/Private/VulkanUAV.cpp
+++ b/Engine/Source/Runtime/VulkanRHI/Private/VulkanUAV.cpp
@@ -76,6 +76,9 @@ void FVulkanShaderResourceView::UpdateView()
 		{
 			VkBuffer SourceVolatileBufferHandle = SourceBuffer->GetHandle();
 
+			//Inherit Size from volatile buffer, as it might change
+			Size = SourceBuffer->GetVolatileLockSize();
+
 			// We might end up with the same BufferView, so do not recreate in that case
 			if (!BufferViews[0] 
 				|| BufferViews[0]->Offset != SourceBuffer->GetOffset() 
diff --git a/Engine/Source/Runtime/VulkanRHI/Public/VulkanMemory.h b/Engine/Source/Runtime/VulkanRHI/Public/VulkanMemory.h
--- a/Engine/Source/Runtime/VulkanRHI/Public/VulkanMemory.h
+++ b/Engine/Source/Runtime/VulkanRHI/Public/VulkanMemory.h
@@ -1290,6 +1290,9 @@ namespace VulkanRHI
 			// Offset into the locked area
 			uint32 CurrentOffset;
 
+			// Size of Allocation
+			uint32 Size;
+
 			// Simple counter used for the SRVs to know a new one is required
 			uint32 LockCounter;
 
@@ -1297,6 +1300,7 @@ namespace VulkanRHI
 				: Data(nullptr)
 				, BufferSuballocation(nullptr)
 				, CurrentOffset(0)
+				, Size(0)
 				, LockCounter(0)
 			{
 			}
