Added Option to prevent multithreaded PSO creation.
Temporarily disabled multithreaded PSO creation on some problematic devices.

diff --git a/Engine/Source/Runtime/VulkanRHI/Private/VulkanPipeline.cpp b/Engine/Source/Runtime/VulkanRHI/Private/VulkanPipeline.cpp
--- a/Engine/Source/Runtime/VulkanRHI/Private/VulkanPipeline.cpp
+++ b/Engine/Source/Runtime/VulkanRHI/Private/VulkanPipeline.cpp
@@ -126,6 +126,15 @@ static FAutoConsoleVariableRef GEnablePipelineCacheCompressionCvar(
 );
 
 
+static int32 GVulkanPSOForceSingleThreaded = 0;
+static FAutoConsoleVariableRef GVulkanPSOForceSingleThreadedCVar(
+	TEXT("r.Vulkan.ForcePSOSingleThreaded"),
+	GVulkanPSOForceSingleThreaded,
+	TEXT("Enable to force singlethreaded creation of PSOs. Only intended as a workaround for buggy drivers\n"),
+	ECVF_ReadOnly | ECVF_RenderThreadSafe
+);
+
+
 
 template <typename TRHIType, typename TVulkanType>
 static inline FSHAHash GetShaderHash(TRHIType* RHIShader)
@@ -553,7 +562,7 @@ void FVulkanPipelineStateCacheManager::Save(const FString& CacheFilename, bool b
 		}
 		else if (Result == VK_INCOMPLETE || Result == VK_ERROR_OUT_OF_HOST_MEMORY)
 		{
-			UE_LOG(LogVulkanRHI, Warning, TEXT("Failed to get Vulkan pipeline cache data."));
+			UE_LOG(LogVulkanRHI, Warning, TEXT("Failed to get Vulkan pipeline cache data. Error %d, %d bytes"), Result, Size);
 
 			VulkanRHI::vkDestroyPipelineCache(Device->GetInstanceHandle(), PipelineCache, VULKAN_CPU_ALLOCATOR);
 			VkPipelineCacheCreateInfo PipelineCacheInfo;
@@ -1659,9 +1668,33 @@ void FVulkanPipelineStateCacheManager::NotifyDeletedGraphicsPSO(FRHIGraphicsPipe
 }
 
 
+//Global lock for PSO creation, only enabled if GVulkanPSOForceSingleThreaded is 1
+struct FPSOGlobalLock
+{
+	FCriticalSection* CriticalSection;
+	FPSOGlobalLock(FCriticalSection* InSynchObject)
+	{
+		
+		CriticalSection = GVulkanPSOForceSingleThreaded ? InSynchObject : nullptr;
+		if (CriticalSection)
+		{
+			CriticalSection->Lock();
+		}
+	}
+	~FPSOGlobalLock()
+	{
+		if (CriticalSection)
+		{
+			CriticalSection->Unlock();
+		}
+	}
+
+};
+
 FVulkanRHIGraphicsPipelineState* FVulkanPipelineStateCacheManager::RHICreateGraphicsPipelineState(const FGraphicsPipelineStateInitializer& Initializer)
 {
 	QUICK_SCOPE_CYCLE_COUNTER(STAT_Vulkan_RHICreateGraphicsPipelineState_NEW);
+	FPSOGlobalLock GlobalLock(&GraphicsPSOLockedCS);
 	FVulkanPSOKey Key;
 	FGfxPipelineDesc Desc;
 	FVulkanDescriptorSetsLayoutInfo DescriptorSetLayoutInfo;
