Fix memory stomp when passing the Texture->SizeX()/SizeY() as the destition Width/Height when our Out is determined by bounding Rect.Max/Min

diff --git a/Engine/Source/Runtime/VulkanRHI/Private/VulkanRenderTarget.cpp b/Engine/Source/Runtime/VulkanRHI/Private/VulkanRenderTarget.cpp
--- a/Engine/Source/Runtime/VulkanRHI/Private/VulkanRenderTarget.cpp
+++ b/Engine/Source/Runtime/VulkanRHI/Private/VulkanRenderTarget.cpp
@@ -876,40 +876,43 @@ void FVulkanDynamicRHI::RHIReadSurfaceData(FRHITexture* TextureRHI, FIntRect Rec
 	OutData.SetNum(NumPixels);
 	FColor* Dest = OutData.GetData();
 
+	uint32 DestWidth  = Rect.Max.X - Rect.Min.X;
+	uint32 DestHeight = Rect.Max.Y - Rect.Min.Y;
+
 	if (Texture2D->Surface.StorageFormat == VK_FORMAT_R16G16B16A16_SFLOAT)
 	{
 		uint32 PixelByteSize = 8u;
 		uint8* In = (uint8*)StagingBuffer->GetMappedPointer() + (Rect.Min.Y * TextureRHI2D->GetSizeX() + Rect.Min.X) * PixelByteSize;
 		uint32 SrcPitch = TextureRHI2D->GetSizeX() * PixelByteSize;
-		ConvertRawR16G16B16A16FDataToFColor(TextureRHI2D->GetSizeX(), TextureRHI2D->GetSizeY(), In, SrcPitch, Dest, false);
+		ConvertRawR16G16B16A16FDataToFColor(DestWidth, DestHeight, In, SrcPitch, Dest, false);
 	}
 	else if (Texture2D->Surface.StorageFormat == VK_FORMAT_A2B10G10R10_UNORM_PACK32)
 	{
 		uint32 PixelByteSize = 4u;
 		uint8* In = (uint8*)StagingBuffer->GetMappedPointer() + (Rect.Min.Y * TextureRHI2D->GetSizeX() + Rect.Min.X) * PixelByteSize;
 		uint32 SrcPitch = TextureRHI2D->GetSizeX() * PixelByteSize;
-		ConvertRawR10G10B10A2DataToFColor(TextureRHI2D->GetSizeX(), TextureRHI2D->GetSizeY(), In, SrcPitch, Dest);
+		ConvertRawR10G10B10A2DataToFColor(DestWidth, DestHeight, In, SrcPitch, Dest);
 	}
 	else if (Texture2D->Surface.StorageFormat == VK_FORMAT_R8G8B8A8_UNORM)
 	{
 		uint32 PixelByteSize = 4u;
 		uint8* In = (uint8*)StagingBuffer->GetMappedPointer() + (Rect.Min.Y * TextureRHI2D->GetSizeX() + Rect.Min.X) * PixelByteSize;
 		uint32 SrcPitch = TextureRHI2D->GetSizeX() * PixelByteSize;
-		ConvertRawR8G8B8A8DataToFColor(TextureRHI2D->GetSizeX(), TextureRHI2D->GetSizeY(), In, SrcPitch, Dest);
+		ConvertRawR8G8B8A8DataToFColor(DestWidth, DestHeight, In, SrcPitch, Dest);
 	}
 	else if (Texture2D->Surface.StorageFormat == VK_FORMAT_R16G16B16A16_UNORM)
 	{
 		uint32 PixelByteSize = 8u;
 		uint8* In = (uint8*)StagingBuffer->GetMappedPointer() + (Rect.Min.Y * TextureRHI2D->GetSizeX() + Rect.Min.X) * PixelByteSize;
 		uint32 SrcPitch = TextureRHI2D->GetSizeX() * PixelByteSize;
-		ConvertRawR16G16B16A16DataToFColor(Texture2D->GetSizeX(), Texture2D->GetSizeY(), In, SrcPitch, Dest);
+		ConvertRawR16G16B16A16DataToFColor(DestWidth, DestHeight, In, SrcPitch, Dest);
 	}
 	else if (Texture2D->Surface.StorageFormat == VK_FORMAT_B8G8R8A8_UNORM)
 	{
 		uint32 PixelByteSize = 4u;
 		uint8* In = (uint8*)StagingBuffer->GetMappedPointer() + (Rect.Min.Y * TextureRHI2D->GetSizeX() + Rect.Min.X) * PixelByteSize;
 		uint32 SrcPitch = TextureRHI2D->GetSizeX() * PixelByteSize;
-		ConvertRawB8G8R8A8DataToFColor(TextureRHI2D->GetSizeX(), TextureRHI2D->GetSizeY(), In, SrcPitch, Dest);
+		ConvertRawB8G8R8A8DataToFColor(DestWidth, DestHeight, In, SrcPitch, Dest);
 	}
 
 	Device->GetStagingManager().ReleaseBuffer(CmdBuffer, StagingBuffer);
