fix invalid check when running GWaitForIdleOnSubmit!=0 in Vulkan

diff --git a/Engine/Source/Runtime/VulkanRHI/Private/VulkanCommandBuffer.cpp b/Engine/Source/Runtime/VulkanRHI/Private/VulkanCommandBuffer.cpp
--- a/Engine/Source/Runtime/VulkanRHI/Private/VulkanCommandBuffer.cpp
+++ b/Engine/Source/Runtime/VulkanRHI/Private/VulkanCommandBuffer.cpp
@@ -428,7 +428,6 @@ void FVulkanCommandBufferManager::FlushResetQueryPools()
 		}
 		PoolResetCommandBuffer->End();
 		Queue->Submit(PoolResetCommandBuffer);
-		check(PoolResetCommandBuffer->State == FVulkanCmdBuffer::EState::Submitted);
 		PoolResets.Empty();
 	}
 }
