Integration from Oculus partner. This is for optimal MSAA support on Oculus Quest specifically.

diff --git a/Engine/Source/Runtime/RHI/Public/RHIDefinitions.h b/Engine/Source/Runtime/RHI/Public/RHIDefinitions.h
--- a/Engine/Source/Runtime/RHI/Public/RHIDefinitions.h
+++ b/Engine/Source/Runtime/RHI/Public/RHIDefinitions.h
@@ -1266,8 +1266,8 @@ inline bool RHINeedsToSwitchVerticalAxis(EShaderPlatform Platform)
 
 inline bool RHISupportsSeparateMSAAAndResolveTextures(const EShaderPlatform Platform)
 {
-	// Metal mobile devices, Vulkan and Android ES2/3.1 need to handle MSAA and resolve textures internally (unless RHICreateTexture2D was changed to take an optional resolve target)
-	return !IsMetalMobilePlatform(Platform) && !IsVulkanPlatform(Platform) && !IsAndroidOpenGLESPlatform(Platform);
+	// Metal mobile devices and Android ES2/3.1 need to handle MSAA and resolve textures internally (unless RHICreateTexture2D was changed to take an optional resolve target)
+	return !IsMetalMobilePlatform(Platform) && !IsAndroidOpenGLESPlatform(Platform);
 }
 
 inline bool RHISupportsComputeShaders(const EShaderPlatform Platform)
diff --git a/Engine/Source/Runtime/VulkanRHI/Private/VulkanRenderTarget.cpp b/Engine/Source/Runtime/VulkanRHI/Private/VulkanRenderTarget.cpp
--- a/Engine/Source/Runtime/VulkanRHI/Private/VulkanRenderTarget.cpp
+++ b/Engine/Source/Runtime/VulkanRHI/Private/VulkanRenderTarget.cpp
@@ -2124,7 +2124,7 @@ FVulkanRenderTargetLayout::FVulkanRenderTargetLayout(FVulkanDevice& InDevice, co
 			Desc[NumAttachmentDescriptions + 1].samples = VK_SAMPLE_COUNT_1_BIT;
 			Desc[NumAttachmentDescriptions + 1].storeOp = VK_ATTACHMENT_STORE_OP_STORE;
 			ResolveReferences[NumColorAttachments].attachment = NumAttachmentDescriptions + 1;
-			ResolveReferences[NumColorAttachments].layout = VK_IMAGE_LAYOUT_GENERAL;
+			ResolveReferences[NumColorAttachments].layout = VK_IMAGE_LAYOUT_COLOR_ATTACHMENT_OPTIMAL;
 			++NumAttachmentDescriptions;
 			bHasResolveAttachments = true;
 		}
@@ -2322,7 +2322,7 @@ FVulkanRenderTargetLayout::FVulkanRenderTargetLayout(const FGraphicsPipelineStat
 				Desc[NumAttachmentDescriptions + 1].samples = VK_SAMPLE_COUNT_1_BIT;
 				Desc[NumAttachmentDescriptions + 1].storeOp = VK_ATTACHMENT_STORE_OP_STORE;
 				ResolveReferences[NumColorAttachments].attachment = NumAttachmentDescriptions + 1;
-				ResolveReferences[NumColorAttachments].layout = VK_IMAGE_LAYOUT_GENERAL;
+				ResolveReferences[NumColorAttachments].layout = VK_IMAGE_LAYOUT_COLOR_ATTACHMENT_OPTIMAL;
 				++NumAttachmentDescriptions;
 				bHasResolveAttachments = true;
 			}
diff --git a/Engine/Source/Runtime/VulkanRHI/Private/VulkanTexture.cpp b/Engine/Source/Runtime/VulkanRHI/Private/VulkanTexture.cpp
--- a/Engine/Source/Runtime/VulkanRHI/Private/VulkanTexture.cpp
+++ b/Engine/Source/Runtime/VulkanRHI/Private/VulkanTexture.cpp
@@ -1869,7 +1869,7 @@ FVulkanTexture2D::FVulkanTexture2D(const FVulkanTexture2D* SrcTexture)
 
 FVulkanTexture2D::~FVulkanTexture2D()
 {
-	if ((Surface.UEFlags & (TexCreate_DepthStencilTargetable | TexCreate_RenderTargetable)) != 0)
+	if ((Surface.UEFlags & (TexCreate_DepthStencilTargetable | TexCreate_RenderTargetable | TexCreate_ResolveTargetable)) != 0)
 	{
 		Surface.Device->NotifyDeletedRenderTarget(Surface.Image);
 	}
@@ -1921,7 +1921,7 @@ FVulkanTextureCube::FVulkanTextureCube(const FVulkanTextureCube* SrcTexture)
 
 FVulkanTextureCube::~FVulkanTextureCube()
 {
-	if ((GetFlags() & (TexCreate_DepthStencilTargetable | TexCreate_RenderTargetable)) != 0)
+	if ((GetFlags() & (TexCreate_DepthStencilTargetable | TexCreate_RenderTargetable | TexCreate_ResolveTargetable)) != 0) 
 	{
 		Surface.Device->NotifyDeletedRenderTarget(Surface.Image);
 	}
@@ -1936,7 +1936,7 @@ FVulkanTexture3D::FVulkanTexture3D(FVulkanDevice& Device, EPixelFormat Format, u
 
 FVulkanTexture3D::~FVulkanTexture3D()
 {
-	if ((GetFlags() & (TexCreate_DepthStencilTargetable | TexCreate_RenderTargetable)) != 0)
+	if ((GetFlags() & (TexCreate_DepthStencilTargetable | TexCreate_RenderTargetable | TexCreate_ResolveTargetable)) != 0) 
 	{
 		Surface.Device->NotifyDeletedRenderTarget(Surface.Image);
 	}
