F - Fix crash when running Windows editor with -vulkan
#rnx

diff --git a/Engine/Source/Runtime/VulkanRHI/Private/VulkanPipeline.cpp b/Engine/Source/Runtime/VulkanRHI/Private/VulkanPipeline.cpp
--- a/Engine/Source/Runtime/VulkanRHI/Private/VulkanPipeline.cpp
+++ b/Engine/Source/Runtime/VulkanRHI/Private/VulkanPipeline.cpp
@@ -1871,7 +1871,7 @@ FVulkanComputePipeline* FVulkanPipelineStateCacheManager::CreateComputePipelineF
 	FUniformBufferGatherInfo UBGatherInfo;
 	DescriptorSetLayoutInfo.ProcessBindingsForStage(VK_SHADER_STAGE_COMPUTE_BIT, ShaderStage::Compute, CSHeader, UBGatherInfo);
 	DescriptorSetLayoutInfo.FinalizeBindings<true>(UBGatherInfo, TArrayView<FRHISamplerState*>());
-	FVulkanLayout* Layout = FindOrAddLayout(DescriptorSetLayoutInfo, true);
+	FVulkanLayout* Layout = FindOrAddLayout(DescriptorSetLayoutInfo, false);
 	FVulkanComputeLayout* ComputeLayout = (FVulkanComputeLayout*)Layout;
 	if (!ComputeLayout->ComputePipelineDescriptorInfo.IsInitialized())
 	{
