VK: Replace ALL_COMMANDS barrier with conditional debug barrier after each compute dispatch

diff --git a/Engine/Source/Runtime/VulkanRHI/Private/VulkanCommands.cpp b/Engine/Source/Runtime/VulkanRHI/Private/VulkanCommands.cpp
--- a/Engine/Source/Runtime/VulkanRHI/Private/VulkanCommands.cpp
+++ b/Engine/Source/Runtime/VulkanRHI/Private/VulkanCommands.cpp
@@ -116,8 +116,7 @@ void FVulkanCommandListContext::RHIDispatchComputeShader(uint32 ThreadGroupCount
 		GpuProfiler.RegisterGPUDispatch(FIntVector(ThreadGroupCountX, ThreadGroupCountY, ThreadGroupCountZ));	
 	}
 
-	//#todo-rco: Temp workaround
-	VulkanRHI::/*Debug*/HeavyWeightBarrier(CmdBuffer/*, 2*/);
+	VulkanRHI::DebugHeavyWeightBarrier(CmdBuffer, 2);
 }
 
 void FVulkanCommandListContext::RHIDispatchIndirectComputeShader(FRHIVertexBuffer* ArgumentBufferRHI, uint32 ArgumentOffset)
@@ -147,8 +146,7 @@ void FVulkanCommandListContext::RHIDispatchIndirectComputeShader(FRHIVertexBuffe
 		GpuProfiler.RegisterGPUDispatch(FIntVector(1, 1, 1));	
 	}
 
-	//#todo-rco: Temp workaround
-	VulkanRHI::/*Debug*/HeavyWeightBarrier(CmdBuffer/*, 2*/);
+	VulkanRHI::DebugHeavyWeightBarrier(CmdBuffer, 2);
 }
 
 void FVulkanCommandListContext::RHISetUAVParameter(FRHIComputeShader* ComputeShaderRHI, uint32 UAVIndex, FRHIUnorderedAccessView* UAVRHI)
