VK: Set default value to 1 for a missing alpha component in a texture. This makes vulkan behavior conisitent with other RHI

diff --git a/Engine/Source/Runtime/VulkanRHI/Private/VulkanDevice.cpp b/Engine/Source/Runtime/VulkanRHI/Private/VulkanDevice.cpp
--- a/Engine/Source/Runtime/VulkanRHI/Private/VulkanDevice.cpp
+++ b/Engine/Source/Runtime/VulkanRHI/Private/VulkanDevice.cpp
@@ -466,13 +466,13 @@ void FVulkanDevice::SetupFormats()
 	SetComponentMapping(PF_B8G8R8A8, VK_COMPONENT_SWIZZLE_R, VK_COMPONENT_SWIZZLE_G, VK_COMPONENT_SWIZZLE_B, VK_COMPONENT_SWIZZLE_A);
 
 	MapFormatSupport(PF_G8, VK_FORMAT_R8_UNORM);
-	SetComponentMapping(PF_G8, VK_COMPONENT_SWIZZLE_R, VK_COMPONENT_SWIZZLE_ZERO, VK_COMPONENT_SWIZZLE_ZERO, VK_COMPONENT_SWIZZLE_ZERO);
+	SetComponentMapping(PF_G8, VK_COMPONENT_SWIZZLE_R, VK_COMPONENT_SWIZZLE_ZERO, VK_COMPONENT_SWIZZLE_ZERO, VK_COMPONENT_SWIZZLE_ONE);
 
 	MapFormatSupportWithFallback(PF_G16, VK_FORMAT_R16_UNORM, {VK_FORMAT_R16_SFLOAT});
-	SetComponentMapping(PF_G16, VK_COMPONENT_SWIZZLE_R, VK_COMPONENT_SWIZZLE_ZERO, VK_COMPONENT_SWIZZLE_ZERO, VK_COMPONENT_SWIZZLE_ZERO);
+	SetComponentMapping(PF_G16, VK_COMPONENT_SWIZZLE_R, VK_COMPONENT_SWIZZLE_ZERO, VK_COMPONENT_SWIZZLE_ZERO, VK_COMPONENT_SWIZZLE_ONE);
 
 	MapFormatSupport(PF_FloatRGB, VK_FORMAT_B10G11R11_UFLOAT_PACK32);
-	SetComponentMapping(PF_FloatRGB, VK_COMPONENT_SWIZZLE_R, VK_COMPONENT_SWIZZLE_G, VK_COMPONENT_SWIZZLE_B, VK_COMPONENT_SWIZZLE_ZERO);
+	SetComponentMapping(PF_FloatRGB, VK_COMPONENT_SWIZZLE_R, VK_COMPONENT_SWIZZLE_G, VK_COMPONENT_SWIZZLE_B, VK_COMPONENT_SWIZZLE_ONE);
 
 	MapFormatSupport(PF_FloatRGBA, VK_FORMAT_R16G16B16A16_SFLOAT, 8);
 	SetComponentMapping(PF_FloatRGBA, VK_COMPONENT_SWIZZLE_R, VK_COMPONENT_SWIZZLE_G, VK_COMPONENT_SWIZZLE_B, VK_COMPONENT_SWIZZLE_A);
@@ -494,46 +494,46 @@ void FVulkanDevice::SetupFormats()
 
 	// Requirement for GPU particles
 	MapFormatSupport(PF_G32R32F, VK_FORMAT_R32G32_SFLOAT, 8);
-	SetComponentMapping(PF_G32R32F, VK_COMPONENT_SWIZZLE_R, VK_COMPONENT_SWIZZLE_G, VK_COMPONENT_SWIZZLE_ZERO, VK_COMPONENT_SWIZZLE_ZERO);
+	SetComponentMapping(PF_G32R32F, VK_COMPONENT_SWIZZLE_R, VK_COMPONENT_SWIZZLE_G, VK_COMPONENT_SWIZZLE_ZERO, VK_COMPONENT_SWIZZLE_ONE);
 
 	MapFormatSupport(PF_A32B32G32R32F, VK_FORMAT_R32G32B32A32_SFLOAT, 16);
 	SetComponentMapping(PF_A32B32G32R32F, VK_COMPONENT_SWIZZLE_R, VK_COMPONENT_SWIZZLE_G, VK_COMPONENT_SWIZZLE_B, VK_COMPONENT_SWIZZLE_A);
 
 	MapFormatSupportWithFallback(PF_G16R16, VK_FORMAT_R16G16_UNORM, {VK_FORMAT_R16G16_SFLOAT});
-	SetComponentMapping(PF_G16R16, VK_COMPONENT_SWIZZLE_R, VK_COMPONENT_SWIZZLE_G, VK_COMPONENT_SWIZZLE_ZERO, VK_COMPONENT_SWIZZLE_ZERO);
+	SetComponentMapping(PF_G16R16, VK_COMPONENT_SWIZZLE_R, VK_COMPONENT_SWIZZLE_G, VK_COMPONENT_SWIZZLE_ZERO, VK_COMPONENT_SWIZZLE_ONE);
 
 	MapFormatSupport(PF_G16R16F, VK_FORMAT_R16G16_SFLOAT);
-	SetComponentMapping(PF_G16R16F, VK_COMPONENT_SWIZZLE_R, VK_COMPONENT_SWIZZLE_G, VK_COMPONENT_SWIZZLE_ZERO, VK_COMPONENT_SWIZZLE_ZERO);
+	SetComponentMapping(PF_G16R16F, VK_COMPONENT_SWIZZLE_R, VK_COMPONENT_SWIZZLE_G, VK_COMPONENT_SWIZZLE_ZERO, VK_COMPONENT_SWIZZLE_ONE);
 
 	MapFormatSupport(PF_G16R16F_FILTER, VK_FORMAT_R16G16_SFLOAT);
-	SetComponentMapping(PF_G16R16F_FILTER, VK_COMPONENT_SWIZZLE_R, VK_COMPONENT_SWIZZLE_G, VK_COMPONENT_SWIZZLE_ZERO, VK_COMPONENT_SWIZZLE_ZERO);
+	SetComponentMapping(PF_G16R16F_FILTER, VK_COMPONENT_SWIZZLE_R, VK_COMPONENT_SWIZZLE_G, VK_COMPONENT_SWIZZLE_ZERO, VK_COMPONENT_SWIZZLE_ONE);
 
 	MapFormatSupport(PF_R16_UINT, VK_FORMAT_R16_UINT);
-	SetComponentMapping(PF_R16_UINT, VK_COMPONENT_SWIZZLE_R, VK_COMPONENT_SWIZZLE_ZERO, VK_COMPONENT_SWIZZLE_ZERO, VK_COMPONENT_SWIZZLE_ZERO);
+	SetComponentMapping(PF_R16_UINT, VK_COMPONENT_SWIZZLE_R, VK_COMPONENT_SWIZZLE_ZERO, VK_COMPONENT_SWIZZLE_ZERO, VK_COMPONENT_SWIZZLE_ONE);
 
 	MapFormatSupport(PF_R16_SINT, VK_FORMAT_R16_SINT);
-	SetComponentMapping(PF_R16_SINT, VK_COMPONENT_SWIZZLE_R, VK_COMPONENT_SWIZZLE_ZERO, VK_COMPONENT_SWIZZLE_ZERO, VK_COMPONENT_SWIZZLE_ZERO);
+	SetComponentMapping(PF_R16_SINT, VK_COMPONENT_SWIZZLE_R, VK_COMPONENT_SWIZZLE_ZERO, VK_COMPONENT_SWIZZLE_ZERO, VK_COMPONENT_SWIZZLE_ONE);
 
 	MapFormatSupport(PF_R32_UINT, VK_FORMAT_R32_UINT);
-	SetComponentMapping(PF_R32_UINT, VK_COMPONENT_SWIZZLE_R, VK_COMPONENT_SWIZZLE_ZERO, VK_COMPONENT_SWIZZLE_ZERO, VK_COMPONENT_SWIZZLE_ZERO);
+	SetComponentMapping(PF_R32_UINT, VK_COMPONENT_SWIZZLE_R, VK_COMPONENT_SWIZZLE_ZERO, VK_COMPONENT_SWIZZLE_ZERO, VK_COMPONENT_SWIZZLE_ONE);
 
 	MapFormatSupport(PF_R32_SINT, VK_FORMAT_R32_SINT);
-	SetComponentMapping(PF_R32_SINT, VK_COMPONENT_SWIZZLE_R, VK_COMPONENT_SWIZZLE_ZERO, VK_COMPONENT_SWIZZLE_ZERO, VK_COMPONENT_SWIZZLE_ZERO);
+	SetComponentMapping(PF_R32_SINT, VK_COMPONENT_SWIZZLE_R, VK_COMPONENT_SWIZZLE_ZERO, VK_COMPONENT_SWIZZLE_ZERO, VK_COMPONENT_SWIZZLE_ONE);
 
 	MapFormatSupport(PF_R8_UINT, VK_FORMAT_R8_UINT);
-	SetComponentMapping(PF_R8_UINT, VK_COMPONENT_SWIZZLE_R, VK_COMPONENT_SWIZZLE_ZERO, VK_COMPONENT_SWIZZLE_ZERO, VK_COMPONENT_SWIZZLE_ZERO);
+	SetComponentMapping(PF_R8_UINT, VK_COMPONENT_SWIZZLE_R, VK_COMPONENT_SWIZZLE_ZERO, VK_COMPONENT_SWIZZLE_ZERO, VK_COMPONENT_SWIZZLE_ONE);
 
 	MapFormatSupportWithFallback(PF_D24, VK_FORMAT_X8_D24_UNORM_PACK32, {VK_FORMAT_D24_UNORM_S8_UINT, VK_FORMAT_D16_UNORM_S8_UINT, VK_FORMAT_D32_SFLOAT, VK_FORMAT_D32_SFLOAT_S8_UINT});
 	SetComponentMapping(PF_D24, VK_COMPONENT_SWIZZLE_R, VK_COMPONENT_SWIZZLE_ZERO, VK_COMPONENT_SWIZZLE_ZERO, VK_COMPONENT_SWIZZLE_ZERO);
 
 	MapFormatSupport(PF_R16F, VK_FORMAT_R16_SFLOAT);
-	SetComponentMapping(PF_R16F, VK_COMPONENT_SWIZZLE_R, VK_COMPONENT_SWIZZLE_ZERO, VK_COMPONENT_SWIZZLE_ZERO, VK_COMPONENT_SWIZZLE_ZERO);
+	SetComponentMapping(PF_R16F, VK_COMPONENT_SWIZZLE_R, VK_COMPONENT_SWIZZLE_ZERO, VK_COMPONENT_SWIZZLE_ZERO, VK_COMPONENT_SWIZZLE_ONE);
 
 	MapFormatSupport(PF_R16F_FILTER, VK_FORMAT_R16_SFLOAT);
-	SetComponentMapping(PF_R16F_FILTER, VK_COMPONENT_SWIZZLE_R, VK_COMPONENT_SWIZZLE_ZERO, VK_COMPONENT_SWIZZLE_ZERO, VK_COMPONENT_SWIZZLE_ZERO);
+	SetComponentMapping(PF_R16F_FILTER, VK_COMPONENT_SWIZZLE_R, VK_COMPONENT_SWIZZLE_ZERO, VK_COMPONENT_SWIZZLE_ZERO, VK_COMPONENT_SWIZZLE_ONE);
 
 	MapFormatSupport(PF_FloatR11G11B10, VK_FORMAT_B10G11R11_UFLOAT_PACK32, 4);
-	SetComponentMapping(PF_FloatR11G11B10, VK_COMPONENT_SWIZZLE_R, VK_COMPONENT_SWIZZLE_G, VK_COMPONENT_SWIZZLE_B, VK_COMPONENT_SWIZZLE_ZERO);
+	SetComponentMapping(PF_FloatR11G11B10, VK_COMPONENT_SWIZZLE_R, VK_COMPONENT_SWIZZLE_G, VK_COMPONENT_SWIZZLE_B, VK_COMPONENT_SWIZZLE_ONE);
 
 	MapFormatSupport(PF_A2B10G10R10, VK_FORMAT_A2B10G10R10_UNORM_PACK32, 4);
 	SetComponentMapping(PF_A2B10G10R10, VK_COMPONENT_SWIZZLE_R, VK_COMPONENT_SWIZZLE_G, VK_COMPONENT_SWIZZLE_B, VK_COMPONENT_SWIZZLE_A);
@@ -558,7 +558,7 @@ void FVulkanDevice::SetupFormats()
 	SetComponentMapping(PF_R8G8B8A8_SNORM, VK_COMPONENT_SWIZZLE_R, VK_COMPONENT_SWIZZLE_G, VK_COMPONENT_SWIZZLE_B, VK_COMPONENT_SWIZZLE_A);
 
 	MapFormatSupport(PF_R16G16_UINT, VK_FORMAT_R16G16_UINT);
-	SetComponentMapping(PF_R16G16_UINT, VK_COMPONENT_SWIZZLE_R, VK_COMPONENT_SWIZZLE_G, VK_COMPONENT_SWIZZLE_ZERO, VK_COMPONENT_SWIZZLE_ZERO);
+	SetComponentMapping(PF_R16G16_UINT, VK_COMPONENT_SWIZZLE_R, VK_COMPONENT_SWIZZLE_G, VK_COMPONENT_SWIZZLE_ZERO, VK_COMPONENT_SWIZZLE_ONE);
 
 	MapFormatSupport(PF_R16G16B16A16_UINT, VK_FORMAT_R16G16B16A16_UINT);
 	SetComponentMapping(PF_R16G16B16A16_UINT, VK_COMPONENT_SWIZZLE_R, VK_COMPONENT_SWIZZLE_G, VK_COMPONENT_SWIZZLE_B, VK_COMPONENT_SWIZZLE_A);
@@ -567,7 +567,7 @@ void FVulkanDevice::SetupFormats()
 	SetComponentMapping(PF_R16G16B16A16_SINT, VK_COMPONENT_SWIZZLE_R, VK_COMPONENT_SWIZZLE_G, VK_COMPONENT_SWIZZLE_B, VK_COMPONENT_SWIZZLE_A);
 
 	MapFormatSupport(PF_R32G32_UINT, VK_FORMAT_R32G32_UINT);
-	SetComponentMapping(PF_R32G32_UINT, VK_COMPONENT_SWIZZLE_R, VK_COMPONENT_SWIZZLE_G, VK_COMPONENT_SWIZZLE_ZERO, VK_COMPONENT_SWIZZLE_ZERO);
+	SetComponentMapping(PF_R32G32_UINT, VK_COMPONENT_SWIZZLE_R, VK_COMPONENT_SWIZZLE_G, VK_COMPONENT_SWIZZLE_ZERO, VK_COMPONENT_SWIZZLE_ONE);
 
 	MapFormatSupport(PF_R32G32B32A32_UINT, VK_FORMAT_R32G32B32A32_UINT);
 	SetComponentMapping(PF_R32G32B32A32_UINT, VK_COMPONENT_SWIZZLE_R, VK_COMPONENT_SWIZZLE_G, VK_COMPONENT_SWIZZLE_B, VK_COMPONENT_SWIZZLE_A);
@@ -579,13 +579,13 @@ void FVulkanDevice::SetupFormats()
 	SetComponentMapping(PF_R16G16B16A16_UNORM, VK_COMPONENT_SWIZZLE_R, VK_COMPONENT_SWIZZLE_G, VK_COMPONENT_SWIZZLE_B, VK_COMPONENT_SWIZZLE_A);
 
 	MapFormatSupport(PF_R8G8, VK_FORMAT_R8G8_UNORM);
-	SetComponentMapping(PF_R8G8, VK_COMPONENT_SWIZZLE_R, VK_COMPONENT_SWIZZLE_G, VK_COMPONENT_SWIZZLE_ZERO, VK_COMPONENT_SWIZZLE_ZERO);
+	SetComponentMapping(PF_R8G8, VK_COMPONENT_SWIZZLE_R, VK_COMPONENT_SWIZZLE_G, VK_COMPONENT_SWIZZLE_ZERO, VK_COMPONENT_SWIZZLE_ONE);
 
 	MapFormatSupport(PF_V8U8, VK_FORMAT_R8G8_UNORM);
-	SetComponentMapping(PF_V8U8, VK_COMPONENT_SWIZZLE_R, VK_COMPONENT_SWIZZLE_G, VK_COMPONENT_SWIZZLE_ZERO, VK_COMPONENT_SWIZZLE_ZERO);
+	SetComponentMapping(PF_V8U8, VK_COMPONENT_SWIZZLE_R, VK_COMPONENT_SWIZZLE_G, VK_COMPONENT_SWIZZLE_ZERO, VK_COMPONENT_SWIZZLE_ONE);
 
 	MapFormatSupport(PF_R32_FLOAT, VK_FORMAT_R32_SFLOAT);
-	SetComponentMapping(PF_R32_FLOAT, VK_COMPONENT_SWIZZLE_R, VK_COMPONENT_SWIZZLE_ZERO, VK_COMPONENT_SWIZZLE_ZERO, VK_COMPONENT_SWIZZLE_ZERO);
+	SetComponentMapping(PF_R32_FLOAT, VK_COMPONENT_SWIZZLE_R, VK_COMPONENT_SWIZZLE_ZERO, VK_COMPONENT_SWIZZLE_ZERO, VK_COMPONENT_SWIZZLE_ONE);
 
 	if (FVulkanPlatform::SupportsBCTextureFormats())
 	{
