Merging CL#10666227

Fix vulkan shutdown issue: when refcount from pending state would hit 0, it'd end up modifying the array being iterated on.

diff --git a/Engine/Source/Runtime/VulkanRHI/Private/VulkanPendingState.cpp b/Engine/Source/Runtime/VulkanRHI/Private/VulkanPendingState.cpp
--- a/Engine/Source/Runtime/VulkanRHI/Private/VulkanPendingState.cpp
+++ b/Engine/Source/Runtime/VulkanRHI/Private/VulkanPendingState.cpp
@@ -402,15 +402,11 @@ void FVulkanPendingComputeState::PrepareForDispatch(FVulkanCmdBuffer* InCmdBuffe
 
 FVulkanPendingGfxState::~FVulkanPendingGfxState()
 {
-	TArray<FVulkanGraphicsPipelineDescriptorState*> DescriptorStates;
-
-	for (auto& Pair : PipelineStates)
+	TMap<FVulkanRHIGraphicsPipelineState*, FVulkanGraphicsPipelineDescriptorState*> Temp;
+	Swap(Temp, PipelineStates);
+	for (auto& Pair : Temp)
 	{
 		FVulkanGraphicsPipelineDescriptorState* State = Pair.Value;
-		DescriptorStates.Push(State);
-	}
-	for(FVulkanGraphicsPipelineDescriptorState* State : DescriptorStates)
-	{
 		delete State;
 	}
 }
