VulkanRHI is using ApplicationCore to InitSDL. Need to depend on it in case we get here before someone else depends on ApplicationCore

Based on this CL 6060555

diff --git a/Engine/Source/Runtime/VulkanRHI/VulkanRHI.Build.cs b/Engine/Source/Runtime/VulkanRHI/VulkanRHI.Build.cs
--- a/Engine/Source/Runtime/VulkanRHI/VulkanRHI.Build.cs
+++ b/Engine/Source/Runtime/VulkanRHI/VulkanRHI.Build.cs
@@ -27,6 +27,7 @@ public class VulkanRHI : ModuleRules
 				"Core", 
 				"CoreUObject", 
 				"Engine", 
+				"ApplicationCore",
 				"RHI", 
 				"RenderCore", 
 				"UtilityShaders",
