Only add ApplicationCore for Linux

diff --git a/Engine/Source/Runtime/VulkanRHI/VulkanRHI.Build.cs b/Engine/Source/Runtime/VulkanRHI/VulkanRHI.Build.cs
--- a/Engine/Source/Runtime/VulkanRHI/VulkanRHI.Build.cs
+++ b/Engine/Source/Runtime/VulkanRHI/VulkanRHI.Build.cs
@@ -27,7 +27,6 @@ public class VulkanRHI : ModuleRules
 				"Core", 
 				"CoreUObject", 
 				"Engine", 
-				"ApplicationCore",
 				"RHI", 
 				"RenderCore", 
 				"UtilityShaders",
@@ -45,6 +44,7 @@ public class VulkanRHI : ModuleRules
 		{
 			if (Target.Platform == UnrealTargetPlatform.Linux)
 			{
+				PrivateDependencyModuleNames.Add("ApplicationCore");
 				AddEngineThirdPartyPrivateStaticDependencies(Target, "SDL2");
 
 				string VulkanSDKPath = Environment.GetEnvironmentVariable("VULKAN_SDK");
