UE4R - More info for Vulkan failing to query extensions & layers

diff --git a/Engine/Source/Runtime/VulkanRHI/Private/VulkanLayers.cpp b/Engine/Source/Runtime/VulkanRHI/Private/VulkanLayers.cpp
--- a/Engine/Source/Runtime/VulkanRHI/Private/VulkanLayers.cpp
+++ b/Engine/Source/Runtime/VulkanRHI/Private/VulkanLayers.cpp
@@ -129,42 +129,26 @@ struct FLayerExtension
 
 static inline void EnumerateInstanceExtensionProperties(const ANSICHAR* LayerName, FLayerExtension& OutLayer)
 {
-	VkResult Result;
-	do
+	uint32 Count = 0;
+	VERIFYVULKANRESULT(VulkanRHI::vkEnumerateInstanceExtensionProperties(LayerName, &Count, nullptr));
+	if (Count > 0)
 	{
-		uint32 Count = 0;
-		Result = VulkanRHI::vkEnumerateInstanceExtensionProperties(LayerName, &Count, nullptr);
-		check(Result >= VK_SUCCESS);
-
-		if (Count > 0)
-		{
-			OutLayer.ExtensionProps.Empty(Count);
-			OutLayer.ExtensionProps.AddUninitialized(Count);
-			Result = VulkanRHI::vkEnumerateInstanceExtensionProperties(LayerName, &Count, OutLayer.ExtensionProps.GetData());
-			check(Result >= VK_SUCCESS);
-		}
+		OutLayer.ExtensionProps.Empty(Count);
+		OutLayer.ExtensionProps.AddUninitialized(Count);
+		VERIFYVULKANRESULT(VulkanRHI::vkEnumerateInstanceExtensionProperties(LayerName, &Count, OutLayer.ExtensionProps.GetData()));
 	}
-	while (Result == VK_INCOMPLETE);
 }
 
 static inline void EnumerateDeviceExtensionProperties(VkPhysicalDevice Device, const ANSICHAR* LayerName, FLayerExtension& OutLayer)
 {
-	VkResult Result;
-	do
+	uint32 Count = 0;
+	VERIFYVULKANRESULT(VulkanRHI::vkEnumerateDeviceExtensionProperties(Device, LayerName, &Count, nullptr));
+	if (Count > 0)
 	{
-		uint32 Count = 0;
-		Result = VulkanRHI::vkEnumerateDeviceExtensionProperties(Device, LayerName, &Count, nullptr);
-		check(Result >= VK_SUCCESS);
-
-		if (Count > 0)
-		{
-			OutLayer.ExtensionProps.Empty(Count);
-			OutLayer.ExtensionProps.AddUninitialized(Count);
-			Result = VulkanRHI::vkEnumerateDeviceExtensionProperties(Device, LayerName, &Count, OutLayer.ExtensionProps.GetData());
-			check(Result >= VK_SUCCESS);
-		}
+		OutLayer.ExtensionProps.Empty(Count);
+		OutLayer.ExtensionProps.AddUninitialized(Count);
+		VERIFYVULKANRESULT(VulkanRHI::vkEnumerateDeviceExtensionProperties(Device, LayerName, &Count, OutLayer.ExtensionProps.GetData()));
 	}
-	while (Result == VK_INCOMPLETE);
 }
 
 
@@ -240,8 +224,6 @@ void FVulkanDynamicRHI::GetInstanceLayersAndExtensions(TArray<const ANSICHAR*>&
 	// 0 is reserved for NULL/instance
 	GlobalLayerExtensions.AddDefaulted();
 
-	VkResult Result;
-
 	// Global extensions
 	EnumerateInstanceExtensionProperties(nullptr, GlobalLayerExtensions[0]);
 
@@ -254,20 +236,13 @@ void FVulkanDynamicRHI::GetInstanceLayersAndExtensions(TArray<const ANSICHAR*>&
 
 	{
 		TArray<VkLayerProperties> GlobalLayerProperties;
-		do
+		uint32 InstanceLayerCount = 0;
+		VERIFYVULKANRESULT(VulkanRHI::vkEnumerateInstanceLayerProperties(&InstanceLayerCount, nullptr));
+		if (InstanceLayerCount > 0)
 		{
-			uint32 InstanceLayerCount = 0;
-			Result = VulkanRHI::vkEnumerateInstanceLayerProperties(&InstanceLayerCount, nullptr);
-			check(Result >= VK_SUCCESS);
-
-			if (InstanceLayerCount > 0)
-			{
-				GlobalLayerProperties.AddZeroed(InstanceLayerCount);
-				Result = VulkanRHI::vkEnumerateInstanceLayerProperties(&InstanceLayerCount, &GlobalLayerProperties[GlobalLayerProperties.Num() - InstanceLayerCount]);
-				check(Result >= VK_SUCCESS);
-			}
+			GlobalLayerProperties.AddZeroed(InstanceLayerCount);
+			VERIFYVULKANRESULT(VulkanRHI::vkEnumerateInstanceLayerProperties(&InstanceLayerCount, &GlobalLayerProperties[GlobalLayerProperties.Num() - InstanceLayerCount]));
 		}
-		while (Result == VK_INCOMPLETE);
 
 		for (int32 Index = 0; Index < GlobalLayerProperties.Num(); ++Index)
 		{
