Change for calling vkCmdClearColorImage without VK_IMAGE_USAGE_TRANSFER_DST_BIT which is against the Vulkan spec and can cause issues on certain drivers.

diff --git a/Engine/Source/Runtime/VulkanRHI/Private/VulkanSwapChain.cpp b/Engine/Source/Runtime/VulkanRHI/Private/VulkanSwapChain.cpp
--- a/Engine/Source/Runtime/VulkanRHI/Private/VulkanSwapChain.cpp
+++ b/Engine/Source/Runtime/VulkanRHI/Private/VulkanSwapChain.cpp
@@ -412,11 +412,7 @@ FVulkanSwapChain::FVulkanSwapChain(VkInstance InInstance, FVulkanDevice& InDevic
 	SwapChainInfo.imageColorSpace = CurrFormat.colorSpace;
 	SwapChainInfo.imageExtent.width = SizeX;
 	SwapChainInfo.imageExtent.height = SizeY;
-	SwapChainInfo.imageUsage = VK_IMAGE_USAGE_COLOR_ATTACHMENT_BIT | VK_IMAGE_USAGE_SAMPLED_BIT;
-	if (GVulkanDelayAcquireImage == EDelayAcquireImageType::DelayAcquire)
-	{
-		SwapChainInfo.imageUsage |= VK_IMAGE_USAGE_TRANSFER_DST_BIT;
-	}
+	SwapChainInfo.imageUsage = VK_IMAGE_USAGE_COLOR_ATTACHMENT_BIT | VK_IMAGE_USAGE_SAMPLED_BIT | VK_IMAGE_USAGE_TRANSFER_DST_BIT;
 	SwapChainInfo.preTransform = PreTransform;
 	SwapChainInfo.imageArrayLayers = 1;
 	SwapChainInfo.imageSharingMode = VK_SHARING_MODE_EXCLUSIVE;
