Add more aggressive barriers when updating Uniform buffers. Fixes an issue with CSM shadows flickering in certain configurations.

diff --git a/Engine/Source/Runtime/VulkanRHI/Private/VulkanCommandBuffer.cpp b/Engine/Source/Runtime/VulkanRHI/Private/VulkanCommandBuffer.cpp
--- a/Engine/Source/Runtime/VulkanRHI/Private/VulkanCommandBuffer.cpp
+++ b/Engine/Source/Runtime/VulkanRHI/Private/VulkanCommandBuffer.cpp
@@ -39,6 +39,7 @@ FVulkanCmdBuffer::FVulkanCmdBuffer(FVulkanDevice* InDevice, FVulkanCommandBuffer
 	, bHasScissor(0)
 	, bHasStencilRef(0)
 	, bIsUploadOnly(bInIsUploadOnly ? 1 : 0)
+	, bIsUniformBufferBarrierAdded(0)
 	, Device(InDevice)
 	, CommandBufferHandle(VK_NULL_HANDLE)
 	, Fence(nullptr)
@@ -122,8 +123,44 @@ void FVulkanCmdBuffer::FreeMemory()
 	State = EState::NotAllocated;
 }
 
+void FVulkanCmdBuffer::BeginUniformUpdateBarrier()
+{
+	if(!bIsUniformBufferBarrierAdded)
+	{
+		VkMemoryBarrier Barrier;
+		ZeroVulkanStruct(Barrier, VK_STRUCTURE_TYPE_MEMORY_BARRIER);
+		Barrier.srcAccessMask = VK_ACCESS_UNIFORM_READ_BIT;
+		Barrier.dstAccessMask = VK_ACCESS_TRANSFER_WRITE_BIT;
+		VulkanRHI::vkCmdPipelineBarrier(GetHandle(), VK_PIPELINE_STAGE_ALL_COMMANDS_BIT, VK_PIPELINE_STAGE_ALL_COMMANDS_BIT, 0, 1, &Barrier, 0, nullptr, 0, nullptr);
+		bIsUniformBufferBarrierAdded = true;
+	}
+}
+
+void FVulkanCmdBuffer::EndUniformUpdateBarrier()
+{
+	if(bIsUniformBufferBarrierAdded)
+	{
+		VkMemoryBarrier Barrier;
+		ZeroVulkanStruct(Barrier, VK_STRUCTURE_TYPE_MEMORY_BARRIER);
+		Barrier.srcAccessMask = VK_ACCESS_TRANSFER_WRITE_BIT;
+		Barrier.dstAccessMask = VK_ACCESS_UNIFORM_READ_BIT;
+		VulkanRHI::vkCmdPipelineBarrier(GetHandle(), VK_PIPELINE_STAGE_ALL_COMMANDS_BIT, VK_PIPELINE_STAGE_ALL_COMMANDS_BIT, 0, 1, &Barrier, 0, nullptr, 0, nullptr);
+		bIsUniformBufferBarrierAdded = false;
+	}
+}
+void FVulkanCmdBuffer::EndRenderPass()
+{
+	checkf(IsInsideRenderPass(), TEXT("Can't EndRP as we're NOT inside one! CmdBuffer 0x%p State=%d"), CommandBufferHandle, (int32)State);
+	VulkanRHI::vkCmdEndRenderPass(CommandBufferHandle);
+	State = EState::IsInsideBegin;
+}
+
 void FVulkanCmdBuffer::BeginRenderPass(const FVulkanRenderTargetLayout& Layout, FVulkanRenderPass* RenderPass, FVulkanFramebuffer* Framebuffer, const VkClearValue* AttachmentClearValues)
 {
+	if(bIsUniformBufferBarrierAdded)
+	{
+		EndUniformUpdateBarrier();
+	}
 	checkf(IsOutsideRenderPass(), TEXT("Can't BeginRP as already inside one! CmdBuffer 0x%p State=%d"), CommandBufferHandle, (int32)State);
 
 	VkRenderPassBeginInfo Info;
diff --git a/Engine/Source/Runtime/VulkanRHI/Private/VulkanCommandBuffer.h b/Engine/Source/Runtime/VulkanRHI/Private/VulkanCommandBuffer.h
--- a/Engine/Source/Runtime/VulkanRHI/Private/VulkanCommandBuffer.h
+++ b/Engine/Source/Runtime/VulkanRHI/Private/VulkanCommandBuffer.h
@@ -39,6 +39,11 @@ public:
 		return CommandBufferPool;
 	}
 
+	bool IsUniformBufferBarrierAdded() const
+	{
+		return bIsUniformBufferBarrierAdded;
+	}
+
 	inline bool IsInsideRenderPass() const
 	{
 		return State == EState::IsInsideRenderPass;
@@ -154,22 +159,21 @@ public:
 	VkRect2D CurrentScissor;
 	uint32 CurrentStencilRef;
 	EState State;
-	uint8 bNeedsDynamicStateSet	: 1;
-	uint8 bHasPipeline			: 1;
-	uint8 bHasViewport			: 1;
-	uint8 bHasScissor			: 1;
-	uint8 bHasStencilRef		: 1;
-	uint8 bIsUploadOnly			: 1;
+	uint8 bNeedsDynamicStateSet			: 1;
+	uint8 bHasPipeline					: 1;
+	uint8 bHasViewport					: 1;
+	uint8 bHasScissor					: 1;
+	uint8 bHasStencilRef				: 1;
+	uint8 bIsUploadOnly					: 1;
+	uint8 bIsUniformBufferBarrierAdded	: 1;
 
 	// You never want to call Begin/EndRenderPass directly as it will mess up with the FTransitionAndLayoutManager
 	void BeginRenderPass(const FVulkanRenderTargetLayout& Layout, class FVulkanRenderPass* RenderPass, class FVulkanFramebuffer* Framebuffer, const VkClearValue* AttachmentClearValues);
-	void EndRenderPass()
-	{
-		checkf(IsInsideRenderPass(), TEXT("Can't EndRP as we're NOT inside one! CmdBuffer 0x%p State=%d"), CommandBufferHandle, (int32)State);
-		VulkanRHI::vkCmdEndRenderPass(CommandBufferHandle);
-		State = EState::IsInsideBegin;
-	}
+	void EndRenderPass();
+
 
+	void BeginUniformUpdateBarrier();
+	void EndUniformUpdateBarrier();
 	//#todo-rco: Hide this
 	FVulkanDescriptorPoolSetContainer* CurrentDescriptorPoolSetContainer = nullptr;
 
diff --git a/Engine/Source/Runtime/VulkanRHI/Private/VulkanRenderTarget.cpp b/Engine/Source/Runtime/VulkanRHI/Private/VulkanRenderTarget.cpp
--- a/Engine/Source/Runtime/VulkanRHI/Private/VulkanRenderTarget.cpp
+++ b/Engine/Source/Runtime/VulkanRHI/Private/VulkanRenderTarget.cpp
@@ -1771,7 +1771,6 @@ void FVulkanCommandListContext::RHIEndRenderPass()
 		}
 	}
 	RHIPopEvent();
-	bUniformBufferUploadRenderPassDirty = true;
 }
 
 void FVulkanCommandListContext::RHINextSubpass()
diff --git a/Engine/Source/Runtime/VulkanRHI/Private/VulkanUniformBuffer.cpp b/Engine/Source/Runtime/VulkanRHI/Private/VulkanUniformBuffer.cpp
--- a/Engine/Source/Runtime/VulkanRHI/Private/VulkanUniformBuffer.cpp
+++ b/Engine/Source/Runtime/VulkanRHI/Private/VulkanUniformBuffer.cpp
@@ -231,16 +231,20 @@ inline void FVulkanDynamicRHI::UpdateUniformBuffer(FVulkanUniformBuffer* Uniform
 		Region.srcOffset = LockInfo.GetBindOffset();
 		Region.dstOffset = VulkanUniformBuffer->GetOffset();
 		VkBuffer UBBuffer = VulkanUniformBuffer->GetBufferAllocation()->GetHandle();
-		if(Context.bUniformBufferUploadRenderPassDirty)
+
+		bool bIsInRenderPass = CmdBuffer->IsInsideRenderPass();
+		bool bIsUniformBarrierAdded = CmdBuffer->IsUniformBufferBarrierAdded();
+		if(bIsInRenderPass || !bIsUniformBarrierAdded)
 		{
-			Context.bUniformBufferUploadRenderPassDirty = false;
-			VkMemoryBarrier Barrier;
-			ZeroVulkanStruct(Barrier, VK_STRUCTURE_TYPE_MEMORY_BARRIER);
-			Barrier.srcAccessMask = VK_ACCESS_UNIFORM_READ_BIT; 
-			Barrier.dstAccessMask = VK_ACCESS_TRANSFER_WRITE_BIT;
-			VulkanRHI::vkCmdPipelineBarrier(CmdBuffer->GetHandle(), VK_PIPELINE_STAGE_ALL_COMMANDS_BIT, VK_PIPELINE_STAGE_ALL_COMMANDS_BIT, 0, 1, &Barrier, 0, nullptr, 0, nullptr);
+			CmdBuffer->BeginUniformUpdateBarrier();
 		}
+
 		VulkanRHI::vkCmdCopyBuffer(CmdBuffer->GetHandle(), LockInfo.GetHandle(), UBBuffer, 1, &Region);
+
+		if(bIsInRenderPass) //when updating outside render passes, the EndUniformUpdateBarrier will be called from EndRenderPass
+		{
+			CmdBuffer->EndUniformUpdateBarrier();
+		}
 	};
 
 	bool bRHIBypass = RHICmdList.Bypass();
