Fix shadow variable error
#rnx

diff --git a/Engine/Source/Runtime/VulkanRHI/Private/VulkanUniformBuffer.cpp b/Engine/Source/Runtime/VulkanRHI/Private/VulkanUniformBuffer.cpp
--- a/Engine/Source/Runtime/VulkanRHI/Private/VulkanUniformBuffer.cpp
+++ b/Engine/Source/Runtime/VulkanRHI/Private/VulkanUniformBuffer.cpp
@@ -206,8 +206,7 @@ inline void FVulkanDynamicRHI::UpdateUniformBuffer(FVulkanUniformBuffer* Uniform
 	FVulkanEmulatedUniformBuffer* EmulatedUniformBuffer = bRealUBs ? nullptr : (FVulkanEmulatedUniformBuffer*)UniformBuffer;
 
 	FBufferSuballocation* NewUBAlloc = nullptr;
-	bool bIsInRenderPass = RHICmdList.IsInsideRenderPass();
-	bool bUseUpload = GVulkanAllowUniformUpload && !bIsInRenderPass; //inside renderpasses, a rename is enforced.
+	bool bUseUpload = GVulkanAllowUniformUpload && !RHICmdList.IsInsideRenderPass(); //inside renderpasses, a rename is enforced.
 
 	if (bRealUBs && !bUseUpload)
 	{
