UE4.24 - More info when failing to map Vulkan memory
#rnx

diff --git a/Engine/Source/Runtime/VulkanRHI/Private/VulkanMemory.cpp b/Engine/Source/Runtime/VulkanRHI/Private/VulkanMemory.cpp
--- a/Engine/Source/Runtime/VulkanRHI/Private/VulkanMemory.cpp
+++ b/Engine/Source/Runtime/VulkanRHI/Private/VulkanMemory.cpp
@@ -432,7 +432,7 @@ namespace VulkanRHI
 	{
 		check(bCanBeMapped);
 		check(!MappedPointer);
-		check(InSize == VK_WHOLE_SIZE || InSize + Offset <= Size);
+		checkf(InSize == VK_WHOLE_SIZE || InSize + Offset <= Size, TEXT("Failed to Map %llu bytes, Offset %llu, AllocSize %llu bytes"), InSize, Offset, Size);
 
 		VERIFYVULKANRESULT(VulkanRHI::vkMapMemory(DeviceHandle, Handle, Offset, InSize, 0, &MappedPointer));
 		return MappedPointer;
