Vulkan: add ability to idle GPU after submitting specific cmdbuffers only.

- r.Vulkan.WaitForIdleOnSubmit=2 will wait only when submitting upload command buffer, =3 - only when submitting command buffer with gfc commands. This can be used to debug and/or workaround issues like UE-86149.

diff --git a/Engine/Source/Runtime/VulkanRHI/Private/VulkanQueue.cpp b/Engine/Source/Runtime/VulkanRHI/Private/VulkanQueue.cpp
--- a/Engine/Source/Runtime/VulkanRHI/Private/VulkanQueue.cpp
+++ b/Engine/Source/Runtime/VulkanRHI/Private/VulkanQueue.cpp
@@ -7,14 +7,17 @@
 #include "VulkanRHIPrivate.h"
 #include "VulkanQueue.h"
 #include "VulkanMemory.h"
+#include "VulkanContext.h"
 
 int32 GWaitForIdleOnSubmit = 0;
 FAutoConsoleVariableRef CVarVulkanWaitForIdleOnSubmit(
 	TEXT("r.Vulkan.WaitForIdleOnSubmit"),
 	GWaitForIdleOnSubmit,
-	TEXT("Waits for the GPU to be idle on every submit. Useful for tracking GPU hangs.\n")
+	TEXT("Waits for the GPU to be idle after submitting a command buffer. Useful for tracking GPU hangs.\n")
 	TEXT(" 0: Do not wait(default)\n")
-	TEXT(" 1: Wait"),
+	TEXT(" 1: Wait on every submit\n")
+	TEXT(" 2: Wait when submitting an upload buffer\n")
+	TEXT(" 3: Wait when submitting an active buffer (one that has gfx commands)\n"),
 	ECVF_Default
 	);
 
@@ -72,7 +75,36 @@ void FVulkanQueue::Submit(FVulkanCmdBuffer* CmdBuffer, uint32 NumSignalSemaphore
 	CmdBuffer->MarkSemaphoresAsSubmitted();
 	CmdBuffer->SubmittedFenceCounter = CmdBuffer->FenceSignaledCounter;
 
+	bool bShouldStall = false;
+
 	if (GWaitForIdleOnSubmit != 0)
+	{
+		// 1 - stall always
+		if (GWaitForIdleOnSubmit == 1)
+		{
+			bShouldStall = true;
+		}
+
+		// modulo check allows combining, e.g. 6 == 2*3, wait for both upload and gfx
+
+		// 2 - stall for Upload buffers
+		if (!bShouldStall && (GWaitForIdleOnSubmit % 2 == 0))
+		{
+			FVulkanCommandBufferManager* CmdBufferMgr = Device->GetImmediateContext().GetCommandBufferManager();
+			FVulkanCmdBuffer* UploadCmdBuffer = CmdBufferMgr->HasPendingUploadCmdBuffer() ? CmdBufferMgr->GetUploadCmdBuffer() : nullptr;
+			bShouldStall = (UploadCmdBuffer == CmdBuffer);
+		}
+
+		// 3 - stall for Active buffers (ones that have gfx commands)
+		if (!bShouldStall && (GWaitForIdleOnSubmit % 3 == 0))
+		{
+			FVulkanCommandBufferManager* CmdBufferMgr = Device->GetImmediateContext().GetCommandBufferManager();
+			FVulkanCmdBuffer* ActiveCmdBuffer = (CmdBufferMgr->HasPendingActiveCmdBuffer() && !CmdBufferMgr->HasPendingUploadCmdBuffer()) ? CmdBufferMgr->GetActiveCmdBuffer() : nullptr;
+			bShouldStall = (ActiveCmdBuffer == CmdBuffer);
+		}
+	}
+
+	if (bShouldStall)
 	{
 		// 200 ms timeout
 		bool bSuccess = Device->GetFenceManager().WaitForFence(CmdBuffer->Fence, 200 * 1000 * 1000);
