Remove unintended debug logging.
#rnx

diff --git a/Engine/Source/Runtime/VulkanRHI/Private/VulkanSwapChain.cpp b/Engine/Source/Runtime/VulkanRHI/Private/VulkanSwapChain.cpp
--- a/Engine/Source/Runtime/VulkanRHI/Private/VulkanSwapChain.cpp
+++ b/Engine/Source/Runtime/VulkanRHI/Private/VulkanSwapChain.cpp
@@ -461,14 +461,12 @@ FVulkanSwapChain::FVulkanSwapChain(VkInstance InInstance, FVulkanDevice& InDevic
 #endif
 	}
 
-	FPlatformMisc::LowLevelOutputDebugStringf(TEXT("FVulkanSwapChain::FVulkanSwapChain calling VulkanRHI::vkCreateSwapchainKHR(Device.GetInstanceHandle() %p, &SwapChainInfo %p, VULKAN_CPU_ALLOCATOR %p, &SwapChain %p)"), Device.GetInstanceHandle(), &SwapChainInfo, VULKAN_CPU_ALLOCATOR, &SwapChain);
 	VERIFYVULKANRESULT_EXPANDED(VulkanRHI::vkCreateSwapchainKHR(Device.GetInstanceHandle(), &SwapChainInfo, VULKAN_CPU_ALLOCATOR, &SwapChain));
-	 
+
 	if (RecreateInfo != nullptr)
 	{
 		if (RecreateInfo->SwapChain != VK_NULL_HANDLE)
 		{
-			FPlatformMisc::LowLevelOutputDebugStringf(TEXT("FVulkanSwapChain::FVulkanSwapChain, calling VulkanRHI::vkDestroySwapchainKHR(Device: %p, Swapchain: %p, Allocator: %p)"), Device.GetInstanceHandle(), RecreateInfo->SwapChain, VULKAN_CPU_ALLOCATOR);
 			VulkanRHI::vkDestroySwapchainKHR(Device.GetInstanceHandle(), RecreateInfo->SwapChain, VULKAN_CPU_ALLOCATOR);
 			RecreateInfo->SwapChain = VK_NULL_HANDLE;
 		}
@@ -524,8 +522,6 @@ void FVulkanSwapChain::Destroy(FVulkanSwapChainRecreateInfo* RecreateInfo)
 	Device.WaitUntilIdle();
 
 	bool bRecreate = RecreateInfo && GVulkanKeepSwapChain;
-	FPlatformMisc::LowLevelOutputDebugStringf(TEXT("FVulkanSwapChain::Destroy, would have called VulkanRHI::vkDestroySwapchainKHR(Device: %p, Swapchain: %p, Allocator: %p)"), Device.GetInstanceHandle(), SwapChain, VULKAN_CPU_ALLOCATOR);
-
 	if (bRecreate)
 	{
 		RecreateInfo->SwapChain = SwapChain;
