Workaround for blinking characters on some hardware (UE-86149).

- use nonCoherent Alignment even on coherent buffers.
- Fix by [REDACTED_PERSON].

diff --git a/Engine/Source/Runtime/VulkanRHI/Private/VulkanMemory.cpp b/Engine/Source/Runtime/VulkanRHI/Private/VulkanMemory.cpp
--- a/Engine/Source/Runtime/VulkanRHI/Private/VulkanMemory.cpp
+++ b/Engine/Source/Runtime/VulkanRHI/Private/VulkanMemory.cpp
@@ -1850,7 +1850,8 @@ namespace VulkanRHI
 
 		// Set minimum alignment to 16 bytes, as some buffers are used with CPU SIMD instructions
 		MemReqs.alignment = FMath::Max<VkDeviceSize>(16, MemReqs.alignment);
-		if (InMemoryReadFlags == VK_MEMORY_PROPERTY_HOST_CACHED_BIT)
+		static const bool bIsAmd = Device->GetDeviceProperties().vendorID == 0x1002;
+		if (InMemoryReadFlags == VK_MEMORY_PROPERTY_HOST_CACHED_BIT || bIsAmd)
 		{
 			uint64 NonCoherentAtomSize = (uint64)Device->GetLimits().nonCoherentAtomSize;
 			MemReqs.alignment = AlignArbitrary(MemReqs.alignment, NonCoherentAtomSize);
