Simplify r.Vulkan.WaitForIdleOnSubmit logic and fix SA warning (UE-87333).

#rnx

diff --git a/Engine/Source/Runtime/VulkanRHI/Private/VulkanQueue.cpp b/Engine/Source/Runtime/VulkanRHI/Private/VulkanQueue.cpp
--- a/Engine/Source/Runtime/VulkanRHI/Private/VulkanQueue.cpp
+++ b/Engine/Source/Runtime/VulkanRHI/Private/VulkanQueue.cpp
@@ -79,28 +79,23 @@ void FVulkanQueue::Submit(FVulkanCmdBuffer* CmdBuffer, uint32 NumSignalSemaphore
 
 	if (GWaitForIdleOnSubmit != 0)
 	{
-		// 1 - stall always
-		if (GWaitForIdleOnSubmit == 1)
-		{
-			bShouldStall = true;
-		}
-
-		// modulo check allows combining, e.g. 6 == 2*3, wait for both upload and gfx
-
-		// 2 - stall for Upload buffers
-		if (!bShouldStall && (GWaitForIdleOnSubmit % 2 == 0))
-		{
-			FVulkanCommandBufferManager* CmdBufferMgr = Device->GetImmediateContext().GetCommandBufferManager();
-			FVulkanCmdBuffer* UploadCmdBuffer = CmdBufferMgr->HasPendingUploadCmdBuffer() ? CmdBufferMgr->GetUploadCmdBuffer() : nullptr;
-			bShouldStall = (UploadCmdBuffer == CmdBuffer);
-		}
+		FVulkanCommandBufferManager* CmdBufferMgr = Device->GetImmediateContext().GetCommandBufferManager();
 
-		// 3 - stall for Active buffers (ones that have gfx commands)
-		if (!bShouldStall && (GWaitForIdleOnSubmit % 3 == 0))
+		switch(GWaitForIdleOnSubmit)
 		{
-			FVulkanCommandBufferManager* CmdBufferMgr = Device->GetImmediateContext().GetCommandBufferManager();
-			FVulkanCmdBuffer* ActiveCmdBuffer = (CmdBufferMgr->HasPendingActiveCmdBuffer() && !CmdBufferMgr->HasPendingUploadCmdBuffer()) ? CmdBufferMgr->GetActiveCmdBuffer() : nullptr;
-			bShouldStall = (ActiveCmdBuffer == CmdBuffer);
+			default:
+				// intentional fall-through
+			case 1:
+				bShouldStall = true;
+				break;
+
+			case 2:
+				bShouldStall = (CmdBufferMgr->HasPendingUploadCmdBuffer() && CmdBufferMgr->GetUploadCmdBuffer() == CmdBuffer);
+				break;
+
+			case 3:
+				bShouldStall = (CmdBufferMgr->HasPendingActiveCmdBuffer() && CmdBufferMgr->GetActiveCmdBufferDirect() == CmdBuffer);
+				break;
 		}
 	}
 
