Fix assert on end PIE (or closing app) using D3D12 and SteamVR; this also fixes other potential issues around lifetime with D3D12 and Vulkan textures with the XrSwapChain class and resource aliasing.
Also switches XrSwapchain and SteamVR to use the new RHICreateAliasedTexture API (to save a redundant/dummy texture creation during initialization).

diff --git a/Engine/Plugins/Lumin/MagicLeap/Source/MagicLeapHelperVulkan/Private/MagicLeapHelperVulkan.cpp b/Engine/Plugins/Lumin/MagicLeap/Source/MagicLeapHelperVulkan/Private/MagicLeapHelperVulkan.cpp
--- a/Engine/Plugins/Lumin/MagicLeap/Source/MagicLeapHelperVulkan/Private/MagicLeapHelperVulkan.cpp
+++ b/Engine/Plugins/Lumin/MagicLeap/Source/MagicLeapHelperVulkan/Private/MagicLeapHelperVulkan.cpp
@@ -320,9 +320,9 @@ bool FMagicLeapHelperVulkan::GetMediaTexture(FTextureRHIRef& ResultTexture, FSam
 	return false;
 }
 
-void FMagicLeapHelperVulkan::AliasMediaTexture(FRHITexture* DestTexture, FRHITexture* SrcTexture)
+void FMagicLeapHelperVulkan::AliasMediaTexture(FTextureRHIRef& DestTexture, FTextureRHIRef& SrcTexture)
 {
 #if PLATFORM_LUMIN
-	GDynamicRHI->RHIAliasTextureResources(DestTexture, SrcTexture);
+	GDynamicRHI->RHIAliasTextureResources((FTextureRHIRef&)DestTexture, (FTextureRHIRef&)SrcTexture);
 #endif
 }
diff --git a/Engine/Plugins/Lumin/MagicLeap/Source/MagicLeapHelperVulkan/Public/MagicLeapHelperVulkan.h b/Engine/Plugins/Lumin/MagicLeap/Source/MagicLeapHelperVulkan/Public/MagicLeapHelperVulkan.h
--- a/Engine/Plugins/Lumin/MagicLeap/Source/MagicLeapHelperVulkan/Public/MagicLeapHelperVulkan.h
+++ b/Engine/Plugins/Lumin/MagicLeap/Source/MagicLeapHelperVulkan/Public/MagicLeapHelperVulkan.h
@@ -16,5 +16,5 @@ public:
 	static bool GetVulkanInstanceExtensionsRequired(TArray<const ANSICHAR*>& Out);
 	static bool GetVulkanDeviceExtensionsRequired(struct VkPhysicalDevice_T* pPhysicalDevice, TArray<const ANSICHAR*>& Out);
 	static bool GetMediaTexture(FTextureRHIRef& Result, FSamplerStateRHIRef& SamplerResult, const uint64 MediaTextureHandle);
-	static void AliasMediaTexture(FRHITexture* DestTexture, FRHITexture* SrcTexture);
+	static void AliasMediaTexture(FTextureRHIRef& DestTexture, FTextureRHIRef& SrcTexture);
 };
diff --git a/Engine/Plugins/Runtime/OpenXR/Source/OpenXRHMD/Private/OpenXRHMD_Swapchain.cpp b/Engine/Plugins/Runtime/OpenXR/Source/OpenXRHMD/Private/OpenXRHMD_Swapchain.cpp
--- a/Engine/Plugins/Runtime/OpenXR/Source/OpenXRHMD/Private/OpenXRHMD_Swapchain.cpp
+++ b/Engine/Plugins/Runtime/OpenXR/Source/OpenXRHMD/Private/OpenXRHMD_Swapchain.cpp
@@ -33,7 +33,7 @@ void FOpenXRSwapchain::IncrementSwapChainIndex_RHIThread(int64 Timeout)
 	WaitInfo.timeout = Timeout;
 	XR_ENSURE(xrWaitSwapchainImage(Handle, &WaitInfo));
 
-	GDynamicRHI->RHIAliasTextureResources(RHITexture, RHITextureSwapChain[SwapChainIndex_RHIThread]);
+	GDynamicRHI->RHIAliasTextureResources((FTextureRHIRef&)RHITexture, (FTextureRHIRef&)RHITextureSwapChain[SwapChainIndex_RHIThread]);
 }
 
 void FOpenXRSwapchain::ReleaseCurrentImage_RHIThread()
@@ -194,7 +194,6 @@ FXRSwapChainPtr CreateSwapchain_D3D11(XrSession InSession, uint8 Format, uint32
 
 	FD3D11DynamicRHI* DynamicRHI = static_cast<FD3D11DynamicRHI*>(GDynamicRHI);
 	TArray<FTextureRHIRef> TextureChain;
-	// @todo: Once things settle down, the chain target will be created below in the CreateXRSwapChain call, via an RHI "CreateAliasedTexture" call.
 	TArray<XrSwapchainImageD3D11KHR> Images = EnumerateImages<XrSwapchainImageD3D11KHR>(Swapchain, XR_TYPE_SWAPCHAIN_IMAGE_D3D11_KHR);
 	for (const auto& Image : Images)
 	{
@@ -234,7 +233,6 @@ FXRSwapChainPtr CreateSwapchain_D3D12(XrSession InSession, uint8 Format, uint32
 
 	FD3D12DynamicRHI* DynamicRHI = static_cast<FD3D12DynamicRHI*>(GDynamicRHI);
 	TArray<FTextureRHIRef> TextureChain;
-	// @todo: Once things settle down, the chain target will be created below in the CreateXRSwapChain call, via an RHI "CreateAliasedTexture" call.
 	TArray<XrSwapchainImageD3D12KHR> Images = EnumerateImages<XrSwapchainImageD3D12KHR>(Swapchain, XR_TYPE_SWAPCHAIN_IMAGE_D3D12_KHR);
 	for (const auto& Image : Images)
 	{
@@ -263,7 +261,6 @@ FXRSwapChainPtr CreateSwapchain_OpenGL(XrSession InSession, uint8 Format, uint32
 
 	TArray<FTextureRHIRef> TextureChain;
 	FOpenGLDynamicRHI* DynamicRHI = static_cast<FOpenGLDynamicRHI*>(GDynamicRHI);
-	// @todo: Once things settle down, the chain target will be created below in the CreateXRSwapChain call, via an RHI "CreateAliasedTexture" call.
 	TArray<XrSwapchainImageOpenGLKHR> Images = EnumerateImages<XrSwapchainImageOpenGLKHR>(Swapchain, XR_TYPE_SWAPCHAIN_IMAGE_OPENGL_KHR);
 	for (const auto& Image : Images)
 	{
@@ -297,7 +294,6 @@ FXRSwapChainPtr CreateSwapchain_Vulkan(XrSession InSession, uint8 Format, uint32
 
 	TArray<FTextureRHIRef> TextureChain;
 	FVulkanDynamicRHI* DynamicRHI = static_cast<FVulkanDynamicRHI*>(GDynamicRHI);
-	// @todo: Once things settle down, the chain target will be created below in the CreateXRSwapChain call, via an RHI "CreateAliasedTexture" call.
 	TArray<XrSwapchainImageVulkanKHR> Images = EnumerateImages<XrSwapchainImageVulkanKHR>(Swapchain, XR_TYPE_SWAPCHAIN_IMAGE_VULKAN_KHR);
 	for (const auto& Image : Images)
 	{
diff --git a/Engine/Plugins/Runtime/Steam/SteamVR/Source/SteamVR/Private/SteamVRHMD.cpp b/Engine/Plugins/Runtime/Steam/SteamVR/Source/SteamVR/Private/SteamVRHMD.cpp
--- a/Engine/Plugins/Runtime/Steam/SteamVR/Source/SteamVR/Private/SteamVRHMD.cpp
+++ b/Engine/Plugins/Runtime/Steam/SteamVR/Source/SteamVR/Private/SteamVRHMD.cpp
@@ -1637,16 +1637,7 @@ bool FSteamVRHMD::AllocateRenderTargetTexture(uint32 Index, uint32 SizeX, uint32
 
 		if (BindingTexture == nullptr)
 		{
-			// For this iteration, create a temporary texture resource that will get stomped the first time we do our bind/alias
-			// pass on this object.
-			// Note, even though this is a wasted allocation, it never gets bound or used in any way, so should be minimally impactful.
-			FTexture2DRHIRef TempTargetableTexture, TempShaderResourceTexture;
-			RHICreateTargetableShaderResource2D(SizeX, SizeY, PF_B8G8R8A8, 1, TexCreate_None, TexCreate_RenderTargetable, false, CreateInfo, TempTargetableTexture, TempShaderResourceTexture, NumSamples);
-			check(TempTargetableTexture == TempShaderResourceTexture);
-			BindingTexture = TempTargetableTexture;
-
-			// Alias immediately.
-			GDynamicRHI->RHIAliasTextureResources(BindingTexture, TargetableTexture);
+			BindingTexture = GDynamicRHI->RHICreateAliasedTexture((FTextureRHIRef&)TargetableTexture);
 		}
 	}
 
@@ -1712,16 +1703,9 @@ bool FSteamVRHMD::AllocateDepthTexture(uint32 Index, uint32 SizeX, uint32 SizeY,
 
 		SwapChainTextures.Add((FTextureRHIRef&)TargetableTexture);
 
-		// As above, create a temp, stompable texture to alias over.
 		if (BindingTexture == nullptr)
 		{
-			FTexture2DRHIRef TempTargetableTexture, TempShaderResourceTexture;
-			RHICreateTargetableShaderResource2D(SizeX, SizeY, PF_DepthStencil, 1, Flags, TargetableTextureFlags, false, CreateInfo, TempTargetableTexture, TempShaderResourceTexture, 1);
-			check(TempTargetableTexture == TempShaderResourceTexture);
-			BindingTexture = TempTargetableTexture;
-
-			// Alias immediately.
-			GDynamicRHI->RHIAliasTextureResources(BindingTexture, TargetableTexture);
+			BindingTexture = GDynamicRHI->RHICreateAliasedTexture((FTextureRHIRef&)TargetableTexture);
 		}
 	}
 
diff --git a/Engine/Source/Runtime/D3D12RHI/Private/D3D12RHIPrivate.h b/Engine/Source/Runtime/D3D12RHI/Private/D3D12RHIPrivate.h
--- a/Engine/Source/Runtime/D3D12RHI/Private/D3D12RHIPrivate.h
+++ b/Engine/Source/Runtime/D3D12RHI/Private/D3D12RHIPrivate.h
@@ -342,8 +342,10 @@ public:
 	virtual uint32 RHIGetViewportNextPresentGPUIndex(FRHIViewport* Viewport) final override;
 	virtual FTexture2DRHIRef RHIGetViewportBackBuffer(FRHIViewport* Viewport) final override;
 	virtual void RHIAliasTextureResources(FRHITexture* DestTexture, FRHITexture* SrcTexture) final override;
-	virtual void RHIAdvanceFrameFence() final override;
 	virtual FTextureRHIRef RHICreateAliasedTexture(FRHITexture* SourceTexture) final override;
+	virtual void RHIAliasTextureResources(FTextureRHIRef& DestTexture, FTextureRHIRef& SrcTexture) final override;
+	virtual FTextureRHIRef RHICreateAliasedTexture(FTextureRHIRef& SourceTexture) final override;
+	virtual void RHIAdvanceFrameFence() final override;
 	virtual void RHIAdvanceFrameForGetViewportBackBuffer(FRHIViewport* Viewport) final override;
 	virtual void RHIAcquireThreadOwnership() final override;
 	virtual void RHIReleaseThreadOwnership() final override;
diff --git a/Engine/Source/Runtime/D3D12RHI/Private/D3D12Texture.cpp b/Engine/Source/Runtime/D3D12RHI/Private/D3D12Texture.cpp
--- a/Engine/Source/Runtime/D3D12RHI/Private/D3D12Texture.cpp
+++ b/Engine/Source/Runtime/D3D12RHI/Private/D3D12Texture.cpp
@@ -364,7 +364,7 @@ void FD3D12TextureStats::D3D12TextureDeleted(TD3D12Texture2D<BaseResourceType>&
 	{
 		const D3D12_RESOURCE_DESC& Desc = D3D12Texture2D->GetDesc();
 		const int64 TextureSize = Texture.GetMemorySize();
-		ensure(TextureSize > 0 || (Texture.Flags & TexCreate_Virtual));
+		ensure(TextureSize > 0 || (Texture.Flags & TexCreate_Virtual) || Texture.GetAliasingSourceTexture() != nullptr);
 
 		UpdateD3D12TextureStats(Desc, -TextureSize, false, Texture.IsCubemap());
 
@@ -2819,10 +2819,32 @@ FTextureCubeRHIRef FD3D12DynamicRHI::RHICreateTextureCubeFromResource(EPixelForm
 }
 
 void FD3D12DynamicRHI::RHIAliasTextureResources(FRHITexture* DestTextureRHI, FRHITexture* SrcTextureRHI)
-{	
+{
 	FD3D12TextureBase* DestTexture = GetD3D12TextureFromRHITexture(DestTextureRHI);
 	FD3D12TextureBase* SrcTexture = GetD3D12TextureFromRHITexture(SrcTextureRHI);
 
+	// This path will potentially cause crashes, if the source texture is destroyed and we're still being used. This
+	// API path will be deprecated post 4.25. To avoid issues, use the version that takes FTextureRHIRef references instead.
+	check(false);
+
+	while (DestTexture && SrcTexture)
+	{
+		DestTexture->AliasResources(SrcTexture);
+
+		DestTexture = DestTexture->GetNextObject();
+		SrcTexture = SrcTexture->GetNextObject();
+	}
+}
+
+void FD3D12DynamicRHI::RHIAliasTextureResources(FTextureRHIRef& DestTextureRHI, FTextureRHIRef& SrcTextureRHI)
+{
+	FD3D12TextureBase* DestTexture = GetD3D12TextureFromRHITexture(DestTextureRHI);
+	FD3D12TextureBase* SrcTexture = GetD3D12TextureFromRHITexture(SrcTextureRHI);
+
+	// Make sure we keep a reference to the source texture we're aliasing, so we don't lose it if all other references
+	// go away but we're kept around.
+	DestTexture->SetAliasingSource(SrcTextureRHI);
+
 	while (DestTexture && SrcTexture)
 	{
 		DestTexture->AliasResources(SrcTexture);
@@ -2900,7 +2922,7 @@ TD3D12Texture2D<BaseResourceType>* FD3D12DynamicRHI::CreateAliasedD3D12Texture2D
 		}
 	}
 
-	RHIAliasTextureResources(Texture2D, SourceTexture);
+	RHIAliasTextureResources((FTextureRHIRef&)Texture2D, (FTextureRHIRef&)SourceTexture);
 
 	return Texture2D;
 }
@@ -2925,6 +2947,35 @@ FTextureRHIRef FD3D12DynamicRHI::RHICreateAliasedTexture(FRHITexture* SourceText
 	return nullptr;
 }
 
+FTextureRHIRef FD3D12DynamicRHI::RHICreateAliasedTexture(FTextureRHIRef& SourceTextureRHI)
+{
+	FD3D12TextureBase* SourceTexture = GetD3D12TextureFromRHITexture(SourceTextureRHI);
+	FTextureRHIRef ReturnTexture;
+	if (SourceTextureRHI->GetTexture2D() != nullptr)
+	{
+		ReturnTexture = CreateAliasedD3D12Texture2D<FD3D12BaseTexture2D>(static_cast<FD3D12Texture2D*>(SourceTextureRHI->GetTexture2D()));
+	}
+	else if (SourceTextureRHI->GetTexture2DArray() != nullptr)
+	{
+		ReturnTexture = CreateAliasedD3D12Texture2D<FD3D12BaseTexture2DArray>(static_cast<FD3D12Texture2DArray*>(SourceTextureRHI->GetTexture2DArray()));
+	}
+	else if (SourceTextureRHI->GetTextureCube() != nullptr)
+	{
+		ReturnTexture = CreateAliasedD3D12Texture2D<FD3D12BaseTextureCube>(static_cast<FD3D12TextureCube*>(SourceTextureRHI->GetTextureCube()));
+	}
+
+	if (ReturnTexture == nullptr)
+	{
+		UE_LOG(LogD3D12RHI, Error, TEXT("Currently FD3D12DynamicRHI::RHICreateAliasedTexture only supports 2D, 2D Array and Cube textures."));
+		return nullptr;
+	}
+
+	FD3D12TextureBase* DestTexture = GetD3D12TextureFromRHITexture(ReturnTexture);
+	DestTexture->SetAliasingSource(SourceTextureRHI);
+
+	return ReturnTexture;
+}
+
 void FD3D12DynamicRHI::RHICopySubTextureRegion(FRHITexture2D* SourceTextureRHI, FRHITexture2D* DestTextureRHI, FBox2D SourceBox, FBox2D DestinationBox)
 {
 	FD3D12TextureBase* SourceTexture = GetD3D12TextureFromRHITexture(SourceTextureRHI);
diff --git a/Engine/Source/Runtime/D3D12RHI/Private/D3D12Texture.h b/Engine/Source/Runtime/D3D12RHI/Private/D3D12Texture.h
--- a/Engine/Source/Runtime/D3D12RHI/Private/D3D12Texture.h
+++ b/Engine/Source/Runtime/D3D12RHI/Private/D3D12Texture.h
@@ -85,11 +85,18 @@ public:
 		MemorySize = InMemorySize;
 	}
 
+	void SetAliasingSource(FTextureRHIRef& SourceTextureRHI)
+	{
+		AliasingSourceTexture = SourceTextureRHI;
+	}
+
 	// Accessors.
 	FD3D12Resource* GetResource() const { return ResourceLocation.GetResource(); }
 	uint64 GetOffset() const { return ResourceLocation.GetOffsetFromBaseOfResource(); }
 	FD3D12ShaderResourceView* GetShaderResourceView() const { return ShaderResourceView; }
 	FD3D12BaseShaderResource* GetBaseShaderResource() const { return BaseShaderResource; }
+	inline const FTextureRHIRef& GetAliasingSourceTexture() const { return AliasingSourceTexture; }
+
 	void SetShaderResourceView(FD3D12ShaderResourceView* InShaderResourceView) { ShaderResourceView = InShaderResourceView; }
 
 	static inline bool ShouldDeferCmdListOperation(FRHICommandList* RHICmdList)
@@ -199,6 +206,8 @@ protected:
 	uint32	NumDepthStencilViews;
 
 	TMap<uint32, FD3D12LockedResource*> LockedMap;
+
+	FTextureRHIRef AliasingSourceTexture;
 };
 
 #if !PLATFORM_SUPPORTS_VIRTUAL_TEXTURES
diff --git a/Engine/Source/Runtime/HeadMountedDisplay/Private/XRSwapChain.cpp b/Engine/Source/Runtime/HeadMountedDisplay/Private/XRSwapChain.cpp
--- a/Engine/Source/Runtime/HeadMountedDisplay/Private/XRSwapChain.cpp
+++ b/Engine/Source/Runtime/HeadMountedDisplay/Private/XRSwapChain.cpp
@@ -117,7 +117,7 @@ void FXRSwapChain::IncrementSwapChainIndex_RHIThread(int64 /* TimeoutNanoseconds
 	CheckInRHIThread();
 
 	SwapChainIndex_RHIThread = (SwapChainIndex_RHIThread + 1) % GetSwapChainLength();
-	GDynamicRHI->RHIAliasTextureResources(RHITexture, RHITextureSwapChain[SwapChainIndex_RHIThread]);
+	GDynamicRHI->RHIAliasTextureResources((FTextureRHIRef&)RHITexture, (FTextureRHIRef&)RHITextureSwapChain[SwapChainIndex_RHIThread]);
 }
 
 
diff --git a/Engine/Source/Runtime/NullDrv/Public/NullRHI.h b/Engine/Source/Runtime/NullDrv/Public/NullRHI.h
--- a/Engine/Source/Runtime/NullDrv/Public/NullRHI.h
+++ b/Engine/Source/Runtime/NullDrv/Public/NullRHI.h
@@ -471,6 +471,10 @@ public:
 	virtual void RHIAliasTextureResources(FRHITexture* DestTexture, FRHITexture* SrcTexture) final override
 	{
 
+	}
+	virtual void RHIAliasTextureResources(FTextureRHIRef& DestTexture, FTextureRHIRef& SrcTexture) final override
+	{
+
 	}
 	virtual void RHIAdvanceFrameForGetViewportBackBuffer(FRHIViewport* Viewport) final override
 	{
diff --git a/Engine/Source/Runtime/RHI/Public/DynamicRHI.h b/Engine/Source/Runtime/RHI/Public/DynamicRHI.h
--- a/Engine/Source/Runtime/RHI/Public/DynamicRHI.h
+++ b/Engine/Source/Runtime/RHI/Public/DynamicRHI.h
@@ -1004,18 +1004,32 @@ public:
 		return nullptr;
 	}
 
+	UE_DEPRECATED(4.25, "RHIAliasTextureResources now takes references to FTextureRHIRef objects as parameters")
 	virtual void RHIAliasTextureResources(FRHITexture* DestTexture, FRHITexture* SrcTexture)
 	{
 		checkNoEntry();
 	}
 
-	virtual void RHIAdvanceFrameFence(){};
+	UE_DEPRECATED(4.25, "RHICreateAliasedTexture now takes a reference to an FTextureRHIRef object")
 	virtual FTextureRHIRef RHICreateAliasedTexture(FRHITexture* SourceTexture)
 	{
 		checkNoEntry();
 		return nullptr;
 	}
 
+	virtual void RHIAliasTextureResources(FTextureRHIRef& DestTexture, FTextureRHIRef& SrcTexture)
+	{
+		checkNoEntry();
+	}
+
+	virtual FTextureRHIRef RHICreateAliasedTexture(FTextureRHIRef& SourceTexture)
+	{
+		checkNoEntry();
+		return nullptr;
+	}
+
+	virtual void RHIAdvanceFrameFence() {};
+
 	// Only relevant with an RHI thread, this advances the backbuffer for the purpose of GetViewportBackBuffer
 	// FlushType: Thread safe
 	virtual void RHIAdvanceFrameForGetViewportBackBuffer(FRHIViewport* Viewport) = 0;
diff --git a/Engine/Source/Runtime/RHI/Public/RHIValidation.h b/Engine/Source/Runtime/RHI/Public/RHIValidation.h
--- a/Engine/Source/Runtime/RHI/Public/RHIValidation.h
+++ b/Engine/Source/Runtime/RHI/Public/RHIValidation.h
@@ -999,6 +999,21 @@ public:
 		return RHI->RHICreateAliasedTexture(SourceTexture);
 	}
 
+	virtual void RHIAliasTextureResources(FTextureRHIRef& DestTexture, FTextureRHIRef& SourceTexture) override final
+	{
+		// Source and target need to be valid objects.
+		check(DestTexture && SourceTexture);
+		// Source texture must have been created (i.e. have a native resource backing).
+		check(SourceTexture->GetNativeResource() != nullptr);
+		RHI->RHIAliasTextureResources(DestTexture, SourceTexture);
+	}
+
+	virtual FTextureRHIRef RHICreateAliasedTexture(FTextureRHIRef& SourceTexture) override final
+	{
+		check(SourceTexture);
+		return RHI->RHICreateAliasedTexture(SourceTexture);
+	}
+
 	// Only relevant with an RHI thread, this advances the backbuffer for the purpose of GetViewportBackBuffer
 	// FlushType: Thread safe
 	virtual void RHIAdvanceFrameForGetViewportBackBuffer(FRHIViewport* Viewport) override final
diff --git a/Engine/Source/Runtime/VulkanRHI/Private/VulkanRHI.cpp b/Engine/Source/Runtime/VulkanRHI/Private/VulkanRHI.cpp
--- a/Engine/Source/Runtime/VulkanRHI/Private/VulkanRHI.cpp
+++ b/Engine/Source/Runtime/VulkanRHI/Private/VulkanRHI.cpp
@@ -1044,37 +1044,49 @@ FTextureCubeRHIRef FVulkanDynamicRHI::RHICreateTextureCubeFromResource(EPixelFor
 }
 
 void FVulkanDynamicRHI::RHIAliasTextureResources(FRHITexture* DestTextureRHI, FRHITexture* SrcTextureRHI)
+{
+	check(false);
+}
+
+FTextureRHIRef FVulkanDynamicRHI::RHICreateAliasedTexture(FRHITexture* SourceTexture)
+{
+	check(false);
+	return nullptr;
+}
+
+void FVulkanDynamicRHI::RHIAliasTextureResources(FTextureRHIRef& DestTextureRHI, FTextureRHIRef& SrcTextureRHI)
 {
 	if (DestTextureRHI && SrcTextureRHI)
 	{
-		FVulkanTextureBase* DestTextureBase = (FVulkanTextureBase*) DestTextureRHI->GetTextureBaseRHI();
-		FVulkanTextureBase* SrcTextureBase = (FVulkanTextureBase*) SrcTextureRHI->GetTextureBaseRHI();
+		FVulkanTextureBase* DestTextureBase = (FVulkanTextureBase*)DestTextureRHI->GetTextureBaseRHI();
+		FVulkanTextureBase* SrcTextureBase = (FVulkanTextureBase*)SrcTextureRHI->GetTextureBaseRHI();
 
 		if (DestTextureBase && SrcTextureBase)
 		{
-			DestTextureBase->AliasTextureResources(SrcTextureBase);
+			DestTextureBase->AliasTextureResources(SrcTextureRHI);
 		}
 	}
 }
 
-FTextureRHIRef FVulkanDynamicRHI::RHICreateAliasedTexture(FRHITexture* SourceTexture)
+FTextureRHIRef FVulkanDynamicRHI::RHICreateAliasedTexture(FTextureRHIRef& SourceTextureRHI)
 {
+	FVulkanTextureBase* SourceTexture = (FVulkanTextureBase*)SourceTextureRHI->GetTextureBaseRHI();
 	FTextureRHIRef AliasedTexture;
-	if (SourceTexture->GetTexture2D() != nullptr)
+	if (SourceTextureRHI->GetTexture2D() != nullptr)
 	{
-		AliasedTexture = new FVulkanTexture2D(static_cast<FVulkanTexture2D*>(SourceTexture));
+		AliasedTexture = new FVulkanTexture2D(SourceTextureRHI, (FVulkanTexture2D*)SourceTexture);
 	}
-	else if (SourceTexture->GetTexture2DArray() != nullptr)
+	else if (SourceTextureRHI->GetTexture2DArray() != nullptr)
 	{
-		AliasedTexture = new FVulkanTexture2DArray(static_cast<FVulkanTexture2DArray*>(SourceTexture));
+		AliasedTexture = new FVulkanTexture2DArray(SourceTextureRHI, (FVulkanTexture2DArray*)SourceTexture);
 	}
-	else if (SourceTexture->GetTextureCube() != nullptr)
+	else if (SourceTextureRHI->GetTextureCube() != nullptr)
 	{
-		AliasedTexture = new FVulkanTextureCube(static_cast<FVulkanTextureCube*>(SourceTexture));
+		AliasedTexture = new FVulkanTextureCube(SourceTextureRHI, (FVulkanTextureCube*)SourceTexture);
 	}
 	else
 	{
-		UE_LOG(LogRHI, Error, TEXT("Currently FOpenGLDynamicRHI::RHICreateAliasedTexture only supports 2D, 2D Array and Cube textures."));
+		UE_LOG(LogRHI, Error, TEXT("Currently FVulkanDynamicRHI::RHICreateAliasedTexture only supports 2D, 2D Array and Cube textures."));
 	}
 
 	return AliasedTexture;
diff --git a/Engine/Source/Runtime/VulkanRHI/Private/VulkanResources.h b/Engine/Source/Runtime/VulkanRHI/Private/VulkanResources.h
--- a/Engine/Source/Runtime/VulkanRHI/Private/VulkanResources.h
+++ b/Engine/Source/Runtime/VulkanRHI/Private/VulkanResources.h
@@ -486,11 +486,11 @@ struct FVulkanTextureBase : public FVulkanBaseShaderResource
 	FVulkanTextureBase(FVulkanDevice& Device, VkImageViewType ResourceType, EPixelFormat Format, uint32 SizeX, uint32 SizeY, uint32 SizeZ, uint32 ArraySize, uint32 NumMips, uint32 NumSamples, VkImage InImage, VkDeviceMemory InMem, FSamplerYcbcrConversionInitializer& ConversionInitializer, uint32 UEFlags, const FRHIResourceCreateInfo& CreateInfo = FRHIResourceCreateInfo());
 
 	// Aliasing constructor.
-	FVulkanTextureBase(const FVulkanTextureBase* SrcTexture, VkImageViewType ResourceType, uint32 SizeX, uint32 SizeY, uint32 sizeZ);
+	FVulkanTextureBase(FTextureRHIRef& SrcTextureRHI, const FVulkanTextureBase* SrcTexture, VkImageViewType ResourceType, uint32 SizeX, uint32 SizeY, uint32 sizeZ);
 
 	virtual ~FVulkanTextureBase();
 
-	void AliasTextureResources(const FVulkanTextureBase* SrcTexture);
+	void AliasTextureResources(FTextureRHIRef& SrcTexture);
 
 	FVulkanSurface Surface;
 
@@ -499,7 +499,7 @@ struct FVulkanTextureBase : public FVulkanBaseShaderResource
 	// View with all mips/layers, but if it's a Depth/Stencil, only the Depth view
 	FVulkanTextureView* PartialView;
 
-	bool bIsAliased;
+	FTextureRHIRef AliasedTexture;
 
 	virtual void OnTransitionResource(FVulkanCommandListContext& Context, EResourceTransitionAccess TransitionType) {};
 
@@ -515,7 +515,7 @@ public:
 	FVulkanTexture2D(FVulkanDevice& Device, EPixelFormat Format, uint32 SizeX, uint32 SizeY, uint32 NumMips, uint32 NumSamples, VkImage Image, struct FSamplerYcbcrConversionInitializer& ConversionInitializer, uint32 UEFlags, const FRHIResourceCreateInfo& CreateInfo);
 
 	// Aliasing constructor
-	FVulkanTexture2D(const FVulkanTexture2D* SrcTexture);
+	FVulkanTexture2D(FTextureRHIRef& SrcTextureRHI, const FVulkanTexture2D* SrcTexture);
 
 	virtual ~FVulkanTexture2D();
 
@@ -553,7 +553,7 @@ public:
 	FVulkanTexture2DArray(FVulkanDevice& Device, EPixelFormat Format, uint32 SizeX, uint32 SizeY, uint32 ArraySize, uint32 NumMips, uint32 NumSamples, VkImage Image, uint32 Flags, FResourceBulkDataInterface* BulkData, const FClearValueBinding& InClearValue);
 
 	// Aliasing constructor
-	FVulkanTexture2DArray(const FVulkanTexture2DArray* SrcTexture);
+	FVulkanTexture2DArray(FTextureRHIRef& SrcTextureRHI, const FVulkanTexture2DArray* SrcTexture);
 
 		
 	// IRefCountedObject interface.
@@ -621,7 +621,7 @@ public:
 	FVulkanTextureCube(FVulkanDevice& Device, EPixelFormat Format, uint32 Size, bool bArray, uint32 ArraySize, uint32 NumMips, VkImage Image, uint32 Flags, FResourceBulkDataInterface* BulkData, const FClearValueBinding& InClearValue);
 
 	// Aliasing constructor
-	FVulkanTextureCube(const FVulkanTextureCube* SrcTexture);
+	FVulkanTextureCube(FTextureRHIRef& SrcTextureRHI, const FVulkanTextureCube* SrcTexture);
 
 	virtual ~FVulkanTextureCube();
 
diff --git a/Engine/Source/Runtime/VulkanRHI/Private/VulkanTexture.cpp b/Engine/Source/Runtime/VulkanRHI/Private/VulkanTexture.cpp
--- a/Engine/Source/Runtime/VulkanRHI/Private/VulkanTexture.cpp
+++ b/Engine/Source/Runtime/VulkanRHI/Private/VulkanTexture.cpp
@@ -1643,7 +1643,6 @@ void FVulkanTextureView::Destroy(FVulkanDevice& Device)
 FVulkanTextureBase::FVulkanTextureBase(FVulkanDevice& Device, VkImageViewType ResourceType, EPixelFormat InFormat, uint32 SizeX, uint32 SizeY, uint32 SizeZ, uint32 ArraySize, uint32 NumMips, uint32 NumSamples, uint32 UEFlags, const FRHIResourceCreateInfo& CreateInfo)
 	: Surface(Device, ResourceType, InFormat, SizeX, SizeY, SizeZ, ArraySize, NumMips, NumSamples, UEFlags, CreateInfo)
 	, PartialView(nullptr)
-	, bIsAliased(false)
 {
 	LLM_SCOPE_VULKAN(ELLMTagVulkan::VulkanTextures);
 	const bool bArray = ResourceType == VK_IMAGE_VIEW_TYPE_1D_ARRAY || ResourceType == VK_IMAGE_VIEW_TYPE_2D_ARRAY || ResourceType == VK_IMAGE_VIEW_TYPE_CUBE_ARRAY;
@@ -1727,7 +1726,6 @@ FVulkanTextureBase::FVulkanTextureBase(FVulkanDevice& Device, VkImageViewType Re
 FVulkanTextureBase::FVulkanTextureBase(FVulkanDevice& Device, VkImageViewType ResourceType, EPixelFormat Format, uint32 SizeX, uint32 SizeY, uint32 SizeZ, uint32 ArraySize, uint32 InNumMips, uint32 InNumSamples, VkImage InImage, VkDeviceMemory InMem, uint32 UEFlags, const FRHIResourceCreateInfo& CreateInfo)
 	: Surface(Device, ResourceType, Format, SizeX, SizeY, SizeZ, ArraySize, InNumMips, InNumSamples, InImage, UEFlags, CreateInfo)
 	, PartialView(nullptr)
-	, bIsAliased(false)
 {
 	check(InMem == VK_NULL_HANDLE);
 	const bool bArray = ResourceType == VK_IMAGE_VIEW_TYPE_1D_ARRAY || ResourceType == VK_IMAGE_VIEW_TYPE_2D_ARRAY || ResourceType == VK_IMAGE_VIEW_TYPE_CUBE_ARRAY;
@@ -1751,7 +1749,6 @@ FVulkanTextureBase::FVulkanTextureBase(FVulkanDevice& Device, VkImageViewType Re
 FVulkanTextureBase::FVulkanTextureBase(FVulkanDevice& Device, VkImageViewType ResourceType, EPixelFormat Format, uint32 SizeX, uint32 SizeY, uint32 SizeZ, uint32 ArraySize, uint32 NumMips, uint32 NumSamples, VkImage InImage, VkDeviceMemory InMem, FSamplerYcbcrConversionInitializer& ConversionInitializer, uint32 UEFlags, const FRHIResourceCreateInfo& CreateInfo)
 	: Surface(Device, ResourceType, Format, SizeX, SizeY, SizeZ, ArraySize, NumMips, NumSamples, InImage, UEFlags, CreateInfo)
 	, PartialView(nullptr)
-	, bIsAliased(false)
 {
 	check(InMem == VK_NULL_HANDLE);
 	const bool bArray = ResourceType == VK_IMAGE_VIEW_TYPE_1D_ARRAY || ResourceType == VK_IMAGE_VIEW_TYPE_2D_ARRAY || ResourceType == VK_IMAGE_VIEW_TYPE_CUBE_ARRAY;
@@ -1782,10 +1779,10 @@ FVulkanTextureBase::FVulkanTextureBase(FVulkanDevice& Device, VkImageViewType Re
 	InsertInitialImageLayout(Device, InImage, VK_IMAGE_LAYOUT_SHADER_READ_ONLY_OPTIMAL);
 }
 
-FVulkanTextureBase::FVulkanTextureBase(const FVulkanTextureBase* SrcTexture, VkImageViewType ResourceType, uint32 SizeX, uint32 SizeY, uint32 SizeZ)
+FVulkanTextureBase::FVulkanTextureBase(FTextureRHIRef& SrcTextureRHI, const FVulkanTextureBase* SrcTexture, VkImageViewType ResourceType, uint32 SizeX, uint32 SizeY, uint32 SizeZ)
 	: Surface(*SrcTexture->Surface.Device, ResourceType, SrcTexture->Surface.PixelFormat, SizeX, SizeY, SizeZ, SrcTexture->Surface.NumArrayLevels, SrcTexture->Surface.NumMips, SrcTexture->Surface.NumSamples, SrcTexture->Surface.Image, SrcTexture->Surface.UEFlags, FRHIResourceCreateInfo())
 	, PartialView(nullptr)
-	, bIsAliased(true)
+	, AliasedTexture(SrcTextureRHI)
 {
 	if (Surface.FullAspectMask == Surface.PartialAspectMask)
 	{
@@ -1796,7 +1793,8 @@ FVulkanTextureBase::FVulkanTextureBase(const FVulkanTextureBase* SrcTexture, VkI
 		PartialView = new FVulkanTextureView;
 		// Skip create, since we're aliasing.
 	}
-	AliasTextureResources(SrcTexture);
+
+	AliasTextureResources(SrcTextureRHI);
 }
 
 FVulkanTextureBase::~FVulkanTextureBase()
@@ -1809,10 +1807,12 @@ FVulkanTextureBase::~FVulkanTextureBase()
 	}
 }
 
-void FVulkanTextureBase::AliasTextureResources(const FVulkanTextureBase* SrcTexture)
+void FVulkanTextureBase::AliasTextureResources(FTextureRHIRef& SrcTextureRHI)
 {
 	DestroyViews();
 
+	FVulkanTextureBase* SrcTexture = (FVulkanTextureBase*)SrcTextureRHI->GetTextureBaseRHI();
+
 	Surface.Destroy();
 	Surface.Image = SrcTexture->Surface.Image;
 	DefaultView.View = SrcTexture->DefaultView.View;
@@ -1825,13 +1825,11 @@ void FVulkanTextureBase::AliasTextureResources(const FVulkanTextureBase* SrcText
 		PartialView->Image = SrcTexture->PartialView->Image;
 		PartialView->ViewId = SrcTexture->PartialView->ViewId;
 	}
-
-	bIsAliased = true;
 }
 
 void FVulkanTextureBase::DestroyViews()
 {
-	if (!bIsAliased)
+	if (AliasedTexture == nullptr)
 	{
 		DefaultView.Destroy(*Surface.Device);
 
@@ -1861,9 +1859,9 @@ FVulkanTexture2D::FVulkanTexture2D(FVulkanDevice& Device, EPixelFormat Format, u
 {
 }
 
-FVulkanTexture2D::FVulkanTexture2D(const FVulkanTexture2D* SrcTexture)
+FVulkanTexture2D::FVulkanTexture2D(FTextureRHIRef& SrcTextureRHI, const FVulkanTexture2D* SrcTexture)
 	: FRHITexture2D(SrcTexture->GetSizeX(), SrcTexture->GetSizeY(), SrcTexture->GetNumMips(), SrcTexture->GetNumSamples(), SrcTexture->GetFormat(), SrcTexture->GetFlags(), SrcTexture->GetClearBinding())
-	, FVulkanTextureBase(static_cast<const FVulkanTextureBase*>(SrcTexture), VK_IMAGE_VIEW_TYPE_2D, SrcTexture->GetSizeX(), SrcTexture->GetSizeY(), 1)
+	, FVulkanTextureBase(SrcTextureRHI, static_cast<const FVulkanTextureBase*>(SrcTexture), VK_IMAGE_VIEW_TYPE_2D, SrcTexture->GetSizeX(), SrcTexture->GetSizeY(), 1)
 {
 }
 
@@ -1887,9 +1885,9 @@ FVulkanTexture2DArray::FVulkanTexture2DArray(FVulkanDevice& Device, EPixelFormat
 {
 }
 
-FVulkanTexture2DArray::FVulkanTexture2DArray(const FVulkanTexture2DArray* SrcTexture)
+FVulkanTexture2DArray::FVulkanTexture2DArray(FTextureRHIRef& SrcTextureRHI, const FVulkanTexture2DArray* SrcTexture)
 	: FRHITexture2DArray(SrcTexture->GetSizeX(), SrcTexture->GetSizeY(), SrcTexture->GetSizeZ(), SrcTexture->GetNumMips(), SrcTexture->GetNumSamples(), SrcTexture->GetFormat(), SrcTexture->GetFlags(), SrcTexture->GetClearBinding())
-	, FVulkanTextureBase(static_cast<const FVulkanTextureBase*>(SrcTexture), VK_IMAGE_VIEW_TYPE_2D_ARRAY, SrcTexture->GetSizeX(), SrcTexture->GetSizeY(), SrcTexture->GetSizeZ())
+	, FVulkanTextureBase(SrcTextureRHI, static_cast<const FVulkanTextureBase*>(SrcTexture), VK_IMAGE_VIEW_TYPE_2D_ARRAY, SrcTexture->GetSizeX(), SrcTexture->GetSizeY(), SrcTexture->GetSizeZ())
 {
 }
 
@@ -1913,9 +1911,9 @@ FVulkanTextureCube::FVulkanTextureCube(FVulkanDevice& Device, EPixelFormat Forma
 {
 }
 
-FVulkanTextureCube::FVulkanTextureCube(const FVulkanTextureCube* SrcTexture)
+FVulkanTextureCube::FVulkanTextureCube(FTextureRHIRef& SrcTextureRHI, const FVulkanTextureCube* SrcTexture)
 	: FRHITextureCube(SrcTexture->GetSize(), SrcTexture->GetNumMips(), SrcTexture->GetFormat(), SrcTexture->GetFlags(), SrcTexture->GetClearBinding())
-	, FVulkanTextureBase(static_cast<const FVulkanTextureBase*>(SrcTexture), VK_IMAGE_VIEW_TYPE_CUBE, SrcTexture->GetSize(), SrcTexture->GetSize(), 1)
+	, FVulkanTextureBase(SrcTextureRHI, static_cast<const FVulkanTextureBase*>(SrcTexture), VK_IMAGE_VIEW_TYPE_CUBE, SrcTexture->GetSize(), SrcTexture->GetSize(), 1)
 {
 }
 
diff --git a/Engine/Source/Runtime/VulkanRHI/Public/VulkanDynamicRHI.h b/Engine/Source/Runtime/VulkanRHI/Public/VulkanDynamicRHI.h
--- a/Engine/Source/Runtime/VulkanRHI/Public/VulkanDynamicRHI.h
+++ b/Engine/Source/Runtime/VulkanRHI/Public/VulkanDynamicRHI.h
@@ -158,6 +158,8 @@ public:
 	virtual FTexture2DRHIRef RHIGetViewportBackBuffer(FRHIViewport* Viewport) final override;
 	virtual void RHIAliasTextureResources(FRHITexture* DestTexture, FRHITexture* SrcTexture) final override;
 	virtual FTextureRHIRef RHICreateAliasedTexture(FRHITexture* SourceTexture) final override;
+	virtual void RHIAliasTextureResources(FTextureRHIRef& DestTexture, FTextureRHIRef& SrcTexture) final override;
+	virtual FTextureRHIRef RHICreateAliasedTexture(FTextureRHIRef& SourceTexture) final override;
 	virtual void RHIAdvanceFrameForGetViewportBackBuffer(FRHIViewport* Viewport) final override;
 	virtual void RHIAcquireThreadOwnership() final override;
 	virtual void RHIReleaseThreadOwnership() final override;
diff --git a/Engine/Source/Runtime/Windows/D3D11RHI/Private/D3D11RHIPrivate.h b/Engine/Source/Runtime/Windows/D3D11RHI/Private/D3D11RHIPrivate.h
--- a/Engine/Source/Runtime/Windows/D3D11RHI/Private/D3D11RHIPrivate.h
+++ b/Engine/Source/Runtime/Windows/D3D11RHI/Private/D3D11RHIPrivate.h
@@ -478,6 +478,8 @@ public:
 	virtual FTexture2DRHIRef RHIGetViewportBackBuffer(FRHIViewport* Viewport) final override;
 	virtual void RHIAliasTextureResources(FRHITexture* DestTexture, FRHITexture* SrcTexture) final override;
 	virtual FTextureRHIRef RHICreateAliasedTexture(FRHITexture* SourceTexture) final override;
+	virtual void RHIAliasTextureResources(FTextureRHIRef& DestTexture, FTextureRHIRef& SrcTexture) final override;
+	virtual FTextureRHIRef RHICreateAliasedTexture(FTextureRHIRef& SourceTexture) final override;
 	virtual void RHIAdvanceFrameForGetViewportBackBuffer(FRHIViewport* Viewport) final override;
 	virtual void RHIAcquireThreadOwnership() final override;
 	virtual void RHIReleaseThreadOwnership() final override;
diff --git a/Engine/Source/Runtime/Windows/D3D11RHI/Private/D3D11Texture.cpp b/Engine/Source/Runtime/Windows/D3D11RHI/Private/D3D11Texture.cpp
--- a/Engine/Source/Runtime/Windows/D3D11RHI/Private/D3D11Texture.cpp
+++ b/Engine/Source/Runtime/Windows/D3D11RHI/Private/D3D11Texture.cpp
@@ -2483,6 +2483,18 @@ FTextureRHIRef FD3D11DynamicRHI::RHICreateAliasedTexture(FRHITexture* SourceText
 	return nullptr;
 }
 
+void FD3D11DynamicRHI::RHIAliasTextureResources(FTextureRHIRef& DestTextureRHI, FTextureRHIRef& SrcTextureRHI)
+{
+	// @todo: Move the raw-pointer implementation down here when it's deprecation is completed.
+	RHIAliasTextureResources((FRHITexture*)DestTextureRHI, (FRHITexture*)SrcTextureRHI);
+}
+
+FTextureRHIRef FD3D11DynamicRHI::RHICreateAliasedTexture(FTextureRHIRef& SourceTexture)
+{
+	// @todo: Move the raw-pointer implementation down here when it's deprecation is completed.
+	return RHICreateAliasedTexture((FRHITexture*)SourceTexture);
+}
+
 void FD3D11DynamicRHI::RHICopyTexture(FRHITexture* SourceTextureRHI, FRHITexture* DestTextureRHI, const FRHICopyTextureInfo& CopyInfo)
 {
 	if (!SourceTextureRHI || !DestTextureRHI || SourceTextureRHI == DestTextureRHI)
