DRPS - Compile fixes from merge
#rnx

diff --git a/Engine/Source/Runtime/Renderer/Private/PostProcess/SceneRenderTargets.cpp b/Engine/Source/Runtime/Renderer/Private/PostProcess/SceneRenderTargets.cpp
--- a/Engine/Source/Runtime/Renderer/Private/PostProcess/SceneRenderTargets.cpp
+++ b/Engine/Source/Runtime/Renderer/Private/PostProcess/SceneRenderTargets.cpp
@@ -869,13 +869,12 @@ FUnorderedAccessViewRHIRef FSceneRenderTargets::GetQuadOverdrawBufferUAV()
 
 FUnorderedAccessViewRHIRef FSceneRenderTargets::GetVirtualTextureFeedbackUAV()
 {
-	TRefCountPtr< IPooledRenderTarget > RenderTarget = VirtualTextureFeedback.FeedbackTextureGPU;
-	if(!RenderTarget)
+	if (!VirtualTextureFeedback.FeedbackBufferUAV)
 	{
 		return GBlackTextureWithUAV->UnorderedAccessViewRHI;
 	}
 	
-	return RenderTarget->GetRenderTargetItem().UAV;
+	return VirtualTextureFeedback.FeedbackBufferUAV;
 }
 
 void FSceneRenderTargets::BindVirtualTextureFeedbackUAV(FRHIRenderPassInfo& RPInfo)
diff --git a/Engine/Source/Runtime/Renderer/Private/VT/VirtualTextureFeedback.cpp b/Engine/Source/Runtime/Renderer/Private/VT/VirtualTextureFeedback.cpp
--- a/Engine/Source/Runtime/Renderer/Private/VT/VirtualTextureFeedback.cpp
+++ b/Engine/Source/Runtime/Renderer/Private/VT/VirtualTextureFeedback.cpp
@@ -184,10 +184,6 @@ void FVirtualTextureFeedback::CreateResourceGPU( FRHICommandListImmediate& RHICm
 		FeedbackBufferUAV = RHICreateUnorderedAccessView(FeedbackBuffer, /*Format=*/ PF_R32_UINT);
 	}
 	
-	// Clear to default value
-	RHICmdList.ClearUAVUint(FeedbackTextureGPU->GetRenderTargetItem().UAV, FUintVector4(~0u, ~0u, ~0u, ~0u));
-	RHICmdList.TransitionResource( EResourceTransitionAccess::ERWNoBarrier, EResourceTransitionPipeline::EGfxToGfx, FeedbackTextureGPU->GetRenderTargetItem().UAV );
-
 	FeedBackFences->Init(RHICmdList);
 }
 
diff --git a/Engine/Source/Runtime/VulkanRHI/Private/VulkanCommands.cpp b/Engine/Source/Runtime/VulkanRHI/Private/VulkanCommands.cpp
--- a/Engine/Source/Runtime/VulkanRHI/Private/VulkanCommands.cpp
+++ b/Engine/Source/Runtime/VulkanRHI/Private/VulkanCommands.cpp
@@ -425,38 +425,6 @@ inline void /*FVulkanCommandListContext::*/SetShaderUniformBufferResources(FVulk
 		const FVulkanShaderHeader::FUBResourceInfo& ResourceInfo = HeaderUBInfo.ResourceEntries[Index];
 		switch (ResourceInfo.UBBaseType)
 		{
-		case UBMT_UAV:
-		{
-			const VkDescriptorType DescriptorType = DescriptorTypes[GlobalInfos[ResourceInfo.GlobalIndex].TypeIndex];
-			
-			FVulkanUnorderedAccessView* UAV = (FVulkanUnorderedAccessView*)(ResourceArray[ResourceInfo.SourceUBResourceIndex].GetReference());
-			if (UAV)
-			{
-				ensure(DescriptorType == VK_DESCRIPTOR_TYPE_STORAGE_IMAGE);
-				FRHITexture* SourceTexture = UAV->SourceTexture.GetReference();
-				if(SourceTexture)
-				{
-					const FVulkanTextureBase* BaseTexture = FVulkanTextureBase::Cast(SourceTexture);
-					VkImageLayout Layout = Context->GetLayoutForDescriptor(BaseTexture->Surface);
-					State->SetUAVForUBResource(GlobalRemappingInfo[ResourceInfo.GlobalIndex].NewDescriptorSet, GlobalRemappingInfo[ResourceInfo.GlobalIndex].NewBindingIndex, UAV);
-					SourceTexture->SetLastRenderTime(CurrentTime);
-				}
-				else
-				{
-					checkNoEntry();
-				}
-			}
-			else
-			{
-#if VULKAN_ENABLE_SHADER_DEBUG_NAMES
-				UE_LOG(LogVulkanRHI, Warning, TEXT("Invalid texture in SRT table for shader '%s'"), *Shader->GetDebugName());
-#else
-				UE_LOG(LogVulkanRHI, Warning, TEXT("Invalid texture in SRT table"));
-#endif
-				checkNoEntry();
-			}
-		}
-		break;
 		case UBMT_SAMPLER:
 		{
 			uint16 CombinedAlias = GlobalInfos[ResourceInfo.GlobalIndex].CombinedSamplerStateAliasIndex;
@@ -529,11 +497,11 @@ inline void /*FVulkanCommandListContext::*/SetShaderUniformBufferResources(FVulk
 		case UBMT_UAV:
 		{
 			const VkDescriptorType DescriptorType = DescriptorTypes[GlobalInfos[ResourceInfo.GlobalIndex].TypeIndex];
-			ensure(DescriptorType == VK_DESCRIPTOR_TYPE_STORAGE_BUFFER 
+			ensure(DescriptorType == VK_DESCRIPTOR_TYPE_STORAGE_BUFFER
 				|| DescriptorType == VK_DESCRIPTOR_TYPE_STORAGE_IMAGE
 				|| DescriptorType == VK_DESCRIPTOR_TYPE_STORAGE_TEXEL_BUFFER);
 			FRHIUnorderedAccessView* CurrentUAV = (FRHIUnorderedAccessView*)(ResourceArray[ResourceInfo.SourceUBResourceIndex].GetReference());
-			if (CurrentUAV)
+			if (UAV)
 			{
 				FVulkanUnorderedAccessView* UAV = ResourceCast(CurrentUAV);
 				State->SetUAVForUBResource(GlobalRemappingInfo[ResourceInfo.GlobalIndex].NewDescriptorSet, GlobalRemappingInfo[ResourceInfo.GlobalIndex].NewBindingIndex, UAV);
