Add a ref to compute shaders when creating the pipeline state

diff --git a/Engine/Source/Runtime/VulkanRHI/Private/VulkanPipeline.cpp b/Engine/Source/Runtime/VulkanRHI/Private/VulkanPipeline.cpp
--- a/Engine/Source/Runtime/VulkanRHI/Private/VulkanPipeline.cpp
+++ b/Engine/Source/Runtime/VulkanRHI/Private/VulkanPipeline.cpp
@@ -203,6 +203,11 @@ FVulkanComputePipeline::FVulkanComputePipeline(FVulkanDevice* InDevice)
 
 FVulkanComputePipeline::~FVulkanComputePipeline()
 {
+	if (ComputeShader)
+	{
+		ComputeShader->Release();
+	}
+
 	Device->NotifyDeletedComputePipeline(this);
 	DEC_DWORD_STAT(STAT_VulkanNumComputePSOs);
 }
@@ -1871,7 +1876,6 @@ FVulkanComputePipeline* FVulkanPipelineStateCacheManager::GetOrCreateComputePipe
 	double BeginTime = FPlatformTime::Seconds();
 
 	FVulkanComputePipeline* ComputePipeline = CreateComputePipelineFromShader(ComputeShader);
-	ComputePipeline->ComputeShader = ComputeShader;
 
 	double EndTime = FPlatformTime::Seconds();
 	double Delta = EndTime - BeginTime;
@@ -1895,6 +1899,9 @@ FVulkanComputePipeline* FVulkanPipelineStateCacheManager::CreateComputePipelineF
 {
 	FVulkanComputePipeline* Pipeline = new FVulkanComputePipeline(Device);
 
+	Pipeline->ComputeShader = Shader;
+	Pipeline->ComputeShader->AddRef();
+
 	FVulkanDescriptorSetsLayoutInfo DescriptorSetLayoutInfo;
 	const FVulkanShaderHeader& CSHeader = Shader->GetCodeHeader();
 	FUniformBufferGatherInfo UBGatherInfo;
