fix an issue with RHIReadSurfaceData not allocating a large enough buffer.
Fixes an issue with texture painting.

diff --git a/Engine/Source/Runtime/VulkanRHI/Private/VulkanRenderTarget.cpp b/Engine/Source/Runtime/VulkanRHI/Private/VulkanRenderTarget.cpp
--- a/Engine/Source/Runtime/VulkanRHI/Private/VulkanRenderTarget.cpp
+++ b/Engine/Source/Runtime/VulkanRHI/Private/VulkanRenderTarget.cpp
@@ -790,7 +790,19 @@ void FVulkanDynamicRHI::RHIReadSurfaceData(FRHITexture* TextureRHI, FIntRect Rec
 	FVulkanCmdBuffer* CmdBuffer = ImmediateContext.GetCommandBufferManager()->GetUploadCmdBuffer();
 
 	ensure(Texture2D->Surface.StorageFormat == VK_FORMAT_R8G8B8A8_UNORM || Texture2D->Surface.StorageFormat == VK_FORMAT_B8G8R8A8_UNORM || Texture2D->Surface.StorageFormat == VK_FORMAT_R16G16B16A16_SFLOAT || Texture2D->Surface.StorageFormat == VK_FORMAT_A2B10G10R10_UNORM_PACK32 || Texture2D->Surface.StorageFormat == VK_FORMAT_R16G16B16A16_UNORM);
-	const bool bIs8Bpp = (Texture2D->Surface.StorageFormat == VK_FORMAT_R16G16B16A16_SFLOAT);
+	bool bIs8Bpp = true;
+	switch(Texture2D->Surface.StorageFormat)
+	{
+	case VK_FORMAT_R16G16B16A16_SFLOAT:
+	case VK_FORMAT_R16G16B16A16_SNORM:
+	case VK_FORMAT_R16G16B16A16_UINT:
+	case VK_FORMAT_R16G16B16A16_SINT:
+		bIs8Bpp = false;
+		break;
+	default:
+		break;
+	}
+
 	const uint32 Size = NumPixels * sizeof(FColor) * (bIs8Bpp ? 2 : 1);
 	VulkanRHI::FStagingBuffer* StagingBuffer = Device->GetStagingManager().AcquireBuffer(Size, VK_BUFFER_USAGE_TRANSFER_DST_BIT, VK_MEMORY_PROPERTY_HOST_CACHED_BIT);
 	if (GIgnoreCPUReads == 0)
