Use vulkan device's minTexelBufferOffsetAlignment for buffer backed SRV alignment.
fixes random polys+ crash on android devices. (which require 64 byte alignment)
#rnx

diff --git a/Engine/Source/Runtime/VulkanRHI/Private/VulkanRHI.cpp b/Engine/Source/Runtime/VulkanRHI/Private/VulkanRHI.cpp
--- a/Engine/Source/Runtime/VulkanRHI/Private/VulkanRHI.cpp
+++ b/Engine/Source/Runtime/VulkanRHI/Private/VulkanRHI.cpp
@@ -1021,6 +1021,12 @@ IRHIComputeContext* FVulkanDynamicRHI::RHIGetDefaultAsyncComputeContext()
 	return &Device->GetImmediateComputeContext();
 }
 
+uint64 FVulkanDynamicRHI::RHIGetMinimumAlignmentForBufferBackedSRV(EPixelFormat Format)
+{
+	const VkPhysicalDeviceLimits& Limits = Device->GetLimits();
+	return Limits.minTexelBufferOffsetAlignment;
+}
+
 IRHICommandContextContainer* FVulkanDynamicRHI::RHIGetCommandContextContainer(int32 Index, int32 Num)
 {
 	if (GRHIThreadCvar.GetValueOnAnyThread() > 1)
@@ -1472,6 +1478,7 @@ void FVulkanBufferView::Create(VkFormat Format, FVulkanResourceMultiBuffer* Buff
 
 	Flags = Buffer->GetBufferUsageFlags() & (VK_BUFFER_USAGE_UNIFORM_TEXEL_BUFFER_BIT | VK_BUFFER_USAGE_STORAGE_TEXEL_BUFFER_BIT);
 	check(Flags);
+	check(IsAligned(Offset, Limits.minTexelBufferOffsetAlignment));
 
 	VERIFYVULKANRESULT(VulkanRHI::vkCreateBufferView(GetParent()->GetInstanceHandle(), &ViewInfo, VULKAN_CPU_ALLOCATOR, &View));
 	
diff --git a/Engine/Source/Runtime/VulkanRHI/Public/VulkanDynamicRHI.h b/Engine/Source/Runtime/VulkanRHI/Public/VulkanDynamicRHI.h
--- a/Engine/Source/Runtime/VulkanRHI/Public/VulkanDynamicRHI.h
+++ b/Engine/Source/Runtime/VulkanRHI/Public/VulkanDynamicRHI.h
@@ -185,6 +185,7 @@ public:
 	virtual class IRHICommandContext* RHIGetDefaultContext() final override;
 	virtual class IRHIComputeContext* RHIGetDefaultAsyncComputeContext() final override;
 	virtual class IRHICommandContextContainer* RHIGetCommandContextContainer(int32 Index, int32 Num) final override;
+	virtual uint64 RHIGetMinimumAlignmentForBufferBackedSRV(EPixelFormat Format) final override;
 
 	///////// Pass through functions that allow RHIs to optimize certain calls.
 	//virtual FVertexBufferRHIRef CreateAndLockVertexBuffer_RenderThread(class FRHICommandListImmediate& RHICmdList, uint32 Size, uint32 InUsage, FRHIResourceCreateInfo& CreateInfo, void*& OutDataBuffer) override final;
