- Fix vulkan RHIReadSurfaceData when reading back mip levels other than 0
- Add missing transition back to Readable in LandscapeEditLayers
- Add VULKAN_EXTERN_EXPORT for VULKAN_ENABLE_IMAGE_TRACKING_LAYER

diff --git a/Engine/Source/Runtime/Landscape/Private/LandscapeEditLayers.cpp b/Engine/Source/Runtime/Landscape/Private/LandscapeEditLayers.cpp
--- a/Engine/Source/Runtime/Landscape/Private/LandscapeEditLayers.cpp
+++ b/Engine/Source/Runtime/Landscape/Private/LandscapeEditLayers.cpp
@@ -1249,6 +1249,7 @@ public:
 		InRHICmdList.DrawPrimitive(0, PrimitiveCount, 1);
 
 		InRHICmdList.EndRenderPass();
+		InRHICmdList.TransitionResource(EResourceTransitionAccess::EReadable, ViewFamily.RenderTarget->GetRenderTargetTexture());
 
 		VertexDeclaration.ReleaseResource();
 		VertexBufferResource.ReleaseResource();
diff --git a/Engine/Source/Runtime/VulkanRHI/Private/VulkanCommandWrappers.h b/Engine/Source/Runtime/VulkanRHI/Private/VulkanCommandWrappers.h
--- a/Engine/Source/Runtime/VulkanRHI/Private/VulkanCommandWrappers.h
+++ b/Engine/Source/Runtime/VulkanRHI/Private/VulkanCommandWrappers.h
@@ -139,7 +139,7 @@ struct FWrapLayer
 	static void CmdUpdateBuffer(VkResult Result, VkCommandBuffer CommandBuffer, VkBuffer DstBuffer, VkDeviceSize DstOffset, VkDeviceSize DataSize, const void* pData) VULKAN_LAYER_BODY
 	static void CmdFillBuffer(VkResult Result, VkCommandBuffer CommandBuffer, VkBuffer DstBuffer, VkDeviceSize DstOffset, VkDeviceSize Size, uint32 Data) VULKAN_LAYER_BODY
 	VULKAN_EXTERN_EXPORT static void CmdClearColorImage(VkResult Result, VkCommandBuffer CommandBuffer, VkImage Image, VkImageLayout ImageLayout, const VkClearColorValue* ColorValue, uint32 RangeCount, const VkImageSubresourceRange* Ranges) VULKAN_LAYER_BODY
-	static void CmdClearDepthStencilImage(VkResult Result, VkCommandBuffer CommandBuffer, VkImage Image, VkImageLayout ImageLayout, const VkClearDepthStencilValue* DepthStencil, uint32 RangeCount, const VkImageSubresourceRange* Ranges) VULKAN_LAYER_BODY
+	VULKAN_EXTERN_EXPORT static void CmdClearDepthStencilImage(VkResult Result, VkCommandBuffer CommandBuffer, VkImage Image, VkImageLayout ImageLayout, const VkClearDepthStencilValue* DepthStencil, uint32 RangeCount, const VkImageSubresourceRange* Ranges) VULKAN_LAYER_BODY
 	static void CmdClearAttachments(VkResult Result, VkCommandBuffer CommandBuffer, uint32 AttachmentCount, const VkClearAttachment* Attachments, uint32 RectCount, const VkClearRect* Rects) VULKAN_LAYER_BODY
 	static void CmdResolveImage(VkResult Result, VkCommandBuffer CommandBuffer, VkImage SrcImage, VkImageLayout SrcImageLayout, VkImage DstImage, VkImageLayout DstImageLayout, uint32 RegionCount, const VkImageResolve* Regions) VULKAN_LAYER_BODY
 	static void CmdSetEvent(VkResult Result, VkCommandBuffer CommandBuffer, VkEvent Event, VkPipelineStageFlags StageMask) VULKAN_LAYER_BODY
diff --git a/Engine/Source/Runtime/VulkanRHI/Private/VulkanRHIPrivate.h b/Engine/Source/Runtime/VulkanRHI/Private/VulkanRHIPrivate.h
--- a/Engine/Source/Runtime/VulkanRHI/Private/VulkanRHIPrivate.h
+++ b/Engine/Source/Runtime/VulkanRHI/Private/VulkanRHIPrivate.h
@@ -383,6 +383,14 @@ inline void VulkanSetImageLayoutSimple(VkCommandBuffer CmdBuffer, VkImage Image,
 	VulkanSetImageLayout(CmdBuffer, Image, OldLayout, NewLayout, SubresourceRange);
 }
 
+// Transitions all mips of Color Image
+inline void VulkanSetImageLayoutAllMips(VkCommandBuffer CmdBuffer, VkImage Image, VkImageLayout OldLayout, VkImageLayout NewLayout, VkImageAspectFlags Aspect = VK_IMAGE_ASPECT_COLOR_BIT)
+{
+	VkImageSubresourceRange SubresourceRange = { Aspect, 0, VK_REMAINING_MIP_LEVELS , 0, 1 };
+	VulkanSetImageLayout(CmdBuffer, Image, OldLayout, NewLayout, SubresourceRange);
+}
+
+
 void VulkanResolveImage(VkCommandBuffer Cmd, FRHITexture* SourceTextureRHI, FRHITexture* DestTextureRHI);
 
 DECLARE_STATS_GROUP(TEXT("Vulkan PSO"), STATGROUP_VulkanPSO, STATCAT_Advanced);
diff --git a/Engine/Source/Runtime/VulkanRHI/Private/VulkanRenderTarget.cpp b/Engine/Source/Runtime/VulkanRHI/Private/VulkanRenderTarget.cpp
--- a/Engine/Source/Runtime/VulkanRHI/Private/VulkanRenderTarget.cpp
+++ b/Engine/Source/Runtime/VulkanRHI/Private/VulkanRenderTarget.cpp
@@ -809,15 +809,18 @@ void FVulkanDynamicRHI::RHIReadSurfaceData(FRHITexture* TextureRHI, FIntRect Rec
 	{
 		VkBufferImageCopy CopyRegion;
 		FMemory::Memzero(CopyRegion);
+		uint32 MipLevel = InFlags.GetMip();
+		uint32 SizeX = FMath::Max(TextureRHI2D->GetSizeX() >> MipLevel, 1u);
+		uint32 SizeY = FMath::Max(TextureRHI2D->GetSizeY() >> MipLevel, 1u);
 		//CopyRegion.bufferOffset = 0;
-		CopyRegion.bufferRowLength = TextureRHI2D->GetSizeX();
-		CopyRegion.bufferImageHeight = TextureRHI2D->GetSizeY();
+		CopyRegion.bufferRowLength = SizeX;
+		CopyRegion.bufferImageHeight = SizeY;
 		CopyRegion.imageSubresource.aspectMask = Texture2D->Surface.GetFullAspectMask();
-		CopyRegion.imageSubresource.mipLevel = InFlags.GetMip();
+		CopyRegion.imageSubresource.mipLevel = MipLevel;
 		//CopyRegion.imageSubresource.baseArrayLayer = 0;
 		CopyRegion.imageSubresource.layerCount = 1;
-		CopyRegion.imageExtent.width = TextureRHI2D->GetSizeX();
-		CopyRegion.imageExtent.height = TextureRHI2D->GetSizeY();
+		CopyRegion.imageExtent.width = SizeX;
+		CopyRegion.imageExtent.height = SizeY;
 		CopyRegion.imageExtent.depth = 1;
 
 		//#todo-rco: Multithreaded!
@@ -825,13 +828,13 @@ void FVulkanDynamicRHI::RHIReadSurfaceData(FRHITexture* TextureRHI, FIntRect Rec
 		bool bHadLayout = (CurrentLayout != VK_IMAGE_LAYOUT_UNDEFINED);
 		if (CurrentLayout != VK_IMAGE_LAYOUT_TRANSFER_SRC_OPTIMAL)
 		{
-			VulkanSetImageLayoutSimple(CmdBuffer->GetHandle(), Texture2D->Surface.Image, CurrentLayout, VK_IMAGE_LAYOUT_TRANSFER_SRC_OPTIMAL);
+			VulkanSetImageLayoutAllMips(CmdBuffer->GetHandle(), Texture2D->Surface.Image, CurrentLayout, VK_IMAGE_LAYOUT_TRANSFER_SRC_OPTIMAL);
 		}
 
 		VulkanRHI::vkCmdCopyImageToBuffer(CmdBuffer->GetHandle(), Texture2D->Surface.Image, VK_IMAGE_LAYOUT_TRANSFER_SRC_OPTIMAL, StagingBuffer->GetHandle(), 1, &CopyRegion);
 		if (bHadLayout && CurrentLayout != VK_IMAGE_LAYOUT_TRANSFER_SRC_OPTIMAL)
 		{
-			VulkanSetImageLayoutSimple(CmdBuffer->GetHandle(), Texture2D->Surface.Image, VK_IMAGE_LAYOUT_TRANSFER_SRC_OPTIMAL, CurrentLayout);
+			VulkanSetImageLayoutAllMips(CmdBuffer->GetHandle(), Texture2D->Surface.Image, VK_IMAGE_LAYOUT_TRANSFER_SRC_OPTIMAL, CurrentLayout);
 		}
 		else
 		{
