Fix for Android Vulkan where UAV to texture transitions do not use ComputeToGFX which can cause corruption on some GPU's due to missing barrier

diff --git a/Engine/Source/Runtime/RHI/Public/RHICommandList.h b/Engine/Source/Runtime/RHI/Public/RHICommandList.h
--- a/Engine/Source/Runtime/RHI/Public/RHICommandList.h
+++ b/Engine/Source/Runtime/RHI/Public/RHICommandList.h
@@ -1555,6 +1555,37 @@ FRHICOMMAND_MACRO(FRHICommandTransitionTexturesArray)
 	RHI_API void Execute(FRHICommandListBase& CmdList);
 };
 
+
+FRHICOMMAND_MACRO(FRHICommandTransitionTexturesPipeline)
+{
+	int32 NumTextures;
+	FRHITexture** Textures; // Pointer to an array of textures, allocated inline with the command list
+	EResourceTransitionAccess TransitionType;
+	EResourceTransitionPipeline TransitionPipeline;
+	FORCEINLINE_DEBUGGABLE FRHICommandTransitionTexturesPipeline(EResourceTransitionAccess InTransitionType, EResourceTransitionPipeline InTransitionPipeline, FRHITexture** InTextures, int32 InNumTextures)
+		: NumTextures(InNumTextures)
+		, Textures(InTextures)
+		, TransitionType(InTransitionType)
+		, TransitionPipeline(InTransitionPipeline)
+	{
+	}
+	RHI_API void Execute(FRHICommandListBase& CmdList);
+};
+
+FRHICOMMAND_MACRO(FRHICommandTransitionTexturesArrayPipeline)
+{
+	TArray<FRHITexture*>& Textures;
+	EResourceTransitionAccess TransitionType;
+	EResourceTransitionPipeline TransitionPipeline;
+	FORCEINLINE_DEBUGGABLE FRHICommandTransitionTexturesArrayPipeline(EResourceTransitionAccess InTransitionType, EResourceTransitionPipeline InTransitionPipeline, TArray<FRHITexture*>& InTextures)
+		: Textures(InTextures)
+		, TransitionType(InTransitionType)
+		, TransitionPipeline(InTransitionPipeline)
+	{
+	}
+	RHI_API void Execute(FRHICommandListBase& CmdList);
+};
+
 struct FRHICommandTransitionUAVsString
 {
 	static const TCHAR* TStr() { return TEXT("FRHICommandTransitionUAVs"); }
@@ -3386,6 +3417,40 @@ public:
 		ALLOC_COMMAND(FRHICommandTransitionTextures)(TransitionType, InlineTextureArray, NumTextures);
 	}
 
+	FORCEINLINE_DEBUGGABLE void TransitionResource(EResourceTransitionAccess TransitionType, EResourceTransitionPipeline TransitionPipeline, FRHITexture* InTexture)
+	{
+		FRHITexture* Texture = InTexture;
+		check(Texture == nullptr || Texture->IsCommitted());
+		if (Bypass())
+		{
+			GetContext().RHITransitionResources(TransitionType, TransitionPipeline, &Texture, 1);
+			return;
+		}
+
+		// Allocate space to hold the single texture pointer inline in the command list itself.
+		FRHITexture** TextureArray = (FRHITexture**)Alloc(sizeof(FRHITexture*), alignof(FRHITexture*));
+		TextureArray[0] = Texture;
+		ALLOC_COMMAND(FRHICommandTransitionTexturesPipeline)(TransitionType, TransitionPipeline, TextureArray, 1);
+	}
+
+	FORCEINLINE_DEBUGGABLE void TransitionResources(EResourceTransitionAccess TransitionType, EResourceTransitionPipeline TransitionPipeline, FRHITexture** InTextures, int32 NumTextures)
+	{
+		if (Bypass())
+		{
+			GetContext().RHITransitionResources(TransitionType, TransitionPipeline, InTextures, NumTextures);
+			return;
+		}
+
+		// Allocate space to hold the list of textures inline in the command list itself.
+		FRHITexture** InlineTextureArray = (FRHITexture**)Alloc(sizeof(FRHITexture*) * NumTextures, alignof(FRHITexture*));
+		for (int32 Index = 0; Index < NumTextures; ++Index)
+		{
+			InlineTextureArray[Index] = InTextures[Index];
+		}
+
+		ALLOC_COMMAND(FRHICommandTransitionTexturesPipeline)(TransitionType, TransitionPipeline,  InlineTextureArray, NumTextures);
+	}	
+
 	FORCEINLINE_DEBUGGABLE void TransitionResourceArrayNoCopy(EResourceTransitionAccess TransitionType, TArray<FRHITexture*>& InTextures)
 	{
 		if (Bypass())
diff --git a/Engine/Source/Runtime/RHI/Public/RHICommandListCommandExecutes.inl b/Engine/Source/Runtime/RHI/Public/RHICommandListCommandExecutes.inl
--- a/Engine/Source/Runtime/RHI/Public/RHICommandListCommandExecutes.inl
+++ b/Engine/Source/Runtime/RHI/Public/RHICommandListCommandExecutes.inl
@@ -54,6 +54,7 @@ struct FRHICommandSetStereoViewport;
 struct FRHICommandSetStreamSource;
 struct FRHICommandSetViewport;
 struct FRHICommandTransitionTextures;
+struct FRHICommandTransitionTexturesPipeline;
 struct FRHICommandTransitionTexturesDepth;
 struct FRHICommandTransitionTexturesArray;
 struct FRHICommandUpdateTextureReference;
@@ -401,6 +402,12 @@ void FRHICommandTransitionTextures::Execute(FRHICommandListBase& CmdList)
 	INTERNAL_DECORATOR(RHITransitionResources)(TransitionType, &Textures[0], NumTextures);
 }
 
+void FRHICommandTransitionTexturesPipeline::Execute(FRHICommandListBase& CmdList)
+{
+	RHISTAT(TransitionTextures);
+	INTERNAL_DECORATOR(RHITransitionResources)(TransitionType, TransitionPipeline, &Textures[0], NumTextures);
+}
+
 void FRHICommandTransitionTexturesDepth::Execute(FRHICommandListBase& CmdList)
 {
 	RHISTAT(TransitionTextures);
diff --git a/Engine/Source/Runtime/RHI/Public/RHIContext.h b/Engine/Source/Runtime/RHI/Public/RHIContext.h
--- a/Engine/Source/Runtime/RHI/Public/RHIContext.h
+++ b/Engine/Source/Runtime/RHI/Public/RHIContext.h
@@ -349,6 +349,7 @@ public:
 	* Explicitly transition a texture resource from readable -> writable by the GPU or vice versa.
 	* We know render targets are only used as rendered targets on the Gfx pipeline, so these transitions are assumed to be implemented such
 	* Gfx->Gfx and Gfx->Compute pipeline transitions are both handled by this call by the RHI implementation.  Hence, no pipeline parameter on this call.
+	* TransitionPipeline version is also available to transition from Compute->Gfx for example
 	*
 	* @param TransitionType - direction of the transition
 	* @param InTextures - array of texture objects to transition
@@ -366,6 +367,21 @@ public:
 		}
 	}
 
+	/**
+	* Explicitly transition a texture resource from readable -> writable by the GPU or vice versa.
+	* Also explicitly states which pipeline the texture can be used on next.  For example, if a Compute job just wrote this UAV for a Pixel shader to read
+	* you would do EResourceTransitionAccess::Readable and EResourceTransitionPipeline::EComputeToGfx
+	*
+	* @param TransitionType - direction of the transition
+	* @param TransitionPipeline - How this Texture is transitioning between Gfx and Compute, if at all.
+	* @param InTextures - array of texture objects to transition
+	* @param NumTextures - number of textures to transition
+	*/
+	virtual void RHITransitionResources(EResourceTransitionAccess TransitionType, EResourceTransitionPipeline TransitionPipeline, FRHITexture** InTextures, int32 NumTextures)
+	{
+		RHITransitionResources(TransitionType, InTextures, NumTextures);
+	}
+
 	/**
 	* Explicitly transition a UAV from readable -> writable by the GPU or vice versa.
 	* Also explicitly states which pipeline the UAV can be used on next.  For example, if a Compute job just wrote this UAV for a Pixel shader to read
diff --git a/Engine/Source/Runtime/RenderCore/Private/RenderGraphBarrierBatcher.cpp b/Engine/Source/Runtime/RenderCore/Private/RenderGraphBarrierBatcher.cpp
--- a/Engine/Source/Runtime/RenderCore/Private/RenderGraphBarrierBatcher.cpp
+++ b/Engine/Source/Runtime/RenderCore/Private/RenderGraphBarrierBatcher.cpp
@@ -143,7 +143,7 @@ FRDGBarrierBatcher::~FRDGBarrierBatcher()
 	{
 		FTransitionParameters TransitionParameters = Element.Key;
 		FTextureBatch& Batch = Element.Value;
-		RHICmdList.TransitionResources(TransitionParameters.TransitionAccess, Batch.GetData(), Batch.Num());
+		RHICmdList.TransitionResources(TransitionParameters.TransitionAccess, TransitionParameters.TransitionPipeline, Batch.GetData(), Batch.Num());
 	}
 
 	for (auto& Element : UAVBatchMap)
@@ -202,7 +202,7 @@ void FRDGBarrierBatcher::QueueTransitionTexture(FRDGTexture* Texture, FRDGResour
 		{
 			FTransitionParameters TransitionParameters;
 			TransitionParameters.TransitionAccess = GetResourceTransitionAccess(StateAfter.Access);
-			TransitionParameters.TransitionPipeline = EResourceTransitionPipeline::EGfxToGfx; // NOTE: Transition API for textures doesn't currently expose pipeline transitions.
+			TransitionParameters.TransitionPipeline = GetResourceTransitionPipeline(StateBefore.Pipeline, StateAfter.Pipeline);
 
 			FTextureBatch& TextureBatch = TextureBatchMap.FindOrAdd(TransitionParameters);
 			TextureBatch.Reserve(kBatchReservationSize);
diff --git a/Engine/Source/Runtime/VulkanRHI/Private/VulkanContext.h b/Engine/Source/Runtime/VulkanRHI/Private/VulkanContext.h
--- a/Engine/Source/Runtime/VulkanRHI/Private/VulkanContext.h
+++ b/Engine/Source/Runtime/VulkanRHI/Private/VulkanContext.h
@@ -195,6 +195,7 @@ public:
 	virtual void RHICopyToResolveTarget(FRHITexture* SourceTexture, FRHITexture* DestTexture, const FResolveParams& ResolveParams) final override;
 	virtual void RHICopyTexture(FRHITexture* SourceTexture, FRHITexture* DestTexture, const FRHICopyTextureInfo& CopyInfo) final override;
 	virtual void RHITransitionResources(EResourceTransitionAccess TransitionType, FRHITexture** InRenderTargets, int32 NumTextures) final override;
+	virtual void RHITransitionResources(EResourceTransitionAccess TransitionType, EResourceTransitionPipeline TransitionPipeline, FRHITexture** InRenderTargets, int32 NumTextures) final override;
 	virtual void RHITransitionResources(EResourceTransitionAccess TransitionType, EResourceTransitionPipeline TransitionPipeline, FRHIUnorderedAccessView** InUAVs, int32 NumUAVs, FRHIComputeFence* WriteComputeFence) final override;
 	virtual void RHITransitionResources(FExclusiveDepthStencil DepthStencilMode, FRHITexture* DepthTexture) final override;
 	virtual void RHICopyToStagingBuffer(FRHIVertexBuffer* SourceBuffer, FRHIStagingBuffer* DestinationStagingBuffer, uint32 Offset, uint32 NumBytes) final override;
diff --git a/Engine/Source/Runtime/VulkanRHI/Private/VulkanRenderTarget.cpp b/Engine/Source/Runtime/VulkanRHI/Private/VulkanRenderTarget.cpp
--- a/Engine/Source/Runtime/VulkanRHI/Private/VulkanRenderTarget.cpp
+++ b/Engine/Source/Runtime/VulkanRHI/Private/VulkanRenderTarget.cpp
@@ -1225,6 +1225,33 @@ void FVulkanCommandListContext::RHITransitionResources(EResourceTransitionAccess
 		if (PendingTransition.Textures.Num() > 0)
 		{
 			PendingTransition.TransitionType = TransitionType;
+			PendingTransition.TransitionPipeline = EResourceTransitionPipeline::EGfxToGfx; // Default to GfxToGfx which is ignored for textures
+			TransitionResources(PendingTransition);
+		}
+	}
+}
+
+void FVulkanCommandListContext::RHITransitionResources(EResourceTransitionAccess TransitionType, EResourceTransitionPipeline TransitionPipeline, FRHITexture** InTextures, int32 NumTextures)
+{
+	if (NumTextures > 0)
+	{
+		FPendingTransition PendingTransition;
+		for (int32 Index = 0; Index < NumTextures; ++Index)
+		{
+			FRHITexture* RHITexture = InTextures[Index];
+			if (RHITexture)
+			{
+				PendingTransition.Textures.Add(RHITexture);
+
+				FVulkanTextureBase* VulkanTexture = FVulkanTextureBase::Cast(RHITexture);
+				VulkanTexture->OnTransitionResource(*this, TransitionType);
+			}
+		}
+
+		if (PendingTransition.Textures.Num() > 0)
+		{
+			PendingTransition.TransitionType = TransitionType;
+			PendingTransition.TransitionPipeline = TransitionPipeline;
 			TransitionResources(PendingTransition);
 		}
 	}
@@ -1449,6 +1476,27 @@ void FVulkanCommandListContext::TransitionResources(const FPendingTransition& Pe
 			VulkanRHI::FPendingBarrier Barrier;
 			for (int32 Index = 0; Index < PendingTransition.Textures.Num(); ++Index)
 			{
+				// If we are transitioning from compute we need additional pipeline stages
+				// We can ignore the other transition types as their barriers are more explicitly handled elsewhere
+				VkPipelineStageFlags SourceStage = 0, DestStage = 0;
+				switch (PendingTransition.TransitionPipeline)
+				{
+				case EResourceTransitionPipeline::EGfxToGfx:
+				case EResourceTransitionPipeline::EGfxToCompute:
+					break;
+				case EResourceTransitionPipeline::EComputeToGfx:
+					SourceStage = VK_PIPELINE_STAGE_COMPUTE_SHADER_BIT;
+					DestStage = VK_PIPELINE_STAGE_ALL_GRAPHICS_BIT;
+					break;
+				case EResourceTransitionPipeline::EComputeToCompute:
+					SourceStage = VK_PIPELINE_STAGE_COMPUTE_SHADER_BIT;
+					DestStage = VK_PIPELINE_STAGE_COMPUTE_SHADER_BIT | VK_PIPELINE_STAGE_ALL_GRAPHICS_BIT;
+					break;
+				default:
+					ensureMsgf(0, TEXT("Unknown transition pipeline %d"), (int32)PendingTransition.TransitionPipeline);
+					break;
+				}
+
 				FVulkanTextureBase* VulkanTexture = FVulkanTextureBase::Cast(PendingTransition.Textures[Index]);
 				VkImageLayout& SrcLayout = TransitionAndLayoutManager.FindOrAddLayoutRW(VulkanTexture->Surface.Image, VK_IMAGE_LAYOUT_UNDEFINED);
 				bool bIsDepthStencil = VulkanTexture->Surface.IsDepthOrStencilAspect();
@@ -1458,7 +1506,7 @@ void FVulkanCommandListContext::TransitionResources(const FPendingTransition& Pe
 
 				int32 BarrierIndex = Barrier.AddImageBarrier(VulkanTexture->Surface.Image, VulkanTexture->Surface.GetFullAspectMask(), VulkanTexture->Surface.GetNumMips(), VulkanTexture->Surface.NumArrayLevels);
 				Barrier.SetTransition(BarrierIndex, VulkanRHI::GetImageLayoutFromVulkanLayout(SrcLayout), VulkanRHI::GetImageLayoutFromVulkanLayout(DstLayout));
-
+				Barrier.AddStages(SourceStage, DestStage);
 				SrcLayout = DstLayout;
 			}
 			//#todo-rco: Temp ensure disabled
diff --git a/Engine/Source/Runtime/VulkanRHI/Public/VulkanMemory.h b/Engine/Source/Runtime/VulkanRHI/Public/VulkanMemory.h
--- a/Engine/Source/Runtime/VulkanRHI/Public/VulkanMemory.h
+++ b/Engine/Source/Runtime/VulkanRHI/Public/VulkanMemory.h
@@ -1619,6 +1619,12 @@ namespace VulkanRHI
 			}
 		}
 
+		inline void AddStages(uint32_t sourceStages, uint32_t destStages)
+		{
+			SourceStage |= sourceStages;
+			DestStage |= destStages;
+		}
+
 		// This is only valid while no other ImageBarriers are added/removed
 		inline VkImageSubresourceRange& GetSubresource(int32 BarrierIndex)
 		{
