UE4.25 - Abstract a little Vulkan shader platform/feature levels
#rnx

diff --git a/Engine/Source/Runtime/VulkanRHI/Private/Android/VulkanAndroidPlatform.cpp b/Engine/Source/Runtime/VulkanRHI/Private/Android/VulkanAndroidPlatform.cpp
--- a/Engine/Source/Runtime/VulkanRHI/Private/Android/VulkanAndroidPlatform.cpp
+++ b/Engine/Source/Runtime/VulkanRHI/Private/Android/VulkanAndroidPlatform.cpp
@@ -200,4 +200,10 @@ void FVulkanAndroidPlatform::OverridePlatformHandlers(bool bInit)
 	}
 }
 
+void FVulkanAndroidPlatform::SetupMaxRHIFeatureLevelAndShaderPlatform(ERHIFeatureLevel::Type InRequestedFeatureLevel)
+{
+	GMaxRHIFeatureLevel = ERHIFeatureLevel::ES3_1;
+	GMaxRHIShaderPlatform = SP_VULKAN_ES3_1_ANDROID;
+}
+
 #endif
diff --git a/Engine/Source/Runtime/VulkanRHI/Private/Android/VulkanAndroidPlatform.h b/Engine/Source/Runtime/VulkanRHI/Private/Android/VulkanAndroidPlatform.h
--- a/Engine/Source/Runtime/VulkanRHI/Private/Android/VulkanAndroidPlatform.h
+++ b/Engine/Source/Runtime/VulkanRHI/Private/Android/VulkanAndroidPlatform.h
@@ -44,6 +44,8 @@
 class FVulkanAndroidPlatform : public FVulkanGenericPlatform
 {
 public:
+	static void SetupMaxRHIFeatureLevelAndShaderPlatform(ERHIFeatureLevel::Type InRequestedFeatureLevel);
+
 	static bool LoadVulkanLibrary();
 	static bool LoadVulkanInstanceFunctions(VkInstance inInstance);
 	static void FreeVulkanLibrary();
diff --git a/Engine/Source/Runtime/VulkanRHI/Private/Lumin/VulkanLuminPlatform.cpp b/Engine/Source/Runtime/VulkanRHI/Private/Lumin/VulkanLuminPlatform.cpp
--- a/Engine/Source/Runtime/VulkanRHI/Private/Lumin/VulkanLuminPlatform.cpp
+++ b/Engine/Source/Runtime/VulkanRHI/Private/Lumin/VulkanLuminPlatform.cpp
@@ -213,3 +213,20 @@ VkBool32 FVulkanLuminPlatform::DebugReportFunction(
 
 	return VK_TRUE;
 }
+
+void FVulkanLuminPlatform::SetupMaxRHIFeatureLevelAndShaderPlatform(ERHIFeatureLevel::Type InRequestedFeatureLevel)
+{
+	if (!GIsEditor &&
+		(FVulkanPlatform::RequiresMobileRenderer() ||
+			InRequestedFeatureLevel == ERHIFeatureLevel::ES3_1 ||
+			FParse::Param(FCommandLine::Get(), TEXT("featureleveles31"))))
+	{
+		GMaxRHIFeatureLevel = ERHIFeatureLevel::ES3_1;
+		GMaxRHIShaderPlatform = SP_VULKAN_ES3_1_LUMIN;
+	}
+	else
+	{
+		GMaxRHIFeatureLevel = ERHIFeatureLevel::SM5;
+		GMaxRHIShaderPlatform = SP_VULKAN_SM5_LUMIN;
+	}
+}
diff --git a/Engine/Source/Runtime/VulkanRHI/Private/Lumin/VulkanLuminPlatform.h b/Engine/Source/Runtime/VulkanRHI/Private/Lumin/VulkanLuminPlatform.h
--- a/Engine/Source/Runtime/VulkanRHI/Private/Lumin/VulkanLuminPlatform.h
+++ b/Engine/Source/Runtime/VulkanRHI/Private/Lumin/VulkanLuminPlatform.h
@@ -31,6 +31,8 @@
 class FVulkanLuminPlatform : public FVulkanGenericPlatform
 {
 public:
+	static void SetupMaxRHIFeatureLevelAndShaderPlatform(ERHIFeatureLevel::Type InRequestedFeatureLevel);
+
 	static bool LoadVulkanLibrary();
 	static bool LoadVulkanInstanceFunctions(VkInstance inInstance);
 	static void FreeVulkanLibrary();
diff --git a/Engine/Source/Runtime/VulkanRHI/Private/VulkanGenericPlatform.cpp b/Engine/Source/Runtime/VulkanRHI/Private/VulkanGenericPlatform.cpp
--- a/Engine/Source/Runtime/VulkanRHI/Private/VulkanGenericPlatform.cpp
+++ b/Engine/Source/Runtime/VulkanRHI/Private/VulkanGenericPlatform.cpp
@@ -79,3 +79,20 @@ VkResult FVulkanGenericPlatform::Present(VkQueue Queue, VkPresentInfoKHR& Presen
 {
 	return VulkanRHI::vkQueuePresentKHR(Queue, &PresentInfo);
 }
+
+void FVulkanGenericPlatform::SetupMaxRHIFeatureLevelAndShaderPlatform(ERHIFeatureLevel::Type InRequestedFeatureLevel)
+{
+	if (!GIsEditor &&
+		(FVulkanPlatform::RequiresMobileRenderer() ||
+			InRequestedFeatureLevel == ERHIFeatureLevel::ES3_1 ||
+			FParse::Param(FCommandLine::Get(), TEXT("featureleveles31"))))
+	{
+		GMaxRHIFeatureLevel = ERHIFeatureLevel::ES3_1;
+		GMaxRHIShaderPlatform = SP_VULKAN_PCES3_1;
+	}
+	else
+	{
+		GMaxRHIFeatureLevel = ERHIFeatureLevel::SM5;
+		GMaxRHIShaderPlatform = SP_VULKAN_SM5;
+	}
+}
diff --git a/Engine/Source/Runtime/VulkanRHI/Private/VulkanGenericPlatform.h b/Engine/Source/Runtime/VulkanRHI/Private/VulkanGenericPlatform.h
--- a/Engine/Source/Runtime/VulkanRHI/Private/VulkanGenericPlatform.h
+++ b/Engine/Source/Runtime/VulkanRHI/Private/VulkanGenericPlatform.h
@@ -15,6 +15,8 @@ class FVulkanDevice;
 class FVulkanGenericPlatform 
 {
 public:
+	static void SetupMaxRHIFeatureLevelAndShaderPlatform(ERHIFeatureLevel::Type InRequestedFeatureLevel);
+
 	static bool IsSupported() { return true; }
 	static void CheckDeviceDriver(uint32 DeviceIndex, EGpuVendorId VendorId, const VkPhysicalDeviceProperties& Props) {}
 
diff --git a/Engine/Source/Runtime/VulkanRHI/Private/VulkanRHI.cpp b/Engine/Source/Runtime/VulkanRHI/Private/VulkanRHI.cpp
--- a/Engine/Source/Runtime/VulkanRHI/Private/VulkanRHI.cpp
+++ b/Engine/Source/Runtime/VulkanRHI/Private/VulkanRHI.cpp
@@ -91,19 +91,8 @@ bool FVulkanDynamicRHIModule::IsSupported()
 
 FDynamicRHI* FVulkanDynamicRHIModule::CreateRHI(ERHIFeatureLevel::Type InRequestedFeatureLevel)
 {
-	if (!GIsEditor &&
-		(FVulkanPlatform::RequiresMobileRenderer() ||
-			InRequestedFeatureLevel == ERHIFeatureLevel::ES3_1 ||
-			FParse::Param(FCommandLine::Get(), TEXT("featureleveles31"))))
-	{
-		GMaxRHIFeatureLevel = ERHIFeatureLevel::ES3_1;
-		GMaxRHIShaderPlatform = PLATFORM_LUMIN ? SP_VULKAN_ES3_1_LUMIN : (PLATFORM_ANDROID ? SP_VULKAN_ES3_1_ANDROID : SP_VULKAN_PCES3_1);
-	}
-	else
-	{
-		GMaxRHIFeatureLevel = ERHIFeatureLevel::SM5;
-		GMaxRHIShaderPlatform = (PLATFORM_LUMIN) ? SP_VULKAN_SM5_LUMIN : SP_VULKAN_SM5;
-	}
+	FVulkanPlatform::SetupMaxRHIFeatureLevelAndShaderPlatform(InRequestedFeatureLevel);
+	check(GMaxRHIFeatureLevel != ERHIFeatureLevel::Num);
 
 	GVulkanRHI = new FVulkanDynamicRHI();
 #if ENABLE_RHI_VALIDATION
