Fix Linux AMD vulkan query invalid results when doing vkCmdResetQueryPool right after vkCmdCopyQueryPoolResults

Move the vkCmdResetQueryPool calls to right before vkCmdWriteTimestamp

diff --git a/Engine/Source/Runtime/VulkanRHI/Private/VulkanCommandBuffer.cpp b/Engine/Source/Runtime/VulkanRHI/Private/VulkanCommandBuffer.cpp
--- a/Engine/Source/Runtime/VulkanRHI/Private/VulkanCommandBuffer.cpp
+++ b/Engine/Source/Runtime/VulkanRHI/Private/VulkanCommandBuffer.cpp
@@ -208,7 +208,6 @@ void FVulkanCmdBuffer::End()
 		VkBuffer BufferHandle = Query.BufferHandle;
 		VkQueryPool PoolHandle = Query.PoolHandle;
 		VulkanRHI::vkCmdCopyQueryPoolResults(GetHandle(), PoolHandle, Index, Query.Count, BufferHandle, sizeof(uint64) * Index, sizeof(uint64), VK_QUERY_RESULT_64_BIT|VK_QUERY_RESULT_WAIT_BIT);
-		VulkanRHI::vkCmdResetQueryPool(GetHandle(), PoolHandle, Index, Query.Count);
 	}
 
 	PendingQueries.Reset();
diff --git a/Engine/Source/Runtime/VulkanRHI/Private/VulkanQuery.cpp b/Engine/Source/Runtime/VulkanRHI/Private/VulkanQuery.cpp
--- a/Engine/Source/Runtime/VulkanRHI/Private/VulkanQuery.cpp
+++ b/Engine/Source/Runtime/VulkanRHI/Private/VulkanQuery.cpp
@@ -590,8 +590,11 @@ void FVulkanCommandListContext::RHIEndRenderQuery(FRHIRenderQuery* QueryRHI)
 		Query->Pool.CurrentTimestamp = (Query->Pool.CurrentTimestamp + 1) % Query->Pool.BufferSize;
 		const uint32 QueryEndIndex = Query->Pool.CurrentTimestamp;
 		FVulkanCmdBuffer* CmdBuffer = CommandBufferManager->GetActiveCmdBuffer();
+
+		VulkanRHI::vkCmdResetQueryPool(CmdBuffer->GetHandle(), Query->Pool.GetHandle(), QueryEndIndex, 1);
 		VulkanRHI::vkCmdWriteTimestamp(CmdBuffer->GetHandle(), VK_PIPELINE_STAGE_TOP_OF_PIPE_BIT, Query->Pool.GetHandle(), QueryEndIndex);
 		CmdBuffer->AddPendingTimestampQuery(QueryEndIndex, 1, Query->Pool.GetHandle(), Query->Pool.ResultsBuffer->GetHandle());
+
 		Query->Pool.TimestampListHandles[QueryEndIndex].CmdBuffer = CmdBuffer;
 		Query->Pool.TimestampListHandles[QueryEndIndex].FenceCounter = CmdBuffer->GetFenceSignaledCounter();
 		Query->Pool.NumIssuedTimestamps = FMath::Min<uint32>(Query->Pool.NumIssuedTimestamps + 1, Query->Pool.BufferSize);
diff --git a/Engine/Source/Runtime/VulkanRHI/Private/VulkanUtil.cpp b/Engine/Source/Runtime/VulkanRHI/Private/VulkanUtil.cpp
--- a/Engine/Source/Runtime/VulkanRHI/Private/VulkanUtil.cpp
+++ b/Engine/Source/Runtime/VulkanRHI/Private/VulkanUtil.cpp
@@ -181,7 +181,10 @@ void FVulkanGPUTiming::StartTiming(FVulkanCmdBuffer* CmdBuffer)
 		}
 		Pool->CurrentTimestamp = (Pool->CurrentTimestamp + 1) % Pool->BufferSize;
 		const uint32 QueryStartIndex = Pool->CurrentTimestamp * 2;
+
+		VulkanRHI::vkCmdResetQueryPool(CmdBuffer->GetHandle(), Pool->GetHandle(), QueryStartIndex, 2);
 		VulkanRHI::vkCmdWriteTimestamp(CmdBuffer->GetHandle(), VK_PIPELINE_STAGE_TOP_OF_PIPE_BIT, Pool->GetHandle(), QueryStartIndex);
+
 		Pool->TimestampListHandles[QueryStartIndex].CmdBuffer = CmdBuffer;
 		Pool->TimestampListHandles[QueryStartIndex].FenceCounter = CmdBuffer->GetFenceSignaledCounter();
 		bIsTiming = true;
@@ -205,8 +208,10 @@ void FVulkanGPUTiming::EndTiming(FVulkanCmdBuffer* CmdBuffer)
 		const uint32 QueryStartIndex = Pool->CurrentTimestamp * 2;
 		const uint32 QueryEndIndex = Pool->CurrentTimestamp * 2 + 1;
 		check(QueryEndIndex == QueryStartIndex + 1);	// Make sure they're adjacent indices.
+
 		VulkanRHI::vkCmdWriteTimestamp(CmdBuffer->GetHandle(), VK_PIPELINE_STAGE_TOP_OF_PIPE_BIT, Pool->GetHandle(), QueryEndIndex);
 		CmdBuffer->AddPendingTimestampQuery(QueryStartIndex, 2, Pool->GetHandle(), Pool->ResultsBuffer->GetHandle());
+
 		Pool->TimestampListHandles[QueryEndIndex].CmdBuffer = CmdBuffer;
 		Pool->TimestampListHandles[QueryEndIndex].FenceCounter = CmdBuffer->GetFenceSignaledCounter();
 		Pool->NumIssuedTimestamps = FMath::Min<uint32>(Pool->NumIssuedTimestamps + 1, Pool->BufferSize);
