Linux: Remove VulkanLinuxPlatform RenderOffScreen logging

diff --git a/Engine/Source/Runtime/VulkanRHI/Private/Linux/VulkanLinuxPlatform.cpp b/Engine/Source/Runtime/VulkanRHI/Private/Linux/VulkanLinuxPlatform.cpp
--- a/Engine/Source/Runtime/VulkanRHI/Private/Linux/VulkanLinuxPlatform.cpp
+++ b/Engine/Source/Runtime/VulkanRHI/Private/Linux/VulkanLinuxPlatform.cpp
@@ -87,7 +87,6 @@ bool FVulkanLinuxPlatform::LoadVulkanLibrary()
 	GForceEnableDebugMarkers = FParse::Param(FCommandLine::Get(), TEXT("vulkandebugmarkers"));
 
 	GRenderOffScreen = FParse::Param(FCommandLine::Get(), TEXT("RenderOffScreen"));
-	UE_LOG(LogRHI, Warning, TEXT("VulkanLinuxPlatform RenderOffScreen: %d"), GRenderOffScreen);
 	return true;
 }
 
