Vulkan: prevent overlap between the upload and graphics cmd buffers.

- Enabled on Windows and Linux(-like) platforms only, Android behavior should be unchaged.
- Perf testing inconclusive, from slight improvement to lower dips in the FPS. The main goal is correctness though.

diff --git a/Engine/Config/Unix/UnixEngine.ini b/Engine/Config/Unix/UnixEngine.ini
--- a/Engine/Config/Unix/UnixEngine.ini
+++ b/Engine/Config/Unix/UnixEngine.ini
@@ -4,3 +4,6 @@ AudioMixerModuleName=AudioMixerSDL
 
 [PlatformCrypto]
 PlatformRequiresDataCrypto=True
+
+[SystemSettings]
+r.Vulkan.UploadCmdBufferSemaphore=1
diff --git a/Engine/Config/Windows/WindowsEngine.ini b/Engine/Config/Windows/WindowsEngine.ini
--- a/Engine/Config/Windows/WindowsEngine.ini
+++ b/Engine/Config/Windows/WindowsEngine.ini
@@ -18,6 +18,7 @@ ConfigClass=/Engine/DeviceProfiles/SetDeviceProfileWindows.SetDeviceProfileWindo
 r.setres=1280x720
 framepro.ScopeMinTimeMicroseconds=10
 fx.NiagaraAllowRuntimeScalabilityChanges=1
+r.Vulkan.UploadCmdBufferSemaphore=1
 
 [PlatformCrypto]
 PlatformRequiresDataCrypto=True
diff --git a/Engine/Source/Runtime/VulkanRHI/Private/VulkanCommandBuffer.cpp b/Engine/Source/Runtime/VulkanRHI/Private/VulkanCommandBuffer.cpp
--- a/Engine/Source/Runtime/VulkanRHI/Private/VulkanCommandBuffer.cpp
+++ b/Engine/Source/Runtime/VulkanRHI/Private/VulkanCommandBuffer.cpp
@@ -26,6 +26,16 @@ static FAutoConsoleVariableRef CVarVulkanProfileCmdBuffers(
 	ECVF_Default
 );
 
+static int32 GVulkanUploadCmdBufferSemaphore = 0;	// despite this behavior being more correct, default to off to avoid changing the existing (as of UE 4.24) behavior
+static FAutoConsoleVariableRef CVarVulkanPreventOverlapWithUpload(
+	TEXT("r.Vulkan.UploadCmdBufferSemaphore"),
+	GVulkanUploadCmdBufferSemaphore,
+	TEXT("Whether command buffers for uploads and graphics can be executed simultaneously.\n")
+	TEXT(" 0: The buffers are submitted without any synch(default)\n")
+	TEXT(" 1: Graphics buffers will not overlap with the upload buffers"),
+	ECVF_ReadOnly
+);
+
 #define CMD_BUFFER_TIME_TO_WAIT_BEFORE_DELETING		10
 
 const uint32 GNumberOfFramesBeforeDeletingDescriptorPool = 300;
@@ -403,11 +413,18 @@ FVulkanCommandBufferManager::FVulkanCommandBufferManager(FVulkanDevice* InDevice
 	, Queue(InContext->GetQueue())
 	, ActiveCmdBuffer(nullptr)
 	, UploadCmdBuffer(nullptr)
+	, ActiveCmdBufferSemaphore(nullptr)
+	, UploadCmdBufferSemaphore(nullptr)
 {
 	check(Device);
 
 	Pool.Create(Queue->GetFamilyIndex());
 
+	if (GVulkanUploadCmdBufferSemaphore)
+	{
+		ActiveCmdBufferSemaphore = new VulkanRHI::FSemaphore(*Device);
+	}
+
 	ActiveCmdBuffer = Pool.Create(false);
 	ActiveCmdBuffer->InitializeTimings(InContext);
 	ActiveCmdBuffer->Begin();
@@ -481,7 +498,39 @@ void FVulkanCommandBufferManager::SubmitUploadCmdBuffer(uint32 NumSignalSemaphor
 		VulkanRHI::DebugHeavyWeightBarrier(UploadCmdBuffer->GetHandle(), 4);
 
 		UploadCmdBuffer->End();
-		Queue->Submit(UploadCmdBuffer, NumSignalSemaphores, SignalSemaphores);
+
+		if (GVulkanUploadCmdBufferSemaphore)
+		{
+			// Add semaphores associated with the recent active cmdbuf(s), if any. That will prevent
+			// the overlap, delaying execution of this cmdbuf until the graphics one(s) is complete.
+			for (FSemaphore* WaitForThis : RenderingCompletedSemaphores)
+			{
+				UploadCmdBuffer->AddWaitSemaphore(VK_PIPELINE_STAGE_TOP_OF_PIPE_BIT, WaitForThis);
+			}
+
+			if (NumSignalSemaphores == 0)
+			{
+				checkf(UploadCmdBufferSemaphore != nullptr, TEXT("FVulkanCommandBufferManager::SubmitUploadCmdBuffer: upload command buffer does not have an associated completion semaphore."));
+				VkSemaphore Sema = UploadCmdBufferSemaphore->GetHandle();
+				Queue->Submit(UploadCmdBuffer, 1, &Sema);
+				UploadCompletedSemaphores.Add(UploadCmdBufferSemaphore);
+				UploadCmdBufferSemaphore = nullptr;
+			}
+			else
+			{
+				TArray<VkSemaphore, TInlineAllocator<16>> CombinedSemaphores;
+				CombinedSemaphores.Add(UploadCmdBufferSemaphore->GetHandle());
+				CombinedSemaphores.Append(SignalSemaphores, NumSignalSemaphores);
+				Queue->Submit(UploadCmdBuffer, CombinedSemaphores.Num(), CombinedSemaphores.GetData());
+			}
+			// the buffer will now hold on to the wait semaphores, so we can clear them here
+			RenderingCompletedSemaphores.Empty();		
+		}
+		else
+		{
+			Queue->Submit(UploadCmdBuffer, NumSignalSemaphores, SignalSemaphores);
+		}
+
 		UploadCmdBuffer->SubmittedTime = FPlatformTime::Seconds();
 	}
 
@@ -505,18 +554,100 @@ void FVulkanCommandBufferManager::SubmitActiveCmdBuffer(VulkanRHI::FSemaphore* S
 		VulkanRHI::DebugHeavyWeightBarrier(ActiveCmdBuffer->GetHandle(), 8);
 
 		ActiveCmdBuffer->End();
-		if (SignalSemaphore)
+
+		if (GVulkanUploadCmdBufferSemaphore)
 		{
-			Queue->Submit(ActiveCmdBuffer, SignalSemaphore->GetHandle());
+			checkf(ActiveCmdBufferSemaphore != nullptr, TEXT("FVulkanCommandBufferManager::SubmitUploadCmdBuffer: graphics command buffer does not have an associated completion semaphore."));
+
+			// Add semaphores associated with the recent upload cmdbuf(s), if any. That will prevent
+			// the overlap, delaying execution of this cmdbuf until upload one(s) are complete.
+			for (FSemaphore* UploadCompleteSema : UploadCompletedSemaphores)
+			{
+				ActiveCmdBuffer->AddWaitSemaphore(VK_PIPELINE_STAGE_TOP_OF_PIPE_BIT, UploadCompleteSema);
+			}
+
+			if (SignalSemaphore)
+			{
+				VkSemaphore SignalThis[2] = 
+				{
+					SignalSemaphore->GetHandle(),
+					ActiveCmdBufferSemaphore->GetHandle()
+				};
+
+				Queue->Submit(ActiveCmdBuffer, UE_ARRAY_COUNT(SignalThis), SignalThis);
+			}
+			else
+			{
+				VkSemaphore SignalThis = ActiveCmdBufferSemaphore->GetHandle();
+				Queue->Submit(ActiveCmdBuffer, 1, &SignalThis);
+			}
+
+			RenderingCompletedSemaphores.Add(ActiveCmdBufferSemaphore);
+			ActiveCmdBufferSemaphore = nullptr;
+
+			// the buffer will now hold on to the wait semaphores, so we can clear them here
+			UploadCompletedSemaphores.Empty();
 		}
 		else
 		{
-			Queue->Submit(ActiveCmdBuffer);
+			if (SignalSemaphore)
+			{
+				Queue->Submit(ActiveCmdBuffer, SignalSemaphore->GetHandle());
+			}
+			else
+			{
+				Queue->Submit(ActiveCmdBuffer);
+			}
 		}
 		ActiveCmdBuffer->SubmittedTime = FPlatformTime::Seconds();
 	}
 
 	ActiveCmdBuffer = nullptr;
+	if (GVulkanUploadCmdBufferSemaphore && ActiveCmdBufferSemaphore != nullptr)
+	{
+		// most likely we didn't submit as it wasn't begun
+		check(ActiveCmdBufferSemaphore->GetRefCount() == 0);
+		delete ActiveCmdBufferSemaphore;
+		ActiveCmdBufferSemaphore = nullptr;
+	}
+}
+
+void FVulkanCommandBufferManager::SubmitActiveCmdBufferFromPresent(VulkanRHI::FSemaphore* SignalSemaphore)
+{
+	if (GVulkanUploadCmdBufferSemaphore)
+	{
+		// unlike more advanced regular SACB(), this is just a wrapper
+		// around Queue->Submit() to avoid rewriting the logic in Present
+		if (SignalSemaphore)
+		{
+			VkSemaphore SignalThis[2] =
+			{
+				SignalSemaphore->GetHandle(),
+				ActiveCmdBufferSemaphore->GetHandle()
+			};
+
+			Queue->Submit(ActiveCmdBuffer, UE_ARRAY_COUNT(SignalThis), SignalThis);
+		}
+		else
+		{
+			VkSemaphore SignalThis = ActiveCmdBufferSemaphore->GetHandle();
+			Queue->Submit(ActiveCmdBuffer, 1, &SignalThis);
+		}
+
+		RenderingCompletedSemaphores.Add(ActiveCmdBufferSemaphore);
+		ActiveCmdBufferSemaphore = nullptr;
+	}
+	else
+	{
+		if (SignalSemaphore)
+		{
+			Queue->Submit(ActiveCmdBuffer, SignalSemaphore->GetHandle());
+		}
+		else
+		{
+			Queue->Submit(ActiveCmdBuffer);
+		}
+	}
 }
 
 FVulkanCmdBuffer* FVulkanCommandBufferPool::Create(bool bIsUploadOnly)
@@ -560,6 +691,13 @@ void FVulkanCommandBufferManager::PrepareForNewActiveCommandBuffer()
 	FScopeLock ScopeLock(&Pool.CS);
 	check(!UploadCmdBuffer);
 
+	if (GVulkanUploadCmdBufferSemaphore)
+	{
+		// create a completion semaphore
+		checkf(ActiveCmdBufferSemaphore == nullptr, TEXT("FVulkanCommandBufferManager::PrepareForNewActiveCommandBuffer: a semaphore from the previous active buffer has been leaked."));
+		ActiveCmdBufferSemaphore = new VulkanRHI::FSemaphore(*Device);
+	}
+
 	for (int32 Index = 0; Index < Pool.CmdBuffers.Num(); ++Index)
 	{
 		FVulkanCmdBuffer* CmdBuffer = Pool.CmdBuffers[Index];
@@ -605,6 +743,13 @@ FVulkanCmdBuffer* FVulkanCommandBufferManager::GetUploadCmdBuffer()
 	FScopeLock ScopeLock(&Pool.CS);
 	if (!UploadCmdBuffer)
 	{
+		if (GVulkanUploadCmdBufferSemaphore)
+		{
+			// create a completion semaphore
+			checkf(UploadCmdBufferSemaphore == nullptr, TEXT("FVulkanCommandBufferManager::GetUploadCmdBuffer: a semaphore from the previous upload buffer has been leaked."));
+			UploadCmdBufferSemaphore = new VulkanRHI::FSemaphore(*Device);
+		}
+
 		for (int32 Index = 0; Index < Pool.CmdBuffers.Num(); ++Index)
 		{
 			FVulkanCmdBuffer* CmdBuffer = Pool.CmdBuffers[Index];
diff --git a/Engine/Source/Runtime/VulkanRHI/Private/VulkanCommandBuffer.h b/Engine/Source/Runtime/VulkanRHI/Private/VulkanCommandBuffer.h
--- a/Engine/Source/Runtime/VulkanRHI/Private/VulkanCommandBuffer.h
+++ b/Engine/Source/Runtime/VulkanRHI/Private/VulkanCommandBuffer.h
@@ -317,6 +317,9 @@ public:
 
 	void SubmitActiveCmdBuffer(VulkanRHI::FSemaphore* SignalSemaphore = nullptr);
 
+	/** Regular SACB() expects not-ended and would rotate the command buffer immediately, but Present has a special logic */
+	void SubmitActiveCmdBufferFromPresent(VulkanRHI::FSemaphore* SignalSemaphore = nullptr);
+
 	void WaitForCmdBuffer(FVulkanCmdBuffer* CmdBuffer, float TimeInSecondsToWait = 10.0f);
 
 
@@ -353,4 +356,16 @@ private:
 	FVulkanCmdBuffer* ActiveCmdBuffer;
 	FVulkanCmdBuffer* UploadCmdBuffer;
 	TArray<FQueryPoolReset> PoolResets;
+
+	/** This semaphore is used to prevent overlaps between the (current) graphics cmdbuf and next upload cmdbuf. */
+	VulkanRHI::FSemaphore* ActiveCmdBufferSemaphore;
+
+	/** Holds semaphores associated with the recent upload cmdbuf(s) - waiting to be added to the next graphics cmdbuf as WaitSemaphores. */
+	TArray<VulkanRHI::FSemaphore*> RenderingCompletedSemaphores;
+	
+	/** This semaphore is used to prevent overlaps between (current) upload cmdbuf and next graphics cmdbuf. */
+	VulkanRHI::FSemaphore* UploadCmdBufferSemaphore;
+	
+	/** Holds semaphores associated with the recent upload cmdbuf(s) - waiting to be added to the next graphics cmdbuf as WaitSemaphores. */
+	TArray<VulkanRHI::FSemaphore*> UploadCompletedSemaphores;
 };
diff --git a/Engine/Source/Runtime/VulkanRHI/Private/VulkanViewport.cpp b/Engine/Source/Runtime/VulkanRHI/Private/VulkanViewport.cpp
--- a/Engine/Source/Runtime/VulkanRHI/Private/VulkanViewport.cpp
+++ b/Engine/Source/Runtime/VulkanRHI/Private/VulkanViewport.cpp
@@ -845,6 +845,8 @@ bool FVulkanViewport::Present(FVulkanCommandListContext* Context, FVulkanCmdBuff
 
 	CmdBuffer->End();
 	GVulkanCommandBufferManager->FlushResetQueryPools();
+	FVulkanCommandBufferManager* ImmediateCmdBufMgr = Device->GetImmediateContext().GetCommandBufferManager();
+	checkf(ImmediateCmdBufMgr->GetActiveCmdBufferDirect() == CmdBuffer, TEXT("Present() is submitting something else than the active command buffer"));
 	if (FVulkanPlatform::SupportsStandardSwapchain())
 	{
 		if (LIKELY(!bFailedToDelayAcquireBackbuffer))
@@ -853,7 +855,9 @@ bool FVulkanViewport::Present(FVulkanCommandListContext* Context, FVulkanCmdBuff
 			{
 				CmdBuffer->AddWaitSemaphore(VK_PIPELINE_STAGE_BOTTOM_OF_PIPE_BIT, AcquiredSemaphore);
 			}
-			Queue->Submit(CmdBuffer, RenderingDoneSemaphores[AcquiredImageIndex]->GetHandle());
+			
+			// submit through the CommandBufferManager as it will add the proper semaphore
+			ImmediateCmdBufMgr->SubmitActiveCmdBufferFromPresent(RenderingDoneSemaphores[AcquiredImageIndex]);
 		}
 		else
 		{
@@ -876,7 +880,8 @@ bool FVulkanViewport::Present(FVulkanCommandListContext* Context, FVulkanCmdBuff
 	}
 	else
 	{
-		Queue->Submit(CmdBuffer);
+		// submit through the CommandBufferManager as it will add the proper semaphore
+		ImmediateCmdBufMgr->SubmitActiveCmdBufferFromPresent(RenderingDoneSemaphores[AcquiredImageIndex]);
 	}
 
 	//Flush all commands
@@ -940,7 +945,6 @@ bool FVulkanViewport::Present(FVulkanCommandListContext* Context, FVulkanCmdBuff
 	//	GInputLatencyTimer.RenderThreadTrigger = false;
 	//}
 
-	FVulkanCommandBufferManager* ImmediateCmdBufMgr = Device->GetImmediateContext().GetCommandBufferManager();
 	// PrepareForNewActiveCommandBuffer might be called by swapchain re-creation routine. Skip prepare if we already have an open active buffer.
 	if (ImmediateCmdBufMgr->GetActiveCmdBuffer() && !ImmediateCmdBufMgr->GetActiveCmdBuffer()->HasBegun())
 	{
