Add missing vulkan pixel formats

diff --git a/Engine/Source/Runtime/VulkanRHI/Private/VulkanRHIPrivate.h b/Engine/Source/Runtime/VulkanRHI/Private/VulkanRHIPrivate.h
--- a/Engine/Source/Runtime/VulkanRHI/Private/VulkanRHIPrivate.h
+++ b/Engine/Source/Runtime/VulkanRHI/Private/VulkanRHIPrivate.h
@@ -513,6 +513,27 @@ namespace VulkanRHI
 		case VK_FORMAT_R32G32B32A32_SFLOAT:
 		case VK_FORMAT_R32G32B32A32_UINT:
 			return 128;
+		case VK_FORMAT_D32_SFLOAT_S8_UINT:
+			return 8;
+		case VK_FORMAT_BC1_RGB_UNORM_BLOCK:
+		case VK_FORMAT_BC1_RGB_SRGB_BLOCK:
+		case VK_FORMAT_BC1_RGBA_UNORM_BLOCK:
+		case VK_FORMAT_BC1_RGBA_SRGB_BLOCK:
+			return 4;
+		case VK_FORMAT_BC2_UNORM_BLOCK:
+		case VK_FORMAT_BC2_SRGB_BLOCK:
+		case VK_FORMAT_BC3_UNORM_BLOCK:
+		case VK_FORMAT_BC3_SRGB_BLOCK:
+		case VK_FORMAT_BC4_UNORM_BLOCK:
+		case VK_FORMAT_BC4_SNORM_BLOCK:
+		case VK_FORMAT_BC5_UNORM_BLOCK:
+		case VK_FORMAT_BC5_SNORM_BLOCK:
+			return 8;
+		case VK_FORMAT_BC6H_UFLOAT_BLOCK:
+		case VK_FORMAT_BC6H_SFLOAT_BLOCK:
+		case VK_FORMAT_BC7_UNORM_BLOCK:
+		case VK_FORMAT_BC7_SRGB_BLOCK:
+			return 8;
 
 			// No pixel, only blocks!
 #if PLATFORM_DESKTOP
