Add tracking code to help debug memory leaks from vulkan objects. Disabled by default.
Use r.vulkan.dumpmemory to dump when enabled.

diff --git a/Engine/Source/Runtime/VulkanRHI/Private/VulkanIndexBuffer.cpp b/Engine/Source/Runtime/VulkanRHI/Private/VulkanIndexBuffer.cpp
--- a/Engine/Source/Runtime/VulkanRHI/Private/VulkanIndexBuffer.cpp
+++ b/Engine/Source/Runtime/VulkanRHI/Private/VulkanIndexBuffer.cpp
@@ -67,6 +67,8 @@ FVulkanResourceMultiBuffer::FVulkanResourceMultiBuffer(FVulkanDevice* InDevice,
 	, NumBuffers(0)
 	, DynamicBufferIndex(0)
 {
+	VULKAN_TRACK_OBJECT_CREATE(FVulkanResourceMultiBuffer, this);
+
 	if (InSize > 0)
 	{
 		const bool bStatic = (InUEUsage & BUF_Static) != 0;
@@ -143,6 +145,7 @@ FVulkanResourceMultiBuffer::FVulkanResourceMultiBuffer(FVulkanDevice* InDevice,
 
 FVulkanResourceMultiBuffer::~FVulkanResourceMultiBuffer()
 {
+	VULKAN_TRACK_OBJECT_DELETE(FVulkanResourceMultiBuffer, this);
 	//#todo-rco: Free VkBuffers
 
 	uint64_t Size = 0;
diff --git a/Engine/Source/Runtime/VulkanRHI/Private/VulkanMemory.cpp b/Engine/Source/Runtime/VulkanRHI/Private/VulkanMemory.cpp
--- a/Engine/Source/Runtime/VulkanRHI/Private/VulkanMemory.cpp
+++ b/Engine/Source/Runtime/VulkanRHI/Private/VulkanMemory.cpp
@@ -58,6 +58,26 @@ constexpr int32 GForceCoherent = 0;
 
 namespace VulkanRHI
 {
+	struct FVulkanMemoryAllocation
+	{
+		const TCHAR* Name;
+		FName ResourceName;
+		void* Address;
+		void* RHIResouce;
+		uint32 Size;
+		uint32 Width;
+		uint32 Height;
+		uint32 Depth;
+		uint32 BytesPerPixel;
+	};
+
+	struct FVulkanMemoryBucket
+	{
+		TArray<FVulkanMemoryAllocation> Allocations;
+	};
+
+
+
 	enum
 	{
 		GPU_ONLY_HEAP_PAGE_SIZE = 256 * 1024 * 1024,
@@ -403,19 +423,24 @@ namespace VulkanRHI
 	void FDeviceMemoryManager::DumpMemory()
 	{
 		SetupAndPrintMemInfo();
-		UE_LOG(LogVulkanRHI, Display, TEXT("Device Memory: %d allocations on %d heaps"), NumAllocations, HeapInfos.Num());
+#if 1 //in case of debugging, it is useful to be able to log directly to LowLevelPrintf, as this is easier to diff. Please do not delete this code.
+#define VULKAN_LOGMEMORY(fmt, ...) UE_LOG(LogVulkanRHI, Display, fmt, __VA_ARGS__)
+#else 
+#define VULKAN_LOGMEMORY(fmt, ...) FPlatformMisc::LowLevelOutputDebugStringf(fmt TEXT("\n"), __VA_ARGS__)
+#endif
+		VULKAN_LOGMEMORY(TEXT("Device Memory: %d allocations on %d heaps"), NumAllocations, HeapInfos.Num());
 		for (int32 Index = 0; Index < HeapInfos.Num(); ++Index)
 		{
 			FHeapInfo& HeapInfo = HeapInfos[Index];
-			UE_LOG(LogVulkanRHI, Display, TEXT("\tHeap %d, %d allocations"), Index, HeapInfo.Allocations.Num());
+			VULKAN_LOGMEMORY(TEXT("\tHeap %d, %d allocations"), Index, HeapInfo.Allocations.Num());
 			uint64 TotalSize = 0;
 
 			if (HeapInfo.Allocations.Num() > 0)
 			{
 #if VULKAN_MEMORY_TRACK_FILE_LINE
-				UE_LOG(LogVulkanRHI, Display, TEXT("\t\tAlloc AllocSize(MB) TotalSize(MB)    Handle  UID  File(Line)"));
+				VULKAN_LOGMEMORY(TEXT("\t\tAlloc AllocSize(MB) TotalSize(MB)    Handle  UID  File(Line)"));
 #else
-				UE_LOG(LogVulkanRHI, Display, TEXT("\t\tAlloc AllocSize(MB) TotalSize(MB)    Handle"));
+				VULKAN_LOGMEMORY(TEXT("\t\tAlloc AllocSize(MB) TotalSize(MB)    Handle"));
 #endif
 			}
 
@@ -423,14 +448,65 @@ namespace VulkanRHI
 			{
 				FDeviceMemoryAllocation* Allocation = HeapInfo.Allocations[SubIndex];
 #if VULKAN_MEMORY_TRACK_FILE_LINE
-				UE_LOG(LogVulkanRHI, Display, TEXT("\t\t%5d %13.3f %13.3f %p %4d %s(%d)"), SubIndex, Allocation->Size / 1024.f / 1024.f, TotalSize / 1024.0f / 1024.0f, (void*)Allocation->Handle, Allocation->UID, ANSI_TO_TCHAR(Allocation->File), Allocation->Line);
+				VULKAN_LOGMEMORY(TEXT("\t\t%5d %13.3f %13.3f %p %4d %s(%d)"), SubIndex, Allocation->Size / 1024.f / 1024.f, TotalSize / 1024.0f / 1024.0f, (void*)Allocation->Handle, Allocation->UID, ANSI_TO_TCHAR(Allocation->File), Allocation->Line);
 #else
-				UE_LOG(LogVulkanRHI, Display, TEXT("\t\t%5d %13.3f %13.3f %p"), SubIndex, Allocation->Size / 1024.f / 1024.f, TotalSize / 1024.0f / 1024.0f, (void*)Allocation->Handle);
+				VULKAN_LOGMEMORY(TEXT("\t\t%5d %13.3f %13.3f %p"), SubIndex, Allocation->Size / 1024.f / 1024.f, TotalSize / 1024.0f / 1024.0f, (void*)Allocation->Handle);
 #endif
 				TotalSize += Allocation->Size;
 			}
-			UE_LOG(LogVulkanRHI, Display, TEXT("\t\tTotal Allocated %.2f MB, Peak %.2f MB"), TotalSize / 1024.0f / 1024.0f, HeapInfo.PeakSize / 1024.0f / 1024.0f);
+			VULKAN_LOGMEMORY(TEXT("\t\tTotal Allocated %.2f MB, Peak %.2f MB"), TotalSize / 1024.0f / 1024.0f, HeapInfo.PeakSize / 1024.0f / 1024.0f);
+		}
+#if VULKAN_OBJECT_TRACKING
+		{
+
+			TSortedMap<uint32, FVulkanMemoryBucket> AllocationBuckets;
+			auto Collector = [&](const TCHAR* Name, FName ResourceName, void* Address, void* RHIRes, uint32 Width, uint32 Height, uint32 Depth, uint32 BytesPerPixel)
+			{
+				uint32 Size = FPlatformMath::Max(Width,1u) * FPlatformMath::Max(Height,1u) * FPlatformMath::Max(Depth, 1u) * BytesPerPixel;
+				uint32 Bucket = Size;
+				if(Bucket >= (1<<20))
+				{
+					Bucket = (Bucket + ((1 << 20) - 1)) & ~((1 << 20)-1);
+				}
+				else
+				{
+					Bucket = (Bucket + ((1 << 10) - 1)) & ~((1 << 10)-1);
+				}
+				FVulkanMemoryAllocation Allocation = 
+				{
+					Name, ResourceName, Address, RHIRes, Size, Width, Height, Depth, BytesPerPixel
+				};
+				FVulkanMemoryBucket& ActualBucket = AllocationBuckets.FindOrAdd(Bucket);
+				ActualBucket.Allocations.Add(Allocation);
+			};
+
+
+			TVulkanTrackBase<FVulkanTextureBase>::CollectAll(Collector);
+			TVulkanTrackBase<FVulkanResourceMultiBuffer>::CollectAll(Collector);
+			for(auto& Itr : AllcationBuckets)
+			{
+				VULKAN_LOGMEMORY(TEXT("***** BUCKET < %d kb *****"), Itr.Key/1024);
+				FVulkanMemoryBucket& B = Itr.Value;
+				uint32 Size = 0;
+				for (FVulkanMemoryAllocation& A : B.Allocations)
+				{
+					Size += A.Size;
+				}
+				VULKAN_LOGMEMORY(TEXT("\t\t%d / %d kb"), B.Allocations.Num(), Size / 1024);
+
+
+				B.Allocations.Sort([](const FVulkanMemoryAllocation& L, const FVulkanMemoryAllocation& R)
+				{
+					return L.Address < R.Address;
+				}
+				);
+				for(FVulkanMemoryAllocation& A : B.Allocations)
+				{
+					VULKAN_LOGMEMORY(TEXT("\t\t%p/%p %6.2fkb (%d) %5d/%5d/%5d %s ::: %s"), A.Address, A.RHIResouce, A.Size / 1024.f, A.Size, A.Width, A.Height, A.Depth, A.Name, *A.ResourceName.ToString());
+				}
+			}
 		}
+#endif
 		Device->GetResourceHeapManager().DumpMemory();
 		GLog->PanicFlushThreadedLogs();
 	}
diff --git a/Engine/Source/Runtime/VulkanRHI/Private/VulkanResources.h b/Engine/Source/Runtime/VulkanRHI/Private/VulkanResources.h
--- a/Engine/Source/Runtime/VulkanRHI/Private/VulkanResources.h
+++ b/Engine/Source/Runtime/VulkanRHI/Private/VulkanResources.h
@@ -517,13 +517,25 @@ struct FVulkanTextureBase : public FVulkanBaseShaderResource
 	FTextureRHIRef AliasedTexture;
 
 	virtual void OnTransitionResource(FVulkanCommandListContext& Context, EResourceTransitionAccess TransitionType) {};
+	
+	template<typename T>
+	void DumpMemory(T Callback)
+	{
+		uint32 BytesPerPixel = (Surface.StorageFormat != VK_FORMAT_UNDEFINED ? GetNumBitsPerPixel(Surface.StorageFormat) : 1) / 8;
+		Callback(TEXT("FVulkanTextureBase"), GetResourceFName(), this, GetRHIResource(), Surface.Width, Surface.Height, Surface.Depth, BytesPerPixel);
+	}
 
 private:
 	void DestroyViews();
+	virtual FName GetResourceFName() = 0;
+	virtual FRHIResource* GetRHIResource(){ return 0; }
+
 };
 
 class FVulkanTexture2D : public FRHITexture2D, public FVulkanTextureBase
 {
+	FName GetResourceFName(){ return GetName(); }
+	virtual FRHIResource* GetRHIResource() { return (FRHITexture2D*)this; }
 public:
 	FVulkanTexture2D(FVulkanDevice& Device, EPixelFormat Format, uint32 SizeX, uint32 SizeY, uint32 NumMips, uint32 NumSamples, uint32 UEFlags, const FRHIResourceCreateInfo& CreateInfo);
 	FVulkanTexture2D(FVulkanDevice& Device, EPixelFormat Format, uint32 SizeX, uint32 SizeY, uint32 NumMips, uint32 NumSamples, VkImage Image, uint32 UEFlags, const FRHIResourceCreateInfo& CreateInfo);
@@ -562,6 +574,7 @@ public:
 
 class FVulkanTexture2DArray : public FRHITexture2DArray, public FVulkanTextureBase
 {
+	FName GetResourceFName() { return GetName(); }
 public:
 	// Constructor, just calls base and Surface constructor
 	FVulkanTexture2DArray(FVulkanDevice& Device, EPixelFormat Format, uint32 SizeX, uint32 SizeY, uint32 ArraySize, uint32 NumMips, uint32 NumSamples, uint32 Flags, FResourceBulkDataInterface* BulkData, const FClearValueBinding& InClearValue);
@@ -598,6 +611,7 @@ public:
 
 class FVulkanTexture3D : public FRHITexture3D, public FVulkanTextureBase
 {
+	FName GetResourceFName() { return GetName(); }
 public:
 	// Constructor, just calls base and Surface constructor
 	FVulkanTexture3D(FVulkanDevice& Device, EPixelFormat Format, uint32 SizeX, uint32 SizeY, uint32 SizeZ, uint32 NumMips, uint32 Flags, FResourceBulkDataInterface* BulkData, const FClearValueBinding& InClearValue);
@@ -631,6 +645,7 @@ public:
 
 class FVulkanTextureCube : public FRHITextureCube, public FVulkanTextureBase
 {
+	FName GetResourceFName() { return GetName(); }
 public:
 	FVulkanTextureCube(FVulkanDevice& Device, EPixelFormat Format, uint32 Size, bool bArray, uint32 ArraySize, uint32 NumMips, uint32 Flags, FResourceBulkDataInterface* BulkData, const FClearValueBinding& InClearValue);
 	FVulkanTextureCube(FVulkanDevice& Device, EPixelFormat Format, uint32 Size, bool bArray, uint32 ArraySize, uint32 NumMips, VkImage Image, uint32 Flags, FResourceBulkDataInterface* BulkData, const FClearValueBinding& InClearValue);
@@ -667,6 +682,7 @@ public:
 
 class FVulkanTextureReference : public FRHITextureReference, public FVulkanTextureBase
 {
+	FName GetResourceFName() { return GetName(); }
 public:
 	explicit FVulkanTextureReference(FVulkanDevice& Device, FLastRenderTimeContainer* InLastRenderTime)
 	:	FRHITextureReference(InLastRenderTime)
@@ -1149,6 +1165,14 @@ public:
 
 	void Swap(FVulkanResourceMultiBuffer& Other);
 
+
+	template<typename T>
+	void DumpMemory(T Callback)
+	{
+		Callback(TEXT("FVulkanResourceMultiBuffer"), FName(), this, 0, GetCurrentSize() * GetNumBuffers(), 1, 1, 1);
+	}
+
+
 protected:
 	uint32 UEUsage;
 	VkBufferUsageFlags BufferUsageFlags;
diff --git a/Engine/Source/Runtime/VulkanRHI/Private/VulkanTexture.cpp b/Engine/Source/Runtime/VulkanRHI/Private/VulkanTexture.cpp
--- a/Engine/Source/Runtime/VulkanRHI/Private/VulkanTexture.cpp
+++ b/Engine/Source/Runtime/VulkanRHI/Private/VulkanTexture.cpp
@@ -1664,6 +1664,8 @@ FVulkanTextureBase::FVulkanTextureBase(FVulkanDevice& Device, VkImageViewType Re
 	: Surface(Device, ResourceType, InFormat, SizeX, SizeY, SizeZ, ArraySize, NumMips, NumSamples, UEFlags, CreateInfo)
 	, PartialView(nullptr)
 {
+	VULKAN_TRACK_OBJECT_CREATE(FVulkanTextureBase, this);
+
 	LLM_SCOPE_VULKAN(ELLMTagVulkan::VulkanTextures);
 	const bool bArray = ResourceType == VK_IMAGE_VIEW_TYPE_1D_ARRAY || ResourceType == VK_IMAGE_VIEW_TYPE_2D_ARRAY || ResourceType == VK_IMAGE_VIEW_TYPE_CUBE_ARRAY;
 	if (Surface.ViewFormat == VK_FORMAT_UNDEFINED)
@@ -1747,6 +1749,7 @@ FVulkanTextureBase::FVulkanTextureBase(FVulkanDevice& Device, VkImageViewType Re
 	: Surface(Device, ResourceType, Format, SizeX, SizeY, SizeZ, ArraySize, InNumMips, InNumSamples, InImage, UEFlags, CreateInfo)
 	, PartialView(nullptr)
 {
+	VULKAN_TRACK_OBJECT_CREATE(FVulkanTextureBase, this);
 	check(InMem == VK_NULL_HANDLE);
 	const bool bArray = ResourceType == VK_IMAGE_VIEW_TYPE_1D_ARRAY || ResourceType == VK_IMAGE_VIEW_TYPE_2D_ARRAY || ResourceType == VK_IMAGE_VIEW_TYPE_CUBE_ARRAY;
 
@@ -1770,6 +1773,8 @@ FVulkanTextureBase::FVulkanTextureBase(FVulkanDevice& Device, VkImageViewType Re
 	: Surface(Device, ResourceType, Format, SizeX, SizeY, SizeZ, ArraySize, NumMips, NumSamples, InImage, UEFlags, CreateInfo)
 	, PartialView(nullptr)
 {
+	VULKAN_TRACK_OBJECT_CREATE(FVulkanTextureBase, this);
+
 	check(InMem == VK_NULL_HANDLE);
 	const bool bArray = ResourceType == VK_IMAGE_VIEW_TYPE_1D_ARRAY || ResourceType == VK_IMAGE_VIEW_TYPE_2D_ARRAY || ResourceType == VK_IMAGE_VIEW_TYPE_CUBE_ARRAY;
 
@@ -1804,6 +1809,8 @@ FVulkanTextureBase::FVulkanTextureBase(FTextureRHIRef& SrcTextureRHI, const FVul
 	, PartialView(nullptr)
 	, AliasedTexture(SrcTextureRHI)
 {
+	VULKAN_TRACK_OBJECT_CREATE(FVulkanTextureBase, this);
+
 	if (Surface.FullAspectMask == Surface.PartialAspectMask)
 	{
 		PartialView = &DefaultView;
@@ -1819,6 +1826,7 @@ FVulkanTextureBase::FVulkanTextureBase(FTextureRHIRef& SrcTextureRHI, const FVul
 
 FVulkanTextureBase::~FVulkanTextureBase()
 {
+	VULKAN_TRACK_OBJECT_DELETE(FVulkanTextureBase, this);
 	DestroyViews();
 
 	if (PartialView != &DefaultView)
diff --git a/Engine/Source/Runtime/VulkanRHI/Public/VulkanConfiguration.h b/Engine/Source/Runtime/VulkanRHI/Public/VulkanConfiguration.h
--- a/Engine/Source/Runtime/VulkanRHI/Public/VulkanConfiguration.h
+++ b/Engine/Source/Runtime/VulkanRHI/Public/VulkanConfiguration.h
@@ -224,6 +224,13 @@
 	#endif
 #endif
 
+
+#ifndef VULKAN_OBJECT_TRACKING 
+#define VULKAN_OBJECT_TRACKING 0 //Track objects created and memory used. use r.vulkan.dumpmemory to dump to console
+#endif
+
+
+
 DECLARE_LOG_CATEGORY_EXTERN(LogVulkanRHI, Log, All);
 
 #if VULKAN_CUSTOM_MEMORY_MANAGER_ENABLED
diff --git a/Engine/Source/Runtime/VulkanRHI/Public/VulkanMemory.h b/Engine/Source/Runtime/VulkanRHI/Public/VulkanMemory.h
--- a/Engine/Source/Runtime/VulkanRHI/Public/VulkanMemory.h
+++ b/Engine/Source/Runtime/VulkanRHI/Public/VulkanMemory.h
@@ -22,6 +22,54 @@ class FVulkanCommandListContext;
 class FVulkanQueue;
 class FVulkanCmdBuffer;
 
+#if !VULKAN_OBJECT_TRACKING
+#define VULKAN_TRACK_OBJECT_CREATE(Type, Ptr) do{}while(0)
+#define VULKAN_TRACK_OBJECT_DELETE(Type, Ptr) do{}while(0)
+#else
+#define VULKAN_TRACK_OBJECT_CREATE(Type, Ptr) do{TVulkanTrackBase<Type>::Add(Ptr);}while(0)
+#define VULKAN_TRACK_OBJECT_DELETE(Type, Ptr) do{TVulkanTrackBase<Type>::Remove(Ptr);}while(0)
+
+template<typename Type>
+class TVulkanTrackBase
+{
+public:
+
+	static FCriticalSection Lock;
+	static TSet<Type*> Objects;
+	template<typename Callback>
+	static uint32 CollectAll(Callback CB)
+	{
+		uint32 Count = 0;
+		FScopeLock L(&Lock);
+		TArray<Type*> Temp;
+		for(Type* Object : Objects)
+		{
+			Object->DumpMemory(CB);
+			Count++;
+		}
+		return Count;
+	}
+	static void Add(Type* Object)
+	{
+		FScopeLock L(&Lock);
+		Objects.Add(Object);
+
+	}
+	static void Remove(Type* Object)
+	{
+		FScopeLock L(&Lock);
+		Objects.Remove(Object);
+	}
+};
+
+template<typename Type>
+FCriticalSection TVulkanTrackBase<Type>::Lock;
+template<typename Type>
+TSet<Type*> TVulkanTrackBase<Type>::Objects;
+#endif
+
+
+
 enum class EDelayAcquireImageType
 {
 	None,			// acquire next image on frame start
