fix unity test breakage caused by vulkan global func.
fix a typo inside DumpMemory

diff --git a/Engine/Source/Runtime/VulkanRHI/Private/VulkanMemory.cpp b/Engine/Source/Runtime/VulkanRHI/Private/VulkanMemory.cpp
--- a/Engine/Source/Runtime/VulkanRHI/Private/VulkanMemory.cpp
+++ b/Engine/Source/Runtime/VulkanRHI/Private/VulkanMemory.cpp
@@ -460,8 +460,9 @@ namespace VulkanRHI
 		{
 
 			TSortedMap<uint32, FVulkanMemoryBucket> AllocationBuckets;
-			auto Collector = [&](const TCHAR* Name, FName ResourceName, void* Address, void* RHIRes, uint32 Width, uint32 Height, uint32 Depth, uint32 BytesPerPixel)
+			auto Collector = [&](const TCHAR* Name, FName ResourceName, void* Address, void* RHIRes, uint32 Width, uint32 Height, uint32 Depth, uint32 Format)
 			{
+				uint32 BytesPerPixel = (Format != VK_FORMAT_UNDEFINED ? GetNumBitsPerPixel((VkFormat)Format) : 8) / 8;
 				uint32 Size = FPlatformMath::Max(Width,1u) * FPlatformMath::Max(Height,1u) * FPlatformMath::Max(Depth, 1u) * BytesPerPixel;
 				uint32 Bucket = Size;
 				if(Bucket >= (1<<20))
@@ -483,7 +484,7 @@ namespace VulkanRHI
 
 			TVulkanTrackBase<FVulkanTextureBase>::CollectAll(Collector);
 			TVulkanTrackBase<FVulkanResourceMultiBuffer>::CollectAll(Collector);
-			for(auto& Itr : AllcationBuckets)
+			for(auto& Itr : AllocationBuckets)
 			{
 				VULKAN_LOGMEMORY(TEXT("***** BUCKET < %d kb *****"), Itr.Key/1024);
 				FVulkanMemoryBucket& B = Itr.Value;
diff --git a/Engine/Source/Runtime/VulkanRHI/Private/VulkanResources.h b/Engine/Source/Runtime/VulkanRHI/Private/VulkanResources.h
--- a/Engine/Source/Runtime/VulkanRHI/Private/VulkanResources.h
+++ b/Engine/Source/Runtime/VulkanRHI/Private/VulkanResources.h
@@ -521,8 +521,7 @@ struct FVulkanTextureBase : public FVulkanBaseShaderResource
 	template<typename T>
 	void DumpMemory(T Callback)
 	{
-		uint32 BytesPerPixel = (Surface.StorageFormat != VK_FORMAT_UNDEFINED ? GetNumBitsPerPixel(Surface.StorageFormat) : 1) / 8;
-		Callback(TEXT("FVulkanTextureBase"), GetResourceFName(), this, GetRHIResource(), Surface.Width, Surface.Height, Surface.Depth, BytesPerPixel);
+		Callback(TEXT("FVulkanTextureBase"), GetResourceFName(), this, GetRHIResource(), Surface.Width, Surface.Height, Surface.Depth, Surface.StorageFormat);
 	}
 
 private:
@@ -1169,7 +1168,7 @@ public:
 	template<typename T>
 	void DumpMemory(T Callback)
 	{
-		Callback(TEXT("FVulkanResourceMultiBuffer"), FName(), this, 0, GetCurrentSize() * GetNumBuffers(), 1, 1, 1);
+		Callback(TEXT("FVulkanResourceMultiBuffer"), FName(), this, 0, GetCurrentSize() * GetNumBuffers(), 1, 1, VK_FORMAT_UNDEFINED);
 	}
 
 
