fix __VA_ARGS__ for gcc

diff --git a/Engine/Source/Runtime/VulkanRHI/Private/VulkanMemory.cpp b/Engine/Source/Runtime/VulkanRHI/Private/VulkanMemory.cpp
--- a/Engine/Source/Runtime/VulkanRHI/Private/VulkanMemory.cpp
+++ b/Engine/Source/Runtime/VulkanRHI/Private/VulkanMemory.cpp
@@ -424,9 +424,9 @@ namespace VulkanRHI
 	{
 		SetupAndPrintMemInfo();
 #if 1 //in case of debugging, it is useful to be able to log directly to LowLevelPrintf, as this is easier to diff. Please do not delete this code.
-#define VULKAN_LOGMEMORY(fmt, ...) UE_LOG(LogVulkanRHI, Display, fmt, __VA_ARGS__)
+#define VULKAN_LOGMEMORY(fmt, ...) UE_LOG(LogVulkanRHI, Display, fmt, ##__VA_ARGS__)
 #else 
-#define VULKAN_LOGMEMORY(fmt, ...) FPlatformMisc::LowLevelOutputDebugStringf(fmt TEXT("\n"), __VA_ARGS__)
+#define VULKAN_LOGMEMORY(fmt, ...) FPlatformMisc::LowLevelOutputDebugStringf(fmt TEXT("\n"), ##__VA_ARGS__)
 #endif
 		VULKAN_LOGMEMORY(TEXT("Device Memory: %d allocations on %d heaps"), NumAllocations, HeapInfos.Num());
 		for (int32 Index = 0; Index < HeapInfos.Num(); ++Index)
