DR - Static analysis fix
#rnx

diff --git a/Engine/Source/Runtime/VulkanRHI/Private/VulkanRenderTarget.cpp b/Engine/Source/Runtime/VulkanRHI/Private/VulkanRenderTarget.cpp
--- a/Engine/Source/Runtime/VulkanRHI/Private/VulkanRenderTarget.cpp
+++ b/Engine/Source/Runtime/VulkanRHI/Private/VulkanRenderTarget.cpp
@@ -943,6 +943,8 @@ void FVulkanDynamicRHI::RHIMapStagingSurface(FRHITexture* TextureRHI, FRHIGPUFen
 		VulkanRHI::vkGetImageSubresourceLayout(Device->GetInstanceHandle(), Texture2D->Surface.Image, &ImageSubResource, &SubResourceLayout);
 
 		int32 BytesPerPixel = GetNumBitsPerPixel(Texture2D->Surface.StorageFormat) / 8;
+		CA_ASSUME(BytesPerPixel > 0);
+
 		Pitch = SubResourceLayout.rowPitch / BytesPerPixel;
 	}
 
