Transition textures that are not rendertargets & depth targets to GENERAL initially. Fixes a rare validation error.
This makes vulkan work as d3d: reading uninitialized textures used to create a validation error

diff --git a/Engine/Source/Runtime/VulkanRHI/Private/VulkanTexture.cpp b/Engine/Source/Runtime/VulkanRHI/Private/VulkanTexture.cpp
--- a/Engine/Source/Runtime/VulkanRHI/Private/VulkanTexture.cpp
+++ b/Engine/Source/Runtime/VulkanRHI/Private/VulkanTexture.cpp
@@ -11,6 +11,18 @@
 #include "Containers/ResourceArray.h"
 #include "VulkanLLM.h"
 
+int32 GVulkanOverrideInitialTextureLayout = VK_IMAGE_LAYOUT_GENERAL;
+static FAutoConsoleVariableRef CVarVulkanOverrideInitialTextureLayout(
+	TEXT("r.Vulkan.OverrideInitialTextureLayout"),
+	GVulkanOverrideInitialTextureLayout,
+	TEXT("Override Initial Texture Layout\n")
+	TEXT(" -1  : Leave Undefined\n")
+	TEXT("  0..: Use assigned value"),
+	ECVF_ReadOnly
+);
+
+
+
 
 int32 GVulkanSubmitOnTextureUnlock = 1;
 static FAutoConsoleVariableRef CVarVulkanSubmitOnTextureUnlock(
@@ -458,6 +470,30 @@ struct FRHICommandRegisterImageLayout final : public FRHICommand<FRHICommandRegi
 	}
 };
 
+static void TransitionInitialImageLayout(FVulkanDevice& Device, VkImage InImage, VkImageLayout InLayout, VkImageAspectFlags InAspectMask)
+{
+	FRHICommandListImmediate& RHICmdList = FRHICommandListExecutor::GetImmediateCommandList();
+	const bool bIsInRenderingThread = IsInRenderingThread();
+	RHICmdList.EnqueueLambda([InImage, InLayout, InAspectMask](FRHICommandListImmediate& Immediate)
+	{
+		FVulkanCommandListContext& Context = (FVulkanCommandListContext&)Immediate.GetContext();
+		FVulkanCmdBuffer* CmdBuffer = Context.GetCommandBufferManager()->GetUploadCmdBuffer();
+		ensure(CmdBuffer->IsOutsideRenderPass());
+		VkImageLayout& Layout = Context.GetTransitionAndLayoutManager().FindOrAddLayoutRW(InImage, VK_IMAGE_LAYOUT_UNDEFINED);
+		check(Layout == VK_IMAGE_LAYOUT_UNDEFINED);
+		FPendingBarrier Barrier;
+		int32 BarrierIndex = Barrier.AddImageBarrier(InImage, InAspectMask, VK_REMAINING_MIP_LEVELS, VK_REMAINING_ARRAY_LAYERS);
+		Barrier.SetTransition(BarrierIndex, VulkanRHI::GetImageLayoutFromVulkanLayout(VK_IMAGE_LAYOUT_UNDEFINED), VulkanRHI::GetImageLayoutFromVulkanLayout(VK_IMAGE_LAYOUT_SHADER_READ_ONLY_OPTIMAL));
+		Barrier.Execute(CmdBuffer);
+		Layout = VK_IMAGE_LAYOUT_SHADER_READ_ONLY_OPTIMAL;
+	});
+	if (bIsInRenderingThread)
+	{
+		// Insert the RHI thread lock fence. This stops any parallel translate tasks running until the command above has completed on the RHI thread.
+		RHICmdList.RHIThreadFence(true);
+	}
+}
+
 static void InsertInitialImageLayout(FVulkanDevice& Device, VkImage InImage, VkImageLayout InLayout)
 {
 	FRHICommandListImmediate& RHICmdList = FRHICommandListExecutor::GetImmediateCommandList();
@@ -1692,10 +1728,15 @@ FVulkanTextureBase::FVulkanTextureBase(FVulkanDevice& Device, VkImageViewType Re
 
 	if (!CreateInfo.BulkData)
 	{
-		FRHICommandList& RHICmdList = FRHICommandListExecutor::GetImmediateCommandList();
-
-		// No initial data, so undefined
-		InsertInitialImageLayout(Device, Surface.Image, VK_IMAGE_LAYOUT_UNDEFINED);
+		if(-1 != GVulkanOverrideInitialTextureLayout && 0 == (UEFlags & (TexCreate_RenderTargetable | TexCreate_DepthStencilTargetable)))
+		{
+			TransitionInitialImageLayout(Device, Surface.Image, (VkImageLayout)GVulkanOverrideInitialTextureLayout, Surface.GetFullAspectMask());
+		}
+		else
+		{
+			// No initial data, so undefined
+			InsertInitialImageLayout(Device, Surface.Image, VK_IMAGE_LAYOUT_UNDEFINED);
+		}
 		return;
 	}
 
