UE4.25 - Fix static analysis
#rnx

diff --git a/Engine/Source/Runtime/VulkanRHI/Private/VulkanRenderTarget.cpp b/Engine/Source/Runtime/VulkanRHI/Private/VulkanRenderTarget.cpp
--- a/Engine/Source/Runtime/VulkanRHI/Private/VulkanRenderTarget.cpp
+++ b/Engine/Source/Runtime/VulkanRHI/Private/VulkanRenderTarget.cpp
@@ -942,10 +942,8 @@ void FVulkanDynamicRHI::RHIMapStagingSurface(FRHITexture* TextureRHI, FRHIGPUFen
 		VkSubresourceLayout SubResourceLayout;
 		VulkanRHI::vkGetImageSubresourceLayout(Device->GetInstanceHandle(), Texture2D->Surface.Image, &ImageSubResource, &SubResourceLayout);
 
-		int32 BytesPerPixel = GetNumBitsPerPixel(Texture2D->Surface.StorageFormat) / 8;
-		CA_ASSUME(BytesPerPixel > 0);
-
-		Pitch = SubResourceLayout.rowPitch / BytesPerPixel;
+		int32 BitsPerPixel = (int32)GetNumBitsPerPixel(Texture2D->Surface.StorageFormat);
+		Pitch = SubResourceLayout.rowPitch / BitsPerPixel / 8;
 	}
 
 	OutWidth = Pitch;
