Fix wrong pitch computation in FVulkanDynamicRHI::RHIMapStagingSurface

diff --git a/Engine/Source/Runtime/VulkanRHI/Private/VulkanRenderTarget.cpp b/Engine/Source/Runtime/VulkanRHI/Private/VulkanRenderTarget.cpp
--- a/Engine/Source/Runtime/VulkanRHI/Private/VulkanRenderTarget.cpp
+++ b/Engine/Source/Runtime/VulkanRHI/Private/VulkanRenderTarget.cpp
@@ -943,7 +943,8 @@ void FVulkanDynamicRHI::RHIMapStagingSurface(FRHITexture* TextureRHI, FRHIGPUFen
 		VulkanRHI::vkGetImageSubresourceLayout(Device->GetInstanceHandle(), Texture2D->Surface.Image, &ImageSubResource, &SubResourceLayout);
 
 		int32 BitsPerPixel = (int32)GetNumBitsPerPixel(Texture2D->Surface.StorageFormat);
-		Pitch = SubResourceLayout.rowPitch / BitsPerPixel / 8;
+		int32 RowPitchBits = SubResourceLayout.rowPitch * 8;
+		Pitch = RowPitchBits / BitsPerPixel;
 	}
 
 	OutWidth = Pitch;
