- Fix inconsistent usage of RHICmdList vs executing command immediately

diff --git a/Engine/Source/Runtime/VulkanRHI/Private/VulkanTexture.cpp b/Engine/Source/Runtime/VulkanRHI/Private/VulkanTexture.cpp
--- a/Engine/Source/Runtime/VulkanRHI/Private/VulkanTexture.cpp
+++ b/Engine/Source/Runtime/VulkanRHI/Private/VulkanTexture.cpp
@@ -474,23 +474,35 @@ static void TransitionInitialImageLayout(FVulkanDevice& Device, VkImage InImage,
 {
 	FRHICommandListImmediate& RHICmdList = FRHICommandListExecutor::GetImmediateCommandList();
 	const bool bIsInRenderingThread = IsInRenderingThread();
-	RHICmdList.EnqueueLambda([InImage, InLayout, InAspectMask](FRHICommandListImmediate& Immediate)
+
+	auto TransitionLambda =
+		[InImage, InLayout, InAspectMask](FRHICommandListImmediate& Immediate)
+		{
+			FVulkanCommandListContext& Context = (FVulkanCommandListContext&)Immediate.GetContext();
+			FVulkanCmdBuffer* CmdBuffer = Context.GetCommandBufferManager()->GetUploadCmdBuffer();
+			ensure(CmdBuffer->IsOutsideRenderPass());
+			VkImageLayout& Layout = Context.GetTransitionAndLayoutManager().FindOrAddLayoutRW(InImage, VK_IMAGE_LAYOUT_UNDEFINED);
+			check(Layout == VK_IMAGE_LAYOUT_UNDEFINED);
+			FPendingBarrier Barrier;
+			int32 BarrierIndex = Barrier.AddImageBarrier(InImage, InAspectMask, VK_REMAINING_MIP_LEVELS, VK_REMAINING_ARRAY_LAYERS);
+			Barrier.SetTransition(BarrierIndex, VulkanRHI::GetImageLayoutFromVulkanLayout(VK_IMAGE_LAYOUT_UNDEFINED), VulkanRHI::GetImageLayoutFromVulkanLayout(VK_IMAGE_LAYOUT_SHADER_READ_ONLY_OPTIMAL));
+			Barrier.Execute(CmdBuffer);
+			Layout = VK_IMAGE_LAYOUT_SHADER_READ_ONLY_OPTIMAL;
+		};
+
+	if (!bIsInRenderingThread || (RHICmdList.Bypass() || !IsRunningRHIInSeparateThread()))
 	{
-		FVulkanCommandListContext& Context = (FVulkanCommandListContext&)Immediate.GetContext();
-		FVulkanCmdBuffer* CmdBuffer = Context.GetCommandBufferManager()->GetUploadCmdBuffer();
-		ensure(CmdBuffer->IsOutsideRenderPass());
-		VkImageLayout& Layout = Context.GetTransitionAndLayoutManager().FindOrAddLayoutRW(InImage, VK_IMAGE_LAYOUT_UNDEFINED);
-		check(Layout == VK_IMAGE_LAYOUT_UNDEFINED);
-		FPendingBarrier Barrier;
-		int32 BarrierIndex = Barrier.AddImageBarrier(InImage, InAspectMask, VK_REMAINING_MIP_LEVELS, VK_REMAINING_ARRAY_LAYERS);
-		Barrier.SetTransition(BarrierIndex, VulkanRHI::GetImageLayoutFromVulkanLayout(VK_IMAGE_LAYOUT_UNDEFINED), VulkanRHI::GetImageLayoutFromVulkanLayout(VK_IMAGE_LAYOUT_SHADER_READ_ONLY_OPTIMAL));
-		Barrier.Execute(CmdBuffer);
-		Layout = VK_IMAGE_LAYOUT_SHADER_READ_ONLY_OPTIMAL;
-	});
-	if (bIsInRenderingThread)
+		TransitionLambda(RHICmdList);
+	}
+	else
 	{
-		// Insert the RHI thread lock fence. This stops any parallel translate tasks running until the command above has completed on the RHI thread.
-		RHICmdList.RHIThreadFence(true);
+		RHICmdList.EnqueueLambda(MoveTemp(TransitionLambda));
+
+		if (bIsInRenderingThread)
+		{
+			// Insert the RHI thread lock fence. This stops any parallel translate tasks running until the command above has completed on the RHI thread.
+			RHICmdList.RHIThreadFence(true);
+		}
 	}
 }
 
