Vulkan - implement RHICopyBufferRegion so that landscape LOD streaming works on Android devices using Vulkan

diff --git a/Engine/Source/Runtime/VulkanRHI/Private/VulkanContext.h b/Engine/Source/Runtime/VulkanRHI/Private/VulkanContext.h
--- a/Engine/Source/Runtime/VulkanRHI/Private/VulkanContext.h
+++ b/Engine/Source/Runtime/VulkanRHI/Private/VulkanContext.h
@@ -194,6 +194,7 @@ public:
 	virtual void RHIClearUAVUint(FRHIUnorderedAccessView* UnorderedAccessViewRHI, const FUintVector4& Values) final override;
 	virtual void RHICopyToResolveTarget(FRHITexture* SourceTexture, FRHITexture* DestTexture, const FResolveParams& ResolveParams) final override;
 	virtual void RHICopyTexture(FRHITexture* SourceTexture, FRHITexture* DestTexture, const FRHICopyTextureInfo& CopyInfo) final override;
+	virtual void RHICopyBufferRegion(FRHIVertexBuffer* DstBuffer, uint64 DstOffset, FRHIVertexBuffer* SrcBuffer, uint64 SrcOffset, uint64 NumBytes) final override;
 	virtual void RHITransitionResources(EResourceTransitionAccess TransitionType, FRHITexture** InRenderTargets, int32 NumTextures) final override;
 	virtual void RHITransitionResources(EResourceTransitionAccess TransitionType, EResourceTransitionPipeline TransitionPipeline, FRHITexture** InRenderTargets, int32 NumTextures) final override;
 	virtual void RHITransitionResources(EResourceTransitionAccess TransitionType, EResourceTransitionPipeline TransitionPipeline, FRHIUnorderedAccessView** InUAVs, int32 NumUAVs, FRHIComputeFence* WriteComputeFence) final override;
diff --git a/Engine/Source/Runtime/VulkanRHI/Private/VulkanTexture.cpp b/Engine/Source/Runtime/VulkanRHI/Private/VulkanTexture.cpp
--- a/Engine/Source/Runtime/VulkanRHI/Private/VulkanTexture.cpp
+++ b/Engine/Source/Runtime/VulkanRHI/Private/VulkanTexture.cpp
@@ -2381,3 +2381,44 @@ void FVulkanCommandListContext::RHICopyTexture(FRHITexture* SourceTexture, FRHIT
 	}
 	Barrier.Execute(InCmdBuffer);
 }
+
+void FVulkanCommandListContext::RHICopyBufferRegion(FRHIVertexBuffer* DstBuffer, uint64 DstOffset, FRHIVertexBuffer* SrcBuffer, uint64 SrcOffset, uint64 NumBytes)
+{
+	if (!DstBuffer || !SrcBuffer || DstBuffer == SrcBuffer || !NumBytes)
+	{
+		return;
+	}
+
+	FVulkanVertexBuffer* DstBufferVk = ResourceCast(DstBuffer);
+	FVulkanVertexBuffer* SrcBufferVk = ResourceCast(SrcBuffer);
+
+	check(DstBufferVk && SrcBufferVk);
+	check(DstOffset + NumBytes <= DstBuffer->GetSize() && SrcOffset + NumBytes <= SrcBuffer->GetSize());
+
+	uint64 DstOffsetVk = DstBufferVk->GetOffset() + DstOffset;
+	uint64 SrcOffsetVk = SrcBufferVk->GetOffset() + SrcOffset;
+
+	FVulkanCmdBuffer* CmdBuffer = CommandBufferManager->GetActiveCmdBuffer();
+	check(CmdBuffer->IsOutsideRenderPass());
+	VkCommandBuffer VkCmdBuffer = CmdBuffer->GetHandle();
+
+	{
+		VkBufferMemoryBarrier Barriers[2];
+		VulkanRHI::SetupAndZeroBufferBarrier(Barriers[0], VK_ACCESS_MEMORY_READ_BIT, VK_ACCESS_TRANSFER_WRITE_BIT, DstBufferVk->GetHandle(), DstOffsetVk, NumBytes);
+		VulkanRHI::SetupAndZeroBufferBarrier(Barriers[1], VK_ACCESS_MEMORY_WRITE_BIT, VK_ACCESS_TRANSFER_READ_BIT, SrcBufferVk->GetHandle(), SrcOffsetVk, NumBytes);
+		VulkanRHI::vkCmdPipelineBarrier(VkCmdBuffer, VK_PIPELINE_STAGE_ALL_COMMANDS_BIT, VK_PIPELINE_STAGE_TRANSFER_BIT, 0, 0, nullptr, UE_ARRAY_COUNT(Barriers), Barriers, 0, nullptr);
+	}
+
+	VkBufferCopy Region = {};
+	Region.srcOffset = SrcOffsetVk;
+	Region.dstOffset = DstOffsetVk;
+	Region.size = NumBytes;
+	VulkanRHI::vkCmdCopyBuffer(VkCmdBuffer, SrcBufferVk->GetHandle(), DstBufferVk->GetHandle(), 1, &Region);
+
+	{
+		VkBufferMemoryBarrier Barriers[2];
+		VulkanRHI::SetupAndZeroBufferBarrier(Barriers[0], VK_ACCESS_TRANSFER_WRITE_BIT, VK_ACCESS_MEMORY_READ_BIT, DstBufferVk->GetHandle(), DstOffsetVk, NumBytes);
+		VulkanRHI::SetupAndZeroBufferBarrier(Barriers[1], VK_ACCESS_TRANSFER_READ_BIT, VK_ACCESS_MEMORY_WRITE_BIT, SrcBufferVk->GetHandle(), SrcOffsetVk, NumBytes);
+		VulkanRHI::vkCmdPipelineBarrier(VkCmdBuffer, VK_PIPELINE_STAGE_TRANSFER_BIT, VK_PIPELINE_STAGE_ALL_COMMANDS_BIT, 0, 0, nullptr, UE_ARRAY_COUNT(Barriers), Barriers, 0, nullptr);
+	}
+}
