Fix Vulkan Submit for HMDs (which do not use swapchains) - was causing assert on negative array index.

diff --git a/Engine/Source/Runtime/VulkanRHI/Private/VulkanViewport.cpp b/Engine/Source/Runtime/VulkanRHI/Private/VulkanViewport.cpp
--- a/Engine/Source/Runtime/VulkanRHI/Private/VulkanViewport.cpp
+++ b/Engine/Source/Runtime/VulkanRHI/Private/VulkanViewport.cpp
@@ -880,8 +880,8 @@ bool FVulkanViewport::Present(FVulkanCommandListContext* Context, FVulkanCmdBuff
 	}
 	else
 	{
-		// submit through the CommandBufferManager as it will add the proper semaphore
-		ImmediateCmdBufMgr->SubmitActiveCmdBufferFromPresent(RenderingDoneSemaphores[AcquiredImageIndex]);
+		// Submit active command buffer if not supporting standard swapchain (e.g. XR devices).
+		ImmediateCmdBufMgr->SubmitActiveCmdBufferFromPresent(nullptr);
 	}
 
 	//Flush all commands
