Dont transition single mip individually. This is not needed, and also causes incosistent layouts across the different mips.
Fixes a crash when enabling edit layers on landscape.

diff --git a/Engine/Source/Runtime/VulkanRHI/Private/VulkanTexture.cpp b/Engine/Source/Runtime/VulkanRHI/Private/VulkanTexture.cpp
--- a/Engine/Source/Runtime/VulkanRHI/Private/VulkanTexture.cpp
+++ b/Engine/Source/Runtime/VulkanRHI/Private/VulkanTexture.cpp
@@ -2296,7 +2296,6 @@ void FVulkanCommandListContext::RHICopyTexture(FRHITexture* SourceTexture, FRHIT
 	VkImageLayout SrcLayout = TransitionAndLayoutManager.FindLayoutChecked(SrcSurface.Image);
 	bool bIsDepth = DstSurface.IsDepthOrStencilAspect();
 	VkImageLayout& DstLayoutRW = TransitionAndLayoutManager.FindOrAddLayoutRW(DstSurface.Image, VK_IMAGE_LAYOUT_UNDEFINED);
-	bool bCopyIntoCPUReadable = (DstSurface.UEFlags & TexCreate_CPUReadback) == TexCreate_CPUReadback;
 
 	FVulkanCmdBuffer* InCmdBuffer = CommandBufferManager->GetActiveCmdBuffer();
 	check(InCmdBuffer->IsOutsideRenderPass());
@@ -2368,17 +2367,8 @@ void FVulkanCommandListContext::RHICopyTexture(FRHITexture* SourceTexture, FRHIT
 
 	Barrier.ResetStages();
 	Barrier.SetTransition(SourceBarrierIndex, EImageLayoutBarrier::TransferSource, VulkanRHI::GetImageLayoutFromVulkanLayout(SrcLayout));
-
-	if (bCopyIntoCPUReadable)
-	{
-		Barrier.SetTransition(DestBarrierIndex, EImageLayoutBarrier::TransferDest, EImageLayoutBarrier::PixelGeneralRW);
-		DstLayoutRW = VK_IMAGE_LAYOUT_GENERAL;
-	}
-	else
-	{
-		Barrier.SetTransition(DestBarrierIndex, EImageLayoutBarrier::TransferDest, bIsDepth ? EImageLayoutBarrier::PixelDepthStencilRead : EImageLayoutBarrier::PixelShaderRead);
-		DstLayoutRW = bIsDepth ? VK_IMAGE_LAYOUT_DEPTH_STENCIL_READ_ONLY_OPTIMAL : VK_IMAGE_LAYOUT_SHADER_READ_ONLY_OPTIMAL;
-	}
+	Barrier.SetTransition(DestBarrierIndex, EImageLayoutBarrier::TransferDest, bIsDepth ? EImageLayoutBarrier::PixelDepthStencilRead : EImageLayoutBarrier::PixelShaderRead);
+	DstLayoutRW = bIsDepth ? VK_IMAGE_LAYOUT_DEPTH_STENCIL_READ_ONLY_OPTIMAL : VK_IMAGE_LAYOUT_SHADER_READ_ONLY_OPTIMAL;
 	Barrier.Execute(InCmdBuffer);
 }
 
