UE4.25 - Get more info for bad PSO state
#rnx

diff --git a/Engine/Source/Runtime/VulkanRHI/Private/VulkanPipeline.cpp b/Engine/Source/Runtime/VulkanRHI/Private/VulkanPipeline.cpp
--- a/Engine/Source/Runtime/VulkanRHI/Private/VulkanPipeline.cpp
+++ b/Engine/Source/Runtime/VulkanRHI/Private/VulkanPipeline.cpp
@@ -1440,6 +1440,39 @@ static inline VkPrimitiveTopology UEToVulkanTopologyType(const FVulkanDevice* In
 			return VK_PRIMITIVE_TOPOLOGY_TRIANGLE_LIST;
 		case PT_TriangleStrip:
 			return VK_PRIMITIVE_TOPOLOGY_TRIANGLE_STRIP;
+		case PT_1_ControlPointPatchList:
+		case PT_2_ControlPointPatchList:
+		case PT_3_ControlPointPatchList:
+		case PT_4_ControlPointPatchList:
+		case PT_5_ControlPointPatchList:
+		case PT_6_ControlPointPatchList:
+		case PT_7_ControlPointPatchList:
+		case PT_8_ControlPointPatchList:
+		case PT_9_ControlPointPatchList:
+		case PT_10_ControlPointPatchList:
+		case PT_12_ControlPointPatchList:
+		case PT_13_ControlPointPatchList:
+		case PT_14_ControlPointPatchList:
+		case PT_15_ControlPointPatchList:
+		case PT_16_ControlPointPatchList:
+		case PT_17_ControlPointPatchList:
+		case PT_18_ControlPointPatchList:
+		case PT_19_ControlPointPatchList:
+		case PT_20_ControlPointPatchList:
+		case PT_22_ControlPointPatchList:
+		case PT_23_ControlPointPatchList:
+		case PT_24_ControlPointPatchList:
+		case PT_25_ControlPointPatchList:
+		case PT_26_ControlPointPatchList:
+		case PT_27_ControlPointPatchList:
+		case PT_28_ControlPointPatchList:
+		case PT_29_ControlPointPatchList:
+		case PT_30_ControlPointPatchList:
+		case PT_31_ControlPointPatchList:
+		case PT_32_ControlPointPatchList:
+			OutControlPoints = (PrimitiveType - PT_1_ControlPointPatchList + 1);
+			checkf(false, TEXT("Missing tessellation shaders, however tried to use EPrimitiveType %d (%d control points)"), (uint32)PrimitiveType, OutControlPoints);
+			break;
 		default:
 			checkf(false, TEXT("Unsupported EPrimitiveType %d"), (uint32)PrimitiveType);
 			break;
