PR #6784: Vulkan LLM improvements: track spare and render target memory (Contributed by [REDACTED_PERSON])

diff --git a/Engine/Source/Runtime/VulkanRHI/Private/VulkanLLM.cpp b/Engine/Source/Runtime/VulkanRHI/Private/VulkanLLM.cpp
--- a/Engine/Source/Runtime/VulkanRHI/Private/VulkanLLM.cpp
+++ b/Engine/Source/Runtime/VulkanRHI/Private/VulkanLLM.cpp
@@ -20,15 +20,17 @@ struct FLLMTagInfoVulkan
 	FName SummaryStatName;		// shows in the LLM summary stat group
 };
 
-DECLARE_LLM_MEMORY_STAT(TEXT("VulkanMisc"), STAT_VulkanMiscLLM, STATGROUP_LLMFULL);
-DECLARE_LLM_MEMORY_STAT(TEXT("VulkanUniformBuffers"), STAT_VulkanUniformBuffersLLM, STATGROUP_LLMFULL);
-DECLARE_LLM_MEMORY_STAT(TEXT("VulkanIndexBuffers"), STAT_VulkanIndexBuffersLLM, STATGROUP_LLMFULL);
-DECLARE_LLM_MEMORY_STAT(TEXT("VulkanVertexBuffers"), STAT_VulkanVertexBuffersLLM, STATGROUP_LLMFULL);
-DECLARE_LLM_MEMORY_STAT(TEXT("VulkanTextures"), STAT_VulkanTexturesLLM, STATGROUP_LLMFULL);
-DECLARE_LLM_MEMORY_STAT(TEXT("VulkanShaders"), STAT_VulkanShadersLLM, STATGROUP_LLMFULL);
-DECLARE_LLM_MEMORY_STAT(TEXT("VulkanFrameTemp"), STAT_VulkanFrameTempLLM, STATGROUP_LLMFULL);
-DECLARE_LLM_MEMORY_STAT(TEXT("VulkanStagingBuffers"), STAT_VulkanStagingBuffersLLM, STATGROUP_LLMFULL);
-DECLARE_LLM_MEMORY_STAT(TEXT("VulkanDriverMemoryCPU"), STAT_VulkanDriverMemoryCPULLM, STATGROUP_LLMFULL);
+DECLARE_LLM_MEMORY_STAT(TEXT("VulkanMisc"), STAT_VulkanMiscLLM, STATGROUP_LLMPlatform);
+DECLARE_LLM_MEMORY_STAT(TEXT("VulkanUniformBuffers"), STAT_VulkanUniformBuffersLLM, STATGROUP_LLMPlatform);
+DECLARE_LLM_MEMORY_STAT(TEXT("VulkanIndexBuffers"), STAT_VulkanIndexBuffersLLM, STATGROUP_LLMPlatform);
+DECLARE_LLM_MEMORY_STAT(TEXT("VulkanVertexBuffers"), STAT_VulkanVertexBuffersLLM, STATGROUP_LLMPlatform);
+DECLARE_LLM_MEMORY_STAT(TEXT("VulkanTextures"), STAT_VulkanTexturesLLM, STATGROUP_LLMPlatform);
+DECLARE_LLM_MEMORY_STAT(TEXT("VulkanRenderTargets"), STAT_VulkanRenderTargetsLLM, STATGROUP_LLMPlatform);
+DECLARE_LLM_MEMORY_STAT(TEXT("VulkanSpareMemoryGPU"), STAT_VulkanSpareMemoryGPULLM, STATGROUP_LLMPlatform);
+DECLARE_LLM_MEMORY_STAT(TEXT("VulkanShaders"), STAT_VulkanShadersLLM, STATGROUP_LLMPlatform);
+DECLARE_LLM_MEMORY_STAT(TEXT("VulkanFrameTemp"), STAT_VulkanFrameTempLLM, STATGROUP_LLMPlatform);
+DECLARE_LLM_MEMORY_STAT(TEXT("VulkanStagingBuffers"), STAT_VulkanStagingBuffersLLM, STATGROUP_LLMPlatform);
+DECLARE_LLM_MEMORY_STAT(TEXT("VulkanDriverMemoryCPU"), STAT_VulkanDriverMemoryCPULLM, STATGROUP_LLMPlatform);
 DECLARE_LLM_MEMORY_STAT(TEXT("VulkanDriverMemoryGPU"), STAT_VulkanDriverMemoryGPULLM, STATGROUP_LLMPlatform);
 
 
@@ -41,6 +43,8 @@ static const FLLMTagInfoVulkan ELLMTagNamesVulkan[] =
 	{ TEXT("VulkanIndexBuffers"),		GET_STATFNAME(STAT_VulkanIndexBuffersLLM),					GET_STATFNAME(STAT_EngineSummaryLLM) },		// ELLMTagVulkan::VulkanIndexBuffers
 	{ TEXT("VulkanVertexBuffers"),		GET_STATFNAME(STAT_VulkanVertexBuffersLLM),					GET_STATFNAME(STAT_EngineSummaryLLM) },		// ELLMTagVulkan::VulkanVertexBuffers
 	{ TEXT("VulkanTextures"),			GET_STATFNAME(STAT_VulkanTexturesLLM),						GET_STATFNAME(STAT_EngineSummaryLLM) },		// ELLMTagVulkan::VulkanTextures
+	{ TEXT("VulkanRenderTargets"),		GET_STATFNAME(STAT_VulkanRenderTargetsLLM),					GET_STATFNAME(STAT_EngineSummaryLLM) },		// ELLMTagVulkan::VulkanRenderTargets
+	{ TEXT("VulkanSpareMemoryGPU"),		GET_STATFNAME(STAT_VulkanSpareMemoryGPULLM),				GET_STATFNAME(STAT_EngineSummaryLLM) },		// ELLMTagVulkan::VulkanSpareMemoryGPU
 	{ TEXT("VulkanShaders"),			GET_STATFNAME(STAT_VulkanShadersLLM),						GET_STATFNAME(STAT_EngineSummaryLLM) },		// ELLMTagVulkan::VulkanShaders
 	{ TEXT("VulkanFrameTemp"),			GET_STATFNAME(STAT_VulkanFrameTempLLM),						GET_STATFNAME(STAT_EngineSummaryLLM) },		// ELLMTagVulkan::VulkanFrameTempGPU
 	{ TEXT("VulkanStagingBuffers"),		GET_STATFNAME(STAT_VulkanStagingBuffersLLM),				GET_STATFNAME(STAT_EngineSummaryLLM) },		// ELLMTagVulkan::VulkanStagingBuffersGPU
diff --git a/Engine/Source/Runtime/VulkanRHI/Private/VulkanLLM.h b/Engine/Source/Runtime/VulkanRHI/Private/VulkanLLM.h
--- a/Engine/Source/Runtime/VulkanRHI/Private/VulkanLLM.h
+++ b/Engine/Source/Runtime/VulkanRHI/Private/VulkanLLM.h
@@ -17,6 +17,8 @@ extern uint64 GVulkanLLMAllocationID;
 #define LLM_TRACK_VULKAN_HIGH_LEVEL_ALLOC(AllocObj, Size) { AllocObj->SetLLMTrackerID(0xDEAD | (++GVulkanLLMAllocationID << 16)); FLowLevelMemTracker::Get().OnLowLevelAlloc(ELLMTracker::Default, (void*)AllocObj->GetLLMTrackerID(), Size, (ELLMTag)ELLMTagVulkan::VulkanMisc); }
 #define LLM_TRACK_VULKAN_HIGH_LEVEL_FREE(AllocObj) { FLowLevelMemTracker::Get().OnLowLevelFree(ELLMTracker::Default, (void*)AllocObj->GetLLMTrackerID()); }
 
+#define LLM_TRACK_VULKAN_SPARE_MEMORY_GPU(Size)	{ LLM_SCOPED_PAUSE_TRACKING_WITH_ENUM_AND_AMOUNT((ELLMTag)ELLMTagVulkan::VulkanSpareMemoryGPU, Size, ELLMTracker::Default, ELLMAllocType::None); }
+
 #define LLM_TRACK_VULKAN_HIGH_LEVEL_ALLOCATION
 enum class ELLMTagVulkan : LLM_TAG_TYPE
 {
@@ -25,6 +27,8 @@ enum class ELLMTagVulkan : LLM_TAG_TYPE
 	VulkanIndexBuffers,
 	VulkanVertexBuffers,
 	VulkanTextures,
+	VulkanRenderTargets,
+	VulkanSpareMemoryGPU,
 	VulkanShaders,	
 	VulkanFrameTemp,
 	VulkanStagingBuffers,
@@ -47,4 +51,6 @@ namespace VulkanLLM
 #define LLM_TRACK_VULKAN_HIGH_LEVEL_ALLOC(...)
 #define LLM_TRACK_VULKAN_HIGH_LEVEL_FREE(...)
 
+#define LLM_TRACK_VULKAN_SPARE_MEMORY_GPU(...)
+
 #endif		// #if VULKAN_USE_LLM
diff --git a/Engine/Source/Runtime/VulkanRHI/Private/VulkanMemory.cpp b/Engine/Source/Runtime/VulkanRHI/Private/VulkanMemory.cpp
--- a/Engine/Source/Runtime/VulkanRHI/Private/VulkanMemory.cpp
+++ b/Engine/Source/Runtime/VulkanRHI/Private/VulkanMemory.cpp
@@ -374,6 +374,7 @@ namespace VulkanRHI
 #if VULKAN_USE_LLM
 		LLM_PLATFORM_SCOPE_VULKAN(ELLMTagVulkan::VulkanDriverMemoryGPU);
 		LLM(FLowLevelMemTracker::Get().OnLowLevelAlloc(ELLMTracker::Platform, (void*)NewAllocation->Handle, AllocationSize, ELLMTag::GraphicsPlatform, ELLMAllocType::System));
+		LLM_TRACK_VULKAN_SPARE_MEMORY_GPU((int64)AllocationSize);
 #endif
 
 		INC_DWORD_STAT(STAT_VulkanNumPhysicalMemAllocations);
@@ -404,6 +405,7 @@ namespace VulkanRHI
 
 #if VULKAN_USE_LLM
 		LLM(FLowLevelMemTracker::Get().OnLowLevelFree(ELLMTracker::Platform, (void*)Allocation->Handle, ELLMAllocType::System));
+		LLM_TRACK_VULKAN_SPARE_MEMORY_GPU(-(int64)Allocation->Size);
 #endif
 
 		--NumAllocations;
@@ -877,6 +879,8 @@ namespace VulkanRHI
 
 				LLM_TRACK_VULKAN_HIGH_LEVEL_ALLOC(NewResourceAllocation, Size);
 
+				LLM_TRACK_VULKAN_SPARE_MEMORY_GPU(-(int64)Size);
+
 				return NewResourceAllocation;
 			}
 		}
@@ -889,6 +893,8 @@ namespace VulkanRHI
 		{
 			LLM_TRACK_VULKAN_HIGH_LEVEL_FREE(Allocation);
 
+			LLM_TRACK_VULKAN_SPARE_MEMORY_GPU((int64)Allocation->RequestedSize);
+
 			FScopeLock ScopeLock(&GOldResourcePageLock);
 			ResourceAllocations.RemoveSingleSwap(Allocation, false);
 
@@ -1829,6 +1835,8 @@ namespace VulkanRHI
 
 				LLM_TRACK_VULKAN_HIGH_LEVEL_ALLOC(NewSuballocation, InSize);
 
+				LLM_TRACK_VULKAN_SPARE_MEMORY_GPU(-(int64)InSize);
+
 				//PeakNumAllocations = FMath::Max(PeakNumAllocations, ResourceAllocations.Num());
 				return NewSuballocation;
 			}
@@ -1851,6 +1859,8 @@ namespace VulkanRHI
 
 			LLM_TRACK_VULKAN_HIGH_LEVEL_FREE(Suballocation);
 
+			LLM_TRACK_VULKAN_SPARE_MEMORY_GPU((int64)Suballocation->RequestedSize);
+
 			FRange NewFree;
 			NewFree.Offset = Suballocation->AllocationOffset;
 			NewFree.Size = Suballocation->AllocationSize;
diff --git a/Engine/Source/Runtime/VulkanRHI/Private/VulkanTexture.cpp b/Engine/Source/Runtime/VulkanRHI/Private/VulkanTexture.cpp
--- a/Engine/Source/Runtime/VulkanRHI/Private/VulkanTexture.cpp
+++ b/Engine/Source/Runtime/VulkanRHI/Private/VulkanTexture.cpp
@@ -50,6 +50,14 @@ struct FTextureLock
 	}
 };
 
+#if VULKAN_USE_LLM
+inline ELLMTagVulkan GetMemoryTagForTextureFlags(uint32 UEFlags)
+{
+	bool bRenderTarget = ((TexCreate_RenderTargetable | TexCreate_ResolveTargetable | TexCreate_DepthStencilTargetable) & UEFlags) != 0u;
+	return bRenderTarget ? ELLMTagVulkan::VulkanRenderTargets : ELLMTagVulkan::VulkanTextures;
+}
+#endif // VULKAN_USE_LLM
+
 inline bool operator == (const FTextureLock& A, const FTextureLock& B)
 {
 	return A.Texture == B.Texture && A.MipIndex == B.MipIndex && A.LayerIndex == B.LayerIndex;
@@ -1026,7 +1034,7 @@ uint32 FVulkanDynamicRHI::RHIComputeMemorySize(FRHITexture* TextureRHI)
 FTexture2DRHIRef FVulkanDynamicRHI::RHICreateTexture2D(uint32 SizeX, uint32 SizeY, uint8 Format, uint32 NumMips, uint32 NumSamples, uint32 Flags, FRHIResourceCreateInfo& CreateInfo)
 
 {
-	LLM_SCOPE_VULKAN(ELLMTagVulkan::VulkanTextures);
+	LLM_SCOPE_VULKAN(GetMemoryTagForTextureFlags(Flags));
 	return new FVulkanTexture2D(*Device, (EPixelFormat)Format, SizeX, SizeY, NumMips, NumSamples, Flags, CreateInfo);
 }
 
@@ -1044,13 +1052,13 @@ void FVulkanDynamicRHI::RHICopySharedMips(FRHITexture2D* DestTexture2D, FRHIText
 
 FTexture2DArrayRHIRef FVulkanDynamicRHI::RHICreateTexture2DArray(uint32 SizeX, uint32 SizeY, uint32 SizeZ, uint8 Format, uint32 NumMips, uint32 NumSamples, uint32 Flags, FRHIResourceCreateInfo& CreateInfo)
 {
-	LLM_SCOPE_VULKAN(ELLMTagVulkan::VulkanTextures);
+	LLM_SCOPE_VULKAN(GetMemoryTagForTextureFlags(Flags));
 	return new FVulkanTexture2DArray(*Device, (EPixelFormat)Format, SizeX, SizeY, SizeZ, NumMips, NumSamples, Flags, CreateInfo.BulkData, CreateInfo.ClearValueBinding);
 }
 
 FTexture3DRHIRef FVulkanDynamicRHI::RHICreateTexture3D(uint32 SizeX, uint32 SizeY, uint32 SizeZ, uint8 Format, uint32 NumMips, uint32 Flags, FRHIResourceCreateInfo& CreateInfo)
 {
-	LLM_SCOPE_VULKAN(ELLMTagVulkan::VulkanTextures);
+	LLM_SCOPE_VULKAN(GetMemoryTagForTextureFlags(Flags));
 	FVulkanTexture3D* Tex3d = new FVulkanTexture3D(*Device, (EPixelFormat)Format, SizeX, SizeY, SizeZ, NumMips, Flags, CreateInfo.BulkData, CreateInfo.ClearValueBinding);
 
 	return Tex3d;
@@ -2026,13 +2034,13 @@ FVulkanTexture3D::~FVulkanTexture3D()
 -----------------------------------------------------------------------------*/
 FTextureCubeRHIRef FVulkanDynamicRHI::RHICreateTextureCube(uint32 Size, uint8 Format, uint32 NumMips, uint32 Flags, FRHIResourceCreateInfo& CreateInfo)
 {
-	LLM_SCOPE_VULKAN(ELLMTagVulkan::VulkanTextures);
+	LLM_SCOPE_VULKAN(GetMemoryTagForTextureFlags(Flags));
 	return new FVulkanTextureCube(*Device, (EPixelFormat)Format, Size, false, 1, NumMips, Flags, CreateInfo.BulkData, CreateInfo.ClearValueBinding);
 }
 
 FTextureCubeRHIRef FVulkanDynamicRHI::RHICreateTextureCubeArray(uint32 Size, uint32 ArraySize, uint8 Format, uint32 NumMips, uint32 Flags, FRHIResourceCreateInfo& CreateInfo)
 {
-	LLM_SCOPE_VULKAN(ELLMTagVulkan::VulkanTextures);
+	LLM_SCOPE_VULKAN(GetMemoryTagForTextureFlags(Flags));
 	return new FVulkanTextureCube(*Device, (EPixelFormat)Format, Size, true, ArraySize, NumMips, Flags, CreateInfo.BulkData, CreateInfo.ClearValueBinding);
 }
 
