VulkanDebug validation fix:
Add pending time stamp query for start when setup instead of done at end - could be that begin/end is called from a different command list and then reset is not called on the same command list which is using it

diff --git a/Engine/Source/Runtime/VulkanRHI/Private/VulkanUtil.cpp b/Engine/Source/Runtime/VulkanRHI/Private/VulkanUtil.cpp
--- a/Engine/Source/Runtime/VulkanRHI/Private/VulkanUtil.cpp
+++ b/Engine/Source/Runtime/VulkanRHI/Private/VulkanUtil.cpp
@@ -182,6 +182,7 @@ void FVulkanGPUTiming::StartTiming(FVulkanCmdBuffer* CmdBuffer)
 		Pool->CurrentTimestamp = (Pool->CurrentTimestamp + 1) % Pool->BufferSize;
 		const uint32 QueryStartIndex = Pool->CurrentTimestamp * 2;
 		VulkanRHI::vkCmdWriteTimestamp(CmdBuffer->GetHandle(), VK_PIPELINE_STAGE_TOP_OF_PIPE_BIT, Pool->GetHandle(), QueryStartIndex);
+		CmdBuffer->AddPendingTimestampQuery(QueryStartIndex, 1, Pool->GetHandle(), Pool->ResultsBuffer->GetHandle());
 		Pool->TimestampListHandles[QueryStartIndex].CmdBuffer = CmdBuffer;
 		Pool->TimestampListHandles[QueryStartIndex].FenceCounter = CmdBuffer->GetFenceSignaledCounter();
 		bIsTiming = true;
@@ -206,7 +207,7 @@ void FVulkanGPUTiming::EndTiming(FVulkanCmdBuffer* CmdBuffer)
 		const uint32 QueryEndIndex = Pool->CurrentTimestamp * 2 + 1;
 		check(QueryEndIndex == QueryStartIndex + 1);	// Make sure they're adjacent indices.
 		VulkanRHI::vkCmdWriteTimestamp(CmdBuffer->GetHandle(), VK_PIPELINE_STAGE_TOP_OF_PIPE_BIT, Pool->GetHandle(), QueryEndIndex);
-		CmdBuffer->AddPendingTimestampQuery(QueryStartIndex, 2, Pool->GetHandle(), Pool->ResultsBuffer->GetHandle());
+		CmdBuffer->AddPendingTimestampQuery(QueryEndIndex, 1, Pool->GetHandle(), Pool->ResultsBuffer->GetHandle());
 		Pool->TimestampListHandles[QueryEndIndex].CmdBuffer = CmdBuffer;
 		Pool->TimestampListHandles[QueryEndIndex].FenceCounter = CmdBuffer->GetFenceSignaledCounter();
 		Pool->NumIssuedTimestamps = FMath::Min<uint32>(Pool->NumIssuedTimestamps + 1, Pool->BufferSize);
