Don't add attachment texture view for resolve texture if RTLayout also didn't add it - fix vulkan debug warning

diff --git a/Engine/Source/Runtime/VulkanRHI/Private/VulkanViewport.cpp b/Engine/Source/Runtime/VulkanRHI/Private/VulkanViewport.cpp
--- a/Engine/Source/Runtime/VulkanRHI/Private/VulkanViewport.cpp
+++ b/Engine/Source/Runtime/VulkanRHI/Private/VulkanViewport.cpp
@@ -382,7 +382,9 @@ FVulkanFramebuffer::FVulkanFramebuffer(FVulkanDevice& Device, const FRHISetRende
 
 		++NumColorAttachments;
 
-		if (InRTInfo.bHasResolveAttachments)
+		// Check the RTLayout as well to make sure the resolve attachment is needed (Vulkan and Feature level specific)
+		// See: FVulkanRenderTargetLayout constructor with FRHIRenderPassInfo
+		if (InRTInfo.bHasResolveAttachments && RTLayout.GetHasResolveAttachments() && RTLayout.GetResolveAttachmentReferences()[Index].layout != VK_IMAGE_LAYOUT_UNDEFINED)
 		{
 			FRHITexture* ResolveRHITexture = InRTInfo.ColorResolveRenderTarget[Index].Texture;
 			FVulkanTextureBase* ResolveTexture = FVulkanTextureBase::Cast(ResolveRHITexture);
