Fix a typo: Initial layout transition wasn't actually using the intended layout type

diff --git a/Engine/Source/Runtime/VulkanRHI/Private/VulkanTexture.cpp b/Engine/Source/Runtime/VulkanRHI/Private/VulkanTexture.cpp
--- a/Engine/Source/Runtime/VulkanRHI/Private/VulkanTexture.cpp
+++ b/Engine/Source/Runtime/VulkanRHI/Private/VulkanTexture.cpp
@@ -493,9 +493,9 @@ static void TransitionInitialImageLayout(FVulkanDevice& Device, VkImage InImage,
 			check(Layout == VK_IMAGE_LAYOUT_UNDEFINED);
 			FPendingBarrier Barrier;
 			int32 BarrierIndex = Barrier.AddImageBarrier(InImage, InAspectMask, VK_REMAINING_MIP_LEVELS, VK_REMAINING_ARRAY_LAYERS);
-			Barrier.SetTransition(BarrierIndex, VulkanRHI::GetImageLayoutFromVulkanLayout(VK_IMAGE_LAYOUT_UNDEFINED), VulkanRHI::GetImageLayoutFromVulkanLayout(VK_IMAGE_LAYOUT_SHADER_READ_ONLY_OPTIMAL));
+			Barrier.SetTransition(BarrierIndex, VulkanRHI::GetImageLayoutFromVulkanLayout(VK_IMAGE_LAYOUT_UNDEFINED), VulkanRHI::GetImageLayoutFromVulkanLayout(InLayout));
 			Barrier.Execute(CmdBuffer);
-			Layout = VK_IMAGE_LAYOUT_SHADER_READ_ONLY_OPTIMAL;
+			Layout = InLayout;
 		};
 
 	if (!bIsInRenderingThread || (RHICmdList.Bypass() || !IsRunningRHIInSeparateThread()))
