Implement Begin/EndUpdateTexture3D_RenderThread in the VulkanRHI and perform non blocking end by enqueuing command in the RHI Thread - fixes barrier order issue

diff --git a/Engine/Source/Runtime/VulkanRHI/Private/VulkanTexture.cpp b/Engine/Source/Runtime/VulkanRHI/Private/VulkanTexture.cpp
--- a/Engine/Source/Runtime/VulkanRHI/Private/VulkanTexture.cpp
+++ b/Engine/Source/Runtime/VulkanRHI/Private/VulkanTexture.cpp
@@ -1464,6 +1464,29 @@ void FVulkanDynamicRHI::InternalUpdateTexture2D(bool bFromRenderingThread, FRHIT
 	}
 }
 
+FUpdateTexture3DData FVulkanDynamicRHI::BeginUpdateTexture3D_RenderThread(class FRHICommandListImmediate& RHICmdList, FRHITexture3D* Texture, uint32 MipIndex, const struct FUpdateTextureRegion3D& UpdateRegion)
+{
+	const int32 FormatSize = PixelFormatBlockBytes[Texture->GetFormat()];
+	const int32 RowPitch = UpdateRegion.Width * FormatSize;
+	const int32 DepthPitch = UpdateRegion.Width * UpdateRegion.Height * FormatSize;
+
+	SIZE_T MemorySize = static_cast<SIZE_T>(DepthPitch)* UpdateRegion.Depth;
+	uint8* Data = (uint8*)FMemory::Malloc(MemorySize);
+
+	return FUpdateTexture3DData(Texture, MipIndex, UpdateRegion, RowPitch, DepthPitch, Data, MemorySize, GFrameNumberRenderThread);
+}
+
+void FVulkanDynamicRHI::EndUpdateTexture3D_RenderThread(class FRHICommandListImmediate& RHICmdList, FUpdateTexture3DData& UpdateData)
+{
+	check(IsInRenderingThread());
+	check(GFrameNumberRenderThread == UpdateData.FrameNumber);
+
+	InternalUpdateTexture3D(true, UpdateData.Texture, UpdateData.MipIndex, UpdateData.UpdateRegion, UpdateData.RowPitch, UpdateData.DepthPitch, UpdateData.Data);
+	
+	FMemory::Free(UpdateData.Data);
+	UpdateData.Data = nullptr;
+}
+
 void FVulkanDynamicRHI::InternalUpdateTexture3D(bool bFromRenderingThread, FRHITexture3D* TextureRHI, uint32 MipIndex, const FUpdateTextureRegion3D& UpdateRegion, uint32 SourceRowPitch, uint32 SourceDepthPitch, const uint8* SourceData)
 {
 	LLM_SCOPE_VULKAN(ELLMTagVulkan::VulkanTextures);
diff --git a/Engine/Source/Runtime/VulkanRHI/Public/VulkanDynamicRHI.h b/Engine/Source/Runtime/VulkanRHI/Public/VulkanDynamicRHI.h
--- a/Engine/Source/Runtime/VulkanRHI/Public/VulkanDynamicRHI.h
+++ b/Engine/Source/Runtime/VulkanRHI/Public/VulkanDynamicRHI.h
@@ -254,6 +254,9 @@ public:
 	{
 		InternalUpdateTexture3D(true, Texture, MipIndex, UpdateRegion, SourceRowPitch, SourceDepthPitch, SourceData);
 	}
+	
+	virtual FUpdateTexture3DData BeginUpdateTexture3D_RenderThread(class FRHICommandListImmediate& RHICmdList, FRHITexture3D* Texture, uint32 MipIndex, const struct FUpdateTextureRegion3D& UpdateRegion) override final;
+	virtual void EndUpdateTexture3D_RenderThread(class FRHICommandListImmediate& RHICmdList, FUpdateTexture3DData& UpdateData) override final;
 
 	virtual FTextureReferenceRHIRef RHICreateTextureReference_RenderThread(class FRHICommandListImmediate& RHICmdList, FLastRenderTimeContainer* LastRenderTime) final override
 	{
