Fix for Android Vulkan SM5 for consolidated sets triggering as errors when gathering uniforms for shaders without uniforms

diff --git a/Engine/Source/Runtime/VulkanRHI/Private/VulkanShaders.cpp b/Engine/Source/Runtime/VulkanRHI/Private/VulkanShaders.cpp
--- a/Engine/Source/Runtime/VulkanRHI/Private/VulkanShaders.cpp
+++ b/Engine/Source/Runtime/VulkanRHI/Private/VulkanShaders.cpp
@@ -575,14 +575,25 @@ void FVulkanDescriptorSetsLayoutInfo::FinalizeBindings(const FUniformBufferGathe
 	CompileTypesUsageID();
 	GenerateHash(ImmutableSamplers);
 
-	// Validate no empty sets were made
-	for (int32 Index = 0; Index < RemappingInfo.SetInfos.Num(); ++Index)
+	// If we are consolidating and no uniforms are present in the shader, then strip the empty set data
+	if (bConsolidateAllIntoOneSet)
 	{
-		check(RemappingInfo.SetInfos[Index].Types.Num() > 0);
+		for (int32 Index = 0; Index < RemappingInfo.SetInfos.Num(); ++Index)
+		{
+			if (RemappingInfo.SetInfos[Index].Types.Num() == 0)
+			{
+				RemappingInfo.SetInfos.RemoveAt(Index);
+			}
+		}
+		check(RemappingInfo.SetInfos.Num() <= 1);
 	}
-
-	// Consolidated only has to have one Set
-	check(!bConsolidateAllIntoOneSet || RemappingInfo.SetInfos.Num() == 1);
+	else
+	{
+		for (int32 Index = 0; Index < RemappingInfo.SetInfos.Num(); ++Index)
+		{
+			check(RemappingInfo.SetInfos[Index].Types.Num() > 0);
+		}
+	}	
 }
 
 void FVulkanComputePipelineDescriptorInfo::Initialize(const FDescriptorSetRemappingInfo& InRemappingInfo)
