Vulkan: Fix handling FRHITextureReference. This fixes a crash when importing UDIM textures.

diff --git a/Engine/Source/Runtime/VulkanRHI/Private/VulkanResources.h b/Engine/Source/Runtime/VulkanRHI/Private/VulkanResources.h
--- a/Engine/Source/Runtime/VulkanRHI/Private/VulkanResources.h
+++ b/Engine/Source/Runtime/VulkanRHI/Private/VulkanResources.h
@@ -724,13 +724,19 @@ inline FVulkanTextureBase* GetVulkanTextureFromRHITexture(FRHITexture* Texture)
 	{
 		return NULL;
 	}
-	else if (FRHITexture2D* Tex2D = Texture->GetTexture2D())
+	if(FRHITextureReference* TexRef = Texture->GetTextureReference())
 	{
-		return static_cast<FVulkanTexture2D*>(Tex2D);
+		Texture = TexRef->GetReferencedTexture();
+		if (!Texture)
+		{
+			return NULL;
+		}
 	}
-	else if (FRHITextureReference* TexRef = Texture->GetTextureReference())
+
+
+	if (FRHITexture2D* Tex2D = Texture->GetTexture2D())
 	{
-		return static_cast<FVulkanTextureReference*>(TexRef);
+		return static_cast<FVulkanTexture2D*>(Tex2D);
 	}
 	else if (FRHITexture2DArray* Tex2DArray = Texture->GetTexture2DArray())
 	{
