Fix for depth corruption on Vulkan Android due to Don't Load, Don't store not working as expected when using D24S8 buffers

diff --git a/Engine/Source/Runtime/RHI/Public/RHIResources.h b/Engine/Source/Runtime/RHI/Public/RHIResources.h
--- a/Engine/Source/Runtime/RHI/Public/RHIResources.h
+++ b/Engine/Source/Runtime/RHI/Public/RHIResources.h
@@ -2364,6 +2364,7 @@ enum class EDepthStencilTargetActions : uint8
 	ClearDepthStencil_StoreStencilNotDepth =	RTACTION_MAKE_MASK(Clear_DontStore, Clear_Store),
 	ClearDepthStencil_ResolveDepthNotStencil =	RTACTION_MAKE_MASK(Clear_Resolve, Clear_DontStore),
 	ClearDepthStencil_ResolveStencilNotDepth =	RTACTION_MAKE_MASK(Clear_DontStore, Clear_Resolve),
+	LoadDepthClearStencil_StoreDepthStencil  =  RTACTION_MAKE_MASK(Load_Store, Clear_Store),
 
 	ClearStencilDontLoadDepth_StoreStencilNotDepth = RTACTION_MAKE_MASK(DontLoad_DontStore, Clear_Store),
 
diff --git a/Engine/Source/Runtime/Renderer/Private/DeferredShadingRenderer.cpp b/Engine/Source/Runtime/Renderer/Private/DeferredShadingRenderer.cpp
--- a/Engine/Source/Runtime/Renderer/Private/DeferredShadingRenderer.cpp
+++ b/Engine/Source/Runtime/Renderer/Private/DeferredShadingRenderer.cpp
@@ -2229,8 +2229,9 @@ void FDeferredShadingSceneRenderer::Render(FRHICommandListImmediate& RHICmdList)
 	{
 		// Clear stencil to 0 now that deferred decals are done using what was setup in the base pass
 		// Shadow passes and other users of stencil assume it is cleared to 0 going in
+		// NB: Depth Load/Write necessary for Vulkan Platforms
 		FRHIRenderPassInfo RPInfo(SceneContext.GetSceneDepthSurface(),
-			EDepthStencilTargetActions::ClearStencilDontLoadDepth_StoreStencilNotDepth);
+			EDepthStencilTargetActions::LoadDepthClearStencil_StoreDepthStencil);
 		RPInfo.DepthStencilRenderTarget.ExclusiveDepthStencil = FExclusiveDepthStencil::DepthNop_StencilWrite;
 
 		RHICmdList.TransitionResource(FExclusiveDepthStencil::DepthNop_StencilWrite, SceneContext.GetSceneDepthSurface());
diff --git a/Engine/Source/Runtime/VulkanRHI/Private/Android/VulkanAndroidPlatform.h b/Engine/Source/Runtime/VulkanRHI/Private/Android/VulkanAndroidPlatform.h
--- a/Engine/Source/Runtime/VulkanRHI/Private/Android/VulkanAndroidPlatform.h
+++ b/Engine/Source/Runtime/VulkanRHI/Private/Android/VulkanAndroidPlatform.h
@@ -116,6 +116,9 @@ public:
 
 	static bool RequiresRenderPassResolveAttachments() { return true; }
 
+	//#todo-rco: Detect Mali? Does the platform require depth to be written on stencil clear
+	static bool RequiresDepthWriteOnStencilClear() { return true; }
+
 protected:
 	static void* VulkanLib;
 	static bool bAttemptedLoad;
diff --git a/Engine/Source/Runtime/VulkanRHI/Private/VulkanGenericPlatform.h b/Engine/Source/Runtime/VulkanRHI/Private/VulkanGenericPlatform.h
--- a/Engine/Source/Runtime/VulkanRHI/Private/VulkanGenericPlatform.h
+++ b/Engine/Source/Runtime/VulkanRHI/Private/VulkanGenericPlatform.h
@@ -114,6 +114,9 @@ public:
 	// Does the platform require resolve attachments in its MSAA renderpasses
 	static bool RequiresRenderPassResolveAttachments() { return false; }
 
+	// Does the platform require depth to be written on stencil clear
+	static bool RequiresDepthWriteOnStencilClear() { return false; }
+
 	// Checks if the PSO cache matches the expected vulkan device properties
 	static bool PSOBinaryCacheMatches(FVulkanDevice* Device, const TArray<uint8>& DeviceCache);
 
diff --git a/Engine/Source/Runtime/VulkanRHI/Private/VulkanRenderTarget.cpp b/Engine/Source/Runtime/VulkanRHI/Private/VulkanRenderTarget.cpp
--- a/Engine/Source/Runtime/VulkanRHI/Private/VulkanRenderTarget.cpp
+++ b/Engine/Source/Runtime/VulkanRHI/Private/VulkanRenderTarget.cpp
@@ -182,6 +182,14 @@ void FTransitionAndLayoutManager::BeginEmulatedRenderPass(FVulkanCommandListCont
 		FVulkanSurface& Surface = FVulkanTextureBase::Cast(DSTexture)->Surface;
 		VkImageLayout& DSLayout = Layouts.FindOrAdd(Surface.Image);
 		FExclusiveDepthStencil RequestedDSAccess = RenderTargetsInfo.DepthStencilRenderTarget.GetDepthStencilAccess();
+		
+		if (FVulkanPlatform::RequiresDepthWriteOnStencilClear() && 
+			RenderTargetsInfo.DepthStencilRenderTarget.DepthStoreAction == ERenderTargetStoreAction::EStore &&
+			RenderTargetsInfo.DepthStencilRenderTarget.GetStencilStoreAction() == ERenderTargetStoreAction::EStore)
+		{
+			RequestedDSAccess = FExclusiveDepthStencil::DepthWrite_StencilWrite;
+		}
+
 		VkImageLayout FinalLayout = VulkanRHI::GetDepthStencilLayout(RequestedDSAccess, InDevice);
 
 		// Check if we need to transition the depth stencil texture(s) based on the current layout and the requested access mode for the render target
@@ -2206,7 +2214,14 @@ FVulkanRenderTargetLayout::FVulkanRenderTargetLayout(FVulkanDevice& InDevice, co
 			ensure(CurrDesc.stencilStoreOp == VK_ATTACHMENT_STORE_OP_DONT_CARE);
 		}
 
-		DepthStencilLayout = VulkanRHI::GetDepthStencilLayout(RPInfo.DepthStencilRenderTarget.ExclusiveDepthStencil, InDevice);
+		FExclusiveDepthStencil ExclusiveDepthStencil = RPInfo.DepthStencilRenderTarget.ExclusiveDepthStencil;
+		if (FVulkanPlatform::RequiresDepthWriteOnStencilClear() &&
+			RPInfo.DepthStencilRenderTarget.Action == EDepthStencilTargetActions::LoadDepthClearStencil_StoreDepthStencil)
+		{
+			ExclusiveDepthStencil = FExclusiveDepthStencil::DepthWrite_StencilWrite;
+		}
+
+		DepthStencilLayout = VulkanRHI::GetDepthStencilLayout(ExclusiveDepthStencil, InDevice);
 		// If the initial != final we need to change the FullHashInfo and use FinalLayout
 		CurrDesc.initialLayout = DepthStencilLayout;
 		CurrDesc.finalLayout = DepthStencilLayout;
