Vulkan: allow forcing the pacers even when vsync is off (UE-91532).

diff --git a/Engine/Source/Runtime/VulkanRHI/Private/VulkanSwapChain.cpp b/Engine/Source/Runtime/VulkanRHI/Private/VulkanSwapChain.cpp
--- a/Engine/Source/Runtime/VulkanRHI/Private/VulkanSwapChain.cpp
+++ b/Engine/Source/Runtime/VulkanRHI/Private/VulkanSwapChain.cpp
@@ -58,6 +58,14 @@ static FAutoConsoleVariableRef CVarVulkanExtensionFramePacer(
 	ECVF_RenderThreadSafe
 );
 
+int32 GVulkanForcePacingWithoutVSync = 0;
+static FAutoConsoleVariableRef CVarVulkanForcePacingWithoutVSync(
+	TEXT("r.Vulkan.ForcePacingWithoutVSync"),
+	GVulkanForcePacingWithoutVSync,
+	TEXT("Whether to CPU pacers remain enabled even if VSync is off"),
+	ECVF_RenderThreadSafe
+);
+
 static int32 GPrintVulkanVsyncDebug = 0;
 
 #if !(UE_BUILD_SHIPPING)
@@ -454,7 +462,8 @@ FVulkanSwapChain::FVulkanSwapChain(VkInstance InInstance, FVulkanDevice& InDevic
 	static bool bPrintSwapchainCreationInfo = true;
 	if (bPrintSwapchainCreationInfo)
 	{
-		UE_LOG(LogVulkanRHI, Log, TEXT("Creating new VK swapchain with format %d, color space %d, num images %d"), static_cast<uint32>(SwapChainInfo.imageFormat), static_cast<uint32>(SwapChainInfo.imageColorSpace), static_cast<uint32>(SwapChainInfo.minImageCount));
+		UE_LOG(LogVulkanRHI, Log, TEXT("Creating new VK swapchain with present mode %d, format %d, color space %d, num images %d"), 
+			static_cast<uint32>(SwapChainInfo.presentMode), static_cast<uint32>(SwapChainInfo.imageFormat), static_cast<uint32>(SwapChainInfo.imageColorSpace), static_cast<uint32>(SwapChainInfo.minImageCount));
 #if WITH_EDITOR
 		bPrintSwapchainCreationInfo = false;
 #endif
@@ -826,7 +835,7 @@ void FGDTimingFramePacer::PollPastFrameInfo()
 void FVulkanSwapChain::RenderThreadPacing()
 {
 	check(IsInRenderingThread());
-	const int32 SyncInterval = LockToVsync ? RHIGetSyncInterval() : 0;
+	const int32 SyncInterval = (LockToVsync || GVulkanForcePacingWithoutVSync) ? RHIGetSyncInterval() : 0;
 
 	//very naive CPU side frame pacer.
 	if (GVulkanCPURenderThreadFramePacer && SyncInterval > 0)
@@ -894,7 +903,7 @@ FVulkanSwapChain::EStatus FVulkanSwapChain::Present(FVulkanQueue* GfxQueue, FVul
 	Info.pSwapchains = &SwapChain;
 	Info.pImageIndices = (uint32*)&CurrentImageIndex;
 
-	const int32 SyncInterval = LockToVsync ? RHIGetSyncInterval() : 0;
+	const int32 SyncInterval = (LockToVsync || GVulkanForcePacingWithoutVSync) ? RHIGetSyncInterval() : 0;
 	ensureMsgf(SyncInterval <= 3 && SyncInterval >= 0, TEXT("Unsupported sync interval: %i"), SyncInterval);
 
 #if VULKAN_SUPPORTS_GOOGLE_DISPLAY_TIMING
