Fix assert FVulkanGPUTiming::Release assert when timestamps not supported

Running on Linux Intel yocto distro which doesn't support timestamps, and
FVulkanGPUTiming::Initialize() checks GIsSupported when creating timing query pool, but
GIsSupported() does not check when releasing pool so Pool is null and check fails.

diff --git a/Engine/Source/Runtime/VulkanRHI/Private/VulkanUtil.cpp b/Engine/Source/Runtime/VulkanRHI/Private/VulkanUtil.cpp
--- a/Engine/Source/Runtime/VulkanRHI/Private/VulkanUtil.cpp
+++ b/Engine/Source/Runtime/VulkanRHI/Private/VulkanUtil.cpp
@@ -158,7 +158,7 @@ void FVulkanGPUTiming::Initialize()
  */
 void FVulkanGPUTiming::Release()
 {
-	if (FVulkanPlatform::SupportsTimestampRenderQueries())
+	if (FVulkanPlatform::SupportsTimestampRenderQueries() && GIsSupported)
 	{
 		check(Pool);
 		Device->GetStagingManager().ReleaseBuffer(nullptr, Pool->ResultsBuffer);
