Fixed a few problems with buffer views on Vulkan:
* SRVs based on volatile buffers had their size reset to the size of the source buffer, which could make them larger than the user intended, and larger than VkPhysicalDeviceLimits.maxTexelBufferElements
* the way FVulkanBufferView::Create() was computing the number of elements for checking against maxTexelBufferElements was wrong
* when a view larger than maxTexelBufferElements was requested, we were mapping VK_WHOLE_SIZE, which is incorrect, because that tries to map to the end of the buffer, which we've just determined to be too large. We need to convert maxTexelBufferElements to bytes and use that as the view size.

#rnx

diff --git a/Engine/Source/Runtime/VulkanRHI/Private/VulkanRHI.cpp b/Engine/Source/Runtime/VulkanRHI/Private/VulkanRHI.cpp
--- a/Engine/Source/Runtime/VulkanRHI/Private/VulkanRHI.cpp
+++ b/Engine/Source/Runtime/VulkanRHI/Private/VulkanRHI.cpp
@@ -1463,8 +1463,9 @@ void FVulkanBufferView::Create(VkFormat Format, FVulkanResourceMultiBuffer* Buff
 
 	//#todo-rco: Revisit this if buffer views become VK_BUFFER_USAGE_STORAGE_BUFFER_BIT instead of VK_BUFFER_USAGE_UNIFORM_TEXEL_BUFFER_BIT
 	const VkPhysicalDeviceLimits& Limits = Device->GetLimits();
-	//#todo-rco: 4 for rgba, do we need 1 for r, 2 for rg and 3 for rgb?
-	ViewInfo.range = ((Size / (GetNumBitsPerPixel(Format) / 4)) < Limits.maxTexelBufferElements) ? Size : VK_WHOLE_SIZE;
+	const uint64 MaxSize = (uint64)Limits.maxTexelBufferElements * GetNumBitsPerPixel(Format) / 8;
+	ViewInfo.range = FMath::Min<uint64>(Size, MaxSize);
+	// TODO: add a check() for exceeding MaxSize, to catch code which blindly makes views without checking the platform limits.
 
 	Flags = Buffer->GetBufferUsageFlags() & (VK_BUFFER_USAGE_UNIFORM_TEXEL_BUFFER_BIT | VK_BUFFER_USAGE_STORAGE_TEXEL_BUFFER_BIT);
 	check(Flags);
diff --git a/Engine/Source/Runtime/VulkanRHI/Private/VulkanUAV.cpp b/Engine/Source/Runtime/VulkanRHI/Private/VulkanUAV.cpp
--- a/Engine/Source/Runtime/VulkanRHI/Private/VulkanUAV.cpp
+++ b/Engine/Source/Runtime/VulkanRHI/Private/VulkanUAV.cpp
@@ -79,8 +79,10 @@ void FVulkanShaderResourceView::UpdateView()
 		{
 			VkBuffer SourceVolatileBufferHandle = SourceBuffer->GetHandle();
 
-			//Inherit Size from volatile buffer, as it might change
-			Size = SourceBuffer->GetVolatileLockSize();
+			// If the volatile buffer shrinks, make sure our size doesn't exceed the new limit.
+			uint32 AvailableSize = SourceBuffer->GetVolatileLockSize();
+			AvailableSize = Offset < AvailableSize ? AvailableSize - Offset : 0;
+			Size = FMath::Min(Size, AvailableSize);
 
 			// We might end up with the same BufferView, so do not recreate in that case
 			if (!BufferViews[0] 
