Add new RHI method RHICalcTexture2DArrayPlatformSize() for PS4, Switch, D3D11, D3D12, OpenGL and Vulkan.

Added support for 3d textures with RHICopyTexture in the Switch and Vulkan implementation.

Fixed D3D12 and PS4 issues in the RHICopyTexture implementation when using compressed formats.

Fixed missing RB swizzle for 3d textures and texture 2d arrays in OpenGL.

Fixed 3d texture bulkdata mip layout when cooking PS4.
Fixed some issues around mip data bulk initialization on PS4.

diff --git a/Engine/Source/Runtime/D3D12RHI/Private/D3D12RHIPrivate.h b/Engine/Source/Runtime/D3D12RHI/Private/D3D12RHIPrivate.h
--- a/Engine/Source/Runtime/D3D12RHI/Private/D3D12RHIPrivate.h
+++ b/Engine/Source/Runtime/D3D12RHI/Private/D3D12RHIPrivate.h
@@ -298,6 +298,7 @@ public:
 	virtual void RHIUpdateShaderResourceView(FRHIShaderResourceView* SRV, FRHIBuffer* Buffer, uint32 Stride, uint8 Format) final override;
 	virtual void RHIUpdateShaderResourceView(FRHIShaderResourceView* SRV, FRHIBuffer* Buffer) final override;
 	virtual uint64 RHICalcTexture2DPlatformSize(uint32 SizeX, uint32 SizeY, uint8 Format, uint32 NumMips, uint32 NumSamples, ETextureCreateFlags InFlags, const FRHIResourceCreateInfo& CreateInfo, uint32& OutAlign) override;
+	virtual uint64 RHICalcTexture2DArrayPlatformSize(uint32 SizeX, uint32 SizeY, uint32 ArraySize, uint8 Format, uint32 NumMips, uint32 NumSamples, ETextureCreateFlags Flags, const FRHIResourceCreateInfo& CreateInfo, uint32& OutAlign) final override;
 	virtual uint64 RHICalcTexture3DPlatformSize(uint32 SizeX, uint32 SizeY, uint32 SizeZ, uint8 Format, uint32 NumMips, ETextureCreateFlags InFlags, const FRHIResourceCreateInfo& CreateInfo, uint32& OutAlign) final override;
 	virtual uint64 RHICalcTextureCubePlatformSize(uint32 Size, uint8 Format, uint32 NumMips, ETextureCreateFlags InFlags, const FRHIResourceCreateInfo& CreateInfo, uint32& OutAlign) final override;
 	virtual uint64 RHIGetMinimumAlignmentForBufferBackedSRV(EPixelFormat Format) final override;
diff --git a/Engine/Source/Runtime/D3D12RHI/Private/D3D12Texture.cpp b/Engine/Source/Runtime/D3D12RHI/Private/D3D12Texture.cpp
--- a/Engine/Source/Runtime/D3D12RHI/Private/D3D12Texture.cpp
+++ b/Engine/Source/Runtime/D3D12RHI/Private/D3D12Texture.cpp
@@ -492,6 +492,24 @@ uint64 FD3D12DynamicRHI::RHICalcTexture2DPlatformSize(uint32 SizeX, uint32 SizeY
 	return AllocationInfo.SizeInBytes;
 }
 
+uint64 FD3D12DynamicRHI::RHICalcTexture2DArrayPlatformSize(uint32 SizeX, uint32 SizeY, uint32 ArraySize, uint8 Format, uint32 NumMips, uint32 NumSamples, ETextureCreateFlags Flags, const FRHIResourceCreateInfo& CreateInfo, uint32& OutAlign)
+{
+	D3D12_RESOURCE_DESC Desc = {};
+	Desc.Dimension = D3D12_RESOURCE_DIMENSION_TEXTURE2D;
+	Desc.Format = (DXGI_FORMAT)GPixelFormats[Format].PlatformFormat;
+	Desc.Height = SizeY;
+	Desc.DepthOrArraySize = ArraySize;
+	Desc.Layout = D3D12_TEXTURE_LAYOUT_UNKNOWN;
+	Desc.MipLevels = NumMips;
+	Desc.SampleDesc.Count = NumSamples;
+	Desc.Width = SizeX;
+
+	const D3D12_RESOURCE_ALLOCATION_INFO AllocationInfo = GetAdapter().GetD3DDevice()->GetResourceAllocationInfo(0, 1, &Desc);
+	OutAlign = static_cast<uint32>(AllocationInfo.Alignment);
+
+	return AllocationInfo.SizeInBytes;
+}
+
 uint64 FD3D12DynamicRHI::RHICalcTexture3DPlatformSize(uint32 SizeX, uint32 SizeY, uint32 SizeZ, uint8 Format, uint32 NumMips, ETextureCreateFlags Flags, const FRHIResourceCreateInfo& CreateInfo, uint32& OutAlign)
 {
 	D3D12_RESOURCE_DESC Desc = {};
@@ -3112,7 +3130,7 @@ void FD3D12CommandContext::RHICopyTexture(FRHITexture* SourceTextureRHI, FRHITex
 		const FIntVector CopySize = CopyInfo.Size == FIntVector::ZeroValue ? SourceTextureRHI->GetSizeXYZ() : CopyInfo.Size;
 
 		// Copy sub texture regions
-		CD3DX12_BOX SourceBoxD3D(
+		const CD3DX12_BOX SourceBoxD3D(
 			CopyInfo.SourcePosition.X,
 			CopyInfo.SourcePosition.Y,
 			CopyInfo.SourcePosition.Z,
@@ -3129,6 +3147,8 @@ void FD3D12CommandContext::RHICopyTexture(FRHITexture* SourceTextureRHI, FRHITex
 		Dst.pResource = DestTexture->GetResource()->GetResource();
 		Dst.Type = bReadback ? D3D12_TEXTURE_COPY_TYPE_PLACED_FOOTPRINT : D3D12_TEXTURE_COPY_TYPE_SUBRESOURCE_INDEX;
 
+		const FPixelFormatInfo& PixelFormatInfo = GPixelFormats[DestTextureRHI->GetFormat()];
+
 		D3D12_RESOURCE_DESC DstDesc = {};
 		FIntVector TextureSize = DestTextureRHI->GetSizeXYZ();
 		DstDesc.Dimension = DestTextureRHI->GetTexture3D() ? D3D12_RESOURCE_DIMENSION_TEXTURE3D : D3D12_RESOURCE_DIMENSION_TEXTURE2D; 
@@ -3136,7 +3156,7 @@ void FD3D12CommandContext::RHICopyTexture(FRHITexture* SourceTextureRHI, FRHITex
 		DstDesc.Height = TextureSize.Y;
 		DstDesc.DepthOrArraySize = TextureSize.Z;
 		DstDesc.MipLevels = DestTextureRHI->GetNumMips();
-		DstDesc.Format = (DXGI_FORMAT)GPixelFormats[DestTextureRHI->GetFormat()].PlatformFormat;
+		DstDesc.Format = (DXGI_FORMAT)PixelFormatInfo.PlatformFormat;
 		DstDesc.SampleDesc.Count = DestTextureRHI->GetNumSamples();
 
 		for (uint32 SliceIndex = 0; SliceIndex < CopyInfo.NumSlices; ++SliceIndex)
@@ -3149,13 +3169,24 @@ void FD3D12CommandContext::RHICopyTexture(FRHITexture* SourceTextureRHI, FRHITex
 				uint32 SourceMipIndex = CopyInfo.SourceMipIndex + MipIndex;
 				uint32 DestMipIndex   = CopyInfo.DestMipIndex   + MipIndex;
 
-				uint32 SizeX = FMath::Max(CopySize.X >> MipIndex, 1);
-				uint32 SizeY = FMath::Max(CopySize.Y >> MipIndex, 1);
-				uint32 SizeZ = FMath::Max(CopySize.Z >> MipIndex, 1);
-
-				SourceBoxD3D.right  = CopyInfo.SourcePosition.X + SizeX;
-				SourceBoxD3D.bottom = CopyInfo.SourcePosition.Y + SizeY;
-				SourceBoxD3D.back   = CopyInfo.SourcePosition.Z + SizeZ;
+				CD3DX12_BOX MipSourceBoxD3D(
+					SourceBoxD3D.left  >> MipIndex,
+					SourceBoxD3D.top   >> MipIndex,
+					SourceBoxD3D.front >> MipIndex,
+					// Align to block size to pad the copy when processing the last surface texels.
+					// This will give inconsistent results otherwise between different RHI.
+					AlignArbitrary<uint32>(FMath::Max<uint32>(SourceBoxD3D.right  >> MipIndex, 1), PixelFormatInfo.BlockSizeX),
+					AlignArbitrary<uint32>(FMath::Max<uint32>(SourceBoxD3D.bottom >> MipIndex, 1), PixelFormatInfo.BlockSizeY),
+					AlignArbitrary<uint32>(FMath::Max<uint32>(SourceBoxD3D.back   >> MipIndex, 1), PixelFormatInfo.BlockSizeZ)
+					);
+
+				const uint32 DestX = CopyInfo.DestPosition.X >> MipIndex;
+				const uint32 DestY = CopyInfo.DestPosition.Y >> MipIndex;
+				const uint32 DestZ = CopyInfo.DestPosition.Z >> MipIndex;
+
+				// RHICopyTexture is allowed to copy mip regions only if are aligned on the block size to prevent unexpected / inconsistent results.
+				ensure(MipSourceBoxD3D.left % PixelFormatInfo.BlockSizeX == 0 && MipSourceBoxD3D.top % PixelFormatInfo.BlockSizeY == 0 && MipSourceBoxD3D.front % PixelFormatInfo.BlockSizeZ == 0);
+				ensure(DestX % PixelFormatInfo.BlockSizeX == 0 && DestY % PixelFormatInfo.BlockSizeY == 0 && DestZ % PixelFormatInfo.BlockSizeZ == 0);
 
 				Src.SubresourceIndex = CalcSubresource(SourceMipIndex, SourceSliceIndex, SourceTextureRHI->GetNumMips());
 				Dst.SubresourceIndex = CalcSubresource(DestMipIndex, DestSliceIndex, DestTextureRHI->GetNumMips());
@@ -3167,11 +3198,9 @@ void FD3D12CommandContext::RHICopyTexture(FRHITexture* SourceTextureRHI, FRHITex
 
 				CommandListHandle->CopyTextureRegion(
 					&Dst, 
-					CopyInfo.DestPosition.X,
-					CopyInfo.DestPosition.Y,
-					CopyInfo.DestPosition.Z,
+					DestX, DestY, DestZ,
 					&Src,
-					&SourceBoxD3D
+					&MipSourceBoxD3D
 				);
 			}
 		}
diff --git a/Engine/Source/Runtime/OpenGLDrv/Private/OpenGLTexture.cpp b/Engine/Source/Runtime/OpenGLDrv/Private/OpenGLTexture.cpp
--- a/Engine/Source/Runtime/OpenGLDrv/Private/OpenGLTexture.cpp
+++ b/Engine/Source/Runtime/OpenGLDrv/Private/OpenGLTexture.cpp
@@ -250,6 +250,12 @@ uint64 FOpenGLDynamicRHI::RHICalcTexture2DPlatformSize(uint32 SizeX, uint32 Size
 	return CalcTextureSize(SizeX, SizeY, (EPixelFormat)Format, NumMips);
 }
 
+uint64 FOpenGLDynamicRHI::RHICalcTexture2DArrayPlatformSize(uint32 SizeX, uint32 SizeY, uint32 ArraySize, uint8 Format, uint32 NumMips, uint32 NumSamples, ETextureCreateFlags Flags, const FRHIResourceCreateInfo& CreateInfo, uint32& OutAlign)
+{
+	OutAlign = 0;
+	return CalcTextureSize(SizeX, SizeY, (EPixelFormat)Format, NumMips) * ArraySize;
+}
+
 uint64 FOpenGLDynamicRHI::RHICalcTexture3DPlatformSize(uint32 SizeX, uint32 SizeY, uint32 SizeZ, uint8 Format, uint32 NumMips, ETextureCreateFlags Flags, const FRHIResourceCreateInfo& CreateInfo, uint32& OutAlign)
 {
 	OutAlign = 0;
@@ -1600,6 +1606,12 @@ FTexture2DArrayRHIRef FOpenGLDynamicRHI::RHICreateTexture2DArray(uint32 SizeX,ui
 		UE_LOG(LogRHI, Fatal,TEXT("Texture format '%s' not supported."), FormatInfo.Name);
 	}
 
+	if (GLFormat.bBGRA && !(Flags & TexCreate_RenderTargetable))
+	{
+		glTexParameteri(Target, GL_TEXTURE_SWIZZLE_R, GL_BLUE);
+		glTexParameteri(Target, GL_TEXTURE_SWIZZLE_B, GL_RED);
+	}
+
 	// Make sure PBO is disabled
 	CachedBindPixelUnpackBuffer(ContextState, 0);
 
@@ -1725,12 +1737,18 @@ FTexture3DRHIRef FOpenGLDynamicRHI::RHICreateTexture3D(uint32 SizeX,uint32 SizeY
 	const bool bSRGB = (Flags&TexCreate_SRGB) != 0;
 	const FOpenGLTextureFormat& GLFormat = GOpenGLTextureFormats[Format];
 	const FPixelFormatInfo& FormatInfo = GPixelFormats[Format];
-
 	if (GLFormat.InternalFormat[bSRGB] == GL_NONE)
 	{
-		UE_LOG(LogRHI, Fatal,TEXT("Texture format '%s' not supported."), FormatInfo.Name);
+		UE_LOG(LogRHI, Fatal, TEXT("Texture format '%s' not supported."), FormatInfo.Name);
+	}
+
+	if (GLFormat.bBGRA && !(Flags & TexCreate_RenderTargetable))
+	{
+		glTexParameteri(Target, GL_TEXTURE_SWIZZLE_R, GL_BLUE);
+		glTexParameteri(Target, GL_TEXTURE_SWIZZLE_B, GL_RED);
 	}
 
+
 	// Make sure PBO is disabled
 	CachedBindPixelUnpackBuffer(ContextState,0);
 
diff --git a/Engine/Source/Runtime/OpenGLDrv/Public/OpenGLDrv.h b/Engine/Source/Runtime/OpenGLDrv/Public/OpenGLDrv.h
--- a/Engine/Source/Runtime/OpenGLDrv/Public/OpenGLDrv.h
+++ b/Engine/Source/Runtime/OpenGLDrv/Public/OpenGLDrv.h
@@ -366,6 +366,7 @@ public:
 	virtual void RHIUpdateShaderResourceView(FRHIShaderResourceView* SRV, FRHIVertexBuffer* VertexBuffer, uint32 Stride, uint8 Format) final override;
 	virtual void RHIUpdateShaderResourceView(FRHIShaderResourceView* SRV, FRHIIndexBuffer* IndexBuffer) final override;
 	virtual uint64 RHICalcTexture2DPlatformSize(uint32 SizeX, uint32 SizeY, uint8 Format, uint32 NumMips, uint32 NumSamples, ETextureCreateFlags Flags, const FRHIResourceCreateInfo& CreateInfo, uint32& OutAlign) final override;
+	virtual uint64 RHICalcTexture2DArrayPlatformSize(uint32 SizeX, uint32 SizeY, uint32 ArraySize, uint8 Format, uint32 NumMips, uint32 NumSamples, ETextureCreateFlags Flags, const FRHIResourceCreateInfo& CreateInfo, uint32& OutAlign) final override;
 	virtual uint64 RHICalcTexture3DPlatformSize(uint32 SizeX, uint32 SizeY, uint32 SizeZ, uint8 Format, uint32 NumMips, ETextureCreateFlags Flags, const FRHIResourceCreateInfo& CreateInfo, uint32& OutAlign) final override;
 	virtual uint64 RHICalcTextureCubePlatformSize(uint32 Size, uint8 Format, uint32 NumMips, ETextureCreateFlags Flags, const FRHIResourceCreateInfo& CreateInfo, uint32& OutAlign) final override;
 	virtual void RHIGetTextureMemoryStats(FTextureMemoryStats& OutStats) final override;
diff --git a/Engine/Source/Runtime/RHI/Private/DynamicRHI.cpp b/Engine/Source/Runtime/RHI/Private/DynamicRHI.cpp
--- a/Engine/Source/Runtime/RHI/Private/DynamicRHI.cpp
+++ b/Engine/Source/Runtime/RHI/Private/DynamicRHI.cpp
@@ -449,6 +449,12 @@ uint64 FDynamicRHI::RHIGetMinimumAlignmentForBufferBackedSRV(EPixelFormat Format
 	return 1;
 }
 
+uint64 FDynamicRHI::RHICalcTexture2DArrayPlatformSize(uint32 SizeX, uint32 SizeY, uint32 ArraySize, uint8 Format, uint32 NumMips, uint32 NumSamples, ETextureCreateFlags Flags, const FRHIResourceCreateInfo& CreateInfo, uint32& OutAlign)
+{
+	// ensureMsgf(false, TEXT("RHICalcTexture2DArrayPlatformSize isn't implemented for the current RHI"));
+	return ArraySize * RHICalcTexture2DPlatformSize(SizeX, SizeY, Format, NumMips, NumSamples, Flags, CreateInfo, OutAlign);
+}
+
 uint64 FDynamicRHI::RHICalcVMTexture2DPlatformSize(uint32 Mip0Width, uint32 Mip0Height, uint8 Format, uint32 NumMips, uint32 FirstMipIdx, uint32 NumSamples, ETextureCreateFlags Flags, uint32& OutAlign)
 {
 	UE_LOG(LogRHI, Fatal, TEXT("RHICalcVMTexture2DPlatformSize isn't implemented for the current RHI"));
diff --git a/Engine/Source/Runtime/RHI/Public/DynamicRHI.h b/Engine/Source/Runtime/RHI/Public/DynamicRHI.h
--- a/Engine/Source/Runtime/RHI/Public/DynamicRHI.h
+++ b/Engine/Source/Runtime/RHI/Public/DynamicRHI.h
@@ -611,6 +611,21 @@ public:
 	// FlushType: Thread safe
 	virtual uint64 RHICalcTexture2DPlatformSize(uint32 SizeX, uint32 SizeY, uint8 Format, uint32 NumMips, uint32 NumSamples, ETextureCreateFlags Flags, const FRHIResourceCreateInfo& CreateInfo, uint32& OutAlign) = 0;
 
+	/**
+	* Computes the total size of a 2D array texture with the specified parameters.
+	*
+	* @param SizeX - width of the texture to compute
+	* @param SizeY - height of the texture to compute
+	* @param ArraySize - number of textures in the array.
+	* @param Format - EPixelFormat texture format
+	* @param NumMips - number of mips to compute or 0 for full mip pyramid
+	* @param NumSamples - number of MSAA samples, usually 1
+	* @param Flags - ETextureCreateFlags creation flags
+	* @param OutAlign - Alignment required for this texture.  Output parameter.
+	*/
+	// FlushType: Thread safe
+	virtual uint64 RHICalcTexture2DArrayPlatformSize(uint32 SizeX, uint32 SizeY, uint32 ArraySize, uint8 Format, uint32 NumMips, uint32 NumSamples, ETextureCreateFlags Flags, const FRHIResourceCreateInfo& CreateInfo, uint32& OutAlign);
+	
 	/**
 	* Computes the total size of a virtual memory (VM) based 2D texture with the specified parameters.
 	*
@@ -1481,6 +1496,11 @@ FORCEINLINE uint64 RHICalcTexture2DPlatformSize(uint32 SizeX, uint32 SizeY, uint
 	return GDynamicRHI->RHICalcTexture2DPlatformSize(SizeX, SizeY, Format, NumMips, NumSamples, Flags, CreateInfo, OutAlign);
 }
 
+FORCEINLINE uint64 RHICalcTexture2DArrayPlatformSize(uint32 SizeX, uint32 SizeY, uint32 ArraySize, uint8 Format, uint32 NumMips, uint32 NumSamples, ETextureCreateFlags Flags, const FRHIResourceCreateInfo& CreateInfo, uint32& OutAlign)
+{
+	return GDynamicRHI->RHICalcTexture2DArrayPlatformSize(SizeX, SizeY, ArraySize, Format, NumMips, NumSamples, Flags, CreateInfo, OutAlign);
+}
+
 FORCEINLINE uint64 RHICalcVMTexture2DPlatformSize(uint32 Mip0Width, uint32 Mip0Height, uint8 Format, uint32 NumMips, uint32 FirstMipIdx, uint32 NumSamples, ETextureCreateFlags Flags, uint32& OutAlign)
 {
 	return GDynamicRHI->RHICalcVMTexture2DPlatformSize(Mip0Width, Mip0Height, Format, NumMips, FirstMipIdx, NumSamples, Flags, OutAlign);
diff --git a/Engine/Source/Runtime/RHI/Public/RHICommandList.h b/Engine/Source/Runtime/RHI/Public/RHICommandList.h
--- a/Engine/Source/Runtime/RHI/Public/RHICommandList.h
+++ b/Engine/Source/Runtime/RHI/Public/RHICommandList.h
@@ -3236,6 +3236,7 @@ public:
 				++PerMipInfo.DestMipIndex;
 				PerMipInfo.Size.X = FMath::Max(1, PerMipInfo.Size.X / 2);
 				PerMipInfo.Size.Y = FMath::Max(1, PerMipInfo.Size.Y / 2);
+				PerMipInfo.Size.Z = FMath::Max(1, PerMipInfo.Size.Z / 2);
 			}
 		}
 	}
diff --git a/Engine/Source/Runtime/VulkanRHI/Private/VulkanTexture.cpp b/Engine/Source/Runtime/VulkanRHI/Private/VulkanTexture.cpp
--- a/Engine/Source/Runtime/VulkanRHI/Private/VulkanTexture.cpp
+++ b/Engine/Source/Runtime/VulkanRHI/Private/VulkanTexture.cpp
@@ -2832,8 +2832,6 @@ static VkMemoryRequirements FindOrCalculateTexturePlatformSize(FVulkanDevice* De
 	return MemReq;
 }
 
-
-
 uint64 FVulkanDynamicRHI::RHICalcTexture2DPlatformSize(uint32 SizeX, uint32 SizeY, uint8 Format, uint32 NumMips, uint32 NumSamples, ETextureCreateFlags Flags, const FRHIResourceCreateInfo& CreateInfo, uint32& OutAlign)
 {
 	const VkMemoryRequirements MemReq = FindOrCalculateTexturePlatformSize(Device, VK_IMAGE_VIEW_TYPE_2D, SizeX, SizeY, 1, Format, NumMips, NumSamples, Flags);
@@ -2841,6 +2839,13 @@ uint64 FVulkanDynamicRHI::RHICalcTexture2DPlatformSize(uint32 SizeX, uint32 Size
 	return MemReq.size;
 }
 
+uint64 FVulkanDynamicRHI::RHICalcTexture2DArrayPlatformSize(uint32 SizeX, uint32 SizeY, uint32 ArraySize, uint8 Format, uint32 NumMips, uint32 NumSamples, ETextureCreateFlags Flags, const FRHIResourceCreateInfo& CreateInfo, uint32& OutAlign)
+{
+	const VkMemoryRequirements MemReq = FindOrCalculateTexturePlatformSize(Device, VK_IMAGE_VIEW_TYPE_2D_ARRAY, SizeX, SizeY, ArraySize, Format, NumMips, NumSamples, Flags);
+	OutAlign = MemReq.alignment;
+	return MemReq.size;
+}
+
 uint64 FVulkanDynamicRHI::RHICalcTexture3DPlatformSize(uint32 SizeX, uint32 SizeY, uint32 SizeZ, uint8 Format, uint32 NumMips, ETextureCreateFlags Flags, const FRHIResourceCreateInfo& CreateInfo, uint32& OutAlign)
 {
 	const VkMemoryRequirements MemReq = FindOrCalculateTexturePlatformSize(Device, VK_IMAGE_VIEW_TYPE_3D, SizeX, SizeY, SizeZ, Format, NumMips, 1, Flags);
@@ -2888,6 +2893,7 @@ void FVulkanCommandListContext::RHICopyTexture(FRHITexture* SourceTexture, FRHIT
 	check(InCmdBuffer->IsOutsideRenderPass());
 	VkCommandBuffer CmdBuffer = InCmdBuffer->GetHandle();
 
+	const FPixelFormatInfo& PixelFormatInfo = GPixelFormats[DestTexture->GetFormat()];
 
 	check((SrcSurface.UEFlags & TexCreate_CPUReadback) == 0);
 	if((DstSurface.UEFlags & TexCreate_CPUReadback) == TexCreate_CPUReadback)
@@ -2899,7 +2905,8 @@ void FVulkanCommandListContext::RHICopyTexture(FRHITexture* SourceTexture, FRHIT
 			ensure(SrcSurface.Width <= DstSurface.Width && SrcSurface.Height <= DstSurface.Height);
 			Size.X = FMath::Max(1u, SrcSurface.Width >> CopyInfo.SourceMipIndex);
 			Size.Y = FMath::Max(1u, SrcSurface.Height >> CopyInfo.SourceMipIndex);
-		}		
+			Size.Z = FMath::Max(1u, SrcSurface.Depth >> CopyInfo.SourceMipIndex);
+		}
 		VkBufferImageCopy CopyRegion[MAX_TEXTURE_MIP_COUNT];
 		FMemory::Memzero(CopyRegion);
 
@@ -2918,13 +2925,13 @@ void FVulkanCommandListContext::RHICopyTexture(FRHITexture* SourceTexture, FRHIT
 			CopyRegion[Index].imageSubresource.layerCount = 1;
 			CopyRegion[Index].imageExtent.width = Size.X;
 			CopyRegion[Index].imageExtent.height = Size.Y;
-			CopyRegion[Index].imageExtent.depth = 1;
+			CopyRegion[Index].imageExtent.depth = Size.Z;
 
 			Size.X = FMath::Max(1, Size.X / 2);
 			Size.Y = FMath::Max(1, Size.Y / 2);
+			Size.Z = FMath::Max(1, Size.Z / 2);
 		}
 
-
 		VulkanRHI::vkCmdCopyImageToBuffer(CmdBuffer, SrcSurface.Image, VK_IMAGE_LAYOUT_TRANSFER_SRC_OPTIMAL, CpuReadbackBuffer->Buffer, CopyInfo.NumMips, &CopyRegion[0]);
 
 		FVulkanPipelineBarrier BarrierMemory;
@@ -2942,35 +2949,37 @@ void FVulkanCommandListContext::RHICopyTexture(FRHITexture* SourceTexture, FRHIT
 		VkImageLayout DstLayout = LayoutManager.FindLayoutChecked(DstSurface.Image);
 		ensureMsgf(DstLayout == VK_IMAGE_LAYOUT_TRANSFER_DST_OPTIMAL, TEXT("Expected destination texture to be in VK_IMAGE_LAYOUT_TRANSFER_DST_OPTIMAL, actual layout is %d"), DstLayout);
 
-
 		VkImageCopy Region;
 		FMemory::Memzero(Region);
 		if (CopyInfo.Size == FIntVector::ZeroValue)
 		{
 			// Copy whole texture when zero vector is specified for region size
 			ensure(SrcSurface.Width <= DstSurface.Width && SrcSurface.Height <= DstSurface.Height);
-			Region.extent.width = FMath::Max(1u, SrcSurface.Width >> CopyInfo.SourceMipIndex);
-			Region.extent.height = FMath::Max(1u, SrcSurface.Height >> CopyInfo.SourceMipIndex);
+			Region.extent.width = FMath::Max<uint32>(PixelFormatInfo.BlockSizeX, SrcSurface.Width >> CopyInfo.SourceMipIndex);
+			Region.extent.height = FMath::Max<uint32>(PixelFormatInfo.BlockSizeY, SrcSurface.Height >> CopyInfo.SourceMipIndex);
+			Region.extent.depth = FMath::Max<uint32>(PixelFormatInfo.BlockSizeZ, SrcSurface.Depth >> CopyInfo.SourceMipIndex);
 		}
 		else
 		{
 			ensure(CopyInfo.Size.X > 0 && (uint32)CopyInfo.Size.X <= DstSurface.Width && CopyInfo.Size.Y > 0 && (uint32)CopyInfo.Size.Y <= DstSurface.Height);
 			Region.extent.width = CopyInfo.Size.X;
 			Region.extent.height = CopyInfo.Size.Y;
+			Region.extent.depth = CopyInfo.Size.Z;
 		}
-		Region.extent.depth = 1;
 		Region.srcSubresource.aspectMask = SrcSurface.GetFullAspectMask();
 		Region.srcSubresource.baseArrayLayer = CopyInfo.SourceSliceIndex;
 		Region.srcSubresource.layerCount = CopyInfo.NumSlices;
 		Region.srcSubresource.mipLevel = CopyInfo.SourceMipIndex;
 		Region.srcOffset.x = CopyInfo.SourcePosition.X;
 		Region.srcOffset.y = CopyInfo.SourcePosition.Y;
+		Region.srcOffset.z = CopyInfo.SourcePosition.Z;
 		Region.dstSubresource.aspectMask = DstSurface.GetFullAspectMask();
 		Region.dstSubresource.baseArrayLayer = CopyInfo.DestSliceIndex;
 		Region.dstSubresource.layerCount = CopyInfo.NumSlices;
 		Region.dstSubresource.mipLevel = CopyInfo.DestMipIndex;
 		Region.dstOffset.x = CopyInfo.DestPosition.X;
 		Region.dstOffset.y = CopyInfo.DestPosition.Y;
+		Region.dstOffset.z = CopyInfo.DestPosition.Z;
 
 		for (uint32 Index = 0; Index < CopyInfo.NumMips; ++Index)
 		{
@@ -2978,10 +2987,30 @@ void FVulkanCommandListContext::RHICopyTexture(FRHITexture* SourceTexture, FRHIT
 				SrcSurface.Image, VK_IMAGE_LAYOUT_TRANSFER_SRC_OPTIMAL,
 				DstSurface.Image, VK_IMAGE_LAYOUT_TRANSFER_DST_OPTIMAL,
 				1, &Region);
-			Region.extent.width = FMath::Max(1u, Region.extent.width / 2);
-			Region.extent.height = FMath::Max(1u, Region.extent.height / 2);
+
 			++Region.srcSubresource.mipLevel;
 			++Region.dstSubresource.mipLevel;
+
+			// Scale down the copy region if there is another mip to proceed.
+			if (Index != CopyInfo.NumMips - 1)
+			{
+				Region.srcOffset.x /= 2;
+				Region.srcOffset.y /= 2;
+				Region.srcOffset.z /= 2;
+
+				Region.dstOffset.x /= 2;
+				Region.dstOffset.y /= 2;
+				Region.dstOffset.z /= 2;
+
+				Region.extent.width  = FMath::Max<uint32>(Region.extent.width / 2, 1);
+				Region.extent.height = FMath::Max<uint32>(Region.extent.height / 2, 1);
+				Region.extent.depth  = FMath::Max<uint32>(Region.extent.depth / 2, 1);
+
+				// RHICopyTexture is allowed to copy mip regions only if are aligned on the block size to prevent unexpected / inconsistent results.
+				ensure(Region.srcOffset.x % PixelFormatInfo.BlockSizeX == 0 && Region.srcOffset.y % PixelFormatInfo.BlockSizeY == 0 && Region.srcOffset.z % PixelFormatInfo.BlockSizeZ == 0);
+				ensure(Region.dstOffset.x % PixelFormatInfo.BlockSizeX == 0 && Region.dstOffset.y % PixelFormatInfo.BlockSizeY == 0 && Region.dstOffset.z % PixelFormatInfo.BlockSizeZ == 0);
+				// For extent, the condition is harder to verify since on Vulkan, the extent must not be aligned on block size if it would exceed the surface limit.
+			}
 		}
 	}
 }
diff --git a/Engine/Source/Runtime/VulkanRHI/Public/VulkanDynamicRHI.h b/Engine/Source/Runtime/VulkanRHI/Public/VulkanDynamicRHI.h
--- a/Engine/Source/Runtime/VulkanRHI/Public/VulkanDynamicRHI.h
+++ b/Engine/Source/Runtime/VulkanRHI/Public/VulkanDynamicRHI.h
@@ -100,6 +100,7 @@ public:
 	virtual void RHIUpdateShaderResourceView(FRHIShaderResourceView* SRV, FRHIBuffer* Buffer, uint32 Stride, uint8 Format) final override;
 	virtual void RHIUpdateShaderResourceView(FRHIShaderResourceView* SRV, FRHIBuffer* Buffer) final override;
 	virtual uint64 RHICalcTexture2DPlatformSize(uint32 SizeX, uint32 SizeY, uint8 Format, uint32 NumMips, uint32 NumSamples, ETextureCreateFlags Flags, const FRHIResourceCreateInfo& CreateInfo, uint32& OutAlign) final override;
+	virtual uint64 RHICalcTexture2DArrayPlatformSize(uint32 SizeX, uint32 SizeY, uint32 ArraySize, uint8 Format, uint32 NumMips, uint32 NumSamples, ETextureCreateFlags Flags, const FRHIResourceCreateInfo& CreateInfo, uint32& OutAlign) final override;
 	virtual uint64 RHICalcTexture3DPlatformSize(uint32 SizeX, uint32 SizeY, uint32 SizeZ, uint8 Format, uint32 NumMips, ETextureCreateFlags Flags, const FRHIResourceCreateInfo& CreateInfo, uint32& OutAlign) final override;
 	virtual uint64 RHICalcTextureCubePlatformSize(uint32 Size, uint8 Format, uint32 NumMips, ETextureCreateFlags Flags, const FRHIResourceCreateInfo& CreateInfo, uint32& OutAlign) final override;
 	virtual void RHIGetTextureMemoryStats(FTextureMemoryStats& OutStats) final override;
diff --git a/Engine/Source/Runtime/Windows/D3D11RHI/Private/D3D11RHIPrivate.h b/Engine/Source/Runtime/Windows/D3D11RHI/Private/D3D11RHIPrivate.h
--- a/Engine/Source/Runtime/Windows/D3D11RHI/Private/D3D11RHIPrivate.h
+++ b/Engine/Source/Runtime/Windows/D3D11RHI/Private/D3D11RHIPrivate.h
@@ -456,6 +456,7 @@ public:
 	virtual void RHIUpdateShaderResourceView(FRHIShaderResourceView* SRV, FRHIBuffer* VertexBuffer, uint32 Stride, uint8 Format) final override;
 	virtual void RHIUpdateShaderResourceView(FRHIShaderResourceView* SRV, FRHIBuffer* IndexBuffer) final override;
 	virtual uint64 RHICalcTexture2DPlatformSize(uint32 SizeX, uint32 SizeY, uint8 Format, uint32 NumMips, uint32 NumSamples, ETextureCreateFlags Flags, const FRHIResourceCreateInfo& CreateInfo, uint32& OutAlign) final override;
+	virtual uint64 RHICalcTexture2DArrayPlatformSize(uint32 SizeX, uint32 SizeY, uint32 ArraySize, uint8 Format, uint32 NumMips, uint32 NumSamples, ETextureCreateFlags Flags, const FRHIResourceCreateInfo& CreateInfo, uint32& OutAlign) final override;
 	virtual uint64 RHICalcTexture3DPlatformSize(uint32 SizeX, uint32 SizeY, uint32 SizeZ, uint8 Format, uint32 NumMips, ETextureCreateFlags Flags, const FRHIResourceCreateInfo& CreateInfo, uint32& OutAlign) final override;
 	virtual uint64 RHICalcTextureCubePlatformSize(uint32 Size, uint8 Format, uint32 NumMips, ETextureCreateFlags Flags, const FRHIResourceCreateInfo& CreateInfo, uint32& OutAlign) final override;
 	virtual uint64 RHIGetMinimumAlignmentForBufferBackedSRV(EPixelFormat Format) final override;
diff --git a/Engine/Source/Runtime/Windows/D3D11RHI/Private/D3D11Texture.cpp b/Engine/Source/Runtime/Windows/D3D11RHI/Private/D3D11Texture.cpp
--- a/Engine/Source/Runtime/Windows/D3D11RHI/Private/D3D11Texture.cpp
+++ b/Engine/Source/Runtime/Windows/D3D11RHI/Private/D3D11Texture.cpp
@@ -252,6 +252,12 @@ uint64 FD3D11DynamicRHI::RHICalcTexture2DPlatformSize(uint32 SizeX, uint32 SizeY
 	return CalcTextureSize(SizeX, SizeY, (EPixelFormat)Format, NumMips);
 }
 
+uint64 FD3D11DynamicRHI::RHICalcTexture2DArrayPlatformSize(uint32 SizeX, uint32 SizeY, uint32 ArraySize, uint8 Format, uint32 NumMips, uint32 NumSamples, ETextureCreateFlags Flags, const FRHIResourceCreateInfo& CreateInfo, uint32& OutAlign)
+{
+	OutAlign = 0;
+	return CalcTextureSize(SizeX, SizeY, (EPixelFormat)Format, NumMips) * ArraySize;
+}
+
 uint64 FD3D11DynamicRHI::RHICalcTexture3DPlatformSize(uint32 SizeX, uint32 SizeY, uint32 SizeZ, uint8 Format, uint32 NumMips, ETextureCreateFlags Flags, const FRHIResourceCreateInfo& CreateInfo, uint32& OutAlign)
 {
 	OutAlign = 0;
@@ -2683,6 +2689,8 @@ void FD3D11DynamicRHI::RHICopyTexture(FRHITexture* SourceTextureRHI, FRHITexture
 
 	if (CopyInfo.Size != FIntVector::ZeroValue)
 	{
+		const FPixelFormatInfo& PixelFormatInfo = GPixelFormats[DestTextureRHI->GetFormat()];
+
 		D3D11_BOX SrcBox;
 		SrcBox.left = CopyInfo.SourcePosition.X;
 		SrcBox.top = CopyInfo.SourcePosition.Y;
@@ -2701,9 +2709,9 @@ void FD3D11DynamicRHI::RHICopyTexture(FRHITexture* SourceTextureRHI, FRHITexture
 				const uint32 SourceSubresource = D3D11CalcSubresource(SourceMipIndex, SourceSliceIndex, SourceTextureRHI->GetNumMips());
 				const uint32 DestSubresource = D3D11CalcSubresource(DestMipIndex, DestSliceIndex, DestTextureRHI->GetNumMips());
 
-				SrcBox.right = CopyInfo.SourcePosition.X + FMath::Max(CopyInfo.Size.X >> MipIndex, 1);
-				SrcBox.bottom = CopyInfo.SourcePosition.Y + FMath::Max(CopyInfo.Size.Y >> MipIndex, 1);
-				SrcBox.back = CopyInfo.SourcePosition.Z + FMath::Max(CopyInfo.Size.Z >> MipIndex, 1);
+				SrcBox.right = CopyInfo.SourcePosition.X  + AlignArbitrary<uint32>(FMath::Max(CopyInfo.Size.X >> MipIndex, 1), PixelFormatInfo.BlockSizeX);
+				SrcBox.bottom = CopyInfo.SourcePosition.Y + AlignArbitrary<uint32>(FMath::Max(CopyInfo.Size.Y >> MipIndex, 1), PixelFormatInfo.BlockSizeY);
+				SrcBox.back = CopyInfo.SourcePosition.Z   + AlignArbitrary<uint32>(FMath::Max(CopyInfo.Size.Z >> MipIndex, 1), PixelFormatInfo.BlockSizeZ);
 
 				Direct3DDeviceIMContext->CopySubresourceRegion(DestTexture->GetResource(), DestSubresource, CopyInfo.DestPosition.X, CopyInfo.DestPosition.Y, CopyInfo.DestPosition.Z, SourceTexture->GetResource(), SourceSubresource, &SrcBox);
 			}
