VK: Avoid int32 overflow for GRHIMaxDispatchThreadGroupsPerDimension initialization

diff --git a/Engine/Source/Runtime/VulkanRHI/Private/VulkanRHI.cpp b/Engine/Source/Runtime/VulkanRHI/Private/VulkanRHI.cpp
--- a/Engine/Source/Runtime/VulkanRHI/Private/VulkanRHI.cpp
+++ b/Engine/Source/Runtime/VulkanRHI/Private/VulkanRHI.cpp
@@ -738,9 +738,9 @@ void FVulkanDynamicRHI::InitInstance()
 		GRHISupportsBaseVertexIndex = true;
 		GSupportsSeparateRenderTargetBlendState = true;
 
-		GRHIMaxDispatchThreadGroupsPerDimension.X = Limits.maxComputeWorkGroupCount[0];
-		GRHIMaxDispatchThreadGroupsPerDimension.Y = Limits.maxComputeWorkGroupCount[1];
-		GRHIMaxDispatchThreadGroupsPerDimension.Z = Limits.maxComputeWorkGroupCount[2];
+		GRHIMaxDispatchThreadGroupsPerDimension.X = FMath::Min<uint32>(Limits.maxComputeWorkGroupCount[0], 0x7fffffff);
+		GRHIMaxDispatchThreadGroupsPerDimension.Y = FMath::Min<uint32>(Limits.maxComputeWorkGroupCount[1], 0x7fffffff);
+		GRHIMaxDispatchThreadGroupsPerDimension.Z = FMath::Min<uint32>(Limits.maxComputeWorkGroupCount[2], 0x7fffffff);
 
 		FVulkanPlatform::SetupFeatureLevels();
 
