UE5 - Support for LunarG GfxReconstruct via -vktrace
VkTrace is no longer available, but if GfxReconstruct is not available it will try to fall down to vktrace

diff --git a/Engine/Source/Runtime/VulkanRHI/Private/VulkanLayers.cpp b/Engine/Source/Runtime/VulkanRHI/Private/VulkanLayers.cpp
--- a/Engine/Source/Runtime/VulkanRHI/Private/VulkanLayers.cpp
+++ b/Engine/Source/Runtime/VulkanRHI/Private/VulkanLayers.cpp
@@ -327,20 +327,29 @@ void FVulkanDynamicRHI::GetInstanceLayersAndExtensions(TArray<const ANSICHAR*>&
 
 	FVulkanPlatform::NotifyFoundInstanceLayersAndExtensions(FoundUniqueLayers, FoundUniqueExtensions);
 
-	bool bVkTrace = false;
+	bool bGfxReconstructOrVkTrace = false;
 	if (FParse::Param(FCommandLine::Get(), TEXT("vktrace")))
 	{
-		const char* VkTraceName = "VK_LAYER_LUNARG_vktrace";
-		if (FindLayerInList(GlobalLayerExtensions, VkTraceName))
+		const char* GfxReconstructName = "VK_LAYER_LUNARG_gfxreconstruct";
+		if (FindLayerInList(GlobalLayerExtensions, GfxReconstructName))
 		{
-			OutInstanceLayers.Add(VkTraceName);
-			bVkTrace = true;
+			OutInstanceLayers.Add(GfxReconstructName);
+			bGfxReconstructOrVkTrace = true;
+		}
+		else
+		{
+			const char* VkTraceName = "VK_LAYER_LUNARG_vktrace";
+			if (FindLayerInList(GlobalLayerExtensions, VkTraceName))
+			{
+				OutInstanceLayers.Add(VkTraceName);
+				bGfxReconstructOrVkTrace = true;
+			}
 		}
 	}
 
 #if VULKAN_HAS_DEBUGGING_ENABLED
 #if VULKAN_ENABLE_API_DUMP
-	if (!bVkTrace)
+	if (!bGfxReconstructOrVkTrace)
 	{
 		const char* VkApiDumpName = "VK_LAYER_LUNARG_api_dump";
 		bool bApiDumpFound = FindLayerInList(GlobalLayerExtensions, VkApiDumpName);
@@ -358,7 +367,7 @@ void FVulkanDynamicRHI::GetInstanceLayersAndExtensions(TArray<const ANSICHAR*>&
 	// At this point the CVar holds the final value
 #if VULKAN_HAS_DEBUGGING_ENABLED
 	const int32 VulkanValidationOption = GValidationCvar.GetValueOnAnyThread();
-	if (!bVkTrace && VulkanValidationOption > 0)
+	if (!bGfxReconstructOrVkTrace && VulkanValidationOption > 0)
 	{
 		bool bSkipStandard = false;
 		bool bStandardAvailable = false;
@@ -417,7 +426,7 @@ void FVulkanDynamicRHI::GetInstanceLayersAndExtensions(TArray<const ANSICHAR*>&
 #endif
 
 #if VULKAN_SUPPORTS_DEBUG_UTILS
-	if (!bVkTrace && VulkanValidationOption > 0)
+	if (!bGfxReconstructOrVkTrace && VulkanValidationOption > 0)
 	{
 		const char* FoundDebugUtilsLayer = nullptr;
 		bOutDebugUtils = FindLayerExtensionInList(GlobalLayerExtensions, VK_EXT_DEBUG_UTILS_EXTENSION_NAME, FoundDebugUtilsLayer);
@@ -467,13 +476,13 @@ void FVulkanDynamicRHI::GetInstanceLayersAndExtensions(TArray<const ANSICHAR*>&
 	}
 
 #if VULKAN_SUPPORTS_DEBUG_UTILS
-	if (!bVkTrace && bOutDebugUtils && FindLayerExtensionInList(GlobalLayerExtensions, VK_EXT_DEBUG_UTILS_EXTENSION_NAME))
+	if (!bGfxReconstructOrVkTrace && bOutDebugUtils && FindLayerExtensionInList(GlobalLayerExtensions, VK_EXT_DEBUG_UTILS_EXTENSION_NAME))
 	{
 		OutInstanceExtensions.Add(VK_EXT_DEBUG_UTILS_EXTENSION_NAME);
 	}
 #endif
 #if VULKAN_HAS_DEBUGGING_ENABLED
-	if (!bVkTrace && !bOutDebugUtils && VulkanValidationOption > 0)
+	if (!bGfxReconstructOrVkTrace && !bOutDebugUtils && VulkanValidationOption > 0)
 	{
 		if (FindLayerExtensionInList(GlobalLayerExtensions, VK_EXT_DEBUG_REPORT_EXTENSION_NAME))
 		{
@@ -481,7 +490,7 @@ void FVulkanDynamicRHI::GetInstanceLayersAndExtensions(TArray<const ANSICHAR*>&
 		}
 	}
 
-	if (VulkanValidationOption > 0 && !bVkTrace)
+	if (VulkanValidationOption > 0 && !bGfxReconstructOrVkTrace)
 	{
 #if VULKAN_HAS_VALIDATION_FEATURES
 		if (FindLayerExtensionInList(GlobalLayerExtensions, VK_EXT_VALIDATION_FEATURES_EXTENSION_NAME) && GGPUValidationCvar.GetValueOnAnyThread() != 0)
