Unified CreateBuffer() API.
Turn CreateAndLock*() functions into utility functions and remove all superfluous RHI implementations that only did Create+Lock anyway.

diff --git a/Engine/Source/Runtime/Apple/MetalRHI/Private/MetalDynamicRHI.h b/Engine/Source/Runtime/Apple/MetalRHI/Private/MetalDynamicRHI.h
--- a/Engine/Source/Runtime/Apple/MetalRHI/Private/MetalDynamicRHI.h
+++ b/Engine/Source/Runtime/Apple/MetalRHI/Private/MetalDynamicRHI.h
@@ -55,11 +55,9 @@ public:
 	virtual FGraphicsPipelineStateRHIRef RHICreateGraphicsPipelineState(const FGraphicsPipelineStateInitializer& Initializer) final override;
 	virtual TRefCountPtr<FRHIComputePipelineState> RHICreateComputePipelineState(FRHIComputeShader* ComputeShader) final override;
 	virtual FUniformBufferRHIRef RHICreateUniformBuffer(const void* Contents, const FRHIUniformBufferLayout& Layout, EUniformBufferUsage Usage, EUniformBufferValidation Validation) final override;
-	virtual FIndexBufferRHIRef RHICreateIndexBuffer(uint32 Stride, uint32 Size, uint32 InUsage, ERHIAccess InResourceState, FRHIResourceCreateInfo& CreateInfo) final override;
+	virtual FBufferRHIRef RHICreateBuffer(uint32 Size, EBufferUsageFlags Usage, uint32 Stride, ERHIAccess InResourceState, FRHIResourceCreateInfo& CreateInfo) final override;
 	virtual void RHITransferBufferUnderlyingResource(FRHIBuffer* DestBuffer, FRHIBuffer* SrcBuffer) final override;
-	virtual FVertexBufferRHIRef RHICreateVertexBuffer(uint32 Size, uint32 InUsage, ERHIAccess InResourceState, FRHIResourceCreateInfo& CreateInfo) final override;
 	virtual void RHICopyBuffer(FRHIBuffer* SourceBuffer, FRHIBuffer* DestBuffer) final override;
-	virtual FStructuredBufferRHIRef RHICreateStructuredBuffer(uint32 Stride, uint32 Size, uint32 InUsage, ERHIAccess InResourceState, FRHIResourceCreateInfo& CreateInfo) final override;
 	virtual void* LockBuffer_BottomOfPipe(FRHICommandListImmediate& RHICmdList, FRHIBuffer* Buffer, uint32 Offset, uint32 SizeRHI, EResourceLockMode LockMode) final override;
 	virtual void UnlockBuffer_BottomOfPipe(FRHICommandListImmediate& RHICmdList, FRHIBuffer* Buffer) final override;
 	virtual FUnorderedAccessViewRHIRef RHICreateUnorderedAccessView(FRHIStructuredBuffer* StructuredBuffer, bool bUseUAVCounter, bool bAppendBuffer) final override;
@@ -131,9 +129,7 @@ public:
 	virtual IRHIComputeContext* RHIGetDefaultAsyncComputeContext() final override;
 	virtual class IRHICommandContextContainer* RHIGetCommandContextContainer(int32 Index, int32 Num) final override;
 	
-	virtual FVertexBufferRHIRef CreateVertexBuffer_RenderThread(class FRHICommandListImmediate& RHICmdList, uint32 Size, uint32 InUsage, ERHIAccess InResourceState, FRHIResourceCreateInfo& CreateInfo) final override;
-	virtual FIndexBufferRHIRef CreateIndexBuffer_RenderThread(class FRHICommandListImmediate& RHICmdList, uint32 Stride, uint32 Size, uint32 InUsage, ERHIAccess InResourceState, FRHIResourceCreateInfo& CreateInfo) final override;
-	virtual FStructuredBufferRHIRef CreateStructuredBuffer_RenderThread(class FRHICommandListImmediate& RHICmdList, uint32 Stride, uint32 Size, uint32 InUsage, ERHIAccess InResourceState, FRHIResourceCreateInfo& CreateInfo) final override;
+	virtual FBufferRHIRef CreateBuffer_RenderThread(class FRHICommandListImmediate& RHICmdList, uint32 Size, EBufferUsageFlags Usage, uint32 Stride, ERHIAccess ResourceState, FRHIResourceCreateInfo& CreateInfo) final override;
 	
 	virtual FTexture2DRHIRef AsyncReallocateTexture2D_RenderThread(class FRHICommandListImmediate& RHICmdList, FRHITexture2D* Texture2D, int32 NewMipCount, int32 NewSizeX, int32 NewSizeY, FThreadSafeCounter* RequestStatus) final override;
 	virtual ETextureReallocationStatus FinalizeAsyncReallocateTexture2D_RenderThread(class FRHICommandListImmediate& RHICmdList, FRHITexture2D* Texture2D, bool bBlockUntilCompleted) final override;
diff --git a/Engine/Source/Runtime/Apple/MetalRHI/Private/MetalIndexBuffer.cpp b/Engine/Source/Runtime/Apple/MetalRHI/Private/MetalIndexBuffer.cpp
--- a/Engine/Source/Runtime/Apple/MetalRHI/Private/MetalIndexBuffer.cpp
+++ b/Engine/Source/Runtime/Apple/MetalRHI/Private/MetalIndexBuffer.cpp
@@ -12,48 +12,6 @@
 #include "RenderUtils.h"
 #include "HAL/LowLevelMemTracker.h"
 
-static uint32 MetalIndexBufferUsage(uint32 InUsage)
-{
-	uint32 Usage = InUsage | BUF_IndexBuffer;
-	if (RHISupportsTessellation(GMaxRHIShaderPlatform))
-	{
-		Usage |= BUF_ShaderResource;
-	}
-	Usage |= (EMetalBufferUsage_GPUOnly | EMetalBufferUsage_LinearTex);
-	return Usage;
-}
-
-FIndexBufferRHIRef FMetalDynamicRHI::RHICreateIndexBuffer(uint32 Stride, uint32 Size, uint32 InUsage, ERHIAccess InResourceState, FRHIResourceCreateInfo& CreateInfo)
-{
-	@autoreleasepool {
-	if (CreateInfo.bWithoutNativeResource)
-	{
-		return new FMetalResourceMultiBuffer(0, MetalIndexBufferUsage(0), 2, nullptr, RRT_IndexBuffer);
-	}
-		
-	// make the RHI object, which will allocate memory
-	FMetalResourceMultiBuffer* IndexBuffer = new FMetalResourceMultiBuffer(Size, MetalIndexBufferUsage(InUsage), Stride, nullptr, RRT_IndexBuffer);
-	
-	if (CreateInfo.ResourceArray)
-	{
-		check(Size == CreateInfo.ResourceArray->GetResourceDataSize());
-
-		// make a buffer usable by CPU
-		void* Buffer = ::RHILockBuffer(IndexBuffer, 0, Size, RLM_WriteOnly);
-
-		// copy the contents of the given data into the buffer
-		FMemory::Memcpy(Buffer, CreateInfo.ResourceArray->GetResourceData(), Size);
-
-		::RHIUnlockBuffer(IndexBuffer);
-
-		// Discard the resource array's contents.
-		CreateInfo.ResourceArray->Discard();
-	}
-
-	return IndexBuffer;
-	}
-}
-
 void FMetalDynamicRHI::RHITransferBufferUnderlyingResource(FRHIBuffer* DestBuffer, FRHIBuffer* SrcBuffer)
 {
 	@autoreleasepool {
@@ -71,20 +29,3 @@ void FMetalDynamicRHI::RHITransferBufferUnderlyingResource(FRHIBuffer* DestBuffe
 	}
 	}
 }
-
-FIndexBufferRHIRef FMetalDynamicRHI::CreateIndexBuffer_RenderThread(class FRHICommandListImmediate& RHICmdList, uint32 Stride, uint32 Size, uint32 InUsage, ERHIAccess InResourceState, FRHIResourceCreateInfo& CreateInfo)
-{
-	@autoreleasepool {
-		if (CreateInfo.bWithoutNativeResource)
-		{
-			return new FMetalResourceMultiBuffer(0, MetalIndexBufferUsage(0), 2, nullptr, RRT_IndexBuffer);
-		}
-		
-		// make the RHI object, which will allocate memory
-		TRefCountPtr<FMetalResourceMultiBuffer> IndexBuffer = new FMetalResourceMultiBuffer(Size, MetalIndexBufferUsage(InUsage), Stride, nullptr, RRT_IndexBuffer);
-		
-		IndexBuffer->Init_RenderThread(RHICmdList, Size, InUsage, CreateInfo, IndexBuffer);
-		
-		return IndexBuffer.GetReference();
-	}
-}
diff --git a/Engine/Source/Runtime/Apple/MetalRHI/Private/MetalStructuredBuffer.cpp b/Engine/Source/Runtime/Apple/MetalRHI/Private/MetalStructuredBuffer.cpp
--- a/Engine/Source/Runtime/Apple/MetalRHI/Private/MetalStructuredBuffer.cpp
+++ b/Engine/Source/Runtime/Apple/MetalRHI/Private/MetalStructuredBuffer.cpp
@@ -5,27 +5,3 @@
 #include "MetalProfiler.h"
 #include "MetalCommandBuffer.h"
 #include "Containers/ResourceArray.h"
-
-static uint32 MetalStructuredBufferUsage(uint32 InUsage)
-{
-	return (InUsage | EMetalBufferUsage_GPUOnly);
-}
-
-FStructuredBufferRHIRef FMetalDynamicRHI::RHICreateStructuredBuffer(uint32 Stride, uint32 Size, uint32 InUsage, ERHIAccess InResourceState, FRHIResourceCreateInfo& CreateInfo)
-{
-	@autoreleasepool {
-		return new FMetalResourceMultiBuffer(Size, MetalStructuredBufferUsage(InUsage), Stride, CreateInfo.ResourceArray, RRT_StructuredBuffer);
-	}
-}
-
-FStructuredBufferRHIRef FMetalDynamicRHI::CreateStructuredBuffer_RenderThread(class FRHICommandListImmediate& RHICmdList, uint32 Stride, uint32 Size, uint32 InUsage, ERHIAccess InResourceState, FRHIResourceCreateInfo& CreateInfo)
-{
-	@autoreleasepool {
-		// make the RHI object, which will allocate memory
-		TRefCountPtr<FMetalResourceMultiBuffer> VertexBuffer = new FMetalResourceMultiBuffer(Size, MetalStructuredBufferUsage(InUsage), Stride, nullptr, RRT_StructuredBuffer);
-		
-		VertexBuffer->Init_RenderThread(RHICmdList, Size, InUsage, CreateInfo, VertexBuffer);
-		
-		return VertexBuffer.GetReference();
-	}
-}
diff --git a/Engine/Source/Runtime/Apple/MetalRHI/Private/MetalVertexBuffer.cpp b/Engine/Source/Runtime/Apple/MetalRHI/Private/MetalVertexBuffer.cpp
--- a/Engine/Source/Runtime/Apple/MetalRHI/Private/MetalVertexBuffer.cpp
+++ b/Engine/Source/Runtime/Apple/MetalRHI/Private/MetalVertexBuffer.cpp
@@ -20,16 +20,14 @@ DECLARE_MEMORY_STAT(TEXT("Used Device Buffer Memory"), STAT_MetalDeviceBufferMem
 #endif
 
 #if STATS
-#define METAL_INC_DWORD_STAT_BY(Type, Name, Size) \
+#define METAL_INC_DWORD_STAT_BY(Type, Name, Size, Usage) \
 	switch(Type)	{ \
 		case RRT_UniformBuffer: INC_DWORD_STAT_BY(STAT_MetalUniform##Name, Size); break; \
-		case RRT_IndexBuffer: INC_DWORD_STAT_BY(STAT_MetalIndex##Name, Size); break; \
-		case RRT_StructuredBuffer: \
-		case RRT_VertexBuffer: INC_DWORD_STAT_BY(STAT_MetalVertex##Name, Size); break; \
+        case RRT_Buffer: if (Usage & BUF_IndexBuffer){ INC_DWORD_STAT_BY(STAT_MetalIndex##Name, Size); } else { INC_DWORD_STAT_BY(STAT_MetalVertex##Name, Size); } break; \
 		default: break; \
 	}
 #else
-#define METAL_INC_DWORD_STAT_BY(Type, Name, Size)
+#define METAL_INC_DWORD_STAT_BY(Type, Name, Size, Usage)
 #endif
 
 @implementation FMetalBufferData
@@ -81,9 +79,30 @@ APPLE_PLATFORM_OBJECT_ALLOC_OVERRIDES(FMetalBufferData)
 }
 @end
 
-static uint32 MetalVertexBufferUsage(uint32 InUsage)
+static uint32 MetalBufferUsage(uint32 InUsage)
 {
-	return (InUsage | BUF_VertexBuffer | EMetalBufferUsage_LinearTex);
+	uint32 Usage = InUsage;
+
+	if (InUsage & BUF_VertexBuffer)
+	{
+		Usage |= EMetalBufferUsage_LinearTex;
+	}
+
+	if (InUsage & BUF_IndexBuffer)
+	{
+		if (RHISupportsTessellation(GMaxRHIShaderPlatform))
+		{
+			Usage |= BUF_ShaderResource;
+		}
+		Usage |= (EMetalBufferUsage_GPUOnly | EMetalBufferUsage_LinearTex);
+	}
+
+	if (InUsage & BUF_StructuredBuffer)
+	{
+		Usage |= EMetalBufferUsage_GPUOnly;
+	}
+
+	return Usage;
 }
 
 void FMetalRHIBuffer::Swap(FMetalRHIBuffer& Other)
@@ -140,7 +159,7 @@ FMetalRHIBuffer::FMetalRHIBuffer(uint32 InSize, uint32 InUsage, ERHIResourceType
 		if (!(InUsage & (BUF_UnorderedAccess|BUF_ShaderResource|EMetalBufferUsage_GPUOnly)) && (InUsage & BUF_Volatile) && InSize < MetalBufferPageSize && (InSize < MetalBufferBytesSize))
 		{
 			Data = [[FMetalBufferData alloc] initWithSize:InSize];
-			METAL_INC_DWORD_STAT_BY(Type, MemAlloc, InSize);
+			METAL_INC_DWORD_STAT_BY(Type, MemAlloc, InSize, Usage);
 		}
 		else
 		{
@@ -274,7 +293,7 @@ FMetalRHIBuffer::~FMetalRHIBuffer()
 {
 	if(TransferBuffer)
 	{
-		METAL_INC_DWORD_STAT_BY(Type, MemFreed, TransferBuffer.GetLength());
+		METAL_INC_DWORD_STAT_BY(Type, MemFreed, TransferBuffer.GetLength(), Usage);
 		SafeReleaseMetalBuffer(TransferBuffer);
 	}
 	
@@ -282,7 +301,7 @@ FMetalRHIBuffer::~FMetalRHIBuffer()
 	{
 		check(Backing.Buffer);
 		
-		METAL_INC_DWORD_STAT_BY(Type, MemFreed, Backing.Buffer.GetLength());
+		METAL_INC_DWORD_STAT_BY(Type, MemFreed, Backing.Buffer.GetLength(), Usage);
 		SafeReleaseMetalBuffer(Backing.Buffer);
 		
 		for (auto& Pair : Backing.Views)
@@ -295,7 +314,7 @@ FMetalRHIBuffer::~FMetalRHIBuffer()
 
 	if (Data)
 	{
-		METAL_INC_DWORD_STAT_BY(Type, MemFreed, Size);
+		METAL_INC_DWORD_STAT_BY(Type, MemFreed, Size, Usage);
 		SafeReleaseMetalObject(Data);
 	}
 }
@@ -307,7 +326,7 @@ void FMetalRHIBuffer::AllocTransferBuffer(bool bOnRHIThread, uint32 InSize, ERes
 	TransferBuffer = GetMetalDeviceContext().CreatePooledBuffer(ArgsCPU);
 	TransferBuffer.SetOwner(nullptr, false);
 	check(TransferBuffer && TransferBuffer.GetPtr());
-	METAL_INC_DWORD_STAT_BY(Type, MemAlloc, InSize);
+	METAL_INC_DWORD_STAT_BY(Type, MemAlloc, InSize, Usage);
 	METAL_FATAL_ASSERT(TransferBuffer, TEXT("Failed to create buffer of size %u and storage mode %u"), InSize, (uint32)mtlpp::StorageMode::Shared);
 }
 
@@ -720,40 +739,40 @@ void FMetalResourceMultiBuffer::Swap(FMetalResourceMultiBuffer& Other)
 	}
 }
 
-FVertexBufferRHIRef FMetalDynamicRHI::RHICreateVertexBuffer(uint32 Size, uint32 InUsage, ERHIAccess InResourceState, FRHIResourceCreateInfo& CreateInfo)
+FBufferRHIRef FMetalDynamicRHI::RHICreateBuffer(uint32 Size, EBufferUsageFlags Usage, uint32 Stride, ERHIAccess InResourceState, FRHIResourceCreateInfo& CreateInfo)
 {
 	check(0);
 	@autoreleasepool {
 	if (CreateInfo.bWithoutNativeResource)
 	{
-		return new FMetalResourceMultiBuffer(0, MetalVertexBufferUsage(0), 0, nullptr, RRT_VertexBuffer);
+		return new FMetalResourceMultiBuffer(0, MetalBufferUsage(0), 0, nullptr, RRT_Buffer);
 	}
 	
 	// make the RHI object, which will allocate memory
-	FMetalResourceMultiBuffer* VertexBuffer = new FMetalResourceMultiBuffer(Size, MetalVertexBufferUsage(InUsage), 0, nullptr, RRT_VertexBuffer);
+	FMetalResourceMultiBuffer* Buffer = new FMetalResourceMultiBuffer(Size, MetalBufferUsage(Usage), 0, nullptr, RRT_Buffer);
 
 	if (CreateInfo.ResourceArray)
 	{
 		check(Size >= CreateInfo.ResourceArray->GetResourceDataSize());
 
 		// make a buffer usable by CPU
-		void* Buffer = ::RHILockBuffer(VertexBuffer, 0, Size, RLM_WriteOnly);
+		void* BufferData = ::RHILockBuffer(Buffer, 0, Size, RLM_WriteOnly);
 		
 		// copy the contents of the given data into the buffer
-		FMemory::Memcpy(Buffer, CreateInfo.ResourceArray->GetResourceData(), Size);
+		FMemory::Memcpy(BufferData, CreateInfo.ResourceArray->GetResourceData(), Size);
 		
-		::RHIUnlockBuffer(VertexBuffer);
+		::RHIUnlockBuffer(Buffer);
 
 		// Discard the resource array's contents.
 		CreateInfo.ResourceArray->Discard();
 	}
-	else if (VertexBuffer->Mode == mtlpp::StorageMode::Private)
+	else if (Buffer->Mode == mtlpp::StorageMode::Private)
 	{
-		check (!VertexBuffer->TransferBuffer);
+		check (!Buffer->TransferBuffer);
 
 		if (GMetalBufferZeroFill && !FMetalCommandQueue::SupportsFeature(EMetalFeaturesFences))
 		{
-			for(FMetalRHIBuffer::FMetalBufferAndViews& Backing : VertexBuffer->BufferPool)
+			for(FMetalRHIBuffer::FMetalBufferAndViews& Backing : Buffer->BufferPool)
 			{
 				FMetalBuffer& TheBuffer = Backing.Buffer;
 				GetMetalDeviceContext().FillBuffer(TheBuffer, ns::Range(0, TheBuffer.GetLength()), 0);
@@ -761,9 +780,9 @@ FVertexBufferRHIRef FMetalDynamicRHI::RHICreateVertexBuffer(uint32 Size, uint32
 		}
 	}
 #if PLATFORM_MAC
-	else if (GMetalBufferZeroFill && VertexBuffer->Mode == mtlpp::StorageMode::Managed)
+	else if (GMetalBufferZeroFill && Buffer->Mode == mtlpp::StorageMode::Managed)
 	{
-		for(FMetalRHIBuffer::FMetalBufferAndViews& Backing : VertexBuffer->BufferPool)
+		for(FMetalRHIBuffer::FMetalBufferAndViews& Backing : Buffer->BufferPool)
 		{
 			FMetalBuffer& TheBuffer = Backing.Buffer;
 			GetMetalDeviceContext().FillBuffer(TheBuffer, ns::Range(0, TheBuffer.GetLength()), 0);
@@ -772,7 +791,7 @@ FVertexBufferRHIRef FMetalDynamicRHI::RHICreateVertexBuffer(uint32 Size, uint32
 	}
 #endif
 
-	return VertexBuffer;
+	return Buffer;
 	}
 }
 
@@ -880,19 +899,19 @@ void FMetalRHIBuffer::Init_RenderThread(FRHICommandListImmediate& RHICmdList, ui
 	}
 }
 
-FVertexBufferRHIRef FMetalDynamicRHI::CreateVertexBuffer_RenderThread(class FRHICommandListImmediate& RHICmdList, uint32 Size, uint32 InUsage, ERHIAccess InResourceState, FRHIResourceCreateInfo& CreateInfo)
+FBufferRHIRef FMetalDynamicRHI::CreateBuffer_RenderThread(class FRHICommandListImmediate& RHICmdList, uint32 Size, EBufferUsageFlags Usage, uint32 Stride, ERHIAccess ResourceState, FRHIResourceCreateInfo& CreateInfo)
 {
 	@autoreleasepool {
 		if (CreateInfo.bWithoutNativeResource)
 		{
-			return new FMetalResourceMultiBuffer(0, MetalVertexBufferUsage(0), 0, nullptr, RRT_VertexBuffer);
+			return new FMetalResourceMultiBuffer(0, MetalBufferUsage(0), 0, nullptr, RRT_Buffer);
 		}
 		
 		// make the RHI object, which will allocate memory
-		TRefCountPtr<FMetalResourceMultiBuffer> VertexBuffer = new FMetalResourceMultiBuffer(Size, MetalVertexBufferUsage(InUsage), 0, nullptr, RRT_VertexBuffer);
+		TRefCountPtr<FMetalResourceMultiBuffer> Buffer = new FMetalResourceMultiBuffer(Size, MetalBufferUsage(Usage), Stride, nullptr, RRT_Buffer);
 		
-		VertexBuffer->Init_RenderThread(RHICmdList, Size, InUsage, CreateInfo, VertexBuffer);
+		Buffer->Init_RenderThread(RHICmdList, Size, Usage, CreateInfo, Buffer);
 		
-		return VertexBuffer.GetReference();
+		return Buffer.GetReference();
 	}
 }
diff --git a/Engine/Source/Runtime/D3D12RHI/Private/D3D12Buffer.cpp b/Engine/Source/Runtime/D3D12RHI/Private/D3D12Buffer.cpp
--- a/Engine/Source/Runtime/D3D12RHI/Private/D3D12Buffer.cpp
+++ b/Engine/Source/Runtime/D3D12RHI/Private/D3D12Buffer.cpp
@@ -367,6 +367,64 @@ void FD3D12Buffer::ReleaseUnderlyingResource()
 	}
 }
 
+D3D12_RESOURCE_DESC CreateBufferResourceDesc(uint32 Size, uint32 Usage)
+{
+	D3D12_RESOURCE_DESC Desc = CD3DX12_RESOURCE_DESC::Buffer(Size);
+
+	if (Usage & BUF_UnorderedAccess)
+	{
+		Desc.Flags |= D3D12_RESOURCE_FLAG_ALLOW_UNORDERED_ACCESS;
+
+		static bool bRequiresRawView = (GMaxRHIFeatureLevel < ERHIFeatureLevel::SM5);
+		if (bRequiresRawView)
+		{
+			// Force the buffer to be a raw, byte address buffer
+			Usage |= BUF_ByteAddressBuffer;
+		}
+	}
+
+	if ((Usage & BUF_ShaderResource) == 0)
+	{
+		Desc.Flags |= D3D12_RESOURCE_FLAG_DENY_SHADER_RESOURCE;
+	}
+
+	if (Usage & BUF_DrawIndirect)
+	{
+		Desc.Flags |= D3D12RHI_RESOURCE_FLAG_ALLOW_INDIRECT_BUFFER;
+	}
+
+	return Desc;
+}
+
+FBufferRHIRef FD3D12DynamicRHI::RHICreateBuffer(uint32 Size, EBufferUsageFlags Usage, uint32 Stride, ERHIAccess InResourceState, FRHIResourceCreateInfo& CreateInfo)
+{
+	if (CreateInfo.bWithoutNativeResource)
+	{
+		return GetAdapter().CreateLinkedObject<FD3D12Buffer>(CreateInfo.GPUMask, [](FD3D12Device* Device)
+			{
+				return new FD3D12Buffer();
+			});
+	}
+
+	const D3D12_RESOURCE_DESC Desc = CreateBufferResourceDesc(Size, Usage);
+	// Structured buffers, non-ByteAddress buffers, need to be aligned to their stride to ensure that they can be addressed correctly with element based offsets.
+	const uint32 Alignment = (Stride > 0) && (((Usage & BUF_StructuredBuffer) != 0) || ((Usage & (BUF_ByteAddressBuffer | BUF_DrawIndirect)) == 0))? Stride : 4;
+
+	FD3D12Buffer* Buffer = GetAdapter().CreateRHIBuffer(nullptr, Desc, Alignment, Stride, Size, Usage, ED3D12ResourceStateMode::Default, InResourceState, CreateInfo);
+	if (Buffer->ResourceLocation.IsTransient())
+	{
+		// TODO: this should ideally be set in platform-independent code, since this tracking is for the high level
+		Buffer->SetCommitted(false);
+	}
+
+	return Buffer;
+}
+
+FBufferRHIRef FD3D12DynamicRHI::CreateBuffer_RenderThread(class FRHICommandListImmediate& RHICmdList, uint32 Size, EBufferUsageFlags Usage, uint32 Stride, ERHIAccess ResourceState, FRHIResourceCreateInfo& CreateInfo)
+{
+	return RHICreateBuffer(Size, Usage, Stride, ResourceState, CreateInfo);
+}
+
 void* FD3D12DynamicRHI::LockBuffer(FRHICommandListImmediate* RHICmdList, FD3D12Buffer* Buffer, uint32 BufferSize, uint32 BufferUsage, uint32 Offset, uint32 Size, EResourceLockMode LockMode)
 {
 	SCOPE_CYCLE_COUNTER(STAT_D3D12LockBufferTime);
diff --git a/Engine/Source/Runtime/D3D12RHI/Private/D3D12IndexBuffer.cpp b/Engine/Source/Runtime/D3D12RHI/Private/D3D12IndexBuffer.cpp
--- a/Engine/Source/Runtime/D3D12RHI/Private/D3D12IndexBuffer.cpp
+++ b/Engine/Source/Runtime/D3D12RHI/Private/D3D12IndexBuffer.cpp
@@ -5,89 +5,3 @@
 =============================================================================*/
 
 #include "D3D12RHIPrivate.h"
-
-D3D12_RESOURCE_DESC CreateIndexBufferResourceDesc(uint32 Size, uint32 InUsage)
-{
-	// Describe the vertex buffer.
-	D3D12_RESOURCE_DESC Desc = CD3DX12_RESOURCE_DESC::Buffer(Size);
-
-	if (InUsage & BUF_UnorderedAccess)
-	{
-		Desc.Flags |= D3D12_RESOURCE_FLAG_ALLOW_UNORDERED_ACCESS;
-	}
-
-	if ((InUsage & BUF_ShaderResource) == 0)
-	{
-		Desc.Flags |= D3D12_RESOURCE_FLAG_DENY_SHADER_RESOURCE;
-	}
-
-	if (InUsage & BUF_DrawIndirect)
-	{
-		Desc.Flags |= D3D12RHI_RESOURCE_FLAG_ALLOW_INDIRECT_BUFFER;
-	}
-
-	return Desc;
-}
-
-FIndexBufferRHIRef FD3D12DynamicRHI::RHICreateIndexBuffer(uint32 Stride, uint32 Size, uint32 InUsage, ERHIAccess InResourceState, FRHIResourceCreateInfo& CreateInfo)
-{
-	if (CreateInfo.bWithoutNativeResource)
-	{
-		return GetAdapter().CreateLinkedObject<FD3D12Buffer>(CreateInfo.GPUMask, [](FD3D12Device* Device)
-			{
-				return new FD3D12Buffer();
-			});
-	}
-
-	const D3D12_RESOURCE_DESC Desc = CreateIndexBufferResourceDesc(Size, InUsage);
-	const uint32 Alignment = 4;
-
-	FD3D12Buffer* Buffer = GetAdapter().CreateRHIBuffer(nullptr, Desc, Alignment, Stride, Size, InUsage | BUF_IndexBuffer, ED3D12ResourceStateMode::Default, InResourceState, CreateInfo);
-	if (Buffer->ResourceLocation.IsTransient())
-	{
-		// TODO: this should ideally be set in platform-independent code, since this tracking is for the high level
-		Buffer->SetCommitted(false);
-	}
-
-	return Buffer;
-}
-
-FIndexBufferRHIRef FD3D12DynamicRHI::CreateIndexBuffer_RenderThread(class FRHICommandListImmediate& RHICmdList, uint32 Stride, uint32 Size, uint32 InUsage, ERHIAccess InResourceState, FRHIResourceCreateInfo& CreateInfo)
-{
-	if (CreateInfo.bWithoutNativeResource)
-	{
-		return GetAdapter().CreateLinkedObject<FD3D12Buffer>(CreateInfo.GPUMask, [](FD3D12Device* Device)
-			{
-				return new FD3D12Buffer();
-			});
-	}
-
-	const D3D12_RESOURCE_DESC Desc = CreateIndexBufferResourceDesc(Size, InUsage);
-	const uint32 Alignment = 4;
-
-	FD3D12Buffer* Buffer = GetAdapter().CreateRHIBuffer(&RHICmdList, Desc, Alignment, Stride, Size, InUsage | BUF_IndexBuffer, ED3D12ResourceStateMode::Default, InResourceState, CreateInfo);
-	if (Buffer->ResourceLocation.IsTransient())
-	{
-		// TODO: this should ideally be set in platform-independent code, since this tracking is for the high level
-		Buffer->SetCommitted(false);
-	}
-
-	return Buffer;
-}
-
-FIndexBufferRHIRef FD3D12DynamicRHI::CreateAndLockIndexBuffer_RenderThread(class FRHICommandListImmediate& RHICmdList, uint32 Stride, uint32 Size, uint32 InUsage, ERHIAccess InResourceState, FRHIResourceCreateInfo& CreateInfo, void*& OutDataBuffer)
-{
-	const D3D12_RESOURCE_DESC Desc = CreateIndexBufferResourceDesc(Size, InUsage);
-	const uint32 Alignment = 4;
-
-	FD3D12Buffer* Buffer = GetAdapter().CreateRHIBuffer(&RHICmdList, Desc, Alignment, Stride, Size, InUsage, ED3D12ResourceStateMode::Default, InResourceState, CreateInfo);
-	if (Buffer->ResourceLocation.IsTransient())
-	{
-		// TODO: this should ideally be set in platform-independent code, since this tracking is for the high level
-		Buffer->SetCommitted(false);
-	}
-
-	OutDataBuffer = LockBuffer(&RHICmdList, Buffer, Buffer->GetSize(), Buffer->GetUsage(), 0, Size, RLM_WriteOnly);
-
-	return Buffer;
-}
diff --git a/Engine/Source/Runtime/D3D12RHI/Private/D3D12RHIPrivate.h b/Engine/Source/Runtime/D3D12RHI/Private/D3D12RHIPrivate.h
--- a/Engine/Source/Runtime/D3D12RHI/Private/D3D12RHIPrivate.h
+++ b/Engine/Source/Runtime/D3D12RHI/Private/D3D12RHIPrivate.h
@@ -284,11 +284,9 @@ public:
 	virtual void RHIReleaseTransition(FRHITransition* Transition) final override;
 	virtual FUniformBufferRHIRef RHICreateUniformBuffer(const void* Contents, const FRHIUniformBufferLayout& Layout, EUniformBufferUsage Usage, EUniformBufferValidation Validation) final override;
 	virtual void RHIUpdateUniformBuffer(FRHIUniformBuffer* UniformBufferRHI, const void* Contents) final override;
-	virtual FIndexBufferRHIRef RHICreateIndexBuffer(uint32 Stride, uint32 Size, uint32 InUsage, ERHIAccess InResourceState, FRHIResourceCreateInfo& CreateInfo) final override;
+	virtual FBufferRHIRef RHICreateBuffer(uint32 Size, EBufferUsageFlags Usage, uint32 Stride, ERHIAccess ResourceState, FRHIResourceCreateInfo& CreateInfo) final override;
 	virtual void RHITransferBufferUnderlyingResource(FRHIBuffer* DestBuffer, FRHIBuffer* SrcBuffer) final override;
-	virtual FVertexBufferRHIRef RHICreateVertexBuffer(uint32 Size, uint32 InUsage, ERHIAccess InResourceState, FRHIResourceCreateInfo& CreateInfo) final override;
 	virtual void RHICopyBuffer(FRHIBuffer* SourceBuffer, FRHIBuffer* DestBuffer) final override;
-	virtual FStructuredBufferRHIRef RHICreateStructuredBuffer(uint32 Stride, uint32 Size, uint32 InUsage, ERHIAccess InResourceState, FRHIResourceCreateInfo& CreateInfo) final override;
 	virtual void* RHILockBuffer(FRHICommandListImmediate& RHICmdList, FRHIBuffer* Buffer, uint32 Offset, uint32 Size, EResourceLockMode LockMode) final override;
 	virtual void RHIUnlockBuffer(FRHICommandListImmediate& RHICmdList, FRHIBuffer* Buffer) final override;
 	virtual FUnorderedAccessViewRHIRef RHICreateUnorderedAccessView(FRHIBuffer* StructuredBuffer, bool bUseUAVCounter, bool bAppendBuffer) final override;
@@ -422,14 +420,10 @@ public:
 		return RHICreateComputeShader(Code, Hash);
 	}
 
-	virtual FVertexBufferRHIRef CreateVertexBuffer_RenderThread(class FRHICommandListImmediate& RHICmdList, uint32 Size, uint32 InUsage, ERHIAccess InResourceState, FRHIResourceCreateInfo& CreateInfo) override final;
-	virtual FStructuredBufferRHIRef CreateStructuredBuffer_RenderThread(class FRHICommandListImmediate& RHICmdList, uint32 Stride, uint32 Size, uint32 InUsage, ERHIAccess InResourceState, FRHIResourceCreateInfo& CreateInfo) override final;
-	virtual FVertexBufferRHIRef CreateAndLockVertexBuffer_RenderThread(class FRHICommandListImmediate& RHICmdList, uint32 Size, uint32 InUsage, ERHIAccess InResourceState, FRHIResourceCreateInfo& CreateInfo, void*& OutDataBuffer) override final;
-	virtual FIndexBufferRHIRef CreateAndLockIndexBuffer_RenderThread(class FRHICommandListImmediate& RHICmdList, uint32 Stride, uint32 Size, uint32 InUsage, ERHIAccess InResourceState, FRHIResourceCreateInfo& CreateInfo, void*& OutDataBuffer) override final;
+	virtual FBufferRHIRef CreateBuffer_RenderThread(class FRHICommandListImmediate& RHICmdList, uint32 Size, EBufferUsageFlags Usage, uint32 Stride, ERHIAccess ResourceState, FRHIResourceCreateInfo& CreateInfo) override final;
 	// virtual FTexture2DRHIRef AsyncReallocateTexture2D_RenderThread(class FRHICommandListImmediate& RHICmdList, FRHITexture2D* Texture2D, int32 NewMipCount, int32 NewSizeX, int32 NewSizeY, FThreadSafeCounter* RequestStatus) override final;
 	// virtual ETextureReallocationStatus FinalizeAsyncReallocateTexture2D_RenderThread(class FRHICommandListImmediate& RHICmdList, FRHITexture2D* Texture2D, bool bBlockUntilCompleted) override final;
 	// virtual ETextureReallocationStatus CancelAsyncReallocateTexture2D_RenderThread(class FRHICommandListImmediate& RHICmdList, FRHITexture2D* Texture2D, bool bBlockUntilCompleted) override final;
-	virtual FIndexBufferRHIRef CreateIndexBuffer_RenderThread(class FRHICommandListImmediate& RHICmdList, uint32 Stride, uint32 Size, uint32 InUsage, ERHIAccess InResourceState, FRHIResourceCreateInfo& CreateInfo) override final;
 	virtual void UpdateTexture2D_RenderThread(class FRHICommandListImmediate& RHICmdList, FRHITexture2D* Texture, uint32 MipIndex, const struct FUpdateTextureRegion2D& UpdateRegion, uint32 SourcePitch, const uint8* SourceData) override final;
 	virtual void* LockTexture2D_RenderThread(class FRHICommandListImmediate& RHICmdList, FRHITexture2D* Texture, uint32 MipIndex, EResourceLockMode LockMode, uint32& DestStride, bool bLockWithinMiptail, bool bNeedsDefaultRHIFlush = true) override final;
 	virtual void UnlockTexture2D_RenderThread(class FRHICommandListImmediate& RHICmdList, FRHITexture2D* Texture, uint32 MipIndex, bool bLockWithinMiptail, bool bNeedsDefaultRHIFlush = true) override final;
diff --git a/Engine/Source/Runtime/D3D12RHI/Private/D3D12StructuredBuffer.cpp b/Engine/Source/Runtime/D3D12RHI/Private/D3D12StructuredBuffer.cpp
--- a/Engine/Source/Runtime/D3D12RHI/Private/D3D12StructuredBuffer.cpp
+++ b/Engine/Source/Runtime/D3D12RHI/Private/D3D12StructuredBuffer.cpp
@@ -2,69 +2,3 @@
 
 #include "D3D12RHIPrivate.h"
 #include "D3D12View.h"
-
-D3D12_RESOURCE_DESC CreateStructuredBufferResourceDesc(uint32 Size, uint32 InUsage)
-{
-	// Describe the structured buffer.
-	D3D12_RESOURCE_DESC Desc = CD3DX12_RESOURCE_DESC::Buffer(Size);
-
-	if ((InUsage & BUF_ShaderResource) == 0)
-	{
-		Desc.Flags |= D3D12_RESOURCE_FLAG_DENY_SHADER_RESOURCE;
-	}
-
-	if (InUsage & BUF_UnorderedAccess)
-	{
-		Desc.Flags |= D3D12_RESOURCE_FLAG_ALLOW_UNORDERED_ACCESS;
-	}
-
-	if (InUsage & BUF_DrawIndirect)
-	{
-		Desc.Flags |= D3D12RHI_RESOURCE_FLAG_ALLOW_INDIRECT_BUFFER;
-	}
-
-
-	return Desc;
-}
-
-FStructuredBufferRHIRef FD3D12DynamicRHI::CreateStructuredBuffer_RenderThread(FRHICommandListImmediate& RHICmdList, uint32 Stride, uint32 Size, uint32 InUsage, ERHIAccess InResourceState, FRHIResourceCreateInfo& CreateInfo)
-{
-	// Check for values that will cause D3D calls to fail
-	check(Size / Stride > 0 && Size % Stride == 0);
-
-	const D3D12_RESOURCE_DESC Desc = CreateStructuredBufferResourceDesc(Size, InUsage);
-
-	// Structured buffers, non-byte address buffers, need to be aligned to their stride to ensure that they
-	// can be addressed correctly with element based offsets.
-	const uint32 Alignment = ((InUsage & (BUF_ByteAddressBuffer | BUF_DrawIndirect)) == 0) ? Stride : 4;
-
-	FD3D12Buffer* NewBuffer = GetAdapter().CreateRHIBuffer(&RHICmdList, Desc, Alignment, Stride, Size, InUsage | BUF_StructuredBuffer, ED3D12ResourceStateMode::Default, InResourceState, CreateInfo);
-	if (NewBuffer->ResourceLocation.IsTransient())
-	{
-		// TODO: this should ideally be set in platform-independent code, since this tracking is for the high level
-		NewBuffer->SetCommitted(false);
-	}
-
-	return NewBuffer;
-}
-
-FStructuredBufferRHIRef FD3D12DynamicRHI::RHICreateStructuredBuffer(uint32 Stride, uint32 Size, uint32 InUsage, ERHIAccess InResourceState, FRHIResourceCreateInfo& CreateInfo)
-{
-	// Check for values that will cause D3D calls to fail
-	check(Size / Stride > 0 && Size % Stride == 0);
-
-	const D3D12_RESOURCE_DESC Desc = CreateStructuredBufferResourceDesc(Size, InUsage);
-
-	// Structured buffers, non-byte address buffers, need to be aligned to their stride to ensure that they
-	// can be addressed correctly with element based offsets.
-	const uint32 Alignment = ((InUsage & (BUF_ByteAddressBuffer | BUF_DrawIndirect)) == 0) ? Stride : 4;
-
-	FD3D12Buffer* NewBuffer = GetAdapter().CreateRHIBuffer(nullptr, Desc, Alignment, Stride, Size, InUsage | BUF_StructuredBuffer, ED3D12ResourceStateMode::Default, InResourceState, CreateInfo);
-	if (NewBuffer->ResourceLocation.IsTransient())
-	{
-		// TODO: this should ideally be set in platform-independent code, since this tracking is for the high level
-		NewBuffer->SetCommitted(false);
-	}
-
-	return NewBuffer;
-}
diff --git a/Engine/Source/Runtime/D3D12RHI/Private/D3D12VertexBuffer.cpp b/Engine/Source/Runtime/D3D12RHI/Private/D3D12VertexBuffer.cpp
--- a/Engine/Source/Runtime/D3D12RHI/Private/D3D12VertexBuffer.cpp
+++ b/Engine/Source/Runtime/D3D12RHI/Private/D3D12VertexBuffer.cpp
@@ -6,82 +6,6 @@
 
 #include "D3D12RHIPrivate.h"
 
-D3D12_RESOURCE_DESC CreateVertexBufferResourceDesc(uint32 Size, uint32 InUsage)
-{
-	// Describe the vertex buffer.
-	D3D12_RESOURCE_DESC Desc = CD3DX12_RESOURCE_DESC::Buffer(Size);
-
-	if (InUsage & BUF_UnorderedAccess)
-	{
-		Desc.Flags |= D3D12_RESOURCE_FLAG_ALLOW_UNORDERED_ACCESS;
-
-		static bool bRequiresRawView = (GMaxRHIFeatureLevel < ERHIFeatureLevel::SM5);
-		if (bRequiresRawView)
-		{
-			// Force the buffer to be a raw, byte address buffer
-			InUsage |= BUF_ByteAddressBuffer;
-		}
-	}
-
-	if ((InUsage & BUF_ShaderResource) == 0)
-	{
-		Desc.Flags |= D3D12_RESOURCE_FLAG_DENY_SHADER_RESOURCE;
-	}
-
-	if (InUsage & BUF_DrawIndirect)
-	{
-		Desc.Flags |= D3D12RHI_RESOURCE_FLAG_ALLOW_INDIRECT_BUFFER;
-	}
-
-	return Desc;
-}
-
-FVertexBufferRHIRef FD3D12DynamicRHI::RHICreateVertexBuffer(uint32 Size, uint32 InUsage, ERHIAccess InResourceState, FRHIResourceCreateInfo& CreateInfo)
-{
-	if (CreateInfo.bWithoutNativeResource)
-	{
-		return GetAdapter().CreateLinkedObject<FD3D12Buffer>(CreateInfo.GPUMask, [](FD3D12Device* Device)
-			{
-				return new FD3D12Buffer();
-			});
-	}
-
-	const D3D12_RESOURCE_DESC Desc = CreateVertexBufferResourceDesc(Size, InUsage);
-	const uint32 Alignment = 4;
-
-	FD3D12Buffer* Buffer = GetAdapter().CreateRHIBuffer(nullptr, Desc, Alignment, 0, Size, InUsage | BUF_VertexBuffer, ED3D12ResourceStateMode::Default, InResourceState, CreateInfo);
-	if (Buffer->ResourceLocation.IsTransient() )
-	{
-		// TODO: this should ideally be set in platform-independent code, since this tracking is for the high level
-		Buffer->SetCommitted(false);
-	}
-
-	return Buffer;
-}
-
-FVertexBufferRHIRef FD3D12DynamicRHI::CreateVertexBuffer_RenderThread(FRHICommandListImmediate& RHICmdList, uint32 Size, uint32 InUsage, ERHIAccess InResourceState, FRHIResourceCreateInfo& CreateInfo)
-{	
-	if (CreateInfo.bWithoutNativeResource)
-	{
-		return GetAdapter().CreateLinkedObject<FD3D12Buffer>(CreateInfo.GPUMask, [](FD3D12Device* Device)
-			{
-				return new FD3D12Buffer();
-			});
-	}
-
-	const D3D12_RESOURCE_DESC Desc = CreateVertexBufferResourceDesc(Size, InUsage);
-	const uint32 Alignment = 4;
-
-	FD3D12Buffer* Buffer = GetAdapter().CreateRHIBuffer(&RHICmdList, Desc, Alignment, 0, Size, InUsage | BUF_VertexBuffer, ED3D12ResourceStateMode::Default, InResourceState, CreateInfo);
-	if (Buffer->ResourceLocation.IsTransient())
-	{
-		// TODO: this should ideally be set in platform-independent code, since this tracking is for the high level
-		Buffer->SetCommitted(false);
-	}
-
-	return Buffer;
-}
-
 void FD3D12DynamicRHI::RHICopyBuffer(FRHIBuffer* SourceBufferRHI, FRHIBuffer* DestBufferRHI)
 {
 	FD3D12Buffer* SrcBuffer = FD3D12DynamicRHI::ResourceCast(SourceBufferRHI);
@@ -291,19 +215,3 @@ void FD3D12CommandContext::RHICopyBufferRegions(const TArrayView<const FCopyBuff
 	TransitionResources(CommandListHandle, DstBuffers, EBatchCopyState::FinalizeDest);
 }
 #endif // D3D12_RHI_RAYTRACING
-
-FVertexBufferRHIRef FD3D12DynamicRHI::CreateAndLockVertexBuffer_RenderThread(FRHICommandListImmediate& RHICmdList, uint32 Size, uint32 InUsage, ERHIAccess InResourceState, FRHIResourceCreateInfo& CreateInfo, void*& OutDataBuffer)
-{
-	const D3D12_RESOURCE_DESC Desc = CreateVertexBufferResourceDesc(Size, InUsage);
-	const uint32 Alignment = 4;
-
-	FD3D12Buffer* Buffer = GetAdapter().CreateRHIBuffer(nullptr, Desc, Alignment, 0, Size, InUsage | BUF_VertexBuffer, ED3D12ResourceStateMode::Default, InResourceState, CreateInfo);
-	if (Buffer->ResourceLocation.IsTransient())
-	{
-		// TODO: this should ideally be set in platform-independent code, since this tracking is for the high level
-		Buffer->SetCommitted(false);
-	}
-	OutDataBuffer = LockBuffer(&RHICmdList, Buffer, Buffer->GetSize(), Buffer->GetUsage(), 0, Size, RLM_WriteOnly);
-
-	return Buffer;
-}
diff --git a/Engine/Source/Runtime/NullDrv/Public/NullRHI.h b/Engine/Source/Runtime/NullDrv/Public/NullRHI.h
--- a/Engine/Source/Runtime/NullDrv/Public/NullRHI.h
+++ b/Engine/Source/Runtime/NullDrv/Public/NullRHI.h
@@ -124,22 +124,13 @@ public:
 
 	}
 
-	virtual FIndexBufferRHIRef RHICreateIndexBuffer(uint32 Stride, uint32 Size, uint32 InUsage, ERHIAccess InResourceState , FRHIResourceCreateInfo& CreateInfo) final override
+	virtual FBufferRHIRef RHICreateBuffer(uint32 Size, EBufferUsageFlags Usage, uint32 Stride, ERHIAccess ResourceState , FRHIResourceCreateInfo& CreateInfo) final override
 	{ 
 		if(CreateInfo.ResourceArray) 
 		{ 
 			CreateInfo.ResourceArray->Discard(); 
 		} 
-		return new FRHIIndexBuffer(Size, InUsage, Stride); 
-	}
-
-	virtual FVertexBufferRHIRef RHICreateVertexBuffer(uint32 Size, uint32 InUsage, ERHIAccess InResourceState, FRHIResourceCreateInfo& CreateInfo) final override
-	{ 
-		if(CreateInfo.ResourceArray) 
-		{ 
-			CreateInfo.ResourceArray->Discard(); 
-		} 
-		return new FRHIVertexBuffer(Size, InUsage, 0); 
+		return new FRHIBuffer(Size, Usage, Stride);
 	}
 
 	virtual void RHICopyBuffer(FRHIBuffer* SourceBuffer, FRHIBuffer* DestBuffer) final override
@@ -147,15 +138,6 @@ public:
 
 	}
 
-	virtual FStructuredBufferRHIRef RHICreateStructuredBuffer(uint32 Stride, uint32 Size, uint32 InUsage, ERHIAccess InResourceState, FRHIResourceCreateInfo& CreateInfo) final override
-	{ 
-		if(CreateInfo.ResourceArray) 
-		{ 
-			CreateInfo.ResourceArray->Discard(); 
-		} 
-		return new FRHIStructuredBuffer(Stride,Size,InUsage); 
-	}
-
 	virtual void* LockBuffer_BottomOfPipe(FRHICommandListImmediate& RHICmdList, FRHIBuffer* Buffer, uint32 Offset, uint32 SizeRHI, EResourceLockMode LockMode) final override
 	{ 
 		return GetStaticBuffer(Buffer->GetSize());
diff --git a/Engine/Source/Runtime/OpenGLDrv/Private/OpenGLIndexBuffer.cpp b/Engine/Source/Runtime/OpenGLDrv/Private/OpenGLIndexBuffer.cpp
--- a/Engine/Source/Runtime/OpenGLDrv/Private/OpenGLIndexBuffer.cpp
+++ b/Engine/Source/Runtime/OpenGLDrv/Private/OpenGLIndexBuffer.cpp
@@ -8,37 +8,6 @@
 #include "Containers/ResourceArray.h"
 #include "OpenGLDrv.h"
 
-FIndexBufferRHIRef FOpenGLDynamicRHI::RHICreateIndexBuffer(uint32 Stride,uint32 Size, uint32 InUsage, ERHIAccess InResourceState, FRHIResourceCreateInfo& CreateInfo)
-{
-	if (CreateInfo.bWithoutNativeResource)
-	{
-		return new FOpenGLBuffer();
-	}
-
-	const void *Data = NULL;
-
-	// If a resource array was provided for the resource, create the resource pre-populated
-	if(CreateInfo.ResourceArray)
-	{
-		check(Size == CreateInfo.ResourceArray->GetResourceDataSize());
-		Data = CreateInfo.ResourceArray->GetResourceData();
-	}
-
-	TRefCountPtr<FOpenGLBuffer> IndexBuffer = new FOpenGLBuffer(GL_ELEMENT_ARRAY_BUFFER, Stride, Size, InUsage | BUF_IndexBuffer, Data);
-
-	if (CreateInfo.ResourceArray)
-	{
-		CreateInfo.ResourceArray->Discard();
-	}
-
-	return IndexBuffer.GetReference();
-}
-
-FIndexBufferRHIRef FOpenGLDynamicRHI::CreateIndexBuffer_RenderThread(class FRHICommandListImmediate& RHICmdList, uint32 Stride, uint32 Size, uint32 InUsage, ERHIAccess InResourceState, FRHIResourceCreateInfo& CreateInfo)
-{
-	return this->RHICreateIndexBuffer(Stride, Size, InUsage, InResourceState, CreateInfo);
-}
-
 void FOpenGLDynamicRHI::RHITransferBufferUnderlyingResource(FRHIBuffer* DestBuffer, FRHIBuffer* SrcBuffer)
 {
 	VERIFY_GL_SCOPE();
diff --git a/Engine/Source/Runtime/OpenGLDrv/Private/OpenGLStructuredBuffer.cpp b/Engine/Source/Runtime/OpenGLDrv/Private/OpenGLStructuredBuffer.cpp
--- a/Engine/Source/Runtime/OpenGLDrv/Private/OpenGLStructuredBuffer.cpp
+++ b/Engine/Source/Runtime/OpenGLDrv/Private/OpenGLStructuredBuffer.cpp
@@ -7,20 +7,3 @@
 #include "CoreMinimal.h"
 #include "Containers/ResourceArray.h"
 #include "OpenGLDrv.h"
-
-FStructuredBufferRHIRef FOpenGLDynamicRHI::RHICreateStructuredBuffer(uint32 Stride,uint32 Size,uint32 InUsage, ERHIAccess InResourceState, FRHIResourceCreateInfo& CreateInfo)
-{
-	VERIFY_GL_SCOPE();
-
-	const void *Data = NULL;
-
-	// If a resource array was provided for the resource, create the resource pre-populated
-	if(CreateInfo.ResourceArray)
-	{
-		check(Size == CreateInfo.ResourceArray->GetResourceDataSize());
-		Data = CreateInfo.ResourceArray->GetResourceData();
-	}
-
-	TRefCountPtr<FOpenGLBuffer> StructuredBuffer = new FOpenGLBuffer(GL_ARRAY_BUFFER, Stride, Size, InUsage | BUF_StructuredBuffer, Data);
-	return StructuredBuffer.GetReference();
-}
diff --git a/Engine/Source/Runtime/OpenGLDrv/Private/OpenGLVertexBuffer.cpp b/Engine/Source/Runtime/OpenGLDrv/Private/OpenGLVertexBuffer.cpp
--- a/Engine/Source/Runtime/OpenGLDrv/Private/OpenGLVertexBuffer.cpp
+++ b/Engine/Source/Runtime/OpenGLDrv/Private/OpenGLVertexBuffer.cpp
@@ -176,7 +176,7 @@ void BeginFrame_VertexBufferCleanup()
 	FrameBytes = 0;
 }
 
-FVertexBufferRHIRef FOpenGLDynamicRHI::RHICreateVertexBuffer(uint32 Size, uint32 InUsage, ERHIAccess InResourceState, FRHIResourceCreateInfo& CreateInfo)
+FBufferRHIRef FOpenGLDynamicRHI::RHICreateBuffer(uint32 Size, EBufferUsageFlags Usage, uint32 Stride, ERHIAccess InResourceState, FRHIResourceCreateInfo& CreateInfo)
 {
 	if (CreateInfo.bWithoutNativeResource)
 	{
@@ -192,14 +192,14 @@ FVertexBufferRHIRef FOpenGLDynamicRHI::RHICreateVertexBuffer(uint32 Size, uint32
 		Data = CreateInfo.ResourceArray->GetResourceData();
 	}
 
-	TRefCountPtr<FOpenGLBuffer> VertexBuffer = new FOpenGLBuffer(GL_ARRAY_BUFFER, 0, Size, InUsage | BUF_VertexBuffer, Data);
+	TRefCountPtr<FOpenGLBuffer> Buffer = new FOpenGLBuffer((Usage & BUF_IndexBuffer)? GL_ARRAY_BUFFER : GL_ELEMENT_ARRAY_BUFFER, Stride, Size, Usage, Data);
 	
 	if (CreateInfo.ResourceArray)
 	{
 		CreateInfo.ResourceArray->Discard();
 	}
 	
-	return VertexBuffer.GetReference();
+	return Buffer.GetReference();
 }
 
 void* FOpenGLDynamicRHI::LockBuffer_BottomOfPipe(FRHICommandListImmediate& RHICmdList, FRHIBuffer* BufferRHI, uint32 Offset, uint32 Size, EResourceLockMode LockMode)
diff --git a/Engine/Source/Runtime/OpenGLDrv/Public/OpenGLDrv.h b/Engine/Source/Runtime/OpenGLDrv/Public/OpenGLDrv.h
--- a/Engine/Source/Runtime/OpenGLDrv/Public/OpenGLDrv.h
+++ b/Engine/Source/Runtime/OpenGLDrv/Public/OpenGLDrv.h
@@ -347,11 +347,9 @@ public:
 
 	virtual FUniformBufferRHIRef RHICreateUniformBuffer(const void* Contents, const FRHIUniformBufferLayout& Layout, EUniformBufferUsage Usage, EUniformBufferValidation Validation) final override;
 	virtual void RHIUpdateUniformBuffer(FRHIUniformBuffer* UniformBufferRHI, const void* Contents) final override;
-	virtual FIndexBufferRHIRef RHICreateIndexBuffer(uint32 Stride, uint32 Size, uint32 InUsage, ERHIAccess InResourceState, FRHIResourceCreateInfo& CreateInfo) final override;
-	virtual FVertexBufferRHIRef RHICreateVertexBuffer(uint32 Size, uint32 InUsage, ERHIAccess InResourceState, FRHIResourceCreateInfo& CreateInfo) final override;
+	virtual FBufferRHIRef RHICreateBuffer(uint32 Size, EBufferUsageFlags Usage, uint32 Stride, ERHIAccess ResourceState, FRHIResourceCreateInfo& CreateInfo) final override;
 	virtual void RHICopyBuffer(FRHIBuffer* SourceBuffer, FRHIBuffer* DestBuffer) final override;
 	virtual void RHITransferBufferUnderlyingResource(FRHIBuffer* DestBuffer, FRHIBuffer* SrcBuffer) final override;
-	virtual FStructuredBufferRHIRef RHICreateStructuredBuffer(uint32 Stride, uint32 Size, uint32 InUsage, ERHIAccess InResourceState, FRHIResourceCreateInfo& CreateInfo) final override;
 	virtual void* LockBuffer_BottomOfPipe(FRHICommandListImmediate& RHICmdList, FRHIBuffer* Buffer, uint32 Offset, uint32 Size, EResourceLockMode LockMode) final override;
 	virtual void UnlockBuffer_BottomOfPipe(FRHICommandListImmediate& RHICmdList, FRHIBuffer* Buffer) final override;
 	virtual FUnorderedAccessViewRHIRef RHICreateUnorderedAccessView(FRHIBuffer* Buffer, bool bUseUAVCounter, bool bAppendBuffer) final override;
@@ -760,18 +758,9 @@ public:
 	virtual void* LockTexture2DArray_RenderThread(class FRHICommandListImmediate& RHICmdList, FRHITexture2DArray* Texture, uint32 ArrayIndex, uint32 MipIndex, EResourceLockMode LockMode, uint32& DestStride, bool bLockWithinMiptail) final override;
 	virtual void UnlockTexture2DArray_RenderThread(class FRHICommandListImmediate& RHICmdList, FRHITexture2DArray* Texture, uint32 ArrayIndex, uint32 MipIndex, bool bLockWithinMiptail) final override;
 
-	virtual FIndexBufferRHIRef CreateIndexBuffer_RenderThread(class FRHICommandListImmediate& RHICmdList, uint32 Stride, uint32 Size, uint32 InUsage, ERHIAccess InResourceState, FRHIResourceCreateInfo& CreateInfo) final override;
-
-	virtual FVertexBufferRHIRef CreateVertexBuffer_RenderThread(class FRHICommandListImmediate& RHICmdList, uint32 Size, uint32 InUsage, ERHIAccess InResourceState, FRHIResourceCreateInfo& CreateInfo) final override
+	virtual FBufferRHIRef CreateBuffer_RenderThread(class FRHICommandListImmediate& RHICmdList, uint32 Size, EBufferUsageFlags Usage, uint32 Stride, ERHIAccess ResourceState, FRHIResourceCreateInfo& CreateInfo) final override
 	{
-		return this->RHICreateVertexBuffer(Size, InUsage, InResourceState, CreateInfo);
-	}
-
-	virtual FStructuredBufferRHIRef CreateStructuredBuffer_RenderThread(class FRHICommandListImmediate& RHICmdList, uint32 Stride, uint32 Size, uint32 InUsage, ERHIAccess InResourceState, FRHIResourceCreateInfo& CreateInfo) final override
-	{
-		RHITHREAD_GLCOMMAND_PROLOGUE();
-		return this->RHICreateStructuredBuffer(Stride, Size, InUsage, InResourceState, CreateInfo);
-		RHITHREAD_GLCOMMAND_EPILOGUE_RETURN(FStructuredBufferRHIRef);
+		return this->RHICreateBuffer(Size, Usage, Stride, ResourceState, CreateInfo);
 	}
 
 	virtual FShaderResourceViewRHIRef CreateShaderResourceView_RenderThread(class FRHICommandListImmediate& RHICmdList, FRHIVertexBuffer* VertexBuffer, uint32 Stride, uint8 Format) final override
diff --git a/Engine/Source/Runtime/RHI/Private/RHICommandList.cpp b/Engine/Source/Runtime/RHI/Private/RHICommandList.cpp
--- a/Engine/Source/Runtime/RHI/Private/RHICommandList.cpp
+++ b/Engine/Source/Runtime/RHI/Private/RHICommandList.cpp
@@ -2096,41 +2096,27 @@ void IRHIComputeContext::PopBreadcrumbRHIThread()
 }
 ///////// Pass through functions that allow RHIs to optimize certain calls.
 
-FVertexBufferRHIRef FDynamicRHI::CreateAndLockVertexBuffer_RenderThread(class FRHICommandListImmediate& RHICmdList, uint32 Size, uint32 InUsage, ERHIAccess InResourceState, FRHIResourceCreateInfo& CreateInfo, void*& OutDataBuffer)
+FBufferRHIRef FDynamicRHI::CreateAndLockVertexBuffer_RenderThread(class FRHICommandListImmediate& RHICmdList, uint32 Size, uint32 Usage, ERHIAccess ResourceState, FRHIResourceCreateInfo& CreateInfo, void*& OutDataBuffer)
 {
-	FVertexBufferRHIRef VertexBuffer = CreateVertexBuffer_RenderThread(RHICmdList, Size, InUsage, InResourceState, CreateInfo);
+	FBufferRHIRef VertexBuffer = CreateBuffer_RenderThread(RHICmdList, Size, (EBufferUsageFlags)Usage | BUF_VertexBuffer, 0, ResourceState, CreateInfo);
 	OutDataBuffer = RHILockBuffer(RHICmdList, VertexBuffer, 0, Size, RLM_WriteOnly);
 
 	return VertexBuffer;
 }
 
-FIndexBufferRHIRef FDynamicRHI::CreateAndLockIndexBuffer_RenderThread(class FRHICommandListImmediate& RHICmdList, uint32 Stride, uint32 Size, uint32 InUsage, ERHIAccess InResourceState, FRHIResourceCreateInfo& CreateInfo, void*& OutDataBuffer)
+FBufferRHIRef FDynamicRHI::CreateAndLockIndexBuffer_RenderThread(class FRHICommandListImmediate& RHICmdList, uint32 Stride, uint32 Size, uint32 Usage, ERHIAccess ResourceState, FRHIResourceCreateInfo& CreateInfo, void*& OutDataBuffer)
 {
-	FIndexBufferRHIRef IndexBuffer = CreateIndexBuffer_RenderThread(RHICmdList, Stride, Size, InUsage, InResourceState, CreateInfo);
+	FBufferRHIRef IndexBuffer = CreateBuffer_RenderThread(RHICmdList, Size, (EBufferUsageFlags)Usage | BUF_IndexBuffer, Stride, ResourceState, CreateInfo);
 	OutDataBuffer = RHILockBuffer(RHICmdList, IndexBuffer, 0, Size, RLM_WriteOnly);
 	return IndexBuffer;
 }
 
 
-FVertexBufferRHIRef FDynamicRHI::CreateVertexBuffer_RenderThread(class FRHICommandListImmediate& RHICmdList, uint32 Size, uint32 InUsage, ERHIAccess InResourceState, FRHIResourceCreateInfo& CreateInfo)
+FBufferRHIRef FDynamicRHI::CreateBuffer_RenderThread(class FRHICommandListImmediate& RHICmdList, uint32 Size, EBufferUsageFlags Usage, uint32 Stride, ERHIAccess ResourceState, FRHIResourceCreateInfo& CreateInfo)
 {
-	CSV_SCOPED_TIMING_STAT(RHITStalls, CreateVertexBuffer_RenderThread);
+	CSV_SCOPED_TIMING_STAT(RHITStalls, CreateBuffer_RenderThread);
 	FScopedRHIThreadStaller StallRHIThread(RHICmdList);
-	return GDynamicRHI->RHICreateVertexBuffer(Size, InUsage, InResourceState, CreateInfo);
-}
-
-FStructuredBufferRHIRef FDynamicRHI::CreateStructuredBuffer_RenderThread(class FRHICommandListImmediate& RHICmdList, uint32 Stride, uint32 Size, uint32 InUsage, ERHIAccess InResourceState, FRHIResourceCreateInfo& CreateInfo)
-{
-	CSV_SCOPED_TIMING_STAT(RHITStalls, CreateStructuredBuffer_RenderThread);
-	FScopedRHIThreadStaller StallRHIThread(RHICmdList);
-	return GDynamicRHI->RHICreateStructuredBuffer(Stride, Size, InUsage, InResourceState, CreateInfo);
-}
-
-FIndexBufferRHIRef FDynamicRHI::CreateIndexBuffer_RenderThread(class FRHICommandListImmediate& RHICmdList, uint32 Stride, uint32 Size, uint32 InUsage, ERHIAccess InResourceState, FRHIResourceCreateInfo& CreateInfo)
-{
-	CSV_SCOPED_TIMING_STAT(RHITStalls, CreateIndexBuffer_RenderThread);
-	FScopedRHIThreadStaller StallRHIThread(RHICmdList);
-	return GDynamicRHI->RHICreateIndexBuffer(Stride, Size, InUsage, InResourceState, CreateInfo);
+	return GDynamicRHI->RHICreateBuffer(Size, Usage, Stride, ResourceState, CreateInfo);
 }
 
 FShaderResourceViewRHIRef FDynamicRHI::CreateShaderResourceView_RenderThread(class FRHICommandListImmediate& RHICmdList, FRHIVertexBuffer* VertexBuffer, uint32 Stride, uint8 Format)
diff --git a/Engine/Source/Runtime/RHI/Public/DynamicRHI.h b/Engine/Source/Runtime/RHI/Public/DynamicRHI.h
--- a/Engine/Source/Runtime/RHI/Public/DynamicRHI.h
+++ b/Engine/Source/Runtime/RHI/Public/DynamicRHI.h
@@ -521,15 +521,6 @@ public:
 
 	virtual void RHIUpdateUniformBuffer(FRHIUniformBuffer* UniformBufferRHI, const void* Contents) = 0;
 
-	// FlushType: Wait RHI Thread
-	virtual FIndexBufferRHIRef RHICreateIndexBuffer(uint32 Stride, uint32 Size, uint32 InUsage, ERHIAccess InResourceState, FRHIResourceCreateInfo& CreateInfo) = 0;
-
-	/**
-	* @param ResourceArray - An optional pointer to a resource array containing the resource's data.
-	*/
-	// FlushType: Wait RHI Thread
-	virtual FVertexBufferRHIRef RHICreateVertexBuffer(uint32 Size, uint32 InUsage, ERHIAccess InResourceState, FRHIResourceCreateInfo& CreateInfo) = 0;
-
 	/** Copies the contents of one buffer to another buffer. They must have identical sizes. */
 	// FlushType: Flush Immediate (seems dangerous)
 	virtual void RHICopyBuffer(FRHIBuffer* SourceBuffer, FRHIBuffer* DestBuffer) = 0;
@@ -545,7 +536,7 @@ public:
 	* @param ResourceArray - An optional pointer to a resource array containing the resource's data.
 	*/
 	// FlushType: Wait RHI Thread
-	virtual FStructuredBufferRHIRef RHICreateStructuredBuffer(uint32 Stride, uint32 Size, uint32 InUsage, ERHIAccess InResourceState, FRHIResourceCreateInfo& CreateInfo) = 0;
+	virtual FBufferRHIRef RHICreateBuffer(uint32 Size, EBufferUsageFlags Usage, uint32 Stride, ERHIAccess InResourceState, FRHIResourceCreateInfo& CreateInfo) = 0;
 
 	// FlushType: Flush RHI Thread
 	virtual void* RHILockBuffer(FRHICommandListImmediate& RHICmdList, FRHIBuffer* Buffer, uint32 Offset, uint32 Size, EResourceLockMode LockMode);
@@ -1256,19 +1247,17 @@ public:
 	FORCEINLINE IRHICommandContextContainer* RHIGetCommandContextContainer(int32 Index, int32 Num, FRHIGPUMask GPUMask) { return RHIGetCommandContextContainer(Index, Num); }
 #endif
 
-	///////// Pass through functions that allow RHIs to optimize certain calls.
-	virtual FVertexBufferRHIRef CreateAndLockVertexBuffer_RenderThread(class FRHICommandListImmediate& RHICmdList, uint32 Size, uint32 InUsage, ERHIAccess InResourceState, FRHIResourceCreateInfo& CreateInfo, void*& OutDataBuffer);
-	virtual FIndexBufferRHIRef CreateAndLockIndexBuffer_RenderThread(class FRHICommandListImmediate& RHICmdList, uint32 Stride, uint32 Size, uint32 InUsage, ERHIAccess InResourceState, FRHIResourceCreateInfo& CreateInfo, void*& OutDataBuffer);
+	// Legacy helper functions
+	FBufferRHIRef CreateAndLockVertexBuffer_RenderThread(class FRHICommandListImmediate& RHICmdList, uint32 Size, uint32 Usage, ERHIAccess ResourceState, FRHIResourceCreateInfo& CreateInfo, void*& OutDataBuffer);
+	FBufferRHIRef CreateAndLockIndexBuffer_RenderThread(class FRHICommandListImmediate& RHICmdList, uint32 Stride, uint32 Size, uint32 Usage, ERHIAccess ResourceState, FRHIResourceCreateInfo& CreateInfo, void*& OutDataBuffer);
 
-	virtual FVertexBufferRHIRef CreateVertexBuffer_RenderThread(class FRHICommandListImmediate& RHICmdList, uint32 Size, uint32 InUsage, ERHIAccess InResourceState, FRHIResourceCreateInfo& CreateInfo);
-	virtual FStructuredBufferRHIRef CreateStructuredBuffer_RenderThread(class FRHICommandListImmediate& RHICmdList, uint32 Stride, uint32 Size, uint32 InUsage, ERHIAccess InResourceState, FRHIResourceCreateInfo& CreateInfo);
+	virtual FBufferRHIRef CreateBuffer_RenderThread(class FRHICommandListImmediate& RHICmdList, uint32 Size, EBufferUsageFlags Usage, uint32 Stride, ERHIAccess ResourceState, FRHIResourceCreateInfo& CreateInfo);
 	virtual FShaderResourceViewRHIRef CreateShaderResourceView_RenderThread(class FRHICommandListImmediate& RHICmdList, FRHIVertexBuffer* VertexBuffer, uint32 Stride, uint8 Format);
 	virtual FShaderResourceViewRHIRef CreateShaderResourceView_RenderThread(class FRHICommandListImmediate& RHICmdList, const FShaderResourceViewInitializer& Initializer);
 	virtual FShaderResourceViewRHIRef CreateShaderResourceView_RenderThread(class FRHICommandListImmediate& RHICmdList, FRHIIndexBuffer* Buffer);
 	virtual FTexture2DRHIRef AsyncReallocateTexture2D_RenderThread(class FRHICommandListImmediate& RHICmdList, FRHITexture2D* Texture2D, int32 NewMipCount, int32 NewSizeX, int32 NewSizeY, FThreadSafeCounter* RequestStatus);
 	virtual ETextureReallocationStatus FinalizeAsyncReallocateTexture2D_RenderThread(class FRHICommandListImmediate& RHICmdList, FRHITexture2D* Texture2D, bool bBlockUntilCompleted);
 	virtual ETextureReallocationStatus CancelAsyncReallocateTexture2D_RenderThread(class FRHICommandListImmediate& RHICmdList, FRHITexture2D* Texture2D, bool bBlockUntilCompleted);
-	virtual FIndexBufferRHIRef CreateIndexBuffer_RenderThread(class FRHICommandListImmediate& RHICmdList, uint32 Stride, uint32 Size, uint32 InUsage, ERHIAccess InResourceState, FRHIResourceCreateInfo& CreateInfo);
 	virtual FVertexShaderRHIRef CreateVertexShader_RenderThread(class FRHICommandListImmediate& RHICmdList, TArrayView<const uint8> Code, const FSHAHash& Hash);
 	virtual FPixelShaderRHIRef CreatePixelShader_RenderThread(class FRHICommandListImmediate& RHICmdList, TArrayView<const uint8> Code, const FSHAHash& Hash);
 	virtual FGeometryShaderRHIRef CreateGeometryShader_RenderThread(class FRHICommandListImmediate& RHICmdList, TArrayView<const uint8> Code, const FSHAHash& Hash);
diff --git a/Engine/Source/Runtime/RHI/Public/RHICommandList.h b/Engine/Source/Runtime/RHI/Public/RHICommandList.h
--- a/Engine/Source/Runtime/RHI/Public/RHICommandList.h
+++ b/Engine/Source/Runtime/RHI/Public/RHICommandList.h
@@ -4080,8 +4080,9 @@ public:
 
 	FORCEINLINE FIndexBufferRHIRef CreateAndLockIndexBuffer(uint32 Stride, uint32 Size, uint32 InUsage, FRHIResourceCreateInfo& CreateInfo, void*& OutDataBuffer)
 	{
-		ERHIAccess ResourceState = RHIGetDefaultResourceState((EBufferUsageFlags) InUsage | BUF_IndexBuffer, true);
-		return CreateAndLockIndexBuffer(Stride, Size, (EBufferUsageFlags) InUsage, ResourceState, CreateInfo, OutDataBuffer);
+		EBufferUsageFlags Usage = (EBufferUsageFlags) InUsage | BUF_IndexBuffer;
+		ERHIAccess ResourceState = RHIGetDefaultResourceState(Usage, true);
+		return CreateAndLockIndexBuffer(Stride, Size, Usage, ResourceState, CreateInfo, OutDataBuffer);
 	}
 	
 	UE_DEPRECATED(5.0, "Buffer locks have been unified. Use LockBuffer() instead.")
@@ -4113,8 +4114,9 @@ public:
 
 	FORCEINLINE FVertexBufferRHIRef CreateAndLockVertexBuffer(uint32 Size, uint32 InUsage, FRHIResourceCreateInfo& CreateInfo, void*& OutDataBuffer)
 	{
-		ERHIAccess ResourceState = RHIGetDefaultResourceState((EBufferUsageFlags) InUsage | BUF_VertexBuffer, true);
-		return CreateAndLockVertexBuffer(Size, (EBufferUsageFlags) InUsage, ResourceState, CreateInfo, OutDataBuffer);
+		EBufferUsageFlags Usage = (EBufferUsageFlags) InUsage | BUF_VertexBuffer;
+		ERHIAccess ResourceState = RHIGetDefaultResourceState(Usage, true);
+		return CreateAndLockVertexBuffer(Size, Usage, ResourceState, CreateInfo, OutDataBuffer);
 	}
 
 	UE_DEPRECATED(5.0, "Buffer locks have been unified. Use LockBuffer() instead.")
@@ -5100,26 +5102,28 @@ FORCEINLINE FIndexBufferRHIRef RHICreateAndLockIndexBuffer(uint32 Stride, uint32
 
 FORCEINLINE FIndexBufferRHIRef RHICreateIndexBuffer(uint32 Stride, uint32 Size, uint32 InUsage, ERHIAccess InResourceState, FRHIResourceCreateInfo& CreateInfo)
 {
-	return GDynamicRHI->CreateIndexBuffer_RenderThread(FRHICommandListExecutor::GetImmediateCommandList(), Stride, Size, InUsage, InResourceState, CreateInfo);
+	return GDynamicRHI->CreateBuffer_RenderThread(FRHICommandListExecutor::GetImmediateCommandList(), Size, (EBufferUsageFlags) InUsage | BUF_IndexBuffer, Stride, InResourceState, CreateInfo);
 }
 
 FORCEINLINE FIndexBufferRHIRef RHIAsyncCreateIndexBuffer(uint32 Stride, uint32 Size, uint32 InUsage, ERHIAccess InResourceState, FRHIResourceCreateInfo& CreateInfo)
 {
-	return GDynamicRHI->RHICreateIndexBuffer(Stride, Size, InUsage, InResourceState, CreateInfo);
+	return GDynamicRHI->RHICreateBuffer(Size, (EBufferUsageFlags)InUsage, Stride, InResourceState, CreateInfo);
 }
 
 FORCEINLINE FIndexBufferRHIRef RHICreateIndexBuffer(uint32 Stride, uint32 Size, uint32 InUsage, FRHIResourceCreateInfo& CreateInfo)
 {
 	bool bHasInitialData = CreateInfo.BulkData != nullptr;
-	ERHIAccess ResourceState = RHIGetDefaultResourceState((EBufferUsageFlags)InUsage | BUF_IndexBuffer, bHasInitialData);
-	return RHICreateIndexBuffer(Stride, Size, InUsage, ResourceState, CreateInfo);
+	EBufferUsageFlags Usage = (EBufferUsageFlags) InUsage | BUF_IndexBuffer;
+	ERHIAccess ResourceState = RHIGetDefaultResourceState(Usage, bHasInitialData);
+	return RHICreateIndexBuffer(Stride, Size, Usage, ResourceState, CreateInfo);
 }
 
 FORCEINLINE FIndexBufferRHIRef RHIAsyncCreateIndexBuffer(uint32 Stride, uint32 Size, uint32 InUsage, FRHIResourceCreateInfo& CreateInfo)
 {
 	bool bHasInitialData = CreateInfo.BulkData != nullptr;
-	ERHIAccess ResourceState = RHIGetDefaultResourceState((EBufferUsageFlags)InUsage | BUF_IndexBuffer, bHasInitialData);
-	return RHIAsyncCreateIndexBuffer(Stride, Size, InUsage, ResourceState, CreateInfo);
+	EBufferUsageFlags Usage = (EBufferUsageFlags) InUsage | BUF_IndexBuffer;
+	ERHIAccess ResourceState = RHIGetDefaultResourceState(Usage, bHasInitialData);
+	return RHIAsyncCreateIndexBuffer(Stride, Size, Usage, ResourceState, CreateInfo);
 }
 
 UE_DEPRECATED(5.0, "Buffer locks have been unified. Use RHILockBuffer() instead.")
@@ -5141,26 +5145,28 @@ FORCEINLINE FVertexBufferRHIRef RHICreateAndLockVertexBuffer(uint32 Size, uint32
 
 FORCEINLINE FVertexBufferRHIRef RHICreateVertexBuffer(uint32 Size, uint32 InUsage, ERHIAccess InResourceState, FRHIResourceCreateInfo& CreateInfo)
 {
-	return GDynamicRHI->CreateVertexBuffer_RenderThread(FRHICommandListExecutor::GetImmediateCommandList(), Size, InUsage, InResourceState, CreateInfo);
+	return GDynamicRHI->CreateBuffer_RenderThread(FRHICommandListExecutor::GetImmediateCommandList(), Size, (EBufferUsageFlags) InUsage | BUF_VertexBuffer, 0, InResourceState, CreateInfo);
 }
 
 FORCEINLINE FVertexBufferRHIRef RHIAsyncCreateVertexBuffer(uint32 Size, uint32 InUsage, ERHIAccess InResourceState, FRHIResourceCreateInfo& CreateInfo)
 {
-	return GDynamicRHI->RHICreateVertexBuffer(Size, InUsage, InResourceState, CreateInfo);
+	return GDynamicRHI->RHICreateBuffer(Size, (EBufferUsageFlags)InUsage, 0, InResourceState, CreateInfo);
 }
 
 FORCEINLINE FVertexBufferRHIRef RHICreateVertexBuffer(uint32 Size, uint32 InUsage, FRHIResourceCreateInfo& CreateInfo)
 {
 	bool bHasInitialData = CreateInfo.BulkData != nullptr;
-	ERHIAccess ResourceState = RHIGetDefaultResourceState((EBufferUsageFlags)InUsage | BUF_VertexBuffer, bHasInitialData);
-	return RHICreateVertexBuffer(Size, InUsage, ResourceState, CreateInfo);
+	EBufferUsageFlags Usage = (EBufferUsageFlags) InUsage | BUF_VertexBuffer;
+	ERHIAccess ResourceState = RHIGetDefaultResourceState(Usage, bHasInitialData);
+	return RHICreateVertexBuffer(Size, Usage, ResourceState, CreateInfo);
 }
 
 FORCEINLINE FVertexBufferRHIRef RHIAsyncCreateVertexBuffer(uint32 Size, uint32 InUsage, FRHIResourceCreateInfo& CreateInfo)
 {
 	bool bHasInitialData = CreateInfo.BulkData != nullptr;
-	ERHIAccess ResourceState = RHIGetDefaultResourceState((EBufferUsageFlags)InUsage | BUF_VertexBuffer, bHasInitialData);
-	return RHIAsyncCreateVertexBuffer(Size, InUsage, ResourceState, CreateInfo);
+	EBufferUsageFlags Usage = (EBufferUsageFlags) InUsage | BUF_VertexBuffer;
+	ERHIAccess ResourceState = RHIGetDefaultResourceState(Usage, bHasInitialData);
+	return RHIAsyncCreateVertexBuffer(Size, Usage, ResourceState, CreateInfo);
 }
 
 UE_DEPRECATED(5.0, "Buffer locks have been unified. Use RHILockBuffer() instead.")
@@ -5177,14 +5183,15 @@ FORCEINLINE void RHIUnlockVertexBuffer(FRHIVertexBuffer* VertexBuffer)
 
 FORCEINLINE FStructuredBufferRHIRef RHICreateStructuredBuffer(uint32 Stride, uint32 Size, uint32 InUsage, ERHIAccess InResourceState, FRHIResourceCreateInfo& CreateInfo)
 {
-	return GDynamicRHI->CreateStructuredBuffer_RenderThread(FRHICommandListExecutor::GetImmediateCommandList(), Stride, Size, InUsage, InResourceState, CreateInfo);
+	return GDynamicRHI->CreateBuffer_RenderThread(FRHICommandListExecutor::GetImmediateCommandList(), Size, (EBufferUsageFlags) InUsage | BUF_StructuredBuffer, Stride, InResourceState, CreateInfo);
 }
 
 FORCEINLINE FStructuredBufferRHIRef RHICreateStructuredBuffer(uint32 Stride, uint32 Size, uint32 InUsage, FRHIResourceCreateInfo& CreateInfo)
 {
 	bool bHasInitialData = CreateInfo.BulkData != nullptr;
-	ERHIAccess ResourceState = RHIGetDefaultResourceState((EBufferUsageFlags)InUsage | BUF_StructuredBuffer, bHasInitialData);
-	return RHICreateStructuredBuffer(Stride, Size, InUsage, ResourceState, CreateInfo);
+	EBufferUsageFlags Usage = (EBufferUsageFlags) InUsage | BUF_StructuredBuffer;
+	ERHIAccess ResourceState = RHIGetDefaultResourceState(Usage, bHasInitialData);
+	return RHICreateStructuredBuffer(Stride, Size, Usage, ResourceState, CreateInfo);
 }
 
 UE_DEPRECATED(5.0, "Buffer locks have been unified. Use RHILockBuffer() instead.")
diff --git a/Engine/Source/Runtime/RHI/Public/RHIDefinitions.h b/Engine/Source/Runtime/RHI/Public/RHIDefinitions.h
--- a/Engine/Source/Runtime/RHI/Public/RHIDefinitions.h
+++ b/Engine/Source/Runtime/RHI/Public/RHIDefinitions.h
@@ -1363,9 +1363,7 @@ enum ERHIResourceType
 	RRT_ComputeShader,
 	RRT_BoundShaderState,
 	RRT_UniformBuffer,
-	RRT_IndexBuffer,
-	RRT_VertexBuffer,
-	RRT_StructuredBuffer,
+	RRT_Buffer,
 	RRT_Texture,
 	RRT_Texture2D,
 	RRT_Texture2DArray,
diff --git a/Engine/Source/Runtime/RHI/Public/RHIResources.h b/Engine/Source/Runtime/RHI/Public/RHIResources.h
--- a/Engine/Source/Runtime/RHI/Public/RHIResources.h
+++ b/Engine/Source/Runtime/RHI/Public/RHIResources.h
@@ -1480,9 +1480,10 @@ typedef TRefCountPtr<FRHIRayTracingShader>          FRayTracingShaderRHIRef;
 typedef TRefCountPtr<FRHIComputeFence>	FComputeFenceRHIRef;
 typedef TRefCountPtr<FRHIBoundShaderState> FBoundShaderStateRHIRef;
 typedef TRefCountPtr<FRHIUniformBuffer> FUniformBufferRHIRef;
-typedef TRefCountPtr<FRHIIndexBuffer> FIndexBufferRHIRef;
-typedef TRefCountPtr<FRHIVertexBuffer> FVertexBufferRHIRef;
-typedef TRefCountPtr<FRHIStructuredBuffer> FStructuredBufferRHIRef;
+typedef TRefCountPtr<FRHIBuffer> FBufferRHIRef;
+typedef FBufferRHIRef FIndexBufferRHIRef;
+typedef FBufferRHIRef FVertexBufferRHIRef;
+typedef FBufferRHIRef FStructuredBufferRHIRef;
 typedef TRefCountPtr<FRHITexture> FTextureRHIRef;
 typedef TRefCountPtr<FRHITexture2D> FTexture2DRHIRef;
 typedef TRefCountPtr<FRHITexture2DArray> FTexture2DArrayRHIRef;
diff --git a/Engine/Source/Runtime/RHI/Public/RHIValidation.h b/Engine/Source/Runtime/RHI/Public/RHIValidation.h
--- a/Engine/Source/Runtime/RHI/Public/RHIValidation.h
+++ b/Engine/Source/Runtime/RHI/Public/RHIValidation.h
@@ -357,23 +357,15 @@ public:
 		UniformBufferRHI->UpdateAllocation(RenderThreadFrameID);
 	}
 
-	// FlushType: Wait RHI Thread
-	virtual FIndexBufferRHIRef RHICreateIndexBuffer(uint32 Stride, uint32 Size, uint32 InUsage, ERHIAccess InResourceState, FRHIResourceCreateInfo& CreateInfo) override final
-	{
-		FIndexBufferRHIRef IndexBuffer = RHI->RHICreateIndexBuffer(Stride, Size, InUsage, InResourceState, CreateInfo);
-		IndexBuffer->InitBarrierTracking(InResourceState, CreateInfo.DebugName);
-		return IndexBuffer;
-	}
-
 	/**
 	* @param ResourceArray - An optional pointer to a resource array containing the resource's data.
 	*/
 	// FlushType: Wait RHI Thread
-	virtual FVertexBufferRHIRef RHICreateVertexBuffer(uint32 Size, uint32 InUsage, ERHIAccess InResourceState, FRHIResourceCreateInfo& CreateInfo) override final
+	virtual FBufferRHIRef RHICreateBuffer(uint32 Size, EBufferUsageFlags Usage, uint32 Stride, ERHIAccess ResourceState, FRHIResourceCreateInfo& CreateInfo) override final
 	{
-		FVertexBufferRHIRef VertexBuffer = RHI->RHICreateVertexBuffer(Size, InUsage, InResourceState, CreateInfo);
-		VertexBuffer->InitBarrierTracking(InResourceState, CreateInfo.DebugName);
-		return VertexBuffer;
+		FBufferRHIRef Buffer = RHI->RHICreateBuffer(Size, Usage, Stride, ResourceState, CreateInfo);
+		Buffer->InitBarrierTracking(ResourceState, CreateInfo.DebugName);
+		return Buffer;
 	}
 
 	/** Copies the contents of one vertex buffer to another vertex buffer.  They must have identical sizes. */
@@ -383,17 +375,6 @@ public:
 		RHI->RHICopyBuffer(SourceBuffer, DestBuffer);
 	}
 
-	/**
-	* @param ResourceArray - An optional pointer to a resource array containing the resource's data.
-	*/
-	// FlushType: Wait RHI Thread
-	virtual FStructuredBufferRHIRef RHICreateStructuredBuffer(uint32 Stride, uint32 Size, uint32 InUsage, ERHIAccess InResourceState, FRHIResourceCreateInfo& CreateInfo) override final
-	{
-		FStructuredBufferRHIRef StructuredBuffer = RHI->RHICreateStructuredBuffer(Stride, Size, InUsage, InResourceState, CreateInfo);
-		StructuredBuffer->InitBarrierTracking(InResourceState, CreateInfo.DebugName);
-		return StructuredBuffer;
-	}
-
 	// FlushType: Flush RHI Thread
 	virtual void* RHILockBuffer(class FRHICommandListImmediate& RHICmdList, FRHIBuffer* Buffer, uint32 Offset, uint32 SizeRHI, EResourceLockMode LockMode) override final
 	{
@@ -1338,33 +1319,11 @@ public:
 		return RHI->RHIGetMinimumAlignmentForBufferBackedSRV(Format);
 	}
 
-	///////// Pass through functions that allow RHIs to optimize certain calls.
-	virtual FVertexBufferRHIRef CreateAndLockVertexBuffer_RenderThread(class FRHICommandListImmediate& RHICmdList, uint32 Size, uint32 InUsage, ERHIAccess InResourceState, FRHIResourceCreateInfo& CreateInfo, void*& OutDataBuffer) override final
+	virtual FBufferRHIRef CreateBuffer_RenderThread(class FRHICommandListImmediate& RHICmdList, uint32 Size, EBufferUsageFlags Usage, uint32 Stride, ERHIAccess ResourceState, FRHIResourceCreateInfo& CreateInfo) override final
 	{
-		FVertexBufferRHIRef VertexBuffer = RHI->CreateAndLockVertexBuffer_RenderThread(RHICmdList, Size, InUsage, InResourceState, CreateInfo, OutDataBuffer);
-		VertexBuffer->InitBarrierTracking(InResourceState, CreateInfo.DebugName);
-		return VertexBuffer;
-	}
-
-	virtual FIndexBufferRHIRef CreateAndLockIndexBuffer_RenderThread(class FRHICommandListImmediate& RHICmdList, uint32 Stride, uint32 Size, uint32 InUsage, ERHIAccess InResourceState, FRHIResourceCreateInfo& CreateInfo, void*& OutDataBuffer) override final
-	{
-		FIndexBufferRHIRef IndexBuffer = RHI->CreateAndLockIndexBuffer_RenderThread(RHICmdList, Stride, Size, InUsage, InResourceState, CreateInfo, OutDataBuffer);
-		IndexBuffer->InitBarrierTracking(InResourceState, CreateInfo.DebugName);
-		return IndexBuffer;
-	}
-
-	virtual FVertexBufferRHIRef CreateVertexBuffer_RenderThread(class FRHICommandListImmediate& RHICmdList, uint32 Size, uint32 InUsage, ERHIAccess InResourceState, FRHIResourceCreateInfo& CreateInfo) override final
-	{
-		FVertexBufferRHIRef VertexBuffer = RHI->CreateVertexBuffer_RenderThread(RHICmdList, Size, InUsage, InResourceState, CreateInfo);
-		VertexBuffer->InitBarrierTracking(InResourceState, CreateInfo.DebugName);
-		return VertexBuffer;
-	}
-
-	virtual FStructuredBufferRHIRef CreateStructuredBuffer_RenderThread(class FRHICommandListImmediate& RHICmdList, uint32 Stride, uint32 Size, uint32 InUsage, ERHIAccess InResourceState, FRHIResourceCreateInfo& CreateInfo) override final
-	{
-		FStructuredBufferRHIRef StructuredBuffer = RHI->CreateStructuredBuffer_RenderThread(RHICmdList, Stride, Size, InUsage, InResourceState, CreateInfo);
-		StructuredBuffer->InitBarrierTracking(InResourceState, CreateInfo.DebugName);
-		return StructuredBuffer;
+		FBufferRHIRef Buffer = RHI->CreateBuffer_RenderThread(RHICmdList, Size, Usage, Stride, ResourceState, CreateInfo);
+		Buffer->InitBarrierTracking(ResourceState, CreateInfo.DebugName);
+		return Buffer;
 	}
 
 	virtual FShaderResourceViewRHIRef CreateShaderResourceView_RenderThread(class FRHICommandListImmediate& RHICmdList, FRHIVertexBuffer* VertexBuffer, uint32 Stride, uint8 Format) override final
@@ -1409,13 +1368,6 @@ public:
 		return RHI->CancelAsyncReallocateTexture2D_RenderThread(RHICmdList, Texture2D, bBlockUntilCompleted);
 	}
 
-	virtual FIndexBufferRHIRef CreateIndexBuffer_RenderThread(class FRHICommandListImmediate& RHICmdList, uint32 Stride, uint32 Size, uint32 InUsage, ERHIAccess InResourceState, FRHIResourceCreateInfo& CreateInfo) override final
-	{
-		FIndexBufferRHIRef IndexBuffer = RHI->CreateIndexBuffer_RenderThread(RHICmdList, Stride, Size, InUsage, InResourceState, CreateInfo);
-		IndexBuffer->InitBarrierTracking(InResourceState, CreateInfo.DebugName);
-		return IndexBuffer;
-	}
-
 	virtual void* LockBuffer_BottomOfPipe(class FRHICommandListImmediate& RHICmdList, FRHIBuffer* Buffer, uint32 Offset, uint32 SizeRHI, EResourceLockMode LockMode) override final
 	{
 		return RHI->LockBuffer_BottomOfPipe(RHICmdList, Buffer, Offset, SizeRHI, LockMode);
diff --git a/Engine/Source/Runtime/VulkanRHI/Private/VulkanBarriers.cpp b/Engine/Source/Runtime/VulkanRHI/Private/VulkanBarriers.cpp
--- a/Engine/Source/Runtime/VulkanRHI/Private/VulkanBarriers.cpp
+++ b/Engine/Source/Runtime/VulkanRHI/Private/VulkanBarriers.cpp
@@ -595,20 +595,11 @@ void FVulkanDynamicRHI::RHICreateTransition(FRHITransition* Transition, ERHIPipe
 				Texture = FVulkanTextureBase::Cast(UAV->SourceTexture);
 				UnderlyingType = FRHITransitionInfo::EType::Texture;
 			}
-			else if (UAV->SourceIndexBuffer)
+			else if (UAV->SourceBuffer)
 			{
-				Buffer = UAV->SourceIndexBuffer;
-				UnderlyingType = FRHITransitionInfo::EType::Buffer;
-			}
-			else if (UAV->SourceVertexBuffer)
-			{
-				Buffer = UAV->SourceVertexBuffer;
-				UnderlyingType = FRHITransitionInfo::EType::Buffer;
-			}
-			else if (UAV->SourceStructuredBuffer)
-			{
-				Buffer = UAV->SourceStructuredBuffer;
+				Buffer = UAV->SourceBuffer;
 				UnderlyingType = FRHITransitionInfo::EType::Buffer;
+				UsageFlags = Buffer->GetBufferUsageFlags();
 			}
 			else
 			{
diff --git a/Engine/Source/Runtime/VulkanRHI/Private/VulkanIndexBuffer.cpp b/Engine/Source/Runtime/VulkanRHI/Private/VulkanIndexBuffer.cpp
--- a/Engine/Source/Runtime/VulkanRHI/Private/VulkanIndexBuffer.cpp
+++ b/Engine/Source/Runtime/VulkanRHI/Private/VulkanIndexBuffer.cpp
@@ -60,11 +60,11 @@ static FORCEINLINE void UpdateVulkanBufferStats(uint64_t Size, VkBufferUsageFlag
 	}
 }
 
-FVulkanResourceMultiBuffer::FVulkanResourceMultiBuffer(FVulkanDevice* InDevice, VkBufferUsageFlags InBufferUsageFlags, uint32 InSize, uint32 InUEUsage, uint32 InStride, FRHIResourceCreateInfo& CreateInfo, class FRHICommandListImmediate* InRHICmdList)
+FVulkanResourceMultiBuffer::FVulkanResourceMultiBuffer(FVulkanDevice* InDevice, uint32 InSize, uint32 InUEUsage, uint32 InStride, FRHIResourceCreateInfo& CreateInfo, class FRHICommandListImmediate* InRHICmdList)
 	: FRHIBuffer(InSize, InUEUsage, InStride)
 	, VulkanRHI::FDeviceChild(InDevice)
 	, UEUsage(InUEUsage)
-	, BufferUsageFlags(InBufferUsageFlags)
+	, BufferUsageFlags(0)
 	, NumBuffers(0)
 	, DynamicBufferIndex(0)
 {
@@ -413,14 +413,14 @@ void FVulkanResourceMultiBuffer::Swap(FVulkanResourceMultiBuffer& Other)
 	::Swap(*this, Other);
 }
 
-FIndexBufferRHIRef FVulkanDynamicRHI::RHICreateIndexBuffer(uint32 Stride, uint32 Size, uint32 InUsage, ERHIAccess InResourceState, FRHIResourceCreateInfo& CreateInfo)
+FBufferRHIRef FVulkanDynamicRHI::RHICreateBuffer(uint32 Size, EBufferUsageFlags Usage, uint32 Stride, ERHIAccess ResourceState, FRHIResourceCreateInfo& CreateInfo)
 {
 	LLM_SCOPE_VULKAN(ELLMTagVulkan::VulkanBuffers);
 	if (CreateInfo.bWithoutNativeResource)
 	{
-		return new FVulkanResourceMultiBuffer(nullptr, 0, 0, 0, 0, CreateInfo);
+		return new FVulkanResourceMultiBuffer(nullptr, 0, 0, 0, CreateInfo, nullptr);
 	}
-	return new FVulkanResourceMultiBuffer(Device, 0, Size, InUsage | BUF_IndexBuffer, Stride, CreateInfo);
+	return new FVulkanResourceMultiBuffer(Device, Size, Usage, Stride, CreateInfo, nullptr);
 }
 
 void* FVulkanDynamicRHI::LockBuffer_BottomOfPipe(FRHICommandListImmediate& RHICmdList, FRHIBuffer* BufferRHI, uint32 Offset, uint32 Size, EResourceLockMode LockMode)
@@ -456,7 +456,7 @@ void FVulkanDynamicRHI::RHITransferBufferUnderlyingResource(FRHIBuffer* DestBuff
 	if (!SrcBuffer)
 	{
 		FRHIResourceCreateInfo CreateInfo;
-		TRefCountPtr<FVulkanResourceMultiBuffer> DeletionProxy = new FVulkanResourceMultiBuffer(Dest->GetParent(), 0, 0, 0, 0, CreateInfo);
+		TRefCountPtr<FVulkanResourceMultiBuffer> DeletionProxy = new FVulkanResourceMultiBuffer(Dest->GetParent(), 0, 0, 0, CreateInfo, nullptr);
 		Dest->Swap(*DeletionProxy);
 	}
 	else
diff --git a/Engine/Source/Runtime/VulkanRHI/Private/VulkanPendingState.cpp b/Engine/Source/Runtime/VulkanRHI/Private/VulkanPendingState.cpp
--- a/Engine/Source/Runtime/VulkanRHI/Private/VulkanPendingState.cpp
+++ b/Engine/Source/Runtime/VulkanRHI/Private/VulkanPendingState.cpp
@@ -342,9 +342,9 @@ void FVulkanPendingComputeState::SetUAVForUBResource(uint32 DescriptorSet, uint3
 	{
 		// make sure any dynamically backed UAV points to current memory
 		UAV->UpdateView();
-		if (UAV->SourceStructuredBuffer)
+		if (UAV->SourceBuffer && UAV->BufferViewFormat == PF_Unknown)
 		{
-			CurrentState->SetStorageBuffer(DescriptorSet, BindingIndex, UAV->SourceStructuredBuffer);
+			CurrentState->SetStorageBuffer(DescriptorSet, BindingIndex, UAV->SourceBuffer);
 		}
 		else if (UAV->BufferView)
 		{
@@ -607,9 +607,9 @@ void FVulkanPendingGfxState::SetUAVForUBResource(uint8 DescriptorSet, uint32 Bin
 	{
 		// make sure any dynamically backed UAV points to current memory
 		UAV->UpdateView();
-		if (UAV->SourceStructuredBuffer)
+		if (UAV->SourceBuffer && UAV->BufferViewFormat == PF_Unknown)
 		{
-			CurrentState->SetStorageBuffer(DescriptorSet, BindingIndex, UAV->SourceStructuredBuffer);
+			CurrentState->SetStorageBuffer(DescriptorSet, BindingIndex, UAV->SourceBuffer);
 		}
 		else if (UAV->BufferView)
 		{
diff --git a/Engine/Source/Runtime/VulkanRHI/Private/VulkanResources.h b/Engine/Source/Runtime/VulkanRHI/Private/VulkanResources.h
--- a/Engine/Source/Runtime/VulkanRHI/Private/VulkanResources.h
+++ b/Engine/Source/Runtime/VulkanRHI/Private/VulkanResources.h
@@ -1101,7 +1101,7 @@ class FVulkanResourceMultiBuffer : public FRHIBuffer, public FVulkanEvictable, p
 	virtual void Move(FVulkanDevice& Device, VulkanRHI::FVulkanAllocation& NewAllocation);
 
 public:
-	FVulkanResourceMultiBuffer(FVulkanDevice* InDevice, VkBufferUsageFlags InBufferUsageFlags, uint32 InSize, uint32 InUEUsage, uint32 InStride, FRHIResourceCreateInfo& CreateInfo, class FRHICommandListImmediate* InRHICmdList = nullptr);
+	FVulkanResourceMultiBuffer(FVulkanDevice* InDevice, uint32 InSize, uint32 InUEUsage, uint32 InStride, FRHIResourceCreateInfo& CreateInfo, class FRHICommandListImmediate* InRHICmdList = nullptr);
 	virtual ~FVulkanResourceMultiBuffer();
 
 	inline const VulkanRHI::FVulkanAllocation& GetCurrentAllocation() const
@@ -1271,16 +1271,13 @@ public:
 	void UpdateView();
 
 protected:
-	// the potential resources to refer to with the UAV object
-	TRefCountPtr<FVulkanResourceMultiBuffer> SourceStructuredBuffer;
 	// The texture that this UAV come from
 	TRefCountPtr<FRHITexture> SourceTexture;
 	FVulkanTextureView TextureView;
 	uint32 MipLevel;
 
-	// The vertex buffer this UAV comes from (can be null)
-	TRefCountPtr<FVulkanResourceMultiBuffer> SourceVertexBuffer;
-	TRefCountPtr<FVulkanResourceMultiBuffer> SourceIndexBuffer;
+	// The buffer this UAV comes from (can be null)
+	TRefCountPtr<FVulkanResourceMultiBuffer> SourceBuffer;
 	TRefCountPtr<FVulkanBufferView> BufferView;
 	EPixelFormat BufferViewFormat;
 
diff --git a/Engine/Source/Runtime/VulkanRHI/Private/VulkanStructuredBuffer.cpp b/Engine/Source/Runtime/VulkanRHI/Private/VulkanStructuredBuffer.cpp
--- a/Engine/Source/Runtime/VulkanRHI/Private/VulkanStructuredBuffer.cpp
+++ b/Engine/Source/Runtime/VulkanRHI/Private/VulkanStructuredBuffer.cpp
@@ -3,8 +3,3 @@
 
 #include "VulkanRHIPrivate.h"
 #include "Containers/ResourceArray.h"
-
-FStructuredBufferRHIRef FVulkanDynamicRHI::RHICreateStructuredBuffer(uint32 InStride, uint32 InSize, uint32 InUsage, ERHIAccess InResourceState, FRHIResourceCreateInfo& CreateInfo)
-{
-	return new FVulkanResourceMultiBuffer(Device, VK_BUFFER_USAGE_STORAGE_BUFFER_BIT, InSize, InUsage | BUF_StructuredBuffer, InStride, CreateInfo);
-}
diff --git a/Engine/Source/Runtime/VulkanRHI/Private/VulkanUAV.cpp b/Engine/Source/Runtime/VulkanRHI/Private/VulkanUAV.cpp
--- a/Engine/Source/Runtime/VulkanRHI/Private/VulkanUAV.cpp
+++ b/Engine/Source/Runtime/VulkanRHI/Private/VulkanUAV.cpp
@@ -201,10 +201,10 @@ void FVulkanShaderResourceView::UpdateView()
 		}
 	}
 }
-FVulkanUnorderedAccessView::FVulkanUnorderedAccessView(FVulkanDevice* Device, FVulkanResourceMultiBuffer* StructuredBuffer, bool bUseUAVCounter, bool bAppendBuffer)
+FVulkanUnorderedAccessView::FVulkanUnorderedAccessView(FVulkanDevice* Device, FVulkanResourceMultiBuffer* Buffer, bool bUseUAVCounter, bool bAppendBuffer)
 	: VulkanRHI::FVulkanViewBase(Device)
-	, SourceStructuredBuffer(StructuredBuffer)
 	, MipLevel(0)
+	, SourceBuffer(Buffer)
 	, BufferViewFormat(PF_Unknown)
 	, VolatileLockCounter(MAX_uint32)
 {
@@ -228,14 +228,7 @@ FVulkanUnorderedAccessView::FVulkanUnorderedAccessView(FVulkanDevice* Device, FV
 	, BufferViewFormat(Format)
 	, VolatileLockCounter(MAX_uint32)
 {
-	if (Buffer->GetUEUsage() & BUF_IndexBuffer)
-	{
-		SourceIndexBuffer = Buffer;
-	}
-	else if (Buffer->GetUEUsage() & BUF_VertexBuffer)
-	{
-		SourceVertexBuffer = Buffer;
-	}
+	SourceBuffer = Buffer;
 }
 
 void FVulkanUnorderedAccessView::Invalidate()
@@ -254,7 +247,7 @@ FVulkanUnorderedAccessView::~FVulkanUnorderedAccessView()
 
 	TextureView.Destroy(*Device);
 	BufferView = nullptr;
-	SourceVertexBuffer = nullptr;
+	SourceBuffer = nullptr;
 	SourceTexture = nullptr;
 	Device = nullptr;
 }
@@ -266,45 +259,24 @@ void FVulkanUnorderedAccessView::UpdateView()
 #endif
 
 	// update the buffer view for dynamic VB backed buffers (or if it was never set)
-	if (SourceVertexBuffer != nullptr)
-	{
-		if (SourceVertexBuffer->IsVolatile() && VolatileLockCounter != SourceVertexBuffer->GetVolatileLockCounter())
-		{
-			BufferView = nullptr;
-			VolatileLockCounter = SourceVertexBuffer->GetVolatileLockCounter();
-		}
-
-		if (BufferView == nullptr || SourceVertexBuffer->IsDynamic())
-		{
-			// thanks to ref counting, overwriting the buffer will toss the old view
-			BufferView = new FVulkanBufferView(Device);
-			BufferView->Create(SourceVertexBuffer.GetReference(), BufferViewFormat, SourceVertexBuffer->GetOffset(), SourceVertexBuffer->GetSize());
-		}
-	}
-	else if (SourceIndexBuffer != nullptr)
+	if (SourceBuffer != nullptr)
 	{
-		if (SourceIndexBuffer->IsVolatile() && VolatileLockCounter != SourceIndexBuffer->GetVolatileLockCounter())
+		if (BufferViewFormat != PF_Unknown)
 		{
-			BufferView = nullptr;
-			VolatileLockCounter = SourceIndexBuffer->GetVolatileLockCounter();
-		}
+			if (SourceBuffer->IsVolatile() && VolatileLockCounter != SourceBuffer->GetVolatileLockCounter())
+			{
+				BufferView = nullptr;
+				VolatileLockCounter = SourceBuffer->GetVolatileLockCounter();
+			}
 
-		if (BufferView == nullptr || SourceIndexBuffer->IsDynamic())
-		{
-			// thanks to ref counting, overwriting the buffer will toss the old view
-			BufferView = new FVulkanBufferView(Device);
-			BufferView->Create(SourceIndexBuffer.GetReference(), BufferViewFormat, SourceIndexBuffer->GetOffset(), SourceIndexBuffer->GetSize());
+			if (BufferView == nullptr || SourceBuffer->IsDynamic())
+			{
+				// thanks to ref counting, overwriting the buffer will toss the old view
+				BufferView = new FVulkanBufferView(Device);
+				BufferView->Create(SourceBuffer.GetReference(), BufferViewFormat, SourceBuffer->GetOffset(), SourceBuffer->GetSize());
+			}
 		}
 	}
-	else if (SourceStructuredBuffer)
-	{
-		// Nothing...
-		//if (SourceStructuredBuffer->IsVolatile() && VolatileLockCounter != SourceStructuredBuffer->GetVolatileLockCounter())
-		//{
-		//	BufferView = nullptr;
-		//	VolatileLockCounter = SourceStructuredBuffer->GetVolatileLockCounter();
-		//}
-	}
 	else if (TextureView.View == VK_NULL_HANDLE)
 	{
 		EPixelFormat Format = (BufferViewFormat == PF_Unknown) ? SourceTexture->GetFormat() : BufferViewFormat;
@@ -369,17 +341,18 @@ FShaderResourceViewRHIRef FVulkanDynamicRHI::RHICreateShaderResourceView(FRHIBuf
 
 FShaderResourceViewRHIRef FVulkanDynamicRHI::RHICreateShaderResourceView(const FShaderResourceViewInitializer& Initializer)
 {
+	const FShaderResourceViewInitializer::FBufferShaderResourceViewInitializer Desc = Initializer.AsBufferSRV();
+	FVulkanResourceMultiBuffer* Buffer = ResourceCast(Desc.Buffer);
+
 	switch (Initializer.GetType())
 	{
 		case FShaderResourceViewInitializer::EType::VertexBufferSRV:
 		{
-			const FShaderResourceViewInitializer::FVertexBufferShaderResourceViewInitializer Desc = Initializer.AsVertexBufferSRV();
 			if (Desc.Buffer)
 			{
 				const uint32 Stride = GPixelFormats[Desc.Format].BlockBytes;
-				FVulkanResourceMultiBuffer* VertexBuffer = ResourceCast(Desc.Buffer);
-				uint32 Size = FMath::Min(VertexBuffer->GetSize() - Desc.StartOffsetBytes, Desc.NumElements * Stride);
-				return new FVulkanShaderResourceView(Device, Desc.Buffer, VertexBuffer, Size, (EPixelFormat)Desc.Format, Desc.StartOffsetBytes);
+				uint32 Size = FMath::Min(Buffer->GetSize() - Desc.StartOffsetBytes, Desc.NumElements * Stride);
+				return new FVulkanShaderResourceView(Device, Desc.Buffer, Buffer, Size, (EPixelFormat)Desc.Format, Desc.StartOffsetBytes);
 			}
 			else
 			{
@@ -388,21 +361,17 @@ FShaderResourceViewRHIRef FVulkanDynamicRHI::RHICreateShaderResourceView(const F
 		}
 		case FShaderResourceViewInitializer::EType::StructuredBufferSRV:
 		{
-			const FShaderResourceViewInitializer::FStructuredBufferShaderResourceViewInitializer Desc = Initializer.AsStructuredBufferSRV();
 			check(Desc.Buffer);
-			FVulkanResourceMultiBuffer* StructuredBuffer = ResourceCast(Desc.Buffer);
-			return new FVulkanShaderResourceView(Device, StructuredBuffer, Desc.StartOffsetBytes);
+			return new FVulkanShaderResourceView(Device, Buffer, Desc.StartOffsetBytes);
 		}			
 		case FShaderResourceViewInitializer::EType::IndexBufferSRV:
 		{
-			const FShaderResourceViewInitializer::FIndexBufferShaderResourceViewInitializer Desc = Initializer.AsIndexBufferSRV();
 			check(Desc.Buffer);
-			FVulkanResourceMultiBuffer* IndexBuffer = ResourceCast(Desc.Buffer);
 			const uint32 Stride = Desc.Buffer->GetStride();
 			check(Stride == 2 || Stride == 4);
 			EPixelFormat Format = (Stride == 4) ? PF_R32_UINT : PF_R16_UINT;
-			uint32 Size = FMath::Min(IndexBuffer->GetSize() - Desc.StartOffsetBytes, Desc.NumElements * Stride);
-			return new FVulkanShaderResourceView(Device, Desc.Buffer, IndexBuffer, Size, Format, Desc.StartOffsetBytes);
+			uint32 Size = FMath::Min(Buffer->GetSize() - Desc.StartOffsetBytes, Desc.NumElements * Stride);
+			return new FVulkanShaderResourceView(Device, Desc.Buffer, Buffer, Size, Format, Desc.StartOffsetBytes);
 		}
 	}
 	checkNoEntry();
@@ -474,7 +443,7 @@ void FVulkanCommandListContext::ClearUAV(TRHICommandList_RecursiveHazardous<FVul
 	if (!bFloat)
 	{
 		EPixelFormat Format;
-		if (UnorderedAccessView->SourceVertexBuffer)
+		if (UnorderedAccessView->SourceBuffer)
 		{
 			Format = UnorderedAccessView->BufferViewFormat;
 		}
@@ -504,10 +473,18 @@ void FVulkanCommandListContext::ClearUAV(TRHICommandList_RecursiveHazardous<FVul
 		ValueType = EClearReplacementValueType::Float;
 	}
 
-	if (UnorderedAccessView->SourceVertexBuffer)
+	if (UnorderedAccessView->SourceBuffer)
 	{
-		uint32 NumElements = UnorderedAccessView->SourceVertexBuffer->GetSize() / GPixelFormats[UnorderedAccessView->BufferViewFormat].BlockBytes;
-		ClearUAVShader_T<EClearReplacementResourceType::Buffer, 4, false>(RHICmdList, UnorderedAccessView, NumElements, 1, 1, ClearValue, ValueType);
+		if (UnorderedAccessView->BufferViewFormat == PF_Unknown)
+		{
+			uint32 NumUints = UnorderedAccessView->SourceBuffer->GetSize() / 4;
+			ClearUAVShader_T<EClearReplacementResourceType::StructuredBuffer, 4, false>(RHICmdList, UnorderedAccessView, NumUints, 1, 1, ClearValue, ValueType);
+		}
+		else
+		{
+			uint32 NumElements = UnorderedAccessView->SourceBuffer->GetSize() / GPixelFormats[UnorderedAccessView->BufferViewFormat].BlockBytes;
+			ClearUAVShader_T<EClearReplacementResourceType::Buffer, 4, false>(RHICmdList, UnorderedAccessView, NumElements, 1, 1, ClearValue, ValueType);
+		}
 	}
 	else if (UnorderedAccessView->SourceTexture)
 	{
@@ -534,12 +511,6 @@ void FVulkanCommandListContext::ClearUAV(TRHICommandList_RecursiveHazardous<FVul
 			ensure(0);
 		}
 	}
-	else if (UnorderedAccessView->SourceStructuredBuffer)
-	{
-		ensure(UnorderedAccessView->BufferViewFormat == PF_Unknown);
-		uint32 NumUints = UnorderedAccessView->SourceStructuredBuffer->GetSize() / 4;
-		ClearUAVShader_T<EClearReplacementResourceType::StructuredBuffer, 4, false>(RHICmdList, UnorderedAccessView, NumUints, 1, 1, ClearValue, ValueType);
-	}
 	else
 	{
 		ensure(0);
diff --git a/Engine/Source/Runtime/VulkanRHI/Private/VulkanVertexBuffer.cpp b/Engine/Source/Runtime/VulkanRHI/Private/VulkanVertexBuffer.cpp
--- a/Engine/Source/Runtime/VulkanRHI/Private/VulkanVertexBuffer.cpp
+++ b/Engine/Source/Runtime/VulkanRHI/Private/VulkanVertexBuffer.cpp
@@ -7,17 +7,6 @@
 #include "VulkanRHIPrivate.h"
 #include "VulkanLLM.h"
 
-FVertexBufferRHIRef FVulkanDynamicRHI::RHICreateVertexBuffer(uint32 Size, uint32 InUsage, ERHIAccess InResourceState, FRHIResourceCreateInfo& CreateInfo)
-{
-	LLM_SCOPE_VULKAN(ELLMTagVulkan::VulkanBuffers);
-	if (CreateInfo.bWithoutNativeResource)
-	{
-		return new FVulkanResourceMultiBuffer(nullptr, VK_BUFFER_USAGE_VERTEX_BUFFER_BIT, 0, 0, 0, CreateInfo);
-	}
-	FVulkanResourceMultiBuffer* VertexBuffer = new FVulkanResourceMultiBuffer(Device, VK_BUFFER_USAGE_VERTEX_BUFFER_BIT, Size, InUsage | BUF_VertexBuffer, 0, CreateInfo);
-	return VertexBuffer;
-}
-
 void FVulkanDynamicRHI::RHICopyBuffer(FRHIBuffer* SourceBufferRHI, FRHIBuffer* DestBufferRHI)
 {
 	VULKAN_SIGNAL_UNIMPLEMENTED();
diff --git a/Engine/Source/Runtime/VulkanRHI/Public/VulkanDynamicRHI.h b/Engine/Source/Runtime/VulkanRHI/Public/VulkanDynamicRHI.h
--- a/Engine/Source/Runtime/VulkanRHI/Public/VulkanDynamicRHI.h
+++ b/Engine/Source/Runtime/VulkanRHI/Public/VulkanDynamicRHI.h
@@ -80,11 +80,9 @@ public:
 	virtual FGraphicsPipelineStateRHIRef RHICreateGraphicsPipelineState(const FGraphicsPipelineStateInitializer& Initializer) final override;
 	virtual FUniformBufferRHIRef RHICreateUniformBuffer(const void* Contents, const FRHIUniformBufferLayout& Layout, EUniformBufferUsage Usage, EUniformBufferValidation Validation) final override;
 	virtual void RHIUpdateUniformBuffer(FRHIUniformBuffer* UniformBufferRHI, const void* Contents) final override;
-	virtual FIndexBufferRHIRef RHICreateIndexBuffer(uint32 Stride, uint32 Size, uint32 InUsage, ERHIAccess InResourceState, FRHIResourceCreateInfo& CreateInfo) final override;
-	virtual FVertexBufferRHIRef RHICreateVertexBuffer(uint32 Size, uint32 InUsage, ERHIAccess InResourceState, FRHIResourceCreateInfo& CreateInfo) final override;
+	virtual FBufferRHIRef RHICreateBuffer(uint32 Size, EBufferUsageFlags Usage, uint32 Stride, ERHIAccess ResourceState, FRHIResourceCreateInfo& CreateInfo) final override;
 	virtual void RHICopyBuffer(FRHIBuffer* SourceBuffer, FRHIBuffer* DestBuffer) final override;
 	virtual void RHITransferBufferUnderlyingResource(FRHIBuffer* DestBuffer, FRHIBuffer* SrcBuffer) final override;
-	virtual FStructuredBufferRHIRef RHICreateStructuredBuffer(uint32 Stride, uint32 Size, uint32 InUsage, ERHIAccess InResourceState, FRHIResourceCreateInfo& CreateInfo) final override;
 	virtual void* LockBuffer_BottomOfPipe(FRHICommandListImmediate& RHICmdList, FRHIBuffer* Buffer, uint32 Offset, uint32 Size, EResourceLockMode LockMode) final override;
 	virtual void UnlockBuffer_BottomOfPipe(FRHICommandListImmediate& RHICmdList, FRHIBuffer* Buffer) final override;
 	virtual FUnorderedAccessViewRHIRef RHICreateUnorderedAccessView(FRHIBuffer* Buffer, bool bUseUAVCounter, bool bAppendBuffer) final override;
@@ -182,25 +180,15 @@ public:
 	virtual class IRHICommandContextContainer* RHIGetCommandContextContainer(int32 Index, int32 Num) final override;
 	virtual uint64 RHIGetMinimumAlignmentForBufferBackedSRV(EPixelFormat Format) final override;
 
-	///////// Pass through functions that allow RHIs to optimize certain calls.
-	//virtual FVertexBufferRHIRef CreateAndLockVertexBuffer_RenderThread(class FRHICommandListImmediate& RHICmdList, uint32 Size, uint32 InUsage, FRHIResourceCreateInfo& CreateInfo, void*& OutDataBuffer) override final;
-	//virtual FIndexBufferRHIRef CreateAndLockIndexBuffer_RenderThread(class FRHICommandListImmediate& RHICmdList, uint32 Stride, uint32 Size, uint32 InUsage, FRHIResourceCreateInfo& CreateInfo, void*& OutDataBuffer) override final;
-	//virtual FStructuredBufferRHIRef CreateStructuredBuffer_RenderThread(class FRHICommandListImmediate& RHICmdList, uint32 Stride, uint32 Size, uint32 InUsage, FRHIResourceCreateInfo& CreateInfo) override final;
-	
-	virtual FVertexBufferRHIRef CreateVertexBuffer_RenderThread(class FRHICommandListImmediate& RHICmdList, uint32 Size, uint32 InUsage, ERHIAccess InResourceState, FRHIResourceCreateInfo& CreateInfo) override final
+	virtual FBufferRHIRef CreateBuffer_RenderThread(class FRHICommandListImmediate& RHICmdList, uint32 Size, EBufferUsageFlags Usage, uint32 Stride, ERHIAccess ResourceState, FRHIResourceCreateInfo& CreateInfo) override final
 	{
 		// this-> is required to avoid calling the global version of this function
-		return this->RHICreateVertexBuffer(Size, InUsage, InResourceState, CreateInfo);
+		return this->RHICreateBuffer(Size, Usage, Stride, ResourceState, CreateInfo);
 	}
 
 	virtual FTexture2DRHIRef AsyncReallocateTexture2D_RenderThread(class FRHICommandListImmediate& RHICmdList, FRHITexture2D* Texture2D, int32 NewMipCount, int32 NewSizeX, int32 NewSizeY, FThreadSafeCounter* RequestStatus) override final;
 	//virtual ETextureReallocationStatus CancelAsyncReallocateTexture2D_RenderThread(class FRHICommandListImmediate& RHICmdList, FRHITexture2D* Texture2D, bool bBlockUntilCompleted) override final;
 	
-	virtual FIndexBufferRHIRef CreateIndexBuffer_RenderThread(class FRHICommandListImmediate& RHICmdList, uint32 Stride, uint32 Size, uint32 InUsage, ERHIAccess InResourceState, FRHIResourceCreateInfo& CreateInfo) override final
-	{
-		// this-> is required to avoid calling the global version of this function
-		return this->RHICreateIndexBuffer(Stride, Size, InUsage, InResourceState, CreateInfo);
-	}
 #if VULKAN_BUFFER_LOCK_THREADSAFE		
 	virtual void* LockBuffer_RenderThread(class FRHICommandListImmediate& RHICmdList, FRHIBuffer* Buffer, uint32 Offset, uint32 SizeRHI, EResourceLockMode LockMode) override final;
 	virtual void UnlockBuffer_RenderThread(class FRHICommandListImmediate& RHICmdList, FRHIBuffer* Buffer) override final;
diff --git a/Engine/Source/Runtime/Windows/D3D11RHI/Private/D3D11IndexBuffer.cpp b/Engine/Source/Runtime/Windows/D3D11RHI/Private/D3D11IndexBuffer.cpp
--- a/Engine/Source/Runtime/Windows/D3D11RHI/Private/D3D11IndexBuffer.cpp
+++ b/Engine/Source/Runtime/Windows/D3D11RHI/Private/D3D11IndexBuffer.cpp
@@ -6,43 +6,68 @@
 
 #include "D3D11RHIPrivate.h"
 
-FIndexBufferRHIRef FD3D11DynamicRHI::RHICreateIndexBuffer(uint32 Stride, uint32 Size, uint32 InUsage, ERHIAccess InResourceState, FRHIResourceCreateInfo& CreateInfo)
+TAutoConsoleVariable<int32> GCVarUseSharedKeyedMutex(
+	TEXT("r.D3D11.UseSharedKeyMutex"),
+	0,
+	TEXT("If 1, BUF_Shared vertex / index buffer and TexCreate_Shared texture will be created\n")
+	TEXT("with the D3D11_RESOURCE_MISC_SHARED_KEYEDMUTEX flag instead of D3D11_RESOURCE_MISC_SHARED (default).\n"),
+	ECVF_Default);
+
+FBufferRHIRef FD3D11DynamicRHI::RHICreateBuffer(uint32 Size, EBufferUsageFlags Usage, uint32 Stride, ERHIAccess ResourceState, FRHIResourceCreateInfo& CreateInfo)
 {
-	InUsage |= BUF_IndexBuffer;
-
 	if (CreateInfo.bWithoutNativeResource)
 	{
 		return new FD3D11Buffer();
 	}
 
-	// Explicitly check that the size is nonzero before allowing CreateIndexBuffer to opaquely fail.
+	// Explicitly check that the size is nonzero before allowing CreateBuffer to opaquely fail.
 	check(Size > 0);
 
-	// Describe the index buffer.
-	D3D11_BUFFER_DESC Desc;
-	ZeroMemory( &Desc, sizeof( D3D11_BUFFER_DESC ) );
-	Desc.ByteWidth = Size;
-	Desc.Usage = (InUsage & BUF_AnyDynamic) ? D3D11_USAGE_DYNAMIC : D3D11_USAGE_DEFAULT;
-	Desc.BindFlags = D3D11_BIND_INDEX_BUFFER;
-	Desc.CPUAccessFlags = (InUsage & BUF_AnyDynamic) ? D3D11_CPU_ACCESS_WRITE : 0;
-	Desc.MiscFlags = 0;
+	// Describe the buffer.
+	D3D11_BUFFER_DESC Desc = { Size };
 
-	if (InUsage & BUF_UnorderedAccess)
+	if (Usage & BUF_AnyDynamic)
 	{
-		Desc.BindFlags |= D3D11_BIND_UNORDERED_ACCESS;
+		Desc.Usage = D3D11_USAGE_DYNAMIC;
+		Desc.CPUAccessFlags = D3D11_CPU_ACCESS_WRITE;
 	}
 
-	if(InUsage & BUF_DrawIndirect)
+	if (Usage & BUF_VertexBuffer)
 	{
-		Desc.MiscFlags |= D3D11_RESOURCE_MISC_DRAWINDIRECT_ARGS;
+		Desc.BindFlags |= D3D11_BIND_VERTEX_BUFFER;
 	}
 
-	if (InUsage & BUF_ShaderResource)
+	if (Usage & BUF_IndexBuffer)
+	{
+		Desc.BindFlags |= D3D11_BIND_INDEX_BUFFER;
+	}
+
+	if (Usage & BUF_ByteAddressBuffer)
+	{
+		Desc.MiscFlags |= D3D11_RESOURCE_MISC_BUFFER_ALLOW_RAW_VIEWS;
+	}
+	else if (Usage & BUF_StructuredBuffer)
+	{
+		Desc.StructureByteStride = Stride;
+		Desc.MiscFlags |= D3D11_RESOURCE_MISC_BUFFER_STRUCTURED;
+	}
+
+	if (Usage & BUF_ShaderResource)
 	{
 		Desc.BindFlags |= D3D11_BIND_SHADER_RESOURCE;
 	}
 
-	if (InUsage & BUF_Shared)
+	if (Usage & BUF_UnorderedAccess)
+	{
+		Desc.BindFlags |= D3D11_BIND_UNORDERED_ACCESS;
+	}
+
+	if (Usage & BUF_DrawIndirect)
+	{
+		Desc.MiscFlags |= D3D11_RESOURCE_MISC_DRAWINDIRECT_ARGS;
+	}
+
+	if (Usage & BUF_Shared)
 	{
 		if (GCVarUseSharedKeyedMutex->GetInt() != 0)
 		{
@@ -54,6 +79,14 @@ FIndexBufferRHIRef FD3D11DynamicRHI::RHICreateIndexBuffer(uint32 Stride, uint32
 		}
 	}
 
+	if (FPlatformMemory::SupportsFastVRAMMemory())
+	{
+		if (Usage & BUF_FastVRAM)
+		{
+			FFastVRAMAllocator::GetFastVRAMAllocator()->AllocUAVBuffer(Desc);
+		}
+	}
+
 	// If a resource array was provided for the resource, create the resource pre-populated
 	D3D11_SUBRESOURCE_DATA InitData;
 	D3D11_SUBRESOURCE_DATA* pInitData = NULL;
@@ -66,29 +99,28 @@ FIndexBufferRHIRef FD3D11DynamicRHI::RHICreateIndexBuffer(uint32 Stride, uint32
 		pInitData = &InitData;
 	}
 
-	TRefCountPtr<ID3D11Buffer> IndexBufferResource;
-	VERIFYD3D11RESULT_EX(Direct3DDevice->CreateBuffer(&Desc,pInitData,IndexBufferResource.GetInitReference()), Direct3DDevice);
+	TRefCountPtr<ID3D11Buffer> BufferResource;
+	VERIFYD3D11RESULT_EX(Direct3DDevice->CreateBuffer(&Desc, pInitData, BufferResource.GetInitReference()), Direct3DDevice);
 
-	UpdateBufferStats(IndexBufferResource, true);
+	UpdateBufferStats(BufferResource, true);
 
-	if(CreateInfo.ResourceArray)
+	if (CreateInfo.DebugName)
+	{
+		BufferResource->SetPrivateData(WKPDID_D3DDebugObjectName, FCString::Strlen(CreateInfo.DebugName) + 1, TCHAR_TO_ANSI(CreateInfo.DebugName));
+	}
+
+	if (CreateInfo.ResourceArray)
 	{
 		// Discard the resource array's contents.
 		CreateInfo.ResourceArray->Discard();
 	}
 
-	return new FD3D11Buffer(IndexBufferResource, Size, InUsage, Stride);
+	return new FD3D11Buffer(BufferResource, Size, Usage, Stride);
 }
 
-FIndexBufferRHIRef FD3D11DynamicRHI::CreateIndexBuffer_RenderThread(
-	class FRHICommandListImmediate& RHICmdList,
-	uint32 Stride,
-	uint32 Size,
-	uint32 InUsage,
-	ERHIAccess InResourceState,
-	FRHIResourceCreateInfo& CreateInfo)
+FBufferRHIRef FD3D11DynamicRHI::CreateBuffer_RenderThread(class FRHICommandListImmediate& RHICmdList, uint32 Size, EBufferUsageFlags Usage, uint32 Stride, ERHIAccess ResourceState, FRHIResourceCreateInfo& CreateInfo)
 {
-	return RHICreateIndexBuffer(Stride, Size, InUsage, InResourceState, CreateInfo);
+	return RHICreateBuffer(Size, Usage, Stride, ResourceState, CreateInfo);
 }
 
 void* FD3D11DynamicRHI::LockBuffer_BottomOfPipe(FRHICommandListImmediate& RHICmdList, FRHIBuffer* BufferRHI, uint32 Offset, uint32 Size, EResourceLockMode LockMode)
diff --git a/Engine/Source/Runtime/Windows/D3D11RHI/Private/D3D11RHIPrivate.h b/Engine/Source/Runtime/Windows/D3D11RHI/Private/D3D11RHIPrivate.h
--- a/Engine/Source/Runtime/Windows/D3D11RHI/Private/D3D11RHIPrivate.h
+++ b/Engine/Source/Runtime/Windows/D3D11RHI/Private/D3D11RHIPrivate.h
@@ -437,11 +437,9 @@ public:
 	virtual FBoundShaderStateRHIRef RHICreateBoundShaderState(FRHIVertexDeclaration* VertexDeclaration, FRHIVertexShader* VertexShader, FRHIHullShader* HullShader, FRHIDomainShader* DomainShader, FRHIPixelShader* PixelShader, FRHIGeometryShader* GeometryShader) final override;
 	virtual FUniformBufferRHIRef RHICreateUniformBuffer(const void* Contents, const FRHIUniformBufferLayout& Layout, EUniformBufferUsage Usage, EUniformBufferValidation Validation) final override;
 	virtual void RHIUpdateUniformBuffer(FRHIUniformBuffer* UniformBufferRHI, const void* Contents) final override;
-	virtual FIndexBufferRHIRef RHICreateIndexBuffer(uint32 Stride, uint32 Size, uint32 InUsage, ERHIAccess InResourceState, FRHIResourceCreateInfo& CreateInfo) final override;
-	virtual FVertexBufferRHIRef RHICreateVertexBuffer(uint32 Size, uint32 InUsage, ERHIAccess InResourceState, FRHIResourceCreateInfo& CreateInfo) final override;
+	virtual FBufferRHIRef RHICreateBuffer(uint32 Size, EBufferUsageFlags Usage, uint32 Stride, ERHIAccess ResourceState, FRHIResourceCreateInfo& CreateInfo) final override;
 	virtual void RHICopyBuffer(FRHIBuffer* SourceBuffer, FRHIBuffer* DestBuffer) final override;
 	virtual void RHITransferBufferUnderlyingResource(FRHIBuffer* DestBuffer, FRHIBuffer* SrcBuffer) final override;
-	virtual FStructuredBufferRHIRef RHICreateStructuredBuffer(uint32 Stride, uint32 Size, uint32 InUsage, ERHIAccess InResourceState, FRHIResourceCreateInfo& CreateInfo) final override;
 	virtual void* LockBuffer_BottomOfPipe(FRHICommandListImmediate& RHICmdList, FRHIBuffer* Buffer, uint32 Offset, uint32 Size, EResourceLockMode LockMode) final override;
 	virtual void UnlockBuffer_BottomOfPipe(FRHICommandListImmediate& RHICmdList, FRHIBuffer* Buffer) final override;
 	virtual FUnorderedAccessViewRHIRef RHICreateUnorderedAccessView(FRHIBuffer* Buffer, bool bUseUAVCounter, bool bAppendBuffer) final override;
@@ -610,15 +608,13 @@ public:
 	// *_RenderThread functions. Command lists call these functions on RT. You can implement your own behavior inside these functions.
 	// For example, deferring the actual creation to RHI thread by sending an RHI command.
 	// For D3D11, these functions mainly just remove RHIT stalls because ID3D11Device is thread safe.
-	virtual FVertexBufferRHIRef CreateVertexBuffer_RenderThread(class FRHICommandListImmediate& RHICmdList, uint32 Size, uint32 InUsage, ERHIAccess InResourceState, FRHIResourceCreateInfo& CreateInfo) final override;
-	virtual FStructuredBufferRHIRef CreateStructuredBuffer_RenderThread(class FRHICommandListImmediate& RHICmdList, uint32 Stride, uint32 Size, uint32 InUsage, ERHIAccess InResourceState, FRHIResourceCreateInfo& CreateInfo) final override;
+	virtual FBufferRHIRef CreateBuffer_RenderThread(class FRHICommandListImmediate& RHICmdList, uint32 Size, EBufferUsageFlags Usage, uint32 Stride, ERHIAccess ResourceState, FRHIResourceCreateInfo& CreateInfo) final override;
 	virtual FShaderResourceViewRHIRef CreateShaderResourceView_RenderThread(class FRHICommandListImmediate& RHICmdList, FRHIVertexBuffer* VertexBuffer, uint32 Stride, uint8 Format) final override;
 	virtual FShaderResourceViewRHIRef CreateShaderResourceView_RenderThread(class FRHICommandListImmediate& RHICmdList, const FShaderResourceViewInitializer& Initializer) final override;
 	virtual FShaderResourceViewRHIRef CreateShaderResourceView_RenderThread(class FRHICommandListImmediate& RHICmdList, FRHIIndexBuffer* Buffer) final override;
 	virtual FTexture2DRHIRef AsyncReallocateTexture2D_RenderThread(class FRHICommandListImmediate& RHICmdList, FRHITexture2D* Texture2D, int32 NewMipCount, int32 NewSizeX, int32 NewSizeY, FThreadSafeCounter* RequestStatus) final override;
 	virtual ETextureReallocationStatus FinalizeAsyncReallocateTexture2D_RenderThread(class FRHICommandListImmediate& RHICmdList, FRHITexture2D* Texture2D, bool bBlockUntilCompleted) final override;
 	virtual ETextureReallocationStatus CancelAsyncReallocateTexture2D_RenderThread(class FRHICommandListImmediate& RHICmdList, FRHITexture2D* Texture2D, bool bBlockUntilCompleted) final override;
-	virtual FIndexBufferRHIRef CreateIndexBuffer_RenderThread(class FRHICommandListImmediate& RHICmdList, uint32 Stride, uint32 Size, uint32 InUsage, ERHIAccess InResourceState, FRHIResourceCreateInfo& CreateInfo) final override;
 	virtual FVertexShaderRHIRef CreateVertexShader_RenderThread(class FRHICommandListImmediate& RHICmdList, TArrayView<const uint8> Code, const FSHAHash& Hash) final override;
 	virtual FPixelShaderRHIRef CreatePixelShader_RenderThread(class FRHICommandListImmediate& RHICmdList, TArrayView<const uint8> Code, const FSHAHash& Hash) final override;
 	virtual FGeometryShaderRHIRef CreateGeometryShader_RenderThread(class FRHICommandListImmediate& RHICmdList, TArrayView<const uint8> Code, const FSHAHash& Hash) final override;
diff --git a/Engine/Source/Runtime/Windows/D3D11RHI/Private/D3D11StructuredBuffer.cpp b/Engine/Source/Runtime/Windows/D3D11RHI/Private/D3D11StructuredBuffer.cpp
--- a/Engine/Source/Runtime/Windows/D3D11RHI/Private/D3D11StructuredBuffer.cpp
+++ b/Engine/Source/Runtime/Windows/D3D11RHI/Private/D3D11StructuredBuffer.cpp
@@ -1,101 +1,3 @@
 // Copyright Epic Games, Inc. All Rights Reserved.
 
 #include "D3D11RHIPrivate.h"
-
-FStructuredBufferRHIRef FD3D11DynamicRHI::RHICreateStructuredBuffer(uint32 Stride, uint32 Size, uint32 InUsage, ERHIAccess InResourceState, FRHIResourceCreateInfo& CreateInfo)
-{
-	InUsage |= BUF_StructuredBuffer;
-
-	// Explicitly check that the size is nonzero before allowing CreateStructuredBuffer to opaquely fail.
-	check(Size > 0);
-	// Check for values that will cause D3D calls to fail
-	check(Size / Stride > 0 && Size % Stride == 0);
-
-	D3D11_BUFFER_DESC Desc;
-	ZeroMemory( &Desc, sizeof( D3D11_BUFFER_DESC ) );
-	Desc.ByteWidth = Size;
-	Desc.Usage = (InUsage & BUF_AnyDynamic) ? D3D11_USAGE_DYNAMIC : D3D11_USAGE_DEFAULT;
-	Desc.BindFlags = 0;
-
-	if(InUsage & BUF_ShaderResource)
-	{
-		// Setup bind flags so we can create a view to read from the buffer in a shader.
-		Desc.BindFlags |= D3D11_BIND_SHADER_RESOURCE;
-	}
-
-	if (InUsage & BUF_UnorderedAccess)
-	{
-		// Setup bind flags so we can create a writeable UAV to the buffer
-		Desc.BindFlags |= D3D11_BIND_UNORDERED_ACCESS;
-	}
-
-	Desc.CPUAccessFlags = (InUsage & BUF_AnyDynamic) ? D3D11_CPU_ACCESS_WRITE : 0;
-	Desc.MiscFlags = 0;
-
-	if (InUsage & BUF_DrawIndirect)
-	{
-		Desc.MiscFlags |= D3D11_RESOURCE_MISC_DRAWINDIRECT_ARGS;
-	}
-	else
-	{
-		if (InUsage & BUF_ByteAddressBuffer)
-		{
-			Desc.MiscFlags |= D3D11_RESOURCE_MISC_BUFFER_ALLOW_RAW_VIEWS;
-		}
-		else
-		{
-			Desc.MiscFlags |= D3D11_RESOURCE_MISC_BUFFER_STRUCTURED;
-		}
-	}
-
-	Desc.StructureByteStride = Stride;
-
-	if (FPlatformMemory::SupportsFastVRAMMemory())
-	{
-		if (InUsage & BUF_FastVRAM)
-		{
-			FFastVRAMAllocator::GetFastVRAMAllocator()->AllocUAVBuffer(Desc);
-		}
-	}
-
-	// If a resource array was provided for the resource, create the resource pre-populated
-	D3D11_SUBRESOURCE_DATA InitData;
-	D3D11_SUBRESOURCE_DATA* pInitData = NULL;
-	if(CreateInfo.ResourceArray)
-	{
-		check(Size == CreateInfo.ResourceArray->GetResourceDataSize());
-		InitData.pSysMem = CreateInfo.ResourceArray->GetResourceData();
-		InitData.SysMemPitch = Size;
-		InitData.SysMemSlicePitch = 0;
-		pInitData = &InitData;
-	}
-
-	TRefCountPtr<ID3D11Buffer> StructuredBufferResource;
-	VERIFYD3D11RESULT_EX(Direct3DDevice->CreateBuffer(&Desc,pInitData,StructuredBufferResource.GetInitReference()), Direct3DDevice);
-
-	if( CreateInfo.DebugName )
-	{
-		StructuredBufferResource->SetPrivateData(WKPDID_D3DDebugObjectName, FCString::Strlen(CreateInfo.DebugName) + 1, TCHAR_TO_ANSI(CreateInfo.DebugName));
-	}
-
-	UpdateBufferStats(StructuredBufferResource, true);
-
-	if(CreateInfo.ResourceArray)
-	{
-		// Discard the resource array's contents.
-		CreateInfo.ResourceArray->Discard();
-	}
-
-	return new FD3D11Buffer(StructuredBufferResource, Size, InUsage, Stride);
-}
-
-FStructuredBufferRHIRef FD3D11DynamicRHI::CreateStructuredBuffer_RenderThread(
-	class FRHICommandListImmediate& RHICmdList,
-	uint32 Stride,
-	uint32 Size,
-	uint32 InUsage,
-	ERHIAccess InResourceState,
-	FRHIResourceCreateInfo& CreateInfo)
-{
-	return RHICreateStructuredBuffer(Stride, Size, InUsage, InResourceState, CreateInfo);
-}
diff --git a/Engine/Source/Runtime/Windows/D3D11RHI/Private/D3D11VertexBuffer.cpp b/Engine/Source/Runtime/Windows/D3D11RHI/Private/D3D11VertexBuffer.cpp
--- a/Engine/Source/Runtime/Windows/D3D11RHI/Private/D3D11VertexBuffer.cpp
+++ b/Engine/Source/Runtime/Windows/D3D11RHI/Private/D3D11VertexBuffer.cpp
@@ -6,126 +6,6 @@
 
 #include "D3D11RHIPrivate.h"
 
-TAutoConsoleVariable<int32> GCVarUseSharedKeyedMutex(
-	TEXT("r.D3D11.UseSharedKeyMutex"),
-	0,
-	TEXT("If 1, BUF_Shared vertex / index buffer and TexCreate_Shared texture will be created\n")
-	TEXT("with the D3D11_RESOURCE_MISC_SHARED_KEYEDMUTEX flag instead of D3D11_RESOURCE_MISC_SHARED (default).\n"),
-	ECVF_Default);
-
-FVertexBufferRHIRef FD3D11DynamicRHI::RHICreateVertexBuffer(uint32 Size, uint32 InUsage, ERHIAccess InResourceState, FRHIResourceCreateInfo& CreateInfo)
-{
-	InUsage |= BUF_VertexBuffer;
-
-	if (CreateInfo.bWithoutNativeResource)
-	{
-		return new FD3D11Buffer();
-	}
-
-	// Explicitly check that the size is nonzero before allowing CreateVertexBuffer to opaquely fail.
-	check(Size > 0);
-
-	D3D11_BUFFER_DESC Desc;
-	ZeroMemory( &Desc, sizeof( D3D11_BUFFER_DESC ) );
-	Desc.ByteWidth = Size;
-	Desc.Usage = (InUsage & BUF_AnyDynamic) ? D3D11_USAGE_DYNAMIC : D3D11_USAGE_DEFAULT;
-	Desc.BindFlags = D3D11_BIND_VERTEX_BUFFER;
-	Desc.CPUAccessFlags = (InUsage & BUF_AnyDynamic) ? D3D11_CPU_ACCESS_WRITE : 0;
-	//Desc.MiscFlags = 0;
-	//Desc.StructureByteStride = 0;
-
-	if (InUsage & BUF_UnorderedAccess)
-	{
-		Desc.BindFlags |= D3D11_BIND_UNORDERED_ACCESS;
-	}
-
-	if (InUsage & BUF_ByteAddressBuffer)
-	{
-		Desc.MiscFlags |= D3D11_RESOURCE_MISC_BUFFER_ALLOW_RAW_VIEWS;
-	}
-
-	if (InUsage & BUF_DrawIndirect)
-	{
-		Desc.MiscFlags |= D3D11_RESOURCE_MISC_DRAWINDIRECT_ARGS;
-	}
-
-	if (InUsage & BUF_ShaderResource)
-	{
-		Desc.BindFlags |= D3D11_BIND_SHADER_RESOURCE;
-	}
-
-	if (InUsage & BUF_Shared)
-	{
-		if (GCVarUseSharedKeyedMutex->GetInt() != 0)
-		{
-			Desc.MiscFlags |= D3D11_RESOURCE_MISC_SHARED_KEYEDMUTEX;
-		}
-		else
-		{
-			Desc.MiscFlags |= D3D11_RESOURCE_MISC_SHARED;
-		}
-	}
-
-	if (FPlatformMemory::SupportsFastVRAMMemory())
-	{
-		if (InUsage & BUF_FastVRAM)
-		{
-			FFastVRAMAllocator::GetFastVRAMAllocator()->AllocUAVBuffer(Desc);
-		}
-	}
-
-	// If a resource array was provided for the resource, create the resource pre-populated
-	D3D11_SUBRESOURCE_DATA InitData;
-	D3D11_SUBRESOURCE_DATA* pInitData = NULL;
-	if(CreateInfo.ResourceArray)
-	{
-		checkf(Size == CreateInfo.ResourceArray->GetResourceDataSize(),
-			TEXT("DebugName: %s, GPU Size: %d, CPU Size: %d, Is Dynamic: %s"),
-			CreateInfo.DebugName, Size, CreateInfo.ResourceArray->GetResourceDataSize(),
-			InUsage & BUF_AnyDynamic ? TEXT("Yes") : TEXT("No"));
-
-		InitData.pSysMem = CreateInfo.ResourceArray->GetResourceData();
-		InitData.SysMemPitch = Size;
-		InitData.SysMemSlicePitch = 0;
-		pInitData = &InitData;
-	}
-
-	TRefCountPtr<ID3D11Buffer> VertexBufferResource;
-	check(Desc.ByteWidth != 0); //tracking down an elusive bug.
-	HRESULT Result = Direct3DDevice->CreateBuffer(&Desc, pInitData, VertexBufferResource.GetInitReference());
-	check(Desc.ByteWidth != 0); //tracking down an elusive bug.
-	if (FAILED(Result))
-	{
-		UE_LOG(LogD3D11RHI, Error, TEXT("D3DDevice failed CreateBuffer VB with ByteWidth=%d, BindFlags=0x%x Usage=%d, CPUAccess=0x%x, MiscFlags=0x%x"), Desc.ByteWidth, (uint32)Desc.BindFlags, (uint32)Desc.Usage, Desc.CPUAccessFlags, Desc.MiscFlags);
-		VERIFYD3D11RESULT_EX(Result, Direct3DDevice);
-	}
-
-	if (CreateInfo.DebugName)
-	{
-		VertexBufferResource->SetPrivateData(WKPDID_D3DDebugObjectName, FCString::Strlen(CreateInfo.DebugName) + 1, TCHAR_TO_ANSI(CreateInfo.DebugName));
-	}
-
-	UpdateBufferStats(VertexBufferResource, true);
-
-	if(CreateInfo.ResourceArray)
-	{
-		// Discard the resource array's contents.
-		CreateInfo.ResourceArray->Discard();
-	}
-
-	return new FD3D11Buffer(VertexBufferResource, Size, InUsage, 0);
-}
-
-FVertexBufferRHIRef FD3D11DynamicRHI::CreateVertexBuffer_RenderThread(
-	class FRHICommandListImmediate& RHICmdList,
-	uint32 Size,
-	uint32 InUsage,
-	ERHIAccess InResourceState,
-	FRHIResourceCreateInfo& CreateInfo)
-{
-	return RHICreateVertexBuffer(Size, InUsage, InResourceState, CreateInfo);
-}
-
 void FD3D11DynamicRHI::RHICopyBuffer(FRHIBuffer* SourceBufferRHI, FRHIBuffer* DestBufferRHI)
 {
 	FD3D11Buffer* SourceBuffer = ResourceCast(SourceBufferRHI);
