Fix for ClearUAV on Vulkan, now using vkCmdFillBuffer

diff --git a/Engine/Source/Runtime/VulkanRHI/Private/VulkanUAV.cpp b/Engine/Source/Runtime/VulkanRHI/Private/VulkanUAV.cpp
--- a/Engine/Source/Runtime/VulkanRHI/Private/VulkanUAV.cpp
+++ b/Engine/Source/Runtime/VulkanRHI/Private/VulkanUAV.cpp
@@ -439,6 +439,26 @@ void FVulkanDynamicRHI::RHIUpdateShaderResourceView(FRHIShaderResourceView* SRV,
 
 void FVulkanCommandListContext::ClearUAV(TRHICommandList_RecursiveHazardous<FVulkanCommandListContext>& RHICmdList, FVulkanUnorderedAccessView* UnorderedAccessView, const void* ClearValue, bool bFloat)
 {
+	struct FVulkanDynamicRHICmdFillBuffer final : public FRHICommand<FVulkanDynamicRHICmdFillBuffer>
+	{
+		FVulkanUnorderedAccessView* UAV;
+		uint32_t ClearValue;
+
+		FORCEINLINE_DEBUGGABLE FVulkanDynamicRHICmdFillBuffer(FVulkanUnorderedAccessView* InUAV, uint32_t InClearValue)
+			: UAV(InUAV), ClearValue(InClearValue)
+		{
+		}
+
+		void Execute(FRHICommandListBase& CmdList)
+		{
+			FVulkanCommandBufferManager* CmdBufferMgr = GVulkanRHI->GetDevice()->GetImmediateContext().GetCommandBufferManager();
+			FVulkanCmdBuffer* CmdBuffer = CmdBufferMgr->GetActiveCmdBuffer();
+
+			FVulkanResourceMultiBuffer* Buffer = UAV->SourceBuffer;
+			VulkanRHI::vkCmdFillBuffer(CmdBuffer->GetHandle(), Buffer->Current.Handle, Buffer->Current.Offset, Buffer->Current.Size, ClearValue);
+		}
+	};
+
 	EClearReplacementValueType ValueType;
 	if (!bFloat)
 	{
@@ -475,16 +495,23 @@ void FVulkanCommandListContext::ClearUAV(TRHICommandList_RecursiveHazardous<FVul
 
 	if (UnorderedAccessView->SourceBuffer)
 	{
-		if (UnorderedAccessView->BufferViewFormat == PF_Unknown)
+		RHICmdList.Transition(FRHITransitionInfo(UnorderedAccessView, ERHIAccess::UAVCompute, ERHIAccess::CopyDest));
+			
+		TRefCountPtr<FVulkanResourceMultiBuffer> Buffer = UnorderedAccessView->SourceBuffer;
+
+		if (RHICmdList.Bypass())
 		{
-			uint32 NumUints = UnorderedAccessView->SourceBuffer->GetSize() / 4;
-			ClearUAVShader_T<EClearReplacementResourceType::StructuredBuffer, 4, false>(RHICmdList, UnorderedAccessView, NumUints, 1, 1, ClearValue, ValueType);
+			FVulkanCommandBufferManager* CmdBufferMgr = GVulkanRHI->GetDevice()->GetImmediateContext().GetCommandBufferManager();
+			FVulkanCmdBuffer* CmdBuffer = CmdBufferMgr->GetActiveCmdBuffer();
+
+			VulkanRHI::vkCmdFillBuffer(CmdBuffer->GetHandle(), Buffer->Current.Handle, Buffer->Current.Offset, Buffer->Current.Size, *(const uint32_t*)ClearValue);
 		}
 		else
 		{
-			uint32 NumElements = UnorderedAccessView->SourceBuffer->GetSize() / GPixelFormats[UnorderedAccessView->BufferViewFormat].BlockBytes;
-			ClearUAVShader_T<EClearReplacementResourceType::Buffer, 4, false>(RHICmdList, UnorderedAccessView, NumElements, 1, 1, ClearValue, ValueType);
+			new (RHICmdList.AllocCommand<FVulkanDynamicRHICmdFillBuffer>()) FVulkanDynamicRHICmdFillBuffer(UnorderedAccessView, *(const uint32_t*)ClearValue);
 		}
+
+		RHICmdList.Transition(FRHITransitionInfo(UnorderedAccessView, ERHIAccess::CopyDest, ERHIAccess::SRVCompute));
 	}
 	else if (UnorderedAccessView->SourceTexture)
 	{
