Fix for ClearUAV on Vulkan

diff --git a/Engine/Source/Runtime/VulkanRHI/Private/VulkanUAV.cpp b/Engine/Source/Runtime/VulkanRHI/Private/VulkanUAV.cpp
--- a/Engine/Source/Runtime/VulkanRHI/Private/VulkanUAV.cpp
+++ b/Engine/Source/Runtime/VulkanRHI/Private/VulkanUAV.cpp
@@ -511,7 +511,7 @@ void FVulkanCommandListContext::ClearUAV(TRHICommandList_RecursiveHazardous<FVul
 			new (RHICmdList.AllocCommand<FVulkanDynamicRHICmdFillBuffer>()) FVulkanDynamicRHICmdFillBuffer(UnorderedAccessView, *(const uint32_t*)ClearValue);
 		}
 
-		RHICmdList.Transition(FRHITransitionInfo(UnorderedAccessView, ERHIAccess::CopyDest, ERHIAccess::SRVCompute));
+		RHICmdList.Transition(FRHITransitionInfo(UnorderedAccessView, ERHIAccess::CopyDest, ERHIAccess::UAVCompute));
 	}
 	else if (UnorderedAccessView->SourceTexture)
 	{
