UE5 - Copy from 15185155
Vulkan fixes for deleted images

diff --git a/Engine/Source/Runtime/VulkanRHI/Private/VulkanDevice.cpp b/Engine/Source/Runtime/VulkanRHI/Private/VulkanDevice.cpp
--- a/Engine/Source/Runtime/VulkanRHI/Private/VulkanDevice.cpp
+++ b/Engine/Source/Runtime/VulkanRHI/Private/VulkanDevice.cpp
@@ -1253,14 +1253,13 @@ const VkComponentMapping& FVulkanDevice::GetFormatComponentMapping(EPixelFormat
 	return PixelFormatComponentMapping[UEFormat];
 }
 
-void FVulkanDevice::NotifyDeletedRenderTarget(VkImage Image)
+void FVulkanDevice::NotifyDeletedImage(VkImage Image, bool bRenderTarget)
 {
-	//#todo-rco: Loop through all contexts!
-	GetImmediateContext().NotifyDeletedRenderTarget(Image);
-}
+	if (bRenderTarget)
+	{
+		GetImmediateContext().NotifyDeletedRenderTarget(Image);
+	}
 
-void FVulkanDevice::NotifyDeletedImage(VkImage Image)
-{
 	//#todo-rco: Loop through all contexts!
 	GetImmediateContext().NotifyDeletedImage(Image);
 }
diff --git a/Engine/Source/Runtime/VulkanRHI/Private/VulkanDevice.h b/Engine/Source/Runtime/VulkanRHI/Private/VulkanDevice.h
--- a/Engine/Source/Runtime/VulkanRHI/Private/VulkanDevice.h
+++ b/Engine/Source/Runtime/VulkanRHI/Private/VulkanDevice.h
@@ -334,8 +334,7 @@ public:
 		return *ComputeContext;
 	}
 
-	void NotifyDeletedRenderTarget(VkImage Image);
-	void NotifyDeletedImage(VkImage Image);
+	void NotifyDeletedImage(VkImage Image, bool bRenderTarget);
 
 #if VULKAN_ENABLE_DRAW_MARKERS
 	inline PFN_vkCmdDebugMarkerBeginEXT GetCmdDbgMarkerBegin() const
diff --git a/Engine/Source/Runtime/VulkanRHI/Private/VulkanMemory.cpp b/Engine/Source/Runtime/VulkanRHI/Private/VulkanMemory.cpp
--- a/Engine/Source/Runtime/VulkanRHI/Private/VulkanMemory.cpp
+++ b/Engine/Source/Runtime/VulkanRHI/Private/VulkanMemory.cpp
@@ -262,8 +262,7 @@ static FAutoConsoleVariableRef CVarVulkanFakeMemoryLimit(
 );
 #endif
 
-// this is disabled currently due causing use-after-free issues, which caused UE-105058
-int32 GVulkanEnableDefrag = 0;
+int32 GVulkanEnableDefrag = 1;
 static FAutoConsoleVariableRef CVarVulkanEnableDefrag(
 	TEXT("r.Vulkan.EnableDefrag"),
 	GVulkanEnableDefrag,
diff --git a/Engine/Source/Runtime/VulkanRHI/Private/VulkanTexture.cpp b/Engine/Source/Runtime/VulkanRHI/Private/VulkanTexture.cpp
--- a/Engine/Source/Runtime/VulkanRHI/Private/VulkanTexture.cpp
+++ b/Engine/Source/Runtime/VulkanRHI/Private/VulkanTexture.cpp
@@ -143,10 +143,10 @@ static void VulkanTextureAllocated(uint64 Size, VkImageViewType ImageType, bool
 	UpdateVulkanTextureStats(Size, bIsCube, bIs3D, bIsRT);
 }
 
-static void VulkanTextureDestroyed(uint64 Size, VkImageViewType ImageTupe, bool bIsRT)
+static void VulkanTextureDestroyed(uint64 Size, VkImageViewType ImageType, bool bIsRT)
 {
-	bool bIsCube = ImageTupe == VK_IMAGE_VIEW_TYPE_CUBE || ImageTupe == VK_IMAGE_VIEW_TYPE_CUBE_ARRAY;
-	bool bIs3D = ImageTupe == VK_IMAGE_VIEW_TYPE_3D;
+	bool bIsCube = ImageType == VK_IMAGE_VIEW_TYPE_CUBE || ImageType == VK_IMAGE_VIEW_TYPE_CUBE_ARRAY;
+	bool bIs3D = ImageType == VK_IMAGE_VIEW_TYPE_3D;
 
 	UpdateVulkanTextureStats(-(int64)Size, bIsCube, bIs3D, bIsRT);
 }
@@ -492,16 +492,18 @@ struct FRHICommandOnDestroyImage final : public FRHICommand<FRHICommandOnDestroy
 {
 	VkImage Image;
 	FVulkanDevice* Device;
+	bool bRenderTarget;
 
-	FRHICommandOnDestroyImage(VkImage InImage, FVulkanDevice* InDevice)
+	FRHICommandOnDestroyImage(VkImage InImage, FVulkanDevice* InDevice, bool bInRenderTarget)
 		: Image(InImage)
 		, Device(InDevice)
+		, bRenderTarget(bInRenderTarget)
 	{
 	}
 
 	void Execute(FRHICommandListBase& RHICmdList)
 	{
-		Device->NotifyDeletedImage(Image);
+		Device->NotifyDeletedImage(Image, bRenderTarget);
 	}
 };
 
@@ -792,7 +794,7 @@ void FVulkanSurface::MoveSurface(FVulkanDevice& InDevice, FVulkanAllocation& New
 		{
 			check(Image != VK_NULL_HANDLE);
 			uint64 Size = GetMemorySize();
-			Device->NotifyDeletedImage(Image);
+			Device->NotifyDeletedImage(Image, bRenderTarget);
 			Device->GetDeferredDeletionQueue().EnqueueResource(VulkanRHI::FDeferredDeletionQueue2::EType::Image, Image);
 			Image = MovedImage;
 			VulkanTextureDestroyed(Size, ViewType, bRenderTarget);
@@ -896,7 +898,7 @@ void FVulkanSurface::OnFullDefrag(FVulkanDevice& InDevice, uint32 NewOffset)
 		{
 			check(Image != VK_NULL_HANDLE);
 			uint64 Size = GetMemorySize();
-			Device->NotifyDeletedRenderTarget(Image);
+			Device->NotifyDeletedImage(Image, bRenderTarget);
 			Device->GetDeferredDeletionQueue().EnqueueResource(VulkanRHI::FDeferredDeletionQueue2::EType::Image, Image);
 			FGenericPlatformMisc::LowLevelOutputDebugStringf(TEXT("** MOVE IMAGE %p -> %p\n"), Image, MovedImage);
 			Image = MovedImage;
@@ -1022,7 +1024,7 @@ void FVulkanSurface::EvictSurface(FVulkanDevice& InDevice)
 	{
 		check(Image != VK_NULL_HANDLE);
 		uint64 Size = GetMemorySize();
-		Device->NotifyDeletedImage(Image);
+		Device->NotifyDeletedImage(Image, bRenderTarget);
 		Device->GetDeferredDeletionQueue().EnqueueResource(VulkanRHI::FDeferredDeletionQueue2::EType::Image, Image);
 		Device->GetMemoryManager().FreeVulkanAllocation(Allocation);
 		check(!Allocation.HasAllocation());
@@ -1143,15 +1145,16 @@ void FVulkanSurface::Destroy()
 	}
 	else if (bIsImageOwner)
 	{
+		const bool bRenderTarget = (UEFlags & (TexCreate_RenderTargetable | TexCreate_DepthStencilTargetable | TexCreate_ResolveTargetable)) != 0;
 		FRHICommandList& RHICmdList = FRHICommandListExecutor::GetImmediateCommandList();
 		if (!IsInRenderingThread() || (RHICmdList.Bypass() || !IsRunningRHIInSeparateThread()))
 		{
-			Device->NotifyDeletedImage(Image);
+			Device->NotifyDeletedImage(Image, bRenderTarget);
 		}
 		else
 		{
 			check(IsInRenderingThread());
-			new (RHICmdList.AllocCommand<FRHICommandOnDestroyImage>()) FRHICommandOnDestroyImage(Image, Device);
+			new (RHICmdList.AllocCommand<FRHICommandOnDestroyImage>()) FRHICommandOnDestroyImage(Image, Device, bRenderTarget);
 		}
 
 		bIsImageOwner = false;
@@ -1166,7 +1169,6 @@ void FVulkanSurface::Destroy()
 			Image = VK_NULL_HANDLE;
 		}
 
-		const bool bRenderTarget = (UEFlags & (TexCreate_RenderTargetable | TexCreate_DepthStencilTargetable | TexCreate_ResolveTargetable)) != 0;
 		VulkanTextureDestroyed(Size, ViewType, bRenderTarget);
 	}
 }
@@ -2538,7 +2540,7 @@ FVulkanTexture2D::~FVulkanTexture2D()
 {
 	if ((Surface.UEFlags & (TexCreate_DepthStencilTargetable | TexCreate_RenderTargetable | TexCreate_ResolveTargetable)) != 0)
 	{
-		Surface.Device->NotifyDeletedRenderTarget(Surface.Image);
+		Surface.Device->GetImmediateContext().NotifyDeletedRenderTarget(Surface.Image);
 	}
 }
 
@@ -2590,7 +2592,7 @@ FVulkanTextureCube::~FVulkanTextureCube()
 {
 	if ((GetFlags() & (TexCreate_DepthStencilTargetable | TexCreate_RenderTargetable | TexCreate_ResolveTargetable)) != 0) 
 	{
-		Surface.Device->NotifyDeletedRenderTarget(Surface.Image);
+		Surface.Device->GetImmediateContext().NotifyDeletedRenderTarget(Surface.Image);
 	}
 }
 
@@ -2605,7 +2607,7 @@ FVulkanTexture3D::~FVulkanTexture3D()
 {
 	if ((GetFlags() & (TexCreate_DepthStencilTargetable | TexCreate_RenderTargetable | TexCreate_ResolveTargetable)) != 0) 
 	{
-		Surface.Device->NotifyDeletedRenderTarget(Surface.Image);
+		Surface.Device->GetImmediateContext().NotifyDeletedRenderTarget(Surface.Image);
 	}
 }
 
diff --git a/Engine/Source/Runtime/VulkanRHI/Private/VulkanUtil.cpp b/Engine/Source/Runtime/VulkanRHI/Private/VulkanUtil.cpp
--- a/Engine/Source/Runtime/VulkanRHI/Private/VulkanUtil.cpp
+++ b/Engine/Source/Runtime/VulkanRHI/Private/VulkanUtil.cpp
@@ -11,6 +11,7 @@
 #include "VulkanMemory.h"
 #include "Misc/OutputDeviceRedirector.h"
 #include "RHIValidationContext.h"
+#include "HAL/FileManager.h"
 
 #if NV_AFTERMATH
 #include "GFSDK_Aftermath_GpuCrashDump.h"
@@ -735,8 +736,15 @@ void AftermathShaderDebugInfoCallback(const void* pShaderDebugInfo, const uint32
 {
 }
 
-void AftermathCrashDumpDescriptionCallback(PFN_GFSDK_Aftermath_AddGpuCrashDumpDescription addDescription, void* pUserData)
+void AftermathCrashDumpDescriptionCallback(PFN_GFSDK_Aftermath_AddGpuCrashDumpDescription AddDescription, void* pUserData)
 {
+	// Add some basic description about the crash. This is called after the GPU crash happens, but before
+	// the actual GPU crash dump callback. The provided data is included in the crash dump and can be
+	// retrieved using GFSDK_Aftermath_GpuCrashDump_GetDescription().
+	FTCHARToUTF8 ProjectNameConverter(FApp::GetProjectName());
+	AddDescription(GFSDK_Aftermath_GpuCrashDumpDescriptionKey_ApplicationName, ProjectNameConverter.Get());
+	AddDescription(GFSDK_Aftermath_GpuCrashDumpDescriptionKey_ApplicationVersion, "v1.0");
+	AddDescription(GFSDK_Aftermath_GpuCrashDumpDescriptionKey_UserDefined, "Vulkan GPU crash");
 }
 #endif
 
diff --git a/Engine/Source/Runtime/VulkanRHI/Private/VulkanViewport.cpp b/Engine/Source/Runtime/VulkanRHI/Private/VulkanViewport.cpp
--- a/Engine/Source/Runtime/VulkanRHI/Private/VulkanViewport.cpp
+++ b/Engine/Source/Runtime/VulkanRHI/Private/VulkanViewport.cpp
@@ -168,8 +168,7 @@ FVulkanViewport::~FVulkanViewport()
 			TextureViews[Index].Destroy(*Device);
 
 			// FIXME: race condition on TransitionAndLayoutManager, could this be called from RT while RHIT is active?
-			Device->NotifyDeletedImage(BackBufferImages[Index]);
-			Device->NotifyDeletedRenderTarget(BackBufferImages[Index]); 
+			Device->NotifyDeletedImage(BackBufferImages[Index], true);
 			BackBufferImages[Index] = VK_NULL_HANDLE;
 		}
 
@@ -723,8 +722,7 @@ void FVulkanViewport::DestroySwapchain(FVulkanSwapChainRecreateInfo* RecreateInf
 		for (int32 Index = 0, NumBuffers = BackBufferImages.Num(); Index < NumBuffers; ++Index)
 		{
 			TextureViews[Index].Destroy(*Device);
-			Device->NotifyDeletedImage(BackBufferImages[Index]);
-			Device->NotifyDeletedRenderTarget(BackBufferImages[Index]);
+			Device->NotifyDeletedImage(BackBufferImages[Index], true);
 			BackBufferImages[Index] = VK_NULL_HANDLE;
 		}
 		
