Fixed Vulkan Timing init crash

diff --git a/Engine/Source/Runtime/VulkanRHI/Private/VulkanCommandBuffer.cpp b/Engine/Source/Runtime/VulkanRHI/Private/VulkanCommandBuffer.cpp
--- a/Engine/Source/Runtime/VulkanRHI/Private/VulkanCommandBuffer.cpp
+++ b/Engine/Source/Runtime/VulkanRHI/Private/VulkanCommandBuffer.cpp
@@ -442,6 +442,10 @@ FVulkanCommandBufferManager::FVulkanCommandBufferManager(FVulkanDevice* InDevice
 	}
 
 	ActiveCmdBuffer = Pool.Create(false);
+}
+
+void FVulkanCommandBufferManager::Init(FVulkanCommandListContext* InContext)
+{
 	ActiveCmdBuffer->InitializeTimings(InContext);
 	ActiveCmdBuffer->Begin();
 }
diff --git a/Engine/Source/Runtime/VulkanRHI/Private/VulkanCommandBuffer.h b/Engine/Source/Runtime/VulkanRHI/Private/VulkanCommandBuffer.h
--- a/Engine/Source/Runtime/VulkanRHI/Private/VulkanCommandBuffer.h
+++ b/Engine/Source/Runtime/VulkanRHI/Private/VulkanCommandBuffer.h
@@ -286,6 +286,8 @@ public:
 	FVulkanCommandBufferManager(FVulkanDevice* InDevice, FVulkanCommandListContext* InContext);
 	~FVulkanCommandBufferManager();
 
+	void Init(FVulkanCommandListContext* InContext);
+
 	inline FVulkanCmdBuffer* GetActiveCmdBuffer()
 	{
 		if (UploadCmdBuffer)
diff --git a/Engine/Source/Runtime/VulkanRHI/Private/VulkanRHI.cpp b/Engine/Source/Runtime/VulkanRHI/Private/VulkanRHI.cpp
--- a/Engine/Source/Runtime/VulkanRHI/Private/VulkanRHI.cpp
+++ b/Engine/Source/Runtime/VulkanRHI/Private/VulkanRHI.cpp
@@ -143,6 +143,7 @@ FVulkanCommandListContext::FVulkanCommandListContext(FVulkanDynamicRHI* InRHI, F
 
 	// Create CommandBufferManager, contain all active buffers
 	CommandBufferManager = new FVulkanCommandBufferManager(InDevice, this);
+	CommandBufferManager->Init(this);
 	FrameTiming->Initialize();
 	if (IsImmediate())
 	{
