Lazily acquire VK backbuffer on any type of layout transition.
This is needed since now we do not transition to write when there is nothing to render. (CL# 15176852)

diff --git a/Engine/Source/Runtime/VulkanRHI/Private/VulkanViewport.cpp b/Engine/Source/Runtime/VulkanRHI/Private/VulkanViewport.cpp
--- a/Engine/Source/Runtime/VulkanRHI/Private/VulkanViewport.cpp
+++ b/Engine/Source/Runtime/VulkanRHI/Private/VulkanViewport.cpp
@@ -65,7 +65,7 @@ void FVulkanBackBuffer::OnAdvanceBackBufferFrame(FRHICommandListImmediate& RHICm
 
 void FVulkanBackBuffer::OnLayoutTransition(FVulkanCommandListContext& Context, VkImageLayout NewLayout)
 {
-	if (NewLayout == VK_IMAGE_LAYOUT_COLOR_ATTACHMENT_OPTIMAL && GVulkanDelayAcquireImage == EDelayAcquireImageType::LazyAcquire)
+	if (GVulkanDelayAcquireImage == EDelayAcquireImageType::LazyAcquire)
 	{
 		AcquireBackBufferImage(Context);
 	}
