Rename some files to more appropriate names in a post-buffer refactor world

diff --git a/Engine/Source/Runtime/OpenGLDrv/Private/OpenGLBuffer.cpp b/Engine/Source/Runtime/OpenGLDrv/Private/OpenGLBuffer.cpp
new file mode 100644
--- /dev/null
+++ b/Engine/Source/Runtime/OpenGLDrv/Private/OpenGLBuffer.cpp
@@ -0,0 +1,382 @@
+// Copyright Epic Games, Inc. All Rights Reserved.
+
+/*=============================================================================
+	OpenGLVertexBuffer.cpp: OpenGL vertex buffer RHI implementation.
+=============================================================================*/
+
+#include "CoreMinimal.h"
+#include "Containers/ResourceArray.h"
+#include "HAL/IConsoleManager.h"
+#include "OpenGLDrv.h"
+
+namespace OpenGLConsoleVariables
+{
+#if PLATFORM_ANDROID
+	int32 bUseStagingBuffer = 0;
+#else
+	int32 bUseStagingBuffer = 1;
+#endif
+
+	static FAutoConsoleVariableRef CVarUseStagingBuffer(
+		TEXT("OpenGL.UseStagingBuffer"),
+		bUseStagingBuffer,
+		TEXT("Enables maps of dynamic vertex buffers to go to a staging buffer"),
+		ECVF_ReadOnly
+		);
+
+	extern int32 bUsePersistentMappingStagingBuffer;
+};
+
+static const uint32 MAX_ALIGNMENT_BITS = 8;
+static const uint32 MAX_OFFSET_BITS = 32 - MAX_ALIGNMENT_BITS;
+
+struct PoolAllocation
+{
+	uint8* BasePointer;
+	uint32 SizeWithoutPadding;
+	uint32 Offset				: MAX_OFFSET_BITS;		// into the target buffer
+	uint32 AlignmentPadding		: MAX_ALIGNMENT_BITS;
+	int32 FrameRetired;
+};
+
+static TArray<PoolAllocation*> AllocationList;
+static TMap<void*,PoolAllocation*> AllocationMap;
+
+static GLuint PoolVB = 0;
+static uint8* PoolPointer = 0;
+static uint32 FrameBytes = 0;
+static uint32 FreeSpace = 0;
+static uint32 OffsetVB = 0;
+static const uint32 PerFrameMax = 1024*1024*4;
+static const uint32 MaxAlignment = 1 << MAX_ALIGNMENT_BITS;
+static const uint32 MaxOffset = 1 << MAX_OFFSET_BITS;
+
+void* GetAllocation( void* Target, uint32 Size, uint32 Offset, uint32 Alignment = 16)
+{
+	check(Alignment < MaxAlignment);
+	check(Offset < MaxOffset);
+	check(FMath::IsPowerOfTwo(Alignment));
+
+	uintptr_t AlignmentSubOne = Alignment - 1;
+
+	if (FOpenGL::SupportsBufferStorage() && OpenGLConsoleVariables::bUseStagingBuffer)
+	{
+		if (PoolVB == 0)
+		{
+			FOpenGL::GenBuffers(1, &PoolVB);
+			glBindBuffer(GL_COPY_READ_BUFFER, PoolVB);
+			FOpenGL::BufferStorage(GL_COPY_READ_BUFFER, PerFrameMax * 4, NULL, GL_MAP_WRITE_BIT | GL_MAP_PERSISTENT_BIT | GL_MAP_COHERENT_BIT);
+			PoolPointer = (uint8*)FOpenGL::MapBufferRange(GL_COPY_READ_BUFFER, 0, PerFrameMax * 4, FOpenGL::EResourceLockMode::RLM_WriteOnlyPersistent);
+
+			FreeSpace = PerFrameMax * 4;
+
+			check(PoolPointer);
+		}
+		check (PoolVB);
+
+		uintptr_t AllocHeadPtr = *reinterpret_cast<const uintptr_t*>(&PoolPointer) + OffsetVB;
+		uint32 AlignmentPadBytes = ((AllocHeadPtr + AlignmentSubOne) & (~AlignmentSubOne)) - AllocHeadPtr;
+		uint32 SizeWithAlignmentPad = Size + AlignmentPadBytes;
+
+		if (SizeWithAlignmentPad > PerFrameMax - FrameBytes || SizeWithAlignmentPad > FreeSpace)
+		{
+			return nullptr;
+		}
+
+		if (SizeWithAlignmentPad > (PerFrameMax*4 - OffsetVB))
+		{
+			// We're wrapping, create dummy allocation and start at the begining
+			uint32 Leftover = PerFrameMax*4 - OffsetVB;
+			PoolAllocation* Alloc = new PoolAllocation;
+			Alloc->BasePointer = 0;
+			Alloc->Offset = 0;
+			Alloc->AlignmentPadding = 0;
+			Alloc->SizeWithoutPadding = Leftover;
+			Alloc->FrameRetired = GFrameNumberRenderThread;
+
+			AllocationList.Add(Alloc);
+			OffsetVB = 0;
+			FreeSpace -= Leftover;
+
+			AllocHeadPtr = *reinterpret_cast<const uintptr_t*>(&PoolPointer) + OffsetVB;
+			AlignmentPadBytes = ((AllocHeadPtr + AlignmentSubOne) & (~AlignmentSubOne)) - AllocHeadPtr;
+			SizeWithAlignmentPad = Size + AlignmentPadBytes;
+		}
+
+		//Again check if we have room
+		if (SizeWithAlignmentPad > FreeSpace)
+		{
+			return nullptr;
+		}
+
+		PoolAllocation* Alloc = new PoolAllocation;
+		Alloc->BasePointer = PoolPointer + OffsetVB;
+		Alloc->Offset = Offset;
+		Alloc->AlignmentPadding = AlignmentPadBytes;
+		Alloc->SizeWithoutPadding = Size;
+		Alloc->FrameRetired = -1;
+
+		AllocationList.Add(Alloc);
+		AllocationMap.Add(Target, Alloc);
+		OffsetVB += SizeWithAlignmentPad;
+		FreeSpace -= SizeWithAlignmentPad;
+		FrameBytes += SizeWithAlignmentPad;
+
+		return Alloc->BasePointer + Alloc->AlignmentPadding;
+
+	}
+
+	return nullptr;
+}
+
+bool RetireAllocation( FOpenGLBuffer* Target)
+{
+	if (FOpenGL::SupportsBufferStorage() && OpenGLConsoleVariables::bUseStagingBuffer)
+	{
+		PoolAllocation *Alloc = 0;
+
+		if ( AllocationMap.RemoveAndCopyValue(Target, Alloc))
+		{
+			check(Alloc);
+			Target->Bind();
+
+			FOpenGL::CopyBufferSubData(GL_COPY_READ_BUFFER, GL_ARRAY_BUFFER, (Alloc->BasePointer + Alloc->AlignmentPadding) - PoolPointer, Alloc->Offset, Alloc->SizeWithoutPadding);
+
+			Alloc->FrameRetired = GFrameNumberRenderThread;
+
+			return true;
+		}
+	}
+	return false;
+}
+
+void BeginFrame_VertexBufferCleanup()
+{
+	if (GFrameNumberRenderThread < 3)
+	{
+		return;
+	}
+
+	int32 NumToRetire = 0;
+	int32 FrameToRecover = GFrameNumberRenderThread - 3;
+
+	while (NumToRetire < AllocationList.Num())
+	{
+		PoolAllocation *Alloc = AllocationList[NumToRetire];
+		if (Alloc->FrameRetired < 0 || Alloc->FrameRetired > FrameToRecover)
+		{
+			break;
+		}
+		FreeSpace += (Alloc->SizeWithoutPadding + Alloc->AlignmentPadding);
+		delete Alloc;
+		NumToRetire++;
+	}
+
+	AllocationList.RemoveAt(0,NumToRetire);
+	FrameBytes = 0;
+}
+
+FBufferRHIRef FOpenGLDynamicRHI::RHICreateBuffer(uint32 Size, EBufferUsageFlags Usage, uint32 Stride, ERHIAccess InResourceState, FRHIResourceCreateInfo& CreateInfo)
+{
+	if (CreateInfo.bWithoutNativeResource)
+	{
+		return new FOpenGLBuffer();
+	}
+
+	const void *Data = NULL;
+
+	// If a resource array was provided for the resource, create the resource pre-populated
+	if(CreateInfo.ResourceArray)
+	{
+		check(Size == CreateInfo.ResourceArray->GetResourceDataSize());
+		Data = CreateInfo.ResourceArray->GetResourceData();
+	}
+
+	TRefCountPtr<FOpenGLBuffer> Buffer = new FOpenGLBuffer((Usage & BUF_IndexBuffer)? GL_ARRAY_BUFFER : GL_ELEMENT_ARRAY_BUFFER, Stride, Size, Usage, Data);
+	
+	if (CreateInfo.ResourceArray)
+	{
+		CreateInfo.ResourceArray->Discard();
+	}
+	
+	return Buffer.GetReference();
+}
+
+void* FOpenGLDynamicRHI::LockBuffer_BottomOfPipe(FRHICommandListImmediate& RHICmdList, FRHIBuffer* BufferRHI, uint32 Offset, uint32 Size, EResourceLockMode LockMode)
+{
+	check(Size > 0);
+	RHITHREAD_GLCOMMAND_PROLOGUE();
+
+	VERIFY_GL_SCOPE();
+	FOpenGLBuffer* Buffer = ResourceCast(BufferRHI);
+	if( !(FOpenGL::SupportsVertexAttribBinding() && OpenGLConsoleVariables::bUseVAB) && Buffer->GetUsage() & BUF_ZeroStride )
+	{
+		check( Offset + Size <= Buffer->GetSize() );
+		// We assume we are only using the first elements from the VB, so we can later on read this memory and create an expanded version of this zero stride buffer
+		check(Offset == 0);
+		return (void*)( (uint8*)Buffer->GetZeroStrideBuffer() + Offset );
+	}
+	else
+	{
+		if (Buffer->IsDynamic() && LockMode == EResourceLockMode::RLM_WriteOnly)
+		{
+			void *Staging = GetAllocation(Buffer, Size, Offset);
+			if (Staging)
+			{
+				return Staging;
+			}
+		}
+		return (void*)Buffer->Lock(Offset, Size, LockMode == EResourceLockMode::RLM_ReadOnly, Buffer->IsDynamic());
+	}
+	RHITHREAD_GLCOMMAND_EPILOGUE_RETURN(void*);
+}
+
+void FOpenGLDynamicRHI::UnlockBuffer_BottomOfPipe(FRHICommandListImmediate& RHICmdList, FRHIBuffer* BufferRHI)
+{
+	RHITHREAD_GLCOMMAND_PROLOGUE();
+	VERIFY_GL_SCOPE();
+	FOpenGLBuffer* Buffer = ResourceCast(BufferRHI);
+	if( (FOpenGL::SupportsVertexAttribBinding() && OpenGLConsoleVariables::bUseVAB) || !( Buffer->GetUsage() & BUF_ZeroStride ) )
+	{
+		if (!RetireAllocation(Buffer))
+		{
+			Buffer->Unlock();
+		}
+	}
+	RHITHREAD_GLCOMMAND_EPILOGUE();
+}
+
+void FOpenGLDynamicRHI::RHICopyBuffer(FRHIBuffer* SourceBufferRHI, FRHIBuffer* DestBufferRHI)
+{
+	check(SourceBufferRHI && DestBufferRHI && SourceBufferRHI->GetSize() == DestBufferRHI->GetSize());
+	RHICopyBufferRegion(DestBufferRHI, 0, SourceBufferRHI, 0, SourceBufferRHI->GetSize());
+}
+
+void FOpenGLDynamicRHI::RHICopyBufferRegion(FRHIBuffer* DestBufferRHI, uint64 DstOffset, FRHIBuffer* SourceBufferRHI, uint64 SrcOffset, uint64 NumBytes)
+{
+	VERIFY_GL_SCOPE();
+	FOpenGLBuffer* SourceBuffer = ResourceCast(SourceBufferRHI);
+	FOpenGLBuffer* DestBuffer = ResourceCast(DestBufferRHI);
+
+	glBindBuffer(GL_COPY_READ_BUFFER, SourceBuffer->Resource);
+	glBindBuffer(GL_COPY_WRITE_BUFFER, DestBuffer->Resource);
+	FOpenGL::CopyBufferSubData(GL_COPY_READ_BUFFER, GL_COPY_WRITE_BUFFER, SrcOffset, DstOffset, NumBytes);
+	glBindBuffer(GL_COPY_READ_BUFFER, 0);
+	glBindBuffer(GL_COPY_WRITE_BUFFER, 0);
+}
+
+void FOpenGLDynamicRHI::RHITransferBufferUnderlyingResource(FRHIBuffer* DestBuffer, FRHIBuffer* SrcBuffer)
+{
+	VERIFY_GL_SCOPE();
+	check(DestBuffer);
+	FOpenGLBuffer* Dest = ResourceCast(DestBuffer);
+	if (!SrcBuffer)
+	{
+		TRefCountPtr<FOpenGLBuffer> Src = new FOpenGLBuffer();
+		Dest->Swap(*Src);
+	}
+	else
+	{
+		FOpenGLBuffer* Src = ResourceCast(SrcBuffer);
+		Dest->Swap(*Src);
+	}
+}
+
+FStagingBufferRHIRef FOpenGLDynamicRHI::RHICreateStagingBuffer()
+{
+	return new FOpenGLStagingBuffer();
+}
+
+void FOpenGLStagingBuffer::Initialize()
+{
+	ShadowBuffer = 0;
+	ShadowSize = 0;
+	Mapping = nullptr;
+	FRHICommandListImmediate& RHICmdList = FRHICommandListExecutor::GetImmediateCommandList();
+	RHITHREAD_GLCOMMAND_PROLOGUE();
+	VERIFY_GL_SCOPE();
+	glGenBuffers(1, &ShadowBuffer);
+	RHITHREAD_GLCOMMAND_EPILOGUE();
+}
+
+FOpenGLStagingBuffer::~FOpenGLStagingBuffer()
+{
+	FRHICommandListImmediate& RHICmdList = FRHICommandListExecutor::GetImmediateCommandList();
+	RHITHREAD_GLCOMMAND_PROLOGUE();
+	VERIFY_GL_SCOPE();
+	glDeleteBuffers(1, &ShadowBuffer);
+	RHITHREAD_GLCOMMAND_EPILOGUE_NORETURN();
+}
+
+// If we do not support the BufferStorage extension or if PersistentMapping is set to false, this will send the command to the RHI and flush it
+// If we do support BufferStorage extension and PersistentMapping is set to true, we just return the pointer + offset
+void* FOpenGLStagingBuffer::Lock(uint32 Offset, uint32 NumBytes)
+{
+	if (!FOpenGL::SupportsBufferStorage() || !OpenGLConsoleVariables::bUsePersistentMappingStagingBuffer)
+	{
+		FRHICommandListImmediate& RHICmdList = FRHICommandListExecutor::GetImmediateCommandList();
+		RHITHREAD_GLCOMMAND_PROLOGUE();
+		VERIFY_GL_SCOPE();
+
+		check(ShadowBuffer != 0);
+		glBindBuffer(GL_COPY_WRITE_BUFFER, ShadowBuffer);
+		void* LocalMapping = FOpenGL::MapBufferRange(GL_COPY_WRITE_BUFFER, 0, NumBytes, FOpenGL::EResourceLockMode::RLM_ReadOnly);
+		check(LocalMapping);
+		return reinterpret_cast<uint8*>(LocalMapping) + Offset;
+
+		RHITHREAD_GLCOMMAND_EPILOGUE_RETURN(void*);
+	}
+	else
+	{
+		check(Mapping != nullptr);
+		return reinterpret_cast<uint8*>(Mapping) + Offset;
+	}
+}
+
+// If we do not support the BufferStorage extension or if PersistentMapping is set to false, this will send the command to the RHI and flush it
+// If we do support BufferStorage extension and PersistentMapping is set to true, we do nothing
+void FOpenGLStagingBuffer::Unlock()
+{
+	if (!FOpenGL::SupportsBufferStorage() || !OpenGLConsoleVariables::bUsePersistentMappingStagingBuffer)
+	{
+		FRHICommandListImmediate& RHICmdList = FRHICommandListExecutor::GetImmediateCommandList();
+		RHITHREAD_GLCOMMAND_PROLOGUE();
+		FOpenGL::UnmapBuffer(GL_COPY_WRITE_BUFFER);
+		Mapping = nullptr;
+		glBindBuffer(GL_COPY_WRITE_BUFFER, 0);
+		RHITHREAD_GLCOMMAND_EPILOGUE();
+	}
+}
+
+void* FOpenGLDynamicRHI::RHILockStagingBuffer(FRHIStagingBuffer* StagingBuffer, FRHIGPUFence* Fence, uint32 Offset, uint32 SizeRHI)
+{
+	FOpenGLStagingBuffer* Buffer = ResourceCast(StagingBuffer);
+	return Buffer->Lock(Offset, SizeRHI);	
+}
+
+void FOpenGLDynamicRHI::RHIUnlockStagingBuffer(FRHIStagingBuffer* StagingBuffer)
+{
+	FOpenGLStagingBuffer* Buffer = ResourceCast(StagingBuffer);
+	Buffer->Unlock();
+}
+
+void* FOpenGLDynamicRHI::LockStagingBuffer_RenderThread(class FRHICommandListImmediate& RHICmdList, FRHIStagingBuffer* StagingBuffer, FRHIGPUFence* Fence, uint32 Offset, uint32 SizeRHI)
+{
+	check(IsInRenderingThread());
+	if (!Fence || !Fence->Poll() || Fence->NumPendingWriteCommands.GetValue() != 0)
+	{
+		QUICK_SCOPE_CYCLE_COUNTER(STAT_FDynamicRHI_LockStagingBuffer_Flush);
+		RHICmdList.ImmediateFlush(EImmediateFlushType::FlushRHIThread);
+	}
+	{
+		QUICK_SCOPE_CYCLE_COUNTER(STAT_FDynamicRHI_LockStagingBuffer_RenderThread);
+		return GDynamicRHI->RHILockStagingBuffer(StagingBuffer, Fence, Offset, SizeRHI);
+	}
+}
+
+void FOpenGLDynamicRHI::UnlockStagingBuffer_RenderThread(class FRHICommandListImmediate& RHICmdList, FRHIStagingBuffer* StagingBuffer)
+{
+	QUICK_SCOPE_CYCLE_COUNTER(STAT_FDynamicRHI_UnlockStagingBuffer_RenderThread);
+	check(IsInRenderingThread());
+	GDynamicRHI->RHIUnlockStagingBuffer(StagingBuffer);
+}
diff --git a/Engine/Source/Runtime/OpenGLDrv/Private/OpenGLVertexBuffer.cpp b/Engine/Source/Runtime/OpenGLDrv/Private/OpenGLVertexBuffer.cpp
deleted file mode 100644
diff --git a/Engine/Source/Runtime/VulkanRHI/Private/VulkanBuffer.cpp b/Engine/Source/Runtime/VulkanRHI/Private/VulkanBuffer.cpp
new file mode 100644
--- /dev/null
+++ b/Engine/Source/Runtime/VulkanRHI/Private/VulkanBuffer.cpp
@@ -0,0 +1,483 @@
+// Copyright Epic Games, Inc. All Rights Reserved.
+
+/*=============================================================================
+	VulkanIndexBuffer.cpp: Vulkan Index buffer RHI implementation.
+=============================================================================*/
+
+#include "VulkanRHIPrivate.h"
+#include "VulkanDevice.h"
+#include "VulkanContext.h"
+#include "Containers/ResourceArray.h"
+#include "VulkanLLM.h"
+
+
+static TMap<FVulkanResourceMultiBuffer*, VulkanRHI::FPendingBufferLock> GPendingLockIBs;
+static FCriticalSection GPendingLockIBsMutex;
+
+static FORCEINLINE void UpdateVulkanBufferStats(uint64_t Size, VkBufferUsageFlags Usage, bool Allocating)
+{
+	const bool bUniformBuffer = !!(Usage & VK_BUFFER_USAGE_UNIFORM_BUFFER_BIT);
+	const bool bIndexBuffer = !!(Usage & VK_BUFFER_USAGE_INDEX_BUFFER_BIT);
+	const bool bVertexBuffer = !!(Usage & VK_BUFFER_USAGE_VERTEX_BUFFER_BIT);
+
+	if (Allocating)
+	{
+		if (bUniformBuffer)
+		{
+			INC_MEMORY_STAT_BY(STAT_UniformBufferMemory, Size);
+		}
+		else if (bIndexBuffer)
+		{
+			INC_MEMORY_STAT_BY(STAT_IndexBufferMemory, Size);
+		}
+		else if (bVertexBuffer)
+		{
+			INC_MEMORY_STAT_BY(STAT_VertexBufferMemory, Size);
+		}
+		else
+		{
+			INC_MEMORY_STAT_BY(STAT_StructuredBufferMemory, Size);
+		}
+	}
+	else
+	{
+		if (bUniformBuffer)
+		{
+			DEC_MEMORY_STAT_BY(STAT_UniformBufferMemory, Size);
+		}
+		else if (bIndexBuffer)
+		{
+			DEC_MEMORY_STAT_BY(STAT_IndexBufferMemory, Size);
+		}
+		else if (bVertexBuffer)
+		{
+			DEC_MEMORY_STAT_BY(STAT_VertexBufferMemory, Size);
+		}
+		else
+		{
+			DEC_MEMORY_STAT_BY(STAT_StructuredBufferMemory, Size);
+		}
+	}
+}
+
+FVulkanResourceMultiBuffer::FVulkanResourceMultiBuffer(FVulkanDevice* InDevice, uint32 InSize, uint32 InUEUsage, uint32 InStride, FRHIResourceCreateInfo& CreateInfo, class FRHICommandListImmediate* InRHICmdList)
+	: FRHIBuffer(InSize, InUEUsage, InStride)
+	, VulkanRHI::FDeviceChild(InDevice)
+	, UEUsage(InUEUsage)
+	, BufferUsageFlags(0)
+	, NumBuffers(0)
+	, DynamicBufferIndex(0)
+{
+	VULKAN_TRACK_OBJECT_CREATE(FVulkanResourceMultiBuffer, this);
+
+	if ((InUEUsage & BUF_VertexBuffer) != 0)
+	{
+		BufferUsageFlags |= VK_BUFFER_USAGE_VERTEX_BUFFER_BIT;
+	}
+	if ((InUEUsage & BUF_IndexBuffer) != 0)
+	{
+		BufferUsageFlags |= VK_BUFFER_USAGE_INDEX_BUFFER_BIT;
+	}
+	if ((InUEUsage & BUF_StructuredBuffer) != 0)
+	{
+		BufferUsageFlags |= VK_BUFFER_USAGE_STORAGE_BUFFER_BIT;
+	}
+
+	if (InSize > 0)
+	{
+		const bool bStatic = (InUEUsage & BUF_Static) != 0;
+		const bool bDynamic = (InUEUsage & BUF_Dynamic) != 0;
+		const bool bVolatile = (InUEUsage & BUF_Volatile) != 0;
+		const bool bShaderResource = (InUEUsage & BUF_ShaderResource) != 0;
+		const bool bIsUniformBuffer = (BufferUsageFlags & VK_BUFFER_USAGE_UNIFORM_BUFFER_BIT) != 0;
+		const bool bUAV = (InUEUsage & BUF_UnorderedAccess) != 0;
+		const bool bIndirect = (InUEUsage & BUF_DrawIndirect) == BUF_DrawIndirect;
+		const bool bCPUReadable = (UEUsage & BUF_KeepCPUAccessible) != 0;
+		const bool bCopySource = (UEUsage & BUF_SourceCopy) != 0;
+
+		BufferUsageFlags |= bVolatile ? 0 : VK_BUFFER_USAGE_TRANSFER_DST_BIT;
+		BufferUsageFlags |= (bShaderResource && !bIsUniformBuffer) ? VK_BUFFER_USAGE_UNIFORM_TEXEL_BUFFER_BIT : 0;
+		BufferUsageFlags |= bUAV ? VK_BUFFER_USAGE_STORAGE_TEXEL_BUFFER_BIT : 0;
+		BufferUsageFlags |= bIndirect ? VK_BUFFER_USAGE_INDIRECT_BUFFER_BIT : 0;
+		BufferUsageFlags |= bCPUReadable ? (VK_BUFFER_USAGE_TRANSFER_SRC_BIT | VK_BUFFER_USAGE_TRANSFER_DST_BIT) : 0;
+		BufferUsageFlags |= bCopySource ? VK_BUFFER_USAGE_TRANSFER_SRC_BIT : 0;
+
+		if (bVolatile)
+		{
+			bool bRenderThread = IsInRenderingThread();
+
+			// Get a dummy buffer as sometimes the high-level misbehaves and tries to use SRVs off volatile buffers before filling them in...
+			void* Data = Lock(bRenderThread, RLM_WriteOnly, InSize, 0);
+			FMemory::Memzero(Data, InSize);
+			Unlock(bRenderThread);
+		}
+		else
+		{
+			VkDevice VulkanDevice = InDevice->GetInstanceHandle();
+
+			VkMemoryPropertyFlags BufferMemFlags = VK_MEMORY_PROPERTY_DEVICE_LOCAL_BIT;
+			const bool bUnifiedMem = InDevice->HasUnifiedMemory();
+			if (bUnifiedMem)
+			{
+				BufferMemFlags |= (VK_MEMORY_PROPERTY_HOST_VISIBLE_BIT | VK_MEMORY_PROPERTY_HOST_COHERENT_BIT);
+			}
+
+			NumBuffers = bDynamic ? NUM_BUFFERS : 1;
+			check(NumBuffers <= UE_ARRAY_COUNT(Buffers));
+
+			for (uint32 Index = 0; Index < NumBuffers; ++Index)
+			{
+				if(!InDevice->GetMemoryManager().AllocateBufferPooled(Buffers[Index], this, InSize, BufferUsageFlags, BufferMemFlags, EVulkanAllocationMetaMultiBuffer, __FILE__, __LINE__))
+				{
+					InDevice->GetMemoryManager().HandleOOM();
+				}
+			}
+			Current.Alloc.Reference(Buffers[DynamicBufferIndex]);
+			Current.Handle = (VkBuffer)Current.Alloc.VulkanHandle;
+			Current.Offset = Current.Alloc.Offset;
+			Current.Size = InSize;
+
+			bool bRenderThread = (InRHICmdList == nullptr);
+			if (bRenderThread)
+			{
+				ensure(IsInRenderingThread());
+			}
+
+			if (CreateInfo.ResourceArray)
+			{
+				uint32 CopyDataSize = FMath::Min(InSize, CreateInfo.ResourceArray->GetResourceDataSize());
+				void* Data = Lock(bRenderThread, RLM_WriteOnly, CopyDataSize, 0);
+				FMemory::Memcpy(Data, CreateInfo.ResourceArray->GetResourceData(), CopyDataSize);
+				Unlock(bRenderThread);
+
+				CreateInfo.ResourceArray->Discard();
+			}
+
+			UpdateVulkanBufferStats(InSize * NumBuffers, BufferUsageFlags, true);
+		}
+	}
+}
+
+FVulkanResourceMultiBuffer::~FVulkanResourceMultiBuffer()
+{
+	VULKAN_TRACK_OBJECT_DELETE(FVulkanResourceMultiBuffer, this);
+	uint64_t TotalSize = 0;
+	for (uint32 Index = 0; Index < NumBuffers; ++Index)
+	{
+		TotalSize += Buffers[Index].Size;
+		Device->GetMemoryManager().FreeVulkanAllocation(Buffers[Index]);
+	}
+	UpdateVulkanBufferStats(TotalSize, BufferUsageFlags, false);
+}
+
+void* FVulkanResourceMultiBuffer::Lock(bool bFromRenderingThread, EResourceLockMode LockMode, uint32 LockSize, uint32 Offset)
+{
+	void* Data = nullptr;
+
+	const bool bStatic = (UEUsage & BUF_Static) != 0;
+	const bool bDynamic = (UEUsage & BUF_Dynamic) != 0;
+	const bool bVolatile = (UEUsage & BUF_Volatile) != 0;
+	const bool bCPUReadable = (UEUsage & BUF_KeepCPUAccessible) != 0;
+	const bool bUAV = (UEUsage & BUF_UnorderedAccess) != 0;
+	const bool bSR = (UEUsage & BUF_ShaderResource) != 0;
+
+	if (bVolatile)
+	{
+		check(NumBuffers == 0);
+		if (LockMode == RLM_ReadOnly)
+		{
+			checkf(0, TEXT("Volatile buffers can't be locked for read."));
+		}
+		else
+		{
+			Device->GetImmediateContext().GetTempFrameAllocationBuffer().Alloc(LockSize + Offset, 256, VolatileLockInfo);
+			Data = VolatileLockInfo.Data;
+			++VolatileLockInfo.LockCounter;
+			check(!VolatileLockInfo.Allocation.HasAllocation());
+			Current.Alloc = VolatileLockInfo.Allocation;
+			Current.Handle = Current.Alloc.GetBufferHandle();
+			Current.Offset = VolatileLockInfo.CurrentOffset + Current.Alloc.Offset;
+			Current.Size = LockSize;
+		}
+	}
+	else
+	{
+		check(bStatic || bDynamic || bUAV || bSR);
+
+		if (LockMode == RLM_ReadOnly)
+		{
+			const bool bUnifiedMem = Device->HasUnifiedMemory();
+			if (bUnifiedMem)
+			{
+				Data = (uint8*)Buffers[DynamicBufferIndex].GetMappedPointer(Device) + Offset;
+			}
+			else 
+			{
+				Device->PrepareForCPURead();
+				FVulkanCommandListContext& ImmediateContext = Device->GetImmediateContext();
+				FVulkanCmdBuffer* CmdBuffer = ImmediateContext.GetCommandBufferManager()->GetUploadCmdBuffer();
+				
+				// Make sure any previous tasks have finished on the source buffer.
+				VkMemoryBarrier BarrierBefore = { VK_STRUCTURE_TYPE_MEMORY_BARRIER, nullptr, VK_ACCESS_MEMORY_WRITE_BIT, VK_ACCESS_MEMORY_READ_BIT };
+				VulkanRHI::vkCmdPipelineBarrier(CmdBuffer->GetHandle(), VK_PIPELINE_STAGE_TRANSFER_BIT, VK_PIPELINE_STAGE_TRANSFER_BIT, 0, 1, &BarrierBefore, 0, nullptr, 0, nullptr);
+
+				// Create a staging buffer we can use to copy data from device to cpu.
+				VulkanRHI::FStagingBuffer* StagingBuffer = Device->GetStagingManager().AcquireBuffer(LockSize, VK_BUFFER_USAGE_TRANSFER_DST_BIT, VK_MEMORY_PROPERTY_HOST_CACHED_BIT);
+
+				// Fill the staging buffer with the data on the device.
+				VkBufferCopy Regions;
+				Regions.size = LockSize;
+				Regions.srcOffset = Offset;
+				Regions.dstOffset = 0;
+
+				VulkanRHI::vkCmdCopyBuffer(CmdBuffer->GetHandle(), Buffers[DynamicBufferIndex].GetBufferHandle(), StagingBuffer->GetHandle(), 1, &Regions);
+
+				// Setup barrier.
+				VkMemoryBarrier BarrierAfter = { VK_STRUCTURE_TYPE_MEMORY_BARRIER, nullptr, VK_ACCESS_MEMORY_WRITE_BIT, VK_ACCESS_HOST_READ_BIT };
+				VulkanRHI::vkCmdPipelineBarrier(CmdBuffer->GetHandle(), VK_PIPELINE_STAGE_TRANSFER_BIT, VK_PIPELINE_STAGE_HOST_BIT, 0, 1, &BarrierAfter, 0, nullptr, 0, nullptr);
+				
+				// Force upload.
+				ImmediateContext.GetCommandBufferManager()->SubmitUploadCmdBuffer();
+				Device->WaitUntilIdle();
+
+				// Flush.
+				StagingBuffer->FlushMappedMemory();
+
+				// Get mapped pointer. 
+				Data = StagingBuffer->GetMappedPointer();
+
+				// Release temp staging buffer during unlock.
+				VulkanRHI::FPendingBufferLock PendingLock;
+				PendingLock.Offset = 0;
+				PendingLock.Size = LockSize;
+				PendingLock.LockMode = LockMode;
+				PendingLock.StagingBuffer = StagingBuffer;
+
+				{
+					FScopeLock ScopeLock(&GPendingLockIBsMutex);
+					check(!GPendingLockIBs.Contains(this));
+					GPendingLockIBs.Add(this, PendingLock);
+				}
+
+				ImmediateContext.GetCommandBufferManager()->PrepareForNewActiveCommandBuffer();
+			}
+		}
+		else
+		{
+			check(LockMode == RLM_WriteOnly);
+			DynamicBufferIndex = (DynamicBufferIndex + 1) % NumBuffers;
+			Current.Alloc.Reference(Buffers[DynamicBufferIndex]);		
+			Current.Handle = (VkBuffer)Current.Alloc.VulkanHandle;
+			Current.Offset = Current.Alloc.Offset;
+			Current.Size = LockSize;
+
+			const bool bUnifiedMem = Device->HasUnifiedMemory();
+			if (bUnifiedMem)
+			{
+				Data = (uint8*)Buffers[DynamicBufferIndex].GetMappedPointer(Device) + Offset;
+			}
+			else
+			{
+				VulkanRHI::FPendingBufferLock PendingLock;
+				PendingLock.Offset = Offset;
+				PendingLock.Size = LockSize;
+				PendingLock.LockMode = LockMode;
+
+				VulkanRHI::FStagingBuffer* StagingBuffer = Device->GetStagingManager().AcquireBuffer(LockSize, VK_BUFFER_USAGE_TRANSFER_SRC_BIT);
+				PendingLock.StagingBuffer = StagingBuffer;
+				Data = StagingBuffer->GetMappedPointer();
+
+				{
+					FScopeLock ScopeLock(&GPendingLockIBsMutex);
+					check(!GPendingLockIBs.Contains(this));
+					GPendingLockIBs.Add(this, PendingLock);
+				}
+			}
+
+		}
+	}
+
+	check(Data);
+	return Data;
+}
+
+inline void FVulkanResourceMultiBuffer::InternalUnlock(FVulkanCommandListContext& Context, VulkanRHI::FPendingBufferLock& PendingLock, FVulkanResourceMultiBuffer* MultiBuffer, int32 InDynamicBufferIndex)
+{
+	uint32 LockSize = PendingLock.Size;
+	uint32 LockOffset = PendingLock.Offset;
+	VulkanRHI::FStagingBuffer* StagingBuffer = PendingLock.StagingBuffer;
+	PendingLock.StagingBuffer = nullptr;
+
+	FVulkanCmdBuffer* Cmd = Context.GetCommandBufferManager()->GetUploadCmdBuffer();
+	if (!Cmd->HasBegun())
+	{
+		Cmd->Begin();
+	}
+	ensure(Cmd->IsOutsideRenderPass());
+	VkCommandBuffer CmdBuffer = Cmd->GetHandle();
+
+	VulkanRHI::DebugHeavyWeightBarrier(CmdBuffer, 16);
+
+	VkBufferCopy Region;
+	FMemory::Memzero(Region);
+	Region.size = LockSize;
+	//Region.srcOffset = 0;
+	Region.dstOffset = LockOffset + MultiBuffer->Buffers[InDynamicBufferIndex].Offset;
+	VulkanRHI::vkCmdCopyBuffer(CmdBuffer, StagingBuffer->GetHandle(), MultiBuffer->Buffers[InDynamicBufferIndex].GetBufferHandle(), 1, &Region);
+
+	MultiBuffer->GetParent()->GetStagingManager().ReleaseBuffer(Cmd, StagingBuffer);
+}
+
+struct FRHICommandMultiBufferUnlock final : public FRHICommand<FRHICommandMultiBufferUnlock>
+{
+	VulkanRHI::FPendingBufferLock PendingLock;
+	FVulkanResourceMultiBuffer* MultiBuffer;
+	FVulkanDevice* Device;
+	int32 DynamicBufferIndex;
+
+	FRHICommandMultiBufferUnlock(FVulkanDevice* InDevice, const VulkanRHI::FPendingBufferLock& InPendingLock, FVulkanResourceMultiBuffer* InMultiBuffer, int32 InDynamicBufferIndex)
+		: PendingLock(InPendingLock)
+		, MultiBuffer(InMultiBuffer)
+		, Device(InDevice)
+		, DynamicBufferIndex(InDynamicBufferIndex)
+	{
+	}
+
+	void Execute(FRHICommandListBase& CmdList)
+	{
+		FVulkanResourceMultiBuffer::InternalUnlock(FVulkanCommandListContext::GetVulkanContext(CmdList.GetContext()), PendingLock, MultiBuffer, DynamicBufferIndex);
+	}
+};
+
+
+void FVulkanResourceMultiBuffer::Unlock(bool bFromRenderingThread)
+{
+	const bool bStatic = (UEUsage & BUF_Static) != 0;
+	const bool bDynamic = (UEUsage & BUF_Dynamic) != 0;
+	const bool bVolatile = (UEUsage & BUF_Volatile) != 0;
+	const bool bCPUReadable = (UEUsage & BUF_KeepCPUAccessible) != 0;
+	const bool bSR = (UEUsage & BUF_ShaderResource) != 0;
+
+	if (bVolatile)
+	{
+		check(NumBuffers == 0);
+
+		// Nothing to do here...
+	}
+	else
+	{
+		check(bStatic || bDynamic || bSR);
+
+		const bool bUnifiedMem = Device->HasUnifiedMemory();
+		if (bUnifiedMem)
+		{
+			// Nothing to do here...
+			return;
+		}
+
+		VulkanRHI::FPendingBufferLock PendingLock;
+		bool bFound = false;
+		{
+			// Found only if it was created for Write
+			FScopeLock ScopeLock(&GPendingLockIBsMutex);
+			bFound = GPendingLockIBs.RemoveAndCopyValue(this, PendingLock);
+		}
+
+		PendingLock.StagingBuffer->FlushMappedMemory();
+
+		checkf(bFound, TEXT("Mismatched lock/unlock IndexBuffer!"));
+		if (PendingLock.LockMode == RLM_WriteOnly)
+		{
+			FRHICommandList& RHICmdList = FRHICommandListExecutor::GetImmediateCommandList();
+			if (!bFromRenderingThread || (RHICmdList.Bypass() || !IsRunningRHIInSeparateThread()))
+			{
+				FVulkanResourceMultiBuffer::InternalUnlock(Device->GetImmediateContext(), PendingLock, this, DynamicBufferIndex);
+			}
+			else
+			{
+				check(IsInRenderingThread());
+				ALLOC_COMMAND_CL(RHICmdList, FRHICommandMultiBufferUnlock)(Device, PendingLock, this, DynamicBufferIndex);
+			}
+		}
+		else if(PendingLock.LockMode == RLM_ReadOnly)
+		{
+			// Just remove the staging buffer here.
+			Device->GetStagingManager().ReleaseBuffer(0, PendingLock.StagingBuffer);
+		}
+	}
+}
+
+void FVulkanResourceMultiBuffer::Swap(FVulkanResourceMultiBuffer& Other)
+{
+	FRHIBuffer::Swap(Other);
+	::Swap(*this, Other);
+}
+
+FBufferRHIRef FVulkanDynamicRHI::RHICreateBuffer(uint32 Size, EBufferUsageFlags Usage, uint32 Stride, ERHIAccess ResourceState, FRHIResourceCreateInfo& CreateInfo)
+{
+	LLM_SCOPE_VULKAN(ELLMTagVulkan::VulkanBuffers);
+	if (CreateInfo.bWithoutNativeResource)
+	{
+		return new FVulkanResourceMultiBuffer(nullptr, 0, 0, 0, CreateInfo, nullptr);
+	}
+	return new FVulkanResourceMultiBuffer(Device, Size, Usage, Stride, CreateInfo, nullptr);
+}
+
+void* FVulkanDynamicRHI::LockBuffer_BottomOfPipe(FRHICommandListImmediate& RHICmdList, FRHIBuffer* BufferRHI, uint32 Offset, uint32 Size, EResourceLockMode LockMode)
+{
+	LLM_SCOPE_VULKAN(ELLMTagVulkan::VulkanBuffers);
+	FVulkanResourceMultiBuffer* Buffer = ResourceCast(BufferRHI);
+	return Buffer->Lock(false, LockMode, Size, Offset);
+}
+
+void FVulkanDynamicRHI::UnlockBuffer_BottomOfPipe(FRHICommandListImmediate& RHICmdList, FRHIBuffer* BufferRHI)
+{
+	LLM_SCOPE_VULKAN(ELLMTagVulkan::VulkanBuffers);
+	FVulkanResourceMultiBuffer* Buffer = ResourceCast(BufferRHI);
+	Buffer->Unlock(false);
+}
+
+#if VULKAN_BUFFER_LOCK_THREADSAFE
+void* FVulkanDynamicRHI::LockBuffer_RenderThread(class FRHICommandListImmediate& RHICmdList, FRHIBuffer* BufferRHI, uint32 Offset, uint32 SizeRHI, EResourceLockMode LockMode)
+{
+	return this->RHILockBuffer(BufferRHI, Offset, SizeRHI, LockMode);
+}
+
+void FVulkanDynamicRHI::UnlockBuffer_RenderThread(FRHICommandListImmediate& RHICmdList, FRHIBuffer* BufferRHI)
+{
+	this->RHIUnlockBuffer(BufferRHI);
+}
+#endif
+
+void FVulkanDynamicRHI::RHICopyBuffer(FRHIBuffer* SourceBufferRHI, FRHIBuffer* DestBufferRHI)
+{
+	VULKAN_SIGNAL_UNIMPLEMENTED();
+}
+
+void FVulkanDynamicRHI::RHITransferBufferUnderlyingResource(FRHIBuffer* DestBuffer, FRHIBuffer* SrcBuffer)
+{
+	check(DestBuffer);
+	FVulkanResourceMultiBuffer* Dest = ResourceCast(DestBuffer);
+	if (!SrcBuffer)
+	{
+		FRHIResourceCreateInfo CreateInfo;
+		TRefCountPtr<FVulkanResourceMultiBuffer> DeletionProxy = new FVulkanResourceMultiBuffer(Dest->GetParent(), 0, 0, 0, CreateInfo, nullptr);
+		Dest->Swap(*DeletionProxy);
+	}
+	else
+	{
+		FVulkanResourceMultiBuffer* Src = ResourceCast(SrcBuffer);
+		Dest->Swap(*Src);
+	}
+}
+
+void FVulkanResourceMultiBuffer::Evict(FVulkanDevice& Device_)
+{
+	checkNoEntry();//Not Implemented, should never be called
+}
+void FVulkanResourceMultiBuffer::Move(FVulkanDevice& Device_, VulkanRHI::FVulkanAllocation& NewAllocation)
+{
+	checkNoEntry();//Not Implemented, should never be called
+}
+
+
diff --git a/Engine/Source/Runtime/VulkanRHI/Private/VulkanIndexBuffer.cpp b/Engine/Source/Runtime/VulkanRHI/Private/VulkanIndexBuffer.cpp
deleted file mode 100644
diff --git a/Engine/Source/Runtime/Windows/D3D11RHI/Private/D3D11Buffer.cpp b/Engine/Source/Runtime/Windows/D3D11RHI/Private/D3D11Buffer.cpp
new file mode 100644
--- /dev/null
+++ b/Engine/Source/Runtime/Windows/D3D11RHI/Private/D3D11Buffer.cpp
@@ -0,0 +1,273 @@
+// Copyright Epic Games, Inc. All Rights Reserved.
+
+/*=============================================================================
+	D3D11IndexBuffer.cpp: D3D Index buffer RHI implementation.
+=============================================================================*/
+
+#include "D3D11RHIPrivate.h"
+
+TAutoConsoleVariable<int32> GCVarUseSharedKeyedMutex(
+	TEXT("r.D3D11.UseSharedKeyMutex"),
+	0,
+	TEXT("If 1, BUF_Shared vertex / index buffer and TexCreate_Shared texture will be created\n")
+	TEXT("with the D3D11_RESOURCE_MISC_SHARED_KEYEDMUTEX flag instead of D3D11_RESOURCE_MISC_SHARED (default).\n"),
+	ECVF_Default);
+
+FBufferRHIRef FD3D11DynamicRHI::RHICreateBuffer(uint32 Size, EBufferUsageFlags Usage, uint32 Stride, ERHIAccess ResourceState, FRHIResourceCreateInfo& CreateInfo)
+{
+	if (CreateInfo.bWithoutNativeResource)
+	{
+		return new FD3D11Buffer();
+	}
+
+	// Explicitly check that the size is nonzero before allowing CreateBuffer to opaquely fail.
+	check(Size > 0);
+
+	// Describe the buffer.
+	D3D11_BUFFER_DESC Desc = { Size };
+
+	if (Usage & BUF_AnyDynamic)
+	{
+		Desc.Usage = D3D11_USAGE_DYNAMIC;
+		Desc.CPUAccessFlags = D3D11_CPU_ACCESS_WRITE;
+	}
+
+	if (Usage & BUF_VertexBuffer)
+	{
+		Desc.BindFlags |= D3D11_BIND_VERTEX_BUFFER;
+	}
+
+	if (Usage & BUF_IndexBuffer)
+	{
+		Desc.BindFlags |= D3D11_BIND_INDEX_BUFFER;
+	}
+
+	if (Usage & BUF_ByteAddressBuffer)
+	{
+		Desc.MiscFlags |= D3D11_RESOURCE_MISC_BUFFER_ALLOW_RAW_VIEWS;
+	}
+	else if (Usage & BUF_StructuredBuffer)
+	{
+		Desc.StructureByteStride = Stride;
+		Desc.MiscFlags |= D3D11_RESOURCE_MISC_BUFFER_STRUCTURED;
+	}
+
+	if (Usage & BUF_ShaderResource)
+	{
+		Desc.BindFlags |= D3D11_BIND_SHADER_RESOURCE;
+	}
+
+	if (Usage & BUF_UnorderedAccess)
+	{
+		Desc.BindFlags |= D3D11_BIND_UNORDERED_ACCESS;
+	}
+
+	if (Usage & BUF_DrawIndirect)
+	{
+		Desc.MiscFlags |= D3D11_RESOURCE_MISC_DRAWINDIRECT_ARGS;
+	}
+
+	if (Usage & BUF_Shared)
+	{
+		if (GCVarUseSharedKeyedMutex->GetInt() != 0)
+		{
+			Desc.MiscFlags |= D3D11_RESOURCE_MISC_SHARED_KEYEDMUTEX;
+		}
+		else
+		{
+			Desc.MiscFlags |= D3D11_RESOURCE_MISC_SHARED;
+		}
+	}
+
+	if (FPlatformMemory::SupportsFastVRAMMemory())
+	{
+		if (Usage & BUF_FastVRAM)
+		{
+			FFastVRAMAllocator::GetFastVRAMAllocator()->AllocUAVBuffer(Desc);
+		}
+	}
+
+	// If a resource array was provided for the resource, create the resource pre-populated
+	D3D11_SUBRESOURCE_DATA InitData;
+	D3D11_SUBRESOURCE_DATA* pInitData = NULL;
+	if(CreateInfo.ResourceArray)
+	{
+		check(Size == CreateInfo.ResourceArray->GetResourceDataSize());
+		InitData.pSysMem = CreateInfo.ResourceArray->GetResourceData();
+		InitData.SysMemPitch = Size;
+		InitData.SysMemSlicePitch = 0;
+		pInitData = &InitData;
+	}
+
+	TRefCountPtr<ID3D11Buffer> BufferResource;
+	VERIFYD3D11RESULT_EX(Direct3DDevice->CreateBuffer(&Desc, pInitData, BufferResource.GetInitReference()), Direct3DDevice);
+
+	if (CreateInfo.DebugName)
+	{
+		BufferResource->SetPrivateData(WKPDID_D3DDebugObjectName, FCString::Strlen(CreateInfo.DebugName) + 1, TCHAR_TO_ANSI(CreateInfo.DebugName));
+	}
+
+	UpdateBufferStats(BufferResource, true);
+
+	if (CreateInfo.DebugName)
+	{
+		BufferResource->SetPrivateData(WKPDID_D3DDebugObjectName, FCString::Strlen(CreateInfo.DebugName) + 1, TCHAR_TO_ANSI(CreateInfo.DebugName));
+	}
+
+	if (CreateInfo.ResourceArray)
+	{
+		// Discard the resource array's contents.
+		CreateInfo.ResourceArray->Discard();
+	}
+
+	return new FD3D11Buffer(BufferResource, Size, Usage, Stride);
+}
+
+FBufferRHIRef FD3D11DynamicRHI::CreateBuffer_RenderThread(class FRHICommandListImmediate& RHICmdList, uint32 Size, EBufferUsageFlags Usage, uint32 Stride, ERHIAccess ResourceState, FRHIResourceCreateInfo& CreateInfo)
+{
+	return RHICreateBuffer(Size, Usage, Stride, ResourceState, CreateInfo);
+}
+
+void* FD3D11DynamicRHI::LockBuffer_BottomOfPipe(FRHICommandListImmediate& RHICmdList, FRHIBuffer* BufferRHI, uint32 Offset, uint32 Size, EResourceLockMode LockMode)
+{
+	FD3D11Buffer* Buffer = ResourceCast(BufferRHI);
+	// If this resource is bound to the device, unbind it
+	ConditionalClearShaderResource(Buffer, true);
+
+	// Determine whether the buffer is dynamic or not.
+	D3D11_BUFFER_DESC Desc;
+	Buffer->Resource->GetDesc(&Desc);
+	const bool bIsDynamic = (Desc.Usage == D3D11_USAGE_DYNAMIC);
+
+	FD3D11LockedKey LockedKey(Buffer->Resource);
+	FD3D11LockedData LockedData;
+
+	if(bIsDynamic)
+	{
+		check(LockMode == RLM_WriteOnly || LockMode == RLM_WriteOnly_NoOverwrite);
+
+		// If the buffer is dynamic, map its memory for writing.
+		D3D11_MAPPED_SUBRESOURCE MappedSubresource;
+
+		D3D11_MAP MapType = (LockMode == RLM_WriteOnly)? D3D11_MAP_WRITE_DISCARD : D3D11_MAP_WRITE_NO_OVERWRITE;
+		VERIFYD3D11RESULT_EX(Direct3DDeviceIMContext->Map(Buffer->Resource, 0, MapType, 0, &MappedSubresource), Direct3DDevice);
+
+		LockedData.SetData(MappedSubresource.pData);
+		LockedData.Pitch = MappedSubresource.RowPitch;
+	}
+	else
+	{
+		if(LockMode == RLM_ReadOnly)
+		{
+			// If the static buffer is being locked for reading, create a staging buffer.
+			D3D11_BUFFER_DESC StagingBufferDesc;
+			ZeroMemory( &StagingBufferDesc, sizeof( D3D11_BUFFER_DESC ) );
+			StagingBufferDesc.ByteWidth = Size;
+			StagingBufferDesc.Usage = D3D11_USAGE_STAGING;
+			StagingBufferDesc.BindFlags = 0;
+			StagingBufferDesc.CPUAccessFlags = D3D11_CPU_ACCESS_READ;
+			StagingBufferDesc.MiscFlags = 0;
+			TRefCountPtr<ID3D11Buffer> StagingBuffer;
+			VERIFYD3D11RESULT_EX(Direct3DDevice->CreateBuffer(&StagingBufferDesc, NULL, StagingBuffer.GetInitReference()), Direct3DDevice);
+			LockedData.StagingResource = StagingBuffer;
+
+			// Copy the contents of the buffer to the staging buffer.
+			D3D11_BOX SourceBox;
+			SourceBox.left = Offset;
+			SourceBox.right = Offset + Size;
+			SourceBox.top = SourceBox.front = 0;
+			SourceBox.bottom = SourceBox.back = 1;
+			Direct3DDeviceIMContext->CopySubresourceRegion(StagingBuffer, 0, 0, 0, 0, Buffer->Resource, 0, &SourceBox);
+
+			// Map the staging buffer's memory for reading.
+			D3D11_MAPPED_SUBRESOURCE MappedSubresource;
+			VERIFYD3D11RESULT_EX(Direct3DDeviceIMContext->Map(StagingBuffer, 0, D3D11_MAP_READ, 0, &MappedSubresource), Direct3DDevice);
+			LockedData.SetData(MappedSubresource.pData);
+			LockedData.Pitch = MappedSubresource.RowPitch;
+			Offset = 0;
+		}
+		else
+		{
+			// If the static buffer is being locked for writing, allocate memory for the contents to be written to.
+			LockedData.AllocData(Desc.ByteWidth);
+			LockedData.Pitch = Desc.ByteWidth;
+		}
+	}
+	
+	// Add the lock to the lock map.
+	AddLockedData(LockedKey, LockedData);
+
+	// Return the offset pointer
+	return (void*)((uint8*)LockedData.GetData() + Offset);
+}
+
+void FD3D11DynamicRHI::UnlockBuffer_BottomOfPipe(FRHICommandListImmediate& RHICmdList, FRHIBuffer* BufferRHI)
+{
+	FD3D11Buffer* Buffer = ResourceCast(BufferRHI);
+
+	// Determine whether the buffer is dynamic or not.
+	D3D11_BUFFER_DESC Desc;
+	Buffer->Resource->GetDesc(&Desc);
+	const bool bIsDynamic = (Desc.Usage == D3D11_USAGE_DYNAMIC);
+
+	// Find the outstanding lock for this buffer and remove it from the tracker.
+	FD3D11LockedData LockedData;
+	verifyf(RemoveLockedData(FD3D11LockedKey(Buffer->Resource), LockedData), TEXT("Buffer is not locked"));
+
+	if(bIsDynamic)
+	{
+		// If the buffer is dynamic, its memory was mapped directly; unmap it.
+		Direct3DDeviceIMContext->Unmap(Buffer->Resource, 0);
+	}
+	else
+	{
+		// If the static buffer lock involved a staging resource, it was locked for reading.
+		if(LockedData.StagingResource)
+		{
+			// Unmap the staging buffer's memory.
+			ID3D11Buffer* StagingBuffer = (ID3D11Buffer*)LockedData.StagingResource.GetReference();
+			Direct3DDeviceIMContext->Unmap(StagingBuffer,0);
+		}
+		else 
+		{
+			// Copy the contents of the temporary memory buffer allocated for writing into the buffer.
+			Direct3DDeviceIMContext->UpdateSubresource(Buffer->Resource, 0, NULL, LockedData.GetData(), LockedData.Pitch, 0);
+
+			// Free the temporary memory buffer.
+			LockedData.FreeData();
+		}
+	}
+}
+
+void FD3D11DynamicRHI::RHICopyBuffer(FRHIBuffer* SourceBufferRHI, FRHIBuffer* DestBufferRHI)
+{
+	FD3D11Buffer* SourceBuffer = ResourceCast(SourceBufferRHI);
+	FD3D11Buffer* DestBuffer = ResourceCast(DestBufferRHI);
+
+	D3D11_BUFFER_DESC SourceBufferDesc;
+	SourceBuffer->Resource->GetDesc(&SourceBufferDesc);
+	
+	D3D11_BUFFER_DESC DestBufferDesc;
+	DestBuffer->Resource->GetDesc(&DestBufferDesc);
+
+	check(SourceBufferDesc.ByteWidth == DestBufferDesc.ByteWidth);
+
+	Direct3DDeviceIMContext->CopyResource(DestBuffer->Resource,SourceBuffer->Resource);
+
+	GPUProfilingData.RegisterGPUWork(1);
+}
+
+void FD3D11DynamicRHI::RHITransferBufferUnderlyingResource(FRHIBuffer* DestBuffer, FRHIBuffer* SrcBuffer)
+{
+	check(DestBuffer);
+	FD3D11Buffer* Dest = ResourceCast(DestBuffer);
+	if (!SrcBuffer)
+	{
+		Dest->ReleaseUnderlyingResource();
+	}
+	else
+	{
+		FD3D11Buffer* Src = ResourceCast(SrcBuffer);
+		Dest->Swap(*Src);
+	}
+}
diff --git a/Engine/Source/Runtime/Windows/D3D11RHI/Private/D3D11IndexBuffer.cpp b/Engine/Source/Runtime/Windows/D3D11RHI/Private/D3D11IndexBuffer.cpp
deleted file mode 100644
