VulkanRT
Loading and enabling extensions, function pointer wrangling and scaffolding for future work.

diff --git a/Engine/Source/Runtime/VulkanRHI/Private/Android/VulkanAndroidPlatform.h b/Engine/Source/Runtime/VulkanRHI/Private/Android/VulkanAndroidPlatform.h
--- a/Engine/Source/Runtime/VulkanRHI/Private/Android/VulkanAndroidPlatform.h
+++ b/Engine/Source/Runtime/VulkanRHI/Private/Android/VulkanAndroidPlatform.h
@@ -21,7 +21,7 @@
 #define VULKAN_SUPPORTS_PHYSICAL_DEVICE_PROPERTIES2	1
 #define VULKAN_SUPPORTS_ASTC_DECODE_MODE			(VULKAN_SUPPORTS_PHYSICAL_DEVICE_PROPERTIES2)
 #define VULKAN_SUPPORTS_NV_DIAGNOSTIC_CHECKPOINT	0
-#define VULKAN_SUPPORTS_NV_DIAGNOSTIC_CONFIG		0
+#define VULKAN_RHI_RAYTRACING						0
 
 // crashing during callback setup on Android, code will fallback to VK_EXT_debug_report instead
 #define VULKAN_SUPPORTS_DEBUG_UTILS					0
diff --git a/Engine/Source/Runtime/VulkanRHI/Private/Linux/VulkanLinuxPlatform.h b/Engine/Source/Runtime/VulkanRHI/Private/Linux/VulkanLinuxPlatform.h
--- a/Engine/Source/Runtime/VulkanRHI/Private/Linux/VulkanLinuxPlatform.h
+++ b/Engine/Source/Runtime/VulkanRHI/Private/Linux/VulkanLinuxPlatform.h
@@ -11,7 +11,7 @@
 #define VULKAN_SHOULD_ENABLE_DRAW_MARKERS			(UE_BUILD_DEBUG || UE_BUILD_DEVELOPMENT)
 #define VULKAN_SIGNAL_UNIMPLEMENTED()				checkf(false, TEXT("Unimplemented vulkan functionality: %s"), __PRETTY_FUNCTION__)
 #define VULKAN_SUPPORTS_AMD_BUFFER_MARKER			1
-
+#define VULKAN_RHI_RAYTRACING						0
 
 #define	UE_VK_API_VERSION							VK_API_VERSION_1_1
 
diff --git a/Engine/Source/Runtime/VulkanRHI/Private/Lumin/VulkanLuminPlatform.cpp b/Engine/Source/Runtime/VulkanRHI/Private/Lumin/VulkanLuminPlatform.cpp
--- a/Engine/Source/Runtime/VulkanRHI/Private/Lumin/VulkanLuminPlatform.cpp
+++ b/Engine/Source/Runtime/VulkanRHI/Private/Lumin/VulkanLuminPlatform.cpp
@@ -152,9 +152,9 @@ bool FVulkanLuminPlatform::ForceEnableDebugMarkers()
 	return GFoundTegraGfxDebugger;
 }
 
-void FVulkanLuminPlatform::EnablePhysicalDeviceFeatureExtensions(VkDeviceCreateInfo& DeviceInfo)
+void FVulkanLuminPlatform::EnablePhysicalDeviceFeatureExtensions(VkDeviceCreateInfo& DeviceInfo, FVulkanDevice& Device)
 {
-	SamplerConversion.pNext = nullptr;
+	SamplerConversion.pNext = const_cast<void*>(DeviceInfo.pNext);
 	SamplerConversion.sType = VK_STRUCTURE_TYPE_PHYSICAL_DEVICE_SAMPLER_YCBCR_CONVERSION_FEATURES;
 	SamplerConversion.samplerYcbcrConversion = VK_TRUE;
 	DeviceInfo.pNext = &SamplerConversion;
diff --git a/Engine/Source/Runtime/VulkanRHI/Private/Lumin/VulkanLuminPlatform.h b/Engine/Source/Runtime/VulkanRHI/Private/Lumin/VulkanLuminPlatform.h
--- a/Engine/Source/Runtime/VulkanRHI/Private/Lumin/VulkanLuminPlatform.h
+++ b/Engine/Source/Runtime/VulkanRHI/Private/Lumin/VulkanLuminPlatform.h
@@ -15,7 +15,7 @@
 #define VULKAN_SUPPORTS_DEDICATED_ALLOCATION			0
 #define VULKAN_SUPPORTS_GPU_CRASH_DUMPS					1
 #define VULKAN_SUPPORTS_ASTC_DECODE_MODE				1
-
+#define VULKAN_RHI_RAYTRACING						0
 
 #define ENUM_VK_ENTRYPOINTS_PLATFORM_BASE(EnumMacro)
 
@@ -63,7 +63,7 @@ public:
 
 	static bool HasUnifiedMemory() { return true; }
 
-	static void EnablePhysicalDeviceFeatureExtensions(VkDeviceCreateInfo& DeviceInfo);
+	static void EnablePhysicalDeviceFeatureExtensions(VkDeviceCreateInfo& DeviceInfo, FVulkanDevice& Device);
 
 	static bool RequiresWaitingForFrameCompletionEvent() { return false; }
 
diff --git a/Engine/Source/Runtime/VulkanRHI/Private/VulkanCommandWrappers.h b/Engine/Source/Runtime/VulkanRHI/Private/VulkanCommandWrappers.h
--- a/Engine/Source/Runtime/VulkanRHI/Private/VulkanCommandWrappers.h
+++ b/Engine/Source/Runtime/VulkanRHI/Private/VulkanCommandWrappers.h
@@ -185,6 +185,17 @@ struct FWrapLayer
 	static void CreateSamplerYcbcrConversionKHR(VkResult Result, VkDevice device, const VkSamplerYcbcrConversionCreateInfo* pCreateInfo, VkSamplerYcbcrConversion* YcbcrConversion) VULKAN_LAYER_BODY
 	static void DestroySamplerYcbcrConversionKHR(VkResult Result, VkDevice device, VkSamplerYcbcrConversion ycbcrConversion) VULKAN_LAYER_BODY
 #endif
+#if VULKAN_RHI_RAYTRACING
+	static void CreateAccelerationStructureKHR(VkResult Result, VkDevice Device, const VkAccelerationStructureCreateInfoKHR* CreateInfo, const VkAllocationCallbacks* Allocator, VkAccelerationStructureKHR* AccelerationStructure);
+	static void DestroyAccelerationStructureKHR(VkResult Result, VkDevice Device, VkAccelerationStructureKHR AccelerationStructure, const VkAllocationCallbacks* Allocator);
+	static void CmdBuildAccelerationStructuresKHR(VkResult Result, VkCommandBuffer CommandBuffer, uint32 InfoCount, const VkAccelerationStructureBuildGeometryInfoKHR* Infos, const VkAccelerationStructureBuildRangeInfoKHR* const* BuildRangeInfos);
+	static void GetAccelerationStructureBuildSizesKHR(VkResult Result, VkDevice Device, VkAccelerationStructureBuildTypeKHR BuildType, const VkAccelerationStructureBuildGeometryInfoKHR* BuildInfo, const uint32* MaxPrimitiveCounts, VkAccelerationStructureBuildSizesInfoKHR* SizeInfo);
+	static void GetAccelerationStructureDeviceAddressKHR(VkResult Result, VkDevice Device, const VkAccelerationStructureDeviceAddressInfoKHR* Info);
+	static void CmdTraceRaysKHR(VkResult Result, VkCommandBuffer CommandBuffer, const VkStridedDeviceAddressRegionKHR* RaygenShaderBindingTable, const VkStridedDeviceAddressRegionKHR* MissShaderBindingTable, const VkStridedDeviceAddressRegionKHR* HitShaderBindingTable, const VkStridedDeviceAddressRegionKHR* CallableShaderBindingTable, uint32 width, uint32 height, uint32 depth);
+	static void CreateRayTracingPipelinesKHR(VkResult Result, VkDevice Device, VkDeferredOperationKHR DeferredOperation, VkPipelineCache PipelineCache, uint32_t CreateInfoCount, const VkRayTracingPipelineCreateInfoKHR* CreateInfos, const VkAllocationCallbacks* Allocator, VkPipeline* Pipelines);
+	static void GetRayTracingShaderGroupHandlesKHR(VkResult Result, VkDevice Device, VkPipeline Pipeline, uint32_t FirstGroup, uint32_t GroupCount, size_t DataSize, void* Data);
+	static void GetBufferDeviceAddressKHR(VkResult Result, VkDevice Device, const VkBufferDeviceAddressInfo* Info);
+#endif
 };
 
 #undef VULKAN_LAYER_BODY
@@ -1372,6 +1383,75 @@ namespace VulkanRHI
 	}
 #endif
 
+#if VULKAN_RHI_RAYTRACING
+	static FORCEINLINE_DEBUGGABLE VkResult vkCreateAccelerationStructureKHR(VkDevice Device, const VkAccelerationStructureCreateInfoKHR* CreateInfo, const VkAllocationCallbacks* Allocator, VkAccelerationStructureKHR* AccelerationStructure)
+	{
+		FWrapLayer::CreateAccelerationStructureKHR(VK_RESULT_MAX_ENUM, Device, CreateInfo, Allocator, AccelerationStructure);
+		const VkResult Result = VULKANAPINAMESPACE::vkCreateAccelerationStructureKHR(Device, CreateInfo, Allocator, AccelerationStructure);
+		FWrapLayer::CreateAccelerationStructureKHR(Result, Device, CreateInfo, Allocator, AccelerationStructure);
+		return Result;
+	}
+
+	static FORCEINLINE_DEBUGGABLE void vkDestroyAccelerationStructureKHR(VkDevice Device, VkAccelerationStructureKHR AccelerationStructure, const VkAllocationCallbacks* Allocator)
+	{
+		FWrapLayer::DestroyAccelerationStructureKHR(VK_RESULT_MAX_ENUM, Device, AccelerationStructure, Allocator);
+		VULKANAPINAMESPACE::vkDestroyAccelerationStructureKHR(Device, AccelerationStructure, Allocator);
+		FWrapLayer::DestroyAccelerationStructureKHR(VK_SUCCESS, Device, AccelerationStructure, Allocator);
+	}
+
+	static FORCEINLINE_DEBUGGABLE void vkCmdBuildAccelerationStructuresKHR(VkCommandBuffer CommandBuffer, uint32 InfoCount, const VkAccelerationStructureBuildGeometryInfoKHR* Infos, const VkAccelerationStructureBuildRangeInfoKHR* const* BuildRangeInfos)
+	{
+		FWrapLayer::CmdBuildAccelerationStructuresKHR(VK_RESULT_MAX_ENUM, CommandBuffer, InfoCount, Infos, BuildRangeInfos);
+		VULKANAPINAMESPACE::vkCmdBuildAccelerationStructuresKHR(CommandBuffer, InfoCount, Infos, BuildRangeInfos);
+		FWrapLayer::CmdBuildAccelerationStructuresKHR(VK_SUCCESS, CommandBuffer, InfoCount, Infos, BuildRangeInfos);
+	}
+
+	static FORCEINLINE_DEBUGGABLE void vkGetAccelerationStructureBuildSizesKHR(VkDevice Device, VkAccelerationStructureBuildTypeKHR BuildType, const VkAccelerationStructureBuildGeometryInfoKHR* BuildInfo, const uint32* MaxPrimitiveCounts, VkAccelerationStructureBuildSizesInfoKHR* SizeInfo)
+	{
+		FWrapLayer::GetAccelerationStructureBuildSizesKHR(VK_RESULT_MAX_ENUM, Device, BuildType, BuildInfo, MaxPrimitiveCounts, SizeInfo);
+		VULKANAPINAMESPACE::vkGetAccelerationStructureBuildSizesKHR(Device, BuildType, BuildInfo, MaxPrimitiveCounts, SizeInfo);
+		FWrapLayer::GetAccelerationStructureBuildSizesKHR(VK_SUCCESS, Device, BuildType, BuildInfo, MaxPrimitiveCounts, SizeInfo);
+	}
+
+	static FORCEINLINE_DEBUGGABLE void vkGetAccelerationStructureDeviceAddressKHR(VkDevice Device, const VkAccelerationStructureDeviceAddressInfoKHR* Info)
+	{
+		FWrapLayer::GetAccelerationStructureDeviceAddressKHR(VK_RESULT_MAX_ENUM, Device, Info);
+		VULKANAPINAMESPACE::vkGetAccelerationStructureDeviceAddressKHR(Device, Info);
+		FWrapLayer::GetAccelerationStructureDeviceAddressKHR(VK_SUCCESS, Device, Info);
+	}
+
+	static FORCEINLINE_DEBUGGABLE void vkCmdTraceRaysKHR(VkCommandBuffer CommandBuffer, const VkStridedDeviceAddressRegionKHR* RaygenShaderBindingTable, const VkStridedDeviceAddressRegionKHR* MissShaderBindingTable, const VkStridedDeviceAddressRegionKHR* HitShaderBindingTable, const VkStridedDeviceAddressRegionKHR* CallableShaderBindingTable, uint32 width, uint32 height, uint32 depth)
+	{
+		FWrapLayer::CmdTraceRaysKHR(VK_RESULT_MAX_ENUM, CommandBuffer, RaygenShaderBindingTable, MissShaderBindingTable, HitShaderBindingTable, CallableShaderBindingTable, width, height, depth);
+		VULKANAPINAMESPACE::vkCmdTraceRaysKHR(CommandBuffer, RaygenShaderBindingTable, MissShaderBindingTable, HitShaderBindingTable, CallableShaderBindingTable, width, height, depth);
+		FWrapLayer::CmdTraceRaysKHR(VK_SUCCESS, CommandBuffer, RaygenShaderBindingTable, MissShaderBindingTable, HitShaderBindingTable, CallableShaderBindingTable, width, height, depth);
+	}
+
+	static FORCEINLINE_DEBUGGABLE VkResult vkCreateRayTracingPipelinesKHR(VkDevice Device, VkDeferredOperationKHR DeferredOperation, VkPipelineCache PipelineCache, uint32_t CreateInfoCount, const VkRayTracingPipelineCreateInfoKHR* CreateInfos, const VkAllocationCallbacks* Allocator, VkPipeline* Pipelines)
+	{
+		FWrapLayer::CreateRayTracingPipelinesKHR(VK_RESULT_MAX_ENUM, Device, DeferredOperation, PipelineCache, CreateInfoCount, CreateInfos, Allocator, Pipelines);
+		const VkResult Result = VULKANAPINAMESPACE::vkCreateRayTracingPipelinesKHR(Device, DeferredOperation, PipelineCache, CreateInfoCount, CreateInfos, Allocator, Pipelines);
+		FWrapLayer::CreateRayTracingPipelinesKHR(Result, Device, DeferredOperation, PipelineCache, CreateInfoCount, CreateInfos, Allocator, Pipelines);
+		return Result;
+	}
+
+	static FORCEINLINE_DEBUGGABLE VkResult vkGetRayTracingShaderGroupHandlesKHR(VkDevice Device, VkPipeline Pipeline, uint32_t FirstGroup, uint32_t GroupCount, size_t DataSize, void* Data)
+	{
+		FWrapLayer::GetRayTracingShaderGroupHandlesKHR(VK_RESULT_MAX_ENUM, Device, Pipeline, FirstGroup, GroupCount, DataSize, Data);
+		const VkResult Result = VULKANAPINAMESPACE::vkGetRayTracingShaderGroupHandlesKHR(Device, Pipeline, FirstGroup, GroupCount, DataSize, Data);
+		FWrapLayer::GetRayTracingShaderGroupHandlesKHR(Result, Device, Pipeline, FirstGroup, GroupCount, DataSize, Data);
+		return Result;
+	}
+
+	static FORCEINLINE_DEBUGGABLE VkDeviceAddress vkGetBufferDeviceAddressKHR(VkDevice Device, const VkBufferDeviceAddressInfo* Info)
+	{
+		FWrapLayer::GetBufferDeviceAddressKHR(VK_RESULT_MAX_ENUM, Device, Info);
+		const VkDeviceAddress Result = VULKANAPINAMESPACE::vkGetBufferDeviceAddressKHR(Device, Info);
+		FWrapLayer::GetBufferDeviceAddressKHR(VK_SUCCESS, Device, Info);
+		return Result;
+	}
+#endif
+
 #if VULKAN_ENABLE_IMAGE_TRACKING_LAYER
 	void BindDebugLabelName(VkImage Image, const TCHAR* Name);
 #endif
diff --git a/Engine/Source/Runtime/VulkanRHI/Private/VulkanDebug.cpp b/Engine/Source/Runtime/VulkanRHI/Private/VulkanDebug.cpp
--- a/Engine/Source/Runtime/VulkanRHI/Private/VulkanDebug.cpp
+++ b/Engine/Source/Runtime/VulkanRHI/Private/VulkanDebug.cpp
@@ -4162,6 +4162,152 @@ void FWrapLayer::DestroySamplerYcbcrConversionKHR(VkResult Result, VkDevice Devi
 }
 #endif
 
+#if VULKAN_RHI_RAYTRACING
+void FWrapLayer::CreateAccelerationStructureKHR(VkResult Result, VkDevice Device, const VkAccelerationStructureCreateInfoKHR* CreateInfo, const VkAllocationCallbacks* Allocator, VkAccelerationStructureKHR* AccelerationStructure)
+{
+	if (Result == VK_RESULT_MAX_ENUM)
+	{
+#if VULKAN_ENABLE_DUMP_LAYER
+	PrintfBeginResult(FString::Printf(TEXT("vkCreateAccelerationStructureKHR(Device=0x%p, CreateInfo=0x%p, AccelerationStructure=0x%p)[...]"), Device, CreateInfo, AccelerationStructure));
+#endif
+	}
+	else
+	{
+#if VULKAN_ENABLE_DUMP_LAYER
+	PrintResult(Result);
+#endif
+	}
+}
+
+void FWrapLayer::DestroyAccelerationStructureKHR(VkResult Result, VkDevice Device, VkAccelerationStructureKHR AccelerationStructure, const VkAllocationCallbacks* Allocator)
+{
+	if (Result == VK_RESULT_MAX_ENUM)
+	{
+#if VULKAN_ENABLE_DUMP_LAYER
+		PrintfBeginResult(FString::Printf(TEXT("vkDestroyAccelerationStructureKHR(Device=0x%p, AccelerationStructure=0x%p)[...]"), Device, AccelerationStructure));
+#endif
+	}
+	else
+	{
+#if VULKAN_ENABLE_DUMP_LAYER
+		PrintResult(Result);
+#endif
+	}
+}
+
+void FWrapLayer::CmdBuildAccelerationStructuresKHR(VkResult Result, VkCommandBuffer CommandBuffer, uint32 InfoCount, const VkAccelerationStructureBuildGeometryInfoKHR* Infos, const VkAccelerationStructureBuildRangeInfoKHR* const* BuildRangeInfos)
+{
+	if (Result == VK_RESULT_MAX_ENUM)
+	{
+#if VULKAN_ENABLE_DUMP_LAYER
+		PrintfBeginResult(FString::Printf(TEXT("vkCmdBuildAccelerationStructuresKHR")));
+#endif
+	}
+	else
+	{
+#if VULKAN_ENABLE_DUMP_LAYER
+		PrintResult(Result);
+#endif
+	}
+}
+
+void FWrapLayer::GetAccelerationStructureBuildSizesKHR(VkResult Result, VkDevice Device, VkAccelerationStructureBuildTypeKHR BuildType, const VkAccelerationStructureBuildGeometryInfoKHR* BuildInfo, const uint32* MaxPrimitiveCounts, VkAccelerationStructureBuildSizesInfoKHR* SizeInfo)
+{
+	if (Result == VK_RESULT_MAX_ENUM)
+	{
+#if VULKAN_ENABLE_DUMP_LAYER
+		PrintfBeginResult(FString::Printf(TEXT("vkGetAccelerationStructureBuildSizesKHR")));
+#endif
+	}
+	else
+	{
+#if VULKAN_ENABLE_DUMP_LAYER
+		PrintResult(Result);
+#endif
+	}
+}
+
+void FWrapLayer::GetAccelerationStructureDeviceAddressKHR(VkResult Result, VkDevice Device, const VkAccelerationStructureDeviceAddressInfoKHR* Info)
+{
+	if (Result == VK_RESULT_MAX_ENUM)
+	{
+#if VULKAN_ENABLE_DUMP_LAYER
+		PrintfBeginResult(FString::Printf(TEXT("vkGetAccelerationStructureDeviceAddressKHR")));
+#endif
+	}
+	else
+	{
+#if VULKAN_ENABLE_DUMP_LAYER
+		PrintResult(Result);
+#endif
+	}
+}
+
+void FWrapLayer::CmdTraceRaysKHR(VkResult Result, VkCommandBuffer CommandBuffer, const VkStridedDeviceAddressRegionKHR* RaygenShaderBindingTable, const VkStridedDeviceAddressRegionKHR* MissShaderBindingTable, const VkStridedDeviceAddressRegionKHR* HitShaderBindingTable, const VkStridedDeviceAddressRegionKHR* CallableShaderBindingTable, uint32 width, uint32 height, uint32 depth)
+{
+	if (Result == VK_RESULT_MAX_ENUM)
+	{
+#if VULKAN_ENABLE_DUMP_LAYER
+		PrintfBeginResult(FString::Printf(TEXT("vkCmdTraceRaysKHR")));
+#endif
+	}
+	else
+	{
+#if VULKAN_ENABLE_DUMP_LAYER
+		PrintResult(Result);
+#endif
+	}
+}
+
+void FWrapLayer::CreateRayTracingPipelinesKHR(VkResult Result, VkDevice Device, VkDeferredOperationKHR DeferredOperation, VkPipelineCache PipelineCache, uint32_t CreateInfoCount, const VkRayTracingPipelineCreateInfoKHR* CreateInfos, const VkAllocationCallbacks* Allocator, VkPipeline* Pipelines)
+{
+	if (Result == VK_RESULT_MAX_ENUM)
+	{
+#if VULKAN_ENABLE_DUMP_LAYER
+		PrintfBeginResult(FString::Printf(TEXT("vkCreateRayTracingPipelinesKHR")));
+#endif
+	}
+	else
+	{
+#if VULKAN_ENABLE_DUMP_LAYER
+		PrintResult(Result);
+#endif
+	}
+}
+
+void FWrapLayer::GetRayTracingShaderGroupHandlesKHR(VkResult Result, VkDevice Device, VkPipeline Pipeline, uint32_t FirstGroup, uint32_t GroupCount, size_t DataSize, void* Data)
+{
+	if (Result == VK_RESULT_MAX_ENUM)
+	{
+#if VULKAN_ENABLE_DUMP_LAYER
+		PrintfBeginResult(FString::Printf(TEXT("vkGetRayTracingShaderGroupHandlesKHR")));
+#endif
+	}
+	else
+	{
+#if VULKAN_ENABLE_DUMP_LAYER
+		PrintResult(Result);
+#endif
+	}
+}
+
+void FWrapLayer::GetBufferDeviceAddressKHR(VkResult Result, VkDevice Device, const VkBufferDeviceAddressInfo* Info)
+{
+	if (Result == VK_RESULT_MAX_ENUM)
+	{
+#if VULKAN_ENABLE_DUMP_LAYER
+		PrintfBeginResult(FString::Printf(TEXT("vkGetBufferDeviceAddressKHR")));
+#endif
+	}
+	else
+	{
+#if VULKAN_ENABLE_DUMP_LAYER
+		PrintResult(Result);
+#endif
+	}
+}
+#endif // VULKAN_RHI_RAYTRACING
+
 #if VULKAN_ENABLE_IMAGE_TRACKING_LAYER
 namespace VulkanRHI
 {
diff --git a/Engine/Source/Runtime/VulkanRHI/Private/VulkanDevice.cpp b/Engine/Source/Runtime/VulkanRHI/Private/VulkanDevice.cpp
--- a/Engine/Source/Runtime/VulkanRHI/Private/VulkanDevice.cpp
+++ b/Engine/Source/Runtime/VulkanRHI/Private/VulkanDevice.cpp
@@ -318,7 +318,7 @@ void FVulkanDevice::CreateDevice()
 	FVulkanPlatform::RestrictEnabledPhysicalDeviceFeatures(PhysicalFeatures);
 	DeviceInfo.pEnabledFeatures = &PhysicalFeatures;
 
-	FVulkanPlatform::EnablePhysicalDeviceFeatureExtensions(DeviceInfo);
+	FVulkanPlatform::EnablePhysicalDeviceFeatureExtensions(DeviceInfo, *this);
 
 #if VULKAN_SUPPORTS_NV_DIAGNOSTICS
 	VkPhysicalDeviceDiagnosticsConfigFeaturesNV DeviceDiagnosticsNV;
@@ -355,24 +355,22 @@ void FVulkanDevice::CreateDevice()
 #endif
 
 #if VULKAN_SUPPORTS_SEPARATE_DEPTH_STENCIL_LAYOUTS
-	VkPhysicalDeviceSeparateDepthStencilLayoutsFeaturesKHR SeparateDepthStencilLayoutsFeatures;
 	if (bHasSeparateDepthStencilLayouts)
 	{
-		ZeroVulkanStruct(SeparateDepthStencilLayoutsFeatures, VK_STRUCTURE_TYPE_PHYSICAL_DEVICE_SEPARATE_DEPTH_STENCIL_LAYOUTS_FEATURES_KHR);
-		SeparateDepthStencilLayoutsFeatures.separateDepthStencilLayouts = VK_TRUE;
-		SeparateDepthStencilLayoutsFeatures.pNext = const_cast<void*>(DeviceInfo.pNext);
-		DeviceInfo.pNext = &SeparateDepthStencilLayoutsFeatures;
+		ZeroVulkanStruct(OptionalFeatures.SeparateDepthStencilLayoutsFeatures, VK_STRUCTURE_TYPE_PHYSICAL_DEVICE_SEPARATE_DEPTH_STENCIL_LAYOUTS_FEATURES_KHR);
+		OptionalFeatures.SeparateDepthStencilLayoutsFeatures.separateDepthStencilLayouts = VK_TRUE;
+		OptionalFeatures.SeparateDepthStencilLayoutsFeatures.pNext = const_cast<void*>(DeviceInfo.pNext);
+		DeviceInfo.pNext = &OptionalFeatures.SeparateDepthStencilLayoutsFeatures;
 	}
 #endif
 
 #if VULKAN_SUPPORTS_SCALAR_BLOCK_LAYOUT
-	VkPhysicalDeviceScalarBlockLayoutFeaturesEXT ScalarBlockLayoutFeatures;
 	if (OptionalDeviceExtensions.HasScalarBlockLayoutFeatures)
 	{
-		ZeroVulkanStruct(ScalarBlockLayoutFeatures, VK_STRUCTURE_TYPE_PHYSICAL_DEVICE_SCALAR_BLOCK_LAYOUT_FEATURES_EXT);
-		ScalarBlockLayoutFeatures.scalarBlockLayout = VK_TRUE;
-		ScalarBlockLayoutFeatures.pNext = (void*)DeviceInfo.pNext;
-		DeviceInfo.pNext = &ScalarBlockLayoutFeatures;
+		ZeroVulkanStruct(OptionalFeatures.ScalarBlockLayoutFeatures, VK_STRUCTURE_TYPE_PHYSICAL_DEVICE_SCALAR_BLOCK_LAYOUT_FEATURES_EXT);
+		OptionalFeatures.ScalarBlockLayoutFeatures.scalarBlockLayout = VK_TRUE;
+		OptionalFeatures.ScalarBlockLayoutFeatures.pNext = (void*)DeviceInfo.pNext;
+		DeviceInfo.pNext = &OptionalFeatures.ScalarBlockLayoutFeatures;
 	}
 #endif
 
@@ -848,14 +846,22 @@ bool FVulkanDevice::QueryGPU(int32 DeviceIndex)
 		GpuProps2.pNext = &GpuIdProps;
 		ZeroVulkanStruct(GpuIdProps, VK_STRUCTURE_TYPE_PHYSICAL_DEVICE_ID_PROPERTIES_KHR);
 
-		void** NextPropsAddr = nullptr;
-		NextPropsAddr = &GpuIdProps.pNext;
-
 #if VULKAN_SUPPORTS_DRIVER_PROPERTIES
 		if (GetOptionalExtensions().HasDriverProperties)
 		{
-			*NextPropsAddr = &PhysicalDeviceProperties;
-			NextPropsAddr = &PhysicalDeviceProperties.pNext;
+			PhysicalDeviceProperties.pNext = GpuProps2.pNext;
+			GpuProps2.pNext = &PhysicalDeviceProperties;
+		}
+#endif
+
+#if VULKAN_RHI_RAYTRACING
+		if (OptionalDeviceExtensions.HasRaytracingExtensions())
+		{
+			ZeroVulkanStruct(RayTracingProperties.AccelerationStructure, VK_STRUCTURE_TYPE_PHYSICAL_DEVICE_ACCELERATION_STRUCTURE_PROPERTIES_KHR);
+			ZeroVulkanStruct(RayTracingProperties.RayTracingPipeline, VK_STRUCTURE_TYPE_PHYSICAL_DEVICE_RAY_TRACING_PIPELINE_PROPERTIES_KHR);
+			RayTracingProperties.RayTracingPipeline.pNext = GpuProps2.pNext;
+			RayTracingProperties.AccelerationStructure.pNext = &RayTracingProperties.RayTracingPipeline;
+			GpuProps2.pNext = &RayTracingProperties.AccelerationStructure;
 		}
 #endif
 
@@ -873,7 +879,7 @@ bool FVulkanDevice::QueryGPU(int32 DeviceIndex)
 #endif
 
 	}
-#endif
+#endif // VULKAN_SUPPORTS_PHYSICAL_DEVICE_PROPERTIES2
 
 	VulkanRHI::vkGetPhysicalDeviceProperties(Gpu, &GpuProps);
 	bool bDiscrete = false;
diff --git a/Engine/Source/Runtime/VulkanRHI/Private/VulkanDevice.h b/Engine/Source/Runtime/VulkanRHI/Private/VulkanDevice.h
--- a/Engine/Source/Runtime/VulkanRHI/Private/VulkanDevice.h
+++ b/Engine/Source/Runtime/VulkanRHI/Private/VulkanDevice.h
@@ -52,6 +52,13 @@ struct FOptionalVulkanDeviceExtensions
 			uint32 HasAtomicInt64 : 1;
 			uint32 HasBufferAtomicInt64 : 1;
 			uint32 HasScalarBlockLayoutFeatures : 1;
+			uint32 HasAccelerationStructure : 1;
+			uint32 HasRayTracingPipeline : 1;
+			uint32 HasDescriptorIndexing : 1;
+			uint32 HasBufferDeviceAddress : 1;
+			uint32 HasDeferredHostOperations : 1;
+			uint32 HasSPIRV_14 : 1;
+			uint32 HasShaderFloatControls : 1;
 		};
 		uint32 Packed;
 	};
@@ -68,7 +75,46 @@ struct FOptionalVulkanDeviceExtensions
 	{
 		return HasAMDBufferMarker || HasNVDiagnosticCheckpoints;
 	}
+
+#if VULKAN_RHI_RAYTRACING
+	inline bool HasRaytracingExtensions() const
+	{
+		return 
+			HasAccelerationStructure && 
+			HasRayTracingPipeline && 
+			HasDescriptorIndexing && 
+			HasBufferDeviceAddress && 
+			HasDeferredHostOperations && 
+			HasSPIRV_14 && 
+			HasShaderFloatControls;
+	}
+#endif
+};
+
+struct FOptionalVulkanDeviceFeatures
+{
+#if VULKAN_SUPPORTS_SEPARATE_DEPTH_STENCIL_LAYOUTS
+	VkPhysicalDeviceSeparateDepthStencilLayoutsFeaturesKHR SeparateDepthStencilLayoutsFeatures;
+#endif
+#if VULKAN_SUPPORTS_SCALAR_BLOCK_LAYOUT
+	VkPhysicalDeviceScalarBlockLayoutFeaturesEXT ScalarBlockLayoutFeatures;
+#endif
+#if VULKAN_RHI_RAYTRACING
+	VkPhysicalDeviceBufferDeviceAddressFeaturesKHR BufferDeviceAddressFeatures;
+	VkPhysicalDeviceAccelerationStructureFeaturesKHR AccelerationStructureFeatures;
+	VkPhysicalDeviceRayTracingPipelineFeaturesKHR RayTracingPipelineFeatures;
+	VkPhysicalDeviceDescriptorIndexingFeaturesEXT DescriptorIndexingFeatures;
+#endif
+};
+
+#if VULKAN_RHI_RAYTRACING
+struct FRayTracingProperties
+{
+	VkPhysicalDeviceAccelerationStructurePropertiesKHR AccelerationStructure;
+	VkPhysicalDeviceRayTracingPipelinePropertiesKHR RayTracingPipeline;
 };
+#endif // VULKAN_RHI_RAYTRACING
+
 namespace VulkanRHI
 {
 	class FDeferredDeletionQueue2 : public FDeviceChild
@@ -223,7 +269,15 @@ public:
 		check(RHI->GetOptionalExtensions().HasKHRGetPhysicalDeviceProperties2);
 		return GpuIdProps;
 	}
-#endif
+
+#if VULKAN_RHI_RAYTRACING
+	inline const FRayTracingProperties& GetRayTracingProperties() const
+	{
+		check(OptionalDeviceExtensions.HasRaytracingExtensions());
+		return RayTracingProperties;
+	}
+#endif // VULKAN_RHI_RAYTRACING
+#endif // VULKAN_SUPPORTS_PHYSICAL_DEVICE_PROPERTIES2
 
 #if VULKAN_SUPPORTS_VALIDATION_CACHE
 	inline VkValidationCacheEXT GetValidationCache() const
@@ -394,6 +448,11 @@ public:
 		return OptionalDeviceExtensions;
 	}
 
+	inline FOptionalVulkanDeviceFeatures& GetOptionalFeatures()
+	{
+		return OptionalFeatures;
+	}
+
 #if VULKAN_SUPPORTS_GPU_CRASH_DUMPS
 	VkBuffer GetCrashMarkerBuffer() const
 	{
@@ -455,9 +514,15 @@ private:
 	VkPhysicalDeviceProperties GpuProps;
 #if VULKAN_SUPPORTS_PHYSICAL_DEVICE_PROPERTIES2
 	VkPhysicalDeviceIDPropertiesKHR GpuIdProps;
-#endif
+
+#if VULKAN_RHI_RAYTRACING
+	FRayTracingProperties RayTracingProperties;
+#endif // VULKAN_RHI_RAYTRACING
+#endif // VULKAN_SUPPORTS_PHYSICAL_DEVICE_PROPERTIES2
 
 	VkPhysicalDeviceFeatures PhysicalFeatures;
+	FOptionalVulkanDeviceFeatures OptionalFeatures;
+
 	bool bHasSeparateDepthStencilLayouts = false;
 
 	TArray<VkQueueFamilyProperties> QueueFamilyProps;
diff --git a/Engine/Source/Runtime/VulkanRHI/Private/VulkanGenericPlatform.h b/Engine/Source/Runtime/VulkanRHI/Private/VulkanGenericPlatform.h
--- a/Engine/Source/Runtime/VulkanRHI/Private/VulkanGenericPlatform.h
+++ b/Engine/Source/Runtime/VulkanRHI/Private/VulkanGenericPlatform.h
@@ -99,7 +99,7 @@ public:
 	static bool SupportsVolumeTextureRendering() { return true; }
 
 	// Allow platforms to add extension features to the DeviceInfo pNext chain
-	static void EnablePhysicalDeviceFeatureExtensions(VkDeviceCreateInfo& DeviceInfo) {}
+	static void EnablePhysicalDeviceFeatureExtensions(VkDeviceCreateInfo& DeviceInfo, FVulkanDevice& Device) {}
 
 	static bool RequiresSwapchainGeneralInitialLayout() { return false; }
 
diff --git a/Engine/Source/Runtime/VulkanRHI/Private/VulkanLayers.cpp b/Engine/Source/Runtime/VulkanRHI/Private/VulkanLayers.cpp
--- a/Engine/Source/Runtime/VulkanRHI/Private/VulkanLayers.cpp
+++ b/Engine/Source/Runtime/VulkanRHI/Private/VulkanLayers.cpp
@@ -849,6 +849,16 @@ void FOptionalVulkanDeviceExtensions::Setup(const TArray<const ANSICHAR*>& Devic
 #if VULKAN_SUPPORTS_SCALAR_BLOCK_LAYOUT
 	HasScalarBlockLayoutFeatures = HasExtension(DeviceExtensions, VK_EXT_SCALAR_BLOCK_LAYOUT_EXTENSION_NAME);
 #endif
+
+#if VULKAN_RHI_RAYTRACING
+	HasAccelerationStructure	= HasExtension(DeviceExtensions, VK_KHR_ACCELERATION_STRUCTURE_EXTENSION_NAME);
+	HasRayTracingPipeline		= HasExtension(DeviceExtensions, VK_KHR_RAY_TRACING_PIPELINE_EXTENSION_NAME);
+	HasDescriptorIndexing		= HasExtension(DeviceExtensions, VK_EXT_DESCRIPTOR_INDEXING_EXTENSION_NAME);
+	HasBufferDeviceAddress		= HasExtension(DeviceExtensions, VK_KHR_BUFFER_DEVICE_ADDRESS_EXTENSION_NAME);
+	HasDeferredHostOperations	= HasExtension(DeviceExtensions, VK_KHR_DEFERRED_HOST_OPERATIONS_EXTENSION_NAME);
+	HasSPIRV_14					= HasExtension(DeviceExtensions, VK_KHR_SPIRV_1_4_EXTENSION_NAME);
+	HasShaderFloatControls		= HasExtension(DeviceExtensions, VK_KHR_SHADER_FLOAT_CONTROLS_EXTENSION_NAME);
+#endif
 }
 
 void FVulkanDynamicRHI::SetupValidationRequests()
diff --git a/Engine/Source/Runtime/VulkanRHI/Private/VulkanRHIPrivate.h b/Engine/Source/Runtime/VulkanRHI/Private/VulkanRHIPrivate.h
--- a/Engine/Source/Runtime/VulkanRHI/Private/VulkanRHIPrivate.h
+++ b/Engine/Source/Runtime/VulkanRHI/Private/VulkanRHIPrivate.h
@@ -19,6 +19,10 @@
 // the configuration will set up anything not set up by the platform
 #include "VulkanConfiguration.h"
 
+#if VULKAN_RHI_RAYTRACING
+#include "VulkanRayTracing.h"
+#endif
+
 #if VULKAN_COMMANDWRAPPERS_ENABLE
 	#if VULKAN_DYNAMICALLYLOADED
 		// Vulkan API is defined in VulkanDynamicAPI namespace.
diff --git a/Engine/Source/Runtime/VulkanRHI/Private/VulkanRayTracing.cpp b/Engine/Source/Runtime/VulkanRHI/Private/VulkanRayTracing.cpp
new file mode 100644
--- /dev/null
+++ b/Engine/Source/Runtime/VulkanRHI/Private/VulkanRayTracing.cpp
@@ -0,0 +1,73 @@
+// Copyright Epic Games, Inc. All Rights Reserved.
+
+#include "VulkanRHIPrivate.h"
+
+#if VULKAN_RHI_RAYTRACING
+
+// Define ray tracing entry points
+#define DEFINE_VK_ENTRYPOINTS(Type,Func) VULKANRHI_API Type VulkanDynamicAPI::Func = NULL;
+ENUM_VK_ENTRYPOINTS_RAYTRACING(DEFINE_VK_ENTRYPOINTS)
+
+void FVulkanRayTracingPlatform::GetDeviceExtensions(EGpuVendorId VendorId, TArray<const ANSICHAR*>& OutExtensions)
+{
+	// Primary extensions
+	OutExtensions.Add(VK_KHR_ACCELERATION_STRUCTURE_EXTENSION_NAME);
+	OutExtensions.Add(VK_KHR_RAY_TRACING_PIPELINE_EXTENSION_NAME);
+
+	// VK_KHR_acceleration_structure dependencies 
+	OutExtensions.Add(VK_EXT_DESCRIPTOR_INDEXING_EXTENSION_NAME); // Note: Promoted to Vulkan 1.2
+	OutExtensions.Add(VK_KHR_BUFFER_DEVICE_ADDRESS_EXTENSION_NAME); // Note: Promoted to Vulkan 1.2
+	OutExtensions.Add(VK_KHR_DEFERRED_HOST_OPERATIONS_EXTENSION_NAME);
+	
+	// VK_KHR_ray_tracing_pipeline dependency
+	OutExtensions.Add(VK_KHR_SPIRV_1_4_EXTENSION_NAME); // Note: Promoted to Vulkan 1.2
+
+	// VK_KHR_spirv_1_4 dependency
+	OutExtensions.Add(VK_KHR_SHADER_FLOAT_CONTROLS_EXTENSION_NAME); // Note: Promoted to Vulkan 1.2
+}
+
+void FVulkanRayTracingPlatform::EnablePhysicalDeviceFeatureExtensions(VkDeviceCreateInfo& DeviceInfo, FVulkanDevice& Device)
+{
+	if (Device.GetOptionalExtensions().HasRaytracingExtensions())
+	{
+		FOptionalVulkanDeviceFeatures& Features = Device.GetOptionalFeatures();
+
+		ZeroVulkanStruct(Features.BufferDeviceAddressFeatures, VK_STRUCTURE_TYPE_PHYSICAL_DEVICE_BUFFER_DEVICE_ADDRESS_FEATURES);
+		Features.BufferDeviceAddressFeatures.bufferDeviceAddress = VK_TRUE;
+		Features.BufferDeviceAddressFeatures.pNext = const_cast<void*>(DeviceInfo.pNext);
+
+		ZeroVulkanStruct(Features.AccelerationStructureFeatures, VK_STRUCTURE_TYPE_PHYSICAL_DEVICE_ACCELERATION_STRUCTURE_FEATURES_KHR);;
+		Features.AccelerationStructureFeatures.accelerationStructure = VK_TRUE;
+		Features.AccelerationStructureFeatures.descriptorBindingAccelerationStructureUpdateAfterBind = VK_TRUE;
+		Features.AccelerationStructureFeatures.pNext = &Features.BufferDeviceAddressFeatures;
+		
+		ZeroVulkanStruct(Features.RayTracingPipelineFeatures, VK_STRUCTURE_TYPE_PHYSICAL_DEVICE_RAY_TRACING_PIPELINE_FEATURES_KHR);
+		Features.RayTracingPipelineFeatures.rayTracingPipeline = VK_TRUE;
+		Features.RayTracingPipelineFeatures.pNext = &Features.AccelerationStructureFeatures;
+
+		ZeroVulkanStruct(Features.DescriptorIndexingFeatures, VK_STRUCTURE_TYPE_PHYSICAL_DEVICE_DESCRIPTOR_INDEXING_FEATURES_EXT);
+		Features.DescriptorIndexingFeatures.pNext = &Features.RayTracingPipelineFeatures;
+
+		DeviceInfo.pNext = &Features.DescriptorIndexingFeatures;
+	}
+}
+
+#if PLATFORM_WINDOWS
+#pragma warning(push)
+#pragma warning(disable : 4191) // warning C4191: 'type cast': unsafe conversion
+#endif // PLATFORM_WINDOWS
+bool FVulkanRayTracingPlatform::LoadVulkanInstanceFunctions(VkInstance inInstance)
+{
+	bool bFoundAllEntryPoints = true;
+#define CHECK_VK_ENTRYPOINTS(Type,Func) if (VulkanDynamicAPI::Func == NULL) { bFoundAllEntryPoints = false; UE_LOG(LogRHI, Warning, TEXT("Failed to find entry point for %s"), TEXT(#Func)); }
+#define GETINSTANCE_VK_ENTRYPOINTS(Type, Func) VulkanDynamicAPI::Func = (Type)VulkanDynamicAPI::vkGetInstanceProcAddr(inInstance, #Func);
+	ENUM_VK_ENTRYPOINTS_RAYTRACING(GETINSTANCE_VK_ENTRYPOINTS);
+#if UE_BUILD_DEBUG
+	ENUM_VK_ENTRYPOINTS_RAYTRACING(CHECK_VK_ENTRYPOINTS);
+#endif
+	return bFoundAllEntryPoints;
+}
+#if PLATFORM_WINDOWS
+#pragma warning(pop) // restore 4191
+#endif
+#endif // #if VULKAN_RHI_RAYTRACING
diff --git a/Engine/Source/Runtime/VulkanRHI/Private/VulkanRayTracing.h b/Engine/Source/Runtime/VulkanRHI/Private/VulkanRayTracing.h
new file mode 100644
--- /dev/null
+++ b/Engine/Source/Runtime/VulkanRHI/Private/VulkanRayTracing.h
@@ -0,0 +1,32 @@
+// Copyright Epic Games, Inc. All Rights Reserved..
+
+#pragma once
+
+#if VULKAN_RHI_RAYTRACING
+
+#define ENUM_VK_ENTRYPOINTS_RAYTRACING(EnumMacro) \
+	EnumMacro(PFN_vkCreateAccelerationStructureKHR, vkCreateAccelerationStructureKHR) \
+	EnumMacro(PFN_vkDestroyAccelerationStructureKHR, vkDestroyAccelerationStructureKHR) \
+	EnumMacro(PFN_vkCmdBuildAccelerationStructuresKHR, vkCmdBuildAccelerationStructuresKHR) \
+	EnumMacro(PFN_vkGetAccelerationStructureBuildSizesKHR, vkGetAccelerationStructureBuildSizesKHR) \
+	EnumMacro(PFN_vkGetAccelerationStructureDeviceAddressKHR, vkGetAccelerationStructureDeviceAddressKHR) \
+	EnumMacro(PFN_vkCmdTraceRaysKHR, vkCmdTraceRaysKHR) \
+	EnumMacro(PFN_vkCreateRayTracingPipelinesKHR, vkCreateRayTracingPipelinesKHR) \
+	EnumMacro(PFN_vkGetRayTracingShaderGroupHandlesKHR, vkGetRayTracingShaderGroupHandlesKHR) \
+	EnumMacro(PFN_vkGetBufferDeviceAddressKHR, vkGetBufferDeviceAddressKHR)
+
+// Declare ray tracing entry points
+namespace VulkanDynamicAPI
+{
+	ENUM_VK_ENTRYPOINTS_RAYTRACING(DECLARE_VK_ENTRYPOINTS);
+}
+
+class FVulkanRayTracingPlatform
+{
+public:
+	static void GetDeviceExtensions(EGpuVendorId VendorId, TArray<const ANSICHAR*>& OutExtensions);
+	static void EnablePhysicalDeviceFeatureExtensions(VkDeviceCreateInfo& DeviceInfo, FVulkanDevice& Device);
+	static bool LoadVulkanInstanceFunctions(VkInstance inInstance);
+};
+
+#endif // #if VULKAN_RHI_RAYTRACING
diff --git a/Engine/Source/Runtime/VulkanRHI/Private/Windows/VulkanWindowsPlatform.cpp b/Engine/Source/Runtime/VulkanRHI/Private/Windows/VulkanWindowsPlatform.cpp
--- a/Engine/Source/Runtime/VulkanRHI/Private/Windows/VulkanWindowsPlatform.cpp
+++ b/Engine/Source/Runtime/VulkanRHI/Private/Windows/VulkanWindowsPlatform.cpp
@@ -137,6 +137,14 @@ bool FVulkanWindowsPlatform::LoadVulkanInstanceFunctions(VkInstance inInstance)
 	ENUM_VK_ENTRYPOINTS_OPTIONAL_PLATFORM_INSTANCE(CHECK_VK_ENTRYPOINTS);
 #endif
 
+#if VULKAN_RHI_RAYTRACING
+	const bool bFoundRayTracingEntries = FVulkanRayTracingPlatform::LoadVulkanInstanceFunctions(inInstance);
+	if (!bFoundRayTracingEntries)
+	{
+		UE_LOG(LogVulkanRHI, Warning, TEXT("Vulkan RHI ray tracing is enabled, but failed to load instance functions."));
+	}
+#endif
+	
 	ENUM_VK_ENTRYPOINTS_PLATFORM_INSTANCE(GETINSTANCE_VK_ENTRYPOINTS);
 	ENUM_VK_ENTRYPOINTS_PLATFORM_INSTANCE(CHECK_VK_ENTRYPOINTS);
 #undef GET_VK_ENTRYPOINTS
@@ -211,6 +219,10 @@ void FVulkanWindowsPlatform::GetDeviceExtensions(EGpuVendorId VendorId, TArray<c
 	// Fullscreen requires Instance capabilities2
 	OutExtensions.Add(VK_EXT_FULL_SCREEN_EXCLUSIVE_EXTENSION_NAME);
 #endif
+
+#if VULKAN_RHI_RAYTRACING
+	FVulkanRayTracingPlatform::GetDeviceExtensions(VendorId, OutExtensions);
+#endif
 }
 
 void FVulkanWindowsPlatform::CreateSurface(void* WindowHandle, VkInstance Instance, VkSurfaceKHR* OutSurface)
@@ -370,3 +382,10 @@ void FVulkanWindowsPlatform::CheckDeviceDriver(uint32 DeviceIndex, EGpuVendorId
 		}
 	}
 }
+
+void FVulkanWindowsPlatform::EnablePhysicalDeviceFeatureExtensions(VkDeviceCreateInfo& DeviceInfo, FVulkanDevice& Device)
+{
+#if VULKAN_RHI_RAYTRACING
+	FVulkanRayTracingPlatform::EnablePhysicalDeviceFeatureExtensions(DeviceInfo, Device);
+#endif
+}
diff --git a/Engine/Source/Runtime/VulkanRHI/Private/Windows/VulkanWindowsPlatform.h b/Engine/Source/Runtime/VulkanRHI/Private/Windows/VulkanWindowsPlatform.h
--- a/Engine/Source/Runtime/VulkanRHI/Private/Windows/VulkanWindowsPlatform.h
+++ b/Engine/Source/Runtime/VulkanRHI/Private/Windows/VulkanWindowsPlatform.h
@@ -19,6 +19,8 @@
 
 #define	UE_VK_API_VERSION							VK_API_VERSION_1_1
 
+//#define VULKAN_RHI_RAYTRACING (RHI_RAYTRACING)
+#define VULKAN_RHI_RAYTRACING 0 // Disabled while under development
 
 // 32-bit windows has warnings on custom mem mgr callbacks
 #define VULKAN_SHOULD_USE_LLM					(UE_BUILD_DEBUG || UE_BUILD_DEVELOPMENT) && !PLATFORM_32BITS
@@ -67,6 +69,8 @@ public:
 		static auto* CVar = IConsoleManager::Get().FindTConsoleVariableDataInt(TEXT("r.Vulkan.UseRealUBs"));
 		return (CVar && CVar->GetValueOnAnyThread() == 0) ? false : bCodeHeaderUseRealUBs;
 	}
+
+	static void EnablePhysicalDeviceFeatureExtensions(VkDeviceCreateInfo& DeviceInfo, FVulkanDevice& Device);
 };
 
 typedef FVulkanWindowsPlatform FVulkanPlatform;
