UE5 - Copy from 15250003
Print file path of the aftermath Vulkan gpu crash dump

diff --git a/Engine/Source/Runtime/VulkanRHI/Private/VulkanUtil.cpp b/Engine/Source/Runtime/VulkanRHI/Private/VulkanUtil.cpp
--- a/Engine/Source/Runtime/VulkanRHI/Private/VulkanUtil.cpp
+++ b/Engine/Source/Runtime/VulkanRHI/Private/VulkanUtil.cpp
@@ -682,11 +682,13 @@ void AftermathGpuCrashDumpCallback(const void* CrashDump, const uint32 CrashDump
 	}
 
 	{
-		FArchive* Writer = IFileManager::Get().CreateFileWriter(*(FPaths::ProjectLogDir() / TEXT("vulkan.nv-gpudmp")));
+		FString Filename = FPaths::ProjectLogDir() / TEXT("vulkan.nv-gpudmp");
+		FArchive* Writer = IFileManager::Get().CreateFileWriter(*Filename);
 		if (Writer)
 		{
 			Writer->Serialize((void*)CrashDump, CrashDumpSize);
 			Writer->Close();
+			UE_LOG(LogVulkanRHI, Warning, TEXT("Generated Aftermath crash dump at '%s'"), *Filename);
 		}
 	}
 
@@ -713,11 +715,13 @@ void AftermathGpuCrashDumpCallback(const void* CrashDump, const uint32 CrashDump
 				GFSDK_Aftermath_Result Result2 = GFSDK_Aftermath_GpuCrashDump_GetJSON(Decoder, (uint32)Json.Num(), Json.GetData());
 				if (Result2 == GFSDK_Aftermath_Result_Success)
 				{
-					FArchive* Writer = IFileManager::Get().CreateFileWriter(*(FPaths::ProjectLogDir() / TEXT("vulkan.nv-gpudmp.json")));
+					FString Filename = FPaths::ProjectLogDir() / TEXT("vulkan.nv-gpudmp.json");
+					FArchive* Writer = IFileManager::Get().CreateFileWriter(*Filename);
 					if (Writer)
 					{
 						Writer->Serialize((void*)Json.GetData(), Json.Num());
 						Writer->Close();
+						UE_LOG(LogVulkanRHI, Warning, TEXT("Generated Aftermath crash dump json at '%s'"), *Filename);
 					}
 				}
 				else
