UE5 - Fix gpu crash markers not being dumped

diff --git a/Engine/Source/Runtime/VulkanRHI/Private/VulkanLayers.cpp b/Engine/Source/Runtime/VulkanRHI/Private/VulkanLayers.cpp
--- a/Engine/Source/Runtime/VulkanRHI/Private/VulkanLayers.cpp
+++ b/Engine/Source/Runtime/VulkanRHI/Private/VulkanLayers.cpp
@@ -786,8 +786,9 @@ void FOptionalVulkanDeviceExtensions::Setup(const TArray<const ANSICHAR*>& Devic
 #if VULKAN_SUPPORTS_NV_DIAGNOSTICS
 	if (GGPUCrashDebuggingEnabled)
 	{
-		HasNVDeviceDiagnosticConfig = HasExtension(DeviceExtensions, VK_NV_DEVICE_DIAGNOSTICS_CONFIG_EXTENSION_NAME) && HasExtension(DeviceExtensions, VK_NV_DEVICE_DIAGNOSTIC_CHECKPOINTS_EXTENSION_NAME);
-		bHasAnyCrashExtension = bHasAnyCrashExtension || HasNVDeviceDiagnosticConfig;
+		HasNVDiagnosticCheckpoints = HasExtension(DeviceExtensions, VK_NV_DEVICE_DIAGNOSTIC_CHECKPOINTS_EXTENSION_NAME);
+		HasNVDeviceDiagnosticConfig = HasExtension(DeviceExtensions, VK_NV_DEVICE_DIAGNOSTICS_CONFIG_EXTENSION_NAME);
+		bHasAnyCrashExtension = bHasAnyCrashExtension || (HasNVDeviceDiagnosticConfig && HasNVDiagnosticCheckpoints);
 	}
 #endif
 
diff --git a/Engine/Source/Runtime/VulkanRHI/Private/VulkanUtil.cpp b/Engine/Source/Runtime/VulkanRHI/Private/VulkanUtil.cpp
--- a/Engine/Source/Runtime/VulkanRHI/Private/VulkanUtil.cpp
+++ b/Engine/Source/Runtime/VulkanRHI/Private/VulkanUtil.cpp
@@ -630,7 +630,7 @@ void FVulkanGPUProfiler::DumpCrashMarkers(void* BufferData)
 					check(Data[Index].sType == VK_STRUCTURE_TYPE_CHECKPOINT_DATA_NV);
 					uint32 Value = (uint32)(size_t)Data[Index].pCheckpointMarker;
 					const FString* Frame = CachedStrings.Find(Value);
-					UE_LOG(LogVulkanRHI, Error, TEXT("[VK_NV_device_diagnostic_checkpoints] %i: Stage 0x%x, %s (CRC 0x%x)"), Index, Data[Index].stage, Frame ? *(*Frame) : TEXT("<undefined>"), Value);
+					UE_LOG(LogVulkanRHI, Error, TEXT("[VK_NV_device_diagnostic_checkpoints] %i: Stage 0x%08x, %s (CRC 0x%x)"), Index, Data[Index].stage, Frame ? *(*Frame) : TEXT("<undefined>"), Value);
 				}
 				GLog->PanicFlushThreadedLogs();
 				GLog->Flush();
