Added support for custom Android Vulkan layers

diff --git a/Engine/Source/Programs/UnrealBuildTool/Platform/Android/UEDeployAndroid.cs b/Engine/Source/Programs/UnrealBuildTool/Platform/Android/UEDeployAndroid.cs
--- a/Engine/Source/Programs/UnrealBuildTool/Platform/Android/UEDeployAndroid.cs
+++ b/Engine/Source/Programs/UnrealBuildTool/Platform/Android/UEDeployAndroid.cs
@@ -1300,6 +1300,32 @@ namespace UnrealBuildTool
 					CopyFileDirectory(VulkanLayersDir, DestDir);
 				}
 			}
+
+			// Copy debug validation layers
+			if (bSupportsVulkan && Configuration != "Shipping")
+			{
+				String LayerDir = "";
+				AndroidPlatformSDK.GetPath(Ini, "/Script/AndroidRuntimeSettings.AndroidRuntimeSettings", "DebugVulkanLayerDirectory", out LayerDir);
+
+				Log.TraceInformation("DebugVulkanLayerDirectory {0}", LayerDir);
+
+				if (LayerDir != "")
+				{
+					string VulkanLayersDir = Path.Combine(Environment.ExpandEnvironmentVariables(LayerDir), NDKArch);
+
+					if (Directory.Exists(VulkanLayersDir))
+					{
+						string DestDir = Path.Combine(UE4BuildPath, "libs", NDKArch);
+						Log.TraceInformation("Copying Debug vulkan layers from {0} to {1}", VulkanLayersDir, DestDir);
+
+						if (!Directory.Exists(DestDir))
+						{
+							Directory.CreateDirectory(DestDir);
+						}
+						CopyFileDirectory(VulkanLayersDir, DestDir);
+					}
+				}
+			}
 		}
 
 		void CopyClangSanitizerLib(string UE4BuildPath, string UE4Arch, string NDKArch, AndroidToolChain.ClangSanitizer Sanitizer)
diff --git a/Engine/Source/Runtime/Android/AndroidRuntimeSettings/Classes/AndroidRuntimeSettings.h b/Engine/Source/Runtime/Android/AndroidRuntimeSettings/Classes/AndroidRuntimeSettings.h
--- a/Engine/Source/Runtime/Android/AndroidRuntimeSettings/Classes/AndroidRuntimeSettings.h
+++ b/Engine/Source/Runtime/Android/AndroidRuntimeSettings/Classes/AndroidRuntimeSettings.h
@@ -412,6 +412,18 @@ public:
 	UPROPERTY(GlobalConfig, EditAnywhere, Category = Build, meta = (DisplayName = "Support Vulkan Desktop [Experimental]"))
 	bool bSupportsVulkanSM5;
 
+	/** Directory for Debug Vulkan Layers to package */
+	UPROPERTY(GlobalConfig, EditAnywhere, Category = Build, meta = (DisplayName = "Debug Vulkan Layer Directory"))
+	FDirectoryPath DebugVulkanLayerDirectory;
+
+	/** Debug Vulkan Device Layers to enable */
+	UPROPERTY(GlobalConfig, EditAnywhere, Category = Build, meta = (DisplayName = "Debug Vulkan Device Layers"))
+	TArray<FString> DebugVulkanDeviceLayers;
+
+	/** Debug Vulkan Instance Layers to enable */
+	UPROPERTY(GlobalConfig, EditAnywhere, Category = Build, meta = (DisplayName = "Debug Vulkan Instance Layers"))
+	TArray<FString> DebugVulkanInstanceLayers;
+
 	// Select a Clang's sanitizer to build the project with
 	UPROPERTY(GlobalConfig, EditAnywhere, Category = Build)
 	TEnumAsByte<EClangSanitizer::Type> ClangSanitizer;
diff --git a/Engine/Source/Runtime/VulkanRHI/Private/Android/VulkanAndroidPlatform.cpp b/Engine/Source/Runtime/VulkanRHI/Private/Android/VulkanAndroidPlatform.cpp
--- a/Engine/Source/Runtime/VulkanRHI/Private/Android/VulkanAndroidPlatform.cpp
+++ b/Engine/Source/Runtime/VulkanRHI/Private/Android/VulkanAndroidPlatform.cpp
@@ -57,6 +57,8 @@ TUniquePtr<struct FAndroidVulkanFramePacer> FVulkanAndroidPlatform::FramePacer;
 int32 FVulkanAndroidPlatform::CachedFramePace = 60;
 int32 FVulkanAndroidPlatform::CachedRefreshRate = 60;
 int32 FVulkanAndroidPlatform::CachedSyncInterval = 1;
+TArray<TArray<ANSICHAR>> FVulkanAndroidPlatform::DebugVulkanDeviceLayers;
+TArray<TArray<ANSICHAR>> FVulkanAndroidPlatform::DebugVulkanInstanceLayers;
 
 #define CHECK_VK_ENTRYPOINTS(Type,Func) if (VulkanDynamicAPI::Func == NULL) { bFoundAllEntryPoints = false; UE_LOG(LogRHI, Warning, TEXT("Failed to find entry point for %s"), TEXT(#Func)); }
 
@@ -355,6 +357,33 @@ void FVulkanAndroidPlatform::GetInstanceExtensions(TArray<const ANSICHAR*>& OutE
 	OutExtensions.Add(VK_GOOGLE_DISPLAY_TIMING_EXTENSION_NAME);
 }
 
+void FVulkanAndroidPlatform::GetInstanceLayers(TArray<const ANSICHAR*>& OutLayers)
+{
+#if !UE_BUILD_SHIPPING
+	if (DebugVulkanInstanceLayers.IsEmpty())
+	{
+		TArray<FString> LayerNames;
+		GConfig->GetArray(TEXT("/Script/AndroidRuntimeSettings.AndroidRuntimeSettings"), TEXT("DebugVulkanInstanceLayers"), LayerNames, GEngineIni);
+
+		if (!LayerNames.IsEmpty())
+		{
+			uint32 Index = 0;
+			for (auto& LayerName : LayerNames)
+			{
+				TArray<ANSICHAR> LayerNameANSI{ TCHAR_TO_ANSI(*LayerName), LayerName.Len() + 1 };
+				DebugVulkanInstanceLayers.Add(LayerNameANSI);
+			}
+		}
+	}
+
+	for (const TArray<ANSICHAR>& LayerName : DebugVulkanInstanceLayers)
+	{
+		OutLayers.Add(LayerName.GetData());
+	}
+#endif
+}
+
+
 static int32 GVulkanQcomRenderPassTransform = 0;
 static FAutoConsoleVariableRef CVarVulkanQcomRenderPassTransform(
 	TEXT("r.Vulkan.UseQcomRenderPassTransform"),
@@ -381,6 +410,32 @@ void FVulkanAndroidPlatform::GetDeviceExtensions(EGpuVendorId VendorId, TArray<c
 #endif
 }
 
+void FVulkanAndroidPlatform::GetDeviceLayers(EGpuVendorId VendorId, TArray<const ANSICHAR*>& OutLayers)
+{
+#if !UE_BUILD_SHIPPING
+	if (DebugVulkanDeviceLayers.IsEmpty())
+	{
+		TArray<FString> LayerNames;
+		GConfig->GetArray(TEXT("/Script/AndroidRuntimeSettings.AndroidRuntimeSettings"), TEXT("DebugVulkanDeviceLayers"), LayerNames, GEngineIni);
+
+		if (!LayerNames.IsEmpty())
+		{
+			uint32 Index = 0;
+			for (auto& LayerName : LayerNames)
+			{
+				TArray<ANSICHAR> LayerNameANSI{ TCHAR_TO_ANSI(*LayerName), LayerName.Len() + 1 };
+				DebugVulkanDeviceLayers.Add(LayerNameANSI);
+			}
+		}
+	}
+
+	for (auto& LayerName : DebugVulkanDeviceLayers)
+	{
+		OutLayers.Add(LayerName.GetData());
+	}
+#endif
+}
+
 void FVulkanAndroidPlatform::NotifyFoundDeviceLayersAndExtensions(VkPhysicalDevice PhysicalDevice, const TArray<FString>& Layers, const TArray<FString>& Extensions)
 {
 #if VULKAN_SUPPORTS_GOOGLE_DISPLAY_TIMING
diff --git a/Engine/Source/Runtime/VulkanRHI/Private/Android/VulkanAndroidPlatform.h b/Engine/Source/Runtime/VulkanRHI/Private/Android/VulkanAndroidPlatform.h
--- a/Engine/Source/Runtime/VulkanRHI/Private/Android/VulkanAndroidPlatform.h
+++ b/Engine/Source/Runtime/VulkanRHI/Private/Android/VulkanAndroidPlatform.h
@@ -57,7 +57,9 @@ public:
 	static void FreeVulkanLibrary();
 
 	static void GetInstanceExtensions(TArray<const ANSICHAR*>& OutExtensions);
+	static void GetInstanceLayers(TArray<const ANSICHAR*>& OutLayers);
 	static void GetDeviceExtensions(EGpuVendorId VendorId, TArray<const ANSICHAR*>& OutExtensions);
+	static void GetDeviceLayers(EGpuVendorId VendorId, TArray<const ANSICHAR*>& OutLayers);
 	static void NotifyFoundDeviceLayersAndExtensions(VkPhysicalDevice PhysicalDevice, const TArray<FString>& Layers, const TArray<FString>& Extensions);
 
 	static void CreateSurface(void* WindowHandle, VkInstance Instance, VkSurfaceKHR* OutSurface);
@@ -146,6 +148,8 @@ protected:
 	static int32 CachedFramePace;
 	static int32 CachedRefreshRate;
 	static int32 CachedSyncInterval;
+	static TArray<TArray<ANSICHAR>> DebugVulkanDeviceLayers;
+	static TArray<TArray<ANSICHAR>> DebugVulkanInstanceLayers;
 };
 
 #if VULKAN_SUPPORTS_GOOGLE_DISPLAY_TIMING
diff --git a/Engine/Source/Runtime/VulkanRHI/Private/Linux/VulkanLinuxPlatform.h b/Engine/Source/Runtime/VulkanRHI/Private/Linux/VulkanLinuxPlatform.h
--- a/Engine/Source/Runtime/VulkanRHI/Private/Linux/VulkanLinuxPlatform.h
+++ b/Engine/Source/Runtime/VulkanRHI/Private/Linux/VulkanLinuxPlatform.h
@@ -42,7 +42,9 @@ public:
 	static void FreeVulkanLibrary();
 
 	static void GetInstanceExtensions(TArray<const ANSICHAR*>& OutExtensions);
+	static void GetInstanceLayers(TArray<const ANSICHAR*>& OutLayers) {}
 	static void GetDeviceExtensions(EGpuVendorId VendorId, TArray<const ANSICHAR*>& OutExtensions);
+	static void GetDeviceLayers(EGpuVendorId VendorId, TArray<const ANSICHAR*>& OutLayers) {}
 
 	static void CreateSurface(void* WindowHandle, VkInstance Instance, VkSurfaceKHR* OutSurface);
 
diff --git a/Engine/Source/Runtime/VulkanRHI/Private/Lumin/VulkanLuminPlatform.h b/Engine/Source/Runtime/VulkanRHI/Private/Lumin/VulkanLuminPlatform.h
--- a/Engine/Source/Runtime/VulkanRHI/Private/Lumin/VulkanLuminPlatform.h
+++ b/Engine/Source/Runtime/VulkanRHI/Private/Lumin/VulkanLuminPlatform.h
@@ -44,7 +44,9 @@ public:
 	static void NotifyFoundDeviceLayersAndExtensions(VkPhysicalDevice PhysicalDevice, const TArray<FString>& Layers, const TArray<FString>& Extensions);
 
 	static void GetInstanceExtensions(TArray<const ANSICHAR*>& OutExtensions);
+	static void GetInstanceLayers(TArray<const ANSICHAR*>& OutLayers) {}
 	static void GetDeviceExtensions(EGpuVendorId VendorId, TArray<const ANSICHAR*>& OutExtensions);
+	static void GetDeviceLayers(EGpuVendorId VendorId, TArray<const ANSICHAR*>& OutLayers) {}
 
 	static void CreateSurface(void* WindowHandle, VkInstance Instance, VkSurfaceKHR* OutSurface);
 
diff --git a/Engine/Source/Runtime/VulkanRHI/Private/VulkanGenericPlatform.h b/Engine/Source/Runtime/VulkanRHI/Private/VulkanGenericPlatform.h
--- a/Engine/Source/Runtime/VulkanRHI/Private/VulkanGenericPlatform.h
+++ b/Engine/Source/Runtime/VulkanRHI/Private/VulkanGenericPlatform.h
@@ -30,7 +30,9 @@ public:
 
 	// Array of required extensions for the platform (Required!)
 	static void GetInstanceExtensions(TArray<const ANSICHAR*>& OutExtensions);
+	static void GetInstanceLayers(TArray<const ANSICHAR*>& OutLayers) {}
 	static void GetDeviceExtensions(EGpuVendorId VendorId, TArray<const ANSICHAR*>& OutExtensions);
+	static void GetDeviceLayers(EGpuVendorId VendorId, TArray<const ANSICHAR*>& OutLayers) {}
 
 	// create the platform-specific surface object - required
 	static void CreateSurface(VkSurfaceKHR* OutSurface);
diff --git a/Engine/Source/Runtime/VulkanRHI/Private/VulkanLayers.cpp b/Engine/Source/Runtime/VulkanRHI/Private/VulkanLayers.cpp
--- a/Engine/Source/Runtime/VulkanRHI/Private/VulkanLayers.cpp
+++ b/Engine/Source/Runtime/VulkanRHI/Private/VulkanLayers.cpp
@@ -467,6 +467,13 @@ void FVulkanDynamicRHI::GetInstanceLayersAndExtensions(TArray<const ANSICHAR*>&
 		}
 	}
 
+	TArray<const ANSICHAR*> PlatformLayers;
+	FVulkanPlatform::GetInstanceLayers(PlatformLayers);
+	for (const ANSICHAR* PlatformLayer : PlatformLayers)
+	{
+		OutInstanceLayers.Add(PlatformLayer);
+	}
+
 	for (int32 j = 0; GInstanceExtensions[j] != nullptr; j++)
 	{
 		if (FindLayerExtensionInList(GlobalLayerExtensions, GInstanceExtensions[j]))
@@ -700,6 +707,13 @@ void FVulkanDevice::GetDeviceExtensionsAndLayers(VkPhysicalDevice Gpu, EGpuVendo
 		}
 	}
 
+	TArray<const ANSICHAR*> PlatformLayers;
+	FVulkanPlatform::GetDeviceLayers(VendorId, PlatformLayers);
+	for (const ANSICHAR* PlatformLayer : PlatformLayers)
+	{
+		OutDeviceLayers.Add(PlatformLayer);
+	}
+
 	for (uint32 Index = 0; Index < UE_ARRAY_COUNT(GDeviceExtensions) && GDeviceExtensions[Index] != nullptr; ++Index)
 	{
 		if (ListContains(AvailableExtensions, GDeviceExtensions[Index]))
diff --git a/Engine/Source/Runtime/VulkanRHI/Private/Windows/VulkanWindowsPlatform.h b/Engine/Source/Runtime/VulkanRHI/Private/Windows/VulkanWindowsPlatform.h
--- a/Engine/Source/Runtime/VulkanRHI/Private/Windows/VulkanWindowsPlatform.h
+++ b/Engine/Source/Runtime/VulkanRHI/Private/Windows/VulkanWindowsPlatform.h
@@ -55,7 +55,9 @@ public:
 	static void FreeVulkanLibrary();
 
 	static void GetInstanceExtensions(TArray<const ANSICHAR*>& OutExtensions);
+	static void GetInstanceLayers(TArray<const ANSICHAR*>& OutLayers) {}
 	static void GetDeviceExtensions(EGpuVendorId VendorId, TArray<const ANSICHAR*>& OutExtensions);
+	static void GetDeviceLayers(EGpuVendorId VendorId, TArray<const ANSICHAR*>& OutLayers) {}
 
 	static void CreateSurface(void* WindowHandle, VkInstance Instance, VkSurfaceKHR* OutSurface);
 
