UE - Compile fix
#rnx

diff --git a/Engine/Source/Runtime/VulkanRHI/Private/Android/VulkanAndroidPlatform.cpp b/Engine/Source/Runtime/VulkanRHI/Private/Android/VulkanAndroidPlatform.cpp
--- a/Engine/Source/Runtime/VulkanRHI/Private/Android/VulkanAndroidPlatform.cpp
+++ b/Engine/Source/Runtime/VulkanRHI/Private/Android/VulkanAndroidPlatform.cpp
@@ -17,6 +17,7 @@
 #include "Android/AndroidPlatformFramePacer.h"
 #include "Math/UnrealMathUtility.h"
 #include "Android/AndroidPlatformMisc.h"
+#include "Misc/ConfigCacheIni.h"
 
 // From VulklanSwapChain.cpp
 extern int32 GVulkanCPURenderThreadFramePacer;
