Fixed ClearUAV unit tests and added support for byte address buffers

diff --git a/Engine/Source/Runtime/VulkanRHI/Private/VulkanContext.h b/Engine/Source/Runtime/VulkanRHI/Private/VulkanContext.h
--- a/Engine/Source/Runtime/VulkanRHI/Private/VulkanContext.h
+++ b/Engine/Source/Runtime/VulkanRHI/Private/VulkanContext.h
@@ -205,6 +205,7 @@ protected:
 
 	void InternalClearMRT(FVulkanCmdBuffer* CmdBuffer, bool bClearColor, int32 NumClearColors, const FLinearColor* ColorArray, bool bClearDepth, float Depth, bool bClearStencil, uint32 Stencil);
 
+	static void ClearUAVFillBuffer(FVulkanUnorderedAccessView* UAV, uint32_t ClearValue);
 	static void ClearUAV(TRHICommandList_RecursiveHazardous<FVulkanCommandListContext>& RHICmdList, FVulkanUnorderedAccessView* UnorderedAccessView, const void* ClearValue, bool bFloat);
 
 public:
diff --git a/Engine/Source/Runtime/VulkanRHI/Private/VulkanUAV.cpp b/Engine/Source/Runtime/VulkanRHI/Private/VulkanUAV.cpp
--- a/Engine/Source/Runtime/VulkanRHI/Private/VulkanUAV.cpp
+++ b/Engine/Source/Runtime/VulkanRHI/Private/VulkanUAV.cpp
@@ -437,6 +437,15 @@ void FVulkanDynamicRHI::RHIUpdateShaderResourceView(FRHIShaderResourceView* SRV,
 	}
 }
 
+void FVulkanCommandListContext::ClearUAVFillBuffer(FVulkanUnorderedAccessView* UAV, uint32_t ClearValue)
+{
+	FVulkanCommandBufferManager* CmdBufferMgr = GVulkanRHI->GetDevice()->GetImmediateContext().GetCommandBufferManager();
+	FVulkanCmdBuffer* CmdBuffer = CmdBufferMgr->GetActiveCmdBuffer();
+
+	FVulkanResourceMultiBuffer* Buffer = UAV->SourceBuffer;
+	VulkanRHI::vkCmdFillBuffer(CmdBuffer->GetHandle(), Buffer->Current.Handle, Buffer->Current.Offset, Buffer->Current.Size, ClearValue);
+}
+
 void FVulkanCommandListContext::ClearUAV(TRHICommandList_RecursiveHazardous<FVulkanCommandListContext>& RHICmdList, FVulkanUnorderedAccessView* UnorderedAccessView, const void* ClearValue, bool bFloat)
 {
 	struct FVulkanDynamicRHICmdFillBuffer final : public FRHICommand<FVulkanDynamicRHICmdFillBuffer>
@@ -451,11 +460,7 @@ void FVulkanCommandListContext::ClearUAV(TRHICommandList_RecursiveHazardous<FVul
 
 		void Execute(FRHICommandListBase& CmdList)
 		{
-			FVulkanCommandBufferManager* CmdBufferMgr = GVulkanRHI->GetDevice()->GetImmediateContext().GetCommandBufferManager();
-			FVulkanCmdBuffer* CmdBuffer = CmdBufferMgr->GetActiveCmdBuffer();
-
-			FVulkanResourceMultiBuffer* Buffer = UAV->SourceBuffer;
-			VulkanRHI::vkCmdFillBuffer(CmdBuffer->GetHandle(), Buffer->Current.Handle, Buffer->Current.Offset, Buffer->Current.Size, ClearValue);
+			ClearUAVFillBuffer(UAV, ClearValue);
 		}
 	};
 
@@ -495,23 +500,30 @@ void FVulkanCommandListContext::ClearUAV(TRHICommandList_RecursiveHazardous<FVul
 
 	if (UnorderedAccessView->SourceBuffer)
 	{
-		RHICmdList.Transition(FRHITransitionInfo(UnorderedAccessView, ERHIAccess::UAVCompute, ERHIAccess::CopyDest));
-			
 		TRefCountPtr<FVulkanResourceMultiBuffer> Buffer = UnorderedAccessView->SourceBuffer;
+		bool bIsByteAddressBuffer = Buffer->GetUsage() & BUF_ByteAddressBuffer;
 
-		if (RHICmdList.Bypass())
+		// Byte address buffers only use the first component, so use vkCmdBufferFill
+		if (UnorderedAccessView->BufferViewFormat == PF_Unknown || bIsByteAddressBuffer)
 		{
-			FVulkanCommandBufferManager* CmdBufferMgr = GVulkanRHI->GetDevice()->GetImmediateContext().GetCommandBufferManager();
-			FVulkanCmdBuffer* CmdBuffer = CmdBufferMgr->GetActiveCmdBuffer();
+			RHICmdList.Transition(FRHITransitionInfo(UnorderedAccessView, ERHIAccess::UAVCompute, ERHIAccess::CopyDest));
+
+			if (RHICmdList.Bypass())
+			{
+				ClearUAVFillBuffer(UnorderedAccessView, *(const uint32_t*)ClearValue);
+			}
+			else
+			{
+				new (RHICmdList.AllocCommand<FVulkanDynamicRHICmdFillBuffer>()) FVulkanDynamicRHICmdFillBuffer(UnorderedAccessView, *(const uint32_t*)ClearValue);
+			}
 
-			VulkanRHI::vkCmdFillBuffer(CmdBuffer->GetHandle(), Buffer->Current.Handle, Buffer->Current.Offset, Buffer->Current.Size, *(const uint32_t*)ClearValue);
+			RHICmdList.Transition(FRHITransitionInfo(UnorderedAccessView, ERHIAccess::CopyDest, ERHIAccess::UAVCompute));
 		}
 		else
 		{
-			new (RHICmdList.AllocCommand<FVulkanDynamicRHICmdFillBuffer>()) FVulkanDynamicRHICmdFillBuffer(UnorderedAccessView, *(const uint32_t*)ClearValue);
+			uint32 NumElements = Buffer->Current.Size / GPixelFormats[UnorderedAccessView->BufferViewFormat].BlockBytes;
+			ClearUAVShader_T<EClearReplacementResourceType::Buffer, 4, false>(RHICmdList, UnorderedAccessView, NumElements, 1, 1, ClearValue, ValueType);
 		}
-
-		RHICmdList.Transition(FRHITransitionInfo(UnorderedAccessView, ERHIAccess::CopyDest, ERHIAccess::UAVCompute));
 	}
 	else if (UnorderedAccessView->SourceTexture)
 	{
