Fix blown out rendering on Vulkan when using an HDR monitor.

Vulkan now respects r.HDR.Display.OutputDevice, and tries to use a swapchain color format that matches the value of that cvar.

diff --git a/Engine/Source/Runtime/VulkanRHI/Private/VulkanSwapChain.cpp b/Engine/Source/Runtime/VulkanRHI/Private/VulkanSwapChain.cpp
--- a/Engine/Source/Runtime/VulkanRHI/Private/VulkanSwapChain.cpp
+++ b/Engine/Source/Runtime/VulkanRHI/Private/VulkanSwapChain.cpp
@@ -149,6 +149,28 @@ FVulkanSwapChain::FVulkanSwapChain(VkInstance InInstance, FVulkanDevice& InDevic
 		FVulkanPlatform::CreateSurface(WindowHandle, Instance, &Surface);
 	}
 
+	static const auto CVarHDROutputDevice = IConsoleManager::Get().FindTConsoleVariableDataInt(TEXT("r.HDR.Display.OutputDevice"));
+	int32 OutputDevice = CVarHDROutputDevice ? CVarHDROutputDevice->GetValueOnAnyThread() : 0;
+	VkColorSpaceKHR RequestedColorSpace;
+	// The possible values are documented in PostProcessTonemap.cpp, where the cvar is defined. They match the ETonemapperOutputDevice enum, which is defined in a header we cannot include.
+	switch (OutputDevice)
+	{
+	case 0:
+		RequestedColorSpace = VK_COLOR_SPACE_SRGB_NONLINEAR_KHR;
+		break;
+	case 1:
+		RequestedColorSpace = VK_COLOR_SPACE_BT709_NONLINEAR_EXT;
+		break;
+	case 3:
+	case 4:
+		RequestedColorSpace = VK_COLOR_SPACE_HDR10_ST2084_EXT;
+		break;
+	default:
+		UE_LOG(LogVulkanRHI, Warning, TEXT("Requested color format %d not supported in Vulkan, falling back to sRGB. Please check the value of r.HDR.Display.OutputDevice."), OutputDevice);
+		RequestedColorSpace = VK_COLOR_SPACE_SRGB_NONLINEAR_KHR;
+		break;
+	}
+
 	// Find Pixel format for presentable images
 	VkSurfaceFormatKHR CurrFormat;
 	FMemory::Memzero(CurrFormat);
@@ -179,7 +201,13 @@ FVulkanSwapChain::FVulkanSwapChain(VkInstance InInstance, FVulkanDevice& InDevic
 					{
 						CurrFormat = Formats[Index];
 						bFound = true;
-						break;
+
+						// We stop the search if both the pixel format and color space are a match. However, if we can't find a matching color space, we'll still use one of the
+						// formats that at least matches the pixel format.
+						if (CurrFormat.colorSpace == RequestedColorSpace)
+						{
+							break;
+						}
 					}
 				}
 
@@ -208,7 +236,7 @@ FVulkanSwapChain::FVulkanSwapChain(VkInstance InInstance, FVulkanDevice& InDevic
 				check(Formats[Index].format != VK_FORMAT_UNDEFINED);
 				for (int32 PFIndex = 0; PFIndex < PF_MAX; ++PFIndex)
 				{
-					if (Formats[Index].format == GPixelFormats[PFIndex].PlatformFormat)
+					if (Formats[Index].format == GPixelFormats[PFIndex].PlatformFormat && Formats[Index].colorSpace == RequestedColorSpace)
 					{
 						InOutPixelFormat = (EPixelFormat)PFIndex;
 						CurrFormat = Formats[Index];
@@ -231,7 +259,7 @@ FVulkanSwapChain::FVulkanSwapChain(VkInstance InInstance, FVulkanDevice& InDevic
 			bool bSupported = false;
 			for (int32 Index = 0; Index < Formats.Num(); ++Index)
 			{
-				if (Formats[Index].format == PlatformFormat)
+				if (Formats[Index].format == PlatformFormat && Formats[Index].colorSpace == RequestedColorSpace)
 				{
 					bSupported = true;
 					CurrFormat = Formats[Index];
@@ -255,7 +283,7 @@ FVulkanSwapChain::FVulkanSwapChain(VkInstance InInstance, FVulkanDevice& InDevic
 				{
 					Msg += TEXT(", ");
 				}
-				Msg += FString::Printf(TEXT("%d"), (int32)Formats[Index].format);
+				Msg += FString::Printf(TEXT("%d/%d"), (int32)Formats[Index].format, (int32)Formats[Index].colorSpace);
 			}
 			if (Formats.Num())
 			{
@@ -265,6 +293,11 @@ FVulkanSwapChain::FVulkanSwapChain(VkInstance InInstance, FVulkanDevice& InDevic
 		}
 	}
 
+	if (CurrFormat.colorSpace != RequestedColorSpace)
+	{
+		UE_LOG(LogVulkanRHI, Warning, TEXT("Requested color format %d not supported by this Vulkan implementation, falling back to %d. Please check the value of r.HDR.Display.OutputDevice."), (int32)RequestedColorSpace, (int32)CurrFormat.colorSpace);
+	}
+
 	VkFormat PlatformFormat = UEToVkTextureFormat(InOutPixelFormat, false);
 
 	Device.SetupPresentQueue(Surface);
