Minimum Vulkan rt blas implementation.
Currently only supports build, not update.
Allocations are not sub-allocated yet.
Plumbing for tlas only.

diff --git a/Engine/Source/Runtime/VulkanRHI/Private/VulkanContext.h b/Engine/Source/Runtime/VulkanRHI/Private/VulkanContext.h
--- a/Engine/Source/Runtime/VulkanRHI/Private/VulkanContext.h
+++ b/Engine/Source/Runtime/VulkanRHI/Private/VulkanContext.h
@@ -102,6 +102,11 @@ public:
 	virtual void RHIEndRenderPass() final override;
 	virtual void RHINextSubpass() final override;
 
+#if VULKAN_RHI_RAYTRACING
+	virtual void RHIBuildAccelerationStructures(const TArrayView<const FAccelerationStructureBuildParams> Params) final override;
+	virtual void RHIBuildAccelerationStructure(FRHIRayTracingScene* Scene) final override;
+#endif
+
 	inline FVulkanCommandBufferManager* GetCommandBufferManager()
 	{
 		return CommandBufferManager;
diff --git a/Engine/Source/Runtime/VulkanRHI/Private/VulkanRHI.cpp b/Engine/Source/Runtime/VulkanRHI/Private/VulkanRHI.cpp
--- a/Engine/Source/Runtime/VulkanRHI/Private/VulkanRHI.cpp
+++ b/Engine/Source/Runtime/VulkanRHI/Private/VulkanRHI.cpp
@@ -712,7 +712,9 @@ void FVulkanDynamicRHI::InitInstance()
 #if VULKAN_SUPPORTS_MULTIVIEW
 		GSupportsMobileMultiView = Device->GetMultiviewFeatures().multiview == VK_TRUE ? true : false;
 #endif
-
+#if VULKAN_RHI_RAYTRACING
+		GRHISupportsRayTracing = Device->GetOptionalExtensions().HasRaytracingExtensions();
+#endif
 #if VULKAN_ENABLE_DUMP_LAYER
 		// Disable RHI thread by default if the dump layer is enabled
 		GRHISupportsRHIThread = false;
diff --git a/Engine/Source/Runtime/VulkanRHI/Private/VulkanRayTracing.cpp b/Engine/Source/Runtime/VulkanRHI/Private/VulkanRayTracing.cpp
--- a/Engine/Source/Runtime/VulkanRHI/Private/VulkanRayTracing.cpp
+++ b/Engine/Source/Runtime/VulkanRHI/Private/VulkanRayTracing.cpp
@@ -1,6 +1,7 @@
 // Copyright Epic Games, Inc. All Rights Reserved.
 
 #include "VulkanRHIPrivate.h"
+#include "VulkanContext.h"
 
 #if VULKAN_RHI_RAYTRACING
 
@@ -70,4 +71,286 @@ bool FVulkanRayTracingPlatform::LoadVulkanInstanceFunctions(VkInstance inInstanc
 #if PLATFORM_WINDOWS
 #pragma warning(pop) // restore 4191
 #endif
+
+static VkDeviceAddress GetDeviceAddress(VkDevice Device, VkBuffer Buffer)
+{
+	VkBufferDeviceAddressInfoKHR DeviceAddressInfo;
+	ZeroVulkanStruct(DeviceAddressInfo, VK_STRUCTURE_TYPE_BUFFER_DEVICE_ADDRESS_INFO);
+	DeviceAddressInfo.buffer = Buffer;
+	return vkGetBufferDeviceAddressKHR(Device, &DeviceAddressInfo);
+}
+
+// Temporary brute force allocation helper, this should be handled by the memory sub-allocator
+static uint32 FindMemoryType(VkPhysicalDevice Gpu, uint32 Filter, VkMemoryPropertyFlags RequestedProperties)
+{
+	VkPhysicalDeviceMemoryProperties Properties = {};
+	vkGetPhysicalDeviceMemoryProperties(Gpu, &Properties);
+
+	uint32 Result = UINT32_MAX;
+	for (uint32 i = 0; i < Properties.memoryTypeCount; ++i)
+	{
+		const bool bTypeFilter = Filter & (1 << i);
+		const bool bPropFilter = (Properties.memoryTypes[i].propertyFlags & RequestedProperties) == RequestedProperties;
+		if (bTypeFilter && bPropFilter)
+		{
+			Result = i;
+			break;
+		}
+	}
+
+	check(Result < UINT32_MAX);
+	return Result;
+}
+
+// Temporary brute force allocation
+static void Allocate(VkPhysicalDevice Gpu, VkDevice Device, VkDeviceSize Size, VkBufferUsageFlags Flags, FAccelerationStructureAllocation& Result)
+{
+	VkBufferCreateInfo CreateInfo;
+	ZeroVulkanStruct(CreateInfo, VK_STRUCTURE_TYPE_BUFFER_CREATE_INFO);
+	CreateInfo.size = Size;
+	CreateInfo.usage = Flags;
+	VERIFYVULKANRESULT(vkCreateBuffer(Device, &CreateInfo, VULKAN_CPU_ALLOCATOR, &Result.Buffer));
+
+	VkMemoryRequirements MemoryRequirements = {};
+	vkGetBufferMemoryRequirements(Device, Result.Buffer, &MemoryRequirements);
+
+	VkMemoryAllocateFlagsInfo MemoryAllocateFlagsInfo;
+	ZeroVulkanStruct(MemoryAllocateFlagsInfo, VK_STRUCTURE_TYPE_MEMORY_ALLOCATE_FLAGS_INFO);
+	MemoryAllocateFlagsInfo.flags = VK_MEMORY_ALLOCATE_DEVICE_ADDRESS_BIT_KHR;
+
+	VkMemoryAllocateInfo MemoryAllocateInfo;
+	ZeroVulkanStruct(MemoryAllocateInfo, VK_STRUCTURE_TYPE_MEMORY_ALLOCATE_INFO);
+	MemoryAllocateInfo.pNext = &MemoryAllocateFlagsInfo;
+	MemoryAllocateInfo.allocationSize = MemoryRequirements.size;
+	MemoryAllocateInfo.memoryTypeIndex = FindMemoryType(Gpu, MemoryRequirements.memoryTypeBits, VK_MEMORY_PROPERTY_DEVICE_LOCAL_BIT);
+	VERIFYVULKANRESULT(vkAllocateMemory(Device, &MemoryAllocateInfo, VULKAN_CPU_ALLOCATOR, &Result.Memory));
+	VERIFYVULKANRESULT(vkBindBufferMemory(Device, Result.Buffer, Result.Memory, 0));
+
+	Result.Device = Device;
+}
+
+// Temporary brute force deallocation
+static void Free(FAccelerationStructureAllocation& Allocation)
+{
+	if (Allocation.Buffer != VK_NULL_HANDLE)
+	{
+		vkDestroyBuffer(Allocation.Device, Allocation.Buffer, VULKAN_CPU_ALLOCATOR);
+		Allocation.Buffer = VK_NULL_HANDLE;
+	}
+	if (Allocation.Memory != VK_NULL_HANDLE)
+	{
+		vkFreeMemory(Allocation.Device, Allocation.Memory, VULKAN_CPU_ALLOCATOR);
+		Allocation.Memory = VK_NULL_HANDLE;
+	}
+}
+
+FVulkanRayTracingGeometry::FVulkanRayTracingGeometry(const FRayTracingGeometryInitializer& Initializer)
+{
+	// Only supporting triangles initially
+	check(Initializer.GeometryType == ERayTracingGeometryType::RTGT_Triangles);
+
+	DebugName = Initializer.DebugName;
+	IndexOffsetInBytes = Initializer.IndexBufferOffset;
+	TotalPrimitiveCount = Initializer.TotalPrimitiveCount;
+	IndexBufferRHI = Initializer.IndexBuffer;
+	bFastBuild = Initializer.bFastBuild;
+	bAllowUpdate = Initializer.bAllowUpdate;
+	Segments = TArray<FRayTracingGeometrySegment>(Initializer.Segments.GetData(), Initializer.Segments.Num());
+}
+
+FVulkanRayTracingGeometry::~FVulkanRayTracingGeometry()
+{
+	Free(Allocation);
+	Free(Scratch);
+}
+
+void FVulkanRayTracingGeometry::BuildAccelerationStructure(FVulkanCommandListContext& CommandContext, EAccelerationStructureBuildMode BuildMode)
+{
+	check(BuildMode == EAccelerationStructureBuildMode::Build); // Temp
+
+	checkf(!IndexBufferRHI || (IndexBufferRHI->GetStride() == 2 || IndexBufferRHI->GetStride() == 4), TEXT("Index buffer must be 16 or 32 bit if in use."));
+	const VkDevice Device = CommandContext.GetDevice()->GetInstanceHandle();
+
+	IndexStrideInBytes = IndexBufferRHI ? IndexBufferRHI->GetStride() : 0;
+
+	FVulkanResourceMultiBuffer* const IndexBuffer = ResourceCast(IndexBufferRHI.GetReference());
+	VkDeviceOrHostAddressConstKHR IndexBufferDeviceAddress = {};
+	IndexBufferDeviceAddress.deviceAddress = IndexBufferRHI ? GetDeviceAddress(Device, IndexBuffer->GetHandle()) : 0;
+
+	TArray<VkAccelerationStructureGeometryKHR, TInlineAllocator<1>> BuildSegments;
+	TArray<VkAccelerationStructureBuildRangeInfoKHR, TInlineAllocator<1>> BuildRanges;
+	TArray<uint32> VertexCounts;
+
+	for (int32 SegmentIndex = 0; SegmentIndex < Segments.Num(); ++SegmentIndex)
+	{
+		const FRayTracingGeometrySegment& Segment = Segments[SegmentIndex];
+
+		FVulkanResourceMultiBuffer* const VertexBuffer = ResourceCast(Segment.VertexBuffer.GetReference());
+		VkDeviceOrHostAddressConstKHR VertexBufferDeviceAddress = {};
+		VertexBufferDeviceAddress.deviceAddress = GetDeviceAddress(Device, VertexBuffer->GetHandle());
+
+		// #vkrt_todo: Add explicit vertex count to FRayTracingGeometrySegment instead of this estimate
+		const uint32 MaxSegmentVertices = (VertexBuffer->GetSize() - Segment.VertexBufferOffset) / Segment.VertexBufferStride;
+		VertexCounts.Add(MaxSegmentVertices);
+
+		VkAccelerationStructureGeometryKHR SegmentGeometry;
+		ZeroVulkanStruct(SegmentGeometry, VK_STRUCTURE_TYPE_ACCELERATION_STRUCTURE_GEOMETRY_KHR);
+		if (Segment.bForceOpaque)
+		{
+			SegmentGeometry.flags = VK_GEOMETRY_OPAQUE_BIT_KHR;
+		}
+
+		// Only support triangles
+		SegmentGeometry.geometryType = VK_GEOMETRY_TYPE_TRIANGLES_KHR;
+		
+		SegmentGeometry.geometry.triangles.sType = VK_STRUCTURE_TYPE_ACCELERATION_STRUCTURE_GEOMETRY_TRIANGLES_DATA_KHR;
+		SegmentGeometry.geometry.triangles.vertexFormat = VK_FORMAT_R32G32B32_SFLOAT;
+		SegmentGeometry.geometry.triangles.vertexData = VertexBufferDeviceAddress;
+		SegmentGeometry.geometry.triangles.maxVertex = MaxSegmentVertices;
+		SegmentGeometry.geometry.triangles.vertexStride = Segment.VertexBufferStride;
+		SegmentGeometry.geometry.triangles.indexData = IndexBufferDeviceAddress;
+
+		// No support for segment transform
+		SegmentGeometry.geometry.triangles.transformData.deviceAddress = 0;
+		SegmentGeometry.geometry.triangles.transformData.hostAddress = nullptr;
+
+		if (IndexBufferRHI)
+		{
+			SegmentGeometry.geometry.triangles.indexType = (IndexStrideInBytes == 2) ? VK_INDEX_TYPE_UINT16 : VK_INDEX_TYPE_UINT32;
+		}
+		else
+		{
+			SegmentGeometry.geometry.triangles.indexType = VK_INDEX_TYPE_NONE_KHR;
+		}
+
+		BuildSegments.Add(SegmentGeometry);
+
+		VkAccelerationStructureBuildRangeInfoKHR RangeInfo = {};
+		RangeInfo.firstVertex = 0;
+
+		// Disabled segments use an empty range. We still build them to keep the sbt valid.
+		RangeInfo.primitiveCount = (Segment.bEnabled) ? MaxSegmentVertices : 0;
+
+		RangeInfo.primitiveOffset = 0;
+		RangeInfo.transformOffset = 0;
+
+		BuildRanges.Add(RangeInfo);
+	}
+
+	VkAccelerationStructureBuildGeometryInfoKHR BuildGeometryInfo;
+	ZeroVulkanStruct(BuildGeometryInfo, VK_STRUCTURE_TYPE_ACCELERATION_STRUCTURE_BUILD_GEOMETRY_INFO_KHR);
+	BuildGeometryInfo.type = VK_ACCELERATION_STRUCTURE_TYPE_BOTTOM_LEVEL_KHR;
+	BuildGeometryInfo.flags = (bFastBuild) ? VK_BUILD_ACCELERATION_STRUCTURE_PREFER_FAST_BUILD_BIT_KHR : VK_BUILD_ACCELERATION_STRUCTURE_PREFER_FAST_TRACE_BIT_KHR;
+	if (bAllowUpdate)
+	{
+		BuildGeometryInfo.flags |= VK_BUILD_ACCELERATION_STRUCTURE_ALLOW_UPDATE_BIT_KHR;
+	}
+	BuildGeometryInfo.mode = (BuildMode == EAccelerationStructureBuildMode::Build) ? VK_BUILD_ACCELERATION_STRUCTURE_MODE_BUILD_KHR : VK_BUILD_ACCELERATION_STRUCTURE_MODE_UPDATE_KHR;
+	BuildGeometryInfo.geometryCount = BuildSegments.Num();
+	BuildGeometryInfo.pGeometries = BuildSegments.GetData();
+
+	VkAccelerationStructureBuildSizesInfoKHR BuildSizesInfo;
+	ZeroVulkanStruct(BuildSizesInfo, VK_STRUCTURE_TYPE_ACCELERATION_STRUCTURE_BUILD_SIZES_INFO_KHR);
+	vkGetAccelerationStructureBuildSizesKHR(
+		Device,
+		VK_ACCELERATION_STRUCTURE_BUILD_TYPE_DEVICE_KHR,
+		&BuildGeometryInfo,
+		VertexCounts.GetData(),
+		&BuildSizesInfo);
+
+	Allocate(
+		CommandContext.GetDevice()->GetPhysicalHandle(), 
+		Device, 
+		BuildSizesInfo.accelerationStructureSize,
+		VK_BUFFER_USAGE_ACCELERATION_STRUCTURE_STORAGE_BIT_KHR | VK_BUFFER_USAGE_SHADER_DEVICE_ADDRESS_BIT, 
+		Allocation);
+
+	Allocate(
+		CommandContext.GetDevice()->GetPhysicalHandle(),
+		Device,
+		BuildSizesInfo.buildScratchSize,
+		VK_BUFFER_USAGE_STORAGE_BUFFER_BIT | VK_BUFFER_USAGE_SHADER_DEVICE_ADDRESS_BIT,
+		Scratch);
+
+	VkBufferDeviceAddressInfoKHR ScratchDeviceAddressInfo;
+	ZeroVulkanStruct(ScratchDeviceAddressInfo, VK_STRUCTURE_TYPE_BUFFER_DEVICE_ADDRESS_INFO);
+	ScratchDeviceAddressInfo.buffer = Scratch.Buffer;
+	Scratch.Address = vkGetBufferDeviceAddressKHR(Device, &ScratchDeviceAddressInfo);
+
+	VkAccelerationStructureCreateInfoKHR CreateInfo;
+	ZeroVulkanStruct(CreateInfo, VK_STRUCTURE_TYPE_ACCELERATION_STRUCTURE_CREATE_INFO_KHR);
+	CreateInfo.buffer = Allocation.Buffer;
+	CreateInfo.size = BuildSizesInfo.accelerationStructureSize;
+	CreateInfo.type = VK_ACCELERATION_STRUCTURE_TYPE_BOTTOM_LEVEL_KHR;
+	VERIFYVULKANRESULT(vkCreateAccelerationStructureKHR(Device, &CreateInfo, VULKAN_CPU_ALLOCATOR, &Handle));
+
+	BuildGeometryInfo.dstAccelerationStructure = Handle;
+	BuildGeometryInfo.scratchData.deviceAddress = Scratch.Address;
+
+	VkAccelerationStructureBuildRangeInfoKHR* const pBuildRanges = BuildRanges.GetData();
+	FVulkanCmdBuffer* const CmdBuffer = CommandContext.GetCommandBufferManager()->GetActiveCmdBuffer();
+	vkCmdBuildAccelerationStructuresKHR(CmdBuffer->GetHandle(), 1, &BuildGeometryInfo, &pBuildRanges);
+
+	CommandContext.GetCommandBufferManager()->SubmitActiveCmdBuffer();
+
+	VkAccelerationStructureDeviceAddressInfoKHR DeviceAddressInfo;
+	ZeroVulkanStruct(DeviceAddressInfo, VK_STRUCTURE_TYPE_ACCELERATION_STRUCTURE_DEVICE_ADDRESS_INFO_KHR);
+	DeviceAddressInfo.accelerationStructure = Handle;
+	Allocation.Address = vkGetAccelerationStructureDeviceAddressKHR(Device, &DeviceAddressInfo);
+
+	// No longer need scratch memory for a static build
+	if (!bAllowUpdate)
+	{
+		Free(Scratch);
+	}
+}
+
+// TLAS Todo
+FVulkanRayTracingScene::FVulkanRayTracingScene(/*FD3D12Adapter* Adapter,*/ const FRayTracingSceneInitializer& Initializer)
+{
+
+}
+
+FVulkanRayTracingScene::~FVulkanRayTracingScene()
+{
+
+}
+
+void FVulkanRayTracingScene::BuildAccelerationStructure()
+{
+
+}
+
+FRayTracingSceneRHIRef FVulkanDynamicRHI::RHICreateRayTracingScene(const FRayTracingSceneInitializer& Initializer)
+{
+	return new FVulkanRayTracingScene(Initializer);
+}
+// TLAS todo
+
+FRayTracingGeometryRHIRef FVulkanDynamicRHI::RHICreateRayTracingGeometry(const FRayTracingGeometryInitializer& Initializer)
+{
+	return new FVulkanRayTracingGeometry(Initializer);
+}
+
+// Todo: High level rhi call should have transitioned and verified vb and ib to read for each segment
+void FVulkanCommandListContext::RHIBuildAccelerationStructures(const TArrayView<const FAccelerationStructureBuildParams> Params)
+{
+	for (const FAccelerationStructureBuildParams& P : Params)
+	{
+		FVulkanRayTracingGeometry* const Geometry = ResourceCast(P.Geometry.GetReference());
+
+		// Todo: Update geometry from params for each segment
+		// Todo: Can we do this only for an update?
+
+		// Build as for each segment
+		Geometry->BuildAccelerationStructure(*this, P.BuildMode);
+	}
+}
+
+void FVulkanCommandListContext::RHIBuildAccelerationStructure(FRHIRayTracingScene* InScene)
+{
+	FVulkanRayTracingScene* const Scene = ResourceCast(InScene);
+	Scene->BuildAccelerationStructure(/**this*/);
+}
+
 #endif // #if VULKAN_RHI_RAYTRACING
diff --git a/Engine/Source/Runtime/VulkanRHI/Private/VulkanRayTracing.h b/Engine/Source/Runtime/VulkanRHI/Private/VulkanRayTracing.h
--- a/Engine/Source/Runtime/VulkanRHI/Private/VulkanRayTracing.h
+++ b/Engine/Source/Runtime/VulkanRHI/Private/VulkanRayTracing.h
@@ -4,6 +4,8 @@
 
 #if VULKAN_RHI_RAYTRACING
 
+class FVulkanCommandListContext;
+
 #define ENUM_VK_ENTRYPOINTS_RAYTRACING(EnumMacro) \
 	EnumMacro(PFN_vkCreateAccelerationStructureKHR, vkCreateAccelerationStructureKHR) \
 	EnumMacro(PFN_vkDestroyAccelerationStructureKHR, vkDestroyAccelerationStructureKHR) \
@@ -29,4 +31,53 @@ public:
 	static bool LoadVulkanInstanceFunctions(VkInstance inInstance);
 };
 
+struct FAccelerationStructureAllocation
+{
+	VkDevice Device = VK_NULL_HANDLE;
+	VkDeviceMemory Memory = VK_NULL_HANDLE;
+	VkBuffer Buffer = VK_NULL_HANDLE;
+	VkDeviceAddress Address = {};
+};
+
+class FVulkanRayTracingGeometry : public FRHIRayTracingGeometry
+{
+public:
+	FVulkanRayTracingGeometry(const FRayTracingGeometryInitializer& Initializer);
+	~FVulkanRayTracingGeometry();
+
+	void BuildAccelerationStructure(FVulkanCommandListContext& CommandContext, EAccelerationStructureBuildMode BuildMode);
+
+private:
+	uint32 IndexStrideInBytes  = 0;
+	uint32 IndexOffsetInBytes  = 0;
+	uint32 TotalPrimitiveCount = 0;
+	
+	FBufferRHIRef IndexBufferRHI;
+	
+	bool bFastBuild   = false;
+	bool bAllowUpdate = false;
+
+	VkAccelerationStructureKHR Handle = VK_NULL_HANDLE;
+	FAccelerationStructureAllocation Allocation;
+	FAccelerationStructureAllocation Scratch;
+
+	TArray<FRayTracingGeometrySegment> Segments;
+
+	FName DebugName;
+};
+
+// Todo
+class FVulkanRayTracingScene : public FRHIRayTracingScene
+{
+public:
+	FVulkanRayTracingScene(/*FD3D12Adapter* Adapter,*/ const FRayTracingSceneInitializer& Initializer);
+	~FVulkanRayTracingScene();
+
+	void BuildAccelerationStructure();
+
+	TArray<FRayTracingGeometryInstance> Instances;
+	TArray<TRefCountPtr<FVulkanRayTracingGeometry>> Geometries;
+	FName DebugName;
+};
+
 #endif // #if VULKAN_RHI_RAYTRACING
diff --git a/Engine/Source/Runtime/VulkanRHI/Private/VulkanResources.h b/Engine/Source/Runtime/VulkanRHI/Private/VulkanResources.h
--- a/Engine/Source/Runtime/VulkanRHI/Private/VulkanResources.h
+++ b/Engine/Source/Runtime/VulkanRHI/Private/VulkanResources.h
@@ -1645,3 +1645,16 @@ static FORCEINLINE typename TVulkanResourceTraits<TRHIType>::TConcreteType* Reso
 {
 	return static_cast<const typename TVulkanResourceTraits<TRHIType>::TConcreteType*>(Resource);
 }
+
+#if VULKAN_RHI_RAYTRACING
+template<>
+struct TVulkanResourceTraits<FRHIRayTracingScene>
+{
+	typedef FVulkanRayTracingScene TConcreteType;
+};
+template<>
+struct TVulkanResourceTraits<FRHIRayTracingGeometry>
+{
+	typedef FVulkanRayTracingGeometry TConcreteType;
+};
+#endif // VULKAN_RHI_RAYTRACING
diff --git a/Engine/Source/Runtime/VulkanRHI/Public/VulkanDynamicRHI.h b/Engine/Source/Runtime/VulkanRHI/Public/VulkanDynamicRHI.h
--- a/Engine/Source/Runtime/VulkanRHI/Public/VulkanDynamicRHI.h
+++ b/Engine/Source/Runtime/VulkanRHI/Public/VulkanDynamicRHI.h
@@ -336,6 +336,11 @@ public:
 
 	void RHICalibrateTimers() override;
 
+#if VULKAN_RHI_RAYTRACING
+	virtual FRayTracingGeometryRHIRef RHICreateRayTracingGeometry(const FRayTracingGeometryInitializer& Initializer) final override;
+	virtual FRayTracingSceneRHIRef RHICreateRayTracingScene(const FRayTracingSceneInitializer& Initializer) final override;
+#endif
+
 	// Historical number of times we've presented any and all viewports
 	uint32 TotalPresentCount = 0;
 
