UE5EA - Missing handling of RDG texture UAVs on Vulkan

diff --git a/Engine/Source/Runtime/VulkanRHI/Private/VulkanCommands.cpp b/Engine/Source/Runtime/VulkanRHI/Private/VulkanCommands.cpp
--- a/Engine/Source/Runtime/VulkanRHI/Private/VulkanCommands.cpp
+++ b/Engine/Source/Runtime/VulkanRHI/Private/VulkanCommands.cpp
@@ -463,6 +463,7 @@ inline void SetShaderUniformBufferResources(FVulkanCommandListContext* Context,
 			break;
 		}
 
+		case UBMT_RDG_TEXTURE_UAV:
 		case UBMT_UAV:
 		case UBMT_RDG_BUFFER_UAV:
 		{
@@ -490,7 +491,7 @@ inline void SetShaderUniformBufferResources(FVulkanCommandListContext* Context,
 		}
 
 		default:
-			check(0);
+			checkf(0, TEXT("Missing handling for UBMT_ %d"), (int32)ResourceInfo.UBBaseType);
 			break;
 		}
 	}
