UE5EA - When doing raster-only passes with no render targets bound to the pass, use DestRect to describe render area

diff --git a/Engine/Source/Runtime/RHI/Public/RHIResources.h b/Engine/Source/Runtime/RHI/Public/RHIResources.h
--- a/Engine/Source/Runtime/RHI/Public/RHIResources.h
+++ b/Engine/Source/Runtime/RHI/Public/RHIResources.h
@@ -2527,6 +2527,8 @@ struct FRHIRenderPassInfo
 	};
 	FDepthStencilEntry DepthStencilRenderTarget;
 
+	// Parameters for resolving a multisampled image
+	// When doing raster-only passes with no render targets bound to the pass, use DestRect to describe render area
 	FResolveParams ResolveParameters;
 
 	// Some RHIs can use a texture to control the sampling and/or shading resolution of different areas 
diff --git a/Engine/Source/Runtime/VulkanRHI/Private/VulkanCommandBuffer.cpp b/Engine/Source/Runtime/VulkanRHI/Private/VulkanCommandBuffer.cpp
--- a/Engine/Source/Runtime/VulkanRHI/Private/VulkanCommandBuffer.cpp
+++ b/Engine/Source/Runtime/VulkanRHI/Private/VulkanCommandBuffer.cpp
@@ -177,10 +177,7 @@ void FVulkanCmdBuffer::BeginRenderPass(const FVulkanRenderTargetLayout& Layout,
 	ZeroVulkanStruct(Info, VK_STRUCTURE_TYPE_RENDER_PASS_BEGIN_INFO);
 	Info.renderPass = RenderPass->GetHandle();
 	Info.framebuffer = Framebuffer->GetHandle();
-	Info.renderArea.offset.x = 0;
-	Info.renderArea.offset.y = 0;
-	Info.renderArea.extent.width = Framebuffer->GetWidth();
-	Info.renderArea.extent.height = Framebuffer->GetHeight();
+	Info.renderArea = Framebuffer->GetRenderArea();
 	Info.clearValueCount = Layout.GetNumUsedClearValues();
 	Info.pClearValues = AttachmentClearValues;
 
diff --git a/Engine/Source/Runtime/VulkanRHI/Private/VulkanRHIPrivate.h b/Engine/Source/Runtime/VulkanRHI/Private/VulkanRHIPrivate.h
--- a/Engine/Source/Runtime/VulkanRHI/Private/VulkanRHIPrivate.h
+++ b/Engine/Source/Runtime/VulkanRHI/Private/VulkanRHIPrivate.h
@@ -176,15 +176,17 @@ protected:
 	ESubpassHint SubpassHint = ESubpassHint::None;
 	uint8 MultiViewCount;
 
-	uint8 Pad0 = 0;
-	uint8 Pad1 = 0;
-	uint8 Pad2 = 0;
-
 	// Hash for a compatible RenderPass
 	uint32 RenderPassCompatibleHash = 0;
 	// Hash for the render pass including the load/store operations
 	uint32 RenderPassFullHash = 0;
 
+	union
+	{
+		VkOffset3D Offset3D;
+		VkOffset3D Offset2D;
+	} Offset;
+
 	union
 	{
 		VkExtent3D	Extent3D;
@@ -271,19 +273,14 @@ public:
 		return (DepthStencilRenderTargetImage == Image);
 	}
 
-	inline uint32 GetWidth() const
-	{
-		return Extents.width;
-	}
-
-	inline uint32 GetHeight() const
+	inline VkRect2D GetRenderArea() const
 	{
-		return Extents.height;
+		return RenderArea;
 	}
 
 private:
 	VkFramebuffer Framebuffer;
-	VkExtent2D Extents;
+	VkRect2D RenderArea;
 
 	// Unadjusted number of color render targets as in FRHISetRenderTargetsInfo 
 	uint32 NumColorRenderTargets;
diff --git a/Engine/Source/Runtime/VulkanRHI/Private/VulkanRenderTarget.cpp b/Engine/Source/Runtime/VulkanRHI/Private/VulkanRenderTarget.cpp
--- a/Engine/Source/Runtime/VulkanRHI/Private/VulkanRenderTarget.cpp
+++ b/Engine/Source/Runtime/VulkanRHI/Private/VulkanRenderTarget.cpp
@@ -1138,6 +1138,7 @@ FVulkanRenderTargetLayout::FVulkanRenderTargetLayout(FVulkanDevice& InDevice, co
 	FMemory::Memzero(ResolveReferences);
 	FMemory::Memzero(InputAttachments);
 	FMemory::Memzero(Desc);
+	FMemory::Memzero(Offset);
 	FMemory::Memzero(Extent);
 
 	FRenderPassCompatibleHashableStruct CompatibleHashInfo;
@@ -1338,6 +1339,18 @@ FVulkanRenderTargetLayout::FVulkanRenderTargetLayout(FVulkanDevice& InDevice, co
 			Extent.Extent3D.depth = Texture->Surface.Depth;
 		}
 	}
+	else if (NumColorRenderTargets == 0)
+	{
+		// No Depth and no color, it's a raster-only pass so make sure the renderArea will be set up properly
+		checkf(RPInfo.ResolveParameters.DestRect.IsValid(), TEXT("For raster-only passes without render targets, ResolveParameters.DestRect has to contain the render area"));
+		bSetExtent = true;
+		Offset.Offset3D.x = RPInfo.ResolveParameters.DestRect.X1;
+		Offset.Offset3D.y = RPInfo.ResolveParameters.DestRect.Y1;
+		Offset.Offset3D.z = 0;
+		Extent.Extent3D.width = RPInfo.ResolveParameters.DestRect.X2 - RPInfo.ResolveParameters.DestRect.X1;
+		Extent.Extent3D.height = RPInfo.ResolveParameters.DestRect.Y2 - RPInfo.ResolveParameters.DestRect.Y1;
+		Extent.Extent3D.depth = 1;
+	}
 
 	if (InDevice.GetOptionalExtensions().HasEXTFragmentDensityMap && RPInfo.ShadingRateTexture)
 	{
@@ -1413,7 +1426,6 @@ FVulkanRenderTargetLayout::FVulkanRenderTargetLayout(const FGraphicsPipelineStat
 	FRenderPassCompatibleHashableStruct CompatibleHashInfo;
 	FRenderPassFullHashableStruct FullHashInfo;
 
-	bool bSetExtent = false;
 	bool bFoundClearOp = false;
 	MultiViewCount = Initializer.MultiViewCount;
 	NumSamples = Initializer.NumSamples;
diff --git a/Engine/Source/Runtime/VulkanRHI/Private/VulkanViewport.cpp b/Engine/Source/Runtime/VulkanRHI/Private/VulkanViewport.cpp
--- a/Engine/Source/Runtime/VulkanRHI/Private/VulkanViewport.cpp
+++ b/Engine/Source/Runtime/VulkanRHI/Private/VulkanViewport.cpp
@@ -470,8 +470,10 @@ FVulkanFramebuffer::FVulkanFramebuffer(FVulkanDevice& Device, const FRHISetRende
 
 	VERIFYVULKANRESULT_EXPANDED(VulkanRHI::vkCreateFramebuffer(Device.GetInstanceHandle(), &CreateInfo, VULKAN_CPU_ALLOCATOR, &Framebuffer));
 
-	Extents.width = RTExtents.width;
-	Extents.height = RTExtents.height;
+	RenderArea.offset.x = 0;
+	RenderArea.offset.y = 0;
+	RenderArea.extent.width = RTExtents.width;
+	RenderArea.extent.height = RTExtents.height;
 
 	INC_DWORD_STAT(STAT_VulkanNumFrameBuffers);
 }
