Fix indexing error (and crash) when FVulkanImageLayout::Set is applied to a texture array range with a non-zero first layer.

diff --git a/Engine/Source/Runtime/VulkanRHI/Private/VulkanBarriers.cpp b/Engine/Source/Runtime/VulkanRHI/Private/VulkanBarriers.cpp
--- a/Engine/Source/Runtime/VulkanRHI/Private/VulkanBarriers.cpp
+++ b/Engine/Source/Runtime/VulkanRHI/Private/VulkanBarriers.cpp
@@ -1330,8 +1330,9 @@ void FVulkanImageLayout::Set(VkImageLayout Layout, const VkImageSubresourceRange
 
 	if (SubresLayouts.Num() == 0)
 	{
-		SubresLayouts.SetNum(NumLayers * NumMips);
-		for (uint32 i = 0; i < NumLayers * NumMips; ++i)
+		uint32 SubresLayoutCount = (NumLayers + FirstLayer) * NumMips;
+		SubresLayouts.SetNum(SubresLayoutCount);
+		for (uint32 i = 0; i < SubresLayoutCount; ++i)
 		{
 			SubresLayouts[i] = MainLayout;
 		}
