VKRT TLAS build minimal implementation

diff --git a/Engine/Source/Runtime/VulkanRHI/Private/VulkanRayTracing.cpp b/Engine/Source/Runtime/VulkanRHI/Private/VulkanRayTracing.cpp
--- a/Engine/Source/Runtime/VulkanRHI/Private/VulkanRayTracing.cpp
+++ b/Engine/Source/Runtime/VulkanRHI/Private/VulkanRayTracing.cpp
@@ -2,6 +2,7 @@
 
 #include "VulkanRHIPrivate.h"
 #include "VulkanContext.h"
+#include "Experimental/Containers/SherwoodHashTable.h"
 
 #if VULKAN_RHI_RAYTRACING
 
@@ -103,12 +104,12 @@ static uint32 FindMemoryType(VkPhysicalDevice Gpu, uint32 Filter, VkMemoryProper
 }
 
 // Temporary brute force allocation
-static void Allocate(VkPhysicalDevice Gpu, VkDevice Device, VkDeviceSize Size, VkBufferUsageFlags Flags, FAccelerationStructureAllocation& Result)
+static void Allocate(VkPhysicalDevice Gpu, VkDevice Device, VkDeviceSize Size, VkBufferUsageFlags UsageFlags, VkMemoryPropertyFlags MemoryFlags, FVkRtAllocation& Result)
 {
 	VkBufferCreateInfo CreateInfo;
 	ZeroVulkanStruct(CreateInfo, VK_STRUCTURE_TYPE_BUFFER_CREATE_INFO);
 	CreateInfo.size = Size;
-	CreateInfo.usage = Flags;
+	CreateInfo.usage = UsageFlags;
 	VERIFYVULKANRESULT(vkCreateBuffer(Device, &CreateInfo, VULKAN_CPU_ALLOCATOR, &Result.Buffer));
 
 	VkMemoryRequirements MemoryRequirements = {};
@@ -116,13 +117,13 @@ static void Allocate(VkPhysicalDevice Gpu, VkDevice Device, VkDeviceSize Size, V
 
 	VkMemoryAllocateFlagsInfo MemoryAllocateFlagsInfo;
 	ZeroVulkanStruct(MemoryAllocateFlagsInfo, VK_STRUCTURE_TYPE_MEMORY_ALLOCATE_FLAGS_INFO);
-	MemoryAllocateFlagsInfo.flags = VK_MEMORY_ALLOCATE_DEVICE_ADDRESS_BIT_KHR;
+	MemoryAllocateFlagsInfo.flags = MemoryFlags;
 
 	VkMemoryAllocateInfo MemoryAllocateInfo;
 	ZeroVulkanStruct(MemoryAllocateInfo, VK_STRUCTURE_TYPE_MEMORY_ALLOCATE_INFO);
 	MemoryAllocateInfo.pNext = &MemoryAllocateFlagsInfo;
 	MemoryAllocateInfo.allocationSize = MemoryRequirements.size;
-	MemoryAllocateInfo.memoryTypeIndex = FindMemoryType(Gpu, MemoryRequirements.memoryTypeBits, VK_MEMORY_PROPERTY_DEVICE_LOCAL_BIT);
+	MemoryAllocateInfo.memoryTypeIndex = FindMemoryType(Gpu, MemoryRequirements.memoryTypeBits, MemoryFlags);
 	VERIFYVULKANRESULT(vkAllocateMemory(Device, &MemoryAllocateInfo, VULKAN_CPU_ALLOCATOR, &Result.Memory));
 	VERIFYVULKANRESULT(vkBindBufferMemory(Device, Result.Buffer, Result.Memory, 0));
 
@@ -130,7 +131,7 @@ static void Allocate(VkPhysicalDevice Gpu, VkDevice Device, VkDeviceSize Size, V
 }
 
 // Temporary brute force deallocation
-static void Free(FAccelerationStructureAllocation& Allocation)
+static void Free(FVkRtAllocation& Allocation)
 {
 	if (Allocation.Buffer != VK_NULL_HANDLE)
 	{
@@ -263,6 +264,7 @@ void FVulkanRayTracingGeometry::BuildAccelerationStructure(FVulkanCommandListCon
 		Device, 
 		BuildSizesInfo.accelerationStructureSize,
 		VK_BUFFER_USAGE_ACCELERATION_STRUCTURE_STORAGE_BIT_KHR | VK_BUFFER_USAGE_SHADER_DEVICE_ADDRESS_BIT, 
+		VK_MEMORY_ALLOCATE_DEVICE_ADDRESS_BIT_KHR, 
 		Allocation);
 
 	Allocate(
@@ -270,6 +272,7 @@ void FVulkanRayTracingGeometry::BuildAccelerationStructure(FVulkanCommandListCon
 		Device,
 		BuildSizesInfo.buildScratchSize,
 		VK_BUFFER_USAGE_STORAGE_BUFFER_BIT | VK_BUFFER_USAGE_SHADER_DEVICE_ADDRESS_BIT,
+		VK_MEMORY_ALLOCATE_DEVICE_ADDRESS_BIT_KHR, 
 		Scratch);
 
 	VkBufferDeviceAddressInfoKHR ScratchDeviceAddressInfo;
@@ -288,10 +291,13 @@ void FVulkanRayTracingGeometry::BuildAccelerationStructure(FVulkanCommandListCon
 	BuildGeometryInfo.scratchData.deviceAddress = Scratch.Address;
 
 	VkAccelerationStructureBuildRangeInfoKHR* const pBuildRanges = BuildRanges.GetData();
-	FVulkanCmdBuffer* const CmdBuffer = CommandContext.GetCommandBufferManager()->GetActiveCmdBuffer();
+
+	FVulkanCommandBufferManager& CommandBufferManager = *CommandContext.GetCommandBufferManager();
+	FVulkanCmdBuffer* const CmdBuffer = CommandBufferManager.GetActiveCmdBuffer();
 	vkCmdBuildAccelerationStructuresKHR(CmdBuffer->GetHandle(), 1, &BuildGeometryInfo, &pBuildRanges);
 
-	CommandContext.GetCommandBufferManager()->SubmitActiveCmdBuffer();
+	CommandBufferManager.SubmitActiveCmdBuffer();
+	CommandBufferManager.PrepareForNewActiveCommandBuffer();
 
 	VkAccelerationStructureDeviceAddressInfoKHR DeviceAddressInfo;
 	ZeroVulkanStruct(DeviceAddressInfo, VK_STRUCTURE_TYPE_ACCELERATION_STRUCTURE_DEVICE_ADDRESS_INFO_KHR);
@@ -305,27 +311,185 @@ void FVulkanRayTracingGeometry::BuildAccelerationStructure(FVulkanCommandListCon
 	}
 }
 
-// TLAS Todo
-FVulkanRayTracingScene::FVulkanRayTracingScene(/*FD3D12Adapter* Adapter,*/ const FRayTracingSceneInitializer& Initializer)
+FVulkanRayTracingScene::FVulkanRayTracingScene(const FRayTracingSceneInitializer& Initializer)
 {
+	DebugName = Initializer.DebugName;
+
+	uint32 NumInstances = 0;
+	Experimental::TSherwoodSet<const FVulkanRayTracingGeometry*> UniqueGeometryReferences;
+
+	for (const FRayTracingGeometryInstance& SceneInstance : Initializer.Instances)
+	{
+		// Don't support GPU transforms yet.
+		ensure(!SceneInstance.GPUTransformsSRV.IsValid());
+
+		// Hold a reference to each unique geometry object in the scene
+		const FVulkanRayTracingGeometry* const Geometry = ResourceCast(SceneInstance.GeometryRHI);
+		bool bIsAlreadyInSet = false;
+		UniqueGeometryReferences.Add(Geometry, &bIsAlreadyInSet);
+		if (!bIsAlreadyInSet)
+		{
+			ReferencedGeometry.Add(Geometry);
+		}
+
+		NumInstances += SceneInstance.NumTransforms;
+	}
+
+	uint32 InstanceDescIndex = 0;
+	InstanceDescs.SetNumUninitialized(NumInstances);
+	InstanceGeometry.SetNumUninitialized(NumInstances);
+	for (const FRayTracingGeometryInstance& SceneInstance : Initializer.Instances)
+	{
+		const bool bUseUniqueUserData = SceneInstance.UserData.Num() > 1;
+		const uint32 CommonUserData = SceneInstance.UserData.Num() == 1 ? SceneInstance.UserData[0] : 0;
+
+		TArrayView<const FMatrix> Transforms = SceneInstance.GetTransforms();
+		for (int32 TransformIndex = 0; TransformIndex < Transforms.Num(); ++TransformIndex)
+		{
+			InstanceGeometry[InstanceDescIndex] = ResourceCast(SceneInstance.GeometryRHI);
+
+			VkAccelerationStructureInstanceKHR& InstanceDesc = InstanceDescs[InstanceDescIndex];
+			const FMatrix TransposedTransform = Transforms[TransformIndex].GetTransposed();
+			FMemory::Memcpy(&InstanceDesc.transform, &TransposedTransform.M[0][0], sizeof(VkTransformMatrixKHR));
+
+			InstanceDesc.accelerationStructureReference = 0; // Set during TLAS build after the BLAS for the referenced geometry is built.
+			InstanceDesc.instanceCustomIndex = (bUseUniqueUserData) ? SceneInstance.UserData[TransformIndex] : CommonUserData;
+			InstanceDesc.mask = SceneInstance.Mask;
+			InstanceDesc.instanceShaderBindingTableRecordOffset = 0; // Todo
+			InstanceDesc.flags = VK_GEOMETRY_INSTANCE_TRIANGLE_FRONT_COUNTERCLOCKWISE_BIT_KHR;
 
+			if (SceneInstance.bForceOpaque)
+			{
+				InstanceDesc.flags |= VK_GEOMETRY_INSTANCE_FORCE_OPAQUE_BIT_KHR;
+			}
+
+			if (SceneInstance.bDoubleSided)
+			{
+				InstanceDesc.flags |= VK_GEOMETRY_INSTANCE_TRIANGLE_FACING_CULL_DISABLE_BIT_KHR;
+			}
+
+			++InstanceDescIndex;
+		}
+	}
 }
 
 FVulkanRayTracingScene::~FVulkanRayTracingScene()
 {
-
+	Free(Allocation);
 }
 
-void FVulkanRayTracingScene::BuildAccelerationStructure()
+void FVulkanRayTracingScene::BuildAccelerationStructure(FVulkanCommandListContext& CommandContext)
 {
+	const VkDevice Device = CommandContext.GetDevice()->GetInstanceHandle();
+	const uint32 NumInstances = static_cast<uint32>(InstanceDescs.Num());
+	const uint32 InstanceBufferByteSize = NumInstances * sizeof(VkAccelerationStructureInstanceKHR);
+	FVkRtAllocation InstanceBuffer;
 
+	Allocate(
+		CommandContext.GetDevice()->GetPhysicalHandle(),
+		Device,
+		InstanceBufferByteSize,
+		VK_BUFFER_USAGE_SHADER_DEVICE_ADDRESS_BIT | VK_BUFFER_USAGE_ACCELERATION_STRUCTURE_BUILD_INPUT_READ_ONLY_BIT_KHR,
+		VK_MEMORY_PROPERTY_HOST_VISIBLE_BIT | VK_MEMORY_PROPERTY_HOST_COHERENT_BIT,
+		InstanceBuffer);
+
+	void* pMappedInstanceBufferMemory = nullptr;
+	VERIFYVULKANRESULT(vkMapMemory(Device, InstanceBuffer.Memory, 0, InstanceBufferByteSize, 0, &pMappedInstanceBufferMemory));
+	{
+		VkAccelerationStructureInstanceKHR* const InstanceDescBuffer = reinterpret_cast<VkAccelerationStructureInstanceKHR*>(pMappedInstanceBufferMemory);
+		FMemory::Memcpy(InstanceDescBuffer, InstanceDescs.GetData(), InstanceBufferByteSize);
+
+		for (int32 InstanceIndex = 0; InstanceIndex < InstanceDescs.Num(); ++InstanceIndex)
+		{
+			InstanceDescBuffer[InstanceIndex].accelerationStructureReference = InstanceGeometry[InstanceIndex]->GetAccelerationStructureAddress();
+		}
+	}
+	vkUnmapMemory(Device, InstanceBuffer.Memory);
+
+	InstanceDescs.Empty();
+	InstanceGeometry.Empty();
+
+	VkDeviceOrHostAddressConstKHR InstanceBufferDeviceAddress = {};
+	InstanceBufferDeviceAddress.deviceAddress = GetDeviceAddress(Device, InstanceBuffer.Buffer);
+
+	VkAccelerationStructureGeometryKHR TLASGeometry;
+	ZeroVulkanStruct(TLASGeometry, VK_STRUCTURE_TYPE_ACCELERATION_STRUCTURE_GEOMETRY_KHR);
+	TLASGeometry.geometryType = VK_GEOMETRY_TYPE_INSTANCES_KHR;
+	TLASGeometry.geometry.instances.sType = VK_STRUCTURE_TYPE_ACCELERATION_STRUCTURE_GEOMETRY_INSTANCES_DATA_KHR;
+	TLASGeometry.geometry.instances.arrayOfPointers = VK_FALSE;
+	TLASGeometry.geometry.instances.data = InstanceBufferDeviceAddress;
+
+	VkAccelerationStructureBuildGeometryInfoKHR TLASGeometryInfo;
+	ZeroVulkanStruct(TLASGeometryInfo, VK_STRUCTURE_TYPE_ACCELERATION_STRUCTURE_BUILD_GEOMETRY_INFO_KHR);
+	TLASGeometryInfo.type = VK_ACCELERATION_STRUCTURE_TYPE_TOP_LEVEL_KHR;
+	TLASGeometryInfo.mode = VK_BUILD_ACCELERATION_STRUCTURE_MODE_BUILD_KHR;
+	TLASGeometryInfo.flags = VK_BUILD_ACCELERATION_STRUCTURE_PREFER_FAST_TRACE_BIT_KHR;
+	TLASGeometryInfo.geometryCount = 1;
+	TLASGeometryInfo.pGeometries = &TLASGeometry;
+
+	VkAccelerationStructureBuildSizesInfoKHR TLASBuildSizesInfo;
+	ZeroVulkanStruct(TLASBuildSizesInfo, VK_STRUCTURE_TYPE_ACCELERATION_STRUCTURE_BUILD_SIZES_INFO_KHR);
+	vkGetAccelerationStructureBuildSizesKHR(
+		Device, 
+		VK_ACCELERATION_STRUCTURE_BUILD_TYPE_DEVICE_KHR, 
+		&TLASGeometryInfo, 
+		&NumInstances, 
+		&TLASBuildSizesInfo);
+
+	Allocate(
+		CommandContext.GetDevice()->GetPhysicalHandle(),
+		Device,
+		TLASBuildSizesInfo.accelerationStructureSize,
+		VK_BUFFER_USAGE_ACCELERATION_STRUCTURE_STORAGE_BIT_KHR | VK_BUFFER_USAGE_SHADER_DEVICE_ADDRESS_BIT,
+		VK_MEMORY_ALLOCATE_DEVICE_ADDRESS_BIT_KHR,
+		Allocation);
+
+	Allocate(
+		CommandContext.GetDevice()->GetPhysicalHandle(),
+		Device,
+		TLASBuildSizesInfo.buildScratchSize,
+		VK_BUFFER_USAGE_STORAGE_BUFFER_BIT | VK_BUFFER_USAGE_SHADER_DEVICE_ADDRESS_BIT,
+		VK_MEMORY_ALLOCATE_DEVICE_ADDRESS_BIT_KHR,
+		Scratch);
+
+	VkBufferDeviceAddressInfoKHR ScratchDeviceAddressInfo;
+	ZeroVulkanStruct(ScratchDeviceAddressInfo, VK_STRUCTURE_TYPE_BUFFER_DEVICE_ADDRESS_INFO);
+	ScratchDeviceAddressInfo.buffer = Scratch.Buffer;
+	Scratch.Address = vkGetBufferDeviceAddressKHR(Device, &ScratchDeviceAddressInfo);
+
+	VkAccelerationStructureCreateInfoKHR TLASCreateInfo;
+	ZeroVulkanStruct(TLASCreateInfo, VK_STRUCTURE_TYPE_ACCELERATION_STRUCTURE_CREATE_INFO_KHR);
+	TLASCreateInfo.buffer = Allocation.Buffer;
+	TLASCreateInfo.size = TLASBuildSizesInfo.accelerationStructureSize;
+	TLASCreateInfo.type = VK_ACCELERATION_STRUCTURE_TYPE_TOP_LEVEL_KHR;
+	VERIFYVULKANRESULT(vkCreateAccelerationStructureKHR(Device, &TLASCreateInfo, VULKAN_CPU_ALLOCATOR, &Handle));
+
+	TLASGeometryInfo.dstAccelerationStructure = Handle;
+	TLASGeometryInfo.scratchData.deviceAddress = Scratch.Address;
+
+	VkAccelerationStructureBuildRangeInfoKHR TLASBuildRangeInfo;
+	TLASBuildRangeInfo.primitiveCount = NumInstances;
+	TLASBuildRangeInfo.primitiveOffset = 0;
+	TLASBuildRangeInfo.transformOffset = 0;
+	TLASBuildRangeInfo.firstVertex = 0;
+
+	VkAccelerationStructureBuildRangeInfoKHR* const pBuildRanges = &TLASBuildRangeInfo;
+
+	FVulkanCommandBufferManager& CommandBufferManager = *CommandContext.GetCommandBufferManager();
+	FVulkanCmdBuffer* const CmdBuffer = CommandBufferManager.GetActiveCmdBuffer();
+	vkCmdBuildAccelerationStructuresKHR(CmdBuffer->GetHandle(), 1, &TLASGeometryInfo, &pBuildRanges);
+
+	CommandBufferManager.SubmitActiveCmdBuffer();
+	CommandBufferManager.PrepareForNewActiveCommandBuffer();
+
+	Free(InstanceBuffer);
+	Free(Scratch);
 }
 
 FRayTracingSceneRHIRef FVulkanDynamicRHI::RHICreateRayTracingScene(const FRayTracingSceneInitializer& Initializer)
 {
 	return new FVulkanRayTracingScene(Initializer);
 }
-// TLAS todo
 
 FRayTracingGeometryRHIRef FVulkanDynamicRHI::RHICreateRayTracingGeometry(const FRayTracingGeometryInitializer& Initializer)
 {
@@ -350,7 +514,7 @@ void FVulkanCommandListContext::RHIBuildAccelerationStructures(const TArrayView<
 void FVulkanCommandListContext::RHIBuildAccelerationStructure(FRHIRayTracingScene* InScene)
 {
 	FVulkanRayTracingScene* const Scene = ResourceCast(InScene);
-	Scene->BuildAccelerationStructure(/**this*/);
+	Scene->BuildAccelerationStructure(*this);
 }
 
 #endif // #if VULKAN_RHI_RAYTRACING
diff --git a/Engine/Source/Runtime/VulkanRHI/Private/VulkanRayTracing.h b/Engine/Source/Runtime/VulkanRHI/Private/VulkanRayTracing.h
--- a/Engine/Source/Runtime/VulkanRHI/Private/VulkanRayTracing.h
+++ b/Engine/Source/Runtime/VulkanRHI/Private/VulkanRayTracing.h
@@ -31,7 +31,7 @@ public:
 	static bool LoadVulkanInstanceFunctions(VkInstance inInstance);
 };
 
-struct FAccelerationStructureAllocation
+struct FVkRtAllocation
 {
 	VkDevice Device = VK_NULL_HANDLE;
 	VkDeviceMemory Memory = VK_NULL_HANDLE;
@@ -46,6 +46,11 @@ public:
 	~FVulkanRayTracingGeometry();
 
 	void BuildAccelerationStructure(FVulkanCommandListContext& CommandContext, EAccelerationStructureBuildMode BuildMode);
+	
+	VkDeviceAddress GetAccelerationStructureAddress() const
+	{
+		return Allocation.Address;
+	}
 
 private:
 	uint32 IndexStrideInBytes  = 0;
@@ -58,25 +63,29 @@ private:
 	bool bAllowUpdate = false;
 
 	VkAccelerationStructureKHR Handle = VK_NULL_HANDLE;
-	FAccelerationStructureAllocation Allocation;
-	FAccelerationStructureAllocation Scratch;
+	FVkRtAllocation Allocation;
+	FVkRtAllocation Scratch;
 
 	TArray<FRayTracingGeometrySegment> Segments;
 
 	FName DebugName;
 };
 
-// Todo
 class FVulkanRayTracingScene : public FRHIRayTracingScene
 {
 public:
-	FVulkanRayTracingScene(/*FD3D12Adapter* Adapter,*/ const FRayTracingSceneInitializer& Initializer);
+	FVulkanRayTracingScene(const FRayTracingSceneInitializer& Initializer);
 	~FVulkanRayTracingScene();
 
-	void BuildAccelerationStructure();
+	void BuildAccelerationStructure(FVulkanCommandListContext& CommandContext);
+
+	VkAccelerationStructureKHR Handle = VK_NULL_HANDLE;
+	FVkRtAllocation Allocation;
+	FVkRtAllocation Scratch;
 
-	TArray<FRayTracingGeometryInstance> Instances;
-	TArray<TRefCountPtr<FVulkanRayTracingGeometry>> Geometries;
+	TArray<VkAccelerationStructureInstanceKHR> InstanceDescs;
+	TArray<const FVulkanRayTracingGeometry*> InstanceGeometry;
+	TArray<TRefCountPtr<const FVulkanRayTracingGeometry>> ReferencedGeometry;
 	FName DebugName;
 };
 
