UE5EA - Enable GPU-assisted validation on Vulkan when regular validation is enabled
Also added support for -vulkanbestpractices

diff --git a/Engine/Source/Runtime/VulkanRHI/Private/VulkanDebug.cpp b/Engine/Source/Runtime/VulkanRHI/Private/VulkanDebug.cpp
--- a/Engine/Source/Runtime/VulkanRHI/Private/VulkanDebug.cpp
+++ b/Engine/Source/Runtime/VulkanRHI/Private/VulkanDebug.cpp
@@ -6,7 +6,7 @@
 
 #include "VulkanRHIPrivate.h"
 
-FAutoConsoleVariable GCVarUniqueValidationMessages(
+static FAutoConsoleVariable GCVarUniqueValidationMessages(
 	TEXT("r.Vulkan.UniqueValidationMessages"),
 	1,
 	TEXT("Filter out validation errors with the same code (only when r.Vulkan.EnableValidation is non zero)")
@@ -276,6 +276,11 @@ static VkBool32 DebugUtilsCallback(VkDebugUtilsMessageSeverityFlagBitsEXT MsgSev
 		// *** [Error:Validation(UNASSIGNED-GPU-Assisted Validation Setup Error.)] Unable to reserve descriptor binding slot on a device with only one slot.
 		return VK_FALSE;
 	}
+	else if (!FCStringAnsi::Strcmp(CallbackData->pMessageIdName, "UNASSIGNED-BestPractices-vkCreateDevice-deprecated-extension"))
+	{
+		// *** CreateDevice(): Attempting to enable deprecated extension VK_KHR_get_memory_requirements2, but this extension has been promoted to VK_VERSION_1_1.
+		return VK_FALSE;
+	}
 
 	static TStaticArray<TSet<uint32>, (int32)EMsgBucket::Count> SeenCodes;
 	bool bUniqueMessages = GCVarUniqueValidationMessages->GetInt() != 0;
diff --git a/Engine/Source/Runtime/VulkanRHI/Private/VulkanLayers.cpp b/Engine/Source/Runtime/VulkanRHI/Private/VulkanLayers.cpp
--- a/Engine/Source/Runtime/VulkanRHI/Private/VulkanLayers.cpp
+++ b/Engine/Source/Runtime/VulkanRHI/Private/VulkanLayers.cpp
@@ -47,10 +47,10 @@ static TAutoConsoleVariable<int32> GStandardValidationCvar(
 
 TAutoConsoleVariable<int32> GGPUValidationCvar(
 	TEXT("r.Vulkan.GPUValidation"),
-	0,
-	TEXT("2 to use enable GPU assised validation AND extra binding slot when using validation layers, or\n")
-	TEXT("1 to use enable GPU assised validation when using validation layers, or\n")
-	TEXT("0 to not use (default)"),
+	2,
+	TEXT("2 to use enable GPU assisted validation AND extra binding slot when using validation layers (default)\n")
+	TEXT("1 to use enable GPU assisted validation when using validation layers, or\n")
+	TEXT("0 to not use"),
 	ECVF_ReadOnly | ECVF_RenderThreadSafe
 );
 
diff --git a/Engine/Source/Runtime/VulkanRHI/Private/VulkanRHI.cpp b/Engine/Source/Runtime/VulkanRHI/Private/VulkanRHI.cpp
--- a/Engine/Source/Runtime/VulkanRHI/Private/VulkanRHI.cpp
+++ b/Engine/Source/Runtime/VulkanRHI/Private/VulkanRHI.cpp
@@ -85,14 +85,23 @@ extern TAutoConsoleVariable<int32> GRHIAllowAsyncComputeCvar;
 uint32 GVulkanDeviceShaderStageBits = 0;
 
 #if VULKAN_HAS_VALIDATION_FEATURES
-static inline TArray<VkValidationFeatureEnableEXT> GetValidationFeaturesEnabled()
+static inline TArray<VkValidationFeatureEnableEXT> GetValidationFeaturesEnabled(bool bEnableValidation)
 {
 	TArray<VkValidationFeatureEnableEXT> Features;
-	Features.Add(VK_VALIDATION_FEATURE_ENABLE_GPU_ASSISTED_EXT);
 	extern TAutoConsoleVariable<int32> GGPUValidationCvar;
-	if (GGPUValidationCvar.GetValueOnAnyThread() > 1)
+	int32 GPUValidationValue = GGPUValidationCvar.GetValueOnAnyThread();
+	if (bEnableValidation && GPUValidationValue > 0)
 	{
-		Features.Add(VK_VALIDATION_FEATURE_ENABLE_GPU_ASSISTED_RESERVE_BINDING_SLOT_EXT);
+		Features.Add(VK_VALIDATION_FEATURE_ENABLE_GPU_ASSISTED_EXT);
+		if (GPUValidationValue > 1)
+		{
+			Features.Add(VK_VALIDATION_FEATURE_ENABLE_GPU_ASSISTED_RESERVE_BINDING_SLOT_EXT);
+		}
+	}
+
+	if (FParse::Param(FCommandLine::Get(), TEXT("vulkanbestpractices")))
+	{
+		Features.Add(VK_VALIDATION_FEATURE_ENABLE_BEST_PRACTICES_EXT);
 	}
 
 	return Features;
@@ -401,7 +410,7 @@ void FVulkanDynamicRHI::CreateInstance()
 			return Key && !FCStringAnsi::Strcmp(Key, VK_EXT_VALIDATION_FEATURES_EXTENSION_NAME);
 		});
 	VkValidationFeaturesEXT ValidationFeatures;
-	TArray<VkValidationFeatureEnableEXT> ValidationFeaturesEnabled = GetValidationFeaturesEnabled();
+	TArray<VkValidationFeatureEnableEXT> ValidationFeaturesEnabled = GetValidationFeaturesEnabled(bHasGPUValidation);
 	if (bHasGPUValidation)
 	{
 		ZeroVulkanStruct(ValidationFeatures, VK_STRUCTURE_TYPE_VALIDATION_FEATURES_EXT);
diff --git a/Engine/Source/Runtime/VulkanRHI/Private/VulkanRHIPrivate.h b/Engine/Source/Runtime/VulkanRHI/Private/VulkanRHIPrivate.h
--- a/Engine/Source/Runtime/VulkanRHI/Private/VulkanRHIPrivate.h
+++ b/Engine/Source/Runtime/VulkanRHI/Private/VulkanRHIPrivate.h
@@ -852,6 +852,8 @@ inline bool UseVulkanDescriptorCache()
 
 extern int32 GVulkanSubmitAfterEveryEndRenderPass;
 extern int32 GWaitForIdleOnSubmit;
+
+// Vendor-specific GPU crash dumps
 extern bool GGPUCrashDebuggingEnabled;
 
 #if VULKAN_HAS_DEBUGGING_ENABLED
