UE5EA - Add support for -gpuvalidation on Vulkan and not enable it by default when using -vulkandebug
- Removed VULKAN_ENABLE_API_DUMP and instead allow dumping using -vulkanapidump
- Setup api_dump layer so the vk_layer_settings.txt file next to the .exe is no longer needed

diff --git a/Engine/Source/Runtime/VulkanRHI/Private/VulkanLayers.cpp b/Engine/Source/Runtime/VulkanRHI/Private/VulkanLayers.cpp
--- a/Engine/Source/Runtime/VulkanRHI/Private/VulkanLayers.cpp
+++ b/Engine/Source/Runtime/VulkanRHI/Private/VulkanLayers.cpp
@@ -47,10 +47,10 @@ static TAutoConsoleVariable<int32> GStandardValidationCvar(
 
 TAutoConsoleVariable<int32> GGPUValidationCvar(
 	TEXT("r.Vulkan.GPUValidation"),
-	2,
-	TEXT("2 to use enable GPU assisted validation AND extra binding slot when using validation layers (default)\n")
+	0,
+	TEXT("2 to use enable GPU assisted validation AND extra binding slot when using validation layers\n")
 	TEXT("1 to use enable GPU assisted validation when using validation layers, or\n")
-	TEXT("0 to not use"),
+	TEXT("0 to not use (default)"),
 	ECVF_ReadOnly | ECVF_RenderThreadSafe
 );
 
@@ -335,21 +335,30 @@ void FVulkanDynamicRHI::GetInstanceLayersAndExtensions(TArray<const ANSICHAR*>&
 	}
 
 #if VULKAN_HAS_DEBUGGING_ENABLED
-#if VULKAN_ENABLE_API_DUMP
-	if (!bGfxReconstructOrVkTrace)
+	if (FParse::Param(FCommandLine::Get(), TEXT("vulkanapidump")))
 	{
-		const char* VkApiDumpName = "VK_LAYER_LUNARG_api_dump";
-		bool bApiDumpFound = FindLayerInList(GlobalLayerExtensions, VkApiDumpName);
-		if (bApiDumpFound)
+		if (bGfxReconstructOrVkTrace)
 		{
-			OutInstanceLayers.Add(VkApiDumpName);
+			UE_LOG(LogVulkanRHI, Warning, TEXT("Can't enable api_dump when GfxReconstruct/VkTrace is enabled"));
 		}
 		else
 		{
-			UE_LOG(LogVulkanRHI, Warning, TEXT("Unable to find Vulkan instance layer %s"), ANSI_TO_TCHAR(VkApiDumpName));
+			const char* VkApiDumpName = "VK_LAYER_LUNARG_api_dump";
+			bool bApiDumpFound = FindLayerInList(GlobalLayerExtensions, VkApiDumpName);
+			if (bApiDumpFound)
+			{
+				OutInstanceLayers.Add(VkApiDumpName);
+				FPlatformMisc::SetEnvironmentVar(TEXT("VK_APIDUMP_LOG_FILENAME"), TEXT("vk_apidump.txt"));
+				FPlatformMisc::SetEnvironmentVar(TEXT("VK_APIDUMP_DETAILED"), TEXT("true"));
+				FPlatformMisc::SetEnvironmentVar(TEXT("VK_APIDUMP_FLUSH"), TEXT("true"));
+				FPlatformMisc::SetEnvironmentVar(TEXT("VK_APIDUMP_OUTPUT_FORMAT"), TEXT("text"));
+			}
+			else
+			{
+				UE_LOG(LogVulkanRHI, Warning, TEXT("Unable to find Vulkan instance layer %s"), ANSI_TO_TCHAR(VkApiDumpName));
+			}
 		}
 	}
-#endif	// VULKAN_ENABLE_API_DUMP
 
 	// At this point the CVar holds the final value
 #if VULKAN_HAS_DEBUGGING_ENABLED
@@ -859,5 +868,14 @@ void FVulkanDynamicRHI::SetupValidationRequests()
 	{
 		GValidationCvar->Set(VulkanValidationOption, ECVF_SetByCommandline);
 	}
+
+	if (FParse::Param(FCommandLine::Get(), TEXT("gpuvalidation")))
+	{
+		if (GValidationCvar->GetInt() < 2)
+		{
+			GValidationCvar->Set(2, ECVF_SetByCommandline);
+		}
+		GGPUValidationCvar->Set(2, ECVF_SetByCommandline);
+	}
 #endif
 }
diff --git a/Engine/Source/Runtime/VulkanRHI/Private/VulkanTexture.cpp b/Engine/Source/Runtime/VulkanRHI/Private/VulkanTexture.cpp
--- a/Engine/Source/Runtime/VulkanRHI/Private/VulkanTexture.cpp
+++ b/Engine/Source/Runtime/VulkanRHI/Private/VulkanTexture.cpp
@@ -2698,17 +2698,14 @@ void FVulkanDynamicRHI::RHIBindDebugLabelName(FRHITexture* TextureRHI, const TCH
 	}
 #endif
 
-#if VULKAN_ENABLE_DUMP_LAYER || VULKAN_ENABLE_API_DUMP
+#if VULKAN_ENABLE_DUMP_LAYER
 	{
 // TODO: this dies in the printf on android. Needs investigation.
 #if !PLATFORM_ANDROID
 		FVulkanTextureBase* Base = (FVulkanTextureBase*)TextureRHI->GetTextureBaseRHI();
 #if VULKAN_ENABLE_DUMP_LAYER
-		VulkanRHI::PrintfBegin
-#elif VULKAN_ENABLE_API_DUMP
-		FPlatformMisc::LowLevelOutputDebugStringf
+		VulkanRHI::PrintfBegin(*FString::Printf(TEXT("vkDebugMarkerSetObjectNameEXT(0x%p=%s)\n"), Base->Surface.Image, Name));
 #endif
-			(*FString::Printf(TEXT("vkDebugMarkerSetObjectNameEXT(0x%p=%s)\n"), Base->Surface.Image, Name));
 #endif
 	}
 #endif
@@ -2736,7 +2733,7 @@ void FVulkanDynamicRHI::RHIBindDebugLabelName(FRHITexture* TextureRHI, const TCH
 
 void FVulkanDynamicRHI::RHIBindDebugLabelName(FRHIUnorderedAccessView* UnorderedAccessViewRHI, const TCHAR* Name)
 {
-#if VULKAN_ENABLE_DUMP_LAYER || VULKAN_ENABLE_API_DUMP
+#if VULKAN_ENABLE_DUMP_LAYER
 	//if (Device->SupportsDebugMarkers())
 	{
 		//if (FRHITexture2D* Tex2d = UnorderedAccessViewRHI->GetTexture2D())
diff --git a/Engine/Source/Runtime/VulkanRHI/Public/VulkanConfiguration.h b/Engine/Source/Runtime/VulkanRHI/Public/VulkanConfiguration.h
--- a/Engine/Source/Runtime/VulkanRHI/Public/VulkanConfiguration.h
+++ b/Engine/Source/Runtime/VulkanRHI/Public/VulkanConfiguration.h
@@ -22,9 +22,6 @@
 
 #define VULKAN_HAS_DEBUGGING_ENABLED							(UE_BUILD_DEBUG || (UE_BUILD_DEVELOPMENT && VULKAN_SHOULD_DEBUG_IN_DEVELOPMENT))
 
-// Enables the VK_LAYER_LUNARG_api_dump layer and the report VK_DEBUG_REPORT_INFORMATION_BIT_EXT flag
-#define VULKAN_ENABLE_API_DUMP									0
-
 #ifndef VULKAN_SHOULD_ENABLE_DRAW_MARKERS
 	#define VULKAN_SHOULD_ENABLE_DRAW_MARKERS					0
 #endif
