Fixes for VKRT validation issues:
Vertex and index buffers must be created with VK_BUFFER_USAGE_SHADER_DEVICE_ADDRESS_BIT set if they are to be used in a BLAS build.
Memory backing a buffer which has VK_BUFFER_USAGE_SHADER_DEVICE_ADDRESS_BIT set needs a VkMemoryAllocateFlagsInfo struct with its flag member set to VK_MEMORY_ALLOCATE_DEVICE_ADDRESS_BIT_KHR on the VkMemoryAllocateInfo pNext chain.
Conservatively allocating all buffers with VK_BUFFER_USAGE_SHADER_DEVICE_ADDRESS_BIT and all allocations with VK_MEMORY_ALLOCATE_DEVICE_ADDRESS_BIT_KHR when vkrt is enabled to minimize code foot print. We don't expect there to be a performance penalty for doing so.

Temp allocation fix:
Don't use VK_MEMORY_ALLOCATE_DEVICE_ADDRESS_BIT_KHR as a memory flag when looking for a memory type index, this was supposed to be VK_MEMORY_PROPERTY_DEVICE_LOCAL_BIT

diff --git a/Engine/Source/Runtime/VulkanRHI/Private/VulkanBuffer.cpp b/Engine/Source/Runtime/VulkanRHI/Private/VulkanBuffer.cpp
--- a/Engine/Source/Runtime/VulkanRHI/Private/VulkanBuffer.cpp
+++ b/Engine/Source/Runtime/VulkanRHI/Private/VulkanBuffer.cpp
@@ -102,6 +102,10 @@ FVulkanResourceMultiBuffer::FVulkanResourceMultiBuffer(FVulkanDevice* InDevice,
 		BufferUsageFlags |= bCPUReadable ? (VK_BUFFER_USAGE_TRANSFER_SRC_BIT | VK_BUFFER_USAGE_TRANSFER_DST_BIT) : 0;
 		BufferUsageFlags |= bCopySource ? VK_BUFFER_USAGE_TRANSFER_SRC_BIT : 0;
 
+#if VULKAN_RHI_RAYTRACING
+		BufferUsageFlags |= VK_BUFFER_USAGE_SHADER_DEVICE_ADDRESS_BIT;
+#endif
+
 		if (bVolatile)
 		{
 			bool bRenderThread = IsInRenderingThread();
diff --git a/Engine/Source/Runtime/VulkanRHI/Private/VulkanMemory.cpp b/Engine/Source/Runtime/VulkanRHI/Private/VulkanMemory.cpp
--- a/Engine/Source/Runtime/VulkanRHI/Private/VulkanMemory.cpp
+++ b/Engine/Source/Runtime/VulkanRHI/Private/VulkanMemory.cpp
@@ -956,6 +956,15 @@ namespace VulkanRHI
 			IncMetaStats(EVulkanAllocationMetaImageRenderTarget, AllocationSize);
 		}
 #endif
+
+#if VULKAN_RHI_RAYTRACING
+		VkMemoryAllocateFlagsInfo MemoryAllocateFlagsInfo;
+		ZeroVulkanStruct(MemoryAllocateFlagsInfo, VK_STRUCTURE_TYPE_MEMORY_ALLOCATE_FLAGS_INFO);
+		MemoryAllocateFlagsInfo.flags = VK_MEMORY_ALLOCATE_DEVICE_ADDRESS_BIT_KHR;
+		MemoryAllocateFlagsInfo.pNext = Info.pNext;
+		Info.pNext = &MemoryAllocateFlagsInfo;
+#endif
+
 		VkDeviceMemory Handle;
 		VkResult Result;
 
diff --git a/Engine/Source/Runtime/VulkanRHI/Private/VulkanRayTracing.cpp b/Engine/Source/Runtime/VulkanRHI/Private/VulkanRayTracing.cpp
--- a/Engine/Source/Runtime/VulkanRHI/Private/VulkanRayTracing.cpp
+++ b/Engine/Source/Runtime/VulkanRHI/Private/VulkanRayTracing.cpp
@@ -117,7 +117,7 @@ static void Allocate(VkPhysicalDevice Gpu, VkDevice Device, VkDeviceSize Size, V
 
 	VkMemoryAllocateFlagsInfo MemoryAllocateFlagsInfo;
 	ZeroVulkanStruct(MemoryAllocateFlagsInfo, VK_STRUCTURE_TYPE_MEMORY_ALLOCATE_FLAGS_INFO);
-	MemoryAllocateFlagsInfo.flags = MemoryFlags;
+	MemoryAllocateFlagsInfo.flags = VK_MEMORY_ALLOCATE_DEVICE_ADDRESS_BIT_KHR;
 
 	VkMemoryAllocateInfo MemoryAllocateInfo;
 	ZeroVulkanStruct(MemoryAllocateInfo, VK_STRUCTURE_TYPE_MEMORY_ALLOCATE_INFO);
@@ -264,7 +264,7 @@ void FVulkanRayTracingGeometry::BuildAccelerationStructure(FVulkanCommandListCon
 		Device, 
 		BuildSizesInfo.accelerationStructureSize,
 		VK_BUFFER_USAGE_ACCELERATION_STRUCTURE_STORAGE_BIT_KHR | VK_BUFFER_USAGE_SHADER_DEVICE_ADDRESS_BIT, 
-		VK_MEMORY_ALLOCATE_DEVICE_ADDRESS_BIT_KHR, 
+		VK_MEMORY_PROPERTY_DEVICE_LOCAL_BIT,
 		Allocation);
 
 	Allocate(
@@ -272,7 +272,7 @@ void FVulkanRayTracingGeometry::BuildAccelerationStructure(FVulkanCommandListCon
 		Device,
 		BuildSizesInfo.buildScratchSize,
 		VK_BUFFER_USAGE_STORAGE_BUFFER_BIT | VK_BUFFER_USAGE_SHADER_DEVICE_ADDRESS_BIT,
-		VK_MEMORY_ALLOCATE_DEVICE_ADDRESS_BIT_KHR, 
+		VK_MEMORY_PROPERTY_DEVICE_LOCAL_BIT,
 		Scratch);
 
 	VkBufferDeviceAddressInfoKHR ScratchDeviceAddressInfo;
@@ -441,7 +441,7 @@ void FVulkanRayTracingScene::BuildAccelerationStructure(FVulkanCommandListContex
 		Device,
 		TLASBuildSizesInfo.accelerationStructureSize,
 		VK_BUFFER_USAGE_ACCELERATION_STRUCTURE_STORAGE_BIT_KHR | VK_BUFFER_USAGE_SHADER_DEVICE_ADDRESS_BIT,
-		VK_MEMORY_ALLOCATE_DEVICE_ADDRESS_BIT_KHR,
+		VK_MEMORY_PROPERTY_DEVICE_LOCAL_BIT,
 		Allocation);
 
 	Allocate(
@@ -449,7 +449,7 @@ void FVulkanRayTracingScene::BuildAccelerationStructure(FVulkanCommandListContex
 		Device,
 		TLASBuildSizesInfo.buildScratchSize,
 		VK_BUFFER_USAGE_STORAGE_BUFFER_BIT | VK_BUFFER_USAGE_SHADER_DEVICE_ADDRESS_BIT,
-		VK_MEMORY_ALLOCATE_DEVICE_ADDRESS_BIT_KHR,
+		VK_MEMORY_PROPERTY_DEVICE_LOCAL_BIT,
 		Scratch);
 
 	VkBufferDeviceAddressInfoKHR ScratchDeviceAddressInfo;
