PR #7795: Simple fix for RT<->RHI thread deadlock (Contributed by [REDACTED_PERSON])

diff --git a/Engine/Source/Runtime/VulkanRHI/Private/VulkanMemory.cpp b/Engine/Source/Runtime/VulkanRHI/Private/VulkanMemory.cpp
--- a/Engine/Source/Runtime/VulkanRHI/Private/VulkanMemory.cpp
+++ b/Engine/Source/Runtime/VulkanRHI/Private/VulkanMemory.cpp
@@ -2564,7 +2564,7 @@ namespace VulkanRHI
 		if(bTryFree)
 		{
 			FScopeLock ScopeLockResourceheap(&GResourceHeapLock);
-			FScopeLock ScopeLock(&CS);
+			FScopeLock ScopeLockResource(&GResourceLock);
 			if(JoinFreeBlocks()) //someone could've allocated memory before we claimed this lock.
 			{
 				Owner->ReleaseSubresourceAllocator(this);
