PR #7781: Pass down calling context during defrag instead of accessing global (Contributed by [REDACTED_PERSON])

diff --git a/Engine/Source/Runtime/VulkanRHI/Private/VulkanBuffer.cpp b/Engine/Source/Runtime/VulkanRHI/Private/VulkanBuffer.cpp
--- a/Engine/Source/Runtime/VulkanRHI/Private/VulkanBuffer.cpp
+++ b/Engine/Source/Runtime/VulkanRHI/Private/VulkanBuffer.cpp
@@ -479,7 +479,7 @@ void FVulkanResourceMultiBuffer::Evict(FVulkanDevice& Device_)
 {
 	checkNoEntry();//Not Implemented, should never be called
 }
-void FVulkanResourceMultiBuffer::Move(FVulkanDevice& Device_, VulkanRHI::FVulkanAllocation& NewAllocation)
+void FVulkanResourceMultiBuffer::Move(FVulkanDevice& Device, FVulkanCommandListContext& Context, VulkanRHI::FVulkanAllocation& NewAllocation)
 {
 	checkNoEntry();//Not Implemented, should never be called
 }
diff --git a/Engine/Source/Runtime/VulkanRHI/Private/VulkanMemory.cpp b/Engine/Source/Runtime/VulkanRHI/Private/VulkanMemory.cpp
--- a/Engine/Source/Runtime/VulkanRHI/Private/VulkanMemory.cpp
+++ b/Engine/Source/Runtime/VulkanRHI/Private/VulkanMemory.cpp
@@ -1801,7 +1801,7 @@ namespace VulkanRHI
 		return Allocator->GetIsDefragging();
 	}
 
-	void FVulkanResourceHeap::DefragTick(FVulkanDevice& Device, uint32 Count)
+	void FVulkanResourceHeap::DefragTick(FVulkanDevice& Device, FVulkanCommandListContext& Context, uint32 Count)
 	{
 		SCOPED_NAMED_EVENT(FVulkanSubresourceAllocator_DefragTick, FColor::Cyan);
 
@@ -1897,7 +1897,7 @@ namespace VulkanRHI
 			{
 				VULKAN_LOGMEMORY(TEXT("Defragging heap [%d] Max %6.2fMB Free %6.2fMB   %2d"), CurrentDefragTarget->AllocatorIndex, MaxSize / (1024.f*1024.f), FreeSpace / (1024.f * 1024.f), DefragCountDown);
 			}
-			if(CurrentDefragTarget->DefragTick(Device, this, Count) > 0)
+			if(CurrentDefragTarget->DefragTick(Device, Context, this, Count) > 0)
 			{
 				DefragCountDown = 3;
 			}
@@ -1941,7 +1941,7 @@ namespace VulkanRHI
 					uint32 Index = rand() % Allocators.Num();
 					PreferredAllocator= Allocators[Index];
 				}
-				PreferredAllocator->DefragFull(Device, this);
+				PreferredAllocator->DefragFull(Device, Context, this);
 
 			}
 		}
@@ -2369,7 +2369,7 @@ namespace VulkanRHI
 		DeviceMemoryManager->TrimMemory(bImmediately);
 	}
 
-	void FMemoryManager::ReleaseFreedPages()
+	void FMemoryManager::ReleaseFreedPages(FVulkanCommandListContext& Context)
 	{
 		auto CanDefragHeap = [](FVulkanResourceHeap* Heap)
 		{
@@ -2460,7 +2460,7 @@ namespace VulkanRHI
 			{
 				DumpMemory();
 			}
-			BestDefragHeap->DefragTick(*Device, Count);
+			BestDefragHeap->DefragTick(*Device, Context, Count);
 
 			if(GVulkanDefragOnce)
 			{
@@ -4016,7 +4016,7 @@ namespace VulkanRHI
 
 	// Fully defrag a page onto a new page
 	// This moves -all- allocations to a new page, and ensures the new page is fully defragged.
-	bool FVulkanSubresourceAllocator::DefragFull(FVulkanDevice& Device, FVulkanResourceHeap* Heap)
+	bool FVulkanSubresourceAllocator::DefragFull(FVulkanDevice& Device, FVulkanCommandListContext& Context, FVulkanResourceHeap* Heap)
 	{
 		TRACE_BOOKMARK(TEXT("DefragFull %d"), AllocatorIndex);
 		SCOPED_NAMED_EVENT(FVulkanSubresourceAllocator_DefragFull, FColor::Cyan);
@@ -4138,7 +4138,7 @@ namespace VulkanRHI
 			Alloc.AllocationOffset = Move.AllocatedOffset;
 			Alloc.AllocationSize = Move.AllocatedSize;
 			FVulkanEvictable* EvictableOwner = Alloc.AllocationOwner;
-			EvictableOwner->OnFullDefrag(Device, Move.AlignedOffset);
+			EvictableOwner->OnFullDefrag(Device, Context, Move.AlignedOffset);
 		}
 
 		FreeList.Empty();
@@ -4178,7 +4178,7 @@ namespace VulkanRHI
 	// One tick of incremental defrag. This will move max Count_ allocations from this suballocator to other SubAllocators on the same Heap
 	// This will eventually free up the page, which can then be released
 	// Returns # of allocations moved.
-	int32 FVulkanSubresourceAllocator::DefragTick(FVulkanDevice& Device, FVulkanResourceHeap* Heap, uint32 Count_)
+	int32 FVulkanSubresourceAllocator::DefragTick(FVulkanDevice& Device, FVulkanCommandListContext& Context, FVulkanResourceHeap* Heap, uint32 Count_)
 	{
 		LastDefragFrame = GFrameNumberRenderThread;
 		int32 DefragCount = 0;
@@ -4205,6 +4205,7 @@ namespace VulkanRHI
 		};
 
 		FScopeLock ScopeLock(&CS);
+
 		//Search for allocations to move to different pages.
 		for (FVulkanAllocationInternal& Alloc : InternalData)
 		{
@@ -4233,7 +4234,7 @@ namespace VulkanRHI
 
 							//Move the Rendertarget to the new allocation
 							//Function swaps the old allocation into the Allocation object
-							Texture->Move(Device, Allocation);
+							Texture->Move(Device, Context, Allocation);
 							DefragCount++;
 							Device.GetMemoryManager().FreeVulkanAllocation(Allocation);
 
@@ -5213,7 +5214,7 @@ void FStagingBuffer::Evict(FVulkanDevice& Device_)
 {
 	checkNoEntry(); //stagingbuffers are always in system memory so this should not happen
 }
-void FStagingBuffer::Move(FVulkanDevice& Device_, FVulkanAllocation& Allocation_)
+void FStagingBuffer::Move(FVulkanDevice& Device_, FVulkanCommandListContext& Context, FVulkanAllocation& Allocation_)
 {
 	checkNoEntry();//stagingbuffers are always in system memory so this should not happen
 }
@@ -5221,13 +5222,13 @@ void FVulkanRingBuffer::Evict(FVulkanDevice& Device_)
 {
 	checkNoEntry(); //stagingbuffers are always in system memory so this should not happen
 }
-void FVulkanRingBuffer::Move(FVulkanDevice& Device_, FVulkanAllocation& Allocation_)
+void FVulkanRingBuffer::Move(FVulkanDevice& Device_, FVulkanCommandListContext& Context, FVulkanAllocation& Allocation_)
 {
 	checkNoEntry();//stagingbuffers are always in system memory so this should not happen
 }
 
 
-void FVulkanEvictable::OnFullDefrag(FVulkanDevice& Device, uint32 NewOffset)
+void FVulkanEvictable::OnFullDefrag(FVulkanDevice& Device, FVulkanCommandListContext& Context, uint32 NewOffset)
 {
 	checkNoEntry();
 }
diff --git a/Engine/Source/Runtime/VulkanRHI/Private/VulkanRHI.cpp b/Engine/Source/Runtime/VulkanRHI/Private/VulkanRHI.cpp
--- a/Engine/Source/Runtime/VulkanRHI/Private/VulkanRHI.cpp
+++ b/Engine/Source/Runtime/VulkanRHI/Private/VulkanRHI.cpp
@@ -956,7 +956,7 @@ void FVulkanCommandListContext::RHIEndFrame()
 	GetGPUProfiler().EndFrame();
 
 	Device->GetStagingManager().ProcessPendingFree(false, true);
-	Device->GetMemoryManager().ReleaseFreedPages();
+	Device->GetMemoryManager().ReleaseFreedPages(*this);
 
 	if (UseVulkanDescriptorCache())
 	{
diff --git a/Engine/Source/Runtime/VulkanRHI/Private/VulkanResources.h b/Engine/Source/Runtime/VulkanRHI/Private/VulkanResources.h
--- a/Engine/Source/Runtime/VulkanRHI/Private/VulkanResources.h
+++ b/Engine/Source/Runtime/VulkanRHI/Private/VulkanResources.h
@@ -273,7 +273,7 @@ struct FVulkanCpuReadbackBuffer
 class FVulkanSurface : public FVulkanEvictable
 {
 	virtual void Evict(FVulkanDevice& Device);
-	virtual void Move(FVulkanDevice& Device, VulkanRHI::FVulkanAllocation& NewAllocation);
+	virtual void Move(FVulkanDevice& Device, FVulkanCommandListContext& Context, VulkanRHI::FVulkanAllocation& NewAllocation);
 	virtual bool CanEvict();
 	virtual bool CanMove();
 public:
@@ -315,8 +315,8 @@ public:
 	void InvalidateMappedMemory();
 	void* GetMappedPointer();
 
-	void MoveSurface(FVulkanDevice& InDevice, VulkanRHI::FVulkanAllocation& NewAllocation);
-	void OnFullDefrag(FVulkanDevice& InDevice, uint32 NewOffset);
+	void MoveSurface(FVulkanDevice& InDevice, FVulkanCommandListContext& Context, VulkanRHI::FVulkanAllocation& NewAllocation);
+	void OnFullDefrag(FVulkanDevice& InDevice, FVulkanCommandListContext& Context, uint32 NewOffset);
 	void EvictSurface(FVulkanDevice& InDevice);
 
 
@@ -515,8 +515,8 @@ struct FVulkanTextureBase : public FVulkanEvictable, public IRefCountedObject
 	}
 
 	void Evict(FVulkanDevice& Device); ///evict to system memory
-	void Move(FVulkanDevice& Device, VulkanRHI::FVulkanAllocation& NewAllocation); //move to a full new allocation
-	void OnFullDefrag(FVulkanDevice& Device, uint32 NewOffset); //called when compacting an allocation. Old image can still be used as a copy source.
+	void Move(FVulkanDevice& Device, FVulkanCommandListContext& Context, VulkanRHI::FVulkanAllocation& NewAllocation); //move to a full new allocation
+	void OnFullDefrag(FVulkanDevice& Device, FVulkanCommandListContext& Context, uint32 NewOffset); //called when compacting an allocation. Old image can still be used as a copy source.
 	FVulkanTextureBase* GetTextureBase() { return this; }
 
 	void AttachView(VulkanRHI::FVulkanViewBase* View);
@@ -956,7 +956,7 @@ private:
 struct FVulkanRingBuffer : public FVulkanEvictable, public VulkanRHI::FDeviceChild
 {
 	virtual void Evict(FVulkanDevice& Device);
-	virtual void Move(FVulkanDevice& Device, VulkanRHI::FVulkanAllocation& NewAllocation);
+	virtual void Move(FVulkanDevice& Device, FVulkanCommandListContext& Context, VulkanRHI::FVulkanAllocation& NewAllocation);
 
 public:
 	FVulkanRingBuffer(FVulkanDevice* InDevice, uint64 TotalSize, VkFlags Usage, VkMemoryPropertyFlags MemPropertyFlags);
@@ -1054,7 +1054,7 @@ protected:
 class FVulkanResourceMultiBuffer : public FRHIBuffer, public FVulkanEvictable, public VulkanRHI::FDeviceChild
 {
 	virtual void Evict(FVulkanDevice& Device);
-	virtual void Move(FVulkanDevice& Device, VulkanRHI::FVulkanAllocation& NewAllocation);
+	virtual void Move(FVulkanDevice& Device, FVulkanCommandListContext& Context, VulkanRHI::FVulkanAllocation& NewAllocation);
 
 public:
 	FVulkanResourceMultiBuffer(FVulkanDevice* InDevice, uint32 InSize, uint32 InUEUsage, uint32 InStride, FRHIResourceCreateInfo& CreateInfo, class FRHICommandListImmediate* InRHICmdList = nullptr);
diff --git a/Engine/Source/Runtime/VulkanRHI/Private/VulkanTexture.cpp b/Engine/Source/Runtime/VulkanRHI/Private/VulkanTexture.cpp
--- a/Engine/Source/Runtime/VulkanRHI/Private/VulkanTexture.cpp
+++ b/Engine/Source/Runtime/VulkanRHI/Private/VulkanTexture.cpp
@@ -685,7 +685,7 @@ FVulkanSurface::FVulkanSurface(FVulkanDevice& InDevice, FVulkanEvictable* Owner,
 	}
 }
 
-void FVulkanSurface::MoveSurface(FVulkanDevice& InDevice, FVulkanAllocation& NewAllocation)
+void FVulkanSurface::MoveSurface(FVulkanDevice& InDevice, FVulkanCommandListContext& Context, FVulkanAllocation& NewAllocation)
 {
 	FImageCreateInfo ImageCreateInfo;
 	FVulkanSurface::GenerateImageCreateInfo(ImageCreateInfo,
@@ -693,8 +693,6 @@ void FVulkanSurface::MoveSurface(FVulkanDevice& InDevice, FVulkanAllocation& New
 		PixelFormat, Width, Height, Depth,
 		ArraySize, NumMips, NumSamples, UEFlags,
 		&StorageFormat, &ViewFormat);
-	FRHICommandListImmediate& RHICmdList = FRHICommandListExecutor::GetImmediateCommandList();
-	FVulkanCommandListContext& Context = (FVulkanCommandListContext&)RHICmdList.GetContext();
 
 	VkImage MovedImage;
 	VERIFYVULKANRESULT(VulkanRHI::vkCreateImage(InDevice.GetInstanceHandle(), &ImageCreateInfo.ImageCreateInfo, VULKAN_CPU_ALLOCATOR, &MovedImage));
@@ -787,7 +785,7 @@ void FVulkanSurface::MoveSurface(FVulkanDevice& InDevice, FVulkanAllocation& New
 }
 
 
-void FVulkanSurface::OnFullDefrag(FVulkanDevice& InDevice, uint32 NewOffset)
+void FVulkanSurface::OnFullDefrag(FVulkanDevice& InDevice, FVulkanCommandListContext& Context, uint32 NewOffset)
 {
 	FImageCreateInfo ImageCreateInfo;
 	FVulkanSurface::GenerateImageCreateInfo(ImageCreateInfo,
@@ -795,8 +793,6 @@ void FVulkanSurface::OnFullDefrag(FVulkanDevice& InDevice, uint32 NewOffset)
 		PixelFormat, Width, Height, Depth,
 		ArraySize, NumMips, NumSamples, UEFlags,
 		&StorageFormat, &ViewFormat);
-	FRHICommandListImmediate& RHICmdList = FRHICommandListExecutor::GetImmediateCommandList();
-	FVulkanCommandListContext& Context = (FVulkanCommandListContext&)RHICmdList.GetContext();
 
 	VkImage MovedImage;
 	VERIFYVULKANRESULT(VulkanRHI::vkCreateImage(InDevice.GetInstanceHandle(), &ImageCreateInfo.ImageCreateInfo, VULKAN_CPU_ALLOCATOR, &MovedImage));
@@ -2388,7 +2384,7 @@ void FVulkanSurface::Evict(FVulkanDevice& Device_)
 {
 	checkNoEntry(); //not supported
 }
-void FVulkanSurface::Move(FVulkanDevice& Device_ , FVulkanAllocation& NewAllocation)
+void FVulkanSurface::Move(FVulkanDevice& Device_, FVulkanCommandListContext& Context, FVulkanAllocation& NewAllocation)
 {
 	checkNoEntry(); //not supported
 }
@@ -2428,7 +2424,7 @@ void FVulkanTextureBase::InvalidateViews(FVulkanDevice& Device)
 	}
 }
 
-void FVulkanTextureBase::Move(FVulkanDevice& Device, FVulkanAllocation& NewAllocation)
+void FVulkanTextureBase::Move(FVulkanDevice& Device, FVulkanCommandListContext& Context, FVulkanAllocation& NewAllocation)
 {
 	FRHITexture* Tex = GetRHITexture();
 	uint64 Size = Surface.GetMemorySize();
@@ -2440,11 +2436,11 @@ void FVulkanTextureBase::Move(FVulkanDevice& Device, FVulkanAllocation& NewAlloc
 		UE_LOG(LogVulkanRHI, Display, TEXT("Evicted %8.4fkb %8.4fkb   TB %p // %p  :: IMG %p   %-40s\n"), Size / (1024.f), TotalSize / (1024.f), this, &Surface, Surface.Image, *GetResourceFName().ToString());
 	}
 
-	Surface.MoveSurface(Device, NewAllocation);
+	Surface.MoveSurface(Device, Context, NewAllocation);
 	InvalidateViews(Device);
 }
 
-void FVulkanTextureBase::OnFullDefrag(FVulkanDevice& Device, uint32 NewOffset)
+void FVulkanTextureBase::OnFullDefrag(FVulkanDevice& Device, FVulkanCommandListContext& Context, uint32 NewOffset)
 {
 	FRHITexture* Tex = GetRHITexture();
 	uint64 Size = Surface.GetMemorySize();
@@ -2456,7 +2452,7 @@ void FVulkanTextureBase::OnFullDefrag(FVulkanDevice& Device, uint32 NewOffset)
 		UE_LOG(LogVulkanRHI, Display, TEXT("Evicted %8.4fkb %8.4fkb   TB %p // %p  :: IMG %p   %-40s\n"), Size / (1024.f), TotalSize / (1024.f), this, &Surface, Surface.Image, *GetResourceFName().ToString());
 	}
 
-	Surface.OnFullDefrag(Device, NewOffset);
+	Surface.OnFullDefrag(Device, Context, NewOffset);
 	InvalidateViews(Device);
 }
 
diff --git a/Engine/Source/Runtime/VulkanRHI/Public/VulkanMemory.h b/Engine/Source/Runtime/VulkanRHI/Public/VulkanMemory.h
--- a/Engine/Source/Runtime/VulkanRHI/Public/VulkanMemory.h
+++ b/Engine/Source/Runtime/VulkanRHI/Public/VulkanMemory.h
@@ -89,9 +89,9 @@ class FVulkanEvictable
 {
 public:
 	virtual void Evict(FVulkanDevice& Device) = 0;
-	virtual void Move(FVulkanDevice& Device, VulkanRHI::FVulkanAllocation& Allocation) = 0;
+	virtual void Move(FVulkanDevice& Device, FVulkanCommandListContext& Context, VulkanRHI::FVulkanAllocation& Allocation) = 0;
 
-	virtual void OnFullDefrag(FVulkanDevice& Device, uint32 NewOffset);
+	virtual void OnFullDefrag(FVulkanDevice& Device, FVulkanCommandListContext& Context, uint32 NewOffset);
 	
 	virtual struct FVulkanTextureBase* GetTextureBase(){return 0;}
 	virtual bool CanMove(){return true;}
@@ -622,8 +622,8 @@ namespace VulkanRHI
 	protected:
 		void SetIsDefragging(bool bInIsDefragging){ bIsDefragging = bInIsDefragging; }
 		void DumpFullHeap();
-		bool DefragFull(FVulkanDevice& Device, FVulkanResourceHeap* Heap);
-		int32 DefragTick(FVulkanDevice& Device, FVulkanResourceHeap* Heap, uint32 Count);
+		bool DefragFull(FVulkanDevice& Device, FVulkanCommandListContext& Context, FVulkanResourceHeap* Heap);
+		int32 DefragTick(FVulkanDevice& Device, FVulkanCommandListContext& Context, FVulkanResourceHeap* Heap, uint32 Count);
 		bool CanDefrag();
 		uint64 EvictToHost(FVulkanDevice& Device);
 		void SetFreePending(FVulkanAllocation& Allocation);
@@ -710,7 +710,7 @@ namespace VulkanRHI
 		}
 
 		uint64 EvictOne(FVulkanDevice& Device);
-		void DefragTick(FVulkanDevice& Device, uint32 Count);
+		void DefragTick(FVulkanDevice& Device, FVulkanCommandListContext& Context, uint32 Count);
 		void DumpMemory(FResourceHeapStats& Stats);
 
 		void SetDefragging(FVulkanSubresourceAllocator* Allocator);
@@ -785,7 +785,7 @@ namespace VulkanRHI
 		void ReleaseSubresourceAllocator(FVulkanSubresourceAllocator* SubresourceAllocator);
 
 
-		void ReleaseFreedPages();
+		void ReleaseFreedPages(FVulkanCommandListContext& Context);
 		void DumpMemory(bool bFullDump = true);
 
 
@@ -903,7 +903,7 @@ namespace VulkanRHI
 	class FStagingBuffer : public FVulkanEvictable, public FRefCount
 	{
 		virtual void Evict(FVulkanDevice& Device);
-		virtual void Move(FVulkanDevice& Device, FVulkanAllocation& Allocation);
+		virtual void Move(FVulkanDevice& Device, FVulkanCommandListContext& Context, FVulkanAllocation& Allocation);
 		virtual bool CanMove(){return false;}
 		virtual bool CanEvict(){return false;}
 
