Fix for variable name collision

diff --git a/Engine/Source/Runtime/VulkanRHI/Private/VulkanBuffer.cpp b/Engine/Source/Runtime/VulkanRHI/Private/VulkanBuffer.cpp
--- a/Engine/Source/Runtime/VulkanRHI/Private/VulkanBuffer.cpp
+++ b/Engine/Source/Runtime/VulkanRHI/Private/VulkanBuffer.cpp
@@ -479,7 +479,7 @@ void FVulkanResourceMultiBuffer::Evict(FVulkanDevice& Device_)
 {
 	checkNoEntry();//Not Implemented, should never be called
 }
-void FVulkanResourceMultiBuffer::Move(FVulkanDevice& Device, FVulkanCommandListContext& Context, VulkanRHI::FVulkanAllocation& NewAllocation)
+void FVulkanResourceMultiBuffer::Move(FVulkanDevice& Device_, FVulkanCommandListContext& Context, VulkanRHI::FVulkanAllocation& NewAllocation)
 {
 	checkNoEntry();//Not Implemented, should never be called
 }
