Fix for command buffers not being freed when using LazyAcquire or when RHIThread is not ticked i.e. Dialogs.

diff --git a/Engine/Source/Runtime/VulkanRHI/Private/VulkanRHI.cpp b/Engine/Source/Runtime/VulkanRHI/Private/VulkanRHI.cpp
--- a/Engine/Source/Runtime/VulkanRHI/Private/VulkanRHI.cpp
+++ b/Engine/Source/Runtime/VulkanRHI/Private/VulkanRHI.cpp
@@ -952,9 +952,10 @@ void FVulkanCommandListContext::RHIEndFrame()
 	check(IsImmediate());
 	//FRCLog::Printf(FString::Printf(TEXT("FVulkanCommandListContext::RHIEndFrame()")));
 
-	
 	GetGPUProfiler().EndFrame();
 
+	GetCommandBufferManager()->FreeUnusedCmdBuffers();
+
 	Device->GetStagingManager().ProcessPendingFree(false, true);
 	Device->GetMemoryManager().ReleaseFreedPages(*this);
 
diff --git a/Engine/Source/Runtime/VulkanRHI/Private/VulkanViewport.cpp b/Engine/Source/Runtime/VulkanRHI/Private/VulkanViewport.cpp
--- a/Engine/Source/Runtime/VulkanRHI/Private/VulkanViewport.cpp
+++ b/Engine/Source/Runtime/VulkanRHI/Private/VulkanViewport.cpp
@@ -1050,12 +1050,6 @@ void FVulkanDynamicRHI::RHITick(float DeltaTime)
 			}
 
 			VulkanDevice->GetImmediateContext().GetTempFrameAllocationBuffer().Reset();
-
-			// Destroy command buffers here when using Delay; when not delaying we'll delete after Acquire
-			if (GVulkanDelayAcquireImage == EDelayAcquireImageType::DelayAcquire)
-			{
-				VulkanDevice->GetImmediateContext().GetCommandBufferManager()->FreeUnusedCmdBuffers();
-			}
 		});
 
 	if (bRequestNULLPixelShader)
