Fix Vulkan crash when calling RHIGetRenderQueryResult on a timer query before the RHI thread executes the preceding RHIEndRenderQuery.

#rnx

diff --git a/Engine/Source/Runtime/VulkanRHI/Private/VulkanQuery.cpp b/Engine/Source/Runtime/VulkanRHI/Private/VulkanQuery.cpp
--- a/Engine/Source/Runtime/VulkanRHI/Private/VulkanQuery.cpp
+++ b/Engine/Source/Runtime/VulkanRHI/Private/VulkanQuery.cpp
@@ -464,6 +464,19 @@ bool FVulkanDynamicRHI::RHIGetRenderQueryResult(FRHIRenderQuery* QueryRHI, uint6
 	else if (BaseQuery->QueryType == RQT_AbsoluteTime)
 	{
 		FVulkanTimingQuery* Query = static_cast<FVulkanTimingQuery*>(BaseQuery);
+		if (Query->Pool == nullptr)
+		{
+			// RHIEndRenderQuery hasn't been executed yet for this query, so it's not ready yet. If it's a non-blocking wait, we can return false right away.
+			if (!bWait)
+			{
+				return false;
+			}
+
+			// Blocking was requested, so, we need to wait for the RHI thread to catch up.
+			FRHICommandListExecutor::GetImmediateCommandList().ImmediateFlush(EImmediateFlushType::FlushRHIThread);
+			check(Query->Pool != nullptr);
+		}
+
 		check(Query->Pool->CurrentTimestamp < Query->Pool->BufferSize);
 		int32 TimestampIndex = Query->Pool->CurrentTimestamp;
 		if (!bWait)
