Tessellation Removal: Removing Tessellation from VulkanRHI.

diff --git a/Engine/Source/Runtime/VulkanRHI/Private/VulkanCommands.cpp b/Engine/Source/Runtime/VulkanRHI/Private/VulkanCommands.cpp
--- a/Engine/Source/Runtime/VulkanRHI/Private/VulkanCommands.cpp
+++ b/Engine/Source/Runtime/VulkanRHI/Private/VulkanCommands.cpp
@@ -55,24 +55,6 @@ static FORCEINLINE ShaderStage::EStage GetAndVerifyShaderStage(FRHIGraphicsShade
 	case SF_Vertex:
 		check(PendingGfxState->GetCurrentShaderKey(ShaderStage::Vertex) == GetShaderKey<FVulkanVertexShader>(ShaderRHI));
 		return ShaderStage::Vertex;
-	case SF_Hull:
-#if PLATFORM_SUPPORTS_TESSELLATION_SHADERS
-		check(PendingGfxState->GetCurrentShaderKey(ShaderStage::Hull) == GetShaderKey<FVulkanHullShader>(ShaderRHI));
-		return ShaderStage::Hull;
-#else
-		checkf(0, TEXT("Tessellation (Hull) not supported on this platform!"));
-		UE_LOG(LogVulkanRHI, Fatal, TEXT("Tessellation (Hull) not supported on this platform!"));
-		break;
-#endif
-	case SF_Domain:
-#if PLATFORM_SUPPORTS_TESSELLATION_SHADERS
-		check(PendingGfxState->GetCurrentShaderKey(ShaderStage::Domain) == GetShaderKey<FVulkanDomainShader>(ShaderRHI));
-		return ShaderStage::Domain;
-#else
-		checkf(0, TEXT("Tessellation (Domain) not supported on this platform!"));
-		UE_LOG(LogVulkanRHI, Fatal, TEXT("Tessellation (Domain) not supported on this platform!"));
-		break;
-#endif
 	case SF_Geometry:
 #if VULKAN_SUPPORTS_GEOMETRY_SHADERS
 		check(PendingGfxState->GetCurrentShaderKey(ShaderStage::Geometry) == GetShaderKey<FVulkanGeometryShader>(ShaderRHI));
@@ -101,26 +83,6 @@ static FORCEINLINE ShaderStage::EStage GetAndVerifyShaderStageAndVulkanShader(FR
 		//check(PendingGfxState->GetCurrentShaderKey(ShaderStage::Vertex) == GetShaderKey<FVulkanVertexShader>(ShaderRHI));
 		OutShader = static_cast<FVulkanVertexShader*>(static_cast<FRHIVertexShader*>(ShaderRHI));
 		return ShaderStage::Vertex;
-	case SF_Hull:
-#if PLATFORM_SUPPORTS_TESSELLATION_SHADERS
-		//check(PendingGfxState->GetCurrentShaderKey(ShaderStage::Hull) == GetShaderKey<FVulkanHullShader>(ShaderRHI));
-		OutShader = static_cast<FVulkanHullShader*>(static_cast<FRHIHullShader*>(ShaderRHI));
-		return ShaderStage::Hull;
-#else
-		checkf(0, TEXT("Tessellation (Domain) not supported on this platform!"));
-		UE_LOG(LogVulkanRHI, Fatal, TEXT("Tessellation (Domain) not supported on this platform!"));
-		break;
-#endif
-	case SF_Domain:
-#if PLATFORM_SUPPORTS_TESSELLATION_SHADERS
-		//check(PendingGfxState->GetCurrentShaderKey(ShaderStage::Domain) == GetShaderKey<FVulkanDomainShader>(ShaderRHI));
-		OutShader = static_cast<FVulkanDomainShader*>(static_cast<FRHIDomainShader*>(ShaderRHI));
-		return ShaderStage::Domain;
-#else
-		checkf(0, TEXT("Tessellation (Hull) not supported on this platform!"));
-		UE_LOG(LogVulkanRHI, Fatal, TEXT("Tessellation (Hull) not supported on this platform!"));
-		break;
-#endif
 	case SF_Geometry:
 #if VULKAN_SUPPORTS_GEOMETRY_SHADERS
 		//check(PendingGfxState->GetCurrentShaderKey(ShaderStage::Geometry) == GetShaderKey<FVulkanGeometryShader>(ShaderRHI));
diff --git a/Engine/Source/Runtime/VulkanRHI/Private/VulkanPipeline.cpp b/Engine/Source/Runtime/VulkanRHI/Private/VulkanPipeline.cpp
--- a/Engine/Source/Runtime/VulkanRHI/Private/VulkanPipeline.cpp
+++ b/Engine/Source/Runtime/VulkanRHI/Private/VulkanPipeline.cpp
@@ -159,10 +159,6 @@ static inline FSHAHash GetShaderHashForStage(const FGraphicsPipelineStateInitial
 	case ShaderStage::Pixel:		return GetShaderHash<FRHIPixelShader, FVulkanPixelShader>(PSOInitializer.BoundShaderState.PixelShaderRHI);
 #if VULKAN_SUPPORTS_GEOMETRY_SHADERS
 	case ShaderStage::Geometry:		return GetShaderHash<FRHIGeometryShader, FVulkanGeometryShader>(PSOInitializer.BoundShaderState.GeometryShaderRHI);
-#endif
-#if PLATFORM_SUPPORTS_TESSELLATION_SHADERS
-	case ShaderStage::Hull:			return GetShaderHash<FRHIHullShader, FVulkanHullShader>(PSOInitializer.BoundShaderState.HullShaderRHI);
-	case ShaderStage::Domain:		return GetShaderHash<FRHIDomainShader, FVulkanDomainShader>(PSOInitializer.BoundShaderState.DomainShaderRHI);
 #endif
 	default:			check(0);	break;
 	}
@@ -1322,10 +1318,6 @@ FVulkanShaderHashes::FVulkanShaderHashes(const FGraphicsPipelineStateInitializer
 	Stages[ShaderStage::Pixel] = GetShaderHash<FRHIPixelShader, FVulkanPixelShader>(PSOInitializer.BoundShaderState.PixelShaderRHI);
 #if VULKAN_SUPPORTS_GEOMETRY_SHADERS
 	Stages[ShaderStage::Geometry] = GetShaderHash<FRHIGeometryShader, FVulkanGeometryShader>(PSOInitializer.BoundShaderState.GeometryShaderRHI);
-#endif
-#if PLATFORM_SUPPORTS_TESSELLATION_SHADERS
-	Stages[ShaderStage::Hull] = GetShaderHash<FRHIHullShader, FVulkanHullShader>(PSOInitializer.BoundShaderState.HullShaderRHI);
-	Stages[ShaderStage::Domain] = GetShaderHash<FRHIDomainShader, FVulkanDomainShader>(PSOInitializer.BoundShaderState.DomainShaderRHI);
 #endif
 	Finalize();
 }
@@ -1514,15 +1506,6 @@ void FVulkanPipelineStateCacheManager::CreateGfxEntry(const FGraphicsPipelineSta
 			}
 #endif
 
-#if PLATFORM_SUPPORTS_TESSELLATION_SHADERS
-			if (Shaders[ShaderStage::Hull])
-			{
-				const FVulkanShaderHeader& HSHeader = Shaders[ShaderStage::Hull]->GetCodeHeader();
-				const FVulkanShaderHeader& DSHeader = Shaders[ShaderStage::Domain]->GetCodeHeader();
-				DescriptorSetLayoutInfo.ProcessBindingsForStage(VK_SHADER_STAGE_TESSELLATION_CONTROL_BIT, ShaderStage::Hull, HSHeader, UBGatherInfo);
-				DescriptorSetLayoutInfo.ProcessBindingsForStage(VK_SHADER_STAGE_TESSELLATION_EVALUATION_BIT, ShaderStage::Domain, DSHeader, UBGatherInfo);
-			}
-#endif
 			// Second pass
 			const int32 NumImmutableSamplers = PSOInitializer.ImmutableSamplerState.ImmutableSamplers.Num();
 			TArrayView<FRHISamplerState*> ImmutableSamplers(NumImmutableSamplers > 0 ? &(FRHISamplerState*&)PSOInitializer.ImmutableSamplerState.ImmutableSamplers[0] : nullptr, NumImmutableSamplers);
@@ -1667,10 +1650,6 @@ FVulkanRHIGraphicsPipelineState::FVulkanRHIGraphicsPipelineState(FVulkanDevice*
 
 	FMemory::Memset(VulkanShaders, 0, sizeof(VulkanShaders));
 	VulkanShaders[ShaderStage::Vertex] = static_cast<FVulkanVertexShader*>(PSOInitializer_.BoundShaderState.VertexShaderRHI);
-#if PLATFORM_SUPPORTS_TESSELLATION_SHADERS
-	VulkanShaders[ShaderStage::Hull] = static_cast<FVulkanHullShader*>(PSOInitializer_.BoundShaderState.HullShaderRHI);
-	VulkanShaders[ShaderStage::Domain] = static_cast<FVulkanDomainShader*>(PSOInitializer_.BoundShaderState.DomainShaderRHI);
-#endif
 #if PLATFORM_SUPPORTS_GEOMETRY_SHADERS
 	VulkanShaders[ShaderStage::Geometry] = static_cast<FVulkanGeometryShader*>(PSOInitializer_.BoundShaderState.GeometryShaderRHI);
 #endif
@@ -1689,11 +1668,6 @@ FVulkanRHIGraphicsPipelineState::FVulkanRHIGraphicsPipelineState(FVulkanDevice*
 	VertexShaderRHI = PSOInitializer_.BoundShaderState.VertexShaderRHI;
 	VertexDeclarationRHI = PSOInitializer_.BoundShaderState.VertexDeclarationRHI;
 
-#if PLATFORM_SUPPORTS_TESSELLATION_SHADERS
-	DomainShaderRHI = PSOInitializer_.BoundShaderState.DomainShaderRHI;
-	HullShaderRHI = PSOInitializer_.BoundShaderState.HullShaderRHI;
-#endif 
-
 #if PLATFORM_SUPPORTS_GEOMETRY_SHADERS
 	GeometryShaderRHI = PSOInitializer_.BoundShaderState.GeometryShaderRHI;
 #endif
@@ -2053,21 +2027,6 @@ void GetVulkanShaders(const FBoundShaderStateInput& BSI, FVulkanShader* OutShade
 #endif
 	}
 #endif
-
-#if PLATFORM_SUPPORTS_TESSELLATION_SHADERS
-	if (BSI.HullShaderRHI)
-	{
-		// Can't have Hull w/o Domain
-		check(BSI.DomainShaderRHI);
-		OutShaders[ShaderStage::Hull] = ResourceCast(BSI.HullShaderRHI);
-		OutShaders[ShaderStage::Domain] = ResourceCast(BSI.DomainShaderRHI);
-	}
-	else
-	{
-		// Can't have Domain w/o Hull
-		check(BSI.DomainShaderRHI == nullptr);
-	}
-#endif
 }
 
 void GetVulkanShaders(FVulkanDevice* Device, const FVulkanRHIGraphicsPipelineState& GfxPipelineState, FVulkanShader* OutShaders[ShaderStage::NumStages])
diff --git a/Engine/Source/Runtime/VulkanRHI/Private/VulkanPipeline.h b/Engine/Source/Runtime/VulkanRHI/Private/VulkanPipeline.h
--- a/Engine/Source/Runtime/VulkanRHI/Private/VulkanPipeline.h
+++ b/Engine/Source/Runtime/VulkanRHI/Private/VulkanPipeline.h
@@ -43,12 +43,6 @@ inline uint64 GetShaderKeyForGfxStage(const FBoundShaderStateInput& BSI, ShaderS
 #if VULKAN_SUPPORTS_GEOMETRY_SHADERS
 	case ShaderStage::Geometry:
 		return GetShaderKey<FVulkanGeometryShader>(BSI.GeometryShaderRHI);
-#endif
-#if PLATFORM_SUPPORTS_TESSELLATION_SHADERS
-	case ShaderStage::Hull:
-		return GetShaderKey<FVulkanHullShader>(BSI.HullShaderRHI);
-	case ShaderStage::Domain:
-		return GetShaderKey<FVulkanDomainShader>(BSI.DomainShaderRHI);
 #endif
 	default:
 		check(0);
@@ -726,10 +720,6 @@ public:
 	FVertexShaderRHIRef					VertexShaderRHI;
 	FVertexDeclarationRHIRef			VertexDeclarationRHI;
 
-#if PLATFORM_SUPPORTS_TESSELLATION_SHADERS
-	FDomainShaderRHIRef					DomainShaderRHI;
-	FHullShaderRHIRef					HullShaderRHI;
-#endif 
 #if PLATFORM_SUPPORTS_GEOMETRY_SHADERS
 	FGeometryShaderRHIRef				GeometryShaderRHI;
 #endif
diff --git a/Engine/Source/Runtime/VulkanRHI/Private/VulkanPipelineState.cpp b/Engine/Source/Runtime/VulkanRHI/Private/VulkanPipelineState.cpp
--- a/Engine/Source/Runtime/VulkanRHI/Private/VulkanPipelineState.cpp
+++ b/Engine/Source/Runtime/VulkanRHI/Private/VulkanPipelineState.cpp
@@ -245,20 +245,6 @@ FVulkanGraphicsPipelineDescriptorState::FVulkanGraphicsPipelineDescriptorState(F
 		}
 #endif
 
-#if PLATFORM_SUPPORTS_TESSELLATION_SHADERS
-		uint64 HullShaderKey = InGfxPipeline->GetShaderKey(SF_Hull);
-		if (HullShaderKey)
-		{
-			const FVulkanHullShader* HullShader = ShaderFactory.LookupShader<FVulkanHullShader>(HullShaderKey);
-			PackedUniformBuffers[ShaderStage::Hull].Init(HullShader->GetCodeHeader(), PackedUniformBuffersMask[ShaderStage::Hull]);
-		}
-		uint64 DomainShaderKey = InGfxPipeline->GetShaderKey(SF_Domain);
-		if (DomainShaderKey)
-		{
-			const FVulkanDomainShader* DomainShader = ShaderFactory.LookupShader<FVulkanDomainShader>(DomainShaderKey);
-			PackedUniformBuffers[ShaderStage::Domain].Init(DomainShader->GetCodeHeader(), PackedUniformBuffersMask[ShaderStage::Domain]);
-		}
-#endif
 		CreateDescriptorWriteInfos();
 
 		//UE_LOG(LogVulkanRHI, Warning, TEXT("GfxPSOState %p For PSO %p Writes:%d"), this, InGfxPipeline, DSWriteContainer.DescriptorWrites.Num());
@@ -372,10 +358,6 @@ void FVulkanCommandListContext::RHISetGraphicsPipelineState(FRHIGraphicsPipeline
 	if (bApplyAdditionalState)
 	{
 		ApplyStaticUniformBuffers(static_cast<FVulkanVertexShader*>(Pipeline->VulkanShaders[ShaderStage::Vertex]));
-#if PLATFORM_SUPPORTS_TESSELLATION_SHADERS
-		ApplyStaticUniformBuffers(static_cast<FVulkanHullShader*>(Pipeline->VulkanShaders[ShaderStage::Hull]));
-		ApplyStaticUniformBuffers(static_cast<FVulkanDomainShader*>(Pipeline->VulkanShaders[ShaderStage::Domain]));
-#endif
 #if PLATFORM_SUPPORTS_GEOMETRY_SHADERS
 		ApplyStaticUniformBuffers(static_cast<FVulkanGeometryShader*>(Pipeline->VulkanShaders[ShaderStage::Geometry]));
 #endif
diff --git a/Engine/Source/Runtime/VulkanRHI/Private/VulkanRHIPrivate.h b/Engine/Source/Runtime/VulkanRHI/Private/VulkanRHIPrivate.h
--- a/Engine/Source/Runtime/VulkanRHI/Private/VulkanRHIPrivate.h
+++ b/Engine/Source/Runtime/VulkanRHI/Private/VulkanRHIPrivate.h
@@ -83,8 +83,6 @@ inline VkShaderStageFlagBits UEFrequencyToVKStageBit(EShaderFrequency InStage)
 	switch (InStage)
 	{
 	case SF_Vertex:		return VK_SHADER_STAGE_VERTEX_BIT;
-	case SF_Hull:		return VK_SHADER_STAGE_TESSELLATION_CONTROL_BIT;
-	case SF_Domain:		return VK_SHADER_STAGE_TESSELLATION_EVALUATION_BIT;
 	case SF_Pixel:		return VK_SHADER_STAGE_FRAGMENT_BIT;
 	case SF_Geometry:	return VK_SHADER_STAGE_GEOMETRY_BIT;
 	case SF_Compute:	return VK_SHADER_STAGE_COMPUTE_BIT;
@@ -101,8 +99,6 @@ inline EShaderFrequency VkStageBitToUEFrequency(VkShaderStageFlagBits FlagBits)
 	switch (FlagBits)
 	{
 	case VK_SHADER_STAGE_VERTEX_BIT:					return SF_Vertex;
-	case VK_SHADER_STAGE_TESSELLATION_CONTROL_BIT:		return SF_Hull;
-	case VK_SHADER_STAGE_TESSELLATION_EVALUATION_BIT:	return SF_Domain;
 	case VK_SHADER_STAGE_FRAGMENT_BIT:					return SF_Pixel;
 	case VK_SHADER_STAGE_GEOMETRY_BIT:					return SF_Geometry;
 	case VK_SHADER_STAGE_COMPUTE_BIT:					return SF_Compute;
diff --git a/Engine/Source/Runtime/VulkanRHI/Private/VulkanResources.h b/Engine/Source/Runtime/VulkanRHI/Private/VulkanResources.h
--- a/Engine/Source/Runtime/VulkanRHI/Private/VulkanResources.h
+++ b/Engine/Source/Runtime/VulkanRHI/Private/VulkanResources.h
@@ -182,8 +182,6 @@ public:
 
 typedef TVulkanBaseShader<FRHIVertexShader, SF_Vertex, VK_SHADER_STAGE_VERTEX_BIT>					FVulkanVertexShader;
 typedef TVulkanBaseShader<FRHIPixelShader, SF_Pixel, VK_SHADER_STAGE_FRAGMENT_BIT>					FVulkanPixelShader;
-typedef TVulkanBaseShader<FRHIHullShader, SF_Hull, VK_SHADER_STAGE_TESSELLATION_CONTROL_BIT>		FVulkanHullShader;
-typedef TVulkanBaseShader<FRHIDomainShader, SF_Domain, VK_SHADER_STAGE_TESSELLATION_EVALUATION_BIT>	FVulkanDomainShader;
 typedef TVulkanBaseShader<FRHIComputeShader, SF_Compute, VK_SHADER_STAGE_COMPUTE_BIT>				FVulkanComputeShader;
 typedef TVulkanBaseShader<FRHIGeometryShader, SF_Geometry, VK_SHADER_STAGE_GEOMETRY_BIT>			FVulkanGeometryShader;
 
@@ -233,8 +231,6 @@ public:
 
 	FORCEINLINE FVulkanVertexShader*   GetVertexShader() const { return (FVulkanVertexShader*)CacheLink.GetVertexShader(); }
 	FORCEINLINE FVulkanPixelShader*    GetPixelShader() const { return (FVulkanPixelShader*)CacheLink.GetPixelShader(); }
-	FORCEINLINE FVulkanHullShader*     GetHullShader() const { return (FVulkanHullShader*)CacheLink.GetHullShader(); }
-	FORCEINLINE FVulkanDomainShader*   GetDomainShader() const { return (FVulkanDomainShader*)CacheLink.GetDomainShader(); }
 	FORCEINLINE FVulkanGeometryShader* GetGeometryShader() const { return (FVulkanGeometryShader*)CacheLink.GetGeometryShader(); }
 
 	const FVulkanShader* GetShader(ShaderStage::EStage Stage) const
@@ -242,10 +238,6 @@ public:
 		switch (Stage)
 		{
 		case ShaderStage::Vertex:		return GetVertexShader();
-#if PLATFORM_SUPPORTS_TESSELLATION_SHADERS
-		case ShaderStage::Hull:			return GetHullShader();
-		case ShaderStage::Domain:		return GetDomainShader();
-#endif
 		case ShaderStage::Pixel:		return GetPixelShader();
 #if VULKAN_SUPPORTS_GEOMETRY_SHADERS
 		case ShaderStage::Geometry:	return GetGeometryShader();
@@ -1479,16 +1471,6 @@ struct TVulkanResourceTraits<FRHIGeometryShader>
 	typedef FVulkanGeometryShader TConcreteType;
 };
 template<>
-struct TVulkanResourceTraits<FRHIHullShader>
-{
-	typedef FVulkanHullShader TConcreteType;
-};
-template<>
-struct TVulkanResourceTraits<FRHIDomainShader>
-{
-	typedef FVulkanDomainShader TConcreteType;
-};
-template<>
 struct TVulkanResourceTraits<FRHIPixelShader>
 {
 	typedef FVulkanPixelShader TConcreteType;
diff --git a/Engine/Source/Runtime/VulkanRHI/Private/VulkanShaders.cpp b/Engine/Source/Runtime/VulkanRHI/Private/VulkanShaders.cpp
--- a/Engine/Source/Runtime/VulkanRHI/Private/VulkanShaders.cpp
+++ b/Engine/Source/Runtime/VulkanRHI/Private/VulkanShaders.cpp
@@ -237,16 +237,6 @@ FPixelShaderRHIRef FVulkanDynamicRHI::RHICreatePixelShader(TArrayView<const uint
 	return Device->GetShaderFactory().CreateShader<FVulkanPixelShader>(Code, Device);
 }
 
-FHullShaderRHIRef FVulkanDynamicRHI::RHICreateHullShader(TArrayView<const uint8> Code, const FSHAHash& Hash)
-{ 
-	return Device->GetShaderFactory().CreateShader<FVulkanHullShader>(Code, Device);
-}
-
-FDomainShaderRHIRef FVulkanDynamicRHI::RHICreateDomainShader(TArrayView<const uint8> Code, const FSHAHash& Hash)
-{ 
-	return Device->GetShaderFactory().CreateShader<FVulkanDomainShader>(Code, Device);
-}
-
 FGeometryShaderRHIRef FVulkanDynamicRHI::RHICreateGeometryShader(TArrayView<const uint8> Code, const FSHAHash& Hash)
 { 
 	return Device->GetShaderFactory().CreateShader<FVulkanGeometryShader>(Code, Device);
diff --git a/Engine/Source/Runtime/VulkanRHI/Public/VulkanCommon.h b/Engine/Source/Runtime/VulkanRHI/Public/VulkanCommon.h
--- a/Engine/Source/Runtime/VulkanRHI/Public/VulkanCommon.h
+++ b/Engine/Source/Runtime/VulkanRHI/Public/VulkanCommon.h
@@ -28,10 +28,6 @@ namespace ShaderStage
 
 #if VULKAN_SUPPORTS_GEOMETRY_SHADERS
 		Geometry		= 2,
-#if PLATFORM_SUPPORTS_TESSELLATION_SHADERS
-		Hull			= 3,
-		Domain			= 4,
-#endif
 		NumStages,
 
 		MaxNumSets		= 8,
@@ -52,10 +48,6 @@ namespace ShaderStage
 		switch (Stage)
 		{
 		case SF_Vertex:		return Vertex;
-#if PLATFORM_SUPPORTS_TESSELLATION_SHADERS
-		case SF_Hull:		return Hull;
-		case SF_Domain:		return Domain;
-#endif
 		case SF_Pixel:		return Pixel;
 #if VULKAN_SUPPORTS_GEOMETRY_SHADERS
 		case SF_Geometry:	return Geometry;
@@ -74,10 +66,6 @@ namespace ShaderStage
 		switch (Stage)
 		{
 		case EStage::Vertex:	return SF_Vertex;
-#if PLATFORM_SUPPORTS_TESSELLATION_SHADERS
-		case EStage::Hull:		return SF_Hull;
-		case EStage::Domain:	return SF_Domain;
-#endif
 		case EStage::Pixel:		return SF_Pixel;
 #if VULKAN_SUPPORTS_GEOMETRY_SHADERS
 		case EStage::Geometry:	return SF_Geometry;
diff --git a/Engine/Source/Runtime/VulkanRHI/Public/VulkanDynamicRHI.h b/Engine/Source/Runtime/VulkanRHI/Public/VulkanDynamicRHI.h
--- a/Engine/Source/Runtime/VulkanRHI/Public/VulkanDynamicRHI.h
+++ b/Engine/Source/Runtime/VulkanRHI/Public/VulkanDynamicRHI.h
@@ -68,8 +68,6 @@ public:
 	virtual FVertexDeclarationRHIRef RHICreateVertexDeclaration(const FVertexDeclarationElementList& Elements) final override;
 	virtual FPixelShaderRHIRef RHICreatePixelShader(TArrayView<const uint8> Code, const FSHAHash& Hash) final override;
 	virtual FVertexShaderRHIRef RHICreateVertexShader(TArrayView<const uint8> Code, const FSHAHash& Hash) final override;
-	virtual FHullShaderRHIRef RHICreateHullShader(TArrayView<const uint8> Code, const FSHAHash& Hash) final override;
-	virtual FDomainShaderRHIRef RHICreateDomainShader(TArrayView<const uint8> Code, const FSHAHash& Hash) final override;
 	virtual FGeometryShaderRHIRef RHICreateGeometryShader(TArrayView<const uint8> Code, const FSHAHash& Hash) final override;
 	virtual FComputeShaderRHIRef RHICreateComputeShader(TArrayView<const uint8> Code, const FSHAHash& Hash) final override;
 	virtual FGPUFenceRHIRef RHICreateGPUFence(const FName &Name) final override;
@@ -209,14 +207,6 @@ public:
 	{
 		return RHICreateComputeShader(Code, Hash);
 	}
-	virtual FHullShaderRHIRef CreateHullShader_RenderThread(class FRHICommandListImmediate& RHICmdList, TArrayView<const uint8> Code, const FSHAHash& Hash) override final
-	{
-		return RHICreateHullShader(Code, Hash);
-	}
-	virtual FDomainShaderRHIRef CreateDomainShader_RenderThread(class FRHICommandListImmediate& RHICmdList, TArrayView<const uint8> Code, const FSHAHash& Hash) override final
-	{
-		return RHICreateDomainShader(Code, Hash);
-	}
 	virtual void* LockTexture2D_RenderThread(class FRHICommandListImmediate& RHICmdList, FRHITexture2D* Texture, uint32 MipIndex, EResourceLockMode LockMode, uint32& DestStride, bool bLockWithinMiptail, bool bNeedsDefaultRHIFlush = true) override final
 	{
 		return RHILockTexture2D(Texture, MipIndex, LockMode, DestStride, bLockWithinMiptail);
