Fix bad sky reflections on Vulkan.

Cubemap array resources weren't being set to the requested initial layout because the code was computing the number of layers incorrectly. It doesn't need to compute it at all, since we need to transition the entire resource, and MakeSubresourceRange() does the right thing if we leave its arguments at default values.

#rnx

diff --git a/Engine/Source/Runtime/VulkanRHI/Private/VulkanTexture.cpp b/Engine/Source/Runtime/VulkanRHI/Private/VulkanTexture.cpp
--- a/Engine/Source/Runtime/VulkanRHI/Private/VulkanTexture.cpp
+++ b/Engine/Source/Runtime/VulkanRHI/Private/VulkanTexture.cpp
@@ -1379,8 +1379,7 @@ void FVulkanSurface::SetInitialImageState(FVulkanCommandListContext& Context, Vk
 	FVulkanCmdBuffer* CmdBuffer = Context.GetCommandBufferManager()->GetUploadCmdBuffer();
 	ensure(CmdBuffer->IsOutsideRenderPass());
 
-	uint32 NumLayers = ViewType == VK_IMAGE_VIEW_TYPE_CUBE ? 6 : 1;
-	VkImageSubresourceRange SubresourceRange = FVulkanPipelineBarrier::MakeSubresourceRange(FullAspectMask, 0, NumMips, 0, NumLayers);
+	VkImageSubresourceRange SubresourceRange = FVulkanPipelineBarrier::MakeSubresourceRange(FullAspectMask);
 
 	VkImageLayout CurrentLayout;
 	if (bClear)
