Fix dumping Vulkan memory for null pointer resource type heap.

#rnx

diff --git a/Engine/Source/Runtime/VulkanRHI/Private/VulkanMemory.cpp b/Engine/Source/Runtime/VulkanRHI/Private/VulkanMemory.cpp
--- a/Engine/Source/Runtime/VulkanRHI/Private/VulkanMemory.cpp
+++ b/Engine/Source/Runtime/VulkanRHI/Private/VulkanMemory.cpp
@@ -3252,17 +3252,20 @@ namespace VulkanRHI
 
 			for(FVulkanResourceHeap* Heap : ResourceTypeHeaps)
 			{				
-				for(FVulkanSubresourceAllocator* Allocator : Heap->UsedDedicatedImagePages)
+				if (Heap)
 				{
-					DedicatedStats.Pages++;
-					DedicatedStats.TotalMemory += Allocator->MaxSize;
-					DedicatedStats.UsedImageMemory += Allocator->UsedSize;
-					DedicatedStats.ImageAllocations += Allocator->NumSubAllocations;
-				}
-				if((int32)HeapIndex < HeapSummary.Num())
-				{
-					HeapSummary[HeapIndex] += DedicatedStats;
-					HeapIndex++;
+					for (FVulkanSubresourceAllocator* Allocator : Heap->UsedDedicatedImagePages)
+					{
+						DedicatedStats.Pages++;
+						DedicatedStats.TotalMemory += Allocator->MaxSize;
+						DedicatedStats.UsedImageMemory += Allocator->UsedSize;
+						DedicatedStats.ImageAllocations += Allocator->NumSubAllocations;
+					}
+					if ((int32)HeapIndex < HeapSummary.Num())
+					{
+						HeapSummary[HeapIndex] += DedicatedStats;
+						HeapIndex++;
+					}
 				}
 			}
 		}
@@ -3519,8 +3522,11 @@ namespace VulkanRHI
 				uint32 SubIndex = 0;
 				for(TArray<FVulkanSubresourceAllocator*>& Foo : ResourceTypeHeaps[Index]->ActivePages)
 				{
-					DumpAllocatorRange(FString::Printf(TEXT("Page[%s] %d"), PageSuffix(SubIndex), Index), Foo);
-					SubIndex++;
+					if (ResourceTypeHeaps[SubIndex])
+					{
+						DumpAllocatorRange(FString::Printf(TEXT("Page[%s] %d"), PageSuffix(SubIndex), Index), Foo);
+						SubIndex++;
+					}
 				}
 				DumpAllocatorRange(FString::Printf(TEXT("UsedDedicatedImagePages %d"), Index ), ResourceTypeHeaps[Index]->UsedDedicatedImagePages);
 
