Workaround for MoltenVK not having RT types.

diff --git a/Engine/Source/Runtime/VulkanRHI/Public/VulkanShaderResources.h b/Engine/Source/Runtime/VulkanRHI/Public/VulkanShaderResources.h
--- a/Engine/Source/Runtime/VulkanRHI/Public/VulkanShaderResources.h
+++ b/Engine/Source/Runtime/VulkanRHI/Public/VulkanShaderResources.h
@@ -25,7 +25,9 @@ static inline VkDescriptorType BindingToDescriptorType(EVulkanBindingType::EType
 	case EVulkanBindingType::StorageTexelBuffer:	return VK_DESCRIPTOR_TYPE_STORAGE_TEXEL_BUFFER;
 	case EVulkanBindingType::StorageBuffer:			return VK_DESCRIPTOR_TYPE_STORAGE_BUFFER;
 	case EVulkanBindingType::InputAttachment:		return VK_DESCRIPTOR_TYPE_INPUT_ATTACHMENT;
+#if !PLATFORM_MAC
 	case EVulkanBindingType::AccelerationStructure:	return VK_DESCRIPTOR_TYPE_ACCELERATION_STRUCTURE_KHR;
+#endif
 	default:
 		check(0);
 		break;
@@ -47,7 +49,9 @@ static inline EVulkanBindingType::EType DescriptorTypeToBinding(VkDescriptorType
 	case VK_DESCRIPTOR_TYPE_STORAGE_TEXEL_BUFFER:		return EVulkanBindingType::StorageTexelBuffer;
 	case VK_DESCRIPTOR_TYPE_STORAGE_BUFFER:				return EVulkanBindingType::StorageBuffer;
 	case VK_DESCRIPTOR_TYPE_INPUT_ATTACHMENT:			return EVulkanBindingType::InputAttachment;
+#if !PLATFORM_MAC
 	case VK_DESCRIPTOR_TYPE_ACCELERATION_STRUCTURE_KHR: return EVulkanBindingType::AccelerationStructure;
+#endif
 	default:
 		check(0);
 		break;
