RHI_RAYTRACING is defined here, we can use it instead of platform defines to ensure an update to date vksdk is available.

diff --git a/Engine/Source/Runtime/VulkanRHI/Public/VulkanShaderResources.h b/Engine/Source/Runtime/VulkanRHI/Public/VulkanShaderResources.h
--- a/Engine/Source/Runtime/VulkanRHI/Public/VulkanShaderResources.h
+++ b/Engine/Source/Runtime/VulkanRHI/Public/VulkanShaderResources.h
@@ -9,7 +9,6 @@
 #include "BoundShaderStateCache.h"
 #include "CrossCompilerCommon.h"
 
-
 static inline VkDescriptorType BindingToDescriptorType(EVulkanBindingType::EType Type)
 {
 	// Make sure these do NOT alias EPackedTypeName*
@@ -25,7 +24,7 @@ static inline VkDescriptorType BindingToDescriptorType(EVulkanBindingType::EType
 	case EVulkanBindingType::StorageTexelBuffer:	return VK_DESCRIPTOR_TYPE_STORAGE_TEXEL_BUFFER;
 	case EVulkanBindingType::StorageBuffer:			return VK_DESCRIPTOR_TYPE_STORAGE_BUFFER;
 	case EVulkanBindingType::InputAttachment:		return VK_DESCRIPTOR_TYPE_INPUT_ATTACHMENT;
-#if !PLATFORM_MAC
+#if RHI_RAYTRACING
 	case EVulkanBindingType::AccelerationStructure:	return VK_DESCRIPTOR_TYPE_ACCELERATION_STRUCTURE_KHR;
 #endif
 	default:
@@ -49,7 +48,7 @@ static inline EVulkanBindingType::EType DescriptorTypeToBinding(VkDescriptorType
 	case VK_DESCRIPTOR_TYPE_STORAGE_TEXEL_BUFFER:		return EVulkanBindingType::StorageTexelBuffer;
 	case VK_DESCRIPTOR_TYPE_STORAGE_BUFFER:				return EVulkanBindingType::StorageBuffer;
 	case VK_DESCRIPTOR_TYPE_INPUT_ATTACHMENT:			return EVulkanBindingType::InputAttachment;
-#if !PLATFORM_MAC
+#if RHI_RAYTRACING
 	case VK_DESCRIPTOR_TYPE_ACCELERATION_STRUCTURE_KHR: return EVulkanBindingType::AccelerationStructure;
 #endif
 	default:
