Split out the TLAS and BLAS size query, allocation and build into separate functions to align with the recent refactors.
Allocation has been moved from the acceleration structure builds to the scene and geometry ctors for now, the next step is to support high level code allocation.
Added RHICalcRayTracingSceneSize and RHICalcRayTracingGeometrySize implementation using the refactor.
Windows compilation for VKRT is on by default now to help catch external regressions. Runtime ray tracing support is still disabled.
A few minor updates to catch up with other changes.

diff --git a/Engine/Source/Runtime/VulkanRHI/Private/VulkanRayTracing.cpp b/Engine/Source/Runtime/VulkanRHI/Private/VulkanRayTracing.cpp
--- a/Engine/Source/Runtime/VulkanRHI/Private/VulkanRayTracing.cpp
+++ b/Engine/Source/Runtime/VulkanRHI/Private/VulkanRayTracing.cpp
@@ -145,43 +145,20 @@ static void Free(FVkRtAllocation& Allocation)
 	}
 }
 
-FVulkanRayTracingGeometry::FVulkanRayTracingGeometry(const FRayTracingGeometryInitializer& Initializer)
+static void GetBLASBuildData(
+	const VkDevice Device,
+	const TArrayView<const FRayTracingGeometrySegment> Segments,
+	const FBufferRHIRef IndexBufferRHI, 
+	const bool bFastBuild,
+	const bool bAllowUpdate,
+	const uint32 IndexStrideInBytes,
+	const EAccelerationStructureBuildMode BuildMode,
+	FVkRtBLASBuildData& BuildData)
 {
-	// Only supporting triangles initially
-	check(Initializer.GeometryType == ERayTracingGeometryType::RTGT_Triangles);
-
-	DebugName = Initializer.DebugName;
-	IndexOffsetInBytes = Initializer.IndexBufferOffset;
-	TotalPrimitiveCount = Initializer.TotalPrimitiveCount;
-	IndexBufferRHI = Initializer.IndexBuffer;
-	bFastBuild = Initializer.bFastBuild;
-	bAllowUpdate = Initializer.bAllowUpdate;
-	Segments = TArray<FRayTracingGeometrySegment>(Initializer.Segments.GetData(), Initializer.Segments.Num());
-}
-
-FVulkanRayTracingGeometry::~FVulkanRayTracingGeometry()
-{
-	Free(Allocation);
-	Free(Scratch);
-}
-
-void FVulkanRayTracingGeometry::BuildAccelerationStructure(FVulkanCommandListContext& CommandContext, EAccelerationStructureBuildMode BuildMode)
-{
-	check(BuildMode == EAccelerationStructureBuildMode::Build); // Temp
-
-	checkf(!IndexBufferRHI || (IndexBufferRHI->GetStride() == 2 || IndexBufferRHI->GetStride() == 4), TEXT("Index buffer must be 16 or 32 bit if in use."));
-	const VkDevice Device = CommandContext.GetDevice()->GetInstanceHandle();
-
-	IndexStrideInBytes = IndexBufferRHI ? IndexBufferRHI->GetStride() : 0;
-
 	FVulkanResourceMultiBuffer* const IndexBuffer = ResourceCast(IndexBufferRHI.GetReference());
 	VkDeviceOrHostAddressConstKHR IndexBufferDeviceAddress = {};
 	IndexBufferDeviceAddress.deviceAddress = IndexBufferRHI ? GetDeviceAddress(Device, IndexBuffer->GetHandle()) : 0;
 
-	TArray<VkAccelerationStructureGeometryKHR, TInlineAllocator<1>> BuildSegments;
-	TArray<VkAccelerationStructureBuildRangeInfoKHR, TInlineAllocator<1>> BuildRanges;
-	TArray<uint32> VertexCounts;
-
 	for (int32 SegmentIndex = 0; SegmentIndex < Segments.Num(); ++SegmentIndex)
 	{
 		const FRayTracingGeometrySegment& Segment = Segments[SegmentIndex];
@@ -190,9 +167,7 @@ void FVulkanRayTracingGeometry::BuildAccelerationStructure(FVulkanCommandListCon
 		VkDeviceOrHostAddressConstKHR VertexBufferDeviceAddress = {};
 		VertexBufferDeviceAddress.deviceAddress = GetDeviceAddress(Device, VertexBuffer->GetHandle());
 
-		// #vkrt_todo: Add explicit vertex count to FRayTracingGeometrySegment instead of this estimate
-		const uint32 MaxSegmentVertices = (VertexBuffer->GetSize() - Segment.VertexBufferOffset) / Segment.VertexBufferStride;
-		VertexCounts.Add(MaxSegmentVertices);
+		BuildData.VertexCounts.Add(Segment.MaxVertices);
 
 		VkAccelerationStructureGeometryKHR SegmentGeometry;
 		ZeroVulkanStruct(SegmentGeometry, VK_STRUCTURE_TYPE_ACCELERATION_STRUCTURE_GEOMETRY_KHR);
@@ -203,11 +178,11 @@ void FVulkanRayTracingGeometry::BuildAccelerationStructure(FVulkanCommandListCon
 
 		// Only support triangles
 		SegmentGeometry.geometryType = VK_GEOMETRY_TYPE_TRIANGLES_KHR;
-		
+
 		SegmentGeometry.geometry.triangles.sType = VK_STRUCTURE_TYPE_ACCELERATION_STRUCTURE_GEOMETRY_TRIANGLES_DATA_KHR;
 		SegmentGeometry.geometry.triangles.vertexFormat = VK_FORMAT_R32G32B32_SFLOAT;
 		SegmentGeometry.geometry.triangles.vertexData = VertexBufferDeviceAddress;
-		SegmentGeometry.geometry.triangles.maxVertex = MaxSegmentVertices;
+		SegmentGeometry.geometry.triangles.maxVertex = Segment.MaxVertices;
 		SegmentGeometry.geometry.triangles.vertexStride = Segment.VertexBufferStride;
 		SegmentGeometry.geometry.triangles.indexData = IndexBufferDeviceAddress;
 
@@ -224,53 +199,78 @@ void FVulkanRayTracingGeometry::BuildAccelerationStructure(FVulkanCommandListCon
 			SegmentGeometry.geometry.triangles.indexType = VK_INDEX_TYPE_NONE_KHR;
 		}
 
-		BuildSegments.Add(SegmentGeometry);
+		BuildData.Segments.Add(SegmentGeometry);
 
 		VkAccelerationStructureBuildRangeInfoKHR RangeInfo = {};
 		RangeInfo.firstVertex = 0;
 
 		// Disabled segments use an empty range. We still build them to keep the sbt valid.
-		RangeInfo.primitiveCount = (Segment.bEnabled) ? MaxSegmentVertices : 0;
+		RangeInfo.primitiveCount = (Segment.bEnabled) ? Segment.MaxVertices : 0;
 
 		RangeInfo.primitiveOffset = 0;
 		RangeInfo.transformOffset = 0;
 
-		BuildRanges.Add(RangeInfo);
+		BuildData.Ranges.Add(RangeInfo);
 	}
 
-	VkAccelerationStructureBuildGeometryInfoKHR BuildGeometryInfo;
-	ZeroVulkanStruct(BuildGeometryInfo, VK_STRUCTURE_TYPE_ACCELERATION_STRUCTURE_BUILD_GEOMETRY_INFO_KHR);
-	BuildGeometryInfo.type = VK_ACCELERATION_STRUCTURE_TYPE_BOTTOM_LEVEL_KHR;
-	BuildGeometryInfo.flags = (bFastBuild) ? VK_BUILD_ACCELERATION_STRUCTURE_PREFER_FAST_BUILD_BIT_KHR : VK_BUILD_ACCELERATION_STRUCTURE_PREFER_FAST_TRACE_BIT_KHR;
+	BuildData.GeometryInfo.type = VK_ACCELERATION_STRUCTURE_TYPE_BOTTOM_LEVEL_KHR;
+	BuildData.GeometryInfo.flags = (bFastBuild) ? VK_BUILD_ACCELERATION_STRUCTURE_PREFER_FAST_BUILD_BIT_KHR : VK_BUILD_ACCELERATION_STRUCTURE_PREFER_FAST_TRACE_BIT_KHR;
 	if (bAllowUpdate)
 	{
-		BuildGeometryInfo.flags |= VK_BUILD_ACCELERATION_STRUCTURE_ALLOW_UPDATE_BIT_KHR;
+		BuildData.GeometryInfo.flags |= VK_BUILD_ACCELERATION_STRUCTURE_ALLOW_UPDATE_BIT_KHR;
 	}
-	BuildGeometryInfo.mode = (BuildMode == EAccelerationStructureBuildMode::Build) ? VK_BUILD_ACCELERATION_STRUCTURE_MODE_BUILD_KHR : VK_BUILD_ACCELERATION_STRUCTURE_MODE_UPDATE_KHR;
-	BuildGeometryInfo.geometryCount = BuildSegments.Num();
-	BuildGeometryInfo.pGeometries = BuildSegments.GetData();
+	BuildData.GeometryInfo.mode = (BuildMode == EAccelerationStructureBuildMode::Build) ? VK_BUILD_ACCELERATION_STRUCTURE_MODE_BUILD_KHR : VK_BUILD_ACCELERATION_STRUCTURE_MODE_UPDATE_KHR;
+	BuildData.GeometryInfo.geometryCount = BuildData.Segments.Num();
+	BuildData.GeometryInfo.pGeometries = BuildData.Segments.GetData();
 
-	VkAccelerationStructureBuildSizesInfoKHR BuildSizesInfo;
-	ZeroVulkanStruct(BuildSizesInfo, VK_STRUCTURE_TYPE_ACCELERATION_STRUCTURE_BUILD_SIZES_INFO_KHR);
 	vkGetAccelerationStructureBuildSizesKHR(
 		Device,
 		VK_ACCELERATION_STRUCTURE_BUILD_TYPE_DEVICE_KHR,
-		&BuildGeometryInfo,
-		VertexCounts.GetData(),
-		&BuildSizesInfo);
+		&BuildData.GeometryInfo,
+		BuildData.VertexCounts.GetData(),
+		&BuildData.SizesInfo);
+}
+
+FVulkanRayTracingGeometry::FVulkanRayTracingGeometry(const FRayTracingGeometryInitializer& Initializer, const FVulkanDevice* InDevice)
+{
+	// Only supporting triangles initially
+	check(Initializer.GeometryType == ERayTracingGeometryType::RTGT_Triangles);
+
+	Device = InDevice;
+	DebugName = Initializer.DebugName;
+	IndexOffsetInBytes = Initializer.IndexBufferOffset;
+	TotalPrimitiveCount = Initializer.TotalPrimitiveCount;
+	IndexBufferRHI = Initializer.IndexBuffer;
+	bFastBuild = Initializer.bFastBuild;
+	bAllowUpdate = Initializer.bAllowUpdate;
+	Segments = TArray<FRayTracingGeometrySegment>(Initializer.Segments.GetData(), Initializer.Segments.Num());
+
+	checkf(!IndexBufferRHI || (IndexBufferRHI->GetStride() == 2 || IndexBufferRHI->GetStride() == 4), TEXT("Index buffer must be 16 or 32 bit if in use."));
+	IndexStrideInBytes = IndexBufferRHI ? IndexBufferRHI->GetStride() : 0;
+
+	FVkRtBLASBuildData BuildData;
+	GetBLASBuildData(
+		Device->GetInstanceHandle(), 
+		MakeArrayView(Segments),
+		IndexBufferRHI, 
+		bFastBuild, 
+		bAllowUpdate, 
+		IndexStrideInBytes, 
+		EAccelerationStructureBuildMode::Build,
+		BuildData);
 
 	Allocate(
-		CommandContext.GetDevice()->GetPhysicalHandle(), 
-		Device, 
-		BuildSizesInfo.accelerationStructureSize,
+		Device->GetPhysicalHandle(), 
+		Device->GetInstanceHandle(),
+		BuildData.SizesInfo.accelerationStructureSize,
 		VK_BUFFER_USAGE_ACCELERATION_STRUCTURE_STORAGE_BIT_KHR | VK_BUFFER_USAGE_SHADER_DEVICE_ADDRESS_BIT, 
 		VK_MEMORY_PROPERTY_DEVICE_LOCAL_BIT,
 		Allocation);
 
 	Allocate(
-		CommandContext.GetDevice()->GetPhysicalHandle(),
-		Device,
-		BuildSizesInfo.buildScratchSize,
+		Device->GetPhysicalHandle(),
+		Device->GetInstanceHandle(),
+		BuildData.SizesInfo.buildScratchSize,
 		VK_BUFFER_USAGE_STORAGE_BUFFER_BIT | VK_BUFFER_USAGE_SHADER_DEVICE_ADDRESS_BIT,
 		VK_MEMORY_PROPERTY_DEVICE_LOCAL_BIT,
 		Scratch);
@@ -278,23 +278,45 @@ void FVulkanRayTracingGeometry::BuildAccelerationStructure(FVulkanCommandListCon
 	VkBufferDeviceAddressInfoKHR ScratchDeviceAddressInfo;
 	ZeroVulkanStruct(ScratchDeviceAddressInfo, VK_STRUCTURE_TYPE_BUFFER_DEVICE_ADDRESS_INFO);
 	ScratchDeviceAddressInfo.buffer = Scratch.Buffer;
-	Scratch.Address = vkGetBufferDeviceAddressKHR(Device, &ScratchDeviceAddressInfo);
+	Scratch.Address = vkGetBufferDeviceAddressKHR(Device->GetInstanceHandle(), &ScratchDeviceAddressInfo);
+}
+
+FVulkanRayTracingGeometry::~FVulkanRayTracingGeometry()
+{
+	Free(Allocation);
+	Free(Scratch);
+}
+
+void FVulkanRayTracingGeometry::BuildAccelerationStructure(FVulkanCommandListContext& CommandContext, EAccelerationStructureBuildMode BuildMode)
+{
+	check(BuildMode == EAccelerationStructureBuildMode::Build); // Temp
+
+	FVkRtBLASBuildData BuildData;
+	GetBLASBuildData(
+		Device->GetInstanceHandle(),
+		MakeArrayView(Segments),
+		IndexBufferRHI,
+		bFastBuild,
+		bAllowUpdate,
+		IndexStrideInBytes,
+		BuildMode,
+		BuildData);
 
 	VkAccelerationStructureCreateInfoKHR CreateInfo;
 	ZeroVulkanStruct(CreateInfo, VK_STRUCTURE_TYPE_ACCELERATION_STRUCTURE_CREATE_INFO_KHR);
 	CreateInfo.buffer = Allocation.Buffer;
-	CreateInfo.size = BuildSizesInfo.accelerationStructureSize;
+	CreateInfo.size = BuildData.SizesInfo.accelerationStructureSize;
 	CreateInfo.type = VK_ACCELERATION_STRUCTURE_TYPE_BOTTOM_LEVEL_KHR;
-	VERIFYVULKANRESULT(vkCreateAccelerationStructureKHR(Device, &CreateInfo, VULKAN_CPU_ALLOCATOR, &Handle));
+	VERIFYVULKANRESULT(vkCreateAccelerationStructureKHR(Device->GetInstanceHandle(), &CreateInfo, VULKAN_CPU_ALLOCATOR, &Handle));
 
-	BuildGeometryInfo.dstAccelerationStructure = Handle;
-	BuildGeometryInfo.scratchData.deviceAddress = Scratch.Address;
+	BuildData.GeometryInfo.dstAccelerationStructure = Handle;
+	BuildData.GeometryInfo.scratchData.deviceAddress = Scratch.Address;
 
-	VkAccelerationStructureBuildRangeInfoKHR* const pBuildRanges = BuildRanges.GetData();
+	VkAccelerationStructureBuildRangeInfoKHR* const pBuildRanges = BuildData.Ranges.GetData();
 
 	FVulkanCommandBufferManager& CommandBufferManager = *CommandContext.GetCommandBufferManager();
 	FVulkanCmdBuffer* const CmdBuffer = CommandBufferManager.GetActiveCmdBuffer();
-	vkCmdBuildAccelerationStructuresKHR(CmdBuffer->GetHandle(), 1, &BuildGeometryInfo, &pBuildRanges);
+	vkCmdBuildAccelerationStructuresKHR(CmdBuffer->GetHandle(), 1, &BuildData.GeometryInfo, &pBuildRanges);
 
 	CommandBufferManager.SubmitActiveCmdBuffer();
 	CommandBufferManager.PrepareForNewActiveCommandBuffer();
@@ -302,7 +324,7 @@ void FVulkanRayTracingGeometry::BuildAccelerationStructure(FVulkanCommandListCon
 	VkAccelerationStructureDeviceAddressInfoKHR DeviceAddressInfo;
 	ZeroVulkanStruct(DeviceAddressInfo, VK_STRUCTURE_TYPE_ACCELERATION_STRUCTURE_DEVICE_ADDRESS_INFO_KHR);
 	DeviceAddressInfo.accelerationStructure = Handle;
-	Allocation.Address = vkGetAccelerationStructureDeviceAddressKHR(Device, &DeviceAddressInfo);
+	Allocation.Address = vkGetAccelerationStructureDeviceAddressKHR(Device->GetInstanceHandle(), &DeviceAddressInfo);
 
 	// No longer need scratch memory for a static build
 	if (!bAllowUpdate)
@@ -311,17 +333,73 @@ void FVulkanRayTracingGeometry::BuildAccelerationStructure(FVulkanCommandListCon
 	}
 }
 
-FVulkanRayTracingScene::FVulkanRayTracingScene(const FRayTracingSceneInitializer& Initializer)
+static void GetTLASBuildData(
+	const VkDevice Device,
+	const uint32 NumInstances,
+	const VkDeviceAddress InstanceBufferAddress,
+	FVkRtTLASBuildData& BuildData)
 {
-	DebugName = Initializer.DebugName;
+	VkDeviceOrHostAddressConstKHR InstanceBufferDeviceAddress = {};
+	InstanceBufferDeviceAddress.deviceAddress = InstanceBufferAddress;
+
+	BuildData.Geometry.geometryType = VK_GEOMETRY_TYPE_INSTANCES_KHR;
+	BuildData.Geometry.geometry.instances.sType = VK_STRUCTURE_TYPE_ACCELERATION_STRUCTURE_GEOMETRY_INSTANCES_DATA_KHR;
+	BuildData.Geometry.geometry.instances.arrayOfPointers = VK_FALSE;
+	BuildData.Geometry.geometry.instances.data = InstanceBufferDeviceAddress;
+
+	BuildData.GeometryInfo.type = VK_ACCELERATION_STRUCTURE_TYPE_TOP_LEVEL_KHR;
+	BuildData.GeometryInfo.mode = VK_BUILD_ACCELERATION_STRUCTURE_MODE_BUILD_KHR;
+	BuildData.GeometryInfo.flags = VK_BUILD_ACCELERATION_STRUCTURE_PREFER_FAST_TRACE_BIT_KHR;
+	BuildData.GeometryInfo.geometryCount = 1;
+	BuildData.GeometryInfo.pGeometries = &BuildData.Geometry;
 
-	uint32 NumInstances = 0;
+	vkGetAccelerationStructureBuildSizesKHR(
+		Device,
+		VK_ACCELERATION_STRUCTURE_BUILD_TYPE_DEVICE_KHR,
+		&BuildData.GeometryInfo,
+		&NumInstances,
+		&BuildData.SizesInfo);
+}
+
+static VkGeometryInstanceFlagsKHR TranslateRayTracingInstanceFlags(ERayTracingInstanceFlags InFlags)
+{
+	VkGeometryInstanceFlagsKHR Result = 0;
+
+	if (EnumHasAnyFlags(InFlags, ERayTracingInstanceFlags::TriangleCullDisable))
+	{
+		Result |= VK_GEOMETRY_INSTANCE_TRIANGLE_FACING_CULL_DISABLE_BIT_KHR;
+	}
+
+	if (!EnumHasAnyFlags(InFlags, ERayTracingInstanceFlags::TriangleCullReverse))
+	{
+		// Counterclockwise is the default for UE.
+		Result |= VK_GEOMETRY_INSTANCE_TRIANGLE_FRONT_COUNTERCLOCKWISE_BIT_KHR;
+	}
+
+	if (EnumHasAnyFlags(InFlags, ERayTracingInstanceFlags::ForceOpaque))
+	{
+		Result |= VK_GEOMETRY_INSTANCE_FORCE_OPAQUE_BIT_KHR;
+	}
+
+	if (EnumHasAnyFlags(InFlags, ERayTracingInstanceFlags::ForceNonOpaque))
+	{
+		Result |= VK_GEOMETRY_INSTANCE_FORCE_NO_OPAQUE_BIT_KHR;
+	}
+
+	return Result;
+}
+
+FVulkanRayTracingScene::FVulkanRayTracingScene(const FRayTracingSceneInitializer& Initializer, const FVulkanDevice* InDevice)
+{
+	Device = InDevice;
+	DebugName = Initializer.DebugName;
+	
 	Experimental::TSherwoodSet<const FVulkanRayTracingGeometry*> UniqueGeometryReferences;
 
 	for (const FRayTracingGeometryInstance& SceneInstance : Initializer.Instances)
 	{
 		// Don't support GPU transforms yet.
-		ensure(!SceneInstance.GPUTransformsSRV.IsValid());
+		ensure(SceneInstance.GPUTransformsSRV == nullptr);
 
 		// Hold a reference to each unique geometry object in the scene
 		const FVulkanRayTracingGeometry* const Geometry = ResourceCast(SceneInstance.GeometryRHI);
@@ -343,58 +421,40 @@ FVulkanRayTracingScene::FVulkanRayTracingScene(const FRayTracingSceneInitializer
 		const bool bUseUniqueUserData = SceneInstance.UserData.Num() > 1;
 		const uint32 CommonUserData = SceneInstance.UserData.Num() == 1 ? SceneInstance.UserData[0] : 0;
 
-		TArrayView<const FMatrix> Transforms = SceneInstance.GetTransforms();
-		for (int32 TransformIndex = 0; TransformIndex < Transforms.Num(); ++TransformIndex)
+		for (int32 TransformIndex = 0; TransformIndex < SceneInstance.Transforms.Num(); ++TransformIndex)
 		{
 			InstanceGeometry[InstanceDescIndex] = ResourceCast(SceneInstance.GeometryRHI);
 
 			VkAccelerationStructureInstanceKHR& InstanceDesc = InstanceDescs[InstanceDescIndex];
-			const FMatrix TransposedTransform = Transforms[TransformIndex].GetTransposed();
+			const FMatrix TransposedTransform = SceneInstance.Transforms[TransformIndex].GetTransposed();
 			FMemory::Memcpy(&InstanceDesc.transform, &TransposedTransform.M[0][0], sizeof(VkTransformMatrixKHR));
 
 			InstanceDesc.accelerationStructureReference = 0; // Set during TLAS build after the BLAS for the referenced geometry is built.
 			InstanceDesc.instanceCustomIndex = (bUseUniqueUserData) ? SceneInstance.UserData[TransformIndex] : CommonUserData;
 			InstanceDesc.mask = SceneInstance.Mask;
 			InstanceDesc.instanceShaderBindingTableRecordOffset = 0; // Todo
-			InstanceDesc.flags = VK_GEOMETRY_INSTANCE_TRIANGLE_FRONT_COUNTERCLOCKWISE_BIT_KHR;
-
-			if (SceneInstance.bForceOpaque)
-			{
-				InstanceDesc.flags |= VK_GEOMETRY_INSTANCE_FORCE_OPAQUE_BIT_KHR;
-			}
-
-			if (SceneInstance.bDoubleSided)
-			{
-				InstanceDesc.flags |= VK_GEOMETRY_INSTANCE_TRIANGLE_FACING_CULL_DISABLE_BIT_KHR;
-			}
+			InstanceDesc.flags = TranslateRayTracingInstanceFlags(SceneInstance.Flags);
 
 			++InstanceDescIndex;
 		}
 	}
-}
 
-FVulkanRayTracingScene::~FVulkanRayTracingScene()
-{
-	Free(Allocation);
-}
-
-void FVulkanRayTracingScene::BuildAccelerationStructure(FVulkanCommandListContext& CommandContext)
-{
-	const VkDevice Device = CommandContext.GetDevice()->GetInstanceHandle();
-	const uint32 NumInstances = static_cast<uint32>(InstanceDescs.Num());
-	const uint32 InstanceBufferByteSize = NumInstances * sizeof(VkAccelerationStructureInstanceKHR);
-	FVkRtAllocation InstanceBuffer;
+	// Allocate instance buffer
+	const uint32 InstanceBufferByteSize = static_cast<uint32>(NumInstances) * sizeof(VkAccelerationStructureInstanceKHR);
 
 	Allocate(
-		CommandContext.GetDevice()->GetPhysicalHandle(),
-		Device,
+		Device->GetPhysicalHandle(),
+		Device->GetInstanceHandle(),
 		InstanceBufferByteSize,
 		VK_BUFFER_USAGE_SHADER_DEVICE_ADDRESS_BIT | VK_BUFFER_USAGE_ACCELERATION_STRUCTURE_BUILD_INPUT_READ_ONLY_BIT_KHR,
 		VK_MEMORY_PROPERTY_HOST_VISIBLE_BIT | VK_MEMORY_PROPERTY_HOST_COHERENT_BIT,
 		InstanceBuffer);
 
+	InstanceBuffer.Address = GetDeviceAddress(Device->GetInstanceHandle(), InstanceBuffer.Buffer);
+
+	// Copy instance data
 	void* pMappedInstanceBufferMemory = nullptr;
-	VERIFYVULKANRESULT(vkMapMemory(Device, InstanceBuffer.Memory, 0, InstanceBufferByteSize, 0, &pMappedInstanceBufferMemory));
+	VERIFYVULKANRESULT(vkMapMemory(Device->GetInstanceHandle(), InstanceBuffer.Memory, 0, InstanceBufferByteSize, 0, &pMappedInstanceBufferMemory));
 	{
 		VkAccelerationStructureInstanceKHR* const InstanceDescBuffer = reinterpret_cast<VkAccelerationStructureInstanceKHR*>(pMappedInstanceBufferMemory);
 		FMemory::Memcpy(InstanceDescBuffer, InstanceDescs.GetData(), InstanceBufferByteSize);
@@ -404,68 +464,53 @@ void FVulkanRayTracingScene::BuildAccelerationStructure(FVulkanCommandListContex
 			InstanceDescBuffer[InstanceIndex].accelerationStructureReference = InstanceGeometry[InstanceIndex]->GetAccelerationStructureAddress();
 		}
 	}
-	vkUnmapMemory(Device, InstanceBuffer.Memory);
+	vkUnmapMemory(Device->GetInstanceHandle(), InstanceBuffer.Memory);
 
 	InstanceDescs.Empty();
 	InstanceGeometry.Empty();
 
-	VkDeviceOrHostAddressConstKHR InstanceBufferDeviceAddress = {};
-	InstanceBufferDeviceAddress.deviceAddress = GetDeviceAddress(Device, InstanceBuffer.Buffer);
-
-	VkAccelerationStructureGeometryKHR TLASGeometry;
-	ZeroVulkanStruct(TLASGeometry, VK_STRUCTURE_TYPE_ACCELERATION_STRUCTURE_GEOMETRY_KHR);
-	TLASGeometry.geometryType = VK_GEOMETRY_TYPE_INSTANCES_KHR;
-	TLASGeometry.geometry.instances.sType = VK_STRUCTURE_TYPE_ACCELERATION_STRUCTURE_GEOMETRY_INSTANCES_DATA_KHR;
-	TLASGeometry.geometry.instances.arrayOfPointers = VK_FALSE;
-	TLASGeometry.geometry.instances.data = InstanceBufferDeviceAddress;
-
-	VkAccelerationStructureBuildGeometryInfoKHR TLASGeometryInfo;
-	ZeroVulkanStruct(TLASGeometryInfo, VK_STRUCTURE_TYPE_ACCELERATION_STRUCTURE_BUILD_GEOMETRY_INFO_KHR);
-	TLASGeometryInfo.type = VK_ACCELERATION_STRUCTURE_TYPE_TOP_LEVEL_KHR;
-	TLASGeometryInfo.mode = VK_BUILD_ACCELERATION_STRUCTURE_MODE_BUILD_KHR;
-	TLASGeometryInfo.flags = VK_BUILD_ACCELERATION_STRUCTURE_PREFER_FAST_TRACE_BIT_KHR;
-	TLASGeometryInfo.geometryCount = 1;
-	TLASGeometryInfo.pGeometries = &TLASGeometry;
-
-	VkAccelerationStructureBuildSizesInfoKHR TLASBuildSizesInfo;
-	ZeroVulkanStruct(TLASBuildSizesInfo, VK_STRUCTURE_TYPE_ACCELERATION_STRUCTURE_BUILD_SIZES_INFO_KHR);
-	vkGetAccelerationStructureBuildSizesKHR(
-		Device, 
-		VK_ACCELERATION_STRUCTURE_BUILD_TYPE_DEVICE_KHR, 
-		&TLASGeometryInfo, 
-		&NumInstances, 
-		&TLASBuildSizesInfo);
+	// Allocate buffers
+	FVkRtTLASBuildData BuildData;
+	GetTLASBuildData(Device->GetInstanceHandle(), NumInstances, 0, BuildData);
 
 	Allocate(
-		CommandContext.GetDevice()->GetPhysicalHandle(),
-		Device,
-		TLASBuildSizesInfo.accelerationStructureSize,
+		Device->GetPhysicalHandle(),
+		Device->GetInstanceHandle(),
+		BuildData.SizesInfo.accelerationStructureSize,
 		VK_BUFFER_USAGE_ACCELERATION_STRUCTURE_STORAGE_BIT_KHR | VK_BUFFER_USAGE_SHADER_DEVICE_ADDRESS_BIT,
 		VK_MEMORY_PROPERTY_DEVICE_LOCAL_BIT,
 		Allocation);
 
 	Allocate(
-		CommandContext.GetDevice()->GetPhysicalHandle(),
-		Device,
-		TLASBuildSizesInfo.buildScratchSize,
+		Device->GetPhysicalHandle(),
+		Device->GetInstanceHandle(),
+		BuildData.SizesInfo.buildScratchSize,
 		VK_BUFFER_USAGE_STORAGE_BUFFER_BIT | VK_BUFFER_USAGE_SHADER_DEVICE_ADDRESS_BIT,
 		VK_MEMORY_PROPERTY_DEVICE_LOCAL_BIT,
 		Scratch);
 
-	VkBufferDeviceAddressInfoKHR ScratchDeviceAddressInfo;
-	ZeroVulkanStruct(ScratchDeviceAddressInfo, VK_STRUCTURE_TYPE_BUFFER_DEVICE_ADDRESS_INFO);
-	ScratchDeviceAddressInfo.buffer = Scratch.Buffer;
-	Scratch.Address = vkGetBufferDeviceAddressKHR(Device, &ScratchDeviceAddressInfo);
+	Scratch.Address = GetDeviceAddress(Device->GetInstanceHandle(), Scratch.Buffer);
+}
+
+FVulkanRayTracingScene::~FVulkanRayTracingScene()
+{
+	Free(Allocation);
+}
+
+void FVulkanRayTracingScene::BuildAccelerationStructure(FVulkanCommandListContext& CommandContext)
+{
+	FVkRtTLASBuildData BuildData;
+	GetTLASBuildData(Device->GetInstanceHandle(), NumInstances, InstanceBuffer.Address, BuildData);
 
 	VkAccelerationStructureCreateInfoKHR TLASCreateInfo;
 	ZeroVulkanStruct(TLASCreateInfo, VK_STRUCTURE_TYPE_ACCELERATION_STRUCTURE_CREATE_INFO_KHR);
 	TLASCreateInfo.buffer = Allocation.Buffer;
-	TLASCreateInfo.size = TLASBuildSizesInfo.accelerationStructureSize;
+	TLASCreateInfo.size = BuildData.SizesInfo.accelerationStructureSize;
 	TLASCreateInfo.type = VK_ACCELERATION_STRUCTURE_TYPE_TOP_LEVEL_KHR;
-	VERIFYVULKANRESULT(vkCreateAccelerationStructureKHR(Device, &TLASCreateInfo, VULKAN_CPU_ALLOCATOR, &Handle));
+	VERIFYVULKANRESULT(vkCreateAccelerationStructureKHR(Device->GetInstanceHandle(), &TLASCreateInfo, VULKAN_CPU_ALLOCATOR, &Handle));
 
-	TLASGeometryInfo.dstAccelerationStructure = Handle;
-	TLASGeometryInfo.scratchData.deviceAddress = Scratch.Address;
+	BuildData.GeometryInfo.dstAccelerationStructure = Handle;
+	BuildData.GeometryInfo.scratchData.deviceAddress = Scratch.Address;
 
 	VkAccelerationStructureBuildRangeInfoKHR TLASBuildRangeInfo;
 	TLASBuildRangeInfo.primitiveCount = NumInstances;
@@ -477,7 +522,7 @@ void FVulkanRayTracingScene::BuildAccelerationStructure(FVulkanCommandListContex
 
 	FVulkanCommandBufferManager& CommandBufferManager = *CommandContext.GetCommandBufferManager();
 	FVulkanCmdBuffer* const CmdBuffer = CommandBufferManager.GetActiveCmdBuffer();
-	vkCmdBuildAccelerationStructuresKHR(CmdBuffer->GetHandle(), 1, &TLASGeometryInfo, &pBuildRanges);
+	vkCmdBuildAccelerationStructuresKHR(CmdBuffer->GetHandle(), 1, &BuildData.GeometryInfo, &pBuildRanges);
 
 	CommandBufferManager.SubmitActiveCmdBuffer();
 	CommandBufferManager.PrepareForNewActiveCommandBuffer();
@@ -486,14 +531,51 @@ void FVulkanRayTracingScene::BuildAccelerationStructure(FVulkanCommandListContex
 	Free(Scratch);
 }
 
+FRayTracingAccelerationStructureSize FVulkanDynamicRHI::RHICalcRayTracingSceneSize(uint32 MaxInstances, ERayTracingAccelerationStructureFlags Flags)
+{
+	FVkRtTLASBuildData BuildData;
+	const VkDeviceAddress InstanceBufferAddress = 0; // No device address available when only querying TLAS size
+	GetTLASBuildData(Device->GetInstanceHandle(), MaxInstances, InstanceBufferAddress, BuildData);
+
+	FRayTracingAccelerationStructureSize Result;
+	Result.ResultSize = BuildData.SizesInfo.accelerationStructureSize;
+	Result.BuildScratchSize = BuildData.SizesInfo.buildScratchSize;
+	Result.UpdateScratchSize = BuildData.SizesInfo.updateScratchSize;
+
+	return Result;
+}
+
+FRayTracingAccelerationStructureSize FVulkanDynamicRHI::RHICalcRayTracingGeometrySize(const FRayTracingGeometryInitializer& Initializer)
+{	
+	const uint32 IndexStrideInBytes = Initializer.IndexBuffer ? Initializer.IndexBuffer->GetStride() : 0;
+
+	FVkRtBLASBuildData BuildData;
+	GetBLASBuildData(
+		Device->GetInstanceHandle(),
+		MakeArrayView(Initializer.Segments),
+		Initializer.IndexBuffer,
+		Initializer.bFastBuild,
+		Initializer.bAllowUpdate,
+		IndexStrideInBytes,
+		EAccelerationStructureBuildMode::Build,
+		BuildData);
+
+	FRayTracingAccelerationStructureSize Result;
+	Result.ResultSize = BuildData.SizesInfo.accelerationStructureSize;
+	Result.BuildScratchSize = BuildData.SizesInfo.buildScratchSize;
+	Result.UpdateScratchSize = BuildData.SizesInfo.updateScratchSize;
+	
+	return Result;
+}
+
 FRayTracingSceneRHIRef FVulkanDynamicRHI::RHICreateRayTracingScene(const FRayTracingSceneInitializer& Initializer)
 {
-	return new FVulkanRayTracingScene(Initializer);
+	return new FVulkanRayTracingScene(Initializer, GetDevice());
 }
 
 FRayTracingGeometryRHIRef FVulkanDynamicRHI::RHICreateRayTracingGeometry(const FRayTracingGeometryInitializer& Initializer)
 {
-	return new FVulkanRayTracingGeometry(Initializer);
+	return new FVulkanRayTracingGeometry(Initializer, GetDevice());
 }
 
 // Todo: High level rhi call should have transitioned and verified vb and ib to read for each segment
diff --git a/Engine/Source/Runtime/VulkanRHI/Private/VulkanRayTracing.h b/Engine/Source/Runtime/VulkanRHI/Private/VulkanRayTracing.h
--- a/Engine/Source/Runtime/VulkanRHI/Private/VulkanRayTracing.h
+++ b/Engine/Source/Runtime/VulkanRHI/Private/VulkanRayTracing.h
@@ -36,13 +36,42 @@ struct FVkRtAllocation
 	VkDevice Device = VK_NULL_HANDLE;
 	VkDeviceMemory Memory = VK_NULL_HANDLE;
 	VkBuffer Buffer = VK_NULL_HANDLE;
-	VkDeviceAddress Address = {};
+	VkDeviceAddress Address = 0;
+};
+
+struct FVkRtTLASBuildData
+{
+	FVkRtTLASBuildData()
+	{
+		ZeroVulkanStruct(Geometry, VK_STRUCTURE_TYPE_ACCELERATION_STRUCTURE_GEOMETRY_KHR);
+		ZeroVulkanStruct(GeometryInfo, VK_STRUCTURE_TYPE_ACCELERATION_STRUCTURE_BUILD_GEOMETRY_INFO_KHR);
+		ZeroVulkanStruct(SizesInfo, VK_STRUCTURE_TYPE_ACCELERATION_STRUCTURE_BUILD_SIZES_INFO_KHR);
+	}
+
+	VkAccelerationStructureGeometryKHR Geometry;
+	VkAccelerationStructureBuildGeometryInfoKHR GeometryInfo;
+	VkAccelerationStructureBuildSizesInfoKHR SizesInfo;
+};
+
+struct FVkRtBLASBuildData
+{
+	FVkRtBLASBuildData()
+	{
+		ZeroVulkanStruct(GeometryInfo, VK_STRUCTURE_TYPE_ACCELERATION_STRUCTURE_BUILD_GEOMETRY_INFO_KHR);
+		ZeroVulkanStruct(SizesInfo, VK_STRUCTURE_TYPE_ACCELERATION_STRUCTURE_BUILD_SIZES_INFO_KHR);
+	}
+
+	TArray<VkAccelerationStructureGeometryKHR, TInlineAllocator<1>> Segments;
+	TArray<VkAccelerationStructureBuildRangeInfoKHR, TInlineAllocator<1>> Ranges;
+	TArray<uint32, TInlineAllocator<1>> VertexCounts;
+	VkAccelerationStructureBuildGeometryInfoKHR GeometryInfo;
+	VkAccelerationStructureBuildSizesInfoKHR SizesInfo;
 };
 
 class FVulkanRayTracingGeometry : public FRHIRayTracingGeometry
 {
 public:
-	FVulkanRayTracingGeometry(const FRayTracingGeometryInitializer& Initializer);
+	FVulkanRayTracingGeometry(const FRayTracingGeometryInitializer& Initializer, const FVulkanDevice* InDevice);
 	~FVulkanRayTracingGeometry();
 
 	void BuildAccelerationStructure(FVulkanCommandListContext& CommandContext, EAccelerationStructureBuildMode BuildMode);
@@ -69,12 +98,13 @@ private:
 	TArray<FRayTracingGeometrySegment> Segments;
 
 	FName DebugName;
+	const FVulkanDevice* Device = nullptr;
 };
 
 class FVulkanRayTracingScene : public FRHIRayTracingScene
 {
 public:
-	FVulkanRayTracingScene(const FRayTracingSceneInitializer& Initializer);
+	FVulkanRayTracingScene(const FRayTracingSceneInitializer& Initializer, const FVulkanDevice* InDevice);
 	~FVulkanRayTracingScene();
 
 	void BuildAccelerationStructure(FVulkanCommandListContext& CommandContext);
@@ -82,11 +112,16 @@ public:
 	VkAccelerationStructureKHR Handle = VK_NULL_HANDLE;
 	FVkRtAllocation Allocation;
 	FVkRtAllocation Scratch;
+	FVkRtAllocation InstanceBuffer;
 
 	TArray<VkAccelerationStructureInstanceKHR> InstanceDescs;
 	TArray<const FVulkanRayTracingGeometry*> InstanceGeometry;
 	TArray<TRefCountPtr<const FVulkanRayTracingGeometry>> ReferencedGeometry;
 	FName DebugName;
+	uint32 NumInstances = 0;
+
+private:
+	const FVulkanDevice* Device = nullptr;
 };
 
 #endif // #if VULKAN_RHI_RAYTRACING
diff --git a/Engine/Source/Runtime/VulkanRHI/Private/Windows/VulkanWindowsPlatform.h b/Engine/Source/Runtime/VulkanRHI/Private/Windows/VulkanWindowsPlatform.h
--- a/Engine/Source/Runtime/VulkanRHI/Private/Windows/VulkanWindowsPlatform.h
+++ b/Engine/Source/Runtime/VulkanRHI/Private/Windows/VulkanWindowsPlatform.h
@@ -19,8 +19,7 @@
 
 #define	UE_VK_API_VERSION							VK_API_VERSION_1_1
 
-//#define VULKAN_RHI_RAYTRACING (RHI_RAYTRACING)
-#define VULKAN_RHI_RAYTRACING 0 // Disabled for development
+#define VULKAN_RHI_RAYTRACING (RHI_RAYTRACING)
 #define VULKAN_SUPPORTS_SCALAR_BLOCK_LAYOUT (VULKAN_RHI_RAYTRACING)
 
 // 32-bit windows has warnings on custom mem mgr callbacks
diff --git a/Engine/Source/Runtime/VulkanRHI/Public/VulkanDynamicRHI.h b/Engine/Source/Runtime/VulkanRHI/Public/VulkanDynamicRHI.h
--- a/Engine/Source/Runtime/VulkanRHI/Public/VulkanDynamicRHI.h
+++ b/Engine/Source/Runtime/VulkanRHI/Public/VulkanDynamicRHI.h
@@ -323,6 +323,8 @@ public:
 	void RHICalibrateTimers() override;
 
 #if VULKAN_RHI_RAYTRACING
+	virtual FRayTracingAccelerationStructureSize RHICalcRayTracingSceneSize(uint32 MaxInstances, ERayTracingAccelerationStructureFlags Flags) final override;
+	virtual FRayTracingAccelerationStructureSize RHICalcRayTracingGeometrySize(const FRayTracingGeometryInitializer& Initializer) final override;
 	virtual FRayTracingGeometryRHIRef RHICreateRayTracingGeometry(const FRayTracingGeometryInitializer& Initializer) final override;
 	virtual FRayTracingSceneRHIRef RHICreateRayTracingScene(const FRayTracingSceneInitializer& Initializer) final override;
 #endif
