Use external TLAS allocations (still backed by the temporary vkrt 'allocator')
Optionally use external scratch allocations (backed by rdg)

diff --git a/Engine/Source/Runtime/VulkanRHI/Private/VulkanBarriers.cpp b/Engine/Source/Runtime/VulkanRHI/Private/VulkanBarriers.cpp
--- a/Engine/Source/Runtime/VulkanRHI/Private/VulkanBarriers.cpp
+++ b/Engine/Source/Runtime/VulkanRHI/Private/VulkanBarriers.cpp
@@ -203,6 +203,19 @@ static void GetVkStageAndAccessFlags(ERHIAccess RHIAccess, FRHITransitionInfo::E
 			AccessFlags = VK_ACCESS_TRANSFER_WRITE_BIT;
 			Layout = VK_IMAGE_LAYOUT_TRANSFER_DST_OPTIMAL;
 			return;
+
+#if VULKAN_RHI_RAYTRACING
+		// vkrt todo: Finer grain stage flags would be ideal here.
+		case ERHIAccess::BVHRead:
+			StageFlags = VK_PIPELINE_STAGE_ACCELERATION_STRUCTURE_BUILD_BIT_KHR | VK_PIPELINE_STAGE_RAY_TRACING_SHADER_BIT_KHR | VK_PIPELINE_STAGE_COMPUTE_SHADER_BIT;
+			AccessFlags = VK_ACCESS_ACCELERATION_STRUCTURE_READ_BIT_KHR;
+			return;
+
+		case ERHIAccess::BVHWrite:
+			StageFlags = VK_PIPELINE_STAGE_ACCELERATION_STRUCTURE_BUILD_BIT_KHR | VK_PIPELINE_STAGE_RAY_TRACING_SHADER_BIT_KHR | VK_PIPELINE_STAGE_COMPUTE_SHADER_BIT;
+			AccessFlags = VK_ACCESS_ACCELERATION_STRUCTURE_WRITE_BIT_KHR;
+			return;
+#endif // VULKAN_RHI_RAYTRACING
 	}
 
 	// If DSVWrite is set, we ignore everything else because it decides the layout.
@@ -626,6 +639,12 @@ void FVulkanDynamicRHI::RHICreateTransition(FRHITransition* Transition, const FR
 			break;
 		}
 
+		case FRHITransitionInfo::EType::BVH:
+		{
+			// Requires memory barrier
+			break;
+		}
+
 		default:
 			checkNoEntry();
 			continue;
diff --git a/Engine/Source/Runtime/VulkanRHI/Private/VulkanBuffer.cpp b/Engine/Source/Runtime/VulkanRHI/Private/VulkanBuffer.cpp
--- a/Engine/Source/Runtime/VulkanRHI/Private/VulkanBuffer.cpp
+++ b/Engine/Source/Runtime/VulkanRHI/Private/VulkanBuffer.cpp
@@ -60,6 +60,31 @@ static FORCEINLINE void UpdateVulkanBufferStats(uint64_t Size, VkBufferUsageFlag
 	}
 }
 
+#if VULKAN_RHI_RAYTRACING
+FVulkanAccelerationStructureBuffer::FVulkanAccelerationStructureBuffer(FVulkanDevice* InDevice, uint32 InSize, uint32 InUEUsage, uint32 InStride, FRHIResourceCreateInfo& CreateInfo)
+	: FRHIBuffer(InSize, InUEUsage, InStride)
+{
+	check(InUEUsage & BUF_AccelerationStructure);
+	if (InSize == 0)
+	{
+		return;
+	}
+	
+	FVulkanRayTracingAllocator::Allocate(
+		InDevice->GetPhysicalHandle(),
+		InDevice->GetInstanceHandle(),
+		InSize,
+		VK_BUFFER_USAGE_ACCELERATION_STRUCTURE_STORAGE_BIT_KHR | VK_BUFFER_USAGE_SHADER_DEVICE_ADDRESS_BIT,
+		VK_MEMORY_PROPERTY_DEVICE_LOCAL_BIT,
+		Allocation);
+}
+
+FVulkanAccelerationStructureBuffer::~FVulkanAccelerationStructureBuffer()
+{
+	FVulkanRayTracingAllocator::Free(Allocation);
+}
+#endif // VULKAN_RHI_RAYTRACING
+
 FVulkanResourceMultiBuffer::FVulkanResourceMultiBuffer(FVulkanDevice* InDevice, uint32 InSize, uint32 InUEUsage, uint32 InStride, FRHIResourceCreateInfo& CreateInfo, class FRHICommandListImmediate* InRHICmdList)
 	: FRHIBuffer(InSize, InUEUsage, InStride)
 	, VulkanRHI::FDeviceChild(InDevice)
@@ -418,6 +443,14 @@ void FVulkanResourceMultiBuffer::Swap(FVulkanResourceMultiBuffer& Other)
 FBufferRHIRef FVulkanDynamicRHI::RHICreateBuffer(uint32 Size, EBufferUsageFlags Usage, uint32 Stride, ERHIAccess ResourceState, FRHIResourceCreateInfo& CreateInfo)
 {
 	LLM_SCOPE_VULKAN(ELLMTagVulkan::VulkanBuffers);
+
+#if VULKAN_RHI_RAYTRACING
+	if (Usage & BUF_AccelerationStructure)
+	{
+		return new FVulkanAccelerationStructureBuffer(Device, Size, Usage, Stride, CreateInfo);
+	}
+#endif
+
 	if (CreateInfo.bWithoutNativeResource)
 	{
 		return new FVulkanResourceMultiBuffer(nullptr, 0, 0, 0, CreateInfo, nullptr);
diff --git a/Engine/Source/Runtime/VulkanRHI/Private/VulkanContext.h b/Engine/Source/Runtime/VulkanRHI/Private/VulkanContext.h
--- a/Engine/Source/Runtime/VulkanRHI/Private/VulkanContext.h
+++ b/Engine/Source/Runtime/VulkanRHI/Private/VulkanContext.h
@@ -101,6 +101,7 @@ public:
 	virtual void RHINextSubpass() final override;
 
 #if VULKAN_RHI_RAYTRACING
+	virtual void RHIBindAccelerationStructureMemory(FRHIRayTracingScene* Scene, FRHIBuffer* Buffer, uint32 BufferOffset) final override;
 	virtual void RHIBuildAccelerationStructures(const TArrayView<const FRayTracingGeometryBuildParams> Params) final override;
 	virtual void RHIBuildAccelerationStructure(const FRayTracingSceneBuildParams& SceneBuildParams) final override;
 	virtual void RHIRayTraceOcclusion(FRHIRayTracingScene* Scene, FRHIShaderResourceView* Rays, FRHIUnorderedAccessView* Output, uint32 NumRays) final override;
diff --git a/Engine/Source/Runtime/VulkanRHI/Private/VulkanRayTracing.cpp b/Engine/Source/Runtime/VulkanRHI/Private/VulkanRayTracing.cpp
--- a/Engine/Source/Runtime/VulkanRHI/Private/VulkanRayTracing.cpp
+++ b/Engine/Source/Runtime/VulkanRHI/Private/VulkanRayTracing.cpp
@@ -104,7 +104,7 @@ static uint32 FindMemoryType(VkPhysicalDevice Gpu, uint32 Filter, VkMemoryProper
 }
 
 // Temporary brute force allocation
-static void Allocate(VkPhysicalDevice Gpu, VkDevice Device, VkDeviceSize Size, VkBufferUsageFlags UsageFlags, VkMemoryPropertyFlags MemoryFlags, FVkRtAllocation& Result)
+void FVulkanRayTracingAllocator::Allocate(VkPhysicalDevice Gpu, VkDevice Device, VkDeviceSize Size, VkBufferUsageFlags UsageFlags, VkMemoryPropertyFlags MemoryFlags, FVkRtAllocation& Result)
 {
 	VkBufferCreateInfo CreateInfo;
 	ZeroVulkanStruct(CreateInfo, VK_STRUCTURE_TYPE_BUFFER_CREATE_INFO);
@@ -131,7 +131,7 @@ static void Allocate(VkPhysicalDevice Gpu, VkDevice Device, VkDeviceSize Size, V
 }
 
 // Temporary brute force deallocation
-static void Free(FVkRtAllocation& Allocation)
+void FVulkanRayTracingAllocator::Free(FVkRtAllocation& Allocation)
 {
 	if (Allocation.Buffer != VK_NULL_HANDLE)
 	{
@@ -259,7 +259,7 @@ FVulkanRayTracingGeometry::FVulkanRayTracingGeometry(const FRayTracingGeometryIn
 		EAccelerationStructureBuildMode::Build,
 		BuildData);
 
-	Allocate(
+	FVulkanRayTracingAllocator::Allocate(
 		Device->GetPhysicalHandle(), 
 		Device->GetInstanceHandle(),
 		BuildData.SizesInfo.accelerationStructureSize,
@@ -267,7 +267,7 @@ FVulkanRayTracingGeometry::FVulkanRayTracingGeometry(const FRayTracingGeometryIn
 		VK_MEMORY_PROPERTY_DEVICE_LOCAL_BIT,
 		Allocation);
 
-	Allocate(
+	FVulkanRayTracingAllocator::Allocate(
 		Device->GetPhysicalHandle(),
 		Device->GetInstanceHandle(),
 		BuildData.SizesInfo.buildScratchSize,
@@ -283,8 +283,8 @@ FVulkanRayTracingGeometry::FVulkanRayTracingGeometry(const FRayTracingGeometryIn
 
 FVulkanRayTracingGeometry::~FVulkanRayTracingGeometry()
 {
-	Free(Allocation);
-	Free(Scratch);
+	FVulkanRayTracingAllocator::Free(Allocation);
+	FVulkanRayTracingAllocator::Free(Scratch);
 }
 
 void FVulkanRayTracingGeometry::BuildAccelerationStructure(FVulkanCommandListContext& CommandContext, EAccelerationStructureBuildMode BuildMode)
@@ -329,7 +329,7 @@ void FVulkanRayTracingGeometry::BuildAccelerationStructure(FVulkanCommandListCon
 	// No longer need scratch memory for a static build
 	if (!bAllowUpdate)
 	{
-		Free(Scratch);
+		FVulkanRayTracingAllocator::Free(Scratch);
 	}
 }
 
@@ -442,7 +442,7 @@ FVulkanRayTracingScene::FVulkanRayTracingScene(const FRayTracingSceneInitializer
 	// Allocate instance buffer
 	const uint32 InstanceBufferByteSize = static_cast<uint32>(NumInstances) * sizeof(VkAccelerationStructureInstanceKHR);
 
-	Allocate(
+	FVulkanRayTracingAllocator::Allocate(
 		Device->GetPhysicalHandle(),
 		Device->GetInstanceHandle(),
 		InstanceBufferByteSize,
@@ -469,48 +469,65 @@ FVulkanRayTracingScene::FVulkanRayTracingScene(const FRayTracingSceneInitializer
 	InstanceDescs.Empty();
 	InstanceGeometry.Empty();
 
-	// Allocate buffers
-	FVkRtTLASBuildData BuildData;
-	GetTLASBuildData(Device->GetInstanceHandle(), NumInstances, 0, BuildData);
-
-	Allocate(
-		Device->GetPhysicalHandle(),
-		Device->GetInstanceHandle(),
-		BuildData.SizesInfo.accelerationStructureSize,
-		VK_BUFFER_USAGE_ACCELERATION_STRUCTURE_STORAGE_BIT_KHR | VK_BUFFER_USAGE_SHADER_DEVICE_ADDRESS_BIT,
-		VK_MEMORY_PROPERTY_DEVICE_LOCAL_BIT,
-		Allocation);
-
-	Allocate(
-		Device->GetPhysicalHandle(),
-		Device->GetInstanceHandle(),
-		BuildData.SizesInfo.buildScratchSize,
-		VK_BUFFER_USAGE_STORAGE_BUFFER_BIT | VK_BUFFER_USAGE_SHADER_DEVICE_ADDRESS_BIT,
-		VK_MEMORY_PROPERTY_DEVICE_LOCAL_BIT,
-		Scratch);
-
-	Scratch.Address = GetDeviceAddress(Device->GetInstanceHandle(), Scratch.Buffer);
+	const ERayTracingAccelerationStructureFlags BuildFlags = ERayTracingAccelerationStructureFlags::FastTrace; // #yuriy_todo: pass this in
+	SizeInfo = RHICalcRayTracingSceneSize(NumInstances, BuildFlags);
 }
 
 FVulkanRayTracingScene::~FVulkanRayTracingScene()
+{}
+
+void FVulkanRayTracingScene::BindBuffer(FRHIBuffer* InBuffer, uint32 InBufferOffset)
 {
-	Free(Allocation);
+	check(SizeInfo.ResultSize + InBufferOffset <= InBuffer->GetSize());
+	check(InBufferOffset % 256 == 0); // Spec requires offset to be a multiple of 256
+	AccelerationStructureBuffer = static_cast<FVulkanAccelerationStructureBuffer*>(InBuffer);
+	BufferOffset = InBufferOffset;
 }
 
-void FVulkanRayTracingScene::BuildAccelerationStructure(FVulkanCommandListContext& CommandContext)
+void FVulkanRayTracingScene::BuildAccelerationStructure(
+	FVulkanCommandListContext& CommandContext,
+	FVulkanResourceMultiBuffer* InScratchBuffer, uint32 InScratchOffset,
+	FVulkanResourceMultiBuffer* InInstanceBuffer, uint32 InInstanceOffset,
+	uint32 NumInstanceDescs)
 {
+	check(AccelerationStructureBuffer.IsValid());
+	check(InInstanceBuffer == nullptr); // External instance buffer not supported yet
+	const bool bExternalScratchBuffer = InScratchBuffer != nullptr;
+
 	FVkRtTLASBuildData BuildData;
 	GetTLASBuildData(Device->GetInstanceHandle(), NumInstances, InstanceBuffer.Address, BuildData);
 
+	if (!bExternalScratchBuffer)
+	{
+		FVulkanRayTracingAllocator::Allocate(
+			Device->GetPhysicalHandle(),
+			Device->GetInstanceHandle(),
+			BuildData.SizesInfo.buildScratchSize,
+			VK_BUFFER_USAGE_STORAGE_BUFFER_BIT | VK_BUFFER_USAGE_SHADER_DEVICE_ADDRESS_BIT,
+			VK_MEMORY_PROPERTY_DEVICE_LOCAL_BIT,
+			Scratch);
+
+		Scratch.Address = GetDeviceAddress(Device->GetInstanceHandle(), Scratch.Buffer);
+	}
+
 	VkAccelerationStructureCreateInfoKHR TLASCreateInfo;
 	ZeroVulkanStruct(TLASCreateInfo, VK_STRUCTURE_TYPE_ACCELERATION_STRUCTURE_CREATE_INFO_KHR);
-	TLASCreateInfo.buffer = Allocation.Buffer;
+	TLASCreateInfo.buffer = AccelerationStructureBuffer->GetBuffer();
+	TLASCreateInfo.offset = BufferOffset;
 	TLASCreateInfo.size = BuildData.SizesInfo.accelerationStructureSize;
 	TLASCreateInfo.type = VK_ACCELERATION_STRUCTURE_TYPE_TOP_LEVEL_KHR;
 	VERIFYVULKANRESULT(vkCreateAccelerationStructureKHR(Device->GetInstanceHandle(), &TLASCreateInfo, VULKAN_CPU_ALLOCATOR, &Handle));
 
 	BuildData.GeometryInfo.dstAccelerationStructure = Handle;
-	BuildData.GeometryInfo.scratchData.deviceAddress = Scratch.Address;
+
+	if (bExternalScratchBuffer)
+	{
+		BuildData.GeometryInfo.scratchData.deviceAddress = GetDeviceAddress(Device->GetInstanceHandle(), InScratchBuffer->GetHandle()) + InScratchOffset;
+	}
+	else
+	{
+		BuildData.GeometryInfo.scratchData.deviceAddress = Scratch.Address;
+	}
 
 	VkAccelerationStructureBuildRangeInfoKHR TLASBuildRangeInfo;
 	TLASBuildRangeInfo.primitiveCount = NumInstances;
@@ -527,8 +544,12 @@ void FVulkanRayTracingScene::BuildAccelerationStructure(FVulkanCommandListContex
 	CommandBufferManager.SubmitActiveCmdBuffer();
 	CommandBufferManager.PrepareForNewActiveCommandBuffer();
 
-	Free(InstanceBuffer);
-	Free(Scratch);
+	FVulkanRayTracingAllocator::Free(InstanceBuffer);
+
+	if (!bExternalScratchBuffer)
+	{
+		FVulkanRayTracingAllocator::Free(Scratch);
+	}
 }
 
 FRayTracingAccelerationStructureSize FVulkanDynamicRHI::RHICalcRayTracingSceneSize(uint32 MaxInstances, ERayTracingAccelerationStructureFlags Flags)
@@ -578,6 +599,11 @@ FRayTracingGeometryRHIRef FVulkanDynamicRHI::RHICreateRayTracingGeometry(const F
 	return new FVulkanRayTracingGeometry(Initializer, GetDevice());
 }
 
+void FVulkanCommandListContext::RHIBindAccelerationStructureMemory(FRHIRayTracingScene* Scene, FRHIBuffer* Buffer, uint32 BufferOffset)
+{
+	ResourceCast(Scene)->BindBuffer(Buffer, BufferOffset);
+}
+
 // Todo: High level rhi call should have transitioned and verified vb and ib to read for each segment
 void FVulkanCommandListContext::RHIBuildAccelerationStructures(const TArrayView<const FRayTracingGeometryBuildParams> Params)
 {
@@ -596,7 +622,13 @@ void FVulkanCommandListContext::RHIBuildAccelerationStructures(const TArrayView<
 void FVulkanCommandListContext::RHIBuildAccelerationStructure(const FRayTracingSceneBuildParams& SceneBuildParams)
 {
 	FVulkanRayTracingScene* const Scene = ResourceCast(SceneBuildParams.Scene);
-	Scene->BuildAccelerationStructure(*this);
+	FVulkanResourceMultiBuffer* const ScratchBuffer = ResourceCast(SceneBuildParams.ScratchBuffer);
+	FVulkanResourceMultiBuffer* const InstanceBuffer = ResourceCast(SceneBuildParams.InstanceBuffer);
+	Scene->BuildAccelerationStructure(
+		*this, 
+		ScratchBuffer, SceneBuildParams.ScratchBufferOffset, 
+		InstanceBuffer, SceneBuildParams.InstanceBufferOffset, 
+		SceneBuildParams.NumInstances);
 }
 
 void FVulkanCommandListContext::RHIRayTraceOcclusion(FRHIRayTracingScene* Scene, FRHIShaderResourceView* Rays, FRHIUnorderedAccessView* Output, uint32 NumRays)
diff --git a/Engine/Source/Runtime/VulkanRHI/Private/VulkanRayTracing.h b/Engine/Source/Runtime/VulkanRHI/Private/VulkanRayTracing.h
--- a/Engine/Source/Runtime/VulkanRHI/Private/VulkanRayTracing.h
+++ b/Engine/Source/Runtime/VulkanRHI/Private/VulkanRayTracing.h
@@ -5,6 +5,8 @@
 #if VULKAN_RHI_RAYTRACING
 
 class FVulkanCommandListContext;
+class FVulkanResourceMultiBuffer;
+class FVulkanAccelerationStructureBuffer;
 
 #define ENUM_VK_ENTRYPOINTS_RAYTRACING(EnumMacro) \
 	EnumMacro(PFN_vkCreateAccelerationStructureKHR, vkCreateAccelerationStructureKHR) \
@@ -39,6 +41,13 @@ struct FVkRtAllocation
 	VkDeviceAddress Address = 0;
 };
 
+class FVulkanRayTracingAllocator
+{
+public:
+	static void Allocate(VkPhysicalDevice Gpu, VkDevice Device, VkDeviceSize Size, VkBufferUsageFlags UsageFlags, VkMemoryPropertyFlags MemoryFlags, FVkRtAllocation& Result);
+	static void Free(FVkRtAllocation& Allocation);
+};
+
 struct FVkRtTLASBuildData
 {
 	FVkRtTLASBuildData()
@@ -107,10 +116,15 @@ public:
 	FVulkanRayTracingScene(const FRayTracingSceneInitializer& Initializer, const FVulkanDevice* InDevice);
 	~FVulkanRayTracingScene();
 
-	void BuildAccelerationStructure(FVulkanCommandListContext& CommandContext);
+	void BindBuffer(FRHIBuffer* InBuffer, uint32 InBufferOffset);
+	void BuildAccelerationStructure(
+		FVulkanCommandListContext& CommandContext, 
+		FVulkanResourceMultiBuffer* ScratchBuffer, uint32 ScratchOffset, 
+		FVulkanResourceMultiBuffer* InstanceBuffer, uint32 InstanceOffset, 
+		uint32 NumInstanceDescs);
 
 	VkAccelerationStructureKHR Handle = VK_NULL_HANDLE;
-	FVkRtAllocation Allocation;
+	TRefCountPtr<FVulkanAccelerationStructureBuffer> AccelerationStructureBuffer;
 	FVkRtAllocation Scratch;
 	FVkRtAllocation InstanceBuffer;
 
@@ -119,6 +133,8 @@ public:
 	TArray<TRefCountPtr<const FVulkanRayTracingGeometry>> ReferencedGeometry;
 	FName DebugName;
 	uint32 NumInstances = 0;
+	uint32 BufferOffset = 0;
+	FRayTracingAccelerationStructureSize SizeInfo;
 
 private:
 	const FVulkanDevice* Device = nullptr;
diff --git a/Engine/Source/Runtime/VulkanRHI/Private/VulkanResources.h b/Engine/Source/Runtime/VulkanRHI/Private/VulkanResources.h
--- a/Engine/Source/Runtime/VulkanRHI/Private/VulkanResources.h
+++ b/Engine/Source/Runtime/VulkanRHI/Private/VulkanResources.h
@@ -1020,6 +1020,23 @@ protected:
 	friend class FVulkanCommandListContext;
 };
 
+#if VULKAN_RHI_RAYTRACING
+class FVulkanAccelerationStructureBuffer : public FRHIBuffer
+{
+public:
+	FVulkanAccelerationStructureBuffer(FVulkanDevice* InDevice, uint32 InSize, uint32 InUEUsage, uint32 InStride, FRHIResourceCreateInfo& CreateInfo);
+	virtual ~FVulkanAccelerationStructureBuffer();
+
+	inline VkBuffer GetBuffer() const
+	{
+		return Allocation.Buffer;
+	}
+
+private:
+	FVkRtAllocation Allocation;
+};
+#endif
+
 class FVulkanResourceMultiBuffer : public FRHIBuffer, public FVulkanEvictable, public VulkanRHI::FDeviceChild
 {
 	virtual void Evict(FVulkanDevice& Device);
