ETextureCreateFlags enum class.

diff --git a/Engine/Plugins/Compositing/OpenColorIO/Source/OpenColorIO/Private/OpenColorIODisplayExtension.cpp b/Engine/Plugins/Compositing/OpenColorIO/Source/OpenColorIO/Private/OpenColorIODisplayExtension.cpp
--- a/Engine/Plugins/Compositing/OpenColorIO/Source/OpenColorIO/Private/OpenColorIODisplayExtension.cpp
+++ b/Engine/Plugins/Compositing/OpenColorIO/Source/OpenColorIO/Private/OpenColorIODisplayExtension.cpp
@@ -174,7 +174,7 @@ FScreenPassTexture FOpenColorIODisplayExtension::PostProcessPassAfterTonemap_Ren
 		return SceneColor;
 	}
 
-	if (!(SceneColor.Texture->Desc.Flags & ETextureCreateFlags::TexCreate_ShaderResource))
+	if (!EnumHasAnyFlags(SceneColor.Texture->Desc.Flags, TexCreate_ShaderResource))
 	{
 		return SceneColor;
 	}
diff --git a/Engine/Plugins/Runtime/Oculus/OculusVR/Source/OculusHMD/Private/OculusHMD_CustomPresent.cpp b/Engine/Plugins/Runtime/Oculus/OculusVR/Source/OculusHMD/Private/OculusHMD_CustomPresent.cpp
--- a/Engine/Plugins/Runtime/Oculus/OculusVR/Source/OculusHMD/Private/OculusHMD_CustomPresent.cpp
+++ b/Engine/Plugins/Runtime/Oculus/OculusVR/Source/OculusHMD/Private/OculusHMD_CustomPresent.cpp
@@ -406,7 +406,7 @@ void FCustomPresent::CopyTexture_RenderThread(FRHICommandListImmediate& RHICmdLi
 
 	if (DstTexture2D)
 	{
-		sRGBSource &= ( ( SrcTexture->GetFlags() & TexCreate_SRGB ) != 0);
+		sRGBSource &= EnumHasAnyFlags(SrcTexture->GetFlags(), TexCreate_SRGB);
 
 		// Need to copy over mip maps on Android since they are not generated like they are on PC
 #if PLATFORM_ANDROID
diff --git a/Engine/Plugins/Runtime/OpenXR/Source/OpenXRHMD/Private/OpenXRHMD_Swapchain.cpp b/Engine/Plugins/Runtime/OpenXR/Source/OpenXRHMD/Private/OpenXRHMD_Swapchain.cpp
--- a/Engine/Plugins/Runtime/OpenXR/Source/OpenXRHMD/Private/OpenXRHMD_Swapchain.cpp
+++ b/Engine/Plugins/Runtime/OpenXR/Source/OpenXRHMD/Private/OpenXRHMD_Swapchain.cpp
@@ -139,19 +139,19 @@ XrSwapchain CreateSwapchain(XrSession InSession, uint32 PlatformFormat, uint32 S
 {
 	// Need a mutable format so we can reinterpret an sRGB format into a linear format
 	XrSwapchainUsageFlags Usage = XR_SWAPCHAIN_USAGE_MUTABLE_FORMAT_BIT;
-	if (TargetableTextureFlags & TexCreate_RenderTargetable)
+	if (EnumHasAnyFlags(TargetableTextureFlags, TexCreate_RenderTargetable))
 	{
 		Usage |= XR_SWAPCHAIN_USAGE_COLOR_ATTACHMENT_BIT;
 	}
-	if (TargetableTextureFlags & TexCreate_DepthStencilTargetable)
+	if (EnumHasAnyFlags(TargetableTextureFlags, TexCreate_DepthStencilTargetable))
 	{
 		Usage |= XR_SWAPCHAIN_USAGE_DEPTH_STENCIL_ATTACHMENT_BIT;
 	}
-	if (TargetableTextureFlags & TexCreate_ShaderResource)
+	if (EnumHasAnyFlags(TargetableTextureFlags, TexCreate_ShaderResource))
 	{
 		Usage |= XR_SWAPCHAIN_USAGE_SAMPLED_BIT;
 	}
-	if (TargetableTextureFlags & TexCreate_UAV)
+	if (EnumHasAnyFlags(TargetableTextureFlags, TexCreate_UAV))
 	{
 		Usage |= XR_SWAPCHAIN_USAGE_UNORDERED_ACCESS_BIT;
 	}
@@ -160,7 +160,7 @@ XrSwapchain CreateSwapchain(XrSession InSession, uint32 PlatformFormat, uint32 S
 	XrSwapchainCreateInfo info;
 	info.type = XR_TYPE_SWAPCHAIN_CREATE_INFO;
 	info.next = nullptr;
-	info.createFlags = Flags & TexCreate_Dynamic ? 0 : XR_SWAPCHAIN_CREATE_STATIC_IMAGE_BIT;
+	info.createFlags = EnumHasAnyFlags(Flags, TexCreate_Dynamic) ? 0 : XR_SWAPCHAIN_CREATE_STATIC_IMAGE_BIT;
 	info.usageFlags = Usage;
 	info.format = PlatformFormat;
 	info.sampleCount = NumSamples;
diff --git a/Engine/Plugins/Runtime/Steam/SteamVR/Source/SteamVR/Private/SteamVRSpectatorScreen.cpp b/Engine/Plugins/Runtime/Steam/SteamVR/Source/SteamVR/Private/SteamVRSpectatorScreen.cpp
--- a/Engine/Plugins/Runtime/Steam/SteamVR/Source/SteamVR/Private/SteamVRSpectatorScreen.cpp
+++ b/Engine/Plugins/Runtime/Steam/SteamVR/Source/SteamVR/Private/SteamVRSpectatorScreen.cpp
@@ -85,7 +85,7 @@ void FSteamVRHMD::CopyTexture_RenderThread(FRHICommandListImmediate& RHICmdList,
 		const bool bSameSize = DstRect.Size() == SrcRect.Size();
 		FRHISamplerState* PixelSampler = bSameSize ? TStaticSamplerState<SF_Point>::GetRHI() : TStaticSamplerState<SF_Bilinear>::GetRHI();
 
-		if ((SrcTexture->GetFlags() & TexCreate_SRGB) != 0)
+		if (EnumHasAnyFlags(SrcTexture->GetFlags(), TexCreate_SRGB))
 		{
 			TShaderMapRef<FScreenPSsRGBSource> PixelShader(ShaderMap);
 			GraphicsPSOInit.BoundShaderState.PixelShaderRHI = PixelShader.GetPixelShader();
diff --git a/Engine/Plugins/Runtime/WindowsMixedReality/Source/WindowsMixedRealityHMD/Private/WindowsMixedRealityHMD.cpp b/Engine/Plugins/Runtime/WindowsMixedReality/Source/WindowsMixedRealityHMD/Private/WindowsMixedRealityHMD.cpp
--- a/Engine/Plugins/Runtime/WindowsMixedReality/Source/WindowsMixedRealityHMD/Private/WindowsMixedRealityHMD.cpp
+++ b/Engine/Plugins/Runtime/WindowsMixedReality/Source/WindowsMixedRealityHMD/Private/WindowsMixedRealityHMD.cpp
@@ -797,7 +797,7 @@ namespace WindowsMixedReality
 			const bool bSameSize = DstRect.Size() == SrcRect.Size();
 			FRHISamplerState* PixelSampler = bSameSize ? TStaticSamplerState<SF_Point>::GetRHI() : TStaticSamplerState<SF_Bilinear>::GetRHI();
 
-			if ((SrcTexture->GetFlags() & TexCreate_SRGB) != 0)
+			if (EnumHasAnyFlags(SrcTexture->GetFlags(), TexCreate_SRGB))
 			{
 				TShaderMapRef<FScreenPSsRGBSource> PixelShader(ShaderMap);
 				GraphicsPSOInit.BoundShaderState.PixelShaderRHI = PixelShader.GetPixelShader();
@@ -1840,7 +1840,7 @@ namespace WindowsMixedReality
 		{
 			FRHIResourceCreateInfo CreateInfo(TEXT("RemappedDepthTexture"));
 			remappedDepthTexture = RHICreateTexture2D(depthFRHITexture->GetSizeX(), depthFRHITexture->GetSizeY(),
-				PF_R32_FLOAT, 1, 1, ETextureCreateFlags::TexCreate_RenderTargetable, CreateInfo);
+				PF_R32_FLOAT, 1, 1, TexCreate_RenderTargetable, CreateInfo);
 		}
 
 		FGraphicsPipelineStateInitializer GraphicsPSOInit;
diff --git a/Engine/Plugins/Runtime/nDisplay/Source/DisplayCluster/Private/Render/Viewport/RenderTarget/DisplayClusterRenderTargetResourcesPool.cpp b/Engine/Plugins/Runtime/nDisplay/Source/DisplayCluster/Private/Render/Viewport/RenderTarget/DisplayClusterRenderTargetResourcesPool.cpp
--- a/Engine/Plugins/Runtime/nDisplay/Source/DisplayCluster/Private/Render/Viewport/RenderTarget/DisplayClusterRenderTargetResourcesPool.cpp
+++ b/Engine/Plugins/Runtime/nDisplay/Source/DisplayCluster/Private/Render/Viewport/RenderTarget/DisplayClusterRenderTargetResourcesPool.cpp
@@ -18,7 +18,7 @@ FDisplayClusterViewportResourceSettings::FDisplayClusterViewportResourceSettings
 		if (ViewportTexture)
 		{
 			Format = ViewportTexture->GetFormat();
-			bShouldUseSRGB = (ViewportTexture->GetFlags() & TexCreate_SRGB) != 0;
+			bShouldUseSRGB = EnumHasAnyFlags(ViewportTexture->GetFlags(), TexCreate_SRGB);
 			DisplayGamma = InViewport->GetDisplayGamma();
 		}
 	}
@@ -41,9 +41,9 @@ FDisplayClusterViewportResourceSettings::FDisplayClusterViewportResourceSettings
 	{
 		Size = InTexture->GetSizeXY();
 		Format = InTexture->GetFormat();
-		bShouldUseSRGB = (InTexture->GetFlags() & TexCreate_SRGB) != 0;
+		bShouldUseSRGB = EnumHasAnyFlags(InTexture->GetFlags(), TexCreate_SRGB);
 
-		bIsRenderTargetable = (InTexture->GetFlags() & TexCreate_RenderTargetable) != 0;
+		bIsRenderTargetable = EnumHasAnyFlags(InTexture->GetFlags(), TexCreate_RenderTargetable);
 		NumMips = InTexture->GetNumMips();
 	}
 }
diff --git a/Engine/Source/Editor/UnrealEd/Private/Commandlets/ShaderPipelineCacheToolsCommandlet.cpp b/Engine/Source/Editor/UnrealEd/Private/Commandlets/ShaderPipelineCacheToolsCommandlet.cpp
--- a/Engine/Source/Editor/UnrealEd/Private/Commandlets/ShaderPipelineCacheToolsCommandlet.cpp
+++ b/Engine/Source/Editor/UnrealEd/Private/Commandlets/ShaderPipelineCacheToolsCommandlet.cpp
@@ -1062,7 +1062,7 @@ int32 ExpandPSOSC(const TArray<FString>& Tokens)
 				// Re-purpose graphics state fields to store RT PSO properties
 				// See corresponding parsing code in ParseStableCSV().
 				Desc.MSAASamples = Item.PSO->RayTracingDesc.MaxPayloadSizeInBytes;
-				Desc.DepthStencilFlags = uint32(Item.PSO->RayTracingDesc.bAllowHitGroupIndexing);
+				Desc.DepthStencilFlags = static_cast<ETextureCreateFlags>(Item.PSO->RayTracingDesc.bAllowHitGroupIndexing);
 
 				PSOLine += FString::Printf(TEXT("\"%s\""), *Desc.StateToString());
 
@@ -1340,7 +1340,7 @@ static TSet<FPipelineCacheFileFormatPSO> ParseStableCSV(const FString& FileName,
 				// See corresponding serialization code in ExpandPSOSC()
 				PSO.RayTracingDesc.Frequency = EShaderFrequency(AdjustedSlotIndex);
 				PSO.RayTracingDesc.MaxPayloadSizeInBytes = PSO.GraphicsDesc.MSAASamples;
-				PSO.RayTracingDesc.bAllowHitGroupIndexing = PSO.GraphicsDesc.DepthStencilFlags != 0;
+				PSO.RayTracingDesc.bAllowHitGroupIndexing = static_cast<uint64>(PSO.GraphicsDesc.DepthStencilFlags) != 0;
 				break;
 			default:
 				UE_LOG(LogShaderPipelineCacheTools, Error, TEXT("Unexpected shader frequency"));
diff --git a/Engine/Source/Runtime/Apple/AGXRHI/Private/AGXPipeline.cpp b/Engine/Source/Runtime/Apple/AGXRHI/Private/AGXPipeline.cpp
--- a/Engine/Source/Runtime/Apple/AGXRHI/Private/AGXPipeline.cpp
+++ b/Engine/Source/Runtime/Apple/AGXRHI/Private/AGXPipeline.cpp
@@ -207,8 +207,8 @@ struct FAGXGraphicsPipelineKey
 			if (TargetFormat == PF_Unknown) { continue; }
 
 			mtlpp::PixelFormat MetalFormat = (mtlpp::PixelFormat)GPixelFormats[TargetFormat].PlatformFormat;
-			ETextureCreateFlags Flags = (ETextureCreateFlags)Init.RenderTargetFlags[i];
-			if (Flags & TexCreate_SRGB)
+			ETextureCreateFlags Flags = Init.RenderTargetFlags[i];
+			if (EnumHasAnyFlags(Flags, TexCreate_SRGB))
 			{
 #if PLATFORM_MAC // Expand as R8_sRGB is iOS only.
 				if (MetalFormat == mtlpp::PixelFormat::R8Unorm)
@@ -653,7 +653,7 @@ static bool ConfigureRenderPipelineDescriptor(mtlpp::RenderPipelineDescriptor& R
 		
 		mtlpp::PixelFormat MetalFormat = (mtlpp::PixelFormat)GPixelFormats[TargetFormat].PlatformFormat;
 		ETextureCreateFlags Flags = (ETextureCreateFlags)Init.RenderTargetFlags[ActiveTargetIndex];
-		if (Flags & TexCreate_SRGB)
+		if (EnumHasAnyFlags(Flags, TexCreate_SRGB))
 		{
 #if PLATFORM_MAC // Expand as R8_sRGB is iOS only.
 			if (MetalFormat == mtlpp::PixelFormat::R8Unorm)
diff --git a/Engine/Source/Runtime/Apple/AGXRHI/Private/AGXRenderTarget.cpp b/Engine/Source/Runtime/Apple/AGXRHI/Private/AGXRenderTarget.cpp
--- a/Engine/Source/Runtime/Apple/AGXRHI/Private/AGXRenderTarget.cpp
+++ b/Engine/Source/Runtime/Apple/AGXRHI/Private/AGXRenderTarget.cpp
@@ -240,7 +240,7 @@ void FAGXDynamicRHI::RHIReadSurfaceData(FRHITexture* TextureRHI, FIntRect Rect,
 	mtlpp::Region Region(Rect.Min.X, Rect.Min.Y, SizeX, SizeY);
     
 	FAGXTexture Texture = Surface->Texture;
-    if(!Texture && (Surface->Flags & TexCreate_Presentable))
+    if(!Texture && EnumHasAnyFlags(Surface->Flags, TexCreate_Presentable))
     {
         Texture = Surface->GetCurrentTexture();
     }
@@ -388,7 +388,7 @@ void FAGXDynamicRHI::RHIReadSurfaceFloatData(FRHITexture* TextureRHI, FIntRect R
 	FAGXSurface* Surface = AGXGetMetalSurfaceFromRHITexture(TextureRHI);
 	
     FAGXTexture Texture = Surface->Texture;
-    if(!Texture && (Surface->Flags & TexCreate_Presentable))
+    if(!Texture && EnumHasAnyFlags(Surface->Flags, TexCreate_Presentable))
     {
 		Texture = Surface->GetCurrentTexture();
     }
diff --git a/Engine/Source/Runtime/Apple/AGXRHI/Private/AGXStateCache.cpp b/Engine/Source/Runtime/Apple/AGXRHI/Private/AGXStateCache.cpp
--- a/Engine/Source/Runtime/Apple/AGXRHI/Private/AGXStateCache.cpp
+++ b/Engine/Source/Runtime/Apple/AGXRHI/Private/AGXStateCache.cpp
@@ -1169,7 +1169,7 @@ void FAGXStateCache::SetVisibilityResultMode(mtlpp::VisibilityResultMode const M
 void FAGXStateCache::ConditionalUpdateBackBuffer(FAGXSurface& Surface)
 {
 	// are we setting the back buffer? if so, make sure we have the drawable
-	if ((Surface.Flags & TexCreate_Presentable))
+	if (EnumHasAnyFlags(Surface.Flags, TexCreate_Presentable))
 	{
 		// update the back buffer texture the first time used this frame
 		if (Surface.Texture.GetPtr() == nil)
diff --git a/Engine/Source/Runtime/Apple/AGXRHI/Private/AGXTexture.cpp b/Engine/Source/Runtime/Apple/AGXRHI/Private/AGXTexture.cpp
--- a/Engine/Source/Runtime/Apple/AGXRHI/Private/AGXTexture.cpp
+++ b/Engine/Source/Runtime/Apple/AGXRHI/Private/AGXTexture.cpp
@@ -42,19 +42,19 @@ FAGXSurface* AGXGetMetalSurfaceFromRHITexture(FRHITexture* Texture)
 
 static bool IsRenderTarget(ETextureCreateFlags Flags)
 {
-	return (Flags & (TexCreate_RenderTargetable | TexCreate_ResolveTargetable | TexCreate_DepthStencilTargetable | TexCreate_DepthStencilResolveTarget)) != 0;
+	return EnumHasAnyFlags(Flags, TexCreate_RenderTargetable | TexCreate_ResolveTargetable | TexCreate_DepthStencilTargetable | TexCreate_DepthStencilResolveTarget);
 }
 
 static mtlpp::TextureUsage ConvertFlagsToUsage(ETextureCreateFlags Flags)
 {
 	NSUInteger Usage = mtlpp::TextureUsage::Unknown;
-    if(Flags & (TexCreate_ShaderResource|TexCreate_ResolveTargetable|TexCreate_DepthStencilTargetable))
+    if (EnumHasAnyFlags(Flags, TexCreate_ShaderResource|TexCreate_ResolveTargetable|TexCreate_DepthStencilTargetable))
 	{
 		Usage |= mtlpp::TextureUsage::ShaderRead;
 		Usage |= mtlpp::TextureUsage::PixelFormatView;
 	}
 	
-	if (Flags & TexCreate_UAV)
+	if (EnumHasAnyFlags(Flags, TexCreate_UAV))
 	{
 		Usage |= mtlpp::TextureUsage::ShaderRead;
 		Usage |= mtlpp::TextureUsage::ShaderWrite;
@@ -62,7 +62,7 @@ static mtlpp::TextureUsage ConvertFlagsToUsage(ETextureCreateFlags Flags)
 	}
 	
 	// offline textures are normal shader read textures
-	if (Flags & TexCreate_OfflineProcessed)
+	if (EnumHasAnyFlags(Flags, TexCreate_OfflineProcessed))
 	{
 		Usage |= mtlpp::TextureUsage::ShaderRead;
 	}
@@ -70,8 +70,8 @@ static mtlpp::TextureUsage ConvertFlagsToUsage(ETextureCreateFlags Flags)
 	//if the high level is doing manual resolves then the textures specifically markes as resolve targets
 	//are likely to be used in a manual shader resolve by the high level and must be bindable as rendertargets.
 	const bool bSeparateResolveTargets = FAGXCommandQueue::SupportsSeparateMSAAAndResolveTarget();
-	const bool bResolveTarget = (Flags & TexCreate_ResolveTargetable);
-	if ((Flags & (TexCreate_RenderTargetable|TexCreate_DepthStencilTargetable|TexCreate_DepthStencilResolveTarget)) || (bResolveTarget && bSeparateResolveTargets))
+	const bool bResolveTarget = EnumHasAnyFlags(Flags, TexCreate_ResolveTargetable);
+	if (EnumHasAnyFlags(Flags, TexCreate_RenderTargetable|TexCreate_DepthStencilTargetable|TexCreate_DepthStencilResolveTarget) || (bResolveTarget && bSeparateResolveTargets))
 	{
 		Usage |= mtlpp::TextureUsage::RenderTarget;
 		Usage |= mtlpp::TextureUsage::ShaderRead;
@@ -481,7 +481,7 @@ void FAGXSurface::Init(FAGXSurface& Source, NSRange MipRange)
 	};
 
 	mtlpp::PixelFormat MetalFormat = (mtlpp::PixelFormat)GPixelFormats[PixelFormat].PlatformFormat;
-	if (Flags & TexCreate_SRGB)
+	if (EnumHasAnyFlags(Flags, TexCreate_SRGB))
 	{
 		// Ensure we have the correct sRGB target format if we create a new texture view rather than using the source texture
 		MetalFormat = AGXToSRGBFormat(MetalFormat);
@@ -530,7 +530,7 @@ void FAGXSurface::Init(FAGXSurface& Source, NSRange MipRange, EPixelFormat Forma
 	
 	bool bUseSourceTex = (Source.PixelFormat != PF_DepthStencil) && Source.PixelFormat == Format && MipRange.location == 0 && MipRange.length == Source.Texture.GetMipmapLevelCount();
 	
-	if (Flags & TexCreate_SRGB)
+	if (EnumHasAnyFlags(Flags, TexCreate_SRGB))
 	{
 		if(bSRGBForceDisable)
 		{
@@ -692,7 +692,7 @@ FAGXSurface::FAGXSurface(ERHIResourceType ResourceType, EPixelFormat Format, uin
 	FPlatformAtomics::InterlockedExchange(&Written, 0);
 	mtlpp::PixelFormat MTLFormat = (mtlpp::PixelFormat)GPixelFormats[Format].PlatformFormat;
 	
-	if (Flags & TexCreate_SRGB)
+	if (EnumHasAnyFlags(Flags, TexCreate_SRGB))
 	{
 		MTLFormat = AGXToSRGBFormat(MTLFormat);
 	}
@@ -702,7 +702,7 @@ FAGXSurface::FAGXSurface(ERHIResourceType ResourceType, EPixelFormat Format, uin
 	
 	
 	// the special back buffer surface will be updated in GetAGXDeviceContext().BeginDrawingViewport - no need to set the texture here
-	if (Flags & TexCreate_Presentable)
+	if (EnumHasAnyFlags(Flags, TexCreate_Presentable))
 	{
 		return;
 	}
@@ -758,7 +758,7 @@ FAGXSurface::FAGXSurface(ERHIResourceType ResourceType, EPixelFormat Format, uin
 	{
 		Desc.SetUsage(ConvertFlagsToUsage(Flags));
 		
-		if((Flags & TexCreate_CPUReadback) && !(Flags & (TexCreate_RenderTargetable|TexCreate_DepthStencilTargetable|TexCreate_FastVRAM)))
+		if (EnumHasAnyFlags(Flags, TexCreate_CPUReadback) && !EnumHasAnyFlags(Flags, TexCreate_RenderTargetable|TexCreate_DepthStencilTargetable|TexCreate_FastVRAM))
 		{
 			Desc.SetCpuCacheMode(mtlpp::CpuCacheMode::DefaultCache);
 #if PLATFORM_MAC
@@ -769,7 +769,7 @@ FAGXSurface::FAGXSurface(ERHIResourceType ResourceType, EPixelFormat Format, uin
 			Desc.SetResourceOptions((mtlpp::ResourceOptions)(mtlpp::ResourceOptions::CpuCacheModeDefaultCache|mtlpp::ResourceOptions::StorageModeShared));
 #endif
 		}
-		else if(((Flags & (TexCreate_NoTiling)) && !(Flags & (TexCreate_FastVRAM|TexCreate_DepthStencilTargetable|TexCreate_RenderTargetable|TexCreate_UAV))))
+		else if(EnumHasAnyFlags(Flags, TexCreate_NoTiling) && !EnumHasAnyFlags(Flags, TexCreate_FastVRAM|TexCreate_DepthStencilTargetable|TexCreate_RenderTargetable|TexCreate_UAV))
 		{
 #if PLATFORM_MAC
 			Desc.SetCpuCacheMode(mtlpp::CpuCacheMode::WriteCombined);
@@ -781,7 +781,7 @@ FAGXSurface::FAGXSurface(ERHIResourceType ResourceType, EPixelFormat Format, uin
 			Desc.SetResourceOptions((mtlpp::ResourceOptions)(mtlpp::ResourceOptions::CpuCacheModeDefaultCache|mtlpp::ResourceOptions::StorageModeShared));
 #endif
 		}
-		else if (Flags & (TexCreate_RenderTargetable|TexCreate_DepthStencilTargetable|TexCreate_ResolveTargetable|TexCreate_DepthStencilResolveTarget))
+		else if (EnumHasAnyFlags(Flags, TexCreate_RenderTargetable|TexCreate_DepthStencilTargetable|TexCreate_ResolveTargetable|TexCreate_DepthStencilResolveTarget))
 		{
 			check(!(Flags & TexCreate_CPUReadback));
 			Desc.SetCpuCacheMode(mtlpp::CpuCacheMode::DefaultCache);
@@ -986,7 +986,7 @@ FAGXSurface::FAGXSurface(ERHIResourceType ResourceType, EPixelFormat Format, uin
 			//we always require an MSAAResolveTexture if MSAATexture is active.
 			check(!MSAATexture || MSAAResolveTexture || bDepthButNoResolveSupported);
 			
-			NSLog(@"Creating %dx MSAA %d x %d %s surface", (int32)Desc.GetSampleCount(), SizeX, SizeY, (Flags & TexCreate_RenderTargetable) ? "Color" : "Depth");
+			NSLog(@"Creating %dx MSAA %d x %d %s surface", (int32)Desc.GetSampleCount(), SizeX, SizeY, EnumHasAnyFlags(Flags, TexCreate_RenderTargetable) ? "Color" : "Depth");
 			if (MSAATexture.GetPtr() == nil)
 			{
 				NSLog(@"Failed to create texture, desc  %@", Desc.GetPtr());
@@ -1175,7 +1175,7 @@ id <MTLBuffer> FAGXSurface::AllocSurface(uint32 MipIndex, uint32 ArrayIndex, ERe
 	
 #if PLATFORM_MAC
 	// Expand R8_sRGB into RGBA8_sRGB for Mac.
-	if (PixelFormat == PF_G8 && (Flags & TexCreate_SRGB) && Type == RRT_Texture2D && LockMode == RLM_WriteOnly)
+	if (PixelFormat == PF_G8 && EnumHasAnyFlags(Flags, TexCreate_SRGB) && Type == RRT_Texture2D && LockMode == RLM_WriteOnly)
 	{
 		DestStride = FMath::Max<uint32>(SizeX >> MipIndex, 1);
 	}
@@ -1209,7 +1209,7 @@ void FAGXSurface::UpdateSurfaceAndDestroySourceBuffer(id <MTLBuffer> SourceBuffe
 	}
 #if PLATFORM_MAC
 	// Expand R8_sRGB into RGBA8_sRGB for Mac.
-	if (PixelFormat == PF_G8 && (Flags & TexCreate_SRGB) && Type == RRT_Texture2D)
+	if (PixelFormat == PF_G8 && EnumHasAnyFlags(Flags, TexCreate_SRGB) && Type == RRT_Texture2D)
 	{
 		TArray<uint8> Data;
 		uint8* ExpandedMem = (uint8*) SourceBuffer.contents;
@@ -1416,7 +1416,7 @@ void* FAGXSurface::Lock(uint32 MipIndex, uint32 ArrayIndex, EResourceLockMode Lo
 			
 #if PLATFORM_MAC
 			// Pack RGBA8_sRGB into R8_sRGB for Mac.
-			if (PixelFormat == PF_G8 && (Flags & TexCreate_SRGB) && Type == RRT_Texture2D)
+			if (PixelFormat == PF_G8 && EnumHasAnyFlags(Flags, TexCreate_SRGB) && Type == RRT_Texture2D)
 			{
 				TArray<uint8> Data;
 				uint8* ExpandedMem = (uint8*)SourceData.GetContents();
@@ -1561,7 +1561,7 @@ uint32 FAGXSurface::GetMipSize(uint32 MipIndex, uint32* Stride, bool bSingleLaye
 		NumBlocksY = FMath::Max<uint32>(NumBlocksY, 2);
 	}
 #if PLATFORM_MAC
-	else if (PixelFormat == PF_G8 && (Flags & TexCreate_SRGB))
+	else if (PixelFormat == PF_G8 && EnumHasAnyFlags(Flags, TexCreate_SRGB))
 	{
 		// RGBA_sRGB is the closest match - so expand the data.
 		NumBlocksX *= 4;
@@ -1621,7 +1621,7 @@ uint32 FAGXSurface::GetNumFaces()
 
 FAGXTexture FAGXSurface::GetDrawableTexture()
 {
-	if (!Texture && (Flags & TexCreate_Presentable))
+	if (!Texture && EnumHasAnyFlags(Flags, TexCreate_Presentable))
 	{
 		check(Viewport);
 		Texture = Viewport->GetDrawableTexture(EAGXViewportAccessRHI);
@@ -1632,7 +1632,7 @@ FAGXTexture FAGXSurface::GetDrawableTexture()
 ns::AutoReleased<FAGXTexture> FAGXSurface::GetCurrentTexture()
 {
 	ns::AutoReleased<FAGXTexture> Tex;
-	if (Viewport && (Flags & TexCreate_Presentable))
+	if (Viewport && EnumHasAnyFlags(Flags, TexCreate_Presentable))
 	{
 		check(Viewport);
 		Tex = Viewport->GetCurrentTexture(EAGXViewportAccessRHI);
@@ -1971,7 +1971,7 @@ static FAGXBuffer InternalCopyTexture2DUpdateRegion(FRHITexture2D* TextureRHI, c
 
 #if PLATFORM_MAC
 	// Expand R8_sRGB into RGBA8_sRGB for Mac.
-	if(Texture->GetFormat() == PF_G8 && (Texture->GetFlags() & TexCreate_SRGB))
+	if(Texture->GetFormat() == PF_G8 && EnumHasAnyFlags(Texture->GetFlags(), TexCreate_SRGB))
 	{
 		const uint32 BufferSize = UpdateRegion.Height * UpdateRegion.Width * sizeof(uint32);		
 		Buffer = GetAGXDeviceContext().CreatePooledBuffer(FAGXPooledBufferArgs(GetAGXDeviceContext().GetDevice(), BufferSize, BUF_Dynamic, mtlpp::StorageMode::Shared));
@@ -2094,7 +2094,7 @@ void FAGXDynamicRHI::RHIUpdateTexture2D(FRHITexture2D* TextureRHI, uint32 MipInd
 		{
 #if PLATFORM_MAC
 			TArray<uint32> ExpandedData;
-			if(Texture->GetFormat() == PF_G8 && (Texture->GetFlags() & TexCreate_SRGB))
+			if(Texture->GetFormat() == PF_G8 && EnumHasAnyFlags(Texture->GetFlags(), TexCreate_SRGB))
 			{
 				ExpandedData.AddZeroed(UpdateRegion.Height * UpdateRegion.Width);
 				InternalExpandR8ToStandardRGBA((uint32*)ExpandedData.GetData(), UpdateRegion, SourcePitch, SourceData);
@@ -2250,7 +2250,7 @@ void FAGXDynamicRHI::RHIUpdateTexture3D(FRHITexture3D* TextureRHI,uint32 MipInde
 		FAGXTexture Tex = Texture->Surface.Texture;
 		
 #if PLATFORM_MAC
-		checkf(!(Texture->GetFormat() == PF_G8 && (Texture->GetFlags() & TexCreate_SRGB)), TEXT("AGXRHI does not support PF_G8_sRGB on 3D, array or cube textures as it requires manual, CPU-side expansion to RGBA8_sRGB which is expensive!"));
+		checkf(!(Texture->GetFormat() == PF_G8 && EnumHasAnyFlags(Texture->GetFlags(), TexCreate_SRGB)), TEXT("AGXRHI does not support PF_G8_sRGB on 3D, array or cube textures as it requires manual, CPU-side expansion to RGBA8_sRGB which is expensive!"));
 #endif
 		if(Tex.GetStorageMode() == mtlpp::StorageMode::Private)
 		{
diff --git a/Engine/Source/Runtime/Apple/MetalRHI/Private/MetalPipeline.cpp b/Engine/Source/Runtime/Apple/MetalRHI/Private/MetalPipeline.cpp
--- a/Engine/Source/Runtime/Apple/MetalRHI/Private/MetalPipeline.cpp
+++ b/Engine/Source/Runtime/Apple/MetalRHI/Private/MetalPipeline.cpp
@@ -208,7 +208,7 @@ struct FMetalGraphicsPipelineKey
 
 			mtlpp::PixelFormat MetalFormat = (mtlpp::PixelFormat)GPixelFormats[TargetFormat].PlatformFormat;
 			ETextureCreateFlags Flags = (ETextureCreateFlags)Init.RenderTargetFlags[i];
-			if (Flags & TexCreate_SRGB)
+			if (EnumHasAnyFlags(Flags, TexCreate_SRGB))
 			{
 #if PLATFORM_MAC // Expand as R8_sRGB is iOS only.
 				if (MetalFormat == mtlpp::PixelFormat::R8Unorm)
@@ -653,7 +653,7 @@ static bool ConfigureRenderPipelineDescriptor(mtlpp::RenderPipelineDescriptor& R
 		
 		mtlpp::PixelFormat MetalFormat = (mtlpp::PixelFormat)GPixelFormats[TargetFormat].PlatformFormat;
 		ETextureCreateFlags Flags = (ETextureCreateFlags)Init.RenderTargetFlags[ActiveTargetIndex];
-		if (Flags & TexCreate_SRGB)
+		if (EnumHasAnyFlags(Flags, TexCreate_SRGB))
 		{
 #if PLATFORM_MAC // Expand as R8_sRGB is iOS only.
 			if (MetalFormat == mtlpp::PixelFormat::R8Unorm)
diff --git a/Engine/Source/Runtime/Apple/MetalRHI/Private/MetalRenderTarget.cpp b/Engine/Source/Runtime/Apple/MetalRHI/Private/MetalRenderTarget.cpp
--- a/Engine/Source/Runtime/Apple/MetalRHI/Private/MetalRenderTarget.cpp
+++ b/Engine/Source/Runtime/Apple/MetalRHI/Private/MetalRenderTarget.cpp
@@ -240,7 +240,7 @@ void FMetalDynamicRHI::RHIReadSurfaceData(FRHITexture* TextureRHI, FIntRect Rect
 	mtlpp::Region Region(Rect.Min.X, Rect.Min.Y, SizeX, SizeY);
     
 	FMetalTexture Texture = Surface->Texture;
-    if(!Texture && (Surface->Flags & TexCreate_Presentable))
+    if(!Texture && EnumHasAnyFlags(Surface->Flags, TexCreate_Presentable))
     {
         Texture = Surface->GetCurrentTexture();
     }
@@ -388,7 +388,7 @@ void FMetalDynamicRHI::RHIReadSurfaceFloatData(FRHITexture* TextureRHI, FIntRect
 	FMetalSurface* Surface = GetMetalSurfaceFromRHITexture(TextureRHI);
 	
     FMetalTexture Texture = Surface->Texture;
-    if(!Texture && (Surface->Flags & TexCreate_Presentable))
+    if(!Texture && EnumHasAnyFlags(Surface->Flags, TexCreate_Presentable))
     {
 		Texture = Surface->GetCurrentTexture();
     }
diff --git a/Engine/Source/Runtime/Apple/MetalRHI/Private/MetalStateCache.cpp b/Engine/Source/Runtime/Apple/MetalRHI/Private/MetalStateCache.cpp
--- a/Engine/Source/Runtime/Apple/MetalRHI/Private/MetalStateCache.cpp
+++ b/Engine/Source/Runtime/Apple/MetalRHI/Private/MetalStateCache.cpp
@@ -1169,7 +1169,7 @@ void FMetalStateCache::SetVisibilityResultMode(mtlpp::VisibilityResultMode const
 void FMetalStateCache::ConditionalUpdateBackBuffer(FMetalSurface& Surface)
 {
 	// are we setting the back buffer? if so, make sure we have the drawable
-	if ((Surface.Flags & TexCreate_Presentable))
+	if (EnumHasAnyFlags(Surface.Flags, TexCreate_Presentable))
 	{
 		// update the back buffer texture the first time used this frame
 		if (Surface.Texture.GetPtr() == nil)
diff --git a/Engine/Source/Runtime/Apple/MetalRHI/Private/MetalTexture.cpp b/Engine/Source/Runtime/Apple/MetalRHI/Private/MetalTexture.cpp
--- a/Engine/Source/Runtime/Apple/MetalRHI/Private/MetalTexture.cpp
+++ b/Engine/Source/Runtime/Apple/MetalRHI/Private/MetalTexture.cpp
@@ -45,19 +45,19 @@ FMetalSurface* GetMetalSurfaceFromRHITexture(FRHITexture* Texture)
 
 static bool IsRenderTarget(ETextureCreateFlags Flags)
 {
-	return (Flags & (TexCreate_RenderTargetable | TexCreate_ResolveTargetable | TexCreate_DepthStencilTargetable | TexCreate_DepthStencilResolveTarget)) != 0;
+	return EnumHasAnyFlags(Flags, TexCreate_RenderTargetable | TexCreate_ResolveTargetable | TexCreate_DepthStencilTargetable | TexCreate_DepthStencilResolveTarget);
 }
 
 static mtlpp::TextureUsage ConvertFlagsToUsage(ETextureCreateFlags Flags)
 {
 	NSUInteger Usage = mtlpp::TextureUsage::Unknown;
-    if(Flags & (TexCreate_ShaderResource|TexCreate_ResolveTargetable|TexCreate_DepthStencilTargetable))
+    if (EnumHasAnyFlags(Flags, TexCreate_ShaderResource|TexCreate_ResolveTargetable|TexCreate_DepthStencilTargetable))
 	{
 		Usage |= mtlpp::TextureUsage::ShaderRead;
 		Usage |= mtlpp::TextureUsage::PixelFormatView;
 	}
 	
-	if (Flags & TexCreate_UAV)
+	if (EnumHasAnyFlags(Flags, TexCreate_UAV))
 	{
 		Usage |= mtlpp::TextureUsage::ShaderRead;
 		Usage |= mtlpp::TextureUsage::ShaderWrite;
@@ -65,7 +65,7 @@ static mtlpp::TextureUsage ConvertFlagsToUsage(ETextureCreateFlags Flags)
 	}
 	
 	// offline textures are normal shader read textures
-	if (Flags & TexCreate_OfflineProcessed)
+	if (EnumHasAnyFlags(Flags, TexCreate_OfflineProcessed))
 	{
 		Usage |= mtlpp::TextureUsage::ShaderRead;
 	}
@@ -73,8 +73,8 @@ static mtlpp::TextureUsage ConvertFlagsToUsage(ETextureCreateFlags Flags)
 	//if the high level is doing manual resolves then the textures specifically markes as resolve targets
 	//are likely to be used in a manual shader resolve by the high level and must be bindable as rendertargets.
 	const bool bSeparateResolveTargets = FMetalCommandQueue::SupportsSeparateMSAAAndResolveTarget();
-	const bool bResolveTarget = (Flags & TexCreate_ResolveTargetable);
-	if ((Flags & (TexCreate_RenderTargetable|TexCreate_DepthStencilTargetable|TexCreate_DepthStencilResolveTarget)) || (bResolveTarget && bSeparateResolveTargets))
+	const bool bResolveTarget = EnumHasAnyFlags(Flags, TexCreate_ResolveTargetable);
+	if (EnumHasAnyFlags(Flags, TexCreate_RenderTargetable|TexCreate_DepthStencilTargetable|TexCreate_DepthStencilResolveTarget) || (bResolveTarget && bSeparateResolveTargets))
 	{
 		Usage |= mtlpp::TextureUsage::RenderTarget;
 		Usage |= mtlpp::TextureUsage::ShaderRead;
@@ -484,7 +484,7 @@ void FMetalSurface::Init(FMetalSurface& Source, NSRange MipRange)
 	};
 
 	mtlpp::PixelFormat MetalFormat = (mtlpp::PixelFormat)GPixelFormats[PixelFormat].PlatformFormat;
-	if (Flags & TexCreate_SRGB)
+	if (EnumHasAnyFlags(Flags, TexCreate_SRGB))
 	{
 		// Ensure we have the correct sRGB target format if we create a new texture view rather than using the source texture
 		MetalFormat = ToSRGBFormat(MetalFormat);
@@ -532,7 +532,7 @@ void FMetalSurface::Init(FMetalSurface& Source, NSRange MipRange, EPixelFormat F
 	
 	bool bUseSourceTex = (Source.PixelFormat != PF_DepthStencil) && Source.PixelFormat == Format && MipRange.location == 0 && MipRange.length == Source.Texture.GetMipmapLevelCount();
 	
-	if (Flags & TexCreate_SRGB)
+	if (EnumHasAnyFlags(Flags, TexCreate_SRGB))
 	{
 		if(bSRGBForceDisable)
 		{
@@ -694,7 +694,7 @@ FMetalSurface::FMetalSurface(ERHIResourceType ResourceType, EPixelFormat Format,
 	FPlatformAtomics::InterlockedExchange(&Written, 0);
 	mtlpp::PixelFormat MTLFormat = (mtlpp::PixelFormat)GPixelFormats[Format].PlatformFormat;
 	
-	if (Flags & TexCreate_SRGB)
+	if (EnumHasAnyFlags(Flags, TexCreate_SRGB))
 	{
 		MTLFormat = ToSRGBFormat(MTLFormat);
 	}
@@ -704,7 +704,7 @@ FMetalSurface::FMetalSurface(ERHIResourceType ResourceType, EPixelFormat Format,
 	
 	
 	// the special back buffer surface will be updated in GetMetalDeviceContext().BeginDrawingViewport - no need to set the texture here
-	if (Flags & TexCreate_Presentable)
+	if (EnumHasAnyFlags(Flags, TexCreate_Presentable))
 	{
 		return;
 	}
@@ -760,7 +760,7 @@ FMetalSurface::FMetalSurface(ERHIResourceType ResourceType, EPixelFormat Format,
 	{
 		Desc.SetUsage(ConvertFlagsToUsage(Flags));
 		
-		if((Flags & TexCreate_CPUReadback) && !(Flags & (TexCreate_RenderTargetable|TexCreate_DepthStencilTargetable|TexCreate_FastVRAM)))
+		if(EnumHasAnyFlags(Flags, TexCreate_CPUReadback) && !EnumHasAnyFlags(Flags, TexCreate_RenderTargetable|TexCreate_DepthStencilTargetable|TexCreate_FastVRAM))
 		{
 			Desc.SetCpuCacheMode(mtlpp::CpuCacheMode::DefaultCache);
 #if PLATFORM_MAC
@@ -771,7 +771,7 @@ FMetalSurface::FMetalSurface(ERHIResourceType ResourceType, EPixelFormat Format,
 			Desc.SetResourceOptions((mtlpp::ResourceOptions)(mtlpp::ResourceOptions::CpuCacheModeDefaultCache|mtlpp::ResourceOptions::StorageModeShared));
 #endif
 		}
-		else if(((Flags & (TexCreate_NoTiling)) && !(Flags & (TexCreate_FastVRAM|TexCreate_DepthStencilTargetable|TexCreate_RenderTargetable|TexCreate_UAV))))
+		else if (EnumHasAnyFlags(Flags, TexCreate_NoTiling) && !EnumHasAnyFlags(Flags, TexCreate_FastVRAM|TexCreate_DepthStencilTargetable|TexCreate_RenderTargetable|TexCreate_UAV))
 		{
 #if PLATFORM_MAC
 			Desc.SetCpuCacheMode(mtlpp::CpuCacheMode::WriteCombined);
@@ -783,7 +783,7 @@ FMetalSurface::FMetalSurface(ERHIResourceType ResourceType, EPixelFormat Format,
 			Desc.SetResourceOptions((mtlpp::ResourceOptions)(mtlpp::ResourceOptions::CpuCacheModeDefaultCache|mtlpp::ResourceOptions::StorageModeShared));
 #endif
 		}
-		else if (Flags & (TexCreate_RenderTargetable|TexCreate_DepthStencilTargetable|TexCreate_ResolveTargetable|TexCreate_DepthStencilResolveTarget))
+		else if (EnumHasAnyFlags(Flags, TexCreate_RenderTargetable|TexCreate_DepthStencilTargetable|TexCreate_ResolveTargetable|TexCreate_DepthStencilResolveTarget))
 		{
 			check(!(Flags & TexCreate_CPUReadback));
 			Desc.SetCpuCacheMode(mtlpp::CpuCacheMode::DefaultCache);
@@ -988,7 +988,7 @@ FMetalSurface::FMetalSurface(ERHIResourceType ResourceType, EPixelFormat Format,
 			//we always require an MSAAResolveTexture if MSAATexture is active.
 			check(!MSAATexture || MSAAResolveTexture || bDepthButNoResolveSupported);
 			
-			NSLog(@"Creating %dx MSAA %d x %d %s surface", (int32)Desc.GetSampleCount(), SizeX, SizeY, (Flags & TexCreate_RenderTargetable) ? "Color" : "Depth");
+			NSLog(@"Creating %dx MSAA %d x %d %s surface", (int32)Desc.GetSampleCount(), SizeX, SizeY, EnumHasAnyFlags(Flags, TexCreate_RenderTargetable) ? "Color" : "Depth");
 			if (MSAATexture.GetPtr() == nil)
 			{
 				NSLog(@"Failed to create texture, desc  %@", Desc.GetPtr());
@@ -1177,7 +1177,7 @@ id <MTLBuffer> FMetalSurface::AllocSurface(uint32 MipIndex, uint32 ArrayIndex, E
 	
 #if PLATFORM_MAC
 	// Expand R8_sRGB into RGBA8_sRGB for Mac.
-	if (PixelFormat == PF_G8 && (Flags & TexCreate_SRGB) && Type == RRT_Texture2D && LockMode == RLM_WriteOnly)
+	if (PixelFormat == PF_G8 && EnumHasAnyFlags(Flags, TexCreate_SRGB) && Type == RRT_Texture2D && LockMode == RLM_WriteOnly)
 	{
 		DestStride = FMath::Max<uint32>(SizeX >> MipIndex, 1);
 	}
@@ -1211,7 +1211,7 @@ void FMetalSurface::UpdateSurfaceAndDestroySourceBuffer(id <MTLBuffer> SourceBuf
 	}
 #if PLATFORM_MAC
 	// Expand R8_sRGB into RGBA8_sRGB for Mac.
-	if (PixelFormat == PF_G8 && (Flags & TexCreate_SRGB) && Type == RRT_Texture2D)
+	if (PixelFormat == PF_G8 && EnumHasAnyFlags(Flags, TexCreate_SRGB) && Type == RRT_Texture2D)
 	{
 		TArray<uint8> Data;
 		uint8* ExpandedMem = (uint8*) SourceBuffer.contents;
@@ -1418,7 +1418,7 @@ void* FMetalSurface::Lock(uint32 MipIndex, uint32 ArrayIndex, EResourceLockMode
 			
 #if PLATFORM_MAC
 			// Pack RGBA8_sRGB into R8_sRGB for Mac.
-			if (PixelFormat == PF_G8 && (Flags & TexCreate_SRGB) && Type == RRT_Texture2D)
+			if (PixelFormat == PF_G8 && EnumHasAnyFlags(Flags, TexCreate_SRGB) && Type == RRT_Texture2D)
 			{
 				TArray<uint8> Data;
 				uint8* ExpandedMem = (uint8*)SourceData.GetContents();
@@ -1563,7 +1563,7 @@ uint32 FMetalSurface::GetMipSize(uint32 MipIndex, uint32* Stride, bool bSingleLa
 		NumBlocksY = FMath::Max<uint32>(NumBlocksY, 2);
 	}
 #if PLATFORM_MAC
-	else if (PixelFormat == PF_G8 && (Flags & TexCreate_SRGB))
+	else if (PixelFormat == PF_G8 && EnumHasAnyFlags(Flags, TexCreate_SRGB))
 	{
 		// RGBA_sRGB is the closest match - so expand the data.
 		NumBlocksX *= 4;
@@ -1623,7 +1623,7 @@ uint32 FMetalSurface::GetNumFaces()
 
 FMetalTexture FMetalSurface::GetDrawableTexture()
 {
-	if (!Texture && (Flags & TexCreate_Presentable))
+	if (!Texture && EnumHasAnyFlags(Flags, TexCreate_Presentable))
 	{
 		check(Viewport);
 		Texture = Viewport->GetDrawableTexture(EMetalViewportAccessRHI);
@@ -1634,7 +1634,7 @@ FMetalTexture FMetalSurface::GetDrawableTexture()
 ns::AutoReleased<FMetalTexture> FMetalSurface::GetCurrentTexture()
 {
 	ns::AutoReleased<FMetalTexture> Tex;
-	if (Viewport && (Flags & TexCreate_Presentable))
+	if (Viewport && EnumHasAnyFlags(Flags, TexCreate_Presentable))
 	{
 		check(Viewport);
 		Tex = Viewport->GetCurrentTexture(EMetalViewportAccessRHI);
@@ -1973,7 +1973,7 @@ static FMetalBuffer InternalCopyTexture2DUpdateRegion(FRHITexture2D* TextureRHI,
 
 #if PLATFORM_MAC
 	// Expand R8_sRGB into RGBA8_sRGB for Mac.
-	if(Texture->GetFormat() == PF_G8 && (Texture->GetFlags() & TexCreate_SRGB))
+	if(Texture->GetFormat() == PF_G8 && EnumHasAnyFlags(Texture->GetFlags(), TexCreate_SRGB))
 	{
 		const uint32 BufferSize = UpdateRegion.Height * UpdateRegion.Width * sizeof(uint32);		
 		Buffer = GetMetalDeviceContext().CreatePooledBuffer(FMetalPooledBufferArgs(GetMetalDeviceContext().GetDevice(), BufferSize, BUF_Dynamic, mtlpp::StorageMode::Shared));
@@ -2096,7 +2096,7 @@ void FMetalDynamicRHI::RHIUpdateTexture2D(FRHITexture2D* TextureRHI, uint32 MipI
 		{
 #if PLATFORM_MAC
 			TArray<uint32> ExpandedData;
-			if(Texture->GetFormat() == PF_G8 && (Texture->GetFlags() & TexCreate_SRGB))
+			if(Texture->GetFormat() == PF_G8 && EnumHasAnyFlags(Texture->GetFlags(), TexCreate_SRGB))
 			{
 				ExpandedData.AddZeroed(UpdateRegion.Height * UpdateRegion.Width);
 				InternalExpandR8ToStandardRGBA((uint32*)ExpandedData.GetData(), UpdateRegion, SourcePitch, SourceData);
@@ -2252,7 +2252,7 @@ void FMetalDynamicRHI::RHIUpdateTexture3D(FRHITexture3D* TextureRHI,uint32 MipIn
 		FMetalTexture Tex = Texture->Surface.Texture;
 		
 #if PLATFORM_MAC
-		checkf(!(Texture->GetFormat() == PF_G8 && (Texture->GetFlags() & TexCreate_SRGB)), TEXT("MetalRHI does not support PF_G8_sRGB on 3D, array or cube textures as it requires manual, CPU-side expansion to RGBA8_sRGB which is expensive!"));
+		checkf(!(Texture->GetFormat() == PF_G8 && EnumHasAnyFlags(Texture->GetFlags(), TexCreate_SRGB)), TEXT("MetalRHI does not support PF_G8_sRGB on 3D, array or cube textures as it requires manual, CPU-side expansion to RGBA8_sRGB which is expensive!"));
 #endif
 		if(Tex.GetStorageMode() == mtlpp::StorageMode::Private)
 		{
diff --git a/Engine/Source/Runtime/D3D12RHI/Private/D3D12RHIPrivate.h b/Engine/Source/Runtime/D3D12RHI/Private/D3D12RHIPrivate.h
--- a/Engine/Source/Runtime/D3D12RHI/Private/D3D12RHIPrivate.h
+++ b/Engine/Source/Runtime/D3D12RHI/Private/D3D12RHIPrivate.h
@@ -998,8 +998,8 @@ public:
 	virtual void* CreateVirtualTexture(ETextureCreateFlags InFlags, D3D12_RESOURCE_DESC& ResourceDesc, const struct FD3D12TextureLayout& TextureLayout, FD3D12Resource** ppResource, FPlatformMemory::FPlatformVirtualMemoryBlock& RawTextureBlock, D3D12_RESOURCE_STATES InitialUsage = D3D12_RESOURCE_STATE_PIXEL_SHADER_RESOURCE) = 0;
 	virtual void DestroyVirtualTexture(ETextureCreateFlags InFlags, void* RawTextureMemory, FPlatformMemory::FPlatformVirtualMemoryBlock& RawTextureBlock, uint64 CommittedTextureSize) = 0;
 #endif
-	virtual bool HandleSpecialLock(void*& MemoryOut, uint32 MipIndex, uint32 ArrayIndex, uint32 InFlags, EResourceLockMode LockMode, const FD3D12TextureLayout& TextureLayout, void* RawTextureMemory, uint32& DestStride) { return false; }
-	virtual bool HandleSpecialUnlock(FRHICommandListBase* RHICmdList, uint32 MipIndex, uint32 InFlags, const struct FD3D12TextureLayout& TextureLayout, void* RawTextureMemory) { return false; }
+	virtual bool HandleSpecialLock(void*& MemoryOut, uint32 MipIndex, uint32 ArrayIndex, ETextureCreateFlags InFlags, EResourceLockMode LockMode, const FD3D12TextureLayout& TextureLayout, void* RawTextureMemory, uint32& DestStride) { return false; }
+	virtual bool HandleSpecialUnlock(FRHICommandListBase* RHICmdList, uint32 MipIndex, ETextureCreateFlags InFlags, const struct FD3D12TextureLayout& TextureLayout, void* RawTextureMemory) { return false; }
 
 	FD3D12Adapter& GetAdapter(uint32_t Index = 0) { return *ChosenAdapters[Index]; }
 	int32 GetNumAdapters() const { return ChosenAdapters.Num(); }
@@ -1400,14 +1400,14 @@ inline DXGI_FORMAT FindDepthStencilResourceDXGIFormat(DXGI_FORMAT InFormat)
 	return InFormat;
 }
 
-inline DXGI_FORMAT GetPlatformTextureResourceFormat(DXGI_FORMAT InFormat, uint32 InFlags)
+inline DXGI_FORMAT GetPlatformTextureResourceFormat(DXGI_FORMAT InFormat, ETextureCreateFlags InFlags)
 {
 	// Find valid shared texture format
-	if (InFlags & TexCreate_Shared)
+	if (EnumHasAnyFlags(InFlags, TexCreate_Shared))
 	{
-		return FindSharedResourceDXGIFormat(InFormat, InFlags & TexCreate_SRGB);
+		return FindSharedResourceDXGIFormat(InFormat, EnumHasAnyFlags(InFlags, TexCreate_SRGB));
 	}
-	if (InFlags & TexCreate_DepthStencilTargetable)
+	if (EnumHasAnyFlags(InFlags, TexCreate_DepthStencilTargetable))
 	{
 		return FindDepthStencilResourceDXGIFormat(InFormat);
 	}
@@ -1511,15 +1511,14 @@ static void TranslateRenderTargetFormats(
 		checkSlow(PsoInit.RenderTargetFormats[RTIdx] == PF_Unknown || GPixelFormats[PsoInit.RenderTargetFormats[RTIdx]].Supported);
 
 		DXGI_FORMAT PlatformFormat = (DXGI_FORMAT)GPixelFormats[PsoInit.RenderTargetFormats[RTIdx]].PlatformFormat;
-		uint32 Flags = PsoInit.RenderTargetFlags[RTIdx];
+		ETextureCreateFlags Flags = PsoInit.RenderTargetFlags[RTIdx];
 
-		RTFormatArray.RTFormats[RTIdx] = D3D12RHI::FindShaderResourceDXGIFormat( GetPlatformTextureResourceFormat(PlatformFormat, Flags), (Flags & TexCreate_SRGB) != 0 );
+		RTFormatArray.RTFormats[RTIdx] = D3D12RHI::FindShaderResourceDXGIFormat( GetPlatformTextureResourceFormat(PlatformFormat, Flags), EnumHasAnyFlags(Flags, TexCreate_SRGB) );
 	}
 
 	checkSlow(PsoInit.DepthStencilTargetFormat == PF_Unknown || GPixelFormats[PsoInit.DepthStencilTargetFormat].Supported);
 
 	DXGI_FORMAT PlatformFormat = (DXGI_FORMAT)GPixelFormats[PsoInit.DepthStencilTargetFormat].PlatformFormat;
-	uint32 Flags = PsoInit.DepthStencilTargetFlag;
 
 	DSVFormat = D3D12RHI::FindDepthStencilDXGIFormat( GetPlatformTextureResourceFormat(PlatformFormat, PsoInit.DepthStencilTargetFlag) );
 }
diff --git a/Engine/Source/Runtime/D3D12RHI/Private/D3D12RenderTarget.cpp b/Engine/Source/Runtime/D3D12RHI/Private/D3D12RenderTarget.cpp
--- a/Engine/Source/Runtime/D3D12RHI/Private/D3D12RenderTarget.cpp
+++ b/Engine/Source/Runtime/D3D12RHI/Private/D3D12RenderTarget.cpp
@@ -204,12 +204,12 @@ void FD3D12CommandContext::RHICopyToResolveTarget(FRHITexture* SourceTextureRHI,
 
 	// Could be back buffer reference texture, so get the correct D3D12 texture here
 	// We know already that it's a FD3D12Texture2D so cast is safe
-	if (SourceTextureRHI->GetFlags() & TexCreate_Presentable)
+	if (EnumHasAnyFlags(SourceTextureRHI->GetFlags(), TexCreate_Presentable))
 	{
 		FD3D12BackBufferReferenceTexture2D* BufferBufferReferenceTexture = (FD3D12BackBufferReferenceTexture2D*)SourceTextureRHI;
 		SourceTextureRHI = BufferBufferReferenceTexture->GetBackBufferTexture();
 	}
-	if (DestTextureRHI->GetFlags() & TexCreate_Presentable)
+	if (EnumHasAnyFlags(DestTextureRHI->GetFlags(), TexCreate_Presentable))
 	{
 		FD3D12BackBufferReferenceTexture2D* BufferBufferReferenceTexture = (FD3D12BackBufferReferenceTexture2D*)DestTextureRHI;
 		DestTextureRHI = BufferBufferReferenceTexture->GetBackBufferTexture();
@@ -1154,7 +1154,7 @@ void FD3D12DynamicRHI::RHIReadSurfaceData(FRHITexture* InRHITexture, FIntRect In
 
 	// Could be back buffer reference texture, so get the correct D3D12 texture here
 	// We know already that it's a FD3D12Texture2D so cast is safe
-	if (InRHITexture->GetFlags() & TexCreate_Presentable)
+	if (EnumHasAnyFlags(InRHITexture->GetFlags(), TexCreate_Presentable))
 	{
 		FD3D12BackBufferReferenceTexture2D* BufferBufferReferenceTexture = (FD3D12BackBufferReferenceTexture2D*)InRHITexture;
 		InRHITexture = BufferBufferReferenceTexture->GetBackBufferTexture();
diff --git a/Engine/Source/Runtime/D3D12RHI/Private/D3D12Texture.cpp b/Engine/Source/Runtime/D3D12RHI/Private/D3D12Texture.cpp
--- a/Engine/Source/Runtime/D3D12RHI/Private/D3D12Texture.cpp
+++ b/Engine/Source/Runtime/D3D12RHI/Private/D3D12Texture.cpp
@@ -390,7 +390,7 @@ void FD3D12TextureStats::D3D12TextureDeleted(TD3D12Texture2D<BaseResourceType>&
 		{
 			const D3D12_RESOURCE_DESC& Desc = D3D12Texture2D->GetDesc();
 			const int64 TextureSize = Texture.GetMemorySize();
-			ensure(TextureSize > 0 || (Texture.Flags & TexCreate_Virtual) || Texture.GetAliasingSourceTexture() != nullptr);
+			ensure(TextureSize > 0 || EnumHasAnyFlags(Texture.Flags, TexCreate_Virtual) || Texture.GetAliasingSourceTexture() != nullptr);
 
 			UpdateD3D12TextureStats(Desc, -TextureSize, false, Texture.IsCubemap(), Texture.IsStreamable());
 
@@ -695,7 +695,7 @@ void SafeCreateTexture2D(FD3D12Device* pDevice,
 	{
 #endif // #if GUARDED_TEXTURE_CREATES
 
-		const D3D12_HEAP_TYPE HeapType = (Flags & TexCreate_CPUReadback) ? D3D12_HEAP_TYPE_READBACK : D3D12_HEAP_TYPE_DEFAULT;
+		const D3D12_HEAP_TYPE HeapType = EnumHasAnyFlags(Flags, TexCreate_CPUReadback) ? D3D12_HEAP_TYPE_READBACK : D3D12_HEAP_TYPE_DEFAULT;
 
 		switch (HeapType)
 		{
@@ -743,7 +743,7 @@ void SafeCreateTexture2D(FD3D12Device* pDevice,
 }
 
 
-static void DetermineTexture2DResourceFlagsAndLayout(uint32 SizeX, uint32 SizeY, uint32 SizeZ, uint32 NumMips, uint32 NumSamples, uint32 Flags, EPixelFormat Format, D3D12_RESOURCE_FLAGS& OutResourceFlags, D3D12_TEXTURE_LAYOUT& OutLayout, bool& bOutCreateRTV, bool& bOutCreateDSV, bool& bOutCreateSRV)
+static void DetermineTexture2DResourceFlagsAndLayout(uint32 SizeX, uint32 SizeY, uint32 SizeZ, uint32 NumMips, uint32 NumSamples, ETextureCreateFlags Flags, EPixelFormat Format, D3D12_RESOURCE_FLAGS& OutResourceFlags, D3D12_TEXTURE_LAYOUT& OutLayout, bool& bOutCreateRTV, bool& bOutCreateDSV, bool& bOutCreateSRV)
 {
 	TRACE_CPUPROFILER_EVENT_SCOPE(DetermineTexture2DResourceFlagsAndLayout);
 
@@ -753,42 +753,37 @@ static void DetermineTexture2DResourceFlagsAndLayout(uint32 SizeX, uint32 SizeY,
 	bOutCreateDSV = false;
 	bOutCreateSRV = true;
 
-	if (Flags & TexCreate_CPUReadback)
+	if (EnumHasAllFlags(Flags, TexCreate_CPUReadback))
 	{
-		check(!(Flags & TexCreate_RenderTargetable));
-		check(!(Flags & TexCreate_DepthStencilTargetable));
-		check(!(Flags & TexCreate_ShaderResource));
+		check(!EnumHasAnyFlags(Flags, TexCreate_RenderTargetable | TexCreate_DepthStencilTargetable | TexCreate_ShaderResource));
 		bOutCreateSRV = false;
 	}
 
-	if (Flags & TexCreate_DisableSRVCreation)
+	if (EnumHasAnyFlags(Flags, TexCreate_DisableSRVCreation))
 	{
 		bOutCreateSRV = false;
 	}
 
-	if (Flags & TexCreate_Shared)
+	if (EnumHasAnyFlags(Flags, TexCreate_Shared))
 	{
 		OutResourceFlags |= D3D12_RESOURCE_FLAG_ALLOW_SIMULTANEOUS_ACCESS;
 	}
 
-	if (Flags & TexCreate_RenderTargetable)
+	if (EnumHasAnyFlags(Flags, TexCreate_RenderTargetable))
 	{
-		check(!(Flags & TexCreate_DepthStencilTargetable));
-		check(!(Flags & TexCreate_ResolveTargetable));
+		check(!EnumHasAnyFlags(Flags, TexCreate_DepthStencilTargetable | TexCreate_ResolveTargetable));
 		OutResourceFlags |= D3D12_RESOURCE_FLAG_ALLOW_RENDER_TARGET;
 		bOutCreateRTV = true;
 	}
-	else if (Flags & TexCreate_DepthStencilTargetable)
+	else if (EnumHasAnyFlags(Flags, TexCreate_DepthStencilTargetable))
 	{
-		check(!(Flags & TexCreate_RenderTargetable));
-		check(!(Flags & TexCreate_ResolveTargetable));
+		check(!EnumHasAnyFlags(Flags, TexCreate_RenderTargetable | TexCreate_ResolveTargetable));
 		OutResourceFlags |= D3D12_RESOURCE_FLAG_ALLOW_DEPTH_STENCIL;
 		bOutCreateDSV = true;
 	}
-	else if (Flags & TexCreate_ResolveTargetable)
+	else if (EnumHasAnyFlags(Flags, TexCreate_ResolveTargetable))
 	{
-		check(!(Flags & TexCreate_RenderTargetable));
-		check(!(Flags & TexCreate_DepthStencilTargetable));
+		check(!EnumHasAnyFlags(Flags, TexCreate_RenderTargetable | TexCreate_DepthStencilTargetable));
 		if (Format == PF_DepthStencil || Format == PF_ShadowDepth || Format == PF_D24)
 		{
 			OutResourceFlags |= D3D12_RESOURCE_FLAG_ALLOW_DEPTH_STENCIL;
@@ -801,12 +796,12 @@ static void DetermineTexture2DResourceFlagsAndLayout(uint32 SizeX, uint32 SizeY,
 		}
 	}
 
-	if (Flags & TexCreate_UAV)
+	if (EnumHasAnyFlags(Flags, TexCreate_UAV))
 	{
 		OutResourceFlags |= D3D12_RESOURCE_FLAG_ALLOW_UNORDERED_ACCESS;
 	}
 
-	if (bOutCreateDSV && !(Flags & TexCreate_ShaderResource))
+	if (bOutCreateDSV && !EnumHasAnyFlags(Flags, TexCreate_ShaderResource))
 	{
 		// Only deny shader resources if it's a depth resource that will never be used as SRV
 		OutResourceFlags |= D3D12_RESOURCE_FLAG_DENY_SHADER_RESOURCE;
@@ -842,7 +837,7 @@ TD3D12Texture2D<BaseResourceType>* FD3D12DynamicRHI::CreateD3D12Texture2D(FRHICo
 
 	SCOPE_CYCLE_COUNTER(STAT_D3D12CreateTextureTime);
 
-	const bool bSRGB = (Flags & TexCreate_SRGB) != 0;
+	const bool bSRGB = EnumHasAnyFlags(Flags, TexCreate_SRGB);
 
 	const DXGI_FORMAT PlatformResourceFormat = GetPlatformTextureResourceFormat((DXGI_FORMAT)GPixelFormats[Format].PlatformFormat, Flags);
 	const DXGI_FORMAT PlatformShaderResourceFormat = FindShaderResourceDXGIFormat(PlatformResourceFormat, bSRGB);
@@ -853,7 +848,7 @@ TD3D12Texture2D<BaseResourceType>* FD3D12DynamicRHI::CreateD3D12Texture2D(FRHICo
 	uint32 ActualMSAAQuality = GetMaxMSAAQuality(ActualMSAACount);
 
 	// 0xffffffff means not supported
-	if (ActualMSAAQuality == 0xffffffff || (Flags & TexCreate_Shared) != 0)
+	if (ActualMSAAQuality == 0xffffffff || EnumHasAnyFlags(Flags, TexCreate_Shared))
 	{
 		// no MSAA
 		ActualMSAACount = 1;
@@ -924,7 +919,7 @@ TD3D12Texture2D<BaseResourceType>* FD3D12DynamicRHI::CreateD3D12Texture2D(FRHICo
 
 		if (ResourceAllocator)
 		{
-			const D3D12_HEAP_TYPE HeapType = (Flags & TexCreate_CPUReadback) ? D3D12_HEAP_TYPE_READBACK : D3D12_HEAP_TYPE_DEFAULT;
+			const D3D12_HEAP_TYPE HeapType = EnumHasAnyFlags(Flags, TexCreate_CPUReadback) ? D3D12_HEAP_TYPE_READBACK : D3D12_HEAP_TYPE_DEFAULT;
 			ResourceAllocator->AllocateTexture(Device->GetGPUIndex(), HeapType, TextureDesc, (EPixelFormat)Format, ED3D12ResourceStateMode::Default, CreateState, ClearValuePtr, CreateInfo.DebugName, Location);
 			Location.SetOwner(NewTexture);
 		}
@@ -956,7 +951,7 @@ TD3D12Texture2D<BaseResourceType>* FD3D12DynamicRHI::CreateD3D12Texture2D(FRHICo
 			TRACE_CPUPROFILER_EVENT_SCOPE(D3D12RHI::CreateViews);
 			if (bCreateRTV)
 			{
-				const bool bCreateRTVsPerSlice = (Flags & TexCreate_TargetArraySlicesIndependently) && (bTextureArray || bCubeTexture);
+				const bool bCreateRTVsPerSlice = EnumHasAnyFlags(Flags, TexCreate_TargetArraySlicesIndependently) && (bTextureArray || bCubeTexture);
 				NewTexture->SetNumRenderTargetViews(bCreateRTVsPerSlice ? NumMips * TextureDesc.DepthOrArraySize : NumMips);
 
 				// Create a render target view for each mip
@@ -1159,7 +1154,7 @@ FD3D12Texture3D* FD3D12DynamicRHI::CreateD3D12Texture3D(FRHICommandListImmediate
 #if PLATFORM_WINDOWS || PLATFORM_HOLOLENS
 	SCOPE_CYCLE_COUNTER(STAT_D3D12CreateTextureTime);
 
-	const bool bSRGB = (Flags & TexCreate_SRGB) != 0;
+	const bool bSRGB = EnumHasAnyFlags(Flags, TexCreate_SRGB);
 
 	const DXGI_FORMAT PlatformResourceFormat = (DXGI_FORMAT)GPixelFormats[Format].PlatformFormat;
 	const DXGI_FORMAT PlatformShaderResourceFormat = FindShaderResourceDXGIFormat(PlatformResourceFormat, bSRGB);
@@ -1173,23 +1168,22 @@ FD3D12Texture3D* FD3D12DynamicRHI::CreateD3D12Texture3D(FRHICommandListImmediate
 		SizeZ,
 		NumMips);
 
-	if (Flags & TexCreate_UAV)
+	if (EnumHasAnyFlags(Flags, TexCreate_UAV))
 	{
 		TextureDesc.Flags |= D3D12_RESOURCE_FLAG_ALLOW_UNORDERED_ACCESS;
 	}
 
 	bool bCreateRTV = false;
 
-	if (Flags & TexCreate_RenderTargetable)
+	if (EnumHasAnyFlags(Flags, TexCreate_RenderTargetable))
 	{
 		TextureDesc.Flags |= D3D12_RESOURCE_FLAG_ALLOW_RENDER_TARGET;
 		bCreateRTV = true;
 	}
 
 	// Set up the texture bind flags.
-	check(!(Flags & TexCreate_DepthStencilTargetable));
-	check(!(Flags & TexCreate_ResolveTargetable));
-	check(Flags & TexCreate_ShaderResource);
+	check(!EnumHasAnyFlags(Flags, TexCreate_DepthStencilTargetable | TexCreate_ResolveTargetable));
+	check(EnumHasAllFlags(Flags, TexCreate_ShaderResource));
 
 	D3D12_CLEAR_VALUE *ClearValuePtr = nullptr;
 	D3D12_CLEAR_VALUE ClearValue;
@@ -1328,11 +1322,11 @@ FTexture2DRHIRef FD3D12DynamicRHI::RHIAsyncCreateTexture2D(uint32 SizeX, uint32
 {
 	check(GRHISupportsAsyncTextureCreation);
 	
-	const uint32 InvalidFlags = TexCreate_RenderTargetable | TexCreate_ResolveTargetable | TexCreate_DepthStencilTargetable | TexCreate_GenerateMipCapable | TexCreate_UAV | TexCreate_Presentable | TexCreate_CPUReadback;
-	check((Flags & InvalidFlags) == 0);
+	const ETextureCreateFlags InvalidFlags = TexCreate_RenderTargetable | TexCreate_ResolveTargetable | TexCreate_DepthStencilTargetable | TexCreate_GenerateMipCapable | TexCreate_UAV | TexCreate_Presentable | TexCreate_CPUReadback;
+	check(!EnumHasAnyFlags(Flags, InvalidFlags));
 
 	const DXGI_FORMAT PlatformResourceFormat = (DXGI_FORMAT)GPixelFormats[Format].PlatformFormat;
-	const DXGI_FORMAT PlatformShaderResourceFormat = FindShaderResourceDXGIFormat(PlatformResourceFormat, ((Flags & TexCreate_SRGB) != 0));
+	const DXGI_FORMAT PlatformShaderResourceFormat = FindShaderResourceDXGIFormat(PlatformResourceFormat, EnumHasAnyFlags(Flags, TexCreate_SRGB));
 	const D3D12_RESOURCE_DESC TextureDesc = CD3DX12_RESOURCE_DESC::Tex2D(
 		PlatformResourceFormat,
 		SizeX,
@@ -2321,7 +2315,7 @@ static void GetReadBackHeapDescImpl(D3D12_PLACED_SUBRESOURCE_FOOTPRINT& OutFootp
 template<typename RHIResourceType>
 void TD3D12Texture2D<RHIResourceType>::GetReadBackHeapDesc(D3D12_PLACED_SUBRESOURCE_FOOTPRINT& OutFootprint, uint32 InSubresource) const
 {
-	check((RHIResourceType::GetFlags() & TexCreate_CPUReadback) != 0);
+	check(EnumHasAnyFlags(RHIResourceType::GetFlags(), TexCreate_CPUReadback));
 
 	if (InSubresource == 0 && FirstSubresourceFootprint)
 	{
@@ -2661,7 +2655,7 @@ FUpdateTexture3DData FD3D12DynamicRHI::BeginUpdateTexture3D_Internal(FRHITexture
 	check(FormatInfo.BlockSizeZ == 1);
 
 	bool bDoComputeShaderCopy = false; // Compute shader can not cast compressed formats into uint
-	if (CVarUseUpdateTexture3DComputeShader.GetValueOnRenderThread() != 0 && FormatInfo.BlockSizeX == 1 && FormatInfo.BlockSizeY == 1 && Texture->ResourceLocation.GetGPUVirtualAddress() && !(Texture->GetFlags() & TexCreate_OfflineProcessed))
+	if (CVarUseUpdateTexture3DComputeShader.GetValueOnRenderThread() != 0 && FormatInfo.BlockSizeX == 1 && FormatInfo.BlockSizeY == 1 && Texture->ResourceLocation.GetGPUVirtualAddress() && !EnumHasAnyFlags(Texture->GetFlags(), TexCreate_OfflineProcessed))
 	{
 		// Try a compute shader update. This does a memory allocation internally
 		bDoComputeShaderCopy = BeginUpdateTexture3D_ComputeShader(UpdateData, UpdateDataD3D12);
@@ -2934,7 +2928,7 @@ TD3D12Texture2D<BaseResourceType>* FD3D12DynamicRHI::CreateTextureFromResource(b
 
 	SCOPE_CYCLE_COUNTER(STAT_D3D12CreateTextureTime);
 
-	const bool bSRGB = (TexCreateFlags & TexCreate_SRGB) != 0;
+	const bool bSRGB = EnumHasAnyFlags(TexCreateFlags, TexCreate_SRGB);
 
 	const DXGI_FORMAT PlatformResourceFormat = TextureDesc.Format;
 	const DXGI_FORMAT PlatformShaderResourceFormat = FindShaderResourceDXGIFormat(PlatformResourceFormat, bSRGB);
@@ -2946,7 +2940,7 @@ TD3D12Texture2D<BaseResourceType>* FD3D12DynamicRHI::CreateTextureFromResource(b
 
 	// The state this resource will be in when it leaves this function
 	const FD3D12Resource::FD3D12ResourceTypeHelper Type(TextureDesc, D3D12_HEAP_TYPE_DEFAULT);
-	const D3D12_RESOURCE_STATES DestinationState = Type.GetOptimalInitialState(ERHIAccess::Unknown, !(TexCreateFlags & TexCreate_Shared));
+	const D3D12_RESOURCE_STATES DestinationState = Type.GetOptimalInitialState(ERHIAccess::Unknown, !EnumHasAnyFlags(TexCreateFlags, TexCreate_Shared));
 
 	FD3D12Device* Device = Adapter->GetDevice(0);
 	FD3D12Resource* TextureResource = new FD3D12Resource(Device, Device->GetGPUMask(), Resource, DestinationState, TextureDesc);
@@ -2965,7 +2959,7 @@ TD3D12Texture2D<BaseResourceType>* FD3D12DynamicRHI::CreateTextureFromResource(b
 
 	if (bCreateRTV)
 	{
-		const bool bCreateRTVsPerSlice = (TexCreateFlags & TexCreate_TargetArraySlicesIndependently) && (bTextureArray || bCubeTexture);
+		const bool bCreateRTVsPerSlice = EnumHasAnyFlags(TexCreateFlags, TexCreate_TargetArraySlicesIndependently) && (bTextureArray || bCubeTexture);
 		Texture2D->SetNumRenderTargetViews(bCreateRTVsPerSlice ? NumMips * TextureDesc.DepthOrArraySize : NumMips);
 
 		// Create a render target view for each mip
@@ -3182,7 +3176,7 @@ TD3D12Texture2D<BaseResourceType>* FD3D12DynamicRHI::CreateAliasedD3D12Texture2D
 
 	FD3D12Device* Device = Adapter->GetDevice(0);
 
-	const bool bSRGB = (SourceTexture->GetFlags() & TexCreate_SRGB) != 0;
+	const bool bSRGB = EnumHasAnyFlags(SourceTexture->GetFlags(), TexCreate_SRGB);
 
 	const DXGI_FORMAT PlatformResourceFormat = TextureDesc.Format;
 	const DXGI_FORMAT PlatformShaderResourceFormat = FindShaderResourceDXGIFormat(PlatformResourceFormat, bSRGB);
@@ -3212,7 +3206,7 @@ TD3D12Texture2D<BaseResourceType>* FD3D12DynamicRHI::CreateAliasedD3D12Texture2D
 		for (uint32 MipIndex = 0; MipIndex < TextureDesc.MipLevels; MipIndex++)
 		{
 			// These are null because we'll be aliasing them shortly.
-			if ((SourceTexture->Flags & TexCreate_TargetArraySlicesIndependently) && (bTextureArray || bCubeTexture))
+			if (EnumHasAnyFlags(SourceTexture->Flags, TexCreate_TargetArraySlicesIndependently) && (bTextureArray || bCubeTexture))
 			{
 				bCreatedRTVPerSlice = true;
 
@@ -3329,7 +3323,7 @@ void FD3D12CommandContext::RHICopyTexture(FRHITexture* SourceTextureRHI, FRHITex
 	numCopies++;
 	CommandListHandle.FlushResourceBarriers();
 
-	const bool bReadback = (DestTextureRHI->GetFlags() & TexCreate_CPUReadback) != 0;
+	const bool bReadback = EnumHasAnyFlags(DestTextureRHI->GetFlags(), TexCreate_CPUReadback);
 
 	if (CopyInfo.Size != FIntVector::ZeroValue || bReadback)
 	{
diff --git a/Engine/Source/Runtime/D3D12RHI/Private/D3D12Texture.h b/Engine/Source/Runtime/D3D12RHI/Private/D3D12Texture.h
--- a/Engine/Source/Runtime/D3D12RHI/Private/D3D12Texture.h
+++ b/Engine/Source/Runtime/D3D12RHI/Private/D3D12Texture.h
@@ -562,7 +562,7 @@ FORCEINLINE FD3D12TextureBase* GetD3D12TextureFromRHITexture(FRHITexture* Textur
 	
 	// If it's the dummy backbuffer then swap with actual current RHI backbuffer right now
 	FRHITexture* RHITexture = Texture;
-	if (RHITexture && RHITexture->GetFlags() & TexCreate_Presentable)
+	if (RHITexture && EnumHasAnyFlags(RHITexture->GetFlags(), TexCreate_Presentable))
 	{
 		FD3D12BackBufferReferenceTexture2D* BufferBufferReferenceTexture = (FD3D12BackBufferReferenceTexture2D*)RHITexture;
 		RHITexture = BufferBufferReferenceTexture->GetBackBufferTexture();
diff --git a/Engine/Source/Runtime/D3D12RHI/Private/D3D12View.cpp b/Engine/Source/Runtime/D3D12RHI/Private/D3D12View.cpp
--- a/Engine/Source/Runtime/D3D12RHI/Private/D3D12View.cpp
+++ b/Engine/Source/Runtime/D3D12RHI/Private/D3D12View.cpp
@@ -152,7 +152,7 @@ FShaderResourceViewRHIRef FD3D12DynamicRHI::RHICreateShaderResourceView(FRHIText
 	}
 
 	// Allow input CreateInfo to override SRGB and/or format
-	const bool bBaseSRGB = (Texture->GetFlags() & TexCreate_SRGB) != 0;
+	const bool bBaseSRGB = EnumHasAnyFlags(Texture->GetFlags(), TexCreate_SRGB);
 	const bool bSRGB = CreateInfo.SRGBOverride != SRGBO_ForceDisable && bBaseSRGB;
 	const DXGI_FORMAT ViewTextureFormat = (CreateInfo.Format == PF_Unknown) ? BaseTextureFormat : (DXGI_FORMAT)GPixelFormats[CreateInfo.Format].PlatformFormat;
 	SRVDesc.Format = FindShaderResourceDXGIFormat(ViewTextureFormat, bSRGB);
diff --git a/Engine/Source/Runtime/Engine/Private/Rendering/StreamableTextureResource.cpp b/Engine/Source/Runtime/Engine/Private/Rendering/StreamableTextureResource.cpp
--- a/Engine/Source/Runtime/Engine/Private/Rendering/StreamableTextureResource.cpp
+++ b/Engine/Source/Runtime/Engine/Private/Rendering/StreamableTextureResource.cpp
@@ -23,20 +23,20 @@ static TAutoConsoleVariable<int32> CVarVirtualTextureEnabled(
 	TEXT("If set to 1, textures will use virtual memory so they can be partially resident."),
 	ECVF_RenderThreadSafe);
 
-bool CanCreateWithPartiallyResidentMips(uint32 TexCreateFlags)
+bool CanCreateWithPartiallyResidentMips(ETextureCreateFlags TexCreateFlags)
 {
 #if PLATFORM_SUPPORTS_VIRTUAL_TEXTURES
-	const uint32 iDisableFlags = 
+	const ETextureCreateFlags iDisableFlags =
 		TexCreate_RenderTargetable |
 		TexCreate_ResolveTargetable |
 		TexCreate_DepthStencilTargetable |
 		TexCreate_Dynamic |
 		TexCreate_UAV |
 		TexCreate_Presentable;
-	const uint32 iRequiredFlags =
+	const ETextureCreateFlags iRequiredFlags =
 		TexCreate_OfflineProcessed;
 
-	return ((TexCreateFlags & (iDisableFlags | iRequiredFlags)) == iRequiredFlags) && CVarVirtualTextureEnabled.GetValueOnAnyThread();
+	return EnumHasAllFlags(TexCreateFlags & iDisableFlags, iRequiredFlags) && CVarVirtualTextureEnabled.GetValueOnAnyThread();
 	
 #else
 	return false;
diff --git a/Engine/Source/Runtime/Engine/Private/TextureRenderTarget2D.cpp b/Engine/Source/Runtime/Engine/Private/TextureRenderTarget2D.cpp
--- a/Engine/Source/Runtime/Engine/Private/TextureRenderTarget2D.cpp
+++ b/Engine/Source/Runtime/Engine/Private/TextureRenderTarget2D.cpp
@@ -523,7 +523,7 @@ void FTextureRenderTarget2DResource::InitDynamicRHI()
 			Texture2DRHI
 			);
 
-		if ((TexCreateFlags & TexCreate_UAV) != 0)
+		if (EnumHasAnyFlags(TexCreateFlags, TexCreate_UAV))
 		{
 			UnorderedAccessViewRHI = RHICreateUnorderedAccessView(RenderTargetTextureRHI);
 		}
diff --git a/Engine/Source/Runtime/Engine/Private/TextureRenderTarget2DArray.cpp b/Engine/Source/Runtime/Engine/Private/TextureRenderTarget2DArray.cpp
--- a/Engine/Source/Runtime/Engine/Private/TextureRenderTarget2DArray.cpp
+++ b/Engine/Source/Runtime/Engine/Private/TextureRenderTarget2DArray.cpp
@@ -255,7 +255,7 @@ void FTextureRenderTarget2DArrayResource::InitDynamicRHI()
 			);
 		}
 
-		if ((TexCreateFlags & TexCreate_UAV) != 0)
+		if (EnumHasAnyFlags(TexCreateFlags, TexCreate_UAV))
 		{
 			UnorderedAccessViewRHI = RHICreateUnorderedAccessView(RenderTarget2DArrayRHI);
 		}
diff --git a/Engine/Source/Runtime/Engine/Private/TextureRenderTargetCube.cpp b/Engine/Source/Runtime/Engine/Private/TextureRenderTargetCube.cpp
--- a/Engine/Source/Runtime/Engine/Private/TextureRenderTargetCube.cpp
+++ b/Engine/Source/Runtime/Engine/Private/TextureRenderTargetCube.cpp
@@ -253,7 +253,7 @@ void FTextureRenderTargetCubeResource::InitDynamicRHI()
 				TextureCubeRHI );
 		}
 
-		if ((TexCreateFlags & TexCreate_UAV) != 0)
+		if (EnumHasAnyFlags(TexCreateFlags, TexCreate_UAV))
 		{
 			UnorderedAccessViewRHI = RHICreateUnorderedAccessView(RenderTargetCubeRHI);
 		}
diff --git a/Engine/Source/Runtime/Engine/Private/TextureRenderTargetVolume.cpp b/Engine/Source/Runtime/Engine/Private/TextureRenderTargetVolume.cpp
--- a/Engine/Source/Runtime/Engine/Private/TextureRenderTargetVolume.cpp
+++ b/Engine/Source/Runtime/Engine/Private/TextureRenderTargetVolume.cpp
@@ -262,7 +262,7 @@ void FTextureRenderTargetVolumeResource::InitDynamicRHI()
 			);
 		}
 
-		if ((TexCreateFlags & TexCreate_UAV) != 0)
+		if (EnumHasAnyFlags(TexCreateFlags, TexCreate_UAV))
 		{
 			UnorderedAccessViewRHI = RHICreateUnorderedAccessView(RenderTargetVolumeRHI);
 		}
diff --git a/Engine/Source/Runtime/Engine/Public/Rendering/StreamableTextureResource.h b/Engine/Source/Runtime/Engine/Public/Rendering/StreamableTextureResource.h
--- a/Engine/Source/Runtime/Engine/Public/Rendering/StreamableTextureResource.h
+++ b/Engine/Source/Runtime/Engine/Public/Rendering/StreamableTextureResource.h
@@ -11,7 +11,7 @@
 #include "Streaming/StreamableRenderResourceState.h"
 #include "Engine/Texture.h"
 
-bool CanCreateWithPartiallyResidentMips(uint32 TexCreateFlags);
+bool CanCreateWithPartiallyResidentMips(ETextureCreateFlags TexCreateFlags);
 
 /** 
  * The rendering resource streamable texture.
diff --git a/Engine/Source/Runtime/OpenGLDrv/Private/Linux/OpenGLLinux.h b/Engine/Source/Runtime/OpenGLDrv/Private/Linux/OpenGLLinux.h
--- a/Engine/Source/Runtime/OpenGLDrv/Private/Linux/OpenGLLinux.h
+++ b/Engine/Source/Runtime/OpenGLDrv/Private/Linux/OpenGLLinux.h
@@ -414,7 +414,7 @@ struct FLinuxOpenGL : public FOpenGL4
 		}
     }
 
-	static FORCEINLINE bool TexStorage2D(GLenum Target, GLint Levels, GLint InternalFormat, GLsizei Width, GLsizei Height, GLenum Format, GLenum Type, uint32 Flags )
+	static FORCEINLINE bool TexStorage2D(GLenum Target, GLint Levels, GLint InternalFormat, GLsizei Width, GLsizei Height, GLenum Format, GLenum Type, ETextureCreateFlags Flags )
 	{
 		if( glTexStorage2D != NULL )
 		{
diff --git a/Engine/Source/Runtime/OpenGLDrv/Private/Lumin/LuminGL4.h b/Engine/Source/Runtime/OpenGLDrv/Private/Lumin/LuminGL4.h
--- a/Engine/Source/Runtime/OpenGLDrv/Private/Lumin/LuminGL4.h
+++ b/Engine/Source/Runtime/OpenGLDrv/Private/Lumin/LuminGL4.h
@@ -466,7 +466,7 @@ struct FLuminOpenGL4 : public FOpenGL4
 		}
 	}
 
-	static FORCEINLINE bool TexStorage2D(GLenum Target, GLint Levels, GLint InternalFormat, GLsizei Width, GLsizei Height, GLenum Format, GLenum Type, uint32 Flags)
+	static FORCEINLINE bool TexStorage2D(GLenum Target, GLint Levels, GLint InternalFormat, GLsizei Width, GLsizei Height, GLenum Format, GLenum Type, ETextureCreateFlags Flags)
 	{
 		if (glTexStorage2D != NULL)
 		{
diff --git a/Engine/Source/Runtime/OpenGLDrv/Private/OpenGLRenderTarget.cpp b/Engine/Source/Runtime/OpenGLDrv/Private/OpenGLRenderTarget.cpp
--- a/Engine/Source/Runtime/OpenGLDrv/Private/OpenGLRenderTarget.cpp
+++ b/Engine/Source/Runtime/OpenGLDrv/Private/OpenGLRenderTarget.cpp
@@ -547,7 +547,7 @@ void FOpenGLDynamicRHI::RHICopyToResolveTarget(FRHITexture* SourceTextureRHI, FR
 	
 		// For CPU readback resolve targets we should issue the resolve to the internal PBO immediately.
 		// This makes any subsequent locking of that texture much cheaper as it won't have to stall on a pixel pack op.
-		bool bLockableTarget = DestTextureRHI->GetTexture2D() && (DestTextureRHI->GetFlags() & TexCreate_CPUReadback) && !(DestTextureRHI->GetFlags() & (TexCreate_RenderTargetable|TexCreate_DepthStencilTargetable)) && !DestTextureRHI->IsMultisampled();
+		bool bLockableTarget = DestTextureRHI->GetTexture2D() && EnumHasAnyFlags(DestTextureRHI->GetFlags(), TexCreate_CPUReadback) && !EnumHasAnyFlags(DestTextureRHI->GetFlags(), TexCreate_RenderTargetable|TexCreate_DepthStencilTargetable) && !DestTextureRHI->IsMultisampled();
 		if(bLockableTarget && !ResolveParams.Rect.IsValid())
 		{
 			FOpenGLTexture2D* DestTex = (FOpenGLTexture2D*)DestTexture;
diff --git a/Engine/Source/Runtime/OpenGLDrv/Private/OpenGLTexture.cpp b/Engine/Source/Runtime/OpenGLDrv/Private/OpenGLTexture.cpp
--- a/Engine/Source/Runtime/OpenGLDrv/Private/OpenGLTexture.cpp
+++ b/Engine/Source/Runtime/OpenGLDrv/Private/OpenGLTexture.cpp
@@ -28,7 +28,7 @@ static TAutoConsoleVariable<int32> CVarDeferTextureCreation(
 
 static TAutoConsoleVariable<int32> CVarDeferTextureCreationExcludeMask(
 	TEXT("r.OpenGL.DeferTextureCreationExcludeFlags"),
-	~(TexCreate_ShaderResource | TexCreate_SRGB | TexCreate_Transient | TexCreate_Streamable | TexCreate_OfflineProcessed),
+	static_cast<int32>(~(TexCreate_ShaderResource | TexCreate_SRGB | TexCreate_Transient | TexCreate_Streamable | TexCreate_OfflineProcessed)),
 	TEXT("Deferred texture creation exclusion mask, any texture requested with flags in this mask will be excluded from deferred creation."),
 	ECVF_RenderThreadSafe);
 
@@ -86,7 +86,7 @@ int64 GOpenGLTotalGraphicsMemory = 0;
 
 static bool ShouldCountAsTextureMemory(ETextureCreateFlags Flags)
 {
-	return (Flags & (TexCreate_RenderTargetable | TexCreate_ResolveTargetable | TexCreate_DepthStencilTargetable)) == 0;
+	return EnumHasAnyFlags(Flags, TexCreate_RenderTargetable | TexCreate_ResolveTargetable | TexCreate_DepthStencilTargetable);
 }
 
 void OpenGLTextureAllocated(FRHITexture* Texture, ETextureCreateFlags Flags)
@@ -431,7 +431,7 @@ void FOpenGLDynamicRHI::InitializeGLTexture(FRHITexture* Texture, uint32 SizeX,
 	VERIFY_GL_SCOPE();
 
 	const uint32 NumSamplesTileMem = bCubeTexture ? 1 : ((FOpenGLTexture2D*)Texture)->GetNumSamplesTileMem();
-	const bool TileMemDepth = NumSamplesTileMem > 1 && (Flags & TexCreate_DepthStencilTargetable);
+	const bool TileMemDepth = NumSamplesTileMem > 1 && EnumHasAnyFlags(Flags, TexCreate_DepthStencilTargetable);
 
 	GLuint TextureID = 0;
 	if (!TileMemDepth)
@@ -492,11 +492,11 @@ void FOpenGLDynamicRHI::InitializeGLTextureInternal(GLuint TextureID, FRHITextur
 
 	GLenum Target = bCubeTexture ? ((FOpenGLTextureCube*)Texture)->Target : ((FOpenGLTexture2D*)Texture)->Target;
 	const uint32 NumSamplesTileMem = bCubeTexture ? 1 : ((FOpenGLTexture2D*)Texture)->GetNumSamplesTileMem();
-	const bool TileMemDepth = NumSamplesTileMem > 1 && (Flags & TexCreate_DepthStencilTargetable);
+	const bool TileMemDepth = NumSamplesTileMem > 1 && EnumHasAnyFlags(Flags, TexCreate_DepthStencilTargetable);
 
 	check(TextureID || TileMemDepth);
 
-	const bool bSRGB = (Flags & TexCreate_SRGB) != 0;
+	const bool bSRGB = EnumHasAnyFlags(Flags, TexCreate_SRGB);
 	const FOpenGLTextureFormat& GLFormat = GOpenGLTextureFormats[Format];
 	if (GLFormat.InternalFormat[bSRGB] == GL_NONE)
 	{
@@ -550,7 +550,7 @@ void FOpenGLDynamicRHI::InitializeGLTextureInternal(GLuint TextureID, FRHITextur
 		
 		TextureMipLimits.Add(TextureID, TPair<GLenum, GLenum>(0, NumMips - 1));
 		
-		if (GLFormat.bBGRA && !(Flags & TexCreate_RenderTargetable))
+		if (GLFormat.bBGRA && !EnumHasAnyFlags(Flags, TexCreate_RenderTargetable))
 		{
 			glTexParameteri(Target, GL_TEXTURE_SWIZZLE_R, GL_BLUE);
 			glTexParameteri(Target, GL_TEXTURE_SWIZZLE_B, GL_RED);
@@ -708,15 +708,15 @@ void FOpenGLDynamicRHI::InitializeGLTextureInternal(GLuint TextureID, FRHITextur
 	
 	// Determine the attachment point for the texture.	
 	GLenum Attachment = GL_NONE;
-	if((Flags & TexCreate_RenderTargetable) || (Flags & TexCreate_CPUReadback))
+	if(EnumHasAnyFlags(Flags, TexCreate_RenderTargetable | TexCreate_CPUReadback))
 	{
 		Attachment = GL_COLOR_ATTACHMENT0;
 	}
-	else if(Flags & TexCreate_DepthStencilTargetable)
+	else if(EnumHasAnyFlags(Flags, TexCreate_DepthStencilTargetable))
 	{
 		Attachment = (Format == PF_DepthStencil) ? GL_DEPTH_STENCIL_ATTACHMENT : GL_DEPTH_ATTACHMENT;
 	}
-	else if(Flags & TexCreate_ResolveTargetable)
+	else if(EnumHasAnyFlags(Flags, TexCreate_ResolveTargetable))
 	{
 		Attachment = (Format == PF_DepthStencil)
 						? GL_DEPTH_STENCIL_ATTACHMENT
@@ -803,7 +803,7 @@ void TOpenGLTexture<RHIResourceType>::Resolve(uint32 MipIndex,uint32 ArrayIndex)
 	// This may be further optimized by caching information if surface content was changed since last lock.
 	
 	const FOpenGLTextureFormat& GLFormat = GOpenGLTextureFormats[PixelFormat];
-	const bool bSRGB = (this->GetFlags() & TexCreate_SRGB) != 0;
+	const bool bSRGB = EnumHasAnyFlags(this->GetFlags(), TexCreate_SRGB);
 	
 	// Use a texture stage that's not likely to be used for draws, to avoid waiting
 	FOpenGLContextState& ContextState = OpenGLRHI->GetContextStateForCurrentContext();
@@ -914,7 +914,7 @@ void* TOpenGLTexture<RHIResourceType>::Lock(uint32 InMipIndex,uint32 ArrayIndex,
 		check(!PixelBuffer->IsLocked());
 		
 		// If the buffer already exists & the flags are such that the texture cannot be rendered to & is CPU accessible then we can skip the internal resolve for read locks. This makes HZB occlusion faster.
-		const bool bCPUTexResolved = bBufferExists && (this->GetFlags() & TexCreate_CPUReadback) && !(this->GetFlags() & (TexCreate_RenderTargetable|TexCreate_DepthStencilTargetable));
+		const bool bCPUTexResolved = bBufferExists && EnumHasAnyFlags(this->GetFlags(), TexCreate_CPUReadback) && !EnumHasAnyFlags(this->GetFlags(), TexCreate_RenderTargetable|TexCreate_DepthStencilTargetable);
 
 		if (LockMode != RLM_WriteOnly && !bCPUTexResolved)
 		{
@@ -995,7 +995,7 @@ void TOpenGLTexture<RHIResourceType>::Unlock(uint32 MipIndex,uint32 ArrayIndex)
 
 	const int32 BufferIndex = MipIndex * (bCubemap ? 6 : 1) * this->GetEffectiveSizeZ() + ArrayIndex;
 	const FOpenGLTextureFormat& GLFormat = GOpenGLTextureFormats[this->GetFormat()];
-	const bool bSRGB = (this->GetFlags() & TexCreate_SRGB) != 0;
+	const bool bSRGB = EnumHasAnyFlags(this->GetFlags(), TexCreate_SRGB);
 	TRefCountPtr<FOpenGLPixelBuffer> PixelBuffer = PixelBuffers[BufferIndex];
 
 	check(IsValidRef(PixelBuffer));
@@ -1202,7 +1202,7 @@ void TOpenGLTexture<RHIResourceType>::RestoreEvictedGLResource(bool bAttemptToRe
 
 	EPixelFormat PixelFormat = this->GetFormat();
 	const FOpenGLTextureFormat& GLFormat = GOpenGLTextureFormats[PixelFormat];
-	const bool bSRGB = (this->GetFlags() & TexCreate_SRGB) != 0;
+	const bool bSRGB = EnumHasAnyFlags(this->GetFlags(), TexCreate_SRGB);
 	checkf(EvictionParamsPtr->MipImageData.Num() == this->GetNumMips(), TEXT("EvictionParamsPtr->MipImageData.Num() =%d, this->GetNumMips() = %d"), EvictionParamsPtr->MipImageData.Num(), this->GetNumMips());
 	
 	for (int i = EvictionParamsPtr->MipImageData.Num() - 1; i >= 0; i--)
@@ -1222,7 +1222,7 @@ void TOpenGLTexture<RHIResourceType>::RestoreEvictedGLResource(bool bAttemptToRe
 	// Use the resident streaming mips if our cvar is -1.
 	uint32 DeferTextureCreationKeepLowerMipCount = (uint32)(GOGLDeferTextureCreationKeepLowerMipCount >= 0 ? GOGLDeferTextureCreationKeepLowerMipCount : UTexture::GetStaticMinTextureResidentMipCount());
 
-	uint32 RetainMips = bAttemptToRetainMips && (this->GetFlags() & TexCreate_Streamable) && this->GetNumMips() > 1 && !this->IsAliased() ? DeferTextureCreationKeepLowerMipCount : 0;
+	uint32 RetainMips = bAttemptToRetainMips && EnumHasAnyFlags(this->GetFlags(), TexCreate_Streamable) && this->GetNumMips() > 1 && !this->IsAliased() ? DeferTextureCreationKeepLowerMipCount : 0;
 
 	if (CanBeEvicted())
 	{
@@ -1312,14 +1312,15 @@ bool TOpenGLTexture<RHIResourceType>::CanCreateAsEvicted()
 	bool bRet =
 		CanDeferTextureCreation()
 		&& FOpenGL::SupportsCopyImage()
-		&& this->GetFlags() // ignore TexCreate_None
-		&& (CVarDeferTextureCreationExcludeMask.GetValueOnAnyThread() & this->GetFlags())==0  // Anything outside of these flags cannot be evicted.
+		&& this->GetFlags() != TexCreate_None // ignore TexCreate_None
+		&& !EnumHasAnyFlags((ETextureCreateFlags)CVarDeferTextureCreationExcludeMask.GetValueOnAnyThread(), this->GetFlags())  // Anything outside of these flags cannot be evicted.
 		&& Target == GL_TEXTURE_2D
 		&& this->GetTexture2D(); // 2d only.
 
 	if (GOGLTextureEvictLogging)
 	{
-		UE_CLOG(!bRet, LogRHI, Warning, TEXT("CanDeferTextureCreation:%d, SupportsCopyImage:%d, Flags:%x Flags&Mask:%d, Target:%x"), CanDeferTextureCreation(), FOpenGL::SupportsCopyImage(), this->GetFlags(), (CVarDeferTextureCreationExcludeMask.GetValueOnAnyThread() & this->GetFlags()), Target);
+		UE_CLOG(!bRet, LogRHI, Warning, TEXT("CanDeferTextureCreation:%d, SupportsCopyImage:%d, Flags:%llx Mask:%x, Target:%x"),
+			CanDeferTextureCreation(), FOpenGL::SupportsCopyImage(), this->GetFlags(), CVarDeferTextureCreationExcludeMask.GetValueOnAnyThread(), Target);
 	}
 
 	return bRet;
@@ -1404,8 +1405,8 @@ void TOpenGLTexture<RHIResourceType>::CloneViaPBO( TOpenGLTexture<RHIResourceTyp
 	check(PixelFormat == Src->GetFormat());
 
 	const FOpenGLTextureFormat& GLFormat = GOpenGLTextureFormats[PixelFormat];
-	const bool bSRGB = (this->GetFlags() & TexCreate_SRGB) != 0;
-	check(bSRGB == ((Src->GetFlags() & TexCreate_SRGB) != 0));
+	const bool bSRGB = EnumHasAnyFlags(this->GetFlags(), TexCreate_SRGB);
+	check(bSRGB == EnumHasAnyFlags(Src->GetFlags(), TexCreate_SRGB));
 	
 	const uint32 BlockSizeX = GPixelFormats[PixelFormat].BlockSizeX;
 	const uint32 BlockSizeY = GPixelFormats[PixelFormat].BlockSizeY;
@@ -1652,7 +1653,7 @@ FTexture2DArrayRHIRef FOpenGLDynamicRHI::RHICreateTexture2DArray(uint32 SizeX,ui
 	
 	TextureMipLimits.Add(TextureID, TPair<GLenum, GLenum>(0, NumMips - 1));
 
-	const bool bSRGB = (Flags&TexCreate_SRGB) != 0;
+	const bool bSRGB = EnumHasAnyFlags(Flags, TexCreate_SRGB);
 	const FOpenGLTextureFormat& GLFormat = GOpenGLTextureFormats[Format];
 	const FPixelFormatInfo& FormatInfo = GPixelFormats[Format];
 	if (GLFormat.InternalFormat[bSRGB] == GL_NONE)
@@ -1660,7 +1661,7 @@ FTexture2DArrayRHIRef FOpenGLDynamicRHI::RHICreateTexture2DArray(uint32 SizeX,ui
 		UE_LOG(LogRHI, Fatal,TEXT("Texture format '%s' not supported."), FormatInfo.Name);
 	}
 
-	if (GLFormat.bBGRA && !(Flags & TexCreate_RenderTargetable))
+	if (GLFormat.bBGRA && !EnumHasAnyFlags(Flags, TexCreate_RenderTargetable))
 	{
 		glTexParameteri(Target, GL_TEXTURE_SWIZZLE_R, GL_BLUE);
 		glTexParameteri(Target, GL_TEXTURE_SWIZZLE_B, GL_RED);
@@ -1725,15 +1726,15 @@ FTexture2DArrayRHIRef FOpenGLDynamicRHI::RHICreateTexture2DArray(uint32 SizeX,ui
 	
 	// Determine the attachment point for the texture.	
 	GLenum Attachment = GL_NONE;
-	if(Flags & TexCreate_RenderTargetable)
+	if(EnumHasAnyFlags(Flags, TexCreate_RenderTargetable))
 	{
 		Attachment = GL_COLOR_ATTACHMENT0;
 	}
-	else if(Flags & TexCreate_DepthStencilTargetable)
+	else if(EnumHasAnyFlags(Flags, TexCreate_DepthStencilTargetable))
 	{
 		Attachment = (Format == PF_DepthStencil) ? GL_DEPTH_STENCIL_ATTACHMENT : GL_DEPTH_ATTACHMENT;
 	}
-	else if(Flags & TexCreate_ResolveTargetable)
+	else if(EnumHasAnyFlags(Flags, TexCreate_ResolveTargetable))
 	{
 		Attachment = (Format == PF_DepthStencil)
 			? GL_DEPTH_STENCIL_ATTACHMENT
@@ -1788,7 +1789,7 @@ FTexture3DRHIRef FOpenGLDynamicRHI::RHICreateTexture3D(uint32 SizeX,uint32 SizeY
 	
 	TextureMipLimits.Add(TextureID, TPair<GLenum, GLenum>(0, NumMips - 1));
 
-	const bool bSRGB = (Flags&TexCreate_SRGB) != 0;
+	const bool bSRGB = EnumHasAnyFlags(Flags, TexCreate_SRGB);
 	const FOpenGLTextureFormat& GLFormat = GOpenGLTextureFormats[Format];
 	const FPixelFormatInfo& FormatInfo = GPixelFormats[Format];
 	if (GLFormat.InternalFormat[bSRGB] == GL_NONE)
@@ -1796,7 +1797,7 @@ FTexture3DRHIRef FOpenGLDynamicRHI::RHICreateTexture3D(uint32 SizeX,uint32 SizeY
 		UE_LOG(LogRHI, Fatal, TEXT("Texture format '%s' not supported."), FormatInfo.Name);
 	}
 
-	if (GLFormat.bBGRA && !(Flags & TexCreate_RenderTargetable))
+	if (GLFormat.bBGRA && !EnumHasAnyFlags(Flags, TexCreate_RenderTargetable))
 	{
 		glTexParameteri(Target, GL_TEXTURE_SWIZZLE_R, GL_BLUE);
 		glTexParameteri(Target, GL_TEXTURE_SWIZZLE_B, GL_RED);
@@ -1865,15 +1866,15 @@ FTexture3DRHIRef FOpenGLDynamicRHI::RHICreateTexture3D(uint32 SizeX,uint32 SizeY
 	
 	// Determine the attachment point for the texture.	
 	GLenum Attachment = GL_NONE;
-	if(Flags & TexCreate_RenderTargetable)
+	if(EnumHasAnyFlags(Flags, TexCreate_RenderTargetable))
 	{
 		Attachment = GL_COLOR_ATTACHMENT0;
 	}
-	else if(Flags & TexCreate_DepthStencilTargetable)
+	else if(EnumHasAnyFlags(Flags, TexCreate_DepthStencilTargetable))
 	{
 		Attachment = Format == PF_DepthStencil ? GL_DEPTH_STENCIL_ATTACHMENT : GL_DEPTH_ATTACHMENT;
 	}
-	else if(Flags & TexCreate_ResolveTargetable)
+	else if(EnumHasAnyFlags(Flags, TexCreate_ResolveTargetable))
 	{
 		Attachment = (Format == PF_DepthStencil)
 			? GL_DEPTH_STENCIL_ATTACHMENT
@@ -1922,7 +1923,7 @@ FShaderResourceViewRHIRef FOpenGLDynamicRHI::RHICreateShaderResourceView(FRHITex
 				{
 					// Choose original format when PF_Unknown is specified (as stated for FRHITextureSRVCreateInfo::Format)
 					const FOpenGLTextureFormat& GLFormat = GOpenGLTextureFormats[Format == PF_Unknown ? Texture2D->GetFormat() : Format];
-					const bool bSRGB = (Texture2D->GetFlags()&TexCreate_SRGB) != 0;
+					const bool bSRGB = EnumHasAnyFlags(Texture2D->GetFlags(), TexCreate_SRGB);
 
 					FOpenGL::TextureView(Resource, Texture2D->Target, Texture2D->GetResource(), GLFormat.InternalFormat[bSRGB], MipLevel, NumMipLevels, 0, 1);
 				}
@@ -1972,7 +1973,7 @@ FShaderResourceViewRHIRef FOpenGLDynamicRHI::RHICreateShaderResourceView(FRHITex
 						uint32 const SizeX = Texture2D->GetSizeX();
 						uint32 const SizeY = Texture2D->GetSizeY();
 						GLenum const Type = GL_UNSIGNED_BYTE;
-						uint32 const Flags = 0;
+						ETextureCreateFlags const Flags = TexCreate_None;
 
 						FOpenGLContextState& ContextState = GetContextStateForCurrentContext();
 						CachedSetupTextureStage(ContextState, FOpenGL::GetMaxCombinedTextureImageUnits() - 1, Target, Texture2D->SRVResource, MipLevel, NumMipLevels);
@@ -2029,7 +2030,7 @@ FShaderResourceViewRHIRef FOpenGLDynamicRHI::RHICreateShaderResourceView(FRHITex
 
 				FOpenGL::GenTextures(1, &Resource);
 				const FOpenGLTextureFormat& GLFormat = GOpenGLTextureFormats[Texture2DArray->GetFormat()];
-				const bool bSRGB = (Texture2DArray->GetFlags()&TexCreate_SRGB) != 0;
+				const bool bSRGB = EnumHasAnyFlags(Texture2DArray->GetFlags(), TexCreate_SRGB);
 
 				FOpenGL::TextureView(Resource, Texture2DArray->Target, Texture2DArray->GetResource(), GLFormat.InternalFormat[bSRGB], MipLevel, 1, 0, 1);
 
@@ -2051,7 +2052,7 @@ FShaderResourceViewRHIRef FOpenGLDynamicRHI::RHICreateShaderResourceView(FRHITex
 
 				FOpenGL::GenTextures(1, &Resource);
 				const FOpenGLTextureFormat& GLFormat = GOpenGLTextureFormats[TextureCube->GetFormat()];
-				const bool bSRGB = (TextureCube->GetFlags()&TexCreate_SRGB) != 0;
+				const bool bSRGB = EnumHasAnyFlags(TextureCube->GetFlags(), TexCreate_SRGB);
 
 				FOpenGL::TextureView(Resource, TextureCube->Target, TextureCube->GetResource(), GLFormat.InternalFormat[bSRGB], MipLevel, 1, 0, 6);
 
@@ -2076,7 +2077,7 @@ FShaderResourceViewRHIRef FOpenGLDynamicRHI::RHICreateShaderResourceView(FRHITex
 
 				FOpenGL::GenTextures(1, &Resource);
 				const FOpenGLTextureFormat& GLFormat = GOpenGLTextureFormats[Texture3D->GetFormat()];
-				const bool bSRGB = (Texture3D->GetFlags()&TexCreate_SRGB) != 0;
+				const bool bSRGB = EnumHasAnyFlags(Texture3D->GetFlags(), TexCreate_SRGB);
 
 				FOpenGL::TextureView(Resource, Texture3D->Target, Texture3D->GetResource(), GLFormat.InternalFormat[bSRGB], MipLevel, 1, 0, 1);
 
@@ -2364,7 +2365,7 @@ void FOpenGLDynamicRHI::RHIUpdateTexture3D(FRHITexture3D* TextureRHI, uint32 Mip
 	
 	glPixelStorei(GL_UNPACK_ALIGNMENT, 1);
 
-	const bool bSRGB = (Texture->GetFlags() & TexCreate_SRGB) != 0;
+	const bool bSRGB = EnumHasAnyFlags(Texture->GetFlags(), TexCreate_SRGB);
 
 	if (GLFormat.bCompressed)
 	{
diff --git a/Engine/Source/Runtime/OpenGLDrv/Private/Windows/OpenGLWindows.h b/Engine/Source/Runtime/OpenGLDrv/Private/Windows/OpenGLWindows.h
--- a/Engine/Source/Runtime/OpenGLDrv/Private/Windows/OpenGLWindows.h
+++ b/Engine/Source/Runtime/OpenGLDrv/Private/Windows/OpenGLWindows.h
@@ -446,7 +446,7 @@ struct FWindowsOpenGL : public FOpenGL4
 		}
 	}
 
-	static FORCEINLINE bool TexStorage2D(GLenum Target, GLint Levels, GLint InternalFormat, GLsizei Width, GLsizei Height, GLenum Format, GLenum Type, uint32 Flags)
+	static FORCEINLINE bool TexStorage2D(GLenum Target, GLint Levels, GLint InternalFormat, GLsizei Width, GLsizei Height, GLenum Format, GLenum Type, ETextureCreateFlags Flags)
 	{
 		if( glTexStorage2D != NULL )
 		{
@@ -634,7 +634,7 @@ struct FWindowsOpenGL : public FOpenGLESDeferred
 		}
 	}
 
-	static FORCEINLINE bool TexStorage2D(GLenum Target, GLint Levels, GLint InternalFormat, GLsizei Width, GLsizei Height, GLenum Format, GLenum Type, uint32 Flags)
+	static FORCEINLINE bool TexStorage2D(GLenum Target, GLint Levels, GLint InternalFormat, GLsizei Width, GLsizei Height, GLenum Format, GLenum Type, ETextureCreateFlags Flags)
 	{
 		if( glTexStorage2D != NULL )
 		{
diff --git a/Engine/Source/Runtime/OpenGLDrv/Public/OpenGL.h b/Engine/Source/Runtime/OpenGLDrv/Public/OpenGL.h
--- a/Engine/Source/Runtime/OpenGLDrv/Public/OpenGL.h
+++ b/Engine/Source/Runtime/OpenGLDrv/Public/OpenGL.h
@@ -283,7 +283,7 @@ public:
 	static FORCEINLINE void DispatchCompute(GLuint NumGroupsX, GLuint NumGroupsY, GLuint NumGroupsZ) UGL_REQUIRED_VOID
 	static FORCEINLINE void DispatchComputeIndirect(GLintptr Offset) UGL_REQUIRED_VOID
 	static FORCEINLINE void MemoryBarrier(GLbitfield Barriers) UGL_REQUIRED_VOID
-	static FORCEINLINE bool TexStorage2D(GLenum Target, GLint Levels, GLint InternalFormat, GLsizei Width, GLsizei Height, GLenum Format, GLenum Type, uint32 Flags) UGL_OPTIONAL(false)
+	static FORCEINLINE bool TexStorage2D(GLenum Target, GLint Levels, GLint InternalFormat, GLsizei Width, GLsizei Height, GLenum Format, GLenum Type, ETextureCreateFlags Flags) UGL_OPTIONAL(false)
 	static FORCEINLINE bool TexStorage2DMultisample(GLenum Target, GLsizei Samples, GLint InternalFormat, GLsizei Width, GLsizei Height, GLboolean FixedSampleLocations) UGL_OPTIONAL(false)
 	static FORCEINLINE void TexStorage3D(GLenum Target, GLint Levels, GLint InternalFormat, GLsizei Width, GLsizei Height, GLsizei Depth, GLenum Format, GLenum Type) UGL_REQUIRED_VOID
 	static FORCEINLINE void CompressedTexSubImage3D(GLenum Target, GLint Level, GLint XOffset, GLint YOffset, GLint ZOffset, GLsizei Width, GLsizei Height, GLsizei Depth, GLenum Format, GLsizei ImageSize, const GLvoid* PixelData) UGL_REQUIRED_VOID
diff --git a/Engine/Source/Runtime/OpenGLDrv/Public/OpenGLES.h b/Engine/Source/Runtime/OpenGLDrv/Public/OpenGLES.h
--- a/Engine/Source/Runtime/OpenGLDrv/Public/OpenGLES.h
+++ b/Engine/Source/Runtime/OpenGLDrv/Public/OpenGLES.h
@@ -652,7 +652,7 @@ struct FOpenGLES : public FOpenGLBase
 		glBlitFramebuffer(SrcX0, SrcY0, SrcX1, SrcY1, DstX0, DstY0, DstX1, DstY1, Mask, Filter);
 	}
 
-	static FORCEINLINE bool TexStorage2D(GLenum Target, GLint Levels, GLint InternalFormat, GLsizei Width, GLsizei Height, GLenum Format, GLenum Type, uint32 Flags)
+	static FORCEINLINE bool TexStorage2D(GLenum Target, GLint Levels, GLint InternalFormat, GLsizei Width, GLsizei Height, GLenum Format, GLenum Type, ETextureCreateFlags Flags)
 	{
 		glTexStorage2D(Target, Levels, InternalFormat, Width, Height);
 		VERIFY_GL(glTexStorage2D);
diff --git a/Engine/Source/Runtime/OpenGLDrv/Public/OpenGLResources.h b/Engine/Source/Runtime/OpenGLDrv/Public/OpenGLResources.h
--- a/Engine/Source/Runtime/OpenGLDrv/Public/OpenGLResources.h
+++ b/Engine/Source/Runtime/OpenGLDrv/Public/OpenGLResources.h
@@ -1594,7 +1594,7 @@ public:
 	// Accessors.
 	bool IsDynamic() const { return (this->GetFlags() & TexCreate_Dynamic) != 0; }
 	bool IsCubemap() const { return bCubemap != 0; }
-	bool IsStaging() const { return (this->GetFlags() & TexCreate_CPUReadback) != 0; }
+	bool IsStaging() const { return EnumHasAnyFlags(this->GetFlags(), TexCreate_CPUReadback); }
 
 
 	/** FRHITexture override.  See FRHITexture::GetNativeResource() */
diff --git a/Engine/Source/Runtime/RHI/Private/PipelineFileCache.cpp b/Engine/Source/Runtime/RHI/Private/PipelineFileCache.cpp
--- a/Engine/Source/Runtime/RHI/Private/PipelineFileCache.cpp
+++ b/Engine/Source/Runtime/RHI/Private/PipelineFileCache.cpp
@@ -49,6 +49,12 @@ DECLARE_MEMORY_STAT(TEXT("New Cached PSO"), STAT_NewCachedPSOMemory, STATGROUP_P
 DECLARE_MEMORY_STAT(TEXT("PSO Stat"), STAT_PSOStatMemory, STATGROUP_PipelineStateCache);
 DECLARE_MEMORY_STAT(TEXT("File Cache"), STAT_FileCacheMemory, STATGROUP_PipelineStateCache);
 
+void LexFromString(ETextureCreateFlags& OutValue, const FStringView& InString)
+{
+	__underlying_type(ETextureCreateFlags) TmpFlags = static_cast<__underlying_type(ETextureCreateFlags)>(OutValue);
+	LexFromString(TmpFlags, InString);
+	OutValue = static_cast<ETextureCreateFlags>(TmpFlags);
+}
 
 enum class EPipelineCacheFileFormatVersions : uint32
 {
@@ -66,12 +72,13 @@ enum class EPipelineCacheFileFormatVersions : uint32
 	AlphaToCoverage = 19,
 	AddingMeshShaders = 20,
 	RemovingTessellationShaders = 21,
+	MoreRenderTargetFlags = 22,
 };
 
 const uint64 FPipelineCacheFileFormatMagic = 0x5049504543414348; // PIPECACH
 const uint64 FPipelineCacheTOCFileFormatMagic = 0x544F435354415232; // TOCSTAR2
 const uint64 FPipelineCacheEOFFileFormatMagic = 0x454F462D4D41524B; // EOF-MARK
-const uint32 FPipelineCacheFileFormatCurrentVersion = (uint32)EPipelineCacheFileFormatVersions::RemovingTessellationShaders;
+const uint32 FPipelineCacheFileFormatCurrentVersion = (uint32)EPipelineCacheFileFormatVersions::MoreRenderTargetFlags;
 const int32  FPipelineCacheGraphicsDescPartsNum = 63; // parser will expect this number of parts in a description string
 
 /**
@@ -373,7 +380,7 @@ FString FPipelineCacheFileFormatPSO::GraphicsDescriptor::StateToString() const
 		, *RasterizerState.ToString()
 		, *DepthStencilState.ToString()
 	);
-	Result += FString::Printf(TEXT("%d,%d,%d,")
+	Result += FString::Printf(TEXT("%d,%d,%lld,")
 		, MSAASamples
 		, uint32(DepthStencilFormat)
 		, DepthStencilFlags
@@ -391,7 +398,7 @@ FString FPipelineCacheFileFormatPSO::GraphicsDescriptor::StateToString() const
 	);
 	for (int32 Index = 0; Index < MaxSimultaneousRenderTargets; Index++)
 	{
-		Result += FString::Printf(TEXT("%d,%d,%d,%d,")
+		Result += FString::Printf(TEXT("%d,%lld,%d,%d,")
 			, uint32(RenderTargetFormats[Index])
 			, RenderTargetFlags[Index]
 			, 0/*Load*/
@@ -468,9 +475,7 @@ bool FPipelineCacheFileFormatPSO::GraphicsDescriptor::StateFromString(const FStr
 	{
 		check(PartEnd - PartIt >= 4 && sizeof(ERenderTargetLoadAction) == 1 && sizeof(ERenderTargetStoreAction) == 1 && sizeof(EPixelFormat) == sizeof(uint32)); //not a very robust parser
 		LexFromString((uint32&)(RenderTargetFormats[Index]), *PartIt++);
-		uint32 RTFlags;
-		LexFromString(RTFlags, *PartIt++);
-		RenderTargetFlags[Index] = (ETextureCreateFlags)RTFlags;
+		LexFromString(RenderTargetFlags[Index], *PartIt++);
 		uint8 Load, Store;
 		LexFromString(Load, *PartIt++);
 		LexFromString(Store, *PartIt++);
@@ -964,9 +969,20 @@ bool FPipelineCacheFileFormatPSO::Verify() const
 				uint32 Format = (uint32)Info.GraphicsDesc.RenderTargetFormats[i];
 				Ar << Format;
 				Info.GraphicsDesc.RenderTargetFormats[i] = (EPixelFormat)Format;
-				uint32 RTFlags = Info.GraphicsDesc.RenderTargetFlags[i];
-				Ar << RTFlags;
-				Info.GraphicsDesc.RenderTargetFlags[i] = (ETextureCreateFlags)RTFlags;
+
+				if (Ar.GameNetVer() < (uint32)EPipelineCacheFileFormatVersions::MoreRenderTargetFlags)
+				{
+					uint32 RTFlags = 0;
+					Ar << RTFlags;
+					Info.GraphicsDesc.RenderTargetFlags[i] = static_cast<ETextureCreateFlags>(RTFlags);
+				}
+				else
+				{
+					static_assert(sizeof(uint64) == sizeof(Info.GraphicsDesc.RenderTargetFlags[i]), "ETextureCreateFlags size changed, please change serialization");
+					uint64 RTFlags = static_cast<uint64>(Info.GraphicsDesc.RenderTargetFlags[i]);
+					Ar << RTFlags;
+					Info.GraphicsDesc.RenderTargetFlags[i] = static_cast<ETextureCreateFlags>(RTFlags);
+				}
 				uint8 LoadStore = 0;
 				Ar << LoadStore;
 				Ar << LoadStore;
@@ -979,7 +995,17 @@ bool FPipelineCacheFileFormatPSO::Verify() const
 			uint32 Format = (uint32)Info.GraphicsDesc.DepthStencilFormat;
 			Ar << Format;
 			Info.GraphicsDesc.DepthStencilFormat = (EPixelFormat)Format;
-			Ar << Info.GraphicsDesc.DepthStencilFlags;
+			if (Ar.GameNetVer() < (uint32)EPipelineCacheFileFormatVersions::MoreRenderTargetFlags)
+			{
+				uint32 DepthStencilFlags = 0;
+				Ar << DepthStencilFlags;
+				Info.GraphicsDesc.DepthStencilFlags = static_cast<ETextureCreateFlags>(DepthStencilFlags);
+			}
+			else
+			{
+				static_assert(sizeof(uint64) == sizeof(Info.GraphicsDesc.DepthStencilFlags), "ETextureCreateFlags size changed, please change serialization");
+				Ar << Info.GraphicsDesc.DepthStencilFlags;
+			}
 			Ar << Info.GraphicsDesc.DepthLoad;
 			Ar << Info.GraphicsDesc.StencilLoad;
 			Ar << Info.GraphicsDesc.DepthStore;
@@ -1147,7 +1173,7 @@ FPipelineCacheFileFormatPSO::FPipelineCacheFileFormatPSO()
 	for (uint32 i = 0; i < MaxSimultaneousRenderTargets; i++)
 	{
 		PSO.GraphicsDesc.RenderTargetFormats[i] = (EPixelFormat)Init.RenderTargetFormats[i];
-		PSO.GraphicsDesc.RenderTargetFlags[i] = (ETextureCreateFlags)Init.RenderTargetFlags[i];
+		PSO.GraphicsDesc.RenderTargetFlags[i] = Init.RenderTargetFlags[i];
 	}
 	
 	PSO.GraphicsDesc.RenderTargetsActive = Init.RenderTargetsEnabled;
diff --git a/Engine/Source/Runtime/RHI/Private/RHI.cpp b/Engine/Source/Runtime/RHI/Private/RHI.cpp
--- a/Engine/Source/Runtime/RHI/Private/RHI.cpp
+++ b/Engine/Source/Runtime/RHI/Private/RHI.cpp
@@ -1854,7 +1854,7 @@ void FRHIRenderPassInfo::Validate() const
 			// for depth read sub-pass
 			// 1. render pass must have depth target
 			// 2. depth target must support InputAttachement
-			ensure((DepthStencilRenderTarget.DepthStencilTarget->GetFlags() & TexCreate_InputAttachmentRead) != 0);
+			ensure(EnumHasAnyFlags(DepthStencilRenderTarget.DepthStencilTarget->GetFlags(), TexCreate_InputAttachmentRead));
 		}
 	}
 	else
diff --git a/Engine/Source/Runtime/RHI/Private/RHIUtilities.cpp b/Engine/Source/Runtime/RHI/Private/RHIUtilities.cpp
--- a/Engine/Source/Runtime/RHI/Private/RHIUtilities.cpp
+++ b/Engine/Source/Runtime/RHI/Private/RHIUtilities.cpp
@@ -535,23 +535,23 @@ RHI_API ERHIAccess RHIGetDefaultResourceState(ETextureCreateFlags InUsage, bool
 
 	if (!bInHasInitialData)
 	{
-		if (InUsage & TexCreate_RenderTargetable)
+		if (EnumHasAnyFlags(InUsage, TexCreate_RenderTargetable))
 		{
 			ResourceState = ERHIAccess::RTV;
 		}
-		else if (InUsage & TexCreate_DepthStencilTargetable)
+		else if (EnumHasAnyFlags(InUsage, TexCreate_DepthStencilTargetable))
 		{
 			ResourceState = ERHIAccess::DSVWrite | ERHIAccess::DSVRead;
 		}
-		else if (InUsage & TexCreate_UAV)
+		else if (EnumHasAnyFlags(InUsage, TexCreate_UAV))
 		{
 			ResourceState = ERHIAccess::UAVMask;
 		}
-		else if (InUsage & TexCreate_Presentable)
+		else if (EnumHasAnyFlags(InUsage, TexCreate_Presentable))
 		{
 			ResourceState = ERHIAccess::Present;
 		}
-		else if (InUsage & TexCreate_ShaderResource)
+		else if (EnumHasAnyFlags(InUsage, TexCreate_ShaderResource))
 		{
 			ResourceState = ERHIAccess::SRVMask;
 		}
diff --git a/Engine/Source/Runtime/RHI/Public/PipelineFileCache.h b/Engine/Source/Runtime/RHI/Public/PipelineFileCache.h
--- a/Engine/Source/Runtime/RHI/Public/PipelineFileCache.h
+++ b/Engine/Source/Runtime/RHI/Public/PipelineFileCache.h
@@ -168,7 +168,7 @@ struct RHI_API FPipelineCacheFileFormatPSO
 		uint32 MSAASamples;
 		
 		EPixelFormat DepthStencilFormat;
-		uint32 DepthStencilFlags;
+		ETextureCreateFlags DepthStencilFlags;
 		ERenderTargetLoadAction DepthLoad;
 		ERenderTargetLoadAction StencilLoad;
 		ERenderTargetStoreAction DepthStore;
diff --git a/Engine/Source/Runtime/RHI/Public/RHICommandList.h b/Engine/Source/Runtime/RHI/Public/RHICommandList.h
--- a/Engine/Source/Runtime/RHI/Public/RHICommandList.h
+++ b/Engine/Source/Runtime/RHI/Public/RHICommandList.h
@@ -5380,19 +5380,19 @@ FORCEINLINE void RHIUpdateTextureReference(FRHITextureReference* TextureRef, FRH
 
 FORCEINLINE FTexture2DRHIRef RHICreateTexture2D(uint32 SizeX, uint32 SizeY, uint8 Format, uint32 NumMips, uint32 NumSamples, ETextureCreateFlags Flags, ERHIAccess InResourceState, FRHIResourceCreateInfo& CreateInfo)
 {
-	LLM_SCOPE((Flags & (TexCreate_RenderTargetable | TexCreate_DepthStencilTargetable)) != 0 ? ELLMTag::RenderTargets : ELLMTag::Textures);
+	LLM_SCOPE(EnumHasAnyFlags(Flags, TexCreate_RenderTargetable | TexCreate_DepthStencilTargetable) ? ELLMTag::RenderTargets : ELLMTag::Textures);
 	return GDynamicRHI->RHICreateTexture2D_RenderThread(FRHICommandListExecutor::GetImmediateCommandList(), SizeX, SizeY, Format, NumMips, NumSamples, Flags, InResourceState, CreateInfo);
 }
 
 FORCEINLINE FTexture2DRHIRef RHICreateTextureExternal2D(uint32 SizeX, uint32 SizeY, uint8 Format, uint32 NumMips, uint32 NumSamples, ETextureCreateFlags Flags, ERHIAccess InResourceState, FRHIResourceCreateInfo& CreateInfo)
 {
-	LLM_SCOPE((Flags & (TexCreate_RenderTargetable | TexCreate_DepthStencilTargetable)) != 0 ? ELLMTag::RenderTargets : ELLMTag::Textures);
+	LLM_SCOPE(EnumHasAnyFlags(Flags, TexCreate_RenderTargetable | TexCreate_DepthStencilTargetable) ? ELLMTag::RenderTargets : ELLMTag::Textures);
 	return GDynamicRHI->RHICreateTextureExternal2D_RenderThread(FRHICommandListExecutor::GetImmediateCommandList(), SizeX, SizeY, Format, NumMips, NumSamples, Flags, InResourceState, CreateInfo);
 }
 
 FORCEINLINE FTexture2DRHIRef RHIAsyncCreateTexture2D(uint32 SizeX, uint32 SizeY, uint8 Format, uint32 NumMips, ETextureCreateFlags Flags, ERHIAccess InResourceState, void** InitialMipData, uint32 NumInitialMips)
 {
-	LLM_SCOPE((Flags & (TexCreate_RenderTargetable | TexCreate_DepthStencilTargetable)) != 0 ? ELLMTag::RenderTargets : ELLMTag::Textures);
+	LLM_SCOPE(EnumHasAnyFlags(Flags, TexCreate_RenderTargetable | TexCreate_DepthStencilTargetable) ? ELLMTag::RenderTargets : ELLMTag::Textures);
 	return GDynamicRHI->RHIAsyncCreateTexture2D(SizeX, SizeY, Format, NumMips, Flags, InResourceState, InitialMipData, NumInitialMips);
 }
 
@@ -5424,7 +5424,7 @@ FORCEINLINE void RHICopySharedMips(FRHITexture2D* DestTexture2D, FRHITexture2D*
 
 FORCEINLINE FTexture2DArrayRHIRef RHICreateTexture2DArray(uint32 SizeX, uint32 SizeY, uint32 SizeZ, uint8 Format, uint32 NumMips, uint32 NumSamples, ETextureCreateFlags Flags, ERHIAccess InResourceState, FRHIResourceCreateInfo& CreateInfo)
 {
-	LLM_SCOPE((Flags & (TexCreate_RenderTargetable | TexCreate_DepthStencilTargetable)) != 0 ? ELLMTag::RenderTargets : ELLMTag::Textures);
+	LLM_SCOPE(EnumHasAnyFlags(Flags, TexCreate_RenderTargetable | TexCreate_DepthStencilTargetable) ? ELLMTag::RenderTargets : ELLMTag::Textures);
 	return GDynamicRHI->RHICreateTexture2DArray_RenderThread(FRHICommandListExecutor::GetImmediateCommandList(), SizeX, SizeY, SizeZ, Format, NumMips, NumSamples, Flags, InResourceState, CreateInfo);
 }
 
@@ -5437,7 +5437,7 @@ FORCEINLINE FTexture2DArrayRHIRef RHICreateTexture2DArray(uint32 SizeX, uint32 S
 
 FORCEINLINE FTexture3DRHIRef RHICreateTexture3D(uint32 SizeX, uint32 SizeY, uint32 SizeZ, uint8 Format, uint32 NumMips, ETextureCreateFlags Flags, ERHIAccess ResourceState, FRHIResourceCreateInfo& CreateInfo)
 {
-	LLM_SCOPE((Flags & (TexCreate_RenderTargetable | TexCreate_DepthStencilTargetable)) != 0 ? ELLMTag::RenderTargets : ELLMTag::Textures);
+	LLM_SCOPE(EnumHasAnyFlags(Flags, TexCreate_RenderTargetable | TexCreate_DepthStencilTargetable) ? ELLMTag::RenderTargets : ELLMTag::Textures);
 	return GDynamicRHI->RHICreateTexture3D_RenderThread(FRHICommandListExecutor::GetImmediateCommandList(), SizeX, SizeY, SizeZ, Format, NumMips, Flags, ResourceState, CreateInfo);
 }
 
@@ -5535,13 +5535,13 @@ FORCEINLINE void RHIUpdateTexture3D(FRHITexture3D* Texture, uint32 MipIndex, con
 
 FORCEINLINE FTextureCubeRHIRef RHICreateTextureCube(uint32 Size, uint8 Format, uint32 NumMips, ETextureCreateFlags Flags, ERHIAccess InResourceState, FRHIResourceCreateInfo& CreateInfo)
 {
-	LLM_SCOPE((Flags & (TexCreate_RenderTargetable | TexCreate_DepthStencilTargetable)) != 0 ? ELLMTag::RenderTargets : ELLMTag::Textures);
+	LLM_SCOPE(EnumHasAnyFlags(Flags, TexCreate_RenderTargetable | TexCreate_DepthStencilTargetable) ? ELLMTag::RenderTargets : ELLMTag::Textures);
 	return GDynamicRHI->RHICreateTextureCube_RenderThread(FRHICommandListExecutor::GetImmediateCommandList(), Size, Format, NumMips, Flags, InResourceState, CreateInfo);
 }
 
 FORCEINLINE FTextureCubeRHIRef RHICreateTextureCubeArray(uint32 Size, uint32 ArraySize, uint8 Format, uint32 NumMips, ETextureCreateFlags Flags, ERHIAccess InResourceState, FRHIResourceCreateInfo& CreateInfo)
 {
-	LLM_SCOPE((Flags & (TexCreate_RenderTargetable | TexCreate_DepthStencilTargetable)) != 0 ? ELLMTag::RenderTargets : ELLMTag::Textures);
+	LLM_SCOPE(EnumHasAnyFlags(Flags, TexCreate_RenderTargetable | TexCreate_DepthStencilTargetable) ? ELLMTag::RenderTargets : ELLMTag::Textures);
 	return GDynamicRHI->RHICreateTextureCubeArray_RenderThread(FRHICommandListExecutor::GetImmediateCommandList(), Size, ArraySize, Format, NumMips, Flags, InResourceState, CreateInfo);
 }
 
diff --git a/Engine/Source/Runtime/RHI/Public/RHIDefinitions.h b/Engine/Source/Runtime/RHI/Public/RHIDefinitions.h
--- a/Engine/Source/Runtime/RHI/Public/RHIDefinitions.h
+++ b/Engine/Source/Runtime/RHI/Public/RHIDefinitions.h
@@ -1528,84 +1528,120 @@ enum class ETextureDimension
 };
 
 /** Flags used for texture creation */
-enum ETextureCreateFlags
-{
-	TexCreate_None					= 0,
-
-	// Texture can be used as a render target
-	TexCreate_RenderTargetable		= 1<<0,
-	// Texture can be used as a resolve target
-	TexCreate_ResolveTargetable		= 1<<1,
-	// Texture can be used as a depth-stencil target.
-	TexCreate_DepthStencilTargetable= 1<<2,
-	// Texture can be used as a shader resource.
-	TexCreate_ShaderResource		= 1<<3,
-	// Texture is encoded in sRGB gamma space
-	TexCreate_SRGB					= 1<<4,
-	// Texture data is writable by the CPU
-	TexCreate_CPUWritable			= 1<<5,
-	// Texture will be created with an un-tiled format
-	TexCreate_NoTiling				= 1<<6,
-	// Texture will be used for video decode
-	TexCreate_VideoDecode			= 1<<7,
-	// Texture that may be updated every frame
-	TexCreate_Dynamic				= 1<<8,
-	// Texture will be used as a render pass attachment that will be read from
-	TexCreate_InputAttachmentRead	= 1 << 9,
-	/** Texture represents a foveation attachment */
-	TexCreate_Foveation				= 1 << 10,
-	// Prefer 3D internal surface tiling mode for volume textures when possible
-	TexCreate_3DTiling				= 1 << 10,
-	// This texture has no GPU or CPU backing. It only exists in tile memory on TBDR GPUs (i.e., mobile).
-	TexCreate_Memoryless			= 1<<11,
-	// Create the texture with the flag that allows mip generation later, only applicable to D3D11
-	TexCreate_GenerateMipCapable	= 1<<12,
-	// The texture can be partially allocated in fastvram
-	TexCreate_FastVRAMPartialAlloc  = 1<<13,
-	// Do not create associated shader resource view, only applicable to D3D11 and D3D12
-	TexCreate_DisableSRVCreation = 1 << 14,
-	// Do not allow Delta Color Compression (DCC) to be used with this texture
-	TexCreate_DisableDCC		    = 1 << 15,
-	// UnorderedAccessView (DX11 only)
-	// Warning: Causes additional synchronization between draw calls when using a render target allocated with this flag, use sparingly
-	// See: GCNPerformanceTweets.pdf Tip 37
-	TexCreate_UAV					= 1<<16,
-	// Render target texture that will be displayed on screen (back buffer)
-	TexCreate_Presentable			= 1<<17,
-	// Texture data is accessible by the CPU
-	TexCreate_CPUReadback			= 1<<18,
-	// Texture was processed offline (via a texture conversion process for the current platform)
-	TexCreate_OfflineProcessed		= 1<<19,
-	// Texture needs to go in fast VRAM if available (HINT only)
-	TexCreate_FastVRAM				= 1<<20,
-	// by default the texture is not showing up in the list - this is to reduce clutter, using the FULL option this can be ignored
-	TexCreate_HideInVisualizeTexture= 1<<21,
-	// Texture should be created in virtual memory, with no physical memory allocation made
-	// You must make further calls to RHIVirtualTextureSetFirstMipInMemory to allocate physical memory
-	// and RHIVirtualTextureSetFirstMipVisible to map the first mip visible to the GPU
-	TexCreate_Virtual				= 1<<22,
-	// Creates a RenderTargetView for each array slice of the texture
-	// Warning: if this was specified when the resource was created, you can't use SV_RenderTargetArrayIndex to route to other slices!
-	TexCreate_TargetArraySlicesIndependently	= 1<<23,
-	// Texture that may be shared with DX9 or other devices
-	TexCreate_Shared = 1 << 24,
-	// RenderTarget will not use full-texture fast clear functionality.
-	TexCreate_NoFastClear = 1 << 25,
-	// Texture is a depth stencil resolve target
-	TexCreate_DepthStencilResolveTarget = 1 << 26,
-	// Flag used to indicted this texture is a streamable 2D texture, and should be counted towards the texture streaming pool budget.
-	TexCreate_Streamable = 1 << 27,
-	// Render target will not FinalizeFastClear; Caches and meta data will be flushed, but clearing will be skipped (avoids potentially trashing metadata)
-	TexCreate_NoFastClearFinalize = 1 << 28,
-	// Hint to the driver that this resource is managed properly by the engine for Alternate-Frame-Rendering in mGPU usage.
-	TexCreate_AFRManual = 1 << 29,
-	// Workaround for 128^3 volume textures getting bloated 4x due to tiling mode on PS4
-	TexCreate_ReduceMemoryWithTilingMode = 1 << 30,
-	/** Texture should be allocated from transient memory. */
-	TexCreate_Transient = 1 << 31
+enum class ETextureCreateFlags : uint64
+{
+    None                              = 0,
+
+    // Texture can be used as a render target
+    RenderTargetable                  = 1ull << 0,
+    // Texture can be used as a resolve target
+    ResolveTargetable                 = 1ull << 1,
+    // Texture can be used as a depth-stencil target.
+    DepthStencilTargetable            = 1ull << 2,
+    // Texture can be used as a shader resource.
+    ShaderResource                    = 1ull << 3,
+    // Texture is encoded in sRGB gamma space
+    SRGB                              = 1ull << 4,
+    // Texture data is writable by the CPU
+    CPUWritable                       = 1ull << 5,
+    // Texture will be created with an un-tiled format
+    NoTiling                          = 1ull << 6,
+    // Texture will be used for video decode
+    VideoDecode                       = 1ull << 7,
+    // Texture that may be updated every frame
+    Dynamic                           = 1ull << 8,
+    // Texture will be used as a render pass attachment that will be read from
+    InputAttachmentRead               = 1ull << 9,
+    /** Texture represents a foveation attachment */
+    Foveation                         = 1ull << 10,
+    // Prefer 3D internal surface tiling mode for volume textures when possible
+    Tiling3D                          = 1ull << 11,
+    // This texture has no GPU or CPU backing. It only exists in tile memory on TBDR GPUs (i.e., mobile).
+    Memoryless                        = 1ull << 12,
+    // Create the texture with the flag that allows mip generation later, only applicable to D3D11
+    GenerateMipCapable                = 1ull << 13,
+    // The texture can be partially allocated in fastvram
+    FastVRAMPartialAlloc              = 1ull << 14,
+    // Do not create associated shader resource view, only applicable to D3D11 and D3D12
+    DisableSRVCreation                = 1ull << 15,
+    // Do not allow Delta Color Compression (DCC) to be used with this texture
+    DisableDCC                        = 1ull << 16,
+    // UnorderedAccessView (DX11 only)
+    // Warning: Causes additional synchronization between draw calls when using a render target allocated with this flag, use sparingly
+    // See: GCNPerformanceTweets.pdf Tip 37
+    UAV                               = 1ull << 17,
+    // Render target texture that will be displayed on screen (back buffer)
+    Presentable                       = 1ull << 18,
+    // Texture data is accessible by the CPU
+    CPUReadback                       = 1ull << 19,
+    // Texture was processed offline (via a texture conversion process for the current platform)
+    OfflineProcessed                  = 1ull << 20,
+    // Texture needs to go in fast VRAM if available (HINT only)
+    FastVRAM                          = 1ull << 21,
+    // by default the texture is not showing up in the list - this is to reduce clutter, using the FULL option this can be ignored
+    HideInVisualizeTexture            = 1ull << 22,
+    // Texture should be created in virtual memory, with no physical memory allocation made
+    // You must make further calls to RHIVirtualTextureSetFirstMipInMemory to allocate physical memory
+    // and RHIVirtualTextureSetFirstMipVisible to map the first mip visible to the GPU
+    Virtual                           = 1ull << 23,
+    // Creates a RenderTargetView for each array slice of the texture
+    // Warning: if this was specified when the resource was created, you can't use SV_RenderTargetArrayIndex to route to other slices!
+    TargetArraySlicesIndependently    = 1ull << 24,
+    // Texture that may be shared with DX9 or other devices
+    Shared                            = 1ull << 25,
+    // RenderTarget will not use full-texture fast clear functionality.
+    NoFastClear                       = 1ull << 26,
+    // Texture is a depth stencil resolve target
+    DepthStencilResolveTarget         = 1ull << 27,
+    // Flag used to indicted this texture is a streamable 2D texture, and should be counted towards the texture streaming pool budget.
+    Streamable                        = 1ull << 28,
+    // Render target will not FinalizeFastClear; Caches and meta data will be flushed, but clearing will be skipped (avoids potentially trashing metadata)
+    NoFastClearFinalize               = 1ull << 29,
+    // Hint to the driver that this resource is managed properly by the engine for Alternate-Frame-Rendering in mGPU usage.
+    AFRManual                         = 1ull << 30,
+    // Workaround for 128^3 volume textures getting bloated 4x due to tiling mode on PS4
+    ReduceMemoryWithTilingMode        = 1ull << 31,
+    /** Texture should be allocated from transient memory. */
+    Transient                         = 1ull << 32,
 };
 ENUM_CLASS_FLAGS(ETextureCreateFlags);
 
+// Compatibility defines
+#define TexCreate_None                           ETextureCreateFlags::None
+#define TexCreate_RenderTargetable               ETextureCreateFlags::RenderTargetable
+#define TexCreate_ResolveTargetable              ETextureCreateFlags::ResolveTargetable
+#define TexCreate_DepthStencilTargetable         ETextureCreateFlags::DepthStencilTargetable
+#define TexCreate_ShaderResource                 ETextureCreateFlags::ShaderResource
+#define TexCreate_SRGB                           ETextureCreateFlags::SRGB
+#define TexCreate_CPUWritable                    ETextureCreateFlags::CPUWritable
+#define TexCreate_NoTiling                       ETextureCreateFlags::NoTiling
+#define TexCreate_VideoDecode                    ETextureCreateFlags::VideoDecode
+#define TexCreate_Dynamic                        ETextureCreateFlags::Dynamic
+#define TexCreate_InputAttachmentRead            ETextureCreateFlags::InputAttachmentRead
+#define TexCreate_Foveation                      ETextureCreateFlags::Foveation
+#define TexCreate_3DTiling                       ETextureCreateFlags::Tiling3D
+#define TexCreate_Memoryless                     ETextureCreateFlags::Memoryless
+#define TexCreate_GenerateMipCapable             ETextureCreateFlags::GenerateMipCapable
+#define TexCreate_FastVRAMPartialAlloc           ETextureCreateFlags::FastVRAMPartialAlloc
+#define TexCreate_DisableSRVCreation             ETextureCreateFlags::DisableSRVCreation
+#define TexCreate_DisableDCC                     ETextureCreateFlags::DisableDCC
+#define TexCreate_UAV                            ETextureCreateFlags::UAV
+#define TexCreate_Presentable                    ETextureCreateFlags::Presentable
+#define TexCreate_CPUReadback                    ETextureCreateFlags::CPUReadback
+#define TexCreate_OfflineProcessed               ETextureCreateFlags::OfflineProcessed
+#define TexCreate_FastVRAM                       ETextureCreateFlags::FastVRAM
+#define TexCreate_HideInVisualizeTexture         ETextureCreateFlags::HideInVisualizeTexture
+#define TexCreate_Virtual                        ETextureCreateFlags::Virtual
+#define TexCreate_TargetArraySlicesIndependently ETextureCreateFlags::TargetArraySlicesIndependently
+#define TexCreate_Shared                         ETextureCreateFlags::Shared
+#define TexCreate_NoFastClear                    ETextureCreateFlags::NoFastClear
+#define TexCreate_DepthStencilResolveTarget      ETextureCreateFlags::DepthStencilResolveTarget
+#define TexCreate_Streamable                     ETextureCreateFlags::Streamable
+#define TexCreate_NoFastClearFinalize            ETextureCreateFlags::NoFastClearFinalize
+#define TexCreate_AFRManual                      ETextureCreateFlags::AFRManual
+#define TexCreate_ReduceMemoryWithTilingMode     ETextureCreateFlags::ReduceMemoryWithTilingMode
+#define TexCreate_Transient                      ETextureCreateFlags::Transient
+
 enum EAsyncComputePriority
 {
 	AsyncComputePriority_Default = 0,
diff --git a/Engine/Source/Runtime/RHI/Public/RHIResources.h b/Engine/Source/Runtime/RHI/Public/RHIResources.h
--- a/Engine/Source/Runtime/RHI/Public/RHIResources.h
+++ b/Engine/Source/Runtime/RHI/Public/RHIResources.h
@@ -2191,7 +2191,7 @@ class FGraphicsPipelineStateInitializer
 public:
 	// Can't use TEnumByte<EPixelFormat> as it changes the struct to be non trivially constructible, breaking memset
 	using TRenderTargetFormats		= TStaticArray<uint8/*EPixelFormat*/, MaxSimultaneousRenderTargets>;
-	using TRenderTargetFlags		= TStaticArray<uint32/*ETextureCreateFlags*/, MaxSimultaneousRenderTargets>;
+	using TRenderTargetFlags		= TStaticArray<ETextureCreateFlags, MaxSimultaneousRenderTargets>;
 
 	FGraphicsPipelineStateInitializer()
 		: BlendState(nullptr)
@@ -2199,9 +2199,9 @@ public:
 		, DepthStencilState(nullptr)
 		, RenderTargetsEnabled(0)
 		, RenderTargetFormats(InPlace, PF_Unknown)
-		, RenderTargetFlags(InPlace, 0)
+		, RenderTargetFlags(InPlace, TexCreate_None)
 		, DepthStencilTargetFormat(PF_Unknown)
-		, DepthStencilTargetFlag(0)
+		, DepthStencilTargetFlag(TexCreate_None)
 		, DepthTargetLoadAction(ERenderTargetLoadAction::ENoAction)
 		, DepthTargetStoreAction(ERenderTargetStoreAction::ENoAction)
 		, StencilTargetLoadAction(ERenderTargetLoadAction::ENoAction)
@@ -2218,7 +2218,6 @@ public:
 	{
 #if PLATFORM_WINDOWS
 		static_assert(sizeof(TRenderTargetFormats::ElementType) == sizeof(uint8/*EPixelFormat*/), "Change TRenderTargetFormats's uint8 to EPixelFormat's size!");
-		static_assert(sizeof(TRenderTargetFlags::ElementType) == sizeof(uint32/*ETextureCreateFlags*/), "Change TRenderTargetFlags's uint32 to ETextureCreateFlags's size!");
 #endif
 		static_assert(PF_MAX < MAX_uint8, "TRenderTargetFormats assumes EPixelFormat can fit in a uint8!");
 	}
@@ -2349,7 +2348,7 @@ public:
 	TRenderTargetFormats			RenderTargetFormats;
 	TRenderTargetFlags				RenderTargetFlags;
 	EPixelFormat					DepthStencilTargetFormat;
-	uint32							DepthStencilTargetFlag;
+	ETextureCreateFlags				DepthStencilTargetFlag;
 	ERenderTargetLoadAction			DepthTargetLoadAction;
 	ERenderTargetStoreAction		DepthTargetStoreAction;
 	ERenderTargetLoadAction			StencilTargetLoadAction;
diff --git a/Engine/Source/Runtime/RHI/Public/RHIUtilities.h b/Engine/Source/Runtime/RHI/Public/RHIUtilities.h
--- a/Engine/Source/Runtime/RHI/Public/RHIUtilities.h
+++ b/Engine/Source/Runtime/RHI/Public/RHIUtilities.h
@@ -692,7 +692,7 @@ inline void RHICreateTargetableShaderResource2D(
 	check(!(bForceSeparateTargetAndShaderResource && bForceSharedTargetAndShaderResource));
 
 	// Ensure that the targetable texture is either render or depth-stencil targetable.
-	check(TargetableTextureFlags & (TexCreate_RenderTargetable | TexCreate_DepthStencilTargetable | TexCreate_UAV));
+	check(EnumHasAnyFlags(TargetableTextureFlags, TexCreate_RenderTargetable | TexCreate_DepthStencilTargetable | TexCreate_UAV));
 
 	if (NumSamples > 1 && !bForceSharedTargetAndShaderResource)
 	{
@@ -707,7 +707,7 @@ inline void RHICreateTargetableShaderResource2D(
 	else
 	{
 		ETextureCreateFlags ResolveTargetableTextureFlags = TexCreate_ResolveTargetable;
-		if (TargetableTextureFlags & TexCreate_DepthStencilTargetable)
+		if (EnumHasAnyFlags(TargetableTextureFlags, TexCreate_DepthStencilTargetable))
 		{
 			ResolveTargetableTextureFlags |= TexCreate_DepthStencilResolveTarget;
 		}
@@ -757,7 +757,7 @@ inline void RHICreateTargetableShaderResource2DArray(
 	check(!(bForceSeparateTargetAndShaderResource && bForceSharedTargetAndShaderResource));
 
 	// Ensure that the targetable texture is either render or depth-stencil targetable.
-	check(TargetableTextureFlags & (TexCreate_RenderTargetable | TexCreate_DepthStencilTargetable | TexCreate_UAV));
+	check(EnumHasAnyFlags(TargetableTextureFlags, TexCreate_RenderTargetable | TexCreate_DepthStencilTargetable | TexCreate_UAV));
 
 	if (NumSamples > 1 && !bForceSharedTargetAndShaderResource)
 	{
@@ -772,7 +772,7 @@ inline void RHICreateTargetableShaderResource2DArray(
 	else
 	{
 		ETextureCreateFlags ResolveTargetableTextureFlags = TexCreate_ResolveTargetable;
-		if (TargetableTextureFlags & TexCreate_DepthStencilTargetable)
+		if (EnumHasAnyFlags(TargetableTextureFlags, TexCreate_DepthStencilTargetable))
 		{
 			ResolveTargetableTextureFlags |= TexCreate_DepthStencilResolveTarget;
 		}
@@ -823,7 +823,7 @@ inline void RHICreateTargetableShaderResourceCube(
 	check(!(Flags & TexCreate_ResolveTargetable));
 
 	// Ensure that the targetable texture is either render or depth-stencil targetable.
-	check(TargetableTextureFlags & (TexCreate_RenderTargetable | TexCreate_DepthStencilTargetable));
+	check(EnumHasAnyFlags(TargetableTextureFlags, TexCreate_RenderTargetable | TexCreate_DepthStencilTargetable));
 
 	if(!bForceSeparateTargetAndShaderResource/* && GSupportsRenderDepthTargetableShaderResources*/)
 	{
@@ -864,7 +864,7 @@ inline void RHICreateTargetableShaderResourceCubeArray(
 	check(!(Flags & TexCreate_ResolveTargetable));
 
 	// Ensure that the targetable texture is either render or depth-stencil targetable.
-	check(TargetableTextureFlags & (TexCreate_RenderTargetable | TexCreate_DepthStencilTargetable));
+	check(EnumHasAnyFlags(TargetableTextureFlags, TexCreate_RenderTargetable | TexCreate_DepthStencilTargetable));
 
 	if(!bForceSeparateTargetAndShaderResource/* && GSupportsRenderDepthTargetableShaderResources*/)
 	{
@@ -906,7 +906,7 @@ inline void RHICreateTargetableShaderResource3D(
 	check(!(Flags & TexCreate_ResolveTargetable));
 
 	// Ensure that the targetable texture is either render or depth-stencil targetable.
-	check(TargetableTextureFlags & (TexCreate_RenderTargetable | TexCreate_DepthStencilTargetable | TexCreate_UAV));
+	check(EnumHasAnyFlags(TargetableTextureFlags, TexCreate_RenderTargetable | TexCreate_DepthStencilTargetable | TexCreate_UAV));
 
 	if (!bForceSeparateTargetAndShaderResource/* && GSupportsRenderDepthTargetableShaderResources*/)
 	{
@@ -916,7 +916,7 @@ inline void RHICreateTargetableShaderResource3D(
 	else
 	{
 		ETextureCreateFlags ResolveTargetableTextureFlags = TexCreate_ResolveTargetable;
-		if (TargetableTextureFlags & TexCreate_DepthStencilTargetable)
+		if (EnumHasAnyFlags(TargetableTextureFlags, TexCreate_DepthStencilTargetable))
 		{
 			ResolveTargetableTextureFlags |= TexCreate_DepthStencilResolveTarget;
 		}
diff --git a/Engine/Source/Runtime/RHI/Public/RHIValidationCommon.h b/Engine/Source/Runtime/RHI/Public/RHIValidationCommon.h
--- a/Engine/Source/Runtime/RHI/Public/RHIValidationCommon.h
+++ b/Engine/Source/Runtime/RHI/Public/RHIValidationCommon.h
@@ -387,7 +387,7 @@ namespace RHIValidation
 
 		virtual FResource* GetTrackerResource() { return &PRIVATE_TrackerResource; }
 
-		inline void InitBarrierTracking(int32 InNumMips, int32 InNumArraySlices, EPixelFormat PixelFormat, uint32 Flags, ERHIAccess InResourceState, const TCHAR* InDebugName)
+		inline void InitBarrierTracking(int32 InNumMips, int32 InNumArraySlices, EPixelFormat PixelFormat, ETextureCreateFlags Flags, ERHIAccess InResourceState, const TCHAR* InDebugName)
 		{
 			FResource* Resource = GetTrackerResource();
 			if (!Resource)
diff --git a/Engine/Source/Runtime/RHI/Public/RHIValidationContext.h b/Engine/Source/Runtime/RHI/Public/RHIValidationContext.h
--- a/Engine/Source/Runtime/RHI/Public/RHIValidationContext.h
+++ b/Engine/Source/Runtime/RHI/Public/RHIValidationContext.h
@@ -346,7 +346,7 @@ public:
 		auto SetAccess = [&](FRHITexture* Texture, ERHIAccess Access)
 		{
 			// If this is a depth texture, only transition the depth plane.
-			RHIValidation::FResourceIdentity Identity = (Texture->GetFlags() & (TexCreate_DepthStencilTargetable | TexCreate_DepthStencilResolveTarget))
+			RHIValidation::FResourceIdentity Identity = EnumHasAnyFlags(Texture->GetFlags(), TexCreate_DepthStencilTargetable | TexCreate_DepthStencilResolveTarget)
 				? Texture->GetViewIdentity(0, 0, 0, 0, uint32(RHIValidation::EResourcePlane::Common), 1)
 				: Texture->GetWholeResourceIdentity();
 
diff --git a/Engine/Source/Runtime/RenderCore/Private/RenderGraphValidation.cpp b/Engine/Source/Runtime/RenderCore/Private/RenderGraphValidation.cpp
--- a/Engine/Source/Runtime/RenderCore/Private/RenderGraphValidation.cpp
+++ b/Engine/Source/Runtime/RenderCore/Private/RenderGraphValidation.cpp
@@ -273,7 +273,7 @@ void FRDGUserValidation::ValidateCreateTexture(const FRDGTextureDesc& Desc, cons
 		TEXT("Failed to create texture %s with pixel format %s because it is not supported."), Name, GPixelFormats[Desc.Format].Name);
 	checkf(Desc.IsValid(), TEXT("Texture %s was created with an invalid descriptor."), Name);
 
-	const bool bCanHaveUAV = (Desc.Flags & TexCreate_UAV) > 0;
+	const bool bCanHaveUAV = EnumHasAnyFlags(Desc.Flags, TexCreate_UAV);
 	const bool bIsMSAA = Desc.NumSamples > 1;
 
 	// D3D11 doesn't allow creating a UAV on MSAA texture.
diff --git a/Engine/Source/Runtime/RenderCore/Private/RenderTargetPool.cpp b/Engine/Source/Runtime/RenderCore/Private/RenderTargetPool.cpp
--- a/Engine/Source/Runtime/RenderCore/Private/RenderTargetPool.cpp
+++ b/Engine/Source/Runtime/RenderCore/Private/RenderTargetPool.cpp
@@ -290,7 +290,7 @@ TRefCountPtr<FPooledRenderTarget> FRenderTargetPool::FindFreeElementInternal(
 		uint32 PassCount = 1;
 		if (AliasingMode == 0)
 		{
-			if ((Desc.Flags & TexCreate_FastVRAM) && bSupportsFastVRAM)
+			if (EnumHasAnyFlags(Desc.Flags, TexCreate_FastVRAM) && bSupportsFastVRAM)
 			{
 				PassCount = 2;
 			}
@@ -316,7 +316,7 @@ TRefCountPtr<FPooledRenderTarget> FRenderTargetPool::FindFreeElementInternal(
 						continue;
 					}
 
-					if ((Desc.Flags & TexCreate_Transient) && bAllowMultipleDiscards == false && Element->HasBeenDiscardedThisFrame())
+					if (EnumHasAnyFlags(Desc.Flags, TexCreate_Transient) && bAllowMultipleDiscards == false && Element->HasBeenDiscardedThisFrame())
 					{
 						// We can't re-use transient resources if they've already been discarded this frame
 						continue;
@@ -351,11 +351,11 @@ Done:
 		PooledRenderTargetHashes.Add(DescHash);
 
 		// TexCreate_UAV should be used on Desc.TargetableFlags
-		check(!(Desc.Flags & TexCreate_UAV));
+		check(!EnumHasAnyFlags(Desc.Flags, TexCreate_UAV));
 
 		FRHIResourceCreateInfo CreateInfo(InDebugName, Desc.ClearValue);
 
-		if (Desc.TargetableFlags & (TexCreate_RenderTargetable | TexCreate_DepthStencilTargetable | TexCreate_UAV))
+		if (EnumHasAnyFlags(Desc.TargetableFlags, TexCreate_RenderTargetable | TexCreate_DepthStencilTargetable | TexCreate_UAV))
 		{
 			// Only create resources if we're not asked to defer creation.
 			if (Desc.Is2DTexture())
@@ -502,7 +502,7 @@ Done:
 #endif
 		}
 
-		if ((Desc.TargetableFlags & TexCreate_UAV))
+		if (EnumHasAnyFlags(Desc.TargetableFlags, TexCreate_UAV))
 		{
 			// The render target desc is invalid if a UAV is requested with an RHI that doesn't support the high-end feature level.
 			check(GMaxRHIFeatureLevel >= ERHIFeatureLevel::SM5 || GMaxRHIFeatureLevel == ERHIFeatureLevel::ES3_1);
@@ -549,7 +549,7 @@ Done:
 	}
 
 	// Transient RTs have to be targettable
-	check((Desc.Flags & TexCreate_Transient) == 0 || Found->GetRenderTargetItem().TargetableTexture != nullptr);
+	check(!EnumHasAnyFlags(Desc.Flags, TexCreate_Transient) || Found->GetRenderTargetItem().TargetableTexture != nullptr);
 
 #if !(UE_BUILD_SHIPPING || UE_BUILD_TEST)
 	if (Found->GetRenderTargetItem().TargetableTexture)
@@ -583,7 +583,7 @@ bool FRenderTargetPool::FindFreeElement(
 	ensure(!(InputDesc.Flags & TexCreate_DepthStencilTargetable) || (InputDesc.ClearValue.ColorBinding == EClearBinding::ENoneBound || InputDesc.ClearValue.ColorBinding == EClearBinding::EDepthStencilBound));
 
 	// TexCreate_FastVRAM should be used on Desc.Flags
-	ensure(!(InputDesc.TargetableFlags & TexCreate_FastVRAM));
+	ensure(!EnumHasAnyFlags(InputDesc.TargetableFlags, TexCreate_FastVRAM));
 
 	// If we're doing aliasing, we may need to override Transient flags, depending on the input format and mode
 	FPooledRenderTargetDesc ModifiedDesc;
@@ -769,7 +769,7 @@ FRenderTargetPool::SMemoryStats FRenderTargetPool::ComputeView()
 			FRTPColumn(const FRenderTargetPoolEvent& Event)
 				: Name(Event.GetDesc().DebugName)
 				, PoolEntryId(Event.GetPoolEntryId())
-				, bVRam((Event.GetDesc().Flags& TexCreate_FastVRAM) != 0)
+				, bVRam(EnumHasAnyFlags(Event.GetDesc().Flags, TexCreate_FastVRAM))
 			{
 				SizeInBytes = Event.GetSizeInBytes();
 			}
diff --git a/Engine/Source/Runtime/RenderCore/Private/VisualizeTexture.cpp b/Engine/Source/Runtime/RenderCore/Private/VisualizeTexture.cpp
--- a/Engine/Source/Runtime/RenderCore/Private/VisualizeTexture.cpp
+++ b/Engine/Source/Runtime/RenderCore/Private/VisualizeTexture.cpp
@@ -411,7 +411,7 @@ void FVisualizeTexture::DisplayPoolResourceListToLog(FVisualizeTexture::ESortBy
 			checkNoEntry();
 		}
 
-		if (Desc.Flags & TexCreate_FastVRAM)
+		if (EnumHasAnyFlags(Desc.Flags, TexCreate_FastVRAM))
 		{
 			FRHIResourceInfo Info;
 
@@ -861,7 +861,7 @@ void FVisualizeTexture::SetCheckPoint(FRDGBuilder& GraphBuilder, IPooledRenderTa
 
 	const FPooledRenderTargetDesc& Desc = PooledRenderTarget->GetDesc();
 
-	if ((Desc.TargetableFlags & TexCreate_ShaderResource) == 0)
+	if (!EnumHasAnyFlags(Desc.TargetableFlags, TexCreate_ShaderResource))
 	{
 		return;
 	}
diff --git a/Engine/Source/Runtime/RenderCore/Public/RendererInterface.h b/Engine/Source/Runtime/RenderCore/Public/RendererInterface.h
--- a/Engine/Source/Runtime/RenderCore/Public/RendererInterface.h
+++ b/Engine/Source/Runtime/RenderCore/Public/RendererInterface.h
@@ -282,7 +282,7 @@ public:
 		}
 
 		return Extent.X != 0 && NumMips != 0 && NumSamples >=1 && NumSamples <=16 && Format != PF_Unknown
-			&& ((TargetableFlags & TexCreate_UAV) == 0 || GMaxRHIFeatureLevel >= ERHIFeatureLevel::SM5 || GMaxRHIFeatureLevel == ERHIFeatureLevel::ES3_1);
+			&& (EnumHasAnyFlags(TargetableFlags, TexCreate_UAV) || GMaxRHIFeatureLevel >= ERHIFeatureLevel::SM5 || GMaxRHIFeatureLevel == ERHIFeatureLevel::ES3_1);
 	}
 
 	FIntVector GetSize() const
@@ -302,11 +302,11 @@ public:
 
 		ETextureCreateFlags LocalFlags = Flags | TargetableFlags;
 
-		if(LocalFlags & TexCreate_RenderTargetable)
+		if(EnumHasAnyFlags(LocalFlags, TexCreate_RenderTargetable))
 		{
 			FlagsString += TEXT(" RT");
 		}
-		if(LocalFlags & TexCreate_SRGB)
+		if(EnumHasAnyFlags(LocalFlags, TexCreate_SRGB))
 		{
 			FlagsString += TEXT(" sRGB");
 		}
@@ -314,17 +314,17 @@ public:
 		{
 			FlagsString += FString::Printf(TEXT(" %dxMSAA"), NumSamples);
 		}
-		if(LocalFlags & TexCreate_UAV)
+		if(EnumHasAnyFlags(LocalFlags, TexCreate_UAV))
 		{
 			FlagsString += TEXT(" UAV");
 		}
 
-		if(LocalFlags & TexCreate_FastVRAM)
+		if(EnumHasAnyFlags(LocalFlags, TexCreate_FastVRAM))
 		{
 			FlagsString += TEXT(" VRam");
 		}
 
-		if (LocalFlags & TexCreate_Transient)
+		if (EnumHasAnyFlags(LocalFlags, TexCreate_Transient))
 		{
 			FlagsString += TEXT(" Transient");
 		}
diff --git a/Engine/Source/Runtime/Renderer/Private/Strata/Strata.cpp b/Engine/Source/Runtime/Renderer/Private/Strata/Strata.cpp
--- a/Engine/Source/Runtime/Renderer/Private/Strata/Strata.cpp
+++ b/Engine/Source/Runtime/Renderer/Private/Strata/Strata.cpp
@@ -218,8 +218,8 @@ void InitialiseStrataFrameSceneData(FSceneRenderer& SceneRenderer, FRDGBuilder&
 		bUpdateLUT = StrataSceneData.GGXEnergyLUT2DTexture == nullptr || StrataSceneData.GGXEnergyLUT2DTexture->GetDesc().Extent.X != LUTResolution || CVarStrataLUTContinousUpdate.GetValueOnAnyThread() > 0;
 		if (bUpdateLUT)
 		{
-			FRDGTextureDesc Desc3D = FRDGTextureDesc::Create3D(FIntVector(LUTResolution, LUTResolution, LUTResolution), EPixelFormat::PF_G16R16F, FClearValueBinding::Black, ETextureCreateFlags::TexCreate_ShaderResource | ETextureCreateFlags::TexCreate_RenderTargetable | ETextureCreateFlags::TexCreate_UAV);
-			FRDGTextureDesc Desc2D = FRDGTextureDesc::Create2D(FIntPoint(LUTResolution, LUTResolution), EPixelFormat::PF_FloatRGBA, FClearValueBinding::Black, ETextureCreateFlags::TexCreate_ShaderResource | ETextureCreateFlags::TexCreate_RenderTargetable | ETextureCreateFlags::TexCreate_UAV);
+			FRDGTextureDesc Desc3D = FRDGTextureDesc::Create3D(FIntVector(LUTResolution, LUTResolution, LUTResolution), EPixelFormat::PF_G16R16F, FClearValueBinding::Black, TexCreate_ShaderResource | TexCreate_RenderTargetable | TexCreate_UAV);
+			FRDGTextureDesc Desc2D = FRDGTextureDesc::Create2D(FIntPoint(LUTResolution, LUTResolution), EPixelFormat::PF_FloatRGBA, FClearValueBinding::Black, TexCreate_ShaderResource | TexCreate_RenderTargetable | TexCreate_UAV);
 			FRDGTextureRef OutLUT3D = GraphBuilder.CreateTexture(Desc3D, TEXT("StrataEnergyLUT3D"));
 			FRDGTextureRef OutLUT2D = GraphBuilder.CreateTexture(Desc2D, TEXT("StrataEnergyLUT2D"));
 
@@ -794,7 +794,7 @@ static void AddStrataLUTPass(FRDGBuilder& GraphBuilder, FSceneRenderer& SceneRen
 	FRDGTextureRef OutLUT3D = GraphBuilder.RegisterExternalTexture(GGXEnergyLUT3DTexture);
 
 	// For debug purpose	
-	FRDGTextureDesc UnfoldLUTDesc = FRDGTextureDesc::Create2D(OutputResolutionRT, EPixelFormat::PF_G16R16F, FClearValueBinding::Black, ETextureCreateFlags::TexCreate_RenderTargetable);
+	FRDGTextureDesc UnfoldLUTDesc = FRDGTextureDesc::Create2D(OutputResolutionRT, EPixelFormat::PF_G16R16F, FClearValueBinding::Black, TexCreate_RenderTargetable);
 	FRDGTextureRef UnfoldLUTTexture = GraphBuilder.CreateTexture(UnfoldLUTDesc, TEXT("StrataEnergyUnfoldLUT"));
 
 	FGlobalShaderMap* GlobalShaderMap = GetGlobalShaderMap(SceneRenderer.FeatureLevel);
diff --git a/Engine/Source/Runtime/Renderer/Private/VisualizeTexturePresent.cpp b/Engine/Source/Runtime/Renderer/Private/VisualizeTexturePresent.cpp
--- a/Engine/Source/Runtime/Renderer/Private/VisualizeTexturePresent.cpp
+++ b/Engine/Source/Runtime/Renderer/Private/VisualizeTexturePresent.cpp
@@ -259,7 +259,7 @@ void FVisualizeTexturePresent::PresentContent(FRDGBuilder& GraphBuilder, const F
 						FLinearColor Color = ElementColor;
 
 						// Highlight EDRAM/FastVRAM usage
-						if (Event->GetDesc().Flags & TexCreate_FastVRAM)
+						if (EnumHasAnyFlags(Event->GetDesc().Flags, TexCreate_FastVRAM))
 						{
 							Color = ElementColorVRam;
 						}
@@ -444,7 +444,7 @@ void FVisualizeTexturePresent::PresentContent(FRDGBuilder& GraphBuilder, const F
 
 		X += 40;
 
-		if (Desc.Flags & TexCreate_CPUReadback)
+		if (EnumHasAnyFlags(Desc.Flags, TexCreate_CPUReadback))
 		{
 			Canvas.DrawShadowedString(X, Y += YStep, TEXT("Content cannot be visualized on the GPU (TexCreate_CPUReadback)"), GetStatsFont(), FLinearColor(1, 1, 0));
 		}
diff --git a/Engine/Source/Runtime/Renderer/Public/MeshPassProcessor.h b/Engine/Source/Runtime/Renderer/Public/MeshPassProcessor.h
--- a/Engine/Source/Runtime/Renderer/Public/MeshPassProcessor.h
+++ b/Engine/Source/Runtime/Renderer/Public/MeshPassProcessor.h
@@ -265,7 +265,7 @@ public:
 			, PrimitiveType
 			, 0
 			, FGraphicsPipelineStateInitializer::TRenderTargetFormats(InPlace, PF_Unknown)
-			, FGraphicsPipelineStateInitializer::TRenderTargetFlags(InPlace, 0)
+			, FGraphicsPipelineStateInitializer::TRenderTargetFlags(InPlace, TexCreate_None)
 			, PF_Unknown
 			, TexCreate_None
 			, ERenderTargetLoadAction::ENoAction
diff --git a/Engine/Source/Runtime/SlateRHIRenderer/Private/SlateRHIRenderer.cpp b/Engine/Source/Runtime/SlateRHIRenderer/Private/SlateRHIRenderer.cpp
--- a/Engine/Source/Runtime/SlateRHIRenderer/Private/SlateRHIRenderer.cpp
+++ b/Engine/Source/Runtime/SlateRHIRenderer/Private/SlateRHIRenderer.cpp
@@ -169,7 +169,7 @@ bool IsMemorylessTexture(const FTexture2DRHIRef& Tex)
 {
 	if (Tex)
 	{
-		return (Tex->GetFlags() & TexCreate_Memoryless) != 0;
+		return EnumHasAnyFlags(Tex->GetFlags(), TexCreate_Memoryless);
 	}
 	return false;
 }
diff --git a/Engine/Source/Runtime/VulkanRHI/Private/VulkanRHI.cpp b/Engine/Source/Runtime/VulkanRHI/Private/VulkanRHI.cpp
--- a/Engine/Source/Runtime/VulkanRHI/Private/VulkanRHI.cpp
+++ b/Engine/Source/Runtime/VulkanRHI/Private/VulkanRHI.cpp
@@ -1147,25 +1147,25 @@ void FVulkanDynamicRHI::RHISubmitCommandsAndFlushGPU()
 
 FTexture2DRHIRef FVulkanDynamicRHI::RHICreateTexture2DFromResource(EPixelFormat Format, uint32 SizeX, uint32 SizeY, uint32 NumMips, uint32 NumSamples, VkImage Resource, ETextureCreateFlags Flags)
 {
-	const FRHIResourceCreateInfo ResourceCreateInfo(TEXT("FVulkanTexture2D"), (Flags & TexCreate_DepthStencilTargetable) ? FClearValueBinding::DepthZero : FClearValueBinding::Transparent);
+	const FRHIResourceCreateInfo ResourceCreateInfo(TEXT("FVulkanTexture2D"), EnumHasAnyFlags(Flags, TexCreate_DepthStencilTargetable) ? FClearValueBinding::DepthZero : FClearValueBinding::Transparent);
 	return new FVulkanTexture2D(*Device, Format, SizeX, SizeY, NumMips, NumSamples, Resource, Flags, ResourceCreateInfo);
 }
 
 FTexture2DRHIRef FVulkanDynamicRHI::RHICreateTexture2DFromResource(EPixelFormat Format, uint32 SizeX, uint32 SizeY, uint32 NumMips, uint32 NumSamples, VkImage Resource, FSamplerYcbcrConversionInitializer& ConversionInitializer, ETextureCreateFlags Flags)
 {
-	const FRHIResourceCreateInfo ResourceCreateInfo(TEXT("FVulkanTexture2D"), (Flags & TexCreate_DepthStencilTargetable) ? FClearValueBinding::DepthZero : FClearValueBinding::Transparent);
+	const FRHIResourceCreateInfo ResourceCreateInfo(TEXT("FVulkanTexture2D"), EnumHasAnyFlags(Flags, TexCreate_DepthStencilTargetable) ? FClearValueBinding::DepthZero : FClearValueBinding::Transparent);
 	return new FVulkanTexture2D(*Device, Format, SizeX, SizeY, NumMips, NumSamples, Resource, ConversionInitializer, Flags, ResourceCreateInfo);
 }
 
 FTexture2DArrayRHIRef FVulkanDynamicRHI::RHICreateTexture2DArrayFromResource(EPixelFormat Format, uint32 SizeX, uint32 SizeY, uint32 ArraySize, uint32 NumMips, uint32 NumSamples, VkImage Resource, ETextureCreateFlags Flags)
 {
-	const FRHIResourceCreateInfo ResourceCreateInfo(TEXT("FVulkanTexture2DArray"), (Flags & TexCreate_DepthStencilTargetable) ? FClearValueBinding::DepthZero : FClearValueBinding::Transparent);
+	const FRHIResourceCreateInfo ResourceCreateInfo(TEXT("FVulkanTexture2DArray"), EnumHasAnyFlags(Flags, TexCreate_DepthStencilTargetable) ? FClearValueBinding::DepthZero : FClearValueBinding::Transparent);
 	return new FVulkanTexture2DArray(*Device, Format, SizeX, SizeY, ArraySize, NumMips, NumSamples, Resource, Flags, ResourceCreateInfo);
 }
 
 FTextureCubeRHIRef FVulkanDynamicRHI::RHICreateTextureCubeFromResource(EPixelFormat Format, uint32 Size, bool bArray, uint32 ArraySize, uint32 NumMips, VkImage Resource, ETextureCreateFlags Flags)
 {
-	const FRHIResourceCreateInfo ResourceCreateInfo(TEXT("FVulkanTextureCube"), (Flags & TexCreate_DepthStencilTargetable) ? FClearValueBinding::DepthZero : FClearValueBinding::Transparent);
+	const FRHIResourceCreateInfo ResourceCreateInfo(TEXT("FVulkanTextureCube"), EnumHasAnyFlags(Flags, TexCreate_DepthStencilTargetable) ? FClearValueBinding::DepthZero : FClearValueBinding::Transparent);
 	return new FVulkanTextureCube(*Device, Format, Size, bArray, ArraySize, NumMips, Resource, Flags, ResourceCreateInfo);
 }
 
diff --git a/Engine/Source/Runtime/VulkanRHI/Private/VulkanRenderTarget.cpp b/Engine/Source/Runtime/VulkanRHI/Private/VulkanRenderTarget.cpp
--- a/Engine/Source/Runtime/VulkanRHI/Private/VulkanRenderTarget.cpp
+++ b/Engine/Source/Runtime/VulkanRHI/Private/VulkanRenderTarget.cpp
@@ -174,10 +174,10 @@ void FVulkanCommandListContext::RHICopyToResolveTarget(FRHITexture* SourceTextur
 
 	ERHIAccess SrcCurrentAccess, DstCurrentAccess;
 
-	check((SrcSurface->UEFlags & TexCreate_CPUReadback) == 0);
+	check(!EnumHasAnyFlags(SrcSurface->UEFlags, TexCreate_CPUReadback));
 	VkImageLayout& SrcLayout = LayoutManager.FindOrAddLayoutRW(*SrcSurface, VK_IMAGE_LAYOUT_UNDEFINED);
 
-	if(DstSurface->UEFlags & TexCreate_CPUReadback)
+	if(EnumHasAnyFlags(DstSurface->UEFlags, TexCreate_CPUReadback))
 	{
 		//Readback textures are represented as a buffer, so we can support miplevels on hardware that does not expose it.
 		FVulkanPipelineBarrier BarrierBefore;
@@ -342,7 +342,7 @@ void FVulkanDynamicRHI::RHIReadSurfaceData(FRHITexture* TextureRHI, FIntRect Rec
 	uint8* MappedPointer = nullptr;
 	VulkanRHI::FStagingBuffer* StagingBuffer = nullptr;
 	FVulkanCmdBuffer* CmdBuffer = nullptr;
-	bool bCPUReadback = (Surface->UEFlags & TexCreate_CPUReadback) == TexCreate_CPUReadback;
+	bool bCPUReadback = EnumHasAllFlags(Surface->UEFlags, TexCreate_CPUReadback);
 	if(!bCPUReadback) //this function supports reading back arbitrary rendertargets, so if its not a cpu readback surface, we do a copy.
 	{
 		ImmediateContext.GetCommandBufferManager()->GetUploadCmdBuffer();
@@ -483,7 +483,7 @@ void FVulkanDynamicRHI::RHIMapStagingSurface(FRHITexture* TextureRHI, FRHIGPUFen
 	}
 
 
-	check((Texture2D->Surface.UEFlags & TexCreate_CPUReadback) == TexCreate_CPUReadback);
+	check(EnumHasAllFlags(Texture2D->Surface.UEFlags, TexCreate_CPUReadback));
 	OutData = Texture2D->Surface.GetMappedPointer();
 	Texture2D->Surface.InvalidateMappedMemory();
 	OutWidth = Texture2D->GetSizeX();
@@ -953,7 +953,7 @@ FVulkanRenderTargetLayout::FVulkanRenderTargetLayout(FVulkanDevice& InDevice, co
 		
 			VkAttachmentDescription& CurrDesc = Desc[NumAttachmentDescriptions];
 			CurrDesc.samples = static_cast<VkSampleCountFlagBits>(NumSamples);
-			CurrDesc.format = UEToVkTextureFormat(RTView.Texture->GetFormat(), (Texture->Surface.UEFlags & TexCreate_SRGB) == TexCreate_SRGB);
+			CurrDesc.format = UEToVkTextureFormat(RTView.Texture->GetFormat(), EnumHasAllFlags(Texture->Surface.UEFlags, TexCreate_SRGB));
 			CurrDesc.loadOp = RenderTargetLoadActionToVulkan(RTView.LoadAction);
 			bFoundClearOp = bFoundClearOp || (CurrDesc.loadOp == VK_ATTACHMENT_LOAD_OP_CLEAR);
 			CurrDesc.storeOp = RenderTargetStoreActionToVulkan(RTView.StoreAction);
@@ -961,7 +961,7 @@ FVulkanRenderTargetLayout::FVulkanRenderTargetLayout(FVulkanDevice& InDevice, co
 			CurrDesc.stencilStoreOp = VK_ATTACHMENT_STORE_OP_DONT_CARE;
 
 			// Removed this temporarily as we need a way to determine if the target is actually memoryless
-			/*if (Texture->Surface.UEFlags & TexCreate_Memoryless)
+			/*if (EnumHasAllFlags(Texture->Surface.UEFlags, TexCreate_Memoryless))
 			{
 				ensure(CurrDesc.storeOp == VK_ATTACHMENT_STORE_OP_DONT_CARE);
 			}*/
@@ -1023,7 +1023,7 @@ FVulkanRenderTargetLayout::FVulkanRenderTargetLayout(FVulkanDevice& InDevice, co
 			CurrDesc.stencilStoreOp = RenderTargetStoreActionToVulkan(RTInfo.DepthStencilRenderTarget.GetStencilStoreAction());
 
 			// Removed this temporarily as we need a way to determine if the target is actually memoryless
-			/*if (Texture->Surface.UEFlags & TexCreate_Memoryless)
+			/*if (EnumHasAllFlags(Texture->Surface.UEFlags, TexCreate_Memoryless))
 			{
 				ensure(CurrDesc.storeOp == VK_ATTACHMENT_STORE_OP_DONT_CARE);
 				ensure(CurrDesc.stencilStoreOp == VK_ATTACHMENT_STORE_OP_DONT_CARE);
@@ -1183,14 +1183,14 @@ FVulkanRenderTargetLayout::FVulkanRenderTargetLayout(FVulkanDevice& InDevice, co
 
 		VkAttachmentDescription& CurrDesc = Desc[NumAttachmentDescriptions];
 		CurrDesc.samples = static_cast<VkSampleCountFlagBits>(NumSamples);
-		CurrDesc.format = UEToVkTextureFormat(ColorEntry.RenderTarget->GetFormat(), (Texture->Surface.UEFlags & TexCreate_SRGB) == TexCreate_SRGB);
+		CurrDesc.format = UEToVkTextureFormat(ColorEntry.RenderTarget->GetFormat(), EnumHasAllFlags(Texture->Surface.UEFlags, TexCreate_SRGB));
 		CurrDesc.loadOp = RenderTargetLoadActionToVulkan(GetLoadAction(ColorEntry.Action));
 		bFoundClearOp = bFoundClearOp || (CurrDesc.loadOp == VK_ATTACHMENT_LOAD_OP_CLEAR);
 		CurrDesc.storeOp = RenderTargetStoreActionToVulkan(GetStoreAction(ColorEntry.Action));
 		CurrDesc.stencilLoadOp = VK_ATTACHMENT_LOAD_OP_DONT_CARE;
 		CurrDesc.stencilStoreOp = VK_ATTACHMENT_STORE_OP_DONT_CARE;
 
-		if (Texture->Surface.UEFlags & TexCreate_Memoryless)
+		if (EnumHasAnyFlags(Texture->Surface.UEFlags, TexCreate_Memoryless))
 		{
 			ensure(CurrDesc.storeOp == VK_ATTACHMENT_STORE_OP_DONT_CARE);
 		}
@@ -1252,7 +1252,7 @@ FVulkanRenderTargetLayout::FVulkanRenderTargetLayout(FVulkanDevice& InDevice, co
 		CurrDesc.storeOp = RenderTargetStoreActionToVulkan(GetStoreAction(GetDepthActions(RPInfo.DepthStencilRenderTarget.Action)));
 		CurrDesc.stencilStoreOp = RenderTargetStoreActionToVulkan(GetStoreAction(GetStencilActions(RPInfo.DepthStencilRenderTarget.Action)));
 
-		if (Texture->Surface.UEFlags & TexCreate_Memoryless)
+		if (EnumHasAnyFlags(Texture->Surface.UEFlags, TexCreate_Memoryless))
 		{
 			ensure(CurrDesc.storeOp == VK_ATTACHMENT_STORE_OP_DONT_CARE);
 			ensure(CurrDesc.stencilStoreOp == VK_ATTACHMENT_STORE_OP_DONT_CARE);
@@ -1436,7 +1436,7 @@ FVulkanRenderTargetLayout::FVulkanRenderTargetLayout(const FGraphicsPipelineStat
 		{
 			VkAttachmentDescription& CurrDesc = Desc[NumAttachmentDescriptions];
 			CurrDesc.samples = static_cast<VkSampleCountFlagBits>(NumSamples);
-			CurrDesc.format = UEToVkTextureFormat(UEFormat, (Initializer.RenderTargetFlags[Index] & TexCreate_SRGB) == TexCreate_SRGB);
+			CurrDesc.format = UEToVkTextureFormat(UEFormat, EnumHasAllFlags(Initializer.RenderTargetFlags[Index], TexCreate_SRGB));
 			CurrDesc.loadOp = VK_ATTACHMENT_LOAD_OP_DONT_CARE;
 			CurrDesc.storeOp = VK_ATTACHMENT_STORE_OP_DONT_CARE;
 			CurrDesc.stencilLoadOp = VK_ATTACHMENT_LOAD_OP_DONT_CARE;
diff --git a/Engine/Source/Runtime/VulkanRHI/Private/VulkanTexture.cpp b/Engine/Source/Runtime/VulkanRHI/Private/VulkanTexture.cpp
--- a/Engine/Source/Runtime/VulkanRHI/Private/VulkanTexture.cpp
+++ b/Engine/Source/Runtime/VulkanRHI/Private/VulkanTexture.cpp
@@ -53,7 +53,7 @@ struct FTextureLock
 #if VULKAN_USE_LLM
 inline ELLMTagVulkan GetMemoryTagForTextureFlags(ETextureCreateFlags UEFlags)
 {
-	bool bRenderTarget = ((TexCreate_RenderTargetable | TexCreate_ResolveTargetable | TexCreate_DepthStencilTargetable) & UEFlags) != 0u;
+	bool bRenderTarget = EnumHasAnyFlags(UEFlags, TexCreate_RenderTargetable | TexCreate_ResolveTargetable | TexCreate_DepthStencilTargetable);
 	return bRenderTarget ? ELLMTagVulkan::VulkanRenderTargets : ELLMTagVulkan::VulkanTextures;
 }
 #endif // VULKAN_USE_LLM
@@ -216,7 +216,7 @@ void FVulkanSurface::GenerateImageCreateInfo(
 	const FPixelFormatInfo& FormatInfo = GPixelFormats[InFormat];
 	VkFormat TextureFormat = (VkFormat)FormatInfo.PlatformFormat;
 
-	if(UEFlags & TexCreate_CPUReadback)
+	if(EnumHasAnyFlags(UEFlags, TexCreate_CPUReadback))
 	{
 		bForceLinearTexture = true;
 	}
@@ -253,10 +253,10 @@ void FVulkanSurface::GenerateImageCreateInfo(
 		break;
 	}
 
-	VkFormat srgbFormat = UEToVkTextureFormat(InFormat, (UEFlags & TexCreate_SRGB) == TexCreate_SRGB);
+	VkFormat srgbFormat = UEToVkTextureFormat(InFormat, EnumHasAllFlags(UEFlags, TexCreate_SRGB));
 	VkFormat nonSrgbFormat = UEToVkTextureFormat(InFormat, false);
 
-	ImageCreateInfo.format = ((UEFlags & TexCreate_UAV) == 0) ? srgbFormat : nonSrgbFormat; 
+	ImageCreateInfo.format = EnumHasAnyFlags(UEFlags, TexCreate_UAV) ? srgbFormat : nonSrgbFormat;
 
 	checkf(ImageCreateInfo.format != VK_FORMAT_UNDEFINED, TEXT("Pixel Format %d not defined!"), (int32)InFormat);
 	if (OutViewFormat)
@@ -279,7 +279,7 @@ void FVulkanSurface::GenerateImageCreateInfo(
 	ImageCreateInfo.flags = (ResourceType == VK_IMAGE_VIEW_TYPE_CUBE || ResourceType == VK_IMAGE_VIEW_TYPE_CUBE_ARRAY) ? VK_IMAGE_CREATE_CUBE_COMPATIBLE_BIT : 0;
 
 
-	if((UEFlags & TexCreate_SRGB) == TexCreate_SRGB)
+	if(EnumHasAllFlags(UEFlags, TexCreate_SRGB))
 	{
 		if(InDevice.GetOptionalExtensions().HasKHRImageFormatList)
 		{
@@ -316,40 +316,40 @@ void FVulkanSurface::GenerateImageCreateInfo(
 	ImageCreateInfo.usage |= VK_IMAGE_USAGE_TRANSFER_SRC_BIT;
 	ImageCreateInfo.usage |= VK_IMAGE_USAGE_SAMPLED_BIT;
 
-	if (UEFlags & TexCreate_Presentable)
+	if (EnumHasAnyFlags(UEFlags, TexCreate_Presentable))
 	{
 		ImageCreateInfo.usage |= VK_IMAGE_USAGE_STORAGE_BIT;		
 	}
-	else if (UEFlags & (TexCreate_RenderTargetable | TexCreate_DepthStencilTargetable))
+	else if (EnumHasAnyFlags(UEFlags, TexCreate_RenderTargetable | TexCreate_DepthStencilTargetable))
 	{
-		if ((UEFlags & TexCreate_InputAttachmentRead) == TexCreate_InputAttachmentRead)
+		if (EnumHasAllFlags(UEFlags, TexCreate_InputAttachmentRead))
 		{
 			ImageCreateInfo.usage |= VK_IMAGE_USAGE_INPUT_ATTACHMENT_BIT;
 		}
-		ImageCreateInfo.usage |= ((UEFlags & TexCreate_RenderTargetable) ? VK_IMAGE_USAGE_COLOR_ATTACHMENT_BIT : VK_IMAGE_USAGE_DEPTH_STENCIL_ATTACHMENT_BIT);
+		ImageCreateInfo.usage |= (EnumHasAnyFlags(UEFlags, TexCreate_RenderTargetable) ? VK_IMAGE_USAGE_COLOR_ATTACHMENT_BIT : VK_IMAGE_USAGE_DEPTH_STENCIL_ATTACHMENT_BIT);
 		ImageCreateInfo.tiling = VK_IMAGE_TILING_OPTIMAL;
-		if ((UEFlags & TexCreate_Memoryless) == TexCreate_Memoryless)
+		if (EnumHasAllFlags(UEFlags, TexCreate_Memoryless))
 		{
 			ImageCreateInfo.usage |= VK_IMAGE_USAGE_TRANSIENT_ATTACHMENT_BIT;
 			// Remove the transfer and sampled bits, as they are incompatible with the transient bit.
 			ImageCreateInfo.usage &= ~(VK_IMAGE_USAGE_TRANSFER_DST_BIT | VK_IMAGE_USAGE_TRANSFER_SRC_BIT | VK_IMAGE_USAGE_SAMPLED_BIT);
 		}
 	}
-	else if (UEFlags & (TexCreate_DepthStencilResolveTarget))
+	else if (EnumHasAnyFlags(UEFlags, TexCreate_DepthStencilResolveTarget))
 	{
 		ImageCreateInfo.usage |= VK_IMAGE_USAGE_DEPTH_STENCIL_ATTACHMENT_BIT;
 		ImageCreateInfo.tiling = VK_IMAGE_TILING_OPTIMAL;
 	}
-	else if (UEFlags & TexCreate_ResolveTargetable)
+	else if (EnumHasAnyFlags(UEFlags, TexCreate_ResolveTargetable))
 	{
 		ImageCreateInfo.usage |= VK_IMAGE_USAGE_COLOR_ATTACHMENT_BIT;
 		ImageCreateInfo.tiling = VK_IMAGE_TILING_OPTIMAL;
 	}
 
-	if (UEFlags & TexCreate_UAV)
+	if (EnumHasAnyFlags(UEFlags, TexCreate_UAV))
 	{
 		//cannot have the storage bit on a memoryless texture
-		ensure((UEFlags & TexCreate_Memoryless) == 0);
+		ensure(!EnumHasAnyFlags(UEFlags, TexCreate_Memoryless));
 		ImageCreateInfo.usage |= VK_IMAGE_USAGE_STORAGE_BIT;
 	}
 
@@ -432,7 +432,7 @@ void FVulkanSurface::GenerateImageCreateInfo(
 		ImageCreateInfo.usage &= ~VK_IMAGE_USAGE_TRANSFER_DST_BIT;
 	}
 
-	if ((UEFlags & TexCreate_DepthStencilTargetable) && GVulkanDepthStencilForceStorageBit)
+	if (EnumHasAnyFlags(UEFlags, TexCreate_DepthStencilTargetable) && GVulkanDepthStencilForceStorageBit)
 	{
 		ImageCreateInfo.usage |= VK_IMAGE_USAGE_STORAGE_BIT;
 	}
@@ -489,7 +489,7 @@ struct FRHICommandOnDestroyImage final : public FRHICommand<FRHICommandOnDestroy
 	}
 };
 
-static VkImageLayout GetInitialLayoutFromRHIAcess(ERHIAccess RHIAccess, uint32 UEFlags)
+static VkImageLayout GetInitialLayoutFromRHIAcess(ERHIAccess RHIAccess, ETextureCreateFlags UEFlags)
 {
 	if (EnumHasAnyFlags(RHIAccess, ERHIAccess::RTV) || RHIAccess == ERHIAccess::Present)
 	{
@@ -508,7 +508,7 @@ static VkImageLayout GetInitialLayoutFromRHIAcess(ERHIAccess RHIAccess, uint32 U
 
 	if (EnumHasAnyFlags(RHIAccess, ERHIAccess::SRVMask))
 	{
-		return (UEFlags & TexCreate_DepthStencilTargetable) ? VK_IMAGE_LAYOUT_DEPTH_STENCIL_READ_ONLY_OPTIMAL : VK_IMAGE_LAYOUT_SHADER_READ_ONLY_OPTIMAL;
+		return EnumHasAnyFlags(UEFlags, TexCreate_DepthStencilTargetable) ? VK_IMAGE_LAYOUT_DEPTH_STENCIL_READ_ONLY_OPTIMAL : VK_IMAGE_LAYOUT_SHADER_READ_ONLY_OPTIMAL;
 	}
 
 	if (EnumHasAnyFlags(RHIAccess, ERHIAccess::UAVMask))
@@ -557,7 +557,7 @@ FVulkanSurface::FVulkanSurface(FVulkanDevice& InDevice, FVulkanEvictable* Owner,
 		InFormat, Width, Height, Depth,
 		ArraySize, NumMips, NumSamples, UEFlags,
 		&StorageFormat, &ViewFormat);
-	if(UEFlags & TexCreate_CPUReadback)
+	if(EnumHasAnyFlags(UEFlags, TexCreate_CPUReadback))
 	{
 		check(NumSamples == 1);	//not implemented
 		check(Depth == 1);		//not implemented
@@ -615,14 +615,14 @@ FVulkanSurface::FVulkanSurface(FVulkanDevice& InDevice, FVulkanEvictable* Owner,
 		MemProps |= VK_MEMORY_PROPERTY_HOST_VISIBLE_BIT;
 	}
 
-	const bool bRenderTarget = (UEFlags & (TexCreate_RenderTargetable | TexCreate_DepthStencilTargetable | TexCreate_ResolveTargetable)) != 0;
-	const bool bUAV = (UEFlags & TexCreate_UAV) != 0;
-	const bool bCPUReadback = (UEFlags & TexCreate_CPUReadback) != 0;
-	const bool bDynamic = (UEFlags & TexCreate_Dynamic) != 0;
+	const bool bRenderTarget = EnumHasAnyFlags(UEFlags, TexCreate_RenderTargetable | TexCreate_DepthStencilTargetable | TexCreate_ResolveTargetable);
+	const bool bUAV = EnumHasAnyFlags(UEFlags, TexCreate_UAV);
+	const bool bCPUReadback = EnumHasAnyFlags(UEFlags, TexCreate_CPUReadback);
+	const bool bDynamic = EnumHasAnyFlags(UEFlags, TexCreate_Dynamic);
 
 	VkMemoryPropertyFlags MemoryFlags = bCPUReadback ? VK_MEMORY_PROPERTY_HOST_VISIBLE_BIT : VK_MEMORY_PROPERTY_DEVICE_LOCAL_BIT;
 
-	bool bMemoryless = (UEFlags & TexCreate_Memoryless) != 0;
+	bool bMemoryless = EnumHasAnyFlags(UEFlags, TexCreate_Memoryless);
 	if (bMemoryless)
 	{
 		if (ensureMsgf(bRenderTarget, TEXT("Memoryless surfaces can only be used for render targets")) && ensureMsgf(!bCPUReadback, TEXT("Memoryless surfaces cannot be read back on CPU")))
@@ -668,7 +668,7 @@ FVulkanSurface::FVulkanSurface(FVulkanDevice& InDevice, FVulkanEvictable* Owner,
 	check(Tiling == VK_IMAGE_TILING_LINEAR || Tiling == VK_IMAGE_TILING_OPTIMAL);
 
 	VkImageLayout InitialLayout = GetInitialLayoutFromRHIAcess(InResourceState, UEFlags);
-	const bool bDoInitialClear = (ImageCreateInfo.ImageCreateInfo.usage & VK_IMAGE_USAGE_SAMPLED_BIT) && (UEFlags & (TexCreate_RenderTargetable | TexCreate_DepthStencilTargetable));
+	const bool bDoInitialClear = (ImageCreateInfo.ImageCreateInfo.usage & VK_IMAGE_USAGE_SAMPLED_BIT) && EnumHasAnyFlags(UEFlags, TexCreate_RenderTargetable | TexCreate_DepthStencilTargetable);
 
 	if (InitialLayout != VK_IMAGE_LAYOUT_UNDEFINED || bDoInitialClear)
 	{
@@ -699,11 +699,11 @@ void FVulkanSurface::MoveSurface(FVulkanDevice& InDevice, FVulkanCommandListCont
 	check(ImageCreateInfo.ImageCreateInfo.tiling == VK_IMAGE_TILING_OPTIMAL);
 
 
-	const bool bRenderTarget = (UEFlags & (TexCreate_RenderTargetable | TexCreate_DepthStencilTargetable | TexCreate_ResolveTargetable)) != 0;
-	const bool bUAV = (UEFlags & TexCreate_UAV) != 0;
-	const bool bCPUReadback = (UEFlags & TexCreate_CPUReadback) != 0;
-	const bool bDynamic = (UEFlags & TexCreate_Dynamic) != 0;
-	const bool bMemoryless = (UEFlags & TexCreate_Memoryless) != 0;
+	const bool bRenderTarget = EnumHasAnyFlags(UEFlags, TexCreate_RenderTargetable | TexCreate_DepthStencilTargetable | TexCreate_ResolveTargetable);
+	const bool bUAV = EnumHasAnyFlags(UEFlags, TexCreate_UAV);
+	const bool bCPUReadback = EnumHasAnyFlags(UEFlags, TexCreate_CPUReadback);
+	const bool bDynamic = EnumHasAnyFlags(UEFlags, TexCreate_Dynamic);
+	const bool bMemoryless = EnumHasAnyFlags(UEFlags, TexCreate_Memoryless);
 	check(!bCPUReadback); //not currently supported
 	check(!bMemoryless);  //not currently supported
 
@@ -799,11 +799,11 @@ void FVulkanSurface::OnFullDefrag(FVulkanDevice& InDevice, FVulkanCommandListCon
 	check(ImageCreateInfo.ImageCreateInfo.tiling == VK_IMAGE_TILING_OPTIMAL);
 
 
-	const bool bRenderTarget = (UEFlags & (TexCreate_RenderTargetable | TexCreate_DepthStencilTargetable | TexCreate_ResolveTargetable)) != 0;
-	const bool bUAV = (UEFlags & TexCreate_UAV) != 0;
-	const bool bCPUReadback = (UEFlags & TexCreate_CPUReadback) != 0;
-	const bool bDynamic = (UEFlags & TexCreate_Dynamic) != 0;
-	const bool bMemoryless = (UEFlags & TexCreate_Memoryless) != 0;
+	const bool bRenderTarget = EnumHasAnyFlags(UEFlags, TexCreate_RenderTargetable | TexCreate_DepthStencilTargetable | TexCreate_ResolveTargetable);
+	const bool bUAV = EnumHasAnyFlags(UEFlags, TexCreate_UAV);
+	const bool bCPUReadback = EnumHasAnyFlags(UEFlags, TexCreate_CPUReadback);
+	const bool bDynamic = EnumHasAnyFlags(UEFlags, TexCreate_Dynamic);
+	const bool bMemoryless = EnumHasAnyFlags(UEFlags, TexCreate_Memoryless);
 	check(!bCPUReadback); //not currently supported
 	check(!bMemoryless);  //not currently supported
 
@@ -917,11 +917,11 @@ void FVulkanSurface::EvictSurface(FVulkanDevice& InDevice)
 
 	check(MemProps == VK_MEMORY_PROPERTY_DEVICE_LOCAL_BIT); //cant evict stuff thats not device local.
 
-	const bool bRenderTarget = (UEFlags & (TexCreate_RenderTargetable | TexCreate_DepthStencilTargetable | TexCreate_ResolveTargetable)) != 0;
-	const bool bUAV = (UEFlags & TexCreate_UAV) != 0;
-	const bool bCPUReadback = (UEFlags & TexCreate_CPUReadback) != 0;
-	const bool bDynamic = (UEFlags & TexCreate_Dynamic) != 0;
-	bool bMemoryless = (UEFlags & TexCreate_Memoryless) != 0;
+	const bool bRenderTarget = EnumHasAnyFlags(UEFlags, TexCreate_RenderTargetable | TexCreate_DepthStencilTargetable | TexCreate_ResolveTargetable);
+	const bool bUAV = EnumHasAnyFlags(UEFlags, TexCreate_UAV);
+	const bool bCPUReadback = EnumHasAnyFlags(UEFlags, TexCreate_CPUReadback);
+	const bool bDynamic = EnumHasAnyFlags(UEFlags, TexCreate_Dynamic);
+	bool bMemoryless = EnumHasAnyFlags(UEFlags, TexCreate_Memoryless);
 	//none of this is supported for eviction
 	check(!bMemoryless);
 	check(!bRenderTarget);
@@ -1042,12 +1042,12 @@ FVulkanSurface::FVulkanSurface(FVulkanDevice& InDevice, VkImageViewType Resource
 
 	checkf(PixelFormat == PF_Unknown || StorageFormat != VK_FORMAT_UNDEFINED, TEXT("PixelFormat %d, is not supported for images"), (int32)PixelFormat);
 
-	ViewFormat = UEToVkTextureFormat(PixelFormat, (UEFlags & TexCreate_SRGB) == TexCreate_SRGB);
+	ViewFormat = UEToVkTextureFormat(PixelFormat, EnumHasAllFlags(UEFlags, TexCreate_SRGB));
 	FullAspectMask = VulkanRHI::GetAspectMaskFromUEFormat(PixelFormat, true, true);
 	PartialAspectMask = VulkanRHI::GetAspectMaskFromUEFormat(PixelFormat, false, true);
 
 	// Purely informative patching, we know that "TexCreate_Presentable" uses optimal tiling
-	if ((UEFlags & TexCreate_Presentable) == TexCreate_Presentable && GetTiling() == VK_IMAGE_TILING_MAX_ENUM)
+	if (EnumHasAllFlags(UEFlags, TexCreate_Presentable) && GetTiling() == VK_IMAGE_TILING_MAX_ENUM)
 	{
 		Tiling = VK_IMAGE_TILING_OPTIMAL;
 	}
@@ -1068,13 +1068,13 @@ FVulkanSurface::FVulkanSurface(FVulkanDevice& InDevice, VkImageViewType Resource
 
 		VkImageLayout InitialLayout;
 		bool bOnlyAddToLayoutManager, bDoInitialClear;
-		if (UEFlags & (TexCreate_RenderTargetable | TexCreate_DepthStencilTargetable))
+		if (EnumHasAnyFlags(UEFlags, TexCreate_RenderTargetable | TexCreate_DepthStencilTargetable))
 		{
-			InitialLayout = (UEFlags & TexCreate_DepthStencilTargetable) ? VK_IMAGE_LAYOUT_DEPTH_STENCIL_ATTACHMENT_OPTIMAL : VK_IMAGE_LAYOUT_COLOR_ATTACHMENT_OPTIMAL;
+			InitialLayout = EnumHasAnyFlags(UEFlags, TexCreate_DepthStencilTargetable) ? VK_IMAGE_LAYOUT_DEPTH_STENCIL_ATTACHMENT_OPTIMAL : VK_IMAGE_LAYOUT_COLOR_ATTACHMENT_OPTIMAL;
 			bOnlyAddToLayoutManager = false;
 			bDoInitialClear = true;
 		}
-		else if (UEFlags & TexCreate_Foveation)
+		else if (EnumHasAnyFlags(UEFlags, TexCreate_Foveation))
 		{
 			// If it's a foveation texture, do not clear but add to layoutmgr, and set correct foveation layout. 
 			InitialLayout = VK_IMAGE_LAYOUT_FRAGMENT_DENSITY_MAP_OPTIMAL_EXT;
@@ -1129,7 +1129,7 @@ void FVulkanSurface::Destroy()
 	}
 	else if (bIsImageOwner)
 	{
-		const bool bRenderTarget = (UEFlags & (TexCreate_RenderTargetable | TexCreate_DepthStencilTargetable | TexCreate_ResolveTargetable)) != 0;
+		const bool bRenderTarget = EnumHasAnyFlags(UEFlags, TexCreate_RenderTargetable | TexCreate_DepthStencilTargetable | TexCreate_ResolveTargetable);
 		FRHICommandList& RHICmdList = FRHICommandListExecutor::GetImmediateCommandList();
 		if (!IsInRenderingThread() || (RHICmdList.Bypass() || !IsRunningRHIInSeparateThread()))
 		{
@@ -2058,7 +2058,7 @@ FVulkanTextureBase::FVulkanTextureBase(FVulkanDevice& Device, VkImageViewType Re
 	Surface.OwningTexture = this;
 	VULKAN_TRACK_OBJECT_CREATE(FVulkanTextureBase, this);
 
-	if(UEFlags & TexCreate_CPUReadback)
+	if(EnumHasAnyFlags(UEFlags, TexCreate_CPUReadback))
 	{
 		return;
 	}
@@ -2068,7 +2068,7 @@ FVulkanTextureBase::FVulkanTextureBase(FVulkanDevice& Device, VkImageViewType Re
 	if (Surface.ViewFormat == VK_FORMAT_UNDEFINED)
 	{
 		Surface.StorageFormat = UEToVkTextureFormat(InFormat, false);
-		Surface.ViewFormat = UEToVkTextureFormat(InFormat, (UEFlags & TexCreate_SRGB) == TexCreate_SRGB);
+		Surface.ViewFormat = UEToVkTextureFormat(InFormat, EnumHasAllFlags(UEFlags, TexCreate_SRGB));
 		checkf(Surface.StorageFormat != VK_FORMAT_UNDEFINED, TEXT("Pixel Format %d not defined!"), (int32)InFormat);
 	}
 
@@ -2186,7 +2186,7 @@ FVulkanTextureBase::FVulkanTextureBase(FVulkanDevice& Device, VkImageViewType Re
 
 	// No MSAA support
 	check(NumSamples == 1);
-	check(!(UEFlags & TexCreate_RenderTargetable));
+	check(!EnumHasAnyFlags(UEFlags, TexCreate_RenderTargetable));
 
 	if (Surface.FullAspectMask == Surface.PartialAspectMask)
 	{
@@ -2730,8 +2730,8 @@ void FVulkanCommandListContext::RHICopyTexture(FRHITexture* SourceTexture, FRHIT
 
 	const FPixelFormatInfo& PixelFormatInfo = GPixelFormats[DestTexture->GetFormat()];
 
-	check((SrcSurface.UEFlags & TexCreate_CPUReadback) == 0);
-	if((DstSurface.UEFlags & TexCreate_CPUReadback) == TexCreate_CPUReadback)
+	check(!EnumHasAnyFlags(SrcSurface.UEFlags, TexCreate_CPUReadback));
+	if(EnumHasAllFlags(DstSurface.UEFlags, TexCreate_CPUReadback))
 	{
 		check(CopyInfo.DestSliceIndex == 0); //slices not supported in TexCreate_CPUReadback textures.
 		FIntVector Size = CopyInfo.Size;
diff --git a/Engine/Source/Runtime/VulkanRHI/Private/VulkanUAV.cpp b/Engine/Source/Runtime/VulkanRHI/Private/VulkanUAV.cpp
--- a/Engine/Source/Runtime/VulkanRHI/Private/VulkanUAV.cpp
+++ b/Engine/Source/Runtime/VulkanRHI/Private/VulkanUAV.cpp
@@ -158,7 +158,7 @@ void FVulkanShaderResourceView::UpdateView()
 	{
 		if (TextureView.View == VK_NULL_HANDLE)
 		{
-			const bool bBaseSRGB = (SourceTexture->GetFlags() & TexCreate_SRGB) != 0;
+			const bool bBaseSRGB = EnumHasAnyFlags(SourceTexture->GetFlags(), TexCreate_SRGB);
 			const bool bSRGB = (SRGBOverride != SRGBO_ForceDisable) && bBaseSRGB;
 
 			EPixelFormat Format = (BufferViewFormat == PF_Unknown) ? SourceTexture->GetFormat() : BufferViewFormat;
diff --git a/Engine/Source/Runtime/Windows/D3D11RHI/Private/D3D11RHIPrivate.h b/Engine/Source/Runtime/Windows/D3D11RHI/Private/D3D11RHIPrivate.h
--- a/Engine/Source/Runtime/Windows/D3D11RHI/Private/D3D11RHIPrivate.h
+++ b/Engine/Source/Runtime/Windows/D3D11RHI/Private/D3D11RHIPrivate.h
@@ -738,14 +738,14 @@ public:
 	void ClearAllShaderResources();
 
 
-	static DXGI_FORMAT GetPlatformTextureResourceFormat(DXGI_FORMAT InFormat, uint32 InFlags);
+	static DXGI_FORMAT GetPlatformTextureResourceFormat(DXGI_FORMAT InFormat, ETextureCreateFlags InFlags);
 
 #if	PLATFORM_SUPPORTS_VIRTUAL_TEXTURES
-	virtual void* CreateVirtualTexture(uint32 SizeX, uint32 SizeY, uint32 SizeZ, uint32 NumMips, bool bCubeTexture, uint32 Flags, void* D3DTextureDesc, void* D3DTextureResource) = 0;
-	virtual void DestroyVirtualTexture(uint32 Flags, void* RawTextureMemory) = 0;
-	virtual bool HandleSpecialLock(FD3D11LockedData& LockedData, uint32 MipIndex, uint32 ArrayIndex, uint32 Flags, EResourceLockMode LockMode,
+	virtual void* CreateVirtualTexture(uint32 SizeX, uint32 SizeY, uint32 SizeZ, uint32 NumMips, bool bCubeTexture, ETextureCreateFlags Flags, void* D3DTextureDesc, void* D3DTextureResource) = 0;
+	virtual void DestroyVirtualTexture(ETextureCreateFlags Flags, void* RawTextureMemory) = 0;
+	virtual bool HandleSpecialLock(FD3D11LockedData& LockedData, uint32 MipIndex, uint32 ArrayIndex, ETextureCreateFlags Flags, EResourceLockMode LockMode,
 		void* D3DTextureResource, void* RawTextureMemory, uint32 NumMips, uint32& DestStride) = 0;
-	virtual bool HandleSpecialUnlock(uint32 MipIndex, uint32 Flags, void* D3DTextureResource, void* RawTextureMemory) = 0;
+	virtual bool HandleSpecialUnlock(uint32 MipIndex, ETextureCreateFlags Flags, void* D3DTextureResource, void* RawTextureMemory) = 0;
 #endif
 
 	uint32 GetHDRDetectedDisplayIndex() const
diff --git a/Engine/Source/Runtime/Windows/D3D11RHI/Private/D3D11Texture.cpp b/Engine/Source/Runtime/Windows/D3D11RHI/Private/D3D11Texture.cpp
--- a/Engine/Source/Runtime/Windows/D3D11RHI/Private/D3D11Texture.cpp
+++ b/Engine/Source/Runtime/Windows/D3D11RHI/Private/D3D11Texture.cpp
@@ -112,7 +112,7 @@ void D3D11TextureAllocated( TD3D11Texture2D<BaseResourceType>& Texture )
 
 	if(D3D11Texture2D)
 	{
-		if ( (Texture.Flags & TexCreate_Virtual) == TexCreate_Virtual )
+		if (EnumHasAllFlags(Texture.Flags, TexCreate_Virtual) )
 		{
 			Texture.SetMemorySize(0);
 		}
@@ -152,7 +152,7 @@ void D3D11TextureDeleted( TD3D11Texture2D<BaseResourceType>& Texture )
 
 		// When using virtual textures use the current memory size, which is the number of physical pages allocated, not virtual
 		int64 TextureSize = 0;
-		if ( (Texture.GetFlags() & TexCreate_Virtual) == TexCreate_Virtual )
+		if (EnumHasAllFlags(Texture.GetFlags(), TexCreate_Virtual))
 		{
 			TextureSize = Texture.GetMemorySize();
 		}
@@ -508,14 +508,14 @@ inline DXGI_FORMAT FindDepthStencilResourceDXGIFormat(DXGI_FORMAT InFormat)
 }
 
 
-DXGI_FORMAT FD3D11DynamicRHI::GetPlatformTextureResourceFormat(DXGI_FORMAT InFormat, uint32 InFlags)
+DXGI_FORMAT FD3D11DynamicRHI::GetPlatformTextureResourceFormat(DXGI_FORMAT InFormat, ETextureCreateFlags InFlags)
 {
 	// Find valid shared texture format
-	if (InFlags & TexCreate_Shared)
+	if (EnumHasAnyFlags(InFlags, TexCreate_Shared))
 	{
-		return FindSharedResourceDXGIFormat(InFormat, InFlags & TexCreate_SRGB);
+		return FindSharedResourceDXGIFormat(InFormat, EnumHasAnyFlags(InFlags, TexCreate_SRGB));
 	}
-	if (InFlags & TexCreate_DepthStencilTargetable)
+	if (EnumHasAnyFlags(InFlags, TexCreate_DepthStencilTargetable))
 	{
 		return FindDepthStencilResourceDXGIFormat(InFormat);
 	}
@@ -622,7 +622,7 @@ TD3D11Texture2D<BaseResourceType>* FD3D11DynamicRHI::CreateD3D11Texture2D(uint32
 
 	bool bPooledTexture = true;
 
-	const bool bSRGB = (Flags & TexCreate_SRGB) != 0;
+	const bool bSRGB = EnumHasAnyFlags(Flags, TexCreate_SRGB);
 
 	const DXGI_FORMAT PlatformResourceFormat = FD3D11DynamicRHI::GetPlatformTextureResourceFormat((DXGI_FORMAT)GPixelFormats[Format].PlatformFormat, Flags);
 	const DXGI_FORMAT PlatformShaderResourceFormat = FindShaderResourceDXGIFormat(PlatformResourceFormat, bSRGB);
@@ -637,7 +637,7 @@ TD3D11Texture2D<BaseResourceType>* FD3D11DynamicRHI::CreateD3D11Texture2D(uint32
 	uint32 ActualMSAAQuality = GetMaxMSAAQuality(ActualMSAACount);
 
 	// 0xffffffff means not supported
-	if (ActualMSAAQuality == 0xffffffff || (Flags & TexCreate_Shared) != 0)
+	if (ActualMSAAQuality == 0xffffffff || EnumHasAnyFlags(Flags, TexCreate_Shared))
 	{
 		// no MSAA
 		ActualMSAACount = 1;
@@ -651,23 +651,21 @@ TD3D11Texture2D<BaseResourceType>* FD3D11DynamicRHI::CreateD3D11Texture2D(uint32
 		bPooledTexture = false;
 	}
 
-	if (NumMips < 1 || SizeX != SizeY || (1 << (NumMips - 1)) != SizeX || (Flags & TexCreate_Shared) != 0)
+	if (NumMips < 1 || SizeX != SizeY || (1 << (NumMips - 1)) != SizeX || EnumHasAnyFlags(Flags, TexCreate_Shared))
 	{
 		bPooledTexture = false;
 	}
 
-	if (Flags & TexCreate_CPUReadback)
+	if (EnumHasAnyFlags(Flags, TexCreate_CPUReadback))
 	{
-		check(!(Flags & TexCreate_RenderTargetable));
-		check(!(Flags & TexCreate_DepthStencilTargetable));
-		check(!(Flags & TexCreate_ShaderResource));
+		check(!EnumHasAnyFlags(Flags, TexCreate_RenderTargetable| TexCreate_DepthStencilTargetable| TexCreate_ShaderResource));
 
 		CPUAccessFlags = D3D11_CPU_ACCESS_READ;
 		TextureUsage = D3D11_USAGE_STAGING;
 		bCreateShaderResource = false;
 	}
 
-	if (Flags & TexCreate_CPUWritable)
+	if (EnumHasAnyFlags(Flags, TexCreate_CPUWritable))
 	{
 		CPUAccessFlags = D3D11_CPU_ACCESS_WRITE;
 		TextureUsage = D3D11_USAGE_STAGING;
@@ -700,12 +698,12 @@ TD3D11Texture2D<BaseResourceType>* FD3D11DynamicRHI::CreateD3D11Texture2D(uint32
 		bCreateShaderResource = false;
 	}
 
-	if (Flags & TexCreate_DisableSRVCreation)
+	if (EnumHasAnyFlags(Flags, TexCreate_DisableSRVCreation))
 	{
 		bCreateShaderResource = false;
 	}
 
-	if (Flags & TexCreate_Shared)
+	if (EnumHasAnyFlags(Flags, TexCreate_Shared))
 	{
 		if (GCVarUseSharedKeyedMutex->GetInt() != 0)
 		{
@@ -717,7 +715,7 @@ TD3D11Texture2D<BaseResourceType>* FD3D11DynamicRHI::CreateD3D11Texture2D(uint32
 		}
 	}
 
-	if (Flags & TexCreate_GenerateMipCapable)
+	if (EnumHasAnyFlags(Flags, TexCreate_GenerateMipCapable))
 	{
 		// Set the flag that allows us to call GenerateMips on this texture later
 		TextureDesc.MiscFlags |= D3D11_RESOURCE_MISC_GENERATE_MIPS;
@@ -729,24 +727,21 @@ TD3D11Texture2D<BaseResourceType>* FD3D11DynamicRHI::CreateD3D11Texture2D(uint32
 	bool bCreateDSV = false;
 	bool bCreatedRTVPerSlice = false;
 
-	if(Flags & TexCreate_RenderTargetable)
+	if(EnumHasAnyFlags(Flags, TexCreate_RenderTargetable))
 	{
-		check(!(Flags & TexCreate_DepthStencilTargetable));
-		check(!(Flags & TexCreate_ResolveTargetable));		
+		check(!EnumHasAnyFlags(Flags, TexCreate_DepthStencilTargetable| TexCreate_ResolveTargetable));
 		TextureDesc.BindFlags |= D3D11_BIND_RENDER_TARGET;
 		bCreateRTV = true;		
 	}
-	else if(Flags & TexCreate_DepthStencilTargetable)
+	else if(EnumHasAnyFlags(Flags, TexCreate_DepthStencilTargetable))
 	{
-		check(!(Flags & TexCreate_RenderTargetable));
-		check(!(Flags & TexCreate_ResolveTargetable));
+		check(!EnumHasAnyFlags(Flags, TexCreate_RenderTargetable| TexCreate_ResolveTargetable));
 		TextureDesc.BindFlags |= D3D11_BIND_DEPTH_STENCIL; 
 		bCreateDSV = true;
 	}
-	else if(Flags & TexCreate_ResolveTargetable)
+	else if(EnumHasAnyFlags(Flags, TexCreate_ResolveTargetable))
 	{
-		check(!(Flags & TexCreate_RenderTargetable));
-		check(!(Flags & TexCreate_DepthStencilTargetable));
+		check(!EnumHasAnyFlags(Flags, TexCreate_RenderTargetable| TexCreate_DepthStencilTargetable));
 		if(Format == PF_DepthStencil || Format == PF_ShadowDepth || Format == PF_D24)
 		{
 			TextureDesc.BindFlags |= D3D11_BIND_DEPTH_STENCIL; 
@@ -765,13 +760,13 @@ TD3D11Texture2D<BaseResourceType>* FD3D11DynamicRHI::CreateD3D11Texture2D(uint32
 		bCreateRTV = false;
 	}
 
-	if (Flags & TexCreate_UAV)
+	if (EnumHasAnyFlags(Flags, TexCreate_UAV))
 	{
 		TextureDesc.BindFlags |= D3D11_BIND_UNORDERED_ACCESS;
 		bPooledTexture = false;
 	}
 
-	if (bCreateDSV && !(Flags & TexCreate_ShaderResource))
+	if (bCreateDSV && !EnumHasAnyFlags(Flags, TexCreate_ShaderResource))
 	{
 		TextureDesc.BindFlags &= ~D3D11_BIND_SHADER_RESOURCE;
 		bCreateShaderResource = false;
@@ -789,7 +784,7 @@ TD3D11Texture2D<BaseResourceType>* FD3D11DynamicRHI::CreateD3D11Texture2D(uint32
 	
 #if PLATFORM_SUPPORTS_VIRTUAL_TEXTURES
 	// Turn off pooling when we are using virtual textures or the texture is offline processed as we control when the memory is released
-	if ( (Flags & (TexCreate_Virtual | TexCreate_OfflineProcessed)) != 0 )
+	if (EnumHasAnyFlags(Flags, TexCreate_Virtual | TexCreate_OfflineProcessed))
 	{
 		bPooledTexture = false;
 	}
@@ -841,7 +836,7 @@ TD3D11Texture2D<BaseResourceType>* FD3D11DynamicRHI::CreateD3D11Texture2D(uint32
 		}
 
 #if PLATFORM_SUPPORTS_VIRTUAL_TEXTURES
-		if ( (Flags & (TexCreate_Virtual | TexCreate_OfflineProcessed)) != 0 )
+		if (EnumHasAnyFlags(Flags, TexCreate_Virtual | TexCreate_OfflineProcessed))
 		{
 			RawTextureMemory = CreateVirtualTexture(SizeX, SizeY, SizeZ, NumMips, bCubeTexture, Flags, &TextureDesc, &TextureResource);
 		}
@@ -856,7 +851,7 @@ TD3D11Texture2D<BaseResourceType>* FD3D11DynamicRHI::CreateD3D11Texture2D(uint32
 			// Create a render target view for each mip
 			for (uint32 MipIndex = 0; MipIndex < NumMips; MipIndex++)
 			{
-				if ((Flags & TexCreate_TargetArraySlicesIndependently) && (bTextureArray || bCubeTexture))
+				if (EnumHasAnyFlags(Flags, TexCreate_TargetArraySlicesIndependently) && (bTextureArray || bCubeTexture))
 				{
 					bCreatedRTVPerSlice = true;
 
@@ -1075,7 +1070,7 @@ TD3D11Texture2D<BaseResourceType>* FD3D11DynamicRHI::CreateD3D11Texture2D(uint32
 	D3D11TextureAllocated(*Texture2D);
 	
 #if !PLATFORM_HOLOLENS
-	if (IsRHIDeviceNVIDIA() && (Flags & TexCreate_AFRManual))
+	if (IsRHIDeviceNVIDIA() && EnumHasAnyFlags(Flags, TexCreate_AFRManual))
 	{
 		// get a resource handle for this texture
 		void* IHVHandle = nullptr;
@@ -1099,7 +1094,7 @@ FD3D11Texture3D* FD3D11DynamicRHI::CreateD3D11Texture3D(uint32 SizeX,uint32 Size
 {
 	SCOPE_CYCLE_COUNTER(STAT_D3D11CreateTextureTime);
 	
-	const bool bSRGB = (Flags & TexCreate_SRGB) != 0;
+	const bool bSRGB = EnumHasAnyFlags(Flags, TexCreate_SRGB);
 
 	const DXGI_FORMAT PlatformResourceFormat = (DXGI_FORMAT)GPixelFormats[Format].PlatformFormat;
 	const DXGI_FORMAT PlatformShaderResourceFormat = FindShaderResourceDXGIFormat(PlatformResourceFormat, bSRGB);
@@ -1118,30 +1113,28 @@ FD3D11Texture3D* FD3D11DynamicRHI::CreateD3D11Texture3D(uint32 SizeX,uint32 Size
 	TextureDesc.CPUAccessFlags = 0;
 	TextureDesc.MiscFlags = 0;
 
-	if (Flags & TexCreate_GenerateMipCapable)
+	if (EnumHasAnyFlags(Flags, TexCreate_GenerateMipCapable))
 	{
 		// Set the flag that allows us to call GenerateMips on this texture later
 		TextureDesc.MiscFlags |= D3D11_RESOURCE_MISC_GENERATE_MIPS;
 	}
 
-	if (Flags & TexCreate_UAV)
+	if (EnumHasAnyFlags(Flags, TexCreate_UAV))
 	{
 		TextureDesc.BindFlags |= D3D11_BIND_UNORDERED_ACCESS;
 	}
 
 	bool bCreateRTV = false;
 
-	if(Flags & TexCreate_RenderTargetable)
+	if(EnumHasAnyFlags(Flags, TexCreate_RenderTargetable))
 	{
 		TextureDesc.BindFlags |= D3D11_BIND_RENDER_TARGET;		
 		bCreateRTV = true;
 	}
 
 	// Set up the texture bind flags.
-	check(!(Flags & TexCreate_DepthStencilTargetable));
-	check(!(Flags & TexCreate_ResolveTargetable));
-	check(Flags & TexCreate_ShaderResource);
-
+	check(!EnumHasAnyFlags(Flags, TexCreate_DepthStencilTargetable | TexCreate_ResolveTargetable));
+	check(EnumHasAllFlags(Flags, TexCreate_ShaderResource));
 
 	TArray<D3D11_SUBRESOURCE_DATA> SubResourceData;
 
@@ -1218,7 +1211,7 @@ FD3D11Texture3D* FD3D11DynamicRHI::CreateD3D11Texture3D(uint32 SizeX,uint32 Size
 
 	D3D11TextureAllocated(*Texture3D);
 #if !PLATFORM_HOLOLENS
-	if (IsRHIDeviceNVIDIA() && (Flags & TexCreate_AFRManual))
+	if (IsRHIDeviceNVIDIA() && EnumHasAnyFlags(Flags, TexCreate_AFRManual))
 	{
 		// get a resource handle for this texture
 		void* IHVHandle = nullptr;
@@ -1272,14 +1265,15 @@ FTexture2DRHIRef FD3D11DynamicRHI::RHIAsyncCreateTexture2D(uint32 SizeX,uint32 S
 	D3D11_SUBRESOURCE_DATA SubResourceData[ MAX_TEXTURE_MIP_COUNT ];
 	FPlatformMemory::Memzero( SubResourceData, sizeof( D3D11_SUBRESOURCE_DATA ) * MAX_TEXTURE_MIP_COUNT );
 
-	uint32 InvalidFlags = TexCreate_RenderTargetable | TexCreate_ResolveTargetable | TexCreate_DepthStencilTargetable | TexCreate_GenerateMipCapable | TexCreate_UAV | TexCreate_Presentable | TexCreate_CPUReadback;
+	ETextureCreateFlags InvalidFlags = TexCreate_RenderTargetable | TexCreate_ResolveTargetable | TexCreate_DepthStencilTargetable | TexCreate_GenerateMipCapable | TexCreate_UAV | TexCreate_Presentable | TexCreate_CPUReadback;
+	check(!EnumHasAnyFlags(Flags, InvalidFlags));
+
 	TArray<TRefCountPtr<ID3D11RenderTargetView> > RenderTargetViews;
 	
 	check(GRHISupportsAsyncTextureCreation);
-	check((Flags & InvalidFlags) == 0);
 
 	const DXGI_FORMAT PlatformResourceFormat = (DXGI_FORMAT)GPixelFormats[Format].PlatformFormat;
-	const DXGI_FORMAT PlatformShaderResourceFormat = FindShaderResourceDXGIFormat(PlatformResourceFormat,((Flags & TexCreate_SRGB) != 0));
+	const DXGI_FORMAT PlatformShaderResourceFormat = FindShaderResourceDXGIFormat(PlatformResourceFormat, EnumHasAnyFlags(Flags, TexCreate_SRGB));
 
 	TextureDesc.Width = SizeX;
 	TextureDesc.Height = SizeY;
@@ -1503,7 +1497,7 @@ FShaderResourceViewRHIRef FD3D11DynamicRHI::RHICreateShaderResourceView(FRHIText
 	}
 
 	// Allow input CreateInfo to override SRGB and/or format
-	const bool bBaseSRGB = (TextureRHI->GetFlags() & TexCreate_SRGB) != 0;
+	const bool bBaseSRGB = EnumHasAnyFlags(TextureRHI->GetFlags(), TexCreate_SRGB);
 	const bool bSRGB = CreateInfo.SRGBOverride != SRGBO_ForceDisable && bBaseSRGB;
 	if (CreateInfo.Format != PF_Unknown)
 	{
@@ -1721,7 +1715,7 @@ void* TD3D11Texture2D<RHIResourceType>::Lock(uint32 MipIndex,uint32 ArrayIndex,E
 #endif
 	if( LockMode == RLM_WriteOnly )
 	{
-		if (!bForceLockDeferred && (Flags & TexCreate_CPUWritable))
+		if (!bForceLockDeferred && EnumHasAnyFlags(Flags, TexCreate_CPUWritable))
 		{
 			D3D11_MAPPED_SUBRESOURCE MappedTexture;
 			VERIFYD3D11RESULT_EX(D3DRHI->GetDeviceContext()->Map(GetResource(), Subresource, D3D11_MAP_WRITE, 0, &MappedTexture), D3DRHI->GetDevice());
@@ -1819,7 +1813,7 @@ void TD3D11Texture2D<RHIResourceType>::Unlock(uint32 MipIndex,uint32 ArrayIndex)
 	}
 	else
 #endif
-	if (!LockedData.bLockDeferred && (Flags & TexCreate_CPUWritable))
+	if (!LockedData.bLockDeferred && EnumHasAnyFlags(Flags, TexCreate_CPUWritable))
 	{
 		D3DRHI->GetDeviceContext()->Unmap(GetResource(), 0);
 	}
@@ -2270,7 +2264,7 @@ TD3D11Texture2D<BaseResourceType>* FD3D11DynamicRHI::CreateTextureFromResource(b
 	D3D11_TEXTURE2D_DESC TextureDesc;
 	TextureResource->GetDesc(&TextureDesc);
 
-	const bool bSRGB = (TexCreateFlags & TexCreate_SRGB) != 0;
+	const bool bSRGB = EnumHasAnyFlags(TexCreateFlags, TexCreate_SRGB);
 
 	const DXGI_FORMAT PlatformResourceFormat = FD3D11DynamicRHI::GetPlatformTextureResourceFormat((DXGI_FORMAT)GPixelFormats[Format].PlatformFormat, TexCreateFlags);
 	const DXGI_FORMAT PlatformShaderResourceFormat = FindShaderResourceDXGIFormat(PlatformResourceFormat, bSRGB);
@@ -2289,7 +2283,7 @@ TD3D11Texture2D<BaseResourceType>* FD3D11DynamicRHI::CreateTextureFromResource(b
 		// Create a render target view for each mip
 		for (uint32 MipIndex = 0; MipIndex < TextureDesc.MipLevels; MipIndex++)
 		{
-			if ((TexCreateFlags & TexCreate_TargetArraySlicesIndependently) && (bTextureArray || bCubeTexture))
+			if (EnumHasAnyFlags(TexCreateFlags, TexCreate_TargetArraySlicesIndependently) && (bTextureArray || bCubeTexture))
 			{
 				bCreatedRTVPerSlice = true;
 
@@ -2506,7 +2500,7 @@ TD3D11Texture2D<BaseResourceType>* FD3D11DynamicRHI::CreateAliasedD3D11Texture2D
 	D3D11_TEXTURE2D_DESC TextureDesc;
 	SourceTexture->GetResource()->GetDesc(&TextureDesc);
 
-	const bool bSRGB = (SourceTexture->Flags & TexCreate_SRGB) != 0;
+	const bool bSRGB = EnumHasAnyFlags(SourceTexture->Flags, TexCreate_SRGB);
 
 	const DXGI_FORMAT PlatformResourceFormat = TextureDesc.Format;
 	const DXGI_FORMAT PlatformShaderResourceFormat = FindShaderResourceDXGIFormat(PlatformResourceFormat, bSRGB);
@@ -2538,7 +2532,7 @@ TD3D11Texture2D<BaseResourceType>* FD3D11DynamicRHI::CreateAliasedD3D11Texture2D
 		for (uint32 MipIndex = 0; MipIndex < TextureDesc.MipLevels; MipIndex++)
 		{
 			// Just add null RTV entries (we'll be aliasing from source shortly).
-			if ((SourceTexture->Flags & TexCreate_TargetArraySlicesIndependently) && (bTextureArray || bCubeTexture))
+			if (EnumHasAnyFlags(SourceTexture->Flags, TexCreate_TargetArraySlicesIndependently) && (bTextureArray || bCubeTexture))
 			{
 				bCreatedRTVPerSlice = true;
 
