Fixed: one more vulkan crash related to skeletal mesh streaming

diff --git a/Engine/Source/Runtime/VulkanRHI/Private/VulkanBuffer.cpp b/Engine/Source/Runtime/VulkanRHI/Private/VulkanBuffer.cpp
--- a/Engine/Source/Runtime/VulkanRHI/Private/VulkanBuffer.cpp
+++ b/Engine/Source/Runtime/VulkanRHI/Private/VulkanBuffer.cpp
@@ -437,6 +437,9 @@ void FVulkanResourceMultiBuffer::Unlock(bool bFromRenderingThread)
 void FVulkanResourceMultiBuffer::Swap(FVulkanResourceMultiBuffer& Other)
 {
 	FRHIBuffer::Swap(Other);
+
+	// FDeviceChild
+	::Swap(Device, Other.Device);
 	
 	::Swap(UEUsage, Other.UEUsage);
 	::Swap(BufferUsageFlags, Other.BufferUsageFlags);
