Make sure Vulkan RHI releases all RHI resources before we do a final FlushPendingDeletes on shutdown

diff --git a/Engine/Source/Runtime/VulkanRHI/Private/VulkanContext.h b/Engine/Source/Runtime/VulkanRHI/Private/VulkanContext.h
--- a/Engine/Source/Runtime/VulkanRHI/Private/VulkanContext.h
+++ b/Engine/Source/Runtime/VulkanRHI/Private/VulkanContext.h
@@ -176,6 +176,8 @@ public:
 
 	void PrepareParallelFromBase(const FVulkanCommandListContext& BaseContext);
 
+	void ReleasePendingState();
+
 	void* Hotfix = nullptr;
 
 protected:
diff --git a/Engine/Source/Runtime/VulkanRHI/Private/VulkanDevice.cpp b/Engine/Source/Runtime/VulkanRHI/Private/VulkanDevice.cpp
--- a/Engine/Source/Runtime/VulkanRHI/Private/VulkanDevice.cpp
+++ b/Engine/Source/Runtime/VulkanRHI/Private/VulkanDevice.cpp
@@ -1237,6 +1237,13 @@ void FVulkanDevice::Destroy()
 	}
 #endif
 
+	// Release pending state that might hold references to RHI resources before we do final FlushPendingDeletes
+	ImmediateContext->ReleasePendingState();
+	if (ComputeContext && ComputeContext != ImmediateContext)
+	{
+		ComputeContext->ReleasePendingState();
+	}
+
 	// Flush all pending deletes before destroying the device and any Vulkan context objects.
 	{
 		FRHICommandListImmediate& RHICmdList = FRHICommandListExecutor::GetImmediateCommandList();
@@ -1462,7 +1469,7 @@ void FVulkanDevice::NotifyDeletedGfxPipeline(class FVulkanRHIGraphicsPipelineSta
 	}
 
 	//#todo-rco: Loop through all contexts!
-	if (ImmediateContext)
+	if (ImmediateContext && ImmediateContext->PendingGfxState)
 	{
 		ImmediateContext->PendingGfxState->NotifyDeletedPipeline(Pipeline);
 	}
@@ -1470,13 +1477,13 @@ void FVulkanDevice::NotifyDeletedGfxPipeline(class FVulkanRHIGraphicsPipelineSta
 
 void FVulkanDevice::NotifyDeletedComputePipeline(class FVulkanComputePipeline* Pipeline)
 {
-	if (ComputeContext && ComputeContext != ImmediateContext)
+	if (ComputeContext && ComputeContext != ImmediateContext && ComputeContext->PendingComputeState)
 	{
 		ComputeContext->PendingComputeState->NotifyDeletedPipeline(Pipeline);
 	}
 
 	//#todo-rco: Loop through all contexts!
-	if (ImmediateContext)
+	if (ImmediateContext && ImmediateContext->PendingComputeState)
 	{
 		ImmediateContext->PendingComputeState->NotifyDeletedPipeline(Pipeline);
 	}
diff --git a/Engine/Source/Runtime/VulkanRHI/Private/VulkanRHI.cpp b/Engine/Source/Runtime/VulkanRHI/Private/VulkanRHI.cpp
--- a/Engine/Source/Runtime/VulkanRHI/Private/VulkanRHI.cpp
+++ b/Engine/Source/Runtime/VulkanRHI/Private/VulkanRHI.cpp
@@ -176,6 +176,15 @@ FVulkanCommandListContext::FVulkanCommandListContext(FVulkanDynamicRHI* InRHI, F
 	GlobalUniformBuffers.AddZeroed(FUniformBufferStaticSlotRegistry::Get().GetSlotCount());
 }
 
+void FVulkanCommandListContext::ReleasePendingState()
+{
+	delete PendingGfxState;
+	PendingGfxState = nullptr;
+	
+	delete PendingComputeState;
+	PendingComputeState = nullptr;
+}
+
 FVulkanCommandListContext::~FVulkanCommandListContext()
 {
 	if (FVulkanPlatform::SupportsTimestampRenderQueries())
