UE5 - Fix Vulkan validation issue

diff --git a/Engine/Source/Runtime/VulkanRHI/Private/Windows/VulkanWindowsPlatform.cpp b/Engine/Source/Runtime/VulkanRHI/Private/Windows/VulkanWindowsPlatform.cpp
--- a/Engine/Source/Runtime/VulkanRHI/Private/Windows/VulkanWindowsPlatform.cpp
+++ b/Engine/Source/Runtime/VulkanRHI/Private/Windows/VulkanWindowsPlatform.cpp
@@ -233,7 +233,8 @@ void FVulkanWindowsPlatform::GetDeviceExtensions(EGpuVendorId VendorId, TArray<c
 #endif
 
 #if VULKAN_SUPPORTS_FRAGMENT_SHADING_RATE
-	OutExtensions.Add(VK_KHR_FRAGMENT_SHADING_RATE_EXTENSION_NAME);
+	// Disabling as validation complains: Missing extension required by the device extension VK_KHR_fragment_shading_rate: VK_KHR_create_renderpass2
+	//OutExtensions.Add(VK_KHR_FRAGMENT_SHADING_RATE_EXTENSION_NAME);
 #endif
 }
 
