UE5 - Disable Vulkan raytracing
Fix currently unused uninitialized variables
Change aftermath engine version to use FApp's version

diff --git a/Engine/Source/Runtime/VulkanRHI/Private/VulkanPipeline.h b/Engine/Source/Runtime/VulkanRHI/Private/VulkanPipeline.h
--- a/Engine/Source/Runtime/VulkanRHI/Private/VulkanPipeline.h
+++ b/Engine/Source/Runtime/VulkanRHI/Private/VulkanPipeline.h
@@ -717,8 +717,8 @@ public:
 
 
 	FVulkanRHIGraphicsPipelineStateLRUNode* LRUNode = nullptr;
-	uint32 LRUFrame;
-	uint32 PipelineCacheSize;
+	uint32 LRUFrame = UINT32_MAX;
+	uint32 PipelineCacheSize = UINT32_MAX;
 	FVulkanPSOKey							VulkanKey;
 
 
diff --git a/Engine/Source/Runtime/VulkanRHI/Private/VulkanUtil.cpp b/Engine/Source/Runtime/VulkanRHI/Private/VulkanUtil.cpp
--- a/Engine/Source/Runtime/VulkanRHI/Private/VulkanUtil.cpp
+++ b/Engine/Source/Runtime/VulkanRHI/Private/VulkanUtil.cpp
@@ -766,8 +766,8 @@ void AftermathGpuCrashDumpCallback(const void* CrashDump, const uint32 CrashDump
 				// Step 2: Allocate a buffer and fetch the generated JSON.
 				TArray<ANSICHAR> Json;
 				Json.AddZeroed(JsonSize);
-				GFSDK_Aftermath_Result Result2 = GFSDK_Aftermath_GpuCrashDump_GetJSON(Decoder, (uint32)Json.Num(), Json.GetData());
-				if (Result2 == GFSDK_Aftermath_Result_Success)
+				GFSDK_Aftermath_Result ResultJson = GFSDK_Aftermath_GpuCrashDump_GetJSON(Decoder, (uint32)Json.Num(), Json.GetData());
+				if (ResultJson == GFSDK_Aftermath_Result_Success)
 				{
 					FString Filename = FPaths::ProjectLogDir() / TEXT("vulkan.nv-gpudmp.json");
 					FArchive* Writer = IFileManager::Get().CreateFileWriter(*Filename);
@@ -800,8 +800,9 @@ void AftermathCrashDumpDescriptionCallback(PFN_GFSDK_Aftermath_AddGpuCrashDumpDe
 	// the actual GPU crash dump callback. The provided data is included in the crash dump and can be
 	// retrieved using GFSDK_Aftermath_GpuCrashDump_GetDescription().
 	FTCHARToUTF8 ProjectNameConverter(FApp::GetProjectName());
+	FTCHARToUTF8 VersionConverter(FApp::GetBuildVersion());
 	AddDescription(GFSDK_Aftermath_GpuCrashDumpDescriptionKey_ApplicationName, ProjectNameConverter.Get());
-	AddDescription(GFSDK_Aftermath_GpuCrashDumpDescriptionKey_ApplicationVersion, "v1.0");
+	AddDescription(GFSDK_Aftermath_GpuCrashDumpDescriptionKey_ApplicationVersion, VersionConverter.Get());
 	AddDescription(GFSDK_Aftermath_GpuCrashDumpDescriptionKey_UserDefined, "Vulkan GPU crash");
 }
 #endif
diff --git a/Engine/Source/Runtime/VulkanRHI/Private/Windows/VulkanWindowsPlatform.h b/Engine/Source/Runtime/VulkanRHI/Private/Windows/VulkanWindowsPlatform.h
--- a/Engine/Source/Runtime/VulkanRHI/Private/Windows/VulkanWindowsPlatform.h
+++ b/Engine/Source/Runtime/VulkanRHI/Private/Windows/VulkanWindowsPlatform.h
@@ -19,8 +19,8 @@
 
 #define	UE_VK_API_VERSION							VK_API_VERSION_1_1
 
-#define VULKAN_RHI_RAYTRACING (RHI_RAYTRACING)
-#define VULKAN_SUPPORTS_SCALAR_BLOCK_LAYOUT (VULKAN_RHI_RAYTRACING)
+#define VULKAN_RHI_RAYTRACING 						0&&(RHI_RAYTRACING)
+#define VULKAN_SUPPORTS_SCALAR_BLOCK_LAYOUT			(VULKAN_RHI_RAYTRACING)
 
 // 32-bit windows has warnings on custom mem mgr callbacks
 #define VULKAN_SHOULD_USE_LLM					(UE_BUILD_DEBUG || UE_BUILD_DEVELOPMENT) && !PLATFORM_32BITS
