UE5EA - Allow Vulkan to initialize if layers or extensions requested by plugins are missing

diff --git a/Engine/Source/Runtime/VulkanRHI/Private/VulkanLayers.cpp b/Engine/Source/Runtime/VulkanRHI/Private/VulkanLayers.cpp
--- a/Engine/Source/Runtime/VulkanRHI/Private/VulkanLayers.cpp
+++ b/Engine/Source/Runtime/VulkanRHI/Private/VulkanLayers.cpp
@@ -620,6 +620,24 @@ void FVulkanDevice::GetDeviceExtensionsAndLayers(VkPhysicalDevice Gpu, EGpuVendo
 				OutDeviceLayers.Add(STANDARD_VALIDATION_LAYER_NAME);
 			}
 		}
+
+		if (!bStandardAvailable)
+		{
+			for (uint32 LayerIndex = 0; GIndividualValidationLayers[LayerIndex] != nullptr; ++LayerIndex)
+			{
+				bool bValidationFound = false;
+				const ANSICHAR* CurrValidationLayer = GIndividualValidationLayers[LayerIndex];
+				for (int32 Index = 1; Index < DeviceLayerExtensions.Num(); ++Index)
+				{
+					if (!FCStringAnsi::Strcmp(DeviceLayerExtensions[Index].LayerProps.layerName, CurrValidationLayer))
+					{
+						bValidationFound = true;
+						OutDeviceLayers.Add(CurrValidationLayer);
+						break;
+					}
+				}
+			}
+		}
 	}
 #endif	// VULKAN_HAS_DEBUGGING_ENABLED
 
