Use dynamic uniform buffers by default for Vulkan. This helps to reduce number of allocated descriptors

diff --git a/Engine/Source/Runtime/VulkanRHI/Private/VulkanShaders.cpp b/Engine/Source/Runtime/VulkanRHI/Private/VulkanShaders.cpp
--- a/Engine/Source/Runtime/VulkanRHI/Private/VulkanShaders.cpp
+++ b/Engine/Source/Runtime/VulkanRHI/Private/VulkanShaders.cpp
@@ -13,9 +13,9 @@
 
 TAutoConsoleVariable<int32> GDynamicGlobalUBs(
 	TEXT("r.Vulkan.DynamicGlobalUBs"),
-	1,
-	TEXT("2 to treat ALL uniform buffers as dynamic\n")\
-	TEXT("1 to treat global/packed uniform buffers as dynamic [default]\n")\
+	2,
+	TEXT("2 to treat ALL uniform buffers as dynamic [default]\n")\
+	TEXT("1 to treat global/packed uniform buffers as dynamic\n")\
 	TEXT("0 to treat them as regular"),
 	ECVF_ReadOnly | ECVF_RenderThreadSafe
 );
