VK: Sub-allocate 'single use' and 'single frame' uniform buffers from a ring buffer

diff --git a/Engine/Source/Runtime/VulkanRHI/Private/VulkanResources.h b/Engine/Source/Runtime/VulkanRHI/Private/VulkanResources.h
--- a/Engine/Source/Runtime/VulkanRHI/Private/VulkanResources.h
+++ b/Engine/Source/Runtime/VulkanRHI/Private/VulkanResources.h
@@ -1178,7 +1178,7 @@ public:
 public:
 	FVulkanDevice* Device;
 	VulkanRHI::FVulkanAllocation Allocation;
-
+	EUniformBufferUsage Usage;
 protected:
 	TArray<TRefCountPtr<FRHIResource>> ResourceTable;
 	
diff --git a/Engine/Source/Runtime/VulkanRHI/Private/VulkanUniformBuffer.cpp b/Engine/Source/Runtime/VulkanRHI/Private/VulkanUniformBuffer.cpp
--- a/Engine/Source/Runtime/VulkanRHI/Private/VulkanUniformBuffer.cpp
+++ b/Engine/Source/Runtime/VulkanRHI/Private/VulkanUniformBuffer.cpp
@@ -44,8 +44,69 @@ static inline EBufferUsageFlags UniformBufferToBufferUsage(EUniformBufferUsage U
 	}
 }
 
+static bool UseRingBuffer(EUniformBufferUsage Usage)
+{
+	// Add a cvar to control this behavior?
+	return false;//(Usage == UniformBuffer_SingleDraw || Usage == UniformBuffer_SingleFrame);
+}
+
+static void UpdateUniformBufferHelper(FVulkanCommandListContext& Context, FVulkanUniformBuffer* VulkanUniformBuffer, int32 DataSize, const void* Data)
+{
+	FVulkanCmdBuffer* CmdBuffer = Context.GetCommandBufferManager()->GetActiveCmdBufferDirect();
+	
+	if (UseRingBuffer(VulkanUniformBuffer->Usage))
+	{
+		FVulkanUniformBufferUploader* UniformBufferUploader = Context.GetUniformBufferUploader();
+		const VkDeviceSize UBOffsetAlignment = Context.GetDevice()->GetLimits().minUniformBufferOffsetAlignment;
+		const FVulkanAllocation& RingBufferAllocation = UniformBufferUploader->GetCPUBufferAllocation();
+		uint64 RingBufferOffset = UniformBufferUploader->AllocateMemory(DataSize, UBOffsetAlignment, CmdBuffer);
+
+		VulkanUniformBuffer->Allocation.Init(
+			EVulkanAllocationEmpty, 
+			EVulkanAllocationMetaUnknown, 
+			RingBufferAllocation.VulkanHandle, 
+			DataSize,
+			RingBufferOffset,
+			RingBufferAllocation.AllocatorIndex,
+			RingBufferAllocation.AllocationIndex,
+			RingBufferAllocation.HandleId);
+		
+		uint8* UploadLocation = UniformBufferUploader->GetCPUMappedPointer() + RingBufferOffset;
+		FMemory::Memcpy(UploadLocation, Data, DataSize);
+	}
+	else
+	{
+		ensure(CmdBuffer->IsOutsideRenderPass());
+		VulkanRHI::FTempFrameAllocationBuffer::FTempAllocInfo LockInfo;
+		Context.GetTempFrameAllocationBuffer().Alloc(DataSize, 16, LockInfo);
+		FMemory::Memcpy(LockInfo.Data, Data, DataSize);
+		VkBufferCopy Region;
+		Region.size = DataSize;
+		Region.srcOffset = LockInfo.GetBindOffset();
+		Region.dstOffset = VulkanUniformBuffer->GetOffset();
+		VkBuffer UBBuffer = VulkanUniformBuffer->Allocation.GetBufferHandle();
+		VkBuffer LockHandle = VulkanUniformBuffer->Allocation.GetBufferHandle();
+
+		bool bIsInRenderPass = CmdBuffer->IsInsideRenderPass();
+		bool bIsUniformBarrierAdded = CmdBuffer->IsUniformBufferBarrierAdded();
+		if(bIsInRenderPass || !bIsUniformBarrierAdded)
+		{
+			CmdBuffer->BeginUniformUpdateBarrier();
+		}
+
+		VulkanRHI::vkCmdCopyBuffer(CmdBuffer->GetHandle(), LockHandle, UBBuffer, 1, &Region);
+
+		if(bIsInRenderPass) //when updating outside render passes, the EndUniformUpdateBarrier will be called from EndRenderPass
+		{
+			CmdBuffer->EndUniformUpdateBarrier();
+		}
+	}
+};
+
 FVulkanUniformBuffer::FVulkanUniformBuffer(FVulkanDevice& Device, const FRHIUniformBufferLayout& InLayout, const void* Contents, EUniformBufferUsage InUsage, EUniformBufferValidation Validation)
-	: FRHIUniformBuffer(InLayout), Device(&Device)
+	: FRHIUniformBuffer(InLayout)
+	, Device(&Device) 
+	, Usage(InUsage)
 {
 #if VULKAN_ENABLE_AGGRESSIVE_STATS
 	SCOPE_CYCLE_COUNTER(STAT_VulkanUniformBufferCreateTime);
@@ -71,10 +132,39 @@ FVulkanUniformBuffer::FVulkanUniformBuffer(FVulkanDevice& Device, const FRHIUnif
 
 	if (InLayout.ConstantBufferSize > 0)
 	{
-		VulkanRHI::FMemoryManager& ResourceMgr = Device.GetMemoryManager();
-
-		// Set it directly as there is no previous one
-		ResourceMgr.AllocUniformBuffer(Allocation, InLayout.ConstantBufferSize, Contents);
+		if (UseRingBuffer(InUsage))
+		{
+			FRHICommandListImmediate& RHICmdList = FRHICommandListExecutor::GetImmediateCommandList();
+			FVulkanUniformBuffer* UniformBuffer = this;
+			int32 DataSize = InLayout.ConstantBufferSize;
+			
+			// make sure we allocate from RingBuffer on RHIT
+			const bool bCanAllocOnThisThread = RHICmdList.Bypass() || (!IsRunningRHIInSeparateThread() && IsInRenderingThread()) || IsInRHIThread();
+			if (bCanAllocOnThisThread)
+			{
+				FVulkanCommandListContextImmediate& Context = Device.GetImmediateContext();
+				UpdateUniformBufferHelper(Context, UniformBuffer, DataSize, Contents);
+			}
+			else
+			{
+				void* CmdListConstantBufferData = RHICmdList.Alloc(DataSize, 16);
+				FMemory::Memcpy(CmdListConstantBufferData, Contents, DataSize);
+				
+				RHICmdList.EnqueueLambda([UniformBuffer, DataSize, CmdListConstantBufferData](FRHICommandList& CmdList)
+				{
+					FVulkanCommandListContext& Context = FVulkanCommandListContext::GetVulkanContext(CmdList.GetContext());
+					UpdateUniformBufferHelper(Context, UniformBuffer, DataSize, CmdListConstantBufferData);
+				});
+
+				RHICmdList.RHIThreadFence(true);
+			}
+		}
+		else
+		{
+			VulkanRHI::FMemoryManager& ResourceMgr = Device.GetMemoryManager();
+			// Set it directly as there is no previous one
+			ResourceMgr.AllocUniformBuffer(Allocation, InLayout.ConstantBufferSize, Contents);
+		}
 	}
 }
 
@@ -115,6 +205,8 @@ FUniformBufferRHIRef FVulkanDynamicRHI::RHICreateUniformBuffer(const void* Conte
 inline void FVulkanDynamicRHI::UpdateUniformBuffer(FVulkanUniformBuffer* UniformBuffer, const void* Contents)
 {
 	SCOPE_CYCLE_COUNTER(STAT_VulkanUpdateUniformBuffers);
+	check(IsInRenderingThread());
+
 	const FRHIUniformBufferLayout& Layout = UniformBuffer->GetLayout();
 
 	const int32 ConstantBufferSize = Layout.ConstantBufferSize;
@@ -124,8 +216,9 @@ inline void FVulkanDynamicRHI::UpdateUniformBuffer(FVulkanUniformBuffer* Uniform
 
 	FVulkanAllocation NewUBAlloc;
 	bool bUseUpload = GVulkanAllowUniformUpload && !RHICmdList.IsInsideRenderPass(); //inside renderpasses, a rename is enforced.
+	const bool bUseRingBuffer = UseRingBuffer(UniformBuffer->Usage);
 
-	if (!bUseUpload)
+	if (!bUseUpload && !bUseRingBuffer)
 	{
 		if (ConstantBufferSize > 0)
 		{
@@ -134,41 +227,12 @@ inline void FVulkanDynamicRHI::UpdateUniformBuffer(FVulkanUniformBuffer* Uniform
 		}
 	}
 
-	auto UpdateUniformBufferHelper = [](FVulkanCommandListContext& Context, FVulkanUniformBuffer* VulkanUniformBuffer, int32 DataSize, const void* Data)
-	{
-		FVulkanCmdBuffer* CmdBuffer = Context.GetCommandBufferManager()->GetActiveCmdBufferDirect();
-		ensure(CmdBuffer->IsOutsideRenderPass());
-		VulkanRHI::FTempFrameAllocationBuffer::FTempAllocInfo LockInfo;
-		Context.GetTempFrameAllocationBuffer().Alloc(DataSize, 16, LockInfo);
-		FMemory::Memcpy(LockInfo.Data, Data, DataSize);
-		VkBufferCopy Region;
-		Region.size = DataSize;
-		Region.srcOffset = LockInfo.GetBindOffset();
-		Region.dstOffset = VulkanUniformBuffer->GetOffset();
-		VkBuffer UBBuffer = VulkanUniformBuffer->Allocation.GetBufferHandle();
-		VkBuffer LockHandle = VulkanUniformBuffer->Allocation.GetBufferHandle();
-
-		bool bIsInRenderPass = CmdBuffer->IsInsideRenderPass();
-		bool bIsUniformBarrierAdded = CmdBuffer->IsUniformBufferBarrierAdded();
-		if(bIsInRenderPass || !bIsUniformBarrierAdded)
-		{
-			CmdBuffer->BeginUniformUpdateBarrier();
-		}
-
-		VulkanRHI::vkCmdCopyBuffer(CmdBuffer->GetHandle(), LockHandle, UBBuffer, 1, &Region);
-
-		if(bIsInRenderPass) //when updating outside render passes, the EndUniformUpdateBarrier will be called from EndRenderPass
-		{
-			CmdBuffer->EndUniformUpdateBarrier();
-		}
-	};
-
 	bool bRHIBypass = RHICmdList.Bypass();
 	if (bRHIBypass)
 	{
 		if (ConstantBufferSize > 0)
 		{
-			if(bUseUpload)
+			if (bUseUpload || bUseRingBuffer)
 			{			
 				FVulkanCommandListContext& Context = Device->GetImmediateContext();
 				UpdateUniformBufferHelper(Context, UniformBuffer, ConstantBufferSize, Contents);
@@ -195,12 +259,12 @@ inline void FVulkanDynamicRHI::UpdateUniformBuffer(FVulkanUniformBuffer* Uniform
 			}
 		}
 
-		if(bUseUpload)
+		if (bUseUpload || bUseRingBuffer)
 		{
 			void* CmdListConstantBufferData = RHICmdList.Alloc(ConstantBufferSize, 16);
 			FMemory::Memcpy(CmdListConstantBufferData, Contents, ConstantBufferSize);
 
-			RHICmdList.EnqueueLambda([UpdateUniformBufferHelper, UniformBuffer, CmdListResources, NumResources, ConstantBufferSize, CmdListConstantBufferData](FRHICommandList& CmdList)
+			RHICmdList.EnqueueLambda([UniformBuffer, CmdListResources, NumResources, ConstantBufferSize, CmdListConstantBufferData](FRHICommandList& CmdList)
 			{
 				FVulkanCommandListContext& Context = (FVulkanCommandListContext&)CmdList.GetContext().GetLowestLevelContext();
 				UpdateUniformBufferHelper(Context, UniformBuffer, ConstantBufferSize, CmdListConstantBufferData);
