VK: Actually using ring buffer in CL# 16583265

diff --git a/Engine/Source/Runtime/VulkanRHI/Private/VulkanUniformBuffer.cpp b/Engine/Source/Runtime/VulkanRHI/Private/VulkanUniformBuffer.cpp
--- a/Engine/Source/Runtime/VulkanRHI/Private/VulkanUniformBuffer.cpp
+++ b/Engine/Source/Runtime/VulkanRHI/Private/VulkanUniformBuffer.cpp
@@ -47,7 +47,7 @@ static inline EBufferUsageFlags UniformBufferToBufferUsage(EUniformBufferUsage U
 static bool UseRingBuffer(EUniformBufferUsage Usage)
 {
 	// Add a cvar to control this behavior?
-	return false;//(Usage == UniformBuffer_SingleDraw || Usage == UniformBuffer_SingleFrame);
+	return (Usage == UniformBuffer_SingleDraw || Usage == UniformBuffer_SingleFrame);
 }
 
 static void UpdateUniformBufferHelper(FVulkanCommandListContext& Context, FVulkanUniformBuffer* VulkanUniformBuffer, int32 DataSize, const void* Data)
