UE5 - Vulkan validation fix
SRGB formats require VK_IMAGE_CREATE_MUTABLE_FORMAT_BIT bit even if VK_KHR_image_format_list is used
#rnx

diff --git a/Engine/Source/Runtime/VulkanRHI/Private/VulkanTexture.cpp b/Engine/Source/Runtime/VulkanRHI/Private/VulkanTexture.cpp
--- a/Engine/Source/Runtime/VulkanRHI/Private/VulkanTexture.cpp
+++ b/Engine/Source/Runtime/VulkanRHI/Private/VulkanTexture.cpp
@@ -292,13 +292,8 @@ void FVulkanSurface::GenerateImageCreateInfo(
 			OutImageCreateInfo.FormatsUsed[0] = nonSrgbFormat;
 			OutImageCreateInfo.FormatsUsed[1] = srgbFormat;
 		}
-		else
-		{
-			// Remove MUTABLE_BIT to enable compression on Mobile Vulkan (QCOM 835)
-			// @todo: Need a long term fix for this.
-			ImageCreateInfo.flags |= VK_IMAGE_CREATE_MUTABLE_FORMAT_BIT;
-		}
 
+		ImageCreateInfo.flags |= VK_IMAGE_CREATE_MUTABLE_FORMAT_BIT;
 	}
 
 #if VULKAN_SUPPORTS_MAINTENANCE_LAYER1
