UE5 - Fix linker error when enabling Vulkan image tracking (missing function)
#rnx

diff --git a/Engine/Source/Runtime/VulkanRHI/Private/VulkanDebug.cpp b/Engine/Source/Runtime/VulkanRHI/Private/VulkanDebug.cpp
--- a/Engine/Source/Runtime/VulkanRHI/Private/VulkanDebug.cpp
+++ b/Engine/Source/Runtime/VulkanRHI/Private/VulkanDebug.cpp
@@ -3792,6 +3792,17 @@ void FWrapLayer::GetPhysicalDeviceFeatures2KHR(VkResult Result, VkPhysicalDevice
 }
 #endif
 
+#if VULKAN_SUPPORTS_FRAGMENT_SHADING_RATE
+void FWrapLayer::GetPhysicalDeviceFragmentShadingRatesKHR(VkResult Result, VkPhysicalDevice PhysicalDevice, uint32* FragmentShadingRateCount, VkPhysicalDeviceFragmentShadingRateKHR* FragmentShadingRates)
+{
+	if (Result == VK_RESULT_MAX_ENUM)
+	{
+#if VULKAN_ENABLE_DUMP_LAYER
+		PrintfBegin(FString::Printf(TEXT("GetPhysicalDeviceFragmentShadingRatesKHR(PhysicalDevice=0x%p, FragmentShadingRateCount=0x%p, FragmentShadingRates=0x%p)[...]"), PhysicalDevice, FragmentShadingRateCount, FragmentShadingRates));
+#endif
+	}
+}
+#endif
 
 void FWrapLayer::CmdSetDepthBias(VkResult Result, VkCommandBuffer CommandBuffer, float DepthBiasConstantFactor, float DepthBiasClamp, float DepthBiasSlopeFactor)
 {
