Vulkan: add RHICreateComputePipelineState implementation.

#[review] [REDACTED_HANDLE], [REDACTED_HANDLE], [REDACTED_HANDLE]
#[robomerge] 17.20

diff --git a/Engine/Source/Runtime/RHI/Private/PipelineStateCache.cpp b/Engine/Source/Runtime/RHI/Private/PipelineStateCache.cpp
--- a/Engine/Source/Runtime/RHI/Private/PipelineStateCache.cpp
+++ b/Engine/Source/Runtime/RHI/Private/PipelineStateCache.cpp
@@ -80,11 +80,7 @@ static FAutoConsoleCommand DumpPipelineCmd(
 
 void SetComputePipelineState(FRHIComputeCommandList& RHICmdList, FRHIComputeShader* ComputeShader)
 {
-#if PLATFORM_ANDROID
-	RHICmdList.SetComputeShader(ComputeShader);
-#else
 	RHICmdList.SetComputePipelineState(PipelineStateCache::GetAndOrCreateComputePipelineState(RHICmdList, ComputeShader), ComputeShader);
-#endif
 }
 
 extern RHI_API FRHIComputePipelineState* ExecuteSetComputePipelineState(FComputePipelineState* ComputePipelineState);
diff --git a/Engine/Source/Runtime/RHI/Public/RHIResources.h b/Engine/Source/Runtime/RHI/Public/RHIResources.h
--- a/Engine/Source/Runtime/RHI/Public/RHIResources.h
+++ b/Engine/Source/Runtime/RHI/Public/RHIResources.h
@@ -1776,6 +1776,7 @@ typedef TRefCountPtr<FRHIViewport> FViewportRHIRef;
 typedef TRefCountPtr<FRHIUnorderedAccessView> FUnorderedAccessViewRHIRef;
 typedef TRefCountPtr<FRHIShaderResourceView> FShaderResourceViewRHIRef;
 typedef TRefCountPtr<FRHIGraphicsPipelineState> FGraphicsPipelineStateRHIRef;
+typedef TRefCountPtr<FRHIComputePipelineState> FComputePipelineStateRHIRef;
 typedef TRefCountPtr<FRHIRayTracingPipelineState> FRayTracingPipelineStateRHIRef;
 
 
diff --git a/Engine/Source/Runtime/VulkanRHI/Private/VulkanPipeline.cpp b/Engine/Source/Runtime/VulkanRHI/Private/VulkanPipeline.cpp
--- a/Engine/Source/Runtime/VulkanRHI/Private/VulkanPipeline.cpp
+++ b/Engine/Source/Runtime/VulkanRHI/Private/VulkanPipeline.cpp
@@ -1805,8 +1805,22 @@ FGraphicsPipelineStateRHIRef FVulkanDynamicRHI::RHICreateGraphicsPipelineState(c
 	return Device->PipelineStateCache->RHICreateGraphicsPipelineState(PSOInitializer);
 }
 
+FVulkanComputePipeline* FVulkanPipelineStateCacheManager::RHICreateComputePipelineState(FRHIComputeShader* ComputeShaderRHI)
+{
+	FVulkanComputeShader* ComputeShader = ResourceCast(ComputeShaderRHI);
+	return Device->GetPipelineStateCache()->GetOrCreateComputePipeline(ComputeShader);
+}
 
+FComputePipelineStateRHIRef FVulkanDynamicRHI::RHICreateComputePipelineState(FRHIComputeShader* ComputeShader)
+{
+#if VULKAN_ENABLE_AGGRESSIVE_STATS
+	SCOPE_CYCLE_COUNTER(STAT_VulkanGetOrCreatePipeline);
+#endif
+	QUICK_SCOPE_CYCLE_COUNTER(STAT_Vulkan_RHICreateComputePipelineState);
+	LLM_SCOPE_VULKAN(ELLMTagVulkan::VulkanShaders);
 
+	return Device->PipelineStateCache->RHICreateComputePipelineState(ComputeShader);
+}
 
 FVulkanComputePipeline* FVulkanPipelineStateCacheManager::GetOrCreateComputePipeline(FVulkanComputeShader* ComputeShader)
 {
diff --git a/Engine/Source/Runtime/VulkanRHI/Private/VulkanPipeline.h b/Engine/Source/Runtime/VulkanRHI/Private/VulkanPipeline.h
--- a/Engine/Source/Runtime/VulkanRHI/Private/VulkanPipeline.h
+++ b/Engine/Source/Runtime/VulkanRHI/Private/VulkanPipeline.h
@@ -504,6 +504,7 @@ private:
 	void DestroyCache();
 
 	FVulkanRHIGraphicsPipelineState* RHICreateGraphicsPipelineState(const FGraphicsPipelineStateInitializer& Initializer);
+	FVulkanComputePipeline* RHICreateComputePipelineState(FRHIComputeShader* ComputeShaderRHI);
 	void NotifyDeletedGraphicsPSO(FRHIGraphicsPipelineState* PSO);
 	bool CreateGfxPipelineFromEntry(FVulkanRHIGraphicsPipelineState* PSO, FVulkanShader* Shaders[ShaderStage::NumStages], VkPipeline* Pipeline);
 
diff --git a/Engine/Source/Runtime/VulkanRHI/Public/VulkanDynamicRHI.h b/Engine/Source/Runtime/VulkanRHI/Public/VulkanDynamicRHI.h
--- a/Engine/Source/Runtime/VulkanRHI/Public/VulkanDynamicRHI.h
+++ b/Engine/Source/Runtime/VulkanRHI/Public/VulkanDynamicRHI.h
@@ -76,6 +76,7 @@ public:
 	virtual FStagingBufferRHIRef RHICreateStagingBuffer() final override;
 	virtual FBoundShaderStateRHIRef RHICreateBoundShaderState(FRHIVertexDeclaration* VertexDeclaration, FRHIVertexShader* VertexShader, FRHIPixelShader* PixelShader, FRHIGeometryShader* GeometryShader) final override;
 	virtual FGraphicsPipelineStateRHIRef RHICreateGraphicsPipelineState(const FGraphicsPipelineStateInitializer& Initializer) final override;
+	virtual FComputePipelineStateRHIRef RHICreateComputePipelineState(FRHIComputeShader* ComputeShader) final override;
 	virtual FUniformBufferRHIRef RHICreateUniformBuffer(const void* Contents, const FRHIUniformBufferLayout& Layout, EUniformBufferUsage Usage, EUniformBufferValidation Validation) final override;
 	virtual void RHIUpdateUniformBuffer(FRHIUniformBuffer* UniformBufferRHI, const void* Contents) final override;
 	virtual FBufferRHIRef RHICreateBuffer(uint32 Size, EBufferUsageFlags Usage, uint32 Stride, ERHIAccess ResourceState, FRHIResourceCreateInfo& CreateInfo) final override;
