Fix issue with flickering in Vulkan.

From the moment a buffer gets locked/unlocked, it modifies the internal Current.Size of that buffer.  That Current.Size is then used as the value to represent the "entire size" of that buffer for clears/fills/binds.  The problem is when a buffer gets reused (rdg's buffer pool, for example):  if it was used in a code path that locks it, then used in a code path that never locks it and expects to use the full size, then the buffer is limited to the size it had in its previous life.  It is also an issue if we lock a buffer at a partial size, but then expect it to be bound at its full size (to read from the filled portions and write to other parts).

NOTE: Fill commands still go beyond used size to clear the entire buffer, but it did before...  Will be tweaked in a later optimization.

diff --git a/Engine/Source/Runtime/VulkanRHI/Private/VulkanResources.h b/Engine/Source/Runtime/VulkanRHI/Private/VulkanResources.h
--- a/Engine/Source/Runtime/VulkanRHI/Private/VulkanResources.h
+++ b/Engine/Source/Runtime/VulkanRHI/Private/VulkanResources.h
@@ -1093,9 +1093,10 @@ public:
 		return Current.Offset;
 	}
 
+	// Remaining size from the current offset
 	inline uint64 GetCurrentSize() const
 	{
-		return Current.Size;
+		return Current.Alloc.Size - (Current.Offset - Current.Alloc.Offset);
 	}
 
 
diff --git a/Engine/Source/Runtime/VulkanRHI/Private/VulkanUAV.cpp b/Engine/Source/Runtime/VulkanRHI/Private/VulkanUAV.cpp
--- a/Engine/Source/Runtime/VulkanRHI/Private/VulkanUAV.cpp
+++ b/Engine/Source/Runtime/VulkanRHI/Private/VulkanUAV.cpp
@@ -443,7 +443,7 @@ void FVulkanCommandListContext::ClearUAVFillBuffer(FVulkanUnorderedAccessView* U
 	FVulkanCmdBuffer* CmdBuffer = CmdBufferMgr->GetActiveCmdBuffer();
 
 	FVulkanResourceMultiBuffer* Buffer = UAV->SourceBuffer;
-	VulkanRHI::vkCmdFillBuffer(CmdBuffer->GetHandle(), Buffer->Current.Handle, Buffer->Current.Offset, Buffer->Current.Size, ClearValue);
+	VulkanRHI::vkCmdFillBuffer(CmdBuffer->GetHandle(), Buffer->GetHandle(), Buffer->GetOffset(), Buffer->GetCurrentSize(), ClearValue);
 }
 
 void FVulkanCommandListContext::ClearUAV(TRHICommandList_RecursiveHazardous<FVulkanCommandListContext>& RHICmdList, FVulkanUnorderedAccessView* UnorderedAccessView, const void* ClearValue, bool bFloat)
@@ -521,7 +521,7 @@ void FVulkanCommandListContext::ClearUAV(TRHICommandList_RecursiveHazardous<FVul
 		}
 		else
 		{
-			uint32 NumElements = Buffer->Current.Size / GPixelFormats[UnorderedAccessView->BufferViewFormat].BlockBytes;
+			uint32 NumElements = Buffer->GetCurrentSize() / GPixelFormats[UnorderedAccessView->BufferViewFormat].BlockBytes;
 			ClearUAVShader_T<EClearReplacementResourceType::Buffer, 4, false>(RHICmdList, UnorderedAccessView, NumElements, 1, 1, ClearValue, ValueType);
 		}
 	}
