VK: Upload uniform buffers when possible instead of always re-allocating. This helps to reduce pressure on a descriptor cache

diff --git a/Engine/Source/Runtime/VulkanRHI/Private/VulkanUniformBuffer.cpp b/Engine/Source/Runtime/VulkanRHI/Private/VulkanUniformBuffer.cpp
--- a/Engine/Source/Runtime/VulkanRHI/Private/VulkanUniformBuffer.cpp
+++ b/Engine/Source/Runtime/VulkanRHI/Private/VulkanUniformBuffer.cpp
@@ -9,7 +9,7 @@
 #include "VulkanLLM.h"
 #include "ShaderParameterStruct.h"
 
-static int32 GVulkanAllowUniformUpload = 0;
+static int32 GVulkanAllowUniformUpload = 1;
 static FAutoConsoleVariableRef CVarVulkanAllowUniformUpload(
 	TEXT("r.Vulkan.AllowUniformUpload"),
 	GVulkanAllowUniformUpload,
@@ -76,30 +76,19 @@ static void UpdateUniformBufferHelper(FVulkanCommandListContext& Context, FVulka
 	}
 	else
 	{
-		ensure(CmdBuffer->IsOutsideRenderPass());
+		check(CmdBuffer->IsOutsideRenderPass());
+
 		VulkanRHI::FTempFrameAllocationBuffer::FTempAllocInfo LockInfo;
 		Context.GetTempFrameAllocationBuffer().Alloc(DataSize, 16, LockInfo);
 		FMemory::Memcpy(LockInfo.Data, Data, DataSize);
 		VkBufferCopy Region;
 		Region.size = DataSize;
-		Region.srcOffset = LockInfo.GetBindOffset();
+		Region.srcOffset = LockInfo.CurrentOffset + LockInfo.Allocation.Offset;
 		Region.dstOffset = VulkanUniformBuffer->GetOffset();
 		VkBuffer UBBuffer = VulkanUniformBuffer->Allocation.GetBufferHandle();
-		VkBuffer LockHandle = VulkanUniformBuffer->Allocation.GetBufferHandle();
-
-		bool bIsInRenderPass = CmdBuffer->IsInsideRenderPass();
-		bool bIsUniformBarrierAdded = CmdBuffer->IsUniformBufferBarrierAdded();
-		if(bIsInRenderPass || !bIsUniformBarrierAdded)
-		{
-			CmdBuffer->BeginUniformUpdateBarrier();
-		}
+		VkBuffer LockHandle = LockInfo.Allocation.GetBufferHandle();
 
 		VulkanRHI::vkCmdCopyBuffer(CmdBuffer->GetHandle(), LockHandle, UBBuffer, 1, &Region);
-
-		if(bIsInRenderPass) //when updating outside render passes, the EndUniformUpdateBarrier will be called from EndRenderPass
-		{
-			CmdBuffer->EndUniformUpdateBarrier();
-		}
 	}
 };
 
