VK: Use runtime descriptor usage statistic to adjust sizes of descriptor cache pools
Do not use descriptor cache for compute and for draws with volatile resources, instead fallback to a default descriptor manager which we reset often. Compute tend to use single use (volatile) resources which polutes descriptor cache

diff --git a/Engine/Source/Runtime/VulkanRHI/Private/VulkanCommandBuffer.cpp b/Engine/Source/Runtime/VulkanRHI/Private/VulkanCommandBuffer.cpp
--- a/Engine/Source/Runtime/VulkanRHI/Private/VulkanCommandBuffer.cpp
+++ b/Engine/Source/Runtime/VulkanRHI/Private/VulkanCommandBuffer.cpp
@@ -206,7 +206,7 @@ void FVulkanCmdBuffer::BeginRenderPass(const FVulkanRenderTargetLayout& Layout,
 	State = EState::IsInsideRenderPass;
 
 	// Acquire a descriptor pool set on a first render pass
-	if (!UseVulkanDescriptorCache() && CurrentDescriptorPoolSetContainer == nullptr)
+	if (CurrentDescriptorPoolSetContainer == nullptr)
 	{
 		AcquirePoolSetContainer();
 	}
@@ -316,7 +316,6 @@ void FVulkanCmdBuffer::Begin()
 
 void FVulkanCmdBuffer::AcquirePoolSetContainer()
 {
-	check(!UseVulkanDescriptorCache())
 	check(!CurrentDescriptorPoolSetContainer);
 	CurrentDescriptorPoolSetContainer = &Device->GetDescriptorPoolsManager().AcquirePoolSetContainer();
 	ensure(TypedDescriptorPoolSets.Num() == 0);
@@ -324,7 +323,6 @@ void FVulkanCmdBuffer::AcquirePoolSetContainer()
 
 bool FVulkanCmdBuffer::AcquirePoolSetAndDescriptorsIfNeeded(const class FVulkanDescriptorSetsLayout& Layout, bool bNeedDescriptors, VkDescriptorSet* OutDescriptors)
 {
-	check(!UseVulkanDescriptorCache())
 	//#todo-rco: This only happens when we call draws outside a render pass...
 	if (!CurrentDescriptorPoolSetContainer)
 	{
@@ -378,7 +376,7 @@ void FVulkanCmdBuffer::RefreshFenceStatus()
 #endif
 			++FenceSignaledCounter;
 
-			if (!UseVulkanDescriptorCache() && CurrentDescriptorPoolSetContainer)
+			if (CurrentDescriptorPoolSetContainer)
 			{
 				//#todo-rco: Reset here?
 				TypedDescriptorPoolSets.Reset();
diff --git a/Engine/Source/Runtime/VulkanRHI/Private/VulkanCommands.cpp b/Engine/Source/Runtime/VulkanRHI/Private/VulkanCommands.cpp
--- a/Engine/Source/Runtime/VulkanRHI/Private/VulkanCommands.cpp
+++ b/Engine/Source/Runtime/VulkanRHI/Private/VulkanCommands.cpp
@@ -135,7 +135,7 @@ void FVulkanCommandListContext::RHISetComputePipelineState(FRHIComputePipelineSt
 		}
 	}
 
-	if (!UseVulkanDescriptorCache() && CmdBuffer->CurrentDescriptorPoolSetContainer == nullptr)
+	if (CmdBuffer->CurrentDescriptorPoolSetContainer == nullptr)
 	{
 		CmdBuffer->CurrentDescriptorPoolSetContainer = &Device->GetDescriptorPoolsManager().AcquirePoolSetContainer();
 	}
diff --git a/Engine/Source/Runtime/VulkanRHI/Private/VulkanDescriptorSets.h b/Engine/Source/Runtime/VulkanRHI/Private/VulkanDescriptorSets.h
--- a/Engine/Source/Runtime/VulkanRHI/Private/VulkanDescriptorSets.h
+++ b/Engine/Source/Runtime/VulkanRHI/Private/VulkanDescriptorSets.h
@@ -1154,8 +1154,11 @@ protected:
 		SetWritten(DescriptorIndex);
 		WriteDescriptors[DescriptorIndex].pTexelBufferView = &View->View;
 		BufferViewReferences[DescriptorIndex] = View;
-		
-		if (UseVulkanDescriptorCache())
+		const bool bVolatile = View->bVolatile;
+
+		bHasVolatileResources|= bVolatile;
+				
+		if (!bVolatile && UseVulkanDescriptorCache())
 		{
 			bool bChanged = false;
 			FVulkanHashableDescriptorInfo& HashableInfo = HashableDescriptorInfos[DescriptorIndex];
@@ -1190,6 +1193,7 @@ protected:
 	FVulkanHashableDescriptorInfo* HashableDescriptorInfos;
 	mutable FVulkanDSetKey Key;
 	mutable bool bIsKeyDirty;
+	bool bHasVolatileResources = false;
 
 	uint32 SetupDescriptorWrites(const TArray<VkDescriptorType>& Types,
 		FVulkanHashableDescriptorInfo* InHashableDescriptorInfos,
@@ -1204,26 +1208,18 @@ protected:
 #if VULKAN_VALIDATE_DESCRIPTORS_WRITTEN
 	TArray<uint32, TInlineAllocator<2> > WrittenMask;
 	TArray<uint32, TInlineAllocator<2> > BaseWrittenMask;
-
+#endif
 	void CheckAllWritten();
 	void Reset();
 	void SetWritten(uint32 DescriptorIndex);
 	void SetWrittenBase(uint32 DescriptorIndex);
 	void InitWrittenMasks(uint32 NumDescriptorWrites);
-#else
-	void Reset(){}
-	void CheckAllWritten() {}
-	void SetWritten(uint32 DescriptorIndex) {}
-	void SetWrittenBase(uint32 DescriptorIndex){}
-	void InitWrittenMasks(uint32 NumDescriptorWrites){}
-#endif
-
 };
 
 class FVulkanGenericDescriptorPool : FNoncopyable
 {
 public:
-	FVulkanGenericDescriptorPool(FVulkanDevice* InDevice, uint32 InMaxDescriptorSets);
+	FVulkanGenericDescriptorPool(FVulkanDevice* InDevice, uint32 InMaxDescriptorSets, const float PoolSizes[VK_DESCRIPTOR_TYPE_RANGE_SIZE]);
 	~FVulkanGenericDescriptorPool();
 
 	FVulkanDevice* GetDevice() const
@@ -1243,6 +1239,8 @@ private:
 	FVulkanDevice* const Device;
 	const uint32 MaxDescriptorSets;
 	VkDescriptorPool DescriptorPool;
+	// information for debugging
+	uint32 PoolSizes[VK_DESCRIPTOR_TYPE_RANGE_SIZE];
 };
 
 class FVulkanDescriptorSetCache : FNoncopyable
@@ -1269,25 +1267,14 @@ private:
 	class FCachedPool : FNoncopyable
 	{
 	public:
-		FCachedPool(FVulkanDevice* InDevice, uint32 InMaxDescriptorSets)
-			: SetCapacity(FMath::RoundToZero(InMaxDescriptorSets * MaxAllocRatio))
-			, Pool(InDevice, InMaxDescriptorSets)
-			, RecentFrame(0)
-		{
-		}
+		FCachedPool(FVulkanDevice* InDevice, uint32 InMaxDescriptorSets, const float PoolSizesRatio[VK_DESCRIPTOR_TYPE_RANGE_SIZE]);
 
 		uint32 GetMaxDescriptorSets() const
 		{
 			return Pool.GetMaxDescriptorSets();
 		}
-
-		void Reset()
-		{
-			Pool.Reset();
-			SetsCache.Reset();
-			SetCache.Reset();
-		}
-
+			
+		void Reset();
 		bool CanGC() const;
 		float CalcAllocRatio() const;
 
@@ -1295,6 +1282,8 @@ private:
 		bool CreateDescriptorSets(const FVulkanDSetsKey& DSetsKey, const FVulkanDescriptorSetsLayout& SetsLayout,
 			TArray<FVulkanDescriptorSetWriter>& DSWriters, VkDescriptorSet* OutSets);
 
+		void CalcPoolSizesRatio(float PoolSizesRatio[VK_DESCRIPTOR_TYPE_RANGE_SIZE]);
+
 	private:
 		static const float MinAllocRatio;
 		static const float MaxAllocRatio;
@@ -1305,6 +1294,9 @@ private:
 		TMap<FVulkanDSetsKey, FSetsEntry> SetsCache;
 		TMap<FVulkanDSetKey, VkDescriptorSet> SetCache;
 		uint32 RecentFrame;
+	
+	public:
+		uint32 PoolSizesStatistic[VK_DESCRIPTOR_TYPE_RANGE_SIZE];
 	};
 
 private:
diff --git a/Engine/Source/Runtime/VulkanRHI/Private/VulkanDevice.cpp b/Engine/Source/Runtime/VulkanRHI/Private/VulkanDevice.cpp
--- a/Engine/Source/Runtime/VulkanRHI/Private/VulkanDevice.cpp
+++ b/Engine/Source/Runtime/VulkanRHI/Private/VulkanDevice.cpp
@@ -1163,11 +1163,9 @@ void FVulkanDevice::InitGPU(int32 DeviceIndex)
 	{
 		DescriptorSetCache = new FVulkanDescriptorSetCache(this);
 	}
-	else
-	{
-		DescriptorPoolsManager = new FVulkanDescriptorPoolsManager();
-		DescriptorPoolsManager->Init(this);
-	}
+	
+	DescriptorPoolsManager = new FVulkanDescriptorPoolsManager();
+	DescriptorPoolsManager->Init(this);
 
 	PipelineStateCache = new FVulkanPipelineStateCacheManager(this);
 
diff --git a/Engine/Source/Runtime/VulkanRHI/Private/VulkanPendingState.cpp b/Engine/Source/Runtime/VulkanRHI/Private/VulkanPendingState.cpp
--- a/Engine/Source/Runtime/VulkanRHI/Private/VulkanPendingState.cpp
+++ b/Engine/Source/Runtime/VulkanRHI/Private/VulkanPendingState.cpp
@@ -643,73 +643,39 @@ FAutoConsoleVariableRef CVarDSetCacheMaxPoolLookups(
 	ECVF_ReadOnly | ECVF_RenderThreadSafe
 );
 
-FVulkanGenericDescriptorPool::FVulkanGenericDescriptorPool(FVulkanDevice* InDevice, uint32 InMaxDescriptorSets)
+const float DefaultPoolSizes[VK_DESCRIPTOR_TYPE_RANGE_SIZE] = 
+{
+	2,		// VK_DESCRIPTOR_TYPE_SAMPLER
+	2,		// VK_DESCRIPTOR_TYPE_COMBINED_IMAGE_SAMPLER
+	2,		// VK_DESCRIPTOR_TYPE_SAMPLED_IMAGE
+	1/8.0,	// VK_DESCRIPTOR_TYPE_STORAGE_IMAGE
+	1/2.0,	// VK_DESCRIPTOR_TYPE_UNIFORM_TEXEL_BUFFER
+	1/8.0,	// VK_DESCRIPTOR_TYPE_STORAGE_TEXEL_BUFFER
+	1/8.0,	// VK_DESCRIPTOR_TYPE_UNIFORM_BUFFER
+	1/8.0,	// VK_DESCRIPTOR_TYPE_STORAGE_BUFFER
+	4,		// VK_DESCRIPTOR_TYPE_UNIFORM_BUFFER_DYNAMIC
+	1/8.0,	// VK_DESCRIPTOR_TYPE_STORAGE_BUFFER_DYNAMIC
+	1/8.0	// VK_DESCRIPTOR_TYPE_INPUT_ATTACHMENT
+};
+
+FVulkanGenericDescriptorPool::FVulkanGenericDescriptorPool(FVulkanDevice* InDevice, uint32 InMaxDescriptorSets, const float PoolSizesRatio[VK_DESCRIPTOR_TYPE_RANGE_SIZE])
 	: Device(InDevice)
 	, MaxDescriptorSets(InMaxDescriptorSets)
 	, DescriptorPool(VK_NULL_HANDLE)
 {
-	// Based on statisticts of runing BR_50v50.replay
-	// TODO Need a better solution
-	const uint32 LimitMaxUniformBuffers = MaxDescriptorSets * 2;
-	const uint32 LimitMaxSamplers = MaxDescriptorSets / 2;
-	const uint32 LimitMaxCombinedImageSamplers = MaxDescriptorSets * 3;
-	const uint32 LimitMaxUniformTexelBuffers = MaxDescriptorSets / 2;
-	const uint32 LimitMaxStorageTexelBuffers = MaxDescriptorSets / 4;
-	const uint32 LimitMaxStorageBuffers = MaxDescriptorSets / 4;
-	const uint32 LimitMaxStorageImage = MaxDescriptorSets / 4;
-	const uint32 LimitMaxSampledImages = MaxDescriptorSets * 2;
-	const uint32 LimitMaxInputAttachments = MaxDescriptorSets / 16;
-
 	TArray<VkDescriptorPoolSize> Types;
-	VkDescriptorPoolSize* Type = new(Types) VkDescriptorPoolSize;
-	FMemory::Memzero(*Type);
-	Type->type = VK_DESCRIPTOR_TYPE_UNIFORM_BUFFER;
-	Type->descriptorCount = LimitMaxUniformBuffers;
-
-	Type = new(Types) VkDescriptorPoolSize;
-	FMemory::Memzero(*Type);
-	Type->type = VK_DESCRIPTOR_TYPE_UNIFORM_BUFFER_DYNAMIC;
-	Type->descriptorCount = LimitMaxUniformBuffers;
-
-	Type = new(Types) VkDescriptorPoolSize;
-	FMemory::Memzero(*Type);
-	Type->type = VK_DESCRIPTOR_TYPE_SAMPLER;
-	Type->descriptorCount = LimitMaxSamplers;
-
-	Type = new(Types) VkDescriptorPoolSize;
-	FMemory::Memzero(*Type);
-	Type->type = VK_DESCRIPTOR_TYPE_COMBINED_IMAGE_SAMPLER;
-	Type->descriptorCount = LimitMaxCombinedImageSamplers;
-
-	Type = new(Types) VkDescriptorPoolSize;
-	FMemory::Memzero(*Type);
-	Type->type = VK_DESCRIPTOR_TYPE_SAMPLED_IMAGE;
-	Type->descriptorCount = LimitMaxSampledImages;
-
-	Type = new(Types) VkDescriptorPoolSize;
-	FMemory::Memzero(*Type);
-	Type->type = VK_DESCRIPTOR_TYPE_UNIFORM_TEXEL_BUFFER;
-	Type->descriptorCount = LimitMaxUniformTexelBuffers;
-
-	Type = new(Types) VkDescriptorPoolSize;
-	FMemory::Memzero(*Type);
-	Type->type = VK_DESCRIPTOR_TYPE_STORAGE_TEXEL_BUFFER;
-	Type->descriptorCount = LimitMaxStorageTexelBuffers;
-
-	Type = new(Types) VkDescriptorPoolSize;
-	FMemory::Memzero(*Type);
-	Type->type = VK_DESCRIPTOR_TYPE_STORAGE_BUFFER;
-	Type->descriptorCount = LimitMaxStorageBuffers;
-
-	Type = new(Types) VkDescriptorPoolSize;
-	FMemory::Memzero(*Type);
-	Type->type = VK_DESCRIPTOR_TYPE_STORAGE_IMAGE;
-	Type->descriptorCount = LimitMaxStorageImage;
-
-	Type = new(Types) VkDescriptorPoolSize;
-	FMemory::Memzero(*Type);
-	Type->type = VK_DESCRIPTOR_TYPE_INPUT_ATTACHMENT;
-	Type->descriptorCount = LimitMaxInputAttachments;
+	Types.Reserve(VK_DESCRIPTOR_TYPE_RANGE_SIZE);
+
+	for (uint32 i = 0; i < VK_DESCRIPTOR_TYPE_RANGE_SIZE; ++i)
+	{
+		float MinSize = FMath::Max(DefaultPoolSizes[i]*InMaxDescriptorSets, 4.0f);
+		PoolSizes[i] = (uint32)FMath::Max(PoolSizesRatio[i]*InMaxDescriptorSets, MinSize);
+
+		VkDescriptorPoolSize* Type = new(Types) VkDescriptorPoolSize;
+		FMemory::Memzero(*Type);
+		Type->type = (VkDescriptorType)i;
+		Type->descriptorCount = PoolSizes[i];
+	}
 
 	VkDescriptorPoolCreateInfo PoolInfo;
 	ZeroVulkanStruct(PoolInfo, VK_STRUCTURE_TYPE_DESCRIPTOR_POOL_CREATE_INFO);
@@ -761,7 +727,7 @@ FVulkanDescriptorSetCache::FVulkanDescriptorSetCache(FVulkanDevice* InDevice)
 	, PoolAllocRatio(0.0f)
 {
 	constexpr uint32 ProbePoolMaxNumSets = 128; // Used for initial estimation of the Allocation Ratio
-	CachedPools.Add(MakeUnique<FCachedPool>(Device, ProbePoolMaxNumSets));
+	CachedPools.Add(MakeUnique<FCachedPool>(Device, ProbePoolMaxNumSets, DefaultPoolSizes));
 }
 
 FVulkanDescriptorSetCache::~FVulkanDescriptorSetCache()
@@ -782,20 +748,24 @@ void FVulkanDescriptorSetCache::AddCachedPool()
 	{
 		constexpr float MinErrorTolerance = -0.10f;
 		constexpr float MaxErrorTolerance = 0.50f;
-		const float Error = ((static_cast<float>(FreePool->GetMaxDescriptorSets()) -
-			static_cast<float>(MaxDescriptorSets)) / static_cast<float>(MaxDescriptorSets));
+		const float Error = ((static_cast<float>(FreePool->GetMaxDescriptorSets()) - static_cast<float>(MaxDescriptorSets)) / static_cast<float>(MaxDescriptorSets));
 		if ((Error >= MinErrorTolerance) && (Error <= MaxErrorTolerance))
 		{
 			FreePool->Reset();
 			CachedPools.EmplaceAt(0, MoveTemp(FreePool));
 			return;
 		}
+
 		// Don't write 'error' as it confuses reporting; it's a perf warning more than an actual error
-		UE_LOG(LogVulkanRHI, Display, TEXT("FVulkanDescriptorSetCache::AddCachedPool() MaxDescriptorSets Delta/Err: %f. Tolerance: [%f..%f]."),
-			static_cast<double>(Error), static_cast<double>(MinErrorTolerance), static_cast<double>(MaxErrorTolerance));
+		UE_LOG(LogVulkanRHI, Display, TEXT("FVulkanDescriptorSetCache::AddCachedPool() MaxDescriptorSets Delta/Err: %f. Tolerance: [%f..%f]."),	Error, MinErrorTolerance, MaxErrorTolerance);
 		FreePool.Reset();
 	}
-	CachedPools.EmplaceAt(0, MakeUnique<FCachedPool>(Device, MaxDescriptorSets));
+	
+	// use current pool sizes statistic for a new pool
+	float PoolSizesRatio[VK_DESCRIPTOR_TYPE_RANGE_SIZE];
+	CachedPools[0]->CalcPoolSizesRatio(PoolSizesRatio);
+
+	CachedPools.EmplaceAt(0, MakeUnique<FCachedPool>(Device, MaxDescriptorSets, PoolSizesRatio));
 }
 
 void FVulkanDescriptorSetCache::GetDescriptorSets(const FVulkanDSetsKey& DSetsKey, const FVulkanDescriptorSetsLayout& SetsLayout,
@@ -839,6 +809,14 @@ void FVulkanDescriptorSetCache::GC()
 const float FVulkanDescriptorSetCache::FCachedPool::MinAllocRatio = 0.5f;
 const float FVulkanDescriptorSetCache::FCachedPool::MaxAllocRatio = 16.0f;
 
+FVulkanDescriptorSetCache::FCachedPool::FCachedPool(FVulkanDevice* InDevice, uint32 InMaxDescriptorSets, const float PoolSizes[VK_DESCRIPTOR_TYPE_RANGE_SIZE])
+	: SetCapacity(FMath::RoundToZero(InMaxDescriptorSets * MaxAllocRatio))
+	, Pool(InDevice, InMaxDescriptorSets, PoolSizes)
+	, RecentFrame(0)
+{
+	FMemory::Memzero(PoolSizesStatistic, sizeof(PoolSizesStatistic));
+}
+
 bool FVulkanDescriptorSetCache::FCachedPool::FindDescriptorSets(const FVulkanDSetsKey& DSetsKey, VkDescriptorSet* OutSets)
 {
 	FSetsEntry* SetsEntry = SetsCache.Find(DSetsKey);
@@ -853,8 +831,11 @@ bool FVulkanDescriptorSetCache::FCachedPool::FindDescriptorSets(const FVulkanDSe
 	return true;
 }
 
-bool FVulkanDescriptorSetCache::FCachedPool::CreateDescriptorSets(const FVulkanDSetsKey& DSetsKey, const FVulkanDescriptorSetsLayout& SetsLayout,
-	TArray<FVulkanDescriptorSetWriter>& DSWriters, VkDescriptorSet* OutSets)
+bool FVulkanDescriptorSetCache::FCachedPool::CreateDescriptorSets(
+	const FVulkanDSetsKey& DSetsKey, 
+	const FVulkanDescriptorSetsLayout& SetsLayout,
+	TArray<FVulkanDescriptorSetWriter>& DSWriters, 
+	VkDescriptorSet* OutSets)
 {
 	FSetsEntry NewSetEntry{};
 
@@ -885,6 +866,12 @@ bool FVulkanDescriptorSetCache::FCachedPool::CreateDescriptorSets(const FVulkanD
 
 		DSWriter.SetDescriptorSet(NewSetEntry.Sets[Index]);
 
+		const uint32* LayoutTypes = SetsLayout.GetLayoutTypes();
+		for (int32 i = 0; i < VK_DESCRIPTOR_TYPE_RANGE_SIZE; ++i)
+		{
+			PoolSizesStatistic[i]+= LayoutTypes[i];
+		}
+
 #if VULKAN_ENABLE_AGGRESSIVE_STATS
 		INC_DWORD_STAT_BY(STAT_VulkanNumUpdateDescriptors, DSWriter.GetNumWrites());
 		INC_DWORD_STAT(STAT_VulkanNumDescSets);
@@ -902,6 +889,14 @@ bool FVulkanDescriptorSetCache::FCachedPool::CreateDescriptorSets(const FVulkanD
 	return true;
 }
 
+void FVulkanDescriptorSetCache::FCachedPool::Reset()
+{
+	Pool.Reset();
+	SetsCache.Reset();
+	SetCache.Reset();
+	FMemory::Memzero(PoolSizesStatistic, sizeof(PoolSizesStatistic));
+}
+
 bool FVulkanDescriptorSetCache::FCachedPool::CanGC() const
 {
 	constexpr uint32 FramesBeforeGC = NUM_FRAMES_TO_WAIT_BEFORE_RELEASING_TO_OS;
@@ -913,12 +908,21 @@ float FVulkanDescriptorSetCache::FCachedPool::CalcAllocRatio() const
 	float AllocRatio = (static_cast<float>(SetCache.Num()) / static_cast<float>(Pool.GetMaxDescriptorSets()));
 	if (AllocRatio < MinAllocRatio)
 	{
-		UE_LOG(LogVulkanRHI, Warning, TEXT("FVulkanDescriptorSetCache::FCachedPool::CalcAllocRatio() Pool Allocation Ratio is too low: %f. Using: %f."),
-			static_cast<double>(AllocRatio), static_cast<double>(MinAllocRatio));
+		UE_LOG(LogVulkanRHI, Warning, TEXT("FVulkanDescriptorSetCache::FCachedPool::CalcAllocRatio() Pool Allocation Ratio is too low: %f. Using: %f."), AllocRatio, MinAllocRatio);
 		AllocRatio = MinAllocRatio;
 	}
 	return AllocRatio;
 }
+
+void FVulkanDescriptorSetCache::FCachedPool::CalcPoolSizesRatio(float PoolSizesRatio[VK_DESCRIPTOR_TYPE_RANGE_SIZE])
+{
+	int32 NumSets = FMath::Max(SetCache.Num(), 1);
+	for (uint32 i = 0; i < VK_DESCRIPTOR_TYPE_RANGE_SIZE; ++i)
+	{
+		PoolSizesRatio[i] = PoolSizesStatistic[i]/(float)NumSets;
+	}
+}
+
 bool FVulkanPendingGfxState::SetGfxPipeline(FVulkanRHIGraphicsPipelineState* InGfxPipeline, bool bForceReset)
 {
 	bool bChanged = bForceReset;
diff --git a/Engine/Source/Runtime/VulkanRHI/Private/VulkanPipelineState.cpp b/Engine/Source/Runtime/VulkanRHI/Private/VulkanPipelineState.cpp
--- a/Engine/Source/Runtime/VulkanRHI/Private/VulkanPipelineState.cpp
+++ b/Engine/Source/Runtime/VulkanRHI/Private/VulkanPipelineState.cpp
@@ -168,34 +168,30 @@ bool FVulkanComputePipelineDescriptorState::InternalUpdateDescriptorSets(FVulkan
 		PackedUniformBuffersDirty = 0;
 	}
 
-	if (UseVulkanDescriptorCache())
+	// We are not using UseVulkanDescriptorCache() for compute pipelines
+	// Compute tend to use volatile resources that polute descriptor cache
+
+	if (!CmdBuffer->AcquirePoolSetAndDescriptorsIfNeeded(*DescriptorSetsLayout, true, DescriptorSetHandles.GetData()))
 	{
-		Device->GetDescriptorSetCache().GetDescriptorSets(GetDSetsKey(), *DescriptorSetsLayout, DSWriter, DescriptorSetHandles.GetData());
+		return false;
 	}
-	else
-	{
-		if (!CmdBuffer->AcquirePoolSetAndDescriptorsIfNeeded(*DescriptorSetsLayout, true, DescriptorSetHandles.GetData()))
-		{
-			return false;
-		}
 
-		const VkDescriptorSet DescriptorSet = DescriptorSetHandles[0];
-		DSWriter[0].SetDescriptorSet(DescriptorSet);
+	const VkDescriptorSet DescriptorSet = DescriptorSetHandles[0];
+	DSWriter[0].SetDescriptorSet(DescriptorSet);
 #if VULKAN_VALIDATE_DESCRIPTORS_WRITTEN
-		for(FVulkanDescriptorSetWriter& Writer : DSWriter)
-		{
-			Writer.CheckAllWritten();
-		}
+	for(FVulkanDescriptorSetWriter& Writer : DSWriter)
+	{
+		Writer.CheckAllWritten();
+	}
 #endif
 
-		{
-	#if VULKAN_ENABLE_AGGRESSIVE_STATS
-			INC_DWORD_STAT_BY(STAT_VulkanNumUpdateDescriptors, DSWriteContainer.DescriptorWrites.Num());
-			INC_DWORD_STAT(STAT_VulkanNumDescSets);
-			SCOPE_CYCLE_COUNTER(STAT_VulkanVkUpdateDS);
-	#endif
-			VulkanRHI::vkUpdateDescriptorSets(Device->GetInstanceHandle(), DSWriteContainer.DescriptorWrites.Num(), DSWriteContainer.DescriptorWrites.GetData(), 0, nullptr);
-		}
+	{
+#if VULKAN_ENABLE_AGGRESSIVE_STATS
+		INC_DWORD_STAT_BY(STAT_VulkanNumUpdateDescriptors, DSWriteContainer.DescriptorWrites.Num());
+		INC_DWORD_STAT(STAT_VulkanNumDescSets);
+		SCOPE_CYCLE_COUNTER(STAT_VulkanVkUpdateDS);
+#endif
+		VulkanRHI::vkUpdateDescriptorSets(Device->GetInstanceHandle(), DSWriteContainer.DescriptorWrites.Num(), DSWriteContainer.DescriptorWrites.GetData(), 0, nullptr);
 	}
 
 	return true;
@@ -288,7 +284,7 @@ bool FVulkanGraphicsPipelineDescriptorState::InternalUpdateDescriptorSets(FVulka
 		}
 	}
 
-	if (UseVulkanDescriptorCache())
+	if (UseVulkanDescriptorCache() && !HasVolatileResources())
 	{
 		if (bIsResourcesDirty)
 		{
@@ -372,35 +368,34 @@ template bool FVulkanComputePipelineDescriptorState::InternalUpdateDescriptorSet
 template bool FVulkanComputePipelineDescriptorState::InternalUpdateDescriptorSets<false>(FVulkanCommandListContext* CmdListContext, FVulkanCmdBuffer* CmdBuffer);
 
 
-#if VULKAN_VALIDATE_DESCRIPTORS_WRITTEN
-
-
 void FVulkanDescriptorSetWriter::CheckAllWritten()
 {
-auto GetVkDescriptorTypeString = [](VkDescriptorType Type)
-{
-	switch (Type)
+#if VULKAN_VALIDATE_DESCRIPTORS_WRITTEN
+	auto GetVkDescriptorTypeString = [](VkDescriptorType Type)
 	{
-		// + 19 to skip "VK_DESCRIPTOR_TYPE_"
-#define VKSWITCHCASE(x)	case x: return FString(&TEXT(#x)[19]);
-		VKSWITCHCASE(VK_DESCRIPTOR_TYPE_SAMPLER)
-			VKSWITCHCASE(VK_DESCRIPTOR_TYPE_COMBINED_IMAGE_SAMPLER)
-			VKSWITCHCASE(VK_DESCRIPTOR_TYPE_SAMPLED_IMAGE)
-			VKSWITCHCASE(VK_DESCRIPTOR_TYPE_STORAGE_IMAGE)
-			VKSWITCHCASE(VK_DESCRIPTOR_TYPE_UNIFORM_TEXEL_BUFFER)
-			VKSWITCHCASE(VK_DESCRIPTOR_TYPE_STORAGE_TEXEL_BUFFER)
-			VKSWITCHCASE(VK_DESCRIPTOR_TYPE_UNIFORM_BUFFER)
-			VKSWITCHCASE(VK_DESCRIPTOR_TYPE_STORAGE_BUFFER)
-			VKSWITCHCASE(VK_DESCRIPTOR_TYPE_UNIFORM_BUFFER_DYNAMIC)
-			VKSWITCHCASE(VK_DESCRIPTOR_TYPE_STORAGE_BUFFER_DYNAMIC)
-			VKSWITCHCASE(VK_DESCRIPTOR_TYPE_INPUT_ATTACHMENT)
-#undef VKSWITCHCASE
-	default:
-		break;
-	}
+		switch (Type)
+		{
+			// + 19 to skip "VK_DESCRIPTOR_TYPE_"
+	#define VKSWITCHCASE(x)	case x: return FString(&TEXT(#x)[19]);
+			VKSWITCHCASE(VK_DESCRIPTOR_TYPE_SAMPLER)
+				VKSWITCHCASE(VK_DESCRIPTOR_TYPE_COMBINED_IMAGE_SAMPLER)
+				VKSWITCHCASE(VK_DESCRIPTOR_TYPE_SAMPLED_IMAGE)
+				VKSWITCHCASE(VK_DESCRIPTOR_TYPE_STORAGE_IMAGE)
+				VKSWITCHCASE(VK_DESCRIPTOR_TYPE_UNIFORM_TEXEL_BUFFER)
+				VKSWITCHCASE(VK_DESCRIPTOR_TYPE_STORAGE_TEXEL_BUFFER)
+				VKSWITCHCASE(VK_DESCRIPTOR_TYPE_UNIFORM_BUFFER)
+				VKSWITCHCASE(VK_DESCRIPTOR_TYPE_STORAGE_BUFFER)
+				VKSWITCHCASE(VK_DESCRIPTOR_TYPE_UNIFORM_BUFFER_DYNAMIC)
+				VKSWITCHCASE(VK_DESCRIPTOR_TYPE_STORAGE_BUFFER_DYNAMIC)
+				VKSWITCHCASE(VK_DESCRIPTOR_TYPE_INPUT_ATTACHMENT)
+	#undef VKSWITCHCASE
+		default:
+			break;
+		}
+
+		return FString::Printf(TEXT("Unknown VkDescriptorType %d"), (int32)Type);
+	};
 
-	return FString::Printf(TEXT("Unknown VkDescriptorType %d"), (int32)Type);
-};
 	const uint32 Writes = NumWrites;
 	if (Writes == 0)
 		return;
@@ -441,32 +436,41 @@ auto GetVkDescriptorTypeString = [](VkDescriptorType Type)
 		UE_LOG(LogVulkanRHI, Warning, TEXT("Not All descriptors where filled out. this can/will cause a driver crash\n%s\n"), *Descriptors);
 		ensureMsgf(false, TEXT("Not All descriptors where filled out. this can/will cause a driver crash\n%s\n"), *Descriptors);
 	}
+#endif
 }
 
 void FVulkanDescriptorSetWriter::Reset()
 {
+	bHasVolatileResources = false;
+
+#if VULKAN_VALIDATE_DESCRIPTORS_WRITTEN
 	WrittenMask = BaseWrittenMask;
+#endif
 }
 void FVulkanDescriptorSetWriter::SetWritten(uint32 DescriptorIndex)
 {
+#if VULKAN_VALIDATE_DESCRIPTORS_WRITTEN
 	uint32 Index = DescriptorIndex / 32;
 	uint32 Mask = DescriptorIndex % 32;
 	WrittenMask[Index] |= (1<<Mask);
+#endif
 }
 void FVulkanDescriptorSetWriter::SetWrittenBase(uint32 DescriptorIndex)
 {
+#if VULKAN_VALIDATE_DESCRIPTORS_WRITTEN	
 	uint32 Index = DescriptorIndex / 32;
 	uint32 Mask = DescriptorIndex % 32;
 	BaseWrittenMask[Index] |= (1<<Mask);
+#endif
 }
 
 void FVulkanDescriptorSetWriter::InitWrittenMasks(uint32 NumDescriptorWrites)
 {
+#if VULKAN_VALIDATE_DESCRIPTORS_WRITTEN	
 	uint32 Size = (NumDescriptorWrites + 31) / 32;
 	WrittenMask.Empty(Size);
 	WrittenMask.SetNumZeroed(Size);
 	BaseWrittenMask.Empty(Size);
 	BaseWrittenMask.SetNumZeroed(Size);
-}
 #endif
-
+}
diff --git a/Engine/Source/Runtime/VulkanRHI/Private/VulkanPipelineState.h b/Engine/Source/Runtime/VulkanRHI/Private/VulkanPipelineState.h
--- a/Engine/Source/Runtime/VulkanRHI/Private/VulkanPipelineState.h
+++ b/Engine/Source/Runtime/VulkanRHI/Private/VulkanPipelineState.h
@@ -40,6 +40,18 @@ public:
 		}
 		return DSetsKey;
 	}
+	
+	bool HasVolatileResources() const
+	{
+		for (const FVulkanDescriptorSetWriter& Writer : DSWriter)
+		{
+			if (Writer.bHasVolatileResources)
+			{
+				return true;
+			}
+		}
+		return false;
+	}
 
 	inline void MarkDirty(bool bDirty)
 	{
diff --git a/Engine/Source/Runtime/VulkanRHI/Private/VulkanRHI.cpp b/Engine/Source/Runtime/VulkanRHI/Private/VulkanRHI.cpp
--- a/Engine/Source/Runtime/VulkanRHI/Private/VulkanRHI.cpp
+++ b/Engine/Source/Runtime/VulkanRHI/Private/VulkanRHI.cpp
@@ -981,10 +981,7 @@ void FVulkanCommandListContext::RHIEndFrame()
 	{
 		Device->GetDescriptorSetCache().GC();
 	}
-	else
-	{
-		Device->GetDescriptorPoolsManager().GC();
-	}
+	Device->GetDescriptorPoolsManager().GC();
 
 	Device->ReleaseUnusedOcclusionQueryPools();
 
@@ -1490,7 +1487,8 @@ void FVulkanBufferView::Create(VkFormat Format, FVulkanResourceMultiBuffer* Buff
 
 	VERIFYVULKANRESULT(VulkanRHI::vkCreateBufferView(GetParent()->GetInstanceHandle(), &ViewInfo, VULKAN_CPU_ALLOCATOR, &View));
 	
-	if (UseVulkanDescriptorCache())
+	bVolatile = Buffer->IsVolatile();
+	if (!bVolatile && UseVulkanDescriptorCache())
 	{
 		ViewId = ++GVulkanBufferViewHandleIdCounter;
 	}
@@ -1506,6 +1504,7 @@ void FVulkanBufferView::Destroy()
 		Device->GetDeferredDeletionQueue().EnqueueResource(FDeferredDeletionQueue2::EType::BufferView, View);
 		View = VK_NULL_HANDLE;
 		ViewId = 0;
+		bVolatile = false;
 	}
 }
 
diff --git a/Engine/Source/Runtime/VulkanRHI/Private/VulkanResources.h b/Engine/Source/Runtime/VulkanRHI/Private/VulkanResources.h
--- a/Engine/Source/Runtime/VulkanRHI/Private/VulkanResources.h
+++ b/Engine/Source/Runtime/VulkanRHI/Private/VulkanResources.h
@@ -903,6 +903,7 @@ struct FVulkanBufferView : public FRHIResource, public VulkanRHI::FDeviceChild
 		, Flags(0)
 		, Offset(0)
 		, Size(0)
+		, bVolatile(false)
 	{
 	}
 
@@ -920,6 +921,8 @@ struct FVulkanBufferView : public FRHIResource, public VulkanRHI::FDeviceChild
 	VkFlags Flags;
 	uint32 Offset;
 	uint32 Size;
+	// Whether source buffer is volatile
+	bool bVolatile;
 };
 
 struct FVulkanRingBuffer : public FVulkanEvictable, public VulkanRHI::FDeviceChild
