VKRT:

Add ray tracing shader types
Add initial ray tracing pipeline creation
Add acceleration structure descriptor type
Add basic ray tracing pipeline, occlusion support initially
Strip reflection from ray tracing shader spirv to deal with validation warnings
Don't use an array of vk descriptor types directly, this breaks with non-contiguous enum values which are common in extensions. Using a TMap from descriptor type to values instead.
Don't store vk types in the serialized shader header, translate to and from our own internal enum types to avoid similar non-contiguous value issues.
Re-enabled ray tracing compilation on windows desktop, explicitly disable runtime support using GRHISupportsRayTracing

Todo:
We need to deal with the lack of a 1:1 mapping between shader stages and frequencies for ray tracing hit groups. This is a one to many mapping which doesn't work with how most of the Vulkan RHI is authored. For now I'm assuming a hitgroup maps to a CHS.
Ray tracing shader descriptor allocation and pending state support. While we are serializing the descriptors for ray tracing shaders, we're not allocating or updating them yet which is the last large chunk needed to trace rays.
Fix spirv binary version mismatch validation error. This basically needs a local dxc modification which explicitly sets the spirv binary version to work around incorrect assuptions in dxc which is causing validation failures.
Add another local dxc modification to allow for arbirary struct size support using VK_EXT_scalar_block_layout for ray tracing shaders.

diff --git a/Engine/Shaders/Public/Platform/Vulkan/VulkanCommon.ush b/Engine/Shaders/Public/Platform/Vulkan/VulkanCommon.ush
--- a/Engine/Shaders/Public/Platform/Vulkan/VulkanCommon.ush
+++ b/Engine/Shaders/Public/Platform/Vulkan/VulkanCommon.ush
@@ -5,6 +5,6 @@
 =============================================================================*/
 
 // Update this GUID to improve shader recompilation for Vulkan only shaders
-// GUID = E28F6421-6CC3-4D01-BCED-D8442AEF87A9
+// GUID = 58EB247E-D403-451E-B31F-4B24D38898A0
 
 #pragma once
diff --git a/Engine/Source/Developer/ShaderCompilerCommon/Private/HlslccHeaderWriter.cpp b/Engine/Source/Developer/ShaderCompilerCommon/Private/HlslccHeaderWriter.cpp
--- a/Engine/Source/Developer/ShaderCompilerCommon/Private/HlslccHeaderWriter.cpp
+++ b/Engine/Source/Developer/ShaderCompilerCommon/Private/HlslccHeaderWriter.cpp
@@ -495,6 +495,11 @@ namespace CrossCompiler
 		MetaDataPrintf(Strings.NumThreads, TEXT("%u, %u, %u"), NumThreadsX, NumThreadsY, NumThreadsZ);
 	}
 
+	void FHlslccHeaderWriter::WriteAccelerationStructures(const TCHAR* ResourceName, uint32 BindingIndex)
+	{
+		MetaDataPrintf(Strings.AccelerationStructures, TEXT("%u:%s"), BindingIndex, ResourceName);
+	}
+
 	void FHlslccHeaderWriter::WriteSideTable(const TCHAR* ResourceName, uint32 SideTableIndex)
 	{
 		MetaDataPrintf(Strings.SideTable, TEXT("%s(%d)"), ResourceName, SideTableIndex);
@@ -548,6 +553,7 @@ namespace CrossCompiler
 		PrintAttributes(TEXT("UAVs"), Strings.UAVs);
 		PrintAttributes(TEXT("SamplerStates"), Strings.SamplerStates);
 		PrintAttributes(TEXT("NumThreads"), Strings.NumThreads);
+		PrintAttributes(TEXT("AccelerationStructures"), Strings.AccelerationStructures);
 		PrintAttributes(TEXT("SideTable"), Strings.SideTable);
 		PrintAttributes(TEXT("ArgumentBuffers"), Strings.ArgumentBuffers);
 
diff --git a/Engine/Source/Developer/ShaderCompilerCommon/Private/ShaderCompilerCommon.cpp b/Engine/Source/Developer/ShaderCompilerCommon/Private/ShaderCompilerCommon.cpp
--- a/Engine/Source/Developer/ShaderCompilerCommon/Private/ShaderCompilerCommon.cpp
+++ b/Engine/Source/Developer/ShaderCompilerCommon/Private/ShaderCompilerCommon.cpp
@@ -2002,6 +2002,7 @@ namespace CrossCompiler
 		DEF_PREFIX_STR(Samplers);
 		DEF_PREFIX_STR(UAVs);
 		DEF_PREFIX_STR(SamplerStates);
+		DEF_PREFIX_STR(AccelerationStructures);
 		DEF_PREFIX_STR(NumThreads);
 #undef DEF_PREFIX_STR
 
@@ -2513,6 +2514,45 @@ namespace CrossCompiler
 			}
 		}
 
+		if (FCStringAnsi::Strncmp(ShaderSource, AccelerationStructuresPrefix, AccelerationStructuresPrefixLen) == 0)
+		{
+			ShaderSource += AccelerationStructuresPrefixLen;
+
+			while (*ShaderSource && *ShaderSource != '\n')
+			{
+				FAccelerationStructure AccelerationStructure;
+
+				if (!ParseIntegerNumber(ShaderSource, AccelerationStructure.Offset))
+				{
+					return false;
+				}
+
+				if (!Match(ShaderSource, ':'))
+				{
+					return false;
+				}
+
+				if (!ParseIdentifier(ShaderSource, AccelerationStructure.Name))
+				{
+					return false;
+				}
+
+				AccelerationStructures.Add(AccelerationStructure);
+
+				if (Match(ShaderSource, '\n'))
+				{
+					break;
+				}
+
+				if (Match(ShaderSource, ','))
+				{
+					continue;
+				}
+
+				return false;
+			}
+		}
+
 		if (FCStringAnsi::Strncmp(ShaderSource, NumThreadsPrefix, NumThreadsPrefixLen) == 0)
 		{
 			ShaderSource += NumThreadsPrefixLen;
diff --git a/Engine/Source/Developer/ShaderCompilerCommon/Public/HlslccHeaderWriter.h b/Engine/Source/Developer/ShaderCompilerCommon/Public/HlslccHeaderWriter.h
--- a/Engine/Source/Developer/ShaderCompilerCommon/Public/HlslccHeaderWriter.h
+++ b/Engine/Source/Developer/ShaderCompilerCommon/Public/HlslccHeaderWriter.h
@@ -35,6 +35,7 @@ namespace CrossCompiler
 		void WriteUAV(const TCHAR* ResourceName, uint32 BindingIndex, uint32 Count = 1);
 		void WriteSamplerState(const TCHAR* ResourceName, uint32 BindingIndex);
 		void WriteNumThreads(uint32 NumThreadsX, uint32 NumThreadsY, uint32 NumThreadsZ);
+		void WriteAccelerationStructures(const TCHAR* ResourceName, uint32 BindingIndex);
 
 		void WriteSideTable(const TCHAR* ResourceName, uint32 SideTableIndex);
 		void WriteArgumentBuffers(uint32 BindingIndex, const TArray<uint32>& ResourceIndices);
@@ -65,6 +66,7 @@ namespace CrossCompiler
 			FString ExternalTextures; // External texture resources (Vulkan ES3.1 profile only)
 			FString SideTable; // Side table for additional indices, e.,g. "spvBufferSizeConstants(31)" (Metal only)
 			FString ArgumentBuffers; // Indirect argument buffers (Metal only)
+			FString AccelerationStructures;
 		};
 
 	private:
diff --git a/Engine/Source/Developer/ShaderCompilerCommon/Public/ShaderCompilerCommon.h b/Engine/Source/Developer/ShaderCompilerCommon/Public/ShaderCompilerCommon.h
--- a/Engine/Source/Developer/ShaderCompilerCommon/Public/ShaderCompilerCommon.h
+++ b/Engine/Source/Developer/ShaderCompilerCommon/Public/ShaderCompilerCommon.h
@@ -208,6 +208,12 @@ namespace CrossCompiler
 			int32 Count;
 		};
 
+		struct FAccelerationStructure
+		{
+			FString Name;
+			int32 Offset = 0;
+		};
+
 		FString Name;
 		TArray<FInOut> Inputs;
 		TArray<FInOut> Outputs;
@@ -220,6 +226,7 @@ namespace CrossCompiler
 		TArray<FSampler> Samplers;
 		TArray<FUAV> UAVs;
 		TArray<FAttribute> SamplerStates;
+		TArray<FAccelerationStructure> AccelerationStructures;
 		uint32 NumThreads[3];
 
 		static bool ReadInOut(const ANSICHAR*& ShaderSource, TArray<FInOut>& OutAttributes);
diff --git a/Engine/Source/Developer/VulkanShaderFormat/Private/VulkanGlslangBridge.cpp b/Engine/Source/Developer/VulkanShaderFormat/Private/VulkanGlslangBridge.cpp
--- a/Engine/Source/Developer/VulkanShaderFormat/Private/VulkanGlslangBridge.cpp
+++ b/Engine/Source/Developer/VulkanShaderFormat/Private/VulkanGlslangBridge.cpp
@@ -179,8 +179,7 @@ void PatchSpirvReflectionEntries(FVulkanSpirv& Spirv)
 	TMap<uint32, uint32> TypePointerUniforms;
 	TMap<uint32, uint32> VariableUniformTypes;
 
-	bool bDone = false;
-	for (FSpirvConstIterator Iter = Spirv.cbegin(); Iter != Spirv.cend() && !bDone; ++Iter)
+	for (FSpirvConstIterator Iter = Spirv.cbegin(); Iter != Spirv.cend(); ++Iter)
 	{
 		switch (Iter.Opcode())
 		{
@@ -236,14 +235,15 @@ void PatchSpirvReflectionEntries(FVulkanSpirv& Spirv)
 			uint32 Type = Iter.Operand(1);
 			uint32 Id = Iter.Operand(2);
 			SpvStorageClass Storage = Iter.OperandAs<SpvStorageClass>(3);
-			if (Storage == SpvStorageClassUniform || Storage == SpvStorageClassUniformConstant)
+			if (Storage == SpvStorageClassUniform || 
+				Storage == SpvStorageClassUniformConstant || 
+				Storage == SpvStorageClassStorageBuffer)
 			{
 				VariableUniformTypes.Add(Id, Type);
 			}
 		}
 			break;
 		case SpvOpFunction:
-			bDone = true;
 			break;
 		default:
 			break;
diff --git a/Engine/Source/Developer/VulkanShaderFormat/Private/VulkanShaderCompiler.cpp b/Engine/Source/Developer/VulkanShaderFormat/Private/VulkanShaderCompiler.cpp
--- a/Engine/Source/Developer/VulkanShaderFormat/Private/VulkanShaderCompiler.cpp
+++ b/Engine/Source/Developer/VulkanShaderFormat/Private/VulkanShaderCompiler.cpp
@@ -366,6 +366,7 @@ public:
 		NumSamplers = 0;
 		NumUniformBuffers = 0;
 		NumUAVs = 0;
+		NumAccelerationStructures = 0;
 		bHasRegularUniformBuffers = 0;
 	}
 };
@@ -473,7 +474,7 @@ static int32 AddGlobal(FOLDVulkanCodeHeader& OLDHeader,
 	}
 	else
 	{
-		int32 GlobalDescriptorTypeIndex = OutHeader.GlobalDescriptorTypes.Add(DescriptorType);
+		int32 GlobalDescriptorTypeIndex = OutHeader.GlobalDescriptorTypes.Add(DescriptorTypeToBinding(DescriptorType));
 		GlobalInfo.TypeIndex = GlobalDescriptorTypeIndex;
 		check(GetCombinedSamplerStateAlias(ParameterName, DescriptorType, BindingTable, CCHeader, GlobalNames) == UINT16_MAX);
 		GlobalInfo.CombinedSamplerStateAliasIndex = UINT16_MAX;
@@ -996,7 +997,7 @@ static void ConvertToNEWHeader(FOLDVulkanCodeHeader& OLDHeader,
 				VkDescriptorType DescriptorType = VK_DESCRIPTOR_TYPE_INPUT_ATTACHMENT;
 				GlobalInfo.OriginalBindingIndex = Entry->Binding;
 				OutHeader.GlobalSpirvInfos[HeaderGlobalIndex] = FVulkanShaderHeader::FSpirvInfo(Entry->WordDescriptorSetIndex, Entry->WordBindingIndex);
-				int32 GlobalDescriptorTypeIndex = OutHeader.GlobalDescriptorTypes.Add(DescriptorType);
+				int32 GlobalDescriptorTypeIndex = OutHeader.GlobalDescriptorTypes.Add(DescriptorTypeToBinding(DescriptorType));
 				GlobalInfo.TypeIndex = GlobalDescriptorTypeIndex;
 				GlobalInfo.CombinedSamplerStateAliasIndex = UINT16_MAX;
 #if VULKAN_ENABLE_BINDING_DEBUG_NAMES
@@ -1408,6 +1409,26 @@ static void BuildShaderOutput(
 			);
 	}
 
+	for (auto& AccelerationStructure : CCHeader.AccelerationStructures)
+	{
+		int32 VulkanBindingIndex = Spirv.FindBinding(AccelerationStructure.Name);
+		check(VulkanBindingIndex != -1);
+
+		ParameterMap.AddParameterAllocation(
+			*AccelerationStructure.Name,
+			AccelerationStructure.Offset,
+			VulkanBindingIndex,
+			1,
+			EShaderParameterType::SRV
+		);
+		NEWEntryTypes.Add(AccelerationStructure.Name, FVulkanShaderHeader::Global);
+
+		OLDHeader.SerializedBindings.NumAccelerationStructures = FMath::Max<uint8>(
+			OLDHeader.SerializedBindings.NumAccelerationStructures,
+			AccelerationStructure.Offset + 1
+			);
+	}
+
 	// Lats make sure that there is some type of name visible
 	OLDHeader.ShaderName = CCHeader.Name.Len() > 0 ? CCHeader.Name : DebugName;
 
@@ -2213,9 +2234,7 @@ static bool BuildShaderOutputFromSpirv(
 
 		const FString ResourceName(ANSI_TO_TCHAR(Binding->name));
 		
-		// VKRT Do we need a WriteAS our should we use reflection directly here?
-		//CCHeaderWriter.WriteAS(*ResourceName, ASBindings++);
-		ASBindings++;
+		CCHeaderWriter.WriteAccelerationStructures(*ResourceName, ASBindings++);
 
 		Spirv.ReflectionInfo.Add(FVulkanSpirv::FEntry(ResourceName, BindingIndex));
 	}
@@ -2358,7 +2377,7 @@ void DoCompileVulkanShader(const FShaderCompilerInput& Input, FShaderCompilerOut
 
 	const bool bIsSM5 = (Version == EVulkanShaderVersion::SM5);
 	const bool bIsMobile = (Version == EVulkanShaderVersion::ES3_1 || Version == EVulkanShaderVersion::ES3_1_ANDROID);
-	const bool bStripReflect = (IsVulkanMobilePlatform(ShaderPlatform) || IsVulkanMobileSM5Platform(ShaderPlatform));
+	const bool bStripReflect = (IsVulkanMobilePlatform(ShaderPlatform) || IsVulkanMobileSM5Platform(ShaderPlatform)) || Input.IsRayTracingShader();
 	const bool bForceDXC = Input.Environment.CompilerFlags.Contains(CFLAG_ForceDXC);
 
 	const EHlslShaderFrequency FrequencyTable[] =
diff --git a/Engine/Source/Runtime/RHI/Public/RHIResources.h b/Engine/Source/Runtime/RHI/Public/RHIResources.h
--- a/Engine/Source/Runtime/RHI/Public/RHIResources.h
+++ b/Engine/Source/Runtime/RHI/Public/RHIResources.h
@@ -713,6 +713,30 @@ public:
 	explicit FRHIRayTracingShader(EShaderFrequency InFrequency) : FRHIShader(RRT_RayTracingShader, InFrequency) {}
 };
 
+class FRHIRayGenShader : public FRHIRayTracingShader
+{
+public:
+	FRHIRayGenShader() : FRHIRayTracingShader(SF_RayGen) {}
+};
+
+class FRHIRayMissShader : public FRHIRayTracingShader
+{
+public:
+	FRHIRayMissShader() : FRHIRayTracingShader(SF_RayMiss) {}
+};
+
+class FRHIRayCallableShader : public FRHIRayTracingShader
+{
+public:
+	FRHIRayCallableShader() : FRHIRayTracingShader(SF_RayCallable) {}
+};
+
+class FRHIRayHitGroupShader : public FRHIRayTracingShader
+{
+public:
+	FRHIRayHitGroupShader() : FRHIRayTracingShader(SF_RayHitGroup) {}
+};
+
 class RHI_API FRHIComputeShader : public FRHIShader
 {
 public:
diff --git a/Engine/Source/Runtime/RenderCore/Public/CrossCompilerCommon.h b/Engine/Source/Runtime/RenderCore/Public/CrossCompilerCommon.h
--- a/Engine/Source/Runtime/RenderCore/Public/CrossCompilerCommon.h
+++ b/Engine/Source/Runtime/RenderCore/Public/CrossCompilerCommon.h
@@ -243,6 +243,7 @@ namespace CrossCompiler
 		uint8					NumSamplers;
 		uint8					NumUniformBuffers;
 		uint8					NumUAVs;
+		uint8					NumAccelerationStructures;
 		bool					bHasRegularUniformBuffers;
 	};
 
diff --git a/Engine/Source/Runtime/VulkanRHI/Private/VulkanCommands.cpp b/Engine/Source/Runtime/VulkanRHI/Private/VulkanCommands.cpp
--- a/Engine/Source/Runtime/VulkanRHI/Private/VulkanCommands.cpp
+++ b/Engine/Source/Runtime/VulkanRHI/Private/VulkanCommands.cpp
@@ -288,7 +288,7 @@ void FVulkanCommandListContext::RHISetShaderParameter(FRHIComputeShader* Compute
 }
 
 template <typename TState>
-inline void SetShaderUniformBufferResources(FVulkanCommandListContext* Context, TState* State, const FVulkanShader* Shader, const TArray<FVulkanShaderHeader::FGlobalInfo>& GlobalInfos, const TArray<TEnumAsByte<VkDescriptorType>>& DescriptorTypes, const FVulkanShaderHeader::FUniformBufferInfo& HeaderUBInfo, const FVulkanUniformBuffer* UniformBuffer, const TArray<FDescriptorSetRemappingInfo::FRemappingInfo>& GlobalRemappingInfo)
+inline void SetShaderUniformBufferResources(FVulkanCommandListContext* Context, TState* State, const FVulkanShader* Shader, const TArray<FVulkanShaderHeader::FGlobalInfo>& GlobalInfos, const TArray<TEnumAsByte<EVulkanBindingType::EType>>& DescriptorTypes, const FVulkanShaderHeader::FUniformBufferInfo& HeaderUBInfo, const FVulkanUniformBuffer* UniformBuffer, const TArray<FDescriptorSetRemappingInfo::FRemappingInfo>& GlobalRemappingInfo)
 {
 #if ENABLE_RHI_VALIDATION
 	static_assert(TIsSame<TState, FVulkanPendingGfxState>::Value || TIsSame<TState, FVulkanPendingComputeState>::Value, "TState must be FVulkanPendingGfxState or FVulkanPendingComputeState");
@@ -309,7 +309,7 @@ inline void SetShaderUniformBufferResources(FVulkanCommandListContext* Context,
 		{
 			uint16 CombinedAlias = GlobalInfos[ResourceInfo.GlobalIndex].CombinedSamplerStateAliasIndex;
 			uint32 GlobalIndex = CombinedAlias == UINT16_MAX ? ResourceInfo.GlobalIndex : CombinedAlias;
-			const VkDescriptorType DescriptorType = DescriptorTypes[GlobalInfos[GlobalIndex].TypeIndex];
+			const VkDescriptorType DescriptorType = BindingToDescriptorType(DescriptorTypes[GlobalInfos[GlobalIndex].TypeIndex]);
 			ensure(DescriptorType == VK_DESCRIPTOR_TYPE_SAMPLER || DescriptorType == VK_DESCRIPTOR_TYPE_COMBINED_IMAGE_SAMPLER);
 			FVulkanSamplerState* CurrSampler = static_cast<FVulkanSamplerState*>(ResourceArray[ResourceInfo.SourceUBResourceIndex].GetReference());
 			if (CurrSampler)
@@ -329,7 +329,7 @@ inline void SetShaderUniformBufferResources(FVulkanCommandListContext* Context,
 		case UBMT_TEXTURE:
 		case UBMT_RDG_TEXTURE:
 		{
-			const VkDescriptorType DescriptorType = DescriptorTypes[GlobalInfos[ResourceInfo.GlobalIndex].TypeIndex];
+			const VkDescriptorType DescriptorType = BindingToDescriptorType(DescriptorTypes[GlobalInfos[ResourceInfo.GlobalIndex].TypeIndex]);
 			ensure(DescriptorType == VK_DESCRIPTOR_TYPE_SAMPLED_IMAGE || DescriptorType == VK_DESCRIPTOR_TYPE_COMBINED_IMAGE_SAMPLER);
 			FRHITexture* TexRef = (FRHITexture*)(ResourceArray[ResourceInfo.SourceUBResourceIndex].GetReference());
 			if (TexRef)
@@ -361,7 +361,7 @@ inline void SetShaderUniformBufferResources(FVulkanCommandListContext* Context,
 		case UBMT_SRV:
 		case UBMT_RDG_BUFFER_SRV:
 		{
-			const VkDescriptorType DescriptorType = DescriptorTypes[GlobalInfos[ResourceInfo.GlobalIndex].TypeIndex];
+			const VkDescriptorType DescriptorType = BindingToDescriptorType(DescriptorTypes[GlobalInfos[ResourceInfo.GlobalIndex].TypeIndex]);
 			ensure(DescriptorType == VK_DESCRIPTOR_TYPE_UNIFORM_TEXEL_BUFFER 
 				|| DescriptorType == VK_DESCRIPTOR_TYPE_COMBINED_IMAGE_SAMPLER
 				|| DescriptorType == VK_DESCRIPTOR_TYPE_SAMPLED_IMAGE
@@ -389,7 +389,7 @@ inline void SetShaderUniformBufferResources(FVulkanCommandListContext* Context,
 		case UBMT_UAV:
 		case UBMT_RDG_BUFFER_UAV:
 		{
-			const VkDescriptorType DescriptorType = DescriptorTypes[GlobalInfos[ResourceInfo.GlobalIndex].TypeIndex];
+			const VkDescriptorType DescriptorType = BindingToDescriptorType(DescriptorTypes[GlobalInfos[ResourceInfo.GlobalIndex].TypeIndex]);
 			ensure(DescriptorType == VK_DESCRIPTOR_TYPE_STORAGE_BUFFER
 				|| DescriptorType == VK_DESCRIPTOR_TYPE_STORAGE_IMAGE
 				|| DescriptorType == VK_DESCRIPTOR_TYPE_STORAGE_TEXEL_BUFFER);
diff --git a/Engine/Source/Runtime/VulkanRHI/Private/VulkanDescriptorSets.h b/Engine/Source/Runtime/VulkanRHI/Private/VulkanDescriptorSets.h
--- a/Engine/Source/Runtime/VulkanRHI/Private/VulkanDescriptorSets.h
+++ b/Engine/Source/Runtime/VulkanRHI/Private/VulkanDescriptorSets.h
@@ -273,12 +273,27 @@ class FVulkanDescriptorSetsLayoutInfo
 public:
 	FVulkanDescriptorSetsLayoutInfo()
 	{
-		FMemory::Memzero(LayoutTypes);
+		// Add expected descriptor types
+		for (uint32 i = VK_DESCRIPTOR_TYPE_BEGIN_RANGE; i <= VK_DESCRIPTOR_TYPE_END_RANGE; ++i)
+		{
+			LayoutTypes.Add(static_cast<VkDescriptorType>(i), 0);
+		}
+
+#if VULKAN_RHI_RAYTRACING
+		LayoutTypes.Add(VK_DESCRIPTOR_TYPE_ACCELERATION_STRUCTURE_KHR, 0);
+#endif
 	}
 
 	inline uint32 GetTypesUsed(VkDescriptorType Type) const
 	{
-		return LayoutTypes[Type];
+		if (LayoutTypes.Contains(Type))
+		{
+			return LayoutTypes[Type];
+		}
+		else
+		{
+			return 0;
+		}
 	}
 
 	struct FSetLayout
@@ -375,14 +390,14 @@ public:
 
 	void CopyFrom(const FVulkanDescriptorSetsLayoutInfo& Info)
 	{
-		FMemory::Memcpy(LayoutTypes, Info.LayoutTypes, sizeof(LayoutTypes));
+		LayoutTypes = Info.LayoutTypes;
 		Hash = Info.Hash;
 		TypesUsageID = Info.TypesUsageID;
 		SetLayouts = Info.SetLayouts;
 		RemappingInfo = Info.RemappingInfo;
 	}
 
-	inline const uint32* GetLayoutTypes() const
+	inline const TMap<VkDescriptorType, uint32>& GetLayoutTypes() const
 	{
 		return LayoutTypes;
 	}
@@ -394,11 +409,11 @@ public:
 
 	inline bool HasInputAttachments() const
 	{
-		return LayoutTypes[VK_DESCRIPTOR_TYPE_INPUT_ATTACHMENT] > 0;
+		return GetTypesUsed(VK_DESCRIPTOR_TYPE_INPUT_ATTACHMENT) > 0;
 	}
 
 protected:
-	uint32 LayoutTypes[VK_DESCRIPTOR_TYPE_RANGE_SIZE];
+	TMap<VkDescriptorType, uint32> LayoutTypes;
 	TArray<FSetLayout> SetLayouts;
 
 	uint32 Hash = 0;
@@ -871,6 +886,9 @@ protected:
 	friend class FVulkanComputePipeline;
 	friend class FVulkanGfxPipeline;
 	friend class FVulkanPipelineStateCacheManager;
+#if VULKAN_RHI_RAYTRACING
+	friend class FVulkanRayTracingPipelineState;
+#endif
 };
 
 class FVulkanGfxLayout : public FVulkanLayout
@@ -921,6 +939,22 @@ protected:
 	friend class FVulkanPipelineStateCacheManager;
 };
 
+#if VULKAN_RHI_RAYTRACING
+class FVulkanRayTracingLayout : public FVulkanLayout
+{
+public:
+	FVulkanRayTracingLayout(FVulkanDevice* InDevice)
+		: FVulkanLayout(InDevice)
+	{
+	}
+
+	virtual bool IsGfxLayout() const final override
+	{
+		return false;
+	}
+};
+#endif // VULKAN_RHI_RAYTRACING
+
 // This class encapsulates updating VkWriteDescriptorSet structures (but doesn't own them), and their flags for dirty ranges; it is intended
 // to be used to access a sub-region of a long array of VkWriteDescriptorSet (ie FVulkanDescriptorSetWriteContainer)
 class FVulkanDescriptorSetWriter
diff --git a/Engine/Source/Runtime/VulkanRHI/Private/VulkanDevice.h b/Engine/Source/Runtime/VulkanRHI/Private/VulkanDevice.h
--- a/Engine/Source/Runtime/VulkanRHI/Private/VulkanDevice.h
+++ b/Engine/Source/Runtime/VulkanRHI/Private/VulkanDevice.h
@@ -579,6 +579,7 @@ private:
 
 #if VULKAN_RHI_RAYTRACING
 	FRayTracingProperties RayTracingProperties;
+	FVulkanBasicRaytracingPipeline* BasicRayTracingPipeline = nullptr;
 #endif // VULKAN_RHI_RAYTRACING
 #endif // VULKAN_SUPPORTS_PHYSICAL_DEVICE_PROPERTIES2
 
diff --git a/Engine/Source/Runtime/VulkanRHI/Private/VulkanPendingState.cpp b/Engine/Source/Runtime/VulkanRHI/Private/VulkanPendingState.cpp
--- a/Engine/Source/Runtime/VulkanRHI/Private/VulkanPendingState.cpp
+++ b/Engine/Source/Runtime/VulkanRHI/Private/VulkanPendingState.cpp
@@ -866,10 +866,9 @@ bool FVulkanDescriptorSetCache::FCachedPool::CreateDescriptorSets(
 
 		DSWriter.SetDescriptorSet(NewSetEntry.Sets[Index]);
 
-		const uint32* LayoutTypes = SetsLayout.GetLayoutTypes();
-		for (int32 i = 0; i < VK_DESCRIPTOR_TYPE_RANGE_SIZE; ++i)
+		for (int32 i = VK_DESCRIPTOR_TYPE_BEGIN_RANGE; i <= VK_DESCRIPTOR_TYPE_END_RANGE; ++i)
 		{
-			PoolSizesStatistic[i]+= LayoutTypes[i];
+			PoolSizesStatistic[i] += SetsLayout.GetTypesUsed(static_cast<VkDescriptorType>(i));
 		}
 
 #if VULKAN_ENABLE_AGGRESSIVE_STATS
diff --git a/Engine/Source/Runtime/VulkanRHI/Private/VulkanRHI.cpp b/Engine/Source/Runtime/VulkanRHI/Private/VulkanRHI.cpp
--- a/Engine/Source/Runtime/VulkanRHI/Private/VulkanRHI.cpp
+++ b/Engine/Source/Runtime/VulkanRHI/Private/VulkanRHI.cpp
@@ -748,7 +748,8 @@ void FVulkanDynamicRHI::InitInstance()
 		GSupportsMobileMultiView = Device->GetMultiviewFeatures().multiview == VK_TRUE ? true : false;
 #endif
 #if VULKAN_RHI_RAYTRACING
-		GRHISupportsRayTracing = Device->GetOptionalExtensions().HasRaytracingExtensions();
+		GRHISupportsRayTracing = false; // Disable runtime hwrt on vk during development.
+		//GRHISupportsRayTracing = Device->GetOptionalExtensions().HasRaytracingExtensions();
 #endif
 #if VULKAN_ENABLE_DUMP_LAYER
 		// Disable RHI thread by default if the dump layer is enabled
@@ -1253,7 +1254,14 @@ FVulkanDescriptorSetsLayout::~FVulkanDescriptorSetsLayout()
 void FVulkanDescriptorSetsLayoutInfo::AddDescriptor(int32 DescriptorSetIndex, const VkDescriptorSetLayoutBinding& Descriptor)
 {
 	// Increment type usage
-	LayoutTypes[Descriptor.descriptorType]++;
+	if (LayoutTypes.Contains(Descriptor.descriptorType))
+	{
+		LayoutTypes[Descriptor.descriptorType]++;
+	}
+	else
+	{
+		LayoutTypes.Add(Descriptor.descriptorType, 1);
+	}
 
 	if (DescriptorSetIndex >= SetLayouts.Num())
 	{
@@ -1283,7 +1291,12 @@ void FVulkanDescriptorSetsLayoutInfo::AddDescriptor(int32 DescriptorSetIndex, co
 		break;
 	case VK_DESCRIPTOR_TYPE_STORAGE_TEXEL_BUFFER:
 	case VK_DESCRIPTOR_TYPE_UNIFORM_TEXEL_BUFFER:
+
+#if VULKAN_RHI_RAYTRACING
+	case VK_DESCRIPTOR_TYPE_ACCELERATION_STRUCTURE_KHR:
 		break;
+#endif
+
 	default:
 		checkf(0, TEXT("Unsupported descriptor type %d"), (int32)Descriptor.descriptorType);
 		break;
@@ -1337,7 +1350,16 @@ void FVulkanDescriptorSetsLayoutInfo::CompileTypesUsageID()
 	static TMap<uint32, uint32> GTypesUsageHashMap;
 	static uint32 GUniqueID = 1;
 
-	const uint32 TypesUsageHash = FCrc::MemCrc32(LayoutTypes, sizeof(LayoutTypes));
+	LayoutTypes.KeySort([](VkDescriptorType A, VkDescriptorType B)
+	{
+		return static_cast<uint32>(A) < static_cast<uint32>(B);
+	});
+
+	uint32 TypesUsageHash = 0;
+	for (const auto& Elem : LayoutTypes)
+	{
+		TypesUsageHash = FCrc::MemCrc32(&Elem.Value, sizeof(uint32), TypesUsageHash);
+	}
 
 	uint32* UniqueID = GTypesUsageHashMap.Find(TypesUsageHash);
 	if (UniqueID == nullptr)
@@ -1397,6 +1419,10 @@ void FVulkanDescriptorSetsLayout::Compile(FVulkanDescriptorSetLayoutMap& DSetLay
 			<	Limits.maxDescriptorSetStorageImages);
 
 	check(LayoutTypes[VK_DESCRIPTOR_TYPE_INPUT_ATTACHMENT] <= Limits.maxDescriptorSetInputAttachments);
+
+#if VULKAN_RHI_RAYTRACING
+	check(LayoutTypes[VK_DESCRIPTOR_TYPE_ACCELERATION_STRUCTURE_KHR] < Device->GetRayTracingProperties().AccelerationStructure.maxDescriptorSetAccelerationStructures);
+#endif
 	
 	LayoutHandles.Empty(SetLayouts.Num());
 
diff --git a/Engine/Source/Runtime/VulkanRHI/Private/VulkanRHIPrivate.h b/Engine/Source/Runtime/VulkanRHI/Private/VulkanRHIPrivate.h
--- a/Engine/Source/Runtime/VulkanRHI/Private/VulkanRHIPrivate.h
+++ b/Engine/Source/Runtime/VulkanRHI/Private/VulkanRHIPrivate.h
@@ -82,10 +82,18 @@ inline VkShaderStageFlagBits UEFrequencyToVKStageBit(EShaderFrequency InStage)
 {
 	switch (InStage)
 	{
-	case SF_Vertex:		return VK_SHADER_STAGE_VERTEX_BIT;
-	case SF_Pixel:		return VK_SHADER_STAGE_FRAGMENT_BIT;
-	case SF_Geometry:	return VK_SHADER_STAGE_GEOMETRY_BIT;
-	case SF_Compute:	return VK_SHADER_STAGE_COMPUTE_BIT;
+	case SF_Vertex:			return VK_SHADER_STAGE_VERTEX_BIT;
+	case SF_Pixel:			return VK_SHADER_STAGE_FRAGMENT_BIT;
+	case SF_Geometry:		return VK_SHADER_STAGE_GEOMETRY_BIT;
+	case SF_Compute:		return VK_SHADER_STAGE_COMPUTE_BIT;
+
+#if VULKAN_RHI_RAYTRACING
+	case SF_RayGen:			return VK_SHADER_STAGE_RAYGEN_BIT_KHR;
+	case SF_RayMiss:		return VK_SHADER_STAGE_MISS_BIT_KHR;
+	case SF_RayHitGroup:	return VK_SHADER_STAGE_CLOSEST_HIT_BIT_KHR; // vkrt todo: How to handle VK_SHADER_STAGE_ANY_HIT_BIT_KHR?
+	case SF_RayCallable:	return VK_SHADER_STAGE_CALLABLE_BIT_KHR;
+#endif // VULKAN_RHI_RAYTRACING
+
 	default:
 		checkf(false, TEXT("Undefined shader stage %d"), (int32)InStage);
 		break;
@@ -102,6 +110,19 @@ inline EShaderFrequency VkStageBitToUEFrequency(VkShaderStageFlagBits FlagBits)
 	case VK_SHADER_STAGE_FRAGMENT_BIT:					return SF_Pixel;
 	case VK_SHADER_STAGE_GEOMETRY_BIT:					return SF_Geometry;
 	case VK_SHADER_STAGE_COMPUTE_BIT:					return SF_Compute;
+
+#if VULKAN_RHI_RAYTRACING
+	case VK_SHADER_STAGE_RAYGEN_BIT_KHR:				return SF_RayGen;
+	case VK_SHADER_STAGE_MISS_BIT_KHR:					return SF_RayMiss;
+	case VK_SHADER_STAGE_CALLABLE_BIT_KHR:				return SF_RayCallable;
+
+	// Hit group frequencies
+	case VK_SHADER_STAGE_ANY_HIT_BIT_KHR:
+	case VK_SHADER_STAGE_CLOSEST_HIT_BIT_KHR:
+	case VK_SHADER_STAGE_INTERSECTION_BIT_KHR:
+		return SF_RayHitGroup;
+#endif // VULKAN_RHI_RAYTRACING
+
 	default:
 		checkf(false, TEXT("Undefined VkShaderStageFlagBits %d"), (int32)FlagBits);
 		break;
diff --git a/Engine/Source/Runtime/VulkanRHI/Private/VulkanRayTracing.cpp b/Engine/Source/Runtime/VulkanRHI/Private/VulkanRayTracing.cpp
--- a/Engine/Source/Runtime/VulkanRHI/Private/VulkanRayTracing.cpp
+++ b/Engine/Source/Runtime/VulkanRHI/Private/VulkanRayTracing.cpp
@@ -2,6 +2,8 @@
 
 #include "VulkanRHIPrivate.h"
 #include "VulkanContext.h"
+#include "VulkanDescriptorSets.h"
+#include "BuiltInRayTracingShaders.h"
 #include "Experimental/Containers/SherwoodHashTable.h"
 
 #if VULKAN_RHI_RAYTRACING
@@ -637,15 +639,244 @@ void FVulkanCommandListContext::RHIRayTraceOcclusion(FRHIRayTracingScene* Scene,
 	return;
 }
 
+template<typename ShaderType>
+static FRHIRayTracingShader* GetBuiltInRayTracingShader()
+{
+	const FGlobalShaderMap* const ShaderMap = GetGlobalShaderMap(GMaxRHIFeatureLevel);
+	auto Shader = ShaderMap->GetShader<ShaderType>();
+	return static_cast<FRHIRayTracingShader*>(Shader.GetRayTracingShader());
+}
+
 void FVulkanDevice::InitializeRayTracing()
 {
-	// todo
-	return;
+	check(BasicRayTracingPipeline == nullptr);
+	BasicRayTracingPipeline = new FVulkanBasicRaytracingPipeline(this);
 }
 
 void FVulkanDevice::CleanUpRayTracing()
 {
-	// todo
-	return;
+	if (BasicRayTracingPipeline != nullptr)
+	{
+		delete BasicRayTracingPipeline;
+		BasicRayTracingPipeline = nullptr;
+	}
+}
+
+static uint32 GetAlignedSize(uint32 Value, uint32 Alignment)
+{
+	return (Value + Alignment - 1) & ~(Alignment - 1);
+}
+
+FVulkanRayTracingPipelineState::FVulkanRayTracingPipelineState(FVulkanDevice* const InDevice, const FRayTracingPipelineStateInitializer& Initializer)
+{	
+	check(Layout == nullptr);
+
+	TArrayView<FRHIRayTracingShader*> InitializerRayGenShaders = Initializer.GetRayGenTable();
+	TArrayView<FRHIRayTracingShader*> InitializerMissShaders = Initializer.GetMissTable();
+	TArrayView<FRHIRayTracingShader*> InitializerHitGroupShaders = Initializer.GetHitGroupTable();
+	// vkrt todo: Callable shader support
+
+	FVulkanDescriptorSetsLayoutInfo DescriptorSetLayoutInfo;
+	FUniformBufferGatherInfo UBGatherInfo;
+	
+	for (FRHIRayTracingShader* RayGenShader : InitializerRayGenShaders)
+	{
+		const FVulkanShaderHeader& Header = static_cast<FVulkanRayGenShader*>(RayGenShader)->GetCodeHeader();
+		DescriptorSetLayoutInfo.ProcessBindingsForStage(VK_SHADER_STAGE_RAYGEN_BIT_KHR, ShaderStage::RayGen, Header, UBGatherInfo);
+	}
+
+	for (FRHIRayTracingShader* MissShader : InitializerMissShaders)
+	{
+		const FVulkanShaderHeader& Header = static_cast<FVulkanRayMissShader*>(MissShader)->GetCodeHeader();
+		DescriptorSetLayoutInfo.ProcessBindingsForStage(VK_SHADER_STAGE_MISS_BIT_KHR, ShaderStage::RayMiss, Header, UBGatherInfo);
+	}
+
+	for (FRHIRayTracingShader* HitGroupShader : InitializerHitGroupShaders)
+	{
+		const FVulkanShaderHeader& Header = static_cast<FVulkanRayHitGroupShader*>(HitGroupShader)->GetCodeHeader();
+		DescriptorSetLayoutInfo.ProcessBindingsForStage(VK_SHADER_STAGE_CLOSEST_HIT_BIT_KHR, ShaderStage::RayHitGroup, Header, UBGatherInfo);
+		// vkrt todo: How to handle any hit for hit group?
+	}
+
+	DescriptorSetLayoutInfo.FinalizeBindings<false>(UBGatherInfo, TArrayView<FRHISamplerState*>());
+
+	Layout = new FVulkanRayTracingLayout(InDevice);
+	Layout->DescriptorSetLayout.CopyFrom(DescriptorSetLayoutInfo);
+	FVulkanDescriptorSetLayoutMap DSetLayoutMap;
+	Layout->Compile(DSetLayoutMap);
+
+	TArray<VkPipelineShaderStageCreateInfo> ShaderStages;
+	TArray<VkRayTracingShaderGroupCreateInfoKHR> ShaderGroups;
+	TArray<ANSICHAR*> EntryPointNames;
+	const uint32 EntryPointNameMaxLength = 24;
+
+	for (FRHIRayTracingShader* const RayGenShaderRHI : InitializerRayGenShaders)
+	{
+		VkPipelineShaderStageCreateInfo ShaderStage;
+		ZeroVulkanStruct(ShaderStage, VK_STRUCTURE_TYPE_PIPELINE_SHADER_STAGE_CREATE_INFO);
+		ShaderStage.module = static_cast<FVulkanRayGenShader*>(RayGenShaderRHI)->GetOrCreateHandle(Layout, Layout->GetDescriptorSetLayoutHash());
+		ShaderStage.stage = VK_SHADER_STAGE_RAYGEN_BIT_KHR;
+			
+		ANSICHAR* const EntryPoint = new ANSICHAR[EntryPointNameMaxLength];
+		static_cast<FVulkanRayGenShader*>(RayGenShaderRHI)->GetEntryPoint(EntryPoint, EntryPointNameMaxLength);
+		EntryPointNames.Add(EntryPoint);
+		ShaderStage.pName = EntryPoint;
+		ShaderStages.Add(ShaderStage);
+
+		VkRayTracingShaderGroupCreateInfoKHR ShaderGroup;
+		ZeroVulkanStruct(ShaderGroup, VK_STRUCTURE_TYPE_RAY_TRACING_SHADER_GROUP_CREATE_INFO_KHR);
+		ShaderGroup.type = VK_RAY_TRACING_SHADER_GROUP_TYPE_GENERAL_KHR;
+		ShaderGroup.generalShader = ShaderStages.Num() - 1;
+		ShaderGroup.closestHitShader = VK_SHADER_UNUSED_KHR;
+		ShaderGroup.anyHitShader = VK_SHADER_UNUSED_KHR;
+		ShaderGroup.intersectionShader = VK_SHADER_UNUSED_KHR;
+		ShaderGroups.Add(ShaderGroup);
+	}
+
+	for (FRHIRayTracingShader* const MissShaderRHI : InitializerMissShaders)
+	{
+		VkPipelineShaderStageCreateInfo ShaderStage;
+		ZeroVulkanStruct(ShaderStage, VK_STRUCTURE_TYPE_PIPELINE_SHADER_STAGE_CREATE_INFO);
+		ShaderStage.module = static_cast<FVulkanRayMissShader*>(MissShaderRHI)->GetOrCreateHandle(Layout, Layout->GetDescriptorSetLayoutHash());
+		ShaderStage.stage = VK_SHADER_STAGE_MISS_BIT_KHR;
+
+		ANSICHAR* const EntryPoint = new char[EntryPointNameMaxLength];
+		static_cast<FVulkanRayGenShader*>(MissShaderRHI)->GetEntryPoint(EntryPoint, EntryPointNameMaxLength);
+		EntryPointNames.Add(EntryPoint);
+		ShaderStage.pName = EntryPoint;
+		ShaderStages.Add(ShaderStage);
+
+		VkRayTracingShaderGroupCreateInfoKHR ShaderGroup;
+		ZeroVulkanStruct(ShaderGroup, VK_STRUCTURE_TYPE_RAY_TRACING_SHADER_GROUP_CREATE_INFO_KHR);
+		ShaderGroup.type = VK_RAY_TRACING_SHADER_GROUP_TYPE_GENERAL_KHR;
+		ShaderGroup.generalShader = ShaderStages.Num() - 1;
+		ShaderGroup.closestHitShader = VK_SHADER_UNUSED_KHR;
+		ShaderGroup.anyHitShader = VK_SHADER_UNUSED_KHR;
+		ShaderGroup.intersectionShader = VK_SHADER_UNUSED_KHR;
+		ShaderGroups.Add(ShaderGroup);
+	}
+
+	for (FRHIRayTracingShader* const HitGroupShaderRHI : InitializerHitGroupShaders)
+	{
+		VkPipelineShaderStageCreateInfo ShaderStage;
+		ZeroVulkanStruct(ShaderStage, VK_STRUCTURE_TYPE_PIPELINE_SHADER_STAGE_CREATE_INFO);
+		ShaderStage.module = static_cast<FVulkanRayHitGroupShader*>(HitGroupShaderRHI)->GetOrCreateHandle(Layout, Layout->GetDescriptorSetLayoutHash());
+		ShaderStage.stage = VK_SHADER_STAGE_CLOSEST_HIT_BIT_KHR;
+
+		ANSICHAR* const EntryPoint = new char[EntryPointNameMaxLength];
+		static_cast<FVulkanRayHitGroupShader*>(HitGroupShaderRHI)->GetEntryPoint(EntryPoint, EntryPointNameMaxLength);
+		EntryPointNames.Add(EntryPoint);
+		ShaderStage.pName = EntryPoint;
+		ShaderStages.Add(ShaderStage);
+
+		VkRayTracingShaderGroupCreateInfoKHR ShaderGroup;
+		ZeroVulkanStruct(ShaderGroup, VK_STRUCTURE_TYPE_RAY_TRACING_SHADER_GROUP_CREATE_INFO_KHR);
+		ShaderGroup.type = VK_RAY_TRACING_SHADER_GROUP_TYPE_TRIANGLES_HIT_GROUP_KHR;
+		ShaderGroup.generalShader = VK_SHADER_UNUSED_KHR;
+		ShaderGroup.closestHitShader = ShaderStages.Num() - 1;
+		ShaderGroup.anyHitShader = VK_SHADER_UNUSED_KHR; // vkrt: todo
+		ShaderGroup.intersectionShader = VK_SHADER_UNUSED_KHR;
+		ShaderGroups.Add(ShaderGroup);
+	}
+
+	VkRayTracingPipelineCreateInfoKHR RayTracingPipelineCreateInfo;
+	ZeroVulkanStruct(RayTracingPipelineCreateInfo, VK_STRUCTURE_TYPE_RAY_TRACING_PIPELINE_CREATE_INFO_KHR);
+	RayTracingPipelineCreateInfo.stageCount = ShaderStages.Num();
+	RayTracingPipelineCreateInfo.pStages = ShaderStages.GetData();
+	RayTracingPipelineCreateInfo.groupCount = ShaderGroups.Num();
+	RayTracingPipelineCreateInfo.pGroups = ShaderGroups.GetData();
+	RayTracingPipelineCreateInfo.maxPipelineRayRecursionDepth = 1;
+	RayTracingPipelineCreateInfo.layout = Layout->GetPipelineLayout();
+	
+	VERIFYVULKANRESULT(vkCreateRayTracingPipelinesKHR(
+		InDevice->GetInstanceHandle(), 
+		VK_NULL_HANDLE, // Deferred Operation 
+		VK_NULL_HANDLE, // Pipeline Cache 
+		1, 
+		&RayTracingPipelineCreateInfo, 
+		VULKAN_CPU_ALLOCATOR, 
+		&Pipeline));
+
+	for (ANSICHAR* const EntryPoint : EntryPointNames)
+	{
+		delete[] EntryPoint;
+	}
+
+	const FRayTracingProperties& Props = InDevice->GetRayTracingProperties();
+	const uint32 HandleSize = Props.RayTracingPipeline.shaderGroupHandleSize;
+	const uint32 HandleSizeAligned = GetAlignedSize(HandleSize, Props.RayTracingPipeline.shaderGroupHandleAlignment);
+	const uint32 GroupCount = ShaderGroups.Num();
+	const uint32 SBTSize = GroupCount * HandleSizeAligned;
+
+	TArray<uint8> ShaderHandleStorage;
+	ShaderHandleStorage.AddUninitialized(SBTSize);
+	VERIFYVULKANRESULT(vkGetRayTracingShaderGroupHandlesKHR(InDevice->GetInstanceHandle(), Pipeline, 0, GroupCount, SBTSize, ShaderHandleStorage.GetData()));
+
+	auto CopyHanldlesToSBT = [InDevice, HandleSize, ShaderHandleStorage](FVkRtAllocation& Allocation, uint32 Offset)
+	{
+		FVulkanRayTracingAllocator::Allocate(
+			InDevice->GetPhysicalHandle(),
+			InDevice->GetInstanceHandle(),
+			HandleSize,
+			VK_BUFFER_USAGE_SHADER_BINDING_TABLE_BIT_KHR | VK_BUFFER_USAGE_SHADER_DEVICE_ADDRESS_BIT,
+			VK_MEMORY_PROPERTY_HOST_VISIBLE_BIT | VK_MEMORY_PROPERTY_HOST_COHERENT_BIT,
+			Allocation);
+
+		void* pMappedBufferMemory = nullptr;
+		VERIFYVULKANRESULT(vkMapMemory(InDevice->GetInstanceHandle(), Allocation.Memory, 0, VK_WHOLE_SIZE, 0, &pMappedBufferMemory));
+		{
+			FMemory::Memcpy(pMappedBufferMemory, ShaderHandleStorage.GetData() + Offset, HandleSize);
+		}
+		vkUnmapMemory(InDevice->GetInstanceHandle(), Allocation.Memory);
+	};
+
+	CopyHanldlesToSBT(RayGenShaderBindingTable, 0);
+	CopyHanldlesToSBT(MissShaderBindingTable, HandleSizeAligned);
+	CopyHanldlesToSBT(HitShaderBindingTable, HandleSizeAligned * 2);
+}
+
+FVulkanRayTracingPipelineState::~FVulkanRayTracingPipelineState()
+{
+	FVulkanRayTracingAllocator::Free(RayGenShaderBindingTable);
+	FVulkanRayTracingAllocator::Free(MissShaderBindingTable);
+	FVulkanRayTracingAllocator::Free(HitShaderBindingTable);
+
+	if (Layout != nullptr)
+	{
+		delete Layout;
+		Layout = nullptr;
+	}
+}
+
+FVulkanBasicRaytracingPipeline::FVulkanBasicRaytracingPipeline(FVulkanDevice* const InDevice)
+{
+	check(Occlusion == nullptr);
+
+	// Occlusion pipeline
+	{
+		FRayTracingPipelineStateInitializer OcclusionInitializer;
+
+		FRHIRayTracingShader* OcclusionRGSTable[] = { GetBuiltInRayTracingShader<FOcclusionMainRG>() };
+		OcclusionInitializer.SetRayGenShaderTable(OcclusionRGSTable);
+
+		FRHIRayTracingShader* OcclusionMSTable[] = { GetBuiltInRayTracingShader<FDefaultPayloadMS>() };
+		OcclusionInitializer.SetMissShaderTable(OcclusionMSTable);
+
+		FRHIRayTracingShader* OcclusionCHSTable[] = { GetBuiltInRayTracingShader<FDefaultMainCHS>() };
+		OcclusionInitializer.SetHitGroupTable(OcclusionCHSTable);
+
+		OcclusionInitializer.bAllowHitGroupIndexing = false;
+
+		Occlusion = new FVulkanRayTracingPipelineState(InDevice, OcclusionInitializer);
+	}
+}
+
+FVulkanBasicRaytracingPipeline::~FVulkanBasicRaytracingPipeline()
+{
+	if (Occlusion != nullptr)
+	{
+		delete Occlusion;
+		Occlusion = nullptr;
+	}
 }
-#endif // #if VULKAN_RHI_RAYTRACING
+#endif // VULKAN_RHI_RAYTRACING
diff --git a/Engine/Source/Runtime/VulkanRHI/Private/VulkanRayTracing.h b/Engine/Source/Runtime/VulkanRHI/Private/VulkanRayTracing.h
--- a/Engine/Source/Runtime/VulkanRHI/Private/VulkanRayTracing.h
+++ b/Engine/Source/Runtime/VulkanRHI/Private/VulkanRayTracing.h
@@ -7,6 +7,7 @@
 class FVulkanCommandListContext;
 class FVulkanResourceMultiBuffer;
 class FVulkanAccelerationStructureBuffer;
+class FVulkanRayTracingLayout;
 
 #define ENUM_VK_ENTRYPOINTS_RAYTRACING(EnumMacro) \
 	EnumMacro(PFN_vkCreateAccelerationStructureKHR, vkCreateAccelerationStructureKHR) \
@@ -140,4 +141,33 @@ private:
 	const FVulkanDevice* Device = nullptr;
 };
 
-#endif // #if VULKAN_RHI_RAYTRACING
+class FVulkanRayTracingPipelineState : public FRHIRayTracingPipelineState
+{
+public:
+
+	UE_NONCOPYABLE(FVulkanRayTracingPipelineState);
+	FVulkanRayTracingPipelineState(FVulkanDevice* const InDevice, const FRayTracingPipelineStateInitializer& Initializer);
+	~FVulkanRayTracingPipelineState();
+
+private:
+
+	FVulkanRayTracingLayout* Layout = nullptr;
+	VkPipeline Pipeline = VK_NULL_HANDLE;
+	FVkRtAllocation RayGenShaderBindingTable;
+	FVkRtAllocation MissShaderBindingTable;
+	FVkRtAllocation HitShaderBindingTable;
+};
+
+class FVulkanBasicRaytracingPipeline
+{
+public:
+
+	UE_NONCOPYABLE(FVulkanBasicRaytracingPipeline);
+	FVulkanBasicRaytracingPipeline(FVulkanDevice* const InDevice);
+	~FVulkanBasicRaytracingPipeline();
+
+private:
+
+	FVulkanRayTracingPipelineState* Occlusion = nullptr;
+};
+#endif // VULKAN_RHI_RAYTRACING
diff --git a/Engine/Source/Runtime/VulkanRHI/Private/VulkanResources.h b/Engine/Source/Runtime/VulkanRHI/Private/VulkanResources.h
--- a/Engine/Source/Runtime/VulkanRHI/Private/VulkanResources.h
+++ b/Engine/Source/Runtime/VulkanRHI/Private/VulkanResources.h
@@ -207,6 +207,13 @@ typedef TVulkanBaseShader<FRHIPixelShader, SF_Pixel, VK_SHADER_STAGE_FRAGMENT_BI
 typedef TVulkanBaseShader<FRHIComputeShader, SF_Compute, VK_SHADER_STAGE_COMPUTE_BIT>				FVulkanComputeShader;
 typedef TVulkanBaseShader<FRHIGeometryShader, SF_Geometry, VK_SHADER_STAGE_GEOMETRY_BIT>			FVulkanGeometryShader;
 
+#if VULKAN_RHI_RAYTRACING
+typedef TVulkanBaseShader<FRHIRayGenShader, SF_RayGen, VK_SHADER_STAGE_RAYGEN_BIT_KHR>					FVulkanRayGenShader;
+typedef TVulkanBaseShader<FRHIRayMissShader, SF_RayMiss, VK_SHADER_STAGE_MISS_BIT_KHR>					FVulkanRayMissShader;
+typedef TVulkanBaseShader<FRHIRayCallableShader, SF_RayCallable, VK_SHADER_STAGE_CALLABLE_BIT_KHR>		FVulkanRayCallableShader;
+typedef TVulkanBaseShader<FRHIRayHitGroupShader, SF_RayHitGroup, VK_SHADER_STAGE_CLOSEST_HIT_BIT_KHR>	FVulkanRayHitGroupShader; // vkrt todo: How to handle VK_SHADER_STAGE_ANY_HIT_BIT_KHR?
+#endif // VULKAN_RHI_RAYTRACING
+
 class FVulkanShaderFactory
 {
 public:
diff --git a/Engine/Source/Runtime/VulkanRHI/Private/VulkanShaders.cpp b/Engine/Source/Runtime/VulkanRHI/Private/VulkanShaders.cpp
--- a/Engine/Source/Runtime/VulkanRHI/Private/VulkanShaders.cpp
+++ b/Engine/Source/Runtime/VulkanRHI/Private/VulkanShaders.cpp
@@ -341,6 +341,29 @@ FComputeShaderRHIRef FVulkanDynamicRHI::RHICreateComputeShader(TArrayView<const
 	return Device->GetShaderFactory().CreateShader<FVulkanComputeShader>(Code, Device);
 }
 
+#if VULKAN_RHI_RAYTRACING
+FRayTracingShaderRHIRef FVulkanDynamicRHI::RHICreateRayTracingShader(TArrayView<const uint8> Code, const FSHAHash& Hash, EShaderFrequency ShaderFrequency)
+{
+	switch (ShaderFrequency)
+	{
+	case EShaderFrequency::SF_RayGen:
+		 return Device->GetShaderFactory().CreateShader<FVulkanRayGenShader>(Code, Device);
+
+	case EShaderFrequency::SF_RayMiss:
+		return Device->GetShaderFactory().CreateShader<FVulkanRayMissShader>(Code, Device);
+
+	case EShaderFrequency::SF_RayCallable:
+		return Device->GetShaderFactory().CreateShader<FVulkanRayCallableShader>(Code, Device);
+
+	case EShaderFrequency::SF_RayHitGroup:
+		return Device->GetShaderFactory().CreateShader<FVulkanRayHitGroupShader>(Code, Device);
+
+	default:
+		check(false);
+		return nullptr;
+	}
+}
+#endif // VULKAN_RHI_RAYTRACING
 
 FVulkanLayout::FVulkanLayout(FVulkanDevice* InDevice)
 	: VulkanRHI::FDeviceChild(InDevice)
@@ -613,7 +636,7 @@ void FVulkanDescriptorSetsLayoutInfo::FinalizeBindings(const FUniformBufferGathe
 			{
 				const FVulkanShaderHeader::FGlobalInfo& GlobalInfo = ShaderHeader->Globals[Index];
 				int32 DescriptorSet = FindOrAddDescriptorSet(Stage);
-				VkDescriptorType Type = ShaderHeader->GlobalDescriptorTypes[GlobalInfo.TypeIndex];
+				VkDescriptorType Type = BindingToDescriptorType(ShaderHeader->GlobalDescriptorTypes[GlobalInfo.TypeIndex]);
 				uint16 CombinedSamplerStateAlias = GlobalInfo.CombinedSamplerStateAliasIndex;
 				uint32 NewBindingIndex = RemappingInfo.AddGlobal(Stage, Index, DescriptorSet, Type, CombinedSamplerStateAlias);
 				Binding.binding = NewBindingIndex;
@@ -647,7 +670,7 @@ void FVulkanDescriptorSetsLayoutInfo::FinalizeBindings(const FUniformBufferGathe
 				{
 					int32 OriginalGlobalIndex = ShaderHeader->InputAttachments[SrcIndex].GlobalIndex;
 					const FVulkanShaderHeader::FGlobalInfo& OriginalGlobalInfo = ShaderHeader->Globals[OriginalGlobalIndex];
-					check(ShaderHeader->GlobalDescriptorTypes[OriginalGlobalInfo.TypeIndex] == VK_DESCRIPTOR_TYPE_INPUT_ATTACHMENT);
+					check(BindingToDescriptorType(ShaderHeader->GlobalDescriptorTypes[OriginalGlobalInfo.TypeIndex]) == VK_DESCRIPTOR_TYPE_INPUT_ATTACHMENT);
 					int32 RemappingIndex = RemappingInfo.InputAttachmentData.AddDefaulted();
 					FInputAttachmentData& AttachmentData = RemappingInfo.InputAttachmentData[RemappingIndex];
 					AttachmentData.BindingIndex = RemappingInfo.StageInfos[Stage].Globals[OriginalGlobalIndex].NewBindingIndex;
diff --git a/Engine/Source/Runtime/VulkanRHI/Private/Windows/VulkanWindowsPlatform.h b/Engine/Source/Runtime/VulkanRHI/Private/Windows/VulkanWindowsPlatform.h
--- a/Engine/Source/Runtime/VulkanRHI/Private/Windows/VulkanWindowsPlatform.h
+++ b/Engine/Source/Runtime/VulkanRHI/Private/Windows/VulkanWindowsPlatform.h
@@ -19,7 +19,7 @@
 
 #define	UE_VK_API_VERSION							VK_API_VERSION_1_1
 
-#define VULKAN_RHI_RAYTRACING 						0&&(RHI_RAYTRACING)
+#define VULKAN_RHI_RAYTRACING 						(RHI_RAYTRACING)
 #define VULKAN_SUPPORTS_SCALAR_BLOCK_LAYOUT			(VULKAN_RHI_RAYTRACING)
 
 // 32-bit windows has warnings on custom mem mgr callbacks
diff --git a/Engine/Source/Runtime/VulkanRHI/Public/VulkanDynamicRHI.h b/Engine/Source/Runtime/VulkanRHI/Public/VulkanDynamicRHI.h
--- a/Engine/Source/Runtime/VulkanRHI/Public/VulkanDynamicRHI.h
+++ b/Engine/Source/Runtime/VulkanRHI/Public/VulkanDynamicRHI.h
@@ -328,6 +328,7 @@ public:
 	virtual FRayTracingAccelerationStructureSize RHICalcRayTracingGeometrySize(const FRayTracingGeometryInitializer& Initializer) final override;
 	virtual FRayTracingGeometryRHIRef RHICreateRayTracingGeometry(const FRayTracingGeometryInitializer& Initializer) final override;
 	virtual FRayTracingSceneRHIRef RHICreateRayTracingScene(const FRayTracingSceneInitializer& Initializer) final override;
+	virtual FRayTracingShaderRHIRef RHICreateRayTracingShader(TArrayView<const uint8> Code, const FSHAHash& Hash, EShaderFrequency ShaderFrequency) final override;
 #endif
 
 	// Historical number of times we've presented any and all viewports
diff --git a/Engine/Source/Runtime/VulkanRHI/Public/VulkanShaderResources.h b/Engine/Source/Runtime/VulkanRHI/Public/VulkanShaderResources.h
--- a/Engine/Source/Runtime/VulkanRHI/Public/VulkanShaderResources.h
+++ b/Engine/Source/Runtime/VulkanRHI/Public/VulkanShaderResources.h
@@ -130,7 +130,7 @@ struct FVulkanShaderHeader
 #endif
 	};
 	TArray<FGlobalInfo>						Globals;
-	TArray<TEnumAsByte<VkDescriptorType>>	GlobalDescriptorTypes;
+	TArray<TEnumAsByte<EVulkanBindingType::EType>>	GlobalDescriptorTypes;
 
 	struct FPackedGlobalInfo
 	{
