VK: more correct accountung of descriptor types used in a descriptor cache

diff --git a/Engine/Source/Runtime/VulkanRHI/Private/VulkanPendingState.cpp b/Engine/Source/Runtime/VulkanRHI/Private/VulkanPendingState.cpp
--- a/Engine/Source/Runtime/VulkanRHI/Private/VulkanPendingState.cpp
+++ b/Engine/Source/Runtime/VulkanRHI/Private/VulkanPendingState.cpp
@@ -663,24 +663,24 @@ FVulkanGenericDescriptorPool::FVulkanGenericDescriptorPool(FVulkanDevice* InDevi
 	, MaxDescriptorSets(InMaxDescriptorSets)
 	, DescriptorPool(VK_NULL_HANDLE)
 {
-	TArray<VkDescriptorPoolSize> Types;
-	Types.Reserve(VK_DESCRIPTOR_TYPE_RANGE_SIZE);
+	VkDescriptorPoolSize Types[VK_DESCRIPTOR_TYPE_RANGE_SIZE];
+	FMemory::Memzero(Types);
 
 	for (uint32 i = 0; i < VK_DESCRIPTOR_TYPE_RANGE_SIZE; ++i)
 	{
+		VkDescriptorType DescriptorType = static_cast<VkDescriptorType>(VK_DESCRIPTOR_TYPE_BEGIN_RANGE + i);
+		
 		float MinSize = FMath::Max(DefaultPoolSizes[i]*InMaxDescriptorSets, 4.0f);
 		PoolSizes[i] = (uint32)FMath::Max(PoolSizesRatio[i]*InMaxDescriptorSets, MinSize);
-
-		VkDescriptorPoolSize* Type = new(Types) VkDescriptorPoolSize;
-		FMemory::Memzero(*Type);
-		Type->type = (VkDescriptorType)i;
-		Type->descriptorCount = PoolSizes[i];
+		
+		Types[i].type = DescriptorType;
+		Types[i].descriptorCount = PoolSizes[i];
 	}
 
 	VkDescriptorPoolCreateInfo PoolInfo;
 	ZeroVulkanStruct(PoolInfo, VK_STRUCTURE_TYPE_DESCRIPTOR_POOL_CREATE_INFO);
-	PoolInfo.poolSizeCount = Types.Num();
-	PoolInfo.pPoolSizes = Types.GetData();
+	PoolInfo.poolSizeCount = VK_DESCRIPTOR_TYPE_RANGE_SIZE;
+	PoolInfo.pPoolSizes = Types;
 	PoolInfo.maxSets = MaxDescriptorSets;
 
 #if VULKAN_ENABLE_AGGRESSIVE_STATS
@@ -866,9 +866,10 @@ bool FVulkanDescriptorSetCache::FCachedPool::CreateDescriptorSets(
 
 		DSWriter.SetDescriptorSet(NewSetEntry.Sets[Index]);
 
-		for (int32 i = VK_DESCRIPTOR_TYPE_BEGIN_RANGE; i <= VK_DESCRIPTOR_TYPE_END_RANGE; ++i)
+		for (int32 i = 0; i < VK_DESCRIPTOR_TYPE_RANGE_SIZE; ++i)
 		{
-			PoolSizesStatistic[i] += SetsLayout.GetTypesUsed(static_cast<VkDescriptorType>(i));
+			VkDescriptorType DescriptorType = static_cast<VkDescriptorType>(VK_DESCRIPTOR_TYPE_BEGIN_RANGE + i);
+			PoolSizesStatistic[i] += SetsLayout.GetTypesUsed(DescriptorType);
 		}
 
 #if VULKAN_ENABLE_AGGRESSIVE_STATS
