Don't initialize Vulkan ray tracing if runtime RT is disabled.

diff --git a/Engine/Source/Runtime/VulkanRHI/Private/VulkanRHI.cpp b/Engine/Source/Runtime/VulkanRHI/Private/VulkanRHI.cpp
--- a/Engine/Source/Runtime/VulkanRHI/Private/VulkanRHI.cpp
+++ b/Engine/Source/Runtime/VulkanRHI/Private/VulkanRHI.cpp
@@ -297,7 +297,10 @@ void FVulkanDynamicRHI::PostInit()
 	TShaderMapRef<FNULLPS>(GetGlobalShaderMap(GMaxRHIFeatureLevel)).GetPixelShader();
 
 #if VULKAN_RHI_RAYTRACING
-	Device->InitializeRayTracing();
+	if (GRHISupportsRayTracing)
+	{
+		Device->InitializeRayTracing();
+	}
 #endif // VULKAN_RHI_RAYTRACING
 }
 
