Disabling vkrt compilation to chase down some failures I didn't have locally.

diff --git a/Engine/Source/Runtime/VulkanRHI/Private/Windows/VulkanWindowsPlatform.h b/Engine/Source/Runtime/VulkanRHI/Private/Windows/VulkanWindowsPlatform.h
--- a/Engine/Source/Runtime/VulkanRHI/Private/Windows/VulkanWindowsPlatform.h
+++ b/Engine/Source/Runtime/VulkanRHI/Private/Windows/VulkanWindowsPlatform.h
@@ -19,7 +19,7 @@
 
 #define	UE_VK_API_VERSION							VK_API_VERSION_1_1
 
-#define VULKAN_RHI_RAYTRACING 						(RHI_RAYTRACING)
+#define VULKAN_RHI_RAYTRACING 						0//(RHI_RAYTRACING)
 #define VULKAN_SUPPORTS_SCALAR_BLOCK_LAYOUT			(VULKAN_RHI_RAYTRACING)
 
 // 32-bit windows has warnings on custom mem mgr callbacks
