We must sub-allocate from the buffer's size, not the allocation (it's possible for MemReq to return a memory size larger than the buffer size).  Fixes some errors.

diff --git a/Engine/Source/Runtime/VulkanRHI/Private/VulkanMemory.cpp b/Engine/Source/Runtime/VulkanRHI/Private/VulkanMemory.cpp
--- a/Engine/Source/Runtime/VulkanRHI/Private/VulkanMemory.cpp
+++ b/Engine/Source/Runtime/VulkanRHI/Private/VulkanMemory.cpp
@@ -2620,12 +2620,13 @@ namespace VulkanRHI
 
 	FVulkanSubresourceAllocator::FVulkanSubresourceAllocator(EVulkanAllocationType InType, FMemoryManager* InOwner, uint8 InSubresourceAllocatorFlags, FDeviceMemoryAllocation* InDeviceMemoryAllocation,
 		uint32 InMemoryTypeIndex, VkMemoryPropertyFlags InMemoryPropertyFlags,
-		uint32 InAlignment, VkBuffer InBuffer, uint32 InBufferId, VkBufferUsageFlags InBufferUsageFlags, int32 InPoolSizeIndex)
+		uint32 InAlignment, VkBuffer InBuffer, uint32 InBufferSize, uint32 InBufferId, VkBufferUsageFlags InBufferUsageFlags, int32 InPoolSizeIndex)
 		: Type(InType)
 		, Owner(InOwner)
 		, MemoryTypeIndex(InMemoryTypeIndex)
 		, MemoryPropertyFlags(InMemoryPropertyFlags)
 		, MemoryAllocation(InDeviceMemoryAllocation)
+		, MaxSize(InBufferSize)
 		, Alignment(InAlignment)
 		, FrameFreed(0)
 		, UsedSize(0)
@@ -2638,7 +2639,6 @@ namespace VulkanRHI
 
 	{
 		FMemory::Memzero(MemoryUsed);
-		MaxSize = InDeviceMemoryAllocation->GetSize();
 
 		if(InDeviceMemoryAllocation->IsMapped())
 		{
@@ -2840,7 +2840,7 @@ namespace VulkanRHI
 			BufferId = ++GVulkanBufferHandleIdCounter;
 		}
 		FVulkanSubresourceAllocator* SubresourceAllocator = new FVulkanSubresourceAllocator(EVulkanAllocationPooledBuffer, this, AllocationFlags, DeviceMemoryAllocation, MemoryTypeIndex,
-			MemoryPropertyFlags, MemReqs.alignment, Buffer, BufferId, BufferUsageFlags, PoolSize);
+			MemoryPropertyFlags, MemReqs.alignment, Buffer, BufferSize, BufferId, BufferUsageFlags, PoolSize);
 
 		RegisterSubresourceAllocator(SubresourceAllocator);
 		UsedBufferAllocations[PoolSize].Add(SubresourceAllocator);
diff --git a/Engine/Source/Runtime/VulkanRHI/Public/VulkanMemory.h b/Engine/Source/Runtime/VulkanRHI/Public/VulkanMemory.h
--- a/Engine/Source/Runtime/VulkanRHI/Public/VulkanMemory.h
+++ b/Engine/Source/Runtime/VulkanRHI/Public/VulkanMemory.h
@@ -568,7 +568,7 @@ namespace VulkanRHI
 	public:
 		FVulkanSubresourceAllocator(EVulkanAllocationType InType, FMemoryManager* InOwner, uint8 InSubResourceAllocatorFlags, FDeviceMemoryAllocation* InDeviceMemoryAllocation,
 			uint32 InMemoryTypeIndex, VkMemoryPropertyFlags InMemoryPropertyFlags,
-			uint32 InAlignment, VkBuffer InBuffer, uint32 InBufferId, VkBufferUsageFlags InBufferUsageFlags, int32 InPoolSizeIndex);
+			uint32 InAlignment, VkBuffer InBuffer, uint32 InBufferSize, uint32 InBufferId, VkBufferUsageFlags InBufferUsageFlags, int32 InPoolSizeIndex);
 
 
 		FVulkanSubresourceAllocator(EVulkanAllocationType InType, FMemoryManager* InOwner, uint8 InSubResourceAllocatorFlags, FDeviceMemoryAllocation* InDeviceMemoryAllocation,
