Fix HDR on Stadia.

diff --git a/Engine/Source/Runtime/VulkanRHI/Private/VulkanSwapChain.cpp b/Engine/Source/Runtime/VulkanRHI/Private/VulkanSwapChain.cpp
--- a/Engine/Source/Runtime/VulkanRHI/Private/VulkanSwapChain.cpp
+++ b/Engine/Source/Runtime/VulkanRHI/Private/VulkanSwapChain.cpp
@@ -147,40 +147,50 @@ FVulkanSwapChain::FVulkanSwapChain(VkInstance InInstance, FVulkanDevice& InDevic
 		FVulkanPlatform::CreateSurface(WindowHandle, Instance, &Surface);
 	}
 
-	static const auto CVarHDROutputDevice = IConsoleManager::Get().FindTConsoleVariableDataInt(TEXT("r.HDR.Display.OutputDevice"));
-	int32 OutputDevice = CVarHDROutputDevice ? CVarHDROutputDevice->GetValueOnAnyThread() : 0;
-	VkColorSpaceKHR RequestedColorSpace;
-	// The possible values are documented in PostProcessTonemap.cpp, where the cvar is defined. They match the ETonemapperOutputDevice enum, which is defined in a header we cannot include.
-	switch (OutputDevice)
-	{
-	case 0:
-		RequestedColorSpace = VK_COLOR_SPACE_SRGB_NONLINEAR_KHR;
-		break;
-	case 1:
-		RequestedColorSpace = VK_COLOR_SPACE_BT709_NONLINEAR_EXT;
-		break;
-	case 3:
-	case 4:
-		RequestedColorSpace = VK_COLOR_SPACE_HDR10_ST2084_EXT;
-		break;
-	default:
-		UE_LOG(LogVulkanRHI, Warning, TEXT("Requested color format %d not supported in Vulkan, falling back to sRGB. Please check the value of r.HDR.Display.OutputDevice."), OutputDevice);
-		RequestedColorSpace = VK_COLOR_SPACE_SRGB_NONLINEAR_KHR;
-		break;
+	uint32 NumFormats;
+	VERIFYVULKANRESULT_EXPANDED(VulkanRHI::vkGetPhysicalDeviceSurfaceFormatsKHR(Device.GetPhysicalHandle(), Surface, &NumFormats, nullptr));
+	check(NumFormats > 0);
+
+	TArray<VkSurfaceFormatKHR> Formats;
+	Formats.AddZeroed(NumFormats);
+	VERIFYVULKANRESULT_EXPANDED(VulkanRHI::vkGetPhysicalDeviceSurfaceFormatsKHR(Device.GetPhysicalHandle(), Surface, &NumFormats, Formats.GetData()));
+
+	VkColorSpaceKHR RequestedColorSpace = Formats[0].colorSpace;
+
+	// If multiple colorspaces are possible, then use CVarHDROutputDevice to narrow it down
+	for (int32 Index = 1; Index < Formats.Num(); ++Index)
+	{
+		if (Formats[Index].colorSpace != RequestedColorSpace)
+		{
+			static const auto CVarHDROutputDevice = IConsoleManager::Get().FindTConsoleVariableDataInt(TEXT("r.HDR.Display.OutputDevice"));
+			int32 OutputDevice = CVarHDROutputDevice ? CVarHDROutputDevice->GetValueOnAnyThread() : 0;
+			// The possible values are documented in PostProcessTonemap.cpp, where the cvar is defined. They match the ETonemapperOutputDevice enum, which is defined in a header we cannot include.
+			switch (OutputDevice)
+			{
+			case 0:
+				RequestedColorSpace = VK_COLOR_SPACE_SRGB_NONLINEAR_KHR;
+				break;
+			case 1:
+				RequestedColorSpace = VK_COLOR_SPACE_BT709_NONLINEAR_EXT;
+				break;
+			case 3:
+			case 4:
+				RequestedColorSpace = VK_COLOR_SPACE_HDR10_ST2084_EXT;
+				break;
+			default:
+				UE_LOG(LogVulkanRHI, Warning, TEXT("Requested color format %d not supported in Vulkan, falling back to sRGB. Please check the value of r.HDR.Display.OutputDevice."), OutputDevice);
+				RequestedColorSpace = VK_COLOR_SPACE_SRGB_NONLINEAR_KHR;
+				break;
+			}
+
+			break;
+		}
 	}
 
 	// Find Pixel format for presentable images
 	VkSurfaceFormatKHR CurrFormat;
 	FMemory::Memzero(CurrFormat);
 	{
-		uint32 NumFormats;
-		VERIFYVULKANRESULT_EXPANDED(VulkanRHI::vkGetPhysicalDeviceSurfaceFormatsKHR(Device.GetPhysicalHandle(), Surface, &NumFormats, nullptr));
-		check(NumFormats > 0);
-
-		TArray<VkSurfaceFormatKHR> Formats;
-		Formats.AddZeroed(NumFormats);
-		VERIFYVULKANRESULT_EXPANDED(VulkanRHI::vkGetPhysicalDeviceSurfaceFormatsKHR(Device.GetPhysicalHandle(), Surface, &NumFormats, Formats.GetData()));
-
 		if (InOutPixelFormat == PF_Unknown)
 		{
 			static const auto* CVarDefaultBackBufferPixelFormat = IConsoleManager::Get().FindTConsoleVariableDataInt(TEXT("r.DefaultBackBufferPixelFormat"));
