VK: Make sure we merge common UBs only across different stages

diff --git a/Engine/Source/Runtime/VulkanRHI/Private/VulkanShaders.cpp b/Engine/Source/Runtime/VulkanRHI/Private/VulkanShaders.cpp
--- a/Engine/Source/Runtime/VulkanRHI/Private/VulkanShaders.cpp
+++ b/Engine/Source/Runtime/VulkanRHI/Private/VulkanShaders.cpp
@@ -496,7 +496,7 @@ void FVulkanDescriptorSetsLayoutInfo::ProcessBindingsForStage(VkShaderStageFlagB
 			{
 				//#todo-rco: Only process constant data part of the UB
 				Found = (UBInfo.ConstantDataOriginalBindingIndex == UINT16_MAX) ? nullptr : OutUBGatherInfo.UBLayoutsToUsedStageMap.Find(UBInfo.LayoutHash);
-				if (Found)
+				if (Found && *Found != StageFlags)
 				{
 					// Move from per stage to common UBs
 					VkShaderStageFlags PrevStage = (VkShaderStageFlags)0;
