VK: Demote dynamic UBs to a regular UBs when we hit device limits on a number of dynamic UBs
Fixed of by one descriptor type limits checks

diff --git a/Engine/Source/Runtime/VulkanRHI/Private/VulkanCommands.cpp b/Engine/Source/Runtime/VulkanRHI/Private/VulkanCommands.cpp
--- a/Engine/Source/Runtime/VulkanRHI/Private/VulkanCommands.cpp
+++ b/Engine/Source/Runtime/VulkanRHI/Private/VulkanCommands.cpp
@@ -433,7 +433,6 @@ inline void FVulkanCommandListContext::SetShaderUniformBuffer(ShaderStage::EStag
 	if (!HeaderUBInfo.bOnlyHasResources)
 	{
 		checkSlow(UniformBuffer->GetLayout().ConstantBufferSize > 0);
-		extern TAutoConsoleVariable<int32> GDynamicGlobalUBs;
 
 		uint8 DescriptorSet;
 		uint32 BindingIndex;
@@ -442,12 +441,15 @@ inline void FVulkanCommandListContext::SetShaderUniformBuffer(ShaderStage::EStag
 			return;
 		}
 
-		if (GDynamicGlobalUBs.GetValueOnAnyThread() > 1)
+		const VkDescriptorType DescriptorType = DescriptorInfo.GetDescriptorType(DescriptorSet, BindingIndex);
+
+		if (DescriptorType == VK_DESCRIPTOR_TYPE_UNIFORM_BUFFER_DYNAMIC)
 		{
 			PendingGfxState->SetUniformBuffer<true>(DescriptorSet, BindingIndex, UniformBuffer);
 		}
 		else
 		{
+			check(DescriptorType == VK_DESCRIPTOR_TYPE_UNIFORM_BUFFER);
 			PendingGfxState->SetUniformBuffer<false>(DescriptorSet, BindingIndex, UniformBuffer);
 		}
 	}
@@ -504,8 +506,7 @@ void FVulkanCommandListContext::RHISetShaderUniformBuffer(FRHIComputeShader* Com
 	if (!HeaderUBInfo.bOnlyHasResources)
 	{
 		checkSlow(UniformBuffer->GetLayout().ConstantBufferSize > 0);
-		extern TAutoConsoleVariable<int32> GDynamicGlobalUBs;
-
+		
 		uint8 DescriptorSet;
 		uint32 BindingIndex;
 		if (!DescriptorInfo.GetDescriptorSetAndBindingIndex(FVulkanShaderHeader::UniformBuffer, BufferIndex, DescriptorSet, BindingIndex))
@@ -513,12 +514,15 @@ void FVulkanCommandListContext::RHISetShaderUniformBuffer(FRHIComputeShader* Com
 			return;
 		}
 
-		if (GDynamicGlobalUBs.GetValueOnAnyThread() > 1)
+		const VkDescriptorType DescriptorType = DescriptorInfo.GetDescriptorType(DescriptorSet, BindingIndex);
+
+		if (DescriptorType == VK_DESCRIPTOR_TYPE_UNIFORM_BUFFER_DYNAMIC)
 		{
 			State.SetUniformBuffer<true>(DescriptorSet, BindingIndex, UniformBuffer);
 		}
 		else
 		{
+			check(DescriptorType == VK_DESCRIPTOR_TYPE_UNIFORM_BUFFER);
 			State.SetUniformBuffer<false>(DescriptorSet, BindingIndex, UniformBuffer);
 		}
 	}
diff --git a/Engine/Source/Runtime/VulkanRHI/Private/VulkanDescriptorSets.h b/Engine/Source/Runtime/VulkanRHI/Private/VulkanDescriptorSets.h
--- a/Engine/Source/Runtime/VulkanRHI/Private/VulkanDescriptorSets.h
+++ b/Engine/Source/Runtime/VulkanRHI/Private/VulkanDescriptorSets.h
@@ -346,7 +346,7 @@ public:
 	void ProcessBindingsForStage(VkShaderStageFlagBits StageFlags, ShaderStage::EStage DescSetStage, const FVulkanShaderHeader& CodeHeader, FUniformBufferGatherInfo& OutUBGatherInfo) const;
 
 	template<bool bIsCompute>
-	void FinalizeBindings(const FUniformBufferGatherInfo& UBGatherInfo, const TArrayView<FRHISamplerState*>& ImmutableSamplers);
+	void FinalizeBindings(const FVulkanDevice& Device, const FUniformBufferGatherInfo& UBGatherInfo, const TArrayView<FRHISamplerState*>& ImmutableSamplers);
 
 	void GenerateHash(const TArrayView<FRHISamplerState*>& ImmutableSamplers);
 
@@ -738,6 +738,11 @@ public:
 		//OutDescriptorSet = RemappingUBInfos[Stage][ParameterIndex].Remapping.NewDescriptorSet;
 		return RemappingInfo->StageInfos[0].Globals;
 	}
+	
+	inline VkDescriptorType GetDescriptorType(uint8 DescriptorSet, int32 DescriptorIndex) const
+	{
+		return RemappingInfo->SetInfos[DescriptorSet].Types[DescriptorIndex];
+	}
 
 	inline bool IsInitialized() const
 	{
@@ -803,6 +808,11 @@ public:
 		return RemappingInfo->StageInfos[Stage].Globals;
 	}
 
+	inline VkDescriptorType GetDescriptorType(uint8 DescriptorSet, int32 DescriptorIndex) const
+	{
+		return RemappingInfo->SetInfos[DescriptorSet].Types[DescriptorIndex];
+	}
+
 	inline bool IsInitialized() const
 	{
 		return bInitialized;
diff --git a/Engine/Source/Runtime/VulkanRHI/Private/VulkanPendingState.cpp b/Engine/Source/Runtime/VulkanRHI/Private/VulkanPendingState.cpp
--- a/Engine/Source/Runtime/VulkanRHI/Private/VulkanPendingState.cpp
+++ b/Engine/Source/Runtime/VulkanRHI/Private/VulkanPendingState.cpp
@@ -651,7 +651,7 @@ const float DefaultPoolSizes[VK_DESCRIPTOR_TYPE_RANGE_SIZE] =
 	1/8.0,	// VK_DESCRIPTOR_TYPE_STORAGE_IMAGE
 	1/2.0,	// VK_DESCRIPTOR_TYPE_UNIFORM_TEXEL_BUFFER
 	1/8.0,	// VK_DESCRIPTOR_TYPE_STORAGE_TEXEL_BUFFER
-	1/8.0,	// VK_DESCRIPTOR_TYPE_UNIFORM_BUFFER
+	1/4.0,	// VK_DESCRIPTOR_TYPE_UNIFORM_BUFFER
 	1/8.0,	// VK_DESCRIPTOR_TYPE_STORAGE_BUFFER
 	4,		// VK_DESCRIPTOR_TYPE_UNIFORM_BUFFER_DYNAMIC
 	1/8.0,	// VK_DESCRIPTOR_TYPE_STORAGE_BUFFER_DYNAMIC
diff --git a/Engine/Source/Runtime/VulkanRHI/Private/VulkanPipeline.cpp b/Engine/Source/Runtime/VulkanRHI/Private/VulkanPipeline.cpp
--- a/Engine/Source/Runtime/VulkanRHI/Private/VulkanPipeline.cpp
+++ b/Engine/Source/Runtime/VulkanRHI/Private/VulkanPipeline.cpp
@@ -1443,7 +1443,7 @@ void FVulkanPipelineStateCacheManager::CreateGfxEntry(const FGraphicsPipelineSta
 			// Second pass
 			const int32 NumImmutableSamplers = PSOInitializer.ImmutableSamplerState.ImmutableSamplers.Num();
 			TArrayView<FRHISamplerState*> ImmutableSamplers(NumImmutableSamplers > 0 ? &(FRHISamplerState*&)PSOInitializer.ImmutableSamplerState.ImmutableSamplers[0] : nullptr, NumImmutableSamplers);
-			DescriptorSetLayoutInfo.FinalizeBindings<false>(UBGatherInfo, ImmutableSamplers);
+			DescriptorSetLayoutInfo.FinalizeBindings<false>(*Device, UBGatherInfo, ImmutableSamplers);
 	}
 
 	FDescriptorSetRemappingInfo& RemappingInfo = DescriptorSetLayoutInfo.RemappingInfo;
@@ -1869,7 +1869,7 @@ FVulkanComputePipeline* FVulkanPipelineStateCacheManager::CreateComputePipelineF
 	const FVulkanShaderHeader& CSHeader = Shader->GetCodeHeader();
 	FUniformBufferGatherInfo UBGatherInfo;
 	DescriptorSetLayoutInfo.ProcessBindingsForStage(VK_SHADER_STAGE_COMPUTE_BIT, ShaderStage::Compute, CSHeader, UBGatherInfo);
-	DescriptorSetLayoutInfo.FinalizeBindings<true>(UBGatherInfo, TArrayView<FRHISamplerState*>());
+	DescriptorSetLayoutInfo.FinalizeBindings<true>(*Device, UBGatherInfo, TArrayView<FRHISamplerState*>());
 	FVulkanLayout* Layout = FindOrAddLayout(DescriptorSetLayoutInfo, false);
 	FVulkanComputeLayout* ComputeLayout = (FVulkanComputeLayout*)Layout;
 	if (!ComputeLayout->ComputePipelineDescriptorInfo.IsInitialized())
diff --git a/Engine/Source/Runtime/VulkanRHI/Private/VulkanPipelineState.cpp b/Engine/Source/Runtime/VulkanRHI/Private/VulkanPipelineState.cpp
--- a/Engine/Source/Runtime/VulkanRHI/Private/VulkanPipelineState.cpp
+++ b/Engine/Source/Runtime/VulkanRHI/Private/VulkanPipelineState.cpp
@@ -18,9 +18,6 @@ enum
 	NumAllocationsPerPool = 8,
 };
 
-
-extern TAutoConsoleVariable<int32> GDynamicGlobalUBs;
-
 #if UE_BUILD_DEBUG || UE_BUILD_DEVELOPMENT
 static TAutoConsoleVariable<int32> GAlwaysWriteDS(
 	TEXT("r.Vulkan.AlwaysWriteDS"),
diff --git a/Engine/Source/Runtime/VulkanRHI/Private/VulkanRHI.cpp b/Engine/Source/Runtime/VulkanRHI/Private/VulkanRHI.cpp
--- a/Engine/Source/Runtime/VulkanRHI/Private/VulkanRHI.cpp
+++ b/Engine/Source/Runtime/VulkanRHI/Private/VulkanRHI.cpp
@@ -1384,41 +1384,37 @@ void FVulkanDescriptorSetsLayout::Compile(FVulkanDescriptorSetLayoutMap& DSetLay
 	// Check for maxDescriptorSetSamplers
 	check(		LayoutTypes[VK_DESCRIPTOR_TYPE_SAMPLER]
 			+	LayoutTypes[VK_DESCRIPTOR_TYPE_COMBINED_IMAGE_SAMPLER]
-			<	Limits.maxDescriptorSetSamplers);
+			<=	Limits.maxDescriptorSetSamplers);
 
 	// Check for maxDescriptorSetUniformBuffers
 	check(		LayoutTypes[VK_DESCRIPTOR_TYPE_UNIFORM_BUFFER]
 			+	LayoutTypes[VK_DESCRIPTOR_TYPE_UNIFORM_BUFFER_DYNAMIC]
-			<	Limits.maxDescriptorSetUniformBuffers);
+			<=	Limits.maxDescriptorSetUniformBuffers);
 
 	// Check for maxDescriptorSetUniformBuffersDynamic
 	check(Device->GetVendorId() == EGpuVendorId::Amd ||
 				LayoutTypes[VK_DESCRIPTOR_TYPE_UNIFORM_BUFFER_DYNAMIC]
-			<	Limits.maxDescriptorSetUniformBuffersDynamic);
+			<=	Limits.maxDescriptorSetUniformBuffersDynamic);
 
 	// Check for maxDescriptorSetStorageBuffers
 	check(		LayoutTypes[VK_DESCRIPTOR_TYPE_STORAGE_BUFFER]
 			+	LayoutTypes[VK_DESCRIPTOR_TYPE_STORAGE_BUFFER_DYNAMIC]
-			<	Limits.maxDescriptorSetStorageBuffers);
+			<=	Limits.maxDescriptorSetStorageBuffers);
 
 	// Check for maxDescriptorSetStorageBuffersDynamic
-	if (LayoutTypes[VK_DESCRIPTOR_TYPE_UNIFORM_BUFFER_DYNAMIC] > Limits.maxDescriptorSetUniformBuffersDynamic)
-	{
-		//#todo-rco: Downgrade to non-dynamic
-	}
 	check(		LayoutTypes[VK_DESCRIPTOR_TYPE_STORAGE_BUFFER_DYNAMIC]
-			<	Limits.maxDescriptorSetStorageBuffersDynamic);
+			<=	Limits.maxDescriptorSetStorageBuffersDynamic);
 
 	// Check for maxDescriptorSetSampledImages
 	check(		LayoutTypes[VK_DESCRIPTOR_TYPE_COMBINED_IMAGE_SAMPLER]
 			+	LayoutTypes[VK_DESCRIPTOR_TYPE_SAMPLED_IMAGE]
 			+	LayoutTypes[VK_DESCRIPTOR_TYPE_UNIFORM_TEXEL_BUFFER]
-			<	Limits.maxDescriptorSetSampledImages);
+			<=	Limits.maxDescriptorSetSampledImages);
 
 	// Check for maxDescriptorSetStorageImages
 	check(		LayoutTypes[VK_DESCRIPTOR_TYPE_STORAGE_IMAGE]
 			+	LayoutTypes[VK_DESCRIPTOR_TYPE_STORAGE_TEXEL_BUFFER]
-			<	Limits.maxDescriptorSetStorageImages);
+			<=	Limits.maxDescriptorSetStorageImages);
 
 	check(LayoutTypes[VK_DESCRIPTOR_TYPE_INPUT_ATTACHMENT] <= Limits.maxDescriptorSetInputAttachments);
 
diff --git a/Engine/Source/Runtime/VulkanRHI/Private/VulkanShaders.cpp b/Engine/Source/Runtime/VulkanRHI/Private/VulkanShaders.cpp
--- a/Engine/Source/Runtime/VulkanRHI/Private/VulkanShaders.cpp
+++ b/Engine/Source/Runtime/VulkanRHI/Private/VulkanShaders.cpp
@@ -521,7 +521,7 @@ void FVulkanDescriptorSetsLayoutInfo::ProcessBindingsForStage(VkShaderStageFlagB
 }
 
 template<bool bIsCompute>
-void FVulkanDescriptorSetsLayoutInfo::FinalizeBindings(const FUniformBufferGatherInfo& UBGatherInfo, const TArrayView<FRHISamplerState*>& ImmutableSamplers)
+void FVulkanDescriptorSetsLayoutInfo::FinalizeBindings(const FVulkanDevice& Device, const FUniformBufferGatherInfo& UBGatherInfo, const TArrayView<FRHISamplerState*>& ImmutableSamplers)
 {
 	checkSlow(RemappingInfo.IsEmpty());
 
@@ -535,6 +535,7 @@ void FVulkanDescriptorSetsLayoutInfo::FinalizeBindings(const FUniformBufferGathe
 	const bool bConvertAllUBsToDynamic = (GDynamicGlobalUBs.GetValueOnAnyThread() > 1);
 	const bool bConvertPackedUBsToDynamic = bConvertAllUBsToDynamic || (GDynamicGlobalUBs.GetValueOnAnyThread() == 1);
 	const bool bConsolidateAllIntoOneSet = GDescriptorSetLayoutMode.GetValueOnAnyThread() == 2;
+	const uint32 MaxDescriptorSetUniformBuffersDynamic = Device.GetLimits().maxDescriptorSetUniformBuffersDynamic;
 
 	uint8	DescriptorStageToSetMapping[ShaderStage::NumStages];
 	FMemory::Memset(DescriptorStageToSetMapping, UINT8_MAX);
@@ -582,7 +583,12 @@ void FVulkanDescriptorSetsLayoutInfo::FinalizeBindings(const FUniformBufferGathe
 			RemappingInfo.StageInfos[Stage].UniformBuffers.Reserve(ShaderHeader->UniformBuffers.Num());
 			for (int32 Index = 0; Index < ShaderHeader->UniformBuffers.Num(); ++Index)
 			{
-				VkDescriptorType Type = bConvertAllUBsToDynamic ? VK_DESCRIPTOR_TYPE_UNIFORM_BUFFER_DYNAMIC : VK_DESCRIPTOR_TYPE_UNIFORM_BUFFER;
+				VkDescriptorType Type = VK_DESCRIPTOR_TYPE_UNIFORM_BUFFER;
+				if (bConvertAllUBsToDynamic && LayoutTypes[VK_DESCRIPTOR_TYPE_UNIFORM_BUFFER_DYNAMIC] < MaxDescriptorSetUniformBuffersDynamic)
+				{
+					Type = VK_DESCRIPTOR_TYPE_UNIFORM_BUFFER_DYNAMIC;
+				}
+				
 				// Here we might mess up with the stageFlags, so reset them every loop
 				Binding.stageFlags = StageFlags;
 				Binding.descriptorType = Type;
@@ -800,5 +806,5 @@ FBoundShaderStateRHIRef FVulkanDynamicRHI::RHICreateBoundShaderState(
 }
 
 
-template void FVulkanDescriptorSetsLayoutInfo::FinalizeBindings<true>(const FUniformBufferGatherInfo& UBGatherInfo, const TArrayView<FRHISamplerState*>& ImmutableSamplers);
-template void FVulkanDescriptorSetsLayoutInfo::FinalizeBindings<false>(const FUniformBufferGatherInfo& UBGatherInfo, const TArrayView<FRHISamplerState*>& ImmutableSamplers);
+template void FVulkanDescriptorSetsLayoutInfo::FinalizeBindings<true>(const FVulkanDevice& Device, const FUniformBufferGatherInfo& UBGatherInfo, const TArrayView<FRHISamplerState*>& ImmutableSamplers);
+template void FVulkanDescriptorSetsLayoutInfo::FinalizeBindings<false>(const FVulkanDevice& Device, const FUniformBufferGatherInfo& UBGatherInfo, const TArrayView<FRHISamplerState*>& ImmutableSamplers);
