PR #8209: Fixed a Bug about Hashing FDeviceMemoryBlockKey (Contributed by [REDACTED_PERSON])

diff --git a/Engine/Source/Runtime/VulkanRHI/Public/VulkanMemory.h b/Engine/Source/Runtime/VulkanRHI/Public/VulkanMemory.h
--- a/Engine/Source/Runtime/VulkanRHI/Public/VulkanMemory.h
+++ b/Engine/Source/Runtime/VulkanRHI/Public/VulkanMemory.h
@@ -333,10 +333,10 @@ namespace VulkanRHI
 		}
 	};
 
-	 FORCEINLINE uint32 GetTypeHash(const FDeviceMemoryBlockKey& BlockKey)
-	 {
-		 return FCrc::MemCrc32(&BlockKey, sizeof(BlockKey), 0);
-	 }
+	FORCEINLINE uint32 GetTypeHash(const FDeviceMemoryBlockKey& BlockKey)
+	{
+		return HashCombine(FCrc::TypeCrc32(BlockKey.MemoryTypeIndex), FCrc::TypeCrc32(BlockKey.BlockSize));
+	}
 
 	struct FDeviceMemoryBlock
 	{
