VulkanRHI: Handle CopySrc/Dest transitions for depth-stencil targets (allows the use of AddCopyTexturePass on depth-stencil of same format).

diff --git a/Engine/Source/Runtime/VulkanRHI/Private/VulkanBarriers.cpp b/Engine/Source/Runtime/VulkanRHI/Private/VulkanBarriers.cpp
--- a/Engine/Source/Runtime/VulkanRHI/Private/VulkanBarriers.cpp
+++ b/Engine/Source/Runtime/VulkanRHI/Private/VulkanBarriers.cpp
@@ -365,7 +365,15 @@ static void GetDepthStencilStageAndAccessFlags(ERHIAccess DepthAccess, ERHIAcces
 
 	Layout = VK_IMAGE_LAYOUT_DEPTH_STENCIL_READ_ONLY_OPTIMAL;
 
-	if (EnumHasAnyFlags(DepthAccess, ERHIAccess::DSVWrite))
+	if (EnumHasAnyFlags(DepthAccess, ERHIAccess::CopySrc))
+	{
+		Layout = VK_IMAGE_LAYOUT_TRANSFER_SRC_OPTIMAL;
+	}
+	else if (EnumHasAnyFlags(DepthAccess, ERHIAccess::CopyDest))
+	{
+		Layout = VK_IMAGE_LAYOUT_TRANSFER_DST_OPTIMAL;
+	}
+	else if (EnumHasAnyFlags(DepthAccess, ERHIAccess::DSVWrite))
 	{
 		if (EnumHasAnyFlags(StencilAccess, ERHIAccess::DSVWrite))
 		{
@@ -431,6 +439,20 @@ static void GetDepthStencilStageAndAccessFlags(ERHIAccess DepthAccess, ERHIAcces
 		ProcessedRHIFlags |= (uint32)ERHIAccess::UAVCompute;
 	}
 
+	if (EnumHasAnyFlags(CombinedAccess, ERHIAccess::CopySrc))
+	{
+		StageFlags |= VK_PIPELINE_STAGE_TRANSFER_BIT;
+		AccessFlags |= VK_ACCESS_TRANSFER_READ_BIT;
+		ProcessedRHIFlags |= (uint32)ERHIAccess::CopySrc;
+	}
+
+	if (EnumHasAnyFlags(CombinedAccess, ERHIAccess::CopyDest))
+	{
+		StageFlags |= VK_PIPELINE_STAGE_TRANSFER_BIT;
+		AccessFlags |= VK_ACCESS_TRANSFER_WRITE_BIT;
+		ProcessedRHIFlags |= (uint32)ERHIAccess::CopyDest;
+	}
+
 	uint32 RemainingFlags = (uint32)CombinedAccess & (~ProcessedRHIFlags);
 	ensureMsgf(RemainingFlags == 0, TEXT("Some access flags were not processed. DepthAccess=%x, StencilAccess=%x, ProcessedRHIFlags=%x, RemainingFlags=%x"), DepthAccess, StencilAccess, ProcessedRHIFlags, RemainingFlags);
 }
