Handle UBMT_RDG_TEXTURE_SRV as SRV in SetShaderUniformBufferResources() function of the Vulkan RHI.

#rnx

diff --git a/Engine/Source/Runtime/VulkanRHI/Private/VulkanCommands.cpp b/Engine/Source/Runtime/VulkanRHI/Private/VulkanCommands.cpp
--- a/Engine/Source/Runtime/VulkanRHI/Private/VulkanCommands.cpp
+++ b/Engine/Source/Runtime/VulkanRHI/Private/VulkanCommands.cpp
@@ -358,6 +358,7 @@ inline void SetShaderUniformBufferResources(FVulkanCommandListContext* Context,
 			break;
 		}
 
+		case UBMT_RDG_TEXTURE_SRV:
 		case UBMT_SRV:
 		case UBMT_RDG_BUFFER_SRV:
 		{
