Handle mixed subresource layouts when moving surfaces.  Return the moved image layout exactly like it was found (keep these background memory operations as transparent as possible).  Add an internal function for the common code between Move/Defrag/Evict.  Copy all layers of images.  Account for layout map resizing.

Remove NumArrayLevels (could be mishandled for cube), use ArraySize and GetNumberOfArrayLevels() instead.

diff --git a/Engine/Source/Runtime/VulkanRHI/Private/VulkanBarriers.cpp b/Engine/Source/Runtime/VulkanRHI/Private/VulkanBarriers.cpp
--- a/Engine/Source/Runtime/VulkanRHI/Private/VulkanBarriers.cpp
+++ b/Engine/Source/Runtime/VulkanRHI/Private/VulkanBarriers.cpp
@@ -1260,6 +1260,63 @@ void FVulkanPipelineBarrier::AddImageLayoutTransition(VkImage Image, VkImageLayo
 	SetupImageBarrier(ImgBarrier, Image, SrcAccessFlags, DstAccessFlags, SrcLayout, DstLayout, SubresourceRange);
 }
 
+void FVulkanPipelineBarrier::AddImageLayoutTransition(VkImage Image, VkImageAspectFlags AspectMask, const FVulkanImageLayout& SrcLayout, VkImageLayout DstLayout)
+{
+	if (SrcLayout.AreAllSubresourcesSameLayout())
+	{
+		AddImageLayoutTransition(Image, SrcLayout.MainLayout, DstLayout, MakeSubresourceRange(AspectMask));
+		return;
+	}
+
+	DstStageMask |= GetVkStageFlagsForLayout(DstLayout);
+	const VkAccessFlags DstAccessFlags = GetVkAccessMaskForLayout(DstLayout);
+
+	VkImageSubresourceRange SubresourceRange = MakeSubresourceRange(AspectMask, 0, 1, 0, 1);
+	for (; SubresourceRange.baseArrayLayer < SrcLayout.NumLayers; ++SubresourceRange.baseArrayLayer)
+	{
+		for (; SubresourceRange.baseMipLevel < SrcLayout.NumMips; ++SubresourceRange.baseMipLevel)
+		{
+			const VkImageLayout SubresourceLayout = SrcLayout.GetSubresLayout(SubresourceRange.baseArrayLayer, SubresourceRange.baseMipLevel);
+			if (SubresourceLayout != DstLayout)
+			{
+				SrcStageMask |= GetVkStageFlagsForLayout(SubresourceLayout);
+				const VkAccessFlags SrcAccessFlags = GetVkAccessMaskForLayout(SubresourceLayout);
+
+				VkImageMemoryBarrier& ImgBarrier = ImageBarriers.AddDefaulted_GetRef();
+				SetupImageBarrier(ImgBarrier, Image, SrcAccessFlags, DstAccessFlags, SubresourceLayout, DstLayout, SubresourceRange);
+			}
+		}
+	}
+}
+
+void FVulkanPipelineBarrier::AddImageLayoutTransition(VkImage Image, VkImageAspectFlags AspectMask, VkImageLayout SrcLayout, const FVulkanImageLayout& DstLayout)
+{
+	if (DstLayout.AreAllSubresourcesSameLayout())
+	{
+		AddImageLayoutTransition(Image, SrcLayout, DstLayout.MainLayout, MakeSubresourceRange(AspectMask));
+		return;
+	}
+
+	SrcStageMask |= GetVkStageFlagsForLayout(SrcLayout);
+	const VkAccessFlags SrcAccessFlags = GetVkAccessMaskForLayout(SrcLayout);
+
+	VkImageSubresourceRange SubresourceRange = MakeSubresourceRange(AspectMask, 0, 1, 0, 1);
+	for (; SubresourceRange.baseArrayLayer < DstLayout.NumLayers; ++SubresourceRange.baseArrayLayer)
+	{
+		for (; SubresourceRange.baseMipLevel < DstLayout.NumMips; ++SubresourceRange.baseMipLevel)
+		{
+			const VkImageLayout SubresourceLayout = DstLayout.GetSubresLayout(SubresourceRange.baseArrayLayer, SubresourceRange.baseMipLevel);
+			if (SubresourceLayout != SrcLayout)
+			{
+				DstStageMask |= GetVkStageFlagsForLayout(SubresourceLayout);
+				const VkAccessFlags DstAccessFlags = GetVkAccessMaskForLayout(SubresourceLayout);
+
+				VkImageMemoryBarrier& ImgBarrier = ImageBarriers.AddDefaulted_GetRef();
+				SetupImageBarrier(ImgBarrier, Image, SrcAccessFlags, DstAccessFlags, SrcLayout, SubresourceLayout, SubresourceRange);
+			}
+		}
+	}
+}
 void FVulkanPipelineBarrier::AddImageAccessTransition(const FVulkanSurface& Surface, ERHIAccess SrcAccess, ERHIAccess DstAccess, const VkImageSubresourceRange& SubresourceRange, VkImageLayout& InOutLayout)
 {
 	// This function should only be used for known states.
diff --git a/Engine/Source/Runtime/VulkanRHI/Private/VulkanBarriers.h b/Engine/Source/Runtime/VulkanRHI/Private/VulkanBarriers.h
--- a/Engine/Source/Runtime/VulkanRHI/Private/VulkanBarriers.h
+++ b/Engine/Source/Runtime/VulkanRHI/Private/VulkanBarriers.h
@@ -25,6 +25,8 @@ struct FVulkanPipelineBarrier
 	TArray<FVulkanTextureBase*, TInlineAllocator<2>> Textures;
 
 	void AddImageLayoutTransition(VkImage Image, VkImageLayout SrcLayout, VkImageLayout DstLayout, const VkImageSubresourceRange& SubresourceRange);
+	void AddImageLayoutTransition(VkImage Image, VkImageAspectFlags AspectMask, const struct FVulkanImageLayout& SrcLayout, VkImageLayout DstLayout);
+	void AddImageLayoutTransition(VkImage Image, VkImageAspectFlags AspectMask, VkImageLayout SrcLayout, const struct FVulkanImageLayout& DstLayout);
 	void AddImageAccessTransition(const FVulkanSurface& Surface, ERHIAccess SrcAccess, ERHIAccess DstAccess, const VkImageSubresourceRange& SubresourceRange, VkImageLayout& InOutLayout);
 	void Execute(VkCommandBuffer CmdBuffer);
 
@@ -176,17 +178,20 @@ public:
 		return Layout.MainLayout;
 	}
 
-	VULKANRHI_API VkImageLayout& FindOrAddLayoutRW(VkImage Image, VkImageLayout LayoutIfNotFound, uint32 NumMips, uint32 NumLayers)
+	VULKANRHI_API FVulkanImageLayout& FindOrAddFullLayoutRW(VkImage Image, VkImageLayout LayoutIfNotFound, uint32 NumMips, uint32 NumLayers)
 	{
 		FVulkanImageLayout* Layout = Layouts.Find(Image);
 		if (Layout)
 		{
 			check(Layout->AreAllSubresourcesSameLayout());
-			return Layout->MainLayout;
+			return *Layout;
 		}
+		return Layouts.Add(Image, FVulkanImageLayout(LayoutIfNotFound, NumMips, NumLayers));
+	}
 
-		Layout = &Layouts.Add(Image, FVulkanImageLayout(LayoutIfNotFound, NumMips, NumLayers));
-		return Layout->MainLayout;
+	VULKANRHI_API VkImageLayout& FindOrAddLayoutRW(VkImage Image, VkImageLayout LayoutIfNotFound, uint32 NumMips, uint32 NumLayers)
+	{
+		return FindOrAddFullLayoutRW(Image, LayoutIfNotFound, NumMips, NumLayers).MainLayout;
 	}
 
 	VULKANRHI_API VkImageLayout& FindOrAddLayoutRW(const FVulkanSurface& Surface, VkImageLayout LayoutIfNotFound)
diff --git a/Engine/Source/Runtime/VulkanRHI/Private/VulkanRenderTarget.cpp b/Engine/Source/Runtime/VulkanRHI/Private/VulkanRenderTarget.cpp
--- a/Engine/Source/Runtime/VulkanRHI/Private/VulkanRenderTarget.cpp
+++ b/Engine/Source/Runtime/VulkanRHI/Private/VulkanRenderTarget.cpp
@@ -1069,7 +1069,7 @@ FVulkanRenderTargetLayout::FVulkanRenderTargetLayout(FVulkanDevice& InDevice, co
 			bSetExtent = true;
 			Extent.Extent3D.width = Texture->Surface.Width;
 			Extent.Extent3D.height = Texture->Surface.Height;
-			Extent.Extent3D.depth = Texture->Surface.NumArrayLevels;
+			Extent.Extent3D.depth = Texture->Surface.GetNumberOfArrayLevels();
 		}
 	}
 
@@ -1178,8 +1178,8 @@ FVulkanRenderTargetLayout::FVulkanRenderTargetLayout(FVulkanDevice& InDevice, co
 		ensure(!NumSamples || NumSamples == ColorEntry.RenderTarget->GetNumSamples());
 		NumSamples = ColorEntry.RenderTarget->GetNumSamples();
 
-		ensure( !bMultiviewRenderTargets || Texture->Surface.NumArrayLevels > 1 );
-		bMultiviewRenderTargets = Texture->Surface.NumArrayLevels > 1;
+		ensure( !bMultiviewRenderTargets || Texture->Surface.GetNumberOfArrayLevels() > 1 );
+		bMultiviewRenderTargets = Texture->Surface.GetNumberOfArrayLevels() > 1;
 
 		VkAttachmentDescription& CurrDesc = Desc[NumAttachmentDescriptions];
 		CurrDesc.samples = static_cast<VkSampleCountFlagBits>(NumSamples);
diff --git a/Engine/Source/Runtime/VulkanRHI/Private/VulkanResources.h b/Engine/Source/Runtime/VulkanRHI/Private/VulkanResources.h
--- a/Engine/Source/Runtime/VulkanRHI/Private/VulkanResources.h
+++ b/Engine/Source/Runtime/VulkanRHI/Private/VulkanResources.h
@@ -379,11 +379,11 @@ public:
 		case VK_IMAGE_VIEW_TYPE_3D:
 			return 1;
 		case VK_IMAGE_VIEW_TYPE_2D_ARRAY:
-			return NumArrayLevels;
+			return ArraySize;
 		case VK_IMAGE_VIEW_TYPE_CUBE:
 			return 6;
 		case VK_IMAGE_VIEW_TYPE_CUBE_ARRAY:
-			return 6 * NumArrayLevels;
+			return 6 * ArraySize;
 		default:
 			ErrorInvalidViewType();
 			return 1;
@@ -431,7 +431,6 @@ public:
 	ETextureCreateFlags UEFlags;
 	VkMemoryPropertyFlags MemProps;
 	VkMemoryRequirements MemoryRequirements;
-	uint32 NumArrayLevels;
 
 	static void InternalLockWrite(FVulkanCommandListContext& Context, FVulkanSurface* Surface, const VkImageSubresourceRange& SubresourceRange, const VkBufferImageCopy& Region, VulkanRHI::FStagingBuffer* StagingBuffer);
 
@@ -441,6 +440,8 @@ private:
 	void SetInitialImageState(FVulkanCommandListContext& Context, VkImageLayout InitialLayout, bool bClear, const FClearValueBinding& ClearValueBinding);
 	friend struct FRHICommandSetInitialImageState;
 
+	void InternalMoveSurface(FVulkanDevice& InDevice, FVulkanCommandListContext& Context, VulkanRHI::FVulkanAllocation& DestAllocation, bool bSwapAllocation);
+
 private:
 	VkImageTiling Tiling;
 	VkImageViewType	ViewType;
diff --git a/Engine/Source/Runtime/VulkanRHI/Private/VulkanTexture.cpp b/Engine/Source/Runtime/VulkanRHI/Private/VulkanTexture.cpp
--- a/Engine/Source/Runtime/VulkanRHI/Private/VulkanTexture.cpp
+++ b/Engine/Source/Runtime/VulkanRHI/Private/VulkanTexture.cpp
@@ -536,7 +536,6 @@ FVulkanSurface::FVulkanSurface(FVulkanDevice& InDevice, FVulkanEvictable* Owner,
 	, PixelFormat(InFormat)
 	, UEFlags(InUEFlags)
 	, MemProps(VK_MEMORY_PROPERTY_DEVICE_LOCAL_BIT)
-	, NumArrayLevels(0)
 	, Tiling(VK_IMAGE_TILING_MAX_ENUM)	// Can be expanded to a per-platform definition
 	, ViewType(ResourceType)
 	, bIsImageOwner(true)
@@ -596,8 +595,6 @@ FVulkanSurface::FVulkanSurface(FVulkanDevice& InDevice, FVulkanEvictable* Owner,
 	VulkanRHI::vkGetImageMemoryRequirements(InDevice.GetInstanceHandle(), Image, &MemoryRequirements);
 
 	VULKAN_SET_DEBUG_NAME(InDevice, VK_OBJECT_TYPE_IMAGE, Image, TEXT("(FVulkanSurface*)0x%p"), this);
-	uint32 LayerCount = (ResourceType == VK_IMAGE_VIEW_TYPE_CUBE || ResourceType == VK_IMAGE_VIEW_TYPE_CUBE_ARRAY) ? 6 : 1;
-	NumArrayLevels = ArraySize * LayerCount;
 
 	FullAspectMask = VulkanRHI::GetAspectMaskFromUEFormat(PixelFormat, true, true);
 	PartialAspectMask = VulkanRHI::GetAspectMaskFromUEFormat(PixelFormat, false, true);
@@ -680,7 +677,9 @@ FVulkanSurface::FVulkanSurface(FVulkanDevice& InDevice, FVulkanEvictable* Owner,
 	}
 }
 
-void FVulkanSurface::MoveSurface(FVulkanDevice& InDevice, FVulkanCommandListContext& Context, FVulkanAllocation& NewAllocation)
+
+
+void FVulkanSurface::InternalMoveSurface(FVulkanDevice& InDevice, FVulkanCommandListContext& Context, FVulkanAllocation& DestAllocation, bool bSwapAllocation)
 {
 	FImageCreateInfo ImageCreateInfo;
 	FVulkanSurface::GenerateImageCreateInfo(ImageCreateInfo,
@@ -691,327 +690,168 @@ void FVulkanSurface::MoveSurface(FVulkanDevice& InDevice, FVulkanCommandListCont
 
 	VkImage MovedImage;
 	VERIFYVULKANRESULT(VulkanRHI::vkCreateImage(InDevice.GetInstanceHandle(), &ImageCreateInfo.ImageCreateInfo, VULKAN_CPU_ALLOCATOR, &MovedImage));
-	check(ImageCreateInfo.ImageCreateInfo.tiling == VK_IMAGE_TILING_OPTIMAL);
-
+	checkf(Tiling == ImageCreateInfo.ImageCreateInfo.tiling, TEXT("Move has changed image tiling:  before [%d] != after [%d]"), (int32)Tiling, (int32)ImageCreateInfo.ImageCreateInfo.tiling);
 
 	const bool bRenderTarget = EnumHasAnyFlags(UEFlags, TexCreate_RenderTargetable | TexCreate_DepthStencilTargetable | TexCreate_ResolveTargetable);
-	const bool bUAV = EnumHasAnyFlags(UEFlags, TexCreate_UAV);
 	const bool bCPUReadback = EnumHasAnyFlags(UEFlags, TexCreate_CPUReadback);
-	const bool bDynamic = EnumHasAnyFlags(UEFlags, TexCreate_Dynamic);
 	const bool bMemoryless = EnumHasAnyFlags(UEFlags, TexCreate_Memoryless);
-	check(!bCPUReadback); //not currently supported
-	check(!bMemoryless);  //not currently supported
-
-	VkMemoryPropertyFlags MemoryFlags = VK_MEMORY_PROPERTY_DEVICE_LOCAL_BIT;
-
-	check(bRenderTarget || bUAV);
+	checkf(!bCPUReadback, TEXT("Move of CPUReadback surfaces not currently supported.   UEFlags=0x%x"), (int32)UEFlags);
+	checkf(!bMemoryless, TEXT("Move of Memoryless surfaces not currently supported.   UEFlags=0x%x"), (int32)UEFlags);
 
+#if UE_BUILD_DEBUG || UE_BUILD_DEVELOPMENT
+	// This shouldn't change
+	VkMemoryRequirements MovedMemReqs;
+	VulkanRHI::vkGetImageMemoryRequirements(InDevice.GetInstanceHandle(), MovedImage, &MovedMemReqs);
+	checkf((MemoryRequirements.alignment == MovedMemReqs.alignment), TEXT("Memory requirements changed: alignment %d -> %d"), (int32)MemoryRequirements.alignment, (int32)MovedMemReqs.alignment);
+	checkf((MemoryRequirements.size == MovedMemReqs.size), TEXT("Memory requirements changed: size %d -> %d"), (int32)MemoryRequirements.size, (int32)MovedMemReqs.size);
+	checkf((MemoryRequirements.memoryTypeBits == MovedMemReqs.memoryTypeBits), TEXT("Memory requirements changed: memoryTypeBits %d -> %d"), (int32)MemoryRequirements.memoryTypeBits, (int32)MovedMemReqs.memoryTypeBits);
+#endif // UE_BUILD_DEBUG || UE_BUILD_DEVELOPMENT
 
-	NewAllocation.BindImage(Device, MovedImage);
+	DestAllocation.BindImage(&InDevice, MovedImage);
 
-	uint32 LayerCount = (ViewType == VK_IMAGE_VIEW_TYPE_CUBE || ViewType == VK_IMAGE_VIEW_TYPE_CUBE_ARRAY) ? 6 : 1;
-	NumArrayLevels = ArraySize * LayerCount;
-
-	//Copy Old -> Moved
+	// Copy Original -> Moved
 	FVulkanCmdBuffer* CmdBuffer = Context.GetCommandBufferManager()->GetActiveCmdBuffer();
 	VkCommandBuffer VkCmdBuffer = CmdBuffer->GetHandle();
-
 	ensure(CmdBuffer->IsOutsideRenderPass());
-	VkImageLayout& Layout = Context.GetLayoutManager().FindOrAddLayoutRW(Image, VK_IMAGE_LAYOUT_UNDEFINED, NumMips, LayerCount);
-	VkImageLayout& MovedLayout = Context.GetLayoutManager().FindOrAddLayoutRW(MovedImage, VK_IMAGE_LAYOUT_UNDEFINED, NumMips, LayerCount);
-	{
-		check(MovedLayout == VK_IMAGE_LAYOUT_UNDEFINED);
-		VkImageSubresourceRange RangeAll;
-		RangeAll.aspectMask = FullAspectMask;
-		RangeAll.baseMipLevel = 0;
-		RangeAll.levelCount = VK_REMAINING_MIP_LEVELS;
-		RangeAll.baseArrayLayer = 0;
-		RangeAll.layerCount = VK_REMAINING_ARRAY_LAYERS;
+
+	const uint32 NumberOfArrayLevels = GetNumberOfArrayLevels();
+	FVulkanImageLayout* MovedLayout = &Context.GetLayoutManager().FindOrAddFullLayoutRW(MovedImage, VK_IMAGE_LAYOUT_UNDEFINED, GetNumMips(), NumberOfArrayLevels);
+	FVulkanImageLayout* OriginalLayout = Context.GetLayoutManager().GetFullLayout(Image);
+	// Account for map resize, should rarely happen...
+	if (OriginalLayout == nullptr)
+	{
+		OriginalLayout = &Context.GetLayoutManager().GetOrAddFullLayout(*this, VK_IMAGE_LAYOUT_UNDEFINED);
+		MovedLayout = &Context.GetLayoutManager().GetFullLayoutChecked(MovedImage);
+	}
+
+	checkf((OriginalLayout->NumMips == GetNumMips()), TEXT("NumMips reported by LayoutManager (%d) differs from surface (%d)"), OriginalLayout->NumMips, GetNumMips());
+	checkf((OriginalLayout->NumLayers == NumberOfArrayLevels), TEXT("NumLayers reported by LayoutManager (%d) differs from surface (%d)"), OriginalLayout->NumLayers, NumberOfArrayLevels);
+	{
+		// Transition to copying layouts
 		{
 			FVulkanPipelineBarrier Barrier;
-			Barrier.AddImageLayoutTransition(Image, Layout, VK_IMAGE_LAYOUT_TRANSFER_SRC_OPTIMAL, RangeAll);
-			Barrier.AddImageLayoutTransition(MovedImage, MovedLayout, VK_IMAGE_LAYOUT_TRANSFER_DST_OPTIMAL, RangeAll);
+			Barrier.AddImageLayoutTransition(Image, FullAspectMask, *OriginalLayout, VK_IMAGE_LAYOUT_TRANSFER_SRC_OPTIMAL);
+			Barrier.AddImageLayoutTransition(MovedImage, FullAspectMask, *MovedLayout, VK_IMAGE_LAYOUT_TRANSFER_DST_OPTIMAL);
 			Barrier.Execute(VkCmdBuffer);
 		}
 		{
+			VkImageCopy Regions[MAX_TEXTURE_MIP_COUNT];
+			check(NumMips <= MAX_TEXTURE_MIP_COUNT);
+			FMemory::Memzero(Regions);
+			for (uint32 i = 0; i < NumMips; ++i)
 			{
-				VkImageCopy Regions[MAX_TEXTURE_MIP_COUNT];
-				check(NumMips <= MAX_TEXTURE_MIP_COUNT);
-				FMemory::Memzero(Regions);
-				for (uint32 i = 0; i < NumMips; ++i)
-				{
-					VkImageCopy& Region = Regions[i];
-					Region.extent.width = FMath::Max(1u, Width >> i);
-					Region.extent.height = FMath::Max(1u, Height >> i);
-					Region.extent.depth = FMath::Max(1u, Depth >> i);
-					Region.srcSubresource.aspectMask = FullAspectMask;
-					Region.dstSubresource.aspectMask = FullAspectMask;
-					Region.srcSubresource.baseArrayLayer = 0;
-					Region.dstSubresource.baseArrayLayer = 0;
-					Region.srcSubresource.layerCount = 1;
-					Region.dstSubresource.layerCount = 1;
-					Region.srcSubresource.mipLevel = i;
-					Region.dstSubresource.mipLevel = i;
-				}
-				VulkanRHI::vkCmdCopyImage(VkCmdBuffer,
-					Image, VK_IMAGE_LAYOUT_TRANSFER_SRC_OPTIMAL,
-					MovedImage, VK_IMAGE_LAYOUT_TRANSFER_DST_OPTIMAL,
-					NumMips, &Regions[0]);
+				VkImageCopy& Region = Regions[i];
+				Region.extent.width = FMath::Max(1u, Width >> i);
+				Region.extent.height = FMath::Max(1u, Height >> i);
+				Region.extent.depth = FMath::Max(1u, Depth >> i);
+				Region.srcSubresource.aspectMask = FullAspectMask;
+				Region.dstSubresource.aspectMask = FullAspectMask;
+				Region.srcSubresource.baseArrayLayer = 0;
+				Region.dstSubresource.baseArrayLayer = 0;
+				Region.srcSubresource.layerCount = NumberOfArrayLevels;
+				Region.dstSubresource.layerCount = NumberOfArrayLevels;
+				Region.srcSubresource.mipLevel = i;
+				Region.dstSubresource.mipLevel = i;
 			}
+			VulkanRHI::vkCmdCopyImage(VkCmdBuffer,
+				Image, VK_IMAGE_LAYOUT_TRANSFER_SRC_OPTIMAL,
+				MovedImage, VK_IMAGE_LAYOUT_TRANSFER_DST_OPTIMAL,
+				NumMips, &Regions[0]);
 		}
-		VkImageLayout FinalLayout = Layout;
-		{
-			VulkanSetImageLayout(VkCmdBuffer, MovedImage, VK_IMAGE_LAYOUT_TRANSFER_DST_OPTIMAL, FinalLayout, RangeAll);
-		}
-		Layout = VK_IMAGE_LAYOUT_TRANSFER_SRC_OPTIMAL;
-		MovedLayout = FinalLayout;
 
+		// Put the destination image in exactly the same layout the original image was
 		{
-			check(Image != VK_NULL_HANDLE);
-			uint64 Size = GetMemorySize();
-			Device->NotifyDeletedImage(Image, bRenderTarget);
-			Device->GetDeferredDeletionQueue().EnqueueResource(VulkanRHI::FDeferredDeletionQueue2::EType::Image, Image);
-			Image = MovedImage;
-			VulkanTextureDestroyed(Size, ViewType, bRenderTarget);
-			Allocation.Swap(NewAllocation);
+			FVulkanPipelineBarrier Barrier;
+			Barrier.AddImageLayoutTransition(MovedImage, FullAspectMask, VK_IMAGE_LAYOUT_TRANSFER_DST_OPTIMAL, *OriginalLayout);
+			Barrier.Execute(VkCmdBuffer);
 		}
-	}
 
-}
+		// Update the tracked layouts
+		*MovedLayout = *OriginalLayout;
+		OriginalLayout->Set(VK_IMAGE_LAYOUT_TRANSFER_SRC_OPTIMAL, FVulkanPipelineBarrier::MakeSubresourceRange(FullAspectMask));
+	}
 
+	{
+		check(Image != VK_NULL_HANDLE);
+		InDevice.NotifyDeletedImage(Image, bRenderTarget);
+		InDevice.GetDeferredDeletionQueue().EnqueueResource(VulkanRHI::FDeferredDeletionQueue2::EType::Image, Image);
+		FGenericPlatformMisc::LowLevelOutputDebugStringf(TEXT("** MOVE IMAGE %p -> %p\n"), Image, MovedImage);
+	}
 
-void FVulkanSurface::OnFullDefrag(FVulkanDevice& InDevice, FVulkanCommandListContext& Context, uint32 NewOffset)
-{
-	FImageCreateInfo ImageCreateInfo;
-	FVulkanSurface::GenerateImageCreateInfo(ImageCreateInfo,
-		InDevice, ViewType,
-		PixelFormat, Width, Height, Depth,
-		ArraySize, NumMips, NumSamples, UEFlags,
-		&StorageFormat, &ViewFormat);
+	Image = MovedImage;
 
-	VkImage MovedImage;
-	VERIFYVULKANRESULT(VulkanRHI::vkCreateImage(InDevice.GetInstanceHandle(), &ImageCreateInfo.ImageCreateInfo, VULKAN_CPU_ALLOCATOR, &MovedImage));
-	check(ImageCreateInfo.ImageCreateInfo.tiling == VK_IMAGE_TILING_OPTIMAL);
+	if (bSwapAllocation)
+	{
+		const uint64 Size = GetMemorySize();
+		VulkanTextureDestroyed(Size, ViewType, bRenderTarget);
 
+		Allocation.Swap(DestAllocation);
+	}
+}
 
+void FVulkanSurface::MoveSurface(FVulkanDevice& InDevice, FVulkanCommandListContext& Context, FVulkanAllocation& NewAllocation)
+{
 	const bool bRenderTarget = EnumHasAnyFlags(UEFlags, TexCreate_RenderTargetable | TexCreate_DepthStencilTargetable | TexCreate_ResolveTargetable);
 	const bool bUAV = EnumHasAnyFlags(UEFlags, TexCreate_UAV);
-	const bool bCPUReadback = EnumHasAnyFlags(UEFlags, TexCreate_CPUReadback);
-	const bool bDynamic = EnumHasAnyFlags(UEFlags, TexCreate_Dynamic);
-	const bool bMemoryless = EnumHasAnyFlags(UEFlags, TexCreate_Memoryless);
-	check(!bCPUReadback); //not currently supported
-	check(!bMemoryless);  //not currently supported
-
-	VkMemoryPropertyFlags MemoryFlags = VK_MEMORY_PROPERTY_DEVICE_LOCAL_BIT;
+	checkf(bRenderTarget || bUAV, TEXT("Surface must be a RenderTarget or a UAV in order to be moved.  UEFlags=0x%x"), (int32)UEFlags);
+	checkf(Tiling == VK_IMAGE_TILING_OPTIMAL, TEXT("Tiling [%d] is not supported for move, only VK_IMAGE_TILING_OPTIMAL"), (int32)Tiling);
 
-	check(bRenderTarget || bUAV);
-
-	//note: this exploits that the unmoved image is still bound to the old allocation, which is freed by the caller in this case.
-	Allocation.Offset = NewOffset;
-	Allocation.BindImage(Device, MovedImage);
-
-	uint32 LayerCount = (ViewType == VK_IMAGE_VIEW_TYPE_CUBE || ViewType == VK_IMAGE_VIEW_TYPE_CUBE_ARRAY) ? 6 : 1;
-	NumArrayLevels = ArraySize * LayerCount;
+	InternalMoveSurface(InDevice, Context, NewAllocation, true);
+}
 
 
-	//Copy Old -> Moved
-	FVulkanCmdBuffer* CmdBuffer = Context.GetCommandBufferManager()->GetActiveCmdBuffer();
-	VkCommandBuffer VkCmdBuffer = CmdBuffer->GetHandle();
+void FVulkanSurface::OnFullDefrag(FVulkanDevice& InDevice, FVulkanCommandListContext& Context, uint32 NewOffset)
+{
+	const bool bRenderTarget = EnumHasAnyFlags(UEFlags, TexCreate_RenderTargetable | TexCreate_DepthStencilTargetable | TexCreate_ResolveTargetable);
+	const bool bUAV = EnumHasAnyFlags(UEFlags, TexCreate_UAV);
+	checkf(bRenderTarget || bUAV, TEXT("Surface must be a RenderTarget or a UAV in order to be defragged.  UEFlags=0x%x"), (int32)UEFlags);
+	checkf(Tiling == VK_IMAGE_TILING_OPTIMAL, TEXT("Tiling [%d] is not supported for defrag, only VK_IMAGE_TILING_OPTIMAL"), (int32)Tiling);
 
-	ensure(CmdBuffer->IsOutsideRenderPass());
-	VkImageLayout& Layout = Context.GetLayoutManager().FindOrAddLayoutRW(Image, VK_IMAGE_LAYOUT_UNDEFINED, NumMips, LayerCount);
-	VkImageLayout& MovedLayout = Context.GetLayoutManager().FindOrAddLayoutRW(MovedImage, VK_IMAGE_LAYOUT_UNDEFINED, NumMips, LayerCount);
-	{
-		check(MovedLayout == VK_IMAGE_LAYOUT_UNDEFINED);
-		VkImageSubresourceRange RangeAll;
-		RangeAll.aspectMask = FullAspectMask;
-		RangeAll.baseMipLevel = 0;
-		RangeAll.levelCount = VK_REMAINING_MIP_LEVELS;
-		RangeAll.baseArrayLayer = 0;
-		RangeAll.layerCount = VK_REMAINING_ARRAY_LAYERS;
-		{
-			FVulkanPipelineBarrier Barrier;
-			Barrier.AddImageLayoutTransition(Image, Layout, VK_IMAGE_LAYOUT_TRANSFER_SRC_OPTIMAL, RangeAll);
-			Barrier.AddImageLayoutTransition(MovedImage, MovedLayout, VK_IMAGE_LAYOUT_TRANSFER_DST_OPTIMAL, RangeAll);
-			Barrier.Execute(VkCmdBuffer);
-		}
-		{
-			{
-				VkImageCopy Regions[MAX_TEXTURE_MIP_COUNT];
-				check(NumMips <= MAX_TEXTURE_MIP_COUNT);
-				FMemory::Memzero(Regions);
-				for (uint32 i = 0; i < NumMips; ++i)
-				{
-					VkImageCopy& Region = Regions[i];
-					Region.extent.width = FMath::Max(1u, Width >> i);
-					Region.extent.height = FMath::Max(1u, Height >> i);
-					Region.extent.depth = FMath::Max(1u, Depth >> i);
-					Region.srcSubresource.aspectMask = FullAspectMask;
-					Region.dstSubresource.aspectMask = FullAspectMask;
-					Region.srcSubresource.baseArrayLayer = 0;
-					Region.dstSubresource.baseArrayLayer = 0;
-					Region.srcSubresource.layerCount = 1;
-					Region.dstSubresource.layerCount = 1;
-					Region.srcSubresource.mipLevel = i;
-					Region.dstSubresource.mipLevel = i;
-				}
-				VulkanRHI::vkCmdCopyImage(VkCmdBuffer,
-					Image, VK_IMAGE_LAYOUT_TRANSFER_SRC_OPTIMAL,
-					MovedImage, VK_IMAGE_LAYOUT_TRANSFER_DST_OPTIMAL,
-					NumMips, &Regions[0]);
-			}
-		}
-		VkImageLayout FinalLayout = Layout;
-		{
-			VulkanSetImageLayout(VkCmdBuffer, MovedImage, VK_IMAGE_LAYOUT_TRANSFER_DST_OPTIMAL, FinalLayout, RangeAll);
-		}
-		Layout = VK_IMAGE_LAYOUT_TRANSFER_SRC_OPTIMAL;
-		MovedLayout = FinalLayout;
+	Allocation.Offset = NewOffset;
+	InternalMoveSurface(InDevice, Context, Allocation, false);
 
-		{
-			check(Image != VK_NULL_HANDLE);
-			Device->NotifyDeletedImage(Image, bRenderTarget);
-			Device->GetDeferredDeletionQueue().EnqueueResource(VulkanRHI::FDeferredDeletionQueue2::EType::Image, Image);
-			FGenericPlatformMisc::LowLevelOutputDebugStringf(TEXT("** MOVE IMAGE %p -> %p\n"), Image, MovedImage);
-			Image = MovedImage;
-		}
-	}
+	//note: this exploits that the unmoved image is still bound to the old allocation, which is freed by the caller in this case.
 }
 
 
 void FVulkanSurface::EvictSurface(FVulkanDevice& InDevice)
 {
 	check(0 == CpuReadbackBuffer);
-	FImageCreateInfo ImageCreateInfo;
-	FVulkanSurface::GenerateImageCreateInfo(ImageCreateInfo,
-		InDevice, ViewType,
-		PixelFormat, Width, Height, Depth,
-		ArraySize, NumMips, NumSamples, UEFlags,
-		&StorageFormat, &ViewFormat);
+	checkf(MemProps == VK_MEMORY_PROPERTY_DEVICE_LOCAL_BIT, TEXT("Can't evict surface that isn't device local.  MemoryProperties=%d"), (int32)MemProps);
+	checkf(VulkanRHI::GetAspectMaskFromUEFormat(PixelFormat, true, true) == FullAspectMask, TEXT("FullAspectMask (%d) does not match with PixelFormat (%d)"), (int32)FullAspectMask, (int32)PixelFormat);
+	checkf(VulkanRHI::GetAspectMaskFromUEFormat(PixelFormat, false, true) == PartialAspectMask, TEXT("PartialAspectMask (%d) does not match with PixelFormat (%d)"), (int32)PartialAspectMask, (int32)PixelFormat);
 
+	const bool bRenderTarget = EnumHasAnyFlags(UEFlags, TexCreate_RenderTargetable | TexCreate_DepthStencilTargetable | TexCreate_ResolveTargetable);
+	const bool bUAV = EnumHasAnyFlags(UEFlags, TexCreate_UAV);
+	//none of this is supported for eviction
+	checkf(!bRenderTarget, TEXT("RenderTargets do not support evict."));
+	checkf(!bUAV, TEXT("UAV do not support evict."));
 
 	FRHICommandListImmediate& RHICmdList = FRHICommandListExecutor::GetImmediateCommandList();
 	FVulkanCommandListContext& Context = (FVulkanCommandListContext&)RHICmdList.GetContext().GetLowestLevelContext();
 
+	MemProps = InDevice.GetDeviceMemoryManager().GetEvictedMemoryProperties();
 
-	VkImage HostImage;
 	FVulkanAllocation HostAllocation;
-	VERIFYVULKANRESULT(VulkanRHI::vkCreateImage(InDevice.GetInstanceHandle(), &ImageCreateInfo.ImageCreateInfo, VULKAN_CPU_ALLOCATOR, &HostImage));
-
-	VulkanRHI::vkGetImageMemoryRequirements(InDevice.GetInstanceHandle(), HostImage, &MemoryRequirements);
-
-	VULKAN_SET_DEBUG_NAME(InDevice, VK_OBJECT_TYPE_IMAGE, HostImage, TEXT("(FVulkanSurface*)0x%p [hostimage]"), this);
-	uint32 LayerCount = (ViewType == VK_IMAGE_VIEW_TYPE_CUBE || ViewType == VK_IMAGE_VIEW_TYPE_CUBE_ARRAY) ? 6 : 1;
-	NumArrayLevels = ArraySize * LayerCount;
-
-	VkImageAspectFlags NewFullAspectMask = VulkanRHI::GetAspectMaskFromUEFormat(PixelFormat, true, true);
-	VkImageAspectFlags NewPartialAspectMask = VulkanRHI::GetAspectMaskFromUEFormat(PixelFormat, false, true);
-	check(NewFullAspectMask == FullAspectMask);
-	check(NewPartialAspectMask == PartialAspectMask);
-
-
-	check(MemProps == VK_MEMORY_PROPERTY_DEVICE_LOCAL_BIT); //cant evict stuff thats not device local.
-
-	const bool bRenderTarget = EnumHasAnyFlags(UEFlags, TexCreate_RenderTargetable | TexCreate_DepthStencilTargetable | TexCreate_ResolveTargetable);
-	const bool bUAV = EnumHasAnyFlags(UEFlags, TexCreate_UAV);
-	const bool bCPUReadback = EnumHasAnyFlags(UEFlags, TexCreate_CPUReadback);
-	const bool bDynamic = EnumHasAnyFlags(UEFlags, TexCreate_Dynamic);
-	bool bMemoryless = EnumHasAnyFlags(UEFlags, TexCreate_Memoryless);
-	//none of this is supported for eviction
-	check(!bMemoryless);
-	check(!bRenderTarget);
-	check(!bCPUReadback);
-	check(!bUAV);
-	MemProps = Device->GetDeviceMemoryManager().GetEvictedMemoryProperties();
-
-	const uint64 TotalGPUMemory = Device->GetDeviceMemoryManager().GetTotalMemory(true);
-
-	EVulkanAllocationMetaType MetaType = EVulkanAllocationMetaImageOther;
+	const EVulkanAllocationMetaType MetaType = EVulkanAllocationMetaImageOther;
 	if (!InDevice.GetMemoryManager().AllocateImageMemory(HostAllocation, this, MemoryRequirements, MemProps, MetaType, __FILE__, __LINE__))
 	{
 		InDevice.GetMemoryManager().HandleOOM();
 		checkNoEntry();
 	}
-	HostAllocation.BindImage(Device, HostImage);
 	VulkanTextureAllocated(MemoryRequirements.size, ViewType, bRenderTarget);
-	check(Tiling == ImageCreateInfo.ImageCreateInfo.tiling);
-
-	{
-		//copy local image -> host image
-		FVulkanCmdBuffer* CmdBuffer = Context.GetCommandBufferManager()->GetActiveCmdBuffer();
-		VkCommandBuffer VkCmdBuffer = CmdBuffer->GetHandle();
 
-		ensure(CmdBuffer->IsOutsideRenderPass());
-		VkImageLayout& Layout = Context.GetLayoutManager().FindOrAddLayoutRW(Image, VK_IMAGE_LAYOUT_UNDEFINED, NumMips, LayerCount);
-		VkImageLayout& HostLayout = Context.GetLayoutManager().FindOrAddLayoutRW(HostImage, VK_IMAGE_LAYOUT_UNDEFINED, NumMips, LayerCount);
-		check(HostLayout == VK_IMAGE_LAYOUT_UNDEFINED);
+	InternalMoveSurface(InDevice, Context, HostAllocation, true);
 
-		VkImageSubresourceRange RangeAll;
-		RangeAll.aspectMask = FullAspectMask;
-		RangeAll.baseMipLevel = 0;
-		RangeAll.levelCount = VK_REMAINING_MIP_LEVELS;
-		RangeAll.baseArrayLayer = 0;
-		RangeAll.layerCount = VK_REMAINING_ARRAY_LAYERS;
-
-
-		{
-			FVulkanPipelineBarrier Barrier;
-			Barrier.AddImageLayoutTransition(Image, Layout, VK_IMAGE_LAYOUT_TRANSFER_SRC_OPTIMAL, RangeAll);
-			Barrier.AddImageLayoutTransition(HostImage, HostLayout, VK_IMAGE_LAYOUT_TRANSFER_DST_OPTIMAL, RangeAll);
-			Barrier.Execute(VkCmdBuffer);
-		}
-		{
-			{
-				VkImageCopy Regions[MAX_TEXTURE_MIP_COUNT];
-				check(NumMips <= MAX_TEXTURE_MIP_COUNT);
-				FMemory::Memzero(Regions);
-				for(uint32 i = 0; i < NumMips; ++i)
-				{
-					VkImageCopy& Region = Regions[i];
-					Region.extent.width = FMath::Max(1u, Width >> i);
-					Region.extent.height = FMath::Max(1u, Height >> i);
-					Region.extent.depth = FMath::Max(1u, Depth >> i);
-					Region.srcSubresource.aspectMask = FullAspectMask;
-					Region.dstSubresource.aspectMask = FullAspectMask;
-					Region.srcSubresource.baseArrayLayer = 0;
-					Region.dstSubresource.baseArrayLayer = 0;
-					Region.srcSubresource.layerCount = 1;
-					Region.dstSubresource.layerCount = 1;
-					Region.srcSubresource.mipLevel = i;
-					Region.dstSubresource.mipLevel = i;
-				}
-					VulkanRHI::vkCmdCopyImage(VkCmdBuffer,
-						Image, VK_IMAGE_LAYOUT_TRANSFER_SRC_OPTIMAL,
-						HostImage, VK_IMAGE_LAYOUT_TRANSFER_DST_OPTIMAL,
-						NumMips, &Regions[0]);
-			}
-		}
-		VkImageLayout FinalLayout = Layout;
-		{
-			VulkanSetImageLayout(VkCmdBuffer, HostImage, VK_IMAGE_LAYOUT_TRANSFER_DST_OPTIMAL, FinalLayout, RangeAll);
-		}
-		Layout = VK_IMAGE_LAYOUT_TRANSFER_SRC_OPTIMAL;
-		HostLayout = FinalLayout;
-	}
-	{
-		check(Image != VK_NULL_HANDLE);
-		uint64 Size = GetMemorySize();
-		Device->NotifyDeletedImage(Image, bRenderTarget);
-		Device->GetDeferredDeletionQueue().EnqueueResource(VulkanRHI::FDeferredDeletionQueue2::EType::Image, Image);
-		Device->GetMemoryManager().FreeVulkanAllocation(Allocation);
-		check(!Allocation.HasAllocation());
-		Image = VK_NULL_HANDLE;
-		VulkanTextureDestroyed(Size, ViewType, bRenderTarget);
-	}
-	Image = HostImage;
-	Allocation.Swap(HostAllocation);
+	// Since the allocations were swapped, HostAllocation now contains the original allocation to be freed
+	Device->GetMemoryManager().FreeVulkanAllocation(HostAllocation);
 
+	VULKAN_SET_DEBUG_NAME(InDevice, VK_OBJECT_TYPE_IMAGE, Image, TEXT("(FVulkanSurface*)0x%p [hostimage]"), this);
 }
 
-
 // This is usually used for the framebuffer image
 FVulkanSurface::FVulkanSurface(FVulkanDevice& InDevice, VkImageViewType ResourceType, EPixelFormat InFormat,
-								uint32 SizeX, uint32 SizeY, uint32 SizeZ, uint32 ArraySize, uint32 InNumMips, uint32 InNumSamples,
+								uint32 SizeX, uint32 SizeY, uint32 SizeZ, uint32 InArraySize, uint32 InNumMips, uint32 InNumSamples,
 								VkImage InImage, ETextureCreateFlags InUEFlags, const FRHIResourceCreateInfo& CreateInfo)
 	: Device(&InDevice)
 	, Image(InImage)
@@ -1020,10 +860,10 @@ FVulkanSurface::FVulkanSurface(FVulkanDevice& InDevice, VkImageViewType Resource
 	, Width(SizeX)
 	, Height(SizeY)
 	, Depth(SizeZ)
+	, ArraySize(InArraySize)
 	, PixelFormat(InFormat)
 	, UEFlags(InUEFlags)
 	, MemProps(VK_MEMORY_PROPERTY_DEVICE_LOCAL_BIT)
-	, NumArrayLevels(ArraySize)
 	, Tiling(VK_IMAGE_TILING_MAX_ENUM)	// Can be expanded to a per-platform definition
 	, ViewType(ResourceType)
 	, bIsImageOwner(false)
@@ -2195,7 +2035,7 @@ FVulkanTextureBase::FVulkanTextureBase(FVulkanDevice& Device, VkImageViewType Re
 }
 
 FVulkanTextureBase::FVulkanTextureBase(FTextureRHIRef& SrcTextureRHI, const FVulkanTextureBase* SrcTexture, VkImageViewType ResourceType, uint32 SizeX, uint32 SizeY, uint32 SizeZ)
-	: Surface(*SrcTexture->Surface.Device, ResourceType, SrcTexture->Surface.PixelFormat, SizeX, SizeY, SizeZ, SrcTexture->Surface.NumArrayLevels, SrcTexture->Surface.NumMips, SrcTexture->Surface.NumSamples, SrcTexture->Surface.Image, SrcTexture->Surface.UEFlags, FRHIResourceCreateInfo(TEXT("FVulkanTextureBase")))
+	: Surface(*SrcTexture->Surface.Device, ResourceType, SrcTexture->Surface.PixelFormat, SizeX, SizeY, SizeZ, SrcTexture->Surface.GetNumberOfArrayLevels(), SrcTexture->Surface.NumMips, SrcTexture->Surface.NumSamples, SrcTexture->Surface.Image, SrcTexture->Surface.UEFlags, FRHIResourceCreateInfo(TEXT("FVulkanTextureBase")))
 	, PartialView(nullptr)
 	, AliasedTexture(SrcTextureRHI)
 {
@@ -2416,7 +2256,7 @@ FVulkanTexture2DArray::FVulkanTexture2DArray(FVulkanDevice& Device, EPixelFormat
 }
 
 FVulkanTexture2DArray::FVulkanTexture2DArray(FTextureRHIRef& SrcTextureRHI, const FVulkanTexture2DArray* SrcTexture)
-	: FRHITexture2DArray(SrcTexture->GetSizeX(), SrcTexture->GetSizeY(), SrcTexture->Surface.NumArrayLevels, SrcTexture->GetNumMips(), SrcTexture->GetNumSamples(), SrcTexture->GetFormat(), SrcTexture->GetFlags(), SrcTexture->GetClearBinding())
+	: FRHITexture2DArray(SrcTexture->GetSizeX(), SrcTexture->GetSizeY(), SrcTexture->Surface.GetNumberOfArrayLevels(), SrcTexture->GetNumMips(), SrcTexture->GetNumSamples(), SrcTexture->GetFormat(), SrcTexture->GetFlags(), SrcTexture->GetClearBinding())
 	, FVulkanTextureBase(SrcTextureRHI, static_cast<const FVulkanTextureBase*>(SrcTexture), VK_IMAGE_VIEW_TYPE_2D_ARRAY, SrcTexture->GetSizeX(), SrcTexture->GetSizeY(), 1)
 {
 }
