Refactor Vulkan renderpass initialization to allow for either vkRenderPassCreateInfo or vkRenderPassCreateInfo2 paths, depending on platform support for the extension.

diff --git a/Engine/Source/Runtime/VulkanRHI/Private/VulkanCommandBuffer.cpp b/Engine/Source/Runtime/VulkanRHI/Private/VulkanCommandBuffer.cpp
--- a/Engine/Source/Runtime/VulkanRHI/Private/VulkanCommandBuffer.cpp
+++ b/Engine/Source/Runtime/VulkanRHI/Private/VulkanCommandBuffer.cpp
@@ -175,7 +175,7 @@ void FVulkanCmdBuffer::EndRenderPass()
 
 void FVulkanCmdBuffer::BeginRenderPass(const FVulkanRenderTargetLayout& Layout, FVulkanRenderPass* RenderPass, FVulkanFramebuffer* Framebuffer, const VkClearValue* AttachmentClearValues)
 {
-	if(bIsUniformBufferBarrierAdded)
+	if (bIsUniformBufferBarrierAdded)
 	{
 		EndUniformUpdateBarrier();
 	}
@@ -201,7 +201,20 @@ void FVulkanCmdBuffer::BeginRenderPass(const FVulkanRenderTargetLayout& Layout,
 		Info.pNext = &RPTransformBeginInfoQCOM;
 	}
 #endif
-	VulkanRHI::vkCmdBeginRenderPass(CommandBufferHandle, &Info, VK_SUBPASS_CONTENTS_INLINE);
+
+#if VULKAN_SUPPORTS_RENDERPASS2
+	if (Device->GetOptionalExtensions().HasKHRRenderPass2)
+	{
+		VkSubpassBeginInfo SubpassInfo;
+		ZeroVulkanStruct(SubpassInfo, VK_STRUCTURE_TYPE_SUBPASS_BEGIN_INFO);
+		SubpassInfo.contents = VK_SUBPASS_CONTENTS_INLINE;
+		VulkanRHI::vkCmdBeginRenderPass2KHR(CommandBufferHandle, &Info, &SubpassInfo);
+	}
+	else
+#endif
+	{
+		VulkanRHI::vkCmdBeginRenderPass(CommandBufferHandle, &Info, VK_SUBPASS_CONTENTS_INLINE);
+	}
 
 	State = EState::IsInsideRenderPass;
 
diff --git a/Engine/Source/Runtime/VulkanRHI/Private/VulkanDevice.h b/Engine/Source/Runtime/VulkanRHI/Private/VulkanDevice.h
--- a/Engine/Source/Runtime/VulkanRHI/Private/VulkanDevice.h
+++ b/Engine/Source/Runtime/VulkanRHI/Private/VulkanDevice.h
@@ -47,6 +47,7 @@ struct FOptionalVulkanDeviceExtensions
 			uint32 HasEXTFragmentDensityMap : 1;
 			uint32 HasEXTFragmentDensityMap2 : 1;
 			uint32 HasKHRFragmentShadingRate : 1;
+			uint32 HasKHRRenderPass2 : 1;
 			uint32 HasEXTFullscreenExclusive : 1;
 			uint32 HasKHRImageFormatList : 1;
 			uint32 HasEXTASTCDecodeMode : 1;
diff --git a/Engine/Source/Runtime/VulkanRHI/Private/VulkanLayers.cpp b/Engine/Source/Runtime/VulkanRHI/Private/VulkanLayers.cpp
--- a/Engine/Source/Runtime/VulkanRHI/Private/VulkanLayers.cpp
+++ b/Engine/Source/Runtime/VulkanRHI/Private/VulkanLayers.cpp
@@ -854,6 +854,10 @@ void FOptionalVulkanDeviceExtensions::Setup(const TArray<const ANSICHAR*>& Devic
 	HasDriverProperties = HasExtension(DeviceExtensions, VK_KHR_DRIVER_PROPERTIES_EXTENSION_NAME);
 #endif
 
+#if VULKAN_SUPPORTS_RENDERPASS2
+	HasKHRRenderPass2 = HasExtension(DeviceExtensions, VK_KHR_CREATE_RENDERPASS_2_EXTENSION_NAME);
+#endif
+
 #if VULKAN_SUPPORTS_FRAGMENT_DENSITY_MAP
 	HasEXTFragmentDensityMap = HasExtension(DeviceExtensions, VK_EXT_FRAGMENT_DENSITY_MAP_EXTENSION_NAME);
 #endif
diff --git a/Engine/Source/Runtime/VulkanRHI/Private/VulkanRHI.cpp b/Engine/Source/Runtime/VulkanRHI/Private/VulkanRHI.cpp
--- a/Engine/Source/Runtime/VulkanRHI/Private/VulkanRHI.cpp
+++ b/Engine/Source/Runtime/VulkanRHI/Private/VulkanRHI.cpp
@@ -22,6 +22,7 @@
 #include "GlobalShader.h"
 #include "RHIValidation.h"
 #include "IHeadMountedDisplayModule.h"
+#include "VulkanRenderpass.h"
 
 static_assert(sizeof(VkStructureType) == sizeof(int32), "ZeroVulkanStruct() assumes VkStructureType is int32!");
 
@@ -1534,201 +1535,6 @@ void FVulkanBufferView::Destroy()
 	}
 }
 
-static VkRenderPass CreateRenderPass(FVulkanDevice& InDevice, const FVulkanRenderTargetLayout& RTLayout)
-{
-	VkRenderPassCreateInfo CreateInfo;
-	ZeroVulkanStruct(CreateInfo, VK_STRUCTURE_TYPE_RENDER_PASS_CREATE_INFO);
-	
-	uint32 NumSubpasses = 0;
-	uint32 NumDependencies = 0;
-
-	VkSubpassDescription SubpassDescriptions[8];
-	VkSubpassDependency SubpassDependencies[8];
-
-	//0b11 for 2, 0b1111 for 4, and so on
-	uint32 MultiviewMask = ( 0b1 << RTLayout.GetMultiViewCount() ) - 1;
-
-	const bool bDeferredShadingSubpass = RTLayout.GetSubpassHint() == ESubpassHint::DeferredShadingSubpass;
-	const bool bDepthReadSubpass =  RTLayout.GetSubpassHint() == ESubpassHint::DepthReadSubpass;
-		
-	// main sub-pass
-	{
-		VkSubpassDescription& SubpassDesc = SubpassDescriptions[NumSubpasses++];
-		FMemory::Memzero(&SubpassDesc, sizeof(VkSubpassDescription));
-
-		SubpassDesc.pipelineBindPoint = VK_PIPELINE_BIND_POINT_GRAPHICS;
-		SubpassDesc.colorAttachmentCount = RTLayout.GetNumColorAttachments();
-		SubpassDesc.pColorAttachments = RTLayout.GetColorAttachmentReferences();
-		SubpassDesc.pResolveAttachments = bDepthReadSubpass ? nullptr : RTLayout.GetResolveAttachmentReferences();
-		SubpassDesc.pDepthStencilAttachment = RTLayout.GetDepthStencilAttachmentReference();
-	}
-
-	// Color write and depth read sub-pass
-	VkAttachmentReference InputAttachments1[MaxSimultaneousRenderTargets + 1];
-	if (bDepthReadSubpass)
-	{
-		VkSubpassDescription& SubpassDesc = SubpassDescriptions[NumSubpasses++];
-		FMemory::Memzero(&SubpassDesc, sizeof(VkSubpassDescription));
-
-		SubpassDesc.pipelineBindPoint = VK_PIPELINE_BIND_POINT_GRAPHICS;
-		SubpassDesc.colorAttachmentCount = RTLayout.GetNumColorAttachments();
-		SubpassDesc.pColorAttachments = RTLayout.GetColorAttachmentReferences();
-		SubpassDesc.pResolveAttachments = RTLayout.GetResolveAttachmentReferences();
-
-		check(RTLayout.GetDepthStencilAttachmentReference());
-
-		// depth as Input0
-		InputAttachments1[0].attachment = RTLayout.GetDepthStencilAttachmentReference()->attachment;
-		InputAttachments1[0].layout = VK_IMAGE_LAYOUT_DEPTH_STENCIL_READ_ONLY_OPTIMAL;
-		
-		SubpassDesc.inputAttachmentCount = 1;
-		SubpassDesc.pInputAttachments = InputAttachments1;
-		// depth attachment is same as input attachment
-		SubpassDesc.pDepthStencilAttachment = InputAttachments1;
-						
-		VkSubpassDependency& SubpassDep = SubpassDependencies[NumDependencies++];
-		SubpassDep.srcSubpass = 0;
-		SubpassDep.dstSubpass = 1;
-	    SubpassDep.srcStageMask = VK_PIPELINE_STAGE_EARLY_FRAGMENT_TESTS_BIT | VK_PIPELINE_STAGE_LATE_FRAGMENT_TESTS_BIT;
-		SubpassDep.dstStageMask = VK_PIPELINE_STAGE_FRAGMENT_SHADER_BIT;
-		SubpassDep.srcAccessMask = VK_ACCESS_DEPTH_STENCIL_ATTACHMENT_WRITE_BIT;
-		SubpassDep.dstAccessMask = VK_ACCESS_INPUT_ATTACHMENT_READ_BIT;
-		SubpassDep.dependencyFlags = VK_DEPENDENCY_BY_REGION_BIT;
-	}
-
-	// Two subpasses for deferred shading
-	VkAttachmentReference InputAttachments2[MaxSimultaneousRenderTargets + 1];
-	VkAttachmentReference DepthStencilAttachment;
-	if (bDeferredShadingSubpass)
-	{
-		// both sub-passes only test DepthStencil
-		DepthStencilAttachment.attachment = RTLayout.GetDepthStencilAttachmentReference()->attachment;
-		DepthStencilAttachment.layout = VK_IMAGE_LAYOUT_DEPTH_STENCIL_READ_ONLY_OPTIMAL;
-
-		const VkAttachmentReference* ColorRef = RTLayout.GetColorAttachmentReferences();
-		uint32 NumColorAttachments = RTLayout.GetNumColorAttachments();
-		check(NumColorAttachments == 4); //current layout is SceneColor, GBufferA/B/C
-
-		// 1. Write to SceneColor and GBuffer, input DepthStencil
-		{
-			VkSubpassDescription& SubpassDesc = SubpassDescriptions[NumSubpasses++];
-			FMemory::Memzero(&SubpassDesc, sizeof(VkSubpassDescription));
-
-			SubpassDesc.pipelineBindPoint = VK_PIPELINE_BIND_POINT_GRAPHICS;
-			SubpassDesc.colorAttachmentCount = 4;
-			SubpassDesc.pColorAttachments = ColorRef;
-			SubpassDesc.pDepthStencilAttachment = &DepthStencilAttachment;
-			// depth as Input0
-			SubpassDesc.inputAttachmentCount = 1;
-			SubpassDesc.pInputAttachments = &DepthStencilAttachment;
-						
-			VkSubpassDependency& SubpassDep = SubpassDependencies[NumDependencies++];
-			SubpassDep.srcSubpass = 0;
-			SubpassDep.dstSubpass = 1;
-			SubpassDep.srcStageMask = VK_PIPELINE_STAGE_COLOR_ATTACHMENT_OUTPUT_BIT;
-			SubpassDep.dstStageMask = VK_PIPELINE_STAGE_FRAGMENT_SHADER_BIT;
-			SubpassDep.srcAccessMask = VK_ACCESS_COLOR_ATTACHMENT_WRITE_BIT;
-			SubpassDep.dstAccessMask = VK_ACCESS_INPUT_ATTACHMENT_READ_BIT;
-			SubpassDep.dependencyFlags = VK_DEPENDENCY_BY_REGION_BIT;
-		}
-
-		// 2. Write to SceneColor, input GBuffer and DepthStencil
-		{
-			VkSubpassDescription& SubpassDesc = SubpassDescriptions[NumSubpasses++];
-			FMemory::Memzero(&SubpassDesc, sizeof(VkSubpassDescription));
-
-			SubpassDesc.pipelineBindPoint = VK_PIPELINE_BIND_POINT_GRAPHICS;
-			SubpassDesc.colorAttachmentCount = 1; // SceneColor only
-			SubpassDesc.pColorAttachments = ColorRef;
-			SubpassDesc.pDepthStencilAttachment = &DepthStencilAttachment;
-			// GBuffer as Input2/3/4
-			InputAttachments2[0].attachment = DepthStencilAttachment.attachment;
-			InputAttachments2[0].layout = DepthStencilAttachment.layout;
-			InputAttachments2[1].attachment = VK_ATTACHMENT_UNUSED;
-			InputAttachments2[1].layout = VK_IMAGE_LAYOUT_UNDEFINED;
-			for (int32 i = 2; i < 5; ++i)
-			{
-				InputAttachments2[i].attachment = ColorRef[i-1].attachment;
-				InputAttachments2[i].layout = VK_IMAGE_LAYOUT_SHADER_READ_ONLY_OPTIMAL;
-			}
-			SubpassDesc.inputAttachmentCount = 5;
-			SubpassDesc.pInputAttachments = InputAttachments2;
-
-			VkSubpassDependency& SubpassDep = SubpassDependencies[NumDependencies++];
-			SubpassDep.srcSubpass = 1;
-			SubpassDep.dstSubpass = 2;
-			SubpassDep.srcStageMask = VK_PIPELINE_STAGE_COLOR_ATTACHMENT_OUTPUT_BIT;
-			SubpassDep.dstStageMask = VK_PIPELINE_STAGE_FRAGMENT_SHADER_BIT;
-			SubpassDep.srcAccessMask = VK_ACCESS_COLOR_ATTACHMENT_WRITE_BIT;
-			SubpassDep.dstAccessMask = VK_ACCESS_INPUT_ATTACHMENT_READ_BIT;
-			if (GVulkanInputAttachmentShaderRead == 1)
-			{
-				// this is not required, but might flicker on some devices without
-				SubpassDep.dstAccessMask|= VK_ACCESS_SHADER_READ_BIT;
-			}
-			SubpassDep.dependencyFlags = VK_DEPENDENCY_BY_REGION_BIT;
-		}
-	}
-	
-	CreateInfo.attachmentCount = RTLayout.GetNumAttachmentDescriptions();
-	CreateInfo.pAttachments = RTLayout.GetAttachmentDescriptions();
-	CreateInfo.subpassCount = NumSubpasses;
-	CreateInfo.pSubpasses = SubpassDescriptions;
-	CreateInfo.dependencyCount = NumDependencies;
-	CreateInfo.pDependencies = SubpassDependencies;
-
-	/*
-	Bit mask that specifies which view rendering is broadcast to
-	0011 = Broadcast to first and second view (layer)
-	*/
-	const uint32_t ViewMask[2] = { MultiviewMask, MultiviewMask };
-
-	/*
-	Bit mask that specifices correlation between views
-	An implementation may use this for optimizations (concurrent render)
-	*/
-	const uint32_t CorrelationMask = MultiviewMask;
-
-	VkRenderPassMultiviewCreateInfo MultiviewInfo;
-	if (RTLayout.GetIsMultiView())
-	{
-		FMemory::Memzero(MultiviewInfo);
-		MultiviewInfo.sType = VK_STRUCTURE_TYPE_RENDER_PASS_MULTIVIEW_CREATE_INFO;
-		MultiviewInfo.pNext = nullptr;
-		MultiviewInfo.subpassCount = NumSubpasses;
-		MultiviewInfo.pViewMasks = ViewMask;
-		MultiviewInfo.dependencyCount = 0;
-		MultiviewInfo.pViewOffsets = nullptr;
-		MultiviewInfo.correlationMaskCount = 1;
-		MultiviewInfo.pCorrelationMasks = &CorrelationMask;
-
-		CreateInfo.pNext = &MultiviewInfo;
-	}
-	
-	VkRenderPassFragmentDensityMapCreateInfoEXT FragDensityCreateInfo;
-	if (InDevice.GetOptionalExtensions().HasEXTFragmentDensityMap && RTLayout.GetHasFragmentDensityAttachment())
-	{
-		ZeroVulkanStruct(FragDensityCreateInfo, VK_STRUCTURE_TYPE_RENDER_PASS_FRAGMENT_DENSITY_MAP_CREATE_INFO_EXT);
-		FragDensityCreateInfo.fragmentDensityMapAttachment = *RTLayout.GetFragmentDensityAttachmentReference();
-
-		// Chain fragment density info onto create info and the rest of the pNexts
-		// onto the fragment density info
-		FragDensityCreateInfo.pNext = CreateInfo.pNext;
-		CreateInfo.pNext = &FragDensityCreateInfo;
-	}
-
-#if VULKAN_SUPPORTS_QCOM_RENDERPASS_TRANSFORM
-	if (RTLayout.GetQCOMRenderPassTransform() != VK_SURFACE_TRANSFORM_IDENTITY_BIT_KHR)
-	{
-		CreateInfo.flags = VK_RENDER_PASS_CREATE_TRANSFORM_BIT_QCOM;
-	}
-#endif
-
-	VkRenderPass RenderPassHandle;
-	VERIFYVULKANRESULT_EXPANDED(VulkanRHI::vkCreateRenderPass(InDevice.GetInstanceHandle(), &CreateInfo, VULKAN_CPU_ALLOCATOR, &RenderPassHandle));
-	return RenderPassHandle;
-}
-
 FVulkanRenderPass::FVulkanRenderPass(FVulkanDevice& InDevice, const FVulkanRenderTargetLayout& InRTLayout) :
 	Layout(InRTLayout),
 	RenderPass(VK_NULL_HANDLE),
@@ -1736,7 +1542,7 @@ FVulkanRenderPass::FVulkanRenderPass(FVulkanDevice& InDevice, const FVulkanRende
 	Device(InDevice)
 {
 	INC_DWORD_STAT(STAT_VulkanNumRenderPasses);
-	RenderPass = CreateRenderPass(InDevice, InRTLayout);
+	RenderPass = CreateVulkanRenderPass(InDevice, InRTLayout);
 }
 
 FVulkanRenderPass::~FVulkanRenderPass()
diff --git a/Engine/Source/Runtime/VulkanRHI/Private/VulkanRenderpass.cpp b/Engine/Source/Runtime/VulkanRHI/Private/VulkanRenderpass.cpp
new file mode 100644
--- /dev/null
+++ b/Engine/Source/Runtime/VulkanRHI/Private/VulkanRenderpass.cpp
@@ -0,0 +1,256 @@
+// Copyright Epic Games, Inc. All Rights Reserved.
+
+/*=============================================================================
+	VulkanState.cpp: Vulkan state implementation.
+=============================================================================*/
+
+#include "VulkanRenderpass.h"
+
+extern int32 GVulkanInputAttachmentShaderRead;
+
+template <typename TSubpassDescriptionClass, typename TSubpassDependencyClass, typename TAttachmentReferenceClass, typename TAttachmentDescriptionClass, typename TRenderPassCreateInfoClass>
+class FVulkanRenderPassBuilder
+{
+public:
+	FVulkanRenderPassBuilder(FVulkanDevice& InDevice)
+		: Device(InDevice)
+	{}
+
+	VkRenderPass Create(const FVulkanRenderTargetLayout& RTLayout)
+	{
+		TRenderPassCreateInfoClass CreateInfo;
+
+		uint32 NumSubpasses = 0;
+		uint32 NumDependencies = 0;
+
+		//0b11 for 2, 0b1111 for 4, and so on
+		uint32 MultiviewMask = (0b1 << RTLayout.GetMultiViewCount()) - 1;
+
+		const bool bDeferredShadingSubpass = RTLayout.GetSubpassHint() == ESubpassHint::DeferredShadingSubpass;
+		const bool bDepthReadSubpass = RTLayout.GetSubpassHint() == ESubpassHint::DepthReadSubpass;
+
+		// Grab (and optionally convert) attachment references.
+		for (uint32 ColorAttachment = 0; ColorAttachment < RTLayout.GetNumColorAttachments(); ++ColorAttachment)
+		{
+			ColorAttachmentReferences.Add(TAttachmentReferenceClass(RTLayout.GetColorAttachmentReferences()[ColorAttachment], 0));
+			if (RTLayout.GetResolveAttachmentReferences() != nullptr)
+			{
+				ResolveAttachmentReferences.Add(TAttachmentReferenceClass(RTLayout.GetResolveAttachmentReferences()[ColorAttachment], 0));
+			}
+		}
+
+		if (RTLayout.GetDepthStencilAttachmentReference() != nullptr)
+		{
+			DepthStencilAttachmentReference = TAttachmentReferenceClass(*RTLayout.GetDepthStencilAttachmentReference(), 0);
+		}
+
+		// main sub-pass
+		{
+			TSubpassDescriptionClass& SubpassDesc = SubpassDescriptions[NumSubpasses++];
+
+			SubpassDesc.SetColorAttachments(ColorAttachmentReferences);
+			SubpassDesc.SetResolveAttachments(ResolveAttachmentReferences);
+			if (RTLayout.GetDepthStencilAttachmentReference() != nullptr)
+			{
+				SubpassDesc.SetDepthStencilAttachment(&DepthStencilAttachmentReference);
+			}
+		}
+
+		// Color write and depth read sub-pass
+		static const uint32 InputAttachment1Count = 1;
+		TAttachmentReferenceClass InputAttachments1[InputAttachment1Count];
+		TAttachmentReferenceClass DepthStencilAttachmentOG;
+		if (bDepthReadSubpass)
+		{
+			DepthStencilAttachmentOG.SetAttachment(*RTLayout.GetDepthStencilAttachmentReference(), VK_IMAGE_ASPECT_DEPTH_BIT);
+			TSubpassDescriptionClass& SubpassDesc = SubpassDescriptions[NumSubpasses++];
+
+			SubpassDesc.SetColorAttachments(ColorAttachmentReferences);
+			SubpassDesc.SetResolveAttachments(ResolveAttachmentReferences);
+
+			check(RTLayout.GetDepthStencilAttachmentReference());
+
+			// depth as Input0
+			InputAttachments1[0].SetAttachment(DepthStencilAttachmentOG, VK_IMAGE_ASPECT_DEPTH_BIT);
+			InputAttachments1[0].layout = VK_IMAGE_LAYOUT_DEPTH_STENCIL_READ_ONLY_OPTIMAL;
+
+			SubpassDesc.SetInputAttachments(InputAttachments1, InputAttachment1Count);
+			// depth attachment is same as input attachment
+			SubpassDesc.SetDepthStencilAttachment(InputAttachments1);
+
+			TSubpassDependencyClass& SubpassDep = SubpassDependencies[NumDependencies++];
+			SubpassDep.srcSubpass = 0;
+			SubpassDep.dstSubpass = 1;
+			SubpassDep.srcStageMask = VK_PIPELINE_STAGE_EARLY_FRAGMENT_TESTS_BIT | VK_PIPELINE_STAGE_LATE_FRAGMENT_TESTS_BIT;
+			SubpassDep.dstStageMask = VK_PIPELINE_STAGE_FRAGMENT_SHADER_BIT;
+			SubpassDep.srcAccessMask = VK_ACCESS_DEPTH_STENCIL_ATTACHMENT_WRITE_BIT;
+			SubpassDep.dstAccessMask = VK_ACCESS_INPUT_ATTACHMENT_READ_BIT;
+			SubpassDep.dependencyFlags = VK_DEPENDENCY_BY_REGION_BIT;
+		}
+
+		// Two subpasses for deferred shading
+		TAttachmentReferenceClass InputAttachments2[MaxSimultaneousRenderTargets + 1];
+		TAttachmentReferenceClass DepthStencilAttachment;
+		if (bDeferredShadingSubpass)
+		{
+			// both sub-passes only test DepthStencil
+			DepthStencilAttachment.attachment = RTLayout.GetDepthStencilAttachmentReference()->attachment;
+			DepthStencilAttachment.layout = VK_IMAGE_LAYOUT_DEPTH_STENCIL_READ_ONLY_OPTIMAL;
+			DepthStencilAttachment.SetAspect(VK_IMAGE_ASPECT_DEPTH_BIT);	// @todo?
+
+			//const VkAttachmentReference* ColorRef = RTLayout.GetColorAttachmentReferences();
+			//uint32 NumColorAttachments = RTLayout.GetNumColorAttachments();
+			check(RTLayout.GetNumColorAttachments() == 4); //current layout is SceneColor, GBufferA/B/C
+
+			// 1. Write to SceneColor and GBuffer, input DepthStencil
+			{
+				TSubpassDescriptionClass& SubpassDesc = SubpassDescriptions[NumSubpasses++];
+				SubpassDesc.SetColorAttachments(ColorAttachmentReferences);
+				SubpassDesc.SetDepthStencilAttachment(&DepthStencilAttachment);
+				SubpassDesc.SetInputAttachments(&DepthStencilAttachment, 1);
+
+				// depth as Input0
+
+				TSubpassDependencyClass& SubpassDep = SubpassDependencies[NumDependencies++];
+				SubpassDep.srcSubpass = 0;
+				SubpassDep.dstSubpass = 1;
+				SubpassDep.srcStageMask = VK_PIPELINE_STAGE_COLOR_ATTACHMENT_OUTPUT_BIT;
+				SubpassDep.dstStageMask = VK_PIPELINE_STAGE_FRAGMENT_SHADER_BIT;
+				SubpassDep.srcAccessMask = VK_ACCESS_COLOR_ATTACHMENT_WRITE_BIT;
+				SubpassDep.dstAccessMask = VK_ACCESS_INPUT_ATTACHMENT_READ_BIT;
+				SubpassDep.dependencyFlags = VK_DEPENDENCY_BY_REGION_BIT;
+			}
+
+			// 2. Write to SceneColor, input GBuffer and DepthStencil
+			{
+				TSubpassDescriptionClass& SubpassDesc = SubpassDescriptions[NumSubpasses++];
+				SubpassDesc.SetColorAttachments(ColorAttachmentReferences, 1); // SceneColor only
+
+				// GBuffer as Input2/3/4
+				InputAttachments2[0].attachment = DepthStencilAttachment.attachment;
+				InputAttachments2[0].layout = DepthStencilAttachment.layout;
+				InputAttachments2[0].SetAspect(VK_IMAGE_ASPECT_DEPTH_BIT);
+
+				InputAttachments2[1].attachment = VK_ATTACHMENT_UNUSED;
+				InputAttachments2[1].layout = VK_IMAGE_LAYOUT_UNDEFINED;
+				InputAttachments2[1].SetAspect(0);
+
+				for (int32 i = 2; i < 5; ++i)
+				{
+					InputAttachments2[i].attachment = ColorAttachmentReferences[i - 1].attachment;
+					InputAttachments2[i].layout = VK_IMAGE_LAYOUT_SHADER_READ_ONLY_OPTIMAL;
+					InputAttachments2[i].SetAspect(VK_IMAGE_ASPECT_COLOR_BIT);
+				}
+
+				SubpassDesc.SetInputAttachments(InputAttachments2, 5);
+
+				TSubpassDependencyClass& SubpassDep = SubpassDependencies[NumDependencies++];
+				SubpassDep.srcSubpass = 1;
+				SubpassDep.dstSubpass = 2;
+				SubpassDep.srcStageMask = VK_PIPELINE_STAGE_COLOR_ATTACHMENT_OUTPUT_BIT;
+				SubpassDep.dstStageMask = VK_PIPELINE_STAGE_FRAGMENT_SHADER_BIT;
+				SubpassDep.srcAccessMask = VK_ACCESS_COLOR_ATTACHMENT_WRITE_BIT;
+				SubpassDep.dstAccessMask = VK_ACCESS_INPUT_ATTACHMENT_READ_BIT;
+				if (GVulkanInputAttachmentShaderRead == 1)
+				{
+					// this is not required, but might flicker on some devices without
+					SubpassDep.dstAccessMask |= VK_ACCESS_SHADER_READ_BIT;
+				}
+				SubpassDep.dependencyFlags = VK_DEPENDENCY_BY_REGION_BIT;
+			}
+		}
+
+		TArray<TAttachmentDescriptionClass> AttachmentDescriptions;
+		for (uint32 Attachment = 0; Attachment < RTLayout.GetNumAttachmentDescriptions(); ++Attachment)
+		{
+			AttachmentDescriptions.Add(TAttachmentDescriptionClass(RTLayout.GetAttachmentDescriptions()[Attachment]));
+		}
+
+		CreateInfo.attachmentCount = AttachmentDescriptions.Num();
+		CreateInfo.pAttachments = AttachmentDescriptions.GetData();
+		CreateInfo.subpassCount = NumSubpasses;
+		CreateInfo.pSubpasses = SubpassDescriptions;
+		CreateInfo.dependencyCount = NumDependencies;
+		CreateInfo.pDependencies = SubpassDependencies;
+
+		/*
+		Bit mask that specifies which view rendering is broadcast to
+		0011 = Broadcast to first and second view (layer)
+		*/
+		const uint32_t ViewMask[2] = { MultiviewMask, MultiviewMask };
+
+		/*
+		Bit mask that specifices correlation between views
+		An implementation may use this for optimizations (concurrent render)
+		*/
+		const uint32_t CorrelationMask = MultiviewMask;
+
+		VkRenderPassMultiviewCreateInfo MultiviewInfo;
+		if (RTLayout.GetIsMultiView())
+		{
+			FMemory::Memzero(MultiviewInfo);
+			MultiviewInfo.sType = VK_STRUCTURE_TYPE_RENDER_PASS_MULTIVIEW_CREATE_INFO;
+			MultiviewInfo.pNext = nullptr;
+			MultiviewInfo.subpassCount = NumSubpasses;
+			MultiviewInfo.pViewMasks = ViewMask;
+			MultiviewInfo.dependencyCount = 0;
+			MultiviewInfo.pViewOffsets = nullptr;
+			MultiviewInfo.correlationMaskCount = 1;
+			MultiviewInfo.pCorrelationMasks = &CorrelationMask;
+
+			CreateInfo.pNext = &MultiviewInfo;
+		}
+
+		VkRenderPassFragmentDensityMapCreateInfoEXT FragDensityCreateInfo;
+		if (Device.GetOptionalExtensions().HasEXTFragmentDensityMap && RTLayout.GetHasFragmentDensityAttachment())
+		{
+			ZeroVulkanStruct(FragDensityCreateInfo, VK_STRUCTURE_TYPE_RENDER_PASS_FRAGMENT_DENSITY_MAP_CREATE_INFO_EXT);
+			FragDensityCreateInfo.fragmentDensityMapAttachment = *RTLayout.GetFragmentDensityAttachmentReference();
+
+			// Chain fragment density info onto create info and the rest of the pNexts
+			// onto the fragment density info
+			FragDensityCreateInfo.pNext = CreateInfo.pNext;
+			CreateInfo.pNext = &FragDensityCreateInfo;
+		}
+
+#if VULKAN_SUPPORTS_QCOM_RENDERPASS_TRANSFORM
+		if (RTLayout.GetQCOMRenderPassTransform() != VK_SURFACE_TRANSFORM_IDENTITY_BIT_KHR)
+		{
+			CreateInfo.flags = VK_RENDER_PASS_CREATE_TRANSFORM_BIT_QCOM;
+		}
+#endif
+
+		return CreateInfo.Create(Device);
+	}
+
+private:
+	TSubpassDescriptionClass SubpassDescriptions[8];
+	TSubpassDependencyClass SubpassDependencies[8];
+
+	TArray<TAttachmentReferenceClass> ColorAttachmentReferences;
+	TArray<TAttachmentReferenceClass> ResolveAttachmentReferences;
+	TAttachmentReferenceClass DepthStencilAttachmentReference;
+
+	FVulkanDevice& Device;
+};
+
+VkRenderPass CreateVulkanRenderPass(FVulkanDevice& InDevice, const FVulkanRenderTargetLayout& RTLayout)
+{
+	VkRenderPass OutRenderpass;
+
+#if VULKAN_SUPPORTS_RENDERPASS2
+	if (InDevice.GetOptionalExtensions().HasKHRRenderPass2)
+	{
+		FVulkanRenderPassBuilder<FVulkanSubpassDescription<VkSubpassDescription2>, FVulkanSubpassDependency<VkSubpassDependency2>, FVulkanAttachmentReference<VkAttachmentReference2>, FVulkanAttachmentDescription<VkAttachmentDescription2>, FVulkanRenderPassCreateInfo<VkRenderPassCreateInfo2>> Creator(InDevice);
+		OutRenderpass = Creator.Create(RTLayout);
+	}
+	else
+#endif
+	{
+		FVulkanRenderPassBuilder<FVulkanSubpassDescription<VkSubpassDescription>, FVulkanSubpassDependency<VkSubpassDependency>, FVulkanAttachmentReference<VkAttachmentReference>, FVulkanAttachmentDescription<VkAttachmentDescription>, FVulkanRenderPassCreateInfo<VkRenderPassCreateInfo>> Creator(InDevice);
+		OutRenderpass = Creator.Create(RTLayout);
+	}
+
+	return OutRenderpass;
+}
+
diff --git a/Engine/Source/Runtime/VulkanRHI/Private/VulkanRenderpass.h b/Engine/Source/Runtime/VulkanRHI/Private/VulkanRenderpass.h
new file mode 100644
--- /dev/null
+++ b/Engine/Source/Runtime/VulkanRHI/Private/VulkanRenderpass.h
@@ -0,0 +1,297 @@
+// Copyright Epic Games, Inc. All Rights Reserved.
+
+/*=============================================================================
+	VulkanState.h: Vulkan state definitions.
+=============================================================================*/
+
+#pragma once
+
+#include "CoreMinimal.h"
+#include "RHIResources.h"
+#include "VulkanRHIPrivate.h"
+#include "VulkanResources.h"
+#include "VulkanPendingState.h"
+#include "VulkanContext.h"
+
+template <typename TAttachmentReferenceType>
+struct FVulkanAttachmentReference
+	: public TAttachmentReferenceType
+{
+	FVulkanAttachmentReference()
+	{
+		ZeroStruct();
+	}
+
+	FVulkanAttachmentReference(const VkAttachmentReference& AttachmentReferenceIn, VkImageAspectFlags AspectMask)
+	{
+		SetAttachment(AttachmentReferenceIn, AspectMask);
+	}
+
+	inline void SetAttachment(const VkAttachmentReference& AttachmentReferenceIn, VkImageAspectFlags AspectMask) { check(false); }
+	inline void SetAttachment(const FVulkanAttachmentReference<TAttachmentReferenceType>& AttachmentReferenceIn, VkImageAspectFlags AspectMask) { *this = AttachmentReferenceIn; }
+	inline void ZeroStruct() {}
+	inline void SetAspect(uint32 Aspect) {}
+};
+
+template <>
+inline void FVulkanAttachmentReference<VkAttachmentReference>::SetAttachment(const VkAttachmentReference& AttachmentReferenceIn, VkImageAspectFlags AspectMask)
+{
+	attachment = AttachmentReferenceIn.attachment;
+	layout = AttachmentReferenceIn.layout;
+}
+
+#if VULKAN_SUPPORTS_RENDERPASS2
+template <>
+inline void FVulkanAttachmentReference<VkAttachmentReference2>::SetAttachment(const VkAttachmentReference& AttachmentReferenceIn, VkImageAspectFlags AspectMask)
+{
+	sType = VK_STRUCTURE_TYPE_ATTACHMENT_REFERENCE_2;
+	pNext = nullptr;
+	attachment = AttachmentReferenceIn.attachment;
+	layout = AttachmentReferenceIn.layout;
+	aspectMask = AspectMask;
+}
+
+template<>
+inline void FVulkanAttachmentReference<VkAttachmentReference2>::SetAttachment(const FVulkanAttachmentReference<VkAttachmentReference2>& AttachmentReferenceIn, VkImageAspectFlags AspectMask)
+{
+	sType = VK_STRUCTURE_TYPE_ATTACHMENT_REFERENCE_2;
+	pNext = nullptr;
+	attachment = AttachmentReferenceIn.attachment;
+	layout = AttachmentReferenceIn.layout;
+	aspectMask = AspectMask;
+}
+#endif
+
+template<>
+inline void FVulkanAttachmentReference<VkAttachmentReference>::ZeroStruct()
+{
+	attachment = 0;
+	layout = VK_IMAGE_LAYOUT_UNDEFINED;
+}
+
+#if VULKAN_SUPPORTS_RENDERPASS2
+template<>
+inline void FVulkanAttachmentReference<VkAttachmentReference2>::ZeroStruct()
+{
+	sType = VK_STRUCTURE_TYPE_ATTACHMENT_REFERENCE_2;
+	pNext = nullptr;
+	attachment = 0;
+	layout = VK_IMAGE_LAYOUT_UNDEFINED;
+	aspectMask = 0;
+}
+
+template<>
+inline void FVulkanAttachmentReference<VkAttachmentReference2>::SetAspect(uint32 Aspect)
+{
+	aspectMask = Aspect;
+}
+#endif
+
+template <typename TSubpassDescriptionType>
+class FVulkanSubpassDescription
+{
+};
+
+template<>
+struct FVulkanSubpassDescription<VkSubpassDescription>
+	: public VkSubpassDescription
+{
+	FVulkanSubpassDescription()
+	{
+		FMemory::Memzero(this, sizeof(VkSubpassDescription));
+		pipelineBindPoint = VK_PIPELINE_BIND_POINT_GRAPHICS;
+	}
+
+	void SetColorAttachments(const TArray<FVulkanAttachmentReference<VkAttachmentReference>>& ColorAttachmentReferences, int OverrideCount = -1)
+	{
+		colorAttachmentCount = (OverrideCount == -1) ? ColorAttachmentReferences.Num() : OverrideCount;
+		pColorAttachments = ColorAttachmentReferences.GetData();
+	}
+
+	void SetResolveAttachments(const TArray<FVulkanAttachmentReference<VkAttachmentReference>>& ResolveAttachmentReferences)
+	{
+		if (ResolveAttachmentReferences.Num() > 0)
+		{
+			check(colorAttachmentCount == ResolveAttachmentReferences.Num());
+			pResolveAttachments = ResolveAttachmentReferences.GetData();
+		}
+	}
+
+	void SetDepthStencilAttachment(FVulkanAttachmentReference<VkAttachmentReference>* DepthStencilAttachmentReference)
+	{
+		pDepthStencilAttachment = static_cast<VkAttachmentReference*>(DepthStencilAttachmentReference);
+	}
+
+	void SetInputAttachments(FVulkanAttachmentReference<VkAttachmentReference>* InputAttachmentReferences, uint32 NumInputAttachmentReferences)
+	{
+		pInputAttachments = static_cast<VkAttachmentReference*>(InputAttachmentReferences);
+		inputAttachmentCount = NumInputAttachmentReferences;
+	}
+};
+
+#if VULKAN_SUPPORTS_RENDERPASS2
+template<>
+struct FVulkanSubpassDescription<VkSubpassDescription2>
+	: public VkSubpassDescription2
+{
+	FVulkanSubpassDescription()
+	{
+		ZeroVulkanStruct(*this, VK_STRUCTURE_TYPE_SUBPASS_DESCRIPTION_2);
+		pipelineBindPoint = VK_PIPELINE_BIND_POINT_GRAPHICS;
+		viewMask = 0;
+	}
+
+	void SetColorAttachments(const TArray<FVulkanAttachmentReference<VkAttachmentReference2>>& ColorAttachmentReferences, int OverrideCount = -1)
+	{
+		colorAttachmentCount = OverrideCount == -1 ? ColorAttachmentReferences.Num() : OverrideCount;
+		pColorAttachments = ColorAttachmentReferences.GetData();
+	}
+
+	void SetResolveAttachments(const TArray<FVulkanAttachmentReference<VkAttachmentReference2>>& ResolveAttachmentReferences)
+	{
+		if (ResolveAttachmentReferences.Num() > 0)
+		{
+			check(colorAttachmentCount == ResolveAttachmentReferences.Num());
+			pResolveAttachments = ResolveAttachmentReferences.GetData();
+		}
+	}
+
+	void SetDepthStencilAttachment(FVulkanAttachmentReference<VkAttachmentReference2>* DepthStencilAttachmentReference)
+	{
+		pDepthStencilAttachment = static_cast<VkAttachmentReference2*>(DepthStencilAttachmentReference);
+	}
+
+	void SetInputAttachments(FVulkanAttachmentReference<VkAttachmentReference2>* InputAttachmentReferences, uint32 NumInputAttachmentReferences)
+	{
+		pInputAttachments = static_cast<VkAttachmentReference2*>(InputAttachmentReferences);
+		inputAttachmentCount = NumInputAttachmentReferences;
+	}
+};
+#endif
+
+template <typename TSubpassDependencyType>
+struct FVulkanSubpassDependency
+	: public TSubpassDependencyType
+{
+};
+
+template<>
+struct FVulkanSubpassDependency<VkSubpassDependency>
+	: public VkSubpassDependency
+{
+	FVulkanSubpassDependency()
+	{
+		FMemory::Memzero(this, sizeof(VkSubpassDependency));
+	}
+};
+
+#if VULKAN_SUPPORTS_RENDERPASS2
+template<>
+struct FVulkanSubpassDependency<VkSubpassDependency2>
+	: public VkSubpassDependency2
+{
+	FVulkanSubpassDependency()
+	{
+		ZeroVulkanStruct(*this, VK_STRUCTURE_TYPE_SUBPASS_DEPENDENCY_2);
+		viewOffset = 0;		// According to the Vulkan spec: "If dependencyFlags does not include VK_DEPENDENCY_VIEW_LOCAL_BIT, viewOffset must be 0"
+	}
+};
+#endif
+
+template<typename TAttachmentDescriptionType>
+struct FVulkanAttachmentDescription
+{
+};
+
+template<>
+struct FVulkanAttachmentDescription<VkAttachmentDescription>
+	: public VkAttachmentDescription
+{
+	FVulkanAttachmentDescription()
+	{
+		FMemory::Memzero(this, sizeof(VkAttachmentDescription));
+	}
+
+	FVulkanAttachmentDescription(const VkAttachmentDescription& InDesc)
+	{
+		flags = InDesc.flags;
+		format = InDesc.format;
+		samples = InDesc.samples;
+		loadOp = InDesc.loadOp;
+		storeOp = InDesc.storeOp;
+		stencilLoadOp = InDesc.stencilLoadOp;
+		stencilStoreOp = InDesc.stencilStoreOp;
+		initialLayout = InDesc.initialLayout;
+		finalLayout = InDesc.finalLayout;
+	}
+};
+
+#if VULKAN_SUPPORTS_RENDERPASS2
+template<>
+struct FVulkanAttachmentDescription<VkAttachmentDescription2>
+	: public VkAttachmentDescription2
+{
+	FVulkanAttachmentDescription()
+	{
+		ZeroVulkanStruct(*this, VK_STRUCTURE_TYPE_ATTACHMENT_DESCRIPTION_2);
+	}
+
+	FVulkanAttachmentDescription(const VkAttachmentDescription& InDesc)
+	{
+		sType = VK_STRUCTURE_TYPE_ATTACHMENT_DESCRIPTION_2;
+		pNext = nullptr;
+		flags = InDesc.flags;
+		format = InDesc.format;
+		samples = InDesc.samples;
+		loadOp = InDesc.loadOp;
+		storeOp = InDesc.storeOp;
+		stencilLoadOp = InDesc.stencilLoadOp;
+		stencilStoreOp = InDesc.stencilStoreOp;
+		initialLayout = InDesc.initialLayout;
+		finalLayout = InDesc.finalLayout;
+	}
+};
+#endif
+
+template <typename T>
+struct FVulkanRenderPassCreateInfo
+{};
+
+template<>
+struct FVulkanRenderPassCreateInfo<VkRenderPassCreateInfo>
+	: public VkRenderPassCreateInfo
+{
+	FVulkanRenderPassCreateInfo()
+	{
+		ZeroVulkanStruct(*this, VK_STRUCTURE_TYPE_RENDER_PASS_CREATE_INFO);
+	}
+
+	VkRenderPass Create(FVulkanDevice& Device)
+	{
+		VkRenderPass Handle = VK_NULL_HANDLE;
+		VERIFYVULKANRESULT_EXPANDED(VulkanRHI::vkCreateRenderPass(Device.GetInstanceHandle(), this, VULKAN_CPU_ALLOCATOR, &Handle));
+		return Handle;
+	}
+};
+
+#if VULKAN_SUPPORTS_RENDERPASS2
+template<>
+struct FVulkanRenderPassCreateInfo<VkRenderPassCreateInfo2>
+	: public VkRenderPassCreateInfo2
+{
+	FVulkanRenderPassCreateInfo()
+	{
+		ZeroVulkanStruct(*this, VK_STRUCTURE_TYPE_RENDER_PASS_CREATE_INFO_2);
+	}
+
+	VkRenderPass Create(FVulkanDevice& Device)
+	{
+		VkRenderPass Handle = VK_NULL_HANDLE;
+		VERIFYVULKANRESULT_EXPANDED(VulkanRHI::vkCreateRenderPass2KHR(Device.GetInstanceHandle(), this, VULKAN_CPU_ALLOCATOR, &Handle));
+		return Handle;
+	}
+};
+#endif
+
+VkRenderPass CreateVulkanRenderPass(FVulkanDevice& Device, const FVulkanRenderTargetLayout& RTLayout);
+
