Fix issue with multi-layer textures that have mixed sub-resource layouts (good catch Mihnea!)

diff --git a/Engine/Source/Runtime/VulkanRHI/Private/VulkanBarriers.cpp b/Engine/Source/Runtime/VulkanRHI/Private/VulkanBarriers.cpp
--- a/Engine/Source/Runtime/VulkanRHI/Private/VulkanBarriers.cpp
+++ b/Engine/Source/Runtime/VulkanRHI/Private/VulkanBarriers.cpp
@@ -1274,7 +1274,7 @@ void FVulkanPipelineBarrier::AddImageLayoutTransition(VkImage Image, VkImageAspe
 	VkImageSubresourceRange SubresourceRange = MakeSubresourceRange(AspectMask, 0, 1, 0, 1);
 	for (; SubresourceRange.baseArrayLayer < SrcLayout.NumLayers; ++SubresourceRange.baseArrayLayer)
 	{
-		for (; SubresourceRange.baseMipLevel < SrcLayout.NumMips; ++SubresourceRange.baseMipLevel)
+		for (SubresourceRange.baseMipLevel = 0; SubresourceRange.baseMipLevel < SrcLayout.NumMips; ++SubresourceRange.baseMipLevel)
 		{
 			const VkImageLayout SubresourceLayout = SrcLayout.GetSubresLayout(SubresourceRange.baseArrayLayer, SubresourceRange.baseMipLevel);
 			if (SubresourceLayout != DstLayout)
@@ -1303,7 +1303,7 @@ void FVulkanPipelineBarrier::AddImageLayoutTransition(VkImage Image, VkImageAspe
 	VkImageSubresourceRange SubresourceRange = MakeSubresourceRange(AspectMask, 0, 1, 0, 1);
 	for (; SubresourceRange.baseArrayLayer < DstLayout.NumLayers; ++SubresourceRange.baseArrayLayer)
 	{
-		for (; SubresourceRange.baseMipLevel < DstLayout.NumMips; ++SubresourceRange.baseMipLevel)
+		for (SubresourceRange.baseMipLevel = 0; SubresourceRange.baseMipLevel < DstLayout.NumMips; ++SubresourceRange.baseMipLevel)
 		{
 			const VkImageLayout SubresourceLayout = DstLayout.GetSubresLayout(SubresourceRange.baseArrayLayer, SubresourceRange.baseMipLevel);
 			if (SubresourceLayout != SrcLayout)
