Reduce spam in logs when using Vulkan.

diff --git a/Engine/Source/Runtime/VulkanRHI/Private/VulkanTexture.cpp b/Engine/Source/Runtime/VulkanRHI/Private/VulkanTexture.cpp
--- a/Engine/Source/Runtime/VulkanRHI/Private/VulkanTexture.cpp
+++ b/Engine/Source/Runtime/VulkanRHI/Private/VulkanTexture.cpp
@@ -775,7 +775,11 @@ void FVulkanSurface::InternalMoveSurface(FVulkanDevice& InDevice, FVulkanCommand
 		check(Image != VK_NULL_HANDLE);
 		InDevice.NotifyDeletedImage(Image, bRenderTarget);
 		InDevice.GetDeferredDeletionQueue().EnqueueResource(VulkanRHI::FDeferredDeletionQueue2::EType::Image, Image);
-		FGenericPlatformMisc::LowLevelOutputDebugStringf(TEXT("** MOVE IMAGE %p -> %p\n"), Image, MovedImage);
+
+		if (GVulkanLogDefrag)
+		{
+			FGenericPlatformMisc::LowLevelOutputDebugStringf(TEXT("** MOVE IMAGE %p -> %p\n"), Image, MovedImage);
+		}
 	}
 
 	Image = MovedImage;
