Before adding Vulkan usage/memorytype bits for memoryless, check if the device has support for it.  When using the Mobile Preview on Linux desktop, a surface could have the memoryless UEFlag but be allocated out of regular heap.  When the heap gets defragged, asserts would trigger incorrectly.  (keeping the UEFlag even when not supported can help track other memoryless related errors from desktop)

diff --git a/Engine/Source/Runtime/VulkanRHI/Private/VulkanMemory.cpp b/Engine/Source/Runtime/VulkanRHI/Private/VulkanMemory.cpp
--- a/Engine/Source/Runtime/VulkanRHI/Private/VulkanMemory.cpp
+++ b/Engine/Source/Runtime/VulkanRHI/Private/VulkanMemory.cpp
@@ -570,10 +570,11 @@ namespace VulkanRHI
 
 	FDeviceMemoryManager::FDeviceMemoryManager() :
 		DeviceHandle(VK_NULL_HANDLE),
-		bHasUnifiedMemory(false),
 		Device(nullptr),
 		NumAllocations(0),
-		PeakNumAllocations(0)
+		PeakNumAllocations(0),
+		bHasUnifiedMemory(false),
+		bSupportsMemoryless(false)
 	{
 		FMemory::Memzero(MemoryProperties);
 	}
@@ -654,6 +655,13 @@ namespace VulkanRHI
 			PrimaryHostHeap = -1; // if there are no non-local heaps, disable eviction and defragmentation
 		}
 
+		// Update bMemoryless support
+		bSupportsMemoryless = false;
+		for (uint32 i = 0; i < MemoryProperties.memoryTypeCount && !bSupportsMemoryless; ++i)
+		{
+			bSupportsMemoryless = ((MemoryProperties.memoryTypes[i].propertyFlags & VK_MEMORY_PROPERTY_LAZILY_ALLOCATED_BIT) == VK_MEMORY_PROPERTY_LAZILY_ALLOCATED_BIT);
+		}
+
 
 		HeapInfos.AddDefaulted(MemoryProperties.memoryHeapCount);
 
diff --git a/Engine/Source/Runtime/VulkanRHI/Private/VulkanTexture.cpp b/Engine/Source/Runtime/VulkanRHI/Private/VulkanTexture.cpp
--- a/Engine/Source/Runtime/VulkanRHI/Private/VulkanTexture.cpp
+++ b/Engine/Source/Runtime/VulkanRHI/Private/VulkanTexture.cpp
@@ -323,7 +323,7 @@ void FVulkanSurface::GenerateImageCreateInfo(
 		}
 		ImageCreateInfo.usage |= (EnumHasAnyFlags(UEFlags, TexCreate_RenderTargetable) ? VK_IMAGE_USAGE_COLOR_ATTACHMENT_BIT : VK_IMAGE_USAGE_DEPTH_STENCIL_ATTACHMENT_BIT);
 		ImageCreateInfo.tiling = VK_IMAGE_TILING_OPTIMAL;
-		if (EnumHasAllFlags(UEFlags, TexCreate_Memoryless))
+		if (EnumHasAllFlags(UEFlags, TexCreate_Memoryless) && InDevice.GetDeviceMemoryManager().SupportsMemoryless())
 		{
 			ImageCreateInfo.usage |= VK_IMAGE_USAGE_TRANSIENT_ATTACHMENT_BIT;
 			// Remove the transfer and sampled bits, as they are incompatible with the transient bit.
@@ -614,7 +614,7 @@ FVulkanSurface::FVulkanSurface(FVulkanDevice& InDevice, FVulkanEvictable* Owner,
 
 	VkMemoryPropertyFlags MemoryFlags = bCPUReadback ? VK_MEMORY_PROPERTY_HOST_VISIBLE_BIT : VK_MEMORY_PROPERTY_DEVICE_LOCAL_BIT;
 
-	bool bMemoryless = EnumHasAnyFlags(UEFlags, TexCreate_Memoryless);
+	bool bMemoryless = EnumHasAnyFlags(UEFlags, TexCreate_Memoryless) && InDevice.GetDeviceMemoryManager().SupportsMemoryless();
 	if (bMemoryless)
 	{
 		if (ensureMsgf(bRenderTarget, TEXT("Memoryless surfaces can only be used for render targets")) && ensureMsgf(!bCPUReadback, TEXT("Memoryless surfaces cannot be read back on CPU")))
@@ -696,7 +696,7 @@ void FVulkanSurface::InternalMoveSurface(FVulkanDevice& InDevice, FVulkanCommand
 	const bool bCPUReadback = EnumHasAnyFlags(UEFlags, TexCreate_CPUReadback);
 	const bool bMemoryless = EnumHasAnyFlags(UEFlags, TexCreate_Memoryless);
 	checkf(!bCPUReadback, TEXT("Move of CPUReadback surfaces not currently supported.   UEFlags=0x%x"), (int32)UEFlags);
-	checkf(!bMemoryless, TEXT("Move of Memoryless surfaces not currently supported.   UEFlags=0x%x"), (int32)UEFlags);
+	checkf(!bMemoryless || !InDevice.GetDeviceMemoryManager().SupportsMemoryless(), TEXT("Move of Memoryless surfaces not currently supported.   UEFlags=0x%x"), (int32)UEFlags);
 
 #if UE_BUILD_DEBUG || UE_BUILD_DEVELOPMENT
 	// This shouldn't change
diff --git a/Engine/Source/Runtime/VulkanRHI/Public/VulkanMemory.h b/Engine/Source/Runtime/VulkanRHI/Public/VulkanMemory.h
--- a/Engine/Source/Runtime/VulkanRHI/Public/VulkanMemory.h
+++ b/Engine/Source/Runtime/VulkanRHI/Public/VulkanMemory.h
@@ -372,6 +372,11 @@ namespace VulkanRHI
 			return bHasUnifiedMemory;
 		}
 
+		inline bool SupportsMemoryless() const
+		{
+			return bSupportsMemoryless;
+		}
+
 		inline uint32 GetNumMemoryTypes() const
 		{
 			return MemoryProperties.memoryTypeCount;
@@ -408,10 +413,11 @@ namespace VulkanRHI
 		VkPhysicalDeviceMemoryBudgetPropertiesEXT MemoryBudget;
 		VkPhysicalDeviceMemoryProperties MemoryProperties;
 		VkDevice DeviceHandle;
-		bool bHasUnifiedMemory;
 		FVulkanDevice* Device;
 		uint32 NumAllocations;
 		uint32 PeakNumAllocations;
+		bool bHasUnifiedMemory;
+		bool bSupportsMemoryless;
 
 		struct FHeapInfo
 		{
