VK: Always use staging buffers to update 'Static' buffers since they could be in use by GPU. Buffer initialization will still use a direct upload on UMA devices

diff --git a/Engine/Source/Runtime/VulkanRHI/Private/VulkanBuffer.cpp b/Engine/Source/Runtime/VulkanRHI/Private/VulkanBuffer.cpp
--- a/Engine/Source/Runtime/VulkanRHI/Private/VulkanBuffer.cpp
+++ b/Engine/Source/Runtime/VulkanRHI/Private/VulkanBuffer.cpp
@@ -91,6 +91,7 @@ FVulkanResourceMultiBuffer::FVulkanResourceMultiBuffer(FVulkanDevice* InDevice,
 	, BufferUsageFlags(0)
 	, NumBuffers(0)
 	, DynamicBufferIndex(0)
+	, LockStatus(ELockStatus::Unlocked)
 {
 	VULKAN_TRACK_OBJECT_CREATE(FVulkanResourceMultiBuffer, this);
 
@@ -174,9 +175,18 @@ FVulkanResourceMultiBuffer::FVulkanResourceMultiBuffer(FVulkanDevice* InDevice,
 			if (CreateInfo.ResourceArray)
 			{
 				uint32 CopyDataSize = FMath::Min(InSize, CreateInfo.ResourceArray->GetResourceDataSize());
-				void* Data = Lock(bRenderThread, RLM_WriteOnly, CopyDataSize, 0);
-				FMemory::Memcpy(Data, CreateInfo.ResourceArray->GetResourceData(), CopyDataSize);
-				Unlock(bRenderThread);
+				// We know this buffer is not in use by GPU atm. If we do have a direct access initialize it without extra copies
+				if (bUnifiedMem)
+				{
+					void* Data = (uint8*)Buffers[DynamicBufferIndex].GetMappedPointer(Device);
+					FMemory::Memcpy(Data, CreateInfo.ResourceArray->GetResourceData(), CopyDataSize);
+				}
+				else
+				{
+					void* Data = Lock(bRenderThread, RLM_WriteOnly, CopyDataSize, 0);
+					FMemory::Memcpy(Data, CreateInfo.ResourceArray->GetResourceData(), CopyDataSize);
+					Unlock(bRenderThread);
+				}
 
 				CreateInfo.ResourceArray->Discard();
 			}
@@ -209,6 +219,8 @@ void* FVulkanResourceMultiBuffer::Lock(bool bFromRenderingThread, EResourceLockM
 	const bool bUAV = EnumHasAnyFlags(GetUsage(), BUF_UnorderedAccess);
 	const bool bSR = EnumHasAnyFlags(GetUsage(), BUF_ShaderResource);
 
+	LockStatus = ELockStatus::Locked;
+
 	if (bVolatile)
 	{
 		check(NumBuffers == 0);
@@ -238,6 +250,7 @@ void* FVulkanResourceMultiBuffer::Lock(bool bFromRenderingThread, EResourceLockM
 			if (bUnifiedMem)
 			{
 				Data = (uint8*)Buffers[DynamicBufferIndex].GetMappedPointer(Device) + Offset;
+				LockStatus = ELockStatus::PersistentMapping;
 			}
 			else 
 			{
@@ -299,12 +312,9 @@ void* FVulkanResourceMultiBuffer::Lock(bool bFromRenderingThread, EResourceLockM
 			Current.Offset = Current.Alloc.Offset;
 			Current.Size = LockSize;
 
-			const bool bUnifiedMem = Device->HasUnifiedMemory();
-			if (bUnifiedMem)
-			{
-				Data = (uint8*)Buffers[DynamicBufferIndex].GetMappedPointer(Device) + Offset;
-			}
-			else
+			// Always use staging buffers to update 'Static' buffers since they maybe be in use by GPU atm
+			const bool bUseStagingBuffer = (bStatic || !Device->HasUnifiedMemory());
+			if (bUseStagingBuffer)
 			{
 				VulkanRHI::FPendingBufferLock PendingLock;
 				PendingLock.Offset = Offset;
@@ -321,7 +331,11 @@ void* FVulkanResourceMultiBuffer::Lock(bool bFromRenderingThread, EResourceLockM
 					GPendingLockIBs.Add(this, PendingLock);
 				}
 			}
-
+			else
+			{
+				Data = (uint8*)Buffers[DynamicBufferIndex].GetMappedPointer(Device) + Offset;
+				LockStatus = ELockStatus::PersistentMapping;
+			}
 		}
 	}
 
@@ -384,23 +398,16 @@ void FVulkanResourceMultiBuffer::Unlock(bool bFromRenderingThread)
 	const bool bCPUReadable = EnumHasAnyFlags(GetUsage(), BUF_KeepCPUAccessible);
 	const bool bSR = EnumHasAnyFlags(GetUsage(), BUF_ShaderResource);
 
-	if (bVolatile)
+	check(LockStatus != ELockStatus::Unlocked);
+	
+	if (bVolatile || LockStatus == ELockStatus::PersistentMapping)
 	{
-		check(NumBuffers == 0);
-
 		// Nothing to do here...
 	}
 	else
 	{
 		check(bStatic || bDynamic || bSR);
-
-		const bool bUnifiedMem = Device->HasUnifiedMemory();
-		if (bUnifiedMem)
-		{
-			// Nothing to do here...
-			return;
-		}
-
+		
 		VulkanRHI::FPendingBufferLock PendingLock;
 		bool bFound = false;
 		{
@@ -431,12 +438,16 @@ void FVulkanResourceMultiBuffer::Unlock(bool bFromRenderingThread)
 			Device->GetStagingManager().ReleaseBuffer(0, PendingLock.StagingBuffer);
 		}
 	}
+
+	LockStatus = ELockStatus::Unlocked;
 }
 
 void FVulkanResourceMultiBuffer::Swap(FVulkanResourceMultiBuffer& Other)
 {
 	FRHIBuffer::Swap(Other);
 
+	check(LockStatus == ELockStatus::Unlocked);
+
 	// FDeviceChild
 	::Swap(Device, Other.Device);
 	
diff --git a/Engine/Source/Runtime/VulkanRHI/Private/VulkanResources.h b/Engine/Source/Runtime/VulkanRHI/Private/VulkanResources.h
--- a/Engine/Source/Runtime/VulkanRHI/Private/VulkanResources.h
+++ b/Engine/Source/Runtime/VulkanRHI/Private/VulkanResources.h
@@ -1135,8 +1135,14 @@ public:
 
 protected:
 	VkBufferUsageFlags BufferUsageFlags;
-	uint32 NumBuffers;
-	uint32 DynamicBufferIndex;
+	uint8 NumBuffers;
+	uint8 DynamicBufferIndex;
+	enum class ELockStatus : uint8
+	{
+		Unlocked,
+		Locked,
+		PersistentMapping,
+	} LockStatus;
 
 	enum
 	{
