Accidentally using the Stencil StoreOp for Depth.

diff --git a/Engine/Source/Runtime/VulkanRHI/Private/VulkanRenderTarget.cpp b/Engine/Source/Runtime/VulkanRHI/Private/VulkanRenderTarget.cpp
--- a/Engine/Source/Runtime/VulkanRHI/Private/VulkanRenderTarget.cpp
+++ b/Engine/Source/Runtime/VulkanRHI/Private/VulkanRenderTarget.cpp
@@ -1490,7 +1490,7 @@ FVulkanRenderTargetLayout::FVulkanRenderTargetLayout(const FGraphicsPipelineStat
 		}
 		if (CurrDesc.samples == VK_SAMPLE_COUNT_1_BIT)
 		{
-			CurrDesc.storeOp = RenderTargetStoreActionToVulkan(Initializer.StencilTargetStoreAction);
+			CurrDesc.storeOp = RenderTargetStoreActionToVulkan(Initializer.DepthTargetStoreAction);
 			CurrDesc.stencilStoreOp = RenderTargetStoreActionToVulkan(Initializer.StencilTargetStoreAction);
 		}
 		else
