Add Vulkan support for the LargeBuffer type that was introduced to fit within the DX11 threadgroup limit issue (UE-118224, #16683799).

Add messages to the ensure(0) in ClearUAV for easier debugging...

diff --git a/Engine/Source/Runtime/VulkanRHI/Private/VulkanUAV.cpp b/Engine/Source/Runtime/VulkanRHI/Private/VulkanUAV.cpp
--- a/Engine/Source/Runtime/VulkanRHI/Private/VulkanUAV.cpp
+++ b/Engine/Source/Runtime/VulkanRHI/Private/VulkanUAV.cpp
@@ -521,8 +521,18 @@ void FVulkanCommandListContext::ClearUAV(TRHICommandList_RecursiveHazardous<FVul
 		}
 		else
 		{
-			uint32 NumElements = Buffer->GetCurrentSize() / GPixelFormats[UnorderedAccessView->BufferViewFormat].BlockBytes;
-			ClearUAVShader_T<EClearReplacementResourceType::Buffer, 4, false>(RHICmdList, UnorderedAccessView, NumElements, 1, 1, ClearValue, ValueType);
+			const uint32 NumElements = Buffer->GetCurrentSize() / GPixelFormats[UnorderedAccessView->BufferViewFormat].BlockBytes;
+			const uint32 ComputeWorkGroupCount = FMath::DivideAndRoundUp(NumElements, (uint32)ClearReplacementCS::TThreadGroupSize<EClearReplacementResourceType::Buffer>::X);
+			FVulkanDevice* TargetDevice = FVulkanCommandListContext::GetVulkanContext(RHICmdList.GetContext()).GetDevice();
+			const bool bOversizedBuffer = (ComputeWorkGroupCount > TargetDevice->GetLimits().maxComputeWorkGroupCount[0]);
+			if (bOversizedBuffer)
+			{
+				ClearUAVShader_T<EClearReplacementResourceType::LargeBuffer, 4, false>(RHICmdList, UnorderedAccessView, NumElements, 1, 1, ClearValue, ValueType);
+			}
+			else
+			{
+				ClearUAVShader_T<EClearReplacementResourceType::Buffer, 4, false>(RHICmdList, UnorderedAccessView, NumElements, 1, 1, ClearValue, ValueType);
+			}
 		}
 	}
 	else if (UnorderedAccessView->SourceTexture)
@@ -547,12 +557,13 @@ void FVulkanCommandListContext::ClearUAV(TRHICommandList_RecursiveHazardous<FVul
 		}
 		else
 		{
-			ensure(0);
+			ensureMsgf(0, TEXT("SourceTexture of unknown type (Name=[%s], Format=%d, Flags=0x%x)!  Skipping ClearUAV..."), 
+				*UnorderedAccessView->SourceTexture->GetName().ToString(), (uint32)UnorderedAccessView->SourceTexture->GetFormat(), (uint32)UnorderedAccessView->SourceTexture->GetFlags());
 		}
 	}
 	else
 	{
-		ensure(0);
+		ensureMsgf(0, TEXT("UnorderedAccessView has no source buffer or texture!  Skipping ClearUAV..."));
 	}
 }
 
