Enable fragment shading rate functionality when supported.

diff --git a/Engine/Source/Runtime/VulkanRHI/Private/Linux/VulkanLinuxPlatform.cpp b/Engine/Source/Runtime/VulkanRHI/Private/Linux/VulkanLinuxPlatform.cpp
--- a/Engine/Source/Runtime/VulkanRHI/Private/Linux/VulkanLinuxPlatform.cpp
+++ b/Engine/Source/Runtime/VulkanRHI/Private/Linux/VulkanLinuxPlatform.cpp
@@ -181,6 +181,12 @@ void FVulkanLinuxPlatform::GetDeviceExtensions(EGpuVendorId VendorId, TArray<con
 
 #if VULKAN_SUPPORTS_RENDERPASS2
 	OutExtensions.Add(VK_KHR_CREATE_RENDERPASS_2_EXTENSION_NAME);
+
+	// Fragment shading rate depends on renderpass2.
+#if VULKAN_SUPPORTS_FRAGMENT_SHADING_RATE
+	OutExtensions.Add(VK_KHR_FRAGMENT_SHADING_RATE_EXTENSION_NAME);
+#endif
+
 #endif
 
 	if (GGPUCrashDebuggingEnabled)
diff --git a/Engine/Source/Runtime/VulkanRHI/Private/VulkanDevice.cpp b/Engine/Source/Runtime/VulkanRHI/Private/VulkanDevice.cpp
--- a/Engine/Source/Runtime/VulkanRHI/Private/VulkanDevice.cpp
+++ b/Engine/Source/Runtime/VulkanRHI/Private/VulkanDevice.cpp
@@ -415,6 +415,16 @@ void FVulkanDevice::CreateDevice()
 	}
 #endif
 
+#if VULKAN_SUPPORTS_FRAGMENT_SHADING_RATE
+	VkPhysicalDeviceFragmentShadingRateFeaturesKHR RequestedFragmentShadingRateFeatures;
+	if (OptionalDeviceExtensions.HasKHRFragmentShadingRate)
+	{
+		RequestedFragmentShadingRateFeatures = GetFragmentShadingRateFeatures();
+		RequestedFragmentShadingRateFeatures.pNext = (void*)DeviceInfo.pNext;
+		DeviceInfo.pNext = &RequestedFragmentShadingRateFeatures;
+	}
+#endif
+
 	// Create the device
 	VkResult Result = VulkanRHI::vkCreateDevice(Gpu, &DeviceInfo, VULKAN_CPU_ALLOCATOR, &Device);
 	if (Result == VK_ERROR_INITIALIZATION_FAILED)
@@ -1473,6 +1483,36 @@ const VkComponentMapping& FVulkanDevice::GetFormatComponentMapping(EPixelFormat
 	return PixelFormatComponentMapping[UEFormat];
 }
 
+#if VULKAN_SUPPORTS_FRAGMENT_SHADING_RATE
+VkExtent2D FVulkanDevice::GetBestMatchedShadingRateExtents(EVRSShadingRate ShadingRate) const
+{
+	// Given that for Vulkan we need to query available device shading rates, we're not guaranteed to have everything that's in our enum;
+	// This function walks the list of supported fragment rates returned by the device, and returns the closest match to the rate requested.
+	const VkExtent2D DirectMappedExtent = { (uint32)(ShadingRate >> 2) + 1, (uint32)((ShadingRate & 0x03) + 1) };
+	VkExtent2D BestMatchedExtent = { 1, 1 };
+	
+	if (BestMatchedExtent.width != DirectMappedExtent.width && BestMatchedExtent.height != DirectMappedExtent.height)
+	{
+		for (auto const& Rate : FragmentShadingRates)
+		{
+			if (Rate.fragmentSize.width == DirectMappedExtent.width && Rate.fragmentSize.height == DirectMappedExtent.height)
+			{
+				BestMatchedExtent = DirectMappedExtent;
+				break;
+			}
+
+			if ((Rate.fragmentSize.width >= BestMatchedExtent.width && Rate.fragmentSize.width <= DirectMappedExtent.width && Rate.fragmentSize.height <= DirectMappedExtent.height && Rate.fragmentSize.height >= BestMatchedExtent.height) ||
+				(Rate.fragmentSize.height >= BestMatchedExtent.height && Rate.fragmentSize.height <= DirectMappedExtent.height && Rate.fragmentSize.width <= DirectMappedExtent.width && Rate.fragmentSize.width >= BestMatchedExtent.width))
+			{
+				BestMatchedExtent = Rate.fragmentSize;
+			}
+		}
+	}
+
+	return BestMatchedExtent;
+}
+#endif
+
 void FVulkanDevice::NotifyDeletedImage(VkImage Image, bool bRenderTarget)
 {
 	if (bRenderTarget)
diff --git a/Engine/Source/Runtime/VulkanRHI/Private/VulkanDevice.h b/Engine/Source/Runtime/VulkanRHI/Private/VulkanDevice.h
--- a/Engine/Source/Runtime/VulkanRHI/Private/VulkanDevice.h
+++ b/Engine/Source/Runtime/VulkanRHI/Private/VulkanDevice.h
@@ -286,8 +286,10 @@ public:
 	{
 		return FragmentShadingRateProperties;
 	}
-#endif
 
+	VkExtent2D GetBestMatchedShadingRateExtents(EVRSShadingRate Rate) const;
+#endif
+	
 #if VULKAN_SUPPORTS_MULTIVIEW
 	inline const VkPhysicalDeviceMultiviewFeatures& GetMultiviewFeatures() const
 	{
diff --git a/Engine/Source/Runtime/VulkanRHI/Private/VulkanLayers.cpp b/Engine/Source/Runtime/VulkanRHI/Private/VulkanLayers.cpp
--- a/Engine/Source/Runtime/VulkanRHI/Private/VulkanLayers.cpp
+++ b/Engine/Source/Runtime/VulkanRHI/Private/VulkanLayers.cpp
@@ -867,10 +867,7 @@ void FOptionalVulkanDeviceExtensions::Setup(const TArray<const ANSICHAR*>& Devic
 #endif
 
 #if VULKAN_SUPPORTS_FRAGMENT_SHADING_RATE
-	// TODO: the VK_KHR_fragment_shading_rate extension is dependent on vkCreateRenderPass2, VkRenderPassCreateInfo2, VkAttachmentDescription2 and VkSubpassDescription2.
-	// Disabling this path for now; adding this support in a later checkin.
-
-	// HasKHRFragmentShadingRate = HasExtension(DeviceExtensions, VK_KHR_FRAGMENT_SHADING_RATE_EXTENSION_NAME);
+	HasKHRFragmentShadingRate = HasExtension(DeviceExtensions, VK_KHR_FRAGMENT_SHADING_RATE_EXTENSION_NAME);
 #endif
 
 #if VULKAN_SUPPORTS_MULTIVIEW
diff --git a/Engine/Source/Runtime/VulkanRHI/Private/VulkanPipeline.cpp b/Engine/Source/Runtime/VulkanRHI/Private/VulkanPipeline.cpp
--- a/Engine/Source/Runtime/VulkanRHI/Private/VulkanPipeline.cpp
+++ b/Engine/Source/Runtime/VulkanRHI/Private/VulkanPipeline.cpp
@@ -1022,6 +1022,14 @@ FArchive& operator << (FArchive& Ar, FGfxPipelineDesc& Entry)
 	}
 #endif
 
+#if VULKAN_SUPPORTS_FRAGMENT_SHADING_RATE
+	uint8 ShadingRate = static_cast<uint8>(Entry.ShadingRate);
+	uint8 Combiner = static_cast<uint8>(Entry.Combiner);
+	
+	Ar << ShadingRate;
+	Ar << Combiner;
+#endif
+
 	Ar << Entry.UseAlphaToCoverage;
 
 	return Ar;
@@ -1051,7 +1059,18 @@ FVulkanPSOKey FGfxPipelineDesc::CreateKey2() const
 	return Result;
 }
 
-
+#if VULKAN_SUPPORTS_FRAGMENT_SHADING_RATE
+// Map Unreal VRS combiner operation enums to Vulkan enums.
+static const TMap<uint8, VkFragmentShadingRateCombinerOpKHR> FragmentCombinerOpMap
+{
+	{ VRSRB_Passthrough,	VK_FRAGMENT_SHADING_RATE_COMBINER_OP_KEEP_KHR },
+	{ VRSRB_Override,	VK_FRAGMENT_SHADING_RATE_COMBINER_OP_REPLACE_KHR },
+	{ VRSRB_Min,			VK_FRAGMENT_SHADING_RATE_COMBINER_OP_MIN_KHR },
+	{ VRSRB_Max,			VK_FRAGMENT_SHADING_RATE_COMBINER_OP_MAX_KHR },
+	{ VRSRB_Sum,			VK_FRAGMENT_SHADING_RATE_COMBINER_OP_MAX_KHR },		// No concept of Sum in Vulkan - fall back to max.
+	// @todo: Add "VK_FRAGMENT_SHADING_RATE_COMBINER_OP_MUL_KHR"?
+};
+#endif
 
 bool FVulkanPipelineStateCacheManager::CreateGfxPipelineFromEntry(FVulkanRHIGraphicsPipelineState* PSO, FVulkanShader* Shaders[ShaderStage::NumStages], VkPipeline* Pipeline)
 {
@@ -1192,6 +1211,23 @@ bool FVulkanPipelineStateCacheManager::CreateGfxPipelineFromEntry(FVulkanRHIGrap
 
 	PipelineInfo.pDynamicState = &DynamicState;
 
+#if VULKAN_SUPPORTS_FRAGMENT_SHADING_RATE
+	const VkExtent2D FragmentSize = Device->GetBestMatchedShadingRateExtents(PSO->Desc.ShadingRate);
+	VkFragmentShadingRateCombinerOpKHR PipelineToPrimitiveCombinerOperation = FragmentCombinerOpMap[(uint8)PSO->Desc.Combiner];
+
+	VkPipelineFragmentShadingRateStateCreateInfoKHR PipelineFragmentShadingRate;
+
+	if (GRHISupportsPipelineVariableRateShading && GRHIVariableRateShadingEnabled)
+	{
+		ZeroVulkanStruct(PipelineFragmentShadingRate, VK_STRUCTURE_TYPE_PIPELINE_FRAGMENT_SHADING_RATE_STATE_CREATE_INFO_KHR);
+		PipelineFragmentShadingRate.fragmentSize = FragmentSize;
+		PipelineFragmentShadingRate.combinerOps[0] = PipelineToPrimitiveCombinerOperation;
+		PipelineFragmentShadingRate.combinerOps[1] = VK_FRAGMENT_SHADING_RATE_COMBINER_OP_MAX_KHR;		// @todo: This needs to be specified too.
+
+		PipelineInfo.pNext = (void*)&PipelineFragmentShadingRate;
+	}
+#endif
+
 	VkResult Result = VK_ERROR_INITIALIZATION_FAILED;
 	double BeginTime = FPlatformTime::Seconds();
 	if(bUseLRU)
@@ -1565,6 +1601,14 @@ void FVulkanPipelineStateCacheManager::CreateGfxEntry(const FGraphicsPipelineSta
 
 	FVulkanRenderTargetLayout RTLayout(PSOInitializer);
 	OutGfxEntry->RenderTargets.ReadFrom(RTLayout);
+
+	// Shading rate:
+	OutGfxEntry->ShadingRate = PSOInitializer.ShadingRate;
+	OutGfxEntry->Combiner = EVRSRateCombiner::VRSRB_Max;		// @todo: This needs to be specified twice; from pipeline-to-primitive, and from primitive-to-attachment. 
+																// We don't have per-primitive VRS so that should just be hard-coded to "passthrough" until this is supported; but we should expose 
+																// this setting in the material properies, especially since there's some materials that don't play nicely with 
+																// shading rates other than 1x1, in which case we'll want to use VRSRB_Min to override e.g. the attachment shading rate.
+																// For now, just locked to "max".
 }
 
 
diff --git a/Engine/Source/Runtime/VulkanRHI/Private/VulkanPipeline.h b/Engine/Source/Runtime/VulkanRHI/Private/VulkanPipeline.h
--- a/Engine/Source/Runtime/VulkanRHI/Private/VulkanPipeline.h
+++ b/Engine/Source/Runtime/VulkanRHI/Private/VulkanPipeline.h
@@ -367,6 +367,9 @@ struct FGfxPipelineDesc
 
 	uint8 UseAlphaToCoverage;
 
+	EVRSShadingRate ShadingRate = EVRSShadingRate::VRSSR_1x1;
+	EVRSRateCombiner Combiner = EVRSRateCombiner::VRSRB_Passthrough;
+
 	bool operator==(const FGfxPipelineDesc& In) const
 	{
 		if (VertexInputKey != In.VertexInputKey)
@@ -450,6 +453,18 @@ struct FGfxPipelineDesc
 			return false;
 		}
 
+#if VULKAN_SUPPORTS_FRAGMENT_SHADING_RATE
+		if (ShadingRate != In.ShadingRate)
+		{
+			return false;
+		}
+		
+		if (Combiner != In.Combiner)
+		{
+			return false;
+		}
+#endif
+
 		return true;
 	}
 };
diff --git a/Engine/Source/Runtime/VulkanRHI/Private/Windows/VulkanWindowsPlatform.cpp b/Engine/Source/Runtime/VulkanRHI/Private/Windows/VulkanWindowsPlatform.cpp
--- a/Engine/Source/Runtime/VulkanRHI/Private/Windows/VulkanWindowsPlatform.cpp
+++ b/Engine/Source/Runtime/VulkanRHI/Private/Windows/VulkanWindowsPlatform.cpp
@@ -242,8 +242,7 @@ void FVulkanWindowsPlatform::GetDeviceExtensions(EGpuVendorId VendorId, TArray<c
 #endif
 
 #if VULKAN_SUPPORTS_FRAGMENT_SHADING_RATE
-	// Disabling as validation complains: Missing extension required by the device extension VK_KHR_fragment_shading_rate: VK_KHR_create_renderpass2
-	//OutExtensions.Add(VK_KHR_FRAGMENT_SHADING_RATE_EXTENSION_NAME);
+	OutExtensions.Add(VK_KHR_FRAGMENT_SHADING_RATE_EXTENSION_NAME);
 #endif
 }
 
