Add RenderArea as a criteria for Framebuffer matching when reusing, fixes some passes that were reusing the same set of targets but different areas (for example Nanite::Rasterize).  Add accessors for Offset2D in FVulkanRenderTargetLayout (and fix copy-pasta typo).

diff --git a/Engine/Source/Runtime/VulkanRHI/Private/VulkanBarriers.cpp b/Engine/Source/Runtime/VulkanRHI/Private/VulkanBarriers.cpp
--- a/Engine/Source/Runtime/VulkanRHI/Private/VulkanBarriers.cpp
+++ b/Engine/Source/Runtime/VulkanRHI/Private/VulkanBarriers.cpp
@@ -1526,7 +1526,11 @@ FVulkanFramebuffer* FVulkanLayoutManager::GetOrCreateFramebuffer(FVulkanDevice&
 
 		for (int32 Index = 0; Index < FramebufferList->Framebuffer.Num(); ++Index)
 		{
-			if (FramebufferList->Framebuffer[Index]->Matches(RenderTargetsInfo))
+			const VkRect2D RenderArea = FramebufferList->Framebuffer[Index]->GetRenderArea();
+
+			if (FramebufferList->Framebuffer[Index]->Matches(RenderTargetsInfo) &&
+				((RTLayout.GetExtent2D().width == RenderArea.extent.width) && (RTLayout.GetExtent2D().height == RenderArea.extent.height) &&
+				(RTLayout.GetOffset2D().x == RenderArea.offset.x) && (RTLayout.GetOffset2D().y == RenderArea.offset.y)))
 			{
 				return FramebufferList->Framebuffer[Index];
 			}
diff --git a/Engine/Source/Runtime/VulkanRHI/Private/VulkanRHIPrivate.h b/Engine/Source/Runtime/VulkanRHI/Private/VulkanRHIPrivate.h
--- a/Engine/Source/Runtime/VulkanRHI/Private/VulkanRHIPrivate.h
+++ b/Engine/Source/Runtime/VulkanRHI/Private/VulkanRHIPrivate.h
@@ -149,6 +149,8 @@ public:
 		check(bCalculatedHash);
 		return RenderPassFullHash;
 	}
+	inline const VkOffset2D& GetOffset2D() const { return Offset.Offset2D; }
+	inline const VkOffset3D& GetOffset3D() const { return Offset.Offset3D; }
 	inline const VkExtent2D& GetExtent2D() const { return Extent.Extent2D; }
 	inline const VkExtent3D& GetExtent3D() const { return Extent.Extent3D; }
 	inline const VkAttachmentDescription* GetAttachmentDescriptions() const { return Desc; }
@@ -201,7 +203,7 @@ protected:
 	union
 	{
 		VkOffset3D Offset3D;
-		VkOffset3D Offset2D;
+		VkOffset2D Offset2D;
 	} Offset;
 
 	union
