Change the use of VK_NULL_HANDLE to 0. Some platforms define it as 0 while others as nullptr. LLVM 11 fails to compile if nullptr is assigned to an int64

diff --git a/Engine/Source/Runtime/VulkanRHI/Private/VulkanMemory.cpp b/Engine/Source/Runtime/VulkanRHI/Private/VulkanMemory.cpp
--- a/Engine/Source/Runtime/VulkanRHI/Private/VulkanMemory.cpp
+++ b/Engine/Source/Runtime/VulkanRHI/Private/VulkanMemory.cpp
@@ -4840,7 +4840,7 @@ namespace VulkanRHI
 		Queue->GetLastSubmittedInfo(Entry.CmdBuffer, Entry.FenceCounter);
 		Entry.FrameNumber = GVulkanRHIDeletionFrameNumber;
 
-		Entry.Handle = VK_NULL_HANDLE;
+		Entry.Handle = 0;
 		Entry.DeviceMemoryAllocation = DeviceMemoryAllocation;
 		{
 			FScopeLock ScopeLock(&CS);
@@ -4862,7 +4862,7 @@ namespace VulkanRHI
 		Queue->GetLastSubmittedInfo(Entry.CmdBuffer, Entry.FenceCounter);
 		Entry.FrameNumber = GVulkanRHIDeletionFrameNumber;
 
-		Entry.Handle = VK_NULL_HANDLE;
+		Entry.Handle = 0;
 		Entry.Allocation = Allocation;
 		Entry.DeviceMemoryAllocation = 0;
 
