Add Attachment VRS functionality for Vulkan.
Plumb through all required stage flags, access masks, layouts, renderpass attachments, etc for the VK_KHR_fragment_shading_rate extension.
Fix up logic used to gather VRS properties (tile sizes, data type/format).
Add templated version of VkFragmentShadingRateAttachmentInfo, attach to renderpasses where supported and when attachment provided.

diff --git a/Engine/Config/Windows/DataDrivenPlatformInfo.ini b/Engine/Config/Windows/DataDrivenPlatformInfo.ini
--- a/Engine/Config/Windows/DataDrivenPlatformInfo.ini
+++ b/Engine/Config/Windows/DataDrivenPlatformInfo.ini
@@ -193,8 +193,7 @@ bSupportsPerPixelDBufferMask = false
 
 bIsHlslcc = false
 
-; //#todo-vulkanvrs
-bSupportsVariableRateShading = false
+bSupportsVariableRateShading = true
 
 NumberOfComputeThreads = 64
 
diff --git a/Engine/Source/Runtime/VulkanRHI/Private/VulkanBarriers.cpp b/Engine/Source/Runtime/VulkanRHI/Private/VulkanBarriers.cpp
--- a/Engine/Source/Runtime/VulkanRHI/Private/VulkanBarriers.cpp
+++ b/Engine/Source/Runtime/VulkanRHI/Private/VulkanBarriers.cpp
@@ -63,9 +63,17 @@ static VkAccessFlags GetVkAccessMaskForLayout(VkImageLayout Layout)
 			Flags = 0;
 			break;
 
+#if VULKAN_SUPPORTS_FRAGMENT_DENSITY_MAP
 		case VK_IMAGE_LAYOUT_FRAGMENT_DENSITY_MAP_OPTIMAL_EXT:
 			Flags = VK_ACCESS_FRAGMENT_DENSITY_MAP_READ_BIT_EXT;
 			break;
+#endif
+
+#if VULKAN_SUPPORTS_FRAGMENT_SHADING_RATE
+		case VK_IMAGE_LAYOUT_FRAGMENT_SHADING_RATE_ATTACHMENT_OPTIMAL_KHR:
+			Flags = VK_ACCESS_FRAGMENT_SHADING_RATE_ATTACHMENT_READ_BIT_KHR;
+			break;
+#endif
 
 		case VK_IMAGE_LAYOUT_GENERAL:
 		case VK_IMAGE_LAYOUT_UNDEFINED:
@@ -126,10 +134,18 @@ static VkPipelineStageFlags GetVkStageFlagsForLayout(VkImageLayout Layout)
 			Flags = VK_PIPELINE_STAGE_BOTTOM_OF_PIPE_BIT;
 			break;
 
+#if VULKAN_SUPPORTS_FRAGMENT_DENSITY_MAP
 		case VK_IMAGE_LAYOUT_FRAGMENT_DENSITY_MAP_OPTIMAL_EXT:
 			Flags = VK_PIPELINE_STAGE_FRAGMENT_DENSITY_PROCESS_BIT_EXT;
 			break;
+#endif
 
+#if VULKAN_SUPPORTS_FRAGMENT_SHADING_RATE
+		case VK_IMAGE_LAYOUT_FRAGMENT_SHADING_RATE_ATTACHMENT_OPTIMAL_KHR:
+			Flags = VK_PIPELINE_STAGE_FRAGMENT_SHADING_RATE_ATTACHMENT_BIT_KHR;
+			break;
+#endif
+			
 		case VK_IMAGE_LAYOUT_GENERAL:
 		case VK_IMAGE_LAYOUT_UNDEFINED:
 			Flags = VK_PIPELINE_STAGE_TOP_OF_PIPE_BIT;
@@ -334,6 +350,31 @@ static void GetVkStageAndAccessFlags(ERHIAccess RHIAccess, FRHITransitionInfo::E
 		ProcessedRHIFlags |= (uint32)(ERHIAccess::CopySrc | ERHIAccess::ResolveSrc);
 	}
 
+	if (EnumHasAnyFlags(RHIAccess, ERHIAccess::ShadingRateSource) && ValidateShadingRateDataType())
+	{
+		checkf(ResourceType == FRHITransitionInfo::EType::Texture, TEXT("A non-texture resource was tagged as a shading rate source; only textures (Texture2D and Texture2DArray) can be used for this purpose."));
+		
+#if VULKAN_SUPPORTS_FRAGMENT_SHADING_RATE
+		if (GRHIVariableRateShadingImageDataType == VRSImage_Palette)
+		{
+			StageFlags = VK_PIPELINE_STAGE_FRAGMENT_SHADING_RATE_ATTACHMENT_BIT_KHR;
+			AccessFlags = VK_ACCESS_FRAGMENT_SHADING_RATE_ATTACHMENT_READ_BIT_KHR;
+			Layout = VK_IMAGE_LAYOUT_FRAGMENT_SHADING_RATE_ATTACHMENT_OPTIMAL_KHR;
+		}
+#endif
+
+#if VULKAN_SUPPORTS_FRAGMENT_DENSITY_MAP
+		if (GRHIVariableRateShadingImageDataType == VRSImage_Fractional)
+		{
+			StageFlags = VK_PIPELINE_STAGE_FRAGMENT_DENSITY_PROCESS_BIT_EXT;
+			AccessFlags = VK_ACCESS_FRAGMENT_DENSITY_MAP_READ_BIT_EXT;
+			Layout = VK_IMAGE_LAYOUT_FRAGMENT_DENSITY_MAP_OPTIMAL_EXT;
+		}
+#endif
+
+		ProcessedRHIFlags |= (uint32)ERHIAccess::ShadingRateSource;
+	}
+
 	uint32 RemainingFlags = (uint32)RHIAccess & (~ProcessedRHIFlags);
 	ensureMsgf(RemainingFlags == 0, TEXT("Some access flags were not processed. RHIAccess=%x, ProcessedRHIFlags=%x, RemainingFlags=%x"), RHIAccess, ProcessedRHIFlags, RemainingFlags);
 }
@@ -1658,11 +1699,26 @@ void FVulkanLayoutManager::BeginRenderPass(FVulkanCommandListContext& Context, F
 	}
 
 	FRHITexture* ShadingRateTexture = RPInfo.ShadingRateTexture;
-	if (ShadingRateTexture)
+	if (ShadingRateTexture && ValidateShadingRateDataType())
 	{
 		FVulkanSurface& Surface = FVulkanTextureBase::Cast(ShadingRateTexture)->Surface;
 		VkImageLayout& DSLayout = FindOrAddLayoutRW(Surface, VK_IMAGE_LAYOUT_UNDEFINED);
-		VkImageLayout ExpectedLayout = VK_IMAGE_LAYOUT_FRAGMENT_DENSITY_MAP_OPTIMAL_EXT;
+		
+		VkImageLayout ExpectedLayout = VK_IMAGE_LAYOUT_UNDEFINED;
+
+#if VULKAN_SUPPORTS_FRAGMENT_SHADING_RATE
+		if (GRHIVariableRateShadingImageDataType == VRSImage_Palette)
+		{
+			ExpectedLayout = VK_IMAGE_LAYOUT_FRAGMENT_SHADING_RATE_ATTACHMENT_OPTIMAL_KHR;
+		}
+#endif
+
+#if VULKAN_SUPPORTS_FRAGMENT_DENSITY_MAP
+		if (GRHIVariableRateShadingImageDataType == VRSImage_Fractional)
+		{
+			ExpectedLayout = VK_IMAGE_LAYOUT_FRAGMENT_DENSITY_MAP_OPTIMAL_EXT;
+		}
+#endif
 
 		// transition shading rate textures to the fragment density map layout for rendering
 		if (DSLayout != ExpectedLayout)
diff --git a/Engine/Source/Runtime/VulkanRHI/Private/VulkanDevice.cpp b/Engine/Source/Runtime/VulkanRHI/Private/VulkanDevice.cpp
--- a/Engine/Source/Runtime/VulkanRHI/Private/VulkanDevice.cpp
+++ b/Engine/Source/Runtime/VulkanRHI/Private/VulkanDevice.cpp
@@ -1040,22 +1040,16 @@ bool FVulkanDevice::QueryGPU(int32 DeviceIndex)
 #endif
 
 #if VULKAN_SUPPORTS_FRAGMENT_SHADING_RATE
-	// TODO: the VK_KHR_fragment_shading_rate extension is dependent on vkCreateRenderPass2, VkRenderPassCreateInfo2, VkAttachmentDescription2 and VkSubpassDescription2.
-	// Disabling this path for now; adding this support in a later checkin.
 	if (GetOptionalExtensions().HasKHRFragmentShadingRate)
 	{
-		GRHISupportsAttachmentVariableRateShading = FragmentShadingRateFeatures.attachmentFragmentShadingRate ? true : false;
-		GRHISupportsPipelineVariableRateShading = FragmentShadingRateFeatures.pipelineFragmentShadingRate ? true : false;
-
 		GRHIVariableRateShadingImageTileMinWidth = FragmentShadingRateProperties.minFragmentShadingRateAttachmentTexelSize.width;
 		GRHIVariableRateShadingImageTileMinHeight = FragmentShadingRateProperties.minFragmentShadingRateAttachmentTexelSize.height;
 		GRHIVariableRateShadingImageTileMaxWidth = FragmentShadingRateProperties.maxFragmentShadingRateAttachmentTexelSize.width;
 		GRHIVariableRateShadingImageTileMaxHeight = FragmentShadingRateProperties.maxFragmentShadingRateAttachmentTexelSize.height;
 
-		GRHIVariableRateShadingImageDataType = VRSImage_Palette;
-		GRHIVariableRateShadingImageFormat = PF_R8_UINT;
+		// todo: We don't currently care much about the other properties here, but at some point in the future we probably will.
 
-		// UE_LOG(LogVulkanRHI, Display, TEXT("Image-based Variable Rate Shading supported via KHRFragmentShadingRate extension. Selected VRS tile size %u by %u pixels per VRS image texel."));
+		UE_LOG(LogVulkanRHI, Verbose, TEXT("Image-based Variable Rate Shading supported via KHRFragmentShadingRate extension. Selected VRS tile size %u by %u pixels per VRS image texel."), GRHIVariableRateShadingImageTileMinWidth, GRHIVariableRateShadingImageTileMinHeight);
 	}
 #endif
 
@@ -1271,6 +1265,25 @@ void FVulkanDevice::InitGPU(int32 DeviceIndex)
 	}
 #endif
 
+#if VULKAN_SUPPORTS_FRAGMENT_SHADING_RATE
+	if (GetOptionalExtensions().HasKHRFragmentShadingRate)
+	{
+		GRHISupportsAttachmentVariableRateShading = FragmentShadingRateFeatures.attachmentFragmentShadingRate ? true : false;
+		GRHISupportsPipelineVariableRateShading = FragmentShadingRateFeatures.pipelineFragmentShadingRate ? true : false;
+
+		if (GRHISupportsAttachmentVariableRateShading)
+		{
+			GRHIVariableRateShadingImageDataType = VRSImage_Palette;
+			GRHIVariableRateShadingImageFormat = PF_R8_UINT;
+		}
+		else
+		{
+			GRHIVariableRateShadingImageDataType = VRSImage_NotSupported;
+			GRHIVariableRateShadingImageFormat = PF_Unknown;
+		}
+	}
+#endif
+
 	PipelineStateCache->InitAndLoad(CacheFilenames);
 
 	// Setup default resource
diff --git a/Engine/Source/Runtime/VulkanRHI/Private/VulkanRHI.cpp b/Engine/Source/Runtime/VulkanRHI/Private/VulkanRHI.cpp
--- a/Engine/Source/Runtime/VulkanRHI/Private/VulkanRHI.cpp
+++ b/Engine/Source/Runtime/VulkanRHI/Private/VulkanRHI.cpp
@@ -803,8 +803,7 @@ void FVulkanDynamicRHI::InitInstance()
 #endif
 
 #if VULKAN_SUPPORTS_FRAGMENT_SHADING_RATE
-		// TODO: Complete logic when render pass support is complete for the KHR_Fragment_shading_rate extension.
-		// GRHISupportsAttachmentVariableRateShading = GetDevice()->GetOptionalExtensions().HasKHRFragmentShadingRate && Device->GetFragmentShadingRateFeatures().attachmentFragmentShadingRate;
+		GRHISupportsAttachmentVariableRateShading = GetDevice()->GetOptionalExtensions().HasKHRFragmentShadingRate && Device->GetFragmentShadingRateFeatures().attachmentFragmentShadingRate;
 #endif
 
 		FVulkanPlatform::SetupFeatureLevels();
diff --git a/Engine/Source/Runtime/VulkanRHI/Private/VulkanRHIPrivate.h b/Engine/Source/Runtime/VulkanRHI/Private/VulkanRHIPrivate.h
--- a/Engine/Source/Runtime/VulkanRHI/Private/VulkanRHIPrivate.h
+++ b/Engine/Source/Runtime/VulkanRHI/Private/VulkanRHIPrivate.h
@@ -173,6 +173,9 @@ public:
 	inline const ESubpassHint GetSubpassHint() const { return SubpassHint; }
 	inline const VkSurfaceTransformFlagBitsKHR GetQCOMRenderPassTransform() const { return QCOMRenderPassTransform; }
 
+protected:
+	VkImageLayout GetVRSImageLayout() const;
+
 protected:
 	VkSurfaceTransformFlagBitsKHR QCOMRenderPassTransform = VK_SURFACE_TRANSFORM_IDENTITY_BIT_KHR;
 	VkAttachmentReference ColorReferences[MaxSimultaneousRenderTargets];
@@ -869,6 +872,38 @@ inline bool UseVulkanDescriptorCache()
 	return (PLATFORM_ANDROID)|| GMaxRHIFeatureLevel <= ERHIFeatureLevel::ES3_1;
 }
 
+inline bool ValidateShadingRateDataType()
+{
+	switch (GRHIVariableRateShadingImageDataType)
+	{
+	case VRSImage_Palette:
+#if VULKAN_SUPPORTS_FRAGMENT_SHADING_RATE
+		return true;
+#else
+		checkf(false, TEXT("GRHIVariableRateShadingImageDataType was specified as VRSImage_Palette, but the VK_KHR_fragment_shading_rate extension is not supported on this platform."));
+		break;
+#endif
+
+	case VRSImage_Fractional:
+#if VULKAN_SUPPORTS_FRAGMENT_DENSITY_MAP
+		return true;
+#else
+		checkf(false, TEXT("GRHIVariableRateShadingImageDataType was specified as VRSImage_Fractional, but the VK_EXT_fragment_density_map extension is not supported on this platform."));
+		break;
+#endif
+
+	case VRSImage_NotSupported:
+		checkf(false, TEXT("A texture was marked as a shading rate source but attachment VRS is not supported on this device. Ensure GRHISupportsAttachmentVariableRateShading and GRHIAttachmentVariableRateShadingEnabled are true before specifying a shading rate attachment."));
+		break;
+
+	default:
+		checkf(false, TEXT("Unrecognized shading rate image data type. Specified type was %d"), (int)GRHIVariableRateShadingImageDataType);
+		break;
+	}
+
+	return false;
+}
+
 extern int32 GVulkanSubmitAfterEveryEndRenderPass;
 extern int32 GWaitForIdleOnSubmit;
 
diff --git a/Engine/Source/Runtime/VulkanRHI/Private/VulkanRenderTarget.cpp b/Engine/Source/Runtime/VulkanRHI/Private/VulkanRenderTarget.cpp
--- a/Engine/Source/Runtime/VulkanRHI/Private/VulkanRenderTarget.cpp
+++ b/Engine/Source/Runtime/VulkanRHI/Private/VulkanRenderTarget.cpp
@@ -895,6 +895,26 @@ struct FRenderPassFullHashableStruct
 #endif
 };
 
+VkImageLayout FVulkanRenderTargetLayout::GetVRSImageLayout() const
+{
+	if (ValidateShadingRateDataType())
+	{
+#if VULKAN_SUPPORTS_FRAGMENT_SHADING_RATE
+		if (GRHIVariableRateShadingImageDataType == VRSImage_Palette)
+		{
+			return VK_IMAGE_LAYOUT_FRAGMENT_SHADING_RATE_ATTACHMENT_OPTIMAL_KHR;
+		}
+#endif
+#if VULKAN_SUPPORTS_FRAGMENT_DENSITY_MAP
+		if (GRHIVariableRateShadingImageDataType == VRSImage_Fractional)
+		{
+			return VK_IMAGE_LAYOUT_FRAGMENT_DENSITY_MAP_OPTIMAL_EXT;
+		}
+#endif
+	}
+
+	return VK_IMAGE_LAYOUT_UNDEFINED;
+}
 
 FVulkanRenderTargetLayout::FVulkanRenderTargetLayout(FVulkanDevice& InDevice, const FRHISetRenderTargetsInfo& RTInfo)
 	: NumAttachmentDescriptions(0)
@@ -1073,16 +1093,18 @@ FVulkanRenderTargetLayout::FVulkanRenderTargetLayout(FVulkanDevice& InDevice, co
 		}
 	}
 
-	if (InDevice.GetOptionalExtensions().HasEXTFragmentDensityMap && RTInfo.ShadingRateTexture)
+	if (GRHISupportsAttachmentVariableRateShading && GRHIVariableRateShadingEnabled && GRHIAttachmentVariableRateShadingEnabled && RTInfo.ShadingRateTexture)
 	{
 		FVulkanTextureBase* Texture = FVulkanTextureBase::Cast(RTInfo.ShadingRateTexture);
 		check(Texture);
+		check(Texture->GetRHITexture());
+		check(Texture->GetRHITexture()->GetFormat() == GRHIVariableRateShadingImageFormat);
 
 		VkAttachmentDescription& CurrDesc = Desc[NumAttachmentDescriptions];
 		FMemory::Memzero(CurrDesc);
 
-		const VkImageLayout FragmentDensityLayout = VK_IMAGE_LAYOUT_FRAGMENT_DENSITY_MAP_OPTIMAL_EXT;
-
+		const VkImageLayout VRSLayout = GetVRSImageLayout();
+		
 		CurrDesc.flags = 0;
 		CurrDesc.format = UEToVkTextureFormat(RTInfo.ShadingRateTexture->GetFormat(), false);
 		CurrDesc.samples = static_cast<VkSampleCountFlagBits>(RTInfo.ShadingRateTexture->GetNumSamples());
@@ -1090,16 +1112,16 @@ FVulkanRenderTargetLayout::FVulkanRenderTargetLayout(FVulkanDevice& InDevice, co
 		CurrDesc.storeOp = VK_ATTACHMENT_STORE_OP_DONT_CARE;
 		CurrDesc.stencilLoadOp = VK_ATTACHMENT_LOAD_OP_DONT_CARE;
 		CurrDesc.stencilStoreOp = VK_ATTACHMENT_STORE_OP_DONT_CARE;
-		CurrDesc.initialLayout = FragmentDensityLayout;
-		CurrDesc.finalLayout = FragmentDensityLayout;
+		CurrDesc.initialLayout = VRSLayout;
+		CurrDesc.finalLayout = VRSLayout;
 
 		FragmentDensityReference.attachment = NumAttachmentDescriptions;
-		FragmentDensityReference.layout = FragmentDensityLayout;
+		FragmentDensityReference.layout = VRSLayout;
 
 		FullHashInfo.LoadOps[MaxSimultaneousRenderTargets + 2] = CurrDesc.stencilLoadOp;
 		FullHashInfo.StoreOps[MaxSimultaneousRenderTargets + 2] = CurrDesc.stencilStoreOp;
 #if VULKAN_USE_REAL_RENDERPASS_COMPATIBILITY
-		FullHashInfo.InitialLayout[MaxSimultaneousRenderTargets + 1] = FragmentDensityLayout;
+		FullHashInfo.InitialLayout[MaxSimultaneousRenderTargets + 1] = VRSLayout;
 #endif
 		CompatibleHashInfo.Formats[MaxSimultaneousRenderTargets + 1] = CurrDesc.format;
 
@@ -1352,15 +1374,17 @@ FVulkanRenderTargetLayout::FVulkanRenderTargetLayout(FVulkanDevice& InDevice, co
 		Extent.Extent3D.depth = 1;
 	}
 
-	if (InDevice.GetOptionalExtensions().HasEXTFragmentDensityMap && RPInfo.ShadingRateTexture)
+	if (GRHISupportsAttachmentVariableRateShading && GRHIVariableRateShadingEnabled && GRHIAttachmentVariableRateShadingEnabled && RPInfo.ShadingRateTexture)
 	{
 		FVulkanTextureBase* Texture = FVulkanTextureBase::Cast(RPInfo.ShadingRateTexture);
 		check(Texture);
+		check(Texture->GetRHITexture());
+		check(Texture->GetRHITexture()->GetFormat() == GRHIVariableRateShadingImageFormat);
 
 		VkAttachmentDescription& CurrDesc = Desc[NumAttachmentDescriptions];
 		FMemory::Memzero(CurrDesc);
 
-		const VkImageLayout FragmentDensityLayout = VK_IMAGE_LAYOUT_FRAGMENT_DENSITY_MAP_OPTIMAL_EXT;
+		const VkImageLayout VRSLayout = GetVRSImageLayout();
 
 		CurrDesc.flags = 0;
 		CurrDesc.format = UEToVkTextureFormat(RPInfo.ShadingRateTexture->GetFormat(), false);
@@ -1369,16 +1393,16 @@ FVulkanRenderTargetLayout::FVulkanRenderTargetLayout(FVulkanDevice& InDevice, co
 		CurrDesc.storeOp = VK_ATTACHMENT_STORE_OP_DONT_CARE;
 		CurrDesc.stencilLoadOp = VK_ATTACHMENT_LOAD_OP_DONT_CARE;
 		CurrDesc.stencilStoreOp = VK_ATTACHMENT_STORE_OP_DONT_CARE;
-		CurrDesc.initialLayout = FragmentDensityLayout;
-		CurrDesc.finalLayout = FragmentDensityLayout;
+		CurrDesc.initialLayout = VRSLayout;
+		CurrDesc.finalLayout = VRSLayout;
 
 		FragmentDensityReference.attachment = NumAttachmentDescriptions;
-		FragmentDensityReference.layout = FragmentDensityLayout;
+		FragmentDensityReference.layout = VRSLayout;
 
 		FullHashInfo.LoadOps[MaxSimultaneousRenderTargets + 2] = CurrDesc.stencilLoadOp;
 		FullHashInfo.StoreOps[MaxSimultaneousRenderTargets + 2] = CurrDesc.stencilStoreOp;
 #if VULKAN_USE_REAL_RENDERPASS_COMPATIBILITY
-		FullHashInfo.InitialLayout[MaxSimultaneousRenderTargets + 1] = FragmentDensityLayout;
+		FullHashInfo.InitialLayout[MaxSimultaneousRenderTargets + 1] = VRSLayout;
 #endif
 		CompatibleHashInfo.Formats[MaxSimultaneousRenderTargets + 1] = CurrDesc.format;
 
@@ -1525,25 +1549,27 @@ FVulkanRenderTargetLayout::FVulkanRenderTargetLayout(const FGraphicsPipelineStat
 		VkAttachmentDescription& CurrDesc = Desc[NumAttachmentDescriptions];
 		FMemory::Memzero(CurrDesc);
 
-		const VkImageLayout FragmentDensityLayout = VK_IMAGE_LAYOUT_FRAGMENT_DENSITY_MAP_OPTIMAL_EXT;
+		const VkImageLayout VRSLayout = GetVRSImageLayout();
+
+		check(GRHIVariableRateShadingImageFormat != PF_Unknown);
 
 		CurrDesc.flags = 0;
-		CurrDesc.format = VK_FORMAT_R8G8_UNORM;
+		CurrDesc.format = UEToVkTextureFormat(GRHIVariableRateShadingImageFormat, false);
 		CurrDesc.samples = VK_SAMPLE_COUNT_1_BIT;
 		CurrDesc.loadOp = VK_ATTACHMENT_LOAD_OP_DONT_CARE;
 		CurrDesc.storeOp = VK_ATTACHMENT_STORE_OP_DONT_CARE;
 		CurrDesc.stencilLoadOp = VK_ATTACHMENT_LOAD_OP_DONT_CARE;
 		CurrDesc.stencilStoreOp = VK_ATTACHMENT_STORE_OP_DONT_CARE;
-		CurrDesc.initialLayout = FragmentDensityLayout;
-		CurrDesc.finalLayout = FragmentDensityLayout;
+		CurrDesc.initialLayout = VRSLayout;
+		CurrDesc.finalLayout = VRSLayout;
 
 		FragmentDensityReference.attachment = NumAttachmentDescriptions;
-		FragmentDensityReference.layout = FragmentDensityLayout;
+		FragmentDensityReference.layout = VRSLayout;
 
 		FullHashInfo.LoadOps[MaxSimultaneousRenderTargets + 2] = CurrDesc.stencilLoadOp;
 		FullHashInfo.StoreOps[MaxSimultaneousRenderTargets + 2] = CurrDesc.stencilStoreOp;
 #if VULKAN_USE_REAL_RENDERPASS_COMPATIBILITY
-		FullHashInfo.InitialLayout[MaxSimultaneousRenderTargets + 1] = FragmentDensityLayout;
+		FullHashInfo.InitialLayout[MaxSimultaneousRenderTargets + 1] = VRSLayout;
 #endif
 		CompatibleHashInfo.Formats[MaxSimultaneousRenderTargets + 1] = CurrDesc.format;
 
diff --git a/Engine/Source/Runtime/VulkanRHI/Private/VulkanRenderpass.cpp b/Engine/Source/Runtime/VulkanRHI/Private/VulkanRenderpass.cpp
--- a/Engine/Source/Runtime/VulkanRHI/Private/VulkanRenderpass.cpp
+++ b/Engine/Source/Runtime/VulkanRHI/Private/VulkanRenderpass.cpp
@@ -28,6 +28,17 @@ public:
 
 		const bool bDeferredShadingSubpass = RTLayout.GetSubpassHint() == ESubpassHint::DeferredShadingSubpass;
 		const bool bDepthReadSubpass = RTLayout.GetSubpassHint() == ESubpassHint::DepthReadSubpass;
+		const bool bApplyFragmentShadingRate =  GRHISupportsAttachmentVariableRateShading && GRHIVariableRateShadingEnabled && GRHIAttachmentVariableRateShadingEnabled && RTLayout.GetFragmentDensityAttachmentReference() != nullptr;
+
+#if VULKAN_SUPPORTS_RENDERPASS2
+		FVulkanAttachmentReference<VkAttachmentReference2> ShadingRateAttachmentReference;
+		FVulkanFragmentShadingRateAttachmentInfo FragmentShadingRateAttachmentInfo;
+		if (bApplyFragmentShadingRate)
+		{
+			ShadingRateAttachmentReference.SetAttachment(*RTLayout.GetFragmentDensityAttachmentReference(), VkImageAspectFlagBits::VK_IMAGE_ASPECT_COLOR_BIT);
+			FragmentShadingRateAttachmentInfo.SetReference(&ShadingRateAttachmentReference);
+		}
+#endif
 
 		// Grab (and optionally convert) attachment references.
 		for (uint32 ColorAttachment = 0; ColorAttachment < RTLayout.GetNumColorAttachments(); ++ColorAttachment)
@@ -54,6 +65,13 @@ public:
 			{
 				SubpassDesc.SetDepthStencilAttachment(&DepthStencilAttachmentReference);
 			}
+
+#if VULKAN_SUPPORTS_RENDERPASS2
+			if (bApplyFragmentShadingRate)
+			{
+				SubpassDesc.SetShadingRateAttachment(&FragmentShadingRateAttachmentInfo);
+			}
+#endif
 		}
 
 		// Color write and depth read sub-pass
@@ -78,6 +96,13 @@ public:
 			// depth attachment is same as input attachment
 			SubpassDesc.SetDepthStencilAttachment(InputAttachments1);
 
+#if VULKAN_SUPPORTS_RENDERPASS2
+			if (bApplyFragmentShadingRate)
+			{
+				SubpassDesc.SetShadingRateAttachment(&FragmentShadingRateAttachmentInfo);
+			}
+#endif
+			
 			TSubpassDependencyClass& SubpassDep = SubpassDependencies[NumDependencies++];
 			SubpassDep.srcSubpass = 0;
 			SubpassDep.dstSubpass = 1;
@@ -109,8 +134,14 @@ public:
 				SubpassDesc.SetDepthStencilAttachment(&DepthStencilAttachment);
 				SubpassDesc.SetInputAttachments(&DepthStencilAttachment, 1);
 
+#if VULKAN_SUPPORTS_RENDERPASS2
+				if (bApplyFragmentShadingRate)
+				{
+					SubpassDesc.SetShadingRateAttachment(&FragmentShadingRateAttachmentInfo);
+				}
+#endif
+				
 				// depth as Input0
-
 				TSubpassDependencyClass& SubpassDep = SubpassDependencies[NumDependencies++];
 				SubpassDep.srcSubpass = 0;
 				SubpassDep.dstSubpass = 1;
@@ -143,6 +174,12 @@ public:
 				}
 
 				SubpassDesc.SetInputAttachments(InputAttachments2, 5);
+#if VULKAN_SUPPORTS_RENDERPASS2
+				if (bApplyFragmentShadingRate)
+				{
+					SubpassDesc.SetShadingRateAttachment(&FragmentShadingRateAttachmentInfo);
+				}
+#endif
 
 				TSubpassDependencyClass& SubpassDep = SubpassDependencies[NumDependencies++];
 				SubpassDep.srcSubpass = 1;
diff --git a/Engine/Source/Runtime/VulkanRHI/Private/VulkanRenderpass.h b/Engine/Source/Runtime/VulkanRHI/Private/VulkanRenderpass.h
--- a/Engine/Source/Runtime/VulkanRHI/Private/VulkanRenderpass.h
+++ b/Engine/Source/Runtime/VulkanRHI/Private/VulkanRenderpass.h
@@ -127,6 +127,11 @@ struct FVulkanSubpassDescription<VkSubpassDescription>
 		pInputAttachments = static_cast<VkAttachmentReference*>(InputAttachmentReferences);
 		inputAttachmentCount = NumInputAttachmentReferences;
 	}
+
+	void SetShadingRateAttachment(void* /* ShadingRateAttachmentInfo */)
+	{
+		// No-op without VK_KHR_create_renderpass2
+	}
 };
 
 #if VULKAN_SUPPORTS_RENDERPASS2
@@ -166,6 +171,11 @@ struct FVulkanSubpassDescription<VkSubpassDescription2>
 		pInputAttachments = static_cast<VkAttachmentReference2*>(InputAttachmentReferences);
 		inputAttachmentCount = NumInputAttachmentReferences;
 	}
+
+	void SetShadingRateAttachment(void* ShadingRateAttachmentInfo)
+	{
+		pNext = ShadingRateAttachmentInfo;
+	}
 };
 #endif
 
@@ -291,6 +301,23 @@ struct FVulkanRenderPassCreateInfo<VkRenderPassCreateInfo2>
 		return Handle;
 	}
 };
+
+struct FVulkanFragmentShadingRateAttachmentInfo
+	: public VkFragmentShadingRateAttachmentInfoKHR
+{
+	FVulkanFragmentShadingRateAttachmentInfo()
+	{
+		ZeroVulkanStruct(*this, VK_STRUCTURE_TYPE_FRAGMENT_SHADING_RATE_ATTACHMENT_INFO_KHR);
+		// For now, just use the smallest tile-size available. TODO: Add a setting to allow prioritizing either higher resolution/larger shading rate attachment targets 
+		// or lower-resolution/smaller attachments.
+		shadingRateAttachmentTexelSize = { (uint32)GRHIVariableRateShadingImageTileMinWidth, (uint32)GRHIVariableRateShadingImageTileMinHeight };
+	}
+
+	void SetReference(FVulkanAttachmentReference<VkAttachmentReference2>* AttachmentReference)
+	{
+		pFragmentShadingRateAttachment = AttachmentReference;
+	}
+};
 #endif
 
 VkRenderPass CreateVulkanRenderPass(FVulkanDevice& Device, const FVulkanRenderTargetLayout& RTLayout);
diff --git a/Engine/Source/Runtime/VulkanRHI/Private/VulkanTexture.cpp b/Engine/Source/Runtime/VulkanRHI/Private/VulkanTexture.cpp
--- a/Engine/Source/Runtime/VulkanRHI/Private/VulkanTexture.cpp
+++ b/Engine/Source/Runtime/VulkanRHI/Private/VulkanTexture.cpp
@@ -341,6 +341,23 @@ void FVulkanSurface::GenerateImageCreateInfo(
 		ImageCreateInfo.tiling = VK_IMAGE_TILING_OPTIMAL;
 	}
 
+	if (EnumHasAnyFlags(UEFlags, TexCreate_Foveation) && ValidateShadingRateDataType())
+	{
+#if VULKAN_SUPPORTS_FRAGMENT_SHADING_RATE
+		if (GRHIVariableRateShadingImageDataType == VRSImage_Palette)
+		{
+			ImageCreateInfo.usage |= VK_IMAGE_USAGE_FRAGMENT_SHADING_RATE_ATTACHMENT_BIT_KHR;
+		}
+#endif
+
+#if VULKAN_SUPPORTS_FRAGMENT_DENSITY_MAP
+		if (GRHIVariableRateShadingImageDataType == VRSImage_Fractional)
+		{
+			ImageCreateInfo.usage |= VK_IMAGE_USAGE_FRAGMENT_DENSITY_MAP_BIT_EXT;
+		}
+#endif
+	}
+	
 	if (EnumHasAnyFlags(UEFlags, TexCreate_UAV))
 	{
 		//cannot have the storage bit on a memoryless texture
@@ -905,7 +922,7 @@ FVulkanSurface::FVulkanSurface(FVulkanDevice& InDevice, VkImageViewType Resource
 #endif
 		VULKAN_SET_DEBUG_NAME(InDevice, VK_OBJECT_TYPE_IMAGE, Image, TEXT("(FVulkanSurface*)0x%p"), this);
 
-		VkImageLayout InitialLayout;
+		VkImageLayout InitialLayout = VK_IMAGE_LAYOUT_UNDEFINED;
 		bool bOnlyAddToLayoutManager, bDoInitialClear;
 		if (EnumHasAnyFlags(UEFlags, TexCreate_RenderTargetable | TexCreate_DepthStencilTargetable))
 		{
@@ -913,10 +930,23 @@ FVulkanSurface::FVulkanSurface(FVulkanDevice& InDevice, VkImageViewType Resource
 			bOnlyAddToLayoutManager = false;
 			bDoInitialClear = true;
 		}
-		else if (EnumHasAnyFlags(UEFlags, TexCreate_Foveation))
+		else if (EnumHasAnyFlags(UEFlags, TexCreate_Foveation) && ValidateShadingRateDataType())
 		{
-			// If it's a foveation texture, do not clear but add to layoutmgr, and set correct foveation layout. 
-			InitialLayout = VK_IMAGE_LAYOUT_FRAGMENT_DENSITY_MAP_OPTIMAL_EXT;
+			// If it's a shading rate attachment texture, do not clear but add to layoutmgr, and set correct shading rate layout.
+#if VULKAN_SUPPORTS_FRAGMENT_SHADING_RATE
+			if (GRHIVariableRateShadingImageDataType == VRSImage_Palette)
+			{
+				InitialLayout = VK_IMAGE_LAYOUT_FRAGMENT_SHADING_RATE_ATTACHMENT_OPTIMAL_KHR;
+			}
+#endif
+
+#if VULKAN_SUPPORTS_FRAGMENT_DENSITY_MAP
+			if (GRHIVariableRateShadingImageDataType == VRSImage_Fractional)
+			{
+				InitialLayout = VK_IMAGE_LAYOUT_FRAGMENT_DENSITY_MAP_OPTIMAL_EXT;
+			}
+#endif
+
 			bOnlyAddToLayoutManager = true;
 			bDoInitialClear = false;
 		}
diff --git a/Engine/Source/Runtime/VulkanRHI/Private/VulkanViewport.cpp b/Engine/Source/Runtime/VulkanRHI/Private/VulkanViewport.cpp
--- a/Engine/Source/Runtime/VulkanRHI/Private/VulkanViewport.cpp
+++ b/Engine/Source/Runtime/VulkanRHI/Private/VulkanViewport.cpp
@@ -434,7 +434,7 @@ FVulkanFramebuffer::FVulkanFramebuffer(FVulkanDevice& Device, const FRHISetRende
 		}
 	}
 
-	if (RTLayout.GetHasFragmentDensityAttachment() && Device.GetOptionalExtensions().HasEXTFragmentDensityMap)
+	if (GRHISupportsAttachmentVariableRateShading && GRHIVariableRateShadingEnabled && GRHIAttachmentVariableRateShadingEnabled && RTLayout.GetHasFragmentDensityAttachment())
 	{
 		FVulkanTextureBase* Texture = FVulkanTextureBase::Cast(InRTInfo.ShadingRateTexture);
 		FragmentDensityImage = Texture->Surface.Image;
