Fix FVulkanUnorderedAccessView for FVulkanTextureCube giving the validation error "requires an image view of type VK_IMAGE_VIEW_TYPE_2D_ARRAY but got VK_IMAGE_VIEW_TYPE_CUBE".
Also fix TextureCube arrays init (broken in 17019036).

diff --git a/Engine/Source/Runtime/VulkanRHI/Private/VulkanTexture.cpp b/Engine/Source/Runtime/VulkanRHI/Private/VulkanTexture.cpp
--- a/Engine/Source/Runtime/VulkanRHI/Private/VulkanTexture.cpp
+++ b/Engine/Source/Runtime/VulkanRHI/Private/VulkanTexture.cpp
@@ -2069,7 +2069,7 @@ FVulkanTextureBase::FVulkanTextureBase(FVulkanDevice& Device, VkImageViewType Re
 }
 
 FVulkanTextureBase::FVulkanTextureBase(FTextureRHIRef& SrcTextureRHI, const FVulkanTextureBase* SrcTexture, VkImageViewType ResourceType, uint32 SizeX, uint32 SizeY, uint32 SizeZ)
-	: Surface(*SrcTexture->Surface.Device, ResourceType, SrcTexture->Surface.PixelFormat, SizeX, SizeY, SizeZ, SrcTexture->Surface.GetNumberOfArrayLevels(), SrcTexture->Surface.NumMips, SrcTexture->Surface.NumSamples, SrcTexture->Surface.Image, SrcTexture->Surface.UEFlags, FRHIResourceCreateInfo(TEXT("FVulkanTextureBase")))
+	: Surface(*SrcTexture->Surface.Device, ResourceType, SrcTexture->Surface.PixelFormat, SizeX, SizeY, SizeZ, SrcTexture->Surface.ArraySize, SrcTexture->Surface.NumMips, SrcTexture->Surface.NumSamples, SrcTexture->Surface.Image, SrcTexture->Surface.UEFlags, FRHIResourceCreateInfo(TEXT("FVulkanTextureBase")))
 	, PartialView(nullptr)
 	, AliasedTexture(SrcTextureRHI)
 {
diff --git a/Engine/Source/Runtime/VulkanRHI/Private/VulkanUAV.cpp b/Engine/Source/Runtime/VulkanRHI/Private/VulkanUAV.cpp
--- a/Engine/Source/Runtime/VulkanRHI/Private/VulkanUAV.cpp
+++ b/Engine/Source/Runtime/VulkanRHI/Private/VulkanUAV.cpp
@@ -288,7 +288,8 @@ void FVulkanUnorderedAccessView::UpdateView()
 		else if (FRHITextureCube* TexCube = SourceTexture->GetTextureCube())
 		{
 			FVulkanTextureCube* VTexCube = ResourceCast(TexCube);
-			TextureView.Create(*Device, VTexCube->Surface.Image, VK_IMAGE_VIEW_TYPE_CUBE, VTexCube->Surface.GetPartialAspectMask(), Format, UEToVkTextureFormat(Format, false), MipLevel, 1, 0, 1, true);
+			// RWTextureCube is defined as RWTexture2DArray in shader source, avoid validation errors by creating the appropriate VK_IMAGE_VIEW_TYPE_2D_ARRAY view (instead of VK_IMAGE_VIEW_TYPE_CUBE)
+			TextureView.Create(*Device, VTexCube->Surface.Image, VK_IMAGE_VIEW_TYPE_2D_ARRAY, VTexCube->Surface.GetPartialAspectMask(), Format, UEToVkTextureFormat(Format, false), MipLevel, 1, 0, VTexCube->Surface.GetNumberOfArrayLevels(), true);
 		}
 		else if (FRHITexture3D* Tex3D = SourceTexture->GetTexture3D())
 		{
