PR #8353: add null pointer check for allocations that use a non-zero offset (Contributed by [REDACTED_PERSON])

Modified for coding standards.

diff --git a/Engine/Source/Runtime/VulkanRHI/Private/VulkanBuffer.cpp b/Engine/Source/Runtime/VulkanRHI/Private/VulkanBuffer.cpp
--- a/Engine/Source/Runtime/VulkanRHI/Private/VulkanBuffer.cpp
+++ b/Engine/Source/Runtime/VulkanRHI/Private/VulkanBuffer.cpp
@@ -211,6 +211,7 @@ FVulkanResourceMultiBuffer::~FVulkanResourceMultiBuffer()
 void* FVulkanResourceMultiBuffer::Lock(bool bFromRenderingThread, EResourceLockMode LockMode, uint32 LockSize, uint32 Offset)
 {
 	void* Data = nullptr;
+	uint32 DataOffset = 0;
 
 	const bool bStatic = EnumHasAnyFlags(GetUsage(), BUF_Static);
 	const bool bDynamic = EnumHasAnyFlags(GetUsage(), BUF_Dynamic);
@@ -249,8 +250,9 @@ void* FVulkanResourceMultiBuffer::Lock(bool bFromRenderingThread, EResourceLockM
 			const bool bUnifiedMem = Device->HasUnifiedMemory();
 			if (bUnifiedMem)
 			{
-				Data = (uint8*)Buffers[DynamicBufferIndex].GetMappedPointer(Device) + Offset;
+				Data = Buffers[DynamicBufferIndex].GetMappedPointer(Device);
 				LockStatus = ELockStatus::PersistentMapping;
+				DataOffset = Offset;
 			}
 			else 
 			{
@@ -333,14 +335,15 @@ void* FVulkanResourceMultiBuffer::Lock(bool bFromRenderingThread, EResourceLockM
 			}
 			else
 			{
-				Data = (uint8*)Buffers[DynamicBufferIndex].GetMappedPointer(Device) + Offset;
+				Data = Buffers[DynamicBufferIndex].GetMappedPointer(Device);
+				DataOffset = Offset;
 				LockStatus = ELockStatus::PersistentMapping;
 			}
 		}
 	}
 
 	check(Data);
-	return Data;
+	return (uint8*)Data + DataOffset;
 }
 
 inline void FVulkanResourceMultiBuffer::InternalUnlock(FVulkanCommandListContext& Context, VulkanRHI::FPendingBufferLock& PendingLock, FVulkanResourceMultiBuffer* MultiBuffer, int32 InDynamicBufferIndex)
