VulkanRHI: Tweak error messages to provide more information for DescriptorType mismatches

diff --git a/Engine/Source/Runtime/VulkanRHI/Private/VulkanDescriptorSets.h b/Engine/Source/Runtime/VulkanRHI/Private/VulkanDescriptorSets.h
--- a/Engine/Source/Runtime/VulkanRHI/Private/VulkanDescriptorSets.h
+++ b/Engine/Source/Runtime/VulkanRHI/Private/VulkanDescriptorSets.h
@@ -1090,11 +1090,15 @@ protected:
 		SetWritten(DescriptorIndex);		
 		if (DescriptorType == VK_DESCRIPTOR_TYPE_STORAGE_BUFFER)
 		{
-			check(WriteDescriptors[DescriptorIndex].descriptorType == VK_DESCRIPTOR_TYPE_STORAGE_BUFFER || WriteDescriptors[DescriptorIndex].descriptorType == VK_DESCRIPTOR_TYPE_STORAGE_BUFFER_DYNAMIC);
+			checkf(WriteDescriptors[DescriptorIndex].descriptorType == VK_DESCRIPTOR_TYPE_STORAGE_BUFFER || WriteDescriptors[DescriptorIndex].descriptorType == VK_DESCRIPTOR_TYPE_STORAGE_BUFFER_DYNAMIC,
+				TEXT("DescriptorType mismatch at index %d: called WriteBuffer<%d> and was expecting %d."), 
+				DescriptorIndex, (uint32)DescriptorType, (uint32)WriteDescriptors[DescriptorIndex].descriptorType);
 		}
 		else
 		{
-			check(WriteDescriptors[DescriptorIndex].descriptorType == DescriptorType);
+			checkf(WriteDescriptors[DescriptorIndex].descriptorType == DescriptorType,
+				TEXT("DescriptorType mismatch at index %d: called WriteBuffer<%d> and was expecting %d."),
+				DescriptorIndex, (uint32)DescriptorType, (uint32)WriteDescriptors[DescriptorIndex].descriptorType);
 		}
 		VkDescriptorBufferInfo* BufferInfo = const_cast<VkDescriptorBufferInfo*>(WriteDescriptors[DescriptorIndex].pBufferInfo);
 		check(BufferInfo);
@@ -1147,7 +1151,9 @@ protected:
 		SetWritten(DescriptorIndex);
 		if (DescriptorType == VK_DESCRIPTOR_TYPE_SAMPLED_IMAGE)
 		{
-			check(WriteDescriptors[DescriptorIndex].descriptorType == VK_DESCRIPTOR_TYPE_SAMPLED_IMAGE || WriteDescriptors[DescriptorIndex].descriptorType == VK_DESCRIPTOR_TYPE_COMBINED_IMAGE_SAMPLER);
+			checkf(WriteDescriptors[DescriptorIndex].descriptorType == VK_DESCRIPTOR_TYPE_SAMPLED_IMAGE || WriteDescriptors[DescriptorIndex].descriptorType == VK_DESCRIPTOR_TYPE_COMBINED_IMAGE_SAMPLER,
+				TEXT("DescriptorType mismatch at index %d: called WriteTextureView<%d> and was expecting %d."),
+				DescriptorIndex, (uint32)DescriptorType, (uint32)WriteDescriptors[DescriptorIndex].descriptorType);
 			ensureMsgf(Layout == VK_IMAGE_LAYOUT_SHARED_PRESENT_KHR ||
 				  Layout == VK_IMAGE_LAYOUT_DEPTH_READ_ONLY_STENCIL_ATTACHMENT_OPTIMAL ||
 				  Layout == VK_IMAGE_LAYOUT_DEPTH_ATTACHMENT_STENCIL_READ_ONLY_OPTIMAL ||
@@ -1157,7 +1163,9 @@ protected:
 		}
 		else
 		{
-			check(WriteDescriptors[DescriptorIndex].descriptorType == DescriptorType);
+			checkf(WriteDescriptors[DescriptorIndex].descriptorType == DescriptorType,
+				TEXT("DescriptorType mismatch at index %d: called WriteTextureView<%d> and was expecting %d."),
+				DescriptorIndex, (uint32)DescriptorType, (uint32)WriteDescriptors[DescriptorIndex].descriptorType);
 		}
 		VkDescriptorImageInfo* ImageInfo = const_cast<VkDescriptorImageInfo*>(WriteDescriptors[DescriptorIndex].pImageInfo);
 		check(ImageInfo);
@@ -1194,7 +1202,9 @@ protected:
 	bool WriteBufferView(uint32 DescriptorIndex, const FVulkanBufferView* View)
 	{
 		check(DescriptorIndex < NumWrites);
-		check(WriteDescriptors[DescriptorIndex].descriptorType == DescriptorType);
+		checkf(WriteDescriptors[DescriptorIndex].descriptorType == DescriptorType, 
+			TEXT("DescriptorType mismatch at index %d: called WriteBufferView<%d> and was expecting %d."), 
+			DescriptorIndex, (uint32)DescriptorType, (uint32)WriteDescriptors[DescriptorIndex].descriptorType);
 		SetWritten(DescriptorIndex);
 		WriteDescriptors[DescriptorIndex].pTexelBufferView = &View->View;
 		BufferViewReferences[DescriptorIndex] = View;
