VulkanRHI: Running out of bits in AllocatorIndex when running AncientGame from scratch in Editor (PIE).

diff --git a/Engine/Source/Runtime/VulkanRHI/Private/VulkanMemory.cpp b/Engine/Source/Runtime/VulkanRHI/Private/VulkanMemory.cpp
--- a/Engine/Source/Runtime/VulkanRHI/Private/VulkanMemory.cpp
+++ b/Engine/Source/Runtime/VulkanRHI/Private/VulkanMemory.cpp
@@ -2512,7 +2512,7 @@ namespace VulkanRHI
 
 	void FVulkanSubresourceAllocator::SetFreePending(FVulkanAllocation& Allocation)
 	{
-		check(Allocation.Type == Type);
+		check(Allocation.GetType() == Type);
 		check(Allocation.AllocatorIndex == GetAllocatorIndex());
 		{
 			FScopeLock ScopeLock(&CS);
@@ -2961,6 +2961,9 @@ namespace VulkanRHI
 		{
 			UE_LOG(LogVulkanRHI, Fatal, TEXT("Missing memory type index %d, MemSize %d, MemPropTypeBits %u, MemPropertyFlags %u, %s(%d)"), TypeIndex, (uint32)MemoryReqs.size, (uint32)MemoryReqs.memoryTypeBits, (uint32)MemoryPropertyFlags, ANSI_TO_TCHAR(File), Line);
 		}
+
+		check(MemoryReqs.size <= (uint64)MAX_uint32);
+
 		const bool bForceSeparateAllocation = (MemoryPropertyFlags & VK_MEMORY_PROPERTY_LAZILY_ALLOCATED_BIT) == VK_MEMORY_PROPERTY_LAZILY_ALLOCATED_BIT;
 		if(!ResourceTypeHeaps[TypeIndex]->AllocateResource(OutAllocation, AllocationOwner, EType::Image, MemoryReqs.size, MemoryReqs.alignment, bMapped, bForceSeparateAllocation, MetaType, File, Line))
 		{
@@ -3015,6 +3018,8 @@ namespace VulkanRHI
 			}
 		}
 
+		check(MemoryReqs.size <= (uint64)MAX_uint32);
+
 		if(!ResourceTypeHeaps[TypeIndex]->AllocateResource(OutAllocation, AllocationOwner, EType::Buffer, MemoryReqs.size, MemoryReqs.alignment, bMapped, false, MetaType, File, Line))
 		{
 			// Try another memory type if the allocation failed
@@ -3057,6 +3062,7 @@ namespace VulkanRHI
 			uint32 TypeIndex = 0;
 			VERIFYVULKANRESULT(DeviceMemoryManager->GetMemoryTypeFromProperties(MemoryReqs.memoryTypeBits, MemoryPropertyFlags, &TypeIndex));
 			ensure((MemoryPropertyFlags & VK_MEMORY_PROPERTY_HOST_VISIBLE_BIT) != VK_MEMORY_PROPERTY_HOST_VISIBLE_BIT);
+			check(MemoryReqs.size <= (uint64)MAX_uint32);
 			if (!ResourceTypeHeaps[TypeIndex])
 			{
 				UE_LOG(LogVulkanRHI, Fatal, TEXT("Missing memory type index %d, MemSize %d, MemPropTypeBits %u, MemPropertyFlags %u, %s(%d)"), TypeIndex, (uint32)MemoryReqs.size, (uint32)MemoryReqs.memoryTypeBits, (uint32)MemoryPropertyFlags, ANSI_TO_TCHAR(File), Line);
@@ -3712,7 +3718,9 @@ namespace VulkanRHI
 	}
 
 	FVulkanAllocation::FVulkanAllocation()
+		: bHasOwnership(0)
 	{
+		SetType(EVulkanAllocationEmpty);
 	}
 	FVulkanAllocation::~FVulkanAllocation()
 	{
@@ -3724,7 +3732,7 @@ namespace VulkanRHI
 	{
 		check(State == EUNUSED);
 		State = EALLOCATED;
-		Type = Alloc.Type;
+		Type = Alloc.GetType();
 		MetaType = Alloc.MetaType;
 
 		Size = Alloc.Size;
@@ -3738,12 +3746,11 @@ namespace VulkanRHI
 	{
 		check(!HasAllocation());
 		bHasOwnership = 1;
-		Type = InType;
+		SetType(InType);
 		MetaType = InMetaType;
 		Size = InSize;
 		Offset = InAlignedOffset;
-		check(InAllocatorIndex < (1<<ALLOCATOR_INDEX_BITS));
-		check(InAllocationIndex < (1<<ALLOCATION_INDEX_BITS));
+		check(InAllocatorIndex < (uint32)MAX_uint16);
 		AllocatorIndex = InAllocatorIndex;
 		AllocationIndex = InAllocationIndex;
 		VulkanHandle = Handle;
diff --git a/Engine/Source/Runtime/VulkanRHI/Private/VulkanRHI.cpp b/Engine/Source/Runtime/VulkanRHI/Private/VulkanRHI.cpp
--- a/Engine/Source/Runtime/VulkanRHI/Private/VulkanRHI.cpp
+++ b/Engine/Source/Runtime/VulkanRHI/Private/VulkanRHI.cpp
@@ -1558,7 +1558,7 @@ FVulkanRingBuffer::FVulkanRingBuffer(FVulkanDevice* InDevice, uint64 TotalSize,
 	, BufferOffset(0)
 	, MinAlignment(0)
 {
-
+	check(TotalSize <= (uint64)MAX_uint32);
 	InDevice->GetMemoryManager().AllocateBufferPooled(Allocation, this, TotalSize, Usage, MemPropertyFlags, EVulkanAllocationMetaRingBuffer, __FILE__, __LINE__);
 	MinAlignment = Allocation.GetBufferAlignment(Device);
 	// Start by wrapping around to set up the correct fence
diff --git a/Engine/Source/Runtime/VulkanRHI/Public/VulkanMemory.h b/Engine/Source/Runtime/VulkanRHI/Public/VulkanMemory.h
--- a/Engine/Source/Runtime/VulkanRHI/Public/VulkanMemory.h
+++ b/Engine/Source/Runtime/VulkanRHI/Public/VulkanMemory.h
@@ -131,6 +131,7 @@ namespace VulkanRHI
 		EVulkanAllocationBuffer,
 		EVulkanAllocationImage,
 		EVulkanAllocationImageDedicated,
+		EVulkanAllocationSize, // keep last
 	};
 
 	enum EVulkanAllocationMetaType : uint8
@@ -145,7 +146,7 @@ namespace VulkanRHI
 		EVulkanAllocationMetaBufferUAV,
 		EVulkanAllocationMetaBufferStaging,
 		EVulkanAllocationMetaBufferOther,
-		EVulkanAllocationMetaSize,
+		EVulkanAllocationMetaSize, // keep last
 	};
 	enum EVulkanAllocationFlags
 	{
@@ -487,11 +488,6 @@ namespace VulkanRHI
 	class FVulkanAllocation
 	{
 	public:
-		enum
-		{
-			ALLOCATOR_INDEX_BITS = 13,
-			ALLOCATION_INDEX_BITS = 32-ALLOCATOR_INDEX_BITS,
-		};
 		FVulkanAllocation();
 		~FVulkanAllocation();
 		void Init(EVulkanAllocationType Type, EVulkanAllocationMetaType MetaType, uint64 Handle, uint32 InSize, uint32 AlignedOffset, uint32 AllocatorIndex, uint32 AllocationIndex, uint32 BufferId);
@@ -504,17 +500,25 @@ namespace VulkanRHI
 		void Own();
 		bool IsValid() const;
 		
-		
-		EVulkanAllocationType Type = EVulkanAllocationEmpty;
-		EVulkanAllocationMetaType MetaType = EVulkanAllocationMetaUnknown;
-		uint8 bHasOwnership = 0;
-
+		uint64 VulkanHandle = 0;
+		uint32 HandleId = 0;
 		uint32 Size = 0;
 		uint32 Offset = 0;
-		uint32 AllocatorIndex : ALLOCATOR_INDEX_BITS;
-		uint32 AllocationIndex : ALLOCATION_INDEX_BITS;
-		int32 HandleId = 0;
-		uint64 VulkanHandle = 0;
+		uint32 AllocationIndex = 0;
+		uint16 AllocatorIndex = 0;
+		EVulkanAllocationMetaType MetaType = EVulkanAllocationMetaUnknown;
+		uint8 Type : 7;
+		uint8 bHasOwnership : 1;
+
+		static_assert(EVulkanAllocationSize < 128, "Not enough bits to hold EVulkanAllocationType");
+		inline EVulkanAllocationType GetType() const
+		{
+			return (EVulkanAllocationType)Type;
+		}
+		inline void SetType(EVulkanAllocationType InType)
+		{
+			Type = (uint8)InType;
+		}
 
 		//helper functions
 		void* GetMappedPointer(FVulkanDevice* Device);
