Fix leak in memory tracking when FStagingManager::AcquireBuffer uses buffer pulled from UsedStagingBuffers.  VULKAN_FILL_TRACK_INFO was called again without freeing previous tracking info.

diff --git a/Engine/Source/Runtime/VulkanRHI/Private/VulkanMemory.cpp b/Engine/Source/Runtime/VulkanRHI/Private/VulkanMemory.cpp
--- a/Engine/Source/Runtime/VulkanRHI/Private/VulkanMemory.cpp
+++ b/Engine/Source/Runtime/VulkanRHI/Private/VulkanMemory.cpp
@@ -367,7 +367,7 @@ void VulkanTrackFillInfo(FVulkanTrackInfo& Track, const char* File, uint32 Line)
 {
 	if (GVulkanMemoryBackTrace > 0)
 	{
-		uint64* Stack = new uint64[GVulkanMemoryBackTrace];
+		uint64* Stack = ((Track.Data != nullptr) && (Track.SizeOrLine < 0)) ? (uint64*)Track.Data : new uint64[GVulkanMemoryBackTrace];
 		int32 Depth = FPlatformStackWalk::CaptureStackBackTrace(Stack, GVulkanMemoryBackTrace);
 		Track.SizeOrLine = -Depth;
 		Track.Data = Stack;
