PR #8412: Buffer usage fix in vulkan surface (Contributed by [REDACTED_PERSON])

diff --git a/Engine/Source/Runtime/VulkanRHI/Private/VulkanTexture.cpp b/Engine/Source/Runtime/VulkanRHI/Private/VulkanTexture.cpp
--- a/Engine/Source/Runtime/VulkanRHI/Private/VulkanTexture.cpp
+++ b/Engine/Source/Runtime/VulkanRHI/Private/VulkanTexture.cpp
@@ -590,7 +590,7 @@ FVulkanSurface::FVulkanSurface(FVulkanDevice& InDevice, FVulkanEvictable* Owner,
 		VkBufferCreateInfo BufferCreateInfo;
 		ZeroVulkanStruct(BufferCreateInfo, VK_STRUCTURE_TYPE_BUFFER_CREATE_INFO);
 		BufferCreateInfo.size = Size;
-		BufferCreateInfo.usage = BufferMemFlags;
+		BufferCreateInfo.usage = VK_BUFFER_USAGE_TRANSFER_DST_BIT;
 
 		VERIFYVULKANRESULT(VulkanRHI::vkCreateBuffer(VulkanDevice, &BufferCreateInfo, VULKAN_CPU_ALLOCATOR, &CpuReadbackBuffer->Buffer));
 		VulkanRHI::vkGetBufferMemoryRequirements(VulkanDevice, CpuReadbackBuffer->Buffer, &MemoryRequirements);
