Added additional debugging when pipeline creation fails in Vulkan

diff --git a/Engine/Source/Runtime/VulkanRHI/Private/VulkanPipeline.cpp b/Engine/Source/Runtime/VulkanRHI/Private/VulkanPipeline.cpp
--- a/Engine/Source/Runtime/VulkanRHI/Private/VulkanPipeline.cpp
+++ b/Engine/Source/Runtime/VulkanRHI/Private/VulkanPipeline.cpp
@@ -1288,10 +1288,23 @@ bool FVulkanPipelineStateCacheManager::CreateGfxPipelineFromEntry(FVulkanRHIGrap
 		Result = VulkanRHI::vkCreateGraphicsPipelines(Device->GetInstanceHandle(), PipelineCache, 1, &PipelineInfo, VULKAN_CPU_ALLOCATOR, Pipeline);
 	}
 
-
 	if (Result != VK_SUCCESS)
 	{
-		UE_LOG(LogVulkanRHI, Error, TEXT("Failed to create graphics pipeline."));
+		FString ShaderHashes = "";
+		if (Shaders[ShaderStage::Vertex] && Shaders[ShaderStage::Vertex]->StageFlag == VK_SHADER_STAGE_VERTEX_BIT)
+		{
+			ShaderHashes += TEXT("VS: ") + static_cast<FVulkanVertexShader*>(Shaders[ShaderStage::Vertex])->GetHash().ToString() + TEXT("\n");
+		}
+		if (Shaders[ShaderStage::Pixel] && Shaders[ShaderStage::Pixel]->StageFlag == VK_SHADER_STAGE_FRAGMENT_BIT)
+		{
+			ShaderHashes += TEXT("PS: ") + static_cast<FVulkanPixelShader*>(Shaders[ShaderStage::Pixel])->GetHash().ToString() + TEXT("\n");
+		}
+		if (Shaders[ShaderStage::Geometry] && Shaders[ShaderStage::Geometry]->StageFlag == VK_SHADER_STAGE_GEOMETRY_BIT)
+		{
+			ShaderHashes += TEXT("GS: ") + static_cast<FVulkanGeometryShader*>(Shaders[ShaderStage::Geometry])->GetHash().ToString() + TEXT("\n");
+		}
+		
+		UE_LOG(LogVulkanRHI, Error, TEXT("Failed to create graphics pipeline.\nShaders in pipeline: %s"), *ShaderHashes);
 		return false;
 	}
 
@@ -1935,7 +1948,13 @@ FVulkanComputePipeline* FVulkanPipelineStateCacheManager::CreateComputePipelineF
 	PipelineInfo.stage.pName = EntryPoint;
 	PipelineInfo.layout = ComputeLayout->GetPipelineLayout();
 		
-	VERIFYVULKANRESULT(VulkanRHI::vkCreateComputePipelines(Device->GetInstanceHandle(), PipelineCache, 1, &PipelineInfo, VULKAN_CPU_ALLOCATOR, &Pipeline->Pipeline));
+	VkResult Result = VulkanRHI::vkCreateComputePipelines(Device->GetInstanceHandle(), PipelineCache, 1, &PipelineInfo, VULKAN_CPU_ALLOCATOR, &Pipeline->Pipeline);
+
+	if (Result != VK_SUCCESS)
+	{
+		FString ComputeHash = Shader->GetHash().ToString();
+		UE_LOG(LogVulkanRHI, Warning, TEXT("Failed to create graphics pipeline.\nShaders in pipeline: CS: %s"), *ComputeHash);
+	}
 
 	Pipeline->Layout = ComputeLayout;
 
