Fix for compiler error

diff --git a/Engine/Source/Runtime/VulkanRHI/Private/VulkanPipeline.cpp b/Engine/Source/Runtime/VulkanRHI/Private/VulkanPipeline.cpp
--- a/Engine/Source/Runtime/VulkanRHI/Private/VulkanPipeline.cpp
+++ b/Engine/Source/Runtime/VulkanRHI/Private/VulkanPipeline.cpp
@@ -1299,11 +1299,12 @@ bool FVulkanPipelineStateCacheManager::CreateGfxPipelineFromEntry(FVulkanRHIGrap
 		{
 			ShaderHashes += TEXT("PS: ") + static_cast<FVulkanPixelShader*>(Shaders[ShaderStage::Pixel])->GetHash().ToString() + TEXT("\n");
 		}
+#if VULKAN_SUPPORTS_GEOMETRY_SHADERS
 		if (Shaders[ShaderStage::Geometry] && Shaders[ShaderStage::Geometry]->StageFlag == VK_SHADER_STAGE_GEOMETRY_BIT)
 		{
 			ShaderHashes += TEXT("GS: ") + static_cast<FVulkanGeometryShader*>(Shaders[ShaderStage::Geometry])->GetHash().ToString() + TEXT("\n");
 		}
-		
+#endif
 		UE_LOG(LogVulkanRHI, Error, TEXT("Failed to create graphics pipeline.\nShaders in pipeline: %s"), *ShaderHashes);
 		return false;
 	}
