Prevent collisions (CurrentlyDeleting assert) between an RHIThread running FRHIResource::FlushPendingDeletes and workers creating PSOs.  Queue duplicate/failed PSO for deletion with Release() instead of using Delete().

diff --git a/Engine/Source/Runtime/VulkanRHI/Private/VulkanPipeline.cpp b/Engine/Source/Runtime/VulkanRHI/Private/VulkanPipeline.cpp
--- a/Engine/Source/Runtime/VulkanRHI/Private/VulkanPipeline.cpp
+++ b/Engine/Source/Runtime/VulkanRHI/Private/VulkanPipeline.cpp
@@ -1766,6 +1766,15 @@ FVulkanRHIGraphicsPipelineState* FVulkanPipelineStateCacheManager::RHICreateGrap
 
 
 	{
+		// Workers can be creating PSOs while FRHIResource::FlushPendingDeletes is running on the RHI thread
+		// so let it get enqueued for a delete with Release() instead.  Only used for failed or duplicate PSOs...
+		auto DeleteNewPSO = [](FVulkanRHIGraphicsPipelineState* PSOPtr)
+		{
+			PSOPtr->AddRef();
+			const uint32 RefCount = PSOPtr->Release();
+			check(RefCount == 0);
+		};
+
 		SCOPE_CYCLE_COUNTER(STAT_VulkanPSOCreationTime);
 		NewPSO = new FVulkanRHIGraphicsPipelineState(Device, Initializer, Desc, &Key);
 		{
@@ -1812,7 +1821,7 @@ FVulkanRHIGraphicsPipelineState* FVulkanPipelineStateCacheManager::RHICreateGrap
 
 				if(!CreateGfxPipelineFromEntry(NewPSO, VulkanShaders, &NewPSO->VulkanPipeline))
 				{
-					delete NewPSO;
+					DeleteNewPSO(NewPSO);
 					return nullptr;
 				}
 				// Recover if we failed to create the pipeline.
@@ -1827,7 +1836,7 @@ FVulkanRHIGraphicsPipelineState* FVulkanPipelineStateCacheManager::RHICreateGrap
 			FVulkanRHIGraphicsPipelineState** MapPSO = GraphicsPSOLockedMap.Find(Key);
 			if(MapPSO)//another thread could end up creating it.
 			{
-				NewPSO->Delete();
+				DeleteNewPSO(NewPSO);
 				NewPSO = *MapPSO;
 			}
 			else
