Move VulkanMemory.h to Private for easier hotfixing in the future.

diff --git a/Engine/Source/Runtime/VulkanRHI/Private/VulkanMemory.h b/Engine/Source/Runtime/VulkanRHI/Private/VulkanMemory.h
new file mode 100644
--- /dev/null
+++ b/Engine/Source/Runtime/VulkanRHI/Private/VulkanMemory.h
@@ -0,0 +1,1227 @@
+// Copyright Epic Games, Inc. All Rights Reserved.
+
+/*=============================================================================
+	VulkanMemory.h: Vulkan Memory RHI definitions.
+=============================================================================*/
+
+#pragma once 
+
+
+//enable to store FILE/LINE, and optionally a stacktrace via r.vulkan.backtrace
+#if !UE_BUILD_SHIPPING
+#define VULKAN_MEMORY_TRACK 1
+#else
+#define VULKAN_MEMORY_TRACK 0
+#endif
+
+
+#define VULKAN_MEMORY_LOW_PRIORITY 0.f
+#define VULKAN_MEMORY_MEDIUM_PRIORITY 0.5f
+#define VULKAN_MEMORY_HIGHER_PRIORITY 0.75f
+#define VULKAN_MEMORY_HIGHEST_PRIORITY 1.f
+
+class FVulkanCommandListContext;
+class FVulkanQueue;
+class FVulkanCmdBuffer;
+
+#if !VULKAN_OBJECT_TRACKING
+#define VULKAN_TRACK_OBJECT_CREATE(Type, Ptr) do{}while(0)
+#define VULKAN_TRACK_OBJECT_DELETE(Type, Ptr) do{}while(0)
+#else
+#define VULKAN_TRACK_OBJECT_CREATE(Type, Ptr) do{TVulkanTrackBase<Type>::Add(Ptr);}while(0)
+#define VULKAN_TRACK_OBJECT_DELETE(Type, Ptr) do{TVulkanTrackBase<Type>::Remove(Ptr);}while(0)
+
+template<typename Type>
+class TVulkanTrackBase
+{
+public:
+
+	static FCriticalSection Lock;
+	static TSet<Type*> Objects;
+	template<typename Callback>
+	static uint32 CollectAll(Callback CB)
+	{
+		uint32 Count = 0;
+		FScopeLock L(&Lock);
+		TArray<Type*> Temp;
+		for(Type* Object : Objects)
+		{
+			Object->DumpMemory(CB);
+			Count++;
+		}
+		return Count;
+	}
+	static void Add(Type* Object)
+	{
+		FScopeLock L(&Lock);
+		Objects.Add(Object);
+
+	}
+	static void Remove(Type* Object)
+	{
+		FScopeLock L(&Lock);
+		Objects.Remove(Object);
+	}
+};
+
+template<typename Type>
+FCriticalSection TVulkanTrackBase<Type>::Lock;
+template<typename Type>
+TSet<Type*> TVulkanTrackBase<Type>::Objects;
+#endif
+
+namespace VulkanRHI
+{
+	class FFenceManager;
+	class FDeviceMemoryManager;					// Manager of low level vulkan memory allocations
+	class FDeviceMemoryAllocation;				// A single low level allocation
+	class FVulkanSubresourceAllocator;			// handles suballocation of one allocation into many
+	class FVulkanResourceHeap;					// has multiple suballocators. On resource heap per memory type
+	class FMemoryManager;						// entry point for everything. Contains all heaps. has a small allocator for small stuff.
+	struct FRange;								// Helper class for suballocation (this is effectively the allocator implementation)
+	class FVulkanAllocation;					// External handle to any allocation handled by FMemoryManager
+	struct FVulkanAllocationInternal;			// Internal representation mirroring vulkan allocation
+	struct FResourceHeapStats;
+}
+
+
+class FVulkanEvictable
+{
+public:
+	virtual void Evict(FVulkanDevice& Device) = 0;
+	virtual void Move(FVulkanDevice& Device, FVulkanCommandListContext& Context, VulkanRHI::FVulkanAllocation& Allocation) = 0;
+
+	virtual void OnFullDefrag(FVulkanDevice& Device, FVulkanCommandListContext& Context, uint32 NewOffset);
+	
+	virtual struct FVulkanTextureBase* GetTextureBase(){return 0;}
+	virtual bool CanMove(){return true;}
+	virtual bool CanEvict(){return true;}
+	
+};
+
+
+struct FVulkanTrackInfo
+{
+	FVulkanTrackInfo();
+	void* Data;
+	int32 SizeOrLine; //negative indicates a stack trace, Positive line number
+};
+
+enum class EDelayAcquireImageType
+{
+	None,			// acquire next image on frame start
+	DelayAcquire,	// acquire next image just before presenting, rendering is done to intermediate image which is copied to real backbuffer
+	LazyAcquire,	// acquire next image on first use
+};
+
+extern EDelayAcquireImageType GVulkanDelayAcquireImage;
+
+extern int32 GVulkanUseBufferBinning;
+
+namespace VulkanRHI
+{
+	enum
+	{
+		NUM_FRAMES_TO_WAIT_BEFORE_RELEASING_TO_OS = 3,
+	};
+
+	enum EVulkanAllocationType : uint8 {
+		EVulkanAllocationEmpty,
+		EVulkanAllocationPooledBuffer,
+		EVulkanAllocationBuffer,
+		EVulkanAllocationImage,
+		EVulkanAllocationImageDedicated,
+		EVulkanAllocationSize, // keep last
+	};
+
+	enum EVulkanAllocationMetaType : uint8
+	{
+		EVulkanAllocationMetaUnknown,
+		EVulkanAllocationMetaUniformBuffer,
+		EVulkanAllocationMetaMultiBuffer,
+		EVulkanAllocationMetaRingBuffer,
+		EVulkanAllocationMetaFrameTempBuffer,
+		EVulkanAllocationMetaImageRenderTarget,
+		EVulkanAllocationMetaImageOther,
+		EVulkanAllocationMetaBufferUAV,
+		EVulkanAllocationMetaBufferStaging,
+		EVulkanAllocationMetaBufferOther,
+		EVulkanAllocationMetaSize, // keep last
+	};
+	enum EVulkanAllocationFlags
+	{
+		VulkanAllocationFlagsMapped = 0x1,
+		VulkanAllocationFlagsCanEvict= 0x2,
+	};
+
+	enum class EType
+	{
+		Image,
+		Buffer,
+	};
+
+
+	// Custom ref counting
+	class FRefCount
+	{
+	public:
+		FRefCount() {}
+		virtual ~FRefCount()
+		{
+			check(NumRefs.GetValue() == 0);
+		}
+
+		inline uint32 AddRef()
+		{
+			int32 NewValue = NumRefs.Increment();
+			check(NewValue > 0);
+			return (uint32)NewValue;
+		}
+
+		inline uint32 Release()
+		{
+			int32 NewValue = NumRefs.Decrement();
+			if (NewValue == 0)
+			{
+				delete this;
+			}
+			check(NewValue >= 0);
+			return (uint32)NewValue;
+		}
+
+		inline uint32 GetRefCount() const
+		{
+			int32 Value = NumRefs.GetValue();
+			check(Value >= 0);
+			return (uint32)Value;
+		}
+
+	private:
+		FThreadSafeCounter NumRefs;
+	};
+
+	class FDeviceChild
+	{
+	public:
+		FDeviceChild(FVulkanDevice* InDevice = nullptr) :
+			Device(InDevice)
+		{
+		}
+
+		inline FVulkanDevice* GetParent() const
+		{
+			// Has to have one if we are asking for it...
+			check(Device);
+			return Device;
+		}
+
+		inline void SetParent(FVulkanDevice* InDevice)
+		{
+			check(!Device);
+			Device = InDevice;
+		}
+
+	protected:
+		FVulkanDevice* Device;
+	};
+
+
+	class FVulkanViewBase : public FDeviceChild
+	{
+	public:
+		virtual ~FVulkanViewBase()
+		{
+			check(NextView == 0);
+		}
+		FVulkanViewBase(FVulkanDevice* Device = nullptr)
+		: FDeviceChild(Device)
+		, NextView(0)
+		{
+		}
+		virtual void Invalidate(){}
+		FVulkanViewBase* NextView;
+	};
+	// An Allocation off a Device Heap. Lowest level of allocations and bounded by VkPhysicalDeviceLimits::maxMemoryAllocationCount.
+	class FDeviceMemoryAllocation
+	{
+	public:
+		FDeviceMemoryAllocation()
+			: Size(0)
+			, DeviceHandle(VK_NULL_HANDLE)
+			, Handle(VK_NULL_HANDLE)
+			, MappedPointer(nullptr)
+			, MemoryTypeIndex(0)
+			, bCanBeMapped(0)
+			, bIsCoherent(0)
+			, bIsCached(0)
+			, bFreedBySystem(false)
+			, bDedicatedMemory(0)
+		{
+		}
+
+		void* Map(VkDeviceSize Size, VkDeviceSize Offset);
+		void Unmap();
+
+		inline bool CanBeMapped() const
+		{
+			return bCanBeMapped != 0;
+		}
+
+		inline bool IsMapped() const
+		{
+			return !!MappedPointer;
+		}
+
+		inline void* GetMappedPointer()
+		{
+			check(IsMapped());
+			return MappedPointer;
+		}
+
+		inline bool IsCoherent() const
+		{
+			return bIsCoherent != 0;
+		}
+
+		void FlushMappedMemory(VkDeviceSize InOffset, VkDeviceSize InSize);
+		void InvalidateMappedMemory(VkDeviceSize InOffset, VkDeviceSize InSize);
+
+		inline VkDeviceMemory GetHandle() const
+		{
+			return Handle;
+		}
+
+		inline VkDeviceSize GetSize() const
+		{
+			return Size;
+		}
+
+		inline uint32 GetMemoryTypeIndex() const
+		{
+			return MemoryTypeIndex;
+		}
+
+	protected:
+		VkDeviceSize Size;
+		VkDevice DeviceHandle;
+		VkDeviceMemory Handle;
+		void* MappedPointer;
+		uint32 MemoryTypeIndex : 8;
+		uint32 bCanBeMapped : 1;
+		uint32 bIsCoherent : 1;
+		uint32 bIsCached : 1;
+		uint32 bFreedBySystem : 1;
+		uint32 bDedicatedMemory : 1;
+		uint32 : 0;
+
+
+#if VULKAN_MEMORY_TRACK
+		FVulkanTrackInfo Track;
+#endif
+		~FDeviceMemoryAllocation();
+
+		friend class FDeviceMemoryManager;
+	};
+
+	struct FDeviceMemoryBlockKey
+	{
+		uint32 MemoryTypeIndex;
+		VkDeviceSize BlockSize;
+
+		bool operator==(const FDeviceMemoryBlockKey& Other) const
+		{
+			return MemoryTypeIndex == Other.MemoryTypeIndex &&  BlockSize == Other.BlockSize;
+		}
+	};
+
+	FORCEINLINE uint32 GetTypeHash(const FDeviceMemoryBlockKey& BlockKey)
+	{
+		return HashCombine(FCrc::TypeCrc32(BlockKey.MemoryTypeIndex), FCrc::TypeCrc32(BlockKey.BlockSize));
+	}
+
+	struct FDeviceMemoryBlock
+	{
+		FDeviceMemoryBlockKey Key;
+		struct FFreeBlock
+		{
+			FDeviceMemoryAllocation* Allocation;
+			uint32 FrameFreed;
+		};
+		TArray<FFreeBlock> Allocations;
+
+	};
+
+
+
+	// Manager of Device Heap Allocations. Calling Alloc/Free is expensive!
+	class FDeviceMemoryManager
+	{
+		TMap<FDeviceMemoryBlockKey, FDeviceMemoryBlock> Allocations;
+		void UpdateMemoryProperties();
+	public:
+		FDeviceMemoryManager();
+		~FDeviceMemoryManager();
+
+		void Init(FVulkanDevice* InDevice);
+
+		void Deinit();
+
+		uint32 GetEvictedMemoryProperties();
+
+		inline bool HasUnifiedMemory() const
+		{
+			return bHasUnifiedMemory;
+		}
+
+		inline bool SupportsMemoryless() const
+		{
+			return bSupportsMemoryless;
+		}
+
+		inline uint32 GetNumMemoryTypes() const
+		{
+			return MemoryProperties.memoryTypeCount;
+		}
+
+		bool SupportsMemoryType(VkMemoryPropertyFlags Properties) const;
+		void GetHostMemoryStatus(uint64* Allocated, uint64* Total);
+		bool IsHostMemory(uint32 MemoryTypeIndex) const;
+
+		VkResult GetMemoryTypeFromProperties(uint32 TypeBits, VkMemoryPropertyFlags Properties, uint32* OutTypeIndex);
+		VkResult GetMemoryTypeFromPropertiesExcluding(uint32 TypeBits, VkMemoryPropertyFlags Properties, uint32 ExcludeTypeIndex, uint32* OutTypeIndex);
+		const VkPhysicalDeviceMemoryProperties& GetMemoryProperties() const;
+
+
+		// bCanFail means an allocation failing is not a fatal error, just returns nullptr
+		FDeviceMemoryAllocation* Alloc(bool bCanFail, VkDeviceSize AllocationSize, uint32 MemoryTypeIndex, void* DedicatedAllocateInfo, float Priority, const char* File, uint32 Line);
+		FDeviceMemoryAllocation* Alloc(bool bCanFail, VkDeviceSize AllocationSize, uint32 MemoryTypeBits, VkMemoryPropertyFlags MemoryPropertyFlags, void* DedicatedAllocateInfo, float Priority, const char* File, uint32 Line);
+
+		// Sets the Allocation to nullptr
+		void Free(FDeviceMemoryAllocation*& Allocation);
+
+		uint64 GetTotalMemory(bool bGPU) const;
+		VkDeviceSize GetBaseHeapSize(uint32 HeapIndex) const;
+		uint32 GetHeapIndex(uint32 MemoryTypeIndex);
+	protected:
+		void FreeInternal(FDeviceMemoryAllocation* Allocation);
+		void TrimMemory(bool bFullTrim);
+		friend class FMemoryManager;
+		void GetMemoryDump(TArray<FResourceHeapStats>& OutDeviceHeapsStats);
+		void DumpMemory();
+
+
+		double MemoryUpdateTime = 0.0;
+		VkPhysicalDeviceMemoryBudgetPropertiesEXT MemoryBudget;
+		VkPhysicalDeviceMemoryProperties MemoryProperties;
+		VkDevice DeviceHandle;
+		FVulkanDevice* Device;
+		uint32 NumAllocations;
+		uint32 PeakNumAllocations;
+		bool bHasUnifiedMemory;
+		bool bSupportsMemoryless;
+
+		struct FHeapInfo
+		{
+			VkDeviceSize TotalSize;
+			VkDeviceSize UsedSize;
+			VkDeviceSize PeakSize;
+			TArray<FDeviceMemoryAllocation*> Allocations;
+
+			FHeapInfo() :
+				TotalSize(0),
+				UsedSize(0),
+				PeakSize(0)
+			{
+			}
+		};
+
+		TArray<FHeapInfo> HeapInfos;
+		int32 PrimaryHostHeap; // memory usage of this heap will decide when to evict.
+
+
+		void SetupAndPrintMemInfo();
+	};
+	struct FRange
+	{
+		uint32 Offset;
+		uint32 Size;
+
+		inline bool operator<(const FRange& In) const
+		{
+			return Offset < In.Offset;
+		}
+
+		static void JoinConsecutiveRanges(TArray<FRange>& Ranges);
+
+		/** Tries to insert the item so it has index ProposedIndex, but may end up merging it with neighbors */
+		static int32 InsertAndTryToMerge(TArray<FRange>& Ranges, const FRange& Item, int32 ProposedIndex);
+
+		/** Tries to append the item to the end but may end up merging it with the neighbor */
+		static int32 AppendAndTryToMerge(TArray<FRange>& Ranges, const FRange& Item);
+
+		/** Attempts to allocate from an entry - can remove it if it was used up*/
+		static void AllocateFromEntry(TArray<FRange>& Ranges, int32 Index, uint32 SizeToAllocate);
+
+		/** Sanity checks an array of ranges */
+		static void SanityCheck(TArray<FRange>& Ranges);
+
+		/** Adds to the array while maintaing the sort. */
+		static int32 Add(TArray<FRange>& Ranges, const FRange& Item);
+	};
+
+	struct FVulkanPageSizeBucket
+	{
+		uint64 AllocationMax;
+		uint32 PageSize;
+		enum
+		{
+			BUCKET_MASK_IMAGE=0x1,
+			BUCKET_MASK_BUFFER=0x2,
+		};
+		uint32 BucketMask;
+	};
+
+
+	// Holds a reference to -any- vulkan gpu allocation
+	// !Intentionally not reference counted
+	// !User must call Free exactly once
+	class FVulkanAllocation
+	{
+	public:
+		FVulkanAllocation();
+		~FVulkanAllocation();
+		void Init(EVulkanAllocationType Type, EVulkanAllocationMetaType MetaType, uint64 Handle, uint32 InSize, uint32 AlignedOffset, uint32 AllocatorIndex, uint32 AllocationIndex, uint32 BufferId);
+		void Free(FVulkanDevice& Device);
+		void Swap(FVulkanAllocation& Other);
+		void Reference(const FVulkanAllocation& Other); //point to other, but don't take ownership
+		bool HasAllocation();
+
+		void Disown(); //disown & own should be used if ownership is transferred.
+		void Own();
+		bool IsValid() const;
+		
+		uint64 VulkanHandle = 0;
+		uint32 HandleId = 0;
+		uint32 Size = 0;
+		uint32 Offset = 0;
+		uint32 AllocationIndex = 0;
+		uint16 AllocatorIndex = 0;
+		EVulkanAllocationMetaType MetaType = EVulkanAllocationMetaUnknown;
+		uint8 Type : 7;
+		uint8 bHasOwnership : 1;
+
+		static_assert(EVulkanAllocationSize < 128, "Not enough bits to hold EVulkanAllocationType");
+		inline EVulkanAllocationType GetType() const
+		{
+			return (EVulkanAllocationType)Type;
+		}
+		inline void SetType(EVulkanAllocationType InType)
+		{
+			Type = (uint8)InType;
+		}
+
+		//helper functions
+		void* GetMappedPointer(FVulkanDevice* Device);
+		void FlushMappedMemory(FVulkanDevice* Device);
+		void InvalidateMappedMemory(FVulkanDevice* Device);
+		VkBuffer GetBufferHandle() const;
+		uint32 GetBufferAlignment(FVulkanDevice* Device) const;
+		VkDeviceMemory GetDeviceMemoryHandle(FVulkanDevice* Device) const;
+		FVulkanSubresourceAllocator* GetSubresourceAllocator(FVulkanDevice* Device) const;
+		void BindBuffer(FVulkanDevice* Device, VkBuffer Buffer);
+		void BindImage(FVulkanDevice* Device, VkImage Image);
+	};
+
+
+
+	struct FVulkanAllocationInternal
+	{
+		enum
+		{
+			EUNUSED,
+			EALLOCATED,
+			EFREED,
+			EFREEPENDING,
+			EFREEDISCARDED, // if a defrag happens with a pending free, its put into this state, so we can ignore the deferred delete once it happens.
+		};
+
+		int State = EUNUSED;
+		void Init(const FVulkanAllocation& Alloc, FVulkanEvictable* InAllocationOwner, uint32 AllocationOffset, uint32 AllocationSize, uint32 Alignment);
+
+		EVulkanAllocationType Type = EVulkanAllocationEmpty;
+		EVulkanAllocationMetaType MetaType = EVulkanAllocationMetaUnknown;
+
+		uint32 Size = 0;
+		uint32 AllocationSize = 0;
+		uint32 AllocationOffset = 0;
+		FVulkanEvictable* AllocationOwner = 0;
+		uint32 Alignment = 0;
+
+		int32 NextFree = -1;
+
+#if VULKAN_MEMORY_TRACK
+		FVulkanTrackInfo Track;
+#endif
+#if VULKAN_USE_LLM
+		uint64 LLMTrackerID;
+		void SetLLMTrackerID(uint64 InTrackerID) { LLMTrackerID = InTrackerID; }
+		uint64 GetLLMTrackerID() { return LLMTrackerID; }
+#endif
+	};
+
+
+
+	class FVulkanSubresourceAllocator
+	{
+		friend class FMemoryManager;
+		friend class FVulkanResourceHeap;
+	public:
+		FVulkanSubresourceAllocator(EVulkanAllocationType InType, FMemoryManager* InOwner, uint8 InSubResourceAllocatorFlags, FDeviceMemoryAllocation* InDeviceMemoryAllocation,
+			uint32 InMemoryTypeIndex, VkMemoryPropertyFlags InMemoryPropertyFlags,
+			uint32 InAlignment, VkBuffer InBuffer, uint32 InBufferSize, uint32 InBufferId, VkBufferUsageFlags InBufferUsageFlags, int32 InPoolSizeIndex);
+
+
+		FVulkanSubresourceAllocator(EVulkanAllocationType InType, FMemoryManager* InOwner, uint8 InSubResourceAllocatorFlags, FDeviceMemoryAllocation* InDeviceMemoryAllocation,
+			uint32 InMemoryTypeIndex, uint32 BufferId = 0xffffffff);
+
+		~FVulkanSubresourceAllocator();
+
+		void Destroy(FVulkanDevice* Device);
+		bool TryAllocate2(FVulkanAllocation& OutAllocation, FVulkanEvictable* Owner, uint32 InSize, uint32 InAlignment, EVulkanAllocationMetaType InMetaType, const char* File, uint32 Line);
+		void Free(FVulkanAllocation& Allocation);
+
+		inline uint32 GetAlignment() const
+		{
+			return Alignment;
+		}
+
+		inline void* GetMappedPointer()
+		{
+			return MemoryAllocation->GetMappedPointer();
+		}
+
+		void Flush(VkDeviceSize Offset, VkDeviceSize AllocationSize);
+		void Invalidate(VkDeviceSize Offset, VkDeviceSize AllocationSize);
+		uint32 GetMaxSize() const { return MaxSize; }
+		uint32 GetUsedSize() const { return UsedSize; }
+
+		inline uint32 GetHandleId() const
+		{
+			return BufferId;
+		}
+		FDeviceMemoryAllocation* GetMemoryAllocation()
+		{
+			return MemoryAllocation;
+		}
+		EVulkanAllocationType GetType(){ return Type; }
+
+		TArrayView<uint32> GetMemoryUsed();
+		uint32 GetNumSubAllocations();
+
+		uint8 GetSubresourceAllocatorFlags()
+		{
+			return SubresourceAllocatorFlags;
+		}
+		uint32 GetId()
+		{
+			return BufferId;
+		}
+		bool GetIsDefragging() { return bIsDefragging; }
+
+
+	protected:
+		void SetIsDefragging(bool bInIsDefragging){ bIsDefragging = bInIsDefragging; }
+		void DumpFullHeap();
+		bool DefragFull(FVulkanDevice& Device, FVulkanCommandListContext& Context, FVulkanResourceHeap* Heap);
+		int32 DefragTick(FVulkanDevice& Device, FVulkanCommandListContext& Context, FVulkanResourceHeap* Heap, uint32 Count);
+		bool CanDefrag();
+		uint64 EvictToHost(FVulkanDevice& Device);
+		void SetFreePending(FVulkanAllocation& Allocation);
+		void FreeInternalData(int32 Index);
+		int32 AllocateInternalData();
+
+		uint32 MemoryUsed[EVulkanAllocationMetaSize];
+
+		EVulkanAllocationType Type;
+		FMemoryManager* Owner;
+		uint32 MemoryTypeIndex;
+		VkMemoryPropertyFlags MemoryPropertyFlags;
+		FDeviceMemoryAllocation* MemoryAllocation;
+		uint32 MaxSize;
+		uint32 Alignment;
+		uint32 FrameFreed;
+		uint32 LastDefragFrame = 0;
+		int64 UsedSize;
+		VkBufferUsageFlags BufferUsageFlags;
+		VkBuffer Buffer;
+		uint32 BufferId;
+		int32 PoolSizeIndex;
+		uint32 AllocatorIndex;
+		uint8 SubresourceAllocatorFlags;
+		uint8 BucketId;
+		bool bIsEvicting = false;
+		bool bLocked = false;
+		bool bIsDefragging = false;
+
+		uint32 NumSubAllocations = 0;
+		uint32 AllocCalls = 0;
+		uint32 FreeCalls = 0;
+
+		// List of free ranges
+		TArray<FRange> FreeList;
+		bool JoinFreeBlocks();
+		uint32 GetAllocatorIndex()
+		{
+			return AllocatorIndex;
+		}
+
+
+#if VULKAN_MEMORY_TRACK
+		FVulkanTrackInfo Track;
+#endif
+		TArray<FVulkanAllocationInternal> InternalData;
+		int32 InternalFreeList= -1;
+		FCriticalSection CS;
+
+	};
+
+
+	// A set of Device Allocations (Heap Pages) for a specific memory type. This handles pooling allocations inside memory pages to avoid
+	// doing allocations directly off the device's heaps
+	class FVulkanResourceHeap
+	{
+	public:
+		FVulkanResourceHeap(FMemoryManager* InOwner, uint32 InMemoryTypeIndex, uint32 InOverridePageSize = 0);
+		~FVulkanResourceHeap();
+
+		void FreePage(FVulkanSubresourceAllocator* InPage);
+		void ReleasePage(FVulkanSubresourceAllocator* InPage);
+
+		//void ReleaseFreedPages(bool bImmediately);
+
+		inline FMemoryManager* GetOwner()
+		{
+			return Owner;
+		}
+
+		inline bool IsHostCachedSupported() const
+		{
+			return bIsHostCachedSupported;
+		}
+
+		inline bool IsLazilyAllocatedSupported() const
+		{
+			return bIsLazilyAllocatedSupported;
+		}
+
+		inline uint32 GetMemoryTypeIndex() const
+		{
+			return MemoryTypeIndex;
+		}
+
+		uint64 EvictOne(FVulkanDevice& Device);
+		void DefragTick(FVulkanDevice& Device, FVulkanCommandListContext& Context, uint32 Count);
+		void DumpMemory(FResourceHeapStats& Stats);
+
+		void SetDefragging(FVulkanSubresourceAllocator* Allocator);
+		bool GetIsDefragging(FVulkanSubresourceAllocator* Allocator);
+		uint32 GetPageSizeBucket(FVulkanPageSizeBucket& BucketOut, EType Type, uint32 AllocationSize, bool bForceSingleAllocation);
+
+	protected:
+		enum {			
+			MAX_BUCKETS = 5,
+		};
+		TArray<FVulkanPageSizeBucket, TFixedAllocator<MAX_BUCKETS> > PageSizeBuckets;
+		uint32 GetPageSize();
+		FMemoryManager* Owner;
+		uint16 MemoryTypeIndex;
+		uint16 HeapIndex;
+
+		bool bIsHostCachedSupported;
+		bool bIsLazilyAllocatedSupported;
+		uint8 DefragCountDown = 0;
+
+		uint32 OverridePageSize;
+
+		uint32 PeakPageSize;
+		uint64 UsedMemory;
+		uint32 PageIDCounter;
+
+		TArray<FVulkanSubresourceAllocator*> ActivePages[MAX_BUCKETS];
+		TArray<FVulkanSubresourceAllocator*> UsedDedicatedImagePages;
+
+		bool TryRealloc(FVulkanAllocation& OutAllocation, FVulkanEvictable* AllocationOwner, EType Type, uint32 Size, uint32 Alignment, EVulkanAllocationMetaType MetaType);
+		bool AllocateResource(FVulkanAllocation& OutAllocation, FVulkanEvictable* AllocationOwner, EType Type, uint32 Size, uint32 Alignment, bool bMapAllocation, bool bForceSeparateAllocation, EVulkanAllocationMetaType MetaType, const char* File, uint32 Line);
+		bool AllocateDedicatedImage(FVulkanAllocation& OutAllocation, FVulkanEvictable* AllocationOwner, VkImage Image, uint32 Size, uint32 Alignment, EVulkanAllocationMetaType MetaType, const char* File, uint32 Line);
+
+		friend class FMemoryManager;
+		friend class FVulkanSubresourceAllocator;
+	};
+
+	enum EVulkanFreeFlags
+	{
+		EVulkanFreeFlag_None = 0x0,
+		EVulkanFreeFlag_DontDefer = 0x1,
+	};
+
+
+	// Manages heaps and their interactions
+	class FMemoryManager : public FDeviceChild
+	{
+		friend class FVulkanAllocation;
+	public:
+
+		FMemoryManager(FVulkanDevice* InDevice);
+		~FMemoryManager();
+
+		void Init();
+		void Deinit();
+
+
+		void FreeVulkanAllocation(FVulkanAllocation& Allocation, EVulkanFreeFlags FreeFlags = EVulkanFreeFlag_None);
+		void FreeVulkanAllocationPooledBuffer(FVulkanAllocation& Allocation);
+		void FreeVulkanAllocationBuffer(FVulkanAllocation& Allocation);
+		void FreeVulkanAllocationImage(FVulkanAllocation& Allocation);
+		void FreeVulkanAllocationImageDedicated(FVulkanAllocation& Allocation);
+
+
+		bool AllocateBufferPooled(FVulkanAllocation& Allocation, FVulkanEvictable* AllocationOwner, uint32 Size, VkBufferUsageFlags BufferUsageFlags, VkMemoryPropertyFlags MemoryPropertyFlags, EVulkanAllocationMetaType MetaType, const char* File, uint32 Line);
+		bool AllocateImageMemory(FVulkanAllocation& Allocation, FVulkanEvictable* AllocationOwner, const VkMemoryRequirements& MemoryReqs, VkMemoryPropertyFlags MemoryPropertyFlags, EVulkanAllocationMetaType MetaType, const char* File, uint32 Line);
+		bool AllocateBufferMemory(FVulkanAllocation& Allocation, FVulkanEvictable* AllocationOwner, const VkMemoryRequirements& MemoryReqs, VkMemoryPropertyFlags MemoryPropertyFlags, EVulkanAllocationMetaType MetaType, const char* File, uint32 Line);
+		bool AllocateDedicatedImageMemory(FVulkanAllocation& Allocation, FVulkanEvictable* AllocationOwner, VkImage Image, const VkMemoryRequirements& MemoryReqs, VkMemoryPropertyFlags MemoryPropertyFlags, EVulkanAllocationMetaType MetaType, const char* File, uint32 Line);
+
+		void RegisterSubresourceAllocator(FVulkanSubresourceAllocator* SubresourceAllocator);
+		void UnregisterSubresourceAllocator(FVulkanSubresourceAllocator* SubresourceAllocator);
+		void ReleaseSubresourceAllocator(FVulkanSubresourceAllocator* SubresourceAllocator);
+
+
+		void ReleaseFreedPages(FVulkanCommandListContext& Context);
+		void DumpMemory(bool bFullDump = true);
+
+
+		void AllocUniformBuffer(FVulkanAllocation& OutAllocation, uint32 Size, const void* Contents);
+		void FreeUniformBuffer(FVulkanAllocation& InAllocation);
+
+		void HandleOOM(bool bCanResume = false, VkResult Result = VK_SUCCESS, uint64 AllocationSize = 0, uint32 MemoryTypeIndex = 0);
+		bool UpdateEvictThreshold(bool bLog);
+
+
+		void* Hotfix = nullptr;
+
+	protected:
+		FDeviceMemoryManager* DeviceMemoryManager;
+		TArray<FVulkanResourceHeap*> ResourceTypeHeaps;
+
+		enum
+		{
+			BufferAllocationSize = 1 * 1024 * 1024,
+			UniformBufferAllocationSize = 2 * 1024 * 1024,
+		};
+
+
+		// pool sizes that we support
+		enum class EPoolSizes : uint8
+		{
+// 			E32,
+// 			E64,
+			E128,
+			E256,
+			E512,
+			E1k,
+			E2k,
+			E8k,
+			E16k,
+			SizesCount,
+		};
+
+		constexpr static uint32 PoolSizes[(int32)EPoolSizes::SizesCount] =
+		{
+// 			32,
+// 			64,
+			128,
+			256,
+			512,
+			1024,
+			2048,
+			8192,
+// 			16 * 1024,
+		};
+
+		constexpr static uint32 BufferSizes[(int32)EPoolSizes::SizesCount + 1] =
+		{
+// 			64 * 1024,
+// 			64 * 1024,
+			128 * 1024,
+			128 * 1024,
+			256 * 1024,
+			256 * 1024,
+			512 * 1024,
+			512 * 1024,
+			1024 * 1024,
+			1 * 1024 * 1024,
+		};
+
+		EPoolSizes GetPoolTypeForAlloc(uint32 Size, uint32 Alignment)
+		{
+			EPoolSizes PoolSize = EPoolSizes::SizesCount;
+			if (GVulkanUseBufferBinning != 0)
+			{
+				for (int32 i = 0; i < (int32)EPoolSizes::SizesCount; ++i)
+				{
+					if (PoolSizes[i] >= Size)
+					{
+						PoolSize = (EPoolSizes)i;
+						break;
+					}
+				}
+			}
+			return PoolSize;
+		}
+
+		TArray<FVulkanSubresourceAllocator*> UsedBufferAllocations[(int32)EPoolSizes::SizesCount + 1];
+		TArray<FVulkanSubresourceAllocator*> FreeBufferAllocations[(int32)EPoolSizes::SizesCount + 1];
+		TArray<FVulkanSubresourceAllocator*> AllBufferAllocations;
+		PTRINT AllBufferAllocationsFreeListHead = (PTRINT)-1;
+
+		uint64 PendingEvictBytes = 0;
+		bool bIsEvicting = false;
+		bool bWantEviction = false;
+
+
+		friend FVulkanSubresourceAllocator* GetSubresourceAllocator(FVulkanDevice* Device, const FVulkanAllocation& Allocation);
+		friend class FVulkanResourceHeap;
+
+		void ReleaseFreedResources(bool bImmediately);
+		void DestroyResourceAllocations();
+		struct FUBPendingFree
+		{
+			FVulkanAllocation Allocation;
+			uint64 Frame = 0;
+		};
+
+		struct
+		{
+			FCriticalSection CS;
+			TArray<FUBPendingFree> PendingFree;
+			uint32 Peak = 0;
+		} UBAllocations;
+
+		void ProcessPendingUBFreesNoLock(bool bForce);
+		void ProcessPendingUBFrees(bool bForce);
+	};
+
+	class FStagingBuffer : public FVulkanEvictable, public FRefCount
+	{
+		virtual void Evict(FVulkanDevice& Device);
+		virtual void Move(FVulkanDevice& Device, FVulkanCommandListContext& Context, FVulkanAllocation& Allocation);
+		virtual bool CanMove(){return false;}
+		virtual bool CanEvict(){return false;}
+
+	public:
+		FStagingBuffer(FVulkanDevice* InDevice);
+		VkBuffer GetHandle() const;
+		void* GetMappedPointer();
+		uint32 GetSize() const;
+		VkDeviceMemory GetDeviceMemoryHandle() const;
+		void FlushMappedMemory();
+		void InvalidateMappedMemory();
+
+#if VULKAN_MEMORY_TRACK
+		FVulkanTrackInfo Track;
+#endif
+
+	protected:
+		FVulkanDevice* Device;
+		FVulkanAllocation Allocation;
+
+		VkBuffer Buffer;
+		VkMemoryPropertyFlagBits MemoryReadFlags;
+		uint32 BufferSize;
+
+		// Owner maintains lifetime
+		virtual ~FStagingBuffer();
+
+		void Destroy();
+
+		friend class FStagingManager;
+	};
+
+	class FStagingManager
+	{
+	public:
+		FStagingManager() :
+			PeakUsedMemory(0),
+			UsedMemory(0),
+			Device(nullptr)
+		{
+		}
+		~FStagingManager();
+
+		void Init(FVulkanDevice* InDevice)
+		{
+			Device = InDevice;
+		}
+
+		void Deinit();
+
+		FStagingBuffer* AcquireBuffer(uint32 Size, VkBufferUsageFlags InUsageFlags = VK_BUFFER_USAGE_TRANSFER_SRC_BIT, VkMemoryPropertyFlagBits InMemoryReadFlags = VK_MEMORY_PROPERTY_HOST_COHERENT_BIT);
+
+		// Sets pointer to nullptr
+		void ReleaseBuffer(FVulkanCmdBuffer* CmdBuffer, FStagingBuffer*& StagingBuffer);
+
+		void ProcessPendingFree(bool bImmediately, bool bFreeToOS);
+
+
+		void GetMemoryDump(FResourceHeapStats& Stats);
+		void DumpMemory();
+
+	protected:
+		friend class FMemoryManager;
+		struct FPendingItemsPerCmdBuffer
+		{
+			FVulkanCmdBuffer* CmdBuffer;
+			struct FPendingItems
+			{
+				uint64 FenceCounter;
+				TArray<FStagingBuffer*> Resources;
+			};
+
+
+			inline FPendingItems* FindOrAddItemsForFence(uint64 Fence);
+
+			TArray<FPendingItems> PendingItems;
+		};
+
+		TArray<FStagingBuffer*> UsedStagingBuffers;
+		TArray<FPendingItemsPerCmdBuffer> PendingFreeStagingBuffers;
+		struct FFreeEntry
+		{
+			FStagingBuffer* StagingBuffer;
+			uint32 FrameNumber;
+		};
+		TArray<FFreeEntry> FreeStagingBuffers;
+
+		uint64 PeakUsedMemory;
+		uint64 UsedMemory;
+
+		FPendingItemsPerCmdBuffer* FindOrAdd(FVulkanCmdBuffer* CmdBuffer);
+
+		void ProcessPendingFreeNoLock(bool bImmediately, bool bFreeToOS);
+
+		FVulkanDevice* Device;
+	};
+
+	class FFence
+	{
+	public:
+		FFence(FVulkanDevice* InDevice, FFenceManager* InOwner, bool bCreateSignaled);
+
+		inline VkFence GetHandle() const
+		{
+			return Handle;
+		}
+
+		inline bool IsSignaled() const
+		{
+			return State == EState::Signaled;
+		}
+
+		FFenceManager* GetOwner()
+		{
+			return Owner;
+		}
+
+	protected:
+		VkFence Handle;
+
+		enum class EState
+		{
+			// Initial state
+			NotReady,
+
+			// After GPU processed it
+			Signaled,
+		};
+
+		EState State;
+
+		FFenceManager* Owner;
+
+		// Only owner can delete!
+		~FFence();
+		friend class FFenceManager;
+	};
+
+	class FFenceManager
+	{
+	public:
+		FFenceManager()
+			: Device(nullptr)
+		{
+		}
+		~FFenceManager();
+
+		void Init(FVulkanDevice* InDevice);
+		void Deinit();
+
+		FFence* AllocateFence(bool bCreateSignaled = false);
+
+		inline bool IsFenceSignaled(FFence* Fence)
+		{
+			if (Fence->IsSignaled())
+			{
+				return true;
+			}
+
+			return CheckFenceState(Fence);
+		}
+
+		// Returns false if it timed out
+		bool WaitForFence(FFence* Fence, uint64 TimeInNanoseconds);
+
+		void ResetFence(FFence* Fence);
+
+		// Sets it to nullptr
+		void ReleaseFence(FFence*& Fence);
+
+		// Sets it to nullptr
+		void WaitAndReleaseFence(FFence*& Fence, uint64 TimeInNanoseconds);
+
+	protected:
+		FVulkanDevice* Device;
+		TArray<FFence*> FreeFences;
+		TArray<FFence*> UsedFences;
+
+		// Returns true if signaled
+		bool CheckFenceState(FFence* Fence);
+
+		void DestroyFence(FFence* Fence);
+	};
+
+	class FGPUEvent : public FDeviceChild, public FRefCount
+	{
+	public:
+		FGPUEvent(FVulkanDevice* InDevice);
+		virtual ~FGPUEvent();
+
+		inline VkEvent GetHandle() const
+		{
+			return Handle;
+		}
+
+	protected:
+		VkEvent Handle;
+	};
+
+
+	// Simple tape allocation per frame for a VkBuffer, used for Volatile allocations
+	class FTempFrameAllocationBuffer : public FDeviceChild
+	{
+		enum
+		{
+			ALLOCATION_SIZE = (4 * 1024 * 1024),
+		};
+
+	public:
+		FTempFrameAllocationBuffer(FVulkanDevice* InDevice);
+		virtual ~FTempFrameAllocationBuffer();
+		void Destroy();
+
+		struct FTempAllocInfo
+		{
+			FVulkanAllocation Allocation;
+			void* Data = 0;
+			uint32 CurrentOffset = 0;
+			uint32 Size = 0;
+			uint32 LockCounter = 0;
+
+			uint32 GetBindOffset()
+			{
+				checkNoEntry();
+				return 0;
+			}
+
+		};
+
+		void Alloc(uint32 InSize, uint32 InAlignment, FTempAllocInfo& OutInfo);
+		void Reset();
+
+	protected:
+		uint32 BufferIndex;
+
+		enum
+		{
+			NUM_BUFFERS = 3,
+		};
+
+		struct FFrameEntry
+		{
+			FVulkanAllocation Allocation;
+			TArray<FVulkanAllocation> PendingDeletionList;
+			uint8* MappedData = nullptr;
+			uint8* CurrentData = nullptr;
+			uint32 Size = 0;
+			uint32 PeakUsed = 0;
+
+			void InitBuffer(FVulkanDevice* InDevice, uint32 InSize);
+			void Reset(FVulkanDevice* Device);
+			bool TryAlloc(uint32 InSize, uint32 InAlignment, FTempAllocInfo& OutInfo);
+		};
+		FFrameEntry Entries[NUM_BUFFERS];
+		FCriticalSection CS;
+
+		friend class FVulkanCommandListContext;
+	};
+
+	class VULKANRHI_API FSemaphore : public FRefCount
+	{
+	public:
+		FSemaphore(FVulkanDevice& InDevice);
+		FSemaphore(FVulkanDevice& InDevice, const VkSemaphore& InExternalSemaphore);
+		virtual ~FSemaphore();
+
+		inline VkSemaphore GetHandle() const
+		{
+			return SemaphoreHandle;
+		}
+
+		inline bool IsExternallyOwned() const
+		{
+			return bExternallyOwned;
+		}
+
+	private:
+		FVulkanDevice& Device;
+		VkSemaphore SemaphoreHandle;
+		bool bExternallyOwned;
+	};
+}
+
+
+#if VULKAN_CUSTOM_MEMORY_MANAGER_ENABLED
+struct FVulkanCustomMemManager
+{
+	static VKAPI_ATTR void* Alloc(void* UserData, size_t Size, size_t Alignment, VkSystemAllocationScope AllocScope);
+	static VKAPI_ATTR void Free(void* pUserData, void* pMem);
+	static VKAPI_ATTR void* Realloc(void* pUserData, void* pOriginal, size_t size, size_t alignment, VkSystemAllocationScope allocScope);
+	static VKAPI_ATTR void InternalAllocationNotification(void* pUserData, size_t size, VkInternalAllocationType allocationType, VkSystemAllocationScope allocationScope);
+	static VKAPI_ATTR void InternalFreeNotification(void* pUserData, size_t size, VkInternalAllocationType allocationType, VkSystemAllocationScope allocationScope);
+
+	FVulkanCustomMemManager();
+
+	enum
+	{
+		VK_SYSTEM_ALLOCATION_SCOPE_RANGE_SIZE = 5,
+	};
+
+	struct FType
+	{
+		size_t UsedMemory = 0;
+		size_t MaxAllocSize = 0;
+		TMap<void*, size_t> Allocs;
+	};
+	TStaticArray<FType, VK_SYSTEM_ALLOCATION_SCOPE_RANGE_SIZE> Types;
+
+	static FType& GetType(void* pUserData, VkSystemAllocationScope AllocScope);
+};
+#endif
diff --git a/Engine/Source/Runtime/VulkanRHI/Public/VulkanMemory.h b/Engine/Source/Runtime/VulkanRHI/Public/VulkanMemory.h
deleted file mode 100644
