Change minimum driver version for Nanite on Vulkan.
Add a check for Linux also.

diff --git a/Engine/Source/Runtime/VulkanRHI/Private/Linux/VulkanLinuxPlatform.cpp b/Engine/Source/Runtime/VulkanRHI/Private/Linux/VulkanLinuxPlatform.cpp
--- a/Engine/Source/Runtime/VulkanRHI/Private/Linux/VulkanLinuxPlatform.cpp
+++ b/Engine/Source/Runtime/VulkanRHI/Private/Linux/VulkanLinuxPlatform.cpp
@@ -243,3 +243,20 @@ void FVulkanLinuxPlatform::WriteCrashMarker(const FOptionalVulkanDeviceExtension
 		}
 	}
 }
+
+void FVulkanLinuxPlatform::CheckDeviceDriver(uint32 DeviceIndex, EGpuVendorId VendorId, const VkPhysicalDeviceProperties& Props)
+{
+	if (VendorId == EGpuVendorId::Nvidia)
+	{
+		UNvidiaDriverVersion NvidiaVersion;
+		static_assert(sizeof(NvidiaVersion) == sizeof(Props.driverVersion), "Mismatched Nvidia pack driver version!");
+		NvidiaVersion.Packed = Props.driverVersion;
+
+		if ((NvidiaVersion.Major < 472) || ((NvidiaVersion.Major == 472) && (NvidiaVersion.Minor < 62)))
+		{
+			UE_LOG(LogVulkanRHI, Warning, TEXT("Nvidia drivers < 472.61.01 do not support Nanite/Lumen in Vulkan."));
+			extern TAutoConsoleVariable<int32> GRHIAllow64bitShaderAtomicsCvar;
+			GRHIAllow64bitShaderAtomicsCvar->SetWithCurrentPriority(0);
+		}
+	}
+}
\ No newline at end of file
diff --git a/Engine/Source/Runtime/VulkanRHI/Private/Linux/VulkanLinuxPlatform.h b/Engine/Source/Runtime/VulkanRHI/Private/Linux/VulkanLinuxPlatform.h
--- a/Engine/Source/Runtime/VulkanRHI/Private/Linux/VulkanLinuxPlatform.h
+++ b/Engine/Source/Runtime/VulkanRHI/Private/Linux/VulkanLinuxPlatform.h
@@ -40,6 +40,7 @@ class FVulkanLinuxPlatform : public FVulkanGenericPlatform
 public:
 	static bool IsSupported();
 
+	static void CheckDeviceDriver(uint32 DeviceIndex, EGpuVendorId VendorId, const VkPhysicalDeviceProperties& Props);
 	static bool LoadVulkanLibrary();
 	static bool LoadVulkanInstanceFunctions(VkInstance inInstance);
 	static void FreeVulkanLibrary();
diff --git a/Engine/Source/Runtime/VulkanRHI/Private/Windows/VulkanWindowsPlatform.cpp b/Engine/Source/Runtime/VulkanRHI/Private/Windows/VulkanWindowsPlatform.cpp
--- a/Engine/Source/Runtime/VulkanRHI/Private/Windows/VulkanWindowsPlatform.cpp
+++ b/Engine/Source/Runtime/VulkanRHI/Private/Windows/VulkanWindowsPlatform.cpp
@@ -402,9 +402,9 @@ void FVulkanWindowsPlatform::CheckDeviceDriver(uint32 DeviceIndex, EGpuVendorId
 			}
 		}
 
-		if ((NvidiaVersion.Major < 471) || ((NvidiaVersion.Major == 471) && (NvidiaVersion.Minor < 46)))
+		if (NvidiaVersion.Major < 472)
 		{
-			UE_LOG(LogVulkanRHI, Warning, TEXT("Nvidia drivers < 471.46 do not support Nanite/Lumen in Vulkan."));
+			UE_LOG(LogVulkanRHI, Warning, TEXT("Nvidia drivers < 472.00 do not support Nanite/Lumen in Vulkan."));
 			extern TAutoConsoleVariable<int32> GRHIAllow64bitShaderAtomicsCvar;
 			GRHIAllow64bitShaderAtomicsCvar->SetWithCurrentPriority(0);
 		}
