Loop on shutdown for as long as there are new entries to delete.
Print ImageAtomic64 support in log on startup (remove Tess).

diff --git a/Engine/Source/Runtime/RHI/Private/RHI.cpp b/Engine/Source/Runtime/RHI/Private/RHI.cpp
--- a/Engine/Source/Runtime/RHI/Private/RHI.cpp
+++ b/Engine/Source/Runtime/RHI/Private/RHI.cpp
@@ -1154,7 +1154,7 @@ bool FRHIResource::Bypass()
 
 DECLARE_CYCLE_STAT(TEXT("Delete Resources"), STAT_DeleteResources, STATGROUP_RHICMDLIST);
 
-void FRHIResource::FlushPendingDeletes(FRHICommandListImmediate& RHICmdList)
+int32 FRHIResource::FlushPendingDeletes(FRHICommandListImmediate& RHICmdList)
 {
 	SCOPE_CYCLE_COUNTER(STAT_DeleteResources);
 
@@ -1177,6 +1177,8 @@ void FRHIResource::FlushPendingDeletes(FRHICommandListImmediate& RHICmdList)
 	});
 	PendingDeletesHPC.Delete(PendingDeletesPtr);
 
+	const int32 NumDeletes = DeletedResources.Num();
+
 	RHICmdList.EnqueueLambda([DeletedResources = MoveTemp(DeletedResources)](FRHICommandListImmediate& RHICmdList) mutable
 	{
 		if (GDynamicRHI)
@@ -1193,6 +1195,8 @@ void FRHIResource::FlushPendingDeletes(FRHICommandListImmediate& RHICmdList)
 			}
 		}
 	});
+
+	return NumDeletes;
 }
 
 static_assert(ERHIZBuffer::FarPlane != ERHIZBuffer::NearPlane, "Near and Far planes must be different!");
diff --git a/Engine/Source/Runtime/RHI/Public/RHIResources.h b/Engine/Source/Runtime/RHI/Public/RHIResources.h
--- a/Engine/Source/Runtime/RHI/Public/RHIResources.h
+++ b/Engine/Source/Runtime/RHI/Public/RHIResources.h
@@ -114,7 +114,7 @@ public:
 		return uint32(CurrentValue);
 	}
 
-	static void FlushPendingDeletes(FRHICommandListImmediate& RHICmdList);
+	static int32 FlushPendingDeletes(FRHICommandListImmediate& RHICmdList);
 
 	static bool Bypass();
 
diff --git a/Engine/Source/Runtime/VulkanRHI/Private/VulkanDevice.cpp b/Engine/Source/Runtime/VulkanRHI/Private/VulkanDevice.cpp
--- a/Engine/Source/Runtime/VulkanRHI/Private/VulkanDevice.cpp
+++ b/Engine/Source/Runtime/VulkanRHI/Private/VulkanDevice.cpp
@@ -1230,7 +1230,7 @@ void FVulkanDevice::InitGPU(int32 DeviceIndex)
 	}
 #endif // VULKAN_SUPPORTS_PHYSICAL_DEVICE_PROPERTIES2
 
-	UE_LOG(LogVulkanRHI, Display, TEXT("Using Device %d: Geometry %d Tessellation %d BufferAtomic64 %d"), DeviceIndex, PhysicalFeatures.geometryShader, PhysicalFeatures.tessellationShader, OptionalDeviceExtensions.HasBufferAtomicInt64);
+	UE_LOG(LogVulkanRHI, Display, TEXT("Using Device %d: Geometry %d BufferAtomic64 %d ImageAtomic64 %d"), DeviceIndex, PhysicalFeatures.geometryShader, OptionalDeviceExtensions.HasBufferAtomicInt64, OptionalDeviceExtensions.HasImageAtomicInt64);
 
 	CreateDevice();
 
@@ -1383,11 +1383,14 @@ void FVulkanDevice::Destroy()
 	}
 
 	// Flush all pending deletes before destroying the device and any Vulkan context objects.
+	// Repeat until no new deletes are added
+	int32 NumDeletes = 0;
+	do
 	{
 		FRHICommandListImmediate& RHICmdList = FRHICommandListExecutor::GetImmediateCommandList();
-		FRHIResource::FlushPendingDeletes(RHICmdList);
+		NumDeletes = FRHIResource::FlushPendingDeletes(RHICmdList);
 		RHICmdList.ImmediateFlush(EImmediateFlushType::FlushRHIThread);
-	}
+	} while (NumDeletes > 0);
 
 	VulkanRHI::vkDestroyImageView(GetInstanceHandle(), DefaultTextureView.View, VULKAN_CPU_ALLOCATOR);
 	DefaultTextureView = {};
