Adds Vulkan external memory support needed by [REDACTED_PERSON]

diff --git a/Engine/Source/Runtime/RHI/Private/RHI.cpp b/Engine/Source/Runtime/RHI/Private/RHI.cpp
--- a/Engine/Source/Runtime/RHI/Private/RHI.cpp
+++ b/Engine/Source/Runtime/RHI/Private/RHI.cpp
@@ -242,6 +242,8 @@ FString GetTextureCreateFlagsName(ETextureCreateFlags TextureCreateFlags)
 		case TexCreate_AFRManual:						return TEXT("TexCreate_AFRManual");
 		case TexCreate_ReduceMemoryWithTilingMode:		return TEXT("TexCreate_ReduceMemoryWithTilingMode");
 		case TexCreate_Transient:						return TEXT("TexCreate_Transient");
+		case TexCreate_AtomicCompatible:				return TEXT("TexCreate_AtomicCompatible");
+		case TexCreate_External:						return TEXT("TexCreate_External");
 		}
 	});
 }
diff --git a/Engine/Source/Runtime/RHI/Public/RHIDefinitions.h b/Engine/Source/Runtime/RHI/Public/RHIDefinitions.h
--- a/Engine/Source/Runtime/RHI/Public/RHIDefinitions.h
+++ b/Engine/Source/Runtime/RHI/Public/RHIDefinitions.h
@@ -1641,6 +1641,8 @@ enum class ETextureCreateFlags : uint64
     Transient                         = 1ull << 32,
     /** Texture needs to support atomic operations */
     AtomicCompatible                  = 1ull << 33,
+	/** Texture should be allocated for external access. Vulkan only */
+	External                		  = 1ull << 34,
 };
 ENUM_CLASS_FLAGS(ETextureCreateFlags);
 
@@ -1680,6 +1682,7 @@ ENUM_CLASS_FLAGS(ETextureCreateFlags);
 #define TexCreate_ReduceMemoryWithTilingMode     ETextureCreateFlags::ReduceMemoryWithTilingMode
 #define TexCreate_Transient                      ETextureCreateFlags::Transient
 #define TexCreate_AtomicCompatible               ETextureCreateFlags::AtomicCompatible
+#define TexCreate_External               		 ETextureCreateFlags::External
 
 enum EAsyncComputePriority
 {
diff --git a/Engine/Source/Runtime/VulkanRHI/Private/VulkanDevice.cpp b/Engine/Source/Runtime/VulkanRHI/Private/VulkanDevice.cpp
--- a/Engine/Source/Runtime/VulkanRHI/Private/VulkanDevice.cpp
+++ b/Engine/Source/Runtime/VulkanRHI/Private/VulkanDevice.cpp
@@ -1260,7 +1260,7 @@ void FVulkanDevice::InitGPU(int32 DeviceIndex)
 			VulkanRHI::vkGetBufferMemoryRequirements(Device, CrashMarker.Buffer, &MemReq);
 
 			CrashMarker.Allocation = DeviceMemoryManager.Alloc(false, CreateInfo.size, MemReq.memoryTypeBits, VK_MEMORY_PROPERTY_HOST_VISIBLE_BIT |
-				VK_MEMORY_PROPERTY_HOST_COHERENT_BIT, nullptr, VULKAN_MEMORY_MEDIUM_PRIORITY, __FILE__, __LINE__);
+				VK_MEMORY_PROPERTY_HOST_COHERENT_BIT, nullptr, VULKAN_MEMORY_MEDIUM_PRIORITY, false, __FILE__, __LINE__);
 
 			uint32* Entry = (uint32*)CrashMarker.Allocation->Map(VK_WHOLE_SIZE, 0);
 			check(Entry);
@@ -1271,7 +1271,7 @@ void FVulkanDevice::InitGPU(int32 DeviceIndex)
 		else if (OptionalDeviceExtensions.HasNVDiagnosticCheckpoints)
 		{
 			CrashMarker.Allocation = DeviceMemoryManager.Alloc(false, GMaxCrashBufferEntries * sizeof(uint32_t), UINT32_MAX, VK_MEMORY_PROPERTY_HOST_VISIBLE_BIT |
-				VK_MEMORY_PROPERTY_HOST_COHERENT_BIT, nullptr, VULKAN_MEMORY_MEDIUM_PRIORITY, __FILE__, __LINE__);
+				VK_MEMORY_PROPERTY_HOST_COHERENT_BIT, nullptr, VULKAN_MEMORY_MEDIUM_PRIORITY, false, __FILE__, __LINE__);
 			uint32* Entry = (uint32*)CrashMarker.Allocation->Map(VK_WHOLE_SIZE, 0);
 			check(Entry);
 			// Start with 0 entries
diff --git a/Engine/Source/Runtime/VulkanRHI/Private/VulkanMemory.cpp b/Engine/Source/Runtime/VulkanRHI/Private/VulkanMemory.cpp
--- a/Engine/Source/Runtime/VulkanRHI/Private/VulkanMemory.cpp
+++ b/Engine/Source/Runtime/VulkanRHI/Private/VulkanMemory.cpp
@@ -12,6 +12,14 @@
 #include "VulkanLLM.h"
 #include "Containers/SortedMap.h"
 
+#if PLATFORM_WINDOWS
+#include "Windows/AllowWindowsPlatformTypes.h"
+THIRD_PARTY_INCLUDES_START
+#include <VersionHelpers.h>
+THIRD_PARTY_INCLUDES_END
+#include "Windows/HideWindowsPlatformTypes.h"
+#endif
+
 // This 'frame number' should only be used for the deletion queue
 uint32 GVulkanRHIDeletionFrameNumber = 0;
 const uint32 NUM_FRAMES_TO_WAIT_FOR_RESOURCE_DELETE = 2;
@@ -880,14 +888,14 @@ namespace VulkanRHI
 		return 0 != (MemoryProperties.memoryTypes[MemoryTypeIndex].propertyFlags & VK_MEMORY_PROPERTY_DEVICE_LOCAL_BIT);
 	}
 
-	FDeviceMemoryAllocation* FDeviceMemoryManager::Alloc(bool bCanFail, VkDeviceSize AllocationSize, uint32 MemoryTypeBits, VkMemoryPropertyFlags MemoryPropertyFlags, void* DedicatedAllocateInfo, float Priority, const char* File, uint32 Line)
+	FDeviceMemoryAllocation* FDeviceMemoryManager::Alloc(bool bCanFail, VkDeviceSize AllocationSize, uint32 MemoryTypeBits, VkMemoryPropertyFlags MemoryPropertyFlags, void* DedicatedAllocateInfo, float Priority, bool bExternal, const char* File, uint32 Line)
 	{
 		uint32 MemoryTypeIndex = ~0;
 		VERIFYVULKANRESULT(this->GetMemoryTypeFromProperties(MemoryTypeBits, MemoryPropertyFlags, &MemoryTypeIndex));
-		return Alloc(bCanFail, AllocationSize, MemoryTypeIndex, DedicatedAllocateInfo, Priority, File, Line);
+		return Alloc(bCanFail, AllocationSize, MemoryTypeIndex, DedicatedAllocateInfo, Priority, bExternal, File, Line);
 	}
 
-	FDeviceMemoryAllocation* FDeviceMemoryManager::Alloc(bool bCanFail, VkDeviceSize AllocationSize, uint32 MemoryTypeIndex, void* DedicatedAllocateInfo, float Priority, const char* File, uint32 Line)
+	FDeviceMemoryAllocation* FDeviceMemoryManager::Alloc(bool bCanFail, VkDeviceSize AllocationSize, uint32 MemoryTypeIndex, void* DedicatedAllocateInfo, float Priority, bool bExternal, const char* File, uint32 Line)
 	{
 		SCOPED_NAMED_EVENT(FDeviceMemoryManager_Alloc, FColor::Cyan);
 		FScopeLock Lock(&GDeviceMemLock);
@@ -965,6 +973,30 @@ namespace VulkanRHI
 		}
 #endif
 
+#if VULKAN_SUPPORTS_EXTERNAL_MEMORY
+		VkExportMemoryAllocateInfoKHR VulkanExportMemoryAllocateInfoKHR = {};
+#if PLATFORM_WINDOWS
+		VkExportMemoryWin32HandleInfoKHR VulkanExportMemoryWin32HandleInfoKHR = {};
+#endif // PLATFORM_WINDOWS
+		if (bExternal)
+		{
+			ZeroVulkanStruct(VulkanExportMemoryAllocateInfoKHR, VK_STRUCTURE_TYPE_EXPORT_MEMORY_ALLOCATE_INFO_KHR);
+#if PLATFORM_WINDOWS
+			ZeroVulkanStruct(VulkanExportMemoryWin32HandleInfoKHR, VK_STRUCTURE_TYPE_EXPORT_MEMORY_WIN32_HANDLE_INFO_KHR);
+			VulkanExportMemoryWin32HandleInfoKHR.pNext = Info.pNext;
+			VulkanExportMemoryWin32HandleInfoKHR.pAttributes = NULL;
+			VulkanExportMemoryWin32HandleInfoKHR.dwAccess =	GENERIC_ALL;
+			VulkanExportMemoryWin32HandleInfoKHR.name = (LPCWSTR)nullptr;
+			VulkanExportMemoryAllocateInfoKHR.pNext = IsWindows8OrGreater() ? &VulkanExportMemoryWin32HandleInfoKHR : nullptr;
+			VulkanExportMemoryAllocateInfoKHR.handleTypes = IsWindows8OrGreater() ? VK_EXTERNAL_MEMORY_HANDLE_TYPE_OPAQUE_WIN32_BIT : VK_EXTERNAL_MEMORY_HANDLE_TYPE_OPAQUE_WIN32_KMT_BIT;
+#else
+			VulkanExportMemoryAllocateInfoKHR.pNext = Info.pNext;
+			VulkanExportMemoryAllocateInfoKHR.handleTypes = VK_EXTERNAL_MEMORY_HANDLE_TYPE_OPAQUE_FD_BIT_KHR;
+#endif // PLATFORM_WINDOWS
+			Info.pNext = &VulkanExportMemoryAllocateInfoKHR;
+		}
+#endif // VULKAN_SUPPORTS_EXTERNAL_MEMORY
+
 #if VULKAN_RHI_RAYTRACING
 		VkMemoryAllocateFlagsInfo MemoryAllocateFlagsInfo;
 		ZeroVulkanStruct(MemoryAllocateFlagsInfo, VK_STRUCTURE_TYPE_MEMORY_ALLOCATE_FLAGS_INFO);
@@ -2043,7 +2075,7 @@ namespace VulkanRHI
 	}
 
 
-	bool FVulkanResourceHeap::AllocateResource(FVulkanAllocation& OutAllocation, FVulkanEvictable* AllocationOwner, EType Type, uint32 Size, uint32 Alignment, bool bMapAllocation, bool bForceSeparateAllocation, EVulkanAllocationMetaType MetaType, const char* File, uint32 Line)
+	bool FVulkanResourceHeap::AllocateResource(FVulkanAllocation& OutAllocation, FVulkanEvictable* AllocationOwner, EType Type, uint32 Size, uint32 Alignment, bool bMapAllocation, bool bForceSeparateAllocation, EVulkanAllocationMetaType MetaType, bool bExternal, const char* File, uint32 Line)
 	{
 		SCOPED_NAMED_EVENT(FResourceHeap_AllocateResource, FColor::Cyan);
 		FScopeLock ScopeLock(&GResourceLock);
@@ -2092,11 +2124,11 @@ namespace VulkanRHI
 			AllocationSize = Size;
 		}
 
-		FDeviceMemoryAllocation* DeviceMemoryAllocation = DeviceMemoryManager.Alloc(true, AllocationSize, MemoryTypeIndex, nullptr, VULKAN_MEMORY_HIGHEST_PRIORITY, File, Line);
+		FDeviceMemoryAllocation* DeviceMemoryAllocation = DeviceMemoryManager.Alloc(true, AllocationSize, MemoryTypeIndex, nullptr, VULKAN_MEMORY_HIGHEST_PRIORITY, bExternal, File, Line);
 		if (!DeviceMemoryAllocation && Size != AllocationSize)
 		{
 			// Retry with a smaller size
-			DeviceMemoryAllocation = DeviceMemoryManager.Alloc(true, Size, MemoryTypeIndex, nullptr, VULKAN_MEMORY_HIGHEST_PRIORITY, File, Line);
+			DeviceMemoryAllocation = DeviceMemoryManager.Alloc(true, Size, MemoryTypeIndex, nullptr, VULKAN_MEMORY_HIGHEST_PRIORITY, bExternal, File, Line);
 			if(!DeviceMemoryAllocation)
 			{
 				return false;
@@ -2130,7 +2162,7 @@ namespace VulkanRHI
 		return bOk;
 	}
 
-	bool FVulkanResourceHeap::AllocateDedicatedImage(FVulkanAllocation& OutAllocation, FVulkanEvictable* AllocationOwner, VkImage Image, uint32 Size, uint32 Alignment, EVulkanAllocationMetaType MetaType, const char* File, uint32 Line)
+	bool FVulkanResourceHeap::AllocateDedicatedImage(FVulkanAllocation& OutAllocation, FVulkanEvictable* AllocationOwner, VkImage Image, uint32 Size, uint32 Alignment, EVulkanAllocationMetaType MetaType, bool bExternal, const char* File, uint32 Line)
 	{
 #if VULKAN_SUPPORTS_DEDICATED_ALLOCATION
 		FScopeLock ScopeLock(&GResourceLock);
@@ -2141,7 +2173,7 @@ namespace VulkanRHI
 		VkMemoryDedicatedAllocateInfoKHR DedicatedAllocInfo;
 		ZeroVulkanStruct(DedicatedAllocInfo, VK_STRUCTURE_TYPE_MEMORY_DEDICATED_ALLOCATE_INFO_KHR);
 		DedicatedAllocInfo.image = Image;
-		FDeviceMemoryAllocation* DeviceMemoryAllocation = Owner->GetParent()->GetDeviceMemoryManager().Alloc(true, AllocationSize, MemoryTypeIndex, &DedicatedAllocInfo, VULKAN_MEMORY_HIGHEST_PRIORITY, File, Line);
+		FDeviceMemoryAllocation* DeviceMemoryAllocation = Owner->GetParent()->GetDeviceMemoryManager().Alloc(true, AllocationSize, MemoryTypeIndex, &DedicatedAllocInfo, VULKAN_MEMORY_HIGHEST_PRIORITY, bExternal, File, Line);
 		if(!DeviceMemoryAllocation)
 		{
 			return false;
@@ -2816,14 +2848,14 @@ namespace VulkanRHI
 		VERIFYVULKANRESULT(Device->GetDeviceMemoryManager().GetMemoryTypeFromProperties(MemReqs.memoryTypeBits, MemoryPropertyFlags, &MemoryTypeIndex));		
 
 		bool bHasUnifiedMemory = DeviceMemoryManager->HasUnifiedMemory();
-		FDeviceMemoryAllocation* DeviceMemoryAllocation = DeviceMemoryManager->Alloc(true, MemReqs.size, MemoryTypeIndex, nullptr, Priority, File, Line);
+		FDeviceMemoryAllocation* DeviceMemoryAllocation = DeviceMemoryManager->Alloc(true, MemReqs.size, MemoryTypeIndex, nullptr, Priority, false, File, Line);
 		if(!DeviceMemoryAllocation)
 		{
 			VkMemoryPropertyFlags MemoryPropertyFallbackFlags = MemoryPropertyFlags & (~VK_MEMORY_PROPERTY_DEVICE_LOCAL_BIT);
 			uint32 MemoryTypeIndexFallback;
 			if(GVulkanMemoryMemoryFallbackToHost && VK_SUCCESS == Device->GetDeviceMemoryManager().GetMemoryTypeFromPropertiesExcluding(MemReqs.memoryTypeBits, MemoryPropertyFallbackFlags, MemoryTypeIndex, &MemoryTypeIndexFallback))
 			{
-				DeviceMemoryAllocation = DeviceMemoryManager->Alloc(false, MemReqs.size, MemoryTypeIndexFallback, nullptr, Priority, File, Line);
+				DeviceMemoryAllocation = DeviceMemoryManager->Alloc(false, MemReqs.size, MemoryTypeIndexFallback, nullptr, Priority, false, File, Line);
 			}			
 		}
 		if(!DeviceMemoryAllocation)
@@ -2933,7 +2965,7 @@ namespace VulkanRHI
 		return bIsEvicting;
 	}
 
-	bool FMemoryManager::AllocateImageMemory(FVulkanAllocation& OutAllocation, FVulkanEvictable* AllocationOwner, const VkMemoryRequirements& MemoryReqs, VkMemoryPropertyFlags MemoryPropertyFlags, EVulkanAllocationMetaType MetaType, const char* File, uint32 Line)
+	bool FMemoryManager::AllocateImageMemory(FVulkanAllocation& OutAllocation, FVulkanEvictable* AllocationOwner, const VkMemoryRequirements& MemoryReqs, VkMemoryPropertyFlags MemoryPropertyFlags, EVulkanAllocationMetaType MetaType, bool bExternal, const char* File, uint32 Line)
 	{
 		bool bHasUnifiedMemory = DeviceMemoryManager->HasUnifiedMemory();
 		bool bCanEvict = MetaTypeCanEvict(MetaType);
@@ -2964,8 +2996,8 @@ namespace VulkanRHI
 
 		check(MemoryReqs.size <= (uint64)MAX_uint32);
 
-		const bool bForceSeparateAllocation = (MemoryPropertyFlags & VK_MEMORY_PROPERTY_LAZILY_ALLOCATED_BIT) == VK_MEMORY_PROPERTY_LAZILY_ALLOCATED_BIT;
-		if(!ResourceTypeHeaps[TypeIndex]->AllocateResource(OutAllocation, AllocationOwner, EType::Image, MemoryReqs.size, MemoryReqs.alignment, bMapped, bForceSeparateAllocation, MetaType, File, Line))
+		const bool bForceSeparateAllocation = bExternal || (MemoryPropertyFlags & VK_MEMORY_PROPERTY_LAZILY_ALLOCATED_BIT) == VK_MEMORY_PROPERTY_LAZILY_ALLOCATED_BIT;
+		if(!ResourceTypeHeaps[TypeIndex]->AllocateResource(OutAllocation, AllocationOwner, EType::Image, MemoryReqs.size, MemoryReqs.alignment, bMapped, bForceSeparateAllocation, MetaType, bExternal, File, Line))
 		{
 			if(GVulkanMemoryMemoryFallbackToHost)
 			{
@@ -2974,7 +3006,7 @@ namespace VulkanRHI
 
 			VERIFYVULKANRESULT(DeviceMemoryManager->GetMemoryTypeFromPropertiesExcluding(MemoryReqs.memoryTypeBits, MemoryPropertyFlags, TypeIndex, &TypeIndex));
 			bMapped = (MemoryPropertyFlags & VK_MEMORY_PROPERTY_HOST_VISIBLE_BIT) == VK_MEMORY_PROPERTY_HOST_VISIBLE_BIT;
-			if(!ResourceTypeHeaps[TypeIndex]->AllocateResource(OutAllocation, AllocationOwner, EType::Image, MemoryReqs.size, MemoryReqs.alignment, bMapped, bForceSeparateAllocation, MetaType, File, Line))
+			if(!ResourceTypeHeaps[TypeIndex]->AllocateResource(OutAllocation, AllocationOwner, EType::Image, MemoryReqs.size, MemoryReqs.alignment, bMapped, bForceSeparateAllocation, MetaType, bExternal, File, Line))
 			{
 				DumpMemory();
 				UE_LOG(LogVulkanRHI, Fatal, TEXT("Out Of Memory, trying to allocate %d bytes\n"), MemoryReqs.size);
@@ -2984,7 +3016,7 @@ namespace VulkanRHI
 		return true;
 	}
 
-	bool FMemoryManager::AllocateBufferMemory(FVulkanAllocation& OutAllocation, FVulkanEvictable* AllocationOwner, const VkMemoryRequirements& MemoryReqs, VkMemoryPropertyFlags MemoryPropertyFlags, EVulkanAllocationMetaType MetaType, const char* File, uint32 Line)
+	bool FMemoryManager::AllocateBufferMemory(FVulkanAllocation& OutAllocation, FVulkanEvictable* AllocationOwner, const VkMemoryRequirements& MemoryReqs, VkMemoryPropertyFlags MemoryPropertyFlags, EVulkanAllocationMetaType MetaType, bool bExternal, const char* File, uint32 Line)
 	{
 		SCOPED_NAMED_EVENT(FResourceHeapManager_AllocateBufferMemory, FColor::Cyan);
 		uint32 TypeIndex = 0;
@@ -3020,7 +3052,7 @@ namespace VulkanRHI
 
 		check(MemoryReqs.size <= (uint64)MAX_uint32);
 
-		if(!ResourceTypeHeaps[TypeIndex]->AllocateResource(OutAllocation, AllocationOwner, EType::Buffer, MemoryReqs.size, MemoryReqs.alignment, bMapped, false, MetaType, File, Line))
+		if(!ResourceTypeHeaps[TypeIndex]->AllocateResource(OutAllocation, AllocationOwner, EType::Buffer, MemoryReqs.size, MemoryReqs.alignment, bMapped, false, MetaType, bExternal, File, Line))
 		{
 			// Try another memory type if the allocation failed
 			VERIFYVULKANRESULT(DeviceMemoryManager->GetMemoryTypeFromPropertiesExcluding(MemoryReqs.memoryTypeBits, MemoryPropertyFlags, TypeIndex, &TypeIndex));
@@ -3029,7 +3061,7 @@ namespace VulkanRHI
 			{
 				UE_LOG(LogVulkanRHI, Fatal, TEXT("Missing memory type index %d, MemSize %d, MemPropTypeBits %u, MemPropertyFlags %u, %s(%d)"), TypeIndex, (uint32)MemoryReqs.size, (uint32)MemoryReqs.memoryTypeBits, (uint32)MemoryPropertyFlags, ANSI_TO_TCHAR(File), Line);
 			}
-			if(!ResourceTypeHeaps[TypeIndex]->AllocateResource(OutAllocation, AllocationOwner, EType::Buffer, MemoryReqs.size, MemoryReqs.alignment, bMapped, false, MetaType, File, Line))
+			if(!ResourceTypeHeaps[TypeIndex]->AllocateResource(OutAllocation, AllocationOwner, EType::Buffer, MemoryReqs.size, MemoryReqs.alignment, bMapped, false, MetaType, bExternal, File, Line))
 			{
 				DumpMemory();
 				UE_LOG(LogVulkanRHI, Fatal, TEXT("Out Of Memory, trying to allocate %d bytes\n"), MemoryReqs.size);
@@ -3039,7 +3071,7 @@ namespace VulkanRHI
 		return true;
 	}
 
-	bool FMemoryManager::AllocateDedicatedImageMemory(FVulkanAllocation& OutAllocation, FVulkanEvictable* AllocationOwner, VkImage Image, const VkMemoryRequirements& MemoryReqs, VkMemoryPropertyFlags MemoryPropertyFlags, EVulkanAllocationMetaType MetaType, const char* File, uint32 Line)
+	bool FMemoryManager::AllocateDedicatedImageMemory(FVulkanAllocation& OutAllocation, FVulkanEvictable* AllocationOwner, VkImage Image, const VkMemoryRequirements& MemoryReqs, VkMemoryPropertyFlags MemoryPropertyFlags, EVulkanAllocationMetaType MetaType, bool bExternal, const char* File, uint32 Line)
 	{
 #if VULKAN_SUPPORTS_DEDICATED_ALLOCATION
 		SCOPED_NAMED_EVENT(FVulkanMemoryManager_AllocateDedicatedImageMemory, FColor::Cyan);
@@ -3067,7 +3099,7 @@ namespace VulkanRHI
 			{
 				UE_LOG(LogVulkanRHI, Fatal, TEXT("Missing memory type index %d, MemSize %d, MemPropTypeBits %u, MemPropertyFlags %u, %s(%d)"), TypeIndex, (uint32)MemoryReqs.size, (uint32)MemoryReqs.memoryTypeBits, (uint32)MemoryPropertyFlags, ANSI_TO_TCHAR(File), Line);
 			}
-			if(!ResourceTypeHeaps[TypeIndex]->AllocateDedicatedImage(OutAllocation, AllocationOwner, Image, MemoryReqs.size, MemoryReqs.alignment, MetaType, File, Line))
+			if(!ResourceTypeHeaps[TypeIndex]->AllocateDedicatedImage(OutAllocation, AllocationOwner, Image, MemoryReqs.size, MemoryReqs.alignment, MetaType, bExternal, File, Line))
 			{
 				if(GVulkanMemoryMemoryFallbackToHost)
 				{
@@ -3075,7 +3107,7 @@ namespace VulkanRHI
 				}
 				if(VK_SUCCESS == DeviceMemoryManager->GetMemoryTypeFromPropertiesExcluding(MemoryReqs.memoryTypeBits, MemoryPropertyFlags, TypeIndex, &TypeIndex))
 				{
-					if(ResourceTypeHeaps[TypeIndex]->AllocateDedicatedImage(OutAllocation, AllocationOwner, Image, MemoryReqs.size, MemoryReqs.alignment, MetaType, File, Line))
+					if(ResourceTypeHeaps[TypeIndex]->AllocateDedicatedImage(OutAllocation, AllocationOwner, Image, MemoryReqs.size, MemoryReqs.alignment, MetaType, bExternal, File, Line))
 					{
 						return true;
 					}
@@ -3086,7 +3118,7 @@ namespace VulkanRHI
 		}
 		else
 		{
-			return AllocateImageMemory(OutAllocation, AllocationOwner, MemoryReqs, MemoryPropertyFlags, MetaType, File, Line);
+			return AllocateImageMemory(OutAllocation, AllocationOwner, MemoryReqs, MemoryPropertyFlags, MetaType, bExternal, File, Line);
 		}
 #else
 		checkNoEntry();
@@ -4141,7 +4173,7 @@ namespace VulkanRHI
 					break;
 			}
 		}
-		FDeviceMemoryAllocation* NewMemoryAllocation = DeviceMemoryManager.Alloc(true, AllocationSize, MemoryTypeIndex, nullptr, VULKAN_MEMORY_HIGHEST_PRIORITY, __FILE__, __LINE__);
+		FDeviceMemoryAllocation* NewMemoryAllocation = DeviceMemoryManager.Alloc(true, AllocationSize, MemoryTypeIndex, nullptr, VULKAN_MEMORY_HIGHEST_PRIORITY, false, __FILE__, __LINE__);
 		if(!NewMemoryAllocation)
 		{
 			Device.GetMemoryManager().DumpMemory();
@@ -4439,7 +4471,7 @@ namespace VulkanRHI
 		}
 
 		VkMemoryPropertyFlags readTypeFlags = InMemoryReadFlags;
-		if(!Device->GetMemoryManager().AllocateBufferMemory(StagingBuffer->Allocation, StagingBuffer, MemReqs, VK_MEMORY_PROPERTY_HOST_VISIBLE_BIT | readTypeFlags, EVulkanAllocationMetaBufferStaging, __FILE__, __LINE__))
+		if(!Device->GetMemoryManager().AllocateBufferMemory(StagingBuffer->Allocation, StagingBuffer, MemReqs, VK_MEMORY_PROPERTY_HOST_VISIBLE_BIT | readTypeFlags, EVulkanAllocationMetaBufferStaging, false, __FILE__, __LINE__))
 		{
 			check(0);
 		}
diff --git a/Engine/Source/Runtime/VulkanRHI/Private/VulkanMemory.h b/Engine/Source/Runtime/VulkanRHI/Private/VulkanMemory.h
--- a/Engine/Source/Runtime/VulkanRHI/Private/VulkanMemory.h
+++ b/Engine/Source/Runtime/VulkanRHI/Private/VulkanMemory.h
@@ -393,8 +393,8 @@ namespace VulkanRHI
 
 
 		// bCanFail means an allocation failing is not a fatal error, just returns nullptr
-		FDeviceMemoryAllocation* Alloc(bool bCanFail, VkDeviceSize AllocationSize, uint32 MemoryTypeIndex, void* DedicatedAllocateInfo, float Priority, const char* File, uint32 Line);
-		FDeviceMemoryAllocation* Alloc(bool bCanFail, VkDeviceSize AllocationSize, uint32 MemoryTypeBits, VkMemoryPropertyFlags MemoryPropertyFlags, void* DedicatedAllocateInfo, float Priority, const char* File, uint32 Line);
+		FDeviceMemoryAllocation* Alloc(bool bCanFail, VkDeviceSize AllocationSize, uint32 MemoryTypeIndex, void* DedicatedAllocateInfo, float Priority, bool bExternal, const char* File, uint32 Line);
+		FDeviceMemoryAllocation* Alloc(bool bCanFail, VkDeviceSize AllocationSize, uint32 MemoryTypeBits, VkMemoryPropertyFlags MemoryPropertyFlags, void* DedicatedAllocateInfo, float Priority, bool bExternal, const char* File, uint32 Line);
 
 		// Sets the Allocation to nullptr
 		void Free(FDeviceMemoryAllocation*& Allocation);
@@ -751,8 +751,8 @@ namespace VulkanRHI
 		TArray<FVulkanSubresourceAllocator*> UsedDedicatedImagePages;
 
 		bool TryRealloc(FVulkanAllocation& OutAllocation, FVulkanEvictable* AllocationOwner, EType Type, uint32 Size, uint32 Alignment, EVulkanAllocationMetaType MetaType);
-		bool AllocateResource(FVulkanAllocation& OutAllocation, FVulkanEvictable* AllocationOwner, EType Type, uint32 Size, uint32 Alignment, bool bMapAllocation, bool bForceSeparateAllocation, EVulkanAllocationMetaType MetaType, const char* File, uint32 Line);
-		bool AllocateDedicatedImage(FVulkanAllocation& OutAllocation, FVulkanEvictable* AllocationOwner, VkImage Image, uint32 Size, uint32 Alignment, EVulkanAllocationMetaType MetaType, const char* File, uint32 Line);
+		bool AllocateResource(FVulkanAllocation& OutAllocation, FVulkanEvictable* AllocationOwner, EType Type, uint32 Size, uint32 Alignment, bool bMapAllocation, bool bForceSeparateAllocation, EVulkanAllocationMetaType MetaType, bool bExternal, const char* File, uint32 Line);
+		bool AllocateDedicatedImage(FVulkanAllocation& OutAllocation, FVulkanEvictable* AllocationOwner, VkImage Image, uint32 Size, uint32 Alignment, EVulkanAllocationMetaType MetaType, bool bExternal, const char* File, uint32 Line);
 
 		friend class FMemoryManager;
 		friend class FVulkanSubresourceAllocator;
@@ -786,9 +786,9 @@ namespace VulkanRHI
 
 
 		bool AllocateBufferPooled(FVulkanAllocation& Allocation, FVulkanEvictable* AllocationOwner, uint32 Size, VkBufferUsageFlags BufferUsageFlags, VkMemoryPropertyFlags MemoryPropertyFlags, EVulkanAllocationMetaType MetaType, const char* File, uint32 Line);
-		bool AllocateImageMemory(FVulkanAllocation& Allocation, FVulkanEvictable* AllocationOwner, const VkMemoryRequirements& MemoryReqs, VkMemoryPropertyFlags MemoryPropertyFlags, EVulkanAllocationMetaType MetaType, const char* File, uint32 Line);
-		bool AllocateBufferMemory(FVulkanAllocation& Allocation, FVulkanEvictable* AllocationOwner, const VkMemoryRequirements& MemoryReqs, VkMemoryPropertyFlags MemoryPropertyFlags, EVulkanAllocationMetaType MetaType, const char* File, uint32 Line);
-		bool AllocateDedicatedImageMemory(FVulkanAllocation& Allocation, FVulkanEvictable* AllocationOwner, VkImage Image, const VkMemoryRequirements& MemoryReqs, VkMemoryPropertyFlags MemoryPropertyFlags, EVulkanAllocationMetaType MetaType, const char* File, uint32 Line);
+		bool AllocateImageMemory(FVulkanAllocation& Allocation, FVulkanEvictable* AllocationOwner, const VkMemoryRequirements& MemoryReqs, VkMemoryPropertyFlags MemoryPropertyFlags, EVulkanAllocationMetaType MetaType, bool bExternal, const char* File, uint32 Line);
+		bool AllocateBufferMemory(FVulkanAllocation& Allocation, FVulkanEvictable* AllocationOwner, const VkMemoryRequirements& MemoryReqs, VkMemoryPropertyFlags MemoryPropertyFlags, EVulkanAllocationMetaType MetaType, bool bExternal, const char* File, uint32 Line);
+		bool AllocateDedicatedImageMemory(FVulkanAllocation& Allocation, FVulkanEvictable* AllocationOwner, VkImage Image, const VkMemoryRequirements& MemoryReqs, VkMemoryPropertyFlags MemoryPropertyFlags, EVulkanAllocationMetaType MetaType, bool bExternal, const char* File, uint32 Line);
 
 		void RegisterSubresourceAllocator(FVulkanSubresourceAllocator* SubresourceAllocator);
 		void UnregisterSubresourceAllocator(FVulkanSubresourceAllocator* SubresourceAllocator);
diff --git a/Engine/Source/Runtime/VulkanRHI/Private/VulkanResources.h b/Engine/Source/Runtime/VulkanRHI/Private/VulkanResources.h
--- a/Engine/Source/Runtime/VulkanRHI/Private/VulkanResources.h
+++ b/Engine/Source/Runtime/VulkanRHI/Private/VulkanResources.h
@@ -301,6 +301,10 @@ public:
 		VkImageCreateInfo ImageCreateInfo;
 		//only used when HasImageFormatListKHR is supported. Otherise VK_IMAGE_CREATE_MUTABLE_FORMAT_BIT is used.
 		VkImageFormatListCreateInfoKHR ImageFormatListCreateInfo;
+#if VULKAN_SUPPORTS_EXTERNAL_MEMORY
+		//used when TexCreate_External is given
+		VkExternalMemoryImageCreateInfoKHR ExternalMemImageCreateInfo;
+#endif // VULKAN_SUPPORTS_EXTERNAL_MEMORY
 		VkFormat FormatsUsed[2];
 	};
 
diff --git a/Engine/Source/Runtime/VulkanRHI/Private/VulkanTexture.cpp b/Engine/Source/Runtime/VulkanRHI/Private/VulkanTexture.cpp
--- a/Engine/Source/Runtime/VulkanRHI/Private/VulkanTexture.cpp
+++ b/Engine/Source/Runtime/VulkanRHI/Private/VulkanTexture.cpp
@@ -365,6 +365,21 @@ void FVulkanSurface::GenerateImageCreateInfo(
 		ImageCreateInfo.usage |= VK_IMAGE_USAGE_STORAGE_BIT;
 	}
 
+	#if VULKAN_SUPPORTS_EXTERNAL_MEMORY
+	if (EnumHasAnyFlags(UEFlags, TexCreate_External))
+	{
+		VkExternalMemoryImageCreateInfoKHR& ExternalMemImageCreateInfo = OutImageCreateInfo.ExternalMemImageCreateInfo;
+		ZeroVulkanStruct(ExternalMemImageCreateInfo, VK_STRUCTURE_TYPE_EXTERNAL_MEMORY_IMAGE_CREATE_INFO_KHR);
+#if PLATFORM_WINDOWS
+		ExternalMemImageCreateInfo.handleTypes = VK_EXTERNAL_MEMORY_HANDLE_TYPE_OPAQUE_WIN32_BIT_KHR;
+#else
+	    ExternalMemImageCreateInfo.handleTypes = VK_EXTERNAL_MEMORY_HANDLE_TYPE_OPAQUE_FD_BIT_KHR;
+#endif
+		ExternalMemImageCreateInfo.pNext = ImageCreateInfo.pNext;
+    	ImageCreateInfo.pNext = &ExternalMemImageCreateInfo;
+	}
+#endif // VULKAN_SUPPORTS_EXTERNAL_MEMORY
+
 	//#todo-rco: If using CONCURRENT, make sure to NOT do so on render targets as that kills DCC compression
 	ImageCreateInfo.sharingMode = VK_SHARING_MODE_EXCLUSIVE;
 	ImageCreateInfo.queueFamilyIndexCount = 0;
@@ -596,7 +611,7 @@ FVulkanSurface::FVulkanSurface(FVulkanDevice& InDevice, FVulkanEvictable* Owner,
 		VulkanRHI::vkGetBufferMemoryRequirements(VulkanDevice, CpuReadbackBuffer->Buffer, &MemoryRequirements);
 		// Set minimum alignment to 16 bytes, as some buffers are used with CPU SIMD instructions
 		MemoryRequirements.alignment = FMath::Max<VkDeviceSize>(16, MemoryRequirements.alignment);
-		if (!InDevice.GetMemoryManager().AllocateBufferMemory(Allocation, Owner, MemoryRequirements, BufferMemFlags, EVulkanAllocationMetaBufferStaging, __FILE__, __LINE__))
+		if (!InDevice.GetMemoryManager().AllocateBufferMemory(Allocation, Owner, MemoryRequirements, BufferMemFlags, EVulkanAllocationMetaBufferStaging, false, __FILE__, __LINE__))
 		{
 			InDevice.GetMemoryManager().HandleOOM();
 		}
@@ -628,6 +643,7 @@ FVulkanSurface::FVulkanSurface(FVulkanDevice& InDevice, FVulkanEvictable* Owner,
 	const bool bUAV = EnumHasAnyFlags(UEFlags, TexCreate_UAV);
 	const bool bCPUReadback = EnumHasAnyFlags(UEFlags, TexCreate_CPUReadback);
 	const bool bDynamic = EnumHasAnyFlags(UEFlags, TexCreate_Dynamic);
+	const bool bExternal = EnumHasAnyFlags(UEFlags, TexCreate_External);
 
 	VkMemoryPropertyFlags MemoryFlags = bCPUReadback ? VK_MEMORY_PROPERTY_HOST_VISIBLE_BIT : VK_MEMORY_PROPERTY_DEVICE_LOCAL_BIT;
 
@@ -655,7 +671,7 @@ FVulkanSurface::FVulkanSurface(FVulkanDevice& InDevice, FVulkanEvictable* Owner,
 	VkDeviceSize SizeToBeConsideredForDedicated = 12 * 1024 * 1024;
 	if ((bRenderTarget || MemoryRequirements.size >= SizeToBeConsideredForDedicated) && !bMemoryless && InDevice.GetOptionalExtensions().HasKHRDedicatedAllocation && GVulkanEnableDedicatedImageMemory)
 	{
-		if(!InDevice.GetMemoryManager().AllocateDedicatedImageMemory(Allocation, Owner, Image, MemoryRequirements, MemoryFlags, MetaType, __FILE__, __LINE__))
+		if(!InDevice.GetMemoryManager().AllocateDedicatedImageMemory(Allocation, Owner, Image, MemoryRequirements, MemoryFlags, MetaType, bExternal, __FILE__, __LINE__))
 		{
 			checkNoEntry();
 		}
@@ -663,7 +679,7 @@ FVulkanSurface::FVulkanSurface(FVulkanDevice& InDevice, FVulkanEvictable* Owner,
 	else
 #endif
 	{
-		if(!InDevice.GetMemoryManager().AllocateImageMemory(Allocation, Owner, MemoryRequirements, MemoryFlags, MetaType, __FILE__, __LINE__))
+		if(!InDevice.GetMemoryManager().AllocateImageMemory(Allocation, Owner, MemoryRequirements, MemoryFlags, MetaType, bExternal, __FILE__, __LINE__))
 		{
 			checkNoEntry();
 		}
@@ -712,8 +728,10 @@ void FVulkanSurface::InternalMoveSurface(FVulkanDevice& InDevice, FVulkanCommand
 	const bool bRenderTarget = EnumHasAnyFlags(UEFlags, TexCreate_RenderTargetable | TexCreate_DepthStencilTargetable | TexCreate_ResolveTargetable);
 	const bool bCPUReadback = EnumHasAnyFlags(UEFlags, TexCreate_CPUReadback);
 	const bool bMemoryless = EnumHasAnyFlags(UEFlags, TexCreate_Memoryless);
+	const bool bExternal = EnumHasAnyFlags(UEFlags, TexCreate_External);
 	checkf(!bCPUReadback, TEXT("Move of CPUReadback surfaces not currently supported.   UEFlags=0x%x"), (int32)UEFlags);
 	checkf(!bMemoryless || !InDevice.GetDeviceMemoryManager().SupportsMemoryless(), TEXT("Move of Memoryless surfaces not currently supported.   UEFlags=0x%x"), (int32)UEFlags);
+	checkf(!bExternal, TEXT("Move of external memory not supported. UEFlags=0x%x"), (int32)UEFlags);
 
 #if UE_BUILD_DEBUG || UE_BUILD_DEVELOPMENT
 	// This shouldn't change
@@ -855,7 +873,7 @@ void FVulkanSurface::EvictSurface(FVulkanDevice& InDevice)
 
 	FVulkanAllocation HostAllocation;
 	const EVulkanAllocationMetaType MetaType = EVulkanAllocationMetaImageOther;
-	if (!InDevice.GetMemoryManager().AllocateImageMemory(HostAllocation, this, MemoryRequirements, MemProps, MetaType, __FILE__, __LINE__))
+	if (!InDevice.GetMemoryManager().AllocateImageMemory(HostAllocation, this, MemoryRequirements, MemProps, MetaType, false, __FILE__, __LINE__))
 	{
 		InDevice.GetMemoryManager().HandleOOM();
 		checkNoEntry();
