Fix bug in memory tracking during Vulkan defrag.

(reported in UDN)

diff --git a/Engine/Source/Runtime/VulkanRHI/Private/VulkanResources.h b/Engine/Source/Runtime/VulkanRHI/Private/VulkanResources.h
--- a/Engine/Source/Runtime/VulkanRHI/Private/VulkanResources.h
+++ b/Engine/Source/Runtime/VulkanRHI/Private/VulkanResources.h
@@ -444,7 +444,7 @@ private:
 	void SetInitialImageState(FVulkanCommandListContext& Context, VkImageLayout InitialLayout, bool bClear, const FClearValueBinding& ClearValueBinding);
 	friend struct FRHICommandSetInitialImageState;
 
-	void InternalMoveSurface(FVulkanDevice& InDevice, FVulkanCommandListContext& Context, VulkanRHI::FVulkanAllocation& DestAllocation, bool bSwapAllocation);
+	void InternalMoveSurface(FVulkanDevice& InDevice, FVulkanCommandListContext& Context, VulkanRHI::FVulkanAllocation& DestAllocation);
 
 private:
 	VkImageTiling Tiling;
diff --git a/Engine/Source/Runtime/VulkanRHI/Private/VulkanTexture.cpp b/Engine/Source/Runtime/VulkanRHI/Private/VulkanTexture.cpp
--- a/Engine/Source/Runtime/VulkanRHI/Private/VulkanTexture.cpp
+++ b/Engine/Source/Runtime/VulkanRHI/Private/VulkanTexture.cpp
@@ -712,7 +712,7 @@ FVulkanSurface::FVulkanSurface(FVulkanDevice& InDevice, FVulkanEvictable* Owner,
 
 
 
-void FVulkanSurface::InternalMoveSurface(FVulkanDevice& InDevice, FVulkanCommandListContext& Context, FVulkanAllocation& DestAllocation, bool bSwapAllocation)
+void FVulkanSurface::InternalMoveSurface(FVulkanDevice& InDevice, FVulkanCommandListContext& Context, FVulkanAllocation& DestAllocation)
 {
 	FImageCreateInfo ImageCreateInfo;
 	FVulkanSurface::GenerateImageCreateInfo(ImageCreateInfo,
@@ -818,14 +818,6 @@ void FVulkanSurface::InternalMoveSurface(FVulkanDevice& InDevice, FVulkanCommand
 	}
 
 	Image = MovedImage;
-
-	if (bSwapAllocation)
-	{
-		const uint64 Size = GetMemorySize();
-		VulkanTextureDestroyed(Size, ViewType, bRenderTarget);
-
-		Allocation.Swap(DestAllocation);
-	}
 }
 
 void FVulkanSurface::MoveSurface(FVulkanDevice& InDevice, FVulkanCommandListContext& Context, FVulkanAllocation& NewAllocation)
@@ -835,7 +827,10 @@ void FVulkanSurface::MoveSurface(FVulkanDevice& InDevice, FVulkanCommandListCont
 	checkf(bRenderTarget || bUAV, TEXT("Surface must be a RenderTarget or a UAV in order to be moved.  UEFlags=0x%x"), (int32)UEFlags);
 	checkf(Tiling == VK_IMAGE_TILING_OPTIMAL, TEXT("Tiling [%d] is not supported for move, only VK_IMAGE_TILING_OPTIMAL"), (int32)Tiling);
 
-	InternalMoveSurface(InDevice, Context, NewAllocation, true);
+	InternalMoveSurface(InDevice, Context, NewAllocation);
+
+	// Swap in the new allocation for this surface
+	Allocation.Swap(NewAllocation);
 }
 
 
@@ -847,7 +842,7 @@ void FVulkanSurface::OnFullDefrag(FVulkanDevice& InDevice, FVulkanCommandListCon
 	checkf(Tiling == VK_IMAGE_TILING_OPTIMAL, TEXT("Tiling [%d] is not supported for defrag, only VK_IMAGE_TILING_OPTIMAL"), (int32)Tiling);
 
 	Allocation.Offset = NewOffset;
-	InternalMoveSurface(InDevice, Context, Allocation, false);
+	InternalMoveSurface(InDevice, Context, Allocation);
 
 	//note: this exploits that the unmoved image is still bound to the old allocation, which is freed by the caller in this case.
 }
@@ -871,6 +866,7 @@ void FVulkanSurface::EvictSurface(FVulkanDevice& InDevice)
 
 	MemProps = InDevice.GetDeviceMemoryManager().GetEvictedMemoryProperties();
 
+	// Create a new host allocation to move the surface to
 	FVulkanAllocation HostAllocation;
 	const EVulkanAllocationMetaType MetaType = EVulkanAllocationMetaImageOther;
 	if (!InDevice.GetMemoryManager().AllocateImageMemory(HostAllocation, this, MemoryRequirements, MemProps, MetaType, false, __FILE__, __LINE__))
@@ -878,12 +874,12 @@ void FVulkanSurface::EvictSurface(FVulkanDevice& InDevice)
 		InDevice.GetMemoryManager().HandleOOM();
 		checkNoEntry();
 	}
-	VulkanTextureAllocated(MemoryRequirements.size, ViewType, bRenderTarget);
 
-	InternalMoveSurface(InDevice, Context, HostAllocation, true);
+	InternalMoveSurface(InDevice, Context, HostAllocation);
 
-	// Since the allocations were swapped, HostAllocation now contains the original allocation to be freed
-	Device->GetMemoryManager().FreeVulkanAllocation(HostAllocation);
+	// Delete the original allocation and swap in the new host allocation
+	Device->GetMemoryManager().FreeVulkanAllocation(Allocation);
+	Allocation.Swap(HostAllocation);
 
 	VULKAN_SET_DEBUG_NAME(InDevice, VK_OBJECT_TYPE_IMAGE, Image, TEXT("(FVulkanSurface*)0x%p [hostimage]"), this);
 }
@@ -2203,7 +2199,7 @@ void FVulkanTextureBase::Move(FVulkanDevice& Device, FVulkanCommandListContext&
 	if (GVulkanLogDefrag)
 	{
 		UE_LOG(LogVulkanRHI, Display, TEXT("Moving Surface, %p <<-- %p    :::: %s\n"), NewAllocation.Offset, 42, *Tex->GetName().ToString());
-		UE_LOG(LogVulkanRHI, Display, TEXT("Evicted %8.4fkb %8.4fkb   TB %p // %p  :: IMG %p   %-40s\n"), Size / (1024.f), TotalSize / (1024.f), this, &Surface, Surface.Image, *GetResourceFName().ToString());
+		UE_LOG(LogVulkanRHI, Display, TEXT("Moved %8.4fkb %8.4fkb   TB %p // %p  :: IMG %p   %-40s\n"), Size / (1024.f), TotalSize / (1024.f), this, &Surface, Surface.Image, *GetResourceFName().ToString());
 	}
 
 	Surface.MoveSurface(Device, Context, NewAllocation);
@@ -2219,7 +2215,7 @@ void FVulkanTextureBase::OnFullDefrag(FVulkanDevice& Device, FVulkanCommandListC
 	if (GVulkanLogDefrag)
 	{
 		UE_LOG(LogVulkanRHI, Display, TEXT("Moving Surface, %p <<-- %p    :::: %s\n"), NewOffset, 42, *Tex->GetName().ToString());
-		UE_LOG(LogVulkanRHI, Display, TEXT("Evicted %8.4fkb %8.4fkb   TB %p // %p  :: IMG %p   %-40s\n"), Size / (1024.f), TotalSize / (1024.f), this, &Surface, Surface.Image, *GetResourceFName().ToString());
+		UE_LOG(LogVulkanRHI, Display, TEXT("Defragged %8.4fkb %8.4fkb   TB %p // %p  :: IMG %p   %-40s\n"), Size / (1024.f), TotalSize / (1024.f), this, &Surface, Surface.Image, *GetResourceFName().ToString());
 	}
 
 	Surface.OnFullDefrag(Device, Context, NewOffset);
