Enable 64bit image atomics by default in Vulkan now that driver support is widely available (and proper checks are in place).

diff --git a/Engine/Source/Runtime/VulkanRHI/Private/VulkanDevice.cpp b/Engine/Source/Runtime/VulkanRHI/Private/VulkanDevice.cpp
--- a/Engine/Source/Runtime/VulkanRHI/Private/VulkanDevice.cpp
+++ b/Engine/Source/Runtime/VulkanRHI/Private/VulkanDevice.cpp
@@ -45,7 +45,7 @@ static TAutoConsoleVariable<int32> CVarVulkanUseD24(
 
 TAutoConsoleVariable<int32> GRHIAllow64bitShaderAtomicsCvar(
 	TEXT("r.Vulkan.Allow64bitShaderAtomics"),
-	0,
+	1,
 	TEXT("Whether to enable 64bit buffer/image atomics required by Nanite/Lumen\n")
 	TEXT("0 to disable 64bit atomics (default)\n")
 	TEXT("1 to enable")
