VulkanRHI: Disable 64bit image atomics on Nvidia.

diff --git a/Engine/Source/Runtime/VulkanRHI/Private/Linux/VulkanLinuxPlatform.cpp b/Engine/Source/Runtime/VulkanRHI/Private/Linux/VulkanLinuxPlatform.cpp
--- a/Engine/Source/Runtime/VulkanRHI/Private/Linux/VulkanLinuxPlatform.cpp
+++ b/Engine/Source/Runtime/VulkanRHI/Private/Linux/VulkanLinuxPlatform.cpp
@@ -252,9 +252,8 @@ void FVulkanLinuxPlatform::CheckDeviceDriver(uint32 DeviceIndex, EGpuVendorId Ve
 		static_assert(sizeof(NvidiaVersion) == sizeof(Props.driverVersion), "Mismatched Nvidia pack driver version!");
 		NvidiaVersion.Packed = Props.driverVersion;
 
-		if ((NvidiaVersion.Major < 472) || ((NvidiaVersion.Major == 472) && (NvidiaVersion.Minor < 62)))
 		{
-			UE_LOG(LogVulkanRHI, Warning, TEXT("Nvidia drivers < 472.61.01 do not support Nanite/Lumen in Vulkan."));
+			UE_LOG(LogVulkanRHI, Warning, TEXT("Nvidia drivers do not currently support Nanite/Lumen in Vulkan."));
 			extern TAutoConsoleVariable<int32> GRHIAllow64bitShaderAtomicsCvar;
 			GRHIAllow64bitShaderAtomicsCvar->SetWithCurrentPriority(0);
 		}
diff --git a/Engine/Source/Runtime/VulkanRHI/Private/Windows/VulkanWindowsPlatform.cpp b/Engine/Source/Runtime/VulkanRHI/Private/Windows/VulkanWindowsPlatform.cpp
--- a/Engine/Source/Runtime/VulkanRHI/Private/Windows/VulkanWindowsPlatform.cpp
+++ b/Engine/Source/Runtime/VulkanRHI/Private/Windows/VulkanWindowsPlatform.cpp
@@ -402,9 +402,8 @@ void FVulkanWindowsPlatform::CheckDeviceDriver(uint32 DeviceIndex, EGpuVendorId
 			}
 		}
 
-		if (NvidiaVersion.Major < 472)
 		{
-			UE_LOG(LogVulkanRHI, Warning, TEXT("Nvidia drivers < 472.00 do not support Nanite/Lumen in Vulkan."));
+			UE_LOG(LogVulkanRHI, Warning, TEXT("Nvidia drivers do not currently support Nanite/Lumen in Vulkan."));
 			extern TAutoConsoleVariable<int32> GRHIAllow64bitShaderAtomicsCvar;
 			GRHIAllow64bitShaderAtomicsCvar->SetWithCurrentPriority(0);
 		}
