Vulkan & AVEncoder compiles correctly for platform extension Windows platforms.

#rnx

diff --git a/Engine/Source/Runtime/VulkanRHI/VulkanRHI.Build.cs b/Engine/Source/Runtime/VulkanRHI/VulkanRHI.Build.cs
--- a/Engine/Source/Runtime/VulkanRHI/VulkanRHI.Build.cs
+++ b/Engine/Source/Runtime/VulkanRHI/VulkanRHI.Build.cs
@@ -6,6 +6,8 @@ using System.IO;
 
 public class VulkanRHI : ModuleRules
 {
+	protected virtual bool bShouldIncludePlatformPrivate { get { return true; } }
+
 	public VulkanRHI(ReadOnlyTargetRules Target) : base(Target)
 	{
 		bLegalToDistributeObjectCode = true;
@@ -23,7 +25,7 @@ public class VulkanRHI : ModuleRules
 				PrivateIncludePaths.Add("Runtime/VulkanRHI/Private/Linux");
 			}
 		}
-		else
+		else if (bShouldIncludePlatformPrivate)
 		{
 			PrivateIncludePaths.Add("Runtime/VulkanRHI/Private/" + Target.Platform);
 		}
