PR #8451: Docs change for r.Vulkan.Allow64bitShaderAtomics (Contributed by [REDACTED_PERSON])

diff --git a/Engine/Source/Runtime/VulkanRHI/Private/VulkanDevice.cpp b/Engine/Source/Runtime/VulkanRHI/Private/VulkanDevice.cpp
--- a/Engine/Source/Runtime/VulkanRHI/Private/VulkanDevice.cpp
+++ b/Engine/Source/Runtime/VulkanRHI/Private/VulkanDevice.cpp
@@ -47,8 +47,8 @@ TAutoConsoleVariable<int32> GRHIAllow64bitShaderAtomicsCvar(
 	TEXT("r.Vulkan.Allow64bitShaderAtomics"),
 	1,
 	TEXT("Whether to enable 64bit buffer/image atomics required by Nanite/Lumen\n")
-	TEXT("0 to disable 64bit atomics (default)\n")
-	TEXT("1 to enable")
+	TEXT("0 to disable 64bit atomics\n")
+	TEXT("1 to enable (default)")
 );
 
 #if NV_AFTERMATH
