Fixes for Vulkan Mobile Desktop renderer
 - Increased max buffered occlusion queries by 1
 - Fixed bug where depth read/stencil write had incorrect layouts

diff --git a/Engine/Source/Runtime/Renderer/Private/ScenePrivateBase.h b/Engine/Source/Runtime/Renderer/Private/ScenePrivateBase.h
--- a/Engine/Source/Runtime/Renderer/Private/ScenePrivateBase.h
+++ b/Engine/Source/Runtime/Renderer/Private/ScenePrivateBase.h
@@ -22,7 +22,7 @@ public:
 
 	enum
 	{
-		MaxBufferedOcclusionFrames = 3
+		MaxBufferedOcclusionFrames = 4
 	};
 
 	// get the system-wide number of frames of buffered occlusion queries.
diff --git a/Engine/Source/Runtime/VulkanRHI/Private/VulkanRenderTarget.cpp b/Engine/Source/Runtime/VulkanRHI/Private/VulkanRenderTarget.cpp
--- a/Engine/Source/Runtime/VulkanRHI/Private/VulkanRenderTarget.cpp
+++ b/Engine/Source/Runtime/VulkanRHI/Private/VulkanRenderTarget.cpp
@@ -1285,6 +1285,7 @@ FVulkanRenderTargetLayout::FVulkanRenderTargetLayout(FVulkanDevice& InDevice, co
 			RPInfo.DepthStencilRenderTarget.Action == EDepthStencilTargetActions::LoadDepthClearStencil_StoreDepthStencil)
 		{
 			ExclusiveDepthStencil = FExclusiveDepthStencil::DepthWrite_StencilWrite;
+			CurrentDSLayout = VK_IMAGE_LAYOUT_DEPTH_STENCIL_ATTACHMENT_OPTIMAL;
 		}
 
 		// Make sure that the requested depth-stencil access is compatible with the current layout of the DS target.
