Updating Vulkan version check in Windows with driver that has everything we need for Nanite/Lumen.

diff --git a/Engine/Source/Runtime/VulkanRHI/Private/Windows/VulkanWindowsPlatform.cpp b/Engine/Source/Runtime/VulkanRHI/Private/Windows/VulkanWindowsPlatform.cpp
--- a/Engine/Source/Runtime/VulkanRHI/Private/Windows/VulkanWindowsPlatform.cpp
+++ b/Engine/Source/Runtime/VulkanRHI/Private/Windows/VulkanWindowsPlatform.cpp
@@ -402,8 +402,9 @@ void FVulkanWindowsPlatform::CheckDeviceDriver(uint32 DeviceIndex, EGpuVendorId
 			}
 		}
 
+		if ((NvidiaVersion.Major < 496) || ((NvidiaVersion.Major == 496) && (NvidiaVersion.Minor < 13)))
 		{
-			UE_LOG(LogVulkanRHI, Warning, TEXT("Nvidia drivers do not currently support Nanite/Lumen in Vulkan."));
+			UE_LOG(LogVulkanRHI, Warning, TEXT("Nvidia drivers < 496.13 do not support Nanite/Lumen in Vulkan."));
 			extern TAutoConsoleVariable<int32> GRHIAllow64bitShaderAtomicsCvar;
 			GRHIAllow64bitShaderAtomicsCvar->SetWithCurrentPriority(0);
 		}
