VulkanMemory cleanup:
- Add VKHasAllFlags/VKHasAnyFlags to make code more compact.
- Clear up usage of "host" and "device local" in memory manager.
- Add FreePending to reports so the ranges are contiguous.
- Fix up some indices that were cumulting stats in the wrong buckets.
- Fix ever growing AllBufferAllocations array (from UDN)
- Remove hotfix pointer
- Other misc code cleanups

diff --git a/Engine/Source/Runtime/VulkanRHI/Private/VulkanMemory.cpp b/Engine/Source/Runtime/VulkanRHI/Private/VulkanMemory.cpp
--- a/Engine/Source/Runtime/VulkanRHI/Private/VulkanMemory.cpp
+++ b/Engine/Source/Runtime/VulkanRHI/Private/VulkanMemory.cpp
@@ -577,13 +577,16 @@ namespace VulkanRHI
 	}
 
 	FDeviceMemoryManager::FDeviceMemoryManager() :
+		MemoryUpdateTime(0.0),
 		DeviceHandle(VK_NULL_HANDLE),
 		Device(nullptr),
 		NumAllocations(0),
 		PeakNumAllocations(0),
 		bHasUnifiedMemory(false),
-		bSupportsMemoryless(false)
+		bSupportsMemoryless(false),
+		PrimaryHeapIndex(-1)
 	{
+		FMemory::Memzero(MemoryBudget);
 		FMemory::Memzero(MemoryProperties);
 	}
 
@@ -596,7 +599,6 @@ namespace VulkanRHI
 	{
 		if (Device->GetOptionalExtensions().HasMemoryBudget)
 		{
-			FMemory::Memzero(MemoryBudget);
 			VkPhysicalDeviceMemoryProperties2 MemoryProperties2;
 			ZeroVulkanStruct(MemoryBudget, VK_STRUCTURE_TYPE_PHYSICAL_DEVICE_MEMORY_BUDGET_PROPERTIES_EXT);
 			ZeroVulkanStruct(MemoryProperties2, VK_STRUCTURE_TYPE_PHYSICAL_DEVICE_MEMORY_PROPERTIES_2);
@@ -604,7 +606,6 @@ namespace VulkanRHI
 			VulkanRHI::vkGetPhysicalDeviceMemoryProperties2(Device->GetPhysicalHandle(), &MemoryProperties2);
 			FMemory::Memcpy(MemoryProperties, MemoryProperties2.memoryProperties);
 
-
 			for (uint32 Heap = 0; Heap < VK_MAX_MEMORY_HEAPS; ++Heap)
 			{
 				MemoryBudget.heapBudget[Heap] = GVulkanBudgetPercentageScale * MemoryBudget.heapBudget[Heap] / 100;
@@ -636,21 +637,24 @@ namespace VulkanRHI
 		NumAllocations = 0;
 		PeakNumAllocations = 0;
 
+		bHasUnifiedMemory = FVulkanPlatform::HasUnifiedMemory();
+
 		DeviceHandle = Device->GetInstanceHandle();
 		UpdateMemoryProperties();
 
-		uint64 HostHeapSize = 0;
-		PrimaryHostHeap = -1; // Primary
+		PrimaryHeapIndex = -1;
+		uint64 PrimaryHeapSize = 0;
 		uint32 NonLocalHeaps = 0;
 
 		for(uint32 i = 0; i < MemoryProperties.memoryHeapCount; ++i)
 		{
-			if(0 != (VK_MEMORY_PROPERTY_DEVICE_LOCAL_BIT & MemoryProperties.memoryHeaps[i].flags))
+			
+			if (VKHasAllFlags(MemoryProperties.memoryHeaps[i].flags, VK_MEMORY_PROPERTY_DEVICE_LOCAL_BIT))
 			{
-				if(MemoryProperties.memoryHeaps[i].size > HostHeapSize)
+				if(MemoryProperties.memoryHeaps[i].size > PrimaryHeapSize)
 				{
-					PrimaryHostHeap = i;
-					HostHeapSize = MemoryProperties.memoryHeaps[i].size;
+					PrimaryHeapIndex = i;
+					PrimaryHeapSize = MemoryProperties.memoryHeaps[i].size;
 				}
 			}
 			else
@@ -660,46 +664,34 @@ namespace VulkanRHI
 		}
 		if(0 == NonLocalHeaps)
 		{
-			PrimaryHostHeap = -1; // if there are no non-local heaps, disable eviction and defragmentation
+			PrimaryHeapIndex = -1; // if there are no non-local heaps, disable eviction and defragmentation
 		}
 
 		// Update bMemoryless support
 		bSupportsMemoryless = false;
 		for (uint32 i = 0; i < MemoryProperties.memoryTypeCount && !bSupportsMemoryless; ++i)
 		{
-			bSupportsMemoryless = ((MemoryProperties.memoryTypes[i].propertyFlags & VK_MEMORY_PROPERTY_LAZILY_ALLOCATED_BIT) == VK_MEMORY_PROPERTY_LAZILY_ALLOCATED_BIT);
+			bSupportsMemoryless = VKHasAllFlags(MemoryProperties.memoryTypes[i].propertyFlags, VK_MEMORY_PROPERTY_LAZILY_ALLOCATED_BIT);
 		}
 
-
 		HeapInfos.AddDefaulted(MemoryProperties.memoryHeapCount);
 
-		for (uint32 Index = 0; Index < MemoryProperties.memoryHeapCount; ++Index)
-		{
-			const bool bIsGPUHeap = ((MemoryProperties.memoryHeaps[Index].flags & VK_MEMORY_HEAP_DEVICE_LOCAL_BIT) == VK_MEMORY_HEAP_DEVICE_LOCAL_BIT);
-			if (bIsGPUHeap)
-			{
-				// Target using 95% of our budget to account for some fragmentation.
-				HeapInfos[Index].TotalSize = (uint64)((float)HeapInfos[Index].TotalSize * 0.95f);
-			}
-		}
-
-
-		SetupAndPrintMemInfo();
+		PrintMemInfo();
 	}
 
 	static FString GetMemoryPropertyFlagsString(VkMemoryPropertyFlags Flags)
 	{
 		FString String;
-		if ((Flags & VK_MEMORY_PROPERTY_DEVICE_LOCAL_BIT) == VK_MEMORY_PROPERTY_DEVICE_LOCAL_BIT)
+		if (VKHasAllFlags(Flags, VK_MEMORY_PROPERTY_DEVICE_LOCAL_BIT))
 		{
-			String += TEXT(" Local");
+			String += TEXT(" DeviceLocal");
 		}
 		else
 		{
 			String += TEXT("      ");
 		}
 
-		if ((Flags & VK_MEMORY_PROPERTY_HOST_VISIBLE_BIT) == VK_MEMORY_PROPERTY_HOST_VISIBLE_BIT)
+		if (VKHasAllFlags(Flags, VK_MEMORY_PROPERTY_HOST_VISIBLE_BIT))
 		{
 			String += TEXT(" HostVisible");
 		}
@@ -708,7 +700,7 @@ namespace VulkanRHI
 			String += TEXT("            ");
 		}
 
-		if ((Flags & VK_MEMORY_PROPERTY_HOST_COHERENT_BIT) == VK_MEMORY_PROPERTY_HOST_COHERENT_BIT)
+		if (VKHasAllFlags(Flags, VK_MEMORY_PROPERTY_HOST_COHERENT_BIT))
 		{
 			String += TEXT(" HostCoherent");
 		}
@@ -717,7 +709,7 @@ namespace VulkanRHI
 			String += TEXT("             ");
 		}
 
-		if ((Flags & VK_MEMORY_PROPERTY_HOST_CACHED_BIT) == VK_MEMORY_PROPERTY_HOST_CACHED_BIT)
+		if (VKHasAllFlags(Flags, VK_MEMORY_PROPERTY_HOST_CACHED_BIT))
 		{
 			String += TEXT(" HostCached");
 		}
@@ -726,7 +718,7 @@ namespace VulkanRHI
 			String += TEXT("           ");
 		}
 
-		if ((Flags & VK_MEMORY_PROPERTY_LAZILY_ALLOCATED_BIT) == VK_MEMORY_PROPERTY_LAZILY_ALLOCATED_BIT)
+		if (VKHasAllFlags(Flags, VK_MEMORY_PROPERTY_LAZILY_ALLOCATED_BIT))
 		{
 			String += TEXT(" Lazy");
 		}
@@ -748,71 +740,69 @@ namespace VulkanRHI
 
 
 
-	void FDeviceMemoryManager::SetupAndPrintMemInfo()
+	void FDeviceMemoryManager::PrintMemInfo() const
 	{
+		auto ToMB = [](uint64 ByteSize) { return (float)ByteSize / (1024.f * 1024.f); };
+		auto ToPct = [](uint64 Used, uint64 Total) { return 100.f * (float)Used / (float)Total; };
+
 		const uint32 MaxAllocations = Device->GetLimits().maxMemoryAllocationCount;
-		VULKAN_LOGMEMORY(TEXT("%d Device Memory Heaps; Max memory allocations %d"), MemoryProperties.memoryHeapCount, MaxAllocations);
-		for (uint32 Index = 0; Index < MemoryProperties.memoryHeapCount; ++Index)
+		VULKAN_LOGMEMORY(TEXT("Max memory allocations %d."), MaxAllocations);
+
+		VULKAN_LOGMEMORY(TEXT("%d Device Memory Heaps:"), MemoryProperties.memoryHeapCount);
+		for (uint32 HeapIndex = 0; HeapIndex < MemoryProperties.memoryHeapCount; ++HeapIndex)
 		{
-			bool bIsGPUHeap = ((MemoryProperties.memoryHeaps[Index].flags & VK_MEMORY_HEAP_DEVICE_LOCAL_BIT) == VK_MEMORY_HEAP_DEVICE_LOCAL_BIT);
-			VULKAN_LOGMEMORY(TEXT("%2d: Flags 0x%x Size %llu (%.2f MB) %s"),
-				Index,
-				MemoryProperties.memoryHeaps[Index].flags,
-				MemoryProperties.memoryHeaps[Index].size,
-				(float)((double)MemoryProperties.memoryHeaps[Index].size / 1024.0 / 1024.0),
-				bIsGPUHeap ? TEXT("GPU") : TEXT(""));
-			HeapInfos[Index].TotalSize = MemoryProperties.memoryHeaps[Index].size;
+			const bool bIsGPUHeap = VKHasAllFlags(MemoryProperties.memoryHeaps[HeapIndex].flags, VK_MEMORY_HEAP_DEVICE_LOCAL_BIT);
+			const VkDeviceSize UsedMemory = (HeapIndex < (uint32)HeapInfos.Num()) ? HeapInfos[HeapIndex].UsedSize : 0;
+
+			VULKAN_LOGMEMORY(TEXT(" %2d: Flags 0x%x - Size %llu (%.2f MB) - Used %llu (%.2f%%)%s%s"),
+				HeapIndex,
+				MemoryProperties.memoryHeaps[HeapIndex].flags,
+				MemoryProperties.memoryHeaps[HeapIndex].size,
+				ToMB(MemoryProperties.memoryHeaps[HeapIndex].size),
+				UsedMemory,
+				ToPct(UsedMemory, MemoryProperties.memoryHeaps[HeapIndex].size),
+				bIsGPUHeap ? TEXT(" - DeviceLocal") : TEXT(""),
+				(PrimaryHeapIndex == (int32)HeapIndex) ? TEXT(" - PrimaryHeap") : TEXT(""));
 		}
 
-		bHasUnifiedMemory = FVulkanPlatform::HasUnifiedMemory();
-		VULKAN_LOGMEMORY(TEXT("%d Device Memory Types (%sunified)"), MemoryProperties.memoryTypeCount, bHasUnifiedMemory ? TEXT("") : TEXT("Not "));
+		VULKAN_LOGMEMORY(TEXT("%d Device Memory Types (%sunified):"), MemoryProperties.memoryTypeCount, bHasUnifiedMemory ? TEXT("") : TEXT("Not "));
 		for (uint32 HeapIndex = 0; HeapIndex < MemoryProperties.memoryHeapCount; ++HeapIndex)
 		{
-			for (uint32 Index = 0; Index < MemoryProperties.memoryTypeCount; ++Index)
+			for (uint32 TypeIndex = 0; TypeIndex < MemoryProperties.memoryTypeCount; ++TypeIndex)
 			{
-				if(HeapIndex == MemoryProperties.memoryTypes[Index].heapIndex)
+				if(HeapIndex == MemoryProperties.memoryTypes[TypeIndex].heapIndex)
 				{
-					VULKAN_LOGMEMORY(TEXT("%2d: Flags 0x%05x Heap %2d %s"),
-						Index,
-						MemoryProperties.memoryTypes[Index].propertyFlags,
-						MemoryProperties.memoryTypes[Index].heapIndex,
-						*GetMemoryPropertyFlagsString(MemoryProperties.memoryTypes[Index].propertyFlags));
+					VULKAN_LOGMEMORY(TEXT(" %2d: Flags 0x%05x - Heap %2d - %s"),
+						TypeIndex,
+						MemoryProperties.memoryTypes[TypeIndex].propertyFlags,
+						MemoryProperties.memoryTypes[TypeIndex].heapIndex,
+						*GetMemoryPropertyFlagsString(MemoryProperties.memoryTypes[TypeIndex].propertyFlags));
 				}
 			}
 		}
-		uint64 HostAllocated, HostLimit;
-		GetHostMemoryStatus(&HostAllocated, &HostLimit);
-		double AllocatedPercentage = 100.0 * HostAllocated / HostLimit;
-		VULKAN_LOGMEMORY(TEXT("Host Allocation Percentage %6.2f%% -      %8.2fMB / %8.3fMB"), AllocatedPercentage, HostAllocated / (1024.f * 1024.f), HostLimit / (1024.f * 1024.f));
+
 		if(Device->GetOptionalExtensions().HasMemoryBudget)
 		{
-			UpdateMemoryProperties();
-			VULKAN_LOGMEMORY(TEXT("Memory Budget"));
+			VULKAN_LOGMEMORY(TEXT("Memory Budget Extension:"));
 			VULKAN_LOGMEMORY(TEXT("\t         | Usage                     | Budget          | Size            |"));
-			VULKAN_LOGMEMORY(TEXT("\t---------|---------------------------------------------------------------|"));
-			for(uint32 Heap = 0; Heap < VK_MAX_MEMORY_HEAPS; ++Heap)
+			VULKAN_LOGMEMORY(TEXT("\t---------|------------------------------------------------------------------|"));
+			for(uint32 HeapIndex = 0; HeapIndex < VK_MAX_MEMORY_HEAPS; ++HeapIndex)
 			{
-				VkDeviceSize Budget = MemoryBudget.heapBudget[Heap];
-				VkDeviceSize Usage = MemoryBudget.heapUsage[Heap];
+				const VkDeviceSize Budget = MemoryBudget.heapBudget[HeapIndex];
+				const VkDeviceSize Usage = MemoryBudget.heapUsage[HeapIndex];
 				if(0 != Budget || 0 != Usage)
 				{
-					VkDeviceSize Size = MemoryProperties.memoryHeaps[Heap].size;
-					double UsagePercentage = 100.0 * Usage  / Size;
-					VULKAN_LOGMEMORY(TEXT("\t HEAP %02d | %6.2f%% / %13.3fmb | %13.3fmb | %13.3fmb |"), Heap,
-						UsagePercentage,
-						Usage / (1024.0 * 1024.0),
-						Budget / (1024.0 * 1024.0),
-						Size / (1024.0 * 1024.0));
+					const VkDeviceSize Size = MemoryProperties.memoryHeaps[HeapIndex].size;
+					VULKAN_LOGMEMORY(TEXT("\t HEAP %02d | %6.2f%% / %13.2f MB | %13.2f MB | %13.2f MB |"),
+						HeapIndex, ToPct(Usage, Size), ToMB(Usage), ToMB(Budget), ToMB(Size));
 				}
 			}
-			VULKAN_LOGMEMORY(TEXT("\t---------|---------------------------------------------------------------|"));
+			VULKAN_LOGMEMORY(TEXT("\t---------|------------------------------------------------------------------|"));
 		}
 		else
 		{
 			VULKAN_LOGMEMORY(TEXT("Memory Budget unavailable"));
 		}
-
-
 	}
 
 	uint32 FDeviceMemoryManager::GetEvictedMemoryProperties()
@@ -857,35 +847,24 @@ namespace VulkanRHI
 		}
 		return false;
 	}
-	void FDeviceMemoryManager::GetHostMemoryStatus(uint64* Allocated, uint64* Total)
+
+	void FDeviceMemoryManager::GetPrimaryHeapStatus(uint64& OutAllocated, uint64& OutLimit)
 	{
-		if(PrimaryHostHeap < 0)
+		if (PrimaryHeapIndex < 0)
 		{
-			*Allocated = 0;
-			*Total = 1;
+			OutAllocated = 0;
+			OutLimit = 1;
 		}
 		else
 		{
-			*Allocated = HeapInfos[PrimaryHostHeap].UsedSize;
-			check(HeapInfos[PrimaryHostHeap].TotalSize == MemoryProperties.memoryHeaps[PrimaryHostHeap].size);
-			*Total = GetBaseHeapSize(PrimaryHostHeap);
-			if (Device->GetOptionalExtensions().HasMemoryBudget)
-			{
-				if(MemoryUpdateTime == 0.0 || FPlatformTime::Seconds() - MemoryUpdateTime >= 1.0 )
+			OutAllocated = HeapInfos[PrimaryHeapIndex].UsedSize;
+			if (Device->GetOptionalExtensions().HasMemoryBudget && (FPlatformTime::Seconds() - MemoryUpdateTime >= 1.0))
 				{
 					MemoryUpdateTime = FPlatformTime::Seconds();
 					UpdateMemoryProperties();
 				}			
-				
-				*Total = MemoryBudget.heapBudget[PrimaryHostHeap];
-			}
-		}
+			OutLimit = GetBaseHeapSize(PrimaryHeapIndex);
 	}
-
-
-	bool FDeviceMemoryManager::IsHostMemory(uint32 MemoryTypeIndex) const
-	{
-		return 0 != (MemoryProperties.memoryTypes[MemoryTypeIndex].propertyFlags & VK_MEMORY_PROPERTY_DEVICE_LOCAL_BIT);
 	}
 
 	FDeviceMemoryAllocation* FDeviceMemoryManager::Alloc(bool bCanFail, VkDeviceSize AllocationSize, uint32 MemoryTypeBits, VkMemoryPropertyFlags MemoryPropertyFlags, void* DedicatedAllocateInfo, float Priority, bool bExternal, const char* File, uint32 Line)
@@ -1009,7 +988,7 @@ namespace VulkanRHI
 		VkResult Result;
 
 #if !UE_BUILD_SHIPPING
-		if(	MemoryTypeIndex == PrimaryHostHeap && GVulkanFakeMemoryLimit && ((uint64)GVulkanFakeMemoryLimit << 20llu) < HeapInfos[PrimaryHostHeap].UsedSize )
+		if (MemoryTypeIndex == PrimaryHeapIndex && GVulkanFakeMemoryLimit && ((uint64)GVulkanFakeMemoryLimit << 20llu) < HeapInfos[PrimaryHeapIndex].UsedSize )
 		{
 			Handle = 0;
 			Result = VK_ERROR_OUT_OF_DEVICE_MEMORY;
@@ -1050,9 +1029,9 @@ namespace VulkanRHI
 		NewAllocation->Handle = Handle;
 		NewAllocation->Size = AllocationSize;
 		NewAllocation->MemoryTypeIndex = MemoryTypeIndex;
-		NewAllocation->bCanBeMapped = ((MemoryProperties.memoryTypes[MemoryTypeIndex].propertyFlags & VK_MEMORY_PROPERTY_HOST_VISIBLE_BIT) == VK_MEMORY_PROPERTY_HOST_VISIBLE_BIT);
-		NewAllocation->bIsCoherent = ((MemoryProperties.memoryTypes[MemoryTypeIndex].propertyFlags & VK_MEMORY_PROPERTY_HOST_COHERENT_BIT) == VK_MEMORY_PROPERTY_HOST_COHERENT_BIT);
-		NewAllocation->bIsCached = ((MemoryProperties.memoryTypes[MemoryTypeIndex].propertyFlags & VK_MEMORY_PROPERTY_HOST_CACHED_BIT) == VK_MEMORY_PROPERTY_HOST_CACHED_BIT);
+		NewAllocation->bCanBeMapped = VKHasAllFlags(MemoryProperties.memoryTypes[MemoryTypeIndex].propertyFlags, VK_MEMORY_PROPERTY_HOST_VISIBLE_BIT);
+		NewAllocation->bIsCoherent = VKHasAllFlags(MemoryProperties.memoryTypes[MemoryTypeIndex].propertyFlags, VK_MEMORY_PROPERTY_HOST_COHERENT_BIT);
+		NewAllocation->bIsCached = VKHasAllFlags(MemoryProperties.memoryTypes[MemoryTypeIndex].propertyFlags, VK_MEMORY_PROPERTY_HOST_CACHED_BIT);
 #if VULKAN_SUPPORTS_DEDICATED_ALLOCATION
 		NewAllocation->bDedicatedMemory = DedicatedAllocateInfo != 0;
 #else
@@ -1307,7 +1286,7 @@ namespace VulkanRHI
 	void FDeviceMemoryManager::DumpMemory()
 	{
 		VULKAN_LOGMEMORY(TEXT("/******************************************* Device Memory ********************************************\\"));
-		SetupAndPrintMemInfo();
+		PrintMemInfo();
 		VULKAN_LOGMEMORY(TEXT("Device Memory: %d allocations on %d heaps"), NumAllocations, HeapInfos.Num());
 		for (int32 Index = 0; Index < HeapInfos.Num(); ++Index)
 		{
@@ -1403,7 +1382,7 @@ namespace VulkanRHI
 
 			if (bIsGPUHeap == bGPU)
 			{
-				TotalMemory += HeapInfos[Index].TotalSize;
+				TotalMemory += MemoryProperties.memoryHeaps[Index].size;
 			}
 		}
 		return TotalMemory;
@@ -1656,7 +1635,7 @@ namespace VulkanRHI
 	{
 		VkDeviceSize HeapSize = MemoryProperties.memoryHeaps[HeapIndex].size;
 #if !UE_BUILD_SHIPPING
-		if(GVulkanFakeMemoryLimit && PrimaryHostHeap == HeapIndex)
+		if (GVulkanFakeMemoryLimit && PrimaryHeapIndex == HeapIndex)
 		{
 			HeapSize = FMath::Min<VkDeviceSize>((uint64)GVulkanFakeMemoryLimit << 20llu, HeapSize);
 		}
@@ -1743,7 +1722,6 @@ namespace VulkanRHI
 		Owner->GetParent()->GetDeviceMemoryManager().Free(Allocation );
 		UsedMemory -= InPage->MaxSize;
 		delete InPage;
-
 	}
 
 	void FVulkanResourceHeap::FreePage(FVulkanSubresourceAllocator* InPage)
@@ -1764,9 +1742,6 @@ namespace VulkanRHI
 			{
 				checkNoEntry();
 			}
-			ReleasePage(InPage);
-
-
 		}
 		else
 		{
@@ -1782,9 +1757,10 @@ namespace VulkanRHI
 				checkNoEntry();
 			}
 			check(!Pages.Find(InPage, Index));
+		}
+
 			ReleasePage(InPage);
 		}
-	}
 
 	uint64 FVulkanResourceHeap::EvictOne(FVulkanDevice& Device)
 	{
@@ -1836,6 +1812,7 @@ namespace VulkanRHI
 			Pages[LastIndex] = Allocator;
 		}
 	}
+
 	bool FVulkanResourceHeap::GetIsDefragging(FVulkanSubresourceAllocator* Allocator)
 	{
 		return Allocator->GetIsDefragging();
@@ -2215,26 +2192,6 @@ namespace VulkanRHI
 
 		ResourceTypeHeaps.AddZeroed(MemoryProperties.memoryTypeCount);
 
-		auto GetMemoryTypesFromProperties = [MemoryProperties](uint32 InTypeBits, VkMemoryPropertyFlags Properties, TArray<uint32>& OutTypeIndices)
-		{
-			// Search memtypes to find first index with those properties
-			for (uint32 i = 0; i < MemoryProperties.memoryTypeCount && InTypeBits; i++)
-			{
-				if ((InTypeBits & 1) == 1)
-				{
-					// Type is available, does it match user properties?
-					if ((MemoryProperties.memoryTypes[i].propertyFlags & Properties) == Properties)
-					{
-						OutTypeIndices.Add(i);
-					}
-				}
-				InTypeBits >>= 1;
-			}
-
-			// No memory types matched, return failure
-			return OutTypeIndices.Num() > 0;
-		};
-
 		// Upload heap. Spec requires this combination to exist.
 		{
 			uint32 TypeIndex = 0;
@@ -2279,25 +2236,23 @@ namespace VulkanRHI
 			FVulkanPageSizeBucket Bucket1 = { UINT64_MAX, 0, FVulkanPageSizeBucket::BUCKET_MASK_IMAGE | FVulkanPageSizeBucket::BUCKET_MASK_BUFFER };
 			PageSizeBuckets.Add(Bucket0);
 			PageSizeBuckets.Add(Bucket1);
-
 		}
 
 
 		// Setup main GPU heap
 		{
-			TArray<uint32> TypeIndices;
-			GetMemoryTypesFromProperties(TypeBits, VK_MEMORY_PROPERTY_DEVICE_LOCAL_BIT, TypeIndices);
-			check(TypeIndices.Num() > 0);
+			uint32 Index;
+			check(DeviceMemoryManager->GetMemoryTypeFromProperties(TypeBits, VK_MEMORY_PROPERTY_DEVICE_LOCAL_BIT, &Index) == VK_SUCCESS);
 
-			for (uint32 Index = 0; Index < MemoryProperties.memoryTypeCount; ++Index)
+			for (Index = 0; Index < MemoryProperties.memoryTypeCount; ++Index)
 			{
-				int32 HeapIndex = MemoryProperties.memoryTypes[Index].heapIndex;
-				VkDeviceSize HeapSize = MemoryProperties.memoryHeaps[HeapIndex].size;
+				const int32 HeapIndex = MemoryProperties.memoryTypes[Index].heapIndex;
+				const VkDeviceSize HeapSize = MemoryProperties.memoryHeaps[HeapIndex].size;
 				if(!ResourceTypeHeaps[Index] )
 				{
 					ResourceTypeHeaps[Index] = new FVulkanResourceHeap(this, Index);
-					ResourceTypeHeaps[Index]->bIsHostCachedSupported = ((MemoryProperties.memoryTypes[Index].propertyFlags & VK_MEMORY_PROPERTY_HOST_CACHED_BIT) == VK_MEMORY_PROPERTY_HOST_CACHED_BIT);
-					ResourceTypeHeaps[Index]->bIsLazilyAllocatedSupported = ((MemoryProperties.memoryTypes[Index].propertyFlags & VK_MEMORY_PROPERTY_LAZILY_ALLOCATED_BIT) == VK_MEMORY_PROPERTY_LAZILY_ALLOCATED_BIT);
+					ResourceTypeHeaps[Index]->bIsHostCachedSupported = VKHasAllFlags(MemoryProperties.memoryTypes[Index].propertyFlags, VK_MEMORY_PROPERTY_HOST_CACHED_BIT);
+					ResourceTypeHeaps[Index]->bIsLazilyAllocatedSupported = VKHasAllFlags(MemoryProperties.memoryTypes[Index].propertyFlags, VK_MEMORY_PROPERTY_LAZILY_ALLOCATED_BIT);
 					auto& PageSizeBuckets = ResourceTypeHeaps[Index]->PageSizeBuckets;
 
 #if PLATFORM_ANDROID
@@ -2380,7 +2335,7 @@ namespace VulkanRHI
 
 	void FMemoryManager::ReleaseFreedResources(bool bImmediately)
 	{
-		FVulkanSubresourceAllocator* BufferAllocationToRelease = nullptr;
+		TArray<FVulkanSubresourceAllocator*> BufferAllocationsToRelease;
 		{
 			FScopeLock ScopeLock(&GResourceHeapLock);
 			for (auto& FreeAllocations : FreeBufferAllocations)
@@ -2390,22 +2345,24 @@ namespace VulkanRHI
 					FVulkanSubresourceAllocator* BufferAllocation = FreeAllocations[Index];
 					if (bImmediately || BufferAllocation->FrameFreed + NUM_FRAMES_TO_WAIT_BEFORE_RELEASING_TO_OS < GFrameNumberRenderThread)
 					{
-						BufferAllocationToRelease = BufferAllocation;
+						BufferAllocationsToRelease.Add(BufferAllocation);
 						FreeAllocations.RemoveAtSwap(Index, 1, false);
-						break;
+						//break;
 					}
 				}
 			}
 		}
 
-		if (BufferAllocationToRelease)
+		for (int32 Index = 0; Index < BufferAllocationsToRelease.Num(); ++Index)
 		{
-			BufferAllocationToRelease->Destroy(GetParent());
-			FDeviceMemoryAllocation* MemoryAllocation = BufferAllocationToRelease->MemoryAllocation;
-			BufferAllocationToRelease->MemoryAllocation = 0;
-			GetParent()->GetDeviceMemoryManager().Free(MemoryAllocation);
-			delete BufferAllocationToRelease;
+			FVulkanSubresourceAllocator* BufferAllocation = BufferAllocationsToRelease[Index];
+
+			BufferAllocation->Destroy(GetParent());
+
+			FVulkanResourceHeap* Heap = ResourceTypeHeaps[BufferAllocation->MemoryTypeIndex];
+			Heap->ReleasePage(BufferAllocation);
 		}
+
 		DeviceMemoryManager->TrimMemory(bImmediately);
 	}
 
@@ -2453,7 +2410,7 @@ namespace VulkanRHI
 
 		ReleaseFreedResources(false);
 
-		int32 PrimaryHostHeap = DeviceMemoryManager->PrimaryHostHeap;
+		const int32 PrimaryHeapIndex = DeviceMemoryManager->PrimaryHeapIndex;
 		FVulkanResourceHeap* BestEvictHeap = 0;
 		uint64 BestEvictHeapSize = 0;
 
@@ -2463,7 +2420,7 @@ namespace VulkanRHI
 		uint32 Idx = 0;
 		for(FVulkanResourceHeap* Heap : ResourceTypeHeaps)
 		{
-			if(Heap->HeapIndex == PrimaryHostHeap)
+			if(Heap->HeapIndex == PrimaryHeapIndex)
 			{
 				uint64 UsedSize = Heap->UsedMemory;
 				if(CanDefragHeap(Heap) && BestDefragHeapSize < UsedSize)
@@ -2481,7 +2438,7 @@ namespace VulkanRHI
 			Idx++;
 		}
 
-		if(BestEvictHeap && ((GVulkanEvictOnePage || UpdateEvictThreshold(true)) && PrimaryHostHeap >= 0))
+		if(BestEvictHeap && ((GVulkanEvictOnePage || UpdateEvictThreshold(true)) && PrimaryHeapIndex >= 0))
 		{
 			GVulkanEvictOnePage = 0;
 			PendingEvictBytes += BestEvictHeap->EvictOne(*Device);		
@@ -2688,10 +2645,12 @@ namespace VulkanRHI
 		{
 			SubresourceAllocatorFlags &= ~VulkanAllocationFlagsMapped;
 		}
+
 		FRange FullRange;
 		FullRange.Offset = 0;
 		FullRange.Size = MaxSize;
 		FreeList.Add(FullRange);
+
 		VULKAN_FILL_TRACK_INFO(Track, __FILE__, __LINE__);
 	}
 
@@ -2715,6 +2674,7 @@ namespace VulkanRHI
 	{
 		FMemory::Memzero(MemoryUsed);
 		MaxSize = InDeviceMemoryAllocation->GetSize();
+
 		if (InDeviceMemoryAllocation->IsMapped())
 		{
 			SubresourceAllocatorFlags |= VulkanAllocationFlagsMapped;
@@ -2731,6 +2691,7 @@ namespace VulkanRHI
 
 		VULKAN_FILL_TRACK_INFO(Track, __FILE__, __LINE__);
 	}
+
 	FVulkanSubresourceAllocator::~FVulkanSubresourceAllocator()
 	{
 		if(!JoinFreeBlocks())
@@ -2749,8 +2710,6 @@ namespace VulkanRHI
 		VULKAN_FREE_TRACK_INFO(Track);
 	}
 
-
-
 	bool FMemoryManager::AllocateBufferPooled(FVulkanAllocation& OutAllocation, FVulkanEvictable* AllocationOwner, uint32 Size, VkBufferUsageFlags BufferUsageFlags, VkMemoryPropertyFlags MemoryPropertyFlags, EVulkanAllocationMetaType MetaType, const char* File, uint32 Line)
 	{
 		SCOPED_NAMED_EVENT(FResourceHeapManager_AllocateBufferPooled, FColor::Cyan);
@@ -2830,7 +2789,7 @@ namespace VulkanRHI
 		}
 
 		// New Buffer
-		uint32 BufferSize = FMath::Max(Size, BufferSizes[PoolSize]);
+		const uint32 BufferSize = FMath::Max(Size, BufferSizes[PoolSize]);
 
 		VkBuffer Buffer;
 		VkBufferCreateInfo BufferCreateInfo;
@@ -2945,10 +2904,9 @@ namespace VulkanRHI
 	//
 	bool FMemoryManager::UpdateEvictThreshold(bool bLog)
 	{
-		uint64 HostAllocated = 0;
-		uint64 HostLimit = 0;
-		DeviceMemoryManager->GetHostMemoryStatus(&HostAllocated, &HostLimit);
-		double AllocatedPercentage = 100.0 * (HostAllocated - PendingEvictBytes) / HostLimit;
+		uint64 PrimaryAllocated, PrimaryLimit;
+		DeviceMemoryManager->GetPrimaryHeapStatus(PrimaryAllocated, PrimaryLimit);
+		double AllocatedPercentage = 100.0 * (PrimaryAllocated - PendingEvictBytes) / PrimaryLimit;
 
 		double EvictionLimit = GVulkanEvictionLimitPercentage;
 		double EvictionLimitLowered = EvictionLimit * (GVulkanEvictionLimitPercentageReenableLimit / 100.f);
@@ -2958,7 +2916,7 @@ namespace VulkanRHI
 		}
 		if(bLog && GVulkanLogEvictStatus)
 		{
-			VULKAN_LOGMEMORY(TEXT("EVICT STATUS %6.2f%%/%6.2f%% :: A:%8.3fMB / E:%8.3fMB / T:%8.3fMB\n"), AllocatedPercentage, EvictionLimit, HostAllocated / (1024.f*1024.f), PendingEvictBytes/ (1024.f*1024.f), HostLimit/ (1024.f*1024.f));
+			VULKAN_LOGMEMORY(TEXT("EVICT STATUS %6.2f%%/%6.2f%% :: A:%8.3fMB / E:%8.3fMB / T:%8.3fMB\n"), AllocatedPercentage, EvictionLimit, PrimaryAllocated / (1024.f*1024.f), PendingEvictBytes/ (1024.f*1024.f), PrimaryLimit / (1024.f*1024.f));
 		}
 
 		bIsEvicting = AllocatedPercentage > EvictionLimit;
@@ -2967,18 +2925,18 @@ namespace VulkanRHI
 
 	bool FMemoryManager::AllocateImageMemory(FVulkanAllocation& OutAllocation, FVulkanEvictable* AllocationOwner, const VkMemoryRequirements& MemoryReqs, VkMemoryPropertyFlags MemoryPropertyFlags, EVulkanAllocationMetaType MetaType, bool bExternal, const char* File, uint32 Line)
 	{
-		bool bHasUnifiedMemory = DeviceMemoryManager->HasUnifiedMemory();
-		bool bCanEvict = MetaTypeCanEvict(MetaType);
+		const bool bHasUnifiedMemory = DeviceMemoryManager->HasUnifiedMemory();
+		const bool bCanEvict = MetaTypeCanEvict(MetaType);
 		if(!bHasUnifiedMemory && bCanEvict && MemoryPropertyFlags == VK_MEMORY_PROPERTY_DEVICE_LOCAL_BIT && UpdateEvictThreshold(false))
 		{
 			MemoryPropertyFlags = DeviceMemoryManager->GetEvictedMemoryProperties();
 		}
-		bool bMapped = (MemoryPropertyFlags & VK_MEMORY_PROPERTY_HOST_VISIBLE_BIT) == VK_MEMORY_PROPERTY_HOST_VISIBLE_BIT;
+		bool bMapped = VKHasAllFlags(MemoryPropertyFlags, VK_MEMORY_PROPERTY_HOST_VISIBLE_BIT);
 		uint32 TypeIndex = 0;
 
 		if (DeviceMemoryManager->GetMemoryTypeFromProperties(MemoryReqs.memoryTypeBits, MemoryPropertyFlags, &TypeIndex) != VK_SUCCESS)
 		{
-			if ((MemoryPropertyFlags & VK_MEMORY_PROPERTY_LAZILY_ALLOCATED_BIT) == VK_MEMORY_PROPERTY_LAZILY_ALLOCATED_BIT)
+			if (VKHasAllFlags(MemoryPropertyFlags, VK_MEMORY_PROPERTY_LAZILY_ALLOCATED_BIT))
 			{
 				// If lazy allocations are not supported, we can fall back to real allocations.
 				MemoryPropertyFlags = MemoryPropertyFlags & ~VK_MEMORY_PROPERTY_LAZILY_ALLOCATED_BIT;
@@ -2996,7 +2954,7 @@ namespace VulkanRHI
 
 		check(MemoryReqs.size <= (uint64)MAX_uint32);
 
-		const bool bForceSeparateAllocation = bExternal || (MemoryPropertyFlags & VK_MEMORY_PROPERTY_LAZILY_ALLOCATED_BIT) == VK_MEMORY_PROPERTY_LAZILY_ALLOCATED_BIT;
+		const bool bForceSeparateAllocation = bExternal || VKHasAllFlags(MemoryPropertyFlags, VK_MEMORY_PROPERTY_LAZILY_ALLOCATED_BIT);
 		if(!ResourceTypeHeaps[TypeIndex]->AllocateResource(OutAllocation, AllocationOwner, EType::Image, MemoryReqs.size, MemoryReqs.alignment, bMapped, bForceSeparateAllocation, MetaType, bExternal, File, Line))
 		{
 			if(GVulkanMemoryMemoryFallbackToHost)
@@ -3005,7 +2963,7 @@ namespace VulkanRHI
 			}
 
 			VERIFYVULKANRESULT(DeviceMemoryManager->GetMemoryTypeFromPropertiesExcluding(MemoryReqs.memoryTypeBits, MemoryPropertyFlags, TypeIndex, &TypeIndex));
-			bMapped = (MemoryPropertyFlags & VK_MEMORY_PROPERTY_HOST_VISIBLE_BIT) == VK_MEMORY_PROPERTY_HOST_VISIBLE_BIT;
+			bMapped = VKHasAllFlags(MemoryPropertyFlags, VK_MEMORY_PROPERTY_HOST_VISIBLE_BIT);
 			if(!ResourceTypeHeaps[TypeIndex]->AllocateResource(OutAllocation, AllocationOwner, EType::Image, MemoryReqs.size, MemoryReqs.alignment, bMapped, bForceSeparateAllocation, MetaType, bExternal, File, Line))
 			{
 				DumpMemory();
@@ -3021,16 +2979,16 @@ namespace VulkanRHI
 		SCOPED_NAMED_EVENT(FResourceHeapManager_AllocateBufferMemory, FColor::Cyan);
 		uint32 TypeIndex = 0;
 		VkResult Result = DeviceMemoryManager->GetMemoryTypeFromProperties(MemoryReqs.memoryTypeBits, MemoryPropertyFlags, &TypeIndex);
-		bool bMapped = (MemoryPropertyFlags & VK_MEMORY_PROPERTY_HOST_VISIBLE_BIT) == VK_MEMORY_PROPERTY_HOST_VISIBLE_BIT;
+		bool bMapped = VKHasAllFlags(MemoryPropertyFlags, VK_MEMORY_PROPERTY_HOST_VISIBLE_BIT);
 		if ((Result != VK_SUCCESS) || !ResourceTypeHeaps[TypeIndex])
 		{
-			if ((MemoryPropertyFlags & VK_MEMORY_PROPERTY_HOST_CACHED_BIT) == VK_MEMORY_PROPERTY_HOST_CACHED_BIT)
+			if (VKHasAllFlags(MemoryPropertyFlags, VK_MEMORY_PROPERTY_HOST_CACHED_BIT))
 			{
 				// Try non-cached flag
 				MemoryPropertyFlags = MemoryPropertyFlags & ~VK_MEMORY_PROPERTY_HOST_CACHED_BIT;
 			}
 
-			if ((MemoryPropertyFlags & VK_MEMORY_PROPERTY_LAZILY_ALLOCATED_BIT) == VK_MEMORY_PROPERTY_LAZILY_ALLOCATED_BIT)
+			if (VKHasAllFlags(MemoryPropertyFlags, VK_MEMORY_PROPERTY_LAZILY_ALLOCATED_BIT))
 			{
 				// Try non-lazy flag
 				MemoryPropertyFlags = MemoryPropertyFlags & ~VK_MEMORY_PROPERTY_LAZILY_ALLOCATED_BIT;
@@ -3136,44 +3094,44 @@ namespace VulkanRHI
 
 		const VkPhysicalDeviceMemoryProperties& MemoryProperties = DeviceMemoryManager->GetMemoryProperties();
 
-		TArray<FResourceHeapStats> Summary;
 		TArray<FResourceHeapStats> HeapSummary;
 		HeapSummary.SetNum(MemoryProperties.memoryHeapCount);
-		for(uint32 Index = 0; Index < MemoryProperties.memoryHeapCount; ++Index)
+		for(uint32 HeapIndex = 0; HeapIndex < MemoryProperties.memoryHeapCount; ++HeapIndex)
 		{
-			HeapSummary[Index].MemoryFlags = 0;
+			HeapSummary[HeapIndex].MemoryFlags = 0;
 			for (uint32 TypeIndex = 0; TypeIndex < MemoryProperties.memoryTypeCount; ++TypeIndex)
 			{
-				if(MemoryProperties.memoryTypes[TypeIndex].heapIndex == Index)
+				if (MemoryProperties.memoryTypes[TypeIndex].heapIndex == HeapIndex)
 				{
-					HeapSummary[Index].MemoryFlags |= MemoryProperties.memoryTypes[TypeIndex].propertyFlags; //since it can be different, just set to the bitwise or of all flags.
+					HeapSummary[HeapIndex].MemoryFlags |= MemoryProperties.memoryTypes[TypeIndex].propertyFlags; //since it can be different, just set to the bitwise or of all flags.
 				}
 			}
 		}
 
-		uint32 NumSmallAllocators = UE_ARRAY_COUNT(UsedBufferAllocations);
-		uint32 NumResourceHeaps = ResourceTypeHeaps.Num();
+		TArray<FResourceHeapStats> OverallSummary;
+		const uint32 NumSmallAllocators = UE_ARRAY_COUNT(UsedBufferAllocations);
+		const uint32 NumResourceHeaps = ResourceTypeHeaps.Num();
 		
-		Summary.SetNum(NumResourceHeaps + NumSmallAllocators * 2 + 1);
+		OverallSummary.SetNum(NumResourceHeaps + NumSmallAllocators * 2 + 1);
 		const uint32 SmallAllocatorsBegin = NumResourceHeaps;
 		const uint32 SmallAllocatorsEnd = NumResourceHeaps + NumSmallAllocators * 2;
 		const uint32 DedicatedAllocatorSummary = SmallAllocatorsEnd;
 
 
-		for (int32 Index = 0; Index < ResourceTypeHeaps.Num(); ++Index)
+		for (int32 TypeIndex = 0; TypeIndex < ResourceTypeHeaps.Num(); ++TypeIndex)
 		{
-			if (ResourceTypeHeaps[Index])
+			if (ResourceTypeHeaps[TypeIndex])
 			{
-				VULKAN_LOGMEMORY(TEXT("Heap %d, Memory Type Index %d"), Index, ResourceTypeHeaps[Index]->MemoryTypeIndex);
-				Summary[Index].MemoryFlags = MemoryProperties.memoryTypes[ResourceTypeHeaps[Index]->MemoryTypeIndex].propertyFlags;
-				ResourceTypeHeaps[Index]->DumpMemory(Summary[Index]);
-				uint32 MemoryTypeIndex = ResourceTypeHeaps[Index]->MemoryTypeIndex;
-				uint32 HeapIndex = MemoryProperties.memoryTypes[MemoryTypeIndex].heapIndex;
-				HeapSummary[HeapIndex] += Summary[Index];
+				const uint32 MemoryTypeIndex = ResourceTypeHeaps[TypeIndex]->MemoryTypeIndex;
+				VULKAN_LOGMEMORY(TEXT("ResourceHeap %d, Memory Type Index %d"), TypeIndex, MemoryTypeIndex);
+				OverallSummary[TypeIndex].MemoryFlags = MemoryProperties.memoryTypes[MemoryTypeIndex].propertyFlags;
+				ResourceTypeHeaps[TypeIndex]->DumpMemory(OverallSummary[TypeIndex]);
+				const uint32 HeapIndex = MemoryProperties.memoryTypes[MemoryTypeIndex].heapIndex;
+				HeapSummary[HeapIndex] += OverallSummary[TypeIndex];
 			}
 			else
 			{
-				VULKAN_LOGMEMORY(TEXT("Heap %d, NOT USED"), Index);
+				VULKAN_LOGMEMORY(TEXT("ResourceHeap %d, NOT USED"), TypeIndex);
 			}
 		}
 
@@ -3184,8 +3142,8 @@ namespace VulkanRHI
 		uint64 AllocLargeTotal = 0;
 		for (int32 PoolSizeIndex = 0; PoolSizeIndex < UE_ARRAY_COUNT(UsedBufferAllocations); PoolSizeIndex++)
 		{
-			FResourceHeapStats& StatsLocal = Summary[NumResourceHeaps + PoolSizeIndex];
-			FResourceHeapStats& StatsHost = Summary[NumResourceHeaps + NumSmallAllocators + PoolSizeIndex];
+			FResourceHeapStats& StatsLocal = OverallSummary[NumResourceHeaps + PoolSizeIndex];
+			FResourceHeapStats& StatsHost = OverallSummary[NumResourceHeaps + NumSmallAllocators + PoolSizeIndex];
 			TArray<FVulkanSubresourceAllocator*>& UsedAllocations = UsedBufferAllocations[PoolSizeIndex];
 			TArray<FVulkanSubresourceAllocator*>& FreeAllocations = FreeBufferAllocations[PoolSizeIndex];
 			if (PoolSizeIndex == (int32)EPoolSizes::SizesCount)
@@ -3201,7 +3159,7 @@ namespace VulkanRHI
 			for (int32 Index = 0; Index < FreeAllocations.Num(); ++Index)
 			{
 				FVulkanSubresourceAllocator* BA = FreeAllocations[Index];
-				if(BA->MemoryPropertyFlags & (VK_MEMORY_PROPERTY_DEVICE_LOCAL_BIT))
+				if (VKHasAnyFlags(BA->MemoryPropertyFlags, VK_MEMORY_PROPERTY_DEVICE_LOCAL_BIT))
 				{
 					StatsLocal.Pages += 1;
 					StatsLocal.BufferPages += 1;
@@ -3216,7 +3174,7 @@ namespace VulkanRHI
 					StatsHost.MemoryFlags |= BA->MemoryPropertyFlags;
 				}
 
-				uint32 HeapIndex = MemoryProperties.memoryTypes[BA->MemoryTypeIndex].heapIndex;
+				const uint32 HeapIndex = MemoryProperties.memoryTypes[BA->MemoryTypeIndex].heapIndex;
 				FResourceHeapStats& HeapStats = HeapSummary[HeapIndex];
 				HeapStats.Pages += 1;
 				HeapStats.BufferPages += 1;
@@ -3251,7 +3209,7 @@ namespace VulkanRHI
 						AllocBinnedTotal += BA->MaxSize;
 					}
 
-					if (BA->MemoryPropertyFlags & (VK_MEMORY_PROPERTY_DEVICE_LOCAL_BIT))
+					if (VKHasAnyFlags(BA->MemoryPropertyFlags, VK_MEMORY_PROPERTY_DEVICE_LOCAL_BIT))
 					{
 						StatsLocal.Pages += 1;
 						StatsLocal.BufferPages += 1;
@@ -3269,7 +3227,7 @@ namespace VulkanRHI
 						StatsHost.BufferAllocations += BA->NumSubAllocations;
 						StatsHost.MemoryFlags |= BA->MemoryPropertyFlags;
 					}
-					uint32 HeapIndex = MemoryProperties.memoryTypes[BA->MemoryTypeIndex].heapIndex;
+					const uint32 HeapIndex = MemoryProperties.memoryTypes[BA->MemoryTypeIndex].heapIndex;
 					FResourceHeapStats& HeapStats = HeapSummary[HeapIndex];
 					HeapStats.Pages += 1;
 					HeapStats.BufferPages += 1;
@@ -3293,25 +3251,27 @@ namespace VulkanRHI
 		VULKAN_LOGMEMORY(TEXT("Large Alloc Used/Max %d/%d %.2f%%"), UsedLargeTotal, AllocLargeTotal, 100.0f * AllocLargeTotal > 0 ? (float)UsedLargeTotal / (float)AllocLargeTotal : 0.0f);
 		VULKAN_LOGMEMORY(TEXT("Binned Alloc Used/Max %d/%d %.2f%%"), UsedBinnedTotal, AllocBinnedTotal, AllocBinnedTotal > 0 ? 100.0f * (float)UsedBinnedTotal / (float)AllocBinnedTotal : 0.0f);
 		{
-			FResourceHeapStats& DedicatedStats = Summary[DedicatedAllocatorSummary];
-			int32 HeapIndex = 0;
+			FResourceHeapStats& DedicatedStats = OverallSummary[DedicatedAllocatorSummary];
 
-			for(FVulkanResourceHeap* Heap : ResourceTypeHeaps)
+			for (int32 TypeIndex = 0; TypeIndex < ResourceTypeHeaps.Num(); ++TypeIndex)
 			{				
-				if (Heap)
+				FVulkanResourceHeap* ResourceHeap = ResourceTypeHeaps[TypeIndex];
+				if (ResourceHeap)
 				{
-					for (FVulkanSubresourceAllocator* Allocator : Heap->UsedDedicatedImagePages)
-					{
-						DedicatedStats.Pages++;
-						DedicatedStats.TotalMemory += Allocator->MaxSize;
-						DedicatedStats.UsedImageMemory += Allocator->UsedSize;
-						DedicatedStats.ImageAllocations += Allocator->NumSubAllocations;
-					}
-					if ((int32)HeapIndex < HeapSummary.Num())
+					FResourceHeapStats AccumulatedStats;
+					for (FVulkanSubresourceAllocator* Allocator : ResourceHeap->UsedDedicatedImagePages)
 					{
-						HeapSummary[HeapIndex] += DedicatedStats;
-						HeapIndex++;
+						AccumulatedStats.Pages++;
+						AccumulatedStats.TotalMemory += Allocator->MaxSize;
+						AccumulatedStats.UsedImageMemory += Allocator->UsedSize;
+						AccumulatedStats.ImageAllocations += Allocator->NumSubAllocations;
 					}
+
+					const uint32 MemoryTypeIndex = ResourceTypeHeaps[TypeIndex]->MemoryTypeIndex;
+					const uint32 HeapIndex = MemoryProperties.memoryTypes[MemoryTypeIndex].heapIndex;
+					HeapSummary[HeapIndex] += AccumulatedStats;
+
+					DedicatedStats += AccumulatedStats;
 				}
 			}
 		}
@@ -3321,7 +3281,7 @@ namespace VulkanRHI
 		{
 			uint64 FreeMemory = Stat.TotalMemory - Stat.UsedBufferMemory - Stat.UsedImageMemory;
 			FString HostString = GetMemoryPropertyFlagsString(Stat.MemoryFlags);
-			VULKAN_LOGMEMORY(TEXT("\t\t%-33s  |%8.2fmb / %8.2fmb / %8.2fmb / %8.2fmb | %10d %10d | %6d %6d %6d | %05x | %s"),
+			VULKAN_LOGMEMORY(TEXT("\t\t%-33s  |%8.2fmb / %8.2fmb / %11.2fmb / %11.2fmb | %10d %10d | %6d %6d %6d | %05x | %s"),
 				*Name,
 				Stat.UsedBufferMemory / (1024.f * 1024.f),
 				Stat.UsedImageMemory / (1024.f * 1024.f),
@@ -3336,19 +3296,16 @@ namespace VulkanRHI
 				*HostString);
 		};
 
-		FResourceHeapStats Total;
-		FResourceHeapStats TotalHost;
-		FResourceHeapStats TotalLocal;
+
+
 		FResourceHeapStats Staging;
 		TArray<FResourceHeapStats> DeviceHeaps;
 		Device->GetStagingManager().GetMemoryDump(Staging);
 		Device->GetDeviceMemoryManager().GetMemoryDump(DeviceHeaps);
 
-
-
 		VULKAN_LOGMEMORY(TEXT("SUMMARY"));
-		VULKAN_LOGMEMORY(TEXT("\t\tDevice Heaps                       |    Memory       Reserved    FreeMem      TotlMem |  Allocs     -         |  Allocs              | Flags | Type   "));
-#define VULKAN_LOGMEMORY_PAD TEXT("\t\t----------------------------------------------------------------------------------------------------------------------------------------")
+		VULKAN_LOGMEMORY(TEXT("\t\tDevice Heaps                       |    Memory       Reserved      FreeMem       TotalMem   |  Allocs     -         |  Allocs              | Flags | Type   "));
+#define VULKAN_LOGMEMORY_PAD TEXT("\t\t---------------------------------------------------------------------------------------------------------------------------------------------------------")
 		VULKAN_LOGMEMORY(VULKAN_LOGMEMORY_PAD);
 		for (int32 Index = 0; Index < DeviceHeaps.Num(); ++Index)
 		{
@@ -3357,16 +3314,18 @@ namespace VulkanRHI
 		}
 		VULKAN_LOGMEMORY(VULKAN_LOGMEMORY_PAD);
 		VULKAN_LOGMEMORY(VULKAN_LOGMEMORY_PAD);
-		VULKAN_LOGMEMORY(TEXT("\t\tAllocators                         |    BufMem       ImgMem      FreeMem      TotlMem |  BufAllocs  ImgAllocs |  Pages BufPgs ImgPgs | Flags | Type   "));
+		VULKAN_LOGMEMORY(TEXT("\t\tAllocators                         |    BufMem       ImgMem        FreeMem       TotalMem   |  BufAllocs  ImgAllocs |  Pages BufPgs ImgPgs | Flags | Type   "));
 		VULKAN_LOGMEMORY(VULKAN_LOGMEMORY_PAD);
 
 
-
-		for (int32 Index = 0; Index < Summary.Num(); ++Index)
+		FResourceHeapStats Total;
+		FResourceHeapStats TotalHost;
+		FResourceHeapStats TotalLocal;
+		for (int32 Index = 0; Index < OverallSummary.Num(); ++Index)
 		{
-			FResourceHeapStats& Stat = Summary[Index];
+			FResourceHeapStats& Stat = OverallSummary[Index];
 			Total += Stat;
-			if(Stat.MemoryFlags & VK_MEMORY_PROPERTY_DEVICE_LOCAL_BIT)
+			if (VKHasAnyFlags(Stat.MemoryFlags, VK_MEMORY_PROPERTY_DEVICE_LOCAL_BIT))
 			{
 				TotalLocal += Stat;
 				TotalLocal.MemoryFlags |= Stat.MemoryFlags;
@@ -3393,7 +3352,7 @@ namespace VulkanRHI
 			}
 			else
 			{
-				WriteLogLine(FString::Printf(TEXT("Heap %d"), Index), Stat);
+				WriteLogLine(FString::Printf(TEXT("ResourceHeap %d"), Index), Stat);
 			}
 		}
 		VULKAN_LOGMEMORY(VULKAN_LOGMEMORY_PAD);
@@ -3401,22 +3360,20 @@ namespace VulkanRHI
 		WriteLogLine(TEXT("TotalLocal"), TotalLocal);
 		WriteLogLine(TEXT("Total"), Total);
 		VULKAN_LOGMEMORY(VULKAN_LOGMEMORY_PAD);
-		for(int32 Index = 0; Index < HeapSummary.Num(); ++Index)
+		for (int32 HeapIndex = 0; HeapIndex < HeapSummary.Num(); ++HeapIndex)
 		{
-			FResourceHeapStats& Stat = HeapSummary[Index];
+			FResourceHeapStats& Stat = HeapSummary[HeapIndex];
 			//for the heaps, show -actual- max size, not reserved.
-			Stat.TotalMemory = MemoryProperties.memoryHeaps[Index].size;
-			WriteLogLine(FString::Printf(TEXT("Allocated Device Heap %d"), Index), Stat);
+			Stat.TotalMemory = MemoryProperties.memoryHeaps[HeapIndex].size;
+			WriteLogLine(FString::Printf(TEXT("Allocated Device Heap %d"), HeapIndex), Stat);
 		}
 		VULKAN_LOGMEMORY(VULKAN_LOGMEMORY_PAD);
-		VULKAN_LOGMEMORY(TEXT("\t\tSubsystems                         |    BufMem       ImgMem      FreeMem      TotlMem |  BufAllocs  ImgAllocs |  Pages BufPgs ImgPgs | Flags | Type   "));
+		VULKAN_LOGMEMORY(TEXT("\t\tSubsystems                         |    BufMem       ImgMem        FreeMem       TotalMem   |  BufAllocs  ImgAllocs |  Pages BufPgs ImgPgs | Flags | Type   "));
 		VULKAN_LOGMEMORY(VULKAN_LOGMEMORY_PAD);
 		WriteLogLine(TEXT("Staging"), Staging);
 		VULKAN_LOGMEMORY(VULKAN_LOGMEMORY_PAD);
 
-#define VULKAN_LOGMEMORY_PAD2 TEXT("\t\t----------------------------------------------------------------------------------------------------------------------------------------------------------------------------------------------------------------------------")
-
-
+#define VULKAN_LOGMEMORY_PAD2 TEXT("\t\t------------------------------------------------------------------------------------------------------------------------------------------------------------------------------------------------------------------------------------------------------")
 
 		VULKAN_LOGMEMORY(TEXT("\n\nSubAllocator Dump\n\n"));
 		auto WriteLogLineSubAllocator = [](const FString& Name, const FString& MemoryString, FVulkanSubresourceAllocator& Allocator)
@@ -3437,7 +3394,7 @@ namespace VulkanRHI
 			}
 			uint64 Free= TotalMemory - TotalUsed;
 			uint8 AllocatorFlags =  Allocator.GetSubresourceAllocatorFlags();
-			VULKAN_LOGMEMORY(TEXT("\t\t%-33s  | %4d %8d | %3d%% / %8.2fmb / %8.2fmb / %8.2fmb / %8.2fmb | %8.2fmb / %8.2fmb / %8.2fmb / %8.2fmb | %8.2fmb / %8.2fmb | %8.2fmb / %8.2fmb / %8.2fmb | %d %d | %s"),
+			VULKAN_LOGMEMORY(TEXT("\t\t%-45s  | %4d %8d | %3d%% / %8.2fmb / %8.2fmb / %8.2fmb / %8.2fmb | %8.2fmb / %8.2fmb / %8.2fmb / %8.2fmb | %8.2fmb / %8.2fmb | %8.2fmb / %8.2fmb / %8.2fmb | %d %d | %s"),
 				*Name,
 				FreeCount,
 				NumAllocations,
@@ -3471,7 +3428,7 @@ namespace VulkanRHI
 					Flags = MemoryProperties.memoryTypes[Allocator->MemoryTypeIndex].propertyFlags;
 				}
 				FString MemoryString = GetMemoryPropertyFlagsString(Flags);
-				FString NameId = FString::Printf(TEXT("%s [%d]"), *Name, Allocator->AllocatorIndex);
+				FString NameId = FString::Printf(TEXT("%s [%4d]"), *Name, Allocator->AllocatorIndex);
 				WriteLogLineSubAllocator(NameId, MemoryString, *Allocator);
 			}
 		};
@@ -3487,7 +3444,7 @@ namespace VulkanRHI
 				}
 				FString MemoryString = GetMemoryPropertyFlagsString(Flags);
 				VULKAN_LOGMEMORY(VULKAN_LOGMEMORY_PAD2);
-				VULKAN_LOGMEMORY(TEXT("\t\t%-33s  | %4s %8s | %4s / %10s / %10s / %10s / %10s | %10s / %10s / %10s / %10s | %10s / %10s | %10s / %10s / %10s | Mapped/Evictable |"),
+				VULKAN_LOGMEMORY(TEXT("\t\t%-45s  | %4s %8s | %4s / %10s / %10s / %10s / %10s | %10s / %10s / %10s / %10s | %10s / %10s | %10s / %10s / %10s | Mapped/Evictable |"),
 					TEXT(""),
 					TEXT("Free"),
 					TEXT("Count"),
@@ -3503,7 +3460,6 @@ namespace VulkanRHI
 					VulkanAllocationMetaTypeToString(EVulkanAllocationMetaImageRenderTarget),
 					VulkanAllocationMetaTypeToString(EVulkanAllocationMetaImageOther),
 					VulkanAllocationMetaTypeToString(EVulkanAllocationMetaBufferUAV),
-					VulkanAllocationMetaTypeToString(EVulkanAllocationMetaBufferUAV),
 					VulkanAllocationMetaTypeToString(EVulkanAllocationMetaBufferStaging),
 					VulkanAllocationMetaTypeToString(EVulkanAllocationMetaBufferOther)
 				);			
@@ -3521,7 +3477,7 @@ namespace VulkanRHI
 
 
 		VULKAN_LOGMEMORY(VULKAN_LOGMEMORY_PAD2);
-		VULKAN_LOGMEMORY(TEXT("\t\t%-33s  | %4s %8s | %4s / %10s / %10s / %10s / %10s | %10s / %10s / %10s / %10s | %10s / %10s | %10s / %10s / %10s |"),
+		VULKAN_LOGMEMORY(TEXT("\t\t%-45s  | %4s %8s | %4s / %10s / %10s / %10s / %10s | %10s / %10s / %10s / %10s | %10s / %10s | %10s / %10s / %10s |"),
 			TEXT(""),
 			TEXT("Free"),
 			TEXT("Count"),
@@ -3537,17 +3493,16 @@ namespace VulkanRHI
 			VulkanAllocationMetaTypeToString(EVulkanAllocationMetaImageRenderTarget),
 			VulkanAllocationMetaTypeToString(EVulkanAllocationMetaImageOther),
 			VulkanAllocationMetaTypeToString(EVulkanAllocationMetaBufferUAV),
-			VulkanAllocationMetaTypeToString(EVulkanAllocationMetaBufferUAV),
 			VulkanAllocationMetaTypeToString(EVulkanAllocationMetaBufferStaging),
 			VulkanAllocationMetaTypeToString(EVulkanAllocationMetaBufferOther)
 		);
-        auto PageSuffix = [this](uint32 Index)
+        auto PageSuffix = [this](uint32 TypeIndex, uint32 BucketIndex)
         {          
-    		auto& PageSizeBuckets = ResourceTypeHeaps[Index]->PageSizeBuckets;
+    		auto& PageSizeBuckets = ResourceTypeHeaps[TypeIndex]->PageSizeBuckets;
             uint32 Mask = 0;
-		    if(Index < (uint32)PageSizeBuckets.Num())
+		    if (BucketIndex < (uint32)PageSizeBuckets.Num())
 			{
-				Mask = PageSizeBuckets[Index].BucketMask;
+				Mask = PageSizeBuckets[BucketIndex].BucketMask;
 			}
 
             switch(Mask)
@@ -3560,22 +3515,16 @@ namespace VulkanRHI
             }
         };
 		VULKAN_LOGMEMORY(VULKAN_LOGMEMORY_PAD2);
-		for (int32 Index = 0; Index < ResourceTypeHeaps.Num(); ++Index)
-		{
-			if (ResourceTypeHeaps[Index])
+		for (int32 TypeIndex = 0; TypeIndex < ResourceTypeHeaps.Num(); ++TypeIndex)
 			{
-				auto& PageSizeBuckets = ResourceTypeHeaps[Index]->PageSizeBuckets;
-				uint32 SubIndex = 0;
-				for(TArray<FVulkanSubresourceAllocator*>& Foo : ResourceTypeHeaps[Index]->ActivePages)
+			if (ResourceTypeHeaps[TypeIndex])
 				{
-					if (ResourceTypeHeaps[SubIndex])
+				for (uint32 SubIndex = 0; SubIndex < FVulkanResourceHeap::MAX_BUCKETS; ++SubIndex)
 					{
-						DumpAllocatorRange(FString::Printf(TEXT("Page[%s] %d"), PageSuffix(SubIndex), Index), Foo);
-						SubIndex++;
+					TArray<FVulkanSubresourceAllocator*>& PageArray = ResourceTypeHeaps[TypeIndex]->ActivePages[SubIndex];
+					DumpAllocatorRange(FString::Printf(TEXT("Page[%s] - ResTypeHeaps[%d] ActivePages[%d]"), PageSuffix(TypeIndex, SubIndex), TypeIndex, SubIndex), PageArray);
 					}
-				}
-				DumpAllocatorRange(FString::Printf(TEXT("UsedDedicatedImagePages %d"), Index ), ResourceTypeHeaps[Index]->UsedDedicatedImagePages);
-
+				DumpAllocatorRange(FString::Printf(TEXT("UsedDedicatedImagePages - ResTypeHeaps[%d]"), TypeIndex), ResourceTypeHeaps[TypeIndex]->UsedDedicatedImagePages);
 			}
 		}
 
@@ -3594,17 +3543,16 @@ namespace VulkanRHI
 		if(bFullDump)
 		{
 			VULKAN_LOGMEMORY(VULKAN_LOGMEMORY_PAD2);
-			for (int32 Index = 0; Index < ResourceTypeHeaps.Num(); ++Index)
+			for (int32 TypeIndex = 0; TypeIndex < ResourceTypeHeaps.Num(); ++TypeIndex)
 			{
-				if (ResourceTypeHeaps[Index])
+				if (ResourceTypeHeaps[TypeIndex])
 				{
-					uint32 SubIndex = 0;
-					for (auto& PageArray : ResourceTypeHeaps[Index]->ActivePages)
+					for (uint32 SubIndex = 0; SubIndex < FVulkanResourceHeap::MAX_BUCKETS; ++SubIndex)
 					{
-						DumpAllocatorRangeContents(FString::Printf(TEXT("Page[%s] %d"), PageSuffix(Index), Index), PageArray);
-						SubIndex++;
+						TArray<FVulkanSubresourceAllocator*>& PageArray = ResourceTypeHeaps[TypeIndex]->ActivePages[SubIndex];
+						DumpAllocatorRangeContents(FString::Printf(TEXT("Page[%s] - ResTypeHeaps[%d] ActivePages[%d]"), PageSuffix(TypeIndex, SubIndex), TypeIndex, SubIndex), PageArray);
 					}
-					DumpAllocatorRangeContents(FString::Printf(TEXT("UsedDedicatedImagePages %d"), Index), ResourceTypeHeaps[Index]->UsedDedicatedImagePages);
+					DumpAllocatorRangeContents(FString::Printf(TEXT("UsedDedicatedImagePages - ResTypeHeaps[%d]"), TypeIndex), ResourceTypeHeaps[TypeIndex]->UsedDedicatedImagePages);
 
 				}
 			}
@@ -4008,7 +3956,7 @@ namespace VulkanRHI
 		TArray<SDumpHeapEntry> Entries;
 		for (FVulkanAllocationInternal& Alloc : InternalData)
 		{
-			if (Alloc.State == FVulkanAllocationInternal::EALLOCATED)
+			if (Alloc.State != FVulkanAllocationInternal::EUNUSED)
 			{
 				int32 Index = (int32)(&Alloc - &InternalData[0]);
 				SDumpHeapEntry H;
@@ -4043,9 +3991,13 @@ namespace VulkanRHI
 			{
 				check(Entry.InternalIndex >= 0);
 				FVulkanAllocationInternal& Alloc = InternalData[Entry.InternalIndex];
+
+				if  (Alloc.State == FVulkanAllocationInternal::EALLOCATED)
+				{
 				const TCHAR* TypeStr = VulkanAllocationMetaTypeToString(Alloc.MetaType);
 				const TCHAR* Name = TEXT("< ? >");
 				FString Tmp;
+
 				if(Alloc.MetaType == EVulkanAllocationMetaImageRenderTarget)
 				{
 					FVulkanEvictable* Evictable = Alloc.AllocationOwner;
@@ -4060,9 +4012,20 @@ namespace VulkanRHI
 						}
 					}
 				}
+
 				VULKAN_LOGMEMORY(TEXT("\t\t    [%08x - %08x]  | %11.3fMB | %20s   %s"), Entry.Offset, Entry.Offset + Entry.Size, Entry.Size / (1024.f * 1024.f), TypeStr, Name);
 			}
+				else
+				{
+					const TCHAR* StateStr =
+						(Alloc.State == FVulkanAllocationInternal::EFREED) ? TEXT("**Freed**") :
+						(Alloc.State == FVulkanAllocationInternal::EFREEPENDING) ? TEXT("**FreePending**") :
+						(Alloc.State == FVulkanAllocationInternal::EFREEDISCARDED) ? TEXT("**FreeDiscarded**") :
+						TEXT("< ? >");
 
+					VULKAN_LOGMEMORY(TEXT("\t\t    [%08x - %08x]  | %11.3fMB | %s"), Entry.Offset, Entry.Offset + Entry.Size, Entry.Size / (1024.f * 1024.f), StateStr);
+				}
+			}
 		}
 	}
 
@@ -4207,8 +4170,6 @@ namespace VulkanRHI
 		Device.GetDeferredDeletionQueue().EnqueueDeviceAllocation(OldAllocation);
 
 		return true;
-
-
 	}
 
 	bool FVulkanSubresourceAllocator::CanDefrag()
@@ -4225,7 +4186,6 @@ namespace VulkanRHI
 			}
 		}
 		return true;
-
 	}
 	
 	// One tick of incremental defrag. This will move max Count_ allocations from this suballocator to other SubAllocators on the same Heap
diff --git a/Engine/Source/Runtime/VulkanRHI/Private/VulkanMemory.h b/Engine/Source/Runtime/VulkanRHI/Private/VulkanMemory.h
--- a/Engine/Source/Runtime/VulkanRHI/Private/VulkanMemory.h
+++ b/Engine/Source/Runtime/VulkanRHI/Private/VulkanMemory.h
@@ -384,8 +384,7 @@ namespace VulkanRHI
 		}
 
 		bool SupportsMemoryType(VkMemoryPropertyFlags Properties) const;
-		void GetHostMemoryStatus(uint64* Allocated, uint64* Total);
-		bool IsHostMemory(uint32 MemoryTypeIndex) const;
+		void GetPrimaryHeapStatus(uint64& OutAllocated, uint64& OutLimit);
 
 		VkResult GetMemoryTypeFromProperties(uint32 TypeBits, VkMemoryPropertyFlags Properties, uint32* OutTypeIndex);
 		VkResult GetMemoryTypeFromPropertiesExcluding(uint32 TypeBits, VkMemoryPropertyFlags Properties, uint32 ExcludeTypeIndex, uint32* OutTypeIndex);
@@ -402,15 +401,16 @@ namespace VulkanRHI
 		uint64 GetTotalMemory(bool bGPU) const;
 		VkDeviceSize GetBaseHeapSize(uint32 HeapIndex) const;
 		uint32 GetHeapIndex(uint32 MemoryTypeIndex);
+
 	protected:
 		void FreeInternal(FDeviceMemoryAllocation* Allocation);
 		void TrimMemory(bool bFullTrim);
 		friend class FMemoryManager;
 		void GetMemoryDump(TArray<FResourceHeapStats>& OutDeviceHeapsStats);
 		void DumpMemory();
+		void PrintMemInfo() const;
 
-
-		double MemoryUpdateTime = 0.0;
+		double MemoryUpdateTime;
 		VkPhysicalDeviceMemoryBudgetPropertiesEXT MemoryBudget;
 		VkPhysicalDeviceMemoryProperties MemoryProperties;
 		VkDevice DeviceHandle;
@@ -422,13 +422,11 @@ namespace VulkanRHI
 
 		struct FHeapInfo
 		{
-			VkDeviceSize TotalSize;
 			VkDeviceSize UsedSize;
 			VkDeviceSize PeakSize;
 			TArray<FDeviceMemoryAllocation*> Allocations;
 
 			FHeapInfo() :
-				TotalSize(0),
 				UsedSize(0),
 				PeakSize(0)
 			{
@@ -436,11 +434,10 @@ namespace VulkanRHI
 		};
 
 		TArray<FHeapInfo> HeapInfos;
-		int32 PrimaryHostHeap; // memory usage of this heap will decide when to evict.
-
 
-		void SetupAndPrintMemInfo();
+		int32 PrimaryHeapIndex; // memory usage of this heap will decide when to evict.
 	};
+
 	struct FRange
 	{
 		uint32 Offset;
@@ -805,9 +802,6 @@ namespace VulkanRHI
 		void HandleOOM(bool bCanResume = false, VkResult Result = VK_SUCCESS, uint64 AllocationSize = 0, uint32 MemoryTypeIndex = 0);
 		bool UpdateEvictThreshold(bool bLog);
 
-
-		void* Hotfix = nullptr;
-
 	protected:
 		FDeviceMemoryManager* DeviceMemoryManager;
 		TArray<FVulkanResourceHeap*> ResourceTypeHeaps;
diff --git a/Engine/Source/Runtime/VulkanRHI/Private/VulkanRHIPrivate.h b/Engine/Source/Runtime/VulkanRHI/Private/VulkanRHIPrivate.h
--- a/Engine/Source/Runtime/VulkanRHI/Private/VulkanRHIPrivate.h
+++ b/Engine/Source/Runtime/VulkanRHI/Private/VulkanRHIPrivate.h
@@ -78,6 +78,18 @@ struct FInputAttachmentData;
 class FValidationContext;
 
 
+template<typename BitsType>
+constexpr bool VKHasAllFlags(VkFlags Flags, BitsType Contains)
+{
+	return (Flags & Contains) == Contains;
+}
+
+template<typename BitsType>
+constexpr bool VKHasAnyFlags(VkFlags Flags, BitsType Contains)
+{
+	return (Flags & Contains) != 0;
+}
+
 inline VkShaderStageFlagBits UEFrequencyToVKStageBit(EShaderFrequency InStage)
 {
 	switch (InStage)
