Fill in new EPixelFormatCapabilities for Vulkan.

diff --git a/Engine/Source/Runtime/VulkanRHI/Private/VulkanDevice.cpp b/Engine/Source/Runtime/VulkanRHI/Private/VulkanDevice.cpp
--- a/Engine/Source/Runtime/VulkanRHI/Private/VulkanDevice.cpp
+++ b/Engine/Source/Runtime/VulkanRHI/Private/VulkanDevice.cpp
@@ -569,6 +569,16 @@ void FVulkanDevice::SetupFormats()
 
 	static_assert(sizeof(VkFormat) <= sizeof(GPixelFormats[0].PlatformFormat), "PlatformFormat must be increased!");
 
+	// Create shortcuts for the possible component mappings
+	const VkComponentMapping ComponentMappingRGBA = { VK_COMPONENT_SWIZZLE_R, VK_COMPONENT_SWIZZLE_G, VK_COMPONENT_SWIZZLE_B, VK_COMPONENT_SWIZZLE_A };
+	const VkComponentMapping ComponentMappingRGB1 = { VK_COMPONENT_SWIZZLE_R, VK_COMPONENT_SWIZZLE_G, VK_COMPONENT_SWIZZLE_B, VK_COMPONENT_SWIZZLE_ONE };
+	const VkComponentMapping ComponentMappingRG01 = { VK_COMPONENT_SWIZZLE_R, VK_COMPONENT_SWIZZLE_G, VK_COMPONENT_SWIZZLE_ZERO, VK_COMPONENT_SWIZZLE_ONE };
+	const VkComponentMapping ComponentMappingR001 = { VK_COMPONENT_SWIZZLE_R, VK_COMPONENT_SWIZZLE_ZERO, VK_COMPONENT_SWIZZLE_ZERO, VK_COMPONENT_SWIZZLE_ONE };
+	const VkComponentMapping ComponentMappingRIII = { VK_COMPONENT_SWIZZLE_R, VK_COMPONENT_SWIZZLE_IDENTITY, VK_COMPONENT_SWIZZLE_IDENTITY, VK_COMPONENT_SWIZZLE_IDENTITY };
+	const VkComponentMapping ComponentMapping000R = { VK_COMPONENT_SWIZZLE_ZERO, VK_COMPONENT_SWIZZLE_ZERO, VK_COMPONENT_SWIZZLE_ZERO, VK_COMPONENT_SWIZZLE_R };
+	const VkComponentMapping ComponentMappingR000 = { VK_COMPONENT_SWIZZLE_R, VK_COMPONENT_SWIZZLE_ZERO, VK_COMPONENT_SWIZZLE_ZERO, VK_COMPONENT_SWIZZLE_ZERO };
+
+
 	// Initialize the platform pixel format map.
 	for (int32 Index = 0; Index < PF_MAX; ++Index)
 	{
@@ -577,300 +587,147 @@ void FVulkanDevice::SetupFormats()
 		GVulkanBufferFormat[Index] = VK_FORMAT_UNDEFINED;
 		
 		// Set default component mapping
-		VkComponentMapping& ComponentMapping = PixelFormatComponentMapping[Index];
-		ComponentMapping.r = VK_COMPONENT_SWIZZLE_R;
-		ComponentMapping.g = VK_COMPONENT_SWIZZLE_G;
-		ComponentMapping.b = VK_COMPONENT_SWIZZLE_B;
-		ComponentMapping.a = VK_COMPONENT_SWIZZLE_A;
+		PixelFormatComponentMapping[Index] = ComponentMappingRGBA;
 	}
 
-	// Required feature flags for color render targets
-	uint32 ColorRenderTargetRequiredFeatures =	VK_FORMAT_FEATURE_SAMPLED_IMAGE_BIT |
-												VK_FORMAT_FEATURE_COLOR_ATTACHMENT_BIT;
-
+	const EPixelFormatCapabilities ColorRenderTargetRequiredCapabilities = (EPixelFormatCapabilities::TextureSample | EPixelFormatCapabilities::RenderTarget);
 
 	// Default formats
-	MapFormatSupport(PF_B8G8R8A8, VK_FORMAT_B8G8R8A8_UNORM);
-	SetComponentMapping(PF_B8G8R8A8, VK_COMPONENT_SWIZZLE_R, VK_COMPONENT_SWIZZLE_G, VK_COMPONENT_SWIZZLE_B, VK_COMPONENT_SWIZZLE_A);
-
-	MapFormatSupport(PF_G8, VK_FORMAT_R8_UNORM);
-	SetComponentMapping(PF_G8, VK_COMPONENT_SWIZZLE_R, VK_COMPONENT_SWIZZLE_ZERO, VK_COMPONENT_SWIZZLE_ZERO, VK_COMPONENT_SWIZZLE_ONE);
-
-	MapFormatSupportWithFallback(PF_G16, VK_FORMAT_R16_UNORM, {VK_FORMAT_R16_SFLOAT});
-	SetComponentMapping(PF_G16, VK_COMPONENT_SWIZZLE_R, VK_COMPONENT_SWIZZLE_ZERO, VK_COMPONENT_SWIZZLE_ZERO, VK_COMPONENT_SWIZZLE_ONE);
-
-	MapFormatSupportWithFallback(PF_FloatRGB, ColorRenderTargetRequiredFeatures, VK_FORMAT_B10G11R11_UFLOAT_PACK32, {VK_FORMAT_R16G16B16_SFLOAT, VK_FORMAT_R16G16B16A16_SFLOAT});
-	SetComponentMapping(PF_FloatRGB, VK_COMPONENT_SWIZZLE_R, VK_COMPONENT_SWIZZLE_G, VK_COMPONENT_SWIZZLE_B, VK_COMPONENT_SWIZZLE_ONE);
-
-	MapFormatSupport(PF_FloatRGBA, VK_FORMAT_R16G16B16A16_SFLOAT, 8);
-	SetComponentMapping(PF_FloatRGBA, VK_COMPONENT_SWIZZLE_R, VK_COMPONENT_SWIZZLE_G, VK_COMPONENT_SWIZZLE_B, VK_COMPONENT_SWIZZLE_A);
+	MapFormatSupport(PF_B8G8R8A8, { VK_FORMAT_B8G8R8A8_UNORM }, ComponentMappingRGBA);
+	MapFormatSupport(PF_G8, { VK_FORMAT_R8_UNORM }, ComponentMappingR001);
+	MapFormatSupport(PF_G16, { VK_FORMAT_R16_UNORM, VK_FORMAT_R16_SFLOAT }, ComponentMappingR001);
+	MapFormatSupport(PF_FloatRGB, { VK_FORMAT_B10G11R11_UFLOAT_PACK32, VK_FORMAT_R16G16B16_SFLOAT, VK_FORMAT_R16G16B16A16_SFLOAT }, ComponentMappingRGB1, ColorRenderTargetRequiredCapabilities);
+	MapFormatSupport(PF_FloatRGBA, { VK_FORMAT_R16G16B16A16_SFLOAT }, ComponentMappingRGBA, 8);
+	MapFormatSupport(PF_ShadowDepth, { VK_FORMAT_D16_UNORM }, ComponentMappingRIII);
+	MapFormatSupport(PF_G32R32F, { VK_FORMAT_R32G32_SFLOAT }, ComponentMappingRG01, 8);  // Requirement for GPU particles
+	MapFormatSupport(PF_A32B32G32R32F, { VK_FORMAT_R32G32B32A32_SFLOAT }, ComponentMappingRGBA, 16);
+	MapFormatSupport(PF_G16R16, { VK_FORMAT_R16G16_UNORM, VK_FORMAT_R16G16_SFLOAT }, ComponentMappingRG01);
+	MapFormatSupport(PF_G16R16F, { VK_FORMAT_R16G16_SFLOAT }, ComponentMappingRG01);
+	MapFormatSupport(PF_G16R16F_FILTER, { VK_FORMAT_R16G16_SFLOAT }, ComponentMappingRG01);
+	MapFormatSupport(PF_R16_UINT, { VK_FORMAT_R16_UINT }, ComponentMappingR001);
+	MapFormatSupport(PF_R16_SINT, { VK_FORMAT_R16_SINT }, ComponentMappingR001);
+	MapFormatSupport(PF_R32_UINT, { VK_FORMAT_R32_UINT }, ComponentMappingR001);
+	MapFormatSupport(PF_R32_SINT, { VK_FORMAT_R32_SINT }, ComponentMappingR001);
+	MapFormatSupport(PF_R8_UINT, { VK_FORMAT_R8_UINT }, ComponentMappingR001);
+	MapFormatSupport(PF_D24, { VK_FORMAT_X8_D24_UNORM_PACK32, VK_FORMAT_D24_UNORM_S8_UINT, VK_FORMAT_D16_UNORM_S8_UINT, VK_FORMAT_D32_SFLOAT, VK_FORMAT_D32_SFLOAT_S8_UINT }, ComponentMappingR000);
+	MapFormatSupport(PF_R16F, { VK_FORMAT_R16_SFLOAT }, ComponentMappingR001);
+	MapFormatSupport(PF_R16F_FILTER, { VK_FORMAT_R16_SFLOAT }, ComponentMappingR001);
+	MapFormatSupport(PF_FloatR11G11B10, { VK_FORMAT_B10G11R11_UFLOAT_PACK32, VK_FORMAT_R16G16B16_SFLOAT, VK_FORMAT_R16G16B16A16_SFLOAT }, ComponentMappingRGB1, ColorRenderTargetRequiredCapabilities);
+	MapFormatSupport(PF_A2B10G10R10, { VK_FORMAT_A2B10G10R10_UNORM_PACK32 }, ComponentMappingRGBA, 4);
+	MapFormatSupport(PF_A16B16G16R16, { VK_FORMAT_R16G16B16A16_UNORM, VK_FORMAT_R16G16B16A16_SFLOAT }, ComponentMappingRGBA, 8);
+	MapFormatSupport(PF_A8, { VK_FORMAT_R8_UNORM }, ComponentMapping000R);
+	MapFormatSupport(PF_R5G6B5_UNORM, { VK_FORMAT_R5G6B5_UNORM_PACK16 }, ComponentMappingRGBA);
+	MapFormatSupport(PF_B5G5R5A1_UNORM, { VK_FORMAT_A1R5G5B5_UNORM_PACK16, VK_FORMAT_R5G5B5A1_UNORM_PACK16, VK_FORMAT_B8G8R8A8_UNORM }, ComponentMappingRGBA);
+	MapFormatSupport(PF_R8G8B8A8, { VK_FORMAT_R8G8B8A8_UNORM }, ComponentMappingRGBA);
+	MapFormatSupport(PF_R8G8B8A8_UINT, { VK_FORMAT_R8G8B8A8_UINT }, ComponentMappingRGBA);
+	MapFormatSupport(PF_R8G8B8A8_SNORM, { VK_FORMAT_R8G8B8A8_SNORM }, ComponentMappingRGBA);
+	MapFormatSupport(PF_R16G16_UINT, { VK_FORMAT_R16G16_UINT }, ComponentMappingRG01);
+	MapFormatSupport(PF_R16G16B16A16_UINT, { VK_FORMAT_R16G16B16A16_UINT }, ComponentMappingRGBA);
+	MapFormatSupport(PF_R16G16B16A16_SINT, { VK_FORMAT_R16G16B16A16_SINT }, ComponentMappingRGBA);
+	MapFormatSupport(PF_R32G32_UINT, { VK_FORMAT_R32G32_UINT }, ComponentMappingRG01);
+	MapFormatSupport(PF_R32G32B32A32_UINT, { VK_FORMAT_R32G32B32A32_UINT }, ComponentMappingRGBA);
+	MapFormatSupport(PF_R16G16B16A16_SNORM, { VK_FORMAT_R16G16B16A16_SNORM, VK_FORMAT_R16G16B16A16_SFLOAT }, ComponentMappingRGBA);
+	MapFormatSupport(PF_R16G16B16A16_UNORM, { VK_FORMAT_R16G16B16A16_UNORM, VK_FORMAT_R16G16B16A16_SFLOAT }, ComponentMappingRGBA);
+	MapFormatSupport(PF_R8G8, { VK_FORMAT_R8G8_UNORM }, ComponentMappingRG01);
+	MapFormatSupport(PF_V8U8, { VK_FORMAT_R8G8_UNORM }, ComponentMappingRG01);
+	MapFormatSupport(PF_R32_FLOAT, { VK_FORMAT_R32_SFLOAT }, ComponentMappingR001);
+	MapFormatSupport(PF_R8, { VK_FORMAT_R8_UNORM }, ComponentMappingR001);
+	MapFormatSupport(PF_G16R16_SNORM, { VK_FORMAT_R16G16_SNORM }, ComponentMappingRG01);
+	MapFormatSupport(PF_R8G8_UINT, { VK_FORMAT_R8G8_UINT }, ComponentMappingRG01);
+	MapFormatSupport(PF_R32G32B32_UINT, { VK_FORMAT_R32G32B32_UINT }, ComponentMappingRGB1);
+	MapFormatSupport(PF_R32G32B32_SINT, { VK_FORMAT_R32G32B32_SINT }, ComponentMappingRGB1);
+	MapFormatSupport(PF_R32G32B32F, { VK_FORMAT_R32G32B32_SFLOAT }, ComponentMappingRGB1);
+	MapFormatSupport(PF_R8_SINT, { VK_FORMAT_R8_SINT }, ComponentMappingR001);
 
 	if (CVarVulkanUseD24.GetValueOnAnyThread() != 0)
 	{
 		// prefer VK_FORMAT_D24_UNORM_S8_UINT
-		MapFormatSupportWithFallback(PF_DepthStencil, VK_FORMAT_D24_UNORM_S8_UINT, {VK_FORMAT_D32_SFLOAT_S8_UINT, VK_FORMAT_D16_UNORM_S8_UINT});
+		MapFormatSupport(PF_DepthStencil, { VK_FORMAT_D24_UNORM_S8_UINT, VK_FORMAT_D32_SFLOAT_S8_UINT, VK_FORMAT_D16_UNORM_S8_UINT }, ComponentMappingRIII);
 	}
 	else
 	{
 		// prefer VK_FORMAT_D32_SFLOAT_S8_UINT
-		MapFormatSupportWithFallback(PF_DepthStencil, VK_FORMAT_D32_SFLOAT_S8_UINT, {VK_FORMAT_D24_UNORM_S8_UINT, VK_FORMAT_D16_UNORM_S8_UINT});
+		MapFormatSupport(PF_DepthStencil, { VK_FORMAT_D32_SFLOAT_S8_UINT, VK_FORMAT_D24_UNORM_S8_UINT, VK_FORMAT_D16_UNORM_S8_UINT }, ComponentMappingRIII);
 	}
-	SetComponentMapping(PF_DepthStencil, VK_COMPONENT_SWIZZLE_R, VK_COMPONENT_SWIZZLE_IDENTITY, VK_COMPONENT_SWIZZLE_IDENTITY, VK_COMPONENT_SWIZZLE_IDENTITY);
-
-	MapFormatSupport(PF_ShadowDepth, VK_FORMAT_D16_UNORM);
-	SetComponentMapping(PF_ShadowDepth, VK_COMPONENT_SWIZZLE_R, VK_COMPONENT_SWIZZLE_IDENTITY, VK_COMPONENT_SWIZZLE_IDENTITY, VK_COMPONENT_SWIZZLE_IDENTITY);
-
-	// Requirement for GPU particles
-	MapFormatSupport(PF_G32R32F, VK_FORMAT_R32G32_SFLOAT, 8);
-	SetComponentMapping(PF_G32R32F, VK_COMPONENT_SWIZZLE_R, VK_COMPONENT_SWIZZLE_G, VK_COMPONENT_SWIZZLE_ZERO, VK_COMPONENT_SWIZZLE_ONE);
-
-	MapFormatSupport(PF_A32B32G32R32F, VK_FORMAT_R32G32B32A32_SFLOAT, 16);
-	SetComponentMapping(PF_A32B32G32R32F, VK_COMPONENT_SWIZZLE_R, VK_COMPONENT_SWIZZLE_G, VK_COMPONENT_SWIZZLE_B, VK_COMPONENT_SWIZZLE_A);
-
-	MapFormatSupportWithFallback(PF_G16R16, VK_FORMAT_R16G16_UNORM, {VK_FORMAT_R16G16_SFLOAT});
-	SetComponentMapping(PF_G16R16, VK_COMPONENT_SWIZZLE_R, VK_COMPONENT_SWIZZLE_G, VK_COMPONENT_SWIZZLE_ZERO, VK_COMPONENT_SWIZZLE_ONE);
-
-	MapFormatSupport(PF_G16R16F, VK_FORMAT_R16G16_SFLOAT);
-	SetComponentMapping(PF_G16R16F, VK_COMPONENT_SWIZZLE_R, VK_COMPONENT_SWIZZLE_G, VK_COMPONENT_SWIZZLE_ZERO, VK_COMPONENT_SWIZZLE_ONE);
-
-	MapFormatSupport(PF_G16R16F_FILTER, VK_FORMAT_R16G16_SFLOAT);
-	SetComponentMapping(PF_G16R16F_FILTER, VK_COMPONENT_SWIZZLE_R, VK_COMPONENT_SWIZZLE_G, VK_COMPONENT_SWIZZLE_ZERO, VK_COMPONENT_SWIZZLE_ONE);
-
-	MapFormatSupport(PF_R16_UINT, VK_FORMAT_R16_UINT);
-	SetComponentMapping(PF_R16_UINT, VK_COMPONENT_SWIZZLE_R, VK_COMPONENT_SWIZZLE_ZERO, VK_COMPONENT_SWIZZLE_ZERO, VK_COMPONENT_SWIZZLE_ONE);
-
-	MapFormatSupport(PF_R16_SINT, VK_FORMAT_R16_SINT);
-	SetComponentMapping(PF_R16_SINT, VK_COMPONENT_SWIZZLE_R, VK_COMPONENT_SWIZZLE_ZERO, VK_COMPONENT_SWIZZLE_ZERO, VK_COMPONENT_SWIZZLE_ONE);
-
-	MapFormatSupport(PF_R32_UINT, VK_FORMAT_R32_UINT);
-	SetComponentMapping(PF_R32_UINT, VK_COMPONENT_SWIZZLE_R, VK_COMPONENT_SWIZZLE_ZERO, VK_COMPONENT_SWIZZLE_ZERO, VK_COMPONENT_SWIZZLE_ONE);
-
-	MapFormatSupport(PF_R32_SINT, VK_FORMAT_R32_SINT);
-	SetComponentMapping(PF_R32_SINT, VK_COMPONENT_SWIZZLE_R, VK_COMPONENT_SWIZZLE_ZERO, VK_COMPONENT_SWIZZLE_ZERO, VK_COMPONENT_SWIZZLE_ONE);
-
-	MapFormatSupport(PF_R8_UINT, VK_FORMAT_R8_UINT);
-	SetComponentMapping(PF_R8_UINT, VK_COMPONENT_SWIZZLE_R, VK_COMPONENT_SWIZZLE_ZERO, VK_COMPONENT_SWIZZLE_ZERO, VK_COMPONENT_SWIZZLE_ONE);
-
-	MapFormatSupportWithFallback(PF_D24, VK_FORMAT_X8_D24_UNORM_PACK32, {VK_FORMAT_D24_UNORM_S8_UINT, VK_FORMAT_D16_UNORM_S8_UINT, VK_FORMAT_D32_SFLOAT, VK_FORMAT_D32_SFLOAT_S8_UINT});
-	SetComponentMapping(PF_D24, VK_COMPONENT_SWIZZLE_R, VK_COMPONENT_SWIZZLE_ZERO, VK_COMPONENT_SWIZZLE_ZERO, VK_COMPONENT_SWIZZLE_ZERO);
-
-	MapFormatSupport(PF_R16F, VK_FORMAT_R16_SFLOAT);
-	SetComponentMapping(PF_R16F, VK_COMPONENT_SWIZZLE_R, VK_COMPONENT_SWIZZLE_ZERO, VK_COMPONENT_SWIZZLE_ZERO, VK_COMPONENT_SWIZZLE_ONE);
-
-	MapFormatSupport(PF_R16F_FILTER, VK_FORMAT_R16_SFLOAT);
-	SetComponentMapping(PF_R16F_FILTER, VK_COMPONENT_SWIZZLE_R, VK_COMPONENT_SWIZZLE_ZERO, VK_COMPONENT_SWIZZLE_ZERO, VK_COMPONENT_SWIZZLE_ONE);
-
-	MapFormatSupportWithFallback(PF_FloatR11G11B10, ColorRenderTargetRequiredFeatures, VK_FORMAT_B10G11R11_UFLOAT_PACK32, {VK_FORMAT_R16G16B16_SFLOAT, VK_FORMAT_R16G16B16A16_SFLOAT});
-	SetComponentMapping(PF_FloatR11G11B10, VK_COMPONENT_SWIZZLE_R, VK_COMPONENT_SWIZZLE_G, VK_COMPONENT_SWIZZLE_B, VK_COMPONENT_SWIZZLE_ONE);
-
-	MapFormatSupport(PF_A2B10G10R10, VK_FORMAT_A2B10G10R10_UNORM_PACK32, 4);
-	SetComponentMapping(PF_A2B10G10R10, VK_COMPONENT_SWIZZLE_R, VK_COMPONENT_SWIZZLE_G, VK_COMPONENT_SWIZZLE_B, VK_COMPONENT_SWIZZLE_A);
-
-	MapFormatSupportWithFallback(PF_A16B16G16R16, VK_FORMAT_R16G16B16A16_UNORM, {VK_FORMAT_R16G16B16A16_SFLOAT});
-	GPixelFormats[PF_A16B16G16R16].BlockBytes = 8;
-	SetComponentMapping(PF_A16B16G16R16, VK_COMPONENT_SWIZZLE_R, VK_COMPONENT_SWIZZLE_G, VK_COMPONENT_SWIZZLE_B, VK_COMPONENT_SWIZZLE_A);
-
-	MapFormatSupport(PF_A8, VK_FORMAT_R8_UNORM);
-	SetComponentMapping(PF_A8, VK_COMPONENT_SWIZZLE_ZERO, VK_COMPONENT_SWIZZLE_ZERO, VK_COMPONENT_SWIZZLE_ZERO, VK_COMPONENT_SWIZZLE_R);
-
-	MapFormatSupport(PF_R5G6B5_UNORM, VK_FORMAT_R5G6B5_UNORM_PACK16);
-	SetComponentMapping(PF_R5G6B5_UNORM, VK_COMPONENT_SWIZZLE_R, VK_COMPONENT_SWIZZLE_G, VK_COMPONENT_SWIZZLE_B, VK_COMPONENT_SWIZZLE_A);
-
-	VkFormat Vulkan5551Format, FallbackFormat0, FallbackFormat1;
-	Vulkan5551Format = FVulkanPlatform::GetPlatform5551Format(FallbackFormat0, FallbackFormat1);
-
-	//use rgb555a1 on Android in order to match OpenGL GL_RGB5_A1
-	MapFormatSupportWithFallback(PF_B5G5R5A1_UNORM, Vulkan5551Format, { FallbackFormat0,  FallbackFormat1 });
-
-	SetComponentMapping(PF_B5G5R5A1_UNORM, VK_COMPONENT_SWIZZLE_R, VK_COMPONENT_SWIZZLE_G, VK_COMPONENT_SWIZZLE_B, VK_COMPONENT_SWIZZLE_A);
-
-	MapFormatSupport(PF_R8G8B8A8, VK_FORMAT_R8G8B8A8_UNORM);
-	SetComponentMapping(PF_R8G8B8A8, VK_COMPONENT_SWIZZLE_R, VK_COMPONENT_SWIZZLE_G, VK_COMPONENT_SWIZZLE_B, VK_COMPONENT_SWIZZLE_A);
-
-	MapFormatSupport(PF_R8G8B8A8_UINT, VK_FORMAT_R8G8B8A8_UINT);
-	SetComponentMapping(PF_R8G8B8A8_UINT, VK_COMPONENT_SWIZZLE_R, VK_COMPONENT_SWIZZLE_G, VK_COMPONENT_SWIZZLE_B, VK_COMPONENT_SWIZZLE_A);
-
-	MapFormatSupport(PF_R8G8B8A8_SNORM, VK_FORMAT_R8G8B8A8_SNORM);
-	SetComponentMapping(PF_R8G8B8A8_SNORM, VK_COMPONENT_SWIZZLE_R, VK_COMPONENT_SWIZZLE_G, VK_COMPONENT_SWIZZLE_B, VK_COMPONENT_SWIZZLE_A);
-
-	MapFormatSupport(PF_R16G16_UINT, VK_FORMAT_R16G16_UINT);
-	SetComponentMapping(PF_R16G16_UINT, VK_COMPONENT_SWIZZLE_R, VK_COMPONENT_SWIZZLE_G, VK_COMPONENT_SWIZZLE_ZERO, VK_COMPONENT_SWIZZLE_ONE);
-
-	MapFormatSupport(PF_R16G16B16A16_UINT, VK_FORMAT_R16G16B16A16_UINT);
-	SetComponentMapping(PF_R16G16B16A16_UINT, VK_COMPONENT_SWIZZLE_R, VK_COMPONENT_SWIZZLE_G, VK_COMPONENT_SWIZZLE_B, VK_COMPONENT_SWIZZLE_A);
-
-	MapFormatSupport(PF_R16G16B16A16_SINT, VK_FORMAT_R16G16B16A16_SINT);
-	SetComponentMapping(PF_R16G16B16A16_SINT, VK_COMPONENT_SWIZZLE_R, VK_COMPONENT_SWIZZLE_G, VK_COMPONENT_SWIZZLE_B, VK_COMPONENT_SWIZZLE_A);
-
-	MapFormatSupport(PF_R32G32_UINT, VK_FORMAT_R32G32_UINT);
-	SetComponentMapping(PF_R32G32_UINT, VK_COMPONENT_SWIZZLE_R, VK_COMPONENT_SWIZZLE_G, VK_COMPONENT_SWIZZLE_ZERO, VK_COMPONENT_SWIZZLE_ONE);
-
-	MapFormatSupport(PF_R32G32B32A32_UINT, VK_FORMAT_R32G32B32A32_UINT);
-	SetComponentMapping(PF_R32G32B32A32_UINT, VK_COMPONENT_SWIZZLE_R, VK_COMPONENT_SWIZZLE_G, VK_COMPONENT_SWIZZLE_B, VK_COMPONENT_SWIZZLE_A);
-
-	MapFormatSupportWithFallback(PF_R16G16B16A16_SNORM, VK_FORMAT_R16G16B16A16_SNORM, {VK_FORMAT_R16G16B16A16_SFLOAT});
-	SetComponentMapping(PF_R16G16B16A16_SNORM, VK_COMPONENT_SWIZZLE_R, VK_COMPONENT_SWIZZLE_G, VK_COMPONENT_SWIZZLE_B, VK_COMPONENT_SWIZZLE_A);
-
-	MapFormatSupportWithFallback(PF_R16G16B16A16_UNORM, VK_FORMAT_R16G16B16A16_UNORM, {VK_FORMAT_R16G16B16A16_SFLOAT});
-	SetComponentMapping(PF_R16G16B16A16_UNORM, VK_COMPONENT_SWIZZLE_R, VK_COMPONENT_SWIZZLE_G, VK_COMPONENT_SWIZZLE_B, VK_COMPONENT_SWIZZLE_A);
-
-	MapFormatSupport(PF_R8G8, VK_FORMAT_R8G8_UNORM);
-	SetComponentMapping(PF_R8G8, VK_COMPONENT_SWIZZLE_R, VK_COMPONENT_SWIZZLE_G, VK_COMPONENT_SWIZZLE_ZERO, VK_COMPONENT_SWIZZLE_ONE);
-
-	MapFormatSupport(PF_V8U8, VK_FORMAT_R8G8_UNORM);
-	SetComponentMapping(PF_V8U8, VK_COMPONENT_SWIZZLE_R, VK_COMPONENT_SWIZZLE_G, VK_COMPONENT_SWIZZLE_ZERO, VK_COMPONENT_SWIZZLE_ONE);
-
-	MapFormatSupport(PF_R32_FLOAT, VK_FORMAT_R32_SFLOAT);
-	SetComponentMapping(PF_R32_FLOAT, VK_COMPONENT_SWIZZLE_R, VK_COMPONENT_SWIZZLE_ZERO, VK_COMPONENT_SWIZZLE_ZERO, VK_COMPONENT_SWIZZLE_ONE);
-
-	MapFormatSupport(PF_R8, VK_FORMAT_R8_UNORM);
-	SetComponentMapping(PF_R8, VK_COMPONENT_SWIZZLE_R, VK_COMPONENT_SWIZZLE_ZERO, VK_COMPONENT_SWIZZLE_ZERO, VK_COMPONENT_SWIZZLE_ONE);
-
-	MapFormatSupport(PF_G16R16_SNORM, VK_FORMAT_R16G16_SNORM);
-	SetComponentMapping(PF_G16R16_SNORM, VK_COMPONENT_SWIZZLE_R, VK_COMPONENT_SWIZZLE_G, VK_COMPONENT_SWIZZLE_ZERO, VK_COMPONENT_SWIZZLE_ONE);
-
-	MapFormatSupport(PF_R8G8_UINT, VK_FORMAT_R8G8_UINT);
-	SetComponentMapping(PF_R8G8_UINT, VK_COMPONENT_SWIZZLE_R, VK_COMPONENT_SWIZZLE_G, VK_COMPONENT_SWIZZLE_ZERO, VK_COMPONENT_SWIZZLE_ONE);
-
-	MapFormatSupport(PF_R32G32B32_UINT, VK_FORMAT_R32G32B32_UINT);
-	SetComponentMapping(PF_R32G32B32_UINT, VK_COMPONENT_SWIZZLE_R, VK_COMPONENT_SWIZZLE_G, VK_COMPONENT_SWIZZLE_B, VK_COMPONENT_SWIZZLE_ONE);
-
-	MapFormatSupport(PF_R32G32B32_SINT, VK_FORMAT_R32G32B32_SINT);
-	SetComponentMapping(PF_R32G32B32_SINT, VK_COMPONENT_SWIZZLE_R, VK_COMPONENT_SWIZZLE_G, VK_COMPONENT_SWIZZLE_B, VK_COMPONENT_SWIZZLE_ONE);
-
-	MapFormatSupport(PF_R32G32B32F, VK_FORMAT_R32G32B32_SFLOAT);
-	SetComponentMapping(PF_R32G32B32F, VK_COMPONENT_SWIZZLE_R, VK_COMPONENT_SWIZZLE_G, VK_COMPONENT_SWIZZLE_B, VK_COMPONENT_SWIZZLE_ONE);
-
-	MapFormatSupport(PF_R8_SINT, VK_FORMAT_R8_SINT);
-	SetComponentMapping(PF_R8_SINT, VK_COMPONENT_SWIZZLE_R, VK_COMPONENT_SWIZZLE_ZERO, VK_COMPONENT_SWIZZLE_ZERO, VK_COMPONENT_SWIZZLE_ONE);
 
 	if (FVulkanPlatform::SupportsBCTextureFormats())
 	{
-		MapFormatSupport(PF_DXT1, VK_FORMAT_BC1_RGB_UNORM_BLOCK);	// Also what OpenGL expects (RGBA instead RGB, but not SRGB)
-		SetComponentMapping(PF_DXT1, VK_COMPONENT_SWIZZLE_R, VK_COMPONENT_SWIZZLE_G, VK_COMPONENT_SWIZZLE_B, VK_COMPONENT_SWIZZLE_ONE);
-
-		MapFormatSupport(PF_DXT3, VK_FORMAT_BC2_UNORM_BLOCK);
-		SetComponentMapping(PF_DXT3, VK_COMPONENT_SWIZZLE_R, VK_COMPONENT_SWIZZLE_G, VK_COMPONENT_SWIZZLE_B, VK_COMPONENT_SWIZZLE_A);
-
-		MapFormatSupport(PF_DXT5, VK_FORMAT_BC3_UNORM_BLOCK);
-		SetComponentMapping(PF_DXT5, VK_COMPONENT_SWIZZLE_R, VK_COMPONENT_SWIZZLE_G, VK_COMPONENT_SWIZZLE_B, VK_COMPONENT_SWIZZLE_A);
-
-		MapFormatSupport(PF_BC4, VK_FORMAT_BC4_UNORM_BLOCK);
-		SetComponentMapping(PF_BC4, VK_COMPONENT_SWIZZLE_R, VK_COMPONENT_SWIZZLE_G, VK_COMPONENT_SWIZZLE_B, VK_COMPONENT_SWIZZLE_A);
-
-		MapFormatSupport(PF_BC5, VK_FORMAT_BC5_UNORM_BLOCK);
-		SetComponentMapping(PF_BC5, VK_COMPONENT_SWIZZLE_R, VK_COMPONENT_SWIZZLE_G, VK_COMPONENT_SWIZZLE_B, VK_COMPONENT_SWIZZLE_A);
-
-		MapFormatSupport(PF_BC6H, VK_FORMAT_BC6H_UFLOAT_BLOCK);
-		SetComponentMapping(PF_BC6H, VK_COMPONENT_SWIZZLE_R, VK_COMPONENT_SWIZZLE_G, VK_COMPONENT_SWIZZLE_B, VK_COMPONENT_SWIZZLE_A);
-
-		MapFormatSupport(PF_BC7, VK_FORMAT_BC7_UNORM_BLOCK);
-		SetComponentMapping(PF_BC7, VK_COMPONENT_SWIZZLE_R, VK_COMPONENT_SWIZZLE_G, VK_COMPONENT_SWIZZLE_B, VK_COMPONENT_SWIZZLE_A);
+		MapFormatSupport(PF_DXT1, { VK_FORMAT_BC1_RGB_UNORM_BLOCK }, ComponentMappingRGB1);	// Also what OpenGL expects (RGBA instead RGB, but not SRGB)
+		MapFormatSupport(PF_DXT3, { VK_FORMAT_BC2_UNORM_BLOCK },     ComponentMappingRGBA);
+		MapFormatSupport(PF_DXT5, { VK_FORMAT_BC3_UNORM_BLOCK },     ComponentMappingRGBA);
+		MapFormatSupport(PF_BC4,  { VK_FORMAT_BC4_UNORM_BLOCK },     ComponentMappingRGBA);
+		MapFormatSupport(PF_BC5,  { VK_FORMAT_BC5_UNORM_BLOCK },     ComponentMappingRGBA);
+		MapFormatSupport(PF_BC6H, { VK_FORMAT_BC6H_UFLOAT_BLOCK },   ComponentMappingRGBA);
+		MapFormatSupport(PF_BC7,  { VK_FORMAT_BC7_UNORM_BLOCK },     ComponentMappingRGBA);
 	}
 
 	if (FVulkanPlatform::SupportsASTCTextureFormats())
 	{
-		MapFormatSupport(PF_ASTC_4x4, VK_FORMAT_ASTC_4x4_UNORM_BLOCK);
-		if (GPixelFormats[PF_ASTC_4x4].Supported)
-		{
-			SetComponentMapping(PF_ASTC_4x4, VK_COMPONENT_SWIZZLE_R, VK_COMPONENT_SWIZZLE_G, VK_COMPONENT_SWIZZLE_B, VK_COMPONENT_SWIZZLE_A);
-		}
-
-		MapFormatSupport(PF_ASTC_6x6, VK_FORMAT_ASTC_6x6_UNORM_BLOCK);
-		if (GPixelFormats[PF_ASTC_6x6].Supported)
-		{
-			SetComponentMapping(PF_ASTC_6x6, VK_COMPONENT_SWIZZLE_R, VK_COMPONENT_SWIZZLE_G, VK_COMPONENT_SWIZZLE_B, VK_COMPONENT_SWIZZLE_A);
-		}
-
-		MapFormatSupport(PF_ASTC_8x8, VK_FORMAT_ASTC_8x8_UNORM_BLOCK);
-		if (GPixelFormats[PF_ASTC_8x8].Supported)
-		{
-			SetComponentMapping(PF_ASTC_8x8, VK_COMPONENT_SWIZZLE_R, VK_COMPONENT_SWIZZLE_G, VK_COMPONENT_SWIZZLE_B, VK_COMPONENT_SWIZZLE_A);
-		}
-
-		MapFormatSupport(PF_ASTC_10x10, VK_FORMAT_ASTC_10x10_UNORM_BLOCK);
-		if (GPixelFormats[PF_ASTC_10x10].Supported)
-		{
-			SetComponentMapping(PF_ASTC_10x10, VK_COMPONENT_SWIZZLE_R, VK_COMPONENT_SWIZZLE_G, VK_COMPONENT_SWIZZLE_B, VK_COMPONENT_SWIZZLE_A);
-		}
-
-		MapFormatSupport(PF_ASTC_12x12, VK_FORMAT_ASTC_12x12_UNORM_BLOCK);
-		if (GPixelFormats[PF_ASTC_12x12].Supported)
-		{
-			SetComponentMapping(PF_ASTC_12x12, VK_COMPONENT_SWIZZLE_R, VK_COMPONENT_SWIZZLE_G, VK_COMPONENT_SWIZZLE_B, VK_COMPONENT_SWIZZLE_A);
-		}
+		MapFormatSupport(PF_ASTC_4x4,   { VK_FORMAT_ASTC_4x4_UNORM_BLOCK },   ComponentMappingRGBA);
+		MapFormatSupport(PF_ASTC_6x6,   { VK_FORMAT_ASTC_6x6_UNORM_BLOCK },   ComponentMappingRGBA);
+		MapFormatSupport(PF_ASTC_8x8,   { VK_FORMAT_ASTC_8x8_UNORM_BLOCK },   ComponentMappingRGBA);
+		MapFormatSupport(PF_ASTC_10x10, { VK_FORMAT_ASTC_10x10_UNORM_BLOCK }, ComponentMappingRGBA);
+		MapFormatSupport(PF_ASTC_12x12, { VK_FORMAT_ASTC_12x12_UNORM_BLOCK }, ComponentMappingRGBA);
 	}
 
 	if (FVulkanPlatform::SupportsETC2TextureFormats())
 	{
-		MapFormatSupport(PF_ETC2_RGB, VK_FORMAT_ETC2_R8G8B8_UNORM_BLOCK);
-		if (GPixelFormats[PF_ETC2_RGB].Supported)
-		{
-			SetComponentMapping(PF_ETC2_RGB, VK_COMPONENT_SWIZZLE_R, VK_COMPONENT_SWIZZLE_G, VK_COMPONENT_SWIZZLE_B, VK_COMPONENT_SWIZZLE_ONE);
-		}
-
-		MapFormatSupport(PF_ETC2_RGBA, VK_FORMAT_ETC2_R8G8B8A8_UNORM_BLOCK);
-		if (GPixelFormats[PF_ETC2_RGB].Supported)
-		{
-			SetComponentMapping(PF_ETC2_RGBA, VK_COMPONENT_SWIZZLE_R, VK_COMPONENT_SWIZZLE_G, VK_COMPONENT_SWIZZLE_B, VK_COMPONENT_SWIZZLE_A);
-		}
-
-		MapFormatSupport(PF_ETC2_R11_EAC, VK_FORMAT_EAC_R11_UNORM_BLOCK);
-		if (GPixelFormats[PF_ETC2_R11_EAC].Supported)
-		{
-			SetComponentMapping(PF_ETC2_R11_EAC, VK_COMPONENT_SWIZZLE_R, VK_COMPONENT_SWIZZLE_ZERO, VK_COMPONENT_SWIZZLE_ZERO, VK_COMPONENT_SWIZZLE_ONE);
-		}
-
-		MapFormatSupport(PF_ETC2_RG11_EAC, VK_FORMAT_EAC_R11G11_UNORM_BLOCK);
-		if (GPixelFormats[PF_ETC2_RG11_EAC].Supported)
-		{
-			SetComponentMapping(PF_ETC2_RG11_EAC, VK_COMPONENT_SWIZZLE_R, VK_COMPONENT_SWIZZLE_G, VK_COMPONENT_SWIZZLE_ZERO, VK_COMPONENT_SWIZZLE_ONE);
-		}
+		MapFormatSupport(PF_ETC2_RGB,      { VK_FORMAT_ETC2_R8G8B8_UNORM_BLOCK },   ComponentMappingRGB1);
+		MapFormatSupport(PF_ETC2_RGBA,     { VK_FORMAT_ETC2_R8G8B8A8_UNORM_BLOCK }, ComponentMappingRGBA);
+		MapFormatSupport(PF_ETC2_R11_EAC,  { VK_FORMAT_EAC_R11_UNORM_BLOCK },       ComponentMappingR001);
+		MapFormatSupport(PF_ETC2_RG11_EAC, { VK_FORMAT_EAC_R11G11_UNORM_BLOCK },    ComponentMappingRG01);
 	}
 
-	if(GetOptionalExtensions().HasEXTASTCDecodeMode)
+	if (GetOptionalExtensions().HasEXTASTCDecodeMode)
 	{
-		MapFormatSupport(PF_ASTC_4x4_HDR, VK_FORMAT_ASTC_4x4_SFLOAT_BLOCK_EXT);
-		if (GPixelFormats[PF_ASTC_4x4_HDR].Supported)
-		{
-			SetComponentMapping(PF_ASTC_4x4_HDR, VK_COMPONENT_SWIZZLE_R, VK_COMPONENT_SWIZZLE_G, VK_COMPONENT_SWIZZLE_B, VK_COMPONENT_SWIZZLE_A);
-		}
-
-		MapFormatSupport(PF_ASTC_6x6_HDR, VK_FORMAT_ASTC_6x6_SFLOAT_BLOCK_EXT);
-		if (GPixelFormats[PF_ASTC_6x6_HDR].Supported)
-		{
-			SetComponentMapping(PF_ASTC_6x6_HDR, VK_COMPONENT_SWIZZLE_R, VK_COMPONENT_SWIZZLE_G, VK_COMPONENT_SWIZZLE_B, VK_COMPONENT_SWIZZLE_A);
-		}
-
-		MapFormatSupport(PF_ASTC_8x8_HDR, VK_FORMAT_ASTC_8x8_SFLOAT_BLOCK_EXT);
-		if (GPixelFormats[PF_ASTC_8x8_HDR].Supported)
-		{
-			SetComponentMapping(PF_ASTC_8x8_HDR, VK_COMPONENT_SWIZZLE_R, VK_COMPONENT_SWIZZLE_G, VK_COMPONENT_SWIZZLE_B, VK_COMPONENT_SWIZZLE_A);
-		}
+		MapFormatSupport(PF_ASTC_4x4_HDR,   { VK_FORMAT_ASTC_4x4_SFLOAT_BLOCK_EXT },   ComponentMappingRGBA);
+		MapFormatSupport(PF_ASTC_6x6_HDR,   { VK_FORMAT_ASTC_6x6_SFLOAT_BLOCK_EXT },   ComponentMappingRGBA);
+		MapFormatSupport(PF_ASTC_8x8_HDR,   { VK_FORMAT_ASTC_8x8_SFLOAT_BLOCK_EXT },   ComponentMappingRGBA);
+		MapFormatSupport(PF_ASTC_10x10_HDR, { VK_FORMAT_ASTC_10x10_SFLOAT_BLOCK_EXT }, ComponentMappingRGBA);
+		MapFormatSupport(PF_ASTC_12x12_HDR, { VK_FORMAT_ASTC_12x12_SFLOAT_BLOCK_EXT }, ComponentMappingRGBA);
+	}
 
-		MapFormatSupport(PF_ASTC_10x10_HDR, VK_FORMAT_ASTC_10x10_SFLOAT_BLOCK_EXT);
-		if (GPixelFormats[PF_ASTC_10x10_HDR].Supported)
-		{
-			SetComponentMapping(PF_ASTC_10x10_HDR, VK_COMPONENT_SWIZZLE_R, VK_COMPONENT_SWIZZLE_G, VK_COMPONENT_SWIZZLE_B, VK_COMPONENT_SWIZZLE_A);
-		}
 
-		MapFormatSupport(PF_ASTC_12x12_HDR, VK_FORMAT_ASTC_12x12_SFLOAT_BLOCK_EXT);
-		if (GPixelFormats[PF_ASTC_12x12_HDR].Supported)
+	// Verify available Vertex Formats
+	{
+		static_assert(VET_None == 0, "Change loop below to skip VET_None");
+		for (int32 VETIndex = (int32)VET_None + 1; VETIndex < VET_MAX; ++VETIndex)
 		{
-			SetComponentMapping(PF_ASTC_12x12_HDR, VK_COMPONENT_SWIZZLE_R, VK_COMPONENT_SWIZZLE_G, VK_COMPONENT_SWIZZLE_B, VK_COMPONENT_SWIZZLE_A);
+			const EVertexElementType UEType = (EVertexElementType)VETIndex;
+			const VkFormat VulkanFormat = UEToVkBufferFormat(UEType);
+			const VkFormatProperties& VertexFormatProperties = GetFormatProperties(VulkanFormat);
+			if (VertexFormatProperties.bufferFeatures == 0)
+			{
+				UE_LOG(LogVulkanRHI, Warning, TEXT("EVertexElementType(%d) is not supported with VkFormat %d"), (int32)UEType, (int32)VulkanFormat);
+			}
 		}
 	}
-	// Verify available Vertex Formats
-	static_assert(VET_None == 0, "Change loop below to skip VET_None");
-	for (int32 Index = (int32)VET_None + 1; Index < VET_MAX; ++Index)
+
+#if !(UE_BUILD_SHIPPING || UE_BUILD_TEST)
+	// Print the resulting pixel format support
+	if (FParse::Param(FCommandLine::Get(), TEXT("PrintVulkanPixelFormatMappings")))
 	{
-		EVertexElementType UEType = (EVertexElementType)Index;
-		VkFormat VulkanFormat = UEToVkBufferFormat(UEType);
-		if (!IsBufferFormatSupported(VulkanFormat))
+		UE_LOG(LogVulkanRHI, Warning, TEXT("Pixel Format Mappings for Vulkan:"));
+		UE_LOG(LogVulkanRHI, Warning, TEXT("%24s | Capabilities | VulkanFormat | BlockBytes | Components | ComponentMapping | BufferFormat"), TEXT("PixelFormatName"));
+		for (int32 PixelFormatIndex = 0; PixelFormatIndex < PF_MAX; ++PixelFormatIndex)
 		{
-			UE_LOG(LogVulkanRHI, Warning, TEXT("EVertexFormat(%d) is not supported with Vk format %d"), (int32)UEType, (int32)VulkanFormat);
+			if (GPixelFormats[PixelFormatIndex].Supported)
+			{
+				const VkComponentMapping& ComponentMapping = PixelFormatComponentMapping[PixelFormatIndex];
+
+				UE_LOG(LogVulkanRHI, Warning, TEXT("%24s |   0x%08X | %12d | %10d | %10d | %10d,%d,%d,%d | %12d"),
+					GPixelFormats[PixelFormatIndex].Name,
+					(uint32)GPixelFormats[PixelFormatIndex].Capabilities,
+					GPixelFormats[PixelFormatIndex].PlatformFormat,
+					GPixelFormats[PixelFormatIndex].BlockBytes,
+					GPixelFormats[PixelFormatIndex].NumComponents,
+					ComponentMapping.r, ComponentMapping.g, ComponentMapping.b, ComponentMapping.a,
+					(int32)GVulkanBufferFormat[PixelFormatIndex]
+				);
+			}
 		}
 	}
+#endif
 }
 
 #if VULKAN_SUPPORTS_COLOR_CONVERSIONS
@@ -892,61 +749,185 @@ VkSamplerYcbcrConversion FVulkanDevice::CreateSamplerColorConversion(const VkSam
 }
 #endif
 
-void FVulkanDevice::MapFormatSupport(EPixelFormat UEFormat, VkFormat VulkanFormat)
+
+const VkFormatProperties& FVulkanDevice::GetFormatProperties(VkFormat InFormat)
 {
-	MapFormatSupportWithFallback(UEFormat, VK_FORMAT_FEATURE_SAMPLED_IMAGE_BIT, VulkanFormat, TArrayView<VkFormat>());
+	if (InFormat >= 0 && InFormat < VK_FORMAT_RANGE_SIZE)
+	{
+		return FormatProperties[InFormat];
+	}
+
+	// Check for extension formats
+	const VkFormatProperties* FoundProperties = ExtensionFormatProperties.Find(InFormat);
+	if (FoundProperties)
+	{
+		return *FoundProperties;
+	}
+
+	// Add it for faster caching next time
+	VkFormatProperties& NewProperties = ExtensionFormatProperties.Add(InFormat);
+	FMemory::Memzero(NewProperties);
+	VulkanRHI::vkGetPhysicalDeviceFormatProperties(Gpu, InFormat, &NewProperties);
+	return NewProperties;
 }
 
-void FVulkanDevice::MapFormatSupportWithFallback(EPixelFormat UEFormat, uint32 RequiredTextureFeatures, VkFormat VulkanFormat, TArrayView<const VkFormat> FallbackTextureFormats)
+void FVulkanDevice::MapBufferFormatSupport(FPixelFormatInfo& PixelFormatInfo, EPixelFormat UEFormat, VkFormat VulkanFormat)
 {
-	VkFormat SupportedTextureFormat = IsTextureFormatSupported(VulkanFormat, RequiredTextureFeatures) ? VulkanFormat : VK_FORMAT_UNDEFINED;
-	VkFormat SupportedBufferFormat = IsBufferFormatSupported(VulkanFormat) ? VulkanFormat : VK_FORMAT_UNDEFINED;
-	
-	FPixelFormatInfo& FormatInfo = GPixelFormats[UEFormat];
-	// at this point we don't know if high level code will use this pixel format for buffers or textures
-	FormatInfo.Supported = (SupportedTextureFormat != VK_FORMAT_UNDEFINED || SupportedBufferFormat!= VK_FORMAT_UNDEFINED);
-	FormatInfo.PlatformFormat = SupportedTextureFormat;
-	
-	GVulkanBufferFormat[UEFormat] = SupportedBufferFormat;
-		
-	if (SupportedTextureFormat == VK_FORMAT_UNDEFINED)
+	check(GVulkanBufferFormat[UEFormat] == VK_FORMAT_UNDEFINED);
+
+	const VkFormatProperties& LocalFormatProperties = GetFormatProperties(VulkanFormat);
+	EPixelFormatCapabilities Capabilities = EPixelFormatCapabilities::None;
+
+	auto ConvertBufferCap = [&Capabilities, &LocalFormatProperties](EPixelFormatCapabilities UnrealCap, VkFormatFeatureFlags InFlag)
+	{
+		const bool HasBufferFeature = VKHasAllFlags(LocalFormatProperties.bufferFeatures, InFlag);
+		if (HasBufferFeature)
+		{
+			EnumAddFlags(Capabilities, UnrealCap);
+		}
+
+		// Make sure we aren't looking in the wrong place for a bit
+		check(!VKHasAnyFlags(LocalFormatProperties.linearTilingFeatures, InFlag));
+		check(!VKHasAnyFlags(LocalFormatProperties.optimalTilingFeatures, InFlag));
+	};
+
+	// Check for buffer caps, use the first one with any caps
+	if (LocalFormatProperties.bufferFeatures != 0)
+	{
+		EnumAddFlags(Capabilities, EPixelFormatCapabilities::Buffer);
+
+		ConvertBufferCap(EPixelFormatCapabilities::VertexBuffer, VK_FORMAT_FEATURE_VERTEX_BUFFER_BIT);
+		ConvertBufferCap(EPixelFormatCapabilities::BufferLoad, VK_FORMAT_FEATURE_UNIFORM_TEXEL_BUFFER_BIT);
+		ConvertBufferCap(EPixelFormatCapabilities::BufferStore, VK_FORMAT_FEATURE_STORAGE_TEXEL_BUFFER_BIT);
+		ConvertBufferCap(EPixelFormatCapabilities::BufferAtomics, VK_FORMAT_FEATURE_STORAGE_TEXEL_BUFFER_ATOMIC_BIT);
+
+		// Vulkan index buffers aren't tied to formats, so any 16 or 32 bit UINT format with a single component will do...
+		// But because we can't check for uint vs float, hardcode supported formats for now
+		if (EnumHasAllFlags(Capabilities, (EPixelFormatCapabilities::BufferLoad | EPixelFormatCapabilities::BufferStore)) &&
+			((VulkanFormat == VK_FORMAT_R16_UINT) || (VulkanFormat == VK_FORMAT_R32_UINT)))
+		{
+			EnumAddFlags(Capabilities, EPixelFormatCapabilities::IndexBuffer);
+		}
+
+		GVulkanBufferFormat[UEFormat] = VulkanFormat;
+		PixelFormatInfo.Capabilities |= Capabilities;
+	}
+}
+
+void FVulkanDevice::MapImageFormatSupport(FPixelFormatInfo& PixelFormatInfo, const TArrayView<const VkFormat>& PrioritizedFormats, EPixelFormatCapabilities RequiredCapabilities)
+{
+	// Query for MipMap support with typical parameters
+	auto SupportsMipMap = [this](VkFormat InFormat)
 	{
-		for (int32 Idx = 0; Idx < FallbackTextureFormats.Num(); ++Idx)
+		VkImageFormatProperties ImageFormatProperties;
+		VkResult RetVal = VulkanRHI::vkGetPhysicalDeviceImageFormatProperties(Gpu, InFormat, VK_IMAGE_TYPE_2D, VK_IMAGE_TILING_OPTIMAL, VK_IMAGE_USAGE_SAMPLED_BIT, 0, &ImageFormatProperties);
+		return (RetVal == VK_SUCCESS) && (ImageFormatProperties.maxMipLevels > 1);
+	};
+
+	EPixelFormatCapabilities Capabilities = EPixelFormatCapabilities::None;
+	auto ConvertImageCap = [&Capabilities](const VkFormatProperties& InFormatProperties, EPixelFormatCapabilities UnrealCap, VkFormatFeatureFlags InFlag)
+	{
+		// Do not distinguish between Linear and Optimal for now.
+		const bool HasImageFeature = VKHasAllFlags(InFormatProperties.linearTilingFeatures, InFlag) || VKHasAllFlags(InFormatProperties.optimalTilingFeatures, InFlag);
+		if (HasImageFeature)
+		{
+			EnumAddFlags(Capabilities, UnrealCap);
+		}
+
+		// Make sure we aren't looking in the wrong place for a bit
+		check(!VKHasAnyFlags(InFormatProperties.bufferFeatures, InFlag));
+	};
+
+	// Go through the PrioritizedFormats and use the first one that meets RequiredCapabilities
+	for (int32 FormatIndex = 0; FormatIndex < PrioritizedFormats.Num(); ++FormatIndex)
+	{
+		Capabilities = EPixelFormatCapabilities::None;
+
+		const VkFormat VulkanFormat = PrioritizedFormats[FormatIndex];
+		const VkFormatProperties& LocalFormatProperties = GetFormatProperties(VulkanFormat);
+
+		// Check for individual texture caps
+		ConvertImageCap(LocalFormatProperties, EPixelFormatCapabilities::AnyTexture | EPixelFormatCapabilities::TextureSample | EPixelFormatCapabilities::TextureLoad, VK_FORMAT_FEATURE_SAMPLED_IMAGE_BIT);
+		ConvertImageCap(LocalFormatProperties, EPixelFormatCapabilities::DepthStencil, VK_FORMAT_FEATURE_DEPTH_STENCIL_ATTACHMENT_BIT);
+		ConvertImageCap(LocalFormatProperties, EPixelFormatCapabilities::RenderTarget, VK_FORMAT_FEATURE_COLOR_ATTACHMENT_BIT);
+		ConvertImageCap(LocalFormatProperties, EPixelFormatCapabilities::TextureBlendable, VK_FORMAT_FEATURE_COLOR_ATTACHMENT_BLEND_BIT);
+		ConvertImageCap(LocalFormatProperties, EPixelFormatCapabilities::UAV | EPixelFormatCapabilities::TypedUAVLoad | EPixelFormatCapabilities::TypedUAVStore, VK_FORMAT_FEATURE_STORAGE_IMAGE_BIT);
+
+		ConvertImageCap(LocalFormatProperties, EPixelFormatCapabilities::TextureAtomics, VK_FORMAT_FEATURE_STORAGE_IMAGE_ATOMIC_BIT);
+
+		if (EnumHasAllFlags(Capabilities, EPixelFormatCapabilities::AnyTexture))
 		{
-			VkFormat FallbackTextureFormat = FallbackTextureFormats[Idx];
-			if (IsTextureFormatSupported(FallbackTextureFormat, RequiredTextureFeatures))
+			// We support gather, but some of our shaders assume offsets so check against features
+			if (GetPhysicalFeatures().shaderImageGatherExtended)
 			{
-				SupportedTextureFormat = FallbackTextureFormat;
-				FormatInfo.PlatformFormat = FallbackTextureFormat;
-				FormatInfo.Supported = true;
-								
-				UE_LOG(LogVulkanRHI, Display, TEXT("MapFormatSupport: %s (images) is not supported with Vk format %d, falling back to Vk format %d"), FormatInfo.Name, (int32)VulkanFormat, (int32)FallbackTextureFormat);
+				EnumAddFlags(Capabilities, EPixelFormatCapabilities::TextureGather);
 			}
+
+			if (SupportsMipMap(VulkanFormat))
+			{
+				EnumAddFlags(Capabilities, EPixelFormatCapabilities::TextureMipmaps);
+			}
+		}
+
+		if (EnumHasAllFlags(Capabilities, RequiredCapabilities))
+		{
+			PixelFormatInfo.PlatformFormat = VulkanFormat;
+			PixelFormatInfo.Capabilities |= Capabilities;
+
+			if (FormatIndex > 0)
+			{
+				UE_LOG(LogVulkanRHI, Display, TEXT("MapImageFormatSupport: %s is not supported with VkFormat %d, falling back to VkFormat %d"), PixelFormatInfo.Name, (int32)PrioritizedFormats[0], (int32)PrioritizedFormats[FormatIndex]);
+			}
+
+			break;
 		}
 	}
-			
-	if (!FormatInfo.Supported)
+}
+
+// Minimum capabilities required for a Vulkan format to be considered as supported
+static constexpr EPixelFormatCapabilities kDefaultTextureCapabilities = EPixelFormatCapabilities::TextureSample;
+// Passthrough to specify we want to keep the initial BlockBytes value set in the PixelFormat
+static constexpr int32 kDefaultBlockBytes = -1;
+
+void FVulkanDevice::MapFormatSupport(EPixelFormat UEFormat, std::initializer_list<VkFormat> InPrioritizedFormats, const VkComponentMapping& ComponentMapping, EPixelFormatCapabilities RequiredCapabilities, int32 BlockBytes)
+{
+	TArrayView<const VkFormat> PrioritizedFormats = MakeArrayView(InPrioritizedFormats);
+	FPixelFormatInfo& PixelFormatInfo = GPixelFormats[UEFormat];
+
+	check(PrioritizedFormats.Num() > 0);
+	check(!PixelFormatInfo.Supported);
+	check(PixelFormatInfo.Capabilities == EPixelFormatCapabilities::None);
+
+	MapBufferFormatSupport(PixelFormatInfo, UEFormat, PrioritizedFormats[0]);
+	MapImageFormatSupport(PixelFormatInfo, PrioritizedFormats, RequiredCapabilities);
+
+	// Flag the pixel format as supported if we can do anything with it
+	PixelFormatInfo.Supported = EnumHasAllFlags(PixelFormatInfo.Capabilities, RequiredCapabilities) || EnumHasAnyFlags(PixelFormatInfo.Capabilities, EPixelFormatCapabilities::Buffer);
+	if (PixelFormatInfo.Supported)
 	{
-		UE_LOG(LogVulkanRHI, Error, TEXT("MapFormatSupport: %s is not supported with Vk format %d"), FormatInfo.Name, (int32)VulkanFormat);
+		PixelFormatComponentMapping[UEFormat] = ComponentMapping;
+		if (BlockBytes > 0)
+		{
+			PixelFormatInfo.BlockBytes = BlockBytes;
+		}
+	}
+	else
+	{
+		UE_LOG(LogVulkanRHI, Error, TEXT("MapFormatSupport: %s is not supported with VkFormat %d"), PixelFormatInfo.Name, (int32)PrioritizedFormats[0]);
 	}
 }
 
-void FVulkanDevice::SetComponentMapping(EPixelFormat UEFormat, VkComponentSwizzle r, VkComponentSwizzle g, VkComponentSwizzle b, VkComponentSwizzle a)
+void FVulkanDevice::MapFormatSupport(EPixelFormat UEFormat, std::initializer_list<VkFormat> PrioritizedFormats, const VkComponentMapping& ComponentMapping)
 {
-	// Please ensure that we support the mapping, otherwise there is no point setting it.
-	check(GPixelFormats[UEFormat].Supported);
-	VkComponentMapping& ComponentMapping = PixelFormatComponentMapping[UEFormat];
-	ComponentMapping.r = r;
-	ComponentMapping.g = g;
-	ComponentMapping.b = b;
-	ComponentMapping.a = a;
+	MapFormatSupport(UEFormat, PrioritizedFormats, ComponentMapping, kDefaultTextureCapabilities, kDefaultBlockBytes);
 }
-
-void FVulkanDevice::MapFormatSupport(EPixelFormat UEFormat, VkFormat VulkanFormat, int32 BlockBytes)
+void FVulkanDevice::MapFormatSupport(EPixelFormat UEFormat, std::initializer_list<VkFormat> PrioritizedFormats, const VkComponentMapping& ComponentMapping, int32 BlockBytes)
+{
+	MapFormatSupport(UEFormat, PrioritizedFormats, ComponentMapping, kDefaultTextureCapabilities, BlockBytes);
+}
+void FVulkanDevice::MapFormatSupport(EPixelFormat UEFormat, std::initializer_list<VkFormat> PrioritizedFormats, const VkComponentMapping& ComponentMapping, EPixelFormatCapabilities RequiredCapabilities)
 {
-	MapFormatSupport(UEFormat, VulkanFormat);
-	FPixelFormatInfo& FormatInfo = GPixelFormats[UEFormat];
-	FormatInfo.BlockBytes = BlockBytes;
+	MapFormatSupport(UEFormat, PrioritizedFormats, ComponentMapping, RequiredCapabilities, kDefaultBlockBytes);
 }
 
 bool FVulkanDevice::QueryGPU(int32 DeviceIndex)
@@ -1490,65 +1471,6 @@ void FVulkanDevice::WaitUntilIdle()
 	GetImmediateContext().GetCommandBufferManager()->RefreshFenceStatus();
 }
 
-bool FVulkanDevice::IsTextureFormatSupported(VkFormat Format, uint32 RequiredFeatures) const
-{
-	check(RequiredFeatures != 0u);
-		
-	auto ArePropertiesSupported = [](const VkFormatProperties& Prop, uint32 RequiredFeatures) -> bool
-	{
-		return	(Prop.linearTilingFeatures & RequiredFeatures) == RequiredFeatures || 
-				(Prop.optimalTilingFeatures & RequiredFeatures) == RequiredFeatures;
-	};
-
-	if (Format >= 0 && Format < VK_FORMAT_RANGE_SIZE)
-	{
-		const VkFormatProperties& Prop = FormatProperties[Format];
-		return ArePropertiesSupported(Prop, RequiredFeatures);
-	}
-
-	// Check for extension formats
-	const VkFormatProperties* FoundProperties = ExtensionFormatProperties.Find(Format);
-	if (FoundProperties)
-	{
-		return ArePropertiesSupported(*FoundProperties, RequiredFeatures);
-	}
-
-	// Add it for faster caching next time
-	VkFormatProperties& NewProperties = ExtensionFormatProperties.Add(Format);
-	FMemory::Memzero(NewProperties);
-	VulkanRHI::vkGetPhysicalDeviceFormatProperties(Gpu, Format, &NewProperties);
-
-	return ArePropertiesSupported(NewProperties, RequiredFeatures);
-}
-
-bool FVulkanDevice::IsBufferFormatSupported(VkFormat Format) const
-{
-	auto ArePropertiesSupported = [](const VkFormatProperties& Prop) -> bool
-	{
-		return (Prop.bufferFeatures != 0);
-	};
-
-	if (Format >= 0 && Format < VK_FORMAT_RANGE_SIZE)
-	{
-		const VkFormatProperties& Prop = FormatProperties[Format];
-		return Prop.bufferFeatures != 0;
-	}
-
-	// Check for extension formats
-	const VkFormatProperties* FoundProperties = ExtensionFormatProperties.Find(Format);
-	if (FoundProperties)
-	{
-		return FoundProperties->bufferFeatures != 0;
-	}
-
-	// Add it for faster caching next time
-	VkFormatProperties& NewProperties = ExtensionFormatProperties.Add(Format);
-	FMemory::Memzero(NewProperties);
-	VulkanRHI::vkGetPhysicalDeviceFormatProperties(Gpu, Format, &NewProperties);
-
-	return NewProperties.bufferFeatures != 0;
-}
-
 const VkComponentMapping& FVulkanDevice::GetFormatComponentMapping(EPixelFormat UEFormat) const
 {
 	if (UEFormat == PF_X24_G8)
diff --git a/Engine/Source/Runtime/VulkanRHI/Private/VulkanDevice.h b/Engine/Source/Runtime/VulkanRHI/Private/VulkanDevice.h
--- a/Engine/Source/Runtime/VulkanRHI/Private/VulkanDevice.h
+++ b/Engine/Source/Runtime/VulkanRHI/Private/VulkanDevice.h
@@ -349,9 +349,6 @@ public:
 		return TimestampValidBitsMask;
 	}
 
-	bool IsTextureFormatSupported(VkFormat Format, uint32 RequiredFeatures) const;
-	bool IsBufferFormatSupported(VkFormat Format) const;
-
 	const VkComponentMapping& GetFormatComponentMapping(EPixelFormat UEFormat) const;
 
 	inline VkDevice GetInstanceHandle() const
@@ -517,20 +514,13 @@ public:
 	void*	Hotfix = nullptr;
 
 private:
-	void MapFormatSupport(EPixelFormat UEFormat, VkFormat VulkanFormat);
-	void MapFormatSupportWithFallback(EPixelFormat UEFormat, uint32 TextureRequiredFeatures, VkFormat VulkanFormat, TArrayView<const VkFormat> FallbackTextureFormats);
-	void MapFormatSupport(EPixelFormat UEFormat, VkFormat VulkanFormat, int32 BlockBytes);
-	void SetComponentMapping(EPixelFormat UEFormat, VkComponentSwizzle r, VkComponentSwizzle g, VkComponentSwizzle b, VkComponentSwizzle a);
-
-	FORCEINLINE void MapFormatSupportWithFallback(EPixelFormat UEFormat, VkFormat VulkanFormat, std::initializer_list<VkFormat> FallbackTextureFormats)
-	{
-		MapFormatSupportWithFallback(UEFormat, VK_FORMAT_FEATURE_SAMPLED_IMAGE_BIT, VulkanFormat, MakeArrayView(FallbackTextureFormats));
-	}
-	
-	FORCEINLINE void MapFormatSupportWithFallback(EPixelFormat UEFormat, uint32 TextureRequiredFeatures, VkFormat VulkanFormat, std::initializer_list<VkFormat> FallbackTextureFormats)
-	{
-		MapFormatSupportWithFallback(UEFormat, TextureRequiredFeatures, VulkanFormat, MakeArrayView(FallbackTextureFormats));
-	}
+	const VkFormatProperties& GetFormatProperties(VkFormat InFormat);
+	void MapBufferFormatSupport(FPixelFormatInfo& PixelFormatInfo, EPixelFormat UEFormat, VkFormat VulkanFormat);
+	void MapImageFormatSupport(FPixelFormatInfo& PixelFormatInfo, const TArrayView<const VkFormat>& PrioritizedFormats, EPixelFormatCapabilities RequiredCapabilities);
+	void MapFormatSupport(EPixelFormat UEFormat, std::initializer_list<VkFormat> PrioritizedFormats, const VkComponentMapping& ComponentMapping, EPixelFormatCapabilities RequiredCapabilities, int32 BlockBytes);
+	void MapFormatSupport(EPixelFormat UEFormat, std::initializer_list<VkFormat> PrioritizedFormats, const VkComponentMapping& ComponentMapping);
+	void MapFormatSupport(EPixelFormat UEFormat, std::initializer_list<VkFormat> PrioritizedFormats, const VkComponentMapping& ComponentMapping, int32 BlockBytes);
+	void MapFormatSupport(EPixelFormat UEFormat, std::initializer_list<VkFormat> PrioritizedFormats, const VkComponentMapping& ComponentMapping, EPixelFormatCapabilities RequiredCapabilities);
 
 	void SubmitCommands(FVulkanCommandListContext* Context);
 
diff --git a/Engine/Source/Runtime/VulkanRHI/Private/VulkanGenericPlatform.h b/Engine/Source/Runtime/VulkanRHI/Private/VulkanGenericPlatform.h
--- a/Engine/Source/Runtime/VulkanRHI/Private/VulkanGenericPlatform.h
+++ b/Engine/Source/Runtime/VulkanRHI/Private/VulkanGenericPlatform.h
@@ -127,11 +127,4 @@ public:
 
 	// Return VK_FALSE if platform wants to suppress the given debug report from the validation layers, VK_TRUE to print it.
 	static VkBool32 DebugReportFunction(VkDebugReportFlagsEXT MsgFlags, VkDebugReportObjectTypeEXT ObjType, uint64_t SrcObject, size_t Location, int32 MsgCode, const ANSICHAR* LayerPrefix, const ANSICHAR* Msg, void* UserData) { return VK_TRUE; }
-
-	static VkFormat GetPlatform5551Format(VkFormat& OutFallbackFormat0, VkFormat& OutFallbackFormat1)
-	{
-		OutFallbackFormat0 = VK_FORMAT_R5G5B5A1_UNORM_PACK16;
-		OutFallbackFormat1 = VK_FORMAT_B8G8R8A8_UNORM;
-		return VK_FORMAT_A1R5G5B5_UNORM_PACK16;
-	};
 };
