Fix build: remove EPixelFormatCapabilities::UAV which doesn't exist in Main.

diff --git a/Engine/Source/Runtime/VulkanRHI/Private/VulkanDevice.cpp b/Engine/Source/Runtime/VulkanRHI/Private/VulkanDevice.cpp
--- a/Engine/Source/Runtime/VulkanRHI/Private/VulkanDevice.cpp
+++ b/Engine/Source/Runtime/VulkanRHI/Private/VulkanDevice.cpp
@@ -851,7 +851,7 @@ void FVulkanDevice::MapImageFormatSupport(FPixelFormatInfo& PixelFormatInfo, con
 		ConvertImageCap(LocalFormatProperties, EPixelFormatCapabilities::DepthStencil, VK_FORMAT_FEATURE_DEPTH_STENCIL_ATTACHMENT_BIT);
 		ConvertImageCap(LocalFormatProperties, EPixelFormatCapabilities::RenderTarget, VK_FORMAT_FEATURE_COLOR_ATTACHMENT_BIT);
 		ConvertImageCap(LocalFormatProperties, EPixelFormatCapabilities::TextureBlendable, VK_FORMAT_FEATURE_COLOR_ATTACHMENT_BLEND_BIT);
-		ConvertImageCap(LocalFormatProperties, EPixelFormatCapabilities::UAV | EPixelFormatCapabilities::TypedUAVLoad | EPixelFormatCapabilities::TypedUAVStore, VK_FORMAT_FEATURE_STORAGE_IMAGE_BIT);
+		ConvertImageCap(LocalFormatProperties, /*EPixelFormatCapabilities::UAV |*/ EPixelFormatCapabilities::TypedUAVLoad | EPixelFormatCapabilities::TypedUAVStore, VK_FORMAT_FEATURE_STORAGE_IMAGE_BIT);
 
 		ConvertImageCap(LocalFormatProperties, EPixelFormatCapabilities::TextureAtomics, VK_FORMAT_FEATURE_STORAGE_IMAGE_ATOMIC_BIT);
 
