Fixed: missing depth test for translucecy with mobile deferred shading enabled

diff --git a/Engine/Source/Runtime/VulkanRHI/Private/VulkanRenderpass.cpp b/Engine/Source/Runtime/VulkanRHI/Private/VulkanRenderpass.cpp
--- a/Engine/Source/Runtime/VulkanRHI/Private/VulkanRenderpass.cpp
+++ b/Engine/Source/Runtime/VulkanRHI/Private/VulkanRenderpass.cpp
@@ -156,6 +156,7 @@ public:
 			{
 				TSubpassDescriptionClass& SubpassDesc = SubpassDescriptions[NumSubpasses++];
 				SubpassDesc.SetColorAttachments(ColorAttachmentReferences, 1); // SceneColor only
+				SubpassDesc.SetDepthStencilAttachment(&DepthStencilAttachment);
 
 				// GBuffer as Input2/3/4
 				InputAttachments2[0].attachment = DepthStencilAttachment.attachment;
