PR #8517: fix VK_FDM usage on Vulkan (Contributed by [REDACTED_PERSON])

3 fixes:
* Fixes the GRHISupportsVariableRateShading logic to ensure that bool doesn't get nuked if the device supports FragmentDensityMap but not ShadingRateAttachment.
* Ensures Scene Textures are created as an array when multiview enabled.
* Ensures the RenderPass2 impl supports the Multiview Mask.

diff --git a/Engine/Source/Runtime/RenderCore/Private/RenderGraphBuilder.cpp b/Engine/Source/Runtime/RenderCore/Private/RenderGraphBuilder.cpp
--- a/Engine/Source/Runtime/RenderCore/Private/RenderGraphBuilder.cpp
+++ b/Engine/Source/Runtime/RenderCore/Private/RenderGraphBuilder.cpp
@@ -711,6 +711,11 @@ FRDGTextureRef FRDGBuilder::RegisterExternalTexture(
 		{
 			Texture->AccessFinal |= ERHIAccess::CopyDest;
 		}
+
+		if (EnumHasAnyFlags(Desc.Flags, TexCreate_Foveation))
+		{
+			Texture->AccessFinal |= ERHIAccess::ShadingRateSource;
+		}
 	}
 
 	FRDGTextureSubresourceState& TextureState = Texture->GetState();
diff --git a/Engine/Source/Runtime/Renderer/Private/SceneTextures.cpp b/Engine/Source/Runtime/Renderer/Private/SceneTextures.cpp
--- a/Engine/Source/Runtime/Renderer/Private/SceneTextures.cpp
+++ b/Engine/Source/Runtime/Renderer/Private/SceneTextures.cpp
@@ -807,7 +807,9 @@ FSceneTextures& FSceneTextures::Create(FRDGBuilder& GraphBuilder, const FSceneTe
 		const float FarDepth = (float)ERHIZBuffer::FarPlane;
 		const FLinearColor FarDepthColor(FarDepth, FarDepth, FarDepth, FarDepth);
 		ETextureCreateFlags FlagsToAdd = IsMobileDeferredShadingEnabled(Config.ShaderPlatform)? TexCreate_Memoryless : TexCreate_None;
-		FRDGTextureDesc Desc = FRDGTextureDesc::Create2D(Config.Extent, GetMobileSceneDepthAuxPixelFormat(Config.ShaderPlatform), FClearValueBinding(FarDepthColor), TexCreate_RenderTargetable | TexCreate_ShaderResource | TexCreate_InputAttachmentRead| FlagsToAdd);
+		FRDGTextureDesc Desc = Config.bRequireMultiView ? 
+			FRDGTextureDesc::Create2DArray(Config.Extent, GetMobileSceneDepthAuxPixelFormat(Config.ShaderPlatform), FClearValueBinding(FarDepthColor), TexCreate_RenderTargetable | TexCreate_ShaderResource | TexCreate_InputAttachmentRead | FlagsToAdd, 2) : 
+			FRDGTextureDesc::Create2D(Config.Extent, GetMobileSceneDepthAuxPixelFormat(Config.ShaderPlatform), FClearValueBinding(FarDepthColor), TexCreate_RenderTargetable | TexCreate_ShaderResource | TexCreate_InputAttachmentRead| FlagsToAdd);
 		Desc.NumSamples = Config.NumSamples;
 		SceneTextures.DepthAux = GraphBuilder.CreateTexture(Desc, TEXT("SceneDepthAux"));
 		
diff --git a/Engine/Source/Runtime/VulkanRHI/Private/VulkanRHI.cpp b/Engine/Source/Runtime/VulkanRHI/Private/VulkanRHI.cpp
--- a/Engine/Source/Runtime/VulkanRHI/Private/VulkanRHI.cpp
+++ b/Engine/Source/Runtime/VulkanRHI/Private/VulkanRHI.cpp
@@ -826,7 +826,7 @@ void FVulkanDynamicRHI::InitInstance()
 #endif
 
 #if VULKAN_SUPPORTS_FRAGMENT_SHADING_RATE
-		GRHISupportsAttachmentVariableRateShading = GetDevice()->GetOptionalExtensions().HasKHRFragmentShadingRate && Device->GetFragmentShadingRateFeatures().attachmentFragmentShadingRate;
+		GRHISupportsAttachmentVariableRateShading |= GetDevice()->GetOptionalExtensions().HasKHRFragmentShadingRate && Device->GetFragmentShadingRateFeatures().attachmentFragmentShadingRate;
 #endif
 
 		FVulkanPlatform::SetupFeatureLevels();
diff --git a/Engine/Source/Runtime/VulkanRHI/Private/VulkanRenderpass.cpp b/Engine/Source/Runtime/VulkanRHI/Private/VulkanRenderpass.cpp
--- a/Engine/Source/Runtime/VulkanRHI/Private/VulkanRenderpass.cpp
+++ b/Engine/Source/Runtime/VulkanRHI/Private/VulkanRenderpass.cpp
@@ -71,6 +71,7 @@ public:
 			{
 				SubpassDesc.SetShadingRateAttachment(&FragmentShadingRateAttachmentInfo);
 			}
+			SubpassDesc.SetMultiViewMask(MultiviewMask);
 #endif
 		}
 
@@ -101,6 +102,7 @@ public:
 			{
 				SubpassDesc.SetShadingRateAttachment(&FragmentShadingRateAttachmentInfo);
 			}
+			SubpassDesc.SetMultiViewMask(MultiviewMask);
 #endif
 			
 			TSubpassDependencyClass& SubpassDep = SubpassDependencies[NumDependencies++];
@@ -139,6 +141,7 @@ public:
 				{
 					SubpassDesc.SetShadingRateAttachment(&FragmentShadingRateAttachmentInfo);
 				}
+				SubpassDesc.SetMultiViewMask(MultiviewMask);
 #endif
 				
 				// depth as Input0
@@ -180,6 +183,7 @@ public:
 				{
 					SubpassDesc.SetShadingRateAttachment(&FragmentShadingRateAttachmentInfo);
 				}
+				SubpassDesc.SetMultiViewMask(MultiviewMask);
 #endif
 
 				TSubpassDependencyClass& SubpassDep = SubpassDependencies[NumDependencies++];
@@ -223,6 +227,9 @@ public:
 		*/
 		const uint32_t CorrelationMask = MultiviewMask;
 
+#if VULKAN_SUPPORTS_RENDERPASS2
+		CreateInfo.SetCorrelationMask(&CorrelationMask);
+#else
 		VkRenderPassMultiviewCreateInfo MultiviewInfo;
 		if (RTLayout.GetIsMultiView())
 		{
@@ -238,6 +245,7 @@ public:
 
 			CreateInfo.pNext = &MultiviewInfo;
 		}
+#endif
 
 		VkRenderPassFragmentDensityMapCreateInfoEXT FragDensityCreateInfo;
 		if (Device.GetOptionalExtensions().HasEXTFragmentDensityMap && RTLayout.GetHasFragmentDensityAttachment())
diff --git a/Engine/Source/Runtime/VulkanRHI/Private/VulkanRenderpass.h b/Engine/Source/Runtime/VulkanRHI/Private/VulkanRenderpass.h
--- a/Engine/Source/Runtime/VulkanRHI/Private/VulkanRenderpass.h
+++ b/Engine/Source/Runtime/VulkanRHI/Private/VulkanRenderpass.h
@@ -132,6 +132,11 @@ struct FVulkanSubpassDescription<VkSubpassDescription>
 	{
 		// No-op without VK_KHR_create_renderpass2
 	}
+
+	void SetMultiViewMask(uint32_t Mask)
+	{
+		// No-op without VK_KHR_create_renderpass2
+	}
 };
 
 #if VULKAN_SUPPORTS_RENDERPASS2
@@ -176,6 +181,11 @@ struct FVulkanSubpassDescription<VkSubpassDescription2>
 	{
 		pNext = ShadingRateAttachmentInfo;
 	}
+
+	void SetMultiViewMask(uint32_t Mask)
+	{
+		viewMask = Mask;
+	}
 };
 #endif
 
@@ -276,6 +286,11 @@ struct FVulkanRenderPassCreateInfo<VkRenderPassCreateInfo>
 		ZeroVulkanStruct(*this, VK_STRUCTURE_TYPE_RENDER_PASS_CREATE_INFO);
 	}
 
+	void SetCorrelationMask(const uint32_t* MaskPtr)
+	{
+		// No-op without VK_KHR_create_renderpass2
+	}
+
 	VkRenderPass Create(FVulkanDevice& Device)
 	{
 		VkRenderPass Handle = VK_NULL_HANDLE;
@@ -294,6 +309,12 @@ struct FVulkanRenderPassCreateInfo<VkRenderPassCreateInfo2>
 		ZeroVulkanStruct(*this, VK_STRUCTURE_TYPE_RENDER_PASS_CREATE_INFO_2);
 	}
 
+	void SetCorrelationMask(const uint32_t* MaskPtr)
+	{
+		correlatedViewMaskCount = 1;
+		pCorrelatedViewMasks = MaskPtr;
+	}
+
 	VkRenderPass Create(FVulkanDevice& Device)
 	{
 		VkRenderPass Handle = VK_NULL_HANDLE;
