VK: Make sure we setup resolve attachments only for the last subpass

diff --git a/Engine/Source/Runtime/VulkanRHI/Private/VulkanRenderpass.cpp b/Engine/Source/Runtime/VulkanRHI/Private/VulkanRenderpass.cpp
--- a/Engine/Source/Runtime/VulkanRHI/Private/VulkanRenderpass.cpp
+++ b/Engine/Source/Runtime/VulkanRHI/Private/VulkanRenderpass.cpp
@@ -60,7 +60,11 @@ public:
 			TSubpassDescriptionClass& SubpassDesc = SubpassDescriptions[NumSubpasses++];
 
 			SubpassDesc.SetColorAttachments(ColorAttachmentReferences);
-			SubpassDesc.SetResolveAttachments(ResolveAttachmentReferences);
+			if (!bDepthReadSubpass)
+			{
+				// only set resolve attachment on the last subpass
+				SubpassDesc.SetResolveAttachments(ResolveAttachmentReferences);
+			}
 			if (RTLayout.GetDepthStencilAttachmentReference() != nullptr)
 			{
 				SubpassDesc.SetDepthStencilAttachment(&DepthStencilAttachmentReference);
