Fix "r.Vulkan.WaitForIdleOnSubmit" crash on a queue submit early during init.

diff --git a/Engine/Source/Runtime/VulkanRHI/Private/VulkanQueue.cpp b/Engine/Source/Runtime/VulkanRHI/Private/VulkanQueue.cpp
--- a/Engine/Source/Runtime/VulkanRHI/Private/VulkanQueue.cpp
+++ b/Engine/Source/Runtime/VulkanRHI/Private/VulkanQueue.cpp
@@ -79,7 +79,7 @@ void FVulkanQueue::Submit(FVulkanCmdBuffer* CmdBuffer, uint32 NumSignalSemaphore
 
 	if (GWaitForIdleOnSubmit != 0)
 	{
-		FVulkanCommandBufferManager* CmdBufferMgr = Device->GetImmediateContext().GetCommandBufferManager();
+		FVulkanCommandBufferManager* CmdBufferMgr = &CmdBuffer->GetOwner()->GetMgr();
 
 		switch(GWaitForIdleOnSubmit)
 		{
