Remove "hotfix variables" used to maintain binary compatibility for past versions.

diff --git a/Engine/Source/Runtime/VulkanRHI/Private/VulkanContext.h b/Engine/Source/Runtime/VulkanRHI/Private/VulkanContext.h
--- a/Engine/Source/Runtime/VulkanRHI/Private/VulkanContext.h
+++ b/Engine/Source/Runtime/VulkanRHI/Private/VulkanContext.h
@@ -178,8 +178,6 @@ public:
 
 	void ReleasePendingState();
 
-	void* Hotfix = nullptr;
-
 protected:
 	FVulkanDynamicRHI* RHI;
 	FVulkanCommandListContext* Immediate;
