Fix bug preventing FVulkanDevice::GpuIdProps being set in pNext chain

diff --git a/Engine/Source/Runtime/VulkanRHI/Private/VulkanDevice.cpp b/Engine/Source/Runtime/VulkanRHI/Private/VulkanDevice.cpp
--- a/Engine/Source/Runtime/VulkanRHI/Private/VulkanDevice.cpp
+++ b/Engine/Source/Runtime/VulkanRHI/Private/VulkanDevice.cpp
@@ -960,30 +960,27 @@ bool FVulkanDevice::QueryGPU(int32 DeviceIndex)
 		GpuProps2.pNext = &GpuIdProps;
 		ZeroVulkanStruct(GpuIdProps, VK_STRUCTURE_TYPE_PHYSICAL_DEVICE_ID_PROPERTIES_KHR);
 
-		void** NextPropsAddr = &GpuProps2.pNext;
-
 #if VULKAN_SUPPORTS_DRIVER_PROPERTIES
 		if (GetOptionalExtensions().HasDriverProperties)
 		{
 			PhysicalDeviceProperties.pNext = GpuProps2.pNext;
 			GpuProps2.pNext = &PhysicalDeviceProperties;
-			NextPropsAddr = &PhysicalDeviceProperties.pNext;
 		}
 #endif
 
 #if VULKAN_SUPPORTS_FRAGMENT_DENSITY_MAP
 		if (GetOptionalExtensions().HasEXTFragmentDensityMap)
 		{
-			*NextPropsAddr = &FragmentDensityMapProperties;
-			NextPropsAddr = &FragmentDensityMapProperties.pNext;
+			FragmentDensityMapProperties.pNext = GpuProps2.pNext;
+			GpuProps2.pNext = &FragmentDensityMapProperties;
 		}
 #endif
 
 #if VULKAN_SUPPORTS_FRAGMENT_SHADING_RATE
 		if (GetOptionalExtensions().HasKHRFragmentShadingRate)
 		{
-			*NextPropsAddr = &FragmentShadingRateProperties;
-			NextPropsAddr = &FragmentShadingRateProperties.pNext;
+			FragmentShadingRateProperties.pNext = GpuProps2.pNext;
+			GpuProps2.pNext = &FragmentShadingRateProperties;
 		}
 #endif
 
