Fix PF_X24_G8 pixel format on Vulkan so that component mapping actually puts stencil values in G.  Fixes issue with HairStrands.DebugMode=6 in Editor.

diff --git a/Engine/Source/Runtime/VulkanRHI/Private/VulkanDevice.cpp b/Engine/Source/Runtime/VulkanRHI/Private/VulkanDevice.cpp
--- a/Engine/Source/Runtime/VulkanRHI/Private/VulkanDevice.cpp
+++ b/Engine/Source/Runtime/VulkanRHI/Private/VulkanDevice.cpp
@@ -577,6 +577,7 @@ void FVulkanDevice::SetupFormats()
 	const VkComponentMapping ComponentMappingRIII = { VK_COMPONENT_SWIZZLE_R, VK_COMPONENT_SWIZZLE_IDENTITY, VK_COMPONENT_SWIZZLE_IDENTITY, VK_COMPONENT_SWIZZLE_IDENTITY };
 	const VkComponentMapping ComponentMapping000R = { VK_COMPONENT_SWIZZLE_ZERO, VK_COMPONENT_SWIZZLE_ZERO, VK_COMPONENT_SWIZZLE_ZERO, VK_COMPONENT_SWIZZLE_R };
 	const VkComponentMapping ComponentMappingR000 = { VK_COMPONENT_SWIZZLE_R, VK_COMPONENT_SWIZZLE_ZERO, VK_COMPONENT_SWIZZLE_ZERO, VK_COMPONENT_SWIZZLE_ZERO };
+	const VkComponentMapping ComponentMappingRR01 = { VK_COMPONENT_SWIZZLE_R, VK_COMPONENT_SWIZZLE_R, VK_COMPONENT_SWIZZLE_ZERO, VK_COMPONENT_SWIZZLE_ONE };
 
 
 	// Initialize the platform pixel format map.
@@ -643,11 +644,13 @@ void FVulkanDevice::SetupFormats()
 	{
 		// prefer VK_FORMAT_D24_UNORM_S8_UINT
 		MapFormatSupport(PF_DepthStencil, { VK_FORMAT_D24_UNORM_S8_UINT, VK_FORMAT_D32_SFLOAT_S8_UINT, VK_FORMAT_D16_UNORM_S8_UINT }, ComponentMappingRIII);
+		MapFormatSupport(PF_X24_G8, { VK_FORMAT_D24_UNORM_S8_UINT, VK_FORMAT_D32_SFLOAT_S8_UINT, VK_FORMAT_D16_UNORM_S8_UINT }, ComponentMappingRR01);
 	}
 	else
 	{
 		// prefer VK_FORMAT_D32_SFLOAT_S8_UINT
 		MapFormatSupport(PF_DepthStencil, { VK_FORMAT_D32_SFLOAT_S8_UINT, VK_FORMAT_D24_UNORM_S8_UINT, VK_FORMAT_D16_UNORM_S8_UINT }, ComponentMappingRIII);
+		MapFormatSupport(PF_X24_G8, { VK_FORMAT_D32_SFLOAT_S8_UINT, VK_FORMAT_D24_UNORM_S8_UINT, VK_FORMAT_D16_UNORM_S8_UINT }, ComponentMappingRR01);
 	}
 
 	if (FVulkanPlatform::SupportsBCTextureFormats())
@@ -1470,10 +1473,6 @@ void FVulkanDevice::WaitUntilIdle()
 
 const VkComponentMapping& FVulkanDevice::GetFormatComponentMapping(EPixelFormat UEFormat) const
 {
-	if (UEFormat == PF_X24_G8)
-	{
-		return GetFormatComponentMapping(PF_DepthStencil);
-	}
 	check(GPixelFormats[UEFormat].Supported);
 	return PixelFormatComponentMapping[UEFormat];
 }
diff --git a/Engine/Source/Runtime/VulkanRHI/Private/VulkanTexture.cpp b/Engine/Source/Runtime/VulkanRHI/Private/VulkanTexture.cpp
--- a/Engine/Source/Runtime/VulkanRHI/Private/VulkanTexture.cpp
+++ b/Engine/Source/Runtime/VulkanRHI/Private/VulkanTexture.cpp
@@ -1869,9 +1869,7 @@ VkImageView FVulkanTextureView::StaticCreate(FVulkanDevice& Device, VkImage InIm
 	//instead of PF_DepthStencil, so the cross-platform code could figure out the proper format to pass in for this.
 	if (UEFormat == PF_X24_G8)
 	{
-		ensure(ViewInfo.format == VK_FORMAT_UNDEFINED);
-		ViewInfo.format = (VkFormat)GPixelFormats[PF_DepthStencil].PlatformFormat;
-		ensure(ViewInfo.format != VK_FORMAT_UNDEFINED);
+		ensure((ViewInfo.format == (VkFormat)GPixelFormats[PF_DepthStencil].PlatformFormat) && (ViewInfo.format != VK_FORMAT_UNDEFINED));
 		ViewInfo.subresourceRange.aspectMask = VK_IMAGE_ASPECT_STENCIL_BIT;
 	}
 
