Fix external textures not removing themselves from layout mgr tracking.  (Thanks RP!)

diff --git a/Engine/Source/Runtime/VulkanRHI/Private/VulkanResources.h b/Engine/Source/Runtime/VulkanRHI/Private/VulkanResources.h
--- a/Engine/Source/Runtime/VulkanRHI/Private/VulkanResources.h
+++ b/Engine/Source/Runtime/VulkanRHI/Private/VulkanResources.h
@@ -316,6 +316,15 @@ struct FVulkanCpuReadbackBuffer
 	uint32 MipSize[MAX_TEXTURE_MIP_COUNT];
 };
 
+
+enum class EImageOwnerType : uint8
+{
+	None,
+	LocalOwner,
+	ExternalOwner,
+	Aliased
+};
+
 /** Texture/RT wrapper. */
 class FVulkanSurface : public FVulkanEvictable
 {
@@ -358,7 +367,7 @@ public:
 	// Constructor for externally owned Image
 	FVulkanSurface(FVulkanDevice& Device, VkImageViewType ResourceType, EPixelFormat Format,
 					uint32 SizeX, uint32 SizeY, uint32 SizeZ, uint32 ArraySize, uint32 NumMips, uint32 NumSamples,
-					VkImage InImage, ETextureCreateFlags UEFlags, const FRHIResourceCreateInfo& CreateInfo);
+					VkImage InImage, ETextureCreateFlags UEFlags, EImageOwnerType InImageOwnerType, const FRHIResourceCreateInfo& CreateInfo);
 
 	virtual ~FVulkanSurface();
 
@@ -442,7 +451,7 @@ public:
 
 	inline bool IsImageOwner() const
 	{
-		return bIsImageOwner;
+		return (ImageOwnerType == EImageOwnerType::LocalOwner);
 	}
 
 	VULKANRHI_API VkDeviceMemory GetAllocationHandle() const;
@@ -478,7 +487,6 @@ private:
 	VkImageTiling Tiling;
 	VkImageViewType	ViewType;
 
-	bool bIsImageOwner;
 	VulkanRHI::FVulkanAllocation Allocation;
 
 	uint32 NumMips;
@@ -490,6 +498,8 @@ private:
 	FVulkanCpuReadbackBuffer* CpuReadbackBuffer;
 	FVulkanTextureBase* OwningTexture = 0;
 
+	EImageOwnerType ImageOwnerType;
+
 	friend struct FVulkanTextureBase;
 };
 
diff --git a/Engine/Source/Runtime/VulkanRHI/Private/VulkanTexture.cpp b/Engine/Source/Runtime/VulkanRHI/Private/VulkanTexture.cpp
--- a/Engine/Source/Runtime/VulkanRHI/Private/VulkanTexture.cpp
+++ b/Engine/Source/Runtime/VulkanRHI/Private/VulkanTexture.cpp
@@ -586,12 +586,12 @@ FVulkanSurface::FVulkanSurface(FVulkanDevice& InDevice, FVulkanEvictable* Owner,
 	, MemProps(VK_MEMORY_PROPERTY_DEVICE_LOCAL_BIT)
 	, Tiling(VK_IMAGE_TILING_MAX_ENUM)	// Can be expanded to a per-platform definition
 	, ViewType(ResourceType)
-	, bIsImageOwner(true)
 	, NumMips(InNumMips)
 	, NumSamples(InNumSamples)
 	, FullAspectMask(0)
 	, PartialAspectMask(0)
 	, CpuReadbackBuffer(nullptr) // for readback textures we use a staging buffer. this is because vulkan only requires implentations to support 1 mip level(which is useless), so we emulate using a buffer
+	, ImageOwnerType(EImageOwnerType::None)
 {
 	FImageCreateInfo ImageCreateInfo;
 	FVulkanSurface::GenerateImageCreateInfo(ImageCreateInfo,
@@ -599,7 +599,7 @@ FVulkanSurface::FVulkanSurface(FVulkanDevice& InDevice, FVulkanEvictable* Owner,
 		InFormat, Width, Height, Depth,
 		ArraySize, NumMips, NumSamples, UEFlags,
 		&StorageFormat, &ViewFormat);
-	if(EnumHasAnyFlags(UEFlags, TexCreate_CPUReadback))
+	if (EnumHasAnyFlags(UEFlags, TexCreate_CPUReadback))
 	{
 		check(NumSamples == 1);	//not implemented
 		check(Depth == 1);		//not implemented
@@ -637,6 +637,7 @@ FVulkanSurface::FVulkanSurface(FVulkanDevice& InDevice, FVulkanEvictable* Owner,
 		return;
 	}
 
+	ImageOwnerType = EImageOwnerType::LocalOwner;
 	VERIFYVULKANRESULT(VulkanRHI::vkCreateImage(InDevice.GetInstanceHandle(), &ImageCreateInfo.ImageCreateInfo, VULKAN_CPU_ALLOCATOR, &Image));
 
 	// Fetch image size
@@ -903,7 +904,7 @@ void FVulkanSurface::EvictSurface(FVulkanDevice& InDevice)
 // This is usually used for the framebuffer image
 FVulkanSurface::FVulkanSurface(FVulkanDevice& InDevice, VkImageViewType ResourceType, EPixelFormat InFormat,
 								uint32 SizeX, uint32 SizeY, uint32 SizeZ, uint32 InArraySize, uint32 InNumMips, uint32 InNumSamples,
-								VkImage InImage, ETextureCreateFlags InUEFlags, const FRHIResourceCreateInfo& CreateInfo)
+								VkImage InImage, ETextureCreateFlags InUEFlags, EImageOwnerType InImageOwnerType, const FRHIResourceCreateInfo& CreateInfo)
 	: Device(&InDevice)
 	, Image(InImage)
 	, StorageFormat(VK_FORMAT_UNDEFINED)
@@ -917,12 +918,12 @@ FVulkanSurface::FVulkanSurface(FVulkanDevice& InDevice, VkImageViewType Resource
 	, MemProps(VK_MEMORY_PROPERTY_DEVICE_LOCAL_BIT)
 	, Tiling(VK_IMAGE_TILING_MAX_ENUM)	// Can be expanded to a per-platform definition
 	, ViewType(ResourceType)
-	, bIsImageOwner(false)
 	, NumMips(InNumMips)
 	, NumSamples(InNumSamples)
 	, FullAspectMask(0)
 	, PartialAspectMask(0)
 	, CpuReadbackBuffer(nullptr)
+	, ImageOwnerType(InImageOwnerType)
 {
 	StorageFormat = UEToVkTextureFormat(PixelFormat, false);
 
@@ -1016,17 +1017,17 @@ FVulkanSurface::~FVulkanSurface()
 
 void FVulkanSurface::Destroy()
 {
-	// An image can be instances.
-	// - Instances VkImage has "bIsImageOwner" set to "false".
-	// - Owner of VkImage has "bIsImageOwner" set to "true".
-	if(CpuReadbackBuffer)
+	const bool bIsLocalOwner = (ImageOwnerType == EImageOwnerType::LocalOwner);
+	const bool bHasExternalOwner = (ImageOwnerType == EImageOwnerType::ExternalOwner);
+
+	if (CpuReadbackBuffer)
 	{
 		Device->GetDeferredDeletionQueue().EnqueueResource(VulkanRHI::FDeferredDeletionQueue2::EType::Buffer, CpuReadbackBuffer->Buffer);
 		Device->GetMemoryManager().FreeVulkanAllocation(Allocation);
 		delete CpuReadbackBuffer;
 
 	}
-	else if (bIsImageOwner)
+	else if (bIsLocalOwner || bHasExternalOwner)
 	{
 		const bool bRenderTarget = EnumHasAnyFlags(UEFlags, TexCreate_RenderTargetable | TexCreate_DepthStencilTargetable | TexCreate_ResolveTargetable);
 		FRHICommandList& RHICmdList = FRHICommandListExecutor::GetImmediateCommandList();
@@ -1040,19 +1041,22 @@ void FVulkanSurface::Destroy()
 			new (RHICmdList.AllocCommand<FRHICommandOnDestroyImage>()) FRHICommandOnDestroyImage(Image, Device, bRenderTarget);
 		}
 
-		bIsImageOwner = false;
+		if (bIsLocalOwner)
+		{
+			uint64 Size = 0;
 
-		uint64 Size = 0;
+			if (Image != VK_NULL_HANDLE)
+			{
+				Size = GetMemorySize();
+				Device->GetDeferredDeletionQueue().EnqueueResource(VulkanRHI::FDeferredDeletionQueue2::EType::Image, Image);
+				Device->GetMemoryManager().FreeVulkanAllocation(Allocation);
+				Image = VK_NULL_HANDLE;
+			}
 
-		if (Image != VK_NULL_HANDLE)
-		{
-			Size = GetMemorySize();
-			Device->GetDeferredDeletionQueue().EnqueueResource(VulkanRHI::FDeferredDeletionQueue2::EType::Image, Image);
-			Device->GetMemoryManager().FreeVulkanAllocation(Allocation);
-			Image = VK_NULL_HANDLE;
+			VulkanTextureDestroyed(Size, ViewType, bRenderTarget);
 		}
 
-		VulkanTextureDestroyed(Size, ViewType, bRenderTarget);
+		ImageOwnerType = EImageOwnerType::None;
 	}
 }
 
@@ -1996,7 +2000,7 @@ FVulkanTextureBase::FVulkanTextureBase(FVulkanDevice& Device, VkImageViewType Re
 }
 
 FVulkanTextureBase::FVulkanTextureBase(FVulkanDevice& Device, VkImageViewType ResourceType, EPixelFormat Format, uint32 SizeX, uint32 SizeY, uint32 SizeZ, uint32 ArraySize, uint32 InNumMips, uint32 InNumSamples, VkImage InImage, VkDeviceMemory InMem, ETextureCreateFlags UEFlags, const FRHIResourceCreateInfo& CreateInfo)
-	: Surface(Device, ResourceType, Format, SizeX, SizeY, SizeZ, ArraySize, InNumMips, InNumSamples, InImage, UEFlags, CreateInfo)
+	: Surface(Device, ResourceType, Format, SizeX, SizeY, SizeZ, ArraySize, InNumMips, InNumSamples, InImage, UEFlags, EImageOwnerType::ExternalOwner, CreateInfo)
 	, PartialView(nullptr)
 {
 	Surface.OwningTexture = this;
@@ -2021,7 +2025,7 @@ FVulkanTextureBase::FVulkanTextureBase(FVulkanDevice& Device, VkImageViewType Re
 }
 
 FVulkanTextureBase::FVulkanTextureBase(FVulkanDevice& Device, VkImageViewType ResourceType, EPixelFormat Format, uint32 SizeX, uint32 SizeY, uint32 SizeZ, uint32 ArraySize, uint32 NumMips, uint32 NumSamples, VkImage InImage, VkDeviceMemory InMem, FSamplerYcbcrConversionInitializer& ConversionInitializer, ETextureCreateFlags UEFlags, const FRHIResourceCreateInfo& CreateInfo)
-	: Surface(Device, ResourceType, Format, SizeX, SizeY, SizeZ, ArraySize, NumMips, NumSamples, InImage, UEFlags, CreateInfo)
+	: Surface(Device, ResourceType, Format, SizeX, SizeY, SizeZ, ArraySize, NumMips, NumSamples, InImage, UEFlags, EImageOwnerType::ExternalOwner, CreateInfo)
 	, PartialView(nullptr)
 {
 	Surface.OwningTexture = this;
@@ -2054,7 +2058,7 @@ FVulkanTextureBase::FVulkanTextureBase(FVulkanDevice& Device, VkImageViewType Re
 }
 
 FVulkanTextureBase::FVulkanTextureBase(FTextureRHIRef& SrcTextureRHI, const FVulkanTextureBase* SrcTexture, VkImageViewType ResourceType, uint32 SizeX, uint32 SizeY, uint32 SizeZ)
-	: Surface(*SrcTexture->Surface.Device, ResourceType, SrcTexture->Surface.PixelFormat, SizeX, SizeY, SizeZ, SrcTexture->Surface.ArraySize, SrcTexture->Surface.NumMips, SrcTexture->Surface.NumSamples, SrcTexture->Surface.Image, SrcTexture->Surface.UEFlags, FRHIResourceCreateInfo(TEXT("FVulkanTextureBase")))
+	: Surface(*SrcTexture->Surface.Device, ResourceType, SrcTexture->Surface.PixelFormat, SizeX, SizeY, SizeZ, SrcTexture->Surface.ArraySize, SrcTexture->Surface.NumMips, SrcTexture->Surface.NumSamples, SrcTexture->Surface.Image, SrcTexture->Surface.UEFlags, EImageOwnerType::Aliased, FRHIResourceCreateInfo(TEXT("FVulkanTextureBase")))
 	, PartialView(nullptr)
 	, AliasedTexture(SrcTextureRHI)
 {
