Enable compilation of Vulkan RHI ray tracing code by default on Windows (off by default at runtime via DDPI and new CVar r.Vulkan.ExperimentalRayTracing).

* Enable Vulkan 1.2 API when ray tracing is enabled (which is the default now)
* Add support for ray query extension

diff --git a/Engine/Source/Runtime/VulkanRHI/Private/VulkanDevice.h b/Engine/Source/Runtime/VulkanRHI/Private/VulkanDevice.h
--- a/Engine/Source/Runtime/VulkanRHI/Private/VulkanDevice.h
+++ b/Engine/Source/Runtime/VulkanRHI/Private/VulkanDevice.h
@@ -63,6 +63,7 @@ struct FOptionalVulkanDeviceExtensions
 			uint32 HasKHRMultiview : 1;
 			uint32 HasAccelerationStructure : 1;
 			uint32 HasRayTracingPipeline : 1;
+			uint32 HasRayQuery : 1;
 			uint32 HasDescriptorIndexing : 1;
 			uint32 HasBufferDeviceAddress : 1;
 			uint32 HasDeferredHostOperations : 1;
@@ -90,7 +91,7 @@ struct FOptionalVulkanDeviceExtensions
 	{
 		return 
 			HasAccelerationStructure && 
-			HasRayTracingPipeline && 
+			(HasRayTracingPipeline || HasRayQuery) &&
 			HasDescriptorIndexing && 
 			HasBufferDeviceAddress && 
 			HasDeferredHostOperations && 
@@ -113,6 +114,7 @@ struct FOptionalVulkanDeviceFeatures
 	VkPhysicalDeviceAccelerationStructureFeaturesKHR AccelerationStructureFeatures;
 	VkPhysicalDeviceRayTracingPipelineFeaturesKHR RayTracingPipelineFeatures;
 	VkPhysicalDeviceDescriptorIndexingFeaturesEXT DescriptorIndexingFeatures;
+	VkPhysicalDeviceRayQueryFeaturesKHR RayQueryFeatures;
 #endif
 };
 
diff --git a/Engine/Source/Runtime/VulkanRHI/Private/VulkanLayers.cpp b/Engine/Source/Runtime/VulkanRHI/Private/VulkanLayers.cpp
--- a/Engine/Source/Runtime/VulkanRHI/Private/VulkanLayers.cpp
+++ b/Engine/Source/Runtime/VulkanRHI/Private/VulkanLayers.cpp
@@ -902,6 +902,7 @@ void FOptionalVulkanDeviceExtensions::Setup(const TArray<const ANSICHAR*>& Devic
 #if VULKAN_RHI_RAYTRACING
 	HasAccelerationStructure	= HasExtension(DeviceExtensions, VK_KHR_ACCELERATION_STRUCTURE_EXTENSION_NAME);
 	HasRayTracingPipeline		= HasExtension(DeviceExtensions, VK_KHR_RAY_TRACING_PIPELINE_EXTENSION_NAME);
+	HasRayQuery		            = HasExtension(DeviceExtensions, VK_KHR_RAY_QUERY_EXTENSION_NAME);
 	HasDescriptorIndexing		= HasExtension(DeviceExtensions, VK_EXT_DESCRIPTOR_INDEXING_EXTENSION_NAME);
 	HasBufferDeviceAddress		= HasExtension(DeviceExtensions, VK_KHR_BUFFER_DEVICE_ADDRESS_EXTENSION_NAME);
 	HasDeferredHostOperations	= HasExtension(DeviceExtensions, VK_KHR_DEFERRED_HOST_OPERATIONS_EXTENSION_NAME);
diff --git a/Engine/Source/Runtime/VulkanRHI/Private/VulkanRHI.cpp b/Engine/Source/Runtime/VulkanRHI/Private/VulkanRHI.cpp
--- a/Engine/Source/Runtime/VulkanRHI/Private/VulkanRHI.cpp
+++ b/Engine/Source/Runtime/VulkanRHI/Private/VulkanRHI.cpp
@@ -775,13 +775,18 @@ void FVulkanDynamicRHI::InitInstance()
 		GSupportsMobileMultiView = Device->GetMultiviewFeatures().multiview == VK_TRUE ? true : false;
 #endif
 #if VULKAN_RHI_RAYTRACING
-		GRHISupportsRayTracing = false; // Disable runtime hwrt on vk during development.
-		//GRHISupportsRayTracing = Device->GetOptionalExtensions().HasRaytracingExtensions();
+		GRHISupportsRayTracing = RHISupportsRayTracing(GMaxRHIShaderPlatform) && Device->GetOptionalExtensions().HasRaytracingExtensions();
 
-		const FRayTracingProperties& RayTracingProps = Device->GetRayTracingProperties();
-		GRHIRayTracingAccelerationStructureAlignment = 256; // TODO (currently handled by FVulkanAccelerationStructureBuffer)
-		GRHIRayTracingScratchBufferAlignment = RayTracingProps.AccelerationStructure.minAccelerationStructureScratchOffsetAlignment;
-		GRHIRayTracingInstanceDescriptorSize = uint32(sizeof(VkAccelerationStructureInstanceKHR));
+		if (GRHISupportsRayTracing)
+		{
+			GRHISupportsRayTracingShaders = RHISupportsRayTracingShaders(GMaxRHIShaderPlatform);
+			GRHISupportsInlineRayTracing = RHISupportsInlineRayTracing(GMaxRHIShaderPlatform) && Device->GetOptionalExtensions().HasRayQuery;
+
+			const FRayTracingProperties& RayTracingProps = Device->GetRayTracingProperties();
+			GRHIRayTracingAccelerationStructureAlignment = 256; // TODO (currently handled by FVulkanAccelerationStructureBuffer)
+			GRHIRayTracingScratchBufferAlignment = RayTracingProps.AccelerationStructure.minAccelerationStructureScratchOffsetAlignment;
+			GRHIRayTracingInstanceDescriptorSize = uint32(sizeof(VkAccelerationStructureInstanceKHR));
+		}
 #endif
 #if VULKAN_ENABLE_DUMP_LAYER
 		// Disable RHI thread by default if the dump layer is enabled
@@ -1452,7 +1457,10 @@ void FVulkanDescriptorSetsLayout::Compile(FVulkanDescriptorSetLayoutMap& DSetLay
 	check(LayoutTypes[VK_DESCRIPTOR_TYPE_INPUT_ATTACHMENT] <= Limits.maxDescriptorSetInputAttachments);
 
 #if VULKAN_RHI_RAYTRACING
-	check(LayoutTypes[VK_DESCRIPTOR_TYPE_ACCELERATION_STRUCTURE_KHR] < Device->GetRayTracingProperties().AccelerationStructure.maxDescriptorSetAccelerationStructures);
+	if (GRHISupportsRayTracing)
+	{
+		check(LayoutTypes[VK_DESCRIPTOR_TYPE_ACCELERATION_STRUCTURE_KHR] < Device->GetRayTracingProperties().AccelerationStructure.maxDescriptorSetAccelerationStructures);
+	}
 #endif
 	
 	LayoutHandles.Empty(SetLayouts.Num());
diff --git a/Engine/Source/Runtime/VulkanRHI/Private/VulkanRayTracing.cpp b/Engine/Source/Runtime/VulkanRHI/Private/VulkanRayTracing.cpp
--- a/Engine/Source/Runtime/VulkanRHI/Private/VulkanRayTracing.cpp
+++ b/Engine/Source/Runtime/VulkanRHI/Private/VulkanRayTracing.cpp
@@ -9,15 +9,29 @@
 #include "BuiltInRayTracingShaders.h"
 #include "Experimental/Containers/SherwoodHashTable.h"
 
+TAutoConsoleVariable<int32> GVulkanExperimentalRayTracing(
+	TEXT("r.Vulkan.ExperimentalRayTracing"),
+	0,
+	TEXT("0: Do not enable Vulkan ray tracing extensions (default)\n")
+	TEXT("1: Enable experimental ray tracing support (for development and testing purposes)"),
+	ECVF_ReadOnly
+);
+
 // Define ray tracing entry points
 #define DEFINE_VK_ENTRYPOINTS(Type,Func) VULKANRHI_API Type VulkanDynamicAPI::Func = NULL;
 ENUM_VK_ENTRYPOINTS_RAYTRACING(DEFINE_VK_ENTRYPOINTS)
 
 void FVulkanRayTracingPlatform::GetDeviceExtensions(EGpuVendorId VendorId, TArray<const ANSICHAR*>& OutExtensions)
 {
+	if (GVulkanExperimentalRayTracing.GetValueOnAnyThread() == 0)
+	{
+		return;
+	}
+
 	// Primary extensions
 	OutExtensions.Add(VK_KHR_ACCELERATION_STRUCTURE_EXTENSION_NAME);
 	OutExtensions.Add(VK_KHR_RAY_TRACING_PIPELINE_EXTENSION_NAME);
+	OutExtensions.Add(VK_KHR_RAY_QUERY_EXTENSION_NAME);
 
 	// VK_KHR_acceleration_structure dependencies 
 	OutExtensions.Add(VK_EXT_DESCRIPTOR_INDEXING_EXTENSION_NAME); // Note: Promoted to Vulkan 1.2
@@ -50,8 +64,12 @@ void FVulkanRayTracingPlatform::EnablePhysicalDeviceFeatureExtensions(VkDeviceCr
 		Features.RayTracingPipelineFeatures.rayTracingPipeline = VK_TRUE;
 		Features.RayTracingPipelineFeatures.pNext = &Features.AccelerationStructureFeatures;
 
+		ZeroVulkanStruct(Features.RayQueryFeatures, VK_STRUCTURE_TYPE_PHYSICAL_DEVICE_RAY_QUERY_FEATURES_KHR);
+		Features.RayQueryFeatures.pNext = &Features.RayTracingPipelineFeatures;
+		Features.RayQueryFeatures.rayQuery = VK_TRUE;
+
 		ZeroVulkanStruct(Features.DescriptorIndexingFeatures, VK_STRUCTURE_TYPE_PHYSICAL_DEVICE_DESCRIPTOR_INDEXING_FEATURES_EXT);
-		Features.DescriptorIndexingFeatures.pNext = &Features.RayTracingPipelineFeatures;
+		Features.DescriptorIndexingFeatures.pNext = &Features.RayQueryFeatures;
 
 		DeviceInfo.pNext = &Features.DescriptorIndexingFeatures;
 	}
diff --git a/Engine/Source/Runtime/VulkanRHI/Private/Windows/VulkanWindowsPlatform.h b/Engine/Source/Runtime/VulkanRHI/Private/Windows/VulkanWindowsPlatform.h
--- a/Engine/Source/Runtime/VulkanRHI/Private/Windows/VulkanWindowsPlatform.h
+++ b/Engine/Source/Runtime/VulkanRHI/Private/Windows/VulkanWindowsPlatform.h
@@ -17,11 +17,15 @@
 #define VULKAN_SUPPORTS_COLOR_CONVERSIONS			1
 #define VULKAN_SUPPORTS_AMD_BUFFER_MARKER			1
 
-#define	UE_VK_API_VERSION							VK_API_VERSION_1_1
-
-#define VULKAN_RHI_RAYTRACING 						0//(RHI_RAYTRACING)
+#define VULKAN_RHI_RAYTRACING 						(RHI_RAYTRACING)
 #define VULKAN_SUPPORTS_SCALAR_BLOCK_LAYOUT			(VULKAN_RHI_RAYTRACING)
 
+#if VULKAN_RHI_RAYTRACING
+#	define UE_VK_API_VERSION						VK_API_VERSION_1_2
+#else
+#	define UE_VK_API_VERSION						VK_API_VERSION_1_1
+#endif // VULKAN_RHI_RAYTRACING
+
 // 32-bit windows has warnings on custom mem mgr callbacks
 #define VULKAN_SHOULD_USE_LLM					(UE_BUILD_DEBUG || UE_BUILD_DEVELOPMENT) && !PLATFORM_32BITS
 
