Implement Vulkan top level acceleration structure SRV management and binding. This allows acceleration structures to be used in shaders that perform inline ray tracing on Vulkan platforms.

diff --git a/Engine/Source/Runtime/VulkanRHI/Private/VulkanDescriptorSets.h b/Engine/Source/Runtime/VulkanRHI/Private/VulkanDescriptorSets.h
--- a/Engine/Source/Runtime/VulkanRHI/Private/VulkanDescriptorSets.h
+++ b/Engine/Source/Runtime/VulkanRHI/Private/VulkanDescriptorSets.h
@@ -67,6 +67,9 @@ struct FDescriptorSetRemappingInfo
 		TArray<VkDescriptorType>	Types;
 		uint16						NumImageInfos = 0;
 		uint16						NumBufferInfos = 0;
+#if VULKAN_RHI_RAYTRACING
+		uint8						NumAccelerationStructures = 0;
+#endif // VULKAN_RHI_RAYTRACING
 	};
 	TArray<FSetInfo>	SetInfos;
 
@@ -124,6 +127,9 @@ struct FDescriptorSetRemappingInfo
 			int32 SetInfosNums = SetInfos[SetInfosIndex].Types.Num();
 			if (SetInfos[SetInfosIndex].NumBufferInfos != In.SetInfos[SetInfosIndex].NumBufferInfos ||
 				SetInfos[SetInfosIndex].NumImageInfos != In.SetInfos[SetInfosIndex].NumImageInfos ||
+#if VULKAN_RHI_RAYTRACING
+				SetInfos[SetInfosIndex].NumAccelerationStructures != In.SetInfos[SetInfosIndex].NumAccelerationStructures ||
+#endif // VULKAN_RHI_RAYTRACING
 				SetInfosNums != In.SetInfos[SetInfosIndex].Types.Num() ||
 				(SetInfosNums != 0 && FMemory::Memcmp(SetInfos[SetInfosIndex].Types.GetData(), In.SetInfos[SetInfosIndex].Types.GetData(), sizeof(VkDescriptorType) * SetInfosNums)))
 			{
@@ -689,6 +695,10 @@ struct FVulkanDescriptorSetWriteContainer
 	TArray<VkDescriptorImageInfo> DescriptorImageInfo;
 	TArray<VkDescriptorBufferInfo> DescriptorBufferInfo;
 	TArray<VkWriteDescriptorSet> DescriptorWrites;
+#if VULKAN_RHI_RAYTRACING
+	TArray<VkAccelerationStructureKHR> AccelerationStructures;
+	TArray<VkWriteDescriptorSetAccelerationStructureKHR> AccelerationStructureWrites;
+#endif // VULKAN_RHI_RAYTRACING
 	TArray<uint8> BindingToDynamicOffsetMap;
 };
 
@@ -1069,6 +1079,44 @@ public:
 		return WriteBufferView<VK_DESCRIPTOR_TYPE_UNIFORM_TEXEL_BUFFER>(DescriptorIndex, View);
 	}
 
+
+#if VULKAN_RHI_RAYTRACING
+	bool WriteAccelerationStructure(uint32 DescriptorIndex, VkAccelerationStructureKHR InAccelerationStructure)
+	{
+		checkf(!UseVulkanDescriptorCache(), TEXT("Descriptor cache path for WriteAccelerationStructure() is not implemented"));
+
+		check(DescriptorIndex < NumWrites);
+		SetWritten(DescriptorIndex);
+
+		check(WriteDescriptors[DescriptorIndex].descriptorType == VK_DESCRIPTOR_TYPE_ACCELERATION_STRUCTURE_KHR);
+
+		// Find the acceleration structure extension in the generic VkWriteDescriptorSet.
+		const VkWriteDescriptorSetAccelerationStructureKHR* FoundWrite = nullptr;
+		const VkBaseInStructure* Cursor = reinterpret_cast<const VkBaseInStructure*>(WriteDescriptors[DescriptorIndex].pNext);
+		while (Cursor)
+		{
+			if (Cursor->sType == VK_STRUCTURE_TYPE_WRITE_DESCRIPTOR_SET_ACCELERATION_STRUCTURE_KHR)
+			{
+				FoundWrite = reinterpret_cast<const VkWriteDescriptorSetAccelerationStructureKHR*>(Cursor);
+				break;
+			}
+			Cursor = Cursor->pNext;
+		}
+
+		checkf(FoundWrite,
+			TEXT("Expected to find a VkWriteDescriptorSetAccelerationStructureKHR that's needed to bind an acceleration structure descriptor. ")
+			TEXT("Possibly something went wrong in SetupDescriptorWrites()."));
+
+		checkf(FoundWrite->accelerationStructureCount == 1, TEXT("Acceleration structure write operation is expected to contain exactly one descriptor"));
+
+		VkAccelerationStructureKHR& AccelerationStructure = *const_cast<VkAccelerationStructureKHR*>(FoundWrite->pAccelerationStructures);
+
+		bool bChanged = CopyAndReturnNotEqual(AccelerationStructure, InAccelerationStructure);
+
+		return bChanged;
+	}
+#endif // VULKAN_RHI_RAYTRACING
+
 	void ClearBufferView(uint32 DescriptorIndex)
 	{
 		BufferViewReferences[DescriptorIndex] = nullptr;
@@ -1253,6 +1301,10 @@ protected:
 		FVulkanHashableDescriptorInfo* InHashableDescriptorInfos,
 		VkWriteDescriptorSet* InWriteDescriptors, VkDescriptorImageInfo* InImageInfo,
 		VkDescriptorBufferInfo* InBufferInfo, uint8* InBindingToDynamicOffsetMap,
+#if VULKAN_RHI_RAYTRACING
+		VkWriteDescriptorSetAccelerationStructureKHR* InAccelerationStructuresWriteDescriptors,
+		VkAccelerationStructureKHR* InAccelerationStructures,
+#endif // VULKAN_RHI_RAYTRACING
 		const FVulkanSamplerState& DefaultSampler, const FVulkanTextureView& DefaultImageView);
 
 	friend class FVulkanCommonPipelineDescriptorState;
diff --git a/Engine/Source/Runtime/VulkanRHI/Private/VulkanDevice.h b/Engine/Source/Runtime/VulkanRHI/Private/VulkanDevice.h
--- a/Engine/Source/Runtime/VulkanRHI/Private/VulkanDevice.h
+++ b/Engine/Source/Runtime/VulkanRHI/Private/VulkanDevice.h
@@ -153,6 +153,7 @@ namespace VulkanRHI
 			ResourceAllocation,
 			DeviceMemoryAllocation,
 			BufferSuballocation,
+			AccelerationStructure,
 		};
 
 		template <typename T>
diff --git a/Engine/Source/Runtime/VulkanRHI/Private/VulkanMemory.cpp b/Engine/Source/Runtime/VulkanRHI/Private/VulkanMemory.cpp
--- a/Engine/Source/Runtime/VulkanRHI/Private/VulkanMemory.cpp
+++ b/Engine/Source/Runtime/VulkanRHI/Private/VulkanMemory.cpp
@@ -20,6 +20,10 @@ THIRD_PARTY_INCLUDES_END
 #include "Windows/HideWindowsPlatformTypes.h"
 #endif
 
+#if VULKAN_RHI_RAYTRACING
+#include "VulkanRayTracing.h"
+#endif
+
 // This 'frame number' should only be used for the deletion queue
 uint32 GVulkanRHIDeletionFrameNumber = 0;
 const uint32 NUM_FRAMES_TO_WAIT_FOR_RESOURCE_DELETE = 2;
@@ -4921,9 +4925,15 @@ namespace VulkanRHI
 				{
 					check(Entry->DeviceMemoryAllocation);
 					Device->GetDeviceMemoryManager().Free(Entry->DeviceMemoryAllocation);
+					break;
 				}
-				break;
-
+#if VULKAN_RHI_RAYTRACING
+				case EType::AccelerationStructure:
+				{
+					VulkanDynamicAPI::vkDestroyAccelerationStructureKHR(DeviceHandle, (VkAccelerationStructureKHR)Entry->Handle, VULKAN_CPU_ALLOCATOR);
+					break;
+				}
+#endif // VULKAN_RHI_RAYTRACING
 
 				default:
 					check(0);
diff --git a/Engine/Source/Runtime/VulkanRHI/Private/VulkanPendingState.cpp b/Engine/Source/Runtime/VulkanRHI/Private/VulkanPendingState.cpp
--- a/Engine/Source/Runtime/VulkanRHI/Private/VulkanPendingState.cpp
+++ b/Engine/Source/Runtime/VulkanRHI/Private/VulkanPendingState.cpp
@@ -323,6 +323,12 @@ void FVulkanPendingComputeState::SetSRVForUBResource(uint32 DescriptorSet, uint3
 		{
 			CurrentState->SetStorageBuffer(DescriptorSet, BindingIndex, SRV->SourceStructuredBuffer);
 		}
+#if VULKAN_RHI_RAYTRACING
+		else if (SRV->AccelerationStructureHandle)
+		{
+			CurrentState->SetAccelerationStructure(DescriptorSet, BindingIndex, SRV->AccelerationStructureHandle);
+		}
+#endif // VULKAN_RHI_RAYTRACING
 		else
 		{
 			checkf(SRV->TextureView.View != VK_NULL_HANDLE, TEXT("Empty SRV"));
diff --git a/Engine/Source/Runtime/VulkanRHI/Private/VulkanPipelineState.cpp b/Engine/Source/Runtime/VulkanRHI/Private/VulkanPipelineState.cpp
--- a/Engine/Source/Runtime/VulkanRHI/Private/VulkanPipelineState.cpp
+++ b/Engine/Source/Runtime/VulkanRHI/Private/VulkanPipelineState.cpp
@@ -77,6 +77,11 @@ void FVulkanCommonPipelineDescriptorState::CreateDescriptorWriteInfos()
 		DSWriteContainer.DescriptorImageInfo.AddZeroed(SetInfo.NumImageInfos);
 		DSWriteContainer.DescriptorBufferInfo.AddZeroed(SetInfo.NumBufferInfos);
 
+#if VULKAN_RHI_RAYTRACING
+		DSWriteContainer.AccelerationStructureWrites.AddZeroed(SetInfo.NumAccelerationStructures);
+		DSWriteContainer.AccelerationStructures.AddZeroed(SetInfo.NumAccelerationStructures);
+#endif // VULKAN_RHI_RAYTRACING
+
 		checkf(SetInfo.Types.Num() < 255, TEXT("Need more bits for BindingToDynamicOffsetMap (currently 8)! Requires %d descriptor bindings in a set!"), SetInfo.Types.Num());
 		DSWriteContainer.BindingToDynamicOffsetMap.AddUninitialized(SetInfo.Types.Num());
 	}
@@ -97,6 +102,12 @@ void FVulkanCommonPipelineDescriptorState::CreateDescriptorWriteInfos()
 	VkWriteDescriptorSet* CurrentDescriptorWrite = DSWriteContainer.DescriptorWrites.GetData();
 	VkDescriptorImageInfo* CurrentImageInfo = DSWriteContainer.DescriptorImageInfo.GetData();
 	VkDescriptorBufferInfo* CurrentBufferInfo = DSWriteContainer.DescriptorBufferInfo.GetData();
+
+#if VULKAN_RHI_RAYTRACING
+	VkWriteDescriptorSetAccelerationStructureKHR* CurrentAccelerationStructuresWriteDescriptors = DSWriteContainer.AccelerationStructureWrites.GetData();
+	VkAccelerationStructureKHR* CurrentAccelerationStructures = DSWriteContainer.AccelerationStructures.GetData();
+#endif // VULKAN_RHI_RAYTRACING
+
 	uint8* CurrentBindingToDynamicOffsetMap = DSWriteContainer.BindingToDynamicOffsetMap.GetData();
 	TArray<uint32> DynamicOffsetsStart;
 	DynamicOffsetsStart.AddZeroed(NumSets);
@@ -111,6 +122,10 @@ void FVulkanCommonPipelineDescriptorState::CreateDescriptorWriteInfos()
 		uint32 NumDynamicOffsets = DSWriter[Set].SetupDescriptorWrites(
 			SetInfo.Types, CurrentHashableDescriptorInfo,
 			CurrentDescriptorWrite, CurrentImageInfo, CurrentBufferInfo, CurrentBindingToDynamicOffsetMap,
+#if VULKAN_RHI_RAYTRACING
+			CurrentAccelerationStructuresWriteDescriptors,
+			CurrentAccelerationStructures,
+#endif // VULKAN_RHI_RAYTRACING
 			DefaultSampler, DefaultImageView);
 
 		TotalNumDynamicOffsets += NumDynamicOffsets;
@@ -127,6 +142,12 @@ void FVulkanCommonPipelineDescriptorState::CreateDescriptorWriteInfos()
 		CurrentDescriptorWrite += SetInfo.Types.Num();
 		CurrentImageInfo += SetInfo.NumImageInfos;
 		CurrentBufferInfo += SetInfo.NumBufferInfos;
+
+#if VULKAN_RHI_RAYTRACING
+		CurrentAccelerationStructuresWriteDescriptors += SetInfo.NumAccelerationStructures;
+		CurrentAccelerationStructures += SetInfo.NumAccelerationStructures;
+#endif // VULKAN_RHI_RAYTRACING
+
 		CurrentBindingToDynamicOffsetMap += SetInfo.Types.Num();
 	}
 
diff --git a/Engine/Source/Runtime/VulkanRHI/Private/VulkanPipelineState.h b/Engine/Source/Runtime/VulkanRHI/Private/VulkanPipelineState.h
--- a/Engine/Source/Runtime/VulkanRHI/Private/VulkanPipelineState.h
+++ b/Engine/Source/Runtime/VulkanRHI/Private/VulkanPipelineState.h
@@ -122,6 +122,12 @@ public:
 		}
 	}
 
+#if VULKAN_RHI_RAYTRACING
+	inline void SetAccelerationStructure(uint8 DescriptorSet, uint32 BindingIndex, VkAccelerationStructureKHR AccelerationStructure)
+	{
+		MarkDirty(DSWriter[DescriptorSet].WriteAccelerationStructure(BindingIndex, AccelerationStructure));
+	}
+#endif // VULKAN_RHI_RAYTRACING
 
 protected:
 	void Reset()
diff --git a/Engine/Source/Runtime/VulkanRHI/Private/VulkanRHI.cpp b/Engine/Source/Runtime/VulkanRHI/Private/VulkanRHI.cpp
--- a/Engine/Source/Runtime/VulkanRHI/Private/VulkanRHI.cpp
+++ b/Engine/Source/Runtime/VulkanRHI/Private/VulkanRHI.cpp
@@ -1292,6 +1292,16 @@ FVulkanDescriptorSetsLayout::~FVulkanDescriptorSetsLayout()
 	LayoutHandles.Reset(0);
 }
 
+// Increments a value and asserts on overflow.
+// FSetInfo uses narrow integer types for descriptor counts,
+// which may feasibly overflow one day (for example if we add bindless resources).
+template <typename T> 
+static void IncrementChecked(T& Value)
+{
+	check(Value < TNumericLimits<T>::Max());
+	++Value;
+}
+
 void FVulkanDescriptorSetsLayoutInfo::AddDescriptor(int32 DescriptorSetIndex, const VkDescriptorSetLayoutBinding& Descriptor)
 {
 	// Increment type usage
@@ -1323,19 +1333,20 @@ void FVulkanDescriptorSetsLayoutInfo::AddDescriptor(int32 DescriptorSetIndex, co
 	case VK_DESCRIPTOR_TYPE_SAMPLED_IMAGE:
 	case VK_DESCRIPTOR_TYPE_STORAGE_IMAGE:
 	case VK_DESCRIPTOR_TYPE_INPUT_ATTACHMENT:
-		++RemappingInfo.SetInfos[DescriptorSetIndex].NumImageInfos;
+		IncrementChecked(RemappingInfo.SetInfos[DescriptorSetIndex].NumImageInfos);
 		break;
 	case VK_DESCRIPTOR_TYPE_UNIFORM_BUFFER:
 	case VK_DESCRIPTOR_TYPE_STORAGE_BUFFER:
 	case VK_DESCRIPTOR_TYPE_UNIFORM_BUFFER_DYNAMIC:
-		++RemappingInfo.SetInfos[DescriptorSetIndex].NumBufferInfos;
+		IncrementChecked(RemappingInfo.SetInfos[DescriptorSetIndex].NumBufferInfos);
 		break;
-	case VK_DESCRIPTOR_TYPE_STORAGE_TEXEL_BUFFER:
-	case VK_DESCRIPTOR_TYPE_UNIFORM_TEXEL_BUFFER:
-
 #if VULKAN_RHI_RAYTRACING
 	case VK_DESCRIPTOR_TYPE_ACCELERATION_STRUCTURE_KHR:
+		IncrementChecked(RemappingInfo.SetInfos[DescriptorSetIndex].NumAccelerationStructures);
+		break;
 #endif
+	case VK_DESCRIPTOR_TYPE_STORAGE_TEXEL_BUFFER:
+	case VK_DESCRIPTOR_TYPE_UNIFORM_TEXEL_BUFFER:
 		break;
 	default:
 		checkf(0, TEXT("Unsupported descriptor type %d"), (int32)Descriptor.descriptorType);
diff --git a/Engine/Source/Runtime/VulkanRHI/Private/VulkanRayTracing.cpp b/Engine/Source/Runtime/VulkanRHI/Private/VulkanRayTracing.cpp
--- a/Engine/Source/Runtime/VulkanRHI/Private/VulkanRayTracing.cpp
+++ b/Engine/Source/Runtime/VulkanRHI/Private/VulkanRayTracing.cpp
@@ -328,15 +328,12 @@ FVulkanRayTracingGeometry::~FVulkanRayTracingGeometry()
 {
 	if (Handle != VK_NULL_HANDLE)
 	{
-		VkDevice NativeDevice = Device->GetInstanceHandle();
-		vkDestroyAccelerationStructureKHR(NativeDevice, Handle, VULKAN_CPU_ALLOCATOR);
+		Device->GetDeferredDeletionQueue().EnqueueResource(VulkanRHI::FDeferredDeletionQueue2::EType::AccelerationStructure, Handle);
 	}
 }
 
 void FVulkanRayTracingGeometry::BuildAccelerationStructure(FVulkanCommandListContext& CommandContext, EAccelerationStructureBuildMode BuildMode)
 {
-	check(BuildMode == EAccelerationStructureBuildMode::Build); // Temp
-
 	FVkRtBLASBuildData BuildData;
 	GetBLASBuildData(
 		Device->GetInstanceHandle(),
@@ -435,29 +432,20 @@ FVulkanRayTracingScene::FVulkanRayTracingScene(FRayTracingSceneInitializer2 InIn
 
 FVulkanRayTracingScene::~FVulkanRayTracingScene()
 {
-	if (Handle != VK_NULL_HANDLE)
-	{
-		VkDevice NativeDevice = Device->GetInstanceHandle();
-		vkDestroyAccelerationStructureKHR(NativeDevice, Handle, VULKAN_CPU_ALLOCATOR);
-	}
 }
 
 void FVulkanRayTracingScene::BindBuffer(FRHIBuffer* InBuffer, uint32 InBufferOffset)
 {
-	checkf(Handle == VK_NULL_HANDLE, TEXT("Binding multiple buffers not currently supported."));
+	checkf(AccelerationStructureView == nullptr, TEXT("Binding multiple buffers is not currently supported."));
+
+	check(IsInRHIThread() || !IsRunningRHIInSeparateThread());
+
 	check(SizeInfo.ResultSize + InBufferOffset <= InBuffer->GetSize());
 	check(InBufferOffset % 256 == 0); // Spec requires offset to be a multiple of 256
 	AccelerationStructureBuffer = static_cast<FVulkanAccelerationStructureBuffer*>(InBuffer);
 
-	VkDevice NativeDevice = Device->GetInstanceHandle();
-
-	VkAccelerationStructureCreateInfoKHR CreateInfo;
-	ZeroVulkanStruct(CreateInfo, VK_STRUCTURE_TYPE_ACCELERATION_STRUCTURE_CREATE_INFO_KHR);
-	CreateInfo.buffer = AccelerationStructureBuffer->GetBuffer();
-	CreateInfo.offset = InBufferOffset;
-	CreateInfo.size = SizeInfo.ResultSize;
-	CreateInfo.type = VK_ACCELERATION_STRUCTURE_TYPE_TOP_LEVEL_KHR;
-	VERIFYVULKANRESULT(vkCreateAccelerationStructureKHR(NativeDevice, &CreateInfo, VULKAN_CPU_ALLOCATOR, &Handle));
+	FShaderResourceViewInitializer ViewInitializer(InBuffer, InBufferOffset, 0);
+	AccelerationStructureView = new FVulkanShaderResourceView(Device, AccelerationStructureBuffer, InBufferOffset);
 }
 
 void FVulkanRayTracingScene::BuildAccelerationStructure(
@@ -492,7 +480,8 @@ void FVulkanRayTracingScene::BuildAccelerationStructure(
 		InScratchBuffer = ScratchBuffer.GetReference();
 	}
 
-	BuildData.GeometryInfo.dstAccelerationStructure = Handle;
+	checkf(AccelerationStructureView, TEXT("A buffer must be bound to the ray tracing scene before it can be built."));
+	BuildData.GeometryInfo.dstAccelerationStructure = AccelerationStructureView->AccelerationStructureHandle;
 
 	BuildData.GeometryInfo.scratchData.deviceAddress = InScratchBuffer->GetDeviceAddress();
 	if (bExternalScratchBuffer)
diff --git a/Engine/Source/Runtime/VulkanRHI/Private/VulkanRayTracing.h b/Engine/Source/Runtime/VulkanRHI/Private/VulkanRayTracing.h
--- a/Engine/Source/Runtime/VulkanRHI/Private/VulkanRayTracing.h
+++ b/Engine/Source/Runtime/VulkanRHI/Private/VulkanRayTracing.h
@@ -116,7 +116,15 @@ private:
 
 	TRefCountPtr<FVulkanResourceMultiBuffer> InstanceBuffer;
 
-	VkAccelerationStructureKHR Handle = VK_NULL_HANDLE;
+	// Native TLAS handles are owned by SRV objects in Vulkan RHI.
+	// D3D12 and other RHIs allow creating TLAS SRVs from any GPU address at any point
+	// and do not require them for operations such as build or update.
+	// FVulkanRayTracingScene can't own the VkAccelerationStructureKHR directly because
+	// we allow TLAS memory to be allocated using transient resource allocator and 
+	// the lifetime of the scene object may be different from the lifetime of the buffer.
+	// Many VkAccelerationStructureKHR-s may be created, pointing at the same buffer.
+	TRefCountPtr<FVulkanShaderResourceView> AccelerationStructureView;
+	
 	TRefCountPtr<FVulkanAccelerationStructureBuffer> AccelerationStructureBuffer;
 };
 
diff --git a/Engine/Source/Runtime/VulkanRHI/Private/VulkanResources.h b/Engine/Source/Runtime/VulkanRHI/Private/VulkanResources.h
--- a/Engine/Source/Runtime/VulkanRHI/Private/VulkanResources.h
+++ b/Engine/Source/Runtime/VulkanRHI/Private/VulkanResources.h
@@ -1288,6 +1288,10 @@ public:
 	FVulkanShaderResourceView(FVulkanDevice* Device, FRHITexture* InSourceTexture, const FRHITextureSRVCreateInfo& InCreateInfo);
 	FVulkanShaderResourceView(FVulkanDevice* Device, FVulkanResourceMultiBuffer* InStructuredBuffer, uint32 InOffset = 0);
 
+#if VULKAN_RHI_RAYTRACING
+	FVulkanShaderResourceView(FVulkanDevice* Device, FVulkanAccelerationStructureBuffer* InSourceBuffer, uint32 InOffset = 0);
+#endif // VULKAN_RHI_RAYTRACING
+
 	void Clear();
 
 	void Rename(FRHIResource* InRHIBuffer, FVulkanResourceMultiBuffer* InSourceBuffer, uint32 InSize, EPixelFormat InFormat);
@@ -1300,13 +1304,13 @@ public:
 		return BufferViews[BufferIndex];
 	}
 
-	EPixelFormat BufferViewFormat;
+	EPixelFormat BufferViewFormat = PF_Unknown;
 	ERHITextureSRVOverrideSRGBType SRGBOverride = SRGBO_Default;
 
 	// The texture that this SRV come from
 	TRefCountPtr<FRHITexture> SourceTexture;
 	FVulkanTextureView TextureView;
-	FVulkanResourceMultiBuffer* SourceStructuredBuffer;
+	FVulkanResourceMultiBuffer* SourceStructuredBuffer = nullptr;
 	uint32 MipLevel = 0;
 	uint32 NumMips = MAX_uint32;
 	uint32 FirstArraySlice = 0;
@@ -1316,13 +1320,17 @@ public:
 
 	TArray<TRefCountPtr<FVulkanBufferView>> BufferViews;
 	uint32 BufferIndex = 0;
-	uint32 Size;
+	uint32 Size = 0;
 	uint32 Offset = 0;
 	// The buffer this SRV comes from (can be null)
-	FVulkanResourceMultiBuffer* SourceBuffer;
+	FVulkanResourceMultiBuffer* SourceBuffer = nullptr;
 	// To keep a reference
 	TRefCountPtr<FRHIResource> SourceRHIBuffer;
 
+#if VULKAN_RHI_RAYTRACING
+	VkAccelerationStructureKHR AccelerationStructureHandle = VK_NULL_HANDLE;
+#endif // VULKAN_RHI_RAYTRACING
+
 protected:
 	// Used to check on volatile buffers if a new BufferView is required
 	VkBuffer VolatileBufferHandle = VK_NULL_HANDLE;
diff --git a/Engine/Source/Runtime/VulkanRHI/Private/VulkanShaders.cpp b/Engine/Source/Runtime/VulkanRHI/Private/VulkanShaders.cpp
--- a/Engine/Source/Runtime/VulkanRHI/Private/VulkanShaders.cpp
+++ b/Engine/Source/Runtime/VulkanRHI/Private/VulkanShaders.cpp
@@ -486,6 +486,10 @@ bool FVulkanGfxLayout::UsesInputAttachment(FVulkanShaderHeader::EAttachmentType
 uint32 FVulkanDescriptorSetWriter::SetupDescriptorWrites(
 	const TArray<VkDescriptorType>& Types, FVulkanHashableDescriptorInfo* InHashableDescriptorInfos,
 	VkWriteDescriptorSet* InWriteDescriptors, VkDescriptorImageInfo* InImageInfo, VkDescriptorBufferInfo* InBufferInfo, uint8* InBindingToDynamicOffsetMap,
+#if VULKAN_RHI_RAYTRACING
+	VkWriteDescriptorSetAccelerationStructureKHR* InAccelerationStructuresWriteDescriptors,
+	VkAccelerationStructureKHR* InAccelerationStructures,
+#endif // VULKAN_RHI_RAYTRACING
 	const FVulkanSamplerState& DefaultSampler, const FVulkanTextureView& DefaultImageView)
 {
 	HashableDescriptorInfos = InHashableDescriptorInfos;
@@ -539,6 +543,15 @@ uint32 FVulkanDescriptorSetWriter::SetupDescriptorWrites(
 		case VK_DESCRIPTOR_TYPE_STORAGE_TEXEL_BUFFER:
 		case VK_DESCRIPTOR_TYPE_UNIFORM_TEXEL_BUFFER:
 			break;
+#if VULKAN_RHI_RAYTRACING
+		case VK_DESCRIPTOR_TYPE_ACCELERATION_STRUCTURE_KHR:
+			InAccelerationStructuresWriteDescriptors->sType = VK_STRUCTURE_TYPE_WRITE_DESCRIPTOR_SET_ACCELERATION_STRUCTURE_KHR;
+			InAccelerationStructuresWriteDescriptors->pNext = nullptr;
+			InAccelerationStructuresWriteDescriptors->accelerationStructureCount = 1;
+			InAccelerationStructuresWriteDescriptors->pAccelerationStructures = InAccelerationStructures++;
+			InWriteDescriptors->pNext = InAccelerationStructuresWriteDescriptors++;
+			break;
+#endif // VULKAN_RHI_RAYTRACING
 		default:
 			checkf(0, TEXT("Unsupported descriptor type %d"), (int32)Types[Index]);
 			break;
diff --git a/Engine/Source/Runtime/VulkanRHI/Private/VulkanUAV.cpp b/Engine/Source/Runtime/VulkanRHI/Private/VulkanUAV.cpp
--- a/Engine/Source/Runtime/VulkanRHI/Private/VulkanUAV.cpp
+++ b/Engine/Source/Runtime/VulkanRHI/Private/VulkanUAV.cpp
@@ -4,6 +4,10 @@
 #include "VulkanContext.h"
 #include "ClearReplacementShaders.h"
 
+#if VULKAN_RHI_RAYTRACING
+#include "VulkanRayTracing.h"
+#endif // VULKAN_RHI_RAYTRACING
+
 FVulkanShaderResourceView::FVulkanShaderResourceView(FVulkanDevice* Device, FRHIResource* InRHIBuffer, FVulkanResourceMultiBuffer* InSourceBuffer, uint32 InSize, EPixelFormat InFormat, uint32 InOffset)
 	: VulkanRHI::FVulkanViewBase(Device)
 	, BufferViewFormat(InFormat)
@@ -54,7 +58,24 @@ FVulkanShaderResourceView::FVulkanShaderResourceView(FVulkanDevice* Device, FVul
 {
 }
 
+#if VULKAN_RHI_RAYTRACING
+FVulkanShaderResourceView::FVulkanShaderResourceView(FVulkanDevice* InDevice, FVulkanAccelerationStructureBuffer* InSourceBuffer, uint32 InOffset)
+	: VulkanRHI::FVulkanViewBase(InDevice)
+	, SourceRHIBuffer(InSourceBuffer)
+{
+	check(InDevice);
 
+	VkAccelerationStructureCreateInfoKHR CreateInfo;
+	ZeroVulkanStruct(CreateInfo, VK_STRUCTURE_TYPE_ACCELERATION_STRUCTURE_CREATE_INFO_KHR);
+	CreateInfo.buffer = InSourceBuffer->GetBuffer();
+	CreateInfo.offset = InOffset;
+	CreateInfo.size = InSourceBuffer->GetSize() - InOffset;
+	CreateInfo.type = VK_ACCELERATION_STRUCTURE_TYPE_TOP_LEVEL_KHR;
+
+	VkDevice NativeDevice = InDevice->GetInstanceHandle();
+	VERIFYVULKANRESULT(VulkanDynamicAPI::vkCreateAccelerationStructureKHR(NativeDevice, &CreateInfo, VULKAN_CPU_ALLOCATOR, &AccelerationStructureHandle));
+}
+#endif // VULKAN_RHI_RAYTRACING
 
 FVulkanShaderResourceView::~FVulkanShaderResourceView()
 {
@@ -70,6 +91,13 @@ FVulkanShaderResourceView::~FVulkanShaderResourceView()
 
 void FVulkanShaderResourceView::Clear()
 {
+#if VULKAN_RHI_RAYTRACING
+	if (AccelerationStructureHandle)
+	{
+		Device->GetDeferredDeletionQueue().EnqueueResource(VulkanRHI::FDeferredDeletionQueue2::EType::AccelerationStructure, AccelerationStructureHandle);
+	}
+#endif // VULKAN_RHI_RAYTRACING
+
 	SourceRHIBuffer = nullptr;
 	SourceBuffer = nullptr;
 	BufferViews.Empty();
@@ -88,6 +116,11 @@ void FVulkanShaderResourceView::Rename(FRHIResource* InRHIBuffer, FVulkanResourc
 {
 	check(Device);
 	check(!Offset);
+
+#if VULKAN_RHI_RAYTRACING
+	checkf(!AccelerationStructureHandle, TEXT("Acceleration structure view renaming is currently not supported"));
+#endif //VULKAN_RHI_RAYTRACING
+
 	BufferViewFormat = InFormat;
 	SourceTexture = nullptr;
 	TextureView.Destroy(*Device);
@@ -154,6 +187,12 @@ void FVulkanShaderResourceView::UpdateView()
 	{
 		// Nothing...
 	}
+#if VULKAN_RHI_RAYTRACING
+	else if (AccelerationStructureHandle)
+	{
+		// Nothing
+	}
+#endif //VULKAN_RHI_RAYTRACING
 	else
 	{
 		if (TextureView.View == VK_NULL_HANDLE)
@@ -377,8 +416,9 @@ FShaderResourceViewRHIRef FVulkanDynamicRHI::RHICreateShaderResourceView(const F
 #if VULKAN_RHI_RAYTRACING
 		case FShaderResourceViewInitializer::EType::AccelerationStructureSRV:
 		{
-			checkNoEntry(); // TODO
-			return nullptr;
+			check(Desc.Buffer);
+			FVulkanAccelerationStructureBuffer* ASBuffer = static_cast<FVulkanAccelerationStructureBuffer*>(Desc.Buffer);
+			return new FVulkanShaderResourceView(Device, ASBuffer, Desc.StartOffsetBytes);
 		}
 #endif // D3D12_RHI_RAYTRACING
 	}
@@ -403,8 +443,8 @@ FShaderResourceViewRHIRef FVulkanDynamicRHI::RHICreateShaderResourceView(FRHIBuf
 #if VULKAN_RHI_RAYTRACING
 	else if (BufferRHI && EnumHasAnyFlags(BufferRHI->GetUsage(), BUF_AccelerationStructure))
 	{
-		FVulkanAccelerationStructureBuffer* Buffer = static_cast<FVulkanAccelerationStructureBuffer*>(BufferRHI);
-		return nullptr; // TODO
+		FVulkanAccelerationStructureBuffer* ASBuffer = static_cast<FVulkanAccelerationStructureBuffer*>(BufferRHI);
+		return new FVulkanShaderResourceView(Device, ASBuffer, 0);
 	}
 #endif
 	else
