Fix compile error due to missing ParallelFor.h include

diff --git a/Engine/Source/Runtime/VulkanRHI/Private/VulkanRayTracing.cpp b/Engine/Source/Runtime/VulkanRHI/Private/VulkanRayTracing.cpp
--- a/Engine/Source/Runtime/VulkanRHI/Private/VulkanRayTracing.cpp
+++ b/Engine/Source/Runtime/VulkanRHI/Private/VulkanRayTracing.cpp
@@ -8,6 +8,7 @@
 #include "VulkanDescriptorSets.h"
 #include "BuiltInRayTracingShaders.h"
 #include "Experimental/Containers/SherwoodHashTable.h"
+#include "Async/ParallelFor.h"
 
 TAutoConsoleVariable<int32> GVulkanExperimentalRayTracing(
 	TEXT("r.Vulkan.ExperimentalRayTracing"),
