Enable inline ray tracing support for Vulkan on PC via DDPI. Feature remains disabled by default via r.Vulkan.RayTracing CVar.

* Minimal set of required ray tracing utility shaders is now always cooked for Windows Vulkan targets
* -noraytracing command line argument can also now be used to disable ray tracing (useful when testing a game that was packaged with r.Vulkan.RayTracing=1)
* Rename r.Vulkan.ExperimentalRayTracing -> r.Vulkan.RayTracing

diff --git a/Engine/Config/VulkanPC/DataDrivenPlatformInfo.ini b/Engine/Config/VulkanPC/DataDrivenPlatformInfo.ini
--- a/Engine/Config/VulkanPC/DataDrivenPlatformInfo.ini
+++ b/Engine/Config/VulkanPC/DataDrivenPlatformInfo.ini
@@ -49,9 +49,9 @@ bSupports4ComponentUAVReadWrite = false
 ; //#todo-vulkanddpi
 bSupportsRenderTargetWriteMask = false
 
-bSupportsRayTracing = false
+bSupportsRayTracing = true
 bSupportsRayTracingShaders = false
-bSupportsInlineRayTracing = false
+bSupportsInlineRayTracing = true
 bSupportsRayTracingIndirectInstanceData = false
 bSupportsGPUSkinCache = true
 bSupportsByteBufferComputeShaders = true
diff --git a/Engine/Source/Runtime/VulkanRHI/Private/VulkanRayTracing.cpp b/Engine/Source/Runtime/VulkanRHI/Private/VulkanRayTracing.cpp
--- a/Engine/Source/Runtime/VulkanRHI/Private/VulkanRayTracing.cpp
+++ b/Engine/Source/Runtime/VulkanRHI/Private/VulkanRayTracing.cpp
@@ -9,9 +9,11 @@
 #include "BuiltInRayTracingShaders.h"
 #include "Experimental/Containers/SherwoodHashTable.h"
 #include "Async/ParallelFor.h"
+#include "Misc/CommandLine.h"
+#include "Misc/Parse.h"
 
-TAutoConsoleVariable<int32> GVulkanExperimentalRayTracing(
-	TEXT("r.Vulkan.ExperimentalRayTracing"),
+TAutoConsoleVariable<int32> GVulkanRayTracingCVar(
+	TEXT("r.Vulkan.RayTracing"),
 	0,
 	TEXT("0: Do not enable Vulkan ray tracing extensions (default)\n")
 	TEXT("1: Enable experimental ray tracing support (for development and testing purposes)"),
@@ -24,7 +26,7 @@ ENUM_VK_ENTRYPOINTS_RAYTRACING(DEFINE_VK_ENTRYPOINTS)
 
 void FVulkanRayTracingPlatform::GetDeviceExtensions(EGpuVendorId VendorId, TArray<const ANSICHAR*>& OutExtensions)
 {
-	if (GVulkanExperimentalRayTracing.GetValueOnAnyThread() == 0)
+	if (!GVulkanRayTracingCVar.GetValueOnAnyThread() || FParse::Param(FCommandLine::Get(), TEXT("noraytracing")))
 	{
 		return;
 	}
