Fix VulkanRT validation errors
- set VK_BUFFER_USAGE_ACCELERATION_STRUCTURE_BUILD_INPUT_READ_ONLY_BIT_KHR on buffers
- enable rayTraversalPrimitiveCulling feature
- create scratch buffers with BUF_StructuredBuffer (VK_BUFFER_USAGE_STORAGE_BUFFER_BIT) usage

diff --git a/Engine/Source/Runtime/VulkanRHI/Private/VulkanBuffer.cpp b/Engine/Source/Runtime/VulkanRHI/Private/VulkanBuffer.cpp
--- a/Engine/Source/Runtime/VulkanRHI/Private/VulkanBuffer.cpp
+++ b/Engine/Source/Runtime/VulkanRHI/Private/VulkanBuffer.cpp
@@ -128,7 +128,15 @@ FVulkanResourceMultiBuffer::FVulkanResourceMultiBuffer(FVulkanDevice* InDevice,
 		BufferUsageFlags |= bCPUReadable ? (VK_BUFFER_USAGE_TRANSFER_SRC_BIT | VK_BUFFER_USAGE_TRANSFER_DST_BIT) : 0;
 
 #if VULKAN_RHI_RAYTRACING
-		BufferUsageFlags |= VK_BUFFER_USAGE_SHADER_DEVICE_ADDRESS_BIT;
+		if (InDevice->GetOptionalExtensions().HasRaytracingExtensions())
+		{
+			BufferUsageFlags |= VK_BUFFER_USAGE_SHADER_DEVICE_ADDRESS_BIT;
+
+			if (!EnumHasAnyFlags(InUEUsage, BUF_AccelerationStructure))
+			{
+				BufferUsageFlags |= VK_BUFFER_USAGE_ACCELERATION_STRUCTURE_BUILD_INPUT_READ_ONLY_BIT_KHR;
+			}
+		}
 #endif
 
 		if (bVolatile)
diff --git a/Engine/Source/Runtime/VulkanRHI/Private/VulkanMemory.cpp b/Engine/Source/Runtime/VulkanRHI/Private/VulkanMemory.cpp
--- a/Engine/Source/Runtime/VulkanRHI/Private/VulkanMemory.cpp
+++ b/Engine/Source/Runtime/VulkanRHI/Private/VulkanMemory.cpp
@@ -981,11 +981,14 @@ namespace VulkanRHI
 #endif // VULKAN_SUPPORTS_EXTERNAL_MEMORY
 
 #if VULKAN_RHI_RAYTRACING
-		VkMemoryAllocateFlagsInfo MemoryAllocateFlagsInfo;
-		ZeroVulkanStruct(MemoryAllocateFlagsInfo, VK_STRUCTURE_TYPE_MEMORY_ALLOCATE_FLAGS_INFO);
-		MemoryAllocateFlagsInfo.flags = VK_MEMORY_ALLOCATE_DEVICE_ADDRESS_BIT_KHR;
-		MemoryAllocateFlagsInfo.pNext = Info.pNext;
-		Info.pNext = &MemoryAllocateFlagsInfo;
+		if (Device->GetOptionalExtensions().HasRaytracingExtensions())
+		{
+			VkMemoryAllocateFlagsInfo MemoryAllocateFlagsInfo;
+			ZeroVulkanStruct(MemoryAllocateFlagsInfo, VK_STRUCTURE_TYPE_MEMORY_ALLOCATE_FLAGS_INFO);
+			MemoryAllocateFlagsInfo.flags = VK_MEMORY_ALLOCATE_DEVICE_ADDRESS_BIT_KHR;
+			MemoryAllocateFlagsInfo.pNext = Info.pNext;
+			Info.pNext = &MemoryAllocateFlagsInfo;
+		}
 #endif
 
 		VkDeviceMemory Handle;
diff --git a/Engine/Source/Runtime/VulkanRHI/Private/VulkanRayTracing.cpp b/Engine/Source/Runtime/VulkanRHI/Private/VulkanRayTracing.cpp
--- a/Engine/Source/Runtime/VulkanRHI/Private/VulkanRayTracing.cpp
+++ b/Engine/Source/Runtime/VulkanRHI/Private/VulkanRayTracing.cpp
@@ -65,6 +65,7 @@ void FVulkanRayTracingPlatform::EnablePhysicalDeviceFeatureExtensions(VkDeviceCr
 		
 		ZeroVulkanStruct(Features.RayTracingPipelineFeatures, VK_STRUCTURE_TYPE_PHYSICAL_DEVICE_RAY_TRACING_PIPELINE_FEATURES_KHR);
 		Features.RayTracingPipelineFeatures.rayTracingPipeline = VK_TRUE;
+		Features.RayTracingPipelineFeatures.rayTraversalPrimitiveCulling = VK_TRUE;
 		Features.RayTracingPipelineFeatures.pNext = &Features.AccelerationStructureFeatures;
 
 		ZeroVulkanStruct(Features.RayQueryFeatures, VK_STRUCTURE_TYPE_PHYSICAL_DEVICE_RAY_QUERY_FEATURES_KHR);
@@ -310,7 +311,7 @@ FVulkanRayTracingGeometry::FVulkanRayTracingGeometry(FRayTracingGeometryInitiali
 	AccelerationStructureBuffer = static_cast<FVulkanAccelerationStructureBuffer*>(RHICreateBuffer(BuildData.SizesInfo.accelerationStructureSize, BUF_AccelerationStructure, 0, ERHIAccess::BVHWrite, BlasBufferCreateInfo).GetReference());
 
 	FRHIResourceCreateInfo ScratchBufferCreateInfo(TEXT("BuildScratchBLAS"));
-	ScratchBuffer = ResourceCast(RHICreateBuffer(BuildData.SizesInfo.buildScratchSize, BUF_UnorderedAccess, 0, ERHIAccess::UAVCompute, ScratchBufferCreateInfo).GetReference());
+	ScratchBuffer = ResourceCast(RHICreateBuffer(BuildData.SizesInfo.buildScratchSize, BUF_UnorderedAccess | BUF_StructuredBuffer, 0, ERHIAccess::UAVCompute, ScratchBufferCreateInfo).GetReference());
 
 	VkDevice NativeDevice = Device->GetInstanceHandle();
 
@@ -479,7 +480,7 @@ void FVulkanRayTracingScene::BuildAccelerationStructure(
 	if (!bExternalScratchBuffer)
 	{
 		FRHIResourceCreateInfo ScratchBufferCreateInfo(TEXT("BuildScratchTLAS"));
-		ScratchBuffer = ResourceCast(RHICreateBuffer(BuildData.SizesInfo.buildScratchSize, BUF_UnorderedAccess, 0, ERHIAccess::UAVCompute, ScratchBufferCreateInfo).GetReference());
+		ScratchBuffer = ResourceCast(RHICreateBuffer(BuildData.SizesInfo.buildScratchSize, BUF_UnorderedAccess | BUF_StructuredBuffer, 0, ERHIAccess::UAVCompute, ScratchBufferCreateInfo).GetReference());
 		InScratchBuffer = ScratchBuffer.GetReference();
 	}
 
