VulkanRT: Allocate AS buffers using FVulkanResourceMultiBuffer
- FVulkanAccelerationStructureBuffer
- Track memory usage using STAT_RTAccelerationStructureMemory

diff --git a/Engine/Source/Runtime/Renderer/Private/RayTracing/RayTracingScene.cpp b/Engine/Source/Runtime/Renderer/Private/RayTracing/RayTracingScene.cpp
--- a/Engine/Source/Runtime/Renderer/Private/RayTracing/RayTracingScene.cpp
+++ b/Engine/Source/Runtime/Renderer/Private/RayTracing/RayTracingScene.cpp
@@ -59,7 +59,7 @@ void FRayTracingScene::Create(FRDGBuilder& GraphBuilder, const FGPUScene& GPUSce
 		RayTracingSceneSRV = nullptr;
 		RayTracingSceneBuffer = nullptr;
 
-		FRHIResourceCreateInfo CreateInfo(TEXT("RayTracingSceneBuffer"));
+		FRHIResourceCreateInfo CreateInfo(TEXT("FRayTracingScene::SceneBuffer"));
 		RayTracingSceneBuffer = RHICreateBuffer(uint32(SizeInfo.ResultSize), BUF_AccelerationStructure, 0, ERHIAccess::BVHWrite, CreateInfo);
 		RayTracingSceneSRV = RHICreateShaderResourceView(RayTracingSceneBuffer);
 	}
@@ -72,7 +72,7 @@ void FRayTracingScene::Create(FRDGBuilder& GraphBuilder, const FGPUScene& GPUSce
 		ScratchBufferDesc.BytesPerElement = uint32(ScratchAlignment);
 		ScratchBufferDesc.NumElements = uint32(FMath::DivideAndRoundUp(SizeInfo.BuildScratchSize, ScratchAlignment));
 
-		BuildScratchBuffer = GraphBuilder.CreateBuffer(ScratchBufferDesc, TEXT("RayTracingScratchBuffer"));
+		BuildScratchBuffer = GraphBuilder.CreateBuffer(ScratchBufferDesc, TEXT("FRayTracingScene::ScratchBuffer"));
 	}
 
 	{
@@ -92,7 +92,7 @@ void FRayTracingScene::Create(FRDGBuilder& GraphBuilder, const FGPUScene& GPUSce
 
 		if (AccelerationStructureAddressesBuffer.NumBytes < AccelerationStructureAddressesBufferSize)
 		{
-			AccelerationStructureAddressesBuffer.Initialize(TEXT("RayTracingAccelerationStructureAddressesBuffer"), AccelerationStructureAddressesBufferSize, BUF_Volatile);
+			AccelerationStructureAddressesBuffer.Initialize(TEXT("FRayTracingScene::AccelerationStructureAddressesBuffer"), AccelerationStructureAddressesBufferSize, BUF_Volatile);
 		}
 	}
 
@@ -104,7 +104,7 @@ void FRayTracingScene::Create(FRDGBuilder& GraphBuilder, const FGPUScene& GPUSce
 			|| UploadBufferSize > InstanceUploadBuffer->GetSize()
 			|| UploadBufferSize < InstanceUploadBuffer->GetSize() / 2)
 		{
-			FRHIResourceCreateInfo CreateInfo(TEXT("RayTracingSceneInstanceUploadBuffer"));
+			FRHIResourceCreateInfo CreateInfo(TEXT("FRayTracingScene::InstanceUploadBuffer"));
 			InstanceUploadBuffer = RHICreateStructuredBuffer(sizeof(FRayTracingInstanceDescriptorInput), UploadBufferSize, BUF_ShaderResource | BUF_Volatile, CreateInfo);
 			InstanceUploadSRV = RHICreateShaderResourceView(InstanceUploadBuffer);
 		}
@@ -118,7 +118,7 @@ void FRayTracingScene::Create(FRDGBuilder& GraphBuilder, const FGPUScene& GPUSce
 			|| UploadBufferSize > TransformUploadBuffer->GetSize()
 			|| UploadBufferSize < TransformUploadBuffer->GetSize() / 2)
 		{
-			FRHIResourceCreateInfo CreateInfo(TEXT("RayTracingSceneTransformUploadBuffer"));
+			FRHIResourceCreateInfo CreateInfo(TEXT("FRayTracingScene::TransformUploadBuffer"));
 			TransformUploadBuffer = RHICreateStructuredBuffer(sizeof(FVector4f), UploadBufferSize, BUF_ShaderResource | BUF_Volatile, CreateInfo);
 			TransformUploadSRV = RHICreateShaderResourceView(TransformUploadBuffer);
 		}
diff --git a/Engine/Source/Runtime/VulkanRHI/Private/VulkanBuffer.cpp b/Engine/Source/Runtime/VulkanRHI/Private/VulkanBuffer.cpp
--- a/Engine/Source/Runtime/VulkanRHI/Private/VulkanBuffer.cpp
+++ b/Engine/Source/Runtime/VulkanRHI/Private/VulkanBuffer.cpp
@@ -19,6 +19,7 @@ static FORCEINLINE void UpdateVulkanBufferStats(uint64_t Size, VkBufferUsageFlag
 	const bool bUniformBuffer = !!(Usage & VK_BUFFER_USAGE_UNIFORM_BUFFER_BIT);
 	const bool bIndexBuffer = !!(Usage & VK_BUFFER_USAGE_INDEX_BUFFER_BIT);
 	const bool bVertexBuffer = !!(Usage & VK_BUFFER_USAGE_VERTEX_BUFFER_BIT);
+	const bool bAccelerationStructure = !!(Usage & VK_BUFFER_USAGE_ACCELERATION_STRUCTURE_STORAGE_BIT_KHR);
 
 	if (Allocating)
 	{
@@ -34,6 +35,10 @@ static FORCEINLINE void UpdateVulkanBufferStats(uint64_t Size, VkBufferUsageFlag
 		{
 			INC_MEMORY_STAT_BY(STAT_VertexBufferMemory, Size);
 		}
+		else if (bAccelerationStructure)
+		{
+			INC_MEMORY_STAT_BY(STAT_RTAccelerationStructureMemory, Size);
+		}
 		else
 		{
 			INC_MEMORY_STAT_BY(STAT_StructuredBufferMemory, Size);
@@ -53,6 +58,10 @@ static FORCEINLINE void UpdateVulkanBufferStats(uint64_t Size, VkBufferUsageFlag
 		{
 			DEC_MEMORY_STAT_BY(STAT_VertexBufferMemory, Size);
 		}
+		else if (bAccelerationStructure)
+		{
+			INC_MEMORY_STAT_BY(STAT_RTAccelerationStructureMemory, Size);
+		}
 		else
 		{
 			DEC_MEMORY_STAT_BY(STAT_StructuredBufferMemory, Size);
@@ -60,30 +69,6 @@ static FORCEINLINE void UpdateVulkanBufferStats(uint64_t Size, VkBufferUsageFlag
 	}
 }
 
-#if VULKAN_RHI_RAYTRACING
-FVulkanAccelerationStructureBuffer::FVulkanAccelerationStructureBuffer(FVulkanDevice* InDevice, uint32 InSize, EBufferUsageFlags InUEUsage, uint32 InStride, FRHIResourceCreateInfo& CreateInfo)
-	: FRHIBuffer(InSize, InUEUsage, InStride)
-{
-	check(EnumHasAnyFlags(InUEUsage, BUF_AccelerationStructure));
-	if (InSize == 0)
-	{
-		return;
-	}
-	
-	FVulkanRayTracingAllocator::Allocate(
-		InDevice,
-		InSize,
-		VK_BUFFER_USAGE_ACCELERATION_STRUCTURE_STORAGE_BIT_KHR | VK_BUFFER_USAGE_SHADER_DEVICE_ADDRESS_BIT,
-		VK_MEMORY_PROPERTY_DEVICE_LOCAL_BIT,
-		Allocation);
-}
-
-FVulkanAccelerationStructureBuffer::~FVulkanAccelerationStructureBuffer()
-{
-	FVulkanRayTracingAllocator::Free(Allocation);
-}
-#endif // VULKAN_RHI_RAYTRACING
-
 FVulkanResourceMultiBuffer::FVulkanResourceMultiBuffer(FVulkanDevice* InDevice, uint32 InSize, EBufferUsageFlags InUEUsage, uint32 InStride, FRHIResourceCreateInfo& CreateInfo, class FRHICommandListImmediate* InRHICmdList)
 	: FRHIBuffer(InSize, InUEUsage, InStride)
 	, VulkanRHI::FDeviceChild(InDevice)
@@ -109,6 +94,12 @@ FVulkanResourceMultiBuffer::FVulkanResourceMultiBuffer(FVulkanDevice* InDevice,
 	{
 		BufferUsageFlags |= VK_BUFFER_USAGE_STORAGE_BUFFER_BIT;
 	}
+#if VULKAN_RHI_RAYTRACING
+	if (EnumHasAnyFlags(InUEUsage, BUF_AccelerationStructure))
+	{
+		BufferUsageFlags |= VK_BUFFER_USAGE_ACCELERATION_STRUCTURE_STORAGE_BIT_KHR;
+	}
+#endif
 
 	if (InSize > 0)
 	{
@@ -474,13 +465,6 @@ FBufferRHIRef FVulkanDynamicRHI::RHICreateBuffer(uint32 Size, EBufferUsageFlags
 {
 	LLM_SCOPE_VULKAN(ELLMTagVulkan::VulkanBuffers);
 
-#if VULKAN_RHI_RAYTRACING
-	if (EnumHasAnyFlags(Usage, BUF_AccelerationStructure))
-	{
-		return new FVulkanAccelerationStructureBuffer(Device, Size, Usage, Stride, CreateInfo);
-	}
-#endif
-
 	if (CreateInfo.bWithoutNativeResource)
 	{
 		return new FVulkanResourceMultiBuffer(nullptr, 0, BUF_None, 0, CreateInfo, nullptr);
diff --git a/Engine/Source/Runtime/VulkanRHI/Private/VulkanMemory.cpp b/Engine/Source/Runtime/VulkanRHI/Private/VulkanMemory.cpp
--- a/Engine/Source/Runtime/VulkanRHI/Private/VulkanMemory.cpp
+++ b/Engine/Source/Runtime/VulkanRHI/Private/VulkanMemory.cpp
@@ -2736,11 +2736,16 @@ namespace VulkanRHI
 		else
 		{
 			bool bIsVertexOrIndexBuffer = (BufferUsageFlags & (VK_BUFFER_USAGE_VERTEX_BUFFER_BIT | VK_BUFFER_USAGE_INDEX_BUFFER_BIT)) != 0;
+			bool bIsAccelerationStructureBuffer = (BufferUsageFlags & VK_BUFFER_USAGE_ACCELERATION_STRUCTURE_STORAGE_BIT_KHR) != 0;
 			if (bIsVertexOrIndexBuffer)
 			{
 				// No alignment restrictions on Vertex or Index buffers, can live on CPU mem
 				Priority = VULKAN_MEMORY_LOW_PRIORITY;
 			}
+			else if (bIsAccelerationStructureBuffer)
+			{
+				Alignment = FMath::Max(Alignment, GRHIRayTracingAccelerationStructureAlignment);
+			}
 			else
 			{
 				// Uniform buffer
diff --git a/Engine/Source/Runtime/VulkanRHI/Private/VulkanRayTracing.cpp b/Engine/Source/Runtime/VulkanRHI/Private/VulkanRayTracing.cpp
--- a/Engine/Source/Runtime/VulkanRHI/Private/VulkanRayTracing.cpp
+++ b/Engine/Source/Runtime/VulkanRHI/Private/VulkanRayTracing.cpp
@@ -292,8 +292,6 @@ FVulkanRayTracingGeometry::FVulkanRayTracingGeometry(FRayTracingGeometryInitiali
 	check(Initializer.GeometryType == ERayTracingGeometryType::RTGT_Triangles);
 	checkf(!Initializer.IndexBuffer || (Initializer.IndexBuffer->GetStride() == 2 || Initializer.IndexBuffer->GetStride() == 4), TEXT("Index buffer must be 16 or 32 bit if in use."));
 
-	IndexStrideInBytes = Initializer.IndexBuffer ? Initializer.IndexBuffer->GetStride() : 0;
-
 	FVkRtBLASBuildData BuildData;
 	GetBLASBuildData(
 		Device->GetInstanceHandle(), 
@@ -302,13 +300,13 @@ FVulkanRayTracingGeometry::FVulkanRayTracingGeometry(FRayTracingGeometryInitiali
 		Initializer.IndexBufferOffset,
 		Initializer.bFastBuild,
 		Initializer.bAllowUpdate,
-		IndexStrideInBytes, 
+		Initializer.IndexBuffer ? Initializer.IndexBuffer->GetStride() : 0,
 		EAccelerationStructureBuildMode::Build,
 		BuildData);
 
 	FString DebugNameString = Initializer.DebugName.ToString();
 	FRHIResourceCreateInfo BlasBufferCreateInfo(*DebugNameString);
-	AccelerationStructureBuffer = static_cast<FVulkanAccelerationStructureBuffer*>(RHICreateBuffer(BuildData.SizesInfo.accelerationStructureSize, BUF_AccelerationStructure, 0, ERHIAccess::BVHWrite, BlasBufferCreateInfo).GetReference());
+	AccelerationStructureBuffer = ResourceCast(RHICreateBuffer(BuildData.SizesInfo.accelerationStructureSize, BUF_AccelerationStructure, 0, ERHIAccess::BVHWrite, BlasBufferCreateInfo).GetReference());
 
 	FRHIResourceCreateInfo ScratchBufferCreateInfo(TEXT("BuildScratchBLAS"));
 	ScratchBuffer = ResourceCast(RHICreateBuffer(BuildData.SizesInfo.buildScratchSize, BUF_UnorderedAccess | BUF_StructuredBuffer, 0, ERHIAccess::UAVCompute, ScratchBufferCreateInfo).GetReference());
@@ -317,7 +315,8 @@ FVulkanRayTracingGeometry::FVulkanRayTracingGeometry(FRayTracingGeometryInitiali
 
 	VkAccelerationStructureCreateInfoKHR CreateInfo;
 	ZeroVulkanStruct(CreateInfo, VK_STRUCTURE_TYPE_ACCELERATION_STRUCTURE_CREATE_INFO_KHR);
-	CreateInfo.buffer = AccelerationStructureBuffer->GetBuffer();
+	CreateInfo.buffer = AccelerationStructureBuffer->GetHandle();
+	CreateInfo.offset = AccelerationStructureBuffer->GetOffset();
 	CreateInfo.size = BuildData.SizesInfo.accelerationStructureSize;
 	CreateInfo.type = VK_ACCELERATION_STRUCTURE_TYPE_BOTTOM_LEVEL_KHR;
 	VERIFYVULKANRESULT(vkCreateAccelerationStructureKHR(NativeDevice, &CreateInfo, VULKAN_CPU_ALLOCATOR, &Handle));
@@ -346,7 +345,7 @@ void FVulkanRayTracingGeometry::BuildAccelerationStructure(FVulkanCommandListCon
 		Initializer.IndexBufferOffset,
 		Initializer.bFastBuild,
 		Initializer.bAllowUpdate,
-		IndexStrideInBytes,
+		Initializer.IndexBuffer ? Initializer.IndexBuffer->GetStride() : 0,
 		BuildMode,
 		BuildData);
 
@@ -446,7 +445,7 @@ void FVulkanRayTracingScene::BindBuffer(FRHIBuffer* InBuffer, uint32 InBufferOff
 
 	check(SizeInfo.ResultSize + InBufferOffset <= InBuffer->GetSize());
 	check(InBufferOffset % 256 == 0); // Spec requires offset to be a multiple of 256
-	AccelerationStructureBuffer = static_cast<FVulkanAccelerationStructureBuffer*>(InBuffer);
+	AccelerationStructureBuffer = ResourceCast(InBuffer);
 
 	FShaderResourceViewInitializer ViewInitializer(InBuffer, InBufferOffset, 0);
 	AccelerationStructureView = new FVulkanShaderResourceView(Device, AccelerationStructureBuffer, InBufferOffset);
diff --git a/Engine/Source/Runtime/VulkanRHI/Private/VulkanRayTracing.h b/Engine/Source/Runtime/VulkanRHI/Private/VulkanRayTracing.h
--- a/Engine/Source/Runtime/VulkanRHI/Private/VulkanRayTracing.h
+++ b/Engine/Source/Runtime/VulkanRHI/Private/VulkanRayTracing.h
@@ -35,6 +35,13 @@ public:
 	static bool LoadVulkanInstanceFunctions(VkInstance inInstance);
 };
 
+struct FVkRtAllocation
+{
+	VkDevice Device = VK_NULL_HANDLE;
+	VkDeviceMemory Memory = VK_NULL_HANDLE;
+	VkBuffer Buffer = VK_NULL_HANDLE;
+};
+
 class FVulkanRayTracingAllocator
 {
 public:
@@ -85,11 +92,10 @@ private:
 	FVulkanDevice* const Device = nullptr;
 
 	const FRayTracingGeometryInitializer Initializer;
-	uint32 IndexStrideInBytes  = 0;
 
 	VkAccelerationStructureKHR Handle = VK_NULL_HANDLE;
 	VkDeviceAddress Address = 0;
-	TRefCountPtr<FVulkanAccelerationStructureBuffer> AccelerationStructureBuffer;
+	TRefCountPtr<FVulkanResourceMultiBuffer> AccelerationStructureBuffer;
 	TRefCountPtr<FVulkanResourceMultiBuffer> ScratchBuffer;
 };
 
@@ -125,7 +131,7 @@ private:
 	// Many VkAccelerationStructureKHR-s may be created, pointing at the same buffer.
 	TRefCountPtr<FVulkanShaderResourceView> AccelerationStructureView;
 	
-	TRefCountPtr<FVulkanAccelerationStructureBuffer> AccelerationStructureBuffer;
+	TRefCountPtr<FVulkanResourceMultiBuffer> AccelerationStructureBuffer;
 };
 
 class FVulkanRayTracingPipelineState : public FRHIRayTracingPipelineState
diff --git a/Engine/Source/Runtime/VulkanRHI/Private/VulkanResources.h b/Engine/Source/Runtime/VulkanRHI/Private/VulkanResources.h
--- a/Engine/Source/Runtime/VulkanRHI/Private/VulkanResources.h
+++ b/Engine/Source/Runtime/VulkanRHI/Private/VulkanResources.h
@@ -1073,30 +1073,6 @@ protected:
 	friend class FVulkanCommandListContext;
 };
 
-#if VULKAN_RHI_RAYTRACING
-struct FVkRtAllocation
-{
-	VkDevice Device = VK_NULL_HANDLE;
-	VkDeviceMemory Memory = VK_NULL_HANDLE;
-	VkBuffer Buffer = VK_NULL_HANDLE;
-};
-
-class FVulkanAccelerationStructureBuffer : public FRHIBuffer
-{
-public:
-	FVulkanAccelerationStructureBuffer(FVulkanDevice* InDevice, uint32 InSize, EBufferUsageFlags InUEUsage, uint32 InStride, FRHIResourceCreateInfo& CreateInfo);
-	virtual ~FVulkanAccelerationStructureBuffer();
-
-	inline VkBuffer GetBuffer() const
-	{
-		return Allocation.Buffer;
-	}
-
-private:
-	FVkRtAllocation Allocation;
-};
-#endif
-
 class FVulkanResourceMultiBuffer : public FRHIBuffer, public FVulkanEvictable, public VulkanRHI::FDeviceChild
 {
 	virtual void Evict(FVulkanDevice& Device);
@@ -1286,11 +1262,7 @@ class FVulkanShaderResourceView : public FRHIShaderResourceView, public VulkanRH
 public:
 	FVulkanShaderResourceView(FVulkanDevice* Device, FRHIResource* InRHIBuffer, FVulkanResourceMultiBuffer* InSourceBuffer, uint32 InSize, EPixelFormat InFormat, uint32 InOffset = 0);
 	FVulkanShaderResourceView(FVulkanDevice* Device, FRHITexture* InSourceTexture, const FRHITextureSRVCreateInfo& InCreateInfo);
-	FVulkanShaderResourceView(FVulkanDevice* Device, FVulkanResourceMultiBuffer* InStructuredBuffer, uint32 InOffset = 0);
-
-#if VULKAN_RHI_RAYTRACING
-	FVulkanShaderResourceView(FVulkanDevice* Device, FVulkanAccelerationStructureBuffer* InSourceBuffer, uint32 InOffset = 0);
-#endif // VULKAN_RHI_RAYTRACING
+	FVulkanShaderResourceView(FVulkanDevice* Device, FVulkanResourceMultiBuffer* InSourceBuffer, uint32 InOffset = 0);
 
 	void Clear();
 
diff --git a/Engine/Source/Runtime/VulkanRHI/Private/VulkanUAV.cpp b/Engine/Source/Runtime/VulkanRHI/Private/VulkanUAV.cpp
--- a/Engine/Source/Runtime/VulkanRHI/Private/VulkanUAV.cpp
+++ b/Engine/Source/Runtime/VulkanRHI/Private/VulkanUAV.cpp
@@ -46,36 +46,34 @@ FVulkanShaderResourceView::FVulkanShaderResourceView(FVulkanDevice* Device, FRHI
 
 }
 
-FVulkanShaderResourceView::FVulkanShaderResourceView(FVulkanDevice* Device, FVulkanResourceMultiBuffer* InStructuredBuffer, uint32 InOffset)
-	: VulkanRHI::FVulkanViewBase(Device)
-	, BufferViewFormat(PF_Unknown)
-	, SourceTexture(nullptr)
-	, SourceStructuredBuffer(InStructuredBuffer)
-	, NumMips(0)
-	, Size(InStructuredBuffer->GetSize() - InOffset)
-	, Offset(InOffset)
-	, SourceBuffer(nullptr)
+FVulkanShaderResourceView::FVulkanShaderResourceView(FVulkanDevice* InDevice, FVulkanResourceMultiBuffer* InSourceBuffer, uint32 InOffset)
+	: VulkanRHI::FVulkanViewBase(InDevice)
 {
-}
+	check(InDevice && InSourceBuffer);
 
 #if VULKAN_RHI_RAYTRACING
-FVulkanShaderResourceView::FVulkanShaderResourceView(FVulkanDevice* InDevice, FVulkanAccelerationStructureBuffer* InSourceBuffer, uint32 InOffset)
-	: VulkanRHI::FVulkanViewBase(InDevice)
-	, SourceRHIBuffer(InSourceBuffer)
-{
-	check(InDevice);
+	if (EnumHasAnyFlags(InSourceBuffer->GetUsage(), BUF_AccelerationStructure))
+	{
+		SourceRHIBuffer = InSourceBuffer;
 
-	VkAccelerationStructureCreateInfoKHR CreateInfo;
-	ZeroVulkanStruct(CreateInfo, VK_STRUCTURE_TYPE_ACCELERATION_STRUCTURE_CREATE_INFO_KHR);
-	CreateInfo.buffer = InSourceBuffer->GetBuffer();
-	CreateInfo.offset = InOffset;
-	CreateInfo.size = InSourceBuffer->GetSize() - InOffset;
-	CreateInfo.type = VK_ACCELERATION_STRUCTURE_TYPE_TOP_LEVEL_KHR;
+		VkAccelerationStructureCreateInfoKHR CreateInfo;
+		ZeroVulkanStruct(CreateInfo, VK_STRUCTURE_TYPE_ACCELERATION_STRUCTURE_CREATE_INFO_KHR);
+		CreateInfo.buffer = InSourceBuffer->GetHandle();
+		CreateInfo.offset = InOffset;
+		CreateInfo.size = InSourceBuffer->GetSize() - InOffset;
+		CreateInfo.type = VK_ACCELERATION_STRUCTURE_TYPE_TOP_LEVEL_KHR;
 
-	VkDevice NativeDevice = InDevice->GetInstanceHandle();
-	VERIFYVULKANRESULT(VulkanDynamicAPI::vkCreateAccelerationStructureKHR(NativeDevice, &CreateInfo, VULKAN_CPU_ALLOCATOR, &AccelerationStructureHandle));
+		VkDevice NativeDevice = InDevice->GetInstanceHandle();
+		VERIFYVULKANRESULT(VulkanDynamicAPI::vkCreateAccelerationStructureKHR(NativeDevice, &CreateInfo, VULKAN_CPU_ALLOCATOR, &AccelerationStructureHandle));
+	}
+	else
+#endif
+	{
+		SourceStructuredBuffer = InSourceBuffer;
+		Size = InSourceBuffer->GetSize() - InOffset;
+		Offset = InOffset;
+	}
 }
-#endif // VULKAN_RHI_RAYTRACING
 
 FVulkanShaderResourceView::~FVulkanShaderResourceView()
 {
@@ -400,6 +398,9 @@ FShaderResourceViewRHIRef FVulkanDynamicRHI::RHICreateShaderResourceView(const F
 			}
 		}
 		case FShaderResourceViewInitializer::EType::StructuredBufferSRV:
+#if VULKAN_RHI_RAYTRACING
+		case FShaderResourceViewInitializer::EType::AccelerationStructureSRV:
+#endif
 		{
 			check(Desc.Buffer);
 			return new FVulkanShaderResourceView(Device, Buffer, Desc.StartOffsetBytes);
@@ -413,14 +414,6 @@ FShaderResourceViewRHIRef FVulkanDynamicRHI::RHICreateShaderResourceView(const F
 			uint32 Size = FMath::Min(Buffer->GetSize() - Desc.StartOffsetBytes, Desc.NumElements * Stride);
 			return new FVulkanShaderResourceView(Device, Desc.Buffer, Buffer, Size, Format, Desc.StartOffsetBytes);
 		}
-#if VULKAN_RHI_RAYTRACING
-		case FShaderResourceViewInitializer::EType::AccelerationStructureSRV:
-		{
-			check(Desc.Buffer);
-			FVulkanAccelerationStructureBuffer* ASBuffer = static_cast<FVulkanAccelerationStructureBuffer*>(Desc.Buffer);
-			return new FVulkanShaderResourceView(Device, ASBuffer, Desc.StartOffsetBytes);
-		}
-#endif // D3D12_RHI_RAYTRACING
 	}
 	checkNoEntry();
 	return nullptr;
@@ -434,19 +427,12 @@ FShaderResourceViewRHIRef FVulkanDynamicRHI::RHICreateShaderResourceView(FRHITex
 
 FShaderResourceViewRHIRef FVulkanDynamicRHI::RHICreateShaderResourceView(FRHIBuffer* BufferRHI)
 {
-	if (BufferRHI && EnumHasAnyFlags(BufferRHI->GetUsage(), BUF_VertexBuffer | BUF_StructuredBuffer))
+	if (BufferRHI && EnumHasAnyFlags(BufferRHI->GetUsage(), BUF_VertexBuffer | BUF_StructuredBuffer | BUF_AccelerationStructure))
 	{
 		FVulkanResourceMultiBuffer* Buffer = ResourceCast(BufferRHI);
 		FVulkanShaderResourceView* SRV = new FVulkanShaderResourceView(Device, Buffer);
 		return SRV;
 	}
-#if VULKAN_RHI_RAYTRACING
-	else if (BufferRHI && EnumHasAnyFlags(BufferRHI->GetUsage(), BUF_AccelerationStructure))
-	{
-		FVulkanAccelerationStructureBuffer* ASBuffer = static_cast<FVulkanAccelerationStructureBuffer*>(BufferRHI);
-		return new FVulkanShaderResourceView(Device, ASBuffer, 0);
-	}
-#endif
 	else
 	{
 		if (!BufferRHI)
