Fix Vulkan issue where transition batches containing only a single execution barriers were not submitted.

"Read -> Write" transitions are expressed as execution barriers and pipeline barrier submission would skip cases without any access masks, so they would only get submitted when batched with other barriers.  Fixes issue on some platforms where the same buffer was re-used by RDG for back to back Culling+Rendering of shadows, and the lone GraphicSRC->ComputeUAV transition for that buffer was skipped.

(a handful of <1us redundant barriers are added along with the handful of necessary ones, the impact on performance was minimal in the tested maps, an upcoming changelist will improve this)

diff --git a/Engine/Source/Runtime/VulkanRHI/Private/VulkanBarriers.cpp b/Engine/Source/Runtime/VulkanRHI/Private/VulkanBarriers.cpp
--- a/Engine/Source/Runtime/VulkanRHI/Private/VulkanBarriers.cpp
+++ b/Engine/Source/Runtime/VulkanRHI/Private/VulkanBarriers.cpp
@@ -553,22 +553,6 @@ static void SetupSubresourceRange(VkImageSubresourceRange& SubresRange, const FR
 	}
 }
 
-static void AddMemoryBarrier(VkMemoryBarrier& MemoryBarrier, VkAccessFlags SrcAccessFlags, VkAccessFlags DstAccessFlags)
-{
-	const VkAccessFlags ReadMask = VK_ACCESS_INDIRECT_COMMAND_READ_BIT | VK_ACCESS_INDEX_READ_BIT | VK_ACCESS_VERTEX_ATTRIBUTE_READ_BIT | VK_ACCESS_UNIFORM_READ_BIT |
-		VK_ACCESS_INPUT_ATTACHMENT_READ_BIT | VK_ACCESS_SHADER_READ_BIT | VK_ACCESS_COLOR_ATTACHMENT_READ_BIT | VK_ACCESS_DEPTH_STENCIL_ATTACHMENT_READ_BIT |
-		VK_ACCESS_TRANSFER_READ_BIT;
-
-	// We only need a memory barrier if the previous commands wrote to the buffer. In case of a transition from read, an execution barrier is enough.
-	const bool SrcAccessIsRead = ((SrcAccessFlags & (~ReadMask)) == 0);
-
-	if (!SrcAccessIsRead)
-	{
-		MemoryBarrier.srcAccessMask |= SrcAccessFlags;
-		MemoryBarrier.dstAccessMask |= DstAccessFlags;
-	}
-}
-
 void FVulkanDynamicRHI::RHICreateTransition(FRHITransition* Transition, const FRHITransitionCreateInfo& CreateInfo)
 {
 	const ERHIPipeline SrcPipelines = CreateInfo.SrcPipelines;
@@ -793,17 +777,18 @@ void FVulkanDynamicRHI::RHICreateTransition(FRHITransition* Transition, const FR
 			DstAccessFlags = SrcAccessFlags & (VK_ACCESS_SHADER_READ_BIT | VK_ACCESS_SHADER_WRITE_BIT);
 		}
 
-		// Add the stages affected by this transition.
-		Data->SrcStageMask |= SrcStageMask;
-		Data->DstStageMask |= DstStageMask;
 
 		// If we're not transitioning across pipes and we don't need to perform layout transitions, we can express memory dependencies through a global memory barrier.
-		if ( (SrcPipelines == DstPipelines) && (Texture == nullptr || SrcLayout == DstLayout) )
+		if ((SrcPipelines == DstPipelines) && (Texture == nullptr || SrcLayout == DstLayout))
 		{
-			AddMemoryBarrier(Data->MemoryBarrier, SrcAccessFlags, DstAccessFlags);
+			Data->AddMemoryBarrier(SrcAccessFlags, DstAccessFlags, SrcStageMask, DstStageMask);
 			continue;
 		}
 
+		// Add the stages affected by this transition.
+		Data->SrcStageMask |= SrcStageMask;
+		Data->DstStageMask |= DstStageMask;
+
 		if (Buffer != nullptr)
 		{
 			// We only add buffer transitions for cross-pipe transfers.
@@ -1157,6 +1142,7 @@ void FVulkanCommandListContext::RHIEndTransitions(TArrayView<const FRHITransitio
 
 	TArray<VkBufferMemoryBarrier, TInlineAllocator<8>> RealBufferBarriers;
 	TArray<VkImageMemoryBarrier, TInlineAllocator<8>> RealImageBarriers;
+	FVulkanPipelineBarrier RealMemoryBarrier;
 
 	for (const FRHITransition* Transition : Transitions)
 	{
@@ -1184,7 +1170,8 @@ void FVulkanCommandListContext::RHIEndTransitions(TArrayView<const FRHITransitio
 		}
 #endif
 
-		VkMemoryBarrier RealMemoryBarrier = Data->MemoryBarrier;
+		RealMemoryBarrier.MemoryBarrier = Data->MemoryBarrier;
+		RealMemoryBarrier.bHasMemoryBarrier = Data->bHasMemoryBarrier;
 
 		check(Data->SrcPipelines != Data->DstPipelines || Data->BufferBarriers.Num() == 0);
 		RealBufferBarriers.Reset();
@@ -1260,7 +1247,7 @@ void FVulkanCommandListContext::RHIEndTransitions(TArrayView<const FRHITransitio
 				{
 					// It turns out that we don't need a layout transition afterall. We may still need a memory barrier if the
 					// previous access was writable.
-					AddMemoryBarrier(RealMemoryBarrier, RealBarrier.srcAccessMask, RealBarrier.dstAccessMask);
+					RealMemoryBarrier.AddMemoryBarrier(RealBarrier.srcAccessMask, RealBarrier.dstAccessMask, 0, 0);
 					continue;
 				}
 			}
@@ -1275,16 +1262,36 @@ void FVulkanCommandListContext::RHIEndTransitions(TArrayView<const FRHITransitio
 			Layout.Set(RealBarrier.newLayout, RealBarrier.subresourceRange);
 		}
 
-		int NumMemoryBarriers = (RealMemoryBarrier.srcAccessMask != 0) || (RealMemoryBarrier.dstAccessMask != 0) ? 1 : 0;
-		if (NumMemoryBarriers == 0 && RealBufferBarriers.Num() == 0 && RealImageBarriers.Num() == 0)
+		if ((!RealMemoryBarrier.bHasMemoryBarrier) && (RealBufferBarriers.Num() == 0) && (RealImageBarriers.Num() == 0))
 		{
 			continue;
 		}
 
-		VulkanRHI::vkCmdPipelineBarrier(CmdBuffer->GetHandle(), RealSrcStageMask, RealDstStageMask, 0, NumMemoryBarriers, &RealMemoryBarrier, RealBufferBarriers.Num(), RealBufferBarriers.GetData(), RealImageBarriers.Num(), RealImageBarriers.GetData());
+		VulkanRHI::vkCmdPipelineBarrier(CmdBuffer->GetHandle(), RealSrcStageMask, RealDstStageMask, 0, RealMemoryBarrier.bHasMemoryBarrier ? 1 : 0, &RealMemoryBarrier.MemoryBarrier, RealBufferBarriers.Num(), RealBufferBarriers.GetData(), RealImageBarriers.Num(), RealImageBarriers.GetData());
 	}
 }
 
+
+void FVulkanPipelineBarrier::AddMemoryBarrier(VkAccessFlags InSrcAccessFlags, VkAccessFlags InDstAccessFlags, VkPipelineStageFlags InSrcStageMask, VkPipelineStageFlags InDstStageMask)
+{
+	const VkAccessFlags ReadMask = VK_ACCESS_INDIRECT_COMMAND_READ_BIT | VK_ACCESS_INDEX_READ_BIT | VK_ACCESS_VERTEX_ATTRIBUTE_READ_BIT | VK_ACCESS_UNIFORM_READ_BIT |
+		VK_ACCESS_INPUT_ATTACHMENT_READ_BIT | VK_ACCESS_SHADER_READ_BIT | VK_ACCESS_COLOR_ATTACHMENT_READ_BIT | VK_ACCESS_DEPTH_STENCIL_ATTACHMENT_READ_BIT |
+		VK_ACCESS_TRANSFER_READ_BIT | VK_ACCESS_MEMORY_READ_BIT;
+
+	// We only need a memory barrier if the previous commands wrote to the buffer. In case of a transition from read, an execution barrier is enough.
+	const bool SrcAccessIsRead = ((InSrcAccessFlags & (~ReadMask)) == 0);
+	if (!SrcAccessIsRead)
+	{
+		MemoryBarrier.srcAccessMask |= InSrcAccessFlags;
+		MemoryBarrier.dstAccessMask |= InDstAccessFlags;
+	}
+
+	SrcStageMask |= InSrcStageMask;
+	DstStageMask |= InDstStageMask;
+	bHasMemoryBarrier = true;
+}
+
+
 //
 // Methods used when the RHI itself needs to perform a layout transition. The public API functions do not call these,
 // they fill in the fields of FVulkanPipelineBarrier using their own logic, based on the ERHIAccess flags.
@@ -1441,10 +1448,9 @@ void FVulkanPipelineBarrier::AddImageAccessTransition(const FVulkanSurface& Surf
 
 void FVulkanPipelineBarrier::Execute(VkCommandBuffer CmdBuffer)
 {
-	int NumMemoryBarriers = (MemoryBarrier.srcAccessMask != 0) || (MemoryBarrier.dstAccessMask != 0) ? 1 : 0;
-	if (NumMemoryBarriers != 0 || ImageBarriers.Num() != 0)
+	if (bHasMemoryBarrier || ImageBarriers.Num() != 0)
 	{
-		VulkanRHI::vkCmdPipelineBarrier(CmdBuffer, SrcStageMask, DstStageMask, 0, NumMemoryBarriers, &MemoryBarrier, 0, nullptr, ImageBarriers.Num(), ImageBarriers.GetData());
+		VulkanRHI::vkCmdPipelineBarrier(CmdBuffer, SrcStageMask, DstStageMask, 0, bHasMemoryBarrier ? 1 : 0, &MemoryBarrier, 0, nullptr, ImageBarriers.Num(), ImageBarriers.GetData());
 	}
 }
 
@@ -1691,10 +1697,9 @@ void FVulkanLayoutManager::BeginRenderPass(FVulkanCommandListContext& Context, F
 		{
 			// Insert a barrier if we're loading from any color targets, to make sure the passes aren't reordered and we end up running before
 			// the pass we're supposed to read from.
-			Barrier.SrcStageMask |= VK_PIPELINE_STAGE_COLOR_ATTACHMENT_OUTPUT_BIT;
-			Barrier.DstStageMask |= VK_PIPELINE_STAGE_COLOR_ATTACHMENT_OUTPUT_BIT;
-			Barrier.MemoryBarrier.srcAccessMask |= VK_ACCESS_COLOR_ATTACHMENT_WRITE_BIT;
-			Barrier.MemoryBarrier.dstAccessMask |= VK_ACCESS_COLOR_ATTACHMENT_WRITE_BIT;
+			const VkAccessFlags AccessMask = VK_ACCESS_COLOR_ATTACHMENT_WRITE_BIT;
+			const VkPipelineStageFlags StageMask = VK_PIPELINE_STAGE_COLOR_ATTACHMENT_OUTPUT_BIT;
+			Barrier.AddMemoryBarrier(AccessMask, AccessMask, StageMask, StageMask);
 		}
 
 		if (bNeedsClearValues)
@@ -1720,10 +1725,9 @@ void FVulkanLayoutManager::BeginRenderPass(FVulkanCommandListContext& Context, F
 		{
 			// If the depth-stencil state doesn't change between passes, the high level code won't perform any transitions.
 			// Make sure we have a barrier in case we're loading depth or stencil, to prevent rearranging passes.
-			Barrier.SrcStageMask |= VK_PIPELINE_STAGE_EARLY_FRAGMENT_TESTS_BIT | VK_PIPELINE_STAGE_LATE_FRAGMENT_TESTS_BIT;
-			Barrier.DstStageMask |= VK_PIPELINE_STAGE_EARLY_FRAGMENT_TESTS_BIT | VK_PIPELINE_STAGE_LATE_FRAGMENT_TESTS_BIT;
-			Barrier.MemoryBarrier.srcAccessMask |= VK_ACCESS_DEPTH_STENCIL_ATTACHMENT_READ_BIT | VK_ACCESS_DEPTH_STENCIL_ATTACHMENT_WRITE_BIT;
-			Barrier.MemoryBarrier.dstAccessMask |= VK_ACCESS_DEPTH_STENCIL_ATTACHMENT_READ_BIT | VK_ACCESS_DEPTH_STENCIL_ATTACHMENT_WRITE_BIT;
+			const VkAccessFlags AccessMask = VK_ACCESS_DEPTH_STENCIL_ATTACHMENT_READ_BIT | VK_ACCESS_DEPTH_STENCIL_ATTACHMENT_WRITE_BIT;
+			const VkPipelineStageFlags StageMask = VK_PIPELINE_STAGE_EARLY_FRAGMENT_TESTS_BIT | VK_PIPELINE_STAGE_LATE_FRAGMENT_TESTS_BIT;
+			Barrier.AddMemoryBarrier(AccessMask, AccessMask, StageMask, StageMask);
 		}
 
 		if (DSTexture->HasClearValue() && bNeedsClearValues)
diff --git a/Engine/Source/Runtime/VulkanRHI/Private/VulkanBarriers.h b/Engine/Source/Runtime/VulkanRHI/Private/VulkanBarriers.h
--- a/Engine/Source/Runtime/VulkanRHI/Private/VulkanBarriers.h
+++ b/Engine/Source/Runtime/VulkanRHI/Private/VulkanBarriers.h
@@ -9,7 +9,7 @@ class FVulkanSurface;
 
 struct FVulkanPipelineBarrier
 {
-	FVulkanPipelineBarrier() : SrcStageMask(0), DstStageMask(0), Semaphore(nullptr)
+	FVulkanPipelineBarrier() : SrcStageMask(0), DstStageMask(0), Semaphore(nullptr), bHasMemoryBarrier(false)
 	{
 		ZeroVulkanStruct(MemoryBarrier, VK_STRUCTURE_TYPE_MEMORY_BARRIER);
 	}
@@ -20,10 +20,12 @@ struct FVulkanPipelineBarrier
 	TArray<VkImageMemoryBarrier, TInlineAllocator<2>> ImageBarriers;
 	TArray<VkBufferMemoryBarrier> BufferBarriers;
 	VulkanRHI::FSemaphore* Semaphore;
+	bool bHasMemoryBarrier;
 
 	// We need to keep the texture pointers around, because we need to call OnTransitionResource on them, and we need mip and layer counts for the tracking code.
 	TArray<FVulkanTextureBase*, TInlineAllocator<2>> Textures;
 
+	void AddMemoryBarrier(VkAccessFlags SrcAccessFlags, VkAccessFlags DstAccessFlags, VkPipelineStageFlags SrcStageMask, VkPipelineStageFlags DstStageMask);
 	void AddImageLayoutTransition(VkImage Image, VkImageLayout SrcLayout, VkImageLayout DstLayout, const VkImageSubresourceRange& SubresourceRange);
 	void AddImageLayoutTransition(VkImage Image, VkImageAspectFlags AspectMask, const struct FVulkanImageLayout& SrcLayout, VkImageLayout DstLayout);
 	void AddImageLayoutTransition(VkImage Image, VkImageAspectFlags AspectMask, VkImageLayout SrcLayout, const struct FVulkanImageLayout& DstLayout);
diff --git a/Engine/Source/Runtime/VulkanRHI/Private/VulkanRenderTarget.cpp b/Engine/Source/Runtime/VulkanRHI/Private/VulkanRenderTarget.cpp
--- a/Engine/Source/Runtime/VulkanRHI/Private/VulkanRenderTarget.cpp
+++ b/Engine/Source/Runtime/VulkanRHI/Private/VulkanRenderTarget.cpp
@@ -218,16 +218,9 @@ void FVulkanCommandListContext::RHICopyToResolveTarget(FRHITexture* SourceTextur
 		CopyRegion.imageExtent.depth = 1;
 		VulkanRHI::vkCmdCopyImageToBuffer(CmdBuffer->GetHandle(), SrcSurface->Image, VK_IMAGE_LAYOUT_TRANSFER_SRC_OPTIMAL, CpuReadbackBuffer->Buffer, 1, &CopyRegion);
 
-
-
 		{
 			FVulkanPipelineBarrier BarrierMemory;
-			BarrierMemory.MemoryBarrier.sType = VK_STRUCTURE_TYPE_MEMORY_BARRIER;
-			BarrierMemory.MemoryBarrier.pNext = nullptr;
-			BarrierMemory.MemoryBarrier.srcAccessMask = VK_ACCESS_TRANSFER_WRITE_BIT;
-			BarrierMemory.MemoryBarrier.dstAccessMask = VK_ACCESS_HOST_READ_BIT;
-			BarrierMemory.SrcStageMask = VK_PIPELINE_STAGE_ALL_COMMANDS_BIT;
-			BarrierMemory.DstStageMask = VK_PIPELINE_STAGE_HOST_BIT;
+			BarrierMemory.AddMemoryBarrier(VK_ACCESS_TRANSFER_WRITE_BIT, VK_ACCESS_HOST_READ_BIT, VK_PIPELINE_STAGE_ALL_COMMANDS_BIT, VK_PIPELINE_STAGE_HOST_BIT);
 			BarrierMemory.Execute(CmdBuffer->GetHandle());
 		}
 		SrcCurrentAccess = ERHIAccess::CopySrc;
diff --git a/Engine/Source/Runtime/VulkanRHI/Private/VulkanTexture.cpp b/Engine/Source/Runtime/VulkanRHI/Private/VulkanTexture.cpp
--- a/Engine/Source/Runtime/VulkanRHI/Private/VulkanTexture.cpp
+++ b/Engine/Source/Runtime/VulkanRHI/Private/VulkanTexture.cpp
@@ -2620,13 +2620,7 @@ void FVulkanCommandListContext::RHICopyTexture(FRHITexture* SourceTexture, FRHIT
 		VulkanRHI::vkCmdCopyImageToBuffer(CmdBuffer, SrcSurface.Image, VK_IMAGE_LAYOUT_TRANSFER_SRC_OPTIMAL, CpuReadbackBuffer->Buffer, CopyInfo.NumMips, &CopyRegion[0]);
 
 		FVulkanPipelineBarrier BarrierMemory;
-		BarrierMemory.MemoryBarrier.sType = VK_STRUCTURE_TYPE_MEMORY_BARRIER;
-		BarrierMemory.MemoryBarrier.pNext = nullptr;
-		BarrierMemory.MemoryBarrier.srcAccessMask = VK_ACCESS_TRANSFER_WRITE_BIT;
-		BarrierMemory.MemoryBarrier.dstAccessMask = VK_ACCESS_HOST_READ_BIT;
-		BarrierMemory.SrcStageMask = VK_PIPELINE_STAGE_ALL_COMMANDS_BIT;
-		BarrierMemory.DstStageMask = VK_PIPELINE_STAGE_HOST_BIT;
-
+		BarrierMemory.AddMemoryBarrier(VK_ACCESS_TRANSFER_WRITE_BIT, VK_ACCESS_HOST_READ_BIT, VK_PIPELINE_STAGE_ALL_COMMANDS_BIT, VK_PIPELINE_STAGE_HOST_BIT);
 		BarrierMemory.Execute(CmdBuffer);
 	}
 	else
