Fix Linux build error when RHIRaytracing enabled ("Instance declaration shadows a field of 'FVulkanDynamicRHI'")
Align Linux defines with Windows for Vullkan Raytracing.

diff --git a/Engine/Source/Runtime/VulkanRHI/Private/Linux/VulkanLinuxPlatform.h b/Engine/Source/Runtime/VulkanRHI/Private/Linux/VulkanLinuxPlatform.h
--- a/Engine/Source/Runtime/VulkanRHI/Private/Linux/VulkanLinuxPlatform.h
+++ b/Engine/Source/Runtime/VulkanRHI/Private/Linux/VulkanLinuxPlatform.h
@@ -11,9 +11,15 @@
 #define VULKAN_SHOULD_ENABLE_DRAW_MARKERS			(UE_BUILD_DEBUG || UE_BUILD_DEVELOPMENT)
 #define VULKAN_SIGNAL_UNIMPLEMENTED()				checkf(false, TEXT("Unimplemented vulkan functionality: %s"), __PRETTY_FUNCTION__)
 #define VULKAN_SUPPORTS_AMD_BUFFER_MARKER			1
-#define VULKAN_RHI_RAYTRACING						0
 
-#define	UE_VK_API_VERSION							VK_API_VERSION_1_1
+#define VULKAN_RHI_RAYTRACING 						(RHI_RAYTRACING)
+#define VULKAN_SUPPORTS_SCALAR_BLOCK_LAYOUT			(VULKAN_RHI_RAYTRACING)
+
+#if VULKAN_RHI_RAYTRACING
+#	define UE_VK_API_VERSION						VK_API_VERSION_1_2
+#else
+#	define UE_VK_API_VERSION						VK_API_VERSION_1_1
+#endif // VULKAN_RHI_RAYTRACING
 
 #define ENUM_VK_ENTRYPOINTS_PLATFORM_BASE(EnumMacro)
 
diff --git a/Engine/Source/Runtime/VulkanRHI/Private/VulkanRayTracing.cpp b/Engine/Source/Runtime/VulkanRHI/Private/VulkanRayTracing.cpp
--- a/Engine/Source/Runtime/VulkanRHI/Private/VulkanRayTracing.cpp
+++ b/Engine/Source/Runtime/VulkanRHI/Private/VulkanRayTracing.cpp
@@ -618,40 +618,40 @@ FRayTracingSceneRHIRef FVulkanDynamicRHI::RHICreateRayTracingScene(const FRayTra
 	const EParallelForFlags ParallelForFlags = EParallelForFlags::None; // set ForceSingleThread for testing
 	ParallelFor(NumSceneInstances, [&Initializer2, Instances = Initializer.Instances, &NativeInstances](int32 InstanceIndex)
 	{
-		const FRayTracingGeometryInstance& Instance = Instances[InstanceIndex];
+		const FRayTracingGeometryInstance& RayTracingGeometryInstance = Instances[InstanceIndex];
 		FVulkanRayTracingGeometry* Geometry = ResourceCast(Initializer2.PerInstanceGeometries[InstanceIndex]);
 
 		const FRayTracingAccelerationStructureAddress AccelerationStructureAddress = Geometry->GetAccelerationStructureAddress(0);
 		check(AccelerationStructureAddress != 0);
 
 		VkAccelerationStructureInstanceKHR InstanceDesc = {};
-		InstanceDesc.mask = Instance.Mask;
+		InstanceDesc.mask = RayTracingGeometryInstance.Mask;
 		InstanceDesc.instanceShaderBindingTableRecordOffset = Initializer2.SegmentPrefixSum[InstanceIndex] * Initializer2.ShaderSlotsPerGeometrySegment; // TODO?
-		InstanceDesc.flags = TranslateRayTracingInstanceFlags(Instance.Flags);
+		InstanceDesc.flags = TranslateRayTracingInstanceFlags(RayTracingGeometryInstance.Flags);
 
-		const uint32 NumTransforms = Instance.NumTransforms;
+		const uint32 NumTransforms = RayTracingGeometryInstance.NumTransforms;
 
-		checkf(Instance.UserData.Num() == 0 || Instance.UserData.Num() >= int32(NumTransforms),
-			TEXT("User data array must be either be empty (Instance.DefaultUserData is used), or contain one entry per entry in Transforms array."));
+		checkf(RayTracingGeometryInstance.UserData.Num() == 0 || RayTracingGeometryInstance.UserData.Num() >= int32(NumTransforms),
+			TEXT("User data array must be either be empty (RayTracingGeometryInstance.DefaultUserData is used), or contain one entry per entry in Transforms array."));
 
-		const bool bUseUniqueUserData = Instance.UserData.Num() != 0;
+		const bool bUseUniqueUserData = RayTracingGeometryInstance.UserData.Num() != 0;
 
 		uint32 DescIndex = Initializer2.BaseInstancePrefixSum[InstanceIndex];
 
 		for (uint32 TransformIndex = 0; TransformIndex < NumTransforms; ++TransformIndex)
 		{
-			InstanceDesc.instanceCustomIndex = bUseUniqueUserData ? Instance.UserData[TransformIndex] : Instance.DefaultUserData;
+			InstanceDesc.instanceCustomIndex = bUseUniqueUserData ? RayTracingGeometryInstance.UserData[TransformIndex] : RayTracingGeometryInstance.DefaultUserData;
 
 			InstanceDesc.accelerationStructureReference = AccelerationStructureAddress;
 
-			if (!Instance.ActivationMask.IsEmpty() && (Instance.ActivationMask[TransformIndex / 32] & (1 << (TransformIndex % 32))) == 0)
+			if (!RayTracingGeometryInstance.ActivationMask.IsEmpty() && (RayTracingGeometryInstance.ActivationMask[TransformIndex / 32] & (1 << (TransformIndex % 32))) == 0)
 			{
 				InstanceDesc.accelerationStructureReference = 0;
 			}
 
-			if (TransformIndex < (uint32)Instance.Transforms.Num())
+			if (TransformIndex < (uint32)RayTracingGeometryInstance.Transforms.Num())
 			{
-				const FMatrix& Transform = Instance.Transforms[TransformIndex];
+				const FMatrix& Transform = RayTracingGeometryInstance.Transforms[TransformIndex];
 
 				InstanceDesc.transform.matrix[0][0] = Transform.M[0][0];
 				InstanceDesc.transform.matrix[0][1] = Transform.M[1][0];
