Merge //UE5/Release-5.0 into //UE5/Private-Frosty-To-5.0-Staging

diff --git a/Engine/Plugins/Experimental/GeometryCache/Source/GeometryCache/Private/GeometryCacheSceneProxy.cpp b/Engine/Plugins/Experimental/GeometryCache/Source/GeometryCache/Private/GeometryCacheSceneProxy.cpp
--- a/Engine/Plugins/Experimental/GeometryCache/Source/GeometryCache/Private/GeometryCacheSceneProxy.cpp
+++ b/Engine/Plugins/Experimental/GeometryCache/Source/GeometryCache/Private/GeometryCacheSceneProxy.cpp
@@ -650,7 +650,7 @@ void FGeometryCacheSceneProxy::UpdateAnimation(float NewTime, bool bNewLooping,
 					Segment.MaxVertices = Section->PositionBuffers[PositionBufferIndex].GetSizeInBytes() / Segment.VertexBufferStride; // conservative estimate
 
 					Segments.Add(Segment);
-					TotalPrimitiveCount += BatchInfo.NumTriangles;
+					Section->RayTracingGeometry.Initializer.TotalPrimitiveCount += BatchInfo.NumTriangles;
 				}
 
 				bRequireRecreate = bRequireRecreate || Section->RayTracingGeometry.Initializer.IndexBuffer != Section->IndexBuffer.IndexBufferRHI;
diff --git a/Engine/Source/Developer/MeshBoneReduction/Private/MeshBoneReduction.cpp b/Engine/Source/Developer/MeshBoneReduction/Private/MeshBoneReduction.cpp
--- a/Engine/Source/Developer/MeshBoneReduction/Private/MeshBoneReduction.cpp
+++ b/Engine/Source/Developer/MeshBoneReduction/Private/MeshBoneReduction.cpp
@@ -456,7 +456,7 @@ public:
 					{
 						FVector TangentX = Vertex.TangentX;
 						FVector TangentY = Vertex.TangentY;
-						FVector TangentZ = Vertex.TangentZ;
+						FVector TangentZ = (FVector4)Vertex.TangentZ;
 						FVector Position = Vertex.Position;
 						for (uint8 InfluenceIndex = 0; InfluenceIndex < MAX_TOTAL_INFLUENCES; ++InfluenceIndex)
 				        {
@@ -469,7 +469,7 @@ public:
 
 									TangentX += ((RemovedBoneMatrices[ArrayIndex].TransformVector(Vertex.TangentX) - Vertex.TangentX) * ((float)Vertex.InfluenceWeights[InfluenceIndex] * InfluenceMultiplier));
 									TangentY += ((RemovedBoneMatrices[ArrayIndex].TransformVector(Vertex.TangentY) - Vertex.TangentY) * ((float)Vertex.InfluenceWeights[InfluenceIndex] * InfluenceMultiplier));
-									TangentZ += ((RemovedBoneMatrices[ArrayIndex].TransformVector(Vertex.TangentZ) - Vertex.TangentZ) * ((float)Vertex.InfluenceWeights[InfluenceIndex] * InfluenceMultiplier));
+									TangentZ += ((RemovedBoneMatrices[ArrayIndex].TransformVector((FVector4)Vertex.TangentZ) - Vertex.TangentZ) * ((float)Vertex.InfluenceWeights[InfluenceIndex] * InfluenceMultiplier));
 								}
 							}
 				        }
diff --git a/Engine/Source/Editor/WorldPartitionEditor/Private/WorldPartition/SWorldPartitionEditorGridSpatialHash.cpp b/Engine/Source/Editor/WorldPartitionEditor/Private/WorldPartition/SWorldPartitionEditorGridSpatialHash.cpp
--- a/Engine/Source/Editor/WorldPartitionEditor/Private/WorldPartition/SWorldPartitionEditorGridSpatialHash.cpp
+++ b/Engine/Source/Editor/WorldPartitionEditor/Private/WorldPartition/SWorldPartitionEditorGridSpatialHash.cpp
@@ -70,6 +70,13 @@ void SWorldPartitionEditorGridSpatialHash::Tick(const FGeometry& AllottedGeometr
 	{
 		UpdateWorldMiniMapDetails();
 	}
+
+	/*static double LastTime = InCurrentTime;
+	if ((InCurrentTime - LastTime) > 15.0)
+	{
+		LastTime = InCurrentTime;
+		ReloadMiniMap();
+	}*/
 }
 
 FReply SWorldPartitionEditorGridSpatialHash::ReloadMiniMap()
diff --git a/Engine/Source/Editor/WorldPartitionEditor/Public/WorldPartition/SWorldPartitionEditorGridSpatialHash.h b/Engine/Source/Editor/WorldPartitionEditor/Public/WorldPartition/SWorldPartitionEditorGridSpatialHash.h
--- a/Engine/Source/Editor/WorldPartitionEditor/Public/WorldPartition/SWorldPartitionEditorGridSpatialHash.h
+++ b/Engine/Source/Editor/WorldPartitionEditor/Public/WorldPartition/SWorldPartitionEditorGridSpatialHash.h
@@ -16,6 +16,8 @@ public:
 
 	void Construct(const FArguments& InArgs);
 
+	virtual void Tick(const FGeometry& AllottedGeometry, const double InCurrentTime, const float InDeltaTime) override;
+
 	int32 PaintGrid(const FGeometry& AllottedGeometry, const FSlateRect& MyCullingRect, FSlateWindowElementList& OutDrawElements, int32 LayerId) const override;
 
 	virtual FReply ReloadMiniMap() override;
diff --git a/Engine/Source/Runtime/Engine/Private/ShaderDerivedDataVersion.h b/Engine/Source/Runtime/Engine/Private/ShaderDerivedDataVersion.h
--- a/Engine/Source/Runtime/Engine/Private/ShaderDerivedDataVersion.h
+++ b/Engine/Source/Runtime/Engine/Private/ShaderDerivedDataVersion.h
@@ -10,5 +10,5 @@
 
 // In case of merge conflicts with DDC versions, you MUST generate a new GUID and set this new
 // guid as version
-#define GLOBALSHADERMAP_DERIVEDDATA_VER			TEXT("2EE7C1F4F6C74732951E5B47720E1CB9")
-#define MATERIALSHADERMAP_DERIVEDDATA_VER		TEXT("321C290437B44458A11C57A76F63914E")
\ No newline at end of file
+#define GLOBALSHADERMAP_DERIVEDDATA_VER			TEXT("7429FF5E4C57472FA538321D6777E8BC")
+#define MATERIALSHADERMAP_DERIVEDDATA_VER		TEXT("0A688C67C250488F8EFB8DCD02A544C9")
\ No newline at end of file
diff --git a/Engine/Source/Runtime/Engine/Public/WorldPartition/WorldPartitionEditorPerProjectUserSettings.h b/Engine/Source/Runtime/Engine/Public/WorldPartition/WorldPartitionEditorPerProjectUserSettings.h
--- a/Engine/Source/Runtime/Engine/Public/WorldPartition/WorldPartitionEditorPerProjectUserSettings.h
+++ b/Engine/Source/Runtime/Engine/Public/WorldPartition/WorldPartitionEditorPerProjectUserSettings.h
@@ -66,6 +66,8 @@ public:
 	{}
 
 #if WITH_EDITOR
+	void UpdateEditorGridConfigHash(UWorld* InWorld);
+	static void UpdateEditorGridConfigHash(UWorld* InWorld, FWorldPartitionPerWorldSettings& PerWorldSettings);
 	TArray<FName> GetEditorGridLoadedCells(UWorld* InWorld) const;
 	void SetEditorGridLoadedCells(UWorld* InWorld, const TArray<FName>& InEditorGridLoadedCells);
 
diff --git a/Engine/Source/Runtime/RHI/Public/RHIDefinitions.h b/Engine/Source/Runtime/RHI/Public/RHIDefinitions.h
--- a/Engine/Source/Runtime/RHI/Public/RHIDefinitions.h
+++ b/Engine/Source/Runtime/RHI/Public/RHIDefinitions.h
@@ -378,6 +378,7 @@ class RHI_API FGenericDataDrivenShaderPlatformInfo
 	uint32 bOverrideFMaterial_NeedsGBufferEnabled : 1;
 	uint32 bSupportsMobileDistanceField : 1;
 	uint32 bSupportsFFTBloom : 1;
+	uint32 bSupportsVertexShaderLayer : 1;
 	uint32 bSupportsInlineRayTracing : 1;
 	uint32 bSupportsRayTracingShaders : 1;
 	uint32 bSupportsVertexShaderLayer : 1;
diff --git a/Engine/Source/Runtime/Renderer/Private/PostProcess/PostProcessMotionBlur.h b/Engine/Source/Runtime/Renderer/Private/PostProcess/PostProcessMotionBlur.h
--- a/Engine/Source/Runtime/Renderer/Private/PostProcess/PostProcessMotionBlur.h
+++ b/Engine/Source/Runtime/Renderer/Private/PostProcess/PostProcessMotionBlur.h
@@ -44,7 +44,7 @@ int32 GetMotionBlurDirections();
 
 // Shader parameters for MotionBlurVelocityFlatten.ush when another pass can also generate the FVelocityFlattenTextures
 BEGIN_SHADER_PARAMETER_STRUCT(FVelocityFlattenParameters, )
-	SHADER_PARAMETER(FVector2D, VelocityScale)
+	SHADER_PARAMETER(FVector2f, VelocityScale)
 	SHADER_PARAMETER(float, VelocityMax)
 END_SHADER_PARAMETER_STRUCT()
 
diff --git a/Engine/Source/Runtime/VulkanRHI/Private/VulkanRayTracing.cpp b/Engine/Source/Runtime/VulkanRHI/Private/VulkanRayTracing.cpp
--- a/Engine/Source/Runtime/VulkanRHI/Private/VulkanRayTracing.cpp
+++ b/Engine/Source/Runtime/VulkanRHI/Private/VulkanRayTracing.cpp
@@ -198,7 +198,6 @@ static void GetBLASBuildData(
 		FVulkanResourceMultiBuffer* const VertexBuffer = ResourceCast(Segment.VertexBuffer.GetReference());
 		VkDeviceOrHostAddressConstKHR VertexBufferDeviceAddress = {};
 		VertexBufferDeviceAddress.deviceAddress = VertexBuffer->GetDeviceAddress() + Segment.VertexBufferOffset;
-
 		VkAccelerationStructureGeometryKHR SegmentGeometry;
 		ZeroVulkanStruct(SegmentGeometry, VK_STRUCTURE_TYPE_ACCELERATION_STRUCTURE_GEOMETRY_KHR);
 
@@ -292,6 +291,8 @@ FVulkanRayTracingGeometry::FVulkanRayTracingGeometry(FRayTracingGeometryInitiali
 	check(Initializer.GeometryType == ERayTracingGeometryType::RTGT_Triangles);
 	checkf(!Initializer.IndexBuffer || (Initializer.IndexBuffer->GetStride() == 2 || Initializer.IndexBuffer->GetStride() == 4), TEXT("Index buffer must be 16 or 32 bit if in use."));
 
+	IndexStrideInBytes = Initializer.IndexBuffer ? Initializer.IndexBuffer->GetStride() : 0;
+
 	FVkRtBLASBuildData BuildData;
 	GetBLASBuildData(
 		Device->GetInstanceHandle(), 
@@ -321,10 +322,6 @@ FVulkanRayTracingGeometry::FVulkanRayTracingGeometry(FRayTracingGeometryInitiali
 	CreateInfo.type = VK_ACCELERATION_STRUCTURE_TYPE_BOTTOM_LEVEL_KHR;
 	VERIFYVULKANRESULT(vkCreateAccelerationStructureKHR(NativeDevice, &CreateInfo, VULKAN_CPU_ALLOCATOR, &Handle));
 
-	VkAccelerationStructureDeviceAddressInfoKHR DeviceAddressInfo;
-	ZeroVulkanStruct(DeviceAddressInfo, VK_STRUCTURE_TYPE_ACCELERATION_STRUCTURE_DEVICE_ADDRESS_INFO_KHR);
-	DeviceAddressInfo.accelerationStructure = Handle;
-	Address = vkGetAccelerationStructureDeviceAddressKHR(NativeDevice, &DeviceAddressInfo);
 	
 	SizeInfo.ResultSize = BuildData.SizesInfo.accelerationStructureSize;
 	SizeInfo.BuildScratchSize = BuildData.SizesInfo.buildScratchSize;
@@ -367,6 +364,11 @@ void FVulkanRayTracingGeometry::BuildAccelerationStructure(FVulkanCommandListCon
 	CommandBufferManager.SubmitActiveCmdBuffer();
 	CommandBufferManager.PrepareForNewActiveCommandBuffer();
 
+	VkAccelerationStructureDeviceAddressInfoKHR DeviceAddressInfo;
+	ZeroVulkanStruct(DeviceAddressInfo, VK_STRUCTURE_TYPE_ACCELERATION_STRUCTURE_DEVICE_ADDRESS_INFO_KHR);
+	DeviceAddressInfo.accelerationStructure = Handle;
+	AccelerationStructureBuffer->SetAddress(vkGetAccelerationStructureDeviceAddressKHR(Device->GetInstanceHandle(), &DeviceAddressInfo));
+
 	// No longer need scratch memory for a static build
 	if (!Initializer.bAllowUpdate)
 	{
diff --git a/Engine/Source/Runtime/VulkanRHI/Private/VulkanRayTracing.h b/Engine/Source/Runtime/VulkanRHI/Private/VulkanRayTracing.h
--- a/Engine/Source/Runtime/VulkanRHI/Private/VulkanRayTracing.h
+++ b/Engine/Source/Runtime/VulkanRHI/Private/VulkanRayTracing.h
@@ -35,13 +35,6 @@ public:
 	static bool LoadVulkanInstanceFunctions(VkInstance inInstance);
 };
 
-struct FVkRtAllocation
-{
-	VkDevice Device = VK_NULL_HANDLE;
-	VkDeviceMemory Memory = VK_NULL_HANDLE;
-	VkBuffer Buffer = VK_NULL_HANDLE;
-};
-
 class FVulkanRayTracingAllocator
 {
 public:
@@ -83,7 +76,7 @@ public:
 	FVulkanRayTracingGeometry(FRayTracingGeometryInitializer Initializer, FVulkanDevice* InDevice);
 	~FVulkanRayTracingGeometry();
 
-	virtual FRayTracingAccelerationStructureAddress GetAccelerationStructureAddress(uint64 GPUIndex) const final override { return Address; }
+	virtual FRayTracingAccelerationStructureAddress GetAccelerationStructureAddress(uint64 GPUIndex) const final override { return AccelerationStructureBuffer->GetAddress(); }
 	virtual uint32 GetNumSegments() const final override { return Initializer.Segments.Num(); }
 	virtual FRayTracingAccelerationStructureSize GetSizeInfo() const final override { return SizeInfo; }
 
@@ -95,7 +88,6 @@ private:
 	const FRayTracingGeometryInitializer Initializer;
 
 	VkAccelerationStructureKHR Handle = VK_NULL_HANDLE;
-	VkDeviceAddress Address = 0;
 	TRefCountPtr<FVulkanResourceMultiBuffer> AccelerationStructureBuffer;
 	TRefCountPtr<FVulkanResourceMultiBuffer> ScratchBuffer;
 	
