Merge //UE5/Release-5.0 (CL-18181837 + compile fixes)

diff --git a/Engine/Source/Runtime/Engine/Private/SkeletalMeshRenderData.cpp b/Engine/Source/Runtime/Engine/Private/SkeletalMeshRenderData.cpp
--- a/Engine/Source/Runtime/Engine/Private/SkeletalMeshRenderData.cpp
+++ b/Engine/Source/Runtime/Engine/Private/SkeletalMeshRenderData.cpp
@@ -186,7 +186,7 @@ static void SerializeLODInfoForDDC(USkeletalMesh* SkeletalMesh, FString& KeySuff
 // differences, etc.) replace the version GUID below with a new one.
 // In case of merge conflicts with DDC versions, you MUST generate a new GUID
 // and set this new GUID as the version.
-#define SKELETALMESH_DERIVEDDATA_VER TEXT("61C5045866H572HEB3C0686189A9E7ED")
+#define SKELETALMESH_DERIVEDDATA_VER TEXT("6F80DD35A8954A16ACFD2341CF8FCA65")
 
 const FString& GetSkeletalMeshDerivedDataVersion()
 {
diff --git a/Engine/Source/Runtime/VulkanRHI/Private/VulkanRayTracing.cpp b/Engine/Source/Runtime/VulkanRHI/Private/VulkanRayTracing.cpp
--- a/Engine/Source/Runtime/VulkanRHI/Private/VulkanRayTracing.cpp
+++ b/Engine/Source/Runtime/VulkanRHI/Private/VulkanRayTracing.cpp
@@ -311,7 +311,7 @@ FVulkanRayTracingGeometry::FVulkanRayTracingGeometry(FRayTracingGeometryInitiali
 
 	FRHIResourceCreateInfo ScratchBufferCreateInfo(TEXT("BuildScratchBLAS"));
 	ScratchBuffer = ResourceCast(RHICreateBuffer(BuildData.SizesInfo.buildScratchSize, BUF_UnorderedAccess | BUF_StructuredBuffer, 0, ERHIAccess::UAVCompute, ScratchBufferCreateInfo).GetReference());
-
+	
 	VkDevice NativeDevice = Device->GetInstanceHandle();
 
 	VkAccelerationStructureCreateInfoKHR CreateInfo;
@@ -326,6 +326,12 @@ FVulkanRayTracingGeometry::FVulkanRayTracingGeometry(FRayTracingGeometryInitiali
 	SizeInfo.ResultSize = BuildData.SizesInfo.accelerationStructureSize;
 	SizeInfo.BuildScratchSize = BuildData.SizesInfo.buildScratchSize;
 	SizeInfo.UpdateScratchSize = BuildData.SizesInfo.updateScratchSize;
+
+
+	VkAccelerationStructureDeviceAddressInfoKHR DeviceAddressInfo;
+	ZeroVulkanStruct(DeviceAddressInfo, VK_STRUCTURE_TYPE_ACCELERATION_STRUCTURE_DEVICE_ADDRESS_INFO_KHR);
+	DeviceAddressInfo.accelerationStructure = Handle;
+	Address = vkGetAccelerationStructureDeviceAddressKHR(NativeDevice, &DeviceAddressInfo);
 }
 
 FVulkanRayTracingGeometry::~FVulkanRayTracingGeometry()
@@ -351,7 +357,6 @@ void FVulkanRayTracingGeometry::BuildAccelerationStructure(FVulkanCommandListCon
 		BuildData);
 
 	check(BuildData.SizesInfo.accelerationStructureSize <= AccelerationStructureBuffer->GetSize());
-
 	BuildData.GeometryInfo.dstAccelerationStructure = Handle;
 	BuildData.GeometryInfo.scratchData.deviceAddress = ScratchBuffer->GetDeviceAddress();
 
@@ -364,11 +369,6 @@ void FVulkanRayTracingGeometry::BuildAccelerationStructure(FVulkanCommandListCon
 	CommandBufferManager.SubmitActiveCmdBuffer();
 	CommandBufferManager.PrepareForNewActiveCommandBuffer();
 
-	VkAccelerationStructureDeviceAddressInfoKHR DeviceAddressInfo;
-	ZeroVulkanStruct(DeviceAddressInfo, VK_STRUCTURE_TYPE_ACCELERATION_STRUCTURE_DEVICE_ADDRESS_INFO_KHR);
-	DeviceAddressInfo.accelerationStructure = Handle;
-	AccelerationStructureBuffer->SetAddress(vkGetAccelerationStructureDeviceAddressKHR(Device->GetInstanceHandle(), &DeviceAddressInfo));
-
 	// No longer need scratch memory for a static build
 	if (!Initializer.bAllowUpdate)
 	{
@@ -441,6 +441,11 @@ FVulkanRayTracingScene::FVulkanRayTracingScene(FRayTracingSceneInitializer2 InIn
 
 FVulkanRayTracingScene::~FVulkanRayTracingScene()
 {
+	if (Handle != VK_NULL_HANDLE)
+	{
+		VkDevice NativeDevice = Device->GetInstanceHandle();
+		vkDestroyAccelerationStructureKHR(NativeDevice, Handle, VULKAN_CPU_ALLOCATOR);
+	}
 }
 
 void FVulkanRayTracingScene::BindBuffer(FRHIBuffer* InBuffer, uint32 InBufferOffset)
@@ -488,7 +493,6 @@ void FVulkanRayTracingScene::BuildAccelerationStructure(
 		ScratchBuffer = ResourceCast(RHICreateBuffer(BuildData.SizesInfo.buildScratchSize, BUF_UnorderedAccess | BUF_StructuredBuffer, 0, ERHIAccess::UAVCompute, ScratchBufferCreateInfo).GetReference());
 		InScratchBuffer = ScratchBuffer.GetReference();
 	}
-
 	checkf(AccelerationStructureView, TEXT("A buffer must be bound to the ray tracing scene before it can be built."));
 	BuildData.GeometryInfo.dstAccelerationStructure = AccelerationStructureView->AccelerationStructureHandle;
 
