Fix bad merge

diff --git a/Engine/Source/Runtime/VulkanRHI/Private/VulkanRayTracing.cpp b/Engine/Source/Runtime/VulkanRHI/Private/VulkanRayTracing.cpp
--- a/Engine/Source/Runtime/VulkanRHI/Private/VulkanRayTracing.cpp
+++ b/Engine/Source/Runtime/VulkanRHI/Private/VulkanRayTracing.cpp
@@ -291,8 +291,6 @@ FVulkanRayTracingGeometry::FVulkanRayTracingGeometry(FRayTracingGeometryInitiali
 	check(Initializer.GeometryType == ERayTracingGeometryType::RTGT_Triangles);
 	checkf(!Initializer.IndexBuffer || (Initializer.IndexBuffer->GetStride() == 2 || Initializer.IndexBuffer->GetStride() == 4), TEXT("Index buffer must be 16 or 32 bit if in use."));
 
-	IndexStrideInBytes = Initializer.IndexBuffer ? Initializer.IndexBuffer->GetStride() : 0;
-
 	FVkRtBLASBuildData BuildData;
 	GetBLASBuildData(
 		Device->GetInstanceHandle(), 
diff --git a/Engine/Source/Runtime/VulkanRHI/Private/VulkanRayTracing.h b/Engine/Source/Runtime/VulkanRHI/Private/VulkanRayTracing.h
--- a/Engine/Source/Runtime/VulkanRHI/Private/VulkanRayTracing.h
+++ b/Engine/Source/Runtime/VulkanRHI/Private/VulkanRayTracing.h
@@ -8,6 +8,7 @@
 
 class FVulkanCommandListContext;
 class FVulkanResourceMultiBuffer;
+class FVulkanAccelerationStructureBuffer;
 class FVulkanRayTracingLayout;
 
 #define ENUM_VK_ENTRYPOINTS_RAYTRACING(EnumMacro) \
@@ -35,6 +36,13 @@ public:
 	static bool LoadVulkanInstanceFunctions(VkInstance inInstance);
 };
 
+struct FVkRtAllocation
+{
+	VkDevice Device = VK_NULL_HANDLE;
+	VkDeviceMemory Memory = VK_NULL_HANDLE;
+	VkBuffer Buffer = VK_NULL_HANDLE;
+};
+
 class FVulkanRayTracingAllocator
 {
 public:
@@ -66,6 +74,7 @@ struct FVkRtBLASBuildData
 
 	TArray<VkAccelerationStructureGeometryKHR, TInlineAllocator<1>> Segments;
 	TArray<VkAccelerationStructureBuildRangeInfoKHR, TInlineAllocator<1>> Ranges;
+	TArray<uint32, TInlineAllocator<1>> VertexCounts;
 	VkAccelerationStructureBuildGeometryInfoKHR GeometryInfo;
 	VkAccelerationStructureBuildSizesInfoKHR SizesInfo;
 };
@@ -76,7 +85,7 @@ public:
 	FVulkanRayTracingGeometry(FRayTracingGeometryInitializer Initializer, FVulkanDevice* InDevice);
 	~FVulkanRayTracingGeometry();
 
-	virtual FRayTracingAccelerationStructureAddress GetAccelerationStructureAddress(uint64 GPUIndex) const final override { return AccelerationStructureBuffer->GetAddress(); }
+	virtual FRayTracingAccelerationStructureAddress GetAccelerationStructureAddress(uint64 GPUIndex) const final override { return Address; }
 	virtual uint32 GetNumSegments() const final override { return Initializer.Segments.Num(); }
 	virtual FRayTracingAccelerationStructureSize GetSizeInfo() const final override { return SizeInfo; }
 
@@ -88,6 +97,7 @@ private:
 	const FRayTracingGeometryInitializer Initializer;
 
 	VkAccelerationStructureKHR Handle = VK_NULL_HANDLE;
+	VkDeviceAddress Address = 0;
 	TRefCountPtr<FVulkanResourceMultiBuffer> AccelerationStructureBuffer;
 	TRefCountPtr<FVulkanResourceMultiBuffer> ScratchBuffer;
 	
