Fix bad merge

diff --git a/Engine/Source/Runtime/VulkanRHI/Private/VulkanRayTracing.cpp b/Engine/Source/Runtime/VulkanRHI/Private/VulkanRayTracing.cpp
--- a/Engine/Source/Runtime/VulkanRHI/Private/VulkanRayTracing.cpp
+++ b/Engine/Source/Runtime/VulkanRHI/Private/VulkanRayTracing.cpp
@@ -439,11 +439,6 @@ FVulkanRayTracingScene::FVulkanRayTracingScene(FRayTracingSceneInitializer2 InIn
 
 FVulkanRayTracingScene::~FVulkanRayTracingScene()
 {
-	if (Handle != VK_NULL_HANDLE)
-	{
-		VkDevice NativeDevice = Device->GetInstanceHandle();
-		vkDestroyAccelerationStructureKHR(NativeDevice, Handle, VULKAN_CPU_ALLOCATOR);
-	}
 }
 
 void FVulkanRayTracingScene::BindBuffer(FRHIBuffer* InBuffer, uint32 InBufferOffset)
