Fixes for RHIReadSurfaceData in Vulkan RHI:
- bIs8Bpp use was inverted (allocate too much memory for RGBA8 and too little memory for RGBA16)
- Rect param was only partially taken into account
- Add more supported formats

diff --git a/Engine/Source/Runtime/VulkanRHI/Private/VulkanRenderTarget.cpp b/Engine/Source/Runtime/VulkanRHI/Private/VulkanRenderTarget.cpp
--- a/Engine/Source/Runtime/VulkanRHI/Private/VulkanRenderTarget.cpp
+++ b/Engine/Source/Runtime/VulkanRHI/Private/VulkanRenderTarget.cpp
@@ -309,14 +309,57 @@ void FVulkanCommandListContext::RHICopyToResolveTarget(FRHITexture* SourceTextur
 	}
 }
 
+
+static void ConvertRawDataToFColor(VkFormat VulkanFormat, uint32 DestWidth, uint32 DestHeight, uint8* In, uint32 SrcPitch, FColor* Dest, const FReadSurfaceDataFlags& InFlags)
+{
+	switch (VulkanFormat)
+	{
+	case VK_FORMAT_R16G16B16A16_SFLOAT:
+		ConvertRawR16G16B16A16FDataToFColor(DestWidth, DestHeight, In, SrcPitch, Dest, false);
+		return;
+
+	case VK_FORMAT_A2B10G10R10_UNORM_PACK32:
+		ConvertRawR10G10B10A2DataToFColor(DestWidth, DestHeight, In, SrcPitch, Dest);
+		return;
+
+	case VK_FORMAT_R8G8B8A8_UNORM:
+		ConvertRawR8G8B8A8DataToFColor(DestWidth, DestHeight, In, SrcPitch, Dest);
+		return;
+
+	case VK_FORMAT_R16G16B16A16_UNORM:
+		ConvertRawR16G16B16A16DataToFColor(DestWidth, DestHeight, In, SrcPitch, Dest);
+		return;
+
+	case VK_FORMAT_B8G8R8A8_UNORM:
+		ConvertRawB8G8R8A8DataToFColor(DestWidth, DestHeight, In, SrcPitch, Dest);
+		return;
+
+	case VK_FORMAT_R8_UNORM:
+		ConvertRawR8DataToFColor(DestWidth, DestHeight, In, SrcPitch, Dest);
+		return;
+
+	case VK_FORMAT_R16_UNORM:
+		ConvertRawR16DataToFColor(DestWidth, DestHeight, In, SrcPitch, Dest);
+		return;
+
+	case VK_FORMAT_R16G16_UNORM:
+		ConvertRawR16G16DataToFColor(DestWidth, DestHeight, In, SrcPitch, Dest);
+		return;
+	}
+
+	checkf(false, TEXT("Unsupported format [%d] for conversion to FColor!"), (uint32)VulkanFormat);
+}
+
 void FVulkanDynamicRHI::RHIReadSurfaceData(FRHITexture* TextureRHI, FIntRect Rect, TArray<FColor>& OutData, FReadSurfaceDataFlags InFlags)
 {
-	uint32 NumPixels = (Rect.Max.X - Rect.Min.X) * (Rect.Max.Y - Rect.Min.Y);
-	if(GIgnoreCPUReads)
+	const uint32 DestWidth = Rect.Max.X - Rect.Min.X;
+	const uint32 DestHeight = Rect.Max.Y - Rect.Min.Y;
+	const uint32 NumRequestedPixels = DestWidth * DestHeight;
+	if (GIgnoreCPUReads)
 	{
 		// Debug: Fill with CPU
 		OutData.Empty(0);
-		OutData.AddZeroed(NumPixels);
+		OutData.AddZeroed(NumRequestedPixels);
 		return;
 	}
 
@@ -345,54 +388,46 @@ void FVulkanDynamicRHI::RHIReadSurfaceData(FRHITexture* TextureRHI, FIntRect Rec
 	{
 		// Just return black for texture types we don't support.
 		OutData.Empty(0);
-		OutData.AddZeroed(NumPixels);
+		OutData.AddZeroed(NumRequestedPixels);
 		return;
 	}
 
 	Device->PrepareForCPURead();
 
-	FVulkanCommandListContext& ImmediateContext = Device->GetImmediateContext();
+	// Figure out the size of the buffer required to hold the requested pixels
+	const uint32 PixelByteSize = GetNumBitsPerPixel(Surface->StorageFormat) / 8;
+	checkf(GPixelFormats[TextureRHI->GetFormat()].Supported && (PixelByteSize > 0), TEXT("Trying to read from unsupported format."));
+	const uint32 BufferSize = NumRequestedPixels * PixelByteSize;
 
-	ensure(Surface->StorageFormat == VK_FORMAT_R8G8B8A8_UNORM || Surface->StorageFormat == VK_FORMAT_B8G8R8A8_UNORM || Surface->StorageFormat == VK_FORMAT_R16G16B16A16_SFLOAT || Surface->StorageFormat == VK_FORMAT_A2B10G10R10_UNORM_PACK32 || Surface->StorageFormat == VK_FORMAT_R16G16B16A16_UNORM);
-	bool bIs8Bpp = true;
-	switch (Surface->StorageFormat)
-	{
-	case VK_FORMAT_R16G16B16A16_SFLOAT:
-	case VK_FORMAT_R16G16B16A16_SNORM:
-	case VK_FORMAT_R16G16B16A16_UINT:
-	case VK_FORMAT_R16G16B16A16_SINT:
-		bIs8Bpp = false;
-		break;
-	default:
-		break;
-	}
-	const uint32 Size = NumPixels * sizeof(FColor) * (bIs8Bpp ? 2 : 1);
+	// Validate that the Rect is within the texture
+	const uint32 MipLevel = InFlags.GetMip();
+	const uint32 MipSizeX = FMath::Max(Surface->Width >> MipLevel, 1u);
+	const uint32 MipSizeY = FMath::Max(Surface->Height >> MipLevel, 1u);
+	checkf((Rect.Max.X <= (int32)MipSizeX) && (Rect.Max.Y <= (int32)MipSizeY), TEXT("The specified Rect [%dx%d] extends beyond this Mip [%dx%d]."), Rect.Max.X, Rect.Max.Y, MipSizeX, MipSizeY);
+
+	FVulkanCommandListContext& ImmediateContext = Device->GetImmediateContext();
 
-	uint8* MappedPointer = nullptr;
 	VulkanRHI::FStagingBuffer* StagingBuffer = nullptr;
 	FVulkanCmdBuffer* CmdBuffer = nullptr;
-	bool bCPUReadback = EnumHasAllFlags(Surface->UEFlags, TexCreate_CPUReadback);
-	if(!bCPUReadback) //this function supports reading back arbitrary rendertargets, so if its not a cpu readback surface, we do a copy.
+	const bool bCPUReadback = EnumHasAllFlags(Surface->UEFlags, TexCreate_CPUReadback);
+	if (!bCPUReadback) //this function supports reading back arbitrary rendertargets, so if its not a cpu readback surface, we do a copy.
 	{
-		ImmediateContext.GetCommandBufferManager()->GetUploadCmdBuffer();
 		CmdBuffer = ImmediateContext.GetCommandBufferManager()->GetUploadCmdBuffer();
-		StagingBuffer = Device->GetStagingManager().AcquireBuffer(Size, VK_BUFFER_USAGE_TRANSFER_DST_BIT, VK_MEMORY_PROPERTY_HOST_CACHED_BIT);
+		StagingBuffer = Device->GetStagingManager().AcquireBuffer(BufferSize, VK_BUFFER_USAGE_TRANSFER_DST_BIT, VK_MEMORY_PROPERTY_HOST_CACHED_BIT);
 		VkBufferImageCopy CopyRegion;
 		FMemory::Memzero(CopyRegion);
-		uint32 MipLevel = InFlags.GetMip();
-		uint32 SizeX = FMath::Max(Surface->Width >> MipLevel, 1u);
-		uint32 SizeY = FMath::Max(Surface->Height >> MipLevel, 1u);
-		CopyRegion.bufferRowLength = SizeX;
-		CopyRegion.bufferImageHeight = SizeY;
+		// Leave bufferRowLength/bufferImageHeight at 0 for tightly packed
 		CopyRegion.imageSubresource.aspectMask = Surface->GetFullAspectMask();
 		CopyRegion.imageSubresource.mipLevel = MipLevel;
 		CopyRegion.imageSubresource.layerCount = 1;
-		CopyRegion.imageExtent.width = SizeX;
-		CopyRegion.imageExtent.height = SizeY;
+		CopyRegion.imageOffset.x = Rect.Min.X;
+		CopyRegion.imageOffset.y = Rect.Min.Y;
+		CopyRegion.imageExtent.width = DestWidth;
+		CopyRegion.imageExtent.height = DestHeight;
 		CopyRegion.imageExtent.depth = 1;
 
 		VkImageLayout& CurrentLayout = Device->GetImmediateContext().GetLayoutManager().FindOrAddLayoutRW(*Surface, VK_IMAGE_LAYOUT_UNDEFINED);
-		bool bHadLayout = (CurrentLayout != VK_IMAGE_LAYOUT_UNDEFINED);
+		const bool bHadLayout = (CurrentLayout != VK_IMAGE_LAYOUT_UNDEFINED);
 		if (CurrentLayout != VK_IMAGE_LAYOUT_TRANSFER_SRC_OPTIMAL)
 		{
 			VulkanSetImageLayoutAllMips(CmdBuffer->GetHandle(), Surface->Image, CurrentLayout, VK_IMAGE_LAYOUT_TRANSFER_SRC_OPTIMAL);
@@ -407,70 +442,37 @@ void FVulkanDynamicRHI::RHIReadSurfaceData(FRHITexture* TextureRHI, FIntRect Rec
 		{
 			CurrentLayout = VK_IMAGE_LAYOUT_TRANSFER_SRC_OPTIMAL;
 		}
-		ensure(StagingBuffer->GetSize() >= Size);
+		ensure(StagingBuffer->GetSize() >= BufferSize);
 
 		VkMemoryBarrier Barrier = { VK_STRUCTURE_TYPE_MEMORY_BARRIER , nullptr, VK_ACCESS_TRANSFER_WRITE_BIT, VK_ACCESS_HOST_READ_BIT };
 		VulkanRHI::vkCmdPipelineBarrier(CmdBuffer->GetHandle(), VK_PIPELINE_STAGE_TRANSFER_BIT, VK_PIPELINE_STAGE_HOST_BIT, 0, 1, &Barrier, 0, nullptr, 0, nullptr);
 
 		// Force upload
 		ImmediateContext.GetCommandBufferManager()->SubmitUploadCmdBuffer();
-
 	}
-	else
-	{
-		MappedPointer = (uint8*)Surface->GetMappedPointer();
-	}
-
-
 
 	Device->WaitUntilIdle();
-	if(!bCPUReadback)
-	{
-		StagingBuffer->InvalidateMappedMemory();
-		MappedPointer = (uint8*)StagingBuffer->GetMappedPointer();
-	}
-
-	OutData.SetNum(NumPixels);
-	FColor* Dest = OutData.GetData();
 
-	uint32 DestWidth = Rect.Max.X - Rect.Min.X;
-	uint32 DestHeight = Rect.Max.Y - Rect.Min.Y;
-	if (Surface->StorageFormat == VK_FORMAT_R16G16B16A16_SFLOAT)
-	{
-		uint32 PixelByteSize = 8u;
-		uint8* In = MappedPointer + (Rect.Min.Y * Surface->Width + Rect.Min.X) * PixelByteSize;
-		uint32 SrcPitch = Surface->Width * PixelByteSize;
-		ConvertRawR16G16B16A16FDataToFColor(DestWidth, DestHeight, In, SrcPitch, Dest, false);
-	}
-	else if (Surface->StorageFormat == VK_FORMAT_A2B10G10R10_UNORM_PACK32)
+	uint8* In;
+	uint32 SrcPitch;
+	if (bCPUReadback)
 	{
-		uint32 PixelByteSize = 4u;
-		uint8* In = MappedPointer + (Rect.Min.Y * Surface->Width + Rect.Min.X) * PixelByteSize;
-		uint32 SrcPitch = Surface->Width * PixelByteSize;
-		ConvertRawR10G10B10A2DataToFColor(DestWidth, DestHeight, In, SrcPitch, Dest);
-	}
-	else if (Surface->StorageFormat == VK_FORMAT_R8G8B8A8_UNORM)
-	{
-		uint32 PixelByteSize = 4u;
-		uint8* In = MappedPointer + (Rect.Min.Y * Surface->Width + Rect.Min.X) * PixelByteSize;
-		uint32 SrcPitch = Surface->Width * PixelByteSize;
-		ConvertRawR8G8B8A8DataToFColor(DestWidth, DestHeight, In, SrcPitch, Dest);
+		// If the text was bCPUReadback, then we have to deal with our Rect potentially being a subset of the total texture
+		In = (uint8*)Surface->GetMappedPointer() + ((Rect.Min.Y * MipSizeX + Rect.Min.X) * PixelByteSize);
+		SrcPitch = MipSizeX * PixelByteSize;
 	}
-	else if (Surface->StorageFormat == VK_FORMAT_R16G16B16A16_UNORM)
-	{
-		uint32 PixelByteSize = 8u;
-		uint8* In = MappedPointer + (Rect.Min.Y * Surface->Width + Rect.Min.X) * PixelByteSize;
-		uint32 SrcPitch = Surface->Width * PixelByteSize;
-		ConvertRawR16G16B16A16DataToFColor(DestWidth, DestHeight, In, SrcPitch, Dest);
-	}
-	else if (Surface->StorageFormat == VK_FORMAT_B8G8R8A8_UNORM)
+	else
 	{
-		uint32 PixelByteSize = 4u;
-		uint8* In = MappedPointer + (Rect.Min.Y * Surface->Width + Rect.Min.X) * PixelByteSize;
-		uint32 SrcPitch = Surface->Width * PixelByteSize;
-		ConvertRawB8G8R8A8DataToFColor(DestWidth, DestHeight, In, SrcPitch, Dest);
+		// If the text was NOT bCPUReadback, the buffer contains only the (tightly packed) Rect we requested
+		StagingBuffer->InvalidateMappedMemory();
+		In = (uint8*)StagingBuffer->GetMappedPointer();
+		SrcPitch = DestWidth * PixelByteSize;
 	}
 
+	OutData.SetNum(NumRequestedPixels);
+	FColor* Dest = OutData.GetData();
+	ConvertRawDataToFColor(Surface->StorageFormat, DestWidth, DestHeight, In, SrcPitch, Dest, InFlags);
+
 	if (!bCPUReadback)
 	{
 		Device->GetStagingManager().ReleaseBuffer(CmdBuffer, StagingBuffer);
