Removed unused field and forward declaration from incorrect merge

diff --git a/Engine/Source/Runtime/VulkanRHI/Private/VulkanRayTracing.h b/Engine/Source/Runtime/VulkanRHI/Private/VulkanRayTracing.h
--- a/Engine/Source/Runtime/VulkanRHI/Private/VulkanRayTracing.h
+++ b/Engine/Source/Runtime/VulkanRHI/Private/VulkanRayTracing.h
@@ -8,7 +8,6 @@
 
 class FVulkanCommandListContext;
 class FVulkanResourceMultiBuffer;
-class FVulkanAccelerationStructureBuffer;
 class FVulkanRayTracingLayout;
 
 #define ENUM_VK_ENTRYPOINTS_RAYTRACING(EnumMacro) \
@@ -74,7 +73,6 @@ struct FVkRtBLASBuildData
 
 	TArray<VkAccelerationStructureGeometryKHR, TInlineAllocator<1>> Segments;
 	TArray<VkAccelerationStructureBuildRangeInfoKHR, TInlineAllocator<1>> Ranges;
-	TArray<uint32, TInlineAllocator<1>> VertexCounts;
 	VkAccelerationStructureBuildGeometryInfoKHR GeometryInfo;
 	VkAccelerationStructureBuildSizesInfoKHR SizesInfo;
 };
