VulkanRHI: Use LinearToGamma specified in flags instead of hardcoding and add support for R32G32B32A32 (thanks Marijn!)

diff --git a/Engine/Source/Runtime/VulkanRHI/Private/VulkanRenderTarget.cpp b/Engine/Source/Runtime/VulkanRHI/Private/VulkanRenderTarget.cpp
--- a/Engine/Source/Runtime/VulkanRHI/Private/VulkanRenderTarget.cpp
+++ b/Engine/Source/Runtime/VulkanRHI/Private/VulkanRenderTarget.cpp
@@ -312,42 +312,50 @@ void FVulkanCommandListContext::RHICopyToResolveTarget(FRHITexture* SourceTextur
 
 static void ConvertRawDataToFColor(VkFormat VulkanFormat, uint32 DestWidth, uint32 DestHeight, uint8* In, uint32 SrcPitch, FColor* Dest, const FReadSurfaceDataFlags& InFlags)
 {
+	const bool bLinearToGamma = InFlags.GetLinearToGamma();
 	switch (VulkanFormat)
 	{
+	case VK_FORMAT_R32G32B32A32_SFLOAT:
+		ConvertRawR32G32B32A32DataToFColor(DestWidth, DestHeight, In, SrcPitch, Dest, bLinearToGamma);
+		break;
+
 	case VK_FORMAT_R16G16B16A16_SFLOAT:
-		ConvertRawR16G16B16A16FDataToFColor(DestWidth, DestHeight, In, SrcPitch, Dest, false);
-		return;
+		ConvertRawR16G16B16A16FDataToFColor(DestWidth, DestHeight, In, SrcPitch, Dest, bLinearToGamma);
+		break;
 
 	case VK_FORMAT_A2B10G10R10_UNORM_PACK32:
 		ConvertRawR10G10B10A2DataToFColor(DestWidth, DestHeight, In, SrcPitch, Dest);
-		return;
+		break;
 
 	case VK_FORMAT_R8G8B8A8_UNORM:
 		ConvertRawR8G8B8A8DataToFColor(DestWidth, DestHeight, In, SrcPitch, Dest);
-		return;
+		break;
 
 	case VK_FORMAT_R16G16B16A16_UNORM:
 		ConvertRawR16G16B16A16DataToFColor(DestWidth, DestHeight, In, SrcPitch, Dest);
-		return;
+		break;
 
 	case VK_FORMAT_B8G8R8A8_UNORM:
 		ConvertRawB8G8R8A8DataToFColor(DestWidth, DestHeight, In, SrcPitch, Dest);
-		return;
+		break;
 
 	case VK_FORMAT_R8_UNORM:
 		ConvertRawR8DataToFColor(DestWidth, DestHeight, In, SrcPitch, Dest);
-		return;
+		break;
 
 	case VK_FORMAT_R16_UNORM:
 		ConvertRawR16DataToFColor(DestWidth, DestHeight, In, SrcPitch, Dest);
-		return;
+		break;
 
 	case VK_FORMAT_R16G16_UNORM:
 		ConvertRawR16G16DataToFColor(DestWidth, DestHeight, In, SrcPitch, Dest);
-		return;
+		break;
+
+	default:
+		checkf(false, TEXT("Unsupported format [%d] for conversion to FColor!"), (uint32)VulkanFormat);
+		break;
 	}
 
-	checkf(false, TEXT("Unsupported format [%d] for conversion to FColor!"), (uint32)VulkanFormat);
 }
 
 void FVulkanDynamicRHI::RHIReadSurfaceData(FRHITexture* TextureRHI, FIntRect Rect, TArray<FColor>& OutData, FReadSurfaceDataFlags InFlags)
