Go back to previous behavior of forcing the tracking to a given value in DoAsyncReallocateTexture2D.  Fixes issue where external images add entries in the layout manager but potentially do not remove them, causing stale entries that can be hit on VkImage ptr re-use (and cause crashes).

diff --git a/Engine/Source/Runtime/VulkanRHI/Private/VulkanTexture.cpp b/Engine/Source/Runtime/VulkanRHI/Private/VulkanTexture.cpp
--- a/Engine/Source/Runtime/VulkanRHI/Private/VulkanTexture.cpp
+++ b/Engine/Source/Runtime/VulkanRHI/Private/VulkanTexture.cpp
@@ -1359,9 +1359,10 @@ static void DoAsyncReallocateTexture2D(FVulkanCommandListContext& Context, FVulk
 		Barrier.Execute(CmdBuffer->GetHandle());
 
 		// Add tracking for the appropriate subresources (intentionally leave added mips in VK_IMAGE_LAYOUT_UNDEFINED)
-		FVulkanImageLayout& NewTextureLayout = Context.GetLayoutManager().GetOrAddFullLayout(NewTexture->Surface, VK_IMAGE_LAYOUT_UNDEFINED);
-		ensure(NewTextureLayout.AreAllSubresourcesSameLayout() && (NewTextureLayout.MainLayout == VK_IMAGE_LAYOUT_UNDEFINED));
+		// NOTE: Overwriting whatever is contained in the layout manager for the new texture (circumvents issue with stale layouts and pointer reuse)
+		FVulkanImageLayout NewTextureLayout(VK_IMAGE_LAYOUT_UNDEFINED, NewTexture->Surface.GetNumMips(), NewTexture->Surface.GetNumberOfArrayLevels());
 		NewTextureLayout.Set(VK_IMAGE_LAYOUT_SHADER_READ_ONLY_OPTIMAL, DestSubResourceRange);
+		Context.GetLayoutManager().GetOrAddFullLayout(NewTexture->Surface, VK_IMAGE_LAYOUT_UNDEFINED) = NewTextureLayout;
 	}
 
 	// request is now complete
