- Remove tesselation stages from GVulkanDeviceShaderStageBits.
- Fix copy-pasta for DstAccessFlags during barrier generation
- Add new function to query stage bits supported by a queue.  Filter stages for barriers sent to async queue.
- Avoid unnecessary semaphores when DstPipelines is ERHIPipeline::All (the same sema wait was added on all queues)

WIP: Allows to run with AsyncCompute enabled (tested with FFTBloom).

diff --git a/Engine/Source/Runtime/VulkanRHI/Private/VulkanBarriers.cpp b/Engine/Source/Runtime/VulkanRHI/Private/VulkanBarriers.cpp
--- a/Engine/Source/Runtime/VulkanRHI/Private/VulkanBarriers.cpp
+++ b/Engine/Source/Runtime/VulkanRHI/Private/VulkanBarriers.cpp
@@ -560,7 +560,7 @@ void FVulkanDynamicRHI::RHICreateTransition(FRHITransition* Transition, const FR
 	const ERHIPipeline SrcPipelines = CreateInfo.SrcPipelines;
 	const ERHIPipeline DstPipelines = CreateInfo.DstPipelines;
 
-	checkf(FMath::IsPowerOfTwo((uint32)SrcPipelines) && FMath::IsPowerOfTwo((uint32)DstPipelines), TEXT("Support for multi-pipe resources is not yet implemented."));
+	//checkf(FMath::IsPowerOfTwo((uint32)SrcPipelines) && FMath::IsPowerOfTwo((uint32)DstPipelines), TEXT("Support for multi-pipe resources is not yet implemented."));
 
 	FVulkanPipelineBarrier* Data = new (Transition->GetPrivateData<FVulkanPipelineBarrier>()) FVulkanPipelineBarrier;
 	Data->SrcPipelines = SrcPipelines;
@@ -571,8 +571,8 @@ void FVulkanDynamicRHI::RHICreateTransition(FRHITransition* Transition, const FR
 	{
 		Data->Semaphore = new VulkanRHI::FSemaphore(*Device);
 
-		uint32 GfxQueueIndex = Device->GetGraphicsQueue()->GetFamilyIndex();
-		uint32 ComputeQueueIndex = Device->GetComputeQueue()->GetFamilyIndex();
+		const uint32 GfxQueueIndex = Device->GetGraphicsQueue()->GetFamilyIndex();
+		const uint32 ComputeQueueIndex = Device->GetComputeQueue()->GetFamilyIndex();
 
 		if (SrcPipelines == ERHIPipeline::Graphics)
 		{
@@ -769,17 +769,17 @@ void FVulkanDynamicRHI::RHICreateTransition(FRHITransition* Transition, const FR
 			}
 		}
 
-		// In case of async compute, override the stage and access flags computed above, since only the compute shader stage is relevant.
+		// In case of async compute, override the stage and access flags computed above, since only a subset of stages is relevant.
 		if (SrcPipelines == ERHIPipeline::AsyncCompute)
 		{
-			SrcStageMask = VK_PIPELINE_STAGE_COMPUTE_SHADER_BIT;
+			SrcStageMask = SrcStageMask & Device->GetComputeQueue()->GetSupportedStageBits();
 			SrcAccessFlags = SrcAccessFlags & (VK_ACCESS_SHADER_READ_BIT | VK_ACCESS_SHADER_WRITE_BIT);
 		}
 
 		if (DstPipelines == ERHIPipeline::AsyncCompute)
 		{
-			DstStageMask = VK_PIPELINE_STAGE_COMPUTE_SHADER_BIT;
-			DstAccessFlags = SrcAccessFlags & (VK_ACCESS_SHADER_READ_BIT | VK_ACCESS_SHADER_WRITE_BIT);
+			DstStageMask = DstStageMask & Device->GetComputeQueue()->GetSupportedStageBits();
+			DstAccessFlags = DstAccessFlags & (VK_ACCESS_SHADER_READ_BIT | VK_ACCESS_SHADER_WRITE_BIT);
 		}
 
 
@@ -1052,6 +1052,11 @@ void FVulkanCommandListContext::RHIBeginTransitions(TArrayView<const FRHITransit
 			// We don't update the image layout here. That will be done in RHIEndTransitions.
 		}
 
+		if (EnumHasAnyFlags(Data->SrcPipelines, ERHIPipeline::AsyncCompute))
+		{
+			RealSrcStageMask = RealSrcStageMask & Queue->GetSupportedStageBits();
+		}
+
 		VulkanRHI::vkCmdPipelineBarrier(CmdBuffer->GetHandle(), RealSrcStageMask, RealDstStageMask, 0, 0, nullptr, RealBufferBarriers.Num(), RealBufferBarriers.GetData(), RealImageBarriers.Num(), RealImageBarriers.GetData());
 
 		check(Data->Semaphore);
@@ -1122,6 +1127,8 @@ void FVulkanCommandListContext::RHIEndTransitions(TArrayView<const FRHITransitio
 
 	FVulkanCmdBuffer* CmdBuffer = CommandBufferManager->GetActiveCmdBuffer();
 
+	const ERHIPipeline CurrentPipeline = Device->IsRealAsyncComputeContext(this) ? ERHIPipeline::AsyncCompute : ERHIPipeline::Graphics;
+
 	bool bSeenWaitSemaphore = false;
 	for (const FRHITransition* Transition : Transitions)
 	{
@@ -1131,6 +1138,12 @@ void FVulkanCommandListContext::RHIEndTransitions(TArrayView<const FRHITransitio
 			continue;
 		}
 
+		// If the source matches our current queue, we don't need a semaphore (avoid unnecessary semaphores when DstPipelines is ERHIPipeline::All)
+		if (EnumHasAllFlags(Data->DstPipelines, Data->SrcPipelines) && (Data->SrcPipelines == CurrentPipeline))
+		{
+			continue;
+		}
+
 		if (!bSeenWaitSemaphore)
 		{
 			if (CommandBufferManager->HasPendingActiveCmdBuffer())
@@ -1272,6 +1285,11 @@ void FVulkanCommandListContext::RHIEndTransitions(TArrayView<const FRHITransitio
 			continue;
 		}
 
+		if (EnumHasAnyFlags(Data->DstPipelines, ERHIPipeline::AsyncCompute))
+		{
+			RealDstStageMask = RealDstStageMask & Queue->GetSupportedStageBits();
+		}
+
 		VulkanRHI::vkCmdPipelineBarrier(CmdBuffer->GetHandle(), RealSrcStageMask, RealDstStageMask, 0, RealMemoryBarrier.bHasMemoryBarrier ? 1 : 0, &RealMemoryBarrier.MemoryBarrier, RealBufferBarriers.Num(), RealBufferBarriers.GetData(), RealImageBarriers.Num(), RealImageBarriers.GetData());
 	}
 }
diff --git a/Engine/Source/Runtime/VulkanRHI/Private/VulkanDevice.h b/Engine/Source/Runtime/VulkanRHI/Private/VulkanDevice.h
--- a/Engine/Source/Runtime/VulkanRHI/Private/VulkanDevice.h
+++ b/Engine/Source/Runtime/VulkanRHI/Private/VulkanDevice.h
@@ -518,7 +518,11 @@ public:
 	VkSamplerYcbcrConversion CreateSamplerColorConversion(const VkSamplerYcbcrConversionCreateInfo& CreateInfo);
 #endif
 
-	void*	Hotfix = nullptr;
+	inline const TArray<VkQueueFamilyProperties>& GetQueueFamilyProps()
+	{
+		return QueueFamilyProps;
+	}
+
 
 private:
 	const VkFormatProperties& GetFormatProperties(VkFormat InFormat);
diff --git a/Engine/Source/Runtime/VulkanRHI/Private/VulkanQueue.cpp b/Engine/Source/Runtime/VulkanRHI/Private/VulkanQueue.cpp
--- a/Engine/Source/Runtime/VulkanRHI/Private/VulkanQueue.cpp
+++ b/Engine/Source/Runtime/VulkanRHI/Private/VulkanQueue.cpp
@@ -31,6 +31,8 @@ FVulkanQueue::FVulkanQueue(FVulkanDevice* InDevice, uint32 InFamilyIndex)
 	, SubmitCounter(0)
 {
 	VulkanRHI::vkGetDeviceQueue(Device->GetInstanceHandle(), FamilyIndex, QueueIndex, &Queue);
+
+	FillSupportedStageBits();
 }
 
 FVulkanQueue::~FVulkanQueue()
@@ -122,3 +124,61 @@ void FVulkanQueue::UpdateLastSubmittedCommandBuffer(FVulkanCmdBuffer* CmdBuffer)
 	LastSubmittedCmdBufferFenceCounter = CmdBuffer->GetFenceSignaledCounterH();
 	++SubmitCounter;
 }
+
+void FVulkanQueue::FillSupportedStageBits()
+{
+	check(Device);
+	check((int32)FamilyIndex < Device->GetQueueFamilyProps().Num());
+
+	const VkQueueFamilyProperties& QueueProps = Device->GetQueueFamilyProps()[FamilyIndex];
+
+	SupportedStages = 
+		VK_PIPELINE_STAGE_TOP_OF_PIPE_BIT | 
+		VK_PIPELINE_STAGE_BOTTOM_OF_PIPE_BIT |
+		VK_PIPELINE_STAGE_HOST_BIT |
+		VK_PIPELINE_STAGE_ALL_COMMANDS_BIT;
+
+	if (VKHasAnyFlags(QueueProps.queueFlags, VK_QUEUE_GRAPHICS_BIT))
+	{
+		SupportedStages |=
+			VK_PIPELINE_STAGE_DRAW_INDIRECT_BIT |
+			VK_PIPELINE_STAGE_VERTEX_INPUT_BIT |
+			VK_PIPELINE_STAGE_VERTEX_SHADER_BIT |
+			VK_PIPELINE_STAGE_FRAGMENT_SHADER_BIT |
+			VK_PIPELINE_STAGE_EARLY_FRAGMENT_TESTS_BIT |
+			VK_PIPELINE_STAGE_LATE_FRAGMENT_TESTS_BIT |
+			VK_PIPELINE_STAGE_COLOR_ATTACHMENT_OUTPUT_BIT |
+			VK_PIPELINE_STAGE_TRANSFER_BIT |
+			VK_PIPELINE_STAGE_ALL_GRAPHICS_BIT;
+
+		if (Device->GetPhysicalFeatures().geometryShader)
+		{
+			SupportedStages |= VK_PIPELINE_STAGE_GEOMETRY_SHADER_BIT;
+		}
+
+#if VULKAN_SUPPORTS_FRAGMENT_SHADING_RATE
+		SupportedStages |= VK_PIPELINE_STAGE_FRAGMENT_SHADING_RATE_ATTACHMENT_BIT_KHR;
+#endif
+
+#if VULKAN_SUPPORTS_FRAGMENT_DENSITY_MAP
+		SupportedStages |= VK_PIPELINE_STAGE_FRAGMENT_DENSITY_PROCESS_BIT_EXT;
+#endif
+	}
+
+	if (VKHasAnyFlags(QueueProps.queueFlags, VK_QUEUE_COMPUTE_BIT))
+	{
+		SupportedStages |=
+			VK_PIPELINE_STAGE_COMPUTE_SHADER_BIT |
+			VK_PIPELINE_STAGE_DRAW_INDIRECT_BIT |
+			VK_PIPELINE_STAGE_TRANSFER_BIT;
+
+#if VULKAN_RHI_RAYTRACING
+		SupportedStages |= VK_PIPELINE_STAGE_ACCELERATION_STRUCTURE_BUILD_BIT_KHR | VK_PIPELINE_STAGE_RAY_TRACING_SHADER_BIT_KHR;
+#endif
+	}
+
+	if (VKHasAnyFlags(QueueProps.queueFlags, VK_QUEUE_TRANSFER_BIT))
+	{
+		SupportedStages |= VK_PIPELINE_STAGE_TRANSFER_BIT;
+	}
+}
diff --git a/Engine/Source/Runtime/VulkanRHI/Private/VulkanQueue.h b/Engine/Source/Runtime/VulkanRHI/Private/VulkanQueue.h
--- a/Engine/Source/Runtime/VulkanRHI/Private/VulkanQueue.h
+++ b/Engine/Source/Runtime/VulkanRHI/Private/VulkanQueue.h
@@ -58,6 +58,11 @@ public:
 		return SubmitCounter;
 	}
 
+	inline VkPipelineStageFlags GetSupportedStageBits()
+	{
+		return SupportedStages;
+	}
+
 private:
 	VkQueue Queue;
 	uint32 FamilyIndex;
@@ -68,6 +73,8 @@ private:
 	FVulkanCmdBuffer* LastSubmittedCmdBuffer;
 	uint64 LastSubmittedCmdBufferFenceCounter;
 	uint64 SubmitCounter;
+	VkPipelineStageFlags SupportedStages;
 
 	void UpdateLastSubmittedCommandBuffer(FVulkanCmdBuffer* CmdBuffer);
+	void FillSupportedStageBits();
 };
diff --git a/Engine/Source/Runtime/VulkanRHI/Private/VulkanRHI.cpp b/Engine/Source/Runtime/VulkanRHI/Private/VulkanRHI.cpp
--- a/Engine/Source/Runtime/VulkanRHI/Private/VulkanRHI.cpp
+++ b/Engine/Source/Runtime/VulkanRHI/Private/VulkanRHI.cpp
@@ -861,12 +861,6 @@ void FVulkanDynamicRHI::InitInstance()
 			GVulkanDeviceShaderStageBits|= VK_PIPELINE_STAGE_GEOMETRY_SHADER_BIT;
 		}
 		
-		if (Device->GetPhysicalFeatures().tessellationShader)
-		{
-			GVulkanDeviceShaderStageBits|= VK_PIPELINE_STAGE_TESSELLATION_CONTROL_SHADER_BIT;
-			GVulkanDeviceShaderStageBits|= VK_PIPELINE_STAGE_TESSELLATION_EVALUATION_SHADER_BIT;
-		}
-
 		FHardwareInfo::RegisterHardwareInfo(NAME_RHI, TEXT("Vulkan"));
 
 		SavePipelineCacheCmd = IConsoleManager::Get().RegisterConsoleCommand(
