Add barrier in Buffer Unlock to meet the expectation of higher level code that data is ready to read. Fixes issue with UI icons in menus sometimes "exploding" under certain conditions (repro'd on Vulkan in Linux Editor with AMD).

diff --git a/Engine/Source/Runtime/VulkanRHI/Private/VulkanBuffer.cpp b/Engine/Source/Runtime/VulkanRHI/Private/VulkanBuffer.cpp
--- a/Engine/Source/Runtime/VulkanRHI/Private/VulkanBuffer.cpp
+++ b/Engine/Source/Runtime/VulkanRHI/Private/VulkanBuffer.cpp
@@ -377,6 +377,10 @@ inline void FVulkanResourceMultiBuffer::InternalUnlock(FVulkanCommandListContext
 	Region.dstOffset = LockOffset + MultiBuffer->Buffers[InDynamicBufferIndex].Offset;
 	VulkanRHI::vkCmdCopyBuffer(CmdBuffer, StagingBuffer->GetHandle(), MultiBuffer->Buffers[InDynamicBufferIndex].GetBufferHandle(), 1, &Region);
 
+	// High level code expects the data in MultiBuffer to be ready to read
+	VkMemoryBarrier BarrierAfter = { VK_STRUCTURE_TYPE_MEMORY_BARRIER, nullptr, VK_ACCESS_TRANSFER_WRITE_BIT, VK_ACCESS_MEMORY_READ_BIT | VK_ACCESS_MEMORY_WRITE_BIT };
+	VulkanRHI::vkCmdPipelineBarrier(CmdBuffer, VK_PIPELINE_STAGE_TRANSFER_BIT, VK_PIPELINE_STAGE_ALL_COMMANDS_BIT, 0, 1, &BarrierAfter, 0, nullptr, 0, nullptr);
+
 	MultiBuffer->GetParent()->GetStagingManager().ReleaseBuffer(Cmd, StagingBuffer);
 }
 
