Missing CS lock in VulkanMemory:  RenderThread (texture creation) overlaps with RHIThread (RHIEndFrame), both can access the ActivePages.

diff --git a/Engine/Source/Runtime/VulkanRHI/Private/VulkanMemory.cpp b/Engine/Source/Runtime/VulkanRHI/Private/VulkanMemory.cpp
--- a/Engine/Source/Runtime/VulkanRHI/Private/VulkanMemory.cpp
+++ b/Engine/Source/Runtime/VulkanRHI/Private/VulkanMemory.cpp
@@ -2424,25 +2424,28 @@ namespace VulkanRHI
 		FVulkanResourceHeap* BestDefragHeap = 0;
 		uint64 BestDefragHeapSize = 0;
 
-		uint32 Idx = 0;
-		for(FVulkanResourceHeap* Heap : ResourceTypeHeaps)
 		{
-			if(Heap->HeapIndex == PrimaryHeapIndex)
+			FScopeLock ScopeLock(&GResourceLock);
+			uint32 Idx = 0;
+			for (FVulkanResourceHeap* Heap : ResourceTypeHeaps)
 			{
-				uint64 UsedSize = Heap->UsedMemory;
-				if(CanDefragHeap(Heap) && BestDefragHeapSize < UsedSize)
+				if (Heap->HeapIndex == PrimaryHeapIndex)
 				{
-					BestDefragHeap = Heap;
-					BestDefragHeapSize = UsedSize;
-				}
-				if(CanEvictHeap(Heap) && BestEvictHeapSize < UsedSize)
-				{
-					BestEvictHeap = Heap;
-					BestEvictHeapSize = UsedSize;
+					uint64 UsedSize = Heap->UsedMemory;
+					if (CanDefragHeap(Heap) && BestDefragHeapSize < UsedSize)
+					{
+						BestDefragHeap = Heap;
+						BestDefragHeapSize = UsedSize;
+					}
+					if (CanEvictHeap(Heap) && BestEvictHeapSize < UsedSize)
+					{
+						BestEvictHeap = Heap;
+						BestEvictHeapSize = UsedSize;
+					}
+
 				}
-				
+				Idx++;
 			}
-			Idx++;
 		}
 
 		if(BestEvictHeap && ((GVulkanEvictOnePage || UpdateEvictThreshold(true)) && PrimaryHeapIndex >= 0))
