Deprecating RHIGetResourceInfo in favor of FRHIResource::GetResourceInfo

diff --git a/Engine/Source/Runtime/Apple/AGXRHI/Private/AGXDynamicRHI.h b/Engine/Source/Runtime/Apple/AGXRHI/Private/AGXDynamicRHI.h
--- a/Engine/Source/Runtime/Apple/AGXRHI/Private/AGXDynamicRHI.h
+++ b/Engine/Source/Runtime/Apple/AGXRHI/Private/AGXDynamicRHI.h
@@ -78,7 +78,6 @@ public:
 	virtual void RHICopySharedMips(FRHITexture2D* DestTexture2D, FRHITexture2D* SrcTexture2D) final override;
 	virtual FTexture2DArrayRHIRef RHICreateTexture2DArray(uint32 SizeX, uint32 SizeY, uint32 SizeZ, uint8 Format, uint32 NumMips, uint32 NumSamples, ETextureCreateFlags Flags, ERHIAccess InResourceState, FRHIResourceCreateInfo& CreateInfo) final override;
 	virtual FTexture3DRHIRef RHICreateTexture3D(uint32 SizeX, uint32 SizeY, uint32 SizeZ, uint8 Format, uint32 NumMips, ETextureCreateFlags Flags, ERHIAccess InResourceState, FRHIResourceCreateInfo& CreateInfo) final override;
-	virtual void RHIGetResourceInfo(FRHITexture* Ref, FRHIResourceInfo& OutInfo) final override;
 	virtual void RHIGenerateMips(FRHITexture* Texture) final override;
 	virtual uint32 RHIComputeMemorySize(FRHITexture* TextureRHI) final override;
 	virtual FTexture2DRHIRef RHIAsyncReallocateTexture2D(FRHITexture2D* Texture2D, int32 NewMipCount, int32 NewSizeX, int32 NewSizeY, FThreadSafeCounter* RequestStatus) final override;
diff --git a/Engine/Source/Runtime/Apple/AGXRHI/Private/AGXTexture.cpp b/Engine/Source/Runtime/Apple/AGXRHI/Private/AGXTexture.cpp
--- a/Engine/Source/Runtime/Apple/AGXRHI/Private/AGXTexture.cpp
+++ b/Engine/Source/Runtime/Apple/AGXRHI/Private/AGXTexture.cpp
@@ -1671,12 +1671,6 @@ FTexture3DRHIRef FAGXDynamicRHI::RHICreateTexture3D(uint32 SizeX, uint32 SizeY,
 	}
 }
 
-void FAGXDynamicRHI::RHIGetResourceInfo(FRHITexture* Ref, FRHIResourceInfo& OutInfo)
-{
-	// @todo Needed for visualisation!!
-	// NOT_SUPPORTED("RHIGetResourceInfo");
-}
-
 void FAGXDynamicRHI::RHIGenerateMips(FRHITexture* SourceSurfaceRHI)
 {
 	@autoreleasepool {
diff --git a/Engine/Source/Runtime/Apple/MetalRHI/Private/MetalDynamicRHI.h b/Engine/Source/Runtime/Apple/MetalRHI/Private/MetalDynamicRHI.h
--- a/Engine/Source/Runtime/Apple/MetalRHI/Private/MetalDynamicRHI.h
+++ b/Engine/Source/Runtime/Apple/MetalRHI/Private/MetalDynamicRHI.h
@@ -79,7 +79,6 @@ public:
 	virtual void RHICopySharedMips(FRHITexture2D* DestTexture2D, FRHITexture2D* SrcTexture2D) final override;
 	virtual FTexture2DArrayRHIRef RHICreateTexture2DArray(uint32 SizeX, uint32 SizeY, uint32 SizeZ, uint8 Format, uint32 NumMips, uint32 NumSamples, ETextureCreateFlags Flags, ERHIAccess InResourceState, FRHIResourceCreateInfo& CreateInfo) final override;
 	virtual FTexture3DRHIRef RHICreateTexture3D(uint32 SizeX, uint32 SizeY, uint32 SizeZ, uint8 Format, uint32 NumMips, ETextureCreateFlags Flags, ERHIAccess InResourceState, FRHIResourceCreateInfo& CreateInfo) final override;
-	virtual void RHIGetResourceInfo(FRHITexture* Ref, FRHIResourceInfo& OutInfo) final override;
 	virtual void RHIGenerateMips(FRHITexture* Texture) final override;
 	virtual uint32 RHIComputeMemorySize(FRHITexture* TextureRHI) final override;
 	virtual FTexture2DRHIRef RHIAsyncReallocateTexture2D(FRHITexture2D* Texture2D, int32 NewMipCount, int32 NewSizeX, int32 NewSizeY, FThreadSafeCounter* RequestStatus) final override;
diff --git a/Engine/Source/Runtime/Apple/MetalRHI/Private/MetalTexture.cpp b/Engine/Source/Runtime/Apple/MetalRHI/Private/MetalTexture.cpp
--- a/Engine/Source/Runtime/Apple/MetalRHI/Private/MetalTexture.cpp
+++ b/Engine/Source/Runtime/Apple/MetalRHI/Private/MetalTexture.cpp
@@ -1761,12 +1761,6 @@ FTexture3DRHIRef FMetalDynamicRHI::RHICreateTexture3D(uint32 SizeX, uint32 SizeY
 	}
 }
 
-void FMetalDynamicRHI::RHIGetResourceInfo(FRHITexture* Ref, FRHIResourceInfo& OutInfo)
-{
-	// @todo Needed for visualisation!!
-	// NOT_SUPPORTED("RHIGetResourceInfo");
-}
-
 void FMetalDynamicRHI::RHIGenerateMips(FRHITexture* SourceSurfaceRHI)
 {
 	@autoreleasepool {
diff --git a/Engine/Source/Runtime/D3D12RHI/Private/D3D12RHIPrivate.h b/Engine/Source/Runtime/D3D12RHI/Private/D3D12RHIPrivate.h
--- a/Engine/Source/Runtime/D3D12RHI/Private/D3D12RHIPrivate.h
+++ b/Engine/Source/Runtime/D3D12RHI/Private/D3D12RHIPrivate.h
@@ -333,7 +333,6 @@ public:
 	virtual void RHICopySharedMips(FRHITexture2D* DestTexture2D, FRHITexture2D* SrcTexture2D) final override;
 	virtual FTexture2DArrayRHIRef RHICreateTexture2DArray(uint32 SizeX, uint32 SizeY, uint32 SizeZ, uint8 Format, uint32 NumMips, uint32 NumSamples, ETextureCreateFlags Flags, ERHIAccess InResourceState, FRHIResourceCreateInfo& CreateInfo) override;
 	virtual FTexture3DRHIRef RHICreateTexture3D(uint32 SizeX, uint32 SizeY, uint32 SizeZ, uint8 Format, uint32 NumMips, ETextureCreateFlags Flags, ERHIAccess InResourceState, FRHIResourceCreateInfo& CreateInfo) override;
-	virtual void RHIGetResourceInfo(FRHITexture* Ref, FRHIResourceInfo& OutInfo) override;
 	virtual FShaderResourceViewRHIRef RHICreateShaderResourceView(FRHITexture* Texture, const FRHITextureSRVCreateInfo& CreateInfo) final override;
 	virtual uint32 RHIComputeMemorySize(FRHITexture* TextureRHI) final override;
 	virtual FTexture2DRHIRef RHIAsyncReallocateTexture2D(FRHITexture2D* Texture2D, int32 NewMipCount, int32 NewSizeX, int32 NewSizeY, FThreadSafeCounter* RequestStatus) override;
diff --git a/Engine/Source/Runtime/D3D12RHI/Private/D3D12Texture.cpp b/Engine/Source/Runtime/D3D12RHI/Private/D3D12Texture.cpp
--- a/Engine/Source/Runtime/D3D12RHI/Private/D3D12Texture.cpp
+++ b/Engine/Source/Runtime/D3D12RHI/Private/D3D12Texture.cpp
@@ -1748,10 +1748,6 @@ FTexture3DRHIRef FD3D12DynamicRHI::RHICreateTexture3D(uint32 SizeX, uint32 SizeY
 #endif // PLATFORM_WINDOWS || PLATFORM_HOLOLENS
 }
 
-void FD3D12DynamicRHI::RHIGetResourceInfo(FRHITexture* Ref, FRHIResourceInfo& OutInfo)
-{
-}
-
 /**
  * Computes the size in memory required by a given texture.
  *
diff --git a/Engine/Source/Runtime/NullDrv/Public/NullRHI.h b/Engine/Source/Runtime/NullDrv/Public/NullRHI.h
--- a/Engine/Source/Runtime/NullDrv/Public/NullRHI.h
+++ b/Engine/Source/Runtime/NullDrv/Public/NullRHI.h
@@ -238,9 +238,6 @@ public:
 	{ 
 		return new FRHITexture3D(SizeX, SizeY, SizeZ, NumMips, (EPixelFormat)Format, Flags, CreateInfo.ClearValueBinding);
 	}
-	virtual void RHIGetResourceInfo(FRHITexture* Ref, FRHIResourceInfo& OutInfo) final override
-	{
-	}
 
 	virtual FShaderResourceViewRHIRef RHICreateShaderResourceView(FRHITexture* Texture, const FRHITextureSRVCreateInfo& CreateInfo) final override
 	{ 
diff --git a/Engine/Source/Runtime/OpenGLDrv/Private/OpenGLTexture.cpp b/Engine/Source/Runtime/OpenGLDrv/Private/OpenGLTexture.cpp
--- a/Engine/Source/Runtime/OpenGLDrv/Private/OpenGLTexture.cpp
+++ b/Engine/Source/Runtime/OpenGLDrv/Private/OpenGLTexture.cpp
@@ -1904,10 +1904,6 @@ FTexture3DRHIRef FOpenGLDynamicRHI::RHICreateTexture3D(uint32 SizeX,uint32 SizeY
 	return Texture;
 }
 
-void FOpenGLDynamicRHI::RHIGetResourceInfo(FRHITexture* Ref, FRHIResourceInfo& OutInfo)
-{
-}
-
 FShaderResourceViewRHIRef FOpenGLDynamicRHI::RHICreateShaderResourceView(FRHITexture* Texture, const FRHITextureSRVCreateInfo& CreateInfo)
 {
 	const uint32 MipLevel = CreateInfo.MipLevel;
diff --git a/Engine/Source/Runtime/OpenGLDrv/Public/OpenGLDrv.h b/Engine/Source/Runtime/OpenGLDrv/Public/OpenGLDrv.h
--- a/Engine/Source/Runtime/OpenGLDrv/Public/OpenGLDrv.h
+++ b/Engine/Source/Runtime/OpenGLDrv/Public/OpenGLDrv.h
@@ -366,7 +366,6 @@ public:
 	virtual void RHICopySharedMips(FRHITexture2D* DestTexture2D, FRHITexture2D* SrcTexture2D) final override;
 	virtual FTexture2DArrayRHIRef RHICreateTexture2DArray(uint32 SizeX, uint32 SizeY, uint32 SizeZ, uint8 Format, uint32 NumMips, uint32 NumSamples, ETextureCreateFlags Flags, ERHIAccess InResourceState, FRHIResourceCreateInfo& CreateInfo) final override;
 	virtual FTexture3DRHIRef RHICreateTexture3D(uint32 SizeX, uint32 SizeY, uint32 SizeZ, uint8 Format, uint32 NumMips, ETextureCreateFlags Flags, ERHIAccess InResourceState, FRHIResourceCreateInfo& CreateInfo) final override;
-	virtual void RHIGetResourceInfo(FRHITexture* Ref, FRHIResourceInfo& OutInfo) final override;
 	virtual FShaderResourceViewRHIRef RHICreateShaderResourceView(FRHITexture* Texture, const FRHITextureSRVCreateInfo& CreateInfo) final override;
 	virtual void RHIGenerateMips(FRHITexture* Texture) final override;
 	virtual uint32 RHIComputeMemorySize(FRHITexture* TextureRHI) final override;
diff --git a/Engine/Source/Runtime/RHI/Public/DynamicRHI.h b/Engine/Source/Runtime/RHI/Public/DynamicRHI.h
--- a/Engine/Source/Runtime/RHI/Public/DynamicRHI.h
+++ b/Engine/Source/Runtime/RHI/Public/DynamicRHI.h
@@ -876,7 +876,16 @@ public:
 	* @param Ref may be 0
 	*/
 	// FlushType: Thread safe
-	virtual void RHIGetResourceInfo(FRHITexture* Ref, FRHIResourceInfo& OutInfo) = 0;
+	UE_DEPRECATED(5.0, "RHIGetResourceInfo is no longer implemented in favor of FRHIResource::GetResourceInfo.")
+	virtual void RHIGetResourceInfo(FRHITexture* Ref, FRHIResourceInfo& OutInfo)
+	{
+#if RHI_ENABLE_RESOURCE_INFO
+		if (Ref)
+		{
+			Ref->GetResourceInfo(OutInfo);
+		}
+#endif
+	}
 
 	/**
 	* Creates a shader resource view for a texture
@@ -1675,9 +1684,12 @@ FORCEINLINE void RHIGetTextureMemoryStats(FTextureMemoryStats& OutStats)
 	GDynamicRHI->RHIGetTextureMemoryStats(OutStats);
 }
 
+UE_DEPRECATED(5.0, "RHIGetResourceInfo is no longer implemented in favor of FRHIResource::GetResourceInfo.")
 FORCEINLINE void RHIGetResourceInfo(FRHITexture* Ref, FRHIResourceInfo& OutInfo)
 {
-	return GDynamicRHI->RHIGetResourceInfo(Ref, OutInfo);
+PRAGMA_DISABLE_DEPRECATION_WARNINGS
+	GDynamicRHI->RHIGetResourceInfo(Ref, OutInfo);
+PRAGMA_ENABLE_DEPRECATION_WARNINGS
 }
 
 FORCEINLINE uint32 RHIComputeMemorySize(FRHITexture* TextureRHI)
diff --git a/Engine/Source/Runtime/RHI/Public/RHICommandList.h b/Engine/Source/Runtime/RHI/Public/RHICommandList.h
--- a/Engine/Source/Runtime/RHI/Public/RHICommandList.h
+++ b/Engine/Source/Runtime/RHI/Public/RHICommandList.h
@@ -4406,9 +4406,12 @@ public:
 		});
 	}
 
+	UE_DEPRECATED(5.0, "RHIGetResourceInfo is no longer implemented in favor of FRHIResource::GetResourceInfo.")
 	FORCEINLINE void GetResourceInfo(FRHITexture* Ref, FRHIResourceInfo& OutInfo)
 	{
+PRAGMA_DISABLE_DEPRECATION_WARNINGS
 		return RHIGetResourceInfo(Ref, OutInfo);
+PRAGMA_ENABLE_DEPRECATION_WARNINGS
 	}
 	
 	FORCEINLINE FShaderResourceViewRHIRef CreateShaderResourceView(FRHITexture* Texture, const FRHITextureSRVCreateInfo& CreateInfo)
diff --git a/Engine/Source/Runtime/RHI/Public/RHIValidation.h b/Engine/Source/Runtime/RHI/Public/RHIValidation.h
--- a/Engine/Source/Runtime/RHI/Public/RHIValidation.h
+++ b/Engine/Source/Runtime/RHI/Public/RHIValidation.h
@@ -678,15 +678,6 @@ public:
 		return Texture;
 	}
 
-	/**
-	* @param Ref may be 0
-	*/
-	// FlushType: Thread safe
-	virtual void RHIGetResourceInfo(FRHITexture* Ref, FRHIResourceInfo& OutInfo) override final
-	{
-		return RHI->RHIGetResourceInfo(Ref, OutInfo);
-	}
-
 	/**
 	* Creates a shader resource view for a texture
 	*/
diff --git a/Engine/Source/Runtime/VulkanRHI/Private/VulkanResources.h b/Engine/Source/Runtime/VulkanRHI/Private/VulkanResources.h
--- a/Engine/Source/Runtime/VulkanRHI/Private/VulkanResources.h
+++ b/Engine/Source/Runtime/VulkanRHI/Private/VulkanResources.h
@@ -575,6 +575,8 @@ struct FVulkanTextureBase : public FVulkanEvictable, public IRefCountedObject
 	void AttachView(VulkanRHI::FVulkanViewBase* View);
 	void DetachView(VulkanRHI::FVulkanViewBase* View);
 
+	bool GetTextureResourceInfo(FRHIResourceInfo& OutResourceInfo) const;
+
 	virtual FRHITexture* GetRHITexture() = 0;
 private:
 	void InvalidateViews(FVulkanDevice& Device);
@@ -628,6 +630,13 @@ public:
 	{
 		return (void*)Surface.Image;
 	}
+
+#if RHI_ENABLE_RESOURCE_INFO
+	virtual bool GetResourceInfo(FRHIResourceInfo& OutResourceInfo) const override final
+	{
+		return GetTextureResourceInfo(OutResourceInfo);
+	}
+#endif
 };
 
 class FVulkanTexture2DArray : public FRHITexture2DArray, public FVulkanTextureBase
@@ -671,6 +680,13 @@ public:
 	{
 		return (void*)Surface.Image;
 	}
+
+#if RHI_ENABLE_RESOURCE_INFO
+	virtual bool GetResourceInfo(FRHIResourceInfo& OutResourceInfo) const override final
+	{
+		return GetTextureResourceInfo(OutResourceInfo);
+	}
+#endif
 };
 
 class FVulkanTexture3D : public FRHITexture3D, public FVulkanTextureBase
@@ -712,6 +728,13 @@ public:
 	{
 		return (void*)Surface.Image;
 	}
+
+#if RHI_ENABLE_RESOURCE_INFO
+	virtual bool GetResourceInfo(FRHIResourceInfo& OutResourceInfo) const override final
+	{
+		return GetTextureResourceInfo(OutResourceInfo);
+	}
+#endif
 };
 
 class FVulkanTextureCube : public FRHITextureCube, public FVulkanTextureBase
@@ -755,6 +778,13 @@ public:
 	{
 		return (void*)Surface.Image;
 	}
+
+#if RHI_ENABLE_RESOURCE_INFO
+	virtual bool GetResourceInfo(FRHIResourceInfo& OutResourceInfo) const override final
+	{
+		return GetTextureResourceInfo(OutResourceInfo);
+	}
+#endif
 };
 
 class FVulkanQueryPool : public VulkanRHI::FDeviceChild
diff --git a/Engine/Source/Runtime/VulkanRHI/Private/VulkanTexture.cpp b/Engine/Source/Runtime/VulkanRHI/Private/VulkanTexture.cpp
--- a/Engine/Source/Runtime/VulkanRHI/Private/VulkanTexture.cpp
+++ b/Engine/Source/Runtime/VulkanRHI/Private/VulkanTexture.cpp
@@ -1293,12 +1293,6 @@ FTexture3DRHIRef FVulkanDynamicRHI::RHICreateTexture3D(uint32 SizeX, uint32 Size
 	return Tex3d;
 }
 
-void FVulkanDynamicRHI::RHIGetResourceInfo(FRHITexture* Ref, FRHIResourceInfo& OutInfo)
-{
-	FVulkanTextureBase* Base = (FVulkanTextureBase*)Ref->GetTextureBaseRHI();
-	OutInfo.VRamAllocation.AllocationSize = Base->Surface.GetMemorySize();
-}
-
 static void DoAsyncReallocateTexture2D(FVulkanCommandListContext& Context, FVulkanTexture2D* OldTexture, FVulkanTexture2D* NewTexture, int32 NewMipCount, int32 NewSizeX, int32 NewSizeY, FThreadSafeCounter* RequestStatus)
 {
 	LLM_SCOPE_VULKAN(ELLMTagVulkan::VulkanTextures);
@@ -2244,6 +2238,12 @@ void FVulkanTextureBase::DetachView(FVulkanViewBase* View)
 	View->NextView = nullptr;
 }
 
+bool FVulkanTextureBase::GetTextureResourceInfo(FRHIResourceInfo& OutResourceInfo) const
+{
+	OutResourceInfo = FRHIResourceInfo();
+	OutResourceInfo.VRamAllocation.AllocationSize = Surface.GetMemorySize();
+	return true;
+}
 
 FVulkanTexture2D::FVulkanTexture2D(FVulkanDevice& Device, EPixelFormat InFormat, uint32 SizeX, uint32 SizeY, uint32 NumMips, uint32 NumSamples, ETextureCreateFlags UEFlags, ERHIAccess InResourceState, const FRHIResourceCreateInfo& CreateInfo)
 :	FRHITexture2D(SizeX, SizeY, FMath::Max(NumMips, 1u), NumSamples, InFormat, UEFlags, CreateInfo.ClearValueBinding)
diff --git a/Engine/Source/Runtime/VulkanRHI/Public/VulkanDynamicRHI.h b/Engine/Source/Runtime/VulkanRHI/Public/VulkanDynamicRHI.h
--- a/Engine/Source/Runtime/VulkanRHI/Public/VulkanDynamicRHI.h
+++ b/Engine/Source/Runtime/VulkanRHI/Public/VulkanDynamicRHI.h
@@ -103,7 +103,6 @@ public:
 	virtual void RHICopySharedMips(FRHITexture2D* DestTexture2D, FRHITexture2D* SrcTexture2D) final override;
 	virtual FTexture2DArrayRHIRef RHICreateTexture2DArray(uint32 SizeX, uint32 SizeY, uint32 SizeZ, uint8 Format, uint32 NumMips, uint32 NumSamples, ETextureCreateFlags Flags, ERHIAccess InResourceState, FRHIResourceCreateInfo& CreateInfo) final override;
 	virtual FTexture3DRHIRef RHICreateTexture3D(uint32 SizeX, uint32 SizeY, uint32 SizeZ, uint8 Format, uint32 NumMips, ETextureCreateFlags Flags, ERHIAccess InResourceState, FRHIResourceCreateInfo& CreateInfo) final override;
-	virtual void RHIGetResourceInfo(FRHITexture* Ref, FRHIResourceInfo& OutInfo) final override;
 	virtual FShaderResourceViewRHIRef RHICreateShaderResourceView(FRHITexture* Texture, const FRHITextureSRVCreateInfo& CreateInfo) final override;
 	virtual uint32 RHIComputeMemorySize(FRHITexture* TextureRHI) final override;
 	virtual FTexture2DRHIRef RHIAsyncReallocateTexture2D(FRHITexture2D* Texture2D, int32 NewMipCount, int32 NewSizeX, int32 NewSizeY, FThreadSafeCounter* RequestStatus) final override;
diff --git a/Engine/Source/Runtime/Windows/D3D11RHI/Private/D3D11RHIPrivate.h b/Engine/Source/Runtime/Windows/D3D11RHI/Private/D3D11RHIPrivate.h
--- a/Engine/Source/Runtime/Windows/D3D11RHI/Private/D3D11RHIPrivate.h
+++ b/Engine/Source/Runtime/Windows/D3D11RHI/Private/D3D11RHIPrivate.h
@@ -499,7 +499,6 @@ public:
 	virtual void RHICopySharedMips(FRHITexture2D* DestTexture2D, FRHITexture2D* SrcTexture2D) final override;
 	virtual FTexture2DArrayRHIRef RHICreateTexture2DArray(uint32 SizeX, uint32 SizeY, uint32 SizeZ, uint8 Format, uint32 NumMips, uint32 NumSamples, ETextureCreateFlags Flags, ERHIAccess InResourceState, FRHIResourceCreateInfo& CreateInfo) final override;
 	virtual FTexture3DRHIRef RHICreateTexture3D(uint32 SizeX, uint32 SizeY, uint32 SizeZ, uint8 Format, uint32 NumMips, ETextureCreateFlags Flags, ERHIAccess InResourceState, FRHIResourceCreateInfo& CreateInfo) final override;
-	virtual void RHIGetResourceInfo(FRHITexture* Ref, FRHIResourceInfo& OutInfo) final override;
 	virtual FShaderResourceViewRHIRef RHICreateShaderResourceView(FRHITexture* TextureRHI, const FRHITextureSRVCreateInfo& CreateInfo) final override;
 	virtual void RHIGenerateMips(FRHITexture* Texture) final override;
 	virtual uint32 RHIComputeMemorySize(FRHITexture* TextureRHI) final override;
diff --git a/Engine/Source/Runtime/Windows/D3D11RHI/Private/D3D11Texture.cpp b/Engine/Source/Runtime/Windows/D3D11RHI/Private/D3D11Texture.cpp
--- a/Engine/Source/Runtime/Windows/D3D11RHI/Private/D3D11Texture.cpp
+++ b/Engine/Source/Runtime/Windows/D3D11RHI/Private/D3D11Texture.cpp
@@ -1425,10 +1425,6 @@ FTexture3DRHIRef FD3D11DynamicRHI::RHICreateTexture3D_RenderThread(
 	return RHICreateTexture3D(SizeX, SizeY, SizeZ, Format, NumMips, Flags, InResourceState, CreateInfo);
 }
 
-void FD3D11DynamicRHI::RHIGetResourceInfo(FRHITexture* Ref, FRHIResourceInfo& OutInfo)
-{
-}
-
 FShaderResourceViewRHIRef FD3D11DynamicRHI::RHICreateShaderResourceView(FRHITexture* TextureRHI, const FRHITextureSRVCreateInfo& CreateInfo)
 {
 	FD3D11TextureBase* Texture = GetD3D11TextureFromRHITexture(TextureRHI);
