Add FRHIRayTracingScene::GetMetadataBufferSRV() method to access RHI-specific data associated with the ray tracing scene. This can be used to access vertex and index buffer binding information for inline ray tracing shaders (implemented for Vulkan).

diff --git a/Engine/Source/Runtime/RHI/Public/RHIResources.h b/Engine/Source/Runtime/RHI/Public/RHIResources.h
--- a/Engine/Source/Runtime/RHI/Public/RHIResources.h
+++ b/Engine/Source/Runtime/RHI/Public/RHIResources.h
@@ -1949,6 +1949,14 @@ class FRHIRayTracingScene : public FRHIRayTracingAccelerationStructure
 {
 public:
 	virtual const FRayTracingSceneInitializer2& GetInitializer() const = 0;
+
+	// Returns a buffer view for RHI-specific system parameters associated with this scene.
+	// This may be needed to access ray tracing geometry data in shaders that use ray queries.
+	// Returns NULL if current RHI does not require this buffer.
+	virtual FRHIShaderResourceView* GetMetadataBufferSRV() const
+	{
+		return nullptr;
+	}
 };
 
 typedef TRefCountPtr<FRHIRayTracingScene>        FRayTracingSceneRHIRef;
diff --git a/Engine/Source/Runtime/VulkanRHI/Private/VulkanRayTracing.cpp b/Engine/Source/Runtime/VulkanRHI/Private/VulkanRayTracing.cpp
--- a/Engine/Source/Runtime/VulkanRHI/Private/VulkanRayTracing.cpp
+++ b/Engine/Source/Runtime/VulkanRHI/Private/VulkanRayTracing.cpp
@@ -106,7 +106,7 @@ static VkDeviceAddress GetDeviceAddress(VkDevice Device, VkBuffer Buffer)
 	return vkGetBufferDeviceAddressKHR(Device, &DeviceAddressInfo);
 }
 
-VkDeviceAddress FVulkanResourceMultiBuffer::GetDeviceAddress()
+VkDeviceAddress FVulkanResourceMultiBuffer::GetDeviceAddress() const
 {
 	return ::GetDeviceAddress(Device->GetInstanceHandle(), GetHandle()) + GetOffset();
 }
@@ -227,6 +227,7 @@ static void GetBLASBuildData(
 		switch (Segment.VertexBufferElementType)
 		{
 		case VET_Float3:
+		case VET_Float4:
 			SegmentGeometry.geometry.triangles.vertexFormat = VK_FORMAT_R32G32B32_SFLOAT;
 			break;
 		default:
@@ -451,11 +452,39 @@ static VkGeometryInstanceFlagsKHR TranslateRayTracingInstanceFlags(ERayTracingIn
 	return Result;
 }
 
+// This structure is analogous to FHitGroupSystemParameters in D3D12 RHI.
+// However, it only contains generic parameters that do not require a full shader binding table (i.e. no per-hit-group user data).
+// It is designed to be used to access vertex and index buffers during inline ray tracing.
+struct FVulkanRayTracingGeometryParameters
+{
+	union
+	{
+		struct
+		{
+			uint32 IndexStride : 8; // Can be just 1 bit to indicate 16 or 32 bit indices
+			uint32 VertexStride : 8; // Can be just 2 bits to indicate float3, float2 or half2 format
+			uint32 Unused : 16;
+		} Config;
+		uint32 ConfigBits = 0;
+	};
+	uint32 IndexBufferOffsetInBytes = 0;
+	uint64 IndexBuffer = 0;
+	uint64 VertexBuffer = 0;
+};
+
 FVulkanRayTracingScene::FVulkanRayTracingScene(FRayTracingSceneInitializer2 InInitializer, FVulkanDevice* InDevice, FVulkanResourceMultiBuffer* InInstanceBuffer)
 	: Device(InDevice), Initializer(MoveTemp(InInitializer)), InstanceBuffer(InInstanceBuffer)
 {
 	const ERayTracingAccelerationStructureFlags BuildFlags = ERayTracingAccelerationStructureFlags::FastTrace; // #yuriy_todo: pass this in
 	SizeInfo = RHICalcRayTracingSceneSize(Initializer.NumNativeInstances, BuildFlags);
+
+	const uint32 ParameterBufferSize = FMath::Max<uint32>(1, Initializer.NumTotalSegments) * sizeof(FVulkanRayTracingGeometryParameters);
+	FRHIResourceCreateInfo ParameterBufferCreateInfo(TEXT("RayTracingSceneMetadata"));
+	PerInstanceGeometryParameterBuffer = ResourceCast(RHICreateBuffer(
+		ParameterBufferSize, BUF_StructuredBuffer | BUF_ShaderResource, sizeof(FVulkanRayTracingGeometryParameters),
+		ERHIAccess::SRVCompute, ParameterBufferCreateInfo).GetReference());
+
+	PerInstanceGeometryParameterSRV = new FVulkanShaderResourceView(Device, PerInstanceGeometryParameterBuffer, 0);
 }
 
 FVulkanRayTracingScene::~FVulkanRayTracingScene()
@@ -496,6 +525,10 @@ void FVulkanRayTracingScene::BuildAccelerationStructure(
 		InstanceBufferAddress = InstanceBuffer->GetDeviceAddress();
 	}
 
+	// Build a metadata buffer	that contains VulkanRHI-specific per-geometry parameters that allow us to access
+	// vertex and index buffers from shaders that use inline ray tracing.
+	BuildPerInstanceGeometryParameterBuffer();
+
 	FVkRtTLASBuildData BuildData;
 	GetTLASBuildData(Device->GetInstanceHandle(), Initializer.NumNativeInstances, InstanceBufferAddress, BuildData);
 
@@ -535,6 +568,64 @@ void FVulkanRayTracingScene::BuildAccelerationStructure(
 	InstanceBuffer = nullptr;
 }
 
+void FVulkanRayTracingScene::BuildPerInstanceGeometryParameterBuffer()
+{
+	// TODO: we could cache parameters in the geometry object to avoid some of the pointer chasing (if this is measured to be a performance issue)
+
+	const uint32 ParameterBufferSize = FMath::Max<uint32>(1, Initializer.NumTotalSegments) * sizeof(FVulkanRayTracingGeometryParameters);
+	check(PerInstanceGeometryParameterBuffer->GetSize() >= ParameterBufferSize);
+
+	check(IsInRHIThread() || !IsRunningRHIInSeparateThread());
+	const bool bTopOfPipe = false; // running on RHI timeline
+
+	void* MappedBuffer = PerInstanceGeometryParameterBuffer->Lock(bTopOfPipe, RLM_WriteOnly, ParameterBufferSize, 0);
+	FVulkanRayTracingGeometryParameters* MappedParameters = reinterpret_cast<FVulkanRayTracingGeometryParameters*>(MappedBuffer);
+	uint32 ParameterIndex = 0;
+
+	for (FRHIRayTracingGeometry* GeometryRHI : Initializer.PerInstanceGeometries)
+	{
+		const FVulkanRayTracingGeometry* Geometry = ResourceCast(GeometryRHI);
+		const FRayTracingGeometryInitializer& GeometryInitializer = Geometry->GetInitializer();
+
+		const FVulkanResourceMultiBuffer* IndexBuffer = ResourceCast(GeometryInitializer.IndexBuffer.GetReference());
+
+		const uint32 IndexStride = IndexBuffer ? IndexBuffer->GetStride() : 0;
+		const uint32 IndexOffsetInBytes = GeometryInitializer.IndexBufferOffset;
+		const VkDeviceAddress IndexBufferAddress = IndexBuffer ? IndexBuffer->GetDeviceAddress() : VkDeviceAddress(0);
+
+		for (const FRayTracingGeometrySegment& Segment : GeometryInitializer.Segments)
+		{
+			const FVulkanResourceMultiBuffer* VertexBuffer = ResourceCast(Segment.VertexBuffer.GetReference());
+			checkf(VertexBuffer, TEXT("All ray tracing geometry segments must have a valid vertex buffer"));
+			const VkDeviceAddress VertexBufferAddress = VertexBuffer->GetDeviceAddress();
+
+			FVulkanRayTracingGeometryParameters SegmentParameters;
+			SegmentParameters.Config.IndexStride = IndexStride;
+			SegmentParameters.Config.VertexStride = Segment.VertexBufferStride;
+
+			if (IndexStride)
+			{
+				SegmentParameters.IndexBufferOffsetInBytes = IndexOffsetInBytes + IndexStride * Segment.FirstPrimitive * 3;
+				SegmentParameters.IndexBuffer = static_cast<uint64>(IndexBufferAddress);
+			}
+			else
+			{
+				SegmentParameters.IndexBuffer = 0;
+			}
+
+			SegmentParameters.VertexBuffer = static_cast<uint64>(VertexBufferAddress) + Segment.VertexBufferOffset;
+
+			check(ParameterIndex < Initializer.NumTotalSegments);
+			MappedParameters[ParameterIndex] = SegmentParameters;
+			ParameterIndex++;
+		}
+	}
+
+	check(ParameterIndex == Initializer.NumTotalSegments);
+
+	PerInstanceGeometryParameterBuffer->Unlock(bTopOfPipe);
+}
+
 void FVulkanDynamicRHI::RHITransferRayTracingGeometryUnderlyingResource(FRHIRayTracingGeometry* DestGeometry, FRHIRayTracingGeometry* SrcGeometry)
 {
 	check(DestGeometry);
diff --git a/Engine/Source/Runtime/VulkanRHI/Private/VulkanRayTracing.h b/Engine/Source/Runtime/VulkanRHI/Private/VulkanRayTracing.h
--- a/Engine/Source/Runtime/VulkanRHI/Private/VulkanRayTracing.h
+++ b/Engine/Source/Runtime/VulkanRHI/Private/VulkanRayTracing.h
@@ -91,6 +91,8 @@ public:
 	void Swap(FVulkanRayTracingGeometry& Other);
 	void BuildAccelerationStructure(FVulkanCommandListContext& CommandContext, EAccelerationStructureBuildMode BuildMode);
 
+	const FRayTracingGeometryInitializer& GetInitializer() const { return Initializer; }
+
 private:
 	FVulkanDevice* const Device = nullptr;
 
@@ -120,6 +122,11 @@ public:
 
 	FRayTracingAccelerationStructureSize SizeInfo;
 
+	virtual FRHIShaderResourceView* GetMetadataBufferSRV() const override final
+	{
+		return PerInstanceGeometryParameterSRV.GetReference();
+	}
+
 private:
 	FVulkanDevice* const Device = nullptr;
 
@@ -137,6 +144,11 @@ private:
 	TRefCountPtr<FVulkanShaderResourceView> AccelerationStructureView;
 	
 	TRefCountPtr<FVulkanResourceMultiBuffer> AccelerationStructureBuffer;
+
+	// Buffer that contains per-instance index and vertex buffer binding data
+	TRefCountPtr<FVulkanResourceMultiBuffer> PerInstanceGeometryParameterBuffer;
+	TRefCountPtr<FVulkanShaderResourceView> PerInstanceGeometryParameterSRV;
+	void BuildPerInstanceGeometryParameterBuffer();
 };
 
 class FVulkanRayTracingPipelineState : public FRHIRayTracingPipelineState
diff --git a/Engine/Source/Runtime/VulkanRHI/Private/VulkanResources.h b/Engine/Source/Runtime/VulkanRHI/Private/VulkanResources.h
--- a/Engine/Source/Runtime/VulkanRHI/Private/VulkanResources.h
+++ b/Engine/Source/Runtime/VulkanRHI/Private/VulkanResources.h
@@ -1193,7 +1193,7 @@ public:
 	}
 
 #if VULKAN_RHI_RAYTRACING
-	VkDeviceAddress GetDeviceAddress();
+	VkDeviceAddress GetDeviceAddress() const;
 #endif
 
 protected:
