Fix PF_R64_UINT support with Renderdoc in Vulkan so that we can still take captures of Nanite scenes (some format capabilities don't make it through the layer).
Add a few of the capabilities in the debug command line for ease of debugging.

diff --git a/Engine/Source/Runtime/VulkanRHI/Private/VulkanDevice.cpp b/Engine/Source/Runtime/VulkanRHI/Private/VulkanDevice.cpp
--- a/Engine/Source/Runtime/VulkanRHI/Private/VulkanDevice.cpp
+++ b/Engine/Source/Runtime/VulkanRHI/Private/VulkanDevice.cpp
@@ -641,8 +641,19 @@ void FVulkanDevice::SetupFormats()
 	MapFormatSupport(PF_R8_SINT, { VK_FORMAT_R8_SINT }, ComponentMappingR001);
 
 	// This will be the format used for 64bit image atomics
-	MapFormatSupport(PF_R64_UINT, { VK_FORMAT_R64_UINT }, ComponentMappingR001, EPixelFormatCapabilities::TextureAtomics);
-	checkf(!GRHISupportsAtomicUInt64 || EnumHasAnyFlags(GPixelFormats[PF_R64_UINT].Capabilities, EPixelFormatCapabilities::TextureSample | EPixelFormatCapabilities::UAV), TEXT("Image atomics were enabled, but the R64 format does not have TextureSample or UAV capabilities!"));
+#if VULKAN_HAS_DEBUGGING_ENABLED
+	const EPixelFormatCapabilities RequiredCaps64U = GRenderDocFound ? EPixelFormatCapabilities::UAV : (EPixelFormatCapabilities::UAV | EPixelFormatCapabilities::TextureAtomics);
+#else
+	const EPixelFormatCapabilities RequiredCaps64U = (EPixelFormatCapabilities::UAV | EPixelFormatCapabilities::TextureAtomics);
+#endif
+	MapFormatSupport(PF_R64_UINT, { VK_FORMAT_R64_UINT, VK_FORMAT_R32G32_UINT }, ComponentMappingR001, RequiredCaps64U);
+	// Shaders were patched to use UAV, make sure we don't expose texture sampling
+	GPixelFormats[PF_R64_UINT].Capabilities &= ~(EPixelFormatCapabilities::AnyTexture | EPixelFormatCapabilities::TextureSample);
+	if (GRHISupportsAtomicUInt64 && !EnumHasAnyFlags(GPixelFormats[PF_R64_UINT].Capabilities, EPixelFormatCapabilities::UAV))
+	{
+		UE_LOG(LogVulkanRHI, Warning, TEXT("64bit image atomics were enabled, but the R64 format does not have UAV capabilities.  Disabling support."));
+		GRHISupportsAtomicUInt64 = false;
+	}
 
 	if (CVarVulkanUseD24.GetValueOnAnyThread() != 0)
 	{
@@ -711,26 +722,53 @@ void FVulkanDevice::SetupFormats()
 	}
 
 #if !(UE_BUILD_SHIPPING || UE_BUILD_TEST)
+
 	// Print the resulting pixel format support
 	if (FParse::Param(FCommandLine::Get(), TEXT("PrintVulkanPixelFormatMappings")))
 	{
+		auto GetFormatCapabilities = [](EPixelFormatCapabilities FormatCapabilities)
+		{
+#define VULKAN_CHECK_FORMAT_CAPABILITY(PF_Name) if (EnumHasAllFlags(FormatCapabilities, EPixelFormatCapabilities::PF_Name)) { CapabilitiesString += TEXT(#PF_Name) TEXT(", ");}
+			FString CapabilitiesString;
+
+			VULKAN_CHECK_FORMAT_CAPABILITY(TextureSample);
+			VULKAN_CHECK_FORMAT_CAPABILITY(TextureCube);
+			VULKAN_CHECK_FORMAT_CAPABILITY(RenderTarget);
+			VULKAN_CHECK_FORMAT_CAPABILITY(DepthStencil);
+			VULKAN_CHECK_FORMAT_CAPABILITY(TextureBlendable);
+			VULKAN_CHECK_FORMAT_CAPABILITY(TextureAtomics);
+
+			VULKAN_CHECK_FORMAT_CAPABILITY(Buffer);
+			VULKAN_CHECK_FORMAT_CAPABILITY(VertexBuffer);
+			VULKAN_CHECK_FORMAT_CAPABILITY(IndexBuffer);
+			VULKAN_CHECK_FORMAT_CAPABILITY(BufferAtomics);
+
+			VULKAN_CHECK_FORMAT_CAPABILITY(UAV);
+
+			return CapabilitiesString;
+#undef VULKAN_CHECK_FORMAT_CAPABILITY
+		};
+
 		UE_LOG(LogVulkanRHI, Warning, TEXT("Pixel Format Mappings for Vulkan:"));
-		UE_LOG(LogVulkanRHI, Warning, TEXT("%24s | Capabilities | VulkanFormat | BlockBytes | Components | ComponentMapping | BufferFormat"), TEXT("PixelFormatName"));
+		UE_LOG(LogVulkanRHI, Warning, TEXT("%24s | VulkanFormat | BlockBytes | Components | ComponentMapping | BufferFormat | Capabilities "), TEXT("PixelFormatName"));
 		for (int32 PixelFormatIndex = 0; PixelFormatIndex < PF_MAX; ++PixelFormatIndex)
 		{
 			if (GPixelFormats[PixelFormatIndex].Supported)
 			{
 				const VkComponentMapping& ComponentMapping = PixelFormatComponentMapping[PixelFormatIndex];
 
-				UE_LOG(LogVulkanRHI, Warning, TEXT("%24s |   0x%08X | %12d | %10d | %10d | %10d,%d,%d,%d | %12d"),
+				FString CapabilitiesString = GetFormatCapabilities(GPixelFormats[PixelFormatIndex].Capabilities);
+
+				UE_LOG(LogVulkanRHI, Warning, TEXT("%24s | %12d | %10d | %10d | %10d,%d,%d,%d | %12d | 0x%08X (%s)"),
 					GPixelFormats[PixelFormatIndex].Name,
-					(uint32)GPixelFormats[PixelFormatIndex].Capabilities,
 					GPixelFormats[PixelFormatIndex].PlatformFormat,
 					GPixelFormats[PixelFormatIndex].BlockBytes,
 					GPixelFormats[PixelFormatIndex].NumComponents,
 					ComponentMapping.r, ComponentMapping.g, ComponentMapping.b, ComponentMapping.a,
-					(int32)GVulkanBufferFormat[PixelFormatIndex]
-				);
+					(int32)GVulkanBufferFormat[PixelFormatIndex],
+					(uint32)GPixelFormats[PixelFormatIndex].Capabilities,
+					*CapabilitiesString
+					);
 			}
 		}
 	}
