Add Vulkan support for the Transient Resource Allocator.  (Still disabled by default)

diff --git a/Engine/Source/Runtime/VulkanRHI/Private/VulkanBarriers.cpp b/Engine/Source/Runtime/VulkanRHI/Private/VulkanBarriers.cpp
--- a/Engine/Source/Runtime/VulkanRHI/Private/VulkanBarriers.cpp
+++ b/Engine/Source/Runtime/VulkanRHI/Private/VulkanBarriers.cpp
@@ -184,6 +184,13 @@ static void GetVkStageAndAccessFlags(ERHIAccess RHIAccess, FRHITransitionInfo::E
 	// States which cannot be combined.
 	switch (RHIAccess)
 	{
+		case ERHIAccess::Discard:
+			// FIXME: Align with Unknown for now, this could perhaps be less brutal
+			StageFlags = VK_PIPELINE_STAGE_ALL_COMMANDS_BIT;
+			AccessFlags = VK_ACCESS_MEMORY_READ_BIT | VK_ACCESS_MEMORY_WRITE_BIT;
+			Layout = bIsSourceState ? VK_IMAGE_LAYOUT_UNDEFINED : SRVLayout;
+			return;
+
 		case ERHIAccess::Unknown:
 			// We don't know where this is coming from, so we'll stall everything.
 			StageFlags = VK_PIPELINE_STAGE_ALL_COMMANDS_BIT;
@@ -394,6 +401,7 @@ struct FDepthStencilSubresTransition
 	ERHIAccess DestStencilAccess;
 	int bDepthAccessSet : 1;
 	int bStencilAccessSet : 1;
+	int bAliasingBarrier : 1;
 };
 
 static void GetDepthStencilStageAndAccessFlags(ERHIAccess DepthAccess, ERHIAccess StencilAccess, VkPipelineStageFlags& StageFlags, VkAccessFlags& AccessFlags, VkImageLayout& Layout, bool bIsSourceState)
@@ -406,6 +414,14 @@ static void GetDepthStencilStageAndAccessFlags(ERHIAccess DepthAccess, ERHIAcces
 		return;
 	}
 
+	if (DepthAccess == ERHIAccess::Discard && StencilAccess == ERHIAccess::Discard)
+	{
+		StageFlags = VK_PIPELINE_STAGE_ALL_COMMANDS_BIT;
+		AccessFlags = VK_ACCESS_MEMORY_READ_BIT | VK_ACCESS_MEMORY_WRITE_BIT;
+		Layout = bIsSourceState ? VK_IMAGE_LAYOUT_UNDEFINED : VK_IMAGE_LAYOUT_DEPTH_STENCIL_READ_ONLY_OPTIMAL;
+		return;
+	}
+
 	Layout = VK_IMAGE_LAYOUT_DEPTH_STENCIL_READ_ONLY_OPTIMAL;
 
 	if (EnumHasAnyFlags(DepthAccess, ERHIAccess::CopySrc))
@@ -600,6 +616,12 @@ void FVulkanDynamicRHI::RHICreateTransition(FRHITransition* Transition, const FR
 			continue;
 		}
 
+		if (Info.AccessAfter == ERHIAccess::Discard)
+		{
+			// Discard as a destination is a no-op
+			continue;
+		}
+
 		if (Info.Type == FRHITransitionInfo::EType::Texture)
 		{
 			// CPU accessible "textures" are implemented as buffers. Check if this is a real texture or a buffer.
@@ -625,7 +647,7 @@ void FVulkanDynamicRHI::RHICreateTransition(FRHITransition* Transition, const FR
 	}
 
 	Data->ImageBarriers.Reserve(NumTextures);
-	Data->Textures.Reserve(NumTextures);
+	Data->ImageBarrierExtras.Reserve(NumTextures);
 	if (SrcPipelines != DstPipelines)
 	{
 		Data->BufferBarriers.Reserve(NumBuffers);
@@ -642,6 +664,12 @@ void FVulkanDynamicRHI::RHICreateTransition(FRHITransition* Transition, const FR
 			continue;
 		}
 
+		if (Info.AccessAfter == ERHIAccess::Discard)
+		{
+			// Discard as a destination is a no-op
+			continue;
+		}
+
 		checkf(Info.AccessAfter != ERHIAccess::Unknown, TEXT("Transitioning a resource to an unknown state is not allowed."));
 
 		FVulkanResourceMultiBuffer* Buffer = nullptr;
@@ -741,6 +769,8 @@ void FVulkanDynamicRHI::RHICreateTransition(FRHITransition* Transition, const FR
 				PendingTransition->bStencilAccessSet = true;
 			}
 
+			PendingTransition->bAliasingBarrier = (Info.AccessBefore == ERHIAccess::Discard);
+
 			if (Index == INDEX_NONE)
 			{
 				// Wait until we find the other aspect of the resource.
@@ -824,7 +854,9 @@ void FVulkanDynamicRHI::RHICreateTransition(FRHITransition* Transition, const FR
 		ImgBarrier.srcQueueFamilyIndex = SrcQueueFamilyIndex;
 		ImgBarrier.dstQueueFamilyIndex = DstQueueFamilyIndex;
 
-		Data->Textures.Add(Texture);
+		FVulkanPipelineBarrier::ImageBarrierExtraData& ExtraData = Data->ImageBarrierExtras.AddDefaulted_GetRef();
+		ExtraData.BaseTexture = Texture;
+		ExtraData.IsAliasingBarrier = (Info.AccessBefore == ERHIAccess::Discard);
 	}
 
 	// Process any depth-stencil transitions which specify a single sub-resource.
@@ -874,7 +906,9 @@ void FVulkanDynamicRHI::RHICreateTransition(FRHITransition* Transition, const FR
 		ImgBarrier.srcQueueFamilyIndex = SrcQueueFamilyIndex;
 		ImgBarrier.dstQueueFamilyIndex = DstQueueFamilyIndex;
 
-		Data->Textures.Add(PendingTransition.Texture);
+		FVulkanPipelineBarrier::ImageBarrierExtraData& ExtraData = Data->ImageBarrierExtras.AddDefaulted_GetRef();
+		ExtraData.BaseTexture = PendingTransition.Texture;
+		ExtraData.IsAliasingBarrier = PendingTransition.bAliasingBarrier;
 	}
 }
 
@@ -1003,7 +1037,7 @@ void FVulkanCommandListContext::RHIBeginTransitions(TArrayView<const FRHITransit
 			RealBarrier.dstAccessMask = 0; // Release resource from current queue.
 		}
 
-		check(Data->ImageBarriers.Num() == Data->Textures.Num());
+		check(Data->ImageBarriers.Num() == Data->ImageBarrierExtras.Num());
 		RealImageBarriers.Reset();
 		RealImageBarriers.Reserve(Data->ImageBarriers.Num());
 
@@ -1013,7 +1047,8 @@ void FVulkanCommandListContext::RHIBeginTransitions(TArrayView<const FRHITransit
 		for (int32 ImgBarrierIdx = 0; ImgBarrierIdx < Data->ImageBarriers.Num(); ++ImgBarrierIdx)
 		{
 			const VkImageMemoryBarrier& ImageBarrier = Data->ImageBarriers[ImgBarrierIdx];
-			FVulkanTextureBase* Texture = Data->Textures[ImgBarrierIdx];
+			const FVulkanPipelineBarrier::ImageBarrierExtraData& ExtraData = Data->ImageBarrierExtras[ImgBarrierIdx];
+			FVulkanTextureBase* Texture = ExtraData.BaseTexture;
 			check(Texture->Surface.Image != VK_NULL_HANDLE);
 
 			FVulkanImageLayout& Layout = LayoutManager.GetOrAddFullLayout(Texture->Surface, VK_IMAGE_LAYOUT_UNDEFINED);
@@ -1024,7 +1059,7 @@ void FVulkanCommandListContext::RHIBeginTransitions(TArrayView<const FRHITransit
 			check(ImageBarrier.newLayout != VK_IMAGE_LAYOUT_UNDEFINED);
 			DstLayout = ImageBarrier.newLayout;
 
-			if (ImageBarrier.oldLayout == VK_IMAGE_LAYOUT_UNDEFINED)
+			if ((ImageBarrier.oldLayout == VK_IMAGE_LAYOUT_UNDEFINED) && (!ExtraData.IsAliasingBarrier))
 			{
 				check(Layout.AreAllSubresourcesSameLayout());
 				SrcLayout = Layout.MainLayout;
@@ -1201,7 +1236,7 @@ void FVulkanCommandListContext::RHIEndTransitions(TArrayView<const FRHITransitio
 			RealBarrier.srcAccessMask = 0; // Acquire resource on current queue.
 		}
 
-		check(Data->ImageBarriers.Num() == Data->Textures.Num());
+		check(Data->ImageBarriers.Num() == Data->ImageBarrierExtras.Num());
 		RealImageBarriers.Reset();
 		RealImageBarriers.Reserve(Data->ImageBarriers.Num());
 
@@ -1211,7 +1246,8 @@ void FVulkanCommandListContext::RHIEndTransitions(TArrayView<const FRHITransitio
 		for (int32 ImgBarrierIdx = 0; ImgBarrierIdx < Data->ImageBarriers.Num(); ++ImgBarrierIdx)
 		{
 			const VkImageMemoryBarrier& ImageBarrier = Data->ImageBarriers[ImgBarrierIdx];
-			FVulkanTextureBase* Texture = Data->Textures[ImgBarrierIdx];
+			const FVulkanPipelineBarrier::ImageBarrierExtraData& ExtraData = Data->ImageBarrierExtras[ImgBarrierIdx];
+			FVulkanTextureBase* Texture = ExtraData.BaseTexture;
 			check(Texture->Surface.GetCpuReadbackBuffer() == nullptr);
 
 			Texture->OnLayoutTransition(*this, ImageBarrier.newLayout);
@@ -1227,7 +1263,7 @@ void FVulkanCommandListContext::RHIEndTransitions(TArrayView<const FRHITransitio
 			check(ImageBarrier.newLayout != VK_IMAGE_LAYOUT_UNDEFINED);
 			DstLayout = ImageBarrier.newLayout;
 
-			if (ImageBarrier.oldLayout == VK_IMAGE_LAYOUT_UNDEFINED)
+			if ((ImageBarrier.oldLayout == VK_IMAGE_LAYOUT_UNDEFINED) && (!ExtraData.IsAliasingBarrier))
 			{
 				if (Layout.AreAllSubresourcesSameLayout())
 				{
diff --git a/Engine/Source/Runtime/VulkanRHI/Private/VulkanBarriers.h b/Engine/Source/Runtime/VulkanRHI/Private/VulkanBarriers.h
--- a/Engine/Source/Runtime/VulkanRHI/Private/VulkanBarriers.h
+++ b/Engine/Source/Runtime/VulkanRHI/Private/VulkanBarriers.h
@@ -23,7 +23,12 @@ struct FVulkanPipelineBarrier
 	bool bHasMemoryBarrier;
 
 	// We need to keep the texture pointers around, because we need to call OnTransitionResource on them, and we need mip and layer counts for the tracking code.
-	TArray<FVulkanTextureBase*, TInlineAllocator<2>> Textures;
+	struct ImageBarrierExtraData
+	{
+		FVulkanTextureBase* BaseTexture = nullptr;
+		bool IsAliasingBarrier = false;
+	};
+	TArray<ImageBarrierExtraData, TInlineAllocator<2>> ImageBarrierExtras;
 
 	void AddMemoryBarrier(VkAccessFlags SrcAccessFlags, VkAccessFlags DstAccessFlags, VkPipelineStageFlags SrcStageMask, VkPipelineStageFlags DstStageMask);
 	void AddImageLayoutTransition(VkImage Image, VkImageLayout SrcLayout, VkImageLayout DstLayout, const VkImageSubresourceRange& SubresourceRange);
diff --git a/Engine/Source/Runtime/VulkanRHI/Private/VulkanBuffer.cpp b/Engine/Source/Runtime/VulkanRHI/Private/VulkanBuffer.cpp
--- a/Engine/Source/Runtime/VulkanRHI/Private/VulkanBuffer.cpp
+++ b/Engine/Source/Runtime/VulkanRHI/Private/VulkanBuffer.cpp
@@ -10,6 +10,7 @@
 #include "Containers/ResourceArray.h"
 #include "VulkanLLM.h"
 #include "VulkanRayTracing.h"
+#include "VulkanTransientResourceAllocator.h"
 
 static TMap<FVulkanResourceMultiBuffer*, VulkanRHI::FPendingBufferLock> GPendingLockIBs;
 static FCriticalSection GPendingLockIBsMutex;
@@ -69,67 +70,64 @@ static FORCEINLINE void UpdateVulkanBufferStats(uint64_t Size, VkBufferUsageFlag
 	}
 }
 
-FVulkanResourceMultiBuffer::FVulkanResourceMultiBuffer(FVulkanDevice* InDevice, uint32 InSize, EBufferUsageFlags InUEUsage, uint32 InStride, FRHIResourceCreateInfo& CreateInfo, class FRHICommandListImmediate* InRHICmdList)
-	: FRHIBuffer(InSize, InUEUsage, InStride)
-	, VulkanRHI::FDeviceChild(InDevice)
-	, BufferUsageFlags(0)
-	, NumBuffers(0)
-	, DynamicBufferIndex(0)
-	, LockStatus(ELockStatus::Unlocked)
-{
-	VULKAN_TRACK_OBJECT_CREATE(FVulkanResourceMultiBuffer, this);
 
+VkBufferUsageFlags FVulkanResourceMultiBuffer::UEToVKBufferUsageFlags(FVulkanDevice* InDevice, EBufferUsageFlags InUEUsage, bool bZeroSize)
+{
 	// Always include TRANSFER_SRC since hardware vendors confirmed it wouldn't have any performance cost and we need it for some debug functionalities.
-	BufferUsageFlags |= VK_BUFFER_USAGE_TRANSFER_SRC_BIT;
+	VkBufferUsageFlags OutVkUsage = VK_BUFFER_USAGE_TRANSFER_SRC_BIT;
 
-	if (EnumHasAnyFlags(InUEUsage, BUF_VertexBuffer))
-	{
-		BufferUsageFlags |= VK_BUFFER_USAGE_VERTEX_BUFFER_BIT;
-	}
-	if (EnumHasAnyFlags(InUEUsage, BUF_IndexBuffer))
-	{
-		BufferUsageFlags |= VK_BUFFER_USAGE_INDEX_BUFFER_BIT;
-	}
-	if (EnumHasAnyFlags(InUEUsage, BUF_StructuredBuffer))
+	auto TranslateFlag = [&OutVkUsage, &InUEUsage](EBufferUsageFlags SearchUEFlag, VkBufferUsageFlags AddedIfFound, VkBufferUsageFlags AddedIfNotFound = 0)
 	{
-		BufferUsageFlags |= VK_BUFFER_USAGE_STORAGE_BUFFER_BIT;
-	}
+		const bool HasFlag = EnumHasAnyFlags(InUEUsage, SearchUEFlag);
+		OutVkUsage |= HasFlag ? AddedIfFound : AddedIfNotFound;
+	};
+
+	TranslateFlag(BUF_VertexBuffer, VK_BUFFER_USAGE_VERTEX_BUFFER_BIT);
+	TranslateFlag(BUF_IndexBuffer, VK_BUFFER_USAGE_INDEX_BUFFER_BIT);
+	TranslateFlag(BUF_StructuredBuffer, VK_BUFFER_USAGE_STORAGE_BUFFER_BIT);
+
 #if VULKAN_RHI_RAYTRACING
-	if (EnumHasAnyFlags(InUEUsage, BUF_AccelerationStructure))
-	{
-		BufferUsageFlags |= VK_BUFFER_USAGE_ACCELERATION_STRUCTURE_STORAGE_BIT_KHR;
-	}
+	TranslateFlag(BUF_AccelerationStructure, VK_BUFFER_USAGE_ACCELERATION_STRUCTURE_STORAGE_BIT_KHR);
 #endif
 
-	if (InSize > 0)
+	if (!bZeroSize)
 	{
-		const bool bStatic = EnumHasAnyFlags(InUEUsage, BUF_Static);
-		const bool bDynamic = EnumHasAnyFlags(InUEUsage, BUF_Dynamic);
-		const bool bVolatile = EnumHasAnyFlags(InUEUsage, BUF_Volatile);
-		const bool bShaderResource = EnumHasAnyFlags(InUEUsage, BUF_ShaderResource);
-		const bool bIsUniformBuffer = (BufferUsageFlags & VK_BUFFER_USAGE_UNIFORM_BUFFER_BIT) != 0;
-		const bool bUAV = EnumHasAnyFlags(InUEUsage, BUF_UnorderedAccess);
-		const bool bIndirect = EnumHasAllFlags(InUEUsage, BUF_DrawIndirect);
-		const bool bCPUReadable = EnumHasAnyFlags(InUEUsage, BUF_KeepCPUAccessible);
-
-		BufferUsageFlags |= bVolatile ? 0 : VK_BUFFER_USAGE_TRANSFER_DST_BIT;
-		BufferUsageFlags |= (bShaderResource && !bIsUniformBuffer) ? VK_BUFFER_USAGE_UNIFORM_TEXEL_BUFFER_BIT : 0;
-		BufferUsageFlags |= bUAV ? VK_BUFFER_USAGE_STORAGE_TEXEL_BUFFER_BIT : 0;
-		BufferUsageFlags |= bIndirect ? VK_BUFFER_USAGE_INDIRECT_BUFFER_BIT : 0;
-		BufferUsageFlags |= bCPUReadable ? (VK_BUFFER_USAGE_TRANSFER_SRC_BIT | VK_BUFFER_USAGE_TRANSFER_DST_BIT) : 0;
+		TranslateFlag(BUF_UnorderedAccess, VK_BUFFER_USAGE_STORAGE_TEXEL_BUFFER_BIT);
+		TranslateFlag(BUF_DrawIndirect, VK_BUFFER_USAGE_INDIRECT_BUFFER_BIT);
+		TranslateFlag(BUF_KeepCPUAccessible, (VK_BUFFER_USAGE_TRANSFER_SRC_BIT | VK_BUFFER_USAGE_TRANSFER_DST_BIT));
+		TranslateFlag(BUF_ShaderResource, VK_BUFFER_USAGE_UNIFORM_TEXEL_BUFFER_BIT);
+
+		TranslateFlag(BUF_Volatile, 0, VK_BUFFER_USAGE_TRANSFER_DST_BIT);
 
 #if VULKAN_RHI_RAYTRACING
 		if (InDevice->GetOptionalExtensions().HasRaytracingExtensions())
 		{
-			BufferUsageFlags |= VK_BUFFER_USAGE_SHADER_DEVICE_ADDRESS_BIT;
+			OutVkUsage |= VK_BUFFER_USAGE_SHADER_DEVICE_ADDRESS_BIT;
 
-			if (!EnumHasAnyFlags(InUEUsage, BUF_AccelerationStructure))
-			{
-				BufferUsageFlags |= VK_BUFFER_USAGE_ACCELERATION_STRUCTURE_BUILD_INPUT_READ_ONLY_BIT_KHR;
-			}
+			TranslateFlag(BUF_AccelerationStructure, 0, VK_BUFFER_USAGE_ACCELERATION_STRUCTURE_BUILD_INPUT_READ_ONLY_BIT_KHR);
 		}
 #endif
+	}
+
+	return OutVkUsage;
+}
+
+FVulkanResourceMultiBuffer::FVulkanResourceMultiBuffer(FVulkanDevice* InDevice, uint32 InSize, EBufferUsageFlags InUEUsage, uint32 InStride, FRHIResourceCreateInfo& CreateInfo, class FRHICommandListImmediate* InRHICmdList, const FRHITransientHeapAllocation* InTransientHeapAllocation)
+	: FRHIBuffer(InSize, InUEUsage, InStride)
+	, VulkanRHI::FDeviceChild(InDevice)
+	, BufferUsageFlags(0)
+	, NumBuffers(0)
+	, DynamicBufferIndex(0)
+	, LockStatus(ELockStatus::Unlocked)
+{
+	VULKAN_TRACK_OBJECT_CREATE(FVulkanResourceMultiBuffer, this);
 
+	const bool bZeroSize = (InSize == 0);
+	BufferUsageFlags = UEToVKBufferUsageFlags(InDevice, InUEUsage, bZeroSize);
+	
+	if (!bZeroSize)
+	{
+		const bool bVolatile = EnumHasAnyFlags(InUEUsage, BUF_Volatile);
 		if (bVolatile)
 		{
 			bool bRenderThread = IsInRenderingThread();
@@ -153,23 +151,44 @@ FVulkanResourceMultiBuffer::FVulkanResourceMultiBuffer(FVulkanDevice* InDevice,
 		{
 			VkDevice VulkanDevice = InDevice->GetInstanceHandle();
 
-			VkMemoryPropertyFlags BufferMemFlags = VK_MEMORY_PROPERTY_DEVICE_LOCAL_BIT;
+			NumBuffers = GetNumBuffersFromUsage(InUEUsage);
+			check(NumBuffers <= UE_ARRAY_COUNT(Buffers));
+
 			const bool bUnifiedMem = InDevice->HasUnifiedMemory();
-			if (bUnifiedMem)
-			{
-				BufferMemFlags |= (VK_MEMORY_PROPERTY_HOST_VISIBLE_BIT | VK_MEMORY_PROPERTY_HOST_COHERENT_BIT);
-			}
 
-			NumBuffers = bDynamic ? NUM_BUFFERS : 1;
-			check(NumBuffers <= UE_ARRAY_COUNT(Buffers));
+			if (InTransientHeapAllocation != nullptr)
+			{
+				const uint32 BufferAlignment = FMemoryManager::CalculateBufferAlignment(*InDevice, BufferUsageFlags);
+				const uint32 AlignedSize = Align(InSize, BufferAlignment);
 
-			for (uint32 Index = 0; Index < NumBuffers; ++Index)
+				Buffers[0] = FVulkanTransientHeap::GetVulkanAllocation(*InTransientHeapAllocation);
+				Buffers[0].Size = InSize;
+				check(Buffers[0].Offset % BufferAlignment == 0);
+				for (uint32 Index = 1; Index < NumBuffers; ++Index)
+				{
+					Buffers[Index] = Buffers[Index - 1];
+					Buffers[Index].Offset += AlignedSize;
+					Buffers[Index].Size = InSize;
+					check((Buffers[Index].Offset + InSize) <= InTransientHeapAllocation->Size);
+				}
+			}
+			else
 			{
-				if(!InDevice->GetMemoryManager().AllocateBufferPooled(Buffers[Index], this, InSize, BufferUsageFlags, BufferMemFlags, EVulkanAllocationMetaMultiBuffer, __FILE__, __LINE__))
+				VkMemoryPropertyFlags BufferMemFlags = VK_MEMORY_PROPERTY_DEVICE_LOCAL_BIT;
+				if (bUnifiedMem)
+				{
+					BufferMemFlags |= (VK_MEMORY_PROPERTY_HOST_VISIBLE_BIT | VK_MEMORY_PROPERTY_HOST_COHERENT_BIT);
+				}
+
+				for (uint32 Index = 0; Index < NumBuffers; ++Index)
 				{
-					InDevice->GetMemoryManager().HandleOOM();
+					if (!InDevice->GetMemoryManager().AllocateBufferPooled(Buffers[Index], this, InSize, BufferUsageFlags, BufferMemFlags, EVulkanAllocationMetaMultiBuffer, __FILE__, __LINE__))
+					{
+						InDevice->GetMemoryManager().HandleOOM();
+					}
 				}
 			}
+
 			Current.Alloc.Reference(Buffers[DynamicBufferIndex]);
 			Current.Handle = (VkBuffer)Current.Alloc.VulkanHandle;
 			Current.Offset = Current.Alloc.Offset;
@@ -475,15 +494,30 @@ void FVulkanResourceMultiBuffer::Swap(FVulkanResourceMultiBuffer& Other)
 	::Swap(VolatileLockInfo, Other.VolatileLockInfo);
 }
 
-FBufferRHIRef FVulkanDynamicRHI::RHICreateBuffer(uint32 Size, EBufferUsageFlags Usage, uint32 Stride, ERHIAccess ResourceState, FRHIResourceCreateInfo& CreateInfo)
+FBufferRHIRef FVulkanDynamicRHI::RHICreateBuffer(uint32 Size, EBufferUsageFlags Usage, uint32 Stride, ERHIAccess ResourceState, FRHIResourceCreateInfo& ResourceCreateInfo)
 {
 	LLM_SCOPE_VULKAN(ELLMTagVulkan::VulkanBuffers);
 
-	if (CreateInfo.bWithoutNativeResource)
+	if (ResourceCreateInfo.bWithoutNativeResource)
 	{
-		return new FVulkanResourceMultiBuffer(nullptr, 0, BUF_None, 0, CreateInfo, nullptr);
+		return new FVulkanResourceMultiBuffer(nullptr, 0, BUF_None, 0, ResourceCreateInfo, nullptr);
 	}
-	return new FVulkanResourceMultiBuffer(Device, Size, Usage, Stride, CreateInfo, nullptr);
+	return new FVulkanResourceMultiBuffer(Device, Size, Usage, Stride, ResourceCreateInfo, nullptr);
+}
+
+FRHIBuffer* FVulkanDynamicRHI::CreateBuffer(const FRHIBufferCreateInfo& InCreateInfo, FRHIResourceCreateInfo& InResourceCreateInfo, const FRHITransientHeapAllocation* InTransientHeapAllocation)
+{
+	LLM_SCOPE_VULKAN(ELLMTagVulkan::VulkanBuffers);
+
+	if (InTransientHeapAllocation == nullptr)
+	{
+		return this->RHICreateBuffer(InCreateInfo.Size, InCreateInfo.Usage, InCreateInfo.Stride, ERHIAccess::None, InResourceCreateInfo);
+	}
+
+	checkf(!EnumHasAnyFlags(InCreateInfo.Usage, BUF_AccelerationStructure), TEXT("AccelerationStructure not yet supported as TransientResource."));
+	checkf(!InResourceCreateInfo.bWithoutNativeResource, TEXT("WithoutNativeResource not yet supported as TransientResource."));
+
+	return new FVulkanResourceMultiBuffer(Device, InCreateInfo.Size, InCreateInfo.Usage, InCreateInfo.Stride, InResourceCreateInfo, nullptr, InTransientHeapAllocation);
 }
 
 void* FVulkanDynamicRHI::LockBuffer_BottomOfPipe(FRHICommandListImmediate& RHICmdList, FRHIBuffer* BufferRHI, uint32 Offset, uint32 Size, EResourceLockMode LockMode)
diff --git a/Engine/Source/Runtime/VulkanRHI/Private/VulkanDevice.cpp b/Engine/Source/Runtime/VulkanRHI/Private/VulkanDevice.cpp
--- a/Engine/Source/Runtime/VulkanRHI/Private/VulkanDevice.cpp
+++ b/Engine/Source/Runtime/VulkanRHI/Private/VulkanDevice.cpp
@@ -13,6 +13,7 @@
 #include "Misc/FileHelper.h"
 #include "VulkanPlatform.h"
 #include "VulkanLLM.h"
+#include "VulkanTransientResourceAllocator.h"
 
 TAutoConsoleVariable<int32> GRHIAllowAsyncComputeCvar(
 	TEXT("r.Vulkan.AllowAsyncCompute"),
@@ -771,6 +772,36 @@ void FVulkanDevice::SetupFormats()
 					);
 			}
 		}
+
+#define VULKAN_CHECK_FORMAT_CAPABILITY(PF_Name) if (EnumHasAllFlags(GPixelFormats[PixelFormatIndex].Capabilities, EPixelFormatCapabilities::PF_Name)) { CapabilitiesString += TEXT(#PF_Name) TEXT(", ");}
+
+		UE_LOG(LogVulkanRHI, Warning, TEXT("Pixel Format Capabilities for Vulkan:"));
+		for (int32 PixelFormatIndex = 0; PixelFormatIndex < PF_MAX; ++PixelFormatIndex)
+		{
+			if (GPixelFormats[PixelFormatIndex].Supported)
+			{
+				FString CapabilitiesString;
+
+				VULKAN_CHECK_FORMAT_CAPABILITY(TextureSample);
+				VULKAN_CHECK_FORMAT_CAPABILITY(TextureCube);
+				VULKAN_CHECK_FORMAT_CAPABILITY(RenderTarget);
+				VULKAN_CHECK_FORMAT_CAPABILITY(DepthStencil);
+				VULKAN_CHECK_FORMAT_CAPABILITY(TextureBlendable);
+				VULKAN_CHECK_FORMAT_CAPABILITY(TextureAtomics);
+				
+				VULKAN_CHECK_FORMAT_CAPABILITY(Buffer);
+				VULKAN_CHECK_FORMAT_CAPABILITY(VertexBuffer);
+				VULKAN_CHECK_FORMAT_CAPABILITY(IndexBuffer);
+				VULKAN_CHECK_FORMAT_CAPABILITY(BufferAtomics);
+
+				VULKAN_CHECK_FORMAT_CAPABILITY(UAV);
+
+				UE_LOG(LogVulkanRHI, Warning, TEXT("%24s : %s"), GPixelFormats[PixelFormatIndex].Name, *CapabilitiesString);
+			}
+		}
+
+#undef VULKAN_CHECK_FORMAT_CAPABILITY
+
 	}
 #endif
 }
@@ -1405,6 +1436,12 @@ void FVulkanDevice::Destroy()
 		ComputeContext->ReleasePendingState();
 	}
 
+	if (TransientHeapCache)
+	{
+		delete TransientHeapCache;
+		TransientHeapCache = nullptr;
+	}
+
 	// Flush all pending deletes before destroying the device and any Vulkan context objects.
 	// Repeat until no new deletes are added
 	int32 NumDeletes = 0;
@@ -1489,7 +1526,6 @@ void FVulkanDevice::Destroy()
 #endif
 	}
 
-	
 	DeferredDeletionQueue.Clear();
 
 	MemoryManager.Deinit();
@@ -1660,3 +1696,12 @@ void FVulkanDevice::VulkanSetObjectName(VkObjectType Type, uint64_t Handle, cons
 	}
 #endif // VULKAN_ENABLE_DRAW_MARKERS
 }
+
+FVulkanTransientHeapCache& FVulkanDevice::GetOrCreateTransientHeapCache()
+{
+	if (!TransientHeapCache)
+	{
+		TransientHeapCache = FVulkanTransientHeapCache::Create(this);
+	}
+	return *TransientHeapCache;
+}
diff --git a/Engine/Source/Runtime/VulkanRHI/Private/VulkanDevice.h b/Engine/Source/Runtime/VulkanRHI/Private/VulkanDevice.h
--- a/Engine/Source/Runtime/VulkanRHI/Private/VulkanDevice.h
+++ b/Engine/Source/Runtime/VulkanRHI/Private/VulkanDevice.h
@@ -12,6 +12,7 @@ class FVulkanDescriptorSetCache;
 class FVulkanDescriptorPool;
 class FVulkanDescriptorPoolsManager;
 class FVulkanCommandListContextImmediate;
+class FVulkanTransientHeapCache;
 #if VULKAN_USE_NEW_QUERIES
 class FVulkanOcclusionQueryPool;
 #else
@@ -523,6 +524,8 @@ public:
 		return QueueFamilyProps;
 	}
 
+	FVulkanTransientHeapCache& GetOrCreateTransientHeapCache();
+
 
 private:
 	const VkFormatProperties& GetFormatProperties(VkFormat InFormat);
@@ -548,6 +551,8 @@ private:
 
 	VulkanRHI::FFenceManager FenceManager;
 
+	FVulkanTransientHeapCache* TransientHeapCache = nullptr;
+
 	// Active on ES3.1
 	FVulkanDescriptorSetCache* DescriptorSetCache = nullptr;
 	// Active on >= SM4
diff --git a/Engine/Source/Runtime/VulkanRHI/Private/VulkanMemory.cpp b/Engine/Source/Runtime/VulkanRHI/Private/VulkanMemory.cpp
--- a/Engine/Source/Runtime/VulkanRHI/Private/VulkanMemory.cpp
+++ b/Engine/Source/Runtime/VulkanRHI/Private/VulkanMemory.cpp
@@ -2354,7 +2354,6 @@ namespace VulkanRHI
 					{
 						BufferAllocationsToRelease.Add(BufferAllocation);
 						FreeAllocations.RemoveAtSwap(Index, 1, false);
-						//break;
 					}
 				}
 			}
@@ -2720,51 +2719,80 @@ namespace VulkanRHI
 		VULKAN_FREE_TRACK_INFO(Track);
 	}
 
-	bool FMemoryManager::AllocateBufferPooled(FVulkanAllocation& OutAllocation, FVulkanEvictable* AllocationOwner, uint32 Size, VkBufferUsageFlags BufferUsageFlags, VkMemoryPropertyFlags MemoryPropertyFlags, EVulkanAllocationMetaType MetaType, const char* File, uint32 Line)
+	uint32 FMemoryManager::CalculateBufferAlignment(FVulkanDevice& InDevice, const VkBufferUsageFlags BufferUsageFlags)
 	{
-		SCOPED_NAMED_EVENT(FResourceHeapManager_AllocateBufferPooled, FColor::Cyan);
-		check(OutAllocation.Type == EVulkanAllocationEmpty);
-		const VkPhysicalDeviceLimits& Limits = Device->GetLimits();
-		uint32 Alignment = 1;
+		const VkPhysicalDeviceLimits& Limits = InDevice.GetLimits();
 
-		float Priority = VULKAN_MEMORY_MEDIUM_PRIORITY;
+		const bool bIsTexelBuffer = VKHasAnyFlags(BufferUsageFlags, (VK_BUFFER_USAGE_UNIFORM_TEXEL_BUFFER_BIT | VK_BUFFER_USAGE_STORAGE_TEXEL_BUFFER_BIT));
+		const bool bIsStorageBuffer = VKHasAnyFlags(BufferUsageFlags, VK_BUFFER_USAGE_STORAGE_BUFFER_BIT);
+		const bool bIsVertexOrIndexBuffer = VKHasAnyFlags(BufferUsageFlags, (VK_BUFFER_USAGE_VERTEX_BUFFER_BIT | VK_BUFFER_USAGE_INDEX_BUFFER_BIT));
+		const bool bIsAccelerationStructureBuffer = VKHasAnyFlags(BufferUsageFlags, VK_BUFFER_USAGE_ACCELERATION_STRUCTURE_STORAGE_BIT_KHR);
+		const bool bIsUniformBuffer = VKHasAnyFlags(BufferUsageFlags, VK_BUFFER_USAGE_UNIFORM_BUFFER_BIT);
+
+		uint32 Alignment = 1;
 
-		bool bIsTexelBuffer = (BufferUsageFlags & (VK_BUFFER_USAGE_UNIFORM_TEXEL_BUFFER_BIT | VK_BUFFER_USAGE_STORAGE_TEXEL_BUFFER_BIT)) != 0;
-		bool bIsStorageBuffer = (BufferUsageFlags & VK_BUFFER_USAGE_STORAGE_BUFFER_BIT) != 0;
 		if (bIsTexelBuffer || bIsStorageBuffer)
 		{
-			Alignment = FMath::Max(Alignment, bIsTexelBuffer ? (uint32)Limits.minTexelBufferOffsetAlignment : 0);
-			Alignment = FMath::Max(Alignment, bIsStorageBuffer ? (uint32)Limits.minStorageBufferOffsetAlignment : 0);
+			Alignment = FMath::Max(Alignment, (uint32)Limits.minTexelBufferOffsetAlignment);
+			Alignment = FMath::Max(Alignment, (uint32)Limits.minStorageBufferOffsetAlignment);
+		}
+		else if (bIsVertexOrIndexBuffer)
+		{
+			// No alignment restrictions on Vertex or Index buffers, leave it at 1
+		}
+		else if (bIsAccelerationStructureBuffer)
+		{
+			Alignment = FMath::Max(Alignment, GRHIRayTracingAccelerationStructureAlignment);
+		}
+		else if (bIsUniformBuffer)
+		{
+			Alignment = FMath::Max(Alignment, (uint32)Limits.minUniformBufferOffsetAlignment);
 		}
 		else
 		{
-			bool bIsVertexOrIndexBuffer = (BufferUsageFlags & (VK_BUFFER_USAGE_VERTEX_BUFFER_BIT | VK_BUFFER_USAGE_INDEX_BUFFER_BIT)) != 0;
-			bool bIsAccelerationStructureBuffer = (BufferUsageFlags & VK_BUFFER_USAGE_ACCELERATION_STRUCTURE_STORAGE_BIT_KHR) != 0;
-			if (bIsVertexOrIndexBuffer)
-			{
-				// No alignment restrictions on Vertex or Index buffers, can live on CPU mem
-				Priority = VULKAN_MEMORY_LOW_PRIORITY;
-			}
-			else if (bIsAccelerationStructureBuffer)
-			{
-				Alignment = FMath::Max(Alignment, GRHIRayTracingAccelerationStructureAlignment);
-			}
-			else
-			{
-				// Uniform buffer
-				ensure((BufferUsageFlags & VK_BUFFER_USAGE_UNIFORM_BUFFER_BIT) == VK_BUFFER_USAGE_UNIFORM_BUFFER_BIT);
-				Alignment = FMath::Max(Alignment, (uint32)Limits.minUniformBufferOffsetAlignment);
-
-				Priority = VULKAN_MEMORY_HIGHER_PRIORITY;
-			}
+			checkf(false, TEXT("Unknown buffer alignment for VkBufferUsageFlags combination: 0x%x"), (uint32)BufferUsageFlags);
 		}
 
-		if (BufferUsageFlags & (VK_BUFFER_USAGE_STORAGE_TEXEL_BUFFER_BIT | VK_BUFFER_USAGE_INDIRECT_BUFFER_BIT))
+		return Alignment;
+	}
+
+	float FMemoryManager::CalculateBufferPriority(const VkBufferUsageFlags BufferUsageFlags)
+	{
+		float Priority = VULKAN_MEMORY_MEDIUM_PRIORITY;
+
+		if (VKHasAnyFlags(BufferUsageFlags, (VK_BUFFER_USAGE_STORAGE_TEXEL_BUFFER_BIT | VK_BUFFER_USAGE_INDIRECT_BUFFER_BIT)))
 		{
 			Priority = VULKAN_MEMORY_HIGHEST_PRIORITY;
 		}
+		else if (VKHasAnyFlags(BufferUsageFlags, (VK_BUFFER_USAGE_UNIFORM_TEXEL_BUFFER_BIT | VK_BUFFER_USAGE_STORAGE_BUFFER_BIT)))
+		{
+			Priority = VULKAN_MEMORY_MEDIUM_PRIORITY;
+		}
+		else if (VKHasAnyFlags(BufferUsageFlags, (VK_BUFFER_USAGE_VERTEX_BUFFER_BIT | VK_BUFFER_USAGE_INDEX_BUFFER_BIT)))
+		{
+			Priority = VULKAN_MEMORY_LOW_PRIORITY;
+		}
+		else if (VKHasAnyFlags(BufferUsageFlags, VK_BUFFER_USAGE_UNIFORM_BUFFER_BIT))
+		{
+			Priority = VULKAN_MEMORY_HIGHER_PRIORITY;
+		}
+		else
+		{
+			checkf(false, TEXT("Unknown priority for VkBufferUsageFlags combination: 0x%x"), (uint32)BufferUsageFlags);
+		}
+
+		return Priority;
+	}
+
+	bool FMemoryManager::AllocateBufferPooled(FVulkanAllocation& OutAllocation, FVulkanEvictable* AllocationOwner, uint32 Size, VkBufferUsageFlags BufferUsageFlags, VkMemoryPropertyFlags MemoryPropertyFlags, EVulkanAllocationMetaType MetaType, const char* File, uint32 Line)
+	{
+		SCOPED_NAMED_EVENT(FResourceHeapManager_AllocateBufferPooled, FColor::Cyan);
+		check(OutAllocation.Type == EVulkanAllocationEmpty);
+
+		uint32 Alignment = CalculateBufferAlignment(*Device, BufferUsageFlags);
+		const float Priority = CalculateBufferPriority(BufferUsageFlags);
 
-		int32 PoolSize = (int32)GetPoolTypeForAlloc(Size, Alignment);
+		const int32 PoolSize = (int32)GetPoolTypeForAlloc(Size, Alignment);
 		if (PoolSize != (int32)EPoolSizes::SizesCount)
 		{
 			Size = PoolSizes[PoolSize];
@@ -4194,7 +4222,7 @@ namespace VulkanRHI
 			if (Alloc.State == FVulkanAllocationInternal::EALLOCATED)
 			{
 				FVulkanEvictable* EvictableOwner = Alloc.AllocationOwner;
-				if(!EvictableOwner->CanMove())
+				if((EvictableOwner == nullptr) || !EvictableOwner->CanMove())
 				{
 					return false;
 				}
diff --git a/Engine/Source/Runtime/VulkanRHI/Private/VulkanMemory.h b/Engine/Source/Runtime/VulkanRHI/Private/VulkanMemory.h
--- a/Engine/Source/Runtime/VulkanRHI/Private/VulkanMemory.h
+++ b/Engine/Source/Runtime/VulkanRHI/Private/VulkanMemory.h
@@ -774,6 +774,8 @@ namespace VulkanRHI
 		void Init();
 		void Deinit();
 
+		static uint32 CalculateBufferAlignment(FVulkanDevice& InDevice, const VkBufferUsageFlags BufferUsageFlags);
+		static float CalculateBufferPriority(const VkBufferUsageFlags BufferUsageFlags);
 
 		void FreeVulkanAllocation(FVulkanAllocation& Allocation, EVulkanFreeFlags FreeFlags = EVulkanFreeFlag_None);
 		void FreeVulkanAllocationPooledBuffer(FVulkanAllocation& Allocation);
diff --git a/Engine/Source/Runtime/VulkanRHI/Private/VulkanRHI.cpp b/Engine/Source/Runtime/VulkanRHI/Private/VulkanRHI.cpp
--- a/Engine/Source/Runtime/VulkanRHI/Private/VulkanRHI.cpp
+++ b/Engine/Source/Runtime/VulkanRHI/Private/VulkanRHI.cpp
@@ -23,6 +23,7 @@
 #include "RHIValidation.h"
 #include "IHeadMountedDisplayModule.h"
 #include "VulkanRenderpass.h"
+#include "VulkanTransientResourceAllocator.h"
 
 static_assert(sizeof(VkStructureType) == sizeof(int32), "ZeroVulkanStruct() assumes VkStructureType is int32!");
 
@@ -77,6 +78,16 @@ static FAutoConsoleVariableRef GCVarInputAttachmentShaderRead(
 	ECVF_ReadOnly
 );
 
+int32 GVulkanEnableTransientResourceAllocator = 0;
+static FAutoConsoleVariableRef GCVarEnableTransientResourceAllocator(
+	TEXT("r.Vulkan.EnableTransientResourceAllocator"),
+	GVulkanEnableTransientResourceAllocator,
+	TEXT("Whether to enable the TransientResourceAllocator to reduce memory usage\n")
+	TEXT("0 to disabled (default)\n")
+	TEXT("1 to enable\n"),
+	ECVF_ReadOnly
+);
+
 bool GGPUCrashDebuggingEnabled = false;
 
 
@@ -1778,4 +1789,13 @@ void FVulkanDynamicRHI::VulkanSetImageLayout( VkCommandBuffer CmdBuffer, VkImage
 	::VulkanSetImageLayout( CmdBuffer, Image, OldLayout, NewLayout, SubresourceRange );
 }
 
+IRHITransientResourceAllocator* FVulkanDynamicRHI::RHICreateTransientResourceAllocator()
+{
+	if (GVulkanEnableTransientResourceAllocator)
+	{
+		return new FVulkanTransientResourceAllocator(Device->GetOrCreateTransientHeapCache());
+	}
+	return nullptr;
+}
+
 #undef LOCTEXT_NAMESPACE
diff --git a/Engine/Source/Runtime/VulkanRHI/Private/VulkanResources.h b/Engine/Source/Runtime/VulkanRHI/Private/VulkanResources.h
--- a/Engine/Source/Runtime/VulkanRHI/Private/VulkanResources.h
+++ b/Engine/Source/Runtime/VulkanRHI/Private/VulkanResources.h
@@ -27,6 +27,7 @@ class FVulkanLayout;
 class FVulkanOcclusionQuery;
 class FVulkanShaderResourceView;
 class FVulkanCommandBufferManager;
+struct FRHITransientHeapAllocation;
 
 namespace VulkanRHI
 {
@@ -362,7 +363,8 @@ public:
 
 	FVulkanSurface(FVulkanDevice& Device, FVulkanEvictable* Owner, VkImageViewType ResourceType, EPixelFormat Format,
 					uint32 SizeX, uint32 SizeY, uint32 SizeZ, uint32 ArraySize,
-					uint32 NumMips, uint32 NumSamples, ETextureCreateFlags UEFlags, ERHIAccess InResourceState, const FRHIResourceCreateInfo& CreateInfo);
+					uint32 NumMips, uint32 NumSamples, ETextureCreateFlags UEFlags, ERHIAccess InResourceState, 
+					const FRHIResourceCreateInfo& CreateInfo, const FRHITransientHeapAllocation* InTransientHeapAllocation = nullptr);
 
 	// Constructor for externally owned Image
 	FVulkanSurface(FVulkanDevice& Device, VkImageViewType ResourceType, EPixelFormat Format,
@@ -539,7 +541,7 @@ struct FVulkanTextureBase : public FVulkanEvictable, public IRefCountedObject
 		return (FVulkanTextureBase*)Texture->GetTextureBaseRHI();
 	}
 
-	FVulkanTextureBase(FVulkanDevice& Device, VkImageViewType ResourceType, EPixelFormat Format, uint32 SizeX, uint32 SizeY, uint32 SizeZ, uint32 ArraySize, uint32 NumMips, uint32 NumSamples, ETextureCreateFlags UEFlags, ERHIAccess InResourceState, const FRHIResourceCreateInfo& CreateInfo);
+	FVulkanTextureBase(FVulkanDevice& Device, VkImageViewType ResourceType, EPixelFormat Format, uint32 SizeX, uint32 SizeY, uint32 SizeZ, uint32 ArraySize, uint32 NumMips, uint32 NumSamples, ETextureCreateFlags UEFlags, ERHIAccess InResourceState, const FRHIResourceCreateInfo& CreateInfo, const FRHITransientHeapAllocation* InTransientHeapAllocation);
 	FVulkanTextureBase(FVulkanDevice& Device, VkImageViewType ResourceType, EPixelFormat Format, uint32 SizeX, uint32 SizeY, uint32 SizeZ, uint32 ArraySize, uint32 NumMips, uint32 NumSamples, VkImage InImage, VkDeviceMemory InMem, ETextureCreateFlags UEFlags, const FRHIResourceCreateInfo& CreateInfo);
 	FVulkanTextureBase(FVulkanDevice& Device, VkImageViewType ResourceType, EPixelFormat Format, uint32 SizeX, uint32 SizeY, uint32 SizeZ, uint32 ArraySize, uint32 NumMips, uint32 NumSamples, VkImage InImage, VkDeviceMemory InMem, FSamplerYcbcrConversionInitializer& ConversionInitializer, ETextureCreateFlags UEFlags, const FRHIResourceCreateInfo& CreateInfo);
 
@@ -593,7 +595,7 @@ class FVulkanTexture2D : public FRHITexture2D, public FVulkanTextureBase
 	FName GetResourceFName(){ return GetName(); }
 	virtual FRHIResource* GetRHIResource() { return (FRHITexture2D*)this; }
 public:
-	FVulkanTexture2D(FVulkanDevice& Device, EPixelFormat Format, uint32 SizeX, uint32 SizeY, uint32 NumMips, uint32 NumSamples, ETextureCreateFlags UEFlags, ERHIAccess InResourceState, const FRHIResourceCreateInfo& CreateInfo);
+	FVulkanTexture2D(FVulkanDevice& Device, EPixelFormat Format, uint32 SizeX, uint32 SizeY, uint32 NumMips, uint32 NumSamples, ETextureCreateFlags UEFlags, ERHIAccess InResourceState, const FRHIResourceCreateInfo& CreateInfo, const FRHITransientHeapAllocation* InTransientHeapAllocation = nullptr);
 	FVulkanTexture2D(FVulkanDevice& Device, EPixelFormat Format, uint32 SizeX, uint32 SizeY, uint32 NumMips, uint32 NumSamples, VkImage Image, ETextureCreateFlags UEFlags, const FRHIResourceCreateInfo& CreateInfo);
 	FVulkanTexture2D(FVulkanDevice& Device, EPixelFormat Format, uint32 SizeX, uint32 SizeY, uint32 NumMips, uint32 NumSamples, VkImage Image, struct FSamplerYcbcrConversionInitializer& ConversionInitializer, ETextureCreateFlags UEFlags, const FRHIResourceCreateInfo& CreateInfo);
 
@@ -644,7 +646,7 @@ class FVulkanTexture2DArray : public FRHITexture2DArray, public FVulkanTextureBa
 	FName GetResourceFName() { return GetName(); }
 public:
 	// Constructor, just calls base and Surface constructor
-	FVulkanTexture2DArray(FVulkanDevice& Device, EPixelFormat Format, uint32 SizeX, uint32 SizeY, uint32 ArraySize, uint32 NumMips, uint32 NumSamples, ETextureCreateFlags Flags, ERHIAccess InResourceState, const FRHIResourceCreateInfo& CreateInfo);
+	FVulkanTexture2DArray(FVulkanDevice& Device, EPixelFormat Format, uint32 SizeX, uint32 SizeY, uint32 ArraySize, uint32 NumMips, uint32 NumSamples, ETextureCreateFlags Flags, ERHIAccess InResourceState, const FRHIResourceCreateInfo& CreateInfo, const FRHITransientHeapAllocation* InTransientHeapAllocation = nullptr);
 	FVulkanTexture2DArray(FVulkanDevice& Device, EPixelFormat Format, uint32 SizeX, uint32 SizeY, uint32 ArraySize, uint32 NumMips, uint32 NumSamples, VkImage Image, ETextureCreateFlags Flags, const FRHIResourceCreateInfo& CreateInfo);
 
 	// Aliasing constructor
@@ -694,7 +696,7 @@ class FVulkanTexture3D : public FRHITexture3D, public FVulkanTextureBase
 	FName GetResourceFName() { return GetName(); }
 public:
 	// Constructor, just calls base and Surface constructor
-	FVulkanTexture3D(FVulkanDevice& Device, EPixelFormat Format, uint32 SizeX, uint32 SizeY, uint32 SizeZ, uint32 NumMips, ETextureCreateFlags Flags, ERHIAccess InResourceState, const FRHIResourceCreateInfo& CreateInfo);
+	FVulkanTexture3D(FVulkanDevice& Device, EPixelFormat Format, uint32 SizeX, uint32 SizeY, uint32 SizeZ, uint32 NumMips, ETextureCreateFlags Flags, ERHIAccess InResourceState, const FRHIResourceCreateInfo& CreateInfo, const FRHITransientHeapAllocation* InTransientHeapAllocation = nullptr);
 	FVulkanTexture3D(FVulkanDevice& Device, EPixelFormat Format, uint32 SizeX, uint32 SizeY, uint32 SizeZ, uint32 NumMips, VkImage Image, ETextureCreateFlags Flags, const FRHIResourceCreateInfo& CreateInfo);
 	virtual ~FVulkanTexture3D();
 
@@ -741,7 +743,7 @@ class FVulkanTextureCube : public FRHITextureCube, public FVulkanTextureBase
 {
 	FName GetResourceFName() { return GetName(); }
 public:
-	FVulkanTextureCube(FVulkanDevice& Device, EPixelFormat Format, uint32 Size, bool bArray, uint32 ArraySize, uint32 NumMips, ETextureCreateFlags Flags, ERHIAccess InResourceState, const FRHIResourceCreateInfo& CreateInfo);
+	FVulkanTextureCube(FVulkanDevice& Device, EPixelFormat Format, uint32 Size, bool bArray, uint32 ArraySize, uint32 NumMips, ETextureCreateFlags Flags, ERHIAccess InResourceState, const FRHIResourceCreateInfo& CreateInfo, const FRHITransientHeapAllocation* InTransientHeapAllocation = nullptr);
 	FVulkanTextureCube(FVulkanDevice& Device, EPixelFormat Format, uint32 Size, bool bArray, uint32 ArraySize, uint32 NumMips, VkImage Image, ETextureCreateFlags Flags, const FRHIResourceCreateInfo& CreateInfo);
 
 	// Aliasing constructor
@@ -1114,7 +1116,7 @@ class FVulkanResourceMultiBuffer : public FRHIBuffer, public FVulkanEvictable, p
 	virtual void Move(FVulkanDevice& Device, FVulkanCommandListContext& Context, VulkanRHI::FVulkanAllocation& NewAllocation);
 
 public:
-	FVulkanResourceMultiBuffer(FVulkanDevice* InDevice, uint32 InSize, EBufferUsageFlags InUEUsage, uint32 InStride, FRHIResourceCreateInfo& CreateInfo, class FRHICommandListImmediate* InRHICmdList = nullptr);
+	FVulkanResourceMultiBuffer(FVulkanDevice* InDevice, uint32 InSize, EBufferUsageFlags InUEUsage, uint32 InStride, FRHIResourceCreateInfo& CreateInfo, class FRHICommandListImmediate* InRHICmdList = nullptr, const FRHITransientHeapAllocation* InTransientHeapAllocation = nullptr);
 	virtual ~FVulkanResourceMultiBuffer();
 
 	inline const VulkanRHI::FVulkanAllocation& GetCurrentAllocation() const
@@ -1196,6 +1198,14 @@ public:
 	VkDeviceAddress GetDeviceAddress() const;
 #endif
 
+	static VkBufferUsageFlags UEToVKBufferUsageFlags(FVulkanDevice* InDevice, EBufferUsageFlags InUEUsage, bool bZeroSize);
+
+	static inline int32 GetNumBuffersFromUsage(EBufferUsageFlags InUEUsage)
+	{
+		const bool bDynamic = EnumHasAnyFlags(InUEUsage, BUF_Dynamic);
+		return bDynamic ? NUM_BUFFERS : 1;
+	}
+
 protected:
 	VkBufferUsageFlags BufferUsageFlags;
 	uint8 NumBuffers;
diff --git a/Engine/Source/Runtime/VulkanRHI/Private/VulkanTexture.cpp b/Engine/Source/Runtime/VulkanRHI/Private/VulkanTexture.cpp
--- a/Engine/Source/Runtime/VulkanRHI/Private/VulkanTexture.cpp
+++ b/Engine/Source/Runtime/VulkanRHI/Private/VulkanTexture.cpp
@@ -11,6 +11,7 @@
 #include "Containers/ResourceArray.h"
 #include "VulkanLLM.h"
 #include "VulkanBarriers.h"
+#include "VulkanTransientResourceAllocator.h"
 
 int32 GVulkanSubmitOnTextureUnlock = 1;
 static FAutoConsoleVariableRef CVarVulkanSubmitOnTextureUnlock(
@@ -578,7 +579,8 @@ static VkImageLayout GetInitialLayoutFromRHIAccess(ERHIAccess RHIAccess, ETextur
 
 FVulkanSurface::FVulkanSurface(FVulkanDevice& InDevice, FVulkanEvictable* Owner, VkImageViewType ResourceType, EPixelFormat InFormat,
 								uint32 SizeX, uint32 SizeY, uint32 SizeZ, uint32 InArraySize, uint32 InNumMips,
-								uint32 InNumSamples, ETextureCreateFlags InUEFlags, ERHIAccess InResourceState, const FRHIResourceCreateInfo& CreateInfo)
+								uint32 InNumSamples, ETextureCreateFlags InUEFlags, ERHIAccess InResourceState, 
+								const FRHIResourceCreateInfo& CreateInfo, const FRHITransientHeapAllocation* InTransientHeapAllocation)
 	: Device(&InDevice)
 	, Image(VK_NULL_HANDLE)
 	, StorageFormat(VK_FORMAT_UNDEFINED)
@@ -687,6 +689,17 @@ FVulkanSurface::FVulkanSurface(FVulkanDevice& InDevice, FVulkanEvictable* Owner,
 		Owner = this;
 	}
 	check(bRenderTarget || bUAV || Owner != nullptr);
+
+	if (InTransientHeapAllocation != nullptr)
+	{
+		check(!bCPUReadback);
+		check(!bMemoryless);
+		check(InTransientHeapAllocation->Offset % MemoryRequirements.alignment == 0);
+		check(InTransientHeapAllocation->Size >= MemoryRequirements.size);
+		Allocation = FVulkanTransientHeap::GetVulkanAllocation(*InTransientHeapAllocation);
+	}
+	else
+	{
 	EVulkanAllocationMetaType MetaType = (bRenderTarget||bUAV) ? EVulkanAllocationMetaImageRenderTarget : EVulkanAllocationMetaImageOther;
 #if VULKAN_SUPPORTS_DEDICATED_ALLOCATION
 	extern int32 GVulkanEnableDedicatedImageMemory;
@@ -707,6 +720,7 @@ FVulkanSurface::FVulkanSurface(FVulkanDevice& InDevice, FVulkanEvictable* Owner,
 			checkNoEntry();
 		}
 	}
+	}
 	Allocation.BindImage(Device, Image);
 
 	// update rhi stats
@@ -1293,6 +1307,41 @@ FTexture3DRHIRef FVulkanDynamicRHI::RHICreateTexture3D(uint32 SizeX, uint32 Size
 	return Tex3d;
 }
 
+FRHITexture* FVulkanDynamicRHI::CreateTexture(const FRHITextureCreateInfo& InCreateInfo, FRHIResourceCreateInfo& InResourceCreateInfo, ERHIAccess InResourceState, const FRHITransientHeapAllocation* InTransientHeapAllocation)
+{
+	LLM_SCOPE_VULKAN(GetMemoryTagForTextureFlags(InCreateInfo.Flags));
+
+	FRHITexture* Texture = nullptr;
+
+	switch (InCreateInfo.Dimension)
+	{
+	case ETextureDimension::Texture2D:
+		Texture = new FVulkanTexture2D(*Device, InCreateInfo.Format, InCreateInfo.Extent.X, InCreateInfo.Extent.Y, InCreateInfo.NumMips, InCreateInfo.NumSamples, InCreateInfo.Flags, InResourceState, InResourceCreateInfo, InTransientHeapAllocation);
+		break;
+
+	case ETextureDimension::Texture2DArray:
+		Texture = new FVulkanTexture2DArray(*Device, InCreateInfo.Format, InCreateInfo.Extent.X, InCreateInfo.Extent.Y, InCreateInfo.ArraySize, InCreateInfo.NumMips, InCreateInfo.NumSamples, InCreateInfo.Flags, InResourceState, InResourceCreateInfo, InTransientHeapAllocation);
+		break;
+
+	case ETextureDimension::TextureCube:
+		Texture = new FVulkanTextureCube(*Device, InCreateInfo.Format, InCreateInfo.Extent.X, false, 1, InCreateInfo.NumMips, InCreateInfo.Flags, InResourceState, InResourceCreateInfo, InTransientHeapAllocation);
+		break;
+
+	case ETextureDimension::TextureCubeArray:
+		Texture = new FVulkanTextureCube(*Device, InCreateInfo.Format, InCreateInfo.Extent.X, true, InCreateInfo.ArraySize, InCreateInfo.NumMips, InCreateInfo.Flags, InResourceState, InResourceCreateInfo, InTransientHeapAllocation);
+		break;
+
+	case ETextureDimension::Texture3D:
+		Texture = new FVulkanTexture3D(*Device, InCreateInfo.Format, InCreateInfo.Extent.X, InCreateInfo.Extent.Y, InCreateInfo.Depth, InCreateInfo.NumMips, InCreateInfo.Flags, InResourceState, InResourceCreateInfo, InTransientHeapAllocation);
+		break;
+
+	default:
+		checkNoEntry();
+	}
+
+	return Texture;
+}
+
 static void DoAsyncReallocateTexture2D(FVulkanCommandListContext& Context, FVulkanTexture2D* OldTexture, FVulkanTexture2D* NewTexture, int32 NewMipCount, int32 NewSizeX, int32 NewSizeY, FThreadSafeCounter* RequestStatus)
 {
 	LLM_SCOPE_VULKAN(ELLMTagVulkan::VulkanTextures);
@@ -1918,8 +1967,8 @@ void FVulkanTextureView::Destroy(FVulkanDevice& Device)
 	}
 }
 
-FVulkanTextureBase::FVulkanTextureBase(FVulkanDevice& Device, VkImageViewType ResourceType, EPixelFormat InFormat, uint32 SizeX, uint32 SizeY, uint32 SizeZ, uint32 ArraySize, uint32 NumMips, uint32 NumSamples, ETextureCreateFlags UEFlags, ERHIAccess InResourceState, const FRHIResourceCreateInfo& CreateInfo)
-	: Surface(Device, this, ResourceType, InFormat, SizeX, SizeY, SizeZ, ArraySize, NumMips, NumSamples, UEFlags, InResourceState, CreateInfo)
+FVulkanTextureBase::FVulkanTextureBase(FVulkanDevice& Device, VkImageViewType ResourceType, EPixelFormat InFormat, uint32 SizeX, uint32 SizeY, uint32 SizeZ, uint32 ArraySize, uint32 NumMips, uint32 NumSamples, ETextureCreateFlags UEFlags, ERHIAccess InResourceState, const FRHIResourceCreateInfo& CreateInfo, const FRHITransientHeapAllocation* InTransientHeapAllocation)
+	: Surface(Device, this, ResourceType, InFormat, SizeX, SizeY, SizeZ, ArraySize, NumMips, NumSamples, UEFlags, InResourceState, CreateInfo, InTransientHeapAllocation)
 	, PartialView(nullptr)
 {
 	Surface.OwningTexture = this;
@@ -2245,9 +2294,9 @@ bool FVulkanTextureBase::GetTextureResourceInfo(FRHIResourceInfo& OutResourceInf
 	return true;
 }
 
-FVulkanTexture2D::FVulkanTexture2D(FVulkanDevice& Device, EPixelFormat InFormat, uint32 SizeX, uint32 SizeY, uint32 NumMips, uint32 NumSamples, ETextureCreateFlags UEFlags, ERHIAccess InResourceState, const FRHIResourceCreateInfo& CreateInfo)
+FVulkanTexture2D::FVulkanTexture2D(FVulkanDevice& Device, EPixelFormat InFormat, uint32 SizeX, uint32 SizeY, uint32 NumMips, uint32 NumSamples, ETextureCreateFlags UEFlags, ERHIAccess InResourceState, const FRHIResourceCreateInfo& CreateInfo, const FRHITransientHeapAllocation* InTransientHeapAllocation)
 :	FRHITexture2D(SizeX, SizeY, FMath::Max(NumMips, 1u), NumSamples, InFormat, UEFlags, CreateInfo.ClearValueBinding)
-,	FVulkanTextureBase(Device, VK_IMAGE_VIEW_TYPE_2D, InFormat, SizeX, SizeY, 1, 1, FMath::Max(NumMips, 1u), NumSamples, UEFlags, InResourceState, CreateInfo)
+,	FVulkanTextureBase(Device, VK_IMAGE_VIEW_TYPE_2D, InFormat, SizeX, SizeY, 1, 1, FMath::Max(NumMips, 1u), NumSamples, UEFlags, InResourceState, CreateInfo, InTransientHeapAllocation)
 {
 }
 
@@ -2273,9 +2322,9 @@ FVulkanTexture2D::~FVulkanTexture2D()
 {
 }
 
-FVulkanTexture2DArray::FVulkanTexture2DArray(FVulkanDevice& Device, EPixelFormat Format, uint32 SizeX, uint32 SizeY, uint32 ArraySize, uint32 NumMips, uint32 NumSamples, ETextureCreateFlags Flags, ERHIAccess InResourceState, const FRHIResourceCreateInfo& CreateInfo)
+FVulkanTexture2DArray::FVulkanTexture2DArray(FVulkanDevice& Device, EPixelFormat Format, uint32 SizeX, uint32 SizeY, uint32 ArraySize, uint32 NumMips, uint32 NumSamples, ETextureCreateFlags Flags, ERHIAccess InResourceState, const FRHIResourceCreateInfo& CreateInfo, const FRHITransientHeapAllocation* InTransientHeapAllocation)
 	:	FRHITexture2DArray(SizeX, SizeY, ArraySize, NumMips, NumSamples, Format, Flags, CreateInfo.ClearValueBinding)
-	,	FVulkanTextureBase(Device, VK_IMAGE_VIEW_TYPE_2D_ARRAY, Format, SizeX, SizeY, 1, ArraySize, NumMips, NumSamples, Flags, InResourceState, CreateInfo)
+	,	FVulkanTextureBase(Device, VK_IMAGE_VIEW_TYPE_2D_ARRAY, Format, SizeX, SizeY, 1, ArraySize, NumMips, NumSamples, Flags, InResourceState, CreateInfo, InTransientHeapAllocation)
 {
 }
 
@@ -2292,10 +2341,10 @@ FVulkanTexture2DArray::FVulkanTexture2DArray(FTextureRHIRef& SrcTextureRHI, cons
 }
 
 
-FVulkanTextureCube::FVulkanTextureCube(FVulkanDevice& Device, EPixelFormat Format, uint32 Size, bool bArray, uint32 ArraySize, uint32 NumMips, ETextureCreateFlags Flags, ERHIAccess InResourceState, const FRHIResourceCreateInfo& CreateInfo)
+FVulkanTextureCube::FVulkanTextureCube(FVulkanDevice& Device, EPixelFormat Format, uint32 Size, bool bArray, uint32 ArraySize, uint32 NumMips, ETextureCreateFlags Flags, ERHIAccess InResourceState, const FRHIResourceCreateInfo& CreateInfo, const FRHITransientHeapAllocation* InTransientHeapAllocation)
 :	 FRHITextureCube(Size, NumMips, Format, Flags, CreateInfo.ClearValueBinding)
 	//#todo-rco: Array/slices count
-,	FVulkanTextureBase(Device, bArray ? VK_IMAGE_VIEW_TYPE_CUBE_ARRAY : VK_IMAGE_VIEW_TYPE_CUBE, Format, Size, Size, 1, ArraySize, NumMips, /*NumSamples=*/ 1, Flags, InResourceState, CreateInfo)
+,	FVulkanTextureBase(Device, bArray ? VK_IMAGE_VIEW_TYPE_CUBE_ARRAY : VK_IMAGE_VIEW_TYPE_CUBE, Format, Size, Size, 1, ArraySize, NumMips, /*NumSamples=*/ 1, Flags, InResourceState, CreateInfo, InTransientHeapAllocation)
 {
 }
 
@@ -2317,9 +2366,9 @@ FVulkanTextureCube::~FVulkanTextureCube()
 }
 
 
-FVulkanTexture3D::FVulkanTexture3D(FVulkanDevice& Device, EPixelFormat Format, uint32 SizeX, uint32 SizeY, uint32 SizeZ, uint32 NumMips, ETextureCreateFlags Flags, ERHIAccess InResourceState, const FRHIResourceCreateInfo& CreateInfo)
+FVulkanTexture3D::FVulkanTexture3D(FVulkanDevice& Device, EPixelFormat Format, uint32 SizeX, uint32 SizeY, uint32 SizeZ, uint32 NumMips, ETextureCreateFlags Flags, ERHIAccess InResourceState, const FRHIResourceCreateInfo& CreateInfo, const FRHITransientHeapAllocation* InTransientHeapAllocation)
 	: FRHITexture3D(SizeX, SizeY, SizeZ, NumMips, Format, Flags, CreateInfo.ClearValueBinding)
-	, FVulkanTextureBase(Device, VK_IMAGE_VIEW_TYPE_3D, Format, SizeX, SizeY, SizeZ, 1, NumMips, /*NumSamples=*/ 1, Flags, InResourceState, CreateInfo)
+	, FVulkanTextureBase(Device, VK_IMAGE_VIEW_TYPE_3D, Format, SizeX, SizeY, SizeZ, 1, NumMips, /*NumSamples=*/ 1, Flags, InResourceState, CreateInfo, InTransientHeapAllocation)
 {
 }
 
@@ -2484,7 +2533,7 @@ void FVulkanDynamicRHI::RHIVirtualTextureSetFirstMipVisible(FRHITexture2D* Textu
 	VULKAN_SIGNAL_UNIMPLEMENTED();
 }
 
-static VkMemoryRequirements FindOrCalculateTexturePlatformSize(FVulkanDevice* Device, VkImageViewType ViewType, uint32 SizeX, uint32 SizeY, uint32 SizeZ, uint8 Format, uint32 NumMips, uint32 NumSamples, ETextureCreateFlags Flags)
+static VkMemoryRequirements FindOrCalculateTexturePlatformSize(FVulkanDevice* Device, VkImageViewType ViewType, uint32 SizeX, uint32 SizeY, uint32 SizeZ, uint32 ArraySize, uint8 Format, uint32 NumMips, uint32 NumSamples, ETextureCreateFlags Flags)
 {
 	// Adjust number of mips as UTexture can request non-valid # of mips
 	NumMips = FMath::Min(FMath::FloorLog2(FMath::Max(SizeX, FMath::Max(SizeY, SizeZ))) + 1, NumMips);
@@ -2523,9 +2572,10 @@ static VkMemoryRequirements FindOrCalculateTexturePlatformSize(FVulkanDevice* De
 	// Create temporary image to measure the memory requirements
 	FVulkanSurface::FImageCreateInfo TmpCreateInfo;
 	FVulkanSurface::GenerateImageCreateInfo(TmpCreateInfo, *Device, ViewType,
-		PixelFormat, SizeX, SizeY, SizeZ, 1, NumMips, NumSamples,
+		PixelFormat, SizeX, SizeY, SizeZ, ArraySize, NumMips, NumSamples,
 		Flags, nullptr, nullptr, false);
 
+	// :TODO: VK_KHR_maintenance4
 	VkImage TmpImage;
 	VERIFYVULKANRESULT(VulkanRHI::vkCreateImage(Device->GetInstanceHandle(), &TmpCreateInfo.ImageCreateInfo, VULKAN_CPU_ALLOCATOR, &TmpImage));
 	VulkanRHI::vkGetImageMemoryRequirements(Device->GetInstanceHandle(), TmpImage, &MemReq);
@@ -2541,28 +2591,73 @@ static VkMemoryRequirements FindOrCalculateTexturePlatformSize(FVulkanDevice* De
 
 uint64 FVulkanDynamicRHI::RHICalcTexture2DPlatformSize(uint32 SizeX, uint32 SizeY, uint8 Format, uint32 NumMips, uint32 NumSamples, ETextureCreateFlags Flags, const FRHIResourceCreateInfo& CreateInfo, uint32& OutAlign)
 {
-	const VkMemoryRequirements MemReq = FindOrCalculateTexturePlatformSize(Device, VK_IMAGE_VIEW_TYPE_2D, SizeX, SizeY, 1, Format, NumMips, NumSamples, Flags);
+	const VkMemoryRequirements MemReq = FindOrCalculateTexturePlatformSize(Device, VK_IMAGE_VIEW_TYPE_2D, SizeX, SizeY, 1, 1, Format, NumMips, NumSamples, Flags);
 	OutAlign = MemReq.alignment;
 	return MemReq.size;
 }
 
 uint64 FVulkanDynamicRHI::RHICalcTexture2DArrayPlatformSize(uint32 SizeX, uint32 SizeY, uint32 ArraySize, uint8 Format, uint32 NumMips, uint32 NumSamples, ETextureCreateFlags Flags, const FRHIResourceCreateInfo& CreateInfo, uint32& OutAlign)
 {
-	const VkMemoryRequirements MemReq = FindOrCalculateTexturePlatformSize(Device, VK_IMAGE_VIEW_TYPE_2D_ARRAY, SizeX, SizeY, ArraySize, Format, NumMips, NumSamples, Flags);
+	const VkMemoryRequirements MemReq = FindOrCalculateTexturePlatformSize(Device, VK_IMAGE_VIEW_TYPE_2D_ARRAY, SizeX, SizeY, 1, ArraySize, Format, NumMips, NumSamples, Flags);
 	OutAlign = MemReq.alignment;
 	return MemReq.size;
 }
 
 uint64 FVulkanDynamicRHI::RHICalcTexture3DPlatformSize(uint32 SizeX, uint32 SizeY, uint32 SizeZ, uint8 Format, uint32 NumMips, ETextureCreateFlags Flags, const FRHIResourceCreateInfo& CreateInfo, uint32& OutAlign)
 {
-	const VkMemoryRequirements MemReq = FindOrCalculateTexturePlatformSize(Device, VK_IMAGE_VIEW_TYPE_3D, SizeX, SizeY, SizeZ, Format, NumMips, 1, Flags);
+	const VkMemoryRequirements MemReq = FindOrCalculateTexturePlatformSize(Device, VK_IMAGE_VIEW_TYPE_3D, SizeX, SizeY, SizeZ, 1, Format, NumMips, 1, Flags);
 	OutAlign = MemReq.alignment;
 	return MemReq.size;
 }
 
 uint64 FVulkanDynamicRHI::RHICalcTextureCubePlatformSize(uint32 Size, uint8 Format, uint32 NumMips, ETextureCreateFlags Flags, const FRHIResourceCreateInfo& CreateInfo, uint32& OutAlign)
 {
-	const VkMemoryRequirements MemReq = FindOrCalculateTexturePlatformSize(Device, VK_IMAGE_VIEW_TYPE_CUBE, Size, Size, 1, Format, NumMips, 1, Flags);
+	const VkMemoryRequirements MemReq = FindOrCalculateTexturePlatformSize(Device, VK_IMAGE_VIEW_TYPE_CUBE, Size, Size, 1, 1, Format, NumMips, 1, Flags);
+	OutAlign = MemReq.alignment;
+	return MemReq.size;
+}
+
+uint64 FVulkanDynamicRHI::RHICalcTexturePlatformSize(const FRHITextureCreateInfo& InCreateInfo, uint32& OutAlign)
+{
+	// Validate some params and determine the view type
+	VkImageViewType ViewType = VK_IMAGE_VIEW_TYPE_2D;
+	switch (InCreateInfo.Dimension)
+	{
+	case ETextureDimension::Texture2D:
+		check(InCreateInfo.ArraySize == 1);
+		check(InCreateInfo.Depth == 1);
+		ViewType = VK_IMAGE_VIEW_TYPE_2D;
+		break;
+
+	case ETextureDimension::Texture2DArray:
+		check(InCreateInfo.Depth == 1);
+		ViewType = VK_IMAGE_VIEW_TYPE_2D_ARRAY;
+		break;
+
+	case ETextureDimension::Texture3D:
+		check(InCreateInfo.ArraySize == 1);
+		check(InCreateInfo.NumSamples == 1);
+		ViewType = VK_IMAGE_VIEW_TYPE_3D;
+		break;
+
+	case ETextureDimension::TextureCube:
+		check(InCreateInfo.ArraySize == 1);
+		check(InCreateInfo.Depth == 1);
+		check(InCreateInfo.NumSamples == 1);
+		ViewType = VK_IMAGE_VIEW_TYPE_CUBE;
+		break;
+
+	case ETextureDimension::TextureCubeArray:
+		check(InCreateInfo.Depth == 1);
+		check(InCreateInfo.NumSamples == 1);
+		ViewType = VK_IMAGE_VIEW_TYPE_CUBE_ARRAY;
+		break;
+
+	default:
+		checkNoEntry();
+	}
+
+	VkMemoryRequirements MemReq = FindOrCalculateTexturePlatformSize(Device, ViewType, InCreateInfo.Extent.X, InCreateInfo.Extent.Y, InCreateInfo.Depth, InCreateInfo.ArraySize, InCreateInfo.Format, InCreateInfo.NumMips, InCreateInfo.NumSamples, InCreateInfo.Flags);
 	OutAlign = MemReq.alignment;
 	return MemReq.size;
 }
diff --git a/Engine/Source/Runtime/VulkanRHI/Private/VulkanTransientResourceAllocator.cpp b/Engine/Source/Runtime/VulkanRHI/Private/VulkanTransientResourceAllocator.cpp
new file mode 100644
--- /dev/null
+++ b/Engine/Source/Runtime/VulkanRHI/Private/VulkanTransientResourceAllocator.cpp
@@ -0,0 +1,147 @@
+// Copyright Epic Games, Inc. All Rights Reserved.
+
+#include "VulkanRHIPrivate.h"
+#include "VulkanTransientResourceAllocator.h"
+
+
+FVulkanTransientHeap::FVulkanTransientHeap(const FInitializer& Initializer, FVulkanDevice* InDevice)
+	: FRHITransientHeap(Initializer)
+	, FDeviceChild(InDevice)
+	, VulkanBuffer(VK_NULL_HANDLE)
+{
+	const VkBufferUsageFlags BufferUsageFlags =
+#if VULKAN_RHI_RAYTRACING
+		VK_BUFFER_USAGE_ACCELERATION_STRUCTURE_STORAGE_BIT_KHR | VK_BUFFER_USAGE_SHADER_DEVICE_ADDRESS_BIT |
+#endif
+		VK_BUFFER_USAGE_VERTEX_BUFFER_BIT | VK_BUFFER_USAGE_INDEX_BUFFER_BIT | VK_BUFFER_USAGE_STORAGE_BUFFER_BIT |
+		VK_BUFFER_USAGE_STORAGE_TEXEL_BUFFER_BIT | VK_BUFFER_USAGE_UNIFORM_TEXEL_BUFFER_BIT |
+		VK_BUFFER_USAGE_INDIRECT_BUFFER_BIT | VK_BUFFER_USAGE_TRANSFER_SRC_BIT | VK_BUFFER_USAGE_TRANSFER_DST_BIT;
+
+	// :TODO: VK_KHR_maintenance4...
+	{
+		VkBufferCreateInfo BufferCreateInfo;
+		ZeroVulkanStruct(BufferCreateInfo, VK_STRUCTURE_TYPE_BUFFER_CREATE_INFO);
+		BufferCreateInfo.size = Initializer.Size;
+		BufferCreateInfo.usage = BufferUsageFlags;
+
+		const VkDevice VulkanDevice = InDevice->GetInstanceHandle();
+
+		VERIFYVULKANRESULT(VulkanRHI::vkCreateBuffer(VulkanDevice, &BufferCreateInfo, VULKAN_CPU_ALLOCATOR, &VulkanBuffer));
+		VulkanRHI::vkGetBufferMemoryRequirements(VulkanDevice, VulkanBuffer, &MemoryRequirements);
+
+		// Find the alignment that works for everyone
+		const uint32 MinBufferAlignment = FMemoryManager::CalculateBufferAlignment(*InDevice, BufferCreateInfo.usage);
+		MemoryRequirements.alignment = FMath::Max<VkDeviceSize>(Initializer.Alignment, MemoryRequirements.alignment);
+		MemoryRequirements.alignment = FMath::Max<VkDeviceSize>(MinBufferAlignment, MemoryRequirements.alignment);
+	}
+
+	VkMemoryPropertyFlags BufferMemFlags = VK_MEMORY_PROPERTY_DEVICE_LOCAL_BIT;
+	const bool bUnifiedMem = InDevice->HasUnifiedMemory();
+	if (bUnifiedMem)
+	{
+		BufferMemFlags |= (VK_MEMORY_PROPERTY_HOST_VISIBLE_BIT | VK_MEMORY_PROPERTY_HOST_COHERENT_BIT);
+	}
+
+	if (!InDevice->GetMemoryManager().AllocateBufferMemory(InternalAllocation, nullptr, MemoryRequirements, BufferMemFlags, EVulkanAllocationMetaBufferOther, false, __FILE__, __LINE__))
+	{
+		InDevice->GetMemoryManager().HandleOOM();
+	}
+
+	InternalAllocation.BindBuffer(InDevice, VulkanBuffer);
+}
+
+FVulkanTransientHeap::~FVulkanTransientHeap()
+{
+	Device->GetMemoryManager().FreeVulkanAllocation(InternalAllocation);
+	Device->GetDeferredDeletionQueue().EnqueueResource(VulkanRHI::FDeferredDeletionQueue2::EType::Buffer, VulkanBuffer);
+	VulkanBuffer = VK_NULL_HANDLE;
+}
+
+VkDeviceMemory FVulkanTransientHeap::GetMemoryHandle()
+{
+	return InternalAllocation.GetDeviceMemoryHandle(Device);
+}
+
+FVulkanAllocation FVulkanTransientHeap::GetVulkanAllocation(const FRHITransientHeapAllocation& HeapAllocation)
+{
+	FVulkanTransientHeap* Heap = static_cast<FVulkanTransientHeap*>(HeapAllocation.Heap);
+	check(Heap);
+
+	FVulkanAllocation TransientAlloc;
+	TransientAlloc.Reference(Heap->InternalAllocation);
+	TransientAlloc.VulkanHandle = (uint64)Heap->VulkanBuffer;
+	TransientAlloc.Offset += HeapAllocation.Offset;
+	TransientAlloc.Size = HeapAllocation.Size;
+	check((TransientAlloc.Offset + TransientAlloc.Size) <= Heap->InternalAllocation.Size);
+	return TransientAlloc;
+}
+
+FVulkanTransientHeapCache* FVulkanTransientHeapCache::Create(FVulkanDevice* InDevice)
+{
+	FRHITransientHeapCache::FInitializer Initializer = FRHITransientHeapCache::FInitializer::CreateDefault();
+
+	// Respect a minimum alignment
+	Initializer.HeapAlignment = FMath::Max((uint32)InDevice->GetLimits().bufferImageGranularity, 256u);
+
+	// Mix resource types onto the same heap.
+	Initializer.bSupportsAllHeapFlags = true;
+
+	return new FVulkanTransientHeapCache(Initializer, InDevice);
+}
+
+FVulkanTransientHeapCache::FVulkanTransientHeapCache(const FRHITransientHeapCache::FInitializer& Initializer, FVulkanDevice* InDevice)
+	: FRHITransientHeapCache(Initializer)
+	, FDeviceChild(InDevice)
+{
+}
+
+FRHITransientHeap* FVulkanTransientHeapCache::CreateHeap(const FRHITransientHeap::FInitializer& HeapInitializer)
+{
+	return new FVulkanTransientHeap(HeapInitializer, Device);
+}
+
+
+FVulkanTransientResourceAllocator::FVulkanTransientResourceAllocator(FVulkanTransientHeapCache& InHeapCache)
+	: FRHITransientResourceHeapAllocator(InHeapCache)
+	, FDeviceChild(InHeapCache.GetParent())
+{
+}
+
+FRHITransientTexture* FVulkanTransientResourceAllocator::CreateTexture(const FRHITextureCreateInfo& InCreateInfo, const TCHAR* InDebugName, uint32 InPassIndex)
+{
+	uint32 ReqAlign = 1;
+	uint64 ReqSize = GVulkanRHI->RHICalcTexturePlatformSize(InCreateInfo, ReqAlign);
+
+	return CreateTextureInternal(InCreateInfo, InDebugName, InPassIndex, ReqSize, ReqAlign,
+		[&](const FRHITransientHeap::FResourceInitializer& Initializer)
+	{
+		ERHIAccess InitialState = ERHIAccess::UAVMask;
+		if (EnumHasAnyFlags(InCreateInfo.Flags, TexCreate_RenderTargetable | TexCreate_ResolveTargetable))
+		{
+			InitialState = ERHIAccess::RTV;
+		}
+		else if (EnumHasAnyFlags(InCreateInfo.Flags, TexCreate_DepthStencilTargetable))
+		{
+			InitialState = ERHIAccess::DSVWrite;
+		}
+
+		FRHIResourceCreateInfo ResourceCreateInfo(InDebugName, InCreateInfo.ClearValue);
+		FRHITexture* Texture = GVulkanRHI->CreateTexture(InCreateInfo, ResourceCreateInfo, InitialState, &Initializer.Allocation);
+		return new FRHITransientTexture(Texture, 0/*GpuVirtualAddress*/, Initializer.Hash, ReqSize, ERHITransientAllocationType::Heap, InCreateInfo);
+	});
+}
+
+FRHITransientBuffer* FVulkanTransientResourceAllocator::CreateBuffer(const FRHIBufferCreateInfo& InCreateInfo, const TCHAR* InDebugName, uint32 InPassIndex)
+{
+	const VkBufferUsageFlags VulkanBufferUsage = FVulkanResourceMultiBuffer::UEToVKBufferUsageFlags(Device, InCreateInfo.Usage, (InCreateInfo.Size == 0));
+	const uint32 Alignment = FMemoryManager::CalculateBufferAlignment(*Device, VulkanBufferUsage);
+	uint64 Size = Align(InCreateInfo.Size, Alignment) * FVulkanResourceMultiBuffer::GetNumBuffersFromUsage(InCreateInfo.Usage);
+
+	return CreateBufferInternal(InCreateInfo, InDebugName, InPassIndex, Size, Alignment,
+		[&](const FRHITransientHeap::FResourceInitializer& Initializer)
+	{
+		FRHIResourceCreateInfo ResourceCreateInfo(InDebugName);
+		FRHIBuffer* Buffer = GVulkanRHI->CreateBuffer(InCreateInfo, ResourceCreateInfo, &Initializer.Allocation);
+		return new FRHITransientBuffer(Buffer, 0/*GpuVirtualAddress*/, Initializer.Hash, Size, ERHITransientAllocationType::Heap, InCreateInfo);
+	});
+}
diff --git a/Engine/Source/Runtime/VulkanRHI/Private/VulkanTransientResourceAllocator.h b/Engine/Source/Runtime/VulkanRHI/Private/VulkanTransientResourceAllocator.h
new file mode 100644
--- /dev/null
+++ b/Engine/Source/Runtime/VulkanRHI/Private/VulkanTransientResourceAllocator.h
@@ -0,0 +1,52 @@
+// Copyright Epic Games, Inc. All Rights Reserved.
+
+#pragma once
+
+#include "VulkanResources.h"
+#include "RHICoreTransientResourceAllocator.h"
+
+
+class FVulkanTransientHeap final
+	: public FRHITransientHeap
+	, public FRefCountBase
+	, public FDeviceChild
+{
+public:
+	FVulkanTransientHeap(const FInitializer& Initializer, FVulkanDevice* InDevice);
+	~FVulkanTransientHeap();
+
+	VkDeviceMemory GetMemoryHandle();
+	static FVulkanAllocation GetVulkanAllocation(const FRHITransientHeapAllocation& TransientInitializer);
+
+private:
+	VkBuffer VulkanBuffer;
+	VkMemoryRequirements MemoryRequirements;
+	FVulkanAllocation InternalAllocation;
+};
+
+
+class FVulkanTransientHeapCache final
+	: public FRHITransientHeapCache
+	, public FDeviceChild
+{
+public:
+	static FVulkanTransientHeapCache* Create(FVulkanDevice* InDevice);
+
+	FRHITransientHeap* CreateHeap(const FRHITransientHeap::FInitializer& Initializer) override;
+
+private:
+	FVulkanTransientHeapCache(const FInitializer& Initializer, FVulkanDevice* InDevice);
+};
+
+
+class FVulkanTransientResourceAllocator final
+	: public FRHITransientResourceHeapAllocator
+	, public FDeviceChild
+{
+public:
+	FVulkanTransientResourceAllocator(FVulkanTransientHeapCache& InHeapCache);
+
+	//! IRHITransientResourceAllocator Overrides
+	FRHITransientTexture* CreateTexture(const FRHITextureCreateInfo& InCreateInfo, const TCHAR* InDebugName, uint32 InPassIndex) override;
+	FRHITransientBuffer* CreateBuffer(const FRHIBufferCreateInfo& InCreateInfo, const TCHAR* InDebugName, uint32 InPassIndex) override;
+};
\ No newline at end of file
diff --git a/Engine/Source/Runtime/VulkanRHI/Public/VulkanDynamicRHI.h b/Engine/Source/Runtime/VulkanRHI/Public/VulkanDynamicRHI.h
--- a/Engine/Source/Runtime/VulkanRHI/Public/VulkanDynamicRHI.h
+++ b/Engine/Source/Runtime/VulkanRHI/Public/VulkanDynamicRHI.h
@@ -11,6 +11,7 @@ class FVulkanFramebuffer;
 class FVulkanDevice;
 class FVulkanQueue;
 class IHeadMountedDisplayVulkanExtensions;
+struct FRHITransientHeapAllocation;
 
 // TODO: fix Lock/Unlock Vertex/Index buffer
 #define VULKAN_BUFFER_LOCK_THREADSAFE 0
@@ -180,6 +181,12 @@ public:
 	virtual class IRHICommandContextContainer* RHIGetCommandContextContainer(int32 Index, int32 Num) final override;
 	virtual uint64 RHIGetMinimumAlignmentForBufferBackedSRV(EPixelFormat Format) final override;
 
+	// Transient Resource Functions and Helpers
+	virtual IRHITransientResourceAllocator* RHICreateTransientResourceAllocator() final override;
+	uint64 RHICalcTexturePlatformSize(const FRHITextureCreateInfo& InCreateInfo, uint32& OutAlign);
+	FRHITexture* CreateTexture(const FRHITextureCreateInfo& InCreateInfo, FRHIResourceCreateInfo& InResourceCreateInfo, ERHIAccess InResourceState, const FRHITransientHeapAllocation* InAllocation);
+	FRHIBuffer* CreateBuffer(const FRHIBufferCreateInfo & InCreateInfo, FRHIResourceCreateInfo & InResourceCreateInfo, const FRHITransientHeapAllocation* InAllocation);
+
 	virtual FBufferRHIRef CreateBuffer_RenderThread(class FRHICommandListImmediate& RHICmdList, uint32 Size, EBufferUsageFlags Usage, uint32 Stride, ERHIAccess ResourceState, FRHIResourceCreateInfo& CreateInfo) override final
 	{
 		// this-> is required to avoid calling the global version of this function
