Add handling of ray tracing buffer usage flags to FMemoryManager::CalculateBufferPriority. Fixes crash at startup when ray tracing is enabled for Vulkan.

diff --git a/Engine/Source/Runtime/VulkanRHI/Private/VulkanMemory.cpp b/Engine/Source/Runtime/VulkanRHI/Private/VulkanMemory.cpp
--- a/Engine/Source/Runtime/VulkanRHI/Private/VulkanMemory.cpp
+++ b/Engine/Source/Runtime/VulkanRHI/Private/VulkanMemory.cpp
@@ -2768,6 +2768,15 @@ namespace VulkanRHI
 		{
 			Priority = VULKAN_MEMORY_MEDIUM_PRIORITY;
 		}
+#if VULKAN_RHI_RAYTRACING
+		else if (VKHasAnyFlags(BufferUsageFlags, 
+			(VK_BUFFER_USAGE_ACCELERATION_STRUCTURE_STORAGE_BIT_KHR |
+			VK_BUFFER_USAGE_SHADER_BINDING_TABLE_BIT_KHR |
+			VK_BUFFER_USAGE_ACCELERATION_STRUCTURE_BUILD_INPUT_READ_ONLY_BIT_KHR)))
+		{
+			Priority = VULKAN_MEMORY_MEDIUM_PRIORITY;
+		}
+#endif // VULKAN_RHI_RAYTRACING
 		else if (VKHasAnyFlags(BufferUsageFlags, (VK_BUFFER_USAGE_VERTEX_BUFFER_BIT | VK_BUFFER_USAGE_INDEX_BUFFER_BIT)))
 		{
 			Priority = VULKAN_MEMORY_LOW_PRIORITY;
