Enable the Transient Resource Allocator by default on desktop Vulkan platforms.

diff --git a/Engine/Source/Runtime/VulkanRHI/Private/VulkanRHI.cpp b/Engine/Source/Runtime/VulkanRHI/Private/VulkanRHI.cpp
--- a/Engine/Source/Runtime/VulkanRHI/Private/VulkanRHI.cpp
+++ b/Engine/Source/Runtime/VulkanRHI/Private/VulkanRHI.cpp
@@ -78,7 +78,7 @@ static FAutoConsoleVariableRef GCVarInputAttachmentShaderRead(
 	ECVF_ReadOnly
 );
 
-int32 GVulkanEnableTransientResourceAllocator = 0;
+int32 GVulkanEnableTransientResourceAllocator = 1;
 static FAutoConsoleVariableRef GCVarEnableTransientResourceAllocator(
 	TEXT("r.Vulkan.EnableTransientResourceAllocator"),
 	GVulkanEnableTransientResourceAllocator,
@@ -1791,7 +1791,7 @@ void FVulkanDynamicRHI::VulkanSetImageLayout( VkCommandBuffer CmdBuffer, VkImage
 
 IRHITransientResourceAllocator* FVulkanDynamicRHI::RHICreateTransientResourceAllocator()
 {
-	if (GVulkanEnableTransientResourceAllocator)
+	if (GVulkanEnableTransientResourceAllocator && IsPCPlatform(GMaxRHIShaderPlatform))
 	{
 		return new FVulkanTransientResourceAllocator(Device->GetOrCreateTransientHeapCache());
 	}
