Add support for run-time detection of vertex shader layer suport to Vulkan RHI (GRHISupportsArrayIndexFromAnyShader)

diff --git a/Engine/Source/Runtime/VulkanRHI/Private/VulkanDevice.h b/Engine/Source/Runtime/VulkanRHI/Private/VulkanDevice.h
--- a/Engine/Source/Runtime/VulkanRHI/Private/VulkanDevice.h
+++ b/Engine/Source/Runtime/VulkanRHI/Private/VulkanDevice.h
@@ -37,41 +37,41 @@ struct FOptionalVulkanDeviceExtensions
 	{
 		struct
 		{
-			uint32 HasKHRMaintenance1 : 1;
-			uint32 HasKHRMaintenance2 : 1;
-			//uint32 HasMirrorClampToEdge : 1;
-			uint32 HasKHRDedicatedAllocation : 1;
-			uint32 HasEXTValidationCache : 1;
-			uint32 HasAMDBufferMarker : 1;
-			uint32 HasNVDiagnosticCheckpoints : 1;
-			uint32 HasNVDeviceDiagnosticConfig : 1;
-			uint32 HasYcbcrSampler : 1;
-			uint32 HasMemoryPriority : 1;
-			uint32 HasMemoryBudget : 1;
-			uint32 HasDriverProperties : 1;
-			uint32 HasEXTFragmentDensityMap : 1;
-			uint32 HasEXTFragmentDensityMap2 : 1;
-			uint32 HasKHRFragmentShadingRate : 1;
-			uint32 HasKHRRenderPass2 : 1;
-			uint32 HasEXTFullscreenExclusive : 1;
-			uint32 HasKHRImageFormatList : 1;
-			uint32 HasEXTTextureCompressionASTCHDR : 1;
-			uint32 HasEXTASTCDecodeMode : 1;
-			uint32 HasQcomRenderPassTransform : 1;
-			uint32 HasImageAtomicInt64 : 1;
-			uint32 HasBufferAtomicInt64 : 1;
-			uint32 HasScalarBlockLayoutFeatures : 1;
-			uint32 HasKHRMultiview : 1;
-			uint32 HasAccelerationStructure : 1;
-			uint32 HasRayTracingPipeline : 1;
-			uint32 HasRayQuery : 1;
-			uint32 HasDescriptorIndexing : 1;
-			uint32 HasBufferDeviceAddress : 1;
-			uint32 HasDeferredHostOperations : 1;
-			uint32 HasSPIRV_14 : 1;
-			uint32 HasShaderFloatControls : 1;
+			uint64 HasKHRMaintenance1 : 1;
+			uint64 HasKHRMaintenance2 : 1;
+			uint64 HasKHRDedicatedAllocation : 1;
+			uint64 HasEXTValidationCache : 1;
+			uint64 HasAMDBufferMarker : 1;
+			uint64 HasNVDiagnosticCheckpoints : 1;
+			uint64 HasNVDeviceDiagnosticConfig : 1;
+			uint64 HasYcbcrSampler : 1;
+			uint64 HasMemoryPriority : 1;
+			uint64 HasMemoryBudget : 1;
+			uint64 HasDriverProperties : 1;
+			uint64 HasEXTFragmentDensityMap : 1;
+			uint64 HasEXTFragmentDensityMap2 : 1;
+			uint64 HasKHRFragmentShadingRate : 1;
+			uint64 HasKHRRenderPass2 : 1;
+			uint64 HasEXTFullscreenExclusive : 1;
+			uint64 HasKHRImageFormatList : 1;
+			uint64 HasEXTTextureCompressionASTCHDR : 1;
+			uint64 HasEXTASTCDecodeMode : 1;
+			uint64 HasQcomRenderPassTransform : 1;
+			uint64 HasImageAtomicInt64 : 1;
+			uint64 HasBufferAtomicInt64 : 1;
+			uint64 HasScalarBlockLayoutFeatures : 1;
+			uint64 HasKHRMultiview : 1;
+			uint64 HasAccelerationStructure : 1;
+			uint64 HasRayTracingPipeline : 1;
+			uint64 HasRayQuery : 1;
+			uint64 HasDescriptorIndexing : 1;
+			uint64 HasBufferDeviceAddress : 1;
+			uint64 HasDeferredHostOperations : 1;
+			uint64 HasSPIRV_14 : 1;
+			uint64 HasShaderFloatControls : 1;
+			uint64 HasEXTShaderViewportIndexLayer : 1;
 		};
-		uint32 Packed;
+		uint64 Packed;
 	};
 
 	FOptionalVulkanDeviceExtensions()
diff --git a/Engine/Source/Runtime/VulkanRHI/Private/VulkanLayers.cpp b/Engine/Source/Runtime/VulkanRHI/Private/VulkanLayers.cpp
--- a/Engine/Source/Runtime/VulkanRHI/Private/VulkanLayers.cpp
+++ b/Engine/Source/Runtime/VulkanRHI/Private/VulkanLayers.cpp
@@ -127,6 +127,10 @@ static const ANSICHAR* GDeviceExtensions[] =
 
 	VK_KHR_IMAGE_FORMAT_LIST_EXTENSION_NAME,
 
+#if VULKAN_SUPPORTS_SHADER_VIEWPORT_INDEX_LAYER
+	VK_EXT_SHADER_VIEWPORT_INDEX_LAYER_EXTENSION_NAME,
+#endif 
+
 	nullptr
 };
 
@@ -909,6 +913,10 @@ void FOptionalVulkanDeviceExtensions::Setup(const TArray<const ANSICHAR*>& Devic
 	HasSPIRV_14					= HasExtension(DeviceExtensions, VK_KHR_SPIRV_1_4_EXTENSION_NAME);
 	HasShaderFloatControls		= HasExtension(DeviceExtensions, VK_KHR_SHADER_FLOAT_CONTROLS_EXTENSION_NAME);
 #endif
+
+#if VULKAN_SUPPORTS_SHADER_VIEWPORT_INDEX_LAYER
+	HasEXTShaderViewportIndexLayer = HasExtension(DeviceExtensions, VK_EXT_SHADER_VIEWPORT_INDEX_LAYER_EXTENSION_NAME);
+#endif
 }
 
 void FVulkanDynamicRHI::SetupValidationRequests()
diff --git a/Engine/Source/Runtime/VulkanRHI/Private/VulkanRHI.cpp b/Engine/Source/Runtime/VulkanRHI/Private/VulkanRHI.cpp
--- a/Engine/Source/Runtime/VulkanRHI/Private/VulkanRHI.cpp
+++ b/Engine/Source/Runtime/VulkanRHI/Private/VulkanRHI.cpp
@@ -840,6 +840,10 @@ void FVulkanDynamicRHI::InitInstance()
 		GRHIMaxDispatchThreadGroupsPerDimension.Y = FMath::Min<uint32>(Limits.maxComputeWorkGroupCount[1], 0x7fffffff);
 		GRHIMaxDispatchThreadGroupsPerDimension.Z = FMath::Min<uint32>(Limits.maxComputeWorkGroupCount[2], 0x7fffffff);
 
+#if VULKAN_SUPPORTS_SHADER_VIEWPORT_INDEX_LAYER
+		GRHISupportsArrayIndexFromAnyShader = GetDevice()->GetOptionalExtensions().HasEXTShaderViewportIndexLayer != 0ULL;
+#endif
+
 #if VULKAN_SUPPORTS_FRAGMENT_DENSITY_MAP
 		GRHISupportsAttachmentVariableRateShading = (GetDevice()->GetOptionalExtensions().HasEXTFragmentDensityMap && Device->GetFragmentDensityMapFeatures().fragmentDensityMap);
 #endif
diff --git a/Engine/Source/Runtime/VulkanRHI/Public/VulkanConfiguration.h b/Engine/Source/Runtime/VulkanRHI/Public/VulkanConfiguration.h
--- a/Engine/Source/Runtime/VulkanRHI/Public/VulkanConfiguration.h
+++ b/Engine/Source/Runtime/VulkanRHI/Public/VulkanConfiguration.h
@@ -331,6 +331,13 @@
 	#define VULKAN_SUPPORTS_IMAGE_64BIT_ATOMICS		0
 #endif
 
+#ifdef VK_EXT_shader_viewport_index_layer
+	#define VULKAN_SUPPORTS_SHADER_VIEWPORT_INDEX_LAYER	1
+#else
+	#define VULKAN_SUPPORTS_SHADER_VIEWPORT_INDEX_LAYER	0
+#endif
+
+
 #ifndef VULKAN_OBJECT_TRACKING 
 #define VULKAN_OBJECT_TRACKING 0 //Track objects created and memory used. use r.vulkan.dumpmemory to dump to console
 #endif
