Fix incorrect checks in RayTracingGeometry:
* Add InitializedState to RayTracingGeometry that will track how it is was crated. We cannot use provided Initializer for that because it's reset during mesh streaming.
* Use BUF_RayTracingScratch for dynamic geometry scratch buffer so we can always guarantee correct alignment.
* Add RayTracingGeometry empty ctor so we don't try to calculate the BLAS size when it is streamed out.

diff --git a/Engine/Source/Runtime/D3D12RHI/Private/D3D12RayTracing.cpp b/Engine/Source/Runtime/D3D12RHI/Private/D3D12RayTracing.cpp
--- a/Engine/Source/Runtime/D3D12RHI/Private/D3D12RayTracing.cpp
+++ b/Engine/Source/Runtime/D3D12RHI/Private/D3D12RayTracing.cpp
@@ -3617,7 +3617,6 @@ void FD3D12RayTracingGeometry::Swap(FD3D12RayTracingGeometry& Other)
 		bIsAccelerationStructureDirty[i] = Other.bIsAccelerationStructureDirty[i];
 	}
 	::Swap(AccelerationStructureCompactedSize, Other.AccelerationStructureCompactedSize);
-	::Swap(SizeInfo, Other.SizeInfo);
 
 	// The rest of the members should be updated using SetInitializer()
 }
@@ -3674,13 +3673,12 @@ void FD3D12RayTracingGeometry::ReleaseUnderlyingResource()
 	for (TRefCountPtr<FD3D12Buffer>& Buffer : AccelerationStructureBuffers)
 	{
 		Buffer.SafeRelease();
-	}	
+	}
 
 	Initializer = {};
 
 	AccelerationStructureCompactedSize = 0;
 	GeometryDescs = {};
-	SizeInfo = {};
 	for (TArray<FHitGroupSystemParameters>& HitGroupParametersForGPU : HitGroupSystemParameters)
 	{
 		HitGroupParametersForGPU.Empty();
@@ -3826,9 +3824,8 @@ void FD3D12RayTracingGeometry::UpdateResidency(FD3D12CommandContext& CommandCont
 
 void FD3D12RayTracingGeometry::SetInitializer(const FRayTracingGeometryInitializer& InInitializer)
 {
-	checkf(Initializer.Type == ERayTracingGeometryInitializerType::StreamingDestination, TEXT("Only FD3D12RayTracingGeometry that was created as StreamingDestination can update their initializer."));
+	checkf(InitializedType == ERayTracingGeometryInitializerType::StreamingDestination, TEXT("Only FD3D12RayTracingGeometry that was created as StreamingDestination can update their initializer."));
 	Initializer = InInitializer;
-	Initializer.Type = ERayTracingGeometryInitializerType::StreamingDestination; // Make sure this flag is set at all times so we can verify the resource state
 
 	for (uint32 GPUIndex = 0; GPUIndex < MAX_NUM_GPUS && GPUIndex < GNumExplicitGPUsForRendering; ++GPUIndex)
 	{
diff --git a/Engine/Source/Runtime/RHI/Public/RHIResources.h b/Engine/Source/Runtime/RHI/Public/RHIResources.h
--- a/Engine/Source/Runtime/RHI/Public/RHIResources.h
+++ b/Engine/Source/Runtime/RHI/Public/RHIResources.h
@@ -2170,6 +2170,7 @@ public:
 	FRHIRayTracingGeometry() = default;
 	FRHIRayTracingGeometry(const FRayTracingGeometryInitializer& InInitializer)
 		: Initializer(InInitializer)
+		, InitializedType(InInitializer.Type)
 	{}
 
 	virtual FRayTracingAccelerationStructureAddress GetAccelerationStructureAddress(uint64 GPUIndex) const = 0;
@@ -2192,6 +2193,7 @@ public:
 protected:
 	FRayTracingAccelerationStructureSize SizeInfo = {};
 	FRayTracingGeometryInitializer Initializer = {};
+	ERayTracingGeometryInitializerType InitializedType = ERayTracingGeometryInitializerType::Rendering;
 };
 
 typedef TRefCountPtr<FRHIRayTracingGeometry>     FRayTracingGeometryRHIRef;
diff --git a/Engine/Source/Runtime/RenderCore/Private/RenderResource.cpp b/Engine/Source/Runtime/RenderCore/Private/RenderResource.cpp
--- a/Engine/Source/Runtime/RenderCore/Private/RenderResource.cpp
+++ b/Engine/Source/Runtime/RenderCore/Private/RenderResource.cpp
@@ -437,7 +437,7 @@ void FRayTracingGeometry::CreateRayTracingGeometry(ERTAccelerationStructureBuild
 	}
 
 	const bool bWithoutNativeResource = Initializer.Type == ERayTracingGeometryInitializerType::StreamingDestination;
-	if (bAllSegmentsAreValid || bWithoutNativeResource)
+	if (bAllSegmentsAreValid)
 	{
 		bValid = !bWithoutNativeResource;
 		RayTracingGeometryRHI = RHICreateRayTracingGeometry(Initializer);
diff --git a/Engine/Source/Runtime/Renderer/Private/DeferredShadingRenderer.cpp b/Engine/Source/Runtime/Renderer/Private/DeferredShadingRenderer.cpp
--- a/Engine/Source/Runtime/Renderer/Private/DeferredShadingRenderer.cpp
+++ b/Engine/Source/Runtime/Renderer/Private/DeferredShadingRenderer.cpp
@@ -1521,7 +1521,7 @@ bool FDeferredShadingSceneRenderer::DispatchRayTracingWorldUpdates(FRDGBuilder&
 		const uint32 ScratchAlignment = GRHIRayTracingAccelerationStructureAlignment;
 		FRDGBufferDesc ScratchBufferDesc;
 		ScratchBufferDesc.UnderlyingType = FRDGBufferDesc::EUnderlyingType::StructuredBuffer;
-		ScratchBufferDesc.Usage = BUF_UnorderedAccess;
+		ScratchBufferDesc.Usage = BUF_RayTracingScratch;
 		ScratchBufferDesc.BytesPerElement = ScratchAlignment;
 		ScratchBufferDesc.NumElements = FMath::DivideAndRoundUp(BLASScratchSize, ScratchAlignment);
 
diff --git a/Engine/Source/Runtime/VulkanRHI/Private/VulkanRayTracing.cpp b/Engine/Source/Runtime/VulkanRHI/Private/VulkanRayTracing.cpp
--- a/Engine/Source/Runtime/VulkanRHI/Private/VulkanRayTracing.cpp
+++ b/Engine/Source/Runtime/VulkanRHI/Private/VulkanRayTracing.cpp
@@ -287,12 +287,25 @@ static void GetBLASBuildData(
 		&BuildData.SizesInfo);
 }
 
+FVulkanRayTracingGeometry::FVulkanRayTracingGeometry(ENoInit)
+{}
+
 FVulkanRayTracingGeometry::FVulkanRayTracingGeometry(const FRayTracingGeometryInitializer& InInitializer, FVulkanDevice* InDevice)
 	: FRHIRayTracingGeometry(InInitializer), Device(InDevice)
 {
 	// Only supporting triangles initially
 	check(Initializer.GeometryType == ERayTracingGeometryType::RTGT_Triangles);
-	checkf(!Initializer.IndexBuffer || (Initializer.IndexBuffer->GetStride() == 2 || Initializer.IndexBuffer->GetStride() == 4), TEXT("Index buffer must be 16 or 32 bit if in use."));
+
+	uint32 IndexBufferStride = 0;
+	if (Initializer.IndexBuffer)
+	{
+		// In case index buffer in initializer is not yet in valid state during streaming we assume the geometry is using UINT32 format.
+		IndexBufferStride = Initializer.IndexBuffer->GetSize() > 0
+			? Initializer.IndexBuffer->GetStride()
+			: 4;
+	}
+
+	checkf(!Initializer.IndexBuffer || (IndexBufferStride == 2 || IndexBufferStride == 4), TEXT("Index buffer must be 16 or 32 bit if in use."));
 
 	FVkRtBLASBuildData BuildData;
 	GetBLASBuildData(
@@ -302,7 +315,7 @@ FVulkanRayTracingGeometry::FVulkanRayTracingGeometry(const FRayTracingGeometryIn
 		Initializer.IndexBufferOffset,
 		Initializer.bFastBuild,
 		Initializer.bAllowUpdate,
-		Initializer.IndexBuffer ? Initializer.IndexBuffer->GetStride() : 0,
+		IndexBufferStride,
 		EAccelerationStructureBuildMode::Build,
 		BuildData);
 
@@ -345,7 +358,7 @@ FVulkanRayTracingGeometry::~FVulkanRayTracingGeometry()
 
 void FVulkanRayTracingGeometry::SetInitializer(const FRayTracingGeometryInitializer& InInitializer)
 {
-	checkf(Initializer.Type == ERayTracingGeometryInitializerType::StreamingDestination, TEXT("Only FVulkanRayTracingGeometry that was created as StreamingDestination can update their initializer."));
+	checkf(InitializedType == ERayTracingGeometryInitializerType::StreamingDestination, TEXT("Only FVulkanRayTracingGeometry that was created as StreamingDestination can update their initializer."));
 	Initializer = InInitializer;
 
 	// TODO: Update HitGroup Parameters
@@ -356,9 +369,10 @@ void FVulkanRayTracingGeometry::Swap(FVulkanRayTracingGeometry& Other)
 	::Swap(Handle, Other.Handle);
 	::Swap(Address, Other.Address);
 
-	Initializer = Other.Initializer;
 	AccelerationStructureBuffer = Other.AccelerationStructureBuffer;
 	ScratchBuffer = Other.ScratchBuffer;
+
+	// The rest of the members should be updated using SetInitializer()
 }
 
 void FVulkanRayTracingGeometry::BuildAccelerationStructure(FVulkanCommandListContext& CommandContext, EAccelerationStructureBuildMode BuildMode)
@@ -631,11 +645,8 @@ void FVulkanDynamicRHI::RHITransferRayTracingGeometryUnderlyingResource(FRHIRayT
 	check(DestGeometry);
 	FVulkanRayTracingGeometry* Dest = ResourceCast(DestGeometry);
 	if (!SrcGeometry)
-	{
-		FRayTracingGeometryInitializer DummyInitializer = {};
-		DummyInitializer.Type = ERayTracingGeometryInitializerType::StreamingDestination;
-		DummyInitializer.DebugName = FName(TEXT("RHITransferRayTracingGeometryUnderlyingResource"));
-		TRefCountPtr<FVulkanRayTracingGeometry> DeletionProxy = new FVulkanRayTracingGeometry(DummyInitializer, GetDevice());
+	{		
+		TRefCountPtr<FVulkanRayTracingGeometry> DeletionProxy = new FVulkanRayTracingGeometry(NoInit);
 		Dest->Swap(*DeletionProxy);
 	}
 	else
diff --git a/Engine/Source/Runtime/VulkanRHI/Private/VulkanRayTracing.h b/Engine/Source/Runtime/VulkanRHI/Private/VulkanRayTracing.h
--- a/Engine/Source/Runtime/VulkanRHI/Private/VulkanRayTracing.h
+++ b/Engine/Source/Runtime/VulkanRHI/Private/VulkanRayTracing.h
@@ -80,6 +80,7 @@ struct FVkRtBLASBuildData
 class FVulkanRayTracingGeometry : public FRHIRayTracingGeometry
 {
 public:
+	FVulkanRayTracingGeometry(ENoInit);
 	FVulkanRayTracingGeometry(const FRayTracingGeometryInitializer& Initializer, FVulkanDevice* InDevice);
 	~FVulkanRayTracingGeometry();
 
