Do not use TransientHeap along with UseVulkanDescriptorCache because of how it generates its hashes (this is essentially the same as not using it on mobile, however I forgot to exclude the mobile preview on PC that uses Vulkan with ERHIFeatureLevel::ES3_1)...

diff --git a/Engine/Source/Runtime/VulkanRHI/Private/VulkanRHI.cpp b/Engine/Source/Runtime/VulkanRHI/Private/VulkanRHI.cpp
--- a/Engine/Source/Runtime/VulkanRHI/Private/VulkanRHI.cpp
+++ b/Engine/Source/Runtime/VulkanRHI/Private/VulkanRHI.cpp
@@ -1795,7 +1795,9 @@ void FVulkanDynamicRHI::VulkanSetImageLayout( VkCommandBuffer CmdBuffer, VkImage
 
 IRHITransientResourceAllocator* FVulkanDynamicRHI::RHICreateTransientResourceAllocator()
 {
-	if (GVulkanEnableTransientResourceAllocator && IsPCPlatform(GMaxRHIShaderPlatform))
+	// Only use transient heap on desktop platforms for now
+	// Not compatible with VulkanDescriptorCache for now because it hashes using the 32bit BufferId instead of the VulkanHandle.
+	if (GVulkanEnableTransientResourceAllocator && IsPCPlatform(GMaxRHIShaderPlatform) && !UseVulkanDescriptorCache())
 	{
 		return new FVulkanTransientResourceAllocator(Device->GetOrCreateTransientHeapCache());
 	}
