VulkanRHI: fix assert condition in Debug builds (aliasing barriers won't match what is held in layout tracking).

diff --git a/Engine/Source/Runtime/VulkanRHI/Private/VulkanBarriers.cpp b/Engine/Source/Runtime/VulkanRHI/Private/VulkanBarriers.cpp
--- a/Engine/Source/Runtime/VulkanRHI/Private/VulkanBarriers.cpp
+++ b/Engine/Source/Runtime/VulkanRHI/Private/VulkanBarriers.cpp
@@ -1281,7 +1281,7 @@ void FVulkanCommandListContext::RHIEndTransitions(TArrayView<const FRHITransitio
 			}
 			else
 			{
-				checkSlow(Layout.AreSubresourcesSameLayout(ImageBarrier.oldLayout, ImageBarrier.subresourceRange));
+				checkSlow(ExtraData.IsAliasingBarrier || Layout.AreSubresourcesSameLayout(ImageBarrier.oldLayout, ImageBarrier.subresourceRange));
 				SrcLayout = ImageBarrier.oldLayout;
 				SrcAccessFlags = ImageBarrier.srcAccessMask;
 			}
