VulkanRHI: Add define to entirely remove transient resource allocator on a per-platforms basis.

#rnx

diff --git a/Engine/Source/Runtime/VulkanRHI/Private/Android/VulkanAndroidPlatform.h b/Engine/Source/Runtime/VulkanRHI/Private/Android/VulkanAndroidPlatform.h
--- a/Engine/Source/Runtime/VulkanRHI/Private/Android/VulkanAndroidPlatform.h
+++ b/Engine/Source/Runtime/VulkanRHI/Private/Android/VulkanAndroidPlatform.h
@@ -22,6 +22,7 @@
 #define VULKAN_SUPPORTS_ASTC_DECODE_MODE			(VULKAN_SUPPORTS_PHYSICAL_DEVICE_PROPERTIES2)
 #define VULKAN_SUPPORTS_NV_DIAGNOSTIC_CHECKPOINT	0
 #define VULKAN_RHI_RAYTRACING						0
+#define VULKAN_SUPPORTS_TRANSIENT_RESOURCE_ALLOCATOR 0
 
 // crashing during callback setup on Android, code will fallback to VK_EXT_debug_report instead
 #define VULKAN_SUPPORTS_DEBUG_UTILS					0
diff --git a/Engine/Source/Runtime/VulkanRHI/Private/VulkanRHI.cpp b/Engine/Source/Runtime/VulkanRHI/Private/VulkanRHI.cpp
--- a/Engine/Source/Runtime/VulkanRHI/Private/VulkanRHI.cpp
+++ b/Engine/Source/Runtime/VulkanRHI/Private/VulkanRHI.cpp
@@ -1769,12 +1769,14 @@ void FVulkanDynamicRHI::VulkanSetImageLayout( VkCommandBuffer CmdBuffer, VkImage
 
 IRHITransientResourceAllocator* FVulkanDynamicRHI::RHICreateTransientResourceAllocator()
 {
+#if VULKAN_SUPPORTS_TRANSIENT_RESOURCE_ALLOCATOR
 	// Only use transient heap on desktop platforms for now
 	// Not compatible with VulkanDescriptorCache for now because it hashes using the 32bit BufferId instead of the VulkanHandle.
 	if (GVulkanEnableTransientResourceAllocator && IsPCPlatform(GMaxRHIShaderPlatform) && !UseVulkanDescriptorCache())
 	{
 		return new FVulkanTransientResourceAllocator(Device->GetOrCreateTransientHeapCache());
 	}
+#endif
 	return nullptr;
 }
 
diff --git a/Engine/Source/Runtime/VulkanRHI/Public/VulkanConfiguration.h b/Engine/Source/Runtime/VulkanRHI/Public/VulkanConfiguration.h
--- a/Engine/Source/Runtime/VulkanRHI/Public/VulkanConfiguration.h
+++ b/Engine/Source/Runtime/VulkanRHI/Public/VulkanConfiguration.h
@@ -354,6 +354,11 @@ VULKANRHI_API DECLARE_LOG_CATEGORY_EXTERN(LogVulkanRHI, Log, All);
 	#define VULKAN_PURGE_SHADER_MODULES							0
 #endif
 
+#ifndef VULKAN_SUPPORTS_TRANSIENT_RESOURCE_ALLOCATOR
+	#define VULKAN_SUPPORTS_TRANSIENT_RESOURCE_ALLOCATOR		1
+#endif
+
+
 #if !defined(NV_AFTERMATH)
 	#define NV_AFTERMATH 0
 #endif
