VulkanRHI: Fix validation error where a view was created beyond it's number of slices when the first slice index was offset.

diff --git a/Engine/Source/Runtime/VulkanRHI/Private/VulkanViewport.cpp b/Engine/Source/Runtime/VulkanRHI/Private/VulkanViewport.cpp
--- a/Engine/Source/Runtime/VulkanRHI/Private/VulkanViewport.cpp
+++ b/Engine/Source/Runtime/VulkanRHI/Private/VulkanViewport.cpp
@@ -338,7 +338,10 @@ FVulkanFramebuffer::FVulkanFramebuffer(FVulkanDevice& Device, const FRHISetRende
 		FVulkanTextureView RTView;
 		if (Texture->Surface.GetViewType() == VK_IMAGE_VIEW_TYPE_2D || Texture->Surface.GetViewType() == VK_IMAGE_VIEW_TYPE_2D_ARRAY)
 		{
-			RTView.Create(*Texture->Surface.Device, Texture->Surface.Image, Texture->Surface.GetViewType(), Texture->Surface.GetFullAspectMask(), Texture->Surface.PixelFormat, Texture->Surface.ViewFormat, MipIndex, 1, FMath::Max(0, (int32)InRTInfo.ColorRenderTarget[Index].ArraySliceIndex), Texture->Surface.GetNumberOfArrayLevels(), true);
+			const uint32 ArraySliceIndex = FMath::Max(0, (int32)InRTInfo.ColorRenderTarget[Index].ArraySliceIndex);
+			check(ArraySliceIndex < Texture->Surface.GetNumberOfArrayLevels());
+			const uint32 NumArraySlices = Texture->Surface.GetNumberOfArrayLevels() - ArraySliceIndex;
+			RTView.Create(*Texture->Surface.Device, Texture->Surface.Image, Texture->Surface.GetViewType(), Texture->Surface.GetFullAspectMask(), Texture->Surface.PixelFormat, Texture->Surface.ViewFormat, MipIndex, 1, ArraySliceIndex, NumArraySlices, true);
 		}
 		else if (Texture->Surface.GetViewType() == VK_IMAGE_VIEW_TYPE_CUBE)
 		{
