Make sure all VK allocations have a valid HandleId on platforms that use 'Descriptor Cache'

diff --git a/Engine/Source/Runtime/VulkanRHI/Private/VulkanDescriptorSets.h b/Engine/Source/Runtime/VulkanRHI/Private/VulkanDescriptorSets.h
--- a/Engine/Source/Runtime/VulkanRHI/Private/VulkanDescriptorSets.h
+++ b/Engine/Source/Runtime/VulkanRHI/Private/VulkanDescriptorSets.h
@@ -1155,11 +1155,10 @@ protected:
 		if (UseVulkanDescriptorCache())
 		{
 			FVulkanHashableDescriptorInfo& HashableInfo = HashableDescriptorInfos[DescriptorIndex];
-			int32 HandleId = Allocation.HandleId;
-			check(HandleId > 0);
-			if (HashableInfo.Buffer.Id != HandleId)
+			check(Allocation.HandleId > 0);
+			if (HashableInfo.Buffer.Id != Allocation.HandleId)
 			{
-				HashableInfo.Buffer.Id = HandleId;
+				HashableInfo.Buffer.Id = Allocation.HandleId;
 				BufferInfo->buffer = (VkBuffer)Allocation.VulkanHandle;
 				bChanged = true;
 			}
diff --git a/Engine/Source/Runtime/VulkanRHI/Private/VulkanMemory.cpp b/Engine/Source/Runtime/VulkanRHI/Private/VulkanMemory.cpp
--- a/Engine/Source/Runtime/VulkanRHI/Private/VulkanMemory.cpp
+++ b/Engine/Source/Runtime/VulkanRHI/Private/VulkanMemory.cpp
@@ -2127,8 +2127,14 @@ namespace VulkanRHI
 			DeviceMemoryAllocation->Map(AllocationSize, 0);
 		}
 
+		uint32 BufferId = 0;
+		if (UseVulkanDescriptorCache())
+		{
+			BufferId = ++GVulkanBufferHandleIdCounter;
+		}
+		
 		++PageIDCounter;
-		FVulkanSubresourceAllocator* Page = new FVulkanSubresourceAllocator(AllocationType, Owner, AllocationFlags, DeviceMemoryAllocation, MemoryTypeIndex, 0);
+		FVulkanSubresourceAllocator* Page = new FVulkanSubresourceAllocator(AllocationType, Owner, AllocationFlags, DeviceMemoryAllocation, MemoryTypeIndex, BufferId);
 		Owner->RegisterSubresourceAllocator(Page);
 		Page->BucketId = BucketId;
 		ActivePages[BucketId].Add(Page);
@@ -2163,8 +2169,15 @@ namespace VulkanRHI
 			return false;
 		}
 
+		
+		uint32 BufferId = 0;
+		if (UseVulkanDescriptorCache())
+		{
+			BufferId = ++GVulkanBufferHandleIdCounter;
+		}
+
 		++PageIDCounter;
-		FVulkanSubresourceAllocator* NewPage = new FVulkanSubresourceAllocator(EVulkanAllocationImageDedicated, Owner, 0, DeviceMemoryAllocation, MemoryTypeIndex, PageIDCounter);
+		FVulkanSubresourceAllocator* NewPage = new FVulkanSubresourceAllocator(EVulkanAllocationImageDedicated, Owner, 0, DeviceMemoryAllocation, MemoryTypeIndex, BufferId);
 		Owner->RegisterSubresourceAllocator(NewPage);
 		NewPage->BucketId = 0xff;
 		UsedDedicatedImagePages.Add(NewPage);
@@ -3787,6 +3800,8 @@ namespace VulkanRHI
 		AllocationIndex = InAllocationIndex;
 		VulkanHandle = Handle;
 		HandleId = BufferId;
+		// Make sure all allocations have a valid Id on platforms that use "Descriptor Cache"
+		ensure(!UseVulkanDescriptorCache() || HandleId != 0);
 	}
 
 	void FVulkanAllocation::Free(FVulkanDevice& Device)
