VulkanRHI: when a slice is specified, create the view for only that slice or layouts will be altered by the render pass for other slices.  Fixes bad layouts in Vulkan Mobile Preview.

diff --git a/Engine/Source/Runtime/VulkanRHI/Private/VulkanViewport.cpp b/Engine/Source/Runtime/VulkanRHI/Private/VulkanViewport.cpp
--- a/Engine/Source/Runtime/VulkanRHI/Private/VulkanViewport.cpp
+++ b/Engine/Source/Runtime/VulkanRHI/Private/VulkanViewport.cpp
@@ -338,9 +338,18 @@ FVulkanFramebuffer::FVulkanFramebuffer(FVulkanDevice& Device, const FRHISetRende
 		FVulkanTextureView RTView;
 		if (Texture->Surface.GetViewType() == VK_IMAGE_VIEW_TYPE_2D || Texture->Surface.GetViewType() == VK_IMAGE_VIEW_TYPE_2D_ARRAY)
 		{
-			const uint32 ArraySliceIndex = FMath::Max(0, (int32)InRTInfo.ColorRenderTarget[Index].ArraySliceIndex);
-			check(ArraySliceIndex < Texture->Surface.GetNumberOfArrayLevels());
-			const uint32 NumArraySlices = Texture->Surface.GetNumberOfArrayLevels() - ArraySliceIndex;
+			uint32 ArraySliceIndex, NumArraySlices;
+			if (InRTInfo.ColorRenderTarget[Index].ArraySliceIndex == -1)
+			{
+				ArraySliceIndex = 0;
+				NumArraySlices = Texture->Surface.GetNumberOfArrayLevels();
+			}
+			else
+			{
+				ArraySliceIndex = InRTInfo.ColorRenderTarget[Index].ArraySliceIndex;
+				NumArraySlices = 1;
+				check(ArraySliceIndex < Texture->Surface.GetNumberOfArrayLevels());
+			}
 			RTView.Create(*Texture->Surface.Device, Texture->Surface.Image, Texture->Surface.GetViewType(), Texture->Surface.GetFullAspectMask(), Texture->Surface.PixelFormat, Texture->Surface.ViewFormat, MipIndex, 1, ArraySliceIndex, NumArraySlices, true);
 		}
 		else if (Texture->Surface.GetViewType() == VK_IMAGE_VIEW_TYPE_CUBE)
