VulkanRHI:
- Split DriverDenyList entries per operating system (driver numbering is different per OS anyhow).
- Add DriverDenyList support to Linux.
- Fix DriverDenyList support for Vulkan (no driver date was invalidating entries).
- Add DriverDenyList entries for Vulkan on Windows and Linux.
Tested on Windows/Linux for AMD/NVIDIA.

diff --git a/Engine/Config/BaseHardware.ini b/Engine/Config/BaseHardware.ini
--- a/Engine/Config/BaseHardware.ini
+++ b/Engine/Config/BaseHardware.ini
@@ -19,7 +19,9 @@
 ;  - Bugs on the latest driver are added to the DriverDenyList but treated as critical bugs for which we must try to find a workaround. We must also notify the IHV to help if we cannot find the problem.
 ;
 ; INI sections:
-;    [GPU_NVIDIA] [GPU_AMD] [GPU_Intel]
+;    [Vendor OS] where
+;    Vendor: GPU_NVIDIA, GPU_AMD, GPU_Intel
+;    OS: Windows, Linux
 ;
 ; Explanation of the INI entries:
 ;  - "SuggestedDriverVersion": string used to communicate to the user which driver version is appropriate (not used by any code), empty if unknown.
@@ -31,9 +33,8 @@
 ;
 ;-----------------------------------------------------------------------------------------------------------------
 
-[GPU_NVIDIA]
-; a newer one is out but has issues (see below)
-SuggestedDriverVersion="461.92"
+[GPU_NVIDIA Windows]
++SuggestedDriverVersion="511.17"
 
 ; Seems this driver was not setting uninitialized global variables
 ; (static, non const) to 0 but leaving them uninitialized.
@@ -63,10 +64,25 @@ SuggestedDriverVersion="461.92"
 
 +DriverDenyList=(DriverVersion="==465.89", Reason="This driver version has known stability issues in DirectX 12")
 
+; UE5 in Windows/Vulkan requires recent drivers for Nanite/Lumen and has other stability issues with older drivers (need >= 496.76)
++DriverDenyList=(DriverVersion="<496.76", Reason="These driver versions have known stability issues and missing features in Vulkan", RHI="Vulkan")
+
 ;-----------------------------------------------------------------------------------------------------------------
 
-[GPU_AMD]
-SuggestedDriverVersion="20.8.3"
+[GPU_NVIDIA Linux]
+SuggestedDriverVersion="510.47"
+
+; UE5 in Linux/Vulkan requires recent drivers for Nanite/Lumen and has other stability issues with older drivers (need >= 495.46)
++DriverDenyList=(DriverVersion="<495.46", Reason="These driver versions have known stability issues and missing features in Vulkan")
+
+
+
+
+;-----------------------------------------------------------------------------------------------------------------
+;-----------------------------------------------------------------------------------------------------------------
+
+[GPU_AMD Windows]
++SuggestedDriverVersion="21.11.3"
 
 ; From OR-27051
 +DriverDenyList=(DriverVersion="==15.300.1025.1001", Reason="Crashes with Paragon content using more recent GPU features")
@@ -76,7 +92,27 @@ SuggestedDriverVersion="20.8.3"
 
 +DriverDenyList=(DriverVersion="<=26.20.13031.15006", Reason="Driver crashes creating PSOs", RHI="D3D12")
 
+; UE5 in Windows/Vulkan requires recent drivers for Nanite/Lumen and has other stability issues with older drivers (need >= 21.11.3)
++DriverDenyList=(DriverVersion="<30.0.13037.1003", Reason="These driver versions have known stability issues and missing features in Vulkan", RHI="Vulkan")
+
+;-----------------------------------------------------------------------------------------------------------------
+
+[GPU_AMD Linux]
+SuggestedDriverVersion="21.2.6"
+
+; UE5 in Linux/Vulkan requires recent drivers for Nanite/Lumen and has other stability issues with older drivers (need >= RADV 21.0.3)
++DriverDenyList=(DriverVersion="<21.2.6", Reason="These driver versions have known stability issues and missing features in Vulkan")
+
+
+
+
+;-----------------------------------------------------------------------------------------------------------------
+;-----------------------------------------------------------------------------------------------------------------
+
+[GPU_Intel Windows]
+SuggestedDriverVersion="27.20.100.8935"
+
 ;-----------------------------------------------------------------------------------------------------------------
 
-[GPU_Intel]
+[GPU_Intel Linux]
 SuggestedDriverVersion="27.20.100.8935"
diff --git a/Engine/Source/Runtime/Core/Public/GenericPlatform/GenericPlatformDriver.h b/Engine/Source/Runtime/Core/Public/GenericPlatform/GenericPlatformDriver.h
--- a/Engine/Source/Runtime/Core/Public/GenericPlatform/GenericPlatformDriver.h
+++ b/Engine/Source/Runtime/Core/Public/GenericPlatform/GenericPlatformDriver.h
@@ -385,10 +385,10 @@ struct FDriverDenyListEntry
 		ensure(!DriverVersionString.IsEmpty() || !DriverDateString.IsEmpty());
 
 		FParse::Value(In, TEXT("RHI="), RHIName);
-		
+
+
 		// later:
 //		FParse::Value(In, TEXT("DeviceId="), DeviceId);
-//		FParse::Value(In, TEXT("OS="), OS);
 //		FParse::Value(In, TEXT("API="), API);
 //		ensure(API == TEXT("DX11"));
 
@@ -567,14 +567,14 @@ struct FGPUHardware
 	// @return a driver version intended to be shown to the user e.g. "15.30.1025.1001 12/17/2015 (Crimson Edition 15.12)"
 	FString GetSuggestedDriverVersion(const FString& InRHIName) const
 	{
-		const TCHAR* Section = GetVendorSectionName();
+		const FString Section = GetVendorSectionName();
 
 		FString Ret;
 
-		if(Section)
+		if(!Section.IsEmpty())
 		{
 			TArray<FString> SuggestedDriverVersions;
-			GConfig->GetArray(Section, TEXT("SuggestedDriverVersion"), SuggestedDriverVersions, GHardwareIni);
+			GConfig->GetArray(*Section, TEXT("SuggestedDriverVersion"), SuggestedDriverVersions, GHardwareIni);
 
 			// Find specific RHI version first
 			if (InRHIName.Len() > 0)
@@ -607,12 +607,12 @@ struct FGPUHardware
 	// @return 0 if there is none
 	FDriverDenyListEntry FindDriverDenyListEntry() const
 	{
-		const TCHAR* Section = GetVendorSectionName();
+		const FString Section = GetVendorSectionName();
 
-		if(Section)
+		if(!Section.IsEmpty())
 		{
 			TArray<FString> DenyListStrings;
-			GConfig->GetArray(GetVendorSectionName(), TEXT("DriverDenyList"), DenyListStrings, GHardwareIni);
+			GConfig->GetArray(*Section, TEXT("DriverDenyList"), DenyListStrings, GHardwareIni);
 
 			for(int32 i = 0; i < DenyListStrings.Num(); ++i)
 			{
@@ -640,12 +640,12 @@ struct FGPUHardware
 	bool IsLatestDenied() const
 	{
 		bool bLatestDenied = false;
-		const TCHAR* Section = GetVendorSectionName();
+		const FString Section = GetVendorSectionName();
 
-		if(Section)
+		if(!Section.IsEmpty())
 		{
 			TArray<FString> DenyListStrings;
-			GConfig->GetArray(GetVendorSectionName(), TEXT("DriverDenyList"), DenyListStrings, GHardwareIni);
+			GConfig->GetArray(*Section, TEXT("DriverDenyList"), DenyListStrings, GHardwareIni);
 
 			for(int32 i = 0; !bLatestDenied && i < DenyListStrings.Num(); ++i)
 			{
@@ -665,22 +665,29 @@ struct FGPUHardware
 
 	// to get a section name in the Hardware.ini file
 	// @return 0 if not found
-	const TCHAR* GetVendorSectionName() const
+	FString GetVendorSectionName() const
 	{
+		const TCHAR* Section = nullptr;
+
 		if(DriverInfo.IsNVIDIA())
 		{
-			return TEXT("GPU_NVIDIA");
+			Section = TEXT("GPU_NVIDIA");
 		}
 		if(DriverInfo.IsAMD())
 		{
-			return TEXT("GPU_AMD");
+			Section = TEXT("GPU_AMD");
 		}
 		else if(DriverInfo.IsIntel())
 		{
-			return TEXT("GPU_Intel");
-		}
+			Section = TEXT("GPU_Intel");
+		}	
 		// more GPU vendors can be added on demand
-		return 0;
+		if (!Section)
+		{
+			return TEXT("");
+		}
+
+		return FString::Printf(TEXT("%s %s"), Section, ANSI_TO_TCHAR(FPlatformProperties::IniPlatformName()));
 	}
 };
 
diff --git a/Engine/Source/Runtime/RHI/Private/DynamicRHI.cpp b/Engine/Source/Runtime/RHI/Private/DynamicRHI.cpp
--- a/Engine/Source/Runtime/RHI/Private/DynamicRHI.cpp
+++ b/Engine/Source/Runtime/RHI/Private/DynamicRHI.cpp
@@ -79,12 +79,12 @@ void InitNullRHI()
 	FGenericCrashContext::SetEngineData(TEXT("RHI.RHIName"), TEXT("NullRHI"));
 }
 
-#if PLATFORM_WINDOWS
+#if PLATFORM_WINDOWS || PLATFORM_UNIX
 static void RHIDetectAndWarnOfBadDrivers(bool bHasEditorToken)
 {
 	int32 CVarValue = CVarWarnOfBadDrivers.GetValueOnGameThread();
 
-	if(!GIsRHIInitialized || !CVarValue || GRHIVendorId == 0)
+	if(!GIsRHIInitialized || !CVarValue || (GRHIVendorId == 0) || FApp::IsUnattended())
 	{
 		return;
 	}
@@ -307,7 +307,7 @@ void RHIInit(bool bHasEditorToken)
 		check(GDynamicRHI);
 	}
 
-#if PLATFORM_WINDOWS || PLATFORM_MAC
+#if PLATFORM_WINDOWS || PLATFORM_MAC || PLATFORM_UNIX
 #if defined(NV_GEFORCENOW) && NV_GEFORCENOW
 	bool bDetectAndWarnBadDrivers = true;
 	if (IsRHIDeviceNVIDIA() && !!CVarDisableDriverWarningPopupIfGFN.GetValueOnAnyThread())
diff --git a/Engine/Source/Runtime/VulkanRHI/Private/VulkanRHI.cpp b/Engine/Source/Runtime/VulkanRHI/Private/VulkanRHI.cpp
--- a/Engine/Source/Runtime/VulkanRHI/Private/VulkanRHI.cpp
+++ b/Engine/Source/Runtime/VulkanRHI/Private/VulkanRHI.cpp
@@ -724,23 +724,43 @@ void FVulkanDynamicRHI::SelectAndInitDevice()
 		GRHIAdapterName.Append(TEXT(" Vulkan"));
 		GRHIAdapterInternalDriverVersion = FString::Printf(TEXT("%d.%d.%d"), VK_VERSION_MAJOR(Props.apiVersion), VK_VERSION_MINOR(Props.apiVersion), VK_VERSION_PATCH(Props.apiVersion));
 	}
-	else if (Device->GetVendorId() == EGpuVendorId::Nvidia)
+	else if (PLATFORM_WINDOWS)
 	{
-		UNvidiaDriverVersion NvidiaVersion;
-		static_assert(sizeof(NvidiaVersion) == sizeof(Props.driverVersion), "Mismatched Nvidia pack driver version!");
-		NvidiaVersion.Packed = Props.driverVersion;
-		GRHIAdapterUserDriverVersion = FString::Printf(TEXT("%d.%02d"), NvidiaVersion.Major, NvidiaVersion.Minor);
-		UE_LOG(LogVulkanRHI, Display, TEXT("Nvidia User Driver Version = %s"), *GRHIAdapterUserDriverVersion);
+		GRHIDeviceId = Props.deviceID;
+		FGPUDriverInfo GPUDriverInfo = FPlatformMisc::GetGPUDriverInfo(GRHIAdapterName);
+		GRHIAdapterUserDriverVersion = GPUDriverInfo.UserDriverVersion;
+		GRHIAdapterInternalDriverVersion = GPUDriverInfo.InternalDriverVersion;
+		GRHIAdapterDriverDate = GPUDriverInfo.DriverDate;
 
-		// Ignore GRHIAdapterInternalDriverVersion for now as the device name doesn't match
+		UE_LOG(LogVulkanRHI, Log, TEXT("    Adapter Name: %s"), *GRHIAdapterName);
+		UE_LOG(LogVulkanRHI, Log, TEXT("     API Version: %d.%d.%d"), VK_VERSION_MAJOR(Props.apiVersion), VK_VERSION_MINOR(Props.apiVersion), VK_VERSION_PATCH(Props.apiVersion));
+		UE_LOG(LogVulkanRHI, Log, TEXT("  Driver Version: %s (0x%X)"), *GRHIAdapterUserDriverVersion, Props.apiVersion);
+		UE_LOG(LogVulkanRHI, Log, TEXT("Internal Version: %s"), *GRHIAdapterInternalDriverVersion);
+		UE_LOG(LogVulkanRHI, Log, TEXT("     Driver Date: %s"), *GRHIAdapterDriverDate);
 	}
 	else if(PLATFORM_UNIX)
 	{
-		GRHIAdapterInternalDriverVersion = FString::Printf(TEXT("%d.%d.%d (0x%X)"), VK_VERSION_MAJOR(Props.apiVersion), VK_VERSION_MINOR(Props.apiVersion), VK_VERSION_PATCH(Props.apiVersion), Props.apiVersion);
-		GRHIAdapterUserDriverVersion = FString::Printf(TEXT("%d.%d.%d (0x%X)"), VK_VERSION_MAJOR(Props.driverVersion), VK_VERSION_MINOR(Props.driverVersion), VK_VERSION_PATCH(Props.driverVersion), Props.driverVersion);
+		// :todo: Create a GetGPUDriverInfo for Linux
+
+		if (Device->GetVendorId() == EGpuVendorId::Nvidia)
+		{
+			UNvidiaDriverVersion NvidiaVersion;
+			static_assert(sizeof(NvidiaVersion) == sizeof(Props.driverVersion), "Mismatched Nvidia pack driver version!");
+			NvidiaVersion.Packed = Props.driverVersion;
+			GRHIAdapterUserDriverVersion = FString::Printf(TEXT("%d.%02d"), NvidiaVersion.Major, NvidiaVersion.Minor);
+		}
+		else
+		{
+			GRHIAdapterUserDriverVersion = FString::Printf(TEXT("%d.%d.%d"), VK_VERSION_MAJOR(Props.driverVersion), VK_VERSION_MINOR(Props.driverVersion), VK_VERSION_PATCH(Props.driverVersion), Props.driverVersion);
+		}
+
 		GRHIDeviceId = Props.deviceID;
+		GRHIAdapterInternalDriverVersion = GRHIAdapterUserDriverVersion;
+		GRHIAdapterDriverDate = TEXT("01-01-01");  // Unused on unix systems, pick a date that will fail test if compared but passes IsValid() check
 
-		UE_LOG(LogVulkanRHI, Display, TEXT("'%s' User Driver Version = %s"), *GRHIAdapterName, *GRHIAdapterUserDriverVersion);
+		UE_LOG(LogVulkanRHI, Log, TEXT("    Adapter Name: %s"), *GRHIAdapterName);
+		UE_LOG(LogVulkanRHI, Log, TEXT("     API Version: %d.%d.%d"), VK_VERSION_MAJOR(Props.apiVersion), VK_VERSION_MINOR(Props.apiVersion), VK_VERSION_PATCH(Props.apiVersion));
+		UE_LOG(LogVulkanRHI, Log, TEXT("  Driver Version: %s (0x%X)"), *GRHIAdapterUserDriverVersion, Props.apiVersion);
 	}
 
 	GRHIPersistentThreadGroupCount = 1440; // TODO: Revisit based on vendor/adapter/perf query
