Quick cleanup of Vulkan extension initialization (before the bigger clean up):
- Make naming a bit more consistent
- Add VK_KHR_maintenance4 support
- Pull VK_EXT_descriptor_indexing out of Vulkan Raycasting so it can be used on its own
- Track extensions that were promoted in later version of Vulkan for eventual cleanup.

diff --git a/Engine/Source/Runtime/VulkanRHI/Private/VulkanDevice.cpp b/Engine/Source/Runtime/VulkanRHI/Private/VulkanDevice.cpp
--- a/Engine/Source/Runtime/VulkanRHI/Private/VulkanDevice.cpp
+++ b/Engine/Source/Runtime/VulkanRHI/Private/VulkanDevice.cpp
@@ -381,7 +381,7 @@ void FVulkanDevice::CreateDevice()
 #endif
 
 #if VULKAN_SUPPORTS_SCALAR_BLOCK_LAYOUT
-	if (OptionalDeviceExtensions.HasScalarBlockLayoutFeatures)
+	if (OptionalDeviceExtensions.HasEXTScalarBlockLayout)
 	{
 		ZeroVulkanStruct(OptionalFeatures.ScalarBlockLayoutFeatures, VK_STRUCTURE_TYPE_PHYSICAL_DEVICE_SCALAR_BLOCK_LAYOUT_FEATURES_EXT);
 		OptionalFeatures.ScalarBlockLayoutFeatures.scalarBlockLayout = VK_TRUE;
@@ -407,7 +407,7 @@ void FVulkanDevice::CreateDevice()
 	// Ensure we have both to enable atomic code paths
 	VkPhysicalDeviceShaderAtomicInt64Features BufferAtomicFeatures;
 	VkPhysicalDeviceShaderImageAtomicInt64FeaturesEXT ImageAtomicFeatures;
-	if (bEnable64bitAtomics && OptionalDeviceExtensions.HasBufferAtomicInt64 && OptionalDeviceExtensions.HasImageAtomicInt64)
+	if (bEnable64bitAtomics && OptionalDeviceExtensions.HasKHRShaderAtomicInt64 && OptionalDeviceExtensions.HasImageAtomicInt64)
 	{
 		ZeroVulkanStruct(BufferAtomicFeatures, VK_STRUCTURE_TYPE_PHYSICAL_DEVICE_SHADER_ATOMIC_INT64_FEATURES_KHR);
 		BufferAtomicFeatures.shaderBufferInt64Atomics = VK_TRUE;
@@ -1206,7 +1206,7 @@ void FVulkanDevice::InitGPU(int32 DeviceIndex)
 #if VULKAN_SUPPORTS_BUFFER_64BIT_ATOMICS
 		VkPhysicalDeviceShaderAtomicInt64FeaturesKHR BufferAtomicFeatures;
 		ZeroVulkanStruct(BufferAtomicFeatures, VK_STRUCTURE_TYPE_PHYSICAL_DEVICE_SHADER_ATOMIC_INT64_FEATURES_KHR);
-		if (OptionalDeviceExtensions.HasBufferAtomicInt64)  // Set by QueryGPU if extension is supported
+		if (OptionalDeviceExtensions.HasKHRShaderAtomicInt64)  // Set by QueryGPU if extension is supported
 		{
 			*NextPropsAddr = &BufferAtomicFeatures;
 			NextPropsAddr = &BufferAtomicFeatures.pNext;
@@ -1272,19 +1272,45 @@ void FVulkanDevice::InitGPU(int32 DeviceIndex)
 		}
 #endif
 
+#if VULKAN_SUPPORTS_MAINTENANCE4
+		VkPhysicalDeviceMaintenance4FeaturesKHR Maintenance4Features;
+		ZeroVulkanStruct(Maintenance4Features, VK_STRUCTURE_TYPE_PHYSICAL_DEVICE_MAINTENANCE_4_FEATURES);
+		if (OptionalDeviceExtensions.HasKHRMaintenance4)  // Set by QueryGPU if extension is supported
+		{
+			*NextPropsAddr = &Maintenance4Features;
+			NextPropsAddr = &Maintenance4Features.pNext;
+		}
+#endif
+
+#if VULKAN_SUPPORTS_DESCRIPTOR_INDEXING
+		VkPhysicalDeviceDescriptorIndexingFeaturesEXT DescriptorIndexingFeatures;
+		ZeroVulkanStruct(DescriptorIndexingFeatures, VK_STRUCTURE_TYPE_PHYSICAL_DEVICE_DESCRIPTOR_INDEXING_FEATURES);
+		if (OptionalDeviceExtensions.HasEXTDescriptorIndexing)  // Set by QueryGPU if extension is supported
+		{
+			*NextPropsAddr = &DescriptorIndexingFeatures;
+			NextPropsAddr = &DescriptorIndexingFeatures.pNext;
+		}
+#endif
+
 		VulkanRHI::vkGetPhysicalDeviceFeatures2KHR(Gpu, &Features2);
 
 		// If we have the extension but the feature is not supported, clear the bits.
 #if VULKAN_SUPPORTS_BUFFER_64BIT_ATOMICS
-		OptionalDeviceExtensions.HasBufferAtomicInt64 = (BufferAtomicFeatures.shaderBufferInt64Atomics == VK_TRUE);
+		OptionalDeviceExtensions.HasKHRShaderAtomicInt64 = (BufferAtomicFeatures.shaderBufferInt64Atomics == VK_TRUE);
 #endif
 #if VULKAN_SUPPORTS_IMAGE_64BIT_ATOMICS
 		OptionalDeviceExtensions.HasImageAtomicInt64 = (ImageAtomicFeatures.shaderImageInt64Atomics == VK_TRUE);
+#endif
+#if VULKAN_SUPPORTS_MAINTENANCE4
+		OptionalDeviceExtensions.HasKHRMaintenance4 = (Maintenance4Features.maintenance4 == VK_TRUE);
+#endif
+#if VULKAN_SUPPORTS_DESCRIPTOR_INDEXING
+		OptionalDeviceExtensions.HasEXTDescriptorIndexing = (DescriptorIndexingFeatures.runtimeDescriptorArray == VK_TRUE); // :todo-jn: add resource specific checks
 #endif
 	}
 #endif // VULKAN_SUPPORTS_PHYSICAL_DEVICE_PROPERTIES2
 
-	UE_LOG(LogVulkanRHI, Display, TEXT("Using Device %d: Geometry %d BufferAtomic64 %d ImageAtomic64 %d"), DeviceIndex, PhysicalFeatures.geometryShader, OptionalDeviceExtensions.HasBufferAtomicInt64, OptionalDeviceExtensions.HasImageAtomicInt64);
+	UE_LOG(LogVulkanRHI, Display, TEXT("Using Device %d: Geometry %d BufferAtomic64 %d ImageAtomic64 %d"), DeviceIndex, PhysicalFeatures.geometryShader, OptionalDeviceExtensions.HasKHRShaderAtomicInt64, OptionalDeviceExtensions.HasImageAtomicInt64);
 
 	CreateDevice();
 
diff --git a/Engine/Source/Runtime/VulkanRHI/Private/VulkanDevice.h b/Engine/Source/Runtime/VulkanRHI/Private/VulkanDevice.h
--- a/Engine/Source/Runtime/VulkanRHI/Private/VulkanDevice.h
+++ b/Engine/Source/Runtime/VulkanRHI/Private/VulkanDevice.h
@@ -37,39 +37,49 @@ struct FOptionalVulkanDeviceExtensions
 	{
 		struct
 		{
-			uint64 HasKHRMaintenance1 : 1;
-			uint64 HasKHRMaintenance2 : 1;
-			uint64 HasKHRDedicatedAllocation : 1;
+			// Optional Extensions
 			uint64 HasEXTValidationCache : 1;
-			uint64 HasAMDBufferMarker : 1;
-			uint64 HasNVDiagnosticCheckpoints : 1;
-			uint64 HasNVDeviceDiagnosticConfig : 1;
-			uint64 HasYcbcrSampler : 1;
 			uint64 HasMemoryPriority : 1;
 			uint64 HasMemoryBudget : 1;
-			uint64 HasDriverProperties : 1;
+			uint64 HasEXTASTCDecodeMode : 1;
 			uint64 HasEXTFragmentDensityMap : 1;
 			uint64 HasEXTFragmentDensityMap2 : 1;
 			uint64 HasKHRFragmentShadingRate : 1;
-			uint64 HasKHRRenderPass2 : 1;
 			uint64 HasEXTFullscreenExclusive : 1;
-			uint64 HasKHRImageFormatList : 1;
-			uint64 HasEXTTextureCompressionASTCHDR : 1;
-			uint64 HasEXTASTCDecodeMode : 1;
-			uint64 HasQcomRenderPassTransform : 1;
 			uint64 HasImageAtomicInt64 : 1;
-			uint64 HasBufferAtomicInt64 : 1;
-			uint64 HasScalarBlockLayoutFeatures : 1;
-			uint64 HasKHRMultiview : 1;
 			uint64 HasAccelerationStructure : 1;
 			uint64 HasRayTracingPipeline : 1;
 			uint64 HasRayQuery : 1;
-			uint64 HasDescriptorIndexing : 1;
-			uint64 HasBufferDeviceAddress : 1;
 			uint64 HasDeferredHostOperations : 1;
+
+			// Vendor specific
+			uint64 HasAMDBufferMarker : 1;
+			uint64 HasNVDiagnosticCheckpoints : 1;
+			uint64 HasNVDeviceDiagnosticConfig : 1;
+			uint64 HasQcomRenderPassTransform : 1;
+
+			// Promoted to 1.1
+			uint64 HasKHRMaintenance1 : 1;
+			uint64 HasKHRMaintenance2 : 1;
+			uint64 HasKHRDedicatedAllocation : 1;
+			uint64 HasYcbcrSampler : 1;  // Combination of VK_KHR_sampler_ycbcr_conversion + VK_KHR_bind_memory2 + VK_KHR_get_memory_requirements2
+			uint64 HasKHRMultiview : 1;
+
+			// Promoted to 1.2
+			uint64 HasDriverProperties : 1;
+			uint64 HasKHRRenderPass2 : 1;
+			uint64 HasKHRImageFormatList : 1;
+			uint64 HasKHRShaderAtomicInt64 : 1;
+			uint64 HasEXTScalarBlockLayout : 1;
+			uint64 HasBufferDeviceAddress : 1;
 			uint64 HasSPIRV_14 : 1;
 			uint64 HasShaderFloatControls : 1;
+			uint64 HasEXTDescriptorIndexing : 1;
 			uint64 HasEXTShaderViewportIndexLayer : 1;
+
+			// Promoted to 1.3
+			uint64 HasEXTTextureCompressionASTCHDR : 1;
+			uint64 HasKHRMaintenance4 : 1;
 		};
 		uint64 Packed;
 	};
@@ -93,7 +103,7 @@ struct FOptionalVulkanDeviceExtensions
 		return 
 			HasAccelerationStructure && 
 			(HasRayTracingPipeline || HasRayQuery) &&
-			HasDescriptorIndexing && 
+			HasEXTDescriptorIndexing &&
 			HasBufferDeviceAddress && 
 			HasDeferredHostOperations && 
 			HasSPIRV_14 && 
diff --git a/Engine/Source/Runtime/VulkanRHI/Private/VulkanLayers.cpp b/Engine/Source/Runtime/VulkanRHI/Private/VulkanLayers.cpp
--- a/Engine/Source/Runtime/VulkanRHI/Private/VulkanLayers.cpp
+++ b/Engine/Source/Runtime/VulkanRHI/Private/VulkanLayers.cpp
@@ -770,11 +770,11 @@ void FOptionalVulkanInstanceExtensions::Setup(const TArray<const ANSICHAR*>& Ins
 	check(Packed == 0);
 
 #if VULKAN_SUPPORTS_EXTERNAL_MEMORY
-	HasKHRExternalMemoryCapabilities = HasExtension(InstanceExtensions, VK_KHR_EXTERNAL_MEMORY_CAPABILITIES_EXTENSION_NAME);
+	HasKHRExternalMemoryCapabilities = HasExtension(InstanceExtensions, VK_KHR_EXTERNAL_MEMORY_CAPABILITIES_EXTENSION_NAME);  // Promoted to 1.1
 #endif
 
 #if VULKAN_SUPPORTS_PHYSICAL_DEVICE_PROPERTIES2
-	HasKHRGetPhysicalDeviceProperties2 = HasExtension(InstanceExtensions, VK_KHR_GET_PHYSICAL_DEVICE_PROPERTIES_2_EXTENSION_NAME);
+	HasKHRGetPhysicalDeviceProperties2 = HasExtension(InstanceExtensions, VK_KHR_GET_PHYSICAL_DEVICE_PROPERTIES_2_EXTENSION_NAME);  // Promoted to 1.1
 #endif
 }
 
@@ -782,141 +782,138 @@ void FOptionalVulkanDeviceExtensions::Setup(const TArray<const ANSICHAR*>& Devic
 {
 	check(Packed == 0);
 
-#if VULKAN_SUPPORTS_MAINTENANCE_LAYER1
-	HasKHRMaintenance1 = HasExtension(DeviceExtensions, VK_KHR_MAINTENANCE1_EXTENSION_NAME);
-#endif
-#if VULKAN_SUPPORTS_MAINTENANCE_LAYER2
-	HasKHRMaintenance2 = HasExtension(DeviceExtensions, VK_KHR_MAINTENANCE2_EXTENSION_NAME);
-#endif
-	//HasMirrorClampToEdge = HasExtension(VK_KHR_SAMPLER_MIRROR_CLAMP_TO_EDGE_EXTENSION_NAME);
-
-#if VULKAN_SUPPORTS_DEDICATED_ALLOCATION
-	HasKHRDedicatedAllocation = HasExtension(DeviceExtensions, VK_KHR_GET_MEMORY_REQUIREMENTS_2_EXTENSION_NAME) && HasExtension(DeviceExtensions, VK_KHR_DEDICATED_ALLOCATION_EXTENSION_NAME);
-#endif
 
+	// Optional Extensions
 #if VULKAN_SUPPORTS_VALIDATION_CACHE
 	HasEXTValidationCache = HasExtension(DeviceExtensions, VK_EXT_VALIDATION_CACHE_EXTENSION_NAME);
 #endif
-
-	bool bHasAnyCrashExtension = false;
-#if VULKAN_SUPPORTS_AMD_BUFFER_MARKER
-	if (GGPUCrashDebuggingEnabled)
-	{
-		HasAMDBufferMarker = HasExtension(DeviceExtensions, VK_AMD_BUFFER_MARKER_EXTENSION_NAME);
-		bHasAnyCrashExtension = bHasAnyCrashExtension || HasAMDBufferMarker;
-	}
-#endif
-
-#if VULKAN_SUPPORTS_NV_DIAGNOSTICS
-	if (GGPUCrashDebuggingEnabled)
-	{
-		HasNVDiagnosticCheckpoints = HasExtension(DeviceExtensions, VK_NV_DEVICE_DIAGNOSTIC_CHECKPOINTS_EXTENSION_NAME);
-		HasNVDeviceDiagnosticConfig = HasExtension(DeviceExtensions, VK_NV_DEVICE_DIAGNOSTICS_CONFIG_EXTENSION_NAME);
-		bHasAnyCrashExtension = bHasAnyCrashExtension || (HasNVDeviceDiagnosticConfig && HasNVDiagnosticCheckpoints);
-	}
-#endif
-
-	if (GGPUCrashDebuggingEnabled && !bHasAnyCrashExtension)
-	{
-		UE_LOG(LogVulkanRHI, Warning, TEXT("Tried to enable GPU crash debugging but no extension found! Will use local tracepoints."));
-	}
-
-#if VULKAN_SUPPORTS_COLOR_CONVERSIONS
-	HasYcbcrSampler = HasExtension(DeviceExtensions, VK_KHR_SAMPLER_YCBCR_CONVERSION_EXTENSION_NAME) && HasExtension(DeviceExtensions, VK_KHR_BIND_MEMORY_2_EXTENSION_NAME) && HasExtension(DeviceExtensions, VK_KHR_GET_MEMORY_REQUIREMENTS_2_EXTENSION_NAME);
-#endif
-
 #if VULKAN_SUPPORTS_MEMORY_PRIORITY
-	HasMemoryPriority = HasExtension(DeviceExtensions, VK_EXT_MEMORY_PRIORITY_EXTENSION_NAME);
-	if (FParse::Param(FCommandLine::Get(), TEXT("disablememorypriority")))
-	{
-		HasMemoryPriority = 0;
-	}
-#else
-	HasMemoryPriority = 0;
+	HasMemoryPriority = HasExtension(DeviceExtensions, VK_EXT_MEMORY_PRIORITY_EXTENSION_NAME) && (!FParse::Param(FCommandLine::Get(), TEXT("disablememorypriority")));
 #endif
-
 #if VULKAN_SUPPORTS_MEMORY_BUDGET
-	HasMemoryBudget = HasExtension(DeviceExtensions, VK_EXT_MEMORY_BUDGET_EXTENSION_NAME);
-	if (FParse::Param(FCommandLine::Get(), TEXT("disablememorybudget")))
-	{
-		HasMemoryBudget = 0;
-	}
-#else
-	HasMemoryBudget = 0;
-#endif
-
-#if VULKAN_SUPPORTS_TEXTURE_COMPRESSION_ASTC_HDR
-	HasEXTTextureCompressionASTCHDR = HasExtension(DeviceExtensions, VK_EXT_TEXTURE_COMPRESSION_ASTC_HDR_EXTENSION_NAME);
-#else
-	HasEXTTextureCompressionASTCHDR = 0;
+	HasMemoryBudget = HasExtension(DeviceExtensions, VK_EXT_MEMORY_BUDGET_EXTENSION_NAME) && (!FParse::Param(FCommandLine::Get(), TEXT("disablememorybudget")));
 #endif
-
 #if VULKAN_SUPPORTS_ASTC_DECODE_MODE
 	HasEXTASTCDecodeMode = HasExtension(DeviceExtensions, VK_EXT_ASTC_DECODE_MODE_EXTENSION_NAME);
-#else
-	HasEXTASTCDecodeMode = 0;
 #endif
-
-#if VULKAN_SUPPORTS_DRIVER_PROPERTIES
-	HasDriverProperties = HasExtension(DeviceExtensions, VK_KHR_DRIVER_PROPERTIES_EXTENSION_NAME);
-#endif
-
-#if VULKAN_SUPPORTS_RENDERPASS2
-	HasKHRRenderPass2 = HasExtension(DeviceExtensions, VK_KHR_CREATE_RENDERPASS_2_EXTENSION_NAME);
-#endif
-
 #if VULKAN_SUPPORTS_FRAGMENT_DENSITY_MAP
 	HasEXTFragmentDensityMap = HasExtension(DeviceExtensions, VK_EXT_FRAGMENT_DENSITY_MAP_EXTENSION_NAME);
 #endif
-
 #if VULKAN_SUPPORTS_FRAGMENT_DENSITY_MAP2
 	HasEXTFragmentDensityMap2 = HasExtension(DeviceExtensions, VK_EXT_FRAGMENT_DENSITY_MAP_2_EXTENSION_NAME);
 #endif
-
 #if VULKAN_SUPPORTS_FRAGMENT_SHADING_RATE
 	HasKHRFragmentShadingRate = HasExtension(DeviceExtensions, VK_KHR_FRAGMENT_SHADING_RATE_EXTENSION_NAME);
 #endif
+#if VULKAN_SUPPORTS_FULLSCREEN_EXCLUSIVE
+	HasEXTFullscreenExclusive = HasExtension(DeviceExtensions, VK_EXT_FULL_SCREEN_EXCLUSIVE_EXTENSION_NAME);
+#endif
+#if VULKAN_SUPPORTS_IMAGE_64BIT_ATOMICS
+	HasImageAtomicInt64 = HasExtension(DeviceExtensions, VK_EXT_SHADER_IMAGE_ATOMIC_INT64_EXTENSION_NAME);
+#endif
 
-#if VULKAN_SUPPORTS_MULTIVIEW
-	HasKHRMultiview = HasExtension(DeviceExtensions, VK_KHR_MULTIVIEW_EXTENSION_NAME);
+
+	// Vendor
+	{
+		bool bHasAnyCrashExtension = false;
+
+#if VULKAN_SUPPORTS_AMD_BUFFER_MARKER
+		if (GGPUCrashDebuggingEnabled)
+		{
+			HasAMDBufferMarker = HasExtension(DeviceExtensions, VK_AMD_BUFFER_MARKER_EXTENSION_NAME);
+			bHasAnyCrashExtension = bHasAnyCrashExtension || HasAMDBufferMarker;
+		}
 #endif
 
-#if VULKAN_SUPPORTS_FULLSCREEN_EXCLUSIVE
-	HasEXTFullscreenExclusive = HasExtension(DeviceExtensions, VK_EXT_FULL_SCREEN_EXCLUSIVE_EXTENSION_NAME);
+#if VULKAN_SUPPORTS_NV_DIAGNOSTICS
+		if (GGPUCrashDebuggingEnabled)
+		{
+			HasNVDiagnosticCheckpoints = HasExtension(DeviceExtensions, VK_NV_DEVICE_DIAGNOSTIC_CHECKPOINTS_EXTENSION_NAME);
+			HasNVDeviceDiagnosticConfig = HasExtension(DeviceExtensions, VK_NV_DEVICE_DIAGNOSTICS_CONFIG_EXTENSION_NAME);
+			bHasAnyCrashExtension = bHasAnyCrashExtension || (HasNVDeviceDiagnosticConfig && HasNVDiagnosticCheckpoints);
+		}
 #endif
 
-	HasKHRImageFormatList = HasExtension(DeviceExtensions, VK_KHR_IMAGE_FORMAT_LIST_EXTENSION_NAME);
+		if (GGPUCrashDebuggingEnabled && !bHasAnyCrashExtension)
+		{
+			UE_LOG(LogVulkanRHI, Warning, TEXT("Tried to enable GPU crash debugging but no extension found! Will use local tracepoints."));
+		}
 
 #if VULKAN_SUPPORTS_QCOM_RENDERPASS_TRANSFORM
-	HasQcomRenderPassTransform = HasExtension(DeviceExtensions, VK_QCOM_RENDER_PASS_TRANSFORM_EXTENSION_NAME);
+		HasQcomRenderPassTransform = HasExtension(DeviceExtensions, VK_QCOM_RENDER_PASS_TRANSFORM_EXTENSION_NAME);
 #endif
+	}
 
-#if VULKAN_SUPPORTS_BUFFER_64BIT_ATOMICS
-	HasBufferAtomicInt64 = HasExtension(DeviceExtensions, VK_KHR_SHADER_ATOMIC_INT64_EXTENSION_NAME);
+
+	// Promoted to 1.1
+	{
+#if VULKAN_SUPPORTS_MAINTENANCE_LAYER1
+		HasKHRMaintenance1 = HasExtension(DeviceExtensions, VK_KHR_MAINTENANCE1_EXTENSION_NAME);
+#endif
+#if VULKAN_SUPPORTS_MAINTENANCE_LAYER2
+		HasKHRMaintenance2 = HasExtension(DeviceExtensions, VK_KHR_MAINTENANCE2_EXTENSION_NAME);
+#endif
+#if VULKAN_SUPPORTS_DEDICATED_ALLOCATION
+		HasKHRDedicatedAllocation = HasExtension(DeviceExtensions, VK_KHR_GET_MEMORY_REQUIREMENTS_2_EXTENSION_NAME) && HasExtension(DeviceExtensions, VK_KHR_DEDICATED_ALLOCATION_EXTENSION_NAME);
 #endif
+#if VULKAN_SUPPORTS_COLOR_CONVERSIONS
+		HasYcbcrSampler = HasExtension(DeviceExtensions, VK_KHR_SAMPLER_YCBCR_CONVERSION_EXTENSION_NAME) && HasExtension(DeviceExtensions, VK_KHR_BIND_MEMORY_2_EXTENSION_NAME) && HasExtension(DeviceExtensions, VK_KHR_GET_MEMORY_REQUIREMENTS_2_EXTENSION_NAME);
+#endif
+#if VULKAN_SUPPORTS_MULTIVIEW
+		HasKHRMultiview = HasExtension(DeviceExtensions, VK_KHR_MULTIVIEW_EXTENSION_NAME);
+#endif
+	}
 
-#if VULKAN_SUPPORTS_IMAGE_64BIT_ATOMICS
-	HasImageAtomicInt64 = HasExtension(DeviceExtensions, VK_EXT_SHADER_IMAGE_ATOMIC_INT64_EXTENSION_NAME);
+
+	// Promoted to 1.2
+	{
+#if VULKAN_SUPPORTS_DRIVER_PROPERTIES
+		HasDriverProperties = HasExtension(DeviceExtensions, VK_KHR_DRIVER_PROPERTIES_EXTENSION_NAME);
 #endif
+#if VULKAN_SUPPORTS_RENDERPASS2
+		HasKHRRenderPass2 = HasExtension(DeviceExtensions, VK_KHR_CREATE_RENDERPASS_2_EXTENSION_NAME);
+#endif
+
+		HasKHRImageFormatList = HasExtension(DeviceExtensions, VK_KHR_IMAGE_FORMAT_LIST_EXTENSION_NAME);
 
+#if VULKAN_SUPPORTS_BUFFER_64BIT_ATOMICS
+		HasKHRShaderAtomicInt64 = HasExtension(DeviceExtensions, VK_KHR_SHADER_ATOMIC_INT64_EXTENSION_NAME);
+#endif
 #if VULKAN_SUPPORTS_SCALAR_BLOCK_LAYOUT
-	HasScalarBlockLayoutFeatures = HasExtension(DeviceExtensions, VK_EXT_SCALAR_BLOCK_LAYOUT_EXTENSION_NAME);
+		HasEXTScalarBlockLayout = HasExtension(DeviceExtensions, VK_EXT_SCALAR_BLOCK_LAYOUT_EXTENSION_NAME);
 #endif
+#if VULKAN_SUPPORTS_DESCRIPTOR_INDEXING
+		HasEXTDescriptorIndexing = HasExtension(DeviceExtensions, VK_EXT_DESCRIPTOR_INDEXING_EXTENSION_NAME);
+#endif
+#if VULKAN_SUPPORTS_SHADER_VIEWPORT_INDEX_LAYER
+		HasEXTShaderViewportIndexLayer = HasExtension(DeviceExtensions, VK_EXT_SHADER_VIEWPORT_INDEX_LAYER_EXTENSION_NAME);
+#endif
+	}
+
+
+	// Promoted to 1.3
+	{
+#if VULKAN_SUPPORTS_TEXTURE_COMPRESSION_ASTC_HDR
+		HasEXTTextureCompressionASTCHDR = HasExtension(DeviceExtensions, VK_EXT_TEXTURE_COMPRESSION_ASTC_HDR_EXTENSION_NAME);
+#endif
+#if VULKAN_SUPPORTS_MAINTENANCE4
+		HasKHRMaintenance4 = HasExtension(DeviceExtensions, VK_KHR_MAINTENANCE_4_EXTENSION_NAME);
+#endif
+	}
+
 
 #if VULKAN_RHI_RAYTRACING
+	// Optional Extensions
 	HasAccelerationStructure	= HasExtension(DeviceExtensions, VK_KHR_ACCELERATION_STRUCTURE_EXTENSION_NAME);
 	HasRayTracingPipeline		= HasExtension(DeviceExtensions, VK_KHR_RAY_TRACING_PIPELINE_EXTENSION_NAME);
 	HasRayQuery		            = HasExtension(DeviceExtensions, VK_KHR_RAY_QUERY_EXTENSION_NAME);
-	HasDescriptorIndexing		= HasExtension(DeviceExtensions, VK_EXT_DESCRIPTOR_INDEXING_EXTENSION_NAME);
-	HasBufferDeviceAddress		= HasExtension(DeviceExtensions, VK_KHR_BUFFER_DEVICE_ADDRESS_EXTENSION_NAME);
 	HasDeferredHostOperations	= HasExtension(DeviceExtensions, VK_KHR_DEFERRED_HOST_OPERATIONS_EXTENSION_NAME);
+
+	// Promoted to 1.2
+	HasBufferDeviceAddress		= HasExtension(DeviceExtensions, VK_KHR_BUFFER_DEVICE_ADDRESS_EXTENSION_NAME);
 	HasSPIRV_14					= HasExtension(DeviceExtensions, VK_KHR_SPIRV_1_4_EXTENSION_NAME);
 	HasShaderFloatControls		= HasExtension(DeviceExtensions, VK_KHR_SHADER_FLOAT_CONTROLS_EXTENSION_NAME);
 #endif
-
-#if VULKAN_SUPPORTS_SHADER_VIEWPORT_INDEX_LAYER
-	HasEXTShaderViewportIndexLayer = HasExtension(DeviceExtensions, VK_EXT_SHADER_VIEWPORT_INDEX_LAYER_EXTENSION_NAME);
-#endif
 }
 
 void FVulkanDynamicRHI::SetupValidationRequests()
diff --git a/Engine/Source/Runtime/VulkanRHI/Public/VulkanConfiguration.h b/Engine/Source/Runtime/VulkanRHI/Public/VulkanConfiguration.h
--- a/Engine/Source/Runtime/VulkanRHI/Public/VulkanConfiguration.h
+++ b/Engine/Source/Runtime/VulkanRHI/Public/VulkanConfiguration.h
@@ -337,6 +337,18 @@
 	#define VULKAN_SUPPORTS_SHADER_VIEWPORT_INDEX_LAYER	0
 #endif
 
+#ifdef VK_KHR_maintenance4
+	#define VULKAN_SUPPORTS_MAINTENANCE4	(VULKAN_SUPPORTS_PHYSICAL_DEVICE_PROPERTIES2)	// Requirement
+#else
+	#define VULKAN_SUPPORTS_MAINTENANCE4	0
+#endif
+
+#ifdef VK_EXT_descriptor_indexing
+	#define VULKAN_SUPPORTS_DESCRIPTOR_INDEXING	(VULKAN_SUPPORTS_PHYSICAL_DEVICE_PROPERTIES2)	// Requirement
+#else
+	#define VULKAN_SUPPORTS_DESCRIPTOR_INDEXING	0
+#endif
+
 
 #ifndef VULKAN_OBJECT_TRACKING 
 #define VULKAN_OBJECT_TRACKING 0 //Track objects created and memory used. use r.vulkan.dumpmemory to dump to console
