VulkanRHI: Support more recent versions of Vulkan 1.2 SDKs in case some devs have it installed locally.

diff --git a/Engine/Source/Runtime/VulkanRHI/Private/VulkanDevice.cpp b/Engine/Source/Runtime/VulkanRHI/Private/VulkanDevice.cpp
--- a/Engine/Source/Runtime/VulkanRHI/Private/VulkanDevice.cpp
+++ b/Engine/Source/Runtime/VulkanRHI/Private/VulkanDevice.cpp
@@ -1274,7 +1274,7 @@ void FVulkanDevice::InitGPU(int32 DeviceIndex)
 
 #if VULKAN_SUPPORTS_MAINTENANCE4
 		VkPhysicalDeviceMaintenance4FeaturesKHR Maintenance4Features;
-		ZeroVulkanStruct(Maintenance4Features, VK_STRUCTURE_TYPE_PHYSICAL_DEVICE_MAINTENANCE_4_FEATURES);
+		ZeroVulkanStruct(Maintenance4Features, VK_STRUCTURE_TYPE_PHYSICAL_DEVICE_MAINTENANCE_4_FEATURES_KHR);
 		if (OptionalDeviceExtensions.HasKHRMaintenance4)  // Set by QueryGPU if extension is supported
 		{
 			*NextPropsAddr = &Maintenance4Features;
