SubmitCommandsAndFlushGPU might update fence state if it was tied to a previously submitted command buffer.  Its state will have been updated from Submitted to NeedReset, and would assert in WaitForCmdBuffer (which is not needed anymore in such a case).

diff --git a/Engine/Source/Runtime/VulkanRHI/Private/VulkanUtil.cpp b/Engine/Source/Runtime/VulkanRHI/Private/VulkanUtil.cpp
--- a/Engine/Source/Runtime/VulkanRHI/Private/VulkanUtil.cpp
+++ b/Engine/Source/Runtime/VulkanRHI/Private/VulkanUtil.cpp
@@ -133,8 +133,14 @@ void* FVulkanDynamicRHI::RHILockStagingBuffer(FRHIStagingBuffer* StagingBufferRH
 	if (FenceRHI && !FenceRHI->Poll())
 	{
 		Device->SubmitCommandsAndFlushGPU();
-		FVulkanGPUFence* Fence = ResourceCast(FenceRHI);
-		Device->GetImmediateContext().GetCommandBufferManager()->WaitForCmdBuffer(Fence->GetCmdBuffer());
+
+		// SubmitCommandsAndFlushGPU might update fence state if it was tied to a previously submitted command buffer.
+		// Its state will have been updated from Submitted to NeedReset, and would assert in WaitForCmdBuffer (which is not needed in such a case)
+		if (!FenceRHI->Poll())
+		{
+			FVulkanGPUFence* Fence = ResourceCast(FenceRHI);
+			Device->GetImmediateContext().GetCommandBufferManager()->WaitForCmdBuffer(Fence->GetCmdBuffer());
+		}
 	}
 
 	return StagingBuffer->Lock(Offset, NumBytes);
