VulkanRHI:  Remove no longer used VULKAN_SUPPORTS_COLOR_CONVERSIONS code (including FSamplerYcbcrConversionInitializer).

diff --git a/Engine/Source/Runtime/VulkanRHI/Private/VulkanCommandWrappers.h b/Engine/Source/Runtime/VulkanRHI/Private/VulkanCommandWrappers.h
--- a/Engine/Source/Runtime/VulkanRHI/Private/VulkanCommandWrappers.h
+++ b/Engine/Source/Runtime/VulkanRHI/Private/VulkanCommandWrappers.h
@@ -190,10 +190,6 @@ struct FWrapLayer
 #if VULKAN_USE_CREATE_WIN32_SURFACE
 	static void CreateWin32SurfaceKHR(VkResult Result, VkInstance instance, const VkWin32SurfaceCreateInfoKHR* pCreateInfo, VkSurfaceKHR* pSurface);
 #endif
-#if VULKAN_SUPPORTS_COLOR_CONVERSIONS
-	static void CreateSamplerYcbcrConversionKHR(VkResult Result, VkDevice device, const VkSamplerYcbcrConversionCreateInfo* pCreateInfo, VkSamplerYcbcrConversion* YcbcrConversion) VULKAN_LAYER_BODY
-	static void DestroySamplerYcbcrConversionKHR(VkResult Result, VkDevice device, VkSamplerYcbcrConversion ycbcrConversion) VULKAN_LAYER_BODY
-#endif
 #if VULKAN_RHI_RAYTRACING
 	static void CreateAccelerationStructureKHR(VkResult Result, VkDevice Device, const VkAccelerationStructureCreateInfoKHR* CreateInfo, const VkAllocationCallbacks* Allocator, VkAccelerationStructureKHR* AccelerationStructure);
 	static void DestroyAccelerationStructureKHR(VkResult Result, VkDevice Device, VkAccelerationStructureKHR AccelerationStructure, const VkAllocationCallbacks* Allocator);
@@ -1403,23 +1399,6 @@ namespace VulkanRHI
 	}
 #endif
 
-#if VULKAN_SUPPORTS_COLOR_CONVERSIONS
-	static FORCEINLINE_DEBUGGABLE VkResult vkCreateSamplerYcbcrConversionKHR(VkDevice Device, const VkSamplerYcbcrConversionCreateInfo* CreateInfo, const VkAllocationCallbacks* Allocator, VkSamplerYcbcrConversion* YcbcrConversion)
-	{
-		FWrapLayer::CreateSamplerYcbcrConversionKHR(VK_RESULT_MAX_ENUM, Device, CreateInfo, YcbcrConversion);
-		VkResult Result = VULKANAPINAMESPACE::vkCreateSamplerYcbcrConversionKHR(Device, CreateInfo, Allocator, YcbcrConversion);
-		FWrapLayer::CreateSamplerYcbcrConversionKHR(Result, Device, CreateInfo, YcbcrConversion);
-		return Result;
-	}
-
-	static FORCEINLINE_DEBUGGABLE void vkDestroySamplerYcbcrConversionKHR(VkDevice Device, VkSamplerYcbcrConversion YcbcrConversion, const VkAllocationCallbacks* Allocator)
-	{
-		FWrapLayer::DestroySamplerYcbcrConversionKHR(VK_RESULT_MAX_ENUM, Device, YcbcrConversion);
-		VULKANAPINAMESPACE::vkDestroySamplerYcbcrConversionKHR(Device, YcbcrConversion, Allocator);
-		FWrapLayer::DestroySamplerYcbcrConversionKHR(VK_SUCCESS, Device, YcbcrConversion);
-	}
-#endif
-
 #if VULKAN_RHI_RAYTRACING
 	static FORCEINLINE_DEBUGGABLE VkResult vkCreateAccelerationStructureKHR(VkDevice Device, const VkAccelerationStructureCreateInfoKHR* CreateInfo, const VkAllocationCallbacks* Allocator, VkAccelerationStructureKHR* AccelerationStructure)
 	{
diff --git a/Engine/Source/Runtime/VulkanRHI/Private/VulkanDebug.cpp b/Engine/Source/Runtime/VulkanRHI/Private/VulkanDebug.cpp
--- a/Engine/Source/Runtime/VulkanRHI/Private/VulkanDebug.cpp
+++ b/Engine/Source/Runtime/VulkanRHI/Private/VulkanDebug.cpp
@@ -4310,33 +4310,6 @@ void FWrapLayer::CreateWin32SurfaceKHR(VkResult Result, VkInstance Instance, con
 }
 #endif
 
-#if VULKAN_SUPPORTS_COLOR_CONVERSIONS
-void FWrapLayer::CreateSamplerYcbcrConversionKHR(VkResult Result, VkDevice Device, const VkSamplerYcbcrConversionCreateInfo* CreateInfo, VkSamplerYcbcrConversion* YcbcrConversion)
-{
-#if VULKAN_ENABLE_DUMP_LAYER
-	if (Result == VK_RESULT_MAX_ENUM)
-	{
-		DevicePrintfBeginResult(Device, FString::Printf(TEXT("vkCreateSamplerYcbcrConversionKHR(CreateInfo=0x%p, YcbcrConversion=0x%p)[...]"), CreateInfo, YcbcrConversion));
-	}
-	else
-	{
-		PrintResultAndNamedHandle(Result, TEXT("SamplerYcbcrConversionKHR"), *YcbcrConversion);
-	}
-	FlushDebugWrapperLog();
-#endif
-}
-
-void FWrapLayer::DestroySamplerYcbcrConversionKHR(VkResult Result, VkDevice Device, VkSamplerYcbcrConversion YcbcrConversion)
-{
-	if (Result == VK_RESULT_MAX_ENUM)
-	{
-#if VULKAN_ENABLE_DUMP_LAYER
-		DevicePrintfBegin(Device, FString::Printf(TEXT("vkDestroySamplerYcbcrConversionKHR(YcbcrConversion=0x%p)"), YcbcrConversion));
-#endif
-	}
-}
-#endif
-
 #if VULKAN_RHI_RAYTRACING
 void FWrapLayer::CreateAccelerationStructureKHR(VkResult Result, VkDevice Device, const VkAccelerationStructureCreateInfoKHR* CreateInfo, const VkAllocationCallbacks* Allocator, VkAccelerationStructureKHR* AccelerationStructure)
 {
diff --git a/Engine/Source/Runtime/VulkanRHI/Private/VulkanDevice.cpp b/Engine/Source/Runtime/VulkanRHI/Private/VulkanDevice.cpp
--- a/Engine/Source/Runtime/VulkanRHI/Private/VulkanDevice.cpp
+++ b/Engine/Source/Runtime/VulkanRHI/Private/VulkanDevice.cpp
@@ -806,26 +806,6 @@ void FVulkanDevice::SetupFormats()
 #endif
 }
 
-#if VULKAN_SUPPORTS_COLOR_CONVERSIONS
-VkSamplerYcbcrConversion FVulkanDevice::CreateSamplerColorConversion(const VkSamplerYcbcrConversionCreateInfo& CreateInfo)
-{
-	const uint32 CreateInfoHash = FCrc::MemCrc32(&CreateInfo, sizeof(CreateInfo));
-	VkSamplerYcbcrConversion* const FindResult = SamplerColorConversionMap.Find(CreateInfoHash);
-	if (FindResult != nullptr)
-	{
-		return *FindResult;
-	}
-	else
-	{
-		VkSamplerYcbcrConversion NewConversion;
-		VERIFYVULKANRESULT(VulkanRHI::vkCreateSamplerYcbcrConversionKHR(GetInstanceHandle(), &CreateInfo, VULKAN_CPU_ALLOCATOR, &NewConversion));
-		SamplerColorConversionMap.Add(CreateInfoHash, NewConversion);
-		return NewConversion;
-	}
-}
-#endif
-
-
 const VkFormatProperties& FVulkanDevice::GetFormatProperties(VkFormat InFormat)
 {
 	if (InFormat >= 0 && InFormat < VK_FORMAT_RANGE_SIZE)
@@ -1493,14 +1473,6 @@ void FVulkanDevice::Destroy()
 	delete DefaultImage;
 	DefaultImage = nullptr;
 
-#if VULKAN_SUPPORTS_COLOR_CONVERSIONS
-	for (const auto& Pair : SamplerColorConversionMap)
-	{
-		VulkanRHI::vkDestroySamplerYcbcrConversionKHR(GetInstanceHandle(), Pair.Value, VULKAN_CPU_ALLOCATOR);
-	}
-	SamplerColorConversionMap.Reset();
-#endif
-
 	for (int32 Index = CommandContexts.Num() - 1; Index >= 0; --Index)
 	{
 		delete CommandContexts[Index];
diff --git a/Engine/Source/Runtime/VulkanRHI/Private/VulkanDevice.h b/Engine/Source/Runtime/VulkanRHI/Private/VulkanDevice.h
--- a/Engine/Source/Runtime/VulkanRHI/Private/VulkanDevice.h
+++ b/Engine/Source/Runtime/VulkanRHI/Private/VulkanDevice.h
@@ -525,10 +525,6 @@ public:
 
 	void SetupPresentQueue(VkSurfaceKHR Surface);
 
-#if VULKAN_SUPPORTS_COLOR_CONVERSIONS
-	VkSamplerYcbcrConversion CreateSamplerColorConversion(const VkSamplerYcbcrConversionCreateInfo& CreateInfo);
-#endif
-
 	inline const TArray<VkQueueFamilyProperties>& GetQueueFamilyProps()
 	{
 		return QueueFamilyProps;
@@ -642,9 +638,6 @@ private:
 	FVulkanCommandListContextImmediate* ImmediateContext;
 	FVulkanCommandListContext* ComputeContext;
 	TArray<FVulkanCommandListContext*> CommandContexts;
-#if VULKAN_SUPPORTS_COLOR_CONVERSIONS
-	TMap<uint32, VkSamplerYcbcrConversion> SamplerColorConversionMap;
-#endif
 
 	FVulkanDynamicRHI* RHI = nullptr;
 	bool bDebugMarkersFound = false;
diff --git a/Engine/Source/Runtime/VulkanRHI/Private/VulkanLayers.cpp b/Engine/Source/Runtime/VulkanRHI/Private/VulkanLayers.cpp
--- a/Engine/Source/Runtime/VulkanRHI/Private/VulkanLayers.cpp
+++ b/Engine/Source/Runtime/VulkanRHI/Private/VulkanLayers.cpp
@@ -856,9 +856,6 @@ void FOptionalVulkanDeviceExtensions::Setup(const TArray<const ANSICHAR*>& Devic
 #if VULKAN_SUPPORTS_DEDICATED_ALLOCATION
 		HasKHRDedicatedAllocation = HasExtension(DeviceExtensions, VK_KHR_GET_MEMORY_REQUIREMENTS_2_EXTENSION_NAME) && HasExtension(DeviceExtensions, VK_KHR_DEDICATED_ALLOCATION_EXTENSION_NAME);
 #endif
-#if VULKAN_SUPPORTS_COLOR_CONVERSIONS
-		HasYcbcrSampler = HasExtension(DeviceExtensions, VK_KHR_SAMPLER_YCBCR_CONVERSION_EXTENSION_NAME) && HasExtension(DeviceExtensions, VK_KHR_BIND_MEMORY_2_EXTENSION_NAME) && HasExtension(DeviceExtensions, VK_KHR_GET_MEMORY_REQUIREMENTS_2_EXTENSION_NAME);
-#endif
 #if VULKAN_SUPPORTS_MULTIVIEW
 		HasKHRMultiview = HasExtension(DeviceExtensions, VK_KHR_MULTIVIEW_EXTENSION_NAME);
 #endif
diff --git a/Engine/Source/Runtime/VulkanRHI/Private/VulkanPipeline.cpp b/Engine/Source/Runtime/VulkanRHI/Private/VulkanPipeline.cpp
--- a/Engine/Source/Runtime/VulkanRHI/Private/VulkanPipeline.cpp
+++ b/Engine/Source/Runtime/VulkanRHI/Private/VulkanPipeline.cpp
@@ -1030,15 +1030,6 @@ FArchive& operator << (FArchive& Ar, FGfxPipelineDesc& Entry)
 #endif
 	Ar << Entry.RenderTargets;
 
-#if VULKAN_SUPPORTS_COLOR_CONVERSIONS
-	for (uint32 Index = 0; Index < MaxImmutableSamplers; ++Index)
-	{
-		uint64 Sampler = (uint64)Entry.ImmutableSamplers[Index];
-		Ar << Sampler;
-		Entry.ImmutableSamplers[Index] = (SIZE_T)Sampler;
-	}
-#endif
-
 #if VULKAN_SUPPORTS_FRAGMENT_SHADING_RATE
 	uint8 ShadingRate = static_cast<uint8>(Entry.ShadingRate);
 	uint8 Combiner = static_cast<uint8>(Entry.Combiner);
@@ -1622,14 +1613,6 @@ void FVulkanPipelineStateCacheManager::CreateGfxEntry(const FGraphicsPipelineSta
 #endif
 	check(NumShaders > 0);
 
-#if VULKAN_SUPPORTS_COLOR_CONVERSIONS
-	for (uint32 Index = 0; Index < MaxImmutableSamplers; ++Index)
-	{
-		OutGfxEntry->ImmutableSamplers[Index] = reinterpret_cast<SIZE_T>(PSOInitializer.ImmutableSamplerState.ImmutableSamplers[Index]);
-	}
-#endif
-
-
 	FVulkanRenderTargetLayout RTLayout(PSOInitializer);
 	OutGfxEntry->RenderTargets.ReadFrom(RTLayout);
 
diff --git a/Engine/Source/Runtime/VulkanRHI/Private/VulkanPipeline.h b/Engine/Source/Runtime/VulkanRHI/Private/VulkanPipeline.h
--- a/Engine/Source/Runtime/VulkanRHI/Private/VulkanPipeline.h
+++ b/Engine/Source/Runtime/VulkanRHI/Private/VulkanPipeline.h
@@ -279,10 +279,6 @@ struct FGfxPipelineDesc
 	FVulkanShaderHashes ShaderHashes;
 #endif
 
-#if VULKAN_SUPPORTS_COLOR_CONVERSIONS
-	SIZE_T ImmutableSamplers[MaxImmutableSamplers] = { 0 };
-#endif
-
 	struct FRenderTargets
 	{
 		struct FAttachmentRef
@@ -428,16 +424,6 @@ struct FGfxPipelineDesc
 		}
 #endif
 
-#if VULKAN_SUPPORTS_COLOR_CONVERSIONS
-		for (uint32 Index = 0; Index < MaxImmutableSamplers; ++Index)
-		{
-			if (ImmutableSamplers[Index] != In.ImmutableSamplers[Index])
-			{
-				return false;
-			}
-		}
-#endif
-
 		if (!(RenderTargets == In.RenderTargets))
 		{
 			return false;
diff --git a/Engine/Source/Runtime/VulkanRHI/Private/VulkanRHI.cpp b/Engine/Source/Runtime/VulkanRHI/Private/VulkanRHI.cpp
--- a/Engine/Source/Runtime/VulkanRHI/Private/VulkanRHI.cpp
+++ b/Engine/Source/Runtime/VulkanRHI/Private/VulkanRHI.cpp
@@ -1230,12 +1230,6 @@ FTexture2DRHIRef FVulkanDynamicRHI::RHICreateTexture2DFromResource(EPixelFormat
 	return new FVulkanTexture2D(*Device, Format, SizeX, SizeY, NumMips, NumSamples, Resource, Flags, ResourceCreateInfo);
 }
 
-FTexture2DRHIRef FVulkanDynamicRHI::RHICreateTexture2DFromResource(EPixelFormat Format, uint32 SizeX, uint32 SizeY, uint32 NumMips, uint32 NumSamples, VkImage Resource, FSamplerYcbcrConversionInitializer& ConversionInitializer, ETextureCreateFlags Flags)
-{
-	const FRHIResourceCreateInfo ResourceCreateInfo(TEXT("FVulkanTexture2D"), EnumHasAnyFlags(Flags, TexCreate_DepthStencilTargetable) ? FClearValueBinding::DepthZero : FClearValueBinding::Transparent);
-	return new FVulkanTexture2D(*Device, Format, SizeX, SizeY, NumMips, NumSamples, Resource, ConversionInitializer, Flags, ResourceCreateInfo);
-}
-
 FTexture2DArrayRHIRef FVulkanDynamicRHI::RHICreateTexture2DArrayFromResource(EPixelFormat Format, uint32 SizeX, uint32 SizeY, uint32 ArraySize, uint32 NumMips, uint32 NumSamples, VkImage Resource, ETextureCreateFlags Flags)
 {
 	const FRHIResourceCreateInfo ResourceCreateInfo(TEXT("FVulkanTexture2DArray"), EnumHasAnyFlags(Flags, TexCreate_DepthStencilTargetable) ? FClearValueBinding::DepthZero : FClearValueBinding::Transparent);
@@ -1384,18 +1378,6 @@ void FVulkanDescriptorSetsLayoutInfo::GenerateHash(const TArrayView<FRHISamplerS
 		TArray<uint16>& PackedUBBindingIndices = RemappingInfo.StageInfos[RemapingIndex].PackedUBBindingIndices;
 		Hash = FCrc::MemCrc32(PackedUBBindingIndices.GetData(), sizeof(uint16) * PackedUBBindingIndices.Num(), Hash);
 	}
-
-#if VULKAN_SUPPORTS_COLOR_CONVERSIONS
-	VkSampler ImmutableSamplers[MaxImmutableSamplers];
-	VkSampler* ImmutableSamplerPtr = ImmutableSamplers;
-	for (int32 Index = 0; Index < InImmutableSamplers.Num(); ++Index)
-	{
-		FRHISamplerState* SamplerState = InImmutableSamplers[Index];
-		*ImmutableSamplerPtr++ = SamplerState ? ResourceCast(SamplerState)->Sampler : VK_NULL_HANDLE;
-	}
-	FMemory::Memzero(ImmutableSamplerPtr, (MaxImmutableSamplers - InImmutableSamplers.Num()));
-	Hash = FCrc::MemCrc32(ImmutableSamplers, sizeof(VkSampler) * MaxImmutableSamplers, Hash);
-#endif
 }
 
 static FCriticalSection GTypesUsageCS;
diff --git a/Engine/Source/Runtime/VulkanRHI/Private/VulkanResources.h b/Engine/Source/Runtime/VulkanRHI/Private/VulkanResources.h
--- a/Engine/Source/Runtime/VulkanRHI/Private/VulkanResources.h
+++ b/Engine/Source/Runtime/VulkanRHI/Private/VulkanResources.h
@@ -44,17 +44,6 @@ enum
 	NUM_TIMESTAMP_QUERIES_PER_POOL = 1024,
 };
 
-struct FSamplerYcbcrConversionInitializer
-{
-	VkFormat Format;
-	uint64 ExternalFormat;
-	VkComponentMapping Components;
-	VkSamplerYcbcrModelConversion Model;
-	VkSamplerYcbcrRange Range;
-	VkChromaLocation XOffset;
-	VkChromaLocation YOffset;
-};
-
 // Mirror GPixelFormats with format information for buffers
 extern VkFormat GVulkanBufferFormat[PF_MAX];
 
@@ -521,7 +510,6 @@ struct FVulkanTextureView
 	}
 
 	void Create(FVulkanDevice& Device, VkImage InImage, VkImageViewType ViewType, VkImageAspectFlags AspectFlags, EPixelFormat UEFormat, VkFormat Format, uint32 FirstMip, uint32 NumMips, uint32 ArraySliceIndex, uint32 NumArraySlices, bool bUseIdentitySwizzle = false);
-	void Create(FVulkanDevice& Device, VkImage InImage, VkImageViewType ViewType, VkImageAspectFlags AspectFlags, EPixelFormat UEFormat, VkFormat Format, uint32 FirstMip, uint32 NumMips, uint32 ArraySliceIndex, uint32 NumArraySlices, FSamplerYcbcrConversionInitializer& ConversionInitializer, bool bUseIdentitySwizzle = false);
 	void Destroy(FVulkanDevice& Device);
 
 	VkImageView View;
@@ -529,7 +517,7 @@ struct FVulkanTextureView
 	uint32 ViewId;
 
 private:
-	static VkImageView StaticCreate(FVulkanDevice& Device, VkImage InImage, VkImageViewType ViewType, VkImageAspectFlags AspectFlags, EPixelFormat UEFormat, VkFormat Format, uint32 FirstMip, uint32 NumMips, uint32 ArraySliceIndex, uint32 NumArraySlices, bool bUseIdentitySwizzle, const FSamplerYcbcrConversionInitializer* ConversionInitializer);
+	static VkImageView StaticCreate(FVulkanDevice& Device, VkImage InImage, VkImageViewType ViewType, VkImageAspectFlags AspectFlags, EPixelFormat UEFormat, VkFormat Format, uint32 FirstMip, uint32 NumMips, uint32 ArraySliceIndex, uint32 NumArraySlices, bool bUseIdentitySwizzle);
 };
 
 
@@ -543,7 +531,6 @@ struct FVulkanTextureBase : public FVulkanEvictable, public IRefCountedObject
 
 	FVulkanTextureBase(FVulkanDevice& Device, VkImageViewType ResourceType, EPixelFormat Format, uint32 SizeX, uint32 SizeY, uint32 SizeZ, uint32 ArraySize, uint32 NumMips, uint32 NumSamples, ETextureCreateFlags UEFlags, ERHIAccess InResourceState, const FRHIResourceCreateInfo& CreateInfo, const FRHITransientHeapAllocation* InTransientHeapAllocation);
 	FVulkanTextureBase(FVulkanDevice& Device, VkImageViewType ResourceType, EPixelFormat Format, uint32 SizeX, uint32 SizeY, uint32 SizeZ, uint32 ArraySize, uint32 NumMips, uint32 NumSamples, VkImage InImage, VkDeviceMemory InMem, ETextureCreateFlags UEFlags, const FRHIResourceCreateInfo& CreateInfo);
-	FVulkanTextureBase(FVulkanDevice& Device, VkImageViewType ResourceType, EPixelFormat Format, uint32 SizeX, uint32 SizeY, uint32 SizeZ, uint32 ArraySize, uint32 NumMips, uint32 NumSamples, VkImage InImage, VkDeviceMemory InMem, FSamplerYcbcrConversionInitializer& ConversionInitializer, ETextureCreateFlags UEFlags, const FRHIResourceCreateInfo& CreateInfo);
 
 	// Aliasing constructor.
 	FVulkanTextureBase(FTextureRHIRef& SrcTextureRHI, const FVulkanTextureBase* SrcTexture, VkImageViewType ResourceType, uint32 SizeX, uint32 SizeY, uint32 sizeZ);
@@ -597,7 +584,6 @@ class FVulkanTexture2D : public FRHITexture2D, public FVulkanTextureBase
 public:
 	FVulkanTexture2D(FVulkanDevice& Device, EPixelFormat Format, uint32 SizeX, uint32 SizeY, uint32 NumMips, uint32 NumSamples, ETextureCreateFlags UEFlags, ERHIAccess InResourceState, const FRHIResourceCreateInfo& CreateInfo, const FRHITransientHeapAllocation* InTransientHeapAllocation = nullptr);
 	FVulkanTexture2D(FVulkanDevice& Device, EPixelFormat Format, uint32 SizeX, uint32 SizeY, uint32 NumMips, uint32 NumSamples, VkImage Image, ETextureCreateFlags UEFlags, const FRHIResourceCreateInfo& CreateInfo);
-	FVulkanTexture2D(FVulkanDevice& Device, EPixelFormat Format, uint32 SizeX, uint32 SizeY, uint32 NumMips, uint32 NumSamples, VkImage Image, struct FSamplerYcbcrConversionInitializer& ConversionInitializer, ETextureCreateFlags UEFlags, const FRHIResourceCreateInfo& CreateInfo);
 
 	// Aliasing constructor
 	FVulkanTexture2D(FTextureRHIRef& SrcTextureRHI, const FVulkanTexture2D* SrcTexture);
diff --git a/Engine/Source/Runtime/VulkanRHI/Private/VulkanState.cpp b/Engine/Source/Runtime/VulkanRHI/Private/VulkanState.cpp
--- a/Engine/Source/Runtime/VulkanRHI/Private/VulkanState.cpp
+++ b/Engine/Source/Runtime/VulkanRHI/Private/VulkanState.cpp
@@ -346,43 +346,6 @@ FSamplerStateRHIRef FVulkanDynamicRHI::RHICreateSamplerState(const FSamplerState
 	}
 }
 
-#if VULKAN_SUPPORTS_COLOR_CONVERSIONS
-FSamplerStateRHIRef FVulkanDynamicRHI::RHICreateSamplerState(
-	const FSamplerStateInitializerRHI& Initializer, 
-	const FSamplerYcbcrConversionInitializer& ConversionInitializer)
-{
-	VkSamplerYcbcrConversionCreateInfo ConversionCreateInfo;
-	FMemory::Memzero(&ConversionCreateInfo, sizeof(VkSamplerYcbcrConversionCreateInfo));
-	ConversionCreateInfo.sType = VK_STRUCTURE_TYPE_SAMPLER_YCBCR_CONVERSION_CREATE_INFO;
-	ConversionCreateInfo.format = ConversionInitializer.Format;
-	
-	ConversionCreateInfo.components.a = ConversionInitializer.Components.a;
-	ConversionCreateInfo.components.r = ConversionInitializer.Components.r;
-	ConversionCreateInfo.components.g = ConversionInitializer.Components.g;
-	ConversionCreateInfo.components.b = ConversionInitializer.Components.b;
-	
-	ConversionCreateInfo.ycbcrModel = ConversionInitializer.Model;
-	ConversionCreateInfo.ycbcrRange = ConversionInitializer.Range;
-	ConversionCreateInfo.xChromaOffset = ConversionInitializer.XOffset;
-	ConversionCreateInfo.yChromaOffset = ConversionInitializer.YOffset;
-	ConversionCreateInfo.chromaFilter = VK_FILTER_NEAREST;
-	ConversionCreateInfo.forceExplicitReconstruction = VK_FALSE;
-
-	check(ConversionInitializer.Format != VK_FORMAT_UNDEFINED); // No support for VkExternalFormatANDROID yet.
-
-	VkSamplerYcbcrConversionInfo ConversionInfo;
-	FMemory::Memzero(&ConversionInfo, sizeof(VkSamplerYcbcrConversionInfo));
-	ConversionInfo.conversion = Device->CreateSamplerColorConversion(ConversionCreateInfo);
-	ConversionInfo.sType = VK_STRUCTURE_TYPE_SAMPLER_YCBCR_CONVERSION_INFO;
-
-	VkSamplerCreateInfo SamplerInfo;
-	FVulkanSamplerState::SetupSamplerCreateInfo(Initializer, *Device, SamplerInfo);
-	SamplerInfo.pNext = &ConversionInfo;
-
-	return new FVulkanSamplerState(SamplerInfo, *Device, true);
-}
-#endif
-
 FRasterizerStateRHIRef FVulkanDynamicRHI::RHICreateRasterizerState(const FRasterizerStateInitializerRHI& Initializer)
 {
 	return new FVulkanRasterizerState(Initializer);
diff --git a/Engine/Source/Runtime/VulkanRHI/Private/VulkanTexture.cpp b/Engine/Source/Runtime/VulkanRHI/Private/VulkanTexture.cpp
--- a/Engine/Source/Runtime/VulkanRHI/Private/VulkanTexture.cpp
+++ b/Engine/Source/Runtime/VulkanRHI/Private/VulkanTexture.cpp
@@ -1802,7 +1802,7 @@ void FVulkanDynamicRHI::InternalUpdateTexture3D(bool bFromRenderingThread, FRHIT
 }
 
 
-VkImageView FVulkanTextureView::StaticCreate(FVulkanDevice& Device, VkImage InImage, VkImageViewType ViewType, VkImageAspectFlags AspectFlags, EPixelFormat UEFormat, VkFormat Format, uint32 FirstMip, uint32 NumMips, uint32 ArraySliceIndex, uint32 NumArraySlices, bool bUseIdentitySwizzle, const FSamplerYcbcrConversionInitializer* ConversionInitializer)
+VkImageView FVulkanTextureView::StaticCreate(FVulkanDevice& Device, VkImage InImage, VkImageViewType ViewType, VkImageAspectFlags AspectFlags, EPixelFormat UEFormat, VkFormat Format, uint32 FirstMip, uint32 NumMips, uint32 ArraySliceIndex, uint32 NumArraySlices, bool bUseIdentitySwizzle)
 {
 	LLM_SCOPE_VULKAN(ELLMTagVulkan::VulkanTextures);
 	VkImageView OutView = VK_NULL_HANDLE;
@@ -1833,37 +1833,6 @@ VkImageView FVulkanTextureView::StaticCreate(FVulkanDevice& Device, VkImage InIm
 		ViewInfo.components = Device.GetFormatComponentMapping(UEFormat);
 	}
 
-#if VULKAN_SUPPORTS_COLOR_CONVERSIONS
-	VkSamplerYcbcrConversionInfo ConversionInfo;
-	if (ConversionInitializer != nullptr)
-	{
-		VkSamplerYcbcrConversionCreateInfo ConversionCreateInfo;
-		FMemory::Memzero(&ConversionCreateInfo, sizeof(VkSamplerYcbcrConversionCreateInfo));
-		ConversionCreateInfo.sType = VK_STRUCTURE_TYPE_SAMPLER_YCBCR_CONVERSION_CREATE_INFO;
-		ConversionCreateInfo.format = ConversionInitializer->Format;
-
-		ConversionCreateInfo.components.a = ConversionInitializer->Components.a;
-		ConversionCreateInfo.components.r = ConversionInitializer->Components.r;
-		ConversionCreateInfo.components.g = ConversionInitializer->Components.g;
-		ConversionCreateInfo.components.b = ConversionInitializer->Components.b;
-
-		ConversionCreateInfo.ycbcrModel = ConversionInitializer->Model;
-		ConversionCreateInfo.ycbcrRange = ConversionInitializer->Range;
-		ConversionCreateInfo.xChromaOffset = ConversionInitializer->XOffset;
-		ConversionCreateInfo.yChromaOffset = ConversionInitializer->YOffset;
-		ConversionCreateInfo.chromaFilter = VK_FILTER_NEAREST;
-		ConversionCreateInfo.forceExplicitReconstruction = VK_FALSE;
-
-		check(ConversionInitializer->Format != VK_FORMAT_UNDEFINED); // No support for VkExternalFormatANDROID yet.
-
-		FMemory::Memzero(&ConversionInfo, sizeof(VkSamplerYcbcrConversionInfo));
-		ConversionInfo.conversion = Device.CreateSamplerColorConversion(ConversionCreateInfo);
-		ConversionInfo.sType = VK_STRUCTURE_TYPE_SAMPLER_YCBCR_CONVERSION_INFO;
-		ConversionInfo.pNext = ViewInfo.pNext;
-		ViewInfo.pNext = &ConversionInfo;
-	}
-#endif
-
 	ViewInfo.subresourceRange.aspectMask = AspectFlags;
 	ViewInfo.subresourceRange.baseMipLevel = FirstMip;
 	ensure(NumMips != 0xFFFFFFFF);
@@ -1935,18 +1904,7 @@ VkImageView FVulkanTextureView::StaticCreate(FVulkanDevice& Device, VkImage InIm
 
 void FVulkanTextureView::Create(FVulkanDevice& Device, VkImage InImage, VkImageViewType ViewType, VkImageAspectFlags AspectFlags, EPixelFormat UEFormat, VkFormat Format, uint32 FirstMip, uint32 NumMips, uint32 ArraySliceIndex, uint32 NumArraySlices, bool bUseIdentitySwizzle)
 {
-	View = StaticCreate(Device, InImage, ViewType, AspectFlags, UEFormat, Format, FirstMip, NumMips, ArraySliceIndex, NumArraySlices, bUseIdentitySwizzle, nullptr);
-	Image = InImage;
-	
-	if (UseVulkanDescriptorCache())
-	{
-		ViewId = ++GVulkanImageViewHandleIdCounter;
-	}
-}
-
-void FVulkanTextureView::Create(FVulkanDevice& Device, VkImage InImage, VkImageViewType ViewType, VkImageAspectFlags AspectFlags, EPixelFormat UEFormat, VkFormat Format, uint32 FirstMip, uint32 NumMips, uint32 ArraySliceIndex, uint32 NumArraySlices, FSamplerYcbcrConversionInitializer& ConversionInitializer, bool bUseIdentitySwizzle)
-{
-	View = StaticCreate(Device, InImage, ViewType, AspectFlags, UEFormat, Format, FirstMip, NumMips, ArraySliceIndex, NumArraySlices, bUseIdentitySwizzle, &ConversionInitializer);
+	View = StaticCreate(Device, InImage, ViewType, AspectFlags, UEFormat, Format, FirstMip, NumMips, ArraySliceIndex, NumArraySlices, bUseIdentitySwizzle);
 	Image = InImage;
 	
 	if (UseVulkanDescriptorCache())
@@ -2074,39 +2032,6 @@ FVulkanTextureBase::FVulkanTextureBase(FVulkanDevice& Device, VkImageViewType Re
 	}
 }
 
-FVulkanTextureBase::FVulkanTextureBase(FVulkanDevice& Device, VkImageViewType ResourceType, EPixelFormat Format, uint32 SizeX, uint32 SizeY, uint32 SizeZ, uint32 ArraySize, uint32 NumMips, uint32 NumSamples, VkImage InImage, VkDeviceMemory InMem, FSamplerYcbcrConversionInitializer& ConversionInitializer, ETextureCreateFlags UEFlags, const FRHIResourceCreateInfo& CreateInfo)
-	: Surface(Device, ResourceType, Format, SizeX, SizeY, SizeZ, ArraySize, NumMips, NumSamples, InImage, UEFlags, EImageOwnerType::ExternalOwner, CreateInfo)
-	, PartialView(nullptr)
-{
-	Surface.OwningTexture = this;
-	VULKAN_TRACK_OBJECT_CREATE(FVulkanTextureBase, this);
-
-	check(InMem == VK_NULL_HANDLE);
-	const bool bArray = ResourceType == VK_IMAGE_VIEW_TYPE_1D_ARRAY || ResourceType == VK_IMAGE_VIEW_TYPE_2D_ARRAY || ResourceType == VK_IMAGE_VIEW_TYPE_CUBE_ARRAY;
-
-	Surface.ViewFormat = ConversionInitializer.Format;
-	Surface.StorageFormat = ConversionInitializer.Format;
-
-	if (ResourceType != VK_IMAGE_VIEW_TYPE_MAX_ENUM && Surface.Image != VK_NULL_HANDLE)
-	{
-		DefaultView.Create(Device, Surface.Image, ResourceType, Surface.GetFullAspectMask(), Format, Surface.ViewFormat, 0, FMath::Max(Surface.NumMips, 1u), 0, bArray ? FMath::Max(1u, ArraySize) : FMath::Max(1u, SizeZ), ConversionInitializer);
-	}
-
-	// No MSAA support
-	check(NumSamples == 1);
-	check(!EnumHasAnyFlags(UEFlags, TexCreate_RenderTargetable));
-
-	if (Surface.FullAspectMask == Surface.PartialAspectMask)
-	{
-		PartialView = &DefaultView;
-	}
-	else
-	{
-		PartialView = new FVulkanTextureView;
-		PartialView->Create(Device, Surface.Image, Surface.ViewType, Surface.PartialAspectMask, Surface.PixelFormat, Surface.ViewFormat, 0, FMath::Max(NumMips, 1u), 0, bArray ? FMath::Max(1u, ArraySize) : FMath::Max(1u, SizeZ), ConversionInitializer);
-	}
-}
-
 FVulkanTextureBase::FVulkanTextureBase(FTextureRHIRef& SrcTextureRHI, const FVulkanTextureBase* SrcTexture, VkImageViewType ResourceType, uint32 SizeX, uint32 SizeY, uint32 SizeZ)
 	: Surface(*SrcTexture->Surface.Device, ResourceType, SrcTexture->Surface.PixelFormat, SizeX, SizeY, SizeZ, SrcTexture->Surface.ArraySize, SrcTexture->Surface.NumMips, SrcTexture->Surface.NumSamples, SrcTexture->Surface.Image, SrcTexture->Surface.UEFlags, EImageOwnerType::Aliased, FRHIResourceCreateInfo(TEXT("FVulkanTextureBase")))
 	, PartialView(nullptr)
@@ -2306,12 +2231,6 @@ FVulkanTexture2D::FVulkanTexture2D(FVulkanDevice& Device, EPixelFormat Format, u
 {
 }
 
-FVulkanTexture2D::FVulkanTexture2D(FVulkanDevice& Device, EPixelFormat Format, uint32 SizeX, uint32 SizeY, uint32 NumMips, uint32 NumSamples, VkImage Image, FSamplerYcbcrConversionInitializer& ConversionInitializer, ETextureCreateFlags UEFlags, const FRHIResourceCreateInfo& CreateInfo)
-	: FRHITexture2D(SizeX, SizeY, NumMips, NumSamples, Format, UEFlags, CreateInfo.ClearValueBinding)
-	, FVulkanTextureBase(Device, VK_IMAGE_VIEW_TYPE_2D, Format, SizeX, SizeY, 1, 1, NumMips, NumSamples, Image, VK_NULL_HANDLE, ConversionInitializer, UEFlags, CreateInfo)
-{
-}
-
 FVulkanTexture2D::FVulkanTexture2D(FTextureRHIRef& SrcTextureRHI, const FVulkanTexture2D* SrcTexture)
 	: FRHITexture2D(SrcTexture->GetSizeX(), SrcTexture->GetSizeY(), SrcTexture->GetNumMips(), SrcTexture->GetNumSamples(), SrcTexture->GetFormat(), SrcTexture->GetFlags(), SrcTexture->GetClearBinding())
 	, FVulkanTextureBase(SrcTextureRHI, static_cast<const FVulkanTextureBase*>(SrcTexture), VK_IMAGE_VIEW_TYPE_2D, SrcTexture->GetSizeX(), SrcTexture->GetSizeY(), 1)
diff --git a/Engine/Source/Runtime/VulkanRHI/Private/Windows/VulkanWindowsPlatform.cpp b/Engine/Source/Runtime/VulkanRHI/Private/Windows/VulkanWindowsPlatform.cpp
--- a/Engine/Source/Runtime/VulkanRHI/Private/Windows/VulkanWindowsPlatform.cpp
+++ b/Engine/Source/Runtime/VulkanRHI/Private/Windows/VulkanWindowsPlatform.cpp
@@ -214,13 +214,6 @@ void FVulkanWindowsPlatform::GetDeviceExtensions(EGpuVendorId VendorId, TArray<c
 #endif
 	}
 
-#if VULKAN_SUPPORTS_COLOR_CONVERSIONS
-	// YCbCr requires BindMem2 and GetMemReqs2
-	OutExtensions.Add(VK_KHR_BIND_MEMORY_2_EXTENSION_NAME);
-	OutExtensions.Add(VK_KHR_GET_MEMORY_REQUIREMENTS_2_EXTENSION_NAME);
-	OutExtensions.Add(VK_KHR_SAMPLER_YCBCR_CONVERSION_EXTENSION_NAME);
-#endif
-
 #if VULKAN_SUPPORTS_FULLSCREEN_EXCLUSIVE
 	// Fullscreen requires Instance capabilities2
 	OutExtensions.Add(VK_EXT_FULL_SCREEN_EXCLUSIVE_EXTENSION_NAME);
diff --git a/Engine/Source/Runtime/VulkanRHI/Private/Windows/VulkanWindowsPlatform.h b/Engine/Source/Runtime/VulkanRHI/Private/Windows/VulkanWindowsPlatform.h
--- a/Engine/Source/Runtime/VulkanRHI/Private/Windows/VulkanWindowsPlatform.h
+++ b/Engine/Source/Runtime/VulkanRHI/Private/Windows/VulkanWindowsPlatform.h
@@ -14,7 +14,6 @@
 #define VULKAN_DYNAMICALLYLOADED					1
 #define VULKAN_SHOULD_ENABLE_DESKTOP_HMD_SUPPORT	1
 #define VULKAN_SIGNAL_UNIMPLEMENTED()				checkf(false, TEXT("Unimplemented vulkan functionality: %s"), StringCast<TCHAR>(__FUNCTION__).Get())
-#define VULKAN_SUPPORTS_COLOR_CONVERSIONS			1
 #define VULKAN_SUPPORTS_AMD_BUFFER_MARKER			1
 
 #define VULKAN_RHI_RAYTRACING 						(RHI_RAYTRACING)
diff --git a/Engine/Source/Runtime/VulkanRHI/Public/VulkanConfiguration.h b/Engine/Source/Runtime/VulkanRHI/Public/VulkanConfiguration.h
--- a/Engine/Source/Runtime/VulkanRHI/Public/VulkanConfiguration.h
+++ b/Engine/Source/Runtime/VulkanRHI/Public/VulkanConfiguration.h
@@ -156,10 +156,6 @@
 	#define VULKAN_DELETE_STALE_CMDBUFFERS						1
 #endif
 
-#ifndef VULKAN_SUPPORTS_COLOR_CONVERSIONS
-	#define VULKAN_SUPPORTS_COLOR_CONVERSIONS					0
-#endif
-
 #ifndef VULKAN_SUPPORTS_AMD_BUFFER_MARKER
 	#define VULKAN_SUPPORTS_AMD_BUFFER_MARKER					0
 #endif
diff --git a/Engine/Source/Runtime/VulkanRHI/Public/VulkanDynamicRHI.h b/Engine/Source/Runtime/VulkanRHI/Public/VulkanDynamicRHI.h
--- a/Engine/Source/Runtime/VulkanRHI/Public/VulkanDynamicRHI.h
+++ b/Engine/Source/Runtime/VulkanRHI/Public/VulkanDynamicRHI.h
@@ -61,9 +61,6 @@ public:
 	void InitInstance();
 
 	virtual FSamplerStateRHIRef RHICreateSamplerState(const FSamplerStateInitializerRHI& Initializer) final override;
-#if VULKAN_SUPPORTS_COLOR_CONVERSIONS
-	virtual FSamplerStateRHIRef RHICreateSamplerState(const FSamplerStateInitializerRHI& Initializer, const FSamplerYcbcrConversionInitializer& ConversionInitializer);
-#endif
 	virtual FRasterizerStateRHIRef RHICreateRasterizerState(const FRasterizerStateInitializerRHI& Initializer) final override;
 	virtual FDepthStencilStateRHIRef RHICreateDepthStencilState(const FDepthStencilStateInitializerRHI& Initializer) final override;
 	virtual FBlendStateRHIRef RHICreateBlendState(const FBlendStateInitializerRHI& Initializer) final override;
@@ -313,7 +310,6 @@ public:
 
 	// FVulkanDynamicRHI interface
 	virtual FTexture2DRHIRef RHICreateTexture2DFromResource(EPixelFormat Format, uint32 SizeX, uint32 SizeY, uint32 NumMips, uint32 NumSamples, VkImage Resource, ETextureCreateFlags Flags);
-	virtual FTexture2DRHIRef RHICreateTexture2DFromResource(EPixelFormat Format, uint32 SizeX, uint32 SizeY, uint32 NumMips, uint32 NumSamples, VkImage Resource, FSamplerYcbcrConversionInitializer& ConversionInitializer, ETextureCreateFlags Flags);
 	virtual FTexture2DArrayRHIRef RHICreateTexture2DArrayFromResource(EPixelFormat Format, uint32 SizeX, uint32 SizeY, uint32 ArraySize, uint32 NumMips, uint32 NumSamples, VkImage Resource, ETextureCreateFlags Flags);
 	virtual FTextureCubeRHIRef RHICreateTextureCubeFromResource(EPixelFormat Format, uint32 Size, bool bArray, uint32 ArraySize, uint32 NumMips, VkImage Resource, ETextureCreateFlags Flags);
 
