VulkanRHI: In cases where the driver version can't be read from the registry on Windows, fall back on reading it from the Vulkan properties (like on Linux).

diff --git a/Engine/Source/Runtime/VulkanRHI/Private/VulkanRHI.cpp b/Engine/Source/Runtime/VulkanRHI/Private/VulkanRHI.cpp
--- a/Engine/Source/Runtime/VulkanRHI/Private/VulkanRHI.cpp
+++ b/Engine/Source/Runtime/VulkanRHI/Private/VulkanRHI.cpp
@@ -719,28 +719,9 @@ void FVulkanDynamicRHI::SelectAndInitDevice()
 
 	Device->InitGPU(DeviceIndex);
 
-	if (PLATFORM_ANDROID)
-	{
-		GRHIAdapterName.Append(TEXT(" Vulkan"));
-		GRHIAdapterInternalDriverVersion = FString::Printf(TEXT("%d.%d.%d"), VK_VERSION_MAJOR(Props.apiVersion), VK_VERSION_MINOR(Props.apiVersion), VK_VERSION_PATCH(Props.apiVersion));
-	}
-	else if (PLATFORM_WINDOWS)
-	{
-		GRHIDeviceId = Props.deviceID;
-		FGPUDriverInfo GPUDriverInfo = FPlatformMisc::GetGPUDriverInfo(GRHIAdapterName);
-		GRHIAdapterUserDriverVersion = GPUDriverInfo.UserDriverVersion;
-		GRHIAdapterInternalDriverVersion = GPUDriverInfo.InternalDriverVersion;
-		GRHIAdapterDriverDate = GPUDriverInfo.DriverDate;
+	auto ReadVulkanDriverVersionFromProps = [](FVulkanDevice* Device) {
 
-		UE_LOG(LogVulkanRHI, Log, TEXT("    Adapter Name: %s"), *GRHIAdapterName);
-		UE_LOG(LogVulkanRHI, Log, TEXT("     API Version: %d.%d.%d"), VK_VERSION_MAJOR(Props.apiVersion), VK_VERSION_MINOR(Props.apiVersion), VK_VERSION_PATCH(Props.apiVersion));
-		UE_LOG(LogVulkanRHI, Log, TEXT("  Driver Version: %s (0x%X)"), *GRHIAdapterUserDriverVersion, Props.apiVersion);
-		UE_LOG(LogVulkanRHI, Log, TEXT("Internal Version: %s"), *GRHIAdapterInternalDriverVersion);
-		UE_LOG(LogVulkanRHI, Log, TEXT("     Driver Date: %s"), *GRHIAdapterDriverDate);
-	}
-	else if(PLATFORM_UNIX)
-	{
-		// :todo: Create a GetGPUDriverInfo for Linux
+		const VkPhysicalDeviceProperties& Props = Device->GetDeviceProperties();
 
 		if (Device->GetVendorId() == EGpuVendorId::Nvidia)
 		{
@@ -760,7 +741,40 @@ void FVulkanDynamicRHI::SelectAndInitDevice()
 
 		UE_LOG(LogVulkanRHI, Log, TEXT("    Adapter Name: %s"), *GRHIAdapterName);
 		UE_LOG(LogVulkanRHI, Log, TEXT("     API Version: %d.%d.%d"), VK_VERSION_MAJOR(Props.apiVersion), VK_VERSION_MINOR(Props.apiVersion), VK_VERSION_PATCH(Props.apiVersion));
-		UE_LOG(LogVulkanRHI, Log, TEXT("  Driver Version: %s (0x%X)"), *GRHIAdapterUserDriverVersion, Props.apiVersion);
+		UE_LOG(LogVulkanRHI, Log, TEXT("  Driver Version: %s"), *GRHIAdapterUserDriverVersion);
+	};
+
+	if (PLATFORM_ANDROID)
+	{
+		GRHIAdapterName.Append(TEXT(" Vulkan"));
+		GRHIAdapterInternalDriverVersion = FString::Printf(TEXT("%d.%d.%d"), VK_VERSION_MAJOR(Props.apiVersion), VK_VERSION_MINOR(Props.apiVersion), VK_VERSION_PATCH(Props.apiVersion));
+	}
+	else if (PLATFORM_WINDOWS)
+	{
+		GRHIDeviceId = Props.deviceID;
+		FGPUDriverInfo GPUDriverInfo = FPlatformMisc::GetGPUDriverInfo(GRHIAdapterName);
+
+		if (GPUDriverInfo.InternalDriverVersion != TEXT("Unknown"))
+		{
+			GRHIAdapterUserDriverVersion = GPUDriverInfo.UserDriverVersion;
+			GRHIAdapterInternalDriverVersion = GPUDriverInfo.InternalDriverVersion;
+			GRHIAdapterDriverDate = GPUDriverInfo.DriverDate;
+
+			UE_LOG(LogVulkanRHI, Log, TEXT("    Adapter Name: %s"), *GRHIAdapterName);
+			UE_LOG(LogVulkanRHI, Log, TEXT("     API Version: %d.%d.%d"), VK_VERSION_MAJOR(Props.apiVersion), VK_VERSION_MINOR(Props.apiVersion), VK_VERSION_PATCH(Props.apiVersion));
+			UE_LOG(LogVulkanRHI, Log, TEXT("  Driver Version: %s (0x%X)"), *GRHIAdapterUserDriverVersion, Props.driverVersion);
+			UE_LOG(LogVulkanRHI, Log, TEXT("Internal Version: %s"), *GRHIAdapterInternalDriverVersion);
+			UE_LOG(LogVulkanRHI, Log, TEXT("     Driver Date: %s"), *GRHIAdapterDriverDate);
+		}
+		else
+		{
+			// If we failed to read from the registry, then use the values provided by Vulkan props
+			ReadVulkanDriverVersionFromProps(Device);
+		}
+	}
+	else if(PLATFORM_UNIX)
+	{
+		ReadVulkanDriverVersionFromProps(Device);
 	}
 
 	GRHIPersistentThreadGroupCount = 1440; // TODO: Revisit based on vendor/adapter/perf query
