Fix Linux build.

diff --git a/Engine/Source/Runtime/VulkanRHI/Private/VulkanRHI.cpp b/Engine/Source/Runtime/VulkanRHI/Private/VulkanRHI.cpp
--- a/Engine/Source/Runtime/VulkanRHI/Private/VulkanRHI.cpp
+++ b/Engine/Source/Runtime/VulkanRHI/Private/VulkanRHI.cpp
@@ -719,11 +719,11 @@ void FVulkanDynamicRHI::SelectAndInitDevice()
 
 	Device->InitGPU(DeviceIndex);
 
-	auto ReadVulkanDriverVersionFromProps = [](FVulkanDevice* Device) {
+	auto ReadVulkanDriverVersionFromProps = [](FVulkanDevice* CurrentDevice) {
 
-		const VkPhysicalDeviceProperties& Props = Device->GetDeviceProperties();
+		const VkPhysicalDeviceProperties& Props = CurrentDevice->GetDeviceProperties();
 
-		if (Device->GetVendorId() == EGpuVendorId::Nvidia)
+		if (CurrentDevice->GetVendorId() == EGpuVendorId::Nvidia)
 		{
 			UNvidiaDriverVersion NvidiaVersion;
 			static_assert(sizeof(NvidiaVersion) == sizeof(Props.driverVersion), "Mismatched Nvidia pack driver version!");
@@ -756,24 +756,24 @@ void FVulkanDynamicRHI::SelectAndInitDevice()
 
 		if (GPUDriverInfo.InternalDriverVersion != TEXT("Unknown"))
 		{
-			GRHIAdapterUserDriverVersion = GPUDriverInfo.UserDriverVersion;
-			GRHIAdapterInternalDriverVersion = GPUDriverInfo.InternalDriverVersion;
-			GRHIAdapterDriverDate = GPUDriverInfo.DriverDate;
+		GRHIAdapterUserDriverVersion = GPUDriverInfo.UserDriverVersion;
+		GRHIAdapterInternalDriverVersion = GPUDriverInfo.InternalDriverVersion;
+		GRHIAdapterDriverDate = GPUDriverInfo.DriverDate;
 
-			UE_LOG(LogVulkanRHI, Log, TEXT("    Adapter Name: %s"), *GRHIAdapterName);
-			UE_LOG(LogVulkanRHI, Log, TEXT("     API Version: %d.%d.%d"), VK_VERSION_MAJOR(Props.apiVersion), VK_VERSION_MINOR(Props.apiVersion), VK_VERSION_PATCH(Props.apiVersion));
+		UE_LOG(LogVulkanRHI, Log, TEXT("    Adapter Name: %s"), *GRHIAdapterName);
+		UE_LOG(LogVulkanRHI, Log, TEXT("     API Version: %d.%d.%d"), VK_VERSION_MAJOR(Props.apiVersion), VK_VERSION_MINOR(Props.apiVersion), VK_VERSION_PATCH(Props.apiVersion));
 			UE_LOG(LogVulkanRHI, Log, TEXT("  Driver Version: %s (0x%X)"), *GRHIAdapterUserDriverVersion, Props.driverVersion);
-			UE_LOG(LogVulkanRHI, Log, TEXT("Internal Version: %s"), *GRHIAdapterInternalDriverVersion);
-			UE_LOG(LogVulkanRHI, Log, TEXT("     Driver Date: %s"), *GRHIAdapterDriverDate);
-		}
+		UE_LOG(LogVulkanRHI, Log, TEXT("Internal Version: %s"), *GRHIAdapterInternalDriverVersion);
+		UE_LOG(LogVulkanRHI, Log, TEXT("     Driver Date: %s"), *GRHIAdapterDriverDate);
+	}
 		else
-		{
+	{
 			// If we failed to read from the registry, then use the values provided by Vulkan props
 			ReadVulkanDriverVersionFromProps(Device);
 		}
-	}
+		}
 	else if(PLATFORM_UNIX)
-	{
+		{
 		ReadVulkanDriverVersionFromProps(Device);
 	}
 
