VulkanRHI: Fix validation issue (in CitySample) with 2D copies using 0 as a Z component.  Add DebugName to validation layer messages.

diff --git a/Engine/Source/Runtime/VulkanRHI/Private/VulkanTexture.cpp b/Engine/Source/Runtime/VulkanRHI/Private/VulkanTexture.cpp
--- a/Engine/Source/Runtime/VulkanRHI/Private/VulkanTexture.cpp
+++ b/Engine/Source/Runtime/VulkanRHI/Private/VulkanTexture.cpp
@@ -651,7 +651,7 @@ FVulkanSurface::FVulkanSurface(FVulkanDevice& InDevice, FVulkanEvictable* Owner,
 	// Fetch image size
 	VulkanRHI::vkGetImageMemoryRequirements(InDevice.GetInstanceHandle(), Image, &MemoryRequirements);
 
-	VULKAN_SET_DEBUG_NAME(InDevice, VK_OBJECT_TYPE_IMAGE, Image, TEXT("(FVulkanSurface*)0x%p"), this);
+	VULKAN_SET_DEBUG_NAME(InDevice, VK_OBJECT_TYPE_IMAGE, Image, TEXT("%s:(FVulkanSurface*)0x%p"), CreateInfo.DebugName ? CreateInfo.DebugName :TEXT("?"), this);
 
 	FullAspectMask = VulkanRHI::GetAspectMaskFromUEFormat(PixelFormat, true, true);
 	PartialAspectMask = VulkanRHI::GetAspectMaskFromUEFormat(PixelFormat, false, true);
@@ -971,7 +971,7 @@ FVulkanSurface::FVulkanSurface(FVulkanDevice& InDevice, VkImageViewType Resource
 			&StorageFormat, &ViewFormat);
 		FWrapLayer::CreateImage(VK_SUCCESS, InDevice.GetInstanceHandle(), &ImageCreateInfo.ImageCreateInfo, &Image);
 #endif
-		VULKAN_SET_DEBUG_NAME(InDevice, VK_OBJECT_TYPE_IMAGE, Image, TEXT("(FVulkanSurface*)0x%p"), this);
+		VULKAN_SET_DEBUG_NAME(InDevice, VK_OBJECT_TYPE_IMAGE, Image, TEXT("%s:(FVulkanSurface*)0x%p"), CreateInfo.DebugName ? CreateInfo.DebugName : TEXT("?"), this);
 
 		VkImageLayout InitialLayout = VK_IMAGE_LAYOUT_UNDEFINED;
 		bool bOnlyAddToLayoutManager, bDoInitialClear;
@@ -2664,7 +2664,7 @@ void FVulkanCommandListContext::RHICopyTexture(FRHITexture* SourceTexture, FRHIT
 			ensure(CopyInfo.Size.X > 0 && (uint32)CopyInfo.Size.X <= DstSurface.Width && CopyInfo.Size.Y > 0 && (uint32)CopyInfo.Size.Y <= DstSurface.Height);
 			Region.extent.width = CopyInfo.Size.X;
 			Region.extent.height = CopyInfo.Size.Y;
-			Region.extent.depth = CopyInfo.Size.Z;
+			Region.extent.depth = FMath::Max(1, CopyInfo.Size.Z);
 		}
 		Region.srcSubresource.aspectMask = SrcSurface.GetFullAspectMask();
 		Region.srcSubresource.baseArrayLayer = CopyInfo.SourceSliceIndex;
