VulkanRHI: Fix bug in CheckAllWritten() and enable it for graphic pipelines as well.

diff --git a/Engine/Source/Runtime/VulkanRHI/Private/VulkanPipelineState.cpp b/Engine/Source/Runtime/VulkanRHI/Private/VulkanPipelineState.cpp
--- a/Engine/Source/Runtime/VulkanRHI/Private/VulkanPipelineState.cpp
+++ b/Engine/Source/Runtime/VulkanRHI/Private/VulkanPipelineState.cpp
@@ -327,6 +327,9 @@ bool FVulkanGraphicsPipelineDescriptorState::InternalUpdateDescriptorSets(FVulka
 				{
 					const VkDescriptorSet DescriptorSet = DescriptorSetHandles[Set];
 					DSWriter[Set].SetDescriptorSet(DescriptorSet);
+#if VULKAN_VALIDATE_DESCRIPTORS_WRITTEN
+					DSWriter[Set].CheckAllWritten();
+#endif
 					++NumSets;
 				}
 
@@ -477,16 +480,16 @@ void FVulkanDescriptorSetWriter::CheckAllWritten()
 	}
 	else
 	{
-		int32 Last = int32(WrittenMask.Num()-1);
+		const int32 Last = int32(WrittenMask.Num()-1);
 		for(int32 i = 0; !bFail && i < Last; ++i)
 		{
 			uint64 Mask = WrittenMask[i];
 			bFail = bFail || Mask != 0xffffffff; 
 		}
 
-		uint32 TailCount = Writes % 32;
+		const uint32 TailCount = Writes - (Last * 32);
 		check(TailCount != 0);
-		uint32 TailMask = (1llu << TailCount)-1;
+		const uint32 TailMask = (1llu << TailCount)-1;
 		bFail = bFail || TailMask != WrittenMask[Last];
 	}
 
