VulkanRHI: Remove hardcoded driver version check in Windows/Linux Platform now that DriverDenyList works with Vulkan.

diff --git a/Engine/Source/Runtime/VulkanRHI/Private/Linux/VulkanLinuxPlatform.cpp b/Engine/Source/Runtime/VulkanRHI/Private/Linux/VulkanLinuxPlatform.cpp
--- a/Engine/Source/Runtime/VulkanRHI/Private/Linux/VulkanLinuxPlatform.cpp
+++ b/Engine/Source/Runtime/VulkanRHI/Private/Linux/VulkanLinuxPlatform.cpp
@@ -266,16 +266,4 @@ void FVulkanLinuxPlatform::WriteCrashMarker(const FOptionalVulkanDeviceExtension
 
 void FVulkanLinuxPlatform::CheckDeviceDriver(uint32 DeviceIndex, EGpuVendorId VendorId, const VkPhysicalDeviceProperties& Props)
 {
-	if (VendorId == EGpuVendorId::Nvidia)
-	{
-		UNvidiaDriverVersion NvidiaVersion;
-		static_assert(sizeof(NvidiaVersion) == sizeof(Props.driverVersion), "Mismatched Nvidia pack driver version!");
-		NvidiaVersion.Packed = Props.driverVersion;
-
-		{
-			UE_LOG(LogVulkanRHI, Warning, TEXT("Nvidia drivers do not currently support Nanite/Lumen in Vulkan."));
-			extern TAutoConsoleVariable<int32> GRHIAllow64bitShaderAtomicsCvar;
-			GRHIAllow64bitShaderAtomicsCvar->SetWithCurrentPriority(0);
-		}
-	}
 }
diff --git a/Engine/Source/Runtime/VulkanRHI/Private/Windows/VulkanWindowsPlatform.cpp b/Engine/Source/Runtime/VulkanRHI/Private/Windows/VulkanWindowsPlatform.cpp
--- a/Engine/Source/Runtime/VulkanRHI/Private/Windows/VulkanWindowsPlatform.cpp
+++ b/Engine/Source/Runtime/VulkanRHI/Private/Windows/VulkanWindowsPlatform.cpp
@@ -395,13 +395,6 @@ void FVulkanWindowsPlatform::CheckDeviceDriver(uint32 DeviceIndex, EGpuVendorId
 				BypassVar->SetWithCurrentPriority(1);
 			}
 		}
-
-		if ((NvidiaVersion.Major < 496) || ((NvidiaVersion.Major == 496) && (NvidiaVersion.Minor < 13)))
-		{
-			UE_LOG(LogVulkanRHI, Warning, TEXT("Nvidia drivers < 496.13 do not support Nanite/Lumen in Vulkan."));
-			extern TAutoConsoleVariable<int32> GRHIAllow64bitShaderAtomicsCvar;
-			GRHIAllow64bitShaderAtomicsCvar->SetWithCurrentPriority(0);
-		}
 	}
 }
 
