VulkanRHI: Give the option to change pipeline queries to BOTTOM_OF_PIPE.  This gives more precise measures for a single span, but might alter overal frame time on some platforms.  Leave TOP_OF_PIPE as default for now until we verified the effect on all platforms.

diff --git a/Engine/Source/Runtime/VulkanRHI/Private/VulkanQuery.cpp b/Engine/Source/Runtime/VulkanRHI/Private/VulkanQuery.cpp
--- a/Engine/Source/Runtime/VulkanRHI/Private/VulkanQuery.cpp
+++ b/Engine/Source/Runtime/VulkanRHI/Private/VulkanQuery.cpp
@@ -21,6 +21,16 @@ TAutoConsoleVariable<int32> GSubmitOcclusionBatchCmdBufferCVar(
 	ECVF_RenderThreadSafe
 );
 
+static uint32 GTimestampQueryStage = 0;
+TAutoConsoleVariable<int32> GTimestampQueryStageCVar(
+	TEXT("r.Vulkan.TimestampQueryStage"),
+	GTimestampQueryStage,
+	TEXT("Defines which pipeline stage is used for timestamp queries.\n")
+	TEXT(" 0: Use VK_PIPELINE_STAGE_TOP_OF_PIPE_BIT, less precise measures but less likely to alter performance (default)\n")
+	TEXT(" 1: Use VK_PIPELINE_STAGE_BOTTOM_OF_PIPE_BIT, more precise measures but may alter performance on some platforms"),
+	ECVF_RenderThreadSafe
+);
+
 constexpr uint32 GMinNumberOfQueriesInPool = 256;
 
 #if PLATFORM_ANDROID
@@ -620,8 +630,9 @@ void FVulkanCommandListContext::RHIEndRenderQuery(FRHIRenderQuery* QueryRHI)
 
 		Query->Pool->CurrentTimestamp = (Query->Pool->CurrentTimestamp + 1) % Query->Pool->BufferSize;
 		const uint32 QueryEndIndex = Query->Pool->CurrentTimestamp;
+		const VkPipelineStageFlagBits QueryPipelineStage = GTimestampQueryStage ? VK_PIPELINE_STAGE_BOTTOM_OF_PIPE_BIT : VK_PIPELINE_STAGE_TOP_OF_PIPE_BIT;
 		FVulkanCmdBuffer* CmdBuffer = CommandBufferManager->GetActiveCmdBuffer();
-		VulkanRHI::vkCmdWriteTimestamp(CmdBuffer->GetHandle(), VK_PIPELINE_STAGE_TOP_OF_PIPE_BIT, Query->Pool->GetHandle(), QueryEndIndex);
+		VulkanRHI::vkCmdWriteTimestamp(CmdBuffer->GetHandle(), QueryPipelineStage, Query->Pool->GetHandle(), QueryEndIndex);
 		CmdBuffer->AddPendingTimestampQuery(QueryEndIndex, 1, Query->Pool->GetHandle(), Query->Pool->ResultsBuffer->GetHandle(), true);
 		Query->Pool->TimestampListHandles[QueryEndIndex].CmdBuffer = CmdBuffer;
 		Query->Pool->TimestampListHandles[QueryEndIndex].FenceCounter = CmdBuffer->GetFenceSignaledCounter();
