VulkanRHI: Because we track a single layout for both depth/stencil, fix rare case when depth and stencil are transitioned out of copy layout in different barrier batches.  Fixes incorrect layout in DumpGPU* engines tests.

diff --git a/Engine/Source/Runtime/VulkanRHI/Private/VulkanBarriers.cpp b/Engine/Source/Runtime/VulkanRHI/Private/VulkanBarriers.cpp
--- a/Engine/Source/Runtime/VulkanRHI/Private/VulkanBarriers.cpp
+++ b/Engine/Source/Runtime/VulkanRHI/Private/VulkanBarriers.cpp
@@ -422,6 +422,24 @@ static void GetDepthStencilStageAndAccessFlags(ERHIAccess DepthAccess, ERHIAcces
 		return;
 	}
 
+	// If one of the source states has copy and the other doesn't, a cache lookup will be necessary
+	const ERHIAccess AnyCopy = ERHIAccess::CopySrc | ERHIAccess::CopyDest;
+	if (!EnumHasAnyFlags(DepthAccess, StencilAccess) &&
+		(EnumHasAnyFlags(DepthAccess, AnyCopy) || EnumHasAnyFlags(StencilAccess, AnyCopy)))
+	{
+		if (bIsSourceState)
+		{
+			StageFlags = VK_PIPELINE_STAGE_ALL_COMMANDS_BIT;
+			AccessFlags = VK_ACCESS_MEMORY_READ_BIT | VK_ACCESS_MEMORY_WRITE_BIT;
+			Layout = VK_IMAGE_LAYOUT_UNDEFINED;
+			return;
+		}
+		else
+		{
+			checkf(!EnumHasAnyFlags(StencilAccess, AnyCopy), TEXT("Only the Stencil being transitioned to Copy state, this is not supported."));
+		}
+	}
+
 	Layout = VK_IMAGE_LAYOUT_DEPTH_STENCIL_READ_ONLY_OPTIMAL;
 
 	if (EnumHasAnyFlags(DepthAccess, ERHIAccess::CopySrc))
@@ -1051,7 +1069,7 @@ void FVulkanCommandListContext::RHIBeginTransitions(TArrayView<const FRHITransit
 			FVulkanTextureBase* Texture = ExtraData.BaseTexture;
 			check(Texture->Surface.Image != VK_NULL_HANDLE);
 
-			FVulkanImageLayout& Layout = LayoutManager.GetOrAddFullLayout(Texture->Surface, VK_IMAGE_LAYOUT_UNDEFINED);
+			const FVulkanImageLayout& Layout = LayoutManager.GetOrAddFullLayout(Texture->Surface, VK_IMAGE_LAYOUT_UNDEFINED);
 
 			VkAccessFlags SrcAccessFlags;
 			VkImageLayout SrcLayout, DstLayout;
