Adding minimal interface to VulkanRHI
- Adding IVulkanDynamicRHI interface for plugins that want to touch Vulkan resources/handles directly.
- Modifying plugins that were using VulkanRHIPrivate.h to use the interface instead.
- Removing plugin references to UE Modules and ThirdParty libraries that were only needed because they included the private RHI headers.
- Removing VulkanRHI/Private from plugin include paths.
- Deprecating VulkanRHIBridge and moved its functionality into the new interface.

diff --git a/Engine/Plugins/Media/HardwareEncoders/Source/EncoderAMF/EncoderAMF.Build.cs b/Engine/Plugins/Media/HardwareEncoders/Source/EncoderAMF/EncoderAMF.Build.cs
--- a/Engine/Plugins/Media/HardwareEncoders/Source/EncoderAMF/EncoderAMF.Build.cs
+++ b/Engine/Plugins/Media/HardwareEncoders/Source/EncoderAMF/EncoderAMF.Build.cs
@@ -37,18 +37,6 @@ public class EncoderAMF : ModuleRules
 		string EngineSourceDirectory = Path.GetFullPath(Target.RelativeEnginePath);
 		PrivateIncludePaths.Add(Path.Combine(EngineSourceDirectory, "Source/Runtime/AVEncoder/Private"));
 
-		AddEngineThirdPartyPrivateStaticDependencies(Target, "NVAftermath");
-
-		PrivateIncludePaths.Add(Path.Combine(EngineSourceDirectory, "Source/Runtime/VulkanRHI/Private"));
 		AddEngineThirdPartyPrivateStaticDependencies(Target, "Vulkan");
-
-		if (Target.IsInPlatformGroup(UnrealPlatformGroup.Windows))
-		{
-			PrivateIncludePaths.Add(Path.Combine(EngineSourceDirectory, "Source/Runtime/VulkanRHI/Private/Windows"));
-		}
-		else if (Target.IsInPlatformGroup(UnrealPlatformGroup.Unix))
-		{
-			PrivateIncludePaths.Add(Path.Combine(EngineSourceDirectory, "Source/Runtime/VulkanRHI/Private/Linux"));
-		}
 	}
 }
diff --git a/Engine/Plugins/Media/HardwareEncoders/Source/EncoderAMF/Private/Amf_Common.h b/Engine/Plugins/Media/HardwareEncoders/Source/EncoderAMF/Private/Amf_Common.h
--- a/Engine/Plugins/Media/HardwareEncoders/Source/EncoderAMF/Private/Amf_Common.h
+++ b/Engine/Plugins/Media/HardwareEncoders/Source/EncoderAMF/Private/Amf_Common.h
@@ -4,7 +4,6 @@
 
 #include "CoreMinimal.h"
 #include "HAL/Thread.h"
-#include "VulkanRHIPrivate.h"
 
 #include <string.h>
 
diff --git a/Engine/Plugins/Media/HardwareEncoders/Source/EncoderAMF/Private/Amf_EncoderH264.cpp b/Engine/Plugins/Media/HardwareEncoders/Source/EncoderAMF/Private/Amf_EncoderH264.cpp
--- a/Engine/Plugins/Media/HardwareEncoders/Source/EncoderAMF/Private/Amf_EncoderH264.cpp
+++ b/Engine/Plugins/Media/HardwareEncoders/Source/EncoderAMF/Private/Amf_EncoderH264.cpp
@@ -5,7 +5,6 @@
 #include "VideoEncoderCommon.h"
 #include "CodecPacket.h"
 #include "AVEncoderDebug.h"
-#include "VulkanRHIBridge.h"
 #include "VideoEncoderInput.h"
 #include "RHI.h"
 #include <stdio.h>
diff --git a/Engine/Plugins/Media/HardwareEncoders/Source/EncoderAMF/Private/Amf_EncoderModule.cpp b/Engine/Plugins/Media/HardwareEncoders/Source/EncoderAMF/Private/Amf_EncoderModule.cpp
--- a/Engine/Plugins/Media/HardwareEncoders/Source/EncoderAMF/Private/Amf_EncoderModule.cpp
+++ b/Engine/Plugins/Media/HardwareEncoders/Source/EncoderAMF/Private/Amf_EncoderModule.cpp
@@ -2,10 +2,7 @@
 
 #include "Modules/ModuleManager.h"
 
-#include "VulkanRHIPrivate.h"
-#include "VulkanRHIBridge.h"
-
-#include "DynamicRHI.h"
+#include "IVulkanDynamicRHI.h"
 
 #include "Amf_Common.h"
 #include "Amf_EncoderH264.h"
@@ -49,7 +46,7 @@ public:
 
 					AMF.DestroyContext();
 
-					VulkanRHIBridge::AddEnabledDeviceExtensionsAndLayers(ExtentionsToAdd, TArray<const ANSICHAR*>());
+					IVulkanDynamicRHI::AddEnabledDeviceExtensionsAndLayers(ExtentionsToAdd, TArray<const ANSICHAR*>());
 				}
 
 				FCoreDelegates::OnPostEngineInit.AddLambda([]() {FVideoEncoderAmf_H264::Register(FVideoEncoderFactory::Get());});
diff --git a/Engine/Plugins/Media/HardwareEncoders/Source/EncoderNVENC/EncoderNVENC.Build.cs b/Engine/Plugins/Media/HardwareEncoders/Source/EncoderNVENC/EncoderNVENC.Build.cs
--- a/Engine/Plugins/Media/HardwareEncoders/Source/EncoderNVENC/EncoderNVENC.Build.cs
+++ b/Engine/Plugins/Media/HardwareEncoders/Source/EncoderNVENC/EncoderNVENC.Build.cs
@@ -30,18 +30,7 @@ public class EncoderNVENC: ModuleRules
 
 		PrivateIncludePaths.Add(Path.Combine(EngineSourceDirectory, "Source/Runtime/AVEncoder/Private"));
 
-		AddEngineThirdPartyPrivateStaticDependencies(Target, "NVAftermath");
-		
-		PrivateIncludePaths.Add(Path.Combine(EngineSourceDirectory, "Source/Runtime/VulkanRHI/Private"));
 		AddEngineThirdPartyPrivateStaticDependencies(Target, "Vulkan");
-
-		if (Target.IsInPlatformGroup(UnrealPlatformGroup.Windows))
-		{
-			PrivateIncludePaths.Add(Path.Combine(EngineSourceDirectory, "Source/Runtime/VulkanRHI/Private/Windows"));
-		}
-		else if (Target.IsInPlatformGroup(UnrealPlatformGroup.Unix))
-		{
-			PrivateIncludePaths.Add(Path.Combine(EngineSourceDirectory, "Source/Runtime/VulkanRHI/Private/Linux"));
-		}
+		AddEngineThirdPartyPrivateStaticDependencies(Target, "DX11");
 	}
 }
diff --git a/Engine/Plugins/Media/HardwareEncoders/Source/EncoderNVENC/Private/NVENC_EncoderModule.cpp b/Engine/Plugins/Media/HardwareEncoders/Source/EncoderNVENC/Private/NVENC_EncoderModule.cpp
--- a/Engine/Plugins/Media/HardwareEncoders/Source/EncoderNVENC/Private/NVENC_EncoderModule.cpp
+++ b/Engine/Plugins/Media/HardwareEncoders/Source/EncoderNVENC/Private/NVENC_EncoderModule.cpp
@@ -2,10 +2,7 @@
 
 #include "Modules/ModuleManager.h"
 
-#include "VulkanRHIPrivate.h"
-#include "VulkanRHIBridge.h"
-
-#include "DynamicRHI.h"
+#include "IVulkanDynamicRHI.h"
 
 #include "NVENC_Common.h"
 #include "NVENC_EncoderH264.h"
@@ -38,7 +35,7 @@ public:
 #elif PLATFORM_LINUX
 					const TArray<const ANSICHAR*> ExtentionsToAdd{ VK_KHR_EXTERNAL_MEMORY_EXTENSION_NAME, VK_KHR_EXTERNAL_MEMORY_FD_EXTENSION_NAME };
 #endif
-					VulkanRHIBridge::AddEnabledDeviceExtensionsAndLayers(ExtentionsToAdd, TArray<const ANSICHAR*>());
+					IVulkanDynamicRHI::AddEnabledDeviceExtensionsAndLayers(ExtentionsToAdd, TArray<const ANSICHAR*>());
 				}
 
 				FModuleManager::LoadModuleChecked<FCUDAModule>("CUDA").OnPostCUDAInit.AddLambda([]() {FVideoEncoderNVENC_H264::Register(FVideoEncoderFactory::Get());});
diff --git a/Engine/Plugins/Media/PixelStreaming/Source/PixelStreaming/PixelStreaming.Build.cs b/Engine/Plugins/Media/PixelStreaming/Source/PixelStreaming/PixelStreaming.Build.cs
--- a/Engine/Plugins/Media/PixelStreaming/Source/PixelStreaming/PixelStreaming.Build.cs
+++ b/Engine/Plugins/Media/PixelStreaming/Source/PixelStreaming/PixelStreaming.Build.cs
@@ -55,7 +55,6 @@ namespace UnrealBuildTool.Rules
 			PrivateIncludePaths.AddRange(new string[]
 				{
 					Path.Combine(EngineDir, "Source/Runtime/AudioMixer/Private"),
-					Path.Combine(EngineDir, "Source/Runtime/VulkanRHI/Private"),
 				});
 
 			PrivateDependencyModuleNames.AddRange(new string[]
@@ -79,16 +78,8 @@ namespace UnrealBuildTool.Rules
 				"AVEncoder"
 			});
 
-			AddEngineThirdPartyPrivateStaticDependencies(Target, "Vulkan");
-
-			// required for casting UE4 BackBuffer to Vulkan Texture2D for NvEnc
-			PrivateDependencyModuleNames.AddRange(new string[] { "CUDA", "VulkanRHI", "nvEncode" });
-
-			AddEngineThirdPartyPrivateStaticDependencies(Target, "NVAftermath");
-			
-			PrivateIncludePathModuleNames.Add("VulkanRHI");
-			PrivateIncludePaths.Add(Path.Combine(EngineDir, "Source/Runtime/VulkanRHI/Private"));
-			AddEngineThirdPartyPrivateStaticDependencies(Target, "Vulkan");
+			PrivateDependencyModuleNames.Add("VulkanRHI");
+			AddEngineThirdPartyPrivateStaticDependencies(Target, "Vulkan", "CUDA");
 
 			if (Target.IsInPlatformGroup(UnrealPlatformGroup.Windows))
 			{
@@ -96,12 +87,6 @@ namespace UnrealBuildTool.Rules
 				PrivateDependencyModuleNames.Add("D3D12RHI");
 
 				AddEngineThirdPartyPrivateStaticDependencies(Target, "DX11", "DX12");
-
-				PrivateIncludePaths.Add(Path.Combine(EngineDir, "Source/Runtime/VulkanRHI/Private/Windows"));
-			}
-			else if (Target.IsInPlatformGroup(UnrealPlatformGroup.Linux))
-			{
-				PrivateIncludePaths.Add(Path.Combine(EngineDir, "Source/Runtime/VulkanRHI/Private/Linux"));
 			}
 
 			AddFolder("SignallingWebServer");
diff --git a/Engine/Plugins/Media/PixelStreaming/Source/PixelStreaming/Private/EncoderFrameFactory.cpp b/Engine/Plugins/Media/PixelStreaming/Source/PixelStreaming/Private/EncoderFrameFactory.cpp
--- a/Engine/Plugins/Media/PixelStreaming/Source/PixelStreaming/Private/EncoderFrameFactory.cpp
+++ b/Engine/Plugins/Media/PixelStreaming/Source/PixelStreaming/Private/EncoderFrameFactory.cpp
@@ -3,11 +3,12 @@
 #include "EncoderFrameFactory.h"
 #include "PixelStreamingPrivate.h"
 #include "CudaModule.h"
-#include "VulkanRHIPrivate.h"
 #include "Settings.h"
 #include "VideoEncoderInput.h"
 #include "Templates/SharedPointer.h"
 
+#include "IVulkanDynamicRHI.h"
+
 #if PLATFORM_WINDOWS
 	#include "VideoCommon.h"
 	#include "D3D11State.h"
@@ -117,8 +118,8 @@ TSharedPtr<AVEncoder::FVideoEncoderInput> UE::PixelStreaming::FEncoderFrameFacto
 	{
 		if (IsRHIDeviceAMD())
 		{
-			FVulkanDynamicRHI* DynamicRHI = GetDynamicRHI<FVulkanDynamicRHI>();
-			AVEncoder::FVulkanDataStruct VulkanData = { DynamicRHI->GetInstance(), DynamicRHI->GetDevice()->GetPhysicalHandle(), DynamicRHI->GetDevice()->GetInstanceHandle() };
+			IVulkanDynamicRHI* DynamicRHI = GetDynamicRHI<IVulkanDynamicRHI>();
+			AVEncoder::FVulkanDataStruct VulkanData = { DynamicRHI->RHIGetVkInstance(), DynamicRHI->RHIGetVkPhysicalDevice(), DynamicRHI->RHIGetVkDevice() };
 
 			return AVEncoder::FVideoEncoderInput::CreateForVulkan(&VulkanData, bIsResizable);
 		}
@@ -173,8 +174,8 @@ void UE::PixelStreaming::FEncoderFrameFactory::SetTexture(TSharedPtr<AVEncoder::
 	{
 		if (IsRHIDeviceAMD())
 		{
-			FVulkanTexture2D* VulkanTexture = static_cast<FVulkanTexture2D*>(Texture.GetReference());
-			InputFrame->SetTexture(VulkanTexture->Surface.Image, [](VkImage NativeTexture) { /* Do something with released texture if needed */ });
+			const VkImage VulkanImage = GetIVulkanDynamicRHI()->RHIGetVkImage(Texture.GetReference());
+			InputFrame->SetTexture(VulkanImage, [](VkImage NativeTexture) { /* Do something with released texture if needed */ });
 		}
 		else if (IsRHIDeviceNVIDIA())
 		{
@@ -217,8 +218,11 @@ void UE::PixelStreaming::FEncoderFrameFactory::SetTexture(TSharedPtr<AVEncoder::
 
 void UE::PixelStreaming::FEncoderFrameFactory::SetTextureCUDAVulkan(TSharedPtr<AVEncoder::FVideoEncoderInputFrame> InputFrame, const FTexture2DRHIRef& Texture)
 {
-	FVulkanTexture2D* VulkanTexture = static_cast<FVulkanTexture2D*>(Texture.GetReference());
-	VkDevice Device = GetDynamicRHI<FVulkanDynamicRHI>()->GetDevice()->GetInstanceHandle();
+	IVulkanDynamicRHI* VulkanRHI = GetIVulkanDynamicRHI();
+
+	const FVulkanRHIAllocationInfo TextureAllocationInfo = VulkanRHI->RHIGetAllocationInfo(Texture.GetReference());
+
+	VkDevice Device = VulkanRHI->RHIGetVkDevice();
 
 #if PLATFORM_WINDOWS
 	HANDLE Handle;
@@ -232,15 +236,11 @@ void UE::PixelStreaming::FEncoderFrameFactory::SetTextureCUDAVulkan(TSharedPtr<A
 		VkMemoryGetWin32HandleInfoKHR MemoryGetHandleInfoKHR = {};
 		MemoryGetHandleInfoKHR.sType = VK_STRUCTURE_TYPE_MEMORY_GET_WIN32_HANDLE_INFO_KHR;
 		MemoryGetHandleInfoKHR.pNext = NULL;
-		MemoryGetHandleInfoKHR.memory = VulkanTexture->Surface.GetAllocationHandle();
+		MemoryGetHandleInfoKHR.memory = TextureAllocationInfo.Handle;
 		MemoryGetHandleInfoKHR.handleType = bUseNTHandle ? VK_EXTERNAL_MEMORY_HANDLE_TYPE_OPAQUE_WIN32_BIT : VK_EXTERNAL_MEMORY_HANDLE_TYPE_OPAQUE_WIN32_KMT_BIT;
 
-		// While this operation is safe (and unavoidable) C4191 has been enabled and this will trigger an error with warnings as errors
-	#pragma warning(push)
-	#pragma warning(disable : 4191)
-		PFN_vkGetMemoryWin32HandleKHR GetMemoryWin32HandleKHR = (PFN_vkGetMemoryWin32HandleKHR)VulkanRHI::vkGetDeviceProcAddr(Device, "vkGetMemoryWin32HandleKHR");
-		VERIFYVULKANRESULT(GetMemoryWin32HandleKHR(Device, &MemoryGetHandleInfoKHR, &Handle));
-	#pragma warning(pop)
+		PFN_vkGetMemoryWin32HandleKHR GetMemoryWin32HandleKHR = (PFN_vkGetMemoryWin32HandleKHR)VulkanRHI->RHIGetVkDeviceProcAddr("vkGetMemoryWin32HandleKHR");
+		VERIFYVULKANRESULT_EXTERNAL(GetMemoryWin32HandleKHR(Device, &MemoryGetHandleInfoKHR, &Handle));
 	}
 
 	FCUDAModule::CUDA().cuCtxPushCurrent(FModuleManager::GetModuleChecked<FCUDAModule>("CUDA").GetCudaContext());
@@ -253,7 +253,7 @@ void UE::PixelStreaming::FEncoderFrameFactory::SetTextureCUDAVulkan(TSharedPtr<A
 		CudaExtMemHandleDesc.type = bUseNTHandle ? CU_EXTERNAL_MEMORY_HANDLE_TYPE_OPAQUE_WIN32 : CU_EXTERNAL_MEMORY_HANDLE_TYPE_OPAQUE_WIN32_KMT;
 		CudaExtMemHandleDesc.handle.win32.name = NULL;
 		CudaExtMemHandleDesc.handle.win32.handle = Handle;
-		CudaExtMemHandleDesc.size = VulkanTexture->Surface.GetAllocationOffset() + VulkanTexture->Surface.GetMemorySize();
+		CudaExtMemHandleDesc.size = TextureAllocationInfo.Offset + TextureAllocationInfo.Size;
 
 		// import external memory
 		CUresult Result = FCUDAModule::CUDA().cuImportExternalMemory(&MappedExternalMemory, &CudaExtMemHandleDesc);
@@ -276,15 +276,11 @@ void UE::PixelStreaming::FEncoderFrameFactory::SetTextureCUDAVulkan(TSharedPtr<A
 		VkMemoryGetFdInfoKHR MemoryGetFdInfoKHR = {};
 		MemoryGetFdInfoKHR.sType = VK_STRUCTURE_TYPE_MEMORY_GET_FD_INFO_KHR;
 		MemoryGetFdInfoKHR.pNext = NULL;
-		MemoryGetFdInfoKHR.memory = VulkanTexture->Surface.GetAllocationHandle();
+		MemoryGetFdInfoKHR.memory = TextureAllocationInfo.Handle;
 		MemoryGetFdInfoKHR.handleType = VK_EXTERNAL_MEMORY_HANDLE_TYPE_OPAQUE_FD_BIT_KHR;
 
-		// While this operation is safe (and unavoidable) C4191 has been enabled and this will trigger an error with warnings as errors
-	#pragma warning(push)
-	#pragma warning(disable : 4191)
-		PFN_vkGetMemoryFdKHR FPGetMemoryFdKHR = (PFN_vkGetMemoryFdKHR)VulkanRHI::vkGetDeviceProcAddr(Device, "vkGetMemoryFdKHR");
-		VERIFYVULKANRESULT(FPGetMemoryFdKHR(Device, &MemoryGetFdInfoKHR, &Fd));
-	#pragma warning(pop)
+		PFN_vkGetMemoryFdKHR FPGetMemoryFdKHR = (PFN_vkGetMemoryFdKHR)VulkanRHI->RHIGetVkDeviceProcAddr("vkGetMemoryFdKHR");
+		VERIFYVULKANRESULT_EXTERNAL(FPGetMemoryFdKHR(Device, &MemoryGetFdInfoKHR, &Fd));
 	}
 
 	FCUDAModule::CUDA().cuCtxPushCurrent(FModuleManager::GetModuleChecked<FCUDAModule>("CUDA").GetCudaContext());
@@ -296,7 +292,7 @@ void UE::PixelStreaming::FEncoderFrameFactory::SetTextureCUDAVulkan(TSharedPtr<A
 		CUDA_EXTERNAL_MEMORY_HANDLE_DESC CudaExtMemHandleDesc = {};
 		CudaExtMemHandleDesc.type = CU_EXTERNAL_MEMORY_HANDLE_TYPE_OPAQUE_FD;
 		CudaExtMemHandleDesc.handle.fd = Fd;
-		CudaExtMemHandleDesc.size = VulkanTexture->Surface.GetAllocationOffset() + VulkanTexture->Surface.GetMemorySize();
+		CudaExtMemHandleDesc.size = TextureAllocationInfo.Offset + TextureAllocationInfo.Size;
 
 		// import external memory
 		CUresult Result = FCUDAModule::CUDA().cuImportExternalMemory(&MappedExternalMemory, &CudaExtMemHandleDesc);
@@ -314,7 +310,7 @@ void UE::PixelStreaming::FEncoderFrameFactory::SetTextureCUDAVulkan(TSharedPtr<A
 	{
 		CUDA_EXTERNAL_MEMORY_MIPMAPPED_ARRAY_DESC MipmapDesc = {};
 		MipmapDesc.numLevels = 1;
-		MipmapDesc.offset = VulkanTexture->Surface.GetAllocationOffset();
+		MipmapDesc.offset = TextureAllocationInfo.Offset;
 		MipmapDesc.arrayDesc.Width = Texture->GetSizeX();
 		MipmapDesc.arrayDesc.Height = Texture->GetSizeY();
 		MipmapDesc.arrayDesc.Depth = 0;
diff --git a/Engine/Plugins/Media/PixelStreaming/Source/PixelStreaming/Private/VideoEncoderH264Wrapper.cpp b/Engine/Plugins/Media/PixelStreaming/Source/PixelStreaming/Private/VideoEncoderH264Wrapper.cpp
--- a/Engine/Plugins/Media/PixelStreaming/Source/PixelStreaming/Private/VideoEncoderH264Wrapper.cpp
+++ b/Engine/Plugins/Media/PixelStreaming/Source/PixelStreaming/Private/VideoEncoderH264Wrapper.cpp
@@ -1,7 +1,6 @@
 // Copyright Epic Games, Inc. All Rights Reserved.
 
 #include "VideoEncoderH264Wrapper.h"
-#include "VulkanRHIPrivate.h"
 #include "Stats.h"
 #include "VideoEncoderFactory.h"
 #include "EncoderFactory.h"
diff --git a/Engine/Plugins/Runtime/Oculus/OculusVR/Source/OculusHMD/OculusHMD.Build.cs b/Engine/Plugins/Runtime/Oculus/OculusVR/Source/OculusHMD/OculusHMD.Build.cs
--- a/Engine/Plugins/Runtime/Oculus/OculusVR/Source/OculusHMD/OculusHMD.Build.cs
+++ b/Engine/Plugins/Runtime/Oculus/OculusVR/Source/OculusHMD/OculusHMD.Build.cs
@@ -13,7 +13,6 @@ namespace UnrealBuildTool.Rules
 				new string[] {
 					// Relative to Engine\Plugins\Runtime\Oculus\OculusVR\Source
 					"../../../../../Source/Runtime/Renderer/Private",
-					"../../../../../Source/Runtime/VulkanRHI/Private",
 					"../../../../../Source/Runtime/Engine/Classes/Components",
 					"../../../../../Source/Runtime/Engine/Classes/Kismet",
 				});
@@ -24,15 +23,6 @@ namespace UnrealBuildTool.Rules
 					"ProceduralMeshComponent",
 				});			
 
-			if (Target.Platform == UnrealTargetPlatform.Win64)
-			{
-				PrivateIncludePaths.Add("../../../../../Source/Runtime/VulkanRHI/Private/Windows");
-			}
-			else
-			{
-				PrivateIncludePaths.Add("../../../../../Source/Runtime/VulkanRHI/Private/" + Target.Platform);
-			}
-
 			PrivateDependencyModuleNames.AddRange(
 				new string[]
 				{
diff --git a/Engine/Plugins/Runtime/Oculus/OculusVR/Source/OculusHMD/Private/OculusHMDPrivateRHI.h b/Engine/Plugins/Runtime/Oculus/OculusVR/Source/OculusHMD/Private/OculusHMDPrivateRHI.h
--- a/Engine/Plugins/Runtime/Oculus/OculusVR/Source/OculusHMD/Private/OculusHMDPrivateRHI.h
+++ b/Engine/Plugins/Runtime/Oculus/OculusVR/Source/OculusHMD/Private/OculusHMDPrivateRHI.h
@@ -36,8 +36,7 @@
 //-------------------------------------------------------------------------------------------------
 
 #if OCULUS_HMD_SUPPORTED_PLATFORMS_VULKAN
-#include "VulkanRHIPrivate.h"
-#include "VulkanResources.h"
+#include "IVulkanDynamicRHI.h"
 #endif // OCULUS_HMD_SUPPORTED_PLATFORMS_VULKAN
 
 #endif // OCULUS_HMD_SUPPORTED_PLATFORMS
\ No newline at end of file
diff --git a/Engine/Plugins/Runtime/Oculus/OculusVR/Source/OculusHMD/Private/OculusHMD_CustomPresent_Vulkan.cpp b/Engine/Plugins/Runtime/Oculus/OculusVR/Source/OculusHMD/Private/OculusHMD_CustomPresent_Vulkan.cpp
--- a/Engine/Plugins/Runtime/Oculus/OculusVR/Source/OculusHMD/Private/OculusHMD_CustomPresent_Vulkan.cpp
+++ b/Engine/Plugins/Runtime/Oculus/OculusVR/Source/OculusHMD/Private/OculusHMD_CustomPresent_Vulkan.cpp
@@ -6,9 +6,7 @@
 #if OCULUS_HMD_SUPPORTED_PLATFORMS_VULKAN
 #include "OculusHMD.h"
 
-#include "VulkanRHIPrivate.h"
-#include "VulkanPendingState.h"
-#include "VulkanContext.h"
+#include "IVulkanDynamicRHI.h"
 
 #if PLATFORM_WINDOWS
 #ifndef WINDOWS_PLATFORM_TYPES_GUARD
@@ -64,26 +62,18 @@ FVulkanCustomPresent::FVulkanCustomPresent(FOculusHMD* InOculusHMD) :
 	}
 */
 #endif
-	FVulkanDynamicRHI* const DynamicRHI = GetDynamicRHI<FVulkanDynamicRHI>();
 
-	bSupportsSubsampled = DynamicRHI->GetDevice()->GetOptionalExtensions().HasEXTFragmentDensityMap2;
+	bSupportsSubsampled = GetIVulkanDynamicRHI()->RHISupportsEXTFragmentDensityMap2();
 }
 
 
 bool FVulkanCustomPresent::IsUsingCorrectDisplayAdapter() const
 {
 #if PLATFORM_WINDOWS
-	const void* luid;
-
-	if (OVRP_SUCCESS(FOculusHMDModule::GetPluginWrapper().GetDisplayAdapterId2(&luid)) &&
-		luid &&
-		GVulkanRHI->GetOptionalExtensions().HasKHRGetPhysicalDeviceProperties2)
+	const void* AdapterId = nullptr;
+	if (OVRP_SUCCESS(FOculusHMDModule::GetPluginWrapper().GetDisplayAdapterId2(&AdapterId)) && AdapterId)
 	{
-		const VkPhysicalDeviceIDPropertiesKHR& vkPhysicalDeviceIDProperties = GVulkanRHI->GetDevice()->GetDeviceIdProperties();
-		if (vkPhysicalDeviceIDProperties.deviceLUIDValid)
-		{
-			return !FMemory::Memcmp(luid, &vkPhysicalDeviceIDProperties.deviceLUID, sizeof(LUID));
-		}
+		return GetIVulkanDynamicRHI()->RHIDoesAdapterMatchDevice(AdapterId);
 	}
 #endif
 
@@ -94,25 +84,25 @@ bool FVulkanCustomPresent::IsUsingCorrectDisplayAdapter() const
 
 void* FVulkanCustomPresent::GetOvrpInstance() const
 {
-	return GVulkanRHI->GetInstance();
+	return GetIVulkanDynamicRHI()->RHIGetVkInstance();
 }
 
 
 void* FVulkanCustomPresent::GetOvrpPhysicalDevice() const
 {
-	return GVulkanRHI->GetDevice()->GetPhysicalHandle();
+	return GetIVulkanDynamicRHI()->RHIGetVkPhysicalDevice();
 }
 
 
 void* FVulkanCustomPresent::GetOvrpDevice() const
 {
-	return GVulkanRHI->GetDevice()->GetInstanceHandle();
+	return GetIVulkanDynamicRHI()->RHIGetVkDevice();
 }
 
 
 void* FVulkanCustomPresent::GetOvrpCommandQueue() const
 {
-	return GVulkanRHI->GetDevice()->GetGraphicsQueue()->GetHandle();
+	return GetIVulkanDynamicRHI()->RHIGetGraphicsVkQueue();
 }
 
 
@@ -120,29 +110,29 @@ FTextureRHIRef FVulkanCustomPresent::CreateTexture_RenderThread(uint32 InSizeX,
 {
 	CheckInRenderThread();
 
-	VkImageSubresourceRange SubresourceRangeAll = { VK_IMAGE_ASPECT_COLOR_BIT, 0, VK_REMAINING_MIP_LEVELS, 0, VK_REMAINING_ARRAY_LAYERS };
-	FVulkanCommandListContext& ImmediateContext = GVulkanRHI->GetDevice()->GetImmediateContext();
-	FVulkanCmdBuffer* CmdBuffer = ImmediateContext.GetCommandBufferManager()->GetActiveCmdBuffer();
+	IVulkanDynamicRHI* VulkanRHI = GetIVulkanDynamicRHI();
+
+	const VkImageSubresourceRange SubresourceRangeAll = { VK_IMAGE_ASPECT_COLOR_BIT, 0, VK_REMAINING_MIP_LEVELS, 0, VK_REMAINING_ARRAY_LAYERS };
 
 	if (EnumHasAnyFlags(InTexCreateFlags,TexCreate_RenderTargetable))
 	{
-		GVulkanRHI->VulkanSetImageLayout(CmdBuffer->GetHandle(), (VkImage)InTexture, VK_IMAGE_LAYOUT_UNDEFINED, VK_IMAGE_LAYOUT_COLOR_ATTACHMENT_OPTIMAL, SubresourceRangeAll);
+		VulkanRHI->RHISetImageLayout((VkImage)InTexture, VK_IMAGE_LAYOUT_UNDEFINED, VK_IMAGE_LAYOUT_COLOR_ATTACHMENT_OPTIMAL, SubresourceRangeAll);
 	}
 	else if (EnumHasAnyFlags(InTexCreateFlags,TexCreate_Foveation))
 	{
-		GVulkanRHI->VulkanSetImageLayout(CmdBuffer->GetHandle(), (VkImage)InTexture, VK_IMAGE_LAYOUT_UNDEFINED, VK_IMAGE_LAYOUT_FRAGMENT_DENSITY_MAP_OPTIMAL_EXT, SubresourceRangeAll);
+		VulkanRHI->RHISetImageLayout((VkImage)InTexture, VK_IMAGE_LAYOUT_UNDEFINED, VK_IMAGE_LAYOUT_FRAGMENT_DENSITY_MAP_OPTIMAL_EXT, SubresourceRangeAll);
 	}
 
 	switch (InResourceType)
 	{
 	case RRT_Texture2D:
-		return GVulkanRHI->RHICreateTexture2DFromResource(InFormat, InSizeX, InSizeY, InNumMips, InNumSamples, (VkImage) InTexture, InTexCreateFlags).GetReference();
+		return VulkanRHI->RHICreateTexture2DFromResource(InFormat, InSizeX, InSizeY, InNumMips, InNumSamples, (VkImage) InTexture, InTexCreateFlags).GetReference();
 
 	case RRT_Texture2DArray:
-		return GVulkanRHI->RHICreateTexture2DArrayFromResource(InFormat, InSizeX, InSizeY, 2, InNumMips, InNumSamples, (VkImage) InTexture, InTexCreateFlags).GetReference();
+		return VulkanRHI->RHICreateTexture2DArrayFromResource(InFormat, InSizeX, InSizeY, 2, InNumMips, InNumSamples, (VkImage) InTexture, InTexCreateFlags).GetReference();
 
 	case RRT_TextureCube:
-		return GVulkanRHI->RHICreateTextureCubeFromResource(InFormat, InSizeX, false, 1, InNumMips, (VkImage) InTexture, InTexCreateFlags).GetReference();
+		return VulkanRHI->RHICreateTextureCubeFromResource(InFormat, InSizeX, false, 1, InNumMips, (VkImage) InTexture, InTexCreateFlags).GetReference();
 
 	default:
 		return nullptr;
diff --git a/Engine/Plugins/Runtime/Oculus/OculusVR/Source/OculusHMD/Private/OculusHMD_VulkanExtensions.cpp b/Engine/Plugins/Runtime/Oculus/OculusVR/Source/OculusHMD/Private/OculusHMD_VulkanExtensions.cpp
--- a/Engine/Plugins/Runtime/Oculus/OculusVR/Source/OculusHMD/Private/OculusHMD_VulkanExtensions.cpp
+++ b/Engine/Plugins/Runtime/Oculus/OculusVR/Source/OculusHMD/Private/OculusHMD_VulkanExtensions.cpp
@@ -17,13 +17,7 @@ namespace OculusHMD
 bool FVulkanExtensions::GetVulkanInstanceExtensionsRequired(TArray<const ANSICHAR*>& Out)
 {
 #if OCULUS_HMD_SUPPORTED_PLATFORMS_VULKAN
-	TArray<VkExtensionProperties> Properties;
-	{
-		uint32_t PropertyCount;
-		VulkanRHI::vkEnumerateInstanceExtensionProperties(nullptr, &PropertyCount, nullptr);
-		Properties.SetNum(PropertyCount);
-		VulkanRHI::vkEnumerateInstanceExtensionProperties(nullptr, &PropertyCount, Properties.GetData());
-	}
+	TArray<VkExtensionProperties> Properties = GetIVulkanDynamicRHI()->RHIGetAllInstanceExtensions();
 
 	TArray<const char*> Extensions;
 	{
@@ -58,13 +52,7 @@ bool FVulkanExtensions::GetVulkanInstanceExtensionsRequired(TArray<const ANSICHA
 bool FVulkanExtensions::GetVulkanDeviceExtensionsRequired(struct VkPhysicalDevice_T *pPhysicalDevice, TArray<const ANSICHAR*>& Out)
 {
 #if OCULUS_HMD_SUPPORTED_PLATFORMS_VULKAN
-	TArray<VkExtensionProperties> Properties;
-	{
-		uint32_t PropertyCount;
-		VulkanRHI::vkEnumerateDeviceExtensionProperties((VkPhysicalDevice) pPhysicalDevice, nullptr, &PropertyCount, nullptr);
-		Properties.SetNum(PropertyCount);
-		VulkanRHI::vkEnumerateDeviceExtensionProperties((VkPhysicalDevice) pPhysicalDevice, nullptr, &PropertyCount, Properties.GetData());
-	}
+	TArray<VkExtensionProperties> Properties = GetIVulkanDynamicRHI()->RHIGetAllDeviceExtensions((VkPhysicalDevice)pPhysicalDevice);
 
 	TArray<const char*> Extensions;
 	{
diff --git a/Engine/Plugins/Runtime/Oculus/OculusVR/Source/OculusMR/OculusMR.Build.cs b/Engine/Plugins/Runtime/Oculus/OculusVR/Source/OculusMR/OculusMR.Build.cs
--- a/Engine/Plugins/Runtime/Oculus/OculusVR/Source/OculusMR/OculusMR.Build.cs
+++ b/Engine/Plugins/Runtime/Oculus/OculusVR/Source/OculusMR/OculusMR.Build.cs
@@ -51,8 +51,6 @@ namespace UnrealBuildTool.Rules
 
 			if (Target.Platform == UnrealTargetPlatform.Android)
             {
-                PrivateIncludePaths.Add("../../../../../Source/Runtime/VulkanRHI/Private");
-                PrivateIncludePaths.Add("../../../../../Source/Runtime/VulkanRHI/Private/" + Target.Platform);
                 AddEngineThirdPartyPrivateStaticDependencies(Target, "Vulkan");
             }
 
diff --git a/Engine/Plugins/Runtime/Oculus/OculusVR/Source/OculusMR/Private/OculusMRModule.cpp b/Engine/Plugins/Runtime/Oculus/OculusVR/Source/OculusMR/Private/OculusMRModule.cpp
--- a/Engine/Plugins/Runtime/Oculus/OculusVR/Source/OculusMR/Private/OculusMRModule.cpp
+++ b/Engine/Plugins/Runtime/Oculus/OculusVR/Source/OculusMR/Private/OculusMRModule.cpp
@@ -20,8 +20,7 @@
 #include "OculusMR_CastingCameraActor.h"
 #include "AudioDevice.h"
 #if PLATFORM_ANDROID
-#include "VulkanRHIBridge.h"
-#include "VulkanRHIPrivate.h"
+#include "IVulkanDynamicRHI.h"
 #endif
 
 #if WITH_EDITOR
@@ -461,16 +460,16 @@ void FOculusMRModule::OnInitialWorldCreated(UWorld* NewWorld)
 						UE_LOG(LogMR, Log, TEXT("MRC Enabled"));
 
 						// Find a free queue index for vulkan
-						if (IsVulkanPlatform(GMaxRHIShaderPlatform))
+						if (RHIGetInterfaceType() == ERHIInterfaceType::Vulkan)
 						{
 							unsigned int queueIndex = 0;
 							ExecuteOnRenderThread([&queueIndex]()
 							{
 								ExecuteOnRHIThread([&queueIndex]()
 								{
-									FVulkanDevice* VulkanDevice = VulkanRHIBridge::GetDevice(GVulkanRHI);
-									int gfxIndex = VulkanDevice->GetGraphicsQueue() ? VulkanDevice->GetGraphicsQueue()->GetQueueIndex() : -1;
-									if (gfxIndex == queueIndex) {
+									const uint32 GraphicsQueueIndex = GetIVulkanDynamicRHI()->RHIGetGraphicsQueueIndex();
+									if (GraphicsQueueIndex == queueIndex)
+									{
 										++queueIndex;
 									}
 								});
diff --git a/Engine/Plugins/Runtime/Oculus/OculusVR/Source/OculusOpenXRHMD/OculusOpenXRHMD.Build.cs b/Engine/Plugins/Runtime/Oculus/OculusVR/Source/OculusOpenXRHMD/OculusOpenXRHMD.Build.cs
--- a/Engine/Plugins/Runtime/Oculus/OculusVR/Source/OculusOpenXRHMD/OculusOpenXRHMD.Build.cs
+++ b/Engine/Plugins/Runtime/Oculus/OculusVR/Source/OculusOpenXRHMD/OculusOpenXRHMD.Build.cs
@@ -14,7 +14,6 @@ namespace UnrealBuildTool.Rules
                     // Relative to Engine\Plugins\Runtime\Oculus\OculusOpenXR\Source
                     "../../../OpenXR/Source/OpenXRHMD/Private",
                     "../../../../../Source/Runtime/Renderer/Private",
-                    "../../../../../Source/Runtime/VulkanRHI/Private",
                     "../../../../../Source/Runtime/Engine/Classes/Components",
                     "../../../../../Source/Runtime/Engine/Classes/Kismet",
                 });
@@ -25,15 +24,6 @@ namespace UnrealBuildTool.Rules
                     "OpenXRHMD",
                 });			
 
-            if (Target.Platform == UnrealTargetPlatform.Win64)
-            {
-                PrivateIncludePaths.Add("../../../../../Source/Runtime/VulkanRHI/Private/Windows");
-            }
-            else
-            {
-                PrivateIncludePaths.Add("../../../../../Source/Runtime/VulkanRHI/Private/" + Target.Platform);
-            }
-
             PrivateDependencyModuleNames.AddRange(
                 new string[]
                 {
diff --git a/Engine/Plugins/Runtime/OpenXR/Source/OpenXRHMD/OpenXRHMD.Build.cs b/Engine/Plugins/Runtime/OpenXR/Source/OpenXRHMD/OpenXRHMD.Build.cs
--- a/Engine/Plugins/Runtime/OpenXR/Source/OpenXRHMD/OpenXRHMD.Build.cs
+++ b/Engine/Plugins/Runtime/OpenXR/Source/OpenXRHMD/OpenXRHMD.Build.cs
@@ -15,22 +15,11 @@ namespace UnrealBuildTool.Rules
 					"OpenXRHMD/Private",
                     EngineDir + "/Source/ThirdParty/OpenXR/include",
                     EngineDir + "/Source/Runtime/Renderer/Private",
-                    EngineDir + "/Source/Runtime/VulkanRHI/Private",
 					// ... add other private include paths required here ...
 				}
 				);
 
-            if (Target.Platform == UnrealTargetPlatform.Win64)
-            {
-                PrivateIncludePaths.Add(EngineDir + "/Source/Runtime/VulkanRHI/Private/Windows");
-            }
-			else if (Target.Platform == UnrealTargetPlatform.Android  || Target.Platform == UnrealTargetPlatform.Linux)
-            {
-                PrivateIncludePaths.Add(EngineDir + "/Source/Runtime/VulkanRHI/Private/" + Target.Platform);
-            }
-
 			PublicIncludePathModuleNames.Add("OpenXR");
-			
 
             PublicDependencyModuleNames.Add("HeadMountedDisplay");
 
diff --git a/Engine/Plugins/Runtime/OpenXR/Source/OpenXRHMD/Private/OpenXRHMD.cpp b/Engine/Plugins/Runtime/OpenXR/Source/OpenXRHMD/Private/OpenXRHMD.cpp
--- a/Engine/Plugins/Runtime/OpenXR/Source/OpenXRHMD/Private/OpenXRHMD.cpp
+++ b/Engine/Plugins/Runtime/OpenXR/Source/OpenXRHMD/Private/OpenXRHMD.cpp
@@ -90,13 +90,7 @@ namespace {
 bool FOpenXRHMD::FVulkanExtensions::GetVulkanInstanceExtensionsRequired(TArray<const ANSICHAR*>& Out)
 {
 #ifdef XR_USE_GRAPHICS_API_VULKAN
-	TArray<VkExtensionProperties> Properties;
-	{
-		uint32_t PropertyCount;
-		VulkanRHI::vkEnumerateInstanceExtensionProperties(nullptr, &PropertyCount, nullptr);
-		Properties.SetNum(PropertyCount);
-		VulkanRHI::vkEnumerateInstanceExtensionProperties(nullptr, &PropertyCount, Properties.GetData());
-	}
+	TArray<VkExtensionProperties> Properties = GetIVulkanDynamicRHI()->RHIGetAllInstanceExtensions();
 
 	{
 		PFN_xrGetVulkanInstanceExtensionsKHR GetVulkanInstanceExtensionsKHR;
@@ -137,13 +131,7 @@ bool FOpenXRHMD::FVulkanExtensions::GetVulkanInstanceExtensionsRequired(TArray<c
 bool FOpenXRHMD::FVulkanExtensions::GetVulkanDeviceExtensionsRequired(VkPhysicalDevice_T *pPhysicalDevice, TArray<const ANSICHAR*>& Out)
 {
 #ifdef XR_USE_GRAPHICS_API_VULKAN
-	TArray<VkExtensionProperties> Properties;
-	{
-		uint32_t PropertyCount;
-		VulkanRHI::vkEnumerateDeviceExtensionProperties((VkPhysicalDevice)pPhysicalDevice, nullptr, &PropertyCount, nullptr);
-		Properties.SetNum(PropertyCount);
-		VulkanRHI::vkEnumerateDeviceExtensionProperties((VkPhysicalDevice)pPhysicalDevice, nullptr, &PropertyCount, Properties.GetData());
-	}
+	TArray<VkExtensionProperties> Properties = GetIVulkanDynamicRHI()->RHIGetAllDeviceExtensions((VkPhysicalDevice)pPhysicalDevice);
 
 	{
 		PFN_xrGetVulkanDeviceExtensionsKHR GetVulkanDeviceExtensionsKHR;
diff --git a/Engine/Plugins/Runtime/OpenXR/Source/OpenXRHMD/Private/OpenXRHMD_RenderBridge.cpp b/Engine/Plugins/Runtime/OpenXR/Source/OpenXRHMD/Private/OpenXRHMD_RenderBridge.cpp
--- a/Engine/Plugins/Runtime/OpenXR/Source/OpenXRHMD/Private/OpenXRHMD_RenderBridge.cpp
+++ b/Engine/Plugins/Runtime/OpenXR/Source/OpenXRHMD/Private/OpenXRHMD_RenderBridge.cpp
@@ -234,11 +234,14 @@ public:
 		Requirements.maxApiVersionSupported = 0;
 		XR_ENSURE(GetVulkanGraphicsRequirementsKHR(InInstance, InSystem, &Requirements));
 
+		IVulkanDynamicRHI* VulkanRHI = GetIVulkanDynamicRHI();
+		const uint32 VulkanVersion = VulkanRHI->RHIGetVulkanVersion();
+
 		// The extension uses the OpenXR version format instead of the Vulkan one
 		XrVersion RHIVersion = XR_MAKE_VERSION(
-			VK_VERSION_MAJOR(UE_VK_API_VERSION),
-			VK_VERSION_MINOR(UE_VK_API_VERSION),
-			VK_VERSION_PATCH(UE_VK_API_VERSION)
+			VK_VERSION_MAJOR(VulkanVersion),
+			VK_VERSION_MINOR(VulkanVersion),
+			VK_VERSION_PATCH(VulkanVersion)
 		);
 		if (RHIVersion < Requirements.minApiVersionSupported) //-V547
 		{
@@ -252,40 +255,30 @@ public:
 
 		PFN_xrGetVulkanGraphicsDeviceKHR GetVulkanGraphicsDeviceKHR;
 		XR_ENSURE(xrGetInstanceProcAddr(Instance, "xrGetVulkanGraphicsDeviceKHR", (PFN_xrVoidFunction*)&GetVulkanGraphicsDeviceKHR));
-		XR_ENSURE(GetVulkanGraphicsDeviceKHR(Instance, System, GVulkanRHI->GetInstance(), &Gpu));
+		XR_ENSURE(GetVulkanGraphicsDeviceKHR(Instance, System, VulkanRHI->RHIGetVkInstance(), &Gpu));
 	}
 
 	virtual void* GetGraphicsBinding() override
 	{
-		FVulkanDevice* Device = GVulkanRHI->GetDevice();
-		FVulkanQueue* Queue = Device->GetGraphicsQueue();
+		IVulkanDynamicRHI* VulkanRHI = GetIVulkanDynamicRHI();
 
 		Binding.type = XR_TYPE_GRAPHICS_BINDING_VULKAN_KHR;
 		Binding.next = nullptr;
-		Binding.instance = GVulkanRHI->GetInstance();
-		Binding.physicalDevice = Device->GetPhysicalHandle();
-		Binding.device = Device->GetInstanceHandle();
-		Binding.queueFamilyIndex = Queue->GetFamilyIndex();
+		Binding.instance = VulkanRHI->RHIGetVkInstance();
+		Binding.physicalDevice = VulkanRHI->RHIGetVkPhysicalDevice();
+		Binding.device = VulkanRHI->RHIGetVkDevice();
+		Binding.queueFamilyIndex = VulkanRHI->RHIGetGraphicsQueueFamilyIndex();
 		Binding.queueIndex = 0;
 		return &Binding;
 	}
 
 	virtual uint64 GetGraphicsAdapterLuid() override
 	{
-#if VULKAN_SUPPORTS_DRIVER_PROPERTIES
-		if (!AdapterLuid && GVulkanRHI->GetOptionalExtensions().HasKHRGetPhysicalDeviceProperties2)
+		if (!AdapterLuid)
 		{
-			VkPhysicalDeviceIDPropertiesKHR GpuIdProps;
-			VkPhysicalDeviceProperties2KHR GpuProps2;
-			ZeroVulkanStruct(GpuProps2, VK_STRUCTURE_TYPE_PHYSICAL_DEVICE_PROPERTIES_2_KHR);
-			GpuProps2.pNext = &GpuIdProps;
-			ZeroVulkanStruct(GpuIdProps, VK_STRUCTURE_TYPE_PHYSICAL_DEVICE_ID_PROPERTIES_KHR);
-
-			VulkanRHI::vkGetPhysicalDeviceProperties2KHR(Gpu, &GpuProps2);
-			check(GpuIdProps.deviceLUIDValid);
-			AdapterLuid = reinterpret_cast<const uint64&>(GpuIdProps.deviceLUID);
+			AdapterLuid = GetIVulkanDynamicRHI()->RHIGetGraphicsAdapterLUID(Gpu);
 		}
-#endif
+
 		return AdapterLuid;
 	}
 
diff --git a/Engine/Plugins/Runtime/OpenXR/Source/OpenXRHMD/Private/OpenXRHMD_Swapchain.cpp b/Engine/Plugins/Runtime/OpenXR/Source/OpenXRHMD/Private/OpenXRHMD_Swapchain.cpp
--- a/Engine/Plugins/Runtime/OpenXR/Source/OpenXRHMD/Private/OpenXRHMD_Swapchain.cpp
+++ b/Engine/Plugins/Runtime/OpenXR/Source/OpenXRHMD/Private/OpenXRHMD_Swapchain.cpp
@@ -387,7 +387,7 @@ FXRSwapChainPtr CreateSwapchain_Vulkan(XrSession InSession, uint8 Format, uint32
 	TFunction<uint32(uint8)> ToPlatformFormat = [](uint8 InFormat)
 	{
 		// UE renders a gamma-corrected image so we need to use an sRGB format if available
-		return UEToVkTextureFormat(GPixelFormats[InFormat].UnrealFormat, true);
+		return GetIVulkanDynamicRHI()->RHIGetSwapChainVkFormat(static_cast<EPixelFormat>(InFormat));
 	};
 	Format = FOpenXRSwapchain::GetNearestSupportedSwapchainFormat(InSession, Format, ToPlatformFormat);
 	if (!Format)
@@ -401,13 +401,15 @@ FXRSwapChainPtr CreateSwapchain_Vulkan(XrSession InSession, uint8 Format, uint32
 		return nullptr;
 	}
 
+	IVulkanDynamicRHI* VulkanRHI = GetIVulkanDynamicRHI();
+
 	TArray<FTextureRHIRef> TextureChain;
 	TArray<XrSwapchainImageVulkanKHR> Images = EnumerateImages<XrSwapchainImageVulkanKHR>(Swapchain, XR_TYPE_SWAPCHAIN_IMAGE_VULKAN_KHR);
 	for (const auto& Image : Images)
 	{
-		TextureChain.Add(static_cast<FTextureRHIRef>(GVulkanRHI->RHICreateTexture2DArrayFromResource(GPixelFormats[Format].UnrealFormat, SizeX, SizeY, ArraySize, NumMips, NumSamples, Image.image, CreateFlags)));
+		TextureChain.Add(static_cast<FTextureRHIRef>(VulkanRHI->RHICreateTexture2DArrayFromResource(GPixelFormats[Format].UnrealFormat, SizeX, SizeY, ArraySize, NumMips, NumSamples, Image.image, CreateFlags)));
 	}
-	FTextureRHIRef ChainTarget = static_cast<FTextureRHIRef>(GDynamicRHI->RHICreateAliasedTexture((FTextureRHIRef&)TextureChain[0]));
+	FTextureRHIRef ChainTarget = static_cast<FTextureRHIRef>(VulkanRHI->RHICreateAliasedTexture((FTextureRHIRef&)TextureChain[0]));
 
 	return CreateXRSwapChain<FOpenXRSwapchain>(MoveTemp(TextureChain), ChainTarget, Swapchain);
 }
diff --git a/Engine/Plugins/Runtime/OpenXR/Source/OpenXRHMD/Private/OpenXRPlatformRHI.h b/Engine/Plugins/Runtime/OpenXR/Source/OpenXRHMD/Private/OpenXRPlatformRHI.h
--- a/Engine/Plugins/Runtime/OpenXR/Source/OpenXRHMD/Private/OpenXRPlatformRHI.h
+++ b/Engine/Plugins/Runtime/OpenXR/Source/OpenXRHMD/Private/OpenXRPlatformRHI.h
@@ -52,8 +52,7 @@
 //-------------------------------------------------------------------------------------------------
 
 #ifdef XR_USE_GRAPHICS_API_VULKAN
-#include "VulkanRHIPrivate.h"
-#include "VulkanResources.h"
+#include "IVulkanDynamicRHI.h"
 #endif // XR_USE_GRAPHICS_API_VULKAN
 
 #if PLATFORM_WINDOWS
diff --git a/Engine/Plugins/Runtime/Steam/SteamVR/Source/SteamVR/Private/SteamVRRender.cpp b/Engine/Plugins/Runtime/Steam/SteamVR/Source/SteamVR/Private/SteamVRRender.cpp
--- a/Engine/Plugins/Runtime/Steam/SteamVR/Source/SteamVR/Private/SteamVRRender.cpp
+++ b/Engine/Plugins/Runtime/Steam/SteamVR/Source/SteamVR/Private/SteamVRRender.cpp
@@ -18,9 +18,7 @@
 #if PLATFORM_MAC
 #include <Metal/Metal.h>
 #else
-#include "VulkanRHIPrivate.h"
-#include "VulkanPendingState.h"
-#include "VulkanContext.h"
+#include "IVulkanDynamicRHI.h"
 #endif
 
 #if PLATFORM_WINDOWS
@@ -356,23 +354,21 @@ void FSteamVRHMD::VulkanBridge::FinishRendering()
 	bool bSubmitDepth = false; // CVarEnableDepthSubmission->GetInt() > 0;
 	vr::EVRSubmitFlags Flags = bSubmitDepth ? vr::EVRSubmitFlags::Submit_TextureWithDepth : vr::EVRSubmitFlags::Submit_Default;
 
-	if(SwapChain->GetTexture2D())
+	if (FRHITexture2D* SwapChainTexture = SwapChain->GetTexture2D())
 	{
-		FVulkanTexture2D* Texture2D = (FVulkanTexture2D*)SwapChain->GetTexture2D();
-		FVulkanTexture2D* DepthTexture2D = (FVulkanTexture2D*)DepthSwapChain->GetTexture2D();
-
-		FVulkanCommandListContext& ImmediateContext = GVulkanRHI->GetDevice()->GetImmediateContext();
+		IVulkanDynamicRHI* VulkanRHI = GetIVulkanDynamicRHI();
+		const VkImage SwapChainImage = VulkanRHI->RHIGetVkImage(SwapChainTexture);
+		const VkFormat SwapChainFormat = VulkanRHI->RHIGetViewVkFormat(SwapChainTexture);
 
 		// Color layout
-		VkImageLayout& CurrentLayout = ImmediateContext.GetLayoutManager().FindOrAddLayoutRW(Texture2D->Surface, VK_IMAGE_LAYOUT_UNDEFINED);
+		VkImageLayout& CurrentLayout = VulkanRHI->RHIFindOrAddLayoutRW(SwapChainTexture, VK_IMAGE_LAYOUT_UNDEFINED);
 		bool bHadLayout = (CurrentLayout != VK_IMAGE_LAYOUT_UNDEFINED);
 		
-		FVulkanCmdBuffer* CmdBuffer = ImmediateContext.GetCommandBufferManager()->GetUploadCmdBuffer();
 		VkImageSubresourceRange SubresourceRange = { VK_IMAGE_ASPECT_COLOR_BIT, 0, 1, 0, 1 };
 
 		if (CurrentLayout != VK_IMAGE_LAYOUT_TRANSFER_SRC_OPTIMAL)
 		{
-			GVulkanRHI->VulkanSetImageLayout(CmdBuffer->GetHandle(), Texture2D->Surface.Image, CurrentLayout, VK_IMAGE_LAYOUT_TRANSFER_SRC_OPTIMAL, SubresourceRange);
+			VulkanRHI->RHISetUploadImageLayout(SwapChainImage, CurrentLayout, VK_IMAGE_LAYOUT_TRANSFER_SRC_OPTIMAL, SubresourceRange);
 		}
 
 		vr::VRTextureBounds_t LeftBounds;
@@ -388,38 +384,43 @@ void FSteamVRHMD::VulkanBridge::FinishRendering()
 		RightBounds.vMax = 1.0f;
 
 		vr::VRVulkanTextureData_t VulkanTextureDataColor {};
-		VulkanTextureDataColor.m_pInstance			= GVulkanRHI->GetInstance();
-		VulkanTextureDataColor.m_pDevice			= GVulkanRHI->GetDevice()->GetInstanceHandle();
-		VulkanTextureDataColor.m_pPhysicalDevice	= GVulkanRHI->GetDevice()->GetPhysicalHandle();
-		VulkanTextureDataColor.m_pQueue				= GVulkanRHI->GetDevice()->GetGraphicsQueue()->GetHandle();
-		VulkanTextureDataColor.m_nQueueFamilyIndex	= GVulkanRHI->GetDevice()->GetGraphicsQueue()->GetFamilyIndex();
-		VulkanTextureDataColor.m_nImage				= (uint64_t)Texture2D->Surface.Image;
-		VulkanTextureDataColor.m_nWidth				= Texture2D->Surface.Width;
-		VulkanTextureDataColor.m_nHeight			= Texture2D->Surface.Height;
-		VulkanTextureDataColor.m_nFormat			= (uint32_t)Texture2D->Surface.ViewFormat;
-		VulkanTextureDataColor.m_nSampleCount = 1;
+		VulkanTextureDataColor.m_pInstance			= VulkanRHI->RHIGetVkInstance();
+		VulkanTextureDataColor.m_pDevice			= VulkanRHI->RHIGetVkDevice();
+		VulkanTextureDataColor.m_pPhysicalDevice	= VulkanRHI->RHIGetVkPhysicalDevice();
+		VulkanTextureDataColor.m_pQueue				= VulkanRHI->RHIGetGraphicsVkQueue();
+		VulkanTextureDataColor.m_nQueueFamilyIndex	= VulkanRHI->RHIGetGraphicsQueueFamilyIndex();
+		VulkanTextureDataColor.m_nImage				= (uint64_t)SwapChainImage;
+		VulkanTextureDataColor.m_nWidth				= SwapChainTexture->GetSizeX();
+		VulkanTextureDataColor.m_nHeight			= SwapChainTexture->GetSizeY();
+		VulkanTextureDataColor.m_nFormat			= (uint32_t)SwapChainFormat;
+		VulkanTextureDataColor.m_nSampleCount       = 1;
 
 		if (bSubmitDepth)
 		{
+			FRHITexture2D* DepthSwapChainTexture = DepthSwapChain->GetTexture2D();
+
+			const VkImage DepthSwapChainImage = VulkanRHI->RHIGetVkImage(DepthSwapChainTexture);
+			const VkFormat DepthSwapChainFormat = VulkanRHI->RHIGetViewVkFormat(DepthSwapChainTexture);
+
 			VkImageSubresourceRange SubresourceRangeDepth = { VK_IMAGE_ASPECT_DEPTH_BIT | VK_IMAGE_ASPECT_STENCIL_BIT, 0, 1, 0, 1 };
-			VkImageLayout& CurrentDepthLayout = ImmediateContext.GetLayoutManager().FindOrAddLayoutRW(DepthTexture2D->Surface, VK_IMAGE_LAYOUT_UNDEFINED);
+			VkImageLayout& CurrentDepthLayout = VulkanRHI->RHIFindOrAddLayoutRW(DepthSwapChainTexture, VK_IMAGE_LAYOUT_UNDEFINED);
 			bool bDepthHadLayout = (CurrentDepthLayout != VK_IMAGE_LAYOUT_UNDEFINED);
 
 			if (CurrentDepthLayout != VK_IMAGE_LAYOUT_TRANSFER_SRC_OPTIMAL)
 			{
-				GVulkanRHI->VulkanSetImageLayout(CmdBuffer->GetHandle(), DepthTexture2D->Surface.Image, CurrentDepthLayout, VK_IMAGE_LAYOUT_TRANSFER_SRC_OPTIMAL, SubresourceRangeDepth);
+				VulkanRHI->RHISetUploadImageLayout(DepthSwapChainImage, CurrentDepthLayout, VK_IMAGE_LAYOUT_TRANSFER_SRC_OPTIMAL, SubresourceRangeDepth);
 			}
 
 			vr::VRVulkanTextureData_t VulkanTextureDataDepth{};
-			VulkanTextureDataDepth.m_pInstance = GVulkanRHI->GetInstance();
-			VulkanTextureDataDepth.m_pDevice = GVulkanRHI->GetDevice()->GetInstanceHandle();
-			VulkanTextureDataDepth.m_pPhysicalDevice = GVulkanRHI->GetDevice()->GetPhysicalHandle();
-			VulkanTextureDataDepth.m_pQueue = GVulkanRHI->GetDevice()->GetGraphicsQueue()->GetHandle();
-			VulkanTextureDataDepth.m_nQueueFamilyIndex = GVulkanRHI->GetDevice()->GetGraphicsQueue()->GetFamilyIndex();
-			VulkanTextureDataDepth.m_nImage = (uint64_t)DepthTexture2D->Surface.Image;
-			VulkanTextureDataDepth.m_nWidth = DepthTexture2D->Surface.Width;
-			VulkanTextureDataDepth.m_nHeight = DepthTexture2D->Surface.Height;
-			VulkanTextureDataDepth.m_nFormat = (uint32_t)DepthTexture2D->Surface.ViewFormat;
+			VulkanTextureDataDepth.m_pInstance         = VulkanRHI->RHIGetVkInstance();
+			VulkanTextureDataDepth.m_pDevice           = VulkanRHI->RHIGetVkDevice();
+			VulkanTextureDataDepth.m_pPhysicalDevice   = VulkanRHI->RHIGetVkPhysicalDevice();
+			VulkanTextureDataDepth.m_pQueue            = VulkanRHI->RHIGetGraphicsVkQueue();
+			VulkanTextureDataDepth.m_nQueueFamilyIndex = VulkanRHI->RHIGetGraphicsQueueFamilyIndex();
+			VulkanTextureDataDepth.m_nImage            = (uint64_t)DepthSwapChainImage;
+			VulkanTextureDataDepth.m_nWidth            = DepthSwapChainTexture->GetSizeX();
+			VulkanTextureDataDepth.m_nHeight           = DepthSwapChainTexture->GetSizeY();
+			VulkanTextureDataDepth.m_nFormat           = (uint32_t)DepthSwapChainFormat;
 			VulkanTextureDataDepth.m_nSampleCount = 1;
 
 			vr::VRTextureWithDepth_t Texture;
@@ -446,7 +447,7 @@ void FSteamVRHMD::VulkanBridge::FinishRendering()
 
 			if (bDepthHadLayout && CurrentDepthLayout != VK_IMAGE_LAYOUT_TRANSFER_SRC_OPTIMAL)
 			{
-				GVulkanRHI->VulkanSetImageLayout(CmdBuffer->GetHandle(), DepthTexture2D->Surface.Image, VK_IMAGE_LAYOUT_TRANSFER_SRC_OPTIMAL, CurrentDepthLayout, SubresourceRangeDepth);
+				VulkanRHI->RHISetUploadImageLayout(DepthSwapChainImage, VK_IMAGE_LAYOUT_TRANSFER_SRC_OPTIMAL, CurrentDepthLayout, SubresourceRangeDepth);
 			}
 			else
 			{
@@ -466,14 +467,14 @@ void FSteamVRHMD::VulkanBridge::FinishRendering()
 
 		if (bHadLayout && CurrentLayout != VK_IMAGE_LAYOUT_TRANSFER_SRC_OPTIMAL)
 		{
-			GVulkanRHI->VulkanSetImageLayout(CmdBuffer->GetHandle(), Texture2D->Surface.Image, VK_IMAGE_LAYOUT_TRANSFER_SRC_OPTIMAL, CurrentLayout, SubresourceRange);
+			VulkanRHI->RHISetUploadImageLayout(SwapChainImage, VK_IMAGE_LAYOUT_TRANSFER_SRC_OPTIMAL, CurrentLayout, SubresourceRange);
 		}
 		else
 		{
 			CurrentLayout = VK_IMAGE_LAYOUT_TRANSFER_SRC_OPTIMAL;
 		}
 
-		ImmediateContext.GetCommandBufferManager()->SubmitUploadCmdBuffer();
+		VulkanRHI->RHISubmitUploadCommandBuffer();
 	}
 }
 
diff --git a/Engine/Plugins/Runtime/Steam/SteamVR/Source/SteamVR/Private/SteamVRSpectatorScreen.cpp b/Engine/Plugins/Runtime/Steam/SteamVR/Source/SteamVR/Private/SteamVRSpectatorScreen.cpp
--- a/Engine/Plugins/Runtime/Steam/SteamVR/Source/SteamVR/Private/SteamVRSpectatorScreen.cpp
+++ b/Engine/Plugins/Runtime/Steam/SteamVR/Source/SteamVR/Private/SteamVRSpectatorScreen.cpp
@@ -15,10 +15,8 @@
 #include "DefaultSpectatorScreenController.h"
 
 #if PLATFORM_LINUX
-#include "VulkanRHIPrivate.h"
+#include "IVulkanDynamicRHI.h"
 #include "ScreenRendering.h"
-#include "VulkanPendingState.h"
-#include "VulkanContext.h"
 #endif
 
 
diff --git a/Engine/Plugins/Runtime/Steam/SteamVR/Source/SteamVR/Private/SteamVRStereoLayers.cpp b/Engine/Plugins/Runtime/Steam/SteamVR/Source/SteamVR/Private/SteamVRStereoLayers.cpp
--- a/Engine/Plugins/Runtime/Steam/SteamVR/Source/SteamVR/Private/SteamVRStereoLayers.cpp
+++ b/Engine/Plugins/Runtime/Steam/SteamVR/Source/SteamVR/Private/SteamVRStereoLayers.cpp
@@ -21,7 +21,7 @@ static vr::EVROverlayError sOvrError;
 #endif
 
 #if !PLATFORM_MAC
-#include "VulkanRHIPrivate.h"
+#include "IVulkanDynamicRHI.h"
 #endif
 
 /*=============================================================================
@@ -251,31 +251,32 @@ void FSteamVRHMD::UpdateStereoLayers_RenderThread()
 			{
 				vr::Texture_t Texture;
 				vr::VRVulkanTextureData_t VulkanTexture{};
-				if ( IsVulkanPlatform( GMaxRHIShaderPlatform ) )
+				if (RHIGetInterfaceType() == ERHIInterfaceType::Vulkan)
 				{
 #if PLATFORM_MAC
 					check( 0 );
 #else
+					IVulkanDynamicRHI* VulkanRHI = GetIVulkanDynamicRHI();
+
 					FRHITexture2D* TextureRHI2D = Layer.LayerDesc.Texture->GetTexture2D();
 					check(TextureRHI2D);
-					FVulkanTexture2D* Texture2D = (FVulkanTexture2D*)TextureRHI2D;
-
-					VulkanTexture.m_pInstance = GVulkanRHI->GetInstance();
-					VulkanTexture.m_pDevice = GVulkanRHI->GetDevice()->GetInstanceHandle();
-					VulkanTexture.m_pPhysicalDevice = GVulkanRHI->GetDevice()->GetPhysicalHandle();
-					VulkanTexture.m_pQueue = GVulkanRHI->GetDevice()->GetGraphicsQueue()->GetHandle();
-					VulkanTexture.m_nQueueFamilyIndex = GVulkanRHI->GetDevice()->GetGraphicsQueue()->GetFamilyIndex();
-					VulkanTexture.m_nImage = (uint64_t)Texture2D->Surface.Image;
-					VulkanTexture.m_nWidth = Texture2D->Surface.Width;
-					VulkanTexture.m_nHeight = Texture2D->Surface.Height;
-					VulkanTexture.m_nFormat = (uint32_t)Texture2D->Surface.ViewFormat;
+
+					VulkanTexture.m_pInstance = VulkanRHI->RHIGetVkInstance();
+					VulkanTexture.m_pDevice = VulkanRHI->RHIGetVkDevice();
+					VulkanTexture.m_pPhysicalDevice = VulkanRHI->RHIGetVkPhysicalDevice();
+					VulkanTexture.m_pQueue = VulkanRHI->RHIGetGraphicsVkQueue();
+					VulkanTexture.m_nQueueFamilyIndex = VulkanRHI->RHIGetGraphicsQueueFamilyIndex();
+					VulkanTexture.m_nImage = (uint64_t)VulkanRHI->RHIGetVkImage(TextureRHI2D);
+					VulkanTexture.m_nWidth = TextureRHI2D->GetSizeX();
+					VulkanTexture.m_nHeight = TextureRHI2D->GetSizeY();
+					VulkanTexture.m_nFormat = (uint32_t)VulkanRHI->RHIGetViewVkFormat(TextureRHI2D);
 					VulkanTexture.m_nSampleCount = 1;
 
 					Texture.handle = &VulkanTexture;
 					Texture.eType = vr::TextureType_Vulkan;
 #endif
 				}
-				else if ( IsOpenGLPlatform( GMaxRHIShaderPlatform ) )
+				else if (RHIGetInterfaceType() == ERHIInterfaceType::OpenGL)
 				{
 					// We need to dereference the pointer to the real handle
 					uintptr_t TextureID = *reinterpret_cast<uint32*>(Layer.LayerDesc.Texture->GetNativeResource());
diff --git a/Engine/Plugins/Runtime/Steam/SteamVR/Source/SteamVR/SteamVR.Build.cs b/Engine/Plugins/Runtime/Steam/SteamVR/Source/SteamVR/SteamVR.Build.cs
--- a/Engine/Plugins/Runtime/Steam/SteamVR/Source/SteamVR/SteamVR.Build.cs
+++ b/Engine/Plugins/Runtime/Steam/SteamVR/Source/SteamVR/SteamVR.Build.cs
@@ -10,20 +10,10 @@ namespace UnrealBuildTool.Rules
 				new string[] {
 					"SteamVR/Private",
 					"../../../../../Source/Runtime/Renderer/Private",
-					"../../../../../Source/Runtime/VulkanRHI/Private",
 					// ... add other private include paths required here ...
 				}
 				);
 
-			if(Target.Platform == UnrealTargetPlatform.Win64)
-			{
-				PrivateIncludePaths.Add("../../../../../Source/Runtime/VulkanRHI/Private/Windows");
-			}
-			else if(Target.Platform != UnrealTargetPlatform.Mac)
-			{
-				PrivateIncludePaths.Add("../../../../../Source/Runtime/VulkanRHI/Private/" + Target.Platform);
-			}
-
 			PrivateDependencyModuleNames.AddRange(
 				new string[]
 				{
diff --git a/Engine/Source/Runtime/AVEncoder/AVEncoder.Build.cs b/Engine/Source/Runtime/AVEncoder/AVEncoder.Build.cs
--- a/Engine/Source/Runtime/AVEncoder/AVEncoder.Build.cs
+++ b/Engine/Source/Runtime/AVEncoder/AVEncoder.Build.cs
@@ -41,17 +41,9 @@ public class AVEncoder : ModuleRules
 			// ... add any modules that your module loads dynamically here ...
 		});
 
-		string EngineSourceDirectory = Path.GetFullPath(Target.RelativeEnginePath);
-
 		if (Target.IsInPlatformGroup(UnrealPlatformGroup.Windows) || Target.IsInPlatformGroup(UnrealPlatformGroup.Unix))
 		{
-			AddEngineThirdPartyPrivateStaticDependencies(Target, "NVAftermath");
-			PrivateIncludePathModuleNames.Add("VulkanRHI");
-
-			PrivateIncludePaths.Add(Path.Combine(EngineSourceDirectory, "Source/Runtime/VulkanRHI/Private"));
-			AddEngineThirdPartyPrivateStaticDependencies(Target, "Vulkan");
-
-			PublicIncludePaths.Add(Path.Combine(EngineSourceDirectory, "Source/ThirdParty/Vulkan/Include"));
+			PublicIncludePathModuleNames.Add("Vulkan");
 		}
 
 		if (Target.IsInPlatformGroup(UnrealPlatformGroup.Windows))
@@ -61,11 +53,7 @@ public class AVEncoder : ModuleRules
 				"D3D12RHI"
 			});
 
-			PublicSystemLibraries.AddRange(new string[] {
-				"DXGI.lib",
-				"d3d11.lib",
-				"d3d12.lib"
-			});
+			AddEngineThirdPartyPrivateStaticDependencies(Target, "DX11", "DX12");
 			
 			PublicDelayLoadDLLs.Add("mfplat.dll");
 			PublicDelayLoadDLLs.Add("mfuuid.dll");
diff --git a/Engine/Source/Runtime/AVEncoder/Private/VideoEncoderInput.cpp b/Engine/Source/Runtime/AVEncoder/Private/VideoEncoderInput.cpp
--- a/Engine/Source/Runtime/AVEncoder/Private/VideoEncoderInput.cpp
+++ b/Engine/Source/Runtime/AVEncoder/Private/VideoEncoderInput.cpp
@@ -10,10 +10,6 @@
 
 #include "Misc/Guid.h"
 
-#if PLATFORM_DESKTOP && !PLATFORM_APPLE
-#include "VulkanRHIBridge.h"
-#endif
-
 #if PLATFORM_WINDOWS
 #include "MicrosoftCommon.h"
 #endif
diff --git a/Engine/Source/Runtime/CUDA/Source/CUDA.Build.cs b/Engine/Source/Runtime/CUDA/Source/CUDA.Build.cs
--- a/Engine/Source/Runtime/CUDA/Source/CUDA.Build.cs
+++ b/Engine/Source/Runtime/CUDA/Source/CUDA.Build.cs
@@ -23,23 +23,14 @@ public class CUDA : ModuleRules
 			PublicDefinitions.Add("PLATFORM_SUPPORTS_CUDA=1");
 			var EngineDir = Path.GetFullPath(Target.RelativeEnginePath);
 
-			AddEngineThirdPartyPrivateStaticDependencies(Target, "NVAftermath");
-			
 			PrivateIncludePathModuleNames.Add("VulkanRHI");
-			PrivateIncludePaths.Add(Path.Combine(EngineDir, "Source/Runtime/VulkanRHI/Private"));	
 
 			if (Target.IsInPlatformGroup(UnrealPlatformGroup.Unix))
 			{
-				PrivateIncludePaths.Add(Path.Combine(EngineDir, "Source/Runtime/VulkanRHI/Private/Linux"));	
-			} else if (Target.IsInPlatformGroup(UnrealPlatformGroup.Windows))
+			}
+			else if (Target.IsInPlatformGroup(UnrealPlatformGroup.Windows))
 			{
-				PublicSystemLibraries.AddRange(new string[] {
-					"DXGI.lib",
-					"d3d11.lib",
-					"d3d12.lib"
-				});
-
-				PrivateIncludePaths.Add(Path.Combine(EngineDir, "Source/Runtime/VulkanRHI/Private/Windows"));
+				AddEngineThirdPartyPrivateStaticDependencies(Target, "DX11", "DX12");
 			}
 
 			AddEngineThirdPartyPrivateStaticDependencies(Target, "Vulkan");
diff --git a/Engine/Source/Runtime/CUDA/Source/Private/CudaModule.cpp b/Engine/Source/Runtime/CUDA/Source/Private/CudaModule.cpp
--- a/Engine/Source/Runtime/CUDA/Source/Private/CudaModule.cpp
+++ b/Engine/Source/Runtime/CUDA/Source/Private/CudaModule.cpp
@@ -3,8 +3,10 @@
 
 #include "Misc/CoreDelegates.h"
 
+#include "RHI.h"
+
 #if PLATFORM_SUPPORTS_CUDA
-#include "VulkanRHIPrivate.h"
+#include "IVulkanDynamicRHI.h"
 #endif
 
 #if PLATFORM_WINDOWS
@@ -16,7 +18,6 @@
     #include "Windows/HideWindowsPlatformTypes.h"
 #endif
 
-#include "RHI.h"
 #include "HAL/UnrealMemory.h"
 
 DEFINE_LOG_CATEGORY(LogCUDA);
@@ -105,8 +106,7 @@ bool FCUDAModule::IsRHISelectedDevice(CUdevice cuDevice)
 		uint8 deviceUUID[16];
 
 		// We find the device that the RHI has selected for us so that later we can create a CUDA context on this device.
-		FVulkanDynamicRHI *vkDynamicRHI = GetDynamicRHI<FVulkanDynamicRHI>();
-		FMemory::Memcpy(deviceUUID, vkDynamicRHI->GetDevice()->GetDeviceIdProperties().deviceUUID, 16);	
+		FMemory::Memcpy(deviceUUID, GetIVulkanDynamicRHI()->RHIGetVulkanDeviceUUID(), 16);	
 
 		// Get the device UUID so we can compare this with what the RHI selected.
 		CUuuid cudaDeviceUUID;
diff --git a/Engine/Source/Runtime/VulkanRHI/Private/VulkanLayers.cpp b/Engine/Source/Runtime/VulkanRHI/Private/VulkanLayers.cpp
--- a/Engine/Source/Runtime/VulkanRHI/Private/VulkanLayers.cpp
+++ b/Engine/Source/Runtime/VulkanRHI/Private/VulkanLayers.cpp
@@ -7,14 +7,7 @@
 #include "VulkanRHIPrivate.h"
 #include "IHeadMountedDisplayModule.h"
 #include "IHeadMountedDisplayVulkanExtensions.h"
-#include "VulkanRHIBridge.h"
-namespace VulkanRHIBridge
-{
-	extern TArray<const ANSICHAR*> InstanceExtensions;
-	extern TArray<const ANSICHAR*> InstanceLayers;
-	extern TArray<const ANSICHAR*> DeviceExtensions;
-	extern TArray<const ANSICHAR*> DeviceLayers;
-}
+
 #if VULKAN_HAS_DEBUGGING_ENABLED
 bool GRenderDocFound = false;
 #endif
@@ -63,6 +56,28 @@ TAutoConsoleVariable<int32> GGPUValidationCvar(
 
 #endif // VULKAN_HAS_DEBUGGING_ENABLED
 
+namespace VulkanExternalExtensions
+{
+	TArray<const ANSICHAR*> InstanceExtensions;
+	TArray<const ANSICHAR*> InstanceLayers;
+	TArray<const ANSICHAR*> DeviceExtensions;
+	TArray<const ANSICHAR*> DeviceLayers;
+}
+
+void IVulkanDynamicRHI::AddEnabledInstanceExtensionsAndLayers(TArrayView<const ANSICHAR* const> InInstanceExtensions, TArrayView<const ANSICHAR* const> InInstanceLayers)
+{
+	checkf(!GVulkanRHI, TEXT("AddEnabledInstanceExtensionsAndLayers should be called before the VulkanRHI has been created"));
+	VulkanExternalExtensions::InstanceExtensions.Append(InInstanceExtensions.GetData(), InInstanceExtensions.Num());
+	VulkanExternalExtensions::InstanceLayers.Append(InInstanceLayers.GetData(), InInstanceLayers.Num());
+}
+
+void IVulkanDynamicRHI::AddEnabledDeviceExtensionsAndLayers(TArrayView<const ANSICHAR* const> InDeviceExtensions, TArrayView<const ANSICHAR* const> InDeviceLayers)
+{
+	checkf(!GVulkanRHI, TEXT("AddEnabledDeviceExtensionsAndLayers should be called before the VulkanRHI has been created"));
+	VulkanExternalExtensions::DeviceExtensions.Append(InDeviceExtensions.GetData(), InDeviceExtensions.Num());
+	VulkanExternalExtensions::DeviceLayers.Append(InDeviceLayers.GetData(), InDeviceLayers.Num());
+}
+
 // Instance Extensions to enable for all platforms
 static const ANSICHAR* GInstanceExtensions[] =
 {
@@ -439,7 +454,7 @@ void FVulkanDynamicRHI::GetInstanceLayersAndExtensions(TArray<const ANSICHAR*>&
 	}
 
 	// Check for layers added outside the RHI (eg plugins)
-	for (const ANSICHAR* VulkanBridgeLayer : VulkanRHIBridge::InstanceLayers)
+	for (const ANSICHAR* VulkanBridgeLayer : VulkanExternalExtensions::InstanceLayers)
 	{
 		if (FindLayerInList(GlobalLayerExtensions, VulkanBridgeLayer))
 		{
@@ -447,7 +462,7 @@ void FVulkanDynamicRHI::GetInstanceLayersAndExtensions(TArray<const ANSICHAR*>&
 		}
 		else
 		{
-			UE_LOG(LogVulkanRHI, Warning, TEXT("Unable to find VulkanRHIBridge instance layer '%s'"), ANSI_TO_TCHAR(VulkanBridgeLayer));
+			UE_LOG(LogVulkanRHI, Warning, TEXT("Unable to find VulkanExternalExtensions instance layer '%s'"), ANSI_TO_TCHAR(VulkanBridgeLayer));
 		}
 	}
 
@@ -478,7 +493,7 @@ void FVulkanDynamicRHI::GetInstanceLayersAndExtensions(TArray<const ANSICHAR*>&
 	}
 
 	// Check for extensions added outside the RHI (eg plugins)
-	for (const ANSICHAR* VulkanBridgeExtension : VulkanRHIBridge::InstanceExtensions)
+	for (const ANSICHAR* VulkanBridgeExtension : VulkanExternalExtensions::InstanceExtensions)
 	{
 		if (FindLayerExtensionInList(GlobalLayerExtensions, VulkanBridgeExtension))
 		{
@@ -486,7 +501,7 @@ void FVulkanDynamicRHI::GetInstanceLayersAndExtensions(TArray<const ANSICHAR*>&
 		}
 		else
 		{
-			UE_LOG(LogVulkanRHI, Warning, TEXT("Unable to find VulkanRHIBridge instance extension '%s'"), ANSI_TO_TCHAR(VulkanBridgeExtension));
+			UE_LOG(LogVulkanRHI, Warning, TEXT("Unable to find VulkanExternalExtensions instance extension '%s'"), ANSI_TO_TCHAR(VulkanBridgeExtension));
 		}
 	}
 
@@ -632,7 +647,7 @@ void FVulkanDevice::GetDeviceExtensionsAndLayers(VkPhysicalDevice Gpu, EGpuVendo
 #endif	// VULKAN_HAS_DEBUGGING_ENABLED
 
 	// Check for layers added outside the RHI (eg plugins)
-	for (const ANSICHAR* VulkanBridgeLayer : VulkanRHIBridge::DeviceLayers)
+	for (const ANSICHAR* VulkanBridgeLayer : VulkanExternalExtensions::DeviceLayers)
 	{
 		if (FindLayerInList(DeviceLayerExtensions, VulkanBridgeLayer))
 		{
@@ -640,7 +655,7 @@ void FVulkanDevice::GetDeviceExtensionsAndLayers(VkPhysicalDevice Gpu, EGpuVendo
 		}
 		else
 		{
-			UE_LOG(LogVulkanRHI, Warning, TEXT("Unable to find VulkanRHIBridge device layer '%s'"), ANSI_TO_TCHAR(VulkanBridgeLayer));
+			UE_LOG(LogVulkanRHI, Warning, TEXT("Unable to find VulkanExternalExtensions device layer '%s'"), ANSI_TO_TCHAR(VulkanBridgeLayer));
 		}
 	}
 
@@ -722,7 +737,7 @@ void FVulkanDevice::GetDeviceExtensionsAndLayers(VkPhysicalDevice Gpu, EGpuVendo
 	}
 	
 	// Check for extensions added outside the RHI (eg plugins)
-	for (const ANSICHAR* VulkanBridgeExtension : VulkanRHIBridge::DeviceExtensions)
+	for (const ANSICHAR* VulkanBridgeExtension : VulkanExternalExtensions::DeviceExtensions)
 	{
 		if (ListContains(AvailableExtensions, VulkanBridgeExtension))
 		{
@@ -730,7 +745,7 @@ void FVulkanDevice::GetDeviceExtensionsAndLayers(VkPhysicalDevice Gpu, EGpuVendo
 		}
 		else
 		{
-			UE_LOG(LogVulkanRHI, Warning, TEXT("Unable to find VulkanRHIBridge device extension '%s'"), ANSI_TO_TCHAR(VulkanBridgeExtension));
+			UE_LOG(LogVulkanRHI, Warning, TEXT("Unable to find VulkanExternalExtensions device extension '%s'"), ANSI_TO_TCHAR(VulkanBridgeExtension));
 		}
 	}
 
diff --git a/Engine/Source/Runtime/VulkanRHI/Private/VulkanRHI.cpp b/Engine/Source/Runtime/VulkanRHI/Private/VulkanRHI.cpp
--- a/Engine/Source/Runtime/VulkanRHI/Private/VulkanRHI.cpp
+++ b/Engine/Source/Runtime/VulkanRHI/Private/VulkanRHI.cpp
@@ -1182,6 +1182,231 @@ void FVulkanDynamicRHI::RHIReleaseThreadOwnership()
 {
 }
 
+// IVulkanDynamicRHI interface
+
+uint32 FVulkanDynamicRHI::RHIGetVulkanVersion() const
+{
+	return UE_VK_API_VERSION;
+}
+
+VkInstance FVulkanDynamicRHI::RHIGetVkInstance() const
+{
+	return GetInstance();
+}
+
+VkDevice FVulkanDynamicRHI::RHIGetVkDevice() const
+{
+	return Device->GetInstanceHandle();
+}
+
+const uint8* FVulkanDynamicRHI::RHIGetVulkanDeviceUUID() const
+{
+	return GetDevice()->GetDeviceIdProperties().deviceUUID;
+}
+
+VkPhysicalDevice FVulkanDynamicRHI::RHIGetVkPhysicalDevice() const
+{
+	return Device->GetPhysicalHandle();
+}
+
+VkQueue FVulkanDynamicRHI::RHIGetGraphicsVkQueue() const
+{
+	return GetDevice()->GetGraphicsQueue()->GetHandle();
+}
+
+uint32 FVulkanDynamicRHI::RHIGetGraphicsQueueIndex() const
+{
+	return GetDevice()->GetGraphicsQueue()->GetQueueIndex();
+}
+
+uint32 FVulkanDynamicRHI::RHIGetGraphicsQueueFamilyIndex() const
+{
+	return GetDevice()->GetGraphicsQueue()->GetFamilyIndex();
+}
+
+VkCommandBuffer FVulkanDynamicRHI::RHIGetActiveVkCommandBuffer()
+{
+	FVulkanCommandListContextImmediate& ImmediateContext = GetDevice()->GetImmediateContext();
+	VkCommandBuffer VulkanCommandBuffer = ImmediateContext.GetCommandBufferManager()->GetActiveCmdBuffer()->GetHandle();
+	return VulkanCommandBuffer;
+}
+
+uint64 FVulkanDynamicRHI::RHIGetGraphicsAdapterLUID(VkPhysicalDevice InPhysicalDevice) const
+{
+	uint64 AdapterLUID = 0;
+#if VULKAN_SUPPORTS_DRIVER_PROPERTIES
+	if (!AdapterLUID && GetOptionalExtensions().HasKHRGetPhysicalDeviceProperties2)
+	{
+		VkPhysicalDeviceIDPropertiesKHR GpuIdProps;
+		VkPhysicalDeviceProperties2KHR GpuProps2;
+		ZeroVulkanStruct(GpuProps2, VK_STRUCTURE_TYPE_PHYSICAL_DEVICE_PROPERTIES_2_KHR);
+		GpuProps2.pNext = &GpuIdProps;
+		ZeroVulkanStruct(GpuIdProps, VK_STRUCTURE_TYPE_PHYSICAL_DEVICE_ID_PROPERTIES_KHR);
+
+		VulkanRHI::vkGetPhysicalDeviceProperties2KHR(InPhysicalDevice, &GpuProps2);
+		check(GpuIdProps.deviceLUIDValid);
+		AdapterLUID = reinterpret_cast<const uint64&>(GpuIdProps.deviceLUID);
+	}
+#endif
+	return AdapterLUID;
+}
+
+bool FVulkanDynamicRHI::RHIDoesAdapterMatchDevice(const void* InAdapterId) const
+{
+#if PLATFORM_WINDOWS
+	if (GetOptionalExtensions().HasKHRGetPhysicalDeviceProperties2)
+	{
+		const VkPhysicalDeviceIDPropertiesKHR& vkPhysicalDeviceIDProperties = GetDevice()->GetDeviceIdProperties();
+		if (vkPhysicalDeviceIDProperties.deviceLUIDValid)
+		{
+			return FMemory::Memcmp(InAdapterId, &vkPhysicalDeviceIDProperties.deviceLUID, sizeof(LUID)) == 0;
+		}
+	}
+#endif
+
+	// Not enough information.  Assume the adapter matches
+	return true;
+}
+
+void* FVulkanDynamicRHI::RHIGetVkDeviceProcAddr(const char* InName) const
+{
+	return (void*)VulkanRHI::vkGetDeviceProcAddr(Device->GetInstanceHandle(), InName);
+}
+
+VkFormat FVulkanDynamicRHI::RHIGetSwapChainVkFormat(EPixelFormat InFormat) const
+{
+	// UE renders a gamma-corrected image so we need to use an sRGB format if available
+	return UEToVkTextureFormat(GPixelFormats[InFormat].UnrealFormat, true);
+}
+
+bool FVulkanDynamicRHI::RHISupportsEXTFragmentDensityMap2() const
+{
+	return GetDevice()->GetOptionalExtensions().HasEXTFragmentDensityMap2;
+}
+
+TArray<VkExtensionProperties> FVulkanDynamicRHI::RHIGetAllInstanceExtensions() const
+{
+	uint32_t ExtensionCount = 0;
+	VulkanRHI::vkEnumerateInstanceExtensionProperties(nullptr, &ExtensionCount, nullptr);
+
+	TArray<VkExtensionProperties> Extensions;
+	Extensions.SetNumUninitialized(ExtensionCount);
+
+	VulkanRHI::vkEnumerateInstanceExtensionProperties(nullptr, &ExtensionCount, Extensions.GetData());
+
+	return Extensions;
+}
+
+TArray<VkExtensionProperties> FVulkanDynamicRHI::RHIGetAllDeviceExtensions(VkPhysicalDevice InPhysicalDevice) const
+{
+	uint32_t ExtensionCount = 0;
+	VulkanRHI::vkEnumerateDeviceExtensionProperties(InPhysicalDevice, nullptr, &ExtensionCount, nullptr);
+
+	TArray<VkExtensionProperties> Extensions;
+	Extensions.SetNumUninitialized(ExtensionCount);
+
+	VulkanRHI::vkEnumerateDeviceExtensionProperties(InPhysicalDevice, nullptr, &ExtensionCount, Extensions.GetData());
+
+	return Extensions;
+}
+
+VkImage FVulkanDynamicRHI::RHIGetVkImage(FRHITexture2D* InTexture) const
+{
+	FVulkanTexture2D* VulkanTexture = ResourceCast(InTexture);
+	return VulkanTexture->Surface.Image;
+}
+
+VkFormat FVulkanDynamicRHI::RHIGetViewVkFormat(FRHITexture2D* InTexture) const
+{
+	FVulkanTexture2D* VulkanTexture = ResourceCast(InTexture);
+	return VulkanTexture->Surface.ViewFormat;
+}
+
+FVulkanRHIAllocationInfo FVulkanDynamicRHI::RHIGetAllocationInfo(FRHITexture2D* InTexture) const
+{
+	FVulkanTexture2D* VulkanTexture = ResourceCast(InTexture);
+
+	FVulkanRHIAllocationInfo NewInfo{};
+	NewInfo.Handle = VulkanTexture->Surface.GetAllocationHandle();
+	NewInfo.Offset = VulkanTexture->Surface.GetAllocationOffset();
+	NewInfo.Size = VulkanTexture->Surface.GetMemorySize();
+
+	return NewInfo;
+}
+
+FVulkanRHIImageViewInfo FVulkanDynamicRHI::RHIGetImageViewInfo(FRHITexture* InTexture) const
+{
+	FVulkanTextureBase* VulkanTexture = (FVulkanTextureBase*)InTexture->GetTextureBaseRHI();
+
+	FVulkanRHIImageViewInfo Info{};
+	Info.ImageView = VulkanTexture->DefaultView.View;
+	Info.Image = VulkanTexture->DefaultView.Image;
+	Info.Format = VulkanTexture->Surface.ViewFormat;
+	Info.Width = VulkanTexture->Surface.Width;
+	Info.Height = VulkanTexture->Surface.Height;
+	Info.Depth = VulkanTexture->Surface.Depth;
+	Info.UEFlags = VulkanTexture->Surface.UEFlags;
+
+	Info.SubresourceRange.aspectMask = VulkanTexture->Surface.GetFullAspectMask();
+	Info.SubresourceRange.layerCount = VulkanTexture->Surface.GetNumberOfArrayLevels();
+	Info.SubresourceRange.levelCount = VulkanTexture->Surface.GetNumMips();
+
+	// TODO: do we need these?
+	Info.SubresourceRange.baseMipLevel = 0;
+	Info.SubresourceRange.baseArrayLayer = 0;
+
+	return Info;
+}
+
+VkImageLayout& FVulkanDynamicRHI::RHIFindOrAddLayoutRW(FRHITexture2D* InTexture, VkImageLayout LayoutIfNotFound)
+{
+	FVulkanTexture2D* VulkanTexture = ResourceCast(InTexture);
+	return GetDevice()->GetImmediateContext().GetLayoutManager().FindOrAddLayoutRW(VulkanTexture->Surface, LayoutIfNotFound);
+}
+
+void FVulkanDynamicRHI::RHISetImageLayout(VkImage Image, VkImageLayout OldLayout, VkImageLayout NewLayout, const VkImageSubresourceRange& SubresourceRange)
+{
+	FVulkanCommandListContext& ImmediateContext = GetDevice()->GetImmediateContext();
+	FVulkanCmdBuffer* CmdBuffer = ImmediateContext.GetCommandBufferManager()->GetActiveCmdBuffer();
+	VulkanSetImageLayout(CmdBuffer->GetHandle(), Image, OldLayout, NewLayout, SubresourceRange);
+}
+
+void FVulkanDynamicRHI::RHISetUploadImageLayout(VkImage Image, VkImageLayout OldLayout, VkImageLayout NewLayout, const VkImageSubresourceRange& SubresourceRange)
+{
+	FVulkanCommandListContext& ImmediateContext = GetDevice()->GetImmediateContext();
+	FVulkanCmdBuffer* CmdBuffer = ImmediateContext.GetCommandBufferManager()->GetUploadCmdBuffer();
+	VulkanSetImageLayout(CmdBuffer->GetHandle(), Image, OldLayout, NewLayout, SubresourceRange);
+}
+
+void FVulkanDynamicRHI::RHIFinishExternalComputeWork(VkCommandBuffer InCommandBuffer)
+{
+	FVulkanCommandListContextImmediate& ImmediateContext = GetDevice()->GetImmediateContext();
+	check(InCommandBuffer == ImmediateContext.GetCommandBufferManager()->GetActiveCmdBuffer()->GetHandle());
+
+	ImmediateContext.GetPendingComputeState()->Reset();
+	ImmediateContext.GetPendingGfxState()->Reset();
+}
+
+void FVulkanDynamicRHI::RHIRegisterWork(uint32 NumPrimitives)
+{
+	FVulkanCommandListContextImmediate& ImmediateContext = GetDevice()->GetImmediateContext();
+	if (FVulkanPlatform::RegisterGPUWork() && ImmediateContext.IsImmediate())
+	{
+		ImmediateContext.GetGPUProfiler().RegisterGPUWork(1);
+	}
+}
+
+void FVulkanDynamicRHI::RHISubmitUploadCommandBuffer()
+{
+	FVulkanCommandListContext& ImmediateContext = GetDevice()->GetImmediateContext();
+	ImmediateContext.GetCommandBufferManager()->SubmitUploadCmdBuffer();
+}
+
+void FVulkanDynamicRHI::RHIVerifyResult(VkResult Result, const ANSICHAR* VkFuntion, const ANSICHAR* Filename, uint32 Line)
+{
+	VulkanRHI::VerifyVulkanResult(Result, VkFuntion, Filename, Line);
+}
+
 void* FVulkanDynamicRHI::RHIGetNativeDevice()
 {
 	return (void*)Device->GetInstanceHandle();
diff --git a/Engine/Source/Runtime/VulkanRHI/Private/VulkanUtil.cpp b/Engine/Source/Runtime/VulkanRHI/Private/VulkanUtil.cpp
--- a/Engine/Source/Runtime/VulkanRHI/Private/VulkanUtil.cpp
+++ b/Engine/Source/Runtime/VulkanRHI/Private/VulkanUtil.cpp
@@ -816,50 +816,12 @@ void AftermathCrashDumpDescriptionCallback(PFN_GFSDK_Aftermath_AddGpuCrashDumpDe
 #include "VulkanRHIBridge.h"
 namespace VulkanRHIBridge
 {
-	TArray<const ANSICHAR*> InstanceExtensions;
-	TArray<const ANSICHAR*> InstanceLayers;
-	TArray<const ANSICHAR*> DeviceExtensions;
-	TArray<const ANSICHAR*> DeviceLayers;
-
-
-	uint64 GetInstance(FVulkanDynamicRHI* RHI)
-	{
-		return (uint64)RHI->GetInstance();
-	}
-
 	FVulkanDevice* GetDevice(FVulkanDynamicRHI* RHI)
 	{
 		return RHI->GetDevice();
 	}
-
-	// Returns a VkDevice
-	uint64 GetLogicalDevice(FVulkanDevice* Device)
-	{
-		return (uint64)Device->GetInstanceHandle();
-	}
-
-	// Returns a VkDeviceVkPhysicalDevice
-	uint64 GetPhysicalDevice(FVulkanDevice* Device)
-	{
-		return (uint64)Device->GetPhysicalHandle();
-	}
-
-	void AddEnabledInstanceExtensionsAndLayers(const TArray<const ANSICHAR*>& InInstanceExtensions, const TArray<const ANSICHAR*>& InInstanceLayers)
-	{
-		checkf(!GVulkanRHI, TEXT("AddEnabledInstanceExtensionsAndLayers should be called before the VulkanRHI has been created"));
-		InstanceExtensions.Append(InInstanceExtensions);
-		InstanceLayers.Append(InInstanceLayers);
-	}
-
-	void AddEnabledDeviceExtensionsAndLayers(const TArray<const ANSICHAR*>& InDeviceExtensions, const TArray<const ANSICHAR*>& InDeviceLayers)
-	{
-		checkf(!GVulkanRHI, TEXT("AddEnabledDeviceExtensionsAndLayers should be called before the VulkanRHI has been created"));
-		DeviceExtensions.Append(InDeviceExtensions);
-		DeviceLayers.Append(InDeviceLayers);
-	}
 }
 
-
 namespace VulkanRHI
 {
 	VkBuffer CreateBuffer(FVulkanDevice* InDevice, VkDeviceSize Size, VkBufferUsageFlags BufferUsageFlags, VkMemoryRequirements& OutMemoryRequirements)
diff --git a/Engine/Source/Runtime/VulkanRHI/Public/IVulkanDynamicRHI.h b/Engine/Source/Runtime/VulkanRHI/Public/IVulkanDynamicRHI.h
new file mode 100644
--- /dev/null
+++ b/Engine/Source/Runtime/VulkanRHI/Public/IVulkanDynamicRHI.h
@@ -0,0 +1,91 @@
+// Copyright Epic Games, Inc. All Rights Reserved.
+
+#pragma once
+
+#include "RHI.h"
+
+#if defined(VK_USE_PLATFORM_WIN32_KHR)
+	#include "Microsoft/AllowMicrosoftPlatformTypes.h"
+	#include <windows.h>
+	#include "Microsoft/HideMicrosoftPlatformTypes.h"
+#endif
+
+THIRD_PARTY_INCLUDES_START
+	#define VK_NO_PROTOTYPES
+	#include <vulkan.h>
+THIRD_PARTY_INCLUDES_END
+
+struct FVulkanRHIAllocationInfo
+{
+	VkDeviceMemory Handle{};
+	uint64 Offset{};
+	uint64 Size{};
+};
+
+struct FVulkanRHIImageViewInfo
+{
+	VkImageView ImageView;
+	VkImage Image;
+	VkImageSubresourceRange SubresourceRange;
+	VkFormat Format;
+	uint32 Width;
+	uint32 Height;
+	uint32 Depth;
+	ETextureCreateFlags UEFlags;
+};
+
+struct IVulkanDynamicRHI : public FDynamicRHI
+{
+	virtual ERHIInterfaceType GetInterfaceType() const override { return ERHIInterfaceType::Vulkan; }
+
+	virtual uint32           RHIGetVulkanVersion() const = 0;
+	virtual VkInstance       RHIGetVkInstance() const = 0;
+	virtual VkDevice         RHIGetVkDevice() const = 0;
+	virtual const uint8*     RHIGetVulkanDeviceUUID() const = 0;
+	virtual VkPhysicalDevice RHIGetVkPhysicalDevice() const = 0;
+
+	virtual VkQueue          RHIGetGraphicsVkQueue() const = 0;
+	virtual uint32           RHIGetGraphicsQueueIndex() const = 0;
+	virtual uint32           RHIGetGraphicsQueueFamilyIndex() const = 0;
+
+	virtual VkCommandBuffer  RHIGetActiveVkCommandBuffer() = 0;
+
+	virtual uint64           RHIGetGraphicsAdapterLUID(VkPhysicalDevice InPhysicalDevice) const = 0;
+	virtual bool             RHIDoesAdapterMatchDevice(const void* InAdapterId) const = 0;
+	virtual void*            RHIGetVkDeviceProcAddr(const char* InName) const = 0;
+	virtual VkFormat         RHIGetSwapChainVkFormat(EPixelFormat InFormat) const = 0;
+	virtual bool             RHISupportsEXTFragmentDensityMap2() const = 0;
+
+	virtual TArray<VkExtensionProperties> RHIGetAllInstanceExtensions() const = 0;
+	virtual TArray<VkExtensionProperties> RHIGetAllDeviceExtensions(VkPhysicalDevice InPhysicalDevice) const = 0;
+
+	virtual FTexture2DRHIRef      RHICreateTexture2DFromResource(EPixelFormat Format, uint32 SizeX, uint32 SizeY, uint32 NumMips, uint32 NumSamples, VkImage Resource, ETextureCreateFlags Flags) = 0;
+	virtual FTexture2DArrayRHIRef RHICreateTexture2DArrayFromResource(EPixelFormat Format, uint32 SizeX, uint32 SizeY, uint32 ArraySize, uint32 NumMips, uint32 NumSamples, VkImage Resource, ETextureCreateFlags Flags) = 0;
+	virtual FTextureCubeRHIRef    RHICreateTextureCubeFromResource(EPixelFormat Format, uint32 Size, bool bArray, uint32 ArraySize, uint32 NumMips, VkImage Resource, ETextureCreateFlags Flags) = 0;
+
+	virtual VkImage                  RHIGetVkImage(FRHITexture2D* InTexture) const = 0;
+	virtual VkFormat                 RHIGetViewVkFormat(FRHITexture2D* InTexture) const = 0;
+	virtual FVulkanRHIAllocationInfo RHIGetAllocationInfo(FRHITexture2D* InTexture) const = 0;
+	virtual FVulkanRHIImageViewInfo  RHIGetImageViewInfo(FRHITexture* InTexture) const = 0;
+
+	virtual VkImageLayout& RHIFindOrAddLayoutRW(FRHITexture2D* InTexture, VkImageLayout LayoutIfNotFound) = 0;
+	virtual void           RHISetImageLayout(VkImage Image, VkImageLayout OldLayout, VkImageLayout NewLayout, const VkImageSubresourceRange& SubresourceRange) = 0;
+	virtual void           RHISetUploadImageLayout(VkImage Image, VkImageLayout OldLayout, VkImageLayout NewLayout, const VkImageSubresourceRange& SubresourceRange) = 0;
+
+	virtual void           RHIFinishExternalComputeWork(VkCommandBuffer InCommandBuffer) = 0;
+	virtual void           RHIRegisterWork(uint32 NumPrimitives) = 0;
+	virtual void           RHISubmitUploadCommandBuffer() = 0;
+
+	virtual void           RHIVerifyResult(VkResult Result, const ANSICHAR* VkFuntion, const ANSICHAR* Filename, uint32 Line) = 0;
+
+	static VULKANRHI_API void AddEnabledInstanceExtensionsAndLayers(TArrayView<const ANSICHAR* const> InInstanceExtensions, TArrayView<const ANSICHAR* const> InInstanceLayers);
+	static VULKANRHI_API void AddEnabledDeviceExtensionsAndLayers(TArrayView<const ANSICHAR* const> InDeviceExtensions, TArrayView<const ANSICHAR* const> InDeviceLayers);
+};
+
+inline IVulkanDynamicRHI* GetIVulkanDynamicRHI()
+{
+	check(GDynamicRHI->GetInterfaceType() == ERHIInterfaceType::Vulkan);
+	return GetDynamicRHI<IVulkanDynamicRHI>();
+}
+
+#define VERIFYVULKANRESULT_EXTERNAL(VkFunction) { const VkResult ScopedResult = VkFunction; if (ScopedResult != VK_SUCCESS) { GetIVulkanDynamicRHI()->RHIVerifyResult(ScopedResult, #VkFunction, __FILE__, __LINE__); }}
diff --git a/Engine/Source/Runtime/VulkanRHI/Public/VulkanDynamicRHI.h b/Engine/Source/Runtime/VulkanRHI/Public/VulkanDynamicRHI.h
--- a/Engine/Source/Runtime/VulkanRHI/Public/VulkanDynamicRHI.h
+++ b/Engine/Source/Runtime/VulkanRHI/Public/VulkanDynamicRHI.h
@@ -4,7 +4,9 @@
 	VulkanRHI.h: Public Vulkan RHI definitions.
 =============================================================================*/
 
-#pragma once 
+#pragma once
+
+#include "IVulkanDynamicRHI.h"
 
 class FVulkanTexture2D;
 class FVulkanFramebuffer;
@@ -40,7 +42,7 @@ struct FOptionalVulkanInstanceExtensions
 PRAGMA_DISABLE_DEPRECATION_WARNINGS
 
 /** The interface which is implemented by the dynamically bound RHI. */
-class FVulkanDynamicRHI : public FDynamicRHI
+class FVulkanDynamicRHI : public IVulkanDynamicRHI
 {
 public:
 
@@ -50,13 +52,43 @@ public:
 	/** Destructor */
 	~FVulkanDynamicRHI() {}
 
+	// IVulkanDynamicRHI interface
+	virtual uint32 RHIGetVulkanVersion() const final override;
+	virtual VkInstance RHIGetVkInstance() const final override;
+	virtual VkDevice RHIGetVkDevice() const final override;
+	virtual const uint8* RHIGetVulkanDeviceUUID() const final override;
+	virtual VkPhysicalDevice RHIGetVkPhysicalDevice() const final override;
+	virtual VkQueue RHIGetGraphicsVkQueue() const final override;
+	virtual uint32 RHIGetGraphicsQueueIndex() const final override;
+	virtual uint32 RHIGetGraphicsQueueFamilyIndex() const final override;
+	virtual VkCommandBuffer RHIGetActiveVkCommandBuffer() final override;
+	virtual uint64 RHIGetGraphicsAdapterLUID(VkPhysicalDevice InPhysicalDevice) const final override;
+	virtual bool RHIDoesAdapterMatchDevice(const void* InAdapterId) const final override;
+	virtual void* RHIGetVkDeviceProcAddr(const char* InName) const final override;
+	virtual VkFormat RHIGetSwapChainVkFormat(EPixelFormat InFormat) const final override;
+	virtual bool RHISupportsEXTFragmentDensityMap2() const final override;
+	virtual TArray<VkExtensionProperties> RHIGetAllInstanceExtensions() const final override;
+	virtual TArray<VkExtensionProperties> RHIGetAllDeviceExtensions(VkPhysicalDevice InPhysicalDevice) const final override;
+	virtual FTexture2DRHIRef RHICreateTexture2DFromResource(EPixelFormat Format, uint32 SizeX, uint32 SizeY, uint32 NumMips, uint32 NumSamples, VkImage Resource, ETextureCreateFlags Flags) final override;
+	virtual FTexture2DArrayRHIRef RHICreateTexture2DArrayFromResource(EPixelFormat Format, uint32 SizeX, uint32 SizeY, uint32 ArraySize, uint32 NumMips, uint32 NumSamples, VkImage Resource, ETextureCreateFlags Flags) final override;
+	virtual FTextureCubeRHIRef RHICreateTextureCubeFromResource(EPixelFormat Format, uint32 Size, bool bArray, uint32 ArraySize, uint32 NumMips, VkImage Resource, ETextureCreateFlags Flags) final override;
+	virtual VkImage RHIGetVkImage(FRHITexture2D* InTexture) const final override;
+	virtual VkFormat RHIGetViewVkFormat(FRHITexture2D* InTexture) const final override;
+	virtual FVulkanRHIAllocationInfo RHIGetAllocationInfo(FRHITexture2D* InTexture) const final override;
+	virtual FVulkanRHIImageViewInfo RHIGetImageViewInfo(FRHITexture* InTexture) const final override;
+	virtual VkImageLayout& RHIFindOrAddLayoutRW(FRHITexture2D* InTexture, VkImageLayout LayoutIfNotFound) final override;
+	virtual void RHISetImageLayout(VkImage Image, VkImageLayout OldLayout, VkImageLayout NewLayout, const VkImageSubresourceRange& SubresourceRange) final override;
+	virtual void RHISetUploadImageLayout(VkImage Image, VkImageLayout OldLayout, VkImageLayout NewLayout, const VkImageSubresourceRange& SubresourceRange) final override;
+	virtual void RHIFinishExternalComputeWork(VkCommandBuffer InCommandBuffer) final override;
+	virtual void RHIRegisterWork(uint32 NumPrimitives) final override;
+	virtual void RHISubmitUploadCommandBuffer() final override;
+	virtual void RHIVerifyResult(VkResult Result, const ANSICHAR* VkFuntion, const ANSICHAR* Filename, uint32 Line) final override;
 
 	// FDynamicRHI interface.
 	virtual void Init() final override;
 	virtual void PostInit() final override;
 	virtual void Shutdown() final override;;
 	virtual const TCHAR* GetName() final override { return TEXT("Vulkan"); }
-	virtual ERHIInterfaceType GetInterfaceType() const override { return ERHIInterfaceType::Vulkan; }
 
 	void InitInstance();
 
@@ -308,11 +340,6 @@ public:
 		return this->RHICreateRenderQuery(QueryType);
 	}
 
-	// FVulkanDynamicRHI interface
-	virtual FTexture2DRHIRef RHICreateTexture2DFromResource(EPixelFormat Format, uint32 SizeX, uint32 SizeY, uint32 NumMips, uint32 NumSamples, VkImage Resource, ETextureCreateFlags Flags);
-	virtual FTexture2DArrayRHIRef RHICreateTexture2DArrayFromResource(EPixelFormat Format, uint32 SizeX, uint32 SizeY, uint32 ArraySize, uint32 NumMips, uint32 NumSamples, VkImage Resource, ETextureCreateFlags Flags);
-	virtual FTextureCubeRHIRef RHICreateTextureCubeFromResource(EPixelFormat Format, uint32 Size, bool bArray, uint32 ArraySize, uint32 NumMips, VkImage Resource, ETextureCreateFlags Flags);
-
 	void RHICalibrateTimers() override;
 
 #if VULKAN_RHI_RAYTRACING
@@ -345,7 +372,7 @@ public:
 		return Instance;
 	}
 	
-	inline FVulkanDevice* GetDevice()
+	inline FVulkanDevice* GetDevice() const
 	{
 		return Device;
 	}
@@ -360,8 +387,8 @@ public:
 		return OptionalInstanceExtensions;
 	}
 
-	virtual void VulkanSetImageLayout( VkCommandBuffer CmdBuffer, VkImage Image, VkImageLayout OldLayout, VkImageLayout NewLayout, const VkImageSubresourceRange& SubresourceRange );
-	
+	void VulkanSetImageLayout( VkCommandBuffer CmdBuffer, VkImage Image, VkImageLayout OldLayout, VkImageLayout NewLayout, const VkImageSubresourceRange& SubresourceRange );
+
 	virtual void* RHILockStagingBuffer(FRHIStagingBuffer* StagingBuffer, FRHIGPUFence* Fence, uint32 Offset, uint32 SizeRHI) final override;
 	virtual void RHIUnlockStagingBuffer(FRHIStagingBuffer* StagingBuffer) final override;
 
diff --git a/Engine/Source/Runtime/VulkanRHI/Public/VulkanRHIBridge.h b/Engine/Source/Runtime/VulkanRHI/Public/VulkanRHIBridge.h
--- a/Engine/Source/Runtime/VulkanRHI/Public/VulkanRHIBridge.h
+++ b/Engine/Source/Runtime/VulkanRHI/Public/VulkanRHIBridge.h
@@ -4,25 +4,47 @@
 	VulkanRHIBridge.h: Utils to interact with the inner RHI.
 =============================================================================*/
 
-#pragma once 
+#pragma once
+
+#include "IVulkanDynamicRHI.h"
 
 class FVulkanDynamicRHI;
 class FVulkanDevice;
 
 namespace VulkanRHIBridge
 {
-
-	VULKANRHI_API void AddEnabledInstanceExtensionsAndLayers(const TArray<const ANSICHAR*>& InInstanceExtensions, const TArray<const ANSICHAR*>& InInstanceLayers);
-	VULKANRHI_API void AddEnabledDeviceExtensionsAndLayers(const TArray<const ANSICHAR*>& InDeviceExtensions, const TArray<const ANSICHAR*>& InDeviceLayers);
-
-
-	VULKANRHI_API uint64 GetInstance(FVulkanDynamicRHI*);
-
+	UE_DEPRECATED(5.1, "Please use IVulkanDynamicRHI::AddEnabledInstanceExtensionsAndLayers()")
+	inline void AddEnabledInstanceExtensionsAndLayers(const TArray<const ANSICHAR*>& InInstanceExtensions, const TArray<const ANSICHAR*>& InInstanceLayers)
+	{
+		IVulkanDynamicRHI::AddEnabledInstanceExtensionsAndLayers(InInstanceExtensions, InInstanceLayers);
+	}
+	
+	UE_DEPRECATED(5.1, "Please use IVulkanDynamicRHI::AddEnabledDeviceExtensionsAndLayers()")
+	inline void AddEnabledDeviceExtensionsAndLayers(const TArray<const ANSICHAR*>& InDeviceExtensions, const TArray<const ANSICHAR*>& InDeviceLayers)
+	{
+		IVulkanDynamicRHI::AddEnabledDeviceExtensionsAndLayers(InDeviceExtensions, InDeviceLayers);
+	}
+
+	UE_DEPRECATED(5.1, "Please use IVulkanDynamicRHI::RHIGetVkInstance()")
+	inline uint64 GetInstance(FVulkanDynamicRHI*)
+	{
+		return (uint64)GetIVulkanDynamicRHI()->RHIGetVkInstance();
+	}
+
+	UE_DEPRECATED(5.1, "Use IVulkanDynamicRHI instead of FVulkanDevice directly.")
 	VULKANRHI_API FVulkanDevice* GetDevice(FVulkanDynamicRHI*);
 
 	// Returns a VkDevice
-	VULKANRHI_API uint64 GetLogicalDevice(FVulkanDevice*);
+	UE_DEPRECATED(5.1, "Please use IVulkanDynamicRHI::RHIGetVkDevice()")
+	inline uint64 GetLogicalDevice(FVulkanDevice*)
+	{
+		return (uint64)GetIVulkanDynamicRHI()->RHIGetVkDevice();
+	}
 
 	// Returns a VkDeviceVkPhysicalDevice
-	VULKANRHI_API uint64 GetPhysicalDevice(FVulkanDevice*);
+	UE_DEPRECATED(5.1, "Please use IVulkanDynamicRHI::RHIGetVkPhysicalDevice()")
+	inline uint64 GetPhysicalDevice(FVulkanDevice*)
+	{
+		return (uint64)GetIVulkanDynamicRHI()->RHIGetVkPhysicalDevice();
+	}
 }
diff --git a/Engine/Source/Runtime/VulkanRHI/VulkanRHI.Build.cs b/Engine/Source/Runtime/VulkanRHI/VulkanRHI.Build.cs
--- a/Engine/Source/Runtime/VulkanRHI/VulkanRHI.Build.cs
+++ b/Engine/Source/Runtime/VulkanRHI/VulkanRHI.Build.cs
@@ -30,6 +30,17 @@ public class VulkanRHI : ModuleRules
 			PrivateIncludePaths.Add("Runtime/VulkanRHI/Private/" + Target.Platform);
 		}
 
+		if (Target.Platform.IsInGroup(UnrealPlatformGroup.Windows))
+		{
+			PublicDefinitions.Add("VK_USE_PLATFORM_WIN32_KHR=1");
+			PublicDefinitions.Add("VK_USE_PLATFORM_WIN32_KHX=1");
+		}
+		else if (Target.Platform.IsInGroup(UnrealPlatformGroup.Android))
+		{
+			PublicDefinitions.Add("VK_USE_PLATFORM_ANDROID_KHR=1");
+		}
+
+
 		PrivateDependencyModuleNames.AddRange(
 			new string[]
 			{
@@ -95,12 +106,7 @@ public class VulkanRHI : ModuleRules
 			{
 				if (Target.Configuration != UnrealTargetConfiguration.Shipping)
 				{
-					PrivateIncludePathModuleNames.AddRange(
-						new string[]
-						{
-							"TaskGraph",
-						}
-					);
+					PrivateIncludePathModuleNames.Add("TaskGraph");
 				}
 			}
 			else
