Fixing Vulkan issues after the interface changes.

diff --git a/Engine/Source/Runtime/VulkanRHI/Private/VulkanRHI.cpp b/Engine/Source/Runtime/VulkanRHI/Private/VulkanRHI.cpp
--- a/Engine/Source/Runtime/VulkanRHI/Private/VulkanRHI.cpp
+++ b/Engine/Source/Runtime/VulkanRHI/Private/VulkanRHI.cpp
@@ -1209,6 +1209,11 @@ VkPhysicalDevice FVulkanDynamicRHI::RHIGetVkPhysicalDevice() const
 	return Device->GetPhysicalHandle();
 }
 
+const VkAllocationCallbacks* FVulkanDynamicRHI::RHIGetVkAllocationCallbacks()
+{
+	return VULKAN_CPU_ALLOCATOR;
+}
+
 VkQueue FVulkanDynamicRHI::RHIGetGraphicsVkQueue() const
 {
 	return GetDevice()->GetGraphicsQueue()->GetHandle();
diff --git a/Engine/Source/Runtime/VulkanRHI/Public/IVulkanDynamicRHI.h b/Engine/Source/Runtime/VulkanRHI/Public/IVulkanDynamicRHI.h
--- a/Engine/Source/Runtime/VulkanRHI/Public/IVulkanDynamicRHI.h
+++ b/Engine/Source/Runtime/VulkanRHI/Public/IVulkanDynamicRHI.h
@@ -43,6 +43,7 @@ struct IVulkanDynamicRHI : public FDynamicRHI
 	virtual VkDevice         RHIGetVkDevice() const = 0;
 	virtual const uint8*     RHIGetVulkanDeviceUUID() const = 0;
 	virtual VkPhysicalDevice RHIGetVkPhysicalDevice() const = 0;
+	virtual const VkAllocationCallbacks* RHIGetVkAllocationCallbacks() = 0;
 
 	virtual VkQueue          RHIGetGraphicsVkQueue() const = 0;
 	virtual uint32           RHIGetGraphicsQueueIndex() const = 0;
diff --git a/Engine/Source/Runtime/VulkanRHI/Public/VulkanDynamicRHI.h b/Engine/Source/Runtime/VulkanRHI/Public/VulkanDynamicRHI.h
--- a/Engine/Source/Runtime/VulkanRHI/Public/VulkanDynamicRHI.h
+++ b/Engine/Source/Runtime/VulkanRHI/Public/VulkanDynamicRHI.h
@@ -58,6 +58,7 @@ public:
 	virtual VkDevice RHIGetVkDevice() const final override;
 	virtual const uint8* RHIGetVulkanDeviceUUID() const final override;
 	virtual VkPhysicalDevice RHIGetVkPhysicalDevice() const final override;
+	virtual const VkAllocationCallbacks* RHIGetVkAllocationCallbacks() final override;
 	virtual VkQueue RHIGetGraphicsVkQueue() const final override;
 	virtual uint32 RHIGetGraphicsQueueIndex() const final override;
 	virtual uint32 RHIGetGraphicsQueueFamilyIndex() const final override;
