Vulkan RT: Update ray tracing BLAS code:
* Use provided ScratchBuffer for building and remove per-geometry ScratchBuffer
* Batch build geometries
* Add missing code for BLAS Update
* Update VertexBuffers when building

diff --git a/Engine/Source/Runtime/VulkanRHI/Private/VulkanRayTracing.cpp b/Engine/Source/Runtime/VulkanRHI/Private/VulkanRayTracing.cpp
--- a/Engine/Source/Runtime/VulkanRHI/Private/VulkanRayTracing.cpp
+++ b/Engine/Source/Runtime/VulkanRHI/Private/VulkanRayTracing.cpp
@@ -307,24 +307,19 @@ FVulkanRayTracingGeometry::FVulkanRayTracingGeometry(const FRayTracingGeometryIn
 
 	checkf(!Initializer.IndexBuffer || (IndexBufferStride == 2 || IndexBufferStride == 4), TEXT("Index buffer must be 16 or 32 bit if in use."));
 
-	FVkRtBLASBuildData BuildData;
-	GetBLASBuildData(
-		Device->GetInstanceHandle(), 
-		MakeArrayView(Initializer.Segments),
-		Initializer.IndexBuffer,
-		Initializer.IndexBufferOffset,
-		Initializer.bFastBuild,
-		Initializer.bAllowUpdate,
-		IndexBufferStride,
-		EAccelerationStructureBuildMode::Build,
-		BuildData);
+	SizeInfo = RHICalcRayTracingGeometrySize(Initializer);
+
+	// If this RayTracingGeometry going to be used as streaming destination 
+	// we don't want to allocate its memory as it will be replaced later by streamed version
+	// but we still need correct SizeInfo as it is used to estimate its memory requirements outside of RHI.
+	if (Initializer.Type == ERayTracingGeometryInitializerType::StreamingDestination)
+	{
+		return;
+	}
 
 	FString DebugNameString = Initializer.DebugName.ToString();
 	FRHIResourceCreateInfo BlasBufferCreateInfo(*DebugNameString);
-	AccelerationStructureBuffer = ResourceCast(RHICreateBuffer(BuildData.SizesInfo.accelerationStructureSize, BUF_AccelerationStructure, 0, ERHIAccess::BVHWrite, BlasBufferCreateInfo).GetReference());
-
-	FRHIResourceCreateInfo ScratchBufferCreateInfo(TEXT("BuildScratchBLAS"));
-	ScratchBuffer = ResourceCast(RHICreateBuffer(BuildData.SizesInfo.buildScratchSize, BUF_StructuredBuffer | BUF_RayTracingScratch, 0, ERHIAccess::UAVCompute, ScratchBufferCreateInfo).GetReference());
+	AccelerationStructureBuffer = ResourceCast(RHICreateBuffer(SizeInfo.ResultSize, BUF_AccelerationStructure, 0, ERHIAccess::BVHWrite, BlasBufferCreateInfo).GetReference());
 
 	VkDevice NativeDevice = Device->GetInstanceHandle();
 
@@ -332,16 +327,10 @@ FVulkanRayTracingGeometry::FVulkanRayTracingGeometry(const FRayTracingGeometryIn
 	ZeroVulkanStruct(CreateInfo, VK_STRUCTURE_TYPE_ACCELERATION_STRUCTURE_CREATE_INFO_KHR);
 	CreateInfo.buffer = AccelerationStructureBuffer->GetHandle();
 	CreateInfo.offset = AccelerationStructureBuffer->GetOffset();
-	CreateInfo.size = BuildData.SizesInfo.accelerationStructureSize;
+	CreateInfo.size = SizeInfo.ResultSize;
 	CreateInfo.type = VK_ACCELERATION_STRUCTURE_TYPE_BOTTOM_LEVEL_KHR;
 	VERIFYVULKANRESULT(vkCreateAccelerationStructureKHR(NativeDevice, &CreateInfo, VULKAN_CPU_ALLOCATOR, &Handle));
-
 	
-	SizeInfo.ResultSize = BuildData.SizesInfo.accelerationStructureSize;
-	SizeInfo.BuildScratchSize = BuildData.SizesInfo.buildScratchSize;
-	SizeInfo.UpdateScratchSize = BuildData.SizesInfo.updateScratchSize;
-
-
 	VkAccelerationStructureDeviceAddressInfoKHR DeviceAddressInfo;
 	ZeroVulkanStruct(DeviceAddressInfo, VK_STRUCTURE_TYPE_ACCELERATION_STRUCTURE_DEVICE_ADDRESS_INFO_KHR);
 	DeviceAddressInfo.accelerationStructure = Handle;
@@ -370,46 +359,10 @@ void FVulkanRayTracingGeometry::Swap(FVulkanRayTracingGeometry& Other)
 	::Swap(Address, Other.Address);
 
 	AccelerationStructureBuffer = Other.AccelerationStructureBuffer;
-	ScratchBuffer = Other.ScratchBuffer;
 
 	// The rest of the members should be updated using SetInitializer()
 }
 
-void FVulkanRayTracingGeometry::BuildAccelerationStructure(FVulkanCommandListContext& CommandContext, EAccelerationStructureBuildMode BuildMode)
-{
-	FVkRtBLASBuildData BuildData;
-	GetBLASBuildData(
-		Device->GetInstanceHandle(),
-		MakeArrayView(Initializer.Segments),
-		Initializer.IndexBuffer,
-		Initializer.IndexBufferOffset,
-		Initializer.bFastBuild,
-		Initializer.bAllowUpdate,
-		Initializer.IndexBuffer ? Initializer.IndexBuffer->GetStride() : 0,
-		BuildMode,
-		BuildData);
-
-	check(BuildData.SizesInfo.accelerationStructureSize <= AccelerationStructureBuffer->GetSize());
-
-	BuildData.GeometryInfo.dstAccelerationStructure = Handle;
-	BuildData.GeometryInfo.scratchData.deviceAddress = ScratchBuffer->GetDeviceAddress();
-
-	VkAccelerationStructureBuildRangeInfoKHR* const pBuildRanges = BuildData.Ranges.GetData();
-
-	FVulkanCommandBufferManager& CommandBufferManager = *CommandContext.GetCommandBufferManager();
-	FVulkanCmdBuffer* const CmdBuffer = CommandBufferManager.GetActiveCmdBuffer();
-	vkCmdBuildAccelerationStructuresKHR(CmdBuffer->GetHandle(), 1, &BuildData.GeometryInfo, &pBuildRanges);
-
-	CommandBufferManager.SubmitActiveCmdBuffer();
-	CommandBufferManager.PrepareForNewActiveCommandBuffer();
-
-	// No longer need scratch memory for a static build
-	if (!Initializer.bAllowUpdate)
-	{
-		ScratchBuffer = nullptr;
-	}
-}
-
 static void GetTLASBuildData(
 	const VkDevice Device,
 	const uint32 NumInstances,
@@ -676,9 +629,9 @@ FRayTracingAccelerationStructureSize FVulkanDynamicRHI::RHICalcRayTracingGeometr
 		BuildData);
 
 	FRayTracingAccelerationStructureSize Result;
-	Result.ResultSize = BuildData.SizesInfo.accelerationStructureSize;
-	Result.BuildScratchSize = BuildData.SizesInfo.buildScratchSize;
-	Result.UpdateScratchSize = BuildData.SizesInfo.updateScratchSize;
+	Result.ResultSize = Align(BuildData.SizesInfo.accelerationStructureSize, GRHIRayTracingAccelerationStructureAlignment);
+	Result.BuildScratchSize = Align(BuildData.SizesInfo.buildScratchSize, GRHIRayTracingScratchBufferAlignment);
+	Result.UpdateScratchSize = Align(BuildData.SizesInfo.updateScratchSize, GRHIRayTracingScratchBufferAlignment);
 	
 	return Result;
 }
@@ -706,17 +659,97 @@ void FVulkanCommandListContext::RHIBindAccelerationStructureMemory(FRHIRayTracin
 // Todo: High level rhi call should have transitioned and verified vb and ib to read for each segment
 void FVulkanCommandListContext::RHIBuildAccelerationStructures(const TArrayView<const FRayTracingGeometryBuildParams> Params, const FRHIBufferRange& ScratchBufferRange)
 {
+	checkf(ScratchBufferRange.Buffer != nullptr, TEXT("BuildAccelerationStructures requires valid scratch buffer"));
+
+	// Update geometry vertex buffers
 	for (const FRayTracingGeometryBuildParams& P : Params)
 	{
 		FVulkanRayTracingGeometry* const Geometry = ResourceCast(P.Geometry.GetReference());
 
-		// Todo: Update geometry from params for each segment
-		// Todo: Can we do this only for an update?
-		// Todo: Use provided ScratchBuffer instead of allocating.
+		if (P.Segments.Num())
+		{
+			checkf(P.Segments.Num() == Geometry->Initializer.Segments.Num(),
+				TEXT("If updated segments are provided, they must exactly match existing geometry segments. Only vertex buffer bindings may change."));
 
-		// Build as for each segment
-		Geometry->BuildAccelerationStructure(*this, P.BuildMode);
+			for (int32 i = 0; i < P.Segments.Num(); ++i)
+			{
+				checkf(P.Segments[i].MaxVertices <= Geometry->Initializer.Segments[i].MaxVertices,
+					TEXT("Maximum number of vertices in a segment (%u) must not be smaller than what was declared during FRHIRayTracingGeometry creation (%u), as this controls BLAS memory allocation."),
+					P.Segments[i].MaxVertices, Geometry->Initializer.Segments[i].MaxVertices
+				);
+
+				Geometry->Initializer.Segments[i].VertexBuffer = P.Segments[i].VertexBuffer;
+				Geometry->Initializer.Segments[i].VertexBufferElementType = P.Segments[i].VertexBufferElementType;
+				Geometry->Initializer.Segments[i].VertexBufferStride = P.Segments[i].VertexBufferStride;
+				Geometry->Initializer.Segments[i].VertexBufferOffset = P.Segments[i].VertexBufferOffset;
+			}
+		}
 	}
+
+	uint32 ScratchBufferSize = ScratchBufferRange.Size ? ScratchBufferRange.Size : ScratchBufferRange.Buffer->GetSize();
+
+	checkf(ScratchBufferSize + ScratchBufferRange.Offset <= ScratchBufferRange.Buffer->GetSize(),
+		TEXT("BLAS scratch buffer range size is %lld bytes with offset %lld, but the buffer only has %lld bytes. "),
+		ScratchBufferRange.Size, ScratchBufferRange.Offset, ScratchBufferRange.Buffer->GetSize());
+
+	const uint64 ScratchAlignment = GRHIRayTracingScratchBufferAlignment;
+	FVulkanResourceMultiBuffer* ScratchBuffer = ResourceCast(ScratchBufferRange.Buffer);
+	uint32 ScratchBufferOffset = ScratchBufferRange.Offset;
+
+	TArray<FVkRtBLASBuildData, TInlineAllocator<32>> TempBuildData;
+	TArray<VkAccelerationStructureBuildGeometryInfoKHR, TInlineAllocator<32>> BuildGeometryInfos;
+	TArray<VkAccelerationStructureBuildRangeInfoKHR*, TInlineAllocator<32>> BuildRangeInfos;
+	TempBuildData.Reserve(Params.Num());
+	BuildGeometryInfos.Reserve(Params.Num());
+	BuildRangeInfos.Reserve(Params.Num());	
+
+	for (const FRayTracingGeometryBuildParams& P : Params)
+	{
+		FVulkanRayTracingGeometry* const Geometry = ResourceCast(P.Geometry.GetReference());
+		const bool bIsUpdate = P.BuildMode == EAccelerationStructureBuildMode::Update;
+
+		uint64 ScratchBufferRequiredSize = bIsUpdate ? Geometry->SizeInfo.UpdateScratchSize : Geometry->SizeInfo.BuildScratchSize;
+		checkf(ScratchBufferRequiredSize + ScratchBufferOffset <= ScratchBufferSize,
+			TEXT("BLAS scratch buffer size is %ld bytes with offset %ld (%ld bytes available), but the build requires %lld bytes. "),
+			ScratchBufferSize, ScratchBufferOffset, ScratchBufferSize - ScratchBufferOffset, ScratchBufferRequiredSize);
+
+		FVkRtBLASBuildData& BuildData = TempBuildData.AddDefaulted_GetRef();
+		GetBLASBuildData(
+			Device->GetInstanceHandle(),
+			MakeArrayView(Geometry->Initializer.Segments),
+			Geometry->Initializer.IndexBuffer,
+			Geometry->Initializer.IndexBufferOffset,
+			Geometry->Initializer.bFastBuild,
+			Geometry->Initializer.bAllowUpdate,
+			Geometry->Initializer.IndexBuffer ? Geometry->Initializer.IndexBuffer->GetStride() : 0,
+			P.BuildMode,
+			BuildData);
+
+		check(BuildData.SizesInfo.accelerationStructureSize <= Geometry->AccelerationStructureBuffer->GetSize());
+
+		BuildData.GeometryInfo.dstAccelerationStructure = Geometry->Handle;
+		BuildData.GeometryInfo.srcAccelerationStructure = bIsUpdate ? Geometry->Handle : VK_NULL_HANDLE;
+
+		VkDeviceAddress ScratchBufferAddress = ScratchBuffer->GetDeviceAddress() + ScratchBufferOffset;
+		ScratchBufferOffset += ScratchBufferRequiredSize;
+
+		checkf(ScratchBufferAddress % GRHIRayTracingScratchBufferAlignment == 0,
+			TEXT("BLAS scratch buffer (plus offset) must be aligned to %ld bytes."),
+			GRHIRayTracingScratchBufferAlignment);
+
+		BuildData.GeometryInfo.scratchData.deviceAddress = ScratchBufferAddress;
+
+		VkAccelerationStructureBuildRangeInfoKHR* const pBuildRanges = BuildData.Ranges.GetData();
+
+		BuildGeometryInfos.Add(BuildData.GeometryInfo);
+		BuildRangeInfos.Add(pBuildRanges);
+	}
+	
+	FVulkanCmdBuffer* const CmdBuffer = CommandBufferManager->GetActiveCmdBuffer();
+	vkCmdBuildAccelerationStructuresKHR(CmdBuffer->GetHandle(), Params.Num(), BuildGeometryInfos.GetData(), BuildRangeInfos.GetData());
+
+	CommandBufferManager->SubmitActiveCmdBuffer();
+	CommandBufferManager->PrepareForNewActiveCommandBuffer();
 }
 
 void FVulkanCommandListContext::RHIBuildAccelerationStructure(const FRayTracingSceneBuildParams& SceneBuildParams)
diff --git a/Engine/Source/Runtime/VulkanRHI/Private/VulkanRayTracing.h b/Engine/Source/Runtime/VulkanRHI/Private/VulkanRayTracing.h
--- a/Engine/Source/Runtime/VulkanRHI/Private/VulkanRayTracing.h
+++ b/Engine/Source/Runtime/VulkanRHI/Private/VulkanRayTracing.h
@@ -88,15 +88,16 @@ public:
 	virtual void SetInitializer(const FRayTracingGeometryInitializer& Initializer) final override;
 
 	void Swap(FVulkanRayTracingGeometry& Other);
-	void BuildAccelerationStructure(FVulkanCommandListContext& CommandContext, EAccelerationStructureBuildMode BuildMode);
 
-private:
-	FVulkanDevice* const Device = nullptr;
+	using FRHIRayTracingGeometry::Initializer;
+	using FRHIRayTracingGeometry::SizeInfo;
 
 	VkAccelerationStructureKHR Handle = VK_NULL_HANDLE;
 	VkDeviceAddress Address = 0;
 	TRefCountPtr<FVulkanResourceMultiBuffer> AccelerationStructureBuffer;
-	TRefCountPtr<FVulkanResourceMultiBuffer> ScratchBuffer;
+
+private:
+	FVulkanDevice* const Device = nullptr;
 };
 
 class FVulkanRayTracingScene : public FRHIRayTracingScene
