Copying //UE5/Dev-ParallelRendering to //UE5/Main

diff --git a/Engine/Source/Runtime/Renderer/Public/MeshPassProcessor.h b/Engine/Source/Runtime/Renderer/Public/MeshPassProcessor.h
--- a/Engine/Source/Runtime/Renderer/Public/MeshPassProcessor.h
+++ b/Engine/Source/Runtime/Renderer/Public/MeshPassProcessor.h
@@ -231,7 +231,7 @@ class FGraphicsMinimalPipelineStateInitializer
 {
 public:
 	// Can't use TEnumByte<EPixelFormat> as it changes the struct to be non trivially constructible, breaking memset
-	using TRenderTargetFormats = TStaticArray<uint8/*EPixelFormat*/, MaxSimultaneousRenderTargets>;
+	using TRenderTargetFormats = TStaticArray<EPixelFormat, MaxSimultaneousRenderTargets>;
 	using TRenderTargetFlags = TStaticArray<uint32/*ETextureCreateFlags*/, MaxSimultaneousRenderTargets>;
 
 	FGraphicsMinimalPipelineStateInitializer()
@@ -239,10 +239,7 @@ public:
 		, RasterizerState(nullptr)
 		, DepthStencilState(nullptr)
 		, PrimitiveType(PT_Num)
-	{
-		static_assert(sizeof(EPixelFormat) != sizeof(uint8), "Change TRenderTargetFormats's uint8 to EPixelFormat");
-		static_assert(PF_MAX < MAX_uint8, "TRenderTargetFormats assumes EPixelFormat can fit in a uint8!");
-	}
+	{}
 
 	FGraphicsMinimalPipelineStateInitializer(
 		FMinimalBoundShaderStateInput	InBoundShaderState,
diff --git a/Engine/Source/Runtime/SlateCore/Public/Rendering/SlateRenderer.h b/Engine/Source/Runtime/SlateCore/Public/Rendering/SlateRenderer.h
--- a/Engine/Source/Runtime/SlateCore/Public/Rendering/SlateRenderer.h
+++ b/Engine/Source/Runtime/SlateCore/Public/Rendering/SlateRenderer.h
@@ -10,7 +10,7 @@
 #include "Fonts/FontTypes.h"
 #include "PixelFormat.h"
 
-class FRHITexture2D;
+class FRHITexture;
 class FRenderTarget;
 class FSlateDrawBuffer;
 class FSlateUpdatableTexture;
@@ -22,7 +22,7 @@ struct Rect;
 class FSceneInterface;
 struct FSlateBrush;
 
-typedef TRefCountPtr<FRHITexture2D> FTexture2DRHIRef;
+typedef TRefCountPtr<FRHITexture> FTexture2DRHIRef;
 
 /**
  * Update context for deferred drawing of widgets to render targets
diff --git a/Engine/Source/Runtime/VulkanRHI/Private/VulkanBarriers.cpp b/Engine/Source/Runtime/VulkanRHI/Private/VulkanBarriers.cpp
--- a/Engine/Source/Runtime/VulkanRHI/Private/VulkanBarriers.cpp
+++ b/Engine/Source/Runtime/VulkanRHI/Private/VulkanBarriers.cpp
@@ -393,7 +393,7 @@ static void GetVkStageAndAccessFlags(ERHIAccess RHIAccess, FRHITransitionInfo::E
 //
 struct FDepthStencilSubresTransition
 {
-	FVulkanTextureBase* Texture;
+	FVulkanTexture* Texture;
 	int ArraySlice;
 	ERHIAccess SrcDepthAccess;
 	ERHIAccess DestDepthAccess;
@@ -643,8 +643,8 @@ void FVulkanDynamicRHI::RHICreateTransition(FRHITransition* Transition, const FR
 		if (Info.Type == FRHITransitionInfo::EType::Texture)
 		{
 			// CPU accessible "textures" are implemented as buffers. Check if this is a real texture or a buffer.
-			FVulkanTextureBase* Texture = FVulkanTextureBase::Cast(Info.Texture);
-			if (Texture->Surface.GetCpuReadbackBuffer() == nullptr)
+			FVulkanTexture* Texture = FVulkanTexture::Cast(Info.Texture);
+			if (Texture->GetCpuReadbackBuffer() == nullptr)
 			{
 				++NumTextures;
 			}
@@ -691,7 +691,7 @@ void FVulkanDynamicRHI::RHICreateTransition(FRHITransition* Transition, const FR
 		checkf(Info.AccessAfter != ERHIAccess::Unknown, TEXT("Transitioning a resource to an unknown state is not allowed."));
 
 		FVulkanResourceMultiBuffer* Buffer = nullptr;
-		FVulkanTextureBase* Texture = nullptr;
+		FVulkanTexture* Texture = nullptr;
 		FRHITransitionInfo::EType UnderlyingType = Info.Type;
 		uint32 UsageFlags = 0;
 
@@ -699,8 +699,8 @@ void FVulkanDynamicRHI::RHICreateTransition(FRHITransition* Transition, const FR
 		{
 		case FRHITransitionInfo::EType::Texture:
 		{
-			Texture = FVulkanTextureBase::Cast(Info.Texture);
-			if (Texture->Surface.GetCpuReadbackBuffer())
+			Texture = FVulkanTexture::Cast(Info.Texture);
+			if (Texture->GetCpuReadbackBuffer())
 			{
 				Texture = nullptr;
 			}
@@ -717,7 +717,7 @@ void FVulkanDynamicRHI::RHICreateTransition(FRHITransition* Transition, const FR
 			FVulkanUnorderedAccessView* UAV = ResourceCast(Info.UAV);
 			if (UAV->SourceTexture)
 			{
-				Texture = FVulkanTextureBase::Cast(UAV->SourceTexture);
+				Texture = FVulkanTexture::Cast(UAV->SourceTexture);
 				UnderlyingType = FRHITransitionInfo::EType::Texture;
 			}
 			else if (UAV->SourceBuffer)
@@ -749,13 +749,13 @@ void FVulkanDynamicRHI::RHICreateTransition(FRHITransition* Transition, const FR
 		VkAccessFlags SrcAccessFlags, DstAccessFlags;
 		VkImageLayout SrcLayout, DstLayout;
 
-		const bool bIsDepthStencil = Texture && Texture->Surface.IsDepthOrStencilAspect();
+		const bool bIsDepthStencil = Texture && Texture->IsDepthOrStencilAspect();
 
 		// If the device doesn't support separate depth-stencil layouts, we must merge depth-stencil subresource transitions so we only do one barrier on the image.
 		if (bIsDepthStencil && Info.PlaneSlice != FRHISubresourceRange::kAllSubresources)
 		{
 			int32 Index = DSSubresTransitions.IndexOfByPredicate([Texture, Info](const FDepthStencilSubresTransition& Entry) -> bool {
-				return Entry.Texture->Surface.Image == Texture->Surface.Image && Entry.ArraySlice == Info.ArraySlice;
+				return Entry.Texture->Image == Texture->Image && Entry.ArraySlice == Info.ArraySlice;
 			});
 
 			FDepthStencilSubresTransition* PendingTransition;
@@ -805,7 +805,7 @@ void FVulkanDynamicRHI::RHICreateTransition(FRHITransition* Transition, const FR
 		else
 		{
 			
-			const bool bSupportsReadOnlyOptimal = (Texture == nullptr) || Texture->Surface.SupportsSampling();
+			const bool bSupportsReadOnlyOptimal = (Texture == nullptr) || Texture->SupportsSampling();
 
 			GetVkStageAndAccessFlags(Info.AccessBefore, UnderlyingType, UsageFlags, bIsDepthStencil, bSupportsReadOnlyOptimal, SrcStageMask, SrcAccessFlags, SrcLayout, true);
 			GetVkStageAndAccessFlags(Info.AccessAfter, UnderlyingType, UsageFlags, bIsDepthStencil, bSupportsReadOnlyOptimal, DstStageMask, DstAccessFlags, DstLayout, false);
@@ -862,7 +862,7 @@ void FVulkanDynamicRHI::RHICreateTransition(FRHITransition* Transition, const FR
 		check(Texture != nullptr);
 
 		VkImageSubresourceRange SubresRange;
-		SetupSubresourceRange(SubresRange, Info, Texture->Surface.GetFullAspectMask());
+		SetupSubresourceRange(SubresRange, Info, Texture->GetFullAspectMask());
 
 		// For some textures, e.g. FVulkanBackBuffer, the image handle may not be set yet, or may be stale, so there's no point storing it here.
 		// We'll set the image to NULL in the barrier info, and RHIEndTransitions will fetch the up to date pointer from the texture, after
@@ -1066,10 +1066,10 @@ void FVulkanCommandListContext::RHIBeginTransitions(TArrayView<const FRHITransit
 		{
 			const VkImageMemoryBarrier& ImageBarrier = Data->ImageBarriers[ImgBarrierIdx];
 			const FVulkanPipelineBarrier::ImageBarrierExtraData& ExtraData = Data->ImageBarrierExtras[ImgBarrierIdx];
-			FVulkanTextureBase* Texture = ExtraData.BaseTexture;
-			check(Texture->Surface.Image != VK_NULL_HANDLE);
+			FVulkanTexture* Texture = ExtraData.BaseTexture;
+			check(Texture->Image != VK_NULL_HANDLE);
 
-			const FVulkanImageLayout& Layout = LayoutManager.GetOrAddFullLayout(Texture->Surface, VK_IMAGE_LAYOUT_UNDEFINED);
+			const FVulkanImageLayout& Layout = LayoutManager.GetOrAddFullLayout(*Texture, VK_IMAGE_LAYOUT_UNDEFINED);
 
 			VkAccessFlags SrcAccessFlags;
 			VkImageLayout SrcLayout, DstLayout;
@@ -1093,14 +1093,14 @@ void FVulkanCommandListContext::RHIBeginTransitions(TArrayView<const FRHITransit
 
 			VkImageMemoryBarrier& RealBarrier = RealImageBarriers.AddDefaulted_GetRef();
 			RealBarrier = ImageBarrier;
-			RealBarrier.image = Texture->Surface.Image;
+			RealBarrier.image = Texture->Image;
 			RealBarrier.srcAccessMask = SrcAccessFlags;
 			RealBarrier.dstAccessMask = 0; // Release resource from current queue.
 			RealBarrier.oldLayout = SrcLayout;
 			RealBarrier.newLayout = DstLayout;
 			
-			// Fix up the destination layout if this barrier specifies a sub-aspect of a depth-stencil surface.
-			AdjustDepthStencilLayout(RealBarrier, Texture->Surface.GetFullAspectMask());
+			// Fix up the destination layout if this barrier specifies a sub-aspect of a depth-stencil 
+			AdjustDepthStencilLayout(RealBarrier, Texture->GetFullAspectMask());
 
 			// We don't update the image layout here. That will be done in RHIEndTransitions.
 		}
@@ -1264,16 +1264,22 @@ void FVulkanCommandListContext::RHIEndTransitions(TArrayView<const FRHITransitio
 		for (int32 ImgBarrierIdx = 0; ImgBarrierIdx < Data->ImageBarriers.Num(); ++ImgBarrierIdx)
 		{
 			const VkImageMemoryBarrier& ImageBarrier = Data->ImageBarriers[ImgBarrierIdx];
+			check(ImageBarrier.image == VK_NULL_HANDLE);
+
 			const FVulkanPipelineBarrier::ImageBarrierExtraData& ExtraData = Data->ImageBarrierExtras[ImgBarrierIdx];
-			FVulkanTextureBase* Texture = ExtraData.BaseTexture;
-			check(Texture->Surface.GetCpuReadbackBuffer() == nullptr);
+			FVulkanTexture* Texture = ExtraData.BaseTexture;
+			check(Texture->GetCpuReadbackBuffer() == nullptr);
 
 			Texture->OnLayoutTransition(*this, ImageBarrier.newLayout);
 
-			// Make sure the texture contains a valid image handle now.
-			check(Texture->Surface.Image != VK_NULL_HANDLE && ImageBarrier.image == VK_NULL_HANDLE);
+			if (Texture->Image == VK_NULL_HANDLE)
+			{
+				// If Texture is a backbuffer, OnLayoutTransition normally updates the image handle. However, if the viewport got invalidated during the
+				// OnLayoutTransition call, we may end up with a null handle, which is fine. Just ignore the transition.
+				continue;
+			}
 
-			FVulkanImageLayout& Layout = LayoutManager.GetOrAddFullLayout(Texture->Surface, VK_IMAGE_LAYOUT_UNDEFINED);
+			FVulkanImageLayout& Layout = LayoutManager.GetOrAddFullLayout(*Texture, VK_IMAGE_LAYOUT_UNDEFINED);
 
 			VkAccessFlags SrcAccessFlags;
 			VkImageLayout SrcLayout, DstLayout;
@@ -1293,7 +1299,7 @@ void FVulkanCommandListContext::RHIEndTransitions(TArrayView<const FRHITransitio
 				{
 					// Slow path, adds one transition per subresource.
 					check(Data->SrcPipelines == Data->DstPipelines);
-					AddSubresourceTransitions(RealImageBarriers, RealSrcStageMask, ImageBarrier, Texture->Surface.Image, Layout, DstLayout);
+					AddSubresourceTransitions(RealImageBarriers, RealSrcStageMask, ImageBarrier, Texture->Image, Layout, DstLayout);
 					continue;
 				}
 			}
@@ -1305,13 +1311,13 @@ void FVulkanCommandListContext::RHIEndTransitions(TArrayView<const FRHITransitio
 			}
 
 			VkImageMemoryBarrier RealBarrier = ImageBarrier;
-			RealBarrier.image = Texture->Surface.Image; // Use the up to date image handle.
+			RealBarrier.image = Texture->Image; // Use the up to date image handle.
 			RealBarrier.srcAccessMask = SrcAccessFlags;
 			RealBarrier.oldLayout = SrcLayout;
 			RealBarrier.newLayout = DstLayout;
 
 			// Fix up the destination layout if this barrier specifies a sub-aspect of a depth-stencil surface.
-			AdjustDepthStencilLayout(RealBarrier, Texture->Surface.GetFullAspectMask());
+			AdjustDepthStencilLayout(RealBarrier, Texture->GetFullAspectMask());
 
 			if (Data->SrcPipelines == Data->DstPipelines)
 			{
@@ -1481,7 +1487,7 @@ void FVulkanPipelineBarrier::AddImageLayoutTransition(VkImage Image, VkImageAspe
 	}
 }
 
-void FVulkanPipelineBarrier::AddImageAccessTransition(const FVulkanSurface& Surface, ERHIAccess SrcAccess, ERHIAccess DstAccess, const VkImageSubresourceRange& SubresourceRange, VkImageLayout& InOutLayout)
+void FVulkanPipelineBarrier::AddImageAccessTransition(const FVulkanTexture& Surface, ERHIAccess SrcAccess, ERHIAccess DstAccess, const VkImageSubresourceRange& SubresourceRange, VkImageLayout& InOutLayout)
 {
 	// This function should only be used for known states.
 	check(DstAccess != ERHIAccess::Unknown);
@@ -1715,7 +1721,7 @@ void FVulkanLayoutManager::ValidateRenderPassColorEntry(const FRHIRenderPassInfo
 {
 	FRHITexture* Texture = bResolveTarget ? ColorEntry.ResolveTarget : ColorEntry.RenderTarget;
 	CA_ASSUME(Texture);
-	FVulkanSurface& Surface = FVulkanTextureBase::Cast(Texture)->Surface;
+	FVulkanTexture& Surface = *FVulkanTexture::Cast(Texture);
 	check(Surface.Image != VK_NULL_HANDLE);
 
 	// Check that the image is in the correct layout for rendering.
@@ -1761,7 +1767,7 @@ void FVulkanLayoutManager::BeginRenderPass(FVulkanCommandListContext& Context, F
 
 		FRHITexture* ColorTexture = ColorEntry.RenderTarget;
 		CA_ASSUME(ColorTexture);
-		FVulkanSurface& ColorSurface = FVulkanTextureBase::Cast(ColorTexture)->Surface;
+		FVulkanTexture& ColorSurface = *FVulkanTexture::Cast(ColorTexture);
 		const bool bPassPerformsResolve = ColorSurface.GetNumSamples() > 1 && ColorEntry.ResolveTarget;
 
 		ValidateRenderPassColorEntry(ColorEntry, false, Barrier);
@@ -1822,7 +1828,7 @@ void FVulkanLayoutManager::BeginRenderPass(FVulkanCommandListContext& Context, F
 	FRHITexture* ShadingRateTexture = RPInfo.ShadingRateTexture;
 	if (ShadingRateTexture && ValidateShadingRateDataType())
 	{
-		FVulkanSurface& Surface = FVulkanTextureBase::Cast(ShadingRateTexture)->Surface;
+		FVulkanTexture& Surface = *FVulkanTexture::Cast(ShadingRateTexture);
 		VkImageLayout& DSLayout = FindOrAddLayoutRW(Surface, VK_IMAGE_LAYOUT_UNDEFINED);
 		
 		VkImageLayout ExpectedLayout = VK_IMAGE_LAYOUT_UNDEFINED;
diff --git a/Engine/Source/Runtime/VulkanRHI/Private/VulkanBarriers.h b/Engine/Source/Runtime/VulkanRHI/Private/VulkanBarriers.h
--- a/Engine/Source/Runtime/VulkanRHI/Private/VulkanBarriers.h
+++ b/Engine/Source/Runtime/VulkanRHI/Private/VulkanBarriers.h
@@ -5,7 +5,6 @@
 #include "CoreMinimal.h"
 #include "VulkanCommon.h"
 
-class FVulkanSurface;
 
 struct FVulkanPipelineBarrier
 {
@@ -25,7 +24,7 @@ struct FVulkanPipelineBarrier
 	// We need to keep the texture pointers around, because we need to call OnTransitionResource on them, and we need mip and layer counts for the tracking code.
 	struct ImageBarrierExtraData
 	{
-		FVulkanTextureBase* BaseTexture = nullptr;
+		FVulkanTexture* BaseTexture = nullptr;
 		bool IsAliasingBarrier = false;
 	};
 	TArray<ImageBarrierExtraData, TInlineAllocator<2>> ImageBarrierExtras;
@@ -35,7 +34,7 @@ struct FVulkanPipelineBarrier
 	void AddImageLayoutTransition(VkImage Image, VkImageAspectFlags AspectMask, const struct FVulkanImageLayout& SrcLayout, VkImageLayout DstLayout);
 	void AddImageLayoutTransition(VkImage Image, VkImageAspectFlags AspectMask, VkImageLayout SrcLayout, const struct FVulkanImageLayout& DstLayout);
 	void AddImageLayoutTransition(VkImage Image, VkImageAspectFlags AspectMask, const struct FVulkanImageLayout& SrcLayout, const struct FVulkanImageLayout& DstLayout);
-	void AddImageAccessTransition(const FVulkanSurface& Surface, ERHIAccess SrcAccess, ERHIAccess DstAccess, const VkImageSubresourceRange& SubresourceRange, VkImageLayout& InOutLayout);
+	void AddImageAccessTransition(const FVulkanTexture& Surface, ERHIAccess SrcAccess, ERHIAccess DstAccess, const VkImageSubresourceRange& SubresourceRange, VkImageLayout& InOutLayout);
 	void Execute(VkCommandBuffer CmdBuffer);
 
 	static VkImageSubresourceRange MakeSubresourceRange(VkImageAspectFlags AspectMask, uint32 FirstMip = 0, uint32 NumMips = VK_REMAINING_MIP_LEVELS, uint32 FirstLayer = 0, uint32 NumLayers = VK_REMAINING_ARRAY_LAYERS);
@@ -164,7 +163,7 @@ public:
 		return Layouts.FindChecked(Image);
 	}
 
-	FVulkanImageLayout& GetOrAddFullLayout(const FVulkanSurface& Surface, VkImageLayout LayoutIfNotFound)
+	FVulkanImageLayout& GetOrAddFullLayout(const FVulkanTexture& Surface, VkImageLayout LayoutIfNotFound)
 	{
 		FVulkanImageLayout* Layout = Layouts.Find(Surface.Image);
 		if (Layout)
@@ -202,12 +201,12 @@ public:
 		return FindOrAddFullLayoutRW(Image, LayoutIfNotFound, NumMips, NumLayers).MainLayout;
 	}
 
-	VULKANRHI_API VkImageLayout& FindOrAddLayoutRW(const FVulkanSurface& Surface, VkImageLayout LayoutIfNotFound)
+	VULKANRHI_API VkImageLayout& FindOrAddLayoutRW(const FVulkanTexture& Surface, VkImageLayout LayoutIfNotFound)
 	{
 		return FindOrAddLayoutRW(Surface.Image, LayoutIfNotFound, Surface.GetNumMips(), Surface.GetNumberOfArrayLevels());
 	}
 
-	VkImageLayout FindOrAddLayout(const FVulkanSurface& Surface, VkImageLayout LayoutIfNotFound)
+	VkImageLayout FindOrAddLayout(const FVulkanTexture& Surface, VkImageLayout LayoutIfNotFound)
 	{
 		return FindOrAddLayoutRW(Surface, LayoutIfNotFound);
 	}
diff --git a/Engine/Source/Runtime/VulkanRHI/Private/VulkanCommands.cpp b/Engine/Source/Runtime/VulkanRHI/Private/VulkanCommands.cpp
--- a/Engine/Source/Runtime/VulkanRHI/Private/VulkanCommands.cpp
+++ b/Engine/Source/Runtime/VulkanRHI/Private/VulkanCommands.cpp
@@ -191,8 +191,8 @@ void FVulkanCommandListContext::RHISetUAVParameter(FRHIComputeShader* ComputeSha
 
 void FVulkanCommandListContext::RHISetShaderTexture(FRHIGraphicsShader* ShaderRHI, uint32 TextureIndex, FRHITexture* NewTextureRHI)
 {
-	FVulkanTextureBase* Texture = FVulkanTextureBase::Cast(NewTextureRHI);
-	VkImageLayout Layout = LayoutManager.FindLayoutChecked(Texture->Surface.Image);
+	FVulkanTexture* Texture = FVulkanTexture::Cast(NewTextureRHI);
+	VkImageLayout Layout = LayoutManager.FindLayoutChecked(Texture->Image);
 
 	ShaderStage::EStage Stage = GetAndVerifyShaderStage(ShaderRHI, PendingGfxState);
 	PendingGfxState->SetTextureForStage(Stage, TextureIndex, Texture, Layout);
@@ -204,8 +204,8 @@ void FVulkanCommandListContext::RHISetShaderTexture(FRHIComputeShader* ComputeSh
 	FVulkanComputeShader* ComputeShader = ResourceCast(ComputeShaderRHI);
 	check(PendingComputeState->GetCurrentShader() == ComputeShader);
 
-	FVulkanTextureBase* VulkanTexture = FVulkanTextureBase::Cast(NewTextureRHI);
-	VkImageLayout Layout = LayoutManager.FindLayoutChecked(VulkanTexture->Surface.Image);
+	FVulkanTexture* VulkanTexture = FVulkanTexture::Cast(NewTextureRHI);
+	VkImageLayout Layout = LayoutManager.FindLayoutChecked(VulkanTexture->Image);
 	PendingComputeState->SetTextureForStage(TextureIndex, VulkanTexture, Layout);
 	NewTextureRHI->SetLastRenderTime((float)FPlatformTime::Seconds());
 }
@@ -303,14 +303,14 @@ inline void SetShaderUniformBufferResources(FVulkanCommandListContext* Context,
 			FRHITexture* TexRef = (FRHITexture*)(ResourceArray[ResourceInfo.SourceUBResourceIndex].GetReference());
 			if (TexRef)
 			{
-				const FVulkanTextureBase* BaseTexture = FVulkanTextureBase::Cast(TexRef);
-				if (!ensure(BaseTexture))
+				const FVulkanTexture* VulkanTexture = FVulkanTexture::Cast(TexRef);
+				if (!ensure(VulkanTexture))
 				{
-					BaseTexture = FVulkanTextureBase::Cast(GBlackTexture->TextureRHI.GetReference());
+					VulkanTexture = FVulkanTexture::Cast(GBlackTexture->TextureRHI.GetReference());
 				}
 
 				// If the descriptor is a storage image in a slot expecting to read only, make sure it's because we don't support sampling
-				ensure(DescriptorType != VK_DESCRIPTOR_TYPE_STORAGE_IMAGE || !BaseTexture->Surface.SupportsSampling());
+				ensure(DescriptorType != VK_DESCRIPTOR_TYPE_STORAGE_IMAGE || !VulkanTexture->SupportsSampling());
 
 #if ENABLE_RHI_VALIDATION
 				if (Context->Tracker)
@@ -319,8 +319,8 @@ inline void SetShaderUniformBufferResources(FVulkanCommandListContext* Context,
 				}
 #endif
 
-				const VkImageLayout Layout = Context->GetLayoutManager().FindLayoutChecked(BaseTexture->Surface.Image);
-				State->SetTextureForUBResource(GlobalRemappingInfo[ResourceInfo.GlobalIndex].NewDescriptorSet, GlobalRemappingInfo[ResourceInfo.GlobalIndex].NewBindingIndex, BaseTexture, Layout);
+				const VkImageLayout Layout = Context->GetLayoutManager().FindLayoutChecked(VulkanTexture->Image);
+				State->SetTextureForUBResource(GlobalRemappingInfo[ResourceInfo.GlobalIndex].NewDescriptorSet, GlobalRemappingInfo[ResourceInfo.GlobalIndex].NewBindingIndex, VulkanTexture, Layout);
 				TexRef->SetLastRenderTime(CurrentTime);
 			}
 			else
diff --git a/Engine/Source/Runtime/VulkanRHI/Private/VulkanDebug.cpp b/Engine/Source/Runtime/VulkanRHI/Private/VulkanDebug.cpp
--- a/Engine/Source/Runtime/VulkanRHI/Private/VulkanDebug.cpp
+++ b/Engine/Source/Runtime/VulkanRHI/Private/VulkanDebug.cpp
@@ -1306,13 +1306,13 @@ namespace VulkanRHI
 
 	void PrintResultAndPointer(VkResult Result, uint64 Handle)
 	{
-		DebugLog += FString::Printf(TEXT(" -> %s => %ull\n"), *GetVkResultErrorString(Result), Handle);
+		DebugLog += FString::Printf(TEXT(" -> %s => %llu\n"), *GetVkResultErrorString(Result), Handle);
 		FlushDebugWrapperLog();
 	}
 
 	void PrintResultAndNamedHandle(VkResult Result, const TCHAR* HandleName, uint64 Handle)
 	{
-		DebugLog += FString::Printf(TEXT(" -> %s => %s=%ull\n"), *GetVkResultErrorString(Result), HandleName, Handle);
+		DebugLog += FString::Printf(TEXT(" -> %s => %s=%llu\n"), *GetVkResultErrorString(Result), HandleName, Handle);
 		FlushDebugWrapperLog();
 	}
 }
diff --git a/Engine/Source/Runtime/VulkanRHI/Private/VulkanDevice.cpp b/Engine/Source/Runtime/VulkanRHI/Private/VulkanDevice.cpp
--- a/Engine/Source/Runtime/VulkanRHI/Private/VulkanDevice.cpp
+++ b/Engine/Source/Runtime/VulkanRHI/Private/VulkanDevice.cpp
@@ -172,7 +172,7 @@ FVulkanDevice::FVulkanDevice(FVulkanDynamicRHI* InRHI, VkPhysicalDevice InGpu)
 	, MemoryManager(this)
 	, DeferredDeletionQueue(this)
 	, DefaultSampler(nullptr)
-	, DefaultImage(nullptr)
+	, DefaultTexture(nullptr)
 	, Gpu(InGpu)
 	, GfxQueue(nullptr)
 	, ComputeQueue(nullptr)
@@ -1411,9 +1411,19 @@ void FVulkanDevice::InitGPU(int32 DeviceIndex)
 		FSamplerStateInitializerRHI Default(SF_Point);
 		DefaultSampler = ResourceCast(RHICreateSamplerState(Default).GetReference());
 
-		FRHIResourceCreateInfo CreateInfo(TEXT("FVulkanDevice_DefaultImage"));
-		DefaultImage = new FVulkanSurface(*this, 0, VK_IMAGE_VIEW_TYPE_2D, PF_B8G8R8A8, 1, 1, 1, 1, 1, 1, TexCreate_RenderTargetable | TexCreate_ShaderResource, ERHIAccess::SRVMask, CreateInfo);
-		DefaultTextureView.Create(*this, DefaultImage->Image, VK_IMAGE_VIEW_TYPE_2D, DefaultImage->GetFullAspectMask(), PF_B8G8R8A8, VK_FORMAT_B8G8R8A8_UNORM, 0, 1, 0, 1);
+		FRHITextureCreateDesc Desc = FRHITextureCreateDesc::Create2D(
+			TEXT("FVulkanDevice_DefaultImage"),
+			{ 1, 1 },
+			PF_B8G8R8A8,
+			FClearValueBinding::None,
+			TexCreate_RenderTargetable | TexCreate_ShaderResource,
+			1,
+			1,
+			0,
+			ERHIAccess::SRVMask
+		);
+
+		DefaultTexture = new FVulkanTexture(*this, Desc, nullptr);
 	}
 }
 
@@ -1458,9 +1468,6 @@ void FVulkanDevice::Destroy()
 		RHICmdList.ImmediateFlush(EImmediateFlushType::FlushRHIThread);
 	} while (NumDeletes > 0);
 
-	VulkanRHI::vkDestroyImageView(GetInstanceHandle(), DefaultTextureView.View, VULKAN_CPU_ALLOCATOR);
-	DefaultTextureView = {};
-
 	delete DescriptorSetCache;
 	DescriptorSetCache = nullptr;
 	
@@ -1470,8 +1477,8 @@ void FVulkanDevice::Destroy()
 	// No need to delete as it's stored in SamplerMap
 	DefaultSampler = nullptr;
 
-	delete DefaultImage;
-	DefaultImage = nullptr;
+	delete DefaultTexture;
+	DefaultTexture = nullptr;
 
 	for (int32 Index = CommandContexts.Num() - 1; Index >= 0; --Index)
 	{
diff --git a/Engine/Source/Runtime/VulkanRHI/Private/VulkanDevice.h b/Engine/Source/Runtime/VulkanRHI/Private/VulkanDevice.h
--- a/Engine/Source/Runtime/VulkanRHI/Private/VulkanDevice.h
+++ b/Engine/Source/Runtime/VulkanRHI/Private/VulkanDevice.h
@@ -381,7 +381,7 @@ public:
 
 	inline const FVulkanTextureView& GetDefaultImageView() const
 	{
-		return DefaultTextureView;
+		return DefaultTexture->DefaultView;
 	}
 
 	inline const VkFormatProperties* GetFormatProperties() const
@@ -567,8 +567,7 @@ private:
 	FVulkanShaderFactory ShaderFactory;
 
 	FVulkanSamplerState* DefaultSampler;
-	FVulkanSurface* DefaultImage;
-	FVulkanTextureView DefaultTextureView;
+	FVulkanTexture* DefaultTexture;
 
 	VkPhysicalDevice Gpu;
 	VkPhysicalDeviceProperties GpuProps;
diff --git a/Engine/Source/Runtime/VulkanRHI/Private/VulkanMemory.cpp b/Engine/Source/Runtime/VulkanRHI/Private/VulkanMemory.cpp
--- a/Engine/Source/Runtime/VulkanRHI/Private/VulkanMemory.cpp
+++ b/Engine/Source/Runtime/VulkanRHI/Private/VulkanMemory.cpp
@@ -1351,7 +1351,7 @@ namespace VulkanRHI
 			};
 
 
-			TVulkanTrackBase<FVulkanTextureBase>::CollectAll(Collector);
+			TVulkanTrackBase<FVulkanTexture>::CollectAll(Collector);
 			TVulkanTrackBase<FVulkanResourceMultiBuffer>::CollectAll(Collector);
 			for(auto& Itr : AllocationBuckets)
 			{
@@ -4068,15 +4068,11 @@ namespace VulkanRHI
 				if(Alloc.MetaType == EVulkanAllocationMetaImageRenderTarget)
 				{
 					FVulkanEvictable* Evictable = Alloc.AllocationOwner;
-					FVulkanTextureBase* Texture = Evictable->GetTextureBase();
+					FVulkanTexture* Texture = Evictable->GetEvictableTexture();
 					if(Texture)
 					{
-						FRHITexture* RHITexture = Texture->GetRHITexture();
-						if(RHITexture)
-						{
-							Tmp = RHITexture->GetName().ToString();
-							Name = *Tmp;
-						}
+						Tmp = Texture->GetName().ToString();
+						Name = *Tmp;
 					}
 				}
 
@@ -4301,7 +4297,7 @@ namespace VulkanRHI
 						if(Heap->TryRealloc(Allocation, EvictableOwner, EType::Image, Alloc.Size, Alloc.Alignment, Alloc.MetaType))
 						{
 							check(Allocation.HasAllocation());
-							FVulkanTextureBase* Texture = (FVulkanTextureBase*)EvictableOwner;
+							FVulkanTexture* Texture = EvictableOwner->GetEvictableTexture();
 
 							if(GVulkanLogDefrag)
 							{
@@ -5294,26 +5290,3 @@ const VkPhysicalDeviceMemoryProperties& FDeviceMemoryManager::GetMemoryPropertie
 {
 	return MemoryProperties;
 }
-
-
-void FStagingBuffer::Evict(FVulkanDevice& InDevice)
-{
-	checkNoEntry(); //stagingbuffers are always in system memory so this should not happen
-}
-void FStagingBuffer::Move(FVulkanDevice& InDevice, FVulkanCommandListContext& Context, FVulkanAllocation& InAllocation)
-{
-	checkNoEntry();//stagingbuffers are always in system memory so this should not happen
-}
-void FVulkanRingBuffer::Evict(FVulkanDevice& InDevice)
-{
-	checkNoEntry(); //stagingbuffers are always in system memory so this should not happen
-}
-void FVulkanRingBuffer::Move(FVulkanDevice& InDevice, FVulkanCommandListContext& Context, FVulkanAllocation& InAllocation)
-{
-	checkNoEntry();//stagingbuffers are always in system memory so this should not happen
-}
-
-void FVulkanEvictable::OnFullDefrag(FVulkanDevice& InDevice, FVulkanCommandListContext& Context, uint32 NewOffset)
-{
-	checkNoEntry();
-}
diff --git a/Engine/Source/Runtime/VulkanRHI/Private/VulkanMemory.h b/Engine/Source/Runtime/VulkanRHI/Private/VulkanMemory.h
--- a/Engine/Source/Runtime/VulkanRHI/Private/VulkanMemory.h
+++ b/Engine/Source/Runtime/VulkanRHI/Private/VulkanMemory.h
@@ -23,6 +23,7 @@
 class FVulkanCommandListContext;
 class FVulkanQueue;
 class FVulkanCmdBuffer;
+class FVulkanTexture;
 
 #if !VULKAN_OBJECT_TRACKING
 #define VULKAN_TRACK_OBJECT_CREATE(Type, Ptr) do{}while(0)
@@ -88,15 +89,12 @@ namespace VulkanRHI
 class FVulkanEvictable
 {
 public:
-	virtual void Evict(FVulkanDevice& Device) = 0;
-	virtual void Move(FVulkanDevice& Device, FVulkanCommandListContext& Context, VulkanRHI::FVulkanAllocation& Allocation) = 0;
-
-	virtual void OnFullDefrag(FVulkanDevice& Device, FVulkanCommandListContext& Context, uint32 NewOffset);
-	
-	virtual struct FVulkanTextureBase* GetTextureBase(){return 0;}
-	virtual bool CanMove(){return true;}
-	virtual bool CanEvict(){return true;}
-	
+	virtual bool CanMove() const { return false; }
+	virtual bool CanEvict() const { return false; }
+	virtual void Evict(FVulkanDevice& Device) { checkNoEntry(); }
+	virtual void Move(FVulkanDevice& Device, FVulkanCommandListContext& Context, VulkanRHI::FVulkanAllocation& Allocation) { checkNoEntry(); }
+	virtual void OnFullDefrag(FVulkanDevice& Device, FVulkanCommandListContext& Context, uint32 NewOffset) { checkNoEntry(); }
+	virtual FVulkanTexture* GetEvictableTexture() { return nullptr; }
 };
 
 
@@ -908,13 +906,9 @@ namespace VulkanRHI
 
 	class FStagingBuffer : public FVulkanEvictable, public FRefCount
 	{
-		virtual void Evict(FVulkanDevice& Device);
-		virtual void Move(FVulkanDevice& Device, FVulkanCommandListContext& Context, FVulkanAllocation& Allocation);
-		virtual bool CanMove(){return false;}
-		virtual bool CanEvict(){return false;}
-
 	public:
 		FStagingBuffer(FVulkanDevice* InDevice);
+
 		VkBuffer GetHandle() const;
 		void* GetMappedPointer();
 		uint32 GetSize() const;
diff --git a/Engine/Source/Runtime/VulkanRHI/Private/VulkanPendingState.h b/Engine/Source/Runtime/VulkanRHI/Private/VulkanPendingState.h
--- a/Engine/Source/Runtime/VulkanRHI/Private/VulkanPendingState.h
+++ b/Engine/Source/Runtime/VulkanRHI/Private/VulkanPendingState.h
@@ -77,7 +77,7 @@ public:
 		SetUAVForUBResource(DescriptorSet, BindingIndex, UAV);
 	}
 
-	inline void SetTextureForStage(uint32 TextureIndex, const FVulkanTextureBase* TextureBase, VkImageLayout Layout)
+	inline void SetTextureForStage(uint32 TextureIndex, const FVulkanTexture* Texture, VkImageLayout Layout)
 	{
 		const FVulkanComputePipelineDescriptorInfo& DescriptorInfo = CurrentState->GetComputePipelineDescriptorInfo();
 		uint8 DescriptorSet;
@@ -87,7 +87,7 @@ public:
 			return;
 		}
 
-		CurrentState->SetTexture(DescriptorSet, BindingIndex, TextureBase, Layout);
+		CurrentState->SetTexture(DescriptorSet, BindingIndex, Texture, Layout);
 	}
 
 	inline void SetSamplerStateForStage(uint32 SamplerIndex, FVulkanSamplerState* Sampler)
@@ -103,9 +103,9 @@ public:
 		CurrentState->SetSamplerState(DescriptorSet, BindingIndex, Sampler);
 	}
 
-	inline void SetTextureForUBResource(int32 DescriptorSet, uint32 BindingIndex, const FVulkanTextureBase* TextureBase, VkImageLayout Layout)
+	inline void SetTextureForUBResource(int32 DescriptorSet, uint32 BindingIndex, const FVulkanTexture* Texture, VkImageLayout Layout)
 	{
-		CurrentState->SetTexture(DescriptorSet, BindingIndex, TextureBase, Layout);
+		CurrentState->SetTexture(DescriptorSet, BindingIndex, Texture, Layout);
 	}
 
 	void SetSRVForUBResource(uint32 DescriptorSet, uint32 BindingIndex, FVulkanShaderResourceView* SRV);
@@ -253,7 +253,7 @@ public:
 		CurrentPipeline->Bind(CmdBuffer);
 	}
 
-	inline void SetTextureForStage(ShaderStage::EStage Stage, uint32 ParameterIndex, const FVulkanTextureBase* TextureBase, VkImageLayout Layout)
+	inline void SetTextureForStage(ShaderStage::EStage Stage, uint32 ParameterIndex, const FVulkanTexture* Texture, VkImageLayout Layout)
 	{
 		const FVulkanGfxPipelineDescriptorInfo& DescriptorInfo = CurrentState->GetGfxPipelineDescriptorInfo();
 		uint8 DescriptorSet;
@@ -263,12 +263,12 @@ public:
 			return;
 		}
 
-		CurrentState->SetTexture(DescriptorSet, BindingIndex, TextureBase, Layout);
+		CurrentState->SetTexture(DescriptorSet, BindingIndex, Texture, Layout);
 	}
 
-	inline void SetTextureForUBResource(uint8 DescriptorSet, uint32 BindingIndex, const FVulkanTextureBase* TextureBase, VkImageLayout Layout)
+	inline void SetTextureForUBResource(uint8 DescriptorSet, uint32 BindingIndex, const FVulkanTexture* Texture, VkImageLayout Layout)
 	{
-		CurrentState->SetTexture(DescriptorSet, BindingIndex, TextureBase, Layout);
+		CurrentState->SetTexture(DescriptorSet, BindingIndex, Texture, Layout);
 	}
 
 	inline void SetUniformBufferConstantData(ShaderStage::EStage Stage, uint32 BindingIndex, const TArray<uint8>& ConstantData)
diff --git a/Engine/Source/Runtime/VulkanRHI/Private/VulkanPipelineState.h b/Engine/Source/Runtime/VulkanRHI/Private/VulkanPipelineState.h
--- a/Engine/Source/Runtime/VulkanRHI/Private/VulkanPipelineState.h
+++ b/Engine/Source/Runtime/VulkanRHI/Private/VulkanPipelineState.h
@@ -76,18 +76,18 @@ public:
 		MarkDirty(DSWriter[DescriptorSet].WriteStorageImage(BindingIndex, TextureView, Layout));
 	}
 
-	inline void SetTexture(uint8 DescriptorSet, uint32 BindingIndex, const FVulkanTextureBase* TextureBase, VkImageLayout Layout)
+	inline void SetTexture(uint8 DescriptorSet, uint32 BindingIndex, const FVulkanTexture* Texture, VkImageLayout Layout)
 	{
-		check(TextureBase && TextureBase->PartialView);
+		check(Texture && Texture->PartialView);
 
 		// If the texture doesn't support sampling, then we read it through a UAV
-		if (TextureBase->Surface.SupportsSampling())
+		if (Texture->SupportsSampling())
 		{
-			MarkDirty(DSWriter[DescriptorSet].WriteImage(BindingIndex, *TextureBase->PartialView, Layout));
+			MarkDirty(DSWriter[DescriptorSet].WriteImage(BindingIndex, *Texture->PartialView, Layout));
 		}
 		else
 		{
-			MarkDirty(DSWriter[DescriptorSet].WriteStorageImage(BindingIndex, *TextureBase->PartialView, Layout));
+			MarkDirty(DSWriter[DescriptorSet].WriteStorageImage(BindingIndex, *Texture->PartialView, Layout));
 		}
 	}
 
diff --git a/Engine/Source/Runtime/VulkanRHI/Private/VulkanRHI.cpp b/Engine/Source/Runtime/VulkanRHI/Private/VulkanRHI.cpp
--- a/Engine/Source/Runtime/VulkanRHI/Private/VulkanRHI.cpp
+++ b/Engine/Source/Runtime/VulkanRHI/Private/VulkanRHI.cpp
@@ -1315,46 +1315,48 @@ TArray<VkExtensionProperties> FVulkanDynamicRHI::RHIGetAllDeviceExtensions(VkPhy
 	return Extensions;
 }
 
-VkImage FVulkanDynamicRHI::RHIGetVkImage(FRHITexture2D* InTexture) const
+VkImage FVulkanDynamicRHI::RHIGetVkImage(FRHITexture* InTexture) const
 {
-	FVulkanTexture2D* VulkanTexture = ResourceCast(InTexture);
-	return VulkanTexture->Surface.Image;
+	FVulkanTexture* VulkanTexture = ResourceCast(InTexture);
+	return VulkanTexture->Image;
 }
 
-VkFormat FVulkanDynamicRHI::RHIGetViewVkFormat(FRHITexture2D* InTexture) const
+VkFormat FVulkanDynamicRHI::RHIGetViewVkFormat(FRHITexture* InTexture) const
 {
-	FVulkanTexture2D* VulkanTexture = ResourceCast(InTexture);
-	return VulkanTexture->Surface.ViewFormat;
+	FVulkanTexture* VulkanTexture = ResourceCast(InTexture);
+	return VulkanTexture->ViewFormat;
 }
 
-FVulkanRHIAllocationInfo FVulkanDynamicRHI::RHIGetAllocationInfo(FRHITexture2D* InTexture) const
+FVulkanRHIAllocationInfo FVulkanDynamicRHI::RHIGetAllocationInfo(FRHITexture* InTexture) const
 {
-	FVulkanTexture2D* VulkanTexture = ResourceCast(InTexture);
+	FVulkanTexture* VulkanTexture = ResourceCast(InTexture);
 
 	FVulkanRHIAllocationInfo NewInfo{};
-	NewInfo.Handle = VulkanTexture->Surface.GetAllocationHandle();
-	NewInfo.Offset = VulkanTexture->Surface.GetAllocationOffset();
-	NewInfo.Size = VulkanTexture->Surface.GetMemorySize();
+	NewInfo.Handle = VulkanTexture->GetAllocationHandle();
+	NewInfo.Offset = VulkanTexture->GetAllocationOffset();
+	NewInfo.Size = VulkanTexture->GetMemorySize();
 
 	return NewInfo;
 }
 
 FVulkanRHIImageViewInfo FVulkanDynamicRHI::RHIGetImageViewInfo(FRHITexture* InTexture) const
 {
-	FVulkanTextureBase* VulkanTexture = (FVulkanTextureBase*)InTexture->GetTextureBaseRHI();
+	FVulkanTexture* VulkanTexture = ResourceCast(InTexture);
+
+	const FRHITextureDesc& Desc = InTexture->GetDesc();
 
 	FVulkanRHIImageViewInfo Info{};
 	Info.ImageView = VulkanTexture->DefaultView.View;
 	Info.Image = VulkanTexture->DefaultView.Image;
-	Info.Format = VulkanTexture->Surface.ViewFormat;
-	Info.Width = VulkanTexture->Surface.Width;
-	Info.Height = VulkanTexture->Surface.Height;
-	Info.Depth = VulkanTexture->Surface.Depth;
-	Info.UEFlags = VulkanTexture->Surface.UEFlags;
+	Info.Format = VulkanTexture->ViewFormat;
+	Info.Width = Desc.Extent.X;
+	Info.Height = Desc.Extent.Y;
+	Info.Depth = Desc.Depth;
+	Info.UEFlags = Desc.Flags;
 
-	Info.SubresourceRange.aspectMask = VulkanTexture->Surface.GetFullAspectMask();
-	Info.SubresourceRange.layerCount = VulkanTexture->Surface.GetNumberOfArrayLevels();
-	Info.SubresourceRange.levelCount = VulkanTexture->Surface.GetNumMips();
+	Info.SubresourceRange.aspectMask = VulkanTexture->GetFullAspectMask();
+	Info.SubresourceRange.layerCount = VulkanTexture->GetNumberOfArrayLevels();
+	Info.SubresourceRange.levelCount = Desc.NumMips;
 
 	// TODO: do we need these?
 	Info.SubresourceRange.baseMipLevel = 0;
@@ -1363,10 +1365,10 @@ FVulkanRHIImageViewInfo FVulkanDynamicRHI::RHIGetImageViewInfo(FRHITexture* InTe
 	return Info;
 }
 
-VkImageLayout& FVulkanDynamicRHI::RHIFindOrAddLayoutRW(FRHITexture2D* InTexture, VkImageLayout LayoutIfNotFound)
+VkImageLayout& FVulkanDynamicRHI::RHIFindOrAddLayoutRW(FRHITexture* InTexture, VkImageLayout LayoutIfNotFound)
 {
-	FVulkanTexture2D* VulkanTexture = ResourceCast(InTexture);
-	return GetDevice()->GetImmediateContext().GetLayoutManager().FindOrAddLayoutRW(VulkanTexture->Surface, LayoutIfNotFound);
+	FVulkanTexture* VulkanTexture = ResourceCast(InTexture);
+	return GetDevice()->GetImmediateContext().GetLayoutManager().FindOrAddLayoutRW(*VulkanTexture, LayoutIfNotFound);
 }
 
 void FVulkanDynamicRHI::RHISetImageLayout(VkImage Image, VkImageLayout OldLayout, VkImageLayout NewLayout, const VkImageSubresourceRange& SubresourceRange)
@@ -1470,58 +1472,81 @@ void FVulkanDynamicRHI::RHISubmitCommandsAndFlushGPU()
 
 FTexture2DRHIRef FVulkanDynamicRHI::RHICreateTexture2DFromResource(EPixelFormat Format, uint32 SizeX, uint32 SizeY, uint32 NumMips, uint32 NumSamples, VkImage Resource, ETextureCreateFlags Flags)
 {
-	const FRHIResourceCreateInfo ResourceCreateInfo(TEXT("FVulkanTexture2D"), EnumHasAnyFlags(Flags, TexCreate_DepthStencilTargetable) ? FClearValueBinding::DepthZero : FClearValueBinding::Transparent);
-	return new FVulkanTexture2D(*Device, Format, SizeX, SizeY, NumMips, NumSamples, Resource, Flags, ResourceCreateInfo);
+	FRHITextureCreateDesc Desc = FRHITextureCreateDesc::Create2D(
+		TEXT("VulkanTexture2DFromResource"),
+		{ (int32)SizeX, (int32)SizeY },
+		Format,
+		EnumHasAnyFlags(Flags, TexCreate_DepthStencilTargetable) ? FClearValueBinding::DepthZero : FClearValueBinding::Transparent,
+		Flags,
+		NumMips,
+		NumSamples
+	);
+
+	return new FVulkanTexture(*Device, Desc, Resource, false);
 }
 
 FTexture2DArrayRHIRef FVulkanDynamicRHI::RHICreateTexture2DArrayFromResource(EPixelFormat Format, uint32 SizeX, uint32 SizeY, uint32 ArraySize, uint32 NumMips, uint32 NumSamples, VkImage Resource, ETextureCreateFlags Flags)
 {
-	const FRHIResourceCreateInfo ResourceCreateInfo(TEXT("FVulkanTexture2DArray"), EnumHasAnyFlags(Flags, TexCreate_DepthStencilTargetable) ? FClearValueBinding::DepthZero : FClearValueBinding::Transparent);
-	return new FVulkanTexture2DArray(*Device, Format, SizeX, SizeY, ArraySize, NumMips, NumSamples, Resource, Flags, ResourceCreateInfo);
+	FRHITextureCreateDesc Desc = FRHITextureCreateDesc::Create2DArray(
+		TEXT("VulkanTextureArrayFromResource"),
+		{ (int32)SizeX, (int32)SizeY },
+		Format,
+		EnumHasAnyFlags(Flags, TexCreate_DepthStencilTargetable) ? FClearValueBinding::DepthZero : FClearValueBinding::Transparent,
+		Flags,
+		ArraySize,
+		NumMips,
+		NumSamples
+	);
+
+	return new FVulkanTexture(*Device, Desc, Resource, false);
 }
 
 FTextureCubeRHIRef FVulkanDynamicRHI::RHICreateTextureCubeFromResource(EPixelFormat Format, uint32 Size, bool bArray, uint32 ArraySize, uint32 NumMips, VkImage Resource, ETextureCreateFlags Flags)
 {
-	const FRHIResourceCreateInfo ResourceCreateInfo(TEXT("FVulkanTextureCube"), EnumHasAnyFlags(Flags, TexCreate_DepthStencilTargetable) ? FClearValueBinding::DepthZero : FClearValueBinding::Transparent);
-	return new FVulkanTextureCube(*Device, Format, Size, bArray, ArraySize, NumMips, Resource, Flags, ResourceCreateInfo);
+	FRHITextureCreateDesc Desc;
+
+	if (ArraySize > 1)
+	{
+		Desc = FRHITextureCreateDesc::CreateCubeArray(
+			TEXT("VulkanTextureCubeFromResource"),
+			Size,
+			Format,
+			EnumHasAnyFlags(Flags, TexCreate_DepthStencilTargetable) ? FClearValueBinding::DepthZero : FClearValueBinding::Transparent,
+			Flags,
+			ArraySize,
+			NumMips,
+			1
+		);
+	}
+	else
+	{
+		Desc = FRHITextureCreateDesc::CreateCube(
+			TEXT("VulkanTextureCubeFromResource"),
+			Size,
+			Format,
+			EnumHasAnyFlags(Flags, TexCreate_DepthStencilTargetable) ? FClearValueBinding::DepthZero : FClearValueBinding::Transparent,
+			Flags,
+			NumMips,
+			1
+		);
+	}
+
+	return new FVulkanTexture(*Device, Desc, Resource, false);
 }
 
 void FVulkanDynamicRHI::RHIAliasTextureResources(FTextureRHIRef& DestTextureRHI, FTextureRHIRef& SrcTextureRHI)
 {
 	if (DestTextureRHI && SrcTextureRHI)
 	{
-		FVulkanTextureBase* DestTextureBase = (FVulkanTextureBase*)DestTextureRHI->GetTextureBaseRHI();
-		FVulkanTextureBase* SrcTextureBase = (FVulkanTextureBase*)SrcTextureRHI->GetTextureBaseRHI();
-
-		if (DestTextureBase && SrcTextureBase)
-		{
-			DestTextureBase->AliasTextureResources(SrcTextureRHI);
-		}
+		FVulkanTexture* DestTexture = FVulkanTexture::Cast(DestTextureRHI);
+		DestTexture->AliasTextureResources(SrcTextureRHI);
 	}
 }
 
 FTextureRHIRef FVulkanDynamicRHI::RHICreateAliasedTexture(FTextureRHIRef& SourceTextureRHI)
 {
-	FVulkanTextureBase* SourceTexture = (FVulkanTextureBase*)SourceTextureRHI->GetTextureBaseRHI();
-	FTextureRHIRef AliasedTexture;
-	if (SourceTextureRHI->GetTexture2D() != nullptr)
-	{
-		AliasedTexture = new FVulkanTexture2D(SourceTextureRHI, (FVulkanTexture2D*)SourceTexture);
-	}
-	else if (SourceTextureRHI->GetTexture2DArray() != nullptr)
-	{
-		AliasedTexture = new FVulkanTexture2DArray(SourceTextureRHI, (FVulkanTexture2DArray*)SourceTexture);
-	}
-	else if (SourceTextureRHI->GetTextureCube() != nullptr)
-	{
-		AliasedTexture = new FVulkanTextureCube(SourceTextureRHI, (FVulkanTextureCube*)SourceTexture);
-	}
-	else
-	{
-		UE_LOG(LogRHI, Error, TEXT("Currently FVulkanDynamicRHI::RHICreateAliasedTexture only supports 2D, 2D Array and Cube textures."));
-	}
-
-	return AliasedTexture;
+	FRHITextureCreateDesc Desc(SourceTextureRHI->GetDesc(), ERHIAccess::SRVMask, *SourceTextureRHI->GetName().ToString());
+	return new FVulkanTexture(*Device, Desc, SourceTextureRHI);
 }
 
 FVulkanDescriptorSetsLayout::FVulkanDescriptorSetsLayout(FVulkanDevice* InDevice) :
diff --git a/Engine/Source/Runtime/VulkanRHI/Private/VulkanRHIPrivate.h b/Engine/Source/Runtime/VulkanRHI/Private/VulkanRHIPrivate.h
--- a/Engine/Source/Runtime/VulkanRHI/Private/VulkanRHIPrivate.h
+++ b/Engine/Source/Runtime/VulkanRHI/Private/VulkanRHIPrivate.h
@@ -139,7 +139,6 @@ inline EShaderFrequency VkStageBitToUEFrequency(VkShaderStageFlagBits FlagBits)
 	return SF_NumFrequencies;
 }
 
-
 class FVulkanRenderTargetLayout
 {
 public:
@@ -285,8 +284,8 @@ public:
 	inline bool ContainsRenderTarget(FRHITexture* Texture) const
 	{
 		ensure(Texture);
-		FVulkanTextureBase* Base = (FVulkanTextureBase*)Texture->GetTextureBaseRHI();
-		return ContainsRenderTarget(Base->Surface.Image);
+		FVulkanTexture* VulkanTexture = FVulkanTexture::Cast(Texture);
+		return ContainsRenderTarget(VulkanTexture->Image);
 	}
 
 	inline bool ContainsRenderTarget(VkImage Image) const
diff --git a/Engine/Source/Runtime/VulkanRHI/Private/VulkanRenderTarget.cpp b/Engine/Source/Runtime/VulkanRHI/Private/VulkanRenderTarget.cpp
--- a/Engine/Source/Runtime/VulkanRHI/Private/VulkanRenderTarget.cpp
+++ b/Engine/Source/Runtime/VulkanRHI/Private/VulkanRenderTarget.cpp
@@ -46,7 +46,7 @@ static FAutoConsoleVariableRef CVarVulkanIgnoreCPUReads(
 	);
 
 static FCriticalSection GStagingMapLock;
-static TMap<FVulkanTextureBase*, VulkanRHI::FStagingBuffer*> GPendingLockedStagingBuffers;
+static TMap<FVulkanTexture*, VulkanRHI::FStagingBuffer*> GPendingLockedStagingBuffers;
 
 #if UE_BUILD_DEBUG || UE_BUILD_DEVELOPMENT
 TAutoConsoleVariable<int32> CVarVulkanDebugBarrier(
@@ -78,16 +78,19 @@ FVulkanRenderPass* FVulkanCommandListContext::PrepareRenderPassForPSOCreation(co
 }
 
 template<typename RegionType>
-static void SetupCopyOrResolveRegion(RegionType& Region, const FVulkanSurface& SrcSurface, const FVulkanSurface& DstSurface, const VkImageSubresourceRange& SrcRange, const VkImageSubresourceRange& DstRange, const FResolveParams& ResolveParams)
+static void SetupCopyOrResolveRegion(RegionType& Region, const FVulkanTexture& SrcSurface, const FVulkanTexture& DstSurface, const VkImageSubresourceRange& SrcRange, const VkImageSubresourceRange& DstRange, const FResolveParams& ResolveParams)
 {
 	FMemory::Memzero(Region);
 
-	bool bCopySrcSubRect = ResolveParams.Rect.IsValid()     && (ResolveParams.Rect.X1 != 0     || ResolveParams.Rect.Y1 != 0     || ResolveParams.Rect.X2 != SrcSurface.Width     || ResolveParams.Rect.Y2 != SrcSurface.Height);
-	bool bCopyDstSubRect = ResolveParams.DestRect.IsValid() && (ResolveParams.DestRect.X1 != 0 || ResolveParams.DestRect.Y1 != 0 || ResolveParams.DestRect.X2 != DstSurface.Width || ResolveParams.DestRect.Y2 != DstSurface.Height);
+	const FRHITextureDesc& SrcDesc = SrcSurface.GetDesc();
+	const FRHITextureDesc& DstDesc = DstSurface.GetDesc();
+
+	bool bCopySrcSubRect = ResolveParams.Rect.IsValid()     && (ResolveParams.Rect.X1 != 0     || ResolveParams.Rect.Y1 != 0     || ResolveParams.Rect.X2 != SrcDesc.Extent.X || ResolveParams.Rect.Y2 != SrcDesc.Extent.Y);
+	bool bCopyDstSubRect = ResolveParams.DestRect.IsValid() && (ResolveParams.DestRect.X1 != 0 || ResolveParams.DestRect.Y1 != 0 || ResolveParams.DestRect.X2 != DstDesc.Extent.X || ResolveParams.DestRect.Y2 != DstDesc.Extent.Y);
 	if(bCopySrcSubRect || bCopyDstSubRect)
 	{
-		const FResolveRect& SrcSubRect = ResolveParams.Rect.IsValid()     ? ResolveParams.Rect     : FResolveRect(0, 0, SrcSurface.Width, SrcSurface.Height);
-		const FResolveRect& DstSubRect = ResolveParams.DestRect.IsValid() ? ResolveParams.DestRect : FResolveRect(0, 0, DstSurface.Width, DstSurface.Height);
+		const FResolveRect& SrcSubRect = ResolveParams.Rect.IsValid()     ? ResolveParams.Rect     : FResolveRect(0, 0, SrcDesc.Extent.X, SrcDesc.Extent.Y);
+		const FResolveRect& DstSubRect = ResolveParams.DestRect.IsValid() ? ResolveParams.DestRect : FResolveRect(0, 0, DstDesc.Extent.X, DstDesc.Extent.Y);
 
 		const uint32 SrcSubRectWidth  = SrcSubRect.X2 - SrcSubRect.X1;
 		const uint32 SrcSubRectHeight = SrcSubRect.Y2 - SrcSubRect.Y1;
@@ -107,10 +110,10 @@ static void SetupCopyOrResolveRegion(RegionType& Region, const FVulkanSurface& S
 	}
 	else
 	{
-		ensure(SrcSurface.Width == DstSurface.Width && SrcSurface.Height == DstSurface.Height);
+		ensure(SrcDesc.Extent.X == DstDesc.Extent.X && SrcDesc.Extent.Y == DstDesc.Extent.Y);
 
-		Region.extent.width = FMath::Max(1u, SrcSurface.Width >> ResolveParams.MipIndex);
-		Region.extent.height = FMath::Max(1u, SrcSurface.Height >> ResolveParams.MipIndex);
+		Region.extent.width = FMath::Max(1, SrcDesc.Extent.X >> ResolveParams.MipIndex);
+		Region.extent.height = FMath::Max(1, SrcDesc.Extent.Y >> ResolveParams.MipIndex);
 	}
 
 	Region.extent.depth = 1;
@@ -135,69 +138,61 @@ void FVulkanCommandListContext::RHICopyToResolveTarget(FRHITexture* SourceTextur
 	FVulkanCmdBuffer* CmdBuffer = CommandBufferManager->GetActiveCmdBuffer();
 	check(CmdBuffer->IsOutsideRenderPass());
 
-	FRHITexture2D* SourceTexture2D = SourceTextureRHI->GetTexture2D();
-	FRHITexture2DArray* SourceTexture2DArray = SourceTextureRHI->GetTexture2DArray();
-	FRHITexture3D* SourceTexture3D = SourceTextureRHI->GetTexture3D();
-	FRHITextureCube* SourceTextureCube = SourceTextureRHI->GetTextureCube();
-	FRHITexture2D* DestTexture2D = DestTextureRHI->GetTexture2D();
-	FRHITexture2DArray* DestTexture2DArray = DestTextureRHI->GetTexture2DArray();
-	FRHITexture3D* DestTexture3D = DestTextureRHI->GetTexture3D();
-	FRHITextureCube* DestTextureCube = DestTextureRHI->GetTextureCube();
-
-	const FVulkanSurface* SrcSurface = nullptr;
-	const FVulkanSurface* DstSurface = nullptr;
+	const FVulkanTexture& SrcSurface = *FVulkanTexture::Cast(SourceTextureRHI);
+	const FVulkanTexture& DstSurface = *FVulkanTexture::Cast(DestTextureRHI);
+
+	const FRHITextureDesc& SrcDesc = SrcSurface.GetDesc();
+	const FRHITextureDesc& DstDesc = DstSurface.GetDesc();
+
 	uint32 SrcNumLayers, DstNumLayers;
 
-	if (SourceTexture2D && DestTexture2D)
+	switch (SourceTextureRHI->GetDesc().Dimension)
 	{
-		SrcSurface = &((FVulkanTexture2D*)SourceTexture2D)->Surface;
-		DstSurface = &((FVulkanTexture2D*)DestTexture2D)->Surface;
+	case ETextureDimension::Texture2D:
+		if (DestTextureRHI->GetDesc().Dimension == ETextureDimension::TextureCube)
+		{
+			SrcNumLayers = 1;
+			DstNumLayers = 6;
+		}
+		else
+		{
+			check(DestTextureRHI->GetDesc().Dimension == ETextureDimension::Texture2D);
+			SrcNumLayers = DstNumLayers = 1;
+		}
+		break;
+
+	case ETextureDimension::Texture2DArray:
+		check(DestTextureRHI->GetDesc().Dimension == ETextureDimension::Texture2DArray);
+		SrcNumLayers = SrcDesc.ArraySize;
+		DstNumLayers = DstDesc.ArraySize;
+		break;
+
+	case ETextureDimension::Texture3D:
+		check(DestTextureRHI->GetDesc().Dimension == ETextureDimension::Texture3D);
 		SrcNumLayers = DstNumLayers = 1;
-	}
-	else if (SourceTextureCube && DestTextureCube) 
-	{
-		SrcSurface = &((FVulkanTextureCube*)SourceTextureCube)->Surface;
-		DstSurface = &((FVulkanTextureCube*)DestTextureCube)->Surface;
+		break;
+
+	case ETextureDimension::TextureCube:
+		check(DestTextureRHI->GetDesc().Dimension == ETextureDimension::TextureCube);
 		SrcNumLayers = DstNumLayers = 6;
-	}
-	else if (SourceTexture2D && DestTextureCube) 
-	{
-		SrcSurface = &((FVulkanTexture2D*)SourceTexture2D)->Surface;
-		DstSurface = &((FVulkanTextureCube*)DestTextureCube)->Surface;
-		SrcNumLayers = 1;
-		DstNumLayers = 6;
-	}
-	else if (SourceTexture3D && DestTexture3D) 
-	{
-		SrcSurface = &((FVulkanTexture3D*)SourceTexture3D)->Surface;
-		DstSurface = &((FVulkanTexture3D*)DestTexture3D)->Surface;
-		SrcNumLayers = DstNumLayers = 1;
-	}
-	else if (SourceTexture2DArray && DestTexture2DArray)
-	{
-		FVulkanTexture2DArray* VulkanSrcTexture = (FVulkanTexture2DArray*)SourceTexture2DArray;
-		SrcSurface = &VulkanSrcTexture->Surface;
-		SrcNumLayers = VulkanSrcTexture->GetSizeZ();
-
-		FVulkanTexture2DArray* VulkanDstTexture = (FVulkanTexture2DArray*)DestTexture2DArray;
-		DstSurface = &VulkanDstTexture->Surface;
-		DstNumLayers = VulkanDstTexture->GetSizeZ();
-	} 
-	else 
-	{
+		break;
+
+	case ETextureDimension::TextureCubeArray:
+		// Not supported. fall through.
+	default:
 		checkNoEntry();
 		return;
 	}
 
 	VkImageSubresourceRange SrcRange;
-	SrcRange.aspectMask = SrcSurface->GetFullAspectMask();
+	SrcRange.aspectMask = SrcSurface.GetFullAspectMask();
 	SrcRange.baseMipLevel = InResolveParams.MipIndex;
 	SrcRange.levelCount = 1;
 	SrcRange.baseArrayLayer = InResolveParams.SourceArrayIndex * SrcNumLayers + (SrcNumLayers == 6 ? InResolveParams.CubeFace : 0);
 	SrcRange.layerCount = 1;
 
 	VkImageSubresourceRange DstRange;
-	DstRange.aspectMask = DstSurface->GetFullAspectMask();
+	DstRange.aspectMask = DstSurface.GetFullAspectMask();
 	DstRange.baseMipLevel = InResolveParams.MipIndex;
 	DstRange.levelCount = 1;
 	DstRange.baseArrayLayer = InResolveParams.DestArrayIndex * DstNumLayers + (DstNumLayers == 6 ? InResolveParams.CubeFace : 0);
@@ -205,38 +200,38 @@ void FVulkanCommandListContext::RHICopyToResolveTarget(FRHITexture* SourceTextur
 
 	ERHIAccess SrcCurrentAccess, DstCurrentAccess;
 
-	check(!EnumHasAnyFlags(SrcSurface->UEFlags, TexCreate_CPUReadback));
-	VkImageLayout& SrcLayout = LayoutManager.FindOrAddLayoutRW(*SrcSurface, VK_IMAGE_LAYOUT_UNDEFINED);
+	check(!EnumHasAnyFlags(SrcSurface.GetDesc().Flags, TexCreate_CPUReadback));
+	VkImageLayout& SrcLayout = LayoutManager.FindOrAddLayoutRW(SrcSurface, VK_IMAGE_LAYOUT_UNDEFINED);
 
-	if(EnumHasAnyFlags(DstSurface->UEFlags, TexCreate_CPUReadback))
+	if(EnumHasAnyFlags(DstSurface.GetDesc().Flags, TexCreate_CPUReadback))
 	{
 		//Readback textures are represented as a buffer, so we can support miplevels on hardware that does not expose it.
 		FVulkanPipelineBarrier BarrierBefore;
 		// We'll transition the entire resources to the correct copy states, so we don't need to worry about sub-resource states.
 		if (SrcLayout != VK_IMAGE_LAYOUT_TRANSFER_SRC_OPTIMAL)
 		{
-			BarrierBefore.AddImageLayoutTransition(SrcSurface->Image, SrcLayout, VK_IMAGE_LAYOUT_TRANSFER_SRC_OPTIMAL, FVulkanPipelineBarrier::MakeSubresourceRange(SrcSurface->GetFullAspectMask()));
+			BarrierBefore.AddImageLayoutTransition(SrcSurface.Image, SrcLayout, VK_IMAGE_LAYOUT_TRANSFER_SRC_OPTIMAL, FVulkanPipelineBarrier::MakeSubresourceRange(SrcSurface.GetFullAspectMask()));
 			SrcLayout = VK_IMAGE_LAYOUT_TRANSFER_SRC_OPTIMAL;
 		}
 		BarrierBefore.Execute(CmdBuffer->GetHandle());
-		const FVulkanCpuReadbackBuffer* CpuReadbackBuffer = DstSurface->GetCpuReadbackBuffer();
+		const FVulkanCpuReadbackBuffer* CpuReadbackBuffer = DstSurface.GetCpuReadbackBuffer();
 		check(DstRange.baseArrayLayer == 0);
 		VkBufferImageCopy CopyRegion;
 		FMemory::Memzero(CopyRegion);
 		uint32 MipLevel = InResolveParams.MipIndex;
-		uint32 SizeX = FMath::Max(SrcSurface->Width >> MipLevel, 1u);
-		uint32 SizeY = FMath::Max(SrcSurface->Height >> MipLevel, 1u);
+		uint32 SizeX = FMath::Max(SrcDesc.Extent.X >> MipLevel, 1);
+		uint32 SizeY = FMath::Max(SrcDesc.Extent.Y >> MipLevel, 1);
 		CopyRegion.bufferOffset = CpuReadbackBuffer->MipOffsets[MipLevel];
 		CopyRegion.bufferRowLength = SizeX;
 		CopyRegion.bufferImageHeight = SizeY;
-		CopyRegion.imageSubresource.aspectMask = SrcSurface->GetFullAspectMask();
+		CopyRegion.imageSubresource.aspectMask = SrcSurface.GetFullAspectMask();
 		CopyRegion.imageSubresource.mipLevel = MipLevel;
 		CopyRegion.imageSubresource.baseArrayLayer = 0;
 		CopyRegion.imageSubresource.layerCount = 1;
 		CopyRegion.imageExtent.width = SizeX;
 		CopyRegion.imageExtent.height = SizeY;
 		CopyRegion.imageExtent.depth = 1;
-		VulkanRHI::vkCmdCopyImageToBuffer(CmdBuffer->GetHandle(), SrcSurface->Image, VK_IMAGE_LAYOUT_TRANSFER_SRC_OPTIMAL, CpuReadbackBuffer->Buffer, 1, &CopyRegion);
+		VulkanRHI::vkCmdCopyImageToBuffer(CmdBuffer->GetHandle(), SrcSurface.Image, VK_IMAGE_LAYOUT_TRANSFER_SRC_OPTIMAL, CpuReadbackBuffer->Buffer, 1, &CopyRegion);
 
 		{
 			FVulkanPipelineBarrier BarrierMemory;
@@ -247,31 +242,31 @@ void FVulkanCommandListContext::RHICopyToResolveTarget(FRHITexture* SourceTextur
 		if(SrcCurrentAccess != InResolveParams.SourceAccessFinal && InResolveParams.SourceAccessFinal != ERHIAccess::Unknown)
 		{
 			FVulkanPipelineBarrier BarrierAfter;
-			BarrierAfter.AddImageAccessTransition(*SrcSurface, SrcCurrentAccess, InResolveParams.SourceAccessFinal, FVulkanPipelineBarrier::MakeSubresourceRange(SrcSurface->GetFullAspectMask()), SrcLayout);
+			BarrierAfter.AddImageAccessTransition(SrcSurface, SrcCurrentAccess, InResolveParams.SourceAccessFinal, FVulkanPipelineBarrier::MakeSubresourceRange(SrcSurface.GetFullAspectMask()), SrcLayout);
 			BarrierAfter.Execute(CmdBuffer->GetHandle());
 
 		}
 	}
 	else
 	{
-		VkImageLayout& DstLayout = LayoutManager.FindOrAddLayoutRW(*DstSurface, VK_IMAGE_LAYOUT_UNDEFINED);
-		if (SrcSurface->Image != DstSurface->Image)
+		VkImageLayout& DstLayout = LayoutManager.FindOrAddLayoutRW(DstSurface, VK_IMAGE_LAYOUT_UNDEFINED);
+		if (SrcSurface.Image != DstSurface.Image)
 		{
-			const bool bIsResolve = SrcSurface->GetNumSamples() > DstSurface->GetNumSamples();
-			checkf(!bIsResolve || !DstSurface->IsDepthOrStencilAspect(), TEXT("Vulkan does not support multisample depth resolve."));
+			const bool bIsResolve = SrcSurface.GetNumSamples() > DstSurface.GetNumSamples();
+			checkf(!bIsResolve || !DstSurface.IsDepthOrStencilAspect(), TEXT("Vulkan does not support multisample depth resolve."));
 
 			FVulkanPipelineBarrier BarrierBefore;
 		
 			// We'll transition the entire resources to the correct copy states, so we don't need to worry about sub-resource states.
 			if (SrcLayout != VK_IMAGE_LAYOUT_TRANSFER_SRC_OPTIMAL)
 			{
-				BarrierBefore.AddImageLayoutTransition(SrcSurface->Image, SrcLayout, VK_IMAGE_LAYOUT_TRANSFER_SRC_OPTIMAL, FVulkanPipelineBarrier::MakeSubresourceRange(SrcSurface->GetFullAspectMask()));
+				BarrierBefore.AddImageLayoutTransition(SrcSurface.Image, SrcLayout, VK_IMAGE_LAYOUT_TRANSFER_SRC_OPTIMAL, FVulkanPipelineBarrier::MakeSubresourceRange(SrcSurface.GetFullAspectMask()));
 				SrcLayout = VK_IMAGE_LAYOUT_TRANSFER_SRC_OPTIMAL;
 			}
 
 			if (DstLayout != VK_IMAGE_LAYOUT_TRANSFER_DST_OPTIMAL)
 			{
-				BarrierBefore.AddImageLayoutTransition(DstSurface->Image, DstLayout, VK_IMAGE_LAYOUT_TRANSFER_DST_OPTIMAL, FVulkanPipelineBarrier::MakeSubresourceRange(DstSurface->GetFullAspectMask()));
+				BarrierBefore.AddImageLayoutTransition(DstSurface.Image, DstLayout, VK_IMAGE_LAYOUT_TRANSFER_DST_OPTIMAL, FVulkanPipelineBarrier::MakeSubresourceRange(DstSurface.GetFullAspectMask()));
 				DstLayout = VK_IMAGE_LAYOUT_TRANSFER_DST_OPTIMAL;
 			}
 
@@ -280,20 +275,20 @@ void FVulkanCommandListContext::RHICopyToResolveTarget(FRHITexture* SourceTextur
 			if (!bIsResolve)
 			{
 				VkImageCopy Region;
-				SetupCopyOrResolveRegion(Region, *SrcSurface, *DstSurface, SrcRange, DstRange, InResolveParams);
+				SetupCopyOrResolveRegion(Region, SrcSurface, DstSurface, SrcRange, DstRange, InResolveParams);
 				VulkanRHI::vkCmdCopyImage(CmdBuffer->GetHandle(),
-					SrcSurface->Image, VK_IMAGE_LAYOUT_TRANSFER_SRC_OPTIMAL,
-					DstSurface->Image, VK_IMAGE_LAYOUT_TRANSFER_DST_OPTIMAL,
+					SrcSurface.Image, VK_IMAGE_LAYOUT_TRANSFER_SRC_OPTIMAL,
+					DstSurface.Image, VK_IMAGE_LAYOUT_TRANSFER_DST_OPTIMAL,
 					1, &Region);
 			}
 			else
 			{
-				check(DstSurface->GetNumSamples() == 1);
+				check(DstSurface.GetNumSamples() == 1);
 				VkImageResolve Region;
-				SetupCopyOrResolveRegion(Region, *SrcSurface, *DstSurface, SrcRange, DstRange, InResolveParams);
+				SetupCopyOrResolveRegion(Region, SrcSurface, DstSurface, SrcRange, DstRange, InResolveParams);
 				VulkanRHI::vkCmdResolveImage(CmdBuffer->GetHandle(),
-					SrcSurface->Image, VK_IMAGE_LAYOUT_TRANSFER_SRC_OPTIMAL,
-					DstSurface->Image, VK_IMAGE_LAYOUT_TRANSFER_DST_OPTIMAL,
+					SrcSurface.Image, VK_IMAGE_LAYOUT_TRANSFER_SRC_OPTIMAL,
+					DstSurface.Image, VK_IMAGE_LAYOUT_TRANSFER_DST_OPTIMAL,
 					1, &Region);
 			}
 
@@ -309,14 +304,14 @@ void FVulkanCommandListContext::RHICopyToResolveTarget(FRHITexture* SourceTextur
 		if (InResolveParams.SourceAccessFinal != ERHIAccess::Unknown && InResolveParams.DestAccessFinal != ERHIAccess::Unknown)
 		{
 			FVulkanPipelineBarrier BarrierAfter;
-			if (SrcSurface->Image != DstSurface->Image && SrcCurrentAccess != InResolveParams.SourceAccessFinal && InResolveParams.SourceAccessFinal != ERHIAccess::Unknown)
+			if (SrcSurface.Image != DstSurface.Image && SrcCurrentAccess != InResolveParams.SourceAccessFinal && InResolveParams.SourceAccessFinal != ERHIAccess::Unknown)
 			{
-				BarrierAfter.AddImageAccessTransition(*SrcSurface, SrcCurrentAccess, InResolveParams.SourceAccessFinal, FVulkanPipelineBarrier::MakeSubresourceRange(SrcSurface->GetFullAspectMask()), SrcLayout);
+				BarrierAfter.AddImageAccessTransition(SrcSurface, SrcCurrentAccess, InResolveParams.SourceAccessFinal, FVulkanPipelineBarrier::MakeSubresourceRange(SrcSurface.GetFullAspectMask()), SrcLayout);
 			}
 
 			if (DstCurrentAccess != InResolveParams.DestAccessFinal && InResolveParams.DestAccessFinal != ERHIAccess::Unknown)
 			{
-				BarrierAfter.AddImageAccessTransition(*DstSurface, DstCurrentAccess, InResolveParams.DestAccessFinal, FVulkanPipelineBarrier::MakeSubresourceRange(DstSurface->GetFullAspectMask()), DstLayout);
+				BarrierAfter.AddImageAccessTransition(DstSurface, DstCurrentAccess, InResolveParams.DestAccessFinal, FVulkanPipelineBarrier::MakeSubresourceRange(DstSurface.GetFullAspectMask()), DstLayout);
 			}
 
 			BarrierAfter.Execute(CmdBuffer->GetHandle());
@@ -391,53 +386,43 @@ void FVulkanDynamicRHI::RHIReadSurfaceData(FRHITexture* TextureRHI, FIntRect Rec
 		return;
 	}
 
-	FVulkanSurface* Surface = nullptr;
-
-	FRHITexture2D* Texture2DRHI = TextureRHI->GetTexture2D();
-	if (Texture2DRHI)
-	{
-		FVulkanTexture2D* Texture2D = (FVulkanTexture2D*)Texture2DRHI;
-		Surface = &Texture2D->Surface;
-	}
-	else
+	const FRHITextureDesc& Desc = TextureRHI->GetDesc();
+	switch (Desc.Dimension)
 	{
+	case ETextureDimension::Texture2D:
+	case ETextureDimension::Texture2DArray:
 		// In VR, the high level code calls this function on the viewport render target, without knowing that it's
 		// actually a texture array created and managed by the VR runtime. In that case we'll just read the first
 		// slice of the array, which corresponds to one of the eyes.
-		FRHITexture2DArray* Texture2DArrayRHI = TextureRHI->GetTexture2DArray();
-		if (Texture2DArrayRHI)
-		{
-			FVulkanTexture2DArray* Texture2DArray = (FVulkanTexture2DArray*)Texture2DArrayRHI;
-			Surface = &Texture2DArray->Surface;
-		}
-	}
+		break;
 
-	if (!Surface)
-	{
+	default:
 		// Just return black for texture types we don't support.
 		OutData.Empty(0);
 		OutData.AddZeroed(NumRequestedPixels);
 		return;
 	}
 
+	FVulkanTexture& Surface = *FVulkanTexture::Cast(TextureRHI);
+
 	Device->PrepareForCPURead();
 
 	// Figure out the size of the buffer required to hold the requested pixels
-	const uint32 PixelByteSize = GetNumBitsPerPixel(Surface->StorageFormat) / 8;
+	const uint32 PixelByteSize = GetNumBitsPerPixel(Surface.StorageFormat) / 8;
 	checkf(GPixelFormats[TextureRHI->GetFormat()].Supported && (PixelByteSize > 0), TEXT("Trying to read from unsupported format."));
 	const uint32 BufferSize = NumRequestedPixels * PixelByteSize;
 
 	// Validate that the Rect is within the texture
 	const uint32 MipLevel = InFlags.GetMip();
-	const uint32 MipSizeX = FMath::Max(Surface->Width >> MipLevel, 1u);
-	const uint32 MipSizeY = FMath::Max(Surface->Height >> MipLevel, 1u);
+	const uint32 MipSizeX = FMath::Max(Desc.Extent.X >> MipLevel, 1);
+	const uint32 MipSizeY = FMath::Max(Desc.Extent.Y >> MipLevel, 1);
 	checkf((Rect.Max.X <= (int32)MipSizeX) && (Rect.Max.Y <= (int32)MipSizeY), TEXT("The specified Rect [%dx%d] extends beyond this Mip [%dx%d]."), Rect.Max.X, Rect.Max.Y, MipSizeX, MipSizeY);
 
 	FVulkanCommandListContext& ImmediateContext = Device->GetImmediateContext();
 
 	VulkanRHI::FStagingBuffer* StagingBuffer = nullptr;
 	FVulkanCmdBuffer* CmdBuffer = nullptr;
-	const bool bCPUReadback = EnumHasAllFlags(Surface->UEFlags, TexCreate_CPUReadback);
+	const bool bCPUReadback = EnumHasAllFlags(Surface.GetDesc().Flags, TexCreate_CPUReadback);
 	if (!bCPUReadback) //this function supports reading back arbitrary rendertargets, so if its not a cpu readback surface, we do a copy.
 	{
 		CmdBuffer = ImmediateContext.GetCommandBufferManager()->GetUploadCmdBuffer();
@@ -445,7 +430,7 @@ void FVulkanDynamicRHI::RHIReadSurfaceData(FRHITexture* TextureRHI, FIntRect Rec
 		VkBufferImageCopy CopyRegion;
 		FMemory::Memzero(CopyRegion);
 		// Leave bufferRowLength/bufferImageHeight at 0 for tightly packed
-		CopyRegion.imageSubresource.aspectMask = Surface->GetFullAspectMask();
+		CopyRegion.imageSubresource.aspectMask = Surface.GetFullAspectMask();
 		CopyRegion.imageSubresource.mipLevel = MipLevel;
 		CopyRegion.imageSubresource.layerCount = 1;
 		CopyRegion.imageOffset.x = Rect.Min.X;
@@ -454,17 +439,17 @@ void FVulkanDynamicRHI::RHIReadSurfaceData(FRHITexture* TextureRHI, FIntRect Rec
 		CopyRegion.imageExtent.height = DestHeight;
 		CopyRegion.imageExtent.depth = 1;
 
-		VkImageLayout& CurrentLayout = Device->GetImmediateContext().GetLayoutManager().FindOrAddLayoutRW(*Surface, VK_IMAGE_LAYOUT_UNDEFINED);
+		VkImageLayout& CurrentLayout = Device->GetImmediateContext().GetLayoutManager().FindOrAddLayoutRW(Surface, VK_IMAGE_LAYOUT_UNDEFINED);
 		const bool bHadLayout = (CurrentLayout != VK_IMAGE_LAYOUT_UNDEFINED);
 		if (CurrentLayout != VK_IMAGE_LAYOUT_TRANSFER_SRC_OPTIMAL)
 		{
-			VulkanSetImageLayoutAllMips(CmdBuffer->GetHandle(), Surface->Image, CurrentLayout, VK_IMAGE_LAYOUT_TRANSFER_SRC_OPTIMAL);
+			VulkanSetImageLayoutAllMips(CmdBuffer->GetHandle(), Surface.Image, CurrentLayout, VK_IMAGE_LAYOUT_TRANSFER_SRC_OPTIMAL);
 		}
 
-		VulkanRHI::vkCmdCopyImageToBuffer(CmdBuffer->GetHandle(), Surface->Image, VK_IMAGE_LAYOUT_TRANSFER_SRC_OPTIMAL, StagingBuffer->GetHandle(), 1, &CopyRegion);
+		VulkanRHI::vkCmdCopyImageToBuffer(CmdBuffer->GetHandle(), Surface.Image, VK_IMAGE_LAYOUT_TRANSFER_SRC_OPTIMAL, StagingBuffer->GetHandle(), 1, &CopyRegion);
 		if (bHadLayout && CurrentLayout != VK_IMAGE_LAYOUT_TRANSFER_SRC_OPTIMAL)
 		{
-			VulkanSetImageLayoutAllMips(CmdBuffer->GetHandle(), Surface->Image, VK_IMAGE_LAYOUT_TRANSFER_SRC_OPTIMAL, CurrentLayout);
+			VulkanSetImageLayoutAllMips(CmdBuffer->GetHandle(), Surface.Image, VK_IMAGE_LAYOUT_TRANSFER_SRC_OPTIMAL, CurrentLayout);
 		}
 		else
 		{
@@ -486,7 +471,7 @@ void FVulkanDynamicRHI::RHIReadSurfaceData(FRHITexture* TextureRHI, FIntRect Rec
 	if (bCPUReadback)
 	{
 		// If the text was bCPUReadback, then we have to deal with our Rect potentially being a subset of the total texture
-		In = (uint8*)Surface->GetMappedPointer() + ((Rect.Min.Y * MipSizeX + Rect.Min.X) * PixelByteSize);
+		In = (uint8*)Surface.GetMappedPointer() + ((Rect.Min.Y * MipSizeX + Rect.Min.X) * PixelByteSize);
 		SrcPitch = MipSizeX * PixelByteSize;
 	}
 	else
@@ -499,7 +484,7 @@ void FVulkanDynamicRHI::RHIReadSurfaceData(FRHITexture* TextureRHI, FIntRect Rec
 
 	OutData.SetNum(NumRequestedPixels);
 	FColor* Dest = OutData.GetData();
-	ConvertRawDataToFColor(Surface->StorageFormat, DestWidth, DestHeight, In, SrcPitch, Dest, InFlags);
+	ConvertRawDataToFColor(Surface.StorageFormat, DestWidth, DestHeight, In, SrcPitch, Dest, InFlags);
 
 	if (!bCPUReadback)
 	{
@@ -521,9 +506,8 @@ void FVulkanDynamicRHI::RHIReadSurfaceData(FRHITexture* TextureRHI, FIntRect Rec
 
 void FVulkanDynamicRHI::RHIMapStagingSurface(FRHITexture* TextureRHI, FRHIGPUFence* FenceRHI, void*& OutData, int32& OutWidth, int32& OutHeight, uint32 GPUIndex)
 {
-	FRHITexture2D* TextureRHI2D = TextureRHI->GetTexture2D();
-	check(TextureRHI2D);
-	FVulkanTexture2D* Texture2D = ResourceCast(TextureRHI2D);
+	check(TextureRHI->GetDesc().Dimension == ETextureDimension::Texture2D);
+	FVulkanTexture* Texture = FVulkanTexture::Cast(TextureRHI);
 
 	if (FenceRHI && !FenceRHI->Poll())
 	{
@@ -542,24 +526,26 @@ void FVulkanDynamicRHI::RHIMapStagingSurface(FRHITexture* TextureRHI, FRHIGPUFen
 	}
 
 
-	check(EnumHasAllFlags(Texture2D->Surface.UEFlags, TexCreate_CPUReadback));
-	OutData = Texture2D->Surface.GetMappedPointer();
-	Texture2D->Surface.InvalidateMappedMemory();
-	OutWidth = Texture2D->GetSizeX();
-	OutHeight = Texture2D->GetSizeY();
+	check(EnumHasAllFlags(Texture->GetDesc().Flags, TexCreate_CPUReadback));
+	OutData = Texture->GetMappedPointer();
+	Texture->InvalidateMappedMemory();
+	OutWidth = Texture->GetSizeX();
+	OutHeight = Texture->GetSizeY();
 }
 
 void FVulkanDynamicRHI::RHIUnmapStagingSurface(FRHITexture* TextureRHI, uint32 GPUIndex)
 {
 }
 
-void FVulkanDynamicRHI::RHIReadSurfaceFloatData(FRHITexture* TextureRHI, FIntRect Rect, TArray<FFloat16Color>& OutData, ECubeFace CubeFace,int32 ArrayIndex,int32 MipIndex)
+void FVulkanDynamicRHI::RHIReadSurfaceFloatData(FRHITexture* TextureRHI, FIntRect Rect, TArray<FFloat16Color>& OutData, ECubeFace CubeFace, int32 ArrayIndex, int32 MipIndex)
 {
-	auto DoCopyFloat = [](FVulkanDevice* InDevice, FVulkanCmdBuffer* InCmdBuffer, const FVulkanSurface& Surface, uint32 InMipIndex, uint32 SrcBaseArrayLayer, FIntRect InRect, TArray<FFloat16Color>& OutputData)
+	auto DoCopyFloat = [](FVulkanDevice* InDevice, FVulkanCmdBuffer* InCmdBuffer, const FVulkanTexture& Surface, uint32 InMipIndex, uint32 SrcBaseArrayLayer, FIntRect InRect, TArray<FFloat16Color>& OutputData)
 	{
 		ensure(Surface.StorageFormat == VK_FORMAT_R16G16B16A16_SFLOAT);
 
-		uint32 NumPixels = (Surface.Width >> InMipIndex) * (Surface.Height >> InMipIndex);
+		const FRHITextureDesc& Desc = Surface.GetDesc();
+
+		uint32 NumPixels = (Desc.Extent.X >> InMipIndex) * (Desc.Extent.Y >> InMipIndex);
 		const uint32 Size = NumPixels * sizeof(FFloat16Color);
 		VulkanRHI::FStagingBuffer* StagingBuffer = InDevice->GetStagingManager().AcquireBuffer(Size, VK_BUFFER_USAGE_TRANSFER_DST_BIT, VK_MEMORY_PROPERTY_HOST_CACHED_BIT);
 
@@ -568,14 +554,14 @@ void FVulkanDynamicRHI::RHIReadSurfaceFloatData(FRHITexture* TextureRHI, FIntRec
 			VkBufferImageCopy CopyRegion;
 			FMemory::Memzero(CopyRegion);
 			//Region.bufferOffset = 0;
-			CopyRegion.bufferRowLength = Surface.Width >> InMipIndex;
-			CopyRegion.bufferImageHeight = Surface.Height >> InMipIndex;
+			CopyRegion.bufferRowLength = Desc.Extent.X >> InMipIndex;
+			CopyRegion.bufferImageHeight = Desc.Extent.Y >> InMipIndex;
 			CopyRegion.imageSubresource.aspectMask = Surface.GetFullAspectMask();
 			CopyRegion.imageSubresource.mipLevel = InMipIndex;
 			CopyRegion.imageSubresource.baseArrayLayer = SrcBaseArrayLayer;
 			CopyRegion.imageSubresource.layerCount = 1;
-			CopyRegion.imageExtent.width = Surface.Width >> InMipIndex;
-			CopyRegion.imageExtent.height = Surface.Height >> InMipIndex;
+			CopyRegion.imageExtent.width = Desc.Extent.X >> InMipIndex;
+			CopyRegion.imageExtent.height = Desc.Extent.Y >> InMipIndex;
 			CopyRegion.imageExtent.depth = 1;
 
 			//#todo-rco: Multithreaded!
@@ -621,7 +607,7 @@ void FVulkanDynamicRHI::RHIReadSurfaceFloatData(FRHITexture* TextureRHI, FIntRec
 		FFloat16Color* Dest = OutputData.GetData();
 		for (int32 Row = InRect.Min.Y; Row < InRect.Max.Y; ++Row)
 		{
-			FFloat16Color* Src = (FFloat16Color*)StagingBuffer->GetMappedPointer() + Row * (Surface.Width >> InMipIndex) + InRect.Min.X;
+			FFloat16Color* Src = (FFloat16Color*)StagingBuffer->GetMappedPointer() + Row * (Desc.Extent.X >> InMipIndex) + InRect.Min.X;
 			for (int32 Col = InRect.Min.X; Col < InRect.Max.X; ++Col)
 			{
 				OutputData[OutIndex++] = *Src++;
@@ -630,22 +616,26 @@ void FVulkanDynamicRHI::RHIReadSurfaceFloatData(FRHITexture* TextureRHI, FIntRec
 		InDevice->GetStagingManager().ReleaseBuffer(InCmdBuffer, StagingBuffer);
 	};
 
+	FVulkanTexture& Surface = *FVulkanTexture::Cast(TextureRHI);
+	const FRHITextureDesc& Desc = Surface.GetDesc();
+
 	if(GIgnoreCPUReads == 1)
 	{
 		// Debug: Fill with CPU
 		uint32 NumPixels = 0;
-		if (TextureRHI->GetTextureCube())
-		{
-			FRHITextureCube* TextureRHICube = TextureRHI->GetTextureCube();
-			FVulkanTextureCube* TextureCube = (FVulkanTextureCube*)TextureRHICube;
-			NumPixels = (TextureCube->Surface.Width >> MipIndex) * (TextureCube->Surface.Height >> MipIndex);
-		}
-		else
+		switch(Desc.Dimension)
 		{
-			FRHITexture2D* TextureRHI2D = TextureRHI->GetTexture2D();
-			check(TextureRHI2D);
-			FVulkanTexture2D* Texture2D = (FVulkanTexture2D*)TextureRHI2D;
-			NumPixels = (Texture2D->Surface.Width >> MipIndex) * (Texture2D->Surface.Height >> MipIndex);
+		case ETextureDimension::TextureCube:
+			NumPixels = (Desc.Extent.X >> MipIndex) * (Desc.Extent.Y >> MipIndex);
+			break;
+
+		case ETextureDimension::Texture2D:
+			NumPixels = (Desc.Extent.X >> MipIndex) * (Desc.Extent.Y >> MipIndex);
+			break;
+
+		default:
+			checkNoEntry();
+			break;
 		}
 
 		OutData.Empty(0);
@@ -656,29 +646,28 @@ void FVulkanDynamicRHI::RHIReadSurfaceFloatData(FRHITexture* TextureRHI, FIntRec
 		Device->PrepareForCPURead();
 
 		FVulkanCmdBuffer* CmdBuffer = Device->GetImmediateContext().GetCommandBufferManager()->GetUploadCmdBuffer();
-		if (TextureRHI->GetTextureCube())
+		switch (TextureRHI->GetDesc().Dimension)
 		{
-			FRHITextureCube* TextureRHICube = TextureRHI->GetTextureCube();
-			FVulkanTextureCube* TextureCube = (FVulkanTextureCube*)TextureRHICube;
-			DoCopyFloat(Device, CmdBuffer, TextureCube->Surface, MipIndex, CubeFace + 6 * ArrayIndex, Rect, OutData);
-		}
-		else
-		{
-			FRHITexture2D* TextureRHI2D = TextureRHI->GetTexture2D();
-			check(TextureRHI2D);
-			FVulkanTexture2D* Texture2D = (FVulkanTexture2D*)TextureRHI2D;
-			DoCopyFloat(Device, CmdBuffer, Texture2D->Surface, MipIndex, ArrayIndex, Rect, OutData);
+			case ETextureDimension::TextureCube:
+				DoCopyFloat(Device, CmdBuffer, Surface, MipIndex, CubeFace + 6 * ArrayIndex, Rect, OutData);
+				break;
+
+			case ETextureDimension::Texture2D:
+				DoCopyFloat(Device, CmdBuffer, Surface, MipIndex, ArrayIndex, Rect, OutData);
+				break;
+
+			default:
+				checkNoEntry();
+				break;
 		}
 		Device->GetImmediateContext().GetCommandBufferManager()->PrepareForNewActiveCommandBuffer();
 	}
 }
 
-void FVulkanDynamicRHI::RHIRead3DSurfaceFloatData(FRHITexture* TextureRHI,FIntRect InRect,FIntPoint ZMinMax,TArray<FFloat16Color>& OutData)
+void FVulkanDynamicRHI::RHIRead3DSurfaceFloatData(FRHITexture* TextureRHI, FIntRect InRect, FIntPoint ZMinMax, TArray<FFloat16Color>& OutData)
 {
-	FRHITexture3D* TextureRHI3D = TextureRHI->GetTexture3D();
-	check(TextureRHI3D);
-	FVulkanTexture3D* Texture3D = (FVulkanTexture3D*)TextureRHI3D;
-	FVulkanSurface& Surface = Texture3D->Surface;
+	FVulkanTexture& Surface = *FVulkanTexture::Cast(TextureRHI);
+	const FRHITextureDesc& Desc = Surface.GetDesc();
 
 	uint32 SizeX = InRect.Width();
 	uint32 SizeY = InRect.Height();
@@ -688,11 +677,11 @@ void FVulkanDynamicRHI::RHIRead3DSurfaceFloatData(FRHITexture* TextureRHI,FIntRe
 
 	// Allocate the output buffer.
 	OutData.Reserve(Size);
+
 	if (GIgnoreCPUReads == 1)
 	{
-		OutData.AddZeroed(Size);
-
 		// Debug: Fill with CPU
+		OutData.AddZeroed(Size);
 		return;
 	}
 
@@ -707,8 +696,8 @@ void FVulkanDynamicRHI::RHIRead3DSurfaceFloatData(FRHITexture* TextureRHI,FIntRe
 		VkBufferImageCopy CopyRegion;
 		FMemory::Memzero(CopyRegion);
 		//Region.bufferOffset = 0;
-		CopyRegion.bufferRowLength = Surface.Width;
-		CopyRegion.bufferImageHeight = Surface.Height;
+		CopyRegion.bufferRowLength = Desc.Extent.X;
+		CopyRegion.bufferImageHeight = Desc.Extent.Y;
 		CopyRegion.imageSubresource.aspectMask = Surface.GetFullAspectMask();
 		//CopyRegion.imageSubresource.mipLevel = 0;
 		//CopyRegion.imageSubresource.baseArrayLayer = 0;
@@ -762,7 +751,7 @@ void FVulkanDynamicRHI::RHIRead3DSurfaceFloatData(FRHITexture* TextureRHI,FIntRe
 	{
 		for (int32 Row = InRect.Min.Y; Row < InRect.Max.Y; ++Row)
 		{
-			FFloat16Color* Src = (FFloat16Color*)StagingBuffer->GetMappedPointer() + Layer * SizeX * SizeY + Row * Surface.Width + InRect.Min.X;
+			FFloat16Color* Src = (FFloat16Color*)StagingBuffer->GetMappedPointer() + Layer * SizeX * SizeY + Row * Desc.Extent.X + InRect.Min.X;
 			for (int32 Col = InRect.Min.X; Col < InRect.Max.X; ++Col)
 			{
 				*Dest++ = *Src++;
@@ -819,7 +808,7 @@ bool FVulkanCommandListContext::IsSwapchainImage(FRHITexture* InTexture) const
 
 	for (uint32 i = 0; i < NumViewports; i++)
 	{
-		VkImage Image = FVulkanTextureBase::Cast(InTexture)->Surface.Image;
+		VkImage Image = FVulkanTexture::Cast(InTexture)->Image;
 		uint32 BackBufferImageCount = Viewports[i]->GetBackBufferImageCount();
 
 		for (uint32 SwapchainImageIdx = 0; SwapchainImageIdx < BackBufferImageCount; SwapchainImageIdx++)
@@ -859,7 +848,7 @@ void FVulkanCommandListContext::RHIBeginRenderPass(const FRHIRenderPassInfo& InI
 	VkImageLayout CurrentDSLayout;
 	if (DSTexture)
 	{
-		FVulkanSurface& Surface = FVulkanTextureBase::Cast(DSTexture)->Surface;
+		FVulkanTexture& Surface = *FVulkanTexture::Cast(DSTexture);
 		CurrentDSLayout = LayoutManager.FindLayoutChecked(Surface.Image);
 	}
 	else
@@ -1005,8 +994,9 @@ FVulkanRenderTargetLayout::FVulkanRenderTargetLayout(FVulkanDevice& InDevice, co
 		const FRHIRenderTargetView& RTView = RTInfo.ColorRenderTarget[Index];
 		if (RTView.Texture)
 		{
-			FVulkanTextureBase* Texture = FVulkanTextureBase::Cast(RTView.Texture);
+			FVulkanTexture* Texture = FVulkanTexture::Cast(RTView.Texture);
 			check(Texture);
+			const FRHITextureDesc& TextureDesc = Texture->GetDesc();
 
 			if (InDevice.GetImmediateContext().IsSwapchainImage(RTView.Texture))
 			{
@@ -1015,26 +1005,24 @@ FVulkanRenderTargetLayout::FVulkanRenderTargetLayout(FVulkanDevice& InDevice, co
 
 			if (bSetExtent)
 			{
-				ensure(Extent.Extent3D.width == FMath::Max(1u, Texture->Surface.Width >> RTView.MipIndex));
-				ensure(Extent.Extent3D.height == FMath::Max(1u, Texture->Surface.Height >> RTView.MipIndex));
-				ensure(Extent.Extent3D.depth == Texture->Surface.Depth);
+				ensure(Extent.Extent3D.width == FMath::Max(1, TextureDesc.Extent.X >> RTView.MipIndex));
+				ensure(Extent.Extent3D.height == FMath::Max(1, TextureDesc.Extent.Y >> RTView.MipIndex));
+				ensure(Extent.Extent3D.depth == TextureDesc.Depth);
 			}
 			else
 			{
 				bSetExtent = true;
-				Extent.Extent3D.width = FMath::Max(1u, Texture->Surface.Width >> RTView.MipIndex);
-				Extent.Extent3D.height = FMath::Max(1u, Texture->Surface.Height >> RTView.MipIndex);
-				Extent.Extent3D.depth = Texture->Surface.Depth;
+				Extent.Extent3D.width = FMath::Max(1, TextureDesc.Extent.X >> RTView.MipIndex);
+				Extent.Extent3D.height = FMath::Max(1, TextureDesc.Extent.Y >> RTView.MipIndex);
+				Extent.Extent3D.depth = TextureDesc.Depth;
 			}
 
-			FVulkanSurface* Surface = &Texture->Surface;
-
-			ensure(!NumSamples || NumSamples == Surface->GetNumSamples());
-			NumSamples = Surface->GetNumSamples();
+			ensure(!NumSamples || NumSamples == Texture->GetNumSamples());
+			NumSamples = Texture->GetNumSamples();
 		
 			VkAttachmentDescription& CurrDesc = Desc[NumAttachmentDescriptions];
 			CurrDesc.samples = static_cast<VkSampleCountFlagBits>(NumSamples);
-			CurrDesc.format = UEToVkTextureFormat(RTView.Texture->GetFormat(), EnumHasAllFlags(Texture->Surface.UEFlags, TexCreate_SRGB));
+			CurrDesc.format = UEToVkTextureFormat(RTView.Texture->GetFormat(), EnumHasAllFlags(TextureDesc.Flags, TexCreate_SRGB));
 			CurrDesc.loadOp = RenderTargetLoadActionToVulkan(RTView.LoadAction);
 			bFoundClearOp = bFoundClearOp || (CurrDesc.loadOp == VK_ATTACHMENT_LOAD_OP_CLEAR);
 			CurrDesc.storeOp = RenderTargetStoreActionToVulkan(RTView.StoreAction);
@@ -1042,7 +1030,7 @@ FVulkanRenderTargetLayout::FVulkanRenderTargetLayout(FVulkanDevice& InDevice, co
 			CurrDesc.stencilStoreOp = VK_ATTACHMENT_STORE_OP_DONT_CARE;
 
 			// Removed this temporarily as we need a way to determine if the target is actually memoryless
-			/*if (EnumHasAllFlags(Texture->Surface.UEFlags, TexCreate_Memoryless))
+			/*if (EnumHasAllFlags(Texture->UEFlags, TexCreate_Memoryless))
 			{
 				ensure(CurrDesc.storeOp == VK_ATTACHMENT_STORE_OP_DONT_CARE);
 			}*/
@@ -1086,12 +1074,12 @@ FVulkanRenderTargetLayout::FVulkanRenderTargetLayout(FVulkanDevice& InDevice, co
 	{
 		VkAttachmentDescription& CurrDesc = Desc[NumAttachmentDescriptions];
 		FMemory::Memzero(CurrDesc);
-		FVulkanTextureBase* Texture = FVulkanTextureBase::Cast(RTInfo.DepthStencilRenderTarget.Texture);
+		FVulkanTexture* Texture = FVulkanTexture::Cast(RTInfo.DepthStencilRenderTarget.Texture);
 		check(Texture);
+		const FRHITextureDesc& TextureDesc = Texture->GetDesc();
 
-		FVulkanSurface* Surface = &Texture->Surface;
-		ensure(!NumSamples || NumSamples == Surface->GetNumSamples());
-		NumSamples = Surface->GetNumSamples();
+		ensure(!NumSamples || NumSamples == Texture->GetNumSamples());
+		NumSamples = TextureDesc.NumSamples;
 
 		CurrDesc.samples = static_cast<VkSampleCountFlagBits>(NumSamples);
 		CurrDesc.format = UEToVkTextureFormat(RTInfo.DepthStencilRenderTarget.Texture->GetFormat(), false);
@@ -1104,7 +1092,7 @@ FVulkanRenderTargetLayout::FVulkanRenderTargetLayout(FVulkanDevice& InDevice, co
 			CurrDesc.stencilStoreOp = RenderTargetStoreActionToVulkan(RTInfo.DepthStencilRenderTarget.GetStencilStoreAction());
 
 			// Removed this temporarily as we need a way to determine if the target is actually memoryless
-			/*if (EnumHasAllFlags(Texture->Surface.UEFlags, TexCreate_Memoryless))
+			/*if (EnumHasAllFlags(Texture->UEFlags, TexCreate_Memoryless))
 			{
 				ensure(CurrDesc.storeOp == VK_ATTACHMENT_STORE_OP_DONT_CARE);
 				ensure(CurrDesc.stencilStoreOp == VK_ATTACHMENT_STORE_OP_DONT_CARE);
@@ -1142,24 +1130,22 @@ FVulkanRenderTargetLayout::FVulkanRenderTargetLayout(FVulkanDevice& InDevice, co
 		if (bSetExtent)
 		{
 			// Depth can be greater or equal to color. Clamp to the smaller size.
-			Extent.Extent3D.width = FMath::Min<uint32>(Extent.Extent3D.width, Texture->Surface.Width);
-			Extent.Extent3D.height = FMath::Min<uint32>(Extent.Extent3D.height, Texture->Surface.Height);
+			Extent.Extent3D.width = FMath::Min<uint32>(Extent.Extent3D.width, TextureDesc.Extent.X);
+			Extent.Extent3D.height = FMath::Min<uint32>(Extent.Extent3D.height, TextureDesc.Extent.Y);
 		}
 		else
 		{
 			bSetExtent = true;
-			Extent.Extent3D.width = Texture->Surface.Width;
-			Extent.Extent3D.height = Texture->Surface.Height;
-			Extent.Extent3D.depth = Texture->Surface.GetNumberOfArrayLevels();
+			Extent.Extent3D.width = TextureDesc.Extent.X;
+			Extent.Extent3D.height = TextureDesc.Extent.Y;
+			Extent.Extent3D.depth = Texture->GetNumberOfArrayLevels();
 		}
 	}
 
 	if (GRHISupportsAttachmentVariableRateShading && GRHIVariableRateShadingEnabled && GRHIAttachmentVariableRateShadingEnabled && RTInfo.ShadingRateTexture)
 	{
-		FVulkanTextureBase* Texture = FVulkanTextureBase::Cast(RTInfo.ShadingRateTexture);
-		check(Texture);
-		check(Texture->GetRHITexture());
-		check(Texture->GetRHITexture()->GetFormat() == GRHIVariableRateShadingImageFormat);
+		FVulkanTexture* Texture = FVulkanTexture::Cast(RTInfo.ShadingRateTexture);
+		check(Texture->GetFormat() == GRHIVariableRateShadingImageFormat);
 
 		VkAttachmentDescription& CurrDesc = Desc[NumAttachmentDescriptions];
 		FMemory::Memzero(CurrDesc);
@@ -1235,8 +1221,9 @@ FVulkanRenderTargetLayout::FVulkanRenderTargetLayout(FVulkanDevice& InDevice, co
 	for (int32 Index = 0; Index < NumColorRenderTargets; ++Index)
 	{
 		const FRHIRenderPassInfo::FColorEntry& ColorEntry = RPInfo.ColorRenderTargets[Index];
-		FVulkanTextureBase* Texture = FVulkanTextureBase::Cast(ColorEntry.RenderTarget);
+		FVulkanTexture* Texture = FVulkanTexture::Cast(ColorEntry.RenderTarget);
 		check(Texture);
+		const FRHITextureDesc& TextureDesc = Texture->GetDesc();
 
 		if (InDevice.GetImmediateContext().IsSwapchainImage(ColorEntry.RenderTarget))
 		{
@@ -1246,34 +1233,34 @@ FVulkanRenderTargetLayout::FVulkanRenderTargetLayout(FVulkanDevice& InDevice, co
 
 		if (bSetExtent)
 		{
-			ensure(Extent.Extent3D.width == FMath::Max(1u, Texture->Surface.Width >> ColorEntry.MipIndex));
-			ensure(Extent.Extent3D.height == FMath::Max(1u, Texture->Surface.Height >> ColorEntry.MipIndex));
-			ensure(Extent.Extent3D.depth == Texture->Surface.Depth);
+			ensure(Extent.Extent3D.width == FMath::Max(1, TextureDesc.Extent.X >> ColorEntry.MipIndex));
+			ensure(Extent.Extent3D.height == FMath::Max(1, TextureDesc.Extent.Y >> ColorEntry.MipIndex));
+			ensure(Extent.Extent3D.depth == TextureDesc.Depth);
 		}
 		else
 		{
 			bSetExtent = true;
-			Extent.Extent3D.width = FMath::Max(1u, Texture->Surface.Width >> ColorEntry.MipIndex);
-			Extent.Extent3D.height = FMath::Max(1u, Texture->Surface.Height >> ColorEntry.MipIndex);
-			Extent.Extent3D.depth = Texture->Surface.Depth;
+			Extent.Extent3D.width = FMath::Max(1, TextureDesc.Extent.X >> ColorEntry.MipIndex);
+			Extent.Extent3D.height = FMath::Max(1, TextureDesc.Extent.Y >> ColorEntry.MipIndex);
+			Extent.Extent3D.depth = TextureDesc.Depth;
 		}
 
 		ensure(!NumSamples || NumSamples == ColorEntry.RenderTarget->GetNumSamples());
 		NumSamples = ColorEntry.RenderTarget->GetNumSamples();
 
-		ensure( !bMultiviewRenderTargets || Texture->Surface.GetNumberOfArrayLevels() > 1 );
-		bMultiviewRenderTargets = Texture->Surface.GetNumberOfArrayLevels() > 1;
+		ensure( !bMultiviewRenderTargets || Texture->GetNumberOfArrayLevels() > 1 );
+		bMultiviewRenderTargets = Texture->GetNumberOfArrayLevels() > 1;
 
 		VkAttachmentDescription& CurrDesc = Desc[NumAttachmentDescriptions];
 		CurrDesc.samples = static_cast<VkSampleCountFlagBits>(NumSamples);
-		CurrDesc.format = UEToVkTextureFormat(ColorEntry.RenderTarget->GetFormat(), EnumHasAllFlags(Texture->Surface.UEFlags, TexCreate_SRGB));
+		CurrDesc.format = UEToVkTextureFormat(ColorEntry.RenderTarget->GetFormat(), EnumHasAllFlags(Texture->GetDesc().Flags, TexCreate_SRGB));
 		CurrDesc.loadOp = RenderTargetLoadActionToVulkan(GetLoadAction(ColorEntry.Action));
 		bFoundClearOp = bFoundClearOp || (CurrDesc.loadOp == VK_ATTACHMENT_LOAD_OP_CLEAR);
 		CurrDesc.storeOp = RenderTargetStoreActionToVulkan(GetStoreAction(ColorEntry.Action));
 		CurrDesc.stencilLoadOp = VK_ATTACHMENT_LOAD_OP_DONT_CARE;
 		CurrDesc.stencilStoreOp = VK_ATTACHMENT_STORE_OP_DONT_CARE;
 
-		if (EnumHasAnyFlags(Texture->Surface.UEFlags, TexCreate_Memoryless))
+		if (EnumHasAnyFlags(Texture->GetDesc().Flags, TexCreate_Memoryless))
 		{
 			ensure(CurrDesc.storeOp == VK_ATTACHMENT_STORE_OP_DONT_CARE);
 		}
@@ -1315,8 +1302,9 @@ FVulkanRenderTargetLayout::FVulkanRenderTargetLayout(FVulkanDevice& InDevice, co
 	{
 		VkAttachmentDescription& CurrDesc = Desc[NumAttachmentDescriptions];
 		FMemory::Memzero(CurrDesc);
-		FVulkanTextureBase* Texture = FVulkanTextureBase::Cast(RPInfo.DepthStencilRenderTarget.DepthStencilTarget);
+		FVulkanTexture* Texture = FVulkanTexture::Cast(RPInfo.DepthStencilRenderTarget.DepthStencilTarget);
 		check(Texture);
+		const FRHITextureDesc& TextureDesc = Texture->GetDesc();
 
 		CurrDesc.samples = static_cast<VkSampleCountFlagBits>(RPInfo.DepthStencilRenderTarget.DepthStencilTarget->GetNumSamples());
 		ensure(!NumSamples || CurrDesc.samples == NumSamples);
@@ -1335,7 +1323,7 @@ FVulkanRenderTargetLayout::FVulkanRenderTargetLayout(FVulkanDevice& InDevice, co
 		CurrDesc.storeOp = RenderTargetStoreActionToVulkan(GetStoreAction(GetDepthActions(RPInfo.DepthStencilRenderTarget.Action)));
 		CurrDesc.stencilStoreOp = RenderTargetStoreActionToVulkan(GetStoreAction(GetStencilActions(RPInfo.DepthStencilRenderTarget.Action)));
 
-		if (EnumHasAnyFlags(Texture->Surface.UEFlags, TexCreate_Memoryless))
+		if (EnumHasAnyFlags(TextureDesc.Flags, TexCreate_Memoryless))
 		{
 			ensure(CurrDesc.storeOp == VK_ATTACHMENT_STORE_OP_DONT_CARE);
 			ensure(CurrDesc.stencilStoreOp == VK_ATTACHMENT_STORE_OP_DONT_CARE);
@@ -1412,15 +1400,15 @@ FVulkanRenderTargetLayout::FVulkanRenderTargetLayout(FVulkanDevice& InDevice, co
 		if (bSetExtent)
 		{
 			// Depth can be greater or equal to color. Clamp to the smaller size.
-			Extent.Extent3D.width = FMath::Min<uint32>(Extent.Extent3D.width, Texture->Surface.Width);
-			Extent.Extent3D.height = FMath::Min<uint32>(Extent.Extent3D.height, Texture->Surface.Height);
+			Extent.Extent3D.width = FMath::Min<uint32>(Extent.Extent3D.width, TextureDesc.Extent.X);
+			Extent.Extent3D.height = FMath::Min<uint32>(Extent.Extent3D.height, TextureDesc.Extent.Y);
 		}
 		else
 		{
 			bSetExtent = true;
-			Extent.Extent3D.width = Texture->Surface.Width;
-			Extent.Extent3D.height = Texture->Surface.Height;
-			Extent.Extent3D.depth = Texture->Surface.Depth;
+			Extent.Extent3D.width = TextureDesc.Extent.X;
+			Extent.Extent3D.height = TextureDesc.Extent.Y;
+			Extent.Extent3D.depth = TextureDesc.Depth;
 		}
 	}
 	else if (NumColorRenderTargets == 0)
@@ -1438,10 +1426,8 @@ FVulkanRenderTargetLayout::FVulkanRenderTargetLayout(FVulkanDevice& InDevice, co
 
 	if (GRHISupportsAttachmentVariableRateShading && GRHIVariableRateShadingEnabled && GRHIAttachmentVariableRateShadingEnabled && RPInfo.ShadingRateTexture)
 	{
-		FVulkanTextureBase* Texture = FVulkanTextureBase::Cast(RPInfo.ShadingRateTexture);
-		check(Texture);
-		check(Texture->GetRHITexture());
-		check(Texture->GetRHITexture()->GetFormat() == GRHIVariableRateShadingImageFormat);
+		FVulkanTexture* Texture = FVulkanTexture::Cast(RPInfo.ShadingRateTexture);
+		check(Texture->GetFormat() == GRHIVariableRateShadingImageFormat);
 
 		VkAttachmentDescription& CurrDesc = Desc[NumAttachmentDescriptions];
 		FMemory::Memzero(CurrDesc);
