Copying //UE5/Dev-ParallelRendering to //UE5/Main

diff --git a/Engine/Source/Runtime/VulkanRHI/Private/VulkanResources.h b/Engine/Source/Runtime/VulkanRHI/Private/VulkanResources.h
--- a/Engine/Source/Runtime/VulkanRHI/Private/VulkanResources.h
+++ b/Engine/Source/Runtime/VulkanRHI/Private/VulkanResources.h
@@ -19,8 +19,7 @@ class FVulkanDevice;
 class FVulkanQueue;
 class FVulkanCmdBuffer;
 class FVulkanBufferCPU;
-struct FVulkanTextureBase;
-class FVulkanTexture2D;
+class FVulkanTexture;
 struct FVulkanBufferView;
 class FVulkanResourceMultiBuffer;
 class FVulkanLayout;
@@ -47,6 +46,20 @@ enum
 // Mirror GPixelFormats with format information for buffers
 extern VkFormat GVulkanBufferFormat[PF_MAX];
 
+// Converts the internal texture dimension to Vulkan view type
+inline VkImageViewType UETextureDimensionToVkImageViewType(ETextureDimension Dimension)
+{
+	switch (Dimension)
+	{
+	case ETextureDimension::Texture2D: return VK_IMAGE_VIEW_TYPE_2D;
+	case ETextureDimension::Texture2DArray: return VK_IMAGE_VIEW_TYPE_2D_ARRAY;
+	case ETextureDimension::Texture3D: return VK_IMAGE_VIEW_TYPE_3D;
+	case ETextureDimension::TextureCube: return VK_IMAGE_VIEW_TYPE_CUBE;
+	case ETextureDimension::TextureCubeArray: return VK_IMAGE_VIEW_TYPE_CUBE_ARRAY;
+	default: checkNoEntry(); return VK_IMAGE_VIEW_TYPE_MAX_ENUM;
+	}
+}
+
 /** This represents a vertex declaration that hasn't been combined with a specific shader to create a bound shader. */
 class FVulkanVertexDeclaration : public FRHIVertexDeclaration
 {
@@ -315,13 +328,80 @@ enum class EImageOwnerType : uint8
 	Aliased
 };
 
-/** Texture/RT wrapper. */
-class FVulkanSurface : public FVulkanEvictable
+
+struct FVulkanTextureView
+{
+	FVulkanTextureView()
+		: View(VK_NULL_HANDLE)
+		, Image(VK_NULL_HANDLE)
+		, ViewId(0)
+	{
+	}
+
+	void Create(FVulkanDevice& Device, VkImage InImage, VkImageViewType ViewType, VkImageAspectFlags AspectFlags, EPixelFormat UEFormat, VkFormat Format, uint32 FirstMip, uint32 NumMips, uint32 ArraySliceIndex, uint32 NumArraySlices, bool bUseIdentitySwizzle = false);
+	void Destroy(FVulkanDevice& Device);
+
+	VkImageView View;
+	VkImage Image;
+	uint32 ViewId;
+};
+
+
+class FVulkanTexture : public FRHITexture, public FVulkanEvictable
 {
-	virtual void Evict(FVulkanDevice& Device);
-	virtual void Move(FVulkanDevice& Device, FVulkanCommandListContext& Context, VulkanRHI::FVulkanAllocation& NewAllocation);
-	virtual bool CanEvict();
-	virtual bool CanMove();
+public:
+	// Regular constructor.
+	FVulkanTexture(FVulkanDevice& InDevice, const FRHITextureCreateDesc& InCreateDesc, const FRHITransientHeapAllocation* InTransientHeapAllocation);
+
+	// Construct from external resource.
+	// FIXME: HUGE HACK: the bUnused argument is there to disambiguate this overload from the one above when passing nullptr, since nullptr is a valid VkImage. Get rid of this code smell when unifying FVulkanSurface and FVulkanTexture.
+	FVulkanTexture(FVulkanDevice& InDevice, const FRHITextureCreateDesc& InCreateDesc, VkImage InImage, bool bUnused);
+
+	// Aliasing constructor.
+	FVulkanTexture(FVulkanDevice& InDevice, const FRHITextureCreateDesc& InCreateDesc, FTextureRHIRef& SrcTextureRHI);
+
+	virtual ~FVulkanTexture();
+
+	void AliasTextureResources(FTextureRHIRef& SrcTextureRHI);
+
+	// View with all mips/layers
+	FVulkanTextureView DefaultView;
+	// View with all mips/layers, but if it's a Depth/Stencil, only the Depth view
+	FVulkanTextureView* PartialView;
+
+	FTextureRHIRef AliasedTexture;
+
+	virtual void OnLayoutTransition(FVulkanCommandListContext& Context, VkImageLayout NewLayout) {}
+
+	template<typename T>
+	void DumpMemory(T Callback)
+	{
+		const FIntVector SizeXYZ = GetSizeXYZ();
+		Callback(TEXT("FVulkanTexture"), GetName(), this, static_cast<FRHIResource*>(this), SizeXYZ.X, SizeXYZ.Y, SizeXYZ.Z, StorageFormat);
+	}
+
+	// FVulkanEvictable interface.
+	bool CanMove() const override { return false; }
+	bool CanEvict() const override { return false; }
+	void Evict(FVulkanDevice& Device) override; ///evict to system memory
+	void Move(FVulkanDevice& Device, FVulkanCommandListContext& Context, VulkanRHI::FVulkanAllocation& NewAllocation) override; //move to a full new allocation
+	void OnFullDefrag(FVulkanDevice& Device, FVulkanCommandListContext& Context, uint32 NewOffset) override; //called when compacting an allocation. Old image can still be used as a copy source.
+	FVulkanTexture* GetEvictableTexture() override { return this; }
+	
+	void AttachView(VulkanRHI::FVulkanViewBase* View);
+	void DetachView(VulkanRHI::FVulkanViewBase* View);
+
+	bool GetTextureResourceInfo(FRHIResourceInfo& OutResourceInfo) const;
+
+	void* GetNativeResource() const override final { return (void*)Image; }
+	void* GetTextureBaseRHI() override final { return this; }
+
+	inline static FVulkanTexture* Cast(FRHITexture* Texture)
+	{
+		check(Texture);
+		return (FVulkanTexture*)Texture->GetTextureBaseRHI();
+	}
+
 public:
 	struct FImageCreateInfo
 	{
@@ -339,38 +419,15 @@ public:
 	static void GenerateImageCreateInfo(
 		FImageCreateInfo& OutImageCreateInfo,
 		FVulkanDevice& InDevice,
-		VkImageViewType ResourceType,
-		EPixelFormat InFormat,
-		uint32 SizeX, uint32 SizeY, uint32 SizeZ,
-		uint32 ArraySize,
-		uint32 NumMips,
-		uint32 NumSamples,
-		ETextureCreateFlags UEFlags,
+		const FRHITextureDesc& InDesc,
 		VkFormat* OutStorageFormat = nullptr,
 		VkFormat* OutViewFormat = nullptr,
 		bool bForceLinearTexture = false);
 
-	FVulkanSurface(FVulkanDevice& Device, FVulkanEvictable* Owner, VkImageViewType ResourceType, EPixelFormat Format,
-					uint32 SizeX, uint32 SizeY, uint32 SizeZ, uint32 ArraySize,
-					uint32 NumMips, uint32 NumSamples, ETextureCreateFlags UEFlags, ERHIAccess InResourceState, 
-					const FRHIResourceCreateInfo& CreateInfo, const FRHITransientHeapAllocation* InTransientHeapAllocation = nullptr);
-
-	// Constructor for externally owned Image
-	FVulkanSurface(FVulkanDevice& Device, VkImageViewType ResourceType, EPixelFormat Format,
-					uint32 SizeX, uint32 SizeY, uint32 SizeZ, uint32 ArraySize, uint32 NumMips, uint32 NumSamples,
-					VkImage InImage, ETextureCreateFlags UEFlags, EImageOwnerType InImageOwnerType, const FRHIResourceCreateInfo& CreateInfo);
-
-	virtual ~FVulkanSurface();
-
-	void Destroy();
+	void DestroySurface();
 	void InvalidateMappedMemory();
 	void* GetMappedPointer();
 
-	void MoveSurface(FVulkanDevice& InDevice, FVulkanCommandListContext& Context, VulkanRHI::FVulkanAllocation& NewAllocation);
-	void OnFullDefrag(FVulkanDevice& InDevice, FVulkanCommandListContext& Context, uint32 NewOffset);
-	void EvictSurface(FVulkanDevice& InDevice);
-
-
 	/**
 	 * Returns how much memory is used by the surface
 	 */
@@ -394,28 +451,27 @@ public:
 	*/
 	void GetMipSize(uint32 MipIndex, uint32& MipBytes);
 
-	inline VkImageViewType GetViewType() const { return ViewType; }
+	inline VkImageViewType GetViewType() const
+	{
+		return UETextureDimensionToVkImageViewType(GetDesc().Dimension);
+	}
 
 	inline VkImageTiling GetTiling() const { return Tiling; }
 
-	inline uint32 GetNumMips() const { return NumMips; }
-
-	inline uint32 GetNumSamples() const { return NumSamples; }
-
 	inline uint32 GetNumberOfArrayLevels() const
 	{
-		switch (ViewType)
+		switch (GetViewType())
 		{
 		case VK_IMAGE_VIEW_TYPE_1D:
 		case VK_IMAGE_VIEW_TYPE_2D:
 		case VK_IMAGE_VIEW_TYPE_3D:
 			return 1;
 		case VK_IMAGE_VIEW_TYPE_2D_ARRAY:
-			return ArraySize;
+			return GetDesc().ArraySize;
 		case VK_IMAGE_VIEW_TYPE_CUBE:
 			return 6;
 		case VK_IMAGE_VIEW_TYPE_CUBE_ARRAY:
-			return 6 * ArraySize;
+			return 6 * GetDesc().ArraySize;
 		default:
 			ErrorInvalidViewType();
 			return 1;
@@ -447,332 +503,41 @@ public:
 
 	inline bool SupportsSampling() const
 	{
-		return EnumHasAllFlags(GPixelFormats[PixelFormat].Capabilities, EPixelFormatCapabilities::TextureSample);
+		return EnumHasAllFlags(GPixelFormats[GetDesc().Format].Capabilities, EPixelFormatCapabilities::TextureSample);
 	}
 
 	VULKANRHI_API VkDeviceMemory GetAllocationHandle() const;
 	VULKANRHI_API uint64 GetAllocationOffset() const;
 
+	static void InternalLockWrite(FVulkanCommandListContext& Context, FVulkanTexture* Surface, const VkBufferImageCopy& Region, VulkanRHI::FStagingBuffer* StagingBuffer);
 
-	FVulkanDevice* Device;
+	const FVulkanCpuReadbackBuffer* GetCpuReadbackBuffer() const { return CpuReadbackBuffer; }
 
+	FVulkanDevice* Device;
 	VkImage Image;
-	
-	// Removes SRGB if requested, used to upload data
-	VkFormat StorageFormat;
-	// Format for SRVs, render targets
-	VkFormat ViewFormat;
-	uint32 Width, Height, Depth, ArraySize;
-	// UE format
-	EPixelFormat PixelFormat;
-	ETextureCreateFlags UEFlags;
+	VkFormat StorageFormat;  // Removes SRGB if requested, used to upload data
+	VkFormat ViewFormat;  // Format for SRVs, render targets
 	VkMemoryPropertyFlags MemProps;
 	VkMemoryRequirements MemoryRequirements;
 
-	static void InternalLockWrite(FVulkanCommandListContext& Context, FVulkanSurface* Surface, const VkBufferImageCopy& Region, VulkanRHI::FStagingBuffer* StagingBuffer);
-
-	const FVulkanCpuReadbackBuffer* GetCpuReadbackBuffer() const { return CpuReadbackBuffer; }
 private:
+	VulkanRHI::FVulkanViewBase* FirstView = nullptr;
 
+	void InvalidateViews(FVulkanDevice& Device);
+	void DestroyViews();
 	void SetInitialImageState(FVulkanCommandListContext& Context, VkImageLayout InitialLayout, bool bClear, const FClearValueBinding& ClearValueBinding);
-	friend struct FRHICommandSetInitialImageState;
-
 	void InternalMoveSurface(FVulkanDevice& InDevice, FVulkanCommandListContext& Context, VulkanRHI::FVulkanAllocation& DestAllocation);
 
-private:
 	VkImageTiling Tiling;
-	VkImageViewType	ViewType;
-
 	VulkanRHI::FVulkanAllocation Allocation;
-
-	uint32 NumMips;
-	uint32 NumSamples;
-
 	VkImageAspectFlags FullAspectMask;
 	VkImageAspectFlags PartialAspectMask;
-
 	FVulkanCpuReadbackBuffer* CpuReadbackBuffer;
-	FVulkanTextureBase* OwningTexture = 0;
-
-	EImageOwnerType ImageOwnerType;
-
-	friend struct FVulkanTextureBase;
-};
-
-
-struct FVulkanTextureView
-{
-	FVulkanTextureView()
-		: View(VK_NULL_HANDLE)
-		, Image(VK_NULL_HANDLE)
-		, ViewId(0)
-	{
-	}
-
-	void Create(FVulkanDevice& Device, VkImage InImage, VkImageViewType ViewType, VkImageAspectFlags AspectFlags, EPixelFormat UEFormat, VkFormat Format, uint32 FirstMip, uint32 NumMips, uint32 ArraySliceIndex, uint32 NumArraySlices, bool bUseIdentitySwizzle = false);
-	void Destroy(FVulkanDevice& Device);
-
-	VkImageView View;
-	VkImage Image;
-	uint32 ViewId;
-
-private:
-	static VkImageView StaticCreate(FVulkanDevice& Device, VkImage InImage, VkImageViewType ViewType, VkImageAspectFlags AspectFlags, EPixelFormat UEFormat, VkFormat Format, uint32 FirstMip, uint32 NumMips, uint32 ArraySliceIndex, uint32 NumArraySlices, bool bUseIdentitySwizzle);
-};
-
-
-struct FVulkanTextureBase : public FVulkanEvictable, public IRefCountedObject
-{
-	inline static FVulkanTextureBase* Cast(FRHITexture* Texture)
-	{
-		check(Texture);
-		return (FVulkanTextureBase*)Texture->GetTextureBaseRHI();
-	}
-
-	FVulkanTextureBase(FVulkanDevice& Device, VkImageViewType ResourceType, EPixelFormat Format, uint32 SizeX, uint32 SizeY, uint32 SizeZ, uint32 ArraySize, uint32 NumMips, uint32 NumSamples, ETextureCreateFlags UEFlags, ERHIAccess InResourceState, const FRHIResourceCreateInfo& CreateInfo, const FRHITransientHeapAllocation* InTransientHeapAllocation);
-	FVulkanTextureBase(FVulkanDevice& Device, VkImageViewType ResourceType, EPixelFormat Format, uint32 SizeX, uint32 SizeY, uint32 SizeZ, uint32 ArraySize, uint32 NumMips, uint32 NumSamples, VkImage InImage, VkDeviceMemory InMem, ETextureCreateFlags UEFlags, const FRHIResourceCreateInfo& CreateInfo);
-
-	// Aliasing constructor.
-	FVulkanTextureBase(FTextureRHIRef& SrcTextureRHI, const FVulkanTextureBase* SrcTexture, VkImageViewType ResourceType, uint32 SizeX, uint32 SizeY, uint32 sizeZ);
-
-	virtual ~FVulkanTextureBase();
-
-	void AliasTextureResources(FTextureRHIRef& SrcTexture);
-
-	FVulkanSurface Surface;
-
-	// View with all mips/layers
-	FVulkanTextureView DefaultView;
-	// View with all mips/layers, but if it's a Depth/Stencil, only the Depth view
-	FVulkanTextureView* PartialView;
-
-	FTextureRHIRef AliasedTexture;
-
-	virtual void OnLayoutTransition(FVulkanCommandListContext& Context, VkImageLayout NewLayout) {}
-
-	template<typename T>
-	void DumpMemory(T Callback)
-	{
-		Callback(TEXT("FVulkanTextureBase"), GetResourceFName(), this, GetRHIResource(), Surface.Width, Surface.Height, Surface.Depth, Surface.StorageFormat);
-	}
-
-	void Evict(FVulkanDevice& Device); ///evict to system memory
-	void Move(FVulkanDevice& Device, FVulkanCommandListContext& Context, VulkanRHI::FVulkanAllocation& NewAllocation); //move to a full new allocation
-	void OnFullDefrag(FVulkanDevice& Device, FVulkanCommandListContext& Context, uint32 NewOffset); //called when compacting an allocation. Old image can still be used as a copy source.
-	FVulkanTextureBase* GetTextureBase() { return this; }
-
-	void AttachView(VulkanRHI::FVulkanViewBase* View);
-	void DetachView(VulkanRHI::FVulkanViewBase* View);
-
-	bool GetTextureResourceInfo(FRHIResourceInfo& OutResourceInfo) const;
-
-	virtual FRHITexture* GetRHITexture() = 0;
-private:
-	void InvalidateViews(FVulkanDevice& Device);
-	VulkanRHI::FVulkanViewBase* FirstView = nullptr;
-
-	void DestroyViews();
-	virtual FName GetResourceFName() = 0;
-	virtual FRHIResource* GetRHIResource(){ return 0; }
-
-};
 
-class FVulkanTexture2D : public FRHITexture2D, public FVulkanTextureBase
-{
-	FName GetResourceFName(){ return GetName(); }
-	virtual FRHIResource* GetRHIResource() { return (FRHITexture2D*)this; }
-public:
-	FVulkanTexture2D(FVulkanDevice& Device, EPixelFormat Format, uint32 SizeX, uint32 SizeY, uint32 NumMips, uint32 NumSamples, ETextureCreateFlags UEFlags, ERHIAccess InResourceState, const FRHIResourceCreateInfo& CreateInfo, const FRHITransientHeapAllocation* InTransientHeapAllocation = nullptr);
-	FVulkanTexture2D(FVulkanDevice& Device, EPixelFormat Format, uint32 SizeX, uint32 SizeY, uint32 NumMips, uint32 NumSamples, VkImage Image, ETextureCreateFlags UEFlags, const FRHIResourceCreateInfo& CreateInfo);
-
-	// Aliasing constructor
-	FVulkanTexture2D(FTextureRHIRef& SrcTextureRHI, const FVulkanTexture2D* SrcTexture);
-
-	virtual ~FVulkanTexture2D();
-	virtual FRHITexture* GetRHITexture()
-	{
-		return this;
-	};
-
-	// IRefCountedObject interface.
-	virtual uint32 AddRef() const override final
-	{
-		return FRHIResource::AddRef();
-	}
-	virtual uint32 Release() const override final
-	{
-		return FRHIResource::Release();
-	}
-	virtual uint32 GetRefCount() const override final
-	{
-		return FRHIResource::GetRefCount();
-	}
-
-	virtual void* GetTextureBaseRHI() override final
-	{
-		FVulkanTextureBase* Base = static_cast<FVulkanTextureBase*>(this);
-		return Base;
-	}
-
-	virtual void* GetNativeResource() const
-	{
-		return (void*)Surface.Image;
-	}
-
-#if RHI_ENABLE_RESOURCE_INFO
-	virtual bool GetResourceInfo(FRHIResourceInfo& OutResourceInfo) const override final
-	{
-		return GetTextureResourceInfo(OutResourceInfo);
-	}
-#endif
-};
-
-class FVulkanTexture2DArray : public FRHITexture2DArray, public FVulkanTextureBase
-{
-	FName GetResourceFName() { return GetName(); }
-public:
-	// Constructor, just calls base and Surface constructor
-	FVulkanTexture2DArray(FVulkanDevice& Device, EPixelFormat Format, uint32 SizeX, uint32 SizeY, uint32 ArraySize, uint32 NumMips, uint32 NumSamples, ETextureCreateFlags Flags, ERHIAccess InResourceState, const FRHIResourceCreateInfo& CreateInfo, const FRHITransientHeapAllocation* InTransientHeapAllocation = nullptr);
-	FVulkanTexture2DArray(FVulkanDevice& Device, EPixelFormat Format, uint32 SizeX, uint32 SizeY, uint32 ArraySize, uint32 NumMips, uint32 NumSamples, VkImage Image, ETextureCreateFlags Flags, const FRHIResourceCreateInfo& CreateInfo);
-
-	// Aliasing constructor
-	FVulkanTexture2DArray(FTextureRHIRef& SrcTextureRHI, const FVulkanTexture2DArray* SrcTexture);
-
-	virtual FRHITexture* GetRHITexture()
-	{
-		return this;
-	};
-
-
-
-	// IRefCountedObject interface.
-	virtual uint32 AddRef() const override final
-	{
-		return FRHIResource::AddRef();
-	}
-	virtual uint32 Release() const override final
-	{
-		return FRHIResource::Release();
-	}
-	virtual uint32 GetRefCount() const override final
-	{
-		return FRHIResource::GetRefCount();
-	}
-
-	virtual void* GetTextureBaseRHI() override final
-	{
-		return (FVulkanTextureBase*)this;
-	}
-
-	virtual void* GetNativeResource() const
-	{
-		return (void*)Surface.Image;
-	}
-
-#if RHI_ENABLE_RESOURCE_INFO
-	virtual bool GetResourceInfo(FRHIResourceInfo& OutResourceInfo) const override final
-	{
-		return GetTextureResourceInfo(OutResourceInfo);
-	}
-#endif
-};
-
-class FVulkanTexture3D : public FRHITexture3D, public FVulkanTextureBase
-{
-	FName GetResourceFName() { return GetName(); }
-public:
-	// Constructor, just calls base and Surface constructor
-	FVulkanTexture3D(FVulkanDevice& Device, EPixelFormat Format, uint32 SizeX, uint32 SizeY, uint32 SizeZ, uint32 NumMips, ETextureCreateFlags Flags, ERHIAccess InResourceState, const FRHIResourceCreateInfo& CreateInfo, const FRHITransientHeapAllocation* InTransientHeapAllocation = nullptr);
-	FVulkanTexture3D(FVulkanDevice& Device, EPixelFormat Format, uint32 SizeX, uint32 SizeY, uint32 SizeZ, uint32 NumMips, VkImage Image, ETextureCreateFlags Flags, const FRHIResourceCreateInfo& CreateInfo);
-	virtual ~FVulkanTexture3D();
-
-
-	virtual FRHITexture* GetRHITexture()
-	{
-		return this;
-	}
-
-
-	// IRefCountedObject interface.
-	virtual uint32 AddRef() const override final
-	{
-		return FRHIResource::AddRef();
-	}
-	virtual uint32 Release() const override final
-	{
-		return FRHIResource::Release();
-	}
-	virtual uint32 GetRefCount() const override final
-	{
-		return FRHIResource::GetRefCount();
-	}
-
-	virtual void* GetTextureBaseRHI() override final
-	{
-		return (FVulkanTextureBase*)this;
-	}
-
-	virtual void* GetNativeResource() const
-	{
-		return (void*)Surface.Image;
-	}
-
-#if RHI_ENABLE_RESOURCE_INFO
-	virtual bool GetResourceInfo(FRHIResourceInfo& OutResourceInfo) const override final
-	{
-		return GetTextureResourceInfo(OutResourceInfo);
-	}
-#endif
-};
-
-class FVulkanTextureCube : public FRHITextureCube, public FVulkanTextureBase
-{
-	FName GetResourceFName() { return GetName(); }
-public:
-	FVulkanTextureCube(FVulkanDevice& Device, EPixelFormat Format, uint32 Size, bool bArray, uint32 ArraySize, uint32 NumMips, ETextureCreateFlags Flags, ERHIAccess InResourceState, const FRHIResourceCreateInfo& CreateInfo, const FRHITransientHeapAllocation* InTransientHeapAllocation = nullptr);
-	FVulkanTextureCube(FVulkanDevice& Device, EPixelFormat Format, uint32 Size, bool bArray, uint32 ArraySize, uint32 NumMips, VkImage Image, ETextureCreateFlags Flags, const FRHIResourceCreateInfo& CreateInfo);
-
-	// Aliasing constructor
-	FVulkanTextureCube(FTextureRHIRef& SrcTextureRHI, const FVulkanTextureCube* SrcTexture);
-
-	virtual ~FVulkanTextureCube();
-
-	virtual FRHITexture* GetRHITexture()
-	{
-		return this;
-	};
-
-
-	// IRefCountedObject interface.
-	virtual uint32 AddRef() const override final
-	{
-		return FRHIResource::AddRef();
-	}
-	virtual uint32 Release() const override final
-	{
-		return FRHIResource::Release();
-	}
-	virtual uint32 GetRefCount() const override final
-	{
-		return FRHIResource::GetRefCount();
-	}
-
-	virtual void* GetTextureBaseRHI() override final
-	{
-		return (FVulkanTextureBase*)this;
-	}
-
-	virtual void* GetNativeResource() const
-	{
-		return (void*)Surface.Image;
-	}
+	friend struct FRHICommandSetInitialImageState;
 
-#if RHI_ENABLE_RESOURCE_INFO
-	virtual bool GetResourceInfo(FRHIResourceInfo& OutResourceInfo) const override final
-	{
-		return GetTextureResourceInfo(OutResourceInfo);
-	}
-#endif
+protected:
+	EImageOwnerType ImageOwnerType;
 };
 
 class FVulkanQueryPool : public VulkanRHI::FDeviceChild
@@ -1000,9 +765,6 @@ struct FVulkanBufferView : public FRHIResource, public VulkanRHI::FDeviceChild
 
 struct FVulkanRingBuffer : public FVulkanEvictable, public VulkanRHI::FDeviceChild
 {
-	virtual void Evict(FVulkanDevice& Device);
-	virtual void Move(FVulkanDevice& Device, FVulkanCommandListContext& Context, VulkanRHI::FVulkanAllocation& NewAllocation);
-
 public:
 	FVulkanRingBuffer(FVulkanDevice* InDevice, uint64 TotalSize, VkFlags Usage, VkMemoryPropertyFlags MemPropertyFlags);
 	virtual ~FVulkanRingBuffer();
@@ -1098,13 +860,16 @@ protected:
 
 class FVulkanResourceMultiBuffer : public FRHIBuffer, public FVulkanEvictable, public VulkanRHI::FDeviceChild
 {
-	virtual void Evict(FVulkanDevice& Device);
-	virtual void Move(FVulkanDevice& Device, FVulkanCommandListContext& Context, VulkanRHI::FVulkanAllocation& NewAllocation);
-
 public:
 	FVulkanResourceMultiBuffer(FVulkanDevice* InDevice, uint32 InSize, EBufferUsageFlags InUEUsage, uint32 InStride, FRHIResourceCreateInfo& CreateInfo, class FRHICommandListImmediate* InRHICmdList = nullptr, const FRHITransientHeapAllocation* InTransientHeapAllocation = nullptr);
 	virtual ~FVulkanResourceMultiBuffer();
 
+	// FVulkanEvictable interface.
+	bool CanMove() const override { return false; }
+	bool CanEvict() const override { return false; }
+	void Evict(FVulkanDevice& Device) override; ///evict to system memory
+	void Move(FVulkanDevice& Device, FVulkanCommandListContext& Context, VulkanRHI::FVulkanAllocation& NewAllocation) override; //move to a full new allocation
+
 	inline const VulkanRHI::FVulkanAllocation& GetCurrentAllocation() const
 	{
 		return Current.Alloc;
@@ -1342,7 +1107,7 @@ protected:
 	uint32 VolatileLockCounter = MAX_uint32;
 
 	FVulkanShaderResourceView* NextView = 0;
-	friend struct FVulkanTextureBase;
+	friend class FVulkanTexture;
 };
 
 class FVulkanVertexInputStateInfo
@@ -1539,29 +1304,9 @@ struct TVulkanResourceTraits<FRHIComputeShader>
 	typedef FVulkanComputeShader TConcreteType;
 };
 template<>
-struct TVulkanResourceTraits<FRHITexture3D>
-{
-	typedef FVulkanTexture3D TConcreteType;
-};
-//template<>
-//struct TVulkanResourceTraits<FRHITexture>
-//{
-//	typedef FVulkanTexture TConcreteType;
-//};
-template<>
-struct TVulkanResourceTraits<FRHITexture2D>
-{
-	typedef FVulkanTexture2D TConcreteType;
-};
-template<>
-struct TVulkanResourceTraits<FRHITexture2DArray>
-{
-	typedef FVulkanTexture2DArray TConcreteType;
-};
-template<>
-struct TVulkanResourceTraits<FRHITextureCube>
+struct TVulkanResourceTraits<FRHITexture>
 {
-	typedef FVulkanTextureCube TConcreteType;
+	typedef FVulkanTexture TConcreteType;
 };
 template<>
 struct TVulkanResourceTraits<FRHIRenderQuery>
diff --git a/Engine/Source/Runtime/VulkanRHI/Private/VulkanSwapChain.cpp b/Engine/Source/Runtime/VulkanRHI/Private/VulkanSwapChain.cpp
--- a/Engine/Source/Runtime/VulkanRHI/Private/VulkanSwapChain.cpp
+++ b/Engine/Source/Runtime/VulkanRHI/Private/VulkanSwapChain.cpp
@@ -908,24 +908,34 @@ FVulkanSwapChain::EStatus FVulkanSwapChain::Present(FVulkanQueue* GfxQueue, FVul
 	return EStatus::Healthy;
 }
 
-void FVulkanSwapChain::CreateQCOMDepthStencil(const FVulkanSurface& InSurface) const
+void FVulkanSwapChain::CreateQCOMDepthStencil(const FVulkanTexture& InSurface) const
 {
 	check(!QCOMDepthStencilSurface);
 	check(!QCOMDepthStencilView);
 	check(!QCOMDepthView);
 
-	ETextureCreateFlags UEFlags = InSurface.UEFlags;
+	const FRHITextureDesc& Desc = InSurface.GetDesc();
+	const ETextureCreateFlags UEFlags = Desc.Flags;
 	check(UEFlags & TexCreate_DepthStencilTargetable);
 
-	QCOMDepthStencilSurface = new FVulkanSurface(*InSurface.Device, nullptr, InSurface.GetViewType(), InSurface.PixelFormat, InSurface.Height, InSurface.Width,
-													InSurface.Depth, 1, InSurface.GetNumMips(), InSurface.GetNumSamples(), UEFlags, ERHIAccess::Unknown, FRHIResourceCreateInfo(TEXT("FVulkanSwapChain")));
+	FRHITextureCreateDesc CreateDesc = FRHITextureCreateDesc::Create2D(
+		TEXT("FVulkanSwapChainQCOM"),
+		{ (int32)Desc.Extent.X, (int32)Desc.Extent.Y }, // Yes, these are intentionally swapped.
+		Desc.Format,
+		FClearValueBinding::None,
+		UEFlags,
+		Desc.NumMips,
+		Desc.NumSamples
+	);
+
+	QCOMDepthStencilSurface = new FVulkanTexture(Device, CreateDesc, nullptr);
 
 	check(QCOMDepthStencilSurface->GetViewType() == VK_IMAGE_VIEW_TYPE_2D);
 	check(QCOMDepthStencilSurface->Image != VK_NULL_HANDLE);
 
 	QCOMDepthStencilView = new FVulkanTextureView;
 	QCOMDepthStencilView->Create(*QCOMDepthStencilSurface->Device, QCOMDepthStencilSurface->Image, QCOMDepthStencilSurface->GetViewType(), QCOMDepthStencilSurface->GetFullAspectMask(),
-								QCOMDepthStencilSurface->PixelFormat, QCOMDepthStencilSurface->ViewFormat, 0, FMath::Max(QCOMDepthStencilSurface->GetNumMips(), 1u), 0, 1u);
+								QCOMDepthStencilSurface->GetDesc().Format, QCOMDepthStencilSurface->ViewFormat, 0, FMath::Max(QCOMDepthStencilSurface->GetNumMips(), 1u), 0, 1u, false);
 
 	if (QCOMDepthStencilSurface->GetFullAspectMask() == QCOMDepthStencilSurface->GetPartialAspectMask())
 	{
@@ -935,11 +945,11 @@ void FVulkanSwapChain::CreateQCOMDepthStencil(const FVulkanSurface& InSurface) c
 	{
 		QCOMDepthView = new FVulkanTextureView;
 		QCOMDepthView->Create(*QCOMDepthStencilSurface->Device, QCOMDepthStencilSurface->Image, QCOMDepthStencilSurface->GetViewType(), QCOMDepthStencilSurface->GetPartialAspectMask(),
-			QCOMDepthStencilSurface->PixelFormat, QCOMDepthStencilSurface->ViewFormat, 0, FMath::Max(QCOMDepthStencilSurface->GetNumMips(), 1u), 0, 1u);
+			QCOMDepthStencilSurface->GetDesc().Format, QCOMDepthStencilSurface->ViewFormat, 0, FMath::Max(QCOMDepthStencilSurface->GetNumMips(), 1u), 0, 1u, false);
 	}
 }
 
-const FVulkanTextureView* FVulkanSwapChain::GetOrCreateQCOMDepthStencilView(const FVulkanSurface& InSurface) const
+const FVulkanTextureView* FVulkanSwapChain::GetOrCreateQCOMDepthStencilView(const FVulkanTexture& InSurface) const
 {
 	if (QCOMDepthStencilView)
 	{
@@ -951,7 +961,7 @@ const FVulkanTextureView* FVulkanSwapChain::GetOrCreateQCOMDepthStencilView(cons
 	return QCOMDepthStencilView;
 }
 
-const FVulkanTextureView* FVulkanSwapChain::GetOrCreateQCOMDepthView(const FVulkanSurface& InSurface) const
+const FVulkanTextureView* FVulkanSwapChain::GetOrCreateQCOMDepthView(const FVulkanTexture& InSurface) const
 {
 	if (QCOMDepthView)
 	{
@@ -963,7 +973,7 @@ const FVulkanTextureView* FVulkanSwapChain::GetOrCreateQCOMDepthView(const FVulk
 	return QCOMDepthView;
 }
 
-const FVulkanSurface* FVulkanSwapChain::GetQCOMDepthStencilSurface() const
+const FVulkanTexture* FVulkanSwapChain::GetQCOMDepthStencilSurface() const
 {
 	return QCOMDepthStencilSurface;
 }
diff --git a/Engine/Source/Runtime/VulkanRHI/Private/VulkanSwapChain.h b/Engine/Source/Runtime/VulkanRHI/Private/VulkanSwapChain.h
--- a/Engine/Source/Runtime/VulkanRHI/Private/VulkanSwapChain.h
+++ b/Engine/Source/Runtime/VulkanRHI/Private/VulkanSwapChain.h
@@ -6,6 +6,8 @@
 
 #pragma once
 
+class FVulkanTexture;
+
 namespace VulkanRHI
 {
 	class FFence;
@@ -40,9 +42,9 @@ public:
 	void RenderThreadPacing();
 	inline int8 DoesLockToVsync() { return LockToVsync; }
 
-	const FVulkanTextureView* GetOrCreateQCOMDepthStencilView(const FVulkanSurface& InSurface) const;
-	const FVulkanTextureView* GetOrCreateQCOMDepthView(const FVulkanSurface& InSurface) const;
-	const FVulkanSurface* GetQCOMDepthStencilSurface() const;
+	const FVulkanTextureView* GetOrCreateQCOMDepthStencilView(const FVulkanTexture& InSurface) const;
+	const FVulkanTextureView* GetOrCreateQCOMDepthView(const FVulkanTexture& InSurface) const;
+	const FVulkanTexture* GetQCOMDepthStencilSurface() const;
 
 protected:
 	VkSurfaceTransformFlagBitsKHR QCOMRenderPassTransform = VK_SURFACE_TRANSFORM_IDENTITY_BIT_KHR;
@@ -82,8 +84,8 @@ protected:
 	// For example in case if QCOMRenderPassTransform is VK_SURFACE_TRANSFORM_ROTATE_90_BIT_KHR we must swap width/height of depth target.
 	// But probably engine can use same depth texture for swapchain and NON swapchain passes. So it is why we have this addditional surface.
 	// With this approach we should be careful in case if depth in swapchain pass is used as attachement and fetched in shader in same time
-	void CreateQCOMDepthStencil(const FVulkanSurface& InSurface) const;
-	mutable FVulkanSurface* QCOMDepthStencilSurface = nullptr;
+	void CreateQCOMDepthStencil(const FVulkanTexture& InSurface) const;
+	mutable FVulkanTexture* QCOMDepthStencilSurface = nullptr;
 	mutable FVulkanTextureView* QCOMDepthStencilView = nullptr;
 	mutable FVulkanTextureView* QCOMDepthView = nullptr;
 
