Copying //UE5/Dev-ParallelRendering to //UE5/Main

diff --git a/Engine/Source/Runtime/VulkanRHI/Private/VulkanTexture.cpp b/Engine/Source/Runtime/VulkanRHI/Private/VulkanTexture.cpp
--- a/Engine/Source/Runtime/VulkanRHI/Private/VulkanTexture.cpp
+++ b/Engine/Source/Runtime/VulkanRHI/Private/VulkanTexture.cpp
@@ -152,7 +152,7 @@ static void VulkanTextureDestroyed(uint64 Size, VkImageViewType ImageType, bool
 	UpdateVulkanTextureStats(-(int64)Size, bIsCube, bIs3D, bIsRT);
 }
 
-inline void FVulkanSurface::InternalLockWrite(FVulkanCommandListContext& Context, FVulkanSurface* Surface, const VkBufferImageCopy& Region, VulkanRHI::FStagingBuffer* StagingBuffer)
+inline void FVulkanTexture::InternalLockWrite(FVulkanCommandListContext& Context, FVulkanTexture* Surface, const VkBufferImageCopy& Region, VulkanRHI::FStagingBuffer* StagingBuffer)
 {
 	FVulkanCmdBuffer* CmdBuffer = Context.GetCommandBufferManager()->GetUploadCmdBuffer();
 	ensure(CmdBuffer->IsOutsideRenderPass());
@@ -190,19 +190,19 @@ inline void FVulkanSurface::InternalLockWrite(FVulkanCommandListContext& Context
 	}
 }
 
-void FVulkanSurface::ErrorInvalidViewType() const
+void FVulkanTexture::ErrorInvalidViewType() const
 {
-	UE_LOG(LogVulkanRHI, Error, TEXT("Invalid ViewType %d"), (uint32)ViewType);
+	UE_LOG(LogVulkanRHI, Error, TEXT("Invalid ViewType %d"), (uint32)GetViewType());
 }
 
 
 struct FRHICommandLockWriteTexture final : public FRHICommand<FRHICommandLockWriteTexture>
 {
-	FVulkanSurface* Surface;
+	FVulkanTexture* Surface;
 	VkBufferImageCopy Region;
 	VulkanRHI::FStagingBuffer* StagingBuffer;
 
-	FRHICommandLockWriteTexture(FVulkanSurface* InSurface, const VkBufferImageCopy& InRegion, VulkanRHI::FStagingBuffer* InStagingBuffer)
+	FRHICommandLockWriteTexture(FVulkanTexture* InSurface, const VkBufferImageCopy& InRegion, VulkanRHI::FStagingBuffer* InStagingBuffer)
 		: Surface(InSurface)
 		, Region(InRegion)
 		, StagingBuffer(InStagingBuffer)
@@ -211,58 +211,54 @@ struct FRHICommandLockWriteTexture final : public FRHICommand<FRHICommandLockWri
 
 	void Execute(FRHICommandListBase& RHICmdList)
 	{
-		FVulkanSurface::InternalLockWrite(FVulkanCommandListContext::GetVulkanContext(RHICmdList.GetContext()), Surface, Region, StagingBuffer);
+		FVulkanTexture::InternalLockWrite(FVulkanCommandListContext::GetVulkanContext(RHICmdList.GetContext()), Surface, Region, StagingBuffer);
 	}
 };
 
-void FVulkanSurface::GenerateImageCreateInfo(
+void FVulkanTexture::GenerateImageCreateInfo(
 	FImageCreateInfo& OutImageCreateInfo,
 	FVulkanDevice& InDevice,
-	VkImageViewType ResourceType,
-	EPixelFormat InFormat,
-	uint32 SizeX, uint32 SizeY, uint32 SizeZ,
-	uint32 ArraySize,
-	uint32 NumMips,
-	uint32 NumSamples,
-	ETextureCreateFlags UEFlags,
+	const FRHITextureDesc& InDesc,
 	VkFormat* OutStorageFormat,
 	VkFormat* OutViewFormat,
 	bool bForceLinearTexture)
 {
 	const VkPhysicalDeviceProperties& DeviceProperties = InDevice.GetDeviceProperties();
-	const FPixelFormatInfo& FormatInfo = GPixelFormats[InFormat];
+	const FPixelFormatInfo& FormatInfo = GPixelFormats[InDesc.Format];
 	VkFormat TextureFormat = (VkFormat)FormatInfo.PlatformFormat;
 
+	const ETextureCreateFlags UEFlags = InDesc.Flags;
 	if(EnumHasAnyFlags(UEFlags, TexCreate_CPUReadback))
 	{
 		bForceLinearTexture = true;
 	}
 
-	checkf(TextureFormat != VK_FORMAT_UNDEFINED, TEXT("PixelFormat %d, is not supported for images"), (int32)InFormat);
+	checkf(TextureFormat != VK_FORMAT_UNDEFINED, TEXT("PixelFormat %d, is not supported for images"), (int32)InDesc.Format);
 	VkImageCreateInfo& ImageCreateInfo = OutImageCreateInfo.ImageCreateInfo;
 	ZeroVulkanStruct(ImageCreateInfo, VK_STRUCTURE_TYPE_IMAGE_CREATE_INFO);
 
+	const VkImageViewType ResourceType = UETextureDimensionToVkImageViewType(InDesc.Dimension);
 	switch(ResourceType)
 	{
 	case VK_IMAGE_VIEW_TYPE_1D:
 		ImageCreateInfo.imageType = VK_IMAGE_TYPE_1D;
-		check(SizeX <= DeviceProperties.limits.maxImageDimension1D);
+		check((uint32)InDesc.Extent.X <= DeviceProperties.limits.maxImageDimension1D);
 		break;
 	case VK_IMAGE_VIEW_TYPE_CUBE:
 	case VK_IMAGE_VIEW_TYPE_CUBE_ARRAY:
-		check(SizeX == SizeY);
-		check(SizeX <= DeviceProperties.limits.maxImageDimensionCube);
-		check(SizeY <= DeviceProperties.limits.maxImageDimensionCube);
+		check(InDesc.Extent.X == InDesc.Extent.Y);
+		check((uint32)InDesc.Extent.X <= DeviceProperties.limits.maxImageDimensionCube);
+		check((uint32)InDesc.Extent.Y <= DeviceProperties.limits.maxImageDimensionCube);
 		ImageCreateInfo.imageType = VK_IMAGE_TYPE_2D;
 		break;
 	case VK_IMAGE_VIEW_TYPE_2D:
 	case VK_IMAGE_VIEW_TYPE_2D_ARRAY:
-		check(SizeX <= DeviceProperties.limits.maxImageDimension2D);
-		check(SizeY <= DeviceProperties.limits.maxImageDimension2D);
+		check((uint32)InDesc.Extent.X <= DeviceProperties.limits.maxImageDimension2D);
+		check((uint32)InDesc.Extent.Y <= DeviceProperties.limits.maxImageDimension2D);
 		ImageCreateInfo.imageType = VK_IMAGE_TYPE_2D;
 		break;
 	case VK_IMAGE_VIEW_TYPE_3D:
-		check(SizeY <= DeviceProperties.limits.maxImageDimension3D);
+		check((uint32)InDesc.Extent.Y <= DeviceProperties.limits.maxImageDimension3D);
 		ImageCreateInfo.imageType = VK_IMAGE_TYPE_3D;
 		break;
 	default:
@@ -270,12 +266,12 @@ void FVulkanSurface::GenerateImageCreateInfo(
 		break;
 	}
 
-	VkFormat srgbFormat = UEToVkTextureFormat(InFormat, EnumHasAllFlags(UEFlags, TexCreate_SRGB));
-	VkFormat nonSrgbFormat = UEToVkTextureFormat(InFormat, false);
+	VkFormat srgbFormat = UEToVkTextureFormat(InDesc.Format, EnumHasAllFlags(UEFlags, TexCreate_SRGB));
+	VkFormat nonSrgbFormat = UEToVkTextureFormat(InDesc.Format, false);
 
 	ImageCreateInfo.format = EnumHasAnyFlags(UEFlags, TexCreate_UAV) ? nonSrgbFormat : srgbFormat;
 
-	checkf(ImageCreateInfo.format != VK_FORMAT_UNDEFINED, TEXT("Pixel Format %d not defined!"), (int32)InFormat);
+	checkf(ImageCreateInfo.format != VK_FORMAT_UNDEFINED, TEXT("Pixel Format %d not defined!"), (int32)InDesc.Format);
 	if (OutViewFormat)
 	{
 		*OutViewFormat = srgbFormat;
@@ -285,12 +281,12 @@ void FVulkanSurface::GenerateImageCreateInfo(
 		*OutStorageFormat = nonSrgbFormat;
 	}
 
-	ImageCreateInfo.extent.width = SizeX;
-	ImageCreateInfo.extent.height = SizeY;
-	ImageCreateInfo.extent.depth = ResourceType == VK_IMAGE_VIEW_TYPE_3D ? SizeZ : 1;
-	ImageCreateInfo.mipLevels = NumMips;
+	ImageCreateInfo.extent.width = InDesc.Extent.X;
+	ImageCreateInfo.extent.height = InDesc.Extent.Y;
+	ImageCreateInfo.extent.depth = ResourceType == VK_IMAGE_VIEW_TYPE_3D ? InDesc.Depth : 1;
+	ImageCreateInfo.mipLevels = InDesc.NumMips;
 	uint32 LayerCount = (ResourceType == VK_IMAGE_VIEW_TYPE_CUBE || ResourceType == VK_IMAGE_VIEW_TYPE_CUBE_ARRAY) ? 6 : 1;
-	ImageCreateInfo.arrayLayers = ArraySize * LayerCount;
+	ImageCreateInfo.arrayLayers = InDesc.ArraySize * LayerCount;
 	check(ImageCreateInfo.arrayLayers <= DeviceProperties.limits.maxImageArrayLayers);
 
 	ImageCreateInfo.flags = (ResourceType == VK_IMAGE_VIEW_TYPE_CUBE || ResourceType == VK_IMAGE_VIEW_TYPE_CUBE_ARRAY) ? VK_IMAGE_CREATE_CUBE_COMPATIBLE_BIT : 0;
@@ -402,6 +398,7 @@ void FVulkanSurface::GenerateImageCreateInfo(
 	ImageCreateInfo.queueFamilyIndexCount = 0;
 	ImageCreateInfo.pQueueFamilyIndices = nullptr;
 
+	uint8 NumSamples = InDesc.NumSamples;
 	if (ImageCreateInfo.tiling == VK_IMAGE_TILING_LINEAR && NumSamples > 1)
 	{
 		UE_LOG(LogVulkanRHI, Warning, TEXT("Not allowed to create Linear textures with %d samples, reverting to 1 sample"), NumSamples);
@@ -485,13 +482,13 @@ void FVulkanSurface::GenerateImageCreateInfo(
 
 struct FRHICommandSetInitialImageState final : public FRHICommand<FRHICommandSetInitialImageState>
 {
-	FVulkanSurface* Surface;
+	FVulkanTexture* Surface;
 	VkImageLayout InitialLayout;
 	bool bOnlyAddToLayoutManager;
 	bool bClear;
 	FClearValueBinding ClearValueBinding;
 
-	FRHICommandSetInitialImageState(FVulkanSurface* InSurface, VkImageLayout InInitialLayout, bool bInOnlyAddToLayoutManager, bool bInClear, const FClearValueBinding& InClearValueBinding)
+	FRHICommandSetInitialImageState(FVulkanTexture* InSurface, VkImageLayout InInitialLayout, bool bInOnlyAddToLayoutManager, bool bInClear, const FClearValueBinding& InClearValueBinding)
 		: Surface(InSurface)
 		, InitialLayout(InInitialLayout)
 		, bOnlyAddToLayoutManager(bInOnlyAddToLayoutManager)
@@ -577,191 +574,16 @@ static VkImageLayout GetInitialLayoutFromRHIAccess(ERHIAccess RHIAccess, ETextur
 	return VK_IMAGE_LAYOUT_UNDEFINED;
 }
 
-FVulkanSurface::FVulkanSurface(FVulkanDevice& InDevice, FVulkanEvictable* Owner, VkImageViewType ResourceType, EPixelFormat InFormat,
-								uint32 SizeX, uint32 SizeY, uint32 SizeZ, uint32 InArraySize, uint32 InNumMips,
-								uint32 InNumSamples, ETextureCreateFlags InUEFlags, ERHIAccess InResourceState, 
-								const FRHIResourceCreateInfo& CreateInfo, const FRHITransientHeapAllocation* InTransientHeapAllocation)
-	: Device(&InDevice)
-	, Image(VK_NULL_HANDLE)
-	, StorageFormat(VK_FORMAT_UNDEFINED)
-	, ViewFormat(VK_FORMAT_UNDEFINED)
-	, Width(SizeX)
-	, Height(SizeY)
-	, Depth(SizeZ)
-	, ArraySize(InArraySize)
-	, PixelFormat(InFormat)
-	, UEFlags(InUEFlags)
-	, MemProps(VK_MEMORY_PROPERTY_DEVICE_LOCAL_BIT)
-	, Tiling(VK_IMAGE_TILING_MAX_ENUM)	// Can be expanded to a per-platform definition
-	, ViewType(ResourceType)
-	, NumMips(InNumMips)
-	, NumSamples(InNumSamples)
-	, FullAspectMask(0)
-	, PartialAspectMask(0)
-	, CpuReadbackBuffer(nullptr) // for readback textures we use a staging buffer. this is because vulkan only requires implentations to support 1 mip level(which is useless), so we emulate using a buffer
-	, ImageOwnerType(EImageOwnerType::None)
-{
-	FImageCreateInfo ImageCreateInfo;
-	FVulkanSurface::GenerateImageCreateInfo(ImageCreateInfo,
-		InDevice, ResourceType,
-		InFormat, Width, Height, Depth,
-		ArraySize, NumMips, NumSamples, UEFlags,
-		&StorageFormat, &ViewFormat);
-	if (EnumHasAnyFlags(UEFlags, TexCreate_CPUReadback))
-	{
-		check(NumSamples == 1);	//not implemented
-		check(Depth == 1);		//not implemented
-		check(ArraySize == 1);	//not implemented
-		CpuReadbackBuffer = new FVulkanCpuReadbackBuffer;
-		uint32 Size = 0;
-		for(uint32 Mip = 0; Mip < NumMips; ++Mip)
-		{
-			uint32 LocalSize;
-			GetMipSize(Mip, LocalSize);
-			CpuReadbackBuffer->MipOffsets[Mip] = Size;
-			CpuReadbackBuffer->MipSize[Mip] = LocalSize;
-			Size += LocalSize;
-		}
-
-		VkDevice VulkanDevice = InDevice.GetInstanceHandle();
-		VkMemoryPropertyFlags BufferMemFlags = (VK_MEMORY_PROPERTY_HOST_VISIBLE_BIT | VK_MEMORY_PROPERTY_HOST_COHERENT_BIT | VK_MEMORY_PROPERTY_HOST_CACHED_BIT);
-		
-		VkBufferCreateInfo BufferCreateInfo;
-		ZeroVulkanStruct(BufferCreateInfo, VK_STRUCTURE_TYPE_BUFFER_CREATE_INFO);
-		BufferCreateInfo.size = Size;
-		BufferCreateInfo.usage = VK_BUFFER_USAGE_TRANSFER_DST_BIT;
-
-		VERIFYVULKANRESULT(VulkanRHI::vkCreateBuffer(VulkanDevice, &BufferCreateInfo, VULKAN_CPU_ALLOCATOR, &CpuReadbackBuffer->Buffer));
-		VulkanRHI::vkGetBufferMemoryRequirements(VulkanDevice, CpuReadbackBuffer->Buffer, &MemoryRequirements);
-		// Set minimum alignment to 16 bytes, as some buffers are used with CPU SIMD instructions
-		MemoryRequirements.alignment = FMath::Max<VkDeviceSize>(16, MemoryRequirements.alignment);
-		if (!InDevice.GetMemoryManager().AllocateBufferMemory(Allocation, Owner, MemoryRequirements, BufferMemFlags, EVulkanAllocationMetaBufferStaging, false, __FILE__, __LINE__))
-		{
-			InDevice.GetMemoryManager().HandleOOM();
-		}
-		Allocation.BindBuffer(Device, CpuReadbackBuffer->Buffer);
-		void* Memory = Allocation.GetMappedPointer(Device);
-		FMemory::Memzero(Memory, MemoryRequirements.size);
-		return;
-	}
-
-	ImageOwnerType = EImageOwnerType::LocalOwner;
-	VERIFYVULKANRESULT(VulkanRHI::vkCreateImage(InDevice.GetInstanceHandle(), &ImageCreateInfo.ImageCreateInfo, VULKAN_CPU_ALLOCATOR, &Image));
-
-	// Fetch image size
-	VulkanRHI::vkGetImageMemoryRequirements(InDevice.GetInstanceHandle(), Image, &MemoryRequirements);
-
-	VULKAN_SET_DEBUG_NAME(InDevice, VK_OBJECT_TYPE_IMAGE, Image, TEXT("%s:(FVulkanSurface*)0x%p"), CreateInfo.DebugName ? CreateInfo.DebugName :TEXT("?"), this);
-
-	FullAspectMask = VulkanRHI::GetAspectMaskFromUEFormat(PixelFormat, true, true);
-	PartialAspectMask = VulkanRHI::GetAspectMaskFromUEFormat(PixelFormat, false, true);
-
-	// If VK_IMAGE_TILING_OPTIMAL is specified,
-	// memoryTypeBits in vkGetImageMemoryRequirements will become 1
-	// which does not support VK_MEMORY_PROPERTY_HOST_VISIBLE_BIT.
-	if (ImageCreateInfo.ImageCreateInfo.tiling != VK_IMAGE_TILING_OPTIMAL)
-	{
-		MemProps |= VK_MEMORY_PROPERTY_HOST_VISIBLE_BIT;
-	}
-
-	const bool bRenderTarget = EnumHasAnyFlags(UEFlags, TexCreate_RenderTargetable | TexCreate_DepthStencilTargetable | TexCreate_ResolveTargetable);
-	const bool bUAV = EnumHasAnyFlags(UEFlags, TexCreate_UAV);
-	const bool bCPUReadback = EnumHasAnyFlags(UEFlags, TexCreate_CPUReadback);
-	const bool bDynamic = EnumHasAnyFlags(UEFlags, TexCreate_Dynamic);
-	const bool bExternal = EnumHasAnyFlags(UEFlags, TexCreate_External);
-
-	VkMemoryPropertyFlags MemoryFlags = bCPUReadback ? VK_MEMORY_PROPERTY_HOST_VISIBLE_BIT : VK_MEMORY_PROPERTY_DEVICE_LOCAL_BIT;
-
-	bool bMemoryless = EnumHasAnyFlags(UEFlags, TexCreate_Memoryless) && InDevice.GetDeviceMemoryManager().SupportsMemoryless();
-	if (bMemoryless)
-	{
-		if (ensureMsgf(bRenderTarget, TEXT("Memoryless surfaces can only be used for render targets")) && ensureMsgf(!bCPUReadback, TEXT("Memoryless surfaces cannot be read back on CPU")))
-		{
-			MemoryFlags |= VK_MEMORY_PROPERTY_LAZILY_ALLOCATED_BIT;
-		}
-		else
-		{
-			bMemoryless = false;
-		}
-	}
-	if(Owner == nullptr)
-	{
-		Owner = this;
-	}
-	check(bRenderTarget || bUAV || Owner != nullptr);
-
-	if (InTransientHeapAllocation != nullptr)
-	{
-		check(!bCPUReadback);
-		check(!bMemoryless);
-		check(InTransientHeapAllocation->Offset % MemoryRequirements.alignment == 0);
-		check(InTransientHeapAllocation->Size >= MemoryRequirements.size);
-		Allocation = FVulkanTransientHeap::GetVulkanAllocation(*InTransientHeapAllocation);
-	}
-	else
-	{
-	EVulkanAllocationMetaType MetaType = (bRenderTarget||bUAV) ? EVulkanAllocationMetaImageRenderTarget : EVulkanAllocationMetaImageOther;
-#if VULKAN_SUPPORTS_DEDICATED_ALLOCATION
-	extern int32 GVulkanEnableDedicatedImageMemory;
-	// Per https://developer.nvidia.com/what%E2%80%99s-your-vulkan-memory-type
-	VkDeviceSize SizeToBeConsideredForDedicated = 12 * 1024 * 1024;
-	if ((bRenderTarget || MemoryRequirements.size >= SizeToBeConsideredForDedicated) && !bMemoryless && InDevice.GetOptionalExtensions().HasKHRDedicatedAllocation && GVulkanEnableDedicatedImageMemory)
-	{
-		if(!InDevice.GetMemoryManager().AllocateDedicatedImageMemory(Allocation, Owner, Image, MemoryRequirements, MemoryFlags, MetaType, bExternal, __FILE__, __LINE__))
-		{
-			checkNoEntry();
-		}
-	}
-	else
-#endif
-	{
-		if(!InDevice.GetMemoryManager().AllocateImageMemory(Allocation, Owner, MemoryRequirements, MemoryFlags, MetaType, bExternal, __FILE__, __LINE__))
-		{
-			checkNoEntry();
-		}
-	}
-	}
-	Allocation.BindImage(Device, Image);
-
-	// update rhi stats
-	VulkanTextureAllocated(MemoryRequirements.size, ResourceType, bRenderTarget);
-
-	Tiling = ImageCreateInfo.ImageCreateInfo.tiling;
-	check(Tiling == VK_IMAGE_TILING_LINEAR || Tiling == VK_IMAGE_TILING_OPTIMAL);
-
-	const VkImageLayout InitialLayout = GetInitialLayoutFromRHIAccess(InResourceState, UEFlags, SupportsSampling());
-	const bool bDoInitialClear = VKHasAnyFlags(ImageCreateInfo.ImageCreateInfo.usage, VK_IMAGE_USAGE_SAMPLED_BIT) && EnumHasAnyFlags(UEFlags, TexCreate_RenderTargetable | TexCreate_DepthStencilTargetable);
-
-	if (InitialLayout != VK_IMAGE_LAYOUT_UNDEFINED || bDoInitialClear)
-	{
-		FRHICommandList& RHICmdList = FRHICommandListExecutor::GetImmediateCommandList();
-		if (!IsInRenderingThread() || (RHICmdList.Bypass() || !IsRunningRHIInSeparateThread()))
-		{
-			SetInitialImageState(Device->GetImmediateContext(), InitialLayout, bDoInitialClear, CreateInfo.ClearValueBinding);
-		}
-		else
-		{
-			check(IsInRenderingThread());
-			ALLOC_COMMAND_CL(RHICmdList, FRHICommandSetInitialImageState)(this, InitialLayout, false, bDoInitialClear, CreateInfo.ClearValueBinding);
-		}
-	}
-}
-
-
-
-void FVulkanSurface::InternalMoveSurface(FVulkanDevice& InDevice, FVulkanCommandListContext& Context, FVulkanAllocation& DestAllocation)
+void FVulkanTexture::InternalMoveSurface(FVulkanDevice& InDevice, FVulkanCommandListContext& Context, FVulkanAllocation& DestAllocation)
 {
 	FImageCreateInfo ImageCreateInfo;
-	FVulkanSurface::GenerateImageCreateInfo(ImageCreateInfo,
-		InDevice, ViewType,
-		PixelFormat, Width, Height, Depth,
-		ArraySize, NumMips, NumSamples, UEFlags,
-		&StorageFormat, &ViewFormat);
+	FVulkanTexture::GenerateImageCreateInfo(ImageCreateInfo, InDevice, GetDesc(), &StorageFormat, &ViewFormat);
 
 	VkImage MovedImage;
 	VERIFYVULKANRESULT(VulkanRHI::vkCreateImage(InDevice.GetInstanceHandle(), &ImageCreateInfo.ImageCreateInfo, VULKAN_CPU_ALLOCATOR, &MovedImage));
 	checkf(Tiling == ImageCreateInfo.ImageCreateInfo.tiling, TEXT("Move has changed image tiling:  before [%d] != after [%d]"), (int32)Tiling, (int32)ImageCreateInfo.ImageCreateInfo.tiling);
 
+	const ETextureCreateFlags UEFlags = GetDesc().Flags;
 	const bool bRenderTarget = EnumHasAnyFlags(UEFlags, TexCreate_RenderTargetable | TexCreate_DepthStencilTargetable | TexCreate_ResolveTargetable);
 	const bool bCPUReadback = EnumHasAnyFlags(UEFlags, TexCreate_CPUReadback);
 	const bool bMemoryless = EnumHasAnyFlags(UEFlags, TexCreate_Memoryless);
@@ -807,15 +629,16 @@ void FVulkanSurface::InternalMoveSurface(FVulkanDevice& InDevice, FVulkanCommand
 			Barrier.Execute(VkCmdBuffer);
 		}
 		{
+			const FIntVector SizeXYZ = GetSizeXYZ();
 			VkImageCopy Regions[MAX_TEXTURE_MIP_COUNT];
-			check(NumMips <= MAX_TEXTURE_MIP_COUNT);
+			check(GetDesc().NumMips <= MAX_TEXTURE_MIP_COUNT);
 			FMemory::Memzero(Regions);
-			for (uint32 i = 0; i < NumMips; ++i)
+			for (uint32 i = 0; i < GetDesc().NumMips; ++i)
 			{
 				VkImageCopy& Region = Regions[i];
-				Region.extent.width = FMath::Max(1u, Width >> i);
-				Region.extent.height = FMath::Max(1u, Height >> i);
-				Region.extent.depth = FMath::Max(1u, Depth >> i);
+				Region.extent.width = FMath::Max(1, SizeXYZ.X >> i);
+				Region.extent.height = FMath::Max(1, SizeXYZ.Y >> i);
+				Region.extent.depth = FMath::Max(1, SizeXYZ.Z >> i);
 				Region.srcSubresource.aspectMask = FullAspectMask;
 				Region.dstSubresource.aspectMask = FullAspectMask;
 				Region.srcSubresource.baseArrayLayer = 0;
@@ -828,7 +651,7 @@ void FVulkanSurface::InternalMoveSurface(FVulkanDevice& InDevice, FVulkanCommand
 			VulkanRHI::vkCmdCopyImage(VkCmdBuffer,
 				Image, VK_IMAGE_LAYOUT_TRANSFER_SRC_OPTIMAL,
 				MovedImage, VK_IMAGE_LAYOUT_TRANSFER_DST_OPTIMAL,
-				NumMips, &Regions[0]);
+				GetDesc().NumMips, &Regions[0]);
 		}
 
 		// Put the destination image in exactly the same layout the original image was
@@ -857,185 +680,7 @@ void FVulkanSurface::InternalMoveSurface(FVulkanDevice& InDevice, FVulkanCommand
 	Image = MovedImage;
 }
 
-void FVulkanSurface::MoveSurface(FVulkanDevice& InDevice, FVulkanCommandListContext& Context, FVulkanAllocation& NewAllocation)
-{
-	const bool bRenderTarget = EnumHasAnyFlags(UEFlags, TexCreate_RenderTargetable | TexCreate_DepthStencilTargetable | TexCreate_ResolveTargetable);
-	const bool bUAV = EnumHasAnyFlags(UEFlags, TexCreate_UAV);
-	checkf(bRenderTarget || bUAV, TEXT("Surface must be a RenderTarget or a UAV in order to be moved.  UEFlags=0x%x"), (int32)UEFlags);
-	checkf(Tiling == VK_IMAGE_TILING_OPTIMAL, TEXT("Tiling [%d] is not supported for move, only VK_IMAGE_TILING_OPTIMAL"), (int32)Tiling);
-
-	InternalMoveSurface(InDevice, Context, NewAllocation);
-
-	// Swap in the new allocation for this surface
-	Allocation.Swap(NewAllocation);
-}
-
-
-void FVulkanSurface::OnFullDefrag(FVulkanDevice& InDevice, FVulkanCommandListContext& Context, uint32 NewOffset)
-{
-	const bool bRenderTarget = EnumHasAnyFlags(UEFlags, TexCreate_RenderTargetable | TexCreate_DepthStencilTargetable | TexCreate_ResolveTargetable);
-	const bool bUAV = EnumHasAnyFlags(UEFlags, TexCreate_UAV);
-	checkf(bRenderTarget || bUAV, TEXT("Surface must be a RenderTarget or a UAV in order to be defragged.  UEFlags=0x%x"), (int32)UEFlags);
-	checkf(Tiling == VK_IMAGE_TILING_OPTIMAL, TEXT("Tiling [%d] is not supported for defrag, only VK_IMAGE_TILING_OPTIMAL"), (int32)Tiling);
-
-	Allocation.Offset = NewOffset;
-	InternalMoveSurface(InDevice, Context, Allocation);
-
-	//note: this exploits that the unmoved image is still bound to the old allocation, which is freed by the caller in this case.
-}
-
-
-void FVulkanSurface::EvictSurface(FVulkanDevice& InDevice)
-{
-	check(0 == CpuReadbackBuffer);
-	checkf(MemProps == VK_MEMORY_PROPERTY_DEVICE_LOCAL_BIT, TEXT("Can't evict surface that isn't device local.  MemoryProperties=%d"), (int32)MemProps);
-	checkf(VulkanRHI::GetAspectMaskFromUEFormat(PixelFormat, true, true) == FullAspectMask, TEXT("FullAspectMask (%d) does not match with PixelFormat (%d)"), (int32)FullAspectMask, (int32)PixelFormat);
-	checkf(VulkanRHI::GetAspectMaskFromUEFormat(PixelFormat, false, true) == PartialAspectMask, TEXT("PartialAspectMask (%d) does not match with PixelFormat (%d)"), (int32)PartialAspectMask, (int32)PixelFormat);
-
-	const bool bRenderTarget = EnumHasAnyFlags(UEFlags, TexCreate_RenderTargetable | TexCreate_DepthStencilTargetable | TexCreate_ResolveTargetable);
-	const bool bUAV = EnumHasAnyFlags(UEFlags, TexCreate_UAV);
-	//none of this is supported for eviction
-	checkf(!bRenderTarget, TEXT("RenderTargets do not support evict."));
-	checkf(!bUAV, TEXT("UAV do not support evict."));
-
-	FRHICommandListImmediate& RHICmdList = FRHICommandListExecutor::GetImmediateCommandList();
-	FVulkanCommandListContext& Context = (FVulkanCommandListContext&)RHICmdList.GetContext().GetLowestLevelContext();
-
-	MemProps = InDevice.GetDeviceMemoryManager().GetEvictedMemoryProperties();
-
-	// Create a new host allocation to move the surface to
-	FVulkanAllocation HostAllocation;
-	const EVulkanAllocationMetaType MetaType = EVulkanAllocationMetaImageOther;
-	if (!InDevice.GetMemoryManager().AllocateImageMemory(HostAllocation, this, MemoryRequirements, MemProps, MetaType, false, __FILE__, __LINE__))
-	{
-		InDevice.GetMemoryManager().HandleOOM();
-		checkNoEntry();
-	}
-
-	InternalMoveSurface(InDevice, Context, HostAllocation);
-
-	// Delete the original allocation and swap in the new host allocation
-	Device->GetMemoryManager().FreeVulkanAllocation(Allocation);
-	Allocation.Swap(HostAllocation);
-
-	VULKAN_SET_DEBUG_NAME(InDevice, VK_OBJECT_TYPE_IMAGE, Image, TEXT("(FVulkanSurface*)0x%p [hostimage]"), this);
-}
-
-// This is usually used for the framebuffer image
-FVulkanSurface::FVulkanSurface(FVulkanDevice& InDevice, VkImageViewType ResourceType, EPixelFormat InFormat,
-								uint32 SizeX, uint32 SizeY, uint32 SizeZ, uint32 InArraySize, uint32 InNumMips, uint32 InNumSamples,
-								VkImage InImage, ETextureCreateFlags InUEFlags, EImageOwnerType InImageOwnerType, const FRHIResourceCreateInfo& CreateInfo)
-	: Device(&InDevice)
-	, Image(InImage)
-	, StorageFormat(VK_FORMAT_UNDEFINED)
-	, ViewFormat(VK_FORMAT_UNDEFINED)
-	, Width(SizeX)
-	, Height(SizeY)
-	, Depth(SizeZ)
-	, ArraySize(InArraySize)
-	, PixelFormat(InFormat)
-	, UEFlags(InUEFlags)
-	, MemProps(VK_MEMORY_PROPERTY_DEVICE_LOCAL_BIT)
-	, Tiling(VK_IMAGE_TILING_MAX_ENUM)	// Can be expanded to a per-platform definition
-	, ViewType(ResourceType)
-	, NumMips(InNumMips)
-	, NumSamples(InNumSamples)
-	, FullAspectMask(0)
-	, PartialAspectMask(0)
-	, CpuReadbackBuffer(nullptr)
-	, ImageOwnerType(InImageOwnerType)
-{
-	StorageFormat = UEToVkTextureFormat(PixelFormat, false);
-
-	checkf(PixelFormat == PF_Unknown || StorageFormat != VK_FORMAT_UNDEFINED, TEXT("PixelFormat %d, is not supported for images"), (int32)PixelFormat);
-
-	ViewFormat = UEToVkTextureFormat(PixelFormat, EnumHasAllFlags(UEFlags, TexCreate_SRGB));
-	FullAspectMask = VulkanRHI::GetAspectMaskFromUEFormat(PixelFormat, true, true);
-	PartialAspectMask = VulkanRHI::GetAspectMaskFromUEFormat(PixelFormat, false, true);
-
-	// Purely informative patching, we know that "TexCreate_Presentable" uses optimal tiling
-	if (EnumHasAllFlags(UEFlags, TexCreate_Presentable) && GetTiling() == VK_IMAGE_TILING_MAX_ENUM)
-	{
-		Tiling = VK_IMAGE_TILING_OPTIMAL;
-	}
-
-	if (Image != VK_NULL_HANDLE)
-	{
-#if VULKAN_ENABLE_WRAP_LAYER
-		FImageCreateInfo ImageCreateInfo;
-		FVulkanSurface::GenerateImageCreateInfo(
-			ImageCreateInfo,
-			InDevice, ResourceType,
-			InFormat, SizeX, SizeY, SizeZ,
-			ArraySize, NumMips, NumSamples, UEFlags,
-			&StorageFormat, &ViewFormat);
-		FWrapLayer::CreateImage(VK_SUCCESS, InDevice.GetInstanceHandle(), &ImageCreateInfo.ImageCreateInfo, &Image);
-#endif
-		VULKAN_SET_DEBUG_NAME(InDevice, VK_OBJECT_TYPE_IMAGE, Image, TEXT("%s:(FVulkanSurface*)0x%p"), CreateInfo.DebugName ? CreateInfo.DebugName : TEXT("?"), this);
-
-		VkImageLayout InitialLayout = VK_IMAGE_LAYOUT_UNDEFINED;
-		bool bOnlyAddToLayoutManager, bDoInitialClear;
-		if (EnumHasAnyFlags(UEFlags, TexCreate_RenderTargetable | TexCreate_DepthStencilTargetable))
-		{
-			InitialLayout = EnumHasAnyFlags(UEFlags, TexCreate_DepthStencilTargetable) ? VK_IMAGE_LAYOUT_DEPTH_STENCIL_ATTACHMENT_OPTIMAL : VK_IMAGE_LAYOUT_COLOR_ATTACHMENT_OPTIMAL;
-			bOnlyAddToLayoutManager = false;
-			bDoInitialClear = true;
-		}
-		else if (EnumHasAnyFlags(UEFlags, TexCreate_Foveation) && ValidateShadingRateDataType())
-		{
-			// If it's a shading rate attachment texture, do not clear but add to layoutmgr, and set correct shading rate layout.
-#if VULKAN_SUPPORTS_FRAGMENT_SHADING_RATE
-			if (GRHIVariableRateShadingImageDataType == VRSImage_Palette)
-			{
-				InitialLayout = VK_IMAGE_LAYOUT_FRAGMENT_SHADING_RATE_ATTACHMENT_OPTIMAL_KHR;
-			}
-#endif
-
-#if VULKAN_SUPPORTS_FRAGMENT_DENSITY_MAP
-			if (GRHIVariableRateShadingImageDataType == VRSImage_Fractional)
-			{
-				InitialLayout = VK_IMAGE_LAYOUT_FRAGMENT_DENSITY_MAP_OPTIMAL_EXT;
-			}
-#endif
-
-			bOnlyAddToLayoutManager = true;
-			bDoInitialClear = false;
-		}
-		else
-		{
-			// If we haven't seen this image before, we assume it's an SRV (VK_IMAGE_LAYOUT_SHADER_READ_ONLY_OPTIMAL) and the call below
-			// tells the layout manager about it. If we've seen it before, the call won't do anything, since the manager already knows the layout.
-			InitialLayout = VK_IMAGE_LAYOUT_SHADER_READ_ONLY_OPTIMAL;
-			bOnlyAddToLayoutManager = true;
-			bDoInitialClear = false;
-		}
-
-		FRHICommandList& RHICmdList = FRHICommandListExecutor::GetImmediateCommandList();
-		if (!IsInRenderingThread() || (RHICmdList.Bypass() || !IsRunningRHIInSeparateThread()))
-		{
-			if (bOnlyAddToLayoutManager)
-			{
-				InDevice.GetImmediateContext().GetLayoutManager().GetOrAddFullLayout(*this, InitialLayout);
-			}
-			else
-			{
-				SetInitialImageState(InDevice.GetImmediateContext(), InitialLayout, true, CreateInfo.ClearValueBinding);
-			}
-		}
-		else
-		{
-			check(IsInRenderingThread());
-			ALLOC_COMMAND_CL(RHICmdList, FRHICommandSetInitialImageState)(this, InitialLayout, bOnlyAddToLayoutManager, bDoInitialClear, CreateInfo.ClearValueBinding);
-		}
-	}
-}
-
-FVulkanSurface::~FVulkanSurface()
-{
-	Destroy();
-}
-
-void FVulkanSurface::Destroy()
+void FVulkanTexture::DestroySurface()
 {
 	const bool bIsLocalOwner = (ImageOwnerType == EImageOwnerType::LocalOwner);
 	const bool bHasExternalOwner = (ImageOwnerType == EImageOwnerType::ExternalOwner);
@@ -1049,7 +694,7 @@ void FVulkanSurface::Destroy()
 	}
 	else if (bIsLocalOwner || bHasExternalOwner)
 	{
-		const bool bRenderTarget = EnumHasAnyFlags(UEFlags, TexCreate_RenderTargetable | TexCreate_DepthStencilTargetable | TexCreate_ResolveTargetable);
+		const bool bRenderTarget = EnumHasAnyFlags(GetDesc().Flags, TexCreate_RenderTargetable | TexCreate_DepthStencilTargetable | TexCreate_ResolveTargetable);
 		FRHICommandList& RHICmdList = FRHICommandListExecutor::GetImmediateCommandList();
 		if (!IsInRenderingThread() || (RHICmdList.Bypass() || !IsRunningRHIInSeparateThread()))
 		{
@@ -1073,7 +718,7 @@ void FVulkanSurface::Destroy()
 				Image = VK_NULL_HANDLE;
 			}
 
-			VulkanTextureDestroyed(Size, ViewType, bRenderTarget);
+			VulkanTextureDestroyed(Size, GetViewType(), bRenderTarget);
 		}
 
 		ImageOwnerType = EImageOwnerType::None;
@@ -1081,19 +726,19 @@ void FVulkanSurface::Destroy()
 }
 
 
-void FVulkanSurface::InvalidateMappedMemory()
+void FVulkanTexture::InvalidateMappedMemory()
 {
 	Allocation.InvalidateMappedMemory(Device);
 
 }
-void* FVulkanSurface::GetMappedPointer()
+void* FVulkanTexture::GetMappedPointer()
 {
 	return Allocation.GetMappedPointer(Device);
 }
 
 
 
-VkDeviceMemory FVulkanSurface::GetAllocationHandle() const
+VkDeviceMemory FVulkanTexture::GetAllocationHandle() const
 {
 	if (Allocation.IsValid())
 	{
@@ -1105,7 +750,7 @@ VkDeviceMemory FVulkanSurface::GetAllocationHandle() const
 	}
 }
 
-uint64 FVulkanSurface::GetAllocationOffset() const
+uint64 FVulkanTexture::GetAllocationOffset() const
 {
 	if (Allocation.IsValid())
 	{
@@ -1117,11 +762,13 @@ uint64 FVulkanSurface::GetAllocationOffset() const
 	}
 }
 
-void FVulkanSurface::GetMipStride(uint32 MipIndex, uint32& Stride)
+void FVulkanTexture::GetMipStride(uint32 MipIndex, uint32& Stride)
 {
 	// Calculate the width of the MipMap.
+	const FRHITextureDesc& Desc = GetDesc();
+	const EPixelFormat PixelFormat = Desc.Format;
 	const uint32 BlockSizeX = GPixelFormats[PixelFormat].BlockSizeX;
-	const uint32 MipSizeX = FMath::Max(Width >> MipIndex, BlockSizeX);
+	const uint32 MipSizeX = FMath::Max<uint32>(Desc.Extent.X >> MipIndex, BlockSizeX);
 	uint32 NumBlocksX = (MipSizeX + BlockSizeX - 1) / BlockSizeX;
 
 	if (PixelFormat == PF_PVRTC2 || PixelFormat == PF_PVRTC4)
@@ -1135,7 +782,7 @@ void FVulkanSurface::GetMipStride(uint32 MipIndex, uint32& Stride)
 	Stride = NumBlocksX * BlockBytes;
 }
 
-void FVulkanSurface::GetMipOffset(uint32 MipIndex, uint32& Offset)
+void FVulkanTexture::GetMipOffset(uint32 MipIndex, uint32& Offset)
 {
 	uint32 offset = Offset = 0;
 	for(uint32 i = 0; i < MipIndex; i++)
@@ -1145,14 +792,16 @@ void FVulkanSurface::GetMipOffset(uint32 MipIndex, uint32& Offset)
 	}
 }
 
-void FVulkanSurface::GetMipSize(uint32 MipIndex, uint32& MipBytes)
+void FVulkanTexture::GetMipSize(uint32 MipIndex, uint32& MipBytes)
 {
 	// Calculate the dimensions of mip-map level.
+	const FRHITextureDesc& Desc = GetDesc();
+	const EPixelFormat PixelFormat = Desc.Format;
 	const uint32 BlockSizeX = GPixelFormats[PixelFormat].BlockSizeX;
 	const uint32 BlockSizeY = GPixelFormats[PixelFormat].BlockSizeY;
 	const uint32 BlockBytes = GPixelFormats[PixelFormat].BlockBytes;
-	const uint32 MipSizeX = FMath::Max(Width >> MipIndex, BlockSizeX);
-	const uint32 MipSizeY = FMath::Max(Height >> MipIndex, BlockSizeY);
+	const uint32 MipSizeX = FMath::Max<uint32>(Desc.Extent.X >> MipIndex, BlockSizeX);
+	const uint32 MipSizeY = FMath::Max<uint32>(Desc.Extent.Y >> MipIndex, BlockSizeY);
 	uint32 NumBlocksX = (MipSizeX + BlockSizeX - 1) / BlockSizeX;
 	uint32 NumBlocksY = (MipSizeY + BlockSizeY - 1) / BlockSizeY;
 
@@ -1179,7 +828,7 @@ void FVulkanSurface::GetMipSize(uint32 MipIndex, uint32& MipBytes)
 */
 }
 
-void FVulkanSurface::SetInitialImageState(FVulkanCommandListContext& Context, VkImageLayout InitialLayout, bool bClear, const FClearValueBinding& ClearValueBinding)
+void FVulkanTexture::SetInitialImageState(FVulkanCommandListContext& Context, VkImageLayout InitialLayout, bool bClear, const FClearValueBinding& ClearValueBinding)
 {
 	// Can't use TransferQueue as Vulkan requires that queue to also have Gfx or Compute capabilities...
 	//#todo-rco: This function is only used during loading currently, if used for regular RHIClear then use the ActiveCmdBuffer
@@ -1267,25 +916,24 @@ uint32 FVulkanDynamicRHI::RHIComputeMemorySize(FRHITexture* TextureRHI)
 		return 0;
 	}
 
-	return FVulkanTextureBase::Cast(TextureRHI)->Surface.GetMemorySize();
+	return FVulkanTexture::Cast(TextureRHI)->GetMemorySize();
 }
 
 /*-----------------------------------------------------------------------------
 	2D texture support.
 -----------------------------------------------------------------------------*/
 
-FTexture2DRHIRef FVulkanDynamicRHI::RHICreateTexture2D(uint32 SizeX, uint32 SizeY, uint8 Format, uint32 NumMips, uint32 NumSamples, ETextureCreateFlags Flags, ERHIAccess InResourceState, FRHIResourceCreateInfo& CreateInfo)
-
+FTextureRHIRef FVulkanDynamicRHI::RHICreateTexture(const FRHITextureCreateDesc& CreateDesc)
 {
-	LLM_SCOPE_VULKAN(GetMemoryTagForTextureFlags(Flags));
-	return new FVulkanTexture2D(*Device, (EPixelFormat)Format, SizeX, SizeY, NumMips, NumSamples, Flags, InResourceState, CreateInfo);
+	LLM_SCOPE_VULKAN(GetMemoryTagForTextureFlags(CreateDesc.Flags));
+	return new FVulkanTexture(*Device, CreateDesc, nullptr);
 }
 
-FTexture2DRHIRef FVulkanDynamicRHI::RHIAsyncCreateTexture2D(uint32 SizeX,uint32 SizeY,uint8 Format,uint32 NumMips,ETextureCreateFlags Flags, ERHIAccess InResourceState,void** InitialMipData,uint32 NumInitialMips)
+FTextureRHIRef FVulkanDynamicRHI::RHIAsyncCreateTexture2D(uint32 SizeX,uint32 SizeY,uint8 Format,uint32 NumMips,ETextureCreateFlags Flags, ERHIAccess InResourceState,void** InitialMipData,uint32 NumInitialMips)
 {
 	UE_LOG(LogVulkan, Fatal, TEXT("RHIAsyncCreateTexture2D is not supported"));
-	VULKAN_SIGNAL_UNIMPLEMENTED(); // Unsupported atm
-	return FTexture2DRHIRef();
+	VULKAN_SIGNAL_UNIMPLEMENTED();
+	return FTextureRHIRef();
 }
 
 void FVulkanDynamicRHI::RHICopySharedMips(FRHITexture2D* DestTexture2D, FRHITexture2D* SrcTexture2D)
@@ -1293,56 +941,7 @@ void FVulkanDynamicRHI::RHICopySharedMips(FRHITexture2D* DestTexture2D, FRHIText
 	VULKAN_SIGNAL_UNIMPLEMENTED();
 }
 
-FTexture2DArrayRHIRef FVulkanDynamicRHI::RHICreateTexture2DArray(uint32 SizeX, uint32 SizeY, uint32 SizeZ, uint8 Format, uint32 NumMips, uint32 NumSamples, ETextureCreateFlags Flags, ERHIAccess InResourceState, FRHIResourceCreateInfo& CreateInfo)
-{
-	LLM_SCOPE_VULKAN(GetMemoryTagForTextureFlags(Flags));
-	return new FVulkanTexture2DArray(*Device, (EPixelFormat)Format, SizeX, SizeY, SizeZ, NumMips, NumSamples, Flags, InResourceState, CreateInfo);
-}
-
-FTexture3DRHIRef FVulkanDynamicRHI::RHICreateTexture3D(uint32 SizeX, uint32 SizeY, uint32 SizeZ, uint8 Format, uint32 NumMips, ETextureCreateFlags Flags, ERHIAccess InResourceState, FRHIResourceCreateInfo& CreateInfo)
-{
-	LLM_SCOPE_VULKAN(GetMemoryTagForTextureFlags(Flags));
-	FVulkanTexture3D* Tex3d = new FVulkanTexture3D(*Device, (EPixelFormat)Format, SizeX, SizeY, SizeZ, NumMips, Flags, InResourceState, CreateInfo);
-
-	return Tex3d;
-}
-
-FRHITexture* FVulkanDynamicRHI::CreateTexture(const FRHITextureCreateInfo& InCreateInfo, FRHIResourceCreateInfo& InResourceCreateInfo, ERHIAccess InResourceState, const FRHITransientHeapAllocation* InTransientHeapAllocation)
-{
-	LLM_SCOPE_VULKAN(GetMemoryTagForTextureFlags(InCreateInfo.Flags));
-
-	FRHITexture* Texture = nullptr;
-
-	switch (InCreateInfo.Dimension)
-	{
-	case ETextureDimension::Texture2D:
-		Texture = new FVulkanTexture2D(*Device, InCreateInfo.Format, InCreateInfo.Extent.X, InCreateInfo.Extent.Y, InCreateInfo.NumMips, InCreateInfo.NumSamples, InCreateInfo.Flags, InResourceState, InResourceCreateInfo, InTransientHeapAllocation);
-		break;
-
-	case ETextureDimension::Texture2DArray:
-		Texture = new FVulkanTexture2DArray(*Device, InCreateInfo.Format, InCreateInfo.Extent.X, InCreateInfo.Extent.Y, InCreateInfo.ArraySize, InCreateInfo.NumMips, InCreateInfo.NumSamples, InCreateInfo.Flags, InResourceState, InResourceCreateInfo, InTransientHeapAllocation);
-		break;
-
-	case ETextureDimension::TextureCube:
-		Texture = new FVulkanTextureCube(*Device, InCreateInfo.Format, InCreateInfo.Extent.X, false, 1, InCreateInfo.NumMips, InCreateInfo.Flags, InResourceState, InResourceCreateInfo, InTransientHeapAllocation);
-		break;
-
-	case ETextureDimension::TextureCubeArray:
-		Texture = new FVulkanTextureCube(*Device, InCreateInfo.Format, InCreateInfo.Extent.X, true, InCreateInfo.ArraySize, InCreateInfo.NumMips, InCreateInfo.Flags, InResourceState, InResourceCreateInfo, InTransientHeapAllocation);
-		break;
-
-	case ETextureDimension::Texture3D:
-		Texture = new FVulkanTexture3D(*Device, InCreateInfo.Format, InCreateInfo.Extent.X, InCreateInfo.Extent.Y, InCreateInfo.Depth, InCreateInfo.NumMips, InCreateInfo.Flags, InResourceState, InResourceCreateInfo, InTransientHeapAllocation);
-		break;
-
-	default:
-		checkNoEntry();
-	}
-
-	return Texture;
-}
-
-static void DoAsyncReallocateTexture2D(FVulkanCommandListContext& Context, FVulkanTexture2D* OldTexture, FVulkanTexture2D* NewTexture, int32 NewMipCount, int32 NewSizeX, int32 NewSizeY, FThreadSafeCounter* RequestStatus)
+static void DoAsyncReallocateTexture2D(FVulkanCommandListContext& Context, FVulkanTexture* OldTexture, FVulkanTexture* NewTexture, int32 NewMipCount, int32 NewSizeX, int32 NewSizeY, FThreadSafeCounter* RequestStatus)
 {
 	LLM_SCOPE_VULKAN(ELLMTagVulkan::VulkanTextures);
 	//QUICK_SCOPE_CYCLE_COUNTER(STAT_FRHICommandGnmAsyncReallocateTexture2D_Execute);
@@ -1385,7 +984,7 @@ static void DoAsyncReallocateTexture2D(FVulkanCommandListContext& Context, FVulk
 	const VkImageSubresourceRange SourceSubResourceRange = FVulkanPipelineBarrier::MakeSubresourceRange(VK_IMAGE_ASPECT_COLOR_BIT, SourceFirstMip, NumSharedMips);
 	const VkImageSubresourceRange DestSubResourceRange = FVulkanPipelineBarrier::MakeSubresourceRange(VK_IMAGE_ASPECT_COLOR_BIT, DestFirstMip, NumSharedMips);
 
-	const FVulkanImageLayout& OldTextureOriginalLayout = Context.GetLayoutManager().GetOrAddFullLayout(OldTexture->Surface, VK_IMAGE_LAYOUT_UNDEFINED);
+	const FVulkanImageLayout& OldTextureOriginalLayout = Context.GetLayoutManager().GetOrAddFullLayout(*OldTexture, VK_IMAGE_LAYOUT_UNDEFINED);
 	ensure(!OldTextureOriginalLayout.AreAllSubresourcesSameLayout() || (OldTextureOriginalLayout.MainLayout != VK_IMAGE_LAYOUT_UNDEFINED));
 	FVulkanImageLayout OldTextureCopyLayout = OldTextureOriginalLayout;
 	OldTextureCopyLayout.Set(VK_IMAGE_LAYOUT_TRANSFER_SRC_OPTIMAL, SourceSubResourceRange);
@@ -1393,45 +992,45 @@ static void DoAsyncReallocateTexture2D(FVulkanCommandListContext& Context, FVulk
 	{
 		// Pre-copy barriers
 		FVulkanPipelineBarrier Barrier;
-		Barrier.AddImageLayoutTransition(OldTexture->Surface.Image, VK_IMAGE_ASPECT_COLOR_BIT, OldTextureOriginalLayout, OldTextureCopyLayout);
-		Barrier.AddImageLayoutTransition(NewTexture->Surface.Image, VK_IMAGE_LAYOUT_UNDEFINED, VK_IMAGE_LAYOUT_TRANSFER_DST_OPTIMAL, DestSubResourceRange);
+		Barrier.AddImageLayoutTransition(OldTexture->Image, VK_IMAGE_ASPECT_COLOR_BIT, OldTextureOriginalLayout, OldTextureCopyLayout);
+		Barrier.AddImageLayoutTransition(NewTexture->Image, VK_IMAGE_LAYOUT_UNDEFINED, VK_IMAGE_LAYOUT_TRANSFER_DST_OPTIMAL, DestSubResourceRange);
 		Barrier.Execute(CmdBuffer->GetHandle());
 	}
 
-	VulkanRHI::vkCmdCopyImage(StagingCommandBuffer, OldTexture->Surface.Image, VK_IMAGE_LAYOUT_TRANSFER_SRC_OPTIMAL, NewTexture->Surface.Image, VK_IMAGE_LAYOUT_TRANSFER_DST_OPTIMAL, NumSharedMips, Regions);
+	VulkanRHI::vkCmdCopyImage(StagingCommandBuffer, OldTexture->Image, VK_IMAGE_LAYOUT_TRANSFER_SRC_OPTIMAL, NewTexture->Image, VK_IMAGE_LAYOUT_TRANSFER_DST_OPTIMAL, NumSharedMips, Regions);
 
 	{
 		// Post-copy barriers
 		FVulkanPipelineBarrier Barrier;
-		Barrier.AddImageLayoutTransition(OldTexture->Surface.Image, VK_IMAGE_ASPECT_COLOR_BIT, OldTextureCopyLayout, OldTextureOriginalLayout);
-		Barrier.AddImageLayoutTransition(NewTexture->Surface.Image, VK_IMAGE_LAYOUT_TRANSFER_DST_OPTIMAL, VK_IMAGE_LAYOUT_SHADER_READ_ONLY_OPTIMAL, DestSubResourceRange);
+		Barrier.AddImageLayoutTransition(OldTexture->Image, VK_IMAGE_ASPECT_COLOR_BIT, OldTextureCopyLayout, OldTextureOriginalLayout);
+		Barrier.AddImageLayoutTransition(NewTexture->Image, VK_IMAGE_LAYOUT_TRANSFER_DST_OPTIMAL, VK_IMAGE_LAYOUT_SHADER_READ_ONLY_OPTIMAL, DestSubResourceRange);
 		Barrier.Execute(CmdBuffer->GetHandle());
 
 		// Add tracking for the appropriate subresources (intentionally leave added mips in VK_IMAGE_LAYOUT_UNDEFINED)
 		// NOTE: Overwriting whatever is contained in the layout manager for the new texture (circumvents issue with stale layouts and pointer reuse)
-		FVulkanImageLayout NewTextureLayout(VK_IMAGE_LAYOUT_UNDEFINED, NewTexture->Surface.GetNumMips(), NewTexture->Surface.GetNumberOfArrayLevels());
+		FVulkanImageLayout NewTextureLayout(VK_IMAGE_LAYOUT_UNDEFINED, NewTexture->GetNumMips(), NewTexture->GetNumberOfArrayLevels());
 		NewTextureLayout.Set(VK_IMAGE_LAYOUT_SHADER_READ_ONLY_OPTIMAL, DestSubResourceRange);
-		Context.GetLayoutManager().GetOrAddFullLayout(NewTexture->Surface, VK_IMAGE_LAYOUT_UNDEFINED) = NewTextureLayout;
+		Context.GetLayoutManager().GetOrAddFullLayout(*NewTexture, VK_IMAGE_LAYOUT_UNDEFINED) = NewTextureLayout;
 	}
 
 	// request is now complete
 	RequestStatus->Decrement();
 
 	// the next unlock for this texture can't block the GPU (it's during runtime)
-	//NewTexture->Surface.bSkipBlockOnUnlock = true;
+	//NewTexture->bSkipBlockOnUnlock = true;
 }
 
 struct FRHICommandVulkanAsyncReallocateTexture2D final : public FRHICommand<FRHICommandVulkanAsyncReallocateTexture2D>
 {
 	FVulkanCommandListContext& Context;
-	FVulkanTexture2D* OldTexture;
-	FVulkanTexture2D* NewTexture;
+	FVulkanTexture* OldTexture;
+	FVulkanTexture* NewTexture;
 	int32 NewMipCount;
 	int32 NewSizeX;
 	int32 NewSizeY;
 	FThreadSafeCounter* RequestStatus;
 
-	FORCEINLINE_DEBUGGABLE FRHICommandVulkanAsyncReallocateTexture2D(FVulkanCommandListContext& InContext, FVulkanTexture2D* InOldTexture, FVulkanTexture2D* InNewTexture, int32 InNewMipCount, int32 InNewSizeX, int32 InNewSizeY, FThreadSafeCounter* InRequestStatus)
+	FORCEINLINE_DEBUGGABLE FRHICommandVulkanAsyncReallocateTexture2D(FVulkanCommandListContext& InContext, FVulkanTexture* InOldTexture, FVulkanTexture* InNewTexture, int32 InNewMipCount, int32 InNewSizeX, int32 InNewSizeY, FThreadSafeCounter* InRequestStatus)
 		: Context(InContext)
 		, OldTexture(InOldTexture)
 		, NewTexture(InNewTexture)
@@ -1457,11 +1056,19 @@ FTexture2DRHIRef FVulkanDynamicRHI::AsyncReallocateTexture2D_RenderThread(FRHICo
 		return FDynamicRHI::AsyncReallocateTexture2D_RenderThread(RHICmdList, OldTextureRHI, NewMipCount, NewSizeX, NewSizeY, RequestStatus);
 	}
 
-	FVulkanTexture2D* OldTexture = ResourceCast(OldTextureRHI);
+	FVulkanTexture* OldTexture = FVulkanTexture::Cast(OldTextureRHI);
 
-	FRHIResourceCreateInfo CreateInfo(TEXT("AsyncReallocateTexture2D_RenderThread"));
-	FVulkanTexture2D* NewTexture = new FVulkanTexture2D(*Device, OldTexture->GetFormat(), NewSizeX, NewSizeY, NewMipCount, OldTexture->GetNumSamples(), OldTexture->GetFlags(), ERHIAccess::Unknown, CreateInfo);
+	FRHITextureCreateDesc Desc = FRHITextureCreateDesc::Create2D(
+		TEXT("AsyncReallocateTexture2D_RenderThread"),
+		{ NewSizeX, NewSizeY },
+		OldTexture->GetFormat(),
+		OldTexture->GetClearBinding(),
+		OldTexture->GetFlags(),
+		NewMipCount,
+		OldTexture->GetNumSamples()
+	);
 
+	FVulkanTexture* NewTexture = new FVulkanTexture(*Device, Desc, nullptr);
 	FVulkanCommandListContext& Context = (FVulkanCommandListContext&)RHICmdList.GetContext().GetLowestLevelContext();
 	ALLOC_COMMAND_CL(RHICmdList, FRHICommandVulkanAsyncReallocateTexture2D)(Context, OldTexture, NewTexture, NewMipCount, NewSizeX, NewSizeY, RequestStatus);
 
@@ -1471,10 +1078,18 @@ FTexture2DRHIRef FVulkanDynamicRHI::AsyncReallocateTexture2D_RenderThread(FRHICo
 FTexture2DRHIRef FVulkanDynamicRHI::RHIAsyncReallocateTexture2D(FRHITexture2D* OldTextureRHI, int32 NewMipCount, int32 NewSizeX, int32 NewSizeY, FThreadSafeCounter* RequestStatus)
 {
 	LLM_SCOPE_VULKAN(ELLMTagVulkan::VulkanTextures);
-	FVulkanTexture2D* OldTexture = ResourceCast(OldTextureRHI);
-
-	FRHIResourceCreateInfo CreateInfo(TEXT("RHIAsyncReallocateTexture2D"));
-	FVulkanTexture2D* NewTexture = new FVulkanTexture2D(*Device, OldTexture->GetFormat(), NewSizeX, NewSizeY, NewMipCount, OldTexture->GetNumSamples(), OldTexture->GetFlags(), ERHIAccess::Unknown, CreateInfo);
+	FVulkanTexture* OldTexture = FVulkanTexture::Cast(OldTextureRHI);
+
+	FRHITextureCreateDesc Desc = FRHITextureCreateDesc::Create2D(
+		TEXT("AsyncReallocateTexture2D_RenderThread"),
+		{ NewSizeX, NewSizeY },
+		OldTexture->GetFormat(),
+		OldTexture->GetClearBinding(),
+		OldTexture->GetFlags(),
+		NewMipCount,
+		OldTexture->GetNumSamples()
+	);
+	FVulkanTexture* NewTexture = new FVulkanTexture(*Device, Desc, nullptr);
 
 	DoAsyncReallocateTexture2D(Device->GetImmediateContext(), OldTexture, NewTexture, NewMipCount, NewSizeX, NewSizeY, RequestStatus);
 
@@ -1491,10 +1106,10 @@ ETextureReallocationStatus FVulkanDynamicRHI::RHICancelAsyncReallocateTexture2D(
 	return TexRealloc_Succeeded;
 }
 
-void* FVulkanDynamicRHI::RHILockTexture2D(FRHITexture2D* TextureRHI,uint32 MipIndex,EResourceLockMode LockMode,uint32& DestStride,bool bLockWithinMiptail)
+void* FVulkanDynamicRHI::RHILockTexture2D(FRHITexture2D* TextureRHI, uint32 MipIndex, EResourceLockMode LockMode, uint32& DestStride, bool bLockWithinMiptail)
 {
 	LLM_SCOPE_VULKAN(ELLMTagVulkan::VulkanTextures);
-	FVulkanTexture2D* Texture = ResourceCast(TextureRHI);
+	FVulkanTexture* Texture = FVulkanTexture::Cast(TextureRHI);
 	check(Texture);
 
 	VulkanRHI::FStagingBuffer** StagingBuffer = nullptr;
@@ -1509,18 +1124,18 @@ void* FVulkanDynamicRHI::RHILockTexture2D(FRHITexture2D* TextureRHI,uint32 MipIn
 
 	uint32 BufferSize = 0;
 	DestStride = 0;
-	Texture->Surface.GetMipSize(MipIndex, BufferSize);
-	Texture->Surface.GetMipStride(MipIndex, DestStride);
+	Texture->GetMipSize(MipIndex, BufferSize);
+	Texture->GetMipStride(MipIndex, DestStride);
 	*StagingBuffer = Device->GetStagingManager().AcquireBuffer(BufferSize);
 
 	void* Data = (*StagingBuffer)->GetMappedPointer();
 	return Data;
 }
 
-void FVulkanDynamicRHI::InternalUnlockTexture2D(bool bFromRenderingThread, FRHITexture2D* TextureRHI,uint32 MipIndex,bool bLockWithinMiptail)
+void FVulkanDynamicRHI::InternalUnlockTexture2D(bool bFromRenderingThread, FRHITexture2D* TextureRHI, uint32 MipIndex, bool bLockWithinMiptail)
 {
 	LLM_SCOPE_VULKAN(ELLMTagVulkan::VulkanTextures);
-	FVulkanTexture2D* Texture = ResourceCast(TextureRHI);
+	FVulkanTexture* Texture = FVulkanTexture::Cast(TextureRHI);
 	check(Texture);
 
 	VkDevice LogicalDevice = Device->GetInstanceHandle();
@@ -1532,13 +1147,14 @@ void FVulkanDynamicRHI::InternalUnlockTexture2D(bool bFromRenderingThread, FRHIT
 		checkf(bFound, TEXT("Texture was not locked!"));
 	}
 
-	EPixelFormat Format = Texture->Surface.PixelFormat;
-	uint32 MipWidth = FMath::Max<uint32>(Texture->Surface.Width >> MipIndex, 0);
-	uint32 MipHeight = FMath::Max<uint32>(Texture->Surface.Height >> MipIndex, 0);
+	const FRHITextureDesc& Desc = Texture->GetDesc();
+	const EPixelFormat Format = Desc.Format;
+	uint32 MipWidth = FMath::Max<uint32>(Desc.Extent.X >> MipIndex, 0);
+	uint32 MipHeight = FMath::Max<uint32>(Desc.Extent.Y >> MipIndex, 0);
 	ensure(!(MipHeight == 0 && MipWidth == 0));
 	MipWidth = FMath::Max<uint32>(MipWidth, 1);
 	MipHeight = FMath::Max<uint32>(MipHeight, 1);
-	uint32 LayerCount = Texture->Surface.GetNumberOfArrayLevels();
+	uint32 LayerCount = Texture->GetNumberOfArrayLevels();
 
 	VkBufferImageCopy Region;
 	FMemory::Memzero(Region);
@@ -1557,19 +1173,19 @@ void FVulkanDynamicRHI::InternalUnlockTexture2D(bool bFromRenderingThread, FRHIT
 	FRHICommandList& RHICmdList = FRHICommandListExecutor::GetImmediateCommandList();
 	if (!bFromRenderingThread || (RHICmdList.Bypass() || !IsRunningRHIInSeparateThread()))
 	{
-		FVulkanSurface::InternalLockWrite(Device->GetImmediateContext(), &Texture->Surface, Region, StagingBuffer);
+		FVulkanTexture::InternalLockWrite(Device->GetImmediateContext(), Texture, Region, StagingBuffer);
 	}
 	else
 	{
 		check(IsInRenderingThread());
-		ALLOC_COMMAND_CL(RHICmdList, FRHICommandLockWriteTexture)(&Texture->Surface, Region, StagingBuffer);
+		ALLOC_COMMAND_CL(RHICmdList, FRHICommandLockWriteTexture)(Texture, Region, StagingBuffer);
 	}
 }
 
 void* FVulkanDynamicRHI::RHILockTexture2DArray(FRHITexture2DArray* TextureRHI,uint32 TextureIndex,uint32 MipIndex,EResourceLockMode LockMode,uint32& DestStride,bool bLockWithinMiptail)
 {
 	LLM_SCOPE_VULKAN(ELLMTagVulkan::VulkanTextures);
-	FVulkanTexture2DArray* Texture = ResourceCast(TextureRHI);
+	FVulkanTexture* Texture = FVulkanTexture::Cast(TextureRHI);
 	check(Texture);
 
 	VulkanRHI::FStagingBuffer** StagingBuffer = nullptr;
@@ -1581,18 +1197,18 @@ void* FVulkanDynamicRHI::RHILockTexture2DArray(FRHITexture2DArray* TextureRHI,ui
 
 	uint32 BufferSize = 0;
 	DestStride = 0;
-	Texture->Surface.GetMipSize(MipIndex, BufferSize);
-	Texture->Surface.GetMipStride(MipIndex, DestStride);
+	Texture->GetMipSize(MipIndex, BufferSize);
+	Texture->GetMipStride(MipIndex, DestStride);
 	*StagingBuffer = Device->GetStagingManager().AcquireBuffer(BufferSize);
 
 	void* Data = (*StagingBuffer)->GetMappedPointer();
 	return Data;
 }
 
-void FVulkanDynamicRHI::RHIUnlockTexture2DArray(FRHITexture2DArray* TextureRHI,uint32 TextureIndex,uint32 MipIndex,bool bLockWithinMiptail)
+void FVulkanDynamicRHI::RHIUnlockTexture2DArray(FRHITexture2DArray* TextureRHI, uint32 TextureIndex, uint32 MipIndex, bool bLockWithinMiptail)
 {
 	LLM_SCOPE_VULKAN(ELLMTagVulkan::VulkanTextures);
-	FVulkanTexture2DArray* Texture = ResourceCast(TextureRHI);
+	FVulkanTexture* Texture = FVulkanTexture::Cast(TextureRHI);
 	check(Texture);
 
 	VkDevice LogicalDevice = Device->GetInstanceHandle();
@@ -1604,9 +1220,10 @@ void FVulkanDynamicRHI::RHIUnlockTexture2DArray(FRHITexture2DArray* TextureRHI,u
 		checkf(bFound, TEXT("Texture was not locked!"));
 	}
 
-	EPixelFormat Format = Texture->Surface.PixelFormat;
-	uint32 MipWidth = FMath::Max<uint32>(Texture->Surface.Width >> MipIndex, GPixelFormats[Format].BlockSizeX);
-	uint32 MipHeight = FMath::Max<uint32>(Texture->Surface.Height >> MipIndex, GPixelFormats[Format].BlockSizeY);
+	const FRHITextureDesc& Desc = Texture->GetDesc();
+	const EPixelFormat Format = Desc.Format;
+	const uint32 MipWidth = FMath::Max<uint32>(Desc.Extent.X >> MipIndex, GPixelFormats[Format].BlockSizeX);
+	const uint32 MipHeight = FMath::Max<uint32>(Desc.Extent.Y >> MipIndex, GPixelFormats[Format].BlockSizeY);
 
 	VkBufferImageCopy Region;
 	FMemory::Memzero(Region);
@@ -1614,7 +1231,7 @@ void FVulkanDynamicRHI::RHIUnlockTexture2DArray(FRHITexture2DArray* TextureRHI,u
 	//Region.bufferOffset = 0;
 	//Region.bufferRowLength = 0;
 	//Region.bufferImageHeight = 0;
-	Region.imageSubresource.aspectMask = Texture->Surface.GetPartialAspectMask();
+	Region.imageSubresource.aspectMask = Texture->GetPartialAspectMask();
 	Region.imageSubresource.mipLevel = MipIndex;
 	Region.imageSubresource.baseArrayLayer = TextureIndex;
 	Region.imageSubresource.layerCount = 1;
@@ -1625,21 +1242,21 @@ void FVulkanDynamicRHI::RHIUnlockTexture2DArray(FRHITexture2DArray* TextureRHI,u
 	FRHICommandList& RHICmdList = FRHICommandListExecutor::GetImmediateCommandList();
 	if (RHICmdList.Bypass() || !IsRunningRHIInSeparateThread())
 	{
-		FVulkanSurface::InternalLockWrite(Device->GetImmediateContext(), &Texture->Surface, Region, StagingBuffer);
+		FVulkanTexture::InternalLockWrite(Device->GetImmediateContext(), Texture, Region, StagingBuffer);
 	}
 	else
 	{
 		check(IsInRenderingThread());
-		ALLOC_COMMAND_CL(RHICmdList, FRHICommandLockWriteTexture)(&Texture->Surface, Region, StagingBuffer);
+		ALLOC_COMMAND_CL(RHICmdList, FRHICommandLockWriteTexture)(Texture, Region, StagingBuffer);
 	}
 }
 
 void FVulkanDynamicRHI::InternalUpdateTexture2D(bool bFromRenderingThread, FRHITexture2D* TextureRHI, uint32 MipIndex, const struct FUpdateTextureRegion2D& UpdateRegion, uint32 SourceRowPitch, const uint8* SourceData)
 {
 	LLM_SCOPE_VULKAN(ELLMTagVulkan::VulkanTextures);
-	FVulkanTexture2D* Texture = ResourceCast(TextureRHI);
+	FVulkanTexture* Texture = FVulkanTexture::Cast(TextureRHI);
 
-	EPixelFormat PixelFormat = Texture->GetFormat();
+	const EPixelFormat PixelFormat = Texture->GetDesc().Format;
 	const int32 BlockSizeX = GPixelFormats[PixelFormat].BlockSizeX;
 	const int32 BlockSizeY = GPixelFormats[PixelFormat].BlockSizeY;
 	const int32 BlockSizeZ = GPixelFormats[PixelFormat].BlockSizeZ;
@@ -1692,12 +1309,12 @@ void FVulkanDynamicRHI::InternalUpdateTexture2D(bool bFromRenderingThread, FRHIT
 	FRHICommandList& RHICmdList = FRHICommandListExecutor::GetImmediateCommandList();
 	if (!bFromRenderingThread || (RHICmdList.Bypass() || !IsRunningRHIInSeparateThread()))
 	{
-		FVulkanSurface::InternalLockWrite(Device->GetImmediateContext(), &Texture->Surface, Region, StagingBuffer);
+		FVulkanTexture::InternalLockWrite(Device->GetImmediateContext(), Texture, Region, StagingBuffer);
 	}
 	else
 	{
 		check(IsInRenderingThread());
-		ALLOC_COMMAND_CL(RHICmdList, FRHICommandLockWriteTexture)(&Texture->Surface, Region, StagingBuffer);
+		ALLOC_COMMAND_CL(RHICmdList, FRHICommandLockWriteTexture)(Texture, Region, StagingBuffer);
 	}
 }
 
@@ -1727,9 +1344,9 @@ void FVulkanDynamicRHI::EndUpdateTexture3D_RenderThread(class FRHICommandListImm
 void FVulkanDynamicRHI::InternalUpdateTexture3D(bool bFromRenderingThread, FRHITexture3D* TextureRHI, uint32 MipIndex, const FUpdateTextureRegion3D& UpdateRegion, uint32 SourceRowPitch, uint32 SourceDepthPitch, const uint8* SourceData)
 {
 	LLM_SCOPE_VULKAN(ELLMTagVulkan::VulkanTextures);
-	FVulkanTexture3D* Texture = ResourceCast(TextureRHI);
+	FVulkanTexture* Texture = FVulkanTexture::Cast(TextureRHI);
 
-	const EPixelFormat PixelFormat = Texture->GetFormat();
+	const EPixelFormat PixelFormat = Texture->GetDesc().Format;
 	const int32 BlockSizeX = GPixelFormats[PixelFormat].BlockSizeX;
 	const int32 BlockSizeY = GPixelFormats[PixelFormat].BlockSizeY;
 	const int32 BlockSizeZ = GPixelFormats[PixelFormat].BlockSizeZ;
@@ -1792,20 +1409,19 @@ void FVulkanDynamicRHI::InternalUpdateTexture3D(bool bFromRenderingThread, FRHIT
 	FRHICommandList& RHICmdList = FRHICommandListExecutor::GetImmediateCommandList();
 	if (!bFromRenderingThread || (RHICmdList.Bypass() || !IsRunningRHIInSeparateThread()))
 	{
-		FVulkanSurface::InternalLockWrite(Device->GetImmediateContext(), &Texture->Surface, Region, StagingBuffer);
+		FVulkanTexture::InternalLockWrite(Device->GetImmediateContext(), Texture, Region, StagingBuffer);
 	}
 	else
 	{
 		check(IsInRenderingThread());
-		ALLOC_COMMAND_CL(RHICmdList, FRHICommandLockWriteTexture)(&Texture->Surface, Region, StagingBuffer);
+		ALLOC_COMMAND_CL(RHICmdList, FRHICommandLockWriteTexture)(Texture, Region, StagingBuffer);
 	}
 }
 
-
-VkImageView FVulkanTextureView::StaticCreate(FVulkanDevice& Device, VkImage InImage, VkImageViewType ViewType, VkImageAspectFlags AspectFlags, EPixelFormat UEFormat, VkFormat Format, uint32 FirstMip, uint32 NumMips, uint32 ArraySliceIndex, uint32 NumArraySlices, bool bUseIdentitySwizzle)
+void FVulkanTextureView::Create(FVulkanDevice& Device, VkImage InImage, VkImageViewType ViewType, VkImageAspectFlags AspectFlags, EPixelFormat UEFormat, VkFormat Format, uint32 FirstMip, uint32 NumMips, uint32 ArraySliceIndex, uint32 NumArraySlices, bool bUseIdentitySwizzle /*= false*/)
 {
 	LLM_SCOPE_VULKAN(ELLMTagVulkan::VulkanTextures);
-	VkImageView OutView = VK_NULL_HANDLE;
+	View = VK_NULL_HANDLE;
 
 	VkImageViewCreateInfo ViewInfo;
 	ZeroVulkanStruct(ViewInfo, VK_STRUCTURE_TYPE_IMAGE_VIEW_CREATE_INFO);
@@ -1897,16 +1513,10 @@ VkImageView FVulkanTextureView::StaticCreate(FVulkanDevice& Device, VkImage InIm
 	}
 
 	INC_DWORD_STAT(STAT_VulkanNumImageViews);
-	VERIFYVULKANRESULT(VulkanRHI::vkCreateImageView(Device.GetInstanceHandle(), &ViewInfo, VULKAN_CPU_ALLOCATOR, &OutView));
-
-	return OutView;
-}
+	VERIFYVULKANRESULT(VulkanRHI::vkCreateImageView(Device.GetInstanceHandle(), &ViewInfo, VULKAN_CPU_ALLOCATOR, &View));
 
-void FVulkanTextureView::Create(FVulkanDevice& Device, VkImage InImage, VkImageViewType ViewType, VkImageAspectFlags AspectFlags, EPixelFormat UEFormat, VkFormat Format, uint32 FirstMip, uint32 NumMips, uint32 ArraySliceIndex, uint32 NumArraySlices, bool bUseIdentitySwizzle)
-{
-	View = StaticCreate(Device, InImage, ViewType, AspectFlags, UEFormat, Format, FirstMip, NumMips, ArraySliceIndex, NumArraySlices, bUseIdentitySwizzle);
 	Image = InImage;
-	
+
 	if (UseVulkanDescriptorCache())
 	{
 		ViewId = ++GVulkanImageViewHandleIdCounter;
@@ -1925,121 +1535,383 @@ void FVulkanTextureView::Destroy(FVulkanDevice& Device)
 	}
 }
 
-FVulkanTextureBase::FVulkanTextureBase(FVulkanDevice& Device, VkImageViewType ResourceType, EPixelFormat InFormat, uint32 SizeX, uint32 SizeY, uint32 SizeZ, uint32 ArraySize, uint32 NumMips, uint32 NumSamples, ETextureCreateFlags UEFlags, ERHIAccess InResourceState, const FRHIResourceCreateInfo& CreateInfo, const FRHITransientHeapAllocation* InTransientHeapAllocation)
-	: Surface(Device, this, ResourceType, InFormat, SizeX, SizeY, SizeZ, ArraySize, NumMips, NumSamples, UEFlags, InResourceState, CreateInfo, InTransientHeapAllocation)
+FVulkanTexture::FVulkanTexture(FVulkanDevice& InDevice, const FRHITextureCreateDesc& InCreateDesc, const FRHITransientHeapAllocation* InTransientHeapAllocation)
+	: FRHITexture(InCreateDesc)
 	, PartialView(nullptr)
+	, Device(&InDevice)
+	, Image(VK_NULL_HANDLE)
+	, StorageFormat(VK_FORMAT_UNDEFINED)
+	, ViewFormat(VK_FORMAT_UNDEFINED)
+	, MemProps(VK_MEMORY_PROPERTY_DEVICE_LOCAL_BIT)
+	, Tiling(VK_IMAGE_TILING_MAX_ENUM)	// Can be expanded to a per-platform definition
+	, FullAspectMask(0)
+	, PartialAspectMask(0)
+	, CpuReadbackBuffer(nullptr) // for readback textures we use a staging buffer. this is because vulkan only requires implementations to support 1 mip level(which is useless), so we emulate using a buffer
 {
-	Surface.OwningTexture = this;
-	VULKAN_TRACK_OBJECT_CREATE(FVulkanTextureBase, this);
+	VULKAN_TRACK_OBJECT_CREATE(FVulkanTexture, this);
 
-	if(EnumHasAnyFlags(UEFlags, TexCreate_CPUReadback))
 	{
-		return;
+		if (EnumHasAnyFlags(InCreateDesc.Flags, TexCreate_CPUReadback))
+		{
+			check(InCreateDesc.NumSamples == 1);	//not implemented
+			check(InCreateDesc.Depth == 1);		//not implemented
+			check(InCreateDesc.ArraySize == 1);	//not implemented
+			CpuReadbackBuffer = new FVulkanCpuReadbackBuffer;
+			uint32 Size = 0;
+			for (uint32 Mip = 0; Mip < InCreateDesc.NumMips; ++Mip)
+			{
+				uint32 LocalSize;
+				GetMipSize(Mip, LocalSize);
+				CpuReadbackBuffer->MipOffsets[Mip] = Size;
+				CpuReadbackBuffer->MipSize[Mip] = LocalSize;
+				Size += LocalSize;
+			}
+
+			VkDevice VulkanDevice = InDevice.GetInstanceHandle();
+			VkMemoryPropertyFlags BufferMemFlags = (VK_MEMORY_PROPERTY_HOST_VISIBLE_BIT | VK_MEMORY_PROPERTY_HOST_COHERENT_BIT | VK_MEMORY_PROPERTY_HOST_CACHED_BIT);
+
+			VkBufferCreateInfo BufferCreateInfo;
+			ZeroVulkanStruct(BufferCreateInfo, VK_STRUCTURE_TYPE_BUFFER_CREATE_INFO);
+			BufferCreateInfo.size = Size;
+			BufferCreateInfo.usage = VK_BUFFER_USAGE_TRANSFER_DST_BIT;
+
+			VERIFYVULKANRESULT(VulkanRHI::vkCreateBuffer(VulkanDevice, &BufferCreateInfo, VULKAN_CPU_ALLOCATOR, &CpuReadbackBuffer->Buffer));
+			VulkanRHI::vkGetBufferMemoryRequirements(VulkanDevice, CpuReadbackBuffer->Buffer, &MemoryRequirements);
+			// Set minimum alignment to 16 bytes, as some buffers are used with CPU SIMD instructions
+			MemoryRequirements.alignment = FMath::Max<VkDeviceSize>(16, MemoryRequirements.alignment);
+			if (!InDevice.GetMemoryManager().AllocateBufferMemory(Allocation, this, MemoryRequirements, BufferMemFlags, EVulkanAllocationMetaBufferStaging, false, __FILE__, __LINE__))
+			{
+				InDevice.GetMemoryManager().HandleOOM();
+			}
+			Allocation.BindBuffer(Device, CpuReadbackBuffer->Buffer);
+			void* Memory = Allocation.GetMappedPointer(Device);
+			FMemory::Memzero(Memory, MemoryRequirements.size);
+
+			ImageOwnerType = EImageOwnerType::None;
+			ViewFormat = StorageFormat = UEToVkTextureFormat(InCreateDesc.Format, false);
+
+			return;
+		}
+
+		ImageOwnerType = EImageOwnerType::LocalOwner;
+
+		FImageCreateInfo ImageCreateInfo;
+		FVulkanTexture::GenerateImageCreateInfo(ImageCreateInfo, InDevice, InCreateDesc, &StorageFormat, &ViewFormat);
+
+		VERIFYVULKANRESULT(VulkanRHI::vkCreateImage(InDevice.GetInstanceHandle(), &ImageCreateInfo.ImageCreateInfo, VULKAN_CPU_ALLOCATOR, &Image));
+
+		// Fetch image size
+		VulkanRHI::vkGetImageMemoryRequirements(InDevice.GetInstanceHandle(), Image, &MemoryRequirements);
+
+		VULKAN_SET_DEBUG_NAME(InDevice, VK_OBJECT_TYPE_IMAGE, Image, TEXT("%s:(FVulkanTexture*)0x%p"), InCreateDesc.DebugName ? InCreateDesc.DebugName : TEXT("?"), this);
+
+		FullAspectMask = VulkanRHI::GetAspectMaskFromUEFormat(InCreateDesc.Format, true, true);
+		PartialAspectMask = VulkanRHI::GetAspectMaskFromUEFormat(InCreateDesc.Format, false, true);
+
+		// If VK_IMAGE_TILING_OPTIMAL is specified,
+		// memoryTypeBits in vkGetImageMemoryRequirements will become 1
+		// which does not support VK_MEMORY_PROPERTY_HOST_VISIBLE_BIT.
+		if (ImageCreateInfo.ImageCreateInfo.tiling != VK_IMAGE_TILING_OPTIMAL)
+		{
+			MemProps |= VK_MEMORY_PROPERTY_HOST_VISIBLE_BIT;
+		}
+
+		const bool bRenderTarget = EnumHasAnyFlags(InCreateDesc.Flags, TexCreate_RenderTargetable | TexCreate_DepthStencilTargetable | TexCreate_ResolveTargetable);
+		const bool bUAV = EnumHasAnyFlags(InCreateDesc.Flags, TexCreate_UAV);
+		const bool bCPUReadback = EnumHasAnyFlags(InCreateDesc.Flags, TexCreate_CPUReadback);
+		const bool bDynamic = EnumHasAnyFlags(InCreateDesc.Flags, TexCreate_Dynamic);
+		const bool bExternal = EnumHasAnyFlags(InCreateDesc.Flags, TexCreate_External);
+
+		VkMemoryPropertyFlags MemoryFlags = bCPUReadback ? VK_MEMORY_PROPERTY_HOST_VISIBLE_BIT : VK_MEMORY_PROPERTY_DEVICE_LOCAL_BIT;
+
+		bool bMemoryless = EnumHasAnyFlags(InCreateDesc.Flags, TexCreate_Memoryless) && InDevice.GetDeviceMemoryManager().SupportsMemoryless();
+		if (bMemoryless)
+		{
+			if (ensureMsgf(bRenderTarget, TEXT("Memoryless surfaces can only be used for render targets")) && ensureMsgf(!bCPUReadback, TEXT("Memoryless surfaces cannot be read back on CPU")))
+			{
+				MemoryFlags |= VK_MEMORY_PROPERTY_LAZILY_ALLOCATED_BIT;
+			}
+			else
+			{
+				bMemoryless = false;
+			}
+		}
+
+		if (InTransientHeapAllocation != nullptr)
+		{
+			check(!bCPUReadback);
+			check(!bMemoryless);
+			check(InTransientHeapAllocation->Offset % MemoryRequirements.alignment == 0);
+			check(InTransientHeapAllocation->Size >= MemoryRequirements.size);
+			Allocation = FVulkanTransientHeap::GetVulkanAllocation(*InTransientHeapAllocation);
+		}
+		else
+		{
+			EVulkanAllocationMetaType MetaType = (bRenderTarget || bUAV) ? EVulkanAllocationMetaImageRenderTarget : EVulkanAllocationMetaImageOther;
+#if VULKAN_SUPPORTS_DEDICATED_ALLOCATION
+			extern int32 GVulkanEnableDedicatedImageMemory;
+			// Per https://developer.nvidia.com/what%E2%80%99s-your-vulkan-memory-type
+			VkDeviceSize SizeToBeConsideredForDedicated = 12 * 1024 * 1024;
+			if ((bRenderTarget || MemoryRequirements.size >= SizeToBeConsideredForDedicated) && !bMemoryless && InDevice.GetOptionalExtensions().HasKHRDedicatedAllocation && GVulkanEnableDedicatedImageMemory)
+			{
+				if (!InDevice.GetMemoryManager().AllocateDedicatedImageMemory(Allocation, this, Image, MemoryRequirements, MemoryFlags, MetaType, bExternal, __FILE__, __LINE__))
+				{
+					checkNoEntry();
+				}
+			}
+			else
+#endif
+			{
+				if (!InDevice.GetMemoryManager().AllocateImageMemory(Allocation, this, MemoryRequirements, MemoryFlags, MetaType, bExternal, __FILE__, __LINE__))
+				{
+					checkNoEntry();
+				}
+			}
+		}
+		Allocation.BindImage(Device, Image);
+
+		// update rhi stats
+		VulkanTextureAllocated(MemoryRequirements.size, GetViewType(), bRenderTarget);
+
+		Tiling = ImageCreateInfo.ImageCreateInfo.tiling;
+		check(Tiling == VK_IMAGE_TILING_LINEAR || Tiling == VK_IMAGE_TILING_OPTIMAL);
+
+		const VkImageLayout InitialLayout = GetInitialLayoutFromRHIAccess(InCreateDesc.InitialState, InCreateDesc.Flags, SupportsSampling());
+		const bool bDoInitialClear = VKHasAnyFlags(ImageCreateInfo.ImageCreateInfo.usage, VK_IMAGE_USAGE_SAMPLED_BIT) && EnumHasAnyFlags(InCreateDesc.Flags, TexCreate_RenderTargetable | TexCreate_DepthStencilTargetable);
+
+		if (InitialLayout != VK_IMAGE_LAYOUT_UNDEFINED || bDoInitialClear)
+		{
+			FRHICommandList& RHICmdList = FRHICommandListExecutor::GetImmediateCommandList();
+			if (!IsInRenderingThread() || (RHICmdList.Bypass() || !IsRunningRHIInSeparateThread()))
+			{
+				SetInitialImageState(Device->GetImmediateContext(), InitialLayout, bDoInitialClear, InCreateDesc.ClearValue);
+			}
+			else
+			{
+				check(IsInRenderingThread());
+				ALLOC_COMMAND_CL(RHICmdList, FRHICommandSetInitialImageState)(this, InitialLayout, false, bDoInitialClear, InCreateDesc.ClearValue);
+			}
+		}
 	}
 
 	LLM_SCOPE_VULKAN(ELLMTagVulkan::VulkanTextures);
-	const bool bArray = ResourceType == VK_IMAGE_VIEW_TYPE_1D_ARRAY || ResourceType == VK_IMAGE_VIEW_TYPE_2D_ARRAY || ResourceType == VK_IMAGE_VIEW_TYPE_CUBE_ARRAY;
-	if (Surface.ViewFormat == VK_FORMAT_UNDEFINED)
+	const VkImageViewType ViewType = GetViewType();
+	const bool bArray = ViewType == VK_IMAGE_VIEW_TYPE_1D_ARRAY || ViewType == VK_IMAGE_VIEW_TYPE_2D_ARRAY || ViewType == VK_IMAGE_VIEW_TYPE_CUBE_ARRAY;
+	if (ViewFormat == VK_FORMAT_UNDEFINED)
 	{
-		Surface.StorageFormat = UEToVkTextureFormat(InFormat, false);
-		Surface.ViewFormat = UEToVkTextureFormat(InFormat, EnumHasAllFlags(UEFlags, TexCreate_SRGB));
-		checkf(Surface.StorageFormat != VK_FORMAT_UNDEFINED, TEXT("Pixel Format %d not defined!"), (int32)InFormat);
+		StorageFormat = UEToVkTextureFormat(InCreateDesc.Format, false);
+		ViewFormat = UEToVkTextureFormat(InCreateDesc.Format, EnumHasAllFlags(InCreateDesc.Flags, TexCreate_SRGB));
+		checkf(StorageFormat != VK_FORMAT_UNDEFINED, TEXT("Pixel Format %d not defined!"), (int32)InCreateDesc.Format);
 	}
 
-	if (ResourceType != VK_IMAGE_VIEW_TYPE_MAX_ENUM)
+	if (ViewType != VK_IMAGE_VIEW_TYPE_MAX_ENUM)
 	{
-		DefaultView.Create(Device, Surface.Image, ResourceType, Surface.GetFullAspectMask(), Surface.PixelFormat, Surface.ViewFormat, 0, FMath::Max(NumMips, 1u), 0, bArray ? FMath::Max(1u, ArraySize) : FMath::Max(1u, SizeZ), !Surface.SupportsSampling());
+		DefaultView.Create(InDevice, Image, ViewType, GetFullAspectMask(), InCreateDesc.Format, ViewFormat, 0, FMath::Max(InCreateDesc.NumMips, (uint8)1u), 0, bArray ? FMath::Max((uint16)1u, InCreateDesc.ArraySize) : FMath::Max((uint16)1u, InCreateDesc.Depth), !SupportsSampling());
 	}
 
-	if (Surface.FullAspectMask == Surface.PartialAspectMask)
+	if (FullAspectMask == PartialAspectMask)
 	{
 		PartialView = &DefaultView;
 	}
 	else
 	{
 		PartialView = new FVulkanTextureView;
-		PartialView->Create(Device, Surface.Image, Surface.ViewType, Surface.PartialAspectMask, Surface.PixelFormat, Surface.ViewFormat, 0, FMath::Max(NumMips, 1u), 0, bArray ? FMath::Max(1u, ArraySize) : FMath::Max(1u, SizeZ));
+		PartialView->Create(InDevice, Image, ViewType, PartialAspectMask, InCreateDesc.Format, ViewFormat, 0, FMath::Max(InCreateDesc.NumMips, (uint8)1u), 0, bArray ? FMath::Max((uint16)1u, InCreateDesc.ArraySize) : FMath::Max((uint16)1u, InCreateDesc.Depth), false);
 	}
 
-	if (!CreateInfo.BulkData)
+	if (!InCreateDesc.BulkData)
 	{
 		return;
 	}
 
 	// InternalLockWrite leaves the image in VK_IMAGE_LAYOUT_SHADER_READ_ONLY_OPTIMAL, so make sure the requested resource state is SRV.
-	check(EnumHasAnyFlags(InResourceState, ERHIAccess::SRVMask));
+	check(EnumHasAnyFlags(InCreateDesc.InitialState, ERHIAccess::SRVMask));
 
 	// Transfer bulk data
-	VulkanRHI::FStagingBuffer* StagingBuffer = Device.GetStagingManager().AcquireBuffer(CreateInfo.BulkData->GetResourceBulkDataSize());
+	VulkanRHI::FStagingBuffer* StagingBuffer = InDevice.GetStagingManager().AcquireBuffer(InCreateDesc.BulkData->GetResourceBulkDataSize());
 	void* Data = StagingBuffer->GetMappedPointer();
 
 	// Do copy
-	FMemory::Memcpy(Data, CreateInfo.BulkData->GetResourceBulkData(), CreateInfo.BulkData->GetResourceBulkDataSize());
-	CreateInfo.BulkData->Discard();
+	FMemory::Memcpy(Data, InCreateDesc.BulkData->GetResourceBulkData(), InCreateDesc.BulkData->GetResourceBulkDataSize());
+	InCreateDesc.BulkData->Discard();
 
-	uint32 LayersPerArrayIndex = (ResourceType == VK_IMAGE_VIEW_TYPE_CUBE_ARRAY || ResourceType == VK_IMAGE_VIEW_TYPE_CUBE) ? 6 : 1;
+	uint32 LayersPerArrayIndex = (ViewType == VK_IMAGE_VIEW_TYPE_CUBE_ARRAY || ViewType == VK_IMAGE_VIEW_TYPE_CUBE) ? 6 : 1;
 
 	VkBufferImageCopy Region;
 	FMemory::Memzero(Region);
 	//#todo-rco: Use real Buffer offset when switching to suballocations!
 	Region.bufferOffset = 0;
-	Region.bufferRowLength = Surface.Width;
-	Region.bufferImageHeight = Surface.Height;
+	Region.bufferRowLength = InCreateDesc.Extent.X;
+	Region.bufferImageHeight = InCreateDesc.Extent.Y;
 	
 	Region.imageSubresource.mipLevel = 0;
 	Region.imageSubresource.baseArrayLayer = 0;
-	Region.imageSubresource.layerCount = ArraySize * LayersPerArrayIndex;
-	Region.imageSubresource.aspectMask = Surface.GetFullAspectMask();
+	Region.imageSubresource.layerCount = InCreateDesc.ArraySize * LayersPerArrayIndex;
+	Region.imageSubresource.aspectMask = GetFullAspectMask();
 
 	Region.imageExtent.width = Region.bufferRowLength;
 	Region.imageExtent.height = Region.bufferImageHeight;
-	Region.imageExtent.depth = Surface.Depth;
+	Region.imageExtent.depth = InCreateDesc.Depth;
 
 	FRHICommandList& RHICmdList = FRHICommandListExecutor::GetImmediateCommandList();
 	if (RHICmdList.Bypass() || !IsRunningRHIInSeparateThread())
 	{
-		FVulkanSurface::InternalLockWrite(Device.GetImmediateContext(), &Surface, Region, StagingBuffer);
+		FVulkanTexture::InternalLockWrite(InDevice.GetImmediateContext(), this, Region, StagingBuffer);
 	}
 	else
 	{
 		check(IsInRenderingThread());
-		ALLOC_COMMAND_CL(RHICmdList, FRHICommandLockWriteTexture)(&Surface, Region, StagingBuffer);
+		ALLOC_COMMAND_CL(RHICmdList, FRHICommandLockWriteTexture)(this, Region, StagingBuffer);
 	}
 }
 
-FVulkanTextureBase::FVulkanTextureBase(FVulkanDevice& Device, VkImageViewType ResourceType, EPixelFormat Format, uint32 SizeX, uint32 SizeY, uint32 SizeZ, uint32 ArraySize, uint32 InNumMips, uint32 InNumSamples, VkImage InImage, VkDeviceMemory InMem, ETextureCreateFlags UEFlags, const FRHIResourceCreateInfo& CreateInfo)
-	: Surface(Device, ResourceType, Format, SizeX, SizeY, SizeZ, ArraySize, InNumMips, InNumSamples, InImage, UEFlags, EImageOwnerType::ExternalOwner, CreateInfo)
-	, PartialView(nullptr)
+FVulkanTexture::FVulkanTexture(FVulkanDevice& InDevice, const FRHITextureCreateDesc& InCreateDesc, VkImage InImage, bool /*bUnused*/)
+	: FRHITexture(InCreateDesc)
+	, Device(&InDevice)
+	, Image(InImage)
+	, StorageFormat(VK_FORMAT_UNDEFINED)
+	, ViewFormat(VK_FORMAT_UNDEFINED)
+	, MemProps(VK_MEMORY_PROPERTY_DEVICE_LOCAL_BIT)
+	, Tiling(VK_IMAGE_TILING_MAX_ENUM)	// Can be expanded to a per-platform definition
+	, FullAspectMask(0)
+	, PartialAspectMask(0)
+	, CpuReadbackBuffer(nullptr)
+	, ImageOwnerType(EImageOwnerType::ExternalOwner)
 {
-	Surface.OwningTexture = this;
-	VULKAN_TRACK_OBJECT_CREATE(FVulkanTextureBase, this);
-	check(InMem == VK_NULL_HANDLE);
-	const bool bArray = ResourceType == VK_IMAGE_VIEW_TYPE_1D_ARRAY || ResourceType == VK_IMAGE_VIEW_TYPE_2D_ARRAY || ResourceType == VK_IMAGE_VIEW_TYPE_CUBE_ARRAY;
+	VULKAN_TRACK_OBJECT_CREATE(FVulkanTexture, this);
+
+	{
+		StorageFormat = UEToVkTextureFormat(InCreateDesc.Format, false);
+
+		checkf(InCreateDesc.Format == PF_Unknown || StorageFormat != VK_FORMAT_UNDEFINED, TEXT("PixelFormat %d, is not supported for images"), (int32)InCreateDesc.Format);
+
+		ViewFormat = UEToVkTextureFormat(InCreateDesc.Format, EnumHasAllFlags(InCreateDesc.Flags, TexCreate_SRGB));
+		FullAspectMask = VulkanRHI::GetAspectMaskFromUEFormat(InCreateDesc.Format, true, true);
+		PartialAspectMask = VulkanRHI::GetAspectMaskFromUEFormat(InCreateDesc.Format, false, true);
+
+		// Purely informative patching, we know that "TexCreate_Presentable" uses optimal tiling
+		if (EnumHasAllFlags(InCreateDesc.Flags, TexCreate_Presentable) && GetTiling() == VK_IMAGE_TILING_MAX_ENUM)
+		{
+			Tiling = VK_IMAGE_TILING_OPTIMAL;
+		}
+
+		if (Image != VK_NULL_HANDLE)
+		{
+#if VULKAN_ENABLE_WRAP_LAYER
+			FImageCreateInfo ImageCreateInfo;
+			FVulkanTexture::GenerateImageCreateInfo(ImageCreateInfo, InDevice, InCreateDesc, &StorageFormat, &ViewFormat);
+			FWrapLayer::CreateImage(VK_SUCCESS, InDevice.GetInstanceHandle(), &ImageCreateInfo.ImageCreateInfo, &Image);
+#endif
+			VULKAN_SET_DEBUG_NAME(InDevice, VK_OBJECT_TYPE_IMAGE, Image, TEXT("%s:(FVulkanTexture*)0x%p"), InCreateDesc.DebugName ? InCreateDesc.DebugName : TEXT("?"), this);
+
+			VkImageLayout InitialLayout = VK_IMAGE_LAYOUT_UNDEFINED;
+			bool bOnlyAddToLayoutManager, bDoInitialClear;
+			if (EnumHasAnyFlags(InCreateDesc.Flags, TexCreate_RenderTargetable | TexCreate_DepthStencilTargetable))
+			{
+				InitialLayout = EnumHasAnyFlags(InCreateDesc.Flags, TexCreate_DepthStencilTargetable) ? VK_IMAGE_LAYOUT_DEPTH_STENCIL_ATTACHMENT_OPTIMAL : VK_IMAGE_LAYOUT_COLOR_ATTACHMENT_OPTIMAL;
+				bOnlyAddToLayoutManager = false;
+				bDoInitialClear = true;
+			}
+			else if (EnumHasAnyFlags(InCreateDesc.Flags, TexCreate_Foveation) && ValidateShadingRateDataType())
+			{
+				// If it's a shading rate attachment texture, do not clear but add to layoutmgr, and set correct shading rate layout.
+#if VULKAN_SUPPORTS_FRAGMENT_SHADING_RATE
+				if (GRHIVariableRateShadingImageDataType == VRSImage_Palette)
+				{
+					InitialLayout = VK_IMAGE_LAYOUT_FRAGMENT_SHADING_RATE_ATTACHMENT_OPTIMAL_KHR;
+				}
+#endif
+
+#if VULKAN_SUPPORTS_FRAGMENT_DENSITY_MAP
+				if (GRHIVariableRateShadingImageDataType == VRSImage_Fractional)
+				{
+					InitialLayout = VK_IMAGE_LAYOUT_FRAGMENT_DENSITY_MAP_OPTIMAL_EXT;
+				}
+#endif
+
+				bOnlyAddToLayoutManager = true;
+				bDoInitialClear = false;
+			}
+			else
+			{
+				// If we haven't seen this image before, we assume it's an SRV (VK_IMAGE_LAYOUT_SHADER_READ_ONLY_OPTIMAL) and the call below
+				// tells the layout manager about it. If we've seen it before, the call won't do anything, since the manager already knows the layout.
+				InitialLayout = VK_IMAGE_LAYOUT_SHADER_READ_ONLY_OPTIMAL;
+				bOnlyAddToLayoutManager = true;
+				bDoInitialClear = false;
+			}
 
-	if (ResourceType != VK_IMAGE_VIEW_TYPE_MAX_ENUM && Surface.Image != VK_NULL_HANDLE)
+			FRHICommandList& RHICmdList = FRHICommandListExecutor::GetImmediateCommandList();
+			if (!IsInRenderingThread() || (RHICmdList.Bypass() || !IsRunningRHIInSeparateThread()))
+			{
+				if (bOnlyAddToLayoutManager)
+				{
+					InDevice.GetImmediateContext().GetLayoutManager().GetOrAddFullLayout(*this, InitialLayout);
+				}
+				else
+				{
+					SetInitialImageState(InDevice.GetImmediateContext(), InitialLayout, true, InCreateDesc.ClearValue);
+				}
+			}
+			else
+			{
+				check(IsInRenderingThread());
+				ALLOC_COMMAND_CL(RHICmdList, FRHICommandSetInitialImageState)(this, InitialLayout, bOnlyAddToLayoutManager, bDoInitialClear, InCreateDesc.ClearValue);
+			}
+		}
+	}
+
+	const VkImageViewType ViewType = GetViewType();
+	const bool bArray = ViewType == VK_IMAGE_VIEW_TYPE_1D_ARRAY || ViewType == VK_IMAGE_VIEW_TYPE_2D_ARRAY || ViewType == VK_IMAGE_VIEW_TYPE_CUBE_ARRAY;
+
+	if (Image != VK_NULL_HANDLE)
 	{
-		DefaultView.Create(Device, Surface.Image, ResourceType, Surface.GetFullAspectMask(), Format, Surface.ViewFormat, 0, FMath::Max(Surface.NumMips, 1u), 0, bArray ? FMath::Max(1u, ArraySize) : FMath::Max(1u, SizeZ), !Surface.SupportsSampling());
+		DefaultView.Create(InDevice, Image, ViewType, GetFullAspectMask(), InCreateDesc.Format, ViewFormat, 0, FMath::Max(InCreateDesc.NumMips, (uint8)1u), 0, bArray ? FMath::Max((uint16)1u, InCreateDesc.ArraySize) : FMath::Max((uint16)1u, InCreateDesc.Depth), !SupportsSampling());
 	}
 
-	if (Surface.FullAspectMask == Surface.PartialAspectMask)
+	if (FullAspectMask == PartialAspectMask)
 	{
 		PartialView = &DefaultView;
 	}
 	else
 	{
 		PartialView = new FVulkanTextureView;
-		PartialView->Create(Device, Surface.Image, Surface.ViewType, Surface.PartialAspectMask, Surface.PixelFormat, Surface.ViewFormat, 0, FMath::Max(InNumMips, 1u), 0, bArray ? FMath::Max(1u, ArraySize) : FMath::Max(1u, SizeZ));
+		PartialView->Create(InDevice, Image, ViewType, PartialAspectMask, InCreateDesc.Format, ViewFormat, 0, FMath::Max(InCreateDesc.NumMips, (uint8)1u), 0, bArray ? FMath::Max((uint16)1u, InCreateDesc.ArraySize) : FMath::Max((uint16)1u, InCreateDesc.Depth), false);
 	}
 }
 
-FVulkanTextureBase::FVulkanTextureBase(FTextureRHIRef& SrcTextureRHI, const FVulkanTextureBase* SrcTexture, VkImageViewType ResourceType, uint32 SizeX, uint32 SizeY, uint32 SizeZ)
-	: Surface(*SrcTexture->Surface.Device, ResourceType, SrcTexture->Surface.PixelFormat, SizeX, SizeY, SizeZ, SrcTexture->Surface.ArraySize, SrcTexture->Surface.NumMips, SrcTexture->Surface.NumSamples, SrcTexture->Surface.Image, SrcTexture->Surface.UEFlags, EImageOwnerType::Aliased, FRHIResourceCreateInfo(TEXT("FVulkanTextureBase")))
-	, PartialView(nullptr)
-	, AliasedTexture(SrcTextureRHI)
+FVulkanTexture::FVulkanTexture(FVulkanDevice& InDevice, const FRHITextureCreateDesc& InCreateDesc, FTextureRHIRef& SrcTextureRHI)
+	: FRHITexture(InCreateDesc)
+	, Device(&InDevice)
+	, Image(VK_NULL_HANDLE)
+	, StorageFormat(VK_FORMAT_UNDEFINED)
+	, ViewFormat(VK_FORMAT_UNDEFINED)
+	, MemProps(VK_MEMORY_PROPERTY_DEVICE_LOCAL_BIT)
+	, Tiling(VK_IMAGE_TILING_MAX_ENUM)	// Can be expanded to a per-platform definition
+	, FullAspectMask(0)
+	, PartialAspectMask(0)
+	, CpuReadbackBuffer(nullptr)
+	, ImageOwnerType(EImageOwnerType::Aliased)
 {
-	VULKAN_TRACK_OBJECT_CREATE(FVulkanTextureBase, this);
+	VULKAN_TRACK_OBJECT_CREATE(FVulkanTexture, this);
 
-	if (Surface.FullAspectMask == Surface.PartialAspectMask)
+	{
+		StorageFormat = UEToVkTextureFormat(InCreateDesc.Format, false);
+
+		checkf(InCreateDesc.Format == PF_Unknown || StorageFormat != VK_FORMAT_UNDEFINED, TEXT("PixelFormat %d, is not supported for images"), (int32)InCreateDesc.Format);
+
+		ViewFormat = UEToVkTextureFormat(InCreateDesc.Format, EnumHasAllFlags(InCreateDesc.Flags, TexCreate_SRGB));
+		FullAspectMask = VulkanRHI::GetAspectMaskFromUEFormat(InCreateDesc.Format, true, true);
+		PartialAspectMask = VulkanRHI::GetAspectMaskFromUEFormat(InCreateDesc.Format, false, true);
+
+		// Purely informative patching, we know that "TexCreate_Presentable" uses optimal tiling
+		if (EnumHasAllFlags(InCreateDesc.Flags, TexCreate_Presentable) && GetTiling() == VK_IMAGE_TILING_MAX_ENUM)
+		{
+			Tiling = VK_IMAGE_TILING_OPTIMAL;
+		}
+	}
+
+	if (FullAspectMask == PartialAspectMask)
 	{
 		PartialView = &DefaultView;
 	}
@@ -2052,25 +1924,27 @@ FVulkanTextureBase::FVulkanTextureBase(FTextureRHIRef& SrcTextureRHI, const FVul
 	AliasTextureResources(SrcTextureRHI);
 }
 
-FVulkanTextureBase::~FVulkanTextureBase()
+FVulkanTexture::~FVulkanTexture()
 {
-	VULKAN_TRACK_OBJECT_DELETE(FVulkanTextureBase, this);
+	VULKAN_TRACK_OBJECT_DELETE(FVulkanTexture, this);
 	DestroyViews();
 
 	if (PartialView != &DefaultView)
 	{
 		delete PartialView;
 	}
+
+	DestroySurface();
 }
 
-void FVulkanTextureBase::AliasTextureResources(FTextureRHIRef& SrcTextureRHI)
+void FVulkanTexture::AliasTextureResources(FTextureRHIRef& SrcTextureRHI)
 {
 	DestroyViews();
+	DestroySurface();
 
-	FVulkanTextureBase* SrcTexture = (FVulkanTextureBase*)SrcTextureRHI->GetTextureBaseRHI();
+	FVulkanTexture* SrcTexture = FVulkanTexture::Cast(SrcTextureRHI);
 
-	Surface.Destroy();
-	Surface.Image = SrcTexture->Surface.Image;
+	Image = SrcTexture->Image;
 	DefaultView.View = SrcTexture->DefaultView.View;
 	DefaultView.Image = SrcTexture->DefaultView.Image;
 	DefaultView.ViewId = SrcTexture->DefaultView.ViewId;
@@ -2081,55 +1955,40 @@ void FVulkanTextureBase::AliasTextureResources(FTextureRHIRef& SrcTextureRHI)
 		PartialView->Image = SrcTexture->PartialView->Image;
 		PartialView->ViewId = SrcTexture->PartialView->ViewId;
 	}
+
+	AliasedTexture = SrcTexture;
 }
 
-void FVulkanTextureBase::DestroyViews()
+void FVulkanTexture::DestroyViews()
 {
 	if (AliasedTexture == nullptr)
 	{
-		DefaultView.Destroy(*Surface.Device);
+		DefaultView.Destroy(*Device);
 
 		if (PartialView != &DefaultView && PartialView != nullptr)
 		{
-			PartialView->Destroy(*Surface.Device);
+			PartialView->Destroy(*Device);
 		}
 	}
 }
 
-void FVulkanSurface::Evict(FVulkanDevice& Device_)
-{
-	checkNoEntry(); //not supported
-}
-void FVulkanSurface::Move(FVulkanDevice& Device_, FVulkanCommandListContext& Context, FVulkanAllocation& NewAllocation)
-{
-	checkNoEntry(); //not supported
-}
-
-bool FVulkanSurface::CanEvict()
+void FVulkanTexture::InvalidateViews(FVulkanDevice& InDevice)
 {
-	return true;
-}
+	DefaultView.Destroy(InDevice);
+	const FRHITextureDesc& Desc = GetDesc();
+	const uint32 NumMips = Desc.NumMips;
+	const VkImageViewType ViewType = GetViewType();
+	const bool bArray = ViewType == VK_IMAGE_VIEW_TYPE_1D_ARRAY || ViewType == VK_IMAGE_VIEW_TYPE_2D_ARRAY || ViewType == VK_IMAGE_VIEW_TYPE_CUBE_ARRAY;
+	const uint32 SizeZOrArraySize = bArray ? FMath::Max((uint16)1u, Desc.ArraySize) : FMath::Max((uint16)1u, Desc.Depth);
 
-bool FVulkanSurface::CanMove()
-{
-	return false;
-}
-
-void FVulkanTextureBase::InvalidateViews(FVulkanDevice& Device)
-{
-	DefaultView.Destroy(Device);
-	uint32 NumMips = Surface.GetNumMips();
-	const bool bArray = Surface.ViewType == VK_IMAGE_VIEW_TYPE_1D_ARRAY || Surface.ViewType == VK_IMAGE_VIEW_TYPE_2D_ARRAY || Surface.ViewType == VK_IMAGE_VIEW_TYPE_CUBE_ARRAY;
-	uint32 SizeZOrArraySize = bArray ? FMath::Max(1u, Surface.ArraySize) : FMath::Max(1u, Surface.Depth);
-
-	if(Surface.ViewType != VK_IMAGE_VIEW_TYPE_MAX_ENUM)
+	if(ViewType != VK_IMAGE_VIEW_TYPE_MAX_ENUM)
 	{
-		DefaultView.Create(Device, Surface.Image, Surface.ViewType, Surface.GetFullAspectMask(), Surface.PixelFormat, Surface.ViewFormat, 0, FMath::Max(NumMips, 1u), 0, SizeZOrArraySize, !Surface.SupportsSampling());
+		DefaultView.Create(InDevice, Image, ViewType, GetFullAspectMask(), GetDesc().Format, ViewFormat, 0, FMath::Max(NumMips, 1u), 0, SizeZOrArraySize, !SupportsSampling());
 	}
 	if(PartialView != &DefaultView)
 	{
-		PartialView->Destroy(*Surface.Device);
-		PartialView->Create(Device, Surface.Image, Surface.ViewType, Surface.PartialAspectMask, Surface.PixelFormat, Surface.ViewFormat, 0, FMath::Max(NumMips, 1u), 0, SizeZOrArraySize);
+		PartialView->Destroy(*Device);
+		PartialView->Create(InDevice, Image, ViewType, PartialAspectMask, GetDesc().Format, ViewFormat, 0, FMath::Max(NumMips, 1u), 0, SizeZOrArraySize, false);
 	}
 
 	VulkanRHI::FVulkanViewBase* View = FirstView;
@@ -2140,59 +1999,111 @@ void FVulkanTextureBase::InvalidateViews(FVulkanDevice& Device)
 	}
 }
 
-void FVulkanTextureBase::Move(FVulkanDevice& Device, FVulkanCommandListContext& Context, FVulkanAllocation& NewAllocation)
+void FVulkanTexture::Move(FVulkanDevice& InDevice, FVulkanCommandListContext& Context, FVulkanAllocation& NewAllocation)
 {
-	FRHITexture* Tex = GetRHITexture();
-	uint64 Size = Surface.GetMemorySize();
+	uint64 Size = GetMemorySize();
 	static uint64 TotalSize = 0;
 	TotalSize += Size;
 	if (GVulkanLogDefrag)
 	{
-		UE_LOG(LogVulkanRHI, Display, TEXT("Moving Surface, %p <<-- %p    :::: %s\n"), NewAllocation.Offset, 42, *Tex->GetName().ToString());
-		UE_LOG(LogVulkanRHI, Display, TEXT("Moved %8.4fkb %8.4fkb   TB %p // %p  :: IMG %p   %-40s\n"), Size / (1024.f), TotalSize / (1024.f), this, &Surface, Surface.Image, *GetResourceFName().ToString());
+		UE_LOG(LogVulkanRHI, Display, TEXT("Moving Surface, %p <<-- %p    :::: %s\n"), NewAllocation.Offset, 42, *GetName().ToString());
+		UE_LOG(LogVulkanRHI, Display, TEXT("Moved %8.4fkb %8.4fkb   TB %p  :: IMG %p   %-40s\n"), Size / (1024.f), TotalSize / (1024.f), this, Image, *GetName().ToString());
 	}
-
-	Surface.MoveSurface(Device, Context, NewAllocation);
-	InvalidateViews(Device);
+	{
+		const ETextureCreateFlags UEFlags = GetDesc().Flags;
+		const bool bRenderTarget = EnumHasAnyFlags(UEFlags, TexCreate_RenderTargetable | TexCreate_DepthStencilTargetable | TexCreate_ResolveTargetable);
+		const bool bUAV = EnumHasAnyFlags(UEFlags, TexCreate_UAV);
+		checkf(bRenderTarget || bUAV, TEXT("Surface must be a RenderTarget or a UAV in order to be moved.  UEFlags=0x%x"), (int32)UEFlags);
+		checkf(Tiling == VK_IMAGE_TILING_OPTIMAL, TEXT("Tiling [%d] is not supported for move, only VK_IMAGE_TILING_OPTIMAL"), (int32)Tiling);
+	
+		InternalMoveSurface(InDevice, Context, NewAllocation);
+	
+		// Swap in the new allocation for this surface
+		Allocation.Swap(NewAllocation);
+	}
+	InvalidateViews(InDevice);
 }
 
-void FVulkanTextureBase::OnFullDefrag(FVulkanDevice& Device, FVulkanCommandListContext& Context, uint32 NewOffset)
+void FVulkanTexture::OnFullDefrag(FVulkanDevice& InDevice, FVulkanCommandListContext& Context, uint32 NewOffset)
 {
-	FRHITexture* Tex = GetRHITexture();
-	uint64 Size = Surface.GetMemorySize();
+	uint64 Size = GetMemorySize();
 	static uint64 TotalSize = 0;
 	TotalSize += Size;
 	if (GVulkanLogDefrag)
 	{
-		UE_LOG(LogVulkanRHI, Display, TEXT("Moving Surface, %p <<-- %p    :::: %s\n"), NewOffset, 42, *Tex->GetName().ToString());
-		UE_LOG(LogVulkanRHI, Display, TEXT("Defragged %8.4fkb %8.4fkb   TB %p // %p  :: IMG %p   %-40s\n"), Size / (1024.f), TotalSize / (1024.f), this, &Surface, Surface.Image, *GetResourceFName().ToString());
+		UE_LOG(LogVulkanRHI, Display, TEXT("Moving Surface, %p <<-- %p    :::: %s\n"), NewOffset, 42, *GetName().ToString());
+		UE_LOG(LogVulkanRHI, Display, TEXT("Defragged %8.4fkb %8.4fkb   TB %p  :: IMG %p   %-40s\n"), Size / (1024.f), TotalSize / (1024.f), this, Image, *GetName().ToString());
 	}
 
-	Surface.OnFullDefrag(Device, Context, NewOffset);
-	InvalidateViews(Device);
-}
+	{
+		const ETextureCreateFlags UEFlags = GetDesc().Flags;
+		const bool bRenderTarget = EnumHasAnyFlags(UEFlags, TexCreate_RenderTargetable | TexCreate_DepthStencilTargetable | TexCreate_ResolveTargetable);
+		const bool bUAV = EnumHasAnyFlags(UEFlags, TexCreate_UAV);
+		checkf(bRenderTarget || bUAV, TEXT("Surface must be a RenderTarget or a UAV in order to be defragged.  UEFlags=0x%x"), (int32)UEFlags);
+		checkf(Tiling == VK_IMAGE_TILING_OPTIMAL, TEXT("Tiling [%d] is not supported for defrag, only VK_IMAGE_TILING_OPTIMAL"), (int32)Tiling);
 
+		Allocation.Offset = NewOffset;
+		InternalMoveSurface(InDevice, Context, Allocation);
 
+		//note: this exploits that the unmoved image is still bound to the old allocation, which is freed by the caller in this case.
+	}
 
+	InvalidateViews(InDevice);
+}
 
-void FVulkanTextureBase::Evict(FVulkanDevice& Device)
+void FVulkanTexture::Evict(FVulkanDevice& InDevice)
 {
 	check(AliasedTexture == nullptr); //can't evict textures we don't own
-	uint64 Size = Surface.GetMemorySize();
+	uint64 Size = GetMemorySize();
 	static uint64 TotalSize = 0;
 	TotalSize += Size;
 	if(GVulkanLogDefrag)
 	{
-		FGenericPlatformMisc::LowLevelOutputDebugStringf(TEXT("Evicted %8.4fkb %8.4fkb   TB %p // %p  :: IMG %p   %-40s\n"), Size / (1024.f), TotalSize / (1024.f), this, &Surface, Surface.Image, *GetResourceFName().ToString());
+		FGenericPlatformMisc::LowLevelOutputDebugStringf(TEXT("Evicted %8.4fkb %8.4fkb   TB %p  :: IMG %p   %-40s\n"), Size / (1024.f), TotalSize / (1024.f), this, Image, *GetName().ToString());
+	}
+
+	{
+		check(0 == CpuReadbackBuffer);
+		checkf(MemProps == VK_MEMORY_PROPERTY_DEVICE_LOCAL_BIT, TEXT("Can't evict surface that isn't device local.  MemoryProperties=%d"), (int32)MemProps);
+		checkf(VulkanRHI::GetAspectMaskFromUEFormat(GetDesc().Format, true, true) == FullAspectMask, TEXT("FullAspectMask (%d) does not match with PixelFormat (%d)"), (int32)FullAspectMask, (int32)GetDesc().Format);
+		checkf(VulkanRHI::GetAspectMaskFromUEFormat(GetDesc().Format, false, true) == PartialAspectMask, TEXT("PartialAspectMask (%d) does not match with PixelFormat (%d)"), (int32)PartialAspectMask, (int32)GetDesc().Format);
+
+		const ETextureCreateFlags UEFlags = GetDesc().Flags;
+		const bool bRenderTarget = EnumHasAnyFlags(UEFlags, TexCreate_RenderTargetable | TexCreate_DepthStencilTargetable | TexCreate_ResolveTargetable);
+		const bool bUAV = EnumHasAnyFlags(UEFlags, TexCreate_UAV);
+		//none of this is supported for eviction
+		checkf(!bRenderTarget, TEXT("RenderTargets do not support evict."));
+		checkf(!bUAV, TEXT("UAV do not support evict."));
+
+		FRHICommandListImmediate& RHICmdList = FRHICommandListExecutor::GetImmediateCommandList();
+		FVulkanCommandListContext& Context = (FVulkanCommandListContext&)RHICmdList.GetContext().GetLowestLevelContext();
+
+		MemProps = InDevice.GetDeviceMemoryManager().GetEvictedMemoryProperties();
+
+		// Create a new host allocation to move the surface to
+		FVulkanAllocation HostAllocation;
+		const EVulkanAllocationMetaType MetaType = EVulkanAllocationMetaImageOther;
+		if (!InDevice.GetMemoryManager().AllocateImageMemory(HostAllocation, this, MemoryRequirements, MemProps, MetaType, false, __FILE__, __LINE__))
+		{
+			InDevice.GetMemoryManager().HandleOOM();
+			checkNoEntry();
+		}
+
+		InternalMoveSurface(InDevice, Context, HostAllocation);
+
+		// Delete the original allocation and swap in the new host allocation
+		Device->GetMemoryManager().FreeVulkanAllocation(Allocation);
+		Allocation.Swap(HostAllocation);
+
+		VULKAN_SET_DEBUG_NAME(InDevice, VK_OBJECT_TYPE_IMAGE, Image, TEXT("(FVulkanTexture*)0x%p [hostimage]"), this);
 	}
 
-	Surface.EvictSurface(Device);
-	InvalidateViews(Device);
+	InvalidateViews(InDevice);
 }
 
 
 static FCriticalSection ViewCritSection;
-void FVulkanTextureBase::AttachView(FVulkanViewBase* View)
+void FVulkanTexture::AttachView(FVulkanViewBase* View)
 {
 	FScopeLock Lock(&ViewCritSection);
 	check(View->NextView == nullptr);
@@ -2200,7 +2111,7 @@ void FVulkanTextureBase::AttachView(FVulkanViewBase* View)
 	FirstView = View;
 }
 
-void FVulkanTextureBase::DetachView(FVulkanViewBase* View)
+void FVulkanTexture::DetachView(FVulkanViewBase* View)
 {
 	FScopeLock Lock(&ViewCritSection);
 	FVulkanViewBase** NextViewPtr = &FirstView;
@@ -2212,108 +2123,21 @@ void FVulkanTextureBase::DetachView(FVulkanViewBase* View)
 	View->NextView = nullptr;
 }
 
-bool FVulkanTextureBase::GetTextureResourceInfo(FRHIResourceInfo& OutResourceInfo) const
+bool FVulkanTexture::GetTextureResourceInfo(FRHIResourceInfo& OutResourceInfo) const
 {
 	OutResourceInfo = FRHIResourceInfo();
-	OutResourceInfo.VRamAllocation.AllocationSize = Surface.GetMemorySize();
+	OutResourceInfo.VRamAllocation.AllocationSize = GetMemorySize();
 	return true;
 }
 
-FVulkanTexture2D::FVulkanTexture2D(FVulkanDevice& Device, EPixelFormat InFormat, uint32 SizeX, uint32 SizeY, uint32 NumMips, uint32 NumSamples, ETextureCreateFlags UEFlags, ERHIAccess InResourceState, const FRHIResourceCreateInfo& CreateInfo, const FRHITransientHeapAllocation* InTransientHeapAllocation)
-:	FRHITexture2D(SizeX, SizeY, FMath::Max(NumMips, 1u), NumSamples, InFormat, UEFlags, CreateInfo.ClearValueBinding)
-,	FVulkanTextureBase(Device, VK_IMAGE_VIEW_TYPE_2D, InFormat, SizeX, SizeY, 1, 1, FMath::Max(NumMips, 1u), NumSamples, UEFlags, InResourceState, CreateInfo, InTransientHeapAllocation)
-{
-}
-
-FVulkanTexture2D::FVulkanTexture2D(FVulkanDevice& Device, EPixelFormat Format, uint32 SizeX, uint32 SizeY, uint32 NumMips, uint32 NumSamples, VkImage Image, ETextureCreateFlags UEFlags, const FRHIResourceCreateInfo& CreateInfo)
-:	FRHITexture2D(SizeX, SizeY, NumMips, NumSamples, Format, UEFlags, CreateInfo.ClearValueBinding)
-,	FVulkanTextureBase(Device, VK_IMAGE_VIEW_TYPE_2D, Format, SizeX, SizeY, 1, 1, NumMips, NumSamples, Image, VK_NULL_HANDLE, UEFlags, CreateInfo)
-{
-}
-
-FVulkanTexture2D::FVulkanTexture2D(FTextureRHIRef& SrcTextureRHI, const FVulkanTexture2D* SrcTexture)
-	: FRHITexture2D(SrcTexture->GetSizeX(), SrcTexture->GetSizeY(), SrcTexture->GetNumMips(), SrcTexture->GetNumSamples(), SrcTexture->GetFormat(), SrcTexture->GetFlags(), SrcTexture->GetClearBinding())
-	, FVulkanTextureBase(SrcTextureRHI, static_cast<const FVulkanTextureBase*>(SrcTexture), VK_IMAGE_VIEW_TYPE_2D, SrcTexture->GetSizeX(), SrcTexture->GetSizeY(), 1)
-{
-}
-
-FVulkanTexture2D::~FVulkanTexture2D()
-{
-}
-
-FVulkanTexture2DArray::FVulkanTexture2DArray(FVulkanDevice& Device, EPixelFormat Format, uint32 SizeX, uint32 SizeY, uint32 ArraySize, uint32 NumMips, uint32 NumSamples, ETextureCreateFlags Flags, ERHIAccess InResourceState, const FRHIResourceCreateInfo& CreateInfo, const FRHITransientHeapAllocation* InTransientHeapAllocation)
-	:	FRHITexture2DArray(SizeX, SizeY, ArraySize, NumMips, NumSamples, Format, Flags, CreateInfo.ClearValueBinding)
-	,	FVulkanTextureBase(Device, VK_IMAGE_VIEW_TYPE_2D_ARRAY, Format, SizeX, SizeY, 1, ArraySize, NumMips, NumSamples, Flags, InResourceState, CreateInfo, InTransientHeapAllocation)
-{
-}
-
-FVulkanTexture2DArray::FVulkanTexture2DArray(FVulkanDevice& Device, EPixelFormat Format, uint32 SizeX, uint32 SizeY, uint32 ArraySize, uint32 NumMips, uint32 NumSamples, VkImage Image, ETextureCreateFlags Flags, const FRHIResourceCreateInfo& CreateInfo)
-	:	FRHITexture2DArray(SizeX, SizeY, ArraySize, NumMips, NumSamples, Format, Flags, CreateInfo.ClearValueBinding)
-	,	FVulkanTextureBase(Device, VK_IMAGE_VIEW_TYPE_2D_ARRAY, Format, SizeX, SizeY, 1, ArraySize, NumMips, NumSamples, Image, VK_NULL_HANDLE, Flags, CreateInfo)
-{
-}
-
-FVulkanTexture2DArray::FVulkanTexture2DArray(FTextureRHIRef& SrcTextureRHI, const FVulkanTexture2DArray* SrcTexture)
-	: FRHITexture2DArray(SrcTexture->GetSizeX(), SrcTexture->GetSizeY(), SrcTexture->Surface.GetNumberOfArrayLevels(), SrcTexture->GetNumMips(), SrcTexture->GetNumSamples(), SrcTexture->GetFormat(), SrcTexture->GetFlags(), SrcTexture->GetClearBinding())
-	, FVulkanTextureBase(SrcTextureRHI, static_cast<const FVulkanTextureBase*>(SrcTexture), VK_IMAGE_VIEW_TYPE_2D_ARRAY, SrcTexture->GetSizeX(), SrcTexture->GetSizeY(), 1)
-{
-}
-
-
-FVulkanTextureCube::FVulkanTextureCube(FVulkanDevice& Device, EPixelFormat Format, uint32 Size, bool bArray, uint32 ArraySize, uint32 NumMips, ETextureCreateFlags Flags, ERHIAccess InResourceState, const FRHIResourceCreateInfo& CreateInfo, const FRHITransientHeapAllocation* InTransientHeapAllocation)
-:	 FRHITextureCube(Size, NumMips, Format, Flags, CreateInfo.ClearValueBinding)
-	//#todo-rco: Array/slices count
-,	FVulkanTextureBase(Device, bArray ? VK_IMAGE_VIEW_TYPE_CUBE_ARRAY : VK_IMAGE_VIEW_TYPE_CUBE, Format, Size, Size, 1, ArraySize, NumMips, /*NumSamples=*/ 1, Flags, InResourceState, CreateInfo, InTransientHeapAllocation)
-{
-}
-
-FVulkanTextureCube::FVulkanTextureCube(FVulkanDevice& Device, EPixelFormat Format, uint32 Size, bool bArray, uint32 ArraySize, uint32 NumMips, VkImage Image, ETextureCreateFlags Flags, const FRHIResourceCreateInfo& CreateInfo)
-	:	 FRHITextureCube(Size, NumMips, Format, Flags, CreateInfo.ClearValueBinding)
-	//#todo-rco: Array/slices count
-	,	FVulkanTextureBase(Device, bArray ? VK_IMAGE_VIEW_TYPE_CUBE_ARRAY : VK_IMAGE_VIEW_TYPE_CUBE, Format, Size, Size, 1, ArraySize, NumMips, /*NumSamples=*/ 1, Image, VK_NULL_HANDLE, Flags, CreateInfo)
-{
-}
-
-FVulkanTextureCube::FVulkanTextureCube(FTextureRHIRef& SrcTextureRHI, const FVulkanTextureCube* SrcTexture)
-	: FRHITextureCube(SrcTexture->GetSize(), SrcTexture->GetNumMips(), SrcTexture->GetFormat(), SrcTexture->GetFlags(), SrcTexture->GetClearBinding())
-	, FVulkanTextureBase(SrcTextureRHI, static_cast<const FVulkanTextureBase*>(SrcTexture), VK_IMAGE_VIEW_TYPE_CUBE, SrcTexture->GetSize(), SrcTexture->GetSize(), 1)
-{
-}
-
-FVulkanTextureCube::~FVulkanTextureCube()
-{
-}
-
-
-FVulkanTexture3D::FVulkanTexture3D(FVulkanDevice& Device, EPixelFormat Format, uint32 SizeX, uint32 SizeY, uint32 SizeZ, uint32 NumMips, ETextureCreateFlags Flags, ERHIAccess InResourceState, const FRHIResourceCreateInfo& CreateInfo, const FRHITransientHeapAllocation* InTransientHeapAllocation)
-	: FRHITexture3D(SizeX, SizeY, SizeZ, NumMips, Format, Flags, CreateInfo.ClearValueBinding)
-	, FVulkanTextureBase(Device, VK_IMAGE_VIEW_TYPE_3D, Format, SizeX, SizeY, SizeZ, 1, NumMips, /*NumSamples=*/ 1, Flags, InResourceState, CreateInfo, InTransientHeapAllocation)
-{
-}
-
-FVulkanTexture3D::~FVulkanTexture3D()
-{
-}
-
 /*-----------------------------------------------------------------------------
 	Cubemap texture support.
 -----------------------------------------------------------------------------*/
-FTextureCubeRHIRef FVulkanDynamicRHI::RHICreateTextureCube(uint32 Size, uint8 Format, uint32 NumMips, ETextureCreateFlags Flags, ERHIAccess InResourceState, FRHIResourceCreateInfo& CreateInfo)
-{
-	LLM_SCOPE_VULKAN(GetMemoryTagForTextureFlags(Flags));
-	return new FVulkanTextureCube(*Device, (EPixelFormat)Format, Size, false, 1, NumMips, Flags, InResourceState, CreateInfo);
-}
-
-FTextureCubeRHIRef FVulkanDynamicRHI::RHICreateTextureCubeArray(uint32 Size, uint32 ArraySize, uint8 Format, uint32 NumMips, ETextureCreateFlags Flags, ERHIAccess InResourceState, FRHIResourceCreateInfo& CreateInfo)
-{
-	LLM_SCOPE_VULKAN(GetMemoryTagForTextureFlags(Flags));
-	return new FVulkanTextureCube(*Device, (EPixelFormat)Format, Size, true, ArraySize, NumMips, Flags, InResourceState, CreateInfo);
-}
 
 void* FVulkanDynamicRHI::RHILockTextureCubeFace(FRHITextureCube* TextureCubeRHI, uint32 FaceIndex, uint32 ArrayIndex, uint32 MipIndex, EResourceLockMode LockMode, uint32& DestStride, bool bLockWithinMiptail)
 {
 	LLM_SCOPE_VULKAN(ELLMTagVulkan::VulkanTextures);
-	FVulkanTextureCube* Texture = ResourceCast(TextureCubeRHI);
+	FVulkanTexture* Texture = FVulkanTexture::Cast(TextureCubeRHI);
 	check(Texture);
 
 	VulkanRHI::FStagingBuffer** StagingBuffer = nullptr;
@@ -2325,18 +2149,18 @@ void* FVulkanDynamicRHI::RHILockTextureCubeFace(FRHITextureCube* TextureCubeRHI,
 
 	uint32 BufferSize = 0;
 	DestStride = 0;
-	Texture->Surface.GetMipSize(MipIndex, BufferSize);
-	Texture->Surface.GetMipStride(MipIndex, DestStride);
+	Texture->GetMipSize(MipIndex, BufferSize);
+	Texture->GetMipStride(MipIndex, DestStride);
 	*StagingBuffer = Device->GetStagingManager().AcquireBuffer(BufferSize);
 
 	void* Data = (*StagingBuffer)->GetMappedPointer();
 	return Data;
 }
 
-void FVulkanDynamicRHI::RHIUnlockTextureCubeFace(FRHITextureCube* TextureCubeRHI,uint32 FaceIndex,uint32 ArrayIndex,uint32 MipIndex,bool bLockWithinMiptail)
+void FVulkanDynamicRHI::RHIUnlockTextureCubeFace(FRHITextureCube* TextureCubeRHI, uint32 FaceIndex, uint32 ArrayIndex, uint32 MipIndex, bool bLockWithinMiptail)
 {
 	LLM_SCOPE_VULKAN(ELLMTagVulkan::VulkanTextures);
-	FVulkanTextureCube* Texture = ResourceCast(TextureCubeRHI);
+	FVulkanTexture* Texture = FVulkanTexture::Cast(TextureCubeRHI);
 	check(Texture);
 
 	VkDevice LogicalDevice = Device->GetInstanceHandle();
@@ -2348,9 +2172,9 @@ void FVulkanDynamicRHI::RHIUnlockTextureCubeFace(FRHITextureCube* TextureCubeRHI
 		checkf(bFound, TEXT("Texture was not locked!"));
 	}
 
-	EPixelFormat Format = Texture->Surface.PixelFormat;
-	uint32 MipWidth = FMath::Max<uint32>(Texture->Surface.Width >> MipIndex, 0);
-	uint32 MipHeight = FMath::Max<uint32>(Texture->Surface.Height >> MipIndex, 0);
+	const FRHITextureDesc& Desc = Texture->GetDesc();
+	uint32 MipWidth = FMath::Max<uint32>(Desc.Extent.X >> MipIndex, 0);
+	uint32 MipHeight = FMath::Max<uint32>(Desc.Extent.Y >> MipIndex, 0);
 	ensure(!(MipHeight == 0 && MipWidth == 0));
 	MipWidth = FMath::Max<uint32>(MipWidth, 1);
 	MipHeight = FMath::Max<uint32>(MipHeight, 1);
@@ -2361,7 +2185,7 @@ void FVulkanDynamicRHI::RHIUnlockTextureCubeFace(FRHITextureCube* TextureCubeRHI
 	//Region.bufferOffset = 0;
 	//Region.bufferRowLength = 0;
 	//Region.bufferImageHeight = 0;
-	Region.imageSubresource.aspectMask = Texture->Surface.GetPartialAspectMask();
+	Region.imageSubresource.aspectMask = Texture->GetPartialAspectMask();
 	Region.imageSubresource.mipLevel = MipIndex;
 	Region.imageSubresource.baseArrayLayer = ArrayIndex * 6 + FaceIndex;
 	Region.imageSubresource.layerCount = 1;
@@ -2372,12 +2196,12 @@ void FVulkanDynamicRHI::RHIUnlockTextureCubeFace(FRHITextureCube* TextureCubeRHI
 	FRHICommandList& RHICmdList = FRHICommandListExecutor::GetImmediateCommandList();
 	if (RHICmdList.Bypass() || !IsRunningRHIInSeparateThread())
 	{
-		FVulkanSurface::InternalLockWrite(Device->GetImmediateContext(), &Texture->Surface, Region, StagingBuffer);
+		FVulkanTexture::InternalLockWrite(Device->GetImmediateContext(), Texture, Region, StagingBuffer);
 	}
 	else
 	{
 		check(IsInRenderingThread());
-		ALLOC_COMMAND_CL(RHICmdList, FRHICommandLockWriteTexture)(&Texture->Surface, Region, StagingBuffer);
+		ALLOC_COMMAND_CL(RHICmdList, FRHICommandLockWriteTexture)(Texture, Region, StagingBuffer);
 	}
 }
 
@@ -2385,8 +2209,8 @@ void FVulkanDynamicRHI::RHIBindDebugLabelName(FRHITexture* TextureRHI, const TCH
 {
 #if VULKAN_ENABLE_IMAGE_TRACKING_LAYER
 	{
-		FVulkanTextureBase* Base = (FVulkanTextureBase*)TextureRHI->GetTextureBaseRHI();
-		VulkanRHI::BindDebugLabelName(Base->Surface.Image, Name);
+		FVulkanTexture* VulkanTexture = FVulkanTexture::Cast(TextureRHI);
+		VulkanRHI::BindDebugLabelName(VulkanTexture->Image, Name);
 	}
 #endif
 
@@ -2394,9 +2218,9 @@ void FVulkanDynamicRHI::RHIBindDebugLabelName(FRHITexture* TextureRHI, const TCH
 	{
 // TODO: this dies in the printf on android. Needs investigation.
 #if !PLATFORM_ANDROID
-		FVulkanTextureBase* Base = (FVulkanTextureBase*)TextureRHI->GetTextureBaseRHI();
+		FVulkanTexture* VulkanTexture = FVulkanTexture::Cast(TextureRHI);
 #if VULKAN_ENABLE_DUMP_LAYER
-		VulkanRHI::PrintfBegin(*FString::Printf(TEXT("vkDebugMarkerSetObjectNameEXT(0x%p=%s)\n"), Base->Surface.Image, Name));
+		VulkanRHI::PrintfBegin(*FString::Printf(TEXT("vkDebugMarkerSetObjectNameEXT(0x%p=%s)\n"), VulkanTexture->Image, Name));
 #endif
 #endif
 	}
@@ -2406,17 +2230,17 @@ void FVulkanDynamicRHI::RHIBindDebugLabelName(FRHITexture* TextureRHI, const TCH
 #if 0//VULKAN_SUPPORTS_DEBUG_UTILS
 	if (auto* SetDebugName = Device->GetSetDebugName())
 	{
-		FVulkanTextureBase* Base = (FVulkanTextureBase*)TextureRHI->GetTextureBaseRHI();
+		FVulkanTexture* VulkanTexture = FVulkanTexture::Cast(TextureRHI);
 		FTCHARToUTF8 Converter(Name);
-		VulkanRHI::SetDebugName(SetDebugName, Device->GetInstanceHandle(), Base->Surface.Image, Converter.Get());
+		VulkanRHI::SetDebugName(SetDebugName, Device->GetInstanceHandle(), VulkanTexture->Image, Converter.Get());
 	}
 	else
 #endif
 	if (auto* SetObjectName = Device->GetDebugMarkerSetObjectName())
 	{
-		FVulkanTextureBase* Base = (FVulkanTextureBase*)TextureRHI->GetTextureBaseRHI();
+		FVulkanTexture* VulkanTexture = FVulkanTexture::Cast(TextureRHI);
 		FTCHARToUTF8 Converter(Name);
-		VulkanRHI::SetDebugMarkerName(SetObjectName, Device->GetInstanceHandle(), Base->Surface.Image, Converter.Get());
+		VulkanRHI::SetDebugMarkerName(SetObjectName, Device->GetInstanceHandle(), VulkanTexture->Image, Converter.Get());
 	}
 #endif
 	FName DebugName(Name);
@@ -2434,7 +2258,7 @@ void FVulkanDynamicRHI::RHIBindDebugLabelName(FRHIUnorderedAccessView* Unordered
 		//	VkDebugMarkerObjectTagInfoEXT Info;
 		//	ZeroVulkanStruct(Info, VK_STRUCTURE_TYPE_DEBUG_MARKER_OBJECT_NAME_INFO_EXT);
 		//	Info.objectType = VK_DEBUG_REPORT_OBJECT_TYPE_IMAGE_EXT;
-		//	Info.object = VulkanTexture->Surface.Image;
+		//	Info.object = VulkanTexture->Image;
 		//	vkDebugMarkerSetObjectNameEXT(Device->GetInstanceHandle(), &Info);
 		//}
 	}
@@ -2452,133 +2276,51 @@ void FVulkanDynamicRHI::RHIVirtualTextureSetFirstMipVisible(FRHITexture2D* Textu
 	VULKAN_SIGNAL_UNIMPLEMENTED();
 }
 
-static VkMemoryRequirements FindOrCalculateTexturePlatformSize(FVulkanDevice* Device, VkImageViewType ViewType, uint32 SizeX, uint32 SizeY, uint32 SizeZ, uint32 ArraySize, uint8 Format, uint32 NumMips, uint32 NumSamples, ETextureCreateFlags Flags)
+FDynamicRHI::FRHICalcTextureSizeResult FVulkanDynamicRHI::RHICalcTexturePlatformSize(FRHITextureDesc const& Desc, uint32 FirstMipIndex)
 {
+	// FIXME: this function ignores FirstMipIndex!
+
+	// Zero out the members which don't affect the size since we'll use this as a key in the map of already computed sizes.
+	FRHITextureDesc CleanDesc = Desc;
+	CleanDesc.UAVFormat = PF_Unknown;
+	CleanDesc.ClearValue = FClearValueBinding::None;
+	CleanDesc.ExtData = 0;
+
 	// Adjust number of mips as UTexture can request non-valid # of mips
-	NumMips = FMath::Min(FMath::FloorLog2(FMath::Max(SizeX, FMath::Max(SizeY, SizeZ))) + 1, NumMips);
-
-	struct FTexturePlatformSizeKey
-	{
-		VkImageViewType ViewType;
-		uint32 SizeX;
-		uint32 SizeY;
-		uint32 SizeZ;
-		uint32 Format;
-		uint32 NumMips;
-		uint32 NumSamples;
-		ETextureCreateFlags Flags;
-	};
+	CleanDesc.NumMips = (uint8)FMath::Min(FMath::FloorLog2(FMath::Max(CleanDesc.Extent.X, FMath::Max(CleanDesc.Extent.Y, (int32)CleanDesc.Depth))) + 1, (uint32)CleanDesc.NumMips);
 
-	static TMap<uint32, VkMemoryRequirements> TextureSizes;
+	static TMap<FRHITextureDesc, VkMemoryRequirements> TextureSizes;
 	static FCriticalSection TextureSizesLock;
 
-	const FTexturePlatformSizeKey Key = { ViewType, SizeX, SizeY, SizeZ, Format, NumMips, NumSamples, Flags};
-	const uint32 Hash = FCrc::MemCrc32(&Key, sizeof(FTexturePlatformSizeKey));
-
 	VkMemoryRequirements* Found = nullptr;
 	{
 		FScopeLock Lock(&TextureSizesLock);
-		Found = TextureSizes.Find(Hash);
+		Found = TextureSizes.Find(CleanDesc);
 		if (Found)
 		{
-			return *Found;
+			return { (uint64)Found->size, (uint32)Found->alignment };
 		}
 	}
 
-	EPixelFormat PixelFormat = (EPixelFormat)Format;
-	VkMemoryRequirements MemReq;
-
-	// Create temporary image to measure the memory requirements
-	FVulkanSurface::FImageCreateInfo TmpCreateInfo;
-	FVulkanSurface::GenerateImageCreateInfo(TmpCreateInfo, *Device, ViewType,
-		PixelFormat, SizeX, SizeY, SizeZ, ArraySize, NumMips, NumSamples,
-		Flags, nullptr, nullptr, false);
+	// Create temporary image to measure the memory requirements.
+	FVulkanTexture::FImageCreateInfo TmpCreateInfo;
+	FVulkanTexture::GenerateImageCreateInfo(TmpCreateInfo, *Device, CleanDesc, nullptr, nullptr, false);
 
-	// :TODO: VK_KHR_maintenance4
+	// TODO: VK_KHR_maintenance4
 	VkImage TmpImage;
 	VERIFYVULKANRESULT(VulkanRHI::vkCreateImage(Device->GetInstanceHandle(), &TmpCreateInfo.ImageCreateInfo, VULKAN_CPU_ALLOCATOR, &TmpImage));
+
+	VkMemoryRequirements MemReq;
 	VulkanRHI::vkGetImageMemoryRequirements(Device->GetInstanceHandle(), TmpImage, &MemReq);
+
 	VulkanRHI::vkDestroyImage(Device->GetInstanceHandle(), TmpImage, VULKAN_CPU_ALLOCATOR);
 
 	{
 		FScopeLock Lock(&TextureSizesLock);
-		TextureSizes.Add(Hash, MemReq);
+		TextureSizes.Add(CleanDesc, MemReq);
 	}
-	
-	return MemReq;
-}
 
-uint64 FVulkanDynamicRHI::RHICalcTexture2DPlatformSize(uint32 SizeX, uint32 SizeY, uint8 Format, uint32 NumMips, uint32 NumSamples, ETextureCreateFlags Flags, const FRHIResourceCreateInfo& CreateInfo, uint32& OutAlign)
-{
-	const VkMemoryRequirements MemReq = FindOrCalculateTexturePlatformSize(Device, VK_IMAGE_VIEW_TYPE_2D, SizeX, SizeY, 1, 1, Format, NumMips, NumSamples, Flags);
-	OutAlign = MemReq.alignment;
-	return MemReq.size;
-}
-
-uint64 FVulkanDynamicRHI::RHICalcTexture2DArrayPlatformSize(uint32 SizeX, uint32 SizeY, uint32 ArraySize, uint8 Format, uint32 NumMips, uint32 NumSamples, ETextureCreateFlags Flags, const FRHIResourceCreateInfo& CreateInfo, uint32& OutAlign)
-{
-	const VkMemoryRequirements MemReq = FindOrCalculateTexturePlatformSize(Device, VK_IMAGE_VIEW_TYPE_2D_ARRAY, SizeX, SizeY, 1, ArraySize, Format, NumMips, NumSamples, Flags);
-	OutAlign = MemReq.alignment;
-	return MemReq.size;
-}
-
-uint64 FVulkanDynamicRHI::RHICalcTexture3DPlatformSize(uint32 SizeX, uint32 SizeY, uint32 SizeZ, uint8 Format, uint32 NumMips, ETextureCreateFlags Flags, const FRHIResourceCreateInfo& CreateInfo, uint32& OutAlign)
-{
-	const VkMemoryRequirements MemReq = FindOrCalculateTexturePlatformSize(Device, VK_IMAGE_VIEW_TYPE_3D, SizeX, SizeY, SizeZ, 1, Format, NumMips, 1, Flags);
-	OutAlign = MemReq.alignment;
-	return MemReq.size;
-}
-
-uint64 FVulkanDynamicRHI::RHICalcTextureCubePlatformSize(uint32 Size, uint8 Format, uint32 NumMips, ETextureCreateFlags Flags, const FRHIResourceCreateInfo& CreateInfo, uint32& OutAlign)
-{
-	const VkMemoryRequirements MemReq = FindOrCalculateTexturePlatformSize(Device, VK_IMAGE_VIEW_TYPE_CUBE, Size, Size, 1, 1, Format, NumMips, 1, Flags);
-	OutAlign = MemReq.alignment;
-	return MemReq.size;
-}
-
-uint64 FVulkanDynamicRHI::RHICalcTexturePlatformSize(const FRHITextureCreateInfo& InCreateInfo, uint32& OutAlign)
-{
-	// Validate some params and determine the view type
-	VkImageViewType ViewType = VK_IMAGE_VIEW_TYPE_2D;
-	switch (InCreateInfo.Dimension)
-	{
-	case ETextureDimension::Texture2D:
-		check(InCreateInfo.ArraySize == 1);
-		check(InCreateInfo.Depth == 1);
-		ViewType = VK_IMAGE_VIEW_TYPE_2D;
-		break;
-
-	case ETextureDimension::Texture2DArray:
-		check(InCreateInfo.Depth == 1);
-		ViewType = VK_IMAGE_VIEW_TYPE_2D_ARRAY;
-		break;
-
-	case ETextureDimension::Texture3D:
-		check(InCreateInfo.ArraySize == 1);
-		check(InCreateInfo.NumSamples == 1);
-		ViewType = VK_IMAGE_VIEW_TYPE_3D;
-		break;
-
-	case ETextureDimension::TextureCube:
-		check(InCreateInfo.ArraySize == 1);
-		check(InCreateInfo.Depth == 1);
-		check(InCreateInfo.NumSamples == 1);
-		ViewType = VK_IMAGE_VIEW_TYPE_CUBE;
-		break;
-
-	case ETextureDimension::TextureCubeArray:
-		check(InCreateInfo.Depth == 1);
-		check(InCreateInfo.NumSamples == 1);
-		ViewType = VK_IMAGE_VIEW_TYPE_CUBE_ARRAY;
-		break;
-
-	default:
-		checkNoEntry();
-	}
-
-	VkMemoryRequirements MemReq = FindOrCalculateTexturePlatformSize(Device, ViewType, InCreateInfo.Extent.X, InCreateInfo.Extent.Y, InCreateInfo.Depth, InCreateInfo.ArraySize, InCreateInfo.Format, InCreateInfo.NumMips, InCreateInfo.NumSamples, InCreateInfo.Flags);
-	OutAlign = MemReq.alignment;
-	return MemReq.size;
+	return { (uint64)MemReq.size, (uint32)MemReq.alignment };
 }
 
 void FVulkanCommandListContext::RHICopyTexture(FRHITexture* SourceTexture, FRHITexture* DestTexture, const FRHICopyTextureInfo& CopyInfo)
@@ -2586,46 +2328,45 @@ void FVulkanCommandListContext::RHICopyTexture(FRHITexture* SourceTexture, FRHIT
 	LLM_SCOPE_VULKAN(ELLMTagVulkan::VulkanTextures);
 	check(SourceTexture && DestTexture);
 
-	FVulkanTextureBase* Source = static_cast<FVulkanTextureBase*>(SourceTexture->GetTextureBaseRHI());
-	FVulkanTextureBase* Dest = static_cast<FVulkanTextureBase*>(DestTexture->GetTextureBaseRHI());
-
-	FVulkanSurface& SrcSurface = Source->Surface;
-	FVulkanSurface& DstSurface = Dest->Surface;
+	FVulkanTexture* Source = FVulkanTexture::Cast(SourceTexture);
+	FVulkanTexture* Dest = FVulkanTexture::Cast(DestTexture);
 
-	VkImageLayout SrcLayout = LayoutManager.FindLayoutChecked(SrcSurface.Image);
+	VkImageLayout SrcLayout = LayoutManager.FindLayoutChecked(Source->Image);
 	ensureMsgf(SrcLayout == VK_IMAGE_LAYOUT_TRANSFER_SRC_OPTIMAL, TEXT("Expected source texture to be in VK_IMAGE_LAYOUT_TRANSFER_SRC_OPTIMAL, actual layout is %d"), SrcLayout);
 
 	FVulkanCmdBuffer* InCmdBuffer = CommandBufferManager->GetActiveCmdBuffer();
 	check(InCmdBuffer->IsOutsideRenderPass());
 	VkCommandBuffer CmdBuffer = InCmdBuffer->GetHandle();
 
-	const FPixelFormatInfo& PixelFormatInfo = GPixelFormats[DestTexture->GetFormat()];
+	const FPixelFormatInfo& PixelFormatInfo = GPixelFormats[DestTexture->GetDesc().Format];
+	const FIntVector SourceXYZ = Source->GetSizeXYZ();
+	const FIntVector DestXYZ = Source->GetSizeXYZ();
 
-	check(!EnumHasAnyFlags(SrcSurface.UEFlags, TexCreate_CPUReadback));
-	if(EnumHasAllFlags(DstSurface.UEFlags, TexCreate_CPUReadback))
+	check(!EnumHasAnyFlags(Source->GetDesc().Flags, TexCreate_CPUReadback));
+	if(EnumHasAllFlags(Dest->GetDesc().Flags, TexCreate_CPUReadback))
 	{
 		check(CopyInfo.DestSliceIndex == 0); //slices not supported in TexCreate_CPUReadback textures.
 		FIntVector Size = CopyInfo.Size;
 		if(Size == FIntVector::ZeroValue)
 		{
-			ensure(SrcSurface.Width <= DstSurface.Width && SrcSurface.Height <= DstSurface.Height);
-			Size.X = FMath::Max(1u, SrcSurface.Width >> CopyInfo.SourceMipIndex);
-			Size.Y = FMath::Max(1u, SrcSurface.Height >> CopyInfo.SourceMipIndex);
-			Size.Z = FMath::Max(1u, SrcSurface.Depth >> CopyInfo.SourceMipIndex);
+			ensure(SourceXYZ.X <= DestXYZ.X && SourceXYZ.Y <= DestXYZ.Y);
+			Size.X = FMath::Max<uint32>(1u, SourceXYZ.X >> CopyInfo.SourceMipIndex);
+			Size.Y = FMath::Max<uint32>(1u, SourceXYZ.Y >> CopyInfo.SourceMipIndex);
+			Size.Z = FMath::Max<uint32>(1u, SourceXYZ.Z >> CopyInfo.SourceMipIndex);
 		}
 		VkBufferImageCopy CopyRegion[MAX_TEXTURE_MIP_COUNT];
 		FMemory::Memzero(CopyRegion);
 
-		const FVulkanCpuReadbackBuffer* CpuReadbackBuffer = DstSurface.GetCpuReadbackBuffer();
-		uint32 SourceSliceIndex = CopyInfo.SourceSliceIndex;
-		uint32 SourceMipIndex = CopyInfo.SourceMipIndex;
-		uint32 DestMipIndex = CopyInfo.DestMipIndex;
+		const FVulkanCpuReadbackBuffer* CpuReadbackBuffer = Dest->GetCpuReadbackBuffer();
+		const uint32 SourceSliceIndex = CopyInfo.SourceSliceIndex;
+		const uint32 SourceMipIndex = CopyInfo.SourceMipIndex;
+		const uint32 DestMipIndex = CopyInfo.DestMipIndex;
 		for (uint32 Index = 0; Index < CopyInfo.NumMips; ++Index)
 		{
 			CopyRegion[Index].bufferOffset = CpuReadbackBuffer->MipOffsets[DestMipIndex + Index];
 			CopyRegion[Index].bufferRowLength = Size.X;
 			CopyRegion[Index].bufferImageHeight = Size.Y;
-			CopyRegion[Index].imageSubresource.aspectMask = SrcSurface.GetFullAspectMask();
+			CopyRegion[Index].imageSubresource.aspectMask = Source->GetFullAspectMask();
 			CopyRegion[Index].imageSubresource.mipLevel = SourceMipIndex;
 			CopyRegion[Index].imageSubresource.baseArrayLayer = SourceSliceIndex;
 			CopyRegion[Index].imageSubresource.layerCount = 1;
@@ -2638,7 +2379,7 @@ void FVulkanCommandListContext::RHICopyTexture(FRHITexture* SourceTexture, FRHIT
 			Size.Z = FMath::Max(1, Size.Z / 2);
 		}
 
-		VulkanRHI::vkCmdCopyImageToBuffer(CmdBuffer, SrcSurface.Image, VK_IMAGE_LAYOUT_TRANSFER_SRC_OPTIMAL, CpuReadbackBuffer->Buffer, CopyInfo.NumMips, &CopyRegion[0]);
+		VulkanRHI::vkCmdCopyImageToBuffer(CmdBuffer, Source->Image, VK_IMAGE_LAYOUT_TRANSFER_SRC_OPTIMAL, CpuReadbackBuffer->Buffer, CopyInfo.NumMips, &CopyRegion[0]);
 
 		FVulkanPipelineBarrier BarrierMemory;
 		BarrierMemory.AddMemoryBarrier(VK_ACCESS_TRANSFER_WRITE_BIT, VK_ACCESS_HOST_READ_BIT, VK_PIPELINE_STAGE_ALL_COMMANDS_BIT, VK_PIPELINE_STAGE_HOST_BIT);
@@ -2646,7 +2387,7 @@ void FVulkanCommandListContext::RHICopyTexture(FRHITexture* SourceTexture, FRHIT
 	}
 	else
 	{
-		VkImageLayout DstLayout = LayoutManager.FindLayoutChecked(DstSurface.Image);
+		VkImageLayout DstLayout = LayoutManager.FindLayoutChecked(Dest->Image);
 		ensureMsgf(DstLayout == VK_IMAGE_LAYOUT_TRANSFER_DST_OPTIMAL, TEXT("Expected destination texture to be in VK_IMAGE_LAYOUT_TRANSFER_DST_OPTIMAL, actual layout is %d"), DstLayout);
 
 		VkImageCopy Region;
@@ -2654,26 +2395,26 @@ void FVulkanCommandListContext::RHICopyTexture(FRHITexture* SourceTexture, FRHIT
 		if (CopyInfo.Size == FIntVector::ZeroValue)
 		{
 			// Copy whole texture when zero vector is specified for region size
-			ensure(SrcSurface.Width <= DstSurface.Width && SrcSurface.Height <= DstSurface.Height);
-			Region.extent.width = FMath::Max<uint32>(PixelFormatInfo.BlockSizeX, SrcSurface.Width >> CopyInfo.SourceMipIndex);
-			Region.extent.height = FMath::Max<uint32>(PixelFormatInfo.BlockSizeY, SrcSurface.Height >> CopyInfo.SourceMipIndex);
-			Region.extent.depth = FMath::Max<uint32>(PixelFormatInfo.BlockSizeZ, SrcSurface.Depth >> CopyInfo.SourceMipIndex);
+			ensure(SourceXYZ.X <= DestXYZ.X && SourceXYZ.Y <= DestXYZ.Y);
+			Region.extent.width = FMath::Max<uint32>(PixelFormatInfo.BlockSizeX, SourceXYZ.X >> CopyInfo.SourceMipIndex);
+			Region.extent.height = FMath::Max<uint32>(PixelFormatInfo.BlockSizeY, SourceXYZ.Y >> CopyInfo.SourceMipIndex);
+			Region.extent.depth = FMath::Max<uint32>(PixelFormatInfo.BlockSizeZ, SourceXYZ.Z >> CopyInfo.SourceMipIndex);
 		}
 		else
 		{
-			ensure(CopyInfo.Size.X > 0 && (uint32)CopyInfo.Size.X <= DstSurface.Width && CopyInfo.Size.Y > 0 && (uint32)CopyInfo.Size.Y <= DstSurface.Height);
+			ensure(CopyInfo.Size.X > 0 && CopyInfo.Size.X <= DestXYZ.X && CopyInfo.Size.Y > 0 && CopyInfo.Size.Y <= DestXYZ.Y);
 			Region.extent.width = CopyInfo.Size.X;
 			Region.extent.height = CopyInfo.Size.Y;
 			Region.extent.depth = FMath::Max(1, CopyInfo.Size.Z);
 		}
-		Region.srcSubresource.aspectMask = SrcSurface.GetFullAspectMask();
+		Region.srcSubresource.aspectMask = Source->GetFullAspectMask();
 		Region.srcSubresource.baseArrayLayer = CopyInfo.SourceSliceIndex;
 		Region.srcSubresource.layerCount = CopyInfo.NumSlices;
 		Region.srcSubresource.mipLevel = CopyInfo.SourceMipIndex;
 		Region.srcOffset.x = CopyInfo.SourcePosition.X;
 		Region.srcOffset.y = CopyInfo.SourcePosition.Y;
 		Region.srcOffset.z = CopyInfo.SourcePosition.Z;
-		Region.dstSubresource.aspectMask = DstSurface.GetFullAspectMask();
+		Region.dstSubresource.aspectMask = Dest->GetFullAspectMask();
 		Region.dstSubresource.baseArrayLayer = CopyInfo.DestSliceIndex;
 		Region.dstSubresource.layerCount = CopyInfo.NumSlices;
 		Region.dstSubresource.mipLevel = CopyInfo.DestMipIndex;
@@ -2684,8 +2425,8 @@ void FVulkanCommandListContext::RHICopyTexture(FRHITexture* SourceTexture, FRHIT
 		for (uint32 Index = 0; Index < CopyInfo.NumMips; ++Index)
 		{
 			VulkanRHI::vkCmdCopyImage(CmdBuffer,
-				SrcSurface.Image, VK_IMAGE_LAYOUT_TRANSFER_SRC_OPTIMAL,
-				DstSurface.Image, VK_IMAGE_LAYOUT_TRANSFER_DST_OPTIMAL,
+				Source->Image, VK_IMAGE_LAYOUT_TRANSFER_SRC_OPTIMAL,
+				Dest->Image, VK_IMAGE_LAYOUT_TRANSFER_DST_OPTIMAL,
 				1, &Region);
 
 			++Region.srcSubresource.mipLevel;
