Copying //UE5/Dev-ParallelRendering to //UE5/Main

diff --git a/Engine/Source/Runtime/VulkanRHI/Private/VulkanTransientResourceAllocator.cpp b/Engine/Source/Runtime/VulkanRHI/Private/VulkanTransientResourceAllocator.cpp
--- a/Engine/Source/Runtime/VulkanRHI/Private/VulkanTransientResourceAllocator.cpp
+++ b/Engine/Source/Runtime/VulkanRHI/Private/VulkanTransientResourceAllocator.cpp
@@ -109,10 +109,9 @@ FVulkanTransientResourceAllocator::FVulkanTransientResourceAllocator(FVulkanTran
 
 FRHITransientTexture* FVulkanTransientResourceAllocator::CreateTexture(const FRHITextureCreateInfo& InCreateInfo, const TCHAR* InDebugName, uint32 InPassIndex)
 {
-	uint32 ReqAlign = 1;
-	uint64 ReqSize = GVulkanRHI->RHICalcTexturePlatformSize(InCreateInfo, ReqAlign);
+	FDynamicRHI::FRHICalcTextureSizeResult MemReq = GVulkanRHI->RHICalcTexturePlatformSize(InCreateInfo, 0);
 
-	return CreateTextureInternal(InCreateInfo, InDebugName, InPassIndex, ReqSize, ReqAlign,
+	return CreateTextureInternal(InCreateInfo, InDebugName, InPassIndex, MemReq.Size, MemReq.Align,
 		[&](const FRHITransientHeap::FResourceInitializer& Initializer)
 	{
 		ERHIAccess InitialState = ERHIAccess::UAVMask;
@@ -125,9 +124,9 @@ FRHITransientTexture* FVulkanTransientResourceAllocator::CreateTexture(const FRH
 			InitialState = ERHIAccess::DSVWrite;
 		}
 
-		FRHIResourceCreateInfo ResourceCreateInfo(InDebugName, InCreateInfo.ClearValue);
-		FRHITexture* Texture = GVulkanRHI->CreateTexture(InCreateInfo, ResourceCreateInfo, InitialState, &Initializer.Allocation);
-		return new FRHITransientTexture(Texture, 0/*GpuVirtualAddress*/, Initializer.Hash, ReqSize, ERHITransientAllocationType::Heap, InCreateInfo);
+		FRHITextureCreateDesc CreateDesc(InCreateInfo, InitialState, InDebugName);
+		FRHITexture* Texture = new FVulkanTexture(*Device, CreateDesc, &Initializer.Allocation);
+		return new FRHITransientTexture(Texture, 0/*GpuVirtualAddress*/, Initializer.Hash, MemReq.Size, ERHITransientAllocationType::Heap, InCreateInfo);
 	});
 }
 
diff --git a/Engine/Source/Runtime/VulkanRHI/Private/VulkanUAV.cpp b/Engine/Source/Runtime/VulkanRHI/Private/VulkanUAV.cpp
--- a/Engine/Source/Runtime/VulkanRHI/Private/VulkanUAV.cpp
+++ b/Engine/Source/Runtime/VulkanRHI/Private/VulkanUAV.cpp
@@ -41,7 +41,7 @@ FVulkanShaderResourceView::FVulkanShaderResourceView(FVulkanDevice* Device, FRHI
 	, Size(0)
 	, SourceBuffer(nullptr)
 {
-	FVulkanTextureBase* VulkanTexture = FVulkanTextureBase::Cast(InSourceTexture);
+	FVulkanTexture* VulkanTexture = FVulkanTexture::Cast(InSourceTexture);
 	VulkanTexture->AttachView(this);
 
 }
@@ -80,7 +80,7 @@ FVulkanShaderResourceView::~FVulkanShaderResourceView()
 	FRHITexture* Texture = SourceTexture.GetReference();
 	if(Texture)
 	{
-		FVulkanTextureBase* VulkanTexture = FVulkanTextureBase::Cast(Texture);
+		FVulkanTexture* VulkanTexture = FVulkanTexture::Cast(Texture);
 		VulkanTexture->DetachView(this);
 	}
 	Clear();
@@ -199,42 +199,34 @@ void FVulkanShaderResourceView::UpdateView()
 			const bool bSRGB = (SRGBOverride != SRGBO_ForceDisable) && bBaseSRGB;
 
 			EPixelFormat Format = (BufferViewFormat == PF_Unknown) ? SourceTexture->GetFormat() : BufferViewFormat;
-			if (FRHITexture2D* Tex2D = SourceTexture->GetTexture2D())
-			{
-				FVulkanTexture2D* VTex2D = ResourceCast(Tex2D);
-				EPixelFormat OriginalFormat = Format;
-				TextureView.Create(*Device, VTex2D->Surface.Image, VK_IMAGE_VIEW_TYPE_2D, VTex2D->Surface.GetPartialAspectMask(), Format, UEToVkTextureFormat(Format, bSRGB), MipLevel, NumMips, 0, 1);
-			}
-			else if (FRHITextureCube* TexCube = SourceTexture->GetTextureCube())
-			{
-				FVulkanTextureCube* VTexCube = ResourceCast(TexCube);
-				TextureView.Create(*Device, VTexCube->Surface.Image, VK_IMAGE_VIEW_TYPE_CUBE, VTexCube->Surface.GetPartialAspectMask(), Format, UEToVkTextureFormat(Format, bSRGB), MipLevel, NumMips, 0, 1);
-			}
-			else if (FRHITexture3D* Tex3D = SourceTexture->GetTexture3D())
-			{
-				FVulkanTexture3D* VTex3D = ResourceCast(Tex3D);
-				TextureView.Create(*Device, VTex3D->Surface.Image, VK_IMAGE_VIEW_TYPE_3D, VTex3D->Surface.GetPartialAspectMask(), Format, UEToVkTextureFormat(Format, bSRGB), MipLevel, NumMips, 0, 1);
-			}
-			else if (FRHITexture2DArray* Tex2DArray = SourceTexture->GetTexture2DArray())
-			{
-				FVulkanTexture2DArray* VTex2DArray = ResourceCast(Tex2DArray);
-				TextureView.Create(
-					*Device,
-					VTex2DArray->Surface.Image,
-					VK_IMAGE_VIEW_TYPE_2D_ARRAY,
-					VTex2DArray->Surface.GetPartialAspectMask(),
-					Format,
-					UEToVkTextureFormat(Format, bSRGB),
-					MipLevel,
-					NumMips,
-					FirstArraySlice,
-					(NumArraySlices == 0 ? VTex2DArray->GetSizeZ() : NumArraySlices)
-				);
-			}
-			else
+			FVulkanTexture* SourceTextureVK = FVulkanTexture::Cast(SourceTexture);
+
+			// A NumArraySlices of 0 means view all the layers, so determine the correct layer count.
+			uint32 ActualNumArraySlices = NumArraySlices;
+			if (ActualNumArraySlices == 0)
 			{
-				ensure(0);
+				switch (SourceTexture->GetDesc().Dimension)
+				{
+					case ETextureDimension::Texture2DArray:
+					case ETextureDimension::TextureCubeArray:
+						ActualNumArraySlices = SourceTexture->GetDesc().ArraySize;
+						break;
+					default:
+						ActualNumArraySlices = 1;
+				}
+
+				// If we're using the total count, subtract the first layer index, to view all the layers after that.
+				if (ensure(FirstArraySlice < ActualNumArraySlices))
+				{
+					ActualNumArraySlices -= FirstArraySlice;
+				}
+				else
+				{
+					ActualNumArraySlices = 1;
+				}
 			}
+
+			TextureView.Create(*Device, SourceTextureVK->Image, SourceTextureVK->GetViewType(), SourceTextureVK->GetPartialAspectMask(), Format, UEToVkTextureFormat(Format, bSRGB), MipLevel, NumMips, FirstArraySlice, ActualNumArraySlices, false);
 		}
 	}
 }
@@ -258,7 +250,7 @@ FVulkanUnorderedAccessView::FVulkanUnorderedAccessView(FVulkanDevice* Device, FR
 	, BufferViewFormat(PF_Unknown)
 	, VolatileLockCounter(MAX_uint32)
 {
-	FVulkanTextureBase* VulkanTexture = FVulkanTextureBase::Cast(TextureRHI);
+	FVulkanTexture* VulkanTexture = FVulkanTexture::Cast(TextureRHI);
 	VulkanTexture->AttachView(this);
 }
 
@@ -284,7 +276,7 @@ FVulkanUnorderedAccessView::~FVulkanUnorderedAccessView()
 {
 	if (SourceTexture)
 	{
-		FVulkanTextureBase* VulkanTexture = FVulkanTextureBase::Cast(SourceTexture);
+		FVulkanTexture* VulkanTexture = FVulkanTexture::Cast(SourceTexture);
 		VulkanTexture->DetachView(this);
 	}
 
@@ -323,32 +315,52 @@ void FVulkanUnorderedAccessView::UpdateView()
 	else if (TextureView.View == VK_NULL_HANDLE)
 	{
 		EPixelFormat Format = (BufferViewFormat == PF_Unknown) ? SourceTexture->GetFormat() : BufferViewFormat;
-		if (FRHITexture2D* Tex2D = SourceTexture->GetTexture2D())
-		{
-			FVulkanTexture2D* VTex2D = ResourceCast(Tex2D);
-			TextureView.Create(*Device, VTex2D->Surface.Image, VK_IMAGE_VIEW_TYPE_2D, VTex2D->Surface.GetPartialAspectMask(), Format, UEToVkTextureFormat(Format, false), MipLevel, 1, 0, 1, true);
-		}
-		else if (FRHITextureCube* TexCube = SourceTexture->GetTextureCube())
-		{
-			FVulkanTextureCube* VTexCube = ResourceCast(TexCube);
-			// RWTextureCube is defined as RWTexture2DArray in shader source, avoid validation errors by creating the appropriate VK_IMAGE_VIEW_TYPE_2D_ARRAY view (instead of VK_IMAGE_VIEW_TYPE_CUBE)
-			TextureView.Create(*Device, VTexCube->Surface.Image, VK_IMAGE_VIEW_TYPE_2D_ARRAY, VTexCube->Surface.GetPartialAspectMask(), Format, UEToVkTextureFormat(Format, false), MipLevel, 1, 0, VTexCube->Surface.GetNumberOfArrayLevels(), true);
-		}
-		else if (FRHITexture3D* Tex3D = SourceTexture->GetTexture3D())
+		FVulkanTexture* SourceTextureVK = FVulkanTexture::Cast(SourceTexture);
+
+		// A NumArraySlices of 0 means view all the layers, so determine the correct layer count.
+		uint32 ActualNumArraySlices = NumArraySlices;
+		if (NumArraySlices == 0)
 		{
-			FVulkanTexture3D* VTex3D = ResourceCast(Tex3D);
-			TextureView.Create(*Device, VTex3D->Surface.Image, VK_IMAGE_VIEW_TYPE_3D, VTex3D->Surface.GetPartialAspectMask(), Format, UEToVkTextureFormat(Format, false), MipLevel, 1, 0, VTex3D->GetSizeZ(), true);
+			switch (SourceTexture->GetDesc().Dimension)
+			{
+			case ETextureDimension::Texture2DArray:
+			case ETextureDimension::TextureCubeArray:
+				ActualNumArraySlices = SourceTexture->GetDesc().ArraySize;
+				break;
+			default:
+				ActualNumArraySlices = 1;
+			}
 		}
-		else if (FRHITexture2DArray* Tex2DArray = SourceTexture->GetTexture2DArray())
+
+		// RWTextureCube is defined as RWTexture2DArray in shader source, so adjust the view type and layer count.
+		VkImageViewType ActualViewType;
+		switch (SourceTexture->GetDesc().Dimension)
 		{
-			FVulkanTexture2DArray* VTex2DArray = ResourceCast(Tex2DArray);
-			TextureView.Create(*Device, VTex2DArray->Surface.Image, VK_IMAGE_VIEW_TYPE_2D_ARRAY, VTex2DArray->Surface.GetPartialAspectMask(), Format, UEToVkTextureFormat(Format, false), MipLevel, 1, 
-				NumArraySlices == 0 ? 0 : FirstArraySlice, NumArraySlices == 0 ? VTex2DArray->GetSizeZ() : NumArraySlices, true);
+		case ETextureDimension::TextureCube:
+		case ETextureDimension::TextureCubeArray:
+			FirstArraySlice *= 6;
+			ActualNumArraySlices *= 6;
+			ActualViewType = VK_IMAGE_VIEW_TYPE_2D_ARRAY;
+			break;
+
+		default:
+			ActualViewType = SourceTextureVK->GetViewType();
 		}
-		else
+
+		if (NumArraySlices == 0)
 		{
-			ensure(0);
+			// If we're using the total count, subtract the first layer index, to view all the layers after that.
+			if (ensure(FirstArraySlice < ActualNumArraySlices))
+			{
+				ActualNumArraySlices -= FirstArraySlice;
+			}
+			else
+			{
+				ActualNumArraySlices = 1;
+			}
 		}
+
+		TextureView.Create(*Device, SourceTextureVK->Image, ActualViewType, SourceTextureVK->GetPartialAspectMask(), Format, UEToVkTextureFormat(Format, false), MipLevel, 1, FirstArraySlice, ActualNumArraySlices, true);
 	}
 }
 
diff --git a/Engine/Source/Runtime/VulkanRHI/Private/VulkanViewport.cpp b/Engine/Source/Runtime/VulkanRHI/Private/VulkanViewport.cpp
--- a/Engine/Source/Runtime/VulkanRHI/Private/VulkanViewport.cpp
+++ b/Engine/Source/Runtime/VulkanRHI/Private/VulkanViewport.cpp
@@ -15,7 +15,7 @@
 
 
 FVulkanBackBuffer::FVulkanBackBuffer(FVulkanDevice& Device, FVulkanViewport* InViewport, EPixelFormat Format, uint32 SizeX, uint32 SizeY, ETextureCreateFlags UEFlags)
-	: FVulkanTexture2D(Device, Format, SizeX, SizeY, 1, 1, VK_NULL_HANDLE, UEFlags, FRHIResourceCreateInfo(TEXT("FVulkanBackBuffer")))
+	: FVulkanTexture(Device, FRHITextureCreateDesc::Create2D(TEXT("FVulkanBackBuffer"), { (int32)SizeX, (int32)SizeY }, Format, FClearValueBinding::None, UEFlags), VK_NULL_HANDLE, false)
 	, Viewport(InViewport)
 {
 }
@@ -24,7 +24,7 @@ void FVulkanBackBuffer::ReleaseAcquiredImage()
 {
 	DefaultView.View = VK_NULL_HANDLE;
 	DefaultView.ViewId = 0;
-	Surface.Image = VK_NULL_HANDLE;
+	Image = VK_NULL_HANDLE;
 }
 
 void FVulkanBackBuffer::ReleaseViewport()
@@ -60,7 +60,7 @@ void FVulkanBackBuffer::OnLayoutTransition(FVulkanCommandListContext& Context, V
 void FVulkanBackBuffer::AcquireBackBufferImage(FVulkanCommandListContext& Context)
 {
 	check(Viewport);
-	if (Surface.Image == VK_NULL_HANDLE)
+	if (Image == VK_NULL_HANDLE)
 	{
 		check(Viewport->AcquiredImageIndex == -1); //-V595
 		
@@ -76,7 +76,7 @@ void FVulkanBackBuffer::AcquireBackBufferImage(FVulkanCommandListContext& Contex
 
 		FVulkanTextureView& ImageView = Viewport->TextureViews[AcquiredImageIndex];
 
-		Surface.Image = ImageView.Image;
+		Image = ImageView.Image;
 		DefaultView.View = ImageView.View;
 		DefaultView.ViewId = ImageView.ViewId;
 
@@ -96,9 +96,9 @@ void FVulkanBackBuffer::AcquireBackBufferImage(FVulkanCommandListContext& Contex
 
 FVulkanBackBuffer::~FVulkanBackBuffer()
 {
-	check(Surface.IsImageOwner() == false);
-	// Clear flags so ~FVulkanTexture2D() doesn't try to re-destroy it
-	Surface.UEFlags = TexCreate_None;
+	check(IsImageOwner() == false);
+	// Clear ImageOwnerType so ~FVulkanTexture2D() doesn't try to re-destroy it
+	ImageOwnerType = EImageOwnerType::None;
 	ReleaseAcquiredImage();
 }
 
@@ -326,41 +326,43 @@ FVulkanFramebuffer::FVulkanFramebuffer(FVulkanDevice& Device, const FRHISetRende
 			continue;
 		}
 
-		FVulkanTextureBase* Texture = FVulkanTextureBase::Cast(RHITexture);
+		FVulkanTexture* Texture = FVulkanTexture::Cast(RHITexture);
+		const FRHITextureDesc& Desc = Texture->GetDesc();
+
 		// this could fire in case one of the textures is FVulkanBackBuffer and it has not acquired an image
 		// with EDelayAcquireImageType::LazyAcquire acquire happens when texture transition to Writeable state
 		// make sure you call TransitionResource(Writable, Tex) before using this texture as a render-target
-		check(Texture->Surface.Image != VK_NULL_HANDLE);
+		check(Texture->Image != VK_NULL_HANDLE);
 
-		ColorRenderTargetImages[Index] = Texture->Surface.Image;
+		ColorRenderTargetImages[Index] = Texture->Image;
 		MipIndex = InRTInfo.ColorRenderTarget[Index].MipIndex;
 
 		FVulkanTextureView RTView;
-		if (Texture->Surface.GetViewType() == VK_IMAGE_VIEW_TYPE_2D || Texture->Surface.GetViewType() == VK_IMAGE_VIEW_TYPE_2D_ARRAY)
+		if (Texture->GetViewType() == VK_IMAGE_VIEW_TYPE_2D || Texture->GetViewType() == VK_IMAGE_VIEW_TYPE_2D_ARRAY)
 		{
 			uint32 ArraySliceIndex, NumArraySlices;
 			if (InRTInfo.ColorRenderTarget[Index].ArraySliceIndex == -1)
 			{
 				ArraySliceIndex = 0;
-				NumArraySlices = Texture->Surface.GetNumberOfArrayLevels();
+				NumArraySlices = Texture->GetNumberOfArrayLevels();
 			}
 			else
 			{
 				ArraySliceIndex = InRTInfo.ColorRenderTarget[Index].ArraySliceIndex;
 				NumArraySlices = 1;
-				check(ArraySliceIndex < Texture->Surface.GetNumberOfArrayLevels());
+				check(ArraySliceIndex < Texture->GetNumberOfArrayLevels());
 			}
-			RTView.Create(*Texture->Surface.Device, Texture->Surface.Image, Texture->Surface.GetViewType(), Texture->Surface.GetFullAspectMask(), Texture->Surface.PixelFormat, Texture->Surface.ViewFormat, MipIndex, 1, ArraySliceIndex, NumArraySlices, true);
+			RTView.Create(*Texture->Device, Texture->Image, Texture->GetViewType(), Texture->GetFullAspectMask(), Desc.Format, Texture->ViewFormat, MipIndex, 1, ArraySliceIndex, NumArraySlices, true);
 		}
-		else if (Texture->Surface.GetViewType() == VK_IMAGE_VIEW_TYPE_CUBE)
+		else if (Texture->GetViewType() == VK_IMAGE_VIEW_TYPE_CUBE)
 		{
 			// Cube always renders one face at a time
 			INC_DWORD_STAT(STAT_VulkanNumImageViews);
-			RTView.Create(*Texture->Surface.Device, Texture->Surface.Image, VK_IMAGE_VIEW_TYPE_2D, Texture->Surface.GetFullAspectMask(), Texture->Surface.PixelFormat, Texture->Surface.ViewFormat, MipIndex, 1, InRTInfo.ColorRenderTarget[Index].ArraySliceIndex, 1, true);
+			RTView.Create(*Texture->Device, Texture->Image, VK_IMAGE_VIEW_TYPE_2D, Texture->GetFullAspectMask(), Desc.Format, Texture->ViewFormat, MipIndex, 1, InRTInfo.ColorRenderTarget[Index].ArraySliceIndex, 1, true);
 		}
-		else if (Texture->Surface.GetViewType() == VK_IMAGE_VIEW_TYPE_3D)
+		else if (Texture->GetViewType() == VK_IMAGE_VIEW_TYPE_3D)
 		{
-			RTView.Create(*Texture->Surface.Device, Texture->Surface.Image, VK_IMAGE_VIEW_TYPE_2D_ARRAY, Texture->Surface.GetFullAspectMask(), Texture->Surface.PixelFormat, Texture->Surface.ViewFormat, MipIndex, 1, 0, Texture->Surface.Depth, true);
+			RTView.Create(*Texture->Device, Texture->Image, VK_IMAGE_VIEW_TYPE_2D_ARRAY, Texture->GetFullAspectMask(), Desc.Format, Texture->ViewFormat, MipIndex, 1, 0, Desc.Depth, true);
 		}
 		else
 		{
@@ -377,15 +379,15 @@ FVulkanFramebuffer::FVulkanFramebuffer(FVulkanDevice& Device, const FRHISetRende
 		if (InRTInfo.bHasResolveAttachments && RTLayout.GetHasResolveAttachments() && RTLayout.GetResolveAttachmentReferences()[Index].layout != VK_IMAGE_LAYOUT_UNDEFINED)
 		{
 			FRHITexture* ResolveRHITexture = InRTInfo.ColorResolveRenderTarget[Index].Texture;
-			FVulkanTextureBase* ResolveTexture = FVulkanTextureBase::Cast(ResolveRHITexture);
-			ColorResolveTargetImages[Index] = ResolveTexture->Surface.Image;
+			FVulkanTexture* ResolveTexture = FVulkanTexture::Cast(ResolveRHITexture);
+			ColorResolveTargetImages[Index] = ResolveTexture->Image;
 
 			//resolve attachments only supported for 2d/2d array textures
 			FVulkanTextureView ResolveRTView;
-			if (ResolveTexture->Surface.GetViewType() == VK_IMAGE_VIEW_TYPE_2D || ResolveTexture->Surface.GetViewType() == VK_IMAGE_VIEW_TYPE_2D_ARRAY)
+			if (ResolveTexture->GetViewType() == VK_IMAGE_VIEW_TYPE_2D || ResolveTexture->GetViewType() == VK_IMAGE_VIEW_TYPE_2D_ARRAY)
 			{
-				ResolveRTView.Create(*ResolveTexture->Surface.Device, ResolveTexture->Surface.Image, ResolveTexture->Surface.GetViewType(), ResolveTexture->Surface.GetFullAspectMask(), ResolveTexture->Surface.PixelFormat, ResolveTexture->Surface.ViewFormat, 
-					MipIndex, 1, FMath::Max(0, (int32)InRTInfo.ColorRenderTarget[Index].ArraySliceIndex), ResolveTexture->Surface.GetNumberOfArrayLevels(), true);
+				ResolveRTView.Create(*ResolveTexture->Device, ResolveTexture->Image, ResolveTexture->GetViewType(), ResolveTexture->GetFullAspectMask(), ResolveTexture->GetDesc().Format, ResolveTexture->ViewFormat,
+					MipIndex, 1, FMath::Max(0, (int32)InRTInfo.ColorRenderTarget[Index].ArraySliceIndex), ResolveTexture->GetNumberOfArrayLevels(), true);
 			}
 
 			AttachmentTextureViews.Add(ResolveRTView);
@@ -397,34 +399,35 @@ FVulkanFramebuffer::FVulkanFramebuffer(FVulkanDevice& Device, const FRHISetRende
 
 	if (RTLayout.GetHasDepthStencil())
 	{
-		FVulkanTextureBase* Texture = FVulkanTextureBase::Cast(InRTInfo.DepthStencilRenderTarget.Texture);
-		DepthStencilRenderTargetImage = Texture->Surface.Image;
-		bool bHasStencil = (Texture->Surface.PixelFormat == PF_DepthStencil || Texture->Surface.PixelFormat == PF_X24_G8);
+		FVulkanTexture* Texture = FVulkanTexture::Cast(InRTInfo.DepthStencilRenderTarget.Texture);
+		const FRHITextureDesc& Desc = Texture->GetDesc();
+		DepthStencilRenderTargetImage = Texture->Image;
+		bool bHasStencil = (Texture->GetDesc().Format == PF_DepthStencil || Texture->GetDesc().Format == PF_X24_G8);
 		check(Texture->PartialView);
 		PartialDepthTextureView = *Texture->PartialView;
 
 		FVulkanTextureView RTView;
-		ensure(Texture->Surface.GetViewType() == VK_IMAGE_VIEW_TYPE_2D || Texture->Surface.GetViewType() == VK_IMAGE_VIEW_TYPE_2D_ARRAY || Texture->Surface.GetViewType() == VK_IMAGE_VIEW_TYPE_CUBE);
-		if (NumColorAttachments == 0 && Texture->Surface.GetViewType() == VK_IMAGE_VIEW_TYPE_CUBE)
+		ensure(Texture->GetViewType() == VK_IMAGE_VIEW_TYPE_2D || Texture->GetViewType() == VK_IMAGE_VIEW_TYPE_2D_ARRAY || Texture->GetViewType() == VK_IMAGE_VIEW_TYPE_CUBE);
+		if (NumColorAttachments == 0 && Texture->GetViewType() == VK_IMAGE_VIEW_TYPE_CUBE)
 		{
-			RTView.Create(*Texture->Surface.Device, Texture->Surface.Image, VK_IMAGE_VIEW_TYPE_2D_ARRAY, Texture->Surface.GetFullAspectMask(), Texture->Surface.PixelFormat, Texture->Surface.ViewFormat, MipIndex, 1, 0, 6, true);
+			RTView.Create(*Texture->Device, Texture->Image, VK_IMAGE_VIEW_TYPE_2D_ARRAY, Texture->GetFullAspectMask(), Texture->GetDesc().Format, Texture->ViewFormat, MipIndex, 1, 0, 6, true);
 			NumLayers = 6;
 			AttachmentTextureViews.Add(RTView);
 			AttachmentViewsToDelete.Add(RTView.View);
 		}
-		else if (Texture->Surface.GetViewType() == VK_IMAGE_VIEW_TYPE_2D  || Texture->Surface.GetViewType() == VK_IMAGE_VIEW_TYPE_2D_ARRAY)
+		else if (Texture->GetViewType() == VK_IMAGE_VIEW_TYPE_2D  || Texture->GetViewType() == VK_IMAGE_VIEW_TYPE_2D_ARRAY)
 		{
 			// depth attachments need a separate view to have no swizzle components, for validation correctness
-			RTView.Create(*Texture->Surface.Device, Texture->Surface.Image, Texture->Surface.GetViewType(), Texture->Surface.GetFullAspectMask(), Texture->Surface.PixelFormat, Texture->Surface.ViewFormat, MipIndex, 1, 0, Texture->Surface.GetNumberOfArrayLevels(), true);
+			RTView.Create(*Texture->Device, Texture->Image, Texture->GetViewType(), Texture->GetFullAspectMask(), Texture->GetDesc().Format, Texture->ViewFormat, MipIndex, 1, 0, Texture->GetNumberOfArrayLevels(), true);
 			AttachmentTextureViews.Add(RTView);
 			AttachmentViewsToDelete.Add(RTView.View);
 		}
 		else if (QCOMRenderPassTransform != VK_SURFACE_TRANSFORM_IDENTITY_BIT_KHR &&
-			Texture->Surface.Width == RTExtents.width && Texture->Surface.Height == RTExtents.height)
+			Desc.Extent.X == RTExtents.width && Desc.Extent.Y == RTExtents.height)
 		{
 			FVulkanSwapChain* SwapChain = Device.GetImmediateContext().GetSwapChain();
-			PartialDepthTextureView = *SwapChain->GetOrCreateQCOMDepthView(Texture->Surface);
-			AttachmentTextureViews.Add(*SwapChain->GetOrCreateQCOMDepthStencilView(Texture->Surface));
+			PartialDepthTextureView = *SwapChain->GetOrCreateQCOMDepthView(*Texture);
+			AttachmentTextureViews.Add(*SwapChain->GetOrCreateQCOMDepthStencilView(*Texture));
 		}
 		else
 		{
@@ -434,13 +437,13 @@ FVulkanFramebuffer::FVulkanFramebuffer(FVulkanDevice& Device, const FRHISetRende
 
 	if (GRHISupportsAttachmentVariableRateShading && GRHIVariableRateShadingEnabled && GRHIAttachmentVariableRateShadingEnabled && RTLayout.GetHasFragmentDensityAttachment())
 	{
-		FVulkanTextureBase* Texture = FVulkanTextureBase::Cast(InRTInfo.ShadingRateTexture);
-		FragmentDensityImage = Texture->Surface.Image;
+		FVulkanTexture* Texture = FVulkanTexture::Cast(InRTInfo.ShadingRateTexture);
+		FragmentDensityImage = Texture->Image;
 
-		ensure(Texture->Surface.GetViewType() == VK_IMAGE_VIEW_TYPE_2D || Texture->Surface.GetViewType() == VK_IMAGE_VIEW_TYPE_2D_ARRAY);
+		ensure(Texture->GetViewType() == VK_IMAGE_VIEW_TYPE_2D || Texture->GetViewType() == VK_IMAGE_VIEW_TYPE_2D_ARRAY);
 
 		FVulkanTextureView RTView;
-		RTView.Create(*Texture->Surface.Device, Texture->Surface.Image, Texture->Surface.GetViewType(), Texture->Surface.GetFullAspectMask(), Texture->Surface.PixelFormat, Texture->Surface.ViewFormat, MipIndex, 1, 0, Texture->Surface.GetNumberOfArrayLevels(), true);
+		RTView.Create(*Texture->Device, Texture->Image, Texture->GetViewType(), Texture->GetFullAspectMask(), Texture->GetDesc().Format, Texture->ViewFormat, MipIndex, 1, 0, Texture->GetNumberOfArrayLevels(), true);
 
 		AttachmentTextureViews.Add(RTView);
 		AttachmentViewsToDelete.Add(RTView.View);
@@ -512,7 +515,7 @@ bool FVulkanFramebuffer::Matches(const FRHISetRenderTargetsInfo& InRTInfo) const
 		if (B.Texture)
 		{
 			VkImage AImage = DepthStencilRenderTargetImage;
-			VkImage BImage = ((FVulkanTextureBase*)B.Texture->GetTextureBaseRHI())->Surface.Image;
+			VkImage BImage = FVulkanTexture::Cast(B.Texture)->Image;
 			if (AImage != BImage)
 			{
 				return false;
@@ -525,7 +528,7 @@ bool FVulkanFramebuffer::Matches(const FRHISetRenderTargetsInfo& InRTInfo) const
 		if (Texture)
 		{
 			VkImage AImage = FragmentDensityImage;
-			VkImage BImage = ((FVulkanTextureBase*)Texture->GetTextureBaseRHI())->Surface.Image;
+			VkImage BImage = FVulkanTexture::Cast(Texture)->Image;
 			if (AImage != BImage)
 			{
 				return false;
@@ -542,7 +545,7 @@ bool FVulkanFramebuffer::Matches(const FRHISetRenderTargetsInfo& InRTInfo) const
 			if (R.Texture)
 			{
 				VkImage AImage = ColorResolveTargetImages[AttachementIndex];
-				VkImage BImage = ((FVulkanTextureBase*)R.Texture->GetTextureBaseRHI())->Surface.Image;
+				VkImage BImage = FVulkanTexture::Cast(R.Texture)->Image;
 				if (AImage != BImage)
 				{
 					return false;
@@ -554,7 +557,7 @@ bool FVulkanFramebuffer::Matches(const FRHISetRenderTargetsInfo& InRTInfo) const
 		if (B.Texture)
 		{
 			VkImage AImage = ColorRenderTargetImages[AttachementIndex];
-			VkImage BImage = ((FVulkanTextureBase*)B.Texture->GetTextureBaseRHI())->Surface.Image;
+			VkImage BImage = FVulkanTexture::Cast(B.Texture)->Image;
 			if (AImage != BImage)
 			{
 				return false;
@@ -645,7 +648,7 @@ void FVulkanViewport::CreateSwapchain(FVulkanSwapChainRecreateInfo* RecreateInfo
 		for (int32 Index = 0; Index < Images.Num(); ++Index)
 		{
 			BackBufferImages[Index] = Images[Index];
-			TextureViews[Index].Create(*Device, Images[Index], VK_IMAGE_VIEW_TYPE_2D, VK_IMAGE_ASPECT_COLOR_BIT, PixelFormat, UEToVkTextureFormat(PixelFormat, false), 0, 1, 0, 1);
+			TextureViews[Index].Create(*Device, Images[Index], VK_IMAGE_VIEW_TYPE_2D, VK_IMAGE_ASPECT_COLOR_BIT, PixelFormat, UEToVkTextureFormat(PixelFormat, false), 0, 1, 0, 1, false);
 
 			// Clear the swapchain to avoid a validation warning, and transition to ColorAttachment
 			{
@@ -693,11 +696,23 @@ void FVulkanViewport::CreateSwapchain(FVulkanSwapChainRecreateInfo* RecreateInfo
 		uint32 BackBufferSizeX = RequiresRenderingBackBuffer() ? SizeX : 1;
 		uint32 BackBufferSizeY = RequiresRenderingBackBuffer() ? SizeY : 1;
 
-		RenderingBackBuffer = new FVulkanTexture2D(*Device, PixelFormat, BackBufferSizeX, BackBufferSizeY, 1, 1, TexCreate_RenderTargetable | TexCreate_ShaderResource, ERHIAccess::Present, FRHIResourceCreateInfo(TEXT("RenderingBackBuffer")));
+		FRHITextureCreateDesc Desc = FRHITextureCreateDesc::Create2D(
+			TEXT("RenderingBackBuffer"),
+			{ (int32)BackBufferSizeX, (int32)BackBufferSizeY },
+			PixelFormat,
+			FClearValueBinding::None,
+			TexCreate_RenderTargetable | TexCreate_ShaderResource,
+			1,
+			1,
+			0,
+			ERHIAccess::Present
+		);
+
+		RenderingBackBuffer = new FVulkanTexture(*Device, Desc, nullptr);
 #if VULKAN_ENABLE_DRAW_MARKERS
 		if (Device->GetDebugMarkerSetObjectName())
 		{
-			VulkanRHI::SetDebugMarkerName(Device->GetDebugMarkerSetObjectName(), Device->GetInstanceHandle(), RenderingBackBuffer->Surface.Image, "RenderingBackBuffer");
+			VulkanRHI::SetDebugMarkerName(Device->GetDebugMarkerSetObjectName(), Device->GetInstanceHandle(), RenderingBackBuffer->Image, "RenderingBackBuffer");
 		}
 #endif
 	}
@@ -829,7 +844,7 @@ bool FVulkanViewport::Present(FVulkanCommandListContext* Context, FVulkanCmdBuff
 				uint32 WindowSizeY = FMath::Min(SizeY, SwapChain->InternalHeight);
 
 				Context->RHIPushEvent(TEXT("CopyImageToBackBuffer"), FColor::Blue);
-				CopyImageToBackBuffer(Context, CmdBuffer, RenderingBackBuffer->Surface.Image, BackBufferImages[AcquiredImageIndex], SizeX, SizeY, WindowSizeX, WindowSizeY);
+				CopyImageToBackBuffer(Context, CmdBuffer, RenderingBackBuffer->Image, BackBufferImages[AcquiredImageIndex], SizeX, SizeY, WindowSizeX, WindowSizeY);
 				Context->RHIPopEvent();
 			}
 			else
@@ -840,7 +855,7 @@ bool FVulkanViewport::Present(FVulkanCommandListContext* Context, FVulkanCmdBuff
 		else
 		{
 			check(AcquiredImageIndex != -1);
-			check(RHIBackBuffer != nullptr && RHIBackBuffer->Surface.Image == BackBufferImages[AcquiredImageIndex]);
+			check(RHIBackBuffer != nullptr && RHIBackBuffer->Image == BackBufferImages[AcquiredImageIndex]);
 
 			VkImageLayout& Layout = Context->GetLayoutManager().FindOrAddLayoutRW(BackBufferImages[AcquiredImageIndex], VK_IMAGE_LAYOUT_UNDEFINED, 1, 1);
 			VulkanSetImageLayout(CmdBuffer->GetHandle(), BackBufferImages[AcquiredImageIndex], Layout, VK_IMAGE_LAYOUT_PRESENT_SRC_KHR, FVulkanPipelineBarrier::MakeSubresourceRange(VK_IMAGE_ASPECT_COLOR_BIT));
diff --git a/Engine/Source/Runtime/VulkanRHI/Private/VulkanViewport.h b/Engine/Source/Runtime/VulkanRHI/Private/VulkanViewport.h
--- a/Engine/Source/Runtime/VulkanRHI/Private/VulkanViewport.h
+++ b/Engine/Source/Runtime/VulkanRHI/Private/VulkanViewport.h
@@ -19,13 +19,13 @@ namespace VulkanRHI
 	class FSemaphore;
 }
 
-class FVulkanBackBuffer : public FVulkanTexture2D
+class FVulkanBackBuffer : public FVulkanTexture
 {
 public:
 	FVulkanBackBuffer(FVulkanDevice& Device, FVulkanViewport* InViewport, EPixelFormat Format, uint32 SizeX, uint32 SizeY, ETextureCreateFlags UEFlags);
 	virtual ~FVulkanBackBuffer();
 	
-	virtual void OnLayoutTransition(FVulkanCommandListContext& Context, VkImageLayout NewLayout) override final;
+	void OnLayoutTransition(FVulkanCommandListContext& Context, VkImageLayout NewLayout) override final;
 
 	void OnGetBackBufferImage(FRHICommandListImmediate& RHICmdList);
 	void OnAdvanceBackBufferFrame(FRHICommandListImmediate& RHICmdList);
@@ -118,7 +118,7 @@ protected:
 	TRefCountPtr<FVulkanBackBuffer> RHIBackBuffer;
 
 	// 'Dummy' back buffer
-	TRefCountPtr<FVulkanTexture2D>	RenderingBackBuffer;
+	TRefCountPtr<FVulkanTexture>	RenderingBackBuffer;
 	
 	/** narrow-scoped section that locks access to back buffer during its recreation*/
 	FCriticalSection RecreatingSwapchain;
diff --git a/Engine/Source/Runtime/VulkanRHI/Public/IVulkanDynamicRHI.h b/Engine/Source/Runtime/VulkanRHI/Public/IVulkanDynamicRHI.h
--- a/Engine/Source/Runtime/VulkanRHI/Public/IVulkanDynamicRHI.h
+++ b/Engine/Source/Runtime/VulkanRHI/Public/IVulkanDynamicRHI.h
@@ -64,12 +64,12 @@ struct IVulkanDynamicRHI : public FDynamicRHI
 	virtual FTexture2DArrayRHIRef RHICreateTexture2DArrayFromResource(EPixelFormat Format, uint32 SizeX, uint32 SizeY, uint32 ArraySize, uint32 NumMips, uint32 NumSamples, VkImage Resource, ETextureCreateFlags Flags) = 0;
 	virtual FTextureCubeRHIRef    RHICreateTextureCubeFromResource(EPixelFormat Format, uint32 Size, bool bArray, uint32 ArraySize, uint32 NumMips, VkImage Resource, ETextureCreateFlags Flags) = 0;
 
-	virtual VkImage                  RHIGetVkImage(FRHITexture2D* InTexture) const = 0;
-	virtual VkFormat                 RHIGetViewVkFormat(FRHITexture2D* InTexture) const = 0;
-	virtual FVulkanRHIAllocationInfo RHIGetAllocationInfo(FRHITexture2D* InTexture) const = 0;
+	virtual VkImage                  RHIGetVkImage(FRHITexture* InTexture) const = 0;
+	virtual VkFormat                 RHIGetViewVkFormat(FRHITexture* InTexture) const = 0;
+	virtual FVulkanRHIAllocationInfo RHIGetAllocationInfo(FRHITexture* InTexture) const = 0;
 	virtual FVulkanRHIImageViewInfo  RHIGetImageViewInfo(FRHITexture* InTexture) const = 0;
 
-	virtual VkImageLayout& RHIFindOrAddLayoutRW(FRHITexture2D* InTexture, VkImageLayout LayoutIfNotFound) = 0;
+	virtual VkImageLayout& RHIFindOrAddLayoutRW(FRHITexture* InTexture, VkImageLayout LayoutIfNotFound) = 0;
 	virtual void           RHISetImageLayout(VkImage Image, VkImageLayout OldLayout, VkImageLayout NewLayout, const VkImageSubresourceRange& SubresourceRange) = 0;
 	virtual void           RHISetUploadImageLayout(VkImage Image, VkImageLayout OldLayout, VkImageLayout NewLayout, const VkImageSubresourceRange& SubresourceRange) = 0;
 
diff --git a/Engine/Source/Runtime/VulkanRHI/Public/VulkanDynamicRHI.h b/Engine/Source/Runtime/VulkanRHI/Public/VulkanDynamicRHI.h
--- a/Engine/Source/Runtime/VulkanRHI/Public/VulkanDynamicRHI.h
+++ b/Engine/Source/Runtime/VulkanRHI/Public/VulkanDynamicRHI.h
@@ -8,7 +8,6 @@
 
 #include "IVulkanDynamicRHI.h"
 
-class FVulkanTexture2D;
 class FVulkanFramebuffer;
 class FVulkanDevice;
 class FVulkanQueue;
@@ -73,11 +72,11 @@ public:
 	virtual FTexture2DRHIRef RHICreateTexture2DFromResource(EPixelFormat Format, uint32 SizeX, uint32 SizeY, uint32 NumMips, uint32 NumSamples, VkImage Resource, ETextureCreateFlags Flags) final override;
 	virtual FTexture2DArrayRHIRef RHICreateTexture2DArrayFromResource(EPixelFormat Format, uint32 SizeX, uint32 SizeY, uint32 ArraySize, uint32 NumMips, uint32 NumSamples, VkImage Resource, ETextureCreateFlags Flags) final override;
 	virtual FTextureCubeRHIRef RHICreateTextureCubeFromResource(EPixelFormat Format, uint32 Size, bool bArray, uint32 ArraySize, uint32 NumMips, VkImage Resource, ETextureCreateFlags Flags) final override;
-	virtual VkImage RHIGetVkImage(FRHITexture2D* InTexture) const final override;
-	virtual VkFormat RHIGetViewVkFormat(FRHITexture2D* InTexture) const final override;
-	virtual FVulkanRHIAllocationInfo RHIGetAllocationInfo(FRHITexture2D* InTexture) const final override;
+	virtual VkImage RHIGetVkImage(FRHITexture* InTexture) const final override;
+	virtual VkFormat RHIGetViewVkFormat(FRHITexture* InTexture) const final override;
+	virtual FVulkanRHIAllocationInfo RHIGetAllocationInfo(FRHITexture* InTexture) const final override;
 	virtual FVulkanRHIImageViewInfo RHIGetImageViewInfo(FRHITexture* InTexture) const final override;
-	virtual VkImageLayout& RHIFindOrAddLayoutRW(FRHITexture2D* InTexture, VkImageLayout LayoutIfNotFound) final override;
+	virtual VkImageLayout& RHIFindOrAddLayoutRW(FRHITexture* InTexture, VkImageLayout LayoutIfNotFound) final override;
 	virtual void RHISetImageLayout(VkImage Image, VkImageLayout OldLayout, VkImageLayout NewLayout, const VkImageSubresourceRange& SubresourceRange) final override;
 	virtual void RHISetUploadImageLayout(VkImage Image, VkImageLayout OldLayout, VkImageLayout NewLayout, const VkImageSubresourceRange& SubresourceRange) final override;
 	virtual void RHIFinishExternalComputeWork(VkCommandBuffer InCommandBuffer) final override;
@@ -124,17 +123,12 @@ public:
 	virtual FShaderResourceViewRHIRef RHICreateShaderResourceView(const FShaderResourceViewInitializer& Initializer) final override;
 	virtual void RHIUpdateShaderResourceView(FRHIShaderResourceView* SRV, FRHIBuffer* Buffer, uint32 Stride, uint8 Format) final override;
 	virtual void RHIUpdateShaderResourceView(FRHIShaderResourceView* SRV, FRHIBuffer* Buffer) final override;
-	virtual uint64 RHICalcTexture2DPlatformSize(uint32 SizeX, uint32 SizeY, uint8 Format, uint32 NumMips, uint32 NumSamples, ETextureCreateFlags Flags, const FRHIResourceCreateInfo& CreateInfo, uint32& OutAlign) final override;
-	virtual uint64 RHICalcTexture2DArrayPlatformSize(uint32 SizeX, uint32 SizeY, uint32 ArraySize, uint8 Format, uint32 NumMips, uint32 NumSamples, ETextureCreateFlags Flags, const FRHIResourceCreateInfo& CreateInfo, uint32& OutAlign) final override;
-	virtual uint64 RHICalcTexture3DPlatformSize(uint32 SizeX, uint32 SizeY, uint32 SizeZ, uint8 Format, uint32 NumMips, ETextureCreateFlags Flags, const FRHIResourceCreateInfo& CreateInfo, uint32& OutAlign) final override;
-	virtual uint64 RHICalcTextureCubePlatformSize(uint32 Size, uint8 Format, uint32 NumMips, ETextureCreateFlags Flags, const FRHIResourceCreateInfo& CreateInfo, uint32& OutAlign) final override;
+	virtual FRHICalcTextureSizeResult RHICalcTexturePlatformSize(FRHITextureDesc const& Desc, uint32 FirstMipIndex) final override;
 	virtual void RHIGetTextureMemoryStats(FTextureMemoryStats& OutStats) final override;
 	virtual bool RHIGetTextureMemoryVisualizeData(FColor* TextureData, int32 SizeX, int32 SizeY, int32 Pitch, int32 PixelSize) final override;
-	virtual FTexture2DRHIRef RHICreateTexture2D(uint32 SizeX, uint32 SizeY, uint8 Format, uint32 NumMips, uint32 NumSamples, ETextureCreateFlags Flags, ERHIAccess InResourceState, FRHIResourceCreateInfo& CreateInfo) final override;
-	virtual FTexture2DRHIRef RHIAsyncCreateTexture2D(uint32 SizeX, uint32 SizeY, uint8 Format, uint32 NumMips, ETextureCreateFlags Flags, ERHIAccess InResourceState, void** InitialMipData, uint32 NumInitialMips) final override;
+	virtual FTextureRHIRef RHICreateTexture(const FRHITextureCreateDesc& CreateDesc) final override;
+	virtual FTextureRHIRef RHIAsyncCreateTexture2D(uint32 SizeX, uint32 SizeY, uint8 Format, uint32 NumMips, ETextureCreateFlags Flags, ERHIAccess InResourceState, void** InitialMipData, uint32 NumInitialMips) final override;
 	virtual void RHICopySharedMips(FRHITexture2D* DestTexture2D, FRHITexture2D* SrcTexture2D) final override;
-	virtual FTexture2DArrayRHIRef RHICreateTexture2DArray(uint32 SizeX, uint32 SizeY, uint32 SizeZ, uint8 Format, uint32 NumMips, uint32 NumSamples, ETextureCreateFlags Flags, ERHIAccess InResourceState, FRHIResourceCreateInfo& CreateInfo) final override;
-	virtual FTexture3DRHIRef RHICreateTexture3D(uint32 SizeX, uint32 SizeY, uint32 SizeZ, uint8 Format, uint32 NumMips, ETextureCreateFlags Flags, ERHIAccess InResourceState, FRHIResourceCreateInfo& CreateInfo) final override;
 	virtual FShaderResourceViewRHIRef RHICreateShaderResourceView(FRHITexture* Texture, const FRHITextureSRVCreateInfo& CreateInfo) final override;
 	virtual uint32 RHIComputeMemorySize(FRHITexture* TextureRHI) final override;
 	virtual FTexture2DRHIRef RHIAsyncReallocateTexture2D(FRHITexture2D* Texture2D, int32 NewMipCount, int32 NewSizeX, int32 NewSizeY, FThreadSafeCounter* RequestStatus) final override;
@@ -160,8 +154,6 @@ public:
 	{
 		InternalUpdateTexture3D(false, Texture, MipIndex, UpdateRegion, SourceRowPitch, SourceDepthPitch, SourceData);
 	}
-	virtual FTextureCubeRHIRef RHICreateTextureCube(uint32 Size, uint8 Format, uint32 NumMips, ETextureCreateFlags Flags, ERHIAccess InResourceState, FRHIResourceCreateInfo& CreateInfo) final override;
-	virtual FTextureCubeRHIRef RHICreateTextureCubeArray(uint32 Size, uint32 ArraySize, uint8 Format, uint32 NumMips, ETextureCreateFlags Flags, ERHIAccess InResourceState, FRHIResourceCreateInfo& CreateInfo) final override;
 	virtual void* RHILockTextureCubeFace(FRHITextureCube* Texture, uint32 FaceIndex, uint32 ArrayIndex, uint32 MipIndex, EResourceLockMode LockMode, uint32& DestStride, bool bLockWithinMiptail) final override;
 	virtual void RHIUnlockTextureCubeFace(FRHITextureCube* Texture, uint32 FaceIndex, uint32 ArrayIndex, uint32 MipIndex, bool bLockWithinMiptail) final override;
 	virtual void RHIBindDebugLabelName(FRHITexture* Texture, const TCHAR* Name) final override;
@@ -210,8 +202,6 @@ public:
 
 	// Transient Resource Functions and Helpers
 	virtual IRHITransientResourceAllocator* RHICreateTransientResourceAllocator() final override;
-	uint64 RHICalcTexturePlatformSize(const FRHITextureCreateInfo& InCreateInfo, uint32& OutAlign);
-	FRHITexture* CreateTexture(const FRHITextureCreateInfo& InCreateInfo, FRHIResourceCreateInfo& InResourceCreateInfo, ERHIAccess InResourceState, const FRHITransientHeapAllocation* InAllocation);
 	FRHIBuffer* CreateBuffer(const FRHIBufferCreateInfo & InCreateInfo, FRHIResourceCreateInfo & InResourceCreateInfo, const FRHITransientHeapAllocation* InAllocation);
 
 	virtual FBufferRHIRef CreateBuffer_RenderThread(class FRHICommandListImmediate& RHICmdList, uint32 Size, EBufferUsageFlags Usage, uint32 Stride, ERHIAccess ResourceState, FRHIResourceCreateInfo& CreateInfo) override final
@@ -264,21 +254,10 @@ public:
 	virtual FUpdateTexture3DData BeginUpdateTexture3D_RenderThread(class FRHICommandListImmediate& RHICmdList, FRHITexture3D* Texture, uint32 MipIndex, const struct FUpdateTextureRegion3D& UpdateRegion) override final;
 	virtual void EndUpdateTexture3D_RenderThread(class FRHICommandListImmediate& RHICmdList, FUpdateTexture3DData& UpdateData) override final;
 
-	virtual FTexture2DRHIRef RHICreateTexture2D_RenderThread(class FRHICommandListImmediate& RHICmdList, uint32 SizeX, uint32 SizeY, uint8 Format, uint32 NumMips, uint32 NumSamples, ETextureCreateFlags Flags, ERHIAccess InResourceState, FRHIResourceCreateInfo& CreateInfo) override final
+	virtual FTextureRHIRef RHICreateTexture_RenderThread(class FRHICommandListImmediate& RHICmdList, const FRHITextureCreateDesc& CreateDesc) final override
 	{
-		return RHICreateTexture2D(SizeX, SizeY, Format, NumMips, NumSamples, Flags, InResourceState, CreateInfo);
-	}
-	virtual FTexture2DRHIRef RHICreateTextureExternal2D_RenderThread(class FRHICommandListImmediate& RHICmdList, uint32 SizeX, uint32 SizeY, uint8 Format, uint32 NumMips, uint32 NumSamples, ETextureCreateFlags Flags, ERHIAccess InResourceState, FRHIResourceCreateInfo& CreateInfo) override final
-	{
-		return RHICreateTexture2D(SizeX, SizeY, Format, NumMips, NumSamples, Flags, InResourceState, CreateInfo);
-	}
-	virtual FTexture2DArrayRHIRef RHICreateTexture2DArray_RenderThread(class FRHICommandListImmediate& RHICmdList, uint32 SizeX, uint32 SizeY, uint32 SizeZ, uint8 Format, uint32 NumMips, uint32 NumSamples, ETextureCreateFlags Flags, ERHIAccess InResourceState, FRHIResourceCreateInfo& CreateInfo) override final
-	{
-		return this->RHICreateTexture2DArray(SizeX, SizeY, SizeZ, Format, NumMips, NumSamples, Flags, InResourceState, CreateInfo);
-	}
-	virtual FTexture3DRHIRef RHICreateTexture3D_RenderThread(class FRHICommandListImmediate& RHICmdList, uint32 SizeX, uint32 SizeY, uint32 SizeZ, uint8 Format, uint32 NumMips, ETextureCreateFlags Flags, ERHIAccess InResourceState, FRHIResourceCreateInfo& CreateInfo) override final
-	{
-		return this->RHICreateTexture3D(SizeX, SizeY, SizeZ, Format, NumMips, Flags, InResourceState, CreateInfo);
+		// this-> is required to avoid calling the global version of this function
+		return this->RHICreateTexture(CreateDesc);
 	}
 	virtual FUnorderedAccessViewRHIRef RHICreateUnorderedAccessView_RenderThread(class FRHICommandListImmediate& RHICmdList, FRHIBuffer* Buffer, bool bUseUAVCounter, bool bAppendBuffer) final override
 	{
@@ -325,15 +304,6 @@ public:
 		// this-> is required to avoid calling the global version of this function
 		return this->RHICreateShaderResourceView(StructuredBuffer);
 	}
-	virtual FTextureCubeRHIRef RHICreateTextureCube_RenderThread(class FRHICommandListImmediate& RHICmdList, uint32 Size, uint8 Format, uint32 NumMips, ETextureCreateFlags Flags, ERHIAccess InResourceState, FRHIResourceCreateInfo& CreateInfo) override final
-	{
-		return this->RHICreateTextureCube(Size, Format, NumMips, Flags, InResourceState, CreateInfo);
-	}
-
-	virtual FTextureCubeRHIRef RHICreateTextureCubeArray_RenderThread(class FRHICommandListImmediate& RHICmdList, uint32 Size, uint32 ArraySize, uint8 Format, uint32 NumMips, ETextureCreateFlags Flags, ERHIAccess InResourceState, FRHIResourceCreateInfo& CreateInfo) override final
-	{
-		return this->RHICreateTextureCubeArray(Size, ArraySize, Format, NumMips, Flags, InResourceState, CreateInfo);
-	}
 
 	virtual FRenderQueryRHIRef RHICreateRenderQuery_RenderThread(class FRHICommandListImmediate& RHICmdList, ERenderQueryType QueryType) override final
 	{
diff --git a/Engine/Source/Runtime/Windows/D3D11RHI/Private/D3D11Commands.cpp b/Engine/Source/Runtime/Windows/D3D11RHI/Private/D3D11Commands.cpp
--- a/Engine/Source/Runtime/Windows/D3D11RHI/Private/D3D11Commands.cpp
+++ b/Engine/Source/Runtime/Windows/D3D11RHI/Private/D3D11Commands.cpp
@@ -88,7 +88,7 @@ void FD3D11DynamicRHI::RHIBeginUpdateMultiFrameResource(FRHITexture* RHITexture)
 {
 	if (!IsRHIDeviceNVIDIA() || GNumAlternateFrameRenderingGroups == 1) return;
 
-	FD3D11TextureBase* Texture = GetD3D11TextureFromRHITexture(RHITexture);
+	FD3D11Texture* Texture = GetD3D11TextureFromRHITexture(RHITexture);
 
 	if (!Texture)
 	{
@@ -112,7 +112,7 @@ void FD3D11DynamicRHI::RHIEndUpdateMultiFrameResource(FRHITexture* RHITexture)
 {
 	if (!IsRHIDeviceNVIDIA() || GNumAlternateFrameRenderingGroups == 1) return;
 
-	FD3D11TextureBase* Texture = GetD3D11TextureFromRHITexture(RHITexture);
+	FD3D11Texture* Texture = GetD3D11TextureFromRHITexture(RHITexture);
 
 	if (!Texture || !Texture->GetIHVResourceHandle())
 	{
@@ -352,7 +352,7 @@ void FD3D11DynamicRHI::RHISetBoundShaderState(FRHIBoundShaderState* BoundShaderS
 
 void FD3D11DynamicRHI::RHISetShaderTexture(FRHIGraphicsShader* ShaderRHI,uint32 TextureIndex, FRHITexture* NewTextureRHI)
 {
-	FD3D11TextureBase* NewTexture = GetD3D11TextureFromRHITexture(NewTextureRHI);
+	FD3D11Texture* NewTexture = GetD3D11TextureFromRHITexture(NewTextureRHI);
 	ID3D11ShaderResourceView* ShaderResourceView = NewTexture ? NewTexture->GetShaderResourceView() : nullptr;
 
 	switch (ShaderRHI->GetFrequency())
@@ -387,7 +387,7 @@ void FD3D11DynamicRHI::RHISetShaderTexture(FRHIComputeShader* ComputeShaderRHI,u
 {
 	//VALIDATE_BOUND_SHADER(ComputeShaderRHI);
 
-	FD3D11TextureBase* NewTexture = GetD3D11TextureFromRHITexture(NewTextureRHI);
+	FD3D11Texture* NewTexture = GetD3D11TextureFromRHITexture(NewTextureRHI);
 	ID3D11ShaderResourceView* ShaderResourceView = NewTexture ? NewTexture->GetShaderResourceView() : nullptr;
 	SetShaderResourceView<SF_Compute>(NewTexture, ShaderResourceView, TextureIndex);
 }
@@ -894,7 +894,7 @@ void FD3D11DynamicRHI::SetRenderTargets(
 	const FRHIRenderTargetView* NewRenderTargetsRHI,
 	const FRHIDepthRenderTargetView* NewDepthStencilTargetRHI)
 {
-	FD3D11TextureBase* NewDepthStencilTarget = GetD3D11TextureFromRHITexture(NewDepthStencilTargetRHI ? NewDepthStencilTargetRHI->Texture : nullptr);
+	FD3D11Texture* NewDepthStencilTarget = GetD3D11TextureFromRHITexture(NewDepthStencilTargetRHI ? NewDepthStencilTargetRHI->Texture : nullptr);
 
 	check(NewNumSimultaneousRenderTargets <= MaxSimultaneousRenderTargets);
 
@@ -929,7 +929,8 @@ void FD3D11DynamicRHI::SetRenderTargets(
 		{
 			int32 RTMipIndex = NewRenderTargetsRHI[RenderTargetIndex].MipIndex;
 			int32 RTSliceIndex = NewRenderTargetsRHI[RenderTargetIndex].ArraySliceIndex;
-			FD3D11TextureBase* NewRenderTarget = GetD3D11TextureFromRHITexture(NewRenderTargetsRHI[RenderTargetIndex].Texture);
+			
+			FD3D11Texture* NewRenderTarget = GetD3D11TextureFromRHITexture(NewRenderTargetsRHI[RenderTargetIndex].Texture);
 			RenderTargetView = NewRenderTarget ? NewRenderTarget->GetRenderTargetView(RTMipIndex, RTSliceIndex) : nullptr;
 
 			ensureMsgf(RenderTargetView, TEXT("Texture being set as render target has no RTV"));
@@ -1185,9 +1186,6 @@ inline int32 SetShaderResourcesFromBuffer_Surface(FD3D11DynamicRHI* RESTRICT D3D
 			const uint16 ResourceIndex = FRHIResourceTableEntry::GetResourceIndex(ResourceInfo);
 			const uint8 BindIndex = FRHIResourceTableEntry::GetBindIndex(ResourceInfo);
 
-			FD3D11BaseShaderResource* ShaderResource = nullptr;
-			ID3D11ShaderResourceView* D3D11Resource = nullptr;
-
 			check(ResourceIndex < NumResourcesInTable);
 			FRHITexture* TextureRHI = (FRHITexture*)Resources[ResourceIndex].GetReference();
 			if (!TextureRHI)
@@ -1195,9 +1193,9 @@ inline int32 SetShaderResourcesFromBuffer_Surface(FD3D11DynamicRHI* RESTRICT D3D
 				UE_LOG(LogD3D11RHI, Fatal, TEXT("Null texture (resource %d bind %d) on UB Layout %s"), ResourceIndex, BindIndex, LayoutName);
 			}
 			TextureRHI->SetLastRenderTime(CurrentTime);
-			FD3D11TextureBase* TextureD3D11 = GetD3D11TextureFromRHITexture(TextureRHI);
-			ShaderResource = TextureD3D11->GetBaseShaderResource();
-			D3D11Resource = TextureD3D11->GetShaderResourceView();
+
+			FD3D11Texture* TextureD3D11 = GetD3D11TextureFromRHITexture(TextureRHI);
+			ID3D11ShaderResourceView* D3D11Resource = TextureD3D11->GetShaderResourceView();
 
 #if ENABLE_RHI_VALIDATION
 			if (D3D11RHI->Tracker)
@@ -1213,7 +1211,7 @@ inline int32 SetShaderResourcesFromBuffer_Surface(FD3D11DynamicRHI* RESTRICT D3D
 #endif
 
 			// todo: could coalesce adjacent bound resources.
-			SetResource<ShaderFrequency>(D3D11RHI, StateCache, BindIndex, ShaderResource, D3D11Resource);
+			SetResource<ShaderFrequency>(D3D11RHI, StateCache, BindIndex, TextureD3D11, D3D11Resource);
 			NumSetCalls++;
 			ResourceInfo = *ResourceInfos++;
 		} while (FRHIResourceTableEntry::GetUniformBufferIndex(ResourceInfo) == BufferIndex);
diff --git a/Engine/Source/Runtime/Windows/D3D11RHI/Private/D3D11RHI.cpp b/Engine/Source/Runtime/Windows/D3D11RHI/Private/D3D11RHI.cpp
--- a/Engine/Source/Runtime/Windows/D3D11RHI/Private/D3D11RHI.cpp
+++ b/Engine/Source/Runtime/Windows/D3D11RHI/Private/D3D11RHI.cpp
@@ -792,25 +792,25 @@ ID3D11Buffer* FD3D11DynamicRHI::RHIGetResource(FRHIBuffer* InBuffer) const
 
 ID3D11Resource* FD3D11DynamicRHI::RHIGetResource(FRHITexture* InTexture) const
 {
-	FD3D11TextureBase* D3D11Texture = GetD3D11TextureFromRHITexture(InTexture);
+	FD3D11Texture* D3D11Texture = GetD3D11TextureFromRHITexture(InTexture);
 	return D3D11Texture->GetResource();
 }
 
 int64 FD3D11DynamicRHI::RHIGetResourceMemorySize(FRHITexture* InTexture) const
 {
-	FD3D11TextureBase* D3D11Texture = GetD3D11TextureFromRHITexture(InTexture);
+	FD3D11Texture* D3D11Texture = GetD3D11TextureFromRHITexture(InTexture);
 	return D3D11Texture->GetMemorySize();
 }
 
 ID3D11RenderTargetView* FD3D11DynamicRHI::RHIGetRenderTargetView(FRHITexture* InTexture, int32 InMipIndex, int32 InArraySliceIndex) const
 {
-	FD3D11TextureBase* D3D11Texture = GetD3D11TextureFromRHITexture(InTexture);
+	FD3D11Texture* D3D11Texture = GetD3D11TextureFromRHITexture(InTexture);
 	return D3D11Texture->GetRenderTargetView(InMipIndex, InArraySliceIndex);
 }
 
 ID3D11ShaderResourceView* FD3D11DynamicRHI::RHIGetShaderResourceView(FRHITexture* InTexture) const
 {
-	FD3D11TextureBase* D3D11Texture = GetD3D11TextureFromRHITexture(InTexture);
+	FD3D11Texture* D3D11Texture = GetD3D11TextureFromRHITexture(InTexture);
 	return D3D11Texture->GetShaderResourceView();
 }
 
diff --git a/Engine/Source/Runtime/Windows/D3D11RHI/Private/D3D11RHIPrivate.h b/Engine/Source/Runtime/Windows/D3D11RHI/Private/D3D11RHIPrivate.h
--- a/Engine/Source/Runtime/Windows/D3D11RHI/Private/D3D11RHIPrivate.h
+++ b/Engine/Source/Runtime/Windows/D3D11RHI/Private/D3D11RHIPrivate.h
@@ -488,18 +488,13 @@ public:
 	virtual FShaderResourceViewRHIRef RHICreateShaderResourceView(const FShaderResourceViewInitializer& Initializer) final override;
 	virtual void RHIUpdateShaderResourceView(FRHIShaderResourceView* SRV, FRHIBuffer* VertexBuffer, uint32 Stride, uint8 Format) final override;
 	virtual void RHIUpdateShaderResourceView(FRHIShaderResourceView* SRV, FRHIBuffer* IndexBuffer) final override;
-	virtual uint64 RHICalcTexture2DPlatformSize(uint32 SizeX, uint32 SizeY, uint8 Format, uint32 NumMips, uint32 NumSamples, ETextureCreateFlags Flags, const FRHIResourceCreateInfo& CreateInfo, uint32& OutAlign) final override;
-	virtual uint64 RHICalcTexture2DArrayPlatformSize(uint32 SizeX, uint32 SizeY, uint32 ArraySize, uint8 Format, uint32 NumMips, uint32 NumSamples, ETextureCreateFlags Flags, const FRHIResourceCreateInfo& CreateInfo, uint32& OutAlign) final override;
-	virtual uint64 RHICalcTexture3DPlatformSize(uint32 SizeX, uint32 SizeY, uint32 SizeZ, uint8 Format, uint32 NumMips, ETextureCreateFlags Flags, const FRHIResourceCreateInfo& CreateInfo, uint32& OutAlign) final override;
-	virtual uint64 RHICalcTextureCubePlatformSize(uint32 Size, uint8 Format, uint32 NumMips, ETextureCreateFlags Flags, const FRHIResourceCreateInfo& CreateInfo, uint32& OutAlign) final override;
+	virtual FRHICalcTextureSizeResult RHICalcTexturePlatformSize(FRHITextureDesc const& Desc, uint32 FirstMipIndex) final override;
 	virtual uint64 RHIGetMinimumAlignmentForBufferBackedSRV(EPixelFormat Format) final override;
 	virtual void RHIGetTextureMemoryStats(FTextureMemoryStats& OutStats) final override;
 	virtual bool RHIGetTextureMemoryVisualizeData(FColor* TextureData,int32 SizeX,int32 SizeY,int32 Pitch,int32 PixelSize) final override;
-	virtual FTexture2DRHIRef RHICreateTexture2D(uint32 SizeX, uint32 SizeY, uint8 Format, uint32 NumMips, uint32 NumSamples, ETextureCreateFlags Flags, ERHIAccess InResourceState, FRHIResourceCreateInfo& CreateInfo) final override;
-	virtual FTexture2DRHIRef RHIAsyncCreateTexture2D(uint32 SizeX, uint32 SizeY, uint8 Format, uint32 NumMips, ETextureCreateFlags Flags, ERHIAccess InResourceState, void** InitialMipData, uint32 NumInitialMips) final override;
+	virtual FTextureRHIRef RHICreateTexture(const FRHITextureCreateDesc& CreateDesc) final override;
+	virtual FTextureRHIRef RHIAsyncCreateTexture2D(uint32 SizeX, uint32 SizeY, uint8 Format, uint32 NumMips, ETextureCreateFlags Flags, ERHIAccess InResourceState, void** InitialMipData, uint32 NumInitialMips) final override;
 	virtual void RHICopySharedMips(FRHITexture2D* DestTexture2D, FRHITexture2D* SrcTexture2D) final override;
-	virtual FTexture2DArrayRHIRef RHICreateTexture2DArray(uint32 SizeX, uint32 SizeY, uint32 SizeZ, uint8 Format, uint32 NumMips, uint32 NumSamples, ETextureCreateFlags Flags, ERHIAccess InResourceState, FRHIResourceCreateInfo& CreateInfo) final override;
-	virtual FTexture3DRHIRef RHICreateTexture3D(uint32 SizeX, uint32 SizeY, uint32 SizeZ, uint8 Format, uint32 NumMips, ETextureCreateFlags Flags, ERHIAccess InResourceState, FRHIResourceCreateInfo& CreateInfo) final override;
 	virtual FShaderResourceViewRHIRef RHICreateShaderResourceView(FRHITexture* TextureRHI, const FRHITextureSRVCreateInfo& CreateInfo) final override;
 	virtual void RHIGenerateMips(FRHITexture* Texture) final override;
 	virtual uint32 RHIComputeMemorySize(FRHITexture* TextureRHI) final override;
@@ -513,8 +508,6 @@ public:
 	virtual void RHIUnlockTexture2DArray(FRHITexture2DArray* Texture, uint32 TextureIndex, uint32 MipIndex, bool bLockWithinMiptail) final override;
 	virtual void RHIUpdateTexture2D(FRHITexture2D* Texture, uint32 MipIndex, const struct FUpdateTextureRegion2D& UpdateRegion, uint32 SourcePitch, const uint8* SourceData) final override;
 	virtual void RHIUpdateTexture3D(FRHITexture3D* Texture, uint32 MipIndex, const struct FUpdateTextureRegion3D& UpdateRegion, uint32 SourceRowPitch, uint32 SourceDepthPitch, const uint8* SourceData) final override;
-	virtual FTextureCubeRHIRef RHICreateTextureCube(uint32 Size, uint8 Format, uint32 NumMips, ETextureCreateFlags Flags, ERHIAccess InResourceState, FRHIResourceCreateInfo& CreateInfo) final override;
-	virtual FTextureCubeRHIRef RHICreateTextureCubeArray(uint32 Size, uint32 ArraySize, uint8 Format, uint32 NumMips, ETextureCreateFlags Flags, ERHIAccess InResourceState, FRHIResourceCreateInfo& CreateInfo) final override;
 	virtual void* RHILockTextureCubeFace(FRHITextureCube* Texture, uint32 FaceIndex, uint32 ArrayIndex, uint32 MipIndex, EResourceLockMode LockMode, uint32& DestStride, bool bLockWithinMiptail) final override;
 	virtual void RHIUnlockTextureCubeFace(FRHITextureCube* Texture, uint32 FaceIndex, uint32 ArrayIndex, uint32 MipIndex, bool bLockWithinMiptail) final override;
 	virtual void RHIBindDebugLabelName(FRHITexture* Texture, const TCHAR* Name) final override;
@@ -651,20 +644,16 @@ public:
 	virtual void UpdateTexture2D_RenderThread(class FRHICommandListImmediate& RHICmdList, FRHITexture2D* Texture, uint32 MipIndex, const struct FUpdateTextureRegion2D& UpdateRegion, uint32 SourcePitch, const uint8* SourceData) final override;
 	virtual void EndUpdateTexture3D_RenderThread(class FRHICommandListImmediate& RHICmdList, FUpdateTexture3DData& UpdateData) final override;
 	virtual void UpdateTexture3D_RenderThread(class FRHICommandListImmediate& RHICmdList, FRHITexture3D* Texture, uint32 MipIndex, const struct FUpdateTextureRegion3D& UpdateRegion, uint32 SourceRowPitch, uint32 SourceDepthPitch, const uint8* SourceData) final override;
-	virtual FTexture2DRHIRef RHICreateTexture2D_RenderThread(class FRHICommandListImmediate& RHICmdList, uint32 SizeX, uint32 SizeY, uint8 Format, uint32 NumMips, uint32 NumSamples, ETextureCreateFlags Flags, ERHIAccess InResourceState, FRHIResourceCreateInfo& CreateInfo) final override;
-	virtual FTexture2DArrayRHIRef RHICreateTexture2DArray_RenderThread(class FRHICommandListImmediate& RHICmdList, uint32 SizeX, uint32 SizeY, uint32 SizeZ, uint8 Format, uint32 NumMips, uint32 NumSamples, ETextureCreateFlags Flags, ERHIAccess InResourceState, FRHIResourceCreateInfo& CreateInfo) final override;
-	virtual FTexture3DRHIRef RHICreateTexture3D_RenderThread(class FRHICommandListImmediate& RHICmdList, uint32 SizeX, uint32 SizeY, uint32 SizeZ, uint8 Format, uint32 NumMips, ETextureCreateFlags Flags, ERHIAccess InResourceState, FRHIResourceCreateInfo& CreateInfo) final override;
 	virtual FUnorderedAccessViewRHIRef RHICreateUnorderedAccessView_RenderThread(class FRHICommandListImmediate& RHICmdList, FRHIBuffer* Buffer, bool bUseUAVCounter, bool bAppendBuffer) final override;
 	virtual FUnorderedAccessViewRHIRef RHICreateUnorderedAccessView_RenderThread(class FRHICommandListImmediate& RHICmdList, FRHITexture* Texture, uint32 MipLevel, uint16 FirstArraySlice, uint16 NumArraySlices) final override;
 	virtual FUnorderedAccessViewRHIRef RHICreateUnorderedAccessView_RenderThread(class FRHICommandListImmediate& RHICmdList, FRHIBuffer* Buffer, uint8 Format) final override;
 	virtual FShaderResourceViewRHIRef RHICreateShaderResourceView_RenderThread(class FRHICommandListImmediate& RHICmdList, const FShaderResourceViewInitializer& Initializer) final override;
 	virtual FShaderResourceViewRHIRef RHICreateShaderResourceView_RenderThread(class FRHICommandListImmediate& RHICmdList, FRHITexture* Texture, const FRHITextureSRVCreateInfo& CreateInfo) final override;
 	virtual FShaderResourceViewRHIRef RHICreateShaderResourceView_RenderThread(class FRHICommandListImmediate& RHICmdList, FRHIBuffer* Buffer) final override;
-	virtual FTextureCubeRHIRef RHICreateTextureCube_RenderThread(class FRHICommandListImmediate& RHICmdList, uint32 Size, uint8 Format, uint32 NumMips, ETextureCreateFlags Flags, ERHIAccess InResourceState, FRHIResourceCreateInfo& CreateInfo) final override;
-	virtual FTextureCubeRHIRef RHICreateTextureCubeArray_RenderThread(class FRHICommandListImmediate& RHICmdList, uint32 Size, uint32 ArraySize, uint8 Format, uint32 NumMips, ETextureCreateFlags Flags, ERHIAccess InResourceState, FRHIResourceCreateInfo& CreateInfo) final override;
 	virtual FRenderQueryRHIRef RHICreateRenderQuery_RenderThread(class FRHICommandListImmediate& RHICmdList, ERenderQueryType QueryType) final override;
 	virtual void* RHILockTextureCubeFace_RenderThread(class FRHICommandListImmediate& RHICmdList, FRHITextureCube* Texture, uint32 FaceIndex, uint32 ArrayIndex, uint32 MipIndex, EResourceLockMode LockMode, uint32& DestStride, bool bLockWithinMiptail) final override;
 	virtual void RHIUnlockTextureCubeFace_RenderThread(class FRHICommandListImmediate& RHICmdList, FRHITextureCube* Texture, uint32 FaceIndex, uint32 ArrayIndex, uint32 MipIndex, bool bLockWithinMiptail) final override;
+	virtual FTextureRHIRef RHICreateTexture_RenderThread(class FRHICommandListImmediate& RHICmdList, const FRHITextureCreateDesc& CreateDesc) final override;
 
 	virtual void RHIBeginRenderPass(const FRHIRenderPassInfo& InInfo, const TCHAR* InName) final override;
 	virtual void RHIEndRenderPass() final override;
@@ -780,14 +769,6 @@ public:
 
 	static DXGI_FORMAT GetPlatformTextureResourceFormat(DXGI_FORMAT InFormat, ETextureCreateFlags InFlags);
 
-#if	PLATFORM_SUPPORTS_VIRTUAL_TEXTURES
-	virtual void* CreateVirtualTexture(uint32 SizeX, uint32 SizeY, uint32 SizeZ, uint32 NumMips, bool bCubeTexture, ETextureCreateFlags Flags, void* D3DTextureDesc, void* D3DTextureResource) = 0;
-	virtual void DestroyVirtualTexture(ETextureCreateFlags Flags, void* RawTextureMemory) = 0;
-	virtual bool HandleSpecialLock(FD3D11LockedData& LockedData, uint32 MipIndex, uint32 ArrayIndex, ETextureCreateFlags Flags, EResourceLockMode LockMode,
-		void* D3DTextureResource, void* RawTextureMemory, uint32 NumMips, uint32& DestStride) = 0;
-	virtual bool HandleSpecialUnlock(uint32 MipIndex, ETextureCreateFlags Flags, void* D3DTextureResource, void* RawTextureMemory) = 0;
-#endif
-
 	uint32 GetHDRDetectedDisplayIndex() const
 	{
 		return HDRDetectedDisplayIndex;
@@ -874,7 +855,7 @@ protected:
 	uint32 CurrentUAVMask;
 
 	TRefCountPtr<ID3D11DepthStencilView> CurrentDepthStencilTarget;
-	TRefCountPtr<FD3D11TextureBase> CurrentDepthTexture;
+	TRefCountPtr<FD3D11Texture> CurrentDepthTexture;
 	FD3D11BaseShaderResource* CurrentResourcesBoundAsSRVs[SF_NumStandardFrequencies][D3D11_COMMONSHADER_INPUT_RESOURCE_SLOT_COUNT];
 	FD3D11BaseShaderResource* CurrentResourcesBoundAsVBs[D3D11_IA_VERTEX_INPUT_RESOURCE_SLOT_COUNT];
 	FD3D11BaseShaderResource* CurrentResourceBoundAsIB;
@@ -960,17 +941,10 @@ protected:
 	// If this is false, disable any IHV optimization/libs
 	bool bAllowVendorDevice;
 
-	template<typename BaseResourceType>
-	TD3D11Texture2D<BaseResourceType>* CreateD3D11Texture2D(uint32 SizeX,uint32 SizeY,uint32 SizeZ,bool bTextureArray,bool CubeTexture,uint8 Format,
-		uint32 NumMips,uint32 NumSamples,ETextureCreateFlags Flags,FRHIResourceCreateInfo& CreateInfo);
-
-	FD3D11Texture3D* CreateD3D11Texture3D(uint32 SizeX,uint32 SizeY,uint32 SizeZ,uint8 Format,uint32 NumMips,ETextureCreateFlags Flags,FRHIResourceCreateInfo& CreateInfo);
-
-	template<typename BaseResourceType>
-	TD3D11Texture2D<BaseResourceType>* CreateTextureFromResource(bool bTextureArray, bool bCubeTexture, EPixelFormat Format, ETextureCreateFlags TexCreateFlags, const FClearValueBinding& ClearValueBinding, ID3D11Texture2D* TextureResource);
+	FD3D11Texture* CreateD3D11Texture2D(FRHITextureCreateDesc const& CreateDesc, TConstArrayView<D3D11_SUBRESOURCE_DATA> InitialData = {});
+	FD3D11Texture* CreateD3D11Texture3D(FRHITextureCreateDesc const& CreateDesc);
 
-	template<typename BaseResourceType>
-	FTextureRHIRef CreateAliasedD3D11Texture2D(TD3D11Texture2D<BaseResourceType>* SourceTexture);
+	FD3D11Texture* CreateTextureFromResource(bool bTextureArray, bool bCubeTexture, EPixelFormat Format, ETextureCreateFlags TexCreateFlags, const FClearValueBinding& ClearValueBinding, ID3D11Texture2D* TextureResource);
 
 	/** Initializes the constant buffers.  Called once at RHI initialization time. */
 	void InitConstantBuffers();
@@ -1024,8 +998,8 @@ protected:
 	template<typename TPixelShader>
 	void ResolveTextureUsingShader(
 		FRHICommandList_RecursiveHazardous& RHICmdList,
-		FD3D11Texture2D* SourceTexture,
-		FD3D11Texture2D* DestTexture,
+		FD3D11Texture* SourceTexture,
+		FD3D11Texture* DestTexture,
 		ID3D11RenderTargetView* DestSurfaceRTV,
 		ID3D11DepthStencilView* DestSurfaceDSV,
 		const D3D11_TEXTURE2D_DESC& ResolveTargetDesc,
diff --git a/Engine/Source/Runtime/Windows/D3D11RHI/Private/D3D11RenderTarget.cpp b/Engine/Source/Runtime/Windows/D3D11RHI/Private/D3D11RenderTarget.cpp
--- a/Engine/Source/Runtime/Windows/D3D11RHI/Private/D3D11RenderTarget.cpp
+++ b/Engine/Source/Runtime/Windows/D3D11RHI/Private/D3D11RenderTarget.cpp
@@ -46,8 +46,8 @@ static FResolveRect GetDefaultRect(const FResolveRect& Rect,uint32 DefaultWidth,
 template<typename TPixelShader>
 void FD3D11DynamicRHI::ResolveTextureUsingShader(
 	FRHICommandList_RecursiveHazardous& RHICmdList,
-	FD3D11Texture2D* SourceTexture,
-	FD3D11Texture2D* DestTexture,
+	FD3D11Texture* SourceTexture,
+	FD3D11Texture* DestTexture,
 	ID3D11RenderTargetView* DestTextureRTV,
 	ID3D11DepthStencilView* DestTextureDSV,
 	const D3D11_TEXTURE2D_DESC& ResolveTargetDesc,
@@ -88,7 +88,7 @@ void FD3D11DynamicRHI::ResolveTextureUsingShader(
 	//we may change rendertargets and depth state behind the RHI's back here.
 	//save off this original state to restore it.
 	FExclusiveDepthStencil OriginalDSVAccessType = CurrentDSVAccessType;
-	TRefCountPtr<FD3D11TextureBase> OriginalDepthTexture = CurrentDepthTexture;
+	TRefCountPtr<FD3D11Texture> OriginalDepthTexture = CurrentDepthTexture;
 
 	if(ResolveTargetDesc.BindFlags & D3D11_BIND_DEPTH_STENCIL)
 	{
@@ -189,16 +189,15 @@ void FD3D11DynamicRHI::RHICopyToResolveTarget(FRHITexture* SourceTextureRHI, FRH
 	}
 
 	FRHICommandList_RecursiveHazardous RHICmdList(this);
-	
 
-	FD3D11Texture2D* SourceTexture2D = static_cast<FD3D11Texture2D*>(SourceTextureRHI->GetTexture2D());
-	FD3D11Texture2D* DestTexture2D = static_cast<FD3D11Texture2D*>(DestTextureRHI->GetTexture2D());
+	FD3D11Texture* SourceTexture2D   = GetD3D11TextureFromRHITexture(SourceTextureRHI->GetTexture2D());
+	FD3D11Texture* DestTexture2D     = GetD3D11TextureFromRHITexture(DestTextureRHI->GetTexture2D());
 
-	FD3D11TextureCube* SourceTextureCube = static_cast<FD3D11TextureCube*>(SourceTextureRHI->GetTextureCube());
-	FD3D11TextureCube* DestTextureCube = static_cast<FD3D11TextureCube*>(DestTextureRHI->GetTextureCube());
+	FD3D11Texture* SourceTextureCube = GetD3D11TextureFromRHITexture(SourceTextureRHI->GetTextureCube());
+	FD3D11Texture* DestTextureCube   = GetD3D11TextureFromRHITexture(DestTextureRHI->GetTextureCube());
 
-	FD3D11Texture3D* SourceTexture3D = static_cast<FD3D11Texture3D*>(SourceTextureRHI->GetTexture3D());
-	FD3D11Texture3D* DestTexture3D = static_cast<FD3D11Texture3D*>(DestTextureRHI->GetTexture3D());
+	FD3D11Texture* SourceTexture3D   = GetD3D11TextureFromRHITexture(SourceTextureRHI->GetTexture3D());
+	FD3D11Texture* DestTexture3D     = GetD3D11TextureFromRHITexture(DestTextureRHI->GetTexture3D());
 		
 	if(SourceTexture2D && DestTexture2D)
 	{
@@ -212,8 +211,7 @@ void FD3D11DynamicRHI::RHICopyToResolveTarget(FRHITexture* SourceTextureRHI, FRH
 				&& !DestTextureRHI->IsMultisampled())
 			{
 				D3D11_TEXTURE2D_DESC ResolveTargetDesc;
-				
-				DestTexture2D->GetResource()->GetDesc(&ResolveTargetDesc);
+				DestTexture2D->GetD3D11Texture2D()->GetDesc(&ResolveTargetDesc);
 
 				ResolveTextureUsingShader<FResolveDepthPS>(
 					RHICmdList,
@@ -430,9 +428,10 @@ static uint32 ComputeBytesPerPixel(DXGI_FORMAT Format)
 
 TRefCountPtr<ID3D11Texture2D> FD3D11DynamicRHI::GetStagingTexture(FRHITexture* TextureRHI,FIntRect InRect, FIntRect& StagingRectOUT, FReadSurfaceDataFlags InFlags)
 {
-	FD3D11TextureBase* Texture = GetD3D11TextureFromRHITexture(TextureRHI);
+	FD3D11Texture* Texture = GetD3D11TextureFromRHITexture(TextureRHI);
+
 	D3D11_TEXTURE2D_DESC SourceDesc; 
-	((ID3D11Texture2D*)Texture->GetResource())->GetDesc(&SourceDesc);// check for 3D textures?
+	Texture->GetD3D11Texture2D()->GetDesc(&SourceDesc);
 	
 	bool bRequiresTempStagingTexture = SourceDesc.Usage != D3D11_USAGE_STAGING; 
 	if(bRequiresTempStagingTexture == false)
@@ -441,7 +440,7 @@ TRefCountPtr<ID3D11Texture2D> FD3D11DynamicRHI::GetStagingTexture(FRHITexture* T
 		// a new staging texture as we do not have to wait for the GPU pipeline to catch up
 		// to the staging texture preparation work.
 		StagingRectOUT = InRect;
-		return ((ID3D11Texture2D*)Texture->GetResource());
+		return Texture->GetD3D11Texture2D();
 	}
 
 	// a temporary staging texture is needed.
@@ -506,14 +505,14 @@ void FD3D11DynamicRHI::ReadSurfaceDataNoMSAARaw(FRHITexture* TextureRHI,FIntRect
 	checkf(InRect.Width() <= TextureRHI->GetSizeXYZ().X >> InFlags.GetMip(), TEXT("Provided rect width (%d), must be smaller or equal to the texture size requested Mip (%d)"), InRect.Width(), TextureRHI->GetSizeXYZ().X >> InFlags.GetMip());
 	checkf(InRect.Height() <= TextureRHI->GetSizeXYZ().Y >> InFlags.GetMip(), TEXT("Provided rect height (%d), must be smaller or equal to the texture size requested Mip (%d)"), InRect.Height(), TextureRHI->GetSizeXYZ().Y >> InFlags.GetMip());
 
-	FD3D11TextureBase* Texture = GetD3D11TextureFromRHITexture(TextureRHI);
+	FD3D11Texture* Texture = GetD3D11TextureFromRHITexture(TextureRHI);
 
 	const uint32 SizeX = InRect.Width();
 	const uint32 SizeY = InRect.Height();
 
 	// Check the format of the surface
 	D3D11_TEXTURE2D_DESC TextureDesc;
-	((ID3D11Texture2D*)Texture->GetResource())->GetDesc(&TextureDesc);
+	Texture->GetD3D11Texture2D()->GetDesc(&TextureDesc);
 	
 	uint32 BytesPerPixel = ComputeBytesPerPixel(TextureDesc.Format);
 	
@@ -645,12 +644,11 @@ void FD3D11DynamicRHI::RHIReadSurfaceData(FRHITexture* TextureRHI,FIntRect InRec
 
 	TArray<uint8> OutDataRaw;
 
-	FD3D11TextureBase* Texture = GetD3D11TextureFromRHITexture(TextureRHI);
+	FD3D11Texture* Texture = GetD3D11TextureFromRHITexture(TextureRHI);
 
 	// Check the format of the surface
 	D3D11_TEXTURE2D_DESC TextureDesc;
-
-	((ID3D11Texture2D*)Texture->GetResource())->GetDesc(&TextureDesc);
+	Texture->GetD3D11Texture2D()->GetDesc(&TextureDesc);
 
 	check(TextureDesc.SampleDesc.Count >= 1);
 
@@ -679,14 +677,14 @@ void FD3D11DynamicRHI::RHIReadSurfaceData(FRHITexture* TextureRHI,FIntRect InRec
 
 void FD3D11DynamicRHI::ReadSurfaceDataMSAARaw(FRHICommandList_RecursiveHazardous& RHICmdList, FRHITexture* TextureRHI,FIntRect InRect,TArray<uint8>& OutData, FReadSurfaceDataFlags InFlags)
 {
-	FD3D11TextureBase* Texture = GetD3D11TextureFromRHITexture(TextureRHI);
+	FD3D11Texture* Texture = GetD3D11TextureFromRHITexture(TextureRHI);
 
 	const uint32 SizeX = InRect.Width();
 	const uint32 SizeY = InRect.Height();
 	
 	// Check the format of the surface
 	D3D11_TEXTURE2D_DESC TextureDesc;
-	((ID3D11Texture2D*)Texture->GetResource())->GetDesc(&TextureDesc);
+	Texture->GetD3D11Texture2D()->GetDesc(&TextureDesc);
 
 	uint32 BytesPerPixel = ComputeBytesPerPixel(TextureDesc.Format);
 
@@ -764,7 +762,7 @@ void FD3D11DynamicRHI::ReadSurfaceDataMSAARaw(FRHICommandList_RecursiveHazardous
 		// Resolve the sample to the non-MSAA render target.
 		ResolveTextureUsingShader<FResolveSingleSamplePS>(
 			RHICmdList,
-			(FD3D11Texture2D*)TextureRHI->GetTexture2D(),
+			Texture,
 			NULL,
 			NonMSAARTV,
 			NULL,
@@ -805,7 +803,7 @@ void FD3D11DynamicRHI::ReadSurfaceDataMSAARaw(FRHICommandList_RecursiveHazardous
 
 void FD3D11DynamicRHI::RHIMapStagingSurface(FRHITexture* TextureRHI, FRHIGPUFence* FenceRHI, void*& OutData, int32& OutWidth, int32& OutHeight, uint32 GPUIndex)
 {
-	ID3D11Texture2D* Texture = (ID3D11Texture2D*)(GetD3D11TextureFromRHITexture(TextureRHI)->GetResource());
+	ID3D11Texture2D* Texture = GetD3D11TextureFromRHITexture(TextureRHI)->GetD3D11Texture2D();
 	
 	D3D11_TEXTURE2D_DESC TextureDesc;
 	Texture->GetDesc(&TextureDesc);
@@ -823,21 +821,21 @@ void FD3D11DynamicRHI::RHIMapStagingSurface(FRHITexture* TextureRHI, FRHIGPUFenc
 
 void FD3D11DynamicRHI::RHIUnmapStagingSurface(FRHITexture* TextureRHI, uint32 GPUIndex)
 {
-	ID3D11Texture2D* Texture = (ID3D11Texture2D*)(GetD3D11TextureFromRHITexture(TextureRHI)->GetResource());
+	ID3D11Texture2D* Texture = GetD3D11TextureFromRHITexture(TextureRHI)->GetD3D11Texture2D();
 
 	Direct3DDeviceIMContext->Unmap(Texture,0);
 }
 
 void FD3D11DynamicRHI::RHIReadSurfaceFloatData(FRHITexture* TextureRHI,FIntRect InRect,TArray<FFloat16Color>& OutData,ECubeFace CubeFace,int32 ArrayIndex,int32 MipIndex)
 {
-	FD3D11TextureBase* Texture = GetD3D11TextureFromRHITexture(TextureRHI);
+	FD3D11Texture* Texture = GetD3D11TextureFromRHITexture(TextureRHI);
 
 	uint32 SizeX = InRect.Width();
 	uint32 SizeY = InRect.Height();
 
 	// Check the format of the surface
 	D3D11_TEXTURE2D_DESC TextureDesc;
-	((ID3D11Texture2D*)Texture->GetResource())->GetDesc(&TextureDesc);
+	Texture->GetD3D11Texture2D()->GetDesc(&TextureDesc);
 
 	check(TextureDesc.Format == GPixelFormats[PF_FloatRGBA].PlatformFormat);
 
@@ -967,12 +965,11 @@ void FD3D11DynamicRHI::RHIReadSurfaceData(FRHITexture* TextureRHI, FIntRect InRe
 {
 	TArray<uint8> OutDataRaw;
 
-	FD3D11TextureBase* Texture = GetD3D11TextureFromRHITexture(TextureRHI);
+	FD3D11Texture* Texture = GetD3D11TextureFromRHITexture(TextureRHI);
 
 	// Check the format of the surface
 	D3D11_TEXTURE2D_DESC TextureDesc;
-
-	((ID3D11Texture2D*)Texture->GetResource())->GetDesc(&TextureDesc);
+	Texture->GetD3D11Texture2D()->GetDesc(&TextureDesc);
 
 	check(TextureDesc.SampleDesc.Count >= 1);
 
@@ -1004,7 +1001,7 @@ void FD3D11DynamicRHI::RHIReadSurfaceData(FRHITexture* TextureRHI, FIntRect InRe
 
 void FD3D11DynamicRHI::RHIRead3DSurfaceFloatData(FRHITexture* TextureRHI,FIntRect InRect,FIntPoint ZMinMax,TArray<FFloat16Color>& OutData)
 {
-	FD3D11TextureBase* Texture = GetD3D11TextureFromRHITexture(TextureRHI);
+	FD3D11Texture* Texture = GetD3D11TextureFromRHITexture(TextureRHI);
 
 	uint32 SizeX = InRect.Width();
 	uint32 SizeY = InRect.Height();
@@ -1012,7 +1009,7 @@ void FD3D11DynamicRHI::RHIRead3DSurfaceFloatData(FRHITexture* TextureRHI,FIntRec
 
 	// Check the format of the surface
 	D3D11_TEXTURE3D_DESC TextureDesc;
-	((ID3D11Texture3D*)Texture->GetResource())->GetDesc(&TextureDesc);
+	Texture->GetD3D11Texture3D()->GetDesc(&TextureDesc);
 
 	bool bIsRGBAFmt = TextureDesc.Format == GPixelFormats[PF_FloatRGBA].PlatformFormat;
 	bool bIsR16FFmt = TextureDesc.Format == GPixelFormats[PF_R16F].PlatformFormat;	
