Vulkan RT: add support for procedural geometry

diff --git a/Engine/Source/Runtime/VulkanRHI/Private/VulkanRayTracing.cpp b/Engine/Source/Runtime/VulkanRHI/Private/VulkanRayTracing.cpp
--- a/Engine/Source/Runtime/VulkanRHI/Private/VulkanRayTracing.cpp
+++ b/Engine/Source/Runtime/VulkanRHI/Private/VulkanRayTracing.cpp
@@ -175,6 +175,7 @@ void FVulkanRayTracingAllocator::Free(FVkRtAllocation& Allocation)
 static void GetBLASBuildData(
 	const VkDevice Device,
 	const TArrayView<const FRayTracingGeometrySegment> Segments,
+	const ERayTracingGeometryType GeometryType,
 	const FBufferRHIRef IndexBufferRHI,
 	const uint32 IndexBufferOffset,
 	const bool bFastBuild,
@@ -214,44 +215,62 @@ static void GetBLASBuildData(
 			SegmentGeometry.flags |= VK_GEOMETRY_NO_DUPLICATE_ANY_HIT_INVOCATION_BIT_KHR;
 		}
 
-		// Only support triangles
-		SegmentGeometry.geometryType = VK_GEOMETRY_TYPE_TRIANGLES_KHR;
-
-		SegmentGeometry.geometry.triangles.sType = VK_STRUCTURE_TYPE_ACCELERATION_STRUCTURE_GEOMETRY_TRIANGLES_DATA_KHR;
-		SegmentGeometry.geometry.triangles.vertexFormat = VK_FORMAT_R32G32B32_SFLOAT;
-		SegmentGeometry.geometry.triangles.vertexData = VertexBufferDeviceAddress;
-		SegmentGeometry.geometry.triangles.maxVertex = Segment.MaxVertices;
-		SegmentGeometry.geometry.triangles.vertexStride = Segment.VertexBufferStride;
-		SegmentGeometry.geometry.triangles.indexData = IndexBufferDeviceAddress;
-
-		switch (Segment.VertexBufferElementType)
-		{
-		case VET_Float3:
-		case VET_Float4:
-			SegmentGeometry.geometry.triangles.vertexFormat = VK_FORMAT_R32G32B32_SFLOAT;
-			break;
-		default:
-			checkNoEntry();
-			break;
-		}
-
-		// No support for segment transform
-		SegmentGeometry.geometry.triangles.transformData.deviceAddress = 0;
-		SegmentGeometry.geometry.triangles.transformData.hostAddress = nullptr;
-
 		uint32 PrimitiveOffset = 0;
-
-		if (IndexBufferRHI)
+		switch (GeometryType)
 		{
-			SegmentGeometry.geometry.triangles.indexType = (IndexStrideInBytes == 2) ? VK_INDEX_TYPE_UINT16 : VK_INDEX_TYPE_UINT32;
-			// offset in bytes into the index buffer where primitive data for the current segment is defined
-			PrimitiveOffset = Segment.FirstPrimitive * IndicesPerPrimitive * IndexStrideInBytes;
-		}
-		else
-		{
-			SegmentGeometry.geometry.triangles.indexType = VK_INDEX_TYPE_NONE_KHR;
-			// for non-indexed geometry, primitiveOffset is applied when reading from vertex buffer
-			PrimitiveOffset = Segment.FirstPrimitive * IndicesPerPrimitive * Segment.VertexBufferStride;
+			case RTGT_Triangles:
+				SegmentGeometry.geometryType = VK_GEOMETRY_TYPE_TRIANGLES_KHR;
+
+				SegmentGeometry.geometry.triangles.sType = VK_STRUCTURE_TYPE_ACCELERATION_STRUCTURE_GEOMETRY_TRIANGLES_DATA_KHR;
+				SegmentGeometry.geometry.triangles.vertexFormat = VK_FORMAT_R32G32B32_SFLOAT;
+				SegmentGeometry.geometry.triangles.vertexData = VertexBufferDeviceAddress;
+				SegmentGeometry.geometry.triangles.maxVertex = Segment.MaxVertices;
+				SegmentGeometry.geometry.triangles.vertexStride = Segment.VertexBufferStride;
+				SegmentGeometry.geometry.triangles.indexData = IndexBufferDeviceAddress;
+
+				switch (Segment.VertexBufferElementType)
+				{
+				case VET_Float3:
+				case VET_Float4:
+					SegmentGeometry.geometry.triangles.vertexFormat = VK_FORMAT_R32G32B32_SFLOAT;
+					break;
+				default:
+					checkNoEntry();
+					break;
+				}
+
+				// No support for segment transform
+				SegmentGeometry.geometry.triangles.transformData.deviceAddress = 0;
+				SegmentGeometry.geometry.triangles.transformData.hostAddress = nullptr;
+
+				if (IndexBufferRHI)
+				{
+					SegmentGeometry.geometry.triangles.indexType = (IndexStrideInBytes == 2) ? VK_INDEX_TYPE_UINT16 : VK_INDEX_TYPE_UINT32;
+					// offset in bytes into the index buffer where primitive data for the current segment is defined
+					PrimitiveOffset = Segment.FirstPrimitive * IndicesPerPrimitive * IndexStrideInBytes;
+				}
+				else
+				{
+					SegmentGeometry.geometry.triangles.indexType = VK_INDEX_TYPE_NONE_KHR;
+					// for non-indexed geometry, primitiveOffset is applied when reading from vertex buffer
+					PrimitiveOffset = Segment.FirstPrimitive * IndicesPerPrimitive * Segment.VertexBufferStride;
+				}
+
+				break;
+			case RTGT_Procedural:
+				checkf(Segment.VertexBufferStride >= (2 * sizeof(FVector3f)), TEXT("Procedural geometry vertex buffer must contain at least 2xFloat3 that defines 3D bounding boxes of primitives."));
+				checkf(Segment.VertexBufferStride % 8 == 0, TEXT("Procedural geometry vertex buffer stride must be a multiple of 8."));
+
+				SegmentGeometry.geometryType = VK_GEOMETRY_TYPE_AABBS_KHR;
+				
+				SegmentGeometry.geometry.aabbs.sType = VK_STRUCTURE_TYPE_ACCELERATION_STRUCTURE_GEOMETRY_AABBS_DATA_KHR;
+				SegmentGeometry.geometry.aabbs.data = VertexBufferDeviceAddress;
+				SegmentGeometry.geometry.aabbs.stride = Segment.VertexBufferStride;
+
+				break;
+			default:
+				checkf(false, TEXT("Unexpected ray tracing geometry type"));
+				break;
 		}
 
 		BuildData.Segments.Add(SegmentGeometry);
@@ -293,9 +312,6 @@ FVulkanRayTracingGeometry::FVulkanRayTracingGeometry(ENoInit)
 FVulkanRayTracingGeometry::FVulkanRayTracingGeometry(const FRayTracingGeometryInitializer& InInitializer, FVulkanDevice* InDevice)
 	: FRHIRayTracingGeometry(InInitializer), Device(InDevice)
 {
-	// Only supporting triangles initially
-	check(Initializer.GeometryType == ERayTracingGeometryType::RTGT_Triangles);
-
 	uint32 IndexBufferStride = 0;
 	if (Initializer.IndexBuffer)
 	{
@@ -620,6 +636,7 @@ FRayTracingAccelerationStructureSize FVulkanDynamicRHI::RHICalcRayTracingGeometr
 	GetBLASBuildData(
 		Device->GetInstanceHandle(),
 		MakeArrayView(Initializer.Segments),
+		Initializer.GeometryType,
 		Initializer.IndexBuffer,
 		Initializer.IndexBufferOffset,
 		Initializer.bFastBuild,
@@ -717,6 +734,7 @@ void FVulkanCommandListContext::RHIBuildAccelerationStructures(const TArrayView<
 		GetBLASBuildData(
 			Device->GetInstanceHandle(),
 			MakeArrayView(Geometry->Initializer.Segments),
+			Geometry->Initializer.GeometryType,
 			Geometry->Initializer.IndexBuffer,
 			Geometry->Initializer.IndexBufferOffset,
 			Geometry->Initializer.bFastBuild,
