VulkanRHI: Use vk_enum_string_helper to print enum names in debug messages.  Add missing namespace to Vulkan API calls in VulkanRaytracing.cpp.

diff --git a/Engine/Source/Runtime/VulkanRHI/Private/Linux/VulkanLinuxPlatform.h b/Engine/Source/Runtime/VulkanRHI/Private/Linux/VulkanLinuxPlatform.h
--- a/Engine/Source/Runtime/VulkanRHI/Private/Linux/VulkanLinuxPlatform.h
+++ b/Engine/Source/Runtime/VulkanRHI/Private/Linux/VulkanLinuxPlatform.h
@@ -21,6 +21,12 @@
 #	define UE_VK_API_VERSION						VK_API_VERSION_1_1
 #endif // VULKAN_RHI_RAYTRACING
 
+#if UE_BUILD_DEBUG || UE_BUILD_DEVELOPMENT
+#	include "vk_enum_string_helper.h"
+#	define VK_TYPE_TO_STRING(Type, Value) ANSI_TO_TCHAR(string_##Type(Value))
+#	define VK_FLAGS_TO_STRING(Type, Value) ANSI_TO_TCHAR(string_##Type(Value).c_str())
+#endif
+
 #define ENUM_VK_ENTRYPOINTS_PLATFORM_BASE(EnumMacro)
 
 #define ENUM_VK_ENTRYPOINTS_PLATFORM_INSTANCE(EnumMacro) \
diff --git a/Engine/Source/Runtime/VulkanRHI/Private/VulkanDescriptorSets.h b/Engine/Source/Runtime/VulkanRHI/Private/VulkanDescriptorSets.h
--- a/Engine/Source/Runtime/VulkanRHI/Private/VulkanDescriptorSets.h
+++ b/Engine/Source/Runtime/VulkanRHI/Private/VulkanDescriptorSets.h
@@ -1139,14 +1139,14 @@ protected:
 		if (DescriptorType == VK_DESCRIPTOR_TYPE_STORAGE_BUFFER)
 		{
 			checkf(WriteDescriptors[DescriptorIndex].descriptorType == VK_DESCRIPTOR_TYPE_STORAGE_BUFFER || WriteDescriptors[DescriptorIndex].descriptorType == VK_DESCRIPTOR_TYPE_STORAGE_BUFFER_DYNAMIC,
-				TEXT("DescriptorType mismatch at index %d: called WriteBuffer<%d> and was expecting %d."), 
-				DescriptorIndex, (uint32)DescriptorType, (uint32)WriteDescriptors[DescriptorIndex].descriptorType);
+				TEXT("DescriptorType mismatch at index %d: called WriteBuffer<%s> and was expecting %s."), 
+				DescriptorIndex, VK_TYPE_TO_STRING(VkDescriptorType, DescriptorType), VK_TYPE_TO_STRING(VkDescriptorType, WriteDescriptors[DescriptorIndex].descriptorType));
 		}
 		else
 		{
 			checkf(WriteDescriptors[DescriptorIndex].descriptorType == DescriptorType,
-				TEXT("DescriptorType mismatch at index %d: called WriteBuffer<%d> and was expecting %d."),
-				DescriptorIndex, (uint32)DescriptorType, (uint32)WriteDescriptors[DescriptorIndex].descriptorType);
+				TEXT("DescriptorType mismatch at index %d: called WriteBuffer<%s> and was expecting %s."),
+				DescriptorIndex, VK_TYPE_TO_STRING(VkDescriptorType, DescriptorType), VK_TYPE_TO_STRING(VkDescriptorType, WriteDescriptors[DescriptorIndex].descriptorType));
 		}
 		VkDescriptorBufferInfo* BufferInfo = const_cast<VkDescriptorBufferInfo*>(WriteDescriptors[DescriptorIndex].pBufferInfo);
 		check(BufferInfo);
@@ -1206,13 +1206,14 @@ protected:
 				  Layout == VK_IMAGE_LAYOUT_DEPTH_ATTACHMENT_STENCIL_READ_ONLY_OPTIMAL ||
 				  Layout == VK_IMAGE_LAYOUT_DEPTH_STENCIL_READ_ONLY_OPTIMAL || 
 				  Layout == VK_IMAGE_LAYOUT_SHADER_READ_ONLY_OPTIMAL ||
-				  Layout == VK_IMAGE_LAYOUT_GENERAL, TEXT("Invalid Layout %d, Index %d, Type %d\n"), Layout, DescriptorIndex, WriteDescriptors[DescriptorIndex].descriptorType);
+				  Layout == VK_IMAGE_LAYOUT_GENERAL, TEXT("Invalid Layout %s, Index %d, Type %s\n"), 
+				VK_TYPE_TO_STRING(VkImageLayout, Layout), DescriptorIndex, VK_TYPE_TO_STRING(VkDescriptorType, WriteDescriptors[DescriptorIndex].descriptorType));
 		}
 		else
 		{
 			checkf(WriteDescriptors[DescriptorIndex].descriptorType == DescriptorType,
-				TEXT("DescriptorType mismatch at index %d: called WriteTextureView<%d> and was expecting %d."),
-				DescriptorIndex, (uint32)DescriptorType, (uint32)WriteDescriptors[DescriptorIndex].descriptorType);
+				TEXT("DescriptorType mismatch at index %d: called WriteTextureView<%s> and was expecting %s."),
+				DescriptorIndex, VK_TYPE_TO_STRING(VkDescriptorType, DescriptorType), VK_TYPE_TO_STRING(VkDescriptorType, WriteDescriptors[DescriptorIndex].descriptorType));
 		}
 		VkDescriptorImageInfo* ImageInfo = const_cast<VkDescriptorImageInfo*>(WriteDescriptors[DescriptorIndex].pImageInfo);
 		check(ImageInfo);
diff --git a/Engine/Source/Runtime/VulkanRHI/Private/VulkanRayTracing.cpp b/Engine/Source/Runtime/VulkanRHI/Private/VulkanRayTracing.cpp
--- a/Engine/Source/Runtime/VulkanRHI/Private/VulkanRayTracing.cpp
+++ b/Engine/Source/Runtime/VulkanRHI/Private/VulkanRayTracing.cpp
@@ -103,7 +103,7 @@ static VkDeviceAddress GetDeviceAddress(VkDevice Device, VkBuffer Buffer)
 	VkBufferDeviceAddressInfoKHR DeviceAddressInfo;
 	ZeroVulkanStruct(DeviceAddressInfo, VK_STRUCTURE_TYPE_BUFFER_DEVICE_ADDRESS_INFO);
 	DeviceAddressInfo.buffer = Buffer;
-	return vkGetBufferDeviceAddressKHR(Device, &DeviceAddressInfo);
+	return VulkanDynamicAPI::vkGetBufferDeviceAddressKHR(Device, &DeviceAddressInfo);
 }
 
 VkDeviceAddress FVulkanResourceMultiBuffer::GetDeviceAddress() const
@@ -115,7 +115,7 @@ VkDeviceAddress FVulkanResourceMultiBuffer::GetDeviceAddress() const
 static uint32 FindMemoryType(VkPhysicalDevice Gpu, uint32 Filter, VkMemoryPropertyFlags RequestedProperties)
 {
 	VkPhysicalDeviceMemoryProperties Properties = {};
-	vkGetPhysicalDeviceMemoryProperties(Gpu, &Properties);
+	VulkanRHI::vkGetPhysicalDeviceMemoryProperties(Gpu, &Properties);
 
 	uint32 Result = UINT32_MAX;
 	for (uint32 i = 0; i < Properties.memoryTypeCount; ++i)
@@ -151,8 +151,8 @@ void FVulkanRayTracingAllocator::Allocate(FVulkanDevice* Device, VkDeviceSize Si
 	MemoryAllocateInfo.pNext = &MemoryAllocateFlagsInfo;
 	MemoryAllocateInfo.allocationSize = MemoryRequirements.size;
 	MemoryAllocateInfo.memoryTypeIndex = FindMemoryType(Gpu, MemoryRequirements.memoryTypeBits, MemoryFlags);
-	VERIFYVULKANRESULT(vkAllocateMemory(DeviceHandle, &MemoryAllocateInfo, VULKAN_CPU_ALLOCATOR, &Result.Memory));
-	VERIFYVULKANRESULT(vkBindBufferMemory(DeviceHandle, Result.Buffer, Result.Memory, 0));
+	VERIFYVULKANRESULT(VulkanRHI::vkAllocateMemory(DeviceHandle, &MemoryAllocateInfo, VULKAN_CPU_ALLOCATOR, &Result.Memory));
+	VERIFYVULKANRESULT(VulkanRHI::vkBindBufferMemory(DeviceHandle, Result.Buffer, Result.Memory, 0));
 
 	Result.Device = DeviceHandle;
 }
@@ -162,12 +162,12 @@ void FVulkanRayTracingAllocator::Free(FVkRtAllocation& Allocation)
 {
 	if (Allocation.Buffer != VK_NULL_HANDLE)
 	{
-		vkDestroyBuffer(Allocation.Device, Allocation.Buffer, VULKAN_CPU_ALLOCATOR);
+		VulkanRHI::vkDestroyBuffer(Allocation.Device, Allocation.Buffer, VULKAN_CPU_ALLOCATOR);
 		Allocation.Buffer = VK_NULL_HANDLE;
 	}
 	if (Allocation.Memory != VK_NULL_HANDLE)
 	{
-		vkFreeMemory(Allocation.Device, Allocation.Memory, VULKAN_CPU_ALLOCATOR);
+		VulkanRHI::vkFreeMemory(Allocation.Device, Allocation.Memory, VULKAN_CPU_ALLOCATOR);
 		Allocation.Memory = VK_NULL_HANDLE;
 	}
 }
@@ -298,7 +298,7 @@ static void GetBLASBuildData(
 	BuildData.GeometryInfo.geometryCount = BuildData.Segments.Num();
 	BuildData.GeometryInfo.pGeometries = BuildData.Segments.GetData();
 
-	vkGetAccelerationStructureBuildSizesKHR(
+	VulkanDynamicAPI::vkGetAccelerationStructureBuildSizesKHR(
 		Device,
 		VK_ACCELERATION_STRUCTURE_BUILD_TYPE_DEVICE_KHR,
 		&BuildData.GeometryInfo,
@@ -345,12 +345,12 @@ FVulkanRayTracingGeometry::FVulkanRayTracingGeometry(const FRayTracingGeometryIn
 	CreateInfo.offset = AccelerationStructureBuffer->GetOffset();
 	CreateInfo.size = SizeInfo.ResultSize;
 	CreateInfo.type = VK_ACCELERATION_STRUCTURE_TYPE_BOTTOM_LEVEL_KHR;
-	VERIFYVULKANRESULT(vkCreateAccelerationStructureKHR(NativeDevice, &CreateInfo, VULKAN_CPU_ALLOCATOR, &Handle));
+	VERIFYVULKANRESULT(VulkanDynamicAPI::vkCreateAccelerationStructureKHR(NativeDevice, &CreateInfo, VULKAN_CPU_ALLOCATOR, &Handle));
 	
 	VkAccelerationStructureDeviceAddressInfoKHR DeviceAddressInfo;
 	ZeroVulkanStruct(DeviceAddressInfo, VK_STRUCTURE_TYPE_ACCELERATION_STRUCTURE_DEVICE_ADDRESS_INFO_KHR);
 	DeviceAddressInfo.accelerationStructure = Handle;
-	Address = vkGetAccelerationStructureDeviceAddressKHR(NativeDevice, &DeviceAddressInfo);
+	Address = VulkanDynamicAPI::vkGetAccelerationStructureDeviceAddressKHR(NativeDevice, &DeviceAddressInfo);
 }
 
 FVulkanRayTracingGeometry::~FVulkanRayTracingGeometry()
@@ -399,7 +399,7 @@ static void GetTLASBuildData(
 	BuildData.GeometryInfo.geometryCount = 1;
 	BuildData.GeometryInfo.pGeometries = &BuildData.Geometry;
 
-	vkGetAccelerationStructureBuildSizesKHR(
+	VulkanDynamicAPI::vkGetAccelerationStructureBuildSizesKHR(
 		Device,
 		VK_ACCELERATION_STRUCTURE_BUILD_TYPE_DEVICE_KHR,
 		&BuildData.GeometryInfo,
@@ -534,7 +534,7 @@ void FVulkanRayTracingScene::BuildAccelerationStructure(
 
 	FVulkanCommandBufferManager& CommandBufferManager = *CommandContext.GetCommandBufferManager();
 	FVulkanCmdBuffer* const CmdBuffer = CommandBufferManager.GetActiveCmdBuffer();
-	vkCmdBuildAccelerationStructuresKHR(CmdBuffer->GetHandle(), 1, &BuildData.GeometryInfo, &pBuildRanges);
+	VulkanDynamicAPI::vkCmdBuildAccelerationStructuresKHR(CmdBuffer->GetHandle(), 1, &BuildData.GeometryInfo, &pBuildRanges);
 
 	CommandBufferManager.SubmitActiveCmdBuffer();
 	CommandBufferManager.PrepareForNewActiveCommandBuffer();
@@ -764,7 +764,7 @@ void FVulkanCommandListContext::RHIBuildAccelerationStructures(const TArrayView<
 	}
 	
 	FVulkanCmdBuffer* const CmdBuffer = CommandBufferManager->GetActiveCmdBuffer();
-	vkCmdBuildAccelerationStructuresKHR(CmdBuffer->GetHandle(), Params.Num(), BuildGeometryInfos.GetData(), BuildRangeInfos.GetData());
+	VulkanDynamicAPI::vkCmdBuildAccelerationStructuresKHR(CmdBuffer->GetHandle(), Params.Num(), BuildGeometryInfos.GetData(), BuildRangeInfos.GetData());
 
 	CommandBufferManager->SubmitActiveCmdBuffer();
 	CommandBufferManager->PrepareForNewActiveCommandBuffer();
@@ -937,7 +937,7 @@ FVulkanRayTracingPipelineState::FVulkanRayTracingPipelineState(FVulkanDevice* co
 	RayTracingPipelineCreateInfo.maxPipelineRayRecursionDepth = 1;
 	RayTracingPipelineCreateInfo.layout = Layout->GetPipelineLayout();
 	
-	VERIFYVULKANRESULT(vkCreateRayTracingPipelinesKHR(
+	VERIFYVULKANRESULT(VulkanDynamicAPI::vkCreateRayTracingPipelinesKHR(
 		InDevice->GetInstanceHandle(), 
 		VK_NULL_HANDLE, // Deferred Operation 
 		VK_NULL_HANDLE, // Pipeline Cache 
@@ -959,7 +959,7 @@ FVulkanRayTracingPipelineState::FVulkanRayTracingPipelineState(FVulkanDevice* co
 
 	TArray<uint8> ShaderHandleStorage;
 	ShaderHandleStorage.AddUninitialized(SBTSize);
-	VERIFYVULKANRESULT(vkGetRayTracingShaderGroupHandlesKHR(InDevice->GetInstanceHandle(), Pipeline, 0, GroupCount, SBTSize, ShaderHandleStorage.GetData()));
+	VERIFYVULKANRESULT(VulkanDynamicAPI::vkGetRayTracingShaderGroupHandlesKHR(InDevice->GetInstanceHandle(), Pipeline, 0, GroupCount, SBTSize, ShaderHandleStorage.GetData()));
 
 	auto CopyHandlesToSBT = [InDevice, HandleSize, ShaderHandleStorage](FVkRtAllocation& Allocation, uint32 Offset)
 	{
@@ -971,11 +971,11 @@ FVulkanRayTracingPipelineState::FVulkanRayTracingPipelineState(FVulkanDevice* co
 			Allocation);
 
 		void* pMappedBufferMemory = nullptr;
-		VERIFYVULKANRESULT(vkMapMemory(InDevice->GetInstanceHandle(), Allocation.Memory, 0, VK_WHOLE_SIZE, 0, &pMappedBufferMemory));
+		VERIFYVULKANRESULT(VulkanRHI::vkMapMemory(InDevice->GetInstanceHandle(), Allocation.Memory, 0, VK_WHOLE_SIZE, 0, &pMappedBufferMemory));
 		{
 			FMemory::Memcpy(pMappedBufferMemory, ShaderHandleStorage.GetData() + Offset, HandleSize);
 		}
-		vkUnmapMemory(InDevice->GetInstanceHandle(), Allocation.Memory);
+		VulkanRHI::vkUnmapMemory(InDevice->GetInstanceHandle(), Allocation.Memory);
 	};
 
 	CopyHandlesToSBT(RayGenShaderBindingTable, 0);
diff --git a/Engine/Source/Runtime/VulkanRHI/Private/Windows/VulkanWindowsPlatform.h b/Engine/Source/Runtime/VulkanRHI/Private/Windows/VulkanWindowsPlatform.h
--- a/Engine/Source/Runtime/VulkanRHI/Private/Windows/VulkanWindowsPlatform.h
+++ b/Engine/Source/Runtime/VulkanRHI/Private/Windows/VulkanWindowsPlatform.h
@@ -25,6 +25,12 @@
 #	define UE_VK_API_VERSION						VK_API_VERSION_1_1
 #endif // VULKAN_RHI_RAYTRACING
 
+#if UE_BUILD_DEBUG || UE_BUILD_DEVELOPMENT
+#	include "vk_enum_string_helper.h"
+#	define VK_TYPE_TO_STRING(Type, Value) ANSI_TO_TCHAR(string_##Type(Value))
+#	define VK_FLAGS_TO_STRING(Type, Value) ANSI_TO_TCHAR(string_##Type(Value).c_str())
+#endif
+
 // 32-bit windows has warnings on custom mem mgr callbacks
 #define VULKAN_SHOULD_USE_LLM					(UE_BUILD_DEBUG || UE_BUILD_DEVELOPMENT) && !PLATFORM_32BITS
 
diff --git a/Engine/Source/Runtime/VulkanRHI/Public/VulkanConfiguration.h b/Engine/Source/Runtime/VulkanRHI/Public/VulkanConfiguration.h
--- a/Engine/Source/Runtime/VulkanRHI/Public/VulkanConfiguration.h
+++ b/Engine/Source/Runtime/VulkanRHI/Public/VulkanConfiguration.h
@@ -372,6 +372,14 @@ VULKANRHI_API DECLARE_LOG_CATEGORY_EXTERN(LogVulkanRHI, Log, All);
 #endif
 
 
+#ifndef VK_TYPE_TO_STRING
+#	define VK_TYPE_TO_STRING(Type, Value) *FString::Printf(TEXT("%u"), (uint32)Value)
+#endif
+#ifndef VK_FLAGS_TO_STRING
+#	define VK_FLAGS_TO_STRING(Type, Value) *FString::Printf(TEXT("%u"), (uint32)Value)
+#endif
+
+
 namespace VulkanRHI
 {
 	static FORCEINLINE const VkAllocationCallbacks* GetMemoryAllocator(const VkAllocationCallbacks* Allocator)
