Fixed: Occasional crashes in FVulkanViewport::DoAcquireImageIndex on app backgrounding
Use a 'dummy' backbuffer surface to gracefully handle 'surface lost' events

diff --git a/Engine/Source/Runtime/VulkanRHI/Private/Android/VulkanAndroidPlatform.h b/Engine/Source/Runtime/VulkanRHI/Private/Android/VulkanAndroidPlatform.h
--- a/Engine/Source/Runtime/VulkanRHI/Private/Android/VulkanAndroidPlatform.h
+++ b/Engine/Source/Runtime/VulkanRHI/Private/Android/VulkanAndroidPlatform.h
@@ -131,6 +131,9 @@ public:
 
 	static void DestroySwapchainKHR(VkDevice Device, VkSwapchainKHR Swapchain, const VkAllocationCallbacks* Allocator);
 
+	// Do not attempt to immediately recreate swapchain
+	static bool RecreateSwapchainOnFail() { return false; }
+
 	static VkFormat GetPlatform5551FormatWithFallback(VkFormat& OutFallbackFormat0, VkFormat& OutFallbackFormat1) 
 	{ 
 		OutFallbackFormat0 = VK_FORMAT_A1R5G5B5_UNORM_PACK16; 
diff --git a/Engine/Source/Runtime/VulkanRHI/Private/VulkanGenericPlatform.h b/Engine/Source/Runtime/VulkanRHI/Private/VulkanGenericPlatform.h
--- a/Engine/Source/Runtime/VulkanRHI/Private/VulkanGenericPlatform.h
+++ b/Engine/Source/Runtime/VulkanRHI/Private/VulkanGenericPlatform.h
@@ -107,6 +107,9 @@ public:
 	// Allow platforms to track swapchain destruction
 	static void DestroySwapchainKHR(VkDevice Device, VkSwapchainKHR Swapchain, const VkAllocationCallbacks* Allocator);
 
+	// Whether to attempt recreate swapchain when present or acqurire operations fail
+	static bool RecreateSwapchainOnFail() { return true; }
+
 	// Ensure the last frame completed on the GPU
 	static bool RequiresWaitingForFrameCompletionEvent() { return true; }
 
diff --git a/Engine/Source/Runtime/VulkanRHI/Private/VulkanViewport.cpp b/Engine/Source/Runtime/VulkanRHI/Private/VulkanViewport.cpp
--- a/Engine/Source/Runtime/VulkanRHI/Private/VulkanViewport.cpp
+++ b/Engine/Source/Runtime/VulkanRHI/Private/VulkanViewport.cpp
@@ -60,37 +60,40 @@ void FVulkanBackBuffer::OnLayoutTransition(FVulkanCommandListContext& Context, V
 void FVulkanBackBuffer::AcquireBackBufferImage(FVulkanCommandListContext& Context)
 {
 	check(Viewport);
-	if (Image == VK_NULL_HANDLE)
+	if (Viewport->AcquiredImageIndex < 0)
 	{
-		check(Viewport->AcquiredImageIndex == -1); //-V595
-		
-		Viewport->AcquireImageIndex(); //-V595
-		// If swapchain got invalidated (OUT_OF_DATE etc) in the above call, we may end up not having a valid viewport pointer at this point. Abort the whole thing.
-		if (Viewport == nullptr)
+		if (Viewport->TryAcquireImageIndex())
 		{
-			return;
-		}
+			int32 AcquiredImageIndex = Viewport->AcquiredImageIndex;
+			check(AcquiredImageIndex >= 0 && AcquiredImageIndex < Viewport->TextureViews.Num());
 
-		int32 AcquiredImageIndex = Viewport->AcquiredImageIndex;
-		check(AcquiredImageIndex >= 0 && AcquiredImageIndex < Viewport->TextureViews.Num());
+			FVulkanTextureView& ImageView = Viewport->TextureViews[AcquiredImageIndex];
 
-		FVulkanTextureView& ImageView = Viewport->TextureViews[AcquiredImageIndex];
+			Image = ImageView.Image;
+			DefaultView.View = ImageView.View;
+			DefaultView.ViewId = ImageView.ViewId;
 
-		Image = ImageView.Image;
-		DefaultView.View = ImageView.View;
-		DefaultView.ViewId = ImageView.ViewId;
+			// right after acquiring image is in undefined state
+			FVulkanLayoutManager& LayoutMgr = Context.GetLayoutManager();
+			VkImageLayout& CurrentLayout = LayoutMgr.FindOrAddLayoutRW(ImageView.Image, VK_IMAGE_LAYOUT_UNDEFINED, 1, 1);
+			CurrentLayout = VK_IMAGE_LAYOUT_UNDEFINED;
 
-		// right after acquiring image is in undefined state
-		FVulkanLayoutManager& LayoutMgr = Context.GetLayoutManager();
-		VkImageLayout& CurrentLayout = LayoutMgr.FindOrAddLayoutRW(ImageView.Image, VK_IMAGE_LAYOUT_UNDEFINED, 1, 1);
-		CurrentLayout = VK_IMAGE_LAYOUT_UNDEFINED;
-		
-		FVulkanCommandBufferManager* CmdBufferManager = Context.GetCommandBufferManager();
-		FVulkanCmdBuffer* CmdBuffer = CmdBufferManager->GetActiveCmdBuffer();
-		check(!CmdBuffer->IsInsideRenderPass());
-			
-		// Wait for semaphore signal before writing to backbuffer image
-		CmdBuffer->AddWaitSemaphore(VK_PIPELINE_STAGE_COLOR_ATTACHMENT_OUTPUT_BIT, Viewport->AcquiredSemaphore);
+			FVulkanCommandBufferManager* CmdBufferManager = Context.GetCommandBufferManager();
+			FVulkanCmdBuffer* CmdBuffer = CmdBufferManager->GetActiveCmdBuffer();
+			check(!CmdBuffer->IsInsideRenderPass());
+
+			// Wait for semaphore signal before writing to backbuffer image
+			CmdBuffer->AddWaitSemaphore(VK_PIPELINE_STAGE_COLOR_ATTACHMENT_OUTPUT_BIT, Viewport->AcquiredSemaphore);
+		}
+		else
+		{
+			// fallback to a 'dummy' backbuffer
+			check(Viewport->RenderingBackBuffer);
+			const FVulkanTextureView& DummyView = Viewport->RenderingBackBuffer->DefaultView;
+			Image = DummyView.Image;
+			DefaultView.View = DummyView.View;
+			DefaultView.ViewId = DummyView.ViewId;
+		}
 	}
 }
 
@@ -168,14 +171,9 @@ FVulkanViewport::~FVulkanViewport()
 	RHI->Viewports.Remove(this);
 }
 
-int32 FVulkanViewport::DoAcquireImageIndex(FVulkanViewport* Viewport)
-{
-	return Viewport->AcquiredImageIndex = Viewport->SwapChain->AcquireImageIndex(&Viewport->AcquiredSemaphore);
-}
-
 bool FVulkanViewport::DoCheckedSwapChainJob(TFunction<int32(FVulkanViewport*)> SwapChainJob)
 {
-	int32 AttemptsPending = 4;
+	int32 AttemptsPending = FVulkanPlatform::RecreateSwapchainOnFail() ? 4 : 0;
 	int32 Status = SwapChainJob(this);
 
 	while (Status < 0 && AttemptsPending > 0)
@@ -207,21 +205,12 @@ bool FVulkanViewport::DoCheckedSwapChainJob(TFunction<int32(FVulkanViewport*)> S
 	return Status >= 0;
 }
 
-void FVulkanViewport::AcquireImageIndex()
-{
-	if (!DoCheckedSwapChainJob(DoAcquireImageIndex))
-	{
-		UE_LOG(LogVulkanRHI, Fatal, TEXT("Swapchain acquire image index failed!"));
-	}
-	check(AcquiredImageIndex != -1);
-}
-
 bool FVulkanViewport::TryAcquireImageIndex()
 {
-	int NewImageIndex = DoAcquireImageIndex(this);
-	if (NewImageIndex != -1)
+	int32 Result = SwapChain->AcquireImageIndex(&AcquiredSemaphore);
+	if (Result >= 0)
 	{
-		AcquiredImageIndex = NewImageIndex;
+		AcquiredImageIndex = Result;
 		return true;
 	}
 
@@ -586,7 +575,7 @@ void FVulkanViewport::Tick(float DeltaTime)
 {
 	check(IsInGameThread());
 
-	if(SwapChain && FPlatformAtomics::AtomicRead(&LockToVsync) != SwapChain->DoesLockToVsync())
+	if (SwapChain && FPlatformAtomics::AtomicRead(&LockToVsync) != SwapChain->DoesLockToVsync())
 	{
 		FlushRenderingCommands();
 		ENQUEUE_RENDER_COMMAND(UpdateVsync)(
@@ -691,7 +680,7 @@ void FVulkanViewport::CreateSwapchain(FVulkanSwapChainRecreateInfo* RecreateInfo
 		}
 	}
 
-	if (!SupportsStandardSwapchain() || GVulkanDelayAcquireImage == EDelayAcquireImageType::DelayAcquire)
+	// We always create a 'dummy' backbuffer to gracefully handle SurfaceLost cases
 	{
 		uint32 BackBufferSizeX = RequiresRenderingBackBuffer() ? SizeX : 1;
 		uint32 BackBufferSizeY = RequiresRenderingBackBuffer() ? SizeY : 1;
@@ -854,12 +843,18 @@ bool FVulkanViewport::Present(FVulkanCommandListContext* Context, FVulkanCmdBuff
 		}
 		else
 		{
-			check(AcquiredImageIndex != -1);
-			check(RHIBackBuffer != nullptr && RHIBackBuffer->Image == BackBufferImages[AcquiredImageIndex]);
-
-			VkImageLayout& Layout = Context->GetLayoutManager().FindOrAddLayoutRW(BackBufferImages[AcquiredImageIndex], VK_IMAGE_LAYOUT_UNDEFINED, 1, 1);
-			VulkanSetImageLayout(CmdBuffer->GetHandle(), BackBufferImages[AcquiredImageIndex], Layout, VK_IMAGE_LAYOUT_PRESENT_SRC_KHR, FVulkanPipelineBarrier::MakeSubresourceRange(VK_IMAGE_ASPECT_COLOR_BIT));
-			Layout = VK_IMAGE_LAYOUT_UNDEFINED;
+			if (AcquiredImageIndex != -1)
+			{
+				check(RHIBackBuffer != nullptr && RHIBackBuffer->Image == BackBufferImages[AcquiredImageIndex]);
+				VkImageLayout& Layout = Context->GetLayoutManager().FindOrAddLayoutRW(BackBufferImages[AcquiredImageIndex], VK_IMAGE_LAYOUT_UNDEFINED, 1, 1);
+				VulkanSetImageLayout(CmdBuffer->GetHandle(), BackBufferImages[AcquiredImageIndex], Layout, VK_IMAGE_LAYOUT_PRESENT_SRC_KHR, FVulkanPipelineBarrier::MakeSubresourceRange(VK_IMAGE_ASPECT_COLOR_BIT));
+				Layout = VK_IMAGE_LAYOUT_UNDEFINED;
+			}
+			else
+			{
+				// When we have failed to acquire backbuffer image we fallback to using 'dummy' backbuffer
+				check(RHIBackBuffer != nullptr && RHIBackBuffer->Image == RenderingBackBuffer->Image);
+			}
 		}
 	}
 
@@ -876,8 +871,9 @@ bool FVulkanViewport::Present(FVulkanCommandListContext* Context, FVulkanCmdBuff
 				CmdBuffer->AddWaitSemaphore(VK_PIPELINE_STAGE_BOTTOM_OF_PIPE_BIT, AcquiredSemaphore);
 			}
 			
+			VulkanRHI::FSemaphore* SignalSemaphore = (AcquiredImageIndex >= 0 ? RenderingDoneSemaphores[AcquiredImageIndex] : nullptr);
 			// submit through the CommandBufferManager as it will add the proper semaphore
-			ImmediateCmdBufMgr->SubmitActiveCmdBufferFromPresent(RenderingDoneSemaphores[AcquiredImageIndex]);
+			ImmediateCmdBufMgr->SubmitActiveCmdBufferFromPresent(SignalSemaphore);
 		}
 		else
 		{
@@ -934,7 +930,7 @@ bool FVulkanViewport::Present(FVulkanCommandListContext* Context, FVulkanCmdBuff
 		};
 		if (SupportsStandardSwapchain() && !DoCheckedSwapChainJob(SwapChainJob))
 		{
-			UE_LOG(LogVulkanRHI, Fatal, TEXT("Swapchain present failed!"));
+			UE_LOG(LogVulkanRHI, Error, TEXT("Swapchain present failed!"));
 			bResult = false;
 		}
 		else
diff --git a/Engine/Source/Runtime/VulkanRHI/Private/VulkanViewport.h b/Engine/Source/Runtime/VulkanRHI/Private/VulkanViewport.h
--- a/Engine/Source/Runtime/VulkanRHI/Private/VulkanViewport.h
+++ b/Engine/Source/Runtime/VulkanRHI/Private/VulkanViewport.h
@@ -146,14 +146,12 @@ protected:
 
 	void CreateSwapchain(struct FVulkanSwapChainRecreateInfo* RecreateInfo);
 	void DestroySwapchain(struct FVulkanSwapChainRecreateInfo* RecreateInfo);
-	void AcquireImageIndex();
 	bool TryAcquireImageIndex();
 
 	void RecreateSwapchain(void* NewNativeWindow);
 	void RecreateSwapchainFromRT(EPixelFormat PreferredPixelFormat);
 	void Resize(uint32 InSizeX, uint32 InSizeY, bool bIsFullscreen, EPixelFormat PreferredPixelFormat);
 
-	static int32 DoAcquireImageIndex(FVulkanViewport* Viewport);
 	bool DoCheckedSwapChainJob(TFunction<int32(FVulkanViewport*)> SwapChainJob);
 	bool SupportsStandardSwapchain();
 	bool RequiresRenderingBackBuffer();
