VulkanRHI: Fix assert when texture is created from TransientHeap with Discard as initial state.  Also a minor clean up to remove unreachable code.

diff --git a/Engine/Source/Runtime/VulkanRHI/Private/VulkanTexture.cpp b/Engine/Source/Runtime/VulkanRHI/Private/VulkanTexture.cpp
--- a/Engine/Source/Runtime/VulkanRHI/Private/VulkanTexture.cpp
+++ b/Engine/Source/Runtime/VulkanRHI/Private/VulkanTexture.cpp
@@ -566,6 +566,7 @@ static VkImageLayout GetInitialLayoutFromRHIAccess(ERHIAccess RHIAccess, ETextur
 	switch (RHIAccess)
 	{
 		case ERHIAccess::Unknown:	return VK_IMAGE_LAYOUT_UNDEFINED;
+		case ERHIAccess::Discard:	return VK_IMAGE_LAYOUT_UNDEFINED;
 		case ERHIAccess::CopySrc:	return VK_IMAGE_LAYOUT_TRANSFER_SRC_OPTIMAL;
 		case ERHIAccess::CopyDest:	return VK_IMAGE_LAYOUT_TRANSFER_DST_OPTIMAL;
 	}
@@ -1568,7 +1569,7 @@ FVulkanTexture::FVulkanTexture(FVulkanDevice& InDevice, const FRHITextureCreateD
 			}
 
 			VkDevice VulkanDevice = InDevice.GetInstanceHandle();
-			VkMemoryPropertyFlags BufferMemFlags = (VK_MEMORY_PROPERTY_HOST_VISIBLE_BIT | VK_MEMORY_PROPERTY_HOST_COHERENT_BIT | VK_MEMORY_PROPERTY_HOST_CACHED_BIT);
+			const VkMemoryPropertyFlags BufferMemFlags = (VK_MEMORY_PROPERTY_HOST_VISIBLE_BIT | VK_MEMORY_PROPERTY_HOST_COHERENT_BIT | VK_MEMORY_PROPERTY_HOST_CACHED_BIT);
 
 			VkBufferCreateInfo BufferCreateInfo;
 			ZeroVulkanStruct(BufferCreateInfo, VK_STRUCTURE_TYPE_BUFFER_CREATE_INFO);
@@ -1618,16 +1619,15 @@ FVulkanTexture::FVulkanTexture(FVulkanDevice& InDevice, const FRHITextureCreateD
 
 		const bool bRenderTarget = EnumHasAnyFlags(InCreateDesc.Flags, TexCreate_RenderTargetable | TexCreate_DepthStencilTargetable | TexCreate_ResolveTargetable);
 		const bool bUAV = EnumHasAnyFlags(InCreateDesc.Flags, TexCreate_UAV);
-		const bool bCPUReadback = EnumHasAnyFlags(InCreateDesc.Flags, TexCreate_CPUReadback);
 		const bool bDynamic = EnumHasAnyFlags(InCreateDesc.Flags, TexCreate_Dynamic);
 		const bool bExternal = EnumHasAnyFlags(InCreateDesc.Flags, TexCreate_External);
 
-		VkMemoryPropertyFlags MemoryFlags = bCPUReadback ? VK_MEMORY_PROPERTY_HOST_VISIBLE_BIT : VK_MEMORY_PROPERTY_DEVICE_LOCAL_BIT;
+		VkMemoryPropertyFlags MemoryFlags = VK_MEMORY_PROPERTY_DEVICE_LOCAL_BIT;
 
 		bool bMemoryless = EnumHasAnyFlags(InCreateDesc.Flags, TexCreate_Memoryless) && InDevice.GetDeviceMemoryManager().SupportsMemoryless();
 		if (bMemoryless)
 		{
-			if (ensureMsgf(bRenderTarget, TEXT("Memoryless surfaces can only be used for render targets")) && ensureMsgf(!bCPUReadback, TEXT("Memoryless surfaces cannot be read back on CPU")))
+			if (ensureMsgf(bRenderTarget, TEXT("Memoryless surfaces can only be used for render targets")))
 			{
 				MemoryFlags |= VK_MEMORY_PROPERTY_LAZILY_ALLOCATED_BIT;
 			}
@@ -1639,7 +1639,6 @@ FVulkanTexture::FVulkanTexture(FVulkanDevice& InDevice, const FRHITextureCreateD
 
 		if (InTransientHeapAllocation != nullptr)
 		{
-			check(!bCPUReadback);
 			check(!bMemoryless);
 			check(InTransientHeapAllocation->Offset % MemoryRequirements.alignment == 0);
 			check(InTransientHeapAllocation->Size >= MemoryRequirements.size);
