VulkanRHI: Fix "CmdBuffer->IsSubmitted()" assert that was being hit by our EngineTests.  Basically repeats the fix from 19146645.

Code clean up, merge two barriers.

diff --git a/Engine/Source/Runtime/VulkanRHI/Private/VulkanRenderTarget.cpp b/Engine/Source/Runtime/VulkanRHI/Private/VulkanRenderTarget.cpp
--- a/Engine/Source/Runtime/VulkanRHI/Private/VulkanRenderTarget.cpp
+++ b/Engine/Source/Runtime/VulkanRHI/Private/VulkanRenderTarget.cpp
@@ -233,18 +233,17 @@ void FVulkanCommandListContext::RHICopyToResolveTarget(FRHITexture* SourceTextur
 		CopyRegion.imageExtent.depth = 1;
 		VulkanRHI::vkCmdCopyImageToBuffer(CmdBuffer->GetHandle(), SrcSurface.Image, VK_IMAGE_LAYOUT_TRANSFER_SRC_OPTIMAL, CpuReadbackBuffer->Buffer, 1, &CopyRegion);
 
-		{
-			FVulkanPipelineBarrier BarrierMemory;
-			BarrierMemory.AddMemoryBarrier(VK_ACCESS_TRANSFER_WRITE_BIT, VK_ACCESS_HOST_READ_BIT, VK_PIPELINE_STAGE_ALL_COMMANDS_BIT, VK_PIPELINE_STAGE_HOST_BIT);
-			BarrierMemory.Execute(CmdBuffer->GetHandle());
-		}
-		SrcCurrentAccess = ERHIAccess::CopySrc;
-		if(SrcCurrentAccess != InResolveParams.SourceAccessFinal && InResolveParams.SourceAccessFinal != ERHIAccess::Unknown)
 		{
 			FVulkanPipelineBarrier BarrierAfter;
-			BarrierAfter.AddImageAccessTransition(SrcSurface, SrcCurrentAccess, InResolveParams.SourceAccessFinal, FVulkanPipelineBarrier::MakeSubresourceRange(SrcSurface.GetFullAspectMask()), SrcLayout);
-			BarrierAfter.Execute(CmdBuffer->GetHandle());
+			BarrierAfter.AddMemoryBarrier(VK_ACCESS_TRANSFER_WRITE_BIT, VK_ACCESS_HOST_READ_BIT, VK_PIPELINE_STAGE_ALL_COMMANDS_BIT, VK_PIPELINE_STAGE_HOST_BIT);
+
+			SrcCurrentAccess = ERHIAccess::CopySrc;
+			if(SrcCurrentAccess != InResolveParams.SourceAccessFinal && InResolveParams.SourceAccessFinal != ERHIAccess::Unknown)
+			{
+				BarrierAfter.AddImageAccessTransition(SrcSurface, SrcCurrentAccess, InResolveParams.SourceAccessFinal, FVulkanPipelineBarrier::MakeSubresourceRange(SrcSurface.GetFullAspectMask()), SrcLayout);
+			}
 
+			BarrierAfter.Execute(CmdBuffer->GetHandle());
 		}
 	}
 	else
@@ -513,8 +512,14 @@ void FVulkanDynamicRHI::RHIMapStagingSurface(FRHITexture* TextureRHI, FRHIGPUFen
 	{
 		FRHICommandListExecutor::GetImmediateCommandList().ImmediateFlush(EImmediateFlushType::FlushRHIThread);
 		Device->SubmitCommandsAndFlushGPU();
-		FVulkanGPUFence* Fence = ResourceCast(FenceRHI);
-		Device->GetImmediateContext().GetCommandBufferManager()->WaitForCmdBuffer(Fence->GetCmdBuffer());
+
+		// SubmitCommandsAndFlushGPU might update fence state if it was tied to a previously submitted command buffer.
+		// Its state will have been updated from Submitted to NeedReset, and would assert in WaitForCmdBuffer (which is not needed in such a case)
+		if (!FenceRHI->Poll())
+		{
+			FVulkanGPUFence* Fence = ResourceCast(FenceRHI);
+			Device->GetImmediateContext().GetCommandBufferManager()->WaitForCmdBuffer(Fence->GetCmdBuffer());
+		}
 	}
 	else
 	{
