VulkanRHI: Make all platforms use the same VulkanSDK flow for compilation (prefer a local SDK only if it's newer) so that we can ensure the engine is always built with the provided version at minimum.

diff --git a/Engine/Source/Runtime/VulkanRHI/VulkanRHI.Build.cs b/Engine/Source/Runtime/VulkanRHI/VulkanRHI.Build.cs
--- a/Engine/Source/Runtime/VulkanRHI/VulkanRHI.Build.cs
+++ b/Engine/Source/Runtime/VulkanRHI/VulkanRHI.Build.cs
@@ -62,34 +62,15 @@ public class VulkanRHI : ModuleRules
 			AddEngineThirdPartyPrivateStaticDependencies(Target, "NVAftermath");
 		}
 
-		if (Target.Platform.IsInGroup(UnrealPlatformGroup.Windows) || Target.Platform == UnrealTargetPlatform.Android)
+		if (Target.Platform.IsInGroup(UnrealPlatformGroup.Windows) || Target.Platform == UnrealTargetPlatform.Android || Target.IsInPlatformGroup(UnrealPlatformGroup.Unix))
 		{
             AddEngineThirdPartyPrivateStaticDependencies(Target, "Vulkan");
         }
-        else if (Target.IsInPlatformGroup(UnrealPlatformGroup.Unix))
-		{
-			if (Target.IsInPlatformGroup(UnrealPlatformGroup.Linux))
-			{
-				PrivateDependencyModuleNames.Add("ApplicationCore");
-				AddEngineThirdPartyPrivateStaticDependencies(Target, "SDL2");
 
-				string VulkanSDKPath = Environment.GetEnvironmentVariable("VULKAN_SDK");
-				bool bSDKInstalled = !String.IsNullOrEmpty(VulkanSDKPath);
-				if (BuildHostPlatform.Current.Platform != UnrealTargetPlatform.Linux || !bSDKInstalled)
-				{
-					AddEngineThirdPartyPrivateStaticDependencies(Target, "Vulkan");
-				}
-				else
-				{
-					PrivateIncludePaths.Add(VulkanSDKPath + "/include");
-					PrivateIncludePaths.Add(VulkanSDKPath + "/include/vulkan");
-					PublicAdditionalLibraries.Add(Path.Combine(VulkanSDKPath, "lib", "libvulkan.so"));
-				}
-			}
-			else
-			{
-				AddEngineThirdPartyPrivateStaticDependencies(Target, "VkHeadersExternal");
-			}
+		if (Target.IsInPlatformGroup(UnrealPlatformGroup.Linux))
+		{
+			PrivateDependencyModuleNames.Add("ApplicationCore");
+			AddEngineThirdPartyPrivateStaticDependencies(Target, "SDL2");
 		}
         else if (Target.Platform == UnrealTargetPlatform.Mac)
         {
