Add support for copying from a cube map or cube map array into a cube map array using RHICopyToResolveTarget on Vulkan.

This fixes a crash when trying to render environment captures, which now uses that function to place the rendered faces into the array used by the scene.

#rnx

diff --git a/Engine/Source/Runtime/VulkanRHI/Private/VulkanRenderTarget.cpp b/Engine/Source/Runtime/VulkanRHI/Private/VulkanRenderTarget.cpp
--- a/Engine/Source/Runtime/VulkanRHI/Private/VulkanRenderTarget.cpp
+++ b/Engine/Source/Runtime/VulkanRHI/Private/VulkanRenderTarget.cpp
@@ -144,20 +144,22 @@ void FVulkanCommandListContext::RHICopyToResolveTarget(FRHITexture* SourceTextur
 	const FRHITextureDesc& SrcDesc = SrcSurface.GetDesc();
 	const FRHITextureDesc& DstDesc = DstSurface.GetDesc();
 
-	uint32 SrcNumLayers, DstNumLayers;
+	uint32 SrcNumLayers, DstNumLayers, SrcCubeFace = 0, DstCubeFace = 0;
 
 	switch (SourceTextureRHI->GetDesc().Dimension)
 	{
 	case ETextureDimension::Texture2D:
-		if (DestTextureRHI->GetDesc().Dimension == ETextureDimension::TextureCube)
+		SrcNumLayers = 1;
+		if (DestTextureRHI->GetDesc().Dimension == ETextureDimension::Texture2D)
 		{
-			SrcNumLayers = 1;
-			DstNumLayers = 6;
+			DstNumLayers = 1;
 		}
 		else
 		{
-			check(DestTextureRHI->GetDesc().Dimension == ETextureDimension::Texture2D);
-			SrcNumLayers = DstNumLayers = 1;
+			// Allow copying a 2D texture to a face of a cube texture.
+			check(DestTextureRHI->GetDesc().Dimension == ETextureDimension::TextureCube);
+			DstNumLayers = 6;
+			DstCubeFace = (uint32)InResolveParams.CubeFace;
 		}
 		break;
 
@@ -173,12 +175,27 @@ void FVulkanCommandListContext::RHICopyToResolveTarget(FRHITexture* SourceTextur
 		break;
 
 	case ETextureDimension::TextureCube:
-		check(DestTextureRHI->GetDesc().Dimension == ETextureDimension::TextureCube);
-		SrcNumLayers = DstNumLayers = 6;
+		SrcNumLayers = 6;
+		SrcCubeFace = DstCubeFace = (uint32)InResolveParams.CubeFace;
+		if (DestTextureRHI->GetDesc().Dimension == ETextureDimension::TextureCube)
+		{
+			DstNumLayers = 6;
+		}
+		else
+		{
+			// Allow copying a cube texture to a slice of a cube texture array.
+			check(DestTextureRHI->GetDesc().Dimension == ETextureDimension::TextureCubeArray);
+			DstNumLayers = 6 * DstDesc.ArraySize;
+		}
 		break;
 
 	case ETextureDimension::TextureCubeArray:
-		// Not supported. fall through.
+		check(DestTextureRHI->GetDesc().Dimension == ETextureDimension::TextureCubeArray);
+		SrcNumLayers = 6 * SrcDesc.ArraySize;
+		DstNumLayers = 6 * DstDesc.ArraySize;
+		SrcCubeFace = DstCubeFace = (uint32)InResolveParams.CubeFace;
+		break;
+
 	default:
 		checkNoEntry();
 		return;
@@ -188,14 +205,14 @@ void FVulkanCommandListContext::RHICopyToResolveTarget(FRHITexture* SourceTextur
 	SrcRange.aspectMask = SrcSurface.GetFullAspectMask();
 	SrcRange.baseMipLevel = InResolveParams.MipIndex;
 	SrcRange.levelCount = 1;
-	SrcRange.baseArrayLayer = InResolveParams.SourceArrayIndex * SrcNumLayers + (SrcNumLayers == 6 ? InResolveParams.CubeFace : 0);
+	SrcRange.baseArrayLayer = InResolveParams.SourceArrayIndex * SrcNumLayers + SrcCubeFace;
 	SrcRange.layerCount = 1;
 
 	VkImageSubresourceRange DstRange;
 	DstRange.aspectMask = DstSurface.GetFullAspectMask();
 	DstRange.baseMipLevel = InResolveParams.MipIndex;
 	DstRange.levelCount = 1;
-	DstRange.baseArrayLayer = InResolveParams.DestArrayIndex * DstNumLayers + (DstNumLayers == 6 ? InResolveParams.CubeFace : 0);
+	DstRange.baseArrayLayer = InResolveParams.DestArrayIndex * DstNumLayers + DstCubeFace;
 	DstRange.layerCount = 1;
 
 	ERHIAccess SrcCurrentAccess, DstCurrentAccess;
