Add option to Serialize precompile PSOs.
Fix possible crash when the LRU is used to evict shader modules.

diff --git a/Engine/Source/Runtime/VulkanRHI/Private/VulkanPipeline.cpp b/Engine/Source/Runtime/VulkanRHI/Private/VulkanPipeline.cpp
--- a/Engine/Source/Runtime/VulkanRHI/Private/VulkanPipeline.cpp
+++ b/Engine/Source/Runtime/VulkanRHI/Private/VulkanPipeline.cpp
@@ -134,6 +134,17 @@ static FAutoConsoleVariableRef GVulkanPSOForceSingleThreadedCVar(
 	ECVF_ReadOnly | ECVF_RenderThreadSafe
 );
 
+static int32 GVulkanPSOPrecompileForceSingleThreaded = 0;
+static FAutoConsoleVariableRef GVulkanPSOPrecompileForceSingleThreadedCVar(
+	TEXT("r.Vulkan.PrecompilePSOForceSingleThreaded"),
+	GVulkanPSOPrecompileForceSingleThreaded,
+	TEXT("Some drivers serialize PSO creation, this setting can help reduce the driver's PSO queue and can decrease the wait time for more important (non-precompile) PSOs.\n")
+	TEXT("0: (default) Allow Async precompile PSO creation.\n")
+	TEXT("1: force singlethreaded creation of precompile PSOs only.")
+	,
+	ECVF_ReadOnly | ECVF_RenderThreadSafe
+);
+
 static int32 GVulkanPSOLRUEvictAfterUnusedFrames = 0;
 static FAutoConsoleVariableRef GVulkanPSOLRUEvictAfterUnusedFramesCVar(
 	TEXT("r.Vulkan.PSOLRUEvictAfterUnusedFrames"),
@@ -251,7 +262,7 @@ FVulkanRHIGraphicsPipelineState::~FVulkanRHIGraphicsPipelineState()
 
 void FVulkanRHIGraphicsPipelineState::GetOrCreateShaderModules(FVulkanShader*const* Shaders)
 {
-	FScopeLock Lock(&FVulkanShaderHandleCS);
+	FScopeLock Lock(&FVulkanShaderHandleCS); // this lock is to stop concurrent access to VulkanShader's ShaderModules container. 
 	for (int32 Index = 0; Index < ShaderStage::NumStages; ++Index)
 	{
 		FVulkanShader* Shader = Shaders[Index];
@@ -1333,6 +1344,7 @@ void FVulkanPipelineStateCacheManager::DestroyCache()
 {
 	VkDevice DeviceHandle = Device->GetInstanceHandle();
 
+	QUICK_SCOPE_CYCLE_COUNTER(STAT_Vulkan_DestroyCache_PSOLock);
 	FScopeLock Lock1(&GraphicsPSOLockedCS);
 	int idx = 0;
 	for (auto& Pair : GraphicsPSOLockedMap)
@@ -1710,33 +1722,39 @@ void FVulkanPipelineStateCacheManager::NotifyDeletedGraphicsPSO(FRHIGraphicsPipe
 }
 
 
-//Global lock for PSO creation, only enabled if GVulkanPSOForceSingleThreaded is 1
-struct FPSOGlobalLock
+struct FPSOOptionalLock
 {
 	FCriticalSection* CriticalSection;
-	FPSOGlobalLock(FCriticalSection* InSynchObject)
+	FPSOOptionalLock(FCriticalSection* InSynchObject)
 	{
-		
-		CriticalSection = GVulkanPSOForceSingleThreaded ? InSynchObject : nullptr;
+		CriticalSection = InSynchObject;
 		if (CriticalSection)
 		{
 			CriticalSection->Lock();
 		}
 	}
-	~FPSOGlobalLock()
+	~FPSOOptionalLock()
 	{
 		if (CriticalSection)
 		{
 			CriticalSection->Unlock();
 		}
 	}
-
 };
 
+static FCriticalSection CreateGraphicsPSOMutex;
+
 FVulkanRHIGraphicsPipelineState* FVulkanPipelineStateCacheManager::RHICreateGraphicsPipelineState(const FGraphicsPipelineStateInitializer& Initializer)
 {
 	QUICK_SCOPE_CYCLE_COUNTER(STAT_Vulkan_RHICreateGraphicsPipelineState_NEW);
-	FPSOGlobalLock GlobalLock(&GraphicsPSOLockedCS);
+
+	// Optional lock for PSO creation, GVulkanPSOForceSingleThreaded is used to work around driver bugs.
+	// GVulkanPSOPrecompileForceSingleThreaded is used when the driver internally serializes PSO creation, this option reduces the driver queue size.
+	// We stall precompile PSOs which increases the likelihood for non-precompile PSO to jump the queue.
+	// Not using GraphicsPSOLockedCS as the create could take a long time on some platforms, holding GraphicsPSOLockedCS the whole time could cause hitching.
+	const bool bShouldLock = GVulkanPSOForceSingleThreaded || (GVulkanPSOPrecompileForceSingleThreaded && Initializer.bFromPSOFileCache);
+	FPSOOptionalLock PSOSingleThreadedLock(bShouldLock ? &CreateGraphicsPSOMutex : nullptr);
+
 	FVulkanPSOKey Key;
 	FGfxPipelineDesc Desc;
 	FVulkanDescriptorSetsLayoutInfo DescriptorSetLayoutInfo;
@@ -1844,6 +1862,7 @@ FVulkanRHIGraphicsPipelineState* FVulkanPipelineStateCacheManager::RHICreateGrap
 				GraphicsPSOLockedMap.Add(MoveTemp(Key), NewPSO);
 				if (bUseLRU && NewPSO->VulkanPipeline != VK_NULL_HANDLE)
 				{
+					QUICK_SCOPE_CYCLE_COUNTER(STAT_Vulkan_RHICreateGraphicsPipelineState_LRU_PSOLock);
 					// we add only created pipelines to the LRU
 					FScopeLock LockRU(&LRUCS);
 					NewPSO->bIsRegistered = true;
@@ -2207,7 +2226,7 @@ void FVulkanPipelineStateCacheManager::LRUTouch(FVulkanRHIGraphicsPipelineState*
 
 			GetVulkanShaders(Device, *PSO, VulkanShaders);
 
-			QUICK_SCOPE_CYCLE_COUNTER(STAT_Vulkan_RHICreateGraphicsPipelineState_CREATE_PART0);
+			QUICK_SCOPE_CYCLE_COUNTER(STAT_Vulkan_LRUMiss);
 
 			if (!CreateGfxPipelineFromEntry(PSO, VulkanShaders, &PSO->VulkanPipeline))
 			{
@@ -2287,6 +2306,7 @@ void FVulkanPipelineStateCacheManager::LRURemove(FVulkanRHIGraphicsPipelineState
 		PSO->DeleteVkPipeline(bImmediate);
 		if (GVulkanReleaseShaderModuleWhenEvictingPSO)
 		{
+			FScopeLock Lock(&FVulkanShaderHandleCS); // this lock is here to stop concurrent access to VulkanShaders ShaderModules container. 
 	        for (int ShaderStageIndex = 0; ShaderStageIndex < ShaderStage::NumStages; ShaderStageIndex++)
 	        {
 				if (PSO->VulkanShaders[ShaderStageIndex] != nullptr)
