Fixed: Client crash at FVulkanBackBuffer::AcquireBackBufferImage

diff --git a/Engine/Source/Runtime/VulkanRHI/Private/VulkanViewport.cpp b/Engine/Source/Runtime/VulkanRHI/Private/VulkanViewport.cpp
--- a/Engine/Source/Runtime/VulkanRHI/Private/VulkanViewport.cpp
+++ b/Engine/Source/Runtime/VulkanRHI/Private/VulkanViewport.cpp
@@ -60,7 +60,8 @@ void FVulkanBackBuffer::OnLayoutTransition(FVulkanCommandListContext& Context, V
 void FVulkanBackBuffer::AcquireBackBufferImage(FVulkanCommandListContext& Context)
 {
 	check(Viewport);
-	if (Viewport->AcquiredImageIndex < 0)
+	
+	if (Surface.Image == VK_NULL_HANDLE)
 	{
 		if (Viewport->TryAcquireImageIndex())
 		{
@@ -207,13 +208,15 @@ bool FVulkanViewport::DoCheckedSwapChainJob(TFunction<int32(FVulkanViewport*)> S
 
 bool FVulkanViewport::TryAcquireImageIndex()
 {
-	int32 Result = SwapChain->AcquireImageIndex(&AcquiredSemaphore);
-	if (Result >= 0)
+	if (SwapChain)
 	{
-		AcquiredImageIndex = Result;
-		return true;
+		int32 Result = SwapChain->AcquireImageIndex(&AcquiredSemaphore);
+		if (Result >= 0)
+		{
+			AcquiredImageIndex = Result;
+			return true;
+		}
 	}
-
 	return false;
 }
 
@@ -612,6 +615,10 @@ void FVulkanViewport::RecreateSwapchainFromRT(EPixelFormat PreferredPixelFormat)
 
 void FVulkanViewport::CreateSwapchain(FVulkanSwapChainRecreateInfo* RecreateInfo)
 {
+	// Release a previous swapchain 'dummy' and a real backbuffer if any
+	RenderingBackBuffer = nullptr;
+	RHIBackBuffer = nullptr;
+
 	if (SupportsStandardSwapchain())
 	{
 		uint32 DesiredNumBackBuffers = NUM_BUFFERS;
@@ -715,12 +722,13 @@ void FVulkanViewport::DestroySwapchain(FVulkanSwapChainRecreateInfo* RecreateInf
 	Device->SubmitCommandsAndFlushGPU();
 	Device->WaitUntilIdle();
 	
-	RenderingBackBuffer = nullptr;
+	// Intentionally leave RenderingBackBuffer alive, so it can be used a dummy backbuffer while we don't have swapchain images
+	// RenderingBackBuffer = nullptr;
 	
 	if (RHIBackBuffer)
 	{
-		RHIBackBuffer->ReleaseViewport();
-		RHIBackBuffer = nullptr;
+		RHIBackBuffer->ReleaseAcquiredImage();
+		// We release this RHIBackBuffer when we create a new swapchain
 	}
 		
 	if (SupportsStandardSwapchain() && SwapChain)
