Fixed: Compile errors in VulkanViewport.cpp

diff --git a/Engine/Source/Runtime/VulkanRHI/Private/VulkanViewport.cpp b/Engine/Source/Runtime/VulkanRHI/Private/VulkanViewport.cpp
--- a/Engine/Source/Runtime/VulkanRHI/Private/VulkanViewport.cpp
+++ b/Engine/Source/Runtime/VulkanRHI/Private/VulkanViewport.cpp
@@ -61,7 +61,7 @@ void FVulkanBackBuffer::AcquireBackBufferImage(FVulkanCommandListContext& Contex
 {
 	check(Viewport);
 	
-	if (Surface.Image == VK_NULL_HANDLE)
+	if (Image == VK_NULL_HANDLE)
 	{
 		if (Viewport->TryAcquireImageIndex())
 		{
