VulkanRHI: Refactor Layers and Extensions:
* Split code for layers and extensions.
* Keep string arrays for layer management, and share as much code as possible between device and instance layers.
* Create a new class for extensions and manage them through a single flow.
    - Package the code for the extension and its features+properties together in a single place, so the device code only iterates on loops of generic classes.
    - Keep distinct states to track if an extensions: is supported by the engine, is enabled in the engine, is supported by the driver at runtime, is activated at runtime.
    - If an extension is present but a required feature isn't, disable the extension.
    - Remove feature structs from Vulkan Device when they were only used for init.
    - Track promotion status of extensions (for future dev)
* Revisit some 'per platform' extensions that were enabled on all platforms.

Other changes:
* QueryGPU has been removed, VulkanDevice constructor fills props and id.
* Remove VendorId patch for Mesa driver from UE4 now that it is supported in the enum.
* Remove unused entrypoints.
* Add entrypoint for vkGetDeviceImageMemoryRequirementsKHR (for future dev)

diff --git a/Engine/Source/Runtime/VulkanRHI/Private/Android/VulkanAndroidPlatform.cpp b/Engine/Source/Runtime/VulkanRHI/Private/Android/VulkanAndroidPlatform.cpp
--- a/Engine/Source/Runtime/VulkanRHI/Private/Android/VulkanAndroidPlatform.cpp
+++ b/Engine/Source/Runtime/VulkanRHI/Private/Android/VulkanAndroidPlatform.cpp
@@ -8,6 +8,7 @@
 #include "Math/UnrealMathUtility.h"
 #include "Android/AndroidPlatformMisc.h"
 #include "Misc/ConfigCacheIni.h"
+#include "../VulkanExtensions.h"
 
 // From VulklanSwapChain.cpp
 extern int32 GVulkanCPURenderThreadFramePacer;
@@ -342,12 +343,12 @@ void FVulkanAndroidPlatform::CreateSurface(void* WindowHandle, VkInstance Instan
 }
 
 
-
-void FVulkanAndroidPlatform::GetInstanceExtensions(TArray<const ANSICHAR*>& OutExtensions)
+void FVulkanAndroidPlatform::GetInstanceExtensions(FVulkanInstanceExtensionArray& OutExtensions)
 {
-	OutExtensions.Add(VK_KHR_SURFACE_EXTENSION_NAME);
-	OutExtensions.Add(VK_KHR_ANDROID_SURFACE_EXTENSION_NAME);
-	OutExtensions.Add(VK_GOOGLE_DISPLAY_TIMING_EXTENSION_NAME);
+	OutExtensions.Add(MakeUnique<FVulkanInstanceExtension>(VK_KHR_ANDROID_SURFACE_EXTENSION_NAME, VULKAN_EXTENSION_ENABLED, VULKAN_EXTENSION_NOT_PROMOTED));
+
+	// VK_GOOGLE_display_timing (as instance extension?)
+	OutExtensions.Add(MakeUnique<FVulkanInstanceExtension>(VK_GOOGLE_DISPLAY_TIMING_EXTENSION_NAME, VULKAN_EXTENSION_ENABLED, VULKAN_EXTENSION_NOT_PROMOTED));
 }
 
 void FVulkanAndroidPlatform::GetInstanceLayers(TArray<const ANSICHAR*>& OutLayers)
@@ -385,45 +386,25 @@ static FAutoConsoleVariableRef CVarVulkanQcomRenderPassTransform(
 	ECVF_ReadOnly
 );
 
-void FVulkanAndroidPlatform::GetDeviceExtensions(EGpuVendorId VendorId, TArray<const ANSICHAR*>& OutExtensions)
+void FVulkanAndroidPlatform::GetDeviceExtensions(FVulkanDevice* Device, FVulkanDeviceExtensionArray& OutExtensions)
 {
-	OutExtensions.Add(VK_KHR_SURFACE_EXTENSION_NAME);
-	OutExtensions.Add(VK_KHR_ANDROID_SURFACE_EXTENSION_NAME);
-	OutExtensions.Add(VK_GOOGLE_DISPLAY_TIMING_EXTENSION_NAME);
-
-	OutExtensions.Add(VK_EXT_ASTC_DECODE_MODE_EXTENSION_NAME);
+	OutExtensions.Add(MakeUnique<FVulkanDeviceExtension>(Device, VK_KHR_ANDROID_SURFACE_EXTENSION_NAME, VULKAN_EXTENSION_ENABLED, VULKAN_EXTENSION_NOT_PROMOTED));
+	OutExtensions.Add(MakeUnique<FVulkanDeviceExtension>(Device, VK_GOOGLE_DISPLAY_TIMING_EXTENSION_NAME, VULKAN_EXTENSION_ENABLED, VULKAN_EXTENSION_NOT_PROMOTED));
+	OutExtensions.Add(MakeUnique<FVulkanDeviceExtension>(Device, VK_EXT_ASTC_DECODE_MODE_EXTENSION_NAME, VULKAN_SUPPORTS_ASTC_DECODE_MODE, VULKAN_EXTENSION_NOT_PROMOTED, DEVICE_EXT_FLAG_SETTER(HasEXTASTCDecodeMode)));
+	//OutExtensions.Add(MakeUnique<FVulkanDeviceExtension>(Device, VK_EXT_TEXTURE_COMPRESSION_ASTC_HDR_EXTENSION_NAME, VULKAN_SUPPORTS_TEXTURE_COMPRESSION_ASTC_HDR, VK_API_VERSION_1_3, DEVICE_EXT_FLAG_SETTER(HasEXTTextureCompressionASTCHDR)));
 
 	if (GVulkanQcomRenderPassTransform)
 	{
-		OutExtensions.Add(VK_QCOM_RENDER_PASS_TRANSFORM_EXTENSION_NAME);
+		OutExtensions.Add(MakeUnique<FVulkanDeviceExtension>(Device, VK_QCOM_RENDER_PASS_TRANSFORM_EXTENSION_NAME, VULKAN_EXTENSION_ENABLED, VK_API_VERSION_1_3, DEVICE_EXT_FLAG_SETTER(HasQcomRenderPassTransform)));
 	}
 
-#if VULKAN_SUPPORTS_FRAGMENT_DENSITY_MAP
-	OutExtensions.Add(VK_EXT_FRAGMENT_DENSITY_MAP_EXTENSION_NAME);
-#endif
-
-#if VULKAN_SUPPORTS_FRAGMENT_DENSITY_MAP2
-	OutExtensions.Add(VK_EXT_FRAGMENT_DENSITY_MAP_2_EXTENSION_NAME);
-#endif
-
-#if VULKAN_SUPPORTS_MULTIVIEW
-	OutExtensions.Add(VK_KHR_MULTIVIEW_EXTENSION_NAME);
-#endif
-
-#if VULKAN_SUPPORTS_RENDERPASS2
-	OutExtensions.Add(VK_KHR_CREATE_RENDERPASS_2_EXTENSION_NAME);
-#endif
-
-#if VULKAN_SUPPORTS_FRAGMENT_SHADING_RATE
-	OutExtensions.Add(VK_KHR_FRAGMENT_SHADING_RATE_EXTENSION_NAME);
-#endif
-
 #if !UE_BUILD_SHIPPING
-	OutExtensions.Add(VULKAN_MALI_LAYER_NAME);
+	// Layer name as extension
+	OutExtensions.Add(MakeUnique<FVulkanDeviceExtension>(Device, VULKAN_MALI_LAYER_NAME, VULKAN_EXTENSION_ENABLED, VULKAN_EXTENSION_NOT_PROMOTED));
 #endif
 }
 
-void FVulkanAndroidPlatform::GetDeviceLayers(EGpuVendorId VendorId, TArray<const ANSICHAR*>& OutLayers)
+void FVulkanAndroidPlatform::GetDeviceLayers(TArray<const ANSICHAR*>& OutLayers)
 {
 #if !UE_BUILD_SHIPPING
 	if (DebugVulkanDeviceLayers.IsEmpty())
@@ -449,10 +430,13 @@ void FVulkanAndroidPlatform::GetDeviceLayers(EGpuVendorId VendorId, TArray<const
 #endif
 }
 
-void FVulkanAndroidPlatform::NotifyFoundDeviceLayersAndExtensions(VkPhysicalDevice PhysicalDevice, const TArray<FString>& Layers, const TArray<FString>& Extensions)
+void FVulkanAndroidPlatform::NotifyFoundDeviceLayersAndExtensions(VkPhysicalDevice PhysicalDevice, const TArray<const ANSICHAR*>& Layers, const TArray<const ANSICHAR*>& Extensions)
 {
 #if VULKAN_SUPPORTS_GOOGLE_DISPLAY_TIMING
-	FVulkanAndroidPlatform::bHasGoogleDisplayTiming = Extensions.Contains(TEXT(VK_GOOGLE_DISPLAY_TIMING_EXTENSION_NAME));
+	FVulkanAndroidPlatform::bHasGoogleDisplayTiming = Extensions.ContainsByPredicate([](const ANSICHAR* Key)
+		{
+			return Key && !FCStringAnsi::Strcmp(Key, VK_GOOGLE_DISPLAY_TIMING_EXTENSION_NAME);
+		});
 	UE_LOG(LogVulkanRHI, Log, TEXT("bHasGoogleDisplayTiming = %d"), FVulkanAndroidPlatform::bHasGoogleDisplayTiming);
 #endif
 }
diff --git a/Engine/Source/Runtime/VulkanRHI/Private/Android/VulkanAndroidPlatform.h b/Engine/Source/Runtime/VulkanRHI/Private/Android/VulkanAndroidPlatform.h
--- a/Engine/Source/Runtime/VulkanRHI/Private/Android/VulkanAndroidPlatform.h
+++ b/Engine/Source/Runtime/VulkanRHI/Private/Android/VulkanAndroidPlatform.h
@@ -23,6 +23,9 @@
 #define VULKAN_SUPPORTS_NV_DIAGNOSTIC_CHECKPOINT	0
 #define VULKAN_RHI_RAYTRACING						0
 #define VULKAN_SUPPORTS_TRANSIENT_RESOURCE_ALLOCATOR 0
+#define VULKAN_SUPPORTS_DRIVER_PROPERTIES			0
+#define VULKAN_SUPPORTS_MULTIVIEW					1
+#define VULKAN_SUPPORTS_DESCRIPTOR_INDEXING			0
 
 // crashing during callback setup on Android, code will fallback to VK_EXT_debug_report instead
 #define VULKAN_SUPPORTS_DEBUG_UTILS					0
@@ -60,11 +63,11 @@ public:
 	static bool LoadVulkanInstanceFunctions(VkInstance inInstance);
 	static void FreeVulkanLibrary();
 
-	static void GetInstanceExtensions(TArray<const ANSICHAR*>& OutExtensions);
+	static void GetInstanceExtensions(FVulkanInstanceExtensionArray& OutExtensions);
 	static void GetInstanceLayers(TArray<const ANSICHAR*>& OutLayers);
-	static void GetDeviceExtensions(EGpuVendorId VendorId, TArray<const ANSICHAR*>& OutExtensions);
-	static void GetDeviceLayers(EGpuVendorId VendorId, TArray<const ANSICHAR*>& OutLayers);
-	static void NotifyFoundDeviceLayersAndExtensions(VkPhysicalDevice PhysicalDevice, const TArray<FString>& Layers, const TArray<FString>& Extensions);
+	static void GetDeviceExtensions(FVulkanDevice* Device, FVulkanDeviceExtensionArray& OutExtensions);
+	static void GetDeviceLayers(TArray<const ANSICHAR*>& OutLayers);
+	static void NotifyFoundDeviceLayersAndExtensions(VkPhysicalDevice PhysicalDevice, const TArray<const ANSICHAR*>& Layers, const TArray<const ANSICHAR*>& Extensions);
 
 	static void CreateSurface(void* WindowHandle, VkInstance Instance, VkSurfaceKHR* OutSurface);
 
diff --git a/Engine/Source/Runtime/VulkanRHI/Private/Linux/VulkanLinuxPlatform.cpp b/Engine/Source/Runtime/VulkanRHI/Private/Linux/VulkanLinuxPlatform.cpp
--- a/Engine/Source/Runtime/VulkanRHI/Private/Linux/VulkanLinuxPlatform.cpp
+++ b/Engine/Source/Runtime/VulkanRHI/Private/Linux/VulkanLinuxPlatform.cpp
@@ -3,6 +3,7 @@
 #include "VulkanLinuxPlatform.h"
 #include "../VulkanRHIPrivate.h"
 #include "../VulkanRayTracing.h"
+#include "../VulkanExtensions.h"
 #include <dlfcn.h>
 #include <SDL.h>
 #include <SDL_vulkan.h>
@@ -147,7 +148,7 @@ namespace
 	}
 }
 
-void FVulkanLinuxPlatform::GetInstanceExtensions(TArray<const ANSICHAR*>& OutExtensions)
+void FVulkanLinuxPlatform::GetInstanceExtensions(FVulkanInstanceExtensionArray& OutExtensions)
 {
 	EnsureSDLIsInited();
 
@@ -170,14 +171,13 @@ void FVulkanLinuxPlatform::GetInstanceExtensions(TArray<const ANSICHAR*>& OutExt
 		return;
 	}
 
-	OutExtensions.Add(VK_KHR_SURFACE_EXTENSION_NAME);
 	if (strcmp(SDLDriver, "x11") == 0)
 	{
-		OutExtensions.Add("VK_KHR_xlib_surface");
+		OutExtensions.Add(MakeUnique<FVulkanInstanceExtension>("VK_KHR_xlib_surface", VULKAN_EXTENSION_ENABLED, VULKAN_EXTENSION_NOT_PROMOTED));
 	}
 	else if (strcmp(SDLDriver, "wayland") == 0)
 	{
-		OutExtensions.Add("VK_KHR_wayland_surface");
+		OutExtensions.Add(MakeUnique<FVulkanInstanceExtension>("VK_KHR_wayland_surface", VULKAN_EXTENSION_ENABLED, VULKAN_EXTENSION_NOT_PROMOTED));
 	}
 	// dummy is when we render offscreen, so ignore warning here
 	else if (strcmp(SDLDriver, "dummy") != 0)
@@ -186,45 +186,10 @@ void FVulkanLinuxPlatform::GetInstanceExtensions(TArray<const ANSICHAR*>& OutExt
 	}
 }
 
-void FVulkanLinuxPlatform::GetDeviceExtensions(EGpuVendorId VendorId, TArray<const ANSICHAR*>& OutExtensions)
+void FVulkanLinuxPlatform::GetDeviceExtensions(FVulkanDevice* Device, FVulkanDeviceExtensionArray& OutExtensions)
 {
-	const bool bAllowVendorDevice = !FParse::Param(FCommandLine::Get(), TEXT("novendordevice"));
-
-#if VULKAN_SUPPORTS_DRIVER_PROPERTIES
-	OutExtensions.Add(VK_KHR_DRIVER_PROPERTIES_EXTENSION_NAME);
-#endif
-
-#if VULKAN_SUPPORTS_DEDICATED_ALLOCATION
-	OutExtensions.Add(VK_KHR_GET_MEMORY_REQUIREMENTS_2_EXTENSION_NAME);
-	OutExtensions.Add(VK_KHR_DEDICATED_ALLOCATION_EXTENSION_NAME);
-#endif
-
-#if VULKAN_SUPPORTS_RENDERPASS2
-	OutExtensions.Add(VK_KHR_CREATE_RENDERPASS_2_EXTENSION_NAME);
-
-	// Fragment shading rate depends on renderpass2.
-#if VULKAN_SUPPORTS_FRAGMENT_SHADING_RATE
-	OutExtensions.Add(VK_KHR_FRAGMENT_SHADING_RATE_EXTENSION_NAME);
-#endif
-
-#endif
-
-	if (GGPUCrashDebuggingEnabled)
-	{
-#if VULKAN_SUPPORTS_AMD_BUFFER_MARKER
-		if (VendorId == EGpuVendorId::Amd && bAllowVendorDevice)
-		{
-			OutExtensions.Add(VK_AMD_BUFFER_MARKER_EXTENSION_NAME);
-		}
-#endif
-#if VULKAN_SUPPORTS_NV_DIAGNOSTICS
-		if (VendorId == EGpuVendorId::Nvidia && bAllowVendorDevice)
-		{
-			OutExtensions.Add(VK_NV_DEVICE_DIAGNOSTIC_CHECKPOINTS_EXTENSION_NAME);
-			OutExtensions.Add(VK_NV_DEVICE_DIAGNOSTICS_CONFIG_EXTENSION_NAME);
-		}
-#endif
-	}
+	// Manually activated extensions
+	OutExtensions.Add(MakeUnique<FVulkanDeviceExtension>(Device, VK_KHR_EXTERNAL_MEMORY_FD_EXTENSION_NAME, VULKAN_SUPPORTS_EXTERNAL_MEMORY, VULKAN_EXTENSION_NOT_PROMOTED, nullptr, FVulkanExtensionBase::ManuallyActivate));
 }
 
 void FVulkanLinuxPlatform::CreateSurface(void* WindowHandle, VkInstance Instance, VkSurfaceKHR* OutSurface)
diff --git a/Engine/Source/Runtime/VulkanRHI/Private/Linux/VulkanLinuxPlatform.h b/Engine/Source/Runtime/VulkanRHI/Private/Linux/VulkanLinuxPlatform.h
--- a/Engine/Source/Runtime/VulkanRHI/Private/Linux/VulkanLinuxPlatform.h
+++ b/Engine/Source/Runtime/VulkanRHI/Private/Linux/VulkanLinuxPlatform.h
@@ -14,6 +14,7 @@
 
 #define VULKAN_RHI_RAYTRACING 						(RHI_RAYTRACING)
 #define VULKAN_SUPPORTS_SCALAR_BLOCK_LAYOUT			(VULKAN_RHI_RAYTRACING)
+#define VULKAN_SUPPORTS_MULTIVIEW					0
 
 #if VULKAN_RHI_RAYTRACING
 #	define UE_VK_API_VERSION						VK_API_VERSION_1_2
@@ -57,10 +58,10 @@ public:
 	static bool LoadVulkanInstanceFunctions(VkInstance inInstance);
 	static void FreeVulkanLibrary();
 
-	static void GetInstanceExtensions(TArray<const ANSICHAR*>& OutExtensions);
+	static void GetInstanceExtensions(FVulkanInstanceExtensionArray& OutExtensions);
 	static void GetInstanceLayers(TArray<const ANSICHAR*>& OutLayers) {}
-	static void GetDeviceExtensions(EGpuVendorId VendorId, TArray<const ANSICHAR*>& OutExtensions);
-	static void GetDeviceLayers(EGpuVendorId VendorId, TArray<const ANSICHAR*>& OutLayers) {}
+	static void GetDeviceExtensions(FVulkanDevice* Device, FVulkanDeviceExtensionArray& OutExtensions);
+	static void GetDeviceLayers(TArray<const ANSICHAR*>& OutLayers) {}
 
 	static void CreateSurface(void* WindowHandle, VkInstance Instance, VkSurfaceKHR* OutSurface);
 
diff --git a/Engine/Source/Runtime/VulkanRHI/Private/VulkanCommandWrappers.h b/Engine/Source/Runtime/VulkanRHI/Private/VulkanCommandWrappers.h
--- a/Engine/Source/Runtime/VulkanRHI/Private/VulkanCommandWrappers.h
+++ b/Engine/Source/Runtime/VulkanRHI/Private/VulkanCommandWrappers.h
@@ -201,6 +201,7 @@ struct FWrapLayer
 	static void GetRayTracingShaderGroupHandlesKHR(VkResult Result, VkDevice Device, VkPipeline Pipeline, uint32_t FirstGroup, uint32_t GroupCount, size_t DataSize, void* Data);
 	static void GetBufferDeviceAddressKHR(VkResult Result, VkDevice Device, const VkBufferDeviceAddressInfo* Info);
 #endif
+	static void GetDeviceImageMemoryRequirementsKHR(VkResult Result, VkDevice Device, const VkDeviceImageMemoryRequirements* Info, VkMemoryRequirements2* MemoryRequirements) VULKAN_LAYER_BODY
 };
 
 #undef VULKAN_LAYER_BODY
@@ -1468,6 +1469,13 @@ namespace VulkanRHI
 	}
 #endif
 
+	static FORCEINLINE_DEBUGGABLE void vkGetDeviceImageMemoryRequirementsKHR(VkDevice Device, const VkDeviceImageMemoryRequirements* Info, VkMemoryRequirements2* MemoryRequirements)
+	{
+		FWrapLayer::GetDeviceImageMemoryRequirementsKHR(VK_RESULT_MAX_ENUM, Device, Info, MemoryRequirements);
+		VULKANAPINAMESPACE::vkGetDeviceImageMemoryRequirementsKHR(Device, Info, MemoryRequirements);
+		FWrapLayer::GetDeviceImageMemoryRequirementsKHR(VK_SUCCESS, Device, Info, MemoryRequirements);
+	}
+
 #if VULKAN_ENABLE_IMAGE_TRACKING_LAYER
 	void BindDebugLabelName(VkImage Image, const TCHAR* Name);
 #endif
diff --git a/Engine/Source/Runtime/VulkanRHI/Private/VulkanDebug.cpp b/Engine/Source/Runtime/VulkanRHI/Private/VulkanDebug.cpp
--- a/Engine/Source/Runtime/VulkanRHI/Private/VulkanDebug.cpp
+++ b/Engine/Source/Runtime/VulkanRHI/Private/VulkanDebug.cpp
@@ -308,7 +308,7 @@ static VkBool32 DebugUtilsCallback(VkDebugUtilsMessageSeverityFlagBitsEXT MsgSev
 void FVulkanDynamicRHI::SetupDebugLayerCallback()
 {
 #if VULKAN_SUPPORTS_DEBUG_UTILS
-	if (bSupportsDebugUtilsExt)
+	if (ActiveDebugLayerExtension == EActiveDebugLayerExtension::DebugUtilsExtension)
 	{
 		PFN_vkCreateDebugUtilsMessengerEXT CreateDebugUtilsMessengerEXT = (PFN_vkCreateDebugUtilsMessengerEXT)(void*)VulkanRHI::vkGetInstanceProcAddr(Instance, "vkCreateDebugUtilsMessengerEXT");
 		if (CreateDebugUtilsMessengerEXT)
@@ -316,7 +316,7 @@ void FVulkanDynamicRHI::SetupDebugLayerCallback()
 			VkDebugUtilsMessengerCreateInfoEXT CreateInfo;
 			ZeroVulkanStruct(CreateInfo, VK_STRUCTURE_TYPE_DEBUG_UTILS_MESSENGER_CREATE_INFO_EXT);
 
-			int32 CVar = GValidationCvar.GetValueOnRenderThread();
+			const int32 CVar = GValidationCvar.GetValueOnRenderThread();
 			CreateInfo.messageSeverity = (CVar >= 1 ? VK_DEBUG_UTILS_MESSAGE_SEVERITY_ERROR_BIT_EXT : 0) |
 				(CVar >= 2 ? VK_DEBUG_UTILS_MESSAGE_SEVERITY_WARNING_BIT_EXT : 0) | (CVar >= 3 ? VK_DEBUG_UTILS_MESSAGE_SEVERITY_INFO_BIT_EXT : 0);
 			CreateInfo.messageType = (CVar >= 1 ? (VK_DEBUG_UTILS_MESSAGE_TYPE_GENERAL_BIT_EXT | VK_DEBUG_UTILS_MESSAGE_TYPE_VALIDATION_BIT_EXT) : 0) |
@@ -328,7 +328,7 @@ void FVulkanDynamicRHI::SetupDebugLayerCallback()
 	}
 	else
 #endif
-	if (bSupportsDebugCallbackExt)
+	if (ActiveDebugLayerExtension == EActiveDebugLayerExtension::DebugReportExtension)
 	{
 		PFN_vkCreateDebugReportCallbackEXT CreateMsgCallback = (PFN_vkCreateDebugReportCallbackEXT)(void*)VulkanRHI::vkGetInstanceProcAddr(Instance, CREATE_MSG_CALLBACK);
 		if (CreateMsgCallback)
@@ -337,7 +337,7 @@ void FVulkanDynamicRHI::SetupDebugLayerCallback()
 			ZeroVulkanStruct(CreateInfo, VK_STRUCTURE_TYPE_DEBUG_REPORT_CREATE_INFO_EXT);
 			CreateInfo.pfnCallback = DebugReportFunction;
 
-			int32 CVar = GValidationCvar.GetValueOnRenderThread();
+			const int32 CVar = GValidationCvar.GetValueOnRenderThread();
 			switch (CVar)
 			{
 			default:
@@ -4456,6 +4456,16 @@ void FWrapLayer::GetBufferDeviceAddressKHR(VkResult Result, VkDevice Device, con
 }
 #endif // VULKAN_RHI_RAYTRACING
 
+void FWrapLayer::GetDeviceImageMemoryRequirementsKHR(VkResult Result, VkDevice Device, const VkDeviceImageMemoryRequirements* Info, VkMemoryRequirements2* MemoryRequirements)
+{
+	if (Result == VK_RESULT_MAX_ENUM)
+	{
+#if VULKAN_ENABLE_DUMP_LAYER
+		PrintfBeginResult(FString::Printf(TEXT("vkGetDeviceImageMemoryRequirementsKHR")));
+#endif
+	}
+}
+
 #if VULKAN_ENABLE_IMAGE_TRACKING_LAYER
 namespace VulkanRHI
 {
diff --git a/Engine/Source/Runtime/VulkanRHI/Private/VulkanDevice.cpp b/Engine/Source/Runtime/VulkanRHI/Private/VulkanDevice.cpp
--- a/Engine/Source/Runtime/VulkanRHI/Private/VulkanDevice.cpp
+++ b/Engine/Source/Runtime/VulkanRHI/Private/VulkanDevice.cpp
@@ -14,6 +14,7 @@
 #include "VulkanPlatform.h"
 #include "VulkanLLM.h"
 #include "VulkanTransientResourceAllocator.h"
+#include "VulkanExtensions.h"
 
 TAutoConsoleVariable<int32> GRHIAllowAsyncComputeCvar(
 	TEXT("r.Vulkan.AllowAsyncCompute"),
@@ -44,13 +45,6 @@ static TAutoConsoleVariable<int32> CVarVulkanUseD24(
 	ECVF_ReadOnly
 );
 
-TAutoConsoleVariable<int32> GRHIAllow64bitShaderAtomicsCvar(
-	TEXT("r.Vulkan.Allow64bitShaderAtomics"),
-	1,
-	TEXT("Whether to enable 64bit buffer/image atomics required by Nanite/Lumen\n")
-	TEXT("0 to disable 64bit atomics\n")
-	TEXT("1 to enable (default)")
-);
 
 #if NV_AFTERMATH
 #include "GFSDK_Aftermath_GpuCrashDump.h"
@@ -184,24 +178,34 @@ FVulkanDevice::FVulkanDevice(FVulkanDynamicRHI* InRHI, VkPhysicalDevice InGpu)
 {
 	RHI = InRHI;
 	FMemory::Memzero(GpuProps);
-#if VULKAN_SUPPORTS_EXTERNAL_MEMORY
-	FMemory::Memzero(GpuIdProps);
-#endif
 	FMemory::Memzero(PhysicalFeatures);
 	FMemory::Memzero(FormatProperties);
 	FMemory::Memzero(PixelFormatComponentMapping);
 
+#if VULKAN_SUPPORTS_PHYSICAL_DEVICE_PROPERTIES2
+	ZeroVulkanStruct(GpuIdProps, VK_STRUCTURE_TYPE_PHYSICAL_DEVICE_ID_PROPERTIES_KHR);
+	if (RHI->GetOptionalExtensions().HasKHRGetPhysicalDeviceProperties2)
+	{
+		VkPhysicalDeviceProperties2KHR PhysicalDeviceProperties2;
+		ZeroVulkanStruct(PhysicalDeviceProperties2, VK_STRUCTURE_TYPE_PHYSICAL_DEVICE_PROPERTIES_2_KHR);
+		PhysicalDeviceProperties2.pNext = &GpuIdProps;
+		VulkanRHI::vkGetPhysicalDeviceProperties2KHR(Gpu, &PhysicalDeviceProperties2);
+		GpuProps = PhysicalDeviceProperties2.properties;
+	}
+	else
+#endif
+	{
+		VulkanRHI::vkGetPhysicalDeviceProperties(Gpu, &GpuProps);
+	}
+
 	// First get the VendorId. We'll have to get properties again after finding out which extensions we want to use
-	VulkanRHI::vkGetPhysicalDeviceProperties(Gpu, &GpuProps);
 	VendorId = RHIConvertToGpuVendorId(GpuProps.vendorID);
+	ensure(VendorId != EGpuVendorId::Unknown);
 
-	// Mesa driver on Linux has vendorID 0x10005:
-	//   https://www.khronos.org/registry/vulkan/specs/1.2-extensions/man/html/VkVendorId.html
-	// Has been added to UE5, but need this workaround on 4.27.1 since we can't modify enum at this point.
-	if (GpuProps.vendorID != 0x10005)
-	{
-		ensure(VendorId != EGpuVendorId::Unknown);
-	}
+	UE_LOG(LogVulkanRHI, Display, TEXT("- DeviceName: %s"), ANSI_TO_TCHAR(GpuProps.deviceName));
+	UE_LOG(LogVulkanRHI, Display, TEXT("- API=%d.%d.%d (0x%x) Driver=0x%x VendorId=0x%x"), VK_VERSION_MAJOR(GpuProps.apiVersion), VK_VERSION_MINOR(GpuProps.apiVersion), VK_VERSION_PATCH(GpuProps.apiVersion), GpuProps.apiVersion, GpuProps.driverVersion, GpuProps.vendorID);
+	UE_LOG(LogVulkanRHI, Display, TEXT("- DeviceID=0x%x Type=%s"), GpuProps.deviceID, VK_TYPE_TO_STRING(VkPhysicalDeviceType, GpuProps.deviceType));
+	UE_LOG(LogVulkanRHI, Display, TEXT("- Max Descriptor Sets Bound %d, Timestamps %d"), GpuProps.limits.maxBoundDescriptorSets, GpuProps.limits.timestampComputeAndGraphics);
 }
 
 FVulkanDevice::~FVulkanDevice()
@@ -236,7 +240,7 @@ static inline FString GetQueueInfoString(const VkQueueFamilyProperties& Props)
 	return Info;
 };
 
-void FVulkanDevice::CreateDevice()
+void FVulkanDevice::CreateDevice(TArray<const ANSICHAR*>& DeviceLayers, FVulkanDeviceExtensionArray& UEExtensions)
 {
 	LLM_SCOPE_VULKAN(ELLMTagVulkan::VulkanMisc);
 	check(Device == VK_NULL_HANDLE);
@@ -245,11 +249,21 @@ void FVulkanDevice::CreateDevice()
 	VkDeviceCreateInfo DeviceInfo;
 	ZeroVulkanStruct(DeviceInfo, VK_STRUCTURE_TYPE_DEVICE_CREATE_INFO);
 	
+	TArray<const ANSICHAR*> DeviceExtensions;
+	for (TUniquePtr<FVulkanDeviceExtension>& UEExtension : UEExtensions)
+	{
+		if (UEExtension->InUse())
+		{
+			DeviceExtensions.Add(UEExtension->GetExtensionName());
+			UEExtension->PreCreateDevice(DeviceInfo);
+		}
+	}
+
 	DeviceInfo.enabledExtensionCount = DeviceExtensions.Num();
 	DeviceInfo.ppEnabledExtensionNames = DeviceExtensions.GetData();
 
-	DeviceInfo.enabledLayerCount = ValidationLayers.Num();
-	DeviceInfo.ppEnabledLayerNames = (DeviceInfo.enabledLayerCount > 0) ? ValidationLayers.GetData() : nullptr;
+	DeviceInfo.enabledLayerCount = DeviceLayers.Num();
+	DeviceInfo.ppEnabledLayerNames = (DeviceInfo.enabledLayerCount > 0) ? DeviceLayers.GetData() : nullptr;
 
 	// Setup Queue info
 	TArray<VkDeviceQueueCreateInfo> QueueFamilyInfos;
@@ -336,23 +350,7 @@ void FVulkanDevice::CreateDevice()
 
 	FVulkanPlatform::EnablePhysicalDeviceFeatureExtensions(DeviceInfo, *this);
 
-#if VULKAN_SUPPORTS_NV_DIAGNOSTICS
-	VkPhysicalDeviceDiagnosticsConfigFeaturesNV DeviceDiagnosticsNV;
-	VkDeviceDiagnosticsConfigCreateInfoNV DeviceDiagnosticsConfigCreateInfoNV;
-	if (OptionalDeviceExtensions.HasNVDeviceDiagnosticConfig)
-	{
-		ZeroVulkanStruct(DeviceDiagnosticsNV, VK_STRUCTURE_TYPE_PHYSICAL_DEVICE_DIAGNOSTICS_CONFIG_FEATURES_NV);
-		DeviceDiagnosticsNV.pNext = const_cast<void*>(DeviceInfo.pNext);
-		DeviceDiagnosticsNV.diagnosticsConfig = VK_TRUE;
-		DeviceInfo.pNext = &DeviceDiagnosticsNV;
-
-		ZeroVulkanStruct(DeviceDiagnosticsConfigCreateInfoNV, VK_STRUCTURE_TYPE_DEVICE_DIAGNOSTICS_CONFIG_CREATE_INFO_NV);
-		DeviceDiagnosticsConfigCreateInfoNV.flags = VK_DEVICE_DIAGNOSTICS_CONFIG_ENABLE_SHADER_DEBUG_INFO_BIT_NV | VK_DEVICE_DIAGNOSTICS_CONFIG_ENABLE_RESOURCE_TRACKING_BIT_NV | VK_DEVICE_DIAGNOSTICS_CONFIG_ENABLE_AUTOMATIC_CHECKPOINTS_BIT_NV;
-		DeviceDiagnosticsConfigCreateInfoNV.pNext = const_cast<void*>(DeviceInfo.pNext);
-		DeviceInfo.pNext = &DeviceDiagnosticsConfigCreateInfoNV;
-	}
-
-	#if NV_AFTERMATH
+#if NV_AFTERMATH && VULKAN_SUPPORTS_NV_DIAGNOSTICS
 	if (GGPUCrashDebuggingEnabled && GVulkanNVAftermathModuleLoaded)
 	{
 		GFSDK_Aftermath_Result Result = GFSDK_Aftermath_EnableGpuCrashDumps(GFSDK_Aftermath_Version_API, 
@@ -367,70 +365,6 @@ void FVulkanDevice::CreateDevice()
 			UE_LOG(LogVulkanRHI, Warning, TEXT("Unable to initialize Aftermath crash dumps (Result %d)"), (int32)Result);
 		}
 	}
-	#endif
-#endif
-
-#if VULKAN_SUPPORTS_SEPARATE_DEPTH_STENCIL_LAYOUTS
-	if (bHasSeparateDepthStencilLayouts)
-	{
-		ZeroVulkanStruct(OptionalFeatures.SeparateDepthStencilLayoutsFeatures, VK_STRUCTURE_TYPE_PHYSICAL_DEVICE_SEPARATE_DEPTH_STENCIL_LAYOUTS_FEATURES_KHR);
-		OptionalFeatures.SeparateDepthStencilLayoutsFeatures.separateDepthStencilLayouts = VK_TRUE;
-		OptionalFeatures.SeparateDepthStencilLayoutsFeatures.pNext = const_cast<void*>(DeviceInfo.pNext);
-		DeviceInfo.pNext = &OptionalFeatures.SeparateDepthStencilLayoutsFeatures;
-	}
-#endif
-
-#if VULKAN_SUPPORTS_SCALAR_BLOCK_LAYOUT
-	if (OptionalDeviceExtensions.HasEXTScalarBlockLayout)
-	{
-		ZeroVulkanStruct(OptionalFeatures.ScalarBlockLayoutFeatures, VK_STRUCTURE_TYPE_PHYSICAL_DEVICE_SCALAR_BLOCK_LAYOUT_FEATURES_EXT);
-		OptionalFeatures.ScalarBlockLayoutFeatures.scalarBlockLayout = VK_TRUE;
-		OptionalFeatures.ScalarBlockLayoutFeatures.pNext = (void*)DeviceInfo.pNext;
-		DeviceInfo.pNext = &OptionalFeatures.ScalarBlockLayoutFeatures;
-	}
-#endif
-
-#if VULKAN_SUPPORTS_MULTIVIEW
-	VkPhysicalDeviceMultiviewFeatures DeviceMultiviewFeatures;
-	if (OptionalDeviceExtensions.HasKHRMultiview)
-	{
-		ZeroVulkanStruct(DeviceMultiviewFeatures, VK_STRUCTURE_TYPE_PHYSICAL_DEVICE_MULTIVIEW_FEATURES);
-		DeviceMultiviewFeatures.multiview = VK_TRUE;
-		DeviceMultiviewFeatures.pNext = (void*)DeviceInfo.pNext;
-		DeviceInfo.pNext = &DeviceMultiviewFeatures;
-	}
-#endif
-
-#if VULKAN_SUPPORTS_BUFFER_64BIT_ATOMICS && VULKAN_SUPPORTS_IMAGE_64BIT_ATOMICS
-	const bool bEnable64bitAtomics = (GRHIAllow64bitShaderAtomicsCvar.GetValueOnAnyThread() != 0);
-
-	// Ensure we have both to enable atomic code paths
-	VkPhysicalDeviceShaderAtomicInt64Features BufferAtomicFeatures;
-	VkPhysicalDeviceShaderImageAtomicInt64FeaturesEXT ImageAtomicFeatures;
-	if (bEnable64bitAtomics && OptionalDeviceExtensions.HasKHRShaderAtomicInt64 && OptionalDeviceExtensions.HasImageAtomicInt64)
-	{
-		ZeroVulkanStruct(BufferAtomicFeatures, VK_STRUCTURE_TYPE_PHYSICAL_DEVICE_SHADER_ATOMIC_INT64_FEATURES_KHR);
-		BufferAtomicFeatures.shaderBufferInt64Atomics = VK_TRUE;
-
-		ZeroVulkanStruct(ImageAtomicFeatures, VK_STRUCTURE_TYPE_PHYSICAL_DEVICE_SHADER_IMAGE_ATOMIC_INT64_FEATURES_EXT);
-		ImageAtomicFeatures.shaderImageInt64Atomics = VK_TRUE;
-
-		BufferAtomicFeatures.pNext = &ImageAtomicFeatures;
-		ImageAtomicFeatures.pNext = (void*)DeviceInfo.pNext;
-		DeviceInfo.pNext = &BufferAtomicFeatures;
-
-		GRHISupportsAtomicUInt64 = true;
-	}
-#endif
-
-#if VULKAN_SUPPORTS_FRAGMENT_SHADING_RATE
-	VkPhysicalDeviceFragmentShadingRateFeaturesKHR RequestedFragmentShadingRateFeatures;
-	if (OptionalDeviceExtensions.HasKHRFragmentShadingRate)
-	{
-		RequestedFragmentShadingRateFeatures = GetFragmentShadingRateFeatures();
-		RequestedFragmentShadingRateFeatures.pNext = (void*)DeviceInfo.pNext;
-		DeviceInfo.pNext = &RequestedFragmentShadingRateFeatures;
-	}
 #endif
 
 	// Create the device
@@ -442,6 +376,8 @@ void FVulkanDevice::CreateDevice()
 	}
 	VERIFYVULKANRESULT_EXPANDED(Result);
 
+	FVulkanPlatform::NotifyFoundDeviceLayersAndExtensions(Gpu, DeviceLayers, DeviceExtensions);
+
 	// Create Graphics Queue, here we submit command buffers for execution
 	GfxQueue = new FVulkanQueue(this, GfxQueueFamilyIndex);
 	if (ComputeQueueFamilyIndex == -1)
@@ -480,13 +416,31 @@ void FVulkanDevice::CreateDevice()
 		}
 	}
 
-	UE_LOG(LogVulkanRHI, Display, TEXT("Using device layers"));
-	for (const ANSICHAR* Layer : ValidationLayers)
+#if VULKAN_SUPPORTS_FRAGMENT_SHADING_RATE
+	// Enumerate the available shading rates if it's supported
+	if (OptionalDeviceExtensions.HasKHRFragmentShadingRate)
+	{
+		uint32 FragmentShadingRateCount = 0;
+		VulkanRHI::vkGetPhysicalDeviceFragmentShadingRatesKHR(Gpu, &FragmentShadingRateCount, nullptr);
+		if (FragmentShadingRateCount != 0)
+		{
+			FragmentShadingRates.SetNum(FragmentShadingRateCount);
+			for (uint32 i = 0; i < FragmentShadingRateCount; ++i)
+			{
+				ZeroVulkanStruct(FragmentShadingRates[i], VK_STRUCTURE_TYPE_PHYSICAL_DEVICE_FRAGMENT_SHADING_RATE_KHR);
+			}
+			VulkanRHI::vkGetPhysicalDeviceFragmentShadingRatesKHR(Gpu, &FragmentShadingRateCount, FragmentShadingRates.GetData());
+		}
+	}
+#endif // VULKAN_SUPPORTS_FRAGMENT_SHADING_RATE
+
+	UE_LOG(LogVulkanRHI, Display, TEXT("Using %d device layers%s"), DeviceLayers.Num(), DeviceLayers.Num() ? TEXT(":") : TEXT("."));
+	for (const ANSICHAR* Layer : DeviceLayers)
 	{
 		UE_LOG(LogVulkanRHI, Display, TEXT("* %s"), ANSI_TO_TCHAR(Layer));
 	}
 
-	UE_LOG(LogVulkanRHI, Display, TEXT("Using device extensions"));
+	UE_LOG(LogVulkanRHI, Display, TEXT("Using %d device extensions:"), DeviceExtensions.Num());
 	for (const ANSICHAR* Extension : DeviceExtensions)
 	{
 		UE_LOG(LogVulkanRHI, Display, TEXT("* %s"), ANSI_TO_TCHAR(Extension));
@@ -907,7 +861,7 @@ void FVulkanDevice::MapImageFormatSupport(FPixelFormatInfo& PixelFormatInfo, con
 		ConvertImageCap(LocalFormatProperties, EPixelFormatCapabilities::DepthStencil, VK_FORMAT_FEATURE_DEPTH_STENCIL_ATTACHMENT_BIT);
 		ConvertImageCap(LocalFormatProperties, EPixelFormatCapabilities::RenderTarget, VK_FORMAT_FEATURE_COLOR_ATTACHMENT_BIT);
 		ConvertImageCap(LocalFormatProperties, EPixelFormatCapabilities::TextureBlendable, VK_FORMAT_FEATURE_COLOR_ATTACHMENT_BLEND_BIT);
-		ConvertImageCap(LocalFormatProperties, EPixelFormatCapabilities::UAV | EPixelFormatCapabilities::TypedUAVLoad | EPixelFormatCapabilities::TypedUAVStore, VK_FORMAT_FEATURE_STORAGE_IMAGE_BIT);
+		ConvertImageCap(LocalFormatProperties, EPixelFormatCapabilities::AllUAVFlags | EPixelFormatCapabilities::TextureStore, VK_FORMAT_FEATURE_STORAGE_IMAGE_BIT);
 
 		ConvertImageCap(LocalFormatProperties, EPixelFormatCapabilities::TextureAtomics, VK_FORMAT_FEATURE_STORAGE_IMAGE_ATOMIC_BIT);
 
@@ -986,176 +940,9 @@ void FVulkanDevice::MapFormatSupport(EPixelFormat UEFormat, std::initializer_lis
 	MapFormatSupport(UEFormat, PrioritizedFormats, ComponentMapping, RequiredCapabilities, kDefaultBlockBytes);
 }
 
-bool FVulkanDevice::QueryGPU(int32 DeviceIndex)
+void FVulkanDevice::InitGPU(int32 DeviceIndex)
 {
-	// Always get the extensions/layers first!
-	TArray<FString> AllDeviceExtensions;
-	TArray<FString> AllValidationLayers;
-	GetDeviceExtensionsAndLayers(Gpu, VendorId, DeviceExtensions, ValidationLayers, AllDeviceExtensions, AllValidationLayers, bDebugMarkersFound);
-	OptionalDeviceExtensions.Setup(DeviceExtensions);
-
-#if VULKAN_SUPPORTS_DRIVER_PROPERTIES
-	VkPhysicalDeviceDriverPropertiesKHR PhysicalDeviceProperties;
-	ZeroVulkanStruct(PhysicalDeviceProperties, VK_STRUCTURE_TYPE_PHYSICAL_DEVICE_DRIVER_PROPERTIES_KHR);
-#endif
-
-#if VULKAN_SUPPORTS_FRAGMENT_DENSITY_MAP
-	VkPhysicalDeviceFragmentDensityMapPropertiesEXT FragmentDensityMapProperties;
-	ZeroVulkanStruct(FragmentDensityMapProperties, VK_STRUCTURE_TYPE_PHYSICAL_DEVICE_FRAGMENT_DENSITY_MAP_PROPERTIES_EXT);
-#endif
-
-#if VULKAN_SUPPORTS_FRAGMENT_SHADING_RATE
-	ZeroVulkanStruct(FragmentShadingRateProperties, VK_STRUCTURE_TYPE_PHYSICAL_DEVICE_FRAGMENT_SHADING_RATE_PROPERTIES_KHR); 
-#endif
-
-#if VULKAN_SUPPORTS_PHYSICAL_DEVICE_PROPERTIES2
-	if (RHI->GetOptionalExtensions().HasKHRGetPhysicalDeviceProperties2)
-	{
-		VkPhysicalDeviceProperties2KHR GpuProps2;
-		ZeroVulkanStruct(GpuProps2, VK_STRUCTURE_TYPE_PHYSICAL_DEVICE_PROPERTIES_2_KHR);
-		GpuProps2.pNext = &GpuIdProps;
-		ZeroVulkanStruct(GpuIdProps, VK_STRUCTURE_TYPE_PHYSICAL_DEVICE_ID_PROPERTIES_KHR);
-
-#if VULKAN_SUPPORTS_DRIVER_PROPERTIES
-		if (GetOptionalExtensions().HasDriverProperties)
-		{
-			PhysicalDeviceProperties.pNext = GpuProps2.pNext;
-			GpuProps2.pNext = &PhysicalDeviceProperties;
-		}
-#endif
-
-#if VULKAN_SUPPORTS_FRAGMENT_DENSITY_MAP
-		if (GetOptionalExtensions().HasEXTFragmentDensityMap)
-		{
-			FragmentDensityMapProperties.pNext = GpuProps2.pNext;
-			GpuProps2.pNext = &FragmentDensityMapProperties;
-		}
-#endif
-
-#if VULKAN_SUPPORTS_FRAGMENT_SHADING_RATE
-		if (GetOptionalExtensions().HasKHRFragmentShadingRate)
-		{
-			FragmentShadingRateProperties.pNext = GpuProps2.pNext;
-			GpuProps2.pNext = &FragmentShadingRateProperties;
-		}
-#endif
-
-#if VULKAN_RHI_RAYTRACING
-		if (OptionalDeviceExtensions.HasRaytracingExtensions())
-		{
-			ZeroVulkanStruct(RayTracingProperties.AccelerationStructure, VK_STRUCTURE_TYPE_PHYSICAL_DEVICE_ACCELERATION_STRUCTURE_PROPERTIES_KHR);
-			ZeroVulkanStruct(RayTracingProperties.RayTracingPipeline, VK_STRUCTURE_TYPE_PHYSICAL_DEVICE_RAY_TRACING_PIPELINE_PROPERTIES_KHR);
-			RayTracingProperties.RayTracingPipeline.pNext = GpuProps2.pNext;
-			RayTracingProperties.AccelerationStructure.pNext = &RayTracingProperties.RayTracingPipeline;
-			GpuProps2.pNext = &RayTracingProperties.AccelerationStructure;
-		}
-#endif
-
-		VulkanRHI::vkGetPhysicalDeviceProperties2KHR(Gpu, &GpuProps2);
-
-#if VULKAN_SUPPORTS_SEPARATE_DEPTH_STENCIL_LAYOUTS
-		// Disabled but kept for reference. The barriers code doesn't currently use separate transitions for depth and stencil.
-		//VkPhysicalDeviceFeatures2KHR PhysicalFeatures2;
-		//ZeroVulkanStruct(PhysicalFeatures2, VK_STRUCTURE_TYPE_PHYSICAL_DEVICE_FEATURES_2_KHR);
-		//VkPhysicalDeviceSeparateDepthStencilLayoutsFeaturesKHR SeparateDepthStencilLayoutsFeatures;
-		//ZeroVulkanStruct(SeparateDepthStencilLayoutsFeatures, VK_STRUCTURE_TYPE_PHYSICAL_DEVICE_SEPARATE_DEPTH_STENCIL_LAYOUTS_FEATURES_KHR);
-		//PhysicalFeatures2.pNext = &SeparateDepthStencilLayoutsFeatures;
-		//VulkanDynamicAPI::vkGetPhysicalDeviceFeatures2KHR(Gpu, &PhysicalFeatures2);
-		//bHasSeparateDepthStencilLayouts = SeparateDepthStencilLayoutsFeatures.separateDepthStencilLayouts ? true : false;
-#endif
-
-	}
-#endif // VULKAN_SUPPORTS_PHYSICAL_DEVICE_PROPERTIES2
-
-	VulkanRHI::vkGetPhysicalDeviceProperties(Gpu, &GpuProps);
-	bool bDiscrete = false;
-	auto GetDeviceTypeString = [&]()
-	{
-		FString Info;
-		switch (GpuProps.deviceType)
-		{
-		case  VK_PHYSICAL_DEVICE_TYPE_OTHER:
-			Info = TEXT("Other");
-			break;
-		case VK_PHYSICAL_DEVICE_TYPE_INTEGRATED_GPU:
-			Info = TEXT("Integrated GPU");
-			break;
-		case VK_PHYSICAL_DEVICE_TYPE_DISCRETE_GPU:
-			Info = TEXT("Discrete GPU");
-			bDiscrete = true;
-			break;
-		case VK_PHYSICAL_DEVICE_TYPE_VIRTUAL_GPU:
-			Info = TEXT("Virtual GPU");
-			break;
-		case VK_PHYSICAL_DEVICE_TYPE_CPU:
-			Info = TEXT("CPU");
-			break;
-		default:
-			Info = TEXT("Unknown");
-			break;
-		}
-		return Info;
-	};
-
-	UE_LOG(LogVulkanRHI, Display, TEXT("Device %d: %s"), DeviceIndex, ANSI_TO_TCHAR(GpuProps.deviceName));
-	UE_LOG(LogVulkanRHI, Display, TEXT("- API %d.%d.%d(0x%x) Driver 0x%x VendorId 0x%x"), VK_VERSION_MAJOR(GpuProps.apiVersion), VK_VERSION_MINOR(GpuProps.apiVersion), VK_VERSION_PATCH(GpuProps.apiVersion), GpuProps.apiVersion, GpuProps.driverVersion, GpuProps.vendorID);
-	UE_LOG(LogVulkanRHI, Display, TEXT("- DeviceID 0x%x Type %s"), GpuProps.deviceID, *GetDeviceTypeString());
-	UE_LOG(LogVulkanRHI, Display, TEXT("- Max Descriptor Sets Bound %d Timestamps %d"), GpuProps.limits.maxBoundDescriptorSets, GpuProps.limits.timestampComputeAndGraphics);
-
-#if VULKAN_SUPPORTS_DRIVER_PROPERTIES
-	if (OptionalDeviceExtensions.HasDriverProperties)
-	{
-		UE_LOG(LogVulkanRHI, Display, TEXT("- Device Properties driverName: %s"), ANSI_TO_TCHAR(PhysicalDeviceProperties.driverName));
-		UE_LOG(LogVulkanRHI, Display, TEXT("- Device Properties driverInfo: %s"), ANSI_TO_TCHAR(PhysicalDeviceProperties.driverInfo));
-	}
-#endif
-
-#if VULKAN_SUPPORTS_FRAGMENT_DENSITY_MAP 
-	// Use the Fragment Density Map extension if and only if the Fragment Shading Rate extension is not available.
-	if (GetOptionalExtensions().HasEXTFragmentDensityMap && !GetOptionalExtensions().HasKHRFragmentShadingRate)
-	{
-		GRHISupportsAttachmentVariableRateShading = true;
-		GRHISupportsPipelineVariableRateShading = false;
-
-		// Go with the smallest tile size for now, and also force to square, since this seems to be standard.
-		// TODO: Eventually we may want to surface the range of possible tile sizes depending on end use cases, but for now this is being used for foveated rendering and smallest tile size
-		// is preferred.
-		
-		GRHIVariableRateShadingImageTileMinWidth = FragmentDensityMapProperties.minFragmentDensityTexelSize.width;
-		GRHIVariableRateShadingImageTileMinHeight = FragmentDensityMapProperties.minFragmentDensityTexelSize.height;
-		GRHIVariableRateShadingImageTileMaxWidth = FragmentDensityMapProperties.maxFragmentDensityTexelSize.width;
-		GRHIVariableRateShadingImageTileMaxHeight = FragmentDensityMapProperties.maxFragmentDensityTexelSize.height;
-
-		GRHIVariableRateShadingImageDataType = VRSImage_Fractional;
-		GRHIVariableRateShadingImageFormat = PF_R8G8;
-
-		// UE_LOG(LogVulkanRHI, Display, TEXT("Image-based Variable Rate Shading supported via EXTFragmentDensityMap extension. Selected VRS tile size %u by %u pixels per VRS image texel."), GRHIVariableRateShadingImageTileMinWidth, GRHIVariableRateShadingImageTileMinHeight);
-	}
-#endif
-
-#if VULKAN_SUPPORTS_FRAGMENT_SHADING_RATE
-	if (GetOptionalExtensions().HasKHRFragmentShadingRate)
-	{
-		GRHIVariableRateShadingImageTileMinWidth = FragmentShadingRateProperties.minFragmentShadingRateAttachmentTexelSize.width;
-		GRHIVariableRateShadingImageTileMinHeight = FragmentShadingRateProperties.minFragmentShadingRateAttachmentTexelSize.height;
-		GRHIVariableRateShadingImageTileMaxWidth = FragmentShadingRateProperties.maxFragmentShadingRateAttachmentTexelSize.width;
-		GRHIVariableRateShadingImageTileMaxHeight = FragmentShadingRateProperties.maxFragmentShadingRateAttachmentTexelSize.height;
-
-		// todo: We don't currently care much about the other properties here, but at some point in the future we probably will.
-
-		UE_LOG(LogVulkanRHI, Verbose, TEXT("Image-based Variable Rate Shading supported via KHRFragmentShadingRate extension. Selected VRS tile size %u by %u pixels per VRS image texel."), GRHIVariableRateShadingImageTileMinWidth, GRHIVariableRateShadingImageTileMinHeight);
-	}
-#endif
-
-	for (const FString& Name : AllValidationLayers)
-	{
-		UE_LOG(LogVulkanRHI, Display, TEXT("-    Found device layer %s"), *Name);
-	}
-
-	for (const FString& Name : AllDeviceExtensions)
-	{
-		UE_LOG(LogVulkanRHI, Display, TEXT("-    Found device extension %s"), *Name);
-	}
+	LLM_SCOPE_VULKAN(ELLMTagVulkan::VulkanMisc);
 
 	uint32 QueueCount = 0;
 	VulkanRHI::vkGetPhysicalDeviceQueueFamilyProperties(Gpu, &QueueCount, nullptr);
@@ -1164,135 +951,71 @@ bool FVulkanDevice::QueryGPU(int32 DeviceIndex)
 	QueueFamilyProps.AddUninitialized(QueueCount);
 	VulkanRHI::vkGetPhysicalDeviceQueueFamilyProperties(Gpu, &QueueCount, QueueFamilyProps.GetData());
 
-	return bDiscrete;
-}
-
-void FVulkanDevice::InitGPU(int32 DeviceIndex)
-{
-	LLM_SCOPE_VULKAN(ELLMTagVulkan::VulkanMisc);
-
-	// Query features
+	// Query base features
 	VulkanRHI::vkGetPhysicalDeviceFeatures(Gpu, &PhysicalFeatures);
 
+	// Setup layers and extensions
+	FVulkanDeviceExtensionArray UEExtensions = FVulkanDeviceExtension::GetUESupportedDeviceExtensions(this);
+	TArray<const ANSICHAR*> DeviceLayers = FVulkanDevice::SetupDeviceLayers(Gpu, UEExtensions);
+
 #if VULKAN_SUPPORTS_PHYSICAL_DEVICE_PROPERTIES2
 	if (RHI->GetOptionalExtensions().HasKHRGetPhysicalDeviceProperties2)
 	{
-		VkPhysicalDeviceFeatures2 Features2;
-		ZeroVulkanStruct(Features2, VK_STRUCTURE_TYPE_PHYSICAL_DEVICE_FEATURES_2);
-
-		void** NextPropsAddr = nullptr;
-		NextPropsAddr = &Features2.pNext;
-
-#if VULKAN_SUPPORTS_BUFFER_64BIT_ATOMICS
-		VkPhysicalDeviceShaderAtomicInt64FeaturesKHR BufferAtomicFeatures;
-		ZeroVulkanStruct(BufferAtomicFeatures, VK_STRUCTURE_TYPE_PHYSICAL_DEVICE_SHADER_ATOMIC_INT64_FEATURES_KHR);
-		if (OptionalDeviceExtensions.HasKHRShaderAtomicInt64)  // Set by QueryGPU if extension is supported
+		// Query advanced features
 		{
-			*NextPropsAddr = &BufferAtomicFeatures;
-			NextPropsAddr = &BufferAtomicFeatures.pNext;
-		}
-#endif
+			VkPhysicalDeviceFeatures2 PhysicalDeviceFeatures2;
+			ZeroVulkanStruct(PhysicalDeviceFeatures2, VK_STRUCTURE_TYPE_PHYSICAL_DEVICE_FEATURES_2);
 
-#if VULKAN_SUPPORTS_FRAGMENT_DENSITY_MAP
-		if (GetOptionalExtensions().HasEXTFragmentDensityMap)
-		{
-			*NextPropsAddr = &FragmentDensityMapFeatures;
-			NextPropsAddr = &FragmentDensityMapFeatures.pNext;
-			ZeroVulkanStruct(FragmentDensityMapFeatures, VK_STRUCTURE_TYPE_PHYSICAL_DEVICE_FRAGMENT_DENSITY_MAP_FEATURES_EXT);
-		}
-#endif
+			for (TUniquePtr<FVulkanDeviceExtension>& UEExtension : UEExtensions)
+			{
+				if (UEExtension->InUse())
+				{
+					UEExtension->PrePhysicalDeviceFeatures(PhysicalDeviceFeatures2);
+				}
+			}
 
-#if VULKAN_SUPPORTS_FRAGMENT_DENSITY_MAP2
-		if (GetOptionalExtensions().HasEXTFragmentDensityMap2)
-		{
-			*NextPropsAddr = &FragmentDensityMap2Features;
-			NextPropsAddr = &FragmentDensityMap2Features.pNext;
-			ZeroVulkanStruct(FragmentDensityMap2Features, VK_STRUCTURE_TYPE_PHYSICAL_DEVICE_FRAGMENT_DENSITY_MAP_2_FEATURES_EXT);
-		}
-#endif
+			VulkanRHI::vkGetPhysicalDeviceFeatures2KHR(Gpu, &PhysicalDeviceFeatures2);
 
-#if VULKAN_SUPPORTS_FRAGMENT_SHADING_RATE
-		if (GetOptionalExtensions().HasKHRFragmentShadingRate)
-		{
-			*NextPropsAddr = &FragmentShadingRateFeatures;
-			NextPropsAddr = &FragmentShadingRateFeatures.pNext;
-			ZeroVulkanStruct(FragmentShadingRateFeatures, VK_STRUCTURE_TYPE_PHYSICAL_DEVICE_FRAGMENT_SHADING_RATE_FEATURES_KHR);
-
-			// While we're here, enumerate the available shading rates.
-			uint32 FragmentShadingRateCount = 0;
-			VulkanRHI::vkGetPhysicalDeviceFragmentShadingRatesKHR(Gpu, &FragmentShadingRateCount, nullptr);
-			if (FragmentShadingRateCount != 0)
+			for (TUniquePtr<FVulkanDeviceExtension>& UEExtension : UEExtensions)
 			{
-				FragmentShadingRates.SetNum(FragmentShadingRateCount);
-				for (uint32 i = 0; i < FragmentShadingRateCount; ++i)
+				if (UEExtension->InUse())
 				{
-					ZeroVulkanStruct(FragmentShadingRates[i], VK_STRUCTURE_TYPE_PHYSICAL_DEVICE_FRAGMENT_SHADING_RATE_KHR);
+					UEExtension->PostPhysicalDeviceFeatures(OptionalDeviceExtensions);
 				}
-				VulkanRHI::vkGetPhysicalDeviceFragmentShadingRatesKHR(Gpu, &FragmentShadingRateCount, FragmentShadingRates.GetData());
 			}
 		}
-#endif
-		
-#if VULKAN_SUPPORTS_MULTIVIEW
-		if (GetOptionalExtensions().HasKHRMultiview)
-		{
-			*NextPropsAddr = &MultiviewFeatures;
-			NextPropsAddr = &MultiviewFeatures.pNext;
-			ZeroVulkanStruct(MultiviewFeatures, VK_STRUCTURE_TYPE_PHYSICAL_DEVICE_MULTIVIEW_FEATURES);
-		}
-#endif
-
-#if VULKAN_SUPPORTS_IMAGE_64BIT_ATOMICS
-		VkPhysicalDeviceShaderImageAtomicInt64FeaturesEXT ImageAtomicFeatures;
-		ZeroVulkanStruct(ImageAtomicFeatures, VK_STRUCTURE_TYPE_PHYSICAL_DEVICE_SHADER_IMAGE_ATOMIC_INT64_FEATURES_EXT);
-		if (OptionalDeviceExtensions.HasImageAtomicInt64)  // Set by QueryGPU if extension is supported
-		{
-			*NextPropsAddr = &ImageAtomicFeatures;
-			NextPropsAddr = &ImageAtomicFeatures.pNext;
-		}
-#endif
 
-#if VULKAN_SUPPORTS_MAINTENANCE4
-		VkPhysicalDeviceMaintenance4FeaturesKHR Maintenance4Features;
-		ZeroVulkanStruct(Maintenance4Features, VK_STRUCTURE_TYPE_PHYSICAL_DEVICE_MAINTENANCE_4_FEATURES_KHR);
-		if (OptionalDeviceExtensions.HasKHRMaintenance4)  // Set by QueryGPU if extension is supported
+		// Query advances properties
 		{
-			*NextPropsAddr = &Maintenance4Features;
-			NextPropsAddr = &Maintenance4Features.pNext;
-		}
-#endif
+			VkPhysicalDeviceProperties2KHR PhysicalDeviceProperties2;
+			ZeroVulkanStruct(PhysicalDeviceProperties2, VK_STRUCTURE_TYPE_PHYSICAL_DEVICE_PROPERTIES_2_KHR);
+			PhysicalDeviceProperties2.pNext = &GpuIdProps;
+			ZeroVulkanStruct(GpuIdProps, VK_STRUCTURE_TYPE_PHYSICAL_DEVICE_ID_PROPERTIES_KHR);
 
-#if VULKAN_SUPPORTS_DESCRIPTOR_INDEXING
-		VkPhysicalDeviceDescriptorIndexingFeaturesEXT DescriptorIndexingFeatures;
-		ZeroVulkanStruct(DescriptorIndexingFeatures, VK_STRUCTURE_TYPE_PHYSICAL_DEVICE_DESCRIPTOR_INDEXING_FEATURES);
-		if (OptionalDeviceExtensions.HasEXTDescriptorIndexing)  // Set by QueryGPU if extension is supported
-		{
-			*NextPropsAddr = &DescriptorIndexingFeatures;
-			NextPropsAddr = &DescriptorIndexingFeatures.pNext;
-		}
-#endif
+			for (TUniquePtr<FVulkanDeviceExtension>& UEExtension : UEExtensions)
+			{
+				if (UEExtension->InUse())
+				{
+					UEExtension->PrePhysicalDeviceProperties(PhysicalDeviceProperties2);
+				}
+			}
 
-		VulkanRHI::vkGetPhysicalDeviceFeatures2KHR(Gpu, &Features2);
+			VulkanRHI::vkGetPhysicalDeviceProperties2KHR(Gpu, &PhysicalDeviceProperties2);
 
-		// If we have the extension but the feature is not supported, clear the bits.
-#if VULKAN_SUPPORTS_BUFFER_64BIT_ATOMICS
-		OptionalDeviceExtensions.HasKHRShaderAtomicInt64 = (BufferAtomicFeatures.shaderBufferInt64Atomics == VK_TRUE);
-#endif
-#if VULKAN_SUPPORTS_IMAGE_64BIT_ATOMICS
-		OptionalDeviceExtensions.HasImageAtomicInt64 = (ImageAtomicFeatures.shaderImageInt64Atomics == VK_TRUE);
-#endif
-#if VULKAN_SUPPORTS_MAINTENANCE4
-		OptionalDeviceExtensions.HasKHRMaintenance4 = (Maintenance4Features.maintenance4 == VK_TRUE);
-#endif
-#if VULKAN_SUPPORTS_DESCRIPTOR_INDEXING
-		OptionalDeviceExtensions.HasEXTDescriptorIndexing = (DescriptorIndexingFeatures.runtimeDescriptorArray == VK_TRUE); // :todo-jn: add resource specific checks
-#endif
+			for (TUniquePtr<FVulkanDeviceExtension>& UEExtension : UEExtensions)
+			{
+				if (UEExtension->InUse())
+				{
+					UEExtension->PostPhysicalDeviceProperties();
+				}
+			}
+		}
 	}
 #endif // VULKAN_SUPPORTS_PHYSICAL_DEVICE_PROPERTIES2
 
 	UE_LOG(LogVulkanRHI, Display, TEXT("Using Device %d: Geometry %d BufferAtomic64 %d ImageAtomic64 %d"), DeviceIndex, PhysicalFeatures.geometryShader, OptionalDeviceExtensions.HasKHRShaderAtomicInt64, OptionalDeviceExtensions.HasImageAtomicInt64);
 
-	CreateDevice();
+	CreateDevice(DeviceLayers, UEExtensions);
 
 	SetupFormats();
 
@@ -1385,25 +1108,6 @@ void FVulkanDevice::InitGPU(int32 DeviceIndex)
 	}
 #endif
 
-#if VULKAN_SUPPORTS_FRAGMENT_SHADING_RATE
-	if (GetOptionalExtensions().HasKHRFragmentShadingRate)
-	{
-		GRHISupportsAttachmentVariableRateShading = FragmentShadingRateFeatures.attachmentFragmentShadingRate ? true : false;
-		GRHISupportsPipelineVariableRateShading = FragmentShadingRateFeatures.pipelineFragmentShadingRate ? true : false;
-
-		if (GRHISupportsAttachmentVariableRateShading)
-		{
-			GRHIVariableRateShadingImageDataType = VRSImage_Palette;
-			GRHIVariableRateShadingImageFormat = PF_R8_UINT;
-		}
-		else
-		{
-			GRHIVariableRateShadingImageDataType = VRSImage_NotSupported;
-			GRHIVariableRateShadingImageFormat = PF_Unknown;
-		}
-	}
-#endif
-
 	PipelineStateCache->InitAndLoad(CacheFilenames);
 
 	// Setup default resource
diff --git a/Engine/Source/Runtime/VulkanRHI/Private/VulkanDevice.h b/Engine/Source/Runtime/VulkanRHI/Private/VulkanDevice.h
--- a/Engine/Source/Runtime/VulkanRHI/Private/VulkanDevice.h
+++ b/Engine/Source/Runtime/VulkanRHI/Private/VulkanDevice.h
@@ -13,6 +13,7 @@ class FVulkanDescriptorPool;
 class FVulkanDescriptorPoolsManager;
 class FVulkanCommandListContextImmediate;
 class FVulkanTransientHeapCache;
+class FVulkanDeviceExtension;
 #if VULKAN_USE_NEW_QUERIES
 class FVulkanOcclusionQueryPool;
 #else
@@ -62,11 +63,9 @@ struct FOptionalVulkanDeviceExtensions
 			uint64 HasKHRMaintenance1 : 1;
 			uint64 HasKHRMaintenance2 : 1;
 			uint64 HasKHRDedicatedAllocation : 1;
-			uint64 HasYcbcrSampler : 1;  // Combination of VK_KHR_sampler_ycbcr_conversion + VK_KHR_bind_memory2 + VK_KHR_get_memory_requirements2
 			uint64 HasKHRMultiview : 1;
 
 			// Promoted to 1.2
-			uint64 HasDriverProperties : 1;
 			uint64 HasKHRRenderPass2 : 1;
 			uint64 HasKHRImageFormatList : 1;
 			uint64 HasKHRShaderAtomicInt64 : 1;
@@ -76,6 +75,7 @@ struct FOptionalVulkanDeviceExtensions
 			uint64 HasShaderFloatControls : 1;
 			uint64 HasEXTDescriptorIndexing : 1;
 			uint64 HasEXTShaderViewportIndexLayer : 1;
+			uint64 HasSeparateDepthStencilLayouts : 1;
 
 			// Promoted to 1.3
 			uint64 HasEXTTextureCompressionASTCHDR : 1;
@@ -90,8 +90,6 @@ struct FOptionalVulkanDeviceExtensions
 		Packed = 0;
 	}
 
-	void Setup(const TArray<const ANSICHAR*>& InDeviceExtensions);
-
 	inline bool HasGPUCrashDumpExtensions() const
 	{
 		return HasAMDBufferMarker || HasNVDiagnosticCheckpoints;
@@ -112,23 +110,6 @@ struct FOptionalVulkanDeviceExtensions
 #endif
 };
 
-struct FOptionalVulkanDeviceFeatures
-{
-#if VULKAN_SUPPORTS_SEPARATE_DEPTH_STENCIL_LAYOUTS
-	VkPhysicalDeviceSeparateDepthStencilLayoutsFeaturesKHR SeparateDepthStencilLayoutsFeatures;
-#endif
-#if VULKAN_SUPPORTS_SCALAR_BLOCK_LAYOUT
-	VkPhysicalDeviceScalarBlockLayoutFeaturesEXT ScalarBlockLayoutFeatures;
-#endif
-#if VULKAN_RHI_RAYTRACING
-	VkPhysicalDeviceBufferDeviceAddressFeaturesKHR BufferDeviceAddressFeatures;
-	VkPhysicalDeviceAccelerationStructureFeaturesKHR AccelerationStructureFeatures;
-	VkPhysicalDeviceRayTracingPipelineFeaturesKHR RayTracingPipelineFeatures;
-	VkPhysicalDeviceDescriptorIndexingFeaturesEXT DescriptorIndexingFeatures;
-	VkPhysicalDeviceRayQueryFeaturesKHR RayQueryFeatures;
-#endif
-};
-
 #if VULKAN_RHI_RAYTRACING
 struct FRayTracingProperties
 {
@@ -213,12 +194,9 @@ public:
 
 	~FVulkanDevice();
 
-	// Returns true if this is a viable candidate for main GPU
-	bool QueryGPU(int32 DeviceIndex);
-
 	void InitGPU(int32 DeviceIndex);
 
-	void CreateDevice();
+	void CreateDevice(TArray<const ANSICHAR*>& DeviceLayers, FVulkanDeviceExtensionArray& UEExtensions);
 
 	void PrepareForDestroy();
 	void Destroy();
@@ -281,41 +259,10 @@ public:
 		return GpuProps;
 	}
 
-#if VULKAN_SUPPORTS_FRAGMENT_DENSITY_MAP
-	inline const VkPhysicalDeviceFragmentDensityMapFeaturesEXT& GetFragmentDensityMapFeatures() const
-	{
-		return FragmentDensityMapFeatures;
-	}
-#endif
-
-#if VULKAN_SUPPORTS_FRAGMENT_DENSITY_MAP2
-	inline const VkPhysicalDeviceFragmentDensityMap2FeaturesEXT& GetFragmentDensityMap2Features() const
-	{
-		return FragmentDensityMap2Features;
-	}
-#endif
-
 #if VULKAN_SUPPORTS_FRAGMENT_SHADING_RATE
-	inline const VkPhysicalDeviceFragmentShadingRateFeaturesKHR& GetFragmentShadingRateFeatures() const
-	{
-		return FragmentShadingRateFeatures;
-	}
-
-	inline const VkPhysicalDeviceFragmentShadingRatePropertiesKHR& GetFragmentShadingRateProperties() const
-	{
-		return FragmentShadingRateProperties;
-	}
-
 	VkExtent2D GetBestMatchedShadingRateExtents(EVRSShadingRate Rate) const;
 #endif
 	
-#if VULKAN_SUPPORTS_MULTIVIEW
-	inline const VkPhysicalDeviceMultiviewFeatures& GetMultiviewFeatures() const
-	{
-		return MultiviewFeatures;
-	}
-#endif
-
 	inline const VkPhysicalDeviceLimits& GetLimits() const
 	{
 		return GpuProps.limits;
@@ -331,7 +278,7 @@ public:
 #if VULKAN_RHI_RAYTRACING
 	inline const FRayTracingProperties& GetRayTracingProperties() const
 	{
-		check(OptionalDeviceExtensions.HasRaytracingExtensions());
+		check(OptionalDeviceExtensions.HasRaytracingExtensions() || (Device == VK_NULL_HANDLE));
 		return RayTracingProperties;
 	}
 
@@ -352,11 +299,6 @@ public:
 		return PhysicalFeatures;
 	}
 
-	inline bool HasSeparateDepthStencilLayouts() const
-	{
-		return bHasSeparateDepthStencilLayouts;
-	}
-
 	inline bool HasUnifiedMemory() const
 	{
 		return DeviceMemoryManager.HasUnifiedMemory();
@@ -506,11 +448,6 @@ public:
 		return OptionalDeviceExtensions;
 	}
 
-	inline FOptionalVulkanDeviceFeatures& GetOptionalFeatures()
-	{
-		return OptionalFeatures;
-	}
-
 #if VULKAN_SUPPORTS_GPU_CRASH_DUMPS
 	VkBuffer GetCrashMarkerBuffer() const
 	{
@@ -532,6 +469,10 @@ public:
 
 	FVulkanTransientHeapCache& GetOrCreateTransientHeapCache();
 
+	inline void SetDebugMarkersFound()
+	{
+		bDebugMarkersFound = true;
+	}
 
 private:
 	const VkFormatProperties& GetFormatProperties(VkFormat InFormat);
@@ -571,24 +512,11 @@ private:
 
 	VkPhysicalDevice Gpu;
 	VkPhysicalDeviceProperties GpuProps;
-#if VULKAN_SUPPORTS_FRAGMENT_DENSITY_MAP
-	VkPhysicalDeviceFragmentDensityMapFeaturesEXT FragmentDensityMapFeatures;
-#endif
-
-#if VULKAN_SUPPORTS_FRAGMENT_DENSITY_MAP2
-	VkPhysicalDeviceFragmentDensityMap2FeaturesEXT FragmentDensityMap2Features;
-#endif
 
 #if VULKAN_SUPPORTS_FRAGMENT_SHADING_RATE
-	VkPhysicalDeviceFragmentShadingRatePropertiesKHR FragmentShadingRateProperties;
-	VkPhysicalDeviceFragmentShadingRateFeaturesKHR FragmentShadingRateFeatures;
 	TArray<VkPhysicalDeviceFragmentShadingRateKHR> FragmentShadingRates;
 #endif
 
-#if VULKAN_SUPPORTS_MULTIVIEW
-	VkPhysicalDeviceMultiviewFeatures MultiviewFeatures;
-#endif
-
 #if VULKAN_SUPPORTS_PHYSICAL_DEVICE_PROPERTIES2
 	VkPhysicalDeviceIDPropertiesKHR GpuIdProps;
 
@@ -599,9 +527,6 @@ private:
 #endif // VULKAN_SUPPORTS_PHYSICAL_DEVICE_PROPERTIES2
 
 	VkPhysicalDeviceFeatures PhysicalFeatures;
-	FOptionalVulkanDeviceFeatures OptionalFeatures;
-
-	bool bHasSeparateDepthStencilLayouts = false;
 
 	TArray<VkQueueFamilyProperties> QueueFamilyProps;
 	VkFormatProperties FormatProperties[VK_FORMAT_RANGE_SIZE];
@@ -640,10 +565,8 @@ private:
 
 	FVulkanDynamicRHI* RHI = nullptr;
 	bool bDebugMarkersFound = false;
-	TArray<const ANSICHAR*> DeviceExtensions;
-	TArray<const ANSICHAR*> ValidationLayers;
-
-	static void GetDeviceExtensionsAndLayers(VkPhysicalDevice Gpu, EGpuVendorId VendorId, TArray<const ANSICHAR*>& OutDeviceExtensions, TArray<const ANSICHAR*>& OutDeviceLayers, TArray<FString>& OutAllDeviceExtensions, TArray<FString>& OutAllDeviceLayers, bool& bOutDebugMarkers);
+	
+	static TArray<const ANSICHAR*> SetupDeviceLayers(VkPhysicalDevice Gpu, FVulkanDeviceExtensionArray& UEExtensions);
 
 	FOptionalVulkanDeviceExtensions OptionalDeviceExtensions;
 
diff --git a/Engine/Source/Runtime/VulkanRHI/Private/VulkanExtensions.cpp b/Engine/Source/Runtime/VulkanRHI/Private/VulkanExtensions.cpp
new file mode 100644
--- /dev/null
+++ b/Engine/Source/Runtime/VulkanRHI/Private/VulkanExtensions.cpp
@@ -0,0 +1,1131 @@
+// Copyright Epic Games, Inc. All Rights Reserved.
+
+#include "VulkanRHIPrivate.h"
+#include "VulkanExtensions.h"
+
+#include "IHeadMountedDisplayModule.h"
+#include "IHeadMountedDisplayVulkanExtensions.h"
+
+
+// ADDING A NEW EXTENSION:
+// 
+// A - If the extension simply needs to be queried for driver support and added at device creation (and set a flag):
+//     Add an FVulkanDeviceExtension/FVulkanInstanceExtension directly in the array returned by GetUESupportedDeviceExtensions/GetUESupportedInstanceExtensions.
+// 
+// B - The extension requires the use of specialized Feature and/or Property structures, or other extended processing:
+//     * Add a [instance/device] class for your extension, try to call it FVulkan[ExtensionNameInCamelCaps]Extension so we have consistency
+//     * Feed it the extension name string, use the SDK's define if possible : VK_[BLA_BLA_BLA]_EXTENSION_NAME
+//     * If there is a VULKAN_SUPPORTS_* define that enables/disables your extension in code, provide it so that the extension knows if it should be used or not.
+//       We don't IFDEF the extensions code so that we still know of its existence and we can report warnings if someone tries to use a disabled extension.
+//       If there is no define to enable/disable the extension in code, then simply use VULKAN_EXTENSION_ENABLED.
+// 
+// Tips:
+// - Feature and Property structures specific to extensions that won't be needed beyond init should be included in the extension's class instead of the device.
+// - To add engine support for a complex extension, but require some kind of external activation (eg plugin) you can use its EExtensionActivation state (see header definition).
+// - If an extension is supported on multiple platforms, it may be cleaner to include it here and simply disable its VULKAN_SUPPORTS_* value in the Vulkan platform header where it's not supported.
+
+
+
+TAutoConsoleVariable<int32> GRHIAllow64bitShaderAtomicsCvar(
+	TEXT("r.Vulkan.Allow64bitShaderAtomics"),
+	1,
+	TEXT("Whether to enable 64bit buffer/image atomics required by Nanite\n")
+	TEXT("0 to disable 64bit atomics\n")
+	TEXT("1 to enable (default)"),
+	ECVF_ReadOnly
+);
+
+TAutoConsoleVariable<int32> GVulkanRayTracingCVar(
+	TEXT("r.Vulkan.RayTracing"),
+	0,
+	TEXT("0: Do not enable Vulkan ray tracing extensions (default)\n")
+	TEXT("1: Enable experimental ray tracing support (for development and testing purposes)"),
+	ECVF_ReadOnly
+);
+
+extern TAutoConsoleVariable<int32> GGPUValidationCvar;
+
+TSharedPtr<IHeadMountedDisplayVulkanExtensions, ESPMode::ThreadSafe> FVulkanDynamicRHI::HMDVulkanExtensions;
+TArray<const ANSICHAR*> FVulkanDeviceExtension::ExternalExtensions;
+TArray<const ANSICHAR*> FVulkanInstanceExtension::ExternalExtensions;
+
+
+template <typename ExistingChainType, typename NewStructType>
+static void AddToPNext(ExistingChainType& Existing, NewStructType& Added)
+{
+	Added.pNext = (void*)Existing.pNext;
+	Existing.pNext = (void*)&Added;
+}
+
+
+#define VERIFYVULKANRESULT_INIT(VkFunction)	{ const VkResult ScopedResult = VkFunction; \
+												if (ScopedResult == VK_ERROR_INITIALIZATION_FAILED) { \
+													UE_LOG(LogVulkanRHI, Error, \
+													TEXT("%s failed\n at %s:%u\nThis typically means Vulkan is not properly set up in your system; try running vulkaninfo from the Vulkan SDK."), \
+													ANSI_TO_TCHAR(#VkFunction), ANSI_TO_TCHAR(__FILE__), __LINE__); } \
+												else if (ScopedResult < VK_SUCCESS) { \
+													VulkanRHI::VerifyVulkanResult(ScopedResult, #VkFunction, __FILE__, __LINE__); }}
+
+
+TArray<VkExtensionProperties> FVulkanDeviceExtension::GetDriverSupportedDeviceExtensions(VkPhysicalDevice Gpu, const ANSICHAR* LayerName)
+{
+	TArray<VkExtensionProperties> OutDeviceExtensions;
+	uint32 Count = 0;
+	VERIFYVULKANRESULT_INIT(VulkanRHI::vkEnumerateDeviceExtensionProperties(Gpu, LayerName, &Count, nullptr));
+	if (Count > 0)
+	{
+		OutDeviceExtensions.AddZeroed(Count);
+		VERIFYVULKANRESULT_INIT(VulkanRHI::vkEnumerateDeviceExtensionProperties(Gpu, LayerName, &Count, OutDeviceExtensions.GetData()));
+	}
+	OutDeviceExtensions.Sort([](const VkExtensionProperties& A, const VkExtensionProperties& B) { return FCStringAnsi::Strcmp(A.extensionName, B.extensionName) < 0; });
+	return OutDeviceExtensions;
+}
+
+TArray<VkExtensionProperties> FVulkanInstanceExtension::GetDriverSupportedInstanceExtensions(const ANSICHAR* LayerName)
+{
+	TArray<VkExtensionProperties> OutInstanceExtensions;
+	uint32 Count = 0;
+	VERIFYVULKANRESULT_INIT(VulkanRHI::vkEnumerateInstanceExtensionProperties(LayerName, &Count, nullptr));
+	if (Count > 0)
+	{
+		OutInstanceExtensions.AddZeroed(Count);
+		VERIFYVULKANRESULT_INIT(VulkanRHI::vkEnumerateInstanceExtensionProperties(LayerName, &Count, OutInstanceExtensions.GetData()));
+	}
+	OutInstanceExtensions.Sort([](const VkExtensionProperties& A, const VkExtensionProperties& B) { return FCStringAnsi::Strcmp(A.extensionName, B.extensionName) < 0; });
+	return OutInstanceExtensions;
+}
+
+#undef VERIFYVULKANRESULT_INIT
+
+
+
+
+// *** Vulkan Device Extension support ***
+// Typical flow:
+// 1- For the selected rendering device, the engine will query the supported extensions (FVulkanDeviceExtension are constructed and support is queried)
+// 2- Followed by a query of Physical Device Features (PrePhysicalDeviceFeatures, PostPhysicalDeviceFeatures)
+// 3- Followed by a query of Physical Device Properties (PrePhysicalDeviceProperties, PostPhysicalDeviceProperties)
+// 4- Finally, the device is created (PreCreateDevice)
+
+
+
+// ***** VK_KHR_maintenance4
+class FVulkanKHRMaintenance4Extension : public FVulkanDeviceExtension
+{
+public:
+
+	FVulkanKHRMaintenance4Extension(FVulkanDevice* InDevice)
+		: FVulkanDeviceExtension(InDevice, VK_KHR_MAINTENANCE_4_EXTENSION_NAME, VULKAN_SUPPORTS_MAINTENANCE4, VK_API_VERSION_1_3)
+	{}
+
+	virtual void PrePhysicalDeviceFeatures(VkPhysicalDeviceFeatures2KHR& PhysicalDeviceFeatures2) override final
+	{
+		ZeroVulkanStruct(Maintenance4Features, VK_STRUCTURE_TYPE_PHYSICAL_DEVICE_MAINTENANCE_4_FEATURES_KHR);
+		AddToPNext(PhysicalDeviceFeatures2, Maintenance4Features);
+	}
+
+	virtual void PostPhysicalDeviceFeatures(FOptionalVulkanDeviceExtensions& ExtensionFlags) override final
+	{
+		ExtensionFlags.HasKHRMaintenance4 = (Maintenance4Features.maintenance4 == VK_TRUE) ? 1 : 0; 
+	}
+
+private:
+	VkPhysicalDeviceMaintenance4FeaturesKHR Maintenance4Features;
+};
+
+
+
+// ***** VK_KHR_driver_properties
+class FVulkanKHRDriverPropertiesExtension : public FVulkanDeviceExtension
+{
+public:
+
+	FVulkanKHRDriverPropertiesExtension(FVulkanDevice* InDevice)
+		: FVulkanDeviceExtension(InDevice, VK_KHR_DRIVER_PROPERTIES_EXTENSION_NAME, VULKAN_SUPPORTS_DRIVER_PROPERTIES, VK_API_VERSION_1_2)
+	{}
+
+	virtual void PrePhysicalDeviceProperties(VkPhysicalDeviceProperties2KHR& PhysicalDeviceProperties2) override final
+	{
+		ZeroVulkanStruct(PhysicalDeviceDriverProperties, VK_STRUCTURE_TYPE_PHYSICAL_DEVICE_DRIVER_PROPERTIES_KHR);
+		AddToPNext(PhysicalDeviceProperties2, PhysicalDeviceDriverProperties);
+	}
+
+	virtual void PostPhysicalDeviceProperties() override final
+	{
+		UE_LOG(LogVulkanRHI, Display, TEXT("Vulkan Physical Device Driver Properties:"));
+		UE_LOG(LogVulkanRHI, Display, TEXT("- driverName: %s"), ANSI_TO_TCHAR(PhysicalDeviceDriverProperties.driverName));
+		UE_LOG(LogVulkanRHI, Display, TEXT("- driverInfo: %s"), ANSI_TO_TCHAR(PhysicalDeviceDriverProperties.driverInfo));
+	}
+
+private:
+	VkPhysicalDeviceDriverPropertiesKHR PhysicalDeviceDriverProperties;
+};
+
+
+
+// ***** VK_KHR_shader_atomic_int64
+class FVulkanShaderAtomicInt64Extension : public FVulkanDeviceExtension
+{
+public:
+
+	FVulkanShaderAtomicInt64Extension(FVulkanDevice* InDevice)
+		: FVulkanDeviceExtension(InDevice, VK_KHR_SHADER_ATOMIC_INT64_EXTENSION_NAME, VULKAN_SUPPORTS_BUFFER_64BIT_ATOMICS, VK_API_VERSION_1_2)
+	{
+		bEnabledInCode = bEnabledInCode && (GRHIAllow64bitShaderAtomicsCvar.GetValueOnAnyThread() != 0);
+	}
+
+#if VULKAN_SUPPORTS_BUFFER_64BIT_ATOMICS
+	virtual void PrePhysicalDeviceFeatures(VkPhysicalDeviceFeatures2KHR& PhysicalDeviceFeatures2) override final
+	{
+		ZeroVulkanStruct(BufferAtomicFeatures, VK_STRUCTURE_TYPE_PHYSICAL_DEVICE_SHADER_ATOMIC_INT64_FEATURES_KHR);
+		AddToPNext(PhysicalDeviceFeatures2, BufferAtomicFeatures);
+	}
+
+	virtual void PostPhysicalDeviceFeatures(FOptionalVulkanDeviceExtensions& ExtensionFlags) override final
+	{
+		ExtensionFlags.HasKHRShaderAtomicInt64 = BufferAtomicFeatures.shaderBufferInt64Atomics;
+		GRHISupportsAtomicUInt64 = (BufferAtomicFeatures.shaderBufferInt64Atomics == VK_TRUE);
+	}
+
+	virtual void PreCreateDevice(VkDeviceCreateInfo& DeviceCreateInfo) override final
+	{
+		if (BufferAtomicFeatures.shaderBufferInt64Atomics == VK_TRUE)
+		{
+			AddToPNext(DeviceCreateInfo, BufferAtomicFeatures);
+		}
+		else
+		{
+			GRHISupportsAtomicUInt64 = false;
+		}
+	}
+
+	VkPhysicalDeviceShaderAtomicInt64Features BufferAtomicFeatures;
+#endif // VULKAN_SUPPORTS_BUFFER_64BIT_ATOMICS
+};
+
+
+
+// ***** VK_EXT_shader_image_atomic_int64
+class FVulkanShaderImageAtomicInt64Extension : public FVulkanDeviceExtension
+{
+public:
+
+	FVulkanShaderImageAtomicInt64Extension(FVulkanDevice* InDevice)
+		: FVulkanDeviceExtension(InDevice, VK_EXT_SHADER_IMAGE_ATOMIC_INT64_EXTENSION_NAME, VULKAN_SUPPORTS_IMAGE_64BIT_ATOMICS)
+	{
+		bEnabledInCode = bEnabledInCode && (GRHIAllow64bitShaderAtomicsCvar.GetValueOnAnyThread() != 0);
+	}
+
+#if VULKAN_SUPPORTS_IMAGE_64BIT_ATOMICS
+	virtual void PrePhysicalDeviceFeatures(VkPhysicalDeviceFeatures2KHR& PhysicalDeviceFeatures2) override final
+	{
+		ZeroVulkanStruct(ImageAtomicFeatures, VK_STRUCTURE_TYPE_PHYSICAL_DEVICE_SHADER_IMAGE_ATOMIC_INT64_FEATURES_EXT);
+		AddToPNext(PhysicalDeviceFeatures2, ImageAtomicFeatures);
+	}
+
+	virtual void PostPhysicalDeviceFeatures(FOptionalVulkanDeviceExtensions& ExtensionFlags) override final
+	{
+		ExtensionFlags.HasImageAtomicInt64 = ImageAtomicFeatures.shaderImageInt64Atomics;
+		GRHISupportsAtomicUInt64 = GRHISupportsAtomicUInt64 && (ImageAtomicFeatures.shaderImageInt64Atomics == VK_TRUE); // GRHISupportsAtomicUInt64 already contains FVulkanShaderAtomicInt64Extension support
+	}
+
+	virtual void PreCreateDevice(VkDeviceCreateInfo& DeviceCreateInfo) override final
+	{
+		const bool bEnable64bitAtomics = (GRHIAllow64bitShaderAtomicsCvar.GetValueOnAnyThread() != 0);
+		if (bEnable64bitAtomics && (ImageAtomicFeatures.shaderImageInt64Atomics == VK_TRUE))
+		{
+			ImageAtomicFeatures.pNext = (void*)DeviceCreateInfo.pNext;
+			DeviceCreateInfo.pNext = &ImageAtomicFeatures;
+		}
+		else
+		{
+			GRHISupportsAtomicUInt64 = false;
+		}
+	}
+
+	VkPhysicalDeviceShaderImageAtomicInt64FeaturesEXT ImageAtomicFeatures;
+#endif // VULKAN_SUPPORTS_PHYSICAL_DEVICE_PROPERTIES2
+};
+
+
+
+// ***** VK_EXT_shader_viewport_index_layer
+class FVulkanEXTShaderViewportIndexLayerExtension : public FVulkanDeviceExtension
+{
+public:
+
+	FVulkanEXTShaderViewportIndexLayerExtension(FVulkanDevice* InDevice)
+		: FVulkanDeviceExtension(InDevice, VK_EXT_SHADER_VIEWPORT_INDEX_LAYER_EXTENSION_NAME, VULKAN_SUPPORTS_SHADER_VIEWPORT_INDEX_LAYER, VK_API_VERSION_1_2)
+	{}
+
+	virtual void PostPhysicalDeviceFeatures(FOptionalVulkanDeviceExtensions& ExtensionFlags) override final
+	{
+		ExtensionFlags.HasEXTShaderViewportIndexLayer = 1; 
+
+		GRHISupportsArrayIndexFromAnyShader = true;
+	}
+};
+
+
+
+// ***** VK_KHR_separate_depth_stencil_layouts
+class FVulkanKHRSeparateDepthStencilLayoutsExtension : public FVulkanDeviceExtension
+{
+public:
+
+	FVulkanKHRSeparateDepthStencilLayoutsExtension(FVulkanDevice* InDevice)
+		: FVulkanDeviceExtension(InDevice, VK_KHR_SEPARATE_DEPTH_STENCIL_LAYOUTS_EXTENSION_NAME, VULKAN_SUPPORTS_SEPARATE_DEPTH_STENCIL_LAYOUTS, VK_API_VERSION_1_2)
+	{
+		// Disabled but kept for reference. The barriers code doesn't currently use separate transitions for depth and stencil.
+		bEnabledInCode = false;
+	}
+
+	virtual void PrePhysicalDeviceFeatures(VkPhysicalDeviceFeatures2KHR& PhysicalDeviceFeatures2) override final
+	{
+		ZeroVulkanStruct(SeparateDepthStencilLayoutsFeatures, VK_STRUCTURE_TYPE_PHYSICAL_DEVICE_SEPARATE_DEPTH_STENCIL_LAYOUTS_FEATURES_KHR);
+		AddToPNext(PhysicalDeviceFeatures2, SeparateDepthStencilLayoutsFeatures);
+	}
+
+	virtual void PostPhysicalDeviceFeatures(FOptionalVulkanDeviceExtensions& ExtensionFlags) override final
+	{
+		ExtensionFlags.HasSeparateDepthStencilLayouts = SeparateDepthStencilLayoutsFeatures.separateDepthStencilLayouts;
+	}
+
+	virtual void PreCreateDevice(VkDeviceCreateInfo& DeviceCreateInfo) override final
+	{
+		if (SeparateDepthStencilLayoutsFeatures.separateDepthStencilLayouts == VK_TRUE)
+		{
+			AddToPNext(DeviceCreateInfo, SeparateDepthStencilLayoutsFeatures);
+		}
+	}
+
+private:
+	VkPhysicalDeviceSeparateDepthStencilLayoutsFeaturesKHR SeparateDepthStencilLayoutsFeatures;
+};
+
+
+
+// ***** VK_KHR_multiview
+class FVulkanKHRMultiviewExtension : public FVulkanDeviceExtension
+{
+public:
+
+	FVulkanKHRMultiviewExtension(FVulkanDevice* InDevice)
+		: FVulkanDeviceExtension(InDevice, VK_KHR_MULTIVIEW_EXTENSION_NAME, VULKAN_SUPPORTS_MULTIVIEW, VK_API_VERSION_1_1)
+	{}
+
+	virtual void PrePhysicalDeviceFeatures(VkPhysicalDeviceFeatures2KHR& PhysicalDeviceFeatures2) override final
+	{
+		ZeroVulkanStruct(DeviceMultiviewFeatures, VK_STRUCTURE_TYPE_PHYSICAL_DEVICE_MULTIVIEW_FEATURES);
+		AddToPNext(PhysicalDeviceFeatures2, DeviceMultiviewFeatures);
+	}
+
+	virtual void PostPhysicalDeviceFeatures(FOptionalVulkanDeviceExtensions& ExtensionFlags) override final 
+	{
+		ExtensionFlags.HasKHRMultiview = DeviceMultiviewFeatures.multiview; 
+	}
+
+	virtual void PreCreateDevice(VkDeviceCreateInfo& DeviceCreateInfo) override final
+	{
+		if (DeviceMultiviewFeatures.multiview == VK_TRUE)
+		{
+			AddToPNext(DeviceCreateInfo, DeviceMultiviewFeatures);
+		}
+	}
+
+private:
+	VkPhysicalDeviceMultiviewFeatures DeviceMultiviewFeatures;
+};
+
+
+
+// ***** VK_EXT_scalar_block_layout
+class FVulkanEXTScalarBlockLayoutExtension : public FVulkanDeviceExtension
+{
+public:
+
+	FVulkanEXTScalarBlockLayoutExtension(FVulkanDevice* InDevice)
+		: FVulkanDeviceExtension(InDevice, VK_EXT_SCALAR_BLOCK_LAYOUT_EXTENSION_NAME, VULKAN_SUPPORTS_SCALAR_BLOCK_LAYOUT, VK_API_VERSION_1_2)
+	{}
+
+	virtual void PrePhysicalDeviceFeatures(VkPhysicalDeviceFeatures2KHR& PhysicalDeviceFeatures2) override final
+	{
+		ZeroVulkanStruct(ScalarBlockLayoutFeatures, VK_STRUCTURE_TYPE_PHYSICAL_DEVICE_SCALAR_BLOCK_LAYOUT_FEATURES_EXT);
+		AddToPNext(PhysicalDeviceFeatures2, ScalarBlockLayoutFeatures);
+	}
+
+	virtual void PostPhysicalDeviceFeatures(FOptionalVulkanDeviceExtensions& ExtensionFlags) override final 
+	{
+		ExtensionFlags.HasEXTScalarBlockLayout = ScalarBlockLayoutFeatures.scalarBlockLayout; 
+	}
+
+	virtual void PreCreateDevice(VkDeviceCreateInfo& DeviceCreateInfo) override final
+	{
+		if (ScalarBlockLayoutFeatures.scalarBlockLayout == VK_TRUE)
+		{
+			AddToPNext(DeviceCreateInfo, ScalarBlockLayoutFeatures);
+		}
+	}
+
+private:
+	VkPhysicalDeviceScalarBlockLayoutFeaturesEXT ScalarBlockLayoutFeatures;
+};
+
+
+
+// ***** VK_EXT_descriptor_indexing
+class FVulkanEXTDescriptorIndexingExtension : public FVulkanDeviceExtension
+{
+public:
+
+	FVulkanEXTDescriptorIndexingExtension(FVulkanDevice* InDevice)
+		: FVulkanDeviceExtension(InDevice, VK_EXT_DESCRIPTOR_INDEXING_EXTENSION_NAME, VULKAN_SUPPORTS_DESCRIPTOR_INDEXING, VK_API_VERSION_1_2)
+	{}
+
+	virtual void PrePhysicalDeviceFeatures(VkPhysicalDeviceFeatures2KHR& PhysicalDeviceFeatures2) override final
+	{
+		ZeroVulkanStruct(DescriptorIndexingFeatures, VK_STRUCTURE_TYPE_PHYSICAL_DEVICE_DESCRIPTOR_INDEXING_FEATURES);
+		AddToPNext(PhysicalDeviceFeatures2, DescriptorIndexingFeatures);
+	}
+
+	virtual void PostPhysicalDeviceFeatures(FOptionalVulkanDeviceExtensions& ExtensionFlags) override final
+	{
+		ExtensionFlags.HasEXTDescriptorIndexing = DescriptorIndexingFeatures.runtimeDescriptorArray; // :todo-jn: add resource specific checks
+	}
+
+	virtual void PreCreateDevice(VkDeviceCreateInfo& DeviceCreateInfo) override final
+	{
+		if (DescriptorIndexingFeatures.runtimeDescriptorArray == VK_TRUE)
+		{
+			AddToPNext(DeviceCreateInfo, DescriptorIndexingFeatures);
+		}
+	}
+
+private:
+	VkPhysicalDeviceDescriptorIndexingFeaturesEXT DescriptorIndexingFeatures;
+};
+
+
+
+// ***** VK_KHR_fragment_shading_rate
+class FVulkanKHRFragmentShadingRateExtension : public FVulkanDeviceExtension
+{
+public:
+
+	FVulkanKHRFragmentShadingRateExtension(FVulkanDevice* InDevice)
+		: FVulkanDeviceExtension(InDevice, VK_KHR_FRAGMENT_SHADING_RATE_EXTENSION_NAME, VULKAN_SUPPORTS_FRAGMENT_SHADING_RATE)
+	{
+	}
+
+	virtual void PrePhysicalDeviceProperties(VkPhysicalDeviceProperties2KHR& PhysicalDeviceProperties2) override final
+	{
+		ZeroVulkanStruct(FragmentShadingRateProperties, VK_STRUCTURE_TYPE_PHYSICAL_DEVICE_FRAGMENT_SHADING_RATE_PROPERTIES_KHR);
+		AddToPNext(PhysicalDeviceProperties2, FragmentShadingRateProperties);
+	}
+
+	virtual void PostPhysicalDeviceFeatures(FOptionalVulkanDeviceExtensions& ExtensionFlags) override final
+	{
+		ExtensionFlags.HasKHRFragmentShadingRate = FragmentShadingRateFeatures.attachmentFragmentShadingRate;
+
+		GRHISupportsAttachmentVariableRateShading = FragmentShadingRateFeatures.attachmentFragmentShadingRate ? true : false;
+		GRHISupportsPipelineVariableRateShading = FragmentShadingRateFeatures.pipelineFragmentShadingRate ? true : false;
+
+		if (FragmentShadingRateFeatures.attachmentFragmentShadingRate == VK_TRUE)
+		{
+			GRHIVariableRateShadingImageDataType = VRSImage_Palette;
+			GRHIVariableRateShadingImageFormat = PF_R8_UINT;
+		}
+		else
+		{
+			GRHIVariableRateShadingImageDataType = VRSImage_NotSupported;
+			GRHIVariableRateShadingImageFormat = PF_Unknown;
+		}
+	}
+
+	virtual void PostPhysicalDeviceProperties() override final
+	{
+		GRHIVariableRateShadingImageTileMinWidth = FragmentShadingRateProperties.minFragmentShadingRateAttachmentTexelSize.width;
+		GRHIVariableRateShadingImageTileMinHeight = FragmentShadingRateProperties.minFragmentShadingRateAttachmentTexelSize.height;
+		GRHIVariableRateShadingImageTileMaxWidth = FragmentShadingRateProperties.maxFragmentShadingRateAttachmentTexelSize.width;
+		GRHIVariableRateShadingImageTileMaxHeight = FragmentShadingRateProperties.maxFragmentShadingRateAttachmentTexelSize.height;
+
+		// todo: We don't currently care much about the other properties here, but at some point in the future we probably will.
+
+		UE_LOG(LogVulkanRHI, Verbose, TEXT("Image-based Variable Rate Shading supported via KHRFragmentShadingRate extension. Selected VRS tile size %u by %u pixels per VRS image texel."), GRHIVariableRateShadingImageTileMinWidth, GRHIVariableRateShadingImageTileMinHeight);
+	}
+
+	virtual void PrePhysicalDeviceFeatures(VkPhysicalDeviceFeatures2KHR& PhysicalDeviceFeatures2) override final
+	{
+		ZeroVulkanStruct(FragmentShadingRateFeatures, VK_STRUCTURE_TYPE_PHYSICAL_DEVICE_FRAGMENT_SHADING_RATE_FEATURES_KHR);
+		AddToPNext(PhysicalDeviceFeatures2, FragmentShadingRateFeatures);
+	}
+
+	virtual void PreCreateDevice(VkDeviceCreateInfo& DeviceCreateInfo) override final
+	{
+		if (FragmentShadingRateFeatures.attachmentFragmentShadingRate == VK_TRUE)
+		{
+			AddToPNext(DeviceCreateInfo, FragmentShadingRateFeatures);
+		}
+	}
+
+private:
+	VkPhysicalDeviceFragmentShadingRatePropertiesKHR FragmentShadingRateProperties;
+	VkPhysicalDeviceFragmentShadingRateFeaturesKHR FragmentShadingRateFeatures;
+};
+
+
+
+// ***** VK_EXT_fragment_density_map
+class FVulkanEXTFragmentDensityMapExtension : public FVulkanDeviceExtension
+{
+public:
+
+	FVulkanEXTFragmentDensityMapExtension(FVulkanDevice* InDevice)
+		: FVulkanDeviceExtension(InDevice, VK_EXT_FRAGMENT_DENSITY_MAP_EXTENSION_NAME, VULKAN_SUPPORTS_FRAGMENT_DENSITY_MAP)
+	{
+	}
+
+	virtual void PrePhysicalDeviceProperties(VkPhysicalDeviceProperties2KHR& PhysicalDeviceProperties2) override final
+	{
+		ZeroVulkanStruct(FragmentDensityMapProperties, VK_STRUCTURE_TYPE_PHYSICAL_DEVICE_FRAGMENT_DENSITY_MAP_PROPERTIES_EXT);
+		AddToPNext(PhysicalDeviceProperties2, FragmentDensityMapProperties);
+	}
+
+	virtual void PostPhysicalDeviceFeatures(FOptionalVulkanDeviceExtensions& ExtensionFlags) override final
+	{
+		ExtensionFlags.HasEXTFragmentDensityMap = 1;
+
+		// Use the Fragment Density Map extension if and only if the Fragment Shading Rate extension is not available.
+		// NOTE: FVulkanKHRFragmentShadingRateExtension must be placed before FVulkanEXTFragmentDensityMapExtension for this to work!
+		if (!GRHISupportsAttachmentVariableRateShading && (FragmentDensityMapFeatures.fragmentDensityMap == VK_TRUE))
+		{
+			GRHISupportsAttachmentVariableRateShading = true;
+			GRHISupportsPipelineVariableRateShading = false;
+
+			// Go with the smallest tile size for now, and also force to square, since this seems to be standard.
+			// TODO: Eventually we may want to surface the range of possible tile sizes depending on end use cases, but for now this is being used for foveated rendering and smallest tile size
+			// is preferred.
+
+			GRHIVariableRateShadingImageTileMinWidth = FragmentDensityMapProperties.minFragmentDensityTexelSize.width;
+			GRHIVariableRateShadingImageTileMinHeight = FragmentDensityMapProperties.minFragmentDensityTexelSize.height;
+			GRHIVariableRateShadingImageTileMaxWidth = FragmentDensityMapProperties.maxFragmentDensityTexelSize.width;
+			GRHIVariableRateShadingImageTileMaxHeight = FragmentDensityMapProperties.maxFragmentDensityTexelSize.height;
+
+			GRHIVariableRateShadingImageDataType = VRSImage_Fractional;
+			GRHIVariableRateShadingImageFormat = PF_R8G8;
+
+			// UE_LOG(LogVulkanRHI, Display, TEXT("Image-based Variable Rate Shading supported via EXTFragmentDensityMap extension. Selected VRS tile size %u by %u pixels per VRS image texel."), GRHIVariableRateShadingImageTileMinWidth, GRHIVariableRateShadingImageTileMinHeight);
+		}
+	}
+
+	virtual void PrePhysicalDeviceFeatures(VkPhysicalDeviceFeatures2KHR& PhysicalDeviceFeatures2) override final
+	{
+		ZeroVulkanStruct(FragmentDensityMapFeatures, VK_STRUCTURE_TYPE_PHYSICAL_DEVICE_FRAGMENT_DENSITY_MAP_FEATURES_EXT);
+		AddToPNext(PhysicalDeviceFeatures2, FragmentDensityMapFeatures);
+	}
+
+	virtual void PreCreateDevice(VkDeviceCreateInfo& DeviceCreateInfo) override final
+	{
+		if (FragmentDensityMapFeatures.fragmentDensityMap == VK_TRUE)
+		{
+			AddToPNext(DeviceCreateInfo, FragmentDensityMapFeatures);
+		}
+	}
+
+private:
+	VkPhysicalDeviceFragmentDensityMapPropertiesEXT FragmentDensityMapProperties;
+	VkPhysicalDeviceFragmentDensityMapFeaturesEXT FragmentDensityMapFeatures;
+};
+
+
+
+// ***** VK_EXT_fragment_density_map2
+class FVulkanEXTFragmentDensityMap2Extension : public FVulkanDeviceExtension
+{
+public:
+
+	FVulkanEXTFragmentDensityMap2Extension(FVulkanDevice* InDevice)
+		: FVulkanDeviceExtension(InDevice, VK_EXT_FRAGMENT_DENSITY_MAP_2_EXTENSION_NAME, VULKAN_SUPPORTS_FRAGMENT_DENSITY_MAP2)
+	{
+	}
+
+	virtual void PrePhysicalDeviceFeatures(VkPhysicalDeviceFeatures2KHR& PhysicalDeviceFeatures2) override final
+	{
+		ZeroVulkanStruct(FragmentDensityMap2Features, VK_STRUCTURE_TYPE_PHYSICAL_DEVICE_FRAGMENT_DENSITY_MAP_2_FEATURES_EXT);
+		AddToPNext(PhysicalDeviceFeatures2, FragmentDensityMap2Features);
+	}
+
+	virtual void PostPhysicalDeviceFeatures(FOptionalVulkanDeviceExtensions& ExtensionFlags) override final
+	{
+		ExtensionFlags.HasEXTFragmentDensityMap2 = FragmentDensityMap2Features.fragmentDensityMapDeferred;
+
+		GRHISupportsLateVariableRateShadingUpdate = (FragmentDensityMap2Features.fragmentDensityMapDeferred == VK_TRUE);
+	}
+
+	virtual void PreCreateDevice(VkDeviceCreateInfo& DeviceCreateInfo) override final
+	{
+		if (FragmentDensityMap2Features.fragmentDensityMapDeferred == VK_TRUE)
+		{
+			AddToPNext(DeviceCreateInfo, FragmentDensityMap2Features);
+		}
+	}
+
+private:
+	VkPhysicalDeviceFragmentDensityMap2FeaturesEXT FragmentDensityMap2Features;
+};
+
+
+
+// ***** VK_KHR_get_memory_requirements2
+class FVulkanKHRGetMemoryRequirements2Extension : public FVulkanDeviceExtension
+{
+public:
+
+	FVulkanKHRGetMemoryRequirements2Extension(FVulkanDevice* InDevice)
+		: FVulkanDeviceExtension(InDevice, VK_KHR_GET_MEMORY_REQUIREMENTS_2_EXTENSION_NAME, VULKAN_EXTENSION_ENABLED, VK_API_VERSION_1_1)
+	{}
+};
+
+
+
+// ***** VK_KHR_buffer_device_address
+class FVulkanKHRBufferDeviceAddressExtension : public FVulkanDeviceExtension
+{
+public:
+
+	FVulkanKHRBufferDeviceAddressExtension(FVulkanDevice* InDevice)
+		: FVulkanDeviceExtension(InDevice, VK_KHR_BUFFER_DEVICE_ADDRESS_EXTENSION_NAME, VULKAN_RHI_RAYTRACING, VK_API_VERSION_1_2) {}
+
+	virtual void PrePhysicalDeviceFeatures(VkPhysicalDeviceFeatures2KHR& PhysicalDeviceFeatures2) override final
+	{
+		ZeroVulkanStruct(BufferDeviceAddressFeatures, VK_STRUCTURE_TYPE_PHYSICAL_DEVICE_BUFFER_DEVICE_ADDRESS_FEATURES);
+		AddToPNext(PhysicalDeviceFeatures2, BufferDeviceAddressFeatures);
+	}
+
+	virtual void PostPhysicalDeviceFeatures(FOptionalVulkanDeviceExtensions& ExtensionFlags) override final 
+	{
+		ExtensionFlags.HasBufferDeviceAddress = BufferDeviceAddressFeatures.bufferDeviceAddress; 
+	}
+
+	virtual void PreCreateDevice(VkDeviceCreateInfo& DeviceCreateInfo) override final
+	{
+		if (BufferDeviceAddressFeatures.bufferDeviceAddress == VK_TRUE)
+		{
+			AddToPNext(DeviceCreateInfo, BufferDeviceAddressFeatures);
+		}
+	}
+
+private:
+	VkPhysicalDeviceBufferDeviceAddressFeaturesKHR BufferDeviceAddressFeatures;
+};
+
+
+
+// ***** VK_KHR_acceleration_structure
+class FVulkanKHRAccelerationStructureExtension : public FVulkanDeviceExtension
+{
+public:
+
+	FVulkanKHRAccelerationStructureExtension(FVulkanDevice* InDevice)
+		: FVulkanDeviceExtension(InDevice, VK_KHR_ACCELERATION_STRUCTURE_EXTENSION_NAME, VULKAN_RHI_RAYTRACING)
+	{
+		bEnabledInCode = bEnabledInCode && GVulkanRayTracingCVar.GetValueOnAnyThread() && !FParse::Param(FCommandLine::Get(), TEXT("noraytracing"));
+	}
+
+	virtual void PrePhysicalDeviceFeatures(VkPhysicalDeviceFeatures2KHR& PhysicalDeviceFeatures2) override final
+	{
+		ZeroVulkanStruct(AccelerationStructureFeatures, VK_STRUCTURE_TYPE_PHYSICAL_DEVICE_ACCELERATION_STRUCTURE_FEATURES_KHR);;
+		AddToPNext(PhysicalDeviceFeatures2, AccelerationStructureFeatures);
+	}
+
+	virtual void PostPhysicalDeviceFeatures(FOptionalVulkanDeviceExtensions& ExtensionFlags) override final
+	{
+		ExtensionFlags.HasAccelerationStructure = 
+			AccelerationStructureFeatures.accelerationStructure & 
+			AccelerationStructureFeatures.descriptorBindingAccelerationStructureUpdateAfterBind;
+	}
+
+	virtual void PrePhysicalDeviceProperties(VkPhysicalDeviceProperties2KHR& PhysicalDeviceProperties2) override final
+	{
+#if VULKAN_RHI_RAYTRACING
+		const FRayTracingProperties& RayTracingProperties = Device->GetRayTracingProperties();
+		VkPhysicalDeviceAccelerationStructurePropertiesKHR& AccelerationStructure = const_cast<VkPhysicalDeviceAccelerationStructurePropertiesKHR&>(RayTracingProperties.AccelerationStructure);
+		ZeroVulkanStruct(AccelerationStructure, VK_STRUCTURE_TYPE_PHYSICAL_DEVICE_ACCELERATION_STRUCTURE_PROPERTIES_KHR);
+		AddToPNext(PhysicalDeviceProperties2, AccelerationStructure);
+#endif
+	}
+
+	virtual void PreCreateDevice(VkDeviceCreateInfo& DeviceCreateInfo) override final
+	{
+		if ((AccelerationStructureFeatures.accelerationStructure == VK_TRUE) &&
+			(AccelerationStructureFeatures.descriptorBindingAccelerationStructureUpdateAfterBind == VK_TRUE))
+		{
+			AddToPNext(DeviceCreateInfo, AccelerationStructureFeatures);
+		}
+	}
+
+private:
+	VkPhysicalDeviceAccelerationStructureFeaturesKHR AccelerationStructureFeatures;
+};
+
+
+
+// ***** VK_KHR_ray_tracing_pipeline
+class FVulkanKHRRayTracingPipelineExtension : public FVulkanDeviceExtension
+{
+public:
+
+	FVulkanKHRRayTracingPipelineExtension(FVulkanDevice* InDevice) 
+		: FVulkanDeviceExtension(InDevice, VK_KHR_RAY_TRACING_PIPELINE_EXTENSION_NAME, VULKAN_RHI_RAYTRACING)
+	{
+		bEnabledInCode = bEnabledInCode && GVulkanRayTracingCVar.GetValueOnAnyThread() && !FParse::Param(FCommandLine::Get(), TEXT("noraytracing"));
+	}
+
+	virtual void PrePhysicalDeviceFeatures(VkPhysicalDeviceFeatures2KHR& PhysicalDeviceFeatures2) override final
+	{
+		ZeroVulkanStruct(RayTracingPipelineFeatures, VK_STRUCTURE_TYPE_PHYSICAL_DEVICE_RAY_TRACING_PIPELINE_FEATURES_KHR);
+		AddToPNext(PhysicalDeviceFeatures2, RayTracingPipelineFeatures);
+	}
+
+	virtual void PostPhysicalDeviceFeatures(FOptionalVulkanDeviceExtensions& ExtensionFlags) override final
+	{
+		ExtensionFlags.HasRayTracingPipeline = RayTracingPipelineFeatures.rayTracingPipeline & RayTracingPipelineFeatures.rayTraversalPrimitiveCulling;
+	}
+
+	virtual void PrePhysicalDeviceProperties(VkPhysicalDeviceProperties2KHR& PhysicalDeviceProperties2) override final
+	{
+#if VULKAN_RHI_RAYTRACING
+		const FRayTracingProperties& RayTracingProperties = Device->GetRayTracingProperties();
+		VkPhysicalDeviceRayTracingPipelinePropertiesKHR& RayTracingPipeline = const_cast<VkPhysicalDeviceRayTracingPipelinePropertiesKHR&>(RayTracingProperties.RayTracingPipeline);
+		ZeroVulkanStruct(RayTracingPipeline, VK_STRUCTURE_TYPE_PHYSICAL_DEVICE_RAY_TRACING_PIPELINE_PROPERTIES_KHR);
+		AddToPNext(PhysicalDeviceProperties2, RayTracingPipeline);
+#endif
+	}
+
+	virtual void PreCreateDevice(VkDeviceCreateInfo& DeviceCreateInfo) override final
+	{
+		if ((RayTracingPipelineFeatures.rayTracingPipeline == VK_TRUE) &&
+			(RayTracingPipelineFeatures.rayTraversalPrimitiveCulling == VK_TRUE))
+		{
+			AddToPNext(DeviceCreateInfo, RayTracingPipelineFeatures);
+		}
+	}
+
+private:
+	VkPhysicalDeviceRayTracingPipelineFeaturesKHR RayTracingPipelineFeatures;
+};
+
+
+
+// ***** VK_KHR_ray_query
+class FVulkanKHRRayQueryExtension : public FVulkanDeviceExtension
+{
+public:
+
+	FVulkanKHRRayQueryExtension(FVulkanDevice* InDevice) 
+		: FVulkanDeviceExtension(InDevice, VK_KHR_RAY_QUERY_EXTENSION_NAME, VULKAN_RHI_RAYTRACING)
+	{
+		bEnabledInCode = bEnabledInCode && GVulkanRayTracingCVar.GetValueOnAnyThread() && !FParse::Param(FCommandLine::Get(), TEXT("noraytracing"));
+	}
+
+	virtual void PrePhysicalDeviceFeatures(VkPhysicalDeviceFeatures2KHR& PhysicalDeviceFeatures2) override final
+	{
+		ZeroVulkanStruct(RayQueryFeatures, VK_STRUCTURE_TYPE_PHYSICAL_DEVICE_RAY_QUERY_FEATURES_KHR);
+		AddToPNext(PhysicalDeviceFeatures2, RayQueryFeatures);
+	}
+
+	virtual void PostPhysicalDeviceFeatures(FOptionalVulkanDeviceExtensions& ExtensionFlags) override final
+	{
+		ExtensionFlags.HasRayQuery = RayQueryFeatures.rayQuery; 
+	}
+
+	virtual void PreCreateDevice(VkDeviceCreateInfo& DeviceCreateInfo) override final
+	{
+		if (RayQueryFeatures.rayQuery == VK_TRUE)
+		{
+			AddToPNext(DeviceCreateInfo, RayQueryFeatures);
+		}
+	}
+
+private:
+	VkPhysicalDeviceRayQueryFeaturesKHR RayQueryFeatures;
+};
+
+
+
+// ***** VK_EXT_debug_marker
+class FVulkanEXTDebugMarkerExtension : public FVulkanDeviceExtension
+{
+public:
+
+	FVulkanEXTDebugMarkerExtension(FVulkanDevice* InDevice)
+		: FVulkanDeviceExtension(InDevice, VK_EXT_DEBUG_MARKER_EXTENSION_NAME, (VULKAN_ENABLE_DRAW_MARKERS & VULKAN_HAS_DEBUGGING_ENABLED))
+	{
+		const int32 VulkanValidationOption = GValidationCvar.GetValueOnAnyThread();
+		bEnabledInCode = bEnabledInCode && ((GRenderDocFound || VulkanValidationOption == 0) || FVulkanPlatform::ForceEnableDebugMarkers());
+	}
+
+	virtual void PostPhysicalDeviceProperties() override final
+	{
+		Device->SetDebugMarkersFound();
+	}
+};
+
+
+
+// ***** VK_AMD_buffer_marker (vendor)
+class FVulkanAMDBufferMarkerExtension : public FVulkanDeviceExtension
+{
+public:
+
+	FVulkanAMDBufferMarkerExtension(FVulkanDevice* InDevice)
+		: FVulkanDeviceExtension(InDevice, VK_AMD_BUFFER_MARKER_EXTENSION_NAME, VULKAN_SUPPORTS_AMD_BUFFER_MARKER)
+	{
+		const bool bAllowVendorDevice = !FParse::Param(FCommandLine::Get(), TEXT("novendordevice"));
+		bEnabledInCode = GGPUCrashDebuggingEnabled && bAllowVendorDevice;
+	}
+
+	virtual void PostPhysicalDeviceFeatures(FOptionalVulkanDeviceExtensions& ExtensionFlags) override final
+	{
+		ExtensionFlags.HasAMDBufferMarker = 1; 
+	}
+};
+
+
+
+// ***** VK_NV_device_diagnostic_checkpoints (vendor)
+class FVulkanNVDeviceDiagnosticCheckpointsExtension : public FVulkanDeviceExtension
+{
+public:
+
+	FVulkanNVDeviceDiagnosticCheckpointsExtension(FVulkanDevice* InDevice)
+		: FVulkanDeviceExtension(InDevice, VK_NV_DEVICE_DIAGNOSTIC_CHECKPOINTS_EXTENSION_NAME, VULKAN_SUPPORTS_NV_DIAGNOSTICS)
+	{
+		const bool bAllowVendorDevice = !FParse::Param(FCommandLine::Get(), TEXT("novendordevice"));
+		bEnabledInCode = GGPUCrashDebuggingEnabled && bAllowVendorDevice;
+	}
+
+	virtual void PostPhysicalDeviceFeatures(FOptionalVulkanDeviceExtensions& ExtensionFlags) override final
+	{
+		ExtensionFlags.HasNVDiagnosticCheckpoints = 1; 
+	}
+};
+
+
+
+// ***** VK_NV_device_diagnostics_config (vendor)
+class FVulkanNVDeviceDiagnosticConfigExtension : public FVulkanDeviceExtension
+{
+public:
+
+	FVulkanNVDeviceDiagnosticConfigExtension(FVulkanDevice* InDevice)
+		: FVulkanDeviceExtension(InDevice, VK_NV_DEVICE_DIAGNOSTICS_CONFIG_EXTENSION_NAME, VULKAN_SUPPORTS_NV_DIAGNOSTICS)
+	{
+		const bool bAllowVendorDevice = !FParse::Param(FCommandLine::Get(), TEXT("novendordevice"));
+		bEnabledInCode = GGPUCrashDebuggingEnabled && bAllowVendorDevice;
+	}
+
+
+	virtual void PrePhysicalDeviceFeatures(VkPhysicalDeviceFeatures2KHR& PhysicalDeviceFeatures2) override final
+	{
+		ZeroVulkanStruct(DeviceDiagnosticsConfigFeaturesNV, VK_STRUCTURE_TYPE_PHYSICAL_DEVICE_DIAGNOSTICS_CONFIG_FEATURES_NV);
+		AddToPNext(PhysicalDeviceFeatures2, DeviceDiagnosticsConfigFeaturesNV);
+	}
+
+	virtual void PostPhysicalDeviceFeatures(FOptionalVulkanDeviceExtensions& ExtensionFlags) override final 
+	{
+		ExtensionFlags.HasNVDeviceDiagnosticConfig = DeviceDiagnosticsConfigFeaturesNV.diagnosticsConfig; 
+	}
+
+	virtual void PreCreateDevice(VkDeviceCreateInfo& DeviceCreateInfo) override final
+	{
+		if (DeviceDiagnosticsConfigFeaturesNV.diagnosticsConfig == VK_TRUE)
+		{
+			AddToPNext(DeviceCreateInfo, DeviceDiagnosticsConfigFeaturesNV);
+
+			ZeroVulkanStruct(DeviceDiagnosticsConfigCreateInfoNV, VK_STRUCTURE_TYPE_DEVICE_DIAGNOSTICS_CONFIG_CREATE_INFO_NV);
+			DeviceDiagnosticsConfigCreateInfoNV.flags = VK_DEVICE_DIAGNOSTICS_CONFIG_ENABLE_SHADER_DEBUG_INFO_BIT_NV | VK_DEVICE_DIAGNOSTICS_CONFIG_ENABLE_RESOURCE_TRACKING_BIT_NV | VK_DEVICE_DIAGNOSTICS_CONFIG_ENABLE_AUTOMATIC_CHECKPOINTS_BIT_NV;
+			AddToPNext(DeviceCreateInfo, DeviceDiagnosticsConfigCreateInfoNV);
+		}
+	}
+
+private:
+	VkPhysicalDeviceDiagnosticsConfigFeaturesNV DeviceDiagnosticsConfigFeaturesNV;
+	VkDeviceDiagnosticsConfigCreateInfoNV DeviceDiagnosticsConfigCreateInfoNV;
+};
+
+
+
+
+
+template <typename ExtensionType>
+static void FlagExtensionSupport(const TArray<VkExtensionProperties>& ExtensionProperties, TArray<TUniquePtr<ExtensionType>>& UEExtensions, const TCHAR* ExtensionTypeName)
+{
+	// Flag the extension support
+	UE_LOG(LogVulkanRHI, Display, TEXT("Found %d available %s extensions :"), ExtensionProperties.Num(), ExtensionTypeName);
+	for (const VkExtensionProperties& Extension : ExtensionProperties)
+	{
+		const int32 ExtensionIndex = ExtensionType::FindExtension(UEExtensions, Extension.extensionName);
+		const bool bFound = (ExtensionIndex != INDEX_NONE);
+		if (bFound)
+		{
+			UEExtensions[ExtensionIndex]->SetSupported();
+		}
+		UE_LOG(LogVulkanRHI, Display, TEXT("  %s %s"), bFound ? TEXT("+") : TEXT("-"), ANSI_TO_TCHAR(Extension.extensionName));
+	}
+}
+
+
+
+FVulkanDeviceExtensionArray FVulkanDeviceExtension::GetUESupportedDeviceExtensions(FVulkanDevice* InDevice)
+{
+	FVulkanDeviceExtensionArray OutUEDeviceExtensions;
+
+	#define ADD_SIMPLE_EXTENSION(EXTENSION_NAME, ENABLED_IN_CODE, PROMOTED_VER, FLAG_SETTER) \
+		OutUEDeviceExtensions.Add(MakeUnique<FVulkanDeviceExtension>(InDevice, EXTENSION_NAME, ENABLED_IN_CODE, PROMOTED_VER, FLAG_SETTER, FVulkanExtensionBase::AutoActivate))
+
+	#define ADD_EXTERNAL_EXTENSION(EXTENSION_NAME, ENABLED_IN_CODE, PROMOTED_VER, FLAG_SETTER) \
+		OutUEDeviceExtensions.Add(MakeUnique<FVulkanDeviceExtension>(InDevice, EXTENSION_NAME, ENABLED_IN_CODE, PROMOTED_VER, FLAG_SETTER, FVulkanExtensionBase::ManuallyActivate))
+
+	#define ADD_CUSTOM_EXTENSION(EXTENSION_CLASS) \
+		OutUEDeviceExtensions.Add(MakeUnique<EXTENSION_CLASS>(InDevice));
+
+
+	// Generic simple extensions :
+
+	ADD_SIMPLE_EXTENSION(VK_KHR_SWAPCHAIN_EXTENSION_NAME,                VULKAN_EXTENSION_ENABLED,             VULKAN_EXTENSION_NOT_PROMOTED, nullptr);
+	ADD_SIMPLE_EXTENSION(VK_KHR_MAINTENANCE1_EXTENSION_NAME,             VULKAN_SUPPORTS_MAINTENANCE_LAYER1,   VK_API_VERSION_1_1,            DEVICE_EXT_FLAG_SETTER(HasKHRMaintenance1));
+	ADD_SIMPLE_EXTENSION(VK_KHR_MAINTENANCE2_EXTENSION_NAME,             VULKAN_SUPPORTS_MAINTENANCE_LAYER2,   VK_API_VERSION_1_1,            DEVICE_EXT_FLAG_SETTER(HasKHRMaintenance2));
+	ADD_SIMPLE_EXTENSION(VK_EXT_MEMORY_BUDGET_EXTENSION_NAME,            VULKAN_SUPPORTS_MEMORY_BUDGET,        VULKAN_EXTENSION_NOT_PROMOTED, DEVICE_EXT_FLAG_SETTER(HasMemoryBudget));
+	ADD_SIMPLE_EXTENSION(VK_EXT_MEMORY_PRIORITY_EXTENSION_NAME,          VULKAN_SUPPORTS_MEMORY_PRIORITY,      VULKAN_EXTENSION_NOT_PROMOTED, DEVICE_EXT_FLAG_SETTER(HasMemoryPriority));
+	ADD_SIMPLE_EXTENSION(VK_KHR_CREATE_RENDERPASS_2_EXTENSION_NAME,      VULKAN_SUPPORTS_RENDERPASS2,          VK_API_VERSION_1_2,            DEVICE_EXT_FLAG_SETTER(HasKHRRenderPass2));
+	ADD_SIMPLE_EXTENSION(VK_KHR_DEDICATED_ALLOCATION_EXTENSION_NAME,     VULKAN_SUPPORTS_DEDICATED_ALLOCATION, VK_API_VERSION_1_1,            DEVICE_EXT_FLAG_SETTER(HasKHRDedicatedAllocation));
+	ADD_SIMPLE_EXTENSION(VK_KHR_DEFERRED_HOST_OPERATIONS_EXTENSION_NAME, VULKAN_RHI_RAYTRACING,                VULKAN_EXTENSION_NOT_PROMOTED, DEVICE_EXT_FLAG_SETTER(HasDeferredHostOperations));
+	ADD_SIMPLE_EXTENSION(VK_KHR_SPIRV_1_4_EXTENSION_NAME,                VULKAN_RHI_RAYTRACING,                VK_API_VERSION_1_2,            DEVICE_EXT_FLAG_SETTER(HasSPIRV_14));
+	ADD_SIMPLE_EXTENSION(VK_KHR_SHADER_FLOAT_CONTROLS_EXTENSION_NAME,    VULKAN_RHI_RAYTRACING,                VK_API_VERSION_1_2,            DEVICE_EXT_FLAG_SETTER(HasShaderFloatControls));
+	ADD_SIMPLE_EXTENSION(VK_KHR_IMAGE_FORMAT_LIST_EXTENSION_NAME,        VULKAN_EXTENSION_ENABLED,             VK_API_VERSION_1_2,            DEVICE_EXT_FLAG_SETTER(HasKHRImageFormatList));
+	ADD_SIMPLE_EXTENSION(VK_EXT_VALIDATION_CACHE_EXTENSION_NAME,         VULKAN_SUPPORTS_VALIDATION_CACHE,     VULKAN_EXTENSION_NOT_PROMOTED, DEVICE_EXT_FLAG_SETTER(HasEXTValidationCache));
+
+
+	// Externally activated extensions (supported by the engine, but enabled externally by plugin or other) :
+
+	ADD_EXTERNAL_EXTENSION(VK_KHR_EXTERNAL_MEMORY_EXTENSION_NAME, VULKAN_SUPPORTS_EXTERNAL_MEMORY, VK_API_VERSION_1_1, nullptr);
+
+
+	// Extensions with custom classes :
+
+	ADD_CUSTOM_EXTENSION(FVulkanKHRDriverPropertiesExtension);
+	ADD_CUSTOM_EXTENSION(FVulkanKHRMaintenance4Extension);
+	ADD_CUSTOM_EXTENSION(FVulkanShaderAtomicInt64Extension);      // must be kept BEFORE ShaderImageAtomicInt64!
+	ADD_CUSTOM_EXTENSION(FVulkanShaderImageAtomicInt64Extension); // must be kept AFTER ShaderAtomicInt64!
+	ADD_CUSTOM_EXTENSION(FVulkanEXTScalarBlockLayoutExtension);
+	ADD_CUSTOM_EXTENSION(FVulkanEXTShaderViewportIndexLayerExtension);
+	ADD_CUSTOM_EXTENSION(FVulkanKHRSeparateDepthStencilLayoutsExtension);
+	ADD_CUSTOM_EXTENSION(FVulkanKHRFragmentShadingRateExtension); // must be kept BEFORE DensityMap!
+	ADD_CUSTOM_EXTENSION(FVulkanEXTFragmentDensityMapExtension);  // must be kept AFTER ShadingRate!
+	ADD_CUSTOM_EXTENSION(FVulkanEXTFragmentDensityMap2Extension);
+	ADD_CUSTOM_EXTENSION(FVulkanKHRMultiviewExtension);
+	ADD_CUSTOM_EXTENSION(FVulkanKHRGetMemoryRequirements2Extension);
+	ADD_CUSTOM_EXTENSION(FVulkanEXTDescriptorIndexingExtension);
+
+	// Needed for Raytracing
+	ADD_CUSTOM_EXTENSION(FVulkanKHRBufferDeviceAddressExtension);
+	ADD_CUSTOM_EXTENSION(FVulkanKHRAccelerationStructureExtension);
+	ADD_CUSTOM_EXTENSION(FVulkanKHRRayTracingPipelineExtension);
+	ADD_CUSTOM_EXTENSION(FVulkanKHRRayQueryExtension);
+
+	// Vendor extensions
+	ADD_CUSTOM_EXTENSION(FVulkanAMDBufferMarkerExtension);
+	ADD_CUSTOM_EXTENSION(FVulkanNVDeviceDiagnosticCheckpointsExtension);
+	ADD_CUSTOM_EXTENSION(FVulkanNVDeviceDiagnosticConfigExtension);
+
+	// Debug
+	ADD_CUSTOM_EXTENSION(FVulkanEXTDebugMarkerExtension);
+
+	// Add in platform specific extensions
+	FVulkanPlatform::GetDeviceExtensions(InDevice, OutUEDeviceExtensions);
+
+	// Helper function to go through a list of extensions and activate them (or add them)
+	auto ActivateExternalExtensions = [&](TArray<const ANSICHAR*> Extensions, const TCHAR* Requester) {
+		for (const ANSICHAR* ExtensionName : Extensions)
+		{
+			const int32 ExtensionIndex = FindExtension(OutUEDeviceExtensions, ExtensionName);
+			if (ExtensionIndex == INDEX_NONE)
+			{
+				OutUEDeviceExtensions.Add(MakeUnique<FVulkanDeviceExtension>(InDevice, ExtensionName, VULKAN_EXTENSION_ENABLED));
+				UE_LOG(LogVulkanRHI, Warning, TEXT("%s requested device extension [%s] isn't part of the engine's core extension list. Adding it on-the-fly..."), Requester, ANSI_TO_TCHAR(ExtensionName));
+			}
+			else
+			{
+				OutUEDeviceExtensions[ExtensionIndex]->SetActivated();
+			}
+		}
+	};
+
+	// Add HMD requested extensions
+	{
+		if (IHeadMountedDisplayModule::IsAvailable())
+		{
+			FVulkanDynamicRHI::HMDVulkanExtensions = IHeadMountedDisplayModule::Get().GetVulkanExtensions();
+		}
+		if (FVulkanDynamicRHI::HMDVulkanExtensions.IsValid())
+		{
+			TArray<const ANSICHAR*> HMDExtensions;
+			FVulkanDynamicRHI::HMDVulkanExtensions->GetVulkanDeviceExtensionsRequired(InDevice->GetPhysicalHandle(), HMDExtensions);
+			ActivateExternalExtensions(HMDExtensions, TEXT("HMD"));
+		}
+	}
+
+	// Add extensions added outside the RHI (eg plugins)
+	ActivateExternalExtensions(ExternalExtensions, TEXT("Externally"));
+
+	// Now that all the extensions are listed, update their support flags
+	FlagExtensionSupport(GetDriverSupportedDeviceExtensions(InDevice->GetPhysicalHandle()), OutUEDeviceExtensions, TEXT("device"));
+
+	#undef ADD_SIMPLE_EXTENSION
+	#undef ADD_EXTERNAL_EXTENSION
+	#undef ADD_CUSTOM_EXTENSION
+
+	return OutUEDeviceExtensions;
+}
+
+
+
+
+
+
+
+
+
+
+
+// *** Vulkan Instance Extension support ***
+// Typical flow:
+// 1- The engine will query the supported extensions at RHI creation (FVulkanDeviceExtension are constructed and support is queried)
+// 2- The supported extensions are added at VkInstance creation (PreCreateInstance)
+
+
+
+// ***** VK_EXT_validation_features
+class FVulkanEXTValidationFeaturesExtension : public FVulkanInstanceExtension
+{
+public:
+
+	FVulkanEXTValidationFeaturesExtension()
+		: FVulkanInstanceExtension(VK_EXT_VALIDATION_FEATURES_EXTENSION_NAME, VULKAN_HAS_DEBUGGING_ENABLED & VULKAN_HAS_VALIDATION_FEATURES, VULKAN_EXTENSION_NOT_PROMOTED, nullptr, FVulkanExtensionBase::ManuallyActivate)
+	{}
+
+	virtual void PreCreateInstance(VkInstanceCreateInfo& InstanceCreateInfo, FOptionalVulkanInstanceExtensions& ExtensionFlags) override final 
+	{
+		check(GValidationCvar.GetValueOnAnyThread() > 0);
+
+		auto GetValidationFeaturesEnabled = []()
+		{
+			TArray<VkValidationFeatureEnableEXT> Features;
+			const int32 GPUValidationValue = GGPUValidationCvar.GetValueOnAnyThread();
+			if (GPUValidationValue > 0)
+			{
+				Features.Add(VK_VALIDATION_FEATURE_ENABLE_GPU_ASSISTED_EXT);
+				if (GPUValidationValue > 1)
+				{
+					Features.Add(VK_VALIDATION_FEATURE_ENABLE_GPU_ASSISTED_RESERVE_BINDING_SLOT_EXT);
+				}
+			}
+
+			if (FParse::Param(FCommandLine::Get(), TEXT("vulkanbestpractices")))
+			{
+				Features.Add(VK_VALIDATION_FEATURE_ENABLE_BEST_PRACTICES_EXT);
+			}
+
+			if (FParse::Param(FCommandLine::Get(), TEXT("vulkandebugsync")))
+			{
+				Features.Add(VK_VALIDATION_FEATURE_ENABLE_SYNCHRONIZATION_VALIDATION_EXT);
+			}
+
+			return Features;
+		};
+
+		ValidationFeaturesEnabled = GetValidationFeaturesEnabled();
+
+		if (ValidationFeaturesEnabled.Num())
+		{
+			ZeroVulkanStruct(ValidationFeatures, VK_STRUCTURE_TYPE_VALIDATION_FEATURES_EXT);
+			ValidationFeatures.enabledValidationFeatureCount = (uint32)ValidationFeaturesEnabled.Num();
+			ValidationFeatures.pEnabledValidationFeatures = ValidationFeaturesEnabled.GetData();
+			AddToPNext(InstanceCreateInfo, ValidationFeatures);
+		}
+	}
+
+private:
+	VkValidationFeaturesEXT ValidationFeatures;
+	TArray<VkValidationFeatureEnableEXT> ValidationFeaturesEnabled;
+};
+
+
+
+
+FVulkanInstanceExtensionArray FVulkanInstanceExtension::GetUESupportedInstanceExtensions()
+{
+	FVulkanInstanceExtensionArray OutUEInstanceExtensions;
+
+	// Generic simple extensions :
+	OutUEInstanceExtensions.Add(MakeUnique<FVulkanInstanceExtension>(VK_KHR_GET_PHYSICAL_DEVICE_PROPERTIES_2_EXTENSION_NAME, VULKAN_SUPPORTS_PHYSICAL_DEVICE_PROPERTIES2, VK_API_VERSION_1_1, INSTANCE_EXT_FLAG_SETTER(HasKHRGetPhysicalDeviceProperties2)));
+	OutUEInstanceExtensions.Add(MakeUnique<FVulkanInstanceExtension>(VK_KHR_EXTERNAL_MEMORY_CAPABILITIES_EXTENSION_NAME, VULKAN_SUPPORTS_EXTERNAL_MEMORY, VK_API_VERSION_1_1, INSTANCE_EXT_FLAG_SETTER(HasKHRExternalMemoryCapabilities)));
+	OutUEInstanceExtensions.Add(MakeUnique<FVulkanInstanceExtension>(VK_KHR_SURFACE_EXTENSION_NAME, VULKAN_EXTENSION_ENABLED, VULKAN_EXTENSION_NOT_PROMOTED));
+	OutUEInstanceExtensions.Add(MakeUnique<FVulkanInstanceExtension>(VK_EXT_SWAPCHAIN_COLOR_SPACE_EXTENSION_NAME, VULKAN_EXTENSION_ENABLED, VULKAN_EXTENSION_NOT_PROMOTED));
+
+	// Debug extensions :
+	OutUEInstanceExtensions.Add(MakeUnique<FVulkanInstanceExtension>(VK_EXT_DEBUG_UTILS_EXTENSION_NAME, VULKAN_HAS_DEBUGGING_ENABLED & VULKAN_SUPPORTS_DEBUG_UTILS, VULKAN_EXTENSION_NOT_PROMOTED, nullptr, FVulkanExtensionBase::ManuallyActivate));
+	OutUEInstanceExtensions.Add(MakeUnique<FVulkanInstanceExtension>(VK_EXT_DEBUG_REPORT_EXTENSION_NAME, VULKAN_HAS_DEBUGGING_ENABLED, VULKAN_EXTENSION_NOT_PROMOTED, nullptr, FVulkanExtensionBase::ManuallyActivate));
+
+	// Extensions with custom classes :
+	OutUEInstanceExtensions.Add(MakeUnique<FVulkanEXTValidationFeaturesExtension>());
+
+
+	// Add in platform specific extensions
+	FVulkanPlatform::GetInstanceExtensions(OutUEInstanceExtensions);
+
+	// Helper function to go through a list of extensions and activate them (or add them)
+	auto ActivateExternalExtensions = [&] (TArray<const ANSICHAR*> Extensions, const TCHAR* Requester) {
+		for (const ANSICHAR* ExtensionName : Extensions)
+		{
+			const int32 ExtensionIndex = FindExtension(OutUEInstanceExtensions, ExtensionName);
+			if (ExtensionIndex == INDEX_NONE)
+			{
+				OutUEInstanceExtensions.Add(MakeUnique<FVulkanInstanceExtension>(ExtensionName, VULKAN_EXTENSION_ENABLED));
+				UE_LOG(LogVulkanRHI, Warning, TEXT("%s requested instance extension [%s] isn't part of the engine's core extension list. Adding it on-the-fly..."), Requester, ANSI_TO_TCHAR(ExtensionName));
+			}
+			else
+			{
+				OutUEInstanceExtensions[ExtensionIndex]->SetActivated();
+			}
+		}
+	};
+
+
+	// Add HMD requested extensions
+	{
+		if (IHeadMountedDisplayModule::IsAvailable())
+		{
+			FVulkanDynamicRHI::HMDVulkanExtensions = IHeadMountedDisplayModule::Get().GetVulkanExtensions();
+		}
+		if (FVulkanDynamicRHI::HMDVulkanExtensions.IsValid())
+		{
+			TArray<const ANSICHAR*> HMDExtensions;
+			FVulkanDynamicRHI::HMDVulkanExtensions->GetVulkanInstanceExtensionsRequired(HMDExtensions);
+			ActivateExternalExtensions(HMDExtensions, TEXT("HMD"));
+		}
+	}
+
+	// Add extensions added outside the RHI (eg plugins)
+	ActivateExternalExtensions(ExternalExtensions, TEXT("Externally"));
+
+	// Now that all the extensions are listed, update their support flags
+	FlagExtensionSupport(GetDriverSupportedInstanceExtensions(), OutUEInstanceExtensions, TEXT("instance"));
+
+	return OutUEInstanceExtensions;
+}
+
+
+
+
diff --git a/Engine/Source/Runtime/VulkanRHI/Private/VulkanExtensions.h b/Engine/Source/Runtime/VulkanRHI/Private/VulkanExtensions.h
new file mode 100644
--- /dev/null
+++ b/Engine/Source/Runtime/VulkanRHI/Private/VulkanExtensions.h
@@ -0,0 +1,150 @@
+// Copyright Epic Games, Inc. All Rights Reserved.
+
+/*=============================================================================
+	VulkanExtensions.h: Vulkan extension definitions.
+=============================================================================*/
+
+#pragma once
+
+#include "CoreMinimal.h"
+#include "VulkanConfiguration.h"
+
+#define VULKAN_EXTENSION_NOT_PROMOTED 0
+#define VULKAN_EXTENSION_ENABLED      1
+#define VULKAN_EXTENSION_DISABLED     0
+
+class FVulkanDevice;
+
+class FVulkanExtensionBase
+{
+public:
+	enum EExtensionActivation : uint8
+	{
+		// If the extension is enabled in code and supported by the driver, it will automatically get used.
+		AutoActivate,
+
+		// The extension will not be used unless SetActivated() is called in the extension.  
+		// Can be used to add engine support for an extension, but not use it by default unless externally requested.
+		// Useful for extensions only activated by plugins, or only present when a certain layer is loaded.
+		ManuallyActivate
+	};
+
+	FVulkanExtensionBase(const ANSICHAR* InExtensionName, int32 InEnabledInCode, uint32 InPromotedVersion, EExtensionActivation InActivation)
+		: ExtensionName(InExtensionName)
+		, PromotedVersion(InPromotedVersion)
+		, bEnabledInCode(InEnabledInCode == VULKAN_EXTENSION_ENABLED)
+		, bSupported(false)
+		, bActivated(InActivation == EExtensionActivation::AutoActivate)
+	{}
+
+	virtual ~FVulkanExtensionBase() {}
+
+	inline void SetSupported() { bSupported = true; }
+	inline void SetActivated() { bActivated = true; }
+
+	inline const ANSICHAR* GetExtensionName() const { return ExtensionName; }
+	inline bool IsEnabled() const { return bEnabledInCode; }
+	inline bool IsSupported() const { return bSupported; }
+
+	inline bool InUse() const
+	{
+		return bEnabledInCode && bSupported && bActivated;
+	}
+
+	template <typename ExtensionType>
+	static int32 FindExtension(const TArray<TUniquePtr<ExtensionType>>& UEExtensions, const ANSICHAR* ExtensionName)
+	{
+		for (int32 Index = 0; Index < UEExtensions.Num(); ++Index)
+		{
+			if (!FCStringAnsi::Strcmp(UEExtensions[Index]->GetExtensionName(), ExtensionName))
+			{
+				return Index;
+			}
+		}
+		return INDEX_NONE;
+	}
+
+protected:
+	const ANSICHAR* ExtensionName;
+	const uint32 PromotedVersion;
+
+	// Extension is enabled in code if the VULKAN_SUPPORTS_* define is TRUE in VulkanConfiguration.h or Vulkan*Platform.h.
+	bool bEnabledInCode;
+
+	// Signals if the extension is support by the driver of the device being used (evaluated at runtime).
+	bool bSupported;
+
+	// Signals if the runtime currently wants this extension to be loaded.
+	bool bActivated;
+};
+
+
+
+class FVulkanDeviceExtension : public FVulkanExtensionBase
+{
+public:
+
+	FVulkanDeviceExtension(FVulkanDevice* InDevice, const ANSICHAR* InExtensionName, int32 InEnabledInCode, uint32 InPromotedVersion = VULKAN_EXTENSION_NOT_PROMOTED, TUniqueFunction<void(FOptionalVulkanDeviceExtensions& ExtensionFlags)>&& InFlagSetter = nullptr, EExtensionActivation InActivation = FVulkanExtensionBase::AutoActivate)
+		: FVulkanExtensionBase(InExtensionName, InEnabledInCode, InPromotedVersion, InActivation)
+		, Device(InDevice)
+		, FlagSetter(MoveTemp(InFlagSetter))
+	{}
+
+	static FVulkanDeviceExtensionArray GetUESupportedDeviceExtensions(FVulkanDevice* InDevice);
+	static TArray<VkExtensionProperties> GetDriverSupportedDeviceExtensions(VkPhysicalDevice Gpu, const ANSICHAR* LayerName = nullptr);
+
+#if VULKAN_SUPPORTS_PHYSICAL_DEVICE_PROPERTIES2
+	virtual void PrePhysicalDeviceProperties(VkPhysicalDeviceProperties2KHR& PhysicalDeviceProperties2) {}
+	virtual void PostPhysicalDeviceProperties() {}
+
+	virtual void PrePhysicalDeviceFeatures(VkPhysicalDeviceFeatures2KHR& PhysicalDeviceFeatures2) {}
+	virtual void PostPhysicalDeviceFeatures(FOptionalVulkanDeviceExtensions& ExtensionFlags)
+	{
+		if (FlagSetter)
+		{
+			FlagSetter(ExtensionFlags);
+		}
+	}
+
+	virtual void PreCreateDevice(VkDeviceCreateInfo& DeviceInfo) {}
+#endif // VULKAN_SUPPORTS_PHYSICAL_DEVICE_PROPERTIES2
+
+	// Holds extensions requested external (eg plugins)
+	static TArray<const ANSICHAR*> ExternalExtensions;
+
+protected:
+	FVulkanDevice* Device;
+	TUniqueFunction<void(FOptionalVulkanDeviceExtensions& ExtensionFlags)> FlagSetter = nullptr;
+};
+
+
+
+class FVulkanInstanceExtension : public FVulkanExtensionBase
+{
+public:
+
+	FVulkanInstanceExtension(const ANSICHAR* InExtensionName, int32 InEnabledInCode, uint32 InPromotedVersion = VULKAN_EXTENSION_NOT_PROMOTED, TUniqueFunction<void(FOptionalVulkanInstanceExtensions& ExtensionFlags)>&& InFlagSetter = nullptr, EExtensionActivation InActivation = FVulkanExtensionBase::AutoActivate)
+		: FVulkanExtensionBase(InExtensionName, InEnabledInCode, InPromotedVersion, InActivation)
+		, FlagSetter(MoveTemp(InFlagSetter))
+	{}
+
+	static FVulkanInstanceExtensionArray GetUESupportedInstanceExtensions();
+	static TArray<VkExtensionProperties> GetDriverSupportedInstanceExtensions(const ANSICHAR* LayerName = nullptr);
+
+	virtual void PreCreateInstance(VkInstanceCreateInfo& CreateInfo, FOptionalVulkanInstanceExtensions& ExtensionFlags)
+	{
+		if (FlagSetter)
+		{
+			FlagSetter(ExtensionFlags);
+		}
+	}
+
+	static TArray<const ANSICHAR*> ExternalExtensions;
+
+protected:
+	TUniqueFunction<void(FOptionalVulkanInstanceExtensions& ExtensionFlags)> FlagSetter = nullptr;
+};
+
+// Helpers to create simple extensions that set a flag
+#define DEVICE_EXT_FLAG_SETTER(FLAG_NAME) [](FOptionalVulkanDeviceExtensions& ExtensionFlags) { ExtensionFlags.FLAG_NAME = 1; }
+#define INSTANCE_EXT_FLAG_SETTER(FLAG_NAME) [](FOptionalVulkanInstanceExtensions& ExtensionFlags) { ExtensionFlags.FLAG_NAME = 1; }
\ No newline at end of file
diff --git a/Engine/Source/Runtime/VulkanRHI/Private/VulkanGenericPlatform.h b/Engine/Source/Runtime/VulkanRHI/Private/VulkanGenericPlatform.h
--- a/Engine/Source/Runtime/VulkanRHI/Private/VulkanGenericPlatform.h
+++ b/Engine/Source/Runtime/VulkanRHI/Private/VulkanGenericPlatform.h
@@ -11,6 +11,9 @@
 struct FOptionalVulkanDeviceExtensions;
 class FVulkanDevice;
 
+using FVulkanDeviceExtensionArray = TArray<TUniquePtr<class FVulkanDeviceExtension>>;
+using FVulkanInstanceExtensionArray = TArray<TUniquePtr<class FVulkanInstanceExtension>>;
+
 // the platform interface, and empty implementations for platforms that don't need em
 class FVulkanGenericPlatform
 {
@@ -25,14 +28,14 @@ public:
 	static void FreeVulkanLibrary() {}
 
 	// Called after querying all the available extensions and layers
-	static void NotifyFoundInstanceLayersAndExtensions(const TArray<FString>& Layers, const TArray<FString>& Extensions) {}
-	static void NotifyFoundDeviceLayersAndExtensions(VkPhysicalDevice PhysicalDevice, const TArray<FString>& Layers, const TArray<FString>& Extensions) {}
+	static void NotifyFoundInstanceLayersAndExtensions(const TArray<const ANSICHAR*>& Layers, const TArray<const ANSICHAR*>& Extensions) {}
+	static void NotifyFoundDeviceLayersAndExtensions(VkPhysicalDevice PhysicalDevice, const TArray<const ANSICHAR*>& Layers, const TArray<const ANSICHAR*>& Extensions) {}
 
 	// Array of required extensions for the platform (Required!)
-	static void GetInstanceExtensions(TArray<const ANSICHAR*>& OutExtensions);
+	static void GetInstanceExtensions(FVulkanInstanceExtensionArray& OutExtensions);
 	static void GetInstanceLayers(TArray<const ANSICHAR*>& OutLayers) {}
-	static void GetDeviceExtensions(EGpuVendorId VendorId, TArray<const ANSICHAR*>& OutExtensions);
-	static void GetDeviceLayers(EGpuVendorId VendorId, TArray<const ANSICHAR*>& OutLayers) {}
+	static void GetDeviceExtensions(FVulkanDevice* Device, FVulkanDeviceExtensionArray& OutExtensions);
+	static void GetDeviceLayers(TArray<const ANSICHAR*>& OutLayers) {}
 
 	// create the platform-specific surface object - required
 	static void CreateSurface(VkSurfaceKHR* OutSurface);
diff --git a/Engine/Source/Runtime/VulkanRHI/Private/VulkanLayers.cpp b/Engine/Source/Runtime/VulkanRHI/Private/VulkanLayers.cpp
--- a/Engine/Source/Runtime/VulkanRHI/Private/VulkanLayers.cpp
+++ b/Engine/Source/Runtime/VulkanRHI/Private/VulkanLayers.cpp
@@ -5,6 +5,7 @@
 =============================================================================*/
 
 #include "VulkanRHIPrivate.h"
+#include "VulkanExtensions.h"
 #include "IHeadMountedDisplayModule.h"
 #include "IHeadMountedDisplayVulkanExtensions.h"
 
@@ -56,321 +57,327 @@ TAutoConsoleVariable<int32> GGPUValidationCvar(
 
 #endif // VULKAN_HAS_DEBUGGING_ENABLED
 
-namespace VulkanExternalExtensions
-{
-	TArray<const ANSICHAR*> InstanceExtensions;
-	TArray<const ANSICHAR*> InstanceLayers;
-	TArray<const ANSICHAR*> DeviceExtensions;
-	TArray<const ANSICHAR*> DeviceLayers;
-}
 
-void IVulkanDynamicRHI::AddEnabledInstanceExtensionsAndLayers(TArrayView<const ANSICHAR* const> InInstanceExtensions, TArrayView<const ANSICHAR* const> InInstanceLayers)
-{
-	checkf(!GVulkanRHI, TEXT("AddEnabledInstanceExtensionsAndLayers should be called before the VulkanRHI has been created"));
-	VulkanExternalExtensions::InstanceExtensions.Append(InInstanceExtensions.GetData(), InInstanceExtensions.Num());
-	VulkanExternalExtensions::InstanceLayers.Append(InInstanceLayers.GetData(), InInstanceLayers.Num());
-}
 
-void IVulkanDynamicRHI::AddEnabledDeviceExtensionsAndLayers(TArrayView<const ANSICHAR* const> InDeviceExtensions, TArrayView<const ANSICHAR* const> InDeviceLayers)
-{
-	checkf(!GVulkanRHI, TEXT("AddEnabledDeviceExtensionsAndLayers should be called before the VulkanRHI has been created"));
-	VulkanExternalExtensions::DeviceExtensions.Append(InDeviceExtensions.GetData(), InDeviceExtensions.Num());
-	VulkanExternalExtensions::DeviceLayers.Append(InDeviceLayers.GetData(), InDeviceLayers.Num());
-}
+#define VERIFYVULKANRESULT_INIT(VkFunction)	{ const VkResult ScopedResult = VkFunction; \
+											if (ScopedResult == VK_ERROR_INITIALIZATION_FAILED) { \
+												UE_LOG(LogVulkanRHI, Error, \
+												TEXT("%s failed\n at %s:%u\nThis typically means Vulkan is not properly set up in your system; try running vulkaninfo from the Vulkan SDK."), \
+												ANSI_TO_TCHAR(#VkFunction), ANSI_TO_TCHAR(__FILE__), __LINE__); } \
+											else if (ScopedResult < VK_SUCCESS) { \
+												VulkanRHI::VerifyVulkanResult(ScopedResult, #VkFunction, __FILE__, __LINE__); }}
 
-// Instance Extensions to enable for all platforms
-static const ANSICHAR* GInstanceExtensions[] =
-{
-#if VULKAN_SUPPORTS_EXTERNAL_MEMORY
-	VK_KHR_EXTERNAL_MEMORY_CAPABILITIES_EXTENSION_NAME,
-#endif
-#if VULKAN_SUPPORTS_PHYSICAL_DEVICE_PROPERTIES2
-	VK_KHR_GET_PHYSICAL_DEVICE_PROPERTIES_2_EXTENSION_NAME,
-#endif
-#if VULKAN_SUPPORTS_VALIDATION_CACHE
-	VK_EXT_VALIDATION_CACHE_EXTENSION_NAME,
-#endif
-	nullptr
-};
 
-// Device Extensions to enable
-static const ANSICHAR* GDeviceExtensions[] =
+// Package a layer with its extensions
+struct FLayerWithExtensions
 {
-	VK_KHR_SWAPCHAIN_EXTENSION_NAME,
-
-	//"VK_KHX_device_group",
-#if VULKAN_SUPPORTS_MAINTENANCE_LAYER1
-	VK_KHR_MAINTENANCE1_EXTENSION_NAME,
-#endif
-
-#if VULKAN_SUPPORTS_MAINTENANCE_LAYER2
-	VK_KHR_MAINTENANCE2_EXTENSION_NAME,
-#endif
-
-#if VULKAN_SUPPORTS_VALIDATION_CACHE
-	VK_EXT_VALIDATION_CACHE_EXTENSION_NAME,
-#endif
-
-#if VULKAN_SUPPORTS_MEMORY_BUDGET
-	VK_EXT_MEMORY_BUDGET_EXTENSION_NAME,
-#endif
-
-#if VULKAN_SUPPORTS_SCALAR_BLOCK_LAYOUT
-	VK_EXT_SCALAR_BLOCK_LAYOUT_EXTENSION_NAME,
-#endif
-
-
-#if VULKAN_SUPPORTS_MEMORY_PRIORITY
-	VK_EXT_MEMORY_PRIORITY_EXTENSION_NAME,
-#endif
+	VkLayerProperties LayerProperties;
+	TArray<VkExtensionProperties> ExtensionProperties;
+};
 
-#if VULKAN_SUPPORTS_SEPARATE_DEPTH_STENCIL_LAYOUTS
-	// If we decide to support separate depth-stencil transitions, enable this.
-	//VK_KHR_CREATE_RENDERPASS_2_EXTENSION_NAME,
-	//VK_KHR_SEPARATE_DEPTH_STENCIL_LAYOUTS_EXTENSION_NAME,
-#endif
 
-	//VK_KHR_SAMPLER_MIRROR_CLAMP_TO_EDGE_EXTENSION_NAME,
-	
-#if VULKAN_SUPPORTS_BUFFER_64BIT_ATOMICS
-	VK_KHR_SHADER_ATOMIC_INT64_EXTENSION_NAME,
-#endif
 
-#if VULKAN_SUPPORTS_IMAGE_64BIT_ATOMICS
-	VK_EXT_SHADER_IMAGE_ATOMIC_INT64_EXTENSION_NAME,
-#endif
+class FVulkanIntanceSetupHelper
+{
+public:
+	TArray<FLayerWithExtensions> EnumerateLayerProperties() const
+	{
+		TArray<FLayerWithExtensions> OutLayerProperties;
+		TArray<VkLayerProperties> TempLayerProperties;
+		uint32 Count = 0;
+		VERIFYVULKANRESULT_INIT(VulkanRHI::vkEnumerateInstanceLayerProperties(&Count, nullptr));
+		if (Count > 0)
+		{
+			TempLayerProperties.AddZeroed(Count);
+			VERIFYVULKANRESULT_INIT(VulkanRHI::vkEnumerateInstanceLayerProperties(&Count, TempLayerProperties.GetData()));
+			OutLayerProperties.SetNum(Count);
+			for (uint32 i=0; i < Count; ++i)
+			{
+				OutLayerProperties[i].LayerProperties = TempLayerProperties[i];
+				OutLayerProperties[i].ExtensionProperties = FVulkanInstanceExtension::GetDriverSupportedInstanceExtensions(TempLayerProperties[i].layerName);
+			}
+		}
+		return OutLayerProperties;
+	}
 
-	VK_KHR_IMAGE_FORMAT_LIST_EXTENSION_NAME,
+	TArray<const ANSICHAR*> GetPlatformLayers()
+	{
+		TArray<const ANSICHAR*> OutPlatformLayers;
+		FVulkanPlatform::GetInstanceLayers(OutPlatformLayers);
+		return OutPlatformLayers;
+	}
 
-#if VULKAN_SUPPORTS_SHADER_VIEWPORT_INDEX_LAYER
-	VK_EXT_SHADER_VIEWPORT_INDEX_LAYER_EXTENSION_NAME,
-#endif 
+	void AddDebugLayers(const TArray<FLayerWithExtensions>& LayerProperties, FVulkanInstanceExtensionArray& UEExtensions, TArray<const ANSICHAR*>& OutLayers);
 
-	nullptr
+	const TCHAR* HelperTypeName = TEXT("instance");
+	static TArray<const ANSICHAR*> ExternalLayers;
+	FVulkanDynamicRHI::EActiveDebugLayerExtension ActiveDebugLayerExtension = FVulkanDynamicRHI::EActiveDebugLayerExtension::None;
 };
 
-TSharedPtr< IHeadMountedDisplayVulkanExtensions, ESPMode::ThreadSafe > FVulkanDynamicRHI::HMDVulkanExtensions;
+TArray<const ANSICHAR*> FVulkanIntanceSetupHelper::ExternalLayers;
 
-struct FLayerExtension
-{
-	FLayerExtension()
-	{
-		FMemory::Memzero(LayerProps);
-	}
 
-	void AddUniqueExtensionNames(TArray<FString>& Out)
+class FVulkanDeviceSetupHelper
+{
+public:
+	TArray<FLayerWithExtensions> EnumerateLayerProperties() const
 	{
-		for (int32 ExtIndex = 0; ExtIndex < ExtensionProps.Num(); ++ExtIndex)
+		TArray<FLayerWithExtensions> OutLayerProperties;
+		TArray<VkLayerProperties> TempLayerProperties;
+		uint32 Count = 0;
+		VERIFYVULKANRESULT_INIT(VulkanRHI::vkEnumerateDeviceLayerProperties(Gpu, &Count, nullptr));
+		if (Count > 0)
 		{
-			Out.AddUnique(ANSI_TO_TCHAR(ExtensionProps[ExtIndex].extensionName));
+			TempLayerProperties.AddZeroed(Count);
+			VERIFYVULKANRESULT_INIT(VulkanRHI::vkEnumerateDeviceLayerProperties(Gpu, &Count, TempLayerProperties.GetData()));
+			OutLayerProperties.SetNum(Count);
+			for (uint32 i = 0; i < Count; ++i)
+			{
+				OutLayerProperties[i].LayerProperties = TempLayerProperties[i];
+				OutLayerProperties[i].ExtensionProperties = FVulkanDeviceExtension::GetDriverSupportedDeviceExtensions(Gpu, TempLayerProperties[i].layerName);
+			}
 		}
+		return OutLayerProperties;
 	}
 
-	void AddAnsiExtensionNames(TArray<const char*>& Out)
+	TArray<const ANSICHAR*> GetPlatformLayers()
 	{
-		for (int32 ExtIndex = 0; ExtIndex < ExtensionProps.Num(); ++ExtIndex)
-		{
-			Out.AddUnique(ExtensionProps[ExtIndex].extensionName);
-		}
+		TArray<const ANSICHAR*> OutPlatformLayers;
+		FVulkanPlatform::GetDeviceLayers(OutPlatformLayers);
+		return OutPlatformLayers;
 	}
 
-	VkLayerProperties LayerProps;
-	TArray<VkExtensionProperties> ExtensionProps;
+	void AddDebugLayers(const TArray<FLayerWithExtensions>& LayerProperties, FVulkanDeviceExtensionArray& UEExtensions, TArray<const ANSICHAR*>& OutLayers);
+
+	const TCHAR* HelperTypeName = TEXT("device");
+	static TArray<const ANSICHAR*> ExternalLayers;
+	VkPhysicalDevice Gpu = VK_NULL_HANDLE;
 };
 
+TArray<const ANSICHAR*> FVulkanDeviceSetupHelper::ExternalLayers;
+
+#undef VERIFYVULKANRESULT_INIT
+
 
-static void ErrorPotentialBadInstallation(const ANSICHAR* VkFunction, const ANSICHAR* Filename, uint32 Line)
-{
-	UE_LOG(LogVulkanRHI, Error, TEXT("%s failed\n at %s:%u\nThis typically means Vulkan is not properly set up in your system; try running vulkaninfo from the Vulkan SDK."), ANSI_TO_TCHAR(VkFunction), ANSI_TO_TCHAR(Filename), Line);
-}
 
-#define VERIFYVULKANRESULT_INIT(VkFunction)		{ const VkResult ScopedResult = VkFunction; \
-													if (ScopedResult == VK_ERROR_INITIALIZATION_FAILED) { ErrorPotentialBadInstallation(#VkFunction, __FILE__, __LINE__); } \
-													else if (ScopedResult < VK_SUCCESS) { VulkanRHI::VerifyVulkanResult(ScopedResult, #VkFunction, __FILE__, __LINE__); }}
 
-static inline void EnumerateInstanceExtensionProperties(const ANSICHAR* LayerName, FLayerExtension& OutLayer)
+// Helper function to find a layer name in a list of LayerProperties
+static inline int32 FindLayerIndexInList(const char* LayerName, const TArray<FLayerWithExtensions>& LayerProperties)
 {
-	uint32 Count = 0;
-	VERIFYVULKANRESULT_INIT(VulkanRHI::vkEnumerateInstanceExtensionProperties(LayerName, &Count, nullptr));
-	if (Count > 0)
+	for (int32 Index = 0; Index < LayerProperties.Num(); ++Index)
 	{
-		OutLayer.ExtensionProps.Empty(Count);
-		OutLayer.ExtensionProps.AddUninitialized(Count);
-		VERIFYVULKANRESULT_INIT(VulkanRHI::vkEnumerateInstanceExtensionProperties(LayerName, &Count, OutLayer.ExtensionProps.GetData()));
+		if (!FCStringAnsi::Strcmp(LayerProperties[Index].LayerProperties.layerName, LayerName))
+		{
+			return Index;
+		}
 	}
+
+	return INDEX_NONE;
 }
 
-static inline void EnumerateDeviceExtensionProperties(VkPhysicalDevice Device, const ANSICHAR* LayerName, FLayerExtension& OutLayer)
+
+
+// Helper function to add a layer to a list and flag its extensions if it's found
+template <typename ExtensionType>
+static inline bool AddRequestedLayer(const ANSICHAR* LayerName, const TArray<FLayerWithExtensions>& LayerProperties, TArray<TUniquePtr<ExtensionType>>& UEExtensions, TArray<const ANSICHAR*>& OutLayers)
 {
-	uint32 Count = 0;
-	VERIFYVULKANRESULT_INIT(VulkanRHI::vkEnumerateDeviceExtensionProperties(Device, LayerName, &Count, nullptr));
-	if (Count > 0)
+	// Find the layer in the list
+	int32 LayerIndex = FindLayerIndexInList(LayerName, LayerProperties);
+	if (LayerIndex == INDEX_NONE)
 	{
-		OutLayer.ExtensionProps.Empty(Count);
-		OutLayer.ExtensionProps.AddUninitialized(Count);
-		VERIFYVULKANRESULT_INIT(VulkanRHI::vkEnumerateDeviceExtensionProperties(Device, LayerName, &Count, OutLayer.ExtensionProps.GetData()));
+		return false;
 	}
-}
 
+	// Add it to the list of used layers
+	OutLayers.Add(LayerName);
 
-static inline void TrimDuplicates(TArray<const ANSICHAR*>& Array)
-{
-	for (int32 OuterIndex = Array.Num() - 1; OuterIndex >= 0; --OuterIndex)
+
+	// Helper function to flag an extension as supported by the driver
+	auto FlagExtensionSupport = [](TArray<TUniquePtr<ExtensionType>>& UEExtensions, const ANSICHAR* ExtensionName)
 	{
-		bool bFound = false;
-		for (int32 InnerIndex = OuterIndex - 1; InnerIndex >= 0; --InnerIndex)
+		for (TUniquePtr<ExtensionType>& Extension : UEExtensions)
 		{
-			if (!FCStringAnsi::Strcmp(Array[OuterIndex], Array[InnerIndex]))
+			if (!FCStringAnsi::Strcmp(Extension->GetExtensionName(), ExtensionName))
 			{
-				bFound = true;
-				break;
+				Extension->SetSupported();
+				return true;
 			}
 		}
+		return false;
+	};
 
-		if (bFound)
-		{
-			Array.RemoveAtSwap(OuterIndex, 1, false);
-		}
+	// Flag its extensions as usable
+	for (const VkExtensionProperties& ExtensionProperties : LayerProperties[LayerIndex].ExtensionProperties)
+	{
+		FlagExtensionSupport(UEExtensions, ExtensionProperties.extensionName);
 	}
+
+	return true;
 }
 
-static inline int32 FindLayerIndexInList(const TArray<FLayerExtension>& List, const char* LayerName)
+
+
+template <typename SetupHelperType, typename ExtensionType>
+static TArray<const ANSICHAR*> SetupLayers(SetupHelperType& VulkanSetupHelper, TArray<TUniquePtr<ExtensionType>>& UEExtensions)
 {
-	// 0 is reserved for NULL/instance
-	for (int32 Index = 1; Index < List.Num(); ++Index)
+	TArray<const ANSICHAR*> OutLayers;
+
+	// Fetch the list of supported layers
+	TArray<FLayerWithExtensions> LayerProperties = VulkanSetupHelper.EnumerateLayerProperties();
+	LayerProperties.Sort([](const FLayerWithExtensions& A, const FLayerWithExtensions& B) { return FCStringAnsi::Strcmp(A.LayerProperties.layerName, B.LayerProperties.layerName) < 0; });
+
+	UE_LOG(LogVulkanRHI, Display, TEXT("Found %d available %s layers %s"), LayerProperties.Num(), VulkanSetupHelper.HelperTypeName, LayerProperties.Num() ? TEXT(":") : TEXT("!"));
+	for (const FLayerWithExtensions& Prop : LayerProperties)
 	{
-		if (!FCStringAnsi::Strcmp(List[Index].LayerProps.layerName, LayerName))
+		UE_LOG(LogVulkanRHI, Display, TEXT("  * %s"), ANSI_TO_TCHAR(Prop.LayerProperties.layerName));
+	}
+
+	// Check for layers added outside the RHI (eg plugins)
+	for (const ANSICHAR* VulkanBridgeLayer : VulkanSetupHelper.ExternalLayers)
+	{
+		if (!AddRequestedLayer(VulkanBridgeLayer, LayerProperties, UEExtensions, OutLayers))
 		{
-			return Index;
+			UE_LOG(LogVulkanRHI, Warning, TEXT("Unable to find VulkanExternalExtensions %s layer '%s'"), VulkanSetupHelper.HelperTypeName, ANSI_TO_TCHAR(VulkanBridgeLayer));
 		}
 	}
 
-	return INDEX_NONE;
-}
+	// Check for platform specific layers
+	TArray<const ANSICHAR*> PlatformLayers = VulkanSetupHelper.GetPlatformLayers();
+	for (const ANSICHAR* PlatformLayer : PlatformLayers)
+	{
+		if (!AddRequestedLayer(PlatformLayer, LayerProperties, UEExtensions, OutLayers))
+		{
+			UE_LOG(LogVulkanRHI, Warning, TEXT("Unable to find platform %s layer '%s'"), VulkanSetupHelper.HelperTypeName, ANSI_TO_TCHAR(PlatformLayer));
+		}
+	}
 
-static inline bool FindLayerInList(const TArray<FLayerExtension>& List, const char* LayerName)
-{
-	return FindLayerIndexInList(List, LayerName) != INDEX_NONE;
-}
+	// Check for any requested debug layers
+	VulkanSetupHelper.AddDebugLayers(LayerProperties, UEExtensions, OutLayers);
 
-static inline bool FindLayerExtensionInList(const TArray<FLayerExtension>& List, const char* ExtensionName, const char*& FoundLayer)
-{
-	for (int32 Index = 0; Index < List.Num(); ++Index)
+	// Clean up the resulting array
+	if (OutLayers.Num() > 0)
 	{
-		for (int32 ExtIndex = 0; ExtIndex < List[Index].ExtensionProps.Num(); ++ExtIndex)
+		auto TrimDuplicates = [](TArray<const ANSICHAR*>& Array)
 		{
-			if (!FCStringAnsi::Strcmp(List[Index].ExtensionProps[ExtIndex].extensionName, ExtensionName))
+			for (int32 OuterIndex = Array.Num() - 1; OuterIndex >= 0; --OuterIndex)
 			{
-				FoundLayer = List[Index].LayerProps.layerName;
-				return true;
+				bool bFound = false;
+				for (int32 InnerIndex = OuterIndex - 1; InnerIndex >= 0; --InnerIndex)
+				{
+					if (!FCStringAnsi::Strcmp(Array[OuterIndex], Array[InnerIndex]))
+					{
+						bFound = true;
+						break;
+					}
+				}
+
+				if (bFound)
+				{
+					Array.RemoveAtSwap(OuterIndex, 1, false);
+				}
 			}
-		}
+		};
+
+		TrimDuplicates(OutLayers);
 	}
+	OutLayers.Sort();
 
-	return false;
+	return OutLayers;
 }
 
-static inline bool FindLayerExtensionInList(const TArray<FLayerExtension>& List, const char* ExtensionName)
+
+void IVulkanDynamicRHI::AddEnabledInstanceExtensionsAndLayers(TArrayView<const ANSICHAR* const> InInstanceExtensions, TArrayView<const ANSICHAR* const> InInstanceLayers)
 {
-	const char* Dummy = nullptr;
-	return FindLayerExtensionInList(List, ExtensionName, Dummy);
+	checkf(!GVulkanRHI, TEXT("AddEnabledInstanceExtensionsAndLayers should be called before the VulkanRHI has been created"));
+	FVulkanInstanceExtension::ExternalExtensions.Append(InInstanceExtensions.GetData(), InInstanceExtensions.Num());
+	FVulkanIntanceSetupHelper::ExternalLayers.Append(InInstanceLayers.GetData(), InInstanceLayers.Num());
 }
 
-
-void FVulkanDynamicRHI::GetInstanceLayersAndExtensions(TArray<const ANSICHAR*>& OutInstanceExtensions, TArray<const ANSICHAR*>& OutInstanceLayers, bool& bOutDebugUtils)
+void IVulkanDynamicRHI::AddEnabledDeviceExtensionsAndLayers(TArrayView<const ANSICHAR* const> InDeviceExtensions, TArrayView<const ANSICHAR* const> InDeviceLayers)
 {
-	bOutDebugUtils = false;
+	checkf(!GVulkanRHI, TEXT("AddEnabledDeviceExtensionsAndLayers should be called before the VulkanRHI has been created"));
+	FVulkanDeviceExtension::ExternalExtensions.Append(InDeviceExtensions.GetData(), InDeviceExtensions.Num());
+	FVulkanDeviceSetupHelper::ExternalLayers.Append(InDeviceLayers.GetData(), InDeviceLayers.Num());
+}
 
-	TArray<FLayerExtension> GlobalLayerExtensions;
-	// 0 is reserved for NULL/instance
-	GlobalLayerExtensions.AddDefaulted();
+TArray<const ANSICHAR*> FVulkanDynamicRHI::SetupInstanceLayers(FVulkanInstanceExtensionArray& UEExtensions)
+{
+	FVulkanIntanceSetupHelper InstanceHelper;
+	TArray<const ANSICHAR*> OutInstanceLayers = SetupLayers(InstanceHelper, UEExtensions);
+	GVulkanRHI->ActiveDebugLayerExtension = InstanceHelper.ActiveDebugLayerExtension;
+	return OutInstanceLayers;
+}
 
-	// Global extensions
-	EnumerateInstanceExtensionProperties(nullptr, GlobalLayerExtensions[0]);
+TArray<const ANSICHAR*> FVulkanDevice::SetupDeviceLayers(VkPhysicalDevice Gpu, FVulkanDeviceExtensionArray& UEExtensions)
+{
+	FVulkanDeviceSetupHelper DeviceHelper;
+	DeviceHelper.Gpu = Gpu;
+	return SetupLayers(DeviceHelper, UEExtensions);
+}
 
-	TArray<FString> FoundUniqueExtensions;
-	TArray<FString> FoundUniqueLayers;
-	for (int32 Index = 0; Index < GlobalLayerExtensions[0].ExtensionProps.Num(); ++Index)
-	{
-		FoundUniqueExtensions.AddUnique(ANSI_TO_TCHAR(GlobalLayerExtensions[0].ExtensionProps[Index].extensionName));
-	}
+void FVulkanDynamicRHI::SetupValidationRequests()
+{
+#if VULKAN_HAS_DEBUGGING_ENABLED
+	int32 VulkanValidationOption = GValidationCvar.GetValueOnAnyThread();
 
+	// Command line overrides Cvar
+	if (FParse::Param(FCommandLine::Get(), TEXT("vulkandebug")))
 	{
-		TArray<VkLayerProperties> GlobalLayerProperties;
-		uint32 InstanceLayerCount = 0;
-		VERIFYVULKANRESULT_INIT(VulkanRHI::vkEnumerateInstanceLayerProperties(&InstanceLayerCount, nullptr));
-		if (InstanceLayerCount > 0)
-		{
-			GlobalLayerProperties.AddZeroed(InstanceLayerCount);
-			VERIFYVULKANRESULT_INIT(VulkanRHI::vkEnumerateInstanceLayerProperties(&InstanceLayerCount, &GlobalLayerProperties[GlobalLayerProperties.Num() - InstanceLayerCount]));
-		}
-
-		for (int32 Index = 0; Index < GlobalLayerProperties.Num(); ++Index)
-		{
-			FLayerExtension* Layer = new(GlobalLayerExtensions) FLayerExtension;
-			Layer->LayerProps = GlobalLayerProperties[Index];
-			EnumerateInstanceExtensionProperties(GlobalLayerProperties[Index].layerName, *Layer);
-			Layer->AddUniqueExtensionNames(FoundUniqueExtensions);
-			FoundUniqueLayers.AddUnique(ANSI_TO_TCHAR(GlobalLayerProperties[Index].layerName));
-		}
+		// Match D3D and GL
+		GValidationCvar->Set(2, ECVF_SetByCommandline);
 	}
-
-	UE_LOG(LogVulkanRHI, Display, TEXT("- Found %d instance layers"), FoundUniqueLayers.Num());
-	if (FoundUniqueLayers.Num() > 0)
+	else if (FParse::Value(FCommandLine::Get(), TEXT("vulkanvalidation="), VulkanValidationOption))
 	{
-		FoundUniqueLayers.Sort();
-		for (const FString& Name : FoundUniqueLayers)
-		{
-			UE_LOG(LogVulkanRHI, Display, TEXT("* %s"), *Name);
-		}
+		GValidationCvar->Set(VulkanValidationOption, ECVF_SetByCommandline);
 	}
 
-	UE_LOG(LogVulkanRHI, Display, TEXT("- Found %d instance extensions"), FoundUniqueExtensions.Num());
-	if (FoundUniqueExtensions.Num() > 0)
+	if (FParse::Param(FCommandLine::Get(), TEXT("gpuvalidation")))
 	{
-		FoundUniqueExtensions.Sort();
-		for (const FString& Name : FoundUniqueExtensions)
+		if (GValidationCvar->GetInt() < 2)
 		{
-			UE_LOG(LogVulkanRHI, Display, TEXT("* %s"), *Name);
+			GValidationCvar->Set(2, ECVF_SetByCommandline);
 		}
+		GGPUValidationCvar->Set(2, ECVF_SetByCommandline);
 	}
+#endif
+}
 
-	FVulkanPlatform::NotifyFoundInstanceLayersAndExtensions(FoundUniqueLayers, FoundUniqueExtensions);
 
-	bool bGfxReconstructOrVkTrace = false;
+// Get a list of debugging layers to activate
+void FVulkanIntanceSetupHelper::AddDebugLayers(const TArray<FLayerWithExtensions>& LayerProperties, FVulkanInstanceExtensionArray& UEExtensions, TArray<const ANSICHAR*>& OutLayers)
+{
 	if (FParse::Param(FCommandLine::Get(), TEXT("vktrace")))
 	{
 		const char* GfxReconstructName = "VK_LAYER_LUNARG_gfxreconstruct";
-		if (FindLayerInList(GlobalLayerExtensions, GfxReconstructName))
+		if (AddRequestedLayer(GfxReconstructName, LayerProperties, UEExtensions, OutLayers))
 		{
-			OutInstanceLayers.Add(GfxReconstructName);
-			bGfxReconstructOrVkTrace = true;
+			ActiveDebugLayerExtension = FVulkanDynamicRHI::EActiveDebugLayerExtension::GfxReconstructLayer;
 		}
 		else
 		{
 			const char* VkTraceName = "VK_LAYER_LUNARG_vktrace";
-			if (FindLayerInList(GlobalLayerExtensions, VkTraceName))
+			if (AddRequestedLayer(VkTraceName, LayerProperties, UEExtensions, OutLayers))
 			{
-				OutInstanceLayers.Add(VkTraceName);
-				bGfxReconstructOrVkTrace = true;
+				ActiveDebugLayerExtension = FVulkanDynamicRHI::EActiveDebugLayerExtension::VkTraceLayer;
 			}
 		}
 	}
 
+	const bool bGfxReconstructOrVkTrace = (ActiveDebugLayerExtension != FVulkanDynamicRHI::EActiveDebugLayerExtension::None);
+
 #if VULKAN_HAS_DEBUGGING_ENABLED
 	if (FParse::Param(FCommandLine::Get(), TEXT("vulkanapidump")))
 	{
 		if (bGfxReconstructOrVkTrace)
 		{
-			UE_LOG(LogVulkanRHI, Warning, TEXT("Can't enable api_dump when GfxReconstruct/VkTrace is enabled"));
+			UE_LOG(LogVulkanRHI, Warning, TEXT("Can't enable api_dump when GfxReconstruct/VkTrace is enabled."));
 		}
 		else
 		{
 			const char* VkApiDumpName = "VK_LAYER_LUNARG_api_dump";
-			bool bApiDumpFound = FindLayerInList(GlobalLayerExtensions, VkApiDumpName);
+			const bool bApiDumpFound = AddRequestedLayer(VkApiDumpName, LayerProperties, UEExtensions, OutLayers);
 			if (bApiDumpFound)
 			{
-				OutInstanceLayers.Add(VkApiDumpName);
 				FPlatformMisc::SetEnvironmentVar(TEXT("VK_APIDUMP_LOG_FILENAME"), TEXT("vk_apidump.txt"));
 				FPlatformMisc::SetEnvironmentVar(TEXT("VK_APIDUMP_DETAILED"), TEXT("true"));
 				FPlatformMisc::SetEnvironmentVar(TEXT("VK_APIDUMP_FLUSH"), TEXT("true"));
@@ -384,27 +391,18 @@ void FVulkanDynamicRHI::GetInstanceLayersAndExtensions(TArray<const ANSICHAR*>&
 	}
 
 	// At this point the CVar holds the final value
-#if VULKAN_HAS_DEBUGGING_ENABLED
-	const int32 VulkanValidationOption = GValidationCvar.GetValueOnAnyThread();
-	if (!bGfxReconstructOrVkTrace && VulkanValidationOption > 0)
+	const bool bUseVulkanValidation = (GValidationCvar.GetValueOnAnyThread() > 0);
+	if (!bGfxReconstructOrVkTrace && bUseVulkanValidation)
 	{
-		bool bSkipStandard = false;
-		bool bStandardAvailable = false;
 		if (GStandardValidationCvar.GetValueOnAnyThread() != 0)
 		{
 			if (GStandardValidationCvar.GetValueOnAnyThread() == 2)
 			{
-				bStandardAvailable = FindLayerInList(GlobalLayerExtensions, KHRONOS_STANDARD_VALIDATION_LAYER_NAME);
-				if (bStandardAvailable)
-				{
-					OutInstanceLayers.Add(KHRONOS_STANDARD_VALIDATION_LAYER_NAME);
-				}
-				else
+				if (!AddRequestedLayer(KHRONOS_STANDARD_VALIDATION_LAYER_NAME, LayerProperties, UEExtensions, OutLayers))
 				{
 #if PLATFORM_WINDOWS || PLATFORM_LINUX
 					//#todo-rco: We don't package DLLs so if this fails it means no DLL was found anywhere, so don't try to load standard validation layers
 					UE_LOG(LogVulkanRHI, Warning, TEXT("Unable to find Vulkan instance validation layer %s;  Do you have the Vulkan SDK Installed?"), TEXT(STANDARD_VALIDATION_LAYER_NAME));
-					bSkipStandard = true;
 #else
 					UE_LOG(LogVulkanRHI, Warning, TEXT("Unable to find Vulkan instance validation layer %s"), TEXT(STANDARD_VALIDATION_LAYER_NAME));
 #endif
@@ -412,545 +410,101 @@ void FVulkanDynamicRHI::GetInstanceLayersAndExtensions(TArray<const ANSICHAR*>&
 			}
 			else
 			{
-				bStandardAvailable = FindLayerInList(GlobalLayerExtensions, STANDARD_VALIDATION_LAYER_NAME);
-				if (bStandardAvailable)
-				{
-					OutInstanceLayers.Add(STANDARD_VALIDATION_LAYER_NAME);
-				}
-				else
+				if (!AddRequestedLayer(STANDARD_VALIDATION_LAYER_NAME, LayerProperties, UEExtensions, OutLayers))
 				{
 					UE_LOG(LogVulkanRHI, Warning, TEXT("Unable to find Vulkan instance validation layer %s"), TEXT(STANDARD_VALIDATION_LAYER_NAME));
 				}
 			}
 		}
 	}
-#endif
 
-#if VULKAN_SUPPORTS_DEBUG_UTILS
-	if (!bGfxReconstructOrVkTrace && VulkanValidationOption > 0)
+	if (bUseVulkanValidation && (ActiveDebugLayerExtension == FVulkanDynamicRHI::EActiveDebugLayerExtension::None))
 	{
-		const char* FoundDebugUtilsLayer = nullptr;
-		bOutDebugUtils = FindLayerExtensionInList(GlobalLayerExtensions, VK_EXT_DEBUG_UTILS_EXTENSION_NAME, FoundDebugUtilsLayer);
-		if (bOutDebugUtils && *FoundDebugUtilsLayer)
+		auto FindLayerContainingExtension = [](const ANSICHAR* ExtensionName, const TArray<FLayerWithExtensions>& LayerProperties)
 		{
-			OutInstanceLayers.Add(FoundDebugUtilsLayer);
-		}
-	}
-#endif
-#endif	// VULKAN_HAS_DEBUGGING_ENABLED
-
-	// Check to see if the HMD requires any specific Vulkan extensions to operate
-	if (IHeadMountedDisplayModule::IsAvailable())
-	{
-		HMDVulkanExtensions = IHeadMountedDisplayModule::Get().GetVulkanExtensions();
-	
-		if (HMDVulkanExtensions.IsValid())
-		{
-			if (!HMDVulkanExtensions->GetVulkanInstanceExtensionsRequired(OutInstanceExtensions))
+			for (int32 LayerIndex = 0; LayerIndex < LayerProperties.Num(); ++LayerIndex)
 			{
-				UE_LOG(LogVulkanRHI, Warning, TEXT("Trying to use Vulkan with an HMD, but required extensions aren't supported!"));
-			}
-		}
-	}
-
-	// Check for layers added outside the RHI (eg plugins)
-	for (const ANSICHAR* VulkanBridgeLayer : VulkanExternalExtensions::InstanceLayers)
-	{
-		if (FindLayerInList(GlobalLayerExtensions, VulkanBridgeLayer))
-		{
-			OutInstanceLayers.Add(VulkanBridgeLayer);
-		}
-		else
-		{
-			UE_LOG(LogVulkanRHI, Warning, TEXT("Unable to find VulkanExternalExtensions instance layer '%s'"), ANSI_TO_TCHAR(VulkanBridgeLayer));
-		}
-	}
-
-	TArray<const ANSICHAR*> PlatformExtensions;
-	FVulkanPlatform::GetInstanceExtensions(PlatformExtensions);
-
-	for (const ANSICHAR* PlatformExtension : PlatformExtensions)
-	{
-		if (FindLayerExtensionInList(GlobalLayerExtensions, PlatformExtension))
-		{
-			OutInstanceExtensions.Add(PlatformExtension);
-		}
-	}
-
-	TArray<const ANSICHAR*> PlatformLayers;
-	FVulkanPlatform::GetInstanceLayers(PlatformLayers);
-	for (const ANSICHAR* PlatformLayer : PlatformLayers)
-	{
-		OutInstanceLayers.Add(PlatformLayer);
-	}
-
-	for (int32 j = 0; GInstanceExtensions[j] != nullptr; j++)
-	{
-		if (FindLayerExtensionInList(GlobalLayerExtensions, GInstanceExtensions[j]))
-		{
-			OutInstanceExtensions.Add(GInstanceExtensions[j]);
-		}
-	}
-
-	// Check for extensions added outside the RHI (eg plugins)
-	for (const ANSICHAR* VulkanBridgeExtension : VulkanExternalExtensions::InstanceExtensions)
-	{
-		if (FindLayerExtensionInList(GlobalLayerExtensions, VulkanBridgeExtension))
-		{
-			OutInstanceExtensions.Add(VulkanBridgeExtension);
-		}
-		else
-		{
-			UE_LOG(LogVulkanRHI, Warning, TEXT("Unable to find VulkanExternalExtensions instance extension '%s'"), ANSI_TO_TCHAR(VulkanBridgeExtension));
-		}
-	}
-
-#if VULKAN_SUPPORTS_DEBUG_UTILS
-	if (!bGfxReconstructOrVkTrace && bOutDebugUtils && FindLayerExtensionInList(GlobalLayerExtensions, VK_EXT_DEBUG_UTILS_EXTENSION_NAME))
-	{
-		OutInstanceExtensions.Add(VK_EXT_DEBUG_UTILS_EXTENSION_NAME);
-	}
-#endif
-#if VULKAN_HAS_DEBUGGING_ENABLED
-	if (!bGfxReconstructOrVkTrace && !bOutDebugUtils && VulkanValidationOption > 0)
-	{
-		if (FindLayerExtensionInList(GlobalLayerExtensions, VK_EXT_DEBUG_REPORT_EXTENSION_NAME))
-		{
-			OutInstanceExtensions.Add(VK_EXT_DEBUG_REPORT_EXTENSION_NAME);
-		}
-	}
-
-	if (VulkanValidationOption > 0 && !bGfxReconstructOrVkTrace)
-	{
-#if VULKAN_HAS_VALIDATION_FEATURES
-		if (FindLayerExtensionInList(GlobalLayerExtensions, VK_EXT_VALIDATION_FEATURES_EXTENSION_NAME) && GGPUValidationCvar.GetValueOnAnyThread() != 0)
-		{
-			OutInstanceExtensions.Add(VK_EXT_VALIDATION_FEATURES_EXTENSION_NAME);
-		}
-#endif
-	}
-#endif
-
-	if (OutInstanceLayers.Num() > 0)
-	{
-		TrimDuplicates(OutInstanceLayers);
-		UE_LOG(LogVulkanRHI, Display, TEXT("Using instance layers"));
-		for (const ANSICHAR* Layer : OutInstanceLayers)
-		{
-			UE_LOG(LogVulkanRHI, Display, TEXT("* %s"), ANSI_TO_TCHAR(Layer));
-		}
-	}
-	else
-	{
-		UE_LOG(LogVulkanRHI, Display, TEXT("Not using instance layers"));
-	}
-
-	if (OutInstanceExtensions.Num() > 0)
-	{
-		TrimDuplicates(OutInstanceExtensions);
-		UE_LOG(LogVulkanRHI, Display, TEXT("Using instance extensions"));
-		for (const ANSICHAR* Extension : OutInstanceExtensions)
-		{
-			UE_LOG(LogVulkanRHI, Display, TEXT("* %s"), ANSI_TO_TCHAR(Extension));
-		}
-	}
-	else
-	{
-		UE_LOG(LogVulkanRHI, Display, TEXT("Not using instance extensions"));
-	}
-}
-
-void FVulkanDevice::GetDeviceExtensionsAndLayers(VkPhysicalDevice Gpu, EGpuVendorId VendorId, TArray<const ANSICHAR*>& OutDeviceExtensions, TArray<const ANSICHAR*>& OutDeviceLayers, TArray<FString>& OutAllDeviceExtensions, TArray<FString>& OutAllDeviceLayers, bool& bOutDebugMarkers)
-{
-	bOutDebugMarkers = false;
-
-	TArray<FLayerExtension> DeviceLayerExtensions;
-	// 0 is reserved for regular device
-	DeviceLayerExtensions.AddDefaulted();
-	{
-		uint32 Count = 0;
-		TArray<VkLayerProperties> Properties;
-		VERIFYVULKANRESULT_INIT(VulkanRHI::vkEnumerateDeviceLayerProperties(Gpu, &Count, nullptr));
-		Properties.AddZeroed(Count);
-		VERIFYVULKANRESULT_INIT(VulkanRHI::vkEnumerateDeviceLayerProperties(Gpu, &Count, Properties.GetData()));
-		check(Count == Properties.Num());
-		for (const VkLayerProperties& Property : Properties)
-		{
-			DeviceLayerExtensions[DeviceLayerExtensions.AddDefaulted()].LayerProps = Property;
-		}
-	}
-
-	TArray<FString> FoundUniqueLayers;
-	TArray<FString> FoundUniqueExtensions;
-
-	for (int32 Index = 0; Index < DeviceLayerExtensions.Num(); ++Index)
-	{
-		if (Index == 0)
-		{
-			EnumerateDeviceExtensionProperties(Gpu, nullptr, DeviceLayerExtensions[Index]);
-		}
-		else
-		{
-			FoundUniqueLayers.AddUnique(ANSI_TO_TCHAR(DeviceLayerExtensions[Index].LayerProps.layerName));
-			EnumerateDeviceExtensionProperties(Gpu, DeviceLayerExtensions[Index].LayerProps.layerName, DeviceLayerExtensions[Index]);
-		}
-
-		DeviceLayerExtensions[Index].AddUniqueExtensionNames(FoundUniqueExtensions);
-	}
-
-	FoundUniqueLayers.Sort();
-	OutAllDeviceLayers = FoundUniqueLayers;
-
-	FoundUniqueExtensions.Sort();
-	OutAllDeviceExtensions = FoundUniqueExtensions;
-
-	FVulkanPlatform::NotifyFoundDeviceLayersAndExtensions(Gpu, FoundUniqueLayers, FoundUniqueExtensions);
-
-	TArray<FString> UniqueUsedDeviceExtensions;
-	auto AddDeviceLayers = [&](const char* LayerName)
-	{
-		int32 LayerIndex = FindLayerIndexInList(DeviceLayerExtensions, LayerName);
-		if (LayerIndex != INDEX_NONE)
-		{
-			DeviceLayerExtensions[LayerIndex].AddUniqueExtensionNames(UniqueUsedDeviceExtensions);
-		}
-	};
-
-#if VULKAN_HAS_DEBUGGING_ENABLED
-	GRenderDocFound = false;
-	#if VULKAN_ENABLE_DRAW_MARKERS
-	{
-		int32 LayerIndex = FindLayerIndexInList(DeviceLayerExtensions, RENDERDOC_LAYER_NAME);
-		if (LayerIndex != INDEX_NONE)
-		{
-			GRenderDocFound = true;
-			DeviceLayerExtensions[LayerIndex].AddUniqueExtensionNames(UniqueUsedDeviceExtensions);
-		}
-	}
-	#endif
-
-	// Verify that all requested debugging device-layers are available. Skip validation layers under RenderDoc
-	const int32 VulkanValidationOption = GValidationCvar.GetValueOnAnyThread();
-	if (!GRenderDocFound && VulkanValidationOption > 0)
-	{
-		// Path for older drivers
-		bool bStandardAvailable = false;
-		if (GStandardValidationCvar.GetValueOnAnyThread() != 0)
-		{
-			bStandardAvailable = FindLayerInList(DeviceLayerExtensions, STANDARD_VALIDATION_LAYER_NAME);
-			if (bStandardAvailable)
-			{
-				OutDeviceLayers.Add(STANDARD_VALIDATION_LAYER_NAME);
+				for (int32 ExtIndex = 0; ExtIndex < LayerProperties[LayerIndex].ExtensionProperties.Num(); ++ExtIndex)
+				{
+					if (!FCStringAnsi::Strcmp(LayerProperties[LayerIndex].ExtensionProperties[ExtIndex].extensionName, ExtensionName))
+					{
+						return LayerIndex;
+					}
+				}
 			}
-		}
-	}
-#endif	// VULKAN_HAS_DEBUGGING_ENABLED
-
-	// Check for layers added outside the RHI (eg plugins)
-	for (const ANSICHAR* VulkanBridgeLayer : VulkanExternalExtensions::DeviceLayers)
-	{
-		if (FindLayerInList(DeviceLayerExtensions, VulkanBridgeLayer))
-		{
-			OutDeviceLayers.Add(VulkanBridgeLayer);
-		}
-		else
-		{
-			UE_LOG(LogVulkanRHI, Warning, TEXT("Unable to find VulkanExternalExtensions device layer '%s'"), ANSI_TO_TCHAR(VulkanBridgeLayer));
-		}
-	}
+			return (int32)INDEX_NONE;
+		};
 
-	if (FVulkanDynamicRHI::HMDVulkanExtensions.IsValid())
-	{
-		if (!FVulkanDynamicRHI::HMDVulkanExtensions->GetVulkanDeviceExtensionsRequired( Gpu, OutDeviceExtensions))
-		{
-			UE_LOG(LogVulkanRHI, Warning, TEXT( "Trying to use Vulkan with an HMD, but required extensions aren't supported on the selected device!"));
-		}
-	}
+		auto ActivateDebuggingExtension = [&](const ANSICHAR* ExtensionName) {
 
-	// Now gather the actually used extensions based on the enabled layers
-	TArray<const ANSICHAR*> AvailableExtensions;
-	{
-		// All global
-		for (int32 ExtIndex = 0; ExtIndex < DeviceLayerExtensions[0].ExtensionProps.Num(); ++ExtIndex)
-		{
-			AvailableExtensions.Add(DeviceLayerExtensions[0].ExtensionProps[ExtIndex].extensionName);
-		}
+			const int32 ExtensionIndex = FVulkanInstanceExtension::FindExtension(UEExtensions, ExtensionName);
+			check(ExtensionIndex != INDEX_NONE);
 
-		// Now only find enabled layers
-		for (int32 LayerIndex = 0; LayerIndex < OutDeviceLayers.Num(); ++LayerIndex)
-		{
-			// Skip 0 as it's the null layer
-			int32 FindLayerIndex;
-			for (FindLayerIndex = 1; FindLayerIndex < DeviceLayerExtensions.Num(); ++FindLayerIndex)
+			// If the extension isn't supported out of the box, check to see if an extension can add support
+			if (!UEExtensions[ExtensionIndex]->IsSupported())
 			{
-				if (!FCStringAnsi::Strcmp(DeviceLayerExtensions[FindLayerIndex].LayerProps.layerName, OutDeviceLayers[LayerIndex]))
+				const int32 LayerIndex = FindLayerContainingExtension(ExtensionName, LayerProperties);
+				if (LayerIndex != INDEX_NONE)
 				{
-					break;
+					static VkLayerProperties StaticLayer;
+					StaticLayer = LayerProperties[LayerIndex].LayerProperties;  // TODO: prevent from pointing to temporary memory
+					AddRequestedLayer(StaticLayer.layerName, LayerProperties, UEExtensions, OutLayers);
 				}
 			}
 
-			if (FindLayerIndex < DeviceLayerExtensions.Num())
-			{
-				DeviceLayerExtensions[FindLayerIndex].AddAnsiExtensionNames(AvailableExtensions);
-			}
-		}
-	}
-	TrimDuplicates(AvailableExtensions);
-
-	auto ListContains = [](const TArray<const ANSICHAR*>& InList, const ANSICHAR* Name)
-	{
-		for (const ANSICHAR* Element : InList)
-		{
-			if (!FCStringAnsi::Strcmp(Element, Name))
+			// If the extension is supported, activate it and set it as our active debugging extension
+			if (UEExtensions[ExtensionIndex]->IsSupported())
 			{
+				UEExtensions[ExtensionIndex]->SetActivated();
 				return true;
 			}
-		}
-
-		return false;
-	};
-
-	// Now go through the actual requested lists
-	TArray<const ANSICHAR*> PlatformExtensions;
-	FVulkanPlatform::GetDeviceExtensions(VendorId, PlatformExtensions);
-	for (const ANSICHAR* PlatformExtension : PlatformExtensions)
-	{
-		if (ListContains(AvailableExtensions, PlatformExtension))
-		{
-			OutDeviceExtensions.Add(PlatformExtension);
-		}
-	}
-
-	TArray<const ANSICHAR*> PlatformLayers;
-	FVulkanPlatform::GetDeviceLayers(VendorId, PlatformLayers);
-	for (const ANSICHAR* PlatformLayer : PlatformLayers)
-	{
-		OutDeviceLayers.Add(PlatformLayer);
-	}
+			return false;
+		};
 
-	for (uint32 Index = 0; Index < UE_ARRAY_COUNT(GDeviceExtensions) && GDeviceExtensions[Index] != nullptr; ++Index)
-	{
-		if (ListContains(AvailableExtensions, GDeviceExtensions[Index]))
-		{
-			OutDeviceExtensions.Add(GDeviceExtensions[Index]);
-		}
-	}
-	
-	// Check for extensions added outside the RHI (eg plugins)
-	for (const ANSICHAR* VulkanBridgeExtension : VulkanExternalExtensions::DeviceExtensions)
-	{
-		if (ListContains(AvailableExtensions, VulkanBridgeExtension))
+#if VULKAN_SUPPORTS_DEBUG_UTILS
+		if (ActivateDebuggingExtension(VK_EXT_DEBUG_UTILS_EXTENSION_NAME))
 		{
-			OutDeviceExtensions.Add(VulkanBridgeExtension);
+			ActiveDebugLayerExtension = FVulkanDynamicRHI::EActiveDebugLayerExtension::DebugUtilsExtension;
 		}
 		else
-		{
-			UE_LOG(LogVulkanRHI, Warning, TEXT("Unable to find VulkanExternalExtensions device extension '%s'"), ANSI_TO_TCHAR(VulkanBridgeExtension));
-		}
-	}
-
-
-#if VULKAN_ENABLE_DRAW_MARKERS && VULKAN_HAS_DEBUGGING_ENABLED
-	if (!bOutDebugMarkers &&
-		(((GRenderDocFound || VulkanValidationOption == 0) && ListContains(AvailableExtensions, VK_EXT_DEBUG_MARKER_EXTENSION_NAME)) || FVulkanPlatform::ForceEnableDebugMarkers()))
-	{
-		OutDeviceExtensions.Add(VK_EXT_DEBUG_MARKER_EXTENSION_NAME);
-		bOutDebugMarkers = true;
-	}
-#endif
-
-	if (OutDeviceExtensions.Num() > 0)
-	{
-		TrimDuplicates(OutDeviceExtensions);
-	}
-
-	if (OutDeviceLayers.Num() > 0)
-	{
-		TrimDuplicates(OutDeviceLayers);
-	}
-}
-
-static inline bool HasExtension(const TArray<const ANSICHAR*> InExtensions, const ANSICHAR* InName)
-{
-	return InExtensions.ContainsByPredicate(
-		[&InName](const ANSICHAR* Extension) -> bool
-		{
-			return FCStringAnsi::Strcmp(Extension, InName) == 0;
-		}
-	);
-};
-
-void FOptionalVulkanInstanceExtensions::Setup(const TArray<const ANSICHAR*>& InstanceExtensions)
-{
-	check(Packed == 0);
-
-#if VULKAN_SUPPORTS_EXTERNAL_MEMORY
-	HasKHRExternalMemoryCapabilities = HasExtension(InstanceExtensions, VK_KHR_EXTERNAL_MEMORY_CAPABILITIES_EXTENSION_NAME);  // Promoted to 1.1
-#endif
-
-#if VULKAN_SUPPORTS_PHYSICAL_DEVICE_PROPERTIES2
-	HasKHRGetPhysicalDeviceProperties2 = HasExtension(InstanceExtensions, VK_KHR_GET_PHYSICAL_DEVICE_PROPERTIES_2_EXTENSION_NAME);  // Promoted to 1.1
-#endif
-}
-
-void FOptionalVulkanDeviceExtensions::Setup(const TArray<const ANSICHAR*>& DeviceExtensions)
-{
-	check(Packed == 0);
-
-
-	// Optional Extensions
-#if VULKAN_SUPPORTS_VALIDATION_CACHE
-	HasEXTValidationCache = HasExtension(DeviceExtensions, VK_EXT_VALIDATION_CACHE_EXTENSION_NAME);
-#endif
-#if VULKAN_SUPPORTS_MEMORY_PRIORITY
-	HasMemoryPriority = HasExtension(DeviceExtensions, VK_EXT_MEMORY_PRIORITY_EXTENSION_NAME) && (!FParse::Param(FCommandLine::Get(), TEXT("disablememorypriority")));
-#endif
-#if VULKAN_SUPPORTS_MEMORY_BUDGET
-	HasMemoryBudget = HasExtension(DeviceExtensions, VK_EXT_MEMORY_BUDGET_EXTENSION_NAME) && (!FParse::Param(FCommandLine::Get(), TEXT("disablememorybudget")));
-#endif
-#if VULKAN_SUPPORTS_ASTC_DECODE_MODE
-	HasEXTASTCDecodeMode = HasExtension(DeviceExtensions, VK_EXT_ASTC_DECODE_MODE_EXTENSION_NAME);
-#endif
-#if VULKAN_SUPPORTS_FRAGMENT_DENSITY_MAP
-	HasEXTFragmentDensityMap = HasExtension(DeviceExtensions, VK_EXT_FRAGMENT_DENSITY_MAP_EXTENSION_NAME);
-#endif
-#if VULKAN_SUPPORTS_FRAGMENT_DENSITY_MAP2
-	HasEXTFragmentDensityMap2 = HasExtension(DeviceExtensions, VK_EXT_FRAGMENT_DENSITY_MAP_2_EXTENSION_NAME);
-#endif
-#if VULKAN_SUPPORTS_FRAGMENT_SHADING_RATE
-	HasKHRFragmentShadingRate = HasExtension(DeviceExtensions, VK_KHR_FRAGMENT_SHADING_RATE_EXTENSION_NAME);
-#endif
-#if VULKAN_SUPPORTS_FULLSCREEN_EXCLUSIVE
-	HasEXTFullscreenExclusive = HasExtension(DeviceExtensions, VK_EXT_FULL_SCREEN_EXCLUSIVE_EXTENSION_NAME);
-#endif
-#if VULKAN_SUPPORTS_IMAGE_64BIT_ATOMICS
-	HasImageAtomicInt64 = HasExtension(DeviceExtensions, VK_EXT_SHADER_IMAGE_ATOMIC_INT64_EXTENSION_NAME);
-#endif
-
-
-	// Vendor
-	{
-		bool bHasAnyCrashExtension = false;
-
-#if VULKAN_SUPPORTS_AMD_BUFFER_MARKER
-		if (GGPUCrashDebuggingEnabled)
-		{
-			HasAMDBufferMarker = HasExtension(DeviceExtensions, VK_AMD_BUFFER_MARKER_EXTENSION_NAME);
-			bHasAnyCrashExtension = bHasAnyCrashExtension || HasAMDBufferMarker;
-		}
 #endif
-
-#if VULKAN_SUPPORTS_NV_DIAGNOSTICS
-		if (GGPUCrashDebuggingEnabled)
+		if (ActivateDebuggingExtension(VK_EXT_DEBUG_REPORT_EXTENSION_NAME))
 		{
-			HasNVDiagnosticCheckpoints = HasExtension(DeviceExtensions, VK_NV_DEVICE_DIAGNOSTIC_CHECKPOINTS_EXTENSION_NAME);
-			HasNVDeviceDiagnosticConfig = HasExtension(DeviceExtensions, VK_NV_DEVICE_DIAGNOSTICS_CONFIG_EXTENSION_NAME);
-			bHasAnyCrashExtension = bHasAnyCrashExtension || (HasNVDeviceDiagnosticConfig && HasNVDiagnosticCheckpoints);
+			ActiveDebugLayerExtension = FVulkanDynamicRHI::EActiveDebugLayerExtension::DebugReportExtension;
 		}
-#endif
 
-		if (GGPUCrashDebuggingEnabled && !bHasAnyCrashExtension)
+		const bool bRequiresValidationFeatures = (GGPUValidationCvar.GetValueOnAnyThread() != 0) ||
+			(FParse::Param(FCommandLine::Get(), TEXT("vulkanbestpractices"))) ||
+			(FParse::Param(FCommandLine::Get(), TEXT("vulkandebugsync")));
+		if (bRequiresValidationFeatures)
 		{
-			UE_LOG(LogVulkanRHI, Warning, TEXT("Tried to enable GPU crash debugging but no extension found! Will use local tracepoints."));
+			ActivateDebuggingExtension(VK_EXT_VALIDATION_FEATURES_EXTENSION_NAME);
 		}
-
-#if VULKAN_SUPPORTS_QCOM_RENDERPASS_TRANSFORM
-		HasQcomRenderPassTransform = HasExtension(DeviceExtensions, VK_QCOM_RENDER_PASS_TRANSFORM_EXTENSION_NAME);
-#endif
-	}
-
-
-	// Promoted to 1.1
-	{
-#if VULKAN_SUPPORTS_MAINTENANCE_LAYER1
-		HasKHRMaintenance1 = HasExtension(DeviceExtensions, VK_KHR_MAINTENANCE1_EXTENSION_NAME);
-#endif
-#if VULKAN_SUPPORTS_MAINTENANCE_LAYER2
-		HasKHRMaintenance2 = HasExtension(DeviceExtensions, VK_KHR_MAINTENANCE2_EXTENSION_NAME);
-#endif
-#if VULKAN_SUPPORTS_DEDICATED_ALLOCATION
-		HasKHRDedicatedAllocation = HasExtension(DeviceExtensions, VK_KHR_GET_MEMORY_REQUIREMENTS_2_EXTENSION_NAME) && HasExtension(DeviceExtensions, VK_KHR_DEDICATED_ALLOCATION_EXTENSION_NAME);
-#endif
-#if VULKAN_SUPPORTS_MULTIVIEW
-		HasKHRMultiview = HasExtension(DeviceExtensions, VK_KHR_MULTIVIEW_EXTENSION_NAME);
-#endif
 	}
+#endif	// VULKAN_HAS_DEBUGGING_ENABLED
+}
 
 
-	// Promoted to 1.2
-	{
-#if VULKAN_SUPPORTS_DRIVER_PROPERTIES
-		HasDriverProperties = HasExtension(DeviceExtensions, VK_KHR_DRIVER_PROPERTIES_EXTENSION_NAME);
-#endif
-#if VULKAN_SUPPORTS_RENDERPASS2
-		HasKHRRenderPass2 = HasExtension(DeviceExtensions, VK_KHR_CREATE_RENDERPASS_2_EXTENSION_NAME);
-#endif
 
-		HasKHRImageFormatList = HasExtension(DeviceExtensions, VK_KHR_IMAGE_FORMAT_LIST_EXTENSION_NAME);
-
-#if VULKAN_SUPPORTS_BUFFER_64BIT_ATOMICS
-		HasKHRShaderAtomicInt64 = HasExtension(DeviceExtensions, VK_KHR_SHADER_ATOMIC_INT64_EXTENSION_NAME);
-#endif
-#if VULKAN_SUPPORTS_SCALAR_BLOCK_LAYOUT
-		HasEXTScalarBlockLayout = HasExtension(DeviceExtensions, VK_EXT_SCALAR_BLOCK_LAYOUT_EXTENSION_NAME);
-#endif
-#if VULKAN_SUPPORTS_DESCRIPTOR_INDEXING
-		HasEXTDescriptorIndexing = HasExtension(DeviceExtensions, VK_EXT_DESCRIPTOR_INDEXING_EXTENSION_NAME);
-#endif
-#if VULKAN_SUPPORTS_SHADER_VIEWPORT_INDEX_LAYER
-		HasEXTShaderViewportIndexLayer = HasExtension(DeviceExtensions, VK_EXT_SHADER_VIEWPORT_INDEX_LAYER_EXTENSION_NAME);
+// Return a list of debug layers to activate
+void FVulkanDeviceSetupHelper::AddDebugLayers(const TArray<FLayerWithExtensions>& LayerProperties, FVulkanDeviceExtensionArray& UEExtensions, TArray<const ANSICHAR*>& OutLayers)
+{
+#if VULKAN_ENABLE_DRAW_MARKERS
+	GRenderDocFound = (FindLayerIndexInList(RENDERDOC_LAYER_NAME, LayerProperties) != INDEX_NONE);
+#else
+	GRenderDocFound = false;
 #endif
-	}
-
 
-	// Promoted to 1.3
+#if VULKAN_HAS_DEBUGGING_ENABLED
+	// Verify that all requested debugging device-layers are available. Skip validation layers under RenderDoc
+	if (!GRenderDocFound && (GValidationCvar.GetValueOnAnyThread() > 0) && (GStandardValidationCvar.GetValueOnAnyThread() != 0))
 	{
-#if VULKAN_SUPPORTS_TEXTURE_COMPRESSION_ASTC_HDR
-		HasEXTTextureCompressionASTCHDR = HasExtension(DeviceExtensions, VK_EXT_TEXTURE_COMPRESSION_ASTC_HDR_EXTENSION_NAME);
-#endif
-#if VULKAN_SUPPORTS_MAINTENANCE4
-		HasKHRMaintenance4 = HasExtension(DeviceExtensions, VK_KHR_MAINTENANCE_4_EXTENSION_NAME);
-#endif
+		// Path for older drivers
+		AddRequestedLayer(STANDARD_VALIDATION_LAYER_NAME, LayerProperties, UEExtensions, OutLayers);
 	}
-
-
-#if VULKAN_RHI_RAYTRACING
-	// Optional Extensions
-	HasAccelerationStructure	= HasExtension(DeviceExtensions, VK_KHR_ACCELERATION_STRUCTURE_EXTENSION_NAME);
-	HasRayTracingPipeline		= HasExtension(DeviceExtensions, VK_KHR_RAY_TRACING_PIPELINE_EXTENSION_NAME);
-	HasRayQuery		            = HasExtension(DeviceExtensions, VK_KHR_RAY_QUERY_EXTENSION_NAME);
-	HasDeferredHostOperations	= HasExtension(DeviceExtensions, VK_KHR_DEFERRED_HOST_OPERATIONS_EXTENSION_NAME);
-
-	// Promoted to 1.2
-	HasBufferDeviceAddress		= HasExtension(DeviceExtensions, VK_KHR_BUFFER_DEVICE_ADDRESS_EXTENSION_NAME);
-	HasSPIRV_14					= HasExtension(DeviceExtensions, VK_KHR_SPIRV_1_4_EXTENSION_NAME);
-	HasShaderFloatControls		= HasExtension(DeviceExtensions, VK_KHR_SHADER_FLOAT_CONTROLS_EXTENSION_NAME);
-#endif
+#endif	// VULKAN_HAS_DEBUGGING_ENABLED
 }
 
-void FVulkanDynamicRHI::SetupValidationRequests()
-{
-#if VULKAN_HAS_DEBUGGING_ENABLED
-	int32 VulkanValidationOption = GValidationCvar.GetValueOnAnyThread();
 
-	// Command line overrides Cvar
-	if (FParse::Param(FCommandLine::Get(), TEXT("vulkandebug")))
-	{
-		// Match D3D and GL
-		GValidationCvar->Set(2, ECVF_SetByCommandline);
-	}
-	else if (FParse::Value(FCommandLine::Get(), TEXT("vulkanvalidation="), VulkanValidationOption))
-	{
-		GValidationCvar->Set(VulkanValidationOption, ECVF_SetByCommandline);
-	}
 
-	if (FParse::Param(FCommandLine::Get(), TEXT("gpuvalidation")))
-	{
-		if (GValidationCvar->GetInt() < 2)
-		{
-			GValidationCvar->Set(2, ECVF_SetByCommandline);
-		}
-		GGPUValidationCvar->Set(2, ECVF_SetByCommandline);
-	}
-#endif
-}
+
diff --git a/Engine/Source/Runtime/VulkanRHI/Private/VulkanLoader.h b/Engine/Source/Runtime/VulkanRHI/Private/VulkanLoader.h
--- a/Engine/Source/Runtime/VulkanRHI/Private/VulkanLoader.h
+++ b/Engine/Source/Runtime/VulkanRHI/Private/VulkanLoader.h
@@ -172,6 +172,7 @@
 
 // List all optional Vulkan entry points used by Unreal that need to be loaded manually
 #define ENUM_VK_ENTRYPOINTS_OPTIONAL_INSTANCE(EnumMacro) \
+	EnumMacro(PFN_vkGetDeviceImageMemoryRequirementsKHR, vkGetDeviceImageMemoryRequirementsKHR) \
 	EnumMacro(PFN_vkCreateSharedSwapchainsKHR, vkCreateSharedSwapchainsKHR)
 
 // List of all Vulkan entry points
diff --git a/Engine/Source/Runtime/VulkanRHI/Private/VulkanRHI.cpp b/Engine/Source/Runtime/VulkanRHI/Private/VulkanRHI.cpp
--- a/Engine/Source/Runtime/VulkanRHI/Private/VulkanRHI.cpp
+++ b/Engine/Source/Runtime/VulkanRHI/Private/VulkanRHI.cpp
@@ -24,6 +24,7 @@
 #include "IHeadMountedDisplayModule.h"
 #include "VulkanRenderpass.h"
 #include "VulkanTransientResourceAllocator.h"
+#include "VulkanExtensions.h"
 
 static_assert(sizeof(VkStructureType) == sizeof(int32), "ZeroVulkanStruct() assumes VkStructureType is int32!");
 
@@ -96,30 +97,6 @@ extern TAutoConsoleVariable<int32> GRHIAllowAsyncComputeCvar;
 // All shader stages supported by VK device - VK_PIPELINE_STAGE_VERTEX_SHADER_BIT, FRAGMENT etc
 uint32 GVulkanDeviceShaderStageBits = 0;
 
-#if VULKAN_HAS_VALIDATION_FEATURES
-static inline TArray<VkValidationFeatureEnableEXT> GetValidationFeaturesEnabled(bool bEnableValidation)
-{
-	TArray<VkValidationFeatureEnableEXT> Features;
-	extern TAutoConsoleVariable<int32> GGPUValidationCvar;
-	int32 GPUValidationValue = GGPUValidationCvar.GetValueOnAnyThread();
-	if (bEnableValidation && GPUValidationValue > 0)
-	{
-		Features.Add(VK_VALIDATION_FEATURE_ENABLE_GPU_ASSISTED_EXT);
-		if (GPUValidationValue > 1)
-		{
-			Features.Add(VK_VALIDATION_FEATURE_ENABLE_GPU_ASSISTED_RESERVE_BINDING_SLOT_EXT);
-		}
-	}
-
-	if (FParse::Param(FCommandLine::Get(), TEXT("vulkanbestpractices")))
-	{
-		Features.Add(VK_VALIDATION_FEATURE_ENABLE_BEST_PRACTICES_EXT);
-	}
-
-	return Features;
-}
-#endif
-
 DEFINE_LOG_CATEGORY(LogVulkan)
 
 bool FVulkanDynamicRHIModule::IsSupported()
@@ -252,6 +229,8 @@ void FVulkanDynamicRHI::Init()
 	// Setup the validation requests ready before we load dlls
 	SetupValidationRequests();
 
+	UE_LOG(LogVulkanRHI, Display, TEXT("Built with Vulkan header version %u.%u.%u"), VK_API_VERSION_MAJOR(VK_HEADER_VERSION_COMPLETE), VK_API_VERSION_MINOR(VK_HEADER_VERSION_COMPLETE), VK_API_VERSION_PATCH(VK_HEADER_VERSION_COMPLETE));
+
 	if (!FVulkanPlatform::LoadVulkanLibrary())
 	{
 #if PLATFORM_LINUX
@@ -428,39 +407,27 @@ void FVulkanDynamicRHI::CreateInstance()
 	ZeroVulkanStruct(InstInfo, VK_STRUCTURE_TYPE_INSTANCE_CREATE_INFO);
 	InstInfo.pApplicationInfo = &AppInfo;
 
-	GetInstanceLayersAndExtensions(InstanceExtensions, InstanceLayers, bSupportsDebugUtilsExt);
+	FVulkanInstanceExtensionArray UEInstanceExtensions = FVulkanInstanceExtension::GetUESupportedInstanceExtensions();
+	InstanceLayers = SetupInstanceLayers(UEInstanceExtensions);
+	for (TUniquePtr<FVulkanInstanceExtension>& Extension : UEInstanceExtensions)
+	{
+		if (Extension->InUse())
+		{
+			InstanceExtensions.Add(Extension->GetExtensionName());
+			Extension->PreCreateInstance(InstInfo, OptionalInstanceExtensions);
+		}
+	}
 
 	InstInfo.enabledExtensionCount = InstanceExtensions.Num();
 	InstInfo.ppEnabledExtensionNames = InstInfo.enabledExtensionCount > 0 ? (const ANSICHAR* const*)InstanceExtensions.GetData() : nullptr;
 	
 	InstInfo.enabledLayerCount = InstanceLayers.Num();
 	InstInfo.ppEnabledLayerNames = InstInfo.enabledLayerCount > 0 ? InstanceLayers.GetData() : nullptr;
-#if VULKAN_HAS_DEBUGGING_ENABLED
-	bSupportsDebugCallbackExt = !bSupportsDebugUtilsExt && InstanceExtensions.ContainsByPredicate([](const ANSICHAR* Key)
-		{ 
-			return Key && !FCStringAnsi::Strcmp(Key, VK_EXT_DEBUG_REPORT_EXTENSION_NAME); 
-		});
-
-#if VULKAN_HAS_VALIDATION_FEATURES
-	bool bHasGPUValidation = InstanceExtensions.ContainsByPredicate([](const ANSICHAR* Key)
-		{
-			return Key && !FCStringAnsi::Strcmp(Key, VK_EXT_VALIDATION_FEATURES_EXTENSION_NAME);
-		});
-	VkValidationFeaturesEXT ValidationFeatures;
-	TArray<VkValidationFeatureEnableEXT> ValidationFeaturesEnabled = GetValidationFeaturesEnabled(bHasGPUValidation);
-	if (bHasGPUValidation)
-	{
-		ZeroVulkanStruct(ValidationFeatures, VK_STRUCTURE_TYPE_VALIDATION_FEATURES_EXT);
-		ValidationFeatures.pNext = InstInfo.pNext;
-		ValidationFeatures.enabledValidationFeatureCount = (uint32)ValidationFeaturesEnabled.Num();
-		ValidationFeatures.pEnabledValidationFeatures = ValidationFeaturesEnabled.GetData();
-		InstInfo.pNext = &ValidationFeatures;
-	}
-#endif
-#endif
 
 	VkResult Result = VulkanRHI::vkCreateInstance(&InstInfo, VULKAN_CPU_ALLOCATOR, &Instance);
 	
+	FVulkanPlatform::NotifyFoundInstanceLayersAndExtensions(InstanceLayers, InstanceExtensions);
+
 	if (Result == VK_ERROR_INCOMPATIBLE_DRIVER)
 	{
 		FPlatformMisc::MessageBoxExt(EAppMsgType::Ok, TEXT(
@@ -531,8 +498,6 @@ void FVulkanDynamicRHI::CreateInstance()
 #if VULKAN_HAS_DEBUGGING_ENABLED
 	SetupDebugLayerCallback();
 #endif
-
-	OptionalInstanceExtensions.Setup(InstanceExtensions);
 }
 
 //#todo-rco: Common RHI should handle this...
@@ -587,20 +552,21 @@ void FVulkanDynamicRHI::SelectAndInitDevice()
 
 #if VULKAN_ENABLE_DESKTOP_HMD_SUPPORT
 	// Allow HMD to override which graphics adapter is chosen, so we pick the adapter where the HMD is connected
-	uint64 HmdGraphicsAdapterLuid  = IHeadMountedDisplayModule::IsAvailable() ? IHeadMountedDisplayModule::Get().GetGraphicsAdapterLuid() : 0;
+	const uint64 HmdGraphicsAdapterLuid  = IHeadMountedDisplayModule::IsAvailable() ? IHeadMountedDisplayModule::Get().GetGraphicsAdapterLuid() : 0;
 #endif
 
 	UE_LOG(LogVulkanRHI, Display, TEXT("Found %d device(s)"), GpuCount);
 	for (uint32 Index = 0; Index < GpuCount; ++Index)
 	{
+		UE_LOG(LogVulkanRHI, Display, TEXT("Device %d:"), Index);
 		FVulkanDevice* NewDevice = new FVulkanDevice(this, PhysicalDevices[Index]);
 		Devices.Add(NewDevice);
 
-		bool bIsDiscrete = NewDevice->QueryGPU(Index);
+		const bool bIsDiscrete = (NewDevice->GetDeviceProperties().deviceType == VK_PHYSICAL_DEVICE_TYPE_DISCRETE_GPU);
 
 #if VULKAN_ENABLE_DESKTOP_HMD_SUPPORT
 		if (!HmdDevice && HmdGraphicsAdapterLuid != 0 &&
-			NewDevice->GetOptionalExtensions().HasKHRGetPhysicalDeviceProperties2 &&
+			GetOptionalExtensions().HasKHRGetPhysicalDeviceProperties2 &&
 			FMemory::Memcmp(&HmdGraphicsAdapterLuid, &NewDevice->GetDeviceIdProperties().deviceLUID, VK_LUID_SIZE_KHR) == 0)
 		{
 			HmdDevice = NewDevice;
@@ -756,24 +722,24 @@ void FVulkanDynamicRHI::SelectAndInitDevice()
 
 		if (GPUDriverInfo.InternalDriverVersion != TEXT("Unknown"))
 		{
-		GRHIAdapterUserDriverVersion = GPUDriverInfo.UserDriverVersion;
-		GRHIAdapterInternalDriverVersion = GPUDriverInfo.InternalDriverVersion;
-		GRHIAdapterDriverDate = GPUDriverInfo.DriverDate;
+			GRHIAdapterUserDriverVersion = GPUDriverInfo.UserDriverVersion;
+			GRHIAdapterInternalDriverVersion = GPUDriverInfo.InternalDriverVersion;
+			GRHIAdapterDriverDate = GPUDriverInfo.DriverDate;
 
-		UE_LOG(LogVulkanRHI, Log, TEXT("    Adapter Name: %s"), *GRHIAdapterName);
-		UE_LOG(LogVulkanRHI, Log, TEXT("     API Version: %d.%d.%d"), VK_VERSION_MAJOR(Props.apiVersion), VK_VERSION_MINOR(Props.apiVersion), VK_VERSION_PATCH(Props.apiVersion));
+			UE_LOG(LogVulkanRHI, Log, TEXT("    Adapter Name: %s"), *GRHIAdapterName);
+			UE_LOG(LogVulkanRHI, Log, TEXT("     API Version: %d.%d.%d"), VK_VERSION_MAJOR(Props.apiVersion), VK_VERSION_MINOR(Props.apiVersion), VK_VERSION_PATCH(Props.apiVersion));
 			UE_LOG(LogVulkanRHI, Log, TEXT("  Driver Version: %s (0x%X)"), *GRHIAdapterUserDriverVersion, Props.driverVersion);
-		UE_LOG(LogVulkanRHI, Log, TEXT("Internal Version: %s"), *GRHIAdapterInternalDriverVersion);
-		UE_LOG(LogVulkanRHI, Log, TEXT("     Driver Date: %s"), *GRHIAdapterDriverDate);
-	}
+			UE_LOG(LogVulkanRHI, Log, TEXT("Internal Version: %s"), *GRHIAdapterInternalDriverVersion);
+			UE_LOG(LogVulkanRHI, Log, TEXT("     Driver Date: %s"), *GRHIAdapterDriverDate);
+		}
 		else
-	{
+		{
 			// If we failed to read from the registry, then use the values provided by Vulkan props
 			ReadVulkanDriverVersionFromProps(Device);
 		}
-		}
+	}
 	else if(PLATFORM_UNIX)
-		{
+	{
 		ReadVulkanDriverVersionFromProps(Device);
 	}
 
@@ -805,7 +771,6 @@ void FVulkanDynamicRHI::InitInstance()
 			EnableIdealGPUCaptureOptions(true);
 		}
 #endif
-		//bool bDeviceSupportsTessellation = Device->GetPhysicalFeatures().tessellationShader != 0;
 
 		const VkPhysicalDeviceProperties& Props = Device->GetDeviceProperties();
 		const VkPhysicalDeviceLimits& Limits = Device->GetLimits();
@@ -819,7 +784,7 @@ void FVulkanDynamicRHI::InitInstance()
 		GRHISupportsTextureStreaming = true;
 		GSupportsTimestampRenderQueries = FVulkanPlatform::SupportsTimestampRenderQueries();
 #if VULKAN_SUPPORTS_MULTIVIEW
-		GSupportsMobileMultiView = Device->GetMultiviewFeatures().multiview == VK_TRUE ? true : false;
+		GSupportsMobileMultiView = Device->GetOptionalExtensions().HasKHRMultiview ? true : false;
 #endif
 #if VULKAN_RHI_RAYTRACING
 		GRHISupportsRayTracing = RHISupportsRayTracing(GMaxRHIShaderPlatform) && Device->GetOptionalExtensions().HasRaytracingExtensions();
@@ -876,22 +841,6 @@ void FVulkanDynamicRHI::InitInstance()
 		GRHIMaxDispatchThreadGroupsPerDimension.Y = FMath::Min<uint32>(Limits.maxComputeWorkGroupCount[1], 0x7fffffff);
 		GRHIMaxDispatchThreadGroupsPerDimension.Z = FMath::Min<uint32>(Limits.maxComputeWorkGroupCount[2], 0x7fffffff);
 
-#if VULKAN_SUPPORTS_SHADER_VIEWPORT_INDEX_LAYER
-		GRHISupportsArrayIndexFromAnyShader = GetDevice()->GetOptionalExtensions().HasEXTShaderViewportIndexLayer != 0ULL;
-#endif
-
-#if VULKAN_SUPPORTS_FRAGMENT_DENSITY_MAP
-		GRHISupportsAttachmentVariableRateShading = (GetDevice()->GetOptionalExtensions().HasEXTFragmentDensityMap && Device->GetFragmentDensityMapFeatures().fragmentDensityMap);
-#endif
-
-#if VULKAN_SUPPORTS_FRAGMENT_DENSITY_MAP2
-		GRHISupportsLateVariableRateShadingUpdate = GetDevice()->GetOptionalExtensions().HasEXTFragmentDensityMap2 && Device->GetFragmentDensityMap2Features().fragmentDensityMapDeferred;
-#endif
-
-#if VULKAN_SUPPORTS_FRAGMENT_SHADING_RATE
-		GRHISupportsAttachmentVariableRateShading |= GetDevice()->GetOptionalExtensions().HasKHRFragmentShadingRate && Device->GetFragmentShadingRateFeatures().attachmentFragmentShadingRate;
-#endif
-
 		FVulkanPlatform::SetupFeatureLevels();
 
 		GRHIRequiresRenderTargetForPixelShaderUAVs = true;
@@ -908,6 +857,11 @@ void FVulkanDynamicRHI::InitInstance()
 			GVulkanDeviceShaderStageBits|= VK_PIPELINE_STAGE_GEOMETRY_SHADER_BIT;
 		}
 		
+		if (GGPUCrashDebuggingEnabled && !Device->GetOptionalExtensions().HasGPUCrashDumpExtensions())
+		{
+			UE_LOG(LogVulkanRHI, Warning, TEXT("Tried to enable GPU crash debugging but no extension found! Will use local tracepoints."));
+		}
+
 		FHardwareInfo::RegisterHardwareInfo(NAME_RHI, TEXT("Vulkan"));
 
 		SavePipelineCacheCmd = IConsoleManager::Get().RegisterConsoleCommand(
@@ -1240,7 +1194,7 @@ uint64 FVulkanDynamicRHI::RHIGetGraphicsAdapterLUID(VkPhysicalDevice InPhysicalD
 {
 	uint64 AdapterLUID = 0;
 #if VULKAN_SUPPORTS_DRIVER_PROPERTIES
-	if (!AdapterLUID && GetOptionalExtensions().HasKHRGetPhysicalDeviceProperties2)
+	if (GetOptionalExtensions().HasKHRGetPhysicalDeviceProperties2)
 	{
 		VkPhysicalDeviceIDPropertiesKHR GpuIdProps;
 		VkPhysicalDeviceProperties2KHR GpuProps2;
diff --git a/Engine/Source/Runtime/VulkanRHI/Private/VulkanRayTracing.cpp b/Engine/Source/Runtime/VulkanRHI/Private/VulkanRayTracing.cpp
--- a/Engine/Source/Runtime/VulkanRHI/Private/VulkanRayTracing.cpp
+++ b/Engine/Source/Runtime/VulkanRHI/Private/VulkanRayTracing.cpp
@@ -12,73 +12,11 @@
 #include "Misc/CommandLine.h"
 #include "Misc/Parse.h"
 
-TAutoConsoleVariable<int32> GVulkanRayTracingCVar(
-	TEXT("r.Vulkan.RayTracing"),
-	0,
-	TEXT("0: Do not enable Vulkan ray tracing extensions (default)\n")
-	TEXT("1: Enable experimental ray tracing support (for development and testing purposes)"),
-	ECVF_ReadOnly
-);
 
 // Define ray tracing entry points
 #define DEFINE_VK_ENTRYPOINTS(Type,Func) VULKANRHI_API Type VulkanDynamicAPI::Func = NULL;
 ENUM_VK_ENTRYPOINTS_RAYTRACING(DEFINE_VK_ENTRYPOINTS)
 
-void FVulkanRayTracingPlatform::GetDeviceExtensions(EGpuVendorId VendorId, TArray<const ANSICHAR*>& OutExtensions)
-{
-	if (!GVulkanRayTracingCVar.GetValueOnAnyThread() || FParse::Param(FCommandLine::Get(), TEXT("noraytracing")))
-	{
-		return;
-	}
-
-	// Primary extensions
-	OutExtensions.Add(VK_KHR_ACCELERATION_STRUCTURE_EXTENSION_NAME);
-	OutExtensions.Add(VK_KHR_RAY_TRACING_PIPELINE_EXTENSION_NAME);
-	OutExtensions.Add(VK_KHR_RAY_QUERY_EXTENSION_NAME);
-
-	// VK_KHR_acceleration_structure dependencies 
-	OutExtensions.Add(VK_EXT_DESCRIPTOR_INDEXING_EXTENSION_NAME); // Note: Promoted to Vulkan 1.2
-	OutExtensions.Add(VK_KHR_BUFFER_DEVICE_ADDRESS_EXTENSION_NAME); // Note: Promoted to Vulkan 1.2
-	OutExtensions.Add(VK_KHR_DEFERRED_HOST_OPERATIONS_EXTENSION_NAME);
-	
-	// VK_KHR_ray_tracing_pipeline dependency
-	OutExtensions.Add(VK_KHR_SPIRV_1_4_EXTENSION_NAME); // Note: Promoted to Vulkan 1.2
-
-	// VK_KHR_spirv_1_4 dependency
-	OutExtensions.Add(VK_KHR_SHADER_FLOAT_CONTROLS_EXTENSION_NAME); // Note: Promoted to Vulkan 1.2
-}
-
-void FVulkanRayTracingPlatform::EnablePhysicalDeviceFeatureExtensions(VkDeviceCreateInfo& DeviceInfo, FVulkanDevice& Device)
-{
-	if (Device.GetOptionalExtensions().HasRaytracingExtensions())
-	{
-		FOptionalVulkanDeviceFeatures& Features = Device.GetOptionalFeatures();
-
-		ZeroVulkanStruct(Features.BufferDeviceAddressFeatures, VK_STRUCTURE_TYPE_PHYSICAL_DEVICE_BUFFER_DEVICE_ADDRESS_FEATURES);
-		Features.BufferDeviceAddressFeatures.bufferDeviceAddress = VK_TRUE;
-		Features.BufferDeviceAddressFeatures.pNext = const_cast<void*>(DeviceInfo.pNext);
-
-		ZeroVulkanStruct(Features.AccelerationStructureFeatures, VK_STRUCTURE_TYPE_PHYSICAL_DEVICE_ACCELERATION_STRUCTURE_FEATURES_KHR);;
-		Features.AccelerationStructureFeatures.accelerationStructure = VK_TRUE;
-		Features.AccelerationStructureFeatures.descriptorBindingAccelerationStructureUpdateAfterBind = VK_TRUE;
-		Features.AccelerationStructureFeatures.pNext = &Features.BufferDeviceAddressFeatures;
-		
-		ZeroVulkanStruct(Features.RayTracingPipelineFeatures, VK_STRUCTURE_TYPE_PHYSICAL_DEVICE_RAY_TRACING_PIPELINE_FEATURES_KHR);
-		Features.RayTracingPipelineFeatures.rayTracingPipeline = VK_TRUE;
-		Features.RayTracingPipelineFeatures.rayTraversalPrimitiveCulling = VK_TRUE;
-		Features.RayTracingPipelineFeatures.pNext = &Features.AccelerationStructureFeatures;
-
-		ZeroVulkanStruct(Features.RayQueryFeatures, VK_STRUCTURE_TYPE_PHYSICAL_DEVICE_RAY_QUERY_FEATURES_KHR);
-		Features.RayQueryFeatures.pNext = &Features.RayTracingPipelineFeatures;
-		Features.RayQueryFeatures.rayQuery = VK_TRUE;
-
-		ZeroVulkanStruct(Features.DescriptorIndexingFeatures, VK_STRUCTURE_TYPE_PHYSICAL_DEVICE_DESCRIPTOR_INDEXING_FEATURES_EXT);
-		Features.DescriptorIndexingFeatures.pNext = &Features.RayQueryFeatures;
-
-		DeviceInfo.pNext = &Features.DescriptorIndexingFeatures;
-	}
-}
-
 #if PLATFORM_WINDOWS
 #pragma warning(push)
 #pragma warning(disable : 4191) // warning C4191: 'type cast': unsafe conversion
diff --git a/Engine/Source/Runtime/VulkanRHI/Private/VulkanRayTracing.h b/Engine/Source/Runtime/VulkanRHI/Private/VulkanRayTracing.h
--- a/Engine/Source/Runtime/VulkanRHI/Private/VulkanRayTracing.h
+++ b/Engine/Source/Runtime/VulkanRHI/Private/VulkanRayTracing.h
@@ -30,8 +30,6 @@ namespace VulkanDynamicAPI
 class FVulkanRayTracingPlatform
 {
 public:
-	static void GetDeviceExtensions(EGpuVendorId VendorId, TArray<const ANSICHAR*>& OutExtensions);
-	static void EnablePhysicalDeviceFeatureExtensions(VkDeviceCreateInfo& DeviceInfo, FVulkanDevice& Device);
 	static bool LoadVulkanInstanceFunctions(VkInstance inInstance);
 };
 
diff --git a/Engine/Source/Runtime/VulkanRHI/Private/Windows/VulkanWindowsPlatform.cpp b/Engine/Source/Runtime/VulkanRHI/Private/Windows/VulkanWindowsPlatform.cpp
--- a/Engine/Source/Runtime/VulkanRHI/Private/Windows/VulkanWindowsPlatform.cpp
+++ b/Engine/Source/Runtime/VulkanRHI/Private/Windows/VulkanWindowsPlatform.cpp
@@ -4,6 +4,7 @@
 #include "../VulkanRHIPrivate.h"
 #include "../VulkanDevice.h"
 #include "../VulkanRayTracing.h"
+#include "../VulkanExtensions.h"
 
 // Disable warning about forward declared enumeration without a type, since the D3D specific enums are not used in this translation unit
 #pragma warning(push)
@@ -170,74 +171,21 @@ void FVulkanWindowsPlatform::FreeVulkanLibrary()
 
 #include "Windows/HideWindowsPlatformTypes.h"
 
-
-
-void FVulkanWindowsPlatform::GetInstanceExtensions(TArray<const ANSICHAR*>& OutExtensions)
+void FVulkanWindowsPlatform::GetInstanceExtensions(FVulkanInstanceExtensionArray& OutExtensions)
 {
-	// windows surface extension
-	OutExtensions.Add(VK_KHR_SURFACE_EXTENSION_NAME);
-	OutExtensions.Add(VK_KHR_WIN32_SURFACE_EXTENSION_NAME);
-
-#if VULKAN_SUPPORTS_FULLSCREEN_EXCLUSIVE
-	// Required by Fullscreen
-	OutExtensions.Add(VK_KHR_GET_SURFACE_CAPABILITIES_2_EXTENSION_NAME);
-#endif
+	OutExtensions.Add(MakeUnique<FVulkanInstanceExtension>(VK_KHR_WIN32_SURFACE_EXTENSION_NAME, VULKAN_EXTENSION_ENABLED, VULKAN_EXTENSION_NOT_PROMOTED));
+	OutExtensions.Add(MakeUnique<FVulkanInstanceExtension>(VK_KHR_GET_SURFACE_CAPABILITIES_2_EXTENSION_NAME, VULKAN_SUPPORTS_FULLSCREEN_EXCLUSIVE, VULKAN_EXTENSION_NOT_PROMOTED));
 }
 
 
-void FVulkanWindowsPlatform::GetDeviceExtensions(EGpuVendorId VendorId, TArray<const ANSICHAR*>& OutExtensions)
+void FVulkanWindowsPlatform::GetDeviceExtensions(FVulkanDevice* Device, FVulkanDeviceExtensionArray& OutExtensions)
 {
-	const bool bAllowVendorDevice = !FParse::Param(FCommandLine::Get(), TEXT("novendordevice"));
-
-#if VULKAN_SUPPORTS_DRIVER_PROPERTIES
-	OutExtensions.Add(VK_KHR_DRIVER_PROPERTIES_EXTENSION_NAME);
-#endif
-
-#if VULKAN_SUPPORTS_DEDICATED_ALLOCATION
-	OutExtensions.Add(VK_KHR_GET_MEMORY_REQUIREMENTS_2_EXTENSION_NAME);
-	OutExtensions.Add(VK_KHR_DEDICATED_ALLOCATION_EXTENSION_NAME);
-#endif
-	if (GGPUCrashDebuggingEnabled)
-	{
-#if VULKAN_SUPPORTS_AMD_BUFFER_MARKER
-		if (VendorId == EGpuVendorId::Amd && bAllowVendorDevice)
-		{
-			OutExtensions.Add(VK_AMD_BUFFER_MARKER_EXTENSION_NAME);
-		}
-#endif
-#if VULKAN_SUPPORTS_NV_DIAGNOSTICS
-		if (VendorId == EGpuVendorId::Nvidia && bAllowVendorDevice)
-		{
-			OutExtensions.Add(VK_NV_DEVICE_DIAGNOSTIC_CHECKPOINTS_EXTENSION_NAME);
-			OutExtensions.Add(VK_NV_DEVICE_DIAGNOSTICS_CONFIG_EXTENSION_NAME);
-		}
-#endif
-	}
+	OutExtensions.Add(MakeUnique<FVulkanDeviceExtension>(Device, VK_EXT_FULL_SCREEN_EXCLUSIVE_EXTENSION_NAME, VULKAN_SUPPORTS_FULLSCREEN_EXCLUSIVE, 
+														VULKAN_EXTENSION_NOT_PROMOTED, DEVICE_EXT_FLAG_SETTER(HasEXTFullscreenExclusive)));
 
-#if VULKAN_SUPPORTS_FULLSCREEN_EXCLUSIVE
-	// Fullscreen requires Instance capabilities2
-	OutExtensions.Add(VK_EXT_FULL_SCREEN_EXCLUSIVE_EXTENSION_NAME);
-#endif
-
-#if VULKAN_RHI_RAYTRACING
-	FVulkanRayTracingPlatform::GetDeviceExtensions(VendorId, OutExtensions);
-#endif
-
-#if VULKAN_SUPPORTS_FRAGMENT_DENSITY_MAP
-	OutExtensions.Add(VK_EXT_FRAGMENT_DENSITY_MAP_EXTENSION_NAME);
-#endif
-
-#if VULKAN_SUPPORTS_FRAGMENT_DENSITY_MAP2
-	OutExtensions.Add(VK_EXT_FRAGMENT_DENSITY_MAP_2_EXTENSION_NAME);
-#endif
-
-#if VULKAN_SUPPORTS_RENDERPASS2
-	OutExtensions.Add(VK_KHR_CREATE_RENDERPASS_2_EXTENSION_NAME);
-#endif
-
-#if VULKAN_SUPPORTS_FRAGMENT_SHADING_RATE
-	OutExtensions.Add(VK_KHR_FRAGMENT_SHADING_RATE_EXTENSION_NAME);
-#endif
+	// Manually activated extensions
+	OutExtensions.Add(MakeUnique<FVulkanDeviceExtension>(Device, VK_KHR_EXTERNAL_MEMORY_WIN32_EXTENSION_NAME, VULKAN_SUPPORTS_EXTERNAL_MEMORY, 
+														VULKAN_EXTENSION_NOT_PROMOTED, nullptr, FVulkanExtensionBase::ManuallyActivate));
 }
 
 void FVulkanWindowsPlatform::CreateSurface(void* WindowHandle, VkInstance Instance, VkSurfaceKHR* OutSurface)
@@ -398,9 +346,3 @@ void FVulkanWindowsPlatform::CheckDeviceDriver(uint32 DeviceIndex, EGpuVendorId
 	}
 }
 
-void FVulkanWindowsPlatform::EnablePhysicalDeviceFeatureExtensions(VkDeviceCreateInfo& DeviceInfo, FVulkanDevice& Device)
-{
-#if VULKAN_RHI_RAYTRACING
-	FVulkanRayTracingPlatform::EnablePhysicalDeviceFeatureExtensions(DeviceInfo, Device);
-#endif
-}
diff --git a/Engine/Source/Runtime/VulkanRHI/Private/Windows/VulkanWindowsPlatform.h b/Engine/Source/Runtime/VulkanRHI/Private/Windows/VulkanWindowsPlatform.h
--- a/Engine/Source/Runtime/VulkanRHI/Private/Windows/VulkanWindowsPlatform.h
+++ b/Engine/Source/Runtime/VulkanRHI/Private/Windows/VulkanWindowsPlatform.h
@@ -18,6 +18,7 @@
 
 #define VULKAN_RHI_RAYTRACING 						(RHI_RAYTRACING)
 #define VULKAN_SUPPORTS_SCALAR_BLOCK_LAYOUT			(VULKAN_RHI_RAYTRACING)
+#define VULKAN_SUPPORTS_MULTIVIEW					0   // TODO-XR: UE-147795
 
 #if VULKAN_RHI_RAYTRACING
 #	define UE_VK_API_VERSION						VK_API_VERSION_1_2
@@ -50,9 +51,7 @@
 	EnumMacro(PFN_vkCmdBeginRenderPass2KHR, vkCmdBeginRenderPass2KHR) \
 	EnumMacro(PFN_vkGetPhysicalDeviceFragmentShadingRatesKHR, vkGetPhysicalDeviceFragmentShadingRatesKHR)
 
-#define ENUM_VK_ENTRYPOINTS_OPTIONAL_PLATFORM_INSTANCE(EnumMacro) \
-	EnumMacro(PFN_vkCreateSamplerYcbcrConversionKHR, vkCreateSamplerYcbcrConversionKHR) \
-	EnumMacro(PFN_vkDestroySamplerYcbcrConversionKHR, vkDestroySamplerYcbcrConversionKHR)
+#define ENUM_VK_ENTRYPOINTS_OPTIONAL_PLATFORM_INSTANCE(EnumMacro)
 
 // and now, include the GenericPlatform class
 #include "../VulkanGenericPlatform.h"
@@ -66,18 +65,16 @@ public:
 	static bool LoadVulkanInstanceFunctions(VkInstance inInstance);
 	static void FreeVulkanLibrary();
 
-	static void GetInstanceExtensions(TArray<const ANSICHAR*>& OutExtensions);
+	static void GetInstanceExtensions(FVulkanInstanceExtensionArray& OutExtensions);
 	static void GetInstanceLayers(TArray<const ANSICHAR*>& OutLayers) {}
-	static void GetDeviceExtensions(EGpuVendorId VendorId, TArray<const ANSICHAR*>& OutExtensions);
-	static void GetDeviceLayers(EGpuVendorId VendorId, TArray<const ANSICHAR*>& OutLayers) {}
+	static void GetDeviceExtensions(FVulkanDevice* Device, FVulkanDeviceExtensionArray& OutExtensions);
+	static void GetDeviceLayers(TArray<const ANSICHAR*>& OutLayers) {}
 
 	static void CreateSurface(void* WindowHandle, VkInstance Instance, VkSurfaceKHR* OutSurface);
 
 	static bool SupportsDeviceLocalHostVisibleWithNoPenalty(EGpuVendorId VendorId);
 
 	static void WriteCrashMarker(const FOptionalVulkanDeviceExtensions& OptionalExtensions, VkCommandBuffer CmdBuffer, VkBuffer DestBuffer, const TArrayView<uint32>& Entries, bool bAdding);
-
-	static void EnablePhysicalDeviceFeatureExtensions(VkDeviceCreateInfo& DeviceInfo, FVulkanDevice& Device);
 };
 
 typedef FVulkanWindowsPlatform FVulkanPlatform;
diff --git a/Engine/Source/Runtime/VulkanRHI/Public/VulkanConfiguration.h b/Engine/Source/Runtime/VulkanRHI/Public/VulkanConfiguration.h
--- a/Engine/Source/Runtime/VulkanRHI/Public/VulkanConfiguration.h
+++ b/Engine/Source/Runtime/VulkanRHI/Public/VulkanConfiguration.h
@@ -339,13 +339,14 @@
 	#define VULKAN_SUPPORTS_MAINTENANCE4	0
 #endif
 
-#ifdef VK_EXT_descriptor_indexing
-	#define VULKAN_SUPPORTS_DESCRIPTOR_INDEXING	(VULKAN_SUPPORTS_PHYSICAL_DEVICE_PROPERTIES2)	// Requirement
-#else
-	#define VULKAN_SUPPORTS_DESCRIPTOR_INDEXING	0
+#ifndef VULKAN_SUPPORTS_DESCRIPTOR_INDEXING
+	#ifdef VK_EXT_descriptor_indexing
+		#define VULKAN_SUPPORTS_DESCRIPTOR_INDEXING	(VULKAN_SUPPORTS_PHYSICAL_DEVICE_PROPERTIES2)	// Requirement
+	#else
+		#define VULKAN_SUPPORTS_DESCRIPTOR_INDEXING	0
+	#endif
 #endif
 
-
 #ifndef VULKAN_OBJECT_TRACKING 
 #define VULKAN_OBJECT_TRACKING 0 //Track objects created and memory used. use r.vulkan.dumpmemory to dump to console
 #endif
diff --git a/Engine/Source/Runtime/VulkanRHI/Public/VulkanDynamicRHI.h b/Engine/Source/Runtime/VulkanRHI/Public/VulkanDynamicRHI.h
--- a/Engine/Source/Runtime/VulkanRHI/Public/VulkanDynamicRHI.h
+++ b/Engine/Source/Runtime/VulkanRHI/Public/VulkanDynamicRHI.h
@@ -34,8 +34,6 @@ struct FOptionalVulkanInstanceExtensions
 		static_assert(sizeof(Packed) == sizeof(FOptionalVulkanInstanceExtensions), "More bits needed!");
 		Packed = 0;
 	}
-
-	void Setup(const TArray<const ANSICHAR*>& InInstanceExtensions);
 };
 
 PRAGMA_DISABLE_DEPRECATION_WARNINGS
@@ -350,7 +348,7 @@ public:
 
 	inline bool SupportsDebugUtilsExt() const
 	{
-		return bSupportsDebugUtilsExt;
+		return (ActiveDebugLayerExtension == EActiveDebugLayerExtension::DebugUtilsExtension);
 	}
 
 	inline const FOptionalVulkanInstanceExtensions& GetOptionalExtensions() const
@@ -405,7 +403,7 @@ protected:
 	static void WriteEndFrameTimestamp(void*);
 	*/
 
-	static void GetInstanceLayersAndExtensions(TArray<const ANSICHAR*>& OutInstanceExtensions, TArray<const ANSICHAR*>& OutInstanceLayers, bool& bOutDebugUtils);
+	static TArray<const ANSICHAR*> SetupInstanceLayers(FVulkanInstanceExtensionArray& UEExtensions);
 
 	IConsoleObject* SavePipelineCacheCmd = nullptr;
 	IConsoleObject* RebuildPipelineCacheCmd = nullptr;
@@ -432,16 +430,26 @@ public:
 	static void TrimLRU();
 #endif
 
+	enum class EActiveDebugLayerExtension
+	{
+		None,
+		GfxReconstructLayer,
+		VkTraceLayer,
+		DebugUtilsExtension,
+		DebugReportExtension
+	};
+
 protected:
 	bool bIsStandaloneStereoDevice = false;
-	bool bSupportsDebugUtilsExt = false;
+
+	EActiveDebugLayerExtension ActiveDebugLayerExtension = EActiveDebugLayerExtension::None;
+
 #if VULKAN_HAS_DEBUGGING_ENABLED
 #if VULKAN_SUPPORTS_DEBUG_UTILS
 	VkDebugUtilsMessengerEXT Messenger = VK_NULL_HANDLE;
 #endif
-
-	bool bSupportsDebugCallbackExt = false;
 	VkDebugReportCallbackEXT MsgCallback = VK_NULL_HANDLE;
+
 	void SetupDebugLayerCallback();
 	void RemoveDebugLayerCallback();
 #endif
