VulkanRHI: Build fixes in Test and Shipping for 19628318

diff --git a/Engine/Source/Runtime/VulkanRHI/Private/VulkanExtensions.cpp b/Engine/Source/Runtime/VulkanRHI/Private/VulkanExtensions.cpp
--- a/Engine/Source/Runtime/VulkanRHI/Private/VulkanExtensions.cpp
+++ b/Engine/Source/Runtime/VulkanRHI/Private/VulkanExtensions.cpp
@@ -761,8 +761,10 @@ public:
 	FVulkanEXTDebugMarkerExtension(FVulkanDevice* InDevice)
 		: FVulkanDeviceExtension(InDevice, VK_EXT_DEBUG_MARKER_EXTENSION_NAME, (VULKAN_ENABLE_DRAW_MARKERS & VULKAN_HAS_DEBUGGING_ENABLED))
 	{
+#if VULKAN_HAS_DEBUGGING_ENABLED
 		const int32 VulkanValidationOption = GValidationCvar.GetValueOnAnyThread();
 		bEnabledInCode = bEnabledInCode && ((GRenderDocFound || VulkanValidationOption == 0) || FVulkanPlatform::ForceEnableDebugMarkers());
+#endif
 	}
 
 	virtual void PostPhysicalDeviceProperties() override final
@@ -1017,7 +1019,9 @@ public:
 
 	virtual void PreCreateInstance(VkInstanceCreateInfo& InstanceCreateInfo, FOptionalVulkanInstanceExtensions& ExtensionFlags) override final 
 	{
+#if VULKAN_HAS_DEBUGGING_ENABLED
 		check(GValidationCvar.GetValueOnAnyThread() > 0);
+#endif 
 
 		auto GetValidationFeaturesEnabled = []()
 		{
diff --git a/Engine/Source/Runtime/VulkanRHI/Public/VulkanDynamicRHI.h b/Engine/Source/Runtime/VulkanRHI/Public/VulkanDynamicRHI.h
--- a/Engine/Source/Runtime/VulkanRHI/Public/VulkanDynamicRHI.h
+++ b/Engine/Source/Runtime/VulkanRHI/Public/VulkanDynamicRHI.h
@@ -430,6 +430,7 @@ public:
 	static void TrimLRU();
 #endif
 
+public:
 	enum class EActiveDebugLayerExtension
 	{
 		None,
