VulkanRHI:  Fix issue with usage of VK_BUFFER_USAGE_SHADER_DEVICE_ADDRESS_BIT and VK_MEMORY_ALLOCATE_DEVICE_ADDRESS_BIT not lining up causing validation issues when raytracing is turned off.

diff --git a/Engine/Source/Runtime/VulkanRHI/Private/VulkanTransientResourceAllocator.cpp b/Engine/Source/Runtime/VulkanRHI/Private/VulkanTransientResourceAllocator.cpp
--- a/Engine/Source/Runtime/VulkanRHI/Private/VulkanTransientResourceAllocator.cpp
+++ b/Engine/Source/Runtime/VulkanRHI/Private/VulkanTransientResourceAllocator.cpp
@@ -9,14 +9,18 @@ FVulkanTransientHeap::FVulkanTransientHeap(const FInitializer& Initializer, FVul
 	, FDeviceChild(InDevice)
 	, VulkanBuffer(VK_NULL_HANDLE)
 {
-	const VkBufferUsageFlags BufferUsageFlags =
-#if VULKAN_RHI_RAYTRACING
-		VK_BUFFER_USAGE_ACCELERATION_STRUCTURE_STORAGE_BIT_KHR | VK_BUFFER_USAGE_SHADER_DEVICE_ADDRESS_BIT |
-#endif
+	VkBufferUsageFlags BufferUsageFlags =
 		VK_BUFFER_USAGE_VERTEX_BUFFER_BIT | VK_BUFFER_USAGE_INDEX_BUFFER_BIT | VK_BUFFER_USAGE_STORAGE_BUFFER_BIT |
 		VK_BUFFER_USAGE_STORAGE_TEXEL_BUFFER_BIT | VK_BUFFER_USAGE_UNIFORM_TEXEL_BUFFER_BIT |
 		VK_BUFFER_USAGE_INDIRECT_BUFFER_BIT | VK_BUFFER_USAGE_TRANSFER_SRC_BIT | VK_BUFFER_USAGE_TRANSFER_DST_BIT;
 
+#if VULKAN_RHI_RAYTRACING
+	if (InDevice->GetOptionalExtensions().HasRaytracingExtensions())
+	{
+		BufferUsageFlags |= VK_BUFFER_USAGE_ACCELERATION_STRUCTURE_STORAGE_BIT_KHR | VK_BUFFER_USAGE_SHADER_DEVICE_ADDRESS_BIT;
+	}
+#endif
+
 	// :TODO: VK_KHR_maintenance4...
 	{
 		VkBufferCreateInfo BufferCreateInfo;
