VulkanRHI: Build fixes in Test and Shipping for 19628318

diff --git a/Engine/Source/Runtime/VulkanRHI/Private/VulkanExtensions.cpp b/Engine/Source/Runtime/VulkanRHI/Private/VulkanExtensions.cpp
--- a/Engine/Source/Runtime/VulkanRHI/Private/VulkanExtensions.cpp
+++ b/Engine/Source/Runtime/VulkanRHI/Private/VulkanExtensions.cpp
@@ -43,7 +43,9 @@ TAutoConsoleVariable<int32> GVulkanRayTracingCVar(
 	ECVF_ReadOnly
 );
 
+#if VULKAN_HAS_DEBUGGING_ENABLED
 extern TAutoConsoleVariable<int32> GGPUValidationCvar;
+#endif
 
 TSharedPtr<IHeadMountedDisplayVulkanExtensions, ESPMode::ThreadSafe> FVulkanDynamicRHI::HMDVulkanExtensions;
 TArray<const ANSICHAR*> FVulkanDeviceExtension::ExternalExtensions;
@@ -1021,7 +1023,6 @@ public:
 	{
 #if VULKAN_HAS_DEBUGGING_ENABLED
 		check(GValidationCvar.GetValueOnAnyThread() > 0);
-#endif 
 
 		auto GetValidationFeaturesEnabled = []()
 		{
@@ -1058,6 +1059,7 @@ public:
 			ValidationFeatures.pEnabledValidationFeatures = ValidationFeaturesEnabled.GetData();
 			AddToPNext(InstanceCreateInfo, ValidationFeatures);
 		}
+#endif 
 	}
 
 private:
