Make Vulkan's LRU shader size calculation thread safe.
Added PSO cache accessor to enforce synchronization choice when using the PSO Cache.

diff --git a/Engine/Source/Runtime/VulkanRHI/Private/VulkanPipeline.cpp b/Engine/Source/Runtime/VulkanRHI/Private/VulkanPipeline.cpp
--- a/Engine/Source/Runtime/VulkanRHI/Private/VulkanPipeline.cpp
+++ b/Engine/Source/Runtime/VulkanRHI/Private/VulkanPipeline.cpp
@@ -310,7 +310,6 @@ FVulkanPipelineStateCacheManager::FVulkanPipelineStateCacheManager(FVulkanDevice
 	, bEvictImmediately(false)
 	, bLinkedToPSOFC(false)
 	, bLinkedToPSOFCSucessfulLoaded(false)
-	, PipelineCache(VK_NULL_HANDLE)
 {
 	bUseLRU = (int32)CVarEnableLRU.GetValueOnAnyThread() != 0;
 	LRUUsedPipelineMax = CVarLRUPipelineCapacity.GetValueOnAnyThread();
@@ -344,12 +343,17 @@ FVulkanPipelineStateCacheManager::~FVulkanPipelineStateCacheManager()
 		VulkanRHI::vkDestroyDescriptorSetLayout(Device->GetInstanceHandle(), Pair.Value.Handle, VULKAN_CPU_ALLOCATOR);
 	}
 
-	VulkanRHI::vkDestroyPipelineCache(Device->GetInstanceHandle(), PipelineCache, VULKAN_CPU_ALLOCATOR);
-	PipelineCache = VK_NULL_HANDLE;
+	{
+		FScopedPipelineCache PipelineCacheExclusive = PSOCache.Get(EPipelineCacheAccess::Exclusive);
+		VulkanRHI::vkDestroyPipelineCache(Device->GetInstanceHandle(), PipelineCacheExclusive.Get(), VULKAN_CPU_ALLOCATOR);
+		PipelineCacheExclusive.Get() = VK_NULL_HANDLE;
+	}
 }
 
 bool FVulkanPipelineStateCacheManager::Load(const TArray<FString>& CacheFilenames)
 {
+	FScopedPipelineCache PipelineCacheExclusive = PSOCache.Get(EPipelineCacheAccess::Exclusive);
+
 	bool bResult = false;
 	// Try to load device cache first
 	for (const FString& CacheFilename : CacheFilenames)
@@ -367,17 +371,17 @@ bool FVulkanPipelineStateCacheManager::Load(const TArray<FString>& CacheFilename
 				PipelineCacheInfo.initialDataSize = DeviceCache.Num();
 				PipelineCacheInfo.pInitialData = DeviceCache.GetData();
 
-				if (PipelineCache == VK_NULL_HANDLE)
+				if (PipelineCacheExclusive.Get() == VK_NULL_HANDLE)
 				{
 					// if we don't have one already, then create our main cache (PipelineCache)
-					VERIFYVULKANRESULT(VulkanRHI::vkCreatePipelineCache(Device->GetInstanceHandle(), &PipelineCacheInfo, VULKAN_CPU_ALLOCATOR, &PipelineCache));
+					VERIFYVULKANRESULT(VulkanRHI::vkCreatePipelineCache(Device->GetInstanceHandle(), &PipelineCacheInfo, VULKAN_CPU_ALLOCATOR, &PipelineCacheExclusive.Get()));
 				}
 				else
 				{
 					// if we have one already, create a temp one and merge into the main cache
 					VkPipelineCache TempPipelineCache;
 					VERIFYVULKANRESULT(VulkanRHI::vkCreatePipelineCache(Device->GetInstanceHandle(), &PipelineCacheInfo, VULKAN_CPU_ALLOCATOR, &TempPipelineCache));
-					VERIFYVULKANRESULT(VulkanRHI::vkMergePipelineCaches(Device->GetInstanceHandle(), PipelineCache, 1, &TempPipelineCache));
+					VERIFYVULKANRESULT(VulkanRHI::vkMergePipelineCaches(Device->GetInstanceHandle(), PipelineCacheExclusive.Get(), 1, &TempPipelineCache));
 					VulkanRHI::vkDestroyPipelineCache(Device->GetInstanceHandle(), TempPipelineCache, VULKAN_CPU_ALLOCATOR);
 				}
 
@@ -433,11 +437,11 @@ bool FVulkanPipelineStateCacheManager::Load(const TArray<FString>& CacheFilename
 	}
 
 	// Lazily create the cache in case the load failed
-	if (PipelineCache == VK_NULL_HANDLE)
+	if (PipelineCacheExclusive.Get() == VK_NULL_HANDLE)
 	{
 		VkPipelineCacheCreateInfo PipelineCacheInfo;
 		ZeroVulkanStruct(PipelineCacheInfo, VK_STRUCTURE_TYPE_PIPELINE_CACHE_CREATE_INFO);
-		VERIFYVULKANRESULT(VulkanRHI::vkCreatePipelineCache(Device->GetInstanceHandle(), &PipelineCacheInfo, VULKAN_CPU_ALLOCATOR, &PipelineCache));
+		VERIFYVULKANRESULT(VulkanRHI::vkCreatePipelineCache(Device->GetInstanceHandle(), &PipelineCacheInfo, VULKAN_CPU_ALLOCATOR, &PipelineCacheExclusive.Get()));
 	}
 
 	return bResult;
@@ -445,6 +449,8 @@ bool FVulkanPipelineStateCacheManager::Load(const TArray<FString>& CacheFilename
 
 void FVulkanPipelineStateCacheManager::InitAndLoad(const TArray<FString>& CacheFilenames)
 {
+	FScopedPipelineCache PipelineCacheExclusive = PSOCache.Get(EPipelineCacheAccess::Exclusive);
+
 	if (GEnablePipelineCacheLoadCvar.GetValueOnAnyThread() == 0)
 	{
 		UE_LOG(LogVulkanRHI, Display, TEXT("Not loading pipeline cache per r.Vulkan.PipelineCacheLoad=0"));
@@ -484,11 +490,11 @@ void FVulkanPipelineStateCacheManager::InitAndLoad(const TArray<FString>& CacheF
 	}
 
 	// Lazily create the cache in case the load failed
-	if (PipelineCache == VK_NULL_HANDLE)
+	if (PipelineCacheExclusive.Get() == VK_NULL_HANDLE)
 	{
 		VkPipelineCacheCreateInfo PipelineCacheInfo;
 		ZeroVulkanStruct(PipelineCacheInfo, VK_STRUCTURE_TYPE_PIPELINE_CACHE_CREATE_INFO);
-		VERIFYVULKANRESULT(VulkanRHI::vkCreatePipelineCache(Device->GetInstanceHandle(), &PipelineCacheInfo, VULKAN_CPU_ALLOCATOR, &PipelineCache));
+		VERIFYVULKANRESULT(VulkanRHI::vkCreatePipelineCache(Device->GetInstanceHandle(), &PipelineCacheInfo, VULKAN_CPU_ALLOCATOR, &PipelineCacheExclusive.Get()));
 	}
 }
 
@@ -567,6 +573,7 @@ void FVulkanPipelineStateCacheManager::OnShaderPipelineCachePrecompilationComple
 
 void FVulkanPipelineStateCacheManager::Save(const FString& CacheFilename, bool bFromPSOFC)
 {
+	FScopedPipelineCache PipelineCacheExclusive = PSOCache.Get(EPipelineCacheAccess::Exclusive);
 	if (bLinkedToPSOFC && !bFromPSOFC)
 	{
 		UE_LOG(LogVulkanRHI, Log, TEXT("FVulkanPipelineStateCacheManager: skipped saving because we only save if the PSOFC based one failed to load."));
@@ -579,13 +586,13 @@ void FVulkanPipelineStateCacheManager::Save(const FString& CacheFilename, bool b
 
 	// First save Device Cache
 	size_t Size = 0;
-	VERIFYVULKANRESULT(VulkanRHI::vkGetPipelineCacheData(Device->GetInstanceHandle(), PipelineCache, &Size, nullptr));
+	VERIFYVULKANRESULT(VulkanRHI::vkGetPipelineCacheData(Device->GetInstanceHandle(), PipelineCacheExclusive.Get(), &Size, nullptr));
 	// 16 is HeaderSize + HeaderVersion
 	if (Size >= 16 + VK_UUID_SIZE)
 	{
 		TArray<uint8> DeviceCache;
 		DeviceCache.AddUninitialized(Size);
-		VkResult Result = VulkanRHI::vkGetPipelineCacheData(Device->GetInstanceHandle(), PipelineCache, &Size, DeviceCache.GetData());
+		VkResult Result = VulkanRHI::vkGetPipelineCacheData(Device->GetInstanceHandle(), PipelineCacheExclusive.Get(), &Size, DeviceCache.GetData());
 		if (Result == VK_SUCCESS)
 		{
 			FString BinaryCacheFilename = FVulkanPlatform::CreatePSOBinaryCacheFilename(Device, CacheFilename);
@@ -603,10 +610,10 @@ void FVulkanPipelineStateCacheManager::Save(const FString& CacheFilename, bool b
 		{
 			UE_LOG(LogVulkanRHI, Warning, TEXT("Failed to get Vulkan pipeline cache data. Error %d, %d bytes"), Result, Size);
 
-			VulkanRHI::vkDestroyPipelineCache(Device->GetInstanceHandle(), PipelineCache, VULKAN_CPU_ALLOCATOR);
+			VulkanRHI::vkDestroyPipelineCache(Device->GetInstanceHandle(), PipelineCacheExclusive.Get(), VULKAN_CPU_ALLOCATOR);
 			VkPipelineCacheCreateInfo PipelineCacheInfo;
 			ZeroVulkanStruct(PipelineCacheInfo, VK_STRUCTURE_TYPE_PIPELINE_CACHE_CREATE_INFO);
-			VERIFYVULKANRESULT(VulkanRHI::vkCreatePipelineCache(Device->GetInstanceHandle(), &PipelineCacheInfo, VULKAN_CPU_ALLOCATOR, &PipelineCache));
+			VERIFYVULKANRESULT(VulkanRHI::vkCreatePipelineCache(Device->GetInstanceHandle(), &PipelineCacheInfo, VULKAN_CPU_ALLOCATOR, &PipelineCacheExclusive.Get()));
 		}
 		else
 		{
@@ -1261,41 +1268,65 @@ bool FVulkanPipelineStateCacheManager::CreateGfxPipelineFromEntry(FVulkanRHIGrap
 			FScopeLock Lock(&LRUCS); 
 			Found = LRU2SizeList.Find(ShaderHash);
 		}
-		size_t PreSize = 0, AfterSize = 0;
+
 		uint32 FoundSize = 0;
+		VkPipelineCache LocalPipelineCache = VK_NULL_HANDLE;
 		if (Found)
 		{
+			SCOPE_CYCLE_COUNTER(STAT_VulkanPSOVulkanCreationTime);
 			FoundSize = Found->PipelineSize;
+			FScopedPipelineCache PipelineCacheShared = PSOCache.Get(EPipelineCacheAccess::Shared);
+			Result = VulkanRHI::vkCreateGraphicsPipelines(Device->GetInstanceHandle(), PipelineCacheShared.Get(), 1, &PipelineInfo, VULKAN_CPU_ALLOCATOR, Pipeline);
 		}
 		else
-		{
-			VulkanRHI::vkGetPipelineCacheData(Device->GetInstanceHandle(), PipelineCache, &PreSize, nullptr);
-		}
-
-
 		{
 			SCOPE_CYCLE_COUNTER(STAT_VulkanPSOVulkanCreationTime);
-			Result = VulkanRHI::vkCreateGraphicsPipelines(Device->GetInstanceHandle(), PipelineCache, 1, &PipelineInfo, VULKAN_CPU_ALLOCATOR, Pipeline);
-		}
+			// We create a single pipeline cache for this create so we can observe the size for LRU cache's accounting.
+			// measuring deltas from the global PipelineCache is not thread safe.
+			VkPipelineCacheCreateInfo PipelineCacheInfo;
+			ZeroVulkanStruct(PipelineCacheInfo, VK_STRUCTURE_TYPE_PIPELINE_CACHE_CREATE_INFO);
+			VERIFYVULKANRESULT(VulkanRHI::vkCreatePipelineCache(Device->GetInstanceHandle(), &PipelineCacheInfo, VULKAN_CPU_ALLOCATOR, &LocalPipelineCache));
+			Result = VulkanRHI::vkCreateGraphicsPipelines(Device->GetInstanceHandle(), LocalPipelineCache, 1, &PipelineInfo, VULKAN_CPU_ALLOCATOR, Pipeline);
 
-		if (!Found && Result == VK_SUCCESS)
-		{
-			VulkanRHI::vkGetPipelineCacheData(Device->GetInstanceHandle(), PipelineCache, &AfterSize, nullptr);
-			uint32 Diff = AfterSize - PreSize;
-			if (!Diff)
+			if (Result == VK_SUCCESS)
 			{
-				UE_LOG(LogVulkanRHI, Warning, TEXT("Shader size was computed as zero, using 20k instead."));
-				Diff = 20 * 1024;
+				FScopeLock Lock(&LRUCS);
+				Found = LRU2SizeList.Find(ShaderHash);
+				if (Found == nullptr) // we may not win..
+				{
+					QUICK_SCOPE_CYCLE_COUNTER(STAT_Vulkan_Calc_LRU_Size);
+					size_t Diff = 0;
+					VulkanRHI::vkGetPipelineCacheData(Device->GetInstanceHandle(), LocalPipelineCache, &Diff, nullptr);
+
+					if (!Diff)
+					{
+						UE_LOG(LogVulkanRHI, Warning, TEXT("Shader size was computed as zero, using 20k instead."));
+						Diff = 20 * 1024;
+					}
+					FVulkanPipelineSize PipelineSize;
+					PipelineSize.ShaderHash = ShaderHash;
+					PipelineSize.PipelineSize = (uint32)Diff;
+					LRU2SizeList.Add(ShaderHash, PipelineSize);
+					FoundSize = Diff;
+
+					{
+						QUICK_SCOPE_CYCLE_COUNTER(STAT_VulkanPSOCacheMerge);
+						FScopedPipelineCache PipelineCacheExclusive = PSOCache.Get(EPipelineCacheAccess::Exclusive);
+						VERIFYVULKANRESULT(VulkanRHI::vkMergePipelineCaches(Device->GetInstanceHandle(), PipelineCacheExclusive.Get(), 1, &LocalPipelineCache));
+					}
+				}
+				else
+				{
+					FoundSize = Found->PipelineSize;;
+				}
 			}
-			FVulkanPipelineSize PipelineSize;
-			PipelineSize.ShaderHash = ShaderHash;
-			PipelineSize.PipelineSize = Diff;
+			else
 			{
-				FScopeLock Lock(&LRUCS);
-				LRU2SizeList.Add(ShaderHash, PipelineSize);
+				UE_LOG(LogVulkanRHI, Warning, TEXT("LRU: PSO create failed."));
 			}
-			FoundSize = Diff;
+			VulkanRHI::vkDestroyPipelineCache(Device->GetInstanceHandle(), LocalPipelineCache, VULKAN_CPU_ALLOCATOR);
 		}
+
 		if(Result == VK_SUCCESS)
 		{
 			PSO->PipelineCacheSize = FoundSize;
@@ -1304,7 +1335,8 @@ bool FVulkanPipelineStateCacheManager::CreateGfxPipelineFromEntry(FVulkanRHIGrap
 	else
 	{
 		SCOPE_CYCLE_COUNTER(STAT_VulkanPSOVulkanCreationTime);
-		Result = VulkanRHI::vkCreateGraphicsPipelines(Device->GetInstanceHandle(), PipelineCache, 1, &PipelineInfo, VULKAN_CPU_ALLOCATOR, Pipeline);
+		FScopedPipelineCache PipelineCacheShared = PSOCache.Get(EPipelineCacheAccess::Shared);
+		Result = VulkanRHI::vkCreateGraphicsPipelines(Device->GetInstanceHandle(), PipelineCacheShared.Get(), 1, &PipelineInfo, VULKAN_CPU_ALLOCATOR, Pipeline);
 	}
 
 	if (Result != VK_SUCCESS)
@@ -1977,8 +2009,12 @@ FVulkanComputePipeline* FVulkanPipelineStateCacheManager::CreateComputePipelineF
 	Shader->GetEntryPoint(EntryPoint, 24);
 	PipelineInfo.stage.pName = EntryPoint;
 	PipelineInfo.layout = ComputeLayout->GetPipelineLayout();
-		
-	VkResult Result = VulkanRHI::vkCreateComputePipelines(Device->GetInstanceHandle(), PipelineCache, 1, &PipelineInfo, VULKAN_CPU_ALLOCATOR, &Pipeline->Pipeline);
+	VkResult Result;
+	{
+		QUICK_SCOPE_CYCLE_COUNTER(STAT_VulkanComputePSOCreate);
+		FScopedPipelineCache PipelineCacheShared = PSOCache.Get(EPipelineCacheAccess::Shared);
+		Result = VulkanRHI::vkCreateComputePipelines(Device->GetInstanceHandle(), PipelineCacheShared.Get(), 1, &PipelineInfo, VULKAN_CPU_ALLOCATOR, &Pipeline->Pipeline);
+	}
 
 	if (Result != VK_SUCCESS)
 	{
diff --git a/Engine/Source/Runtime/VulkanRHI/Private/VulkanPipeline.h b/Engine/Source/Runtime/VulkanRHI/Private/VulkanPipeline.h
--- a/Engine/Source/Runtime/VulkanRHI/Private/VulkanPipeline.h
+++ b/Engine/Source/Runtime/VulkanRHI/Private/VulkanPipeline.h
@@ -520,7 +520,33 @@ private:
 	FRWLock ComputePipelineLock;
 	TMap<uint64, FVulkanComputePipeline*> ComputePipelineEntries;
 
-	VkPipelineCache PipelineCache;
+	template<typename TType>
+	class FScopedRWAccessor
+	{
+		bool bWriteAccess;
+		TType& ProtectedObj;
+		FRWLock& RWLock;
+	public:
+		FScopedRWAccessor(bool bWriteAccessIn, TType& ProtectedObjIn, FRWLock& RWLockIn) : bWriteAccess(bWriteAccessIn), ProtectedObj(ProtectedObjIn), RWLock(RWLockIn) { bWriteAccess ? RWLock.WriteLock() : RWLock.ReadLock(); }
+		~FScopedRWAccessor() { bWriteAccess ? RWLock.WriteUnlock() : RWLock.ReadUnlock(); }
+		TType& Get() { return ProtectedObj; }
+	};
+
+	using FScopedPipelineCache = FScopedRWAccessor<VkPipelineCache>;
+
+	enum class EPipelineCacheAccess : uint8
+	{
+		Shared,			// 'read' access, or for use when the API does its own synchronization.
+		Exclusive		// 'write' access, excludes all other usage for the duration. 
+	};
+	class FPipelineCache
+	{
+		VkPipelineCache PipelineCache = VK_NULL_HANDLE;
+		FRWLock PipelineCacheLock;
+	public:
+		FScopedPipelineCache Get(EPipelineCacheAccess PipelineAccessType) { return FScopedPipelineCache(PipelineAccessType == EPipelineCacheAccess::Exclusive, PipelineCache, PipelineCacheLock); }
+	};
+	FPipelineCache PSOCache;
 
 	FCriticalSection LayoutMapCS;
 	TMap<FVulkanDescriptorSetsLayoutInfo, FVulkanLayout*> LayoutMap;
