VulkanRHI: Use vkGetDeviceImageMemoryRequirementsKHR from VK_KHR_maintenance4 to get MemReq for images in a single API call (no create/destroy image).
Start converting log messages to use TYPE_TO_STRING for easier debugging.
SET_DEBUG_NAME for buffers.

diff --git a/Engine/Source/Runtime/VulkanRHI/Private/VulkanBarriers.cpp b/Engine/Source/Runtime/VulkanRHI/Private/VulkanBarriers.cpp
--- a/Engine/Source/Runtime/VulkanRHI/Private/VulkanBarriers.cpp
+++ b/Engine/Source/Runtime/VulkanRHI/Private/VulkanBarriers.cpp
@@ -1727,8 +1727,8 @@ void FVulkanLayoutManager::ValidateRenderPassColorEntry(const FRHIRenderPassInfo
 	// Check that the image is in the correct layout for rendering.
 	const VkImageLayout ExpectedVkLayout = VK_IMAGE_LAYOUT_COLOR_ATTACHMENT_OPTIMAL;
 	FVulkanImageLayout& LayoutInfo = GetOrAddFullLayout(Surface, VK_IMAGE_LAYOUT_UNDEFINED);
-	VkImageLayout CurrentVkLayout = LayoutInfo.GetSubresLayout(ColorEntry.ArraySlice, ColorEntry.MipIndex);
-	if (ensureMsgf(CurrentVkLayout == ExpectedVkLayout, TEXT("%s target is in layout %u, expected %u. Please add a transition before starting the render pass."), bResolveTarget ? TEXT("Resolve") : TEXT("Color"), CurrentVkLayout, ExpectedVkLayout))
+	const VkImageLayout CurrentVkLayout = LayoutInfo.GetSubresLayout(ColorEntry.ArraySlice, ColorEntry.MipIndex);
+	if (ensureMsgf(CurrentVkLayout == ExpectedVkLayout, TEXT("%s target is in layout %s, expected %s. Please add a transition before starting the render pass."), bResolveTarget ? TEXT("Resolve") : TEXT("Color"), VK_TYPE_TO_STRING(VkImageLayout, CurrentVkLayout), VK_TYPE_TO_STRING(VkImageLayout, ExpectedVkLayout)))
 	{
 		return;
 	}
diff --git a/Engine/Source/Runtime/VulkanRHI/Private/VulkanBuffer.cpp b/Engine/Source/Runtime/VulkanRHI/Private/VulkanBuffer.cpp
--- a/Engine/Source/Runtime/VulkanRHI/Private/VulkanBuffer.cpp
+++ b/Engine/Source/Runtime/VulkanRHI/Private/VulkanBuffer.cpp
@@ -171,6 +171,8 @@ FVulkanResourceMultiBuffer::FVulkanResourceMultiBuffer(FVulkanDevice* InDevice,
 					Buffers[Index].Size = InSize;
 					check((Buffers[Index].Offset + InSize) <= InTransientHeapAllocation->Size);
 				}
+				
+				VULKAN_SET_DEBUG_NAME((*InDevice), VK_OBJECT_TYPE_BUFFER, Buffers[0].VulkanHandle, TEXT("%s:(FVulkanResourceMultiBuffer*)0x%p:Buffer[1,2,3]"), CreateInfo.DebugName ? CreateInfo.DebugName : TEXT("?"), this);
 			}
 			else
 			{
@@ -186,6 +188,7 @@ FVulkanResourceMultiBuffer::FVulkanResourceMultiBuffer(FVulkanDevice* InDevice,
 					{
 						InDevice->GetMemoryManager().HandleOOM();
 					}
+					VULKAN_SET_DEBUG_NAME((*InDevice), VK_OBJECT_TYPE_BUFFER, Buffers[Index].VulkanHandle, TEXT("%s:(FVulkanResourceMultiBuffer*)0x%p:Buffer[%d]"), CreateInfo.DebugName ? CreateInfo.DebugName : TEXT("?"), this, Index);
 				}
 			}
 
diff --git a/Engine/Source/Runtime/VulkanRHI/Private/VulkanMemory.cpp b/Engine/Source/Runtime/VulkanRHI/Private/VulkanMemory.cpp
--- a/Engine/Source/Runtime/VulkanRHI/Private/VulkanMemory.cpp
+++ b/Engine/Source/Runtime/VulkanRHI/Private/VulkanMemory.cpp
@@ -683,56 +683,6 @@ namespace VulkanRHI
 		PrintMemInfo();
 	}
 
-	static FString GetMemoryPropertyFlagsString(VkMemoryPropertyFlags Flags)
-	{
-		FString String;
-		if (VKHasAllFlags(Flags, VK_MEMORY_PROPERTY_DEVICE_LOCAL_BIT))
-		{
-			String += TEXT(" DeviceLocal");
-		}
-		else
-		{
-			String += TEXT("      ");
-		}
-
-		if (VKHasAllFlags(Flags, VK_MEMORY_PROPERTY_HOST_VISIBLE_BIT))
-		{
-			String += TEXT(" HostVisible");
-		}
-		else
-		{
-			String += TEXT("            ");
-		}
-
-		if (VKHasAllFlags(Flags, VK_MEMORY_PROPERTY_HOST_COHERENT_BIT))
-		{
-			String += TEXT(" HostCoherent");
-		}
-		else
-		{
-			String += TEXT("             ");
-		}
-
-		if (VKHasAllFlags(Flags, VK_MEMORY_PROPERTY_HOST_CACHED_BIT))
-		{
-			String += TEXT(" HostCached");
-		}
-		else
-		{
-			String += TEXT("           ");
-		}
-
-		if (VKHasAllFlags(Flags, VK_MEMORY_PROPERTY_LAZILY_ALLOCATED_BIT))
-		{
-			String += TEXT(" Lazy");
-		}
-		else
-		{
-			String += TEXT("     ");
-		}
-		return String;
-	};
-
 	bool MetaTypeCanEvict(EVulkanAllocationMetaType MetaType)
 	{
 		switch(MetaType)
@@ -750,7 +700,7 @@ namespace VulkanRHI
 		auto ToPct = [](uint64 Used, uint64 Total) { return 100.f * (float)Used / (float)Total; };
 
 		const uint32 MaxAllocations = Device->GetLimits().maxMemoryAllocationCount;
-		VULKAN_LOGMEMORY(TEXT("Max memory allocations %d."), MaxAllocations);
+		VULKAN_LOGMEMORY(TEXT("Max memory allocations %u."), MaxAllocations);
 
 		VULKAN_LOGMEMORY(TEXT("%d Device Memory Heaps:"), MemoryProperties.memoryHeapCount);
 		for (uint32 HeapIndex = 0; HeapIndex < MemoryProperties.memoryHeapCount; ++HeapIndex)
@@ -776,11 +726,10 @@ namespace VulkanRHI
 			{
 				if(HeapIndex == MemoryProperties.memoryTypes[TypeIndex].heapIndex)
 				{
+					const VkMemoryPropertyFlags MemoryPropertyFlags = MemoryProperties.memoryTypes[TypeIndex].propertyFlags;
 					VULKAN_LOGMEMORY(TEXT(" %2d: Flags 0x%05x - Heap %2d - %s"),
-						TypeIndex,
-						MemoryProperties.memoryTypes[TypeIndex].propertyFlags,
-						MemoryProperties.memoryTypes[TypeIndex].heapIndex,
-						*GetMemoryPropertyFlagsString(MemoryProperties.memoryTypes[TypeIndex].propertyFlags));
+						TypeIndex, MemoryPropertyFlags, HeapIndex,
+						MemoryPropertyFlags ? VK_FLAGS_TO_STRING(VkMemoryPropertyFlags, MemoryPropertyFlags) : TEXT(""));
 				}
 			}
 		}
@@ -2763,7 +2712,7 @@ namespace VulkanRHI
 		}
 		else
 		{
-			checkf(false, TEXT("Unknown buffer alignment for VkBufferUsageFlags combination: 0x%x"), (uint32)BufferUsageFlags);
+			checkf(false, TEXT("Unknown buffer alignment for VkBufferUsageFlags combination: 0x%x"), VK_FLAGS_TO_STRING(VkBufferUsageFlags, BufferUsageFlags));
 		}
 
 		return Alignment;
@@ -2800,7 +2749,7 @@ namespace VulkanRHI
 		}
 		else
 		{
-			checkf(false, TEXT("Unknown priority for VkBufferUsageFlags combination: 0x%x"), (uint32)BufferUsageFlags);
+			checkf(false, TEXT("Unknown priority for VkBufferUsageFlags combination: 0x%x"), VK_FLAGS_TO_STRING(VkBufferUsageFlags, BufferUsageFlags));
 		}
 
 		return Priority;
@@ -3345,7 +3294,7 @@ namespace VulkanRHI
 		auto WriteLogLine = [](const FString& Name, FResourceHeapStats& Stat)
 		{
 			uint64 FreeMemory = Stat.TotalMemory - Stat.UsedBufferMemory - Stat.UsedImageMemory;
-			FString HostString = GetMemoryPropertyFlagsString(Stat.MemoryFlags);
+			FString HostString = VK_FLAGS_TO_STRING(VkMemoryPropertyFlags, Stat.MemoryFlags);
 			VULKAN_LOGMEMORY(TEXT("\t\t%-33s  |%8.2fmb / %8.2fmb / %11.2fmb / %11.2fmb | %10d %10d | %6d %6d %6d | %05x | %s"),
 				*Name,
 				Stat.UsedBufferMemory / (1024.f * 1024.f),
@@ -3492,7 +3441,7 @@ namespace VulkanRHI
 				{
 					Flags = MemoryProperties.memoryTypes[Allocator->MemoryTypeIndex].propertyFlags;
 				}
-				FString MemoryString = GetMemoryPropertyFlagsString(Flags);
+				FString MemoryString = VK_FLAGS_TO_STRING(VkMemoryPropertyFlags, Flags);
 				FString NameId = FString::Printf(TEXT("%s [%4d]"), *Name, Allocator->AllocatorIndex);
 				WriteLogLineSubAllocator(NameId, MemoryString, *Allocator);
 			}
@@ -3507,7 +3456,7 @@ namespace VulkanRHI
 				{
 					Flags = MemoryProperties.memoryTypes[Allocator->MemoryTypeIndex].propertyFlags;
 				}
-				FString MemoryString = GetMemoryPropertyFlagsString(Flags);
+				FString MemoryString = VK_FLAGS_TO_STRING(VkMemoryPropertyFlags, Flags);
 				VULKAN_LOGMEMORY(VULKAN_LOGMEMORY_PAD2);
 				VULKAN_LOGMEMORY(TEXT("\t\t%-45s  | %4s %8s | %4s / %10s / %10s / %10s / %10s | %10s / %10s / %10s / %10s | %10s / %10s | %10s / %10s / %10s | Mapped/Evictable |"),
 					TEXT(""),
diff --git a/Engine/Source/Runtime/VulkanRHI/Private/VulkanSwapChain.cpp b/Engine/Source/Runtime/VulkanRHI/Private/VulkanSwapChain.cpp
--- a/Engine/Source/Runtime/VulkanRHI/Private/VulkanSwapChain.cpp
+++ b/Engine/Source/Runtime/VulkanRHI/Private/VulkanSwapChain.cpp
@@ -260,7 +260,7 @@ FVulkanSwapChain::FVulkanSwapChain(VkInstance InInstance, FVulkanDevice& InDevic
 					{
 						InOutPixelFormat = (EPixelFormat)PFIndex;
 						CurrFormat = Formats[Index];
-						UE_LOG(LogVulkanRHI, Verbose, TEXT("No swapchain format requested, picking up VulkanFormat %d"), (uint32)CurrFormat.format);
+						UE_LOG(LogVulkanRHI, Verbose, TEXT("No swapchain format requested, picking up VulkanFormat %s"), VK_TYPE_TO_STRING(VkFormat, CurrFormat.format));
 						break;
 					}
 				}
@@ -425,12 +425,12 @@ FVulkanSwapChain::FVulkanSwapChain(VkInstance InInstance, FVulkanDevice& InDevic
 			}
 			else
 			{
-				UE_LOG(LogVulkanRHI, Warning, TEXT("Couldn't find desired PresentMode! Using %d"), static_cast<int32>(FoundPresentModes[0]));
+				UE_LOG(LogVulkanRHI, Warning, TEXT("Couldn't find desired PresentMode! Using %s"), VK_TYPE_TO_STRING(VkPresentModeKHR, FoundPresentModes[0]));
 				PresentMode = FoundPresentModes[0];
 			}
 		}
 
-		UE_CLOG(bFirstTimeLog, LogVulkanRHI, Display, TEXT("Selected VkPresentModeKHR mode %d"), PresentMode);
+		UE_CLOG(bFirstTimeLog, LogVulkanRHI, Display, TEXT("Selected VkPresentModeKHR mode %s"), VK_TYPE_TO_STRING(VkPresentModeKHR, PresentMode));
 		bFirstTimeLog = false;
 	}
 
@@ -517,8 +517,9 @@ FVulkanSwapChain::FVulkanSwapChain(VkInstance InInstance, FVulkanDevice& InDevic
 	static bool bPrintSwapchainCreationInfo = true;
 	if (bPrintSwapchainCreationInfo)
 	{
-		UE_LOG(LogVulkanRHI, Log, TEXT("Creating new VK swapchain with present mode %d, format %d, color space %d, num images %d"), 
-			static_cast<uint32>(SwapChainInfo.presentMode), static_cast<uint32>(SwapChainInfo.imageFormat), static_cast<uint32>(SwapChainInfo.imageColorSpace), static_cast<uint32>(SwapChainInfo.minImageCount));
+		UE_LOG(LogVulkanRHI, Log, TEXT("Creating new VK swapchain with %s, %s, %s, num images %d"), 
+			VK_TYPE_TO_STRING(VkPresentModeKHR, SwapChainInfo.presentMode), VK_TYPE_TO_STRING(VkFormat, SwapChainInfo.imageFormat), 
+			VK_TYPE_TO_STRING(VkColorSpaceKHR, SwapChainInfo.imageColorSpace), static_cast<uint32>(SwapChainInfo.minImageCount));
 #if WITH_EDITOR
 		bPrintSwapchainCreationInfo = false;
 #endif
diff --git a/Engine/Source/Runtime/VulkanRHI/Private/VulkanTexture.cpp b/Engine/Source/Runtime/VulkanRHI/Private/VulkanTexture.cpp
--- a/Engine/Source/Runtime/VulkanRHI/Private/VulkanTexture.cpp
+++ b/Engine/Source/Runtime/VulkanRHI/Private/VulkanTexture.cpp
@@ -192,7 +192,7 @@ inline void FVulkanTexture::InternalLockWrite(FVulkanCommandListContext& Context
 
 void FVulkanTexture::ErrorInvalidViewType() const
 {
-	UE_LOG(LogVulkanRHI, Error, TEXT("Invalid ViewType %d"), (uint32)GetViewType());
+	UE_LOG(LogVulkanRHI, Error, TEXT("Invalid ViewType %s"), VK_TYPE_TO_STRING(VkImageViewType, GetViewType()));
 }
 
 
@@ -582,7 +582,7 @@ void FVulkanTexture::InternalMoveSurface(FVulkanDevice& InDevice, FVulkanCommand
 
 	VkImage MovedImage;
 	VERIFYVULKANRESULT(VulkanRHI::vkCreateImage(InDevice.GetInstanceHandle(), &ImageCreateInfo.ImageCreateInfo, VULKAN_CPU_ALLOCATOR, &MovedImage));
-	checkf(Tiling == ImageCreateInfo.ImageCreateInfo.tiling, TEXT("Move has changed image tiling:  before [%d] != after [%d]"), (int32)Tiling, (int32)ImageCreateInfo.ImageCreateInfo.tiling);
+	checkf(Tiling == ImageCreateInfo.ImageCreateInfo.tiling, TEXT("Move has changed image tiling:  before [%s] != after [%s]"), VK_TYPE_TO_STRING(VkImageTiling, Tiling), VK_TYPE_TO_STRING(VkImageTiling, ImageCreateInfo.ImageCreateInfo.tiling));
 
 	const ETextureCreateFlags UEFlags = GetDesc().Flags;
 	const bool bRenderTarget = EnumHasAnyFlags(UEFlags, TexCreate_RenderTargetable | TexCreate_DepthStencilTargetable | TexCreate_ResolveTargetable);
@@ -2013,7 +2013,7 @@ void FVulkanTexture::Move(FVulkanDevice& InDevice, FVulkanCommandListContext& Co
 		const bool bRenderTarget = EnumHasAnyFlags(UEFlags, TexCreate_RenderTargetable | TexCreate_DepthStencilTargetable | TexCreate_ResolveTargetable);
 		const bool bUAV = EnumHasAnyFlags(UEFlags, TexCreate_UAV);
 		checkf(bRenderTarget || bUAV, TEXT("Surface must be a RenderTarget or a UAV in order to be moved.  UEFlags=0x%x"), (int32)UEFlags);
-		checkf(Tiling == VK_IMAGE_TILING_OPTIMAL, TEXT("Tiling [%d] is not supported for move, only VK_IMAGE_TILING_OPTIMAL"), (int32)Tiling);
+		checkf(Tiling == VK_IMAGE_TILING_OPTIMAL, TEXT("Tiling [%s] is not supported for move, only VK_IMAGE_TILING_OPTIMAL"), VK_TYPE_TO_STRING(VkImageTiling, Tiling));
 	
 		InternalMoveSurface(InDevice, Context, NewAllocation);
 	
@@ -2063,9 +2063,9 @@ void FVulkanTexture::Evict(FVulkanDevice& InDevice)
 
 	{
 		check(0 == CpuReadbackBuffer);
-		checkf(MemProps == VK_MEMORY_PROPERTY_DEVICE_LOCAL_BIT, TEXT("Can't evict surface that isn't device local.  MemoryProperties=%d"), (int32)MemProps);
-		checkf(VulkanRHI::GetAspectMaskFromUEFormat(GetDesc().Format, true, true) == FullAspectMask, TEXT("FullAspectMask (%d) does not match with PixelFormat (%d)"), (int32)FullAspectMask, (int32)GetDesc().Format);
-		checkf(VulkanRHI::GetAspectMaskFromUEFormat(GetDesc().Format, false, true) == PartialAspectMask, TEXT("PartialAspectMask (%d) does not match with PixelFormat (%d)"), (int32)PartialAspectMask, (int32)GetDesc().Format);
+		checkf(MemProps == VK_MEMORY_PROPERTY_DEVICE_LOCAL_BIT, TEXT("Can't evict surface that isn't device local.  MemoryProperties=%s"), VK_FLAGS_TO_STRING(VkMemoryPropertyFlags, MemProps));
+		checkf(VulkanRHI::GetAspectMaskFromUEFormat(GetDesc().Format, true, true) == FullAspectMask, TEXT("FullAspectMask (%s) does not match with PixelFormat (%d)"), VK_FLAGS_TO_STRING(VkImageAspectFlags, FullAspectMask), (int32)GetDesc().Format);
+		checkf(VulkanRHI::GetAspectMaskFromUEFormat(GetDesc().Format, false, true) == PartialAspectMask, TEXT("PartialAspectMask (%s) does not match with PixelFormat (%d)"), VK_FLAGS_TO_STRING(VkImageAspectFlags, PartialAspectMask), (int32)GetDesc().Format);
 
 		const ETextureCreateFlags UEFlags = GetDesc().Flags;
 		const bool bRenderTarget = EnumHasAnyFlags(UEFlags, TexCreate_RenderTargetable | TexCreate_DepthStencilTargetable | TexCreate_ResolveTargetable);
@@ -2305,21 +2305,35 @@ FDynamicRHI::FRHICalcTextureSizeResult FVulkanDynamicRHI::RHICalcTexturePlatform
 	FVulkanTexture::FImageCreateInfo TmpCreateInfo;
 	FVulkanTexture::GenerateImageCreateInfo(TmpCreateInfo, *Device, CleanDesc, nullptr, nullptr, false);
 
-	// TODO: VK_KHR_maintenance4
-	VkImage TmpImage;
-	VERIFYVULKANRESULT(VulkanRHI::vkCreateImage(Device->GetInstanceHandle(), &TmpCreateInfo.ImageCreateInfo, VULKAN_CPU_ALLOCATOR, &TmpImage));
+	VkMemoryRequirements OutMemReq;
 
-	VkMemoryRequirements MemReq;
-	VulkanRHI::vkGetImageMemoryRequirements(Device->GetInstanceHandle(), TmpImage, &MemReq);
+	if (Device->GetOptionalExtensions().HasKHRMaintenance4)
+	{
+		VkDeviceImageMemoryRequirementsKHR ImageMemReq;
+		ZeroVulkanStruct(ImageMemReq, VK_STRUCTURE_TYPE_DEVICE_IMAGE_MEMORY_REQUIREMENTS);
+		ImageMemReq.pCreateInfo = &TmpCreateInfo.ImageCreateInfo;
+		ImageMemReq.planeAspect = (VulkanRHI::GetAspectMaskFromUEFormat(CleanDesc.Format, true, true) == VK_IMAGE_ASPECT_COLOR_BIT) ? VK_IMAGE_ASPECT_COLOR_BIT : VK_IMAGE_ASPECT_DEPTH_BIT;  // should be ignored
+
+		VkMemoryRequirements2KHR MemReq2;
+		ZeroVulkanStruct(MemReq2, VK_STRUCTURE_TYPE_MEMORY_REQUIREMENTS_2);
 
-	VulkanRHI::vkDestroyImage(Device->GetInstanceHandle(), TmpImage, VULKAN_CPU_ALLOCATOR);
+		VulkanRHI::vkGetDeviceImageMemoryRequirementsKHR(Device->GetInstanceHandle(), &ImageMemReq, &MemReq2);
+		OutMemReq = MemReq2.memoryRequirements;
+	}
+	else
+	{
+		VkImage TmpImage;
+		VERIFYVULKANRESULT(VulkanRHI::vkCreateImage(Device->GetInstanceHandle(), &TmpCreateInfo.ImageCreateInfo, VULKAN_CPU_ALLOCATOR, &TmpImage));
+		VulkanRHI::vkGetImageMemoryRequirements(Device->GetInstanceHandle(), TmpImage, &OutMemReq);
+		VulkanRHI::vkDestroyImage(Device->GetInstanceHandle(), TmpImage, VULKAN_CPU_ALLOCATOR);
+	}
 
 	{
 		FScopeLock Lock(&TextureSizesLock);
-		TextureSizes.Add(CleanDesc, MemReq);
+		TextureSizes.Add(CleanDesc, OutMemReq);
 	}
 
-	return { (uint64)MemReq.size, (uint32)MemReq.alignment };
+	return { (uint64)OutMemReq.size, (uint32)OutMemReq.alignment };
 }
 
 void FVulkanCommandListContext::RHICopyTexture(FRHITexture* SourceTexture, FRHITexture* DestTexture, const FRHICopyTextureInfo& CopyInfo)
@@ -2331,7 +2345,7 @@ void FVulkanCommandListContext::RHICopyTexture(FRHITexture* SourceTexture, FRHIT
 	FVulkanTexture* Dest = FVulkanTexture::Cast(DestTexture);
 
 	VkImageLayout SrcLayout = LayoutManager.FindLayoutChecked(Source->Image);
-	ensureMsgf(SrcLayout == VK_IMAGE_LAYOUT_TRANSFER_SRC_OPTIMAL, TEXT("Expected source texture to be in VK_IMAGE_LAYOUT_TRANSFER_SRC_OPTIMAL, actual layout is %d"), SrcLayout);
+	ensureMsgf(SrcLayout == VK_IMAGE_LAYOUT_TRANSFER_SRC_OPTIMAL, TEXT("Expected source texture to be in VK_IMAGE_LAYOUT_TRANSFER_SRC_OPTIMAL, actual layout is %s"), VK_TYPE_TO_STRING(VkImageLayout, SrcLayout));
 
 	FVulkanCmdBuffer* InCmdBuffer = CommandBufferManager->GetActiveCmdBuffer();
 	check(InCmdBuffer->IsOutsideRenderPass());
@@ -2387,7 +2401,7 @@ void FVulkanCommandListContext::RHICopyTexture(FRHITexture* SourceTexture, FRHIT
 	else
 	{
 		VkImageLayout DstLayout = LayoutManager.FindLayoutChecked(Dest->Image);
-		ensureMsgf(DstLayout == VK_IMAGE_LAYOUT_TRANSFER_DST_OPTIMAL, TEXT("Expected destination texture to be in VK_IMAGE_LAYOUT_TRANSFER_DST_OPTIMAL, actual layout is %d"), DstLayout);
+		ensureMsgf(DstLayout == VK_IMAGE_LAYOUT_TRANSFER_DST_OPTIMAL, TEXT("Expected destination texture to be in VK_IMAGE_LAYOUT_TRANSFER_DST_OPTIMAL, actual layout is %s"), VK_TYPE_TO_STRING(VkImageLayout, DstLayout));
 
 		VkImageCopy Region;
 		FMemory::Memzero(Region);
