VulkanRHI: Clean up.  Remove Maintenance defines that are guaranteed with new SDK version minimum.

#rnx

diff --git a/Engine/Source/Runtime/VulkanRHI/Private/VulkanBarriers.cpp b/Engine/Source/Runtime/VulkanRHI/Private/VulkanBarriers.cpp
--- a/Engine/Source/Runtime/VulkanRHI/Private/VulkanBarriers.cpp
+++ b/Engine/Source/Runtime/VulkanRHI/Private/VulkanBarriers.cpp
@@ -40,12 +40,10 @@ static VkAccessFlags GetVkAccessMaskForLayout(VkImageLayout Layout)
 			Flags = VK_ACCESS_DEPTH_STENCIL_ATTACHMENT_READ_BIT | VK_ACCESS_DEPTH_STENCIL_ATTACHMENT_WRITE_BIT;
 			break;
 
-#if VULKAN_SUPPORTS_MAINTENANCE_LAYER2
 		case VK_IMAGE_LAYOUT_DEPTH_READ_ONLY_STENCIL_ATTACHMENT_OPTIMAL_KHR:
 		case VK_IMAGE_LAYOUT_DEPTH_ATTACHMENT_STENCIL_READ_ONLY_OPTIMAL_KHR:
 			Flags = VK_ACCESS_SHADER_READ_BIT | VK_ACCESS_DEPTH_STENCIL_ATTACHMENT_READ_BIT | VK_ACCESS_DEPTH_STENCIL_ATTACHMENT_WRITE_BIT;
 			break;
-#endif
 
 		case VK_IMAGE_LAYOUT_SHADER_READ_ONLY_OPTIMAL:
 			Flags = VK_ACCESS_SHADER_READ_BIT;
@@ -119,10 +117,8 @@ static VkPipelineStageFlags GetVkStageFlagsForLayout(VkImageLayout Layout)
 			break;
 
 		case VK_IMAGE_LAYOUT_DEPTH_STENCIL_READ_ONLY_OPTIMAL:
-#if VULKAN_SUPPORTS_MAINTENANCE_LAYER2
 		case VK_IMAGE_LAYOUT_DEPTH_READ_ONLY_STENCIL_ATTACHMENT_OPTIMAL_KHR:
 		case VK_IMAGE_LAYOUT_DEPTH_ATTACHMENT_STENCIL_READ_ONLY_OPTIMAL_KHR:
-#endif
 #if VULKAN_SUPPORTS_SEPARATE_DEPTH_STENCIL_LAYOUTS
 		case VK_IMAGE_LAYOUT_DEPTH_READ_ONLY_OPTIMAL_KHR:
 		case VK_IMAGE_LAYOUT_STENCIL_READ_ONLY_OPTIMAL_KHR:
diff --git a/Engine/Source/Runtime/VulkanRHI/Private/VulkanExtensions.cpp b/Engine/Source/Runtime/VulkanRHI/Private/VulkanExtensions.cpp
--- a/Engine/Source/Runtime/VulkanRHI/Private/VulkanExtensions.cpp
+++ b/Engine/Source/Runtime/VulkanRHI/Private/VulkanExtensions.cpp
@@ -117,7 +117,7 @@ class FVulkanKHRMaintenance4Extension : public FVulkanDeviceExtension
 public:
 
 	FVulkanKHRMaintenance4Extension(FVulkanDevice* InDevice)
-		: FVulkanDeviceExtension(InDevice, VK_KHR_MAINTENANCE_4_EXTENSION_NAME, VULKAN_SUPPORTS_MAINTENANCE4, VK_API_VERSION_1_3)
+		: FVulkanDeviceExtension(InDevice, VK_KHR_MAINTENANCE_4_EXTENSION_NAME, VULKAN_EXTENSION_ENABLED, VK_API_VERSION_1_3)
 	{}
 
 	virtual void PrePhysicalDeviceFeatures(VkPhysicalDeviceFeatures2KHR& PhysicalDeviceFeatures2) override final
@@ -898,8 +898,8 @@ FVulkanDeviceExtensionArray FVulkanDeviceExtension::GetUESupportedDeviceExtensio
 	// Generic simple extensions :
 
 	ADD_SIMPLE_EXTENSION(VK_KHR_SWAPCHAIN_EXTENSION_NAME,                VULKAN_EXTENSION_ENABLED,             VULKAN_EXTENSION_NOT_PROMOTED, nullptr);
-	ADD_SIMPLE_EXTENSION(VK_KHR_MAINTENANCE1_EXTENSION_NAME,             VULKAN_SUPPORTS_MAINTENANCE_LAYER1,   VK_API_VERSION_1_1,            DEVICE_EXT_FLAG_SETTER(HasKHRMaintenance1));
-	ADD_SIMPLE_EXTENSION(VK_KHR_MAINTENANCE2_EXTENSION_NAME,             VULKAN_SUPPORTS_MAINTENANCE_LAYER2,   VK_API_VERSION_1_1,            DEVICE_EXT_FLAG_SETTER(HasKHRMaintenance2));
+	ADD_SIMPLE_EXTENSION(VK_KHR_MAINTENANCE1_EXTENSION_NAME,             VULKAN_EXTENSION_ENABLED,             VK_API_VERSION_1_1,            DEVICE_EXT_FLAG_SETTER(HasKHRMaintenance1));
+	ADD_SIMPLE_EXTENSION(VK_KHR_MAINTENANCE2_EXTENSION_NAME,             VULKAN_EXTENSION_ENABLED,             VK_API_VERSION_1_1,            DEVICE_EXT_FLAG_SETTER(HasKHRMaintenance2));
 	ADD_SIMPLE_EXTENSION(VK_EXT_MEMORY_BUDGET_EXTENSION_NAME,            VULKAN_SUPPORTS_MEMORY_BUDGET,        VULKAN_EXTENSION_NOT_PROMOTED, DEVICE_EXT_FLAG_SETTER(HasMemoryBudget));
 	ADD_SIMPLE_EXTENSION(VK_EXT_MEMORY_PRIORITY_EXTENSION_NAME,          VULKAN_SUPPORTS_MEMORY_PRIORITY,      VULKAN_EXTENSION_NOT_PROMOTED, DEVICE_EXT_FLAG_SETTER(HasMemoryPriority));
 	ADD_SIMPLE_EXTENSION(VK_KHR_CREATE_RENDERPASS_2_EXTENSION_NAME,      VULKAN_SUPPORTS_RENDERPASS2,          VK_API_VERSION_1_2,            DEVICE_EXT_FLAG_SETTER(HasKHRRenderPass2));
diff --git a/Engine/Source/Runtime/VulkanRHI/Private/VulkanRHIPrivate.h b/Engine/Source/Runtime/VulkanRHI/Private/VulkanRHIPrivate.h
--- a/Engine/Source/Runtime/VulkanRHI/Private/VulkanRHIPrivate.h
+++ b/Engine/Source/Runtime/VulkanRHI/Private/VulkanRHIPrivate.h
@@ -813,7 +813,6 @@ namespace VulkanRHI
 		{
 			return VK_IMAGE_LAYOUT_DEPTH_STENCIL_READ_ONLY_OPTIMAL;
 		}
-#if VULKAN_SUPPORTS_MAINTENANCE_LAYER2
 		else if (RequestedDSAccess == FExclusiveDepthStencil::DepthRead_StencilWrite && InDevice.GetOptionalExtensions().HasKHRMaintenance2)
 		{
 			return VK_IMAGE_LAYOUT_DEPTH_READ_ONLY_STENCIL_ATTACHMENT_OPTIMAL_KHR;
@@ -822,7 +821,6 @@ namespace VulkanRHI
 		{
 			return VK_IMAGE_LAYOUT_DEPTH_READ_ONLY_STENCIL_ATTACHMENT_OPTIMAL_KHR;
 		}
-#endif
 
 		ensure(RequestedDSAccess.IsDepthWrite() || RequestedDSAccess.IsStencilWrite());
 		return VK_IMAGE_LAYOUT_DEPTH_STENCIL_ATTACHMENT_OPTIMAL;
diff --git a/Engine/Source/Runtime/VulkanRHI/Private/VulkanTexture.cpp b/Engine/Source/Runtime/VulkanRHI/Private/VulkanTexture.cpp
--- a/Engine/Source/Runtime/VulkanRHI/Private/VulkanTexture.cpp
+++ b/Engine/Source/Runtime/VulkanRHI/Private/VulkanTexture.cpp
@@ -309,12 +309,10 @@ void FVulkanTexture::GenerateImageCreateInfo(
 		ImageCreateInfo.flags |= VK_IMAGE_CREATE_MUTABLE_FORMAT_BIT;
 	}
 
-#if VULKAN_SUPPORTS_MAINTENANCE_LAYER1
 	if (InDevice.GetOptionalExtensions().HasKHRMaintenance1 && ImageCreateInfo.imageType == VK_IMAGE_TYPE_3D)
 	{
 		ImageCreateInfo.flags |= VK_IMAGE_CREATE_2D_ARRAY_COMPATIBLE_BIT_KHR;
 	}
-#endif
 
 	ImageCreateInfo.tiling = bForceLinearTexture ? VK_IMAGE_TILING_LINEAR : GVulkanViewTypeTilingMode[ResourceType];
 
diff --git a/Engine/Source/Runtime/VulkanRHI/Public/VulkanConfiguration.h b/Engine/Source/Runtime/VulkanRHI/Public/VulkanConfiguration.h
--- a/Engine/Source/Runtime/VulkanRHI/Public/VulkanConfiguration.h
+++ b/Engine/Source/Runtime/VulkanRHI/Public/VulkanConfiguration.h
@@ -100,18 +100,6 @@
 	#define VULKAN_ENABLE_LRU_CACHE								0
 #endif
 
-#ifdef VK_KHR_maintenance1
-	#define VULKAN_SUPPORTS_MAINTENANCE_LAYER1					1
-#else
-	#define VULKAN_SUPPORTS_MAINTENANCE_LAYER1					0
-#endif
-
-#ifdef VK_KHR_maintenance2
-	#define VULKAN_SUPPORTS_MAINTENANCE_LAYER2					1
-#else
-	#define VULKAN_SUPPORTS_MAINTENANCE_LAYER2					0
-#endif
-
 #ifdef VK_EXT_validation_cache
 	#define VULKAN_SUPPORTS_VALIDATION_CACHE					1
 #else
@@ -333,12 +321,6 @@
 	#define VULKAN_SUPPORTS_SHADER_VIEWPORT_INDEX_LAYER	0
 #endif
 
-#ifdef VK_KHR_maintenance4
-	#define VULKAN_SUPPORTS_MAINTENANCE4	(VULKAN_SUPPORTS_PHYSICAL_DEVICE_PROPERTIES2)	// Requirement
-#else
-	#define VULKAN_SUPPORTS_MAINTENANCE4	0
-#endif
-
 #ifndef VULKAN_SUPPORTS_DESCRIPTOR_INDEXING
 	#ifdef VK_EXT_descriptor_indexing
 		#define VULKAN_SUPPORTS_DESCRIPTOR_INDEXING	(VULKAN_SUPPORTS_PHYSICAL_DEVICE_PROPERTIES2)	// Requirement
