VulkanRHI: Add support for Host Query Reset extension.  If it is supported, reset all new query pools from host (no cmdbuffer needed) on creation to fix validation issue.  Can be disabled using "r.Vulkan.AllowHostQueryReset".

#rnx

diff --git a/Engine/Source/Runtime/VulkanRHI/Private/VulkanCommandWrappers.h b/Engine/Source/Runtime/VulkanRHI/Private/VulkanCommandWrappers.h
--- a/Engine/Source/Runtime/VulkanRHI/Private/VulkanCommandWrappers.h
+++ b/Engine/Source/Runtime/VulkanRHI/Private/VulkanCommandWrappers.h
@@ -202,6 +202,7 @@ struct FWrapLayer
 	static void GetBufferDeviceAddressKHR(VkResult Result, VkDevice Device, const VkBufferDeviceAddressInfo* Info);
 #endif
 	static void GetDeviceImageMemoryRequirementsKHR(VkResult Result, VkDevice Device, const VkDeviceImageMemoryRequirements* Info, VkMemoryRequirements2* MemoryRequirements) VULKAN_LAYER_BODY
+	static void ResetQueryPoolEXT(VkResult Result, VkDevice Device, VkQueryPool QueryPool, uint32_t FirstQuery, uint32_t QueryCount) VULKAN_LAYER_BODY
 };
 
 #undef VULKAN_LAYER_BODY
@@ -1476,6 +1477,13 @@ namespace VulkanRHI
 		FWrapLayer::GetDeviceImageMemoryRequirementsKHR(VK_SUCCESS, Device, Info, MemoryRequirements);
 	}
 
+	static FORCEINLINE_DEBUGGABLE void vkResetQueryPoolEXT(VkDevice Device, VkQueryPool QueryPool, uint32_t FirstQuery, uint32_t QueryCount)
+	{
+		FWrapLayer::ResetQueryPoolEXT(VK_RESULT_MAX_ENUM, Device, QueryPool, FirstQuery, QueryCount);
+		VULKANAPINAMESPACE::vkResetQueryPoolEXT(Device, QueryPool, FirstQuery, QueryCount);
+		FWrapLayer::ResetQueryPoolEXT(VK_SUCCESS, Device, QueryPool, FirstQuery, QueryCount);
+	}
+
 #if VULKAN_ENABLE_IMAGE_TRACKING_LAYER
 	void BindDebugLabelName(VkImage Image, const TCHAR* Name);
 #endif
diff --git a/Engine/Source/Runtime/VulkanRHI/Private/VulkanDebug.cpp b/Engine/Source/Runtime/VulkanRHI/Private/VulkanDebug.cpp
--- a/Engine/Source/Runtime/VulkanRHI/Private/VulkanDebug.cpp
+++ b/Engine/Source/Runtime/VulkanRHI/Private/VulkanDebug.cpp
@@ -4466,6 +4466,16 @@ void FWrapLayer::GetDeviceImageMemoryRequirementsKHR(VkResult Result, VkDevice D
 	}
 }
 
+void FWrapLayer::ResetQueryPoolEXT(VkResult Result, VkDevice Device, VkQueryPool QueryPool, uint32_t FirstQuery, uint32_t QueryCount)
+{
+	if (Result == VK_RESULT_MAX_ENUM)
+	{
+#if VULKAN_ENABLE_DUMP_LAYER
+		PrintfBeginResult(FString::Printf(TEXT("ResetQueryPoolEXT(Device=0x%p, QueryPool=0x%p, FirstQuery=%u, QueryCount=%u)")), Device, QueryPool, FirstQuery, QueryCount);
+#endif
+	}
+}
+
 #if VULKAN_ENABLE_IMAGE_TRACKING_LAYER
 namespace VulkanRHI
 {
diff --git a/Engine/Source/Runtime/VulkanRHI/Private/VulkanDevice.h b/Engine/Source/Runtime/VulkanRHI/Private/VulkanDevice.h
--- a/Engine/Source/Runtime/VulkanRHI/Private/VulkanDevice.h
+++ b/Engine/Source/Runtime/VulkanRHI/Private/VulkanDevice.h
@@ -76,6 +76,7 @@ struct FOptionalVulkanDeviceExtensions
 			uint64 HasEXTDescriptorIndexing : 1;
 			uint64 HasEXTShaderViewportIndexLayer : 1;
 			uint64 HasSeparateDepthStencilLayouts : 1;
+			uint64 HasEXTHostQueryReset : 1;
 
 			// Promoted to 1.3
 			uint64 HasEXTTextureCompressionASTCHDR : 1;
diff --git a/Engine/Source/Runtime/VulkanRHI/Private/VulkanExtensions.cpp b/Engine/Source/Runtime/VulkanRHI/Private/VulkanExtensions.cpp
--- a/Engine/Source/Runtime/VulkanRHI/Private/VulkanExtensions.cpp
+++ b/Engine/Source/Runtime/VulkanRHI/Private/VulkanExtensions.cpp
@@ -43,6 +43,15 @@ TAutoConsoleVariable<int32> GVulkanRayTracingCVar(
 	ECVF_ReadOnly
 );
 
+TAutoConsoleVariable<int32> GVulkanAllowHostQueryResetCVar(
+	TEXT("r.Vulkan.AllowHostQueryReset"),
+	1,
+	TEXT("0: Do not enable support for Host Query Reset extension\n")
+	TEXT("1: Enable Host Query Reset (default)"),
+	ECVF_ReadOnly
+);
+
+
 #if VULKAN_HAS_DEBUGGING_ENABLED
 extern TAutoConsoleVariable<int32> GGPUValidationCvar;
 #endif
@@ -786,7 +795,7 @@ public:
 		: FVulkanDeviceExtension(InDevice, VK_AMD_BUFFER_MARKER_EXTENSION_NAME, VULKAN_SUPPORTS_AMD_BUFFER_MARKER)
 	{
 		const bool bAllowVendorDevice = !FParse::Param(FCommandLine::Get(), TEXT("novendordevice"));
-		bEnabledInCode = GGPUCrashDebuggingEnabled && bAllowVendorDevice;
+		bEnabledInCode = bEnabledInCode && GGPUCrashDebuggingEnabled && bAllowVendorDevice;
 	}
 
 	virtual void PostPhysicalDeviceFeatures(FOptionalVulkanDeviceExtensions& ExtensionFlags) override final
@@ -806,7 +815,7 @@ public:
 		: FVulkanDeviceExtension(InDevice, VK_NV_DEVICE_DIAGNOSTIC_CHECKPOINTS_EXTENSION_NAME, VULKAN_SUPPORTS_NV_DIAGNOSTICS)
 	{
 		const bool bAllowVendorDevice = !FParse::Param(FCommandLine::Get(), TEXT("novendordevice"));
-		bEnabledInCode = GGPUCrashDebuggingEnabled && bAllowVendorDevice;
+		bEnabledInCode = bEnabledInCode && GGPUCrashDebuggingEnabled && bAllowVendorDevice;
 	}
 
 	virtual void PostPhysicalDeviceFeatures(FOptionalVulkanDeviceExtensions& ExtensionFlags) override final
@@ -826,7 +835,7 @@ public:
 		: FVulkanDeviceExtension(InDevice, VK_NV_DEVICE_DIAGNOSTICS_CONFIG_EXTENSION_NAME, VULKAN_SUPPORTS_NV_DIAGNOSTICS)
 	{
 		const bool bAllowVendorDevice = !FParse::Param(FCommandLine::Get(), TEXT("novendordevice"));
-		bEnabledInCode = GGPUCrashDebuggingEnabled && bAllowVendorDevice;
+		bEnabledInCode = bEnabledInCode && GGPUCrashDebuggingEnabled && bAllowVendorDevice;
 	}
 
 
@@ -859,6 +868,41 @@ private:
 };
 
 
+// ***** VK_EXT_host_query_reset
+class FVulkanEXTHostQueryResetExtension : public FVulkanDeviceExtension
+{
+public:
+
+	FVulkanEXTHostQueryResetExtension(FVulkanDevice* InDevice)
+		: FVulkanDeviceExtension(InDevice, VK_EXT_HOST_QUERY_RESET_EXTENSION_NAME, VULKAN_EXTENSION_ENABLED, VK_API_VERSION_1_2)
+	{
+		bEnabledInCode = bEnabledInCode && (GVulkanAllowHostQueryResetCVar.GetValueOnAnyThread() != 0);
+	}
+
+	virtual void PrePhysicalDeviceFeatures(VkPhysicalDeviceFeatures2KHR& PhysicalDeviceFeatures2) override final
+	{
+		ZeroVulkanStruct(HostQueryResetFeatures, VK_STRUCTURE_TYPE_PHYSICAL_DEVICE_HOST_QUERY_RESET_FEATURES);
+		AddToPNext(PhysicalDeviceFeatures2, HostQueryResetFeatures);
+	}
+
+	virtual void PostPhysicalDeviceFeatures(FOptionalVulkanDeviceExtensions& ExtensionFlags) override final
+	{
+		ExtensionFlags.HasEXTHostQueryReset = HostQueryResetFeatures.hostQueryReset;
+	}
+
+	virtual void PreCreateDevice(VkDeviceCreateInfo& DeviceCreateInfo) override final
+	{
+		if (HostQueryResetFeatures.hostQueryReset == VK_TRUE)
+		{
+			AddToPNext(DeviceCreateInfo, HostQueryResetFeatures);
+		}
+	}
+
+private:
+	VkPhysicalDeviceHostQueryResetFeaturesEXT HostQueryResetFeatures;
+};
+
+
 
 
 
@@ -931,6 +975,7 @@ FVulkanDeviceExtensionArray FVulkanDeviceExtension::GetUESupportedDeviceExtensio
 	ADD_CUSTOM_EXTENSION(FVulkanKHRMultiviewExtension);
 	ADD_CUSTOM_EXTENSION(FVulkanKHRGetMemoryRequirements2Extension);
 	ADD_CUSTOM_EXTENSION(FVulkanEXTDescriptorIndexingExtension);
+	ADD_CUSTOM_EXTENSION(FVulkanEXTHostQueryResetExtension);
 
 	// Needed for Raytracing
 	ADD_CUSTOM_EXTENSION(FVulkanKHRBufferDeviceAddressExtension);
diff --git a/Engine/Source/Runtime/VulkanRHI/Private/VulkanLoader.h b/Engine/Source/Runtime/VulkanRHI/Private/VulkanLoader.h
--- a/Engine/Source/Runtime/VulkanRHI/Private/VulkanLoader.h
+++ b/Engine/Source/Runtime/VulkanRHI/Private/VulkanLoader.h
@@ -173,6 +173,7 @@
 // List all optional Vulkan entry points used by Unreal that need to be loaded manually
 #define ENUM_VK_ENTRYPOINTS_OPTIONAL_INSTANCE(EnumMacro) \
 	EnumMacro(PFN_vkGetDeviceImageMemoryRequirementsKHR, vkGetDeviceImageMemoryRequirementsKHR) \
+	EnumMacro(PFN_vkResetQueryPoolEXT, vkResetQueryPoolEXT) \
 	EnumMacro(PFN_vkCreateSharedSwapchainsKHR, vkCreateSharedSwapchainsKHR)
 
 // List of all Vulkan entry points
diff --git a/Engine/Source/Runtime/VulkanRHI/Private/VulkanQuery.cpp b/Engine/Source/Runtime/VulkanRHI/Private/VulkanQuery.cpp
--- a/Engine/Source/Runtime/VulkanRHI/Private/VulkanQuery.cpp
+++ b/Engine/Source/Runtime/VulkanRHI/Private/VulkanQuery.cpp
@@ -57,7 +57,12 @@ FVulkanQueryPool::FVulkanQueryPool(FVulkanDevice* InDevice, FVulkanCommandBuffer
 
 	VERIFYVULKANRESULT(VulkanRHI::vkCreateQueryPool(Device->GetInstanceHandle(), &PoolCreateInfo, VULKAN_CPU_ALLOCATOR, &QueryPool));
 	
-	if (bInShouldAddReset && CommandBufferManager)
+	// If host query resets are supported, reset all new query pools on creation
+	if (Device->GetOptionalExtensions().HasEXTHostQueryReset)
+	{
+		VulkanRHI::vkResetQueryPoolEXT(Device->GetInstanceHandle(), QueryPool, 0, MaxQueries);
+	}
+	else if (bInShouldAddReset && CommandBufferManager)
 	{
 		CommandBufferManager->AddQueryPoolForReset(QueryPool, InMaxQueries);
 	}
@@ -105,7 +110,7 @@ bool FVulkanOcclusionQueryPool::InternalTryGetResults(bool bWait)
 	if (VulkanRHI::vkGetEventStatus(Device->GetInstanceHandle(), ResetEvent) == VK_EVENT_SET)
 	{
 		Result = VulkanRHI::vkGetQueryPoolResults(Device->GetInstanceHandle(), QueryPool, 0, NumUsedQueries, NumUsedQueries * sizeof(uint64), QueryOutput.GetData(), sizeof(uint64), VK_QUERY_RESULT_64_BIT);
-				if (Result == VK_SUCCESS)
+		if (Result == VK_SUCCESS)
 		{
 			State = EState::RT_PostGetResults;
 			return true;
diff --git a/Engine/Source/Runtime/VulkanRHI/Private/VulkanUtil.cpp b/Engine/Source/Runtime/VulkanRHI/Private/VulkanUtil.cpp
--- a/Engine/Source/Runtime/VulkanRHI/Private/VulkanUtil.cpp
+++ b/Engine/Source/Runtime/VulkanRHI/Private/VulkanUtil.cpp
@@ -203,6 +203,12 @@ void FVulkanGPUTiming::StartTiming(FVulkanCmdBuffer* CmdBuffer)
 		Pool->CurrentTimestamp = (Pool->CurrentTimestamp + 1) % Pool->BufferSize;
 		const uint32 QueryStartIndex = Pool->CurrentTimestamp * 2;
 
+		// If host query resets are supported, reset timestamp queries before writing to them (no need to consider host/GPU sync)
+		if (Device->GetOptionalExtensions().HasEXTHostQueryReset)
+		{
+			VulkanRHI::vkResetQueryPoolEXT(Device->GetInstanceHandle(), Pool->GetHandle(), QueryStartIndex, 2);
+		}
+
 		VulkanRHI::vkCmdWriteTimestamp(CmdBuffer->GetHandle(), VK_PIPELINE_STAGE_BOTTOM_OF_PIPE_BIT, Pool->GetHandle(), QueryStartIndex);
 		Pool->TimestampListHandles[QueryStartIndex].CmdBuffer = CmdBuffer;
 		Pool->TimestampListHandles[QueryStartIndex].FenceCounter = CmdBuffer->GetFenceSignaledCounter();
@@ -226,8 +232,8 @@ void FVulkanGPUTiming::EndTiming(FVulkanCmdBuffer* CmdBuffer)
 		}
 		check(Pool->CurrentTimestamp < Pool->BufferSize);
 		const uint32 QueryStartIndex = Pool->CurrentTimestamp * 2;
-		const uint32 QueryEndIndex = Pool->CurrentTimestamp * 2 + 1;
-		check(QueryEndIndex == QueryStartIndex + 1);	// Make sure they're adjacent indices.
+		// Keep Start and End contiguous to fetch them together with a single AddPendingTimestampQuery(QueryStartIndex,2,...)
+		const uint32 QueryEndIndex = QueryStartIndex + 1;   
 
 		// In case we aren't reading queries, remove oldest
 		if (NumPendingQueries >= Pool->BufferSize)
@@ -240,8 +246,7 @@ void FVulkanGPUTiming::EndTiming(FVulkanCmdBuffer* CmdBuffer)
 		NumPendingQueries++;
 
 		VulkanRHI::vkCmdWriteTimestamp(CmdBuffer->GetHandle(), VK_PIPELINE_STAGE_BOTTOM_OF_PIPE_BIT, Pool->GetHandle(), QueryEndIndex);
-		CmdBuffer->AddPendingTimestampQuery(QueryStartIndex, 1, Pool->GetHandle(), Pool->ResultsBuffer->GetHandle(), false);
-		CmdBuffer->AddPendingTimestampQuery(QueryEndIndex, 1, Pool->GetHandle(), Pool->ResultsBuffer->GetHandle(), false);
+		CmdBuffer->AddPendingTimestampQuery(QueryStartIndex, 2, Pool->GetHandle(), Pool->ResultsBuffer->GetHandle(), false);
 		Pool->TimestampListHandles[QueryEndIndex].CmdBuffer = CmdBuffer;
 		Pool->TimestampListHandles[QueryEndIndex].FenceCounter = CmdBuffer->GetFenceSignaledCounter();
 		Pool->TimestampListHandles[QueryEndIndex].FrameCount = CmdContext->GetFrameCounter();
@@ -281,7 +286,7 @@ uint64 FVulkanGPUTiming::GetTiming(bool bGetCurrentResultsAndBlock)
 		while (PendingQueries.Peek(TimestampIndex))
 		{
 			const uint32 QueryStartIndex = TimestampIndex * 2;
-			const uint32 QueryEndIndex = TimestampIndex * 2 + 1;
+			const uint32 QueryEndIndex = QueryStartIndex + 1;
 
 			const FVulkanTimingQueryPool::FCmdBufferFence& StartQuerySyncPoint = Pool->TimestampListHandles[QueryStartIndex];
 			const FVulkanTimingQueryPool::FCmdBufferFence& EndQuerySyncPoint = Pool->TimestampListHandles[QueryEndIndex];
