VulkanRHI: Disable UE4 specific fix causing performance hitches on UE5.

diff --git a/Engine/Source/Runtime/VulkanRHI/Private/VulkanRenderTarget.cpp b/Engine/Source/Runtime/VulkanRHI/Private/VulkanRenderTarget.cpp
--- a/Engine/Source/Runtime/VulkanRHI/Private/VulkanRenderTarget.cpp
+++ b/Engine/Source/Runtime/VulkanRHI/Private/VulkanRenderTarget.cpp
@@ -12,9 +12,9 @@
 #include "SceneUtils.h"
 #include "RHISurfaceDataConversion.h"
 
-// Enable this by default for desktop platforms, but avoid it for others such as mobile due to a DEVICE LOST when alt+tabing
+// Debug mode used as workaround when a DEVICE LOST occurs on alt+tab on some platforms
 // This is a workaround and may end up causing some hitches on the rendering thread
-static int32 GVulkanFlushOnMapStaging = PLATFORM_DESKTOP;
+static int32 GVulkanFlushOnMapStaging = 0;
 static FAutoConsoleVariableRef CVarGVulkanFlushOnMapStaging(
 	TEXT("r.Vulkan.FlushOnMapStaging"),
 	GVulkanFlushOnMapStaging,
