VulkanRHI: Add CPUReadback as an alternative to read images without SAMPLED_IMAGE_BIT.

#rnx

diff --git a/Engine/Source/Runtime/VulkanRHI/Private/VulkanTexture.cpp b/Engine/Source/Runtime/VulkanRHI/Private/VulkanTexture.cpp
--- a/Engine/Source/Runtime/VulkanRHI/Private/VulkanTexture.cpp
+++ b/Engine/Source/Runtime/VulkanRHI/Private/VulkanTexture.cpp
@@ -438,7 +438,7 @@ void FVulkanTexture::GenerateImageCreateInfo(
 	if (!VKHasAnyFlags(FormatFlags, VK_FORMAT_FEATURE_SAMPLED_IMAGE_BIT))
 	{
 		// Some formats don't support sampling and that's ok, we'll use a STORAGE_IMAGE
-		check(EnumHasAnyFlags(UEFlags, TexCreate_UAV));
+		check(EnumHasAnyFlags(UEFlags, TexCreate_UAV | TexCreate_CPUReadback));
 		ImageCreateInfo.usage &= ~VK_IMAGE_USAGE_SAMPLED_BIT;
 	}
 
