Fixed: Vulkan validation warning related to a buffer upload

diff --git a/Engine/Source/Runtime/VulkanRHI/Private/VulkanUniformBuffer.cpp b/Engine/Source/Runtime/VulkanRHI/Private/VulkanUniformBuffer.cpp
--- a/Engine/Source/Runtime/VulkanRHI/Private/VulkanUniformBuffer.cpp
+++ b/Engine/Source/Runtime/VulkanRHI/Private/VulkanUniformBuffer.cpp
@@ -88,7 +88,13 @@ static void UpdateUniformBufferHelper(FVulkanCommandListContext& Context, FVulka
 		VkBuffer UBBuffer = VulkanUniformBuffer->Allocation.GetBufferHandle();
 		VkBuffer LockHandle = LockInfo.Allocation.GetBufferHandle();
 
+		VkMemoryBarrier BarrierBefore = { VK_STRUCTURE_TYPE_MEMORY_BARRIER, nullptr, VK_ACCESS_UNIFORM_READ_BIT, VK_ACCESS_TRANSFER_WRITE_BIT };
+		VulkanRHI::vkCmdPipelineBarrier(CmdBuffer->GetHandle(), VK_PIPELINE_STAGE_FRAGMENT_SHADER_BIT, VK_PIPELINE_STAGE_TRANSFER_BIT, 0, 1, &BarrierBefore, 0, nullptr, 0, nullptr);
+
 		VulkanRHI::vkCmdCopyBuffer(CmdBuffer->GetHandle(), LockHandle, UBBuffer, 1, &Region);
+
+		VkMemoryBarrier BarrierAfter = { VK_STRUCTURE_TYPE_MEMORY_BARRIER, nullptr, VK_ACCESS_TRANSFER_WRITE_BIT, VK_ACCESS_UNIFORM_READ_BIT };
+		VulkanRHI::vkCmdPipelineBarrier(CmdBuffer->GetHandle(), VK_PIPELINE_STAGE_TRANSFER_BIT, VK_PIPELINE_STAGE_VERTEX_SHADER_BIT|VK_PIPELINE_STAGE_COMPUTE_SHADER_BIT, 0, 1, &BarrierAfter, 0, nullptr, 0, nullptr);
 	}
 };
 
