Fixed: Pixel 6 Pro crashes on start up. Disabled a workaround for 4 swapchain images

diff --git a/Engine/Source/Runtime/VulkanRHI/Private/VulkanSwapChain.cpp b/Engine/Source/Runtime/VulkanRHI/Private/VulkanSwapChain.cpp
--- a/Engine/Source/Runtime/VulkanRHI/Private/VulkanSwapChain.cpp
+++ b/Engine/Source/Runtime/VulkanRHI/Private/VulkanSwapChain.cpp
@@ -25,12 +25,7 @@ static FAutoConsoleVariableRef CVarVulkanKeepSwapChain(
 	ECVF_RenderThreadSafe
 );
 
-#if PLATFORM_ANDROID
-// On Android swapchain may include more images than we have requested and rendering to those images does not work correctly, so we just ignore them
-int32 GVulkanSwapChainIgnoreExtraImages = 1;
-#else
 int32 GVulkanSwapChainIgnoreExtraImages = 0;
-#endif
 static FAutoConsoleVariableRef CVarVulkanSwapChainIgnoreExtraImages(
 	TEXT("r.Vulkan.SwapChainIgnoreExtraImages"),
 	GVulkanSwapChainIgnoreExtraImages,
