Set Backbuffer's TexCreate_ResolveTargetable flag for cases when MobileHDR=false and MSAA is enabled

diff --git a/Engine/Source/Runtime/OpenGLDrv/Private/Android/AndroidOpenGL.cpp b/Engine/Source/Runtime/OpenGLDrv/Private/Android/AndroidOpenGL.cpp
--- a/Engine/Source/Runtime/OpenGLDrv/Private/Android/AndroidOpenGL.cpp
+++ b/Engine/Source/Runtime/OpenGLDrv/Private/Android/AndroidOpenGL.cpp
@@ -456,7 +456,7 @@ FOpenGLTexture* PlatformCreateBuiltinBackBuffer(FOpenGLDynamicRHI* OpenGLRHI, ui
 			{ (int32)SizeX, (int32)SizeY },
 			PF_B8G8R8A8,
 			FClearValueBinding::Transparent,
-			TexCreate_RenderTargetable | TexCreate_Presentable
+			TexCreate_RenderTargetable | TexCreate_Presentable | TexCreate_ResolveTargetable
 		);
 
 		return new FOpenGLTexture(CreateDesc);
diff --git a/Engine/Source/Runtime/OpenGLDrv/Private/OpenGLViewport.cpp b/Engine/Source/Runtime/OpenGLDrv/Private/OpenGLViewport.cpp
--- a/Engine/Source/Runtime/OpenGLDrv/Private/OpenGLViewport.cpp
+++ b/Engine/Source/Runtime/OpenGLDrv/Private/OpenGLViewport.cpp
@@ -320,7 +320,7 @@ void FOpenGLViewport::Resize(uint32 InSizeX,uint32 InSizeY,bool bInIsFullscreen)
 				{ (int32)InSizeX, (int32)InSizeY },
 				PixelFormat,
 				FClearValueBinding::Transparent,
-				TexCreate_RenderTargetable
+				TexCreate_RenderTargetable | TexCreate_ResolveTargetable
 			));
 		}
 
diff --git a/Engine/Source/Runtime/OpenGLDrv/Private/Windows/OpenGLWindows.cpp b/Engine/Source/Runtime/OpenGLDrv/Private/Windows/OpenGLWindows.cpp
--- a/Engine/Source/Runtime/OpenGLDrv/Private/Windows/OpenGLWindows.cpp
+++ b/Engine/Source/Runtime/OpenGLDrv/Private/Windows/OpenGLWindows.cpp
@@ -1008,7 +1008,7 @@ FOpenGLTexture* PlatformCreateBuiltinBackBuffer(FOpenGLDynamicRHI* OpenGLRHI, ui
 			{ (int32)SizeX, (int32)SizeY },
 			PF_B8G8R8A8,
 			FClearValueBinding::Transparent,
-			TexCreate_RenderTargetable | TexCreate_Presentable
+			TexCreate_RenderTargetable | TexCreate_Presentable | TexCreate_ResolveTargetable
 		));
 	}
 
diff --git a/Engine/Source/Runtime/VulkanRHI/Private/VulkanViewport.cpp b/Engine/Source/Runtime/VulkanRHI/Private/VulkanViewport.cpp
--- a/Engine/Source/Runtime/VulkanRHI/Private/VulkanViewport.cpp
+++ b/Engine/Source/Runtime/VulkanRHI/Private/VulkanViewport.cpp
@@ -667,7 +667,7 @@ void FVulkanViewport::CreateSwapchain(FVulkanSwapChainRecreateInfo* RecreateInfo
 		
 		Device->GetImmediateContext().GetCommandBufferManager()->SubmitUploadCmdBuffer();
 
-		RHIBackBuffer = new FVulkanBackBuffer(*Device, this, PixelFormat, SizeX, SizeY, TexCreate_RenderTargetable | TexCreate_ShaderResource);
+		RHIBackBuffer = new FVulkanBackBuffer(*Device, this, PixelFormat, SizeX, SizeY, TexCreate_RenderTargetable | TexCreate_ShaderResource | TexCreate_ResolveTargetable);
 	}
 	else
 	{
@@ -697,7 +697,7 @@ void FVulkanViewport::CreateSwapchain(FVulkanSwapChainRecreateInfo* RecreateInfo
 			{ (int32)BackBufferSizeX, (int32)BackBufferSizeY },
 			PixelFormat,
 			FClearValueBinding::None,
-			TexCreate_RenderTargetable | TexCreate_ShaderResource,
+			TexCreate_RenderTargetable | TexCreate_ShaderResource | TexCreate_ResolveTargetable,
 			1,
 			1,
 			0,
