Remove vulkan overly conservative barriers on UB upload introduced in CL#20006576
Mark mobile directional light UBs as SingleFrame since they are updated every frame and should use ring buffer

diff --git a/Engine/Source/Runtime/Renderer/Private/RendererScene.cpp b/Engine/Source/Runtime/Renderer/Private/RendererScene.cpp
--- a/Engine/Source/Runtime/Renderer/Private/RendererScene.cpp
+++ b/Engine/Source/Runtime/Renderer/Private/RendererScene.cpp
@@ -1071,7 +1071,8 @@ void FPersistentUniformBuffers::Initialize()
 	FMobileDirectionalLightShaderParameters MobileDirectionalLightShaderParameters = {};
 	for (int32 Index = 0; Index < UE_ARRAY_COUNT(MobileDirectionalLightUniformBuffers); ++Index)
 	{
-		MobileDirectionalLightUniformBuffers[Index] = TUniformBufferRef<FMobileDirectionalLightShaderParameters>::CreateUniformBufferImmediate(MobileDirectionalLightShaderParameters, UniformBuffer_MultiFrame, EUniformBufferValidation::None);
+		// UniformBuffer_SingleFrame here is an optimization as this buffer gets uploaded everyframe
+		MobileDirectionalLightUniformBuffers[Index] = TUniformBufferRef<FMobileDirectionalLightShaderParameters>::CreateUniformBufferImmediate(MobileDirectionalLightShaderParameters, UniformBuffer_SingleFrame, EUniformBufferValidation::None);
 	}
 
 	const FMobileReflectionCaptureShaderParameters* DefaultMobileSkyReflectionParameters = (const FMobileReflectionCaptureShaderParameters*)GDefaultMobileReflectionCaptureUniformBuffer.GetContents();
diff --git a/Engine/Source/Runtime/VulkanRHI/Private/VulkanUniformBuffer.cpp b/Engine/Source/Runtime/VulkanRHI/Private/VulkanUniformBuffer.cpp
--- a/Engine/Source/Runtime/VulkanRHI/Private/VulkanUniformBuffer.cpp
+++ b/Engine/Source/Runtime/VulkanRHI/Private/VulkanUniformBuffer.cpp
@@ -88,13 +88,7 @@ static void UpdateUniformBufferHelper(FVulkanCommandListContext& Context, FVulka
 		VkBuffer UBBuffer = VulkanUniformBuffer->Allocation.GetBufferHandle();
 		VkBuffer LockHandle = LockInfo.Allocation.GetBufferHandle();
 
-		VkMemoryBarrier BarrierBefore = { VK_STRUCTURE_TYPE_MEMORY_BARRIER, nullptr, VK_ACCESS_UNIFORM_READ_BIT, VK_ACCESS_TRANSFER_WRITE_BIT };
-		VulkanRHI::vkCmdPipelineBarrier(CmdBuffer->GetHandle(), VK_PIPELINE_STAGE_FRAGMENT_SHADER_BIT, VK_PIPELINE_STAGE_TRANSFER_BIT, 0, 1, &BarrierBefore, 0, nullptr, 0, nullptr);
-
 		VulkanRHI::vkCmdCopyBuffer(CmdBuffer->GetHandle(), LockHandle, UBBuffer, 1, &Region);
-
-		VkMemoryBarrier BarrierAfter = { VK_STRUCTURE_TYPE_MEMORY_BARRIER, nullptr, VK_ACCESS_TRANSFER_WRITE_BIT, VK_ACCESS_UNIFORM_READ_BIT };
-		VulkanRHI::vkCmdPipelineBarrier(CmdBuffer->GetHandle(), VK_PIPELINE_STAGE_TRANSFER_BIT, VK_PIPELINE_STAGE_VERTEX_SHADER_BIT|VK_PIPELINE_STAGE_COMPUTE_SHADER_BIT, 0, 1, &BarrierAfter, 0, nullptr, 0, nullptr);
 	}
 };
 
