VulkanRHI:  Integrate and adapt 19922484 from Release-5.0 (missing lock, deadlock and heap-use-after-free).

Misc VulkanMemory clean up:
- Move locks into the classes where they apply, no change in functionality: GStagingLock, GDeviceMemLock, GFenceLock, GResourceHeapLock (renamed to UsedFreeBufferAllocationsLock)
- Move GResourceLock into FVulkanResourceHeap, rename to PagesLock, one lock per heap now.
- Rename the SubresourceAllocator CS to make it easier to identify in all the locking.
- Rework locking around ReleaseSubresourceAllocator (when we JoinFreeBlocks)

Other WIP clean up:
- Remove FVulkanEvictable from StagingBuffer and UniformRingBuffer since they couldn't move and evict. WIP
- Get rid of unused code (like OnFullDefrag).
- Initialize FenceSignaledCounter to MAX_uint64 to match Clear state

diff --git a/Engine/Source/Runtime/VulkanRHI/Private/VulkanMemory.cpp b/Engine/Source/Runtime/VulkanRHI/Private/VulkanMemory.cpp
--- a/Engine/Source/Runtime/VulkanRHI/Private/VulkanMemory.cpp
+++ b/Engine/Source/Runtime/VulkanRHI/Private/VulkanMemory.cpp
@@ -293,14 +293,6 @@ static FAutoConsoleVariableRef CVarVulkanDefragOnce(
 	ECVF_RenderThreadSafe
 );
 
-int32 GVulkanDefragRepeat = 0;
-static FAutoConsoleVariableRef CVarVulkanDefragRepeat (
-	TEXT("r.Vulkan.DefragRepeatDebug"),
-	GVulkanDefragRepeat ,
-	TEXT("Set to 1 to test defrag, repeatedly"),
-	ECVF_RenderThreadSafe
-);
-
 
 
 static float GVulkanEvictionLimitPercentage = 80.f;
@@ -465,12 +457,6 @@ namespace VulkanRHI
 	constexpr uint32 FMemoryManager::PoolSizes[(int32)FMemoryManager::EPoolSizes::SizesCount];
 	constexpr uint32 FMemoryManager::BufferSizes[(int32)FMemoryManager::EPoolSizes::SizesCount + 1];
 
-	static FCriticalSection GResourceLock;
-	static FCriticalSection GStagingLock;
-	static FCriticalSection GDeviceMemLock;
-	static FCriticalSection GFenceLock;
-	static FCriticalSection GResourceHeapLock;
-
 
 	const TCHAR* VulkanAllocationTypeToString(EVulkanAllocationType Type)
 	{
@@ -812,10 +798,10 @@ namespace VulkanRHI
 		{
 			OutAllocated = HeapInfos[PrimaryHeapIndex].UsedSize;
 			if (Device->GetOptionalExtensions().HasMemoryBudget && (FPlatformTime::Seconds() - MemoryUpdateTime >= 1.0))
-				{
-					MemoryUpdateTime = FPlatformTime::Seconds();
-					UpdateMemoryProperties();
-				}			
+			{
+				MemoryUpdateTime = FPlatformTime::Seconds();
+				UpdateMemoryProperties();
+			}			
 			OutLimit = GetBaseHeapSize(PrimaryHeapIndex);
 	}
 	}
@@ -830,7 +816,7 @@ namespace VulkanRHI
 	FDeviceMemoryAllocation* FDeviceMemoryManager::Alloc(bool bCanFail, VkDeviceSize AllocationSize, uint32 MemoryTypeIndex, void* DedicatedAllocateInfo, float Priority, bool bExternal, const char* File, uint32 Line)
 	{
 		SCOPED_NAMED_EVENT(FDeviceMemoryManager_Alloc, FColor::Cyan);
-		FScopeLock Lock(&GDeviceMemLock);
+		FScopeLock Lock(&DeviceMemLock);
 
 		if(!DedicatedAllocateInfo)
 		{
@@ -887,10 +873,10 @@ namespace VulkanRHI
 
 #if VULKAN_SUPPORTS_MEMORY_PRIORITY
 		VkMemoryPriorityAllocateInfoEXT Prio;
-		ZeroVulkanStruct(Prio, VK_STRUCTURE_TYPE_MEMORY_PRIORITY_ALLOCATE_INFO_EXT);
-		Prio.priority = Priority;
 		if (Device->GetOptionalExtensions().HasMemoryPriority)
 		{
+			ZeroVulkanStruct(Prio, VK_STRUCTURE_TYPE_MEMORY_PRIORITY_ALLOCATE_INFO_EXT);
+			Prio.priority = Priority;
 			Info.pNext = &Prio;
 		}
 #endif
@@ -1032,7 +1018,7 @@ namespace VulkanRHI
 	void FDeviceMemoryManager::Free(FDeviceMemoryAllocation*& Allocation)
 	{
 		SCOPED_NAMED_EVENT(FDeviceMemoryManager_Free, FColor::Cyan);
-		FScopeLock Lock(&GDeviceMemLock);
+		FScopeLock Lock(&DeviceMemLock);
 
 		check(Allocation);
 		check(Allocation->Handle != VK_NULL_HANDLE);
@@ -1109,7 +1095,7 @@ namespace VulkanRHI
 	}
 	void FDeviceMemoryManager::TrimMemory(bool bFullTrim)
 	{
-		FScopeLock Lock(&GDeviceMemLock);
+		FScopeLock Lock(&DeviceMemLock);
 		//blocks are always freed after being reserved for FrameThresholdFull frames.
 		const uint32 FrameThresholdFull = 100; 
 		// After being held for FrameThresholdPartial frames, only LargeThresholdPartial/SmallThresholdPartial pages are kept reserved
@@ -1682,7 +1668,7 @@ namespace VulkanRHI
 
 	void FVulkanResourceHeap::FreePage(FVulkanSubresourceAllocator* InPage)
 	{
-		FScopeLock ScopeLock(&GResourceLock);
+		FScopeLock ScopeLock(&PagesLock);
 		check(InPage->JoinFreeBlocks());
 		int32 Index = -1;
 
@@ -1715,18 +1701,19 @@ namespace VulkanRHI
 			check(!Pages.Find(InPage, Index));
 		}
 
-			ReleasePage(InPage);
-		}
+		ReleasePage(InPage);
+	}
 
 	uint64 FVulkanResourceHeap::EvictOne(FVulkanDevice& Device)
 	{
+		FScopeLock ScopeLock(&PagesLock);
 		for(int32 Index = MAX_BUCKETS - 2; Index >= 0; Index--)
 		{
 			TArray<FVulkanSubresourceAllocator*>& Pages = ActivePages[Index];
 			for (int32 Index2 = 0; Index2 < Pages.Num(); ++Index2)
 			{
 				FVulkanSubresourceAllocator* Allocator = Pages[Index2];
-				if(!Allocator->bIsEvicting && Allocator->GetSubresourceAllocatorFlags() & VulkanAllocationFlagsCanEvict)
+				if (!Allocator->bIsEvicting && Allocator->GetSubresourceAllocatorFlags() & VulkanAllocationFlagsCanEvict)
 				{
 					return Allocator->EvictToHost(Device);
 				}
@@ -1778,9 +1765,7 @@ namespace VulkanRHI
 	{
 		SCOPED_NAMED_EVENT(FVulkanSubresourceAllocator_DefragTick, FColor::Cyan);
 
-		FScopeLock ScopeLock(&GResourceLock);
-		const uint32 DefragMask = VulkanAllocationFlagsCanEvict;
-		const uint32 DefragValue = 0;
+		FScopeLock ScopeLock(&PagesLock);
 
 		FVulkanSubresourceAllocator* CurrentDefragTarget = 0;
 		//Continue if defrag currently in progress.
@@ -1847,7 +1832,7 @@ namespace VulkanRHI
 					}
 				}
 			}
-			if(PotentialDefrag.Num())
+			if (PotentialDefrag.Num())
 			{
 				PotentialDefrag.Sort([](const FPotentialDefragMove& LHS, const FPotentialDefragMove& RHS)
 					{
@@ -1860,23 +1845,23 @@ namespace VulkanRHI
 		}
 
 
-		if(CurrentDefragTarget)
+		if (CurrentDefragTarget)
 		{
 			uint32 MaxSize = CurrentDefragTarget->GetMaxSize();
 			uint32 UsedSize = CurrentDefragTarget->GetUsedSize();
 			uint32 FreeSpace = MaxSize - UsedSize;
 
-			if(GVulkanLogDefrag)
+			if (GVulkanLogDefrag)
 			{
 				VULKAN_LOGMEMORY(TEXT("Defragging heap [%d] Max %6.2fMB Free %6.2fMB   %2d"), CurrentDefragTarget->AllocatorIndex, MaxSize / (1024.f*1024.f), FreeSpace / (1024.f * 1024.f), DefragCountDown);
 			}
-			if(CurrentDefragTarget->DefragTick(Device, Context, this, Count) > 0)
+			if (CurrentDefragTarget->DefragTick(Device, Context, this, Count) > 0)
 			{
 				DefragCountDown = 3;
 			}
 			else
 			{
-				if(0 == --DefragCountDown)
+				if (0 == --DefragCountDown)
 				{
 					if(GVulkanDefragAutoPause)
 					{
@@ -1886,38 +1871,6 @@ namespace VulkanRHI
 				}
 			}
 		}
-		else
-		{
-			if(GVulkanDefragOnce > 0 || GVulkanDefragRepeat > 0)
-			{
-				//Note: This code-path is not yet active. its intended for use in "panic" defrag scenarios.
-				uint32 Heap = GVulkanDefragRepeat ? GVulkanDefragRepeat :  GVulkanDefragOnce;
-				TArray<FVulkanSubresourceAllocator*> Allocators;
-				FVulkanSubresourceAllocator* PreferredAllocator = 0;
-				for (TArray<FVulkanSubresourceAllocator*>& Pages : ActivePages)
-				{
-					for (int32 Index = 0; Index < Pages.Num(); ++Index)
-					{
-						FVulkanSubresourceAllocator* Allocator = Pages[Index];
-						if (0 == (Allocator->GetSubresourceAllocatorFlags() & VulkanAllocationFlagsCanEvict))
-						{
-							if(Allocator->AllocatorIndex == Heap)
-							{
-								PreferredAllocator = Allocator;
-							}
-							Allocators.Add(Allocator);
-						}
-					}
-				}
-				if(!PreferredAllocator)
-				{
-					uint32 Index = rand() % Allocators.Num();
-					PreferredAllocator= Allocators[Index];
-				}
-				PreferredAllocator->DefragFull(Device, Context, this);
-
-			}
-		}
 	}
 
 
@@ -2011,7 +1964,7 @@ namespace VulkanRHI
 	bool FVulkanResourceHeap::AllocateResource(FVulkanAllocation& OutAllocation, FVulkanEvictable* AllocationOwner, EType Type, uint32 Size, uint32 Alignment, bool bMapAllocation, bool bForceSeparateAllocation, EVulkanAllocationMetaType MetaType, bool bExternal, const char* File, uint32 Line)
 	{
 		SCOPED_NAMED_EVENT(FResourceHeap_AllocateResource, FColor::Cyan);
-		FScopeLock ScopeLock(&GResourceLock);
+		FScopeLock ScopeLock(&PagesLock);
 
 		bForceSeparateAllocation = bForceSeparateAllocation || GVulkanSingleAllocationPerResource != 0;
 
@@ -2027,7 +1980,7 @@ namespace VulkanRHI
 		{
 			AllocationFlags |= VulkanAllocationFlagsMapped;
 		}
-
+		
 		uint32 AllocationSize;
 
 		if (!bForceSeparateAllocation)
@@ -2104,7 +2057,7 @@ namespace VulkanRHI
 	bool FVulkanResourceHeap::AllocateDedicatedImage(FVulkanAllocation& OutAllocation, FVulkanEvictable* AllocationOwner, VkImage Image, uint32 Size, uint32 Alignment, EVulkanAllocationMetaType MetaType, bool bExternal, const char* File, uint32 Line)
 	{
 #if VULKAN_SUPPORTS_DEDICATED_ALLOCATION
-		FScopeLock ScopeLock(&GResourceLock);
+		FScopeLock ScopeLock(&PagesLock);
 
 		uint32 AllocationSize = Size;
 
@@ -2306,7 +2259,7 @@ namespace VulkanRHI
 	{
 		TArray<FVulkanSubresourceAllocator*> BufferAllocationsToRelease;
 		{
-			FScopeLock ScopeLock(&GResourceHeapLock);
+			FScopeLock ScopeLock(&UsedFreeBufferAllocationsLock);
 			for (auto& FreeAllocations : FreeBufferAllocations)
 			{
 				for (int32 Index = 0; Index < FreeAllocations.Num(); ++Index)
@@ -2386,12 +2339,12 @@ namespace VulkanRHI
 		uint64 BestDefragHeapSize = 0;
 
 		{
-			FScopeLock ScopeLock(&GResourceLock);
-			uint32 Idx = 0;
 			for (FVulkanResourceHeap* Heap : ResourceTypeHeaps)
 			{
 				if (Heap->HeapIndex == PrimaryHeapIndex)
 				{
+					FScopeLock ScopeLock(&Heap->PagesLock);
+
 					uint64 UsedSize = Heap->UsedMemory;
 					if (CanDefragHeap(Heap) && BestDefragHeapSize < UsedSize)
 					{
@@ -2403,9 +2356,7 @@ namespace VulkanRHI
 						BestEvictHeap = Heap;
 						BestEvictHeapSize = UsedSize;
 					}
-
 				}
-				Idx++;
 			}
 		}
 
@@ -2444,30 +2395,30 @@ namespace VulkanRHI
 	{
 		SCOPED_NAMED_EVENT(FResourceHeapManager_FreeVulkanAllocationPooledBuffer, FColor::Cyan);
 		DecMetaStats(Allocation.MetaType, Allocation.Size);
-		uint32 Index = Allocation.AllocatorIndex;
-		AllBufferAllocations[Index]->Free(Allocation);
+		const uint32 Index = Allocation.AllocatorIndex;
+		GetSubresourceAllocator(Index)->Free(Allocation);
 	}
 	void FMemoryManager::FreeVulkanAllocationBuffer(FVulkanAllocation& Allocation)
 	{
 		SCOPED_NAMED_EVENT(FResourceHeapManager_FreeVulkanAllocationBuffer, FColor::Cyan);
 		DecMetaStats(Allocation.MetaType, Allocation.Size);
-		uint32 Index = Allocation.AllocatorIndex;
-		AllBufferAllocations[Index]->Free(Allocation);
+		const uint32 Index = Allocation.AllocatorIndex;
+		GetSubresourceAllocator(Index)->Free(Allocation);
 	}
 
 	void FMemoryManager::FreeVulkanAllocationImage(FVulkanAllocation& Allocation)
 	{
 		SCOPED_NAMED_EVENT(FResourceHeapManager_FreeVulkanAllocationImage, FColor::Cyan);
 		DecMetaStats(Allocation.MetaType, Allocation.Size);
-		uint32 Index = Allocation.AllocatorIndex;
-		AllBufferAllocations[Index]->Free(Allocation);
+		const uint32 Index = Allocation.AllocatorIndex;
+		GetSubresourceAllocator(Index)->Free(Allocation);
 	}
 	void FMemoryManager::FreeVulkanAllocationImageDedicated(FVulkanAllocation& Allocation)
 	{
 		SCOPED_NAMED_EVENT(FResourceHeapManager_FreeVulkanAllocationImageDedicated, FColor::Cyan);
 		DecMetaStats(Allocation.MetaType, Allocation.Size);
-		uint32 Index = Allocation.AllocatorIndex;
-		AllBufferAllocations[Index]->Free(Allocation);
+		const uint32 Index = Allocation.AllocatorIndex;
+		GetSubresourceAllocator(Index)->Free(Allocation);
 	}
 
 	void FVulkanSubresourceAllocator::SetFreePending(FVulkanAllocation& Allocation)
@@ -2475,7 +2426,7 @@ namespace VulkanRHI
 		check(Allocation.GetType() == Type);
 		check(Allocation.AllocatorIndex == GetAllocatorIndex());
 		{
-			FScopeLock ScopeLock(&CS);
+			FScopeLock ScopeLock(&SubresourceAllocatorCS);
 			FVulkanAllocationInternal& Data = InternalData[Allocation.AllocationIndex];
 			Data.State = FVulkanAllocationInternal::EFREEPENDING;
 		}
@@ -2489,7 +2440,7 @@ namespace VulkanRHI
 		check(Allocation.AllocatorIndex == GetAllocatorIndex());
 		bool bTryFree = false;
 		{
-			FScopeLock ScopeLock(&CS);
+			FScopeLock ScopeLock(&SubresourceAllocatorCS);
 			FreeCalls++;
 			uint32 AllocationOffset;
 			uint32 AllocationSize;
@@ -2529,14 +2480,10 @@ namespace VulkanRHI
 				bTryFree = true; //cannot free here as it will cause incorrect lock ordering
 			}
 		}
-		if(bTryFree)
+
+		if (bTryFree)
 		{
-			FScopeLock ScopeLockResourceheap(&GResourceHeapLock);
-			FScopeLock ScopeLockResource(&GResourceLock);
-			if(JoinFreeBlocks()) //someone could've allocated memory before we claimed this lock.
-			{
-				Owner->ReleaseSubresourceAllocator(this);
-			}
+			Owner->ReleaseSubresourceAllocator(this);
 		}
 	}
 
@@ -2567,9 +2514,7 @@ namespace VulkanRHI
 		}
 		else
 		{
-
-			uint32 Index = Allocation.AllocatorIndex;
-			AllBufferAllocations[Index]->SetFreePending(Allocation);
+			GetSubresourceAllocator(Allocation.AllocatorIndex)->SetFreePending(Allocation);
 			Device->GetDeferredDeletionQueue().EnqueueResourceAllocation(Allocation);
 		}
 		check(!Allocation.HasAllocation());
@@ -2769,7 +2714,7 @@ namespace VulkanRHI
 			Size = PoolSizes[PoolSize];
 		}
 
-		FScopeLock ScopeLock(&GResourceHeapLock);
+		FScopeLock ScopeLock(&UsedFreeBufferAllocationsLock);
 
 		for (int32 Index = 0; Index < UsedBufferAllocations[PoolSize].Num(); ++Index)
 		{
@@ -2868,14 +2813,13 @@ namespace VulkanRHI
 		return false;
 	}
 
-
-
 	void FMemoryManager::RegisterSubresourceAllocator(FVulkanSubresourceAllocator* SubresourceAllocator)
 	{
 		check(SubresourceAllocator->AllocatorIndex == 0xffffffff);
-		if(AllBufferAllocationsFreeListHead != (PTRINT)-1)
+		FRWScopeLock ScopedLock(AllBufferAllocationsLock, SLT_Write);
+		if (AllBufferAllocationsFreeListHead != (PTRINT)-1)
 		{
-			uint32 Index = AllBufferAllocationsFreeListHead;
+			const uint32 Index = AllBufferAllocationsFreeListHead;
 			AllBufferAllocationsFreeListHead = (PTRINT)AllBufferAllocations[Index];
 			SubresourceAllocator->AllocatorIndex = Index;
 			AllBufferAllocations[Index] = SubresourceAllocator;
@@ -2887,35 +2831,45 @@ namespace VulkanRHI
 		}
 
 	}
+
 	void FMemoryManager::UnregisterSubresourceAllocator(FVulkanSubresourceAllocator* SubresourceAllocator)
 	{
-		if(SubresourceAllocator->bIsEvicting)
+		if (SubresourceAllocator->bIsEvicting)
 		{
 			PendingEvictBytes -= SubresourceAllocator->GetMemoryAllocation()->GetSize();
 		}
-		uint32 Index = SubresourceAllocator->AllocatorIndex;
+		FRWScopeLock ScopedLock(AllBufferAllocationsLock, SLT_Write);
+		const uint32 Index = SubresourceAllocator->AllocatorIndex;
 		check(Index != 0xffffffff);
 		AllBufferAllocations[Index] = (FVulkanSubresourceAllocator*)AllBufferAllocationsFreeListHead;
 		AllBufferAllocationsFreeListHead = Index;
 	}
 
-
-	void FMemoryManager::ReleaseSubresourceAllocator(FVulkanSubresourceAllocator* SubresourceAllocator)
+	bool FMemoryManager::ReleaseSubresourceAllocator(FVulkanSubresourceAllocator* SubresourceAllocator)
 	{
-		if(SubresourceAllocator->Type == EVulkanAllocationPooledBuffer)
-		{
-			check(SubresourceAllocator->JoinFreeBlocks());
-			UsedBufferAllocations[SubresourceAllocator->PoolSizeIndex].RemoveSingleSwap(SubresourceAllocator, false);
-			SubresourceAllocator->FrameFreed = GFrameNumberRenderThread;
-			FreeBufferAllocations[SubresourceAllocator->PoolSizeIndex].Add(SubresourceAllocator);
-		}
-		else
+		// Lock everything to make sure no one adds anything while we get rid of it
+		FScopeLock ScopeLockBufferAllocations(&UsedFreeBufferAllocationsLock);
+		FScopeLock ScopeLockPages(&ResourceTypeHeaps[SubresourceAllocator->MemoryTypeIndex]->PagesLock);
+
+		if (SubresourceAllocator->JoinFreeBlocks())
 		{
-			FVulkanResourceHeap* Heap = ResourceTypeHeaps[SubresourceAllocator->MemoryTypeIndex];
-			Heap->FreePage(SubresourceAllocator);
+			if (SubresourceAllocator->Type == EVulkanAllocationPooledBuffer)
+			{
+				UsedBufferAllocations[SubresourceAllocator->PoolSizeIndex].RemoveSingleSwap(SubresourceAllocator, false);
+				SubresourceAllocator->FrameFreed = GFrameNumberRenderThread;
+				FreeBufferAllocations[SubresourceAllocator->PoolSizeIndex].Add(SubresourceAllocator);
+			}
+			else
+			{
+				FVulkanResourceHeap* Heap = ResourceTypeHeaps[SubresourceAllocator->MemoryTypeIndex];
+				Heap->FreePage(SubresourceAllocator);
+			}
+
+			return true;
 		}
+		return false;
 	}
-	//
+
 	bool FMemoryManager::UpdateEvictThreshold(bool bLog)
 	{
 		uint64 PrimaryAllocated, PrimaryLimit;
@@ -3101,7 +3055,7 @@ namespace VulkanRHI
 
 	void FMemoryManager::DumpMemory(bool bFullDump)
 	{
-		FScopeLock ScopeLock(&GResourceHeapLock);
+		FScopeLock ScopeLock(&UsedFreeBufferAllocationsLock);
 		Device->GetDeviceMemoryManager().DumpMemory();
 		VULKAN_LOGMEMORY(TEXT("/******************************************* FMemoryManager ********************************************\\"));
 		VULKAN_LOGMEMORY(TEXT("HEAP DUMP"));
@@ -3694,7 +3648,7 @@ namespace VulkanRHI
 
 	bool FVulkanSubresourceAllocator::JoinFreeBlocks()
 	{
-		FScopeLock ScopeLock(&CS);
+		FScopeLock ScopeLock(&SubresourceAllocatorCS);
 #if !UE_VK_MEMORY_JOIN_FREELIST_ON_THE_FLY
 		FRange::JoinConsecutiveRanges(FreeList);
 #endif
@@ -3857,7 +3811,7 @@ namespace VulkanRHI
 		case EVulkanAllocationBuffer:
 		case EVulkanAllocationImage:
 		case EVulkanAllocationImageDedicated:
-			return Device->GetMemoryManager().AllBufferAllocations[AllocatorIndex];
+			return Device->GetMemoryManager().GetSubresourceAllocator(AllocatorIndex);
 		break;
 		default:
 			check(0);
@@ -3898,8 +3852,8 @@ namespace VulkanRHI
 
 	bool FVulkanSubresourceAllocator::TryAllocate2(FVulkanAllocation& OutAllocation, FVulkanEvictable* AllocationOwner, uint32 InSize, uint32 InAlignment, EVulkanAllocationMetaType InMetaType, const char* File, uint32 Line)
 	{
-		FScopeLock ScopeLock(&CS);
-		if(bIsEvicting||bLocked)
+		FScopeLock ScopeLock(&SubresourceAllocatorCS);
+		if (bIsEvicting || bLocked)
 		{
 			return false;
 		}
@@ -4041,149 +3995,6 @@ namespace VulkanRHI
 		}
 	}
 
-
-	// Fully defrag a page onto a new page
-	// This moves -all- allocations to a new page, and ensures the new page is fully defragged.
-	bool FVulkanSubresourceAllocator::DefragFull(FVulkanDevice& Device, FVulkanCommandListContext& Context, FVulkanResourceHeap* Heap)
-	{
-		TRACE_BOOKMARK(TEXT("DefragFull %d"), AllocatorIndex);
-		SCOPED_NAMED_EVENT(FVulkanSubresourceAllocator_DefragFull, FColor::Cyan);
-		FScopeLock ScopeLock(&CS);
-		LastDefragFrame = GFrameNumberRenderThread;
-
-		#if !UE_BUILD_SHIPPING
-		for (FVulkanAllocationInternal& Alloc : InternalData)
-		{
-			switch(Alloc.State)
-			{
-				case FVulkanAllocationInternal::EALLOCATED:
-					{
-						FVulkanEvictable* EvictableOwner = Alloc.AllocationOwner;
-						if(!EvictableOwner->CanMove())
-						{
-							FGenericPlatformMisc::LowLevelOutputDebugStringf(TEXT("Immovable heap!!\n"));
-							return false;
-						}
-						check(Alloc.MetaType == EVulkanAllocationMetaImageRenderTarget); //defrag only done for RTs, images still pending.
-					}
-					break;
-				default:
-					break;
-			}
-		}
-		#endif
-		FDeviceMemoryManager& DeviceMemoryManager = Owner->GetParent()->GetDeviceMemoryManager();
-
-		FDeviceMemoryAllocation* OldAllocation = MemoryAllocation;
-		uint32 AllocationSize = OldAllocation->GetSize();
-
-		uint32 FreeListSize = FreeList.Num();
-		if(0 == GVulkanDefragRepeat 
-			&& (FreeListSize == 0 ||  // Allocation is fully occupied
-			   (FreeListSize == 1 && FreeList[0].Offset + FreeList[0].Size == AllocationSize)) //Allocation only has 1 allocation and its in the correct place
-		  )
-		{
-			if (GVulkanLogDefrag)
-			{
-				VULKAN_LOGMEMORY(TEXT("Memory already defragged\n"));
-			}
-			return true;
-		}
-
-
-		struct SMoveAllocation
-		{
-			uint32 Index;
-			uint32 AlignedOffset;
-			uint32 AllocatedOffset;
-			uint32 AllocatedSize;
-		};
-		TArray<SMoveAllocation> Moves;
-		uint32 Offset = 0; // Current offset in the new page
-		
-		// Build a list of Src -> Dst Moves that are needed in order to fully defrag onto another page
-		// The actual copies are only enqueued once everything succeeds
-		for(int32 Index = 0; Index < InternalData.Num(); ++Index)
-		{
-			FVulkanAllocationInternal& Alloc  = InternalData[Index];
-			switch(Alloc.State)
-			{
-				case FVulkanAllocationInternal::EALLOCATED:
-					{
-						check(Alloc.MetaType == EVulkanAllocationMetaImageRenderTarget);
-						uint32 AllocatedOffset = Offset;
-						uint32 Size = Alloc.Size;
-						uint32 AlignedOffset = Align(Offset, Alloc.Alignment); //Align to the original alignment requested by the allocation
-						uint32 AlignmentAdjustment = AlignedOffset - Offset;
-						uint32 AllocatedSize = AlignmentAdjustment + Size;
-						Offset += AllocatedSize;
-						
-						//If there are multiple allocations, with high alignment requirements, we can end up with the defragging actually using more than the full memory.
-						//in that case, bail out, and dont defrag.
-						if(Offset > AllocationSize)
-						{
-							return false;
-						}
-						if(!Alloc.AllocationOwner->CanMove())
-						{
-							return false;
-						}
-
-						SMoveAllocation M = { (uint32)Index, AlignedOffset, AllocatedOffset, AllocatedSize } ;
-						Moves.Add(M);
-					}
-					break;
-				case FVulkanAllocationInternal::EFREEPENDING:
-					{
-						//since this has already been queued to be freed, there is no need to move this, so we just tag it to be discarded.
-						LLM_TRACK_VULKAN_HIGH_LEVEL_FREE(Alloc);
-						VULKAN_FREE_TRACK_INFO(Alloc.Track);
-						Alloc.State = FVulkanAllocationInternal::EFREEDISCARDED;
-					}
-					break;
-				case FVulkanAllocationInternal::EUNUSED:
-					break;
-				default:
-					checkf(false, TEXT("Unexpected State %d\n"), Alloc.State);
-					break;
-			}
-		}
-		FDeviceMemoryAllocation* NewMemoryAllocation = DeviceMemoryManager.Alloc(true, AllocationSize, MemoryTypeIndex, nullptr, VULKAN_MEMORY_HIGHEST_PRIORITY, false, __FILE__, __LINE__);
-		if(!NewMemoryAllocation)
-		{
-			Device.GetMemoryManager().DumpMemory();
-			checkNoEntry();
-			return false;
-		}
-		MemoryAllocation = NewMemoryAllocation;
-
-		for(SMoveAllocation& Move : Moves)
-		{
-			FVulkanAllocationInternal& Alloc = InternalData[Move.Index];
-			int64 Delta = Move.AllocatedSize - Alloc.AllocationSize;
-			LLM_TRACK_VULKAN_SPARE_MEMORY_GPU(-Delta);
-			MemoryUsed[Alloc.MetaType] += Delta;
-			Alloc.AllocationOffset = Move.AllocatedOffset;
-			Alloc.AllocationSize = Move.AllocatedSize;
-			FVulkanEvictable* EvictableOwner = Alloc.AllocationOwner;
-			EvictableOwner->OnFullDefrag(Device, Context, Move.AlignedOffset);
-		}
-
-		FreeList.Empty();
-		FRange NewFree;
-		NewFree.Offset = Offset;
-		NewFree.Size = AllocationSize - Offset;
-		FRange::Add(FreeList, NewFree);
-		UsedSize = Offset;
-		if (GVulkanLogDefrag)
-		{
-			VULKAN_LOGMEMORY(TEXT("** DEFRAG Moved %6.2fMB\n"), Offset / (1024.f * 1024.f));
-		}
-		Device.GetDeferredDeletionQueue().EnqueueDeviceAllocation(OldAllocation);
-
-		return true;
-	}
-
 	bool FVulkanSubresourceAllocator::CanDefrag()
 	{
 		for (FVulkanAllocationInternal& Alloc : InternalData)
@@ -4229,7 +4040,7 @@ namespace VulkanRHI
 
 		};
 
-		FScopeLock ScopeLock(&CS);
+		FScopeLock ScopeLock(&SubresourceAllocatorCS);
 
 		//Search for allocations to move to different pages.
 		for (FVulkanAllocationInternal& Alloc : InternalData)
@@ -4290,11 +4101,11 @@ namespace VulkanRHI
 
 	uint64 FVulkanSubresourceAllocator::EvictToHost(FVulkanDevice& Device)
 	{
-		FScopeLock ScopeLock(&CS);
+		FScopeLock ScopeLock(&SubresourceAllocatorCS);
 		bIsEvicting = true;
-		for(FVulkanAllocationInternal& Alloc: InternalData)
+		for (FVulkanAllocationInternal& Alloc : InternalData)
 		{
-			if(Alloc.State == FVulkanAllocationInternal::EALLOCATED)
+			if (Alloc.State == FVulkanAllocationInternal::EALLOCATED)
 			{
 				switch(Alloc.MetaType)
 				{
@@ -4403,7 +4214,7 @@ namespace VulkanRHI
 
 		//#todo-rco: Better locking!
 		{
-			FScopeLock Lock(&GStagingLock);
+			FScopeLock Lock(&StagingLock);
 			for (int32 Index = 0; Index < FreeStagingBuffers.Num(); ++Index)
 			{
 				FFreeEntry& FreeBuffer = FreeStagingBuffers[Index];
@@ -4443,7 +4254,7 @@ namespace VulkanRHI
 		}
 
 		VkMemoryPropertyFlags readTypeFlags = InMemoryReadFlags;
-		if (!Device->GetMemoryManager().AllocateBufferMemory(StagingBuffer->Allocation, StagingBuffer, MemReqs, VK_MEMORY_PROPERTY_HOST_VISIBLE_BIT | readTypeFlags, EVulkanAllocationMetaBufferStaging, false, __FILE__, __LINE__))
+		if (!Device->GetMemoryManager().AllocateBufferMemory(StagingBuffer->Allocation, nullptr, MemReqs, VK_MEMORY_PROPERTY_HOST_VISIBLE_BIT | readTypeFlags, EVulkanAllocationMetaBufferStaging, false, __FILE__, __LINE__))
 		{
 			check(0);
 		}
@@ -4453,7 +4264,7 @@ namespace VulkanRHI
 		//StagingBuffer->ResourceAllocation->BindBuffer(Device, StagingBuffer->Buffer);
 
 		{
-			FScopeLock Lock(&GStagingLock);
+			FScopeLock Lock(&StagingLock);
 			UsedStagingBuffers.Add(StagingBuffer);
 			UsedMemory += StagingBuffer->GetSize();
 			PeakUsedMemory = FMath::Max(UsedMemory, PeakUsedMemory);
@@ -4499,7 +4310,7 @@ namespace VulkanRHI
 		SCOPE_CYCLE_COUNTER(STAT_VulkanStagingBuffer);
 #endif
 
-		FScopeLock Lock(&GStagingLock);
+		FScopeLock Lock(&StagingLock);
 		UsedStagingBuffers.RemoveSingleSwap(StagingBuffer, false);
 
 		if (CmdBuffer)
@@ -4638,7 +4449,7 @@ namespace VulkanRHI
 		SCOPE_CYCLE_COUNTER(STAT_VulkanStagingBuffer);
 #endif
 
-		FScopeLock Lock(&GStagingLock);
+		FScopeLock Lock(&StagingLock);
 		ProcessPendingFreeNoLock(bImmediately, bFreeToOS);
 	}
 
@@ -4677,7 +4488,7 @@ namespace VulkanRHI
 
 	void FFenceManager::Deinit()
 	{
-		FScopeLock Lock(&GFenceLock);
+		FScopeLock Lock(&FenceLock);
 		ensureMsgf(UsedFences.Num() == 0, TEXT("No all fences are done!"));
 		VkDevice DeviceHandle = Device->GetInstanceHandle();
 		for (FFence* Fence : FreeFences)
@@ -4688,7 +4499,7 @@ namespace VulkanRHI
 
 	FFence* FFenceManager::AllocateFence(bool bCreateSignaled)
 	{
-		FScopeLock Lock(&GFenceLock);
+		FScopeLock Lock(&FenceLock);
 		if (FreeFences.Num() != 0)
 		{
 			FFence* Fence = FreeFences[0];
@@ -4710,7 +4521,7 @@ namespace VulkanRHI
 	// Sets it to nullptr
 	void FFenceManager::ReleaseFence(FFence*& Fence)
 	{
-		FScopeLock Lock(&GFenceLock);
+		FScopeLock Lock(&FenceLock);
 		ResetFence(Fence);
 		UsedFences.RemoveSingleSwap(Fence, false);
 #if VULKAN_REUSE_FENCES
@@ -4723,7 +4534,7 @@ namespace VulkanRHI
 
 	void FFenceManager::WaitAndReleaseFence(FFence*& Fence, uint64 TimeInNanoseconds)
 	{
-		FScopeLock Lock(&GFenceLock);
+		FScopeLock Lock(&FenceLock);
 		if (!Fence->IsSignaled())
 		{
 			WaitForFence(Fence, TimeInNanoseconds);
diff --git a/Engine/Source/Runtime/VulkanRHI/Private/VulkanMemory.h b/Engine/Source/Runtime/VulkanRHI/Private/VulkanMemory.h
--- a/Engine/Source/Runtime/VulkanRHI/Private/VulkanMemory.h
+++ b/Engine/Source/Runtime/VulkanRHI/Private/VulkanMemory.h
@@ -6,6 +6,7 @@
 
 #pragma once 
 
+#include "Misc/ScopeRWLock.h"
 
 //enable to store FILE/LINE, and optionally a stacktrace via r.vulkan.backtrace
 #if !UE_BUILD_SHIPPING
@@ -93,7 +94,6 @@ public:
 	virtual bool CanEvict() const { return false; }
 	virtual void Evict(FVulkanDevice& Device) { checkNoEntry(); }
 	virtual void Move(FVulkanDevice& Device, FVulkanCommandListContext& Context, VulkanRHI::FVulkanAllocation& Allocation) { checkNoEntry(); }
-	virtual void OnFullDefrag(FVulkanDevice& Device, FVulkanCommandListContext& Context, uint32 NewOffset) { checkNoEntry(); }
 	virtual FVulkanTexture* GetEvictableTexture() { return nullptr; }
 };
 
@@ -434,6 +434,8 @@ namespace VulkanRHI
 		TArray<FHeapInfo> HeapInfos;
 
 		int32 PrimaryHeapIndex; // memory usage of this heap will decide when to evict.
+
+		FCriticalSection DeviceMemLock;
 	};
 
 	struct FRange
@@ -627,7 +629,6 @@ namespace VulkanRHI
 	protected:
 		void SetIsDefragging(bool bInIsDefragging){ bIsDefragging = bInIsDefragging; }
 		void DumpFullHeap();
-		bool DefragFull(FVulkanDevice& Device, FVulkanCommandListContext& Context, FVulkanResourceHeap* Heap);
 		int32 DefragTick(FVulkanDevice& Device, FVulkanCommandListContext& Context, FVulkanResourceHeap* Heap, uint32 Count);
 		bool CanDefrag();
 		uint64 EvictToHost(FVulkanDevice& Device);
@@ -676,8 +677,7 @@ namespace VulkanRHI
 #endif
 		TArray<FVulkanAllocationInternal> InternalData;
 		int32 InternalFreeList= -1;
-		FCriticalSection CS;
-
+		FCriticalSection SubresourceAllocatorCS;
 	};
 
 
@@ -730,7 +730,7 @@ namespace VulkanRHI
 		uint32 GetPageSize();
 		FMemoryManager* Owner;
 		uint16 MemoryTypeIndex;
-		uint16 HeapIndex;
+		const uint16 HeapIndex;
 
 		bool bIsHostCachedSupported;
 		bool bIsLazilyAllocatedSupported;
@@ -742,6 +742,7 @@ namespace VulkanRHI
 		uint64 UsedMemory;
 		uint32 PageIDCounter;
 
+		FCriticalSection PagesLock;
 		TArray<FVulkanSubresourceAllocator*> ActivePages[MAX_BUCKETS];
 		TArray<FVulkanSubresourceAllocator*> UsedDedicatedImagePages;
 
@@ -789,7 +790,7 @@ namespace VulkanRHI
 
 		void RegisterSubresourceAllocator(FVulkanSubresourceAllocator* SubresourceAllocator);
 		void UnregisterSubresourceAllocator(FVulkanSubresourceAllocator* SubresourceAllocator);
-		void ReleaseSubresourceAllocator(FVulkanSubresourceAllocator* SubresourceAllocator);
+		bool ReleaseSubresourceAllocator(FVulkanSubresourceAllocator* SubresourceAllocator);
 
 
 		void ReleaseFreedPages(FVulkanCommandListContext& Context);
@@ -872,17 +873,23 @@ namespace VulkanRHI
 			return PoolSize;
 		}
 
+		FCriticalSection UsedFreeBufferAllocationsLock;
 		TArray<FVulkanSubresourceAllocator*> UsedBufferAllocations[(int32)EPoolSizes::SizesCount + 1];
 		TArray<FVulkanSubresourceAllocator*> FreeBufferAllocations[(int32)EPoolSizes::SizesCount + 1];
+
+		FRWLock AllBufferAllocationsLock;  // protects against resizing of array (RenderThread<->RHIThread)
 		TArray<FVulkanSubresourceAllocator*> AllBufferAllocations;
 		PTRINT AllBufferAllocationsFreeListHead = (PTRINT)-1;
+		inline FVulkanSubresourceAllocator* GetSubresourceAllocator(const uint32 AllocatorIndex)
+		{
+			FRWScopeLock ScopedReadLock(AllBufferAllocationsLock, SLT_ReadOnly);
+			return AllBufferAllocations[AllocatorIndex];
+		}
 
 		uint64 PendingEvictBytes = 0;
 		bool bIsEvicting = false;
 		bool bWantEviction = false;
 
-
-		friend FVulkanSubresourceAllocator* GetSubresourceAllocator(FVulkanDevice* Device, const FVulkanAllocation& Allocation);
 		friend class FVulkanResourceHeap;
 
 		void ReleaseFreedResources(bool bImmediately);
@@ -904,7 +911,7 @@ namespace VulkanRHI
 		void ProcessPendingUBFrees(bool bForce);
 	};
 
-	class FStagingBuffer : public FVulkanEvictable, public FRefCount
+	class FStagingBuffer : public FRefCount
 	{
 	public:
 		FStagingBuffer(FVulkanDevice* InDevice);
@@ -982,6 +989,8 @@ namespace VulkanRHI
 			TArray<FPendingItems> PendingItems;
 		};
 
+		FCriticalSection StagingLock;
+
 		TArray<FStagingBuffer*> UsedStagingBuffers;
 		TArray<FPendingItemsPerCmdBuffer> PendingFreeStagingBuffers;
 		struct FFreeEntry
@@ -1079,6 +1088,7 @@ namespace VulkanRHI
 
 	protected:
 		FVulkanDevice* Device;
+		FCriticalSection FenceLock;
 		TArray<FFence*> FreeFences;
 		TArray<FFence*> UsedFences;
 
diff --git a/Engine/Source/Runtime/VulkanRHI/Private/VulkanRHI.cpp b/Engine/Source/Runtime/VulkanRHI/Private/VulkanRHI.cpp
--- a/Engine/Source/Runtime/VulkanRHI/Private/VulkanRHI.cpp
+++ b/Engine/Source/Runtime/VulkanRHI/Private/VulkanRHI.cpp
@@ -1820,7 +1820,7 @@ FVulkanRingBuffer::FVulkanRingBuffer(FVulkanDevice* InDevice, uint64 TotalSize,
 	, MinAlignment(0)
 {
 	check(TotalSize <= (uint64)MAX_uint32);
-	InDevice->GetMemoryManager().AllocateBufferPooled(Allocation, this, TotalSize, Usage, MemPropertyFlags, EVulkanAllocationMetaRingBuffer, __FILE__, __LINE__);
+	InDevice->GetMemoryManager().AllocateBufferPooled(Allocation, nullptr, TotalSize, Usage, MemPropertyFlags, EVulkanAllocationMetaRingBuffer, __FILE__, __LINE__);
 	MinAlignment = Allocation.GetBufferAlignment(Device);
 	// Start by wrapping around to set up the correct fence
 	BufferOffset = TotalSize;
diff --git a/Engine/Source/Runtime/VulkanRHI/Private/VulkanResources.h b/Engine/Source/Runtime/VulkanRHI/Private/VulkanResources.h
--- a/Engine/Source/Runtime/VulkanRHI/Private/VulkanResources.h
+++ b/Engine/Source/Runtime/VulkanRHI/Private/VulkanResources.h
@@ -385,7 +385,6 @@ public:
 	bool CanEvict() const override { return false; }
 	void Evict(FVulkanDevice& Device) override; ///evict to system memory
 	void Move(FVulkanDevice& Device, FVulkanCommandListContext& Context, VulkanRHI::FVulkanAllocation& NewAllocation) override; //move to a full new allocation
-	void OnFullDefrag(FVulkanDevice& Device, FVulkanCommandListContext& Context, uint32 NewOffset) override; //called when compacting an allocation. Old image can still be used as a copy source.
 	FVulkanTexture* GetEvictableTexture() override { return this; }
 	
 	void AttachView(VulkanRHI::FVulkanViewBase* View);
@@ -763,7 +762,7 @@ struct FVulkanBufferView : public FRHIResource, public VulkanRHI::FDeviceChild
 	bool bVolatile;
 };
 
-struct FVulkanRingBuffer : public FVulkanEvictable, public VulkanRHI::FDeviceChild
+struct FVulkanRingBuffer : public VulkanRHI::FDeviceChild
 {
 public:
 	FVulkanRingBuffer(FVulkanDevice* InDevice, uint64 TotalSize, VkFlags Usage, VkMemoryPropertyFlags MemPropertyFlags);
@@ -1269,7 +1268,7 @@ public:
 
 protected:
 	FVulkanCmdBuffer*	CmdBuffer = nullptr;
-	uint64				FenceSignaledCounter = 0;
+	uint64				FenceSignaledCounter = MAX_uint64;
 
 	friend class FVulkanCommandListContext;
 };
diff --git a/Engine/Source/Runtime/VulkanRHI/Private/VulkanTexture.cpp b/Engine/Source/Runtime/VulkanRHI/Private/VulkanTexture.cpp
--- a/Engine/Source/Runtime/VulkanRHI/Private/VulkanTexture.cpp
+++ b/Engine/Source/Runtime/VulkanRHI/Private/VulkanTexture.cpp
@@ -2021,33 +2021,6 @@ void FVulkanTexture::Move(FVulkanDevice& InDevice, FVulkanCommandListContext& Co
 	InvalidateViews(InDevice);
 }
 
-void FVulkanTexture::OnFullDefrag(FVulkanDevice& InDevice, FVulkanCommandListContext& Context, uint32 NewOffset)
-{
-	uint64 Size = GetMemorySize();
-	static uint64 TotalSize = 0;
-	TotalSize += Size;
-	if (GVulkanLogDefrag)
-	{
-		UE_LOG(LogVulkanRHI, Display, TEXT("Moving Surface, %p <<-- %p    :::: %s\n"), NewOffset, 42, *GetName().ToString());
-		UE_LOG(LogVulkanRHI, Display, TEXT("Defragged %8.4fkb %8.4fkb   TB %p  :: IMG %p   %-40s\n"), Size / (1024.f), TotalSize / (1024.f), this, Image, *GetName().ToString());
-	}
-
-	{
-		const ETextureCreateFlags UEFlags = GetDesc().Flags;
-		const bool bRenderTarget = EnumHasAnyFlags(UEFlags, TexCreate_RenderTargetable | TexCreate_DepthStencilTargetable | TexCreate_ResolveTargetable);
-		const bool bUAV = EnumHasAnyFlags(UEFlags, TexCreate_UAV);
-		checkf(bRenderTarget || bUAV, TEXT("Surface must be a RenderTarget or a UAV in order to be defragged.  UEFlags=0x%x"), (int32)UEFlags);
-		checkf(Tiling == VK_IMAGE_TILING_OPTIMAL, TEXT("Tiling [%d] is not supported for defrag, only VK_IMAGE_TILING_OPTIMAL"), (int32)Tiling);
-
-		Allocation.Offset = NewOffset;
-		InternalMoveSurface(InDevice, Context, Allocation);
-
-		//note: this exploits that the unmoved image is still bound to the old allocation, which is freed by the caller in this case.
-	}
-
-	InvalidateViews(InDevice);
-}
-
 void FVulkanTexture::Evict(FVulkanDevice& InDevice)
 {
 	check(AliasedTexture == nullptr); //can't evict textures we don't own
