VulkanRHI: Print pipeline stage name from VkCheckpointDataNV in GPU crash report.

#rnx

diff --git a/Engine/Source/Runtime/VulkanRHI/Private/VulkanUtil.cpp b/Engine/Source/Runtime/VulkanRHI/Private/VulkanUtil.cpp
--- a/Engine/Source/Runtime/VulkanRHI/Private/VulkanUtil.cpp
+++ b/Engine/Source/Runtime/VulkanRHI/Private/VulkanUtil.cpp
@@ -695,7 +695,8 @@ void FVulkanGPUProfiler::DumpCrashMarkers(void* BufferData)
 					check(Data[Index].sType == VK_STRUCTURE_TYPE_CHECKPOINT_DATA_NV);
 					uint32 Value = (uint32)(size_t)Data[Index].pCheckpointMarker;
 					const FString* Frame = CachedStrings.Find(Value);
-					UE_LOG(LogVulkanRHI, Error, TEXT("[VK_NV_device_diagnostic_checkpoints] %i: Stage 0x%08x, %s (CRC 0x%x)"), Index, Data[Index].stage, Frame ? *(*Frame) : TEXT("<undefined>"), Value);
+					UE_LOG(LogVulkanRHI, Error, TEXT("[VK_NV_device_diagnostic_checkpoints] %i: Stage %s (0x%08x), %s (CRC 0x%x)"), 
+						Index, VK_TYPE_TO_STRING(VkPipelineStageFlagBits, Data[Index].stage), Data[Index].stage, Frame ? *(*Frame) : TEXT("<undefined>"), Value);
 				}
 				GLog->Panic();
 			}
