VulkanRHI: Move FVulkanRingBuffer code in VulkanUniformBuffer.cpp next to FVulkanUniformBufferUploader (only place where it's used) before refactor.

#rnx

diff --git a/Engine/Source/Runtime/VulkanRHI/Private/VulkanRHI.cpp b/Engine/Source/Runtime/VulkanRHI/Private/VulkanRHI.cpp
--- a/Engine/Source/Runtime/VulkanRHI/Private/VulkanRHI.cpp
+++ b/Engine/Source/Runtime/VulkanRHI/Private/VulkanRHI.cpp
@@ -1814,55 +1814,6 @@ FVulkanRenderPass::~FVulkanRenderPass()
 	RenderPass = VK_NULL_HANDLE;
 }
 
-FVulkanRingBuffer::FVulkanRingBuffer(FVulkanDevice* InDevice, uint64 TotalSize, VkFlags Usage, VkMemoryPropertyFlags MemPropertyFlags)
-	: VulkanRHI::FDeviceChild(InDevice)
-	, BufferSize(TotalSize)
-	, BufferOffset(0)
-	, MinAlignment(0)
-{
-	check(TotalSize <= (uint64)MAX_uint32);
-	InDevice->GetMemoryManager().AllocateBufferPooled(Allocation, nullptr, TotalSize, Usage, MemPropertyFlags, EVulkanAllocationMetaRingBuffer, __FILE__, __LINE__);
-	MinAlignment = Allocation.GetBufferAlignment(Device);
-	// Start by wrapping around to set up the correct fence
-	BufferOffset = TotalSize;
-}
-
-FVulkanRingBuffer::~FVulkanRingBuffer()
-{
-	Device->GetMemoryManager().FreeVulkanAllocation(Allocation);
-}
-
-uint64 FVulkanRingBuffer::WrapAroundAllocateMemory(uint64 Size, uint32 Alignment, FVulkanCmdBuffer* InCmdBuffer)
-{
-	CA_ASSUME(InCmdBuffer != nullptr); // Suppress static analysis warning
-	uint64 AllocationOffset = Align<uint64>(BufferOffset, Alignment);
-	ensure(AllocationOffset + Size > BufferSize);
-
-	// Check to see if we can wrap around the ring buffer
-	if (FenceCmdBuffer)
-	{
-		if (FenceCounter == FenceCmdBuffer->GetFenceSignaledCounterI())
-		{
-			//if (FenceCounter == FenceCmdBuffer->GetSubmittedFenceCounter())
-			{
-				//UE_LOG(LogVulkanRHI, Error, TEXT("Ringbuffer overflow during the same cmd buffer!"));
-			}
-			//else
-			{
-				//UE_LOG(LogVulkanRHI, Error, TEXT("Wrapped around the ring buffer! Waiting for the GPU..."));
-				//Device->GetImmediateContext().GetCommandBufferManager()->WaitForCmdBuffer(FenceCmdBuffer, 0.5f);
-			}
-		}
-	}
-
-	BufferOffset = Size;
-
-	FenceCmdBuffer = InCmdBuffer;
-	FenceCounter = InCmdBuffer->GetSubmittedFenceCounter();
-
-	return 0;
-}
-
 void FVulkanDynamicRHI::SavePipelineCache()
 {
 	FString CacheFile = GetPipelineCacheFilename();
diff --git a/Engine/Source/Runtime/VulkanRHI/Private/VulkanUniformBuffer.cpp b/Engine/Source/Runtime/VulkanRHI/Private/VulkanUniformBuffer.cpp
--- a/Engine/Source/Runtime/VulkanRHI/Private/VulkanUniformBuffer.cpp
+++ b/Engine/Source/Runtime/VulkanRHI/Private/VulkanUniformBuffer.cpp
@@ -313,3 +313,51 @@ FVulkanUniformBufferUploader::~FVulkanUniformBufferUploader()
 }
 
 
+FVulkanRingBuffer::FVulkanRingBuffer(FVulkanDevice* InDevice, uint64 TotalSize, VkFlags Usage, VkMemoryPropertyFlags MemPropertyFlags)
+	: VulkanRHI::FDeviceChild(InDevice)
+	, BufferSize(TotalSize)
+	, BufferOffset(0)
+	, MinAlignment(0)
+{
+	check(TotalSize <= (uint64)MAX_uint32);
+	InDevice->GetMemoryManager().AllocateBufferPooled(Allocation, nullptr, TotalSize, Usage, MemPropertyFlags, EVulkanAllocationMetaRingBuffer, __FILE__, __LINE__);
+	MinAlignment = Allocation.GetBufferAlignment(Device);
+	// Start by wrapping around to set up the correct fence
+	BufferOffset = TotalSize;
+}
+
+FVulkanRingBuffer::~FVulkanRingBuffer()
+{
+	Device->GetMemoryManager().FreeVulkanAllocation(Allocation);
+}
+
+uint64 FVulkanRingBuffer::WrapAroundAllocateMemory(uint64 Size, uint32 Alignment, FVulkanCmdBuffer* InCmdBuffer)
+{
+	CA_ASSUME(InCmdBuffer != nullptr); // Suppress static analysis warning
+	uint64 AllocationOffset = Align<uint64>(BufferOffset, Alignment);
+	ensure(AllocationOffset + Size > BufferSize);
+
+	// Check to see if we can wrap around the ring buffer
+	if (FenceCmdBuffer)
+	{
+		if (FenceCounter == FenceCmdBuffer->GetFenceSignaledCounterI())
+		{
+			//if (FenceCounter == FenceCmdBuffer->GetSubmittedFenceCounter())
+			{
+				//UE_LOG(LogVulkanRHI, Error, TEXT("Ringbuffer overflow during the same cmd buffer!"));
+			}
+			//else
+			{
+				//UE_LOG(LogVulkanRHI, Error, TEXT("Wrapped around the ring buffer! Waiting for the GPU..."));
+				//Device->GetImmediateContext().GetCommandBufferManager()->WaitForCmdBuffer(FenceCmdBuffer, 0.5f);
+			}
+		}
+	}
+
+	BufferOffset = Size;
+
+	FenceCmdBuffer = InCmdBuffer;
+	FenceCounter = InCmdBuffer->GetSubmittedFenceCounter();
+
+	return 0;
+}
