VulkanRHI: Fix crash when trying to read data from TextureCubeArray.

diff --git a/Engine/Source/Runtime/VulkanRHI/Private/VulkanRenderTarget.cpp b/Engine/Source/Runtime/VulkanRHI/Private/VulkanRenderTarget.cpp
--- a/Engine/Source/Runtime/VulkanRHI/Private/VulkanRenderTarget.cpp
+++ b/Engine/Source/Runtime/VulkanRHI/Private/VulkanRenderTarget.cpp
@@ -641,16 +641,18 @@ void FVulkanDynamicRHI::RHIReadSurfaceFloatData(FRHITexture* TextureRHI, FIntRec
 	FVulkanTexture& Surface = *FVulkanTexture::Cast(TextureRHI);
 	const FRHITextureDesc& Desc = Surface.GetDesc();
 
-	if(GIgnoreCPUReads == 1)
+	if (GIgnoreCPUReads == 1)
 	{
 		// Debug: Fill with CPU
 		uint32 NumPixels = 0;
 		switch(Desc.Dimension)
 		{
+		case ETextureDimension::TextureCubeArray:
 		case ETextureDimension::TextureCube:
 			NumPixels = (Desc.Extent.X >> MipIndex) * (Desc.Extent.Y >> MipIndex);
 			break;
 
+		case ETextureDimension::Texture2DArray:
 		case ETextureDimension::Texture2D:
 			NumPixels = (Desc.Extent.X >> MipIndex) * (Desc.Extent.Y >> MipIndex);
 			break;
@@ -670,10 +672,12 @@ void FVulkanDynamicRHI::RHIReadSurfaceFloatData(FRHITexture* TextureRHI, FIntRec
 		FVulkanCmdBuffer* CmdBuffer = Device->GetImmediateContext().GetCommandBufferManager()->GetUploadCmdBuffer();
 		switch (TextureRHI->GetDesc().Dimension)
 		{
+			case ETextureDimension::TextureCubeArray:
 			case ETextureDimension::TextureCube:
 				DoCopyFloat(Device, CmdBuffer, Surface, MipIndex, CubeFace + 6 * ArrayIndex, Rect, OutData);
 				break;
 
+			case ETextureDimension::Texture2DArray:
 			case ETextureDimension::Texture2D:
 				DoCopyFloat(Device, CmdBuffer, Surface, MipIndex, ArrayIndex, Rect, OutData);
 				break;
