VulkanRHI:  Fix issue I introduced in 19291469, make the code match the comment again (thanks Chris!)

diff --git a/Engine/Source/Runtime/VulkanRHI/Private/VulkanSwapChain.cpp b/Engine/Source/Runtime/VulkanRHI/Private/VulkanSwapChain.cpp
--- a/Engine/Source/Runtime/VulkanRHI/Private/VulkanSwapChain.cpp
+++ b/Engine/Source/Runtime/VulkanRHI/Private/VulkanSwapChain.cpp
@@ -915,7 +915,7 @@ void FVulkanSwapChain::CreateQCOMDepthStencil(const FVulkanTexture& InSurface) c
 
 	FRHITextureCreateDesc CreateDesc = FRHITextureCreateDesc::Create2D(
 		TEXT("FVulkanSwapChainQCOM"),
-		{ (int32)Desc.Extent.X, (int32)Desc.Extent.Y }, // Yes, these are intentionally swapped.
+		{ (int32)Desc.Extent.Y, (int32)Desc.Extent.X }, // Yes, these are intentionally swapped.
 		Desc.Format,
 		FClearValueBinding::None,
 		UEFlags,
