Fixed: Local lights flickering on mobile vulkan

diff --git a/Engine/Source/Runtime/VulkanRHI/Private/VulkanBuffer.cpp b/Engine/Source/Runtime/VulkanRHI/Private/VulkanBuffer.cpp
--- a/Engine/Source/Runtime/VulkanRHI/Private/VulkanBuffer.cpp
+++ b/Engine/Source/Runtime/VulkanRHI/Private/VulkanBuffer.cpp
@@ -244,9 +244,9 @@
 	void* Data = nullptr;
 	uint32 DataOffset = 0;
 
-	const bool bStatic = EnumHasAnyFlags(GetUsage(), BUF_Static);
 	const bool bDynamic = EnumHasAnyFlags(GetUsage(), BUF_Dynamic);
 	const bool bVolatile = EnumHasAnyFlags(GetUsage(), BUF_Volatile);
+	const bool bStatic = EnumHasAnyFlags(GetUsage(), BUF_Static) || !(bVolatile || bDynamic);
 	const bool bCPUReadable = EnumHasAnyFlags(GetUsage(), BUF_KeepCPUAccessible);
 	const bool bUAV = EnumHasAnyFlags(GetUsage(), BUF_UnorderedAccess);
 	const bool bSR = EnumHasAnyFlags(GetUsage(), BUF_ShaderResource);
@@ -430,9 +430,9 @@
 
 void FVulkanResourceMultiBuffer::Unlock(bool bFromRenderingThread)
 {
-	const bool bStatic = EnumHasAnyFlags(GetUsage(), BUF_Static);
 	const bool bDynamic = EnumHasAnyFlags(GetUsage(), BUF_Dynamic);
 	const bool bVolatile = EnumHasAnyFlags(GetUsage(), BUF_Volatile);
+	const bool bStatic = EnumHasAnyFlags(GetUsage(), BUF_Static) || !(bVolatile || bDynamic);
 	const bool bCPUReadable = EnumHasAnyFlags(GetUsage(), BUF_KeepCPUAccessible);
 	const bool bSR = EnumHasAnyFlags(GetUsage(), BUF_ShaderResource);
 
