VulkanRHI:  Properly set GRHISupportsWaveOperations at runtime.  Add subgroup properties query to device initialization if the platform's Vulkan instance is >= 1.1.  Update VulkanPC ddpi to report wave op support.  Cleanup: start removing VULKAN_SUPPORTS_PHYSICAL_DEVICE_PROPERTIES2 define.

#rnx

diff --git a/Engine/Config/VulkanPC/DataDrivenPlatformInfo.ini b/Engine/Config/VulkanPC/DataDrivenPlatformInfo.ini
--- a/Engine/Config/VulkanPC/DataDrivenPlatformInfo.ini
+++ b/Engine/Config/VulkanPC/DataDrivenPlatformInfo.ini
@@ -83,9 +83,7 @@ bSupportsIntrinsicWaveOnce = false
 ; //#todo-vulkanddpi
 bSupportsConservativeRasterization = false
 
-
-; //#todo-vulkanddpi
-bSupportsWaveOperations = false
+bSupportsWaveOperations = true
 
 ; //#todo-vulkanddpi
 bRequiresExplicit128bitRT = false
diff --git a/Engine/Source/Runtime/VulkanRHI/Private/VulkanDevice.cpp b/Engine/Source/Runtime/VulkanRHI/Private/VulkanDevice.cpp
--- a/Engine/Source/Runtime/VulkanRHI/Private/VulkanDevice.cpp
+++ b/Engine/Source/Runtime/VulkanRHI/Private/VulkanDevice.cpp
@@ -182,18 +182,22 @@ FVulkanDevice::FVulkanDevice(FVulkanDynamicRHI* InRHI, VkPhysicalDevice InGpu)
 	FMemory::Memzero(FormatProperties);
 	FMemory::Memzero(PixelFormatComponentMapping);
 
-#if VULKAN_SUPPORTS_PHYSICAL_DEVICE_PROPERTIES2
 	ZeroVulkanStruct(GpuIdProps, VK_STRUCTURE_TYPE_PHYSICAL_DEVICE_ID_PROPERTIES_KHR);
+	ZeroVulkanStruct(GpuSubgroupProps, VK_STRUCTURE_TYPE_PHYSICAL_DEVICE_SUBGROUP_PROPERTIES);
 	if (RHI->GetOptionalExtensions().HasKHRGetPhysicalDeviceProperties2)
 	{
 		VkPhysicalDeviceProperties2KHR PhysicalDeviceProperties2;
 		ZeroVulkanStruct(PhysicalDeviceProperties2, VK_STRUCTURE_TYPE_PHYSICAL_DEVICE_PROPERTIES_2_KHR);
 		PhysicalDeviceProperties2.pNext = &GpuIdProps;
+		if (UE_VK_API_VERSION >= VK_API_VERSION_1_1)
+		{
+			// Only consider wave ops on platforms creating a Vulkan 1.1 instance (or greater)
+			GpuIdProps.pNext = &GpuSubgroupProps;
+		}
 		VulkanRHI::vkGetPhysicalDeviceProperties2KHR(Gpu, &PhysicalDeviceProperties2);
 		GpuProps = PhysicalDeviceProperties2.properties;
 	}
 	else
-#endif
 	{
 		VulkanRHI::vkGetPhysicalDeviceProperties(Gpu, &GpuProps);
 	}
@@ -959,7 +963,6 @@ void FVulkanDevice::InitGPU(int32 DeviceIndex)
 	FVulkanDeviceExtensionArray UEExtensions = FVulkanDeviceExtension::GetUESupportedDeviceExtensions(this);
 	TArray<const ANSICHAR*> DeviceLayers = FVulkanDevice::SetupDeviceLayers(Gpu, UEExtensions);
 
-#if VULKAN_SUPPORTS_PHYSICAL_DEVICE_PROPERTIES2
 	if (RHI->GetOptionalExtensions().HasKHRGetPhysicalDeviceProperties2)
 	{
 		// Query advanced features
@@ -1012,7 +1015,6 @@ void FVulkanDevice::InitGPU(int32 DeviceIndex)
 			}
 		}
 	}
-#endif // VULKAN_SUPPORTS_PHYSICAL_DEVICE_PROPERTIES2
 
 	UE_LOG(LogVulkanRHI, Display, TEXT("Using Device %d: Geometry %d BufferAtomic64 %d ImageAtomic64 %d"), DeviceIndex, PhysicalFeatures.geometryShader, OptionalDeviceExtensions.HasKHRShaderAtomicInt64, OptionalDeviceExtensions.HasImageAtomicInt64);
 
diff --git a/Engine/Source/Runtime/VulkanRHI/Private/VulkanDevice.h b/Engine/Source/Runtime/VulkanRHI/Private/VulkanDevice.h
--- a/Engine/Source/Runtime/VulkanRHI/Private/VulkanDevice.h
+++ b/Engine/Source/Runtime/VulkanRHI/Private/VulkanDevice.h
@@ -269,13 +269,18 @@ public:
 		return GpuProps.limits;
 	}
 
-#if VULKAN_SUPPORTS_PHYSICAL_DEVICE_PROPERTIES2
 	inline const VkPhysicalDeviceIDPropertiesKHR& GetDeviceIdProperties() const
 	{
 		check(RHI->GetOptionalExtensions().HasKHRGetPhysicalDeviceProperties2);
 		return GpuIdProps;
 	}
 
+	inline const VkPhysicalDeviceSubgroupProperties& GetDeviceSubgroupProperties() const
+	{
+		check(RHI->GetOptionalExtensions().HasKHRGetPhysicalDeviceProperties2);
+		return GpuSubgroupProps;
+	}
+
 #if VULKAN_RHI_RAYTRACING
 	inline const FRayTracingProperties& GetRayTracingProperties() const
 	{
@@ -286,7 +291,6 @@ public:
 	void InitializeRayTracing();
 	void CleanUpRayTracing();
 #endif // VULKAN_RHI_RAYTRACING
-#endif // VULKAN_SUPPORTS_PHYSICAL_DEVICE_PROPERTIES2
 
 #if VULKAN_SUPPORTS_VALIDATION_CACHE
 	inline VkValidationCacheEXT GetValidationCache() const
@@ -518,14 +522,13 @@ private:
 	TArray<VkPhysicalDeviceFragmentShadingRateKHR> FragmentShadingRates;
 #endif
 
-#if VULKAN_SUPPORTS_PHYSICAL_DEVICE_PROPERTIES2
 	VkPhysicalDeviceIDPropertiesKHR GpuIdProps;
+	VkPhysicalDeviceSubgroupProperties GpuSubgroupProps;
 
 #if VULKAN_RHI_RAYTRACING
 	FRayTracingProperties RayTracingProperties;
 	FVulkanBasicRaytracingPipeline* BasicRayTracingPipeline = nullptr;
 #endif // VULKAN_RHI_RAYTRACING
-#endif // VULKAN_SUPPORTS_PHYSICAL_DEVICE_PROPERTIES2
 
 	VkPhysicalDeviceFeatures PhysicalFeatures;
 
diff --git a/Engine/Source/Runtime/VulkanRHI/Private/VulkanRHI.cpp b/Engine/Source/Runtime/VulkanRHI/Private/VulkanRHI.cpp
--- a/Engine/Source/Runtime/VulkanRHI/Private/VulkanRHI.cpp
+++ b/Engine/Source/Runtime/VulkanRHI/Private/VulkanRHI.cpp
@@ -856,6 +856,11 @@ void FVulkanDynamicRHI::InitInstance()
 			GVulkanDeviceShaderStageBits|= VK_PIPELINE_STAGE_GEOMETRY_SHADER_BIT;
 		}
 		
+		// Check for wave ops support (only filled on platforms creating Vulkan 1.1 or greater instances)
+		const VkSubgroupFeatureFlags RequiredSubgroupFlags = VK_SUBGROUP_FEATURE_BASIC_BIT | VK_SUBGROUP_FEATURE_VOTE_BIT | VK_SUBGROUP_FEATURE_ARITHMETIC_BIT | VK_SUBGROUP_FEATURE_BALLOT_BIT;
+		GRHISupportsWaveOperations = VKHasAllFlags(Device->GetDeviceSubgroupProperties().supportedStages, GVulkanDeviceShaderStageBits) &&
+			VKHasAllFlags(Device->GetDeviceSubgroupProperties().supportedOperations, RequiredSubgroupFlags);
+
 		if (GGPUCrashDebuggingEnabled && !Device->GetOptionalExtensions().HasGPUCrashDumpExtensions())
 		{
 			UE_LOG(LogVulkanRHI, Warning, TEXT("Tried to enable GPU crash debugging but no extension found! Will use local tracepoints."));
