VulkanRHI: Match the behavior of other RHIs and submit commands on RHIBlockUntilGPUIdle.  Fix "RHI Clear" engine tests for Vulkan.

#rnx

diff --git a/Engine/Source/Runtime/VulkanRHI/Private/VulkanCommands.cpp b/Engine/Source/Runtime/VulkanRHI/Private/VulkanCommands.cpp
--- a/Engine/Source/Runtime/VulkanRHI/Private/VulkanCommands.cpp
+++ b/Engine/Source/Runtime/VulkanRHI/Private/VulkanCommands.cpp
@@ -713,6 +713,7 @@ bool FVulkanDynamicRHI::RHIIsRenderingSuspended()
 
 void FVulkanDynamicRHI::RHIBlockUntilGPUIdle()
 {
+	Device->SubmitCommandsAndFlushGPU();
 	Device->WaitUntilIdle();
 }
 
