VulkanRHI: Texture unlock was altering texture layout in a way high level code couldn't know.  Only return to SHADER_READ_ONLY if we can't return to previous layout (UNDEFINED), otherwise return the texture to how it was before the call.  Fixes Engine Test "RHI Copy Texture".

#rnx

diff --git a/Engine/Source/Runtime/VulkanRHI/Private/VulkanTexture.cpp b/Engine/Source/Runtime/VulkanRHI/Private/VulkanTexture.cpp
--- a/Engine/Source/Runtime/VulkanRHI/Private/VulkanTexture.cpp
+++ b/Engine/Source/Runtime/VulkanRHI/Private/VulkanTexture.cpp
@@ -174,8 +174,16 @@ inline void FVulkanTexture::InternalLockWrite(FVulkanCommandListContext& Context
 
 	VulkanRHI::vkCmdCopyBufferToImage(StagingCommandBuffer, StagingBuffer->GetHandle(), Surface->Image, VK_IMAGE_LAYOUT_TRANSFER_DST_OPTIMAL, 1, &Region);
 
-	// Transition the subresource layouts from the copy state to a regular read state
-	TrackedTextureLayout.Set(VK_IMAGE_LAYOUT_SHADER_READ_ONLY_OPTIMAL, SubresourceRange);
+	// Transition the subresource layouts from the copy state to a regular read state if it was undefined, otherwise return it how it was
+	for (uint32 LayerIndex = ImageSubresource.baseArrayLayer; LayerIndex < ImageSubresource.baseArrayLayer + ImageSubresource.layerCount; ++LayerIndex)
+	{
+		if (TrackedTextureLayout.GetSubresLayout(LayerIndex, ImageSubresource.mipLevel) == VK_IMAGE_LAYOUT_UNDEFINED)
+		{
+			const VkImageSubresourceRange SingleLayerSubresourceRange = FVulkanPipelineBarrier::MakeSubresourceRange(ImageSubresource.aspectMask, ImageSubresource.mipLevel, 1, LayerIndex, 1);
+			TrackedTextureLayout.Set(VK_IMAGE_LAYOUT_SHADER_READ_ONLY_OPTIMAL, SingleLayerSubresourceRange);
+		}
+	}
+
 	{
 		FVulkanPipelineBarrier Barrier;
 		Barrier.AddImageLayoutTransition(Surface->Image, SubresourceRange.aspectMask, TransferTextureLayout, TrackedTextureLayout);
