VulkanRHI: Stop including CPU drivers by default (leads to a lot of bugs and confusion when the discrete GPUs doesn't get picked up and llvmpipe is used).  Add a command line argument to explicitely allow them in the device list.

#rnx

diff --git a/Engine/Source/Runtime/VulkanRHI/Private/VulkanRHI.cpp b/Engine/Source/Runtime/VulkanRHI/Private/VulkanRHI.cpp
--- a/Engine/Source/Runtime/VulkanRHI/Private/VulkanRHI.cpp
+++ b/Engine/Source/Runtime/VulkanRHI/Private/VulkanRHI.cpp
@@ -556,6 +556,7 @@ void FVulkanDynamicRHI::SelectAndInitDevice()
 #endif
 
 	UE_LOG(LogVulkanRHI, Display, TEXT("Found %d device(s)"), GpuCount);
+	const bool bAllowCPUDevices = FParse::Param(FCommandLine::Get(), TEXT("AllowCPUDevices"));
 	for (uint32 Index = 0; Index < GpuCount; ++Index)
 	{
 		UE_LOG(LogVulkanRHI, Display, TEXT("Device %d:"), Index);
@@ -563,6 +564,7 @@ void FVulkanDynamicRHI::SelectAndInitDevice()
 		Devices.Add(NewDevice);
 
 		const bool bIsDiscrete = (NewDevice->GetDeviceProperties().deviceType == VK_PHYSICAL_DEVICE_TYPE_DISCRETE_GPU);
+		const bool bIsCPUDevice = (NewDevice->GetDeviceProperties().deviceType == VK_PHYSICAL_DEVICE_TYPE_CPU);
 
 #if VULKAN_ENABLE_DESKTOP_HMD_SUPPORT
 		if (!HmdDevice && HmdGraphicsAdapterLuid != 0 &&
@@ -577,6 +579,10 @@ void FVulkanDynamicRHI::SelectAndInitDevice()
 		{
 			DiscreteDevices.Add({NewDevice, Index});
 		}
+		else if (bIsCPUDevice && !bAllowCPUDevices)
+		{
+			UE_LOG(LogVulkanRHI, Display, TEXT("Skipping device [%s] of type VK_PHYSICAL_DEVICE_TYPE_CPU (add -AllowCPUDevices to your command line to include it)."), ANSI_TO_TCHAR(NewDevice->GetDeviceProperties().deviceName));
+		}
 		else
 		{
 			IntegratedDevices.Add({NewDevice, Index});
@@ -648,8 +654,10 @@ void FVulkanDynamicRHI::SelectAndInitDevice()
 			}
 			else
 			{
-				checkf(0, TEXT("No devices found!"));
-				DeviceIndex = 0;
+				FPlatformMisc::MessageBoxExt(EAppMsgType::Ok, TEXT("Cannot find a compatible Vulkan device."), TEXT("No devices found!"));
+				FPlatformMisc::RequestExitWithStatus(true, 1);
+				// unreachable
+				return;
 			}
 		}
 	}
