VulkanRHI: Rename GVulkanDeviceShaderStageBits to GVulkanDevicePipelineStageBits (since it contains VkPipelineStageFlagBits).  Use VkShaderStageFlagBits to fix Wave Ops check.

diff --git a/Engine/Source/Runtime/VulkanRHI/Private/VulkanBarriers.cpp b/Engine/Source/Runtime/VulkanRHI/Private/VulkanBarriers.cpp
--- a/Engine/Source/Runtime/VulkanRHI/Private/VulkanBarriers.cpp
+++ b/Engine/Source/Runtime/VulkanRHI/Private/VulkanBarriers.cpp
@@ -9,7 +9,7 @@
 FVulkanLayoutManager FVulkanCommandListContext::LayoutManager;
 
 // All shader stages supported by VK device - VK_PIPELINE_STAGE_VERTEX_SHADER_BIT, FRAGMENT etc
-extern uint32 GVulkanDeviceShaderStageBits;
+extern uint32 GVulkanDevicePipelineStageBits;
 
 //
 // The following two functions are used when the RHI needs to do image layout transitions internally.
@@ -302,7 +302,7 @@ static void GetVkStageAndAccessFlags(ERHIAccess RHIAccess, FRHITransitionInfo::E
 
 	if (EnumHasAnyFlags(RHIAccess, ERHIAccess::SRVGraphics))
 	{
-		StageFlags |= (GVulkanDeviceShaderStageBits & ~VK_PIPELINE_STAGE_COMPUTE_SHADER_BIT);
+		StageFlags |= (GVulkanDevicePipelineStageBits & ~VK_PIPELINE_STAGE_COMPUTE_SHADER_BIT);
 		AccessFlags |= VK_ACCESS_SHADER_READ_BIT;
 		Layout = SRVLayout;
 
@@ -323,7 +323,7 @@ static void GetVkStageAndAccessFlags(ERHIAccess RHIAccess, FRHITransitionInfo::E
 
 	if (EnumHasAnyFlags(RHIAccess, ERHIAccess::UAVGraphics))
 	{
-		StageFlags |= (GVulkanDeviceShaderStageBits & ~VK_PIPELINE_STAGE_COMPUTE_SHADER_BIT);
+		StageFlags |= (GVulkanDevicePipelineStageBits & ~VK_PIPELINE_STAGE_COMPUTE_SHADER_BIT);
 		AccessFlags |= VK_ACCESS_SHADER_READ_BIT | VK_ACCESS_SHADER_WRITE_BIT;
 		Layout = VK_IMAGE_LAYOUT_GENERAL;
 
diff --git a/Engine/Source/Runtime/VulkanRHI/Private/VulkanRHI.cpp b/Engine/Source/Runtime/VulkanRHI/Private/VulkanRHI.cpp
--- a/Engine/Source/Runtime/VulkanRHI/Private/VulkanRHI.cpp
+++ b/Engine/Source/Runtime/VulkanRHI/Private/VulkanRHI.cpp
@@ -95,7 +95,7 @@ bool GGPUCrashDebuggingEnabled = false;
 extern TAutoConsoleVariable<int32> GRHIAllowAsyncComputeCvar;
 
 // All shader stages supported by VK device - VK_PIPELINE_STAGE_VERTEX_SHADER_BIT, FRAGMENT etc
-uint32 GVulkanDeviceShaderStageBits = 0;
+uint32 GVulkanDevicePipelineStageBits = 0;
 
 DEFINE_LOG_CATEGORY(LogVulkan)
 
@@ -855,20 +855,23 @@ void FVulkanDynamicRHI::InitInstance()
 		GUseTexture3DBulkDataRHI = false;
 
 		// these are supported by all devices
-		GVulkanDeviceShaderStageBits = 	VK_PIPELINE_STAGE_VERTEX_SHADER_BIT | 
-										VK_PIPELINE_STAGE_FRAGMENT_SHADER_BIT |
-										VK_PIPELINE_STAGE_COMPUTE_SHADER_BIT;
+		GVulkanDevicePipelineStageBits = VK_PIPELINE_STAGE_VERTEX_SHADER_BIT | VK_PIPELINE_STAGE_FRAGMENT_SHADER_BIT | VK_PIPELINE_STAGE_COMPUTE_SHADER_BIT;
+		uint32 VulkanDeviceShaderStageBits = VK_SHADER_STAGE_VERTEX_BIT | VK_SHADER_STAGE_FRAGMENT_BIT | VK_SHADER_STAGE_COMPUTE_BIT;
+
 		// optional shader stages
 		if (Device->GetPhysicalFeatures().geometryShader) 
 		{
-			GVulkanDeviceShaderStageBits|= VK_PIPELINE_STAGE_GEOMETRY_SHADER_BIT;
+			GVulkanDevicePipelineStageBits |= VK_PIPELINE_STAGE_GEOMETRY_SHADER_BIT;
+			VulkanDeviceShaderStageBits |= VK_SHADER_STAGE_GEOMETRY_BIT;
 		}
 		
 		// Check for wave ops support (only filled on platforms creating Vulkan 1.1 or greater instances)
 		const VkSubgroupFeatureFlags RequiredSubgroupFlags = VK_SUBGROUP_FEATURE_BASIC_BIT | VK_SUBGROUP_FEATURE_VOTE_BIT | VK_SUBGROUP_FEATURE_ARITHMETIC_BIT | VK_SUBGROUP_FEATURE_BALLOT_BIT;
-		GRHISupportsWaveOperations = VKHasAllFlags(Device->GetDeviceSubgroupProperties().supportedStages, GVulkanDeviceShaderStageBits) &&
+		GRHISupportsWaveOperations = VKHasAllFlags(Device->GetDeviceSubgroupProperties().supportedStages, VulkanDeviceShaderStageBits) &&
 			VKHasAllFlags(Device->GetDeviceSubgroupProperties().supportedOperations, RequiredSubgroupFlags);
 
+		UE_LOG(LogVulkanRHI, Display, TEXT("Wave Operations have been %s."), GRHISupportsWaveOperations ? TEXT("ENABLED") : TEXT("DISABLED"));
+
 		if (GGPUCrashDebuggingEnabled && !Device->GetOptionalExtensions().HasGPUCrashDumpExtensions())
 		{
 			UE_LOG(LogVulkanRHI, Warning, TEXT("Tried to enable GPU crash debugging but no extension found! Will use local tracepoints."));
