VK: Added a workaround for AFBC to reduce uncompressed textures memory requirements
Changed how we apply compression for mobile landscape weightmap textures, so it works with mesh and heightmap based landscape

diff --git a/Engine/Source/Developer/Android/AndroidTargetPlatform/Private/AndroidTargetPlatform.cpp b/Engine/Source/Developer/Android/AndroidTargetPlatform/Private/AndroidTargetPlatform.cpp
--- a/Engine/Source/Developer/Android/AndroidTargetPlatform/Private/AndroidTargetPlatform.cpp
+++ b/Engine/Source/Developer/Android/AndroidTargetPlatform/Private/AndroidTargetPlatform.cpp
@@ -492,6 +492,10 @@ void FAndroidTargetPlatform::GetTextureFormats( const UTexture* Texture, TArray<
 	// OpenGL ES has F32 textures but doesn't allow linear filtering unless OES_texture_float_linear
 	const bool bSupportFilteredFloat32Textures = false;
 
+	// optionaly compress landscape weightmaps for a mobile rendering
+	static const auto CompressLandscapeWeightMapsVar = IConsoleManager::Get().FindTConsoleVariableDataInt(TEXT("r.Mobile.CompressLandscapeWeightMaps"));
+	static const bool bCompressLandscapeWeightMaps = (CompressLandscapeWeightMapsVar && CompressLandscapeWeightMapsVar->GetValueOnAnyThread() != 0);
+
 	TArray<FName>& LayerFormats = OutFormats.AddDefaulted_GetRef();
 	int32 BlockSize = 1; // this looks wrong? should be 4 for FAndroid_DXTTargetPlatform ? - it is wrong, but BlockSize is ignored
 	GetDefaultTextureFormatNamePerLayer(LayerFormats, this, Texture, bSupportCompressedVolumeTexture, BlockSize, bSupportFilteredFloat32Textures);
@@ -503,6 +507,10 @@ void FAndroidTargetPlatform::GetTextureFormats( const UTexture* Texture, TArray<
 			// forward rendering only needs one channel for shadow maps
 			TextureFormatName = FName(TEXT("G8"));
 		}
+		else if (Texture->LODGroup == TEXTUREGROUP_Terrain_Weightmap && bCompressLandscapeWeightMaps)
+		{
+			TextureFormatName = AndroidTexFormat::NameAutoDXT;
+		}
 		
 		if (Texture->IsA(UTextureCube::StaticClass()))
 		{
diff --git a/Engine/Source/Developer/IOS/IOSTargetPlatform/Private/IOSTargetPlatform.cpp b/Engine/Source/Developer/IOS/IOSTargetPlatform/Private/IOSTargetPlatform.cpp
--- a/Engine/Source/Developer/IOS/IOSTargetPlatform/Private/IOSTargetPlatform.cpp
+++ b/Engine/Source/Developer/IOS/IOSTargetPlatform/Private/IOSTargetPlatform.cpp
@@ -603,6 +603,14 @@ void FIOSTargetPlatform::GetTextureFormats( const UTexture* Texture, TArray< TAr
 		TextureFormatNames.Init(NameG8, NumLayers);
 	}
 
+	// optionaly compress landscape weightmaps for a mobile rendering
+	static const auto CompressLandscapeWeightMapsVar = IConsoleManager::Get().FindTConsoleVariableDataInt(TEXT("r.Mobile.CompressLandscapeWeightMaps"));
+	static const bool bCompressLandscapeWeightMaps = (CompressLandscapeWeightMapsVar && CompressLandscapeWeightMapsVar->GetValueOnAnyThread() != 0);
+	if (Texture->LODGroup == TEXTUREGROUP_Terrain_Weightmap && bCompressLandscapeWeightMaps)
+	{
+		TextureFormatNames.Init(FName(TEXT("AutoDXT")), NumLayers);
+	}
+
 	// if we didn't assign anything specially, then use the defaults
 	if (TextureFormatNames.Num() == 0)
 	{
diff --git a/Engine/Source/Runtime/Landscape/Private/LandscapeEdit.cpp b/Engine/Source/Runtime/Landscape/Private/LandscapeEdit.cpp
--- a/Engine/Source/Runtime/Landscape/Private/LandscapeEdit.cpp
+++ b/Engine/Source/Runtime/Landscape/Private/LandscapeEdit.cpp
@@ -89,13 +89,11 @@ DEFINE_LOG_CATEGORY(LogLandscapeBP);
 
 #define LOCTEXT_NAMESPACE "Landscape"
 
-int32 GMobileCompressLandscapeWeightMaps = 0;
-FAutoConsoleVariableRef CVarMobileCompressLanscapeWeightMaps(
-    TEXT("r.Mobile.CompressLandscapeWeightMaps"),
-    GMobileCompressLandscapeWeightMaps,
-    TEXT("Whether to compress the terrain weight maps for mobile."),
-    ECVF_ReadOnly
-);
+static TAutoConsoleVariable<int32> CVarMobileCompressLanscapeWeightMaps(
+	TEXT("r.Mobile.CompressLandscapeWeightMaps"),
+	0,
+	TEXT("Whether to compress the terrain weight maps for mobile."),
+	ECVF_ReadOnly);
 
 #if WITH_EDITOR
 
@@ -6542,7 +6540,7 @@ void ULandscapeComponent::GenerateMobilePlatformPixelData(bool bIsCooking, const
 
 	MobileWeightmapTextures.Empty();
 
-    UTexture2D* MobileWeightNormalmapTexture = GetLandscapeProxy()->CreateLandscapeTexture(WeightmapSize, WeightmapSize, TEXTUREGROUP_Terrain_Weightmap, TSF_BGRA8, nullptr, GMobileCompressLandscapeWeightMaps ? true : false );
+    UTexture2D* MobileWeightNormalmapTexture = GetLandscapeProxy()->CreateLandscapeTexture(WeightmapSize, WeightmapSize, TEXTUREGROUP_Terrain_Weightmap, TSF_BGRA8);
 	CreateEmptyTextureMips(MobileWeightNormalmapTexture, true);
 
 	{
@@ -6593,7 +6591,7 @@ void ULandscapeComponent::GenerateMobilePlatformPixelData(bool bIsCooking, const
 					// create a new weightmap texture if we've run out of channels
 					CurrentChannel = 0;
 					RemainingChannels = 4;
-                    CurrentWeightmapTexture = GetLandscapeProxy()->CreateLandscapeTexture(WeightmapSize, WeightmapSize, TEXTUREGROUP_Terrain_Weightmap, TSF_BGRA8, nullptr, GMobileCompressLandscapeWeightMaps ? true : false);
+                    CurrentWeightmapTexture = GetLandscapeProxy()->CreateLandscapeTexture(WeightmapSize, WeightmapSize, TEXTUREGROUP_Terrain_Weightmap, TSF_BGRA8);
 					CreateEmptyTextureMips(CurrentWeightmapTexture, true);
 					MobileWeightmapTextures.Add(CurrentWeightmapTexture);
 				}
@@ -7314,6 +7312,7 @@ UTexture2D* ALandscapeProxy::CreateLandscapeTexture(int32 InSizeX, int32 InSizeY
 	NewTexture->Source.Init2DWithMipChain(InSizeX, InSizeY, InFormat);
 	NewTexture->SRGB = false;
 	NewTexture->CompressionNone = !bCompress;
+	NewTexture->CompressionQuality = TCQ_Highest;
 	NewTexture->MipGenSettings = TMGS_LeaveExistingMips;
 	NewTexture->AddressX = TA_Clamp;
 	NewTexture->AddressY = TA_Clamp;
diff --git a/Engine/Source/Runtime/VulkanRHI/Private/Android/VulkanAndroidPlatform.cpp b/Engine/Source/Runtime/VulkanRHI/Private/Android/VulkanAndroidPlatform.cpp
--- a/Engine/Source/Runtime/VulkanRHI/Private/Android/VulkanAndroidPlatform.cpp
+++ b/Engine/Source/Runtime/VulkanRHI/Private/Android/VulkanAndroidPlatform.cpp
@@ -53,6 +53,7 @@ int32 FVulkanAndroidPlatform::SuccessfulRefreshRateFrames = 1;
 int32 FVulkanAndroidPlatform::UnsuccessfulRefreshRateFrames = 0;
 TArray<TArray<ANSICHAR>> FVulkanAndroidPlatform::DebugVulkanDeviceLayers;
 TArray<TArray<ANSICHAR>> FVulkanAndroidPlatform::DebugVulkanInstanceLayers;
+int32 FVulkanAndroidPlatform::AFBCWorkaroundOption = 0;
 
 #define CHECK_VK_ENTRYPOINTS(Type,Func) if (VulkanDynamicAPI::Func == NULL) { bFoundAllEntryPoints = false; UE_LOG(LogRHI, Warning, TEXT("Failed to find entry point for %s"), TEXT(#Func)); }
 
@@ -613,3 +614,95 @@ bool FAndroidVulkanFramePacer::SupportsFramePace(int32 QueryFramePace)
 	int32 TempRefreshRate, TempSyncInterval;
 	return SupportsFramePaceInternal(QueryFramePace, TempRefreshRate, TempSyncInterval);
 }
+
+//
+// Test whether we should enable a workaround for uncompressed textures
+// Arm GPUs use an optimization "Arm FrameBuffer Compression - AFBC" that can significanly inflate (~5x) uncompressed texture memory requirements
+// For now AFBC and similar optimizations can be disabled by using VK_IMAGE_CREATE_MUTABLE_FORMAT_BIT or VK_IMAGE_USAGE_STORAGE_BIT flags on a texture
+//
+void FVulkanAndroidPlatform::SetupAFBCWorkaround(const FVulkanDevice& InDevice)
+{
+	AFBCWorkaroundOption = 0;
+
+	const VkFormat FormatAFBC = VK_FORMAT_B8G8R8A8_UNORM;
+	const VkFormatFeatureFlags FormatFlags = InDevice.GetFormatProperties()[FormatAFBC].optimalTilingFeatures;
+
+	VkImageCreateInfo ImageCreateInfo;
+	ZeroVulkanStruct(ImageCreateInfo, VK_STRUCTURE_TYPE_IMAGE_CREATE_INFO);
+	ImageCreateInfo.imageType = VK_IMAGE_TYPE_2D;
+	ImageCreateInfo.samples = VK_SAMPLE_COUNT_1_BIT;
+	ImageCreateInfo.format = FormatAFBC;
+	ImageCreateInfo.tiling = VK_IMAGE_TILING_OPTIMAL;
+	ImageCreateInfo.arrayLayers = 1;
+	ImageCreateInfo.extent = {128, 128, 1};
+	ImageCreateInfo.mipLevels = 8;
+	ImageCreateInfo.flags = 0;// VK_IMAGE_CREATE_MUTABLE_FORMAT_BIT;
+	ImageCreateInfo.usage = VK_IMAGE_USAGE_TRANSFER_DST_BIT | VK_IMAGE_USAGE_SAMPLED_BIT;
+	ImageCreateInfo.sharingMode = VK_SHARING_MODE_EXCLUSIVE;
+	ImageCreateInfo.queueFamilyIndexCount = 0;
+	ImageCreateInfo.pQueueFamilyIndices = nullptr;
+	ImageCreateInfo.initialLayout = VK_IMAGE_LAYOUT_SHADER_READ_ONLY_OPTIMAL;
+
+	VkImage Image0;
+	VkMemoryRequirements Image0Mem;
+	VERIFYVULKANRESULT(VulkanRHI::vkCreateImage(InDevice.GetInstanceHandle(), &ImageCreateInfo, VULKAN_CPU_ALLOCATOR, &Image0));
+	VulkanRHI::vkGetImageMemoryRequirements(InDevice.GetInstanceHandle(), Image0, &Image0Mem);
+	VulkanRHI::vkDestroyImage(InDevice.GetInstanceHandle(), Image0, VULKAN_CPU_ALLOCATOR);
+
+	VkImage ImageMutable;
+	VkMemoryRequirements ImageMutableMem;
+	ImageCreateInfo.flags = VK_IMAGE_CREATE_MUTABLE_FORMAT_BIT;
+	VERIFYVULKANRESULT(VulkanRHI::vkCreateImage(InDevice.GetInstanceHandle(), &ImageCreateInfo, VULKAN_CPU_ALLOCATOR, &ImageMutable));
+	VulkanRHI::vkGetImageMemoryRequirements(InDevice.GetInstanceHandle(), ImageMutable, &ImageMutableMem);
+	VulkanRHI::vkDestroyImage(InDevice.GetInstanceHandle(), ImageMutable, VULKAN_CPU_ALLOCATOR);
+
+	VkImage ImageStorage;
+	VkMemoryRequirements ImageStorageMem;
+	if ((FormatFlags & VK_FORMAT_FEATURE_STORAGE_IMAGE_BIT) != 0)
+	{
+		ImageCreateInfo.usage |= VK_IMAGE_USAGE_STORAGE_BIT;
+		ImageCreateInfo.flags = 0;
+	}
+	VERIFYVULKANRESULT(VulkanRHI::vkCreateImage(InDevice.GetInstanceHandle(), &ImageCreateInfo, VULKAN_CPU_ALLOCATOR, &ImageStorage));
+	VulkanRHI::vkGetImageMemoryRequirements(InDevice.GetInstanceHandle(), ImageStorage, &ImageStorageMem);
+	VulkanRHI::vkDestroyImage(InDevice.GetInstanceHandle(), ImageStorage, VULKAN_CPU_ALLOCATOR);
+
+	const float MEM_SIZE_THRESHOLD = 1.5f;
+	const float IMAGE0_SIZE = (float)Image0Mem.size;
+
+	if (ImageMutableMem.size * MEM_SIZE_THRESHOLD < IMAGE0_SIZE)
+	{
+		AFBCWorkaroundOption = 1;
+	}
+	else if (ImageStorageMem.size * MEM_SIZE_THRESHOLD < IMAGE0_SIZE)
+	{
+		AFBCWorkaroundOption = 2;
+	}
+
+	if (AFBCWorkaroundOption != 0)
+	{
+		UE_LOG(LogRHI, Display, TEXT("Enabling workaround to reduce memory requrement for BGRA textures (%s flag). 128x128 - 8 Mips BGRA texture: %u KiB -> %u KiB"), 
+			AFBCWorkaroundOption == 1 ? TEXT("MUTABLE") : TEXT("STORAGE"),
+			Image0Mem.size / 1024,
+			AFBCWorkaroundOption == 1 ? ImageMutableMem.size / 1024 : ImageStorageMem.size / 1024
+		);
+	}
+}
+
+void FVulkanAndroidPlatform::SetImageAFBCWorkaround(VkImageCreateInfo& ImageCreateInfo)
+{
+	if (AFBCWorkaroundOption != 0 &&
+		ImageCreateInfo.imageType == VK_IMAGE_TYPE_2D && 
+		ImageCreateInfo.format == VK_FORMAT_B8G8R8A8_UNORM && 
+		ImageCreateInfo.mipLevels >= 8) // its worth enabling for 128x128 and up
+	{
+		if (AFBCWorkaroundOption == 1)
+		{
+			ImageCreateInfo.flags |= VK_IMAGE_CREATE_MUTABLE_FORMAT_BIT;
+		}
+		else if (AFBCWorkaroundOption == 2)
+		{
+			ImageCreateInfo.usage |= VK_IMAGE_USAGE_STORAGE_BIT;
+		}
+	}
+}
diff --git a/Engine/Source/Runtime/VulkanRHI/Private/Android/VulkanAndroidPlatform.h b/Engine/Source/Runtime/VulkanRHI/Private/Android/VulkanAndroidPlatform.h
--- a/Engine/Source/Runtime/VulkanRHI/Private/Android/VulkanAndroidPlatform.h
+++ b/Engine/Source/Runtime/VulkanRHI/Private/Android/VulkanAndroidPlatform.h
@@ -143,6 +143,11 @@ public:
 		OutFallbackFormat1 = VK_FORMAT_B8G8R8A8_UNORM;
 		return VK_FORMAT_R5G5B5A1_UNORM_PACK16; 
 	};
+
+	// Setup platform to use AFBC workaround to reduce uncompressed textures memory requirements
+	static void SetupAFBCWorkaround(const FVulkanDevice& InDevice);
+	static void SetImageAFBCWorkaround(VkImageCreateInfo& ImageCreateInfo);
+
 protected:
 	static void* VulkanLib;
 	static bool bAttemptedLoad;
@@ -160,6 +165,8 @@ protected:
 	static int32 UnsuccessfulRefreshRateFrames;
 	static TArray<TArray<ANSICHAR>> DebugVulkanDeviceLayers;
 	static TArray<TArray<ANSICHAR>> DebugVulkanInstanceLayers;
+
+	static int32 AFBCWorkaroundOption;
 };
 
 #if VULKAN_SUPPORTS_GOOGLE_DISPLAY_TIMING
diff --git a/Engine/Source/Runtime/VulkanRHI/Private/VulkanDevice.cpp b/Engine/Source/Runtime/VulkanRHI/Private/VulkanDevice.cpp
--- a/Engine/Source/Runtime/VulkanRHI/Private/VulkanDevice.cpp
+++ b/Engine/Source/Runtime/VulkanRHI/Private/VulkanDevice.cpp
@@ -1126,6 +1126,8 @@ void FVulkanDevice::InitGPU(int32 DeviceIndex)
 
 		DefaultTexture = new FVulkanTexture(*this, Desc, nullptr);
 	}
+
+	FVulkanPlatform::SetupAFBCWorkaround(*this);
 }
 
 void FVulkanDevice::PrepareForDestroy()
diff --git a/Engine/Source/Runtime/VulkanRHI/Private/VulkanGenericPlatform.h b/Engine/Source/Runtime/VulkanRHI/Private/VulkanGenericPlatform.h
--- a/Engine/Source/Runtime/VulkanRHI/Private/VulkanGenericPlatform.h
+++ b/Engine/Source/Runtime/VulkanRHI/Private/VulkanGenericPlatform.h
@@ -133,4 +133,8 @@ public:
 
 	// Return VK_FALSE if platform wants to suppress the given debug report from the validation layers, VK_TRUE to print it.
 	static VkBool32 DebugReportFunction(VkDebugReportFlagsEXT MsgFlags, VkDebugReportObjectTypeEXT ObjType, uint64_t SrcObject, size_t Location, int32 MsgCode, const ANSICHAR* LayerPrefix, const ANSICHAR* Msg, void* UserData) { return VK_TRUE; }
+
+	// Setup platform to use AFBC workaround to reduce uncompressed textures memory requirements
+	static void SetupAFBCWorkaround(const FVulkanDevice& InDevice) {};
+	static void SetImageAFBCWorkaround(VkImageCreateInfo& ImageCreateInfo) {};
 };
diff --git a/Engine/Source/Runtime/VulkanRHI/Private/VulkanTexture.cpp b/Engine/Source/Runtime/VulkanRHI/Private/VulkanTexture.cpp
--- a/Engine/Source/Runtime/VulkanRHI/Private/VulkanTexture.cpp
+++ b/Engine/Source/Runtime/VulkanRHI/Private/VulkanTexture.cpp
@@ -438,6 +438,8 @@ void FVulkanTexture::GenerateImageCreateInfo(
 		checkf(0, TEXT("Unsupported number of samples %d"), NumSamples);
 		break;
 	}
+
+	FVulkanPlatform::SetImageAFBCWorkaround(ImageCreateInfo);
 		
 	const VkFormatFeatureFlags FormatFlags = ImageCreateInfo.tiling == VK_IMAGE_TILING_LINEAR ? 
 		InDevice.GetFormatProperties()[ImageCreateInfo.format].linearTilingFeatures : 
