VulkanRHI:  Add VK_SUBGROUP_FEATURE_SHUFFLE_BIT to the check for GRHISupportsWaveOperations (for WaveReadLaneAt)

#rnx

diff --git a/Engine/Source/Runtime/VulkanRHI/Private/VulkanRHI.cpp b/Engine/Source/Runtime/VulkanRHI/Private/VulkanRHI.cpp
--- a/Engine/Source/Runtime/VulkanRHI/Private/VulkanRHI.cpp
+++ b/Engine/Source/Runtime/VulkanRHI/Private/VulkanRHI.cpp
@@ -866,7 +866,9 @@ void FVulkanDynamicRHI::InitInstance()
 		}
 		
 		// Check for wave ops support (only filled on platforms creating Vulkan 1.1 or greater instances)
-		const VkSubgroupFeatureFlags RequiredSubgroupFlags = VK_SUBGROUP_FEATURE_BASIC_BIT | VK_SUBGROUP_FEATURE_VOTE_BIT | VK_SUBGROUP_FEATURE_ARITHMETIC_BIT | VK_SUBGROUP_FEATURE_BALLOT_BIT;
+		const VkSubgroupFeatureFlags RequiredSubgroupFlags =	VK_SUBGROUP_FEATURE_BASIC_BIT | VK_SUBGROUP_FEATURE_VOTE_BIT | 
+																VK_SUBGROUP_FEATURE_ARITHMETIC_BIT | VK_SUBGROUP_FEATURE_BALLOT_BIT | 
+																VK_SUBGROUP_FEATURE_SHUFFLE_BIT;
 		GRHISupportsWaveOperations = VKHasAllFlags(Device->GetDeviceSubgroupProperties().supportedStages, VulkanDeviceShaderStageBits) &&
 			VKHasAllFlags(Device->GetDeviceSubgroupProperties().supportedOperations, RequiredSubgroupFlags);
 
