VulkanRHI:  Set GRHIMinimumWaveSize/GRHIMaximumWaveSize when wave ops are supported (use VK_EXT_subgroup_size_control if available, or the default subgroup size provided by device sungroup properties).  Improve log messages: provide subgroup sizes when supported, or missing flags when not supported.

#rnx

diff --git a/Engine/Source/Runtime/VulkanRHI/Private/VulkanExtensions.cpp b/Engine/Source/Runtime/VulkanRHI/Private/VulkanExtensions.cpp
--- a/Engine/Source/Runtime/VulkanRHI/Private/VulkanExtensions.cpp
+++ b/Engine/Source/Runtime/VulkanRHI/Private/VulkanExtensions.cpp
@@ -903,6 +903,50 @@ private:
 };
 
 
+// ***** VK_EXT_subgroup_size_control
+class FVulkanEXTSubgroupSizeControlExtension : public FVulkanDeviceExtension
+{
+public:
+
+	FVulkanEXTSubgroupSizeControlExtension(FVulkanDevice* InDevice)
+		: FVulkanDeviceExtension(InDevice, VK_EXT_SUBGROUP_SIZE_CONTROL_EXTENSION_NAME, VULKAN_EXTENSION_ENABLED, VK_API_VERSION_1_3)
+	{
+	}
+
+	virtual void PrePhysicalDeviceFeatures(VkPhysicalDeviceFeatures2KHR& PhysicalDeviceFeatures2) override final
+	{
+		ZeroVulkanStruct(SubgroupSizeControlFeatures, VK_STRUCTURE_TYPE_PHYSICAL_DEVICE_SUBGROUP_SIZE_CONTROL_FEATURES);
+		AddToPNext(PhysicalDeviceFeatures2, SubgroupSizeControlFeatures);
+	}
+
+	virtual void PostPhysicalDeviceFeatures(FOptionalVulkanDeviceExtensions& ExtensionFlags) override final
+	{
+ 		bSupportsSubgroupSizeControl = (SubgroupSizeControlFeatures.subgroupSizeControl == VK_TRUE);
+	}
+
+	virtual void PrePhysicalDeviceProperties(VkPhysicalDeviceProperties2KHR& PhysicalDeviceProperties2) override final
+	{
+		if (bSupportsSubgroupSizeControl)
+		{
+			ZeroVulkanStruct(SubgroupSizeControlProperties, VK_STRUCTURE_TYPE_PHYSICAL_DEVICE_SUBGROUP_SIZE_CONTROL_PROPERTIES);
+			AddToPNext(PhysicalDeviceProperties2, SubgroupSizeControlProperties);
+		}
+	}
+
+	virtual void PreCreateDevice(VkDeviceCreateInfo& DeviceCreateInfo) override final
+	{
+		if (bSupportsSubgroupSizeControl)
+		{
+			GRHIMinimumWaveSize = SubgroupSizeControlProperties.minSubgroupSize;
+			GRHIMaximumWaveSize = SubgroupSizeControlProperties.maxSubgroupSize;
+		}
+	}
+
+private:
+	VkPhysicalDeviceSubgroupSizeControlFeaturesEXT SubgroupSizeControlFeatures;
+	VkPhysicalDeviceSubgroupSizeControlPropertiesEXT SubgroupSizeControlProperties;
+	bool bSupportsSubgroupSizeControl = false;
+};
 
 
 
@@ -976,6 +1020,7 @@ FVulkanDeviceExtensionArray FVulkanDeviceExtension::GetUESupportedDeviceExtensio
 	ADD_CUSTOM_EXTENSION(FVulkanKHRGetMemoryRequirements2Extension);
 	ADD_CUSTOM_EXTENSION(FVulkanEXTDescriptorIndexingExtension);
 	ADD_CUSTOM_EXTENSION(FVulkanEXTHostQueryResetExtension);
+	ADD_CUSTOM_EXTENSION(FVulkanEXTSubgroupSizeControlExtension);
 
 	// Needed for Raytracing
 	ADD_CUSTOM_EXTENSION(FVulkanKHRBufferDeviceAddressExtension);
diff --git a/Engine/Source/Runtime/VulkanRHI/Private/VulkanRHI.cpp b/Engine/Source/Runtime/VulkanRHI/Private/VulkanRHI.cpp
--- a/Engine/Source/Runtime/VulkanRHI/Private/VulkanRHI.cpp
+++ b/Engine/Source/Runtime/VulkanRHI/Private/VulkanRHI.cpp
@@ -872,7 +872,23 @@ void FVulkanDynamicRHI::InitInstance()
 		GRHISupportsWaveOperations = VKHasAllFlags(Device->GetDeviceSubgroupProperties().supportedStages, VulkanDeviceShaderStageBits) &&
 			VKHasAllFlags(Device->GetDeviceSubgroupProperties().supportedOperations, RequiredSubgroupFlags);
 
-		UE_LOG(LogVulkanRHI, Display, TEXT("Wave Operations have been %s."), GRHISupportsWaveOperations ? TEXT("ENABLED") : TEXT("DISABLED"));
+		if (GRHISupportsWaveOperations)
+		{
+			// Use default size if VK_EXT_subgroup_size_control didn't fill them
+			if (!GRHIMinimumWaveSize || !GRHIMaximumWaveSize)
+			{
+				GRHIMinimumWaveSize = GRHIMaximumWaveSize = Device->GetDeviceSubgroupProperties().subgroupSize;
+			}
+
+			UE_LOG(LogVulkanRHI, Display, TEXT("Wave Operations have been ENABLED (wave size: min=%d max=%d)."), GRHIMinimumWaveSize, GRHIMaximumWaveSize);
+		}
+		else
+		{
+			const uint32 MissingStageFlags = (Device->GetDeviceSubgroupProperties().supportedStages & VulkanDeviceShaderStageBits) ^ VulkanDeviceShaderStageBits;
+			const uint32 MissingOperationFlags = (Device->GetDeviceSubgroupProperties().supportedOperations & RequiredSubgroupFlags) ^ RequiredSubgroupFlags;
+			UE_LOG(LogVulkanRHI, Display, TEXT("Wave Operations have been DISABLED (missing stages=0x%x operations=0x%x)."), MissingStageFlags, MissingOperationFlags);
+		}
+
 
 		if (GGPUCrashDebuggingEnabled && !Device->GetOptionalExtensions().HasGPUCrashDumpExtensions())
 		{
