Allocate FVulkanRayTracingScene::PerInstanceGeometryParameterBuffer directly as FVulkanResourceMultiBuffer, instead of going through RHICreateBuffer.

* Regular RHICreateBuffer must be called from the rendering thread, however RT scene can now be created from any thread.
* FVulkanResourceMultiBuffer can already be created from any thread if it does not require initialization
* Fixes crash on startup when running in Vulkan mode with HWRT

diff --git a/Engine/Source/Runtime/VulkanRHI/Private/VulkanRayTracing.cpp b/Engine/Source/Runtime/VulkanRHI/Private/VulkanRayTracing.cpp
--- a/Engine/Source/Runtime/VulkanRHI/Private/VulkanRayTracing.cpp
+++ b/Engine/Source/Runtime/VulkanRHI/Private/VulkanRayTracing.cpp
@@ -420,9 +420,9 @@ FVulkanRayTracingScene::FVulkanRayTracingScene(FRayTracingSceneInitializer2 InIn
 
 	const uint32 ParameterBufferSize = FMath::Max<uint32>(1, Initializer.NumTotalSegments) * sizeof(FVulkanRayTracingGeometryParameters);
 	FRHIResourceCreateInfo ParameterBufferCreateInfo(TEXT("RayTracingSceneMetadata"));
-	PerInstanceGeometryParameterBuffer = ResourceCast(RHICreateBuffer(
+	PerInstanceGeometryParameterBuffer = new FVulkanResourceMultiBuffer(Device,
 		ParameterBufferSize, BUF_StructuredBuffer | BUF_ShaderResource, sizeof(FVulkanRayTracingGeometryParameters),
-		ERHIAccess::SRVCompute, ParameterBufferCreateInfo).GetReference());
+		ParameterBufferCreateInfo);
 
 	PerInstanceGeometryParameterSRV = new FVulkanShaderResourceView(Device, PerInstanceGeometryParameterBuffer, 0);
 }
