Make UpdateUniformBuffer thread safe by taking a command list.

diff --git a/Engine/Source/Runtime/Apple/AGXRHI/Private/AGXDynamicRHI.h b/Engine/Source/Runtime/Apple/AGXRHI/Private/AGXDynamicRHI.h
--- a/Engine/Source/Runtime/Apple/AGXRHI/Private/AGXDynamicRHI.h
+++ b/Engine/Source/Runtime/Apple/AGXRHI/Private/AGXDynamicRHI.h
@@ -162,7 +162,7 @@ public:
 	virtual FComputeShaderRHIRef CreateComputeShader_RenderThread(class FRHICommandListImmediate& RHICmdList, TArrayView<const uint8> Code, const FSHAHash& Hash) final override;
 	virtual FRHIShaderLibraryRef RHICreateShaderLibrary_RenderThread(class FRHICommandListImmediate& RHICmdList, EShaderPlatform Platform, FString FilePath, FString Name) final override;
 
-	virtual void RHIUpdateUniformBuffer(FRHIUniformBuffer* UniformBufferRHI, const void* Contents) final override;
+	virtual void RHIUpdateUniformBuffer(FRHICommandListBase& RHICmdList, FRHIUniformBuffer* UniformBufferRHI, const void* Contents) final override;
 
 	virtual uint16 RHIGetPlatformTextureMaxSampleCount() override;
 
diff --git a/Engine/Source/Runtime/Apple/AGXRHI/Private/AGXUniformBufferCommands.cpp b/Engine/Source/Runtime/Apple/AGXRHI/Private/AGXUniformBufferCommands.cpp
--- a/Engine/Source/Runtime/Apple/AGXRHI/Private/AGXUniformBufferCommands.cpp
+++ b/Engine/Source/Runtime/Apple/AGXRHI/Private/AGXUniformBufferCommands.cpp
@@ -4,44 +4,34 @@
 #include "AGXShaderTypes.h"
 #include "AGXFrameAllocator.h"
 
-static void DoUpdateUniformBuffer(FAGXUniformBuffer* UB, const void* Contents)
-{
-    check(IsInRenderingThread() || IsInRHIThread());
-    
-    FRHICommandListImmediate& RHICmdList = FRHICommandListExecutor::GetImmediateCommandList();
-    // The only way we can be on the RHI thread here is if we're in the process of creating a FLocalUniformBuffer.
-    bool bUpdateImmediately = RHICmdList.Bypass() || IsInRHIThread();
-
-    if(bUpdateImmediately)
-    {
-        UB->Update(Contents);
-    }
-    else
-    {
-        const uint32 NumBytes = UB->GetLayout().ConstantBufferSize;
-        void* Data = RHICmdList.Alloc(NumBytes, 16);
-        FMemory::Memcpy(Data, Contents, NumBytes);
-
-        RHICmdList.EnqueueLambda([Data, UB](FRHICommandListImmediate& RHICmdList)
-        {
-            UB->Update(Data);
-        });
-
-        RHICmdList.RHIThreadFence(true);
-    }
-}
-
 FUniformBufferRHIRef FAGXDynamicRHI::RHICreateUniformBuffer(const void* Contents, const FRHIUniformBufferLayout* Layout, EUniformBufferUsage Usage, EUniformBufferValidation Validation)
 {
 	return new FAGXUniformBuffer(Contents, Layout, Usage, Validation);
 }
 
-void FAGXDynamicRHI::RHIUpdateUniformBuffer(FRHIUniformBuffer* UniformBufferRHI, const void* Contents)
+void FAGXDynamicRHI::RHIUpdateUniformBuffer(FRHICommandListBase& RHICmdList, FRHIUniformBuffer* UniformBufferRHI, const void* Contents)
 {
-    check(IsInRenderingThread());
-    
-    FAGXUniformBuffer* UB = ResourceCast(UniformBufferRHI);
-    DoUpdateUniformBuffer(UB, Contents);
+    FAGXUniformBuffer* UniformBuffer = ResourceCast(UniformBufferRHI);
+
+	const void* SrcContents = Contents;
+
+	if (RHICmdList.IsTopOfPipe())
+	{
+		const FRHIUniformBufferLayout& Layout = UniformBuffer->GetLayout();
+
+		// Copy the contents memory region into the RHICmdList to allow deferred execution on the RHI thread.
+		void* DstContents = RHICmdList.Alloc(Layout.ConstantBufferSize, alignof(FRHIResource*));
+		FMemory::ParallelMemcpy(DstContents, Contents, Layout.ConstantBufferSize, EMemcpyCachePolicy::StoreUncached);
+
+		SrcContents = DstContents;
+	}
+
+	RHICmdList.EnqueueLambda([UniformBuffer, SrcContents](FRHICommandListBase& RHICmdList)
+	{
+		UniformBuffer->Update(SrcContents);
+	});
+
+	RHICmdList.RHIThreadFence(true);
 }
 
 template<EAGXShaderStages Stage, typename RHIShaderType>
diff --git a/Engine/Source/Runtime/Apple/MetalRHI/Private/MetalDynamicRHI.h b/Engine/Source/Runtime/Apple/MetalRHI/Private/MetalDynamicRHI.h
--- a/Engine/Source/Runtime/Apple/MetalRHI/Private/MetalDynamicRHI.h
+++ b/Engine/Source/Runtime/Apple/MetalRHI/Private/MetalDynamicRHI.h
@@ -166,7 +166,7 @@ public:
 	virtual FComputeShaderRHIRef CreateComputeShader_RenderThread(class FRHICommandListImmediate& RHICmdList, TArrayView<const uint8> Code, const FSHAHash& Hash) final override;
 	virtual FRHIShaderLibraryRef RHICreateShaderLibrary_RenderThread(class FRHICommandListImmediate& RHICmdList, EShaderPlatform Platform, FString FilePath, FString Name) final override;
 
-	virtual void RHIUpdateUniformBuffer(FRHIUniformBuffer* UniformBufferRHI, const void* Contents) final override;
+	virtual void RHIUpdateUniformBuffer(FRHICommandListBase& RHICmdList, FRHIUniformBuffer* UniformBufferRHI, const void* Contents) final override;
 
 	virtual uint16 RHIGetPlatformTextureMaxSampleCount() override;
 
diff --git a/Engine/Source/Runtime/Apple/MetalRHI/Private/MetalUniformBufferCommands.cpp b/Engine/Source/Runtime/Apple/MetalRHI/Private/MetalUniformBufferCommands.cpp
--- a/Engine/Source/Runtime/Apple/MetalRHI/Private/MetalUniformBufferCommands.cpp
+++ b/Engine/Source/Runtime/Apple/MetalRHI/Private/MetalUniformBufferCommands.cpp
@@ -4,44 +4,32 @@
 #include "MetalShaderTypes.h"
 #include "MetalFrameAllocator.h"
 
-static void DoUpdateUniformBuffer(FMetalUniformBuffer* UB, const void* Contents)
+FUniformBufferRHIRef FMetalDynamicRHI::RHICreateUniformBuffer(const void* Contents, const FRHIUniformBufferLayout* Layout, EUniformBufferUsage Usage, EUniformBufferValidation Validation)
 {
-    check(IsInRenderingThread() || IsInRHIThread());
+	return new FMetalUniformBuffer(Contents, Layout, Usage, Validation);
+}
 
-    FRHICommandListImmediate& RHICmdList = FRHICommandListExecutor::GetImmediateCommandList();
-    // The only way we can be on the RHI thread here is if we're in the process of creating a FLocalUniformBuffer.
-    bool bUpdateImmediately = RHICmdList.Bypass() || IsInRHIThread();
+void FMetalDynamicRHI::RHIUpdateUniformBuffer(FRHICommandListBase& RHICmdList, FRHIUniformBuffer* UniformBufferRHI, const void* Contents)
+{
+	FMetalUniformBuffer* UniformBuffer = ResourceCast(UniformBufferRHI);
 
-    if(bUpdateImmediately)
-    {
-        UB->Update(Contents);
-    }
-    else
-    {
-        const uint32 NumBytes = UB->GetLayout().ConstantBufferSize;
-        void* Data = RHICmdList.Alloc(NumBytes, 16);
-        FMemory::Memcpy(Data, Contents, NumBytes);
+	const void* SrcContents = Contents;
 
-        RHICmdList.EnqueueLambda([Data, UB](FRHICommandListImmediate& RHICmdList)
-        {
-            UB->Update(Data);
-        });
+	if (RHICmdList.IsTopOfPipe())
+	{
+		const FRHIUniformBufferLayout& Layout = UniformBuffer->GetLayout();
 
-        RHICmdList.RHIThreadFence(true);
-    }
-}
+		// Copy the contents memory region into the RHICmdList to allow deferred execution on the RHI thread.
+		void* DstContents = RHICmdList.Alloc(Layout.ConstantBufferSize, alignof(FRHIResource*));
+		FMemory::ParallelMemcpy(DstContents, Contents, Layout.ConstantBufferSize, EMemcpyCachePolicy::StoreUncached);
 
-FUniformBufferRHIRef FMetalDynamicRHI::RHICreateUniformBuffer(const void* Contents, const FRHIUniformBufferLayout* Layout, EUniformBufferUsage Usage, EUniformBufferValidation Validation)
-{
-	return new FMetalUniformBuffer(Contents, Layout, Usage, Validation);
-}
+		SrcContents = DstContents;
+	}
 
-void FMetalDynamicRHI::RHIUpdateUniformBuffer(FRHIUniformBuffer* UniformBufferRHI, const void* Contents)
-{
-    check(IsInRenderingThread());
-    
-    FMetalUniformBuffer* UB = ResourceCast(UniformBufferRHI);
-    DoUpdateUniformBuffer(UB, Contents);
+	RHICmdList.EnqueueLambda([UniformBuffer, SrcContents](FRHICommandListBase& RHICmdList)
+	{
+		UniformBuffer->Update(SrcContents);
+	});
 }
 
 template<EMetalShaderStages Stage, typename RHIShaderType>
diff --git a/Engine/Source/Runtime/D3D12RHI/Private/D3D12RHIPrivate.h b/Engine/Source/Runtime/D3D12RHI/Private/D3D12RHIPrivate.h
--- a/Engine/Source/Runtime/D3D12RHI/Private/D3D12RHIPrivate.h
+++ b/Engine/Source/Runtime/D3D12RHI/Private/D3D12RHIPrivate.h
@@ -338,7 +338,7 @@ public:
 	virtual void RHICreateTransition(FRHITransition* Transition, const FRHITransitionCreateInfo& CreateInfo) final override;
 	virtual void RHIReleaseTransition(FRHITransition* Transition) final override;
 	virtual FUniformBufferRHIRef RHICreateUniformBuffer(const void* Contents, const FRHIUniformBufferLayout* Layout, EUniformBufferUsage Usage, EUniformBufferValidation Validation) final override;
-	virtual void RHIUpdateUniformBuffer(FRHIUniformBuffer* UniformBufferRHI, const void* Contents) final override;
+	virtual void RHIUpdateUniformBuffer(FRHICommandListBase& RHICmdList, FRHIUniformBuffer* UniformBufferRHI, const void* Contents) final override;
 	virtual FBufferRHIRef RHICreateBuffer(FRHICommandListBase& RHICmdList, uint32 Size, EBufferUsageFlags Usage, uint32 Stride, ERHIAccess ResourceState, FRHIResourceCreateInfo& CreateInfo) final override;
 	virtual void RHITransferBufferUnderlyingResource(FRHIBuffer* DestBuffer, FRHIBuffer* SrcBuffer) final override;
 	virtual void RHICopyBuffer(FRHIBuffer* SourceBuffer, FRHIBuffer* DestBuffer) final override;
diff --git a/Engine/Source/Runtime/D3D12RHI/Private/D3D12UniformBuffer.cpp b/Engine/Source/Runtime/D3D12RHI/Private/D3D12UniformBuffer.cpp
--- a/Engine/Source/Runtime/D3D12RHI/Private/D3D12UniformBuffer.cpp
+++ b/Engine/Source/Runtime/D3D12RHI/Private/D3D12UniformBuffer.cpp
@@ -238,15 +238,13 @@ FRHICOMMAND_MACRO(FRHICommandD3D12UpdateUniformBuffer)
 	}
 };
 
-void FD3D12DynamicRHI::RHIUpdateUniformBuffer(FRHIUniformBuffer* UniformBufferRHI, const void* Contents)
+void FD3D12DynamicRHI::RHIUpdateUniformBuffer(FRHICommandListBase& RHICmdList, FRHIUniformBuffer* UniformBufferRHI, const void* Contents)
 {
-	check(IsInRenderingThread());
 	check(UniformBufferRHI);
 
 	const FRHIUniformBufferLayout& Layout = UniformBufferRHI->GetLayout();
 	ValidateShaderParameterResourcesRHI(Contents, Layout);
 
-	FRHICommandListImmediate& RHICmdList = FRHICommandListExecutor::GetImmediateCommandList();
 	const bool bBypass = RHICmdList.Bypass();
 
 	FD3D12UniformBuffer* FirstUniformBuffer = ResourceCast(UniformBufferRHI);
diff --git a/Engine/Source/Runtime/NullDrv/Public/NullRHI.h b/Engine/Source/Runtime/NullDrv/Public/NullRHI.h
--- a/Engine/Source/Runtime/NullDrv/Public/NullRHI.h
+++ b/Engine/Source/Runtime/NullDrv/Public/NullRHI.h
@@ -110,7 +110,7 @@ public:
 		return new FRHIUniformBuffer(Layout); 
 	}
 
-	virtual void RHIUpdateUniformBuffer(FRHIUniformBuffer* UniformBufferRHI, const void* Contents) final override
+	virtual void RHIUpdateUniformBuffer(FRHICommandListBase& RHICmdList, FRHIUniformBuffer* UniformBufferRHI, const void* Contents) final override
 	{
 
 	}
diff --git a/Engine/Source/Runtime/OpenGLDrv/Private/OpenGLUniformBuffer.cpp b/Engine/Source/Runtime/OpenGLDrv/Private/OpenGLUniformBuffer.cpp
--- a/Engine/Source/Runtime/OpenGLDrv/Private/OpenGLUniformBuffer.cpp
+++ b/Engine/Source/Runtime/OpenGLDrv/Private/OpenGLUniformBuffer.cpp
@@ -774,11 +774,10 @@ void UpdateUniformBufferContents(FOpenGLUniformBuffer* UniformBuffer, const void
 	}
 }
 
-void FOpenGLDynamicRHI::RHIUpdateUniformBuffer(FRHIUniformBuffer* UniformBufferRHI, const void* Contents)
+void FOpenGLDynamicRHI::RHIUpdateUniformBuffer(FRHICommandListBase& RHICmdList, FRHIUniformBuffer* UniformBufferRHI, const void* Contents)
 {
 	FOpenGLUniformBuffer* UniformBuffer = ResourceCast(UniformBufferRHI);
 
-	FRHICommandListImmediate& RHICmdList = FRHICommandListExecutor::GetImmediateCommandList();
 	const FRHIUniformBufferLayout& Layout = UniformBufferRHI->GetLayout();
 	ValidateShaderParameterResourcesRHI(Contents, Layout);
 
@@ -822,7 +821,7 @@ void FOpenGLDynamicRHI::RHIUpdateUniformBuffer(FRHIUniformBuffer* UniformBufferR
 			FMemory::Memcpy(CmdListConstantBufferData, Contents, ConstantBufferSize);
 		}
 
-		RHICmdList.EnqueueLambda([UniformBuffer, CmdListResources, NumResources, CmdListConstantBufferData, ConstantBufferSize, NextUniqueID](FRHICommandList&)
+		RHICmdList.EnqueueLambda([UniformBuffer, CmdListResources, NumResources, CmdListConstantBufferData, ConstantBufferSize, NextUniqueID](FRHICommandListBase&)
 		{
 			UpdateUniformBufferContents(UniformBuffer, CmdListConstantBufferData, ConstantBufferSize);
 
diff --git a/Engine/Source/Runtime/OpenGLDrv/Public/OpenGLDrv.h b/Engine/Source/Runtime/OpenGLDrv/Public/OpenGLDrv.h
--- a/Engine/Source/Runtime/OpenGLDrv/Public/OpenGLDrv.h
+++ b/Engine/Source/Runtime/OpenGLDrv/Public/OpenGLDrv.h
@@ -370,7 +370,7 @@ public:
 	virtual FComputeShaderRHIRef RHICreateComputeShader(TArrayView<const uint8> Code, const FSHAHash& Hash) final override;
 
 	virtual FUniformBufferRHIRef RHICreateUniformBuffer(const void* Contents, const FRHIUniformBufferLayout* Layout, EUniformBufferUsage Usage, EUniformBufferValidation Validation) final override;
-	virtual void RHIUpdateUniformBuffer(FRHIUniformBuffer* UniformBufferRHI, const void* Contents) final override;
+	virtual void RHIUpdateUniformBuffer(FRHICommandListBase& RHICmdList, FRHIUniformBuffer* UniformBufferRHI, const void* Contents) final override;
 	virtual FBufferRHIRef RHICreateBuffer(FRHICommandListBase& RHICmdList, uint32 Size, EBufferUsageFlags Usage, uint32 Stride, ERHIAccess ResourceState, FRHIResourceCreateInfo& CreateInfo) final override;
 	virtual void RHICopyBuffer(FRHIBuffer* SourceBuffer, FRHIBuffer* DestBuffer) final override;
 	virtual void RHITransferBufferUnderlyingResource(FRHIBuffer* DestBuffer, FRHIBuffer* SrcBuffer) final override;
diff --git a/Engine/Source/Runtime/RHI/Public/DynamicRHI.h b/Engine/Source/Runtime/RHI/Public/DynamicRHI.h
--- a/Engine/Source/Runtime/RHI/Public/DynamicRHI.h
+++ b/Engine/Source/Runtime/RHI/Public/DynamicRHI.h
@@ -407,7 +407,7 @@ public:
 		return RHICreateUniformBuffer(Contents, &Layout, Usage, Validation);
 	}
 
-	virtual void RHIUpdateUniformBuffer(FRHIUniformBuffer* UniformBufferRHI, const void* Contents) = 0;
+	virtual void RHIUpdateUniformBuffer(FRHICommandListBase& RHICmdList, FRHIUniformBuffer* UniformBufferRHI, const void* Contents) = 0;
 
 	/** Copies the contents of one buffer to another buffer. They must have identical sizes. */
 	// FlushType: Flush Immediate (seems dangerous)
@@ -1318,11 +1318,6 @@ FORCEINLINE FUniformBufferRHIRef RHICreateUniformBuffer(const void* Contents, co
 	return GDynamicRHI->RHICreateUniformBuffer(Contents, &Layout, Usage, Validation);
 }
 
-FORCEINLINE void RHIUpdateUniformBuffer(FRHIUniformBuffer* UniformBufferRHI, const void* Contents)
-{
-	return GDynamicRHI->RHIUpdateUniformBuffer(UniformBufferRHI, Contents);
-}
-
 FORCEINLINE FDynamicRHI::FRHICalcTextureSizeResult RHICalcTexturePlatformSize(FRHITextureDesc const& Desc, uint32 FirstMipIndex = 0)
 {
 	if ( ! Desc.IsValid() )
diff --git a/Engine/Source/Runtime/RHI/Public/RHICommandList.h b/Engine/Source/Runtime/RHI/Public/RHICommandList.h
--- a/Engine/Source/Runtime/RHI/Public/RHICommandList.h
+++ b/Engine/Source/Runtime/RHI/Public/RHICommandList.h
@@ -649,6 +649,11 @@ public:
 		return Buffer;
 	}
 
+	FORCEINLINE void UpdateUniformBuffer(FRHIUniformBuffer* UniformBufferRHI, const void* Contents)
+	{
+		GDynamicRHI->RHIUpdateUniformBuffer(*this, UniformBufferRHI, Contents);
+	}
+
 	void DisallowBypass()
 	{
 #if CAN_DISALLOW_COMMAND_LIST_BYPASS
@@ -5303,6 +5308,11 @@ FORCEINLINE void RHIUnlockIndexBuffer(FRHIBuffer* IndexBuffer)
 	FRHICommandListExecutor::GetImmediateCommandList().UnlockBuffer(IndexBuffer);
 }
 
+FORCEINLINE void RHIUpdateUniformBuffer(FRHIUniformBuffer* UniformBufferRHI, const void* Contents)
+{
+	return FRHICommandListExecutor::GetImmediateCommandList().UpdateUniformBuffer(UniformBufferRHI, Contents);
+}
+
 UE_DEPRECATED(5.0, "Use RHICreateBuffer() and RHILockBuffer() instead.")
 FORCEINLINE FBufferRHIRef RHICreateAndLockVertexBuffer(uint32 Size, EBufferUsageFlags InUsage, FRHIResourceCreateInfo& CreateInfo, void*& OutDataBuffer)
 {
diff --git a/Engine/Source/Runtime/RHI/Public/RHIValidation.h b/Engine/Source/Runtime/RHI/Public/RHIValidation.h
--- a/Engine/Source/Runtime/RHI/Public/RHIValidation.h
+++ b/Engine/Source/Runtime/RHI/Public/RHIValidation.h
@@ -369,10 +369,9 @@ public:
 		return UniformBuffer;
 	}
 
-	virtual void RHIUpdateUniformBuffer(FRHIUniformBuffer* UniformBufferRHI, const void* Contents) override final
+	virtual void RHIUpdateUniformBuffer(FRHICommandListBase& RHICmdList, FRHIUniformBuffer* UniformBufferRHI, const void* Contents) override final
 	{
-		check(IsInRenderingThread());
-		RHI->RHIUpdateUniformBuffer(UniformBufferRHI, Contents);
+		RHI->RHIUpdateUniformBuffer(RHICmdList, UniformBufferRHI, Contents);
 		UniformBufferRHI->UpdateAllocation(RenderThreadFrameID);
 	}
 
diff --git a/Engine/Source/Runtime/VulkanRHI/Private/VulkanUniformBuffer.cpp b/Engine/Source/Runtime/VulkanRHI/Private/VulkanUniformBuffer.cpp
--- a/Engine/Source/Runtime/VulkanRHI/Private/VulkanUniformBuffer.cpp
+++ b/Engine/Source/Runtime/VulkanRHI/Private/VulkanUniformBuffer.cpp
@@ -194,18 +194,15 @@ FUniformBufferRHIRef FVulkanDynamicRHI::RHICreateUniformBuffer(const void* Conte
 	return new FVulkanUniformBuffer(*Device, Layout, Contents, Usage, Validation);
 }
 
-inline void FVulkanDynamicRHI::UpdateUniformBuffer(FVulkanUniformBuffer* UniformBuffer, const void* Contents)
+inline void FVulkanDynamicRHI::UpdateUniformBuffer(FRHICommandListBase& RHICmdList, FVulkanUniformBuffer* UniformBuffer, const void* Contents)
 {
 	SCOPE_CYCLE_COUNTER(STAT_VulkanUpdateUniformBuffers);
-	check(IsInRenderingThread());
 
 	const FRHIUniformBufferLayout& Layout = UniformBuffer->GetLayout();
 
 	const int32 ConstantBufferSize = Layout.ConstantBufferSize;
 	const int32 NumResources = Layout.Resources.Num();
 
-	FRHICommandListImmediate& RHICmdList = FRHICommandListExecutor::GetImmediateCommandList();
-
 	FVulkanAllocation NewUBAlloc;
 	bool bUseUpload = GVulkanAllowUniformUpload && !RHICmdList.IsInsideRenderPass(); //inside renderpasses, a rename is enforced.
 	const bool bUseRingBuffer = UseRingBuffer(UniformBuffer->Usage);
@@ -256,7 +253,7 @@ inline void FVulkanDynamicRHI::UpdateUniformBuffer(FVulkanUniformBuffer* Uniform
 			void* CmdListConstantBufferData = RHICmdList.Alloc(ConstantBufferSize, 16);
 			FMemory::Memcpy(CmdListConstantBufferData, Contents, ConstantBufferSize);
 
-			RHICmdList.EnqueueLambda([UniformBuffer, CmdListResources, NumResources, ConstantBufferSize, CmdListConstantBufferData](FRHICommandList& CmdList)
+			RHICmdList.EnqueueLambda([UniformBuffer, CmdListResources, NumResources, ConstantBufferSize, CmdListConstantBufferData](FRHICommandListBase& CmdList)
 			{
 				FVulkanCommandListContext& Context = (FVulkanCommandListContext&)CmdList.GetContext().GetLowestLevelContext();
 				UpdateUniformBufferHelper(Context, UniformBuffer, ConstantBufferSize, CmdListConstantBufferData);
@@ -266,7 +263,7 @@ inline void FVulkanDynamicRHI::UpdateUniformBuffer(FVulkanUniformBuffer* Uniform
 		else
 		{
 			NewUBAlloc.Disown(); //this releases ownership while its put into the lambda
-			RHICmdList.EnqueueLambda([UniformBuffer, NewUBAlloc, CmdListResources, NumResources](FRHICommandList& CmdList)
+			RHICmdList.EnqueueLambda([UniformBuffer, NewUBAlloc, CmdListResources, NumResources](FRHICommandListBase& CmdList)
 			{
 				FVulkanAllocation Alloc;
 				Alloc.Reference(NewUBAlloc);
@@ -282,10 +279,10 @@ inline void FVulkanDynamicRHI::UpdateUniformBuffer(FVulkanUniformBuffer* Uniform
 }
 
 
-void FVulkanDynamicRHI::RHIUpdateUniformBuffer(FRHIUniformBuffer* UniformBufferRHI, const void* Contents)
+void FVulkanDynamicRHI::RHIUpdateUniformBuffer(FRHICommandListBase& RHICmdList, FRHIUniformBuffer* UniformBufferRHI, const void* Contents)
 {
 	FVulkanUniformBuffer* UniformBuffer = ResourceCast(UniformBufferRHI);
-	UpdateUniformBuffer(UniformBuffer, Contents);
+	UpdateUniformBuffer(RHICmdList, UniformBuffer, Contents);
 }
 
 FVulkanUniformBufferUploader::FVulkanUniformBufferUploader(FVulkanDevice* InDevice)
diff --git a/Engine/Source/Runtime/VulkanRHI/Public/VulkanDynamicRHI.h b/Engine/Source/Runtime/VulkanRHI/Public/VulkanDynamicRHI.h
--- a/Engine/Source/Runtime/VulkanRHI/Public/VulkanDynamicRHI.h
+++ b/Engine/Source/Runtime/VulkanRHI/Public/VulkanDynamicRHI.h
@@ -104,7 +104,7 @@ public:
 	virtual FGraphicsPipelineStateRHIRef RHICreateGraphicsPipelineState(const FGraphicsPipelineStateInitializer& Initializer) final override;
 	virtual FComputePipelineStateRHIRef RHICreateComputePipelineState(FRHIComputeShader* ComputeShader) final override;
 	virtual FUniformBufferRHIRef RHICreateUniformBuffer(const void* Contents, const FRHIUniformBufferLayout* Layout, EUniformBufferUsage Usage, EUniformBufferValidation Validation) final override;
-	virtual void RHIUpdateUniformBuffer(FRHIUniformBuffer* UniformBufferRHI, const void* Contents) final override;
+	virtual void RHIUpdateUniformBuffer(FRHICommandListBase& RHICmdList, FRHIUniformBuffer* UniformBufferRHI, const void* Contents) final override;
 	virtual FBufferRHIRef RHICreateBuffer(FRHICommandListBase& RHICmdList, uint32 Size, EBufferUsageFlags Usage, uint32 Stride, ERHIAccess ResourceState, FRHIResourceCreateInfo& CreateInfo) final override;
 	virtual void RHICopyBuffer(FRHIBuffer* SourceBuffer, FRHIBuffer* DestBuffer) final override;
 	virtual void RHITransferBufferUnderlyingResource(FRHIBuffer* DestBuffer, FRHIBuffer* SrcBuffer) final override;
@@ -446,7 +446,7 @@ protected:
 	void InternalUpdateTexture2D(bool bFromRenderingThread, FRHITexture2D* TextureRHI, uint32 MipIndex, const struct FUpdateTextureRegion2D& UpdateRegion, uint32 SourcePitch, const uint8* SourceData);
 	void InternalUpdateTexture3D(bool bFromRenderingThread, FRHITexture3D* TextureRHI, uint32 MipIndex, const struct FUpdateTextureRegion3D& UpdateRegion, uint32 SourceRowPitch, uint32 SourceDepthPitch, const uint8* SourceData);
 
-	void UpdateUniformBuffer(FVulkanUniformBuffer* UniformBuffer, const void* Contents);
+	void UpdateUniformBuffer(FRHICommandListBase& RHICmdList, FVulkanUniformBuffer* UniformBuffer, const void* Contents);
 
 	static void SetupValidationRequests();
 
diff --git a/Engine/Source/Runtime/Windows/D3D11RHI/Private/D3D11RHIPrivate.h b/Engine/Source/Runtime/Windows/D3D11RHI/Private/D3D11RHIPrivate.h
--- a/Engine/Source/Runtime/Windows/D3D11RHI/Private/D3D11RHIPrivate.h
+++ b/Engine/Source/Runtime/Windows/D3D11RHI/Private/D3D11RHIPrivate.h
@@ -475,7 +475,7 @@ public:
     virtual void RHIUnlockStagingBuffer(FRHIStagingBuffer* StagingBuffer) final override;
 	virtual FBoundShaderStateRHIRef RHICreateBoundShaderState(FRHIVertexDeclaration* VertexDeclaration, FRHIVertexShader* VertexShader, FRHIPixelShader* PixelShader, FRHIGeometryShader* GeometryShader) final override;
 	virtual FUniformBufferRHIRef RHICreateUniformBuffer(const void* Contents, const FRHIUniformBufferLayout* Layout, EUniformBufferUsage Usage, EUniformBufferValidation Validation) final override;
-	virtual void RHIUpdateUniformBuffer(FRHIUniformBuffer* UniformBufferRHI, const void* Contents) final override;
+	virtual void RHIUpdateUniformBuffer(FRHICommandListBase& RHICmdList, FRHIUniformBuffer* UniformBufferRHI, const void* Contents) final override;
 	virtual FBufferRHIRef RHICreateBuffer(FRHICommandListBase& RHICmdList, uint32 Size, EBufferUsageFlags Usage, uint32 Stride, ERHIAccess ResourceState, FRHIResourceCreateInfo& CreateInfo) final override;
 	virtual void RHICopyBuffer(FRHIBuffer* SourceBuffer, FRHIBuffer* DestBuffer) final override;
 	virtual void RHITransferBufferUnderlyingResource(FRHIBuffer* DestBuffer, FRHIBuffer* SrcBuffer) final override;
diff --git a/Engine/Source/Runtime/Windows/D3D11RHI/Private/D3D11UniformBuffer.cpp b/Engine/Source/Runtime/Windows/D3D11RHI/Private/D3D11UniformBuffer.cpp
--- a/Engine/Source/Runtime/Windows/D3D11RHI/Private/D3D11UniformBuffer.cpp
+++ b/Engine/Source/Runtime/Windows/D3D11RHI/Private/D3D11UniformBuffer.cpp
@@ -281,7 +281,7 @@ void UpdateUniformBufferContents(FD3D11Device* Direct3DDevice, FD3D11DeviceConte
 	}
 }
 
-void FD3D11DynamicRHI::RHIUpdateUniformBuffer(FRHIUniformBuffer* UniformBufferRHI, const void* Contents)
+void FD3D11DynamicRHI::RHIUpdateUniformBuffer(FRHICommandListBase& RHICmdList, FRHIUniformBuffer* UniformBufferRHI, const void* Contents)
 {
 	check(IsInRenderingThread());
 	check(UniformBufferRHI);
@@ -295,8 +295,6 @@ void FD3D11DynamicRHI::RHIUpdateUniformBuffer(FRHIUniformBuffer* UniformBufferRH
 
 	check(UniformBuffer->ResourceTable.Num() == NumResources);
 
-	FRHICommandListImmediate& RHICmdList = FRHICommandListExecutor::GetImmediateCommandList();
-
 	if (RHICmdList.Bypass())
 	{
 		UpdateUniformBufferContents(Direct3DDevice, Direct3DDeviceIMContext, UniformBuffer, Contents, ConstantBufferSize);
@@ -335,7 +333,7 @@ void FD3D11DynamicRHI::RHIUpdateUniformBuffer(FRHIUniformBuffer* UniformBufferRH
 			CmdListResources,
 			NumResources,
 			CmdListConstantBufferData,
-			ConstantBufferSize](FRHICommandList&)
+			ConstantBufferSize](FRHICommandListBase&)
 		{
 			UpdateUniformBufferContents(Direct3DDevice, Direct3DDeviceIMContext, UniformBuffer, CmdListConstantBufferData, ConstantBufferSize);
 
