Added support for multiple vkPipelineCaches

diff --git a/Engine/Source/Runtime/VulkanRHI/Private/VulkanPipeline.cpp b/Engine/Source/Runtime/VulkanRHI/Private/VulkanPipeline.cpp
--- a/Engine/Source/Runtime/VulkanRHI/Private/VulkanPipeline.cpp
+++ b/Engine/Source/Runtime/VulkanRHI/Private/VulkanPipeline.cpp
@@ -308,8 +308,7 @@ void FVulkanRHIGraphicsPipelineState::PurgeLoadedShaderModules(FVulkanDevice* In
 FVulkanPipelineStateCacheManager::FVulkanPipelineStateCacheManager(FVulkanDevice* InDevice)
 	: Device(InDevice)
 	, bEvictImmediately(false)
-	, bLinkedToPSOFC(false)
-	, bLinkedToPSOFCSucessfulLoaded(false)
+	, bPrecompilingCacheLoadedFromFile(false)
 {
 	bUseLRU = (int32)CVarEnableLRU.GetValueOnAnyThread() != 0;
 	LRUUsedPipelineMax = CVarLRUPipelineCapacity.GetValueOnAnyThread();
@@ -318,18 +317,14 @@ FVulkanPipelineStateCacheManager::FVulkanPipelineStateCacheManager(FVulkanDevice
 
 FVulkanPipelineStateCacheManager::~FVulkanPipelineStateCacheManager()
 {
-
-	if (bLinkedToPSOFC)
+	if (OnShaderPipelineCacheOpenedDelegate.IsValid())
 	{
-		if (OnShaderPipelineCacheOpenedDelegate.IsValid())
-		{
-			FShaderPipelineCache::GetCacheOpenedDelegate().Remove(OnShaderPipelineCacheOpenedDelegate);
-		}
+		FShaderPipelineCache::GetCacheOpenedDelegate().Remove(OnShaderPipelineCacheOpenedDelegate);
+	}
 
-		if (OnShaderPipelineCachePrecompilationCompleteDelegate.IsValid())
-		{
-			FShaderPipelineCache::GetPrecompilationCompleteDelegate().Remove(OnShaderPipelineCachePrecompilationCompleteDelegate);
-		}
+	if (OnShaderPipelineCachePrecompilationCompleteDelegate.IsValid())
+	{
+		FShaderPipelineCache::GetPrecompilationCompleteDelegate().Remove(OnShaderPipelineCachePrecompilationCompleteDelegate);
 	}
 	DestroyCache();
 
@@ -344,15 +339,24 @@ FVulkanPipelineStateCacheManager::~FVulkanPipelineStateCacheManager()
 	}
 
 	{
-		FScopedPipelineCache PipelineCacheExclusive = PSOCache.Get(EPipelineCacheAccess::Exclusive);
+		//TODO: Save PSOCache here?!
+		FScopedPipelineCache PipelineCacheExclusive = GlobalPSOCache.Get(EPipelineCacheAccess::Exclusive);
 		VulkanRHI::vkDestroyPipelineCache(Device->GetInstanceHandle(), PipelineCacheExclusive.Get(), VULKAN_CPU_ALLOCATOR);
-		PipelineCacheExclusive.Get() = VK_NULL_HANDLE;
+	}
+
+	{
+		FScopedPipelineCache PipelineCacheExclusive = CurrentPrecompilingPSOCache.Get(EPipelineCacheAccess::Exclusive);
+		if (PipelineCacheExclusive.Get() != VK_NULL_HANDLE)
+		{
+			//If CurrentOpenedPSOCache is still valid then it has never received OnShaderPipelineCachePrecompilationComplete callback, so we are not going to save it's content to disk as it's most likely is incomplete at this point
+			VulkanRHI::vkDestroyPipelineCache(Device->GetInstanceHandle(), PipelineCacheExclusive.Get(), VULKAN_CPU_ALLOCATOR);
+		}
 	}
 }
 
-bool FVulkanPipelineStateCacheManager::Load(const TArray<FString>& CacheFilenames)
+bool FVulkanPipelineStateCacheManager::Load(const TArray<FString>& CacheFilenames, FPipelineCache& Cache)
 {
-	FScopedPipelineCache PipelineCacheExclusive = PSOCache.Get(EPipelineCacheAccess::Exclusive);
+	FScopedPipelineCache PipelineCacheExclusive = Cache.Get(EPipelineCacheAccess::Exclusive);
 
 	bool bResult = false;
 	// Try to load device cache first
@@ -373,12 +377,12 @@ bool FVulkanPipelineStateCacheManager::Load(const TArray<FString>& CacheFilename
 
 				if (PipelineCacheExclusive.Get() == VK_NULL_HANDLE)
 				{
-					// if we don't have one already, then create our main cache (PipelineCache)
 					VERIFYVULKANRESULT(VulkanRHI::vkCreatePipelineCache(Device->GetInstanceHandle(), &PipelineCacheInfo, VULKAN_CPU_ALLOCATOR, &PipelineCacheExclusive.Get()));
 				}
 				else
 				{
-					// if we have one already, create a temp one and merge into the main cache
+					//TODO: assert on reopening the same cache twice?!
+					// if we have one already, create a temp one and merge it
 					VkPipelineCache TempPipelineCache;
 					VERIFYVULKANRESULT(VulkanRHI::vkCreatePipelineCache(Device->GetInstanceHandle(), &PipelineCacheInfo, VULKAN_CPU_ALLOCATOR, &TempPipelineCache));
 					VERIFYVULKANRESULT(VulkanRHI::vkMergePipelineCaches(Device->GetInstanceHandle(), PipelineCacheExclusive.Get(), 1, &TempPipelineCache));
@@ -400,6 +404,7 @@ bool FVulkanPipelineStateCacheManager::Load(const TArray<FString>& CacheFilename
 		}
 	}
 
+	//TODO: how to load LRU cache as it will have info about PSOs from multiple caches
 	if(CVarEnableLRU.GetValueOnAnyThread() != 0)
 	{
 		for (const FString& CacheFilename : CacheFilenames)
@@ -457,29 +462,27 @@ void FVulkanPipelineStateCacheManager::InitAndLoad(const TArray<FString>& CacheF
 	{
 		if (GPipelineCacheFromShaderPipelineCacheCvar.GetValueOnAnyThread() == 0)
 		{
-			Load(CacheFilenames);
+			Load(CacheFilenames, GlobalPSOCache);
 		}
 		else
 		{
-			bLinkedToPSOFC = true;
 			UE_LOG(LogVulkanRHI, Log, TEXT("FVulkanPipelineStateCacheManager will check for loading, etc when ShaderPipelineCache opens its file"));
 
-
 #if PLATFORM_ANDROID && USE_ANDROID_FILE
 			// @todo Lumin: Use that GetPathForExternalWrite or something?
 			// BTW, this is totally bad. We should not platform ifdefs like this, rather the HAL needs to be extended!
 			extern FString GExternalFilePath;
-			LinkedToPSOFCCacheFolderPath = GExternalFilePath / TEXT("VulkanProgramBinaryCache");
+			CompiledPSOCacheTopFolderPath = GExternalFilePath / TEXT("VulkanProgramBinaryCache");
 
 #else
-			LinkedToPSOFCCacheFolderPath = FPaths::ProjectSavedDir() / TEXT("VulkanProgramBinaryCache");
+			CompiledPSOCacheTopFolderPath = FPaths::ProjectSavedDir() / TEXT("VulkanProgramBinaryCache");
 #endif
 
 			// Remove entire ProgramBinaryCache folder if -ClearOpenGLBinaryProgramCache is specified on command line
 			if (FParse::Param(FCommandLine::Get(), TEXT("ClearVulkanBinaryProgramCache")))
 			{
-				UE_LOG(LogVulkanRHI, Log, TEXT("FVulkanPipelineStateCacheManager: Deleting binary program cache folder for -ClearVulkanBinaryProgramCache: %s"), *LinkedToPSOFCCacheFolderPath);
-				FPlatformFileManager::Get().GetPlatformFile().DeleteDirectoryRecursively(*LinkedToPSOFCCacheFolderPath);
+				UE_LOG(LogVulkanRHI, Log, TEXT("FVulkanPipelineStateCacheManager: Deleting binary program cache folder for -ClearVulkanBinaryProgramCache: %s"), *CompiledPSOCacheTopFolderPath);
+				FPlatformFileManager::Get().GetPlatformFile().DeleteDirectoryRecursively(*CompiledPSOCacheTopFolderPath);
 			}
 
 			OnShaderPipelineCacheOpenedDelegate = FShaderPipelineCache::GetCacheOpenedDelegate().AddRaw(this, &FVulkanPipelineStateCacheManager::OnShaderPipelineCacheOpened);
@@ -487,7 +490,7 @@ void FVulkanPipelineStateCacheManager::InitAndLoad(const TArray<FString>& CacheF
 		}
 	}
 
-	FScopedPipelineCache PipelineCacheExclusive = PSOCache.Get(EPipelineCacheAccess::Exclusive);
+	FScopedPipelineCache PipelineCacheExclusive = GlobalPSOCache.Get(EPipelineCacheAccess::Exclusive);
 	// Lazily create the cache in case the load failed
 	if (PipelineCacheExclusive.Get() == VK_NULL_HANDLE)
 	{
@@ -497,9 +500,24 @@ void FVulkanPipelineStateCacheManager::InitAndLoad(const TArray<FString>& CacheF
 	}
 }
 
+void FVulkanPipelineStateCacheManager::Save(const FString& CacheFilename)
+{
+	SavePSOCache(CacheFilename, GlobalPSOCache);
+
+	//TODO: Save LRU cache here
+}
+
 void FVulkanPipelineStateCacheManager::OnShaderPipelineCacheOpened(FString const& Name, EShaderPlatform Platform, uint32 Count, const FGuid& VersionGuid, FShaderPipelineCache::FShaderCachePrecompileContext& ShaderCachePrecompileContext)
 {
-	check(bLinkedToPSOFC);
+	//TODO: support reloading the same cache
+	if (CompiledPSOCaches.Contains(VersionGuid))
+	{
+		UE_LOG(LogVulkanRHI, Log, TEXT("FVulkanPipelineStateCacheManager::OnShaderPipelineCacheOpened attempts to load a cache that was already loaded before %s %s"), *Name, *VersionGuid.ToString());
+		return;
+	}
+
+	CurrentPrecompilingPSOCacheGuid = VersionGuid;
+
 	UE_LOG(LogVulkanRHI, Log, TEXT("FVulkanPipelineStateCacheManager::OnShaderPipelineCacheOpened %s %d %s"), *Name, Count, *VersionGuid.ToString());
 
 	IPlatformFile& PlatformFile = FPlatformFileManager::Get().GetPlatformFile();
@@ -507,82 +525,92 @@ void FVulkanPipelineStateCacheManager::OnShaderPipelineCacheOpened(FString const
 	const VkPhysicalDeviceProperties& DeviceProperties = Device->GetDeviceProperties();
 	FString BinaryCacheAppendage = FString::Printf(TEXT(".%x.%x"), DeviceProperties.vendorID, DeviceProperties.deviceID);
 
-	LinkedToPSOFCCacheFolderFilename = LinkedToPSOFCCacheFolderPath / TEXT("VulkanPSO_") + VersionGuid.ToString() + BinaryCacheAppendage;
-	FString TempName = LinkedToPSOFCCacheFolderPath / TEXT("TempScanVulkanPSO_") + VersionGuid.ToString() + BinaryCacheAppendage;
+	CompiledPSOCacheFolderName = CompiledPSOCacheTopFolderPath / TEXT("VulkanPSO_") + VersionGuid.ToString() + BinaryCacheAppendage;
+	FString TempName = CompiledPSOCacheTopFolderPath / TEXT("TempScanVulkanPSO_") + VersionGuid.ToString() + BinaryCacheAppendage;
 
-	bool bSuccess = false;
+	{
+		FScopedPipelineCache PipelineCacheExclusive = CurrentPrecompilingPSOCache.Get(EPipelineCacheAccess::Exclusive);
+		checkf(PipelineCacheExclusive.Get() == VK_NULL_HANDLE, TEXT("Trying to open more than one shader pipeline cache"));
 
-	if (PlatformFile.FileExists(*LinkedToPSOFCCacheFolderFilename))
+		VkPipelineCacheCreateInfo PipelineCacheInfo;
+		ZeroVulkanStruct(PipelineCacheInfo, VK_STRUCTURE_TYPE_PIPELINE_CACHE_CREATE_INFO);
+		VERIFYVULKANRESULT(VulkanRHI::vkCreatePipelineCache(Device->GetInstanceHandle(), &PipelineCacheInfo, VULKAN_CPU_ALLOCATOR, &PipelineCacheExclusive.Get()));
+	}
+
+	if (PlatformFile.FileExists(*CompiledPSOCacheFolderName))
 	{
 		// Try to move the file to a temporary filename before the scan, so we won't try to read it again if it's corrupted
 		PlatformFile.DeleteFile(*TempName);
-		PlatformFile.MoveFile(*TempName, *LinkedToPSOFCCacheFolderFilename);
+		PlatformFile.MoveFile(*TempName, *CompiledPSOCacheFolderName);
 
 		TArray<FString> CacheFilenames;
 		CacheFilenames.Add(TempName);
-		bSuccess = Load(CacheFilenames);
 
 		// Rename the file back after a successful scan.
-		if (bSuccess)
+		if (Load(CacheFilenames, CurrentPrecompilingPSOCache))
 		{
-			bLinkedToPSOFCSucessfulLoaded = true;
-			PlatformFile.MoveFile(*LinkedToPSOFCCacheFolderFilename, *TempName);
+			bPrecompilingCacheLoadedFromFile = true;
+			PlatformFile.MoveFile(*CompiledPSOCacheFolderName, *TempName);
 
 			if (CVarPipelineLRUCacheEvictBinary.GetValueOnAnyThread())
 			{
 				bEvictImmediately = true;
 			}
 		}
+		else
+		{
+			UE_LOG(LogVulkanRHI, Log, TEXT("FVulkanPipelineStateCacheManager: PSO cache failed to load, deleting file: %s"), *CompiledPSOCacheFolderName);
+			FPlatformFileManager::Get().GetPlatformFile().DeleteFile(*CompiledPSOCacheFolderName);
+		}
 	}
 	else
 	{
-		UE_LOG(LogVulkanRHI, Log, TEXT("FVulkanPipelineStateCacheManager: %s does not exist."), *LinkedToPSOFCCacheFolderFilename);
-	}
-	if (!bSuccess)
-	{
-		UE_LOG(LogVulkanRHI, Log, TEXT("FVulkanPipelineStateCacheManager: No matching vulkan PSO cache found or it failed to load, deleting binary program cache folder: %s"), *LinkedToPSOFCCacheFolderPath);
-		FPlatformFileManager::Get().GetPlatformFile().DeleteDirectoryRecursively(*LinkedToPSOFCCacheFolderPath);
+		UE_LOG(LogVulkanRHI, Log, TEXT("FVulkanPipelineStateCacheManager: %s does not exist."), *CompiledPSOCacheFolderName);
 	}
 
+	if (!bPrecompilingCacheLoadedFromFile || (bEvictImmediately && CVarPipelineLRUCacheEvictBinaryPreloadScreen.GetValueOnAnyThread()))
 	{
-		if (!bLinkedToPSOFCSucessfulLoaded || (bEvictImmediately && CVarPipelineLRUCacheEvictBinaryPreloadScreen.GetValueOnAnyThread()))
-		{
-			ShaderCachePrecompileContext.SetPrecompilationIsSlowTask();
-		}
+		//TODO: This!
+		ShaderCachePrecompileContext.SetPrecompilationIsSlowTask();
 	}
 }
 
 void FVulkanPipelineStateCacheManager::OnShaderPipelineCachePrecompilationComplete(uint32 Count, double Seconds, const FShaderPipelineCache::FShaderCachePrecompileContext& ShaderCachePrecompileContext)
 {
-	check(bLinkedToPSOFC);
 	UE_LOG(LogVulkanRHI, Log, TEXT("FVulkanPipelineStateCacheManager::OnShaderPipelineCachePrecompilationComplete"));
 
 	bEvictImmediately = false;
-	if (!bLinkedToPSOFCSucessfulLoaded)
+	if (!bPrecompilingCacheLoadedFromFile)
+	{
+		//Save PSO cache only if it failed to load
+		SavePSOCache(CompiledPSOCacheFolderName, CurrentPrecompilingPSOCache);
+	}
+
+	if (!CompiledPSOCaches.Contains(CurrentPrecompilingPSOCacheGuid))
 	{
-		Save(LinkedToPSOFCCacheFolderFilename, true);
+		CompiledPSOCaches.Add(CurrentPrecompilingPSOCacheGuid);
+
+		//Merge this CurrentOpenedPSOCache with global PSOCache
+		QUICK_SCOPE_CYCLE_COUNTER(STAT_VulkanPSOCacheMerge);
+		FScopedPipelineCache GlobalPipelineCacheExclusive = GlobalPSOCache.Get(EPipelineCacheAccess::Exclusive);
+		FScopedPipelineCache CurrentPipelineCacheExclusive = CurrentPrecompilingPSOCache.Get(EPipelineCacheAccess::Exclusive);
+		VERIFYVULKANRESULT(VulkanRHI::vkMergePipelineCaches(Device->GetInstanceHandle(), GlobalPipelineCacheExclusive.Get(), 1, &CurrentPipelineCacheExclusive.Get()));
+		VulkanRHI::vkDestroyPipelineCache(Device->GetInstanceHandle(), CurrentPipelineCacheExclusive.Get(), VULKAN_CPU_ALLOCATOR);
+
+		CurrentPipelineCacheExclusive.Get() = VK_NULL_HANDLE;
 	}
 
-	// Want to ignore any subsequent Shader Pipeline Cache opening/closing, eg when loading modules
-	FShaderPipelineCache::GetCacheOpenedDelegate().Remove(OnShaderPipelineCacheOpenedDelegate);
-	FShaderPipelineCache::GetPrecompilationCompleteDelegate().Remove(OnShaderPipelineCachePrecompilationCompleteDelegate);
-	OnShaderPipelineCacheOpenedDelegate.Reset();
-	OnShaderPipelineCachePrecompilationCompleteDelegate.Reset();
+	bPrecompilingCacheLoadedFromFile = false;
+	CurrentPrecompilingPSOCacheGuid = FGuid();
 }
 
-void FVulkanPipelineStateCacheManager::Save(const FString& CacheFilename, bool bFromPSOFC)
+void FVulkanPipelineStateCacheManager::SavePSOCache(const FString& CacheFilename, FPipelineCache& Cache)
 {
-	FScopedPipelineCache PipelineCacheExclusive = PSOCache.Get(EPipelineCacheAccess::Exclusive);
-	if (bLinkedToPSOFC && !bFromPSOFC)
-	{
-		UE_LOG(LogVulkanRHI, Log, TEXT("FVulkanPipelineStateCacheManager: skipped saving because we only save if the PSOFC based one failed to load."));
-		return;
-	}
-	FScopeLock Lock1(&GraphicsPSOLockedCS);
+	FScopedPipelineCache PipelineCacheExclusive = Cache.Get(EPipelineCacheAccess::Exclusive);
+	FScopeLock Lock1(&GraphicsPSOLockedCS);		//TODO: Do we really need this here?!
 	FScopeLock Lock2(&LRUCS);
 
 
-
 	// First save Device Cache
 	size_t Size = 0;
 	VERIFYVULKANRESULT(VulkanRHI::vkGetPipelineCacheData(Device->GetInstanceHandle(), PipelineCacheExclusive.Get(), &Size, nullptr));
@@ -609,6 +637,7 @@ void FVulkanPipelineStateCacheManager::Save(const FString& CacheFilename, bool b
 		{
 			UE_LOG(LogVulkanRHI, Warning, TEXT("Failed to get Vulkan pipeline cache data. Error %d, %d bytes"), Result, Size);
 
+			//TODO: Resave it when we shutdown the manager?!
 			VulkanRHI::vkDestroyPipelineCache(Device->GetInstanceHandle(), PipelineCacheExclusive.Get(), VULKAN_CPU_ALLOCATOR);
 			VkPipelineCacheCreateInfo PipelineCacheInfo;
 			ZeroVulkanStruct(PipelineCacheInfo, VK_STRUCTURE_TYPE_PIPELINE_CACHE_CREATE_INFO);
@@ -1097,7 +1126,7 @@ static const TMap<uint8, VkFragmentShadingRateCombinerOpKHR> FragmentCombinerOpM
 };
 #endif
 
-bool FVulkanPipelineStateCacheManager::CreateGfxPipelineFromEntry(FVulkanRHIGraphicsPipelineState* PSO, FVulkanShader* Shaders[ShaderStage::NumStages], VkPipeline* Pipeline)
+bool FVulkanPipelineStateCacheManager::CreateGfxPipelineFromEntry(FVulkanRHIGraphicsPipelineState* PSO, FVulkanShader* Shaders[ShaderStage::NumStages], VkPipeline* Pipeline, FPipelineCache& Cache)
 {
 	FGfxPipelineDesc* GfxEntry = &PSO->Desc;
 	if (Shaders[ShaderStage::Pixel] == nullptr && !FVulkanPlatform::SupportsNullPixelShader())
@@ -1274,7 +1303,7 @@ bool FVulkanPipelineStateCacheManager::CreateGfxPipelineFromEntry(FVulkanRHIGrap
 		{
 			SCOPE_CYCLE_COUNTER(STAT_VulkanPSOVulkanCreationTime);
 			FoundSize = Found->PipelineSize;
-			FScopedPipelineCache PipelineCacheShared = PSOCache.Get(EPipelineCacheAccess::Shared);
+			FScopedPipelineCache PipelineCacheShared = Cache.Get(EPipelineCacheAccess::Shared);
 			Result = VulkanRHI::vkCreateGraphicsPipelines(Device->GetInstanceHandle(), PipelineCacheShared.Get(), 1, &PipelineInfo, VULKAN_CPU_ALLOCATOR, Pipeline);
 		}
 		else
@@ -1310,13 +1339,13 @@ bool FVulkanPipelineStateCacheManager::CreateGfxPipelineFromEntry(FVulkanRHIGrap
 
 					{
 						QUICK_SCOPE_CYCLE_COUNTER(STAT_VulkanPSOCacheMerge);
-						FScopedPipelineCache PipelineCacheExclusive = PSOCache.Get(EPipelineCacheAccess::Exclusive);
+						FScopedPipelineCache PipelineCacheExclusive = Cache.Get(EPipelineCacheAccess::Exclusive);
 						VERIFYVULKANRESULT(VulkanRHI::vkMergePipelineCaches(Device->GetInstanceHandle(), PipelineCacheExclusive.Get(), 1, &LocalPipelineCache));
 					}
 				}
 				else
 				{
-					FoundSize = Found->PipelineSize;;
+					FoundSize = Found->PipelineSize;
 				}
 			}
 			else
@@ -1334,7 +1363,7 @@ bool FVulkanPipelineStateCacheManager::CreateGfxPipelineFromEntry(FVulkanRHIGrap
 	else
 	{
 		SCOPE_CYCLE_COUNTER(STAT_VulkanPSOVulkanCreationTime);
-		FScopedPipelineCache PipelineCacheShared = PSOCache.Get(EPipelineCacheAccess::Shared);
+		FScopedPipelineCache PipelineCacheShared = Cache.Get(EPipelineCacheAccess::Shared);
 		Result = VulkanRHI::vkCreateGraphicsPipelines(Device->GetInstanceHandle(), PipelineCacheShared.Get(), 1, &PipelineInfo, VULKAN_CPU_ALLOCATOR, Pipeline);
 	}
 
@@ -1868,7 +1897,8 @@ FVulkanRHIGraphicsPipelineState* FVulkanPipelineStateCacheManager::RHICreateGrap
 			
 				QUICK_SCOPE_CYCLE_COUNTER(STAT_Vulkan_RHICreateGraphicsPipelineState_CREATE_PART0);
 
-				if(!CreateGfxPipelineFromEntry(NewPSO, VulkanShaders, &NewPSO->VulkanPipeline))
+				FPipelineCache& Cache = Initializer.bFromPSOFileCache ? CurrentPrecompilingPSOCache : GlobalPSOCache;
+				if(!CreateGfxPipelineFromEntry(NewPSO, VulkanShaders, &NewPSO->VulkanPipeline, Cache))
 				{
 					DeleteNewPSO(NewPSO);
 					return nullptr;
@@ -1904,8 +1934,6 @@ FVulkanRHIGraphicsPipelineState* FVulkanPipelineStateCacheManager::RHICreateGrap
 				{
 					NewPSO->bIsRegistered = true;
 				}
-
-
 			}
 		}
 	}
@@ -2011,7 +2039,7 @@ FVulkanComputePipeline* FVulkanPipelineStateCacheManager::CreateComputePipelineF
 	VkResult Result;
 	{
 		QUICK_SCOPE_CYCLE_COUNTER(STAT_VulkanComputePSOCreate);
-		FScopedPipelineCache PipelineCacheShared = PSOCache.Get(EPipelineCacheAccess::Shared);
+		FScopedPipelineCache PipelineCacheShared = GlobalPSOCache.Get(EPipelineCacheAccess::Shared);
 		Result = VulkanRHI::vkCreateComputePipelines(Device->GetInstanceHandle(), PipelineCacheShared.Get(), 1, &PipelineInfo, VULKAN_CPU_ALLOCATOR, &Pipeline->Pipeline);
 	}
 
@@ -2263,7 +2291,7 @@ void FVulkanPipelineStateCacheManager::LRUTouch(FVulkanRHIGraphicsPipelineState*
 
 			QUICK_SCOPE_CYCLE_COUNTER(STAT_Vulkan_LRUMiss);
 
-			if (!CreateGfxPipelineFromEntry(PSO, VulkanShaders, &PSO->VulkanPipeline))
+			if (!CreateGfxPipelineFromEntry(PSO, VulkanShaders, &PSO->VulkanPipeline, GlobalPSOCache))
 			{
 				check(0);
 			}
@@ -2304,6 +2332,7 @@ void FVulkanRHIGraphicsPipelineState::DeleteVkPipeline(bool bImmediate)
 
 	Device->PipelineStateCache->LRUCheckNotInside(this);
 }
+
 void FVulkanPipelineStateCacheManager::LRUCheckNotInside(FVulkanRHIGraphicsPipelineState* PSO)
 {
 	FScopeLock Lock(&LRUCS);
diff --git a/Engine/Source/Runtime/VulkanRHI/Private/VulkanPipeline.h b/Engine/Source/Runtime/VulkanRHI/Private/VulkanPipeline.h
--- a/Engine/Source/Runtime/VulkanRHI/Private/VulkanPipeline.h
+++ b/Engine/Source/Runtime/VulkanRHI/Private/VulkanPipeline.h
@@ -98,6 +98,7 @@ struct FVulkanShaderHashes
 
 struct FVulkanPipelineSize
 {
+	//TODO: increase hash to 12 bytes to reduce collisions
 	uint64 ShaderHash;
 	uint32 PipelineSize;
 
@@ -464,7 +465,7 @@ public:
 
 	// Array of potential cache locations; first entries have highest priority. Only one cache file is loaded. If unsuccessful, tries next entry in the array.
 	void InitAndLoad(const TArray<FString>& CacheFilenames);
-	void Save(const FString& CacheFilename, bool bFromPSOFC = false);
+	void Save(const FString& CacheFilename);
 
 	FVulkanPipelineStateCacheManager(FVulkanDevice* InParent);
 	~FVulkanPipelineStateCacheManager();
@@ -475,18 +476,21 @@ public:
 	void NotifyDeletedComputePipeline(FVulkanComputePipeline* Pipeline);
 
 private:
+	class FPipelineCache;
+
 	/** Delegate handlers to track the ShaderPipelineCache precompile. */
 	void OnShaderPipelineCacheOpened(FString const& Name, EShaderPlatform Platform, uint32 Count, const FGuid& VersionGuid, FShaderPipelineCache::FShaderCachePrecompileContext& ShaderCachePrecompileContext);
 	void OnShaderPipelineCachePrecompilationComplete(uint32 Count, double Seconds, const FShaderPipelineCache::FShaderCachePrecompileContext& ShaderCachePrecompileContext);
 
 	void CreateGfxEntry(const FGraphicsPipelineStateInitializer& PSOInitializer, FVulkanDescriptorSetsLayoutInfo& DescriptorSetLayoutInfo, FGfxPipelineDesc* Desc);
-	bool Load(const TArray<FString>& CacheFilenames);
+	bool Load(const TArray<FString>& CacheFilenames, FPipelineCache& Cache);
+	void SavePSOCache(const FString& CacheFilename, FPipelineCache& Cache);
 	void DestroyCache();
 
 	FVulkanRHIGraphicsPipelineState* RHICreateGraphicsPipelineState(const FGraphicsPipelineStateInitializer& Initializer);
 	FVulkanComputePipeline* RHICreateComputePipelineState(FRHIComputeShader* ComputeShaderRHI);
 	void NotifyDeletedGraphicsPSO(FRHIGraphicsPipelineState* PSO);
-	bool CreateGfxPipelineFromEntry(FVulkanRHIGraphicsPipelineState* PSO, FVulkanShader* Shaders[ShaderStage::NumStages], VkPipeline* Pipeline);
+	bool CreateGfxPipelineFromEntry(FVulkanRHIGraphicsPipelineState* PSO, FVulkanShader* Shaders[ShaderStage::NumStages], VkPipeline* Pipeline, FPipelineCache& Cache);
 
 
 	FVulkanLayout* FindOrAddLayout(const FVulkanDescriptorSetsLayoutInfo& DescriptorSetLayoutInfo, bool bGfxLayout);
@@ -508,11 +512,8 @@ private:
 
 	FVulkanDevice* Device;
 	bool bEvictImmediately;
-	// if true, we will link to the PSOFC, loading later, when we have that guid and only if the guid matches, saving only if there is no match, and only saving after the PSOFC is done.
-	bool bLinkedToPSOFC;
-	bool bLinkedToPSOFCSucessfulLoaded;
-	FString LinkedToPSOFCCacheFolderPath;
-	FString LinkedToPSOFCCacheFolderFilename;
+	FString CompiledPSOCacheTopFolderPath;
+	FString CompiledPSOCacheFolderName;
 	FDelegateHandle OnShaderPipelineCacheOpenedDelegate;
 	FDelegateHandle OnShaderPipelineCachePrecompilationCompleteDelegate;
 
@@ -546,7 +547,14 @@ private:
 	public:
 		FScopedPipelineCache Get(EPipelineCacheAccess PipelineAccessType) { return FScopedPipelineCache(PipelineAccessType == EPipelineCacheAccess::Exclusive, PipelineCache, PipelineCacheLock); }
 	};
-	FPipelineCache PSOCache;
+	FPipelineCache GlobalPSOCache;		// contains all PSO caches opened during the program run as well as PSO objects created on the fly
+
+	FPipelineCache CurrentPrecompilingPSOCache;
+	// if true, we will link to the PSOFC, loading later, when we have that guid and only if the guid matches, saving only if there is no match, and only saving after the PSOFC is done.
+	bool bPrecompilingCacheLoadedFromFile;
+	FGuid CurrentPrecompilingPSOCacheGuid;
+
+	TSet<FGuid> CompiledPSOCaches;
 
 	FCriticalSection LayoutMapCS;
 	TMap<FVulkanDescriptorSetsLayoutInfo, FVulkanLayout*> LayoutMap;
@@ -562,7 +570,7 @@ private:
 	uint32 LRUUsedPipelineSize = 0;
 	uint32 LRUUsedPipelineCount = 0;
 	uint32 LRUUsedPipelineMax = 0;
-	TMap<uint32, FVulkanPipelineSize> LRU2SizeList;	// key: Shader hash (FShaderHash), value: pipeline size
+	TMap<uint64, FVulkanPipelineSize> LRU2SizeList;	// key: Shader hash (FShaderHash), value: pipeline size
 	bool bUseLRU = true;
 	friend class FVulkanDynamicRHI;
 	friend class FVulkanCommandListContext;
@@ -729,7 +737,6 @@ public:
 	uint32 PipelineCacheSize = UINT32_MAX;
 	FVulkanPSOKey							VulkanKey;
 
-
 #if VULKAN_PSO_CACHE_DEBUG
 	FPixelShaderRHIRef					PixelShaderRHI;
 	FVertexShaderRHIRef					VertexShaderRHI;
