Fixed: VK validation warning related to AFBC workaround setup

diff --git a/Engine/Source/Runtime/VulkanRHI/Private/Android/VulkanAndroidPlatform.cpp b/Engine/Source/Runtime/VulkanRHI/Private/Android/VulkanAndroidPlatform.cpp
--- a/Engine/Source/Runtime/VulkanRHI/Private/Android/VulkanAndroidPlatform.cpp
+++ b/Engine/Source/Runtime/VulkanRHI/Private/Android/VulkanAndroidPlatform.cpp
@@ -641,7 +641,7 @@ void FVulkanAndroidPlatform::SetupAFBCWorkaround(const FVulkanDevice& InDevice)
 	ImageCreateInfo.sharingMode = VK_SHARING_MODE_EXCLUSIVE;
 	ImageCreateInfo.queueFamilyIndexCount = 0;
 	ImageCreateInfo.pQueueFamilyIndices = nullptr;
-	ImageCreateInfo.initialLayout = VK_IMAGE_LAYOUT_SHADER_READ_ONLY_OPTIMAL;
+	ImageCreateInfo.initialLayout = VK_IMAGE_LAYOUT_UNDEFINED;
 
 	VkImage Image0;
 	VkMemoryRequirements Image0Mem;
