Do not try to filter out VK validation messages on Android as layer does not pass a valid MsgCode anymore

diff --git a/Engine/Source/Runtime/VulkanRHI/Private/VulkanDebug.cpp b/Engine/Source/Runtime/VulkanRHI/Private/VulkanDebug.cpp
--- a/Engine/Source/Runtime/VulkanRHI/Private/VulkanDebug.cpp
+++ b/Engine/Source/Runtime/VulkanRHI/Private/VulkanDebug.cpp
@@ -162,17 +162,8 @@ static VkBool32 VKAPI_PTR DebugReportFunction(
 		return VK_FALSE;
 	}
 
-	FString LayerCode = FString::Printf(TEXT("%s%x"), ANSI_TO_TCHAR(LayerPrefix), MsgCode);
-
-	static TSet<FString> SeenCodes;
-	if (GCVarUniqueValidationMessages->GetInt() == 0 || !SeenCodes.Contains(LayerCode))
-	{
-		VULKAN_REPORT_LOG(TEXT("*** [%s:%s] Obj 0x%p Loc %d %s"), ANSI_TO_TCHAR(MsgPrefix), *LayerCode, (void*)SrcObject, (uint32)Location, ANSI_TO_TCHAR(Msg));
-		if (GCVarUniqueValidationMessages->GetInt() == 1)
-		{
-			SeenCodes.Add(LayerCode);
-		}
-	}
+	// MsgCode seem to be always 0
+	VULKAN_REPORT_LOG(TEXT("*** [%s] Obj 0x%p Loc %d %s"), ANSI_TO_TCHAR(MsgPrefix), (void*)SrcObject, (uint32)Location, ANSI_TO_TCHAR(Msg));
 
 	return VK_FALSE;
 }
