Storing VK shader modules in TRefCountPtrs to prevent them from being destroyed while still being referenced from other threads.

diff --git a/Engine/Source/Runtime/VulkanRHI/Private/VulkanPipeline.cpp b/Engine/Source/Runtime/VulkanRHI/Private/VulkanPipeline.cpp
--- a/Engine/Source/Runtime/VulkanRHI/Private/VulkanPipeline.cpp
+++ b/Engine/Source/Runtime/VulkanRHI/Private/VulkanPipeline.cpp
@@ -52,9 +52,6 @@ static TAtomic<uint64> SPipelineGfxCount;
 
 static const double HitchTime = 1.0 / 1000.0;
 
-static FCriticalSection FVulkanShaderHandleCS;
-
-
 TAutoConsoleVariable<int32> CVarPipelineDebugForceEvictImmediately(
 	TEXT("r.Vulkan.PipelineDebugForceEvictImmediately"),
 	0,
@@ -254,57 +251,33 @@ FVulkanRHIGraphicsPipelineState::~FVulkanRHIGraphicsPipelineState()
 	}
 
 	Device->PipelineStateCache->NotifyDeletedGraphicsPSO(this);
-	if (bShaderModulesLoaded)
-	{
-		PurgeLoadedShaderModules(Device);
-	}
 }
 
-void FVulkanRHIGraphicsPipelineState::GetOrCreateShaderModules(FVulkanShader*const* Shaders)
+void FVulkanRHIGraphicsPipelineState::GetOrCreateShaderModules(TRefCountPtr<FVulkanShaderModule> (&ShaderModulesOUT)[ShaderStage::NumStages], FVulkanShader*const* Shaders)
 {
-	FScopeLock Lock(&FVulkanShaderHandleCS); // this lock is to stop concurrent access to VulkanShader's ShaderModules container. 
 	for (int32 Index = 0; Index < ShaderStage::NumStages; ++Index)
 	{
+		check(!ShaderModulesOUT[Index].IsValid());
 		FVulkanShader* Shader = Shaders[Index];
 		if (Shader)
 		{
-			ShaderModules[Index] = Shader->GetOrCreateHandle(Desc, Layout, Layout->GetDescriptorSetLayoutHash());
+			ShaderModulesOUT[Index] = Shader->GetOrCreateHandle(Desc, Layout, Layout->GetDescriptorSetLayoutHash());
 		}
 	}
 }
 
 void FVulkanRHIGraphicsPipelineState::PurgeShaderModules(FVulkanShader*const* Shaders)
 {
-	check(!bShaderModulesLoaded);
-
 	for (int32 Index = 0; Index < ShaderStage::NumStages; ++Index)
 	{
 		FVulkanShader* Shader = Shaders[Index];
 		if (Shader)
 		{
 			Shader->PurgeShaderModules();
-			ShaderModules[Index] = VK_NULL_HANDLE;
 		}
 	}
 }
 
-void FVulkanRHIGraphicsPipelineState::PurgeLoadedShaderModules(FVulkanDevice* InDevice)
-{
-	check(bShaderModulesLoaded);
-
-	for (int32 Index = 0; Index < ShaderStage::NumStages; ++Index)
-	{
-		if (ShaderModules[Index] != VK_NULL_HANDLE)
-		{
-			VulkanRHI::vkDestroyShaderModule(InDevice->GetInstanceHandle(), ShaderModules[Index], VULKAN_CPU_ALLOCATOR);
-			ShaderModules[Index] = VK_NULL_HANDLE;
-		}
-	}
-
-	bShaderModulesLoaded = false;
-}
-
-
 FVulkanPipelineStateCacheManager::FVulkanPipelineStateCacheManager(FVulkanDevice* InDevice)
 	: Device(InDevice)
 	, bEvictImmediately(false)
@@ -1134,10 +1107,9 @@ bool FVulkanPipelineStateCacheManager::CreateGfxPipelineFromEntry(FVulkanRHIGrap
 		Shaders[ShaderStage::Pixel] = ResourceCast(TShaderMapRef<FNULLPS>(GetGlobalShaderMap(GMaxRHIFeatureLevel)).GetPixelShader());
 	}
 
-	if (!PSO->bShaderModulesLoaded)
-	{
-		PSO->GetOrCreateShaderModules(Shaders);
-	}
+	TRefCountPtr<FVulkanShaderModule> ShaderModules[ShaderStage::NumStages];
+
+	PSO->GetOrCreateShaderModules(ShaderModules, Shaders);
 
 	// Pipeline
 	VkGraphicsPipelineCreateInfo PipelineInfo;
@@ -1191,7 +1163,7 @@ bool FVulkanPipelineStateCacheManager::CreateGfxPipelineFromEntry(FVulkanRHIGrap
 	ANSICHAR EntryPoints[ShaderStage::NumStages][24];
 	for (int32 ShaderStage = 0; ShaderStage < ShaderStage::NumStages; ++ShaderStage)
 	{
-		if (!PSO->ShaderModules[ShaderStage])
+		if (!ShaderModules[ShaderStage].IsValid())
 		{
 			continue;
 		}
@@ -1200,7 +1172,7 @@ bool FVulkanPipelineStateCacheManager::CreateGfxPipelineFromEntry(FVulkanRHIGrap
 		ShaderStages[PipelineInfo.stageCount].sType = VK_STRUCTURE_TYPE_PIPELINE_SHADER_STAGE_CREATE_INFO;
 		VkShaderStageFlagBits Stage = UEFrequencyToVKStageBit(ShaderStage::GetFrequencyForGfxStage(CurrStage));
 		ShaderStages[PipelineInfo.stageCount].stage = Stage;
-		ShaderStages[PipelineInfo.stageCount].module = PSO->ShaderModules[CurrStage];
+		ShaderStages[PipelineInfo.stageCount].module = ShaderModules[CurrStage]->GetVkShaderModule();
 		Shaders[ShaderStage]->GetEntryPoint(EntryPoints[PipelineInfo.stageCount], 24);
 		ShaderStages[PipelineInfo.stageCount].pName = EntryPoints[PipelineInfo.stageCount];
 		PipelineInfo.stageCount++;
@@ -1739,7 +1711,6 @@ FVulkanRHIGraphicsPipelineState::FVulkanRHIGraphicsPipelineState(FVulkanDevice*
 
 	PSOInitializer = PSOInitializer_;
 #endif
-	FMemory::Memset(ShaderModules, 0, sizeof(ShaderModules));
 	INC_DWORD_STAT(STAT_VulkanNumGraphicsPSOs);
 	INC_DWORD_STAT_BY(STAT_VulkanPSOKeyMemory, this->VulkanKey.GetDataRef().Num());
 }
@@ -2024,13 +1995,13 @@ FVulkanComputePipeline* FVulkanPipelineStateCacheManager::CreateComputePipelineF
 		ComputeLayout->ComputePipelineDescriptorInfo.Initialize(Layout->GetDescriptorSetsLayout().RemappingInfo);
 	}
 
-	VkShaderModule ShaderModule = Shader->GetOrCreateHandle(Layout, Layout->GetDescriptorSetLayoutHash());
+	TRefCountPtr<FVulkanShaderModule> ShaderModule = Shader->GetOrCreateHandle(Layout, Layout->GetDescriptorSetLayoutHash());
 
 	VkComputePipelineCreateInfo PipelineInfo;
 	ZeroVulkanStruct(PipelineInfo, VK_STRUCTURE_TYPE_COMPUTE_PIPELINE_CREATE_INFO);
 	PipelineInfo.stage.sType = VK_STRUCTURE_TYPE_PIPELINE_SHADER_STAGE_CREATE_INFO;
 	PipelineInfo.stage.stage = VK_SHADER_STAGE_COMPUTE_BIT;
-	PipelineInfo.stage.module = ShaderModule;
+	PipelineInfo.stage.module = ShaderModule->GetVkShaderModule();
 	// main_00000000_00000000
 	ANSICHAR EntryPoint[24];
 	Shader->GetEntryPoint(EntryPoint, 24);
@@ -2267,7 +2238,7 @@ void FVulkanPipelineStateCacheManager::LRUTouch(FVulkanRHIGraphicsPipelineState*
 	}
 	FScopeLock Lock(&LRUCS);
 	check((PSO->GetVulkanPipeline() == 0) == (PSO->LRUNode == 0));
-	
+
 	if(PSO->LRUNode)
 	{
 		check(PSO->GetVulkanPipeline());
@@ -2370,7 +2341,6 @@ void FVulkanPipelineStateCacheManager::LRURemove(FVulkanRHIGraphicsPipelineState
 		PSO->DeleteVkPipeline(bImmediate);
 		if (GVulkanReleaseShaderModuleWhenEvictingPSO)
 		{
-			FScopeLock Lock(&FVulkanShaderHandleCS); // this lock is here to stop concurrent access to VulkanShaders ShaderModules container. 
 	        for (int ShaderStageIndex = 0; ShaderStageIndex < ShaderStage::NumStages; ShaderStageIndex++)
 	        {
 				if (PSO->VulkanShaders[ShaderStageIndex] != nullptr)
diff --git a/Engine/Source/Runtime/VulkanRHI/Private/VulkanPipeline.h b/Engine/Source/Runtime/VulkanRHI/Private/VulkanPipeline.h
--- a/Engine/Source/Runtime/VulkanRHI/Private/VulkanPipeline.h
+++ b/Engine/Source/Runtime/VulkanRHI/Private/VulkanPipeline.h
@@ -708,15 +708,12 @@ public:
 	}
 
 	void DeleteVkPipeline(bool bImmediate);
-	void GetOrCreateShaderModules(FVulkanShader*const* Shaders);
+	void GetOrCreateShaderModules(TRefCountPtr<FVulkanShaderModule>(&ShaderModulesOUT)[ShaderStage::NumStages], FVulkanShader* const* Shaders);
 	void PurgeShaderModules(FVulkanShader*const* Shaders);
-	void PurgeLoadedShaderModules(FVulkanDevice* InDevice);
-
 
 
 	bool								bHasInputAttachments = false;
 	bool								bIsRegistered;
-	bool								bShaderModulesLoaded = false;
 
 	uint64								ShaderKeys[ShaderStage::NumStages];
 	TEnumAsByte<EPrimitiveType>			PrimitiveType;
@@ -728,7 +725,6 @@ public:
 	FVulkanDevice* Device;
 	FGfxPipelineDesc Desc;
 	FVulkanShader* VulkanShaders[ShaderStage::NumStages];
-	VkShaderModule ShaderModules[ShaderStage::NumStages];
 	const FVulkanRenderPass* RenderPass;
 
 
diff --git a/Engine/Source/Runtime/VulkanRHI/Private/VulkanResources.h b/Engine/Source/Runtime/VulkanRHI/Private/VulkanResources.h
--- a/Engine/Source/Runtime/VulkanRHI/Private/VulkanResources.h
+++ b/Engine/Source/Runtime/VulkanRHI/Private/VulkanResources.h
@@ -79,8 +79,24 @@ public:
 
 struct FGfxPipelineDesc;
 
+class FVulkanShaderModule : public FThreadSafeRefCountedObject
+{
+	static FVulkanDevice* Device;
+	VkShaderModule ActualShaderModule;
+public:
+	FVulkanShaderModule(FVulkanDevice* DeviceIn, VkShaderModule ShaderModuleIn) : ActualShaderModule(ShaderModuleIn) 
+	{
+		check(DeviceIn && (Device == DeviceIn || !Device));
+		Device = DeviceIn;
+	}
+	virtual ~FVulkanShaderModule();
+	VkShaderModule& GetVkShaderModule() { return ActualShaderModule; }
+};
+
 class FVulkanShader : public IRefCountedObject
 {
+	static FCriticalSection VulkanShaderModulesMapCS;
+
 public:
 	FVulkanShader(FVulkanDevice* InDevice, EShaderFrequency InFrequency, VkShaderStageFlagBits InStageFlag)
 		: ShaderKey(0)
@@ -96,9 +112,10 @@ public:
 
 	void Setup(TArrayView<const uint8> InShaderHeaderAndCode, uint64 InShaderKey);
 
-	VkShaderModule GetOrCreateHandle(const FVulkanLayout* Layout, uint32 LayoutHash)
+	TRefCountPtr<FVulkanShaderModule> GetOrCreateHandle(const FVulkanLayout* Layout, uint32 LayoutHash)
 	{
-		VkShaderModule* Found = ShaderModules.Find(LayoutHash);
+		FScopeLock Lock(&VulkanShaderModulesMapCS);
+		TRefCountPtr<FVulkanShaderModule>* Found = ShaderModules.Find(LayoutHash);
 		if (Found)
 		{
 			return *Found;
@@ -107,14 +124,15 @@ public:
 		return CreateHandle(Layout, LayoutHash);
 	}
 	
-	VkShaderModule GetOrCreateHandle(const FGfxPipelineDesc& Desc, const FVulkanLayout* Layout, uint32 LayoutHash)
+	TRefCountPtr<FVulkanShaderModule> GetOrCreateHandle(const FGfxPipelineDesc& Desc, const FVulkanLayout* Layout, uint32 LayoutHash)
 	{
+		FScopeLock Lock(&VulkanShaderModulesMapCS);
 		if (NeedsSpirvInputAttachmentPatching(Desc))
 		{
 			LayoutHash = HashCombine(LayoutHash, 1);
 		}
 		
-		VkShaderModule* Found = ShaderModules.Find(LayoutHash);
+		TRefCountPtr<FVulkanShaderModule>* Found = ShaderModules.Find(LayoutHash);
 		if (Found)
 		{
 			return *Found;
@@ -168,7 +186,7 @@ protected:
 
 	/** External bindings for this shader. */
 	FVulkanShaderHeader				CodeHeader;
-	TMap<uint32, VkShaderModule>	ShaderModules;
+	TMap<uint32, TRefCountPtr<FVulkanShaderModule>>	ShaderModules;
 	const VkShaderStageFlagBits		StageFlag;
 	EShaderFrequency				Frequency;
 
@@ -193,8 +211,8 @@ protected:
 
 	FVulkanDevice*					Device;
 
-	VkShaderModule CreateHandle(const FVulkanLayout* Layout, uint32 LayoutHash);
-	VkShaderModule CreateHandle(const FGfxPipelineDesc& Desc, const FVulkanLayout* Layout, uint32 LayoutHash);
+	TRefCountPtr<FVulkanShaderModule> CreateHandle(const FVulkanLayout* Layout, uint32 LayoutHash);
+	TRefCountPtr<FVulkanShaderModule> CreateHandle(const FGfxPipelineDesc& Desc, const FVulkanLayout* Layout, uint32 LayoutHash);
 	
 	bool NeedsSpirvInputAttachmentPatching(const FGfxPipelineDesc& Desc) const;
 
diff --git a/Engine/Source/Runtime/VulkanRHI/Private/VulkanShaders.cpp b/Engine/Source/Runtime/VulkanRHI/Private/VulkanShaders.cpp
--- a/Engine/Source/Runtime/VulkanRHI/Private/VulkanShaders.cpp
+++ b/Engine/Source/Runtime/VulkanRHI/Private/VulkanShaders.cpp
@@ -38,6 +38,7 @@ static FAutoConsoleVariableRef GVulkanCompressSPIRVCVar(
 	TEXT("1 SPIRV source is compressed on load and decompressed as when needed, this saves RAM but can introduce hitching when creating shaders."),
 	ECVF_ReadOnly | ECVF_RenderThreadSafe
 );
+FCriticalSection FVulkanShader::VulkanShaderModulesMapCS;
 
 FVulkanShaderFactory::~FVulkanShaderFactory()
 {
@@ -138,6 +139,13 @@ FArchive& operator<<(FArchive& Ar, FVulkanShader::FSpirvContainer& SpirvContaine
 	return Ar;
 }
 
+FVulkanDevice* FVulkanShaderModule::Device = nullptr;
+
+FVulkanShaderModule::~FVulkanShaderModule()
+{
+	Device->GetDeferredDeletionQueue().EnqueueResource(VulkanRHI::FDeferredDeletionQueue2::EType::ShaderModule, ActualShaderModule);
+}
+
 FVulkanShader::FSpirvCode FVulkanShader::GetSpirvCode()
 {
 	if (SpirvContainer.IsCompressed())
@@ -198,7 +206,7 @@ void FVulkanShader::Setup(TArrayView<const uint8> InShaderHeaderAndCode, uint64
 #endif
 }
 
-static VkShaderModule CreateShaderModule(FVulkanDevice* Device, FVulkanShader::FSpirvCode& SpirvCode)
+static TRefCountPtr<FVulkanShaderModule> CreateShaderModule(FVulkanDevice* Device, FVulkanShader::FSpirvCode& SpirvCode)
 {
 	const TArrayView<uint32> Spirv = SpirvCode.GetCodeView();
 	VkShaderModule ShaderModule;
@@ -220,7 +228,10 @@ static VkShaderModule CreateShaderModule(FVulkanDevice* Device, FVulkanShader::F
 #endif
 
 	VERIFYVULKANRESULT(VulkanRHI::vkCreateShaderModule(Device->GetInstanceHandle(), &ModuleCreateInfo, VULKAN_CPU_ALLOCATOR, &ShaderModule));
-	return ShaderModule;
+	
+	TRefCountPtr<FVulkanShaderModule> ReturnPtr = TRefCountPtr<FVulkanShaderModule>(new FVulkanShaderModule(Device, ShaderModule));
+
+	return ReturnPtr;
 }
 
 /*
@@ -307,8 +318,9 @@ bool FVulkanShader::NeedsSpirvInputAttachmentPatching(const FGfxPipelineDesc& De
 	return (Desc.RasterizationSamples > 1 && CodeHeader.InputAttachments.Num() > 0);
 }
 
-VkShaderModule FVulkanShader::CreateHandle(const FGfxPipelineDesc& Desc, const FVulkanLayout* Layout, uint32 LayoutHash)
+TRefCountPtr<FVulkanShaderModule> FVulkanShader::CreateHandle(const FGfxPipelineDesc& Desc, const FVulkanLayout* Layout, uint32 LayoutHash)
 {
+	FScopeLock Lock(&VulkanShaderModulesMapCS);
 	FSpirvCode Spirv = GetSpirvCode();
 
 	Layout->PatchSpirvBindings(Spirv, Frequency, CodeHeader, StageFlag);
@@ -317,17 +329,18 @@ VkShaderModule FVulkanShader::CreateHandle(const FGfxPipelineDesc& Desc, const F
 		Spirv = PatchSpirvInputAttachments(Spirv);
 	}
 
-	VkShaderModule Module = CreateShaderModule(Device, Spirv);
+	TRefCountPtr<FVulkanShaderModule> Module = CreateShaderModule(Device, Spirv);
 	ShaderModules.Add(LayoutHash, Module);
 	return Module;
 }
 
-VkShaderModule FVulkanShader::CreateHandle(const FVulkanLayout* Layout, uint32 LayoutHash)
+TRefCountPtr<FVulkanShaderModule> FVulkanShader::CreateHandle(const FVulkanLayout* Layout, uint32 LayoutHash)
 {
+	FScopeLock Lock(&VulkanShaderModulesMapCS);
 	FSpirvCode Spirv = GetSpirvCode();
 
 	Layout->PatchSpirvBindings(Spirv, Frequency, CodeHeader, StageFlag);
-	VkShaderModule Module = CreateShaderModule(Device, Spirv);
+	TRefCountPtr<FVulkanShaderModule> Module = CreateShaderModule(Device, Spirv);
 	ShaderModules.Add(LayoutHash, Module);
 	return Module;
 }
@@ -340,11 +353,7 @@ FVulkanShader::~FVulkanShader()
 
 void FVulkanShader::PurgeShaderModules()
 {
-	for (const auto& Pair : ShaderModules)
-	{
-		VkShaderModule ShaderModule = Pair.Value;
-		Device->GetDeferredDeletionQueue().EnqueueResource(VulkanRHI::FDeferredDeletionQueue2::EType::ShaderModule, ShaderModule);
-	}
+	FScopeLock Lock(&VulkanShaderModulesMapCS);
 	ShaderModules.Empty(0);
 }
 
