Back out changelist 20838318 to address CIS content error.

#rnx

diff --git a/Engine/Source/Runtime/VulkanRHI/Private/VulkanPipeline.cpp b/Engine/Source/Runtime/VulkanRHI/Private/VulkanPipeline.cpp
--- a/Engine/Source/Runtime/VulkanRHI/Private/VulkanPipeline.cpp
+++ b/Engine/Source/Runtime/VulkanRHI/Private/VulkanPipeline.cpp
@@ -52,6 +52,9 @@ static TAtomic<uint64> SPipelineGfxCount;
 
 static const double HitchTime = 1.0 / 1000.0;
 
+static FCriticalSection FVulkanShaderHandleCS;
+
+
 TAutoConsoleVariable<int32> CVarPipelineDebugForceEvictImmediately(
 	TEXT("r.Vulkan.PipelineDebugForceEvictImmediately"),
 	0,
@@ -251,33 +254,57 @@ FVulkanRHIGraphicsPipelineState::~FVulkanRHIGraphicsPipelineState()
 	}
 
 	Device->PipelineStateCache->NotifyDeletedGraphicsPSO(this);
+	if (bShaderModulesLoaded)
+	{
+		PurgeLoadedShaderModules(Device);
+	}
 }
 
-void FVulkanRHIGraphicsPipelineState::GetOrCreateShaderModules(TRefCountPtr<FVulkanShaderModule> (&ShaderModulesOUT)[ShaderStage::NumStages], FVulkanShader*const* Shaders)
+void FVulkanRHIGraphicsPipelineState::GetOrCreateShaderModules(FVulkanShader*const* Shaders)
 {
+	FScopeLock Lock(&FVulkanShaderHandleCS); // this lock is to stop concurrent access to VulkanShader's ShaderModules container. 
 	for (int32 Index = 0; Index < ShaderStage::NumStages; ++Index)
 	{
-		check(!ShaderModulesOUT[Index].IsValid());
 		FVulkanShader* Shader = Shaders[Index];
 		if (Shader)
 		{
-			ShaderModulesOUT[Index] = Shader->GetOrCreateHandle(Desc, Layout, Layout->GetDescriptorSetLayoutHash());
+			ShaderModules[Index] = Shader->GetOrCreateHandle(Desc, Layout, Layout->GetDescriptorSetLayoutHash());
 		}
 	}
 }
 
 void FVulkanRHIGraphicsPipelineState::PurgeShaderModules(FVulkanShader*const* Shaders)
 {
+	check(!bShaderModulesLoaded);
+
 	for (int32 Index = 0; Index < ShaderStage::NumStages; ++Index)
 	{
 		FVulkanShader* Shader = Shaders[Index];
 		if (Shader)
 		{
 			Shader->PurgeShaderModules();
+			ShaderModules[Index] = VK_NULL_HANDLE;
 		}
 	}
 }
 
+void FVulkanRHIGraphicsPipelineState::PurgeLoadedShaderModules(FVulkanDevice* InDevice)
+{
+	check(bShaderModulesLoaded);
+
+	for (int32 Index = 0; Index < ShaderStage::NumStages; ++Index)
+	{
+		if (ShaderModules[Index] != VK_NULL_HANDLE)
+		{
+			VulkanRHI::vkDestroyShaderModule(InDevice->GetInstanceHandle(), ShaderModules[Index], VULKAN_CPU_ALLOCATOR);
+			ShaderModules[Index] = VK_NULL_HANDLE;
+		}
+	}
+
+	bShaderModulesLoaded = false;
+}
+
+
 FVulkanPipelineStateCacheManager::FVulkanPipelineStateCacheManager(FVulkanDevice* InDevice)
 	: Device(InDevice)
 	, bEvictImmediately(false)
@@ -1107,9 +1134,10 @@ bool FVulkanPipelineStateCacheManager::CreateGfxPipelineFromEntry(FVulkanRHIGrap
 		Shaders[ShaderStage::Pixel] = ResourceCast(TShaderMapRef<FNULLPS>(GetGlobalShaderMap(GMaxRHIFeatureLevel)).GetPixelShader());
 	}
 
-	TRefCountPtr<FVulkanShaderModule> ShaderModules[ShaderStage::NumStages];
-
-	PSO->GetOrCreateShaderModules(ShaderModules, Shaders);
+	if (!PSO->bShaderModulesLoaded)
+	{
+		PSO->GetOrCreateShaderModules(Shaders);
+	}
 
 	// Pipeline
 	VkGraphicsPipelineCreateInfo PipelineInfo;
@@ -1163,7 +1191,7 @@ bool FVulkanPipelineStateCacheManager::CreateGfxPipelineFromEntry(FVulkanRHIGrap
 	ANSICHAR EntryPoints[ShaderStage::NumStages][24];
 	for (int32 ShaderStage = 0; ShaderStage < ShaderStage::NumStages; ++ShaderStage)
 	{
-		if (!ShaderModules[ShaderStage].IsValid())
+		if (!PSO->ShaderModules[ShaderStage])
 		{
 			continue;
 		}
@@ -1172,7 +1200,7 @@ bool FVulkanPipelineStateCacheManager::CreateGfxPipelineFromEntry(FVulkanRHIGrap
 		ShaderStages[PipelineInfo.stageCount].sType = VK_STRUCTURE_TYPE_PIPELINE_SHADER_STAGE_CREATE_INFO;
 		VkShaderStageFlagBits Stage = UEFrequencyToVKStageBit(ShaderStage::GetFrequencyForGfxStage(CurrStage));
 		ShaderStages[PipelineInfo.stageCount].stage = Stage;
-		ShaderStages[PipelineInfo.stageCount].module = ShaderModules[CurrStage]->GetVkShaderModule();
+		ShaderStages[PipelineInfo.stageCount].module = PSO->ShaderModules[CurrStage];
 		Shaders[ShaderStage]->GetEntryPoint(EntryPoints[PipelineInfo.stageCount], 24);
 		ShaderStages[PipelineInfo.stageCount].pName = EntryPoints[PipelineInfo.stageCount];
 		PipelineInfo.stageCount++;
@@ -1711,6 +1739,7 @@ FVulkanRHIGraphicsPipelineState::FVulkanRHIGraphicsPipelineState(FVulkanDevice*
 
 	PSOInitializer = PSOInitializer_;
 #endif
+	FMemory::Memset(ShaderModules, 0, sizeof(ShaderModules));
 	INC_DWORD_STAT(STAT_VulkanNumGraphicsPSOs);
 	INC_DWORD_STAT_BY(STAT_VulkanPSOKeyMemory, this->VulkanKey.GetDataRef().Num());
 }
@@ -1995,13 +2024,13 @@ FVulkanComputePipeline* FVulkanPipelineStateCacheManager::CreateComputePipelineF
 		ComputeLayout->ComputePipelineDescriptorInfo.Initialize(Layout->GetDescriptorSetsLayout().RemappingInfo);
 	}
 
-	TRefCountPtr<FVulkanShaderModule> ShaderModule = Shader->GetOrCreateHandle(Layout, Layout->GetDescriptorSetLayoutHash());
+	VkShaderModule ShaderModule = Shader->GetOrCreateHandle(Layout, Layout->GetDescriptorSetLayoutHash());
 
 	VkComputePipelineCreateInfo PipelineInfo;
 	ZeroVulkanStruct(PipelineInfo, VK_STRUCTURE_TYPE_COMPUTE_PIPELINE_CREATE_INFO);
 	PipelineInfo.stage.sType = VK_STRUCTURE_TYPE_PIPELINE_SHADER_STAGE_CREATE_INFO;
 	PipelineInfo.stage.stage = VK_SHADER_STAGE_COMPUTE_BIT;
-	PipelineInfo.stage.module = ShaderModule->GetVkShaderModule();
+	PipelineInfo.stage.module = ShaderModule;
 	// main_00000000_00000000
 	ANSICHAR EntryPoint[24];
 	Shader->GetEntryPoint(EntryPoint, 24);
@@ -2238,7 +2267,7 @@ void FVulkanPipelineStateCacheManager::LRUTouch(FVulkanRHIGraphicsPipelineState*
 	}
 	FScopeLock Lock(&LRUCS);
 	check((PSO->GetVulkanPipeline() == 0) == (PSO->LRUNode == 0));
-
+	
 	if(PSO->LRUNode)
 	{
 		check(PSO->GetVulkanPipeline());
@@ -2341,6 +2370,7 @@ void FVulkanPipelineStateCacheManager::LRURemove(FVulkanRHIGraphicsPipelineState
 		PSO->DeleteVkPipeline(bImmediate);
 		if (GVulkanReleaseShaderModuleWhenEvictingPSO)
 		{
+			FScopeLock Lock(&FVulkanShaderHandleCS); // this lock is here to stop concurrent access to VulkanShaders ShaderModules container. 
 	        for (int ShaderStageIndex = 0; ShaderStageIndex < ShaderStage::NumStages; ShaderStageIndex++)
 	        {
 				if (PSO->VulkanShaders[ShaderStageIndex] != nullptr)
diff --git a/Engine/Source/Runtime/VulkanRHI/Private/VulkanPipeline.h b/Engine/Source/Runtime/VulkanRHI/Private/VulkanPipeline.h
--- a/Engine/Source/Runtime/VulkanRHI/Private/VulkanPipeline.h
+++ b/Engine/Source/Runtime/VulkanRHI/Private/VulkanPipeline.h
@@ -708,12 +708,15 @@ public:
 	}
 
 	void DeleteVkPipeline(bool bImmediate);
-	void GetOrCreateShaderModules(TRefCountPtr<FVulkanShaderModule>(&ShaderModulesOUT)[ShaderStage::NumStages], FVulkanShader* const* Shaders);
+	void GetOrCreateShaderModules(FVulkanShader*const* Shaders);
 	void PurgeShaderModules(FVulkanShader*const* Shaders);
+	void PurgeLoadedShaderModules(FVulkanDevice* InDevice);
+
 
 
 	bool								bHasInputAttachments = false;
 	bool								bIsRegistered;
+	bool								bShaderModulesLoaded = false;
 
 	uint64								ShaderKeys[ShaderStage::NumStages];
 	TEnumAsByte<EPrimitiveType>			PrimitiveType;
@@ -725,6 +728,7 @@ public:
 	FVulkanDevice* Device;
 	FGfxPipelineDesc Desc;
 	FVulkanShader* VulkanShaders[ShaderStage::NumStages];
+	VkShaderModule ShaderModules[ShaderStage::NumStages];
 	const FVulkanRenderPass* RenderPass;
 
 
diff --git a/Engine/Source/Runtime/VulkanRHI/Private/VulkanResources.h b/Engine/Source/Runtime/VulkanRHI/Private/VulkanResources.h
--- a/Engine/Source/Runtime/VulkanRHI/Private/VulkanResources.h
+++ b/Engine/Source/Runtime/VulkanRHI/Private/VulkanResources.h
@@ -79,24 +79,8 @@ public:
 
 struct FGfxPipelineDesc;
 
-class FVulkanShaderModule : public FThreadSafeRefCountedObject
-{
-	static FVulkanDevice* Device;
-	VkShaderModule ActualShaderModule;
-public:
-	FVulkanShaderModule(FVulkanDevice* DeviceIn, VkShaderModule ShaderModuleIn) : ActualShaderModule(ShaderModuleIn) 
-	{
-		check(DeviceIn && (Device == DeviceIn || !Device));
-		Device = DeviceIn;
-	}
-	virtual ~FVulkanShaderModule();
-	VkShaderModule& GetVkShaderModule() { return ActualShaderModule; }
-};
-
 class FVulkanShader : public IRefCountedObject
 {
-	static FCriticalSection VulkanShaderModulesMapCS;
-
 public:
 	FVulkanShader(FVulkanDevice* InDevice, EShaderFrequency InFrequency, VkShaderStageFlagBits InStageFlag)
 		: ShaderKey(0)
@@ -112,10 +96,9 @@ public:
 
 	void Setup(TArrayView<const uint8> InShaderHeaderAndCode, uint64 InShaderKey);
 
-	TRefCountPtr<FVulkanShaderModule> GetOrCreateHandle(const FVulkanLayout* Layout, uint32 LayoutHash)
+	VkShaderModule GetOrCreateHandle(const FVulkanLayout* Layout, uint32 LayoutHash)
 	{
-		FScopeLock Lock(&VulkanShaderModulesMapCS);
-		TRefCountPtr<FVulkanShaderModule>* Found = ShaderModules.Find(LayoutHash);
+		VkShaderModule* Found = ShaderModules.Find(LayoutHash);
 		if (Found)
 		{
 			return *Found;
@@ -124,15 +107,14 @@ public:
 		return CreateHandle(Layout, LayoutHash);
 	}
 	
-	TRefCountPtr<FVulkanShaderModule> GetOrCreateHandle(const FGfxPipelineDesc& Desc, const FVulkanLayout* Layout, uint32 LayoutHash)
+	VkShaderModule GetOrCreateHandle(const FGfxPipelineDesc& Desc, const FVulkanLayout* Layout, uint32 LayoutHash)
 	{
-		FScopeLock Lock(&VulkanShaderModulesMapCS);
 		if (NeedsSpirvInputAttachmentPatching(Desc))
 		{
 			LayoutHash = HashCombine(LayoutHash, 1);
 		}
 		
-		TRefCountPtr<FVulkanShaderModule>* Found = ShaderModules.Find(LayoutHash);
+		VkShaderModule* Found = ShaderModules.Find(LayoutHash);
 		if (Found)
 		{
 			return *Found;
@@ -186,7 +168,7 @@ protected:
 
 	/** External bindings for this shader. */
 	FVulkanShaderHeader				CodeHeader;
-	TMap<uint32, TRefCountPtr<FVulkanShaderModule>>	ShaderModules;
+	TMap<uint32, VkShaderModule>	ShaderModules;
 	const VkShaderStageFlagBits		StageFlag;
 	EShaderFrequency				Frequency;
 
@@ -211,8 +193,8 @@ protected:
 
 	FVulkanDevice*					Device;
 
-	TRefCountPtr<FVulkanShaderModule> CreateHandle(const FVulkanLayout* Layout, uint32 LayoutHash);
-	TRefCountPtr<FVulkanShaderModule> CreateHandle(const FGfxPipelineDesc& Desc, const FVulkanLayout* Layout, uint32 LayoutHash);
+	VkShaderModule CreateHandle(const FVulkanLayout* Layout, uint32 LayoutHash);
+	VkShaderModule CreateHandle(const FGfxPipelineDesc& Desc, const FVulkanLayout* Layout, uint32 LayoutHash);
 	
 	bool NeedsSpirvInputAttachmentPatching(const FGfxPipelineDesc& Desc) const;
 
diff --git a/Engine/Source/Runtime/VulkanRHI/Private/VulkanShaders.cpp b/Engine/Source/Runtime/VulkanRHI/Private/VulkanShaders.cpp
--- a/Engine/Source/Runtime/VulkanRHI/Private/VulkanShaders.cpp
+++ b/Engine/Source/Runtime/VulkanRHI/Private/VulkanShaders.cpp
@@ -38,7 +38,6 @@ static FAutoConsoleVariableRef GVulkanCompressSPIRVCVar(
 	TEXT("1 SPIRV source is compressed on load and decompressed as when needed, this saves RAM but can introduce hitching when creating shaders."),
 	ECVF_ReadOnly | ECVF_RenderThreadSafe
 );
-FCriticalSection FVulkanShader::VulkanShaderModulesMapCS;
 
 FVulkanShaderFactory::~FVulkanShaderFactory()
 {
@@ -139,13 +138,6 @@ FArchive& operator<<(FArchive& Ar, FVulkanShader::FSpirvContainer& SpirvContaine
 	return Ar;
 }
 
-FVulkanDevice* FVulkanShaderModule::Device = nullptr;
-
-FVulkanShaderModule::~FVulkanShaderModule()
-{
-	Device->GetDeferredDeletionQueue().EnqueueResource(VulkanRHI::FDeferredDeletionQueue2::EType::ShaderModule, ActualShaderModule);
-}
-
 FVulkanShader::FSpirvCode FVulkanShader::GetSpirvCode()
 {
 	if (SpirvContainer.IsCompressed())
@@ -206,7 +198,7 @@ void FVulkanShader::Setup(TArrayView<const uint8> InShaderHeaderAndCode, uint64
 #endif
 }
 
-static TRefCountPtr<FVulkanShaderModule> CreateShaderModule(FVulkanDevice* Device, FVulkanShader::FSpirvCode& SpirvCode)
+static VkShaderModule CreateShaderModule(FVulkanDevice* Device, FVulkanShader::FSpirvCode& SpirvCode)
 {
 	const TArrayView<uint32> Spirv = SpirvCode.GetCodeView();
 	VkShaderModule ShaderModule;
@@ -228,10 +220,7 @@ static TRefCountPtr<FVulkanShaderModule> CreateShaderModule(FVulkanDevice* Devic
 #endif
 
 	VERIFYVULKANRESULT(VulkanRHI::vkCreateShaderModule(Device->GetInstanceHandle(), &ModuleCreateInfo, VULKAN_CPU_ALLOCATOR, &ShaderModule));
-	
-	TRefCountPtr<FVulkanShaderModule> ReturnPtr = TRefCountPtr<FVulkanShaderModule>(new FVulkanShaderModule(Device, ShaderModule));
-
-	return ReturnPtr;
+	return ShaderModule;
 }
 
 /*
@@ -318,9 +307,8 @@ bool FVulkanShader::NeedsSpirvInputAttachmentPatching(const FGfxPipelineDesc& De
 	return (Desc.RasterizationSamples > 1 && CodeHeader.InputAttachments.Num() > 0);
 }
 
-TRefCountPtr<FVulkanShaderModule> FVulkanShader::CreateHandle(const FGfxPipelineDesc& Desc, const FVulkanLayout* Layout, uint32 LayoutHash)
+VkShaderModule FVulkanShader::CreateHandle(const FGfxPipelineDesc& Desc, const FVulkanLayout* Layout, uint32 LayoutHash)
 {
-	FScopeLock Lock(&VulkanShaderModulesMapCS);
 	FSpirvCode Spirv = GetSpirvCode();
 
 	Layout->PatchSpirvBindings(Spirv, Frequency, CodeHeader, StageFlag);
@@ -329,18 +317,17 @@ TRefCountPtr<FVulkanShaderModule> FVulkanShader::CreateHandle(const FGfxPipeline
 		Spirv = PatchSpirvInputAttachments(Spirv);
 	}
 
-	TRefCountPtr<FVulkanShaderModule> Module = CreateShaderModule(Device, Spirv);
+	VkShaderModule Module = CreateShaderModule(Device, Spirv);
 	ShaderModules.Add(LayoutHash, Module);
 	return Module;
 }
 
-TRefCountPtr<FVulkanShaderModule> FVulkanShader::CreateHandle(const FVulkanLayout* Layout, uint32 LayoutHash)
+VkShaderModule FVulkanShader::CreateHandle(const FVulkanLayout* Layout, uint32 LayoutHash)
 {
-	FScopeLock Lock(&VulkanShaderModulesMapCS);
 	FSpirvCode Spirv = GetSpirvCode();
 
 	Layout->PatchSpirvBindings(Spirv, Frequency, CodeHeader, StageFlag);
-	TRefCountPtr<FVulkanShaderModule> Module = CreateShaderModule(Device, Spirv);
+	VkShaderModule Module = CreateShaderModule(Device, Spirv);
 	ShaderModules.Add(LayoutHash, Module);
 	return Module;
 }
@@ -353,7 +340,11 @@ FVulkanShader::~FVulkanShader()
 
 void FVulkanShader::PurgeShaderModules()
 {
-	FScopeLock Lock(&VulkanShaderModulesMapCS);
+	for (const auto& Pair : ShaderModules)
+	{
+		VkShaderModule ShaderModule = Pair.Value;
+		Device->GetDeferredDeletionQueue().EnqueueResource(VulkanRHI::FDeferredDeletionQueue2::EType::ShaderModule, ShaderModule);
+	}
 	ShaderModules.Empty(0);
 }
 
