Expand 'r.Vulkan.ForcePSOSingleThreaded' option to allow or disallow multithreaded PSO creates for precompile and/or rendering PSOs.

diff --git a/Engine/Source/Runtime/VulkanRHI/Private/VulkanPipeline.cpp b/Engine/Source/Runtime/VulkanRHI/Private/VulkanPipeline.cpp
--- a/Engine/Source/Runtime/VulkanRHI/Private/VulkanPipeline.cpp
+++ b/Engine/Source/Runtime/VulkanRHI/Private/VulkanPipeline.cpp
@@ -123,22 +123,23 @@ static FAutoConsoleVariableRef GEnablePipelineCacheCompressionCvar(
 );
 
 
-static int32 GVulkanPSOForceSingleThreaded = 0;
+enum class ESingleThreadedPSOCreateMode
+{
+	None = 0,
+	All = 1,
+	Precompile = 2,
+	NonPrecompiled = 3,
+};
+
+static int32 GVulkanPSOForceSingleThreaded = (int32)ESingleThreadedPSOCreateMode::None;
 static FAutoConsoleVariableRef GVulkanPSOForceSingleThreadedCVar(
 	TEXT("r.Vulkan.ForcePSOSingleThreaded"),
 	GVulkanPSOForceSingleThreaded,
-	TEXT("Enable to force singlethreaded creation of PSOs. Only intended as a workaround for buggy drivers\n"),
-	ECVF_ReadOnly | ECVF_RenderThreadSafe
-);
-
-static int32 GVulkanPSOPrecompileForceSingleThreaded = 0;
-static FAutoConsoleVariableRef GVulkanPSOPrecompileForceSingleThreadedCVar(
-	TEXT("r.Vulkan.PrecompilePSOForceSingleThreaded"),
-	GVulkanPSOPrecompileForceSingleThreaded,
-	TEXT("Some drivers serialize PSO creation, this setting can help reduce the driver's PSO queue and can decrease the wait time for more important (non-precompile) PSOs.\n")
+	TEXT("Enable to force singlethreaded creation of PSOs. Only intended as a workaround for buggy drivers\n")
 	TEXT("0: (default) Allow Async precompile PSO creation.\n")
-	TEXT("1: force singlethreaded creation of precompile PSOs only.")
-	,
+	TEXT("1: force singlethreaded creation of all PSOs.\n")
+	TEXT("2: force singlethreaded creation of precompile PSOs only.\n")
+	TEXT("3: force singlethreaded creation of non-precompile PSOs only."),
 	ECVF_ReadOnly | ECVF_RenderThreadSafe
 );
 
@@ -1780,10 +1781,14 @@ FVulkanRHIGraphicsPipelineState* FVulkanPipelineStateCacheManager::RHICreateGrap
 	QUICK_SCOPE_CYCLE_COUNTER(STAT_Vulkan_RHICreateGraphicsPipelineState_NEW);
 
 	// Optional lock for PSO creation, GVulkanPSOForceSingleThreaded is used to work around driver bugs.
-	// GVulkanPSOPrecompileForceSingleThreaded is used when the driver internally serializes PSO creation, this option reduces the driver queue size.
+	// GVulkanPSOForceSingleThreaded == Precompile can be used when the driver internally serializes PSO creation, this option reduces the driver queue size.
 	// We stall precompile PSOs which increases the likelihood for non-precompile PSO to jump the queue.
 	// Not using GraphicsPSOLockedCS as the create could take a long time on some platforms, holding GraphicsPSOLockedCS the whole time could cause hitching.
-	const bool bShouldLock = GVulkanPSOForceSingleThreaded || (GVulkanPSOPrecompileForceSingleThreaded && Initializer.bFromPSOFileCache);
+	const ESingleThreadedPSOCreateMode ThreadingMode = (ESingleThreadedPSOCreateMode)GVulkanPSOForceSingleThreaded;
+	bool bShouldLock = ThreadingMode == ESingleThreadedPSOCreateMode::All
+		|| (ThreadingMode == ESingleThreadedPSOCreateMode::Precompile && Initializer.bFromPSOFileCache)
+		|| (ThreadingMode == ESingleThreadedPSOCreateMode::NonPrecompiled && !Initializer.bFromPSOFileCache);
+
 	FPSOOptionalLock PSOSingleThreadedLock(bShouldLock ? &CreateGraphicsPSOMutex : nullptr);
 
 	FVulkanPSOKey Key;
