VulkanRHI:  avoid the automatic clearing of new targets when they are allocated in transient memory since it gets put in the upload command buffer that can execute earlier in the frame (when the same memory is potentially used by a different transient resource).  Possibly a nemesis bug, seems to resolve flickering that occurs in editor when using Vulkan.

diff --git a/Engine/Source/Runtime/VulkanRHI/Private/VulkanTexture.cpp b/Engine/Source/Runtime/VulkanRHI/Private/VulkanTexture.cpp
--- a/Engine/Source/Runtime/VulkanRHI/Private/VulkanTexture.cpp
+++ b/Engine/Source/Runtime/VulkanRHI/Private/VulkanTexture.cpp
@@ -1644,7 +1644,8 @@ FVulkanTexture::FVulkanTexture(FVulkanDevice& InDevice, const FRHITextureCreateD
 			}
 		}
 
-		if (InTransientHeapAllocation != nullptr)
+		const bool bIsTransientResource = (InTransientHeapAllocation != nullptr);
+		if (bIsTransientResource)
 		{
 			check(!bMemoryless);
 			check(InTransientHeapAllocation->Offset % MemoryRequirements.alignment == 0);
@@ -1683,7 +1684,7 @@ FVulkanTexture::FVulkanTexture(FVulkanDevice& InDevice, const FRHITextureCreateD
 		check(Tiling == VK_IMAGE_TILING_LINEAR || Tiling == VK_IMAGE_TILING_OPTIMAL);
 
 		const VkImageLayout InitialLayout = GetInitialLayoutFromRHIAccess(InCreateDesc.InitialState, InCreateDesc.Flags, SupportsSampling());
-		const bool bDoInitialClear = VKHasAnyFlags(ImageCreateInfo.ImageCreateInfo.usage, VK_IMAGE_USAGE_SAMPLED_BIT) && EnumHasAnyFlags(InCreateDesc.Flags, TexCreate_RenderTargetable | TexCreate_DepthStencilTargetable);
+		const bool bDoInitialClear = !bIsTransientResource && VKHasAnyFlags(ImageCreateInfo.ImageCreateInfo.usage, VK_IMAGE_USAGE_SAMPLED_BIT) && EnumHasAnyFlags(InCreateDesc.Flags, TexCreate_RenderTargetable | TexCreate_DepthStencilTargetable);
 
 		if (InitialLayout != VK_IMAGE_LAYOUT_UNDEFINED || bDoInitialClear)
 		{
