Fixed: Android Adreno driver 0549.0 causes low-res textures in game
Added a workaround for a certain ASTC textures to use a Linear tiling mode if it significantly reduces texture memory requirements

diff --git a/Engine/Source/Runtime/VulkanRHI/Private/Android/VulkanAndroidPlatform.cpp b/Engine/Source/Runtime/VulkanRHI/Private/Android/VulkanAndroidPlatform.cpp
--- a/Engine/Source/Runtime/VulkanRHI/Private/Android/VulkanAndroidPlatform.cpp
+++ b/Engine/Source/Runtime/VulkanRHI/Private/Android/VulkanAndroidPlatform.cpp
@@ -54,6 +54,7 @@ int32 FVulkanAndroidPlatform::UnsuccessfulRefreshRateFrames = 0;
 TArray<TArray<ANSICHAR>> FVulkanAndroidPlatform::DebugVulkanDeviceLayers;
 TArray<TArray<ANSICHAR>> FVulkanAndroidPlatform::DebugVulkanInstanceLayers;
 int32 FVulkanAndroidPlatform::AFBCWorkaroundOption = 0;
+int32 FVulkanAndroidPlatform::ASTCWorkaroundOption = 0;
 
 #define CHECK_VK_ENTRYPOINTS(Type,Func) if (VulkanDynamicAPI::Func == NULL) { bFoundAllEntryPoints = false; UE_LOG(LogRHI, Warning, TEXT("Failed to find entry point for %s"), TEXT(#Func)); }
 
@@ -616,80 +617,120 @@ bool FAndroidVulkanFramePacer::SupportsFramePace(int32 QueryFramePace)
 }
 
 //
-// Test whether we should enable a workaround for uncompressed textures
+// Test whether we should enable workarounds for textures
 // Arm GPUs use an optimization "Arm FrameBuffer Compression - AFBC" that can significanly inflate (~5x) uncompressed texture memory requirements
 // For now AFBC and similar optimizations can be disabled by using VK_IMAGE_CREATE_MUTABLE_FORMAT_BIT or VK_IMAGE_USAGE_STORAGE_BIT flags on a texture
+// On Adreno GPUs ASTC textures with optimial tiling may require 8x more memory
 //
-void FVulkanAndroidPlatform::SetupAFBCWorkaround(const FVulkanDevice& InDevice)
+void FVulkanAndroidPlatform::SetupImageMemoryRequirementWorkaround(const FVulkanDevice& InDevice)
 {
 	AFBCWorkaroundOption = 0;
-
-	const VkFormat FormatAFBC = VK_FORMAT_B8G8R8A8_UNORM;
-	const VkFormatFeatureFlags FormatFlags = InDevice.GetFormatProperties()[FormatAFBC].optimalTilingFeatures;
+	ASTCWorkaroundOption = 0;
 
 	VkImageCreateInfo ImageCreateInfo;
 	ZeroVulkanStruct(ImageCreateInfo, VK_STRUCTURE_TYPE_IMAGE_CREATE_INFO);
 	ImageCreateInfo.imageType = VK_IMAGE_TYPE_2D;
 	ImageCreateInfo.samples = VK_SAMPLE_COUNT_1_BIT;
-	ImageCreateInfo.format = FormatAFBC;
+	ImageCreateInfo.format = VK_FORMAT_B8G8R8A8_UNORM;
 	ImageCreateInfo.tiling = VK_IMAGE_TILING_OPTIMAL;
 	ImageCreateInfo.arrayLayers = 1;
 	ImageCreateInfo.extent = {128, 128, 1};
 	ImageCreateInfo.mipLevels = 8;
-	ImageCreateInfo.flags = 0;// VK_IMAGE_CREATE_MUTABLE_FORMAT_BIT;
+	ImageCreateInfo.flags = 0;
 	ImageCreateInfo.usage = VK_IMAGE_USAGE_TRANSFER_DST_BIT | VK_IMAGE_USAGE_SAMPLED_BIT;
 	ImageCreateInfo.sharingMode = VK_SHARING_MODE_EXCLUSIVE;
 	ImageCreateInfo.queueFamilyIndexCount = 0;
 	ImageCreateInfo.pQueueFamilyIndices = nullptr;
 	ImageCreateInfo.initialLayout = VK_IMAGE_LAYOUT_UNDEFINED;
 
-	VkImage Image0;
-	VkMemoryRequirements Image0Mem;
-	VERIFYVULKANRESULT(VulkanRHI::vkCreateImage(InDevice.GetInstanceHandle(), &ImageCreateInfo, VULKAN_CPU_ALLOCATOR, &Image0));
-	VulkanRHI::vkGetImageMemoryRequirements(InDevice.GetInstanceHandle(), Image0, &Image0Mem);
-	VulkanRHI::vkDestroyImage(InDevice.GetInstanceHandle(), Image0, VULKAN_CPU_ALLOCATOR);
+	// AFBC workarounds
+	{
+		const VkFormatFeatureFlags FormatFlags = InDevice.GetFormatProperties()[VK_FORMAT_B8G8R8A8_UNORM].optimalTilingFeatures;
 
-	VkImage ImageMutable;
-	VkMemoryRequirements ImageMutableMem;
-	ImageCreateInfo.flags = VK_IMAGE_CREATE_MUTABLE_FORMAT_BIT;
-	VERIFYVULKANRESULT(VulkanRHI::vkCreateImage(InDevice.GetInstanceHandle(), &ImageCreateInfo, VULKAN_CPU_ALLOCATOR, &ImageMutable));
-	VulkanRHI::vkGetImageMemoryRequirements(InDevice.GetInstanceHandle(), ImageMutable, &ImageMutableMem);
-	VulkanRHI::vkDestroyImage(InDevice.GetInstanceHandle(), ImageMutable, VULKAN_CPU_ALLOCATOR);
+		VkImage Image0;
+		VkMemoryRequirements Image0Mem;
+		VERIFYVULKANRESULT(VulkanRHI::vkCreateImage(InDevice.GetInstanceHandle(), &ImageCreateInfo, VULKAN_CPU_ALLOCATOR, &Image0));
+		VulkanRHI::vkGetImageMemoryRequirements(InDevice.GetInstanceHandle(), Image0, &Image0Mem);
+		VulkanRHI::vkDestroyImage(InDevice.GetInstanceHandle(), Image0, VULKAN_CPU_ALLOCATOR);
 
-	VkImage ImageStorage;
-	VkMemoryRequirements ImageStorageMem;
-	if ((FormatFlags & VK_FORMAT_FEATURE_STORAGE_IMAGE_BIT) != 0)
-	{
-		ImageCreateInfo.usage |= VK_IMAGE_USAGE_STORAGE_BIT;
-		ImageCreateInfo.flags = 0;
-	}
-	VERIFYVULKANRESULT(VulkanRHI::vkCreateImage(InDevice.GetInstanceHandle(), &ImageCreateInfo, VULKAN_CPU_ALLOCATOR, &ImageStorage));
-	VulkanRHI::vkGetImageMemoryRequirements(InDevice.GetInstanceHandle(), ImageStorage, &ImageStorageMem);
-	VulkanRHI::vkDestroyImage(InDevice.GetInstanceHandle(), ImageStorage, VULKAN_CPU_ALLOCATOR);
+		VkImage ImageMutable;
+		VkMemoryRequirements ImageMutableMem;
+		ImageCreateInfo.flags = VK_IMAGE_CREATE_MUTABLE_FORMAT_BIT;
+		VERIFYVULKANRESULT(VulkanRHI::vkCreateImage(InDevice.GetInstanceHandle(), &ImageCreateInfo, VULKAN_CPU_ALLOCATOR, &ImageMutable));
+		VulkanRHI::vkGetImageMemoryRequirements(InDevice.GetInstanceHandle(), ImageMutable, &ImageMutableMem);
+		VulkanRHI::vkDestroyImage(InDevice.GetInstanceHandle(), ImageMutable, VULKAN_CPU_ALLOCATOR);
+
+		VkImage ImageStorage;
+		VkMemoryRequirements ImageStorageMem;
+		if ((FormatFlags & VK_FORMAT_FEATURE_STORAGE_IMAGE_BIT) != 0)
+		{
+			ImageCreateInfo.usage |= VK_IMAGE_USAGE_STORAGE_BIT;
+			ImageCreateInfo.flags = 0;
+		}
+		VERIFYVULKANRESULT(VulkanRHI::vkCreateImage(InDevice.GetInstanceHandle(), &ImageCreateInfo, VULKAN_CPU_ALLOCATOR, &ImageStorage));
+		VulkanRHI::vkGetImageMemoryRequirements(InDevice.GetInstanceHandle(), ImageStorage, &ImageStorageMem);
+		VulkanRHI::vkDestroyImage(InDevice.GetInstanceHandle(), ImageStorage, VULKAN_CPU_ALLOCATOR);
 
-	const float MEM_SIZE_THRESHOLD = 1.5f;
-	const float IMAGE0_SIZE = (float)Image0Mem.size;
+		const float MEM_SIZE_THRESHOLD = 1.5f;
+		const float IMAGE0_SIZE = (float)Image0Mem.size;
 
-	if (ImageMutableMem.size * MEM_SIZE_THRESHOLD < IMAGE0_SIZE)
-	{
-		AFBCWorkaroundOption = 1;
-	}
-	else if (ImageStorageMem.size * MEM_SIZE_THRESHOLD < IMAGE0_SIZE)
-	{
-		AFBCWorkaroundOption = 2;
+		if (ImageMutableMem.size * MEM_SIZE_THRESHOLD < IMAGE0_SIZE)
+		{
+			AFBCWorkaroundOption = 1;
+		}
+		else if (ImageStorageMem.size * MEM_SIZE_THRESHOLD < IMAGE0_SIZE)
+		{
+			AFBCWorkaroundOption = 2;
+		}
+
+		if (AFBCWorkaroundOption != 0)
+		{
+			UE_LOG(LogRHI, Display, TEXT("Enabling workaround to reduce memory requrement for BGRA textures (%s flag). 128x128 - 8 Mips BGRA texture: %u KiB -> %u KiB"),
+				AFBCWorkaroundOption == 1 ? TEXT("MUTABLE") : TEXT("STORAGE"),
+				Image0Mem.size / 1024,
+				AFBCWorkaroundOption == 1 ? ImageMutableMem.size / 1024 : ImageStorageMem.size / 1024
+			);
+		}
 	}
 
-	if (AFBCWorkaroundOption != 0)
+	// ASTC workarounds
 	{
-		UE_LOG(LogRHI, Display, TEXT("Enabling workaround to reduce memory requrement for BGRA textures (%s flag). 128x128 - 8 Mips BGRA texture: %u KiB -> %u KiB"), 
-			AFBCWorkaroundOption == 1 ? TEXT("MUTABLE") : TEXT("STORAGE"),
-			Image0Mem.size / 1024,
-			AFBCWorkaroundOption == 1 ? ImageMutableMem.size / 1024 : ImageStorageMem.size / 1024
-		);
+		ImageCreateInfo.flags = 0;
+		ImageCreateInfo.format = VK_FORMAT_ASTC_8x8_UNORM_BLOCK;
+		ImageCreateInfo.tiling = VK_IMAGE_TILING_OPTIMAL;
+		ImageCreateInfo.usage = VK_IMAGE_USAGE_TRANSFER_DST_BIT | VK_IMAGE_USAGE_SAMPLED_BIT;
+
+		VkImage ImageOptimal_ASTC;
+		VkMemoryRequirements ImageOptimalMem_ASTC;
+		VERIFYVULKANRESULT(VulkanRHI::vkCreateImage(InDevice.GetInstanceHandle(), &ImageCreateInfo, VULKAN_CPU_ALLOCATOR, &ImageOptimal_ASTC));
+		VulkanRHI::vkGetImageMemoryRequirements(InDevice.GetInstanceHandle(), ImageOptimal_ASTC, &ImageOptimalMem_ASTC);
+		VulkanRHI::vkDestroyImage(InDevice.GetInstanceHandle(), ImageOptimal_ASTC, VULKAN_CPU_ALLOCATOR);
+
+		ImageCreateInfo.tiling = VK_IMAGE_TILING_LINEAR;
+
+		VkImage ImageLinear_ASTC;
+		VkMemoryRequirements ImageLinearMem_ASTC;
+		VERIFYVULKANRESULT(VulkanRHI::vkCreateImage(InDevice.GetInstanceHandle(), &ImageCreateInfo, VULKAN_CPU_ALLOCATOR, &ImageLinear_ASTC));
+		VulkanRHI::vkGetImageMemoryRequirements(InDevice.GetInstanceHandle(), ImageLinear_ASTC, &ImageLinearMem_ASTC);
+		VulkanRHI::vkDestroyImage(InDevice.GetInstanceHandle(), ImageLinear_ASTC, VULKAN_CPU_ALLOCATOR);
+
+
+		const float MEM_SIZE_THRESHOLD = 2.0f;
+		const float ImageOptimal_SIZE = (float)ImageOptimalMem_ASTC.size;
+
+		if (ImageLinearMem_ASTC.size * MEM_SIZE_THRESHOLD <= ImageOptimal_SIZE)
+		{
+			ASTCWorkaroundOption = 1;
+
+			UE_LOG(LogRHI, Display, TEXT("Enabling workaround to reduce memory requrement for ASTC textures (VK_IMAGE_TILING_LINEAR). 128x128 - 8 Mips ASTC_8x8 texture: %u KiB -> %u KiB"),
+				ImageOptimalMem_ASTC.size / 1024,
+				ImageLinearMem_ASTC.size / 1024
+			);
+		}
 	}
 }
 
-void FVulkanAndroidPlatform::SetImageAFBCWorkaround(VkImageCreateInfo& ImageCreateInfo)
+void FVulkanAndroidPlatform::SetImageMemoryRequirementWorkaround(VkImageCreateInfo& ImageCreateInfo)
 {
 	if (AFBCWorkaroundOption != 0 &&
 		ImageCreateInfo.imageType == VK_IMAGE_TYPE_2D && 
@@ -705,4 +746,13 @@ void FVulkanAndroidPlatform::SetImageAFBCWorkaround(VkImageCreateInfo& ImageCrea
 			ImageCreateInfo.usage |= VK_IMAGE_USAGE_STORAGE_BIT;
 		}
 	}
+
+	// Use ASTC workaround for textures ASTC_6x6 and ASTC_8x8 with mips and size up to 128x128
+	if (ASTCWorkaroundOption != 0 &&
+		ImageCreateInfo.imageType == VK_IMAGE_TYPE_2D &&
+		(ImageCreateInfo.format >= VK_FORMAT_ASTC_6x6_UNORM_BLOCK && ImageCreateInfo.format <= VK_FORMAT_ASTC_8x8_SRGB_BLOCK) &&
+		(ImageCreateInfo.mipLevels > 1 && ImageCreateInfo.extent.width <= 128 && ImageCreateInfo.extent.height <= 128))
+	{
+		ImageCreateInfo.tiling = VK_IMAGE_TILING_LINEAR;
+	}
 }
diff --git a/Engine/Source/Runtime/VulkanRHI/Private/Android/VulkanAndroidPlatform.h b/Engine/Source/Runtime/VulkanRHI/Private/Android/VulkanAndroidPlatform.h
--- a/Engine/Source/Runtime/VulkanRHI/Private/Android/VulkanAndroidPlatform.h
+++ b/Engine/Source/Runtime/VulkanRHI/Private/Android/VulkanAndroidPlatform.h
@@ -144,9 +144,9 @@ public:
 		return VK_FORMAT_R5G5B5A1_UNORM_PACK16; 
 	};
 
-	// Setup platform to use AFBC workaround to reduce uncompressed textures memory requirements
-	static void SetupAFBCWorkaround(const FVulkanDevice& InDevice);
-	static void SetImageAFBCWorkaround(VkImageCreateInfo& ImageCreateInfo);
+	// Setup platform to use a workaround to reduce textures memory requirements
+	static void SetupImageMemoryRequirementWorkaround(const FVulkanDevice& InDevice);
+	static void SetImageMemoryRequirementWorkaround(VkImageCreateInfo& ImageCreateInfo);
 
 protected:
 	static void* VulkanLib;
@@ -167,6 +167,7 @@ protected:
 	static TArray<TArray<ANSICHAR>> DebugVulkanInstanceLayers;
 
 	static int32 AFBCWorkaroundOption;
+	static int32 ASTCWorkaroundOption;
 };
 
 #if VULKAN_SUPPORTS_GOOGLE_DISPLAY_TIMING
diff --git a/Engine/Source/Runtime/VulkanRHI/Private/VulkanDevice.cpp b/Engine/Source/Runtime/VulkanRHI/Private/VulkanDevice.cpp
--- a/Engine/Source/Runtime/VulkanRHI/Private/VulkanDevice.cpp
+++ b/Engine/Source/Runtime/VulkanRHI/Private/VulkanDevice.cpp
@@ -1127,7 +1127,7 @@ void FVulkanDevice::InitGPU(int32 DeviceIndex)
 		DefaultTexture = new FVulkanTexture(*this, Desc, nullptr);
 	}
 
-	FVulkanPlatform::SetupAFBCWorkaround(*this);
+	FVulkanPlatform::SetupImageMemoryRequirementWorkaround(*this);
 }
 
 void FVulkanDevice::PrepareForDestroy()
diff --git a/Engine/Source/Runtime/VulkanRHI/Private/VulkanGenericPlatform.h b/Engine/Source/Runtime/VulkanRHI/Private/VulkanGenericPlatform.h
--- a/Engine/Source/Runtime/VulkanRHI/Private/VulkanGenericPlatform.h
+++ b/Engine/Source/Runtime/VulkanRHI/Private/VulkanGenericPlatform.h
@@ -134,7 +134,7 @@ public:
 	// Return VK_FALSE if platform wants to suppress the given debug report from the validation layers, VK_TRUE to print it.
 	static VkBool32 DebugReportFunction(VkDebugReportFlagsEXT MsgFlags, VkDebugReportObjectTypeEXT ObjType, uint64_t SrcObject, size_t Location, int32 MsgCode, const ANSICHAR* LayerPrefix, const ANSICHAR* Msg, void* UserData) { return VK_TRUE; }
 
-	// Setup platform to use AFBC workaround to reduce uncompressed textures memory requirements
-	static void SetupAFBCWorkaround(const FVulkanDevice& InDevice) {};
-	static void SetImageAFBCWorkaround(VkImageCreateInfo& ImageCreateInfo) {};
+	// Setup platform to use a workaround to reduce textures memory requirements
+	static void SetupImageMemoryRequirementWorkaround(const FVulkanDevice& InDevice) {};
+	static void SetImageMemoryRequirementWorkaround(VkImageCreateInfo& ImageCreateInfo) {};
 };
diff --git a/Engine/Source/Runtime/VulkanRHI/Private/VulkanTexture.cpp b/Engine/Source/Runtime/VulkanRHI/Private/VulkanTexture.cpp
--- a/Engine/Source/Runtime/VulkanRHI/Private/VulkanTexture.cpp
+++ b/Engine/Source/Runtime/VulkanRHI/Private/VulkanTexture.cpp
@@ -439,7 +439,7 @@ void FVulkanTexture::GenerateImageCreateInfo(
 		break;
 	}
 
-	FVulkanPlatform::SetImageAFBCWorkaround(ImageCreateInfo);
+	FVulkanPlatform::SetImageMemoryRequirementWorkaround(ImageCreateInfo);
 		
 	const VkFormatFeatureFlags FormatFlags = ImageCreateInfo.tiling == VK_IMAGE_TILING_LINEAR ? 
 		InDevice.GetFormatProperties()[ImageCreateInfo.format].linearTilingFeatures : 
