Adding GRHISupportsLargerVariableRateShadingSizes to control support of 2x4/4x2/4x4 tile sizes.
On D3D12, this maps to D3D12_FEATURE_DATA_D3D12_OPTIONS6::AdditionalShadingRatesSupported.

diff --git a/Engine/Source/Runtime/D3D12RHI/Private/Windows/WindowsD3D12Device.cpp b/Engine/Source/Runtime/D3D12RHI/Private/Windows/WindowsD3D12Device.cpp
--- a/Engine/Source/Runtime/D3D12RHI/Private/Windows/WindowsD3D12Device.cpp
+++ b/Engine/Source/Runtime/D3D12RHI/Private/Windows/WindowsD3D12Device.cpp
@@ -1457,6 +1457,8 @@ void FD3D12DynamicRHI::Init()
 	if(hr == S_OK && options.VariableShadingRateTier != D3D12_VARIABLE_SHADING_RATE_TIER_NOT_SUPPORTED)
 	{
 		GRHISupportsPipelineVariableRateShading = true;		// We have at least tier 1.
+		GRHISupportsLargerVariableRateShadingSizes = (options.AdditionalShadingRatesSupported != 0);
+
 		if (options.VariableShadingRateTier == D3D12_VARIABLE_SHADING_RATE_TIER_2)
 		{
 			GRHISupportsAttachmentVariableRateShading = true;
diff --git a/Engine/Source/Runtime/RHI/Private/RHI.cpp b/Engine/Source/Runtime/RHI/Private/RHI.cpp
--- a/Engine/Source/Runtime/RHI/Private/RHI.cpp
+++ b/Engine/Source/Runtime/RHI/Private/RHI.cpp
@@ -1473,6 +1473,7 @@ bool GRHISupportsHDROutput = false;
 bool GRHIVariableRateShadingEnabled = true;
 bool GRHIAttachmentVariableRateShadingEnabled = true;
 bool GRHISupportsPipelineVariableRateShading = false;
+bool GRHISupportsLargerVariableRateShadingSizes = false;
 bool GRHISupportsAttachmentVariableRateShading = false;
 bool GRHISupportsComplexVariableRateShadingCombinerOps = false;
 bool GRHISupportsVariableRateShadingAttachmentArrayTextures = false;
diff --git a/Engine/Source/Runtime/RHI/Public/RHI.h b/Engine/Source/Runtime/RHI/Public/RHI.h
--- a/Engine/Source/Runtime/RHI/Public/RHI.h
+++ b/Engine/Source/Runtime/RHI/Public/RHI.h
@@ -722,6 +722,9 @@ extern RHI_API FIntVector GRHIMaxDispatchThreadGroupsPerDimension;
 /** Whether or not the RHI can support per-draw Variable Rate Shading. */
 extern RHI_API bool GRHISupportsPipelineVariableRateShading;
 
+/** Whether or not the Variable Rate Shading can be done at larger (2x4 or 4x2 or 4x4) sizes. */
+extern RHI_API bool GRHISupportsLargerVariableRateShadingSizes;
+
 /** Whether or not the RHI can support image-based Variable Rate Shading. */
 extern RHI_API bool GRHISupportsAttachmentVariableRateShading;
 
diff --git a/Engine/Source/Runtime/Renderer/Public/MeshPassProcessor.inl b/Engine/Source/Runtime/Renderer/Public/MeshPassProcessor.inl
--- a/Engine/Source/Runtime/Renderer/Public/MeshPassProcessor.inl
+++ b/Engine/Source/Runtime/Renderer/Public/MeshPassProcessor.inl
@@ -20,12 +20,19 @@ static EVRSShadingRate GetShadingRateFromMaterial(EMaterialShadingRate MaterialS
 			return EVRSShadingRate::VRSSR_2x1;
 		case MSR_2x2:
 			return EVRSShadingRate::VRSSR_2x2;
-		case MSR_4x2:
-			return EVRSShadingRate::VRSSR_4x2;
-		case MSR_2x4:
-			return EVRSShadingRate::VRSSR_2x4;						
-		case MSR_4x4:
-			return EVRSShadingRate::VRSSR_4x4;
+		}
+
+		if (GRHISupportsLargerVariableRateShadingSizes)
+		{
+			switch (MaterialShadingRate)
+			{
+			case MSR_4x2:
+				return EVRSShadingRate::VRSSR_4x2;
+			case MSR_2x4:
+				return EVRSShadingRate::VRSSR_2x4;
+			case MSR_4x4:
+				return EVRSShadingRate::VRSSR_4x4;
+			}
 		}
 	}
 	return EVRSShadingRate::VRSSR_1x1;
diff --git a/Engine/Source/Runtime/VulkanRHI/Private/VulkanExtensions.cpp b/Engine/Source/Runtime/VulkanRHI/Private/VulkanExtensions.cpp
--- a/Engine/Source/Runtime/VulkanRHI/Private/VulkanExtensions.cpp
+++ b/Engine/Source/Runtime/VulkanRHI/Private/VulkanExtensions.cpp
@@ -460,6 +460,12 @@ public:
 		GRHIVariableRateShadingImageTileMaxWidth = FragmentShadingRateProperties.maxFragmentShadingRateAttachmentTexelSize.width;
 		GRHIVariableRateShadingImageTileMaxHeight = FragmentShadingRateProperties.maxFragmentShadingRateAttachmentTexelSize.height;
 
+		if (FragmentShadingRateProperties.maxFragmentSize.width >= 4 && FragmentShadingRateProperties.maxFragmentSize.height >= 4)
+		{
+			// FYI FVulkanDevice::GetBestMatchedShadingRateExtents does extent filtering
+			GRHISupportsLargerVariableRateShadingSizes = GRHISupportsPipelineVariableRateShading;
+		}
+
 		// todo: We don't currently care much about the other properties here, but at some point in the future we probably will.
 
 		UE_LOG(LogVulkanRHI, Verbose, TEXT("Image-based Variable Rate Shading supported via KHRFragmentShadingRate extension. Selected VRS tile size %u by %u pixels per VRS image texel."), GRHIVariableRateShadingImageTileMinWidth, GRHIVariableRateShadingImageTileMinHeight);
