VulkanRHI:
- Fix the VULKAN_ENABLE_DUMP_LAYER and VULKAN_ENABLE_TRACKING_LAYER build errors:  fix parentheses, duplicate function bodies and raytracing entrypoints.
- Add VK_KHR_synchronization2 entrypoints (WIP)

diff --git a/Engine/Source/Runtime/VulkanRHI/Private/Linux/VulkanLinuxPlatform.cpp b/Engine/Source/Runtime/VulkanRHI/Private/Linux/VulkanLinuxPlatform.cpp
--- a/Engine/Source/Runtime/VulkanRHI/Private/Linux/VulkanLinuxPlatform.cpp
+++ b/Engine/Source/Runtime/VulkanRHI/Private/Linux/VulkanLinuxPlatform.cpp
@@ -110,7 +110,7 @@ bool FVulkanLinuxPlatform::LoadVulkanInstanceFunctions(VkInstance inInstance)
 	}
 
 #if VULKAN_RHI_RAYTRACING
-	const bool bFoundRayTracingEntries = FVulkanRayTracingPlatform::LoadVulkanInstanceFunctions(inInstance);
+	const bool bFoundRayTracingEntries = FVulkanRayTracingPlatform::CheckVulkanInstanceFunctions(inInstance);
 	if (!bFoundRayTracingEntries)
 	{
 		UE_LOG(LogVulkanRHI, Warning, TEXT("Vulkan RHI ray tracing is enabled, but failed to load instance functions."));
diff --git a/Engine/Source/Runtime/VulkanRHI/Private/VulkanDebug.cpp b/Engine/Source/Runtime/VulkanRHI/Private/VulkanDebug.cpp
--- a/Engine/Source/Runtime/VulkanRHI/Private/VulkanDebug.cpp
+++ b/Engine/Source/Runtime/VulkanRHI/Private/VulkanDebug.cpp
@@ -1086,38 +1086,6 @@ namespace VulkanRHI
 		return String;
 	}
 
-	static FString GetAccessFlagString(VkAccessFlags Flags)
-	{
-		if (Flags == 0)
-		{
-			return TEXT("0");
-		}
-		FString String;
-		AppendBitFieldName(VK_ACCESS_INDIRECT_COMMAND_READ_BIT, TEXT("INDIRECT_COMMAND"));
-		AppendBitFieldName(VK_ACCESS_INDEX_READ_BIT, TEXT("INDEX_READ"));
-		AppendBitFieldName(VK_ACCESS_VERTEX_ATTRIBUTE_READ_BIT, TEXT("VERTEX_ATTR_READ"));
-		AppendBitFieldName(VK_ACCESS_UNIFORM_READ_BIT, TEXT("UNIF_READ"));
-		AppendBitFieldName(VK_ACCESS_INPUT_ATTACHMENT_READ_BIT, TEXT("INPUT_ATT_READ"));
-		AppendBitFieldName(VK_ACCESS_SHADER_READ_BIT, TEXT("SHADER_READ"));
-		AppendBitFieldName(VK_ACCESS_SHADER_WRITE_BIT, TEXT("SHADER_WRITE"));
-		AppendBitFieldName(VK_ACCESS_COLOR_ATTACHMENT_READ_BIT, TEXT("COLOR_ATT_READ"));
-		AppendBitFieldName(VK_ACCESS_COLOR_ATTACHMENT_WRITE_BIT, TEXT("COLOR_ATT_WRITE"));
-		AppendBitFieldName(VK_ACCESS_DEPTH_STENCIL_ATTACHMENT_READ_BIT, TEXT("DS_ATT_READ"));
-		AppendBitFieldName(VK_ACCESS_DEPTH_STENCIL_ATTACHMENT_WRITE_BIT, TEXT("DS_ATT_WRITE"));
-		AppendBitFieldName(VK_ACCESS_TRANSFER_READ_BIT, TEXT("TRANSFER_READ"));
-		AppendBitFieldName(VK_ACCESS_TRANSFER_WRITE_BIT, TEXT("TRANSFER_WRITE"));
-		AppendBitFieldName(VK_ACCESS_HOST_READ_BIT, TEXT("HOST_READ"));
-		AppendBitFieldName(VK_ACCESS_HOST_WRITE_BIT, TEXT("HOST_WRITE"));
-		AppendBitFieldName(VK_ACCESS_MEMORY_READ_BIT, TEXT("MEM_READ"));
-		AppendBitFieldName(VK_ACCESS_MEMORY_WRITE_BIT, TEXT("MEM_WRITE"));
-		if (Flags != 0)
-		{
-			FString Unknown = FString::Printf(TEXT("%d"), Flags);
-			AppendBitFieldName(Flags, Unknown);
-		}
-		return String;
-	}
-
 	FString GetSampleCountString(VkSampleCountFlags Flags)
 	{
 		if (Flags == 0)
@@ -1206,11 +1174,6 @@ namespace VulkanRHI
 		return FString::Printf(TEXT("AspectMask=%s, BaseMip=%d, NumLevels=%d, BaseArrayLayer=%d, NumLayers=%d"), *GetAspectMaskString(Range.aspectMask), Range.baseMipLevel, Range.levelCount, Range.baseArrayLayer, Range.layerCount);
 	}
 
-	static FString GetStageMaskString(VkPipelineStageFlags Flags)
-	{
-		return FString::Printf(TEXT("VkPipelineStageFlags=0x%x"), (uint32)Flags);
-	}
-
 	static FString GetClearColorValueString(const VkClearColorValue& Value)
 	{
 		return FString::Printf(TEXT("0x%x(%f), 0x%x(%f), 0x%x(%f), 0x%x(%f)"),
@@ -1803,9 +1766,20 @@ static void DumpImageMemoryBarriers(uint32 ImageMemoryBarrierCount, const VkImag
 {
 	for (uint32 Index = 0; Index < ImageMemoryBarrierCount; ++Index)
 	{
-		DebugLog += FString::Printf(TEXT("%s\tImageBarrier[%d]: srcAccess=%s, oldLayout=%s, srcQueueFamilyIndex=%d\n"), Tabs, Index, *GetAccessFlagString(ImageMemoryBarriers[Index].srcAccessMask), *GetVkImageLayoutString(ImageMemoryBarriers[Index].oldLayout), ImageMemoryBarriers[Index].srcQueueFamilyIndex);
-		DebugLog += FString::Printf(TEXT("%s\t\tdstAccess=%s, newLayout=%s, dstQueueFamilyIndex=%d\n"), Tabs, *GetAccessFlagString(ImageMemoryBarriers[Index].dstAccessMask), *GetVkImageLayoutString(ImageMemoryBarriers[Index].newLayout), ImageMemoryBarriers[Index].dstQueueFamilyIndex);
-		DebugLog += FString::Printf(TEXT("%s\t\tImage=0x%p, subresourceRange=(%s)\n"), Tabs, ImageMemoryBarriers[Index].image, *GetImageSubResourceRangeString(ImageMemoryBarriers[Index].subresourceRange));
+		const VkImageMemoryBarrier& ImageBarrier = ImageMemoryBarriers[Index];
+		DebugLog += FString::Printf(TEXT("%s\tImageBarrier[%d]: srcAccess=%s, oldLayout=%s, srcQueueFamilyIndex=%d\n"), Tabs, Index, VK_FLAGS_TO_STRING(VkAccessFlags, ImageBarrier.srcAccessMask), VK_TYPE_TO_STRING(VkImageLayout, ImageBarrier.oldLayout), ImageBarrier.srcQueueFamilyIndex);
+		DebugLog += FString::Printf(TEXT("%s\t\tdstAccess=%s, newLayout=%s, dstQueueFamilyIndex=%d\n"), Tabs, VK_FLAGS_TO_STRING(VkAccessFlags, ImageBarrier.dstAccessMask), VK_TYPE_TO_STRING(VkImageLayout, ImageBarrier.newLayout), ImageBarrier.dstQueueFamilyIndex);
+		DebugLog += FString::Printf(TEXT("%s\t\tImage=0x%p, subresourceRange=(%s)\n"), Tabs, ImageBarrier.image, *GetImageSubResourceRangeString(ImageBarrier.subresourceRange));
+	}
+}
+static void DumpImageMemoryBarriers(uint32 ImageMemoryBarrierCount, const VkImageMemoryBarrier2* ImageMemoryBarriers)
+{
+	for (uint32 Index = 0; Index < ImageMemoryBarrierCount; ++Index)
+	{
+		const VkImageMemoryBarrier2& ImageBarrier = ImageMemoryBarriers[Index];
+		DebugLog += FString::Printf(TEXT("%s\tImageBarrier[%d]: srcStage=%s, srcAccess=%s, oldLayout=%s, srcQueueFamilyIndex=%d\n"), Tabs, Index, VK_FLAGS_TO_STRING(VkPipelineStageFlags2, ImageBarrier.srcStageMask), VK_FLAGS_TO_STRING(VkAccessFlags2, ImageBarrier.srcAccessMask), VK_TYPE_TO_STRING(VkImageLayout, ImageBarrier.oldLayout), ImageBarrier.srcQueueFamilyIndex);
+		DebugLog += FString::Printf(TEXT("%s\t\tdstStage=%s, dstAccess=%s, newLayout=%s, dstQueueFamilyIndex=%d\n"), Tabs, VK_FLAGS_TO_STRING(VkPipelineStageFlags2, ImageBarrier.dstStageMask), VK_FLAGS_TO_STRING(VkAccessFlags2, ImageBarrier.dstAccessMask), VK_TYPE_TO_STRING(VkImageLayout, ImageBarrier.newLayout), ImageBarrier.dstQueueFamilyIndex);
+		DebugLog += FString::Printf(TEXT("%s\t\tImage=0x%p, subresourceRange=(%s)\n"), Tabs, ImageBarrier.image, *GetImageSubResourceRangeString(ImageBarrier.subresourceRange));
 	}
 }
 #endif
@@ -1815,7 +1789,7 @@ void FWrapLayer::CmdPipelineBarrier(VkResult Result, VkCommandBuffer CommandBuff
 	if (Result == VK_RESULT_MAX_ENUM)
 	{
 #if VULKAN_ENABLE_DUMP_LAYER
-		CmdPrintfBegin(CommandBuffer, FString::Printf(TEXT("vkCmdPipelineBarrier(SrcMask=%s, DestMask=%s, Flags=%d, NumMemB=%d, MemB=0x%p,"), *GetStageMaskString(SrcStageMask), *GetStageMaskString(DstStageMask), (uint32)DependencyFlags, MemoryBarrierCount, MemoryBarriers));
+		CmdPrintfBegin(CommandBuffer, FString::Printf(TEXT("vkCmdPipelineBarrier(srcStageMask=%s, destStageMask=%s, Flags=%d, NumMemB=%d, MemB=0x%p,"), VK_FLAGS_TO_STRING(VkPipelineStageFlags, SrcStageMask), VK_FLAGS_TO_STRING(VkPipelineStageFlags, DstStageMask), (uint32)DependencyFlags, MemoryBarrierCount, MemoryBarriers));
 		DebugLog += FString::Printf(TEXT("%s\tNumBufferB=%d, BufferB=0x%p, NumImageB=%d, ImageB=0x%p)[...]\n"), Tabs, BufferMemoryBarrierCount, BufferMemoryBarriers, ImageMemoryBarrierCount, ImageMemoryBarriers);
 		DumpImageMemoryBarriers(ImageMemoryBarrierCount, ImageMemoryBarriers);
 		FlushDebugWrapperLog();
@@ -1829,7 +1803,7 @@ void FWrapLayer::CmdPipelineBarrier(VkResult Result, VkCommandBuffer CommandBuff
 				BreakOnTrackingImage(ImageMemoryBarriers[Index].image);
 				if(DumpTrackImage(ImageMemoryBarriers[Index].image))
 				{
-					CmdPrintfBegin(CommandBuffer, FString::Printf(TEXT("vkCmdPipelineBarrier(SrcMask=%s, DestMask=%s, Flags=%d, NumMemB=%d, MemB=0x%p,"), *GetStageMaskString(SrcStageMask), *GetStageMaskString(DstStageMask), (uint32)DependencyFlags, MemoryBarrierCount, MemoryBarriers));
+					CmdPrintfBegin(CommandBuffer, FString::Printf(TEXT("vkCmdPipelineBarrier(SrcStageMask=%s, DstStageMask=%s, Flags=%d, NumMemB=%d, MemB=0x%p,"), VK_FLAGS_TO_STRING(VkPipelineStageFlags, SrcStageMask), VK_FLAGS_TO_STRING(VkPipelineStageFlags, DstStageMask), (uint32)DependencyFlags, MemoryBarrierCount, MemoryBarriers));
 					DebugLog += FString::Printf(TEXT("%s\tNumBufferB=%d, BufferB=0x%p, NumImageB=%d, ImageB=0x%p)[...]\n"), Tabs, BufferMemoryBarrierCount, BufferMemoryBarriers, ImageMemoryBarrierCount, ImageMemoryBarriers);
 					DumpImageMemoryBarriers(ImageMemoryBarrierCount, ImageMemoryBarriers);
 					FlushDebugWrapperLog();
@@ -1870,7 +1844,7 @@ void FWrapLayer::CmdWaitEvents(VkResult Result, VkCommandBuffer CommandBuffer, u
 	if (Result == VK_RESULT_MAX_ENUM)
 	{
 #if VULKAN_ENABLE_DUMP_LAYER
-		CmdPrintfBegin(CommandBuffer, FString::Printf(TEXT("vkCmdWaitEvents(NumEvents=%d, Events=0x%p, SrcMask=%s, DestMask=%s, NumMemB=%d, MemB=0x%p,"), EventCount, Events, *GetStageMaskString(SrcStageMask), *GetStageMaskString(DstStageMask), MemoryBarrierCount, MemoryBarriers));
+		CmdPrintfBegin(CommandBuffer, FString::Printf(TEXT("vkCmdWaitEvents(NumEvents=%d, Events=0x%p, SrcStageMask=%s, DestStageMask=%s, NumMemB=%d, MemB=0x%p,"), EventCount, Events, VK_FLAGS_TO_STRING(VkPipelineStageFlags, SrcStageMask), VK_FLAGS_TO_STRING(VkPipelineStageFlags, DstStageMask), MemoryBarrierCount, MemoryBarriers));
 		DebugLog += FString::Printf(TEXT("%s\tNumBufferB=%d, BufferB=0x%p, NumImageB=%d, ImageB=0x%p)[...]\n"), Tabs, BufferMemoryBarrierCount, BufferMemoryBarriers, ImageMemoryBarrierCount, ImageMemoryBarriers);
 		for (uint32 Index = 0; Index < EventCount; ++Index)
 		{
@@ -4215,24 +4189,6 @@ void FWrapLayer::GetPhysicalDeviceSurfaceFormatsKHR(VkResult Result, VkPhysicalD
 	}
 }
 
-#if VULKAN_SUPPORTS_FRAGMENT_SHADING_RATE
-void FWrapLayer::GetPhysicalDeviceFragmentShadingRatesKHR(VkResult Result, VkPhysicalDevice PhysicalDevice, uint32* FragmentShadingRateCount, VkPhysicalDeviceFragmentShadingRateKHR* FragmentShadingRates)
-{
-	if (Result == VK_RESULT_MAX_ENUM)
-	{
-#if VULKAN_ENABLE_DUMP_LAYER
-		PrintfBeginResult(FString::Printf(TEXT("vkGetPhysicalDeviceFragmentShadingRatesKHR(PhysicalDevice=0x%p, FragmentShadingRateCount=0x%p, FragmentShadingRates=0x%p)[...]"), PhysicalDevice, FragmentShadingRateCount, FragmentShadingRates));
-#endif
-	}
-	else
-	{
-#if VULKAN_ENABLE_DUMP_LAYER
-		PrintResult(Result);
-#endif
-	}
-}
-#endif
-
 void FWrapLayer::GetPhysicalDeviceSurfaceSupportKHR(VkResult Result, VkPhysicalDevice PhysicalDevice, uint32_t QueueFamilyIndex, VkSurfaceKHR Surface, VkBool32* SupportedPtr)
 {
 	if (Result == VK_RESULT_MAX_ENUM)
@@ -4462,11 +4418,109 @@ void FWrapLayer::ResetQueryPoolEXT(VkResult Result, VkDevice Device, VkQueryPool
 	if (Result == VK_RESULT_MAX_ENUM)
 	{
 #if VULKAN_ENABLE_DUMP_LAYER
-		PrintfBeginResult(FString::Printf(TEXT("ResetQueryPoolEXT(Device=0x%p, QueryPool=0x%p, FirstQuery=%u, QueryCount=%u)")), Device, QueryPool, FirstQuery, QueryCount);
+		PrintfBeginResult(FString::Printf(TEXT("ResetQueryPoolEXT(Device=0x%p, QueryPool=0x%p, FirstQuery=%u, QueryCount=%u)"), Device, QueryPool, FirstQuery, QueryCount));
 #endif
 	}
 }
 
+
+void FWrapLayer::CmdPipelineBarrier2KHR(VkResult Result, VkCommandBuffer CommandBuffer, const VkDependencyInfo* DependencyInfo)
+{
+	if (Result == VK_RESULT_MAX_ENUM)
+	{
+#if VULKAN_ENABLE_DUMP_LAYER
+		CmdPrintfBegin(CommandBuffer, FString::Printf(TEXT("vkCmdPipelineBarrier2(CommandBuffer=0x%p, Flags=%d, NumMemB=%d, MemB=0x%p,"), CommandBuffer, (uint32)DependencyInfo->dependencyFlags, DependencyInfo->memoryBarrierCount, DependencyInfo->pMemoryBarriers));
+		DebugLog += FString::Printf(TEXT("%s\tNumBufferB=%d, BufferB=0x%p, NumImageB=%d, ImageB=0x%p)[...]\n"), Tabs, DependencyInfo->bufferMemoryBarrierCount, DependencyInfo->pBufferMemoryBarriers, DependencyInfo->imageMemoryBarrierCount, DependencyInfo->pImageMemoryBarriers);
+		DumpImageMemoryBarriers(DependencyInfo->imageMemoryBarrierCount, DependencyInfo->pImageMemoryBarriers);
+		FlushDebugWrapperLog();
+#endif
+
+#if VULKAN_ENABLE_IMAGE_TRACKING_LAYER
+		{
+			FScopeLock ScopeLock(&GTrackingCS);
+			for (uint32 Index = 0; Index < DependencyInfo->imageMemoryBarrierCount; ++Index)
+			{
+				const VkImageMemoryBarrier2& ImageBarrier = DependencyInfo->pImageMemoryBarriers[Index];
+
+				BreakOnTrackingImage(ImageBarrier.image);
+				if (DumpTrackImage(ImageBarrier.image))
+				{
+					CmdPrintfBegin(CommandBuffer, FString::Printf(TEXT("vkCmdPipelineBarrier(Flags=%d, NumMemB=%d, MemB=0x%p,"), (uint32)DependencyInfo->dependencyFlags, DependencyInfo->memoryBarrierCount, DependencyInfo->pMemoryBarriers));
+					DebugLog += FString::Printf(TEXT("%s\tNumBufferB=%d, BufferB=0x%p, NumImageB=%d, ImageB=0x%p)[...]\n"), Tabs, DependencyInfo->bufferMemoryBarrierCount, DependencyInfo->pBufferMemoryBarriers, DependencyInfo->imageMemoryBarrierCount, DependencyInfo->pImageMemoryBarriers);
+					DumpImageMemoryBarriers(DependencyInfo->imageMemoryBarrierCount, DependencyInfo->pImageMemoryBarriers);
+					FlushDebugWrapperLog();
+				}
+				FTrackingImage* TrackingImage = GVulkanTrackingImageLayouts.Find(ImageBarrier.image);
+				check(TrackingImage);
+#if VULKAN_ENABLE_TRACKING_CALLSTACK
+				FTrackingImage::FHistoryEntry* HistoryEntry = new (TrackingImage->History) FTrackingImage::FHistoryEntry;
+				HistoryEntry->ArrayLayouts = TrackingImage->ArrayLayouts;
+#endif
+				const VkImageSubresourceRange& Range = ImageBarrier.subresourceRange;
+				uint32 NumLayers = (Range.layerCount == VK_REMAINING_ARRAY_LAYERS) ? (TrackingImage->Info.CreateInfo.arrayLayers - Range.baseArrayLayer) : Range.layerCount;
+				for (uint32 LIndex = Range.baseArrayLayer; LIndex < Range.baseArrayLayer + NumLayers; ++LIndex)
+				{
+					TArray<VkImageLayout>& MipLayouts = TrackingImage->ArrayLayouts[LIndex];
+					uint32 NumLevels = (Range.levelCount == VK_REMAINING_MIP_LEVELS) ? (TrackingImage->Info.CreateInfo.mipLevels - Range.baseMipLevel) : Range.levelCount;
+					for (uint32 MIndex = Range.baseMipLevel; MIndex < Range.baseMipLevel + NumLevels; ++MIndex)
+					{
+						if (ImageBarrier.oldLayout != VK_IMAGE_LAYOUT_UNDEFINED && MipLayouts[MIndex] != ImageBarrier.oldLayout)
+						{
+							ensure(0);
+						}
+						MipLayouts[MIndex] = ImageBarrier.newLayout;
+					}
+				}
+#if VULKAN_ENABLE_TRACKING_CALLSTACK
+				CaptureCallStack(HistoryEntry->Callstack, 2);
+#endif
+			}
+		}
+#endif
+	}
+}
+
+void FWrapLayer::CmdResetEvent2KHR(VkResult Result, VkCommandBuffer CommandBuffer, VkEvent Event, VkPipelineStageFlags2 StageMask)
+{
+	if (Result == VK_RESULT_MAX_ENUM)
+	{
+#if VULKAN_ENABLE_DUMP_LAYER
+		PrintfBeginResult(FString::Printf(TEXT("CmdResetEvent2KHR(CommandBuffer=0x%p, Event=0x%p, StageMask=0xllx)"), CommandBuffer, Event, StageMask));
+#endif
+	}
+}
+
+void FWrapLayer::CmdSetEvent2KHR(VkResult Result, VkCommandBuffer CommandBuffer, VkEvent Event, const VkDependencyInfo* DependencyInfo)
+{
+	if (Result == VK_RESULT_MAX_ENUM)
+	{
+#if VULKAN_ENABLE_DUMP_LAYER
+		PrintfBeginResult(FString::Printf(TEXT("CmdSetEvent2KHR(CommandBuffer=0x%p, Event=0x%p)"), CommandBuffer, Event));
+#endif
+	}
+}
+
+void FWrapLayer::CmdWaitEvents2KHR(VkResult Result, VkCommandBuffer CommandBuffer, uint32_t EventCount, const VkEvent* Events, const VkDependencyInfo* DependencyInfos)
+{
+	if (Result == VK_RESULT_MAX_ENUM)
+	{
+#if VULKAN_ENABLE_DUMP_LAYER
+		PrintfBeginResult(FString::Printf(TEXT("CmdWaitEvents2KHR(CommandBuffer=0x%p, EventCount=%u, Events=0x%p)"), CommandBuffer, EventCount, Events));
+#endif
+	}
+}
+
+void FWrapLayer::QueueSubmit2KHR(VkResult Result, VkQueue Queue, uint32_t SubmitCount, const VkSubmitInfo2* Submits, VkFence Fence)
+{
+	if (Result == VK_RESULT_MAX_ENUM)
+	{
+#if VULKAN_ENABLE_DUMP_LAYER
+		PrintfBeginResult(FString::Printf(TEXT("QueueSubmit2KHR(Queue=0x%p, SubmitCount=%u, Submits=0x%p, Fence=0x%p)"), Queue, SubmitCount, Submits, Fence));
+#endif
+	}
+}
+
+
 #if VULKAN_ENABLE_IMAGE_TRACKING_LAYER
 namespace VulkanRHI
 {
diff --git a/Engine/Source/Runtime/VulkanRHI/Private/VulkanLoader.h b/Engine/Source/Runtime/VulkanRHI/Private/VulkanLoader.h
--- a/Engine/Source/Runtime/VulkanRHI/Private/VulkanLoader.h
+++ b/Engine/Source/Runtime/VulkanRHI/Private/VulkanLoader.h
@@ -170,10 +170,28 @@
 	EnumMacro(PFN_vkCreateDisplayModeKHR, vkCreateDisplayModeKHR) \
 	EnumMacro(PFN_vkGetDisplayPlaneCapabilitiesKHR, vkGetDisplayPlaneCapabilitiesKHR)
 
+// List all optional Vulkan raytracing entrypoints used by Unreal
+#define ENUM_VK_ENTRYPOINTS_RAYTRACING(EnumMacro) \
+	EnumMacro(PFN_vkCreateAccelerationStructureKHR, vkCreateAccelerationStructureKHR) \
+	EnumMacro(PFN_vkDestroyAccelerationStructureKHR, vkDestroyAccelerationStructureKHR) \
+	EnumMacro(PFN_vkCmdBuildAccelerationStructuresKHR, vkCmdBuildAccelerationStructuresKHR) \
+	EnumMacro(PFN_vkGetAccelerationStructureBuildSizesKHR, vkGetAccelerationStructureBuildSizesKHR) \
+	EnumMacro(PFN_vkGetAccelerationStructureDeviceAddressKHR, vkGetAccelerationStructureDeviceAddressKHR) \
+	EnumMacro(PFN_vkCmdTraceRaysKHR, vkCmdTraceRaysKHR) \
+	EnumMacro(PFN_vkCreateRayTracingPipelinesKHR, vkCreateRayTracingPipelinesKHR) \
+	EnumMacro(PFN_vkGetRayTracingShaderGroupHandlesKHR, vkGetRayTracingShaderGroupHandlesKHR) \
+	EnumMacro(PFN_vkGetBufferDeviceAddressKHR, vkGetBufferDeviceAddressKHR)
+
 // List all optional Vulkan entry points used by Unreal that need to be loaded manually
 #define ENUM_VK_ENTRYPOINTS_OPTIONAL_INSTANCE(EnumMacro) \
+	ENUM_VK_ENTRYPOINTS_RAYTRACING(EnumMacro) \
 	EnumMacro(PFN_vkGetDeviceImageMemoryRequirementsKHR, vkGetDeviceImageMemoryRequirementsKHR) \
 	EnumMacro(PFN_vkResetQueryPoolEXT, vkResetQueryPoolEXT) \
+	EnumMacro(PFN_vkCmdPipelineBarrier2KHR, vkCmdPipelineBarrier2KHR) \
+	EnumMacro(PFN_vkCmdResetEvent2KHR, vkCmdResetEvent2KHR) \
+	EnumMacro(PFN_vkCmdSetEvent2KHR, vkCmdSetEvent2KHR) \
+	EnumMacro(PFN_vkCmdWaitEvents2KHR, vkCmdWaitEvents2KHR) \
+	EnumMacro(PFN_vkQueueSubmit2KHR, vkQueueSubmit2KHR) \
 	EnumMacro(PFN_vkCreateSharedSwapchainsKHR, vkCreateSharedSwapchainsKHR)
 
 // List of all Vulkan entry points
diff --git a/Engine/Source/Runtime/VulkanRHI/Private/VulkanRayTracing.cpp b/Engine/Source/Runtime/VulkanRHI/Private/VulkanRayTracing.cpp
--- a/Engine/Source/Runtime/VulkanRHI/Private/VulkanRayTracing.cpp
+++ b/Engine/Source/Runtime/VulkanRHI/Private/VulkanRayTracing.cpp
@@ -13,20 +13,14 @@
 #include "Misc/Parse.h"
 
 
-// Define ray tracing entry points
-#define DEFINE_VK_ENTRYPOINTS(Type,Func) VULKANRHI_API Type VulkanDynamicAPI::Func = NULL;
-ENUM_VK_ENTRYPOINTS_RAYTRACING(DEFINE_VK_ENTRYPOINTS)
-
 #if PLATFORM_WINDOWS
 #pragma warning(push)
 #pragma warning(disable : 4191) // warning C4191: 'type cast': unsafe conversion
 #endif // PLATFORM_WINDOWS
-bool FVulkanRayTracingPlatform::LoadVulkanInstanceFunctions(VkInstance inInstance)
+bool FVulkanRayTracingPlatform::CheckVulkanInstanceFunctions(VkInstance inInstance)
 {
 	bool bFoundAllEntryPoints = true;
 #define CHECK_VK_ENTRYPOINTS(Type,Func) if (VulkanDynamicAPI::Func == NULL) { bFoundAllEntryPoints = false; UE_LOG(LogRHI, Warning, TEXT("Failed to find entry point for %s"), TEXT(#Func)); }
-#define GETINSTANCE_VK_ENTRYPOINTS(Type, Func) VulkanDynamicAPI::Func = (Type)VulkanDynamicAPI::vkGetInstanceProcAddr(inInstance, #Func);
-	ENUM_VK_ENTRYPOINTS_RAYTRACING(GETINSTANCE_VK_ENTRYPOINTS);
 #if UE_BUILD_DEBUG
 	ENUM_VK_ENTRYPOINTS_RAYTRACING(CHECK_VK_ENTRYPOINTS);
 #endif
diff --git a/Engine/Source/Runtime/VulkanRHI/Private/VulkanRayTracing.h b/Engine/Source/Runtime/VulkanRHI/Private/VulkanRayTracing.h
--- a/Engine/Source/Runtime/VulkanRHI/Private/VulkanRayTracing.h
+++ b/Engine/Source/Runtime/VulkanRHI/Private/VulkanRayTracing.h
@@ -10,27 +10,10 @@ class FVulkanCommandListContext;
 class FVulkanResourceMultiBuffer;
 class FVulkanRayTracingLayout;
 
-#define ENUM_VK_ENTRYPOINTS_RAYTRACING(EnumMacro) \
-	EnumMacro(PFN_vkCreateAccelerationStructureKHR, vkCreateAccelerationStructureKHR) \
-	EnumMacro(PFN_vkDestroyAccelerationStructureKHR, vkDestroyAccelerationStructureKHR) \
-	EnumMacro(PFN_vkCmdBuildAccelerationStructuresKHR, vkCmdBuildAccelerationStructuresKHR) \
-	EnumMacro(PFN_vkGetAccelerationStructureBuildSizesKHR, vkGetAccelerationStructureBuildSizesKHR) \
-	EnumMacro(PFN_vkGetAccelerationStructureDeviceAddressKHR, vkGetAccelerationStructureDeviceAddressKHR) \
-	EnumMacro(PFN_vkCmdTraceRaysKHR, vkCmdTraceRaysKHR) \
-	EnumMacro(PFN_vkCreateRayTracingPipelinesKHR, vkCreateRayTracingPipelinesKHR) \
-	EnumMacro(PFN_vkGetRayTracingShaderGroupHandlesKHR, vkGetRayTracingShaderGroupHandlesKHR) \
-	EnumMacro(PFN_vkGetBufferDeviceAddressKHR, vkGetBufferDeviceAddressKHR)
-
-// Declare ray tracing entry points
-namespace VulkanDynamicAPI
-{
-	ENUM_VK_ENTRYPOINTS_RAYTRACING(DECLARE_VK_ENTRYPOINTS);
-}
-
 class FVulkanRayTracingPlatform
 {
 public:
-	static bool LoadVulkanInstanceFunctions(VkInstance inInstance);
+	static bool CheckVulkanInstanceFunctions(VkInstance inInstance);
 };
 
 struct FVkRtAllocation
diff --git a/Engine/Source/Runtime/VulkanRHI/Private/Windows/VulkanWindowsPlatform.cpp b/Engine/Source/Runtime/VulkanRHI/Private/Windows/VulkanWindowsPlatform.cpp
--- a/Engine/Source/Runtime/VulkanRHI/Private/Windows/VulkanWindowsPlatform.cpp
+++ b/Engine/Source/Runtime/VulkanRHI/Private/Windows/VulkanWindowsPlatform.cpp
@@ -145,7 +145,7 @@ bool FVulkanWindowsPlatform::LoadVulkanInstanceFunctions(VkInstance inInstance)
 #endif
 
 #if VULKAN_RHI_RAYTRACING
-	const bool bFoundRayTracingEntries = FVulkanRayTracingPlatform::LoadVulkanInstanceFunctions(inInstance);
+	const bool bFoundRayTracingEntries = FVulkanRayTracingPlatform::CheckVulkanInstanceFunctions(inInstance);
 	if (!bFoundRayTracingEntries)
 	{
 		UE_LOG(LogVulkanRHI, Warning, TEXT("Vulkan RHI ray tracing is enabled, but failed to load instance functions."));
