VulkanRHI: skip harmless SPIRV extension validation errors for SPV_GOOGLE_hlsl_functionality1 and SPV_GOOGLE_user_type.

#rnx

diff --git a/Engine/Source/Runtime/VulkanRHI/Private/VulkanDebug.cpp b/Engine/Source/Runtime/VulkanRHI/Private/VulkanDebug.cpp
--- a/Engine/Source/Runtime/VulkanRHI/Private/VulkanDebug.cpp
+++ b/Engine/Source/Runtime/VulkanRHI/Private/VulkanDebug.cpp
@@ -272,6 +272,16 @@ static VkBool32 DebugUtilsCallback(VkDebugUtilsMessageSeverityFlagBitsEXT MsgSev
 		// *** CreateDevice(): Attempting to enable deprecated extension VK_KHR_get_memory_requirements2, but this extension has been promoted to VK_VERSION_1_1.
 		return VK_FALSE;
 	}
+	else if (FCStringAnsi::Strstr(CallbackData->pMessage, "The SPIR-V Extension (SPV_GOOGLE_hlsl_functionality1) was declared") != nullptr)
+	{
+		// *** [Error:Validation(VUID-VkShaderModuleCreateInfo-pCode-04147)] vkCreateShaderModule(): The SPIR-V Extension (SPV_GOOGLE_hlsl_functionality1) was declared, but none of the requirements were met to use it.
+		return VK_FALSE;
+	 }
+	else if (FCStringAnsi::Strstr(CallbackData->pMessage, "The SPIR-V Extension (SPV_GOOGLE_user_type) was declared") != nullptr)
+	{
+		// *** [Error:Validation(VUID-VkShaderModuleCreateInfo-pCode-04147)] vkCreateShaderModule(): The SPIR-V Extension (SPV_GOOGLE_user_type) was declared, but none of the requirements were met to use it.
+		return VK_FALSE;
+	}
 
 	static TStaticArray<TSet<uint32>, (int32)EMsgBucket::Count> SeenCodes;
 	bool bUniqueMessages = GCVarUniqueValidationMessages->GetInt() != 0;
