VulkanRHI: Fix memory corruption and crashes in Vulkan raytracing.  Make buffer lock/unlock use provided context instead of fetching immediate context.

diff --git a/Engine/Source/Runtime/VulkanRHI/Private/VulkanBuffer.cpp b/Engine/Source/Runtime/VulkanRHI/Private/VulkanBuffer.cpp
--- a/Engine/Source/Runtime/VulkanRHI/Private/VulkanBuffer.cpp
+++ b/Engine/Source/Runtime/VulkanRHI/Private/VulkanBuffer.cpp
@@ -15,6 +15,18 @@
 static TMap<FVulkanResourceMultiBuffer*, VulkanRHI::FPendingBufferLock> GPendingLockIBs;
 static FCriticalSection GPendingLockIBsMutex;
 
+static FORCEINLINE VulkanRHI::FPendingBufferLock GetPendingBufferLock(FVulkanResourceMultiBuffer* Buffer)
+{
+	VulkanRHI::FPendingBufferLock PendingLock;
+
+	// Found only if it was created for Write
+	FScopeLock ScopeLock(&GPendingLockIBsMutex);
+	const bool bFound = GPendingLockIBs.RemoveAndCopyValue(Buffer, PendingLock);
+
+	checkf(bFound, TEXT("Mismatched Buffer Lock/Unlock!"));
+	return PendingLock;
+}
+
 static FORCEINLINE void UpdateVulkanBufferStats(uint64_t Size, VkBufferUsageFlags Usage, bool Allocating)
 {
 	const bool bUniformBuffer = !!(Usage & VK_BUFFER_USAGE_UNIFORM_BUFFER_BIT);
@@ -236,6 +248,11 @@ FVulkanResourceMultiBuffer::~FVulkanResourceMultiBuffer()
 }
 
 void* FVulkanResourceMultiBuffer::Lock(FRHICommandListBase& RHICmdList, EResourceLockMode LockMode, uint32 LockSize, uint32 Offset)
+{
+	return Lock(FVulkanCommandListContext::GetVulkanContext(RHICmdList.GetContext()), LockMode, LockSize, Offset);
+}
+
+void* FVulkanResourceMultiBuffer::Lock(FVulkanCommandListContext& Context, EResourceLockMode LockMode, uint32 LockSize, uint32 Offset)
 {
 	void* Data = nullptr;
 	uint32 DataOffset = 0;
@@ -258,7 +275,7 @@ void* FVulkanResourceMultiBuffer::Lock(FRHICommandListBase& RHICmdList, EResourc
 		}
 		else
 		{
-			Device->GetImmediateContext().GetTempFrameAllocationBuffer().Alloc(LockSize + Offset, 256, VolatileLockInfo);
+			Context.GetTempFrameAllocationBuffer().Alloc(LockSize + Offset, 256, VolatileLockInfo);
 			Data = VolatileLockInfo.Data;
 			++VolatileLockInfo.LockCounter;
 			check(!VolatileLockInfo.Allocation.HasAllocation());
@@ -274,7 +291,7 @@ void* FVulkanResourceMultiBuffer::Lock(FRHICommandListBase& RHICmdList, EResourc
 
 		if (LockMode == RLM_ReadOnly)
 		{
-			check(IsInRenderingThread() && RHICmdList.IsImmediate());
+			check(IsInRenderingThread() && Context.IsImmediate());
 
 			const bool bUnifiedMem = Device->HasUnifiedMemory();
 			if (bUnifiedMem)
@@ -286,8 +303,7 @@ void* FVulkanResourceMultiBuffer::Lock(FRHICommandListBase& RHICmdList, EResourc
 			else 
 			{
 				Device->PrepareForCPURead();
-				FVulkanCommandListContext& ImmediateContext = Device->GetImmediateContext();
-				FVulkanCmdBuffer* CmdBuffer = ImmediateContext.GetCommandBufferManager()->GetUploadCmdBuffer();
+				FVulkanCmdBuffer* CmdBuffer = Context.GetCommandBufferManager()->GetUploadCmdBuffer();
 				
 				// Make sure any previous tasks have finished on the source buffer.
 				VkMemoryBarrier BarrierBefore = { VK_STRUCTURE_TYPE_MEMORY_BARRIER, nullptr, VK_ACCESS_MEMORY_WRITE_BIT, VK_ACCESS_MEMORY_READ_BIT };
@@ -309,7 +325,7 @@ void* FVulkanResourceMultiBuffer::Lock(FRHICommandListBase& RHICmdList, EResourc
 				VulkanRHI::vkCmdPipelineBarrier(CmdBuffer->GetHandle(), VK_PIPELINE_STAGE_TRANSFER_BIT, VK_PIPELINE_STAGE_HOST_BIT, 0, 1, &BarrierAfter, 0, nullptr, 0, nullptr);
 				
 				// Force upload.
-				ImmediateContext.GetCommandBufferManager()->SubmitUploadCmdBuffer();
+				Context.GetCommandBufferManager()->SubmitUploadCmdBuffer();
 				Device->WaitUntilIdle();
 
 				// Flush.
@@ -331,7 +347,7 @@ void* FVulkanResourceMultiBuffer::Lock(FRHICommandListBase& RHICmdList, EResourc
 					GPendingLockIBs.Add(this, PendingLock);
 				}
 
-				ImmediateContext.GetCommandBufferManager()->PrepareForNewActiveCommandBuffer();
+				Context.GetCommandBufferManager()->PrepareForNewActiveCommandBuffer();
 			}
 		}
 		else
@@ -425,8 +441,20 @@ struct FRHICommandMultiBufferUnlock final : public FRHICommand<FRHICommandMultiB
 	}
 };
 
-
 void FVulkanResourceMultiBuffer::Unlock(FRHICommandListBase& RHICmdList)
+{
+	const bool bVolatile = EnumHasAnyFlags(GetUsage(), BUF_Volatile);
+	if (!bVolatile && (LockStatus != ELockStatus::PersistentMapping) && !RHICmdList.IsBottomOfPipe())
+	{
+		ALLOC_COMMAND_CL(RHICmdList, FRHICommandMultiBufferUnlock)(Device, GetPendingBufferLock(this), this, DynamicBufferIndex);
+	}
+	else
+	{
+		Unlock(FVulkanCommandListContext::GetVulkanContext(RHICmdList.GetContext()));
+	}
+}
+
+void FVulkanResourceMultiBuffer::Unlock(FVulkanCommandListContext& Context)
 {
 	const bool bDynamic = EnumHasAnyFlags(GetUsage(), BUF_Dynamic);
 	const bool bVolatile = EnumHasAnyFlags(GetUsage(), BUF_Volatile);
@@ -444,27 +472,13 @@ void FVulkanResourceMultiBuffer::Unlock(FRHICommandListBase& RHICmdList)
 	{
 		check(bStatic || bDynamic || bSR);
 		
-		VulkanRHI::FPendingBufferLock PendingLock;
-		bool bFound = false;
-		{
-			// Found only if it was created for Write
-			FScopeLock ScopeLock(&GPendingLockIBsMutex);
-			bFound = GPendingLockIBs.RemoveAndCopyValue(this, PendingLock);
-		}
+		VulkanRHI::FPendingBufferLock PendingLock = GetPendingBufferLock(this);
 
 		PendingLock.StagingBuffer->FlushMappedMemory();
 
-		checkf(bFound, TEXT("Mismatched lock/unlock IndexBuffer!"));
 		if (PendingLock.LockMode == RLM_WriteOnly)
 		{
-			if (RHICmdList.IsBottomOfPipe())
-			{
-				FVulkanResourceMultiBuffer::InternalUnlock(Device->GetImmediateContext(), PendingLock, this, DynamicBufferIndex);
-			}
-			else
-			{
-				ALLOC_COMMAND_CL(RHICmdList, FRHICommandMultiBufferUnlock)(Device, PendingLock, this, DynamicBufferIndex);
-			}
+			FVulkanResourceMultiBuffer::InternalUnlock(Context, PendingLock, this, DynamicBufferIndex);
 		}
 		else if(PendingLock.LockMode == RLM_ReadOnly)
 		{
diff --git a/Engine/Source/Runtime/VulkanRHI/Private/VulkanRayTracing.cpp b/Engine/Source/Runtime/VulkanRHI/Private/VulkanRayTracing.cpp
--- a/Engine/Source/Runtime/VulkanRHI/Private/VulkanRayTracing.cpp
+++ b/Engine/Source/Runtime/VulkanRHI/Private/VulkanRayTracing.cpp
@@ -451,7 +451,7 @@ void FVulkanRayTracingScene::BuildAccelerationStructure(
 {
 	// Build a metadata buffer	that contains VulkanRHI-specific per-geometry parameters that allow us to access
 	// vertex and index buffers from shaders that use inline ray tracing.
-	BuildPerInstanceGeometryParameterBuffer();
+	BuildPerInstanceGeometryParameterBuffer(CommandContext);
 
 	check(AccelerationStructureBuffer.IsValid());
 	check(InInstanceBuffer != nullptr);
@@ -516,7 +516,7 @@ void FVulkanRayTracingScene::BuildAccelerationStructure(
 	CommandBufferManager.PrepareForNewActiveCommandBuffer();
 }
 
-void FVulkanRayTracingScene::BuildPerInstanceGeometryParameterBuffer()
+void FVulkanRayTracingScene::BuildPerInstanceGeometryParameterBuffer(FVulkanCommandListContext& CommandContext)
 {
 	// TODO: we could cache parameters in the geometry object to avoid some of the pointer chasing (if this is measured to be a performance issue)
 
@@ -524,9 +524,8 @@ void FVulkanRayTracingScene::BuildPerInstanceGeometryParameterBuffer()
 	check(PerInstanceGeometryParameterBuffer->GetSize() >= ParameterBufferSize);
 
 	check(IsInRHIThread() || !IsRunningRHIInSeparateThread());
-	FRHICommandListImmediate& RHICmdList = FRHICommandListExecutor::GetImmediateCommandList();
 
-	void* MappedBuffer = PerInstanceGeometryParameterBuffer->Lock(RHICmdList, RLM_WriteOnly, ParameterBufferSize, 0);
+	void* MappedBuffer = PerInstanceGeometryParameterBuffer->Lock(CommandContext, RLM_WriteOnly, ParameterBufferSize, 0);
 	FVulkanRayTracingGeometryParameters* MappedParameters = reinterpret_cast<FVulkanRayTracingGeometryParameters*>(MappedBuffer);
 	uint32 ParameterIndex = 0;
 
@@ -571,7 +570,7 @@ void FVulkanRayTracingScene::BuildPerInstanceGeometryParameterBuffer()
 
 	check(ParameterIndex == Initializer.NumTotalSegments);
 
-	PerInstanceGeometryParameterBuffer->Unlock(RHICmdList);
+	PerInstanceGeometryParameterBuffer->Unlock(CommandContext);
 }
 
 void FVulkanDynamicRHI::RHITransferRayTracingGeometryUnderlyingResource(FRHIRayTracingGeometry* DestGeometry, FRHIRayTracingGeometry* SrcGeometry)
diff --git a/Engine/Source/Runtime/VulkanRHI/Private/VulkanRayTracing.h b/Engine/Source/Runtime/VulkanRHI/Private/VulkanRayTracing.h
--- a/Engine/Source/Runtime/VulkanRHI/Private/VulkanRayTracing.h
+++ b/Engine/Source/Runtime/VulkanRHI/Private/VulkanRayTracing.h
@@ -128,7 +128,7 @@ private:
 	// Buffer that contains per-instance index and vertex buffer binding data
 	TRefCountPtr<FVulkanResourceMultiBuffer> PerInstanceGeometryParameterBuffer;
 	TRefCountPtr<FVulkanShaderResourceView> PerInstanceGeometryParameterSRV;
-	void BuildPerInstanceGeometryParameterBuffer();
+	void BuildPerInstanceGeometryParameterBuffer(FVulkanCommandListContext& CommandContext);
 };
 
 class FVulkanRayTracingPipelineState : public FRHIRayTracingPipelineState
diff --git a/Engine/Source/Runtime/VulkanRHI/Private/VulkanResources.h b/Engine/Source/Runtime/VulkanRHI/Private/VulkanResources.h
--- a/Engine/Source/Runtime/VulkanRHI/Private/VulkanResources.h
+++ b/Engine/Source/Runtime/VulkanRHI/Private/VulkanResources.h
@@ -953,6 +953,9 @@ public:
 	void* Lock(FRHICommandListBase& RHICmdList, EResourceLockMode LockMode, uint32 Size, uint32 Offset);
 	void Unlock(FRHICommandListBase& RHICmdList);
 
+	void* Lock(FVulkanCommandListContext& Context, EResourceLockMode LockMode, uint32 Size, uint32 Offset);
+	void Unlock(FVulkanCommandListContext& Context);
+
 	void Swap(FVulkanResourceMultiBuffer& Other);
 
 
