For Vulkan BUF_RayTracingScratch is special as it requires different alignment but has no special storage bit. In order to account for the alignment requirements:
* Changed FMemoryManager::CalculateBufferAlignment to take EBufferUsageFlags as input
* Allow passing MinAlignment to AllocateBufferPooled, if not provided (ie. set to 0) the alignment will be calculated from VkBufferUsageFlags as before.
* Modify TransientResource allocator to use EBufferUsageFlags and do not set AccelerationStructure bit as we cannot allocate AccelerationStructures from transient heap.

diff --git a/Engine/Source/Runtime/VulkanRHI/Private/VulkanBuffer.cpp b/Engine/Source/Runtime/VulkanRHI/Private/VulkanBuffer.cpp
--- a/Engine/Source/Runtime/VulkanRHI/Private/VulkanBuffer.cpp
+++ b/Engine/Source/Runtime/VulkanRHI/Private/VulkanBuffer.cpp
@@ -167,10 +167,10 @@ FVulkanResourceMultiBuffer::FVulkanResourceMultiBuffer(FVulkanDevice* InDevice,
 			check(NumBuffers <= UE_ARRAY_COUNT(Buffers));
 
 			const bool bUnifiedMem = InDevice->HasUnifiedMemory();
+			const uint32 BufferAlignment = FMemoryManager::CalculateBufferAlignment(*InDevice, InUEUsage, bZeroSize);
 
 			if (InTransientHeapAllocation != nullptr)
 			{
-				const uint32 BufferAlignment = FMemoryManager::CalculateBufferAlignment(*InDevice, BufferUsageFlags);
 				const uint32 AlignedSize = Align(InSize, BufferAlignment);
 
 				Buffers[0] = FVulkanTransientHeap::GetVulkanAllocation(*InTransientHeapAllocation);
@@ -196,7 +196,7 @@ FVulkanResourceMultiBuffer::FVulkanResourceMultiBuffer(FVulkanDevice* InDevice,
 
 				for (uint32 Index = 0; Index < NumBuffers; ++Index)
 				{
-					if (!InDevice->GetMemoryManager().AllocateBufferPooled(Buffers[Index], this, InSize, BufferUsageFlags, BufferMemFlags, EVulkanAllocationMetaMultiBuffer, __FILE__, __LINE__))
+					if (!InDevice->GetMemoryManager().AllocateBufferPooled(Buffers[Index], this, InSize, BufferAlignment, BufferUsageFlags, BufferMemFlags, EVulkanAllocationMetaMultiBuffer, __FILE__, __LINE__))
 					{
 						InDevice->GetMemoryManager().HandleOOM();
 					}
diff --git a/Engine/Source/Runtime/VulkanRHI/Private/VulkanMemory.cpp b/Engine/Source/Runtime/VulkanRHI/Private/VulkanMemory.cpp
--- a/Engine/Source/Runtime/VulkanRHI/Private/VulkanMemory.cpp
+++ b/Engine/Source/Runtime/VulkanRHI/Private/VulkanMemory.cpp
@@ -2626,7 +2626,7 @@ namespace VulkanRHI
 		VULKAN_FREE_TRACK_INFO(Track);
 	}
 
-	uint32 FMemoryManager::CalculateBufferAlignment(FVulkanDevice& InDevice, const VkBufferUsageFlags BufferUsageFlags)
+	static uint32 CalculateBufferAlignmentFromVKUsageFlags(FVulkanDevice& InDevice, const VkBufferUsageFlags BufferUsageFlags)
 	{
 		const VkPhysicalDeviceLimits& Limits = InDevice.GetLimits();
 
@@ -2663,6 +2663,21 @@ namespace VulkanRHI
 		return Alignment;
 	}
 
+	uint32 FMemoryManager::CalculateBufferAlignment(FVulkanDevice& InDevice, EBufferUsageFlags InUEUsage, bool bZeroSize)
+	{
+		const VkBufferUsageFlags VulkanBufferUsage = FVulkanResourceMultiBuffer::UEToVKBufferUsageFlags(&InDevice, InUEUsage, bZeroSize);
+
+		uint32 Alignment = CalculateBufferAlignmentFromVKUsageFlags(InDevice, VulkanBufferUsage);
+
+#if VULKAN_RHI_RAYTRACING
+		if (EnumHasAnyFlags(InUEUsage, BUF_RayTracingScratch))
+		{
+			Alignment = GRHIRayTracingScratchBufferAlignment;
+		}
+#endif
+		return Alignment;
+	}
+
 	float FMemoryManager::CalculateBufferPriority(const VkBufferUsageFlags BufferUsageFlags)
 	{
 		float Priority = VULKAN_MEMORY_MEDIUM_PRIORITY;
@@ -2700,12 +2715,12 @@ namespace VulkanRHI
 		return Priority;
 	}
 
-	bool FMemoryManager::AllocateBufferPooled(FVulkanAllocation& OutAllocation, FVulkanEvictable* AllocationOwner, uint32 Size, VkBufferUsageFlags BufferUsageFlags, VkMemoryPropertyFlags MemoryPropertyFlags, EVulkanAllocationMetaType MetaType, const char* File, uint32 Line)
+	bool FMemoryManager::AllocateBufferPooled(FVulkanAllocation& OutAllocation, FVulkanEvictable* AllocationOwner, uint32 Size, uint32 MinAlignment, VkBufferUsageFlags BufferUsageFlags, VkMemoryPropertyFlags MemoryPropertyFlags, EVulkanAllocationMetaType MetaType, const char* File, uint32 Line)
 	{
 		SCOPED_NAMED_EVENT(FResourceHeapManager_AllocateBufferPooled, FColor::Cyan);
 		check(OutAllocation.Type == EVulkanAllocationEmpty);
 
-		uint32 Alignment = CalculateBufferAlignment(*Device, BufferUsageFlags);
+		uint32 Alignment = FMath::Max(MinAlignment, CalculateBufferAlignmentFromVKUsageFlags(*Device, BufferUsageFlags));
 		const float Priority = CalculateBufferPriority(BufferUsageFlags);
 
 		const int32 PoolSize = (int32)GetPoolTypeForAlloc(Size, Alignment);
@@ -3556,7 +3571,7 @@ namespace VulkanRHI
 
 	void FMemoryManager::AllocUniformBuffer(FVulkanAllocation& OutAllocation, uint32 Size, const void* Contents)
 	{
-		if(!AllocateBufferPooled(OutAllocation, nullptr, Size, VK_BUFFER_USAGE_UNIFORM_BUFFER_BIT | VK_BUFFER_USAGE_TRANSFER_DST_BIT, VK_MEMORY_PROPERTY_HOST_VISIBLE_BIT | VK_MEMORY_PROPERTY_HOST_COHERENT_BIT, EVulkanAllocationMetaUniformBuffer, __FILE__, __LINE__))
+		if(!AllocateBufferPooled(OutAllocation, nullptr, Size, 0, VK_BUFFER_USAGE_UNIFORM_BUFFER_BIT | VK_BUFFER_USAGE_TRANSFER_DST_BIT, VK_MEMORY_PROPERTY_HOST_VISIBLE_BIT | VK_MEMORY_PROPERTY_HOST_COHERENT_BIT, EVulkanAllocationMetaUniformBuffer, __FILE__, __LINE__))
 		{
 			HandleOOM(false);
 			checkNoEntry();
@@ -4806,7 +4821,7 @@ namespace VulkanRHI
 		PeakUsed = 0;
 		FMemoryManager& ResourceHeapManager = InDevice->GetMemoryManager();
 		check(Allocation.Type == EVulkanAllocationEmpty);
-		if(ResourceHeapManager.AllocateBufferPooled(Allocation, nullptr, InSize,
+		if(ResourceHeapManager.AllocateBufferPooled(Allocation, nullptr, InSize, 0,
 			VK_BUFFER_USAGE_TRANSFER_SRC_BIT | VK_BUFFER_USAGE_VERTEX_BUFFER_BIT |
 			VK_BUFFER_USAGE_INDEX_BUFFER_BIT | VK_BUFFER_USAGE_UNIFORM_BUFFER_BIT |
 			VK_BUFFER_USAGE_UNIFORM_TEXEL_BUFFER_BIT | VK_BUFFER_USAGE_STORAGE_BUFFER_BIT,
diff --git a/Engine/Source/Runtime/VulkanRHI/Private/VulkanMemory.h b/Engine/Source/Runtime/VulkanRHI/Private/VulkanMemory.h
--- a/Engine/Source/Runtime/VulkanRHI/Private/VulkanMemory.h
+++ b/Engine/Source/Runtime/VulkanRHI/Private/VulkanMemory.h
@@ -773,7 +773,7 @@ namespace VulkanRHI
 		void Init();
 		void Deinit();
 
-		static uint32 CalculateBufferAlignment(FVulkanDevice& InDevice, const VkBufferUsageFlags BufferUsageFlags);
+		static uint32 CalculateBufferAlignment(FVulkanDevice& InDevice, EBufferUsageFlags InUEUsage, bool bZeroSize);
 		static float CalculateBufferPriority(const VkBufferUsageFlags BufferUsageFlags);
 
 		void FreeVulkanAllocation(FVulkanAllocation& Allocation, EVulkanFreeFlags FreeFlags = EVulkanFreeFlag_None);
@@ -783,7 +783,7 @@ namespace VulkanRHI
 		void FreeVulkanAllocationImageDedicated(FVulkanAllocation& Allocation);
 
 
-		bool AllocateBufferPooled(FVulkanAllocation& Allocation, FVulkanEvictable* AllocationOwner, uint32 Size, VkBufferUsageFlags BufferUsageFlags, VkMemoryPropertyFlags MemoryPropertyFlags, EVulkanAllocationMetaType MetaType, const char* File, uint32 Line);
+		bool AllocateBufferPooled(FVulkanAllocation& Allocation, FVulkanEvictable* AllocationOwner, uint32 Size, uint32 MinAlignment, VkBufferUsageFlags BufferUsageFlags, VkMemoryPropertyFlags MemoryPropertyFlags, EVulkanAllocationMetaType MetaType, const char* File, uint32 Line);
 		bool AllocateImageMemory(FVulkanAllocation& Allocation, FVulkanEvictable* AllocationOwner, const VkMemoryRequirements& MemoryReqs, VkMemoryPropertyFlags MemoryPropertyFlags, EVulkanAllocationMetaType MetaType, bool bExternal, const char* File, uint32 Line);
 		bool AllocateBufferMemory(FVulkanAllocation& Allocation, FVulkanEvictable* AllocationOwner, const VkMemoryRequirements& MemoryReqs, VkMemoryPropertyFlags MemoryPropertyFlags, EVulkanAllocationMetaType MetaType, bool bExternal, const char* File, uint32 Line);
 		bool AllocateDedicatedImageMemory(FVulkanAllocation& Allocation, FVulkanEvictable* AllocationOwner, VkImage Image, const VkMemoryRequirements& MemoryReqs, VkMemoryPropertyFlags MemoryPropertyFlags, EVulkanAllocationMetaType MetaType, bool bExternal, const char* File, uint32 Line);
diff --git a/Engine/Source/Runtime/VulkanRHI/Private/VulkanTransientResourceAllocator.cpp b/Engine/Source/Runtime/VulkanRHI/Private/VulkanTransientResourceAllocator.cpp
--- a/Engine/Source/Runtime/VulkanRHI/Private/VulkanTransientResourceAllocator.cpp
+++ b/Engine/Source/Runtime/VulkanRHI/Private/VulkanTransientResourceAllocator.cpp
@@ -9,17 +9,19 @@ FVulkanTransientHeap::FVulkanTransientHeap(const FInitializer& Initializer, FVul
 	, FDeviceChild(InDevice)
 	, VulkanBuffer(VK_NULL_HANDLE)
 {
-	VkBufferUsageFlags BufferUsageFlags =
-		VK_BUFFER_USAGE_VERTEX_BUFFER_BIT | VK_BUFFER_USAGE_INDEX_BUFFER_BIT | VK_BUFFER_USAGE_STORAGE_BUFFER_BIT |
-		VK_BUFFER_USAGE_STORAGE_TEXEL_BUFFER_BIT | VK_BUFFER_USAGE_UNIFORM_TEXEL_BUFFER_BIT |
-		VK_BUFFER_USAGE_INDIRECT_BUFFER_BIT | VK_BUFFER_USAGE_TRANSFER_SRC_BIT | VK_BUFFER_USAGE_TRANSFER_DST_BIT;
+	EBufferUsageFlags UEBufferUsageFlags = BUF_VertexBuffer | BUF_IndexBuffer | BUF_DrawIndirect 
+		| BUF_UnorderedAccess | BUF_StructuredBuffer | BUF_ShaderResource | BUF_KeepCPUAccessible;
 
 #if VULKAN_RHI_RAYTRACING
 	if (InDevice->GetOptionalExtensions().HasRaytracingExtensions())
 	{
-		BufferUsageFlags |= VK_BUFFER_USAGE_ACCELERATION_STRUCTURE_STORAGE_BIT_KHR | VK_BUFFER_USAGE_SHADER_DEVICE_ADDRESS_BIT;
+		UEBufferUsageFlags |= BUF_RayTracingScratch;
+		// AccelerationStructure not yet supported as TransientResource see FVulkanTransientResourceAllocator::CreateBuffer
+		//UEBufferUsageFlags |= BUF_AccelerationStructure;
 	}
 #endif
+	const bool bZeroSize = false;
+	VkBufferUsageFlags BufferUsageFlags = FVulkanResourceMultiBuffer::UEToVKBufferUsageFlags(InDevice, UEBufferUsageFlags, bZeroSize);
 
 	// :TODO: VK_KHR_maintenance4...
 	{
@@ -34,7 +36,7 @@ FVulkanTransientHeap::FVulkanTransientHeap(const FInitializer& Initializer, FVul
 		VulkanRHI::vkGetBufferMemoryRequirements(VulkanDevice, VulkanBuffer, &MemoryRequirements);
 
 		// Find the alignment that works for everyone
-		const uint32 MinBufferAlignment = FMemoryManager::CalculateBufferAlignment(*InDevice, BufferCreateInfo.usage);
+		const uint32 MinBufferAlignment = FMemoryManager::CalculateBufferAlignment(*InDevice, UEBufferUsageFlags, bZeroSize);
 		MemoryRequirements.alignment = FMath::Max<VkDeviceSize>(Initializer.Alignment, MemoryRequirements.alignment);
 		MemoryRequirements.alignment = FMath::Max<VkDeviceSize>(MinBufferAlignment, MemoryRequirements.alignment);
 	}
@@ -129,8 +131,8 @@ FRHITransientBuffer* FVulkanTransientResourceAllocator::CreateBuffer(const FRHIB
 	checkf(!EnumHasAnyFlags(InCreateInfo.Usage, BUF_AccelerationStructure), TEXT("AccelerationStructure not yet supported as TransientResource."));
 	checkf(!EnumHasAnyFlags(InCreateInfo.Usage, BUF_Volatile), TEXT("The volatile flag is not supported for transient resources."));
 
-	const VkBufferUsageFlags VulkanBufferUsage = FVulkanResourceMultiBuffer::UEToVKBufferUsageFlags(Device, InCreateInfo.Usage, (InCreateInfo.Size == 0));
-	const uint32 Alignment = FMemoryManager::CalculateBufferAlignment(*Device, VulkanBufferUsage);
+	const bool bZeroSize = (InCreateInfo.Size == 0);
+	const uint32 Alignment = FMemoryManager::CalculateBufferAlignment(*Device, InCreateInfo.Usage, bZeroSize);
 	uint64 Size = Align(InCreateInfo.Size, Alignment) * FVulkanResourceMultiBuffer::GetNumBuffersFromUsage(InCreateInfo.Usage);
 
 	return CreateBufferInternal(InCreateInfo, InDebugName, InPassIndex, Size, Alignment,
diff --git a/Engine/Source/Runtime/VulkanRHI/Private/VulkanUniformBuffer.cpp b/Engine/Source/Runtime/VulkanRHI/Private/VulkanUniformBuffer.cpp
--- a/Engine/Source/Runtime/VulkanRHI/Private/VulkanUniformBuffer.cpp
+++ b/Engine/Source/Runtime/VulkanRHI/Private/VulkanUniformBuffer.cpp
@@ -327,7 +327,7 @@ FVulkanRingBuffer::FVulkanRingBuffer(FVulkanDevice* InDevice, uint64 TotalSize,
 	, MinAlignment(0)
 {
 	check(TotalSize <= (uint64)MAX_uint32);
-	InDevice->GetMemoryManager().AllocateBufferPooled(Allocation, nullptr, TotalSize, Usage, MemPropertyFlags, EVulkanAllocationMetaRingBuffer, __FILE__, __LINE__);
+	InDevice->GetMemoryManager().AllocateBufferPooled(Allocation, nullptr, TotalSize, 0, Usage, MemPropertyFlags, EVulkanAllocationMetaRingBuffer, __FILE__, __LINE__);
 	MinAlignment = Allocation.GetBufferAlignment(Device);
 	// Start by wrapping around to set up the correct fence
 	BufferOffset = TotalSize;
