VulkanRHI: Move a few functions from ENUM_VK_ENTRYPOINTS_PLATFORM_INSTANCE to the more general ENUM_VK_ENTRYPOINTS_OPTIONAL_INSTANCE.  Either they were used on all platforms, or we want to keep them available in code built for all platforms (they remain optional and the calls check for extension support).

#rnx

diff --git a/Engine/Source/Runtime/VulkanRHI/Private/Android/VulkanAndroidPlatform.h b/Engine/Source/Runtime/VulkanRHI/Private/Android/VulkanAndroidPlatform.h
--- a/Engine/Source/Runtime/VulkanRHI/Private/Android/VulkanAndroidPlatform.h
+++ b/Engine/Source/Runtime/VulkanRHI/Private/Android/VulkanAndroidPlatform.h
@@ -44,11 +44,6 @@
 #define ENUM_VK_ENTRYPOINTS_OPTIONAL_PLATFORM_INSTANCE(EnumMacro) \
 	EnumMacro(PFN_vkGetRefreshCycleDurationGOOGLE, vkGetRefreshCycleDurationGOOGLE) \
 	EnumMacro(PFN_vkGetPastPresentationTimingGOOGLE, vkGetPastPresentationTimingGOOGLE) \
-	EnumMacro(PFN_vkGetPhysicalDeviceProperties2KHR, vkGetPhysicalDeviceProperties2KHR) \
-	EnumMacro(PFN_vkGetPhysicalDeviceFeatures2KHR, vkGetPhysicalDeviceFeatures2KHR) \
-	EnumMacro(PFN_vkGetPhysicalDeviceMemoryProperties2, vkGetPhysicalDeviceMemoryProperties2) \
-	EnumMacro(PFN_vkCreateRenderPass2KHR, vkCreateRenderPass2KHR) \
-	EnumMacro(PFN_vkCmdBeginRenderPass2KHR, vkCmdBeginRenderPass2KHR) \
 	EnumMacro(PFN_vkGetPhysicalDeviceFragmentShadingRatesKHR, vkGetPhysicalDeviceFragmentShadingRatesKHR)
 
 // and now, include the GenericPlatform class
diff --git a/Engine/Source/Runtime/VulkanRHI/Private/Linux/VulkanLinuxPlatform.h b/Engine/Source/Runtime/VulkanRHI/Private/Linux/VulkanLinuxPlatform.h
--- a/Engine/Source/Runtime/VulkanRHI/Private/Linux/VulkanLinuxPlatform.h
+++ b/Engine/Source/Runtime/VulkanRHI/Private/Linux/VulkanLinuxPlatform.h
@@ -36,13 +36,6 @@
     EnumMacro(PFN_vkCmdWriteBufferMarkerAMD, vkCmdWriteBufferMarkerAMD) \
     EnumMacro(PFN_vkCmdSetCheckpointNV, vkCmdSetCheckpointNV) \
     EnumMacro(PFN_vkGetQueueCheckpointDataNV, vkGetQueueCheckpointDataNV) \
-    EnumMacro(PFN_vkGetPhysicalDeviceProperties2KHR, vkGetPhysicalDeviceProperties2KHR) \
-    EnumMacro(PFN_vkGetPhysicalDeviceFeatures2KHR, vkGetPhysicalDeviceFeatures2KHR) \
-    EnumMacro(PFN_vkGetImageMemoryRequirements2KHR , vkGetImageMemoryRequirements2KHR) \
-    EnumMacro(PFN_vkGetBufferMemoryRequirements2KHR , vkGetBufferMemoryRequirements2KHR) \
-    EnumMacro(PFN_vkGetPhysicalDeviceMemoryProperties2, vkGetPhysicalDeviceMemoryProperties2) \
-	EnumMacro(PFN_vkCreateRenderPass2KHR, vkCreateRenderPass2KHR) \
-	EnumMacro(PFN_vkCmdBeginRenderPass2KHR, vkCmdBeginRenderPass2KHR) \
 	EnumMacro(PFN_vkGetPhysicalDeviceFragmentShadingRatesKHR, vkGetPhysicalDeviceFragmentShadingRatesKHR)
 
 // and now, include the GenericPlatform class
diff --git a/Engine/Source/Runtime/VulkanRHI/Private/VulkanLoader.h b/Engine/Source/Runtime/VulkanRHI/Private/VulkanLoader.h
--- a/Engine/Source/Runtime/VulkanRHI/Private/VulkanLoader.h
+++ b/Engine/Source/Runtime/VulkanRHI/Private/VulkanLoader.h
@@ -185,6 +185,13 @@
 // List all optional Vulkan entry points used by Unreal that need to be loaded manually
 #define ENUM_VK_ENTRYPOINTS_OPTIONAL_INSTANCE(EnumMacro) \
 	ENUM_VK_ENTRYPOINTS_RAYTRACING(EnumMacro) \
+	EnumMacro(PFN_vkGetPhysicalDeviceProperties2KHR, vkGetPhysicalDeviceProperties2KHR) \
+	EnumMacro(PFN_vkGetPhysicalDeviceFeatures2KHR, vkGetPhysicalDeviceFeatures2KHR) \
+	EnumMacro(PFN_vkGetImageMemoryRequirements2KHR , vkGetImageMemoryRequirements2KHR) \
+	EnumMacro(PFN_vkGetBufferMemoryRequirements2KHR , vkGetBufferMemoryRequirements2KHR) \
+	EnumMacro(PFN_vkGetPhysicalDeviceMemoryProperties2, vkGetPhysicalDeviceMemoryProperties2) \
+	EnumMacro(PFN_vkCreateRenderPass2KHR, vkCreateRenderPass2KHR) \
+	EnumMacro(PFN_vkCmdBeginRenderPass2KHR, vkCmdBeginRenderPass2KHR) \
 	EnumMacro(PFN_vkGetDeviceImageMemoryRequirementsKHR, vkGetDeviceImageMemoryRequirementsKHR) \
 	EnumMacro(PFN_vkResetQueryPoolEXT, vkResetQueryPoolEXT) \
 	EnumMacro(PFN_vkCmdPipelineBarrier2KHR, vkCmdPipelineBarrier2KHR) \
diff --git a/Engine/Source/Runtime/VulkanRHI/Private/Windows/VulkanWindowsPlatform.h b/Engine/Source/Runtime/VulkanRHI/Private/Windows/VulkanWindowsPlatform.h
--- a/Engine/Source/Runtime/VulkanRHI/Private/Windows/VulkanWindowsPlatform.h
+++ b/Engine/Source/Runtime/VulkanRHI/Private/Windows/VulkanWindowsPlatform.h
@@ -39,16 +39,9 @@
 
 #define ENUM_VK_ENTRYPOINTS_PLATFORM_INSTANCE(EnumMacro)	\
 	EnumMacro(PFN_vkCreateWin32SurfaceKHR, vkCreateWin32SurfaceKHR) \
-	EnumMacro(PFN_vkGetPhysicalDeviceProperties2KHR, vkGetPhysicalDeviceProperties2KHR) \
-	EnumMacro(PFN_vkGetPhysicalDeviceFeatures2KHR, vkGetPhysicalDeviceFeatures2KHR) \
-	EnumMacro(PFN_vkGetImageMemoryRequirements2KHR , vkGetImageMemoryRequirements2KHR) \
 	EnumMacro(PFN_vkCmdWriteBufferMarkerAMD, vkCmdWriteBufferMarkerAMD) \
 	EnumMacro(PFN_vkCmdSetCheckpointNV, vkCmdSetCheckpointNV) \
 	EnumMacro(PFN_vkGetQueueCheckpointDataNV, vkGetQueueCheckpointDataNV) \
-	EnumMacro(PFN_vkGetBufferMemoryRequirements2KHR , vkGetBufferMemoryRequirements2KHR) \
-	EnumMacro(PFN_vkGetPhysicalDeviceMemoryProperties2, vkGetPhysicalDeviceMemoryProperties2) \
-	EnumMacro(PFN_vkCreateRenderPass2KHR, vkCreateRenderPass2KHR) \
-	EnumMacro(PFN_vkCmdBeginRenderPass2KHR, vkCmdBeginRenderPass2KHR) \
 	EnumMacro(PFN_vkGetPhysicalDeviceFragmentShadingRatesKHR, vkGetPhysicalDeviceFragmentShadingRatesKHR)
 
 #define ENUM_VK_ENTRYPOINTS_OPTIONAL_PLATFORM_INSTANCE(EnumMacro)
