VulkanRHI:  VULKAN_SUPPORTS_RENDERPASS2 is a compile-time check, also make sure VK_KHR_create_renderpass2 is supported by driver at runtime before using the structures.

#rnx

diff --git a/Engine/Source/Runtime/VulkanRHI/Private/VulkanRenderpass.cpp b/Engine/Source/Runtime/VulkanRHI/Private/VulkanRenderpass.cpp
--- a/Engine/Source/Runtime/VulkanRHI/Private/VulkanRenderpass.cpp
+++ b/Engine/Source/Runtime/VulkanRHI/Private/VulkanRenderpass.cpp
@@ -234,28 +234,31 @@ public:
 		*/
 		const uint32_t CorrelationMask = MultiviewMask;
 
-#if VULKAN_SUPPORTS_RENDERPASS2
-		if (RTLayout.GetIsMultiView())
-		{
-			CreateInfo.SetCorrelationMask(&CorrelationMask);
-		}
-#else
 		VkRenderPassMultiviewCreateInfo MultiviewInfo;
 		if (RTLayout.GetIsMultiView())
 		{
-			FMemory::Memzero(MultiviewInfo);
-			MultiviewInfo.sType = VK_STRUCTURE_TYPE_RENDER_PASS_MULTIVIEW_CREATE_INFO;
-			MultiviewInfo.pNext = nullptr;
-			MultiviewInfo.subpassCount = NumSubpasses;
-			MultiviewInfo.pViewMasks = ViewMask;
-			MultiviewInfo.dependencyCount = 0;
-			MultiviewInfo.pViewOffsets = nullptr;
-			MultiviewInfo.correlationMaskCount = 1;
-			MultiviewInfo.pCorrelationMasks = &CorrelationMask;
-
-			CreateInfo.pNext = &MultiviewInfo;
-		}
+#if VULKAN_SUPPORTS_RENDERPASS2
+			if (Device.GetOptionalExtensions().HasKHRRenderPass2)
+			{
+				CreateInfo.SetCorrelationMask(&CorrelationMask);
+			}
+			else
 #endif
+			{
+				checkf(Device.GetOptionalExtensions().HasKHRMultiview, TEXT("Layout is multiview but extension is not supported!"))
+				ZeroVulkanStruct(MultiviewInfo, VK_STRUCTURE_TYPE_RENDER_PASS_MULTIVIEW_CREATE_INFO);
+				MultiviewInfo.subpassCount = NumSubpasses;
+				MultiviewInfo.pViewMasks = ViewMask;
+				MultiviewInfo.dependencyCount = 0;
+				MultiviewInfo.pViewOffsets = nullptr;
+				MultiviewInfo.correlationMaskCount = 1;
+				MultiviewInfo.pCorrelationMasks = &CorrelationMask;
+
+				MultiviewInfo.pNext = CreateInfo.pNext;
+				CreateInfo.pNext = &MultiviewInfo;
+			}
+		}
+
 
 		VkRenderPassFragmentDensityMapCreateInfoEXT FragDensityCreateInfo;
 		if (Device.GetOptionalExtensions().HasEXTFragmentDensityMap && RTLayout.GetHasFragmentDensityAttachment())
