Needed to use EnumHasAllFlags instead of EnumHasAnyFlags when testing BUF_RayTracingScratch, because it's a compound flag. This fixes an assert when booting the QAGame editor.

diff --git a/Engine/Source/Runtime/VulkanRHI/Private/VulkanMemory.cpp b/Engine/Source/Runtime/VulkanRHI/Private/VulkanMemory.cpp
--- a/Engine/Source/Runtime/VulkanRHI/Private/VulkanMemory.cpp
+++ b/Engine/Source/Runtime/VulkanRHI/Private/VulkanMemory.cpp
@@ -2670,7 +2670,7 @@ namespace VulkanRHI
 		uint32 Alignment = CalculateBufferAlignmentFromVKUsageFlags(InDevice, VulkanBufferUsage);
 
 #if VULKAN_RHI_RAYTRACING
-		if (EnumHasAnyFlags(InUEUsage, BUF_RayTracingScratch))
+		if (EnumHasAllFlags(InUEUsage, BUF_RayTracingScratch))
 		{
 			Alignment = GRHIRayTracingScratchBufferAlignment;
 		}
