Add ray tracing compaction handler for VulkanRT
* Based on D3D12 version
* Add FVulkanRayTracingCompactedSizeQueryPool to query BLAS compacted size
* Synchronization handled with a simple fence.
* Add missing barriers around build commands

diff --git a/Engine/Source/Runtime/VulkanRHI/Private/VulkanCommandWrappers.h b/Engine/Source/Runtime/VulkanRHI/Private/VulkanCommandWrappers.h
--- a/Engine/Source/Runtime/VulkanRHI/Private/VulkanCommandWrappers.h
+++ b/Engine/Source/Runtime/VulkanRHI/Private/VulkanCommandWrappers.h
@@ -200,6 +200,8 @@ struct FWrapLayer
 	static void CreateRayTracingPipelinesKHR(VkResult Result, VkDevice Device, VkDeferredOperationKHR DeferredOperation, VkPipelineCache PipelineCache, uint32_t CreateInfoCount, const VkRayTracingPipelineCreateInfoKHR* CreateInfos, const VkAllocationCallbacks* Allocator, VkPipeline* Pipelines);
 	static void GetRayTracingShaderGroupHandlesKHR(VkResult Result, VkDevice Device, VkPipeline Pipeline, uint32_t FirstGroup, uint32_t GroupCount, size_t DataSize, void* Data);
 	static void GetBufferDeviceAddressKHR(VkResult Result, VkDevice Device, const VkBufferDeviceAddressInfo* Info);
+	static void CmdWriteAccelerationStructuresPropertiesKHR(VkResult Result, VkCommandBuffer CommandBuffer, uint32_t AccelerationStructureCount, const VkAccelerationStructureKHR* AccelerationStructures, VkQueryType QueryType, VkQueryPool QueryPool, uint32_t FirstQuery);
+	static void CmdCopyAccelerationStructureKHR(VkResult Result, VkCommandBuffer CommandBuffer, const VkCopyAccelerationStructureInfoKHR* Info);
 #endif
 	static void GetDeviceImageMemoryRequirementsKHR(VkResult Result, VkDevice Device, const VkDeviceImageMemoryRequirements* Info, VkMemoryRequirements2* MemoryRequirements) VULKAN_LAYER_BODY
 	static void ResetQueryPoolEXT(VkResult Result, VkDevice Device, VkQueryPool QueryPool, uint32_t FirstQuery, uint32_t QueryCount) VULKAN_LAYER_BODY
@@ -1468,6 +1470,20 @@ namespace VulkanRHI
 		FWrapLayer::GetBufferDeviceAddressKHR(VK_SUCCESS, Device, Info);
 		return Result;
 	}
+
+	static FORCEINLINE_DEBUGGABLE void vkCmdWriteAccelerationStructuresPropertiesKHR(VkCommandBuffer CommandBuffer,	uint32_t AccelerationStructureCount, const VkAccelerationStructureKHR* AccelerationStructures, VkQueryType QueryType, VkQueryPool QueryPool, uint32_t FirstQuery)
+	{
+		FWrapLayer::CmdWriteAccelerationStructuresPropertiesKHR(VK_RESULT_MAX_ENUM, CommandBuffer, AccelerationStructureCount, AccelerationStructures, QueryType, QueryPool, FirstQuery);
+		VULKANAPINAMESPACE::vkCmdWriteAccelerationStructuresPropertiesKHR(CommandBuffer, AccelerationStructureCount, AccelerationStructures, QueryType, QueryPool, FirstQuery);
+		FWrapLayer::CmdWriteAccelerationStructuresPropertiesKHR(VK_SUCCESS, CommandBuffer, AccelerationStructureCount, AccelerationStructures, QueryType, QueryPool, FirstQuery);
+	}
+
+	static FORCEINLINE_DEBUGGABLE void vkCmdCopyAccelerationStructureKHR(VkCommandBuffer CommandBuffer, const VkCopyAccelerationStructureInfoKHR* Info)
+	{
+		FWrapLayer::CmdCopyAccelerationStructureKHR(VK_RESULT_MAX_ENUM, CommandBuffer, Info);
+		VULKANAPINAMESPACE::vkCmdCopyAccelerationStructureKHR(CommandBuffer, Info);
+		FWrapLayer::CmdCopyAccelerationStructureKHR(VK_SUCCESS, CommandBuffer, CommandBuffer, Info);
+	}
 #endif
 
 	static FORCEINLINE_DEBUGGABLE void vkGetDeviceImageMemoryRequirementsKHR(VkDevice Device, const VkDeviceImageMemoryRequirements* Info, VkMemoryRequirements2* MemoryRequirements)
diff --git a/Engine/Source/Runtime/VulkanRHI/Private/VulkanDebug.cpp b/Engine/Source/Runtime/VulkanRHI/Private/VulkanDebug.cpp
--- a/Engine/Source/Runtime/VulkanRHI/Private/VulkanDebug.cpp
+++ b/Engine/Source/Runtime/VulkanRHI/Private/VulkanDebug.cpp
@@ -4411,6 +4411,38 @@ void FWrapLayer::GetBufferDeviceAddressKHR(VkResult Result, VkDevice Device, con
 #endif
 	}
 }
+
+void FWrapLayer::CmdWriteAccelerationStructuresPropertiesKHR(VkResult Result, VkCommandBuffer CommandBuffer, uint32_t AccelerationStructureCount, const VkAccelerationStructureKHR* AccelerationStructures, VkQueryType QueryType, VkQueryPool QueryPool, uint32_t FirstQuery)
+{
+	if (Result == VK_RESULT_MAX_ENUM)
+	{
+#if VULKAN_ENABLE_DUMP_LAYER
+		PrintfBeginResult(FString::Printf(TEXT("vkCmdWriteAccelerationStructuresPropertiesKHR")));
+#endif
+	}
+	else
+	{
+#if VULKAN_ENABLE_DUMP_LAYER
+		PrintResult(Result);
+#endif
+	}
+}
+
+void FWrapLayer::CmdCopyAccelerationStructureKHR(VkResult Result, VkCommandBuffer CommandBuffer, const VkCopyAccelerationStructureInfoKHR* Info)
+{
+	if (Result == VK_RESULT_MAX_ENUM)
+	{
+#if VULKAN_ENABLE_DUMP_LAYER
+		PrintfBeginResult(FString::Printf(TEXT("vkCmdCopyAccelerationStructureKHR")));
+#endif
+	}
+	else
+	{
+#if VULKAN_ENABLE_DUMP_LAYER
+		PrintResult(Result);
+#endif
+	}
+}
 #endif // VULKAN_RHI_RAYTRACING
 
 void FWrapLayer::GetDeviceImageMemoryRequirementsKHR(VkResult Result, VkDevice Device, const VkDeviceImageMemoryRequirements* Info, VkMemoryRequirements2* MemoryRequirements)
diff --git a/Engine/Source/Runtime/VulkanRHI/Private/VulkanDevice.cpp b/Engine/Source/Runtime/VulkanRHI/Private/VulkanDevice.cpp
--- a/Engine/Source/Runtime/VulkanRHI/Private/VulkanDevice.cpp
+++ b/Engine/Source/Runtime/VulkanRHI/Private/VulkanDevice.cpp
@@ -15,6 +15,7 @@
 #include "VulkanLLM.h"
 #include "VulkanTransientResourceAllocator.h"
 #include "VulkanExtensions.h"
+#include "VulkanRayTracing.h"
 
 TAutoConsoleVariable<int32> GRHIAllowAsyncComputeCvar(
 	TEXT("r.Vulkan.AllowAsyncCompute"),
@@ -1126,6 +1127,11 @@ void FVulkanDevice::InitGPU(int32 DeviceIndex)
 		DefaultTexture = new FVulkanTexture(*this, Desc, nullptr);
 	}
 
+#if VULKAN_RHI_RAYTRACING
+	check(RayTracingCompactionRequestHandler == nullptr);
+	RayTracingCompactionRequestHandler = new FVulkanRayTracingCompactionRequestHandler(this);
+#endif
+
 	FVulkanPlatform::SetupImageMemoryRequirementWorkaround(*this);
 }
 
diff --git a/Engine/Source/Runtime/VulkanRHI/Private/VulkanDevice.h b/Engine/Source/Runtime/VulkanRHI/Private/VulkanDevice.h
--- a/Engine/Source/Runtime/VulkanRHI/Private/VulkanDevice.h
+++ b/Engine/Source/Runtime/VulkanRHI/Private/VulkanDevice.h
@@ -22,6 +22,7 @@ class FOLDVulkanQueryPool;
 
 #if VULKAN_RHI_RAYTRACING
 class FVulkanBasicRaytracingPipeline;
+class FVulkanRayTracingCompactionRequestHandler;
 #endif
 
 #define VULKAN_USE_DEBUG_NAMES 1
@@ -287,6 +288,8 @@ public:
 		check(OptionalDeviceExtensions.HasRaytracingExtensions() || (Device == VK_NULL_HANDLE));
 		return RayTracingProperties;
 	}
+	
+	FVulkanRayTracingCompactionRequestHandler* GetRayTracingCompactionRequestHandler() { return RayTracingCompactionRequestHandler; }
 
 	void InitializeRayTracing();
 	void CleanUpRayTracing();
@@ -530,6 +533,7 @@ private:
 #if VULKAN_RHI_RAYTRACING
 	FRayTracingProperties RayTracingProperties;
 	FVulkanBasicRaytracingPipeline* BasicRayTracingPipeline = nullptr;
+	FVulkanRayTracingCompactionRequestHandler* RayTracingCompactionRequestHandler = nullptr;
 #endif // VULKAN_RHI_RAYTRACING
 
 	VkPhysicalDeviceFeatures PhysicalFeatures;
diff --git a/Engine/Source/Runtime/VulkanRHI/Private/VulkanLoader.h b/Engine/Source/Runtime/VulkanRHI/Private/VulkanLoader.h
--- a/Engine/Source/Runtime/VulkanRHI/Private/VulkanLoader.h
+++ b/Engine/Source/Runtime/VulkanRHI/Private/VulkanLoader.h
@@ -180,7 +180,9 @@
 	EnumMacro(PFN_vkCmdTraceRaysKHR, vkCmdTraceRaysKHR) \
 	EnumMacro(PFN_vkCreateRayTracingPipelinesKHR, vkCreateRayTracingPipelinesKHR) \
 	EnumMacro(PFN_vkGetRayTracingShaderGroupHandlesKHR, vkGetRayTracingShaderGroupHandlesKHR) \
-	EnumMacro(PFN_vkGetBufferDeviceAddressKHR, vkGetBufferDeviceAddressKHR)
+	EnumMacro(PFN_vkGetBufferDeviceAddressKHR, vkGetBufferDeviceAddressKHR) \
+	EnumMacro(PFN_vkCmdWriteAccelerationStructuresPropertiesKHR, vkCmdWriteAccelerationStructuresPropertiesKHR) \
+	EnumMacro(PFN_vkCmdCopyAccelerationStructureKHR, vkCmdCopyAccelerationStructureKHR)
 
 // List all optional Vulkan entry points used by Unreal that need to be loaded manually
 #define ENUM_VK_ENTRYPOINTS_OPTIONAL_INSTANCE(EnumMacro) \
diff --git a/Engine/Source/Runtime/VulkanRHI/Private/VulkanRHI.cpp b/Engine/Source/Runtime/VulkanRHI/Private/VulkanRHI.cpp
--- a/Engine/Source/Runtime/VulkanRHI/Private/VulkanRHI.cpp
+++ b/Engine/Source/Runtime/VulkanRHI/Private/VulkanRHI.cpp
@@ -25,6 +25,7 @@
 #include "VulkanRenderpass.h"
 #include "VulkanTransientResourceAllocator.h"
 #include "VulkanExtensions.h"
+#include "VulkanRayTracing.h"
 
 static_assert(sizeof(VkStructureType) == sizeof(int32), "ZeroVulkanStruct() assumes VkStructureType is int32!");
 
@@ -979,6 +980,10 @@ void FVulkanCommandListContext::RHIBeginFrame()
 	++GVulkanRHIDeletionFrameNumber;
 
 	GpuProfiler.BeginFrame();
+
+#if VULKAN_RHI_RAYTRACING
+	Device->GetRayTracingCompactionRequestHandler()->Update(*this);
+#endif // D3D12_RHI_RAYTRACING
 }
 
 
diff --git a/Engine/Source/Runtime/VulkanRHI/Private/VulkanRayTracing.cpp b/Engine/Source/Runtime/VulkanRHI/Private/VulkanRayTracing.cpp
--- a/Engine/Source/Runtime/VulkanRHI/Private/VulkanRayTracing.cpp
+++ b/Engine/Source/Runtime/VulkanRHI/Private/VulkanRayTracing.cpp
@@ -12,6 +12,21 @@
 #include "Misc/CommandLine.h"
 #include "Misc/Parse.h"
 
+static int32 GVulkanRayTracingAllowCompaction = 1;
+static FAutoConsoleVariableRef CVarVulkanRayTracingAllowCompaction(
+	TEXT("r.Vulkan.RayTracing.AllowCompaction"),
+	GVulkanRayTracingAllowCompaction,
+	TEXT("Whether to automatically perform compaction for static acceleration structures to save GPU memory. (default = 1)\n"),
+	ECVF_ReadOnly
+);
+
+static int32 GVulkanRayTracingMaxBatchedCompaction = 64;
+static FAutoConsoleVariableRef CVarVulkanRayTracingMaxBatchedCompaction(
+	TEXT("r.Vulkan.RayTracing.MaxBatchedCompaction"),
+	GVulkanRayTracingMaxBatchedCompaction,
+	TEXT("Maximum of amount of compaction requests and rebuilds per frame. (default = 64)\n"),
+	ECVF_ReadOnly
+);
 
 #if PLATFORM_WINDOWS
 #pragma warning(push)
@@ -104,15 +119,55 @@ void FVulkanRayTracingAllocator::Free(FVkRtAllocation& Allocation)
 	}
 }
 
+static void AddAccelerationStructureBuildBarrier(VkCommandBuffer CommandBuffer)
+{
+	VkMemoryBarrier Barrier;
+	ZeroVulkanStruct(Barrier, VK_STRUCTURE_TYPE_MEMORY_BARRIER);
+	Barrier.srcAccessMask = VK_ACCESS_ACCELERATION_STRUCTURE_WRITE_BIT_KHR;
+	Barrier.dstAccessMask = VK_ACCESS_ACCELERATION_STRUCTURE_READ_BIT_KHR | VK_ACCESS_ACCELERATION_STRUCTURE_WRITE_BIT_KHR;
+	VulkanRHI::vkCmdPipelineBarrier(CommandBuffer, VK_PIPELINE_STAGE_ACCELERATION_STRUCTURE_BUILD_BIT_KHR, VK_PIPELINE_STAGE_ACCELERATION_STRUCTURE_BUILD_BIT_KHR, 0, 1, &Barrier, 0, nullptr, 0, nullptr);
+}
+
+static bool ShouldCompactAfterBuild(ERayTracingAccelerationStructureFlags BuildFlags)
+{
+	return EnumHasAllFlags(BuildFlags, ERayTracingAccelerationStructureFlags::AllowCompaction | ERayTracingAccelerationStructureFlags::FastTrace)
+		&& !EnumHasAnyFlags(BuildFlags, ERayTracingAccelerationStructureFlags::AllowUpdate);
+}
+
+static ERayTracingAccelerationStructureFlags GetRayTracingAccelerationStructureBuildFlags(const FRayTracingGeometryInitializer& Initializer)
+{
+	ERayTracingAccelerationStructureFlags BuildFlags = ERayTracingAccelerationStructureFlags::None;
+
+	if (Initializer.bFastBuild)
+	{
+		BuildFlags = ERayTracingAccelerationStructureFlags::FastBuild;
+	}
+	else
+	{
+		BuildFlags = ERayTracingAccelerationStructureFlags::FastTrace;
+	}
+
+	if (Initializer.bAllowUpdate)
+	{
+		EnumAddFlags(BuildFlags, ERayTracingAccelerationStructureFlags::AllowUpdate);
+	}
+
+	if (!Initializer.bFastBuild && !Initializer.bAllowUpdate && Initializer.bAllowCompaction && GVulkanRayTracingAllowCompaction)
+	{
+		EnumAddFlags(BuildFlags, ERayTracingAccelerationStructureFlags::AllowCompaction);
+	}
+
+	return BuildFlags;
+}
+
 static void GetBLASBuildData(
 	const VkDevice Device,
 	const TArrayView<const FRayTracingGeometrySegment> Segments,
 	const ERayTracingGeometryType GeometryType,
 	const FBufferRHIRef IndexBufferRHI,
-	const uint32 IndexBufferOffset,
-	const bool bFastBuild,
-	const bool bAllowUpdate,
+	const uint32 IndexBufferOffset,	
 	const uint32 IndexStrideInBytes,
+	ERayTracingAccelerationStructureFlags BuildFlags,
 	const EAccelerationStructureBuildMode BuildMode,
 	FVkRtBLASBuildData& BuildData)
 {
@@ -221,11 +276,17 @@ static void GetBLASBuildData(
 	}
 
 	BuildData.GeometryInfo.type = VK_ACCELERATION_STRUCTURE_TYPE_BOTTOM_LEVEL_KHR;
-	BuildData.GeometryInfo.flags = (bFastBuild) ? VK_BUILD_ACCELERATION_STRUCTURE_PREFER_FAST_BUILD_BIT_KHR : VK_BUILD_ACCELERATION_STRUCTURE_PREFER_FAST_TRACE_BIT_KHR;
-	if (bAllowUpdate)
+	BuildData.GeometryInfo.flags = (EnumHasAnyFlags(BuildFlags, ERayTracingAccelerationStructureFlags::FastBuild))
+		? VK_BUILD_ACCELERATION_STRUCTURE_PREFER_FAST_BUILD_BIT_KHR 
+		: VK_BUILD_ACCELERATION_STRUCTURE_PREFER_FAST_TRACE_BIT_KHR;
+	if (EnumHasAnyFlags(BuildFlags, ERayTracingAccelerationStructureFlags::AllowUpdate))
 	{
 		BuildData.GeometryInfo.flags |= VK_BUILD_ACCELERATION_STRUCTURE_ALLOW_UPDATE_BIT_KHR;
 	}
+	if (EnumHasAnyFlags(BuildFlags, ERayTracingAccelerationStructureFlags::AllowCompaction))
+	{
+		BuildData.GeometryInfo.flags |= VK_BUILD_ACCELERATION_STRUCTURE_ALLOW_COMPACTION_BIT_KHR;
+	}
 	BuildData.GeometryInfo.mode = (BuildMode == EAccelerationStructureBuildMode::Build) ? VK_BUILD_ACCELERATION_STRUCTURE_MODE_BUILD_KHR : VK_BUILD_ACCELERATION_STRUCTURE_MODE_UPDATE_KHR;
 	BuildData.GeometryInfo.geometryCount = BuildData.Segments.Num();
 	BuildData.GeometryInfo.pGeometries = BuildData.Segments.GetData();
@@ -291,6 +352,7 @@ FVulkanRayTracingGeometry::~FVulkanRayTracingGeometry()
 	{
 		Device->GetDeferredDeletionQueue().EnqueueResource(VulkanRHI::FDeferredDeletionQueue2::EType::AccelerationStructure, Handle);
 	}
+	RemoveCompactionRequest();
 }
 
 void FVulkanRayTracingGeometry::SetInitializer(const FRayTracingGeometryInitializer& InInitializer)
@@ -305,12 +367,73 @@ void FVulkanRayTracingGeometry::Swap(FVulkanRayTracingGeometry& Other)
 {
 	::Swap(Handle, Other.Handle);
 	::Swap(Address, Other.Address);
+	::Swap(AccelerationStructureCompactedSize, Other.AccelerationStructureCompactedSize);
 
 	AccelerationStructureBuffer = Other.AccelerationStructureBuffer;
 
 	// The rest of the members should be updated using SetInitializer()
 }
 
+void FVulkanRayTracingGeometry::RemoveCompactionRequest()
+{
+	if (bHasPendingCompactionRequests)
+	{
+		check(AccelerationStructureBuffer);
+		bool bRequestFound = Device->GetRayTracingCompactionRequestHandler()->ReleaseRequest(this);
+		check(bRequestFound);
+		bHasPendingCompactionRequests = false;
+	}
+}
+
+void FVulkanRayTracingGeometry::CompactAccelerationStructure(FVulkanCmdBuffer& CmdBuffer, uint64 InSizeAfterCompaction)
+{
+	check(bHasPendingCompactionRequests);
+	bHasPendingCompactionRequests = false;
+
+	ensureMsgf(InSizeAfterCompaction > 0, TEXT("Compacted acceleration structure size is expected to be non-zero. This error suggests that GPU readback synchronization is broken."));
+	if (InSizeAfterCompaction == 0)
+	{
+		return;
+	}
+
+	// Move old AS into this temporary variable which gets released when this function returns	
+	TRefCountPtr<FVulkanResourceMultiBuffer> OldAccelerationStructure = AccelerationStructureBuffer;
+	VkAccelerationStructureKHR OldHandle = Handle;
+
+	FString DebugNameString = Initializer.DebugName.ToString();
+	FRHIResourceCreateInfo BlasBufferCreateInfo(*DebugNameString);
+	AccelerationStructureBuffer = new FVulkanResourceMultiBuffer(Device, InSizeAfterCompaction, BUF_AccelerationStructure, 0, BlasBufferCreateInfo);
+	
+	VkDevice NativeDevice = Device->GetInstanceHandle();
+
+	VkAccelerationStructureCreateInfoKHR CreateInfo;
+	ZeroVulkanStruct(CreateInfo, VK_STRUCTURE_TYPE_ACCELERATION_STRUCTURE_CREATE_INFO_KHR);
+	CreateInfo.buffer = AccelerationStructureBuffer->GetHandle();
+	CreateInfo.offset = AccelerationStructureBuffer->GetOffset();
+	CreateInfo.size = InSizeAfterCompaction;
+	CreateInfo.type = VK_ACCELERATION_STRUCTURE_TYPE_BOTTOM_LEVEL_KHR;
+	VERIFYVULKANRESULT(VulkanDynamicAPI::vkCreateAccelerationStructureKHR(NativeDevice, &CreateInfo, VULKAN_CPU_ALLOCATOR, &Handle));
+
+	VkAccelerationStructureDeviceAddressInfoKHR DeviceAddressInfo;
+	ZeroVulkanStruct(DeviceAddressInfo, VK_STRUCTURE_TYPE_ACCELERATION_STRUCTURE_DEVICE_ADDRESS_INFO_KHR);
+	DeviceAddressInfo.accelerationStructure = Handle;
+	Address = VulkanDynamicAPI::vkGetAccelerationStructureDeviceAddressKHR(NativeDevice, &DeviceAddressInfo);
+
+	// Add a barrier to make sure acceleration structure are synchronized correctly for the copy command.
+	AddAccelerationStructureBuildBarrier(CmdBuffer.GetHandle());
+
+	VkCopyAccelerationStructureInfoKHR CopyInfo;
+	ZeroVulkanStruct(CopyInfo, VK_STRUCTURE_TYPE_COPY_ACCELERATION_STRUCTURE_INFO_KHR);
+	CopyInfo.src = OldHandle;
+	CopyInfo.dst = Handle;
+	CopyInfo.mode = VK_COPY_ACCELERATION_STRUCTURE_MODE_COMPACT_KHR;
+	VulkanDynamicAPI::vkCmdCopyAccelerationStructureKHR(CmdBuffer.GetHandle(), &CopyInfo);
+
+	AccelerationStructureCompactedSize = InSizeAfterCompaction;
+
+	Device->GetDeferredDeletionQueue().EnqueueResource(VulkanRHI::FDeferredDeletionQueue2::EType::AccelerationStructure, OldHandle);
+}
+
 static void GetTLASBuildData(
 	const VkDevice Device,
 	const uint32 NumInstances,
@@ -510,8 +633,15 @@ void FVulkanRayTracingScene::BuildAccelerationStructure(
 
 	FVulkanCommandBufferManager& CommandBufferManager = *CommandContext.GetCommandBufferManager();
 	FVulkanCmdBuffer* const CmdBuffer = CommandBufferManager.GetActiveCmdBuffer();
+
+	// Force a memory barrier to make sure all previous builds ops are finished before building the TLAS
+	AddAccelerationStructureBuildBarrier(CmdBuffer->GetHandle());
+
 	VulkanDynamicAPI::vkCmdBuildAccelerationStructuresKHR(CmdBuffer->GetHandle(), NumLayers, GeometryInfos.GetData(), pBuildRanges.GetData());
 
+	// Acceleration structure build barrier is used here to ensure that the acceleration structure build is complete before any rays are traced
+	AddAccelerationStructureBuildBarrier(CmdBuffer->GetHandle());
+
 	CommandBufferManager.SubmitActiveCmdBuffer();
 	CommandBufferManager.PrepareForNewActiveCommandBuffer();
 }
@@ -578,8 +708,9 @@ void FVulkanDynamicRHI::RHITransferRayTracingGeometryUnderlyingResource(FRHIRayT
 	check(DestGeometry);
 	FVulkanRayTracingGeometry* Dest = ResourceCast(DestGeometry);
 	if (!SrcGeometry)
-	{		
+	{
 		TRefCountPtr<FVulkanRayTracingGeometry> DeletionProxy = new FVulkanRayTracingGeometry(NoInit);
+		Dest->RemoveCompactionRequest();
 		Dest->Swap(*DeletionProxy);
 	}
 	else
@@ -614,9 +745,8 @@ FRayTracingAccelerationStructureSize FVulkanDynamicRHI::RHICalcRayTracingGeometr
 		Initializer.GeometryType,
 		Initializer.IndexBuffer,
 		Initializer.IndexBufferOffset,
-		Initializer.bFastBuild,
-		Initializer.bAllowUpdate,
 		IndexStrideInBytes,
+		GetRayTracingAccelerationStructureBuildFlags(Initializer),
 		EAccelerationStructureBuildMode::Build,
 		BuildData);
 
@@ -711,10 +841,9 @@ void FVulkanCommandListContext::RHIBuildAccelerationStructures(const TArrayView<
 			MakeArrayView(Geometry->Initializer.Segments),
 			Geometry->Initializer.GeometryType,
 			Geometry->Initializer.IndexBuffer,
-			Geometry->Initializer.IndexBufferOffset,
-			Geometry->Initializer.bFastBuild,
-			Geometry->Initializer.bAllowUpdate,
+			Geometry->Initializer.IndexBufferOffset,			
 			Geometry->Initializer.IndexBuffer ? Geometry->Initializer.IndexBuffer->GetStride() : 0,
+			GetRayTracingAccelerationStructureBuildFlags(Geometry->Initializer),
 			P.BuildMode,
 			BuildData);
 
@@ -741,8 +870,26 @@ void FVulkanCommandListContext::RHIBuildAccelerationStructures(const TArrayView<
 	FVulkanCmdBuffer* const CmdBuffer = CommandBufferManager->GetActiveCmdBuffer();
 	VulkanDynamicAPI::vkCmdBuildAccelerationStructuresKHR(CmdBuffer->GetHandle(), Params.Num(), BuildGeometryInfos.GetData(), BuildRangeInfos.GetData());
 
+	// Add an acceleration structure build barrier after each acceleration structure build batch.
+	// This is required because there are currently no explicit read/write barriers
+	// for acceleration structures, but we need to ensure that all commands
+	// are complete before BLAS is used again on the GPU.
+	AddAccelerationStructureBuildBarrier(CmdBuffer->GetHandle());
+
 	CommandBufferManager->SubmitActiveCmdBuffer();
 	CommandBufferManager->PrepareForNewActiveCommandBuffer();
+
+	for (const FRayTracingGeometryBuildParams& P : Params)
+	{
+		FVulkanRayTracingGeometry* const Geometry = ResourceCast(P.Geometry.GetReference());
+
+		ERayTracingAccelerationStructureFlags GeometryBuildFlags = GetRayTracingAccelerationStructureBuildFlags(Geometry->Initializer);
+		if (ShouldCompactAfterBuild(GeometryBuildFlags))
+		{
+			Device->GetRayTracingCompactionRequestHandler()->RequestCompact(Geometry);
+			Geometry->bHasPendingCompactionRequests = true;
+		}
+	}
 }
 
 void FVulkanCommandListContext::RHIBuildAccelerationStructure(const FRayTracingSceneBuildParams& SceneBuildParams)
@@ -1002,4 +1149,149 @@ FVulkanBasicRaytracingPipeline::~FVulkanBasicRaytracingPipeline()
 		Occlusion = nullptr;
 	}
 }
+
+void FVulkanRayTracingCompactedSizeQueryPool::EndBatch(FVulkanCmdBuffer* InCmdBuffer)
+{
+	check(CmdBuffer == nullptr);
+	CmdBuffer = InCmdBuffer;
+	FenceSignaledCounter = InCmdBuffer->GetFenceSignaledCounter();
+}
+
+void FVulkanRayTracingCompactedSizeQueryPool::Reset(FVulkanCmdBuffer* InCmdBuffer)
+{
+	VulkanRHI::vkCmdResetQueryPool(InCmdBuffer->GetHandle(), QueryPool, 0, MaxQueries);
+	FenceSignaledCounter = 0;
+	CmdBuffer = nullptr;
+}
+
+bool FVulkanRayTracingCompactedSizeQueryPool::TryGetResults(uint32 NumResults)
+{
+	if (CmdBuffer == nullptr) return false;
+
+	const uint64 FenceCurrentSignaledCounter = CmdBuffer->GetFenceSignaledCounter();
+	if (FenceSignaledCounter > FenceCurrentSignaledCounter)
+	{
+		return false;
+	}
+
+	VkResult Result = VulkanRHI::vkGetQueryPoolResults(Device->GetInstanceHandle(), QueryPool, 0, NumResults, NumResults * sizeof(uint64), QueryOutput.GetData(), sizeof(uint64), VK_QUERY_RESULT_WAIT_BIT);
+	if (Result == VK_SUCCESS)
+	{
+		return true;
+	}
+	return false;
+}
+
+FVulkanRayTracingCompactionRequestHandler::FVulkanRayTracingCompactionRequestHandler(FVulkanDevice* const InDevice)
+	: VulkanRHI::FDeviceChild(InDevice)
+{
+	QueryPool = new FVulkanRayTracingCompactedSizeQueryPool(InDevice, GVulkanRayTracingMaxBatchedCompaction);
+}
+
+void FVulkanRayTracingCompactionRequestHandler::RequestCompact(FVulkanRayTracingGeometry* InRTGeometry)
+{
+	check(InRTGeometry->AccelerationStructureBuffer);
+	ERayTracingAccelerationStructureFlags GeometryBuildFlags = GetRayTracingAccelerationStructureBuildFlags(InRTGeometry->Initializer);
+	check(EnumHasAllFlags(GeometryBuildFlags, ERayTracingAccelerationStructureFlags::AllowCompaction) &&
+		EnumHasAllFlags(GeometryBuildFlags, ERayTracingAccelerationStructureFlags::FastTrace) &&
+		!EnumHasAnyFlags(GeometryBuildFlags, ERayTracingAccelerationStructureFlags::AllowUpdate));
+
+	FScopeLock Lock(&CS);
+	PendingRequests.Add(InRTGeometry);
+}
+
+bool FVulkanRayTracingCompactionRequestHandler::ReleaseRequest(FVulkanRayTracingGeometry* InRTGeometry)
+{
+	FScopeLock Lock(&CS);
+
+	// Remove from pending list, not found then try active requests
+	if (PendingRequests.Remove(InRTGeometry) <= 0)
+	{
+		// If currently enqueued, then clear pointer to not handle the compaction request anymore			
+		for (int32 BLASIndex = 0; BLASIndex < ActiveBLASes.Num(); ++BLASIndex)
+		{
+			if (ActiveRequests[BLASIndex] == InRTGeometry)
+			{
+				ActiveRequests[BLASIndex] = nullptr;
+				return true;
+			}
+		}
+
+		return false;
+	}
+	else
+	{
+		return true;
+	}
+}
+
+void FVulkanRayTracingCompactionRequestHandler::Update(FVulkanCommandListContext& InCommandContext)
+{
+	LLM_SCOPE_BYNAME(TEXT("FVulkanRT/Compaction"));
+	FScopeLock Lock(&CS);
+
+	if (ActiveBLASes.Num() > 0)
+	{		
+		FVulkanCommandBufferManager& CommandBufferManager = *InCommandContext.GetCommandBufferManager();
+		FVulkanCmdBuffer* const CmdBuffer = CommandBufferManager.GetActiveCmdBuffer();
+
+		if (QueryPool->TryGetResults(ActiveBLASes.Num()))
+		{
+			// Compact
+			for (int32 BLASIndex = 0; BLASIndex < ActiveBLASes.Num(); ++BLASIndex)
+			{
+				if (ActiveRequests[BLASIndex] != nullptr)
+				{
+					ActiveRequests[BLASIndex]->CompactAccelerationStructure(*CmdBuffer, QueryPool->GetResultValue(BLASIndex));
+				}
+			}
+
+			QueryPool->Reset(CmdBuffer);
+
+			ActiveRequests.Empty(ActiveRequests.Num());
+			ActiveBLASes.Empty(ActiveBLASes.Num());
+		}
+	}
+
+	// build a new set of build requests to extract the build data	
+	for (FVulkanRayTracingGeometry* RTGeometry : PendingRequests)
+	{
+		ActiveRequests.Add(RTGeometry);
+		ActiveBLASes.Add(RTGeometry->Handle);
+
+		// enqueued enough requests for this update round
+		if (ActiveRequests.Num() >= GVulkanRayTracingMaxBatchedCompaction)
+		{
+			break;
+		}
+	}
+
+	// Do we have requests?
+	if (ActiveRequests.Num() > 0)
+	{
+		// clear out all of the pending requests, don't allow the array to shrink
+		PendingRequests.RemoveAt(0, ActiveRequests.Num(), false);
+
+		FVulkanCommandBufferManager& CommandBufferManager = *InCommandContext.GetCommandBufferManager();
+		FVulkanCmdBuffer* const CmdBuffer = CommandBufferManager.GetActiveCmdBuffer();
+
+		// Barrier here is not stricly necessary as it is added after the build.
+		// AddAccelerationStructureBuildBarrier(CmdBuffer->GetHandle());
+
+		// Write compacted size info from the selected requests
+		VulkanDynamicAPI::vkCmdWriteAccelerationStructuresPropertiesKHR(
+			CmdBuffer->GetHandle(),
+			ActiveBLASes.Num(), ActiveBLASes.GetData(),
+			VK_QUERY_TYPE_ACCELERATION_STRUCTURE_COMPACTED_SIZE_KHR,
+			QueryPool->GetHandle(),
+			0
+		);
+
+		CommandBufferManager.SubmitActiveCmdBuffer();
+		CommandBufferManager.PrepareForNewActiveCommandBuffer();
+
+		QueryPool->EndBatch(CmdBuffer);
+	}
+}
+
 #endif // VULKAN_RHI_RAYTRACING
diff --git a/Engine/Source/Runtime/VulkanRHI/Private/VulkanRayTracing.h b/Engine/Source/Runtime/VulkanRHI/Private/VulkanRayTracing.h
--- a/Engine/Source/Runtime/VulkanRHI/Private/VulkanRayTracing.h
+++ b/Engine/Source/Runtime/VulkanRHI/Private/VulkanRayTracing.h
@@ -72,11 +72,15 @@ public:
 
 	using FRHIRayTracingGeometry::Initializer;
 	using FRHIRayTracingGeometry::SizeInfo;
+	
+	void RemoveCompactionRequest();
+	void CompactAccelerationStructure(FVulkanCmdBuffer& CmdBuffer, uint64 InSizeAfterCompaction);
 
 	VkAccelerationStructureKHR Handle = VK_NULL_HANDLE;
 	VkDeviceAddress Address = 0;
 	TRefCountPtr<FVulkanResourceMultiBuffer> AccelerationStructureBuffer;
-
+	bool bHasPendingCompactionRequests = false;
+	uint64 AccelerationStructureCompactedSize = 0;
 private:
 	FVulkanDevice* const Device = nullptr;
 };
@@ -160,4 +164,48 @@ private:
 
 	FVulkanRayTracingPipelineState* Occlusion = nullptr;
 };
+
+class FVulkanRayTracingCompactedSizeQueryPool : public FVulkanQueryPool
+{
+public:
+	FVulkanRayTracingCompactedSizeQueryPool(FVulkanDevice* InDevice, uint32 InMaxQueries)
+		: FVulkanQueryPool(InDevice, nullptr, InMaxQueries, VK_QUERY_TYPE_ACCELERATION_STRUCTURE_COMPACTED_SIZE_KHR, false)
+	{}
+
+	void EndBatch(FVulkanCmdBuffer* InCmdBuffer);
+	bool TryGetResults(uint32 NumResults);
+	void Reset(FVulkanCmdBuffer* InCmdBuffer);
+
+	FVulkanCmdBuffer* CmdBuffer = nullptr;
+	uint64 FenceSignaledCounter = 0;
+};
+
+// Manages all the pending BLAS compaction requests
+class FVulkanRayTracingCompactionRequestHandler : public VulkanRHI::FDeviceChild
+{
+public:
+	UE_NONCOPYABLE(FVulkanRayTracingCompactionRequestHandler)
+
+	FVulkanRayTracingCompactionRequestHandler(FVulkanDevice* const InDevice);
+	~FVulkanRayTracingCompactionRequestHandler()
+	{
+		check(PendingRequests.IsEmpty());
+		delete QueryPool;
+	}
+
+	void RequestCompact(FVulkanRayTracingGeometry* InRTGeometry);
+	bool ReleaseRequest(FVulkanRayTracingGeometry* InRTGeometry);
+
+	void Update(FVulkanCommandListContext& InCommandContext);
+
+private:
+
+	FCriticalSection CS;
+	TArray<FVulkanRayTracingGeometry*> PendingRequests;
+	TArray<FVulkanRayTracingGeometry*> ActiveRequests;
+	TArray<VkAccelerationStructureKHR> ActiveBLASes;
+
+	FVulkanRayTracingCompactedSizeQueryPool* QueryPool = nullptr;
+};
+
 #endif // VULKAN_RHI_RAYTRACING
