VulkanRHI: Use VkImageViewUsageCreateInfo on some views created for framebuffer to specify they will only be used with COLOR_ATTACHMENT_BIT.

diff --git a/Engine/Source/Runtime/VulkanRHI/Private/VulkanResources.h b/Engine/Source/Runtime/VulkanRHI/Private/VulkanResources.h
--- a/Engine/Source/Runtime/VulkanRHI/Private/VulkanResources.h
+++ b/Engine/Source/Runtime/VulkanRHI/Private/VulkanResources.h
@@ -357,7 +357,7 @@ struct FVulkanTextureView
 	{
 	}
 
-	void Create(FVulkanDevice& Device, VkImage InImage, VkImageViewType ViewType, VkImageAspectFlags AspectFlags, EPixelFormat UEFormat, VkFormat Format, uint32 FirstMip, uint32 NumMips, uint32 ArraySliceIndex, uint32 NumArraySlices, bool bUseIdentitySwizzle = false);
+	void Create(FVulkanDevice& Device, VkImage InImage, VkImageViewType ViewType, VkImageAspectFlags AspectFlags, EPixelFormat UEFormat, VkFormat Format, uint32 FirstMip, uint32 NumMips, uint32 ArraySliceIndex, uint32 NumArraySlices, bool bUseIdentitySwizzle = false, VkImageUsageFlags ImageUsageFlags = 0);
 	void Destroy(FVulkanDevice& Device);
 
 	VkImageView View;
diff --git a/Engine/Source/Runtime/VulkanRHI/Private/VulkanTexture.cpp b/Engine/Source/Runtime/VulkanRHI/Private/VulkanTexture.cpp
--- a/Engine/Source/Runtime/VulkanRHI/Private/VulkanTexture.cpp
+++ b/Engine/Source/Runtime/VulkanRHI/Private/VulkanTexture.cpp
@@ -1429,7 +1429,7 @@ void FVulkanDynamicRHI::InternalUpdateTexture3D(bool bFromRenderingThread, FRHIT
 	}
 }
 
-void FVulkanTextureView::Create(FVulkanDevice& Device, VkImage InImage, VkImageViewType ViewType, VkImageAspectFlags AspectFlags, EPixelFormat UEFormat, VkFormat Format, uint32 FirstMip, uint32 NumMips, uint32 ArraySliceIndex, uint32 NumArraySlices, bool bUseIdentitySwizzle /*= false*/)
+void FVulkanTextureView::Create(FVulkanDevice& Device, VkImage InImage, VkImageViewType ViewType, VkImageAspectFlags AspectFlags, EPixelFormat UEFormat, VkFormat Format, uint32 FirstMip, uint32 NumMips, uint32 ArraySliceIndex, uint32 NumArraySlices, bool bUseIdentitySwizzle /*= false*/, VkImageUsageFlags ImageUsageFlags /*= 0*/)
 {
 	LLM_SCOPE_VULKAN(ELLMTagVulkan::VulkanTextures);
 	View = VK_NULL_HANDLE;
@@ -1523,6 +1523,17 @@ void FVulkanTextureView::Create(FVulkanDevice& Device, VkImage InImage, VkImageV
 		ViewInfo.subresourceRange.aspectMask = VK_IMAGE_ASPECT_STENCIL_BIT;
 	}
 
+	// Inform the driver the view will only be used with a subset of usage flags (to help performance and/or compatibility)
+	VkImageViewUsageCreateInfo ImageViewUsageCreateInfo;
+	if ((ImageUsageFlags != 0) && Device.GetOptionalExtensions().HasKHRMaintenance2)
+	{
+		ZeroVulkanStruct(ImageViewUsageCreateInfo, VK_STRUCTURE_TYPE_IMAGE_VIEW_USAGE_CREATE_INFO);
+		ImageViewUsageCreateInfo.usage = ImageUsageFlags;
+
+		ImageViewUsageCreateInfo.pNext = (void*)ViewInfo.pNext;
+		ViewInfo.pNext = &ImageViewUsageCreateInfo;
+	}
+
 	INC_DWORD_STAT(STAT_VulkanNumImageViews);
 	VERIFYVULKANRESULT(VulkanRHI::vkCreateImageView(Device.GetInstanceHandle(), &ViewInfo, VULKAN_CPU_ALLOCATOR, &View));
 
diff --git a/Engine/Source/Runtime/VulkanRHI/Private/VulkanViewport.cpp b/Engine/Source/Runtime/VulkanRHI/Private/VulkanViewport.cpp
--- a/Engine/Source/Runtime/VulkanRHI/Private/VulkanViewport.cpp
+++ b/Engine/Source/Runtime/VulkanRHI/Private/VulkanViewport.cpp
@@ -354,7 +354,7 @@ FVulkanFramebuffer::FVulkanFramebuffer(FVulkanDevice& Device, const FRHISetRende
 		}
 		else if (Texture->GetViewType() == VK_IMAGE_VIEW_TYPE_3D)
 		{
-			RTView.Create(*Texture->Device, Texture->Image, VK_IMAGE_VIEW_TYPE_2D_ARRAY, Texture->GetFullAspectMask(), Desc.Format, Texture->ViewFormat, MipIndex, 1, 0, Desc.Depth, true);
+			RTView.Create(*Texture->Device, Texture->Image, VK_IMAGE_VIEW_TYPE_2D_ARRAY, Texture->GetFullAspectMask(), Desc.Format, Texture->ViewFormat, MipIndex, 1, 0, Desc.Depth, true, VK_IMAGE_USAGE_COLOR_ATTACHMENT_BIT);
 		}
 		else
 		{
