Fix for crash in Vulkan when creating a volatile buffer.

diff --git a/Engine/Source/Runtime/VulkanRHI/Private/VulkanBuffer.cpp b/Engine/Source/Runtime/VulkanRHI/Private/VulkanBuffer.cpp
--- a/Engine/Source/Runtime/VulkanRHI/Private/VulkanBuffer.cpp
+++ b/Engine/Source/Runtime/VulkanRHI/Private/VulkanBuffer.cpp
@@ -249,7 +249,10 @@ FVulkanResourceMultiBuffer::~FVulkanResourceMultiBuffer()
 
 void* FVulkanResourceMultiBuffer::Lock(FRHICommandListBase& RHICmdList, EResourceLockMode LockMode, uint32 LockSize, uint32 Offset)
 {
-	return Lock(FVulkanCommandListContext::GetVulkanContext(RHICmdList.GetContext()), LockMode, LockSize, Offset);
+	// Use the immediate context for write operations, since we are only accessing allocators.
+	FVulkanCommandListContext& Context = FVulkanCommandListContext::GetVulkanContext(LockMode == RLM_WriteOnly ? *RHIGetDefaultContext() : RHICmdList.GetContext());
+
+	return Lock(Context, LockMode, LockSize, Offset);
 }
 
 void* FVulkanResourceMultiBuffer::Lock(FVulkanCommandListContext& Context, EResourceLockMode LockMode, uint32 LockSize, uint32 Offset)
@@ -443,8 +446,12 @@ struct FRHICommandMultiBufferUnlock final : public FRHICommand<FRHICommandMultiB
 
 void FVulkanResourceMultiBuffer::Unlock(FRHICommandListBase& RHICmdList)
 {
-	const bool bVolatile = EnumHasAnyFlags(GetUsage(), BUF_Volatile);
-	if (!bVolatile && (LockStatus != ELockStatus::PersistentMapping) && !RHICmdList.IsBottomOfPipe())
+	if (EnumHasAnyFlags(GetUsage(), BUF_Volatile) || LockStatus == ELockStatus::PersistentMapping)
+	{
+		return;
+	}
+
+	if (RHICmdList.IsTopOfPipe())
 	{
 		ALLOC_COMMAND_CL(RHICmdList, FRHICommandMultiBufferUnlock)(Device, GetPendingBufferLock(this), this, DynamicBufferIndex);
 	}
