Added FSR support to Vulkan PSO Processes

diff --git a/Engine/Build/Android/Prebuilt/PSOService/Android/CMakeLists.txt b/Engine/Build/Android/Prebuilt/PSOService/Android/CMakeLists.txt
--- a/Engine/Build/Android/Prebuilt/PSOService/Android/CMakeLists.txt
+++ b/Engine/Build/Android/Prebuilt/PSOService/Android/CMakeLists.txt
@@ -14,10 +14,10 @@ set(SOURCE_ROOT
 
 set(libpsoservice_includes
     ${SOURCE_ROOT}/src
+    ${CMAKE_SOURCE_DIR}/../../../../../Source/ThirdParty/Vulkan/Include
 )
 
 include_directories(${libpsoservice_includes}) 
-
 set(libpsoservice_source
     ${SOURCE_ROOT}/PSOService.cpp
 )
diff --git a/Engine/Build/Android/Prebuilt/PSOService/Source/PSOService.cpp b/Engine/Build/Android/Prebuilt/PSOService/Source/PSOService.cpp
--- a/Engine/Build/Android/Prebuilt/PSOService/Source/PSOService.cpp
+++ b/Engine/Build/Android/Prebuilt/PSOService/Source/PSOService.cpp
@@ -6,8 +6,9 @@
 
 #include <jni.h>
 #include <android/log.h>
-#include <vulkan/vulkan.h>
+#include "vulkan/vulkan.h"
 #include <vector>
+#include <list>
 #include <string>
 
 #define APPNAME "UEPSOService"
@@ -288,7 +289,6 @@ public:
 		}
 
 		// Create PSO
-
 		COPY_FROM_BUFFER(&PipelineCreateInfo, PSO, MemoryOffset, sizeof(GraphicsPipelineCreateInfo));
 
 		memset(&CreateInfo, 0, sizeof(VkGraphicsPipelineCreateInfo));
@@ -297,6 +297,18 @@ public:
 		CreateInfo.stageCount = PipelineCreateInfo.StageCount;
 		CreateInfo.subpass = PipelineCreateInfo.subpass;
 
+		// FSR
+		bool bHasFSRCreateInfo = false;
+		COPY_FROM_BUFFER(&bHasFSRCreateInfo, PSO, MemoryOffset, sizeof(bool));
+
+		VkPipelineFragmentShadingRateStateCreateInfoKHR FSRCreateInfo;
+		if (bHasFSRCreateInfo)
+		{
+			COPY_FROM_BUFFER(&FSRCreateInfo, PSO, MemoryOffset, sizeof(VkPipelineFragmentShadingRateStateCreateInfoKHR));
+			FSRCreateInfo.pNext = nullptr;
+			CreateInfo.pNext = &FSRCreateInfo;
+		}
+
 		VkPipelineShaderStageCreateInfo ShaderStages[2];
 
 		// VkPipelineShaderStageCreateInfo
@@ -467,6 +479,7 @@ public:
 		{
 			// Render pass
 			VkRenderPassCreateInfo2KHR RenderPassCreateInfo;
+
 			COPY_FROM_BUFFER(&RenderPassCreateInfo, PSO, MemoryOffset, sizeof(VkRenderPassCreateInfo2KHR));
 
 			// Check for VK_STRUCTURE_TYPE_RENDER_PASS_FRAGMENT_DENSITY_MAP_CREATE_INFO_EXT
@@ -492,24 +505,31 @@ public:
 			}
 
 			VkSubpassDescription2KHR* SubpassDescriptions = new VkSubpassDescription2KHR[RenderPassCreateInfo.subpassCount];
+			std::vector<VkFragmentShadingRateAttachmentInfoKHR> FSRAttachmentInfos;
 
 			for (uint32_t Idx = 0; Idx < RenderPassCreateInfo.subpassCount; ++Idx)
 			{
 				COPY_FROM_BUFFER(&SubpassDescriptions[Idx], PSO, MemoryOffset, sizeof(VkSubpassDescription2KHR));
 
-				// Check for fragment shading rate info
-				/*bool bHasNext = false;
-				COPY_FROM_BUFFER(&bHasNext, PSO, MemoryOffset, sizeof(bool));
+				// Add additional pNext structs
+				// FSR
+				bool bHasFSRAttachmentInfo = false;
+				COPY_FROM_BUFFER(&bHasFSRAttachmentInfo, PSO, MemoryOffset, sizeof(bool));;
 
-				if (bHasNext)
+				if (bHasFSRAttachmentInfo)
 				{
-					SubpassDescriptions[Idx].pNext = &PSO[MemoryOffset];
-					MemoryOffset += sizeof(VkFragmentShadingRateAttachmentInfoKHR);
-
-					VkFragmentShadingRateAttachmentInfoKHR* FragmentShadingRate = (VkFragmentShadingRateAttachmentInfoKHR*)SubpassDescriptions[Idx].pNext;
-					FragmentShadingRate->pFragmentShadingRateAttachment = (VkAttachmentReference2KHR*)&PSO[MemoryOffset];
+					FSRAttachmentInfos.push_back(VkFragmentShadingRateAttachmentInfoKHR());
+					auto& FSRAttachmentInfo = FSRAttachmentInfos.back();
+					FSRAttachmentInfo.pNext = nullptr;
+					FSRAttachmentInfo.sType = VK_STRUCTURE_TYPE_FRAGMENT_SHADING_RATE_ATTACHMENT_INFO_KHR;
+					FSRAttachmentInfo.pFragmentShadingRateAttachment = (VkAttachmentReference2KHR*)&PSO[MemoryOffset];
 					MemoryOffset += sizeof(VkAttachmentReference2KHR);
-				}*/
+
+					FSRAttachmentInfo.shadingRateAttachmentTexelSize = *(VkExtent2D*)&PSO[MemoryOffset];
+					MemoryOffset += sizeof(VkExtent2D);
+
+					SubpassDescriptions[Idx].pNext = &FSRAttachmentInfo;
+				}
 
 				if (SubpassDescriptions[Idx].colorAttachmentCount > 0)
 				{
diff --git a/Engine/Source/Runtime/VulkanRHI/Private/Android/VulkanAndroidPlatform.cpp b/Engine/Source/Runtime/VulkanRHI/Private/Android/VulkanAndroidPlatform.cpp
--- a/Engine/Source/Runtime/VulkanRHI/Private/Android/VulkanAndroidPlatform.cpp
+++ b/Engine/Source/Runtime/VulkanRHI/Private/Android/VulkanAndroidPlatform.cpp
@@ -649,6 +649,78 @@ void CharArrayToBuffer(const TArray<const ANSICHAR*>& CharArray, TArray<char>& M
 	}
 }
 
+void GetVKStructsFromPNext(const void* pNext, TMap<uint32_t, const void*>& VkStructs)
+{
+	const void* PointerNext = pNext;
+
+	while (PointerNext)
+	{
+		const VkStructureType* NextType = (VkStructureType*)PointerNext;
+		VkStructs.FindOrAdd((uint32_t)*NextType) = PointerNext;
+
+		PointerNext = *((void**)(((uint8*)NextType) + sizeof(size_t)));
+	}
+}
+
+void HandleGraphicsPipelineCreatePNext(VkGraphicsPipelineCreateInfo* PipelineCreateInfo, TArray<char>& MemoryStream)
+{
+	TMap<uint32_t, const void*> VkStructs;
+	GetVKStructsFromPNext(PipelineCreateInfo->pNext, VkStructs);
+
+	int32_t HandledCount = 0;
+
+	// FSR Create Info
+	bool bHasFSRCreateInfo = false;
+
+#if !VULKAN_SUPPORTS_FRAGMENT_SHADING_RATE
+	COPY_TO_BUFFER(MemoryStream, &bHasFSRCreateInfo, sizeof(bool));
+#else
+	const void** Struct = VkStructs.Find((uint32_t)VK_STRUCTURE_TYPE_PIPELINE_FRAGMENT_SHADING_RATE_STATE_CREATE_INFO_KHR);
+	bHasFSRCreateInfo = Struct != nullptr;
+	COPY_TO_BUFFER(MemoryStream, &bHasFSRCreateInfo, sizeof(bool));
+
+	if (bHasFSRCreateInfo)
+	{
+		VkPipelineFragmentShadingRateStateCreateInfoKHR FSRCreateInfo = *(VkPipelineFragmentShadingRateStateCreateInfoKHR*)*Struct;
+		FSRCreateInfo.pNext = nullptr;
+
+		COPY_TO_BUFFER(MemoryStream, &FSRCreateInfo, sizeof(VkPipelineFragmentShadingRateStateCreateInfoKHR));
+		HandledCount++;
+	} 
+#endif
+
+	check(HandledCount == VkStructs.Num());
+}
+
+void HandleSubpassDescriptionPNext(VkSubpassDescription2* SubpassDescription, TArray<char>& MemoryStream)
+{
+	TMap<uint32_t, const void*> VkStructs;
+	GetVKStructsFromPNext(SubpassDescription->pNext, VkStructs);
+
+	int32_t HandledCount = 0;
+
+	// FSR Create Info 
+	bool bHasFSRAttachmentCreateInfo = false;
+
+#if !VULKAN_SUPPORTS_FRAGMENT_SHADING_RATE
+	COPY_TO_BUFFER(MemoryStream, &bHasFSRAttachmentCreateInfo, sizeof(bool));
+#else
+	const void** Struct = VkStructs.Find((uint32_t)VK_STRUCTURE_TYPE_FRAGMENT_SHADING_RATE_ATTACHMENT_INFO_KHR);
+	bHasFSRAttachmentCreateInfo = Struct != nullptr;
+	COPY_TO_BUFFER(MemoryStream, &bHasFSRAttachmentCreateInfo, sizeof(bool));
+
+	if (bHasFSRAttachmentCreateInfo)
+	{
+		VkFragmentShadingRateAttachmentInfoKHR* FragmentShadingRateCreateInfo = (VkFragmentShadingRateAttachmentInfoKHR*)*Struct;
+		COPY_TO_BUFFER(MemoryStream, FragmentShadingRateCreateInfo->pFragmentShadingRateAttachment, sizeof(VkAttachmentReference2));
+		COPY_TO_BUFFER(MemoryStream, &FragmentShadingRateCreateInfo->shadingRateAttachmentTexelSize, sizeof(VkExtent2D));
+		HandledCount++;
+	}
+#endif
+
+	check(HandledCount == VkStructs.Num());
+}
+
 void PipelineToBinary(FVulkanDevice* Device, VkGraphicsPipelineCreateInfo* PipelineInfo, FGfxPipelineDesc* GfxEntry, const FVulkanRenderTargetLayout* RTLayout, TArray<char>& MemoryStream)
 {
 	static const unsigned int INITIAL_PSO_STREAM_SIZE = 64 * 1024;
@@ -679,7 +751,7 @@ void PipelineToBinary(FVulkanDevice* Device, VkGraphicsPipelineCreateInfo* Pipel
 	CharArrayToBuffer(DeviceExtensions, MemoryStream);
 	COPY_TO_BUFFER(MemoryStream, &pipelineCreateInfo, sizeof(GraphicsPipelineCreateInfo));
 
-	check(PipelineInfo->pNext == nullptr);
+	HandleGraphicsPipelineCreatePNext(PipelineInfo, MemoryStream);
 
 	// VkPipelineShaderStageCreateInfo
 	for (int32_t Idx = 0; Idx < PipelineInfo->stageCount; ++Idx)
@@ -912,22 +984,7 @@ void PipelineToBinary(FVulkanDevice* Device, VkGraphicsPipelineCreateInfo* Pipel
 			
 			COPY_TO_BUFFER(MemoryStream, &CreateInfo.pSubpasses[Idx], sizeof(VkSubpassDescription2));
 
-			check(CreateInfo.pSubpasses[Idx].pNext == nullptr);
-
-			// Check for fragment shading rate info
-			/*bool bHasNext = CreateInfo.pSubpasses[Idx].pNext != nullptr;
-			COPY_TO_BUFFER(MemoryStream, &bHasNext, MemoryOffset, sizeof(bool));
-			
-			if (bHasNext)
-			{
-				VkStructureType NextType = *(VkStructureType*)CreateInfo.pSubpasses[Idx].pNext;
-				check(NextType == VK_STRUCTURE_TYPE_FRAGMENT_SHADING_RATE_ATTACHMENT_INFO_KHR);
-
-				VkFragmentShadingRateAttachmentInfoKHR* FragmentShadingRate = (VkFragmentShadingRateAttachmentInfoKHR*)CreateInfo.pSubpasses[Idx].pNext;
-
-				COPY_TO_BUFFER(MemoryStream, FragmentShadingRate, MemoryOffset, sizeof(VkFragmentShadingRateAttachmentInfoKHR));
-				COPY_TO_BUFFER(MemoryStream, FragmentShadingRate->pFragmentShadingRateAttachment, MemoryOffset, sizeof(VkAttachmentReference2));
-			}*/
+			HandleSubpassDescriptionPNext(&SubpassDescription, MemoryStream);
 
 			if(SubpassDescription.colorAttachmentCount > 0)
 			{
