Use the VULKAN_RHI_RAYTRACING preprocessor instead of the RHI_RAYTRACING to fix an issue where the platform declares RHI_RAYTRACING flag but doesn’t use the Vulkan RHI, resulting in undeclared Vulkan enums.

diff --git a/Engine/Source/Runtime/VulkanRHI/Public/VulkanShaderResources.h b/Engine/Source/Runtime/VulkanRHI/Public/VulkanShaderResources.h
--- a/Engine/Source/Runtime/VulkanRHI/Public/VulkanShaderResources.h
+++ b/Engine/Source/Runtime/VulkanRHI/Public/VulkanShaderResources.h
@@ -24,7 +24,7 @@ static inline VkDescriptorType BindingToDescriptorType(EVulkanBindingType::EType
 	case EVulkanBindingType::StorageTexelBuffer:	return VK_DESCRIPTOR_TYPE_STORAGE_TEXEL_BUFFER;
 	case EVulkanBindingType::StorageBuffer:			return VK_DESCRIPTOR_TYPE_STORAGE_BUFFER;
 	case EVulkanBindingType::InputAttachment:		return VK_DESCRIPTOR_TYPE_INPUT_ATTACHMENT;
-#if RHI_RAYTRACING
+#ifdef VULKAN_RHI_RAYTRACING
 	case EVulkanBindingType::AccelerationStructure:	return VK_DESCRIPTOR_TYPE_ACCELERATION_STRUCTURE_KHR;
 #endif
 	default:
@@ -48,7 +48,7 @@ static inline EVulkanBindingType::EType DescriptorTypeToBinding(VkDescriptorType
 	case VK_DESCRIPTOR_TYPE_STORAGE_TEXEL_BUFFER:		return EVulkanBindingType::StorageTexelBuffer;
 	case VK_DESCRIPTOR_TYPE_STORAGE_BUFFER:				return EVulkanBindingType::StorageBuffer;
 	case VK_DESCRIPTOR_TYPE_INPUT_ATTACHMENT:			return EVulkanBindingType::InputAttachment;
-#if RHI_RAYTRACING
+#ifdef VULKAN_RHI_RAYTRACING
 	case VK_DESCRIPTOR_TYPE_ACCELERATION_STRUCTURE_KHR: return EVulkanBindingType::AccelerationStructure;
 #endif
 	default:
