VulkanRHI: Fix layout check to use first slice when using array in UAV.  Fixes asserts in Strata startup on Vulkan.

#rnx

diff --git a/Engine/Source/Runtime/VulkanRHI/Private/VulkanPendingState.cpp b/Engine/Source/Runtime/VulkanRHI/Private/VulkanPendingState.cpp
--- a/Engine/Source/Runtime/VulkanRHI/Private/VulkanPendingState.cpp
+++ b/Engine/Source/Runtime/VulkanRHI/Private/VulkanPendingState.cpp
@@ -624,8 +624,8 @@ void FVulkanPendingGfxState::SetUAVForUBResource(uint8 DescriptorSet, uint32 Bin
 		}
 		else if (UAV->SourceTexture)
 		{
-			VkImageLayout Layout = Context.GetLayoutManager().FindLayoutChecked(UAV->TextureView.Image);
-			CurrentState->SetUAVTextureView(DescriptorSet, BindingIndex, UAV->TextureView, Layout);
+			const FVulkanImageLayout& Layout = Context.GetLayoutManager().GetFullLayoutChecked(UAV->TextureView.Image);
+			CurrentState->SetUAVTextureView(DescriptorSet, BindingIndex, UAV->TextureView, Layout.GetSubresLayout(UAV->FirstArraySlice, UAV->MipLevel));
 		}
 		else
 		{
diff --git a/Engine/Source/Runtime/VulkanRHI/Private/VulkanRenderTarget.cpp b/Engine/Source/Runtime/VulkanRHI/Private/VulkanRenderTarget.cpp
--- a/Engine/Source/Runtime/VulkanRHI/Private/VulkanRenderTarget.cpp
+++ b/Engine/Source/Runtime/VulkanRHI/Private/VulkanRenderTarget.cpp
@@ -1257,7 +1257,7 @@ FVulkanRenderTargetLayout::FVulkanRenderTargetLayout(FVulkanDevice& InDevice, co
 		ensure(!NumSamples || NumSamples == ColorEntry.RenderTarget->GetNumSamples());
 		NumSamples = ColorEntry.RenderTarget->GetNumSamples();
 
-		ensure( !bMultiviewRenderTargets || Texture->GetNumberOfArrayLevels() > 1 );
+		ensure(!GetIsMultiView() || !bMultiviewRenderTargets || Texture->GetNumberOfArrayLevels() > 1);
 		bMultiviewRenderTargets = Texture->GetNumberOfArrayLevels() > 1;
 
 		VkAttachmentDescription& CurrDesc = Desc[NumAttachmentDescriptions];
