Only use android's Vulkan PSO compiler service for precompile PSO requests.
precompile jobs no longer touch vulkan PSO LRU (preventing unnecessary VK PSO creates during precompile)
Add an option to force PSO create tasks to run on the RHI thread. This fixes issues some GPU hanging issues with some android vulkan drivers.

diff --git a/Engine/Source/Runtime/RHI/Private/PipelineStateCache.cpp b/Engine/Source/Runtime/RHI/Private/PipelineStateCache.cpp
--- a/Engine/Source/Runtime/RHI/Private/PipelineStateCache.cpp
+++ b/Engine/Source/Runtime/RHI/Private/PipelineStateCache.cpp
@@ -75,6 +75,15 @@ static TAutoConsoleVariable<int32> GCVarAsyncPipelineCompile(
 	ECVF_ReadOnly | ECVF_RenderThreadSafe
 );
 
+bool GRunPSOCreateTasksOnRHIT = false;
+static FAutoConsoleVariableRef CVarCreatePSOsOnRHIThread(
+	TEXT("r.pso.CreateOnRHIThread"),
+	GRunPSOCreateTasksOnRHIT,
+	TEXT("0: Run PSO creation on task threads\n")
+	TEXT("1: Run PSO creation on RHI thread."),
+	ECVF_RenderThreadSafe
+);
+
 static TAutoConsoleVariable<int32> CVarPSOEvictionTime(
 	TEXT("r.pso.evictiontime"),
 	60,
@@ -1435,9 +1444,11 @@ public:
 
 	ENamedThreads::Type GetDesiredThread()
 	{
+		ENamedThreads::Type DesiredThread = GRunPSOCreateTasksOnRHIT && IsRunningRHIInSeparateThread() ? ENamedThreads::RHIThread : CPrio_FCompilePipelineStateTask.Get();
+
 		// On Mac the compilation is handled using external processes, so engine threads have very little work to do
 		// and it's better to leave more CPU time to these extrenal processes and other engine threads.
-		return PLATFORM_MAC ? ENamedThreads::AnyBackgroundThreadNormalTask : CPrio_FCompilePipelineStateTask.Get();
+		return PLATFORM_MAC ? ENamedThreads::AnyBackgroundThreadNormalTask : DesiredThread;
 	}
 };
 
diff --git a/Engine/Source/Runtime/VulkanRHI/Private/VulkanPipeline.cpp b/Engine/Source/Runtime/VulkanRHI/Private/VulkanPipeline.cpp
--- a/Engine/Source/Runtime/VulkanRHI/Private/VulkanPipeline.cpp
+++ b/Engine/Source/Runtime/VulkanRHI/Private/VulkanPipeline.cpp
@@ -1133,8 +1133,30 @@ static const TMap<uint8, VkFragmentShadingRateCombinerOpKHR> FragmentCombinerOpM
 	// @todo: Add "VK_FRAGMENT_SHADING_RATE_COMBINER_OP_MUL_KHR"?
 };
 #endif
-bool FVulkanPipelineStateCacheManager::CreateGfxPipelineFromEntry(FVulkanRHIGraphicsPipelineState* PSO, FVulkanShader* Shaders[ShaderStage::NumStages], VkPipeline* Pipeline, FPipelineCache& Cache)
+
+FString FVulkanPipelineStateCacheManager::ShaderHashesToString(FVulkanShader* Shaders[ShaderStage::NumStages])
 {
+	FString ShaderHashes = "";
+	if (Shaders[ShaderStage::Vertex] && Shaders[ShaderStage::Vertex]->StageFlag == VK_SHADER_STAGE_VERTEX_BIT)
+	{
+		ShaderHashes += TEXT("VS: ") + static_cast<FVulkanVertexShader*>(Shaders[ShaderStage::Vertex])->GetHash().ToString() + TEXT("\n");
+	}
+	if (Shaders[ShaderStage::Pixel] && Shaders[ShaderStage::Pixel]->StageFlag == VK_SHADER_STAGE_FRAGMENT_BIT)
+	{
+		ShaderHashes += TEXT("PS: ") + static_cast<FVulkanPixelShader*>(Shaders[ShaderStage::Pixel])->GetHash().ToString() + TEXT("\n");
+	}
+#if VULKAN_SUPPORTS_GEOMETRY_SHADERS
+	if (Shaders[ShaderStage::Geometry] && Shaders[ShaderStage::Geometry]->StageFlag == VK_SHADER_STAGE_GEOMETRY_BIT)
+	{
+		ShaderHashes += TEXT("GS: ") + static_cast<FVulkanGeometryShader*>(Shaders[ShaderStage::Geometry])->GetHash().ToString() + TEXT("\n");
+	}
+#endif
+	return ShaderHashes;
+}
+
+bool FVulkanPipelineStateCacheManager::CreateGfxPipelineFromEntry(FVulkanRHIGraphicsPipelineState* PSO, FVulkanShader* Shaders[ShaderStage::NumStages], bool bPrecompile)
+{
+	VkPipeline* Pipeline = &PSO->VulkanPipeline;
 	FGfxPipelineDesc* GfxEntry = &PSO->Desc;
 	if (Shaders[ShaderStage::Pixel] == nullptr && !FVulkanPlatform::SupportsNullPixelShader())
 	{
@@ -1291,164 +1313,155 @@ bool FVulkanPipelineStateCacheManager::CreateGfxPipelineFromEntry(FVulkanRHIGrap
 	VkResult Result = VK_ERROR_INITIALIZATION_FAILED;
 
 	double BeginTime = FPlatformTime::Seconds();
-	if(bUseLRU)
+
+	Result = CreateVKPipeline(PSO, Shaders, PipelineInfo, bPrecompile);
+
+	if (Result != VK_SUCCESS)
 	{
-#if VULKAN_USE_SHADERKEYS
-		const uint64 ShaderHash = GfxEntry->ShaderKeyShared;
-#else
-		const uint64 ShaderHash = GfxEntry->ShaderHashes.Hash;
-#endif
-		FVulkanPipelineSize* Found;
-		{
-			FScopeLock Lock(&LRUCS); 
-			Found = LRU2SizeList.Find(ShaderHash);
-		}
+		FString ShaderHashes = ShaderHashesToString(Shaders);
 
-		uint32 FoundSize = 0;
-		VkPipelineCache LocalPipelineCache = VK_NULL_HANDLE;
-		if (Found)
-		{
-			SCOPE_CYCLE_COUNTER(STAT_VulkanPSOVulkanCreationTime);
-			FoundSize = Found->PipelineSize;
-			FScopedPipelineCache PipelineCacheShared = Cache.Get(EPipelineCacheAccess::Shared);
-			Result = VulkanRHI::vkCreateGraphicsPipelines(Device->GetInstanceHandle(), PipelineCacheShared.Get(), 1, &PipelineInfo, VULKAN_CPU_ALLOCATOR, Pipeline);
-		}
-		else
-		{
-			SCOPE_CYCLE_COUNTER(STAT_VulkanPSOVulkanCreationTime);
+		UE_LOG(LogVulkanRHI, Error, TEXT("Failed to create graphics pipeline.\nShaders in pipeline: %s"), *ShaderHashes);
+		return false;
+	}
 
-			size_t AfterSize = 0;
+	double EndTime = FPlatformTime::Seconds();
+	double Delta = EndTime - BeginTime;
+	if (Delta > HitchTime)
+	{
+		UE_LOG(LogVulkanRHI, Verbose, TEXT("Hitchy gfx pipeline key CS (%.3f ms)"), (float)(Delta * 1000.0));
+	}
 
-			bool PreCompiled = false;
+	INC_DWORD_STAT(STAT_VulkanNumPSOs);
+	return true;
+}
 
-#if PLATFORM_ANDROID
-			if (FVulkanAndroidPlatform::AreRemoteCompileServicesActive() &&
-				CVarPipelineLRUCacheEvictBinary.GetValueOnAnyThread())
-			{
-				FVulkanShader::FSpirvCode VS = PSO->GetPatchedSpirvCode(Shaders[ShaderStage::Vertex]);
-				FVulkanShader::FSpirvCode PS = PSO->GetPatchedSpirvCode(Shaders[ShaderStage::Pixel]);
-				TArrayView<uint32_t> VSCode = VS.GetCodeView();
-				TArrayView<uint32_t> PSCode = PS.GetCodeView();
+VkResult FVulkanPipelineStateCacheManager::CreateVKPipeline(FVulkanRHIGraphicsPipelineState* PSO, FVulkanShader* Shaders[ShaderStage::NumStages], VkGraphicsPipelineCreateInfo PipelineInfo, bool bIsPrecompileJob)
+{
+	VkPipeline* Pipeline = &PSO->VulkanPipeline;
 
-				LocalPipelineCache = FVulkanPlatform::PrecompilePSO(Device, &PipelineInfo, GfxEntry, &PSO->RenderPass->GetLayout(), VSCode, PSCode, AfterSize);
+	FPipelineCache& Cache = bIsPrecompileJob ? CurrentPrecompilingPSOCache : GlobalPSOCache;
 
-				if (LocalPipelineCache != VK_NULL_HANDLE)
-				{
-					PreCompiled = true;
-				}
-			}
-#endif
+	VkPipelineCache LocalPipelineCache = VK_NULL_HANDLE;
+	VkResult Result = VK_ERROR_INITIALIZATION_FAILED;
+	uint32 PSOSize = 0;
+	bool bWantPSOSize = false;
+	FGfxPipelineDesc* GfxEntry = &PSO->Desc;
+	uint64 ShaderHash = 0;
+	if (bUseLRU)
+	{
+		QUICK_SCOPE_CYCLE_COUNTER(STAT_VulkanPSOLRUSizeLookup);
 
-			if (PreCompiled)
+#if VULKAN_USE_SHADERKEYS
+		ShaderHash = GfxEntry->ShaderKeyShared;
+#else
+		ShaderHash = GfxEntry->ShaderHashes.Hash;
+#endif
+		{
+			FScopeLock Lock(&LRUCS);
+			FVulkanPipelineSize* Found = LRU2SizeList.Find(ShaderHash);
+			if (Found)
 			{
-				// The PSO has been sent to a separate process, this avoids process wide mutex on PSO creation with some drivers.
-				// We have not created the PSO in this case, it is passed to a separate process and added to the pipeline cache at the end of PSO precompile.
-				// TODO errors.
-				Pipeline = nullptr;
-				Result = VK_SUCCESS;
+				PSOSize = Found->PipelineSize;
 			}
 			else
 			{
-				// We create a single pipeline cache for this create so we can observe the size for LRU cache's accounting.
-				// measuring deltas from the global PipelineCache is not thread safe.
-				VkPipelineCacheCreateInfo PipelineCacheInfo;
-				ZeroVulkanStruct(PipelineCacheInfo, VK_STRUCTURE_TYPE_PIPELINE_CACHE_CREATE_INFO);
-				VERIFYVULKANRESULT(VulkanRHI::vkCreatePipelineCache(Device->GetInstanceHandle(), &PipelineCacheInfo, VULKAN_CPU_ALLOCATOR, &LocalPipelineCache));
-				Result = VulkanRHI::vkCreateGraphicsPipelines(Device->GetInstanceHandle(), LocalPipelineCache, 1, &PipelineInfo, VULKAN_CPU_ALLOCATOR, Pipeline);
+				bWantPSOSize = true;
 			}
+		}
+	}
 
-			if (Result == VK_SUCCESS)
-			{
-				FScopeLock Lock(&LRUCS);
-				Found = LRU2SizeList.Find(ShaderHash);
-				if (Found == nullptr) // we may not win..
-				{
-					QUICK_SCOPE_CYCLE_COUNTER(STAT_Vulkan_Calc_LRU_Size);
-					size_t Diff = 0;
-
-					if (LocalPipelineCache != VK_NULL_HANDLE)
-					{
-						VulkanRHI::vkGetPipelineCacheData(Device->GetInstanceHandle(), LocalPipelineCache, &Diff, nullptr);
-					}
-
-					if (!Diff)
-					{
-						UE_LOG(LogVulkanRHI, Warning, TEXT("Shader size was computed as zero, using 20k instead."));
-						Diff = 20 * 1024;
-					}
-					FVulkanPipelineSize PipelineSize;
-					PipelineSize.ShaderHash = ShaderHash;
-					PipelineSize.PipelineSize = (uint32)Diff;
-					LRU2SizeList.Add(ShaderHash, PipelineSize);
-					FoundSize = Diff;
-
-					if (LocalPipelineCache != VK_NULL_HANDLE)
-					{
-						QUICK_SCOPE_CYCLE_COUNTER(STAT_VulkanPSOCacheMerge);
-						FScopedPipelineCache PipelineCacheExclusive = Cache.Get(EPipelineCacheAccess::Exclusive);
-						VERIFYVULKANRESULT(VulkanRHI::vkMergePipelineCaches(Device->GetInstanceHandle(), PipelineCacheExclusive.Get(), 1, &LocalPipelineCache));
-					}
-				}
-				else
-				{
-					FoundSize = Found->PipelineSize;
-				}
-			}
-			else
-			{
-				UE_LOG(LogVulkanRHI, Warning, TEXT("LRU: PSO create failed."));
-			}
+#if PLATFORM_ANDROID
+	if (bIsPrecompileJob && FVulkanAndroidPlatform::AreRemoteCompileServicesActive() /*&&
+		CVarPipelineLRUCacheEvictBinary.GetValueOnAnyThread()*/)
+	{
+		FVulkanShader::FSpirvCode VS = PSO->GetPatchedSpirvCode(Shaders[ShaderStage::Vertex]);
+		FVulkanShader::FSpirvCode PS = PSO->GetPatchedSpirvCode(Shaders[ShaderStage::Pixel]);
+		TArrayView<uint32_t> VSCode = VS.GetCodeView();
+		TArrayView<uint32_t> PSCode = PS.GetCodeView();
+		size_t AfterSize = 0;
+
+		LocalPipelineCache = FVulkanPlatform::PrecompilePSO(Device, &PipelineInfo, GfxEntry, &PSO->RenderPass->GetLayout(), VSCode, PSCode, AfterSize);
 
-			if (LocalPipelineCache != VK_NULL_HANDLE)
+		if (ensure(LocalPipelineCache != VK_NULL_HANDLE))
+		{
+			Pipeline[0] = VK_NULL_HANDLE;
+			Result = VK_SUCCESS;
+			if(bWantPSOSize)
 			{
-				VulkanRHI::vkDestroyPipelineCache(Device->GetInstanceHandle(), LocalPipelineCache, VULKAN_CPU_ALLOCATOR);
+				PSOSize = AfterSize;
 			}
 		}
-
-		if(Result == VK_SUCCESS)
+		else
 		{
-			PSO->PipelineCacheSize = FoundSize;
+			FString ShaderHashes = ShaderHashesToString(Shaders);
+			UE_LOG(LogVulkanRHI, Error, TEXT("Android RemoteCompileServices Failed to create graphics pipeline.\nShaders in pipeline: %s"), *ShaderHashes);
 		}
 	}
-	else
-	{
-		SCOPE_CYCLE_COUNTER(STAT_VulkanPSOVulkanCreationTime);
-		FScopedPipelineCache PipelineCacheShared = Cache.Get(EPipelineCacheAccess::Shared);
-		Result = VulkanRHI::vkCreateGraphicsPipelines(Device->GetInstanceHandle(), PipelineCacheShared.Get(), 1, &PipelineInfo, VULKAN_CPU_ALLOCATOR, Pipeline);
-	}
-
-	if (Result != VK_SUCCESS)
+#endif
+	if(Result != VK_SUCCESS)
 	{
-		FString ShaderHashes = "";
-		if (Shaders[ShaderStage::Vertex] && Shaders[ShaderStage::Vertex]->StageFlag == VK_SHADER_STAGE_VERTEX_BIT)
-		{
-			ShaderHashes += TEXT("VS: ") + static_cast<FVulkanVertexShader*>(Shaders[ShaderStage::Vertex])->GetHash().ToString() + TEXT("\n");
-		}
-		if (Shaders[ShaderStage::Pixel] && Shaders[ShaderStage::Pixel]->StageFlag == VK_SHADER_STAGE_FRAGMENT_BIT)
+		if (bWantPSOSize)
 		{
-			ShaderHashes += TEXT("PS: ") + static_cast<FVulkanPixelShader*>(Shaders[ShaderStage::Pixel])->GetHash().ToString() + TEXT("\n");
+			// We create a single pipeline cache for this create so we can observe the size for LRU cache's accounting.
+			// measuring deltas from the global PipelineCache is not thread safe.
+			VkPipelineCacheCreateInfo PipelineCacheInfo;
+			ZeroVulkanStruct(PipelineCacheInfo, VK_STRUCTURE_TYPE_PIPELINE_CACHE_CREATE_INFO);
+			VERIFYVULKANRESULT(VulkanRHI::vkCreatePipelineCache(Device->GetInstanceHandle(), &PipelineCacheInfo, VULKAN_CPU_ALLOCATOR, &LocalPipelineCache));
+			Result = VulkanRHI::vkCreateGraphicsPipelines(Device->GetInstanceHandle(), LocalPipelineCache, 1, &PipelineInfo, VULKAN_CPU_ALLOCATOR, Pipeline);
 		}
-#if VULKAN_SUPPORTS_GEOMETRY_SHADERS
-		if (Shaders[ShaderStage::Geometry] && Shaders[ShaderStage::Geometry]->StageFlag == VK_SHADER_STAGE_GEOMETRY_BIT)
+		else
 		{
-			ShaderHashes += TEXT("GS: ") + static_cast<FVulkanGeometryShader*>(Shaders[ShaderStage::Geometry])->GetHash().ToString() + TEXT("\n");
+			SCOPE_CYCLE_COUNTER(STAT_VulkanPSOVulkanCreationTime);
+			FScopedPipelineCache PipelineCacheShared = Cache.Get(EPipelineCacheAccess::Shared);
+			Result = VulkanRHI::vkCreateGraphicsPipelines(Device->GetInstanceHandle(), PipelineCacheShared.Get(), 1, &PipelineInfo, VULKAN_CPU_ALLOCATOR, Pipeline);
 		}
-#endif
-		UE_LOG(LogVulkanRHI, Error, TEXT("Failed to create graphics pipeline.\nShaders in pipeline: %s"), *ShaderHashes);
-		return false;
 	}
 
-	double EndTime = FPlatformTime::Seconds();
-	double Delta = EndTime - BeginTime;
-	if (Delta > HitchTime)
+	if (LocalPipelineCache != VK_NULL_HANDLE)
 	{
-		UE_LOG(LogVulkanRHI, Verbose, TEXT("Hitchy gfx pipeline key CS (%.3f ms)"), (float)(Delta * 1000.0));
+		if (bWantPSOSize)
+		{
+			FScopeLock Lock(&LRUCS);
+			FVulkanPipelineSize* Found = LRU2SizeList.Find(ShaderHash);
+			if (Found == nullptr) // Check we're not beaten to it..
+			{
+				QUICK_SCOPE_CYCLE_COUNTER(STAT_Vulkan_Calc_LRU_Size);
+				size_t Diff = 0;
+
+				if (LocalPipelineCache != VK_NULL_HANDLE)
+				{
+					VulkanRHI::vkGetPipelineCacheData(Device->GetInstanceHandle(), LocalPipelineCache, &Diff, nullptr);
+				}
+
+				if (!Diff)
+				{
+					UE_LOG(LogVulkanRHI, Warning, TEXT("Shader size was computed as zero, using 20k instead."));
+					Diff = 20 * 1024;
+				}
+				FVulkanPipelineSize PipelineSize;
+				PipelineSize.ShaderHash = ShaderHash;
+				PipelineSize.PipelineSize = (uint32)Diff;
+				LRU2SizeList.Add(ShaderHash, PipelineSize);
+				PSOSize = Diff;
+			}
+			else
+			{
+				PSOSize = Found->PipelineSize;
+			}
+		}
+
+		QUICK_SCOPE_CYCLE_COUNTER(STAT_VulkanPSOCacheMerge);
+		FScopedPipelineCache PipelineCacheExclusive = Cache.Get(EPipelineCacheAccess::Exclusive);
+		VERIFYVULKANRESULT(VulkanRHI::vkMergePipelineCaches(Device->GetInstanceHandle(), PipelineCacheExclusive.Get(), 1, &LocalPipelineCache));
+		VulkanRHI::vkDestroyPipelineCache(Device->GetInstanceHandle(), LocalPipelineCache, VULKAN_CPU_ALLOCATOR);
 	}
 
-	INC_DWORD_STAT(STAT_VulkanNumPSOs);
-	return true;
-}
+	VERIFYVULKANRESULT(Result);
 
+	PSO->PipelineCacheSize = PSOSize;
+
+ 	return Result;
+}
 
 void FVulkanPipelineStateCacheManager::DestroyCache()
 {
@@ -1888,7 +1901,10 @@ FVulkanRHIGraphicsPipelineState* FVulkanPipelineStateCacheManager::RHICreateGrap
 			if(PSO)
 			{
 				check(*PSO);
-				LRUTouch(*PSO);
+				if(!Initializer.bFromPSOFileCache)
+				{
+					LRUTouch(*PSO);
+				}
 				return *PSO;
 			}
 		}
@@ -1950,8 +1966,7 @@ FVulkanRHIGraphicsPipelineState* FVulkanPipelineStateCacheManager::RHICreateGrap
 			
 				QUICK_SCOPE_CYCLE_COUNTER(STAT_Vulkan_RHICreateGraphicsPipelineState_CREATE_PART0);
 
-				FPipelineCache& Cache = Initializer.bFromPSOFileCache ? CurrentPrecompilingPSOCache : GlobalPSOCache;
-				if(!CreateGfxPipelineFromEntry(NewPSO, VulkanShaders, &NewPSO->VulkanPipeline, Cache))
+				if(!CreateGfxPipelineFromEntry(NewPSO, VulkanShaders, Initializer.bFromPSOFileCache))
 				{
 					DeleteNewPSO(NewPSO);
 					return nullptr;
@@ -2344,7 +2359,7 @@ void FVulkanPipelineStateCacheManager::LRUTouch(FVulkanRHIGraphicsPipelineState*
 
 			QUICK_SCOPE_CYCLE_COUNTER(STAT_Vulkan_LRUMiss);
 
-			if (!CreateGfxPipelineFromEntry(PSO, VulkanShaders, &PSO->VulkanPipeline, GlobalPSOCache))
+			if (!CreateGfxPipelineFromEntry(PSO, VulkanShaders, false))
 			{
 				check(0);
 			}
diff --git a/Engine/Source/Runtime/VulkanRHI/Private/VulkanPipeline.h b/Engine/Source/Runtime/VulkanRHI/Private/VulkanPipeline.h
--- a/Engine/Source/Runtime/VulkanRHI/Private/VulkanPipeline.h
+++ b/Engine/Source/Runtime/VulkanRHI/Private/VulkanPipeline.h
@@ -490,8 +490,10 @@ private:
 	FVulkanRHIGraphicsPipelineState* RHICreateGraphicsPipelineState(const FGraphicsPipelineStateInitializer& Initializer);
 	FVulkanComputePipeline* RHICreateComputePipelineState(FRHIComputeShader* ComputeShaderRHI);
 	void NotifyDeletedGraphicsPSO(FRHIGraphicsPipelineState* PSO);
-	bool CreateGfxPipelineFromEntry(FVulkanRHIGraphicsPipelineState* PSO, FVulkanShader* Shaders[ShaderStage::NumStages], VkPipeline* Pipeline, FPipelineCache& Cache);
+	bool CreateGfxPipelineFromEntry(FVulkanRHIGraphicsPipelineState* PSO, FVulkanShader* Shaders[ShaderStage::NumStages], bool bPrecompile);
 
+	VkResult CreateVKPipeline(FVulkanRHIGraphicsPipelineState* PSO, FVulkanShader* Shaders[ShaderStage::NumStages], VkGraphicsPipelineCreateInfo PipelineInfo, bool bIsPrecompileJob);
+	static FString ShaderHashesToString(FVulkanShader* Shaders[ShaderStage::NumStages]);
 
 	FVulkanLayout* FindOrAddLayout(const FVulkanDescriptorSetsLayoutInfo& DescriptorSetLayoutInfo, bool bGfxLayout);
 	FVulkanComputePipeline* CreateComputePipelineFromShader(FVulkanComputeShader* Shader);
