VulkanRHI: Add implementation of bindless samplers using VK_EXT_descriptor_indexing.

diff --git a/Engine/Config/VulkanPC/DataDrivenPlatformInfo.ini b/Engine/Config/VulkanPC/DataDrivenPlatformInfo.ini
--- a/Engine/Config/VulkanPC/DataDrivenPlatformInfo.ini
+++ b/Engine/Config/VulkanPC/DataDrivenPlatformInfo.ini
@@ -133,6 +133,9 @@ bSupportsOIT=false
 
 EnablesHLSL2021ByDefault=1
 
+bSupportsBindless=true
+
+
 [ShaderPlatform VULKAN_PCES3_1]
 Language=Vulkan
 MaxFeatureLevel=ES3_1
diff --git a/Engine/Shaders/Public/Platform/Vulkan/VulkanCommon.ush b/Engine/Shaders/Public/Platform/Vulkan/VulkanCommon.ush
--- a/Engine/Shaders/Public/Platform/Vulkan/VulkanCommon.ush
+++ b/Engine/Shaders/Public/Platform/Vulkan/VulkanCommon.ush
@@ -35,3 +35,18 @@ void ImageInterlockedMaxUInt64(RWTexture2D<UlongType> Dest, uint2 Coord, UlongTy
 	InterlockedMax(Dest[Coord], Value);
 }
 #endif
+
+#if PLATFORM_SUPPORTS_BINDLESS
+	// While we use descriptor indexing, declare unbounded resource arrays
+	COMPILER_DEFINE VULKAN_HEAP_NAME(Prefix, Type) Prefix##Type
+	COMPILER_DEFINE VULKAN_SAMPLER_HEAP(Type) VULKAN_HEAP_NAME(VULKAN_BINDLESS_SAMPLER_ARRAY_PREFIX,Type)
+	COMPILER_DEFINE VULKAN_RESOURCE_HEAP(Type) VULKAN_HEAP_NAME(VULKAN_BINDLESS_RESOURCE_ARRAY_PREFIX,Type)
+
+	// todo-jn: bindless resource heaps
+	COMPILER_DEFINE GetResourceFromHeap(Type, Index) VULKAN_RESOURCE_HEAP(Type)[Index]
+
+	// bindless sampler heaps (specify a register so AUTO_BINDLESS macros aren't applied on them)
+	SamplerState VULKAN_SAMPLER_HEAP(SamplerState)[] : register(s0, space1000);
+	SamplerComparisonState VULKAN_SAMPLER_HEAP(SamplerComparisonState)[] : register(s0, space1000);
+	COMPILER_DEFINE GetSamplerFromHeap(Type, Index) VULKAN_SAMPLER_HEAP(Type)[Index]
+#endif
diff --git a/Engine/Source/Developer/VulkanShaderFormat/Private/VulkanShaderCompiler.cpp b/Engine/Source/Developer/VulkanShaderFormat/Private/VulkanShaderCompiler.cpp
--- a/Engine/Source/Developer/VulkanShaderFormat/Private/VulkanShaderCompiler.cpp
+++ b/Engine/Source/Developer/VulkanShaderFormat/Private/VulkanShaderCompiler.cpp
@@ -899,7 +899,6 @@ static void BuildShaderOutput(
 
 	FOLDVulkanCodeHeader OLDHeader;
 
-	FShaderParameterMap& ParameterMap = ShaderOutput.ParameterMap;
 	EShaderFrequency Frequency = (EShaderFrequency)ShaderOutput.Target.Frequency;
 
 	TBitArray<> UsedUniformBufferSlots;
@@ -1031,7 +1030,9 @@ static void BuildShaderOutput(
 			ShaderOutput
 		);
 
-		NEWEntryTypes.Add(PackedGlobal.Name, FVulkanShaderHeader::PackedGlobal);
+		FString ParamName = PackedGlobal.Name;
+		UE::ShaderCompilerCommon::RemoveBindlessParameterPrefix(ParamName);
+		NEWEntryTypes.Add(ParamName, FVulkanShaderHeader::PackedGlobal);
 
 		uint32& Size = PackedGlobalArraySize.FindOrAdd((CrossCompiler::EPackedTypeName)PackedGlobal.PackedType);
 		Size = FMath::Max<uint32>(BytesPerComponent * (PackedGlobal.Offset + PackedGlobal.Count), Size);
@@ -1197,7 +1198,7 @@ static void BuildShaderOutput(
 			if (!SharedSamplerStates.Contains(SamplerState))
 			{
 				// ParameterMap does not use a TMultiMap, so we cannot push the same entry to it more than once!  if we try to, we've done something wrong...
-				check(!ParameterMap.ContainsParameterAllocation(*SamplerState));
+				check(!ShaderOutput.ParameterMap.ContainsParameterAllocation(*SamplerState));
 
 				HandleReflectedShaderSampler(SamplerState, Sampler.Offset, VulkanBindingIndex, Sampler.Count, ShaderOutput);
 
@@ -1352,9 +1353,29 @@ FCompilerInfo::FCompilerInfo(const FShaderCompilerInput& InInput, const FString&
 static void GatherSpirvReflectionBindings(
 	spv_reflect::ShaderModule&	Reflection,
 	FSpirvReflectBindings&		OutBindings,
-	const EShaderFrequency		ShaderFrequency)
+	const EShaderFrequency		ShaderFrequency,
+	const bool					bSupportsBindless)
 {
-	SpvReflectResult SpvResult = SPV_REFLECT_RESULT_NOT_READY;
+	// Change descriptor set numbers
+	TArray<SpvReflectDescriptorSet*> DescriptorSets;
+	uint32 NumDescriptorSets = 0;
+
+	// If bindless is supported, then offset the descriptor set to fit the bindless heaps at the beginning
+	const uint32 DescSetNo = bSupportsBindless ? VulkanBindless::NumBindlessSets + ShaderStage::GetStageForFrequency(ShaderFrequency) : ShaderStage::GetStageForFrequency(ShaderFrequency);
+
+	SpvReflectResult SpvResult = Reflection.EnumerateDescriptorSets(&NumDescriptorSets, nullptr);
+	check(SpvResult == SPV_REFLECT_RESULT_SUCCESS);
+	if (NumDescriptorSets > 0)
+	{
+		DescriptorSets.SetNum(NumDescriptorSets);
+		SpvResult = Reflection.EnumerateDescriptorSets(&NumDescriptorSets, DescriptorSets.GetData());
+		check(SpvResult == SPV_REFLECT_RESULT_SUCCESS);
+
+		for (const SpvReflectDescriptorSet* DescSet : DescriptorSets)
+		{
+			Reflection.ChangeDescriptorSetNumber(DescSet, DescSetNo);
+		}
+	}
 
 	OutBindings.GatherInputAttributes(Reflection);
 	OutBindings.GatherOutputAttributes(Reflection);
@@ -1370,23 +1391,29 @@ static void GatherSpirvReflectionBindings(
 		OutBindings.AssignInputAttributeLocationsBySemanticIndex(Reflection, CrossCompiler::FShaderConductorContext::GetIdentifierTable().InputAttribute);
 	}
 
-	// Change descriptor set numbers
-	TArray<SpvReflectDescriptorSet*> DescriptorSets;
-	uint32 NumDescriptorSets = 0;
-
-	SpvResult = Reflection.EnumerateDescriptorSets(&NumDescriptorSets, nullptr);
-	check(SpvResult == SPV_REFLECT_RESULT_SUCCESS);
-	if (NumDescriptorSets > 0)
+	// Patch resource heaps descriptor set numbers
+	if (bSupportsBindless)
 	{
-		DescriptorSets.SetNum(NumDescriptorSets);
-		SpvResult = Reflection.EnumerateDescriptorSets(&NumDescriptorSets, DescriptorSets.GetData());
-		check(SpvResult == SPV_REFLECT_RESULT_SUCCESS);
-
-		for (const SpvReflectDescriptorSet* DescSet : DescriptorSets)
+		// Move the bindless heap to its dedicated descriptor set and remove it from our regular binding arrays
+		auto MoveBindlessHeaps = [&](TArray<SpvReflectDescriptorBinding*>& BindingArray, const TCHAR* HeapPrefix, uint32 BinldessDescSetNo)
 		{
-			const uint32 DescSetNo = ShaderStage::GetStageForFrequency(ShaderFrequency);
-			Reflection.ChangeDescriptorSetNumber(DescSet, DescSetNo);
-		}
+			for (int32 Index = BindingArray.Num() - 1; Index >= 0; --Index)
+			{
+				const SpvReflectDescriptorBinding* pBinding = BindingArray[Index];
+				FString BindingName(ANSI_TO_TCHAR(pBinding->name));
+				if (BindingName.StartsWith(HeapPrefix))
+				{
+					const uint32 Binding = 0;  // single bindless heap per descriptor set
+					Reflection.ChangeDescriptorBindingNumbers(pBinding, Binding, BinldessDescSetNo);
+					BindingArray.RemoveAtSwap(Index);
+				}
+			}
+		};
+
+		// Remove sampler heaps from binding arrays
+		MoveBindlessHeaps(OutBindings.Samplers, VulkanBindless::kBindlessSamplerArrayPrefix, VulkanBindless::BindlessSamplerSet);
+
+		// todo-jn: Remove resource heaps from binding arrays
 	}
 }
 
@@ -1493,8 +1520,10 @@ static bool BuildShaderOutputFromSpirv(
 		check(Result == SPV_REFLECT_RESULT_SUCCESS);
 	}
 
+	const bool bSupportsBindless = Input.Environment.CompilerFlags.Contains(CFLAG_BindlessResources) || Input.Environment.CompilerFlags.Contains(CFLAG_BindlessSamplers);
+
 	FSpirvReflectBindings Bindings;
-	GatherSpirvReflectionBindings(Reflection, Bindings, static_cast<EShaderFrequency>(Input.Target.Frequency));
+	GatherSpirvReflectionBindings(Reflection, Bindings, static_cast<EShaderFrequency>(Input.Target.Frequency), bSupportsBindless);
 
 	// Register how often a sampler-state is used
 	for (const SpvReflectDescriptorBinding* Binding : Bindings.TextureSRVs)
@@ -1832,7 +1861,7 @@ static bool BuildShaderOutputFromSpirv(
 	}
 
 	// For Android run an additional pass to patch spirv to be compatible across drivers
-	if(IsAndroidShaderPlatform(Input.Target.GetPlatform()))
+	if (IsAndroidShaderPlatform(Input.Target.GetPlatform()))
 	{
 		const char* OptArgs[] = { "--android-driver-patch" };
 		if (!CompilerContext.OptimizeSpirv(Spirv.Data, OptArgs, UE_ARRAY_COUNT(OptArgs)))
@@ -2118,6 +2147,7 @@ static bool CompileWithShaderConductor(
 	const FShaderCompilerInput& Input = CompilerInfo.Input;
 
 	const bool bIsRayTracingShader = Input.IsRayTracingShader();
+	const bool bHasBindless = Input.Environment.CompilerFlags.Contains(CFLAG_BindlessResources) || Input.Environment.CompilerFlags.Contains(CFLAG_BindlessSamplers);
 	const bool bRewriteHlslSource = !bIsRayTracingShader;
 	const bool bDebugDump = CompilerInfo.bDebugDump;
 
@@ -2288,6 +2318,13 @@ void DoCompileVulkanShader(const FShaderCompilerInput& Input, FShaderCompilerOut
 	{
 		AdditionalDefines.SetDefine(TEXT("PLATFORM_SUPPORTS_SM6_0_WAVE_OPERATIONS"), 1);
 	}
+	else
+	{
+		check(!Input.Environment.CompilerFlags.Contains(CFLAG_WaveOperations));
+	}
+
+	AdditionalDefines.SetDefine(TEXT("VULKAN_BINDLESS_SAMPLER_ARRAY_PREFIX"), VulkanBindless::kBindlessSamplerArrayPrefix);
+	AdditionalDefines.SetDefine(TEXT("VULKAN_BINDLESS_RESOURCE_ARRAY_PREFIX"), VulkanBindless::kBindlessResourceArrayPrefix);
 
 	// Preprocess the shader.
 	FString PreprocessedShaderSource;
diff --git a/Engine/Source/Runtime/VulkanRHI/Private/VulkanCommandBuffer.cpp b/Engine/Source/Runtime/VulkanRHI/Private/VulkanCommandBuffer.cpp
--- a/Engine/Source/Runtime/VulkanRHI/Private/VulkanCommandBuffer.cpp
+++ b/Engine/Source/Runtime/VulkanRHI/Private/VulkanCommandBuffer.cpp
@@ -252,12 +252,12 @@ void FVulkanCmdBuffer::End()
 
 	for (PendingQuery& Query : PendingTimestampQueries)
 	{
-		uint64 Index = Query.Index;
-		VkBuffer BufferHandle = Query.BufferHandle;
-		VkQueryPool PoolHandle = Query.PoolHandle;
-		VkQueryResultFlags BlockingFlags = Query.bBlocking ?  VK_QUERY_RESULT_WAIT_BIT : VK_QUERY_RESULT_WITH_AVAILABILITY_BIT;
-		uint32 Width = (Query.bBlocking ? 1 : 2);
-		uint32 Stride = sizeof(uint64) * Width;
+		const uint64 Index = Query.Index;
+		const VkBuffer BufferHandle = Query.BufferHandle;
+		const VkQueryPool PoolHandle = Query.PoolHandle;
+		const VkQueryResultFlags BlockingFlags = Query.bBlocking ?  VK_QUERY_RESULT_WAIT_BIT : VK_QUERY_RESULT_WITH_AVAILABILITY_BIT;
+		const uint32 Width = (Query.bBlocking ? 1 : 2);
+		const uint32 Stride = sizeof(uint64) * Width;
 
 		VulkanRHI::vkCmdCopyQueryPoolResults(GetHandle(), PoolHandle, Index, Query.Count, BufferHandle, Stride * Index, Stride, VK_QUERY_RESULT_64_BIT | BlockingFlags);
 		VulkanRHI::vkCmdResetQueryPool(GetHandle(), PoolHandle, Index, Query.Count);
@@ -279,7 +279,7 @@ inline void FVulkanCmdBuffer::InitializeTimings(FVulkanCommandListContext* InCon
 
 			// Upload cb's can be submitted multiple times in a single frame, so we use an expanded pool to catch timings
 			// Any overflow will wrap
-			uint32 PoolSize = bIsUploadOnly ? 256 : 32;
+			const uint32 PoolSize = bIsUploadOnly ? 256 : 32;
 			Timing->Initialize(PoolSize);
 		}
 	}
@@ -324,6 +324,22 @@ void FVulkanCmdBuffer::Begin()
 	}
 	check(!CurrentDescriptorPoolSetContainer);
 
+	if (!bIsUploadOnly && Device->SupportsBindless())
+	{
+		FVulkanBindlessDescriptorManager* BindlessDescriptorManager = Device->GetBindlessDescriptorManager();
+		FVulkanQueue* Queue = GetOwner()->GetMgr().GetQueue();
+		const VkPipelineStageFlags SupportedStages = Queue->GetSupportedStageBits();
+		if (SupportedStages & VK_PIPELINE_STAGE_FRAGMENT_SHADER_BIT)
+		{
+			BindlessDescriptorManager->BindDescriptorSets(CommandBufferHandle, VK_PIPELINE_BIND_POINT_GRAPHICS);
+		}
+		if (SupportedStages & VK_PIPELINE_STAGE_COMPUTE_SHADER_BIT)
+		{
+			BindlessDescriptorManager->BindDescriptorSets(CommandBufferHandle, VK_PIPELINE_BIND_POINT_COMPUTE);
+		}
+		// todo-jn: bindless : raytracing
+	}
+
 	bNeedsDynamicStateSet = true;
 }
 
diff --git a/Engine/Source/Runtime/VulkanRHI/Private/VulkanCommandBuffer.h b/Engine/Source/Runtime/VulkanRHI/Private/VulkanCommandBuffer.h
--- a/Engine/Source/Runtime/VulkanRHI/Private/VulkanCommandBuffer.h
+++ b/Engine/Source/Runtime/VulkanRHI/Private/VulkanCommandBuffer.h
@@ -365,6 +365,11 @@ public:
 		return Context;
 	}
 
+	FVulkanQueue* GetQueue()
+	{
+		return Queue;
+	}
+
 private:
 	struct FQueryPoolReset
 	{
diff --git a/Engine/Source/Runtime/VulkanRHI/Private/VulkanDescriptorSets.cpp b/Engine/Source/Runtime/VulkanRHI/Private/VulkanDescriptorSets.cpp
new file mode 100644
--- /dev/null
+++ b/Engine/Source/Runtime/VulkanRHI/Private/VulkanDescriptorSets.cpp
@@ -0,0 +1,194 @@
+// Copyright Epic Games, Inc. All Rights Reserved.
+
+/*=============================================================================
+	VulkanDescriptorSets.cpp: Vulkan descriptor set RHI implementation.
+=============================================================================*/
+
+#include "VulkanRHIPrivate.h"
+#include "VulkanDescriptorSets.h"
+
+
+int32 GBindlessMaxResourceDescriptorCount = 1000 * 1000;
+static FAutoConsoleVariableRef CVarBindlessResourceDescriptorCount(
+	TEXT("r.Vulkan.Bindless.MaxResourceDescriptorCount"),
+	GBindlessMaxResourceDescriptorCount,
+	TEXT("Maximum bindless resource descriptor count"),
+	ECVF_ReadOnly
+);
+
+int32 GBindlessMaxSamplerDescriptorCount = 2048;
+static FAutoConsoleVariableRef CVarBindlessSamplerDescriptorCount(
+	TEXT("r.Vulkan.Bindless.MaxSamplerDescriptorCount"),
+	GBindlessMaxSamplerDescriptorCount,
+	TEXT("Maximum bindless sampler descriptor count"),
+	ECVF_ReadOnly
+);
+
+
+FVulkanBindlessDescriptorManager::FVulkanBindlessDescriptorManager(FVulkanDevice* InDevice)
+	: VulkanRHI::FDeviceChild(InDevice)
+	, bBindlessResourcesAllowed(RHIGetBindlessResourcesConfiguration(GMaxRHIShaderPlatform) != ERHIBindlessConfiguration::Disabled)
+	, bBindlessSamplersAllowed(RHIGetBindlessSamplersConfiguration(GMaxRHIShaderPlatform) != ERHIBindlessConfiguration::Disabled)
+{
+
+}
+
+FVulkanBindlessDescriptorManager::~FVulkanBindlessDescriptorManager()
+{
+	VulkanRHI::vkDestroyPipelineLayout(Device->GetInstanceHandle(), BindlessPipelineLayout, VULKAN_CPU_ALLOCATOR);
+
+	VulkanRHI::vkDestroyDescriptorSetLayout(Device->GetInstanceHandle(), EmptyDescriptorSetLayout, VULKAN_CPU_ALLOCATOR);
+	VulkanRHI::vkDestroyDescriptorSetLayout(Device->GetInstanceHandle(), SamplerDescriptorSetLayout, VULKAN_CPU_ALLOCATOR);
+	VulkanRHI::vkDestroyDescriptorSetLayout(Device->GetInstanceHandle(), ResourceDescriptorSetLayout, VULKAN_CPU_ALLOCATOR);
+
+	VulkanRHI::vkDestroyDescriptorPool(Device->GetInstanceHandle(), DescriptorPool, VULKAN_CPU_ALLOCATOR);
+}
+
+void FVulkanBindlessDescriptorManager::Init()
+{
+	MaxResourceDescriptors = GBindlessMaxResourceDescriptorCount;
+	MaxSamplerDescriptors = GBindlessMaxSamplerDescriptorCount;
+
+	// Create the empty descriptor set (used to pad empty indices)
+	{
+		VkDescriptorSetLayoutCreateInfo EmptyDescriptorSetLayoutCreateInfo;
+		ZeroVulkanStruct(EmptyDescriptorSetLayoutCreateInfo, VK_STRUCTURE_TYPE_DESCRIPTOR_SET_LAYOUT_CREATE_INFO);
+		VERIFYVULKANRESULT(VulkanRHI::vkCreateDescriptorSetLayout(Device->GetInstanceHandle(), &EmptyDescriptorSetLayoutCreateInfo, VULKAN_CPU_ALLOCATOR, &EmptyDescriptorSetLayout));
+	}
+
+	// Create the descriptor pool for bindless
+	{
+		const uint32 DescriptorTypeCount = 1;
+		VkDescriptorPoolSize DescriptorPoolSize[DescriptorTypeCount];
+		DescriptorPoolSize[0].type = VK_DESCRIPTOR_TYPE_SAMPLER;
+		DescriptorPoolSize[0].descriptorCount = MaxSamplerDescriptors;
+		// todo-jn: bindless: add 'real' resource descriptors
+
+		VkDescriptorPoolCreateInfo DescriptorPoolCreateInfo;
+		ZeroVulkanStruct(DescriptorPoolCreateInfo, VK_STRUCTURE_TYPE_DESCRIPTOR_POOL_CREATE_INFO);
+		DescriptorPoolCreateInfo.maxSets = VulkanBindless::NumBindlessSets;
+		DescriptorPoolCreateInfo.poolSizeCount = DescriptorTypeCount;
+		DescriptorPoolCreateInfo.pPoolSizes = DescriptorPoolSize;
+		DescriptorPoolCreateInfo.flags = VK_DESCRIPTOR_POOL_CREATE_UPDATE_AFTER_BIND_BIT;
+
+		VERIFYVULKANRESULT(VulkanRHI::vkCreateDescriptorPool(Device->GetInstanceHandle(), &DescriptorPoolCreateInfo, VULKAN_CPU_ALLOCATOR, &DescriptorPool));
+	}
+
+	// Create the sampler descriptor set layout
+	{
+		VkDescriptorSetLayoutBinding DescriptorSetLayoutBinding;
+		DescriptorSetLayoutBinding.binding = 0;
+		DescriptorSetLayoutBinding.descriptorType = VK_DESCRIPTOR_TYPE_SAMPLER;
+		DescriptorSetLayoutBinding.descriptorCount = MaxSamplerDescriptors;
+		DescriptorSetLayoutBinding.stageFlags = VK_SHADER_STAGE_ALL;
+		DescriptorSetLayoutBinding.pImmutableSamplers = nullptr;
+
+		const VkDescriptorBindingFlags DescriptorBindingFlags =
+			VK_DESCRIPTOR_BINDING_UPDATE_AFTER_BIND_BIT |
+			VK_DESCRIPTOR_BINDING_PARTIALLY_BOUND_BIT |
+			VK_DESCRIPTOR_BINDING_UPDATE_UNUSED_WHILE_PENDING_BIT |
+			VK_DESCRIPTOR_BINDING_VARIABLE_DESCRIPTOR_COUNT_BIT;
+		VkDescriptorSetLayoutBindingFlagsCreateInfo DescriptorSetLayoutBindingFlagsCreateInfo;
+		ZeroVulkanStruct(DescriptorSetLayoutBindingFlagsCreateInfo, VK_STRUCTURE_TYPE_DESCRIPTOR_SET_LAYOUT_BINDING_FLAGS_CREATE_INFO);
+		DescriptorSetLayoutBindingFlagsCreateInfo.bindingCount = 1;
+		DescriptorSetLayoutBindingFlagsCreateInfo.pBindingFlags = &DescriptorBindingFlags;
+
+		VkDescriptorSetLayoutCreateInfo DescriptorSetLayoutCreateInfo;
+		ZeroVulkanStruct(DescriptorSetLayoutCreateInfo, VK_STRUCTURE_TYPE_DESCRIPTOR_SET_LAYOUT_CREATE_INFO);
+		DescriptorSetLayoutCreateInfo.pBindings = &DescriptorSetLayoutBinding;
+		DescriptorSetLayoutCreateInfo.bindingCount = 1;
+		DescriptorSetLayoutCreateInfo.flags = VK_DESCRIPTOR_SET_LAYOUT_CREATE_UPDATE_AFTER_BIND_POOL_BIT;
+		DescriptorSetLayoutCreateInfo.pNext = &DescriptorSetLayoutBindingFlagsCreateInfo;
+
+		VERIFYVULKANRESULT(VulkanRHI::vkCreateDescriptorSetLayout(Device->GetInstanceHandle(), &DescriptorSetLayoutCreateInfo, VULKAN_CPU_ALLOCATOR, &SamplerDescriptorSetLayout));
+	}
+
+	// Create the sampler descriptor set
+	{
+		VkDescriptorSetVariableDescriptorCountAllocateInfo VariableDescriptorCountAllocateInfo;
+		ZeroVulkanStruct(VariableDescriptorCountAllocateInfo, VK_STRUCTURE_TYPE_DESCRIPTOR_SET_VARIABLE_DESCRIPTOR_COUNT_ALLOCATE_INFO);
+		VariableDescriptorCountAllocateInfo.descriptorSetCount = 1;
+		VariableDescriptorCountAllocateInfo.pDescriptorCounts = &MaxSamplerDescriptors;
+
+		VkDescriptorSetAllocateInfo DescriptorSetAllocateInfo;
+		ZeroVulkanStruct(DescriptorSetAllocateInfo, VK_STRUCTURE_TYPE_DESCRIPTOR_SET_ALLOCATE_INFO);
+		DescriptorSetAllocateInfo.descriptorPool = DescriptorPool;
+		DescriptorSetAllocateInfo.descriptorSetCount = 1;
+		DescriptorSetAllocateInfo.pSetLayouts = &SamplerDescriptorSetLayout;
+		DescriptorSetAllocateInfo.pNext = &VariableDescriptorCountAllocateInfo;
+
+		VERIFYVULKANRESULT(VulkanRHI::vkAllocateDescriptorSets(Device->GetInstanceHandle(), &DescriptorSetAllocateInfo, &DescriptorSets[VulkanBindless::BindlessSamplerSet]));
+	}
+
+	// todo-jn: bindless: Create the resource descriptor set and layout
+	{
+		ResourceDescriptorSetLayout = EmptyDescriptorSetLayout;
+		DescriptorSets[VulkanBindless::BindlessResourceSet] = VK_NULL_HANDLE;
+	}
+
+	// Now create the basic pipeline
+	{
+		VkDescriptorSetLayout DescriptorSetLayouts[VulkanBindless::NumBindlessSets];
+		DescriptorSetLayouts[VulkanBindless::BindlessSamplerSet] = SamplerDescriptorSetLayout;
+		DescriptorSetLayouts[VulkanBindless::BindlessResourceSet] = ResourceDescriptorSetLayout;
+
+		VkPipelineLayoutCreateInfo PipelineLayoutCreateInfo;
+		ZeroVulkanStruct(PipelineLayoutCreateInfo, VK_STRUCTURE_TYPE_PIPELINE_LAYOUT_CREATE_INFO);
+		PipelineLayoutCreateInfo.setLayoutCount = VulkanBindless::NumBindlessSets;
+		PipelineLayoutCreateInfo.pSetLayouts = DescriptorSetLayouts;
+
+		VERIFYVULKANRESULT(VulkanRHI::vkCreatePipelineLayout(Device->GetInstanceHandle(), &PipelineLayoutCreateInfo, VULKAN_CPU_ALLOCATOR, &BindlessPipelineLayout));
+	}
+}
+
+FVulkanBindlessDescriptorManager::BindlessLayoutArray FVulkanBindlessDescriptorManager::GeneratePipelineLayout(const TArray<VkDescriptorSetLayout>& LayoutArray) const
+{
+	BindlessLayoutArray PatchedArray;
+	PatchedArray.Add(SamplerDescriptorSetLayout);
+	PatchedArray.Add(ResourceDescriptorSetLayout);
+	PatchedArray.Append(LayoutArray);
+	return PatchedArray;
+}
+
+void FVulkanBindlessDescriptorManager::BindDescriptorSets(VkCommandBuffer CommandBuffer, VkPipelineBindPoint BindPoint)
+{
+	uint32 FirstDescriptorSet = VulkanBindless::NumBindlessSets;
+	uint32 NumDescriptorSets = 0;
+	if (bBindlessSamplersAllowed)
+	{
+		FirstDescriptorSet = FMath::Min<uint32>(VulkanBindless::BindlessSamplerSet, FirstDescriptorSet);
+		++NumDescriptorSets;
+	}
+	if (bBindlessResourcesAllowed)
+	{
+		FirstDescriptorSet = FMath::Min<uint32>(VulkanBindless::BindlessResourceSet, FirstDescriptorSet);
+		++NumDescriptorSets;
+	}
+
+	if (NumDescriptorSets > 0)
+	{
+		VulkanRHI::vkCmdBindDescriptorSets(CommandBuffer, BindPoint, BindlessPipelineLayout, 0, NumDescriptorSets, &DescriptorSets[FirstDescriptorSet], 0, nullptr);
+	}
+}
+
+FRHIDescriptorHandle FVulkanBindlessDescriptorManager::RegisterSampler(VkSampler VulkanSampler)
+{
+	const uint32 SamplerIndex = BindlessSamplerCount++;
+	checkf(SamplerIndex < MaxSamplerDescriptors, TEXT("You need to grow the sampler array size!"));
+
+	VkDescriptorImageInfo DescriptorImageInfo;
+	FMemory::Memzero(DescriptorImageInfo);
+	DescriptorImageInfo.sampler = VulkanSampler;
+
+	VkWriteDescriptorSet WriteDescriptorSet;
+	ZeroVulkanStruct(WriteDescriptorSet, VK_STRUCTURE_TYPE_WRITE_DESCRIPTOR_SET);
+	WriteDescriptorSet.dstSet = DescriptorSets[VulkanBindless::BindlessSamplerSet];
+	WriteDescriptorSet.dstBinding = 0;
+	WriteDescriptorSet.dstArrayElement = SamplerIndex;
+	WriteDescriptorSet.descriptorCount = 1;
+	WriteDescriptorSet.descriptorType = VK_DESCRIPTOR_TYPE_SAMPLER;
+	WriteDescriptorSet.pImageInfo = &DescriptorImageInfo;
+
+	VulkanRHI::vkUpdateDescriptorSets(Device->GetInstanceHandle(), 1, &WriteDescriptorSet, 0, nullptr);
+	return FRHIDescriptorHandle(ERHIDescriptorHeapType::Sampler, SamplerIndex);
+}
\ No newline at end of file
diff --git a/Engine/Source/Runtime/VulkanRHI/Private/VulkanDescriptorSets.h b/Engine/Source/Runtime/VulkanRHI/Private/VulkanDescriptorSets.h
--- a/Engine/Source/Runtime/VulkanRHI/Private/VulkanDescriptorSets.h
+++ b/Engine/Source/Runtime/VulkanRHI/Private/VulkanDescriptorSets.h
@@ -1375,7 +1375,7 @@ private:
 	public:
 		FCachedPool(FVulkanDevice* InDevice, uint32 InMaxDescriptorSets, const float PoolSizesRatio[VK_DESCRIPTOR_TYPE_RANGE_SIZE]);
 
-		uint32 GetMaxDescriptorSets() const
+		inline uint32 GetMaxDescriptorSets() const
 		{
 			return Pool.GetMaxDescriptorSets();
 		}
@@ -1412,3 +1412,40 @@ private:
 	float PoolAllocRatio;
 };
 
+// Manager for resource descriptors used in bindless rendering.
+class FVulkanBindlessDescriptorManager : public VulkanRHI::FDeviceChild
+{
+public:
+	using BindlessLayoutArray = TArray<VkDescriptorSetLayout, TInlineAllocator<VulkanBindless::MaxNumSets>>;
+
+	FVulkanBindlessDescriptorManager(FVulkanDevice* InDevice);
+	~FVulkanBindlessDescriptorManager();
+
+	void Init();
+
+	BindlessLayoutArray GeneratePipelineLayout(const TArray<VkDescriptorSetLayout>& LayoutArray) const;
+
+	void BindDescriptorSets(VkCommandBuffer CommandBuffer, VkPipelineBindPoint BindPoint);
+
+	FRHIDescriptorHandle RegisterSampler(VkSampler VulkanSampler);
+
+private:
+	const bool bBindlessResourcesAllowed;
+	const bool bBindlessSamplersAllowed;
+
+	uint32 MaxResourceDescriptors = 0;
+	uint32 MaxSamplerDescriptors = 0;
+
+	std::atomic<uint32> BindlessSamplerCount = 0;
+
+	VkDescriptorSetLayout EmptyDescriptorSetLayout = VK_NULL_HANDLE;
+
+	VkDescriptorPool DescriptorPool = VK_NULL_HANDLE;
+
+	VkDescriptorSetLayout SamplerDescriptorSetLayout = VK_NULL_HANDLE;
+	VkDescriptorSetLayout ResourceDescriptorSetLayout = VK_NULL_HANDLE;
+
+	VkDescriptorSet  DescriptorSets[VulkanBindless::NumBindlessSets] = { VK_NULL_HANDLE };
+
+	VkPipelineLayout BindlessPipelineLayout = VK_NULL_HANDLE;
+};
diff --git a/Engine/Source/Runtime/VulkanRHI/Private/VulkanDevice.cpp b/Engine/Source/Runtime/VulkanRHI/Private/VulkanDevice.cpp
--- a/Engine/Source/Runtime/VulkanRHI/Private/VulkanDevice.cpp
+++ b/Engine/Source/Runtime/VulkanRHI/Private/VulkanDevice.cpp
@@ -16,6 +16,7 @@
 #include "VulkanTransientResourceAllocator.h"
 #include "VulkanExtensions.h"
 #include "VulkanRayTracing.h"
+#include "VulkanDescriptorSets.h"
 
 TAutoConsoleVariable<int32> GRHIAllowAsyncComputeCvar(
 	TEXT("r.Vulkan.AllowAsyncCompute"),
@@ -1074,6 +1075,20 @@ void FVulkanDevice::InitGPU(int32 DeviceIndex)
 	DescriptorPoolsManager = new FVulkanDescriptorPoolsManager();
 	DescriptorPoolsManager->Init(this);
 
+	if (RHISupportsBindless(GMaxRHIShaderPlatform) &&
+		OptionalDeviceExtensions.HasEXTDescriptorIndexing &&
+		(GpuProps.limits.maxBoundDescriptorSets >= VulkanBindless::MaxNumSets))
+	{
+		const bool bBindlessResources = RHIGetBindlessResourcesConfiguration(GMaxRHIShaderPlatform) != ERHIBindlessConfiguration::Disabled;
+		const bool bBindlessSamplers = RHIGetBindlessSamplersConfiguration(GMaxRHIShaderPlatform) != ERHIBindlessConfiguration::Disabled;
+
+		if (bBindlessResources || bBindlessSamplers)
+		{
+			BindlessDescriptorManager = new FVulkanBindlessDescriptorManager(this);
+			BindlessDescriptorManager->Init();
+		}
+	}
+
 	PipelineStateCache = new FVulkanPipelineStateCacheManager(this);
 
 	TArray<FString> CacheFilenames = FVulkanPlatform::GetPSOCacheFilenames();
@@ -1182,6 +1197,12 @@ void FVulkanDevice::Destroy()
 	delete DescriptorPoolsManager;
 	DescriptorPoolsManager = nullptr;
 
+	if (BindlessDescriptorManager)
+	{
+		delete BindlessDescriptorManager;
+		BindlessDescriptorManager = nullptr;
+	}
+
 	// No need to delete as it's stored in SamplerMap
 	DefaultSampler = nullptr;
 
diff --git a/Engine/Source/Runtime/VulkanRHI/Private/VulkanDevice.h b/Engine/Source/Runtime/VulkanRHI/Private/VulkanDevice.h
--- a/Engine/Source/Runtime/VulkanRHI/Private/VulkanDevice.h
+++ b/Engine/Source/Runtime/VulkanRHI/Private/VulkanDevice.h
@@ -11,6 +11,7 @@
 class FVulkanDescriptorSetCache;
 class FVulkanDescriptorPool;
 class FVulkanDescriptorPoolsManager;
+class FVulkanBindlessDescriptorManager;
 class FVulkanCommandListContextImmediate;
 class FVulkanTransientHeapCache;
 class FVulkanDeviceExtension;
@@ -312,6 +313,11 @@ public:
 		return DeviceMemoryManager.HasUnifiedMemory();
 	}
 
+	inline bool SupportsBindless() const
+	{
+		return (BindlessDescriptorManager != nullptr);
+	}
+
 	inline uint64 GetTimestampValidBitsMask() const
 	{
 		return TimestampValidBitsMask;
@@ -379,6 +385,11 @@ public:
 		return *DescriptorPoolsManager;
 	}
 
+	inline FVulkanBindlessDescriptorManager* GetBindlessDescriptorManager()
+	{
+		return BindlessDescriptorManager;
+	}
+
 	inline TMap<uint32, FSamplerStateRHIRef>& GetSamplerMap()
 	{
 		return SamplerMap;
@@ -515,6 +526,8 @@ private:
 	// Active on >= SM4
 	FVulkanDescriptorPoolsManager* DescriptorPoolsManager = nullptr;
 
+	FVulkanBindlessDescriptorManager* BindlessDescriptorManager = nullptr;
+
 	FVulkanShaderFactory ShaderFactory;
 
 	FVulkanSamplerState* DefaultSampler;
diff --git a/Engine/Source/Runtime/VulkanRHI/Private/VulkanExtensions.cpp b/Engine/Source/Runtime/VulkanRHI/Private/VulkanExtensions.cpp
--- a/Engine/Source/Runtime/VulkanRHI/Private/VulkanExtensions.cpp
+++ b/Engine/Source/Runtime/VulkanRHI/Private/VulkanExtensions.cpp
@@ -401,12 +401,17 @@ public:
 
 	virtual void PostPhysicalDeviceFeatures(FOptionalVulkanDeviceExtensions& ExtensionFlags) override final
 	{
-		ExtensionFlags.HasEXTDescriptorIndexing = DescriptorIndexingFeatures.runtimeDescriptorArray; // :todo-jn: add resource specific checks
+		bFullySupported = (DescriptorIndexingFeatures.runtimeDescriptorArray == VK_TRUE) &&
+			(DescriptorIndexingFeatures.descriptorBindingPartiallyBound == VK_TRUE) &&
+			(DescriptorIndexingFeatures.descriptorBindingUpdateUnusedWhilePending == VK_TRUE) &&
+			(DescriptorIndexingFeatures.descriptorBindingVariableDescriptorCount == VK_TRUE);
+
+		ExtensionFlags.HasEXTDescriptorIndexing = bFullySupported ? 1 : 0;
 	}
 
 	virtual void PreCreateDevice(VkDeviceCreateInfo& DeviceCreateInfo) override final
 	{
-		if (DescriptorIndexingFeatures.runtimeDescriptorArray == VK_TRUE)
+		if (bFullySupported)
 		{
 			AddToPNext(DeviceCreateInfo, DescriptorIndexingFeatures);
 		}
@@ -414,6 +419,7 @@ public:
 
 private:
 	VkPhysicalDeviceDescriptorIndexingFeaturesEXT DescriptorIndexingFeatures;
+	bool bFullySupported = false;
 };
 
 
diff --git a/Engine/Source/Runtime/VulkanRHI/Private/VulkanMemory.cpp b/Engine/Source/Runtime/VulkanRHI/Private/VulkanMemory.cpp
--- a/Engine/Source/Runtime/VulkanRHI/Private/VulkanMemory.cpp
+++ b/Engine/Source/Runtime/VulkanRHI/Private/VulkanMemory.cpp
@@ -3569,19 +3569,13 @@ namespace VulkanRHI
 		}
 	}
 
-	void FMemoryManager::AllocUniformBuffer(FVulkanAllocation& OutAllocation, uint32 Size, const void* Contents)
+	void FMemoryManager::AllocUniformBuffer(FVulkanAllocation& OutAllocation, uint32 Size)
 	{
 		if(!AllocateBufferPooled(OutAllocation, nullptr, Size, 0, VK_BUFFER_USAGE_UNIFORM_BUFFER_BIT | VK_BUFFER_USAGE_TRANSFER_DST_BIT, VK_MEMORY_PROPERTY_HOST_VISIBLE_BIT | VK_MEMORY_PROPERTY_HOST_COHERENT_BIT, EVulkanAllocationMetaUniformBuffer, __FILE__, __LINE__))
 		{
 			HandleOOM(false);
 			checkNoEntry();
 		}
-
-		if (Contents)
-		{
-			FMemory::Memcpy(OutAllocation.GetMappedPointer(Device), Contents, Size);
-			OutAllocation.FlushMappedMemory(Device);
-		}
 	}
 	void FMemoryManager::FreeUniformBuffer(FVulkanAllocation& InAllocation)
 	{
diff --git a/Engine/Source/Runtime/VulkanRHI/Private/VulkanMemory.h b/Engine/Source/Runtime/VulkanRHI/Private/VulkanMemory.h
--- a/Engine/Source/Runtime/VulkanRHI/Private/VulkanMemory.h
+++ b/Engine/Source/Runtime/VulkanRHI/Private/VulkanMemory.h
@@ -797,7 +797,7 @@ namespace VulkanRHI
 		void DumpMemory(bool bFullDump = true);
 
 
-		void AllocUniformBuffer(FVulkanAllocation& OutAllocation, uint32 Size, const void* Contents);
+		void AllocUniformBuffer(FVulkanAllocation& OutAllocation, uint32 Size);
 		void FreeUniformBuffer(FVulkanAllocation& InAllocation);
 
 		void HandleOOM(bool bCanResume = false, VkResult Result = VK_SUCCESS, uint64 AllocationSize = 0, uint32 MemoryTypeIndex = 0);
diff --git a/Engine/Source/Runtime/VulkanRHI/Private/VulkanPipelineState.h b/Engine/Source/Runtime/VulkanRHI/Private/VulkanPipelineState.h
--- a/Engine/Source/Runtime/VulkanRHI/Private/VulkanPipelineState.h
+++ b/Engine/Source/Runtime/VulkanRHI/Private/VulkanPipelineState.h
@@ -25,6 +25,8 @@ public:
 	FVulkanCommonPipelineDescriptorState(FVulkanDevice* InDevice)
 		: VulkanRHI::FDeviceChild(InDevice)
 	{
+		bBindlessSamplers = (RHIGetBindlessSamplersConfiguration(GMaxRHIShaderPlatform) != ERHIBindlessConfiguration::Disabled);
+		bBindlessResources = (RHIGetBindlessResourcesConfiguration(GMaxRHIShaderPlatform) != ERHIBindlessConfiguration::Disabled);
 	}
 
 	virtual ~FVulkanCommonPipelineDescriptorState() {}
@@ -104,8 +106,11 @@ public:
 
 	inline void SetSamplerState(uint8 DescriptorSet, uint32 BindingIndex, const FVulkanSamplerState* Sampler)
 	{
-		check(Sampler && Sampler->Sampler != VK_NULL_HANDLE);
-		MarkDirty(DSWriter[DescriptorSet].WriteSampler(BindingIndex, *Sampler));
+		if (!bBindlessSamplers)
+		{
+			check(Sampler && Sampler->Sampler != VK_NULL_HANDLE);
+			MarkDirty(DSWriter[DescriptorSet].WriteSampler(BindingIndex, *Sampler));
+		}
 	}
 
 	inline void SetInputAttachment(uint8 DescriptorSet, uint32 BindingIndex, const FVulkanTextureView& TextureView, VkImageLayout Layout)
@@ -148,10 +153,12 @@ protected:
 	}
 	inline void Bind(VkCommandBuffer CmdBuffer, VkPipelineLayout PipelineLayout, VkPipelineBindPoint BindPoint)
 	{
+		// Bindless will sneak in two sets before the others
+		const uint32 FirstSet = Device->SupportsBindless() ? VulkanBindless::NumBindlessSets : 0;
 		VulkanRHI::vkCmdBindDescriptorSets(CmdBuffer,
 			BindPoint,
 			PipelineLayout,
-			0, DescriptorSetHandles.Num(), DescriptorSetHandles.GetData(),
+			FirstSet, DescriptorSetHandles.Num(), DescriptorSetHandles.GetData(),
 			(uint32)DynamicOffsets.Num(), DynamicOffsets.GetData());
 	}
 
@@ -177,6 +184,9 @@ protected:
 	
 	mutable FVulkanDSetsKey DSetsKey;
 	mutable bool bIsDSetsKeyDirty = true;
+
+	bool bBindlessSamplers = false;
+	bool bBindlessResources = false;
 };
 
 
diff --git a/Engine/Source/Runtime/VulkanRHI/Private/VulkanRHI.cpp b/Engine/Source/Runtime/VulkanRHI/Private/VulkanRHI.cpp
--- a/Engine/Source/Runtime/VulkanRHI/Private/VulkanRHI.cpp
+++ b/Engine/Source/Runtime/VulkanRHI/Private/VulkanRHI.cpp
@@ -845,6 +845,7 @@ void FVulkanDynamicRHI::InitInstance()
 		GRHISupportsBaseVertexIndex = true;
 		GSupportsSeparateRenderTargetBlendState = true;
 		GRHISupportsSeparateDepthStencilCopyAccess = false;
+		GRHISupportsBindless =	RHISupportsBindless(GMaxRHIShaderPlatform) && Device->SupportsBindless();
 
 		GRHIMaxDispatchThreadGroupsPerDimension.X = FMath::Min<uint32>(Limits.maxComputeWorkGroupCount[0], 0x7fffffff);
 		GRHIMaxDispatchThreadGroupsPerDimension.Y = FMath::Min<uint32>(Limits.maxComputeWorkGroupCount[1], 0x7fffffff);
diff --git a/Engine/Source/Runtime/VulkanRHI/Private/VulkanShaders.cpp b/Engine/Source/Runtime/VulkanRHI/Private/VulkanShaders.cpp
--- a/Engine/Source/Runtime/VulkanRHI/Private/VulkanShaders.cpp
+++ b/Engine/Source/Runtime/VulkanRHI/Private/VulkanShaders.cpp
@@ -10,6 +10,7 @@
 #include "GlobalShader.h"
 #include "Serialization/MemoryReader.h"
 #include "VulkanLLM.h"
+#include "VulkanDescriptorSets.h"
 
 TAutoConsoleVariable<int32> GDynamicGlobalUBs(
 	TEXT("r.Vulkan.DynamicGlobalUBs"),
@@ -369,7 +370,8 @@ void FVulkanLayout::PatchSpirvBindings(FVulkanShader::FSpirvCode& SprivCode, ESh
 	TArrayView<uint32> Spirv = SprivCode.GetCodeView();	//#todo-rco: Do we need an actual copy of the SPIR-V?
 	ShaderStage::EStage Stage = ShaderStage::GetStageForFrequency(Frequency);
 	const FDescriptorSetRemappingInfo::FStageInfo& StageInfo = DescriptorSetLayout.RemappingInfo.StageInfos[Stage];
-	
+	const uint32 FirstSet = Device->SupportsBindless() ? VulkanBindless::NumBindlessSets : 0;
+
 	checkSlow(StageInfo.UniformBuffers.Num() == CodeHeader.UniformBufferSpirvInfos.Num());
 	for (int32 Index = 0; Index < CodeHeader.UniformBufferSpirvInfos.Num(); ++Index)
 	{
@@ -379,7 +381,7 @@ void FVulkanLayout::PatchSpirvBindings(FVulkanShader::FSpirvCode& SprivCode, ESh
 			const uint32 OffsetBindingIndex = CodeHeader.UniformBufferSpirvInfos[Index].BindingIndexOffset;
 			check(OffsetDescriptorSet != UINT32_MAX && OffsetBindingIndex != UINT32_MAX);
 			uint16 NewDescriptorSet = StageInfo.UniformBuffers[Index].Remapping.NewDescriptorSet;
-			Spirv[OffsetDescriptorSet] = NewDescriptorSet;
+			Spirv[OffsetDescriptorSet] = FirstSet + NewDescriptorSet;
 			uint16 NewBindingIndex = StageInfo.UniformBuffers[Index].Remapping.NewBindingIndex;
 			Spirv[OffsetBindingIndex] = NewBindingIndex;
 		}
@@ -392,7 +394,7 @@ void FVulkanLayout::PatchSpirvBindings(FVulkanShader::FSpirvCode& SprivCode, ESh
 		const uint32 OffsetBindingIndex = CodeHeader.GlobalSpirvInfos[Index].BindingIndexOffset;
 		check(OffsetDescriptorSet != UINT32_MAX && OffsetBindingIndex != UINT32_MAX);
 		uint16 NewDescriptorSet = StageInfo.Globals[Index].NewDescriptorSet;
-		Spirv[OffsetDescriptorSet] = NewDescriptorSet;
+		Spirv[OffsetDescriptorSet] = FirstSet + NewDescriptorSet;
 		uint16 NewBindingIndex = StageInfo.Globals[Index].NewBindingIndex;
 		Spirv[OffsetBindingIndex] = NewBindingIndex;
 	}
@@ -403,7 +405,7 @@ void FVulkanLayout::PatchSpirvBindings(FVulkanShader::FSpirvCode& SprivCode, ESh
 		const uint32 OffsetDescriptorSet = CodeHeader.PackedUBs[Index].SPIRVDescriptorSetOffset;
 		const uint32 OffsetBindingIndex = CodeHeader.PackedUBs[Index].SPIRVBindingIndexOffset;
 		check(OffsetDescriptorSet != UINT32_MAX && OffsetBindingIndex != UINT32_MAX);
-		Spirv[OffsetDescriptorSet] = StageInfo.PackedUBDescriptorSet;
+		Spirv[OffsetDescriptorSet] = FirstSet + StageInfo.PackedUBDescriptorSet;
 		Spirv[OffsetBindingIndex] = StageInfo.PackedUBBindingIndices[Index];
 	}
 }
@@ -474,15 +476,26 @@ void FVulkanLayout::Compile(FVulkanDescriptorSetLayoutMap& DSetLayoutMap)
 
 	DescriptorSetLayout.Compile(DSetLayoutMap);
 
-	const TArray<VkDescriptorSetLayout>& LayoutHandles = DescriptorSetLayout.GetHandles();
-
 	VkPipelineLayoutCreateInfo PipelineLayoutCreateInfo;
 	ZeroVulkanStruct(PipelineLayoutCreateInfo, VK_STRUCTURE_TYPE_PIPELINE_LAYOUT_CREATE_INFO);
-	PipelineLayoutCreateInfo.setLayoutCount = LayoutHandles.Num();
-	PipelineLayoutCreateInfo.pSetLayouts = LayoutHandles.GetData();
-	//PipelineLayoutCreateInfo.pushConstantRangeCount = 0;
 
-	VERIFYVULKANRESULT(VulkanRHI::vkCreatePipelineLayout(Device->GetInstanceHandle(), &PipelineLayoutCreateInfo, VULKAN_CPU_ALLOCATOR, &PipelineLayout));
+	// Sneak in extra descriptor sets when in bindless
+	FVulkanBindlessDescriptorManager* BindlessDescriptorManager = Device->GetBindlessDescriptorManager();
+	if (BindlessDescriptorManager)
+	{
+		FVulkanBindlessDescriptorManager::BindlessLayoutArray LayoutHandles = BindlessDescriptorManager->GeneratePipelineLayout(DescriptorSetLayout.GetHandles());
+		PipelineLayoutCreateInfo.setLayoutCount = LayoutHandles.Num();
+		PipelineLayoutCreateInfo.pSetLayouts = LayoutHandles.GetData();
+		VERIFYVULKANRESULT(VulkanRHI::vkCreatePipelineLayout(Device->GetInstanceHandle(), &PipelineLayoutCreateInfo, VULKAN_CPU_ALLOCATOR, &PipelineLayout));
+	}
+	else
+	{
+		const TArray<VkDescriptorSetLayout>& LayoutHandles = DescriptorSetLayout.GetHandles();
+		PipelineLayoutCreateInfo.setLayoutCount = LayoutHandles.Num();
+		PipelineLayoutCreateInfo.pSetLayouts = LayoutHandles.GetData();
+		//PipelineLayoutCreateInfo.pushConstantRangeCount = 0;
+		VERIFYVULKANRESULT(VulkanRHI::vkCreatePipelineLayout(Device->GetInstanceHandle(), &PipelineLayoutCreateInfo, VULKAN_CPU_ALLOCATOR, &PipelineLayout));
+	}
 }
 
 bool FVulkanGfxLayout::UsesInputAttachment(FVulkanShaderHeader::EAttachmentType AttachmentType) const
diff --git a/Engine/Source/Runtime/VulkanRHI/Private/VulkanState.cpp b/Engine/Source/Runtime/VulkanRHI/Private/VulkanState.cpp
--- a/Engine/Source/Runtime/VulkanRHI/Private/VulkanState.cpp
+++ b/Engine/Source/Runtime/VulkanRHI/Private/VulkanState.cpp
@@ -6,6 +6,7 @@
 
 #include "VulkanRHIPrivate.h"
 #include "VulkanPendingState.h"
+#include "VulkanDescriptorSets.h"
 
 static FCriticalSection GSamplerHashLock;
 
@@ -221,6 +222,7 @@ void FVulkanSamplerState::SetupSamplerCreateInfo(const FSamplerStateInitializerR
 FVulkanSamplerState::FVulkanSamplerState(const VkSamplerCreateInfo& InInfo, FVulkanDevice& InDevice, const bool bInIsImmutable)
 	: Sampler(VK_NULL_HANDLE)
 	, SamplerId(0)
+	, BindlessHandle()
 	, bIsImmutable(bInIsImmutable)
 {
 	VERIFYVULKANRESULT(VulkanRHI::vkCreateSampler(InDevice.GetInstanceHandle(), &InInfo, VULKAN_CPU_ALLOCATOR, &Sampler));
@@ -229,6 +231,11 @@ FVulkanSamplerState::FVulkanSamplerState(const VkSamplerCreateInfo& InInfo, FVul
 	{
 		SamplerId = ++GVulkanSamplerHandleIdCounter;
 	}
+
+	if (InDevice.SupportsBindless() && (RHIGetBindlessSamplersConfiguration(GMaxRHIShaderPlatform) != ERHIBindlessConfiguration::Disabled))
+	{
+		BindlessHandle = InDevice.GetBindlessDescriptorManager()->RegisterSampler(Sampler);
+	}
 }
 
 FVulkanRasterizerState::FVulkanRasterizerState(const FRasterizerStateInitializerRHI& InInitializer)
diff --git a/Engine/Source/Runtime/VulkanRHI/Private/VulkanState.h b/Engine/Source/Runtime/VulkanRHI/Private/VulkanState.h
--- a/Engine/Source/Runtime/VulkanRHI/Private/VulkanState.h
+++ b/Engine/Source/Runtime/VulkanRHI/Private/VulkanState.h
@@ -20,9 +20,12 @@ public:
 
 	VkSampler Sampler;
 	uint32 SamplerId;
+	FRHIDescriptorHandle BindlessHandle;
 
 	static void SetupSamplerCreateInfo(const FSamplerStateInitializerRHI& Initializer, FVulkanDevice& InDevice, VkSamplerCreateInfo& OutSamplerInfo);
 
+	virtual FRHIDescriptorHandle GetBindlessHandle() const override final { return BindlessHandle; }
+
 private:
 	bool bIsImmutable;
 };
diff --git a/Engine/Source/Runtime/VulkanRHI/Private/VulkanUniformBuffer.cpp b/Engine/Source/Runtime/VulkanRHI/Private/VulkanUniformBuffer.cpp
--- a/Engine/Source/Runtime/VulkanRHI/Private/VulkanUniformBuffer.cpp
+++ b/Engine/Source/Runtime/VulkanRHI/Private/VulkanUniformBuffer.cpp
@@ -28,6 +28,63 @@ enum
 	Uniform buffer RHI object
 -----------------------------------------------------------------------------*/
 
+
+
+// Utility class for reading shader parameters out of the data blob passed in.
+struct FUniformDataReader
+{
+	FUniformDataReader() = delete;
+	FUniformDataReader(const void* InData) : Data(reinterpret_cast<const uint8*>(InData)) { }
+
+	template<typename TResourceOut>
+	const TResourceOut& Read(const FRHIUniformBufferResource& InResource) const
+	{
+		return *reinterpret_cast<const TResourceOut*>(Data + InResource.MemberOffset);
+	}
+
+	const uint8* Data;
+};
+
+static void UpdateUniformBufferConstants(FVulkanDevice* Device, void* DestinationData, const void* SourceData, const FRHIUniformBufferLayout* Layout)
+{
+	check(DestinationData != nullptr);
+	check(SourceData != nullptr);
+
+	// First copy wholesale
+	FMemory::Memcpy(DestinationData, SourceData, Layout->ConstantBufferSize);
+
+	// Then copy indices over
+	if (Device->SupportsBindless())
+	{
+		FUniformDataReader Reader(SourceData);
+
+		for (const FRHIUniformBufferResource& Resource : Layout->Resources)
+		{
+			FRHIDescriptorHandle Handle;
+
+			if (Resource.MemberType == UBMT_SAMPLER)
+			{
+				FRHISamplerState* SamplerState = Reader.Read<FRHISamplerState*>(Resource);
+				if (SamplerState)
+				{
+					Handle = SamplerState->GetBindlessHandle();
+				}
+			}
+			else
+			{
+				// todo-jn: bindless resources : handle other types
+			}
+
+			if (Handle.IsValid())
+			{
+				const uint32 BindlessIndex = Handle.GetIndex();
+				FMemory::Memcpy(reinterpret_cast<uint8*>(DestinationData) + Resource.MemberOffset, &BindlessIndex, sizeof(BindlessIndex));
+			}
+		}
+	}
+}
+
+
 static inline EBufferUsageFlags UniformBufferToBufferUsage(EUniformBufferUsage Usage)
 {
 	switch (Usage)
@@ -50,10 +107,12 @@ static bool UseRingBuffer(EUniformBufferUsage Usage)
 	return (Usage == UniformBuffer_SingleDraw || Usage == UniformBuffer_SingleFrame);
 }
 
-static void UpdateUniformBufferHelper(FVulkanCommandListContext& Context, FVulkanUniformBuffer* VulkanUniformBuffer, int32 DataSize, const void* Data)
+static void UpdateUniformBufferHelper(FVulkanCommandListContext& Context, FVulkanUniformBuffer* VulkanUniformBuffer, const void* Data)
 {
 	FVulkanCmdBuffer* CmdBuffer = Context.GetCommandBufferManager()->GetActiveCmdBufferDirect();
 	
+	const int32 DataSize = VulkanUniformBuffer->GetLayout().ConstantBufferSize;
+
 	if (UseRingBuffer(VulkanUniformBuffer->Usage))
 	{
 		FVulkanUniformBufferUploader* UniformBufferUploader = Context.GetUniformBufferUploader();
@@ -72,7 +131,7 @@ static void UpdateUniformBufferHelper(FVulkanCommandListContext& Context, FVulka
 			RingBufferAllocation.HandleId);
 		
 		uint8* UploadLocation = UniformBufferUploader->GetCPUMappedPointer() + RingBufferOffset;
-		FMemory::Memcpy(UploadLocation, Data, DataSize);
+		UpdateUniformBufferConstants(Context.GetDevice(), UploadLocation, Data, VulkanUniformBuffer->GetLayoutPtr());
 	}
 	else
 	{
@@ -80,7 +139,8 @@ static void UpdateUniformBufferHelper(FVulkanCommandListContext& Context, FVulka
 
 		VulkanRHI::FTempFrameAllocationBuffer::FTempAllocInfo LockInfo;
 		Context.GetTempFrameAllocationBuffer().Alloc(DataSize, 16, LockInfo);
-		FMemory::Memcpy(LockInfo.Data, Data, DataSize);
+		UpdateUniformBufferConstants(Context.GetDevice(), LockInfo.Data, Data, VulkanUniformBuffer->GetLayoutPtr());
+
 		VkBufferCopy Region;
 		Region.size = DataSize;
 		Region.srcOffset = LockInfo.CurrentOffset + LockInfo.Allocation.Offset;
@@ -92,9 +152,9 @@ static void UpdateUniformBufferHelper(FVulkanCommandListContext& Context, FVulka
 	}
 };
 
-FVulkanUniformBuffer::FVulkanUniformBuffer(FVulkanDevice& Device, const FRHIUniformBufferLayout* InLayout, const void* Contents, EUniformBufferUsage InUsage, EUniformBufferValidation Validation)
+FVulkanUniformBuffer::FVulkanUniformBuffer(FVulkanDevice& InDevice, const FRHIUniformBufferLayout* InLayout, const void* Contents, EUniformBufferUsage InUsage, EUniformBufferValidation Validation)
 	: FRHIUniformBuffer(InLayout)
-	, Device(&Device) 
+	, Device(&InDevice)
 	, Usage(InUsage)
 {
 #if VULKAN_ENABLE_AGGRESSIVE_STATS
@@ -140,8 +200,8 @@ FVulkanUniformBuffer::FVulkanUniformBuffer(FVulkanDevice& Device, const FRHIUnif
 				const bool bCanAllocOnThisThread = RHICmdList.Bypass() || (!IsRunningRHIInSeparateThread() && bInRenderingThread) || bInRHIThread;
 				if (bCanAllocOnThisThread)
 				{
-					FVulkanCommandListContextImmediate& Context = Device.GetImmediateContext();
-					UpdateUniformBufferHelper(Context, UniformBuffer, DataSize, Contents);
+					FVulkanCommandListContextImmediate& Context = Device->GetImmediateContext();
+					UpdateUniformBufferHelper(Context, UniformBuffer, Contents);
 				}
 				else
 				{
@@ -151,7 +211,7 @@ FVulkanUniformBuffer::FVulkanUniformBuffer(FVulkanDevice& Device, const FRHIUnif
 					RHICmdList.EnqueueLambda([UniformBuffer, DataSize, CmdListConstantBufferData](FRHICommandList& CmdList)
 					{
 						FVulkanCommandListContext& Context = FVulkanCommandListContext::GetVulkanContext(CmdList.GetContext());
-						UpdateUniformBufferHelper(Context, UniformBuffer, DataSize, CmdListConstantBufferData);
+						UpdateUniformBufferHelper(Context, UniformBuffer, CmdListConstantBufferData);
 					});
 
 					RHICmdList.RHIThreadFence(true);
@@ -160,9 +220,14 @@ FVulkanUniformBuffer::FVulkanUniformBuffer(FVulkanDevice& Device, const FRHIUnif
 		}
 		else
 		{
-			VulkanRHI::FMemoryManager& ResourceMgr = Device.GetMemoryManager();
+			VulkanRHI::FMemoryManager& ResourceMgr = Device->GetMemoryManager();
 			// Set it directly as there is no previous one
-			ResourceMgr.AllocUniformBuffer(Allocation, InLayout->ConstantBufferSize, Contents);
+			ResourceMgr.AllocUniformBuffer(Allocation, InLayout->ConstantBufferSize);
+			if (Contents)
+			{
+				UpdateUniformBufferConstants(Device, Allocation.GetMappedPointer(Device), Contents, InLayout);
+				Allocation.FlushMappedMemory(Device);
+			}
 		}
 	}
 }
@@ -219,7 +284,12 @@ inline void FVulkanDynamicRHI::UpdateUniformBuffer(FRHICommandListBase& RHICmdLi
 		if (ConstantBufferSize > 0)
 		{
 			SCOPE_CYCLE_COUNTER(STAT_VulkanUpdateUniformBuffersRename);
-			Device->GetMemoryManager().AllocUniformBuffer(NewUBAlloc, ConstantBufferSize, Contents);
+			Device->GetMemoryManager().AllocUniformBuffer(NewUBAlloc, ConstantBufferSize);
+			if (Contents)
+			{
+				UpdateUniformBufferConstants(Device, NewUBAlloc.GetMappedPointer(Device), Contents, &Layout);
+				NewUBAlloc.FlushMappedMemory(Device);
+			}
 		}
 	}
 
@@ -231,7 +301,7 @@ inline void FVulkanDynamicRHI::UpdateUniformBuffer(FRHICommandListBase& RHICmdLi
 			if (bUseUpload || bUseRingBuffer)
 			{			
 				FVulkanCommandListContext& Context = Device->GetImmediateContext();
-				UpdateUniformBufferHelper(Context, UniformBuffer, ConstantBufferSize, Contents);
+				UpdateUniformBufferHelper(Context, UniformBuffer, Contents);
 			}
 			else
 			{
@@ -263,7 +333,7 @@ inline void FVulkanDynamicRHI::UpdateUniformBuffer(FRHICommandListBase& RHICmdLi
 			RHICmdList.EnqueueLambda([UniformBuffer, CmdListResources, NumResources, ConstantBufferSize, CmdListConstantBufferData](FRHICommandListBase& CmdList)
 			{
 				FVulkanCommandListContext& Context = (FVulkanCommandListContext&)CmdList.GetContext().GetLowestLevelContext();
-				UpdateUniformBufferHelper(Context, UniformBuffer, ConstantBufferSize, CmdListConstantBufferData);
+				UpdateUniformBufferHelper(Context, UniformBuffer, CmdListConstantBufferData);
 				UniformBuffer->UpdateResourceTable(CmdListResources, NumResources);
 			});
 		}
diff --git a/Engine/Source/Runtime/VulkanRHI/Public/VulkanCommon.h b/Engine/Source/Runtime/VulkanRHI/Public/VulkanCommon.h
--- a/Engine/Source/Runtime/VulkanRHI/Public/VulkanCommon.h
+++ b/Engine/Source/Runtime/VulkanRHI/Public/VulkanCommon.h
@@ -111,6 +111,22 @@ namespace ShaderStage
 	}
 };
 
+namespace VulkanBindless
+{
+	enum EDescriptorSets
+	{
+		BindlessSamplerSet = 0,
+		BindlessResourceSet = 1,
+
+		NumBindlessSets = 2,
+		MaxNumSets = ShaderStage::EStage::MaxNumSets + NumBindlessSets
+	};
+
+	// Prefix used to declare arrays of samplers/resources for bindless
+	static constexpr const TCHAR* kBindlessResourceArrayPrefix = TEXT("ResourceDescriptorHeap_");
+	static constexpr const TCHAR* kBindlessSamplerArrayPrefix = TEXT("SamplerDescriptorHeap_");
+};
+
 namespace EVulkanBindingType
 {
 	enum EType : uint8
