Ensure PSO size parameter is returned when using Android's out-of-process PSO service.
Fixes potential issues with vulkan PSO LRU.
silence some direct to logcat spam from the vulkan PSO compiler service
Add simple verification code to compare PSOs from android services vs engine.

diff --git a/Engine/Build/Android/Prebuilt/PSOService/Source/PSOService.cpp b/Engine/Build/Android/Prebuilt/PSOService/Source/PSOService.cpp
--- a/Engine/Build/Android/Prebuilt/PSOService/Source/PSOService.cpp
+++ b/Engine/Build/Android/Prebuilt/PSOService/Source/PSOService.cpp
@@ -76,7 +76,7 @@ public:
 		bool bEnableValidation = false;
 		for (uint32_t Idx; Idx < InstanceLayers.size(); ++Idx)
 		{
-			bEnableValidation = true;//strcmp(InstanceLayers[Idx], "VK_LAYER_KHRONOS_validation") == 0;
+			bEnableValidation = strcmp(InstanceLayers[Idx], "VK_LAYER_KHRONOS_validation") == 0;
 
 			if (bEnableValidation)
 			{
@@ -276,7 +276,7 @@ public:
 		GraphicsPipelineCreateInfo PipelineCreateInfo;
 		VkGraphicsPipelineCreateInfo CreateInfo;
 
-		__android_log_print(ANDROID_LOG_VERBOSE, APPNAME, "CompileGFXPSO: VSSize %d, PSSize %d, PSOSize %d", (uint32_t)VSSize, (uint32_t)PSSize, (uint32_t)PSOSize);
+		//__android_log_print(ANDROID_LOG_VERBOSE, APPNAME, "CompileGFXPSO: VSSize %d, PSSize %d, PSOSize %d", (uint32_t)VSSize, (uint32_t)PSSize, (uint32_t)PSOSize);
 
 		if (bSinglePSOPerCache && PSOCacheData)
 		{
diff --git a/Engine/Source/Runtime/VulkanRHI/Private/Android/VulkanAndroidPlatform.cpp b/Engine/Source/Runtime/VulkanRHI/Private/Android/VulkanAndroidPlatform.cpp
--- a/Engine/Source/Runtime/VulkanRHI/Private/Android/VulkanAndroidPlatform.cpp
+++ b/Engine/Source/Runtime/VulkanRHI/Private/Android/VulkanAndroidPlatform.cpp
@@ -1289,7 +1289,7 @@ VkPipelineCache FVulkanAndroidPlatform::PrecompilePSO(FVulkanDevice* Device, VkG
 					PipelineCacheCreateInfo.pInitialData = CompiledProgramBinary.GetData();
 					PipelineCacheCreateInfo.initialDataSize = CompiledProgramBinary.Num();
 					VERIFYVULKANRESULT(VulkanRHI::vkCreatePipelineCache(Device->GetInstanceHandle(), &PipelineCacheCreateInfo, VULKAN_CPU_ALLOCATOR, &PipelineCache));
-
+					AfterSize = CompiledProgramBinary.Num();
 					return PipelineCache;
 				}
 
diff --git a/Engine/Source/Runtime/VulkanRHI/Private/VulkanPipeline.cpp b/Engine/Source/Runtime/VulkanRHI/Private/VulkanPipeline.cpp
--- a/Engine/Source/Runtime/VulkanRHI/Private/VulkanPipeline.cpp
+++ b/Engine/Source/Runtime/VulkanRHI/Private/VulkanPipeline.cpp
@@ -1347,6 +1347,7 @@ VkResult FVulkanPipelineStateCacheManager::CreateVKPipeline(FVulkanRHIGraphicsPi
 	bool bWantPSOSize = false;
 	FGfxPipelineDesc* GfxEntry = &PSO->Desc;
 	uint64 ShaderHash = 0;
+	bool bValidateServicePSO = false;
 	if (bUseLRU)
 	{
 		QUICK_SCOPE_CYCLE_COUNTER(STAT_VulkanPSOLRUSizeLookup);
@@ -1386,6 +1387,17 @@ VkResult FVulkanPipelineStateCacheManager::CreateVKPipeline(FVulkanRHIGraphicsPi
 		{
 			Pipeline[0] = VK_NULL_HANDLE;
 			Result = VK_SUCCESS;
+
+			// enable bValidateServicePSO to compare PSOService result against engine's result.
+			//bValidateServicePSO = true;
+			if(bValidateServicePSO)
+			{
+				Result = VK_ERROR_INITIALIZATION_FAILED;
+				bWantPSOSize = true;
+				VulkanRHI::vkDestroyPipelineCache(Device->GetInstanceHandle(), LocalPipelineCache, VULKAN_CPU_ALLOCATOR);
+				LocalPipelineCache = VK_NULL_HANDLE;
+			}
+
 			if(bWantPSOSize)
 			{
 				PSOSize = AfterSize;
@@ -1408,6 +1420,15 @@ VkResult FVulkanPipelineStateCacheManager::CreateVKPipeline(FVulkanRHIGraphicsPi
 			ZeroVulkanStruct(PipelineCacheInfo, VK_STRUCTURE_TYPE_PIPELINE_CACHE_CREATE_INFO);
 			VERIFYVULKANRESULT(VulkanRHI::vkCreatePipelineCache(Device->GetInstanceHandle(), &PipelineCacheInfo, VULKAN_CPU_ALLOCATOR, &LocalPipelineCache));
 			Result = VulkanRHI::vkCreateGraphicsPipelines(Device->GetInstanceHandle(), LocalPipelineCache, 1, &PipelineInfo, VULKAN_CPU_ALLOCATOR, Pipeline);
+			if (bValidateServicePSO)
+			{
+				size_t Diff = 0;
+				if (ensure(LocalPipelineCache != VK_NULL_HANDLE))
+				{
+					VulkanRHI::vkGetPipelineCacheData(Device->GetInstanceHandle(), LocalPipelineCache, &Diff, nullptr);
+				}
+				UE_CLOG(Diff != PSOSize, LogVulkanRHI, Warning, TEXT("PSO service size mismatches engine size! [PSOService = %d, Game Process = %d]"), PSOSize, Diff);
+			}
 		}
 		else
 		{
