Silence a static analysis warning

diff --git a/Engine/Source/Runtime/VulkanRHI/Private/VulkanPipeline.cpp b/Engine/Source/Runtime/VulkanRHI/Private/VulkanPipeline.cpp
--- a/Engine/Source/Runtime/VulkanRHI/Private/VulkanPipeline.cpp
+++ b/Engine/Source/Runtime/VulkanRHI/Private/VulkanPipeline.cpp
@@ -1410,7 +1410,9 @@ VkResult FVulkanPipelineStateCacheManager::CreateVKPipeline(FVulkanRHIGraphicsPi
 		}
 	}
 #endif
-	if(Result != VK_SUCCESS)
+	// Disabling 'V547: Expression is always true'
+	// The precompile code above can set Result to success on Android platforms.
+	if(Result != VK_SUCCESS)  //-V547
 	{
 		if (bWantPSOSize)
 		{
