Vk: Attempt to initialize PF_R64_UINT only for SM5+ feature level

diff --git a/Engine/Source/Runtime/VulkanRHI/Private/VulkanDevice.cpp b/Engine/Source/Runtime/VulkanRHI/Private/VulkanDevice.cpp
--- a/Engine/Source/Runtime/VulkanRHI/Private/VulkanDevice.cpp
+++ b/Engine/Source/Runtime/VulkanRHI/Private/VulkanDevice.cpp
@@ -601,18 +601,22 @@ void FVulkanDevice::SetupFormats()
 	MapFormatSupport(PF_R8_SINT, { VK_FORMAT_R8_SINT }, ComponentMappingR001);
 
 	// This will be the format used for 64bit image atomics
+	// This format is SM5 only, skip it for mobile to not confuse QA with a logged error about missing pixel format
+	if (GMaxRHIFeatureLevel >= ERHIFeatureLevel::SM5)
+	{
 #if VULKAN_HAS_DEBUGGING_ENABLED
-	const EPixelFormatCapabilities RequiredCaps64U = GRenderDocFound ? EPixelFormatCapabilities::UAV : (EPixelFormatCapabilities::UAV | EPixelFormatCapabilities::TextureAtomics);
+		const EPixelFormatCapabilities RequiredCaps64U = GRenderDocFound ? EPixelFormatCapabilities::UAV : (EPixelFormatCapabilities::UAV | EPixelFormatCapabilities::TextureAtomics);
 #else
-	const EPixelFormatCapabilities RequiredCaps64U = (EPixelFormatCapabilities::UAV | EPixelFormatCapabilities::TextureAtomics);
+		const EPixelFormatCapabilities RequiredCaps64U = (EPixelFormatCapabilities::UAV | EPixelFormatCapabilities::TextureAtomics);
 #endif
-	MapFormatSupport(PF_R64_UINT, { VK_FORMAT_R64_UINT, VK_FORMAT_R32G32_UINT }, ComponentMappingR001, RequiredCaps64U);
-	// Shaders were patched to use UAV, make sure we don't expose texture sampling
-	GPixelFormats[PF_R64_UINT].Capabilities &= ~(EPixelFormatCapabilities::AnyTexture | EPixelFormatCapabilities::TextureSample);
-	if (GRHISupportsAtomicUInt64 && !EnumHasAnyFlags(GPixelFormats[PF_R64_UINT].Capabilities, EPixelFormatCapabilities::UAV))
-	{
-		UE_LOG(LogVulkanRHI, Warning, TEXT("64bit image atomics were enabled, but the R64 format does not have UAV capabilities.  Disabling support."));
-		GRHISupportsAtomicUInt64 = false;
+		MapFormatSupport(PF_R64_UINT, { VK_FORMAT_R64_UINT, VK_FORMAT_R32G32_UINT }, ComponentMappingR001, RequiredCaps64U);
+		// Shaders were patched to use UAV, make sure we don't expose texture sampling
+		GPixelFormats[PF_R64_UINT].Capabilities &= ~(EPixelFormatCapabilities::AnyTexture | EPixelFormatCapabilities::TextureSample);
+		if (GRHISupportsAtomicUInt64 && !EnumHasAnyFlags(GPixelFormats[PF_R64_UINT].Capabilities, EPixelFormatCapabilities::UAV))
+		{
+			UE_LOG(LogVulkanRHI, Warning, TEXT("64bit image atomics were enabled, but the R64 format does not have UAV capabilities.  Disabling support."));
+			GRHISupportsAtomicUInt64 = false;
+		}
 	}
 
 	if (CVarVulkanUseD24.GetValueOnAnyThread() != 0)
