Fixed: Frame stuterring on devices wtih a non standard 'Motion Smoothness' on

diff --git a/Engine/Build/Android/Java/src/com/epicgames/unreal/GameActivity.java.template b/Engine/Build/Android/Java/src/com/epicgames/unreal/GameActivity.java.template
--- a/Engine/Build/Android/Java/src/com/epicgames/unreal/GameActivity.java.template
+++ b/Engine/Build/Android/Java/src/com/epicgames/unreal/GameActivity.java.template
@@ -5391,7 +5391,22 @@ public class GameActivity extends $${gameActivitySuperClass}$$ implements Surfac
 
 	public boolean AndroidThunkJava_SetNativeDisplayRefreshRate(int RefreshRate)
 	{
-		if(ANDROID_BUILD_VERSION >= 24)
+		if (ANDROID_BUILD_VERSION >= 30)
+		{
+			final int FinalRefreshRate = RefreshRate;
+			_activity.runOnUiThread(new Runnable()
+			{
+				@Override
+				public void run()
+				{
+					MySurfaceView.getHolder().getSurface().setFrameRate(FinalRefreshRate, android.view.Surface.FRAME_RATE_COMPATIBILITY_FIXED_SOURCE);
+				}
+			});
+			Log.debug("SetNativeDisplayRefreshRate - setFrameRate " + RefreshRate);
+
+			return true;
+		}
+		else if(ANDROID_BUILD_VERSION >= 24)
 		{
 			WindowManager windowManager =  getWindowManager();
 			Display display = windowManager.getDefaultDisplay();
diff --git a/Engine/Source/Runtime/VulkanRHI/Private/Android/VulkanAndroidPlatform.cpp b/Engine/Source/Runtime/VulkanRHI/Private/Android/VulkanAndroidPlatform.cpp
--- a/Engine/Source/Runtime/VulkanRHI/Private/Android/VulkanAndroidPlatform.cpp
+++ b/Engine/Source/Runtime/VulkanRHI/Private/Android/VulkanAndroidPlatform.cpp
@@ -494,9 +494,10 @@ bool FVulkanAndroidPlatform::FramePace(FVulkanDevice& Device, VkSwapchainKHR Swa
 	int32 CurrentFramePace = FAndroidPlatformRHIFramePacer::GetFramePace();
 	if (CurrentFramePace != 0)
 	{
-		int32 CurrentRefreshRate = FAndroidMisc::GetNativeDisplayRefreshRate();
-
-		bool bRefreshRateInvalid = (CurrentRefreshRate != CachedRefreshRate);
+		CachedRefreshRate = FAndroidMisc::GetNativeDisplayRefreshRate();
+		
+		// we don't need to change display refresh rate as long as it's not lower than game frame pace 
+		bool bRefreshRateInvalid = (CachedRefreshRate < CurrentFramePace);
 		bool bTryChangingRefreshRate = (bRefreshRateInvalid && (SuccessfulRefreshRateFrames > 0 || UnsuccessfulRefreshRateFrames > 1000));
 
 		if (bRefreshRateInvalid)
@@ -515,16 +516,15 @@ bool FVulkanAndroidPlatform::FramePace(FVulkanDevice& Device, VkSwapchainKHR Swa
 		// or periodically if not successfully running at the desired rate
 		if (CurrentFramePace != CachedFramePace || bTryChangingRefreshRate)
 		{
+			UE_LOG(LogRHI, Log, TEXT("Requesting a new display refresh rate - %d Hz"), CurrentFramePace);
+			
 			CachedFramePace = CurrentFramePace;
-			if (FramePacer->SupportsFramePaceInternal(CurrentFramePace, CachedRefreshRate, CachedSyncInterval))
-			{
-				FAndroidMisc::SetNativeDisplayRefreshRate(CachedRefreshRate);
-			}
-			else
-			{
-				// Desired frame pace not supported, save current refresh rate to prevent logspam.
-				CachedRefreshRate = CurrentRefreshRate;
-			}
+			FAndroidMisc::SetNativeDisplayRefreshRate(CurrentFramePace);
+						
+			// assume device succesfully sets desired display refresh rate
+			// actual refresh rate could be higher than requested, we will query exact value on a next frame
+			CachedRefreshRate = CurrentFramePace;
+			
 			UnsuccessfulRefreshRateFrames = 0;
 			SuccessfulRefreshRateFrames = 0;
 		}
@@ -570,6 +570,10 @@ VkResult FVulkanAndroidPlatform::CreateSwapchainKHR(VkDevice Device, const VkSwa
 void FVulkanAndroidPlatform::DestroySwapchainKHR(VkDevice Device, VkSwapchainKHR Swapchain, const VkAllocationCallbacks* Allocator)
 {
 	VulkanRHI::vkDestroySwapchainKHR(Device, Swapchain, Allocator);
+
+	// reset frame pace, to force display refresh rate update after we create a new swapchain
+	// see FVulkanAndroidPlatform::FramePace
+	CachedFramePace = 0;
 }
 
 
