Invoke the Android Vulkan PSO compiler services only when the binary cache needs creating.

diff --git a/Engine/Source/Runtime/VulkanRHI/Private/VulkanPipeline.cpp b/Engine/Source/Runtime/VulkanRHI/Private/VulkanPipeline.cpp
--- a/Engine/Source/Runtime/VulkanRHI/Private/VulkanPipeline.cpp
+++ b/Engine/Source/Runtime/VulkanRHI/Private/VulkanPipeline.cpp
@@ -562,17 +562,17 @@ void FVulkanPipelineStateCacheManager::OnShaderPipelineCacheOpened(FString const
 		UE_LOG(LogVulkanRHI, Log, TEXT("FVulkanPipelineStateCacheManager: %s does not exist."), *CompiledPSOCacheFolderName);
 	}
 
-#if PLATFORM_ANDROID
-	if (GNumRemoteProgramCompileServices)
-	{
-		FVulkanAndroidPlatform::StartAndWaitForRemoteCompileServices(GNumRemoteProgramCompileServices);
-	}
-#endif
+
 
 	if (!bPrecompilingCacheLoadedFromFile || (bEvictImmediately && CVarPipelineLRUCacheEvictBinaryPreloadScreen.GetValueOnAnyThread()))
 	{
-		//TODO: This!
 		ShaderCachePrecompileContext.SetPrecompilationIsSlowTask();
+#if PLATFORM_ANDROID
+		if (GNumRemoteProgramCompileServices)
+		{
+			FVulkanAndroidPlatform::StartAndWaitForRemoteCompileServices(GNumRemoteProgramCompileServices);
+		}
+#endif
 	}
 }
 
@@ -581,7 +581,7 @@ void FVulkanPipelineStateCacheManager::OnShaderPipelineCachePrecompilationComple
 	UE_LOG(LogVulkanRHI, Log, TEXT("FVulkanPipelineStateCacheManager::OnShaderPipelineCachePrecompilationComplete"));
 
 #if PLATFORM_ANDROID
-	if (GNumRemoteProgramCompileServices)
+	if (FVulkanAndroidPlatform::AreRemoteCompileServicesActive())
 	{
 		FVulkanAndroidPlatform::StopRemoteCompileServices();
 	}
