Make Android Compiler service failed JNI bindings a fatal error message.

diff --git a/Engine/Source/Runtime/OpenGLDrv/Private/Android/AndroidOpenGL.cpp b/Engine/Source/Runtime/OpenGLDrv/Private/Android/AndroidOpenGL.cpp
--- a/Engine/Source/Runtime/OpenGLDrv/Private/Android/AndroidOpenGL.cpp
+++ b/Engine/Source/Runtime/OpenGLDrv/Private/Android/AndroidOpenGL.cpp
@@ -144,7 +144,7 @@ struct FOpenGLRemoteGLProgramCompileJNI
 		}
 
 		bAllFound = OGLServiceAccessor && DispatchProgramLink && StartRemoteProgramLink && StopRemoteProgramLink && ProgramResponseClass && ProgramResponse_SuccessField && ProgramResponse_CompiledBinaryField && ProgramResponse_ErrorField;
-		UE_CLOG(!bAllFound, LogRHI, Error, TEXT("Failed to find JNI GL remote program compiler."));
+		UE_CLOG(!bAllFound, LogRHI, Fatal, TEXT("Failed to find JNI GL remote program compiler."));
 	}
 }OpenGLRemoteGLProgramCompileJNI;
 
diff --git a/Engine/Source/Runtime/VulkanRHI/Private/Android/VulkanAndroidPlatform.cpp b/Engine/Source/Runtime/VulkanRHI/Private/Android/VulkanAndroidPlatform.cpp
--- a/Engine/Source/Runtime/VulkanRHI/Private/Android/VulkanAndroidPlatform.cpp
+++ b/Engine/Source/Runtime/VulkanRHI/Private/Android/VulkanAndroidPlatform.cpp
@@ -1170,7 +1170,7 @@ struct FVKRemoteProgramCompileJNI
 		}
 
 		bAllFound = PSOServiceAccessor && DispatchPSOCompile && StartRemoteProgramLink && StopRemoteProgramLink && ProgramResponseClass && ProgramResponse_SuccessField && ProgramResponse_CompiledBinaryField && ProgramResponse_ErrorField;
-		UE_CLOG(!bAllFound, LogRHI, Error, TEXT("Failed to find JNI Vulkan remote program compiler."));
+		UE_CLOG(!bAllFound, LogRHI, Fatal, TEXT("Failed to find JNI Vulkan remote program compiler."));
 	}
 }VKRemoteProgramCompileJNI;
 
