Restore changes lost in platform RHI implementations when the DevPR copy-up robomerged back to UE5 Main

diff --git a/Engine/Source/Runtime/Apple/MetalRHI/Private/MetalRHI.cpp b/Engine/Source/Runtime/Apple/MetalRHI/Private/MetalRHI.cpp
--- a/Engine/Source/Runtime/Apple/MetalRHI/Private/MetalRHI.cpp
+++ b/Engine/Source/Runtime/Apple/MetalRHI/Private/MetalRHI.cpp
@@ -716,7 +716,8 @@ FMetalDynamicRHI::FMetalDynamicRHI(ERHIFeatureLevel::Type RequestedFeatureLevel)
 	GMetalBufferFormats[PF_R8_SINT				] = { mtlpp::PixelFormat::R8Sint, (uint8)EMetalBufferFormat::R8Sint };
 	GMetalBufferFormats[PF_R64_UINT				] = { mtlpp::PixelFormat::Invalid, (uint8)EMetalBufferFormat::Unknown };
 	GMetalBufferFormats[PF_R9G9B9EXP5			] = { mtlpp::PixelFormat::Invalid, (uint8)EMetalBufferFormat::Unknown };
-	static_assert(PF_MAX == 86, "Please setup GMetalBufferFormats properly for the new pixel format");
+	GMetalBufferFormats[PF_P010					] = { mtlpp::PixelFormat::Invalid, (uint8)EMetalBufferFormat::Unknown };
+	static_assert(PF_MAX == 87, "Please setup GMetalBufferFormats properly for the new pixel format");
 
 	// Initialize the platform pixel format map.
 	GPixelFormats[PF_Unknown			].PlatformFormat	= (uint32)mtlpp::PixelFormat::Invalid;
diff --git a/Engine/Source/Runtime/D3D12RHI/Private/D3D12RHI.cpp b/Engine/Source/Runtime/D3D12RHI/Private/D3D12RHI.cpp
--- a/Engine/Source/Runtime/D3D12RHI/Private/D3D12RHI.cpp
+++ b/Engine/Source/Runtime/D3D12RHI/Private/D3D12RHI.cpp
@@ -185,6 +185,9 @@ FD3D12DynamicRHI::FD3D12DynamicRHI(const TArray<TSharedPtr<FD3D12Adapter>>& Chos
 
 	GPixelFormats[PF_R9G9B9EXP5	    ].PlatformFormat = DXGI_FORMAT_R9G9B9E5_SHAREDEXP;
 
+	GPixelFormats[PF_P010			].PlatformFormat = DXGI_FORMAT_P010;
+	GPixelFormats[PF_P010			].Supported = true;
+
 	// MS - Not doing any feature level checks. D3D12 currently supports these limits.
 	// However this may need to be revisited if new feature levels are introduced with different HW requirement
 	GSupportsSeparateRenderTargetBlendState = true;
diff --git a/Engine/Source/Runtime/D3D12RHI/Private/D3D12Texture.cpp b/Engine/Source/Runtime/D3D12RHI/Private/D3D12Texture.cpp
--- a/Engine/Source/Runtime/D3D12RHI/Private/D3D12Texture.cpp
+++ b/Engine/Source/Runtime/D3D12RHI/Private/D3D12Texture.cpp
@@ -470,6 +470,13 @@ void FD3D12TextureStats::D3D12TextureDeleted(FD3D12Texture& Texture)
 
 bool FD3D12Texture::CanBe4KAligned(const FD3D12ResourceDesc& Desc, EPixelFormat UEFormat)
 {
+	// Exclude video related formats
+	if (UEFormat == PF_NV12 ||
+		UEFormat == PF_P010)
+	{
+		return false;
+	}
+
 	// 4KB alignment is only available for read only textures
 	if (!EnumHasAnyFlags(Desc.Flags, D3D12_RESOURCE_FLAG_ALLOW_RENDER_TARGET | D3D12_RESOURCE_FLAG_ALLOW_DEPTH_STENCIL | D3D12_RESOURCE_FLAG_ALLOW_UNORDERED_ACCESS) &&
 		!Desc.NeedsUAVAliasWorkarounds() && // UAV aliased resources are secretly writable.
@@ -1425,7 +1432,8 @@ void FD3D12Texture::CreateViews()
 		bCreateShaderResource = false;
 	}
 
-	if (Desc.Format == PF_NV12)
+	if (Desc.Format == PF_NV12 ||
+		Desc.Format == PF_P010)
 	{
 		bCreateRTV = false;
 		bCreateShaderResource = false;
diff --git a/Engine/Source/Runtime/D3D12RHI/Private/Windows/WindowsD3D12Viewport.cpp b/Engine/Source/Runtime/D3D12RHI/Private/Windows/WindowsD3D12Viewport.cpp
--- a/Engine/Source/Runtime/D3D12RHI/Private/Windows/WindowsD3D12Viewport.cpp
+++ b/Engine/Source/Runtime/D3D12RHI/Private/Windows/WindowsD3D12Viewport.cpp
@@ -33,7 +33,6 @@ FD3D12Viewport::FD3D12Viewport(class FD3D12Adapter* InParent, HWND InWindowHandl
 	, bFullscreenLost(false)
 	, PixelFormat(InPreferredPixelFormat)
 	, bIsValid(true)
-	, bHDRMetaDataSet(false)
 	, ColorSpace(DXGI_COLOR_SPACE_RGB_FULL_G22_NONE_P709)
 	, NumBackBuffers(WindowsDefaultNumBackBuffers)
 	, DummyBackBuffer_RenderThread(nullptr)
@@ -397,20 +396,6 @@ void FD3D12Viewport::EnableHDR()
 {
 	if ( GRHISupportsHDROutput && IsHDREnabled() )
 	{
-		const float DisplayMaxOutputNits = HDRGetDisplayMaximumLuminance();
-		const float DisplayMinOutputNits = 0.0f;	// Min output of the display
-		const float DisplayMaxCLL = 0.0f;			// Max content light level in lumens (0.0 == unknown)
-		const float DisplayFALL = 0.0f;				// Frame average light level (0.0 == unknown)
-
-		// Ideally we can avoid setting TV meta data and instead the engine can do tone mapping based on the
-		// actual current display properties (display mapping).
-		SetHDRTVMode(true,
-			DisplayColorGamut,
-			DisplayMaxOutputNits,
-			DisplayMinOutputNits,
-			DisplayMaxCLL,
-			DisplayFALL);
-
 		// Ensure we have the correct color space set.
 		EnsureColorSpace(DisplayColorGamut, DisplayOutputFormat);
 	}
@@ -420,28 +405,8 @@ void FD3D12Viewport::ShutdownHDR()
 {
 	if (GRHISupportsHDROutput)
 	{
-		// Default SDR display data
-		const EDisplayColorGamut DisplayGamut = EDisplayColorGamut::sRGB_D65;
-		const EDisplayOutputFormat OutputDevice = EDisplayOutputFormat::SDR_sRGB;
-
-		// Note: These values aren't actually used.
-		const float DisplayMaxOutputNits = 100.0f;	// Max output of the display
-		const float DisplayMinOutputNits = 0.0f;	// Min output of the display
-		const float DisplayMaxCLL = 100.0f;			// Max content light level in lumens
-		const float DisplayFALL = 20.0f;			// Frame average light level
-
-		// Ideally we can avoid setting TV meta data and instead the engine can do tone mapping based on the
-		// actual current display properties (display mapping).
-		SetHDRTVMode(
-			false,
-			DisplayGamut,
-			DisplayMaxOutputNits,
-			DisplayMinOutputNits,
-			DisplayMaxCLL,
-			DisplayFALL);
-
 		// Ensure we have the correct color space set.
-		EnsureColorSpace(DisplayGamut, OutputDevice);
+		EnsureColorSpace(EDisplayColorGamut::sRGB_D65, EDisplayOutputFormat::SDR_sRGB);
 	}
 }
 
@@ -552,59 +517,6 @@ void FD3D12Viewport::EnsureColorSpace(EDisplayColorGamut DisplayGamut, EDisplayO
 	}
 }
 
-void FD3D12Viewport::SetHDRTVMode(bool bEnableHDR, EDisplayColorGamut DisplayGamut, float MaxOutputNits, float MinOutputNits, float MaxCLL, float MaxFALL)
-{
-	ensure(SwapChain4.GetReference());
-
-	static const DisplayChromacities DisplayChromacityList[] =
-	{
-		{ 0.64000f, 0.33000f, 0.30000f, 0.60000f, 0.15000f, 0.06000f, 0.31270f, 0.32900f }, // EDisplayColorGamut::sRGB_D65
-		{ 0.68000f, 0.32000f, 0.26500f, 0.69000f, 0.15000f, 0.06000f, 0.31270f, 0.32900f }, // EDisplayColorGamut::DCIP3_D65
-		{ 0.70800f, 0.29200f, 0.17000f, 0.79700f, 0.13100f, 0.04600f, 0.31270f, 0.32900f }, // EDisplayColorGamut::Rec2020_D65
-		{ 0.73470f, 0.26530f, 0.00000f, 1.00000f, 0.00010f,-0.07700f, 0.32168f, 0.33767f }, // EDisplayColorGamut::ACES_D60
-		{ 0.71300f, 0.29300f, 0.16500f, 0.83000f, 0.12800f, 0.04400f, 0.32168f, 0.33767f }, // EDisplayColorGamut::ACEScg_D60
-	};
-
-	if (bEnableHDR)
-	{
-		const DisplayChromacities& Chroma = DisplayChromacityList[(int32)DisplayGamut];
-
-		// Set HDR meta data
-		DXGI_HDR_METADATA_HDR10 HDR10MetaData = {};
-		HDR10MetaData.RedPrimary[0] = static_cast<uint16>(Chroma.RedX * 50000.0f);
-		HDR10MetaData.RedPrimary[1] = static_cast<uint16>(Chroma.RedY * 50000.0f);
-		HDR10MetaData.GreenPrimary[0] = static_cast<uint16>(Chroma.GreenX * 50000.0f);
-		HDR10MetaData.GreenPrimary[1] = static_cast<uint16>(Chroma.GreenY * 50000.0f);
-		HDR10MetaData.BluePrimary[0] = static_cast<uint16>(Chroma.BlueX * 50000.0f);
-		HDR10MetaData.BluePrimary[1] = static_cast<uint16>(Chroma.BlueY * 50000.0f);
-		HDR10MetaData.WhitePoint[0] = static_cast<uint16>(Chroma.WpX * 50000.0f);
-		HDR10MetaData.WhitePoint[1] = static_cast<uint16>(Chroma.WpY * 50000.0f);
-		HDR10MetaData.MaxMasteringLuminance = static_cast<uint32>(MaxOutputNits * 10000.0f);
-		HDR10MetaData.MinMasteringLuminance = static_cast<uint32>(MinOutputNits * 10000.0f);
-		HDR10MetaData.MaxContentLightLevel = static_cast<uint16>(MaxCLL);
-		HDR10MetaData.MaxFrameAverageLightLevel = static_cast<uint16>(MaxFALL);
-
-		VERIFYD3D12RESULT(SwapChain4->SetHDRMetaData(DXGI_HDR_METADATA_TYPE_HDR10, sizeof(HDR10MetaData), &HDR10MetaData));
-		UE_LOG(LogD3D12RHI, Log, TEXT("Setting HDR meta data on swap chain (%#016llx) using DisplayGamut %u:"), SwapChain4.GetReference(), static_cast<uint32>(DisplayGamut));
-		UE_LOG(LogD3D12RHI, Log, TEXT("\t\tMaxMasteringLuminance = %.4f nits"), HDR10MetaData.MaxMasteringLuminance * .0001f);
-		UE_LOG(LogD3D12RHI, Log, TEXT("\t\tMinMasteringLuminance = %.4f nits"), HDR10MetaData.MinMasteringLuminance * .0001f);
-		UE_LOG(LogD3D12RHI, Log, TEXT("\t\tMaxContentLightLevel = %u nits"), HDR10MetaData.MaxContentLightLevel);
-		UE_LOG(LogD3D12RHI, Log, TEXT("\t\tMaxFrameAverageLightLevel %u = nits"), HDR10MetaData.MaxFrameAverageLightLevel);
-		bHDRMetaDataSet = true;
-	}
-	else
-	{
-		if (bHDRMetaDataSet)
-		{
-			// Clear meta data.
-			VERIFYD3D12RESULT(SwapChain4->SetHDRMetaData(DXGI_HDR_METADATA_TYPE_NONE, 0, nullptr));
-			UE_LOG(LogD3D12RHI, Log, TEXT("Clearing HDR meta data on swap chain (%#016llx)."), SwapChain4.GetReference());
-			bHDRMetaDataSet = false;
-		}
-	}
-}
-
-
 void FD3D12Viewport::OnResumeRendering()
 {}
 
diff --git a/Engine/Source/Runtime/D3D12RHI/Public/D3D12Viewport.h b/Engine/Source/Runtime/D3D12RHI/Public/D3D12Viewport.h
--- a/Engine/Source/Runtime/D3D12RHI/Public/D3D12Viewport.h
+++ b/Engine/Source/Runtime/D3D12RHI/Public/D3D12Viewport.h
@@ -199,7 +199,6 @@ private:
 #endif
 
 #if PLATFORM_WINDOWS || PLATFORM_HOLOLENS
-	bool bHDRMetaDataSet;
 	DXGI_COLOR_SPACE_TYPE ColorSpace;
 #endif
 #endif // D3D12_VIEWPORT_EXPOSES_SWAP_CHAIN
@@ -268,13 +267,6 @@ private:
 #if PLATFORM_WINDOWS || PLATFORM_HOLOLENS
 	/** Ensure the correct color space is set on the swap chain */
 	void EnsureColorSpace(EDisplayColorGamut DisplayGamut, EDisplayOutputFormat OutputDevice);
-
-	/** 
-	 * Set HDR meta data. 
-	 * Note: Meta data should only be provided for TVs, not monitors. 
-	 * This is because the TV is doing the work to display the colors correctly.
-	 */
-	void SetHDRTVMode(bool bEnableHDR, EDisplayColorGamut DisplayGamut, float MaxOutputNits, float MinOutputNits, float MaxCLL, float MaxFALL);
 #endif
 };
 
diff --git a/Engine/Source/Runtime/VulkanRHI/Private/VulkanMemory.cpp b/Engine/Source/Runtime/VulkanRHI/Private/VulkanMemory.cpp
--- a/Engine/Source/Runtime/VulkanRHI/Private/VulkanMemory.cpp
+++ b/Engine/Source/Runtime/VulkanRHI/Private/VulkanMemory.cpp
@@ -3569,19 +3569,13 @@ namespace VulkanRHI
 		}
 	}
 
-	void FMemoryManager::AllocUniformBuffer(FVulkanAllocation& OutAllocation, uint32 Size, const void* Contents)
+	void FMemoryManager::AllocUniformBuffer(FVulkanAllocation& OutAllocation, uint32 Size)
 	{
 		if(!AllocateBufferPooled(OutAllocation, nullptr, Size, 0, VK_BUFFER_USAGE_UNIFORM_BUFFER_BIT | VK_BUFFER_USAGE_TRANSFER_DST_BIT, VK_MEMORY_PROPERTY_HOST_VISIBLE_BIT | VK_MEMORY_PROPERTY_HOST_COHERENT_BIT, EVulkanAllocationMetaUniformBuffer, __FILE__, __LINE__))
 		{
 			HandleOOM(false);
 			checkNoEntry();
 		}
-
-		if (Contents)
-		{
-			FMemory::Memcpy(OutAllocation.GetMappedPointer(Device), Contents, Size);
-			OutAllocation.FlushMappedMemory(Device);
-		}
 	}
 	void FMemoryManager::FreeUniformBuffer(FVulkanAllocation& InAllocation)
 	{
diff --git a/Engine/Source/Runtime/VulkanRHI/Private/VulkanMemory.h b/Engine/Source/Runtime/VulkanRHI/Private/VulkanMemory.h
--- a/Engine/Source/Runtime/VulkanRHI/Private/VulkanMemory.h
+++ b/Engine/Source/Runtime/VulkanRHI/Private/VulkanMemory.h
@@ -797,7 +797,7 @@ namespace VulkanRHI
 		void DumpMemory(bool bFullDump = true);
 
 
-		void AllocUniformBuffer(FVulkanAllocation& OutAllocation, uint32 Size, const void* Contents);
+		void AllocUniformBuffer(FVulkanAllocation& OutAllocation, uint32 Size);
 		void FreeUniformBuffer(FVulkanAllocation& InAllocation);
 
 		void HandleOOM(bool bCanResume = false, VkResult Result = VK_SUCCESS, uint64 AllocationSize = 0, uint32 MemoryTypeIndex = 0);
diff --git a/Engine/Source/Runtime/VulkanRHI/Private/VulkanRHI.cpp b/Engine/Source/Runtime/VulkanRHI/Private/VulkanRHI.cpp
--- a/Engine/Source/Runtime/VulkanRHI/Private/VulkanRHI.cpp
+++ b/Engine/Source/Runtime/VulkanRHI/Private/VulkanRHI.cpp
@@ -845,6 +845,7 @@ void FVulkanDynamicRHI::InitInstance()
 		GRHISupportsBaseVertexIndex = true;
 		GSupportsSeparateRenderTargetBlendState = true;
 		GRHISupportsSeparateDepthStencilCopyAccess = false;
+		GRHISupportsBindless =	RHISupportsBindless(GMaxRHIShaderPlatform) && Device->SupportsBindless();
 
 		GRHIMaxDispatchThreadGroupsPerDimension.X = FMath::Min<uint32>(Limits.maxComputeWorkGroupCount[0], 0x7fffffff);
 		GRHIMaxDispatchThreadGroupsPerDimension.Y = FMath::Min<uint32>(Limits.maxComputeWorkGroupCount[1], 0x7fffffff);
diff --git a/Engine/Source/Runtime/VulkanRHI/Private/VulkanUniformBuffer.cpp b/Engine/Source/Runtime/VulkanRHI/Private/VulkanUniformBuffer.cpp
--- a/Engine/Source/Runtime/VulkanRHI/Private/VulkanUniformBuffer.cpp
+++ b/Engine/Source/Runtime/VulkanRHI/Private/VulkanUniformBuffer.cpp
@@ -28,6 +28,63 @@ enum
 	Uniform buffer RHI object
 -----------------------------------------------------------------------------*/
 
+
+
+// Utility class for reading shader parameters out of the data blob passed in.
+struct FUniformDataReader
+{
+	FUniformDataReader() = delete;
+	FUniformDataReader(const void* InData) : Data(reinterpret_cast<const uint8*>(InData)) { }
+
+	template<typename TResourceOut>
+	const TResourceOut& Read(const FRHIUniformBufferResource& InResource) const
+	{
+		return *reinterpret_cast<const TResourceOut*>(Data + InResource.MemberOffset);
+	}
+
+	const uint8* Data;
+};
+
+static void UpdateUniformBufferConstants(FVulkanDevice* Device, void* DestinationData, const void* SourceData, const FRHIUniformBufferLayout* Layout)
+{
+	check(DestinationData != nullptr);
+	check(SourceData != nullptr);
+
+	// First copy wholesale
+	FMemory::Memcpy(DestinationData, SourceData, Layout->ConstantBufferSize);
+
+	// Then copy indices over
+	if (Device->SupportsBindless())
+	{
+		FUniformDataReader Reader(SourceData);
+
+		for (const FRHIUniformBufferResource& Resource : Layout->Resources)
+		{
+			FRHIDescriptorHandle Handle;
+
+			if (Resource.MemberType == UBMT_SAMPLER)
+			{
+				FRHISamplerState* SamplerState = Reader.Read<FRHISamplerState*>(Resource);
+				if (SamplerState)
+				{
+					Handle = SamplerState->GetBindlessHandle();
+				}
+			}
+			else
+			{
+				// todo-jn: bindless resources : handle other types
+			}
+
+			if (Handle.IsValid())
+			{
+				const uint32 BindlessIndex = Handle.GetIndex();
+				FMemory::Memcpy(reinterpret_cast<uint8*>(DestinationData) + Resource.MemberOffset, &BindlessIndex, sizeof(BindlessIndex));
+			}
+		}
+	}
+}
+
+
 static inline EBufferUsageFlags UniformBufferToBufferUsage(EUniformBufferUsage Usage)
 {
 	switch (Usage)
@@ -50,10 +107,12 @@ static bool UseRingBuffer(EUniformBufferUsage Usage)
 	return (Usage == UniformBuffer_SingleDraw || Usage == UniformBuffer_SingleFrame);
 }
 
-static void UpdateUniformBufferHelper(FVulkanCommandListContext& Context, FVulkanUniformBuffer* VulkanUniformBuffer, int32 DataSize, const void* Data)
+static void UpdateUniformBufferHelper(FVulkanCommandListContext& Context, FVulkanUniformBuffer* VulkanUniformBuffer, const void* Data)
 {
 	FVulkanCmdBuffer* CmdBuffer = Context.GetCommandBufferManager()->GetActiveCmdBufferDirect();
 	
+	const int32 DataSize = VulkanUniformBuffer->GetLayout().ConstantBufferSize;
+
 	if (UseRingBuffer(VulkanUniformBuffer->Usage))
 	{
 		FVulkanUniformBufferUploader* UniformBufferUploader = Context.GetUniformBufferUploader();
@@ -72,7 +131,7 @@ static void UpdateUniformBufferHelper(FVulkanCommandListContext& Context, FVulka
 			RingBufferAllocation.HandleId);
 		
 		uint8* UploadLocation = UniformBufferUploader->GetCPUMappedPointer() + RingBufferOffset;
-		FMemory::Memcpy(UploadLocation, Data, DataSize);
+		UpdateUniformBufferConstants(Context.GetDevice(), UploadLocation, Data, VulkanUniformBuffer->GetLayoutPtr());
 	}
 	else
 	{
@@ -80,7 +139,8 @@ static void UpdateUniformBufferHelper(FVulkanCommandListContext& Context, FVulka
 
 		VulkanRHI::FTempFrameAllocationBuffer::FTempAllocInfo LockInfo;
 		Context.GetTempFrameAllocationBuffer().Alloc(DataSize, 16, LockInfo);
-		FMemory::Memcpy(LockInfo.Data, Data, DataSize);
+		UpdateUniformBufferConstants(Context.GetDevice(), LockInfo.Data, Data, VulkanUniformBuffer->GetLayoutPtr());
+
 		VkBufferCopy Region;
 		Region.size = DataSize;
 		Region.srcOffset = LockInfo.CurrentOffset + LockInfo.Allocation.Offset;
@@ -92,9 +152,9 @@ static void UpdateUniformBufferHelper(FVulkanCommandListContext& Context, FVulka
 	}
 };
 
-FVulkanUniformBuffer::FVulkanUniformBuffer(FVulkanDevice& Device, const FRHIUniformBufferLayout* InLayout, const void* Contents, EUniformBufferUsage InUsage, EUniformBufferValidation Validation)
+FVulkanUniformBuffer::FVulkanUniformBuffer(FVulkanDevice& InDevice, const FRHIUniformBufferLayout* InLayout, const void* Contents, EUniformBufferUsage InUsage, EUniformBufferValidation Validation)
 	: FRHIUniformBuffer(InLayout)
-	, Device(&Device) 
+	, Device(&InDevice)
 	, Usage(InUsage)
 {
 #if VULKAN_ENABLE_AGGRESSIVE_STATS
@@ -140,8 +200,8 @@ FVulkanUniformBuffer::FVulkanUniformBuffer(FVulkanDevice& Device, const FRHIUnif
 				const bool bCanAllocOnThisThread = RHICmdList.Bypass() || (!IsRunningRHIInSeparateThread() && bInRenderingThread) || bInRHIThread;
 				if (bCanAllocOnThisThread)
 				{
-					FVulkanCommandListContextImmediate& Context = Device.GetImmediateContext();
-					UpdateUniformBufferHelper(Context, UniformBuffer, DataSize, Contents);
+					FVulkanCommandListContextImmediate& Context = Device->GetImmediateContext();
+					UpdateUniformBufferHelper(Context, UniformBuffer, Contents);
 				}
 				else
 				{
@@ -151,7 +211,7 @@ FVulkanUniformBuffer::FVulkanUniformBuffer(FVulkanDevice& Device, const FRHIUnif
 					RHICmdList.EnqueueLambda([UniformBuffer, DataSize, CmdListConstantBufferData](FRHICommandList& CmdList)
 					{
 						FVulkanCommandListContext& Context = FVulkanCommandListContext::GetVulkanContext(CmdList.GetContext());
-						UpdateUniformBufferHelper(Context, UniformBuffer, DataSize, CmdListConstantBufferData);
+						UpdateUniformBufferHelper(Context, UniformBuffer, CmdListConstantBufferData);
 					});
 
 					RHICmdList.RHIThreadFence(true);
@@ -160,9 +220,14 @@ FVulkanUniformBuffer::FVulkanUniformBuffer(FVulkanDevice& Device, const FRHIUnif
 		}
 		else
 		{
-			VulkanRHI::FMemoryManager& ResourceMgr = Device.GetMemoryManager();
+			VulkanRHI::FMemoryManager& ResourceMgr = Device->GetMemoryManager();
 			// Set it directly as there is no previous one
-			ResourceMgr.AllocUniformBuffer(Allocation, InLayout->ConstantBufferSize, Contents);
+			ResourceMgr.AllocUniformBuffer(Allocation, InLayout->ConstantBufferSize);
+			if (Contents)
+			{
+				UpdateUniformBufferConstants(Device, Allocation.GetMappedPointer(Device), Contents, InLayout);
+				Allocation.FlushMappedMemory(Device);
+			}
 		}
 	}
 }
@@ -219,7 +284,12 @@ inline void FVulkanDynamicRHI::UpdateUniformBuffer(FRHICommandListBase& RHICmdLi
 		if (ConstantBufferSize > 0)
 		{
 			SCOPE_CYCLE_COUNTER(STAT_VulkanUpdateUniformBuffersRename);
-			Device->GetMemoryManager().AllocUniformBuffer(NewUBAlloc, ConstantBufferSize, Contents);
+			Device->GetMemoryManager().AllocUniformBuffer(NewUBAlloc, ConstantBufferSize);
+			if (Contents)
+			{
+				UpdateUniformBufferConstants(Device, NewUBAlloc.GetMappedPointer(Device), Contents, &Layout);
+				NewUBAlloc.FlushMappedMemory(Device);
+			}
 		}
 	}
 
@@ -231,7 +301,7 @@ inline void FVulkanDynamicRHI::UpdateUniformBuffer(FRHICommandListBase& RHICmdLi
 			if (bUseUpload || bUseRingBuffer)
 			{			
 				FVulkanCommandListContext& Context = Device->GetImmediateContext();
-				UpdateUniformBufferHelper(Context, UniformBuffer, ConstantBufferSize, Contents);
+				UpdateUniformBufferHelper(Context, UniformBuffer, Contents);
 			}
 			else
 			{
@@ -271,7 +341,7 @@ inline void FVulkanDynamicRHI::UpdateUniformBuffer(FRHICommandListBase& RHICmdLi
 			RHICmdList.EnqueueLambda([UniformBuffer, CmdListResources, NumResources, ConstantBufferSize, CmdListConstantBufferData](FRHICommandListBase& CmdList)
 			{
 				FVulkanCommandListContext& Context = (FVulkanCommandListContext&)CmdList.GetContext().GetLowestLevelContext();
-				UpdateUniformBufferHelper(Context, UniformBuffer, ConstantBufferSize, CmdListConstantBufferData);
+				UpdateUniformBufferHelper(Context, UniformBuffer, CmdListConstantBufferData);
 				UniformBuffer->UpdateResourceTable(CmdListResources, NumResources);
 			});
 		}
