VulkanRHI: Fix offset calculation in check

#rnx

diff --git a/Engine/Source/Runtime/VulkanRHI/Private/VulkanCommands.cpp b/Engine/Source/Runtime/VulkanRHI/Private/VulkanCommands.cpp
--- a/Engine/Source/Runtime/VulkanRHI/Private/VulkanCommands.cpp
+++ b/Engine/Source/Runtime/VulkanRHI/Private/VulkanCommands.cpp
@@ -812,7 +812,6 @@ void FVulkanCommandListContext::RHICopyToStagingBuffer(FRHIBuffer* SourceBufferR
 		StagingBuffer->Device = Device;
 	}
 
-	StagingBuffer->QueuedOffset = Offset;
 	StagingBuffer->QueuedNumBytes = NumBytes;
 
 	VkBufferCopy Region;
diff --git a/Engine/Source/Runtime/VulkanRHI/Private/VulkanResources.h b/Engine/Source/Runtime/VulkanRHI/Private/VulkanResources.h
--- a/Engine/Source/Runtime/VulkanRHI/Private/VulkanResources.h
+++ b/Engine/Source/Runtime/VulkanRHI/Private/VulkanResources.h
@@ -1272,7 +1272,6 @@ public:
 
 private:
 	VulkanRHI::FStagingBuffer* StagingBuffer = nullptr;
-	uint32 QueuedOffset = 0;
 	uint32 QueuedNumBytes = 0;
 	// The staging buffer was allocated from this device.
 	FVulkanDevice* Device;
diff --git a/Engine/Source/Runtime/VulkanRHI/Private/VulkanUtil.cpp b/Engine/Source/Runtime/VulkanRHI/Private/VulkanUtil.cpp
--- a/Engine/Source/Runtime/VulkanRHI/Private/VulkanUtil.cpp
+++ b/Engine/Source/Runtime/VulkanRHI/Private/VulkanUtil.cpp
@@ -108,10 +108,8 @@ void* FVulkanStagingBuffer::Lock(uint32 Offset, uint32 NumBytes)
 {
 	check(!bIsLocked);
 	bIsLocked = true;
-	uint32 QueuedEndOffset = QueuedNumBytes + QueuedOffset;
-	uint32 EndOffset = Offset + NumBytes;
-	check(Offset < QueuedNumBytes && EndOffset <= QueuedEndOffset);
-	//#todo-rco: Apply the offset in case it doesn't match
+	const uint32 EndOffset = Offset + NumBytes;
+	checkf(EndOffset <= QueuedNumBytes, TEXT("Lock at Offset (%u) and NumBytes (%u) reads beyond the allocated size of the staging buffer (%u)"), Offset, NumBytes, QueuedNumBytes);
 	return (void*)((uint8*)StagingBuffer->GetMappedPointer() + Offset);
 }
 
