Perform driver version checks before initializing the RHI device.

The driver version is logged and checked after the RHI instance is created, before calling GDynamicRHI->Init(). This worked on D3D11 and 12 without many changes because those RHIs select the adapter before the Init call, so they can fill in the information necessary for GetGPUDriverInfo() before Init. However, the Vulkan RHI init code had to be restructured quite a bit to do the same thing.

Also removed some obsolete hardcoded driver version checks from the Vulkan RHI, since we have better checks in RHIDetectAndWarnOfBadDrivers().

#rnx

diff --git a/Engine/Source/Runtime/Apple/MetalRHI/Private/MetalRHI.cpp b/Engine/Source/Runtime/Apple/MetalRHI/Private/MetalRHI.cpp
--- a/Engine/Source/Runtime/Apple/MetalRHI/Private/MetalRHI.cpp
+++ b/Engine/Source/Runtime/Apple/MetalRHI/Private/MetalRHI.cpp
@@ -987,31 +987,18 @@ FMetalDynamicRHI::FMetalDynamicRHI(ERHIFeatureLevel::Type RequestedFeatureLevel)
 		break;
 	};
 
-	// get driver version (todo: share with other RHIs)
-	{
-		FGPUDriverInfo GPUDriverInfo = FPlatformMisc::GetGPUDriverInfo(GRHIAdapterName);
-		
-		GRHIAdapterUserDriverVersion = GPUDriverInfo.UserDriverVersion;
-		GRHIAdapterInternalDriverVersion = GPUDriverInfo.InternalDriverVersion;
-		GRHIAdapterDriverDate = GPUDriverInfo.DriverDate;
-		
-		UE_LOG(LogMetal, Display, TEXT("    Adapter Name: %s"), *GRHIAdapterName);
-		UE_LOG(LogMetal, Display, TEXT("  Driver Version: %s (internal:%s, unified:%s)"), *GRHIAdapterUserDriverVersion, *GRHIAdapterInternalDriverVersion, *GPUDriverInfo.GetUnifiedDriverVersion());
-		UE_LOG(LogMetal, Display, TEXT("     Driver Date: %s"), *GRHIAdapterDriverDate);
-		UE_LOG(LogMetal, Display, TEXT("          Vendor: %s"), *GPUDriverInfo.ProviderName);
 #if PLATFORM_MAC
-		if(GPUDesc.GPUVendorId == GRHIVendorId)
-		{
-			UE_LOG(LogMetal, Display,  TEXT("      Vendor ID: %d"), GPUDesc.GPUVendorId);
-			UE_LOG(LogMetal, Display,  TEXT("      Device ID: %d"), GPUDesc.GPUDeviceId);
-			UE_LOG(LogMetal, Display,  TEXT("      VRAM (MB): %d"), GPUDesc.GPUMemoryMB);
-		}
-		else
-		{
-			UE_LOG(LogMetal, Warning,  TEXT("GPU descriptor (%s) from IORegistry failed to match Metal (%s)"), *FString(GPUDesc.GPUName), *GRHIAdapterName);
-		}
-#endif
+	if(GPUDesc.GPUVendorId == GRHIVendorId)
+	{
+		UE_LOG(LogMetal, Display,  TEXT("      Vendor ID: %d"), GPUDesc.GPUVendorId);
+		UE_LOG(LogMetal, Display,  TEXT("      Device ID: %d"), GPUDesc.GPUDeviceId);
+		UE_LOG(LogMetal, Display,  TEXT("      VRAM (MB): %d"), GPUDesc.GPUMemoryMB);
 	}
+	else
+	{
+		UE_LOG(LogMetal, Warning,  TEXT("GPU descriptor (%s) from IORegistry failed to match Metal (%s)"), *FString(GPUDesc.GPUName), *GRHIAdapterName);
+	}
+#endif
 
 #if PLATFORM_MAC
 	if (!FPlatformProcess::IsSandboxedApplication())
diff --git a/Engine/Source/Runtime/D3D12RHI/Private/Windows/WindowsD3D12Device.cpp b/Engine/Source/Runtime/D3D12RHI/Private/Windows/WindowsD3D12Device.cpp
--- a/Engine/Source/Runtime/D3D12RHI/Private/Windows/WindowsD3D12Device.cpp
+++ b/Engine/Source/Runtime/D3D12RHI/Private/Windows/WindowsD3D12Device.cpp
@@ -855,6 +855,12 @@ void FD3D12DynamicRHIModule::FindAdapter()
 	if (ChosenAdapters.Num() > 0 && ChosenAdapters[0]->GetDesc().IsValid())
 	{
 		UE_LOG(LogD3D12RHI, Log, TEXT("Chosen D3D12 Adapter Id = %u"), ChosenAdapters[0]->GetAdapterIndex());
+
+		const DXGI_ADAPTER_DESC& AdapterDesc = ChosenAdapters[0]->GetD3DAdapterDesc();
+		GRHIVendorId = AdapterDesc.VendorId;
+		GRHIAdapterName = AdapterDesc.Description;
+		GRHIDeviceId = AdapterDesc.DeviceId;
+		GRHIDeviceRevision = AdapterDesc.Revision;
 	}
 	else
 	{
@@ -1144,8 +1150,7 @@ void FD3D12DynamicRHI::Init()
 
 	const DXGI_ADAPTER_DESC& AdapterDesc = GetAdapter().GetD3DAdapterDesc();
 
-	// Need to set GRHIVendorId before calling IsRHIDevice* functions
-	GRHIVendorId = AdapterDesc.VendorId;
+	UE_LOG(LogD3D12RHI, Log, TEXT("    GPU DeviceId: 0x%x (for the marketing name, search the web for \"GPU Device Id\")"), AdapterDesc.DeviceId);
 
 #if AMD_API_ENABLE
 	// Initialize the AMD AGS utility library, when running on an AMD device
@@ -1323,26 +1328,6 @@ void FD3D12DynamicRHI::Init()
 
 	GTexturePoolSize = 0;
 
-	GRHIAdapterName = AdapterDesc.Description;
-	GRHIDeviceId = AdapterDesc.DeviceId;
-	GRHIDeviceRevision = AdapterDesc.Revision;
-
-	UE_LOG(LogD3D12RHI, Log, TEXT("    GPU DeviceId: 0x%x (for the marketing name, search the web for \"GPU Device Id\")"),
-		AdapterDesc.DeviceId);
-
-	// get driver version (todo: share with other RHIs)
-	{
-		FGPUDriverInfo GPUDriverInfo = FPlatformMisc::GetGPUDriverInfo(GRHIAdapterName);
-
-		GRHIAdapterUserDriverVersion = GPUDriverInfo.UserDriverVersion;
-		GRHIAdapterInternalDriverVersion = GPUDriverInfo.InternalDriverVersion;
-		GRHIAdapterDriverDate = GPUDriverInfo.DriverDate;
-
-		UE_LOG(LogD3D12RHI, Log, TEXT("    Adapter Name: %s"), *GRHIAdapterName);
-		UE_LOG(LogD3D12RHI, Log, TEXT("  Driver Version: %s (internal:%s, unified:%s)"), *GRHIAdapterUserDriverVersion, *GRHIAdapterInternalDriverVersion, *GPUDriverInfo.GetUnifiedDriverVersion());
-		UE_LOG(LogD3D12RHI, Log, TEXT("     Driver Date: %s"), *GRHIAdapterDriverDate);
-	}
-
 	// Issue: 32bit windows doesn't report 64bit value, we take what we get.
 	FD3D12GlobalStats::GDedicatedVideoMemory = int64(AdapterDesc.DedicatedVideoMemory);
 	FD3D12GlobalStats::GDedicatedSystemMemory = int64(AdapterDesc.DedicatedSystemMemory);
diff --git a/Engine/Source/Runtime/RHI/Private/DynamicRHI.cpp b/Engine/Source/Runtime/RHI/Private/DynamicRHI.cpp
--- a/Engine/Source/Runtime/RHI/Private/DynamicRHI.cpp
+++ b/Engine/Source/Runtime/RHI/Private/DynamicRHI.cpp
@@ -82,9 +82,9 @@ void InitNullRHI()
 #if PLATFORM_WINDOWS || PLATFORM_UNIX
 static void RHIDetectAndWarnOfBadDrivers(bool bHasEditorToken)
 {
-	if (!GIsRHIInitialized || (GRHIVendorId == 0))
+	if (GRHIVendorId == 0)
 	{
-		UE_LOG(LogRHI, Log, TEXT("Skipping Driver Check: RHI%s initialized, VendorId=0x%x"), GIsRHIInitialized ? TEXT("") : TEXT(" NOT"), GRHIVendorId);
+		UE_LOG(LogRHI, Log, TEXT("Skipping Driver Check, no vendor ID set."));
 		return;
 	}
 
@@ -268,7 +268,7 @@ static void RHIDetectAndWarnOfBadDrivers(bool bHasEditorToken)
 {
 	int32 CVarValue = CVarWarnOfBadDrivers.GetValueOnGameThread();
 
-	if (!GIsRHIInitialized || !CVarValue || GRHIVendorId == 0 || bHasEditorToken || FApp::IsUnattended())
+	if (!CVarValue || GRHIVendorId == 0 || bHasEditorToken || FApp::IsUnattended())
 	{
 		return;
 	}
@@ -316,6 +316,27 @@ void RHIInit(bool bHasEditorToken)
 			GDynamicRHI = PlatformCreateDynamicRHI();
 			if (GDynamicRHI)
 			{
+#if PLATFORM_WINDOWS || PLATFORM_MAC || PLATFORM_UNIX
+
+				// Get driver version. Creating GDynamicRHI is expected to set GRHIAdapterName.
+				FGPUDriverInfo GPUDriverInfo = FPlatformMisc::GetGPUDriverInfo(GRHIAdapterName);
+				if (GPUDriverInfo.IsValid())
+				{
+					// GetGPUDriverInfo is not implemented on Linux, so it returns an invalid driver info object. However, the FVulkanDynamicRHI constructor
+					// sets these values on that platform, so we'll still have data we can log.
+					GRHIAdapterUserDriverVersion = GPUDriverInfo.UserDriverVersion;
+					GRHIAdapterInternalDriverVersion = GPUDriverInfo.InternalDriverVersion;
+					GRHIAdapterDriverDate = GPUDriverInfo.DriverDate;
+				}
+
+				UE_LOG(LogRHI, Log, TEXT("RHI Adapter Info:"));
+				UE_LOG(LogRHI, Log, TEXT("            Name: %s"), *GRHIAdapterName);
+				UE_LOG(LogRHI, Log, TEXT("  Driver Version: %s (internal:%s, unified:%s)"), *GRHIAdapterUserDriverVersion, *GRHIAdapterInternalDriverVersion, *GPUDriverInfo.GetUnifiedDriverVersion());
+				UE_LOG(LogRHI, Log, TEXT("     Driver Date: %s"), *GRHIAdapterDriverDate);
+
+				RHIDetectAndWarnOfBadDrivers(bHasEditorToken);
+#endif
+
 				GDynamicRHI->Init();
 
 #if WITH_MGPU
@@ -367,10 +388,6 @@ void RHIInit(bool bHasEditorToken)
 
 		check(GDynamicRHI);
 	}
-
-#if PLATFORM_WINDOWS || PLATFORM_MAC || PLATFORM_UNIX
-	RHIDetectAndWarnOfBadDrivers(bHasEditorToken);
-#endif
 }
 
 void RHIPostInit(const TArray<uint32>& InPixelFormatByteWidth)
diff --git a/Engine/Source/Runtime/VulkanRHI/Private/Linux/VulkanLinuxPlatform.cpp b/Engine/Source/Runtime/VulkanRHI/Private/Linux/VulkanLinuxPlatform.cpp
--- a/Engine/Source/Runtime/VulkanRHI/Private/Linux/VulkanLinuxPlatform.cpp
+++ b/Engine/Source/Runtime/VulkanRHI/Private/Linux/VulkanLinuxPlatform.cpp
@@ -236,7 +236,3 @@ void FVulkanLinuxPlatform::WriteCrashMarker(const FOptionalVulkanDeviceExtension
 		}
 	}
 }
-
-void FVulkanLinuxPlatform::CheckDeviceDriver(uint32 DeviceIndex, EGpuVendorId VendorId, const VkPhysicalDeviceProperties& Props)
-{
-}
diff --git a/Engine/Source/Runtime/VulkanRHI/Private/Linux/VulkanLinuxPlatform.h b/Engine/Source/Runtime/VulkanRHI/Private/Linux/VulkanLinuxPlatform.h
--- a/Engine/Source/Runtime/VulkanRHI/Private/Linux/VulkanLinuxPlatform.h
+++ b/Engine/Source/Runtime/VulkanRHI/Private/Linux/VulkanLinuxPlatform.h
@@ -46,7 +46,6 @@ class FVulkanLinuxPlatform : public FVulkanGenericPlatform
 public:
 	static bool IsSupported();
 
-	static void CheckDeviceDriver(uint32 DeviceIndex, EGpuVendorId VendorId, const VkPhysicalDeviceProperties& Props);
 	static bool LoadVulkanLibrary();
 	static bool LoadVulkanInstanceFunctions(VkInstance inInstance);
 	static void FreeVulkanLibrary();
diff --git a/Engine/Source/Runtime/VulkanRHI/Private/VulkanDevice.cpp b/Engine/Source/Runtime/VulkanRHI/Private/VulkanDevice.cpp
--- a/Engine/Source/Runtime/VulkanRHI/Private/VulkanDevice.cpp
+++ b/Engine/Source/Runtime/VulkanRHI/Private/VulkanDevice.cpp
@@ -950,7 +950,7 @@ void FVulkanDevice::MapFormatSupport(EPixelFormat UEFormat, std::initializer_lis
 	MapFormatSupport(UEFormat, PrioritizedFormats, ComponentMapping, RequiredCapabilities, kDefaultBlockBytes);
 }
 
-void FVulkanDevice::InitGPU(int32 DeviceIndex)
+void FVulkanDevice::InitGPU()
 {
 	LLM_SCOPE_VULKAN(ELLMTagVulkan::VulkanMisc);
 
@@ -1021,7 +1021,7 @@ void FVulkanDevice::InitGPU(int32 DeviceIndex)
 		}
 	}
 
-	UE_LOG(LogVulkanRHI, Display, TEXT("Using Device %d: Geometry %d BufferAtomic64 %d ImageAtomic64 %d"), DeviceIndex, PhysicalFeatures.geometryShader, OptionalDeviceExtensions.HasKHRShaderAtomicInt64, OptionalDeviceExtensions.HasImageAtomicInt64);
+	UE_LOG(LogVulkanRHI, Display, TEXT("Device properties: Geometry %d BufferAtomic64 %d ImageAtomic64 %d"), PhysicalFeatures.geometryShader, OptionalDeviceExtensions.HasKHRShaderAtomicInt64, OptionalDeviceExtensions.HasImageAtomicInt64);
 
 	CreateDevice(DeviceLayers, UEExtensions);
 
diff --git a/Engine/Source/Runtime/VulkanRHI/Private/VulkanDevice.h b/Engine/Source/Runtime/VulkanRHI/Private/VulkanDevice.h
--- a/Engine/Source/Runtime/VulkanRHI/Private/VulkanDevice.h
+++ b/Engine/Source/Runtime/VulkanRHI/Private/VulkanDevice.h
@@ -197,7 +197,7 @@ public:
 
 	~FVulkanDevice();
 
-	void InitGPU(int32 DeviceIndex);
+	void InitGPU();
 
 	void CreateDevice(TArray<const ANSICHAR*>& DeviceLayers, FVulkanDeviceExtensionArray& UEExtensions);
 
diff --git a/Engine/Source/Runtime/VulkanRHI/Private/VulkanGenericPlatform.h b/Engine/Source/Runtime/VulkanRHI/Private/VulkanGenericPlatform.h
--- a/Engine/Source/Runtime/VulkanRHI/Private/VulkanGenericPlatform.h
+++ b/Engine/Source/Runtime/VulkanRHI/Private/VulkanGenericPlatform.h
@@ -23,7 +23,6 @@ public:
 	static void SetupMaxRHIFeatureLevelAndShaderPlatform(ERHIFeatureLevel::Type InRequestedFeatureLevel);
 
 	static bool IsSupported() { return true; }
-	static void CheckDeviceDriver(uint32 DeviceIndex, EGpuVendorId VendorId, const VkPhysicalDeviceProperties& Props) {}
 
 	static bool LoadVulkanLibrary() { return true; }
 	static bool LoadVulkanInstanceFunctions(VkInstance inInstance) { return true; }
diff --git a/Engine/Source/Runtime/VulkanRHI/Private/VulkanLayers.cpp b/Engine/Source/Runtime/VulkanRHI/Private/VulkanLayers.cpp
--- a/Engine/Source/Runtime/VulkanRHI/Private/VulkanLayers.cpp
+++ b/Engine/Source/Runtime/VulkanRHI/Private/VulkanLayers.cpp
@@ -304,7 +304,7 @@ TArray<const ANSICHAR*> FVulkanDynamicRHI::SetupInstanceLayers(FVulkanInstanceEx
 {
 	FVulkanIntanceSetupHelper InstanceHelper;
 	TArray<const ANSICHAR*> OutInstanceLayers = SetupLayers(InstanceHelper, UEExtensions);
-	GVulkanRHI->ActiveDebugLayerExtension = InstanceHelper.ActiveDebugLayerExtension;
+	ActiveDebugLayerExtension = InstanceHelper.ActiveDebugLayerExtension;
 	return OutInstanceLayers;
 }
 
diff --git a/Engine/Source/Runtime/VulkanRHI/Private/VulkanRHI.cpp b/Engine/Source/Runtime/VulkanRHI/Private/VulkanRHI.cpp
--- a/Engine/Source/Runtime/VulkanRHI/Private/VulkanRHI.cpp
+++ b/Engine/Source/Runtime/VulkanRHI/Private/VulkanRHI.cpp
@@ -224,10 +224,7 @@ FVulkanDynamicRHI::FVulkanDynamicRHI()
 
 	// Copy source requires its own image layout.
 	EnumRemoveFlags(GRHIMergeableAccessMask, ERHIAccess::CopySrc);
-}
 
-void FVulkanDynamicRHI::Init()
-{
 	// Setup the validation requests ready before we load dlls
 	SetupValidationRequests();
 
@@ -238,7 +235,7 @@ void FVulkanDynamicRHI::Init()
 #if PLATFORM_LINUX
 		// be more verbose on Linux
 		FPlatformMisc::MessageBoxExt(EAppMsgType::Ok, *LOCTEXT("UnableToInitializeVulkanLinux", "Unable to load Vulkan library and/or acquire the necessary function pointers. Make sure an up-to-date libvulkan.so.1 is installed.").ToString(),
-									 *LOCTEXT("UnableToInitializeVulkanLinuxTitle", "Unable to initialize Vulkan.").ToString());
+			*LOCTEXT("UnableToInitializeVulkanLinuxTitle", "Unable to initialize Vulkan.").ToString());
 #endif // PLATFORM_LINUX
 		UE_LOG(LogVulkanRHI, Fatal, TEXT("Failed to find all required Vulkan entry points; make sure your driver supports Vulkan!"));
 	}
@@ -252,6 +249,12 @@ void FVulkanDynamicRHI::Init()
 	LLM(VulkanLLM::Initialize());
 #endif
 
+	CreateInstance();
+	SelectDevice();
+}
+
+void FVulkanDynamicRHI::Init()
+{
 	InitInstance();
 
 	bIsStandaloneStereoDevice = IHeadMountedDisplayModule::IsAvailable() && IHeadMountedDisplayModule::Get().IsStandaloneStereoOnlyDevice();
@@ -523,7 +526,7 @@ static inline int32 PreferAdapterVendor()
 	return -1;
 }
 
-void FVulkanDynamicRHI::SelectAndInitDevice()
+void FVulkanDynamicRHI::SelectDevice()
 {
 	uint32 GpuCount = 0;
 	VkResult Result = VulkanRHI::vkEnumeratePhysicalDevices(Instance, &GpuCount, nullptr);
@@ -685,16 +688,14 @@ void FVulkanDynamicRHI::SelectAndInitDevice()
 		}
 	}
 
+	UE_LOG(LogVulkanRHI, Log, TEXT("Chosen device index: %d"), DeviceIndex);
+
 	if (bUseVendorIdAsIs)
 	{
 		GRHIVendorId = Props.vendorID;
 	}
 	GRHIAdapterName = ANSI_TO_TCHAR(Props.deviceName);
 
-	FVulkanPlatform::CheckDeviceDriver(DeviceIndex, Device->GetVendorId(), Props);
-
-	Device->InitGPU(DeviceIndex);
-
 	auto ReadVulkanDriverVersionFromProps = [](FVulkanDevice* CurrentDevice) {
 
 		const VkPhysicalDeviceProperties& Props = CurrentDevice->GetDeviceProperties();
@@ -714,10 +715,7 @@ void FVulkanDynamicRHI::SelectAndInitDevice()
 		GRHIDeviceId = Props.deviceID;
 		GRHIAdapterInternalDriverVersion = GRHIAdapterUserDriverVersion;
 		GRHIAdapterDriverDate = TEXT("01-01-01");  // Unused on unix systems, pick a date that will fail test if compared but passes IsValid() check
-
-		UE_LOG(LogVulkanRHI, Log, TEXT("    Adapter Name: %s"), *GRHIAdapterName);
 		UE_LOG(LogVulkanRHI, Log, TEXT("     API Version: %d.%d.%d"), VK_VERSION_MAJOR(Props.apiVersion), VK_VERSION_MINOR(Props.apiVersion), VK_VERSION_PATCH(Props.apiVersion));
-		UE_LOG(LogVulkanRHI, Log, TEXT("  Driver Version: %s"), *GRHIAdapterUserDriverVersion);
 	};
 
 	if (PLATFORM_ANDROID)
@@ -728,25 +726,7 @@ void FVulkanDynamicRHI::SelectAndInitDevice()
 	else if (PLATFORM_WINDOWS)
 	{
 		GRHIDeviceId = Props.deviceID;
-		FGPUDriverInfo GPUDriverInfo = FPlatformMisc::GetGPUDriverInfo(GRHIAdapterName);
-
-		if (GPUDriverInfo.InternalDriverVersion != TEXT("Unknown"))
-		{
-			GRHIAdapterUserDriverVersion = GPUDriverInfo.UserDriverVersion;
-			GRHIAdapterInternalDriverVersion = GPUDriverInfo.InternalDriverVersion;
-			GRHIAdapterDriverDate = GPUDriverInfo.DriverDate;
-
-			UE_LOG(LogVulkanRHI, Log, TEXT("    Adapter Name: %s"), *GRHIAdapterName);
-			UE_LOG(LogVulkanRHI, Log, TEXT("     API Version: %d.%d.%d"), VK_VERSION_MAJOR(Props.apiVersion), VK_VERSION_MINOR(Props.apiVersion), VK_VERSION_PATCH(Props.apiVersion));
-			UE_LOG(LogVulkanRHI, Log, TEXT("  Driver Version: %s (0x%X)"), *GRHIAdapterUserDriverVersion, Props.driverVersion);
-			UE_LOG(LogVulkanRHI, Log, TEXT("Internal Version: %s"), *GRHIAdapterInternalDriverVersion);
-			UE_LOG(LogVulkanRHI, Log, TEXT("     Driver Date: %s"), *GRHIAdapterDriverDate);
-		}
-		else
-		{
-			// If we failed to read from the registry, then use the values provided by Vulkan props
-			ReadVulkanDriverVersionFromProps(Device);
-		}
+		UE_LOG(LogVulkanRHI, Log, TEXT("API Version: %d.%d.%d"), VK_VERSION_MAJOR(Props.apiVersion), VK_VERSION_MINOR(Props.apiVersion), VK_VERSION_PATCH(Props.apiVersion));
 	}
 	else if(PLATFORM_UNIX)
 	{
@@ -760,20 +740,17 @@ void FVulkanDynamicRHI::InitInstance()
 {
 	check(IsInGameThread());
 
-	if (!Device)
+	if (!GIsRHIInitialized)
 	{
 		// Wait for the rendering thread to go idle.
 		FlushRenderingCommands();
 
-		check(!GIsRHIInitialized);
-
 		FVulkanPlatform::OverridePlatformHandlers(true);
 
 		GRHISupportsAsyncTextureCreation = false;
 		GEnableAsyncCompute = false;
 
-		CreateInstance();
-		SelectAndInitDevice();
+		Device->InitGPU();
 
 #if VULKAN_HAS_DEBUGGING_ENABLED
 		if (GRenderDocFound)
diff --git a/Engine/Source/Runtime/VulkanRHI/Private/Windows/VulkanWindowsPlatform.cpp b/Engine/Source/Runtime/VulkanRHI/Private/Windows/VulkanWindowsPlatform.cpp
--- a/Engine/Source/Runtime/VulkanRHI/Private/Windows/VulkanWindowsPlatform.cpp
+++ b/Engine/Source/Runtime/VulkanRHI/Private/Windows/VulkanWindowsPlatform.cpp
@@ -229,120 +229,3 @@ void FVulkanWindowsPlatform::WriteCrashMarker(const FOptionalVulkanDeviceExtensi
 		}
 	}
 }
-
-void FVulkanWindowsPlatform::CheckDeviceDriver(uint32 DeviceIndex, EGpuVendorId VendorId, const VkPhysicalDeviceProperties& Props)
-{
-	const bool bAllowVendorDevice = !FParse::Param(FCommandLine::Get(), TEXT("novendordevice"));
-	if (VendorId == EGpuVendorId::Amd && bAllowVendorDevice)
-	{
-		AGSGPUInfo AmdGpuInfo;
-		AGSContext* AmdAgsContext = nullptr;
-		if (agsInitialize(AGS_MAKE_VERSION(AMD_AGS_VERSION_MAJOR, AMD_AGS_VERSION_MINOR, AMD_AGS_VERSION_PATCH), nullptr, &AmdAgsContext, &AmdGpuInfo) == AGS_SUCCESS)
-		{
-			const char* Version = AmdGpuInfo.radeonSoftwareVersion;
-			if (DeviceIndex < (uint32)AmdGpuInfo.numDevices && Version && *Version)
-			{
-				auto& DeviceInfo = AmdGpuInfo.devices[DeviceIndex];
-				bool bIsPreGCN = DeviceInfo.asicFamily == AGSDeviceInfo::AsicFamily_PreGCN;
-				if (DeviceInfo.asicFamily != AGSDeviceInfo::AsicFamily_Unknown)
-				{
-					// "Major.Minor.Revision"
-					do
-					{
-						int32 MajorVersion = FCStringAnsi::Atoi(Version);
-						while (*Version >= '0' && *Version <= '9')
-						{
-							++Version;
-						}
-
-						if (*Version != '.')
-						{
-							break;
-						}
-						++Version;
-
-						int32 MinorVersion = FCStringAnsi::Atoi(Version);
-						while (*Version >= '0' && *Version <= '9')
-						{
-							++Version;
-						}
-
-						if (*Version != '.')
-						{
-							break;
-						}
-						++Version;
-
-						int32 RevisionVersion = FCStringAnsi::Atoi(Version);
-
-						if (MajorVersion > 0)
-						{
-							if (MajorVersion < 18)
-							{
-								// Deny drivers older than 18.xx.xx drivers
-								FPlatformMisc::MessageBoxExt(EAppMsgType::Ok, TEXT("There are known issues with older Vulkan Radeon drivers; the recommended version is 19.4.1; please try updating your driver to that version."), TEXT("Vulkan driver version"));
-								FPlatformMisc::RequestExitWithStatus(true, 1);
-							}
-							else if (WITH_EDITOR)
-							{
-								bool bBadVersion = false;
-								if (MajorVersion == 19)
-								{
-									if (MinorVersion < 4 || (MinorVersion == 4 && RevisionVersion < 1))
-									{
-										bBadVersion = true;
-									}
-								}
-								else if (MajorVersion == 18)
-								{
-									if (MinorVersion > 12 || (MinorVersion == 12 && RevisionVersion >= 2))
-									{
-										bBadVersion = true;
-									}
-								}
-
-								if (bBadVersion)
-								{
-									// Deny drivers between 18.12.2 and 19.2.1, as they as it introduced an issue with Slate windows/Vulkan viewports on the editor; 19.3.x also have crashes
-									FPlatformMisc::MessageBoxExt(EAppMsgType::Ok, TEXT("There are known issues with Vulkan on the editor with the some \nRadeon drivers; the recommended version is 19.4.1: please try updating your driver to that version."), TEXT("Vulkan driver version"));
-									FPlatformMisc::RequestExitWithStatus(true, 1);
-								}
-							}
-						}
-					}
-					while (0);
-
-					GRHIDeviceIsAMDPreGCNArchitecture = GRHIDeviceIsAMDPreGCNArchitecture || bIsPreGCN;
-					if (GRHIDeviceIsAMDPreGCNArchitecture)
-					{
-						UE_LOG(LogVulkanRHI, Log, TEXT("AMD Pre GCN architecture detected, some driver workarounds will be in place"));
-					}
-					UE_LOG(LogVulkanRHI, Display, TEXT("AMD User Driver Version = %s"), ANSI_TO_TCHAR(AmdGpuInfo.radeonSoftwareVersion));
-				}
-			}
-
-			agsDeInitialize(AmdAgsContext);
-		}
-	}
-	else if (VendorId == EGpuVendorId::Nvidia)
-	{
-		UNvidiaDriverVersion NvidiaVersion;
-		static_assert(sizeof(NvidiaVersion) == sizeof(Props.driverVersion), "Mismatched Nvidia pack driver version!");
-		NvidiaVersion.Packed = Props.driverVersion;
-
-		if (GRHIAdapterName.Contains(TEXT("RTX 20")))
-		{
-			if (NvidiaVersion.Major < 430)
-			{
-				// Workaround a crash on 20xx family
-				UE_LOG(LogVulkanRHI, Warning, TEXT("Nvidia 20xx family of GPUs have a known crash on drivers < 430. Compatibility mode (slow!) will now be enabled"));
-
-				extern TAutoConsoleVariable<int32> GRHIThreadCvar;
-				GRHIThreadCvar->SetWithCurrentPriority(0);
-				IConsoleVariable* BypassVar = IConsoleManager::Get().FindConsoleVariable(TEXT("r.RHICmdBypass"));
-				BypassVar->SetWithCurrentPriority(1);
-			}
-		}
-	}
-}
-
diff --git a/Engine/Source/Runtime/VulkanRHI/Private/Windows/VulkanWindowsPlatform.h b/Engine/Source/Runtime/VulkanRHI/Private/Windows/VulkanWindowsPlatform.h
--- a/Engine/Source/Runtime/VulkanRHI/Private/Windows/VulkanWindowsPlatform.h
+++ b/Engine/Source/Runtime/VulkanRHI/Private/Windows/VulkanWindowsPlatform.h
@@ -52,8 +52,6 @@
 class FVulkanWindowsPlatform : public FVulkanGenericPlatform
 {
 public:
-	static void CheckDeviceDriver(uint32 DeviceIndex, EGpuVendorId VendorId, const VkPhysicalDeviceProperties& Props);
-
 	static bool LoadVulkanLibrary();
 	static bool LoadVulkanInstanceFunctions(VkInstance inInstance);
 	static void FreeVulkanLibrary();
diff --git a/Engine/Source/Runtime/VulkanRHI/Public/VulkanDynamicRHI.h b/Engine/Source/Runtime/VulkanRHI/Public/VulkanDynamicRHI.h
--- a/Engine/Source/Runtime/VulkanRHI/Public/VulkanDynamicRHI.h
+++ b/Engine/Source/Runtime/VulkanRHI/Public/VulkanDynamicRHI.h
@@ -379,7 +379,7 @@ protected:
 	TRefCountPtr<FVulkanViewport> DrawingViewport;
 
 	void CreateInstance();
-	void SelectAndInitDevice();
+	void SelectDevice();
 	void InitGPU(FVulkanDevice* Device);
 	void InitDevice(FVulkanDevice* Device);
 
@@ -391,7 +391,7 @@ protected:
 	static void WriteEndFrameTimestamp(void*);
 	*/
 
-	static TArray<const ANSICHAR*> SetupInstanceLayers(FVulkanInstanceExtensionArray& UEExtensions);
+	TArray<const ANSICHAR*> SetupInstanceLayers(FVulkanInstanceExtensionArray& UEExtensions);
 
 	IConsoleObject* SavePipelineCacheCmd = nullptr;
 	IConsoleObject* RebuildPipelineCacheCmd = nullptr;
diff --git a/Engine/Source/Runtime/Windows/D3D11RHI/Private/Windows/WindowsD3D11Device.cpp b/Engine/Source/Runtime/Windows/D3D11RHI/Private/Windows/WindowsD3D11Device.cpp
--- a/Engine/Source/Runtime/Windows/D3D11RHI/Private/Windows/WindowsD3D11Device.cpp
+++ b/Engine/Source/Runtime/Windows/D3D11RHI/Private/Windows/WindowsD3D11Device.cpp
@@ -1087,6 +1087,10 @@ void FD3D11DynamicRHIModule::FindAdapter()
 		UE_LOG(LogD3D11RHI, Error, TEXT("Failed to choose a D3D11 Adapter."));
 	}
 
+	GRHIAdapterName = ChosenAdapter.DXGIAdapterDesc.Description;
+	GRHIVendorId = ChosenAdapter.DXGIAdapterDesc.VendorId;
+	GRHIDeviceId = ChosenAdapter.DXGIAdapterDesc.DeviceId;
+	GRHIDeviceRevision = ChosenAdapter.DXGIAdapterDesc.Revision;
 	GRHIDeviceIsIntegrated = ChosenAdapter.bIsIntegrated;
 }
 
@@ -1727,28 +1731,10 @@ void FD3D11DynamicRHI::InitD3DDevice()
 
 		GTexturePoolSize = 0;
 
-		GRHIAdapterName = Adapter.DXGIAdapterDesc.Description;
-		GRHIVendorId = Adapter.DXGIAdapterDesc.VendorId;
-		GRHIDeviceId = Adapter.DXGIAdapterDesc.DeviceId;
-		GRHIDeviceRevision = Adapter.DXGIAdapterDesc.Revision;
 		// turn off creation on other threads for NVidia since a driver heuristic will notice that and make the creation synchronous, and that is not desirable given that large number of shaders will still be created on a single thread
 		GRHISupportsMultithreadedShaderCreation = !IsRHIDeviceNVIDIA(); 
 
-		UE_LOG(LogD3D11RHI, Log, TEXT("    GPU DeviceId: 0x%x (for the marketing name, search the web for \"GPU Device Id\")"), 
-			Adapter.DXGIAdapterDesc.DeviceId);
-
-		// get driver version (todo: share with other RHIs)
-		{
-			FGPUDriverInfo GPUDriverInfo = FPlatformMisc::GetGPUDriverInfo(GRHIAdapterName);
-
-			GRHIAdapterUserDriverVersion = GPUDriverInfo.UserDriverVersion;
-			GRHIAdapterInternalDriverVersion = GPUDriverInfo.InternalDriverVersion;
-			GRHIAdapterDriverDate = GPUDriverInfo.DriverDate;
-
-			UE_LOG(LogD3D11RHI, Log, TEXT("    Adapter Name: %s"), *GRHIAdapterName);
-			UE_LOG(LogD3D11RHI, Log, TEXT("  Driver Version: %s (internal:%s, unified:%s)"), *GRHIAdapterUserDriverVersion, *GRHIAdapterInternalDriverVersion, *GPUDriverInfo.GetUnifiedDriverVersion());
-			UE_LOG(LogD3D11RHI, Log, TEXT("     Driver Date: %s"), *GRHIAdapterDriverDate);
-		}
+		UE_LOG(LogD3D11RHI, Log, TEXT("    GPU DeviceId: 0x%x (for the marketing name, search the web for \"GPU Device Id\")"), Adapter.DXGIAdapterDesc.DeviceId);
 
 		// Issue: 32bit windows doesn't report 64bit value, we take what we get.
 		FD3D11GlobalStats::GDedicatedVideoMemory = int64(Adapter.DXGIAdapterDesc.DedicatedVideoMemory);
