VulkanRHI: Fix assert in Uniform Buffer

#rnx

diff --git a/Engine/Source/Runtime/VulkanRHI/Private/VulkanUniformBuffer.cpp b/Engine/Source/Runtime/VulkanRHI/Private/VulkanUniformBuffer.cpp
--- a/Engine/Source/Runtime/VulkanRHI/Private/VulkanUniformBuffer.cpp
+++ b/Engine/Source/Runtime/VulkanRHI/Private/VulkanUniformBuffer.cpp
@@ -185,7 +185,7 @@ FVulkanUniformBuffer::FVulkanUniformBuffer(FVulkanDevice& InDevice, const FRHIUn
 
 	if (InLayout->ConstantBufferSize > 0)
 	{
-		const bool bInRenderingThread = IsInRenderingThread();
+		const bool bInRenderingThread = IsInParallelRenderingThread();
 		const bool bInRHIThread = IsInRHIThread();
 
 		if (UseRingBuffer(InUsage) && (bInRenderingThread || bInRHIThread))
