Rollback 22308522.

diff --git a/Engine/Source/Runtime/VulkanRHI/Private/VulkanUniformBuffer.cpp b/Engine/Source/Runtime/VulkanRHI/Private/VulkanUniformBuffer.cpp
--- a/Engine/Source/Runtime/VulkanRHI/Private/VulkanUniformBuffer.cpp
+++ b/Engine/Source/Runtime/VulkanRHI/Private/VulkanUniformBuffer.cpp
@@ -185,7 +185,7 @@ FVulkanUniformBuffer::FVulkanUniformBuffer(FVulkanDevice& InDevice, const FRHIUn
 
 	if (InLayout->ConstantBufferSize > 0)
 	{
-		const bool bInRenderingThread = IsInParallelRenderingThread();
+		const bool bInRenderingThread = IsInRenderingThread();
 		const bool bInRHIThread = IsInRHIThread();
 
 		if (UseRingBuffer(InUsage) && (bInRenderingThread || bInRHIThread))
