VulkanRHI: FVulkanGPUTiming: Non-blocking queries were reset right after copying result, without knowing the availability of the results. Add a retry (and eventually a discard) for non-blocking queries.  Make it less likely to hit the code path that flushes everything.  Cache the MappedPointer for Query result buffers to avoid the return trip (and lock) in VulkanMemory since they never change.  Fix the return value of GetTiming that would return a sum if multiple results are available instead of only the latest one (as per function description).   Add support for VK_EXT_calibrated_timestamps to fix assert with RADV on startup by reading calibration value instead of a query result (that might not be avail).
NOTE: will add matching CPU timestamp from VK_EXT_calibrated_timestamps in a follow up changelist.

diff --git a/Engine/Source/Runtime/VulkanRHI/Private/VulkanCommandBuffer.cpp b/Engine/Source/Runtime/VulkanRHI/Private/VulkanCommandBuffer.cpp
--- a/Engine/Source/Runtime/VulkanRHI/Private/VulkanCommandBuffer.cpp
+++ b/Engine/Source/Runtime/VulkanRHI/Private/VulkanCommandBuffer.cpp
@@ -260,7 +260,10 @@ void FVulkanCmdBuffer::End()
 		const uint32 Stride = sizeof(uint64) * Width;
 
 		VulkanRHI::vkCmdCopyQueryPoolResults(GetHandle(), PoolHandle, Index, Query.Count, BufferHandle, Stride * Index, Stride, VK_QUERY_RESULT_64_BIT | BlockingFlags);
-		VulkanRHI::vkCmdResetQueryPool(GetHandle(), PoolHandle, Index, Query.Count);
+		if (Query.bBlocking)
+		{
+			VulkanRHI::vkCmdResetQueryPool(GetHandle(), PoolHandle, Index, Query.Count);
+		}
 	}
 
 	PendingTimestampQueries.Reset();
diff --git a/Engine/Source/Runtime/VulkanRHI/Private/VulkanCommandWrappers.h b/Engine/Source/Runtime/VulkanRHI/Private/VulkanCommandWrappers.h
--- a/Engine/Source/Runtime/VulkanRHI/Private/VulkanCommandWrappers.h
+++ b/Engine/Source/Runtime/VulkanRHI/Private/VulkanCommandWrappers.h
@@ -205,6 +205,8 @@ struct FWrapLayer
 #endif
 	static void GetDeviceImageMemoryRequirementsKHR(VkResult Result, VkDevice Device, const VkDeviceImageMemoryRequirements* Info, VkMemoryRequirements2* MemoryRequirements) VULKAN_LAYER_BODY
 	static void ResetQueryPoolEXT(VkResult Result, VkDevice Device, VkQueryPool QueryPool, uint32_t FirstQuery, uint32_t QueryCount) VULKAN_LAYER_BODY
+	static void GetPhysicalDeviceCalibrateableTimeDomainsEXT(VkResult Result, VkPhysicalDevice PhysicalDevice, uint32_t* TimeDomainCount, VkTimeDomainEXT* TimeDomains) VULKAN_LAYER_BODY
+	static void GetCalibratedTimestampsEXT(VkResult Result, VkDevice Device, uint32_t TimestampCount, const VkCalibratedTimestampInfoEXT* TimestampInfos, uint64_t* Timestamps, uint64_t* MaxDeviation) VULKAN_LAYER_BODY
 };
 
 #undef VULKAN_LAYER_BODY
@@ -1500,6 +1502,22 @@ namespace VulkanRHI
 		FWrapLayer::ResetQueryPoolEXT(VK_SUCCESS, Device, QueryPool, FirstQuery, QueryCount);
 	}
 
+	static FORCEINLINE_DEBUGGABLE VkResult vkGetPhysicalDeviceCalibrateableTimeDomainsEXT(VkPhysicalDevice PhysicalDevice, uint32_t* TimeDomainCount, VkTimeDomainEXT* TimeDomains)
+	{
+		FWrapLayer::GetPhysicalDeviceCalibrateableTimeDomainsEXT(VK_RESULT_MAX_ENUM, PhysicalDevice, TimeDomainCount, TimeDomains);
+		const VkResult Result = VULKANAPINAMESPACE::vkGetPhysicalDeviceCalibrateableTimeDomainsEXT(PhysicalDevice, TimeDomainCount, TimeDomains);
+		FWrapLayer::GetPhysicalDeviceCalibrateableTimeDomainsEXT(Result, PhysicalDevice, TimeDomainCount, TimeDomains);
+		return Result;
+	}
+
+	static FORCEINLINE_DEBUGGABLE VkResult vkGetCalibratedTimestampsEXT(VkDevice Device, uint32_t TimestampCount, const VkCalibratedTimestampInfoEXT* TimestampInfos, uint64_t* Timestamps, uint64_t* MaxDeviation)
+	{
+		FWrapLayer::GetCalibratedTimestampsEXT(VK_RESULT_MAX_ENUM, Device, TimestampCount, TimestampInfos, Timestamps, MaxDeviation);
+		const VkResult Result = VULKANAPINAMESPACE::vkGetCalibratedTimestampsEXT(Device, TimestampCount, TimestampInfos, Timestamps, MaxDeviation);
+		FWrapLayer::GetCalibratedTimestampsEXT(Result, Device, TimestampCount, TimestampInfos, Timestamps, MaxDeviation);
+		return Result;
+	}
+
 #if VULKAN_ENABLE_IMAGE_TRACKING_LAYER
 	void BindDebugLabelName(VkImage Image, const TCHAR* Name);
 #endif
diff --git a/Engine/Source/Runtime/VulkanRHI/Private/VulkanContext.h b/Engine/Source/Runtime/VulkanRHI/Private/VulkanContext.h
--- a/Engine/Source/Runtime/VulkanRHI/Private/VulkanContext.h
+++ b/Engine/Source/Runtime/VulkanRHI/Private/VulkanContext.h
@@ -83,6 +83,7 @@ public:
 	// Render time measurement
 	virtual void RHIBeginRenderQuery(FRHIRenderQuery* RenderQuery) final override;
 	virtual void RHIEndRenderQuery(FRHIRenderQuery* RenderQuery) final override;
+	virtual void RHICalibrateTimers(FRHITimestampCalibrationQuery* CalibrationQuery) final override;
 
 	virtual void RHISubmitCommandsHint() final override;
 
diff --git a/Engine/Source/Runtime/VulkanRHI/Private/VulkanDebug.cpp b/Engine/Source/Runtime/VulkanRHI/Private/VulkanDebug.cpp
--- a/Engine/Source/Runtime/VulkanRHI/Private/VulkanDebug.cpp
+++ b/Engine/Source/Runtime/VulkanRHI/Private/VulkanDebug.cpp
@@ -4465,6 +4465,25 @@ void FWrapLayer::ResetQueryPoolEXT(VkResult Result, VkDevice Device, VkQueryPool
 	}
 }
 
+void FWrapLayer::GetPhysicalDeviceCalibrateableTimeDomainsEXT(VkPhysicalDevice PhysicalDevice, uint32_t* TimeDomainCount, VkTimeDomainEXT* TimeDomains)
+{
+	if (Result == VK_RESULT_MAX_ENUM)
+	{
+#if VULKAN_ENABLE_DUMP_LAYER
+		PrintfBeginResult(FString::Printf(TEXT("PhysicalDeviceCalibrateableTimeDomainsEXT(PhysicalDevice=0x%p, TimeDomains=0x%p)"), PhysicalDevice, TimeDomains));
+#endif
+	}
+}
+
+void FWrapLayer::GetCalibratedTimestampsEXT(VkResult Result, VkDevice Device, uint32_t TimestampCount, const VkCalibratedTimestampInfoEXT* TimestampInfos, uint64_t* Timestamps, uint64_t* MaxDeviation)
+{
+	if (Result == VK_RESULT_MAX_ENUM)
+	{
+#if VULKAN_ENABLE_DUMP_LAYER
+		PrintfBeginResult(FString::Printf(TEXT("GetCalibratedTimestampsEXT(Device=0x%p, TimestampCount=%u)"), Device, TimestampCount));
+#endif
+	}
+}
 
 void FWrapLayer::CmdPipelineBarrier2KHR(VkResult Result, VkCommandBuffer CommandBuffer, const VkDependencyInfo* DependencyInfo)
 {
diff --git a/Engine/Source/Runtime/VulkanRHI/Private/VulkanDevice.cpp b/Engine/Source/Runtime/VulkanRHI/Private/VulkanDevice.cpp
--- a/Engine/Source/Runtime/VulkanRHI/Private/VulkanDevice.cpp
+++ b/Engine/Source/Runtime/VulkanRHI/Private/VulkanDevice.cpp
@@ -1443,3 +1443,30 @@ FVulkanTransientHeapCache& FVulkanDevice::GetOrCreateTransientHeapCache()
 	}
 	return *TransientHeapCache;
 }
+
+FGPUTimingCalibrationTimestamp FVulkanDevice::GetCalibrationTimestamp()
+{
+	auto ToMicroseconds = [](uint64_t Timestamp)
+	{
+		const double Frequency = double(FVulkanGPUTiming::GetTimingFrequency());
+		uint64 Microseconds = (uint64)((double(Timestamp) / Frequency) * 1000.0 * 1000.0);
+		return Microseconds;
+	};
+
+	FGPUTimingCalibrationTimestamp CalibrationTimestamp;
+	if (OptionalDeviceExtensions.HasEXTCalibratedTimestamps)
+	{
+		VkCalibratedTimestampInfoEXT TimestampInfo;
+		ZeroVulkanStruct(TimestampInfo, VK_STRUCTURE_TYPE_CALIBRATED_TIMESTAMP_INFO_EXT);
+		TimestampInfo.timeDomain = VK_TIME_DOMAIN_DEVICE_EXT;
+
+		uint64_t GPUTimestamp = 0;
+		uint64_t MaxDeviation = 0;
+		VERIFYVULKANRESULT(VulkanRHI::vkGetCalibratedTimestampsEXT(Device, 1, &TimestampInfo, &GPUTimestamp, &MaxDeviation));
+		CalibrationTimestamp.GPUMicroseconds = ToMicroseconds(GPUTimestamp);
+
+		const uint64 CPUTimestamp = FPlatformTime::Cycles64();
+		CalibrationTimestamp.CPUMicroseconds = uint64(FPlatformTime::ToSeconds64(CPUTimestamp) * 1e6);
+	}
+	return CalibrationTimestamp;
+}
diff --git a/Engine/Source/Runtime/VulkanRHI/Private/VulkanDevice.h b/Engine/Source/Runtime/VulkanRHI/Private/VulkanDevice.h
--- a/Engine/Source/Runtime/VulkanRHI/Private/VulkanDevice.h
+++ b/Engine/Source/Runtime/VulkanRHI/Private/VulkanDevice.h
@@ -54,6 +54,7 @@ struct FOptionalVulkanDeviceExtensions
 			uint64 HasRayTracingPipeline : 1;
 			uint64 HasRayQuery : 1;
 			uint64 HasDeferredHostOperations : 1;
+			uint64 HasEXTCalibratedTimestamps : 1;
 
 			// Vendor specific
 			uint64 HasAMDBufferMarker : 1;
@@ -495,6 +496,10 @@ public:
 		bDebugMarkersFound = true;
 	}
 
+	// Performs a GPU and CPU timestamp at nearly the same time.
+	// This allows aligning GPU and CPU events on the same timeline in profile visualization.
+	FGPUTimingCalibrationTimestamp GetCalibrationTimestamp();
+
 private:
 	const VkFormatProperties& GetFormatProperties(VkFormat InFormat);
 	void MapBufferFormatSupport(FPixelFormatInfo& PixelFormatInfo, EPixelFormat UEFormat, VkFormat VulkanFormat);
diff --git a/Engine/Source/Runtime/VulkanRHI/Private/VulkanExtensions.cpp b/Engine/Source/Runtime/VulkanRHI/Private/VulkanExtensions.cpp
--- a/Engine/Source/Runtime/VulkanRHI/Private/VulkanExtensions.cpp
+++ b/Engine/Source/Runtime/VulkanRHI/Private/VulkanExtensions.cpp
@@ -961,6 +961,36 @@ private:
 };
 
 
+// ***** VK_EXT_calibrated_timestamps
+class FVulkanEXTCalibratedTimestampsExtension : public FVulkanDeviceExtension
+{
+public:
+	FVulkanEXTCalibratedTimestampsExtension(FVulkanDevice* InDevice)
+		: FVulkanDeviceExtension(InDevice, VK_EXT_CALIBRATED_TIMESTAMPS_EXTENSION_NAME, VULKAN_EXTENSION_ENABLED, VULKAN_EXTENSION_NOT_PROMOTED)
+	{
+	}
+
+	virtual void PostPhysicalDeviceFeatures(FOptionalVulkanDeviceExtensions& ExtensionFlags) override final
+	{
+		uint32 TimeDomainCount = 0;
+		VulkanRHI::vkGetPhysicalDeviceCalibrateableTimeDomainsEXT(Device->GetPhysicalHandle(), &TimeDomainCount, nullptr);
+
+		TArray<VkTimeDomainEXT, TInlineAllocator<4>> TimeDomains;
+		TimeDomains.SetNumZeroed(TimeDomainCount);
+		VulkanRHI::vkGetPhysicalDeviceCalibrateableTimeDomainsEXT(Device->GetPhysicalHandle(), &TimeDomainCount, &TimeDomains[0]);
+
+		for (VkTimeDomainEXT TimeDomain : TimeDomains)
+		{
+			if (TimeDomain == VK_TIME_DOMAIN_DEVICE_EXT)
+			{
+				ExtensionFlags.HasEXTCalibratedTimestamps = 1;
+				break;
+			}
+		}
+	}
+};
+
+
 
 template <typename ExtensionType>
 static void FlagExtensionSupport(const TArray<VkExtensionProperties>& ExtensionProperties, TArray<TUniquePtr<ExtensionType>>& UEExtensions, const TCHAR* ExtensionTypeName)
@@ -1033,6 +1063,7 @@ FVulkanDeviceExtensionArray FVulkanDeviceExtension::GetUESupportedDeviceExtensio
 	ADD_CUSTOM_EXTENSION(FVulkanEXTDescriptorIndexingExtension);
 	ADD_CUSTOM_EXTENSION(FVulkanEXTHostQueryResetExtension);
 	ADD_CUSTOM_EXTENSION(FVulkanEXTSubgroupSizeControlExtension);
+	ADD_CUSTOM_EXTENSION(FVulkanEXTCalibratedTimestampsExtension);
 
 	// Needed for Raytracing
 	ADD_CUSTOM_EXTENSION(FVulkanKHRBufferDeviceAddressExtension);
diff --git a/Engine/Source/Runtime/VulkanRHI/Private/VulkanGPUProfiler.h b/Engine/Source/Runtime/VulkanRHI/Private/VulkanGPUProfiler.h
--- a/Engine/Source/Runtime/VulkanRHI/Private/VulkanGPUProfiler.h
+++ b/Engine/Source/Runtime/VulkanRHI/Private/VulkanGPUProfiler.h
@@ -63,6 +63,17 @@ public:
 	static void CalibrateTimers(FVulkanCommandListContext& InCmdContext);
 
 private:
+
+	/**
+	 * If results weren't available when first queried, make sure we try again.
+	 */
+	void RefreshPendingQueries(FVulkanCmdBuffer& CmdBuffer);
+
+	/**
+	 * Remove expired queries to make space for new ones
+	 */
+	void DiscardExpiredQueries(bool bNeedSpace);
+
 	/**
 	 * Initializes the static variables, if necessary.
 	 */
diff --git a/Engine/Source/Runtime/VulkanRHI/Private/VulkanLoader.h b/Engine/Source/Runtime/VulkanRHI/Private/VulkanLoader.h
--- a/Engine/Source/Runtime/VulkanRHI/Private/VulkanLoader.h
+++ b/Engine/Source/Runtime/VulkanRHI/Private/VulkanLoader.h
@@ -201,7 +201,10 @@
 	EnumMacro(PFN_vkCmdSetEvent2KHR, vkCmdSetEvent2KHR) \
 	EnumMacro(PFN_vkCmdWaitEvents2KHR, vkCmdWaitEvents2KHR) \
 	EnumMacro(PFN_vkQueueSubmit2KHR, vkQueueSubmit2KHR) \
-	EnumMacro(PFN_vkCreateSharedSwapchainsKHR, vkCreateSharedSwapchainsKHR)
+	EnumMacro(PFN_vkCreateSharedSwapchainsKHR, vkCreateSharedSwapchainsKHR) \
+	EnumMacro(PFN_vkGetPhysicalDeviceCalibrateableTimeDomainsEXT, vkGetPhysicalDeviceCalibrateableTimeDomainsEXT) \
+	EnumMacro(PFN_vkGetCalibratedTimestampsEXT, vkGetCalibratedTimestampsEXT)
+
 
 // List of all Vulkan entry points
 #define ENUM_VK_ENTRYPOINTS_ALL(EnumMacro) \
diff --git a/Engine/Source/Runtime/VulkanRHI/Private/VulkanQuery.cpp b/Engine/Source/Runtime/VulkanRHI/Private/VulkanQuery.cpp
--- a/Engine/Source/Runtime/VulkanRHI/Private/VulkanQuery.cpp
+++ b/Engine/Source/Runtime/VulkanRHI/Private/VulkanQuery.cpp
@@ -517,7 +517,7 @@ bool FVulkanDynamicRHI::RHIGetRenderQueryResult(FRHIRenderQuery* QueryRHI, uint6
 				if (StartQuerySyncPoint.FenceCounter < StartQuerySyncPoint.CmdBuffer->GetFenceSignaledCounter())
 				{
 					Query->Pool->ResultsBuffer->InvalidateMappedMemory();
-					uint64* Data = (uint64*)Query->Pool->ResultsBuffer->GetMappedPointer();
+					const uint64* Data = Query->Pool->MappedPointer;
 					OutNumPixels = ToMicroseconds(Data[TimestampIndex]);
 					return true;
 				}
@@ -560,7 +560,7 @@ bool FVulkanDynamicRHI::RHIGetRenderQueryResult(FRHIRenderQuery* QueryRHI, uint6
 			GRenderThreadIdle[ERenderThreadIdleTypes::WaitingForGPUQuery] += FPlatformTime::Cycles() - IdleStart;
 			GRenderThreadNumIdle[ERenderThreadIdleTypes::WaitingForGPUQuery]++;
 
-			uint64* Data = (uint64*)Query->Pool->ResultsBuffer->GetMappedPointer();
+			const uint64* Data = Query->Pool->MappedPointer;
 			OutNumPixels = ToMicroseconds(Data[TimestampIndex]);
 			return true;
 		}
@@ -631,6 +631,7 @@ void FVulkanCommandListContext::RHIEndRenderQuery(FRHIRenderQuery* QueryRHI)
 		{
 			Query->Pool = new FVulkanTimingQueryPool(Device, CommandBufferManager, 4);
 			Query->Pool->ResultsBuffer = Device->GetStagingManager().AcquireBuffer(Query->Pool->BufferSize * sizeof(uint64), VK_BUFFER_USAGE_TRANSFER_DST_BIT, VK_MEMORY_PROPERTY_HOST_COHERENT_BIT);
+			Query->Pool->MappedPointer = (uint64*)Query->Pool->ResultsBuffer->GetMappedPointer();
 		}
 
 		Query->Pool->CurrentTimestamp = (Query->Pool->CurrentTimestamp + 1) % Query->Pool->BufferSize;
@@ -646,6 +647,16 @@ void FVulkanCommandListContext::RHIEndRenderQuery(FRHIRenderQuery* QueryRHI)
 	}
 }
 
+void FVulkanCommandListContext::RHICalibrateTimers(FRHITimestampCalibrationQuery* CalibrationQuery)
+{
+	if (Device->GetOptionalExtensions().HasEXTCalibratedTimestamps)
+	{
+		FGPUTimingCalibrationTimestamp CalibrationTimestamp = Device->GetCalibrationTimestamp();
+		CalibrationQuery->CPUMicroseconds[0] = CalibrationTimestamp.CPUMicroseconds;
+		CalibrationQuery->GPUMicroseconds[0] = CalibrationTimestamp.GPUMicroseconds;
+	}
+}
+
 void FVulkanCommandListContext::WriteBeginTimestamp(FVulkanCmdBuffer* CmdBuffer)
 {
 	FrameTiming->StartTiming(CmdBuffer);
diff --git a/Engine/Source/Runtime/VulkanRHI/Private/VulkanRHIPrivate.h b/Engine/Source/Runtime/VulkanRHI/Private/VulkanRHIPrivate.h
--- a/Engine/Source/Runtime/VulkanRHI/Private/VulkanRHIPrivate.h
+++ b/Engine/Source/Runtime/VulkanRHI/Private/VulkanRHIPrivate.h
@@ -56,6 +56,7 @@
 
 #endif
 
+#include "GPUProfiler.h"
 #include "VulkanDevice.h"
 #include "VulkanQueue.h"
 #include "VulkanCommandBuffer.h"
diff --git a/Engine/Source/Runtime/VulkanRHI/Private/VulkanResources.h b/Engine/Source/Runtime/VulkanRHI/Private/VulkanResources.h
--- a/Engine/Source/Runtime/VulkanRHI/Private/VulkanResources.h
+++ b/Engine/Source/Runtime/VulkanRHI/Private/VulkanResources.h
@@ -696,10 +696,12 @@ public:
 		FVulkanCmdBuffer* CmdBuffer;
 		uint64 FenceCounter;
 		uint64 FrameCount = UINT64_MAX;
+		uint32 Attempts = 0;
 	};
 	TArray<FCmdBufferFence> TimestampListHandles;
 
 	VulkanRHI::FStagingBuffer* ResultsBuffer = nullptr;
+	uint64* MappedPointer = nullptr;
 };
 
 class FVulkanRenderQuery : public FRHIRenderQuery
diff --git a/Engine/Source/Runtime/VulkanRHI/Private/VulkanUtil.cpp b/Engine/Source/Runtime/VulkanRHI/Private/VulkanUtil.cpp
--- a/Engine/Source/Runtime/VulkanRHI/Private/VulkanUtil.cpp
+++ b/Engine/Source/Runtime/VulkanRHI/Private/VulkanUtil.cpp
@@ -26,6 +26,7 @@ extern CORE_API bool GIsGPUCrashed;
 static FString		EventDeepString(TEXT("EventTooDeep"));
 static const uint32	EventDeepCRC = FCrc::StrCrc32<TCHAR>(*EventDeepString);
 static const uint32 BUFFERED_TIMING_QUERIES = 1;
+static const uint32 TIMING_QUERY_RETRIES = 1;
 
 /**
  * Initializes the static variables, if necessary.
@@ -54,35 +55,12 @@ void FVulkanGPUTiming::PlatformStaticInitialize(void* UserData)
 
 void FVulkanGPUTiming::CalibrateTimers(FVulkanCommandListContext& InCmdContext)
 {
-#if VULKAN_USE_NEW_QUERIES
-
-	// TODO: Implement VULKAN_USE_NEW_QUERIES version
-
-#else
 	FVulkanDevice* Device = InCmdContext.GetDevice();
-	FVulkanRenderQuery* TimestampQuery = new FVulkanRenderQuery(RQT_AbsoluteTime);
-
-	{
-		FVulkanCmdBuffer* CmdBuffer = InCmdContext.GetCommandBufferManager()->GetUploadCmdBuffer();
-		InCmdContext.EndRenderQueryInternal(CmdBuffer, TimestampQuery);
-		InCmdContext.GetCommandBufferManager()->SubmitUploadCmdBuffer();
-	}
-
-	uint64 CPUTimestamp = 0;
-	uint64 GPUTimestampMicroseconds = 0;
-
-	const bool bWait = true;
-	if (TimestampQuery->GetResult(Device, GPUTimestampMicroseconds, bWait))
+	if (Device->GetOptionalExtensions().HasEXTCalibratedTimestamps)
 	{
-		CPUTimestamp = FPlatformTime::Cycles64();
-
-		GCalibrationTimestamp.CPUMicroseconds = uint64(FPlatformTime::ToSeconds64(CPUTimestamp) * 1e6);
-		GCalibrationTimestamp.GPUMicroseconds = GPUTimestampMicroseconds;
+		FGPUTimingCalibrationTimestamp CalibrationTimestamp = Device->GetCalibrationTimestamp();
+		SetCalibrationTimestamp(CalibrationTimestamp);
 	}
-
-	delete TimestampQuery;
-
-#endif
 }
 
 void FVulkanDynamicRHI::RHICalibrateTimers()
@@ -169,6 +147,7 @@ void FVulkanGPUTiming::Initialize(uint32 PoolSize)
 		check(!Pool);
 		Pool = new FVulkanTimingQueryPool(Device, CmdContext->GetCommandBufferManager(), PoolSize);
 		Pool->ResultsBuffer = Device->GetStagingManager().AcquireBuffer(Pool->GetMaxQueries() * sizeof(uint64) * 2, VK_BUFFER_USAGE_TRANSFER_DST_BIT, VK_MEMORY_PROPERTY_HOST_COHERENT_BIT);
+		Pool->MappedPointer = (uint64*)Pool->ResultsBuffer->GetMappedPointer();
 	}
 }
 
@@ -198,19 +177,39 @@ void FVulkanGPUTiming::StartTiming(FVulkanCmdBuffer* CmdBuffer)
 		{
 			CmdBuffer = CmdContext->GetCommandBufferManager()->GetActiveCmdBuffer();
 		}
+
+		// In case we aren't reading queries, remove oldest
+		if (NumPendingQueries >= Pool->BufferSize)
+		{
+			DiscardExpiredQueries(true);
+
+			if (NumPendingQueries >= Pool->BufferSize)
+			{
+				return;
+			}
+		}
+
+		RefreshPendingQueries(*CmdBuffer);
+
 		Pool->CurrentTimestamp = (Pool->CurrentTimestamp + 1) % Pool->BufferSize;
 		const uint32 QueryStartIndex = Pool->CurrentTimestamp * 2;
 
-		// If host query resets are supported, reset timestamp queries before writing to them (no need to consider host/GPU sync)
+		// Since we use non-blocking queries, we need to reset them before using them
+		// Use host query resets if supported since there's no need to consider host/GPU sync
 		if (Device->GetOptionalExtensions().HasEXTHostQueryReset)
 		{
 			VulkanRHI::vkResetQueryPoolEXT(Device->GetInstanceHandle(), Pool->GetHandle(), QueryStartIndex, 2);
 		}
+		else
+		{
+			VulkanRHI::vkCmdResetQueryPool(CmdBuffer->GetHandle(), Pool->GetHandle(), QueryStartIndex, 2);
+		}
 
 		VulkanRHI::vkCmdWriteTimestamp(CmdBuffer->GetHandle(), VK_PIPELINE_STAGE_BOTTOM_OF_PIPE_BIT, Pool->GetHandle(), QueryStartIndex);
 		Pool->TimestampListHandles[QueryStartIndex].CmdBuffer = CmdBuffer;
 		Pool->TimestampListHandles[QueryStartIndex].FenceCounter = CmdBuffer->GetFenceSignaledCounter();
 		Pool->TimestampListHandles[QueryStartIndex].FrameCount = CmdContext->GetFrameCounter();
+		Pool->TimestampListHandles[QueryStartIndex].Attempts = 0;
 		bIsTiming = true;
 	}
 }
@@ -233,21 +232,17 @@ void FVulkanGPUTiming::EndTiming(FVulkanCmdBuffer* CmdBuffer)
 		// Keep Start and End contiguous to fetch them together with a single AddPendingTimestampQuery(QueryStartIndex,2,...)
 		const uint32 QueryEndIndex = QueryStartIndex + 1;   
 
-		// In case we aren't reading queries, remove oldest
-		if (NumPendingQueries >= Pool->BufferSize)
-		{
-			PendingQueries.Pop();
-			NumPendingQueries--;
-		}
-
 		PendingQueries.Enqueue(Pool->CurrentTimestamp);
 		NumPendingQueries++;
 
 		VulkanRHI::vkCmdWriteTimestamp(CmdBuffer->GetHandle(), VK_PIPELINE_STAGE_BOTTOM_OF_PIPE_BIT, Pool->GetHandle(), QueryEndIndex);
 		CmdBuffer->AddPendingTimestampQuery(QueryStartIndex, 2, Pool->GetHandle(), Pool->ResultsBuffer->GetHandle(), false);
+		Pool->TimestampListHandles[QueryStartIndex].CmdBuffer = CmdBuffer;
+		Pool->TimestampListHandles[QueryStartIndex].FenceCounter = CmdBuffer->GetFenceSignaledCounter();
 		Pool->TimestampListHandles[QueryEndIndex].CmdBuffer = CmdBuffer;
 		Pool->TimestampListHandles[QueryEndIndex].FenceCounter = CmdBuffer->GetFenceSignaledCounter();
 		Pool->TimestampListHandles[QueryEndIndex].FrameCount = CmdContext->GetFrameCounter();
+		Pool->TimestampListHandles[QueryEndIndex].Attempts = 0;
 		Pool->NumIssuedTimestamps = FMath::Min<uint32>(Pool->NumIssuedTimestamps + 1, Pool->BufferSize);
 
 		bIsTiming = false;
@@ -255,6 +250,93 @@ void FVulkanGPUTiming::EndTiming(FVulkanCmdBuffer* CmdBuffer)
 	}
 }
 
+void FVulkanGPUTiming::DiscardExpiredQueries(bool bNeedSpace)
+{
+	// Returns true if we can safely throw away the query
+	auto CanBeDiscarded = [this](uint32 QueryIndex)
+	{
+		FVulkanTimingQueryPool::FCmdBufferFence& QuerySyncPoint = Pool->TimestampListHandles[QueryIndex];
+		if (QuerySyncPoint.CmdBuffer && (QuerySyncPoint.FenceCounter < QuerySyncPoint.CmdBuffer->GetFenceSignaledCounter()))
+		{
+			const uint64 Availability = Pool->MappedPointer[QueryIndex * 2 + 1];
+
+			if (Availability || (QuerySyncPoint.Attempts < TIMING_QUERY_RETRIES))
+			{
+				return false;
+			}
+		}
+		return true;
+	};
+
+	auto IsQueryExpired = [&](uint32 QueryIndex)
+	{
+		bool QueryExpired = CanBeDiscarded(QueryIndex);
+		if (!QueryExpired && bNeedSpace)
+		{
+			// If the hold back for a query is the fence and the command buffer was submitted, take a shot at updating its status
+			FVulkanTimingQueryPool::FCmdBufferFence& QuerySyncPoint = Pool->TimestampListHandles[QueryIndex];
+			if ((QuerySyncPoint.FenceCounter >= QuerySyncPoint.CmdBuffer->GetFenceSignaledCounter()) && QuerySyncPoint.CmdBuffer->IsSubmitted())
+			{
+				QuerySyncPoint.CmdBuffer->GetOwner()->RefreshFenceStatus();
+				if (QuerySyncPoint.FenceCounter < QuerySyncPoint.CmdBuffer->GetFenceSignaledCounter())
+				{
+					QueryExpired = CanBeDiscarded(QueryIndex);
+				}
+			}
+		}
+		return QueryExpired;
+	};
+
+	Pool->ResultsBuffer->InvalidateMappedMemory();
+
+	uint32 OldestTimestamp;
+	while (PendingQueries.Peek(OldestTimestamp))
+	{
+		const uint32 QueryStartIndex = OldestTimestamp * 2;
+		const uint32 QueryEndIndex = QueryStartIndex + 1;
+
+		const bool StartQueryExpired = IsQueryExpired(QueryStartIndex);
+		const bool EndQueryExpired = IsQueryExpired(QueryEndIndex);
+
+		if (StartQueryExpired || EndQueryExpired)
+		{
+			// If the queries is no good anymore, throw it away and try the next one
+			PendingQueries.Pop();
+			NumPendingQueries--;
+		}
+		else
+		{
+			break;
+		}
+	}
+}
+
+void FVulkanGPUTiming::RefreshPendingQueries(FVulkanCmdBuffer& CmdBuffer)
+{
+	for (int32 QueryIndex=0; QueryIndex < Pool->TimestampListHandles.Num(); ++QueryIndex)
+	{
+		FVulkanTimingQueryPool::FCmdBufferFence& QuerySyncPoint = Pool->TimestampListHandles[QueryIndex];
+
+		if (QuerySyncPoint.CmdBuffer && (QuerySyncPoint.FenceCounter < QuerySyncPoint.CmdBuffer->GetFenceSignaledCounter()))
+		{
+			Pool->ResultsBuffer->InvalidateMappedMemory();
+			const uint64 Availability = Pool->MappedPointer[QueryIndex * 2 + 1];
+
+			if (!Availability && (QuerySyncPoint.Attempts < TIMING_QUERY_RETRIES))
+			{
+				QuerySyncPoint.CmdBuffer = &CmdBuffer;
+				QuerySyncPoint.FenceCounter = CmdBuffer.GetFenceSignaledCounter();
+				QuerySyncPoint.FrameCount = CmdContext->GetFrameCounter();
+				++QuerySyncPoint.Attempts;
+
+				CmdBuffer.AddPendingTimestampQuery(
+					QueryIndex, 1, Pool->GetHandle(), Pool->ResultsBuffer->GetHandle(), false);
+			}
+		}
+	}
+}
+
+
 /**
  * Retrieves the most recently resolved timing measurement.
  * The unit is the same as for FPlatformTime::Cycles(). Returns 0 if there are no resolved measurements.
@@ -281,6 +363,8 @@ uint64 FVulkanGPUTiming::GetTiming(bool bGetCurrentResultsAndBlock)
 
 		PreviousFrame = FrameCount;
 
+		DiscardExpiredQueries(false);
+
 		while (PendingQueries.Peek(TimestampIndex))
 		{
 			const uint32 QueryStartIndex = TimestampIndex * 2;
@@ -294,7 +378,7 @@ uint64 FVulkanGPUTiming::GetTiming(bool bGetCurrentResultsAndBlock)
 
 			if (!bBlocking)
 			{
-				uint64_t QueryFrameCount = EndQuerySyncPoint.FrameCount;
+				const uint64_t QueryFrameCount = EndQuerySyncPoint.FrameCount;
 				
 				// Allow queries to back up if we are non-blocking
 				if (FrameCount < BUFFERED_TIMING_QUERIES || FrameCount - BUFFERED_TIMING_QUERIES < QueryFrameCount)
@@ -313,11 +397,10 @@ uint64 FVulkanGPUTiming::GetTiming(bool bGetCurrentResultsAndBlock)
 				StartQuerySyncPoint.FenceCounter < StartQuerySyncPoint.CmdBuffer->GetFenceSignaledCounter())
 			{
 				Pool->ResultsBuffer->InvalidateMappedMemory();
-				uint64* Data = (uint64*)Pool->ResultsBuffer->GetMappedPointer();
-				StartTimeAvailability = Data[QueryStartIndex * 2 + 1];
-				EndTimeAvailability = Data[QueryEndIndex * 2 + 1];
-				StartTime = Data[QueryStartIndex * 2];
-				EndTime = Data[QueryEndIndex * 2];
+				StartTimeAvailability = Pool->MappedPointer[QueryStartIndex * 2 + 1];
+				EndTimeAvailability = Pool->MappedPointer[QueryEndIndex * 2 + 1];
+				StartTime = Pool->MappedPointer[QueryStartIndex * 2];
+				EndTime = Pool->MappedPointer[QueryEndIndex * 2];
 
 				if (!StartTimeAvailability || !EndTimeAvailability)
 				{
@@ -329,7 +412,8 @@ uint64 FVulkanGPUTiming::GetTiming(bool bGetCurrentResultsAndBlock)
 
 				if (EndTime > StartTime)
 				{
-					TotalTime += EndTime - StartTime;
+					// Only keep the most recent result
+					TotalTime = EndTime - StartTime;
 				}
 
 				continue;
@@ -363,11 +447,10 @@ uint64 FVulkanGPUTiming::GetTiming(bool bGetCurrentResultsAndBlock)
 				GRenderThreadNumIdle[ERenderThreadIdleTypes::WaitingForGPUQuery]++;
 
 				Pool->ResultsBuffer->InvalidateMappedMemory();
-				uint64* Data = (uint64*)Pool->ResultsBuffer->GetMappedPointer();
-				StartTimeAvailability = Data[QueryStartIndex * 2 + 1];
-				EndTimeAvailability = Data[QueryEndIndex * 2 + 1];
-				StartTime = Data[QueryStartIndex * 2];
-				EndTime = Data[QueryEndIndex * 2];
+				StartTimeAvailability = Pool->MappedPointer[QueryStartIndex * 2 + 1];
+				EndTimeAvailability = Pool->MappedPointer[QueryEndIndex * 2 + 1];
+				StartTime = Pool->MappedPointer[QueryStartIndex * 2];
+				EndTime = Pool->MappedPointer[QueryEndIndex * 2];
 
 				PendingQueries.Pop();
 				NumPendingQueries--;
