PR #9416: Various fixes in memory reporting code (Contributed by [REDACTED_PERSON])

#rnx

diff --git a/Engine/Source/Runtime/VulkanRHI/Private/VulkanMemory.cpp b/Engine/Source/Runtime/VulkanRHI/Private/VulkanMemory.cpp
--- a/Engine/Source/Runtime/VulkanRHI/Private/VulkanMemory.cpp
+++ b/Engine/Source/Runtime/VulkanRHI/Private/VulkanMemory.cpp
@@ -1249,7 +1249,8 @@ namespace VulkanRHI
 			}
 			VULKAN_LOGMEMORY(TEXT("\t\tTotal Allocated %.2f MB, Peak %.2f MB"), TotalSize / 1024.0f / 1024.0f, HeapInfo.PeakSize / 1024.0f / 1024.0f);
 		}
-		VULKAN_LOGMEMORY(TEXT("Blocks Kept alive Current Frame (%d)\n"), GFrameNumberRenderThread);
+		VULKAN_LOGMEMORY(TEXT("Blocks Kept alive Current Frame (%d) "), GFrameNumberRenderThread);
+		VULKAN_LOGMEMORY(TEXT("Free Blocks in Allocations:"));
 		for (TPair<FDeviceMemoryBlockKey, FDeviceMemoryBlock>& Pair : Allocations)
 		{
 			FDeviceMemoryBlock& Block = Pair.Value;
@@ -1262,7 +1263,6 @@ namespace VulkanRHI
 
 #if VULKAN_OBJECT_TRACKING
 		{
-
 			TSortedMap<uint32, FVulkanMemoryBucket> AllocationBuckets;
 			auto Collector = [&](const TCHAR* Name, FName ResourceName, void* Address, void* RHIRes, uint32 Width, uint32 Height, uint32 Depth, uint32 Format)
 			{
@@ -2618,7 +2618,7 @@ namespace VulkanRHI
 			{
 				if(Data.State == FVulkanAllocationInternal::EALLOCATED)
 				{
-					UE_LOG(LogVulkanRHI, Warning, TEXT(" ** LEAK %03d [%08x-%08x] %d  %s \n%s"), LeakCount++, Data.AllocationOffset, Data.AllocationSize,  Data.Size, VulkanAllocationMetaTypeToString(Data.MetaType), *VULKAN_TRACK_STRING(Data.Track));
+					UE_LOG(LogVulkanRHI, Warning, TEXT(" ** LEAK %03d [%08x-%08x] %u  %s \n%s"), LeakCount++, Data.AllocationOffset, Data.AllocationSize,  Data.Size, VulkanAllocationMetaTypeToString(Data.MetaType), *VULKAN_TRACK_STRING(Data.Track));
 				}
 			}
 		}
@@ -2784,12 +2784,12 @@ namespace VulkanRHI
 		FDeviceMemoryAllocation* DeviceMemoryAllocation = DeviceMemoryManager->Alloc(true, MemReqs.size, MemoryTypeIndex, nullptr, Priority, false, File, Line);
 		if(!DeviceMemoryAllocation)
 		{
-			VkMemoryPropertyFlags MemoryPropertyFallbackFlags = MemoryPropertyFlags & (~VK_MEMORY_PROPERTY_DEVICE_LOCAL_BIT);
-			uint32 MemoryTypeIndexFallback;
-			if(GVulkanMemoryMemoryFallbackToHost && VK_SUCCESS == Device->GetDeviceMemoryManager().GetMemoryTypeFromPropertiesExcluding(MemReqs.memoryTypeBits, MemoryPropertyFallbackFlags, MemoryTypeIndex, &MemoryTypeIndexFallback))
+			MemoryPropertyFlags &= (~VK_MEMORY_PROPERTY_DEVICE_LOCAL_BIT);
+			const uint32 ForbiddenMemoryTypeIndex = MemoryTypeIndex;
+			if (GVulkanMemoryMemoryFallbackToHost && VK_SUCCESS == Device->GetDeviceMemoryManager().GetMemoryTypeFromPropertiesExcluding(MemReqs.memoryTypeBits, MemoryPropertyFlags, ForbiddenMemoryTypeIndex, &MemoryTypeIndex))
 			{
-				DeviceMemoryAllocation = DeviceMemoryManager->Alloc(false, MemReqs.size, MemoryTypeIndexFallback, nullptr, Priority, false, File, Line);
-			}			
+				DeviceMemoryAllocation = DeviceMemoryManager->Alloc(false, MemReqs.size, MemoryTypeIndex, nullptr, Priority, false, File, Line);
+			}
 		}
 		if(!DeviceMemoryAllocation)
 		{
@@ -3175,7 +3175,7 @@ namespace VulkanRHI
 				for (int32 Index = 0; Index < UsedAllocations.Num(); ++Index)
 				{
 					FVulkanSubresourceAllocator* BA = UsedAllocations[Index];
-					VULKAN_LOGMEMORY(TEXT("%6d 0x%016llx 0x%016llx 0x%06x 0x%08x %6d   %6d        %d/%d"), Index, (void*)BA->Buffer, (void*)BA->MemoryAllocation->GetHandle(), BA->MemoryPropertyFlags, BA->BufferUsageFlags, BA->NumSubAllocations, BA->FreeList.Num(), BA->UsedSize, BA->MaxSize);
+					VULKAN_LOGMEMORY(TEXT("%6d 0x%016llx 0x%016llx 0x%06x 0x%08x %6d   %6d        %u/%u"), Index, (void*)BA->Buffer, (void*)BA->MemoryAllocation->GetHandle(), BA->MemoryPropertyFlags, BA->BufferUsageFlags, BA->NumSubAllocations, BA->FreeList.Num(), BA->UsedSize, BA->MaxSize);
 
 					if (PoolSizeIndex == (int32)EPoolSizes::SizesCount)
 					{
@@ -3221,18 +3221,18 @@ namespace VulkanRHI
 
 				if (PoolSizeIndex == (int32)EPoolSizes::SizesCount)
 				{
-					VULKAN_LOGMEMORY(TEXT(" Large Alloc Used/Max %d/%d %6.2f%%"), _UsedLargeTotal, _AllocLargeTotal, 100.0f * (float)_UsedLargeTotal / (float)_AllocLargeTotal);
+					VULKAN_LOGMEMORY(TEXT(" Large Alloc Used/Max %llu/%llu %6.2f%%"), _UsedLargeTotal, _AllocLargeTotal, 100.0f * (float)_UsedLargeTotal / (float)_AllocLargeTotal);
 				}
 				else
 				{
-					VULKAN_LOGMEMORY(TEXT(" Binned [%d] Alloc Used/Max %d/%d %6.2f%%"), PoolSizes[PoolSizeIndex], _UsedBinnedTotal, _AllocBinnedTotal, 100.0f * (float)_UsedBinnedTotal / (float)_AllocBinnedTotal);
+					VULKAN_LOGMEMORY(TEXT(" Binned [%d] Alloc Used/Max %llu/%llu %6.2f%%"), PoolSizes[PoolSizeIndex], _UsedBinnedTotal, _AllocBinnedTotal, 100.0f * (float)_UsedBinnedTotal / (float)_AllocBinnedTotal);
 				}
 			}
 		}
 
 		VULKAN_LOGMEMORY(TEXT("::Totals::"));
-		VULKAN_LOGMEMORY(TEXT("Large Alloc Used/Max %d/%d %.2f%%"), UsedLargeTotal, AllocLargeTotal, 100.0f * AllocLargeTotal > 0 ? (float)UsedLargeTotal / (float)AllocLargeTotal : 0.0f);
-		VULKAN_LOGMEMORY(TEXT("Binned Alloc Used/Max %d/%d %.2f%%"), UsedBinnedTotal, AllocBinnedTotal, AllocBinnedTotal > 0 ? 100.0f * (float)UsedBinnedTotal / (float)AllocBinnedTotal : 0.0f);
+		VULKAN_LOGMEMORY(TEXT("Large Alloc Used/Max %llu/%llu %.2f%%"), UsedLargeTotal, AllocLargeTotal, AllocLargeTotal > 0 ? 100.0f * (float)UsedLargeTotal / (float)AllocLargeTotal : 0.0f);
+		VULKAN_LOGMEMORY(TEXT("Binned Alloc Used/Max %llu/%llu %.2f%%"), UsedBinnedTotal, AllocBinnedTotal, AllocBinnedTotal > 0 ? 100.0f * (float)UsedBinnedTotal / (float)AllocBinnedTotal : 0.0f);
 		{
 			FResourceHeapStats& DedicatedStats = OverallSummary[DedicatedAllocatorSummary];
 
@@ -3262,7 +3262,7 @@ namespace VulkanRHI
 
 		auto WriteLogLine = [](const FString& Name, FResourceHeapStats& Stat)
 		{
-			uint64 FreeMemory = Stat.TotalMemory - Stat.UsedBufferMemory - Stat.UsedImageMemory;
+			uint64 FreeMemory = Stat.TotalMemory - FMath::Min<uint64>(Stat.TotalMemory, Stat.UsedBufferMemory + Stat.UsedImageMemory);
 			FString HostString = VK_FLAGS_TO_STRING(VkMemoryPropertyFlags, Stat.MemoryFlags);
 			VULKAN_LOGMEMORY(TEXT("\t\t%-33s  |%8.2fmb / %8.2fmb / %11.2fmb / %11.2fmb | %10d %10d | %6d %6d %6d | %05x | %s"),
 				*Name,
