VulkanRHI: ArraySize was used instead of depth during defrag copies for Texture2dArray.

#rnx

diff --git a/Engine/Source/Runtime/VulkanRHI/Private/VulkanTexture.cpp b/Engine/Source/Runtime/VulkanRHI/Private/VulkanTexture.cpp
--- a/Engine/Source/Runtime/VulkanRHI/Private/VulkanTexture.cpp
+++ b/Engine/Source/Runtime/VulkanRHI/Private/VulkanTexture.cpp
@@ -586,13 +586,14 @@ static VkImageLayout GetInitialLayoutFromRHIAccess(ERHIAccess RHIAccess, ETextur
 void FVulkanTexture::InternalMoveSurface(FVulkanDevice& InDevice, FVulkanCommandListContext& Context, FVulkanAllocation& DestAllocation)
 {
 	FImageCreateInfo ImageCreateInfo;
-	FVulkanTexture::GenerateImageCreateInfo(ImageCreateInfo, InDevice, GetDesc(), &StorageFormat, &ViewFormat);
+	const FRHITextureDesc& Desc = GetDesc();
+	FVulkanTexture::GenerateImageCreateInfo(ImageCreateInfo, InDevice, Desc, &StorageFormat, &ViewFormat);
 
 	VkImage MovedImage;
 	VERIFYVULKANRESULT(VulkanRHI::vkCreateImage(InDevice.GetInstanceHandle(), &ImageCreateInfo.ImageCreateInfo, VULKAN_CPU_ALLOCATOR, &MovedImage));
 	checkf(Tiling == ImageCreateInfo.ImageCreateInfo.tiling, TEXT("Move has changed image tiling:  before [%s] != after [%s]"), VK_TYPE_TO_STRING(VkImageTiling, Tiling), VK_TYPE_TO_STRING(VkImageTiling, ImageCreateInfo.ImageCreateInfo.tiling));
 
-	const ETextureCreateFlags UEFlags = GetDesc().Flags;
+	const ETextureCreateFlags UEFlags = Desc.Flags;
 	const bool bRenderTarget = EnumHasAnyFlags(UEFlags, TexCreate_RenderTargetable | TexCreate_DepthStencilTargetable | TexCreate_ResolveTargetable);
 	const bool bCPUReadback = EnumHasAnyFlags(UEFlags, TexCreate_CPUReadback);
 	const bool bMemoryless = EnumHasAnyFlags(UEFlags, TexCreate_Memoryless);
@@ -638,16 +639,15 @@ void FVulkanTexture::InternalMoveSurface(FVulkanDevice& InDevice, FVulkanCommand
 			Barrier.Execute(VkCmdBuffer);
 		}
 		{
-			const FIntVector SizeXYZ = GetSizeXYZ();
 			VkImageCopy Regions[MAX_TEXTURE_MIP_COUNT];
-			check(GetDesc().NumMips <= MAX_TEXTURE_MIP_COUNT);
+			check(Desc.NumMips <= MAX_TEXTURE_MIP_COUNT);
 			FMemory::Memzero(Regions);
-			for (uint32 i = 0; i < GetDesc().NumMips; ++i)
+			for (uint32 i = 0; i < Desc.NumMips; ++i)
 			{
 				VkImageCopy& Region = Regions[i];
-				Region.extent.width = FMath::Max(1, SizeXYZ.X >> i);
-				Region.extent.height = FMath::Max(1, SizeXYZ.Y >> i);
-				Region.extent.depth = FMath::Max(1, SizeXYZ.Z >> i);
+				Region.extent.width = FMath::Max(1, Desc.Extent.X >> i);
+				Region.extent.height = FMath::Max(1, Desc.Extent.Y >> i);
+				Region.extent.depth = FMath::Max(1, Desc.Depth >> i);
 				Region.srcSubresource.aspectMask = FullAspectMask;
 				Region.dstSubresource.aspectMask = FullAspectMask;
 				Region.srcSubresource.baseArrayLayer = 0;
@@ -660,7 +660,7 @@ void FVulkanTexture::InternalMoveSurface(FVulkanDevice& InDevice, FVulkanCommand
 			VulkanRHI::vkCmdCopyImage(VkCmdBuffer,
 				Image, VK_IMAGE_LAYOUT_TRANSFER_SRC_OPTIMAL,
 				MovedImage, VK_IMAGE_LAYOUT_TRANSFER_DST_OPTIMAL,
-				GetDesc().NumMips, &Regions[0]);
+				Desc.NumMips, &Regions[0]);
 		}
 
 		// Put the destination image in exactly the same layout the original image was
