PR #9026: Pass clearbinding to RHICreateTexture2DFromResource in VulkanRHI (Contributed by [REDACTED_PERSON])

Modified to include all RHICreateTexture*FromResource functions.  Remove the check for Depth, it would be misleading to not have the specified value applied (default value of Transparent is equivalent to ZeroDepth anyhow).

#rnx

diff --git a/Engine/Plugins/Runtime/Oculus/OculusVR/Source/OculusHMD/Private/OculusHMD_CustomPresent_Vulkan.cpp b/Engine/Plugins/Runtime/Oculus/OculusVR/Source/OculusHMD/Private/OculusHMD_CustomPresent_Vulkan.cpp
--- a/Engine/Plugins/Runtime/Oculus/OculusVR/Source/OculusHMD/Private/OculusHMD_CustomPresent_Vulkan.cpp
+++ b/Engine/Plugins/Runtime/Oculus/OculusVR/Source/OculusHMD/Private/OculusHMD_CustomPresent_Vulkan.cpp
@@ -129,7 +129,7 @@ FTextureRHIRef FVulkanCustomPresent::CreateTexture_RenderThread(uint32 InSizeX,
 		return VulkanRHI->RHICreateTexture2DFromResource(InFormat, InSizeX, InSizeY, InNumMips, InNumSamples, (VkImage) InTexture, InTexCreateFlags).GetReference();
 
 	case RRT_Texture2DArray:
-		return VulkanRHI->RHICreateTexture2DArrayFromResource(InFormat, InSizeX, InSizeY, 2, InNumMips, InNumSamples, (VkImage) InTexture, InTexCreateFlags).GetReference();
+		return VulkanRHI->RHICreateTexture2DArrayFromResource(InFormat, InSizeX, InSizeY, 2, InNumMips, InNumSamples, (VkImage) InTexture, InTexCreateFlags, InBinding).GetReference();
 
 	case RRT_TextureCube:
 		return VulkanRHI->RHICreateTextureCubeFromResource(InFormat, InSizeX, false, 1, InNumMips, (VkImage) InTexture, InTexCreateFlags).GetReference();
diff --git a/Engine/Plugins/Runtime/OpenXR/Source/OpenXRHMD/Private/OpenXRHMD_Swapchain.cpp b/Engine/Plugins/Runtime/OpenXR/Source/OpenXRHMD/Private/OpenXRHMD_Swapchain.cpp
--- a/Engine/Plugins/Runtime/OpenXR/Source/OpenXRHMD/Private/OpenXRHMD_Swapchain.cpp
+++ b/Engine/Plugins/Runtime/OpenXR/Source/OpenXRHMD/Private/OpenXRHMD_Swapchain.cpp
@@ -416,7 +416,7 @@ FXRSwapChainPtr CreateSwapchain_Vulkan(XrSession InSession, uint8 Format, uint32
 	TArray<XrSwapchainImageVulkanKHR> Images = EnumerateImages<XrSwapchainImageVulkanKHR>(Swapchain, XR_TYPE_SWAPCHAIN_IMAGE_VULKAN_KHR);
 	for (const auto& Image : Images)
 	{
-		TextureChain.Add(static_cast<FTextureRHIRef>(VulkanRHI->RHICreateTexture2DArrayFromResource(GPixelFormats[Format].UnrealFormat, SizeX, SizeY, ArraySize, NumMips, NumSamples, Image.image, CreateFlags)));
+		TextureChain.Add(static_cast<FTextureRHIRef>(VulkanRHI->RHICreateTexture2DArrayFromResource(GPixelFormats[Format].UnrealFormat, SizeX, SizeY, ArraySize, NumMips, NumSamples, Image.image, CreateFlags, ClearValueBinding)));
 	}
 	FTextureRHIRef ChainTarget = static_cast<FTextureRHIRef>(VulkanRHI->RHICreateAliasedTexture((FTextureRHIRef&)TextureChain[0]));
 
diff --git a/Engine/Source/Runtime/VulkanRHI/Private/VulkanRHI.cpp b/Engine/Source/Runtime/VulkanRHI/Private/VulkanRHI.cpp
--- a/Engine/Source/Runtime/VulkanRHI/Private/VulkanRHI.cpp
+++ b/Engine/Source/Runtime/VulkanRHI/Private/VulkanRHI.cpp
@@ -1431,11 +1431,11 @@ void FVulkanDynamicRHI::RHISubmitCommandsAndFlushGPU()
 	Device->SubmitCommandsAndFlushGPU();
 }
 
-FTexture2DRHIRef FVulkanDynamicRHI::RHICreateTexture2DFromResource(EPixelFormat Format, uint32 SizeX, uint32 SizeY, uint32 NumMips, uint32 NumSamples, VkImage Resource, ETextureCreateFlags Flags)
+FTexture2DRHIRef FVulkanDynamicRHI::RHICreateTexture2DFromResource(EPixelFormat Format, uint32 SizeX, uint32 SizeY, uint32 NumMips, uint32 NumSamples, VkImage Resource, ETextureCreateFlags Flags, const FClearValueBinding& ClearValueBinding)
 {
 	const FRHITextureCreateDesc Desc =
 		FRHITextureCreateDesc::Create2D(TEXT("VulkanTexture2DFromResource"), SizeX, SizeY, Format)
-		.SetClearValue(EnumHasAnyFlags(Flags, ETextureCreateFlags::DepthStencilTargetable) ? FClearValueBinding::DepthZero : FClearValueBinding::Transparent)
+		.SetClearValue(ClearValueBinding)
 		.SetFlags(Flags)
 		.SetNumMips(NumMips)
 		.SetNumSamples(NumSamples)
@@ -1444,11 +1444,11 @@ FTexture2DRHIRef FVulkanDynamicRHI::RHICreateTexture2DFromResource(EPixelFormat
 	return new FVulkanTexture(*Device, Desc, Resource, false);
 }
 
-FTexture2DArrayRHIRef FVulkanDynamicRHI::RHICreateTexture2DArrayFromResource(EPixelFormat Format, uint32 SizeX, uint32 SizeY, uint32 ArraySize, uint32 NumMips, uint32 NumSamples, VkImage Resource, ETextureCreateFlags Flags)
+FTexture2DArrayRHIRef FVulkanDynamicRHI::RHICreateTexture2DArrayFromResource(EPixelFormat Format, uint32 SizeX, uint32 SizeY, uint32 ArraySize, uint32 NumMips, uint32 NumSamples, VkImage Resource, ETextureCreateFlags Flags, const FClearValueBinding& ClearValueBinding)
 {
 	const FRHITextureCreateDesc Desc =
 		FRHITextureCreateDesc::Create2DArray(TEXT("VulkanTextureArrayFromResource"), SizeX, SizeY, ArraySize, Format)
-		.SetClearValue(EnumHasAnyFlags(Flags, ETextureCreateFlags::DepthStencilTargetable) ? FClearValueBinding::DepthZero : FClearValueBinding::Transparent)
+		.SetClearValue(ClearValueBinding)
 		.SetFlags(Flags)
 		.SetNumMips(NumMips)
 		.SetNumSamples(NumSamples)
@@ -1457,14 +1457,14 @@ FTexture2DArrayRHIRef FVulkanDynamicRHI::RHICreateTexture2DArrayFromResource(EPi
 	return new FVulkanTexture(*Device, Desc, Resource, false);
 }
 
-FTextureCubeRHIRef FVulkanDynamicRHI::RHICreateTextureCubeFromResource(EPixelFormat Format, uint32 Size, bool bArray, uint32 ArraySize, uint32 NumMips, VkImage Resource, ETextureCreateFlags Flags)
+FTextureCubeRHIRef FVulkanDynamicRHI::RHICreateTextureCubeFromResource(EPixelFormat Format, uint32 Size, bool bArray, uint32 ArraySize, uint32 NumMips, VkImage Resource, ETextureCreateFlags Flags, const FClearValueBinding& ClearValueBinding)
 {
 	const FRHITextureCreateDesc Desc =
 		FRHITextureCreateDesc::Create(TEXT("VulkanTextureCubeFromResource"), ArraySize > 1 ? ETextureDimension::TextureCubeArray : ETextureDimension::TextureCube)
 		.SetExtent(Size)
 		.SetArraySize(ArraySize)
 		.SetFormat(Format)
-		.SetClearValue(EnumHasAnyFlags(Flags, ETextureCreateFlags::DepthStencilTargetable) ? FClearValueBinding::DepthZero : FClearValueBinding::Transparent)
+		.SetClearValue(ClearValueBinding)
 		.SetFlags(Flags)
 		.SetNumMips(NumMips)
 		.DetermineInititialState();
diff --git a/Engine/Source/Runtime/VulkanRHI/Public/IVulkanDynamicRHI.h b/Engine/Source/Runtime/VulkanRHI/Public/IVulkanDynamicRHI.h
--- a/Engine/Source/Runtime/VulkanRHI/Public/IVulkanDynamicRHI.h
+++ b/Engine/Source/Runtime/VulkanRHI/Public/IVulkanDynamicRHI.h
@@ -60,9 +60,9 @@ struct IVulkanDynamicRHI : public FDynamicRHI
 	virtual TArray<VkExtensionProperties> RHIGetAllInstanceExtensions() const = 0;
 	virtual TArray<VkExtensionProperties> RHIGetAllDeviceExtensions(VkPhysicalDevice InPhysicalDevice) const = 0;
 
-	virtual FTexture2DRHIRef      RHICreateTexture2DFromResource(EPixelFormat Format, uint32 SizeX, uint32 SizeY, uint32 NumMips, uint32 NumSamples, VkImage Resource, ETextureCreateFlags Flags) = 0;
-	virtual FTexture2DArrayRHIRef RHICreateTexture2DArrayFromResource(EPixelFormat Format, uint32 SizeX, uint32 SizeY, uint32 ArraySize, uint32 NumMips, uint32 NumSamples, VkImage Resource, ETextureCreateFlags Flags) = 0;
-	virtual FTextureCubeRHIRef    RHICreateTextureCubeFromResource(EPixelFormat Format, uint32 Size, bool bArray, uint32 ArraySize, uint32 NumMips, VkImage Resource, ETextureCreateFlags Flags) = 0;
+	virtual FTexture2DRHIRef      RHICreateTexture2DFromResource(EPixelFormat Format, uint32 SizeX, uint32 SizeY, uint32 NumMips, uint32 NumSamples, VkImage Resource, ETextureCreateFlags Flags, const FClearValueBinding& ClearValueBinding = FClearValueBinding::Transparent) = 0;
+	virtual FTexture2DArrayRHIRef RHICreateTexture2DArrayFromResource(EPixelFormat Format, uint32 SizeX, uint32 SizeY, uint32 ArraySize, uint32 NumMips, uint32 NumSamples, VkImage Resource, ETextureCreateFlags Flags, const FClearValueBinding& ClearValueBinding = FClearValueBinding::Transparent) = 0;
+	virtual FTextureCubeRHIRef    RHICreateTextureCubeFromResource(EPixelFormat Format, uint32 Size, bool bArray, uint32 ArraySize, uint32 NumMips, VkImage Resource, ETextureCreateFlags Flags, const FClearValueBinding& ClearValueBinding = FClearValueBinding::Transparent) = 0;
 
 	virtual VkImage                  RHIGetVkImage(FRHITexture* InTexture) const = 0;
 	virtual VkFormat                 RHIGetViewVkFormat(FRHITexture* InTexture) const = 0;
diff --git a/Engine/Source/Runtime/VulkanRHI/Public/VulkanDynamicRHI.h b/Engine/Source/Runtime/VulkanRHI/Public/VulkanDynamicRHI.h
--- a/Engine/Source/Runtime/VulkanRHI/Public/VulkanDynamicRHI.h
+++ b/Engine/Source/Runtime/VulkanRHI/Public/VulkanDynamicRHI.h
@@ -64,9 +64,9 @@ public:
 	virtual bool RHISupportsEXTFragmentDensityMap2() const final override;
 	virtual TArray<VkExtensionProperties> RHIGetAllInstanceExtensions() const final override;
 	virtual TArray<VkExtensionProperties> RHIGetAllDeviceExtensions(VkPhysicalDevice InPhysicalDevice) const final override;
-	virtual FTexture2DRHIRef RHICreateTexture2DFromResource(EPixelFormat Format, uint32 SizeX, uint32 SizeY, uint32 NumMips, uint32 NumSamples, VkImage Resource, ETextureCreateFlags Flags) final override;
-	virtual FTexture2DArrayRHIRef RHICreateTexture2DArrayFromResource(EPixelFormat Format, uint32 SizeX, uint32 SizeY, uint32 ArraySize, uint32 NumMips, uint32 NumSamples, VkImage Resource, ETextureCreateFlags Flags) final override;
-	virtual FTextureCubeRHIRef RHICreateTextureCubeFromResource(EPixelFormat Format, uint32 Size, bool bArray, uint32 ArraySize, uint32 NumMips, VkImage Resource, ETextureCreateFlags Flags) final override;
+	virtual FTexture2DRHIRef RHICreateTexture2DFromResource(EPixelFormat Format, uint32 SizeX, uint32 SizeY, uint32 NumMips, uint32 NumSamples, VkImage Resource, ETextureCreateFlags Flags, const FClearValueBinding& ClearValueBinding = FClearValueBinding::Transparent) final override;
+	virtual FTexture2DArrayRHIRef RHICreateTexture2DArrayFromResource(EPixelFormat Format, uint32 SizeX, uint32 SizeY, uint32 ArraySize, uint32 NumMips, uint32 NumSamples, VkImage Resource, ETextureCreateFlags Flags, const FClearValueBinding& ClearValueBinding = FClearValueBinding::Transparent) final override;
+	virtual FTextureCubeRHIRef RHICreateTextureCubeFromResource(EPixelFormat Format, uint32 Size, bool bArray, uint32 ArraySize, uint32 NumMips, VkImage Resource, ETextureCreateFlags Flags, const FClearValueBinding& ClearValueBinding = FClearValueBinding::Transparent) final override;
 	virtual VkImage RHIGetVkImage(FRHITexture* InTexture) const final override;
 	virtual VkFormat RHIGetViewVkFormat(FRHITexture* InTexture) const final override;
 	virtual FVulkanRHIAllocationInfo RHIGetAllocationInfo(FRHITexture* InTexture) const final override;
