VulkanRHI:  Reduce reliance on layout tracking (bring us a step closer to parallel cmdbuffer generation).

- Moved layout tracking to CmdBuf.
- Added layout tracking to Queue, submitted CmdBuf transfer their layout to the queue.
- Defrag uses Queue tracking for layout queries.
- Add support for VK_KHR_synchronization2 and VK_KHR_separate_depth_stencil_layouts.  Add a new SupportsParallelRendering code path when both are supported that (mostly) disables layout tracking for regular command buffer generation.  Also sets GRHISupportsSeparateDepthStencilCopyAccess when supported (to re-enable access to async compute on Vulkan in RDG).
- Change layout tracking to track depth and stencil separately and to use new VK_KHR_synchronization2 structs and layouts for all of tracking (like VK_IMAGE_LAYOUT_READ_ONLY_OPTIMAL/VK_IMAGE_LAYOUT_ATTACHMENT_OPTIMAL).  If Sync2 is not supported, they are converted back before being submitted (added GetMergedDepthStencilLayout to merge separate depth and stencil into a single layout).
- Merged RHIBeginTransitions/RHIEndTransitions into a single TransitionProcessor (with a bool to distinguish the differences  between begin/end) to reduce code duplication.
- Remove some of the places where tracking was queried since RDG should put us in the right state already.
- Added 'r.Vulkan.AllowSynchronization2' to disable Sync2 on systems that support it.
- Added 'r.Vulkan.AutoCorrectExpectedLayouts' which is enabled by default to allow some parts of code to read from layout tracking while in SupportsParallelRendering mode.  These will need to get fixed...

Related clean up:
- Create a common code path to reset FVulkanRenderTargetLayout (some constructors left some params uninit)
- Fix barrier errors on RHIRead*SurfaceFloatData functions when reading a mip or layer index different than 0.
- Remove VulkanSetImageLayoutSimple and VulkanSetImageLayoutAllMips.  Reduce the usage of VulkanSetImageLayout to a minimum and change it to take in a FVulkanCmdBuffer to be able to branch on Sync2 support.  Using FVulkanPipelineBarrier instead where possible...

LayoutManger refactor:
- Pull Renderpass/Framebuffer management out of LayoutManager, create new RenderPassManager for it.  Manage render passes per device instead of per context. Still keep the currently set Renderpass/Framebuffer in Context.
- Make framebuffer creation threadsafe
- Remove 'static' for LayoutManager, actually create one per context.

Other clean up:
- Remove most of VULKAN_SUPPORTS_RENDERPASS2 defines since minimum VK SDK guarantees the presence of all the structures.
- Remove VULKAN_USE_NEW_QUERIES.
- Use VK_TYPE_TO_STRING for a few more type in the DUMP_LAYER to reduce maintenance
- Fill CreateRenderPass2KHR in the WrapLayer
- Fix some of the message in memory dump (missing params, wrong types)
- Add 'r.Vulkan.UseMemoryBarrierOpt' to make it possible to disable the memory barrier optimization (easier to track down missing buffer synchronization)

NOTE: Not converting everything to VkPipelineStageFlags2 and VkAccessFlags2 everywhere until necessary.  More work needed for parallel rendering (for thread safety, and for correctness without 'r.Vulkan.AutoCorrectExpectedLayouts').

Prerequisite for:

diff --git a/Engine/Plugins/Runtime/Steam/SteamVR/Source/SteamVR/Private/SteamVRRender.cpp b/Engine/Plugins/Runtime/Steam/SteamVR/Source/SteamVR/Private/SteamVRRender.cpp
--- a/Engine/Plugins/Runtime/Steam/SteamVR/Source/SteamVR/Private/SteamVRRender.cpp
+++ b/Engine/Plugins/Runtime/Steam/SteamVR/Source/SteamVR/Private/SteamVRRender.cpp
@@ -360,7 +360,9 @@ void FSteamVRHMD::VulkanBridge::FinishRendering()
 		const VkFormat SwapChainFormat = VulkanRHI->RHIGetViewVkFormat(SwapChainTexture);
 
 		// Color layout
+		PRAGMA_DISABLE_DEPRECATION_WARNINGS		
 		VkImageLayout& CurrentLayout = VulkanRHI->RHIFindOrAddLayoutRW(SwapChainTexture, VK_IMAGE_LAYOUT_UNDEFINED);
+		PRAGMA_ENABLE_DEPRECATION_WARNINGS
 		bool bHadLayout = (CurrentLayout != VK_IMAGE_LAYOUT_UNDEFINED);
 		
 		VkImageSubresourceRange SubresourceRange = { VK_IMAGE_ASPECT_COLOR_BIT, 0, 1, 0, 1 };
@@ -402,7 +404,9 @@ void FSteamVRHMD::VulkanBridge::FinishRendering()
 			const VkFormat DepthSwapChainFormat = VulkanRHI->RHIGetViewVkFormat(DepthSwapChainTexture);
 
 			VkImageSubresourceRange SubresourceRangeDepth = { VK_IMAGE_ASPECT_DEPTH_BIT | VK_IMAGE_ASPECT_STENCIL_BIT, 0, 1, 0, 1 };
+			PRAGMA_DISABLE_DEPRECATION_WARNINGS
 			VkImageLayout& CurrentDepthLayout = VulkanRHI->RHIFindOrAddLayoutRW(DepthSwapChainTexture, VK_IMAGE_LAYOUT_UNDEFINED);
+			PRAGMA_ENABLE_DEPRECATION_WARNINGS
 			bool bDepthHadLayout = (CurrentDepthLayout != VK_IMAGE_LAYOUT_UNDEFINED);
 
 			if (CurrentDepthLayout != VK_IMAGE_LAYOUT_TRANSFER_SRC_OPTIMAL)
diff --git a/Engine/Source/Runtime/VulkanRHI/Private/VulkanBarriers.cpp b/Engine/Source/Runtime/VulkanRHI/Private/VulkanBarriers.cpp
--- a/Engine/Source/Runtime/VulkanRHI/Private/VulkanBarriers.cpp
+++ b/Engine/Source/Runtime/VulkanRHI/Private/VulkanBarriers.cpp
@@ -6,17 +6,40 @@
 #include "VulkanPendingState.h"
 #include "ProfilingDebugging/RealtimeGPUProfiler.h"
 
-FVulkanLayoutManager FVulkanCommandListContext::LayoutManager;
+
 
 // All shader stages supported by VK device - VK_PIPELINE_STAGE_VERTEX_SHADER_BIT, FRAGMENT etc
 extern uint32 GVulkanDevicePipelineStageBits;
 
+
+int32 GVulkanUseMemoryBarrierOpt = 1;
+static FAutoConsoleVariableRef CVarVulkanUseMemoryBarrierOpt(
+	TEXT("r.Vulkan.UseMemoryBarrierOpt"),
+	GVulkanUseMemoryBarrierOpt,
+	TEXT("Simplify buffer barriers and image barriers without layout transitions to a memory barrier.\n")
+	TEXT(" 0: Do not collapse to a single memory barrier, useful for tracking single resource transitions in external tools\n")
+	TEXT(" 1: Collapse to a memory barrier when appropriate (default)"),
+	ECVF_Default
+);
+
+int32 GVulkanAutoCorrectExpectedLayouts = 1;
+static FAutoConsoleVariableRef CVarVulkanAutoCorrectExpectedLayouts(
+	TEXT("r.Vulkan.AutoCorrectExpectedLayouts"),
+	GVulkanAutoCorrectExpectedLayouts,
+	TEXT("Will use layout tracking to correct mismatched layouts when setting expected layouts.  \n")
+	TEXT("This is unsafe for multi-threaded command buffer generations\n")
+	TEXT(" 0: Do not correct layouts.\n")
+	TEXT(" 1: Correct the layouts using layout tracking. (default)"),
+	ECVF_Default
+);
+
+
 //
 // The following two functions are used when the RHI needs to do image layout transitions internally.
 // They are not used for the transitions requested through the public API (RHICreate/Begin/EndTransition)
 // unless the initial state in ERHIAccess::Unknown, in which case the tracking code kicks in.
 //
-static VkAccessFlags GetVkAccessMaskForLayout(VkImageLayout Layout)
+static VkAccessFlags GetVkAccessMaskForLayout(const VkImageLayout Layout)
 {
 	VkAccessFlags Flags = 0;
 
@@ -33,15 +56,13 @@ static VkAccessFlags GetVkAccessMaskForLayout(VkImageLayout Layout)
 			break;
 
 		case VK_IMAGE_LAYOUT_DEPTH_STENCIL_ATTACHMENT_OPTIMAL:
-#if VULKAN_SUPPORTS_SEPARATE_DEPTH_STENCIL_LAYOUTS
-		case VK_IMAGE_LAYOUT_DEPTH_ATTACHMENT_OPTIMAL_KHR:
-		case VK_IMAGE_LAYOUT_STENCIL_ATTACHMENT_OPTIMAL_KHR:
-#endif
+		case VK_IMAGE_LAYOUT_DEPTH_ATTACHMENT_OPTIMAL:
+		case VK_IMAGE_LAYOUT_STENCIL_ATTACHMENT_OPTIMAL:
 			Flags = VK_ACCESS_DEPTH_STENCIL_ATTACHMENT_READ_BIT | VK_ACCESS_DEPTH_STENCIL_ATTACHMENT_WRITE_BIT;
 			break;
 
-		case VK_IMAGE_LAYOUT_DEPTH_READ_ONLY_STENCIL_ATTACHMENT_OPTIMAL_KHR:
-		case VK_IMAGE_LAYOUT_DEPTH_ATTACHMENT_STENCIL_READ_ONLY_OPTIMAL_KHR:
+		case VK_IMAGE_LAYOUT_DEPTH_READ_ONLY_STENCIL_ATTACHMENT_OPTIMAL:
+		case VK_IMAGE_LAYOUT_DEPTH_ATTACHMENT_STENCIL_READ_ONLY_OPTIMAL:
 			Flags = VK_ACCESS_SHADER_READ_BIT | VK_ACCESS_DEPTH_STENCIL_ATTACHMENT_READ_BIT | VK_ACCESS_DEPTH_STENCIL_ATTACHMENT_WRITE_BIT;
 			break;
 
@@ -50,10 +71,8 @@ static VkAccessFlags GetVkAccessMaskForLayout(VkImageLayout Layout)
 			break;
 
 		case VK_IMAGE_LAYOUT_DEPTH_STENCIL_READ_ONLY_OPTIMAL:
-#if VULKAN_SUPPORTS_SEPARATE_DEPTH_STENCIL_LAYOUTS
-		case VK_IMAGE_LAYOUT_DEPTH_READ_ONLY_OPTIMAL_KHR:
-		case VK_IMAGE_LAYOUT_STENCIL_READ_ONLY_OPTIMAL_KHR:
-#endif
+		case VK_IMAGE_LAYOUT_DEPTH_READ_ONLY_OPTIMAL:
+		case VK_IMAGE_LAYOUT_STENCIL_READ_ONLY_OPTIMAL:
 			Flags = VK_ACCESS_SHADER_READ_BIT | VK_ACCESS_DEPTH_STENCIL_ATTACHMENT_READ_BIT;
 			break;
 
@@ -61,23 +80,30 @@ static VkAccessFlags GetVkAccessMaskForLayout(VkImageLayout Layout)
 			Flags = 0;
 			break;
 
-#if VULKAN_SUPPORTS_FRAGMENT_DENSITY_MAP
 		case VK_IMAGE_LAYOUT_FRAGMENT_DENSITY_MAP_OPTIMAL_EXT:
 			Flags = VK_ACCESS_FRAGMENT_DENSITY_MAP_READ_BIT_EXT;
 			break;
-#endif
 
-#if VULKAN_SUPPORTS_FRAGMENT_SHADING_RATE
 		case VK_IMAGE_LAYOUT_FRAGMENT_SHADING_RATE_ATTACHMENT_OPTIMAL_KHR:
 			Flags = VK_ACCESS_FRAGMENT_SHADING_RATE_ATTACHMENT_READ_BIT_KHR;
 			break;
-#endif
 
 		case VK_IMAGE_LAYOUT_GENERAL:
+			// todo-jn: could be used for R64 in read layout
 		case VK_IMAGE_LAYOUT_UNDEFINED:
 			Flags = 0;
 			break;
 
+		case VK_IMAGE_LAYOUT_ATTACHMENT_OPTIMAL:
+			// todo-jn: sync2 currently only used by depth/stencil targets
+			Flags = VK_ACCESS_DEPTH_STENCIL_ATTACHMENT_READ_BIT | VK_ACCESS_DEPTH_STENCIL_ATTACHMENT_WRITE_BIT;
+			break;
+
+		case VK_IMAGE_LAYOUT_READ_ONLY_OPTIMAL:
+			// todo-jn: sync2 currently only used by depth/stencil targets
+			Flags = VK_ACCESS_SHADER_READ_BIT;
+			break;
+
 		default:
 			checkNoEntry();
 			break;
@@ -105,10 +131,8 @@ static VkPipelineStageFlags GetVkStageFlagsForLayout(VkImageLayout Layout)
 			break;
 
 		case VK_IMAGE_LAYOUT_DEPTH_STENCIL_ATTACHMENT_OPTIMAL:
-#if VULKAN_SUPPORTS_SEPARATE_DEPTH_STENCIL_LAYOUTS
-		case VK_IMAGE_LAYOUT_DEPTH_ATTACHMENT_OPTIMAL_KHR:
-		case VK_IMAGE_LAYOUT_STENCIL_ATTACHMENT_OPTIMAL_KHR:
-#endif
+		case VK_IMAGE_LAYOUT_DEPTH_ATTACHMENT_OPTIMAL:
+		case VK_IMAGE_LAYOUT_STENCIL_ATTACHMENT_OPTIMAL:
 			Flags = VK_PIPELINE_STAGE_EARLY_FRAGMENT_TESTS_BIT | VK_PIPELINE_STAGE_LATE_FRAGMENT_TESTS_BIT;
 			break;
 
@@ -117,12 +141,10 @@ static VkPipelineStageFlags GetVkStageFlagsForLayout(VkImageLayout Layout)
 			break;
 
 		case VK_IMAGE_LAYOUT_DEPTH_STENCIL_READ_ONLY_OPTIMAL:
-		case VK_IMAGE_LAYOUT_DEPTH_READ_ONLY_STENCIL_ATTACHMENT_OPTIMAL_KHR:
-		case VK_IMAGE_LAYOUT_DEPTH_ATTACHMENT_STENCIL_READ_ONLY_OPTIMAL_KHR:
-#if VULKAN_SUPPORTS_SEPARATE_DEPTH_STENCIL_LAYOUTS
-		case VK_IMAGE_LAYOUT_DEPTH_READ_ONLY_OPTIMAL_KHR:
-		case VK_IMAGE_LAYOUT_STENCIL_READ_ONLY_OPTIMAL_KHR:
-#endif
+		case VK_IMAGE_LAYOUT_DEPTH_READ_ONLY_STENCIL_ATTACHMENT_OPTIMAL:
+		case VK_IMAGE_LAYOUT_DEPTH_ATTACHMENT_STENCIL_READ_ONLY_OPTIMAL:
+		case VK_IMAGE_LAYOUT_DEPTH_READ_ONLY_OPTIMAL:
+		case VK_IMAGE_LAYOUT_STENCIL_READ_ONLY_OPTIMAL:
 			Flags = VK_PIPELINE_STAGE_FRAGMENT_SHADER_BIT | VK_PIPELINE_STAGE_EARLY_FRAGMENT_TESTS_BIT | VK_PIPELINE_STAGE_LATE_FRAGMENT_TESTS_BIT;
 			break;
 
@@ -130,23 +152,29 @@ static VkPipelineStageFlags GetVkStageFlagsForLayout(VkImageLayout Layout)
 			Flags = VK_PIPELINE_STAGE_BOTTOM_OF_PIPE_BIT;
 			break;
 
-#if VULKAN_SUPPORTS_FRAGMENT_DENSITY_MAP
 		case VK_IMAGE_LAYOUT_FRAGMENT_DENSITY_MAP_OPTIMAL_EXT:
 			Flags = VK_PIPELINE_STAGE_FRAGMENT_DENSITY_PROCESS_BIT_EXT;
 			break;
-#endif
 
-#if VULKAN_SUPPORTS_FRAGMENT_SHADING_RATE
 		case VK_IMAGE_LAYOUT_FRAGMENT_SHADING_RATE_ATTACHMENT_OPTIMAL_KHR:
 			Flags = VK_PIPELINE_STAGE_FRAGMENT_SHADING_RATE_ATTACHMENT_BIT_KHR;
 			break;
-#endif
 			
 		case VK_IMAGE_LAYOUT_GENERAL:
 		case VK_IMAGE_LAYOUT_UNDEFINED:
 			Flags = VK_PIPELINE_STAGE_TOP_OF_PIPE_BIT;
 			break;
 
+		case VK_IMAGE_LAYOUT_ATTACHMENT_OPTIMAL:
+			// todo-jn: sync2 currently only used by depth/stencil targets
+			Flags = VK_PIPELINE_STAGE_EARLY_FRAGMENT_TESTS_BIT | VK_PIPELINE_STAGE_LATE_FRAGMENT_TESTS_BIT;
+			break;
+
+		case VK_IMAGE_LAYOUT_READ_ONLY_OPTIMAL:
+			// todo-jn: sync2 currently only used by depth/stencil targets
+			Flags = VK_PIPELINE_STAGE_FRAGMENT_SHADER_BIT | VK_PIPELINE_STAGE_EARLY_FRAGMENT_TESTS_BIT | VK_PIPELINE_STAGE_LATE_FRAGMENT_TESTS_BIT;
+			break;
+
 		default:
 			checkNoEntry();
 			break;
@@ -182,7 +210,7 @@ static void GetVkStageAndAccessFlags(ERHIAccess RHIAccess, FRHITransitionInfo::E
 	// The layout to use if SRV access is requested. In case of depth/stencil buffers, we don't need to worry about different states for the separate aspects, since that's handled explicitly elsewhere,
 	// and this function is never called for depth-only or stencil-only transitions.
 	const VkImageLayout SRVLayout = 
-		bIsDepthStencil ? VK_IMAGE_LAYOUT_DEPTH_STENCIL_READ_ONLY_OPTIMAL : 
+		bIsDepthStencil ? VK_IMAGE_LAYOUT_READ_ONLY_OPTIMAL : 
 		bSupportsReadOnlyOptimal ? VK_IMAGE_LAYOUT_SHADER_READ_ONLY_OPTIMAL : VK_IMAGE_LAYOUT_GENERAL;
 
 	// States which cannot be combined.
@@ -202,7 +230,6 @@ static void GetVkStageAndAccessFlags(ERHIAccess RHIAccess, FRHITransitionInfo::E
 			return;
 
 		case ERHIAccess::CPURead:
-			// FIXME: is this correct?
 			StageFlags = VK_PIPELINE_STAGE_HOST_BIT;
 			AccessFlags = VK_ACCESS_HOST_READ_BIT;
 			Layout = VK_IMAGE_LAYOUT_GENERAL;
@@ -233,18 +260,17 @@ static void GetVkStageAndAccessFlags(ERHIAccess RHIAccess, FRHITransitionInfo::E
 			Layout = VK_IMAGE_LAYOUT_TRANSFER_DST_OPTIMAL;
 			return;
 
-#if VULKAN_RHI_RAYTRACING
-		// vkrt todo: Finer grain stage flags would be ideal here.
 		case ERHIAccess::BVHRead:
+			// vkrt todo: Finer grain stage flags would be ideal here.
 			StageFlags = VK_PIPELINE_STAGE_ACCELERATION_STRUCTURE_BUILD_BIT_KHR | VK_PIPELINE_STAGE_RAY_TRACING_SHADER_BIT_KHR | VK_PIPELINE_STAGE_COMPUTE_SHADER_BIT;
 			AccessFlags = VK_ACCESS_ACCELERATION_STRUCTURE_READ_BIT_KHR;
 			return;
 
 		case ERHIAccess::BVHWrite:
+			// vkrt todo: Finer grain stage flags would be ideal here.
 			StageFlags = VK_PIPELINE_STAGE_ACCELERATION_STRUCTURE_BUILD_BIT_KHR | VK_PIPELINE_STAGE_RAY_TRACING_SHADER_BIT_KHR | VK_PIPELINE_STAGE_COMPUTE_SHADER_BIT;
 			AccessFlags = VK_ACCESS_ACCELERATION_STRUCTURE_WRITE_BIT_KHR;
 			return;
-#endif // VULKAN_RHI_RAYTRACING
 	}
 
 	// If DSVWrite is set, we ignore everything else because it decides the layout.
@@ -253,7 +279,7 @@ static void GetVkStageAndAccessFlags(ERHIAccess RHIAccess, FRHITransitionInfo::E
 		check(bIsDepthStencil);
 		StageFlags = VK_PIPELINE_STAGE_EARLY_FRAGMENT_TESTS_BIT | VK_PIPELINE_STAGE_LATE_FRAGMENT_TESTS_BIT;
 		AccessFlags = VK_ACCESS_DEPTH_STENCIL_ATTACHMENT_READ_BIT | VK_ACCESS_DEPTH_STENCIL_ATTACHMENT_WRITE_BIT;
-		Layout = VK_IMAGE_LAYOUT_DEPTH_STENCIL_ATTACHMENT_OPTIMAL;
+		Layout = VK_IMAGE_LAYOUT_ATTACHMENT_OPTIMAL;
 		return;
 	}
 
@@ -303,7 +329,7 @@ static void GetVkStageAndAccessFlags(ERHIAccess RHIAccess, FRHITransitionInfo::E
 
 		// If any of the SRV flags is set, the code below will set Layout to SRVLayout again, but it's fine since
 		// SRVLayout takes into account bIsDepthStencil and ends up being the same as what we set here.
-		Layout = VK_IMAGE_LAYOUT_DEPTH_STENCIL_READ_ONLY_OPTIMAL;
+		Layout = VK_IMAGE_LAYOUT_READ_ONLY_OPTIMAL;
 
 		ProcessedRHIFlags |= (uint32)ERHIAccess::DSVRead;
 	}
@@ -392,141 +418,159 @@ static void GetVkStageAndAccessFlags(ERHIAccess RHIAccess, FRHITransitionInfo::E
 	ensureMsgf(RemainingFlags == 0, TEXT("Some access flags were not processed. RHIAccess=%x, ProcessedRHIFlags=%x, RemainingFlags=%x"), RHIAccess, ProcessedRHIFlags, RemainingFlags);
 }
 
-//
-// Helpers for merging separate depth-stencil transitions into a single transition.
-//
-struct FDepthStencilSubresTransition
-{
-	FVulkanTexture* Texture;
-	int ArraySlice;
-	ERHIAccess SrcDepthAccess;
-	ERHIAccess DestDepthAccess;
-	ERHIAccess SrcStencilAccess;
-	ERHIAccess DestStencilAccess;
-	int bDepthAccessSet : 1;
-	int bStencilAccessSet : 1;
-	int bAliasingBarrier : 1;
-};
 
-static void GetDepthStencilStageAndAccessFlags(ERHIAccess DepthAccess, ERHIAccess StencilAccess, VkPipelineStageFlags& StageFlags, VkAccessFlags& AccessFlags, VkImageLayout& Layout, bool bIsSourceState)
+static VkImageAspectFlags GetDepthStencilAspectMask(uint32 PlaneSlice)
 {
-	if (DepthAccess == ERHIAccess::Unknown || StencilAccess == ERHIAccess::Unknown)
+	VkImageAspectFlags AspectFlags = 0;
+
+	if (PlaneSlice == FRHISubresourceRange::kAllSubresources)
 	{
-		StageFlags = VK_PIPELINE_STAGE_ALL_COMMANDS_BIT;
-		AccessFlags = VK_ACCESS_MEMORY_READ_BIT | VK_ACCESS_MEMORY_WRITE_BIT;
-		Layout = VK_IMAGE_LAYOUT_UNDEFINED;
-		return;
+		AspectFlags = VK_IMAGE_ASPECT_DEPTH_BIT | VK_IMAGE_ASPECT_STENCIL_BIT;
 	}
 
-	if (DepthAccess == ERHIAccess::Discard && StencilAccess == ERHIAccess::Discard)
+	if (PlaneSlice == FRHISubresourceRange::kDepthPlaneSlice)
 	{
-		StageFlags = VK_PIPELINE_STAGE_ALL_COMMANDS_BIT;
-		AccessFlags = VK_ACCESS_MEMORY_READ_BIT | VK_ACCESS_MEMORY_WRITE_BIT;
-		Layout = bIsSourceState ? VK_IMAGE_LAYOUT_UNDEFINED : VK_IMAGE_LAYOUT_DEPTH_STENCIL_READ_ONLY_OPTIMAL;
-		return;
+		AspectFlags = VK_IMAGE_ASPECT_DEPTH_BIT;
+	}
+
+	if (PlaneSlice == FRHISubresourceRange::kStencilPlaneSlice)
+	{
+		AspectFlags = VK_IMAGE_ASPECT_STENCIL_BIT;
 	}
 
-	Layout = VK_IMAGE_LAYOUT_DEPTH_STENCIL_READ_ONLY_OPTIMAL;
+	return AspectFlags;
+}
+
 
-	if (EnumHasAnyFlags(DepthAccess, ERHIAccess::CopySrc))
+// Returns the VK_KHR_synchronization2 layout corresponding to an access type
+static VkImageLayout GetDepthOrStencilLayout(ERHIAccess Access)
+{
+	VkImageLayout Layout;
+	if (Access == ERHIAccess::Unknown)
+	{
+		Layout = VK_IMAGE_LAYOUT_UNDEFINED;
+	}
+	else if (Access == ERHIAccess::Discard)
+	{
+		Layout = VK_IMAGE_LAYOUT_UNDEFINED;
+	}
+	else if (EnumHasAnyFlags(Access, ERHIAccess::CopySrc))
 	{
 		Layout = VK_IMAGE_LAYOUT_TRANSFER_SRC_OPTIMAL;
 	}
-	else if (EnumHasAnyFlags(DepthAccess, ERHIAccess::CopyDest))
+	else if (EnumHasAnyFlags(Access, ERHIAccess::CopyDest))
 	{
 		Layout = VK_IMAGE_LAYOUT_TRANSFER_DST_OPTIMAL;
 	}
-	else if (EnumHasAnyFlags(DepthAccess, ERHIAccess::DSVWrite))
+	else if (EnumHasAnyFlags(Access, ERHIAccess::DSVWrite))
 	{
-		if (EnumHasAnyFlags(StencilAccess, ERHIAccess::DSVWrite))
-		{
-			Layout = VK_IMAGE_LAYOUT_DEPTH_STENCIL_ATTACHMENT_OPTIMAL;
-		}
-		else
-		{
-			Layout = VK_IMAGE_LAYOUT_DEPTH_ATTACHMENT_STENCIL_READ_ONLY_OPTIMAL_KHR;
-		}
+		Layout = VK_IMAGE_LAYOUT_ATTACHMENT_OPTIMAL;
 	}
 	else
 	{
-		if (EnumHasAnyFlags(StencilAccess, ERHIAccess::DSVWrite))
-		{
-			Layout = VK_IMAGE_LAYOUT_DEPTH_READ_ONLY_STENCIL_ATTACHMENT_OPTIMAL_KHR;
-		}
+		Layout = VK_IMAGE_LAYOUT_READ_ONLY_OPTIMAL;
+	}
+	return Layout;
+}
+
+static void GetDepthOrStencilStageAndAccessFlags(ERHIAccess Access, VkPipelineStageFlags& StageFlags, VkAccessFlags& AccessFlags)
+{
+	if (Access == ERHIAccess::Unknown || Access == ERHIAccess::Discard)
+	{
+		StageFlags = VK_PIPELINE_STAGE_ALL_COMMANDS_BIT;
+		AccessFlags = VK_ACCESS_MEMORY_READ_BIT | VK_ACCESS_MEMORY_WRITE_BIT;
+		return;
 	}
 
 	StageFlags = 0;
 	AccessFlags = 0;
-	ERHIAccess CombinedAccess = DepthAccess | StencilAccess;
 	uint32 ProcessedRHIFlags = 0;
 
-	if (EnumHasAnyFlags(CombinedAccess, ERHIAccess::DSVWrite))
+	if (EnumHasAnyFlags(Access, ERHIAccess::DSVWrite))
 	{
 		StageFlags |= VK_PIPELINE_STAGE_EARLY_FRAGMENT_TESTS_BIT | VK_PIPELINE_STAGE_LATE_FRAGMENT_TESTS_BIT;
 		AccessFlags |= VK_ACCESS_DEPTH_STENCIL_ATTACHMENT_READ_BIT | VK_ACCESS_DEPTH_STENCIL_ATTACHMENT_WRITE_BIT;
 		ProcessedRHIFlags |= (uint32)ERHIAccess::DSVWrite;
 	}
 	
-	if (EnumHasAnyFlags(CombinedAccess, ERHIAccess::DSVRead))
+	if (EnumHasAnyFlags(Access, ERHIAccess::DSVRead))
 	{
 		StageFlags |= VK_PIPELINE_STAGE_EARLY_FRAGMENT_TESTS_BIT | VK_PIPELINE_STAGE_LATE_FRAGMENT_TESTS_BIT;
 		AccessFlags |= VK_ACCESS_DEPTH_STENCIL_ATTACHMENT_READ_BIT;
 		ProcessedRHIFlags |= (uint32)ERHIAccess::DSVRead;
 	}
 
-	if (EnumHasAnyFlags(CombinedAccess, ERHIAccess::SRVGraphics))
+	if (EnumHasAnyFlags(Access, ERHIAccess::SRVGraphics))
 	{
 		StageFlags |= VK_PIPELINE_STAGE_FRAGMENT_SHADER_BIT;
 		AccessFlags |= VK_ACCESS_SHADER_READ_BIT;
 		ProcessedRHIFlags |= (uint32)ERHIAccess::SRVGraphics;
 	}
 
-	if (EnumHasAnyFlags(CombinedAccess, ERHIAccess::UAVGraphics))
+	if (EnumHasAnyFlags(Access, ERHIAccess::UAVGraphics))
 	{
 		StageFlags |= VK_PIPELINE_STAGE_FRAGMENT_SHADER_BIT;
 		AccessFlags |= VK_ACCESS_SHADER_READ_BIT | VK_ACCESS_SHADER_WRITE_BIT;
 		ProcessedRHIFlags |= (uint32)ERHIAccess::UAVGraphics;
 	}
 
-	if (EnumHasAnyFlags(CombinedAccess, ERHIAccess::SRVCompute))
+	if (EnumHasAnyFlags(Access, ERHIAccess::SRVCompute))
 	{
 		StageFlags |= VK_PIPELINE_STAGE_COMPUTE_SHADER_BIT;
 		AccessFlags |= VK_ACCESS_SHADER_READ_BIT;
 		ProcessedRHIFlags |= (uint32)ERHIAccess::SRVCompute;
 	}
 
-	if (EnumHasAnyFlags(CombinedAccess, ERHIAccess::UAVCompute))
+	if (EnumHasAnyFlags(Access, ERHIAccess::UAVCompute))
 	{
 		StageFlags |= VK_PIPELINE_STAGE_COMPUTE_SHADER_BIT;
 		AccessFlags |= VK_ACCESS_SHADER_READ_BIT | VK_ACCESS_SHADER_WRITE_BIT;
 		ProcessedRHIFlags |= (uint32)ERHIAccess::UAVCompute;
 	}
 
-	if (EnumHasAnyFlags(CombinedAccess, ERHIAccess::CopySrc))
+	if (EnumHasAnyFlags(Access, ERHIAccess::CopySrc))
 	{
 		StageFlags |= VK_PIPELINE_STAGE_TRANSFER_BIT;
 		AccessFlags |= VK_ACCESS_TRANSFER_READ_BIT;
 		ProcessedRHIFlags |= (uint32)ERHIAccess::CopySrc;
 	}
 
-	if (EnumHasAnyFlags(CombinedAccess, ERHIAccess::CopyDest))
+	if (EnumHasAnyFlags(Access, ERHIAccess::CopyDest))
 	{
 		StageFlags |= VK_PIPELINE_STAGE_TRANSFER_BIT;
 		AccessFlags |= VK_ACCESS_TRANSFER_WRITE_BIT;
 		ProcessedRHIFlags |= (uint32)ERHIAccess::CopyDest;
 	}
 
-	uint32 RemainingFlags = (uint32)CombinedAccess & (~ProcessedRHIFlags);
-	ensureMsgf(RemainingFlags == 0, TEXT("Some access flags were not processed. DepthAccess=%x, StencilAccess=%x, ProcessedRHIFlags=%x, RemainingFlags=%x"), DepthAccess, StencilAccess, ProcessedRHIFlags, RemainingFlags);
+	const uint32 RemainingFlags = (uint32)Access & (~ProcessedRHIFlags);
+	ensureMsgf(RemainingFlags == 0, TEXT("Some access flags were not processed. Access=%x, ProcessedRHIFlags=%x, RemainingFlags=%x"), Access, ProcessedRHIFlags, RemainingFlags);
+}
+
+
+// Helper To apply to each bit in an aspect
+template <typename TFunction>
+static void ForEachAspect(VkImageAspectFlags AspectFlags, TFunction Function)
+{
+	// Keep it simple for now, can iterate on more bits when needed
+	for (uint32 SingleBit = 1; SingleBit <= VK_IMAGE_ASPECT_STENCIL_BIT; SingleBit <<= 1)
+	{
+		if ((AspectFlags & SingleBit) != 0)
+		{
+			Function((VkImageAspectFlagBits)SingleBit);
+		}
+	}
 }
 
 //
 // Helpers for filling in the fields of a VkImageMemoryBarrier structure.
 //
-static void SetupImageBarrier(VkImageMemoryBarrier& ImgBarrier, VkImage Image, VkAccessFlags SrcAccessFlags, VkAccessFlags DstAccessFlags, VkImageLayout SrcLayout, VkImageLayout DstLayout, const VkImageSubresourceRange& SubresRange)
+static void SetupImageBarrier(VkImageMemoryBarrier2& ImgBarrier, VkImage Image, VkPipelineStageFlags SrcStageFlags, VkPipelineStageFlags DstStageFlags, 
+	VkAccessFlags SrcAccessFlags, VkAccessFlags DstAccessFlags, VkImageLayout SrcLayout, VkImageLayout DstLayout, const VkImageSubresourceRange& SubresRange)
 {
-	ImgBarrier.sType = VK_STRUCTURE_TYPE_IMAGE_MEMORY_BARRIER;
+	ImgBarrier.sType = VK_STRUCTURE_TYPE_IMAGE_MEMORY_BARRIER_2;
 	ImgBarrier.pNext = nullptr;
+	ImgBarrier.srcStageMask = SrcStageFlags;
+	ImgBarrier.dstStageMask = DstStageFlags;
 	ImgBarrier.srcAccessMask = SrcAccessFlags;
 	ImgBarrier.dstAccessMask = DstAccessFlags;
 	ImgBarrier.oldLayout = SrcLayout;
@@ -537,7 +581,8 @@ static void SetupImageBarrier(VkImageMemoryBarrier& ImgBarrier, VkImage Image, V
 	ImgBarrier.subresourceRange = SubresRange;
 }
 
-static void SetupImageBarrierEntireRes(VkImageMemoryBarrier& ImgBarrier, VkImage Image, VkAccessFlags SrcAccessFlags, VkAccessFlags DstAccessFlags, VkImageLayout SrcLayout, VkImageLayout DstLayout, VkImageAspectFlags AspectMask)
+static void SetupImageBarrierEntireRes(VkImageMemoryBarrier2& ImgBarrier, VkImage Image, VkPipelineStageFlags SrcStageFlags, VkPipelineStageFlags DstStageFlags, 
+	VkAccessFlags SrcAccessFlags, VkAccessFlags DstAccessFlags, VkImageLayout SrcLayout, VkImageLayout DstLayout, VkImageAspectFlags AspectMask)
 {
 	VkImageSubresourceRange SubresRange;
 	SubresRange.aspectMask = AspectMask;
@@ -545,7 +590,7 @@ static void SetupImageBarrierEntireRes(VkImageMemoryBarrier& ImgBarrier, VkImage
 	SubresRange.levelCount = VK_REMAINING_MIP_LEVELS;
 	SubresRange.baseArrayLayer = 0;
 	SubresRange.layerCount = VK_REMAINING_ARRAY_LAYERS;
-	SetupImageBarrier(ImgBarrier, Image, SrcAccessFlags, DstAccessFlags, SrcLayout, DstLayout, SubresRange);
+	SetupImageBarrier(ImgBarrier, Image, SrcStageFlags, DstStageFlags, SrcAccessFlags, DstAccessFlags, SrcLayout, DstLayout, SubresRange);
 }
 
 // Fill in a VkImageSubresourceRange struct from the data contained inside a transition info struct coming from the public API.
@@ -577,11 +622,11 @@ static void SetupSubresourceRange(VkImageSubresourceRange& SubresRange, const FR
 
 void FVulkanDynamicRHI::RHICreateTransition(FRHITransition* Transition, const FRHITransitionCreateInfo& CreateInfo)
 {
+	TRACE_CPUPROFILER_EVENT_SCOPE(RHICreateTransition);
+
 	const ERHIPipeline SrcPipelines = CreateInfo.SrcPipelines;
 	const ERHIPipeline DstPipelines = CreateInfo.DstPipelines;
 
-	//checkf(FMath::IsPowerOfTwo((uint32)SrcPipelines) && FMath::IsPowerOfTwo((uint32)DstPipelines), TEXT("Support for multi-pipe resources is not yet implemented."));
-
 	FVulkanPipelineBarrier* Data = new (Transition->GetPrivateData<FVulkanPipelineBarrier>()) FVulkanPipelineBarrier;
 	Data->SrcPipelines = SrcPipelines;
 	Data->DstPipelines = DstPipelines;
@@ -659,8 +704,6 @@ void FVulkanDynamicRHI::RHICreateTransition(FRHITransition* Transition, const FR
 
 	const ERHIAccess DepthStencilFlags = ERHIAccess::DSVRead | ERHIAccess::DSVWrite;
 
-	TArray<FDepthStencilSubresTransition, TInlineAllocator<4>> DSSubresTransitions;
-
 	for (const FRHITransitionInfo& Info : CreateInfo.TransitionInfos)
 	{
 		if (!Info.Resource)
@@ -694,9 +737,11 @@ void FVulkanDynamicRHI::RHICreateTransition(FRHITransition* Transition, const FR
 		}
 
 		case FRHITransitionInfo::EType::Buffer:
+		{
 			Buffer = ResourceCast(Info.Buffer);
 			UsageFlags = Buffer->GetBufferUsageFlags();
 			break;
+		}
 
 		case FRHITransitionInfo::EType::UAV:
 		{
@@ -737,56 +782,13 @@ void FVulkanDynamicRHI::RHICreateTransition(FRHITransition* Transition, const FR
 
 		const bool bIsDepthStencil = Texture && Texture->IsDepthOrStencilAspect();
 
-		// If the device doesn't support separate depth-stencil layouts, we must merge depth-stencil subresource transitions so we only do one barrier on the image.
-		if (bIsDepthStencil && Info.PlaneSlice != FRHISubresourceRange::kAllSubresources)
+		if (bIsDepthStencil)
 		{
-			int32 Index = DSSubresTransitions.IndexOfByPredicate([Texture, Info](const FDepthStencilSubresTransition& Entry) -> bool {
-				return Entry.Texture->Image == Texture->Image && Entry.ArraySlice == Info.ArraySlice;
-			});
-
-			FDepthStencilSubresTransition* PendingTransition;
-			if (Index == INDEX_NONE)
-			{
-				PendingTransition = &DSSubresTransitions.AddDefaulted_GetRef();
-				PendingTransition->Texture = Texture;
-				PendingTransition->bDepthAccessSet = PendingTransition->bStencilAccessSet = false;
-				PendingTransition->ArraySlice = Info.ArraySlice;
-			}
-			else
-			{
-				PendingTransition = &DSSubresTransitions[Index];
-			}
-
-			if (Info.PlaneSlice == FRHISubresourceRange::kDepthPlaneSlice)
-			{
-				// We don't support multiple transitions on the same aspect.
-				ensure(PendingTransition->bDepthAccessSet == false);
-				PendingTransition->SrcDepthAccess = Info.AccessBefore;
-				PendingTransition->DestDepthAccess = Info.AccessAfter;
-				PendingTransition->bDepthAccessSet = true;
-			}
-			else
-			{
-				ensure(PendingTransition->bStencilAccessSet == false);
-				PendingTransition->SrcStencilAccess = Info.AccessBefore;
-				PendingTransition->DestStencilAccess = Info.AccessAfter;
-				PendingTransition->bStencilAccessSet = true;
-			}
-
-			PendingTransition->bAliasingBarrier = (Info.AccessBefore == ERHIAccess::Discard);
-
-			if (Index == INDEX_NONE)
-			{
-				// Wait until we find the other aspect of the resource.
-				continue;
-			}
-
-			// Now that we have both aspect transitions filled in, we can figure out the layout.
-			GetDepthStencilStageAndAccessFlags(PendingTransition->SrcDepthAccess, PendingTransition->SrcStencilAccess, SrcStageMask, SrcAccessFlags, SrcLayout, true);
-			GetDepthStencilStageAndAccessFlags(PendingTransition->DestDepthAccess, PendingTransition->DestStencilAccess, DstStageMask, DstAccessFlags, DstLayout, false);
-
-			// Remove the pending transition and add the barrier.
-			DSSubresTransitions.RemoveAtSwap(Index);
+			// if we use separate transitions, then just feed them in as they are
+			SrcLayout = GetDepthOrStencilLayout(Info.AccessBefore);
+			DstLayout = GetDepthOrStencilLayout(Info.AccessAfter);
+			GetDepthOrStencilStageAndAccessFlags(Info.AccessBefore, SrcStageMask, SrcAccessFlags);
+			GetDepthOrStencilStageAndAccessFlags(Info.AccessAfter, DstStageMask, DstAccessFlags);
 		}
 		else
 		{
@@ -799,42 +801,26 @@ void FVulkanDynamicRHI::RHICreateTransition(FRHITransition* Transition, const FR
 			// If not compute, remove vertex pipeline bits as only compute updates vertex buffers
 			if (!(SrcStageMask & VK_PIPELINE_STAGE_COMPUTE_SHADER_BIT))
 			{
-				DstStageMask &= ~(VK_PIPELINE_STAGE_VERTEX_SHADER_BIT | VK_PIPELINE_STAGE_TESSELLATION_CONTROL_SHADER_BIT | VK_PIPELINE_STAGE_TESSELLATION_EVALUATION_SHADER_BIT | VK_PIPELINE_STAGE_GEOMETRY_SHADER_BIT);
+				DstStageMask &= ~(VK_PIPELINE_STAGE_VERTEX_SHADER_BIT | VK_PIPELINE_STAGE_GEOMETRY_SHADER_BIT);
 			}
 		}
 
-		// In case of async compute, override the stage and access flags computed above, since only a subset of stages is relevant.
-		if (SrcPipelines == ERHIPipeline::AsyncCompute)
-		{
-			SrcStageMask = SrcStageMask & Device->GetComputeQueue()->GetSupportedStageBits();
-			SrcAccessFlags = SrcAccessFlags & (VK_ACCESS_SHADER_READ_BIT | VK_ACCESS_SHADER_WRITE_BIT);
-		}
-
-		if (DstPipelines == ERHIPipeline::AsyncCompute)
-		{
-			DstStageMask = DstStageMask & Device->GetComputeQueue()->GetSupportedStageBits();
-			DstAccessFlags = DstAccessFlags & (VK_ACCESS_SHADER_READ_BIT | VK_ACCESS_SHADER_WRITE_BIT);
-		}
-
-
 		// If we're not transitioning across pipes and we don't need to perform layout transitions, we can express memory dependencies through a global memory barrier.
-		if ((SrcPipelines == DstPipelines) && (Texture == nullptr || SrcLayout == DstLayout))
+		if ((SrcPipelines == DstPipelines) && (Texture == nullptr || (SrcLayout == DstLayout)) && GVulkanUseMemoryBarrierOpt)
 		{
 			Data->AddMemoryBarrier(SrcAccessFlags, DstAccessFlags, SrcStageMask, DstStageMask);
 			continue;
 		}
 
-		// Add the stages affected by this transition.
-		Data->SrcStageMask |= SrcStageMask;
-		Data->DstStageMask |= DstStageMask;
-
 		if (Buffer != nullptr)
 		{
 			// We only add buffer transitions for cross-pipe transfers.
 			checkSlow( (SrcPipelines != DstPipelines) && (Texture == nullptr) );
-			VkBufferMemoryBarrier& BufferBarrier = Data->BufferBarriers.AddDefaulted_GetRef();
-			BufferBarrier.sType = VK_STRUCTURE_TYPE_BUFFER_MEMORY_BARRIER;
+			VkBufferMemoryBarrier2& BufferBarrier = Data->BufferBarriers.AddDefaulted_GetRef();
+			BufferBarrier.sType = VK_STRUCTURE_TYPE_BUFFER_MEMORY_BARRIER_2;
 			BufferBarrier.pNext = nullptr;
+			BufferBarrier.srcStageMask = SrcStageMask;
+			BufferBarrier.dstStageMask = DstStageMask;
 			BufferBarrier.srcAccessMask = SrcAccessFlags;
 			BufferBarrier.dstAccessMask = DstAccessFlags;
 			BufferBarrier.srcQueueFamilyIndex = SrcQueueFamilyIndex;
@@ -847,14 +833,15 @@ void FVulkanDynamicRHI::RHICreateTransition(FRHITransition* Transition, const FR
 
 		check(Texture != nullptr);
 
+		const VkImageAspectFlags AspectFlags = bIsDepthStencil ? GetDepthStencilAspectMask(Info.PlaneSlice) : Texture->GetFullAspectMask();
 		VkImageSubresourceRange SubresRange;
-		SetupSubresourceRange(SubresRange, Info, Texture->GetFullAspectMask());
+		SetupSubresourceRange(SubresRange, Info, AspectFlags);
 
 		// For some textures, e.g. FVulkanBackBuffer, the image handle may not be set yet, or may be stale, so there's no point storing it here.
 		// We'll set the image to NULL in the barrier info, and RHIEndTransitions will fetch the up to date pointer from the texture, after
 		// OnLayoutTransition is called.
-		VkImageMemoryBarrier& ImgBarrier = Data->ImageBarriers.AddDefaulted_GetRef();
-		SetupImageBarrier(ImgBarrier, VK_NULL_HANDLE, SrcAccessFlags, DstAccessFlags, SrcLayout, DstLayout, SubresRange);
+		VkImageMemoryBarrier2& ImgBarrier = Data->ImageBarriers.AddDefaulted_GetRef();
+		SetupImageBarrier(ImgBarrier, VK_NULL_HANDLE, SrcStageMask, DstStageMask, SrcAccessFlags, DstAccessFlags, SrcLayout, DstLayout, SubresRange);
 		ImgBarrier.srcQueueFamilyIndex = SrcQueueFamilyIndex;
 		ImgBarrier.dstQueueFamilyIndex = DstQueueFamilyIndex;
 
@@ -862,65 +849,6 @@ void FVulkanDynamicRHI::RHICreateTransition(FRHITransition* Transition, const FR
 		ExtraData.BaseTexture = Texture;
 		ExtraData.IsAliasingBarrier = (Info.AccessBefore == ERHIAccess::Discard);
 	}
-
-	// Process any depth-stencil transitions which specify a single sub-resource.
-	for (const FDepthStencilSubresTransition& PendingTransition : DSSubresTransitions)
-	{
-		VkPipelineStageFlags SrcStageMask, DstStageMask;
-		VkAccessFlags SrcAccessFlags, DstAccessFlags;
-		VkImageLayout SrcLayout, DstLayout;
-
-		auto GetOtherAspectAccess = [DepthStencilFlags](ERHIAccess ExplicitAspectAccess) -> ERHIAccess {
-			// If the aspect that was explicitly set has any depth-stencil flags, we'll assume that the other aspect is read-only (DSVRead). When the barrier
-			// is executed, we check the actual state of the other aspect and adjust the layout if it turns out that the assumption was wrong.
-			// If we don't have any depth-stencil flags, we assume that both aspects will be in the same state, since it's presumably SRVGraphics or SRVCompute.
-			return EnumHasAnyFlags(ExplicitAspectAccess, DepthStencilFlags) ? ERHIAccess::DSVRead : ExplicitAspectAccess;
-		};
-
-		const ERHIAccess CopyAccess = ERHIAccess::CopySrc | ERHIAccess::CopyDest;
-		ERHIAccess SrcDepthAccess, SrcStencilAccess, DstDepthAccess, DstStencilAccess;
-		VkImageAspectFlags AspectMask;
-		if (PendingTransition.bDepthAccessSet)
-		{
-			SrcDepthAccess = PendingTransition.SrcDepthAccess;
-			DstDepthAccess = PendingTransition.DestDepthAccess;
-			SrcStencilAccess = GetOtherAspectAccess(SrcDepthAccess);
-			DstStencilAccess = GetOtherAspectAccess(DstDepthAccess);
-			AspectMask = VK_IMAGE_ASPECT_DEPTH_BIT;
-
-			check(!EnumHasAnyFlags(PendingTransition.SrcDepthAccess, CopyAccess));
-			check(!EnumHasAnyFlags(PendingTransition.DestDepthAccess, CopyAccess));
-		}
-		else
-		{
-			SrcStencilAccess = PendingTransition.SrcStencilAccess;
-			DstStencilAccess = PendingTransition.DestStencilAccess;
-			SrcDepthAccess = GetOtherAspectAccess(SrcStencilAccess);
-			DstDepthAccess = GetOtherAspectAccess(DstStencilAccess);
-			AspectMask = VK_IMAGE_ASPECT_STENCIL_BIT;
-
-			check(!EnumHasAnyFlags(PendingTransition.SrcStencilAccess, CopyAccess));
-			check(!EnumHasAnyFlags(PendingTransition.DestStencilAccess, CopyAccess));
-		}
-
-		GetDepthStencilStageAndAccessFlags(SrcDepthAccess, SrcStencilAccess, SrcStageMask, SrcAccessFlags, SrcLayout, true);
-		GetDepthStencilStageAndAccessFlags(DstDepthAccess, DstStencilAccess, DstStageMask, DstAccessFlags, DstLayout, false);
-		
-		// Don't bother trying to figure out the source layout, let the cache fill it in.
-		SrcLayout = VK_IMAGE_LAYOUT_UNDEFINED;
-
-		Data->SrcStageMask |= SrcStageMask;
-		Data->DstStageMask |= DstStageMask;
-
-		VkImageMemoryBarrier& ImgBarrier = Data->ImageBarriers.AddDefaulted_GetRef();
-		SetupImageBarrierEntireRes(ImgBarrier, VK_NULL_HANDLE, SrcAccessFlags, DstAccessFlags, SrcLayout, DstLayout, AspectMask);
-		ImgBarrier.srcQueueFamilyIndex = SrcQueueFamilyIndex;
-		ImgBarrier.dstQueueFamilyIndex = DstQueueFamilyIndex;
-
-		FVulkanPipelineBarrier::ImageBarrierExtraData& ExtraData = Data->ImageBarrierExtras.AddDefaulted_GetRef();
-		ExtraData.BaseTexture = PendingTransition.Texture;
-		ExtraData.IsAliasingBarrier = PendingTransition.bAliasingBarrier;
-	}
 }
 
 void FVulkanDynamicRHI::RHIReleaseTransition(FRHITransition* Transition)
@@ -928,422 +856,592 @@ void FVulkanDynamicRHI::RHIReleaseTransition(FRHITransition* Transition)
 	Transition->GetPrivateData<FVulkanPipelineBarrier>()->~FVulkanPipelineBarrier();
 }
 
-static void GetDepthStencilWritableState(VkImageLayout Layout, bool& bDepthWritable, bool& bStencilWritable)
+static void AddSubresourceTransitions(TArray<VkImageMemoryBarrier>& Barriers, VkPipelineStageFlags& SrcStageMask, const VkImageMemoryBarrier& TemplateBarrier, VkImage ImageHandle, FVulkanImageLayout& CurrentLayout, VkImageLayout DstLayout)
 {
-	switch (Layout)
-	{
-	case VK_IMAGE_LAYOUT_DEPTH_STENCIL_ATTACHMENT_OPTIMAL:
-		bDepthWritable = true;
-		bStencilWritable = true;
-		break;
+	const uint32 FirstPlane = 0;
+	const uint32 LastPlane = CurrentLayout.NumPlanes;
 
-	case VK_IMAGE_LAYOUT_DEPTH_READ_ONLY_STENCIL_ATTACHMENT_OPTIMAL_KHR:
-		bDepthWritable = false;
-		bStencilWritable = true;
-		break;
+	const uint32 FirstLayer = TemplateBarrier.subresourceRange.baseArrayLayer;
+	const uint32 LastLayer = FirstLayer + CurrentLayout.GetSubresRangeLayerCount(TemplateBarrier.subresourceRange);
 
-	case VK_IMAGE_LAYOUT_DEPTH_ATTACHMENT_STENCIL_READ_ONLY_OPTIMAL_KHR:
-		bDepthWritable = true;
-		bStencilWritable = false;
-		break;
+	const uint32 FirstMip = TemplateBarrier.subresourceRange.baseMipLevel;
+	const uint32 LastMip = FirstMip + CurrentLayout.GetSubresRangeMipCount(TemplateBarrier.subresourceRange);
 
-	default:
-		// This includes VK_IMAGE_LAYOUT_DEPTH_STENCIL_READ_ONLY_OPTIMAL, which is what we normally expect for read-only DS.
-		bDepthWritable = false;
-		bStencilWritable = false;
-		break;
+	for (uint32 PlaneIdx = FirstPlane; PlaneIdx < LastPlane; ++PlaneIdx)
+	{
+		for (uint32 LayerIdx = FirstLayer; LayerIdx < LastLayer; ++LayerIdx)
+		{
+			VkImageMemoryBarrier* PrevMipBarrier = nullptr;
+
+			for (uint32 MipIdx = FirstMip; MipIdx < LastMip; ++MipIdx)
+			{
+				VkImageLayout SrcLayout = CurrentLayout.GetSubresLayout(LayerIdx, MipIdx, PlaneIdx);
+
+				// Merge with the previous transition if the previous mip was in the same state as this mip.
+				if (PrevMipBarrier && PrevMipBarrier->oldLayout == SrcLayout)
+				{
+					PrevMipBarrier->subresourceRange.levelCount += 1;
+				}
+				else
+				{
+					if (SrcLayout == DstLayout)
+					{
+						continue;
+					}
+
+					SrcStageMask |= GetVkStageFlagsForLayout(SrcLayout);
+
+					VkImageMemoryBarrier& Barrier = Barriers.AddDefaulted_GetRef();
+					Barrier = TemplateBarrier;
+					Barrier.srcAccessMask = GetVkAccessMaskForLayout(SrcLayout);
+					Barrier.oldLayout = SrcLayout;
+					Barrier.newLayout = DstLayout;
+					Barrier.image = ImageHandle;
+					Barrier.subresourceRange.baseMipLevel = MipIdx;
+					Barrier.subresourceRange.levelCount = 1;
+					Barrier.subresourceRange.baseArrayLayer = LayerIdx;
+					Barrier.subresourceRange.layerCount = 1;
+
+					PrevMipBarrier = &Barrier;
+				}
+			}
+		}
 	}
+	CurrentLayout.Set(DstLayout, TemplateBarrier.subresourceRange);
 }
 
-static void AdjustDepthStencilLayout(VkImageMemoryBarrier& Barrier, VkImageAspectFlags FullAspectMask)
+static void DowngradeBarrier(VkMemoryBarrier& OutBarrier, const VkMemoryBarrier2& InBarrier)
 {
-	// If this barrier specifies a sub-aspect of a depth-stencil surface, we need to make sure that the layout reflects the current state
-	// of the other aspect. The code in RHICreateTransition cannot know the current state, so it assumed it's read-only when it set up the
-	// barrier. If it turns out it's writable, we need to change the layout accordingly.
-	if (FullAspectMask == Barrier.subresourceRange.aspectMask)
-	{
-		return;
-	}
-
-	if (!ensure(Barrier.subresourceRange.aspectMask == VK_IMAGE_ASPECT_DEPTH_BIT || Barrier.subresourceRange.aspectMask == VK_IMAGE_ASPECT_STENCIL_BIT))
-	{
-		return;
-	}
+	OutBarrier.sType = VK_STRUCTURE_TYPE_MEMORY_BARRIER;
+	OutBarrier.pNext = InBarrier.pNext;
+	OutBarrier.srcAccessMask = InBarrier.srcAccessMask;
+	OutBarrier.dstAccessMask = InBarrier.dstAccessMask;
+}
 
-	bool bDepthWritableCurrent, bStencilWritableCurrent;
-	GetDepthStencilWritableState(Barrier.oldLayout, bDepthWritableCurrent, bStencilWritableCurrent);
+static void DowngradeBarrier(VkBufferMemoryBarrier& OutBarrier, const VkBufferMemoryBarrier2& InBarrier)
+{
+	OutBarrier.sType = VK_STRUCTURE_TYPE_BUFFER_MEMORY_BARRIER;
+	OutBarrier.pNext = InBarrier.pNext;
+	OutBarrier.srcAccessMask = InBarrier.srcAccessMask;
+	OutBarrier.dstAccessMask = InBarrier.dstAccessMask;
+	OutBarrier.srcQueueFamilyIndex = InBarrier.srcQueueFamilyIndex;
+	OutBarrier.dstQueueFamilyIndex = InBarrier.dstQueueFamilyIndex;
+	OutBarrier.buffer = InBarrier.buffer;
+	OutBarrier.offset = InBarrier.offset;
+	OutBarrier.size = InBarrier.size;
+}
 
-	bool bDepthWritableNew, bStencilWritableNew;
-	GetDepthStencilWritableState(Barrier.newLayout, bDepthWritableNew, bStencilWritableNew);
+static void DowngradeBarrier(VkImageMemoryBarrier& OutBarrier, const VkImageMemoryBarrier2& InBarrier)
+{
+	OutBarrier.sType = VK_STRUCTURE_TYPE_IMAGE_MEMORY_BARRIER;
+	OutBarrier.pNext = InBarrier.pNext;
+	OutBarrier.srcAccessMask = InBarrier.srcAccessMask;
+	OutBarrier.dstAccessMask = InBarrier.dstAccessMask;
+	OutBarrier.oldLayout = InBarrier.oldLayout;
+	OutBarrier.newLayout = InBarrier.newLayout;
+	OutBarrier.srcQueueFamilyIndex = InBarrier.srcQueueFamilyIndex;
+	OutBarrier.dstQueueFamilyIndex = InBarrier.dstQueueFamilyIndex;
+	OutBarrier.image = InBarrier.image;
+	OutBarrier.subresourceRange = InBarrier.subresourceRange;
 
-	if (Barrier.subresourceRange.aspectMask == VK_IMAGE_ASPECT_DEPTH_BIT)
+	// Last minute conversion if both aspects are on the same sync2 layout then we can do the simple conversion (workaround for InitialLayout)
+	if (OutBarrier.subresourceRange.aspectMask == (VK_IMAGE_ASPECT_DEPTH_BIT | VK_IMAGE_ASPECT_STENCIL_BIT))
 	{
-		if (bStencilWritableCurrent)
+		if (OutBarrier.newLayout == VK_IMAGE_LAYOUT_READ_ONLY_OPTIMAL)
 		{
-			Barrier.newLayout = bDepthWritableNew ? VK_IMAGE_LAYOUT_DEPTH_STENCIL_ATTACHMENT_OPTIMAL : VK_IMAGE_LAYOUT_DEPTH_READ_ONLY_STENCIL_ATTACHMENT_OPTIMAL_KHR;
+			OutBarrier.newLayout = VK_IMAGE_LAYOUT_DEPTH_STENCIL_READ_ONLY_OPTIMAL;
 		}
-	}
-	else
-	{
-		if (bDepthWritableCurrent)
+		else if (OutBarrier.newLayout == VK_IMAGE_LAYOUT_ATTACHMENT_OPTIMAL)
 		{
-			Barrier.newLayout = bStencilWritableNew ? VK_IMAGE_LAYOUT_DEPTH_STENCIL_ATTACHMENT_OPTIMAL : VK_IMAGE_LAYOUT_DEPTH_ATTACHMENT_STENCIL_READ_ONLY_OPTIMAL_KHR;
+			OutBarrier.newLayout = VK_IMAGE_LAYOUT_DEPTH_STENCIL_ATTACHMENT_OPTIMAL;
 		}
 	}
-
-	// Now that we have the correct layout, we can set the mask to include both aspects.
-	Barrier.subresourceRange.aspectMask = VK_IMAGE_ASPECT_DEPTH_BIT | VK_IMAGE_ASPECT_STENCIL_BIT;
 }
 
-void FVulkanCommandListContext::RHIBeginTransitions(TArrayView<const FRHITransition*> Transitions)
+template <typename DstArrayType, typename SrcArrayType>
+static void DowngradeBarrierArray(DstArrayType& TargetArray, const SrcArrayType& SrcArray, VkPipelineStageFlags& MergedSrcStageMask, VkPipelineStageFlags& MergedDstStageMask)
 {
-	static IConsoleVariable* CVarShowTransitions = IConsoleManager::Get().FindConsoleVariable(TEXT("r.ProfileGPU.ShowTransitions"));
-	const bool bShowTransitionEvents = CVarShowTransitions->GetInt() != 0;
-	SCOPED_RHI_CONDITIONAL_DRAW_EVENTF(*this, RHIBeginTransitions, bShowTransitionEvents, TEXT("RHIBeginTransitions"));
+	TargetArray.Reserve(TargetArray.Num() + SrcArray.Num());
+	for (const auto& SrcBarrier : SrcArray)
+	{
+		auto& DstBarrier = TargetArray.AddDefaulted_GetRef();
+		DowngradeBarrier(DstBarrier, SrcBarrier);
+		MergedSrcStageMask |= SrcBarrier.srcStageMask;
+		MergedDstStageMask |= SrcBarrier.dstStageMask;
+	}
+}
 
-	TArray<VkBufferMemoryBarrier, TInlineAllocator<8>> RealBufferBarriers;
-	TArray<VkImageMemoryBarrier, TInlineAllocator<8>> RealImageBarriers;
-	TArray<VulkanRHI::FSemaphore*, TInlineAllocator<8>> SignalSemaphores;
+// The transition processor takes RHITransitions and transforms them into Vulkan barriers (of the template type)
+template<typename MemoryBarrierType, typename BufferBarrierType, typename ImageBarrierType>
+class FTransitionProcessor
+{
+public:
 
-	FVulkanCmdBuffer* CmdBuffer = CommandBufferManager->GetActiveCmdBuffer();
+	FTransitionProcessor(FVulkanCommandListContext& InContext, bool InIsBeginTransition)
+		: Context(InContext)
+		, IsBeginTransition(InIsBeginTransition)
+	{}
 
-	for (const FRHITransition* Transition : Transitions)
+	void Process(TArrayView<const FRHITransition*>& Transitions)
 	{
-		const FVulkanPipelineBarrier* Data = Transition->GetPrivateData<FVulkanPipelineBarrier>();
-
-		// We only care about cross-pipe transitions in this function.
-		if (Data->SrcPipelines == Data->DstPipelines)
+		for (const FRHITransition* Transition : Transitions)
 		{
-			continue;
-		}
+			const FVulkanPipelineBarrier* Data = Transition->GetPrivateData<FVulkanPipelineBarrier>();
 
-		if (Data->SrcStageMask == 0 || Data->DstStageMask == 0)
-		{
-			// This transition didn't specify any resources.
-			check(Data->ImageBarriers.Num() == 0 && Data->BufferBarriers.Num() == 0);
-			continue;
-		}
+			const bool IsCrossPipe = (Data->SrcPipelines != Data->DstPipelines);
 
-#if DO_GUARD_SLOW
-		switch (Data->SrcPipelines)
-		{
-			case ERHIPipeline::Graphics:
-				checkf(!Device->IsRealAsyncComputeContext(this), TEXT("Attempt to begin Graphics -> AsyncCompute transition on the async compute command list."));
-				break;
-			case ERHIPipeline::AsyncCompute:
-				checkf(Device->IsRealAsyncComputeContext(this), TEXT("Attempt to begin AsyncCompute -> Graphics transition on the graphics command list."));
-				break;
-			default:
-				checkNoEntry();
-				break;
-		}
-#endif
+			// We only care about cross-pipe transitions when we begin transitions
+			if (IsBeginTransition && !IsCrossPipe)
+			{
+				continue;
+			}
 
-		RealBufferBarriers.Reset();
-		RealBufferBarriers.Reserve(Data->BufferBarriers.Num());
-		for (const VkBufferMemoryBarrier& BufferBarrier : Data->BufferBarriers)
-		{
-			VkBufferMemoryBarrier& RealBarrier = RealBufferBarriers.AddDefaulted_GetRef();
-			RealBarrier = BufferBarrier;
-			RealBarrier.dstAccessMask = 0; // Release resource from current queue.
-		}
+			ValidateQueue(Data);
 
-		check(Data->ImageBarriers.Num() == Data->ImageBarrierExtras.Num());
-		RealImageBarriers.Reset();
-		RealImageBarriers.Reserve(Data->ImageBarriers.Num());
+			const int32 BufferBatchStartIndex = BufferBarriers.Num();
+			const int32 ImageBatchStartIndex = ImageBarriers.Num();
+			AppendBarriers(Data);
 
-		VkPipelineStageFlags RealSrcStageMask = Data->SrcStageMask;
-		VkPipelineStageFlags RealDstStageMask = Data->DstStageMask;
+			// Patch the src/dst access mask for acquire/release on queue ownership transfers
+			if (IsCrossPipe)
+			{
+				PatchCrossPipeTransitions(BufferBarriers, BufferBatchStartIndex);
+				PatchCrossPipeTransitions(ImageBarriers, ImageBatchStartIndex);
+			}
 
-		for (int32 ImgBarrierIdx = 0; ImgBarrierIdx < Data->ImageBarriers.Num(); ++ImgBarrierIdx)
-		{
-			const VkImageMemoryBarrier& ImageBarrier = Data->ImageBarriers[ImgBarrierIdx];
-			const FVulkanPipelineBarrier::ImageBarrierExtraData& ExtraData = Data->ImageBarrierExtras[ImgBarrierIdx];
-			FVulkanTexture* Texture = ExtraData.BaseTexture;
-			check(Texture->Image != VK_NULL_HANDLE);
+			// Set the images for the barriers from the texture in the extra data and notify of layout changes
+			for (int32 Index = 0; Index < Data->ImageBarrierExtras.Num(); ++Index)
+			{
+				const FVulkanPipelineBarrier::ImageBarrierExtraData* RemainingExtras = &Data->ImageBarrierExtras[Index];
+				FVulkanTexture* Texture = RemainingExtras->BaseTexture;
+				check(Texture->Image != VK_NULL_HANDLE);  // coming in, we should always have an image
 
-			const FVulkanImageLayout& Layout = LayoutManager.GetOrAddFullLayout(*Texture, VK_IMAGE_LAYOUT_UNDEFINED);
+				const int32 TargetIndex = ImageBatchStartIndex + Index;
+				ImageBarrierType* RemainingBarriers = &ImageBarriers[TargetIndex];
 
-			VkAccessFlags SrcAccessFlags;
-			VkImageLayout SrcLayout, DstLayout;
+				// EndTransition handles layout transition notifications
+				if (!IsBeginTransition)
+				{
+					Texture->OnLayoutTransition(Context, RemainingBarriers->newLayout);
+				}
 
-			check(ImageBarrier.newLayout != VK_IMAGE_LAYOUT_UNDEFINED);
-			DstLayout = ImageBarrier.newLayout;
+				RemainingBarriers->image = Texture->Image;
 
-			if ((ImageBarrier.oldLayout == VK_IMAGE_LAYOUT_UNDEFINED) && (!ExtraData.IsAliasingBarrier))
-			{
-				check(Layout.AreAllSubresourcesSameLayout());
-				SrcLayout = Layout.MainLayout;
-				SrcAccessFlags = GetVkAccessMaskForLayout(SrcLayout);
-				RealSrcStageMask |= GetVkStageFlagsForLayout(SrcLayout);
+				if ((RemainingBarriers->image != VK_NULL_HANDLE) && (RemainingBarriers->subresourceRange.aspectMask != 0))
+				{
+					const int32 RemainingCount = Data->ImageBarrierExtras.Num() - Index;
+					ApplyTracking(RemainingBarriers, RemainingExtras, RemainingCount, IsCrossPipe);
+				}
 			}
-			else
+
+			// Loop again to remove invalidated entries
+			for (int32 DstIndex = ImageBatchStartIndex; DstIndex < ImageBarriers.Num();)
 			{
-				checkSlow(Layout.AreSubresourcesSameLayout(ImageBarrier.oldLayout, ImageBarrier.subresourceRange));
-				SrcLayout = ImageBarrier.oldLayout;
-				SrcAccessFlags = ImageBarrier.srcAccessMask;
+				// If Texture is a backbuffer, OnLayoutTransition normally updates the image handle. However, if the viewport got invalidated during the
+				// OnLayoutTransition call, we may end up with a null handle, which is fine.  Merged Depth and Stencil transitions will also result in
+				// in null aspectMask for the extra transition.  Either is fine, just remove the transition.
+				if ((ImageBarriers[DstIndex].image == VK_NULL_HANDLE) || ((ImageBarriers[DstIndex].subresourceRange.aspectMask == 0)))
+				{
+					ImageBarriers.RemoveAtSwap(DstIndex);
+				}
+				else
+				{
+					++DstIndex;
+				}
 			}
 
-			VkImageMemoryBarrier& RealBarrier = RealImageBarriers.AddDefaulted_GetRef();
-			RealBarrier = ImageBarrier;
-			RealBarrier.image = Texture->Image;
-			RealBarrier.srcAccessMask = SrcAccessFlags;
-			RealBarrier.dstAccessMask = 0; // Release resource from current queue.
-			RealBarrier.oldLayout = SrcLayout;
-			RealBarrier.newLayout = DstLayout;
-			
-			// Fix up the destination layout if this barrier specifies a sub-aspect of a depth-stencil 
-			AdjustDepthStencilLayout(RealBarrier, Texture->GetFullAspectMask());
-
-			// We don't update the image layout here. That will be done in RHIEndTransitions.
-		}
-
-		if (EnumHasAnyFlags(Data->SrcPipelines, ERHIPipeline::AsyncCompute))
-		{
-			RealSrcStageMask = RealSrcStageMask & Queue->GetSupportedStageBits();
+			FinishBatch(Data->SrcPipelines, Data->DstPipelines, BufferBatchStartIndex, ImageBatchStartIndex);
 		}
 
-		VulkanRHI::vkCmdPipelineBarrier(CmdBuffer->GetHandle(), RealSrcStageMask, RealDstStageMask, 0, 0, nullptr, RealBufferBarriers.Num(), RealBufferBarriers.GetData(), RealImageBarriers.Num(), RealImageBarriers.GetData());
-
-		check(Data->Semaphore);
-		SignalSemaphores.Add(Data->Semaphore);
-	}
-
-	if (SignalSemaphores.Num() > 0)
-	{
-		CommandBufferManager->SubmitActiveCmdBuffer(SignalSemaphores);
-		CommandBufferManager->PrepareForNewActiveCommandBuffer();
+		FinishAll();
 	}
-}
 
-static void AddSubresourceTransitions(TArray<VkImageMemoryBarrier, TInlineAllocator<8>>& Barriers, VkPipelineStageFlags& SrcStageMask, const VkImageMemoryBarrier& TemplateBarrier, VkImage ImageHandle, FVulkanImageLayout& CurrentLayout, VkImageLayout DstLayout)
-{
-	const uint32 FirstLayer = TemplateBarrier.subresourceRange.baseArrayLayer;
-	const uint32 LastLayer = FirstLayer + CurrentLayout.GetSubresRangeLayerCount(TemplateBarrier.subresourceRange);
-	
-	const uint32 FirstMip = TemplateBarrier.subresourceRange.baseMipLevel;
-	const uint32 LastMip = FirstMip + CurrentLayout.GetSubresRangeMipCount(TemplateBarrier.subresourceRange);
+private:
+	void AppendBarriers(const FVulkanPipelineBarrier* Data);
+	void ApplyTracking(ImageBarrierType* RemainingBarriers, const FVulkanPipelineBarrier::ImageBarrierExtraData* RemainingExtras, int32 RemainingCount, bool IsCrossPipe);
+	void FinishBatch(ERHIPipeline SrcPipelines, ERHIPipeline DstPipelines, int32 BufferBatchStartIndex, int32 ImageBatchStartIndex);
+	void FinishAll();
 
-	for (uint32 LayerIdx = FirstLayer; LayerIdx < LastLayer; ++LayerIdx)
+	// Patches for acquire/release of resources during queue ownership transfers
+	// For async compute, sanitize the stage and access flags to what is supported by the queue
+	template<typename ArrayType>
+	void PatchCrossPipeTransitions(ArrayType& Barriers, int32 StartIndex)
 	{
-		VkImageMemoryBarrier* PrevMipBarrier = nullptr;
-
-		for (uint32 MipIdx = FirstMip; MipIdx < LastMip; ++MipIdx)
+		for (int32 Index = StartIndex; Index < Barriers.Num(); ++Index)
 		{
-			VkImageLayout SrcLayout = CurrentLayout.GetSubresLayout(LayerIdx, MipIdx);
+			auto& Barrier = Barriers[Index];
 
-			// Merge with the previous transition if the previous mip was in the same state as this mip.
-			if (PrevMipBarrier && PrevMipBarrier->oldLayout == SrcLayout)
-			{
-				PrevMipBarrier->subresourceRange.levelCount += 1;
-			}
-			else
+			if (Barrier.srcQueueFamilyIndex != Barrier.dstQueueFamilyIndex)
 			{
-				if (SrcLayout == DstLayout)
+				if (IsBeginTransition)
 				{
-					continue;
+					Barrier.dstAccessMask = 0; // Release resource from current queue.
+				}
+				else
+				{
+					Barrier.srcAccessMask = 0; // Acquire resource on current queue.
 				}
-
-				SrcStageMask |= GetVkStageFlagsForLayout(SrcLayout);
-
-				VkImageMemoryBarrier& Barrier = Barriers.AddDefaulted_GetRef();
-				Barrier = TemplateBarrier;
-				Barrier.srcAccessMask = GetVkAccessMaskForLayout(SrcLayout);
-				Barrier.oldLayout = SrcLayout;
-				Barrier.newLayout = DstLayout;
-				Barrier.image = ImageHandle;
-				Barrier.subresourceRange.baseMipLevel = MipIdx;
-				Barrier.subresourceRange.levelCount = 1;
-				Barrier.subresourceRange.baseArrayLayer = LayerIdx;
-				Barrier.subresourceRange.layerCount = 1;
-
-				PrevMipBarrier = &Barrier;
 			}
 		}
 	}
 
-	CurrentLayout.Set(DstLayout, TemplateBarrier.subresourceRange);
-}
-
-void FVulkanCommandListContext::RHIEndTransitions(TArrayView<const FRHITransition*> Transitions)
-{
-	static IConsoleVariable* CVarShowTransitions = IConsoleManager::Get().FindConsoleVariable(TEXT("r.ProfileGPU.ShowTransitions"));
-	const bool bShowTransitionEvents = CVarShowTransitions->GetInt() != 0;
-	SCOPED_RHI_CONDITIONAL_DRAW_EVENTF(*this, RHIEndTransitions, bShowTransitionEvents, TEXT("RHIEndTransitions"));
-
-	FVulkanCmdBuffer* CmdBuffer = CommandBufferManager->GetActiveCmdBuffer();
-
-	const ERHIPipeline CurrentPipeline = Device->IsRealAsyncComputeContext(this) ? ERHIPipeline::AsyncCompute : ERHIPipeline::Graphics;
-
-	bool bSeenWaitSemaphore = false;
-	for (const FRHITransition* Transition : Transitions)
+	void ValidateQueue(const FVulkanPipelineBarrier* Data)
 	{
-		const FVulkanPipelineBarrier* Data = Transition->GetPrivateData<FVulkanPipelineBarrier>();
-		if (Data->Semaphore == nullptr)
+#if DO_GUARD_SLOW
+		ERHIPipeline RequestedPipeline = IsBeginTransition ? Data->SrcPipelines : Data->DstPipelines;
+		const TCHAR* RequestedPipelineName = IsBeginTransition ? TEXT("SRC") : TEXT("DST");
+		switch (RequestedPipeline)
 		{
-			continue;
+		case ERHIPipeline::Graphics:
+			checkf(!Device->IsRealAsyncComputeContext(this), TEXT("The %s pipeline for this transition is the Graphics queue, but it's submitted on the AsyncCompute queue."), IsBeginTransition ? TEXT("BEGIN") : TEXT("END"));
+			break;
+		case ERHIPipeline::AsyncCompute:
+			checkf(Device->IsRealAsyncComputeContext(this), TEXT("The %s pipeline for this transition is the AsyncCompute queue, but it's submitted on the Graphics queue."), IsBeginTransition ? TEXT("BEGIN") : TEXT("END"));
+			break;
+		default:
+			checkNoEntry();
+			break;
 		}
+#endif
+	}
+
+	FVulkanCommandListContext& Context;
+	const bool IsBeginTransition;
+
+	VkPipelineStageFlags MergedSrcStageMask = 0;
+	VkPipelineStageFlags MergedDstStageMask = 0;
+
+	TArray<MemoryBarrierType> MemoryBarriers;
+	TArray<BufferBarrierType> BufferBarriers;
+	TArray<ImageBarrierType> ImageBarriers;
+};
+
+
+
+// *** FTransitionProcessor for Original Barriers ***
+
+template <> 
+void FTransitionProcessor<VkMemoryBarrier, VkBufferMemoryBarrier, VkImageMemoryBarrier>::FinishBatch(ERHIPipeline SrcPipelines, ERHIPipeline DstPipelines, int32 BufferBatchStartIndex, int32 ImageBatchStartIndex)
+{
+	if ((MemoryBarriers.Num() == 0) && (BufferBarriers.Num() == 0) && (ImageBarriers.Num() == 0))
+	{
+		return;
+	}
+
+	if (SrcPipelines == ERHIPipeline::AsyncCompute)
+	{
+		MergedSrcStageMask &= Context.GetDevice()->GetComputeQueue()->GetSupportedStageBits();
+	}
+	if (DstPipelines == ERHIPipeline::AsyncCompute)
+	{
+		MergedDstStageMask &= Context.GetDevice()->GetComputeQueue()->GetSupportedStageBits();
+	}
+
+	// Submit merged stage masks with arrays of barriers
+	VulkanRHI::vkCmdPipelineBarrier(Context.GetCommandBufferManager()->GetActiveCmdBuffer()->GetHandle(), MergedSrcStageMask, MergedDstStageMask, 0, MemoryBarriers.Num(), MemoryBarriers.GetData(), BufferBarriers.Num(), BufferBarriers.GetData(), ImageBarriers.Num(), ImageBarriers.GetData());
+	MemoryBarriers.Reset();
+	BufferBarriers.Reset();
+	ImageBarriers.Reset();
+	MergedSrcStageMask = 0;
+	MergedDstStageMask = 0;
+}
+
+template <> 
+void FTransitionProcessor<VkMemoryBarrier, VkBufferMemoryBarrier, VkImageMemoryBarrier>::FinishAll()
+{ 
+	// Non-Sync2 already submitted each batch
+}
+
+template <>
+void FTransitionProcessor<VkMemoryBarrier, VkBufferMemoryBarrier, VkImageMemoryBarrier>::AppendBarriers(const FVulkanPipelineBarrier* Data)
+{
+	DowngradeBarrierArray(MemoryBarriers, Data->MemoryBarriers, MergedSrcStageMask, MergedDstStageMask);
+	DowngradeBarrierArray(BufferBarriers, Data->BufferBarriers, MergedSrcStageMask, MergedDstStageMask);
+	DowngradeBarrierArray(ImageBarriers, Data->ImageBarriers, MergedSrcStageMask, MergedDstStageMask);
+}
+
+template <>
+void FTransitionProcessor<VkMemoryBarrier, VkBufferMemoryBarrier, VkImageMemoryBarrier>::ApplyTracking(VkImageMemoryBarrier* RemainingBarriers, const FVulkanPipelineBarrier::ImageBarrierExtraData* RemainingExtras, int32 RemainingCount, bool IsCrossPipe)
+{
+	FVulkanLayoutManager& LayoutManager = Context.GetCommandBufferManager()->GetActiveCmdBuffer()->GetLayoutManager();
+
+	VkImageMemoryBarrier& ImageBarrier = *RemainingBarriers;
+	FVulkanTexture* Texture = RemainingExtras->BaseTexture;
+	check(Texture->Image == ImageBarrier.image);
+
+	// Get a copy of the current layout, we will alter it and feed it back in the manager
+	FVulkanImageLayout Layout = *LayoutManager.GetFullLayout(*Texture, true);
+
+	check((Layout.NumMips == Texture->GetNumMips()) && (Layout.NumLayers == Texture->GetNumberOfArrayLevels()));
+	check((ImageBarrier.subresourceRange.baseArrayLayer + Layout.GetSubresRangeLayerCount(ImageBarrier.subresourceRange)) <= Layout.NumLayers);
+	check((ImageBarrier.subresourceRange.baseMipLevel + Layout.GetSubresRangeMipCount(ImageBarrier.subresourceRange)) <= Layout.NumMips);
+	check(ImageBarrier.newLayout != VK_IMAGE_LAYOUT_UNDEFINED);
 
-		// If the source matches our current queue, we don't need a semaphore (avoid unnecessary semaphores when DstPipelines is ERHIPipeline::All)
-		if (EnumHasAllFlags(Data->DstPipelines, Data->SrcPipelines) && (Data->SrcPipelines == CurrentPipeline))
+	// For Depth/Stencil formats where only one of the aspects is transitioned, look ahead for other barriers on the same resource
+	if (Texture->IsDepthOrStencilAspect() && (Texture->GetFullAspectMask() != ImageBarrier.subresourceRange.aspectMask))
+	{
+		check((ImageBarrier.subresourceRange.aspectMask == VK_IMAGE_ASPECT_DEPTH_BIT) || (ImageBarrier.subresourceRange.aspectMask == VK_IMAGE_ASPECT_STENCIL_BIT));
+		if ((ImageBarrier.oldLayout == VK_IMAGE_LAYOUT_UNDEFINED) && (!RemainingExtras->IsAliasingBarrier))
 		{
-			continue;
+			ImageBarrier.oldLayout = Layout.GetSubresLayout(ImageBarrier.subresourceRange.baseArrayLayer, ImageBarrier.subresourceRange.baseMipLevel, (VkImageAspectFlagBits)ImageBarrier.subresourceRange.aspectMask);
 		}
 
-		if (!bSeenWaitSemaphore)
+		const VkImageAspectFlagBits OtherAspectMask = (VkImageAspectFlagBits)(Texture->GetFullAspectMask() ^ ImageBarrier.subresourceRange.aspectMask);
+		VkImageLayout OtherAspectOldLayout = Layout.GetSubresLayout(ImageBarrier.subresourceRange.baseArrayLayer, ImageBarrier.subresourceRange.baseMipLevel, OtherAspectMask);
+		VkImageLayout OtherAspectNewLayout = OtherAspectOldLayout;
+		for (int32 OtherBarrierIndex = 0; OtherBarrierIndex < RemainingCount; ++OtherBarrierIndex)
 		{
-			if (CommandBufferManager->HasPendingActiveCmdBuffer())
+			VkImageMemoryBarrier& OtherImageBarrier = RemainingBarriers[OtherBarrierIndex];
+			FVulkanTexture* OtherTexture = RemainingExtras[OtherBarrierIndex].BaseTexture;
+			if ((OtherTexture->Image == ImageBarrier.image) && (OtherImageBarrier.subresourceRange.aspectMask == OtherAspectMask))
 			{
-				CommandBufferManager->SubmitActiveCmdBuffer();
-				CommandBufferManager->PrepareForNewActiveCommandBuffer();
-				CmdBuffer = CommandBufferManager->GetActiveCmdBuffer();
+				check(ImageBarrier.subresourceRange.baseArrayLayer == OtherImageBarrier.subresourceRange.baseArrayLayer);
+				check(ImageBarrier.subresourceRange.baseMipLevel == OtherImageBarrier.subresourceRange.baseMipLevel);
+
+				OtherAspectNewLayout = OtherImageBarrier.newLayout;
+				if ((OtherImageBarrier.oldLayout != VK_IMAGE_LAYOUT_UNDEFINED) || RemainingExtras[OtherBarrierIndex].IsAliasingBarrier)
+				{
+					OtherAspectOldLayout = OtherImageBarrier.oldLayout;
+				}
+
+				Layout.Set(OtherImageBarrier.newLayout, OtherImageBarrier.subresourceRange);
+
+				// Make it invalid so that it gets removed when we reach it
+				OtherImageBarrier.subresourceRange.aspectMask = 0;
+
+				break;
 			}
-			bSeenWaitSemaphore = true;
 		}
 
-		CmdBuffer->AddWaitSemaphore(VK_PIPELINE_STAGE_TOP_OF_PIPE_BIT, Data->Semaphore);
-	}
-
-	TArray<VkBufferMemoryBarrier, TInlineAllocator<8>> RealBufferBarriers;
-	TArray<VkImageMemoryBarrier, TInlineAllocator<8>> RealImageBarriers;
-	FVulkanPipelineBarrier RealMemoryBarrier;
+		Layout.Set(ImageBarrier.newLayout, ImageBarrier.subresourceRange);
 
-	for (const FRHITransition* Transition : Transitions)
+		// Fix up the barrier to include both layouts
+		ImageBarrier.oldLayout = GetMergedDepthStencilLayout(ImageBarrier.oldLayout, OtherAspectOldLayout);
+		ImageBarrier.newLayout = GetMergedDepthStencilLayout(ImageBarrier.newLayout, OtherAspectNewLayout);
+		ImageBarrier.subresourceRange.aspectMask |= OtherAspectMask;
+	}
+	else
 	{
-		const FVulkanPipelineBarrier* Data = Transition->GetPrivateData<FVulkanPipelineBarrier>();
-		
-		if (Data->SrcStageMask == 0 || Data->DstStageMask == 0)
+		if ((ImageBarrier.oldLayout == VK_IMAGE_LAYOUT_UNDEFINED) && (!RemainingExtras->IsAliasingBarrier))
 		{
-			// This transition didn't specify any resources.
-			check(Data->ImageBarriers.Num() == 0 && Data->BufferBarriers.Num() == 0);
-			continue;
+			if (Layout.AreAllSubresourcesSameLayout())
+			{
+				ImageBarrier.oldLayout = Layout.MainLayout;
+				ImageBarrier.srcAccessMask = GetVkAccessMaskForLayout(ImageBarrier.oldLayout);
+				MergedSrcStageMask |= GetVkStageFlagsForLayout(ImageBarrier.oldLayout);
+			}
+			else
+			{
+				// Slow path, adds one transition per subresource.
+				check(!IsCrossPipe);
+				AddSubresourceTransitions(ImageBarriers, MergedSrcStageMask, ImageBarrier, Texture->Image, Layout, ImageBarrier.newLayout);
+			}
+		}
+		else
+		{
+			checkSlow(ExtraData.IsAliasingBarrier || Layout.AreSubresourcesSameLayout(ImageBarrier.oldLayout, ImageBarrier.subresourceRange));
 		}
 
-#if DO_GUARD_SLOW
-		switch (Data->DstPipelines)
+		MergedDstStageMask |= GetVkStageFlagsForLayout(ImageBarrier.newLayout);
+
+		if (!IsCrossPipe)
 		{
-		case ERHIPipeline::Graphics:
-			checkf(!Device->IsRealAsyncComputeContext(this), TEXT("Attempt to end AsyncCompute -> Graphics transition on the async compute command list."));
-			break;
-		case ERHIPipeline::AsyncCompute:
-			checkf(Device->IsRealAsyncComputeContext(this), TEXT("Attempt to begin Graphics -> AsyncCompute transition on the graphics command list."));
-			break;
-		default:
-			checkNoEntry();
-			break;
+			if (ImageBarrier.oldLayout == ImageBarrier.newLayout)
+			{
+				// It turns out that we don't need a layout transition after all. We may still need a memory barrier if the
+				// previous access was writable.
+
+				if (MemoryBarriers.Num() == 0)
+				{
+					MemoryBarriers.AddDefaulted();
+					ZeroVulkanStruct(MemoryBarriers[0], VK_STRUCTURE_TYPE_MEMORY_BARRIER);
+				}
+
+				MemoryBarriers[0].srcAccessMask |= ImageBarrier.srcAccessMask;
+				MemoryBarriers[0].dstAccessMask |= ImageBarrier.dstAccessMask;
+			}
 		}
-#endif
 
-		RealMemoryBarrier.MemoryBarrier = Data->MemoryBarrier;
-		RealMemoryBarrier.bHasMemoryBarrier = Data->bHasMemoryBarrier;
+		Layout.Set(ImageBarrier.newLayout, ImageBarrier.subresourceRange);
+	}
+
+	LayoutManager.SetFullLayout(*Texture, Layout);
+}
+
+
+
+// *** FTransitionProcessor for Sync2 Barriers ***
+
+template <> 
+void FTransitionProcessor<VkMemoryBarrier2, VkBufferMemoryBarrier2, VkImageMemoryBarrier2>::FinishBatch(ERHIPipeline SrcPipelines, ERHIPipeline DstPipelines, int32 BufferBatchStartIndex, int32 ImageBatchStartIndex)
+{
+	// :todo-jn: clean up
 
-		check(Data->SrcPipelines != Data->DstPipelines || Data->BufferBarriers.Num() == 0);
-		RealBufferBarriers.Reset();
-		RealBufferBarriers.Reserve(Data->BufferBarriers.Num());
-		for (const VkBufferMemoryBarrier& BufferBarrier : Data->BufferBarriers)
+	if (SrcPipelines == ERHIPipeline::AsyncCompute)
+	{
+		for (int32 Index = BufferBatchStartIndex; Index < BufferBarriers.Num(); ++Index)
 		{
-			VkBufferMemoryBarrier& RealBarrier = RealBufferBarriers.AddDefaulted_GetRef();
-			RealBarrier = BufferBarrier;
-			RealBarrier.srcAccessMask = 0; // Acquire resource on current queue.
+			VkBufferMemoryBarrier2& BufferBarrier = BufferBarriers[Index];
+			BufferBarrier.srcStageMask &= Context.GetDevice()->GetComputeQueue()->GetSupportedStageBits();
+			BufferBarrier.srcAccessMask &= (VK_ACCESS_SHADER_READ_BIT | VK_ACCESS_SHADER_WRITE_BIT);
 		}
 
-		check(Data->ImageBarriers.Num() == Data->ImageBarrierExtras.Num());
-		RealImageBarriers.Reset();
-		RealImageBarriers.Reserve(Data->ImageBarriers.Num());
+		for (int32 Index = ImageBatchStartIndex; Index < ImageBarriers.Num(); ++Index)
+		{
+			VkImageMemoryBarrier2& Barrier = ImageBarriers[Index];
+			Barrier.srcStageMask &= Context.GetDevice()->GetComputeQueue()->GetSupportedStageBits();
+			Barrier.srcAccessMask &= (VK_ACCESS_SHADER_READ_BIT | VK_ACCESS_SHADER_WRITE_BIT);
+		}
+	}
 
-		VkPipelineStageFlags RealSrcStageMask = Data->SrcStageMask;
-		VkPipelineStageFlags RealDstStageMask = Data->DstStageMask;
+	if (DstPipelines == ERHIPipeline::AsyncCompute)
+	{
+		for (int32 Index = BufferBatchStartIndex; Index < BufferBarriers.Num(); ++Index)
+		{
+			VkBufferMemoryBarrier2& BufferBarrier = BufferBarriers[Index];
+			BufferBarrier.dstStageMask &= Context.GetDevice()->GetComputeQueue()->GetSupportedStageBits();
+			BufferBarrier.dstAccessMask &= (VK_ACCESS_SHADER_READ_BIT | VK_ACCESS_SHADER_WRITE_BIT);
+		}
 
-		for (int32 ImgBarrierIdx = 0; ImgBarrierIdx < Data->ImageBarriers.Num(); ++ImgBarrierIdx)
+		for (int32 Index = ImageBatchStartIndex; Index < ImageBarriers.Num(); ++Index)
 		{
-			const VkImageMemoryBarrier& ImageBarrier = Data->ImageBarriers[ImgBarrierIdx];
-			check(ImageBarrier.image == VK_NULL_HANDLE);
+			VkImageMemoryBarrier2& Barrier = ImageBarriers[Index];
+			Barrier.dstStageMask &= Context.GetDevice()->GetComputeQueue()->GetSupportedStageBits();
+			Barrier.dstAccessMask &= (VK_ACCESS_SHADER_READ_BIT | VK_ACCESS_SHADER_WRITE_BIT);
+		}
+	}
 
-			const FVulkanPipelineBarrier::ImageBarrierExtraData& ExtraData = Data->ImageBarrierExtras[ImgBarrierIdx];
-			FVulkanTexture* Texture = ExtraData.BaseTexture;
-			check(Texture->GetCpuReadbackBuffer() == nullptr);
+	// Sync2 does not submit batches, only one big clump at the end
+}
 
-			Texture->OnLayoutTransition(*this, ImageBarrier.newLayout);
+template <>
+void FTransitionProcessor<VkMemoryBarrier2, VkBufferMemoryBarrier2, VkImageMemoryBarrier2>::FinishAll()
+{
+	if ((MemoryBarriers.Num() == 0) && (BufferBarriers.Num() == 0) && (ImageBarriers.Num() == 0))
+	{
+		return;
+	}
 
-			if (Texture->Image == VK_NULL_HANDLE)
-			{
-				// If Texture is a backbuffer, OnLayoutTransition normally updates the image handle. However, if the viewport got invalidated during the
-				// OnLayoutTransition call, we may end up with a null handle, which is fine. Just ignore the transition.
-				continue;
-			}
+	VkDependencyInfo DependencyInfo;
+	DependencyInfo.sType = VK_STRUCTURE_TYPE_DEPENDENCY_INFO;
+	DependencyInfo.pNext = nullptr;
+	DependencyInfo.dependencyFlags = 0;
+	DependencyInfo.memoryBarrierCount = MemoryBarriers.Num();
+	DependencyInfo.pMemoryBarriers = MemoryBarriers.GetData();
+	DependencyInfo.bufferMemoryBarrierCount = BufferBarriers.Num();
+	DependencyInfo.pBufferMemoryBarriers = BufferBarriers.GetData();
+	DependencyInfo.imageMemoryBarrierCount = ImageBarriers.Num();
+	DependencyInfo.pImageMemoryBarriers = ImageBarriers.GetData();
+	VulkanRHI::vkCmdPipelineBarrier2KHR(Context.GetCommandBufferManager()->GetActiveCmdBuffer()->GetHandle(), &DependencyInfo);
+}
 
-			FVulkanImageLayout& Layout = LayoutManager.GetOrAddFullLayout(*Texture, VK_IMAGE_LAYOUT_UNDEFINED);
+template <>
+void FTransitionProcessor<VkMemoryBarrier2, VkBufferMemoryBarrier2, VkImageMemoryBarrier2>::AppendBarriers(const FVulkanPipelineBarrier* Data)
+{
+	MemoryBarriers.Append(Data->MemoryBarriers);
+	BufferBarriers.Append(Data->BufferBarriers);
+	ImageBarriers.Append(Data->ImageBarriers);
+}
 
-			VkAccessFlags SrcAccessFlags;
-			VkImageLayout SrcLayout, DstLayout;
+template <>
+void FTransitionProcessor<VkMemoryBarrier2, VkBufferMemoryBarrier2, VkImageMemoryBarrier2>::ApplyTracking(VkImageMemoryBarrier2* RemainingBarriers, const FVulkanPipelineBarrier::ImageBarrierExtraData* RemainingExtras, int32 RemainingCount, bool IsCrossPipe)
+{
+	// Sync2 only needs to inform tracking of the new layout
+	Context.GetCommandBufferManager()->GetActiveCmdBuffer()->GetLayoutManager().SetLayout(*RemainingExtras->BaseTexture, RemainingBarriers->subresourceRange, RemainingBarriers->newLayout);
 
-			check(ImageBarrier.newLayout != VK_IMAGE_LAYOUT_UNDEFINED);
-			DstLayout = ImageBarrier.newLayout;
+	// No needs to worry about 'surprise' subresource transition
+	// No need to worry about partial depth/stencil layout changes
+	// Identical layouts are chill ("Image memory barriers that do not perform an image layout transition can be specified by setting oldLayout equal to newLayout")
+}
 
-			if ((ImageBarrier.oldLayout == VK_IMAGE_LAYOUT_UNDEFINED) && (!ExtraData.IsAliasingBarrier))
-			{
-				if (Layout.AreAllSubresourcesSameLayout())
-				{
-					SrcLayout = Layout.MainLayout;
-					SrcAccessFlags = GetVkAccessMaskForLayout(SrcLayout);
-					RealSrcStageMask |= GetVkStageFlagsForLayout(SrcLayout);
-				}
-				else
-				{
-					// Slow path, adds one transition per subresource.
-					check(Data->SrcPipelines == Data->DstPipelines);
-					AddSubresourceTransitions(RealImageBarriers, RealSrcStageMask, ImageBarrier, Texture->Image, Layout, DstLayout);
-					continue;
-				}
-			}
-			else
-			{
-				checkSlow(ExtraData.IsAliasingBarrier || Layout.AreSubresourcesSameLayout(ImageBarrier.oldLayout, ImageBarrier.subresourceRange));
-				SrcLayout = ImageBarrier.oldLayout;
-				SrcAccessFlags = ImageBarrier.srcAccessMask;
-			}
 
-			VkImageMemoryBarrier RealBarrier = ImageBarrier;
-			RealBarrier.image = Texture->Image; // Use the up to date image handle.
-			RealBarrier.srcAccessMask = SrcAccessFlags;
-			RealBarrier.oldLayout = SrcLayout;
-			RealBarrier.newLayout = DstLayout;
 
-			// Fix up the destination layout if this barrier specifies a sub-aspect of a depth-stencil surface.
-			AdjustDepthStencilLayout(RealBarrier, Texture->GetFullAspectMask());
 
-			if (Data->SrcPipelines == Data->DstPipelines)
-			{
-				if (RealBarrier.oldLayout == RealBarrier.newLayout)
-				{
-					// It turns out that we don't need a layout transition afterall. We may still need a memory barrier if the
-					// previous access was writable.
-					RealMemoryBarrier.AddMemoryBarrier(RealBarrier.srcAccessMask, RealBarrier.dstAccessMask, 0, 0);
-					continue;
-				}
-			}
-			else
-			{
-				// Acquire resource on current queue.
-				RealBarrier.srcAccessMask = 0;
-			}
 
-			RealImageBarriers.Add(RealBarrier);
+static TArray<VulkanRHI::FSemaphore*> ExtractTransitionSemaphores(TArrayView<const FRHITransition*>& Transitions, ERHIPipeline SkipSrcPipeline = ERHIPipeline::None)
+{
+	TArray<VulkanRHI::FSemaphore*> Semaphores;
 
-			Layout.Set(RealBarrier.newLayout, RealBarrier.subresourceRange);
+	for (const FRHITransition* Transition : Transitions)
+	{
+		const FVulkanPipelineBarrier* Data = Transition->GetPrivateData<FVulkanPipelineBarrier>();
+		if (Data->Semaphore == nullptr)
+		{
+			continue;
 		}
 
-		if ((!RealMemoryBarrier.bHasMemoryBarrier) && (RealBufferBarriers.Num() == 0) && (RealImageBarriers.Num() == 0))
+		// Avoid unnecessary semaphores when DstPipelines is ERHIPipeline::All
+		if ((SkipSrcPipeline != ERHIPipeline::None) && EnumHasAllFlags(Data->DstPipelines, Data->SrcPipelines) && (Data->SrcPipelines == SkipSrcPipeline))
 		{
 			continue;
 		}
 
-		if (EnumHasAnyFlags(Data->DstPipelines, ERHIPipeline::AsyncCompute))
+		Semaphores.Add(Data->Semaphore);
+	}
+
+	return Semaphores;
+}
+
+
+void FVulkanCommandListContext::RHIBeginTransitions(TArrayView<const FRHITransition*> Transitions)
+{
+	static IConsoleVariable* CVarShowTransitions = IConsoleManager::Get().FindConsoleVariable(TEXT("r.ProfileGPU.ShowTransitions"));
+	const bool bShowTransitionEvents = CVarShowTransitions->GetInt() != 0;
+	SCOPED_RHI_CONDITIONAL_DRAW_EVENTF(*this, RHIBeginTransitions, bShowTransitionEvents, TEXT("RHIBeginTransitions"));
+
+	TRACE_CPUPROFILER_EVENT_SCOPE(RHIBeginTransitions);
+
+	if (Device->SupportsParallelRendering())
+	{
+		FTransitionProcessor<VkMemoryBarrier2, VkBufferMemoryBarrier2, VkImageMemoryBarrier2> Processor(*this, true);
+		Processor.Process(Transitions);
+	}
+	else
+	{
+		FTransitionProcessor<VkMemoryBarrier, VkBufferMemoryBarrier, VkImageMemoryBarrier> Processor(*this, true);
+		Processor.Process(Transitions);
+	}
+
+	// Signal semaphores
+	TArray<VulkanRHI::FSemaphore*> SignalSemaphores = ExtractTransitionSemaphores(Transitions);
+	if (SignalSemaphores.Num() > 0)
+	{
+		CommandBufferManager->SubmitActiveCmdBuffer(SignalSemaphores);
+		CommandBufferManager->PrepareForNewActiveCommandBuffer();
+	}
+}
+
+void FVulkanCommandListContext::RHIEndTransitions(TArrayView<const FRHITransition*> Transitions)
+{
+	static IConsoleVariable* CVarShowTransitions = IConsoleManager::Get().FindConsoleVariable(TEXT("r.ProfileGPU.ShowTransitions"));
+	const bool bShowTransitionEvents = CVarShowTransitions->GetInt() != 0;
+	SCOPED_RHI_CONDITIONAL_DRAW_EVENTF(*this, RHIEndTransitions, bShowTransitionEvents, TEXT("RHIEndTransitions"));
+
+	TRACE_CPUPROFILER_EVENT_SCOPE(RHIEndTransitions);
+
+	const ERHIPipeline CurrentPipeline = Device->IsRealAsyncComputeContext(this) ? ERHIPipeline::AsyncCompute : ERHIPipeline::Graphics;
+	TArray<VulkanRHI::FSemaphore*> WaitSemaphores = ExtractTransitionSemaphores(Transitions, CurrentPipeline);
+	if (WaitSemaphores.Num() > 0)
+	{
+		if (CommandBufferManager->HasPendingActiveCmdBuffer())
 		{
-			RealDstStageMask = RealDstStageMask & Queue->GetSupportedStageBits();
+			CommandBufferManager->SubmitActiveCmdBuffer();
+			CommandBufferManager->PrepareForNewActiveCommandBuffer();
 		}
 
-		VulkanRHI::vkCmdPipelineBarrier(CmdBuffer->GetHandle(), RealSrcStageMask, RealDstStageMask, 0, RealMemoryBarrier.bHasMemoryBarrier ? 1 : 0, &RealMemoryBarrier.MemoryBarrier, RealBufferBarriers.Num(), RealBufferBarriers.GetData(), RealImageBarriers.Num(), RealImageBarriers.GetData());
+		FVulkanCmdBuffer* CmdBuffer = CommandBufferManager->GetActiveCmdBuffer();
+		CmdBuffer->AddWaitSemaphore(VK_PIPELINE_STAGE_TOP_OF_PIPE_BIT, WaitSemaphores);
+	}
+
+	if (Device->SupportsParallelRendering())
+	{
+		FTransitionProcessor<VkMemoryBarrier2, VkBufferMemoryBarrier2, VkImageMemoryBarrier2> Processor(*this, false);
+		Processor.Process(Transitions);
+	}
+	else
+	{
+		FTransitionProcessor<VkMemoryBarrier, VkBufferMemoryBarrier, VkImageMemoryBarrier> Processor(*this, false);
+		Processor.Process(Transitions);
 	}
 }
 
@@ -1354,6 +1452,15 @@ void FVulkanPipelineBarrier::AddMemoryBarrier(VkAccessFlags InSrcAccessFlags, Vk
 		VK_ACCESS_INPUT_ATTACHMENT_READ_BIT | VK_ACCESS_SHADER_READ_BIT | VK_ACCESS_COLOR_ATTACHMENT_READ_BIT | VK_ACCESS_DEPTH_STENCIL_ATTACHMENT_READ_BIT |
 		VK_ACCESS_TRANSFER_READ_BIT | VK_ACCESS_MEMORY_READ_BIT;
 
+	if (MemoryBarriers.Num() == 0)
+	{
+		VkMemoryBarrier2& NewBarrier = MemoryBarriers.AddDefaulted_GetRef();
+		ZeroVulkanStruct(NewBarrier, VK_STRUCTURE_TYPE_MEMORY_BARRIER_2);
+	}
+
+	// Mash everything into a single barrier
+	VkMemoryBarrier2& MemoryBarrier = MemoryBarriers[0];
+
 	// We only need a memory barrier if the previous commands wrote to the buffer. In case of a transition from read, an execution barrier is enough.
 	const bool SrcAccessIsRead = ((InSrcAccessFlags & (~ReadMask)) == 0);
 	if (!SrcAccessIsRead)
@@ -1362,9 +1469,8 @@ void FVulkanPipelineBarrier::AddMemoryBarrier(VkAccessFlags InSrcAccessFlags, Vk
 		MemoryBarrier.dstAccessMask |= InDstAccessFlags;
 	}
 
-	SrcStageMask |= InSrcStageMask;
-	DstStageMask |= InDstStageMask;
-	bHasMemoryBarrier = true;
+	MemoryBarrier.srcStageMask |= InSrcStageMask;
+	MemoryBarrier.dstStageMask |= InDstStageMask;
 }
 
 
@@ -1374,14 +1480,14 @@ void FVulkanPipelineBarrier::AddMemoryBarrier(VkAccessFlags InSrcAccessFlags, Vk
 //
 void FVulkanPipelineBarrier::AddImageLayoutTransition(VkImage Image, VkImageLayout SrcLayout, VkImageLayout DstLayout, const VkImageSubresourceRange& SubresourceRange)
 {
-	SrcStageMask |= GetVkStageFlagsForLayout(SrcLayout);
-	DstStageMask |= GetVkStageFlagsForLayout(DstLayout);
+	const VkPipelineStageFlags SrcStageMask = GetVkStageFlagsForLayout(SrcLayout);
+	const VkPipelineStageFlags DstStageMask = GetVkStageFlagsForLayout(DstLayout);
 
-	VkAccessFlags SrcAccessFlags = GetVkAccessMaskForLayout(SrcLayout);
-	VkAccessFlags DstAccessFlags = GetVkAccessMaskForLayout(DstLayout);
+	const VkAccessFlags SrcAccessFlags = GetVkAccessMaskForLayout(SrcLayout);
+	const VkAccessFlags DstAccessFlags = GetVkAccessMaskForLayout(DstLayout);
 
-	VkImageMemoryBarrier& ImgBarrier = ImageBarriers.AddDefaulted_GetRef();
-	SetupImageBarrier(ImgBarrier, Image, SrcAccessFlags, DstAccessFlags, SrcLayout, DstLayout, SubresourceRange);
+	VkImageMemoryBarrier2& ImgBarrier = ImageBarriers.AddDefaulted_GetRef();
+	SetupImageBarrier(ImgBarrier, Image, SrcStageMask, DstStageMask, SrcAccessFlags, DstAccessFlags, SrcLayout, DstLayout, SubresourceRange);
 }
 
 void FVulkanPipelineBarrier::AddImageLayoutTransition(VkImage Image, VkImageAspectFlags AspectMask, const FVulkanImageLayout& SrcLayout, VkImageLayout DstLayout)
@@ -1392,25 +1498,28 @@ void FVulkanPipelineBarrier::AddImageLayoutTransition(VkImage Image, VkImageAspe
 		return;
 	}
 
-	DstStageMask |= GetVkStageFlagsForLayout(DstLayout);
+	const VkPipelineStageFlags DstStageMask = GetVkStageFlagsForLayout(DstLayout);
 	const VkAccessFlags DstAccessFlags = GetVkAccessMaskForLayout(DstLayout);
 
-	VkImageSubresourceRange SubresourceRange = MakeSubresourceRange(AspectMask, 0, 1, 0, 1);
-	for (; SubresourceRange.baseArrayLayer < SrcLayout.NumLayers; ++SubresourceRange.baseArrayLayer)
-	{
-		for (SubresourceRange.baseMipLevel = 0; SubresourceRange.baseMipLevel < SrcLayout.NumMips; ++SubresourceRange.baseMipLevel)
+	ForEachAspect(AspectMask, [&](VkImageAspectFlagBits SingleAspect)
 		{
-			const VkImageLayout SubresourceLayout = SrcLayout.GetSubresLayout(SubresourceRange.baseArrayLayer, SubresourceRange.baseMipLevel);
-			if (SubresourceLayout != DstLayout)
+			VkImageSubresourceRange SubresourceRange = MakeSubresourceRange(SingleAspect, 0, 1, 0, 1);
+			for (; SubresourceRange.baseArrayLayer < SrcLayout.NumLayers; ++SubresourceRange.baseArrayLayer)
 			{
-				SrcStageMask |= GetVkStageFlagsForLayout(SubresourceLayout);
-				const VkAccessFlags SrcAccessFlags = GetVkAccessMaskForLayout(SubresourceLayout);
-
-				VkImageMemoryBarrier& ImgBarrier = ImageBarriers.AddDefaulted_GetRef();
-				SetupImageBarrier(ImgBarrier, Image, SrcAccessFlags, DstAccessFlags, SubresourceLayout, DstLayout, SubresourceRange);
+				for (SubresourceRange.baseMipLevel = 0; SubresourceRange.baseMipLevel < SrcLayout.NumMips; ++SubresourceRange.baseMipLevel)
+				{
+					const VkImageLayout SubresourceLayout = SrcLayout.GetSubresLayout(SubresourceRange.baseArrayLayer, SubresourceRange.baseMipLevel, SingleAspect);
+					if (SubresourceLayout != DstLayout)
+					{
+						const VkPipelineStageFlags SrcStageMask = GetVkStageFlagsForLayout(SubresourceLayout);
+						const VkAccessFlags SrcAccessFlags = GetVkAccessMaskForLayout(SubresourceLayout);
+
+						VkImageMemoryBarrier2& ImgBarrier = ImageBarriers.AddDefaulted_GetRef();
+						SetupImageBarrier(ImgBarrier, Image, SrcStageMask, DstStageMask, SrcAccessFlags, DstAccessFlags, SubresourceLayout, DstLayout, SubresourceRange);
+					}
+				}
 			}
-		}
-	}
+		});
 }
 
 void FVulkanPipelineBarrier::AddImageLayoutTransition(VkImage Image, VkImageAspectFlags AspectMask, VkImageLayout SrcLayout, const FVulkanImageLayout& DstLayout)
@@ -1421,25 +1530,28 @@ void FVulkanPipelineBarrier::AddImageLayoutTransition(VkImage Image, VkImageAspe
 		return;
 	}
 
-	SrcStageMask |= GetVkStageFlagsForLayout(SrcLayout);
+	const VkPipelineStageFlags SrcStageMask = GetVkStageFlagsForLayout(SrcLayout);
 	const VkAccessFlags SrcAccessFlags = GetVkAccessMaskForLayout(SrcLayout);
 
-	VkImageSubresourceRange SubresourceRange = MakeSubresourceRange(AspectMask, 0, 1, 0, 1);
-	for (; SubresourceRange.baseArrayLayer < DstLayout.NumLayers; ++SubresourceRange.baseArrayLayer)
-	{
-		for (SubresourceRange.baseMipLevel = 0; SubresourceRange.baseMipLevel < DstLayout.NumMips; ++SubresourceRange.baseMipLevel)
+	ForEachAspect(AspectMask, [&](VkImageAspectFlagBits SingleAspect)
 		{
-			const VkImageLayout SubresourceLayout = DstLayout.GetSubresLayout(SubresourceRange.baseArrayLayer, SubresourceRange.baseMipLevel);
-			if (SubresourceLayout != SrcLayout)
+			VkImageSubresourceRange SubresourceRange = MakeSubresourceRange(SingleAspect, 0, 1, 0, 1);
+			for (; SubresourceRange.baseArrayLayer < DstLayout.NumLayers; ++SubresourceRange.baseArrayLayer)
 			{
-				DstStageMask |= GetVkStageFlagsForLayout(SubresourceLayout);
-				const VkAccessFlags DstAccessFlags = GetVkAccessMaskForLayout(SubresourceLayout);
-
-				VkImageMemoryBarrier& ImgBarrier = ImageBarriers.AddDefaulted_GetRef();
-				SetupImageBarrier(ImgBarrier, Image, SrcAccessFlags, DstAccessFlags, SrcLayout, SubresourceLayout, SubresourceRange);
+				for (SubresourceRange.baseMipLevel = 0; SubresourceRange.baseMipLevel < DstLayout.NumMips; ++SubresourceRange.baseMipLevel)
+				{
+					const VkImageLayout SubresourceLayout = DstLayout.GetSubresLayout(SubresourceRange.baseArrayLayer, SubresourceRange.baseMipLevel, SingleAspect);
+					if (SubresourceLayout != SrcLayout)
+					{
+						const VkPipelineStageFlags DstStageMask = GetVkStageFlagsForLayout(SubresourceLayout);
+						const VkAccessFlags DstAccessFlags = GetVkAccessMaskForLayout(SubresourceLayout);
+
+						VkImageMemoryBarrier2& ImgBarrier = ImageBarriers.AddDefaulted_GetRef();
+						SetupImageBarrier(ImgBarrier, Image, SrcStageMask, DstStageMask, SrcAccessFlags, DstAccessFlags, SrcLayout, SubresourceLayout, SubresourceRange);
+					}
+				}
 			}
-		}
-	}
+		});
 }
 
 void FVulkanPipelineBarrier::AddImageLayoutTransition(VkImage Image, VkImageAspectFlags AspectMask, const FVulkanImageLayout& SrcLayout, const FVulkanImageLayout& DstLayout)
@@ -1457,26 +1569,29 @@ void FVulkanPipelineBarrier::AddImageLayoutTransition(VkImage Image, VkImageAspe
 		checkf(SrcLayout.NumLayers == DstLayout.NumLayers, TEXT("Source (%d) and Destination (%d) layer count mismatch!"), SrcLayout.NumLayers, DstLayout.NumLayers);
 		checkf(SrcLayout.NumMips == DstLayout.NumMips, TEXT("Source (%d) and Destination (%d) mip count mismatch!"), SrcLayout.NumMips, DstLayout.NumMips);
 
-		VkImageSubresourceRange SubresourceRange = MakeSubresourceRange(AspectMask, 0, 1, 0, 1);
-		for (; SubresourceRange.baseArrayLayer < DstLayout.NumLayers; ++SubresourceRange.baseArrayLayer)
-		{
-			for (SubresourceRange.baseMipLevel = 0; SubresourceRange.baseMipLevel < DstLayout.NumMips; ++SubresourceRange.baseMipLevel)
+		ForEachAspect(AspectMask, [&](VkImageAspectFlagBits SingleAspect)
 			{
-				const VkImageLayout SrcSubresourceLayout = SrcLayout.GetSubresLayout(SubresourceRange.baseArrayLayer, SubresourceRange.baseMipLevel);
-				const VkImageLayout DstSubresourceLayout = DstLayout.GetSubresLayout(SubresourceRange.baseArrayLayer, SubresourceRange.baseMipLevel);
-				if (SrcSubresourceLayout != DstSubresourceLayout)
+				VkImageSubresourceRange SubresourceRange = MakeSubresourceRange(SingleAspect, 0, 1, 0, 1);
+				for (; SubresourceRange.baseArrayLayer < DstLayout.NumLayers; ++SubresourceRange.baseArrayLayer)
 				{
-					SrcStageMask |= GetVkStageFlagsForLayout(SrcSubresourceLayout);
-					const VkAccessFlags SrcAccessFlags = GetVkAccessMaskForLayout(SrcSubresourceLayout);
-
-					DstStageMask |= GetVkStageFlagsForLayout(DstSubresourceLayout);
-					const VkAccessFlags DstAccessFlags = GetVkAccessMaskForLayout(DstSubresourceLayout);
-
-					VkImageMemoryBarrier& ImgBarrier = ImageBarriers.AddDefaulted_GetRef();
-					SetupImageBarrier(ImgBarrier, Image, SrcAccessFlags, DstAccessFlags, SrcSubresourceLayout, DstSubresourceLayout, SubresourceRange);
+					for (SubresourceRange.baseMipLevel = 0; SubresourceRange.baseMipLevel < DstLayout.NumMips; ++SubresourceRange.baseMipLevel)
+					{
+						const VkImageLayout SrcSubresourceLayout = SrcLayout.GetSubresLayout(SubresourceRange.baseArrayLayer, SubresourceRange.baseMipLevel, SingleAspect);
+						const VkImageLayout DstSubresourceLayout = DstLayout.GetSubresLayout(SubresourceRange.baseArrayLayer, SubresourceRange.baseMipLevel, SingleAspect);
+						if (SrcSubresourceLayout != DstSubresourceLayout)
+						{
+							const VkPipelineStageFlags SrcStageMask = GetVkStageFlagsForLayout(SrcSubresourceLayout);
+							const VkAccessFlags SrcAccessFlags = GetVkAccessMaskForLayout(SrcSubresourceLayout);
+
+							const VkPipelineStageFlags DstStageMask = GetVkStageFlagsForLayout(DstSubresourceLayout);
+							const VkAccessFlags DstAccessFlags = GetVkAccessMaskForLayout(DstSubresourceLayout);
+
+							VkImageMemoryBarrier2& ImgBarrier = ImageBarriers.AddDefaulted_GetRef();
+							SetupImageBarrier(ImgBarrier, Image, SrcStageMask, DstStageMask, SrcAccessFlags, DstAccessFlags, SrcSubresourceLayout, DstSubresourceLayout, SubresourceRange);
+						}
+					}
 				}
-			}
-		}
+			});
 	}
 }
 
@@ -1499,9 +1614,6 @@ void FVulkanPipelineBarrier::AddImageAccessTransition(const FVulkanTexture& Surf
 		ImgDstStage &= ~(VK_PIPELINE_STAGE_VERTEX_SHADER_BIT | VK_PIPELINE_STAGE_TESSELLATION_CONTROL_SHADER_BIT | VK_PIPELINE_STAGE_TESSELLATION_EVALUATION_SHADER_BIT | VK_PIPELINE_STAGE_GEOMETRY_SHADER_BIT);
 	}
 
-	SrcStageMask |= ImgSrcStage;
-	DstStageMask |= ImgDstStage;
-
 	if (SrcLayout == VK_IMAGE_LAYOUT_UNDEFINED)
 	{
 		SrcLayout = InOutLayout;
@@ -1517,17 +1629,56 @@ void FVulkanPipelineBarrier::AddImageAccessTransition(const FVulkanTexture& Surf
 		DstLayout = VK_IMAGE_LAYOUT_GENERAL;
 	}
 
-	VkImageMemoryBarrier& ImgBarrier = ImageBarriers.AddDefaulted_GetRef();
-	SetupImageBarrier(ImgBarrier, Surface.Image, SrcAccessFlags, DstAccessFlags, SrcLayout, DstLayout, SubresourceRange);
+	VkImageMemoryBarrier2& ImgBarrier = ImageBarriers.AddDefaulted_GetRef();
+	SetupImageBarrier(ImgBarrier, Surface.Image, ImgSrcStage, ImgDstStage, SrcAccessFlags, DstAccessFlags, SrcLayout, DstLayout, SubresourceRange);
 
 	InOutLayout = DstLayout;
 }
 
 void FVulkanPipelineBarrier::Execute(VkCommandBuffer CmdBuffer)
 {
-	if (bHasMemoryBarrier || ImageBarriers.Num() != 0)
+	if (MemoryBarriers.Num() != 0 || BufferBarriers.Num() != 0 || ImageBarriers.Num() != 0)
 	{
-		VulkanRHI::vkCmdPipelineBarrier(CmdBuffer, SrcStageMask, DstStageMask, 0, bHasMemoryBarrier ? 1 : 0, &MemoryBarrier, 0, nullptr, ImageBarriers.Num(), ImageBarriers.GetData());
+		VkPipelineStageFlags SrcStageMask = 0;
+		VkPipelineStageFlags DstStageMask = 0;
+
+		TArray<VkMemoryBarrier, TInlineAllocator<1>> TempMemoryBarriers;
+		DowngradeBarrierArray(TempMemoryBarriers, MemoryBarriers, SrcStageMask, DstStageMask);
+
+		TArray<VkBufferMemoryBarrier> TempBufferBarriers;
+		DowngradeBarrierArray(TempBufferBarriers, BufferBarriers, SrcStageMask, DstStageMask);
+
+		TArray<VkImageMemoryBarrier, TInlineAllocator<2>> TempImageBarriers;
+		DowngradeBarrierArray(TempImageBarriers, ImageBarriers, SrcStageMask, DstStageMask);
+
+		VulkanRHI::vkCmdPipelineBarrier(CmdBuffer, SrcStageMask, DstStageMask, 0, TempMemoryBarriers.Num(), TempMemoryBarriers.GetData(), 
+			TempBufferBarriers.Num(), TempBufferBarriers.GetData(), TempImageBarriers.Num(), TempImageBarriers.GetData());
+	}
+}
+
+void FVulkanPipelineBarrier::Execute(FVulkanCmdBuffer* CmdBuffer)
+{
+	if (MemoryBarriers.Num() != 0 || BufferBarriers.Num() != 0 || ImageBarriers.Num() != 0)
+	{
+		if (CmdBuffer->GetDevice()->SupportsParallelRendering())
+		{
+			VkDependencyInfo DependencyInfo;
+			DependencyInfo.sType = VK_STRUCTURE_TYPE_DEPENDENCY_INFO;
+			DependencyInfo.pNext = nullptr;
+			DependencyInfo.dependencyFlags = 0;
+			DependencyInfo.memoryBarrierCount = MemoryBarriers.Num();
+			DependencyInfo.pMemoryBarriers = MemoryBarriers.GetData();
+			DependencyInfo.bufferMemoryBarrierCount = BufferBarriers.Num();
+			DependencyInfo.pBufferMemoryBarriers = BufferBarriers.GetData();
+			DependencyInfo.imageMemoryBarrierCount = ImageBarriers.Num();
+			DependencyInfo.pImageMemoryBarriers = ImageBarriers.GetData();
+			VulkanRHI::vkCmdPipelineBarrier2KHR(CmdBuffer->GetHandle(), &DependencyInfo);
+		}
+		else
+		{
+			// Call the original execute with older types
+			Execute(CmdBuffer->GetHandle());
+		}
 	}
 }
 
@@ -1545,7 +1696,7 @@ VkImageSubresourceRange FVulkanPipelineBarrier::MakeSubresourceRange(VkImageAspe
 //
 // Used when we need to change the layout of a single image. Some plug-ins call this function from outside the RHI (Steam VR, at the time of writing this).
 //
-void VulkanSetImageLayout(VkCommandBuffer CmdBuffer, VkImage Image, VkImageLayout OldLayout, VkImageLayout NewLayout, const VkImageSubresourceRange& SubresourceRange)
+void VulkanSetImageLayout(FVulkanCmdBuffer* CmdBuffer, VkImage Image, VkImageLayout OldLayout, VkImageLayout NewLayout, const VkImageSubresourceRange& SubresourceRange)
 {
 	FVulkanPipelineBarrier Barrier;
 	Barrier.AddImageLayoutTransition(Image, OldLayout, NewLayout, SubresourceRange);
@@ -1559,19 +1710,25 @@ bool FVulkanImageLayout::AreSubresourcesSameLayout(VkImageLayout Layout, const V
 		return MainLayout == Layout;
 	}
 
+	const uint32 FirstPlane = (SubresourceRange.aspectMask == VK_IMAGE_ASPECT_STENCIL_BIT) ? NumPlanes - 1 : 0;
+	const uint32 LastPlane = (SubresourceRange.aspectMask & VK_IMAGE_ASPECT_STENCIL_BIT) ? NumPlanes : 1;
+
 	const uint32 FirstLayer = SubresourceRange.baseArrayLayer;
 	const uint32 LastLayer = FirstLayer + GetSubresRangeLayerCount(SubresourceRange);
 
 	const uint32 FirstMip = SubresourceRange.baseMipLevel;
 	const uint32 LastMip = FirstMip + GetSubresRangeMipCount(SubresourceRange);
 
-	for (uint32 LayerIdx = FirstLayer; LayerIdx < LastLayer; ++LayerIdx)
+	for (uint32 PlaneIdx = FirstPlane; PlaneIdx < LastPlane; ++PlaneIdx)
 	{
-		for (uint32 MipIdx = FirstMip; MipIdx < LastMip; ++MipIdx)
+		for (uint32 LayerIdx = FirstLayer; LayerIdx < LastLayer; ++LayerIdx)
 		{
-			if (SubresLayouts[LayerIdx * NumMips + MipIdx] != Layout)
+			for (uint32 MipIdx = FirstMip; MipIdx < LastMip; ++MipIdx)
 			{
-				return false;
+				if (SubresLayouts[(PlaneIdx * NumLayers * NumMips) + (LayerIdx * NumMips) + MipIdx] != Layout)
+				{
+					return false;
+				}
 			}
 		}
 	}
@@ -1586,8 +1743,8 @@ void FVulkanImageLayout::CollapseSubresLayoutsIfSame()
 		return;
 	}
 
-	VkImageLayout Layout = SubresLayouts[0];
-	for (uint32 i = 1; i < NumLayers * NumMips; ++i)
+	const VkImageLayout Layout = SubresLayouts[0];
+	for (uint32 i = 1; i < NumPlanes * NumLayers * NumMips; ++i)
 	{
 		if (SubresLayouts[i] != Layout)
 		{
@@ -1601,13 +1758,18 @@ void FVulkanImageLayout::CollapseSubresLayoutsIfSame()
 
 void FVulkanImageLayout::Set(VkImageLayout Layout, const VkImageSubresourceRange& SubresourceRange)
 {
+	const uint32 FirstPlane = (SubresourceRange.aspectMask == VK_IMAGE_ASPECT_STENCIL_BIT) ? NumPlanes - 1 : 0;
+	const uint32 LastPlane = (SubresourceRange.aspectMask & VK_IMAGE_ASPECT_STENCIL_BIT) ? NumPlanes : 1;
+
 	const uint32 FirstLayer = SubresourceRange.baseArrayLayer;
 	const uint32 LayerCount = GetSubresRangeLayerCount(SubresourceRange);
 
 	const uint32 FirstMip = SubresourceRange.baseMipLevel;
 	const uint32 MipCount = GetSubresRangeMipCount(SubresourceRange);
 
-	if (FirstLayer == 0 && LayerCount == NumLayers && FirstMip == 0 && MipCount == NumMips)
+	if (FirstPlane == 0 && LastPlane == NumPlanes &&
+		FirstLayer == 0 && LayerCount == NumLayers && 
+		FirstMip == 0 && MipCount == NumMips)
 	{
 		// We're setting the entire resource to the same layout.
 		MainLayout = Layout;
@@ -1617,7 +1779,7 @@ void FVulkanImageLayout::Set(VkImageLayout Layout, const VkImageSubresourceRange
 
 	if (SubresLayouts.Num() == 0)
 	{
-		uint32 SubresLayoutCount = (NumLayers + FirstLayer) * NumMips;
+		const uint32 SubresLayoutCount = NumPlanes * NumLayers * NumMips;
 		SubresLayouts.SetNum(SubresLayoutCount);
 		for (uint32 i = 0; i < SubresLayoutCount; ++i)
 		{
@@ -1625,11 +1787,14 @@ void FVulkanImageLayout::Set(VkImageLayout Layout, const VkImageSubresourceRange
 		}
 	}
 
-	for (uint32 Layer = FirstLayer; Layer < FirstLayer + LayerCount; ++Layer)
+	for (uint32 Plane = FirstPlane; Plane < LastPlane; ++Plane)
 	{
-		for (uint32 Mip = FirstMip; Mip < FirstMip + MipCount; ++Mip)
+		for (uint32 Layer = FirstLayer; Layer < FirstLayer + LayerCount; ++Layer)
 		{
-			SubresLayouts[Layer * NumMips + Mip] = Layout;
+			for (uint32 Mip = FirstMip; Mip < FirstMip + MipCount; ++Mip)
+			{
+				SubresLayouts[Plane * (NumLayers * NumMips) + Layer * NumMips + Mip] = Layout;
+			}
 		}
 	}
 
@@ -1638,239 +1803,119 @@ void FVulkanImageLayout::Set(VkImageLayout Layout, const VkImageSubresourceRange
 	CollapseSubresLayoutsIfSame();
 }
 
-void FVulkanLayoutManager::Destroy(FVulkanDevice& InDevice, FVulkanLayoutManager* Immediate)
+VkImageLayout FVulkanLayoutManager::GetDefaultLayout(FVulkanCmdBuffer* CmdBuffer, const FVulkanTexture& VulkanTexture, ERHIAccess DesiredAccess)
 {
-	check(!GIsRHIInitialized);
-
-	if (Immediate)
+	switch (DesiredAccess)
 	{
-		Immediate->RenderPasses.Append(RenderPasses);
-		Immediate->Framebuffers.Append(Framebuffers);
-	}
-	else
+	case ERHIAccess::SRVCompute:
+	case ERHIAccess::SRVGraphics:
+	case ERHIAccess::SRVMask:
 	{
-		for (auto& Pair : RenderPasses)
+		if (VulkanTexture.IsDepthOrStencilAspect())
 		{
-			delete Pair.Value;
-		}
+			const bool bSupportsParallelRendering = CmdBuffer->GetDevice()->SupportsParallelRendering();
+			if (bSupportsParallelRendering)
+			{
+				return VK_IMAGE_LAYOUT_READ_ONLY_OPTIMAL;
 
-		for (auto& Pair : Framebuffers)
-		{
-			FFramebufferList* List = Pair.Value;
-			for (int32 Index = List->Framebuffer.Num() - 1; Index >= 0; --Index)
+			}
+			else
 			{
-				List->Framebuffer[Index]->Destroy(InDevice);
-				delete List->Framebuffer[Index];
+				// The only layout that can't be hardcoded...  Even if we only use depth, we need to know stencil
+				FVulkanLayoutManager& LayoutMgr = CmdBuffer->GetLayoutManager();
+				const VkImageLayout DepthLayout = LayoutMgr.GetDepthStencilHint(VulkanTexture, VK_IMAGE_ASPECT_DEPTH_BIT);
+				const VkImageLayout StencilLayout = LayoutMgr.GetDepthStencilHint(VulkanTexture, VK_IMAGE_ASPECT_STENCIL_BIT);
+				return GetMergedDepthStencilLayout(DepthLayout, StencilLayout);
 			}
-			delete List;
+		}
+		else
+		{
+			return VulkanTexture.SupportsSampling() ? VK_IMAGE_LAYOUT_SHADER_READ_ONLY_OPTIMAL : VK_IMAGE_LAYOUT_GENERAL;
 		}
 	}
 
-	RenderPasses.Reset();
-	Framebuffers.Reset();
-}
+	case ERHIAccess::UAVCompute:
+	case ERHIAccess::UAVGraphics:
+	case ERHIAccess::UAVMask: return VK_IMAGE_LAYOUT_GENERAL;
 
-FVulkanFramebuffer* FVulkanLayoutManager::GetOrCreateFramebuffer(FVulkanDevice& InDevice, const FRHISetRenderTargetsInfo& RenderTargetsInfo, const FVulkanRenderTargetLayout& RTLayout, FVulkanRenderPass* RenderPass)
-{
-	uint32 RTLayoutHash = RTLayout.GetRenderPassCompatibleHash();
+	case ERHIAccess::CopySrc: return VK_IMAGE_LAYOUT_TRANSFER_SRC_OPTIMAL;
+	case ERHIAccess::CopyDest: return VK_IMAGE_LAYOUT_TRANSFER_DST_OPTIMAL;
 
-	uint64 MipsAndSlicesValues[MaxSimultaneousRenderTargets];
-	for (int32 Index = 0; Index < MaxSimultaneousRenderTargets; ++Index)
-	{
-		MipsAndSlicesValues[Index] = ((uint64)RenderTargetsInfo.ColorRenderTarget[Index].ArraySliceIndex << (uint64)32) | (uint64)RenderTargetsInfo.ColorRenderTarget[Index].MipIndex;
-	}
-	RTLayoutHash = FCrc::MemCrc32(MipsAndSlicesValues, sizeof(MipsAndSlicesValues), RTLayoutHash);
+	case ERHIAccess::DSVRead: return VK_IMAGE_LAYOUT_READ_ONLY_OPTIMAL;
+	case ERHIAccess::DSVWrite: return VK_IMAGE_LAYOUT_ATTACHMENT_OPTIMAL;
 
-	FFramebufferList** FoundFramebufferList = Framebuffers.Find(RTLayoutHash);
-	FFramebufferList* FramebufferList = nullptr;
-	if (FoundFramebufferList)
+	case ERHIAccess::ShadingRateSource:
 	{
-		FramebufferList = *FoundFramebufferList;
-
-		for (int32 Index = 0; Index < FramebufferList->Framebuffer.Num(); ++Index)
+		if (GRHIVariableRateShadingImageDataType == VRSImage_Palette)
 		{
-			const VkRect2D RenderArea = FramebufferList->Framebuffer[Index]->GetRenderArea();
-
-			if (FramebufferList->Framebuffer[Index]->Matches(RenderTargetsInfo) &&
-				((RTLayout.GetExtent2D().width == RenderArea.extent.width) && (RTLayout.GetExtent2D().height == RenderArea.extent.height) &&
-				(RTLayout.GetOffset2D().x == RenderArea.offset.x) && (RTLayout.GetOffset2D().y == RenderArea.offset.y)))
-			{
-				return FramebufferList->Framebuffer[Index];
-			}
+			return VK_IMAGE_LAYOUT_FRAGMENT_SHADING_RATE_ATTACHMENT_OPTIMAL_KHR;
+		}
+		else if (GRHIVariableRateShadingImageDataType == VRSImage_Fractional)
+		{
+			return VK_IMAGE_LAYOUT_FRAGMENT_DENSITY_MAP_OPTIMAL_EXT;
 		}
-	}
-	else
-	{
-		FramebufferList = new FFramebufferList;
-		Framebuffers.Add(RTLayoutHash, FramebufferList);
-	}
-
-	FVulkanFramebuffer* Framebuffer = new FVulkanFramebuffer(InDevice, RenderTargetsInfo, RTLayout, *RenderPass);
-	FramebufferList->Framebuffer.Add(Framebuffer);
-	return Framebuffer;
-}
-
-void FVulkanLayoutManager::ValidateRenderPassColorEntry(const FRHIRenderPassInfo::FColorEntry& ColorEntry, bool bResolveTarget, FVulkanPipelineBarrier& Barrier)
-{
-	FRHITexture* Texture = bResolveTarget ? ColorEntry.ResolveTarget : ColorEntry.RenderTarget;
-	CA_ASSUME(Texture);
-	FVulkanTexture& Surface = *FVulkanTexture::Cast(Texture);
-	check(Surface.Image != VK_NULL_HANDLE);
-
-	// Check that the image is in the correct layout for rendering.
-	const VkImageLayout ExpectedVkLayout = VK_IMAGE_LAYOUT_COLOR_ATTACHMENT_OPTIMAL;
-	FVulkanImageLayout& LayoutInfo = GetOrAddFullLayout(Surface, VK_IMAGE_LAYOUT_UNDEFINED);
-	const VkImageLayout CurrentVkLayout = LayoutInfo.GetSubresLayout(ColorEntry.ArraySlice, ColorEntry.MipIndex);
-	if (ensureMsgf(CurrentVkLayout == ExpectedVkLayout, TEXT("%s target is in layout %s, expected %s. Please add a transition before starting the render pass."), bResolveTarget ? TEXT("Resolve") : TEXT("Color"), VK_TYPE_TO_STRING(VkImageLayout, CurrentVkLayout), VK_TYPE_TO_STRING(VkImageLayout, ExpectedVkLayout)))
-	{
-		return;
 	}
 
-	// If all the subresources are in the same layout, add a transition for the entire image.
-	if (LayoutInfo.AreAllSubresourcesSameLayout())
-	{
-		Barrier.AddImageLayoutTransition(Surface.Image, CurrentVkLayout, ExpectedVkLayout, FVulkanPipelineBarrier::MakeSubresourceRange(VK_IMAGE_ASPECT_COLOR_BIT));
-		LayoutInfo.MainLayout = ExpectedVkLayout;
-		return;
+	default:
+		checkNoEntry();
+		return VK_IMAGE_LAYOUT_UNDEFINED;
 	}
-
-	// Transition only the mip and layer we're rendering to.
-	VkImageSubresourceRange SubresRange = FVulkanPipelineBarrier::MakeSubresourceRange(VK_IMAGE_ASPECT_COLOR_BIT, ColorEntry.MipIndex, 1, ColorEntry.ArraySlice, 1);
-	Barrier.AddImageLayoutTransition(Surface.Image, CurrentVkLayout, ExpectedVkLayout, SubresRange);
-	LayoutInfo.Set(ExpectedVkLayout, SubresRange);
 }
 
-void FVulkanLayoutManager::BeginRenderPass(FVulkanCommandListContext& Context, FVulkanDevice& InDevice, FVulkanCmdBuffer* CmdBuffer, const FRHIRenderPassInfo& RPInfo, const FVulkanRenderTargetLayout& RTLayout, FVulkanRenderPass* RenderPass, FVulkanFramebuffer* Framebuffer)
+VkImageLayout FVulkanLayoutManager::SetExpectedLayout(FVulkanCmdBuffer* CmdBuffer, const FVulkanTexture& VulkanTexture, ERHIAccess DesiredAccess)
 {
-	check(!CurrentRenderPass);
-	// (NumRT + 1 [Depth] ) * 2 [surface + resolve]
-	VkClearValue ClearValues[(MaxSimultaneousRenderTargets + 1) * 2];
-	uint32 ClearValueIndex = 0;
-	bool bNeedsClearValues = RenderPass->GetNumUsedClearValues() > 0;
-	FMemory::Memzero(ClearValues);
-
-	int32 NumColorTargets = RPInfo.GetNumColorRenderTargets();
-	int32 Index = 0;
+	const VkImageLayout ExpectedLayout = GetDefaultLayout(CmdBuffer, VulkanTexture, DesiredAccess);
+	VkImageLayout PreviousLayout = ExpectedLayout;
 
-	FVulkanPipelineBarrier Barrier;
+	FVulkanLayoutManager& LayoutMgr = CmdBuffer->GetLayoutManager();
 
-	for (Index = 0; Index < NumColorTargets; ++Index)
+	// This code path is not safe for multi-threaded code
+	if (!LayoutMgr.bWriteOnly || GVulkanAutoCorrectExpectedLayouts)
 	{
-		const FRHIRenderPassInfo::FColorEntry& ColorEntry = RPInfo.ColorRenderTargets[Index];
-
-		FRHITexture* ColorTexture = ColorEntry.RenderTarget;
-		CA_ASSUME(ColorTexture);
-		FVulkanTexture& ColorSurface = *FVulkanTexture::Cast(ColorTexture);
-		const bool bPassPerformsResolve = ColorSurface.GetNumSamples() > 1 && ColorEntry.ResolveTarget;
-
-		ValidateRenderPassColorEntry(ColorEntry, false, Barrier);
-		if (bPassPerformsResolve)
-		{
-			check(ColorEntry.ResolveTarget != ColorEntry.RenderTarget);
-			ValidateRenderPassColorEntry(ColorEntry, true, Barrier);
-		}
+		const FVulkanImageLayout* SrcLayout = LayoutMgr.Layouts.Find(VulkanTexture.Image);
 
-		if (GetLoadAction(ColorEntry.Action) == ERenderTargetLoadAction::ELoad)
+		if (!SrcLayout)
 		{
-			// Insert a barrier if we're loading from any color targets, to make sure the passes aren't reordered and we end up running before
-			// the pass we're supposed to read from.
-			const VkAccessFlags AccessMask = VK_ACCESS_COLOR_ATTACHMENT_WRITE_BIT;
-			const VkPipelineStageFlags StageMask = VK_PIPELINE_STAGE_COLOR_ATTACHMENT_OUTPUT_BIT;
-			Barrier.AddMemoryBarrier(AccessMask, AccessMask, StageMask, StageMask);
+			FVulkanLayoutManager& QueueLayoutMgr = CmdBuffer->GetOwner()->GetMgr().GetCommandListContext()->GetQueue()->GetLayoutManager();
+			SrcLayout = QueueLayoutMgr.Layouts.Find(VulkanTexture.Image);
 		}
 
-		if (bNeedsClearValues)
+		if (SrcLayout && (!SrcLayout->AreAllSubresourcesSameLayout() || (SrcLayout->MainLayout != ExpectedLayout)))
 		{
-			const FLinearColor& ClearColor = ColorTexture->HasClearValue() ? ColorTexture->GetClearColor() : FLinearColor::Black;
-			ClearValues[ClearValueIndex].color.float32[0] = ClearColor.R;
-			ClearValues[ClearValueIndex].color.float32[1] = ClearColor.G;
-			ClearValues[ClearValueIndex].color.float32[2] = ClearColor.B;
-			ClearValues[ClearValueIndex].color.float32[3] = ClearColor.A;
-			++ClearValueIndex;
-			if (bPassPerformsResolve)
+			FVulkanPipelineBarrier Barrier;
+			if (VulkanTexture.IsDepthOrStencilAspect() && (SrcLayout->NumPlanes > 1) && !LayoutMgr.bWriteOnly)
+			{
+				const VkImageLayout MergedLayout = GetMergedDepthStencilLayout(SrcLayout->GetSubresLayout(0, 0, VK_IMAGE_ASPECT_DEPTH_BIT), SrcLayout->GetSubresLayout(0, 0, VK_IMAGE_ASPECT_STENCIL_BIT));
+				Barrier.AddImageLayoutTransition(VulkanTexture.Image, MergedLayout, ExpectedLayout, FVulkanPipelineBarrier::MakeSubresourceRange(VulkanTexture.GetFullAspectMask()));
+				PreviousLayout = MergedLayout;
+			}
+			else
 			{
-				++ClearValueIndex;
+				Barrier.AddImageLayoutTransition(VulkanTexture.Image, VulkanTexture.GetFullAspectMask(), *SrcLayout, ExpectedLayout);
+				check(SrcLayout->AreAllSubresourcesSameLayout());
+				PreviousLayout = SrcLayout->MainLayout;
 			}
+			Barrier.Execute(CmdBuffer);
 		}
 	}
 
-	FRHITexture* DSTexture = RPInfo.DepthStencilRenderTarget.DepthStencilTarget;
-	if (DSTexture)
-	{
-		FExclusiveDepthStencil RequestedDSAccess = RPInfo.DepthStencilRenderTarget.ExclusiveDepthStencil;
-		if (RequestedDSAccess.IsDepthRead() || RequestedDSAccess.IsStencilRead())
-		{
-			// If the depth-stencil state doesn't change between passes, the high level code won't perform any transitions.
-			// Make sure we have a barrier in case we're loading depth or stencil, to prevent rearranging passes.
-			const VkAccessFlags AccessMask = VK_ACCESS_DEPTH_STENCIL_ATTACHMENT_READ_BIT | VK_ACCESS_DEPTH_STENCIL_ATTACHMENT_WRITE_BIT;
-			const VkPipelineStageFlags StageMask = VK_PIPELINE_STAGE_EARLY_FRAGMENT_TESTS_BIT | VK_PIPELINE_STAGE_LATE_FRAGMENT_TESTS_BIT;
-			Barrier.AddMemoryBarrier(AccessMask, AccessMask, StageMask, StageMask);
-		}
-
-		if (DSTexture->HasClearValue() && bNeedsClearValues)
-		{
-			float Depth = 0;
-			uint32 Stencil = 0;
-			DSTexture->GetDepthStencilClearValue(Depth, Stencil);
-			ClearValues[ClearValueIndex].depthStencil.depth = Depth;
-			ClearValues[ClearValueIndex].depthStencil.stencil = Stencil;
-			++ClearValueIndex;
-		}
-	}
+	return PreviousLayout;
+}
 
-	FRHITexture* ShadingRateTexture = RPInfo.ShadingRateTexture;
-	if (ShadingRateTexture && ValidateShadingRateDataType())
+// todo-jn: these function are called hints, but it more of a temporary hack to access depth/stencil information for render passes
+VkImageLayout FVulkanLayoutManager::GetDepthStencilHint(const FVulkanTexture& VulkanTexture, VkImageAspectFlagBits AspectBit)
+{
+	check(AspectBit == VK_IMAGE_ASPECT_DEPTH_BIT || AspectBit == VK_IMAGE_ASPECT_STENCIL_BIT);
+	FVulkanImageLayout* Layout = Layouts.Find(VulkanTexture.Image);
+	if (Layout)
 	{
-		FVulkanTexture& Surface = *FVulkanTexture::Cast(ShadingRateTexture);
-		VkImageLayout& DSLayout = FindOrAddLayoutRW(Surface, VK_IMAGE_LAYOUT_UNDEFINED);
-		
-		VkImageLayout ExpectedLayout = VK_IMAGE_LAYOUT_UNDEFINED;
-
-#if VULKAN_SUPPORTS_FRAGMENT_SHADING_RATE
-		if (GRHIVariableRateShadingImageDataType == VRSImage_Palette)
-		{
-			ExpectedLayout = VK_IMAGE_LAYOUT_FRAGMENT_SHADING_RATE_ATTACHMENT_OPTIMAL_KHR;
-		}
-#endif
-
-#if VULKAN_SUPPORTS_FRAGMENT_DENSITY_MAP
-		if (GRHIVariableRateShadingImageDataType == VRSImage_Fractional)
-		{
-			ExpectedLayout = VK_IMAGE_LAYOUT_FRAGMENT_DENSITY_MAP_OPTIMAL_EXT;
-		}
-#endif
-
-		// transition shading rate textures to the fragment density map layout for rendering
-		if (DSLayout != ExpectedLayout)
-		{
-			Barrier.AddImageLayoutTransition(Surface.Image, DSLayout, ExpectedLayout, FVulkanPipelineBarrier::MakeSubresourceRange(VK_IMAGE_ASPECT_COLOR_BIT));
-			DSLayout = ExpectedLayout;
-		}
+		return Layout->GetSubresLayout(0, 0, AspectBit);
 	}
-
-	ensure(ClearValueIndex <= RenderPass->GetNumUsedClearValues());
-
-	Barrier.Execute(CmdBuffer->GetHandle());
-
-	CmdBuffer->BeginRenderPass(RenderPass->GetLayout(), RenderPass, Framebuffer, ClearValues);
-
+	else if (Fallback)
 	{
-		const VkExtent3D& Extents = RTLayout.GetExtent3D();
-		Context.GetPendingGfxState()->SetViewport(0, 0, 0, Extents.width, Extents.height, 1);
+		return Fallback->GetDepthStencilHint(VulkanTexture, AspectBit);
 	}
-
-	CurrentFramebuffer = Framebuffer;
-	CurrentRenderPass = RenderPass;
-}
-
-void FVulkanLayoutManager::EndRenderPass(FVulkanCmdBuffer* CmdBuffer)
-{
-	check(CurrentRenderPass);
-	CmdBuffer->EndRenderPass();
-
-	CurrentRenderPass = nullptr;
-
-	VulkanRHI::DebugHeavyWeightBarrier(CmdBuffer->GetHandle(), 1);
+	return VK_IMAGE_LAYOUT_UNDEFINED;
 }
 
 void FVulkanLayoutManager::NotifyDeletedImage(VkImage Image)
@@ -1878,33 +1923,11 @@ void FVulkanLayoutManager::NotifyDeletedImage(VkImage Image)
 	Layouts.Remove(Image);
 }
 
-void FVulkanLayoutManager::NotifyDeletedRenderTarget(FVulkanDevice& InDevice, VkImage Image)
+void FVulkanLayoutManager::TransferTo(FVulkanLayoutManager& Destination)
 {
-	Layouts.Remove(Image);
-	for (auto It = Framebuffers.CreateIterator(); It; ++It)
+	if (Layouts.Num())
 	{
-		FFramebufferList* List = It->Value;
-		for (int32 Index = List->Framebuffer.Num() - 1; Index >= 0; --Index)
-		{
-			FVulkanFramebuffer* Framebuffer = List->Framebuffer[Index];
-			if (Framebuffer->ContainsRenderTarget(Image))
-			{
-				List->Framebuffer.RemoveAtSwap(Index, 1, false);
-				Framebuffer->Destroy(InDevice);
-
-				if (Framebuffer == CurrentFramebuffer)
-				{
-					CurrentFramebuffer = nullptr;
-				}
-
-				delete Framebuffer;
-			}
-		}
-
-		if (List->Framebuffer.Num() == 0)
-		{
-			delete List;
-			It.RemoveCurrent();
-		}
+		Destination.Layouts.Append(Layouts);
+		Layouts.Empty();
 	}
 }
diff --git a/Engine/Source/Runtime/VulkanRHI/Private/VulkanBarriers.h b/Engine/Source/Runtime/VulkanRHI/Private/VulkanBarriers.h
--- a/Engine/Source/Runtime/VulkanRHI/Private/VulkanBarriers.h
+++ b/Engine/Source/Runtime/VulkanRHI/Private/VulkanBarriers.h
@@ -5,21 +5,22 @@
 #include "CoreMinimal.h"
 #include "VulkanCommon.h"
 
+class FVulkanCmdBuffer;
 
 struct FVulkanPipelineBarrier
 {
-	FVulkanPipelineBarrier() : SrcStageMask(0), DstStageMask(0), Semaphore(nullptr), bHasMemoryBarrier(false)
-	{
-		ZeroVulkanStruct(MemoryBarrier, VK_STRUCTURE_TYPE_MEMORY_BARRIER);
-	}
+	FVulkanPipelineBarrier() : Semaphore(nullptr)
+	{}
+
+	using MemoryBarrierArrayType = TArray<VkMemoryBarrier2, TInlineAllocator<1>>;
+	using ImageBarrierArrayType = TArray<VkImageMemoryBarrier2, TInlineAllocator<2>>;
+	using BufferBarrierArrayType = TArray<VkBufferMemoryBarrier2>;
 
 	ERHIPipeline SrcPipelines, DstPipelines;
-	VkPipelineStageFlags SrcStageMask, DstStageMask;
-	VkMemoryBarrier MemoryBarrier;
-	TArray<VkImageMemoryBarrier, TInlineAllocator<2>> ImageBarriers;
-	TArray<VkBufferMemoryBarrier> BufferBarriers;
+	MemoryBarrierArrayType MemoryBarriers;
+	ImageBarrierArrayType ImageBarriers;
+	BufferBarrierArrayType BufferBarriers;
 	VulkanRHI::FSemaphore* Semaphore;
-	bool bHasMemoryBarrier;
 
 	// We need to keep the texture pointers around, because we need to call OnTransitionResource on them, and we need mip and layer counts for the tracking code.
 	struct ImageBarrierExtraData
@@ -36,21 +37,24 @@ struct FVulkanPipelineBarrier
 	void AddImageLayoutTransition(VkImage Image, VkImageAspectFlags AspectMask, const struct FVulkanImageLayout& SrcLayout, const struct FVulkanImageLayout& DstLayout);
 	void AddImageAccessTransition(const FVulkanTexture& Surface, ERHIAccess SrcAccess, ERHIAccess DstAccess, const VkImageSubresourceRange& SubresourceRange, VkImageLayout& InOutLayout);
 	void Execute(VkCommandBuffer CmdBuffer);
+	void Execute(FVulkanCmdBuffer* CmdBuffer);
 
 	static VkImageSubresourceRange MakeSubresourceRange(VkImageAspectFlags AspectMask, uint32 FirstMip = 0, uint32 NumMips = VK_REMAINING_MIP_LEVELS, uint32 FirstLayer = 0, uint32 NumLayers = VK_REMAINING_ARRAY_LAYERS);
 };
 
 struct FVulkanImageLayout
 {
-	FVulkanImageLayout(VkImageLayout InitialLayout, uint32 InNumMips, uint32 InNumLayers) :
+	FVulkanImageLayout(VkImageLayout InitialLayout, uint32 InNumMips, uint32 InNumLayers, VkImageAspectFlags Aspect) :
 		NumMips(InNumMips),
 		NumLayers(InNumLayers),
+		NumPlanes((Aspect == (VK_IMAGE_ASPECT_DEPTH_BIT | VK_IMAGE_ASPECT_STENCIL_BIT)) ? 2 : 1),
 		MainLayout(InitialLayout)
 	{
 	}
 
 	uint32 NumMips;
 	uint32 NumLayers;
+	uint32 NumPlanes;
 
 	// The layout when all the subresources are in the same state.
 	VkImageLayout MainLayout;
@@ -58,12 +62,17 @@ struct FVulkanImageLayout
 	// Explicit subresource layouts. Always NumLayers*NumMips elements.
 	TArray<VkImageLayout> SubresLayouts;
 
-	bool AreAllSubresourcesSameLayout() const
+	inline bool AreAllSubresourcesSameLayout() const
 	{
 		return SubresLayouts.Num() == 0;
 	}
 
-	VkImageLayout GetSubresLayout(uint32 Layer, uint32 Mip) const
+	VkImageLayout GetSubresLayout(uint32 Layer, uint32 Mip, VkImageAspectFlagBits Aspect) const
+	{
+		return GetSubresLayout(Layer, Mip, (Aspect==VK_IMAGE_ASPECT_STENCIL_BIT) ? NumPlanes - 1 : 0);
+	}
+
+	VkImageLayout GetSubresLayout(uint32 Layer, uint32 Mip, uint32 Plane) const
 	{
 		if (SubresLayouts.Num() == 0)
 		{
@@ -75,20 +84,22 @@ struct FVulkanImageLayout
 			Layer = 0;
 		}
 
-		check(Layer < NumLayers && Mip < NumMips);
-		return SubresLayouts[Layer * NumMips + Mip];
+		check(Plane < NumPlanes && Layer < NumLayers && Mip < NumMips);
+		return SubresLayouts[(Plane * NumLayers * NumMips) + (Layer * NumMips) + Mip];
 	}
 
 	bool AreSubresourcesSameLayout(VkImageLayout Layout, const VkImageSubresourceRange& SubresourceRange) const;
 
-	uint32 GetSubresRangeLayerCount(const VkImageSubresourceRange& SubresourceRange) const
+	inline uint32 GetSubresRangeLayerCount(const VkImageSubresourceRange& SubresourceRange) const
 	{
-		return (SubresourceRange.layerCount == VK_REMAINING_ARRAY_LAYERS) ? NumLayers : SubresourceRange.layerCount;
+		check(SubresourceRange.baseArrayLayer < NumLayers);
+		return (SubresourceRange.layerCount == VK_REMAINING_ARRAY_LAYERS) ? (NumLayers - SubresourceRange.baseArrayLayer) : SubresourceRange.layerCount;
 	}
 
-	uint32 GetSubresRangeMipCount(const VkImageSubresourceRange& SubresourceRange) const
+	inline uint32 GetSubresRangeMipCount(const VkImageSubresourceRange& SubresourceRange) const
 	{
-		return (SubresourceRange.levelCount == VK_REMAINING_MIP_LEVELS) ? NumMips : SubresourceRange.levelCount;
+		check(SubresourceRange.baseMipLevel < NumMips);
+		return (SubresourceRange.levelCount == VK_REMAINING_MIP_LEVELS) ? (NumMips - SubresourceRange.baseMipLevel) : SubresourceRange.levelCount;
 	}
 
 	void CollapseSubresLayoutsIfSame();
@@ -99,128 +110,116 @@ struct FVulkanImageLayout
 class FVulkanLayoutManager
 {
 public:
-	FVulkanLayoutManager()
-		: CurrentRenderPass(nullptr)
-		, CurrentFramebuffer(nullptr)
+	FVulkanLayoutManager(bool InWriteOnly, FVulkanLayoutManager* InFallback)
+		: bWriteOnly(InWriteOnly)
+		, Fallback(InFallback)
 	{
 	}
 
-	void TempCopy(const FVulkanLayoutManager& In)
+	void NotifyDeletedImage(VkImage Image);
+
+	// Predetermined layouts for given RHIAccess
+	static VkImageLayout GetDefaultLayout(FVulkanCmdBuffer* CmdBuffer, const FVulkanTexture& VulkanTexture, ERHIAccess DesiredAccess);
+
+	// Expected layouts and Hints are workarounds until we can use 'hardcoded layouts' everywhere.
+	static VkImageLayout SetExpectedLayout(FVulkanCmdBuffer* CmdBuffer, const FVulkanTexture& VulkanTexture, ERHIAccess DesiredAccess);
+	VkImageLayout GetDepthStencilHint(const FVulkanTexture& VulkanTexture, VkImageAspectFlagBits AspectBit);
+
+	VULKANRHI_API const FVulkanImageLayout* GetFullLayout(VkImage Image) const
 	{
-		Framebuffers = In.Framebuffers;
-		RenderPasses = In.RenderPasses;
-		Layouts = In.Layouts;
+		check(!bWriteOnly);
+		const FVulkanImageLayout* Layout = Layouts.Find(Image);
+		if (!Layout && Fallback)
+		{
+			return Fallback->GetFullLayout(Image);
+		}
+		return Layout;
 	}
 
-	void Destroy(FVulkanDevice& InDevice, FVulkanLayoutManager* Immediate);
-
-	FVulkanFramebuffer* GetOrCreateFramebuffer(FVulkanDevice& InDevice, const FRHISetRenderTargetsInfo& RenderTargetsInfo, const FVulkanRenderTargetLayout& RTLayout, FVulkanRenderPass* RenderPass);
-	FVulkanRenderPass* GetOrCreateRenderPass(FVulkanDevice& InDevice, const FVulkanRenderTargetLayout& RTLayout)
+	VULKANRHI_API const FVulkanImageLayout* GetFullLayout(const FVulkanTexture& VulkanTexture, bool bAddIfNotFound = false, VkImageLayout LayoutIfNotFound = VK_IMAGE_LAYOUT_UNDEFINED)
 	{
-		uint32 RenderPassHash = RTLayout.GetRenderPassFullHash();
-		FVulkanRenderPass** FoundRenderPass = nullptr;
+		check(!bWriteOnly);
+		const FVulkanImageLayout* Layout = Layouts.Find(VulkanTexture.Image);
+		
+		if (!Layout && Fallback)
 		{
-			FScopeLock Lock(&RenderPassesCS);
-			FoundRenderPass = RenderPasses.Find(RenderPassHash);
+			Layout = Fallback->GetFullLayout(VulkanTexture, false);
 		}
-		if (FoundRenderPass)
+		
+		if (Layout)
 		{
-			return *FoundRenderPass;
+			return Layout;
 		}
-
-		FVulkanRenderPass* RenderPass = new FVulkanRenderPass(InDevice, RTLayout);
+		else if (!bAddIfNotFound)
 		{
-			FScopeLock Lock(&RenderPassesCS);
-			FoundRenderPass = RenderPasses.Find(RenderPassHash);
-			if (FoundRenderPass)
-			{
-				delete RenderPass;
-				return *FoundRenderPass;
-			}
-			RenderPasses.Add(RenderPassHash, RenderPass);
+			return nullptr;
 		}
-		return RenderPass;
-	}
-
-	void BeginRenderPass(FVulkanCommandListContext& Context, FVulkanDevice& InDevice, FVulkanCmdBuffer* CmdBuffer, const FRHIRenderPassInfo& RPInfo, const FVulkanRenderTargetLayout& RTLayout, FVulkanRenderPass* RenderPass, FVulkanFramebuffer* Framebuffer);
-	void EndRenderPass(FVulkanCmdBuffer* CmdBuffer);
-
-	FVulkanRenderPass* CurrentRenderPass;
-	FVulkanFramebuffer* CurrentFramebuffer;
-
-	FCriticalSection RenderPassesCS;
-
-	void NotifyDeletedRenderTarget(FVulkanDevice& InDevice, VkImage Image);
-	void NotifyDeletedImage(VkImage Image);
-
-	FVulkanImageLayout* GetFullLayout(VkImage Image)
-	{
-		return Layouts.Find(Image);
-	}
 
-	FVulkanImageLayout& GetFullLayoutChecked(VkImage Image)
-	{
-		return Layouts.FindChecked(Image);
+		return &Layouts.Add(VulkanTexture.Image, FVulkanImageLayout(LayoutIfNotFound, VulkanTexture.GetNumMips(), VulkanTexture.GetNumberOfArrayLevels(), VulkanTexture.GetFullAspectMask()));
 	}
 
-	FVulkanImageLayout& GetOrAddFullLayout(const FVulkanTexture& Surface, VkImageLayout LayoutIfNotFound)
+	// Not the preferred path because we can't ensure Mip and Layer counts match, but still necessary for images like the backbuffer
+	VULKANRHI_API void SetFullLayout(VkImage Image, const FVulkanImageLayout& NewLayout)
 	{
-		FVulkanImageLayout* Layout = Layouts.Find(Surface.Image);
+		FVulkanImageLayout* Layout = Layouts.Find(Image);
 		if (Layout)
 		{
-			return *Layout;
+			*Layout = NewLayout;
+		}
+		else
+		{
+			Layouts.Add(Image, NewLayout);
 		}
-
-		return Layouts.Add(Surface.Image, FVulkanImageLayout(LayoutIfNotFound, Surface.GetNumMips(), Surface.GetNumberOfArrayLevels()));
 	}
 
-	//
-	// The following functions should only be used when the calling code is sure that the image has all its sub-resources
-	// in the same state. They will assert if that's not the case.
-	//
-	VkImageLayout FindLayoutChecked(VkImage Image) const
+	VULKANRHI_API void SetFullLayout(const FVulkanTexture& VulkanTexture, const FVulkanImageLayout& NewLayout)
 	{
-		const FVulkanImageLayout& Layout = Layouts.FindChecked(Image);
-		check(Layout.AreAllSubresourcesSameLayout());
-		return Layout.MainLayout;
+		check((VulkanTexture.GetNumMips() == NewLayout.NumMips) && (VulkanTexture.GetNumberOfArrayLevels() == NewLayout.NumLayers));
+		SetFullLayout(VulkanTexture.Image, NewLayout);
 	}
 
-	VULKANRHI_API FVulkanImageLayout& FindOrAddFullLayoutRW(VkImage Image, VkImageLayout LayoutIfNotFound, uint32 NumMips, uint32 NumLayers)
+	VULKANRHI_API void SetFullLayout(const FVulkanTexture& VulkanTexture, VkImageLayout InLayout, bool bOnlyIfNotFound=false)
 	{
-		FVulkanImageLayout* Layout = Layouts.Find(Image);
+		FVulkanImageLayout* Layout = Layouts.Find(VulkanTexture.Image);
 		if (Layout)
 		{
-			check(Layout->AreAllSubresourcesSameLayout());
-			return *Layout;
+			if (!bOnlyIfNotFound)
+			{
+				Layout->Set(InLayout, FVulkanPipelineBarrier::MakeSubresourceRange(VulkanTexture.GetFullAspectMask()));
+			}
+		}
+		else
+		{
+			Layouts.Add(VulkanTexture.Image, FVulkanImageLayout(InLayout, VulkanTexture.GetNumMips(), VulkanTexture.GetNumberOfArrayLevels(), VulkanTexture.GetFullAspectMask()));
 		}
-		return Layouts.Add(Image, FVulkanImageLayout(LayoutIfNotFound, NumMips, NumLayers));
-	}
-
-	VULKANRHI_API VkImageLayout& FindOrAddLayoutRW(VkImage Image, VkImageLayout LayoutIfNotFound, uint32 NumMips, uint32 NumLayers)
-	{
-		return FindOrAddFullLayoutRW(Image, LayoutIfNotFound, NumMips, NumLayers).MainLayout;
 	}
 
-	VULKANRHI_API VkImageLayout& FindOrAddLayoutRW(const FVulkanTexture& Surface, VkImageLayout LayoutIfNotFound)
+	VULKANRHI_API void SetLayout(const FVulkanTexture& VulkanTexture, const VkImageSubresourceRange& InSubresourceRange, VkImageLayout InLayout)
 	{
-		return FindOrAddLayoutRW(Surface.Image, LayoutIfNotFound, Surface.GetNumMips(), Surface.GetNumberOfArrayLevels());
+		FVulkanImageLayout* Layout = Layouts.Find(VulkanTexture.Image);
+		if (Layout)
+		{
+			Layout->Set(InLayout, InSubresourceRange);
+		}
+		else
+		{
+			FVulkanImageLayout NewLayout(VK_IMAGE_LAYOUT_UNDEFINED, VulkanTexture.GetNumMips(), VulkanTexture.GetNumberOfArrayLevels(), VulkanTexture.GetFullAspectMask());
+			NewLayout.Set(InLayout, InSubresourceRange);
+			Layouts.Add(VulkanTexture.Image, NewLayout);
+		}
 	}
 
-	VkImageLayout FindOrAddLayout(const FVulkanTexture& Surface, VkImageLayout LayoutIfNotFound)
-	{
-		return FindOrAddLayoutRW(Surface, LayoutIfNotFound);
-	}
+	// Transfers our layouts into the destination
+	void TransferTo(FVulkanLayoutManager& Destination);
 
 private:
-	TMap<uint32, FVulkanRenderPass*> RenderPasses;
-
-	struct FFramebufferList
-	{
-		TArray<FVulkanFramebuffer*> Framebuffer;
-	};
-	TMap<uint32, FFramebufferList*> Framebuffers;
 
 	TMap<VkImage, FVulkanImageLayout> Layouts;
 
-	void ValidateRenderPassColorEntry(const FRHIRenderPassInfo::FColorEntry& ColorEntry, bool bResolveTarget, FVulkanPipelineBarrier& Barrier);
+	// If we're WriteOnly, we should never read layout from this instance.  This is important for parallel rendering.
+	// When WriteOnly, this instance of the layout manager should only collect layouts to later feed them to the another central mgr.
+	const bool bWriteOnly;
+
+	// If parallel command list creation is NOT supported, then the queue's layout mgr can be used as a fallback to fetch previous layouts.
+	FVulkanLayoutManager* Fallback;
 };
diff --git a/Engine/Source/Runtime/VulkanRHI/Private/VulkanCommandBuffer.cpp b/Engine/Source/Runtime/VulkanRHI/Private/VulkanCommandBuffer.cpp
--- a/Engine/Source/Runtime/VulkanRHI/Private/VulkanCommandBuffer.cpp
+++ b/Engine/Source/Runtime/VulkanRHI/Private/VulkanCommandBuffer.cpp
@@ -66,8 +66,10 @@ FVulkanCmdBuffer::FVulkanCmdBuffer(FVulkanDevice* InDevice, FVulkanCommandBuffer
 	, CommandBufferPool(InCommandBufferPool)
 	, Timing(nullptr)
 	, LastValidTiming(0)
+	, LayoutManager(InDevice->SupportsParallelRendering(), 
+		InDevice->SupportsParallelRendering() ? nullptr :
+		&InCommandBufferPool->GetMgr().GetCommandListContext()->GetQueue()->GetLayoutManager())
 {
-
 	{
 		FScopeLock ScopeLock(CommandBufferPool->GetCS());
 		AllocMemory();
@@ -202,7 +204,6 @@ void FVulkanCmdBuffer::BeginRenderPass(const FVulkanRenderTargetLayout& Layout,
 	}
 #endif
 
-#if VULKAN_SUPPORTS_RENDERPASS2
 	if (Device->GetOptionalExtensions().HasKHRRenderPass2)
 	{
 		VkSubpassBeginInfo SubpassInfo;
@@ -211,7 +212,6 @@ void FVulkanCmdBuffer::BeginRenderPass(const FVulkanRenderTargetLayout& Layout,
 		VulkanRHI::vkCmdBeginRenderPass2KHR(CommandBufferHandle, &Info, &SubpassInfo);
 	}
 	else
-#endif
 	{
 		VulkanRHI::vkCmdBeginRenderPass(CommandBufferHandle, &Info, VK_SUBPASS_CONTENTS_INLINE);
 	}
@@ -288,12 +288,18 @@ inline void FVulkanCmdBuffer::InitializeTimings(FVulkanCommandListContext* InCon
 	}
 }
 
-void FVulkanCmdBuffer::AddWaitSemaphore(VkPipelineStageFlags InWaitFlags, VulkanRHI::FSemaphore* InWaitSemaphore)
+void FVulkanCmdBuffer::AddWaitSemaphore(VkPipelineStageFlags InWaitFlags, TArrayView<VulkanRHI::FSemaphore*> InWaitSemaphores)
 {
-	WaitFlags.Add(InWaitFlags);
-	InWaitSemaphore->AddRef();
-	check(!WaitSemaphores.Contains(InWaitSemaphore));
-	WaitSemaphores.Add(InWaitSemaphore);
+	WaitFlags.Reserve(WaitFlags.Num() + InWaitSemaphores.Num());
+
+	for (VulkanRHI::FSemaphore* Sema : InWaitSemaphores)
+	{
+		WaitFlags.Add(InWaitFlags);
+		Sema->AddRef();
+		check(!WaitSemaphores.Contains(Sema));
+	}
+
+	WaitSemaphores.Append(InWaitSemaphores);
 }
 
 void FVulkanCmdBuffer::Begin()
diff --git a/Engine/Source/Runtime/VulkanRHI/Private/VulkanCommandBuffer.h b/Engine/Source/Runtime/VulkanRHI/Private/VulkanCommandBuffer.h
--- a/Engine/Source/Runtime/VulkanRHI/Private/VulkanCommandBuffer.h
+++ b/Engine/Source/Runtime/VulkanRHI/Private/VulkanCommandBuffer.h
@@ -8,6 +8,7 @@
 
 #include "CoreMinimal.h"
 #include "VulkanConfiguration.h"
+#include "VulkanBarriers.h"
 
 extern int32 GVulkanProfileCmdBuffers;
 extern int32 GVulkanUseCmdBufferTimingForGPUTime;
@@ -42,6 +43,11 @@ public:
 		return CommandBufferPool;
 	}
 
+	FVulkanDevice* GetDevice()
+	{
+		return Device;
+	}
+
 	bool IsUniformBufferBarrierAdded() const
 	{
 		return bIsUniformBufferBarrierAdded;
@@ -138,12 +144,22 @@ public:
 		return SubmittedFenceCounter;
 	}
 
+	inline FVulkanLayoutManager& GetLayoutManager()
+	{
+		return LayoutManager;
+	}
+
 	inline bool HasValidTiming() const
 	{
 		return (Timing != nullptr) && (FMath::Abs((int64)FenceSignaledCounter - (int64)LastValidTiming) < 3);
 	}
 
-	void VULKANRHI_API AddWaitSemaphore(VkPipelineStageFlags InWaitFlags, VulkanRHI::FSemaphore* InWaitSemaphore);
+	void VULKANRHI_API AddWaitSemaphore(VkPipelineStageFlags InWaitFlags, VulkanRHI::FSemaphore* InWaitSemaphore)
+	{
+		AddWaitSemaphore(InWaitFlags, MakeArrayView<VulkanRHI::FSemaphore*>(&InWaitSemaphore, 1));
+	}
+
+	void VULKANRHI_API AddWaitSemaphore(VkPipelineStageFlags InWaitFlags, TArrayView<VulkanRHI::FSemaphore*> InWaitSemaphores);
 
 	void Begin();
 	void End();
@@ -233,6 +249,8 @@ private:
 	void AllocMemory();
 	void FreeMemory();
 
+	FVulkanLayoutManager LayoutManager;
+
 public:
 	//#todo-rco: Hide this
 	TMap<uint32, class FVulkanTypedDescriptorPoolSet*> TypedDescriptorPoolSets;
@@ -360,16 +378,28 @@ public:
 
 	void FreeUnusedCmdBuffers();
 
-	FVulkanCommandListContext* GetCommandListContext()
+	inline FVulkanCommandListContext* GetCommandListContext()
 	{
 		return Context;
 	}
 
-	FVulkanQueue* GetQueue()
+	inline FVulkanQueue* GetQueue()
 	{
 		return Queue;
 	}
 
+	inline void NotifyDeletedImage(VkImage Image)
+	{
+		if (UploadCmdBuffer)
+		{
+			UploadCmdBuffer->GetLayoutManager().NotifyDeletedImage(Image);
+		}
+		if (ActiveCmdBuffer)
+		{
+			ActiveCmdBuffer->GetLayoutManager().NotifyDeletedImage(Image);
+		}
+	}
+
 private:
 	struct FQueryPoolReset
 	{
diff --git a/Engine/Source/Runtime/VulkanRHI/Private/VulkanCommandWrappers.h b/Engine/Source/Runtime/VulkanRHI/Private/VulkanCommandWrappers.h
--- a/Engine/Source/Runtime/VulkanRHI/Private/VulkanCommandWrappers.h
+++ b/Engine/Source/Runtime/VulkanRHI/Private/VulkanCommandWrappers.h
@@ -101,9 +101,7 @@ struct FWrapLayer
 	static void CreateFramebuffer(VkResult Result, VkDevice Device, const VkFramebufferCreateInfo* CreateInfo, VkFramebuffer* Framebuffer) VULKAN_LAYER_BODY
 	static void DestroyFramebuffer(VkResult Result, VkDevice Device, VkFramebuffer Framebuffer) VULKAN_LAYER_BODY
 	static void CreateRenderPass(VkResult Result, VkDevice Device, const VkRenderPassCreateInfo* CreateInfo, VkRenderPass* RenderPass) VULKAN_LAYER_BODY
-#if VULKAN_SUPPORTS_RENDERPASS2
 	static void CreateRenderPass2KHR(VkResult Result, VkDevice Device, const VkRenderPassCreateInfo2* CreateInfo, VkRenderPass* RenderPass) VULKAN_LAYER_BODY
-#endif
 	static void DestroyRenderPass(VkResult Result, VkDevice Device, VkRenderPass RenderPass) VULKAN_LAYER_BODY
 	static void GetRenderAreaGranularity(VkResult Result, VkDevice Device, VkRenderPass RenderPass, VkExtent2D* pGranularity) VULKAN_LAYER_BODY
 	static void CreateCommandPool(VkResult Result, VkDevice Device, const VkCommandPoolCreateInfo* CreateInfo, VkCommandPool* CommandPool) VULKAN_LAYER_BODY
@@ -156,9 +154,7 @@ struct FWrapLayer
 	static void CmdCopyQueryPoolResults(VkResult Result, VkCommandBuffer CommandBuffer, VkQueryPool QueryPool, uint32 FirstQuery, uint32 QueryCount, VkBuffer DstBuffer, VkDeviceSize DstOffset, VkDeviceSize Stride, VkQueryResultFlags Flags) VULKAN_LAYER_BODY
 	static void CmdPushConstants(VkResult Result, VkCommandBuffer CommandBuffer, VkPipelineLayout Layout, VkShaderStageFlags StageFlags, uint32_t Offset, uint32_t Size, const void* pValues) VULKAN_LAYER_BODY
 	static void CmdBeginRenderPass(VkResult Result, VkCommandBuffer CommandBuffer, const VkRenderPassBeginInfo* RenderPassBegin, VkSubpassContents Contents) VULKAN_LAYER_BODY
-#if VULKAN_SUPPORTS_RENDERPASS2
 	static void CmdBeginRenderPass2KHR(VkResult Result, VkCommandBuffer CommandBuffer, const VkRenderPassBeginInfo* RenderPassBegin, const VkSubpassBeginInfo* pSubpassBeginInfo) VULKAN_LAYER_BODY
-#endif
 	static void CmdNextSubpass(VkResult Result, VkCommandBuffer CommandBuffer, VkSubpassContents Contents) VULKAN_LAYER_BODY
 	static void CmdEndRenderPass(VkResult Result, VkCommandBuffer CommandBuffer) VULKAN_LAYER_BODY
 	static void CmdExecuteCommands(VkResult Result, VkCommandBuffer CommandBuffer, uint32 CommandBufferCount, const VkCommandBuffer* pCommandBuffers) VULKAN_LAYER_BODY
@@ -207,6 +203,11 @@ struct FWrapLayer
 	static void ResetQueryPoolEXT(VkResult Result, VkDevice Device, VkQueryPool QueryPool, uint32_t FirstQuery, uint32_t QueryCount) VULKAN_LAYER_BODY
 	static void GetPhysicalDeviceCalibrateableTimeDomainsEXT(VkResult Result, VkPhysicalDevice PhysicalDevice, uint32_t* TimeDomainCount, VkTimeDomainEXT* TimeDomains) VULKAN_LAYER_BODY
 	static void GetCalibratedTimestampsEXT(VkResult Result, VkDevice Device, uint32_t TimestampCount, const VkCalibratedTimestampInfoEXT* TimestampInfos, uint64_t* Timestamps, uint64_t* MaxDeviation) VULKAN_LAYER_BODY
+	static void CmdPipelineBarrier2KHR(VkResult Result, VkCommandBuffer CommandBuffer, const VkDependencyInfo* DependencyInfo) VULKAN_LAYER_BODY
+	static void CmdResetEvent2KHR(VkResult Result, VkCommandBuffer CommandBuffer, VkEvent Event, VkPipelineStageFlags2 StageMask) VULKAN_LAYER_BODY
+	static void CmdSetEvent2KHR(VkResult Result, VkCommandBuffer CommandBuffer, VkEvent Event, const VkDependencyInfo* DependencyInfo) VULKAN_LAYER_BODY
+	static void CmdWaitEvents2KHR(VkResult Result, VkCommandBuffer CommandBuffer, uint32_t EventCount, const VkEvent* Events, const VkDependencyInfo* DependencyInfos) VULKAN_LAYER_BODY
+	static void QueueSubmit2KHR(VkResult Result, VkQueue Queue, uint32_t SubmitCount, const VkSubmitInfo2* Submits, VkFence Fence) VULKAN_LAYER_BODY
 };
 
 #undef VULKAN_LAYER_BODY
@@ -879,7 +880,6 @@ namespace VulkanRHI
 		return Result;
 	}
 
-#if VULKAN_SUPPORTS_RENDERPASS2
 	static FORCEINLINE_DEBUGGABLE VkResult vkCreateRenderPass2KHR(VkDevice Device, const VkRenderPassCreateInfo2* CreateInfo, const VkAllocationCallbacks* Allocator, VkRenderPass* RenderPass)
 	{
 		FWrapLayer::CreateRenderPass2KHR(VK_RESULT_MAX_ENUM, Device, CreateInfo, RenderPass);
@@ -887,7 +887,6 @@ namespace VulkanRHI
 		FWrapLayer::CreateRenderPass2KHR(Result, Device, CreateInfo, RenderPass);
 		return Result;
 	}
-#endif
 
 	static FORCEINLINE_DEBUGGABLE void  vkDestroyRenderPass(VkDevice Device, VkRenderPass RenderPass, const VkAllocationCallbacks* Allocator)
 	{
@@ -1268,14 +1267,12 @@ namespace VulkanRHI
 		FWrapLayer::CmdBeginRenderPass(VK_SUCCESS, CommandBuffer, RenderPassBegin, Contents);
 	}
 
-#if VULKAN_SUPPORTS_RENDERPASS2
 	static FORCEINLINE_DEBUGGABLE void vkCmdBeginRenderPass2KHR(VkCommandBuffer CommandBuffer, const VkRenderPassBeginInfo* RenderPassBegin, const VkSubpassBeginInfo* SubpassBeginInfo)
 	{
 		FWrapLayer::CmdBeginRenderPass2KHR(VK_RESULT_MAX_ENUM, CommandBuffer, RenderPassBegin, SubpassBeginInfo);
 		VULKANAPINAMESPACE::vkCmdBeginRenderPass2KHR(CommandBuffer, RenderPassBegin, SubpassBeginInfo);
 		FWrapLayer::CmdBeginRenderPass2KHR(VK_SUCCESS, CommandBuffer, RenderPassBegin, SubpassBeginInfo);
 	}
-#endif
 
 	static FORCEINLINE_DEBUGGABLE void  vkCmdNextSubpass(VkCommandBuffer CommandBuffer, VkSubpassContents Contents)
 	{
@@ -1518,6 +1515,41 @@ namespace VulkanRHI
 		return Result;
 	}
 
+	static FORCEINLINE_DEBUGGABLE void vkCmdPipelineBarrier2KHR(VkCommandBuffer CommandBuffer, const VkDependencyInfo* DependencyInfo)
+	{
+		FWrapLayer::CmdPipelineBarrier2KHR(VK_RESULT_MAX_ENUM, CommandBuffer, DependencyInfo);
+		VULKANAPINAMESPACE::vkCmdPipelineBarrier2KHR(CommandBuffer, DependencyInfo);
+		FWrapLayer::CmdPipelineBarrier2KHR(VK_SUCCESS, CommandBuffer, DependencyInfo);
+	}
+
+	static FORCEINLINE_DEBUGGABLE void vkCmdResetEvent2KHR(VkCommandBuffer CommandBuffer, VkEvent Event, VkPipelineStageFlags2 StageMask)
+	{
+		FWrapLayer::CmdResetEvent2KHR(VK_RESULT_MAX_ENUM, CommandBuffer, Event, StageMask);
+		VULKANAPINAMESPACE::vkCmdResetEvent2KHR(CommandBuffer, Event, StageMask);
+		FWrapLayer::CmdResetEvent2KHR(VK_SUCCESS, CommandBuffer, Event, StageMask);
+	}
+	
+	static FORCEINLINE_DEBUGGABLE void vkCmdSetEvent2KHR(VkCommandBuffer CommandBuffer, VkEvent Event, const VkDependencyInfo* DependencyInfo)
+	{
+		FWrapLayer::CmdSetEvent2KHR(VK_RESULT_MAX_ENUM, CommandBuffer, Event, DependencyInfo);
+		VULKANAPINAMESPACE::vkCmdSetEvent2KHR(CommandBuffer, Event, DependencyInfo);
+		FWrapLayer::CmdSetEvent2KHR(VK_SUCCESS, CommandBuffer, Event, DependencyInfo);
+	}
+
+	static FORCEINLINE_DEBUGGABLE void vkCmdWaitEvents2KHR(VkCommandBuffer CommandBuffer, uint32_t EventCount, const VkEvent* Events, const VkDependencyInfo* DependencyInfos)
+	{
+		FWrapLayer::CmdWaitEvents2KHR(VK_RESULT_MAX_ENUM, CommandBuffer, EventCount, Events, DependencyInfos);
+		VULKANAPINAMESPACE::vkCmdWaitEvents2KHR(CommandBuffer, EventCount, Events, DependencyInfos);
+		FWrapLayer::CmdWaitEvents2KHR(VK_SUCCESS, CommandBuffer, EventCount, Events, DependencyInfos);
+	}
+
+	static FORCEINLINE_DEBUGGABLE void vkQueueSubmit2KHR(VkQueue Queue, uint32_t SubmitCount, const VkSubmitInfo2* Submits, VkFence Fence)
+	{
+		FWrapLayer::QueueSubmit2KHR(VK_RESULT_MAX_ENUM, Queue, SubmitCount, Submits, Fence);
+		VULKANAPINAMESPACE::vkQueueSubmit2KHR(Queue, SubmitCount, Submits, Fence);
+		FWrapLayer::QueueSubmit2KHR(VK_SUCCESS, Queue, SubmitCount, Submits, Fence);
+	}
+
 #if VULKAN_ENABLE_IMAGE_TRACKING_LAYER
 	void BindDebugLabelName(VkImage Image, const TCHAR* Name);
 #endif
diff --git a/Engine/Source/Runtime/VulkanRHI/Private/VulkanCommands.cpp b/Engine/Source/Runtime/VulkanRHI/Private/VulkanCommands.cpp
--- a/Engine/Source/Runtime/VulkanRHI/Private/VulkanCommands.cpp
+++ b/Engine/Source/Runtime/VulkanRHI/Private/VulkanCommands.cpp
@@ -191,11 +191,11 @@ void FVulkanCommandListContext::RHISetUAVParameter(FRHIComputeShader* ComputeSha
 
 void FVulkanCommandListContext::RHISetShaderTexture(FRHIGraphicsShader* ShaderRHI, uint32 TextureIndex, FRHITexture* NewTextureRHI)
 {
-	FVulkanTexture* Texture = FVulkanTexture::Cast(NewTextureRHI);
-	VkImageLayout Layout = LayoutManager.FindLayoutChecked(Texture->Image);
+	FVulkanTexture* VulkanTexture = FVulkanTexture::Cast(NewTextureRHI);
+	const VkImageLayout ExpectedLayout = FVulkanLayoutManager::GetDefaultLayout(GetCommandBufferManager()->GetActiveCmdBuffer(), *VulkanTexture, ERHIAccess::SRVGraphics);
 
 	ShaderStage::EStage Stage = GetAndVerifyShaderStage(ShaderRHI, PendingGfxState);
-	PendingGfxState->SetTextureForStage(Stage, TextureIndex, Texture, Layout);
+	PendingGfxState->SetTextureForStage(Stage, TextureIndex, VulkanTexture, ExpectedLayout);
 	NewTextureRHI->SetLastRenderTime((float)FPlatformTime::Seconds());
 }
 
@@ -205,8 +205,8 @@ void FVulkanCommandListContext::RHISetShaderTexture(FRHIComputeShader* ComputeSh
 	check(PendingComputeState->GetCurrentShader() == ComputeShader);
 
 	FVulkanTexture* VulkanTexture = FVulkanTexture::Cast(NewTextureRHI);
-	VkImageLayout Layout = LayoutManager.FindLayoutChecked(VulkanTexture->Image);
-	PendingComputeState->SetTextureForStage(TextureIndex, VulkanTexture, Layout);
+	const VkImageLayout ExpectedLayout = FVulkanLayoutManager::GetDefaultLayout(GetCommandBufferManager()->GetActiveCmdBuffer(), *VulkanTexture, ERHIAccess::SRVCompute);
+	PendingComputeState->SetTextureForStage(TextureIndex, VulkanTexture, ExpectedLayout);
 	NewTextureRHI->SetLastRenderTime((float)FPlatformTime::Seconds());
 }
 
@@ -319,8 +319,8 @@ inline void SetShaderUniformBufferResources(FVulkanCommandListContext* Context,
 				}
 #endif
 
-				const VkImageLayout Layout = Context->GetLayoutManager().FindLayoutChecked(VulkanTexture->Image);
-				State->SetTextureForUBResource(GlobalRemappingInfo[ResourceInfo.GlobalIndex].NewDescriptorSet, GlobalRemappingInfo[ResourceInfo.GlobalIndex].NewBindingIndex, VulkanTexture, Layout);
+				const VkImageLayout ExpectedLayout = FVulkanLayoutManager::GetDefaultLayout(Context->GetCommandBufferManager()->GetActiveCmdBuffer(), *VulkanTexture, ERHIAccess::SRVMask);
+				State->SetTextureForUBResource(GlobalRemappingInfo[ResourceInfo.GlobalIndex].NewDescriptorSet, GlobalRemappingInfo[ResourceInfo.GlobalIndex].NewBindingIndex, VulkanTexture, ExpectedLayout);
 				TexRef->SetLastRenderTime(CurrentTime);
 			}
 			else
@@ -646,16 +646,16 @@ void FVulkanCommandListContext::RHIClearMRT(bool bClearColor, int32 NumClearColo
 	FVulkanCmdBuffer* CmdBuffer = CommandBufferManager->GetActiveCmdBuffer();
 	//FRCLog::Printf(TEXT("RHIClearMRT"));
 
-	const uint32 NumColorAttachments = LayoutManager.CurrentFramebuffer->GetNumColorAttachments();
+	const uint32 NumColorAttachments = CurrentFramebuffer->GetNumColorAttachments();
 	check(!bClearColor || (uint32)NumClearColors <= NumColorAttachments);
 	InternalClearMRT(CmdBuffer, bClearColor, bClearColor ? NumClearColors : 0, ClearColorArray, bClearDepth, Depth, bClearStencil, Stencil);
 }
 
 void FVulkanCommandListContext::InternalClearMRT(FVulkanCmdBuffer* CmdBuffer, bool bClearColor, int32 NumClearColors, const FLinearColor* ClearColorArray, bool bClearDepth, float Depth, bool bClearStencil, uint32 Stencil)
 {
-	if (LayoutManager.CurrentRenderPass)
+	if (CurrentRenderPass)
 	{
-		const VkExtent2D& Extents = LayoutManager.CurrentRenderPass->GetLayout().GetExtent2D();
+		const VkExtent2D& Extents = CurrentRenderPass->GetLayout().GetExtent2D();
 		VkClearRect Rect;
 		FMemory::Memzero(Rect);
 		Rect.rect.offset.x = 0;
@@ -786,12 +786,6 @@ void FVulkanCommandListContext::RHISubmitCommandsHint()
 	CommandBufferManager->RefreshFenceStatus();
 }
 
-void FVulkanCommandListContext::PrepareParallelFromBase(const FVulkanCommandListContext& BaseContext)
-{
-	//#todo-rco: Temp
-	LayoutManager.TempCopy(BaseContext.LayoutManager);
-}
-
 void FVulkanCommandListContext::RHICopyToStagingBuffer(FRHIBuffer* SourceBufferRHI, FRHIStagingBuffer* StagingBufferRHI, uint32 Offset, uint32 NumBytes)
 {
 	FVulkanCmdBuffer* CmdBuffer = CommandBufferManager->GetActiveCmdBuffer();
diff --git a/Engine/Source/Runtime/VulkanRHI/Private/VulkanContext.h b/Engine/Source/Runtime/VulkanRHI/Private/VulkanContext.h
--- a/Engine/Source/Runtime/VulkanRHI/Private/VulkanContext.h
+++ b/Engine/Source/Runtime/VulkanRHI/Private/VulkanContext.h
@@ -8,7 +8,7 @@
 
 #include "VulkanResources.h"
 #include "VulkanGPUProfiler.h"
-#include "VulkanBarriers.h"
+#include "VulkanRenderpass.h"
 
 class FVulkanDevice;
 class FVulkanCommandBufferManager;
@@ -130,17 +130,26 @@ public:
 
 	inline void NotifyDeletedRenderTarget(VkImage Image)
 	{
-		LayoutManager.NotifyDeletedRenderTarget(*Device, Image);
+		if (CurrentFramebuffer && CurrentFramebuffer->ContainsRenderTarget(Image))
+		{
+			CurrentFramebuffer = nullptr;
+		}
 	}
 
 	inline void NotifyDeletedImage(VkImage Image)
 	{
-		LayoutManager.NotifyDeletedImage(Image);
+		CommandBufferManager->NotifyDeletedImage(Image);
+		Queue->NotifyDeletedImage(Image);
 	}
 
 	inline FVulkanRenderPass* GetCurrentRenderPass()
 	{
-		return LayoutManager.CurrentRenderPass;
+		return CurrentRenderPass;
+	}
+
+	inline FVulkanFramebuffer* GetCurrentFramebuffer()
+	{
+		return CurrentFramebuffer;
 	}
 
 	inline uint64 GetFrameCounter() const
@@ -175,8 +184,6 @@ public:
 
 	void EndRenderQueryInternal(FVulkanCmdBuffer* CmdBuffer, FVulkanRenderQuery* Query);
 
-	void PrepareParallelFromBase(const FVulkanCommandListContext& BaseContext);
-
 	void ReleasePendingState();
 
 protected:
@@ -199,7 +206,8 @@ protected:
 
 	FVulkanCommandBufferManager* CommandBufferManager;
 
-	static VULKANRHI_API FVulkanLayoutManager LayoutManager;
+	FVulkanRenderPass* CurrentRenderPass = nullptr;
+	FVulkanFramebuffer* CurrentFramebuffer = nullptr;
 
 	FVulkanOcclusionQueryPool* CurrentOcclusionQueryPool = nullptr;
 
@@ -220,11 +228,6 @@ public:
 	VkFormat GetSwapchainImageFormat() const;
 	FVulkanSwapChain* GetSwapChain() const;
 
-	inline FVulkanLayoutManager& GetLayoutManager()
-	{
-		return LayoutManager;
-	}
-
 	FVulkanRenderPass* PrepareRenderPassForPSOCreation(const FGraphicsPipelineStateInitializer& Initializer);
 	FVulkanRenderPass* PrepareRenderPassForPSOCreation(const FVulkanRenderTargetLayout& Initializer);
 
diff --git a/Engine/Source/Runtime/VulkanRHI/Private/VulkanDebug.cpp b/Engine/Source/Runtime/VulkanRHI/Private/VulkanDebug.cpp
--- a/Engine/Source/Runtime/VulkanRHI/Private/VulkanDebug.cpp
+++ b/Engine/Source/Runtime/VulkanRHI/Private/VulkanDebug.cpp
@@ -915,64 +915,17 @@ namespace VulkanRHI
 
 	static FString GetVkImageLayoutString(VkImageLayout Layout)
 	{
-		switch (Layout)
-		{
-			// + 16 to skip "VK_IMAGE_LAYOUT"
-#define VKSWITCHCASE(x)	case x: return &TEXT(#x)[16];
-			VKSWITCHCASE(VK_IMAGE_LAYOUT_UNDEFINED)
-			VKSWITCHCASE(VK_IMAGE_LAYOUT_GENERAL)
-			VKSWITCHCASE(VK_IMAGE_LAYOUT_COLOR_ATTACHMENT_OPTIMAL)
-			VKSWITCHCASE(VK_IMAGE_LAYOUT_DEPTH_STENCIL_ATTACHMENT_OPTIMAL)
-			VKSWITCHCASE(VK_IMAGE_LAYOUT_DEPTH_STENCIL_READ_ONLY_OPTIMAL)
-			VKSWITCHCASE(VK_IMAGE_LAYOUT_SHADER_READ_ONLY_OPTIMAL)
-			VKSWITCHCASE(VK_IMAGE_LAYOUT_TRANSFER_SRC_OPTIMAL)
-			VKSWITCHCASE(VK_IMAGE_LAYOUT_TRANSFER_DST_OPTIMAL)
-			VKSWITCHCASE(VK_IMAGE_LAYOUT_PREINITIALIZED)
-			VKSWITCHCASE(VK_IMAGE_LAYOUT_PRESENT_SRC_KHR)
-#undef VKSWITCHCASE
-		default:
-			break;
-		}
-
-		return FString::Printf(TEXT("Unknown VkImageLayout %d"), (int32)Layout);
+		return FString(VK_TYPE_TO_STRING(VkImageLayout, Layout)).RightChop(16);
 	}
 
 	static FString GetVkImageViewTypeString(VkImageViewType Type)
 	{
-		switch (Type)
-		{
-			// + 19 to skip "VK_IMAGE_VIEW_TYPE_"
-#define VKSWITCHCASE(x)	case x: return &TEXT(#x)[19];
-		VKSWITCHCASE(VK_IMAGE_VIEW_TYPE_1D)
-		VKSWITCHCASE(VK_IMAGE_VIEW_TYPE_2D)
-		VKSWITCHCASE(VK_IMAGE_VIEW_TYPE_3D)
-		VKSWITCHCASE(VK_IMAGE_VIEW_TYPE_CUBE)
-		VKSWITCHCASE(VK_IMAGE_VIEW_TYPE_1D_ARRAY)
-		VKSWITCHCASE(VK_IMAGE_VIEW_TYPE_2D_ARRAY)
-		VKSWITCHCASE(VK_IMAGE_VIEW_TYPE_CUBE_ARRAY)
-#undef VKSWITCHCASE
-		default:
-			break;
-		}
-
-		return FString::Printf(TEXT("Unknown VkImageViewType %d"), (int32)Type);
+		return FString(VK_TYPE_TO_STRING(VkImageViewType, Type)).RightChop(19);
 	}
 
 	static FString GetVkImageTypeString(VkImageType Type)
 	{
-		switch (Type)
-		{
-			// + 14 to skip "VK_IMAGE_TYPE_1D"
-#define VKSWITCHCASE(x)	case x: return &TEXT(#x)[14];
-		VKSWITCHCASE(VK_IMAGE_TYPE_1D)
-		VKSWITCHCASE(VK_IMAGE_TYPE_2D)
-		VKSWITCHCASE(VK_IMAGE_TYPE_3D)
-#undef VKSWITCHCASE
-		default:
-			break;
-		}
-
-		return FString::Printf(TEXT("Unknown VkImageType %d"), (int32)Type);
+		return FString(VK_TYPE_TO_STRING(VkImageType, Type)).RightChop(14);
 	}
 
 	static FString GetVkDescriptorTypeString(VkDescriptorType Type)
@@ -2066,9 +2019,12 @@ void FWrapLayer::CreateFramebuffer(VkResult Result, VkDevice Device, const VkFra
 		{
 			FFBInfo Info;
 			Info.Info = *CreateInfo;
-			Info.Attachments.AddUninitialized(CreateInfo->attachmentCount);
-			FMemory::Memcpy(&Info.Attachments[0], CreateInfo->pAttachments, CreateInfo->attachmentCount * sizeof(VkImageView));
-			Info.Info.pAttachments = &Info.Attachments[0];
+			if (CreateInfo->attachmentCount)
+			{
+				Info.Attachments.AddUninitialized(CreateInfo->attachmentCount);
+				FMemory::Memcpy(&Info.Attachments[0], CreateInfo->pAttachments, CreateInfo->attachmentCount * sizeof(VkImageView));
+				Info.Info.pAttachments = &Info.Attachments[0];
+			}
 			GFramebufferInfo.Add(*Framebuffer, Info);
 		}
 #endif
@@ -2240,25 +2196,34 @@ void FWrapLayer::CreateRenderPass2KHR(VkResult Result, VkDevice Device, const Vk
 	else
 	{
 #if VULKAN_ENABLE_DUMP_LAYER
-		/*
-		* @todo
 		PrintResultAndNamedHandle(Result, TEXT("RenderPass"), *RenderPass);
 		if (Result == VK_SUCCESS)
 		{
 			FRenderPassInfo Info;
-			Info.Info = *CreateInfo;
+			Info.Info.sType = CreateInfo->sType;
+			Info.Info.attachmentCount = CreateInfo->attachmentCount;
+			Info.Info.dependencyCount = CreateInfo->dependencyCount;
+			Info.Info.subpassCount = CreateInfo->subpassCount;
+			Info.Info.flags = CreateInfo->flags;
 			Info.Info.pAttachments = nullptr;
 			Info.Info.pSubpasses = nullptr;
 			Info.Info.pDependencies = nullptr;
 			Info.Descriptions.AddUninitialized(CreateInfo->attachmentCount);
-			if (CreateInfo->attachmentCount)
+			for (uint32 Index = 0; Index < CreateInfo->attachmentCount; ++Index)
 			{
-				FMemory::Memcpy(&Info.Descriptions[0], CreateInfo->pAttachments, CreateInfo->attachmentCount * sizeof(VkAttachmentDescription));
+				Info.Descriptions[Index].flags = CreateInfo->pAttachments[Index].flags;
+				Info.Descriptions[Index].format = CreateInfo->pAttachments[Index].format;
+				Info.Descriptions[Index].samples = CreateInfo->pAttachments[Index].samples;
+				Info.Descriptions[Index].loadOp = CreateInfo->pAttachments[Index].loadOp;
+				Info.Descriptions[Index].storeOp = CreateInfo->pAttachments[Index].storeOp;
+				Info.Descriptions[Index].stencilLoadOp = CreateInfo->pAttachments[Index].stencilLoadOp;
+				Info.Descriptions[Index].stencilStoreOp = CreateInfo->pAttachments[Index].stencilStoreOp;
+				Info.Descriptions[Index].initialLayout = CreateInfo->pAttachments[Index].initialLayout;
+				Info.Descriptions[Index].finalLayout = CreateInfo->pAttachments[Index].finalLayout;
 			}
 			GRenderPassInfo.Add(*RenderPass, Info);
 		}
 		FlushDebugWrapperLog();
-		*/
 #endif
 	}
 }
@@ -4506,7 +4471,7 @@ void FWrapLayer::CmdPipelineBarrier2KHR(VkResult Result, VkCommandBuffer Command
 				BreakOnTrackingImage(ImageBarrier.image);
 				if (DumpTrackImage(ImageBarrier.image))
 				{
-					CmdPrintfBegin(CommandBuffer, FString::Printf(TEXT("vkCmdPipelineBarrier(Flags=%d, NumMemB=%d, MemB=0x%p,"), (uint32)DependencyInfo->dependencyFlags, DependencyInfo->memoryBarrierCount, DependencyInfo->pMemoryBarriers));
+					CmdPrintfBegin(CommandBuffer, FString::Printf(TEXT("vkCmdPipelineBarrier2(Flags=%d, NumMemB=%d, MemB=0x%p,"), (uint32)DependencyInfo->dependencyFlags, DependencyInfo->memoryBarrierCount, DependencyInfo->pMemoryBarriers));
 					DebugLog += FString::Printf(TEXT("%s\tNumBufferB=%d, BufferB=0x%p, NumImageB=%d, ImageB=0x%p)[...]\n"), Tabs, DependencyInfo->bufferMemoryBarrierCount, DependencyInfo->pBufferMemoryBarriers, DependencyInfo->imageMemoryBarrierCount, DependencyInfo->pImageMemoryBarriers);
 					DumpImageMemoryBarriers(DependencyInfo->imageMemoryBarrierCount, DependencyInfo->pImageMemoryBarriers);
 					FlushDebugWrapperLog();
diff --git a/Engine/Source/Runtime/VulkanRHI/Private/VulkanDescriptorSets.h b/Engine/Source/Runtime/VulkanRHI/Private/VulkanDescriptorSets.h
--- a/Engine/Source/Runtime/VulkanRHI/Private/VulkanDescriptorSets.h
+++ b/Engine/Source/Runtime/VulkanRHI/Private/VulkanDescriptorSets.h
@@ -1201,7 +1201,8 @@ protected:
 			checkf(WriteDescriptors[DescriptorIndex].descriptorType == VK_DESCRIPTOR_TYPE_SAMPLED_IMAGE || WriteDescriptors[DescriptorIndex].descriptorType == VK_DESCRIPTOR_TYPE_COMBINED_IMAGE_SAMPLER,
 				TEXT("DescriptorType mismatch at index %d: called WriteTextureView<%d> and was expecting %d."),
 				DescriptorIndex, (uint32)DescriptorType, (uint32)WriteDescriptors[DescriptorIndex].descriptorType);
-			ensureMsgf(Layout == VK_IMAGE_LAYOUT_SHARED_PRESENT_KHR ||
+			ensureMsgf(Layout == VK_IMAGE_LAYOUT_READ_ONLY_OPTIMAL ||
+				  Layout == VK_IMAGE_LAYOUT_SHARED_PRESENT_KHR ||
 				  Layout == VK_IMAGE_LAYOUT_DEPTH_READ_ONLY_STENCIL_ATTACHMENT_OPTIMAL ||
 				  Layout == VK_IMAGE_LAYOUT_DEPTH_ATTACHMENT_STENCIL_READ_ONLY_OPTIMAL ||
 				  Layout == VK_IMAGE_LAYOUT_DEPTH_STENCIL_READ_ONLY_OPTIMAL || 
diff --git a/Engine/Source/Runtime/VulkanRHI/Private/VulkanDevice.cpp b/Engine/Source/Runtime/VulkanRHI/Private/VulkanDevice.cpp
--- a/Engine/Source/Runtime/VulkanRHI/Private/VulkanDevice.cpp
+++ b/Engine/Source/Runtime/VulkanRHI/Private/VulkanDevice.cpp
@@ -15,6 +15,7 @@
 #include "VulkanLLM.h"
 #include "VulkanTransientResourceAllocator.h"
 #include "VulkanExtensions.h"
+#include "VulkanRenderpass.h"
 #include "VulkanRayTracing.h"
 #include "VulkanDescriptorSets.h"
 
@@ -1072,6 +1073,8 @@ void FVulkanDevice::InitGPU()
 	}
 #endif
 
+	RenderPassManager = new FVulkanRenderPassManager(this);
+
 	if (UseVulkanDescriptorCache())
 	{
 		DescriptorSetCache = new FVulkanDescriptorSetCache(this);
@@ -1227,6 +1230,9 @@ void FVulkanDevice::Destroy()
 	delete ImmediateContext;
 	ImmediateContext = nullptr;
 
+	delete RenderPassManager;
+	RenderPassManager = nullptr;
+
 	for (FVulkanOcclusionQueryPool* Pool : UsedOcclusionQueryPools)
 	{
 		delete Pool;
@@ -1326,10 +1332,13 @@ void FVulkanDevice::NotifyDeletedImage(VkImage Image, bool bRenderTarget)
 {
 	if (bRenderTarget)
 	{
+		// Contexts first, as it may clear the current framebuffer
 		GetImmediateContext().NotifyDeletedRenderTarget(Image);
+		// Delete framebuffers using this image
+		GetRenderPassManager().NotifyDeletedRenderTarget(Image);
 	}
 
-	//#todo-rco: Loop through all contexts!
+	//#todo-jn: Loop through all contexts!  And all queues!
 	GetImmediateContext().NotifyDeletedImage(Image);
 }
 
diff --git a/Engine/Source/Runtime/VulkanRHI/Private/VulkanDevice.h b/Engine/Source/Runtime/VulkanRHI/Private/VulkanDevice.h
--- a/Engine/Source/Runtime/VulkanRHI/Private/VulkanDevice.h
+++ b/Engine/Source/Runtime/VulkanRHI/Private/VulkanDevice.h
@@ -15,11 +15,8 @@ class FVulkanBindlessDescriptorManager;
 class FVulkanCommandListContextImmediate;
 class FVulkanTransientHeapCache;
 class FVulkanDeviceExtension;
-#if VULKAN_USE_NEW_QUERIES
 class FVulkanOcclusionQueryPool;
-#else
-class FOLDVulkanQueryPool;
-#endif
+class FVulkanRenderPassManager;
 
 #if VULKAN_RHI_RAYTRACING
 class FVulkanBasicRaytracingPipeline;
@@ -84,6 +81,7 @@ struct FOptionalVulkanDeviceExtensions
 			// Promoted to 1.3
 			uint64 HasEXTTextureCompressionASTCHDR : 1;
 			uint64 HasKHRMaintenance4 : 1;
+			uint64 HasKHRSynchronization2 : 1;
 		};
 		uint64 Packed;
 	};
@@ -376,6 +374,11 @@ public:
 		return FenceManager;
 	}
 
+	inline FVulkanRenderPassManager& GetRenderPassManager()
+	{
+		return *RenderPassManager;
+	}
+
 	inline FVulkanDescriptorSetCache& GetDescriptorSetCache()
 	{
 		return *DescriptorSetCache;
@@ -468,6 +471,11 @@ public:
 		return OptionalDeviceExtensions;
 	}
 
+	inline bool SupportsParallelRendering() const
+	{
+		return OptionalDeviceExtensions.HasSeparateDepthStencilLayouts && OptionalDeviceExtensions.HasKHRSynchronization2 && OptionalDeviceExtensions.HasKHRRenderPass2;
+	}
+
 #if VULKAN_SUPPORTS_GPU_CRASH_DUMPS
 	VkBuffer GetCrashMarkerBuffer() const
 	{
@@ -524,6 +532,8 @@ private:
 
 	VulkanRHI::FFenceManager FenceManager;
 
+	FVulkanRenderPassManager* RenderPassManager;
+
 	FVulkanTransientHeapCache* TransientHeapCache = nullptr;
 
 	// Active on ES3.1
diff --git a/Engine/Source/Runtime/VulkanRHI/Private/VulkanExtensions.cpp b/Engine/Source/Runtime/VulkanRHI/Private/VulkanExtensions.cpp
--- a/Engine/Source/Runtime/VulkanRHI/Private/VulkanExtensions.cpp
+++ b/Engine/Source/Runtime/VulkanRHI/Private/VulkanExtensions.cpp
@@ -51,6 +51,16 @@ TAutoConsoleVariable<int32> GVulkanAllowHostQueryResetCVar(
 	ECVF_ReadOnly
 );
 
+TAutoConsoleVariable<int32> GVulkaAllowSync2BarriersCVar(
+	TEXT("r.Vulkan.AllowSynchronization2"),
+	1,
+	TEXT("Enables the use of advanced barriers that combine the use of the VK_KHR_separate_depth_stencil_layouts \n")
+	TEXT("and VK_KHR_synchronization2 to reduce the reliance on layout tracking (except for defragging).\n")
+	TEXT("This is necessary in order to support parallel command buffer generation.\n")
+	TEXT("0: Do not enable support for sync2 barriers.\n")
+	TEXT("1: Enable sync2 barriers (default)"),
+	ECVF_ReadOnly
+);
 
 #if VULKAN_HAS_DEBUGGING_ENABLED
 extern TAutoConsoleVariable<int32> GGPUValidationCvar;
@@ -285,10 +295,9 @@ class FVulkanKHRSeparateDepthStencilLayoutsExtension : public FVulkanDeviceExten
 public:
 
 	FVulkanKHRSeparateDepthStencilLayoutsExtension(FVulkanDevice* InDevice)
-		: FVulkanDeviceExtension(InDevice, VK_KHR_SEPARATE_DEPTH_STENCIL_LAYOUTS_EXTENSION_NAME, VULKAN_SUPPORTS_SEPARATE_DEPTH_STENCIL_LAYOUTS, VK_API_VERSION_1_2)
+		: FVulkanDeviceExtension(InDevice, VK_KHR_SEPARATE_DEPTH_STENCIL_LAYOUTS_EXTENSION_NAME, VULKAN_EXTENSION_ENABLED, VK_API_VERSION_1_2)
 	{
-		// Disabled but kept for reference. The barriers code doesn't currently use separate transitions for depth and stencil.
-		bEnabledInCode = false;
+		bEnabledInCode = bEnabledInCode && (GVulkaAllowSync2BarriersCVar.GetValueOnAnyThread() != 0);
 	}
 
 	virtual void PrePhysicalDeviceFeatures(VkPhysicalDeviceFeatures2KHR& PhysicalDeviceFeatures2) override final
@@ -316,6 +325,42 @@ private:
 
 
 
+// ***** VK_KHR_synchronization2
+class FVulkanKHRSynchronization2 : public FVulkanDeviceExtension
+{
+public:
+
+	FVulkanKHRSynchronization2(FVulkanDevice* InDevice)
+		: FVulkanDeviceExtension(InDevice, VK_KHR_SYNCHRONIZATION_2_EXTENSION_NAME, VULKAN_EXTENSION_ENABLED, VK_API_VERSION_1_3)
+	{
+		bEnabledInCode = bEnabledInCode && (GVulkaAllowSync2BarriersCVar.GetValueOnAnyThread() != 0);
+	}
+
+	virtual void PrePhysicalDeviceFeatures(VkPhysicalDeviceFeatures2KHR& PhysicalDeviceFeatures2) override final
+	{
+		ZeroVulkanStruct(Synchronization2Features, VK_STRUCTURE_TYPE_PHYSICAL_DEVICE_SYNCHRONIZATION_2_FEATURES);
+		AddToPNext(PhysicalDeviceFeatures2, Synchronization2Features);
+	}
+
+	virtual void PostPhysicalDeviceFeatures(FOptionalVulkanDeviceExtensions& ExtensionFlags) override final
+	{
+		ExtensionFlags.HasKHRSynchronization2 = Synchronization2Features.synchronization2;
+	}
+
+	virtual void PreCreateDevice(VkDeviceCreateInfo& DeviceCreateInfo) override final
+	{
+		if (Synchronization2Features.synchronization2 == VK_TRUE)
+		{
+			AddToPNext(DeviceCreateInfo, Synchronization2Features);
+		}
+	}
+
+private:
+	VkPhysicalDeviceSynchronization2FeaturesKHR Synchronization2Features;
+};
+
+
+
 // ***** VK_KHR_multiview
 class FVulkanKHRMultiviewExtension : public FVulkanDeviceExtension
 {
@@ -914,7 +959,6 @@ private:
 	VkPhysicalDeviceHostQueryResetFeaturesEXT HostQueryResetFeatures;
 };
 
-
 // ***** VK_EXT_subgroup_size_control
 class FVulkanEXTSubgroupSizeControlExtension : public FVulkanDeviceExtension
 {
@@ -1055,6 +1099,7 @@ FVulkanDeviceExtensionArray FVulkanDeviceExtension::GetUESupportedDeviceExtensio
 	ADD_CUSTOM_EXTENSION(FVulkanEXTScalarBlockLayoutExtension);
 	ADD_CUSTOM_EXTENSION(FVulkanEXTShaderViewportIndexLayerExtension);
 	ADD_CUSTOM_EXTENSION(FVulkanKHRSeparateDepthStencilLayoutsExtension);
+	ADD_CUSTOM_EXTENSION(FVulkanKHRSynchronization2);
 	ADD_CUSTOM_EXTENSION(FVulkanKHRFragmentShadingRateExtension); // must be kept BEFORE DensityMap!
 	ADD_CUSTOM_EXTENSION(FVulkanEXTFragmentDensityMapExtension);  // must be kept AFTER ShadingRate!
 	ADD_CUSTOM_EXTENSION(FVulkanEXTFragmentDensityMap2Extension);
diff --git a/Engine/Source/Runtime/VulkanRHI/Private/VulkanMemory.cpp b/Engine/Source/Runtime/VulkanRHI/Private/VulkanMemory.cpp
--- a/Engine/Source/Runtime/VulkanRHI/Private/VulkanMemory.cpp
+++ b/Engine/Source/Runtime/VulkanRHI/Private/VulkanMemory.cpp
@@ -2610,13 +2610,13 @@ namespace VulkanRHI
 
 	FVulkanSubresourceAllocator::~FVulkanSubresourceAllocator()
 	{
-		if(!JoinFreeBlocks())
+		if (!JoinFreeBlocks())
 		{
 			UE_LOG(LogVulkanRHI, Warning, TEXT("FVulkanSubresourceAllocator %p has unfreed %s resources %s"), (void*)this, VulkanAllocationTypeToString(Type), *VULKAN_TRACK_STRING(Track));
 			uint32 LeakCount = 0;
-			for(FVulkanAllocationInternal& Data : InternalData)
+			for (FVulkanAllocationInternal& Data : InternalData)
 			{
-				if(Data.State == FVulkanAllocationInternal::EALLOCATED)
+				if (Data.State == FVulkanAllocationInternal::EALLOCATED)
 				{
 					UE_LOG(LogVulkanRHI, Warning, TEXT(" ** LEAK %03d [%08x-%08x] %u  %s \n%s"), LeakCount++, Data.AllocationOffset, Data.AllocationSize,  Data.Size, VulkanAllocationMetaTypeToString(Data.MetaType), *VULKAN_TRACK_STRING(Data.Track));
 				}
@@ -2657,7 +2657,7 @@ namespace VulkanRHI
 		}
 		else
 		{
-			checkf(false, TEXT("Unknown buffer alignment for VkBufferUsageFlags combination: 0x%x"), VK_FLAGS_TO_STRING(VkBufferUsageFlags, BufferUsageFlags));
+			checkf(false, TEXT("Unknown buffer alignment for VkBufferUsageFlags combination: 0x%x (%s)"), BufferUsageFlags, VK_FLAGS_TO_STRING(VkBufferUsageFlags, BufferUsageFlags));
 		}
 
 		return Alignment;
@@ -2709,7 +2709,7 @@ namespace VulkanRHI
 		}
 		else
 		{
-			checkf(false, TEXT("Unknown priority for VkBufferUsageFlags combination: 0x%x"), VK_FLAGS_TO_STRING(VkBufferUsageFlags, BufferUsageFlags));
+			checkf(false, TEXT("Unknown priority for VkBufferUsageFlags combination: 0x%x (%s)"), BufferUsageFlags, VK_FLAGS_TO_STRING(VkBufferUsageFlags, BufferUsageFlags));
 		}
 
 		return Priority;
@@ -3262,8 +3262,8 @@ namespace VulkanRHI
 
 		auto WriteLogLine = [](const FString& Name, FResourceHeapStats& Stat)
 		{
-			uint64 FreeMemory = Stat.TotalMemory - FMath::Min<uint64>(Stat.TotalMemory, Stat.UsedBufferMemory + Stat.UsedImageMemory);
-			FString HostString = VK_FLAGS_TO_STRING(VkMemoryPropertyFlags, Stat.MemoryFlags);
+			const uint64 FreeMemory = Stat.TotalMemory - FMath::Min<uint64>(Stat.TotalMemory, Stat.UsedBufferMemory + Stat.UsedImageMemory);
+			FString HostString = (Stat.MemoryFlags != 0) ? VK_FLAGS_TO_STRING(VkMemoryPropertyFlags, Stat.MemoryFlags) : TEXT("");
 			VULKAN_LOGMEMORY(TEXT("\t\t%-33s  |%8.2fmb / %8.2fmb / %11.2fmb / %11.2fmb | %10d %10d | %6d %6d %6d | %05x | %s"),
 				*Name,
 				Stat.UsedBufferMemory / (1024.f * 1024.f),
@@ -3410,7 +3410,7 @@ namespace VulkanRHI
 				{
 					Flags = MemoryProperties.memoryTypes[Allocator->MemoryTypeIndex].propertyFlags;
 				}
-				FString MemoryString = VK_FLAGS_TO_STRING(VkMemoryPropertyFlags, Flags);
+				FString MemoryString = (Flags != 0) ? VK_FLAGS_TO_STRING(VkMemoryPropertyFlags, Flags) : TEXT("");
 				FString NameId = FString::Printf(TEXT("%s [%4d]"), *Name, Allocator->AllocatorIndex);
 				WriteLogLineSubAllocator(NameId, MemoryString, *Allocator);
 			}
@@ -3425,7 +3425,7 @@ namespace VulkanRHI
 				{
 					Flags = MemoryProperties.memoryTypes[Allocator->MemoryTypeIndex].propertyFlags;
 				}
-				FString MemoryString = VK_FLAGS_TO_STRING(VkMemoryPropertyFlags, Flags);
+				FString MemoryString = (Flags != 0) ? VK_FLAGS_TO_STRING(VkMemoryPropertyFlags, Flags) : TEXT("");
 				VULKAN_LOGMEMORY(VULKAN_LOGMEMORY_PAD2);
 				VULKAN_LOGMEMORY(TEXT("\t\t%-45s  | %4s %8s | %4s / %10s / %10s / %10s / %10s | %10s / %10s / %10s / %10s | %10s / %10s | %10s / %10s / %10s | Mapped/Evictable |"),
 					TEXT(""),
@@ -3722,11 +3722,10 @@ namespace VulkanRHI
 
 	void FVulkanAllocation::Free(FVulkanDevice& Device)
 	{
-		if(HasAllocation())
+		if (HasAllocation())
 		{
 			Device.GetMemoryManager().FreeVulkanAllocation(*this);
 			check(EVulkanAllocationEmpty != Type);
-
 		}
 	}
 	void FVulkanAllocation::Swap(FVulkanAllocation& Other)
@@ -4055,7 +4054,7 @@ namespace VulkanRHI
 
 		FScopeLock ScopeLock(&SubresourceAllocatorCS);
 
-		//Search for allocations to move to different pages.
+		// Search for allocations to move to different pages.
 		for (FVulkanAllocationInternal& Alloc : InternalData)
 		{
 			if (Alloc.State == FVulkanAllocationInternal::EALLOCATED)
@@ -4065,44 +4064,50 @@ namespace VulkanRHI
 				{
 					case EVulkanAllocationMetaImageRenderTarget: //only rendertargets can be defragged
 					{
-						FVulkanAllocation Allocation;
-						// The current SubAllocator is tagged as locked, this will never allocate in the current SubAllocator.
-						if(Heap->TryRealloc(Allocation, EvictableOwner, EType::Image, Alloc.Size, Alloc.Alignment, Alloc.MetaType))
-						{
-							check(Allocation.HasAllocation());
-							FVulkanTexture* Texture = EvictableOwner->GetEvictableTexture();
+						FVulkanTexture* Texture = EvictableOwner->GetEvictableTexture();
 
-							if(GVulkanLogDefrag)
+						// Defrag works from the queue for layouts
+						// Only work with straightforward targets that are in a single layout
+						const FVulkanImageLayout* OriginalLayout = Context.GetQueue()->GetLayoutManager().GetFullLayout(Texture->Image);
+						if (OriginalLayout && OriginalLayout->AreAllSubresourcesSameLayout() && (OriginalLayout->MainLayout != VK_IMAGE_LAYOUT_UNDEFINED))
+						{
+							FVulkanAllocation Allocation;
+							// The current SubAllocator is tagged as locked, this will never allocate in the current SubAllocator.
+							if (Heap->TryRealloc(Allocation, EvictableOwner, EType::Image, Alloc.Size, Alloc.Alignment, Alloc.MetaType))
 							{
-								VULKAN_LOGMEMORY(TEXT("Moving %6.2fMB : %d:%08x -> %d/%08x\n"), Alloc.Size / (1024.f*1024.f),
-									AllocatorIndex,
-									Alloc.AllocationOffset,
-									Allocation.AllocatorIndex,
-									Allocation.Offset);
+								check(Allocation.HasAllocation());
+
+								if (GVulkanLogDefrag)
+								{
+									VULKAN_LOGMEMORY(TEXT("Moving %6.2fMB : %d:%08x -> %d/%08x\n"), Alloc.Size / (1024.f * 1024.f),
+										AllocatorIndex,
+										Alloc.AllocationOffset,
+										Allocation.AllocatorIndex,
+										Allocation.Offset);
+								}
+
+								//Move the Rendertarget to the new allocation
+								//Function swaps the old allocation into the Allocation object
+								Texture->Move(Device, Context, Allocation);
+								DefragCount++;
+								Device.GetMemoryManager().FreeVulkanAllocation(Allocation);
+
+								check(Alloc.State != FVulkanAllocationInternal::EALLOCATED);
+								check(!Allocation.HasAllocation()); //must be consumed by Move
+							}
+							else
+							{
+								check(!Allocation.HasAllocation());
+								bLocked = false;
+								return DefragCount;
 							}
-
-							//Move the Rendertarget to the new allocation
-							//Function swaps the old allocation into the Allocation object
-							Texture->Move(Device, Context, Allocation);
-							DefragCount++;
-							Device.GetMemoryManager().FreeVulkanAllocation(Allocation);
-
-							check(Alloc.State != FVulkanAllocationInternal::EALLOCATED);
-							check(!Allocation.HasAllocation()); //must be consumed by Move
-						}
-						else
-						{
-							check(!Allocation.HasAllocation());
-							bLocked = false;
-							return DefragCount;
 						}
-
 					}
 					break;
 					default:
 						checkNoEntry(); //not implemented.
 				}
-				if(0 >= --Count)
+				if (0 >= --Count)
 				{
 					break;
 				}
@@ -4124,7 +4129,7 @@ namespace VulkanRHI
 				{
 				case EVulkanAllocationMetaImageOther:
 				{
-					FVulkanEvictable* Texture= Alloc.AllocationOwner;
+					FVulkanEvictable* Texture = Alloc.AllocationOwner;
 					Texture->Evict(Device, Context);
 				}
 				break;
diff --git a/Engine/Source/Runtime/VulkanRHI/Private/VulkanPendingState.cpp b/Engine/Source/Runtime/VulkanRHI/Private/VulkanPendingState.cpp
--- a/Engine/Source/Runtime/VulkanRHI/Private/VulkanPendingState.cpp
+++ b/Engine/Source/Runtime/VulkanRHI/Private/VulkanPendingState.cpp
@@ -331,9 +331,11 @@ void FVulkanPendingComputeState::SetSRVForUBResource(uint32 DescriptorSet, uint3
 #endif // VULKAN_RHI_RAYTRACING
 		else
 		{
-			checkf(SRV->TextureView.View != VK_NULL_HANDLE, TEXT("Empty SRV"));
-			const FVulkanImageLayout& Layout = Context.GetLayoutManager().GetFullLayoutChecked(SRV->TextureView.Image);
-			CurrentState->SetSRVTextureView(DescriptorSet, BindingIndex, SRV->TextureView, Layout.GetSubresLayout(SRV->FirstArraySlice, SRV->MipLevel));
+			const FVulkanTexture* VulkanTexture = FVulkanTexture::Cast(SRV->SourceTexture.GetReference());
+			checkf((SRV->TextureView.View != VK_NULL_HANDLE) && VulkanTexture, TEXT("Empty SRV"));
+			const VkImageLayout ExpectedLayout = FVulkanLayoutManager::GetDefaultLayout(
+				Context.GetCommandBufferManager()->GetActiveCmdBuffer(), *VulkanTexture, ERHIAccess::SRVCompute);
+			CurrentState->SetSRVTextureView(DescriptorSet, BindingIndex, SRV->TextureView, ExpectedLayout);
 		}
 	}
 	else
@@ -358,7 +360,9 @@ void FVulkanPendingComputeState::SetUAVForUBResource(uint32 DescriptorSet, uint3
 		}
 		else if (UAV->SourceTexture)
 		{
-			CurrentState->SetUAVTextureView(DescriptorSet, BindingIndex, UAV->TextureView, VK_IMAGE_LAYOUT_GENERAL);
+			const FVulkanTexture* VulkanTexture = FVulkanTexture::Cast(UAV->SourceTexture.GetReference());
+			const VkImageLayout ExpectedLayout = FVulkanLayoutManager::GetDefaultLayout(Context.GetCommandBufferManager()->GetActiveCmdBuffer(), *VulkanTexture, ERHIAccess::UAVCompute);
+			CurrentState->SetUAVTextureView(DescriptorSet, BindingIndex, UAV->TextureView, ExpectedLayout);
 		}
 		else
 		{
@@ -415,7 +419,7 @@ void FVulkanPendingGfxState::PrepareForDraw(FVulkanCmdBuffer* CmdBuffer)
 	// TODO: Add 'dirty' flag? Need to rebind only on PSO change
 	if (CurrentPipeline->bHasInputAttachments)
 	{
-		FVulkanFramebuffer* CurrentFramebuffer = Context.GetLayoutManager().CurrentFramebuffer;
+		FVulkanFramebuffer* CurrentFramebuffer = Context.GetCurrentFramebuffer();
 		UpdateInputAttachments(CurrentFramebuffer);
 	}
 	
@@ -597,9 +601,10 @@ void FVulkanPendingGfxState::SetSRVForUBResource(uint8 DescriptorSet, uint32 Bin
 		}
 		else
 		{
-			checkf(SRV->TextureView.View != VK_NULL_HANDLE, TEXT("Empty SRV"));
-			const FVulkanImageLayout& Layout = Context.GetLayoutManager().GetFullLayoutChecked(SRV->TextureView.Image);
-			CurrentState->SetSRVTextureView(DescriptorSet, BindingIndex, SRV->TextureView, Layout.GetSubresLayout(SRV->FirstArraySlice, SRV->MipLevel));
+			const FVulkanTexture* VulkanTexture = FVulkanTexture::Cast(SRV->SourceTexture.GetReference());
+			checkf((SRV->TextureView.View != VK_NULL_HANDLE) && VulkanTexture, TEXT("Empty SRV"));
+			const VkImageLayout ExpectedLayout = FVulkanLayoutManager::GetDefaultLayout(Context.GetCommandBufferManager()->GetActiveCmdBuffer(), *VulkanTexture, ERHIAccess::SRVGraphics);
+			CurrentState->SetSRVTextureView(DescriptorSet, BindingIndex, SRV->TextureView, ExpectedLayout);
 		}
 	}
 	else
@@ -624,8 +629,10 @@ void FVulkanPendingGfxState::SetUAVForUBResource(uint8 DescriptorSet, uint32 Bin
 		}
 		else if (UAV->SourceTexture)
 		{
-			const FVulkanImageLayout& Layout = Context.GetLayoutManager().GetFullLayoutChecked(UAV->TextureView.Image);
-			CurrentState->SetUAVTextureView(DescriptorSet, BindingIndex, UAV->TextureView, Layout.GetSubresLayout(UAV->FirstArraySlice, UAV->MipLevel));
+			const FVulkanTexture* VulkanTexture = FVulkanTexture::Cast(UAV->SourceTexture.GetReference());
+			checkf(VulkanTexture, TEXT("Empty UAV"));
+			const VkImageLayout Layout = FVulkanLayoutManager::GetDefaultLayout(Context.GetCommandBufferManager()->GetActiveCmdBuffer(), *VulkanTexture, ERHIAccess::UAVGraphics);
+			CurrentState->SetUAVTextureView(DescriptorSet, BindingIndex, UAV->TextureView, Layout);
 		}
 		else
 		{
diff --git a/Engine/Source/Runtime/VulkanRHI/Private/VulkanPipeline.cpp b/Engine/Source/Runtime/VulkanRHI/Private/VulkanPipeline.cpp
--- a/Engine/Source/Runtime/VulkanRHI/Private/VulkanPipeline.cpp
+++ b/Engine/Source/Runtime/VulkanRHI/Private/VulkanPipeline.cpp
@@ -918,6 +918,23 @@ FArchive& operator << (FArchive& Ar, FGfxPipelineDesc::FRenderTargets::FAttachme
 	return Ar;
 }
 
+void FGfxPipelineDesc::FRenderTargets::FStencilAttachmentRef::ReadFrom(const VkAttachmentReferenceStencilLayout& InState)
+{
+	Layout = (uint64)InState.stencilLayout;
+}
+
+void FGfxPipelineDesc::FRenderTargets::FStencilAttachmentRef::WriteInto(VkAttachmentReferenceStencilLayout& Out) const
+{
+	Out.stencilLayout = (VkImageLayout)Layout;
+}
+
+FArchive& operator << (FArchive& Ar, FGfxPipelineDesc::FRenderTargets::FStencilAttachmentRef& AttachmentRef)
+{
+	// Modify VERSION if serialization changes
+	Ar << AttachmentRef.Layout;
+	return Ar;
+}
+
 void FGfxPipelineDesc::FRenderTargets::FAttachmentDesc::ReadFrom(const VkAttachmentDescription &InState)
 {
 	Format =			(uint32)InState.format;
@@ -960,6 +977,27 @@ FArchive& operator << (FArchive& Ar, FGfxPipelineDesc::FRenderTargets::FAttachme
 	return Ar;
 }
 
+void FGfxPipelineDesc::FRenderTargets::FStencilAttachmentDesc::ReadFrom(const VkAttachmentDescriptionStencilLayout& InState)
+{
+	InitialLayout = (uint64)InState.stencilInitialLayout;
+	FinalLayout = (uint64)InState.stencilFinalLayout;
+}
+
+void FGfxPipelineDesc::FRenderTargets::FStencilAttachmentDesc::WriteInto(VkAttachmentDescriptionStencilLayout& Out) const
+{
+	Out.stencilInitialLayout = (VkImageLayout)InitialLayout;
+	Out.stencilFinalLayout = (VkImageLayout)FinalLayout;
+}
+
+FArchive& operator << (FArchive& Ar, FGfxPipelineDesc::FRenderTargets::FStencilAttachmentDesc& StencilAttachmentDesc)
+{
+	// Modify VERSION if serialization changes
+	Ar << StencilAttachmentDesc.InitialLayout;
+	Ar << StencilAttachmentDesc.FinalLayout;
+
+	return Ar;
+}
+
 void FGfxPipelineDesc::FRenderTargets::ReadFrom(const FVulkanRenderTargetLayout& RTLayout)
 {
 	NumAttachments =			RTLayout.NumAttachmentDescriptions;
@@ -986,7 +1024,8 @@ void FGfxPipelineDesc::FRenderTargets::ReadFrom(const FVulkanRenderTargetLayout&
 	};
 	CopyAttachmentRefs(ColorAttachments, RTLayout.ColorReferences, UE_ARRAY_COUNT(RTLayout.ColorReferences));
 	CopyAttachmentRefs(ResolveAttachments, RTLayout.ResolveReferences, UE_ARRAY_COUNT(RTLayout.ResolveReferences));
-	DepthStencil.ReadFrom(RTLayout.DepthStencilReference);
+	Depth.ReadFrom(RTLayout.DepthReference);
+	Stencil.ReadFrom(RTLayout.StencilReference);
 	FragmentDensity.ReadFrom(RTLayout.FragmentDensityReference);
 
 	Descriptions.AddZeroed(UE_ARRAY_COUNT(RTLayout.Desc));
@@ -994,6 +1033,7 @@ void FGfxPipelineDesc::FRenderTargets::ReadFrom(const FVulkanRenderTargetLayout&
 	{
 		Descriptions[Index].ReadFrom(RTLayout.Desc[Index]);
 	}
+	StencilDescription.ReadFrom(RTLayout.StencilDesc);
 }
 
 void FGfxPipelineDesc::FRenderTargets::WriteInto(FVulkanRenderTargetLayout& Out) const
@@ -1022,7 +1062,8 @@ void FGfxPipelineDesc::FRenderTargets::WriteInto(FVulkanRenderTargetLayout& Out)
 	};
 	CopyAttachmentRefs(ColorAttachments, Out.ColorReferences, UE_ARRAY_COUNT(Out.ColorReferences));
 	CopyAttachmentRefs(ResolveAttachments, Out.ResolveReferences, UE_ARRAY_COUNT(Out.ResolveReferences));
-	DepthStencil.WriteInto(Out.DepthStencilReference);
+	Depth.WriteInto(Out.DepthReference);
+	Stencil.WriteInto(Out.StencilReference);
 	FragmentDensity.WriteInto(Out.FragmentDensityReference);
 
 
@@ -1030,6 +1071,7 @@ void FGfxPipelineDesc::FRenderTargets::WriteInto(FVulkanRenderTargetLayout& Out)
 	{
 		Descriptions[Index].WriteInto(Out.Desc[Index]);
 	}
+	StencilDescription.WriteInto(Out.StencilDesc);
 }
 
 FArchive& operator << (FArchive& Ar, FGfxPipelineDesc::FRenderTargets& RTs)
@@ -1040,10 +1082,12 @@ FArchive& operator << (FArchive& Ar, FGfxPipelineDesc::FRenderTargets& RTs)
 	Ar << RTs.NumUsedClearValues;
 	Ar << RTs.ColorAttachments;
 	Ar << RTs.ResolveAttachments;
-	Ar << RTs.DepthStencil;
+	Ar << RTs.Depth;
+	Ar << RTs.Stencil;
 	Ar << RTs.FragmentDensity;
 
 	Ar << RTs.Descriptions;
+	Ar << RTs.StencilDescription;
 
 	Ar << RTs.bHasDepthStencil;
 	Ar << RTs.bHasResolveAttachments;
diff --git a/Engine/Source/Runtime/VulkanRHI/Private/VulkanPipeline.h b/Engine/Source/Runtime/VulkanRHI/Private/VulkanPipeline.h
--- a/Engine/Source/Runtime/VulkanRHI/Private/VulkanPipeline.h
+++ b/Engine/Source/Runtime/VulkanRHI/Private/VulkanPipeline.h
@@ -294,9 +294,23 @@ struct FGfxPipelineDesc
 				return Attachment == In.Attachment && Layout == In.Layout;
 			}
 		};
+
+		struct FStencilAttachmentRef
+		{
+			uint64 Layout;
+
+			void ReadFrom(const VkAttachmentReferenceStencilLayout& InState);
+			void WriteInto(VkAttachmentReferenceStencilLayout& OutState) const;
+			bool operator == (const FStencilAttachmentRef& In) const
+			{
+				return Layout == In.Layout;
+			}
+		};
+
 		TArray<FAttachmentRef> ColorAttachments;
 		TArray<FAttachmentRef> ResolveAttachments;
-		FAttachmentRef DepthStencil;
+		FAttachmentRef Depth;
+		FStencilAttachmentRef Stencil;
 		FAttachmentRef FragmentDensity;
 
 		struct FAttachmentDesc
@@ -327,7 +341,24 @@ struct FGfxPipelineDesc
 			void ReadFrom(const VkAttachmentDescription &InState);
 			void WriteInto(VkAttachmentDescription& OutState) const;
 		};
+
+		struct FStencilAttachmentDesc
+		{
+			uint64 InitialLayout;
+			uint64 FinalLayout;
+
+			bool operator==(const FStencilAttachmentDesc& In) const
+			{
+				return InitialLayout == In.InitialLayout &&
+					FinalLayout == In.FinalLayout;
+			}
+
+			void ReadFrom(const VkAttachmentDescriptionStencilLayout& InState);
+			void WriteInto(VkAttachmentDescriptionStencilLayout& OutState) const;
+		};
+
 		TArray<FAttachmentDesc> Descriptions;
+		FStencilAttachmentDesc StencilDescription;
 
 		uint8 NumAttachments;
 		uint8 NumColorAttachments;
@@ -345,9 +376,11 @@ struct FGfxPipelineDesc
 		{
 			return ColorAttachments == In.ColorAttachments &&
 				ResolveAttachments == In.ResolveAttachments &&
-				DepthStencil == In.DepthStencil &&
+				Depth == In.Depth &&
+				Stencil == In.Stencil &&
 				FragmentDensity == In.FragmentDensity &&
 				Descriptions == In.Descriptions &&
+				StencilDescription == In.StencilDescription &&
 				NumAttachments == In.NumAttachments &&
 				NumColorAttachments == In.NumColorAttachments &&
 				bHasDepthStencil == In.bHasDepthStencil &&
diff --git a/Engine/Source/Runtime/VulkanRHI/Private/VulkanQuery.cpp b/Engine/Source/Runtime/VulkanRHI/Private/VulkanQuery.cpp
--- a/Engine/Source/Runtime/VulkanRHI/Private/VulkanQuery.cpp
+++ b/Engine/Source/Runtime/VulkanRHI/Private/VulkanQuery.cpp
@@ -14,13 +14,6 @@
 #include "HAL/PlatformStackwalk.h"
 #endif
 
-TAutoConsoleVariable<int32> GSubmitOcclusionBatchCmdBufferCVar(
-	TEXT("r.Vulkan.SubmitOcclusionBatchCmdBuffer"),
-	1,
-	TEXT("1 to submit the cmd buffer after end occlusion query batch (default)"),
-	ECVF_RenderThreadSafe
-);
-
 static uint32 GTimestampQueryStage = 0;
 TAutoConsoleVariable<int32> GTimestampQueryStageCVar(
 	TEXT("r.Vulkan.TimestampQueryStage"),
@@ -364,14 +357,6 @@ void FVulkanCommandListContext::EndOcclusionQueryBatch(FVulkanCmdBuffer* CmdBuff
 	checkf(CurrentOcclusionQueryPool, TEXT("EndOcclusionQueryBatch called without corresponding BeginOcclusionQueryBatch!"));
 	CurrentOcclusionQueryPool->EndBatch(CmdBuffer);
 	CurrentOcclusionQueryPool = nullptr;
-	LayoutManager.EndRenderPass(CmdBuffer);
-
-	// Sync point
-	if (GSubmitOcclusionBatchCmdBufferCVar.GetValueOnAnyThread())
-	{
-		RequestSubmitCurrentCommands();
-		SafePointSubmit();
-	}
 }
 
 
diff --git a/Engine/Source/Runtime/VulkanRHI/Private/VulkanQueue.cpp b/Engine/Source/Runtime/VulkanRHI/Private/VulkanQueue.cpp
--- a/Engine/Source/Runtime/VulkanRHI/Private/VulkanQueue.cpp
+++ b/Engine/Source/Runtime/VulkanRHI/Private/VulkanQueue.cpp
@@ -29,6 +29,7 @@ FVulkanQueue::FVulkanQueue(FVulkanDevice* InDevice, uint32 InFamilyIndex)
 	, LastSubmittedCmdBuffer(nullptr)
 	, LastSubmittedCmdBufferFenceCounter(0)
 	, SubmitCounter(0)
+	, LayoutManager(false, nullptr)
 {
 	VulkanRHI::vkGetDeviceQueue(Device->GetInstanceHandle(), FamilyIndex, QueueIndex, &Queue);
 
@@ -115,6 +116,9 @@ void FVulkanQueue::Submit(FVulkanCmdBuffer* CmdBuffer, uint32 NumSignalSemaphore
 	CmdBuffer->GetOwner()->RefreshFenceStatus(CmdBuffer);
 
 	Device->GetStagingManager().ProcessPendingFree(false, false);
+
+	// If we're tracking layouts for the queue, merge in the changes recorded in this command buffer's context
+	CmdBuffer->GetLayoutManager().TransferTo(LayoutManager);
 }
 
 void FVulkanQueue::UpdateLastSubmittedCommandBuffer(FVulkanCmdBuffer* CmdBuffer)
@@ -182,3 +186,8 @@ void FVulkanQueue::FillSupportedStageBits()
 		SupportedStages |= VK_PIPELINE_STAGE_TRANSFER_BIT;
 	}
 }
+
+void FVulkanQueue::NotifyDeletedImage(VkImage Image)
+{
+	LayoutManager.NotifyDeletedImage(Image);
+}
diff --git a/Engine/Source/Runtime/VulkanRHI/Private/VulkanQueue.h b/Engine/Source/Runtime/VulkanRHI/Private/VulkanQueue.h
--- a/Engine/Source/Runtime/VulkanRHI/Private/VulkanQueue.h
+++ b/Engine/Source/Runtime/VulkanRHI/Private/VulkanQueue.h
@@ -7,10 +7,12 @@
 #pragma once
 
 #include "VulkanConfiguration.h"
+#include "VulkanBarriers.h"
 
 class FVulkanDevice;
 class FVulkanCmdBuffer;
 class FVulkanCommandListContext;
+class FVulkanLayoutManager;
 
 namespace VulkanRHI
 {
@@ -46,7 +48,7 @@ public:
 		return Queue;
 	}
 
-	void GetLastSubmittedInfo(FVulkanCmdBuffer*& OutCmdBuffer, uint64& OutFenceCounter) const
+	inline void GetLastSubmittedInfo(FVulkanCmdBuffer*& OutCmdBuffer, uint64& OutFenceCounter) const
 	{
 		FScopeLock ScopeLock(&CS);
 		OutCmdBuffer = LastSubmittedCmdBuffer;
@@ -58,11 +60,18 @@ public:
 		return SubmitCounter;
 	}
 
-	inline VkPipelineStageFlags GetSupportedStageBits()
+	inline VkPipelineStageFlags GetSupportedStageBits() const
 	{
 		return SupportedStages;
 	}
 
+	inline FVulkanLayoutManager& GetLayoutManager()
+	{
+		return LayoutManager;
+	}
+
+	void NotifyDeletedImage(VkImage Image);
+
 private:
 	VkQueue Queue;
 	uint32 FamilyIndex;
@@ -75,6 +84,9 @@ private:
 	uint64 SubmitCounter;
 	VkPipelineStageFlags SupportedStages;
 
+	// Last known layouts submitted on this queue, used for defrag
+	FVulkanLayoutManager LayoutManager;
+
 	void UpdateLastSubmittedCommandBuffer(FVulkanCmdBuffer* CmdBuffer);
 	void FillSupportedStageBits();
 };
diff --git a/Engine/Source/Runtime/VulkanRHI/Private/VulkanRHI.cpp b/Engine/Source/Runtime/VulkanRHI/Private/VulkanRHI.cpp
--- a/Engine/Source/Runtime/VulkanRHI/Private/VulkanRHI.cpp
+++ b/Engine/Source/Runtime/VulkanRHI/Private/VulkanRHI.cpp
@@ -188,8 +188,6 @@ FVulkanCommandListContext::~FVulkanCommandListContext()
 	delete CommandBufferManager;
 	CommandBufferManager = nullptr;
 
-	LayoutManager.Destroy(*Device, Immediate ? &LayoutManager : nullptr);
-
 	delete UniformBufferUploader;
 	delete PendingGfxState;
 	delete PendingComputeState;
@@ -821,7 +819,7 @@ void FVulkanDynamicRHI::InitInstance()
 		GMaxTextureArrayLayers = Props.limits.maxImageArrayLayers;
 		GRHISupportsBaseVertexIndex = true;
 		GSupportsSeparateRenderTargetBlendState = true;
-		GRHISupportsSeparateDepthStencilCopyAccess = false;
+		GRHISupportsSeparateDepthStencilCopyAccess = Device->SupportsParallelRendering();
 		GRHIBindlessSupport = Device->SupportsBindless() ? RHIGetBindlessSupport(GMaxRHIShaderPlatform) : ERHIBindlessSupport::Unsupported;
 PRAGMA_DISABLE_DEPRECATION_WARNINGS
 		GRHISupportsBindless = GRHIBindlessSupport != ERHIBindlessSupport::Unsupported;
@@ -1336,10 +1334,15 @@ FVulkanRHIImageViewInfo FVulkanDynamicRHI::RHIGetImageViewInfo(FRHITexture* InTe
 	return Info;
 }
 
+// todo-jn: deprecate
 VkImageLayout& FVulkanDynamicRHI::RHIFindOrAddLayoutRW(FRHITexture* InTexture, VkImageLayout LayoutIfNotFound)
 {
 	FVulkanTexture* VulkanTexture = ResourceCast(InTexture);
-	return GetDevice()->GetImmediateContext().GetLayoutManager().FindOrAddLayoutRW(*VulkanTexture, LayoutIfNotFound);
+	FVulkanCommandListContext& ImmediateContext = GetDevice()->GetImmediateContext();
+	FVulkanCmdBuffer* CmdBuffer = ImmediateContext.GetCommandBufferManager()->GetActiveCmdBuffer();
+	// Removing const to allow original functionality even when parallel rendering is enabled.
+	FVulkanImageLayout* VulkanImageLayout = const_cast<FVulkanImageLayout*>(CmdBuffer->GetLayoutManager().GetFullLayout(*VulkanTexture, true, LayoutIfNotFound));
+	return VulkanImageLayout->MainLayout;
 }
 
 void FVulkanDynamicRHI::RHISetImageLayout(VkImage Image, VkImageLayout OldLayout, VkImageLayout NewLayout, const VkImageSubresourceRange& SubresourceRange)
@@ -1925,9 +1928,11 @@ void FVulkanDynamicRHI::RecreateSwapChain(void* NewNativeWindow)
 	}
 }
 
-void FVulkanDynamicRHI::VulkanSetImageLayout( VkCommandBuffer CmdBuffer, VkImage Image, VkImageLayout OldLayout, VkImageLayout NewLayout, const VkImageSubresourceRange& SubresourceRange )
+void FVulkanDynamicRHI::VulkanSetImageLayout(VkCommandBuffer CmdBuffer, VkImage Image, VkImageLayout OldLayout, VkImageLayout NewLayout, const VkImageSubresourceRange& SubresourceRange)
 {
-	::VulkanSetImageLayout( CmdBuffer, Image, OldLayout, NewLayout, SubresourceRange );
+	FVulkanPipelineBarrier Barrier;
+	Barrier.AddImageLayoutTransition(Image, OldLayout, NewLayout, SubresourceRange);
+	Barrier.Execute(CmdBuffer);
 }
 
 IRHITransientResourceAllocator* FVulkanDynamicRHI::RHICreateTransientResourceAllocator()
diff --git a/Engine/Source/Runtime/VulkanRHI/Private/VulkanRHIPrivate.h b/Engine/Source/Runtime/VulkanRHI/Private/VulkanRHIPrivate.h
--- a/Engine/Source/Runtime/VulkanRHI/Private/VulkanRHIPrivate.h
+++ b/Engine/Source/Runtime/VulkanRHI/Private/VulkanRHIPrivate.h
@@ -145,7 +145,7 @@ class FVulkanRenderTargetLayout
 public:
 	FVulkanRenderTargetLayout(const FGraphicsPipelineStateInitializer& Initializer);
 	FVulkanRenderTargetLayout(FVulkanDevice& InDevice, const FRHISetRenderTargetsInfo& RTInfo);
-	FVulkanRenderTargetLayout(FVulkanDevice& InDevice, const FRHIRenderPassInfo& RPInfo, VkImageLayout CurrentDSLayout);
+	FVulkanRenderTargetLayout(FVulkanDevice& InDevice, const FRHIRenderPassInfo& RPInfo, VkImageLayout CurrentDepthLayout, VkImageLayout CurrentStencilLayout);
 
 	inline uint32 GetRenderPassCompatibleHash() const
 	{
@@ -175,9 +175,12 @@ public:
 
 	inline const VkAttachmentReference* GetColorAttachmentReferences() const { return NumColorAttachments > 0 ? ColorReferences : nullptr; }
 	inline const VkAttachmentReference* GetResolveAttachmentReferences() const { return bHasResolveAttachments ? ResolveReferences : nullptr; }
-	inline const VkAttachmentReference* GetDepthStencilAttachmentReference() const { return bHasDepthStencil ? &DepthStencilReference : nullptr; }
+	inline const VkAttachmentReference* GetDepthAttachmentReference() const { return bHasDepthStencil ? &DepthReference : nullptr; }
+	inline const VkAttachmentReferenceStencilLayout* GetStencilAttachmentReference() const { return bHasDepthStencil ? &StencilReference : nullptr; }
 	inline const VkAttachmentReference* GetFragmentDensityAttachmentReference() const { return bHasFragmentDensityAttachment ? &FragmentDensityReference : nullptr; }
 
+	inline const VkAttachmentDescriptionStencilLayout* GetStencilDesc() const { return bHasDepthStencil ? &StencilDesc : nullptr; }
+
 	inline const ESubpassHint GetSubpassHint() const { return SubpassHint; }
 	inline const VkSurfaceTransformFlagBitsKHR GetQCOMRenderPassTransform() const { return QCOMRenderPassTransform; }
 
@@ -187,13 +190,14 @@ protected:
 protected:
 	VkSurfaceTransformFlagBitsKHR QCOMRenderPassTransform = VK_SURFACE_TRANSFORM_IDENTITY_BIT_KHR;
 	VkAttachmentReference ColorReferences[MaxSimultaneousRenderTargets];
-	VkAttachmentReference DepthStencilReference;
+	VkAttachmentReference DepthReference;
+	VkAttachmentReferenceStencilLayout StencilReference;
 	VkAttachmentReference FragmentDensityReference;
 	VkAttachmentReference ResolveReferences[MaxSimultaneousRenderTargets];
-	VkAttachmentReference InputAttachments[MaxSimultaneousRenderTargets + 1];
 
 	// Depth goes in the "+1" slot and the Shading Rate texture goes in the "+2" slot.
 	VkAttachmentDescription Desc[MaxSimultaneousRenderTargets * 2 + 2];
+	VkAttachmentDescriptionStencilLayout StencilDesc;
 
 	uint8 NumAttachmentDescriptions;
 	uint8 NumColorAttachments;
@@ -223,23 +227,32 @@ protected:
 		VkExtent2D	Extent2D;
 	} Extent;
 
-	FVulkanRenderTargetLayout()
+	inline void ResetAttachments()
 	{
 		FMemory::Memzero(ColorReferences);
-		FMemory::Memzero(DepthStencilReference);
+		FMemory::Memzero(DepthReference);
 		FMemory::Memzero(FragmentDensityReference);
 		FMemory::Memzero(ResolveReferences);
-		FMemory::Memzero(InputAttachments);
 		FMemory::Memzero(Desc);
+		FMemory::Memzero(Offset);
+		FMemory::Memzero(Extent);
+
+		ZeroVulkanStruct(StencilReference, VK_STRUCTURE_TYPE_ATTACHMENT_REFERENCE_STENCIL_LAYOUT);
+		ZeroVulkanStruct(StencilDesc, VK_STRUCTURE_TYPE_ATTACHMENT_DESCRIPTION_STENCIL_LAYOUT);
+	}
+
+	FVulkanRenderTargetLayout()
+	{
 		NumAttachmentDescriptions = 0;
 		NumColorAttachments = 0;
 		bHasDepthStencil = 0;
 		bHasResolveAttachments = 0;
 		bHasFragmentDensityAttachment = 0;
-		Extent.Extent3D.width = 0;
-		Extent.Extent3D.height = 0;
-		Extent.Extent3D.depth = 0;
+		NumSamples = 0;
+		NumUsedClearValues = 0;
 		MultiViewCount = 0;
+
+		ResetAttachments();
 	}
 
 	bool bCalculatedHash = false;
@@ -347,7 +360,7 @@ public:
 	}
 
 private:
-	friend class FVulkanLayoutManager;
+	friend class FVulkanRenderPassManager;
 	friend class FVulkanPipelineStateCacheManager;
 
 	FVulkanRenderPass(FVulkanDevice& Device, const FVulkanRenderTargetLayout& RTLayout);
@@ -380,21 +393,8 @@ union UNvidiaDriverVersion
 };
 
 // Transitions an image to the specified layout. This does not update the layout cached internally by the RHI; the calling code must do that explicitly via FVulkanCommandListContext::GetLayoutManager() if necessary.
-void VulkanSetImageLayout(VkCommandBuffer CmdBuffer, VkImage Image, VkImageLayout OldLayout, VkImageLayout NewLayout, const VkImageSubresourceRange& SubresourceRange);
-
-// Transitions Color Images's first mip/layer/face
-inline void VulkanSetImageLayoutSimple(VkCommandBuffer CmdBuffer, VkImage Image, VkImageLayout OldLayout, VkImageLayout NewLayout, VkImageAspectFlags Aspect = VK_IMAGE_ASPECT_COLOR_BIT)
-{
-	VkImageSubresourceRange SubresourceRange = { Aspect, 0, 1, 0, 1 };
-	VulkanSetImageLayout(CmdBuffer, Image, OldLayout, NewLayout, SubresourceRange);
-}
+void VulkanSetImageLayout(FVulkanCmdBuffer* CmdBuffer, VkImage Image, VkImageLayout OldLayout, VkImageLayout NewLayout, const VkImageSubresourceRange& SubresourceRange);
 
-// Transitions all mips of Color Image
-inline void VulkanSetImageLayoutAllMips(VkCommandBuffer CmdBuffer, VkImage Image, VkImageLayout OldLayout, VkImageLayout NewLayout, VkImageAspectFlags Aspect = VK_IMAGE_ASPECT_COLOR_BIT)
-{
-	VkImageSubresourceRange SubresourceRange = { Aspect, 0, VK_REMAINING_MIP_LEVELS , 0, 1 };
-	VulkanSetImageLayout(CmdBuffer, Image, OldLayout, NewLayout, SubresourceRange);
-}
 
 
 DECLARE_STATS_GROUP(TEXT("Vulkan PSO"), STATGROUP_VulkanPSO, STATCAT_Advanced);
@@ -437,6 +437,7 @@ DECLARE_CYCLE_STAT_EXTERN(TEXT("DrawPrim UP Prep Time"), STAT_VulkanUPPrepTime,
 DECLARE_CYCLE_STAT_EXTERN(TEXT("Uniform Buffer Creation Time"), STAT_VulkanUniformBufferCreateTime, STATGROUP_VulkanRHI, );
 DECLARE_CYCLE_STAT_EXTERN(TEXT("Apply DS Uniform Buffers"), STAT_VulkanApplyDSUniformBuffers, STATGROUP_VulkanRHI, );
 DECLARE_CYCLE_STAT_EXTERN(TEXT("Apply Packed Uniform Buffers"), STAT_VulkanApplyPackedUniformBuffers, STATGROUP_VulkanRHI, );
+DECLARE_CYCLE_STAT_EXTERN(TEXT("Barrier Time"), STAT_VulkanBarrierTime, STATGROUP_VulkanRHI, );
 DECLARE_CYCLE_STAT_EXTERN(TEXT("SRV Update Time"), STAT_VulkanSRVUpdateTime, STATGROUP_VulkanRHI, );
 DECLARE_CYCLE_STAT_EXTERN(TEXT("UAV Update Time"), STAT_VulkanUAVUpdateTime, STATGROUP_VulkanRHI, );
 DECLARE_CYCLE_STAT_EXTERN(TEXT("Deletion Queue"), STAT_VulkanDeletionQueue, STATGROUP_VulkanRHI, );
@@ -810,23 +811,97 @@ namespace VulkanRHI
 
 	inline VkImageLayout GetDepthStencilLayout(FExclusiveDepthStencil RequestedDSAccess, FVulkanDevice& InDevice)
 	{
+		if (!RequestedDSAccess.IsUsingStencil() && InDevice.SupportsParallelRendering())
+		{
+			if (RequestedDSAccess.IsDepthRead())
+			{
+				return VK_IMAGE_LAYOUT_READ_ONLY_OPTIMAL;
+			}
+			else if (RequestedDSAccess.IsDepthWrite())
+			{
+				return VK_IMAGE_LAYOUT_ATTACHMENT_OPTIMAL;
+			}
+		}
+
 		if (RequestedDSAccess == FExclusiveDepthStencil::DepthRead_StencilNop || RequestedDSAccess == FExclusiveDepthStencil::DepthRead_StencilRead)
 		{
 			return VK_IMAGE_LAYOUT_DEPTH_STENCIL_READ_ONLY_OPTIMAL;
 		}
 		else if (RequestedDSAccess == FExclusiveDepthStencil::DepthRead_StencilWrite && InDevice.GetOptionalExtensions().HasKHRMaintenance2)
 		{
-			return VK_IMAGE_LAYOUT_DEPTH_READ_ONLY_STENCIL_ATTACHMENT_OPTIMAL_KHR;
+			return VK_IMAGE_LAYOUT_DEPTH_READ_ONLY_STENCIL_ATTACHMENT_OPTIMAL;
+		}
+
+		else if (RequestedDSAccess.IsDepthWrite() && !RequestedDSAccess.IsStencilWrite() && InDevice.SupportsParallelRendering())
+		{
+			return VK_IMAGE_LAYOUT_DEPTH_ATTACHMENT_STENCIL_READ_ONLY_OPTIMAL;
 		}
-		if(!RequestedDSAccess.IsUsingDepth() && RequestedDSAccess.IsUsingStencil())
+
+		if(!RequestedDSAccess.IsUsingDepth() && RequestedDSAccess.IsUsingStencil())  // todo-jn: wussat?  still needs maintenance2?
 		{
-			return VK_IMAGE_LAYOUT_DEPTH_READ_ONLY_STENCIL_ATTACHMENT_OPTIMAL_KHR;
+			return VK_IMAGE_LAYOUT_DEPTH_READ_ONLY_STENCIL_ATTACHMENT_OPTIMAL;
 		}
 
 		ensure(RequestedDSAccess.IsDepthWrite() || RequestedDSAccess.IsStencilWrite());
 		return VK_IMAGE_LAYOUT_DEPTH_STENCIL_ATTACHMENT_OPTIMAL;
 	}
 
+	// Merge a depth and a stencil layout for drivers that don't support VK_KHR_separate_depth_stencil_layouts
+	inline VkImageLayout GetMergedDepthStencilLayout(VkImageLayout DepthLayout, VkImageLayout StencilLayout)
+	{
+		if ((DepthLayout == VK_IMAGE_LAYOUT_TRANSFER_SRC_OPTIMAL) || (StencilLayout == VK_IMAGE_LAYOUT_TRANSFER_SRC_OPTIMAL))
+		{
+			checkf(StencilLayout == DepthLayout,
+				TEXT("You can't merge transfer src layout without anything else than transfer src (%s != %s).  ")
+				TEXT("You need either VK_KHR_separate_depth_stencil_layouts or GRHISupportsSeparateDepthStencilCopyAccess enabled."),
+				VK_TYPE_TO_STRING(VkImageLayout, DepthLayout), VK_TYPE_TO_STRING(VkImageLayout, StencilLayout));
+			return VK_IMAGE_LAYOUT_TRANSFER_SRC_OPTIMAL;
+		}
+
+		if ((DepthLayout == VK_IMAGE_LAYOUT_TRANSFER_DST_OPTIMAL) || (StencilLayout == VK_IMAGE_LAYOUT_TRANSFER_DST_OPTIMAL))
+		{
+			checkf(StencilLayout == DepthLayout,
+				TEXT("You can't merge transfer dst layout without anything else than transfer dst (%s != %s).  ")
+				TEXT("You need either VK_KHR_separate_depth_stencil_layouts or GRHISupportsSeparateDepthStencilCopyAccess enabled."),
+				VK_TYPE_TO_STRING(VkImageLayout, DepthLayout), VK_TYPE_TO_STRING(VkImageLayout, StencilLayout));
+			return VK_IMAGE_LAYOUT_TRANSFER_DST_OPTIMAL;
+		}
+
+		if ((DepthLayout == VK_IMAGE_LAYOUT_UNDEFINED) && (StencilLayout == VK_IMAGE_LAYOUT_UNDEFINED))
+		{
+			return VK_IMAGE_LAYOUT_UNDEFINED;
+		}
+
+		if (DepthLayout == VK_IMAGE_LAYOUT_ATTACHMENT_OPTIMAL)
+		{
+			if ((StencilLayout == VK_IMAGE_LAYOUT_ATTACHMENT_OPTIMAL) || (StencilLayout == VK_IMAGE_LAYOUT_UNDEFINED))
+			{
+				return VK_IMAGE_LAYOUT_DEPTH_STENCIL_ATTACHMENT_OPTIMAL;
+			}
+			else
+			{
+				check(StencilLayout == VK_IMAGE_LAYOUT_READ_ONLY_OPTIMAL);
+				return VK_IMAGE_LAYOUT_DEPTH_ATTACHMENT_STENCIL_READ_ONLY_OPTIMAL;
+			}
+		}
+		else if (DepthLayout == VK_IMAGE_LAYOUT_READ_ONLY_OPTIMAL)
+		{
+			if ((StencilLayout == VK_IMAGE_LAYOUT_READ_ONLY_OPTIMAL) || (StencilLayout == VK_IMAGE_LAYOUT_UNDEFINED))
+			{
+				return VK_IMAGE_LAYOUT_DEPTH_STENCIL_READ_ONLY_OPTIMAL;
+			}
+			else
+			{
+				check(StencilLayout == VK_IMAGE_LAYOUT_ATTACHMENT_OPTIMAL);
+				return VK_IMAGE_LAYOUT_DEPTH_READ_ONLY_STENCIL_ATTACHMENT_OPTIMAL;
+			}
+		}
+		else
+		{
+			return (StencilLayout == VK_IMAGE_LAYOUT_ATTACHMENT_OPTIMAL) ? VK_IMAGE_LAYOUT_DEPTH_STENCIL_ATTACHMENT_OPTIMAL : VK_IMAGE_LAYOUT_DEPTH_STENCIL_READ_ONLY_OPTIMAL;
+		}
+	}
+
 	inline void HeavyWeightBarrier(VkCommandBuffer CmdBuffer)
 	{
 		VkMemoryBarrier Barrier;
diff --git a/Engine/Source/Runtime/VulkanRHI/Private/VulkanRenderTarget.cpp b/Engine/Source/Runtime/VulkanRHI/Private/VulkanRenderTarget.cpp
--- a/Engine/Source/Runtime/VulkanRHI/Private/VulkanRenderTarget.cpp
+++ b/Engine/Source/Runtime/VulkanRHI/Private/VulkanRenderTarget.cpp
@@ -45,6 +45,13 @@ static FAutoConsoleVariableRef CVarVulkanIgnoreCPUReads(
 	ECVF_Default
 	);
 
+TAutoConsoleVariable<int32> GSubmitOcclusionBatchCmdBufferCVar(
+	TEXT("r.Vulkan.SubmitOcclusionBatchCmdBuffer"),
+	1,
+	TEXT("1 to submit the cmd buffer after end occlusion query batch (default)"),
+	ECVF_RenderThreadSafe
+);
+
 static FCriticalSection GStagingMapLock;
 static TMap<FVulkanTexture*, VulkanRHI::FStagingBuffer*> GPendingLockedStagingBuffers;
 
@@ -73,7 +80,7 @@ FVulkanRenderPass* FVulkanCommandListContext::PrepareRenderPassForPSOCreation(co
 FVulkanRenderPass* FVulkanCommandListContext::PrepareRenderPassForPSOCreation(const FVulkanRenderTargetLayout& RTLayout)
 {
 	FVulkanRenderPass* RenderPass = nullptr;
-	RenderPass = LayoutManager.GetOrCreateRenderPass(*Device, RTLayout);
+	RenderPass = Device->GetRenderPassManager().GetOrCreateRenderPass(RTLayout);
 	return RenderPass;
 }
 
@@ -215,22 +222,13 @@ void FVulkanCommandListContext::RHICopyToResolveTarget(FRHITexture* SourceTextur
 	DstRange.baseArrayLayer = InResolveParams.DestArrayIndex * DstNumLayers + DstCubeFace;
 	DstRange.layerCount = 1;
 
-	ERHIAccess SrcCurrentAccess, DstCurrentAccess;
-
 	check(!EnumHasAnyFlags(SrcSurface.GetDesc().Flags, TexCreate_CPUReadback));
-	VkImageLayout& SrcLayout = LayoutManager.FindOrAddLayoutRW(SrcSurface, VK_IMAGE_LAYOUT_UNDEFINED);
+	VkImageLayout SrcLayout = FVulkanLayoutManager::SetExpectedLayout(CmdBuffer, SrcSurface, ERHIAccess::CopySrc);
 
-	if(EnumHasAnyFlags(DstSurface.GetDesc().Flags, TexCreate_CPUReadback))
+	if (EnumHasAnyFlags(DstSurface.GetDesc().Flags, TexCreate_CPUReadback))
 	{
 		//Readback textures are represented as a buffer, so we can support miplevels on hardware that does not expose it.
-		FVulkanPipelineBarrier BarrierBefore;
-		// We'll transition the entire resources to the correct copy states, so we don't need to worry about sub-resource states.
-		if (SrcLayout != VK_IMAGE_LAYOUT_TRANSFER_SRC_OPTIMAL)
-		{
-			BarrierBefore.AddImageLayoutTransition(SrcSurface.Image, SrcLayout, VK_IMAGE_LAYOUT_TRANSFER_SRC_OPTIMAL, FVulkanPipelineBarrier::MakeSubresourceRange(SrcSurface.GetFullAspectMask()));
-			SrcLayout = VK_IMAGE_LAYOUT_TRANSFER_SRC_OPTIMAL;
-		}
-		BarrierBefore.Execute(CmdBuffer->GetHandle());
+
 		const FVulkanCpuReadbackBuffer* CpuReadbackBuffer = DstSurface.GetCpuReadbackBuffer();
 		check(DstRange.baseArrayLayer == 0);
 		VkBufferImageCopy CopyRegion;
@@ -254,18 +252,19 @@ void FVulkanCommandListContext::RHICopyToResolveTarget(FRHITexture* SourceTextur
 			FVulkanPipelineBarrier BarrierAfter;
 			BarrierAfter.AddMemoryBarrier(VK_ACCESS_TRANSFER_WRITE_BIT, VK_ACCESS_HOST_READ_BIT, VK_PIPELINE_STAGE_ALL_COMMANDS_BIT, VK_PIPELINE_STAGE_HOST_BIT);
 
-			SrcCurrentAccess = ERHIAccess::CopySrc;
-			if(SrcCurrentAccess != InResolveParams.SourceAccessFinal && InResolveParams.SourceAccessFinal != ERHIAccess::Unknown)
+			if (InResolveParams.SourceAccessFinal != ERHIAccess::CopySrc && InResolveParams.SourceAccessFinal != ERHIAccess::Unknown)
 			{
-				BarrierAfter.AddImageAccessTransition(SrcSurface, SrcCurrentAccess, InResolveParams.SourceAccessFinal, FVulkanPipelineBarrier::MakeSubresourceRange(SrcSurface.GetFullAspectMask()), SrcLayout);
+				BarrierAfter.AddImageAccessTransition(SrcSurface, ERHIAccess::CopySrc, InResolveParams.SourceAccessFinal, FVulkanPipelineBarrier::MakeSubresourceRange(SrcSurface.GetFullAspectMask()), SrcLayout);
 			}
+			CmdBuffer->GetLayoutManager().SetFullLayout(SrcSurface, SrcLayout);
 
-			BarrierAfter.Execute(CmdBuffer->GetHandle());
+			BarrierAfter.Execute(CmdBuffer);
 		}
 	}
 	else
 	{
-		VkImageLayout& DstLayout = LayoutManager.FindOrAddLayoutRW(DstSurface, VK_IMAGE_LAYOUT_UNDEFINED);
+		VkImageLayout DstLayout = FVulkanLayoutManager::SetExpectedLayout(CmdBuffer, DstSurface, ERHIAccess::CopyDest);
+
 		if (SrcSurface.Image != DstSurface.Image)
 		{
 			const bool bIsResolve = SrcSurface.GetNumSamples() > DstSurface.GetNumSamples();
@@ -274,19 +273,13 @@ void FVulkanCommandListContext::RHICopyToResolveTarget(FRHITexture* SourceTextur
 			FVulkanPipelineBarrier BarrierBefore;
 		
 			// We'll transition the entire resources to the correct copy states, so we don't need to worry about sub-resource states.
-			if (SrcLayout != VK_IMAGE_LAYOUT_TRANSFER_SRC_OPTIMAL)
-			{
-				BarrierBefore.AddImageLayoutTransition(SrcSurface.Image, SrcLayout, VK_IMAGE_LAYOUT_TRANSFER_SRC_OPTIMAL, FVulkanPipelineBarrier::MakeSubresourceRange(SrcSurface.GetFullAspectMask()));
-				SrcLayout = VK_IMAGE_LAYOUT_TRANSFER_SRC_OPTIMAL;
-			}
-
 			if (DstLayout != VK_IMAGE_LAYOUT_TRANSFER_DST_OPTIMAL)
 			{
 				BarrierBefore.AddImageLayoutTransition(DstSurface.Image, DstLayout, VK_IMAGE_LAYOUT_TRANSFER_DST_OPTIMAL, FVulkanPipelineBarrier::MakeSubresourceRange(DstSurface.GetFullAspectMask()));
 				DstLayout = VK_IMAGE_LAYOUT_TRANSFER_DST_OPTIMAL;
 			}
 
-			BarrierBefore.Execute(CmdBuffer->GetHandle());
+			BarrierBefore.Execute(CmdBuffer);
 
 			if (!bIsResolve)
 			{
@@ -307,30 +300,25 @@ void FVulkanCommandListContext::RHICopyToResolveTarget(FRHITexture* SourceTextur
 					DstSurface.Image, VK_IMAGE_LAYOUT_TRANSFER_DST_OPTIMAL,
 					1, &Region);
 			}
-
-			SrcCurrentAccess = ERHIAccess::CopySrc;
-			DstCurrentAccess = ERHIAccess::CopyDest;
-		}
-		else
-		{
-			SrcCurrentAccess = ERHIAccess::Unknown;
-			DstCurrentAccess = ERHIAccess::Unknown;
 		}
 
-		if (InResolveParams.SourceAccessFinal != ERHIAccess::Unknown && InResolveParams.DestAccessFinal != ERHIAccess::Unknown)
+		if ((InResolveParams.SourceAccessFinal != ERHIAccess::Unknown) || (InResolveParams.DestAccessFinal != ERHIAccess::Unknown))
 		{
 			FVulkanPipelineBarrier BarrierAfter;
-			if (SrcSurface.Image != DstSurface.Image && SrcCurrentAccess != InResolveParams.SourceAccessFinal && InResolveParams.SourceAccessFinal != ERHIAccess::Unknown)
+
+			if ((SrcSurface.Image != DstSurface.Image) && (InResolveParams.SourceAccessFinal != ERHIAccess::CopySrc) && (InResolveParams.SourceAccessFinal != ERHIAccess::Unknown))
 			{
-				BarrierAfter.AddImageAccessTransition(SrcSurface, SrcCurrentAccess, InResolveParams.SourceAccessFinal, FVulkanPipelineBarrier::MakeSubresourceRange(SrcSurface.GetFullAspectMask()), SrcLayout);
+				BarrierAfter.AddImageAccessTransition(SrcSurface, ERHIAccess::CopySrc, InResolveParams.SourceAccessFinal, FVulkanPipelineBarrier::MakeSubresourceRange(SrcSurface.GetFullAspectMask()), SrcLayout);
+				CmdBuffer->GetLayoutManager().SetFullLayout(SrcSurface, SrcLayout);
 			}
 
-			if (DstCurrentAccess != InResolveParams.DestAccessFinal && InResolveParams.DestAccessFinal != ERHIAccess::Unknown)
+			if ((InResolveParams.DestAccessFinal != ERHIAccess::CopyDest) && (InResolveParams.DestAccessFinal != ERHIAccess::Unknown))
 			{
-				BarrierAfter.AddImageAccessTransition(DstSurface, DstCurrentAccess, InResolveParams.DestAccessFinal, FVulkanPipelineBarrier::MakeSubresourceRange(DstSurface.GetFullAspectMask()), DstLayout);
+				BarrierAfter.AddImageAccessTransition(DstSurface, ERHIAccess::CopyDest, InResolveParams.DestAccessFinal, FVulkanPipelineBarrier::MakeSubresourceRange(DstSurface.GetFullAspectMask()), DstLayout);
+				CmdBuffer->GetLayoutManager().SetFullLayout(DstSurface, DstLayout);
 			}
 
-			BarrierAfter.Execute(CmdBuffer->GetHandle());
+			BarrierAfter.Execute(CmdBuffer);
 		}
 	}
 
@@ -391,6 +379,8 @@ static void ConvertRawDataToFColor(VkFormat VulkanFormat, uint32 DestWidth, uint
 
 void FVulkanDynamicRHI::RHIReadSurfaceData(FRHITexture* TextureRHI, FIntRect Rect, TArray<FColor>& OutData, FReadSurfaceDataFlags InFlags)
 {
+	checkf((InFlags.GetCubeFace() == CubeFace_MAX), TEXT("Cube faces not supported yet."));
+
 	const uint32 DestWidth = Rect.Max.X - Rect.Min.X;
 	const uint32 DestHeight = Rect.Max.Y - Rect.Min.Y;
 	const uint32 NumRequestedPixels = DestWidth * DestHeight;
@@ -448,6 +438,7 @@ void FVulkanDynamicRHI::RHIReadSurfaceData(FRHITexture* TextureRHI, FIntRect Rec
 		// Leave bufferRowLength/bufferImageHeight at 0 for tightly packed
 		CopyRegion.imageSubresource.aspectMask = Surface.GetFullAspectMask();
 		CopyRegion.imageSubresource.mipLevel = MipLevel;
+		CopyRegion.imageSubresource.baseArrayLayer = InFlags.GetArrayIndex();
 		CopyRegion.imageSubresource.layerCount = 1;
 		CopyRegion.imageOffset.x = Rect.Min.X;
 		CopyRegion.imageOffset.y = Rect.Min.Y;
@@ -455,26 +446,23 @@ void FVulkanDynamicRHI::RHIReadSurfaceData(FRHITexture* TextureRHI, FIntRect Rec
 		CopyRegion.imageExtent.height = DestHeight;
 		CopyRegion.imageExtent.depth = 1;
 
-		VkImageLayout& CurrentLayout = Device->GetImmediateContext().GetLayoutManager().FindOrAddLayoutRW(Surface, VK_IMAGE_LAYOUT_UNDEFINED);
-		const bool bHadLayout = (CurrentLayout != VK_IMAGE_LAYOUT_UNDEFINED);
-		if (CurrentLayout != VK_IMAGE_LAYOUT_TRANSFER_SRC_OPTIMAL)
-		{
-			VulkanSetImageLayoutAllMips(CmdBuffer->GetHandle(), Surface.Image, CurrentLayout, VK_IMAGE_LAYOUT_TRANSFER_SRC_OPTIMAL);
-		}
+		const VkImageLayout PreviousLayout = FVulkanLayoutManager::SetExpectedLayout(CmdBuffer, Surface, ERHIAccess::CopySrc);
 
 		VulkanRHI::vkCmdCopyImageToBuffer(CmdBuffer->GetHandle(), Surface.Image, VK_IMAGE_LAYOUT_TRANSFER_SRC_OPTIMAL, StagingBuffer->GetHandle(), 1, &CopyRegion);
-		if (bHadLayout && CurrentLayout != VK_IMAGE_LAYOUT_TRANSFER_SRC_OPTIMAL)
+
+		FVulkanPipelineBarrier AfterBarrier;
+		if ((PreviousLayout != VK_IMAGE_LAYOUT_UNDEFINED) && (PreviousLayout != VK_IMAGE_LAYOUT_TRANSFER_SRC_OPTIMAL))
 		{
-			VulkanSetImageLayoutAllMips(CmdBuffer->GetHandle(), Surface.Image, VK_IMAGE_LAYOUT_TRANSFER_SRC_OPTIMAL, CurrentLayout);
+			AfterBarrier.AddImageLayoutTransition(Surface.Image, VK_IMAGE_LAYOUT_TRANSFER_SRC_OPTIMAL, PreviousLayout, FVulkanPipelineBarrier::MakeSubresourceRange(Surface.GetFullAspectMask()));
 		}
 		else
 		{
-			CurrentLayout = VK_IMAGE_LAYOUT_TRANSFER_SRC_OPTIMAL;
+			CmdBuffer->GetLayoutManager().SetFullLayout(Surface, VK_IMAGE_LAYOUT_TRANSFER_SRC_OPTIMAL);
 		}
 		ensure(StagingBuffer->GetSize() >= BufferSize);
 
-		VkMemoryBarrier Barrier = { VK_STRUCTURE_TYPE_MEMORY_BARRIER , nullptr, VK_ACCESS_TRANSFER_WRITE_BIT, VK_ACCESS_HOST_READ_BIT };
-		VulkanRHI::vkCmdPipelineBarrier(CmdBuffer->GetHandle(), VK_PIPELINE_STAGE_TRANSFER_BIT, VK_PIPELINE_STAGE_HOST_BIT, 0, 1, &Barrier, 0, nullptr, 0, nullptr);
+		AfterBarrier.AddMemoryBarrier(VK_ACCESS_TRANSFER_WRITE_BIT, VK_ACCESS_HOST_READ_BIT, VK_PIPELINE_STAGE_TRANSFER_BIT, VK_PIPELINE_STAGE_HOST_BIT);
+		AfterBarrier.Execute(CmdBuffer);
 
 		// Force upload
 		ImmediateContext.GetCommandBufferManager()->SubmitUploadCmdBuffer();
@@ -567,10 +555,11 @@ void FVulkanDynamicRHI::RHIReadSurfaceFloatData(FRHITexture* TextureRHI, FIntRec
 
 		const FRHITextureDesc& Desc = Surface.GetDesc();
 
-		uint32 NumPixels = (Desc.Extent.X >> InMipIndex) * (Desc.Extent.Y >> InMipIndex);
+		const uint32 NumPixels = (Desc.Extent.X >> InMipIndex) * (Desc.Extent.Y >> InMipIndex);
 		const uint32 Size = NumPixels * sizeof(FFloat16Color);
 		VulkanRHI::FStagingBuffer* StagingBuffer = InDevice->GetStagingManager().AcquireBuffer(Size, VK_BUFFER_USAGE_TRANSFER_DST_BIT, VK_MEMORY_PROPERTY_HOST_CACHED_BIT);
 
+		FVulkanPipelineBarrier AfterBarrier;
 		if (GIgnoreCPUReads == 0)
 		{
 			VkBufferImageCopy CopyRegion;
@@ -586,23 +575,18 @@ void FVulkanDynamicRHI::RHIReadSurfaceFloatData(FRHITexture* TextureRHI, FIntRec
 			CopyRegion.imageExtent.height = Desc.Extent.Y >> InMipIndex;
 			CopyRegion.imageExtent.depth = 1;
 
-			//#todo-rco: Multithreaded!
-			VkImageLayout& CurrentLayout = InDevice->GetImmediateContext().GetLayoutManager().FindOrAddLayoutRW(Surface, VK_IMAGE_LAYOUT_UNDEFINED);
-			bool bHadLayout = (CurrentLayout != VK_IMAGE_LAYOUT_UNDEFINED);
-			if (CurrentLayout != VK_IMAGE_LAYOUT_TRANSFER_SRC_OPTIMAL)
-			{
-				VulkanSetImageLayoutSimple(InCmdBuffer->GetHandle(), Surface.Image, CurrentLayout, VK_IMAGE_LAYOUT_TRANSFER_SRC_OPTIMAL);
-			}
+			const VkImageLayout OriginalLayout = FVulkanLayoutManager::SetExpectedLayout(InCmdBuffer, Surface, ERHIAccess::CopySrc);
 
 			VulkanRHI::vkCmdCopyImageToBuffer(InCmdBuffer->GetHandle(), Surface.Image, VK_IMAGE_LAYOUT_TRANSFER_SRC_OPTIMAL, StagingBuffer->GetHandle(), 1, &CopyRegion);
 
-			if (bHadLayout && CurrentLayout != VK_IMAGE_LAYOUT_TRANSFER_SRC_OPTIMAL)
+			if ((OriginalLayout != VK_IMAGE_LAYOUT_UNDEFINED) && (OriginalLayout != VK_IMAGE_LAYOUT_TRANSFER_SRC_OPTIMAL))
 			{
-				VulkanSetImageLayoutSimple(InCmdBuffer->GetHandle(), Surface.Image, VK_IMAGE_LAYOUT_TRANSFER_SRC_OPTIMAL, CurrentLayout);
+				const VkImageSubresourceRange ImageSubresource = FVulkanPipelineBarrier::MakeSubresourceRange(Surface.GetFullAspectMask());
+				AfterBarrier.AddImageLayoutTransition(Surface.Image, VK_IMAGE_LAYOUT_TRANSFER_SRC_OPTIMAL, OriginalLayout, ImageSubresource);
 			}
 			else
 			{
-				CurrentLayout = VK_IMAGE_LAYOUT_TRANSFER_SRC_OPTIMAL;
+				InCmdBuffer->GetLayoutManager().SetFullLayout(Surface, VK_IMAGE_LAYOUT_TRANSFER_SRC_OPTIMAL);
 			}
 		}
 		else
@@ -613,8 +597,8 @@ void FVulkanDynamicRHI::RHIReadSurfaceFloatData(FRHITexture* TextureRHI, FIntRec
 		// the staging buffer size may be bigger then the size due to alignment, etc. but it must not be smaller!
 		ensure(StagingBuffer->GetSize() >= Size);
 
-		VkMemoryBarrier Barrier = { VK_STRUCTURE_TYPE_MEMORY_BARRIER , nullptr, VK_ACCESS_TRANSFER_WRITE_BIT, VK_ACCESS_HOST_READ_BIT };
-		VulkanRHI::vkCmdPipelineBarrier(InCmdBuffer->GetHandle(), VK_PIPELINE_STAGE_TRANSFER_BIT, VK_PIPELINE_STAGE_HOST_BIT, 0, 1, &Barrier, 0, nullptr, 0, nullptr);
+		AfterBarrier.AddMemoryBarrier(VK_ACCESS_TRANSFER_WRITE_BIT, VK_ACCESS_HOST_READ_BIT, VK_PIPELINE_STAGE_TRANSFER_BIT, VK_PIPELINE_STAGE_HOST_BIT);
+		AfterBarrier.Execute(InCmdBuffer);
 
 		// Force upload
 		InDevice->GetImmediateContext().GetCommandBufferManager()->SubmitUploadCmdBuffer();
@@ -695,10 +679,10 @@ void FVulkanDynamicRHI::RHIRead3DSurfaceFloatData(FRHITexture* TextureRHI, FIntR
 	FVulkanTexture& Surface = *FVulkanTexture::Cast(TextureRHI);
 	const FRHITextureDesc& Desc = Surface.GetDesc();
 
-	uint32 SizeX = InRect.Width();
-	uint32 SizeY = InRect.Height();
-	uint32 SizeZ = ZMinMax.Y - ZMinMax.X;
-	uint32 NumPixels = SizeX * SizeY * SizeZ;
+	const uint32 SizeX = InRect.Width();
+	const uint32 SizeY = InRect.Height();
+	const uint32 SizeZ = ZMinMax.Y - ZMinMax.X;
+	const uint32 NumPixels = SizeX * SizeY * SizeZ;
 	const uint32 Size = NumPixels * sizeof(FFloat16Color);
 
 	// Allocate the output buffer.
@@ -717,6 +701,8 @@ void FVulkanDynamicRHI::RHIRead3DSurfaceFloatData(FRHITexture* TextureRHI, FIntR
 	ensure(Surface.StorageFormat == VK_FORMAT_R16G16B16A16_SFLOAT);
 
 	VulkanRHI::FStagingBuffer* StagingBuffer = Device->GetStagingManager().AcquireBuffer(Size, VK_BUFFER_USAGE_TRANSFER_DST_BIT, VK_MEMORY_PROPERTY_HOST_CACHED_BIT);
+
+	FVulkanPipelineBarrier AfterBarrier;
 	if (GIgnoreCPUReads == 0)
 	{
 		VkBufferImageCopy CopyRegion;
@@ -736,22 +722,17 @@ void FVulkanDynamicRHI::RHIRead3DSurfaceFloatData(FRHITexture* TextureRHI, FIntR
 		CopyRegion.imageExtent.depth = SizeZ;
 
 		//#todo-rco: Multithreaded!
-		VkImageLayout& CurrentLayout = Device->GetImmediateContext().GetLayoutManager().FindOrAddLayoutRW(Surface, VK_IMAGE_LAYOUT_UNDEFINED);
-		bool bHadLayout = (CurrentLayout != VK_IMAGE_LAYOUT_UNDEFINED);
-		if (CurrentLayout != VK_IMAGE_LAYOUT_TRANSFER_SRC_OPTIMAL)
-		{
-			VulkanSetImageLayoutSimple(CmdBuffer->GetHandle(), Surface.Image, CurrentLayout, VK_IMAGE_LAYOUT_TRANSFER_SRC_OPTIMAL);
-		}
+		const VkImageLayout OriginalLayout = FVulkanLayoutManager::SetExpectedLayout(CmdBuffer, Surface, ERHIAccess::CopySrc);
 
 		VulkanRHI::vkCmdCopyImageToBuffer(CmdBuffer->GetHandle(), Surface.Image, VK_IMAGE_LAYOUT_TRANSFER_SRC_OPTIMAL, StagingBuffer->GetHandle(), 1, &CopyRegion);
 
-		if (bHadLayout && CurrentLayout != VK_IMAGE_LAYOUT_TRANSFER_SRC_OPTIMAL)
+		if ((OriginalLayout != VK_IMAGE_LAYOUT_UNDEFINED) && (OriginalLayout != VK_IMAGE_LAYOUT_TRANSFER_SRC_OPTIMAL))
 		{
-			VulkanSetImageLayoutSimple(CmdBuffer->GetHandle(), Surface.Image, VK_IMAGE_LAYOUT_TRANSFER_SRC_OPTIMAL, CurrentLayout);
+			AfterBarrier.AddImageLayoutTransition(Surface.Image, VK_IMAGE_LAYOUT_TRANSFER_SRC_OPTIMAL, OriginalLayout, FVulkanPipelineBarrier::MakeSubresourceRange(Surface.GetFullAspectMask()));
 		}
 		else
 		{
-			CurrentLayout = VK_IMAGE_LAYOUT_TRANSFER_SRC_OPTIMAL;
+			CmdBuffer->GetLayoutManager().SetFullLayout(Surface, VK_IMAGE_LAYOUT_TRANSFER_SRC_OPTIMAL);
 		}
 	}
 	else
@@ -762,8 +743,8 @@ void FVulkanDynamicRHI::RHIRead3DSurfaceFloatData(FRHITexture* TextureRHI, FIntR
 	// the staging buffer size may be bigger then the size due to alignment, etc. but it must not be smaller!
 	ensure(StagingBuffer->GetSize() >= Size);
 
-	VkMemoryBarrier Barrier = { VK_STRUCTURE_TYPE_MEMORY_BARRIER , nullptr, VK_ACCESS_TRANSFER_WRITE_BIT, VK_ACCESS_HOST_READ_BIT };
-	VulkanRHI::vkCmdPipelineBarrier(CmdBuffer->GetHandle(), VK_PIPELINE_STAGE_TRANSFER_BIT, VK_PIPELINE_STAGE_HOST_BIT, 0, 1, &Barrier, 0, nullptr, 0, nullptr);
+	AfterBarrier.AddMemoryBarrier(VK_ACCESS_TRANSFER_WRITE_BIT, VK_ACCESS_HOST_READ_BIT, VK_PIPELINE_STAGE_TRANSFER_BIT, VK_PIPELINE_STAGE_HOST_BIT);
+	AfterBarrier.Execute(CmdBuffer);
 
 	// Force upload
 	Device->GetImmediateContext().GetCommandBufferManager()->SubmitUploadCmdBuffer();
@@ -871,47 +852,128 @@ void FVulkanCommandListContext::RHIBeginRenderPass(const FRHIRenderPassInfo& InI
 	}
 
 	FRHITexture* DSTexture = InInfo.DepthStencilRenderTarget.DepthStencilTarget;
-	VkImageLayout CurrentDSLayout;
+	VkImageLayout CurrentDepthLayout = VK_IMAGE_LAYOUT_UNDEFINED;
+	VkImageLayout CurrentStencilLayout = VK_IMAGE_LAYOUT_UNDEFINED;
 	if (DSTexture)
 	{
-		FVulkanTexture& Surface = *FVulkanTexture::Cast(DSTexture);
-		CurrentDSLayout = LayoutManager.FindLayoutChecked(Surface.Image);
-	}
-	else
-	{
-		CurrentDSLayout = VK_IMAGE_LAYOUT_UNDEFINED;
+		FVulkanTexture& VulkanTexture = *FVulkanTexture::Cast(DSTexture);
+		const VkImageAspectFlags AspectFlags = VulkanTexture.GetFullAspectMask();
+
+		if (GetDevice()->SupportsParallelRendering())
+		{
+			const FExclusiveDepthStencil ExclusiveDepthStencil = InInfo.DepthStencilRenderTarget.ExclusiveDepthStencil;
+			if (VKHasAnyFlags(AspectFlags, VK_IMAGE_ASPECT_DEPTH_BIT))
+			{
+				if (ExclusiveDepthStencil.IsDepthWrite())
+				{
+					CurrentDepthLayout = VK_IMAGE_LAYOUT_ATTACHMENT_OPTIMAL;
+				}
+				else
+				{
+					// todo-jn: temporarily use tracking illegally because sometimes depth is left in ATTACHMENT_OPTIMAL for Read passes
+					CurrentDepthLayout = CmdBuffer->GetLayoutManager().GetDepthStencilHint(VulkanTexture, VK_IMAGE_ASPECT_DEPTH_BIT);
+					if (CurrentDepthLayout == VK_IMAGE_LAYOUT_UNDEFINED)
+					{
+						CurrentDepthLayout = GetQueue()->GetLayoutManager().GetDepthStencilHint(VulkanTexture, VK_IMAGE_ASPECT_DEPTH_BIT);
+						if (CurrentDepthLayout == VK_IMAGE_LAYOUT_UNDEFINED)
+						{
+							if (ExclusiveDepthStencil.IsDepthRead())
+							{
+								CurrentDepthLayout = VK_IMAGE_LAYOUT_READ_ONLY_OPTIMAL;
+							}
+							else
+							{
+								UE_LOG(LogVulkanRHI, Warning, TEXT("Render pass [%s] used an unknown depth state!"), InName ? InName : TEXT("unknown"));
+							}
+						}
+					}
+				}
+			}
+
+			if (VKHasAnyFlags(AspectFlags, VK_IMAGE_ASPECT_STENCIL_BIT))
+			{
+				if (ExclusiveDepthStencil.IsStencilWrite())
+				{
+					CurrentStencilLayout = VK_IMAGE_LAYOUT_ATTACHMENT_OPTIMAL;
+				}
+				else if (ExclusiveDepthStencil.IsStencilRead())
+				{
+					CurrentStencilLayout = VK_IMAGE_LAYOUT_READ_ONLY_OPTIMAL;
+				}
+				else
+				{
+					// todo-jn: temporarily use tracking illegally because stencil is no-op so we can't know expected layout
+					CurrentStencilLayout = CmdBuffer->GetLayoutManager().GetDepthStencilHint(VulkanTexture, VK_IMAGE_ASPECT_STENCIL_BIT);
+					if (CurrentStencilLayout == VK_IMAGE_LAYOUT_UNDEFINED)
+					{
+						CurrentStencilLayout = GetQueue()->GetLayoutManager().GetDepthStencilHint(VulkanTexture, VK_IMAGE_ASPECT_STENCIL_BIT);
+						if (CurrentStencilLayout == VK_IMAGE_LAYOUT_UNDEFINED)
+						{
+							UE_LOG(LogVulkanRHI, Warning, TEXT("Render pass [%s] used an unknown stencil state!"), InName ? InName : TEXT("unknown"));
+						}
+					}
+				}
+			}
+		}
+		else
+		{
+			const FVulkanImageLayout* FullLayout = CmdBuffer->GetLayoutManager().GetFullLayout(VulkanTexture);
+			check(FullLayout);
+			if (VKHasAnyFlags(AspectFlags, VK_IMAGE_ASPECT_DEPTH_BIT))
+			{
+				CurrentDepthLayout = FullLayout->GetSubresLayout(0, 0, VK_IMAGE_ASPECT_DEPTH_BIT);
+			}
+			if (VKHasAnyFlags(AspectFlags, VK_IMAGE_ASPECT_STENCIL_BIT))
+			{
+				CurrentStencilLayout = FullLayout->GetSubresLayout(0, 0, VK_IMAGE_ASPECT_STENCIL_BIT);
+			}
+		}
 	}
 
-	FVulkanRenderTargetLayout RTLayout(*Device, InInfo, CurrentDSLayout);
+	FVulkanRenderTargetLayout RTLayout(*Device, InInfo, CurrentDepthLayout, CurrentStencilLayout);
 	check(RTLayout.GetExtent2D().width != 0 && RTLayout.GetExtent2D().height != 0);
 
-	FVulkanRenderPass* RenderPass = LayoutManager.GetOrCreateRenderPass(*Device, RTLayout);
+	FVulkanRenderPass* RenderPass = Device->GetRenderPassManager().GetOrCreateRenderPass(RTLayout);
 	FRHISetRenderTargetsInfo RTInfo;
 	InInfo.ConvertToRenderTargetsInfo(RTInfo);
 
-	FVulkanFramebuffer* Framebuffer = LayoutManager.GetOrCreateFramebuffer(*Device, RTInfo, RTLayout, RenderPass);
+	FVulkanFramebuffer* Framebuffer = Device->GetRenderPassManager().GetOrCreateFramebuffer(RTInfo, RTLayout, RenderPass);
 	checkf(RenderPass != nullptr && Framebuffer != nullptr, TEXT("RenderPass not started! Bad combination of values? Depth %p #Color %d Color0 %p"), (void*)InInfo.DepthStencilRenderTarget.DepthStencilTarget, InInfo.GetNumColorRenderTargets(), (void*)InInfo.ColorRenderTargets[0].RenderTarget);
-	LayoutManager.BeginRenderPass(*this, *Device, CmdBuffer, InInfo, RTLayout, RenderPass, Framebuffer);
+	Device->GetRenderPassManager().BeginRenderPass(*this, *Device, CmdBuffer, InInfo, RTLayout, RenderPass, Framebuffer);
+
+	check(!CurrentRenderPass);
+	CurrentRenderPass = RenderPass;
+	CurrentFramebuffer = Framebuffer;
 }
 
 void FVulkanCommandListContext::RHIEndRenderPass()
 {
+	const bool bHasOcclusionQueries = (RenderPassInfo.NumOcclusionQueries > 0);
 	FVulkanCmdBuffer* CmdBuffer = CommandBufferManager->GetActiveCmdBuffer();
-	if (RenderPassInfo.NumOcclusionQueries > 0)
+
+	if (bHasOcclusionQueries)
 	{
 		EndOcclusionQueryBatch(CmdBuffer);
 	}
-	else
-	{
-		LayoutManager.EndRenderPass(CmdBuffer);
-	}
+	
+	Device->GetRenderPassManager().EndRenderPass(CmdBuffer);
+
+	check(CurrentRenderPass);
+	CurrentRenderPass = nullptr;
 
 	RHIPopEvent();
+
+	// Sync point for passes with occlusion queries
+	if (bHasOcclusionQueries && GSubmitOcclusionBatchCmdBufferCVar.GetValueOnAnyThread())
+	{
+		RequestSubmitCurrentCommands();
+		SafePointSubmit();
+	}
 }
 
 void FVulkanCommandListContext::RHINextSubpass()
 {
-	check(LayoutManager.CurrentRenderPass);
+	check(CurrentRenderPass);
 	FVulkanCmdBuffer* CmdBuffer = CommandBufferManager->GetActiveCmdBuffer();
 	VkCommandBuffer Cmd = CmdBuffer->GetHandle();
 	VulkanRHI::vkCmdNextSubpass(Cmd, VK_SUBPASS_CONTENTS_INLINE);
@@ -930,8 +992,8 @@ struct FRenderPassCompatibleHashableStruct
 	uint8							NumSamples;
 	uint8							SubpassHint;
 	VkSurfaceTransformFlagBitsKHR	QCOMRenderPassTransform;
-	// +1 for DepthStencil, +1 for Fragment Density
-	VkFormat						Formats[MaxSimultaneousRenderTargets + 2];
+	// +1 for Depth, +1 for Stencil, +1 for Fragment Density
+	VkFormat						Formats[MaxSimultaneousRenderTargets + 3];
 	uint16							AttachmentsToResolve;
 };
 
@@ -948,9 +1010,9 @@ struct FRenderPassFullHashableStruct
 	TEnumAsByte<VkAttachmentStoreOp>	StoreOps[MaxSimultaneousRenderTargets + 3];
 #if VULKAN_USE_REAL_RENDERPASS_COMPATIBILITY
 	// If the initial != final we need to add FinalLayout and potentially RefLayout
-	VkImageLayout						InitialLayout[MaxSimultaneousRenderTargets + 2];
-	//VkImageLayout						FinalLayout[MaxSimultaneousRenderTargets + 2];
-	//VkImageLayout						RefLayout[MaxSimultaneousRenderTargets + 2];
+	VkImageLayout						InitialLayout[MaxSimultaneousRenderTargets + 3];
+	//VkImageLayout						FinalLayout[MaxSimultaneousRenderTargets + 3];
+	//VkImageLayout						RefLayout[MaxSimultaneousRenderTargets + 3];
 #endif
 };
 
@@ -985,13 +1047,7 @@ FVulkanRenderTargetLayout::FVulkanRenderTargetLayout(FVulkanDevice& InDevice, co
 	, NumUsedClearValues(0)
 	, MultiViewCount(0)
 {
-	FMemory::Memzero(ColorReferences);
-	FMemory::Memzero(DepthStencilReference);
-	FMemory::Memzero(FragmentDensityReference);
-	FMemory::Memzero(ResolveReferences);
-	FMemory::Memzero(InputAttachments);
-	FMemory::Memzero(Desc);
-	FMemory::Memzero(Extent);
+	ResetAttachments();
 
 	FRenderPassCompatibleHashableStruct CompatibleHashInfo;
 	FRenderPassFullHashableStruct FullHashInfo;
@@ -1078,7 +1134,6 @@ FVulkanRenderTargetLayout::FVulkanRenderTargetLayout(FVulkanDevice& InDevice, co
 		}
 	}
 
-	VkImageLayout DepthStencilLayout = VK_IMAGE_LAYOUT_UNDEFINED;
 	if (RTInfo.DepthStencilRenderTarget.Texture)
 	{
 		VkAttachmentDescription& CurrDesc = Desc[NumAttachmentDescriptions];
@@ -1114,21 +1169,26 @@ FVulkanRenderTargetLayout::FVulkanRenderTargetLayout(FVulkanDevice& InDevice, co
 			CurrDesc.stencilStoreOp = VK_ATTACHMENT_STORE_OP_DONT_CARE;
 		}
 
-		DepthStencilLayout = VulkanRHI::GetDepthStencilLayout(RTInfo.DepthStencilRenderTarget.GetDepthStencilAccess(), InDevice);
+		const VkImageLayout DepthLayout = RTInfo.DepthStencilRenderTarget.GetDepthStencilAccess().IsDepthWrite() ? VK_IMAGE_LAYOUT_ATTACHMENT_OPTIMAL : VK_IMAGE_LAYOUT_READ_ONLY_OPTIMAL;
+		const VkImageLayout StencilLayout = RTInfo.DepthStencilRenderTarget.GetDepthStencilAccess().IsStencilWrite() ? VK_IMAGE_LAYOUT_ATTACHMENT_OPTIMAL : VK_IMAGE_LAYOUT_READ_ONLY_OPTIMAL;
 
 		// If the initial != final we need to change the FullHashInfo and use FinalLayout
-		CurrDesc.initialLayout = DepthStencilLayout;
-		CurrDesc.finalLayout = DepthStencilLayout;
+		CurrDesc.initialLayout = DepthLayout;
+		CurrDesc.finalLayout = DepthLayout;
+		StencilDesc.stencilInitialLayout = StencilLayout;
+		StencilDesc.stencilFinalLayout = StencilLayout;
 
-		DepthStencilReference.attachment = NumAttachmentDescriptions;
-		DepthStencilReference.layout = DepthStencilLayout;
+		DepthReference.attachment = NumAttachmentDescriptions;
+		DepthReference.layout = DepthLayout;
+		StencilReference.stencilLayout = StencilLayout;
 
 		FullHashInfo.LoadOps[MaxSimultaneousRenderTargets] = CurrDesc.loadOp;
 		FullHashInfo.LoadOps[MaxSimultaneousRenderTargets + 1] = CurrDesc.stencilLoadOp;
 		FullHashInfo.StoreOps[MaxSimultaneousRenderTargets] = CurrDesc.storeOp;
 		FullHashInfo.StoreOps[MaxSimultaneousRenderTargets + 1] = CurrDesc.stencilStoreOp;
 #if VULKAN_USE_REAL_RENDERPASS_COMPATIBILITY
-		FullHashInfo.InitialLayout[MaxSimultaneousRenderTargets] = DepthStencilLayout;
+		FullHashInfo.InitialLayout[MaxSimultaneousRenderTargets] = DepthLayout;
+		FullHashInfo.InitialLayout[MaxSimultaneousRenderTargets + 1] = StencilLayout;
 #endif
 		CompatibleHashInfo.Formats[MaxSimultaneousRenderTargets] = CurrDesc.format;
 
@@ -1177,7 +1237,7 @@ FVulkanRenderTargetLayout::FVulkanRenderTargetLayout(FVulkanDevice& InDevice, co
 		FullHashInfo.LoadOps[MaxSimultaneousRenderTargets + 2] = CurrDesc.stencilLoadOp;
 		FullHashInfo.StoreOps[MaxSimultaneousRenderTargets + 2] = CurrDesc.stencilStoreOp;
 #if VULKAN_USE_REAL_RENDERPASS_COMPATIBILITY
-		FullHashInfo.InitialLayout[MaxSimultaneousRenderTargets + 1] = VRSLayout;
+		FullHashInfo.InitialLayout[MaxSimultaneousRenderTargets + 2] = VRSLayout;
 #endif
 		CompatibleHashInfo.Formats[MaxSimultaneousRenderTargets + 1] = CurrDesc.format;
 
@@ -1200,7 +1260,7 @@ FVulkanRenderTargetLayout::FVulkanRenderTargetLayout(FVulkanDevice& InDevice, co
 }
 
 
-FVulkanRenderTargetLayout::FVulkanRenderTargetLayout(FVulkanDevice& InDevice, const FRHIRenderPassInfo& RPInfo, VkImageLayout CurrentDSLayout)
+FVulkanRenderTargetLayout::FVulkanRenderTargetLayout(FVulkanDevice& InDevice, const FRHIRenderPassInfo& RPInfo, VkImageLayout CurrentDepthLayout, VkImageLayout CurrentStencilLayout)
 	: NumAttachmentDescriptions(0)
 	, NumColorAttachments(0)
 	, bHasDepthStencil(false)
@@ -1210,14 +1270,7 @@ FVulkanRenderTargetLayout::FVulkanRenderTargetLayout(FVulkanDevice& InDevice, co
 	, NumUsedClearValues(0)
 	, MultiViewCount(RPInfo.MultiViewCount)
 {
-	FMemory::Memzero(ColorReferences);
-	FMemory::Memzero(DepthStencilReference);
-	FMemory::Memzero(FragmentDensityReference);
-	FMemory::Memzero(ResolveReferences);
-	FMemory::Memzero(InputAttachments);
-	FMemory::Memzero(Desc);
-	FMemory::Memzero(Offset);
-	FMemory::Memzero(Extent);
+	ResetAttachments();
 
 	FRenderPassCompatibleHashableStruct CompatibleHashInfo;
 	FRenderPassFullHashableStruct FullHashInfo;
@@ -1306,7 +1359,6 @@ FVulkanRenderTargetLayout::FVulkanRenderTargetLayout(FVulkanDevice& InDevice, co
 		++NumColorAttachments;
 	}
 
-	VkImageLayout DepthStencilLayout = VK_IMAGE_LAYOUT_UNDEFINED;
 	if (RPInfo.DepthStencilRenderTarget.DepthStencilTarget)
 	{
 		VkAttachmentDescription& CurrDesc = Desc[NumAttachmentDescriptions];
@@ -1343,66 +1395,46 @@ FVulkanRenderTargetLayout::FVulkanRenderTargetLayout(FVulkanDevice& InDevice, co
 			RPInfo.DepthStencilRenderTarget.Action == EDepthStencilTargetActions::LoadDepthClearStencil_StoreDepthStencil)
 		{
 			ExclusiveDepthStencil = FExclusiveDepthStencil::DepthWrite_StencilWrite;
-			CurrentDSLayout = VK_IMAGE_LAYOUT_DEPTH_STENCIL_ATTACHMENT_OPTIMAL;
+			CurrentDepthLayout = VK_IMAGE_LAYOUT_ATTACHMENT_OPTIMAL;
+			CurrentStencilLayout = VK_IMAGE_LAYOUT_ATTACHMENT_OPTIMAL;
 		}
 
-		// Make sure that the requested depth-stencil access is compatible with the current layout of the DS target.
-		const bool bWritableDepth = ExclusiveDepthStencil.IsDepthWrite();
-		const bool bWritableStencil = ExclusiveDepthStencil.IsStencilWrite();
-		switch (CurrentDSLayout)
+		// If the initial != final we need to change the FullHashInfo and use FinalLayout
+		CurrDesc.initialLayout = CurrentDepthLayout;
+		CurrDesc.finalLayout = CurrentDepthLayout;
+		StencilDesc.stencilInitialLayout = CurrentStencilLayout;
+		StencilDesc.stencilFinalLayout = CurrentStencilLayout;
+
+		// We can't have the final layout be UNDEFINED, but it's possible that we get here from a transient texture
+		// where the stencil was never used yet.  We can set the layout to whatever we want, the next transition will
+		// happen from UNDEFINED anyhow.
+		if (CurrentDepthLayout == VK_IMAGE_LAYOUT_UNDEFINED)
 		{
-			case VK_IMAGE_LAYOUT_DEPTH_STENCIL_ATTACHMENT_OPTIMAL:
-				// Writable depth-stencil is compatible with all the requested modes.
-				break;
-
-			case VK_IMAGE_LAYOUT_DEPTH_STENCIL_READ_ONLY_OPTIMAL:
-				// Read-only on both aspects requires the requested access to be read-only.
-				ensureMsgf(!bWritableDepth && !bWritableStencil, TEXT("Both aspects of the DS target are read-only, but the requested mode requires write access: D=%s S=%s."),
-					ExclusiveDepthStencil.IsUsingDepth() ? (ExclusiveDepthStencil.IsDepthWrite() ? TEXT("Write") : TEXT("Read")) : TEXT("Nop"),
-					ExclusiveDepthStencil.IsUsingStencil() ? (ExclusiveDepthStencil.IsStencilWrite() ? TEXT("Write") : TEXT("Read")) : TEXT("Nop")
-				);
-				break;
-
-			case VK_IMAGE_LAYOUT_DEPTH_READ_ONLY_STENCIL_ATTACHMENT_OPTIMAL_KHR:
-				// If only stencil is writable, the requested depth access must be read-only.
-				ensureMsgf(!bWritableDepth, TEXT("The depth aspect is read-only, but the requested mode requires depth writes: D=%s S=%s."),
-					ExclusiveDepthStencil.IsUsingDepth() ? (ExclusiveDepthStencil.IsDepthWrite() ? TEXT("Write") : TEXT("Read")) : TEXT("Nop"),
-					ExclusiveDepthStencil.IsUsingStencil() ? (ExclusiveDepthStencil.IsStencilWrite() ? TEXT("Write") : TEXT("Read")) : TEXT("Nop")
-				);
-				break;
-
-			case VK_IMAGE_LAYOUT_DEPTH_ATTACHMENT_STENCIL_READ_ONLY_OPTIMAL_KHR:
-				// If only depth is writable, the requested stencil access must be read-only.
-				ensureMsgf(!bWritableStencil, TEXT("The stencil aspect is read-only, but the requested mode requires stencil writes: D=%s S=%s."),
-					ExclusiveDepthStencil.IsUsingDepth() ? (ExclusiveDepthStencil.IsDepthWrite() ? TEXT("Write") : TEXT("Read")) : TEXT("Nop"),
-					ExclusiveDepthStencil.IsUsingStencil() ? (ExclusiveDepthStencil.IsStencilWrite() ? TEXT("Write") : TEXT("Read")) : TEXT("Nop")
-				);
-				break;
-
-			default:
-				// Any other layout is invalid when starting a render pass.
-				ensureMsgf(false, TEXT("Depth target is in layout %u, which is invalid for a render pass."), CurrentDSLayout);
-				break;
+			check(CurrDesc.storeOp == VK_ATTACHMENT_STORE_OP_DONT_CARE);
+			CurrDesc.finalLayout = VK_IMAGE_LAYOUT_ATTACHMENT_OPTIMAL;
+		}
+		if (CurrentStencilLayout == VK_IMAGE_LAYOUT_UNDEFINED)
+		{
+			check(CurrDesc.stencilStoreOp == VK_ATTACHMENT_STORE_OP_DONT_CARE);
+			StencilDesc.stencilFinalLayout = VK_IMAGE_LAYOUT_ATTACHMENT_OPTIMAL;
 		}
 
-		DepthStencilLayout = CurrentDSLayout;
+		DepthReference.attachment = NumAttachmentDescriptions;
+		DepthReference.layout = CurrentDepthLayout;
+		StencilReference.stencilLayout = CurrentStencilLayout;
 
-		// If the initial != final we need to change the FullHashInfo and use FinalLayout
-		CurrDesc.initialLayout = DepthStencilLayout;
-		CurrDesc.finalLayout = DepthStencilLayout;
-		DepthStencilReference.attachment = NumAttachmentDescriptions;
-		DepthStencilReference.layout = DepthStencilLayout;
+		++NumAttachmentDescriptions;
 
 		FullHashInfo.LoadOps[MaxSimultaneousRenderTargets] = CurrDesc.loadOp;
 		FullHashInfo.LoadOps[MaxSimultaneousRenderTargets + 1] = CurrDesc.stencilLoadOp;
 		FullHashInfo.StoreOps[MaxSimultaneousRenderTargets] = CurrDesc.storeOp;
 		FullHashInfo.StoreOps[MaxSimultaneousRenderTargets + 1] = CurrDesc.stencilStoreOp;
 #if VULKAN_USE_REAL_RENDERPASS_COMPATIBILITY
-		FullHashInfo.InitialLayout[MaxSimultaneousRenderTargets] = DepthStencilLayout;
+		FullHashInfo.InitialLayout[MaxSimultaneousRenderTargets] = CurrentDepthLayout;
+		FullHashInfo.InitialLayout[MaxSimultaneousRenderTargets + 1] = CurrentStencilLayout;
 #endif
 		CompatibleHashInfo.Formats[MaxSimultaneousRenderTargets] = CurrDesc.format;
 
-		++NumAttachmentDescriptions;
 
 		bHasDepthStencil = true;
 
@@ -1459,7 +1491,7 @@ FVulkanRenderTargetLayout::FVulkanRenderTargetLayout(FVulkanDevice& InDevice, co
 		FullHashInfo.LoadOps[MaxSimultaneousRenderTargets + 2] = CurrDesc.stencilLoadOp;
 		FullHashInfo.StoreOps[MaxSimultaneousRenderTargets + 2] = CurrDesc.stencilStoreOp;
 #if VULKAN_USE_REAL_RENDERPASS_COMPATIBILITY
-		FullHashInfo.InitialLayout[MaxSimultaneousRenderTargets + 1] = VRSLayout;
+		FullHashInfo.InitialLayout[MaxSimultaneousRenderTargets + 2] = VRSLayout;
 #endif
 		CompatibleHashInfo.Formats[MaxSimultaneousRenderTargets + 1] = CurrDesc.format;
 
@@ -1496,13 +1528,7 @@ FVulkanRenderTargetLayout::FVulkanRenderTargetLayout(const FGraphicsPipelineStat
 	, NumUsedClearValues(0)
 	, MultiViewCount(0)
 {
-	FMemory::Memzero(ColorReferences);
-	FMemory::Memzero(DepthStencilReference);
-	FMemory::Memzero(FragmentDensityReference);
-	FMemory::Memzero(ResolveReferences);
-	FMemory::Memzero(InputAttachments);
-	FMemory::Memzero(Desc);
-	FMemory::Memzero(Extent);
+	ResetAttachments();
 
 	FRenderPassCompatibleHashableStruct CompatibleHashInfo;
 	FRenderPassFullHashableStruct FullHashInfo;
@@ -1582,18 +1608,22 @@ FVulkanRenderTargetLayout::FVulkanRenderTargetLayout(const FGraphicsPipelineStat
 		}
 
 		// If the initial != final we need to change the FullHashInfo and use FinalLayout
-		CurrDesc.initialLayout = VK_IMAGE_LAYOUT_DEPTH_STENCIL_ATTACHMENT_OPTIMAL;
-		CurrDesc.finalLayout = VK_IMAGE_LAYOUT_DEPTH_STENCIL_ATTACHMENT_OPTIMAL;
+		CurrDesc.initialLayout = VK_IMAGE_LAYOUT_ATTACHMENT_OPTIMAL;
+		CurrDesc.finalLayout = VK_IMAGE_LAYOUT_ATTACHMENT_OPTIMAL;
+		StencilDesc.stencilInitialLayout = VK_IMAGE_LAYOUT_ATTACHMENT_OPTIMAL;
+		StencilDesc.stencilFinalLayout = VK_IMAGE_LAYOUT_ATTACHMENT_OPTIMAL;
 
-		DepthStencilReference.attachment = NumAttachmentDescriptions;
-		DepthStencilReference.layout = VK_IMAGE_LAYOUT_DEPTH_STENCIL_ATTACHMENT_OPTIMAL;
+		DepthReference.attachment = NumAttachmentDescriptions;
+		DepthReference.layout = VK_IMAGE_LAYOUT_ATTACHMENT_OPTIMAL;
+		StencilReference.stencilLayout = VK_IMAGE_LAYOUT_ATTACHMENT_OPTIMAL;
 
 		FullHashInfo.LoadOps[MaxSimultaneousRenderTargets] = CurrDesc.loadOp;
 		FullHashInfo.LoadOps[MaxSimultaneousRenderTargets + 1] = CurrDesc.stencilLoadOp;
 		FullHashInfo.StoreOps[MaxSimultaneousRenderTargets] = CurrDesc.storeOp;
 		FullHashInfo.StoreOps[MaxSimultaneousRenderTargets + 1] = CurrDesc.stencilStoreOp;
 #if VULKAN_USE_REAL_RENDERPASS_COMPATIBILITY
-		FullHashInfo.InitialLayout[MaxSimultaneousRenderTargets] = VK_IMAGE_LAYOUT_DEPTH_STENCIL_ATTACHMENT_OPTIMAL;
+		FullHashInfo.InitialLayout[MaxSimultaneousRenderTargets] = VK_IMAGE_LAYOUT_ATTACHMENT_OPTIMAL;
+		FullHashInfo.InitialLayout[MaxSimultaneousRenderTargets + 1] = VK_IMAGE_LAYOUT_ATTACHMENT_OPTIMAL;
 #endif
 		CompatibleHashInfo.Formats[MaxSimultaneousRenderTargets] = CurrDesc.format;
 
@@ -1626,7 +1656,7 @@ FVulkanRenderTargetLayout::FVulkanRenderTargetLayout(const FGraphicsPipelineStat
 		FullHashInfo.LoadOps[MaxSimultaneousRenderTargets + 2] = CurrDesc.stencilLoadOp;
 		FullHashInfo.StoreOps[MaxSimultaneousRenderTargets + 2] = CurrDesc.stencilStoreOp;
 #if VULKAN_USE_REAL_RENDERPASS_COMPATIBILITY
-		FullHashInfo.InitialLayout[MaxSimultaneousRenderTargets + 1] = VRSLayout;
+		FullHashInfo.InitialLayout[MaxSimultaneousRenderTargets + 2] = VRSLayout;
 #endif
 		CompatibleHashInfo.Formats[MaxSimultaneousRenderTargets + 1] = CurrDesc.format;
 
diff --git a/Engine/Source/Runtime/VulkanRHI/Private/VulkanRenderpass.cpp b/Engine/Source/Runtime/VulkanRHI/Private/VulkanRenderpass.cpp
--- a/Engine/Source/Runtime/VulkanRHI/Private/VulkanRenderpass.cpp
+++ b/Engine/Source/Runtime/VulkanRHI/Private/VulkanRenderpass.cpp
@@ -10,14 +10,12 @@ VkRenderPass CreateVulkanRenderPass(FVulkanDevice& InDevice, const FVulkanRender
 {
 	VkRenderPass OutRenderpass;
 
-#if VULKAN_SUPPORTS_RENDERPASS2
 	if (InDevice.GetOptionalExtensions().HasKHRRenderPass2)
 	{
 		FVulkanRenderPassBuilder<FVulkanSubpassDescription<VkSubpassDescription2>, FVulkanSubpassDependency<VkSubpassDependency2>, FVulkanAttachmentReference<VkAttachmentReference2>, FVulkanAttachmentDescription<VkAttachmentDescription2>, FVulkanRenderPassCreateInfo<VkRenderPassCreateInfo2>> Creator(InDevice);
 		OutRenderpass = Creator.Create(RTLayout);
 	}
 	else
-#endif
 	{
 		FVulkanRenderPassBuilder<FVulkanSubpassDescription<VkSubpassDescription>, FVulkanSubpassDependency<VkSubpassDependency>, FVulkanAttachmentReference<VkAttachmentReference>, FVulkanAttachmentDescription<VkAttachmentDescription>, FVulkanRenderPassCreateInfo<VkRenderPassCreateInfo>> Creator(InDevice);
 		OutRenderpass = Creator.Create(RTLayout);
@@ -26,3 +24,199 @@ VkRenderPass CreateVulkanRenderPass(FVulkanDevice& InDevice, const FVulkanRender
 	return OutRenderpass;
 }
 
+
+
+FVulkanRenderPassManager::~FVulkanRenderPassManager()
+{
+	check(!GIsRHIInitialized);
+
+	for (auto& Pair : RenderPasses)
+	{
+		delete Pair.Value;
+	}
+
+	for (auto& Pair : Framebuffers)
+	{
+		FFramebufferList* List = Pair.Value;
+		for (int32 Index = List->Framebuffer.Num() - 1; Index >= 0; --Index)
+		{
+			List->Framebuffer[Index]->Destroy(*Device);
+			delete List->Framebuffer[Index];
+		}
+		delete List;
+	}
+
+	RenderPasses.Reset();
+	Framebuffers.Reset();
+}
+
+FVulkanFramebuffer* FVulkanRenderPassManager::GetOrCreateFramebuffer(const FRHISetRenderTargetsInfo& RenderTargetsInfo, const FVulkanRenderTargetLayout& RTLayout, FVulkanRenderPass* RenderPass)
+{
+	// todo-jn: threadsafe?
+
+	uint32 RTLayoutHash = RTLayout.GetRenderPassCompatibleHash();
+
+	uint64 MipsAndSlicesValues[MaxSimultaneousRenderTargets];
+	for (int32 Index = 0; Index < MaxSimultaneousRenderTargets; ++Index)
+	{
+		MipsAndSlicesValues[Index] = ((uint64)RenderTargetsInfo.ColorRenderTarget[Index].ArraySliceIndex << (uint64)32) | (uint64)RenderTargetsInfo.ColorRenderTarget[Index].MipIndex;
+	}
+	RTLayoutHash = FCrc::MemCrc32(MipsAndSlicesValues, sizeof(MipsAndSlicesValues), RTLayoutHash);
+
+
+	FFramebufferList** FoundFramebufferList = nullptr;
+	FFramebufferList* FramebufferList = nullptr;
+
+	{
+		FRWScopeLock ScopedReadLock(FramebuffersLock, SLT_ReadOnly);
+
+		FoundFramebufferList = Framebuffers.Find(RTLayoutHash);
+		if (FoundFramebufferList)
+		{
+			FramebufferList = *FoundFramebufferList;
+
+			for (int32 Index = 0; Index < FramebufferList->Framebuffer.Num(); ++Index)
+			{
+				const VkRect2D RenderArea = FramebufferList->Framebuffer[Index]->GetRenderArea();
+
+				if (FramebufferList->Framebuffer[Index]->Matches(RenderTargetsInfo) &&
+					((RTLayout.GetExtent2D().width == RenderArea.extent.width) && (RTLayout.GetExtent2D().height == RenderArea.extent.height) &&
+						(RTLayout.GetOffset2D().x == RenderArea.offset.x) && (RTLayout.GetOffset2D().y == RenderArea.offset.y)))
+				{
+					return FramebufferList->Framebuffer[Index];
+				}
+			}
+		}
+	}
+
+	FRWScopeLock ScopedWriteLock(FramebuffersLock, SLT_Write);
+	FoundFramebufferList = Framebuffers.Find(RTLayoutHash);
+	if (!FoundFramebufferList)
+	{
+		FramebufferList = new FFramebufferList;
+		Framebuffers.Add(RTLayoutHash, FramebufferList);
+	}
+
+	FVulkanFramebuffer* Framebuffer = new FVulkanFramebuffer(*Device, RenderTargetsInfo, RTLayout, *RenderPass);
+	FramebufferList->Framebuffer.Add(Framebuffer);
+	return Framebuffer;
+}
+
+void FVulkanRenderPassManager::BeginRenderPass(FVulkanCommandListContext& Context, FVulkanDevice& InDevice, FVulkanCmdBuffer* CmdBuffer, const FRHIRenderPassInfo& RPInfo, const FVulkanRenderTargetLayout& RTLayout, FVulkanRenderPass* RenderPass, FVulkanFramebuffer* Framebuffer)
+{
+	// (NumRT + 1 [Depth] ) * 2 [surface + resolve]
+	VkClearValue ClearValues[(MaxSimultaneousRenderTargets + 1) * 2];
+	uint32 ClearValueIndex = 0;
+	bool bNeedsClearValues = RenderPass->GetNumUsedClearValues() > 0;
+	FMemory::Memzero(ClearValues);
+
+	int32 NumColorTargets = RPInfo.GetNumColorRenderTargets();
+	int32 Index = 0;
+
+	FVulkanPipelineBarrier Barrier;
+
+	for (Index = 0; Index < NumColorTargets; ++Index)
+	{
+		const FRHIRenderPassInfo::FColorEntry& ColorEntry = RPInfo.ColorRenderTargets[Index];
+
+		FRHITexture* ColorTexture = ColorEntry.RenderTarget;
+		CA_ASSUME(ColorTexture);
+		FVulkanTexture& ColorSurface = *FVulkanTexture::Cast(ColorTexture);
+		const bool bPassPerformsResolve = ColorSurface.GetNumSamples() > 1 && ColorEntry.ResolveTarget;
+
+		if (GetLoadAction(ColorEntry.Action) == ERenderTargetLoadAction::ELoad)
+		{
+			// Insert a barrier if we're loading from any color targets, to make sure the passes aren't reordered and we end up running before
+			// the pass we're supposed to read from.
+			const VkAccessFlags AccessMask = VK_ACCESS_COLOR_ATTACHMENT_WRITE_BIT;
+			const VkPipelineStageFlags StageMask = VK_PIPELINE_STAGE_COLOR_ATTACHMENT_OUTPUT_BIT;
+			Barrier.AddMemoryBarrier(AccessMask, AccessMask, StageMask, StageMask);
+		}
+
+		if (bNeedsClearValues)
+		{
+			const FLinearColor& ClearColor = ColorTexture->HasClearValue() ? ColorTexture->GetClearColor() : FLinearColor::Black;
+			ClearValues[ClearValueIndex].color.float32[0] = ClearColor.R;
+			ClearValues[ClearValueIndex].color.float32[1] = ClearColor.G;
+			ClearValues[ClearValueIndex].color.float32[2] = ClearColor.B;
+			ClearValues[ClearValueIndex].color.float32[3] = ClearColor.A;
+			++ClearValueIndex;
+			if (bPassPerformsResolve)
+			{
+				++ClearValueIndex;
+			}
+		}
+	}
+
+	FRHITexture* DSTexture = RPInfo.DepthStencilRenderTarget.DepthStencilTarget;
+	if (DSTexture)
+	{
+		const FExclusiveDepthStencil RequestedDSAccess = RPInfo.DepthStencilRenderTarget.ExclusiveDepthStencil;
+		if (RequestedDSAccess.IsDepthRead() || RequestedDSAccess.IsStencilRead())
+		{
+			// If the depth-stencil state doesn't change between passes, the high level code won't perform any transitions.
+			// Make sure we have a barrier in case we're loading depth or stencil, to prevent rearranging passes.
+			const VkAccessFlags AccessMask = VK_ACCESS_DEPTH_STENCIL_ATTACHMENT_READ_BIT | VK_ACCESS_DEPTH_STENCIL_ATTACHMENT_WRITE_BIT;
+			const VkPipelineStageFlags StageMask = VK_PIPELINE_STAGE_EARLY_FRAGMENT_TESTS_BIT | VK_PIPELINE_STAGE_LATE_FRAGMENT_TESTS_BIT;
+			Barrier.AddMemoryBarrier(AccessMask, AccessMask, StageMask, StageMask);
+		}
+
+		if (DSTexture->HasClearValue() && bNeedsClearValues)
+		{
+			float Depth = 0;
+			uint32 Stencil = 0;
+			DSTexture->GetDepthStencilClearValue(Depth, Stencil);
+			ClearValues[ClearValueIndex].depthStencil.depth = Depth;
+			ClearValues[ClearValueIndex].depthStencil.stencil = Stencil;
+			++ClearValueIndex;
+		}
+	}
+
+	FRHITexture* ShadingRateTexture = RPInfo.ShadingRateTexture;
+	if (ShadingRateTexture)
+	{
+		ValidateShadingRateDataType();
+	}
+
+	ensure(ClearValueIndex <= RenderPass->GetNumUsedClearValues());
+
+	Barrier.Execute(CmdBuffer);
+
+	CmdBuffer->BeginRenderPass(RenderPass->GetLayout(), RenderPass, Framebuffer, ClearValues);
+
+	{
+		const VkExtent3D& Extents = RTLayout.GetExtent3D();
+		Context.GetPendingGfxState()->SetViewport(0, 0, 0, Extents.width, Extents.height, 1);
+	}
+}
+
+void FVulkanRenderPassManager::EndRenderPass(FVulkanCmdBuffer* CmdBuffer)
+{
+	CmdBuffer->EndRenderPass();
+
+	VulkanRHI::DebugHeavyWeightBarrier(CmdBuffer->GetHandle(), 1);
+}
+
+void FVulkanRenderPassManager::NotifyDeletedRenderTarget(VkImage Image)
+{
+	for (auto It = Framebuffers.CreateIterator(); It; ++It)
+	{
+		FFramebufferList* List = It->Value;
+		for (int32 Index = List->Framebuffer.Num() - 1; Index >= 0; --Index)
+		{
+			FVulkanFramebuffer* Framebuffer = List->Framebuffer[Index];
+			if (Framebuffer->ContainsRenderTarget(Image))
+			{
+				List->Framebuffer.RemoveAtSwap(Index, 1, false);
+				Framebuffer->Destroy(*Device);
+				delete Framebuffer;
+			}
+		}
+
+		if (List->Framebuffer.Num() == 0)
+		{
+			delete List;
+			It.RemoveCurrent();
+		}
+	}
+}
diff --git a/Engine/Source/Runtime/VulkanRHI/Private/VulkanRenderpass.h b/Engine/Source/Runtime/VulkanRHI/Private/VulkanRenderpass.h
--- a/Engine/Source/Runtime/VulkanRHI/Private/VulkanRenderpass.h
+++ b/Engine/Source/Runtime/VulkanRHI/Private/VulkanRenderpass.h
@@ -27,8 +27,9 @@ struct FVulkanAttachmentReference
 		SetAttachment(AttachmentReferenceIn, AspectMask);
 	}
 
-	inline void SetAttachment(const VkAttachmentReference& AttachmentReferenceIn, VkImageAspectFlags AspectMask) { check(false); }
+	inline void SetAttachment(const VkAttachmentReference& AttachmentReferenceIn, VkImageAspectFlags AspectMask) { checkNoEntry(); }
 	inline void SetAttachment(const FVulkanAttachmentReference<TAttachmentReferenceType>& AttachmentReferenceIn, VkImageAspectFlags AspectMask) { *this = AttachmentReferenceIn; }
+	inline void SetDepthStencilAttachment(const VkAttachmentReference& AttachmentReferenceIn, const VkAttachmentReferenceStencilLayout* StencilReference, VkImageAspectFlags AspectMask, bool bSupportsParallelRendering) { checkNoEntry(); }
 	inline void ZeroStruct() {}
 	inline void SetAspect(uint32 Aspect) {}
 };
@@ -40,7 +41,6 @@ inline void FVulkanAttachmentReference<VkAttachmentReference>::SetAttachment(con
 	layout = AttachmentReferenceIn.layout;
 }
 
-#if VULKAN_SUPPORTS_RENDERPASS2
 template <>
 inline void FVulkanAttachmentReference<VkAttachmentReference2>::SetAttachment(const VkAttachmentReference& AttachmentReferenceIn, VkImageAspectFlags AspectMask)
 {
@@ -60,7 +60,27 @@ inline void FVulkanAttachmentReference<VkAttachmentReference2>::SetAttachment(co
 	layout = AttachmentReferenceIn.layout;
 	aspectMask = AspectMask;
 }
-#endif
+
+template <>
+inline void FVulkanAttachmentReference<VkAttachmentReference>::SetDepthStencilAttachment(const VkAttachmentReference& AttachmentReferenceIn, 
+	const VkAttachmentReferenceStencilLayout* StencilReference, VkImageAspectFlags AspectMask, bool bSupportsParallelRendering)
+{
+	attachment = AttachmentReferenceIn.attachment;
+	const VkImageLayout StencilLayout = StencilReference ? StencilReference->stencilLayout : VK_IMAGE_LAYOUT_UNDEFINED;
+	layout = GetMergedDepthStencilLayout(AttachmentReferenceIn.layout, StencilLayout);
+}
+
+template <>
+inline void FVulkanAttachmentReference<VkAttachmentReference2>::SetDepthStencilAttachment(const VkAttachmentReference& AttachmentReferenceIn, 
+	const VkAttachmentReferenceStencilLayout* StencilReference, VkImageAspectFlags AspectMask, bool bSupportsParallelRendering)
+{
+	sType = VK_STRUCTURE_TYPE_ATTACHMENT_REFERENCE_2;
+	pNext = (bSupportsParallelRendering && StencilReference && StencilReference->stencilLayout != VK_IMAGE_LAYOUT_UNDEFINED) ? StencilReference : nullptr;
+	attachment = AttachmentReferenceIn.attachment;
+	layout = bSupportsParallelRendering ? AttachmentReferenceIn.layout : GetMergedDepthStencilLayout(AttachmentReferenceIn.layout, StencilReference->stencilLayout);
+	aspectMask = AspectMask;
+}
+
 
 template<>
 inline void FVulkanAttachmentReference<VkAttachmentReference>::ZeroStruct()
@@ -69,7 +89,6 @@ inline void FVulkanAttachmentReference<VkAttachmentReference>::ZeroStruct()
 	layout = VK_IMAGE_LAYOUT_UNDEFINED;
 }
 
-#if VULKAN_SUPPORTS_RENDERPASS2
 template<>
 inline void FVulkanAttachmentReference<VkAttachmentReference2>::ZeroStruct()
 {
@@ -85,7 +104,6 @@ inline void FVulkanAttachmentReference<VkAttachmentReference2>::SetAspect(uint32
 {
 	aspectMask = Aspect;
 }
-#endif
 
 template <typename TSubpassDescriptionType>
 class FVulkanSubpassDescription
@@ -139,7 +157,6 @@ struct FVulkanSubpassDescription<VkSubpassDescription>
 	}
 };
 
-#if VULKAN_SUPPORTS_RENDERPASS2
 template<>
 struct FVulkanSubpassDescription<VkSubpassDescription2>
 	: public VkSubpassDescription2
@@ -187,7 +204,6 @@ struct FVulkanSubpassDescription<VkSubpassDescription2>
 		viewMask = Mask;
 	}
 };
-#endif
 
 template <typename TSubpassDependencyType>
 struct FVulkanSubpassDependency
@@ -205,7 +221,6 @@ struct FVulkanSubpassDependency<VkSubpassDependency>
 	}
 };
 
-#if VULKAN_SUPPORTS_RENDERPASS2
 template<>
 struct FVulkanSubpassDependency<VkSubpassDependency2>
 	: public VkSubpassDependency2
@@ -216,7 +231,6 @@ struct FVulkanSubpassDependency<VkSubpassDependency2>
 		viewOffset = 0;		// According to the Vulkan spec: "If dependencyFlags does not include VK_DEPENDENCY_VIEW_LOCAL_BIT, viewOffset must be 0"
 	}
 };
-#endif
 
 template<typename TAttachmentDescriptionType>
 struct FVulkanAttachmentDescription
@@ -244,9 +258,25 @@ struct FVulkanAttachmentDescription<VkAttachmentDescription>
 		initialLayout = InDesc.initialLayout;
 		finalLayout = InDesc.finalLayout;
 	}
+
+	FVulkanAttachmentDescription(const VkAttachmentDescription& InDesc, const VkAttachmentDescriptionStencilLayout* InStencilDesc, bool bSupportsParallelRendering)
+	{
+		flags = InDesc.flags;
+		format = InDesc.format;
+		samples = InDesc.samples;
+		loadOp = InDesc.loadOp;
+		storeOp = InDesc.storeOp;
+		stencilLoadOp = InDesc.stencilLoadOp;
+		stencilStoreOp = InDesc.stencilStoreOp;
+
+		const VkImageLayout StencilInitialLayout = InStencilDesc ? InStencilDesc->stencilInitialLayout : VK_IMAGE_LAYOUT_UNDEFINED;
+		initialLayout = GetMergedDepthStencilLayout(InDesc.initialLayout, StencilInitialLayout);
+
+		const VkImageLayout StencilFinalLayout = InStencilDesc ? InStencilDesc->stencilFinalLayout : VK_IMAGE_LAYOUT_UNDEFINED;
+		finalLayout = GetMergedDepthStencilLayout(InDesc.finalLayout, StencilFinalLayout);
+	}
 };
 
-#if VULKAN_SUPPORTS_RENDERPASS2
 template<>
 struct FVulkanAttachmentDescription<VkAttachmentDescription2>
 	: public VkAttachmentDescription2
@@ -270,8 +300,22 @@ struct FVulkanAttachmentDescription<VkAttachmentDescription2>
 		initialLayout = InDesc.initialLayout;
 		finalLayout = InDesc.finalLayout;
 	}
+
+	FVulkanAttachmentDescription(const VkAttachmentDescription& InDesc, const VkAttachmentDescriptionStencilLayout* InStencilDesc, bool bSupportsParallelRendering)
+	{
+		sType = VK_STRUCTURE_TYPE_ATTACHMENT_DESCRIPTION_2;
+		pNext = (bSupportsParallelRendering && InStencilDesc && InStencilDesc->stencilFinalLayout != VK_IMAGE_LAYOUT_UNDEFINED) ? InStencilDesc : nullptr;
+		flags = InDesc.flags;
+		format = InDesc.format;
+		samples = InDesc.samples;
+		loadOp = InDesc.loadOp;
+		storeOp = InDesc.storeOp;
+		stencilLoadOp = InDesc.stencilLoadOp;
+		stencilStoreOp = InDesc.stencilStoreOp;
+		initialLayout = bSupportsParallelRendering ? InDesc.initialLayout : GetMergedDepthStencilLayout(InDesc.initialLayout, InStencilDesc->stencilInitialLayout);
+		finalLayout = bSupportsParallelRendering ? InDesc.finalLayout : GetMergedDepthStencilLayout(InDesc.finalLayout, InStencilDesc->stencilFinalLayout);
+	}
 };
-#endif
 
 template <typename T>
 struct FVulkanRenderPassCreateInfo
@@ -299,7 +343,6 @@ struct FVulkanRenderPassCreateInfo<VkRenderPassCreateInfo>
 	}
 };
 
-#if VULKAN_SUPPORTS_RENDERPASS2
 template<>
 struct FVulkanRenderPassCreateInfo<VkRenderPassCreateInfo2>
 	: public VkRenderPassCreateInfo2
@@ -339,7 +382,6 @@ struct FVulkanFragmentShadingRateAttachmentInfo
 		pFragmentShadingRateAttachment = AttachmentReference;
 	}
 };
-#endif
 
 extern int32 GVulkanInputAttachmentShaderRead;
 
@@ -361,15 +403,14 @@ public:
 
 		const bool bDeferredShadingSubpass = RTLayout.GetSubpassHint() == ESubpassHint::DeferredShadingSubpass;
 		const bool bDepthReadSubpass = RTLayout.GetSubpassHint() == ESubpassHint::DepthReadSubpass;
-		const bool bApplyFragmentShadingRate = GRHISupportsAttachmentVariableRateShading && GRHIVariableRateShadingEnabled && GRHIAttachmentVariableRateShadingEnabled && RTLayout.GetFragmentDensityAttachmentReference() != nullptr;
+		const bool bApplyFragmentShadingRate =  GRHISupportsAttachmentVariableRateShading && GRHIVariableRateShadingEnabled && GRHIAttachmentVariableRateShadingEnabled && RTLayout.GetFragmentDensityAttachmentReference() != nullptr;
+		const bool bHasDepthStencilAttachmentReference = (RTLayout.GetDepthAttachmentReference() != nullptr);
 
-#if VULKAN_SUPPORTS_RENDERPASS2
 		if (bApplyFragmentShadingRate)
 		{
 			ShadingRateAttachmentReference.SetAttachment(*RTLayout.GetFragmentDensityAttachmentReference(), VkImageAspectFlagBits::VK_IMAGE_ASPECT_COLOR_BIT);
 			FragmentShadingRateAttachmentInfo.SetReference(&ShadingRateAttachmentReference);
 		}
-#endif
 
 		// Grab (and optionally convert) attachment references.
 		for (uint32 ColorAttachment = 0; ColorAttachment < RTLayout.GetNumColorAttachments(); ++ColorAttachment)
@@ -381,9 +422,10 @@ public:
 			}
 		}
 
-		if (RTLayout.GetDepthStencilAttachmentReference() != nullptr)
+
+		if (bHasDepthStencilAttachmentReference)
 		{
-			DepthStencilAttachmentReference = TAttachmentReferenceClass(*RTLayout.GetDepthStencilAttachmentReference(), 0);
+			DepthStencilAttachmentReference.SetDepthStencilAttachment(*RTLayout.GetDepthAttachmentReference(), RTLayout.GetStencilAttachmentReference(), 0, Device.SupportsParallelRendering());
 		}
 
 		// main sub-pass
@@ -396,30 +438,28 @@ public:
 				// only set resolve attachment on the last subpass
 				SubpassDesc.SetResolveAttachments(ResolveAttachmentReferences);
 			}
-			if (RTLayout.GetDepthStencilAttachmentReference() != nullptr)
+			if (bHasDepthStencilAttachmentReference)
 			{
 				SubpassDesc.SetDepthStencilAttachment(&DepthStencilAttachmentReference);
 			}
 
-#if VULKAN_SUPPORTS_RENDERPASS2
 			if (bApplyFragmentShadingRate)
 			{
 				SubpassDesc.SetShadingRateAttachment(&FragmentShadingRateAttachmentInfo);
 			}
 			SubpassDesc.SetMultiViewMask(MultiviewMask);
-#endif
 		}
 
 		// Color write and depth read sub-pass
 		if (bDepthReadSubpass)
 		{
-			DepthStencilAttachmentOG.SetAttachment(*RTLayout.GetDepthStencilAttachmentReference(), VK_IMAGE_ASPECT_DEPTH_BIT);
+			DepthStencilAttachmentOG.SetAttachment(*RTLayout.GetDepthAttachmentReference(), VK_IMAGE_ASPECT_DEPTH_BIT);
 			TSubpassDescriptionClass& SubpassDesc = SubpassDescriptions[NumSubpasses++];
 
 			SubpassDesc.SetColorAttachments(ColorAttachmentReferences);
 			SubpassDesc.SetResolveAttachments(ResolveAttachmentReferences);
 
-			check(RTLayout.GetDepthStencilAttachmentReference());
+			check(RTLayout.GetDepthAttachmentReference());
 
 			// Depth as Input0
 			InputAttachments1[0].SetAttachment(DepthStencilAttachmentOG, VK_IMAGE_ASPECT_DEPTH_BIT);
@@ -429,13 +469,11 @@ public:
 			// depth attachment is same as input attachment
 			SubpassDesc.SetDepthStencilAttachment(InputAttachments1);
 
-#if VULKAN_SUPPORTS_RENDERPASS2
 			if (bApplyFragmentShadingRate)
 			{
 				SubpassDesc.SetShadingRateAttachment(&FragmentShadingRateAttachmentInfo);
 			}
 			SubpassDesc.SetMultiViewMask(MultiviewMask);
-#endif
 
 			TSubpassDependencyClass& SubpassDep = SubpassDependencies[NumDependencies++];
 			SubpassDep.srcSubpass = 0;
@@ -451,7 +489,7 @@ public:
 		if (bDeferredShadingSubpass)
 		{
 			// both sub-passes only test DepthStencil
-			DepthStencilAttachment.attachment = RTLayout.GetDepthStencilAttachmentReference()->attachment;
+			DepthStencilAttachment.attachment = RTLayout.GetDepthAttachmentReference()->attachment;
 			DepthStencilAttachment.layout = VK_IMAGE_LAYOUT_DEPTH_STENCIL_READ_ONLY_OPTIMAL;
 			DepthStencilAttachment.SetAspect(VK_IMAGE_ASPECT_DEPTH_BIT);	// @todo?
 
@@ -466,13 +504,11 @@ public:
 				SubpassDesc.SetDepthStencilAttachment(&DepthStencilAttachment);
 				SubpassDesc.SetInputAttachments(&DepthStencilAttachment, 1);
 
-#if VULKAN_SUPPORTS_RENDERPASS2
 				if (bApplyFragmentShadingRate)
 				{
 					SubpassDesc.SetShadingRateAttachment(&FragmentShadingRateAttachmentInfo);
 				}
 				SubpassDesc.SetMultiViewMask(MultiviewMask);
-#endif
 
 				// Depth as Input0
 				TSubpassDependencyClass& SubpassDep = SubpassDependencies[NumDependencies++];
@@ -511,13 +547,11 @@ public:
 				}
 
 				SubpassDesc.SetInputAttachments(InputAttachments2, NumColorInputs + 2);
-#if VULKAN_SUPPORTS_RENDERPASS2
 				if (bApplyFragmentShadingRate)
 				{
 					SubpassDesc.SetShadingRateAttachment(&FragmentShadingRateAttachmentInfo);
 				}
 				SubpassDesc.SetMultiViewMask(MultiviewMask);
-#endif
 
 				TSubpassDependencyClass& SubpassDep = SubpassDependencies[NumDependencies++];
 				SubpassDep.srcSubpass = 1;
@@ -537,7 +571,14 @@ public:
 
 		for (uint32 Attachment = 0; Attachment < RTLayout.GetNumAttachmentDescriptions(); ++Attachment)
 		{
-			AttachmentDescriptions.Add(TAttachmentDescriptionClass(RTLayout.GetAttachmentDescriptions()[Attachment]));
+			if (bHasDepthStencilAttachmentReference && (Attachment == DepthStencilAttachmentReference.attachment))
+			{
+				AttachmentDescriptions.Add(TAttachmentDescriptionClass(RTLayout.GetAttachmentDescriptions()[Attachment], RTLayout.GetStencilDesc(), Device.SupportsParallelRendering()));
+			}
+			else
+			{
+				AttachmentDescriptions.Add(TAttachmentDescriptionClass(RTLayout.GetAttachmentDescriptions()[Attachment]));
+			}
 		}
 
 		CreateInfo.attachmentCount = AttachmentDescriptions.Num();
@@ -562,13 +603,11 @@ public:
 		VkRenderPassMultiviewCreateInfo MultiviewInfo;
 		if (RTLayout.GetIsMultiView())
 		{
-#if VULKAN_SUPPORTS_RENDERPASS2
 			if (Device.GetOptionalExtensions().HasKHRRenderPass2)
 			{
 				CreateInfo.SetCorrelationMask(&CorrelationMask);
 			}
 			else
-#endif
 			{
 				checkf(Device.GetOptionalExtensions().HasKHRMultiview, TEXT("Layout is multiview but extension is not supported!"))
 				ZeroVulkanStruct(MultiviewInfo, VK_STRUCTURE_TYPE_RENDER_PASS_MULTIVIEW_CREATE_INFO);
@@ -635,10 +674,8 @@ private:
 	TAttachmentReferenceClass DepthStencilAttachmentReference;
 	TArray<TAttachmentDescriptionClass> AttachmentDescriptions;
 
-#if VULKAN_SUPPORTS_RENDERPASS2
 	FVulkanAttachmentReference<VkAttachmentReference2> ShadingRateAttachmentReference;
 	FVulkanFragmentShadingRateAttachmentInfo FragmentShadingRateAttachmentInfo;
-#endif
 
 	TRenderPassCreateInfoClass CreateInfo;
 	FVulkanDevice& Device;
@@ -646,3 +683,57 @@ private:
 
 VkRenderPass CreateVulkanRenderPass(FVulkanDevice& Device, const FVulkanRenderTargetLayout& RTLayout);
 
+
+
+class FVulkanRenderPassManager : public VulkanRHI::FDeviceChild
+{
+public:
+	FVulkanRenderPassManager(FVulkanDevice* InDevice) : VulkanRHI::FDeviceChild(InDevice) {}
+	~FVulkanRenderPassManager();
+
+	FVulkanFramebuffer* GetOrCreateFramebuffer(const FRHISetRenderTargetsInfo& RenderTargetsInfo, const FVulkanRenderTargetLayout& RTLayout, FVulkanRenderPass* RenderPass);
+	FVulkanRenderPass* GetOrCreateRenderPass(const FVulkanRenderTargetLayout& RTLayout)
+	{
+		const uint32 RenderPassHash = RTLayout.GetRenderPassFullHash();
+
+		{
+			FRWScopeLock ScopedReadLock(RenderPassesLock, SLT_ReadOnly);
+			FVulkanRenderPass** FoundRenderPass = RenderPasses.Find(RenderPassHash);
+			if (FoundRenderPass)
+			{
+				return *FoundRenderPass;
+			}
+		}
+
+		FVulkanRenderPass* RenderPass = new FVulkanRenderPass(*Device, RTLayout);
+		{
+			FRWScopeLock ScopedWriteLock(RenderPassesLock, SLT_Write);
+			FVulkanRenderPass** FoundRenderPass = RenderPasses.Find(RenderPassHash);
+			if (FoundRenderPass)
+			{
+				delete RenderPass;
+				return *FoundRenderPass;
+			}
+			RenderPasses.Add(RenderPassHash, RenderPass);
+		}
+		return RenderPass;
+	}
+
+	void BeginRenderPass(FVulkanCommandListContext& Context, FVulkanDevice& InDevice, FVulkanCmdBuffer* CmdBuffer, const FRHIRenderPassInfo& RPInfo, const FVulkanRenderTargetLayout& RTLayout, FVulkanRenderPass* RenderPass, FVulkanFramebuffer* Framebuffer);
+	void EndRenderPass(FVulkanCmdBuffer* CmdBuffer);
+
+	FRWLock RenderPassesLock;
+	FRWLock FramebuffersLock;
+
+	void NotifyDeletedRenderTarget(VkImage Image);
+
+private:
+
+	TMap<uint32, FVulkanRenderPass*> RenderPasses;
+
+	struct FFramebufferList
+	{
+		TArray<FVulkanFramebuffer*> Framebuffer;
+	};
+	TMap<uint32, FFramebufferList*> Framebuffers;
+};
diff --git a/Engine/Source/Runtime/VulkanRHI/Private/VulkanResources.h b/Engine/Source/Runtime/VulkanRHI/Private/VulkanResources.h
--- a/Engine/Source/Runtime/VulkanRHI/Private/VulkanResources.h
+++ b/Engine/Source/Runtime/VulkanRHI/Private/VulkanResources.h
@@ -449,7 +449,7 @@ public:
 	/**
 	 * Returns how much memory is used by the surface
 	 */
-	uint32 GetMemorySize() const
+	inline uint32 GetMemorySize() const
 	{
 		return MemoryRequirements.size;
 	}
@@ -524,6 +524,11 @@ public:
 		return EnumHasAllFlags(GPixelFormats[GetDesc().Format].Capabilities, EPixelFormatCapabilities::TextureSample);
 	}
 
+	inline VkImageLayout GetDefaultLayout() const
+	{
+		return DefaultLayout;
+	}
+
 	VULKANRHI_API VkDeviceMemory GetAllocationHandle() const;
 	VULKANRHI_API uint64 GetAllocationOffset() const;
 
@@ -543,14 +548,15 @@ private:
 
 	void InvalidateViews(FVulkanDevice& Device);
 	void DestroyViews();
-	void SetInitialImageState(FVulkanCommandListContext& Context, VkImageLayout InitialLayout, bool bClear, const FClearValueBinding& ClearValueBinding);
-	void InternalMoveSurface(FVulkanDevice& InDevice, FVulkanCommandListContext& Context, VulkanRHI::FVulkanAllocation& DestAllocation);
+	void SetInitialImageState(FVulkanCommandListContext& Context, VkImageLayout InitialLayout, bool bClear, const FClearValueBinding& ClearValueBinding, bool bIsTransientResource);
+	void InternalMoveSurface(FVulkanDevice& InDevice, FVulkanCommandListContext& Context, VulkanRHI::FVulkanAllocation& DestAllocation, VkImageLayout OriginalLayout);
 
 	VkImageTiling Tiling;
 	VulkanRHI::FVulkanAllocation Allocation;
 	VkImageAspectFlags FullAspectMask;
 	VkImageAspectFlags PartialAspectMask;
 	FVulkanCpuReadbackBuffer* CpuReadbackBuffer;
+	VkImageLayout DefaultLayout;
 
 	friend struct FRHICommandSetInitialImageState;
 
diff --git a/Engine/Source/Runtime/VulkanRHI/Private/VulkanTexture.cpp b/Engine/Source/Runtime/VulkanRHI/Private/VulkanTexture.cpp
--- a/Engine/Source/Runtime/VulkanRHI/Private/VulkanTexture.cpp
+++ b/Engine/Source/Runtime/VulkanRHI/Private/VulkanTexture.cpp
@@ -161,33 +161,18 @@ inline void FVulkanTexture::InternalLockWrite(FVulkanCommandListContext& Context
 	const VkImageSubresourceLayers& ImageSubresource = Region.imageSubresource;
 	const VkImageSubresourceRange SubresourceRange = FVulkanPipelineBarrier::MakeSubresourceRange(ImageSubresource.aspectMask, ImageSubresource.mipLevel, 1, ImageSubresource.baseArrayLayer, ImageSubresource.layerCount);
 
-	FVulkanImageLayout& TrackedTextureLayout = Context.GetLayoutManager().GetOrAddFullLayout(*Surface, VK_IMAGE_LAYOUT_UNDEFINED);
-
-	// Transition the subresource layouts from their tracked state to the copy state
-	FVulkanImageLayout TransferTextureLayout = TrackedTextureLayout;
-	TransferTextureLayout.Set(VK_IMAGE_LAYOUT_TRANSFER_DST_OPTIMAL, SubresourceRange);
 	{
 		FVulkanPipelineBarrier Barrier;
-		Barrier.AddImageLayoutTransition(Surface->Image, SubresourceRange.aspectMask, TrackedTextureLayout, TransferTextureLayout);
-		Barrier.Execute(StagingCommandBuffer);
+		Barrier.AddImageLayoutTransition(Surface->Image, VK_IMAGE_LAYOUT_UNDEFINED, VK_IMAGE_LAYOUT_TRANSFER_DST_OPTIMAL, SubresourceRange);
+		Barrier.Execute(CmdBuffer);
 	}
 
 	VulkanRHI::vkCmdCopyBufferToImage(StagingCommandBuffer, StagingBuffer->GetHandle(), Surface->Image, VK_IMAGE_LAYOUT_TRANSFER_DST_OPTIMAL, 1, &Region);
 
-	// Transition the subresource layouts from the copy state to a regular read state if it was undefined, otherwise return it how it was
-	for (uint32 LayerIndex = ImageSubresource.baseArrayLayer; LayerIndex < ImageSubresource.baseArrayLayer + ImageSubresource.layerCount; ++LayerIndex)
-	{
-		if (TrackedTextureLayout.GetSubresLayout(LayerIndex, ImageSubresource.mipLevel) == VK_IMAGE_LAYOUT_UNDEFINED)
-		{
-			const VkImageSubresourceRange SingleLayerSubresourceRange = FVulkanPipelineBarrier::MakeSubresourceRange(ImageSubresource.aspectMask, ImageSubresource.mipLevel, 1, LayerIndex, 1);
-			TrackedTextureLayout.Set(VK_IMAGE_LAYOUT_SHADER_READ_ONLY_OPTIMAL, SingleLayerSubresourceRange);
-		}
-	}
-
 	{
 		FVulkanPipelineBarrier Barrier;
-		Barrier.AddImageLayoutTransition(Surface->Image, SubresourceRange.aspectMask, TransferTextureLayout, TrackedTextureLayout);
-		Barrier.Execute(StagingCommandBuffer);
+		Barrier.AddImageLayoutTransition(Surface->Image, VK_IMAGE_LAYOUT_TRANSFER_DST_OPTIMAL, Surface->GetDefaultLayout(), SubresourceRange);
+		Barrier.Execute(CmdBuffer);
 	}
 
 	Surface->Device->GetStagingManager().ReleaseBuffer(CmdBuffer, StagingBuffer);
@@ -490,17 +475,19 @@ void FVulkanTexture::GenerateImageCreateInfo(
 
 struct FRHICommandSetInitialImageState final : public FRHICommand<FRHICommandSetInitialImageState>
 {
-	FVulkanTexture* Surface;
+	FVulkanTexture* VulkanTexture;
 	VkImageLayout InitialLayout;
 	bool bOnlyAddToLayoutManager;
 	bool bClear;
+	bool bIsTransientResource;
 	FClearValueBinding ClearValueBinding;
 
-	FRHICommandSetInitialImageState(FVulkanTexture* InSurface, VkImageLayout InInitialLayout, bool bInOnlyAddToLayoutManager, bool bInClear, const FClearValueBinding& InClearValueBinding)
-		: Surface(InSurface)
+	FRHICommandSetInitialImageState(FVulkanTexture* InVulkanTexture, VkImageLayout InInitialLayout, bool bInOnlyAddToLayoutManager, bool bInClear, const FClearValueBinding& InClearValueBinding, bool bInIsTransientResource)
+		: VulkanTexture(InVulkanTexture)
 		, InitialLayout(InInitialLayout)
 		, bOnlyAddToLayoutManager(bInOnlyAddToLayoutManager)
 		, bClear(bInClear)
+		, bIsTransientResource(bInIsTransientResource)
 		, ClearValueBinding(InClearValueBinding)
 	{
 	}
@@ -511,11 +498,12 @@ struct FRHICommandSetInitialImageState final : public FRHICommand<FRHICommandSet
 		
 		if (bOnlyAddToLayoutManager)
 		{
-			Context.GetLayoutManager().GetOrAddFullLayout(*Surface, InitialLayout);
+			FVulkanCmdBuffer* CmdBuffer = Context.GetCommandBufferManager()->GetActiveCmdBuffer();
+			CmdBuffer->GetLayoutManager().SetFullLayout(*VulkanTexture, InitialLayout, true);
 		}
 		else
 		{
-			Surface->SetInitialImageState(Context, InitialLayout, bClear, ClearValueBinding);
+			VulkanTexture->SetInitialImageState(Context, InitialLayout, bClear, ClearValueBinding, bIsTransientResource);
 		}
 	}
 };
@@ -539,7 +527,7 @@ struct FRHICommandOnDestroyImage final : public FRHICommand<FRHICommandOnDestroy
 	}
 };
 
-static VkImageLayout GetInitialLayoutFromRHIAccess(ERHIAccess RHIAccess, ETextureCreateFlags UEFlags, bool bSupportReadOnlyOptimal)
+static VkImageLayout GetInitialLayoutFromRHIAccess(ERHIAccess RHIAccess, bool bIsDepthStencil, bool bSupportReadOnlyOptimal)
 {
 	if (EnumHasAnyFlags(RHIAccess, ERHIAccess::RTV) || RHIAccess == ERHIAccess::Present)
 	{
@@ -548,19 +536,19 @@ static VkImageLayout GetInitialLayoutFromRHIAccess(ERHIAccess RHIAccess, ETextur
 
 	if (EnumHasAnyFlags(RHIAccess, ERHIAccess::DSVWrite))
 	{
-		return VK_IMAGE_LAYOUT_DEPTH_STENCIL_ATTACHMENT_OPTIMAL;
+		return VK_IMAGE_LAYOUT_ATTACHMENT_OPTIMAL;
 	}
 
 	if (EnumHasAnyFlags(RHIAccess, ERHIAccess::DSVRead))
 	{
-		return VK_IMAGE_LAYOUT_DEPTH_STENCIL_READ_ONLY_OPTIMAL;
+		return VK_IMAGE_LAYOUT_READ_ONLY_OPTIMAL;
 	}
 
 	if (EnumHasAnyFlags(RHIAccess, ERHIAccess::SRVMask))
 	{
-		if (EnumHasAnyFlags(UEFlags, TexCreate_DepthStencilTargetable))
+		if (bIsDepthStencil)
 		{
-			return VK_IMAGE_LAYOUT_DEPTH_STENCIL_READ_ONLY_OPTIMAL;
+			return VK_IMAGE_LAYOUT_READ_ONLY_OPTIMAL;
 		}
 
 		return bSupportReadOnlyOptimal ? VK_IMAGE_LAYOUT_SHADER_READ_ONLY_OPTIMAL : VK_IMAGE_LAYOUT_GENERAL;
@@ -583,7 +571,7 @@ static VkImageLayout GetInitialLayoutFromRHIAccess(ERHIAccess RHIAccess, ETextur
 	return VK_IMAGE_LAYOUT_UNDEFINED;
 }
 
-void FVulkanTexture::InternalMoveSurface(FVulkanDevice& InDevice, FVulkanCommandListContext& Context, FVulkanAllocation& DestAllocation)
+void FVulkanTexture::InternalMoveSurface(FVulkanDevice& InDevice, FVulkanCommandListContext& Context, FVulkanAllocation& DestAllocation, VkImageLayout OriginalLayout)
 {
 	FImageCreateInfo ImageCreateInfo;
 	const FRHITextureDesc& Desc = GetDesc();
@@ -618,25 +606,16 @@ void FVulkanTexture::InternalMoveSurface(FVulkanDevice& InDevice, FVulkanCommand
 	VkCommandBuffer VkCmdBuffer = CmdBuffer->GetHandle();
 	ensure(CmdBuffer->IsOutsideRenderPass());
 
-	const uint32 NumberOfArrayLevels = GetNumberOfArrayLevels();
-	FVulkanImageLayout* MovedLayout = &Context.GetLayoutManager().FindOrAddFullLayoutRW(MovedImage, VK_IMAGE_LAYOUT_UNDEFINED, GetNumMips(), NumberOfArrayLevels);
-	FVulkanImageLayout* OriginalLayout = Context.GetLayoutManager().GetFullLayout(Image);
-	// Account for map resize, should rarely happen...
-	if (OriginalLayout == nullptr)
 	{
-		OriginalLayout = &Context.GetLayoutManager().GetOrAddFullLayout(*this, VK_IMAGE_LAYOUT_UNDEFINED);
-		MovedLayout = &Context.GetLayoutManager().GetFullLayoutChecked(MovedImage);
-	}
+		const uint32 NumberOfArrayLevels = GetNumberOfArrayLevels();
+		const VkImageSubresourceRange FullSubresourceRange = FVulkanPipelineBarrier::MakeSubresourceRange(FullAspectMask);
 
-	checkf((OriginalLayout->NumMips == GetNumMips()), TEXT("NumMips reported by LayoutManager (%d) differs from surface (%d)"), OriginalLayout->NumMips, GetNumMips());
-	checkf((OriginalLayout->NumLayers == NumberOfArrayLevels), TEXT("NumLayers reported by LayoutManager (%d) differs from surface (%d)"), OriginalLayout->NumLayers, NumberOfArrayLevels);
-	{
 		// Transition to copying layouts
 		{
 			FVulkanPipelineBarrier Barrier;
-			Barrier.AddImageLayoutTransition(Image, FullAspectMask, *OriginalLayout, VK_IMAGE_LAYOUT_TRANSFER_SRC_OPTIMAL);
-			Barrier.AddImageLayoutTransition(MovedImage, FullAspectMask, *MovedLayout, VK_IMAGE_LAYOUT_TRANSFER_DST_OPTIMAL);
-			Barrier.Execute(VkCmdBuffer);
+			Barrier.AddImageLayoutTransition(Image, OriginalLayout, VK_IMAGE_LAYOUT_TRANSFER_SRC_OPTIMAL, FullSubresourceRange);
+			Barrier.AddImageLayoutTransition(MovedImage, VK_IMAGE_LAYOUT_UNDEFINED, VK_IMAGE_LAYOUT_TRANSFER_DST_OPTIMAL, FullSubresourceRange);
+			Barrier.Execute(CmdBuffer);
 		}
 		{
 			VkImageCopy Regions[MAX_TEXTURE_MIP_COUNT];
@@ -666,13 +645,10 @@ void FVulkanTexture::InternalMoveSurface(FVulkanDevice& InDevice, FVulkanCommand
 		// Put the destination image in exactly the same layout the original image was
 		{
 			FVulkanPipelineBarrier Barrier;
-			Barrier.AddImageLayoutTransition(MovedImage, FullAspectMask, VK_IMAGE_LAYOUT_TRANSFER_DST_OPTIMAL, *OriginalLayout);
-			Barrier.Execute(VkCmdBuffer);
+			Barrier.AddImageLayoutTransition(Image, VK_IMAGE_LAYOUT_TRANSFER_SRC_OPTIMAL, OriginalLayout, FullSubresourceRange);
+			Barrier.AddImageLayoutTransition(MovedImage, VK_IMAGE_LAYOUT_TRANSFER_DST_OPTIMAL, OriginalLayout, FullSubresourceRange);
+			Barrier.Execute(CmdBuffer);
 		}
-
-		// Update the tracked layouts
-		*MovedLayout = *OriginalLayout;
-		OriginalLayout->Set(VK_IMAGE_LAYOUT_TRANSFER_SRC_OPTIMAL, FVulkanPipelineBarrier::MakeSubresourceRange(FullAspectMask));
 	}
 
 	{
@@ -687,6 +663,9 @@ void FVulkanTexture::InternalMoveSurface(FVulkanDevice& InDevice, FVulkanCommand
 	}
 
 	Image = MovedImage;
+
+	// Move is used for defrag, which uses layouts stored in the queue, update the layout stored there
+	Context.GetQueue()->GetLayoutManager().SetFullLayout(*this, OriginalLayout);
 }
 
 void FVulkanTexture::DestroySurface()
@@ -837,19 +816,24 @@ void FVulkanTexture::GetMipSize(uint32 MipIndex, uint32& MipBytes)
 */
 }
 
-void FVulkanTexture::SetInitialImageState(FVulkanCommandListContext& Context, VkImageLayout InitialLayout, bool bClear, const FClearValueBinding& ClearValueBinding)
+void FVulkanTexture::SetInitialImageState(FVulkanCommandListContext& Context, VkImageLayout InitialLayout, bool bClear, const FClearValueBinding& ClearValueBinding, bool bIsTransientResource)
 {
 	// Can't use TransferQueue as Vulkan requires that queue to also have Gfx or Compute capabilities...
 	//#todo-rco: This function is only used during loading currently, if used for regular RHIClear then use the ActiveCmdBuffer
-	FVulkanCmdBuffer* CmdBuffer = Context.GetCommandBufferManager()->GetUploadCmdBuffer();
+	// NOTE: Transient resources' memory might have belonged to another resource earlier in the ActiveCmdBuffer, so we can't use UploadCmdBuffer
+	FVulkanCmdBuffer* CmdBuffer = bIsTransientResource ? Context.GetCommandBufferManager()->GetActiveCmdBuffer() : Context.GetCommandBufferManager()->GetUploadCmdBuffer();
 	ensure(CmdBuffer->IsOutsideRenderPass());
 
 	VkImageSubresourceRange SubresourceRange = FVulkanPipelineBarrier::MakeSubresourceRange(FullAspectMask);
 
-	VkImageLayout CurrentLayout;
-	if (bClear)
+	VkImageLayout CurrentLayout = VK_IMAGE_LAYOUT_UNDEFINED;
+	if (bClear && !bIsTransientResource)
 	{
-		VulkanSetImageLayout(CmdBuffer->GetHandle(), Image, VK_IMAGE_LAYOUT_UNDEFINED, VK_IMAGE_LAYOUT_TRANSFER_DST_OPTIMAL, SubresourceRange);
+		{
+			FVulkanPipelineBarrier Barrier;
+			Barrier.AddImageLayoutTransition(Image, VK_IMAGE_LAYOUT_UNDEFINED, VK_IMAGE_LAYOUT_TRANSFER_DST_OPTIMAL, SubresourceRange);
+			Barrier.Execute(CmdBuffer);
+		}
 
 		if (FullAspectMask == VK_IMAGE_ASPECT_COLOR_BIT)
 		{
@@ -875,19 +859,15 @@ void FVulkanTexture::SetInitialImageState(FVulkanCommandListContext& Context, Vk
 
 		CurrentLayout = VK_IMAGE_LAYOUT_TRANSFER_DST_OPTIMAL;
 	}
-	else
-	{
-		CurrentLayout = VK_IMAGE_LAYOUT_UNDEFINED;
-	}
 
-	if (InitialLayout != CurrentLayout && InitialLayout != VK_IMAGE_LAYOUT_UNDEFINED)
+	if ((InitialLayout != CurrentLayout) && (InitialLayout != VK_IMAGE_LAYOUT_UNDEFINED))
 	{
-		VulkanSetImageLayout(CmdBuffer->GetHandle(), Image, CurrentLayout, InitialLayout, SubresourceRange);
+		FVulkanPipelineBarrier Barrier;
+		Barrier.AddImageLayoutTransition(Image, CurrentLayout, InitialLayout, SubresourceRange);
+		Barrier.Execute(CmdBuffer);
 	}
 
-	FVulkanImageLayout& FullLayout = Context.GetLayoutManager().GetOrAddFullLayout(*this, InitialLayout);
-	checkSlow(FullLayout.AreAllSubresourcesSameLayout());
-	FullLayout.MainLayout = InitialLayout;
+	CmdBuffer->GetLayoutManager().SetFullLayout(*this, InitialLayout);
 }
 
 /*-----------------------------------------------------------------------------
@@ -993,17 +973,12 @@ static void DoAsyncReallocateTexture2D(FVulkanCommandListContext& Context, FVulk
 	const VkImageSubresourceRange SourceSubResourceRange = FVulkanPipelineBarrier::MakeSubresourceRange(VK_IMAGE_ASPECT_COLOR_BIT, SourceFirstMip, NumSharedMips);
 	const VkImageSubresourceRange DestSubResourceRange = FVulkanPipelineBarrier::MakeSubresourceRange(VK_IMAGE_ASPECT_COLOR_BIT, DestFirstMip, NumSharedMips);
 
-	const FVulkanImageLayout& OldTextureOriginalLayout = Context.GetLayoutManager().GetOrAddFullLayout(*OldTexture, VK_IMAGE_LAYOUT_UNDEFINED);
-	ensure(!OldTextureOriginalLayout.AreAllSubresourcesSameLayout() || (OldTextureOriginalLayout.MainLayout != VK_IMAGE_LAYOUT_UNDEFINED));
-	FVulkanImageLayout OldTextureCopyLayout = OldTextureOriginalLayout;
-	OldTextureCopyLayout.Set(VK_IMAGE_LAYOUT_TRANSFER_SRC_OPTIMAL, SourceSubResourceRange);
-
 	{
 		// Pre-copy barriers
 		FVulkanPipelineBarrier Barrier;
-		Barrier.AddImageLayoutTransition(OldTexture->Image, VK_IMAGE_ASPECT_COLOR_BIT, OldTextureOriginalLayout, OldTextureCopyLayout);
+		Barrier.AddImageLayoutTransition(OldTexture->Image, OldTexture->GetDefaultLayout(), VK_IMAGE_LAYOUT_TRANSFER_SRC_OPTIMAL, SourceSubResourceRange);
 		Barrier.AddImageLayoutTransition(NewTexture->Image, VK_IMAGE_LAYOUT_UNDEFINED, VK_IMAGE_LAYOUT_TRANSFER_DST_OPTIMAL, DestSubResourceRange);
-		Barrier.Execute(CmdBuffer->GetHandle());
+		Barrier.Execute(CmdBuffer);
 	}
 
 	VulkanRHI::vkCmdCopyImage(StagingCommandBuffer, OldTexture->Image, VK_IMAGE_LAYOUT_TRANSFER_SRC_OPTIMAL, NewTexture->Image, VK_IMAGE_LAYOUT_TRANSFER_DST_OPTIMAL, NumSharedMips, Regions);
@@ -1011,15 +986,13 @@ static void DoAsyncReallocateTexture2D(FVulkanCommandListContext& Context, FVulk
 	{
 		// Post-copy barriers
 		FVulkanPipelineBarrier Barrier;
-		Barrier.AddImageLayoutTransition(OldTexture->Image, VK_IMAGE_ASPECT_COLOR_BIT, OldTextureCopyLayout, OldTextureOriginalLayout);
-		Barrier.AddImageLayoutTransition(NewTexture->Image, VK_IMAGE_LAYOUT_TRANSFER_DST_OPTIMAL, VK_IMAGE_LAYOUT_SHADER_READ_ONLY_OPTIMAL, DestSubResourceRange);
-		Barrier.Execute(CmdBuffer->GetHandle());
+		Barrier.AddImageLayoutTransition(OldTexture->Image, VK_IMAGE_LAYOUT_TRANSFER_SRC_OPTIMAL, OldTexture->GetDefaultLayout(), SourceSubResourceRange);
+		Barrier.AddImageLayoutTransition(NewTexture->Image, VK_IMAGE_LAYOUT_TRANSFER_DST_OPTIMAL, NewTexture->GetDefaultLayout(), DestSubResourceRange);
+		Barrier.Execute(CmdBuffer);
 
 		// Add tracking for the appropriate subresources (intentionally leave added mips in VK_IMAGE_LAYOUT_UNDEFINED)
 		// NOTE: Overwriting whatever is contained in the layout manager for the new texture (circumvents issue with stale layouts and pointer reuse)
-		FVulkanImageLayout NewTextureLayout(VK_IMAGE_LAYOUT_UNDEFINED, NewTexture->GetNumMips(), NewTexture->GetNumberOfArrayLevels());
-		NewTextureLayout.Set(VK_IMAGE_LAYOUT_SHADER_READ_ONLY_OPTIMAL, DestSubResourceRange);
-		Context.GetLayoutManager().GetOrAddFullLayout(*NewTexture, VK_IMAGE_LAYOUT_UNDEFINED) = NewTextureLayout;
+		CmdBuffer->GetLayoutManager().SetLayout(*NewTexture, DestSubResourceRange, NewTexture->GetDefaultLayout());
 	}
 
 	// request is now complete
@@ -1567,6 +1540,7 @@ FVulkanTexture::FVulkanTexture(FVulkanDevice& InDevice, const FRHITextureCreateD
 	, FullAspectMask(0)
 	, PartialAspectMask(0)
 	, CpuReadbackBuffer(nullptr) // for readback textures we use a staging buffer. this is because vulkan only requires implementations to support 1 mip level(which is useless), so we emulate using a buffer
+	, DefaultLayout(VK_IMAGE_LAYOUT_UNDEFINED)
 {
 	VULKAN_TRACK_OBJECT_CREATE(FVulkanTexture, this);
 
@@ -1695,22 +1669,24 @@ FVulkanTexture::FVulkanTexture(FVulkanDevice& InDevice, const FRHITextureCreateD
 		Tiling = ImageCreateInfo.ImageCreateInfo.tiling;
 		check(Tiling == VK_IMAGE_TILING_LINEAR || Tiling == VK_IMAGE_TILING_OPTIMAL);
 
-		const VkImageLayout InitialLayout = GetInitialLayoutFromRHIAccess(InCreateDesc.InitialState, InCreateDesc.Flags, SupportsSampling());
-		const bool bDoInitialClear = !bIsTransientResource && VKHasAnyFlags(ImageCreateInfo.ImageCreateInfo.usage, VK_IMAGE_USAGE_SAMPLED_BIT) && EnumHasAnyFlags(InCreateDesc.Flags, TexCreate_RenderTargetable | TexCreate_DepthStencilTargetable);
+		const VkImageLayout InitialLayout = GetInitialLayoutFromRHIAccess(InCreateDesc.InitialState, IsDepthOrStencilAspect(), SupportsSampling());
+		const bool bDoInitialClear = VKHasAnyFlags(ImageCreateInfo.ImageCreateInfo.usage, VK_IMAGE_USAGE_SAMPLED_BIT) && EnumHasAnyFlags(InCreateDesc.Flags, TexCreate_RenderTargetable | TexCreate_DepthStencilTargetable);
 
 		if (InitialLayout != VK_IMAGE_LAYOUT_UNDEFINED || bDoInitialClear)
 		{
 			FRHICommandList& RHICmdList = FRHICommandListExecutor::GetImmediateCommandList();
 			if (!IsInRenderingThread() || (RHICmdList.Bypass() || !IsRunningRHIInSeparateThread()))
 			{
-				SetInitialImageState(Device->GetImmediateContext(), InitialLayout, bDoInitialClear, InCreateDesc.ClearValue);
+				SetInitialImageState(Device->GetImmediateContext(), InitialLayout, bDoInitialClear, InCreateDesc.ClearValue, bIsTransientResource);
 			}
 			else
 			{
 				check(IsInRenderingThread());
-				ALLOC_COMMAND_CL(RHICmdList, FRHICommandSetInitialImageState)(this, InitialLayout, false, bDoInitialClear, InCreateDesc.ClearValue);
+				ALLOC_COMMAND_CL(RHICmdList, FRHICommandSetInitialImageState)(this, InitialLayout, false, bDoInitialClear, InCreateDesc.ClearValue, bIsTransientResource);
 			}
 		}
+
+		DefaultLayout = InitialLayout;
 	}
 
 	LLM_SCOPE_VULKAN(ELLMTagVulkan::VulkanTextures);
@@ -1745,6 +1721,7 @@ FVulkanTexture::FVulkanTexture(FVulkanDevice& InDevice, const FRHITextureCreateD
 
 	// InternalLockWrite leaves the image in VK_IMAGE_LAYOUT_SHADER_READ_ONLY_OPTIMAL, so make sure the requested resource state is SRV.
 	check(EnumHasAnyFlags(InCreateDesc.InitialState, ERHIAccess::SRVMask));
+	DefaultLayout = VK_IMAGE_LAYOUT_SHADER_READ_ONLY_OPTIMAL;
 
 	// Transfer bulk data
 	VulkanRHI::FStagingBuffer* StagingBuffer = InDevice.GetStagingManager().AcquireBuffer(InCreateDesc.BulkData->GetResourceBulkDataSize());
@@ -1795,6 +1772,7 @@ FVulkanTexture::FVulkanTexture(FVulkanDevice& InDevice, const FRHITextureCreateD
 	, FullAspectMask(0)
 	, PartialAspectMask(0)
 	, CpuReadbackBuffer(nullptr)
+	, DefaultLayout(VK_IMAGE_LAYOUT_UNDEFINED)
 	, ImageOwnerType(EImageOwnerType::ExternalOwner)
 {
 	VULKAN_TRACK_OBJECT_CREATE(FVulkanTexture, this);
@@ -1827,7 +1805,7 @@ FVulkanTexture::FVulkanTexture(FVulkanDevice& InDevice, const FRHITextureCreateD
 			bool bOnlyAddToLayoutManager, bDoInitialClear;
 			if (EnumHasAnyFlags(InCreateDesc.Flags, TexCreate_RenderTargetable | TexCreate_DepthStencilTargetable))
 			{
-				InitialLayout = EnumHasAnyFlags(InCreateDesc.Flags, TexCreate_DepthStencilTargetable) ? VK_IMAGE_LAYOUT_DEPTH_STENCIL_ATTACHMENT_OPTIMAL : VK_IMAGE_LAYOUT_COLOR_ATTACHMENT_OPTIMAL;
+				InitialLayout = EnumHasAnyFlags(InCreateDesc.Flags, TexCreate_DepthStencilTargetable) ? VK_IMAGE_LAYOUT_ATTACHMENT_OPTIMAL : VK_IMAGE_LAYOUT_COLOR_ATTACHMENT_OPTIMAL;
 				bOnlyAddToLayoutManager = false;
 				bDoInitialClear = true;
 			}
@@ -1860,22 +1838,26 @@ FVulkanTexture::FVulkanTexture(FVulkanDevice& InDevice, const FRHITextureCreateD
 				bDoInitialClear = false;
 			}
 
+			DefaultLayout = InitialLayout;
+
 			FRHICommandList& RHICmdList = FRHICommandListExecutor::GetImmediateCommandList();
 			if (!IsInRenderingThread() || (RHICmdList.Bypass() || !IsRunningRHIInSeparateThread()))
 			{
+				FVulkanCommandListContext& Context = InDevice.GetImmediateContext();
 				if (bOnlyAddToLayoutManager)
 				{
-					InDevice.GetImmediateContext().GetLayoutManager().GetOrAddFullLayout(*this, InitialLayout);
+					FVulkanCmdBuffer* CmdBuffer = Context.GetCommandBufferManager()->GetActiveCmdBuffer();
+					CmdBuffer->GetLayoutManager().SetFullLayout(*this, InitialLayout, true);
 				}
 				else
 				{
-					SetInitialImageState(InDevice.GetImmediateContext(), InitialLayout, true, InCreateDesc.ClearValue);
+					SetInitialImageState(Context, InitialLayout, true, InCreateDesc.ClearValue, false);
 				}
 			}
 			else
 			{
 				check(IsInRenderingThread());
-				ALLOC_COMMAND_CL(RHICmdList, FRHICommandSetInitialImageState)(this, InitialLayout, bOnlyAddToLayoutManager, bDoInitialClear, InCreateDesc.ClearValue);
+				ALLOC_COMMAND_CL(RHICmdList, FRHICommandSetInitialImageState)(this, InitialLayout, bOnlyAddToLayoutManager, bDoInitialClear, InCreateDesc.ClearValue, false);
 			}
 		}
 	}
@@ -1910,6 +1892,7 @@ FVulkanTexture::FVulkanTexture(FVulkanDevice& InDevice, const FRHITextureCreateD
 	, FullAspectMask(0)
 	, PartialAspectMask(0)
 	, CpuReadbackBuffer(nullptr)
+	, DefaultLayout(VK_IMAGE_LAYOUT_UNDEFINED)
 	, ImageOwnerType(EImageOwnerType::Aliased)
 {
 	VULKAN_TRACK_OBJECT_CREATE(FVulkanTexture, this);
@@ -1976,6 +1959,7 @@ void FVulkanTexture::AliasTextureResources(FTextureRHIRef& SrcTextureRHI)
 	}
 
 	AliasedTexture = SrcTexture;
+	DefaultLayout = SrcTexture->DefaultLayout;
 }
 
 void FVulkanTexture::DestroyViews()
@@ -2020,7 +2004,7 @@ void FVulkanTexture::InvalidateViews(FVulkanDevice& InDevice)
 
 void FVulkanTexture::Move(FVulkanDevice& InDevice, FVulkanCommandListContext& Context, FVulkanAllocation& NewAllocation)
 {
-	uint64 Size = GetMemorySize();
+	const uint64 Size = GetMemorySize();
 	static uint64 TotalSize = 0;
 	TotalSize += Size;
 	if (GVulkanLogDefrag)
@@ -2028,32 +2012,41 @@ void FVulkanTexture::Move(FVulkanDevice& InDevice, FVulkanCommandListContext& Co
 		UE_LOG(LogVulkanRHI, Display, TEXT("Moving Surface, %p <<-- %p    :::: %s\n"), NewAllocation.Offset, 42, *GetName().ToString());
 		UE_LOG(LogVulkanRHI, Display, TEXT("Moved %8.4fkb %8.4fkb   TB %p  :: IMG %p   %-40s\n"), Size / (1024.f), TotalSize / (1024.f), this, Image, *GetName().ToString());
 	}
-	{
-		const ETextureCreateFlags UEFlags = GetDesc().Flags;
-		const bool bRenderTarget = EnumHasAnyFlags(UEFlags, TexCreate_RenderTargetable | TexCreate_DepthStencilTargetable | TexCreate_ResolveTargetable);
-		const bool bUAV = EnumHasAnyFlags(UEFlags, TexCreate_UAV);
-		checkf(bRenderTarget || bUAV, TEXT("Surface must be a RenderTarget or a UAV in order to be moved.  UEFlags=0x%x"), (int32)UEFlags);
-		checkf(Tiling == VK_IMAGE_TILING_OPTIMAL, TEXT("Tiling [%s] is not supported for move, only VK_IMAGE_TILING_OPTIMAL"), VK_TYPE_TO_STRING(VkImageTiling, Tiling));
-	
-		InternalMoveSurface(InDevice, Context, NewAllocation);
+
+	// Move is used for defrag, which uses layouts stored in the queue
+	const FVulkanImageLayout* OriginalLayout = Context.GetQueue()->GetLayoutManager().GetFullLayout(Image);
+	check(OriginalLayout && OriginalLayout->AreAllSubresourcesSameLayout());
+
+	const ETextureCreateFlags UEFlags = GetDesc().Flags;
+	const bool bRenderTarget = EnumHasAnyFlags(UEFlags, TexCreate_RenderTargetable | TexCreate_DepthStencilTargetable | TexCreate_ResolveTargetable);
+	const bool bUAV = EnumHasAnyFlags(UEFlags, TexCreate_UAV);
+	checkf(bRenderTarget || bUAV, TEXT("Surface must be a RenderTarget or a UAV in order to be moved.  UEFlags=0x%x"), (int32)UEFlags);
+	checkf(Tiling == VK_IMAGE_TILING_OPTIMAL, TEXT("Tiling [%s] is not supported for move, only VK_IMAGE_TILING_OPTIMAL"), VK_TYPE_TO_STRING(VkImageTiling, Tiling));
+	checkf((OriginalLayout->NumMips == GetNumMips()), TEXT("NumMips reported by LayoutManager (%d) differs from surface (%d)"), OriginalLayout->NumMips, GetNumMips());
+	checkf((OriginalLayout->NumLayers == GetNumberOfArrayLevels()), TEXT("NumLayers reported by LayoutManager (%d) differs from surface (%d)"), OriginalLayout->NumLayers, GetNumberOfArrayLevels());
+
+	InternalMoveSurface(InDevice, Context, NewAllocation, OriginalLayout->MainLayout);
 	
-		// Swap in the new allocation for this surface
-		Allocation.Swap(NewAllocation);
-	}
+	// Swap in the new allocation for this surface
+	Allocation.Swap(NewAllocation);
+
 	InvalidateViews(InDevice);
 }
 
 void FVulkanTexture::Evict(FVulkanDevice& InDevice, FVulkanCommandListContext& Context)
 {
 	check(AliasedTexture == nullptr); //can't evict textures we don't own
-	uint64 Size = GetMemorySize();
+	const uint64 Size = GetMemorySize();
 	static uint64 TotalSize = 0;
 	TotalSize += Size;
-	if(GVulkanLogDefrag)
+	if (GVulkanLogDefrag)
 	{
 		FGenericPlatformMisc::LowLevelOutputDebugStringf(TEXT("Evicted %8.4fkb %8.4fkb   TB %p  :: IMG %p   %-40s\n"), Size / (1024.f), TotalSize / (1024.f), this, Image, *GetName().ToString());
 	}
 
+	// Eviction layouts are read from the queue
+	const FVulkanImageLayout* OriginalLayout = Context.GetQueue()->GetLayoutManager().GetFullLayout(Image);
+	if (OriginalLayout && OriginalLayout->AreAllSubresourcesSameLayout())
 	{
 		check(0 == CpuReadbackBuffer);
 		checkf(MemProps == VK_MEMORY_PROPERTY_DEVICE_LOCAL_BIT, TEXT("Can't evict surface that isn't device local.  MemoryProperties=%s"), VK_FLAGS_TO_STRING(VkMemoryPropertyFlags, MemProps));
@@ -2066,6 +2059,8 @@ void FVulkanTexture::Evict(FVulkanDevice& InDevice, FVulkanCommandListContext& C
 		//none of this is supported for eviction
 		checkf(!bRenderTarget, TEXT("RenderTargets do not support evict."));
 		checkf(!bUAV, TEXT("UAV do not support evict."));
+		checkf((OriginalLayout->NumMips == GetNumMips()), TEXT("NumMips reported by LayoutManager (%d) differs from surface (%d)"), OriginalLayout->NumMips, GetNumMips());
+		checkf((OriginalLayout->NumLayers == GetNumberOfArrayLevels()), TEXT("NumLayers reported by LayoutManager (%d) differs from surface (%d)"), OriginalLayout->NumLayers, GetNumberOfArrayLevels());
 
 		MemProps = InDevice.GetDeviceMemoryManager().GetEvictedMemoryProperties();
 
@@ -2078,16 +2073,16 @@ void FVulkanTexture::Evict(FVulkanDevice& InDevice, FVulkanCommandListContext& C
 			checkNoEntry();
 		}
 
-		InternalMoveSurface(InDevice, Context, HostAllocation);
+		InternalMoveSurface(InDevice, Context, HostAllocation, OriginalLayout->MainLayout);
 
 		// Delete the original allocation and swap in the new host allocation
 		Device->GetMemoryManager().FreeVulkanAllocation(Allocation);
 		Allocation.Swap(HostAllocation);
 
 		VULKAN_SET_DEBUG_NAME(InDevice, VK_OBJECT_TYPE_IMAGE, Image, TEXT("(FVulkanTexture*)0x%p [hostimage]"), this);
-	}
 
-	InvalidateViews(InDevice);
+		InvalidateViews(InDevice);
+	}
 }
 
 
@@ -2334,12 +2329,13 @@ void FVulkanCommandListContext::RHICopyTexture(FRHITexture* SourceTexture, FRHIT
 	FVulkanTexture* Source = FVulkanTexture::Cast(SourceTexture);
 	FVulkanTexture* Dest = FVulkanTexture::Cast(DestTexture);
 
-	VkImageLayout SrcLayout = LayoutManager.FindLayoutChecked(Source->Image);
-	ensureMsgf(SrcLayout == VK_IMAGE_LAYOUT_TRANSFER_SRC_OPTIMAL, TEXT("Expected source texture to be in VK_IMAGE_LAYOUT_TRANSFER_SRC_OPTIMAL, actual layout is %s"), VK_TYPE_TO_STRING(VkImageLayout, SrcLayout));
+	FVulkanCmdBuffer* CmdBuffer = CommandBufferManager->GetActiveCmdBuffer();
+	check(CmdBuffer->IsOutsideRenderPass());
 
-	FVulkanCmdBuffer* InCmdBuffer = CommandBufferManager->GetActiveCmdBuffer();
-	check(InCmdBuffer->IsOutsideRenderPass());
-	VkCommandBuffer CmdBuffer = InCmdBuffer->GetHandle();
+	{
+		const VkImageLayout ExpectedSrcLayout = FVulkanLayoutManager::SetExpectedLayout(CmdBuffer, *Source, ERHIAccess::CopySrc);
+		ensureMsgf((ExpectedSrcLayout == VK_IMAGE_LAYOUT_TRANSFER_SRC_OPTIMAL), TEXT("Expected source texture to be in VK_IMAGE_LAYOUT_TRANSFER_SRC_OPTIMAL, actual layout is %s"), VK_TYPE_TO_STRING(VkImageLayout, ExpectedSrcLayout));
+	}
 
 	const FPixelFormatInfo& PixelFormatInfo = GPixelFormats[DestTexture->GetDesc().Format];
 	const FRHITextureDesc& SourceDesc = SourceTexture->GetDesc();
@@ -2348,7 +2344,7 @@ void FVulkanCommandListContext::RHICopyTexture(FRHITexture* SourceTexture, FRHIT
 	const FIntVector DestXYZ = DestDesc.GetSize();
 
 	check(!EnumHasAnyFlags(Source->GetDesc().Flags, TexCreate_CPUReadback));
-	if(EnumHasAllFlags(Dest->GetDesc().Flags, TexCreate_CPUReadback))
+	if (EnumHasAllFlags(Dest->GetDesc().Flags, TexCreate_CPUReadback))
 	{
 		check(CopyInfo.DestSliceIndex == 0); //slices not supported in TexCreate_CPUReadback textures.
 		FIntVector Size = CopyInfo.Size;
@@ -2384,7 +2380,7 @@ void FVulkanCommandListContext::RHICopyTexture(FRHITexture* SourceTexture, FRHIT
 			Size.Z = FMath::Max(1, Size.Z / 2);
 		}
 
-		VulkanRHI::vkCmdCopyImageToBuffer(CmdBuffer, Source->Image, VK_IMAGE_LAYOUT_TRANSFER_SRC_OPTIMAL, CpuReadbackBuffer->Buffer, CopyInfo.NumMips, &CopyRegion[0]);
+		VulkanRHI::vkCmdCopyImageToBuffer(CmdBuffer->GetHandle(), Source->Image, VK_IMAGE_LAYOUT_TRANSFER_SRC_OPTIMAL, CpuReadbackBuffer->Buffer, CopyInfo.NumMips, &CopyRegion[0]);
 
 		FVulkanPipelineBarrier BarrierMemory;
 		BarrierMemory.AddMemoryBarrier(VK_ACCESS_TRANSFER_WRITE_BIT, VK_ACCESS_HOST_READ_BIT, VK_PIPELINE_STAGE_ALL_COMMANDS_BIT, VK_PIPELINE_STAGE_HOST_BIT);
@@ -2392,8 +2388,10 @@ void FVulkanCommandListContext::RHICopyTexture(FRHITexture* SourceTexture, FRHIT
 	}
 	else
 	{
-		VkImageLayout DstLayout = LayoutManager.FindLayoutChecked(Dest->Image);
-		ensureMsgf(DstLayout == VK_IMAGE_LAYOUT_TRANSFER_DST_OPTIMAL, TEXT("Expected destination texture to be in VK_IMAGE_LAYOUT_TRANSFER_DST_OPTIMAL, actual layout is %s"), VK_TYPE_TO_STRING(VkImageLayout, DstLayout));
+		{
+			const VkImageLayout ExpectedDstLayout = FVulkanLayoutManager::SetExpectedLayout(CmdBuffer, *Dest, ERHIAccess::CopyDest);
+			ensureMsgf((ExpectedDstLayout == VK_IMAGE_LAYOUT_TRANSFER_DST_OPTIMAL), TEXT("Expected destination texture to be in VK_IMAGE_LAYOUT_TRANSFER_DST_OPTIMAL, actual layout is %s"), VK_TYPE_TO_STRING(VkImageLayout, ExpectedDstLayout));
+		}
 
 		VkImageCopy Region;
 		FMemory::Memzero(Region);
@@ -2429,7 +2427,7 @@ void FVulkanCommandListContext::RHICopyTexture(FRHITexture* SourceTexture, FRHIT
 
 		for (uint32 Index = 0; Index < CopyInfo.NumMips; ++Index)
 		{
-			VulkanRHI::vkCmdCopyImage(CmdBuffer,
+			VulkanRHI::vkCmdCopyImage(CmdBuffer->GetHandle(),
 				Source->Image, VK_IMAGE_LAYOUT_TRANSFER_SRC_OPTIMAL,
 				Dest->Image, VK_IMAGE_LAYOUT_TRANSFER_DST_OPTIMAL,
 				1, &Region);
diff --git a/Engine/Source/Runtime/VulkanRHI/Private/VulkanUtil.cpp b/Engine/Source/Runtime/VulkanRHI/Private/VulkanUtil.cpp
--- a/Engine/Source/Runtime/VulkanRHI/Private/VulkanUtil.cpp
+++ b/Engine/Source/Runtime/VulkanRHI/Private/VulkanUtil.cpp
@@ -425,8 +425,8 @@ uint64 FVulkanGPUTiming::GetTiming(bool bGetCurrentResultsAndBlock)
 
 				SCOPE_CYCLE_COUNTER(STAT_RenderQueryResultTime);
 
-				bool bWaitForStart = StartQuerySyncPoint.FenceCounter == StartQuerySyncPoint.CmdBuffer->GetFenceSignaledCounter();
-				bool bWaitForEnd = EndQuerySyncPoint.FenceCounter == EndQuerySyncPoint.CmdBuffer->GetFenceSignaledCounter();
+				const bool bWaitForStart = StartQuerySyncPoint.FenceCounter >= StartQuerySyncPoint.CmdBuffer->GetFenceSignaledCounter();
+				const bool bWaitForEnd = EndQuerySyncPoint.FenceCounter >= EndQuerySyncPoint.CmdBuffer->GetFenceSignaledCounter();
 				if (bWaitForEnd || bWaitForStart)
 				{
 					// Need to submit the open command lists.
@@ -434,11 +434,11 @@ uint64 FVulkanGPUTiming::GetTiming(bool bGetCurrentResultsAndBlock)
 				}
 
 				// CPU wait for query results to be ready.
-				if (bWaitForStart && StartQuerySyncPoint.FenceCounter == StartQuerySyncPoint.CmdBuffer->GetFenceSignaledCounter())
+				if (bWaitForStart && StartQuerySyncPoint.FenceCounter >= StartQuerySyncPoint.CmdBuffer->GetFenceSignaledCounter())
 				{
 					CmdContext->GetCommandBufferManager()->WaitForCmdBuffer(StartQuerySyncPoint.CmdBuffer);
 				}
-				if (bWaitForEnd && EndQuerySyncPoint.FenceCounter == EndQuerySyncPoint.CmdBuffer->GetFenceSignaledCounter())
+				if (bWaitForEnd && EndQuerySyncPoint.FenceCounter >= EndQuerySyncPoint.CmdBuffer->GetFenceSignaledCounter())
 				{
 					CmdContext->GetCommandBufferManager()->WaitForCmdBuffer(EndQuerySyncPoint.CmdBuffer);
 				}
@@ -1049,6 +1049,7 @@ DEFINE_STAT(STAT_VulkanUPPrepTime);
 DEFINE_STAT(STAT_VulkanUniformBufferCreateTime);
 DEFINE_STAT(STAT_VulkanApplyDSUniformBuffers);
 DEFINE_STAT(STAT_VulkanApplyPackedUniformBuffers);
+DEFINE_STAT(STAT_VulkanBarrierTime);
 DEFINE_STAT(STAT_VulkanSRVUpdateTime);
 DEFINE_STAT(STAT_VulkanUAVUpdateTime);
 DEFINE_STAT(STAT_VulkanDeletionQueue);
diff --git a/Engine/Source/Runtime/VulkanRHI/Private/VulkanViewport.cpp b/Engine/Source/Runtime/VulkanRHI/Private/VulkanViewport.cpp
--- a/Engine/Source/Runtime/VulkanRHI/Private/VulkanViewport.cpp
+++ b/Engine/Source/Runtime/VulkanRHI/Private/VulkanViewport.cpp
@@ -74,15 +74,15 @@ void FVulkanBackBuffer::AcquireBackBufferImage(FVulkanCommandListContext& Contex
 			DefaultView.View = ImageView.View;
 			DefaultView.ViewId = ImageView.ViewId;
 
-			// right after acquiring image is in undefined state
-			FVulkanLayoutManager& LayoutMgr = Context.GetLayoutManager();
-			VkImageLayout& CurrentLayout = LayoutMgr.FindOrAddLayoutRW(ImageView.Image, VK_IMAGE_LAYOUT_UNDEFINED, 1, 1);
-			CurrentLayout = VK_IMAGE_LAYOUT_UNDEFINED;
-
 			FVulkanCommandBufferManager* CmdBufferManager = Context.GetCommandBufferManager();
 			FVulkanCmdBuffer* CmdBuffer = CmdBufferManager->GetActiveCmdBuffer();
 			check(!CmdBuffer->IsInsideRenderPass());
 
+			// right after acquiring image is in undefined state
+			FVulkanLayoutManager& LayoutMgr = CmdBuffer->GetLayoutManager();
+			const FVulkanImageLayout CustomLayout(VK_IMAGE_LAYOUT_UNDEFINED, 1, 1, VK_IMAGE_ASPECT_COLOR_BIT);
+			LayoutMgr.SetFullLayout(ImageView.Image, CustomLayout);
+
 			// Wait for semaphore signal before writing to backbuffer image
 			CmdBuffer->AddWaitSemaphore(VK_PIPELINE_STAGE_COLOR_ATTACHMENT_OUTPUT_BIT, Viewport->AcquiredSemaphore);
 		}
@@ -641,6 +641,11 @@ void FVulkanViewport::CreateSwapchain(FVulkanSwapChainRecreateInfo* RecreateInfo
 		FVulkanCmdBuffer* CmdBuffer = Device->GetImmediateContext().GetCommandBufferManager()->GetUploadCmdBuffer();
 		ensure(CmdBuffer->IsOutsideRenderPass());
 
+		VkClearColorValue ClearColor;
+		FMemory::Memzero(ClearColor);
+
+		const VkImageSubresourceRange Range = FVulkanPipelineBarrier::MakeSubresourceRange(VK_IMAGE_ASPECT_COLOR_BIT);
+		const FVulkanImageLayout InitialLayout(VK_IMAGE_LAYOUT_COLOR_ATTACHMENT_OPTIMAL, 1, 1, VK_IMAGE_ASPECT_COLOR_BIT);
 		for (int32 Index = 0; Index < Images.Num(); ++Index)
 		{
 			BackBufferImages[Index] = Images[Index];
@@ -648,13 +653,10 @@ void FVulkanViewport::CreateSwapchain(FVulkanSwapChainRecreateInfo* RecreateInfo
 
 			// Clear the swapchain to avoid a validation warning, and transition to ColorAttachment
 			{
-				VkImageSubresourceRange Range = FVulkanPipelineBarrier::MakeSubresourceRange(VK_IMAGE_ASPECT_COLOR_BIT);
-
-				VkClearColorValue Color;
-				FMemory::Memzero(Color);
-				VulkanSetImageLayout(CmdBuffer->GetHandle(), Images[Index], VK_IMAGE_LAYOUT_UNDEFINED, VK_IMAGE_LAYOUT_TRANSFER_DST_OPTIMAL, Range);
-				VulkanRHI::vkCmdClearColorImage(CmdBuffer->GetHandle(), Images[Index], VK_IMAGE_LAYOUT_TRANSFER_DST_OPTIMAL, &Color, 1, &Range);
-				VulkanSetImageLayout(CmdBuffer->GetHandle(), Images[Index], VK_IMAGE_LAYOUT_TRANSFER_DST_OPTIMAL, VK_IMAGE_LAYOUT_COLOR_ATTACHMENT_OPTIMAL, Range);
+				VulkanSetImageLayout(CmdBuffer, Images[Index], VK_IMAGE_LAYOUT_UNDEFINED, VK_IMAGE_LAYOUT_TRANSFER_DST_OPTIMAL, Range);
+				VulkanRHI::vkCmdClearColorImage(CmdBuffer->GetHandle(), Images[Index], VK_IMAGE_LAYOUT_TRANSFER_DST_OPTIMAL, &ClearColor, 1, &Range);
+				VulkanSetImageLayout(CmdBuffer, Images[Index], VK_IMAGE_LAYOUT_TRANSFER_DST_OPTIMAL, VK_IMAGE_LAYOUT_COLOR_ATTACHMENT_OPTIMAL, Range);
+				CmdBuffer->GetLayoutManager().SetFullLayout(Images[Index], InitialLayout);
 			}
 
 #if VULKAN_ENABLE_DRAW_MARKERS
@@ -746,16 +748,15 @@ void FVulkanViewport::DestroySwapchain(FVulkanSwapChainRecreateInfo* RecreateInf
 	AcquiredImageIndex = -1;
 }
 
-inline static void CopyImageToBackBuffer(FVulkanCommandListContext* Context, FVulkanCmdBuffer* CmdBuffer, VkImage SrcSurface, VkImage DstSurface, int32 SizeX, int32 SizeY, int32 WindowSizeX, int32 WindowSizeY)
+inline static void CopyImageToBackBuffer(FVulkanCommandListContext* Context, FVulkanCmdBuffer* CmdBuffer, FVulkanTexture& SrcSurface, VkImage DstSurface, int32 SizeX, int32 SizeY, int32 WindowSizeX, int32 WindowSizeY)
 {
-	FVulkanLayoutManager& LayoutManager = Context->GetLayoutManager();
-	VkImageLayout SrcLayout = LayoutManager.FindLayoutChecked(SrcSurface);
+	FVulkanLayoutManager& LayoutManager = CmdBuffer->GetLayoutManager();
+	const VkImageLayout PreviousSrcLayout = FVulkanLayoutManager::SetExpectedLayout(CmdBuffer, SrcSurface, ERHIAccess::CopySrc);
 
 	{
 		FVulkanPipelineBarrier Barrier;
-		Barrier.AddImageLayoutTransition(SrcSurface, SrcLayout, VK_IMAGE_LAYOUT_TRANSFER_SRC_OPTIMAL, FVulkanPipelineBarrier::MakeSubresourceRange(VK_IMAGE_ASPECT_COLOR_BIT, 0, 1));
 		Barrier.AddImageLayoutTransition(DstSurface, VK_IMAGE_LAYOUT_UNDEFINED, VK_IMAGE_LAYOUT_TRANSFER_DST_OPTIMAL, FVulkanPipelineBarrier::MakeSubresourceRange(VK_IMAGE_ASPECT_COLOR_BIT, 0, 1));
-		Barrier.Execute(CmdBuffer->GetHandle());
+		Barrier.Execute(CmdBuffer);
 	}
 
 	VulkanRHI::DebugHeavyWeightBarrier(CmdBuffer->GetHandle(), 32);
@@ -782,7 +783,7 @@ inline static void CopyImageToBackBuffer(FVulkanCommandListContext* Context, FVu
 		Region.dstSubresource.baseArrayLayer = 0;
 		Region.dstSubresource.layerCount = 1;
 		VulkanRHI::vkCmdBlitImage(CmdBuffer->GetHandle(),
-			SrcSurface, VK_IMAGE_LAYOUT_TRANSFER_SRC_OPTIMAL,
+			SrcSurface.Image, VK_IMAGE_LAYOUT_TRANSFER_SRC_OPTIMAL,
 			DstSurface, VK_IMAGE_LAYOUT_TRANSFER_DST_OPTIMAL,
 			1, &Region, VK_FILTER_LINEAR);
 	}
@@ -802,16 +803,16 @@ inline static void CopyImageToBackBuffer(FVulkanCommandListContext* Context, FVu
 		Region.dstSubresource.layerCount = 1;
 		//Region.dstSubresource.mipLevel = 0;
 		VulkanRHI::vkCmdCopyImage(CmdBuffer->GetHandle(),
-			SrcSurface, VK_IMAGE_LAYOUT_TRANSFER_SRC_OPTIMAL,
+			SrcSurface.Image, VK_IMAGE_LAYOUT_TRANSFER_SRC_OPTIMAL,
 			DstSurface, VK_IMAGE_LAYOUT_TRANSFER_DST_OPTIMAL,
 			1, &Region);
 	}
 
 	{
 		FVulkanPipelineBarrier Barrier;
-		Barrier.AddImageLayoutTransition(SrcSurface, VK_IMAGE_LAYOUT_TRANSFER_SRC_OPTIMAL, SrcLayout, FVulkanPipelineBarrier::MakeSubresourceRange(VK_IMAGE_ASPECT_COLOR_BIT, 0, 1));
+		Barrier.AddImageLayoutTransition(SrcSurface.Image, VK_IMAGE_LAYOUT_TRANSFER_SRC_OPTIMAL, PreviousSrcLayout, FVulkanPipelineBarrier::MakeSubresourceRange(VK_IMAGE_ASPECT_COLOR_BIT, 0, 1));
 		Barrier.AddImageLayoutTransition(DstSurface, VK_IMAGE_LAYOUT_TRANSFER_DST_OPTIMAL, VK_IMAGE_LAYOUT_PRESENT_SRC_KHR, FVulkanPipelineBarrier::MakeSubresourceRange(VK_IMAGE_ASPECT_COLOR_BIT, 0, 1));
-		Barrier.Execute(CmdBuffer->GetHandle());
+		Barrier.Execute(CmdBuffer);
 	}
 }
 
@@ -835,7 +836,7 @@ bool FVulkanViewport::Present(FVulkanCommandListContext* Context, FVulkanCmdBuff
 				uint32 WindowSizeY = FMath::Min(SizeY, SwapChain->InternalHeight);
 
 				Context->RHIPushEvent(TEXT("CopyImageToBackBuffer"), FColor::Blue);
-				CopyImageToBackBuffer(Context, CmdBuffer, RenderingBackBuffer->Image, BackBufferImages[AcquiredImageIndex], SizeX, SizeY, WindowSizeX, WindowSizeY);
+				CopyImageToBackBuffer(Context, CmdBuffer, *RenderingBackBuffer.GetReference(), BackBufferImages[AcquiredImageIndex], SizeX, SizeY, WindowSizeX, WindowSizeY);
 				Context->RHIPopEvent();
 			}
 			else
@@ -848,9 +849,21 @@ bool FVulkanViewport::Present(FVulkanCommandListContext* Context, FVulkanCmdBuff
 			if (AcquiredImageIndex != -1)
 			{
 				check(RHIBackBuffer != nullptr && RHIBackBuffer->Image == BackBufferImages[AcquiredImageIndex]);
-				VkImageLayout& Layout = Context->GetLayoutManager().FindOrAddLayoutRW(BackBufferImages[AcquiredImageIndex], VK_IMAGE_LAYOUT_UNDEFINED, 1, 1);
-				VulkanSetImageLayout(CmdBuffer->GetHandle(), BackBufferImages[AcquiredImageIndex], Layout, VK_IMAGE_LAYOUT_PRESENT_SRC_KHR, FVulkanPipelineBarrier::MakeSubresourceRange(VK_IMAGE_ASPECT_COLOR_BIT));
-				Layout = VK_IMAGE_LAYOUT_UNDEFINED;
+
+				FVulkanLayoutManager& LayoutManager = CmdBuffer->GetLayoutManager();
+				const FVulkanImageLayout* TrackedLayout = LayoutManager.GetFullLayout(BackBufferImages[AcquiredImageIndex]);
+
+				// One of the rare cases we'll let parallel rendering check the tracking (legacy path already checked as a fallback)
+				if (!TrackedLayout && Device->SupportsParallelRendering())
+				{
+					TrackedLayout = Context->GetQueue()->GetLayoutManager().GetFullLayout(BackBufferImages[AcquiredImageIndex]);
+				}
+
+				VkImageLayout LastLayout = TrackedLayout ? TrackedLayout->MainLayout : VK_IMAGE_LAYOUT_UNDEFINED;
+				VulkanSetImageLayout(CmdBuffer, BackBufferImages[AcquiredImageIndex], LastLayout, VK_IMAGE_LAYOUT_PRESENT_SRC_KHR, FVulkanPipelineBarrier::MakeSubresourceRange(VK_IMAGE_ASPECT_COLOR_BIT));
+
+				const FVulkanImageLayout UndefinedLayout(VK_IMAGE_LAYOUT_UNDEFINED, 1, 1, VK_IMAGE_ASPECT_COLOR_BIT);
+				LayoutManager.SetFullLayout(BackBufferImages[AcquiredImageIndex], UndefinedLayout);
 			}
 			else
 			{
diff --git a/Engine/Source/Runtime/VulkanRHI/Public/IVulkanDynamicRHI.h b/Engine/Source/Runtime/VulkanRHI/Public/IVulkanDynamicRHI.h
--- a/Engine/Source/Runtime/VulkanRHI/Public/IVulkanDynamicRHI.h
+++ b/Engine/Source/Runtime/VulkanRHI/Public/IVulkanDynamicRHI.h
@@ -69,7 +69,9 @@ struct IVulkanDynamicRHI : public FDynamicRHI
 	virtual FVulkanRHIAllocationInfo RHIGetAllocationInfo(FRHITexture* InTexture) const = 0;
 	virtual FVulkanRHIImageViewInfo  RHIGetImageViewInfo(FRHITexture* InTexture) const = 0;
 
+	UE_DEPRECATED(5.2, "Altering image layout tracking directly is not permitted anymore.")
 	virtual VkImageLayout& RHIFindOrAddLayoutRW(FRHITexture* InTexture, VkImageLayout LayoutIfNotFound) = 0;
+
 	virtual void           RHISetImageLayout(VkImage Image, VkImageLayout OldLayout, VkImageLayout NewLayout, const VkImageSubresourceRange& SubresourceRange) = 0;
 	virtual void           RHISetUploadImageLayout(VkImage Image, VkImageLayout OldLayout, VkImageLayout NewLayout, const VkImageSubresourceRange& SubresourceRange) = 0;
 
diff --git a/Engine/Source/Runtime/VulkanRHI/Public/VulkanConfiguration.h b/Engine/Source/Runtime/VulkanRHI/Public/VulkanConfiguration.h
--- a/Engine/Source/Runtime/VulkanRHI/Public/VulkanConfiguration.h
+++ b/Engine/Source/Runtime/VulkanRHI/Public/VulkanConfiguration.h
@@ -52,10 +52,6 @@
 	#define VULKAN_FREEPAGE_FOR_TYPE							0
 #endif
 
-#ifndef VULKAN_USE_NEW_QUERIES
-	#define VULKAN_USE_NEW_QUERIES								1
-#endif
-
 #ifndef VULKAN_SHOULD_USE_LLM
 	#define VULKAN_SHOULD_USE_LLM								0
 #endif
@@ -295,14 +291,6 @@
 	#endif
 #endif
 
-#ifndef VULKAN_SUPPORTS_SEPARATE_DEPTH_STENCIL_LAYOUTS
-	#ifdef VK_KHR_separate_depth_stencil_layouts
-		#define VULKAN_SUPPORTS_SEPARATE_DEPTH_STENCIL_LAYOUTS	1
-	#else
-		#define VULKAN_SUPPORTS_SEPARATE_DEPTH_STENCIL_LAYOUTS	0
-	#endif
-#endif
-
 #ifdef VK_KHR_shader_atomic_int64
 	#define VULKAN_SUPPORTS_BUFFER_64BIT_ATOMICS	(VULKAN_SUPPORTS_PHYSICAL_DEVICE_PROPERTIES2)	// Requirement
 #else
diff --git a/Engine/Source/Runtime/VulkanRHI/Public/VulkanDynamicRHI.h b/Engine/Source/Runtime/VulkanRHI/Public/VulkanDynamicRHI.h
--- a/Engine/Source/Runtime/VulkanRHI/Public/VulkanDynamicRHI.h
+++ b/Engine/Source/Runtime/VulkanRHI/Public/VulkanDynamicRHI.h
@@ -71,7 +71,10 @@ public:
 	virtual VkFormat RHIGetViewVkFormat(FRHITexture* InTexture) const final override;
 	virtual FVulkanRHIAllocationInfo RHIGetAllocationInfo(FRHITexture* InTexture) const final override;
 	virtual FVulkanRHIImageViewInfo RHIGetImageViewInfo(FRHITexture* InTexture) const final override;
+
+	UE_DEPRECATED(5.2, "Altering image layout tracking directly is not permitted anymore.")
 	virtual VkImageLayout& RHIFindOrAddLayoutRW(FRHITexture* InTexture, VkImageLayout LayoutIfNotFound) final override;
+
 	virtual void RHISetImageLayout(VkImage Image, VkImageLayout OldLayout, VkImageLayout NewLayout, const VkImageSubresourceRange& SubresourceRange) final override;
 	virtual void RHISetUploadImageLayout(VkImage Image, VkImageLayout OldLayout, VkImageLayout NewLayout, const VkImageSubresourceRange& SubresourceRange) final override;
 	virtual void RHIFinishExternalComputeWork(VkCommandBuffer InCommandBuffer) final override;
