VulkanRHI:  Fix build (errors in checkSlow)

diff --git a/Engine/Source/Runtime/VulkanRHI/Private/VulkanBarriers.cpp b/Engine/Source/Runtime/VulkanRHI/Private/VulkanBarriers.cpp
--- a/Engine/Source/Runtime/VulkanRHI/Private/VulkanBarriers.cpp
+++ b/Engine/Source/Runtime/VulkanRHI/Private/VulkanBarriers.cpp
@@ -1093,10 +1093,10 @@ private:
 		switch (RequestedPipeline)
 		{
 		case ERHIPipeline::Graphics:
-			checkf(!Device->IsRealAsyncComputeContext(this), TEXT("The %s pipeline for this transition is the Graphics queue, but it's submitted on the AsyncCompute queue."), IsBeginTransition ? TEXT("BEGIN") : TEXT("END"));
+			checkf(!Context.GetDevice()->IsRealAsyncComputeContext(this), TEXT("The %s pipeline for this transition is the Graphics queue, but it's submitted on the AsyncCompute queue."), IsBeginTransition ? TEXT("BEGIN") : TEXT("END"));
 			break;
 		case ERHIPipeline::AsyncCompute:
-			checkf(Device->IsRealAsyncComputeContext(this), TEXT("The %s pipeline for this transition is the AsyncCompute queue, but it's submitted on the Graphics queue."), IsBeginTransition ? TEXT("BEGIN") : TEXT("END"));
+			checkf(Context.GetDevice()->IsRealAsyncComputeContext(this), TEXT("The %s pipeline for this transition is the AsyncCompute queue, but it's submitted on the Graphics queue."), IsBeginTransition ? TEXT("BEGIN") : TEXT("END"));
 			break;
 		default:
 			checkNoEntry();
@@ -1239,7 +1239,7 @@ void FTransitionProcessor<VkMemoryBarrier, VkBufferMemoryBarrier, VkImageMemoryB
 		}
 		else
 		{
-			checkSlow(ExtraData.IsAliasingBarrier || Layout.AreSubresourcesSameLayout(ImageBarrier.oldLayout, ImageBarrier.subresourceRange));
+			checkSlow(RemainingExtras->IsAliasingBarrier || Layout.AreSubresourcesSameLayout(ImageBarrier.oldLayout, ImageBarrier.subresourceRange));
 		}
 
 		MergedDstStageMask |= GetVkStageFlagsForLayout(ImageBarrier.newLayout);
