VulkanRHI: Fix static analysis warnings from 22531238.

diff --git a/Engine/Source/Runtime/VulkanRHI/Private/VulkanBarriers.cpp b/Engine/Source/Runtime/VulkanRHI/Private/VulkanBarriers.cpp
--- a/Engine/Source/Runtime/VulkanRHI/Private/VulkanBarriers.cpp
+++ b/Engine/Source/Runtime/VulkanRHI/Private/VulkanBarriers.cpp
@@ -1093,10 +1093,10 @@ private:
 		switch (RequestedPipeline)
 		{
 		case ERHIPipeline::Graphics:
-			checkf(!Context.GetDevice()->IsRealAsyncComputeContext(this), TEXT("The %s pipeline for this transition is the Graphics queue, but it's submitted on the AsyncCompute queue."), IsBeginTransition ? TEXT("BEGIN") : TEXT("END"));
+			checkf(!Context.GetDevice()->IsRealAsyncComputeContext(&Context), TEXT("The %s pipeline for this transition is the Graphics queue, but it's submitted on the AsyncCompute queue."), IsBeginTransition ? TEXT("BEGIN") : TEXT("END"));
 			break;
 		case ERHIPipeline::AsyncCompute:
-			checkf(Context.GetDevice()->IsRealAsyncComputeContext(this), TEXT("The %s pipeline for this transition is the AsyncCompute queue, but it's submitted on the Graphics queue."), IsBeginTransition ? TEXT("BEGIN") : TEXT("END"));
+			checkf(Context.GetDevice()->IsRealAsyncComputeContext(&Context), TEXT("The %s pipeline for this transition is the AsyncCompute queue, but it's submitted on the Graphics queue."), IsBeginTransition ? TEXT("BEGIN") : TEXT("END"));
 			break;
 		default:
 			checkNoEntry();
diff --git a/Engine/Source/Runtime/VulkanRHI/Private/VulkanRenderpass.cpp b/Engine/Source/Runtime/VulkanRHI/Private/VulkanRenderpass.cpp
--- a/Engine/Source/Runtime/VulkanRHI/Private/VulkanRenderpass.cpp
+++ b/Engine/Source/Runtime/VulkanRHI/Private/VulkanRenderpass.cpp
@@ -63,6 +63,25 @@ FVulkanFramebuffer* FVulkanRenderPassManager::GetOrCreateFramebuffer(const FRHIS
 	}
 	RTLayoutHash = FCrc::MemCrc32(MipsAndSlicesValues, sizeof(MipsAndSlicesValues), RTLayoutHash);
 
+	auto FindFramebufferInList = [&](FFramebufferList* InFramebufferList)
+	{
+		FVulkanFramebuffer* OutFramebuffer = nullptr;
+
+		for (int32 Index = 0; Index < InFramebufferList->Framebuffer.Num(); ++Index)
+		{
+			const VkRect2D RenderArea = InFramebufferList->Framebuffer[Index]->GetRenderArea();
+
+			if (InFramebufferList->Framebuffer[Index]->Matches(RenderTargetsInfo) &&
+				((RTLayout.GetExtent2D().width == RenderArea.extent.width) && (RTLayout.GetExtent2D().height == RenderArea.extent.height) &&
+					(RTLayout.GetOffset2D().x == RenderArea.offset.x) && (RTLayout.GetOffset2D().y == RenderArea.offset.y)))
+			{
+				OutFramebuffer = InFramebufferList->Framebuffer[Index];
+				break;
+			}
+		}
+
+		return OutFramebuffer;
+	};
 
 	FFramebufferList** FoundFramebufferList = nullptr;
 	FFramebufferList* FramebufferList = nullptr;
@@ -75,16 +94,10 @@ FVulkanFramebuffer* FVulkanRenderPassManager::GetOrCreateFramebuffer(const FRHIS
 		{
 			FramebufferList = *FoundFramebufferList;
 
-			for (int32 Index = 0; Index < FramebufferList->Framebuffer.Num(); ++Index)
+			FVulkanFramebuffer* ExistingFramebuffer = FindFramebufferInList(FramebufferList);
+			if (ExistingFramebuffer)
 			{
-				const VkRect2D RenderArea = FramebufferList->Framebuffer[Index]->GetRenderArea();
-
-				if (FramebufferList->Framebuffer[Index]->Matches(RenderTargetsInfo) &&
-					((RTLayout.GetExtent2D().width == RenderArea.extent.width) && (RTLayout.GetExtent2D().height == RenderArea.extent.height) &&
-						(RTLayout.GetOffset2D().x == RenderArea.offset.x) && (RTLayout.GetOffset2D().y == RenderArea.offset.y)))
-				{
-					return FramebufferList->Framebuffer[Index];
-				}
+				return ExistingFramebuffer;
 			}
 		}
 	}
@@ -96,6 +109,15 @@ FVulkanFramebuffer* FVulkanRenderPassManager::GetOrCreateFramebuffer(const FRHIS
 		FramebufferList = new FFramebufferList;
 		Framebuffers.Add(RTLayoutHash, FramebufferList);
 	}
+	else
+	{
+		FramebufferList = *FoundFramebufferList;
+		FVulkanFramebuffer* ExistingFramebuffer = FindFramebufferInList(FramebufferList);
+		if (ExistingFramebuffer)
+		{
+			return ExistingFramebuffer;
+		}
+	}
 
 	FVulkanFramebuffer* Framebuffer = new FVulkanFramebuffer(*Device, RenderTargetsInfo, RTLayout, *RenderPass);
 	FramebufferList->Framebuffer.Add(Framebuffer);
