VulkanRHI: clean up header inclusion.

diff --git a/Engine/Source/Runtime/VulkanRHI/Private/VulkanContext.h b/Engine/Source/Runtime/VulkanRHI/Private/VulkanContext.h
--- a/Engine/Source/Runtime/VulkanRHI/Private/VulkanContext.h
+++ b/Engine/Source/Runtime/VulkanRHI/Private/VulkanContext.h
@@ -8,7 +8,6 @@
 
 #include "VulkanResources.h"
 #include "VulkanGPUProfiler.h"
-#include "VulkanRenderpass.h"
 
 class FVulkanDevice;
 class FVulkanCommandBufferManager;
diff --git a/Engine/Source/Runtime/VulkanRHI/Private/VulkanRenderTarget.cpp b/Engine/Source/Runtime/VulkanRHI/Private/VulkanRenderTarget.cpp
--- a/Engine/Source/Runtime/VulkanRHI/Private/VulkanRenderTarget.cpp
+++ b/Engine/Source/Runtime/VulkanRHI/Private/VulkanRenderTarget.cpp
@@ -11,6 +11,7 @@
 #include "VulkanSwapChain.h"
 #include "SceneUtils.h"
 #include "RHISurfaceDataConversion.h"
+#include "VulkanRenderpass.h"
 
 // Debug mode used as workaround when a DEVICE LOST occurs on alt+tab on some platforms
 // This is a workaround and may end up causing some hitches on the rendering thread
diff --git a/Engine/Source/Runtime/VulkanRHI/Private/VulkanRenderpass.cpp b/Engine/Source/Runtime/VulkanRHI/Private/VulkanRenderpass.cpp
--- a/Engine/Source/Runtime/VulkanRHI/Private/VulkanRenderpass.cpp
+++ b/Engine/Source/Runtime/VulkanRHI/Private/VulkanRenderpass.cpp
@@ -4,7 +4,9 @@
 	VulkanState.cpp: Vulkan state implementation.
 =============================================================================*/
 
+#include "VulkanRHIPrivate.h"
 #include "VulkanRenderpass.h"
+#include "VulkanContext.h"
 
 VkRenderPass CreateVulkanRenderPass(FVulkanDevice& InDevice, const FVulkanRenderTargetLayout& RTLayout)
 {
diff --git a/Engine/Source/Runtime/VulkanRHI/Private/VulkanRenderpass.h b/Engine/Source/Runtime/VulkanRHI/Private/VulkanRenderpass.h
--- a/Engine/Source/Runtime/VulkanRHI/Private/VulkanRenderpass.h
+++ b/Engine/Source/Runtime/VulkanRHI/Private/VulkanRenderpass.h
@@ -11,7 +11,8 @@
 #include "VulkanRHIPrivate.h"
 #include "VulkanResources.h"
 #include "VulkanPendingState.h"
-#include "VulkanContext.h"
+
+class FVulkanCommandListContext;
 
 template <typename TAttachmentReferenceType>
 struct FVulkanAttachmentReference
