VulkanRHI: Fix validation error where texture2darray was set in texture2d param where FRDGTextureSRVDesc::CreateForSlice was used.

diff --git a/Engine/Source/Runtime/VulkanRHI/Private/VulkanUAV.cpp b/Engine/Source/Runtime/VulkanRHI/Private/VulkanUAV.cpp
--- a/Engine/Source/Runtime/VulkanRHI/Private/VulkanUAV.cpp
+++ b/Engine/Source/Runtime/VulkanRHI/Private/VulkanUAV.cpp
@@ -229,7 +229,14 @@ void FVulkanShaderResourceView::UpdateView()
 				}
 			}
 
-			TextureView.Create(*Device, SourceTextureVK->Image, SourceTextureVK->GetViewType(), SourceTextureVK->GetPartialAspectMask(), Format, UEToVkTextureFormat(Format, bSRGB), MipLevel, NumMips, FirstArraySlice, ActualNumArraySlices, false);
+			// If we explicitely request a single slice of a 2darray (such as with FRDGTextureSRVDesc::CreateForSlice), return it as a regular 2d view
+			VkImageViewType ActualViewType = SourceTextureVK->GetViewType();
+			if ((NumArraySlices == 1) && (SourceTexture->GetDesc().Dimension == ETextureDimension::Texture2DArray))
+			{
+				ActualViewType = VK_IMAGE_VIEW_TYPE_2D;
+			}
+
+			TextureView.Create(*Device, SourceTextureVK->Image, ActualViewType, SourceTextureVK->GetPartialAspectMask(), Format, UEToVkTextureFormat(Format, bSRGB), MipLevel, NumMips, FirstArraySlice, ActualNumArraySlices, false);
 		}
 	}
 }
