VulkanRHI:  Don't assume SRGB formats are supported, create a forwarding table and fill it by checking VkFormatProperties.  Fix up PrintVulkanPixelFormatMappings to also print the SRGB mapping when there is one (and use VK_TYPE_TO_STRING on formats for readable output).  Remove duplicate print of EPixelFormatCapabilities.

diff --git a/Engine/Source/Runtime/VulkanRHI/Private/VulkanDevice.cpp b/Engine/Source/Runtime/VulkanRHI/Private/VulkanDevice.cpp
--- a/Engine/Source/Runtime/VulkanRHI/Private/VulkanDevice.cpp
+++ b/Engine/Source/Runtime/VulkanRHI/Private/VulkanDevice.cpp
@@ -59,6 +59,9 @@ void AftermathCrashDumpDescriptionCallback(PFN_GFSDK_Aftermath_AddGpuCrashDumpDe
 // Mirror GPixelFormats with format information for buffers
 VkFormat GVulkanBufferFormat[PF_MAX];
 
+// Mirror GPixelFormats with format information for buffers
+VkFormat GVulkanSRGBFormat[PF_MAX];
+
 EDelayAcquireImageType GVulkanDelayAcquireImage = EDelayAcquireImageType::DelayAcquire;
 
 TAutoConsoleVariable<int32> CVarDelayAcquireBackBuffer(
@@ -687,7 +690,78 @@ void FVulkanDevice::SetupFormats()
 		}
 	}
 
-#if !(UE_BUILD_SHIPPING || UE_BUILD_TEST)
+	// Verify the potential SRGB formats and fill GVulkanSRGBFormat
+	{
+		auto GetSRGBMapping = [this](const VkFormat InFormat)
+		{
+			VkFormat SRGBFormat = InFormat;
+			switch (InFormat)
+			{
+			case VK_FORMAT_B8G8R8A8_UNORM:				SRGBFormat = VK_FORMAT_B8G8R8A8_SRGB; break;
+			case VK_FORMAT_A8B8G8R8_UNORM_PACK32:		SRGBFormat = VK_FORMAT_A8B8G8R8_SRGB_PACK32; break;
+			case VK_FORMAT_R8_UNORM:					SRGBFormat = ((GMaxRHIFeatureLevel <= ERHIFeatureLevel::ES3_1) ? VK_FORMAT_R8_UNORM : VK_FORMAT_R8_SRGB); break;
+			case VK_FORMAT_R8G8_UNORM:					SRGBFormat = VK_FORMAT_R8G8_SRGB; break;
+			case VK_FORMAT_R8G8B8_UNORM:				SRGBFormat = VK_FORMAT_R8G8B8_SRGB; break;
+			case VK_FORMAT_R8G8B8A8_UNORM:				SRGBFormat = VK_FORMAT_R8G8B8A8_SRGB; break;
+			case VK_FORMAT_BC1_RGB_UNORM_BLOCK:			SRGBFormat = VK_FORMAT_BC1_RGB_SRGB_BLOCK; break;
+			case VK_FORMAT_BC1_RGBA_UNORM_BLOCK:		SRGBFormat = VK_FORMAT_BC1_RGBA_SRGB_BLOCK; break;
+			case VK_FORMAT_BC2_UNORM_BLOCK:				SRGBFormat = VK_FORMAT_BC2_SRGB_BLOCK; break;
+			case VK_FORMAT_BC3_UNORM_BLOCK:				SRGBFormat = VK_FORMAT_BC3_SRGB_BLOCK; break;
+			case VK_FORMAT_BC7_UNORM_BLOCK:				SRGBFormat = VK_FORMAT_BC7_SRGB_BLOCK; break;
+			case VK_FORMAT_ETC2_R8G8B8_UNORM_BLOCK:		SRGBFormat = VK_FORMAT_ETC2_R8G8B8_SRGB_BLOCK; break;
+			case VK_FORMAT_ETC2_R8G8B8A1_UNORM_BLOCK:	SRGBFormat = VK_FORMAT_ETC2_R8G8B8A1_SRGB_BLOCK; break;
+			case VK_FORMAT_ETC2_R8G8B8A8_UNORM_BLOCK:	SRGBFormat = VK_FORMAT_ETC2_R8G8B8A8_SRGB_BLOCK; break;
+			case VK_FORMAT_ASTC_4x4_UNORM_BLOCK:		SRGBFormat = VK_FORMAT_ASTC_4x4_SRGB_BLOCK; break;
+			case VK_FORMAT_ASTC_5x4_UNORM_BLOCK:		SRGBFormat = VK_FORMAT_ASTC_5x4_SRGB_BLOCK; break;
+			case VK_FORMAT_ASTC_5x5_UNORM_BLOCK:		SRGBFormat = VK_FORMAT_ASTC_5x5_SRGB_BLOCK; break;
+			case VK_FORMAT_ASTC_6x5_UNORM_BLOCK:		SRGBFormat = VK_FORMAT_ASTC_6x5_SRGB_BLOCK; break;
+			case VK_FORMAT_ASTC_6x6_UNORM_BLOCK:		SRGBFormat = VK_FORMAT_ASTC_6x6_SRGB_BLOCK; break;
+			case VK_FORMAT_ASTC_8x5_UNORM_BLOCK:		SRGBFormat = VK_FORMAT_ASTC_8x5_SRGB_BLOCK; break;
+			case VK_FORMAT_ASTC_8x6_UNORM_BLOCK:		SRGBFormat = VK_FORMAT_ASTC_8x6_SRGB_BLOCK; break;
+			case VK_FORMAT_ASTC_8x8_UNORM_BLOCK:		SRGBFormat = VK_FORMAT_ASTC_8x8_SRGB_BLOCK; break;
+			case VK_FORMAT_ASTC_10x5_UNORM_BLOCK:		SRGBFormat = VK_FORMAT_ASTC_10x5_SRGB_BLOCK; break;
+			case VK_FORMAT_ASTC_10x6_UNORM_BLOCK:		SRGBFormat = VK_FORMAT_ASTC_10x6_SRGB_BLOCK; break;
+			case VK_FORMAT_ASTC_10x8_UNORM_BLOCK:		SRGBFormat = VK_FORMAT_ASTC_10x8_SRGB_BLOCK; break;
+			case VK_FORMAT_ASTC_10x10_UNORM_BLOCK:		SRGBFormat = VK_FORMAT_ASTC_10x10_SRGB_BLOCK; break;
+			case VK_FORMAT_ASTC_12x10_UNORM_BLOCK:		SRGBFormat = VK_FORMAT_ASTC_12x10_SRGB_BLOCK; break;
+			case VK_FORMAT_ASTC_12x12_UNORM_BLOCK:		SRGBFormat = VK_FORMAT_ASTC_12x12_SRGB_BLOCK; break;
+				//		case VK_FORMAT_PVRTC1_2BPP_UNORM_BLOCK_IMG:	Format = VK_FORMAT_PVRTC1_2BPP_SRGB_BLOCK_IMG; break;
+				//		case VK_FORMAT_PVRTC1_4BPP_UNORM_BLOCK_IMG:	Format = VK_FORMAT_PVRTC1_4BPP_SRGB_BLOCK_IMG; break;
+				//		case VK_FORMAT_PVRTC2_2BPP_UNORM_BLOCK_IMG:	Format = VK_FORMAT_PVRTC2_2BPP_SRGB_BLOCK_IMG; break;
+				//		case VK_FORMAT_PVRTC2_4BPP_UNORM_BLOCK_IMG:	Format = VK_FORMAT_PVRTC2_4BPP_SRGB_BLOCK_IMG; break;
+			default:	break;
+			}
+
+			// If we're introducing a new format, make sure it's supported
+			if (InFormat != SRGBFormat)
+			{
+				const VkFormatProperties& SRGBFormatProperties = GetFormatProperties(SRGBFormat);
+				if (!VKHasAnyFlags(SRGBFormatProperties.optimalTilingFeatures, VK_FORMAT_FEATURE_SAMPLED_IMAGE_BIT))
+				{
+					// If we can't even sample from it, then reject the suggested SRGB format
+					SRGBFormat = InFormat;
+				}
+			}
+
+			return SRGBFormat;
+		};
+
+		for (int32 PixelFormatIndex = 0; PixelFormatIndex < PF_MAX; ++PixelFormatIndex)
+		{
+			const FPixelFormatInfo& PixelFormatInfo = GPixelFormats[PixelFormatIndex];
+			if (PixelFormatInfo.Supported)
+			{
+				const VkFormat OriginalFormat = (VkFormat)PixelFormatInfo.PlatformFormat;
+				GVulkanSRGBFormat[PixelFormatIndex] = GetSRGBMapping(OriginalFormat);
+			}
+			else
+			{
+				GVulkanSRGBFormat[PixelFormatIndex] = VK_FORMAT_UNDEFINED;
+			}
+		}
+	}
+
+#if UE_BUILD_DEBUG || UE_BUILD_DEVELOPMENT
 
 	// Print the resulting pixel format support
 	if (FParse::Param(FCommandLine::Get(), TEXT("PrintVulkanPixelFormatMappings")))
@@ -716,59 +790,49 @@ void FVulkanDevice::SetupFormats()
 		};
 
 		UE_LOG(LogVulkanRHI, Warning, TEXT("Pixel Format Mappings for Vulkan:"));
-		UE_LOG(LogVulkanRHI, Warning, TEXT("%24s | VulkanFormat | BlockBytes | Components | ComponentMapping | BufferFormat | Capabilities "), TEXT("PixelFormatName"));
+		UE_LOG(LogVulkanRHI, Warning, TEXT("%24s | %24s | BlockBytes | Components | ComponentMapping | BufferFormat | Capabilities | SRGBFormat"), 
+			TEXT("PixelFormatName"), TEXT("VulkanFormat"));
 		for (int32 PixelFormatIndex = 0; PixelFormatIndex < PF_MAX; ++PixelFormatIndex)
 		{
 			if (GPixelFormats[PixelFormatIndex].Supported)
 			{
 				const VkComponentMapping& ComponentMapping = PixelFormatComponentMapping[PixelFormatIndex];
 
-				FString CapabilitiesString = GetFormatCapabilities(GPixelFormats[PixelFormatIndex].Capabilities);
+				const VkFormat VulkanFormat = (VkFormat)GPixelFormats[PixelFormatIndex].PlatformFormat;
+				FString VulkanFormatStr(VK_TYPE_TO_STRING(VkFormat, VulkanFormat));
+				VulkanFormatStr.RightChopInline(10);  // Chop the VK_FORMAT_
+
+				FString SRGBFormat;
+				if (VulkanFormat != GVulkanSRGBFormat[PixelFormatIndex])
+				{
+					SRGBFormat = VK_TYPE_TO_STRING(VkFormat, GVulkanSRGBFormat[PixelFormatIndex]);
+					SRGBFormat.RightChopInline(10);  // Chop the VK_FORMAT_
+				}
 
-				UE_LOG(LogVulkanRHI, Warning, TEXT("%24s | %12d | %10d | %10d | %10d,%d,%d,%d | %12d | 0x%08X (%s)"),
+				UE_LOG(LogVulkanRHI, Warning, TEXT("%24s | %24s | %10d | %10d | %10d,%d,%d,%d | %12d |  0x%08X  | %s"),
 					GPixelFormats[PixelFormatIndex].Name,
-					GPixelFormats[PixelFormatIndex].PlatformFormat,
+					*VulkanFormatStr,
 					GPixelFormats[PixelFormatIndex].BlockBytes,
 					GPixelFormats[PixelFormatIndex].NumComponents,
 					ComponentMapping.r, ComponentMapping.g, ComponentMapping.b, ComponentMapping.a,
 					(int32)GVulkanBufferFormat[PixelFormatIndex],
 					(uint32)GPixelFormats[PixelFormatIndex].Capabilities,
-					*CapabilitiesString
+					*SRGBFormat
 					);
 			}
 		}
 
-#define VULKAN_CHECK_FORMAT_CAPABILITY(PF_Name) if (EnumHasAllFlags(GPixelFormats[PixelFormatIndex].Capabilities, EPixelFormatCapabilities::PF_Name)) { CapabilitiesString += TEXT(#PF_Name) TEXT(", ");}
-
 		UE_LOG(LogVulkanRHI, Warning, TEXT("Pixel Format Capabilities for Vulkan:"));
 		for (int32 PixelFormatIndex = 0; PixelFormatIndex < PF_MAX; ++PixelFormatIndex)
 		{
 			if (GPixelFormats[PixelFormatIndex].Supported)
 			{
-				FString CapabilitiesString;
-
-				VULKAN_CHECK_FORMAT_CAPABILITY(TextureSample);
-				VULKAN_CHECK_FORMAT_CAPABILITY(TextureCube);
-				VULKAN_CHECK_FORMAT_CAPABILITY(RenderTarget);
-				VULKAN_CHECK_FORMAT_CAPABILITY(DepthStencil);
-				VULKAN_CHECK_FORMAT_CAPABILITY(TextureBlendable);
-				VULKAN_CHECK_FORMAT_CAPABILITY(TextureAtomics);
-				
-				VULKAN_CHECK_FORMAT_CAPABILITY(Buffer);
-				VULKAN_CHECK_FORMAT_CAPABILITY(VertexBuffer);
-				VULKAN_CHECK_FORMAT_CAPABILITY(IndexBuffer);
-				VULKAN_CHECK_FORMAT_CAPABILITY(BufferAtomics);
-
-				VULKAN_CHECK_FORMAT_CAPABILITY(UAV);
-
+				const FString CapabilitiesString = GetFormatCapabilities(GPixelFormats[PixelFormatIndex].Capabilities);
 				UE_LOG(LogVulkanRHI, Warning, TEXT("%24s : %s"), GPixelFormats[PixelFormatIndex].Name, *CapabilitiesString);
 			}
 		}
-
-#undef VULKAN_CHECK_FORMAT_CAPABILITY
-
 	}
-#endif
+#endif  // UE_BUILD_DEBUG || UE_BUILD_DEVELOPMENT
 }
 
 const VkFormatProperties& FVulkanDevice::GetFormatProperties(VkFormat InFormat)
diff --git a/Engine/Source/Runtime/VulkanRHI/Private/VulkanRHIPrivate.h b/Engine/Source/Runtime/VulkanRHI/Private/VulkanRHIPrivate.h
--- a/Engine/Source/Runtime/VulkanRHI/Private/VulkanRHIPrivate.h
+++ b/Engine/Source/Runtime/VulkanRHI/Private/VulkanRHIPrivate.h
@@ -629,50 +629,17 @@ static inline VkAttachmentStoreOp RenderTargetStoreActionToVulkan(ERenderTargetS
 	return OutStoreAction;
 }
 
+extern VkFormat GVulkanSRGBFormat[PF_MAX];
 inline VkFormat UEToVkTextureFormat(EPixelFormat UEFormat, const bool bIsSRGB)
 {
-	VkFormat Format = (VkFormat)GPixelFormats[UEFormat].PlatformFormat;
 	if (bIsSRGB)
 	{
-		switch (Format)
-		{
-		case VK_FORMAT_B8G8R8A8_UNORM:				Format = VK_FORMAT_B8G8R8A8_SRGB; break;
-		case VK_FORMAT_A8B8G8R8_UNORM_PACK32:		Format = VK_FORMAT_A8B8G8R8_SRGB_PACK32; break;
-		case VK_FORMAT_R8_UNORM:					Format = ((GMaxRHIFeatureLevel <= ERHIFeatureLevel::ES3_1) ? VK_FORMAT_R8_UNORM : VK_FORMAT_R8_SRGB); break;
-		case VK_FORMAT_R8G8_UNORM:					Format = VK_FORMAT_R8G8_SRGB; break;
-		case VK_FORMAT_R8G8B8_UNORM:				Format = VK_FORMAT_R8G8B8_SRGB; break;
-		case VK_FORMAT_R8G8B8A8_UNORM:				Format = VK_FORMAT_R8G8B8A8_SRGB; break;
-		case VK_FORMAT_BC1_RGB_UNORM_BLOCK:			Format = VK_FORMAT_BC1_RGB_SRGB_BLOCK; break;
-		case VK_FORMAT_BC1_RGBA_UNORM_BLOCK:		Format = VK_FORMAT_BC1_RGBA_SRGB_BLOCK; break;
-		case VK_FORMAT_BC2_UNORM_BLOCK:				Format = VK_FORMAT_BC2_SRGB_BLOCK; break;
-		case VK_FORMAT_BC3_UNORM_BLOCK:				Format = VK_FORMAT_BC3_SRGB_BLOCK; break;
-		case VK_FORMAT_BC7_UNORM_BLOCK:				Format = VK_FORMAT_BC7_SRGB_BLOCK; break;
-		case VK_FORMAT_ETC2_R8G8B8_UNORM_BLOCK:		Format = VK_FORMAT_ETC2_R8G8B8_SRGB_BLOCK; break;
-		case VK_FORMAT_ETC2_R8G8B8A1_UNORM_BLOCK:	Format = VK_FORMAT_ETC2_R8G8B8A1_SRGB_BLOCK; break;
-		case VK_FORMAT_ETC2_R8G8B8A8_UNORM_BLOCK:	Format = VK_FORMAT_ETC2_R8G8B8A8_SRGB_BLOCK; break;
-		case VK_FORMAT_ASTC_4x4_UNORM_BLOCK:		Format = VK_FORMAT_ASTC_4x4_SRGB_BLOCK; break;
-		case VK_FORMAT_ASTC_5x4_UNORM_BLOCK:		Format = VK_FORMAT_ASTC_5x4_SRGB_BLOCK; break;
-		case VK_FORMAT_ASTC_5x5_UNORM_BLOCK:		Format = VK_FORMAT_ASTC_5x5_SRGB_BLOCK; break;
-		case VK_FORMAT_ASTC_6x5_UNORM_BLOCK:		Format = VK_FORMAT_ASTC_6x5_SRGB_BLOCK; break;
-		case VK_FORMAT_ASTC_6x6_UNORM_BLOCK:		Format = VK_FORMAT_ASTC_6x6_SRGB_BLOCK; break;
-		case VK_FORMAT_ASTC_8x5_UNORM_BLOCK:		Format = VK_FORMAT_ASTC_8x5_SRGB_BLOCK; break;
-		case VK_FORMAT_ASTC_8x6_UNORM_BLOCK:		Format = VK_FORMAT_ASTC_8x6_SRGB_BLOCK; break;
-		case VK_FORMAT_ASTC_8x8_UNORM_BLOCK:		Format = VK_FORMAT_ASTC_8x8_SRGB_BLOCK; break;
-		case VK_FORMAT_ASTC_10x5_UNORM_BLOCK:		Format = VK_FORMAT_ASTC_10x5_SRGB_BLOCK; break;
-		case VK_FORMAT_ASTC_10x6_UNORM_BLOCK:		Format = VK_FORMAT_ASTC_10x6_SRGB_BLOCK; break;
-		case VK_FORMAT_ASTC_10x8_UNORM_BLOCK:		Format = VK_FORMAT_ASTC_10x8_SRGB_BLOCK; break;
-		case VK_FORMAT_ASTC_10x10_UNORM_BLOCK:		Format = VK_FORMAT_ASTC_10x10_SRGB_BLOCK; break;
-		case VK_FORMAT_ASTC_12x10_UNORM_BLOCK:		Format = VK_FORMAT_ASTC_12x10_SRGB_BLOCK; break;
-		case VK_FORMAT_ASTC_12x12_UNORM_BLOCK:		Format = VK_FORMAT_ASTC_12x12_SRGB_BLOCK; break;
-//		case VK_FORMAT_PVRTC1_2BPP_UNORM_BLOCK_IMG:	Format = VK_FORMAT_PVRTC1_2BPP_SRGB_BLOCK_IMG; break;
-//		case VK_FORMAT_PVRTC1_4BPP_UNORM_BLOCK_IMG:	Format = VK_FORMAT_PVRTC1_4BPP_SRGB_BLOCK_IMG; break;
-//		case VK_FORMAT_PVRTC2_2BPP_UNORM_BLOCK_IMG:	Format = VK_FORMAT_PVRTC2_2BPP_SRGB_BLOCK_IMG; break;
-//		case VK_FORMAT_PVRTC2_4BPP_UNORM_BLOCK_IMG:	Format = VK_FORMAT_PVRTC2_4BPP_SRGB_BLOCK_IMG; break;
-		default:	break;
-		}
+		return GVulkanSRGBFormat[UEFormat];
+	}
+	else
+	{
+		return (VkFormat)GPixelFormats[UEFormat].PlatformFormat;
 	}
-
-	return Format;
 }
 
 static inline VkFormat UEToVkBufferFormat(EVertexElementType Type)
