Log vkCreateComputePipelines failure as an error

diff --git a/Engine/Source/Runtime/VulkanRHI/Private/VulkanPipeline.cpp b/Engine/Source/Runtime/VulkanRHI/Private/VulkanPipeline.cpp
--- a/Engine/Source/Runtime/VulkanRHI/Private/VulkanPipeline.cpp
+++ b/Engine/Source/Runtime/VulkanRHI/Private/VulkanPipeline.cpp
@@ -2181,7 +2181,7 @@ FVulkanComputePipeline* FVulkanPipelineStateCacheManager::CreateComputePipelineF
 	if (Result != VK_SUCCESS)
 	{
 		FString ComputeHash = Shader->GetHash().ToString();
-		UE_LOG(LogVulkanRHI, Warning, TEXT("Failed to create graphics pipeline.\nShaders in pipeline: CS: %s"), *ComputeHash);
+		UE_LOG(LogVulkanRHI, Error, TEXT("Failed to create compute pipeline.\nShaders in pipeline: CS: %s"), *ComputeHash);
 	}
 
 	Pipeline->Layout = ComputeLayout;
