VulkanRHI: Fix assert (in Debug builds only) when querying the cube face on a non-cube texture.

diff --git a/Engine/Source/Runtime/VulkanRHI/Private/VulkanRenderTarget.cpp b/Engine/Source/Runtime/VulkanRHI/Private/VulkanRenderTarget.cpp
--- a/Engine/Source/Runtime/VulkanRHI/Private/VulkanRenderTarget.cpp
+++ b/Engine/Source/Runtime/VulkanRHI/Private/VulkanRenderTarget.cpp
@@ -380,7 +380,7 @@ static void ConvertRawDataToFColor(VkFormat VulkanFormat, uint32 DestWidth, uint
 
 void FVulkanDynamicRHI::RHIReadSurfaceData(FRHITexture* TextureRHI, FIntRect Rect, TArray<FColor>& OutData, FReadSurfaceDataFlags InFlags)
 {
-	checkf((InFlags.GetCubeFace() == CubeFace_MAX), TEXT("Cube faces not supported yet."));
+	checkf((!TextureRHI->GetDesc().IsTextureCube()) || (InFlags.GetCubeFace() == CubeFace_MAX), TEXT("Cube faces not supported yet."));
 
 	const uint32 DestWidth = Rect.Max.X - Rect.Min.X;
 	const uint32 DestHeight = Rect.Max.Y - Rect.Min.Y;
