VulkanRHI: Write query results as 64bit values.  Clear Output array on reset to make it more obvious if some values aren't written.

#rnx

diff --git a/Engine/Source/Runtime/VulkanRHI/Private/VulkanRayTracing.cpp b/Engine/Source/Runtime/VulkanRHI/Private/VulkanRayTracing.cpp
--- a/Engine/Source/Runtime/VulkanRHI/Private/VulkanRayTracing.cpp
+++ b/Engine/Source/Runtime/VulkanRHI/Private/VulkanRayTracing.cpp
@@ -1116,6 +1116,8 @@ void FVulkanRayTracingCompactedSizeQueryPool::Reset(FVulkanCmdBuffer* InCmdBuffe
 	VulkanRHI::vkCmdResetQueryPool(InCmdBuffer->GetHandle(), QueryPool, 0, MaxQueries);
 	FenceSignaledCounter = 0;
 	CmdBuffer = nullptr;
+	check(QueryOutput.Num() == MaxQueries);
+	FMemory::Memzero(QueryOutput.GetData(), MaxQueries * sizeof(uint64));
 }
 
 bool FVulkanRayTracingCompactedSizeQueryPool::TryGetResults(uint32 NumResults)
@@ -1128,7 +1130,7 @@ bool FVulkanRayTracingCompactedSizeQueryPool::TryGetResults(uint32 NumResults)
 		return false;
 	}
 
-	VkResult Result = VulkanRHI::vkGetQueryPoolResults(Device->GetInstanceHandle(), QueryPool, 0, NumResults, NumResults * sizeof(uint64), QueryOutput.GetData(), sizeof(uint64), VK_QUERY_RESULT_WAIT_BIT);
+	VkResult Result = VulkanRHI::vkGetQueryPoolResults(Device->GetInstanceHandle(), QueryPool, 0, NumResults, NumResults * sizeof(uint64), QueryOutput.GetData(), sizeof(uint64), VK_QUERY_RESULT_64_BIT | VK_QUERY_RESULT_WAIT_BIT);
 	if (Result == VK_SUCCESS)
 	{
 		return true;
