VulkanRHI: Fix warning from Renderdoc where stencil struct was still being attached on formats without stencil.

#rnx

diff --git a/Engine/Source/Runtime/VulkanRHI/Private/VulkanRHIPrivate.h b/Engine/Source/Runtime/VulkanRHI/Private/VulkanRHIPrivate.h
--- a/Engine/Source/Runtime/VulkanRHI/Private/VulkanRHIPrivate.h
+++ b/Engine/Source/Runtime/VulkanRHI/Private/VulkanRHIPrivate.h
@@ -584,6 +584,20 @@ namespace VulkanRHI
 			return VK_IMAGE_ASPECT_COLOR_BIT;
 		}
 	}
+
+	static bool VulkanFormatHasStencil(VkFormat Format)
+	{
+		switch (Format)
+		{
+		case VK_FORMAT_D16_UNORM_S8_UINT:
+		case VK_FORMAT_D24_UNORM_S8_UINT:
+		case VK_FORMAT_D32_SFLOAT_S8_UINT:
+		case VK_FORMAT_S8_UINT:
+			return true;
+		default:
+			return false;
+		}
+	}
 }
 
 #if VULKAN_HAS_DEBUGGING_ENABLED
diff --git a/Engine/Source/Runtime/VulkanRHI/Private/VulkanRenderpass.h b/Engine/Source/Runtime/VulkanRHI/Private/VulkanRenderpass.h
--- a/Engine/Source/Runtime/VulkanRHI/Private/VulkanRenderpass.h
+++ b/Engine/Source/Runtime/VulkanRHI/Private/VulkanRenderpass.h
@@ -270,10 +270,10 @@ struct FVulkanAttachmentDescription<VkAttachmentDescription>
 		stencilLoadOp = InDesc.stencilLoadOp;
 		stencilStoreOp = InDesc.stencilStoreOp;
 
-		const VkImageLayout StencilInitialLayout = InStencilDesc ? InStencilDesc->stencilInitialLayout : VK_IMAGE_LAYOUT_UNDEFINED;
+		const bool bHasStencilLayout = VulkanFormatHasStencil(InDesc.format) && (InStencilDesc != nullptr);
+		const VkImageLayout StencilInitialLayout = bHasStencilLayout ? InStencilDesc->stencilInitialLayout : VK_IMAGE_LAYOUT_UNDEFINED;
 		initialLayout = GetMergedDepthStencilLayout(InDesc.initialLayout, StencilInitialLayout);
-
-		const VkImageLayout StencilFinalLayout = InStencilDesc ? InStencilDesc->stencilFinalLayout : VK_IMAGE_LAYOUT_UNDEFINED;
+		const VkImageLayout StencilFinalLayout = bHasStencilLayout ? InStencilDesc->stencilFinalLayout : VK_IMAGE_LAYOUT_UNDEFINED;
 		finalLayout = GetMergedDepthStencilLayout(InDesc.finalLayout, StencilFinalLayout);
 	}
 };
@@ -304,8 +304,10 @@ struct FVulkanAttachmentDescription<VkAttachmentDescription2>
 
 	FVulkanAttachmentDescription(const VkAttachmentDescription& InDesc, const VkAttachmentDescriptionStencilLayout* InStencilDesc, bool bSupportsParallelRendering)
 	{
+		const bool bHasStencilLayout = bSupportsParallelRendering && VulkanFormatHasStencil(InDesc.format) && (InStencilDesc != nullptr);
+
 		sType = VK_STRUCTURE_TYPE_ATTACHMENT_DESCRIPTION_2;
-		pNext = (bSupportsParallelRendering && InStencilDesc && InStencilDesc->stencilFinalLayout != VK_IMAGE_LAYOUT_UNDEFINED) ? InStencilDesc : nullptr;
+		pNext = (bHasStencilLayout && (InStencilDesc->stencilFinalLayout != VK_IMAGE_LAYOUT_UNDEFINED)) ? InStencilDesc : nullptr;
 		flags = InDesc.flags;
 		format = InDesc.format;
 		samples = InDesc.samples;
