VulkanRHI: Bring back layout tracking in the Sync2 path.  Fix undesired discards in frame.

#rnx

diff --git a/Engine/Source/Runtime/VulkanRHI/Private/VulkanBarriers.cpp b/Engine/Source/Runtime/VulkanRHI/Private/VulkanBarriers.cpp
--- a/Engine/Source/Runtime/VulkanRHI/Private/VulkanBarriers.cpp
+++ b/Engine/Source/Runtime/VulkanRHI/Private/VulkanBarriers.cpp
@@ -33,6 +33,16 @@ static FAutoConsoleVariableRef CVarVulkanAutoCorrectExpectedLayouts(
 	ECVF_Default
 );
 
+int32 GVulkanAutoCorrectUnknownLayouts = 1;
+static FAutoConsoleVariableRef CVarVulkanAutoCorrectUnknownLayouts(
+	TEXT("r.Vulkan.AutoCorrectUnknownLayouts"),
+	GVulkanAutoCorrectUnknownLayouts,
+	TEXT("Will use layout tracking to correct unknown layouts.  \n")
+	TEXT("This is unsafe for multi-threaded command buffer generations\n")
+	TEXT(" 0: Do not correct layouts.\n")
+	TEXT(" 1: Correct the layouts using layout tracking. (default)"),
+	ECVF_Default
+);
 
 //
 // The following two functions are used when the RHI needs to do image layout transitions internally.
@@ -1346,6 +1356,46 @@ void FTransitionProcessor<VkMemoryBarrier2, VkBufferMemoryBarrier2, VkImageMemor
 template <>
 void FTransitionProcessor<VkMemoryBarrier2, VkBufferMemoryBarrier2, VkImageMemoryBarrier2>::ApplyTracking(VkImageMemoryBarrier2* RemainingBarriers, const FVulkanPipelineBarrier::ImageBarrierExtraData* RemainingExtras, int32 RemainingCount, bool IsCrossPipe)
 {
+	if (GVulkanAutoCorrectUnknownLayouts)
+	{
+		VkImageMemoryBarrier2& ImageBarrier = *RemainingBarriers;
+
+		if ((ImageBarrier.oldLayout == VK_IMAGE_LAYOUT_UNDEFINED) && (!RemainingExtras->IsAliasingBarrier))
+		{
+			FVulkanLayoutManager& LayoutManager = Context.GetCommandBufferManager()->GetActiveCmdBuffer()->GetLayoutManager();
+			FVulkanTexture* Texture = RemainingExtras->BaseTexture;
+			check(Texture->Image == ImageBarrier.image);
+			const FVulkanImageLayout* Layout = LayoutManager.GetFullLayout(*Texture, false);
+			if (Layout)
+			{
+				if (Layout->AreAllSubresourcesSameLayout())
+				{
+					ImageBarrier.oldLayout = Layout->MainLayout;
+				}
+				else
+				{
+					ForEachAspect(ImageBarrier.subresourceRange.aspectMask, [&](VkImageAspectFlagBits SingleAspect)
+						{
+							const uint32 LastLayer = ImageBarrier.subresourceRange.baseArrayLayer + ImageBarrier.subresourceRange.layerCount;
+							const uint32 LastMip = ImageBarrier.subresourceRange.baseMipLevel + ImageBarrier.subresourceRange.levelCount;
+							for (uint32 LayerIndex = ImageBarrier.subresourceRange.baseArrayLayer; LayerIndex < LastLayer; ++LayerIndex)
+							{
+								for (uint32 MipIndex = ImageBarrier.subresourceRange.baseMipLevel; MipIndex < LastMip; ++MipIndex)
+								{
+									const VkImageLayout SubresourceLayout = Layout->GetSubresLayout(LayerIndex, MipIndex, SingleAspect);
+									check((ImageBarrier.oldLayout == VK_IMAGE_LAYOUT_UNDEFINED) || (ImageBarrier.oldLayout == SubresourceLayout));
+									ImageBarrier.oldLayout = SubresourceLayout;
+								}
+							}
+						});
+				}
+
+				ImageBarrier.srcAccessMask = GetVkAccessMaskForLayout(ImageBarrier.oldLayout);
+				ImageBarrier.srcStageMask = GetVkStageFlagsForLayout(ImageBarrier.oldLayout);
+			}
+		}
+	}
+
 	// Sync2 only needs to inform tracking of the new layout
 	Context.GetCommandBufferManager()->GetActiveCmdBuffer()->GetLayoutManager().SetLayout(*RemainingExtras->BaseTexture, RemainingBarriers->subresourceRange, RemainingBarriers->newLayout);
 
diff --git a/Engine/Source/Runtime/VulkanRHI/Private/VulkanCommandBuffer.cpp b/Engine/Source/Runtime/VulkanRHI/Private/VulkanCommandBuffer.cpp
--- a/Engine/Source/Runtime/VulkanRHI/Private/VulkanCommandBuffer.cpp
+++ b/Engine/Source/Runtime/VulkanRHI/Private/VulkanCommandBuffer.cpp
@@ -47,6 +47,7 @@ static FAutoConsoleVariableRef CVarVulkanPreventOverlapWithUpload(
 #define CMD_BUFFER_TIME_TO_WAIT_BEFORE_DELETING		10
 
 const uint32 GNumberOfFramesBeforeDeletingDescriptorPool = 300;
+extern int32 GVulkanAutoCorrectUnknownLayouts;
 
 FVulkanCmdBuffer::FVulkanCmdBuffer(FVulkanDevice* InDevice, FVulkanCommandBufferPool* InCommandBufferPool, bool bInIsUploadOnly)
 	: CurrentStencilRef(0)
@@ -66,8 +67,7 @@ FVulkanCmdBuffer::FVulkanCmdBuffer(FVulkanDevice* InDevice, FVulkanCommandBuffer
 	, CommandBufferPool(InCommandBufferPool)
 	, Timing(nullptr)
 	, LastValidTiming(0)
-	, LayoutManager(InDevice->SupportsParallelRendering(), 
-		InDevice->SupportsParallelRendering() ? nullptr :
+	, LayoutManager(InDevice->SupportsParallelRendering() && !GVulkanAutoCorrectUnknownLayouts,
 		&InCommandBufferPool->GetMgr().GetCommandListContext()->GetQueue()->GetLayoutManager())
 {
 	{
