PR #9760: support implementations supporting draw-only VRS (Contributed by [REDACTED_PERSON])
(with some minor formatting cleanups)

Adds support for implementations where the VR_KFR_fragment_shading_rate extension supports per-draw VRS, but not image attachment-based VRS.

diff --git a/Engine/Source/Runtime/VulkanRHI/Private/VulkanExtensions.cpp b/Engine/Source/Runtime/VulkanRHI/Private/VulkanExtensions.cpp
--- a/Engine/Source/Runtime/VulkanRHI/Private/VulkanExtensions.cpp
+++ b/Engine/Source/Runtime/VulkanRHI/Private/VulkanExtensions.cpp
@@ -487,10 +487,10 @@ public:
 
 	virtual void PostPhysicalDeviceFeatures(FOptionalVulkanDeviceExtensions& ExtensionFlags) override final
 	{
-		ExtensionFlags.HasKHRFragmentShadingRate = FragmentShadingRateFeatures.attachmentFragmentShadingRate;
+		ExtensionFlags.HasKHRFragmentShadingRate = 1;
 
-		GRHISupportsAttachmentVariableRateShading = FragmentShadingRateFeatures.attachmentFragmentShadingRate ? true : false;
-		GRHISupportsPipelineVariableRateShading = FragmentShadingRateFeatures.pipelineFragmentShadingRate ? true : false;
+		GRHISupportsAttachmentVariableRateShading = (FragmentShadingRateFeatures.attachmentFragmentShadingRate == VK_TRUE);
+		GRHISupportsPipelineVariableRateShading = (FragmentShadingRateFeatures.pipelineFragmentShadingRate == VK_TRUE);
 
 		if (FragmentShadingRateFeatures.attachmentFragmentShadingRate == VK_TRUE)
 		{
@@ -506,20 +506,23 @@ public:
 
 	virtual void PostPhysicalDeviceProperties() override final
 	{
-		GRHIVariableRateShadingImageTileMinWidth = FragmentShadingRateProperties.minFragmentShadingRateAttachmentTexelSize.width;
-		GRHIVariableRateShadingImageTileMinHeight = FragmentShadingRateProperties.minFragmentShadingRateAttachmentTexelSize.height;
-		GRHIVariableRateShadingImageTileMaxWidth = FragmentShadingRateProperties.maxFragmentShadingRateAttachmentTexelSize.width;
-		GRHIVariableRateShadingImageTileMaxHeight = FragmentShadingRateProperties.maxFragmentShadingRateAttachmentTexelSize.height;
-
-		if (FragmentShadingRateProperties.maxFragmentSize.width >= 4 && FragmentShadingRateProperties.maxFragmentSize.height >= 4)
+		if (FragmentShadingRateFeatures.attachmentFragmentShadingRate == VK_TRUE)
 		{
-			// FYI FVulkanDevice::GetBestMatchedShadingRateExtents does extent filtering
-			GRHISupportsLargerVariableRateShadingSizes = GRHISupportsPipelineVariableRateShading;
-		}
+			GRHIVariableRateShadingImageTileMinWidth = FragmentShadingRateProperties.minFragmentShadingRateAttachmentTexelSize.width;
+			GRHIVariableRateShadingImageTileMinHeight = FragmentShadingRateProperties.minFragmentShadingRateAttachmentTexelSize.height;
+			GRHIVariableRateShadingImageTileMaxWidth = FragmentShadingRateProperties.maxFragmentShadingRateAttachmentTexelSize.width;
+			GRHIVariableRateShadingImageTileMaxHeight = FragmentShadingRateProperties.maxFragmentShadingRateAttachmentTexelSize.height;
+
+			if (FragmentShadingRateProperties.maxFragmentSize.width >= 4 && FragmentShadingRateProperties.maxFragmentSize.height >= 4)
+			{
+				// FYI FVulkanDevice::GetBestMatchedShadingRateExtents does extent filtering
+				GRHISupportsLargerVariableRateShadingSizes = GRHISupportsPipelineVariableRateShading;
+			}
 
-		// todo: We don't currently care much about the other properties here, but at some point in the future we probably will.
+			// todo: We don't currently care much about the other properties here, but at some point in the future we probably will.
 
-		UE_LOG(LogVulkanRHI, Verbose, TEXT("Image-based Variable Rate Shading supported via KHRFragmentShadingRate extension. Selected VRS tile size %u by %u pixels per VRS image texel."), GRHIVariableRateShadingImageTileMinWidth, GRHIVariableRateShadingImageTileMinHeight);
+			UE_LOG(LogVulkanRHI, Verbose, TEXT("Image-based Variable Rate Shading supported via KHRFragmentShadingRate extension. Selected VRS tile size %u by %u pixels per VRS image texel."), GRHIVariableRateShadingImageTileMinWidth, GRHIVariableRateShadingImageTileMinHeight);
+		}
 	}
 
 	virtual void PrePhysicalDeviceFeatures(VkPhysicalDeviceFeatures2KHR& PhysicalDeviceFeatures2) override final
@@ -530,7 +533,7 @@ public:
 
 	virtual void PreCreateDevice(VkDeviceCreateInfo& DeviceCreateInfo) override final
 	{
-		if (FragmentShadingRateFeatures.attachmentFragmentShadingRate == VK_TRUE)
+		if (FragmentShadingRateFeatures.attachmentFragmentShadingRate == VK_TRUE || FragmentShadingRateFeatures.pipelineFragmentShadingRate == VK_TRUE)
 		{
 			AddToPNext(DeviceCreateInfo, FragmentShadingRateFeatures);
 		}
@@ -568,7 +571,6 @@ public:
 		if (!GRHISupportsAttachmentVariableRateShading && (FragmentDensityMapFeatures.fragmentDensityMap == VK_TRUE))
 		{
 			GRHISupportsAttachmentVariableRateShading = true;
-			GRHISupportsPipelineVariableRateShading = false;
 
 			// Go with the smallest tile size for now, and also force to square, since this seems to be standard.
 			// TODO: Eventually we may want to surface the range of possible tile sizes depending on end use cases, but for now this is being used for foveated rendering and smallest tile size
@@ -582,7 +584,7 @@ public:
 			GRHIVariableRateShadingImageDataType = VRSImage_Fractional;
 			GRHIVariableRateShadingImageFormat = PF_R8G8;
 
-			// UE_LOG(LogVulkanRHI, Display, TEXT("Image-based Variable Rate Shading supported via EXTFragmentDensityMap extension. Selected VRS tile size %u by %u pixels per VRS image texel."), GRHIVariableRateShadingImageTileMinWidth, GRHIVariableRateShadingImageTileMinHeight);
+			UE_LOG(LogVulkanRHI, Display, TEXT("Image-based Variable Rate Shading supported via EXTFragmentDensityMap extension. Selected VRS tile size %u by %u pixels per VRS image texel."), GRHIVariableRateShadingImageTileMinWidth, GRHIVariableRateShadingImageTileMinHeight);
 		}
 	}
 
