VulkanRHI: fix static analysis warning.

#rnx

diff --git a/Engine/Source/Runtime/VulkanRHI/Private/VulkanPipeline.cpp b/Engine/Source/Runtime/VulkanRHI/Private/VulkanPipeline.cpp
--- a/Engine/Source/Runtime/VulkanRHI/Private/VulkanPipeline.cpp
+++ b/Engine/Source/Runtime/VulkanRHI/Private/VulkanPipeline.cpp
@@ -1277,7 +1277,7 @@ bool FVulkanPipelineStateCacheManager::CreateGfxPipelineFromEntry(FVulkanRHIGrap
 		Shaders[ShaderStage]->GetEntryPoint(EntryPoints[PipelineInfo.stageCount], 24);
 		ShaderStages[PipelineInfo.stageCount].pName = EntryPoints[PipelineInfo.stageCount];
 
-		const FVulkanShaderHeader& ShaderHeader = Shaders[CurrStage]->GetCodeHeader();
+		const FVulkanShaderHeader& ShaderHeader = Shaders[ShaderStage]->GetCodeHeader();
 		if ((ShaderHeader.WaveSize > 0) && Device->GetOptionalExtensions().HasEXTSubgroupSizeControl)
 		{
 			check((ShaderHeader.WaveSize >= GRHIMinimumWaveSize) && (ShaderHeader.WaveSize <= GRHIMaximumWaveSize));
