VulkanRHI: Fix static analysis warning.

#rnx

diff --git a/Engine/Source/Runtime/VulkanRHI/Private/VulkanPipeline.cpp b/Engine/Source/Runtime/VulkanRHI/Private/VulkanPipeline.cpp
--- a/Engine/Source/Runtime/VulkanRHI/Private/VulkanPipeline.cpp
+++ b/Engine/Source/Runtime/VulkanRHI/Private/VulkanPipeline.cpp
@@ -1277,13 +1277,16 @@ bool FVulkanPipelineStateCacheManager::CreateGfxPipelineFromEntry(FVulkanRHIGrap
 		Shaders[ShaderStage]->GetEntryPoint(EntryPoints[PipelineInfo.stageCount], 24);
 		ShaderStages[PipelineInfo.stageCount].pName = EntryPoints[PipelineInfo.stageCount];
 
-		const FVulkanShaderHeader& ShaderHeader = Shaders[ShaderStage]->GetCodeHeader();
-		if ((ShaderHeader.WaveSize > 0) && Device->GetOptionalExtensions().HasEXTSubgroupSizeControl)
+		if (Device->GetOptionalExtensions().HasEXTSubgroupSizeControl && (Shaders[ShaderStage] != nullptr))
 		{
-			check((ShaderHeader.WaveSize >= GRHIMinimumWaveSize) && (ShaderHeader.WaveSize <= GRHIMaximumWaveSize));
-			ZeroVulkanStruct(RequiredSubgroupSizeCreateInfo[PipelineInfo.stageCount], VK_STRUCTURE_TYPE_PIPELINE_SHADER_STAGE_REQUIRED_SUBGROUP_SIZE_CREATE_INFO);
-			RequiredSubgroupSizeCreateInfo[PipelineInfo.stageCount].requiredSubgroupSize = ShaderHeader.WaveSize;
-			ShaderStages[PipelineInfo.stageCount].pNext = &RequiredSubgroupSizeCreateInfo[PipelineInfo.stageCount];
+			const FVulkanShaderHeader& ShaderHeader = Shaders[ShaderStage]->GetCodeHeader();
+			if (ShaderHeader.WaveSize > 0)
+			{
+				check((ShaderHeader.WaveSize >= GRHIMinimumWaveSize) && (ShaderHeader.WaveSize <= GRHIMaximumWaveSize));
+				ZeroVulkanStruct(RequiredSubgroupSizeCreateInfo[PipelineInfo.stageCount], VK_STRUCTURE_TYPE_PIPELINE_SHADER_STAGE_REQUIRED_SUBGROUP_SIZE_CREATE_INFO);
+				RequiredSubgroupSizeCreateInfo[PipelineInfo.stageCount].requiredSubgroupSize = ShaderHeader.WaveSize;
+				ShaderStages[PipelineInfo.stageCount].pNext = &RequiredSubgroupSizeCreateInfo[PipelineInfo.stageCount];
+			}
 		}
 
 		PipelineInfo.stageCount++;
