VulkanRHI: Fix static analysis again.

#rnx

diff --git a/Engine/Source/Runtime/VulkanRHI/Private/VulkanPipeline.cpp b/Engine/Source/Runtime/VulkanRHI/Private/VulkanPipeline.cpp
--- a/Engine/Source/Runtime/VulkanRHI/Private/VulkanPipeline.cpp
+++ b/Engine/Source/Runtime/VulkanRHI/Private/VulkanPipeline.cpp
@@ -1264,7 +1264,7 @@ bool FVulkanPipelineStateCacheManager::CreateGfxPipelineFromEntry(FVulkanRHIGrap
 	VkPipelineShaderStageRequiredSubgroupSizeCreateInfo RequiredSubgroupSizeCreateInfo[ShaderStage::NumStages];
 	for (int32 ShaderStage = 0; ShaderStage < ShaderStage::NumStages; ++ShaderStage)
 	{
-		if (!ShaderModules[ShaderStage].IsValid())
+		if (!ShaderModules[ShaderStage].IsValid() || (Shaders[ShaderStage] == nullptr))
 		{
 			continue;
 		}
@@ -1277,7 +1277,7 @@ bool FVulkanPipelineStateCacheManager::CreateGfxPipelineFromEntry(FVulkanRHIGrap
 		Shaders[ShaderStage]->GetEntryPoint(EntryPoints[PipelineInfo.stageCount], 24);
 		ShaderStages[PipelineInfo.stageCount].pName = EntryPoints[PipelineInfo.stageCount];
 
-		if (Device->GetOptionalExtensions().HasEXTSubgroupSizeControl && (Shaders[ShaderStage] != nullptr))
+		if (Device->GetOptionalExtensions().HasEXTSubgroupSizeControl)
 		{
 			const FVulkanShaderHeader& ShaderHeader = Shaders[ShaderStage]->GetCodeHeader();
 			if (ShaderHeader.WaveSize > 0)
