VulkanRHI: Fix a hole in the locking of the FVulkanShaderFactory.  Reduce lock contention by having one lock per map.

diff --git a/Engine/Source/Runtime/VulkanRHI/Private/VulkanResources.h b/Engine/Source/Runtime/VulkanRHI/Private/VulkanResources.h
--- a/Engine/Source/Runtime/VulkanRHI/Private/VulkanResources.h
+++ b/Engine/Source/Runtime/VulkanRHI/Private/VulkanResources.h
@@ -110,8 +110,6 @@ public:
 
 	void PurgeShaderModules();
 
-	void Setup(TArrayView<const uint8> InShaderHeaderAndCode, uint64 InShaderKey);
-
 	TRefCountPtr<FVulkanShaderModule> GetOrCreateHandle(const FVulkanLayout* Layout, uint32 LayoutHash)
 	{
 		FScopeLock Lock(&VulkanShaderModulesMapCS);
@@ -209,6 +207,7 @@ private:
 	static FSpirvCode PatchSpirvInputAttachments(FSpirvCode& SpirvCode);
 
 protected:
+	void Setup(FVulkanShaderHeader&& InCodeHeader, FSpirvContainer&& InSpirvContainer, uint64 InShaderKey);
 
 	FVulkanDevice*					Device;
 
@@ -277,7 +276,7 @@ public:
 	{
 		if (ShaderKey)
 		{
-			FRWScopeLock ScopedLock(Lock, SLT_ReadOnly);
+			FRWScopeLock ScopedLock(RWLock[ShaderType::StaticFrequency], SLT_ReadOnly);
 			FVulkanShader* const * FoundShaderPtr = ShaderMap[ShaderType::StaticFrequency].Find(ShaderKey);
 			if (FoundShaderPtr)
 			{
@@ -292,7 +291,7 @@ public:
 	void OnDeleteShader(const FVulkanShader& Shader);
 
 private:
-	mutable FRWLock Lock;
+	mutable FRWLock RWLock[SF_NumFrequencies];
 	TMap<uint64, FVulkanShader*> ShaderMap[SF_NumFrequencies];
 };
 
diff --git a/Engine/Source/Runtime/VulkanRHI/Private/VulkanShaders.cpp b/Engine/Source/Runtime/VulkanRHI/Private/VulkanShaders.cpp
--- a/Engine/Source/Runtime/VulkanRHI/Private/VulkanShaders.cpp
+++ b/Engine/Source/Runtime/VulkanRHI/Private/VulkanShaders.cpp
@@ -52,31 +52,49 @@ FVulkanShaderFactory::~FVulkanShaderFactory()
 template <typename ShaderType> 
 ShaderType* FVulkanShaderFactory::CreateShader(TArrayView<const uint8> Code, FVulkanDevice* Device)
 {
-	uint32 ShaderCodeLen = Code.Num();
-	uint32 ShaderCodeCRC = FCrc::MemCrc32(Code.GetData(), Code.Num());
-	uint64 ShaderKey = ((uint64)ShaderCodeLen | ((uint64)ShaderCodeCRC << 32));
+	const uint32 ShaderCodeLen = Code.Num();
+	const uint32 ShaderCodeCRC = FCrc::MemCrc32(Code.GetData(), Code.Num());
+	const uint64 ShaderKey = ((uint64)ShaderCodeLen | ((uint64)ShaderCodeCRC << 32));
 
 	ShaderType* RetShader = LookupShader<ShaderType>(ShaderKey);
+
 	if (RetShader == nullptr)
 	{
-		RetShader = new ShaderType(Device);
-		RetShader->Setup(Code, ShaderKey);
-		if constexpr (
-			ShaderType::StaticFrequency == SF_RayCallable ||
-			ShaderType::StaticFrequency == SF_RayGen ||
-			ShaderType::StaticFrequency == SF_RayHitGroup ||
-			ShaderType::StaticFrequency == SF_RayMiss)
+		// Do serialize outside of lock
+		FMemoryReaderView Ar(Code, true);
+		FVulkanShaderHeader CodeHeader;
+		Ar << CodeHeader;
+		FVulkanShader::FSpirvContainer SpirvContainer;
+		Ar << SpirvContainer;
+
 		{
-			RetShader->RayTracingPayloadType = RetShader->CodeHeader.RayTracingPayloadType;
-		}
-			
-		FRWScopeLock ScopedLock(Lock, SLT_Write);
-		ShaderMap[ShaderType::StaticFrequency].Add(ShaderKey, RetShader);
+			FRWScopeLock ScopedLock(RWLock[ShaderType::StaticFrequency], SLT_Write);
+			FVulkanShader* const* FoundShaderPtr = ShaderMap[ShaderType::StaticFrequency].Find(ShaderKey);
+			if (FoundShaderPtr)
+			{
+				RetShader = static_cast<ShaderType*>(*FoundShaderPtr);
+			}
+			else
+			{
+				RetShader = new ShaderType(Device);
+				RetShader->Setup(MoveTemp(CodeHeader), MoveTemp(SpirvContainer), ShaderKey);
+				if constexpr (
+					ShaderType::StaticFrequency == SF_RayCallable ||
+					ShaderType::StaticFrequency == SF_RayGen ||
+					ShaderType::StaticFrequency == SF_RayHitGroup ||
+					ShaderType::StaticFrequency == SF_RayMiss)
+				{
+					RetShader->RayTracingPayloadType = RetShader->CodeHeader.RayTracingPayloadType;
+				}
+
+				ShaderMap[ShaderType::StaticFrequency].Add(ShaderKey, RetShader);
 
 #if !(UE_BUILD_SHIPPING || UE_BUILD_TEST)
-		FShaderCodeReader ShaderCode(Code);
-		RetShader->ShaderName = ShaderCode.FindOptionalData(FShaderCodeName::Key);
+				FShaderCodeReader ShaderCode(Code);
+				RetShader->ShaderName = ShaderCode.FindOptionalData(FShaderCodeName::Key);
 #endif
+			}
+		}
 	}
 
 	return RetShader;
@@ -84,14 +102,13 @@ ShaderType* FVulkanShaderFactory::CreateShader(TArrayView<const uint8> Code, FVu
 
 void FVulkanShaderFactory::LookupShaders(const uint64 InShaderKeys[ShaderStage::NumStages], FVulkanShader* OutShaders[ShaderStage::NumStages]) const
 {
-	FRWScopeLock ScopedLock(Lock, SLT_ReadOnly);
-	
 	for (int32 Idx = 0; Idx < ShaderStage::NumStages; ++Idx)
 	{
 		uint64 ShaderKey = InShaderKeys[Idx];
 		if (ShaderKey)
 		{
 			EShaderFrequency ShaderFrequency = ShaderStage::GetFrequencyForGfxStage((ShaderStage::EStage)Idx);
+			FRWScopeLock ScopedLock(RWLock[ShaderFrequency], SLT_ReadOnly);
 			FVulkanShader* const * FoundShaderPtr = ShaderMap[ShaderFrequency].Find(ShaderKey);
 			if (FoundShaderPtr)
 			{
@@ -103,8 +120,8 @@ void FVulkanShaderFactory::LookupShaders(const uint64 InShaderKeys[ShaderStage::
 
 void FVulkanShaderFactory::OnDeleteShader(const FVulkanShader& Shader)
 {
-	FRWScopeLock ScopedLock(Lock, SLT_Write);
-	uint64 ShaderKey = Shader.GetShaderKey(); 
+	const uint64 ShaderKey = Shader.GetShaderKey(); 
+	FRWScopeLock ScopedLock(RWLock[Shader.Frequency], SLT_Write);
 	ShaderMap[Shader.Frequency].Remove(ShaderKey);
 }
 
@@ -174,18 +191,16 @@ FVulkanShader::FSpirvCode FVulkanShader::GetSpirvCode()
 }
 
 
-void FVulkanShader::Setup(TArrayView<const uint8> InShaderHeaderAndCode, uint64 InShaderKey)
+void FVulkanShader::Setup(FVulkanShaderHeader&& InCodeHeader, FSpirvContainer&& InSpirvContainer, uint64 InShaderKey)
 {
 	LLM_SCOPE_VULKAN(ELLMTagVulkan::VulkanShaders);
 	check(Device);
 
 	ShaderKey = InShaderKey;
 
-	FMemoryReaderView Ar(InShaderHeaderAndCode, true);
-
-	Ar << CodeHeader;
+	CodeHeader = MoveTemp(InCodeHeader);
 
-	Ar << SpirvContainer;
+	SpirvContainer = MoveTemp(InSpirvContainer);
 
 	checkf(SpirvContainer.GetSizeBytes() != 0, TEXT("Empty SPIR-V! %s"), *CodeHeader.DebugName);
 
