VulkanRHI:  Add support for VK_KHR_bind_memory2 and add missing vkGetDeviceBufferMemoryRequirementsKHR from Maintenance4. (needed for upcoming allocator changes)

#rnx

diff --git a/Engine/Source/Runtime/VulkanRHI/Private/VulkanCommandWrappers.h b/Engine/Source/Runtime/VulkanRHI/Private/VulkanCommandWrappers.h
--- a/Engine/Source/Runtime/VulkanRHI/Private/VulkanCommandWrappers.h
+++ b/Engine/Source/Runtime/VulkanRHI/Private/VulkanCommandWrappers.h
@@ -200,9 +200,12 @@ struct FWrapLayer
 	static void CmdCopyAccelerationStructureKHR(VkResult Result, VkCommandBuffer CommandBuffer, const VkCopyAccelerationStructureInfoKHR* Info);
 #endif
 	static void GetDeviceImageMemoryRequirementsKHR(VkResult Result, VkDevice Device, const VkDeviceImageMemoryRequirements* Info, VkMemoryRequirements2* MemoryRequirements) VULKAN_LAYER_BODY
+	static void GetDeviceBufferMemoryRequirementsKHR(VkResult Result, VkDevice Device, const VkDeviceBufferMemoryRequirements* Info, VkMemoryRequirements2* MemoryRequirements) VULKAN_LAYER_BODY
 	static void ResetQueryPoolEXT(VkResult Result, VkDevice Device, VkQueryPool QueryPool, uint32_t FirstQuery, uint32_t QueryCount) VULKAN_LAYER_BODY
 	static void GetPhysicalDeviceCalibrateableTimeDomainsEXT(VkResult Result, VkPhysicalDevice PhysicalDevice, uint32_t* TimeDomainCount, VkTimeDomainEXT* TimeDomains) VULKAN_LAYER_BODY
 	static void GetCalibratedTimestampsEXT(VkResult Result, VkDevice Device, uint32_t TimestampCount, const VkCalibratedTimestampInfoEXT* TimestampInfos, uint64_t* Timestamps, uint64_t* MaxDeviation) VULKAN_LAYER_BODY
+	static void BindBufferMemory2KHR(VkResult Result, VkDevice Device, uint32_t BindInfoCount, const VkBindBufferMemoryInfo* BindInfos) VULKAN_LAYER_BODY
+	static void BindImageMemory2KHR(VkResult Result, VkDevice Device, uint32_t BindInfoCount, const VkBindImageMemoryInfo* BindInfos) VULKAN_LAYER_BODY
 	static void CmdPipelineBarrier2KHR(VkResult Result, VkCommandBuffer CommandBuffer, const VkDependencyInfo* DependencyInfo) VULKAN_LAYER_BODY
 	static void CmdResetEvent2KHR(VkResult Result, VkCommandBuffer CommandBuffer, VkEvent Event, VkPipelineStageFlags2 StageMask) VULKAN_LAYER_BODY
 	static void CmdSetEvent2KHR(VkResult Result, VkCommandBuffer CommandBuffer, VkEvent Event, const VkDependencyInfo* DependencyInfo) VULKAN_LAYER_BODY
@@ -1492,6 +1495,13 @@ namespace VulkanRHI
 		FWrapLayer::GetDeviceImageMemoryRequirementsKHR(VK_SUCCESS, Device, Info, MemoryRequirements);
 	}
 
+	static FORCEINLINE_DEBUGGABLE void vkGetDeviceBufferMemoryRequirementsKHR(VkDevice Device, const VkDeviceBufferMemoryRequirements* Info, VkMemoryRequirements2* MemoryRequirements)
+	{
+		FWrapLayer::GetDeviceBufferMemoryRequirementsKHR(VK_RESULT_MAX_ENUM, Device, Info, MemoryRequirements);
+		VULKANAPINAMESPACE::vkGetDeviceBufferMemoryRequirementsKHR(Device, Info, MemoryRequirements);
+		FWrapLayer::GetDeviceBufferMemoryRequirementsKHR(VK_SUCCESS, Device, Info, MemoryRequirements);
+	}
+
 	static FORCEINLINE_DEBUGGABLE void vkResetQueryPoolEXT(VkDevice Device, VkQueryPool QueryPool, uint32_t FirstQuery, uint32_t QueryCount)
 	{
 		FWrapLayer::ResetQueryPoolEXT(VK_RESULT_MAX_ENUM, Device, QueryPool, FirstQuery, QueryCount);
@@ -1515,6 +1525,22 @@ namespace VulkanRHI
 		return Result;
 	}
 
+	static FORCEINLINE_DEBUGGABLE VkResult vkBindBufferMemory2KHR(VkDevice Device, uint32_t BindInfoCount, const VkBindBufferMemoryInfo* BindInfos)
+	{
+		FWrapLayer::BindBufferMemory2KHR(VK_RESULT_MAX_ENUM, Device, BindInfoCount, BindInfos);
+		const VkResult Result = VULKANAPINAMESPACE::vkBindBufferMemory2KHR(Device, BindInfoCount, BindInfos);
+		FWrapLayer::BindBufferMemory2KHR(Result, Device, BindInfoCount, BindInfos);
+		return Result;
+	}
+
+	static FORCEINLINE_DEBUGGABLE VkResult vkBindImageMemory2KHR(VkDevice Device, uint32_t BindInfoCount, const VkBindImageMemoryInfo* BindInfos)
+	{
+		FWrapLayer::BindImageMemory2KHR(VK_RESULT_MAX_ENUM, Device, BindInfoCount, BindInfos);
+		const VkResult Result = VULKANAPINAMESPACE::vkBindImageMemory2KHR(Device, BindInfoCount, BindInfos);
+		FWrapLayer::BindImageMemory2KHR(Result, Device, BindInfoCount, BindInfos);
+		return Result;
+	}
+
 	static FORCEINLINE_DEBUGGABLE void vkCmdPipelineBarrier2KHR(VkCommandBuffer CommandBuffer, const VkDependencyInfo* DependencyInfo)
 	{
 		FWrapLayer::CmdPipelineBarrier2KHR(VK_RESULT_MAX_ENUM, CommandBuffer, DependencyInfo);
diff --git a/Engine/Source/Runtime/VulkanRHI/Private/VulkanDebug.cpp b/Engine/Source/Runtime/VulkanRHI/Private/VulkanDebug.cpp
--- a/Engine/Source/Runtime/VulkanRHI/Private/VulkanDebug.cpp
+++ b/Engine/Source/Runtime/VulkanRHI/Private/VulkanDebug.cpp
@@ -4420,6 +4420,16 @@ void FWrapLayer::GetDeviceImageMemoryRequirementsKHR(VkResult Result, VkDevice D
 	}
 }
 
+void FWrapLayer::GetDeviceBufferMemoryRequirementsKHR(VkResult Result, VkDevice Device, const VkDeviceBufferMemoryRequirements* Info, VkMemoryRequirements2* MemoryRequirements)
+{
+	if (Result == VK_RESULT_MAX_ENUM)
+	{
+#if VULKAN_ENABLE_DUMP_LAYER
+		PrintfBeginResult(FString::Printf(TEXT("vkGetDeviceBufferMemoryRequirementsKHR")));
+#endif
+	}
+}
+
 void FWrapLayer::ResetQueryPoolEXT(VkResult Result, VkDevice Device, VkQueryPool QueryPool, uint32_t FirstQuery, uint32_t QueryCount)
 {
 	if (Result == VK_RESULT_MAX_ENUM)
@@ -4450,6 +4460,26 @@ void FWrapLayer::GetCalibratedTimestampsEXT(VkResult Result, VkDevice Device, ui
 	}
 }
 
+void FWrapLayer::BindBufferMemory2KHR(VkResult Result, VkDevice Device, uint32_t BindInfoCount, const VkBindBufferMemoryInfo* BindInfos)
+{
+	if (Result == VK_RESULT_MAX_ENUM)
+	{
+#if VULKAN_ENABLE_DUMP_LAYER
+		PrintfBeginResult(FString::Printf(TEXT("BindBufferMemory2KHR(Device=0x%p, BindInfoCount=%u)"), Device, BindInfoCount));
+#endif
+	}
+}
+
+void FWrapLayer::BindImageMemory2KHR(VkResult Result, VkDevice Device, uint32_t BindInfoCount, const VkBindImageMemoryInfo* BindInfos)
+{
+	if (Result == VK_RESULT_MAX_ENUM)
+	{
+#if VULKAN_ENABLE_DUMP_LAYER
+		PrintfBeginResult(FString::Printf(TEXT("BindImageMemory2KHR(Device=0x%p, BindInfoCount=%u)"), Device, BindInfoCount));
+#endif
+	}
+}
+
 void FWrapLayer::CmdPipelineBarrier2KHR(VkResult Result, VkCommandBuffer CommandBuffer, const VkDependencyInfo* DependencyInfo)
 {
 	if (Result == VK_RESULT_MAX_ENUM)
diff --git a/Engine/Source/Runtime/VulkanRHI/Private/VulkanExtensions.cpp b/Engine/Source/Runtime/VulkanRHI/Private/VulkanExtensions.cpp
--- a/Engine/Source/Runtime/VulkanRHI/Private/VulkanExtensions.cpp
+++ b/Engine/Source/Runtime/VulkanRHI/Private/VulkanExtensions.cpp
@@ -1088,6 +1088,7 @@ FVulkanDeviceExtensionArray FVulkanDeviceExtension::GetUESupportedDeviceExtensio
 	ADD_SIMPLE_EXTENSION(VK_KHR_SHADER_FLOAT_CONTROLS_EXTENSION_NAME,    VULKAN_RHI_RAYTRACING,                VK_API_VERSION_1_2,            DEVICE_EXT_FLAG_SETTER(HasShaderFloatControls));
 	ADD_SIMPLE_EXTENSION(VK_KHR_IMAGE_FORMAT_LIST_EXTENSION_NAME,        VULKAN_EXTENSION_ENABLED,             VK_API_VERSION_1_2,            DEVICE_EXT_FLAG_SETTER(HasKHRImageFormatList));
 	ADD_SIMPLE_EXTENSION(VK_EXT_VALIDATION_CACHE_EXTENSION_NAME,         VULKAN_SUPPORTS_VALIDATION_CACHE,     VULKAN_EXTENSION_NOT_PROMOTED, DEVICE_EXT_FLAG_SETTER(HasEXTValidationCache));
+	ADD_SIMPLE_EXTENSION(VK_KHR_BIND_MEMORY_2_EXTENSION_NAME,            VULKAN_EXTENSION_ENABLED,             VK_API_VERSION_1_1,            nullptr);
 
 
 	// Externally activated extensions (supported by the engine, but enabled externally by plugin or other) :
diff --git a/Engine/Source/Runtime/VulkanRHI/Private/VulkanLoader.h b/Engine/Source/Runtime/VulkanRHI/Private/VulkanLoader.h
--- a/Engine/Source/Runtime/VulkanRHI/Private/VulkanLoader.h
+++ b/Engine/Source/Runtime/VulkanRHI/Private/VulkanLoader.h
@@ -197,6 +197,7 @@
 	EnumMacro(PFN_vkCreateRenderPass2KHR, vkCreateRenderPass2KHR) \
 	EnumMacro(PFN_vkCmdBeginRenderPass2KHR, vkCmdBeginRenderPass2KHR) \
 	EnumMacro(PFN_vkGetDeviceImageMemoryRequirementsKHR, vkGetDeviceImageMemoryRequirementsKHR) \
+	EnumMacro(PFN_vkGetDeviceBufferMemoryRequirementsKHR, vkGetDeviceBufferMemoryRequirementsKHR) \
 	EnumMacro(PFN_vkResetQueryPoolEXT, vkResetQueryPoolEXT) \
 	EnumMacro(PFN_vkCmdPipelineBarrier2KHR, vkCmdPipelineBarrier2KHR) \
 	EnumMacro(PFN_vkCmdResetEvent2KHR, vkCmdResetEvent2KHR) \
@@ -205,7 +206,9 @@
 	EnumMacro(PFN_vkQueueSubmit2KHR, vkQueueSubmit2KHR) \
 	EnumMacro(PFN_vkCreateSharedSwapchainsKHR, vkCreateSharedSwapchainsKHR) \
 	EnumMacro(PFN_vkGetPhysicalDeviceCalibrateableTimeDomainsEXT, vkGetPhysicalDeviceCalibrateableTimeDomainsEXT) \
-	EnumMacro(PFN_vkGetCalibratedTimestampsEXT, vkGetCalibratedTimestampsEXT)
+	EnumMacro(PFN_vkGetCalibratedTimestampsEXT, vkGetCalibratedTimestampsEXT) \
+	EnumMacro(PFN_vkBindBufferMemory2KHR, vkBindBufferMemory2KHR) \
+	EnumMacro(PFN_vkBindImageMemory2KHR, vkBindImageMemory2KHR)
 
 
 // List of all Vulkan entry points
