VulkanRHI: Minor clean up: Remove legacy unused cvars. Add extension reported memory usage to vulkanmemoryraw stats.

#rnx

diff --git a/Engine/Source/Runtime/VulkanRHI/Private/Android/VulkanAndroidPlatform.h b/Engine/Source/Runtime/VulkanRHI/Private/Android/VulkanAndroidPlatform.h
--- a/Engine/Source/Runtime/VulkanRHI/Private/Android/VulkanAndroidPlatform.h
+++ b/Engine/Source/Runtime/VulkanRHI/Private/Android/VulkanAndroidPlatform.h
@@ -15,7 +15,6 @@
 #define VULKAN_SHOULD_USE_COMMANDWRAPPERS			VULKAN_SHOULD_USE_LLM //LLM on Vulkan needs command wrappers to account for vkallocs
 #define VULKAN_ENABLE_LRU_CACHE						1
 #define VULKAN_SUPPORTS_GOOGLE_DISPLAY_TIMING		1
-#define VULKAN_FREEPAGE_FOR_TYPE					1
 #define VULKAN_PURGE_SHADER_MODULES					0
 #define VULKAN_SUPPORTS_DEDICATED_ALLOCATION		0
 #define VULKAN_SUPPORTS_PHYSICAL_DEVICE_PROPERTIES2	1
diff --git a/Engine/Source/Runtime/VulkanRHI/Private/VulkanMemory.cpp b/Engine/Source/Runtime/VulkanRHI/Private/VulkanMemory.cpp
--- a/Engine/Source/Runtime/VulkanRHI/Private/VulkanMemory.cpp
+++ b/Engine/Source/Runtime/VulkanRHI/Private/VulkanMemory.cpp
@@ -33,7 +33,6 @@ const uint32 NUM_FRAMES_TO_WAIT_FOR_RESOURCE_DELETE = 2;
 #define UE_VK_MEMORY_KEEP_FREELIST_SORTED					1
 #define UE_VK_MEMORY_JOIN_FREELIST_ON_THE_FLY				(UE_VK_MEMORY_KEEP_FREELIST_SORTED && 1)
 #define UE_VK_MEMORY_KEEP_FREELIST_SORTED_CATCHBUGS			0 // debugging
-#define VULKAN_FREE_ALL_PAGES ((PLATFORM_ANDROID) ? 1 : 0)
 
 #define VULKAN_LOG_MEMORY_UELOG 1 //in case of debugging, it is useful to be able to log directly to LowLevelPrintf, as this is easier to diff. Please do not delete this code.
 
@@ -44,64 +43,41 @@ const uint32 NUM_FRAMES_TO_WAIT_FOR_RESOURCE_DELETE = 2;
 #endif
 
 
-DECLARE_STATS_GROUP(TEXT("Vulkan Memory Raw"), STATGROUP_VulkanMemoryRaw, STATCAT_Advanced);
-DECLARE_MEMORY_STAT_EXTERN(TEXT("Dedicated Memory"), STAT_VulkanDedicatedMemory, STATGROUP_VulkanMemoryRaw, );
-DECLARE_MEMORY_STAT_EXTERN(TEXT("MemoryPool 0"), STAT_VulkanMemory0, STATGROUP_VulkanMemoryRaw, );
-DECLARE_MEMORY_STAT_EXTERN(TEXT("MemoryPool 1"), STAT_VulkanMemory1, STATGROUP_VulkanMemoryRaw, );
-DECLARE_MEMORY_STAT_EXTERN(TEXT("MemoryPool 2"), STAT_VulkanMemory2, STATGROUP_VulkanMemoryRaw, );
-DECLARE_MEMORY_STAT_EXTERN(TEXT("MemoryPool 3"), STAT_VulkanMemory3, STATGROUP_VulkanMemoryRaw, );
-DECLARE_MEMORY_STAT_EXTERN(TEXT("MemoryPool 4"), STAT_VulkanMemory4, STATGROUP_VulkanMemoryRaw, );
-DECLARE_MEMORY_STAT_EXTERN(TEXT("MemoryPool 5"), STAT_VulkanMemory5, STATGROUP_VulkanMemoryRaw, );
-DECLARE_MEMORY_STAT_EXTERN(TEXT("MemoryPool (remaining)"), STAT_VulkanMemoryX, STATGROUP_VulkanMemoryRaw, );
-
-DECLARE_MEMORY_STAT_EXTERN(TEXT("MemoryPool 0 Reserved"), STAT_VulkanMemory0Reserved, STATGROUP_VulkanMemoryRaw, );
-DECLARE_MEMORY_STAT_EXTERN(TEXT("MemoryPool 1 Reserved"), STAT_VulkanMemory1Reserved, STATGROUP_VulkanMemoryRaw, );
-DECLARE_MEMORY_STAT_EXTERN(TEXT("MemoryPool 2 Reserved"), STAT_VulkanMemory2Reserved, STATGROUP_VulkanMemoryRaw, );
-DECLARE_MEMORY_STAT_EXTERN(TEXT("MemoryPool 3 Reserved"), STAT_VulkanMemory3Reserved, STATGROUP_VulkanMemoryRaw, );
-DECLARE_MEMORY_STAT_EXTERN(TEXT("MemoryPool 4 Reserved"), STAT_VulkanMemory4Reserved, STATGROUP_VulkanMemoryRaw, );
-DECLARE_MEMORY_STAT_EXTERN(TEXT("MemoryPool 5 Reserved"), STAT_VulkanMemory5Reserved, STATGROUP_VulkanMemoryRaw, );
-DECLARE_MEMORY_STAT_EXTERN(TEXT("MemoryPool (remaining) Reserved"), STAT_VulkanMemoryXReserved, STATGROUP_VulkanMemoryRaw, );
-
-DECLARE_MEMORY_STAT_EXTERN(TEXT("_Total Allocated"), STAT_VulkanMemoryTotal, STATGROUP_VulkanMemoryRaw, );
-DECLARE_MEMORY_STAT_EXTERN(TEXT("_Reserved"), STAT_VulkanMemoryReserved, STATGROUP_VulkanMemoryRaw, );
-
-DECLARE_MEMORY_STAT_EXTERN(TEXT("MemoryPool 0 Budget"), STAT_VulkanMemoryBudget0, STATGROUP_VulkanMemoryRaw, );
-DECLARE_MEMORY_STAT_EXTERN(TEXT("MemoryPool 1 Budget"), STAT_VulkanMemoryBudget1, STATGROUP_VulkanMemoryRaw, );
-DECLARE_MEMORY_STAT_EXTERN(TEXT("MemoryPool 2 Budget"), STAT_VulkanMemoryBudget2, STATGROUP_VulkanMemoryRaw, );
-DECLARE_MEMORY_STAT_EXTERN(TEXT("MemoryPool 3 Budget"), STAT_VulkanMemoryBudget3, STATGROUP_VulkanMemoryRaw, );
-DECLARE_MEMORY_STAT_EXTERN(TEXT("MemoryPool 4 Budget"), STAT_VulkanMemoryBudget4, STATGROUP_VulkanMemoryRaw, );
-DECLARE_MEMORY_STAT_EXTERN(TEXT("MemoryPool 5 Budget"), STAT_VulkanMemoryBudget5, STATGROUP_VulkanMemoryRaw, );
-DECLARE_MEMORY_STAT_EXTERN(TEXT("MemoryPool (remaining) Budget"), STAT_VulkanMemoryBudgetX, STATGROUP_VulkanMemoryRaw, );
-
-
-DEFINE_STAT(STAT_VulkanDedicatedMemory);
-DEFINE_STAT(STAT_VulkanMemory0);
-DEFINE_STAT(STAT_VulkanMemory1);
-DEFINE_STAT(STAT_VulkanMemory2);
-DEFINE_STAT(STAT_VulkanMemory3);
-DEFINE_STAT(STAT_VulkanMemory4);
-DEFINE_STAT(STAT_VulkanMemory5);
-DEFINE_STAT(STAT_VulkanMemoryX);
-
-DEFINE_STAT(STAT_VulkanMemory0Reserved);
-DEFINE_STAT(STAT_VulkanMemory1Reserved);
-DEFINE_STAT(STAT_VulkanMemory2Reserved);
-DEFINE_STAT(STAT_VulkanMemory3Reserved);
-DEFINE_STAT(STAT_VulkanMemory4Reserved);
-DEFINE_STAT(STAT_VulkanMemory5Reserved);
-DEFINE_STAT(STAT_VulkanMemoryXReserved);
-DEFINE_STAT(STAT_VulkanMemoryReserved);
-
-DEFINE_STAT(STAT_VulkanMemoryBudget0);
-DEFINE_STAT(STAT_VulkanMemoryBudget1);
-DEFINE_STAT(STAT_VulkanMemoryBudget2);
-DEFINE_STAT(STAT_VulkanMemoryBudget3);
-DEFINE_STAT(STAT_VulkanMemoryBudget4);
-DEFINE_STAT(STAT_VulkanMemoryBudget5);
-DEFINE_STAT(STAT_VulkanMemoryBudgetX);
-
-
-DEFINE_STAT(STAT_VulkanMemoryTotal);
+DECLARE_STATS_GROUP_SORTBYNAME(TEXT("Vulkan Memory Raw"), STATGROUP_VulkanMemoryRaw, STATCAT_Advanced);
+DECLARE_MEMORY_STAT(TEXT("Dedicated Memory"), STAT_VulkanDedicatedMemory, STATGROUP_VulkanMemoryRaw);
+DECLARE_MEMORY_STAT(TEXT("MemoryPool 0"), STAT_VulkanMemory0, STATGROUP_VulkanMemoryRaw);
+DECLARE_MEMORY_STAT(TEXT("MemoryPool 1"), STAT_VulkanMemory1, STATGROUP_VulkanMemoryRaw);
+DECLARE_MEMORY_STAT(TEXT("MemoryPool 2"), STAT_VulkanMemory2, STATGROUP_VulkanMemoryRaw);
+DECLARE_MEMORY_STAT(TEXT("MemoryPool 3"), STAT_VulkanMemory3, STATGROUP_VulkanMemoryRaw);
+DECLARE_MEMORY_STAT(TEXT("MemoryPool 4"), STAT_VulkanMemory4, STATGROUP_VulkanMemoryRaw);
+DECLARE_MEMORY_STAT(TEXT("MemoryPool 5"), STAT_VulkanMemory5, STATGROUP_VulkanMemoryRaw);
+DECLARE_MEMORY_STAT(TEXT("MemoryPool (remaining)"), STAT_VulkanMemoryX, STATGROUP_VulkanMemoryRaw);
+
+DECLARE_MEMORY_STAT(TEXT("MemoryPool 0 Reserved"), STAT_VulkanMemory0Reserved, STATGROUP_VulkanMemoryRaw);
+DECLARE_MEMORY_STAT(TEXT("MemoryPool 1 Reserved"), STAT_VulkanMemory1Reserved, STATGROUP_VulkanMemoryRaw);
+DECLARE_MEMORY_STAT(TEXT("MemoryPool 2 Reserved"), STAT_VulkanMemory2Reserved, STATGROUP_VulkanMemoryRaw);
+DECLARE_MEMORY_STAT(TEXT("MemoryPool 3 Reserved"), STAT_VulkanMemory3Reserved, STATGROUP_VulkanMemoryRaw);
+DECLARE_MEMORY_STAT(TEXT("MemoryPool 4 Reserved"), STAT_VulkanMemory4Reserved, STATGROUP_VulkanMemoryRaw);
+DECLARE_MEMORY_STAT(TEXT("MemoryPool 5 Reserved"), STAT_VulkanMemory5Reserved, STATGROUP_VulkanMemoryRaw);
+DECLARE_MEMORY_STAT(TEXT("MemoryPool (remaining) Reserved"), STAT_VulkanMemoryXReserved, STATGROUP_VulkanMemoryRaw);
+
+DECLARE_MEMORY_STAT(TEXT("_Total Allocated"), STAT_VulkanMemoryTotal, STATGROUP_VulkanMemoryRaw);
+DECLARE_MEMORY_STAT(TEXT("_Reserved"), STAT_VulkanMemoryReserved, STATGROUP_VulkanMemoryRaw);
+
+DECLARE_MEMORY_STAT(TEXT("Memory Heap 0 Budget"), STAT_VulkanMemoryBudget0, STATGROUP_VulkanMemoryRaw);
+DECLARE_MEMORY_STAT(TEXT("Memory Heap 1 Budget"), STAT_VulkanMemoryBudget1, STATGROUP_VulkanMemoryRaw);
+DECLARE_MEMORY_STAT(TEXT("Memory Heap 2 Budget"), STAT_VulkanMemoryBudget2, STATGROUP_VulkanMemoryRaw);
+DECLARE_MEMORY_STAT(TEXT("Memory Heap 3 Budget"), STAT_VulkanMemoryBudget3, STATGROUP_VulkanMemoryRaw);
+DECLARE_MEMORY_STAT(TEXT("Memory Heap 4 Budget"), STAT_VulkanMemoryBudget4, STATGROUP_VulkanMemoryRaw);
+DECLARE_MEMORY_STAT(TEXT("Memory Heap 5 Budget"), STAT_VulkanMemoryBudget5, STATGROUP_VulkanMemoryRaw);
+DECLARE_MEMORY_STAT(TEXT("Memory Heap (remaining) Budget"), STAT_VulkanMemoryBudgetX, STATGROUP_VulkanMemoryRaw);
+
+DECLARE_MEMORY_STAT(TEXT("Memory Heap 0 Usage"), STAT_VulkanMemoryUsage0, STATGROUP_VulkanMemoryRaw);
+DECLARE_MEMORY_STAT(TEXT("Memory Heap 1 Usage"), STAT_VulkanMemoryUsage1, STATGROUP_VulkanMemoryRaw);
+DECLARE_MEMORY_STAT(TEXT("Memory Heap 2 Usage"), STAT_VulkanMemoryUsage2, STATGROUP_VulkanMemoryRaw);
+DECLARE_MEMORY_STAT(TEXT("Memory Heap 3 Usage"), STAT_VulkanMemoryUsage3, STATGROUP_VulkanMemoryRaw);
+DECLARE_MEMORY_STAT(TEXT("Memory Heap 4 Usage"), STAT_VulkanMemoryUsage4, STATGROUP_VulkanMemoryRaw);
+DECLARE_MEMORY_STAT(TEXT("Memory Heap 5 Usage"), STAT_VulkanMemoryUsage5, STATGROUP_VulkanMemoryRaw);
 
 
 DECLARE_STATS_GROUP(TEXT("Vulkan Memory"), STATGROUP_VulkanMemory, STATCAT_Advanced);
@@ -205,22 +181,6 @@ static FAutoConsoleVariableRef CVarVulkanUseBufferBinning(
 	ECVF_ReadOnly
 );
 
-static int32 GVulkanFreePageForType = VULKAN_FREEPAGE_FOR_TYPE;
-static FAutoConsoleVariableRef CVarVulkanFreePageForType(
-	TEXT("r.Vulkan.FreePageForType"),
-	GVulkanFreePageForType,
-	TEXT("Enable separate free page list for images and buffers."),
-	ECVF_ReadOnly
-);
-
-static int32 GVulkanFreeAllPages = VULKAN_FREE_ALL_PAGES;
-static FAutoConsoleVariableRef CVarVulkanFreeAllPages(
-	TEXT("r.Vulkan.FreeAllPages"),
-	GVulkanFreeAllPages,
-	TEXT("Enable to fully free all pages early. default on android only"),
-	ECVF_ReadOnly
-);
-
 static int32 GVulkanLogEvictStatus = 0;
 static FAutoConsoleVariableRef GVarVulkanLogEvictStatus(
 	TEXT("r.Vulkan.LogEvictStatus"),
@@ -613,6 +573,13 @@ namespace VulkanRHI
 			SET_DWORD_STAT(STAT_VulkanMemoryBudget4, MemoryBudget.heapBudget[4]);
 			SET_DWORD_STAT(STAT_VulkanMemoryBudget5, MemoryBudget.heapBudget[5]);
 			SET_DWORD_STAT(STAT_VulkanMemoryBudgetX, BudgetX);
+
+			SET_DWORD_STAT(STAT_VulkanMemoryUsage0, MemoryBudget.heapUsage[0]);
+			SET_DWORD_STAT(STAT_VulkanMemoryUsage1, MemoryBudget.heapUsage[1]);
+			SET_DWORD_STAT(STAT_VulkanMemoryUsage2, MemoryBudget.heapUsage[2]);
+			SET_DWORD_STAT(STAT_VulkanMemoryUsage3, MemoryBudget.heapUsage[3]);
+			SET_DWORD_STAT(STAT_VulkanMemoryUsage4, MemoryBudget.heapUsage[4]);
+			SET_DWORD_STAT(STAT_VulkanMemoryUsage5, MemoryBudget.heapUsage[5]);
 		}
 		else
 		{
diff --git a/Engine/Source/Runtime/VulkanRHI/Public/VulkanConfiguration.h b/Engine/Source/Runtime/VulkanRHI/Public/VulkanConfiguration.h
--- a/Engine/Source/Runtime/VulkanRHI/Public/VulkanConfiguration.h
+++ b/Engine/Source/Runtime/VulkanRHI/Public/VulkanConfiguration.h
@@ -51,10 +51,6 @@ struct VkAllocationCallbacks;
 
 #define VULKAN_SINGLE_ALLOCATION_PER_RESOURCE					0
 
-#ifndef VULKAN_FREEPAGE_FOR_TYPE
-	#define VULKAN_FREEPAGE_FOR_TYPE							0
-#endif
-
 #ifndef VULKAN_SHOULD_USE_LLM
 	#define VULKAN_SHOULD_USE_LLM								0
 #endif
