PR# 9814: Fix vkCreateRenderPass2 crash accessing EXTFragmentDensityMap pNext (Contributed by [REDACTED_PERSON])

CreateInfo.pNext was being assigned a pointer to a struct created on the stack, causing a crash when accessed after it went out of scope.

diff --git a/Engine/Source/Runtime/VulkanRHI/Private/VulkanRenderpass.h b/Engine/Source/Runtime/VulkanRHI/Private/VulkanRenderpass.h
--- a/Engine/Source/Runtime/VulkanRHI/Private/VulkanRenderpass.h
+++ b/Engine/Source/Runtime/VulkanRHI/Private/VulkanRenderpass.h
@@ -346,6 +346,16 @@ struct FVulkanRenderPassCreateInfo<VkRenderPassCreateInfo>
 	}
 };
 
+struct FVulkanRenderPassFragmentDensityMapCreateInfoEXT
+	: public VkRenderPassFragmentDensityMapCreateInfoEXT
+{
+	FVulkanRenderPassFragmentDensityMapCreateInfoEXT()
+	{
+		ZeroVulkanStruct(*this, VK_STRUCTURE_TYPE_RENDER_PASS_FRAGMENT_DENSITY_MAP_CREATE_INFO_EXT);
+	}
+};
+
+
 template<>
 struct FVulkanRenderPassCreateInfo<VkRenderPassCreateInfo2>
 	: public VkRenderPassCreateInfo2
@@ -626,10 +636,8 @@ public:
 			}
 		}
 
-		VkRenderPassFragmentDensityMapCreateInfoEXT FragDensityCreateInfo;
 		if (Device.GetOptionalExtensions().HasEXTFragmentDensityMap && RTLayout.GetHasFragmentDensityAttachment())
 		{
-			ZeroVulkanStruct(FragDensityCreateInfo, VK_STRUCTURE_TYPE_RENDER_PASS_FRAGMENT_DENSITY_MAP_CREATE_INFO_EXT);
 			FragDensityCreateInfo.fragmentDensityMapAttachment = *RTLayout.GetFragmentDensityAttachmentReference();
 
 			// Chain fragment density info onto create info and the rest of the pNexts
@@ -679,6 +687,7 @@ private:
 
 	FVulkanAttachmentReference<VkAttachmentReference2> ShadingRateAttachmentReference;
 	FVulkanFragmentShadingRateAttachmentInfo FragmentShadingRateAttachmentInfo;
+	FVulkanRenderPassFragmentDensityMapCreateInfoEXT FragDensityCreateInfo;
 
 	TRenderPassCreateInfoClass CreateInfo;
 	FVulkanDevice& Device;
