PR #9816: Correctly read CorrelationMask when creting vulkanrender pass

CorrelationMask was created on the stack, but we attempted to read it after it was out of scope. Now it's a member of of FVulkanRenderPassBuilder.

diff --git a/Engine/Source/Runtime/VulkanRHI/Private/VulkanRenderpass.h b/Engine/Source/Runtime/VulkanRHI/Private/VulkanRenderpass.h
--- a/Engine/Source/Runtime/VulkanRHI/Private/VulkanRenderpass.h
+++ b/Engine/Source/Runtime/VulkanRHI/Private/VulkanRenderpass.h
@@ -404,6 +404,7 @@ class FVulkanRenderPassBuilder
 public:
 	FVulkanRenderPassBuilder(FVulkanDevice& InDevice)
 		: Device(InDevice)
+		, CorrelationMask(0)
 	{}
 
 	void BuildCreateInfo(const FVulkanRenderTargetLayout& RTLayout)
@@ -611,7 +612,7 @@ public:
 		Bit mask that specifices correlation between views
 		An implementation may use this for optimizations (concurrent render)
 		*/
-		const uint32_t CorrelationMask = MultiviewMask;
+		CorrelationMask = MultiviewMask;
 
 		VkRenderPassMultiviewCreateInfo MultiviewInfo;
 		if (RTLayout.GetIsMultiView())
@@ -691,6 +692,8 @@ private:
 
 	TRenderPassCreateInfoClass CreateInfo;
 	FVulkanDevice& Device;
+
+	uint32_t CorrelationMask;
 };
 
 VkRenderPass CreateVulkanRenderPass(FVulkanDevice& Device, const FVulkanRenderTargetLayout& RTLayout);
