VulkanRHI: Add missing vkGetBufferMemoryRequirements2KHR in command wrapper.

#rnx

diff --git a/Engine/Source/Runtime/VulkanRHI/Private/VulkanCommandWrappers.h b/Engine/Source/Runtime/VulkanRHI/Private/VulkanCommandWrappers.h
--- a/Engine/Source/Runtime/VulkanRHI/Private/VulkanCommandWrappers.h
+++ b/Engine/Source/Runtime/VulkanRHI/Private/VulkanCommandWrappers.h
@@ -165,10 +165,10 @@ struct FWrapLayer
 	static void AcquireNextImageKHR(VkResult Result, VkDevice Device, VkSwapchainKHR Swapchain, uint64_t Timeout, VkSemaphore Semaphore, VkFence Fence, uint32_t* ImageIndex) VULKAN_LAYER_BODY
 	static void DestroySurfaceKHR(VkResult Result, VkInstance Instance, VkSurfaceKHR SurfaceKHR) VULKAN_LAYER_BODY
 	static void DestroySwapchainKHR(VkResult Result, VkDevice Device, VkSwapchainKHR Swapchain) VULKAN_LAYER_BODY
-#if VULKAN_SUPPORTS_DEDICATED_ALLOCATION
 	static void GetImageMemoryRequirements2KHR(VkResult Result, VkDevice Device, const VkImageMemoryRequirementsInfo2KHR* Info, VkMemoryRequirements2KHR* MemoryRequirements) VULKAN_LAYER_BODY
-#endif
-	VULKAN_EXTERN_EXPORT static void GetPhysicalDeviceMemoryProperties2(VkResult Result, VkPhysicalDevice PhysicalDevice, VkPhysicalDeviceMemoryProperties2* MemoryProperties) VULKAN_LAYER_BODY 
+	static void GetBufferMemoryRequirements2KHR(VkResult Result, VkDevice Device, const VkBufferMemoryRequirementsInfo2KHR* Info, VkMemoryRequirements2KHR* MemoryRequirements) VULKAN_LAYER_BODY
+
+		VULKAN_EXTERN_EXPORT static void GetPhysicalDeviceMemoryProperties2(VkResult Result, VkPhysicalDevice PhysicalDevice, VkPhysicalDeviceMemoryProperties2* MemoryProperties) VULKAN_LAYER_BODY 
 #if VULKAN_SUPPORTS_PHYSICAL_DEVICE_PROPERTIES2
 	VULKAN_EXTERN_EXPORT static void GetPhysicalDeviceProperties2KHR(VkResult Result, VkPhysicalDevice PhysicalDevice, VkPhysicalDeviceProperties2KHR* Properties) VULKAN_LAYER_BODY
 	VULKAN_EXTERN_EXPORT static void GetPhysicalDeviceFeatures2KHR(VkResult Result, VkPhysicalDevice PhysicalDevice, VkPhysicalDeviceFeatures2KHR* Features) VULKAN_LAYER_BODY
@@ -1396,14 +1396,19 @@ namespace VulkanRHI
 		FWrapLayer::DestroySurfaceKHR(VK_SUCCESS, Instance, Surface);
 	}
 
-#if VULKAN_SUPPORTS_DEDICATED_ALLOCATION
 	static FORCEINLINE_DEBUGGABLE void vkGetImageMemoryRequirements2KHR(VkDevice Device, const VkImageMemoryRequirementsInfo2KHR* Info, VkMemoryRequirements2KHR* MemoryRequirements)
 	{
 		FWrapLayer::GetImageMemoryRequirements2KHR(VK_RESULT_MAX_ENUM, Device, Info, MemoryRequirements);
 		VULKANAPINAMESPACE::vkGetImageMemoryRequirements2KHR(Device, Info, MemoryRequirements);
 		FWrapLayer::GetImageMemoryRequirements2KHR(VK_SUCCESS, Device, Info, MemoryRequirements);
 	}
-#endif
+
+	static FORCEINLINE_DEBUGGABLE void vkGetBufferMemoryRequirements2KHR(VkDevice Device, const VkBufferMemoryRequirementsInfo2KHR* Info, VkMemoryRequirements2KHR* MemoryRequirements)
+	{
+		FWrapLayer::GetBufferMemoryRequirements2KHR(VK_RESULT_MAX_ENUM, Device, Info, MemoryRequirements);
+		VULKANAPINAMESPACE::vkGetBufferMemoryRequirements2KHR(Device, Info, MemoryRequirements);
+		FWrapLayer::GetBufferMemoryRequirements2KHR(VK_SUCCESS, Device, Info, MemoryRequirements);
+	}
 
 #if VULKAN_RHI_RAYTRACING
 	static FORCEINLINE_DEBUGGABLE VkResult vkCreateAccelerationStructureKHR(VkDevice Device, const VkAccelerationStructureCreateInfoKHR* CreateInfo, const VkAllocationCallbacks* Allocator, VkAccelerationStructureKHR* AccelerationStructure)
diff --git a/Engine/Source/Runtime/VulkanRHI/Private/VulkanDebug.cpp b/Engine/Source/Runtime/VulkanRHI/Private/VulkanDebug.cpp
--- a/Engine/Source/Runtime/VulkanRHI/Private/VulkanDebug.cpp
+++ b/Engine/Source/Runtime/VulkanRHI/Private/VulkanDebug.cpp
@@ -3788,8 +3788,6 @@ void FWrapLayer::GetPhysicalDeviceQueueFamilyProperties(VkResult Result, VkPhysi
 	}
 }
 
-
-#if VULKAN_SUPPORTS_DEDICATED_ALLOCATION
 void FWrapLayer::GetImageMemoryRequirements2KHR(VkResult Result, VkDevice Device, const VkImageMemoryRequirementsInfo2KHR* Info, VkMemoryRequirements2KHR* MemoryRequirements)
 {
 	if (Result == VK_RESULT_MAX_ENUM)
@@ -3799,7 +3797,16 @@ void FWrapLayer::GetImageMemoryRequirements2KHR(VkResult Result, VkDevice Device
 #endif
 	}
 }
-#endif	// VULKAN_SUPPORTS_DEDICATED_ALLOCATION
+
+void FWrapLayer::GetBufferMemoryRequirements2KHR(VkResult Result, VkDevice Device, const VkBufferMemoryRequirementsInfo2KHR* Info, VkMemoryRequirements2KHR* MemoryRequirements)
+{
+	if (Result == VK_RESULT_MAX_ENUM)
+	{
+#if VULKAN_ENABLE_DUMP_LAYER
+		DevicePrintfBegin(Device, FString::Printf(TEXT("vkGetBufferemoryRequirements2KHR(Info=0x%p, MemReqs=0x%p)[...]"), Info, MemoryRequirements));
+#endif
+	}
+}
 
 void FWrapLayer::GetImageSparseMemoryRequirements(VkResult Result, VkDevice Device, VkImage Image, uint32_t* pSparseMemoryRequirementCount, VkSparseImageMemoryRequirements* pSparseMemoryRequirements)
 {
