VulkanRHI: Fix texture memory stat to exclude transient memory.

#rnx

diff --git a/Engine/Source/Runtime/VulkanRHI/Private/VulkanTexture.cpp b/Engine/Source/Runtime/VulkanRHI/Private/VulkanTexture.cpp
--- a/Engine/Source/Runtime/VulkanRHI/Private/VulkanTexture.cpp
+++ b/Engine/Source/Runtime/VulkanRHI/Private/VulkanTexture.cpp
@@ -696,17 +696,18 @@ void FVulkanTexture::DestroySurface()
 
 		if (bIsLocalOwner)
 		{
-			uint64 Size = 0;
+			// If we don't own the allocation, it's transient memory not included in stats
+			if (Allocation.HasAllocation())
+			{
+				VulkanTextureDestroyed(Allocation.Size, GetViewType(), bRenderTarget);
+			}
 
 			if (Image != VK_NULL_HANDLE)
 			{
-				Size = GetMemorySize();
 				Device->GetDeferredDeletionQueue().EnqueueResource(VulkanRHI::FDeferredDeletionQueue2::EType::Image, Image);
 				Device->GetMemoryManager().FreeVulkanAllocation(Allocation);
 				Image = VK_NULL_HANDLE;
 			}
-
-			VulkanTextureDestroyed(Size, GetViewType(), bRenderTarget);
 		}
 
 		ImageOwnerType = EImageOwnerType::None;
@@ -1649,12 +1650,12 @@ FVulkanTexture::FVulkanTexture(FVulkanDevice& InDevice, const FRHITextureCreateD
 					checkNoEntry();
 				}
 			}
+
+			// update rhi stats
+			VulkanTextureAllocated(Allocation.Size, GetViewType(), bRenderTarget);
 		}
 		Allocation.BindImage(Device, Image);
 
-		// update rhi stats
-		VulkanTextureAllocated(MemoryRequirements.size, GetViewType(), bRenderTarget);
-
 		Tiling = ImageCreateInfo.ImageCreateInfo.tiling;
 		check(Tiling == VK_IMAGE_TILING_LINEAR || Tiling == VK_IMAGE_TILING_OPTIMAL);
 
