VulkanRHI:  Remove warning spam in logs from undefined stencil state.

#rnx

diff --git a/Engine/Source/Runtime/VulkanRHI/Private/VulkanRenderTarget.cpp b/Engine/Source/Runtime/VulkanRHI/Private/VulkanRenderTarget.cpp
--- a/Engine/Source/Runtime/VulkanRHI/Private/VulkanRenderTarget.cpp
+++ b/Engine/Source/Runtime/VulkanRHI/Private/VulkanRenderTarget.cpp
@@ -908,10 +908,7 @@ void FVulkanCommandListContext::RHIBeginRenderPass(const FRHIRenderPassInfo& InI
 					if (CurrentStencilLayout == VK_IMAGE_LAYOUT_UNDEFINED)
 					{
 						CurrentStencilLayout = GetQueue()->GetLayoutManager().GetDepthStencilHint(VulkanTexture, VK_IMAGE_ASPECT_STENCIL_BIT);
-						if (CurrentStencilLayout == VK_IMAGE_LAYOUT_UNDEFINED)
-						{
-							UE_LOG(LogVulkanRHI, Warning, TEXT("Render pass [%s] used an unknown stencil state!"), InName ? InName : TEXT("unknown"));
-						}
+						// If the layout is still UNDEFINED, FVulkanRenderTargetLayout will force it into a known state
 					}
 				}
 			}
