VK: make sure we invalidate cached memory when locking staging buffer for readback

diff --git a/Engine/Source/Runtime/VulkanRHI/Private/VulkanUtil.cpp b/Engine/Source/Runtime/VulkanRHI/Private/VulkanUtil.cpp
--- a/Engine/Source/Runtime/VulkanRHI/Private/VulkanUtil.cpp
+++ b/Engine/Source/Runtime/VulkanRHI/Private/VulkanUtil.cpp
@@ -89,6 +89,8 @@ void* FVulkanStagingBuffer::Lock(uint32 Offset, uint32 NumBytes)
 	bIsLocked = true;
 	const uint32 EndOffset = Offset + NumBytes;
 	checkf(EndOffset <= QueuedNumBytes, TEXT("Lock at Offset (%u) and NumBytes (%u) reads beyond the allocated size of the staging buffer (%u)"), Offset, NumBytes, QueuedNumBytes);
+	// make sure cached memory is invalidated
+	StagingBuffer->InvalidateMappedMemory();
 	return (void*)((uint8*)StagingBuffer->GetMappedPointer() + Offset);
 }
 
