Fixed: broken ASTC HDR compression

diff --git a/Engine/Source/Developer/TargetPlatform/Private/TargetPlatformBase.cpp b/Engine/Source/Developer/TargetPlatform/Private/TargetPlatformBase.cpp
--- a/Engine/Source/Developer/TargetPlatform/Private/TargetPlatformBase.cpp
+++ b/Engine/Source/Developer/TargetPlatform/Private/TargetPlatformBase.cpp
@@ -94,7 +94,10 @@ static FAutoConsoleVariableRef CVarAllowASTCHDRProfile(
 
 bool FTargetPlatformBase::UsesASTCHDR() const
 {
-	return GASTCHDRProfile != 0;
+	static IConsoleVariable* CVar = IConsoleManager::Get().FindConsoleVariable(TEXT("cook.ASTCTextureCompressor"));
+	const bool bUsesARMCompressor = (CVar ? (CVar->GetInt() != 0) : false);
+	
+	return (bUsesARMCompressor && GASTCHDRProfile != 0);
 }
 
 void FTargetPlatformBase::GetRayTracingShaderFormats(TArray<FName>& OutFormats) const
diff --git a/Engine/Source/Developer/TextureFormatASTC/Private/TextureFormatASTC.cpp b/Engine/Source/Developer/TextureFormatASTC/Private/TextureFormatASTC.cpp
--- a/Engine/Source/Developer/TextureFormatASTC/Private/TextureFormatASTC.cpp
+++ b/Engine/Source/Developer/TextureFormatASTC/Private/TextureFormatASTC.cpp
@@ -67,7 +67,7 @@ static FAutoConsoleVariableRef CVarDebugWriteDecodedImage(
 #endif
 
 // increment this if you change anything that will affect compression in this file
-#define BASE_ASTC_FORMAT_VERSION 45
+#define BASE_ASTC_FORMAT_VERSION 46
 
 #define MAX_QUALITY_BY_SIZE 4
 #define MAX_QUALITY_BY_SPEED 3
diff --git a/Engine/Source/Runtime/CoreUObject/Public/UObject/NoExportTypes.h b/Engine/Source/Runtime/CoreUObject/Public/UObject/NoExportTypes.h
--- a/Engine/Source/Runtime/CoreUObject/Public/UObject/NoExportTypes.h
+++ b/Engine/Source/Runtime/CoreUObject/Public/UObject/NoExportTypes.h
@@ -282,6 +282,11 @@ enum EPixelFormat : int
 	PF_ETC2_RG11_EAC,
 	PF_R8,
 	PF_B5G5R5A1_UNORM,
+	PF_ASTC_4x4_HDR,	
+	PF_ASTC_6x6_HDR,	
+	PF_ASTC_8x8_HDR,	
+	PF_ASTC_10x10_HDR,	
+	PF_ASTC_12x12_HDR,
 	PF_G16R16_SNORM,
 	PF_R8G8_UINT,
 	PF_R32G32B32_UINT,
diff --git a/Engine/Source/Runtime/VulkanRHI/Private/Android/VulkanAndroidPlatform.cpp b/Engine/Source/Runtime/VulkanRHI/Private/Android/VulkanAndroidPlatform.cpp
--- a/Engine/Source/Runtime/VulkanRHI/Private/Android/VulkanAndroidPlatform.cpp
+++ b/Engine/Source/Runtime/VulkanRHI/Private/Android/VulkanAndroidPlatform.cpp
@@ -1481,7 +1481,7 @@ void FVulkanAndroidPlatform::SetupImageMemoryRequirementWorkaround(const FVulkan
 
 	// AFBC workarounds
 	{
-		const VkFormatFeatureFlags FormatFlags = InDevice.GetFormatProperties()[VK_FORMAT_B8G8R8A8_UNORM].optimalTilingFeatures;
+		const VkFormatFeatureFlags FormatFlags = InDevice.GetFormatProperties(VK_FORMAT_B8G8R8A8_UNORM).optimalTilingFeatures;
 
 		VkImage Image0;
 		VkMemoryRequirements Image0Mem;
diff --git a/Engine/Source/Runtime/VulkanRHI/Private/VulkanDevice.cpp b/Engine/Source/Runtime/VulkanRHI/Private/VulkanDevice.cpp
--- a/Engine/Source/Runtime/VulkanRHI/Private/VulkanDevice.cpp
+++ b/Engine/Source/Runtime/VulkanRHI/Private/VulkanDevice.cpp
@@ -837,7 +837,7 @@ void FVulkanDevice::SetupFormats()
 #endif  // UE_BUILD_DEBUG || UE_BUILD_DEVELOPMENT
 }
 
-const VkFormatProperties& FVulkanDevice::GetFormatProperties(VkFormat InFormat)
+const VkFormatProperties& FVulkanDevice::GetFormatProperties(VkFormat InFormat) const
 {
 	if (InFormat >= 0 && InFormat < VK_FORMAT_RANGE_SIZE)
 	{
diff --git a/Engine/Source/Runtime/VulkanRHI/Private/VulkanDevice.h b/Engine/Source/Runtime/VulkanRHI/Private/VulkanDevice.h
--- a/Engine/Source/Runtime/VulkanRHI/Private/VulkanDevice.h
+++ b/Engine/Source/Runtime/VulkanRHI/Private/VulkanDevice.h
@@ -339,10 +339,7 @@ public:
 		return DefaultTexture->DefaultView;
 	}
 
-	inline const VkFormatProperties* GetFormatProperties() const
-	{
-		return FormatProperties;
-	}
+	const VkFormatProperties& GetFormatProperties(VkFormat InFormat) const;
 
 	inline VulkanRHI::FDeviceMemoryManager& GetDeviceMemoryManager()
 	{
@@ -509,7 +506,6 @@ public:
 	FGPUTimingCalibrationTimestamp GetCalibrationTimestamp();
 
 private:
-	const VkFormatProperties& GetFormatProperties(VkFormat InFormat);
 	void MapBufferFormatSupport(FPixelFormatInfo& PixelFormatInfo, EPixelFormat UEFormat, VkFormat VulkanFormat);
 	void MapImageFormatSupport(FPixelFormatInfo& PixelFormatInfo, const TArrayView<const VkFormat>& PrioritizedFormats, EPixelFormatCapabilities RequiredCapabilities);
 	void MapFormatSupport(EPixelFormat UEFormat, std::initializer_list<VkFormat> PrioritizedFormats, const VkComponentMapping& ComponentMapping, EPixelFormatCapabilities RequiredCapabilities, int32 BlockBytes);
diff --git a/Engine/Source/Runtime/VulkanRHI/Private/VulkanTexture.cpp b/Engine/Source/Runtime/VulkanRHI/Private/VulkanTexture.cpp
--- a/Engine/Source/Runtime/VulkanRHI/Private/VulkanTexture.cpp
+++ b/Engine/Source/Runtime/VulkanRHI/Private/VulkanTexture.cpp
@@ -425,10 +425,11 @@ void FVulkanTexture::GenerateImageCreateInfo(
 	}
 
 	FVulkanPlatform::SetImageMemoryRequirementWorkaround(ImageCreateInfo);
-		
+
+	const VkFormatProperties& FormatProperties = InDevice.GetFormatProperties(ImageCreateInfo.format);
 	const VkFormatFeatureFlags FormatFlags = ImageCreateInfo.tiling == VK_IMAGE_TILING_LINEAR ? 
-		InDevice.GetFormatProperties()[ImageCreateInfo.format].linearTilingFeatures : 
-		InDevice.GetFormatProperties()[ImageCreateInfo.format].optimalTilingFeatures;
+		FormatProperties.linearTilingFeatures : 
+		FormatProperties.optimalTilingFeatures;
 
 	if (!VKHasAnyFlags(FormatFlags, VK_FORMAT_FEATURE_SAMPLED_IMAGE_BIT))
 	{
