Fix incorrect RayTracingScene BuildScratchSize.

diff --git a/Engine/Source/Runtime/D3D12RHI/Private/D3D12RayTracing.cpp b/Engine/Source/Runtime/D3D12RHI/Private/D3D12RayTracing.cpp
--- a/Engine/Source/Runtime/D3D12RHI/Private/D3D12RayTracing.cpp
+++ b/Engine/Source/Runtime/D3D12RHI/Private/D3D12RayTracing.cpp
@@ -3930,7 +3930,7 @@ FD3D12RayTracingScene::FD3D12RayTracingScene(FD3D12Adapter* Adapter, FRayTracing
 		Layer.ScratchBufferOffset = Align(SizeInfo.BuildScratchSize, GRHIRayTracingScratchBufferAlignment);
 
 		SizeInfo.ResultSize = Layer.BufferOffset + Layer.SizeInfo.ResultSize;
-		SizeInfo.BuildScratchSize = Layer.BufferOffset + Layer.SizeInfo.BuildScratchSize;
+		SizeInfo.BuildScratchSize = Layer.ScratchBufferOffset + Layer.SizeInfo.BuildScratchSize;
 	}
 };
 
diff --git a/Engine/Source/Runtime/VulkanRHI/Private/VulkanRayTracing.cpp b/Engine/Source/Runtime/VulkanRHI/Private/VulkanRayTracing.cpp
--- a/Engine/Source/Runtime/VulkanRHI/Private/VulkanRayTracing.cpp
+++ b/Engine/Source/Runtime/VulkanRHI/Private/VulkanRayTracing.cpp
@@ -532,7 +532,7 @@ FVulkanRayTracingScene::FVulkanRayTracingScene(FRayTracingSceneInitializer2 InIn
 		Layer.ScratchBufferOffset = Align(SizeInfo.BuildScratchSize, GRHIRayTracingScratchBufferAlignment);
 
 		SizeInfo.ResultSize = Layer.BufferOffset + LayerSizeInfo.ResultSize;
-		SizeInfo.BuildScratchSize = Layer.BufferOffset + LayerSizeInfo.BuildScratchSize;
+		SizeInfo.BuildScratchSize = Layer.ScratchBufferOffset + LayerSizeInfo.BuildScratchSize;
 	}
 
 	const uint32 ParameterBufferSize = FMath::Max<uint32>(1, Initializer.NumTotalSegments) * sizeof(FVulkanRayTracingGeometryParameters);
