VulkanRHI: Small clean up:
- Remove duplicate NextView declaration in FVulkanShaderResourceView (see VulkanRHI::FVulkanViewBase).
- Clearly label Invalidate() function as an override in FVulkanShaderResourceView and FVulkanUnorderedAccessView.

#rnx

diff --git a/Engine/Source/Runtime/VulkanRHI/Private/VulkanResources.h b/Engine/Source/Runtime/VulkanRHI/Private/VulkanResources.h
--- a/Engine/Source/Runtime/VulkanRHI/Private/VulkanResources.h
+++ b/Engine/Source/Runtime/VulkanRHI/Private/VulkanResources.h
@@ -1057,8 +1057,7 @@ public:
 
 	~FVulkanUnorderedAccessView();
 
-	void Invalidate();
-
+	void Invalidate() override final;
 	void UpdateView();
 
 protected:
@@ -1094,7 +1093,7 @@ public:
 
 	void Rename(FRHIResource* InRHIBuffer, FVulkanResourceMultiBuffer* InSourceBuffer, uint32 InSize, EPixelFormat InFormat);
 
-	void Invalidate();
+	void Invalidate() override final;
 	void UpdateView();
 
 	inline FVulkanBufferView* GetBufferView()
@@ -1133,9 +1132,6 @@ protected:
 	// Used to check on volatile buffers if a new BufferView is required
 	VkBuffer VolatileBufferHandle = VK_NULL_HANDLE;
 	uint32 VolatileLockCounter = MAX_uint32;
-
-	FVulkanShaderResourceView* NextView = 0;
-	friend class FVulkanTexture;
 };
 
 class FVulkanVertexInputStateInfo
