VulkanRHI:  Vulkan Descriptor Buffer WIP:
- Move declaration of FVulkanLayouts after the BindlessManager so they can use the BindlessPipelineLayout.
- Add stubs for other FVulkanBindlessDescriptorManager functions that will be needed for descriptor_buffers.
- Start replacing old decriptor_indexing code with the new stubs for descriptor_buffers.

#rnx

diff --git a/Engine/Source/Runtime/VulkanRHI/Private/VulkanCommandBuffer.cpp b/Engine/Source/Runtime/VulkanRHI/Private/VulkanCommandBuffer.cpp
--- a/Engine/Source/Runtime/VulkanRHI/Private/VulkanCommandBuffer.cpp
+++ b/Engine/Source/Runtime/VulkanRHI/Private/VulkanCommandBuffer.cpp
@@ -338,15 +338,7 @@ void FVulkanCmdBuffer::Begin()
 		FVulkanBindlessDescriptorManager* BindlessDescriptorManager = Device->GetBindlessDescriptorManager();
 		FVulkanQueue* Queue = GetOwner()->GetMgr().GetQueue();
 		const VkPipelineStageFlags SupportedStages = Queue->GetSupportedStageBits();
-		if (SupportedStages & VK_PIPELINE_STAGE_FRAGMENT_SHADER_BIT)
-		{
-			BindlessDescriptorManager->BindDescriptorSets(CommandBufferHandle, VK_PIPELINE_BIND_POINT_GRAPHICS);
-		}
-		if (SupportedStages & VK_PIPELINE_STAGE_COMPUTE_SHADER_BIT)
-		{
-			BindlessDescriptorManager->BindDescriptorSets(CommandBufferHandle, VK_PIPELINE_BIND_POINT_COMPUTE);
-		}
-		// todo-jn: bindless : raytracing
+		BindlessDescriptorManager->BindDescriptorBuffers(CommandBufferHandle, SupportedStages);
 	}
 
 	bNeedsDynamicStateSet = true;
diff --git a/Engine/Source/Runtime/VulkanRHI/Private/VulkanDescriptorSets.cpp b/Engine/Source/Runtime/VulkanRHI/Private/VulkanDescriptorSets.cpp
--- a/Engine/Source/Runtime/VulkanRHI/Private/VulkanDescriptorSets.cpp
+++ b/Engine/Source/Runtime/VulkanRHI/Private/VulkanDescriptorSets.cpp
@@ -8,187 +8,201 @@
 #include "VulkanDescriptorSets.h"
 
 
-int32 GBindlessMaxResourceDescriptorCount = 1000 * 1000;
-static FAutoConsoleVariableRef CVarBindlessResourceDescriptorCount(
-	TEXT("r.Vulkan.Bindless.MaxResourceDescriptorCount"),
-	GBindlessMaxResourceDescriptorCount,
-	TEXT("Maximum bindless resource descriptor count"),
+int32 GVulkanBindlessMaxSamplerDescriptorCount = 2048;
+static FAutoConsoleVariableRef CVarVulkanBindlessMaxSamplerDescriptorCount(
+	TEXT("r.Vulkan.Bindless.MaxSamplerDescriptorCount"),
+	GVulkanBindlessMaxSamplerDescriptorCount,
+	TEXT("Maximum bindless sampler descriptor count"),
 	ECVF_ReadOnly
 );
 
-int32 GBindlessMaxSamplerDescriptorCount = 2048;
-static FAutoConsoleVariableRef CVarBindlessSamplerDescriptorCount(
-	TEXT("r.Vulkan.Bindless.MaxSamplerDescriptorCount"),
-	GBindlessMaxSamplerDescriptorCount,
-	TEXT("Maximum bindless sampler descriptor count"),
+int32 GVulkanBindlessMaxSampledImageDescriptorCount = 256 * 1024;
+static FAutoConsoleVariableRef CVarVulkanBindlessMaxSampledImageCount(
+	TEXT("r.Vulkan.Bindless.MaxResourceSampledImageCount"),
+	GVulkanBindlessMaxSampledImageDescriptorCount,
+	TEXT("Maximum bindless Sampled Image descriptor count"),
+	ECVF_ReadOnly
+);
+
+int32 GVulkanBindlessMaxStorageImageDescriptorCount = 64 * 1024;
+static FAutoConsoleVariableRef CVarVulkanBindlessMaxStorageImageCount(
+	TEXT("r.Vulkan.Bindless.MaxResourceStorageImageCount"),
+	GVulkanBindlessMaxStorageImageDescriptorCount,
+	TEXT("Maximum bindless Storage Image descriptor count"),
+	ECVF_ReadOnly
+);
+
+int32 GVulkanBindlessMaxUniformTexelBufferDescriptorCount = 64 * 1024;
+static FAutoConsoleVariableRef CVarVulkanBindlessMaxUniformTexelBufferCount(
+	TEXT("r.Vulkan.Bindless.MaxResourceUniformTexelBufferCount"),
+	GVulkanBindlessMaxUniformTexelBufferDescriptorCount,
+	TEXT("Maximum bindless Uniform Texel Buffer descriptor count"),
+	ECVF_ReadOnly
+);
+
+int32 GVulkanBindlessMaxStorageTexelBufferDescriptorCount = 64 * 1024;
+static FAutoConsoleVariableRef CVarVulkanBindlessMaxStorageTexelBufferCount(
+	TEXT("r.Vulkan.Bindless.MaxResourceStorageTexelBufferCount"),
+	GVulkanBindlessMaxStorageTexelBufferDescriptorCount,
+	TEXT("Maximum bindless Storage Texel Buffer descriptor count"),
+	ECVF_ReadOnly
+);
+
+int32 GVulkanBindlessMaxUniformBufferDescriptorCount = 2 * 1024 * 1024;
+static FAutoConsoleVariableRef CVarVulkanBindlessMaxUniformBufferCount(
+	TEXT("r.Vulkan.Bindless.MaxResourceUniformBufferCount"),
+	GVulkanBindlessMaxUniformBufferDescriptorCount,
+	TEXT("Maximum bindless Uniform Buffer descriptor count"),
+	ECVF_ReadOnly
+);
+
+int32 GVulkanBindlessMaxStorageBufferDescriptorCount = 64 * 1024;
+static FAutoConsoleVariableRef CVarVulkanBindlessMaxStorageBufferCount(
+	TEXT("r.Vulkan.Bindless.MaxResourceStorageBufferCount"),
+	GVulkanBindlessMaxStorageBufferDescriptorCount,
+	TEXT("Maximum bindless Storage Buffer descriptor count"),
+	ECVF_ReadOnly
+);
+
+int32 GVulkanBindlessMaxAccelerationStructureDescriptorCount = 64 * 1024;
+static FAutoConsoleVariableRef CVarVulkanBindlessMaxAccelerationStructureCount(
+	TEXT("r.Vulkan.Bindless.MaxResourceAccelerationStructureCount"),
+	GVulkanBindlessMaxAccelerationStructureDescriptorCount,
+	TEXT("Maximum bindless Acceleration Structure descriptor count"),
+	ECVF_ReadOnly
+);
+
+
+int32 GVulkanBindlessMaxUniformBuffersPerStage = 32;
+static FAutoConsoleVariableRef CVarVulkanBindlessMaxUniformBuffersPerStage(
+	TEXT("r.Vulkan.Bindless.MaxUniformBuffersPerStage"),
+	GVulkanBindlessMaxUniformBuffersPerStage,
+	TEXT("Maximum Uniform Buffers per shader stage"),
 	ECVF_ReadOnly
 );
 
 
+static inline uint32 GetInitialDescriptorCount(VkDescriptorType DescriptorType)
+{
+	switch (DescriptorType)
+	{
+	case VK_DESCRIPTOR_TYPE_SAMPLER:                    return GVulkanBindlessMaxSamplerDescriptorCount;
+	case VK_DESCRIPTOR_TYPE_SAMPLED_IMAGE:              return GVulkanBindlessMaxSampledImageDescriptorCount;
+	case VK_DESCRIPTOR_TYPE_STORAGE_IMAGE:              return GVulkanBindlessMaxStorageImageDescriptorCount;
+	case VK_DESCRIPTOR_TYPE_UNIFORM_TEXEL_BUFFER:       return GVulkanBindlessMaxUniformTexelBufferDescriptorCount;
+	case VK_DESCRIPTOR_TYPE_STORAGE_TEXEL_BUFFER:       return GVulkanBindlessMaxStorageTexelBufferDescriptorCount;
+	case VK_DESCRIPTOR_TYPE_UNIFORM_BUFFER:             return GVulkanBindlessMaxUniformBufferDescriptorCount;
+	case VK_DESCRIPTOR_TYPE_STORAGE_BUFFER:             return GVulkanBindlessMaxStorageBufferDescriptorCount;
+	case VK_DESCRIPTOR_TYPE_ACCELERATION_STRUCTURE_KHR: return GVulkanBindlessMaxAccelerationStructureDescriptorCount;
+	default: checkNoEntry();
+	}
+	return 0;
+}
+
+
+static inline VkMemoryPropertyFlags GetDescriptorBufferMemoryType(FVulkanDevice* Device)
+{
+	if (Device->HasUnifiedMemory() || (FVulkanPlatform::SupportsDeviceLocalHostVisibleWithNoPenalty(Device->GetVendorId()) &&
+		Device->GetDeviceMemoryManager().SupportsMemoryType(VK_MEMORY_PROPERTY_HOST_VISIBLE_BIT | VK_MEMORY_PROPERTY_HOST_COHERENT_BIT | VK_MEMORY_PROPERTY_DEVICE_LOCAL_BIT)))
+	{
+		return VK_MEMORY_PROPERTY_DEVICE_LOCAL_BIT | VK_MEMORY_PROPERTY_HOST_VISIBLE_BIT | VK_MEMORY_PROPERTY_HOST_COHERENT_BIT;
+	}
+	else
+	{
+		return VK_MEMORY_PROPERTY_HOST_VISIBLE_BIT | VK_MEMORY_PROPERTY_HOST_COHERENT_BIT;
+	}
+}
+
+
+
+
+
 FVulkanBindlessDescriptorManager::FVulkanBindlessDescriptorManager(FVulkanDevice* InDevice)
 	: VulkanRHI::FDeviceChild(InDevice)
-	, bBindlessResourcesAllowed(RHIGetBindlessResourcesConfiguration(GMaxRHIShaderPlatform) != ERHIBindlessConfiguration::Disabled)
-	, bBindlessSamplersAllowed(RHIGetBindlessSamplersConfiguration(GMaxRHIShaderPlatform) != ERHIBindlessConfiguration::Disabled)
 {
-
+	const bool bFullyDisabled = (RHIGetBindlessResourcesConfiguration(GMaxRHIShaderPlatform) == ERHIBindlessConfiguration::Disabled) &&
+		(RHIGetBindlessSamplersConfiguration(GMaxRHIShaderPlatform) == ERHIBindlessConfiguration::Disabled);
+	const bool bFullyEnabled = (RHIGetBindlessResourcesConfiguration(GMaxRHIShaderPlatform) == ERHIBindlessConfiguration::AllShaders) &&
+		(RHIGetBindlessSamplersConfiguration(GMaxRHIShaderPlatform) == ERHIBindlessConfiguration::AllShaders);
+	checkf(bFullyDisabled || bFullyEnabled, TEXT("Bindless must be fully enabled or fully disabled in Vulkan because of the way it uses descriptor buffers."));
 }
 
 FVulkanBindlessDescriptorManager::~FVulkanBindlessDescriptorManager()
 {
-	VulkanRHI::vkDestroyPipelineLayout(Device->GetInstanceHandle(), BindlessPipelineLayout, VULKAN_CPU_ALLOCATOR);
+	check(BindlessPipelineLayout == VK_NULL_HANDLE);
+}
 
-	VulkanRHI::vkDestroyDescriptorSetLayout(Device->GetInstanceHandle(), EmptyDescriptorSetLayout, VULKAN_CPU_ALLOCATOR);
-	VulkanRHI::vkDestroyDescriptorSetLayout(Device->GetInstanceHandle(), SamplerDescriptorSetLayout, VULKAN_CPU_ALLOCATOR);
-	VulkanRHI::vkDestroyDescriptorSetLayout(Device->GetInstanceHandle(), ResourceDescriptorSetLayout, VULKAN_CPU_ALLOCATOR);
+void FVulkanBindlessDescriptorManager::Deinit()
+{
+}
 
-	VulkanRHI::vkDestroyDescriptorPool(Device->GetInstanceHandle(), DescriptorPool, VULKAN_CPU_ALLOCATOR);
+bool FVulkanBindlessDescriptorManager::VerifySupport(FVulkanDevice* Device)
+{
+	return false;
 }
 
 void FVulkanBindlessDescriptorManager::Init()
 {
-	MaxResourceDescriptors = GBindlessMaxResourceDescriptorCount;
-	MaxSamplerDescriptors = GBindlessMaxSamplerDescriptorCount;
-
-	// Create the empty descriptor set (used to pad empty indices)
-	{
-		VkDescriptorSetLayoutCreateInfo EmptyDescriptorSetLayoutCreateInfo;
-		ZeroVulkanStruct(EmptyDescriptorSetLayoutCreateInfo, VK_STRUCTURE_TYPE_DESCRIPTOR_SET_LAYOUT_CREATE_INFO);
-		VERIFYVULKANRESULT(VulkanRHI::vkCreateDescriptorSetLayout(Device->GetInstanceHandle(), &EmptyDescriptorSetLayoutCreateInfo, VULKAN_CPU_ALLOCATOR, &EmptyDescriptorSetLayout));
-	}
+}
 
-	// Create the descriptor pool for bindless
-	{
-		const uint32 DescriptorTypeCount = 1;
-		VkDescriptorPoolSize DescriptorPoolSize[DescriptorTypeCount];
-		DescriptorPoolSize[0].type = VK_DESCRIPTOR_TYPE_SAMPLER;
-		DescriptorPoolSize[0].descriptorCount = MaxSamplerDescriptors;
-		// todo-jn: bindless: add 'real' resource descriptors
-
-		VkDescriptorPoolCreateInfo DescriptorPoolCreateInfo;
-		ZeroVulkanStruct(DescriptorPoolCreateInfo, VK_STRUCTURE_TYPE_DESCRIPTOR_POOL_CREATE_INFO);
-		DescriptorPoolCreateInfo.maxSets = VulkanBindless::NumBindlessSets;
-		DescriptorPoolCreateInfo.poolSizeCount = DescriptorTypeCount;
-		DescriptorPoolCreateInfo.pPoolSizes = DescriptorPoolSize;
-		DescriptorPoolCreateInfo.flags = VK_DESCRIPTOR_POOL_CREATE_UPDATE_AFTER_BIND_BIT;
-
-		VERIFYVULKANRESULT(VulkanRHI::vkCreateDescriptorPool(Device->GetInstanceHandle(), &DescriptorPoolCreateInfo, VULKAN_CPU_ALLOCATOR, &DescriptorPool));
-	}
+void FVulkanBindlessDescriptorManager::BindDescriptorBuffers(VkCommandBuffer CommandBuffer, VkPipelineStageFlags SupportedStages)
+{
+}
 
-	// Create the sampler descriptor set layout
-	{
-		VkDescriptorSetLayoutBinding DescriptorSetLayoutBinding;
-		DescriptorSetLayoutBinding.binding = 0;
-		DescriptorSetLayoutBinding.descriptorType = VK_DESCRIPTOR_TYPE_SAMPLER;
-		DescriptorSetLayoutBinding.descriptorCount = MaxSamplerDescriptors;
-		DescriptorSetLayoutBinding.stageFlags = VK_SHADER_STAGE_ALL;
-		DescriptorSetLayoutBinding.pImmutableSamplers = nullptr;
-
-		const VkDescriptorBindingFlags DescriptorBindingFlags =
-			VK_DESCRIPTOR_BINDING_UPDATE_AFTER_BIND_BIT |
-			VK_DESCRIPTOR_BINDING_PARTIALLY_BOUND_BIT |
-			VK_DESCRIPTOR_BINDING_UPDATE_UNUSED_WHILE_PENDING_BIT |
-			VK_DESCRIPTOR_BINDING_VARIABLE_DESCRIPTOR_COUNT_BIT;
-		VkDescriptorSetLayoutBindingFlagsCreateInfo DescriptorSetLayoutBindingFlagsCreateInfo;
-		ZeroVulkanStruct(DescriptorSetLayoutBindingFlagsCreateInfo, VK_STRUCTURE_TYPE_DESCRIPTOR_SET_LAYOUT_BINDING_FLAGS_CREATE_INFO);
-		DescriptorSetLayoutBindingFlagsCreateInfo.bindingCount = 1;
-		DescriptorSetLayoutBindingFlagsCreateInfo.pBindingFlags = &DescriptorBindingFlags;
-
-		VkDescriptorSetLayoutCreateInfo DescriptorSetLayoutCreateInfo;
-		ZeroVulkanStruct(DescriptorSetLayoutCreateInfo, VK_STRUCTURE_TYPE_DESCRIPTOR_SET_LAYOUT_CREATE_INFO);
-		DescriptorSetLayoutCreateInfo.pBindings = &DescriptorSetLayoutBinding;
-		DescriptorSetLayoutCreateInfo.bindingCount = 1;
-		DescriptorSetLayoutCreateInfo.flags = VK_DESCRIPTOR_SET_LAYOUT_CREATE_UPDATE_AFTER_BIND_POOL_BIT;
-		DescriptorSetLayoutCreateInfo.pNext = &DescriptorSetLayoutBindingFlagsCreateInfo;
-
-		VERIFYVULKANRESULT(VulkanRHI::vkCreateDescriptorSetLayout(Device->GetInstanceHandle(), &DescriptorSetLayoutCreateInfo, VULKAN_CPU_ALLOCATOR, &SamplerDescriptorSetLayout));
-	}
+void FVulkanBindlessDescriptorManager::RegisterUniformBuffers(VkCommandBuffer CommandBuffer, VkPipelineBindPoint BindPoint, ShaderStage::EStage Stage, const FVulkanDescriptorSetWriter& SetWriter)
+{
+}
 
-	// Create the sampler descriptor set
-	{
-		VkDescriptorSetVariableDescriptorCountAllocateInfo VariableDescriptorCountAllocateInfo;
-		ZeroVulkanStruct(VariableDescriptorCountAllocateInfo, VK_STRUCTURE_TYPE_DESCRIPTOR_SET_VARIABLE_DESCRIPTOR_COUNT_ALLOCATE_INFO);
-		VariableDescriptorCountAllocateInfo.descriptorSetCount = 1;
-		VariableDescriptorCountAllocateInfo.pDescriptorCounts = &MaxSamplerDescriptors;
-
-		VkDescriptorSetAllocateInfo DescriptorSetAllocateInfo;
-		ZeroVulkanStruct(DescriptorSetAllocateInfo, VK_STRUCTURE_TYPE_DESCRIPTOR_SET_ALLOCATE_INFO);
-		DescriptorSetAllocateInfo.descriptorPool = DescriptorPool;
-		DescriptorSetAllocateInfo.descriptorSetCount = 1;
-		DescriptorSetAllocateInfo.pSetLayouts = &SamplerDescriptorSetLayout;
-		DescriptorSetAllocateInfo.pNext = &VariableDescriptorCountAllocateInfo;
-
-		VERIFYVULKANRESULT(VulkanRHI::vkAllocateDescriptorSets(Device->GetInstanceHandle(), &DescriptorSetAllocateInfo, &DescriptorSets[VulkanBindless::BindlessSamplerSet]));
-	}
 
-	// todo-jn: bindless: Create the resource descriptor set and layout
-	{
-		ResourceDescriptorSetLayout = EmptyDescriptorSetLayout;
-		DescriptorSets[VulkanBindless::BindlessResourceSet] = VK_NULL_HANDLE;
-	}
+FRHIDescriptorHandle FVulkanBindlessDescriptorManager::RegisterSampler(VkSampler VulkanSampler)
+{
+	return FRHIDescriptorHandle();
+}
 
-	// Now create the basic pipeline
-	{
-		VkDescriptorSetLayout DescriptorSetLayouts[VulkanBindless::NumBindlessSets];
-		DescriptorSetLayouts[VulkanBindless::BindlessSamplerSet] = SamplerDescriptorSetLayout;
-		DescriptorSetLayouts[VulkanBindless::BindlessResourceSet] = ResourceDescriptorSetLayout;
+FRHIDescriptorHandle FVulkanBindlessDescriptorManager::RegisterImage(VkImageView ImageView, VkDescriptorType DescriptorType, bool bIsDepthStencil)
+{
+	return FRHIDescriptorHandle();
+}
 
-		VkPipelineLayoutCreateInfo PipelineLayoutCreateInfo;
-		ZeroVulkanStruct(PipelineLayoutCreateInfo, VK_STRUCTURE_TYPE_PIPELINE_LAYOUT_CREATE_INFO);
-		PipelineLayoutCreateInfo.setLayoutCount = VulkanBindless::NumBindlessSets;
-		PipelineLayoutCreateInfo.pSetLayouts = DescriptorSetLayouts;
+FRHIDescriptorHandle FVulkanBindlessDescriptorManager::RegisterBuffer(VkBuffer VulkanBuffer, VkDeviceSize BufferOffset, VkDeviceSize BufferSize, VkDescriptorType DescriptorType)
+{
+	return FRHIDescriptorHandle();
+}
 
-		VERIFYVULKANRESULT(VulkanRHI::vkCreatePipelineLayout(Device->GetInstanceHandle(), &PipelineLayoutCreateInfo, VULKAN_CPU_ALLOCATOR, &BindlessPipelineLayout));
-	}
+FRHIDescriptorHandle FVulkanBindlessDescriptorManager::RegisterTexelBuffer(const VkBufferViewCreateInfo& ViewInfo, VkDescriptorType DescriptorType)
+{
+	return FRHIDescriptorHandle();
 }
 
-FVulkanBindlessDescriptorManager::BindlessLayoutArray FVulkanBindlessDescriptorManager::GeneratePipelineLayout(const TArray<VkDescriptorSetLayout>& LayoutArray) const
+FRHIDescriptorHandle FVulkanBindlessDescriptorManager::RegisterAccelerationStructure(VkAccelerationStructureKHR AccelerationStructure)
 {
-	BindlessLayoutArray PatchedArray;
-	PatchedArray.Add(SamplerDescriptorSetLayout);
-	PatchedArray.Add(ResourceDescriptorSetLayout);
-	PatchedArray.Append(LayoutArray);
-	return PatchedArray;
+	return FRHIDescriptorHandle();
 }
 
-void FVulkanBindlessDescriptorManager::BindDescriptorSets(VkCommandBuffer CommandBuffer, VkPipelineBindPoint BindPoint)
+
+uint32 FVulkanBindlessDescriptorManager::GetFreeResourceIndex(FVulkanBindlessDescriptorManager::BindlessSetState& State)
 {
-	uint32 FirstDescriptorSet = VulkanBindless::NumBindlessSets;
-	uint32 NumDescriptorSets = 0;
-	if (bBindlessSamplersAllowed)
-	{
-		FirstDescriptorSet = FMath::Min<uint32>(VulkanBindless::BindlessSamplerSet, FirstDescriptorSet);
-		++NumDescriptorSets;
-	}
-	if (bBindlessResourcesAllowed)
 	{
-		FirstDescriptorSet = FMath::Min<uint32>(VulkanBindless::BindlessResourceSet, FirstDescriptorSet);
-		++NumDescriptorSets;
+		FScopeLock ScopeLock(&State.FreeListCS);
+		if ((State.FreeListHead != MAX_uint32) && (State.PeakDescriptorCount >= State.MaxDescriptorCount)) // todo-jn: temp
+		{
+			const uint32 FreeIndex = State.FreeListHead;
+			const uint32 ByteOffset = State.FreeListHead * State.DescriptorSize;
+			uint32* NextSlot = (uint32*)(&State.DebugDescriptors[ByteOffset]);
+			State.FreeListHead = *NextSlot;
+			return FreeIndex;
+		}
 	}
 
-	if (NumDescriptorSets > 0)
-	{
-		VulkanRHI::vkCmdBindDescriptorSets(CommandBuffer, BindPoint, BindlessPipelineLayout, 0, NumDescriptorSets, &DescriptorSets[FirstDescriptorSet], 0, nullptr);
-	}
+	const uint32 ResourceIndex = State.PeakDescriptorCount++;
+	checkf(ResourceIndex < State.MaxDescriptorCount, TEXT("You need to grow the resource array size for [%s]!"), VK_TYPE_TO_STRING(VkDescriptorType, State.DescriptorType));
+	return ResourceIndex;
 }
 
-FRHIDescriptorHandle FVulkanBindlessDescriptorManager::RegisterSampler(VkSampler VulkanSampler)
+void FVulkanBindlessDescriptorManager::Unregister(FRHIDescriptorHandle DescriptorHandle)
 {
-	const uint32 SamplerIndex = BindlessSamplerCount++;
-	checkf(SamplerIndex < MaxSamplerDescriptors, TEXT("You need to grow the sampler array size!"));
-
-	VkDescriptorImageInfo DescriptorImageInfo;
-	FMemory::Memzero(DescriptorImageInfo);
-	DescriptorImageInfo.sampler = VulkanSampler;
-
-	VkWriteDescriptorSet WriteDescriptorSet;
-	ZeroVulkanStruct(WriteDescriptorSet, VK_STRUCTURE_TYPE_WRITE_DESCRIPTOR_SET);
-	WriteDescriptorSet.dstSet = DescriptorSets[VulkanBindless::BindlessSamplerSet];
-	WriteDescriptorSet.dstBinding = 0;
-	WriteDescriptorSet.dstArrayElement = SamplerIndex;
-	WriteDescriptorSet.descriptorCount = 1;
-	WriteDescriptorSet.descriptorType = VK_DESCRIPTOR_TYPE_SAMPLER;
-	WriteDescriptorSet.pImageInfo = &DescriptorImageInfo;
-
-	VulkanRHI::vkUpdateDescriptorSets(Device->GetInstanceHandle(), 1, &WriteDescriptorSet, 0, nullptr);
-	return FRHIDescriptorHandle(ERHIDescriptorHeapType::Sampler, SamplerIndex);
-}
\ No newline at end of file
+}
+
+void FVulkanBindlessDescriptorManager::CopyDescriptor(VkCommandBuffer CommandBuffer, FRHIDescriptorHandle DstHandle, const FRHIDescriptorHandle SrcHandle)
+{
+}
diff --git a/Engine/Source/Runtime/VulkanRHI/Private/VulkanDescriptorSets.h b/Engine/Source/Runtime/VulkanRHI/Private/VulkanDescriptorSets.h
--- a/Engine/Source/Runtime/VulkanRHI/Private/VulkanDescriptorSets.h
+++ b/Engine/Source/Runtime/VulkanRHI/Private/VulkanDescriptorSets.h
@@ -850,131 +850,7 @@
 };
 
 
-// Layout for a Pipeline, also includes DescriptorSets layout
-class FVulkanLayout : public VulkanRHI::FDeviceChild
-{
-public:
-	FVulkanLayout(FVulkanDevice* InDevice);
-	virtual ~FVulkanLayout();
 
-	virtual bool IsGfxLayout() const = 0;
 
-	inline const FVulkanDescriptorSetsLayout& GetDescriptorSetsLayout() const
-	{
-		return DescriptorSetLayout;
-	}
-
-	inline VkPipelineLayout GetPipelineLayout() const
-	{
-		return PipelineLayout;
-	}
-
-	inline bool HasDescriptors() const
-	{
-		return DescriptorSetLayout.GetLayouts().Num() > 0;
-	}
-
-	inline uint32 GetDescriptorSetLayoutHash() const
-	{
-		return DescriptorSetLayout.GetHash();
-	}
-
-	void PatchSpirvBindings(FVulkanShader::FSpirvCode& SpirvCode, EShaderFrequency Frequency, const FVulkanShaderHeader& CodeHeader, VkShaderStageFlagBits InStageFlag) const;
-
-protected:
-	FVulkanDescriptorSetsLayout	DescriptorSetLayout;
-	VkPipelineLayout			PipelineLayout;
-
-	template <bool bIsCompute>
-	inline void FinalizeBindings(const FUniformBufferGatherInfo& UBGatherInfo)
-	{
-		// Setting descriptor is only allowed prior to compiling the layout
-		check(DescriptorSetLayout.GetHandles().Num() == 0);
-
-		DescriptorSetLayout.FinalizeBindings<bIsCompute>(UBGatherInfo);
-	}
-
-	inline void ProcessBindingsForStage(VkShaderStageFlagBits StageFlags, ShaderStage::EStage DescSet, const FVulkanShaderHeader& CodeHeader, FUniformBufferGatherInfo& OutUBGatherInfo) const
-	{
-		// Setting descriptor is only allowed prior to compiling the layout
-		check(DescriptorSetLayout.GetHandles().Num() == 0);
-
-		DescriptorSetLayout.ProcessBindingsForStage(StageFlags, DescSet, CodeHeader, OutUBGatherInfo);
-	}
-
-	void Compile(FVulkanDescriptorSetLayoutMap& DSetLayoutMap);
-
-	friend class FVulkanComputePipeline;
-	friend class FVulkanGfxPipeline;
-	friend class FVulkanPipelineStateCacheManager;
-#if VULKAN_RHI_RAYTRACING
-	friend class FVulkanRayTracingPipelineState;
-#endif
-};
-
-class FVulkanGfxLayout : public FVulkanLayout
-{
-public:
-	FVulkanGfxLayout(FVulkanDevice* InDevice)
-		: FVulkanLayout(InDevice)
-	{
-	}
-
-	virtual bool IsGfxLayout() const final override
-	{
-		return true;
-	}
-
-	inline const FVulkanGfxPipelineDescriptorInfo& GetGfxPipelineDescriptorInfo() const
-	{
-		return GfxPipelineDescriptorInfo;
-	}
-
-	bool UsesInputAttachment(FVulkanShaderHeader::EAttachmentType AttachmentType) const;
-	
-protected:
-	FVulkanGfxPipelineDescriptorInfo		GfxPipelineDescriptorInfo;
-	friend class FVulkanPipelineStateCacheManager;
-};
-
-class FVulkanComputeLayout : public FVulkanLayout
-{
-public:
-	FVulkanComputeLayout(FVulkanDevice* InDevice)
-		: FVulkanLayout(InDevice)
-	{
-	}
-
-	virtual bool IsGfxLayout() const final override
-	{
-		return false;
-	}
-
-	inline const FVulkanComputePipelineDescriptorInfo& GetComputePipelineDescriptorInfo() const
-	{
-		return ComputePipelineDescriptorInfo;
-	}
-
-protected:
-	FVulkanComputePipelineDescriptorInfo		ComputePipelineDescriptorInfo;
-	friend class FVulkanPipelineStateCacheManager;
-};
-
-#if VULKAN_RHI_RAYTRACING
-class FVulkanRayTracingLayout : public FVulkanLayout
-{
-public:
-	FVulkanRayTracingLayout(FVulkanDevice* InDevice)
-		: FVulkanLayout(InDevice)
-	{
-	}
-
-	virtual bool IsGfxLayout() const final override
-	{
-		return false;
-	}
-};
-#endif // VULKAN_RHI_RAYTRACING
-
 // This class encapsulates updating VkWriteDescriptorSet structures (but doesn't own them), and their flags for dirty ranges; it is intended
 // to be used to access a sub-region of a long array of VkWriteDescriptorSet (ie FVulkanDescriptorSetWriteContainer)
@@ -1324,6 +1200,199 @@
 	void InitWrittenMasks(uint32 NumDescriptorWrites);
 };
 
+
+
+// Manager for resource descriptors used in bindless rendering.
+class FVulkanBindlessDescriptorManager : public VulkanRHI::FDeviceChild
+{
+public:
+	FVulkanBindlessDescriptorManager(FVulkanDevice* InDevice);
+	~FVulkanBindlessDescriptorManager();
+
+	void Init();
+	void Deinit();
+
+	static bool VerifySupport(FVulkanDevice* InDevice);
+
+	inline VkPipelineLayout GetPipelineLayout() const
+	{
+		return BindlessPipelineLayout;
+	}
+
+	void BindDescriptorBuffers(VkCommandBuffer CommandBuffer, VkPipelineStageFlags SupportedStages);
+
+	FRHIDescriptorHandle RegisterSampler(VkSampler VulkanSampler);
+	FRHIDescriptorHandle RegisterImage(VkImageView VulkanImage, VkDescriptorType DescriptorType, bool bIsDepthStencil);
+	FRHIDescriptorHandle RegisterBuffer(VkBuffer VulkanBuffer, VkDeviceSize BufferOffset, VkDeviceSize BufferSize, VkDescriptorType DescriptorType);
+	FRHIDescriptorHandle RegisterTexelBuffer(const VkBufferViewCreateInfo& ViewInfo, VkDescriptorType DescriptorType);
+	FRHIDescriptorHandle RegisterAccelerationStructure(VkAccelerationStructureKHR AccelerationStructure);
+	void RegisterUniformBuffers(VkCommandBuffer CommandBuffer, VkPipelineBindPoint BindPoint, ShaderStage::EStage Stage, const FVulkanDescriptorSetWriter& SetWriter);
+
+	void Unregister(FRHIDescriptorHandle DescriptorHandle);
+
+	void CopyDescriptor(VkCommandBuffer CommandBuffer, FRHIDescriptorHandle DstHandle, const FRHIDescriptorHandle SrcHandle);
+
+private:
+
+	VkDescriptorSetLayout EmptyDescriptorSetLayout = VK_NULL_HANDLE;
+
+	struct BindlessSetState
+	{
+		VkDescriptorType DescriptorType = VK_DESCRIPTOR_TYPE_MAX_ENUM;
+		uint32 MaxDescriptorCount = 0;
+		std::atomic<uint32> PeakDescriptorCount = 1;  // always keep a null descriptor in slot 0
+		VkDescriptorSetLayout DescriptorSetLayout = VK_NULL_HANDLE;
+
+		FCriticalSection FreeListCS;
+		uint32 FreeListHead = MAX_uint32;
+
+		uint32 DescriptorSize = 0;
+		VkBuffer BufferHandle = VK_NULL_HANDLE;
+		VkDeviceMemory MemoryHandle = VK_NULL_HANDLE;
+		uint8* MappedPointer = nullptr;
+
+		TArray<uint8> DebugDescriptors;
+	};
+	BindlessSetState BindlessSetStates[VulkanBindless::NumBindlessSets];
+
+	BindlessSetState BindlessUniformBufferSetState;
+	std::atomic<uint32> UniformBlockIndex = 0;
+
+	VkPipelineLayout BindlessPipelineLayout = VK_NULL_HANDLE;
+
+	uint32 GetFreeResourceIndex(BindlessSetState& Desc);
+};
+
+
+
+// Layout for a Pipeline, also includes DescriptorSets layout
+class FVulkanLayout : public VulkanRHI::FDeviceChild
+{
+public:
+	FVulkanLayout(FVulkanDevice* InDevice);
+	virtual ~FVulkanLayout();
+
+	virtual bool IsGfxLayout() const = 0;
+
+	inline const FVulkanDescriptorSetsLayout& GetDescriptorSetsLayout() const
+	{
+		return DescriptorSetLayout;
+	}
+
+	inline VkPipelineLayout GetPipelineLayout() const
+	{
+		return Device->SupportsBindless() ? Device->GetBindlessDescriptorManager()->GetPipelineLayout() : PipelineLayout;
+	}
+
+	inline bool HasDescriptors() const
+	{
+		return DescriptorSetLayout.GetLayouts().Num() > 0;
+	}
+
+	inline uint32 GetDescriptorSetLayoutHash() const
+	{
+		return DescriptorSetLayout.GetHash();
+	}
+
+	void PatchSpirvBindings(FVulkanShader::FSpirvCode& SpirvCode, EShaderFrequency Frequency, const FVulkanShaderHeader& CodeHeader, VkShaderStageFlagBits InStageFlag) const;
+
+protected:
+	FVulkanDescriptorSetsLayout	DescriptorSetLayout;
+	VkPipelineLayout			PipelineLayout;
+
+	template <bool bIsCompute>
+	inline void FinalizeBindings(const FUniformBufferGatherInfo& UBGatherInfo)
+	{
+		// Setting descriptor is only allowed prior to compiling the layout
+		check(DescriptorSetLayout.GetHandles().Num() == 0);
+
+		DescriptorSetLayout.FinalizeBindings<bIsCompute>(UBGatherInfo);
+	}
+
+	inline void ProcessBindingsForStage(VkShaderStageFlagBits StageFlags, ShaderStage::EStage DescSet, const FVulkanShaderHeader& CodeHeader, FUniformBufferGatherInfo& OutUBGatherInfo) const
+	{
+		// Setting descriptor is only allowed prior to compiling the layout
+		check(DescriptorSetLayout.GetHandles().Num() == 0);
+
+		DescriptorSetLayout.ProcessBindingsForStage(StageFlags, DescSet, CodeHeader, OutUBGatherInfo);
+	}
+
+	void Compile(FVulkanDescriptorSetLayoutMap& DSetLayoutMap);
+
+	friend class FVulkanComputePipeline;
+	friend class FVulkanGfxPipeline;
+	friend class FVulkanPipelineStateCacheManager;
+#if VULKAN_RHI_RAYTRACING
+	friend class FVulkanRayTracingPipelineState;
+#endif
+};
+
+class FVulkanGfxLayout : public FVulkanLayout
+{
+public:
+	FVulkanGfxLayout(FVulkanDevice* InDevice)
+		: FVulkanLayout(InDevice)
+	{
+	}
+
+	virtual bool IsGfxLayout() const final override
+	{
+		return true;
+	}
+
+	inline const FVulkanGfxPipelineDescriptorInfo& GetGfxPipelineDescriptorInfo() const
+	{
+		return GfxPipelineDescriptorInfo;
+	}
+
+	bool UsesInputAttachment(FVulkanShaderHeader::EAttachmentType AttachmentType) const;
+
+protected:
+	FVulkanGfxPipelineDescriptorInfo		GfxPipelineDescriptorInfo;
+	friend class FVulkanPipelineStateCacheManager;
+};
+
+class FVulkanComputeLayout : public FVulkanLayout
+{
+public:
+	FVulkanComputeLayout(FVulkanDevice* InDevice)
+		: FVulkanLayout(InDevice)
+	{
+	}
+
+	virtual bool IsGfxLayout() const final override
+	{
+		return false;
+	}
+
+	inline const FVulkanComputePipelineDescriptorInfo& GetComputePipelineDescriptorInfo() const
+	{
+		return ComputePipelineDescriptorInfo;
+	}
+
+protected:
+	FVulkanComputePipelineDescriptorInfo		ComputePipelineDescriptorInfo;
+	friend class FVulkanPipelineStateCacheManager;
+};
+
+#if VULKAN_RHI_RAYTRACING
+class FVulkanRayTracingLayout : public FVulkanLayout
+{
+public:
+	FVulkanRayTracingLayout(FVulkanDevice* InDevice)
+		: FVulkanLayout(InDevice)
+	{
+	}
+
+	virtual bool IsGfxLayout() const final override
+	{
+		return false;
+	}
+};
+#endif // VULKAN_RHI_RAYTRACING
+
+
+
 class FVulkanGenericDescriptorPool : FNoncopyable
 {
 public:
@@ -1414,40 +1483,6 @@
 	float PoolAllocRatio;
 };
 
-// Manager for resource descriptors used in bindless rendering.
-class FVulkanBindlessDescriptorManager : public VulkanRHI::FDeviceChild
-{
-public:
-	using BindlessLayoutArray = TArray<VkDescriptorSetLayout, TInlineAllocator<VulkanBindless::MaxNumSets>>;
 
-	FVulkanBindlessDescriptorManager(FVulkanDevice* InDevice);
-	~FVulkanBindlessDescriptorManager();
 
-	void Init();
 
-	BindlessLayoutArray GeneratePipelineLayout(const TArray<VkDescriptorSetLayout>& LayoutArray) const;
-
-	void BindDescriptorSets(VkCommandBuffer CommandBuffer, VkPipelineBindPoint BindPoint);
-
-	FRHIDescriptorHandle RegisterSampler(VkSampler VulkanSampler);
-
-private:
-	const bool bBindlessResourcesAllowed;
-	const bool bBindlessSamplersAllowed;
-
-	uint32 MaxResourceDescriptors = 0;
-	uint32 MaxSamplerDescriptors = 0;
-
-	std::atomic<uint32> BindlessSamplerCount = 0;
-
-	VkDescriptorSetLayout EmptyDescriptorSetLayout = VK_NULL_HANDLE;
-
-	VkDescriptorPool DescriptorPool = VK_NULL_HANDLE;
-
-	VkDescriptorSetLayout SamplerDescriptorSetLayout = VK_NULL_HANDLE;
-	VkDescriptorSetLayout ResourceDescriptorSetLayout = VK_NULL_HANDLE;
-
-	VkDescriptorSet  DescriptorSets[VulkanBindless::NumBindlessSets] = { VK_NULL_HANDLE };
-
-	VkPipelineLayout BindlessPipelineLayout = VK_NULL_HANDLE;
-};
diff --git a/Engine/Source/Runtime/VulkanRHI/Private/VulkanDevice.cpp b/Engine/Source/Runtime/VulkanRHI/Private/VulkanDevice.cpp
--- a/Engine/Source/Runtime/VulkanRHI/Private/VulkanDevice.cpp
+++ b/Engine/Source/Runtime/VulkanRHI/Private/VulkanDevice.cpp
@@ -1151,16 +1151,10 @@ void FVulkanDevice::InitGPU()
 	DescriptorPoolsManager = new FVulkanDescriptorPoolsManager();
 	DescriptorPoolsManager->Init(this);
 
-	if (OptionalDeviceExtensions.HasEXTDescriptorIndexing && GpuProps.limits.maxBoundDescriptorSets >= VulkanBindless::MaxNumSets)
+	if (FVulkanBindlessDescriptorManager::VerifySupport(this))
 	{
-		const bool bBindlessResources = RHIGetBindlessResourcesConfiguration(GMaxRHIShaderPlatform) != ERHIBindlessConfiguration::Disabled;
-		const bool bBindlessSamplers = RHIGetBindlessSamplersConfiguration(GMaxRHIShaderPlatform) != ERHIBindlessConfiguration::Disabled;
-
-		if (bBindlessResources || bBindlessSamplers)
-		{
-			BindlessDescriptorManager = new FVulkanBindlessDescriptorManager(this);
-			BindlessDescriptorManager->Init();
-		}
+		BindlessDescriptorManager = new FVulkanBindlessDescriptorManager(this);
+		BindlessDescriptorManager->Init();
 	}
 
 	PipelineStateCache = new FVulkanPipelineStateCacheManager(this);
@@ -1271,12 +1265,6 @@ void FVulkanDevice::Destroy()
 	delete DescriptorPoolsManager;
 	DescriptorPoolsManager = nullptr;
 
-	if (BindlessDescriptorManager)
-	{
-		delete BindlessDescriptorManager;
-		BindlessDescriptorManager = nullptr;
-	}
-
 	// No need to delete as it's stored in SamplerMap
 	DefaultSampler = nullptr;
 
@@ -1339,6 +1327,13 @@ void FVulkanDevice::Destroy()
 
 	DeferredDeletionQueue.Clear();
 
+	if (BindlessDescriptorManager)
+	{
+		BindlessDescriptorManager->Deinit();
+		delete BindlessDescriptorManager;
+		BindlessDescriptorManager = nullptr;
+	}
+
 	MemoryManager.Deinit();
 
 	delete TransferQueue;
diff --git a/Engine/Source/Runtime/VulkanRHI/Private/VulkanShaders.cpp b/Engine/Source/Runtime/VulkanRHI/Private/VulkanShaders.cpp
--- a/Engine/Source/Runtime/VulkanRHI/Private/VulkanShaders.cpp
+++ b/Engine/Source/Runtime/VulkanRHI/Private/VulkanShaders.cpp
@@ -500,20 +500,11 @@ void FVulkanLayout::Compile(FVulkanDescriptorSetLayoutMap& DSetLayoutMap)
 
 	DescriptorSetLayout.Compile(DSetLayoutMap);
 
-	VkPipelineLayoutCreateInfo PipelineLayoutCreateInfo;
-	ZeroVulkanStruct(PipelineLayoutCreateInfo, VK_STRUCTURE_TYPE_PIPELINE_LAYOUT_CREATE_INFO);
-
-	// Sneak in extra descriptor sets when in bindless
-	FVulkanBindlessDescriptorManager* BindlessDescriptorManager = Device->GetBindlessDescriptorManager();
-	if (BindlessDescriptorManager)
-	{
-		FVulkanBindlessDescriptorManager::BindlessLayoutArray LayoutHandles = BindlessDescriptorManager->GeneratePipelineLayout(DescriptorSetLayout.GetHandles());
-		PipelineLayoutCreateInfo.setLayoutCount = LayoutHandles.Num();
-		PipelineLayoutCreateInfo.pSetLayouts = LayoutHandles.GetData();
-		VERIFYVULKANRESULT(VulkanRHI::vkCreatePipelineLayout(Device->GetInstanceHandle(), &PipelineLayoutCreateInfo, VULKAN_CPU_ALLOCATOR, &PipelineLayout));
-	}
-	else
+	if (!Device->SupportsBindless())
 	{
+		VkPipelineLayoutCreateInfo PipelineLayoutCreateInfo;
+		ZeroVulkanStruct(PipelineLayoutCreateInfo, VK_STRUCTURE_TYPE_PIPELINE_LAYOUT_CREATE_INFO);
+
 		const TArray<VkDescriptorSetLayout>& LayoutHandles = DescriptorSetLayout.GetHandles();
 		PipelineLayoutCreateInfo.setLayoutCount = LayoutHandles.Num();
 		PipelineLayoutCreateInfo.pSetLayouts = LayoutHandles.GetData();
