VulkanRHI:  Add bindless handle support to the Vulkan deferred deletion queue.  Add VK_BUFFER_USAGE_SHADER_DEVICE_ADDRESS_BIT flag to pooled buffers and staging buffers for bindless.  (Vulkan Bindless WIP)

#rnx

diff --git a/Engine/Source/Runtime/VulkanRHI/Private/VulkanDevice.h b/Engine/Source/Runtime/VulkanRHI/Private/VulkanDevice.h
--- a/Engine/Source/Runtime/VulkanRHI/Private/VulkanDevice.h
+++ b/Engine/Source/Runtime/VulkanRHI/Private/VulkanDevice.h
@@ -51,6 +51,7 @@ struct FOptionalVulkanDeviceExtensions
 			uint64 HasRayQuery : 1;
 			uint64 HasDeferredHostOperations : 1;
 			uint64 HasEXTCalibratedTimestamps : 1;
+			uint64 HasEXTDescriptorBuffer : 1;
 
 			// Vendor specific
 			uint64 HasAMDBufferMarker : 1;
@@ -148,6 +149,7 @@ namespace VulkanRHI
 			DeviceMemoryAllocation,
 			BufferSuballocation,
 			AccelerationStructure,
+			BindlessHandle,
 		};
 
 		template <typename T>
@@ -157,6 +159,16 @@ namespace VulkanRHI
 			EnqueueGenericResource(Type, (uint64)Handle);
 		}
 
+		inline void EnqueueBindlessHandle(FRHIDescriptorHandle DescriptorHandle)
+		{
+			if (DescriptorHandle.IsValid())
+			{
+				const uint64 Type = (uint64)DescriptorHandle.GetRawType();
+				const uint64 Index = (uint64)DescriptorHandle.GetIndex();
+				const uint64 AsUInt64 = (Type << 32) | Index;
+				EnqueueResource(EType::BindlessHandle, AsUInt64);
+			}
+		}
 
 		void EnqueueResourceAllocation(FVulkanAllocation& Allocation);
 		void EnqueueDeviceAllocation(FDeviceMemoryAllocation* DeviceMemoryAllocation);
diff --git a/Engine/Source/Runtime/VulkanRHI/Private/VulkanMemory.cpp b/Engine/Source/Runtime/VulkanRHI/Private/VulkanMemory.cpp
--- a/Engine/Source/Runtime/VulkanRHI/Private/VulkanMemory.cpp
+++ b/Engine/Source/Runtime/VulkanRHI/Private/VulkanMemory.cpp
@@ -10,6 +10,7 @@
 #include "HAL/PlatformStackWalk.h"
 #include "VulkanContext.h"
 #include "VulkanLLM.h"
+#include "VulkanDescriptorSets.h"
 #include "Containers/SortedMap.h"
 
 #if PLATFORM_WINDOWS
@@ -2737,6 +2738,11 @@ namespace VulkanRHI
 		ZeroVulkanStruct(BufferCreateInfo, VK_STRUCTURE_TYPE_BUFFER_CREATE_INFO);
 		BufferCreateInfo.size = BufferSize;
 		BufferCreateInfo.usage = BufferUsageFlags;
+		// For descriptors buffers
+		if (Device->SupportsBindless())
+		{
+			BufferCreateInfo.usage |= VK_BUFFER_USAGE_SHADER_DEVICE_ADDRESS_BIT;
+		}
 		VERIFYVULKANRESULT(VulkanRHI::vkCreateBuffer(Device->GetInstanceHandle(), &BufferCreateInfo, VULKAN_CPU_ALLOCATOR, &Buffer));
 
 		VkMemoryRequirements MemReqs;
@@ -4220,6 +4226,11 @@ namespace VulkanRHI
 		ZeroVulkanStruct(StagingBufferCreateInfo, VK_STRUCTURE_TYPE_BUFFER_CREATE_INFO);
 		StagingBufferCreateInfo.size = Size;
 		StagingBufferCreateInfo.usage = InUsageFlags;
+		// For descriptors buffers
+		if (Device->SupportsBindless())
+		{
+			StagingBufferCreateInfo.usage |= VK_BUFFER_USAGE_SHADER_DEVICE_ADDRESS_BIT;
+		}
 
 		VkDevice VulkanDevice = Device->GetInstanceHandle();
 
@@ -4628,7 +4639,7 @@ namespace VulkanRHI
 #if VULKAN_HAS_DEBUGGING_ENABLED
 			FEntry* ExistingEntry = Entries.FindByPredicate([&](const FEntry& InEntry)
 				{
-					return InEntry.Handle == Entry.Handle;
+					return (InEntry.Handle == Entry.Handle) && (InEntry.StructureType == Entry.StructureType);
 				});
 			checkf(ExistingEntry == nullptr, TEXT("Attempt to double-delete resource, FDeferredDeletionQueue2::EType: %d, Handle: %llu"), (int32)Type, Handle);
 #endif
@@ -4738,6 +4749,16 @@ namespace VulkanRHI
 					break;
 				}
 #endif // VULKAN_RHI_RAYTRACING
+				case EType::BindlessHandle:
+				{
+					check(Device->SupportsBindless());
+					uint32 BindlessType = (uint32)(Entry->Handle >> 32);
+					check(BindlessType < MAX_uint8);
+					uint32 BindlessIndex = (uint32)(Entry->Handle & 0xFFFFFFFF);
+					FRHIDescriptorHandle DescriptorHandle((uint8)BindlessType, BindlessIndex);
+					Device->GetBindlessDescriptorManager()->Unregister(DescriptorHandle);
+					break;
+				}
 
 				default:
 					check(0);
