VulkanRHI: Always have an instance of BindlessDescriptorManager to match more closely with other RHI.  Also removes the need for SupportsBindless() checks everywhere (it will be verified internally).

#rnx

diff --git a/Engine/Source/Runtime/VulkanRHI/Private/VulkanDescriptorSets.cpp b/Engine/Source/Runtime/VulkanRHI/Private/VulkanDescriptorSets.cpp
--- a/Engine/Source/Runtime/VulkanRHI/Private/VulkanDescriptorSets.cpp
+++ b/Engine/Source/Runtime/VulkanRHI/Private/VulkanDescriptorSets.cpp
@@ -119,6 +119,7 @@ static inline VkMemoryPropertyFlags GetDescriptorBufferMemoryType(FVulkanDevice*
 
 FVulkanBindlessDescriptorManager::FVulkanBindlessDescriptorManager(FVulkanDevice* InDevice)
 	: VulkanRHI::FDeviceChild(InDevice)
+	, bIsSupported(VerifySupport(InDevice))
 {
 	const bool bFullyDisabled = (RHIGetBindlessResourcesConfiguration(GMaxRHIShaderPlatform) == ERHIBindlessConfiguration::Disabled) &&
 		(RHIGetBindlessSamplersConfiguration(GMaxRHIShaderPlatform) == ERHIBindlessConfiguration::Disabled);
diff --git a/Engine/Source/Runtime/VulkanRHI/Private/VulkanDescriptorSets.h b/Engine/Source/Runtime/VulkanRHI/Private/VulkanDescriptorSets.h
--- a/Engine/Source/Runtime/VulkanRHI/Private/VulkanDescriptorSets.h
+++ b/Engine/Source/Runtime/VulkanRHI/Private/VulkanDescriptorSets.h
@@ -1213,6 +1213,10 @@ public:
 	void Deinit();
 
 	static bool VerifySupport(FVulkanDevice* InDevice);
+	inline bool IsSupported()
+	{
+		return bIsSupported;
+	}
 
 	inline VkPipelineLayout GetPipelineLayout() const
 	{
@@ -1233,6 +1237,7 @@ public:
 	void CopyDescriptor(VkCommandBuffer CommandBuffer, FRHIDescriptorHandle DstHandle, const FRHIDescriptorHandle SrcHandle);
 
 private:
+	const bool bIsSupported;
 
 	VkDescriptorSetLayout EmptyDescriptorSetLayout = VK_NULL_HANDLE;
 
diff --git a/Engine/Source/Runtime/VulkanRHI/Private/VulkanDevice.cpp b/Engine/Source/Runtime/VulkanRHI/Private/VulkanDevice.cpp
--- a/Engine/Source/Runtime/VulkanRHI/Private/VulkanDevice.cpp
+++ b/Engine/Source/Runtime/VulkanRHI/Private/VulkanDevice.cpp
@@ -1018,6 +1018,12 @@ void FVulkanDevice::MapFormatSupport(EPixelFormat UEFormat, std::initializer_lis
 	MapFormatSupport(UEFormat, PrioritizedFormats, ComponentMapping, RequiredCapabilities, kDefaultBlockBytes);
 }
 
+bool FVulkanDevice::SupportsBindless() const
+{
+	checkSlow(BindlessDescriptorManager != nullptr);
+	return BindlessDescriptorManager->IsSupported();
+}
+
 void FVulkanDevice::InitGPU()
 {
 	LLM_SCOPE_VULKAN(ELLMTagVulkan::VulkanMisc);
@@ -1151,11 +1157,8 @@ void FVulkanDevice::InitGPU()
 	DescriptorPoolsManager = new FVulkanDescriptorPoolsManager();
 	DescriptorPoolsManager->Init(this);
 
-	if (FVulkanBindlessDescriptorManager::VerifySupport(this))
-	{
-		BindlessDescriptorManager = new FVulkanBindlessDescriptorManager(this);
-		BindlessDescriptorManager->Init();
-	}
+	BindlessDescriptorManager = new FVulkanBindlessDescriptorManager(this);
+	BindlessDescriptorManager->Init();
 
 	PipelineStateCache = new FVulkanPipelineStateCacheManager(this);
 
@@ -1327,12 +1330,9 @@ void FVulkanDevice::Destroy()
 
 	DeferredDeletionQueue.Clear();
 
-	if (BindlessDescriptorManager)
-	{
-		BindlessDescriptorManager->Deinit();
-		delete BindlessDescriptorManager;
-		BindlessDescriptorManager = nullptr;
-	}
+	BindlessDescriptorManager->Deinit();
+	delete BindlessDescriptorManager;
+	BindlessDescriptorManager = nullptr;
 
 	MemoryManager.Deinit();
 
diff --git a/Engine/Source/Runtime/VulkanRHI/Private/VulkanDevice.h b/Engine/Source/Runtime/VulkanRHI/Private/VulkanDevice.h
--- a/Engine/Source/Runtime/VulkanRHI/Private/VulkanDevice.h
+++ b/Engine/Source/Runtime/VulkanRHI/Private/VulkanDevice.h
@@ -324,10 +324,7 @@ public:
 		return DeviceMemoryManager.HasUnifiedMemory();
 	}
 
-	inline bool SupportsBindless() const
-	{
-		return (BindlessDescriptorManager != nullptr);
-	}
+	bool SupportsBindless() const;
 
 	inline uint64 GetTimestampValidBitsMask() const
 	{
