VulkanRHI: Add bindless handles to SRV, UAV and textures.  (Vulkan Bindless WIP)

#rnx

diff --git a/Engine/Source/Runtime/VulkanRHI/Private/VulkanMemory.h b/Engine/Source/Runtime/VulkanRHI/Private/VulkanMemory.h
--- a/Engine/Source/Runtime/VulkanRHI/Private/VulkanMemory.h
+++ b/Engine/Source/Runtime/VulkanRHI/Private/VulkanMemory.h
@@ -232,13 +232,19 @@ namespace VulkanRHI
 		{
 			check(NextView == 0);
 		}
+
 		FVulkanViewBase(FVulkanDevice* Device = nullptr)
 		: FDeviceChild(Device)
 		, NextView(0)
 		{
 		}
-		virtual void Invalidate(){}
+
+		virtual void Invalidate() {}
 		FVulkanViewBase* NextView;
+
+	protected:
+		void FreeBindlessHandle();
+		FRHIDescriptorHandle BindlessHandle;
 	};
 	// An Allocation off a Device Heap. Lowest level of allocations and bounded by VkPhysicalDeviceLimits::maxMemoryAllocationCount.
 	class FDeviceMemoryAllocation
diff --git a/Engine/Source/Runtime/VulkanRHI/Private/VulkanResources.h b/Engine/Source/Runtime/VulkanRHI/Private/VulkanResources.h
--- a/Engine/Source/Runtime/VulkanRHI/Private/VulkanResources.h
+++ b/Engine/Source/Runtime/VulkanRHI/Private/VulkanResources.h
@@ -413,13 +413,18 @@ public:
 	void* GetNativeResource() const override final { return (void*)Image; }
 	void* GetTextureBaseRHI() override final { return this; }
 
+	virtual FRHIDescriptorHandle GetDefaultBindlessHandle() const override final
+	{
+		checkSlow(DefaultBindlessHandle.IsValid());
+		return DefaultBindlessHandle;
+	}
+
 	inline static FVulkanTexture* Cast(FRHITexture* Texture)
 	{
 		check(Texture);
 		return (FVulkanTexture*)Texture->GetTextureBaseRHI();
 	}
 
-public:
 	struct FImageCreateInfo
 	{
 		VkImageCreateInfo ImageCreateInfo;
@@ -556,6 +561,7 @@ private:
 	VkImageAspectFlags PartialAspectMask;
 	FVulkanCpuReadbackBuffer* CpuReadbackBuffer;
 	VkImageLayout DefaultLayout;
+	FRHIDescriptorHandle DefaultBindlessHandle;
 
 	friend struct FRHICommandSetInitialImageState;
 
@@ -1060,6 +1066,8 @@ public:
 	void Invalidate() override final;
 	void UpdateView();
 
+	virtual FRHIDescriptorHandle GetBindlessHandle() const override final;
+
 protected:
 	// The texture that this UAV come from
 	TRefCountPtr<FRHITexture> SourceTexture;
@@ -1096,6 +1104,8 @@ public:
 	void Invalidate() override final;
 	void UpdateView();
 
+	virtual FRHIDescriptorHandle GetBindlessHandle() const override final;
+
 	inline FVulkanBufferView* GetBufferView()
 	{
 		return BufferViews[BufferIndex];
diff --git a/Engine/Source/Runtime/VulkanRHI/Private/VulkanTexture.cpp b/Engine/Source/Runtime/VulkanRHI/Private/VulkanTexture.cpp
--- a/Engine/Source/Runtime/VulkanRHI/Private/VulkanTexture.cpp
+++ b/Engine/Source/Runtime/VulkanRHI/Private/VulkanTexture.cpp
@@ -1705,6 +1705,8 @@ FVulkanTexture::FVulkanTexture(FVulkanDevice& InDevice, const FRHITextureCreateD
 		PartialView->Create(InDevice, Image, ViewType, PartialAspectMask, InCreateDesc.Format, ViewFormat, 0, FMath::Max(InCreateDesc.NumMips, (uint8)1u), 0, bArray ? FMath::Max((uint16)1u, InCreateDesc.ArraySize) : FMath::Max((uint16)1u, InCreateDesc.Depth), false);
 	}
 
+	DefaultBindlessHandle = Device->GetBindlessDescriptorManager()->RegisterImage(PartialView->View, SupportsSampling() ? VK_DESCRIPTOR_TYPE_SAMPLED_IMAGE : VK_DESCRIPTOR_TYPE_STORAGE_IMAGE, IsDepthOrStencilAspect());
+
 	if (!InCreateDesc.BulkData)
 	{
 		return;
@@ -1870,6 +1872,11 @@ FVulkanTexture::FVulkanTexture(FVulkanDevice& InDevice, const FRHITextureCreateD
 		PartialView = new FVulkanTextureView;
 		PartialView->Create(InDevice, Image, ViewType, PartialAspectMask, InCreateDesc.Format, ViewFormat, 0, FMath::Max(InCreateDesc.NumMips, (uint8)1u), 0, bArray ? FMath::Max((uint16)1u, InCreateDesc.ArraySize) : FMath::Max((uint16)1u, InCreateDesc.Depth), false);
 	}
+
+	if (PartialView->View)
+	{
+		DefaultBindlessHandle = Device->GetBindlessDescriptorManager()->RegisterImage(PartialView->View, SupportsSampling() ? VK_DESCRIPTOR_TYPE_SAMPLED_IMAGE : VK_DESCRIPTOR_TYPE_STORAGE_IMAGE, IsDepthOrStencilAspect());
+	}
 }
 
 FVulkanTexture::FVulkanTexture(FVulkanDevice& InDevice, const FRHITextureCreateDesc& InCreateDesc, FTextureRHIRef& SrcTextureRHI)
@@ -1951,6 +1958,7 @@ void FVulkanTexture::AliasTextureResources(FTextureRHIRef& SrcTextureRHI)
 
 	AliasedTexture = SrcTexture;
 	DefaultLayout = SrcTexture->DefaultLayout;
+	DefaultBindlessHandle = SrcTexture->DefaultBindlessHandle;
 }
 
 void FVulkanTexture::DestroyViews()
@@ -1963,6 +1971,12 @@ void FVulkanTexture::DestroyViews()
 		{
 			PartialView->Destroy(*Device);
 		}
+
+		if (DefaultBindlessHandle.IsValid())
+		{
+			Device->GetDeferredDeletionQueue().EnqueueBindlessHandle(DefaultBindlessHandle);
+			DefaultBindlessHandle = FRHIDescriptorHandle();
+		}
 	}
 }
 
@@ -1975,18 +1989,25 @@ void FVulkanTexture::InvalidateViews(FVulkanDevice& InDevice)
 	const bool bArray = ViewType == VK_IMAGE_VIEW_TYPE_1D_ARRAY || ViewType == VK_IMAGE_VIEW_TYPE_2D_ARRAY || ViewType == VK_IMAGE_VIEW_TYPE_CUBE_ARRAY;
 	const uint32 SizeZOrArraySize = bArray ? FMath::Max((uint16)1u, Desc.ArraySize) : FMath::Max((uint16)1u, Desc.Depth);
 
-	if(ViewType != VK_IMAGE_VIEW_TYPE_MAX_ENUM)
+	if (ViewType != VK_IMAGE_VIEW_TYPE_MAX_ENUM)
 	{
 		DefaultView.Create(InDevice, Image, ViewType, GetFullAspectMask(), GetDesc().Format, ViewFormat, 0, FMath::Max(NumMips, 1u), 0, SizeZOrArraySize, !SupportsSampling());
 	}
-	if(PartialView != &DefaultView)
+	if (PartialView != &DefaultView)
 	{
 		PartialView->Destroy(*Device);
 		PartialView->Create(InDevice, Image, ViewType, PartialAspectMask, GetDesc().Format, ViewFormat, 0, FMath::Max(NumMips, 1u), 0, SizeZOrArraySize, false);
 	}
 
+	if (DefaultBindlessHandle.IsValid())
+	{
+		Device->GetDeferredDeletionQueue().EnqueueBindlessHandle(DefaultBindlessHandle);
+	}
+
+	DefaultBindlessHandle = Device->GetBindlessDescriptorManager()->RegisterImage(PartialView->View, SupportsSampling() ? VK_DESCRIPTOR_TYPE_SAMPLED_IMAGE : VK_DESCRIPTOR_TYPE_STORAGE_IMAGE, IsDepthOrStencilAspect());
+
 	VulkanRHI::FVulkanViewBase* View = FirstView;
-	while(View)
+	while (View)
 	{
 		View->Invalidate();
 		View = View->NextView;
diff --git a/Engine/Source/Runtime/VulkanRHI/Private/VulkanUAV.cpp b/Engine/Source/Runtime/VulkanRHI/Private/VulkanUAV.cpp
--- a/Engine/Source/Runtime/VulkanRHI/Private/VulkanUAV.cpp
+++ b/Engine/Source/Runtime/VulkanRHI/Private/VulkanUAV.cpp
@@ -2,12 +2,24 @@
 
 #include "VulkanRHIPrivate.h"
 #include "VulkanContext.h"
+#include "VulkanDescriptorSets.h"
 #include "ClearReplacementShaders.h"
 
 #if VULKAN_RHI_RAYTRACING
 #include "VulkanRayTracing.h"
 #endif // VULKAN_RHI_RAYTRACING
 
+
+void VulkanRHI::FVulkanViewBase::FreeBindlessHandle()
+{
+	if (BindlessHandle.IsValid())
+	{
+		Device->GetDeferredDeletionQueue().EnqueueBindlessHandle(BindlessHandle);
+		BindlessHandle = FRHIDescriptorHandle();
+	}
+}
+
+
 FVulkanShaderResourceView::FVulkanShaderResourceView(FVulkanDevice* Device, FRHIViewableResource* InRHIBuffer, FVulkanResourceMultiBuffer* InSourceBuffer, uint32 InSize, EPixelFormat InFormat, uint32 InOffset)
 	: FRHIShaderResourceView(InRHIBuffer)
 	, VulkanRHI::FVulkanViewBase(Device)
@@ -45,7 +57,7 @@ FVulkanShaderResourceView::FVulkanShaderResourceView(FVulkanDevice* Device, FRHI
 {
 	FVulkanTexture* VulkanTexture = FVulkanTexture::Cast(InSourceTexture);
 	VulkanTexture->AttachView(this);
-
+	UpdateView();
 }
 
 FVulkanShaderResourceView::FVulkanShaderResourceView(FVulkanDevice* InDevice, FVulkanResourceMultiBuffer* InSourceBuffer, uint32 InOffset)
@@ -68,6 +80,8 @@ FVulkanShaderResourceView::FVulkanShaderResourceView(FVulkanDevice* InDevice, FV
 
 		VkDevice NativeDevice = InDevice->GetInstanceHandle();
 		VERIFYVULKANRESULT(VulkanDynamicAPI::vkCreateAccelerationStructureKHR(NativeDevice, &CreateInfo, VULKAN_CPU_ALLOCATOR, &AccelerationStructureHandle));
+
+		BindlessHandle = Device->GetBindlessDescriptorManager()->RegisterAccelerationStructure(AccelerationStructureHandle);
 	}
 	else
 #endif
@@ -75,6 +89,9 @@ FVulkanShaderResourceView::FVulkanShaderResourceView(FVulkanDevice* InDevice, FV
 		SourceStructuredBuffer = InSourceBuffer;
 		Size = InSourceBuffer->GetSize() - InOffset;
 		Offset = InOffset;
+
+		const uint32 FullMemoryOffset = InSourceBuffer->GetOffset() + InOffset;
+		BindlessHandle = Device->GetBindlessDescriptorManager()->RegisterBuffer(InSourceBuffer->GetHandle(), FullMemoryOffset, Size, VK_DESCRIPTOR_TYPE_STORAGE_BUFFER);
 	}
 }
 
@@ -92,19 +109,22 @@ FVulkanShaderResourceView::~FVulkanShaderResourceView()
 
 void FVulkanShaderResourceView::Clear()
 {
-#if VULKAN_RHI_RAYTRACING
-	if (Device && AccelerationStructureHandle)
-	{
-		Device->GetDeferredDeletionQueue().EnqueueResource(VulkanRHI::FDeferredDeletionQueue2::EType::AccelerationStructure, AccelerationStructureHandle);
-	}
-#endif // VULKAN_RHI_RAYTRACING
-
 	SourceRHIBuffer = nullptr;
 	SourceBuffer = nullptr;
 	BufferViews.Empty();
 	SourceStructuredBuffer = nullptr;
 	if (Device)
 	{
+#if VULKAN_RHI_RAYTRACING
+		if (AccelerationStructureHandle)
+		{
+			Device->GetDeferredDeletionQueue().EnqueueResource(VulkanRHI::FDeferredDeletionQueue2::EType::AccelerationStructure, AccelerationStructureHandle);
+			AccelerationStructureHandle = VK_NULL_HANDLE;
+		}
+#endif // VULKAN_RHI_RAYTRACING
+
+		FreeBindlessHandle();
+
 		TextureView.Destroy(*Device);
 	}
 	SourceTexture = nullptr;
@@ -136,9 +156,14 @@ void FVulkanShaderResourceView::Rename(FRHIResource* InRHIBuffer, FVulkanResourc
 	SourceRHIBuffer = InRHIBuffer;
 	VolatileBufferHandle = VK_NULL_HANDLE;
 	VolatileLockCounter = MAX_uint32;
+
+	// todo-jn: make the new bindless handle live in the old one's slot?
+	FreeBindlessHandle();
 }
+
 void FVulkanShaderResourceView::Invalidate()
 {
+	FreeBindlessHandle();
 	TextureView.Destroy(*Device);
 }
 
@@ -168,25 +193,42 @@ void FVulkanShaderResourceView::UpdateView()
 				|| VolatileBufferHandle != SourceVolatileBufferHandle)
 			{
 				BufferViews[0] = nullptr;
+				FreeBindlessHandle();
 			}
 
 			VolatileLockCounter = SourceBuffer->GetVolatileLockCounter();
 			VolatileBufferHandle = SourceVolatileBufferHandle;
 		}
-		else if (SourceBuffer->IsDynamic())
+		else if (SourceBuffer->IsDynamic() && (BufferIndex != SourceBuffer->GetDynamicIndex()))
 		{
 			BufferIndex = SourceBuffer->GetDynamicIndex();
+			FreeBindlessHandle();
 		}
 
 		if (!BufferViews[BufferIndex])
 		{
 			BufferViews[BufferIndex] = new FVulkanBufferView(Device);
 			BufferViews[BufferIndex]->Create(SourceBuffer, BufferViewFormat, SourceBuffer->GetOffset() + Offset, CurrentViewSize);
+			// :todo-jn: the buffer view is actually not needed in bindless anymore
+		}
+
+		if (!BindlessHandle.IsValid())
+		{
+			VkBufferViewCreateInfo ViewInfo;
+			ViewInfo.buffer = SourceBuffer->GetHandle();
+			ViewInfo.offset = SourceBuffer->GetOffset() + Offset;
+			ViewInfo.range = CurrentViewSize;;
+			ViewInfo.format = UEToVkTextureFormat(BufferViewFormat, false);
+			BindlessHandle = Device->GetBindlessDescriptorManager()->RegisterTexelBuffer(ViewInfo, VK_DESCRIPTOR_TYPE_UNIFORM_TEXEL_BUFFER);
 		}
 	}
 	else if (SourceStructuredBuffer)
 	{
-		// Nothing...
+		if (!BindlessHandle.IsValid())
+		{
+			const uint32 FullMemoryOffset = SourceStructuredBuffer->GetOffset() + Offset;
+			BindlessHandle = Device->GetBindlessDescriptorManager()->RegisterBuffer(SourceStructuredBuffer->GetHandle(), FullMemoryOffset, Size, VK_DESCRIPTOR_TYPE_STORAGE_BUFFER);
+		}
 	}
 #if VULKAN_RHI_RAYTRACING
 	else if (AccelerationStructureHandle)
@@ -237,9 +279,19 @@ void FVulkanShaderResourceView::UpdateView()
 			}
 
 			TextureView.Create(*Device, SourceTextureVK->Image, ActualViewType, SourceTextureVK->GetPartialAspectMask(), Format, UEToVkTextureFormat(Format, bSRGB), MipLevel, NumMips, FirstArraySlice, ActualNumArraySlices, false);
+
+			checkSlow(!BindlessHandle.IsValid());
+			BindlessHandle = Device->GetBindlessDescriptorManager()->RegisterImage(TextureView.View, VK_DESCRIPTOR_TYPE_SAMPLED_IMAGE, SourceTextureVK->IsDepthOrStencilAspect());
 		}
 	}
 }
+
+FRHIDescriptorHandle FVulkanShaderResourceView::GetBindlessHandle() const
+{
+	checkSlow(BindlessHandle.IsValid());
+	return BindlessHandle;
+}
+
 FVulkanUnorderedAccessView::FVulkanUnorderedAccessView(FVulkanDevice* Device, FVulkanResourceMultiBuffer* Buffer, bool bUseUAVCounter, bool bAppendBuffer)
 	: FRHIUnorderedAccessView(Buffer)
 	, VulkanRHI::FVulkanViewBase(Device)
@@ -250,6 +302,8 @@ FVulkanUnorderedAccessView::FVulkanUnorderedAccessView(FVulkanDevice* Device, FV
 	, BufferViewFormat(PF_Unknown)
 	, VolatileLockCounter(MAX_uint32)
 {
+	checkSlow(!bUseUAVCounter && !bAppendBuffer);
+	BindlessHandle = Device->GetBindlessDescriptorManager()->RegisterBuffer(SourceBuffer->GetHandle(), SourceBuffer->GetOffset(), SourceBuffer->GetSize(), VK_DESCRIPTOR_TYPE_STORAGE_BUFFER);
 }
 
 FVulkanUnorderedAccessView::FVulkanUnorderedAccessView(FVulkanDevice* Device, FRHITexture* TextureRHI, uint32 MipLevel, uint16 InFirstArraySlice, uint16 InNumArraySlices)
@@ -282,6 +336,7 @@ FVulkanUnorderedAccessView::FVulkanUnorderedAccessView(FVulkanDevice* Device, FV
 void FVulkanUnorderedAccessView::Invalidate()
 {
 	check(SourceTexture);
+	FreeBindlessHandle();
 	TextureView.Destroy(*Device);
 }
 
@@ -293,6 +348,8 @@ FVulkanUnorderedAccessView::~FVulkanUnorderedAccessView()
 		VulkanTexture->DetachView(this);
 	}
 
+	FreeBindlessHandle();
+
 	TextureView.Destroy(*Device);
 	BufferView = nullptr;
 	SourceBuffer = nullptr;
@@ -314,6 +371,7 @@ void FVulkanUnorderedAccessView::UpdateView()
 			if (SourceBuffer->IsVolatile() && VolatileLockCounter != SourceBuffer->GetVolatileLockCounter())
 			{
 				BufferView = nullptr;
+				FreeBindlessHandle();
 				VolatileLockCounter = SourceBuffer->GetVolatileLockCounter();
 			}
 
@@ -323,6 +381,22 @@ void FVulkanUnorderedAccessView::UpdateView()
 				BufferView = new FVulkanBufferView(Device);
 				BufferView->Create(SourceBuffer.GetReference(), BufferViewFormat, SourceBuffer->GetOffset(), SourceBuffer->GetSize());
 			}
+
+			if ((!BindlessHandle.IsValid()) || SourceBuffer->IsDynamic())
+			{
+				FreeBindlessHandle();
+				VkBufferViewCreateInfo ViewInfo;
+				ViewInfo.buffer = SourceBuffer->GetHandle();
+				ViewInfo.offset = SourceBuffer->GetOffset();
+				ViewInfo.range = SourceBuffer->GetSize();
+				ViewInfo.format = UEToVkTextureFormat(BufferViewFormat, false);
+				BindlessHandle = Device->GetBindlessDescriptorManager()->RegisterTexelBuffer(ViewInfo, VK_DESCRIPTOR_TYPE_STORAGE_TEXEL_BUFFER);
+			}
+		}
+		else if (!BindlessHandle.IsValid())
+		{
+			checkSlow(!SourceBuffer->IsDynamic());
+			BindlessHandle = Device->GetBindlessDescriptorManager()->RegisterBuffer(SourceBuffer->GetHandle(), SourceBuffer->GetOffset(), SourceBuffer->GetSize(), VK_DESCRIPTOR_TYPE_STORAGE_BUFFER);
 		}
 	}
 	else if (TextureView.View == VK_NULL_HANDLE)
@@ -374,9 +448,18 @@ void FVulkanUnorderedAccessView::UpdateView()
 		}
 
 		TextureView.Create(*Device, SourceTextureVK->Image, ActualViewType, SourceTextureVK->GetPartialAspectMask(), Format, UEToVkTextureFormat(Format, false), MipLevel, 1, FirstArraySlice, ActualNumArraySlices, true);
+
+		checkSlow(!BindlessHandle.IsValid());
+		BindlessHandle = Device->GetBindlessDescriptorManager()->RegisterImage(TextureView.View, VK_DESCRIPTOR_TYPE_STORAGE_IMAGE, SourceTextureVK->IsDepthOrStencilAspect());
 	}
 }
 
+FRHIDescriptorHandle FVulkanUnorderedAccessView::GetBindlessHandle() const
+{
+	checkSlow(BindlessHandle.IsValid());
+	return BindlessHandle;
+}
+
 FUnorderedAccessViewRHIRef FVulkanDynamicRHI::RHICreateUnorderedAccessView(FRHIBuffer* BufferRHI, bool bUseUAVCounter, bool bAppendBuffer)
 {
 	FVulkanResourceMultiBuffer* Buffer = ResourceCast(BufferRHI);
