VulkanRHI: Add VK_EXT_descriptor_buffer extension.  Add all necessary functions to loader and wrap layers.

diff --git a/Engine/Source/Runtime/VulkanRHI/Private/VulkanCommandWrappers.h b/Engine/Source/Runtime/VulkanRHI/Private/VulkanCommandWrappers.h
--- a/Engine/Source/Runtime/VulkanRHI/Private/VulkanCommandWrappers.h
+++ b/Engine/Source/Runtime/VulkanRHI/Private/VulkanCommandWrappers.h
@@ -211,6 +211,11 @@ struct FWrapLayer
 	static void CmdSetEvent2KHR(VkResult Result, VkCommandBuffer CommandBuffer, VkEvent Event, const VkDependencyInfo* DependencyInfo) VULKAN_LAYER_BODY
 	static void CmdWaitEvents2KHR(VkResult Result, VkCommandBuffer CommandBuffer, uint32_t EventCount, const VkEvent* Events, const VkDependencyInfo* DependencyInfos) VULKAN_LAYER_BODY
 	static void QueueSubmit2KHR(VkResult Result, VkQueue Queue, uint32_t SubmitCount, const VkSubmitInfo2* Submits, VkFence Fence) VULKAN_LAYER_BODY
+	static void GetDescriptorSetLayoutSizeEXT(VkResult Result, VkDevice Device, VkDescriptorSetLayout Layout, VkDeviceSize* OutLayoutSizeInBytes) VULKAN_LAYER_BODY
+	static void GetDescriptorSetLayoutBindingOffsetEXT(VkResult Result, VkDevice Device, VkDescriptorSetLayout Layout, uint32_t Binding, VkDeviceSize* Offset) VULKAN_LAYER_BODY
+	static void CmdBindDescriptorBuffersEXT(VkResult Result, VkCommandBuffer CommandBuffer, uint32_t BufferCount, const VkDescriptorBufferBindingInfoEXT* BindingInfos) VULKAN_LAYER_BODY
+	static void CmdSetDescriptorBufferOffsetsEXT(VkResult Result, VkCommandBuffer CommandBuffer, VkPipelineBindPoint PipelineBindPoint, VkPipelineLayout Layout, uint32_t FirstSet, uint32_t SetCount, const uint32_t* BufferIndices, const VkDeviceSize* Offsets) VULKAN_LAYER_BODY
+	static void GetDescriptorEXT(VkResult Result, VkDevice Device, const VkDescriptorGetInfoEXT* DescriptorInfo, size_t DataSize, void* Descriptor) VULKAN_LAYER_BODY
 };
 
 #undef VULKAN_LAYER_BODY
@@ -1440,11 +1445,12 @@ namespace VulkanRHI
 		FWrapLayer::GetAccelerationStructureBuildSizesKHR(VK_SUCCESS, Device, BuildType, BuildInfo, MaxPrimitiveCounts, SizeInfo);
 	}
 
-	static FORCEINLINE_DEBUGGABLE void vkGetAccelerationStructureDeviceAddressKHR(VkDevice Device, const VkAccelerationStructureDeviceAddressInfoKHR* Info)
+	static FORCEINLINE_DEBUGGABLE VkDeviceAddress vkGetAccelerationStructureDeviceAddressKHR(VkDevice Device, const VkAccelerationStructureDeviceAddressInfoKHR* Info)
 	{
 		FWrapLayer::GetAccelerationStructureDeviceAddressKHR(VK_RESULT_MAX_ENUM, Device, Info);
-		VULKANAPINAMESPACE::vkGetAccelerationStructureDeviceAddressKHR(Device, Info);
+		VkDeviceAddress Result = VULKANAPINAMESPACE::vkGetAccelerationStructureDeviceAddressKHR(Device, Info);
 		FWrapLayer::GetAccelerationStructureDeviceAddressKHR(VK_SUCCESS, Device, Info);
+		return Result;
 	}
 
 	static FORCEINLINE_DEBUGGABLE void vkCmdTraceRaysKHR(VkCommandBuffer CommandBuffer, const VkStridedDeviceAddressRegionKHR* RaygenShaderBindingTable, const VkStridedDeviceAddressRegionKHR* MissShaderBindingTable, const VkStridedDeviceAddressRegionKHR* HitShaderBindingTable, const VkStridedDeviceAddressRegionKHR* CallableShaderBindingTable, uint32 width, uint32 height, uint32 depth)
@@ -1489,7 +1495,7 @@ namespace VulkanRHI
 	{
 		FWrapLayer::CmdCopyAccelerationStructureKHR(VK_RESULT_MAX_ENUM, CommandBuffer, Info);
 		VULKANAPINAMESPACE::vkCmdCopyAccelerationStructureKHR(CommandBuffer, Info);
-		FWrapLayer::CmdCopyAccelerationStructureKHR(VK_SUCCESS, CommandBuffer, CommandBuffer, Info);
+		FWrapLayer::CmdCopyAccelerationStructureKHR(VK_SUCCESS, CommandBuffer, Info);
 	}
 #endif
 
@@ -1581,6 +1587,41 @@ namespace VulkanRHI
 		FWrapLayer::QueueSubmit2KHR(VK_SUCCESS, Queue, SubmitCount, Submits, Fence);
 	}
 
+	static FORCEINLINE_DEBUGGABLE void vkGetDescriptorSetLayoutSizeEXT(VkDevice Device, VkDescriptorSetLayout Layout, VkDeviceSize* OutLayoutSizeInBytes)
+	{
+		FWrapLayer::GetDescriptorSetLayoutSizeEXT(VK_RESULT_MAX_ENUM, Device, Layout, OutLayoutSizeInBytes);
+		VULKANAPINAMESPACE::vkGetDescriptorSetLayoutSizeEXT(Device, Layout, OutLayoutSizeInBytes);
+		FWrapLayer::GetDescriptorSetLayoutSizeEXT(VK_SUCCESS, Device, Layout, OutLayoutSizeInBytes);
+	}
+
+	static FORCEINLINE_DEBUGGABLE void vkGetDescriptorSetLayoutBindingOffsetEXT(VkDevice Device, VkDescriptorSetLayout Layout, uint32_t Binding, VkDeviceSize* Offset)
+	{
+		FWrapLayer::GetDescriptorSetLayoutBindingOffsetEXT(VK_RESULT_MAX_ENUM, Device, Layout, Binding, Offset);
+		VULKANAPINAMESPACE::vkGetDescriptorSetLayoutBindingOffsetEXT(Device, Layout, Binding, Offset);
+		FWrapLayer::GetDescriptorSetLayoutBindingOffsetEXT(VK_SUCCESS, Device, Layout, Binding, Offset);
+	}
+
+	static FORCEINLINE_DEBUGGABLE void vkCmdBindDescriptorBuffersEXT(VkCommandBuffer CommandBuffer, uint32_t BufferCount, const VkDescriptorBufferBindingInfoEXT* BindingInfos)
+	{
+		FWrapLayer::CmdBindDescriptorBuffersEXT(VK_RESULT_MAX_ENUM, CommandBuffer, BufferCount, BindingInfos);
+		VULKANAPINAMESPACE::vkCmdBindDescriptorBuffersEXT(CommandBuffer, BufferCount, BindingInfos);
+		FWrapLayer::CmdBindDescriptorBuffersEXT(VK_SUCCESS, CommandBuffer, BufferCount, BindingInfos);
+	}
+
+	static FORCEINLINE_DEBUGGABLE void vkCmdSetDescriptorBufferOffsetsEXT(VkCommandBuffer CommandBuffer, VkPipelineBindPoint PipelineBindPoint, VkPipelineLayout Layout, uint32_t FirstSet, uint32_t SetCount, const uint32_t* BufferIndices, const VkDeviceSize* Offsets)
+	{
+		FWrapLayer::CmdSetDescriptorBufferOffsetsEXT(VK_RESULT_MAX_ENUM, CommandBuffer, PipelineBindPoint, Layout, FirstSet, SetCount, BufferIndices, Offsets);
+		VULKANAPINAMESPACE::vkCmdSetDescriptorBufferOffsetsEXT(CommandBuffer, PipelineBindPoint, Layout, FirstSet, SetCount, BufferIndices, Offsets);
+		FWrapLayer::CmdSetDescriptorBufferOffsetsEXT(VK_SUCCESS, CommandBuffer, PipelineBindPoint, Layout, FirstSet, SetCount, BufferIndices, Offsets);
+	}
+
+	static FORCEINLINE_DEBUGGABLE void vkGetDescriptorEXT(VkDevice Device, const VkDescriptorGetInfoEXT* DescriptorInfo, size_t DataSize, void* Descriptor)
+	{
+		FWrapLayer::GetDescriptorEXT(VK_RESULT_MAX_ENUM, Device, DescriptorInfo, DataSize, Descriptor);
+		VULKANAPINAMESPACE::vkGetDescriptorEXT(Device, DescriptorInfo,  DataSize, Descriptor);
+		FWrapLayer::GetDescriptorEXT(VK_SUCCESS, Device, DescriptorInfo,  DataSize, Descriptor);
+	}
+
 #if VULKAN_ENABLE_IMAGE_TRACKING_LAYER
 	void BindDebugLabelName(VkImage Image, const TCHAR* Name);
 #endif
diff --git a/Engine/Source/Runtime/VulkanRHI/Private/VulkanDebug.cpp b/Engine/Source/Runtime/VulkanRHI/Private/VulkanDebug.cpp
--- a/Engine/Source/Runtime/VulkanRHI/Private/VulkanDebug.cpp
+++ b/Engine/Source/Runtime/VulkanRHI/Private/VulkanDebug.cpp
@@ -4373,7 +4373,7 @@ void FWrapLayer::GetBufferDeviceAddressKHR(VkResult Result, VkDevice Device, con
 	if (Result == VK_RESULT_MAX_ENUM)
 	{
 #if VULKAN_ENABLE_DUMP_LAYER
-		PrintfBeginResult(FString::Printf(TEXT("vkGetBufferDeviceAddressKHR")));
+		PrintfBeginResult(FString::Printf(TEXT("vkGetBufferDeviceAddressKHR(Device=0x%p, Buffer=0x%p)"), Device, Info->buffer));
 #endif
 	}
 	else
@@ -4447,7 +4447,7 @@ void FWrapLayer::ResetQueryPoolEXT(VkResult Result, VkDevice Device, VkQueryPool
 	}
 }
 
-void FWrapLayer::GetPhysicalDeviceCalibrateableTimeDomainsEXT(VkPhysicalDevice PhysicalDevice, uint32_t* TimeDomainCount, VkTimeDomainEXT* TimeDomains)
+void FWrapLayer::GetPhysicalDeviceCalibrateableTimeDomainsEXT(VkResult Result, VkPhysicalDevice PhysicalDevice, uint32_t* TimeDomainCount, VkTimeDomainEXT* TimeDomains)
 {
 	if (Result == VK_RESULT_MAX_ENUM)
 	{
@@ -4487,6 +4487,65 @@ void FWrapLayer::BindImageMemory2KHR(VkResult Result, VkDevice Device, uint32_t
 	}
 }
 
+void FWrapLayer::GetDescriptorSetLayoutSizeEXT(VkResult Result, VkDevice Device, VkDescriptorSetLayout Layout, VkDeviceSize* OutLayoutSizeInBytes)
+{
+	if (Result == VK_RESULT_MAX_ENUM)
+	{
+#if VULKAN_ENABLE_DUMP_LAYER
+		PrintfBeginResult(FString::Printf(TEXT("GetDescriptorSetLayoutSizeEXT(Device=0x%p, Layout=0x%p)"), Device, Layout));
+#endif
+	}
+}
+
+void FWrapLayer::GetDescriptorSetLayoutBindingOffsetEXT(VkResult Result, VkDevice Device, VkDescriptorSetLayout Layout, uint32_t Binding, VkDeviceSize* Offset)
+{
+	if (Result == VK_RESULT_MAX_ENUM)
+	{
+#if VULKAN_ENABLE_DUMP_LAYER
+		PrintfBeginResult(FString::Printf(TEXT("GetDescriptorSetLayoutBindingOffsetEXT(Device=0x%p, Layout=0x%p)"), Device, Layout));
+#endif
+	}
+}
+
+void FWrapLayer::CmdBindDescriptorBuffersEXT(VkResult Result, VkCommandBuffer CommandBuffer, uint32_t BufferCount, const VkDescriptorBufferBindingInfoEXT* BindingInfos)
+{
+	if (Result == VK_RESULT_MAX_ENUM)
+	{
+#if VULKAN_ENABLE_DUMP_LAYER
+		PrintfBeginResult(FString::Printf(TEXT("CmdBindDescriptorBuffersEXT(CommandBuffer=0x%p, BufferCount=%u)"), CommandBuffer, BufferCount));
+		for (uint32_t Index = 0; Index  < BufferCount; ++Index)
+		{
+			DebugLog += FString::Printf(TEXT("%s%d - Address=0x%p Usage=0x%X\n"), Tabs, Index, BindingInfos[Index].address, BindingInfos[Index].usage);
+		}
+#endif
+	}
+}
+
+void FWrapLayer::CmdSetDescriptorBufferOffsetsEXT(VkResult Result, VkCommandBuffer CommandBuffer, VkPipelineBindPoint PipelineBindPoint, VkPipelineLayout Layout, uint32_t FirstSet, uint32_t SetCount, const uint32_t* BufferIndices, const VkDeviceSize* Offsets)
+{
+	if (Result == VK_RESULT_MAX_ENUM)
+	{
+#if VULKAN_ENABLE_DUMP_LAYER
+		PrintfBeginResult(FString::Printf(TEXT("CmdSetDescriptorBufferOffsetsEXT(CommandBuffer=0x%p, BindPoint=%s, PipelineLayout=0x%p, FirstSet=%u, SetCount=%u)"), 
+			CommandBuffer, VK_TYPE_TO_STRING(VkPipelineBindPoint, PipelineBindPoint), Layout, FirstSet, SetCount));
+		for (uint32_t Index = 0; Index < SetCount; ++Index)
+		{
+			DebugLog += FString::Printf(TEXT("%sSet %u - BufferIndex=%u Offset=%llu\n"), Tabs, FirstSet + Index, BufferIndices[Index], Offsets[Index]);
+		}
+#endif
+	}
+}
+
+void FWrapLayer::GetDescriptorEXT(VkResult Result, VkDevice Device, const VkDescriptorGetInfoEXT* DescriptorInfo, size_t DataSize, void* Descriptor)
+{
+	if (Result == VK_RESULT_MAX_ENUM)
+	{
+#if VULKAN_ENABLE_DUMP_LAYER
+		PrintfBeginResult(FString::Printf(TEXT("GetDescriptorEXT(Device=0x%p, DataSize=%llu DescriptorType=%s)"), Device, DataSize, VK_TYPE_TO_STRING(VkDescriptorType, DescriptorInfo->type)));
+#endif
+	}
+}
+
 void FWrapLayer::CmdPipelineBarrier2KHR(VkResult Result, VkCommandBuffer CommandBuffer, const VkDependencyInfo* DependencyInfo)
 {
 	if (Result == VK_RESULT_MAX_ENUM)
diff --git a/Engine/Source/Runtime/VulkanRHI/Private/VulkanDevice.h b/Engine/Source/Runtime/VulkanRHI/Private/VulkanDevice.h
--- a/Engine/Source/Runtime/VulkanRHI/Private/VulkanDevice.h
+++ b/Engine/Source/Runtime/VulkanRHI/Private/VulkanDevice.h
@@ -294,6 +294,12 @@ public:
 		return GpuSubgroupProps;
 	}
 
+	inline const VkPhysicalDeviceDescriptorBufferPropertiesEXT& GetDescriptorBufferProperties() const
+	{
+		check(RHI->GetOptionalExtensions().HasKHRGetPhysicalDeviceProperties2);
+		return DescriptorBufferProps;
+	}
+
 #if VULKAN_RHI_RAYTRACING
 	inline const FRayTracingProperties& GetRayTracingProperties() const
 	{
@@ -560,8 +566,10 @@ private:
 	TArray<VkPhysicalDeviceFragmentShadingRateKHR> FragmentShadingRates;
 #endif
 
+	// Extension specific properties
 	VkPhysicalDeviceIDPropertiesKHR GpuIdProps;
 	VkPhysicalDeviceSubgroupProperties GpuSubgroupProps;
+	VkPhysicalDeviceDescriptorBufferPropertiesEXT DescriptorBufferProps;
 
 #if VULKAN_RHI_RAYTRACING
 	FRayTracingProperties RayTracingProperties;
diff --git a/Engine/Source/Runtime/VulkanRHI/Private/VulkanExtensions.cpp b/Engine/Source/Runtime/VulkanRHI/Private/VulkanExtensions.cpp
--- a/Engine/Source/Runtime/VulkanRHI/Private/VulkanExtensions.cpp
+++ b/Engine/Source/Runtime/VulkanRHI/Private/VulkanExtensions.cpp
@@ -1040,6 +1040,69 @@ public:
 };
 
 
+// ***** VK_EXT_descriptor_buffer
+class FVulkanEXTDescriptorBuffer : public FVulkanDeviceExtension
+{
+public:
+	FVulkanEXTDescriptorBuffer(FVulkanDevice* InDevice)
+		: FVulkanDeviceExtension(InDevice, VK_EXT_DESCRIPTOR_BUFFER_EXTENSION_NAME, VULKAN_EXTENSION_ENABLED, VULKAN_EXTENSION_NOT_PROMOTED)
+	{
+	}
+
+	virtual void PrePhysicalDeviceFeatures(VkPhysicalDeviceFeatures2KHR& PhysicalDeviceFeatures2) override final
+	{
+		ZeroVulkanStruct(DescriptorBufferFeatures, VK_STRUCTURE_TYPE_PHYSICAL_DEVICE_DESCRIPTOR_BUFFER_FEATURES_EXT);
+		AddToPNext(PhysicalDeviceFeatures2, DescriptorBufferFeatures);
+	}
+
+	virtual void PostPhysicalDeviceFeatures(FOptionalVulkanDeviceExtensions& ExtensionFlags) override final
+	{
+		// only enable descriptor buffers if we also support mutable descriptor types (value filled prior)
+		bSupportsDescriptorBuffer = (DescriptorBufferFeatures.descriptorBuffer == VK_TRUE);
+		ExtensionFlags.HasEXTDescriptorBuffer = bSupportsDescriptorBuffer ? 1 : 0;
+	}
+
+	virtual void PrePhysicalDeviceProperties(VkPhysicalDeviceProperties2KHR& PhysicalDeviceProperties2) override final
+	{
+		if (bSupportsDescriptorBuffer)
+		{
+			VkPhysicalDeviceDescriptorBufferPropertiesEXT& DescriptorBufferProperties = const_cast<VkPhysicalDeviceDescriptorBufferPropertiesEXT&>(Device->GetDescriptorBufferProperties());
+			ZeroVulkanStruct(DescriptorBufferProperties, VK_STRUCTURE_TYPE_PHYSICAL_DEVICE_DESCRIPTOR_BUFFER_PROPERTIES_EXT);
+			AddToPNext(PhysicalDeviceProperties2, DescriptorBufferProperties);
+		}
+	}
+
+	virtual void PreCreateDevice(VkDeviceCreateInfo& DeviceCreateInfo) override final
+	{
+		if (bSupportsDescriptorBuffer)
+		{
+			AddToPNext(DeviceCreateInfo, DescriptorBufferFeatures);
+
+			VkPhysicalDeviceDescriptorBufferPropertiesEXT& DescriptorBufferProperties = const_cast<VkPhysicalDeviceDescriptorBufferPropertiesEXT&>(Device->GetDescriptorBufferProperties());
+
+			UE_LOG(LogVulkanRHI, Display, TEXT("Enabling Vulkan Descriptor Buffers with: ")
+				TEXT("allowSamplerImageViewPostSubmitCreation=%u, maxDescriptorBufferBindings=%u, ")
+				TEXT("maxSamplerDescriptorBufferBindings=%u, maxResourceDescriptorBufferBindings=%u, ")
+				TEXT("samplerDescriptorBufferAddressSpaceSize=%llu, resourceDescriptorBufferAddressSpaceSize=%llu, ")
+				TEXT("maxSamplerDescriptorBufferRange=%llu, maxResourceDescriptorBufferRange=%llu, ")
+				TEXT("descriptorBufferAddressSpaceSize=%llu, descriptorBufferOffsetAlignment=%llu, ")
+				TEXT("samplerDescriptorSize=%llu"),
+				DescriptorBufferProperties.allowSamplerImageViewPostSubmitCreation, DescriptorBufferProperties.maxDescriptorBufferBindings,
+				DescriptorBufferProperties.maxSamplerDescriptorBufferBindings, DescriptorBufferProperties.maxResourceDescriptorBufferBindings,
+				DescriptorBufferProperties.samplerDescriptorBufferAddressSpaceSize, DescriptorBufferProperties.resourceDescriptorBufferAddressSpaceSize,
+				DescriptorBufferProperties.maxSamplerDescriptorBufferRange, DescriptorBufferProperties.maxResourceDescriptorBufferRange,
+				DescriptorBufferProperties.descriptorBufferAddressSpaceSize, DescriptorBufferProperties.descriptorBufferOffsetAlignment,
+				DescriptorBufferProperties.samplerDescriptorSize);
+		}
+	}
+
+private:
+	VkPhysicalDeviceDescriptorBufferFeaturesEXT DescriptorBufferFeatures;
+	bool bSupportsDescriptorBuffer = false;
+};
+
+
+
 
 template <typename ExtensionType>
 static void FlagExtensionSupport(const TArray<VkExtensionProperties>& ExtensionProperties, TArray<TUniquePtr<ExtensionType>>& UEExtensions, const TCHAR* ExtensionTypeName)
@@ -1115,6 +1178,7 @@ FVulkanDeviceExtensionArray FVulkanDeviceExtension::GetUESupportedDeviceExtensio
 	ADD_CUSTOM_EXTENSION(FVulkanEXTHostQueryResetExtension);
 	ADD_CUSTOM_EXTENSION(FVulkanEXTSubgroupSizeControlExtension);
 	ADD_CUSTOM_EXTENSION(FVulkanEXTCalibratedTimestampsExtension);
+	ADD_CUSTOM_EXTENSION(FVulkanEXTDescriptorBuffer);
 
 	// Needed for Raytracing
 	ADD_CUSTOM_EXTENSION(FVulkanKHRBufferDeviceAddressExtension);
diff --git a/Engine/Source/Runtime/VulkanRHI/Private/VulkanLoader.h b/Engine/Source/Runtime/VulkanRHI/Private/VulkanLoader.h
--- a/Engine/Source/Runtime/VulkanRHI/Private/VulkanLoader.h
+++ b/Engine/Source/Runtime/VulkanRHI/Private/VulkanLoader.h
@@ -208,7 +208,12 @@
 	EnumMacro(PFN_vkGetPhysicalDeviceCalibrateableTimeDomainsEXT, vkGetPhysicalDeviceCalibrateableTimeDomainsEXT) \
 	EnumMacro(PFN_vkGetCalibratedTimestampsEXT, vkGetCalibratedTimestampsEXT) \
 	EnumMacro(PFN_vkBindBufferMemory2KHR, vkBindBufferMemory2KHR) \
-	EnumMacro(PFN_vkBindImageMemory2KHR, vkBindImageMemory2KHR)
+	EnumMacro(PFN_vkBindImageMemory2KHR, vkBindImageMemory2KHR) \
+	EnumMacro(PFN_vkGetDescriptorSetLayoutSizeEXT, vkGetDescriptorSetLayoutSizeEXT) \
+	EnumMacro(PFN_vkGetDescriptorSetLayoutBindingOffsetEXT, vkGetDescriptorSetLayoutBindingOffsetEXT) \
+	EnumMacro(PFN_vkCmdBindDescriptorBuffersEXT, vkCmdBindDescriptorBuffersEXT) \
+	EnumMacro(PFN_vkCmdSetDescriptorBufferOffsetsEXT, vkCmdSetDescriptorBufferOffsetsEXT) \
+	EnumMacro(PFN_vkGetDescriptorEXT, vkGetDescriptorEXT)
 
 
 // List of all Vulkan entry points
