VulkanRHI:  Add VK_BUFFER_USAGE_SHADER_DEVICE_ADDRESS_BIT to buffer usage flags when bindless is enabled.  Make GCVarRobustBufferAccess accessible from bindless descriptor manager (can affect descriptor sizes).

diff --git a/Engine/Source/Runtime/VulkanRHI/Private/VulkanBuffer.cpp b/Engine/Source/Runtime/VulkanRHI/Private/VulkanBuffer.cpp
--- a/Engine/Source/Runtime/VulkanRHI/Private/VulkanBuffer.cpp
+++ b/Engine/Source/Runtime/VulkanRHI/Private/VulkanBuffer.cpp
@@ -119,6 +119,11 @@ VkBufferUsageFlags FVulkanResourceMultiBuffer::UEToVKBufferUsageFlags(FVulkanDev
 			TranslateFlag(BUF_AccelerationStructure, 0, VK_BUFFER_USAGE_ACCELERATION_STRUCTURE_BUILD_INPUT_READ_ONLY_BIT_KHR);
 		}
 #endif
+		// For descriptors buffers
+		if (InDevice->SupportsBindless())
+		{
+			OutVkUsage |= VK_BUFFER_USAGE_SHADER_DEVICE_ADDRESS_BIT;
+		}
 	}
 
 	return OutVkUsage;
diff --git a/Engine/Source/Runtime/VulkanRHI/Private/VulkanDevice.cpp b/Engine/Source/Runtime/VulkanRHI/Private/VulkanDevice.cpp
--- a/Engine/Source/Runtime/VulkanRHI/Private/VulkanDevice.cpp
+++ b/Engine/Source/Runtime/VulkanRHI/Private/VulkanDevice.cpp
@@ -33,7 +33,7 @@ TAutoConsoleVariable<int32> GAllowPresentOnComputeQueue(
 	TEXT("1 to allow presenting on the compute queue if available")
 );
 
-static TAutoConsoleVariable<int32> GCVarRobustBufferAccess(
+TAutoConsoleVariable<int32> GCVarRobustBufferAccess(
 	TEXT("r.Vulkan.RobustBufferAccess"),
 	1,
 	TEXT("0 to disable robust buffer access")
