VulkanRHI: Add VK_PIPELINE_CREATE_DESCRIPTOR_BUFFER_BIT_EXT flag to pipelines creating with bindless enabled.

diff --git a/Engine/Source/Runtime/VulkanRHI/Private/VulkanPipeline.cpp b/Engine/Source/Runtime/VulkanRHI/Private/VulkanPipeline.cpp
--- a/Engine/Source/Runtime/VulkanRHI/Private/VulkanPipeline.cpp
+++ b/Engine/Source/Runtime/VulkanRHI/Private/VulkanPipeline.cpp
@@ -1399,6 +1399,11 @@ VkResult FVulkanPipelineStateCacheManager::CreateVKPipeline(FVulkanRHIGraphicsPi
 
 	FPipelineCache& Cache = bIsPrecompileJob ? CurrentPrecompilingPSOCache : GlobalPSOCache;
 
+	if (Device->SupportsBindless())
+	{
+		PipelineInfo.flags |= VK_PIPELINE_CREATE_DESCRIPTOR_BUFFER_BIT_EXT;
+	}
+
 	VkPipelineCache LocalPipelineCache = VK_NULL_HANDLE;
 	VkResult Result = VK_ERROR_INITIALIZATION_FAILED;
 	uint32 PSOSize = 0;
@@ -1658,34 +1663,33 @@ void FVulkanPipelineStateCacheManager::CreateGfxEntry(const FGraphicsPipelineSta
 	FVulkanVertexInputStateInfo VertexInputState;
 	
 	{
+		const FBoundShaderStateInput& BSI = PSOInitializer.BoundShaderState;
 
-			const FBoundShaderStateInput& BSI = PSOInitializer.BoundShaderState;
-
-			const FVulkanShaderHeader& VSHeader = Shaders[ShaderStage::Vertex]->GetCodeHeader();
-			VertexInputState.Generate(ResourceCast(PSOInitializer.BoundShaderState.VertexDeclarationRHI), VSHeader.InOutMask);
+		const FVulkanShaderHeader& VSHeader = Shaders[ShaderStage::Vertex]->GetCodeHeader();
+		VertexInputState.Generate(ResourceCast(PSOInitializer.BoundShaderState.VertexDeclarationRHI), VSHeader.InOutMask);
 
-			FUniformBufferGatherInfo UBGatherInfo;
+		FUniformBufferGatherInfo UBGatherInfo;
 
-			DescriptorSetLayoutInfo.ProcessBindingsForStage(VK_SHADER_STAGE_VERTEX_BIT, ShaderStage::Vertex, VSHeader, UBGatherInfo);
+		DescriptorSetLayoutInfo.ProcessBindingsForStage(VK_SHADER_STAGE_VERTEX_BIT, ShaderStage::Vertex, VSHeader, UBGatherInfo);
 
-			if (Shaders[ShaderStage::Pixel])
-			{
-				const FVulkanShaderHeader& PSHeader = Shaders[ShaderStage::Pixel]->GetCodeHeader();
-				DescriptorSetLayoutInfo.ProcessBindingsForStage(VK_SHADER_STAGE_FRAGMENT_BIT, ShaderStage::Pixel, PSHeader, UBGatherInfo);
-			}
+		if (Shaders[ShaderStage::Pixel])
+		{
+			const FVulkanShaderHeader& PSHeader = Shaders[ShaderStage::Pixel]->GetCodeHeader();
+			DescriptorSetLayoutInfo.ProcessBindingsForStage(VK_SHADER_STAGE_FRAGMENT_BIT, ShaderStage::Pixel, PSHeader, UBGatherInfo);
+		}
 
 #if VULKAN_SUPPORTS_GEOMETRY_SHADERS
-			if (Shaders[ShaderStage::Geometry])
-			{
-				const FVulkanShaderHeader& GSHeader = Shaders[ShaderStage::Geometry]->GetCodeHeader();
-				DescriptorSetLayoutInfo.ProcessBindingsForStage(VK_SHADER_STAGE_GEOMETRY_BIT, ShaderStage::Geometry, GSHeader, UBGatherInfo);
-			}
+		if (Shaders[ShaderStage::Geometry])
+		{
+			const FVulkanShaderHeader& GSHeader = Shaders[ShaderStage::Geometry]->GetCodeHeader();
+			DescriptorSetLayoutInfo.ProcessBindingsForStage(VK_SHADER_STAGE_GEOMETRY_BIT, ShaderStage::Geometry, GSHeader, UBGatherInfo);
+		}
 #endif
 
-			// Second pass
-			const int32 NumImmutableSamplers = PSOInitializer.ImmutableSamplerState.ImmutableSamplers.Num();
-			TArrayView<FRHISamplerState*> ImmutableSamplers(NumImmutableSamplers > 0 ? &(FRHISamplerState*&)PSOInitializer.ImmutableSamplerState.ImmutableSamplers[0] : nullptr, NumImmutableSamplers);
-			DescriptorSetLayoutInfo.FinalizeBindings<false>(*Device, UBGatherInfo, ImmutableSamplers);
+		// Second pass
+		const int32 NumImmutableSamplers = PSOInitializer.ImmutableSamplerState.ImmutableSamplers.Num();
+		TArrayView<FRHISamplerState*> ImmutableSamplers(NumImmutableSamplers > 0 ? &(FRHISamplerState*&)PSOInitializer.ImmutableSamplerState.ImmutableSamplers[0] : nullptr, NumImmutableSamplers);
+		DescriptorSetLayoutInfo.FinalizeBindings<false>(*Device, UBGatherInfo, ImmutableSamplers);
 	}
 
 	FDescriptorSetRemappingInfo& RemappingInfo = DescriptorSetLayoutInfo.RemappingInfo;
@@ -2153,6 +2157,11 @@ FVulkanComputePipeline* FVulkanPipelineStateCacheManager::CreateComputePipelineF
 	PipelineInfo.stage.pName = EntryPoint;
 	PipelineInfo.layout = ComputeLayout->GetPipelineLayout();
 
+	if (Device->SupportsBindless())
+	{
+		PipelineInfo.flags |= VK_PIPELINE_CREATE_DESCRIPTOR_BUFFER_BIT_EXT;
+	}
+
 	VkPipelineShaderStageRequiredSubgroupSizeCreateInfo RequiredSubgroupSizeCreateInfo;
 	if ((CSHeader.WaveSize > 0) && Device->GetOptionalExtensions().HasEXTSubgroupSizeControl)
 	{
