VulkanRHI: Add checks in all the old descriptor writer functions to make sure they aren't called in bindless.  Also skip descriptor binding.  (Vulkan Bindless WIP)

diff --git a/Engine/Source/Runtime/VulkanRHI/Private/VulkanPipelineState.h b/Engine/Source/Runtime/VulkanRHI/Private/VulkanPipelineState.h
--- a/Engine/Source/Runtime/VulkanRHI/Private/VulkanPipelineState.h
+++ b/Engine/Source/Runtime/VulkanRHI/Private/VulkanPipelineState.h
@@ -24,9 +24,8 @@ class FVulkanCommonPipelineDescriptorState : public VulkanRHI::FDeviceChild
 public:
 	FVulkanCommonPipelineDescriptorState(FVulkanDevice* InDevice)
 		: VulkanRHI::FDeviceChild(InDevice)
+		, bUseBindless(InDevice->SupportsBindless())
 	{
-		bBindlessSamplers = (RHIGetBindlessSamplersConfiguration(GMaxRHIShaderPlatform) != ERHIBindlessConfiguration::Disabled);
-		bBindlessResources = (RHIGetBindlessResourcesConfiguration(GMaxRHIShaderPlatform) != ERHIBindlessConfiguration::Disabled);
 	}
 
 	virtual ~FVulkanCommonPipelineDescriptorState() {}
@@ -42,7 +41,7 @@ public:
 		}
 		return DSetsKey;
 	}
-	
+
 	bool HasVolatileResources() const
 	{
 		for (const FVulkanDescriptorSetWriter& Writer : DSWriter)
@@ -63,23 +62,27 @@ public:
 
 	inline void SetStorageBuffer(uint8 DescriptorSet, uint32 BindingIndex, const FVulkanResourceMultiBuffer* StructuredBuffer)
 	{
+		check(!bUseBindless);
 		check(StructuredBuffer && (StructuredBuffer->GetBufferUsageFlags() & VK_BUFFER_USAGE_STORAGE_BUFFER_BIT) == VK_BUFFER_USAGE_STORAGE_BUFFER_BIT);
 		MarkDirty(DSWriter[DescriptorSet].WriteStorageBuffer(BindingIndex, StructuredBuffer->GetCurrentAllocation(), StructuredBuffer->GetOffset(), StructuredBuffer->GetCurrentSize()));
 	}
 
 	inline void SetUAVTexelBufferViewState(uint8 DescriptorSet, uint32 BindingIndex, const FVulkanBufferView* View)
 	{
+		check(!bUseBindless);
 		check(View && (View->Flags & VK_BUFFER_USAGE_STORAGE_TEXEL_BUFFER_BIT) == VK_BUFFER_USAGE_STORAGE_TEXEL_BUFFER_BIT);
 		MarkDirty(DSWriter[DescriptorSet].WriteStorageTexelBuffer(BindingIndex, View));
 	}
 
 	inline void SetUAVTextureView(uint8 DescriptorSet, uint32 BindingIndex, const FVulkanTextureView& TextureView, VkImageLayout Layout)
 	{
+		check(!bUseBindless);
 		MarkDirty(DSWriter[DescriptorSet].WriteStorageImage(BindingIndex, TextureView, Layout));
 	}
 
 	inline void SetTexture(uint8 DescriptorSet, uint32 BindingIndex, const FVulkanTexture* Texture, VkImageLayout Layout)
 	{
+		check(!bUseBindless);
 		check(Texture && Texture->PartialView);
 
 		// If the texture doesn't support sampling, then we read it through a UAV
@@ -95,50 +98,47 @@ public:
 
 	inline void SetSRVBufferViewState(uint8 DescriptorSet, uint32 BindingIndex, const FVulkanBufferView* View)
 	{
+		check(!bUseBindless);
 		check(View && (View->Flags & VK_BUFFER_USAGE_UNIFORM_TEXEL_BUFFER_BIT) == VK_BUFFER_USAGE_UNIFORM_TEXEL_BUFFER_BIT);
 		MarkDirty(DSWriter[DescriptorSet].WriteUniformTexelBuffer(BindingIndex, View));
 	}
 
 	inline void SetSRVTextureView(uint8 DescriptorSet, uint32 BindingIndex, const FVulkanTextureView& TextureView, VkImageLayout Layout)
 	{
+		check(!bUseBindless);
 		MarkDirty(DSWriter[DescriptorSet].WriteImage(BindingIndex, TextureView, Layout));
 	}
 
 	inline void SetSamplerState(uint8 DescriptorSet, uint32 BindingIndex, const FVulkanSamplerState* Sampler)
 	{
-		if (!bBindlessSamplers)
-		{
-			check(Sampler && Sampler->Sampler != VK_NULL_HANDLE);
-			MarkDirty(DSWriter[DescriptorSet].WriteSampler(BindingIndex, *Sampler));
-		}
+		check(!bUseBindless);
+		check(Sampler && Sampler->Sampler != VK_NULL_HANDLE);
+		MarkDirty(DSWriter[DescriptorSet].WriteSampler(BindingIndex, *Sampler));
 	}
 
 	inline void SetInputAttachment(uint8 DescriptorSet, uint32 BindingIndex, const FVulkanTextureView& TextureView, VkImageLayout Layout)
 	{
+		check(!bUseBindless);
 		MarkDirty(DSWriter[DescriptorSet].WriteInputAttachment(BindingIndex, TextureView, Layout));
 	}
 
 	template<bool bDynamic>
 	inline void SetUniformBuffer(uint8 DescriptorSet, uint32 BindingIndex, const FVulkanUniformBuffer* UniformBuffer)
 	{
-/*
-		if ((UniformBuffersWithDataMask[DescriptorSet] & (1ULL << (uint64)BindingIndex)) != 0)
-*/
+		if (bDynamic)
 		{
-			if (bDynamic)
-			{
-				MarkDirty(DSWriter[DescriptorSet].WriteDynamicUniformBuffer(BindingIndex, UniformBuffer->Allocation, 0, UniformBuffer->GetSize(), UniformBuffer->GetOffset()));
-			}
-			else
-			{
-				MarkDirty(DSWriter[DescriptorSet].WriteUniformBuffer(BindingIndex, UniformBuffer->Allocation, UniformBuffer->GetOffset(), UniformBuffer->GetSize()));
-			}
+			MarkDirty(DSWriter[DescriptorSet].WriteDynamicUniformBuffer(BindingIndex, UniformBuffer->Allocation, 0, UniformBuffer->GetSize(), UniformBuffer->GetOffset()));
+		}
+		else
+		{
+			MarkDirty(DSWriter[DescriptorSet].WriteUniformBuffer(BindingIndex, UniformBuffer->Allocation, UniformBuffer->GetOffset(), UniformBuffer->GetSize()));
 		}
 	}
 
 #if VULKAN_RHI_RAYTRACING
 	inline void SetAccelerationStructure(uint8 DescriptorSet, uint32 BindingIndex, VkAccelerationStructureKHR AccelerationStructure)
 	{
+		check(!bUseBindless);
 		MarkDirty(DSWriter[DescriptorSet].WriteAccelerationStructure(BindingIndex, AccelerationStructure));
 	}
 #endif // VULKAN_RHI_RAYTRACING
@@ -153,13 +153,15 @@ protected:
 	}
 	inline void Bind(VkCommandBuffer CmdBuffer, VkPipelineLayout PipelineLayout, VkPipelineBindPoint BindPoint)
 	{
-		// Bindless will sneak in two sets before the others
-		const uint32 FirstSet = Device->SupportsBindless() ? VulkanBindless::NumBindlessSets : 0;
-		VulkanRHI::vkCmdBindDescriptorSets(CmdBuffer,
-			BindPoint,
-			PipelineLayout,
-			FirstSet, DescriptorSetHandles.Num(), DescriptorSetHandles.GetData(),
-			(uint32)DynamicOffsets.Num(), DynamicOffsets.GetData());
+		// Bindless will replace with global sets
+		if (!bUseBindless)
+		{
+			VulkanRHI::vkCmdBindDescriptorSets(CmdBuffer,
+				BindPoint,
+				PipelineLayout,
+				0, DescriptorSetHandles.Num(), DescriptorSetHandles.GetData(),
+				(uint32)DynamicOffsets.Num(), DynamicOffsets.GetData());
+		}
 	}
 
 	void CreateDescriptorWriteInfos();
@@ -185,8 +187,7 @@ protected:
 	mutable FVulkanDSetsKey DSetsKey;
 	mutable bool bIsDSetsKeyDirty = true;
 
-	bool bBindlessSamplers = false;
-	bool bBindlessResources = false;
+	const bool bUseBindless;
 };
 
 
@@ -217,6 +218,8 @@ public:
 
 	bool UpdateDescriptorSets(FVulkanCommandListContext* CmdListContext, FVulkanCmdBuffer* CmdBuffer)
 	{
+		check(!bUseBindless);
+
 		const bool bUseDynamicGlobalUBs = (GDynamicGlobalUBs->GetInt() > 0);
 		if (bUseDynamicGlobalUBs)
 		{
@@ -274,19 +277,10 @@ public:
 		PackedUniformBuffers[Stage].SetEmulatedUniformBufferIntoPacked(BindingIndex, ConstantData, PackedUniformBuffersDirty[Stage]);
 	}
 
-	inline void SetDynamicUniformBuffer(uint8 DescriptorSet, uint32 BindingIndex, const FVulkanUniformBuffer* UniformBuffer)
-	{
-		ensure(0);
-/*
-		if ((UniformBuffersWithDataMask[DescriptorSet] & (1ULL << (uint64)BindingIndex)) != 0)
-		{
-			MarkDirty(DSWriter[DescriptorSet].WriteDynamicUniformBuffer(BindingIndex, *UniformBuffer->GetBufferAllocation(), 0, UniformBuffer->GetSize(), UniformBuffer->GetOffset()));
-		}
-*/
-	}
-
 	bool UpdateDescriptorSets(FVulkanCommandListContext* CmdListContext, FVulkanCmdBuffer* CmdBuffer)
 	{
+		check(!bUseBindless);
+
 		const bool bUseDynamicGlobalUBs = (GDynamicGlobalUBs->GetInt() > 0);
 		if (bUseDynamicGlobalUBs)
 		{
