VulkanRHI: Update the packed uniform buffers in bindless (Vulkan Bindless WIP)

#rnx

diff --git a/Engine/Source/Runtime/VulkanRHI/Private/VulkanPendingState.cpp b/Engine/Source/Runtime/VulkanRHI/Private/VulkanPendingState.cpp
--- a/Engine/Source/Runtime/VulkanRHI/Private/VulkanPendingState.cpp
+++ b/Engine/Source/Runtime/VulkanRHI/Private/VulkanPendingState.cpp
@@ -380,11 +380,17 @@ void FVulkanPendingComputeState::PrepareForDispatch(FVulkanCmdBuffer* InCmdBuffe
 
 	check(CurrentState);
 
-	const bool bHasDescriptorSets = CurrentState->UpdateDescriptorSets(&Context, InCmdBuffer);
-
 	VkCommandBuffer CmdBuffer = InCmdBuffer->GetHandle();
 
+	if (Device->SupportsBindless())
+	{
+		CurrentState->UpdateBindlessDescriptors(&Context, InCmdBuffer);
+		CurrentPipeline->Bind(CmdBuffer);
+	}
+	else
 	{
+		const bool bHasDescriptorSets = CurrentState->UpdateDescriptorSets(&Context, InCmdBuffer);
+
 		//#todo-rco: Move this to SetComputePipeline()
 #if VULKAN_ENABLE_AGGRESSIVE_STATS
 		SCOPE_CYCLE_COUNTER(STAT_VulkanPipelineBind);
@@ -416,20 +422,29 @@ void FVulkanPendingGfxState::PrepareForDraw(FVulkanCmdBuffer* CmdBuffer)
 
 	check(CmdBuffer->bHasPipeline);
 
-	// TODO: Add 'dirty' flag? Need to rebind only on PSO change
-	if (CurrentPipeline->bHasInputAttachments)
+	if (Device->SupportsBindless())
 	{
-		FVulkanFramebuffer* CurrentFramebuffer = Context.GetCurrentFramebuffer();
-		UpdateInputAttachments(CurrentFramebuffer);
+		check(!CurrentPipeline->bHasInputAttachments); // todo-jn: bindless + InputAttachments
+		UpdateDynamicStates(CmdBuffer);
+		CurrentState->UpdateBindlessDescriptors(&Context, CmdBuffer);
 	}
-	
-	bool bHasDescriptorSets = CurrentState->UpdateDescriptorSets(&Context, CmdBuffer);
+	else
+	{
+		// TODO: Add 'dirty' flag? Need to rebind only on PSO change
+		if (CurrentPipeline->bHasInputAttachments)
+		{
+			FVulkanFramebuffer* CurrentFramebuffer = Context.GetCurrentFramebuffer();
+			UpdateInputAttachments(CurrentFramebuffer);
+		}
 
-	UpdateDynamicStates(CmdBuffer);
+		const bool bHasDescriptorSets = CurrentState->UpdateDescriptorSets(&Context, CmdBuffer);
 
-	if (bHasDescriptorSets)
-	{
-		CurrentState->BindDescriptorSets(CmdBuffer->GetHandle());
+		UpdateDynamicStates(CmdBuffer);
+
+		if (bHasDescriptorSets)
+		{
+			CurrentState->BindDescriptorSets(CmdBuffer->GetHandle());
+		}
 	}
 
 	if (bDirtyVertexStreams)
diff --git a/Engine/Source/Runtime/VulkanRHI/Private/VulkanPipelineState.cpp b/Engine/Source/Runtime/VulkanRHI/Private/VulkanPipelineState.cpp
--- a/Engine/Source/Runtime/VulkanRHI/Private/VulkanPipelineState.cpp
+++ b/Engine/Source/Runtime/VulkanRHI/Private/VulkanPipelineState.cpp
@@ -216,6 +216,54 @@ bool FVulkanComputePipelineDescriptorState::InternalUpdateDescriptorSets(FVulkan
 	return true;
 }
 
+void FVulkanComputePipelineDescriptorState::UpdateBindlessDescriptors(FVulkanCommandListContext* CmdListContext, FVulkanCmdBuffer* CmdBuffer)
+{
+	check(Device->SupportsBindless());
+
+	// We should only have uniform buffers at this point
+	check(DSWriteContainer.DescriptorBufferInfo.Num() == DSWriteContainer.DescriptorWrites.Num());
+	check(DSWriteContainer.DescriptorImageInfo.Num() == 0);
+
+	FVulkanUniformBufferUploader* UniformBufferUploader = CmdListContext->GetUniformBufferUploader();
+	uint8* CPURingBufferBase = (uint8*)UniformBufferUploader->GetCPUMappedPointer();
+	const VkDeviceSize UBOffsetAlignment = Device->GetLimits().minUniformBufferOffsetAlignment;
+
+	if (PackedUniformBuffersDirty != 0)
+	{
+#if VULKAN_ENABLE_AGGRESSIVE_STATS
+		SCOPE_CYCLE_COUNTER(STAT_VulkanApplyPackedUniformBuffers);
+#endif
+		uint64 RemainingPackedUniformsMask = PackedUniformBuffersDirty;
+		int32 PackedUBIndex = 0;
+		while (RemainingPackedUniformsMask)
+		{
+			if (RemainingPackedUniformsMask & 1)
+			{
+				const FPackedUniformBuffers::FPackedBuffer& StagedUniformBuffer = PackedUniformBuffers.GetBuffer(PackedUBIndex);
+				const int32 BindingIndex = PipelineDescriptorInfo->RemappingInfo->StageInfos[0].PackedUBBindingIndices[PackedUBIndex];
+				VkDescriptorBufferInfo& BufferInfo = *const_cast<VkDescriptorBufferInfo*>(DSWriter[0].WriteDescriptors[BindingIndex].pBufferInfo);
+				BufferInfo.buffer = UniformBufferUploader->GetCPUBufferHandle();
+				BufferInfo.range = StagedUniformBuffer.Num();
+				check((BufferInfo.range & 0xF) == 0);   // :todo-jn: nv driver issue
+
+				const uint64 RingBufferOffset = UniformBufferUploader->AllocateMemory(BufferInfo.range, UBOffsetAlignment, CmdBuffer);
+				BufferInfo.offset = UniformBufferUploader->GetCPUBufferOffset() + RingBufferOffset;
+
+				// get location in the ring buffer to use
+				FMemory::Memcpy(CPURingBufferBase + RingBufferOffset, StagedUniformBuffer.GetData(), BufferInfo.range);
+			}
+			RemainingPackedUniformsMask = RemainingPackedUniformsMask >> 1;
+			++PackedUBIndex;
+		}
+
+		PackedUniformBuffersDirty = 0;
+	}
+
+	if (DSWriteContainer.DescriptorBufferInfo.Num())
+	{
+		Device->GetBindlessDescriptorManager()->RegisterUniformBuffers(CmdBuffer->GetHandle(), VK_PIPELINE_BIND_POINT_COMPUTE, ShaderStage::EStage::Compute, DSWriter[0]);
+	}
+}
 
 FVulkanGraphicsPipelineDescriptorState::FVulkanGraphicsPipelineDescriptorState(FVulkanDevice* InDevice, FVulkanRHIGraphicsPipelineState* InGfxPipeline)
 	: FVulkanCommonPipelineDescriptorState(InDevice)
@@ -351,6 +399,66 @@ bool FVulkanGraphicsPipelineDescriptorState::InternalUpdateDescriptorSets(FVulka
 	return true;
 }
 
+void FVulkanGraphicsPipelineDescriptorState::UpdateBindlessDescriptors(FVulkanCommandListContext* CmdListContext, FVulkanCmdBuffer* CmdBuffer)
+{
+	check(Device->SupportsBindless());
+
+	// We should only have uniform buffers at this point
+	check(DSWriteContainer.DescriptorBufferInfo.Num() == DSWriteContainer.DescriptorWrites.Num());
+	check(DSWriteContainer.DescriptorImageInfo.Num() == 0);
+
+	FVulkanUniformBufferUploader* UniformBufferUploader = CmdListContext->GetUniformBufferUploader();
+	uint8* CPURingBufferBase = (uint8*)UniformBufferUploader->GetCPUMappedPointer();
+	const VkDeviceSize UBOffsetAlignment = Device->GetLimits().minUniformBufferOffsetAlignment;
+
+	const FDescriptorSetRemappingInfo* RESTRICT RemappingInfo = PipelineDescriptorInfo->RemappingInfo;
+
+	// Process updates
+	{
+#if VULKAN_ENABLE_AGGRESSIVE_STATS
+		SCOPE_CYCLE_COUNTER(STAT_VulkanApplyPackedUniformBuffers);
+#endif
+		int32 DSWriterIndex = 0;
+		for (int32 Stage = 0; Stage < ShaderStage::NumStages; ++Stage)
+		{
+			if (PackedUniformBuffersDirty[Stage] != 0)
+			{
+				const uint32 DescriptorSet = RemappingInfo->StageInfos[Stage].PackedUBDescriptorSet;
+				uint64 RemainingPackedUniformsMask = PackedUniformBuffersDirty[Stage];
+				int32 PackedUBIndex = 0;
+				while (RemainingPackedUniformsMask)
+				{
+					if (RemainingPackedUniformsMask & 1)
+					{
+						const FPackedUniformBuffers::FPackedBuffer& StagedUniformBuffer = PackedUniformBuffers[Stage].GetBuffer(PackedUBIndex);
+						const int32 BindingIndex = RemappingInfo->StageInfos[Stage].PackedUBBindingIndices[PackedUBIndex];
+						VkDescriptorBufferInfo& BufferInfo = *const_cast<VkDescriptorBufferInfo*>(DSWriter[DescriptorSet].WriteDescriptors[BindingIndex].pBufferInfo);
+						BufferInfo.buffer = UniformBufferUploader->GetCPUBufferHandle();
+						BufferInfo.range = StagedUniformBuffer.Num();
+						check((BufferInfo.range & 0xF) == 0);   // :todo-jn: nv driver issue
+
+						const uint64 RingBufferOffset = UniformBufferUploader->AllocateMemory(BufferInfo.range, UBOffsetAlignment, CmdBuffer);
+						BufferInfo.offset = UniformBufferUploader->GetCPUBufferOffset() + RingBufferOffset;
+
+						// get location in the ring buffer to use
+						FMemory::Memcpy(CPURingBufferBase + RingBufferOffset, StagedUniformBuffer.GetData(), BufferInfo.range);
+					}
+					RemainingPackedUniformsMask = RemainingPackedUniformsMask >> 1;
+					++PackedUBIndex;
+				}
+
+				PackedUniformBuffersDirty[Stage] = 0;
+			}
+
+			if ((RemappingInfo->StageInfos[Stage].PackedUBDescriptorSet != UINT16_MAX) || RemappingInfo->StageInfos[Stage].UniformBuffers.Num())
+			{
+				Device->GetBindlessDescriptorManager()->RegisterUniformBuffers(CmdBuffer->GetHandle(), VK_PIPELINE_BIND_POINT_GRAPHICS, (ShaderStage::EStage)Stage, DSWriter[DSWriterIndex++]);
+			}
+		}
+	}
+}
+
+
 template <typename TRHIShader>
 void FVulkanCommandListContext::ApplyStaticUniformBuffers(TRHIShader* Shader)
 {
diff --git a/Engine/Source/Runtime/VulkanRHI/Private/VulkanPipelineState.h b/Engine/Source/Runtime/VulkanRHI/Private/VulkanPipelineState.h
--- a/Engine/Source/Runtime/VulkanRHI/Private/VulkanPipelineState.h
+++ b/Engine/Source/Runtime/VulkanRHI/Private/VulkanPipelineState.h
@@ -231,6 +231,8 @@ public:
 		}
 	}
 
+	void UpdateBindlessDescriptors(FVulkanCommandListContext* CmdListContext, FVulkanCmdBuffer* CmdBuffer);
+
 	inline void BindDescriptorSets(VkCommandBuffer CmdBuffer)
 	{
 		Bind(CmdBuffer, ComputePipeline->GetLayout().GetPipelineLayout(), VK_PIPELINE_BIND_POINT_COMPUTE);
@@ -292,6 +294,8 @@ public:
 		}
 	}
 
+	void UpdateBindlessDescriptors(FVulkanCommandListContext* CmdListContext, FVulkanCmdBuffer* CmdBuffer);
+
 	inline void BindDescriptorSets(VkCommandBuffer CmdBuffer)
 	{
 		Bind(CmdBuffer, GfxPipeline->GetLayout().GetPipelineLayout(), VK_PIPELINE_BIND_POINT_GRAPHICS);
