VulkanRHI: Update views before using them in bindless (Vulkan Bindless WIP).

#rnx

diff --git a/Engine/Source/Runtime/VulkanRHI/Private/VulkanCommands.cpp b/Engine/Source/Runtime/VulkanRHI/Private/VulkanCommands.cpp
--- a/Engine/Source/Runtime/VulkanRHI/Private/VulkanCommands.cpp
+++ b/Engine/Source/Runtime/VulkanRHI/Private/VulkanCommands.cpp
@@ -257,8 +257,40 @@ void FVulkanCommandListContext::RHISetShaderParameter(FRHIComputeShader* Compute
 	PendingComputeState->SetPackedGlobalShaderParameter(BufferIndex, BaseIndex, NumBytes, NewValue);
 }
 
+static void UpdateBindlessHandles(TConstArrayView<FRHIShaderParameterResource> InBindlessParameters)
+{
+	for (const FRHIShaderParameterResource& Parameter : InBindlessParameters)
+	{
+		if (FRHIResource* Resource = Parameter.Resource)
+		{
+			switch (Parameter.Type)
+			{
+			case FRHIShaderParameterResource::EType::ResourceView:
+			{
+				FRHIShaderResourceView* ShaderResourceView = static_cast<FRHIShaderResourceView*>(Resource);
+				FVulkanShaderResourceView* VKShaderResourceView = ResourceCast(ShaderResourceView);
+				VKShaderResourceView->UpdateView();
+				checkSlow(VKShaderResourceView->GetBindlessHandle().IsValid());
+			}
+			break;
+			case FRHIShaderParameterResource::EType::UnorderedAccessView:
+			{
+				FRHIUnorderedAccessView* UnorderedAccessView = static_cast<FRHIUnorderedAccessView*>(Resource);
+				FVulkanUnorderedAccessView* VKUnorderedAccessView = ResourceCast(UnorderedAccessView);
+				VKUnorderedAccessView->UpdateView();
+				checkSlow(VKUnorderedAccessView->GetBindlessHandle().IsValid());
+			}
+			default:
+				break;
+			}
+		}
+	}
+}
+
 void FVulkanCommandListContext::RHISetShaderParameters(FRHIGraphicsShader* Shader, TConstArrayView<uint8> InParametersData, TConstArrayView<FRHIShaderParameter> InParameters, TConstArrayView<FRHIShaderParameterResource> InResourceParameters, TConstArrayView<FRHIShaderParameterResource> InBindlessParameters)
 {
+	UpdateBindlessHandles(InBindlessParameters);
+
 	UE::RHICore::RHISetShaderParametersShared(
 		*this
 		, Shader
@@ -271,6 +303,8 @@ void FVulkanCommandListContext::RHISetShaderParameters(FRHIGraphicsShader* Shade
 
 void FVulkanCommandListContext::RHISetShaderParameters(FRHIComputeShader* Shader, TConstArrayView<uint8> InParametersData, TConstArrayView<FRHIShaderParameter> InParameters, TConstArrayView<FRHIShaderParameterResource> InResourceParameters, TConstArrayView<FRHIShaderParameterResource> InBindlessParameters)
 {
+	UpdateBindlessHandles(InBindlessParameters);
+
 	UE::RHICore::RHISetShaderParametersShared(
 		*this
 		, Shader
diff --git a/Engine/Source/Runtime/VulkanRHI/Private/VulkanUniformBuffer.cpp b/Engine/Source/Runtime/VulkanRHI/Private/VulkanUniformBuffer.cpp
--- a/Engine/Source/Runtime/VulkanRHI/Private/VulkanUniformBuffer.cpp
+++ b/Engine/Source/Runtime/VulkanRHI/Private/VulkanUniformBuffer.cpp
@@ -30,9 +30,83 @@ enum
 -----------------------------------------------------------------------------*/
 
 
+static void UpdateBindlessHandles(const void* SourceData, const FRHIUniformBufferLayout* Layout)
+{
+	UE::RHICore::FUniformDataReader Reader(SourceData);
+
+	// Update views before picking up their bindless handles
+	for (const FRHIUniformBufferResource& Resource : Layout->Resources)
+	{
+		switch (Resource.MemberType)
+		{
+			case UBMT_SRV:
+			{
+				FVulkanShaderResourceView* ShaderResourceView = Reader.Read<FVulkanShaderResourceView*>(Resource);
+				if (ShaderResourceView)
+				{
+					ShaderResourceView->UpdateView();
+				}
+			}
+			break;
+			case UBMT_RDG_TEXTURE_SRV:
+			case UBMT_RDG_BUFFER_SRV:
+			{
+				FRDGShaderResourceView* RDGShaderResourceView = Reader.Read<FRDGShaderResourceView*>(Resource);
+				if (RDGShaderResourceView && RDGShaderResourceView->GetRHI())
+				{
+					FVulkanShaderResourceView* ShaderResourceView = ResourceCast(RDGShaderResourceView->GetRHI());
+					ShaderResourceView->UpdateView();
+				}
+			}
+			break;
+			case UBMT_UAV:
+			{
+				FVulkanUnorderedAccessView* UnorderedAccessView = Reader.Read<FVulkanUnorderedAccessView*>(Resource);
+				if (UnorderedAccessView)
+				{
+					UnorderedAccessView->UpdateView();
+				}
+			}
+			break;
+			case UBMT_RDG_TEXTURE_UAV:
+			case UBMT_RDG_BUFFER_UAV:
+			{
+				FRDGUnorderedAccessView* RDGUnorderedAccessView = Reader.Read<FRDGUnorderedAccessView*>(Resource);
+				if (RDGUnorderedAccessView && RDGUnorderedAccessView->GetRHI())
+				{
+					FVulkanUnorderedAccessView* UnorderedAccessView = ResourceCast(RDGUnorderedAccessView->GetRHI());
+					UnorderedAccessView->UpdateView();
+				}
+			}
+			break;
+			case UBMT_SAMPLER:
+			case UBMT_TEXTURE:
+			case UBMT_RDG_TEXTURE:
+			case UBMT_NESTED_STRUCT:
+			case UBMT_INCLUDED_STRUCT:
+			case UBMT_REFERENCED_STRUCT:
+			{
+				// Do nothing
+			}
+			break;
+
+			default:
+			//checkf(false, TEXT("Unhandled resource type?"));
+			break;
+		}
+	}
+}
+
 static void UpdateUniformBufferConstants(FVulkanDevice* Device, void* DestinationData, const void* SourceData, const FRHIUniformBufferLayout* Layout)
 {
-	UE::RHICore::UpdateUniformBufferConstants(DestinationData, SourceData, *Layout, Device->SupportsBindless());
+	const bool bSupportsBindless = Device->SupportsBindless();
+	if (bSupportsBindless)
+	{
+		// Update the bindless handles in views before we use them
+		UpdateBindlessHandles(SourceData, Layout);
+	}
+
+	UE::RHICore::UpdateUniformBufferConstants(DestinationData, SourceData, *Layout, bSupportsBindless);
 }
 
 static inline EBufferUsageFlags UniformBufferToBufferUsage(EUniformBufferUsage Usage)
@@ -57,7 +131,7 @@ static bool UseRingBuffer(EUniformBufferUsage Usage)
 	return (Usage == UniformBuffer_SingleDraw || Usage == UniformBuffer_SingleFrame);
 }
 
-static void UpdateUniformBufferHelper(FVulkanCommandListContext& Context, FVulkanUniformBuffer* VulkanUniformBuffer, const void* Data)
+static void UpdateUniformBufferHelper(FVulkanCommandListContext& Context, FVulkanUniformBuffer* VulkanUniformBuffer, const void* Data, bool bUpdateConstants = true)
 {
 	FVulkanCmdBuffer* CmdBuffer = Context.GetCommandBufferManager()->GetActiveCmdBufferDirect();
 	
@@ -81,7 +155,10 @@ static void UpdateUniformBufferHelper(FVulkanCommandListContext& Context, FVulka
 			RingBufferAllocation.HandleId);
 		
 		uint8* UploadLocation = UniformBufferUploader->GetCPUMappedPointer() + RingBufferOffset;
-		UpdateUniformBufferConstants(Context.GetDevice(), UploadLocation, Data, VulkanUniformBuffer->GetLayoutPtr());
+		if (bUpdateConstants)
+		{
+			UpdateUniformBufferConstants(Context.GetDevice(), UploadLocation, Data, VulkanUniformBuffer->GetLayoutPtr());
+		}
 	}
 	else
 	{
@@ -89,7 +166,10 @@ static void UpdateUniformBufferHelper(FVulkanCommandListContext& Context, FVulka
 
 		VulkanRHI::FTempFrameAllocationBuffer::FTempAllocInfo LockInfo;
 		Context.GetTempFrameAllocationBuffer().Alloc(DataSize, 16, LockInfo);
-		UpdateUniformBufferConstants(Context.GetDevice(), LockInfo.Data, Data, VulkanUniformBuffer->GetLayoutPtr());
+		if (bUpdateConstants)
+		{
+			UpdateUniformBufferConstants(Context.GetDevice(), LockInfo.Data, Data, VulkanUniformBuffer->GetLayoutPtr());
+		}
 
 		VkBufferCopy Region;
 		Region.size = DataSize;
@@ -156,12 +236,12 @@ FVulkanUniformBuffer::FVulkanUniformBuffer(FVulkanDevice& InDevice, const FRHIUn
 				else
 				{
 					void* CmdListConstantBufferData = RHICmdList.Alloc(DataSize, 16);
-					FMemory::Memcpy(CmdListConstantBufferData, Contents, DataSize);
+					UpdateUniformBufferConstants(Device, CmdListConstantBufferData, Contents, InLayout);
 
 					RHICmdList.EnqueueLambda([UniformBuffer, DataSize, CmdListConstantBufferData](FRHICommandList& CmdList)
 					{
 						FVulkanCommandListContext& Context = FVulkanCommandListContext::GetVulkanContext(CmdList.GetContext());
-						UpdateUniformBufferHelper(Context, UniformBuffer, CmdListConstantBufferData);
+						UpdateUniformBufferHelper(Context, UniformBuffer, CmdListConstantBufferData, false);
 					});
 
 					RHICmdList.RHIThreadFence(true);
