VulkanRHI: Overload RHIUpdateTextureReference when using bindless to copy descriptor into a reserved slot when we update a texture reference (avoids the need to update uniform buffers).  (Vulkan Bindless WIP)

#rnx

diff --git a/Engine/Source/Runtime/VulkanRHI/Private/VulkanResources.h b/Engine/Source/Runtime/VulkanRHI/Private/VulkanResources.h
--- a/Engine/Source/Runtime/VulkanRHI/Private/VulkanResources.h
+++ b/Engine/Source/Runtime/VulkanRHI/Private/VulkanResources.h
@@ -419,6 +419,8 @@ public:
 		return DefaultBindlessHandle;
 	}
 
+	void CopyBindlessHandle(FVulkanCommandListContext& Context, FRHIDescriptorHandle DestHandle);
+
 	inline static FVulkanTexture* Cast(FRHITexture* Texture)
 	{
 		check(Texture);
diff --git a/Engine/Source/Runtime/VulkanRHI/Private/VulkanTexture.cpp b/Engine/Source/Runtime/VulkanRHI/Private/VulkanTexture.cpp
--- a/Engine/Source/Runtime/VulkanRHI/Private/VulkanTexture.cpp
+++ b/Engine/Source/Runtime/VulkanRHI/Private/VulkanTexture.cpp
@@ -872,6 +872,18 @@ void FVulkanTexture::SetInitialImageState(FVulkanCommandListContext& Context, Vk
 	CmdBuffer->GetLayoutManager().SetFullLayout(*this, InitialLayout);
 }
 
+void FVulkanTexture::CopyBindlessHandle(FVulkanCommandListContext& Context, FRHIDescriptorHandle DestHandle)
+{
+	FVulkanCmdBuffer* CmdBuffer = Context.GetCommandBufferManager()->GetActiveCmdBuffer();
+
+	// Copy the our descriptor in the old texture's slot
+	checkSlow(DefaultBindlessHandle.IsValid());
+	checkSlow(DestHandle.IsValid());
+	Device->GetBindlessDescriptorManager()->CopyDescriptor(CmdBuffer->GetHandle(), DestHandle, DefaultBindlessHandle);
+}
+
+
+
 /*-----------------------------------------------------------------------------
 	Texture allocator support.
 -----------------------------------------------------------------------------*/
@@ -910,6 +922,21 @@ uint32 FVulkanDynamicRHI::RHIComputeMemorySize(FRHITexture* TextureRHI)
 	return FVulkanTexture::Cast(TextureRHI)->GetMemorySize();
 }
 
+void FVulkanDynamicRHI::RHIUpdateTextureReference(FRHITextureReference* TextureRef, FRHITexture* NewTexture)
+{
+	if (Device->SupportsBindless())
+	{
+		if (TextureRef->GetReferencedTexture() && NewTexture)
+		{
+			// Continue using the old slot for the new texture, this way indexes stored in buffers are still valid
+			FVulkanTexture* NextTexture = ResourceCast(NewTexture);
+			NextTexture->CopyBindlessHandle(Device->GetImmediateContext(), TextureRef->GetDefaultBindlessHandle());
+		}
+	}
+
+	FDynamicRHI::RHIUpdateTextureReference(TextureRef, NewTexture);
+}
+
 /*-----------------------------------------------------------------------------
 	2D texture support.
 -----------------------------------------------------------------------------*/
diff --git a/Engine/Source/Runtime/VulkanRHI/Public/VulkanDynamicRHI.h b/Engine/Source/Runtime/VulkanRHI/Public/VulkanDynamicRHI.h
--- a/Engine/Source/Runtime/VulkanRHI/Public/VulkanDynamicRHI.h
+++ b/Engine/Source/Runtime/VulkanRHI/Public/VulkanDynamicRHI.h
@@ -122,6 +122,7 @@ public:
 	virtual FShaderResourceViewRHIRef RHICreateShaderResourceView(const FShaderResourceViewInitializer& Initializer) final override;
 	virtual void RHIUpdateShaderResourceView(FRHIShaderResourceView* SRV, FRHIBuffer* Buffer, uint32 Stride, uint8 Format) final override;
 	virtual void RHIUpdateShaderResourceView(FRHIShaderResourceView* SRV, FRHIBuffer* Buffer) final override;
+	virtual void RHIUpdateTextureReference(FRHITextureReference* TextureRef, FRHITexture* NewTexture) final override;
 	virtual FRHICalcTextureSizeResult RHICalcTexturePlatformSize(FRHITextureDesc const& Desc, uint32 FirstMipIndex) final override;
 	virtual void RHIGetTextureMemoryStats(FTextureMemoryStats& OutStats) final override;
 	virtual bool RHIGetTextureMemoryVisualizeData(FColor* TextureData, int32 SizeX, int32 SizeY, int32 Pitch, int32 PixelSize) final override;
