VulkanRHI: Add missing copy when uniform buffer constants aren't updated (introduced in 23737362).  Fix GPU crash in Vulkan editor.

#rnx

diff --git a/Engine/Source/Runtime/VulkanRHI/Private/VulkanUniformBuffer.cpp b/Engine/Source/Runtime/VulkanRHI/Private/VulkanUniformBuffer.cpp
--- a/Engine/Source/Runtime/VulkanRHI/Private/VulkanUniformBuffer.cpp
+++ b/Engine/Source/Runtime/VulkanRHI/Private/VulkanUniformBuffer.cpp
@@ -134,13 +134,28 @@ static bool UseRingBuffer(EUniformBufferUsage Usage)
 static void UpdateUniformBufferHelper(FVulkanCommandListContext& Context, FVulkanUniformBuffer* VulkanUniformBuffer, const void* Data, bool bUpdateConstants = true)
 {
 	FVulkanCmdBuffer* CmdBuffer = Context.GetCommandBufferManager()->GetActiveCmdBufferDirect();
-	
+
+	FVulkanDevice* Device = Context.GetDevice();
 	const int32 DataSize = VulkanUniformBuffer->GetLayout().ConstantBufferSize;
 
+	auto CopyUniformBufferData = [&](void* DestinationData, const void* SourceData) {
+
+		if (bUpdateConstants)
+		{
+			// Update constants as the data is copied
+			UpdateUniformBufferConstants(Device, DestinationData, SourceData, VulkanUniformBuffer->GetLayoutPtr());
+		}
+		else
+		{
+			// Don't touch constant, copy the data as-is
+			FMemory::Memcpy(DestinationData, SourceData, DataSize);
+		}
+	};
+
 	if (UseRingBuffer(VulkanUniformBuffer->Usage))
 	{
 		FVulkanUniformBufferUploader* UniformBufferUploader = Context.GetUniformBufferUploader();
-		const VkDeviceSize UBOffsetAlignment = Context.GetDevice()->GetLimits().minUniformBufferOffsetAlignment;
+		const VkDeviceSize UBOffsetAlignment = Device->GetLimits().minUniformBufferOffsetAlignment;
 		const FVulkanAllocation& RingBufferAllocation = UniformBufferUploader->GetCPUBufferAllocation();
 		uint64 RingBufferOffset = UniformBufferUploader->AllocateMemory(DataSize, UBOffsetAlignment, CmdBuffer);
 
@@ -155,10 +170,7 @@ static void UpdateUniformBufferHelper(FVulkanCommandListContext& Context, FVulka
 			RingBufferAllocation.HandleId);
 		
 		uint8* UploadLocation = UniformBufferUploader->GetCPUMappedPointer() + RingBufferOffset;
-		if (bUpdateConstants)
-		{
-			UpdateUniformBufferConstants(Context.GetDevice(), UploadLocation, Data, VulkanUniformBuffer->GetLayoutPtr());
-		}
+		CopyUniformBufferData(UploadLocation, Data);
 	}
 	else
 	{
@@ -166,10 +178,7 @@ static void UpdateUniformBufferHelper(FVulkanCommandListContext& Context, FVulka
 
 		VulkanRHI::FTempFrameAllocationBuffer::FTempAllocInfo LockInfo;
 		Context.GetTempFrameAllocationBuffer().Alloc(DataSize, 16, LockInfo);
-		if (bUpdateConstants)
-		{
-			UpdateUniformBufferConstants(Context.GetDevice(), LockInfo.Data, Data, VulkanUniformBuffer->GetLayoutPtr());
-		}
+		CopyUniformBufferData(LockInfo.Data, Data);
 
 		VkBufferCopy Region;
 		Region.size = DataSize;
