VulkanRHI: Fix security issue regarding loading of vulkan pipeline cache.

#rnx

diff --git a/Engine/Source/Runtime/VulkanRHI/Private/VulkanPipeline.cpp b/Engine/Source/Runtime/VulkanRHI/Private/VulkanPipeline.cpp
--- a/Engine/Source/Runtime/VulkanRHI/Private/VulkanPipeline.cpp
+++ b/Engine/Source/Runtime/VulkanRHI/Private/VulkanPipeline.cpp
@@ -2203,16 +2203,23 @@ void FVulkanPipelineStateCacheManager::NotifyDeletedComputePipeline(FVulkanCompu
 }
 
 template<typename T>
-static void SerializeArray(FArchive& Ar, TArray<T>& Array)
+static bool SerializeArray(FArchive& Ar, TArray<T>& Array)
 {
 	int32 Num = Array.Num();
 	Ar << Num;
 	if (Ar.IsLoading())
 	{
-		Array.SetNum(Num);
-		for (int32 Index = 0; Index < Num; ++Index)
+		if (Num < 0)
 		{
-			Ar << Array[Index];
+			return false;
+		}
+		else
+		{
+			Array.SetNum(Num);
+			for (int32 Index = 0; Index < Num; ++Index)
+			{
+				Ar << Array[Index];
+			}
 		}
 	}
 	else
@@ -2222,6 +2229,7 @@ static void SerializeArray(FArchive& Ar, TArray<T>& Array)
 			Ar << Array[Index];
 		}
 	}
+	return true;
 }
 
 
@@ -2251,7 +2259,12 @@ bool FVulkanPipelineStateCacheManager::FVulkanLRUCacheFile::Load(FArchive& Ar)
 		return false;
 	}
 
-	SerializeArray(Ar, PipelineSizes);
+	if (!SerializeArray(Ar, PipelineSizes))
+	{
+		UE_LOG(LogVulkanRHI, Warning, TEXT("Unable to load lru pipeline cache due to invalid archive data!"));
+		return false;
+	}
+
 	return true;
 }
 
