VulkanRHI:  Vulkan RHICopyTexture was not using using SourcePosition, causing an engine test failure on 3D texture copies when z offsets were used.

diff --git a/Engine/Source/Runtime/VulkanRHI/Private/VulkanTexture.cpp b/Engine/Source/Runtime/VulkanRHI/Private/VulkanTexture.cpp
--- a/Engine/Source/Runtime/VulkanRHI/Private/VulkanTexture.cpp
+++ b/Engine/Source/Runtime/VulkanRHI/Private/VulkanTexture.cpp
@@ -2385,9 +2385,10 @@ void FVulkanCommandListContext::RHICopyTexture(FRHITexture* SourceTexture, FRHIT
 	check(!EnumHasAnyFlags(Source->GetDesc().Flags, TexCreate_CPUReadback));
 	if (EnumHasAllFlags(Dest->GetDesc().Flags, TexCreate_CPUReadback))
 	{
-		check(CopyInfo.DestSliceIndex == 0); //slices not supported in TexCreate_CPUReadback textures.
+		checkf(CopyInfo.DestSliceIndex == 0, TEXT("Slices not supported in TexCreate_CPUReadback textures"));
+		checkf(CopyInfo.DestPosition.IsZero(), TEXT("Destination position not supported in TexCreate_CPUReadback textures"));
 		FIntVector Size = CopyInfo.Size;
-		if(Size == FIntVector::ZeroValue)
+		if (Size == FIntVector::ZeroValue)
 		{
 			ensure(SourceXYZ.X <= DestXYZ.X && SourceXYZ.Y <= DestXYZ.Y);
 			Size.X = FMath::Max<uint32>(1u, SourceXYZ.X >> CopyInfo.SourceMipIndex);
@@ -2410,6 +2411,9 @@ void FVulkanCommandListContext::RHICopyTexture(FRHITexture* SourceTexture, FRHIT
 			CopyRegion[Index].imageSubresource.mipLevel = SourceMipIndex;
 			CopyRegion[Index].imageSubresource.baseArrayLayer = SourceSliceIndex;
 			CopyRegion[Index].imageSubresource.layerCount = 1;
+			CopyRegion[Index].imageOffset.x = CopyInfo.SourcePosition.X;
+			CopyRegion[Index].imageOffset.y = CopyInfo.SourcePosition.Y;
+			CopyRegion[Index].imageOffset.z = CopyInfo.SourcePosition.Z;
 			CopyRegion[Index].imageExtent.width = Size.X;
 			CopyRegion[Index].imageExtent.height = Size.Y;
 			CopyRegion[Index].imageExtent.depth = Size.Z;
