Added WITH_AMD_AGS guard around includes.

diff --git a/Engine/Source/Runtime/VulkanRHI/Private/Windows/VulkanWindowsPlatform.cpp b/Engine/Source/Runtime/VulkanRHI/Private/Windows/VulkanWindowsPlatform.cpp
--- a/Engine/Source/Runtime/VulkanRHI/Private/Windows/VulkanWindowsPlatform.cpp
+++ b/Engine/Source/Runtime/VulkanRHI/Private/Windows/VulkanWindowsPlatform.cpp
@@ -7,10 +7,12 @@
 #include "../VulkanExtensions.h"
 
 // Disable warning about forward declared enumeration without a type, since the D3D specific enums are not used in this translation unit
+#if WITH_AMD_AGS
 #pragma warning(push)
 #pragma warning(disable : 4471)
 #include "amd_ags.h"
 #pragma warning(pop)
+#endif
 
 #include "Windows/AllowWindowsPlatformTypes.h"
 static HMODULE GVulkanDLLModule = nullptr;
diff --git a/Engine/Source/Runtime/Windows/D3D11RHI/Private/D3D11Device.cpp b/Engine/Source/Runtime/Windows/D3D11RHI/Private/D3D11Device.cpp
--- a/Engine/Source/Runtime/Windows/D3D11RHI/Private/D3D11Device.cpp
+++ b/Engine/Source/Runtime/Windows/D3D11RHI/Private/D3D11Device.cpp
@@ -11,7 +11,7 @@
 #include "Modules/ModuleManager.h"
 #include "Windows/AllowWindowsPlatformTypes.h"
 	#include <delayimp.h>
-	#if !PLATFORM_HOLOLENS
+	#if WITH_AMD_AGS
 	#include "amd_ags.h"
 	#endif
 #include "Windows/HideWindowsPlatformTypes.h"
