VulkanRHI:  Match other RHI and skip shader check when the SRV/UAV being set is null (the RHI will do nothing anyhow).

#rnx

diff --git a/Engine/Source/Runtime/VulkanRHI/Private/VulkanCommands.cpp b/Engine/Source/Runtime/VulkanRHI/Private/VulkanCommands.cpp
--- a/Engine/Source/Runtime/VulkanRHI/Private/VulkanCommands.cpp
+++ b/Engine/Source/Runtime/VulkanRHI/Private/VulkanCommands.cpp
@@ -176,7 +176,7 @@ void FVulkanCommandListContext::RHISetUAVParameter(FRHIPixelShader* PixelShaderR
 
 void FVulkanCommandListContext::RHISetUAVParameter(FRHIComputeShader* ComputeShaderRHI, uint32 UAVIndex, FRHIUnorderedAccessView* UAVRHI)
 {
-	check(PendingComputeState->GetCurrentShader() == ResourceCast(ComputeShaderRHI));
+	check((UAVRHI == nullptr) || (PendingComputeState->GetCurrentShader() == ResourceCast(ComputeShaderRHI)));
 
 	FVulkanUnorderedAccessView* UAV = ResourceCast(UAVRHI);
 	PendingComputeState->SetUAVForStage(UAVIndex, UAV);
@@ -184,9 +184,6 @@ void FVulkanCommandListContext::RHISetUAVParameter(FRHIComputeShader* ComputeSha
 
 void FVulkanCommandListContext::RHISetUAVParameter(FRHIComputeShader* ComputeShaderRHI,uint32 UAVIndex, FRHIUnorderedAccessView* UAVRHI, uint32 InitialCount)
 {
-	check(PendingComputeState->GetCurrentShader() == ResourceCast(ComputeShaderRHI));
-
-	FVulkanUnorderedAccessView* UAV = ResourceCast(UAVRHI);
 	ensure(0);
 }
 
@@ -221,7 +218,7 @@ void FVulkanCommandListContext::RHISetShaderResourceViewParameter(FRHIGraphicsSh
 
 void FVulkanCommandListContext::RHISetShaderResourceViewParameter(FRHIComputeShader* ComputeShaderRHI,uint32 TextureIndex, FRHIShaderResourceView* SRVRHI)
 {
-	check(PendingComputeState->GetCurrentShader() == ResourceCast(ComputeShaderRHI));
+	check((SRVRHI == nullptr) || (PendingComputeState->GetCurrentShader() == ResourceCast(ComputeShaderRHI)));
 
 	FVulkanShaderResourceView* SRV = ResourceCast(SRVRHI);
 	PendingComputeState->SetSRVForStage(TextureIndex, SRV);
