VulkanRHI: Move GetBufferDeviceAddressKHR out of Raytracing only functions in VulkanLoader.h.

#rnx

diff --git a/Engine/Source/Runtime/VulkanRHI/Private/VulkanCommandWrappers.h b/Engine/Source/Runtime/VulkanRHI/Private/VulkanCommandWrappers.h
--- a/Engine/Source/Runtime/VulkanRHI/Private/VulkanCommandWrappers.h
+++ b/Engine/Source/Runtime/VulkanRHI/Private/VulkanCommandWrappers.h
@@ -195,10 +195,10 @@ struct FWrapLayer
 	static void CmdTraceRaysKHR(VkResult Result, VkCommandBuffer CommandBuffer, const VkStridedDeviceAddressRegionKHR* RaygenShaderBindingTable, const VkStridedDeviceAddressRegionKHR* MissShaderBindingTable, const VkStridedDeviceAddressRegionKHR* HitShaderBindingTable, const VkStridedDeviceAddressRegionKHR* CallableShaderBindingTable, uint32 width, uint32 height, uint32 depth);
 	static void CreateRayTracingPipelinesKHR(VkResult Result, VkDevice Device, VkDeferredOperationKHR DeferredOperation, VkPipelineCache PipelineCache, uint32_t CreateInfoCount, const VkRayTracingPipelineCreateInfoKHR* CreateInfos, const VkAllocationCallbacks* Allocator, VkPipeline* Pipelines);
 	static void GetRayTracingShaderGroupHandlesKHR(VkResult Result, VkDevice Device, VkPipeline Pipeline, uint32_t FirstGroup, uint32_t GroupCount, size_t DataSize, void* Data);
-	static void GetBufferDeviceAddressKHR(VkResult Result, VkDevice Device, const VkBufferDeviceAddressInfo* Info);
 	static void CmdWriteAccelerationStructuresPropertiesKHR(VkResult Result, VkCommandBuffer CommandBuffer, uint32_t AccelerationStructureCount, const VkAccelerationStructureKHR* AccelerationStructures, VkQueryType QueryType, VkQueryPool QueryPool, uint32_t FirstQuery);
 	static void CmdCopyAccelerationStructureKHR(VkResult Result, VkCommandBuffer CommandBuffer, const VkCopyAccelerationStructureInfoKHR* Info);
 #endif
+	static void GetBufferDeviceAddressKHR(VkResult Result, VkDevice Device, const VkBufferDeviceAddressInfo* Info) VULKAN_LAYER_BODY
 	static void GetDeviceImageMemoryRequirementsKHR(VkResult Result, VkDevice Device, const VkDeviceImageMemoryRequirements* Info, VkMemoryRequirements2* MemoryRequirements) VULKAN_LAYER_BODY
 	static void GetDeviceBufferMemoryRequirementsKHR(VkResult Result, VkDevice Device, const VkDeviceBufferMemoryRequirements* Info, VkMemoryRequirements2* MemoryRequirements) VULKAN_LAYER_BODY
 	static void ResetQueryPoolEXT(VkResult Result, VkDevice Device, VkQueryPool QueryPool, uint32_t FirstQuery, uint32_t QueryCount) VULKAN_LAYER_BODY
@@ -1476,14 +1476,6 @@ namespace VulkanRHI
 		return Result;
 	}
 
-	static FORCEINLINE_DEBUGGABLE VkDeviceAddress vkGetBufferDeviceAddressKHR(VkDevice Device, const VkBufferDeviceAddressInfo* Info)
-	{
-		FWrapLayer::GetBufferDeviceAddressKHR(VK_RESULT_MAX_ENUM, Device, Info);
-		const VkDeviceAddress Result = VULKANAPINAMESPACE::vkGetBufferDeviceAddressKHR(Device, Info);
-		FWrapLayer::GetBufferDeviceAddressKHR(VK_SUCCESS, Device, Info);
-		return Result;
-	}
-
 	static FORCEINLINE_DEBUGGABLE void vkCmdWriteAccelerationStructuresPropertiesKHR(VkCommandBuffer CommandBuffer,	uint32_t AccelerationStructureCount, const VkAccelerationStructureKHR* AccelerationStructures, VkQueryType QueryType, VkQueryPool QueryPool, uint32_t FirstQuery)
 	{
 		FWrapLayer::CmdWriteAccelerationStructuresPropertiesKHR(VK_RESULT_MAX_ENUM, CommandBuffer, AccelerationStructureCount, AccelerationStructures, QueryType, QueryPool, FirstQuery);
@@ -1499,6 +1491,14 @@ namespace VulkanRHI
 	}
 #endif
 
+	static FORCEINLINE_DEBUGGABLE VkDeviceAddress vkGetBufferDeviceAddressKHR(VkDevice Device, const VkBufferDeviceAddressInfo* Info)
+	{
+		FWrapLayer::GetBufferDeviceAddressKHR(VK_RESULT_MAX_ENUM, Device, Info);
+		const VkDeviceAddress Result = VULKANAPINAMESPACE::vkGetBufferDeviceAddressKHR(Device, Info);
+		FWrapLayer::GetBufferDeviceAddressKHR(VK_SUCCESS, Device, Info);
+		return Result;
+	}
+
 	static FORCEINLINE_DEBUGGABLE void vkGetDeviceImageMemoryRequirementsKHR(VkDevice Device, const VkDeviceImageMemoryRequirements* Info, VkMemoryRequirements2* MemoryRequirements)
 	{
 		FWrapLayer::GetDeviceImageMemoryRequirementsKHR(VK_RESULT_MAX_ENUM, Device, Info, MemoryRequirements);
diff --git a/Engine/Source/Runtime/VulkanRHI/Private/VulkanDebug.cpp b/Engine/Source/Runtime/VulkanRHI/Private/VulkanDebug.cpp
--- a/Engine/Source/Runtime/VulkanRHI/Private/VulkanDebug.cpp
+++ b/Engine/Source/Runtime/VulkanRHI/Private/VulkanDebug.cpp
@@ -4368,12 +4368,12 @@ void FWrapLayer::GetRayTracingShaderGroupHandlesKHR(VkResult Result, VkDevice De
 	}
 }
 
-void FWrapLayer::GetBufferDeviceAddressKHR(VkResult Result, VkDevice Device, const VkBufferDeviceAddressInfo* Info)
+void FWrapLayer::CmdWriteAccelerationStructuresPropertiesKHR(VkResult Result, VkCommandBuffer CommandBuffer, uint32_t AccelerationStructureCount, const VkAccelerationStructureKHR* AccelerationStructures, VkQueryType QueryType, VkQueryPool QueryPool, uint32_t FirstQuery)
 {
 	if (Result == VK_RESULT_MAX_ENUM)
 	{
 #if VULKAN_ENABLE_DUMP_LAYER
-		PrintfBeginResult(FString::Printf(TEXT("vkGetBufferDeviceAddressKHR(Device=0x%p, Buffer=0x%p)"), Device, Info->buffer));
+		PrintfBeginResult(FString::Printf(TEXT("vkCmdWriteAccelerationStructuresPropertiesKHR")));
 #endif
 	}
 	else
@@ -4384,12 +4384,12 @@ void FWrapLayer::GetBufferDeviceAddressKHR(VkResult Result, VkDevice Device, con
 	}
 }
 
-void FWrapLayer::CmdWriteAccelerationStructuresPropertiesKHR(VkResult Result, VkCommandBuffer CommandBuffer, uint32_t AccelerationStructureCount, const VkAccelerationStructureKHR* AccelerationStructures, VkQueryType QueryType, VkQueryPool QueryPool, uint32_t FirstQuery)
+void FWrapLayer::CmdCopyAccelerationStructureKHR(VkResult Result, VkCommandBuffer CommandBuffer, const VkCopyAccelerationStructureInfoKHR* Info)
 {
 	if (Result == VK_RESULT_MAX_ENUM)
 	{
 #if VULKAN_ENABLE_DUMP_LAYER
-		PrintfBeginResult(FString::Printf(TEXT("vkCmdWriteAccelerationStructuresPropertiesKHR")));
+		PrintfBeginResult(FString::Printf(TEXT("vkCmdCopyAccelerationStructureKHR")));
 #endif
 	}
 	else
@@ -4399,13 +4399,14 @@ void FWrapLayer::CmdWriteAccelerationStructuresPropertiesKHR(VkResult Result, Vk
 #endif
 	}
 }
+#endif // VULKAN_RHI_RAYTRACING
 
-void FWrapLayer::CmdCopyAccelerationStructureKHR(VkResult Result, VkCommandBuffer CommandBuffer, const VkCopyAccelerationStructureInfoKHR* Info)
+void FWrapLayer::GetBufferDeviceAddressKHR(VkResult Result, VkDevice Device, const VkBufferDeviceAddressInfo* Info)
 {
 	if (Result == VK_RESULT_MAX_ENUM)
 	{
 #if VULKAN_ENABLE_DUMP_LAYER
-		PrintfBeginResult(FString::Printf(TEXT("vkCmdCopyAccelerationStructureKHR")));
+		PrintfBeginResult(FString::Printf(TEXT("vkGetBufferDeviceAddressKHR(Device=0x%p, Buffer=0x%p)"), Device, Info->buffer));
 #endif
 	}
 	else
@@ -4415,7 +4416,6 @@ void FWrapLayer::CmdCopyAccelerationStructureKHR(VkResult Result, VkCommandBuffe
 #endif
 	}
 }
-#endif // VULKAN_RHI_RAYTRACING
 
 void FWrapLayer::GetDeviceImageMemoryRequirementsKHR(VkResult Result, VkDevice Device, const VkDeviceImageMemoryRequirements* Info, VkMemoryRequirements2* MemoryRequirements)
 {
