VulkanRHI:  Cache buffer address for the ring buffer used for packed uniform buffers to reduce the amount of calls to vkGetBufferDeviceAddressKHR in bindless.  (Vulkan Bindless WIP)

#rnx

diff --git a/Engine/Source/Runtime/VulkanRHI/Private/VulkanResources.h b/Engine/Source/Runtime/VulkanRHI/Private/VulkanResources.h
--- a/Engine/Source/Runtime/VulkanRHI/Private/VulkanResources.h
+++ b/Engine/Source/Runtime/VulkanRHI/Private/VulkanResources.h
@@ -821,6 +821,11 @@ public:
 		return Allocation.Offset;
 	}
 
+	inline VkDeviceAddress GetBufferAddress() const
+	{
+		return BufferAddress;
+	}
+
 	inline VkBuffer GetHandle() const
 	{
 		return Allocation.GetBufferHandle();
@@ -831,12 +836,12 @@ public:
 		return Allocation.GetMappedPointer(Device);
 	}
 
-	VulkanRHI::FVulkanAllocation& GetAllocation()
+	inline VulkanRHI::FVulkanAllocation& GetAllocation()
 	{
 		return Allocation;
 	}
 
-	const VulkanRHI::FVulkanAllocation& GetAllocation() const
+	inline const VulkanRHI::FVulkanAllocation& GetAllocation() const
 	{
 		return Allocation;
 	}
@@ -845,6 +850,7 @@ public:
 protected:
 	uint64 BufferSize;
 	uint64 BufferOffset;
+	VkDeviceAddress BufferAddress;
 	uint32 MinAlignment;
 	VulkanRHI::FVulkanAllocation Allocation;
 
@@ -861,22 +867,22 @@ public:
 	FVulkanUniformBufferUploader(FVulkanDevice* InDevice);
 	~FVulkanUniformBufferUploader();
 
-	uint8* GetCPUMappedPointer()
+	inline uint8* GetCPUMappedPointer()
 	{
 		return (uint8*)CPUBuffer->GetMappedPointer();
 	}
 
-	uint64 AllocateMemory(uint64 Size, uint32 Alignment, FVulkanCmdBuffer* InCmdBuffer)
+	inline uint64 AllocateMemory(uint64 Size, uint32 Alignment, FVulkanCmdBuffer* InCmdBuffer)
 	{
 		return CPUBuffer->AllocateMemory(Size, Alignment, InCmdBuffer);
 	}
 
-	const VulkanRHI::FVulkanAllocation& GetCPUBufferAllocation() const
+	inline const VulkanRHI::FVulkanAllocation& GetCPUBufferAllocation() const
 	{
 		return CPUBuffer->GetAllocation();
 	}
 
-	VkBuffer GetCPUBufferHandle() const
+	inline VkBuffer GetCPUBufferHandle() const
 	{
 		return CPUBuffer->GetHandle();
 	}
@@ -886,6 +892,11 @@ public:
 		return CPUBuffer->GetBufferOffset();
 	}
 
+	inline VkDeviceAddress GetCPUBufferAddress() const
+	{
+		return CPUBuffer->GetBufferAddress();
+	}
+
 protected:
 	FVulkanRingBuffer* CPUBuffer;
 	friend class FVulkanCommandListContext;
diff --git a/Engine/Source/Runtime/VulkanRHI/Private/VulkanUniformBuffer.cpp b/Engine/Source/Runtime/VulkanRHI/Private/VulkanUniformBuffer.cpp
--- a/Engine/Source/Runtime/VulkanRHI/Private/VulkanUniformBuffer.cpp
+++ b/Engine/Source/Runtime/VulkanRHI/Private/VulkanUniformBuffer.cpp
@@ -435,6 +435,7 @@ FVulkanRingBuffer::FVulkanRingBuffer(FVulkanDevice* InDevice, uint64 TotalSize,
 	: VulkanRHI::FDeviceChild(InDevice)
 	, BufferSize(TotalSize)
 	, BufferOffset(0)
+	, BufferAddress(0)
 	, MinAlignment(0)
 {
 	check(TotalSize <= (uint64)MAX_uint32);
@@ -442,6 +443,14 @@ FVulkanRingBuffer::FVulkanRingBuffer(FVulkanDevice* InDevice, uint64 TotalSize,
 	MinAlignment = Allocation.GetBufferAlignment(Device);
 	// Start by wrapping around to set up the correct fence
 	BufferOffset = TotalSize;
+
+	if (InDevice->GetOptionalExtensions().HasBufferDeviceAddress)
+	{
+		VkBufferDeviceAddressInfo BufferInfo;
+		ZeroVulkanStruct(BufferInfo, VK_STRUCTURE_TYPE_BUFFER_DEVICE_ADDRESS_INFO);
+		BufferInfo.buffer = Allocation.GetBufferHandle();
+		BufferAddress = VulkanRHI::vkGetBufferDeviceAddressKHR(Device->GetInstanceHandle(), &BufferInfo);
+	}
 }
 
 FVulkanRingBuffer::~FVulkanRingBuffer()
