VulkanRHI: Set uniform buffers for all stages at once so all offsets can be set in a single call.  Allows for more debug modes.  (Vulkan Bindless WIP)

#rnx

diff --git a/Engine/Source/Runtime/VulkanRHI/Private/VulkanDescriptorSets.cpp b/Engine/Source/Runtime/VulkanRHI/Private/VulkanDescriptorSets.cpp
--- a/Engine/Source/Runtime/VulkanRHI/Private/VulkanDescriptorSets.cpp
+++ b/Engine/Source/Runtime/VulkanRHI/Private/VulkanDescriptorSets.cpp
@@ -8,6 +8,14 @@
 #include "VulkanDescriptorSets.h"
 
 
+int32 GVulkanBindlessEnabled = 0;
+static FAutoConsoleVariableRef CVarVulkanBindlessEnabled(
+	TEXT("r.Vulkan.Bindless.Enabled"),
+	GVulkanBindlessEnabled,
+	TEXT("Enable the use of bindless if all conditions are met to support it"),
+	ECVF_ReadOnly
+);
+
 int32 GVulkanBindlessMaxSamplerDescriptorCount = 2048;
 static FAutoConsoleVariableRef CVarVulkanBindlessMaxSamplerDescriptorCount(
 	TEXT("r.Vulkan.Bindless.MaxSamplerDescriptorCount"),
@@ -81,6 +89,15 @@ static FAutoConsoleVariableRef CVarVulkanBindlessMaxUniformBuffersPerStage(
 	ECVF_ReadOnly
 );
 
+int32 GVulkanBindlessRebindBuffers = 0;
+static FAutoConsoleVariableRef CVarVulkanBindlessRebindBuffers(
+	TEXT("r.Vulkan.Bindless.RebindBuffers"),
+	GVulkanBindlessRebindBuffers,
+	TEXT("Rebind buffers for every draw or dispatch.  Handy for debugging but not great for performance."),
+	ECVF_RenderThreadSafe
+);
+
+
 
 static inline uint32 GetInitialDescriptorCount(VkDescriptorType DescriptorType)
 {
@@ -121,11 +138,20 @@ FVulkanBindlessDescriptorManager::FVulkanBindlessDescriptorManager(FVulkanDevice
 	: VulkanRHI::FDeviceChild(InDevice)
 	, bIsSupported(VerifySupport(InDevice))
 {
-	const bool bFullyDisabled = (RHIGetBindlessResourcesConfiguration(GMaxRHIShaderPlatform) == ERHIBindlessConfiguration::Disabled) &&
-		(RHIGetBindlessSamplersConfiguration(GMaxRHIShaderPlatform) == ERHIBindlessConfiguration::Disabled);
-	const bool bFullyEnabled = (RHIGetBindlessResourcesConfiguration(GMaxRHIShaderPlatform) == ERHIBindlessConfiguration::AllShaders) &&
-		(RHIGetBindlessSamplersConfiguration(GMaxRHIShaderPlatform) == ERHIBindlessConfiguration::AllShaders);
-	checkf(bFullyDisabled || bFullyEnabled, TEXT("Bindless must be fully enabled or fully disabled in Vulkan because of the way it uses descriptor buffers."));
+	FMemory::Memzero(BufferBindingInfo);
+
+	// Setup up the buffer indices according to the constants currently used in code to end up with 
+	// a different index for each type of bindless resource (up to VulkanBindless::NumBindlessSets)
+	// and then all the uniform buffers in the same descriptor buffer at the end
+	// For example:   BufferIndices[] = { 0,1,2,3,4,5,6,7,7,7,7,7,7,7,7,7,7 };
+	for (uint32 Index = 0; Index < VulkanBindless::NumBindlessSets; Index++)
+	{
+		BufferIndices[Index] = Index;
+	}
+	for (uint32 Index = VulkanBindless::NumBindlessSets; Index < VulkanBindless::MaxNumSets; Index++)
+	{
+		BufferIndices[Index] = VulkanBindless::NumBindlessSets;
+	}
 }
 
 FVulkanBindlessDescriptorManager::~FVulkanBindlessDescriptorManager()
@@ -150,7 +176,7 @@ void FVulkanBindlessDescriptorManager::BindDescriptorBuffers(VkCommandBuffer Com
 {
 }
 
-void FVulkanBindlessDescriptorManager::RegisterUniformBuffers(VkCommandBuffer CommandBuffer, VkPipelineBindPoint BindPoint, ShaderStage::EStage Stage, const FVulkanDescriptorSetWriter& SetWriter)
+void FVulkanBindlessDescriptorManager::RegisterUniformBuffers(VkCommandBuffer CommandBuffer, VkPipelineBindPoint BindPoint, const FUniformBufferDescriptorArrays& StageUBs)
 {
 }
 
diff --git a/Engine/Source/Runtime/VulkanRHI/Private/VulkanDescriptorSets.h b/Engine/Source/Runtime/VulkanRHI/Private/VulkanDescriptorSets.h
--- a/Engine/Source/Runtime/VulkanRHI/Private/VulkanDescriptorSets.h
+++ b/Engine/Source/Runtime/VulkanRHI/Private/VulkanDescriptorSets.h
@@ -1209,6 +1209,8 @@ public:
 	FVulkanBindlessDescriptorManager(FVulkanDevice* InDevice);
 	~FVulkanBindlessDescriptorManager();
 
+	typedef TStaticArray<TArray<VkDescriptorAddressInfoEXT>, ShaderStage::NumStages> FUniformBufferDescriptorArrays;
+
 	void Init();
 	void Deinit();
 
@@ -1230,7 +1232,7 @@ public:
 	FRHIDescriptorHandle RegisterBuffer(VkBuffer VulkanBuffer, VkDeviceSize BufferOffset, VkDeviceSize BufferSize, VkDescriptorType DescriptorType);
 	FRHIDescriptorHandle RegisterTexelBuffer(const VkBufferViewCreateInfo& ViewInfo, VkDescriptorType DescriptorType);
 	FRHIDescriptorHandle RegisterAccelerationStructure(VkAccelerationStructureKHR AccelerationStructure);
-	void RegisterUniformBuffers(VkCommandBuffer CommandBuffer, VkPipelineBindPoint BindPoint, ShaderStage::EStage Stage, const FVulkanDescriptorSetWriter& SetWriter);
+	void RegisterUniformBuffers(VkCommandBuffer CommandBuffer, VkPipelineBindPoint BindPoint, const FUniformBufferDescriptorArrays& StageUBs);
 
 	void Unregister(FRHIDescriptorHandle DescriptorHandle);
 
@@ -1263,6 +1265,9 @@ private:
 	BindlessSetState BindlessUniformBufferSetState;
 	std::atomic<uint32> UniformBlockIndex = 0;
 
+	VkDescriptorBufferBindingInfoEXT BufferBindingInfo[VulkanBindless::NumBindlessSets + 1];
+	uint32_t BufferIndices[VulkanBindless::MaxNumSets];
+
 	VkPipelineLayout BindlessPipelineLayout = VK_NULL_HANDLE;
 
 	uint32 GetFreeResourceIndex(BindlessSetState& Desc);
diff --git a/Engine/Source/Runtime/VulkanRHI/Private/VulkanPipelineState.cpp b/Engine/Source/Runtime/VulkanRHI/Private/VulkanPipelineState.cpp
--- a/Engine/Source/Runtime/VulkanRHI/Private/VulkanPipelineState.cpp
+++ b/Engine/Source/Runtime/VulkanRHI/Private/VulkanPipelineState.cpp
@@ -218,7 +218,7 @@ bool FVulkanComputePipelineDescriptorState::InternalUpdateDescriptorSets(FVulkan
 
 void FVulkanComputePipelineDescriptorState::UpdateBindlessDescriptors(FVulkanCommandListContext* CmdListContext, FVulkanCmdBuffer* CmdBuffer)
 {
-	check(Device->SupportsBindless());
+	check(bUseBindless);
 
 	// We should only have uniform buffers at this point
 	check(DSWriteContainer.DescriptorBufferInfo.Num() == DSWriteContainer.DescriptorWrites.Num());
@@ -228,29 +228,41 @@ void FVulkanComputePipelineDescriptorState::UpdateBindlessDescriptors(FVulkanCom
 	uint8* CPURingBufferBase = (uint8*)UniformBufferUploader->GetCPUMappedPointer();
 	const VkDeviceSize UBOffsetAlignment = Device->GetLimits().minUniformBufferOffsetAlignment;
 
-	if (PackedUniformBuffersDirty != 0)
+	FVulkanBindlessDescriptorManager::FUniformBufferDescriptorArrays StageUBs;
+
+	const int32 Stage = (int32)ShaderStage::EStage::Compute;
+	const FDescriptorSetRemappingInfo* RESTRICT RemappingInfo = PipelineDescriptorInfo->RemappingInfo;
+	const FDescriptorSetRemappingInfo::FStageInfo& StageInfo = RemappingInfo->StageInfos[Stage];
+
+	TArray<VkDescriptorAddressInfoEXT>& DescriptorAddressInfos = StageUBs[Stage];
+	DescriptorAddressInfos.SetNumZeroed(StageInfo.PackedUBBindingIndices.Num() + StageInfo.UniformBuffers.Num());
+
+	// PackedUniformBuffersDirty ?
+	check((PackedUniformBuffersMask == 0) || (PackedUniformBuffersMask == 1));
+	if (PackedUniformBuffersMask != 0)
 	{
-#if VULKAN_ENABLE_AGGRESSIVE_STATS
-		SCOPE_CYCLE_COUNTER(STAT_VulkanApplyPackedUniformBuffers);
-#endif
-		uint64 RemainingPackedUniformsMask = PackedUniformBuffersDirty;
+		check((StageInfo.PackedUBDescriptorSet != UINT16_MAX) && (StageInfo.PackedUBDescriptorSet <= (uint16)Stage));
+		uint64 RemainingPackedUniformsMask = PackedUniformBuffersMask;
 		int32 PackedUBIndex = 0;
 		while (RemainingPackedUniformsMask)
 		{
 			if (RemainingPackedUniformsMask & 1)
 			{
 				const FPackedUniformBuffers::FPackedBuffer& StagedUniformBuffer = PackedUniformBuffers.GetBuffer(PackedUBIndex);
-				const int32 BindingIndex = PipelineDescriptorInfo->RemappingInfo->StageInfos[0].PackedUBBindingIndices[PackedUBIndex];
-				VkDescriptorBufferInfo& BufferInfo = *const_cast<VkDescriptorBufferInfo*>(DSWriter[0].WriteDescriptors[BindingIndex].pBufferInfo);
-				BufferInfo.buffer = UniformBufferUploader->GetCPUBufferHandle();
-				BufferInfo.range = StagedUniformBuffer.Num();
-				check((BufferInfo.range & 0xF) == 0);   // :todo-jn: nv driver issue
+				const int32 UBSize = StagedUniformBuffer.Num();
+				const int32 BindingIndex = StageInfo.PackedUBBindingIndices[PackedUBIndex];
+
+				const uint64 RingBufferOffset = UniformBufferUploader->AllocateMemory(UBSize, UBOffsetAlignment, CmdBuffer);
 
-				const uint64 RingBufferOffset = UniformBufferUploader->AllocateMemory(BufferInfo.range, UBOffsetAlignment, CmdBuffer);
-				BufferInfo.offset = UniformBufferUploader->GetCPUBufferOffset() + RingBufferOffset;
+				// Make sure it wasn't written to already
+				VkDescriptorAddressInfoEXT& DescriptorAddressInfo = DescriptorAddressInfos[BindingIndex];
+				check(DescriptorAddressInfo.sType == 0);
+				DescriptorAddressInfo.sType = VK_STRUCTURE_TYPE_BUFFER_DEVICE_ADDRESS_INFO;
+				DescriptorAddressInfo.address = UniformBufferUploader->GetCPUBufferAddress() + RingBufferOffset;
+				DescriptorAddressInfo.range = UBSize;
 
 				// get location in the ring buffer to use
-				FMemory::Memcpy(CPURingBufferBase + RingBufferOffset, StagedUniformBuffer.GetData(), BufferInfo.range);
+				FMemory::Memcpy(CPURingBufferBase + RingBufferOffset, StagedUniformBuffer.GetData(), UBSize);
 			}
 			RemainingPackedUniformsMask = RemainingPackedUniformsMask >> 1;
 			++PackedUBIndex;
@@ -259,10 +271,34 @@ void FVulkanComputePipelineDescriptorState::UpdateBindlessDescriptors(FVulkanCom
 		PackedUniformBuffersDirty = 0;
 	}
 
-	if (DSWriteContainer.DescriptorBufferInfo.Num())
+	for (int32 UBIndex = 0; UBIndex < StageInfo.UniformBuffers.Num(); ++UBIndex)
 	{
-		Device->GetBindlessDescriptorManager()->RegisterUniformBuffers(CmdBuffer->GetHandle(), VK_PIPELINE_BIND_POINT_COMPUTE, ShaderStage::EStage::Compute, DSWriter[0]);
+		const FDescriptorSetRemappingInfo::FUBRemappingInfo& UBRemappingInfo = StageInfo.UniformBuffers[UBIndex];
+		check((StageInfo.PackedUBDescriptorSet == UINT16_MAX) || UBRemappingInfo.Remapping.NewDescriptorSet == StageInfo.PackedUBDescriptorSet);
+		check(UBRemappingInfo.bHasConstantData);
+
+		VkDescriptorAddressInfoEXT& DescriptorAddressInfo = DescriptorAddressInfos[UBRemappingInfo.Remapping.NewBindingIndex];
+		check(DescriptorAddressInfo.sType == 0);
+
+		VkWriteDescriptorSet& WriteDescriptorSet = DSWriter[UBRemappingInfo.Remapping.NewDescriptorSet].WriteDescriptors[UBRemappingInfo.Remapping.NewBindingIndex];
+		check(WriteDescriptorSet.dstBinding == UBRemappingInfo.Remapping.NewBindingIndex);
+		check(WriteDescriptorSet.dstArrayElement == 0);
+		check(WriteDescriptorSet.descriptorCount == 1);
+		check(WriteDescriptorSet.descriptorType == VK_DESCRIPTOR_TYPE_UNIFORM_BUFFER);
+		checkSlow(WriteDescriptorSet.pBufferInfo);
+
+		VkBufferDeviceAddressInfo BufferInfo;
+		ZeroVulkanStruct(BufferInfo, VK_STRUCTURE_TYPE_BUFFER_DEVICE_ADDRESS_INFO);
+		BufferInfo.buffer = WriteDescriptorSet.pBufferInfo->buffer;
+		VkDeviceAddress BufferAddress = VulkanRHI::vkGetBufferDeviceAddressKHR(Device->GetInstanceHandle(), &BufferInfo);
+
+		DescriptorAddressInfo.sType = VK_STRUCTURE_TYPE_BUFFER_DEVICE_ADDRESS_INFO;
+		DescriptorAddressInfo.address = BufferAddress + WriteDescriptorSet.pBufferInfo->offset;
+		DescriptorAddressInfo.range = WriteDescriptorSet.pBufferInfo->range;
 	}
+
+	// Send to descriptor manager
+	Device->GetBindlessDescriptorManager()->RegisterUniformBuffers(CmdBuffer->GetHandle(), VK_PIPELINE_BIND_POINT_COMPUTE, StageUBs);
 }
 
 FVulkanGraphicsPipelineDescriptorState::FVulkanGraphicsPipelineDescriptorState(FVulkanDevice* InDevice, FVulkanRHIGraphicsPipelineState* InGfxPipeline)
@@ -401,7 +437,7 @@ bool FVulkanGraphicsPipelineDescriptorState::InternalUpdateDescriptorSets(FVulka
 
 void FVulkanGraphicsPipelineDescriptorState::UpdateBindlessDescriptors(FVulkanCommandListContext* CmdListContext, FVulkanCmdBuffer* CmdBuffer)
 {
-	check(Device->SupportsBindless());
+	check(bUseBindless);
 
 	// We should only have uniform buffers at this point
 	check(DSWriteContainer.DescriptorBufferInfo.Num() == DSWriteContainer.DescriptorWrites.Num());
@@ -413,35 +449,47 @@ void FVulkanGraphicsPipelineDescriptorState::UpdateBindlessDescriptors(FVulkanCo
 
 	const FDescriptorSetRemappingInfo* RESTRICT RemappingInfo = PipelineDescriptorInfo->RemappingInfo;
 
+	FVulkanBindlessDescriptorManager::FUniformBufferDescriptorArrays StageUBs;
+
 	// Process updates
 	{
 #if VULKAN_ENABLE_AGGRESSIVE_STATS
 		SCOPE_CYCLE_COUNTER(STAT_VulkanApplyPackedUniformBuffers);
 #endif
-		int32 DSWriterIndex = 0;
 		for (int32 Stage = 0; Stage < ShaderStage::NumStages; ++Stage)
 		{
-			if (PackedUniformBuffersDirty[Stage] != 0)
+			const FDescriptorSetRemappingInfo::FStageInfo& StageInfo = RemappingInfo->StageInfos[Stage];
+
+			TArray<VkDescriptorAddressInfoEXT>& DescriptorAddressInfos = StageUBs[Stage];
+			DescriptorAddressInfos.SetNumZeroed(StageInfo.PackedUBBindingIndices.Num() + StageInfo.UniformBuffers.Num());
+
+			// PackedUniformBuffersDirty ?
+			check((PackedUniformBuffersMask[Stage] == 0) || (PackedUniformBuffersMask[Stage] == 1));
+			if (PackedUniformBuffersMask[Stage] != 0)
 			{
-				const uint32 DescriptorSet = RemappingInfo->StageInfos[Stage].PackedUBDescriptorSet;
-				uint64 RemainingPackedUniformsMask = PackedUniformBuffersDirty[Stage];
+				check((StageInfo.PackedUBDescriptorSet != UINT16_MAX) && (StageInfo.PackedUBDescriptorSet <= (uint16)Stage));
+				uint64 RemainingPackedUniformsMask = PackedUniformBuffersMask[Stage];
 				int32 PackedUBIndex = 0;
 				while (RemainingPackedUniformsMask)
 				{
 					if (RemainingPackedUniformsMask & 1)
 					{
 						const FPackedUniformBuffers::FPackedBuffer& StagedUniformBuffer = PackedUniformBuffers[Stage].GetBuffer(PackedUBIndex);
-						const int32 BindingIndex = RemappingInfo->StageInfos[Stage].PackedUBBindingIndices[PackedUBIndex];
-						VkDescriptorBufferInfo& BufferInfo = *const_cast<VkDescriptorBufferInfo*>(DSWriter[DescriptorSet].WriteDescriptors[BindingIndex].pBufferInfo);
-						BufferInfo.buffer = UniformBufferUploader->GetCPUBufferHandle();
-						BufferInfo.range = StagedUniformBuffer.Num();
-						check((BufferInfo.range & 0xF) == 0);   // :todo-jn: nv driver issue
+						const int32 UBSize = StagedUniformBuffer.Num();
+						const int32 PaddedUBSize = Align<int32>(UBSize, 64);  // :todo-jn: work around for NV driver issue
+						const int32 BindingIndex = StageInfo.PackedUBBindingIndices[PackedUBIndex];
 
-						const uint64 RingBufferOffset = UniformBufferUploader->AllocateMemory(BufferInfo.range, UBOffsetAlignment, CmdBuffer);
-						BufferInfo.offset = UniformBufferUploader->GetCPUBufferOffset() + RingBufferOffset;
+						const uint64 RingBufferOffset = UniformBufferUploader->AllocateMemory(PaddedUBSize, UBOffsetAlignment, CmdBuffer);
+
+						// Make sure it wasn't written to already
+						VkDescriptorAddressInfoEXT& DescriptorAddressInfo = DescriptorAddressInfos[BindingIndex];
+						check(DescriptorAddressInfo.sType == 0);
+						DescriptorAddressInfo.sType = VK_STRUCTURE_TYPE_BUFFER_DEVICE_ADDRESS_INFO;
+						DescriptorAddressInfo.address = UniformBufferUploader->GetCPUBufferAddress() + RingBufferOffset;
+						DescriptorAddressInfo.range = PaddedUBSize;
 
 						// get location in the ring buffer to use
-						FMemory::Memcpy(CPURingBufferBase + RingBufferOffset, StagedUniformBuffer.GetData(), BufferInfo.range);
+						FMemory::Memcpy(CPURingBufferBase + RingBufferOffset, StagedUniformBuffer.GetData(), UBSize);
 					}
 					RemainingPackedUniformsMask = RemainingPackedUniformsMask >> 1;
 					++PackedUBIndex;
@@ -450,12 +498,36 @@ void FVulkanGraphicsPipelineDescriptorState::UpdateBindlessDescriptors(FVulkanCo
 				PackedUniformBuffersDirty[Stage] = 0;
 			}
 
-			if ((RemappingInfo->StageInfos[Stage].PackedUBDescriptorSet != UINT16_MAX) || RemappingInfo->StageInfos[Stage].UniformBuffers.Num())
+			for (int32 UBIndex = 0; UBIndex < StageInfo.UniformBuffers.Num(); ++UBIndex)
 			{
-				Device->GetBindlessDescriptorManager()->RegisterUniformBuffers(CmdBuffer->GetHandle(), VK_PIPELINE_BIND_POINT_GRAPHICS, (ShaderStage::EStage)Stage, DSWriter[DSWriterIndex++]);
+				const FDescriptorSetRemappingInfo::FUBRemappingInfo& UBRemappingInfo = StageInfo.UniformBuffers[UBIndex];
+				check((StageInfo.PackedUBDescriptorSet == UINT16_MAX) || UBRemappingInfo.Remapping.NewDescriptorSet == StageInfo.PackedUBDescriptorSet);
+				check(UBRemappingInfo.bHasConstantData);
+
+				VkDescriptorAddressInfoEXT& DescriptorAddressInfo = DescriptorAddressInfos[UBRemappingInfo.Remapping.NewBindingIndex];
+				check(DescriptorAddressInfo.sType == 0);
+
+				VkWriteDescriptorSet& WriteDescriptorSet = DSWriter[UBRemappingInfo.Remapping.NewDescriptorSet].WriteDescriptors[UBRemappingInfo.Remapping.NewBindingIndex];
+				check(WriteDescriptorSet.dstBinding == UBRemappingInfo.Remapping.NewBindingIndex);
+				check(WriteDescriptorSet.dstArrayElement == 0);
+				check(WriteDescriptorSet.descriptorCount == 1);
+				check(WriteDescriptorSet.descriptorType == VK_DESCRIPTOR_TYPE_UNIFORM_BUFFER);
+				checkSlow(WriteDescriptorSet.pBufferInfo);
+
+				VkBufferDeviceAddressInfo BufferInfo;
+				ZeroVulkanStruct(BufferInfo, VK_STRUCTURE_TYPE_BUFFER_DEVICE_ADDRESS_INFO);
+				BufferInfo.buffer = WriteDescriptorSet.pBufferInfo->buffer;
+				VkDeviceAddress BufferAddress = VulkanRHI::vkGetBufferDeviceAddressKHR(Device->GetInstanceHandle(), &BufferInfo);
+
+				DescriptorAddressInfo.sType = VK_STRUCTURE_TYPE_BUFFER_DEVICE_ADDRESS_INFO;
+				DescriptorAddressInfo.address = BufferAddress + WriteDescriptorSet.pBufferInfo->offset;
+				DescriptorAddressInfo.range = Align<VkDeviceSize>(WriteDescriptorSet.pBufferInfo->range, 64u); // :todo-jn: work around for NV driver issue
 			}
 		}
 	}
+
+	// Send to descriptor manager
+	Device->GetBindlessDescriptorManager()->RegisterUniformBuffers(CmdBuffer->GetHandle(), VK_PIPELINE_BIND_POINT_GRAPHICS, StageUBs);
 }
 
 
