VulkanRHI: Print descriptor type names in error message.

#rnx

diff --git a/Engine/Source/Runtime/VulkanRHI/Private/VulkanDescriptorSets.h b/Engine/Source/Runtime/VulkanRHI/Private/VulkanDescriptorSets.h
--- a/Engine/Source/Runtime/VulkanRHI/Private/VulkanDescriptorSets.h
+++ b/Engine/Source/Runtime/VulkanRHI/Private/VulkanDescriptorSets.h
@@ -1129,8 +1129,8 @@ protected:
 	{
 		check(DescriptorIndex < NumWrites);
 		checkf(WriteDescriptors[DescriptorIndex].descriptorType == DescriptorType, 
-			TEXT("DescriptorType mismatch at index %d: called WriteBufferView<%d> and was expecting %d."), 
-			DescriptorIndex, (uint32)DescriptorType, (uint32)WriteDescriptors[DescriptorIndex].descriptorType);
+			TEXT("DescriptorType mismatch at index %d: called WriteBufferView<%s> and was expecting %s."), 
+			DescriptorIndex, VK_TYPE_TO_STRING(VkDescriptorType, DescriptorType), VK_TYPE_TO_STRING(VkDescriptorType, WriteDescriptors[DescriptorIndex].descriptorType));
 		SetWritten(DescriptorIndex);
 		WriteDescriptors[DescriptorIndex].pTexelBufferView = &View->View;
 		BufferViewReferences[DescriptorIndex] = View;
