VulkanRHI: Fix validation warning because of missing flag.

#rnx

diff --git a/Engine/Source/Runtime/VulkanRHI/Private/VulkanUniformBuffer.cpp b/Engine/Source/Runtime/VulkanRHI/Private/VulkanUniformBuffer.cpp
--- a/Engine/Source/Runtime/VulkanRHI/Private/VulkanUniformBuffer.cpp
+++ b/Engine/Source/Runtime/VulkanRHI/Private/VulkanUniformBuffer.cpp
@@ -438,13 +438,19 @@ FVulkanRingBuffer::FVulkanRingBuffer(FVulkanDevice* InDevice, uint64 TotalSize,
 	, BufferAddress(0)
 	, MinAlignment(0)
 {
+	const bool bHasBufferDeviceAddress = InDevice->GetOptionalExtensions().HasBufferDeviceAddress;
+	if (bHasBufferDeviceAddress)
+	{
+		Usage |= VK_BUFFER_USAGE_SHADER_DEVICE_ADDRESS_BIT;
+	}
+
 	check(TotalSize <= (uint64)MAX_uint32);
 	InDevice->GetMemoryManager().AllocateBufferPooled(Allocation, nullptr, TotalSize, 0, Usage, MemPropertyFlags, EVulkanAllocationMetaRingBuffer, __FILE__, __LINE__);
 	MinAlignment = Allocation.GetBufferAlignment(Device);
 	// Start by wrapping around to set up the correct fence
 	BufferOffset = TotalSize;
 
-	if (InDevice->GetOptionalExtensions().HasBufferDeviceAddress)
+	if (bHasBufferDeviceAddress)
 	{
 		VkBufferDeviceAddressInfo BufferInfo;
 		ZeroVulkanStruct(BufferInfo, VK_STRUCTURE_TYPE_BUFFER_DEVICE_ADDRESS_INFO);
