VulkanRHI: Fix static analysis warning.

diff --git a/Engine/Source/Runtime/VulkanRHI/Private/VulkanDescriptorSets.cpp b/Engine/Source/Runtime/VulkanRHI/Private/VulkanDescriptorSets.cpp
--- a/Engine/Source/Runtime/VulkanRHI/Private/VulkanDescriptorSets.cpp
+++ b/Engine/Source/Runtime/VulkanRHI/Private/VulkanDescriptorSets.cpp
@@ -140,18 +140,10 @@ FVulkanBindlessDescriptorManager::FVulkanBindlessDescriptorManager(FVulkanDevice
 {
 	FMemory::Memzero(BufferBindingInfo);
 
-	// Setup up the buffer indices according to the constants currently used in code to end up with 
-	// a different index for each type of bindless resource (up to VulkanBindless::NumBindlessSets)
-	// and then all the uniform buffers in the same descriptor buffer at the end
-	// For example:   BufferIndices[] = { 0,1,2,3,4,5,6,7,7,7,7,7,7,7,7,7,7 };
 	for (uint32 Index = 0; Index < VulkanBindless::NumBindlessSets; Index++)
 	{
 		BufferIndices[Index] = Index;
 	}
-	for (uint32 Index = VulkanBindless::NumBindlessSets; Index < VulkanBindless::MaxNumSets; Index++)
-	{
-		BufferIndices[Index] = VulkanBindless::NumBindlessSets;
-	}
 }
 
 FVulkanBindlessDescriptorManager::~FVulkanBindlessDescriptorManager()
