VulkanRHI: Move pointer check to before SetSRVForStage/SetUAVForStage to fix out-of-bounds in GetDescriptorSetAndBindingIndex when srv/uav are set to nullptr.  Change pointers to ArrayView to catch out-of-bounds before ASan next time.

diff --git a/Engine/Source/Runtime/VulkanRHI/Private/VulkanCommands.cpp b/Engine/Source/Runtime/VulkanRHI/Private/VulkanCommands.cpp
--- a/Engine/Source/Runtime/VulkanRHI/Private/VulkanCommands.cpp
+++ b/Engine/Source/Runtime/VulkanRHI/Private/VulkanCommands.cpp
@@ -171,17 +171,23 @@ void FVulkanCommandListContext::RHIDispatchIndirectComputeShader(FRHIBuffer* Arg
 
 void FVulkanCommandListContext::RHISetUAVParameter(FRHIPixelShader* PixelShaderRHI, uint32 UAVIndex, FRHIUnorderedAccessView* UAVRHI)
 {
+	if (UAVRHI)
+	{
 	FVulkanUnorderedAccessView* UAV = ResourceCast(UAVRHI);
 	PendingGfxState->SetUAVForStage(ShaderStage::Pixel, UAVIndex, UAV);
 }
+}
 
 void FVulkanCommandListContext::RHISetUAVParameter(FRHIComputeShader* ComputeShaderRHI, uint32 UAVIndex, FRHIUnorderedAccessView* UAVRHI)
 {
-	check((UAVRHI == nullptr) || (PendingComputeState->GetCurrentShader() == ResourceCast(ComputeShaderRHI)));
+	if (UAVRHI)
+	{
+		check(PendingComputeState->GetCurrentShader() == ResourceCast(ComputeShaderRHI));
 
 	FVulkanUnorderedAccessView* UAV = ResourceCast(UAVRHI);
 	PendingComputeState->SetUAVForStage(UAVIndex, UAV);
 }
+}
 
 void FVulkanCommandListContext::RHISetUAVParameter(FRHIComputeShader* ComputeShaderRHI,uint32 UAVIndex, FRHIUnorderedAccessView* UAVRHI, uint32 InitialCount)
 {
@@ -212,18 +218,24 @@ void FVulkanCommandListContext::RHISetShaderTexture(FRHIComputeShader* ComputeSh
 
 void FVulkanCommandListContext::RHISetShaderResourceViewParameter(FRHIGraphicsShader* ShaderRHI, uint32 TextureIndex, FRHIShaderResourceView* SRVRHI)
 {
+	if (SRVRHI)
+	{
 	ShaderStage::EStage Stage = GetAndVerifyShaderStage(ShaderRHI, PendingGfxState);
 	FVulkanShaderResourceView* SRV = ResourceCast(SRVRHI);
 	PendingGfxState->SetSRVForStage(Stage, TextureIndex, SRV);
 }
+}
 
 void FVulkanCommandListContext::RHISetShaderResourceViewParameter(FRHIComputeShader* ComputeShaderRHI,uint32 TextureIndex, FRHIShaderResourceView* SRVRHI)
 {
-	check((SRVRHI == nullptr) || (PendingComputeState->GetCurrentShader() == ResourceCast(ComputeShaderRHI)));
+	if (SRVRHI)
+	{
+		check(PendingComputeState->GetCurrentShader() == ResourceCast(ComputeShaderRHI));
 
 	FVulkanShaderResourceView* SRV = ResourceCast(SRVRHI);
 	PendingComputeState->SetSRVForStage(TextureIndex, SRV);
 }
+}
 
 void FVulkanCommandListContext::RHISetShaderSampler(FRHIGraphicsShader* ShaderRHI, uint32 SamplerIndex, FRHISamplerState* NewStateRHI)
 {
diff --git a/Engine/Source/Runtime/VulkanRHI/Private/VulkanDescriptorSets.h b/Engine/Source/Runtime/VulkanRHI/Private/VulkanDescriptorSets.h
--- a/Engine/Source/Runtime/VulkanRHI/Private/VulkanDescriptorSets.h
+++ b/Engine/Source/Runtime/VulkanRHI/Private/VulkanDescriptorSets.h
@@ -709,10 +709,7 @@ class FVulkanComputePipelineDescriptorInfo
 {
 public:
 	FVulkanComputePipelineDescriptorInfo()
-		: RemappingUBInfos(nullptr)
-		, RemappingGlobalInfos(nullptr)
-		, RemappingPackedUBInfos(nullptr)
-		, HasDescriptorsInSetMask(0)
+		: HasDescriptorsInSetMask(0)
 		, RemappingInfo(nullptr)
 		, bInitialized(false)
 	{
@@ -764,12 +761,12 @@ public:
 
 protected:
 	// Cached data from FDescriptorSetRemappingInfo
-	const FDescriptorSetRemappingInfo::FUBRemappingInfo*	RemappingUBInfos;
-	const FDescriptorSetRemappingInfo::FRemappingInfo*		RemappingGlobalInfos;
-	const uint16*											RemappingPackedUBInfos;
-	uint32													HasDescriptorsInSetMask;
-	const FDescriptorSetRemappingInfo*						RemappingInfo;
-	bool													bInitialized;
+	TArrayView<const FDescriptorSetRemappingInfo::FUBRemappingInfo>	RemappingUBInfos;
+	TArrayView<const FDescriptorSetRemappingInfo::FRemappingInfo>	RemappingGlobalInfos;
+	TArrayView<const uint16>										RemappingPackedUBInfos;
+	uint32															HasDescriptorsInSetMask;
+	const FDescriptorSetRemappingInfo*								RemappingInfo;
+	bool															bInitialized;
 
 	friend class FVulkanComputePipelineDescriptorState;
 };
@@ -783,9 +780,6 @@ public:
 		, RemappingInfo(nullptr)
 		, bInitialized(false)
 	{
-		FMemory::Memzero(RemappingUBInfos);
-		FMemory::Memzero(RemappingGlobalInfos);
-		FMemory::Memzero(RemappingPackedUBInfos);
 	}
 
 	inline bool GetDescriptorSetAndBindingIndex(const FVulkanShaderHeader::EType Type, const ShaderStage::EStage Stage, int32 ParameterIndex, uint8& OutDescriptorSet, uint32& OutBindingIndex) const
@@ -838,9 +832,9 @@ public:
 
 protected:
 	// Cached data from FDescriptorSetRemappingInfo
-	const FDescriptorSetRemappingInfo::FUBRemappingInfo*	RemappingUBInfos[ShaderStage::NumStages];
-	const FDescriptorSetRemappingInfo::FRemappingInfo*		RemappingGlobalInfos[ShaderStage::NumStages];
-	const uint16*											RemappingPackedUBInfos[ShaderStage::NumStages];
+	TArrayView<const FDescriptorSetRemappingInfo::FUBRemappingInfo>	RemappingUBInfos[ShaderStage::NumStages];
+	TArrayView<const FDescriptorSetRemappingInfo::FRemappingInfo>	RemappingGlobalInfos[ShaderStage::NumStages];
+	TArrayView<const uint16>										RemappingPackedUBInfos[ShaderStage::NumStages];
 	uint32													HasDescriptorsInSetMask;
 
 	const FDescriptorSetRemappingInfo*						RemappingInfo;
diff --git a/Engine/Source/Runtime/VulkanRHI/Private/VulkanPendingState.cpp b/Engine/Source/Runtime/VulkanRHI/Private/VulkanPendingState.cpp
--- a/Engine/Source/Runtime/VulkanRHI/Private/VulkanPendingState.cpp
+++ b/Engine/Source/Runtime/VulkanRHI/Private/VulkanPendingState.cpp
@@ -309,65 +309,59 @@ FVulkanPendingComputeState::~FVulkanPendingComputeState()
 
 void FVulkanPendingComputeState::SetSRVForUBResource(uint32 DescriptorSet, uint32 BindingIndex, FVulkanShaderResourceView* SRV)
 {
-	if (SRV)
+	check(SRV);
+
+	// make sure any dynamically backed SRV points to current memory
+	SRV->UpdateView();
+	if (SRV->BufferViews.Num() != 0)
 	{
-		// make sure any dynamically backed SRV points to current memory
-		SRV->UpdateView();
-		if (SRV->BufferViews.Num() != 0)
-		{
-			FVulkanBufferView* BufferView = SRV->GetBufferView();
-			checkf(BufferView->View != VK_NULL_HANDLE, TEXT("Empty SRV"));
-			CurrentState->SetSRVBufferViewState(DescriptorSet, BindingIndex, BufferView);
-		}
-		else if (SRV->SourceStructuredBuffer)
-		{
-			CurrentState->SetStorageBuffer(DescriptorSet, BindingIndex, SRV->SourceStructuredBuffer);
-		}
+		FVulkanBufferView* BufferView = SRV->GetBufferView();
+		checkf(BufferView->View != VK_NULL_HANDLE, TEXT("Empty SRV"));
+		CurrentState->SetSRVBufferViewState(DescriptorSet, BindingIndex, BufferView);
+	}
+	else if (SRV->SourceStructuredBuffer)
+	{
+		CurrentState->SetStorageBuffer(DescriptorSet, BindingIndex, SRV->SourceStructuredBuffer);
+	}
 #if VULKAN_RHI_RAYTRACING
-		else if (SRV->AccelerationStructureHandle)
-		{
-			CurrentState->SetAccelerationStructure(DescriptorSet, BindingIndex, SRV->AccelerationStructureHandle);
-		}
-#endif // VULKAN_RHI_RAYTRACING
-		else
-		{
-			const FVulkanTexture* VulkanTexture = FVulkanTexture::Cast(SRV->SourceTexture.GetReference());
-			checkf((SRV->TextureView.View != VK_NULL_HANDLE) && VulkanTexture, TEXT("Empty SRV"));
-			const VkImageLayout ExpectedLayout = FVulkanLayoutManager::GetDefaultLayout(
-				Context.GetCommandBufferManager()->GetActiveCmdBuffer(), *VulkanTexture, ERHIAccess::SRVCompute);
-			CurrentState->SetSRVTextureView(DescriptorSet, BindingIndex, SRV->TextureView, ExpectedLayout);
-		}
+	else if (SRV->AccelerationStructureHandle)
+	{
+		CurrentState->SetAccelerationStructure(DescriptorSet, BindingIndex, SRV->AccelerationStructureHandle);
 	}
+#endif // VULKAN_RHI_RAYTRACING
 	else
 	{
-		//CurrentState->SetSRVBufferViewState(BindIndex, nullptr);
+		const FVulkanTexture* VulkanTexture = FVulkanTexture::Cast(SRV->SourceTexture.GetReference());
+		checkf((SRV->TextureView.View != VK_NULL_HANDLE) && VulkanTexture, TEXT("Empty SRV"));
+		const VkImageLayout ExpectedLayout = FVulkanLayoutManager::GetDefaultLayout(
+			Context.GetCommandBufferManager()->GetActiveCmdBuffer(), *VulkanTexture, ERHIAccess::SRVCompute);
+		CurrentState->SetSRVTextureView(DescriptorSet, BindingIndex, SRV->TextureView, ExpectedLayout);
 	}
 }
 
 void FVulkanPendingComputeState::SetUAVForUBResource(uint32 DescriptorSet, uint32 BindingIndex, FVulkanUnorderedAccessView* UAV)
 {
-	if (UAV)
+	check(UAV);
+
+	// make sure any dynamically backed UAV points to current memory
+	UAV->UpdateView();
+	if (UAV->SourceBuffer && UAV->BufferViewFormat == PF_Unknown)
 	{
-		// make sure any dynamically backed UAV points to current memory
-		UAV->UpdateView();
-		if (UAV->SourceBuffer && UAV->BufferViewFormat == PF_Unknown)
-		{
-			CurrentState->SetStorageBuffer(DescriptorSet, BindingIndex, UAV->SourceBuffer);
-		}
-		else if (UAV->BufferView)
-		{
-			CurrentState->SetUAVTexelBufferViewState(DescriptorSet, BindingIndex, UAV->BufferView);
-		}
-		else if (UAV->SourceTexture)
-		{
-			const FVulkanTexture* VulkanTexture = FVulkanTexture::Cast(UAV->SourceTexture.GetReference());
-			const VkImageLayout ExpectedLayout = FVulkanLayoutManager::GetDefaultLayout(Context.GetCommandBufferManager()->GetActiveCmdBuffer(), *VulkanTexture, ERHIAccess::UAVCompute);
-			CurrentState->SetUAVTextureView(DescriptorSet, BindingIndex, UAV->TextureView, ExpectedLayout);
-		}
-		else
-		{
-			ensure(0);
-		}
+		CurrentState->SetStorageBuffer(DescriptorSet, BindingIndex, UAV->SourceBuffer);
+	}
+	else if (UAV->BufferView)
+	{
+		CurrentState->SetUAVTexelBufferViewState(DescriptorSet, BindingIndex, UAV->BufferView);
+	}
+	else if (UAV->SourceTexture)
+	{
+		const FVulkanTexture* VulkanTexture = FVulkanTexture::Cast(UAV->SourceTexture.GetReference());
+		const VkImageLayout ExpectedLayout = FVulkanLayoutManager::GetDefaultLayout(Context.GetCommandBufferManager()->GetActiveCmdBuffer(), *VulkanTexture, ERHIAccess::UAVCompute);
+		CurrentState->SetUAVTextureView(DescriptorSet, BindingIndex, UAV->TextureView, ExpectedLayout);
+	}
+	else
+	{
+		ensure(0);
 	}
 }
 
@@ -599,60 +593,54 @@ void FVulkanPendingGfxState::UpdateInputAttachments(FVulkanFramebuffer* Framebuf
 
 void FVulkanPendingGfxState::SetSRVForUBResource(uint8 DescriptorSet, uint32 BindingIndex, FVulkanShaderResourceView* SRV)
 {
-	if (SRV)
+	check(SRV);
+
+	// make sure any dynamically backed SRV points to current memory
+	SRV->UpdateView();
+	if (SRV->BufferViews.Num() != 0)
 	{
-		// make sure any dynamically backed SRV points to current memory
-		SRV->UpdateView();
-		if (SRV->BufferViews.Num() != 0)
-		{
-			FVulkanBufferView* BufferView = SRV->GetBufferView();
-			checkf(BufferView->View != VK_NULL_HANDLE, TEXT("Empty SRV"));
+		FVulkanBufferView* BufferView = SRV->GetBufferView();
+		checkf(BufferView->View != VK_NULL_HANDLE, TEXT("Empty SRV"));
 
-			CurrentState->SetSRVBufferViewState(DescriptorSet, BindingIndex, BufferView);
-		}
-		else if (SRV->SourceStructuredBuffer)
-		{
-			CurrentState->SetStorageBuffer(DescriptorSet, BindingIndex, SRV->SourceStructuredBuffer);
-		}
-		else
-		{
-			const FVulkanTexture* VulkanTexture = FVulkanTexture::Cast(SRV->SourceTexture.GetReference());
-			checkf((SRV->TextureView.View != VK_NULL_HANDLE) && VulkanTexture, TEXT("Empty SRV"));
-			const VkImageLayout ExpectedLayout = FVulkanLayoutManager::GetDefaultLayout(Context.GetCommandBufferManager()->GetActiveCmdBuffer(), *VulkanTexture, ERHIAccess::SRVGraphics);
-			CurrentState->SetSRVTextureView(DescriptorSet, BindingIndex, SRV->TextureView, ExpectedLayout);
-		}
+		CurrentState->SetSRVBufferViewState(DescriptorSet, BindingIndex, BufferView);
+	}
+	else if (SRV->SourceStructuredBuffer)
+	{
+		CurrentState->SetStorageBuffer(DescriptorSet, BindingIndex, SRV->SourceStructuredBuffer);
 	}
 	else
 	{
-		//CurrentState->SetSRVBufferViewState(Stage, BindIndex, nullptr);
+		const FVulkanTexture* VulkanTexture = FVulkanTexture::Cast(SRV->SourceTexture.GetReference());
+		checkf((SRV->TextureView.View != VK_NULL_HANDLE) && VulkanTexture, TEXT("Empty SRV"));
+		const VkImageLayout ExpectedLayout = FVulkanLayoutManager::GetDefaultLayout(Context.GetCommandBufferManager()->GetActiveCmdBuffer(), *VulkanTexture, ERHIAccess::SRVGraphics);
+		CurrentState->SetSRVTextureView(DescriptorSet, BindingIndex, SRV->TextureView, ExpectedLayout);
 	}
 }
 
 void FVulkanPendingGfxState::SetUAVForUBResource(uint8 DescriptorSet, uint32 BindingIndex, FVulkanUnorderedAccessView* UAV)
 {
-	if (UAV)
+	check(UAV);
+
+	// make sure any dynamically backed UAV points to current memory
+	UAV->UpdateView();
+	if (UAV->SourceBuffer && UAV->BufferViewFormat == PF_Unknown)
 	{
-		// make sure any dynamically backed UAV points to current memory
-		UAV->UpdateView();
-		if (UAV->SourceBuffer && UAV->BufferViewFormat == PF_Unknown)
-		{
-			CurrentState->SetStorageBuffer(DescriptorSet, BindingIndex, UAV->SourceBuffer);
-		}
-		else if (UAV->BufferView)
-		{
-			CurrentState->SetUAVTexelBufferViewState(DescriptorSet, BindingIndex, UAV->BufferView);
-		}
-		else if (UAV->SourceTexture)
-		{
-			const FVulkanTexture* VulkanTexture = FVulkanTexture::Cast(UAV->SourceTexture.GetReference());
-			checkf(VulkanTexture, TEXT("Empty UAV"));
-			const VkImageLayout Layout = FVulkanLayoutManager::GetDefaultLayout(Context.GetCommandBufferManager()->GetActiveCmdBuffer(), *VulkanTexture, ERHIAccess::UAVGraphics);
-			CurrentState->SetUAVTextureView(DescriptorSet, BindingIndex, UAV->TextureView, Layout);
-		}
-		else
-		{
-			ensure(0);
-		}
+		CurrentState->SetStorageBuffer(DescriptorSet, BindingIndex, UAV->SourceBuffer);
+	}
+	else if (UAV->BufferView)
+	{
+		CurrentState->SetUAVTexelBufferViewState(DescriptorSet, BindingIndex, UAV->BufferView);
+	}
+	else if (UAV->SourceTexture)
+	{
+		const FVulkanTexture* VulkanTexture = FVulkanTexture::Cast(UAV->SourceTexture.GetReference());
+		checkf(VulkanTexture, TEXT("Empty UAV"));
+		const VkImageLayout Layout = FVulkanLayoutManager::GetDefaultLayout(Context.GetCommandBufferManager()->GetActiveCmdBuffer(), *VulkanTexture, ERHIAccess::UAVGraphics);
+		CurrentState->SetUAVTextureView(DescriptorSet, BindingIndex, UAV->TextureView, Layout);
+	}
+	else
+	{
+		ensure(0);
 	}
 }
 
diff --git a/Engine/Source/Runtime/VulkanRHI/Private/VulkanShaders.cpp b/Engine/Source/Runtime/VulkanRHI/Private/VulkanShaders.cpp
--- a/Engine/Source/Runtime/VulkanRHI/Private/VulkanShaders.cpp
+++ b/Engine/Source/Runtime/VulkanRHI/Private/VulkanShaders.cpp
@@ -848,9 +848,9 @@ void FVulkanComputePipelineDescriptorInfo::Initialize(const FDescriptorSetRemapp
 {
 	check(!bInitialized);
 
-	RemappingGlobalInfos = InRemappingInfo.StageInfos[0].Globals.GetData();
-	RemappingUBInfos = InRemappingInfo.StageInfos[0].UniformBuffers.GetData();
-	RemappingPackedUBInfos = InRemappingInfo.StageInfos[0].PackedUBBindingIndices.GetData();
+	RemappingGlobalInfos = InRemappingInfo.StageInfos[0].Globals;
+	RemappingUBInfos = InRemappingInfo.StageInfos[0].UniformBuffers;
+	RemappingPackedUBInfos = InRemappingInfo.StageInfos[0].PackedUBBindingIndices;
 
 	RemappingInfo = &InRemappingInfo;
 
@@ -878,9 +878,9 @@ void FVulkanGfxPipelineDescriptorInfo::Initialize(const FDescriptorSetRemappingI
 	for (int32 StageIndex = 0; StageIndex < ShaderStage::NumStages; ++StageIndex)
 	{
 		//#todo-rco: Enable this!
-		RemappingUBInfos[StageIndex] = InRemappingInfo.StageInfos[StageIndex].UniformBuffers.GetData();
-		RemappingGlobalInfos[StageIndex] = InRemappingInfo.StageInfos[StageIndex].Globals.GetData();
-		RemappingPackedUBInfos[StageIndex] = InRemappingInfo.StageInfos[StageIndex].PackedUBBindingIndices.GetData();
+		RemappingUBInfos[StageIndex] = InRemappingInfo.StageInfos[StageIndex].UniformBuffers;
+		RemappingGlobalInfos[StageIndex] = InRemappingInfo.StageInfos[StageIndex].Globals;
+		RemappingPackedUBInfos[StageIndex] = InRemappingInfo.StageInfos[StageIndex].PackedUBBindingIndices;
 	}
 
 	RemappingInfo = &InRemappingInfo;
