Fixed: Swappy VK initialization which was missing a WindowHandle

diff --git a/Engine/Source/Runtime/VulkanRHI/Private/Android/VulkanAndroidPlatform.cpp b/Engine/Source/Runtime/VulkanRHI/Private/Android/VulkanAndroidPlatform.cpp
--- a/Engine/Source/Runtime/VulkanRHI/Private/Android/VulkanAndroidPlatform.cpp
+++ b/Engine/Source/Runtime/VulkanRHI/Private/Android/VulkanAndroidPlatform.cpp
@@ -346,25 +346,31 @@ void FVulkanAndroidPlatform::InitDevice(FVulkanDevice* InDevice)
 #endif
 }
 
-void FVulkanAndroidPlatform::CreateSurface(void* WindowHandle, VkInstance Instance, VkSurfaceKHR* OutSurface)
+void* FVulkanAndroidPlatform::GetHardwareWindowHandle()
 {
 	// don't use cached window handle coming from VulkanViewport, as it could be gone by now
-	WindowHandle = FAndroidWindow::GetHardwareWindow_EventThread();
-	if (WindowHandle == NULL)
+	void* WindowHandle = FAndroidWindow::GetHardwareWindow_EventThread();
+	if (WindowHandle == nullptr)
 	{
-
 		// Sleep if the hardware window isn't currently available.
 		// The Window may not exist if the activity is pausing/resuming, in which case we make this thread wait
 		FPlatformMisc::LowLevelOutputDebugString(TEXT("Waiting for Native window in FVulkanAndroidPlatform::CreateSurface"));
 		WindowHandle = FAndroidWindow::WaitForHardwareWindow();
 
-		if (WindowHandle == NULL)
+		if (WindowHandle == nullptr)
 		{
 			FPlatformMisc::LowLevelOutputDebugString(TEXT("Aborting FVulkanAndroidPlatform::CreateSurface, FAndroidWindow::WaitForHardwareWindow() returned null"));
-			return;
 		}
 	}
 
+	return WindowHandle;
+}
+
+void FVulkanAndroidPlatform::CreateSurface(void* WindowHandle, VkInstance Instance, VkSurfaceKHR* OutSurface)
+{
+	// don't use cached window handle coming from VulkanViewport, as it could be gone by now
+	WindowHandle = GetHardwareWindowHandle();
+
 	VkAndroidSurfaceCreateInfoKHR SurfaceCreateInfo;
 	ZeroVulkanStruct(SurfaceCreateInfo, VK_STRUCTURE_TYPE_ANDROID_SURFACE_CREATE_INFO_KHR);
 	SurfaceCreateInfo.window = (ANativeWindow*)WindowHandle;
@@ -655,6 +661,10 @@ VkResult FVulkanAndroidPlatform::CreateSwapchainKHR(void* WindowHandle, VkPhysic
 			JNIEnv* Env = FAndroidApplication::GetJavaEnv();
 			if (ensure(Env))
 			{
+				// don't use cached window handle coming from VulkanViewport, as it could be gone by now
+				WindowHandle = GetHardwareWindowHandle();
+				check(WindowHandle);
+								
 				uint64_t RefreshDuration; // in nanoseconds
 				SwappyVk_initAndGetRefreshCycleDuration(Env, FJavaWrapper::GameActivityThis, PhysicalDevice, Device, *Swapchain, &RefreshDuration);
 				SwappyVk_setWindow(Device, *Swapchain, (ANativeWindow*)WindowHandle);	
diff --git a/Engine/Source/Runtime/VulkanRHI/Private/Android/VulkanAndroidPlatform.h b/Engine/Source/Runtime/VulkanRHI/Private/Android/VulkanAndroidPlatform.h
--- a/Engine/Source/Runtime/VulkanRHI/Private/Android/VulkanAndroidPlatform.h
+++ b/Engine/Source/Runtime/VulkanRHI/Private/Android/VulkanAndroidPlatform.h
@@ -66,6 +66,8 @@ public:
 
 	static void CreateSurface(void* WindowHandle, VkInstance Instance, VkSurfaceKHR* OutSurface);
 
+	static void* GetHardwareWindowHandle();
+
 	static bool SupportsBCTextureFormats() { return false; }
 	static bool SupportsASTCTextureFormats() { return true; }
 	static bool SupportsETC2TextureFormats() { return true; }
