Fixed: Issues with 90Hz framerate settings

diff --git a/Engine/Source/Runtime/VulkanRHI/Private/Android/VulkanAndroidPlatform.cpp b/Engine/Source/Runtime/VulkanRHI/Private/Android/VulkanAndroidPlatform.cpp
--- a/Engine/Source/Runtime/VulkanRHI/Private/Android/VulkanAndroidPlatform.cpp
+++ b/Engine/Source/Runtime/VulkanRHI/Private/Android/VulkanAndroidPlatform.cpp
@@ -735,16 +735,33 @@ bool FAndroidVulkanFramePacer::SupportsFramePaceInternal(int32 QueryFramePace, i
 	// check if we want to use CPU frame pacing at less than a multiple of supported refresh rate
 	if (FAndroidPlatformRHIFramePacer::CVarSupportNonVSyncMultipleFrameRates.GetValueOnAnyThread() == 1)
 	{
+		int32 ClosestRefreshRate = 0;
+		
 		for (int32 Rate : RefreshRates)
 		{
 			if (Rate > QueryFramePace)
 			{
-				UE_LOG(LogRHI, Log, TEXT("Supports %d using refresh rate %d with CPU frame pacing"), QueryFramePace, Rate);
-				OutRefreshRate = Rate;
-				OutSyncInterval = 0;
-				return true;
+				if (ClosestRefreshRate < QueryFramePace)
+				{
+					ClosestRefreshRate = Rate;
+				}
+
+				// prefer 30hz multiple
+				if ((Rate % 30) == 0)
+				{
+					ClosestRefreshRate = Rate;
+					break;
+				}
 			}
 		}
+
+		if (ClosestRefreshRate > QueryFramePace)
+		{
+			OutRefreshRate = ClosestRefreshRate;
+			OutSyncInterval = 0;
+			UE_LOG(LogRHI, Log, TEXT("Supports %d using refresh rate %d with CPU frame pacing"), QueryFramePace, OutRefreshRate);
+			return true;
+		}
 	}
 
 	OutRefreshRate = QueryFramePace;
