Added r.Vulkan.UseASTCDecodeMode cvar to have an option to disable VK_EXT_astc_decode_mode usage

diff --git a/Engine/Source/Runtime/VulkanRHI/Private/Android/VulkanAndroidPlatform.cpp b/Engine/Source/Runtime/VulkanRHI/Private/Android/VulkanAndroidPlatform.cpp
--- a/Engine/Source/Runtime/VulkanRHI/Private/Android/VulkanAndroidPlatform.cpp
+++ b/Engine/Source/Runtime/VulkanRHI/Private/Android/VulkanAndroidPlatform.cpp
@@ -422,11 +422,22 @@ static FAutoConsoleVariableRef CVarVulkanQcomRenderPassTransform(
 	ECVF_ReadOnly
 );
 
+static int32 GVulkanUseASTCDecodeMode = 1;
+static FAutoConsoleVariableRef CVarVulkanUseASTCDecodeMode(
+	TEXT("r.Vulkan.UseASTCDecodeMode"),
+	GVulkanUseASTCDecodeMode,
+	TEXT("Whether to use VK_EXT_astc_decode_mode extension\n"),
+	ECVF_ReadOnly
+);
+
 void FVulkanAndroidPlatform::GetDeviceExtensions(FVulkanDevice* Device, FVulkanDeviceExtensionArray& OutExtensions)
 {
 	OutExtensions.Add(MakeUnique<FVulkanDeviceExtension>(Device, VK_KHR_ANDROID_SURFACE_EXTENSION_NAME, VULKAN_EXTENSION_ENABLED, VULKAN_EXTENSION_NOT_PROMOTED));
 	OutExtensions.Add(MakeUnique<FVulkanDeviceExtension>(Device, VK_GOOGLE_DISPLAY_TIMING_EXTENSION_NAME, VULKAN_EXTENSION_ENABLED, VULKAN_EXTENSION_NOT_PROMOTED));
-	OutExtensions.Add(MakeUnique<FVulkanDeviceExtension>(Device, VK_EXT_ASTC_DECODE_MODE_EXTENSION_NAME, VULKAN_SUPPORTS_ASTC_DECODE_MODE, VULKAN_EXTENSION_NOT_PROMOTED, DEVICE_EXT_FLAG_SETTER(HasEXTASTCDecodeMode)));
+	if (GVulkanUseASTCDecodeMode)
+	{
+		OutExtensions.Add(MakeUnique<FVulkanDeviceExtension>(Device, VK_EXT_ASTC_DECODE_MODE_EXTENSION_NAME, VULKAN_SUPPORTS_ASTC_DECODE_MODE, VULKAN_EXTENSION_NOT_PROMOTED, DEVICE_EXT_FLAG_SETTER(HasEXTASTCDecodeMode)));
+	}
 	//OutExtensions.Add(MakeUnique<FVulkanDeviceExtension>(Device, VK_EXT_TEXTURE_COMPRESSION_ASTC_HDR_EXTENSION_NAME, VULKAN_SUPPORTS_TEXTURE_COMPRESSION_ASTC_HDR, VK_API_VERSION_1_3, DEVICE_EXT_FLAG_SETTER(HasEXTTextureCompressionASTCHDR)));
 
 	if (GVulkanQcomRenderPassTransform)
