VulkanRHI:  Push a draft implementation of the FVulkanBindlessDescriptorManager that uses VK_EXT_descriptor_buffer.  Uses one descriptor set for each type of resource, each backed by a descriptor buffer.  Everything is bound and visible to every shader, except uniform buffers which are bound per draw/dispatch.  Uses a single pipeline layout for everything.  Initial buffer sizes for each descrciptor type can be configured through CVars (they do not automatically grow yet).  Still very experimental, requires multiple other features to be disabled.  Vulkan Bindless WIP.

diff --git a/Engine/Source/Runtime/VulkanRHI/Private/VulkanDescriptorSets.cpp b/Engine/Source/Runtime/VulkanRHI/Private/VulkanDescriptorSets.cpp
--- a/Engine/Source/Runtime/VulkanRHI/Private/VulkanDescriptorSets.cpp
+++ b/Engine/Source/Runtime/VulkanRHI/Private/VulkanDescriptorSets.cpp
@@ -80,16 +80,7 @@ static FAutoConsoleVariableRef CVarVulkanBindlessMaxAccelerationStructureCount(
 	ECVF_ReadOnly
 );
 
-
-int32 GVulkanBindlessMaxUniformBuffersPerStage = 32;
-static FAutoConsoleVariableRef CVarVulkanBindlessMaxUniformBuffersPerStage(
-	TEXT("r.Vulkan.Bindless.MaxUniformBuffersPerStage"),
-	GVulkanBindlessMaxUniformBuffersPerStage,
-	TEXT("Maximum Uniform Buffers per shader stage"),
-	ECVF_ReadOnly
-);
-
-int32 GVulkanBindlessRebindBuffers = 0;
+int32 GVulkanBindlessRebindBuffers = 1;
 static FAutoConsoleVariableRef CVarVulkanBindlessRebindBuffers(
 	TEXT("r.Vulkan.Bindless.RebindBuffers"),
 	GVulkanBindlessRebindBuffers,
@@ -97,6 +88,104 @@ static FAutoConsoleVariableRef CVarVulkanBindlessRebindBuffers(
 	ECVF_RenderThreadSafe
 );
 
+int32 GVulkanBindlessBufferOffsetUpdates = 0;
+static FAutoConsoleVariableRef CVarVulkanBindlessBufferOffsetUpdates(
+	TEXT("r.Vulkan.Bindless.BufferOffsetUpdates"),
+	GVulkanBindlessBufferOffsetUpdates,
+	TEXT("0 to set all offsets for each draw/dispatch\n")\
+	TEXT("1 to set resource descriptor buffer offsets once, and only update for uniform buffer offsets on draw/dispatch\n"),
+	ECVF_ReadOnly
+);
+
+
+DECLARE_STATS_GROUP(TEXT("Vulkan Bindless"), STATGROUP_VulkanBindless, STATCAT_Advanced);
+
+DECLARE_DWORD_ACCUMULATOR_STAT_EXTERN(TEXT("Num Peak Descriptor Count"), STAT_VulkanBindlessPeakDescriptorCount, STATGROUP_VulkanBindless, );
+DECLARE_DWORD_ACCUMULATOR_STAT_EXTERN(TEXT("Peak Samplers"), STAT_VulkanBindlessPeakSampler, STATGROUP_VulkanBindless, );
+DECLARE_DWORD_ACCUMULATOR_STAT_EXTERN(TEXT("Peak Images"), STAT_VulkanBindlessPeakImage, STATGROUP_VulkanBindless, );
+DECLARE_DWORD_ACCUMULATOR_STAT_EXTERN(TEXT("Peak Buffers"), STAT_VulkanBindlessPeakBuffer, STATGROUP_VulkanBindless, );
+DECLARE_DWORD_ACCUMULATOR_STAT_EXTERN(TEXT("Peak Texel Buffers"), STAT_VulkanBindlessPeakTexelBuffer, STATGROUP_VulkanBindless, );
+DECLARE_DWORD_ACCUMULATOR_STAT_EXTERN(TEXT("Peak Acceleration Structures"), STAT_VulkanBindlessPeakAccelerationStructure, STATGROUP_VulkanBindless, );
+DECLARE_DWORD_COUNTER_STAT_EXTERN(TEXT("Write Per Frame"), STAT_VulkanBindlessWritePerFrame, STATGROUP_VulkanBindless, );
+
+DEFINE_STAT(STAT_VulkanBindlessPeakDescriptorCount);
+DEFINE_STAT(STAT_VulkanBindlessPeakSampler);
+DEFINE_STAT(STAT_VulkanBindlessPeakImage);
+DEFINE_STAT(STAT_VulkanBindlessPeakBuffer);
+DEFINE_STAT(STAT_VulkanBindlessPeakTexelBuffer);
+DEFINE_STAT(STAT_VulkanBindlessPeakAccelerationStructure);
+DEFINE_STAT(STAT_VulkanBindlessWritePerFrame);
+
+
+static inline uint8 GetIndexForDescriptorType(VkDescriptorType DescriptorType)
+{
+	switch (DescriptorType)
+	{
+	case VK_DESCRIPTOR_TYPE_SAMPLER:                    return VulkanBindless::BindlessSamplerSet;
+	case VK_DESCRIPTOR_TYPE_SAMPLED_IMAGE:              return VulkanBindless::BindlessSampledImageSet;
+	case VK_DESCRIPTOR_TYPE_STORAGE_IMAGE:              return VulkanBindless::BindlessStorageImageSet;
+	case VK_DESCRIPTOR_TYPE_UNIFORM_TEXEL_BUFFER:       return VulkanBindless::BindlessUniformTexelBufferSet;
+	case VK_DESCRIPTOR_TYPE_STORAGE_TEXEL_BUFFER:       return VulkanBindless::BindlessStorageTexelBufferSet;
+	case VK_DESCRIPTOR_TYPE_STORAGE_BUFFER:             return VulkanBindless::BindlessStorageBufferSet;
+	case VK_DESCRIPTOR_TYPE_UNIFORM_BUFFER:             return VulkanBindless::BindlessUniformBufferSet;
+	case VK_DESCRIPTOR_TYPE_ACCELERATION_STRUCTURE_KHR: return VulkanBindless::BindlessAccelerationStructureSet;
+	default: checkNoEntry();
+	}
+
+	return VulkanBindless::NumBindlessSets;
+}
+
+
+static inline VkDescriptorType GetDescriptorTypeForSetIndex(uint8 SetIndex)
+{
+	switch (SetIndex)
+	{
+	case VulkanBindless::BindlessSamplerSet:                return VK_DESCRIPTOR_TYPE_SAMPLER;
+	case VulkanBindless::BindlessSampledImageSet:           return VK_DESCRIPTOR_TYPE_SAMPLED_IMAGE;
+	case VulkanBindless::BindlessStorageImageSet:           return VK_DESCRIPTOR_TYPE_STORAGE_IMAGE;
+	case VulkanBindless::BindlessUniformTexelBufferSet:     return VK_DESCRIPTOR_TYPE_UNIFORM_TEXEL_BUFFER;
+	case VulkanBindless::BindlessStorageTexelBufferSet:     return VK_DESCRIPTOR_TYPE_STORAGE_TEXEL_BUFFER;
+	case VulkanBindless::BindlessStorageBufferSet:          return VK_DESCRIPTOR_TYPE_STORAGE_BUFFER;
+	case VulkanBindless::BindlessUniformBufferSet:          return VK_DESCRIPTOR_TYPE_UNIFORM_BUFFER;
+	case VulkanBindless::BindlessAccelerationStructureSet:  return VK_DESCRIPTOR_TYPE_ACCELERATION_STRUCTURE_KHR;
+	default: checkNoEntry();
+	}
+	return VK_DESCRIPTOR_TYPE_MAX_ENUM;
+}
+
+
+extern TAutoConsoleVariable<int32> GCVarRobustBufferAccess;
+static inline uint32 GetDescriptorTypeSize(FVulkanDevice* Device, VkDescriptorType DescriptorType)
+{
+	const bool bRobustBufferAccess = (GCVarRobustBufferAccess.GetValueOnAnyThread() > 0);
+	const VkPhysicalDeviceDescriptorBufferPropertiesEXT& DescriptorBufferProperties = Device->GetDescriptorBufferProperties();
+
+	switch (DescriptorType)
+	{
+	case VK_DESCRIPTOR_TYPE_SAMPLER:                    
+		return DescriptorBufferProperties.samplerDescriptorSize;
+	case VK_DESCRIPTOR_TYPE_SAMPLED_IMAGE:              
+		return DescriptorBufferProperties.sampledImageDescriptorSize;
+	case VK_DESCRIPTOR_TYPE_STORAGE_IMAGE:              
+		return DescriptorBufferProperties.storageImageDescriptorSize;
+	case VK_DESCRIPTOR_TYPE_UNIFORM_TEXEL_BUFFER:       
+		return bRobustBufferAccess	? DescriptorBufferProperties.robustUniformTexelBufferDescriptorSize 
+									: DescriptorBufferProperties.uniformTexelBufferDescriptorSize;
+	case VK_DESCRIPTOR_TYPE_STORAGE_TEXEL_BUFFER:       
+		return bRobustBufferAccess	? DescriptorBufferProperties.robustStorageTexelBufferDescriptorSize 
+									: DescriptorBufferProperties.storageTexelBufferDescriptorSize;
+	case VK_DESCRIPTOR_TYPE_UNIFORM_BUFFER:             
+		return bRobustBufferAccess	? DescriptorBufferProperties.robustUniformBufferDescriptorSize 
+									: DescriptorBufferProperties.uniformBufferDescriptorSize;
+	case VK_DESCRIPTOR_TYPE_STORAGE_BUFFER:             
+		return bRobustBufferAccess	? DescriptorBufferProperties.robustStorageBufferDescriptorSize 
+									: DescriptorBufferProperties.storageBufferDescriptorSize;
+	case VK_DESCRIPTOR_TYPE_ACCELERATION_STRUCTURE_KHR: 
+		return DescriptorBufferProperties.accelerationStructureDescriptorSize;
+	default: checkNoEntry();
+	}
+	return 0;
+}
 
 
 static inline uint32 GetInitialDescriptorCount(VkDescriptorType DescriptorType)
@@ -132,6 +221,77 @@ static inline VkMemoryPropertyFlags GetDescriptorBufferMemoryType(FVulkanDevice*
 
 
 
+// Check all the requirements to be running in Bindless using Descriptor Buffers
+bool FVulkanBindlessDescriptorManager::VerifySupport(FVulkanDevice* Device)
+{
+	const bool bFullyDisabled =
+		(RHIGetBindlessResourcesConfiguration(GMaxRHIShaderPlatform) == ERHIBindlessConfiguration::Disabled) &&
+		(RHIGetBindlessSamplersConfiguration(GMaxRHIShaderPlatform) == ERHIBindlessConfiguration::Disabled);
+
+	if ((GVulkanBindlessEnabled == 0) || bFullyDisabled)
+	{
+		return false;
+	}
+
+	extern TAutoConsoleVariable<int32> GDynamicGlobalUBs;
+	if (GDynamicGlobalUBs->GetInt() != 0)  // :todo-jn:
+	{
+		UE_LOG(LogRHI, Warning, TEXT("Please disable Dynamic Uniform Buffers (r.Vulkan.DynamicGlobalUBs=0) if you want to use Bindless in Vulkan."));
+		return false;
+	}
+
+	extern int32 GVulkanEnableDefrag;
+	if (GVulkanEnableDefrag != 0)  // :todo-jn:
+	{
+		UE_LOG(LogRHI, Warning, TEXT("Please disable memory defrag (r.Vulkan.EnableDefrag=0) if you want to use Bindless in Vulkan."));
+		return false;
+	}
+
+	const bool bFullyEnabled =
+		(RHIGetBindlessResourcesConfiguration(GMaxRHIShaderPlatform) == ERHIBindlessConfiguration::AllShaders) &&
+		(RHIGetBindlessSamplersConfiguration(GMaxRHIShaderPlatform) == ERHIBindlessConfiguration::AllShaders);
+
+	if (bFullyEnabled)
+	{
+		const VkPhysicalDeviceProperties& GpuProps = Device->GetDeviceProperties();
+		const FOptionalVulkanDeviceExtensions& OptionalDeviceExtensions = Device->GetOptionalExtensions();
+		const VkPhysicalDeviceDescriptorBufferPropertiesEXT& DescriptorBufferProperties = Device->GetDescriptorBufferProperties();
+
+		const bool bMeetsExtensionsRequirements =
+			OptionalDeviceExtensions.HasEXTDescriptorIndexing &&
+			OptionalDeviceExtensions.HasBufferDeviceAddress &&
+			OptionalDeviceExtensions.HasEXTDescriptorBuffer;
+
+		if (bMeetsExtensionsRequirements)
+		{
+			const bool bMeetsPropertiesRequirements =
+				(GpuProps.limits.maxBoundDescriptorSets >= VulkanBindless::NumBindlessSets) &&
+				(DescriptorBufferProperties.maxDescriptorBufferBindings >= VulkanBindless::NumBindlessSets) &&
+				(DescriptorBufferProperties.maxResourceDescriptorBufferBindings >= VulkanBindless::NumBindlessSets) &&
+				(DescriptorBufferProperties.maxSamplerDescriptorBufferBindings >= 1) &&
+				Device->GetDeviceMemoryManager().SupportsMemoryType(GetDescriptorBufferMemoryType(Device));
+
+			if (bMeetsPropertiesRequirements)
+			{
+				return true;
+			}
+			else
+			{
+				UE_LOG(LogRHI, Warning, TEXT("Bindless descriptor were requested but NOT enabled because of insufficient property support."));
+			}
+		}
+		else
+		{
+			UE_LOG(LogRHI, Warning, TEXT("Bindless descriptor were requested but NOT enabled because of missing extension support."));
+		}
+	}
+	else
+	{
+		UE_LOG(LogRHI, Warning, TEXT("Bindless in Vulkan must currently be fully enabled (all samplers and resources) or fully disabled."));
+	}
+
+	return false;
+}
 
 
 FVulkanBindlessDescriptorManager::FVulkanBindlessDescriptorManager(FVulkanDevice* InDevice)
@@ -139,7 +299,6 @@ FVulkanBindlessDescriptorManager::FVulkanBindlessDescriptorManager(FVulkanDevice
 	, bIsSupported(VerifySupport(InDevice))
 {
 	FMemory::Memzero(BufferBindingInfo);
-
 	for (uint32 Index = 0; Index < VulkanBindless::NumBindlessSets; Index++)
 	{
 		BufferIndices[Index] = Index;
@@ -148,59 +307,552 @@ FVulkanBindlessDescriptorManager::FVulkanBindlessDescriptorManager(FVulkanDevice
 
 FVulkanBindlessDescriptorManager::~FVulkanBindlessDescriptorManager()
 {
-	check(BindlessPipelineLayout == VK_NULL_HANDLE);
+	checkf(BindlessPipelineLayout == VK_NULL_HANDLE, TEXT("DeInit() was not called on FVulkanBindlessDescriptorManager!"));
 }
 
 void FVulkanBindlessDescriptorManager::Deinit()
 {
-}
+	const VkDevice DeviceHandle = Device->GetInstanceHandle();
 
-bool FVulkanBindlessDescriptorManager::VerifySupport(FVulkanDevice* Device)
-{
-	return false;
+	if (bIsSupported)
+	{
+		VulkanRHI::vkDestroyPipelineLayout(DeviceHandle, BindlessPipelineLayout, VULKAN_CPU_ALLOCATOR);
+		BindlessPipelineLayout = VK_NULL_HANDLE;
+
+		auto DestroyBindlessState = [DeviceHandle](BindlessSetState& State)
+		{
+			VulkanRHI::vkDestroyDescriptorSetLayout(DeviceHandle, State.DescriptorSetLayout, VULKAN_CPU_ALLOCATOR);
+			State.DescriptorSetLayout = VK_NULL_HANDLE;
+
+			VulkanRHI::vkDestroyBuffer(DeviceHandle, State.BufferHandle, VULKAN_CPU_ALLOCATOR);
+			State.BufferHandle = VK_NULL_HANDLE;
+
+			VulkanRHI::vkUnmapMemory(DeviceHandle, State.MemoryHandle);
+			VulkanRHI::vkFreeMemory(DeviceHandle, State.MemoryHandle, VULKAN_CPU_ALLOCATOR);
+			State.MemoryHandle = VK_NULL_HANDLE;
+		};
+
+		for (uint32 SetIndex = 0; SetIndex < VulkanBindless::NumBindlessSets; ++SetIndex)
+	{
+			BindlessSetState& State = BindlessSetStates[SetIndex];
+			if (State.DescriptorType != VK_DESCRIPTOR_TYPE_MAX_ENUM)
+			{
+				DestroyBindlessState(State);
+			}
+		}
+
+		VulkanRHI::vkDestroyDescriptorSetLayout(DeviceHandle, EmptyDescriptorSetLayout, VULKAN_CPU_ALLOCATOR);
+		EmptyDescriptorSetLayout = VK_NULL_HANDLE;
+	}
 }
 
 void FVulkanBindlessDescriptorManager::Init()
 {
+	if (!bIsSupported)
+	{
+		return;
+	}
+
+	const VkDevice DeviceHandle = Device->GetInstanceHandle();
+	const VkPhysicalDeviceDescriptorBufferPropertiesEXT& DescriptorBufferProperties = Device->GetDescriptorBufferProperties();
+
+	// Create the dummy layout for unsupported descriptor types
+	{
+		VkDescriptorSetLayoutCreateInfo EmptyDescriptorSetLayoutCreateInfo;
+		ZeroVulkanStruct(EmptyDescriptorSetLayoutCreateInfo, VK_STRUCTURE_TYPE_DESCRIPTOR_SET_LAYOUT_CREATE_INFO);
+		VERIFYVULKANRESULT(VulkanRHI::vkCreateDescriptorSetLayout(DeviceHandle, &EmptyDescriptorSetLayoutCreateInfo, VULKAN_CPU_ALLOCATOR, &EmptyDescriptorSetLayout));
+	}
+
+	{
+		auto InitBindlessSetState = [&](VkDescriptorType DescriptorType, BindlessSetState& OutState) {
+
+			OutState.DescriptorType = DescriptorType;
+
+			OutState.DescriptorSize = GetDescriptorTypeSize(Device, DescriptorType);
+			checkf((OutState.DescriptorSize > 0), TEXT("Descriptor Type [%s] returned an invalid descriptor size!"), VK_TYPE_TO_STRING(VkDescriptorType, DescriptorType));
+
+			OutState.MaxDescriptorCount = GetInitialDescriptorCount(DescriptorType);
+			checkf((OutState.MaxDescriptorCount > 0), TEXT("Descriptor Type [%s] returned an invalid descriptor count!"), VK_TYPE_TO_STRING(VkDescriptorType, DescriptorType));
+		};
+
+		// Fill the DescriptorSetLayout for a BindlessSetState
+		auto CreateDescriptorSetLayout = [&](const BindlessSetState& State) {
+
+			if (State.DescriptorType == VK_DESCRIPTOR_TYPE_MAX_ENUM)
+			{
+				return EmptyDescriptorSetLayout;
+			}
+			else
+			{
+				VkDescriptorSetLayoutBinding DescriptorSetLayoutBinding;
+				DescriptorSetLayoutBinding.binding = 0;
+				DescriptorSetLayoutBinding.descriptorType = State.DescriptorType;
+				DescriptorSetLayoutBinding.descriptorCount = State.MaxDescriptorCount;  // todo-jn: resizable
+				DescriptorSetLayoutBinding.stageFlags = VK_SHADER_STAGE_ALL;
+				DescriptorSetLayoutBinding.pImmutableSamplers = nullptr; // todo-jn: ImmutableSamplers
+
+				// These flags are implied with descriptor_buffer: VK_DESCRIPTOR_BINDING_UPDATE_AFTER_BIND_BIT,VK_DESCRIPTOR_BINDING_PARTIALLY_BOUND_BIT,VK_DESCRIPTOR_BINDING_UPDATE_UNUSED_WHILE_PENDING_BIT
+				// :todo-jn: add support for VK_DESCRIPTOR_BINDING_VARIABLE_DESCRIPTOR_COUNT_BIT when drivers are fixed to allow for buffers to grow
+				const VkDescriptorBindingFlags DescriptorBindingFlags = 0; // VK_DESCRIPTOR_BINDING_VARIABLE_DESCRIPTOR_COUNT_BIT;
+
+				VkDescriptorSetLayoutBindingFlagsCreateInfo DescriptorSetLayoutBindingFlagsCreateInfo;
+				ZeroVulkanStruct(DescriptorSetLayoutBindingFlagsCreateInfo, VK_STRUCTURE_TYPE_DESCRIPTOR_SET_LAYOUT_BINDING_FLAGS_CREATE_INFO);
+				DescriptorSetLayoutBindingFlagsCreateInfo.bindingCount = 1;
+				DescriptorSetLayoutBindingFlagsCreateInfo.pBindingFlags = &DescriptorBindingFlags;
+
+				VkDescriptorSetLayoutCreateInfo DescriptorSetLayoutCreateInfo;
+				ZeroVulkanStruct(DescriptorSetLayoutCreateInfo, VK_STRUCTURE_TYPE_DESCRIPTOR_SET_LAYOUT_CREATE_INFO);
+				DescriptorSetLayoutCreateInfo.pBindings = &DescriptorSetLayoutBinding;
+				DescriptorSetLayoutCreateInfo.bindingCount = 1;
+				DescriptorSetLayoutCreateInfo.flags = VK_DESCRIPTOR_SET_LAYOUT_CREATE_DESCRIPTOR_BUFFER_BIT_EXT;
+				DescriptorSetLayoutCreateInfo.pNext = &DescriptorSetLayoutBindingFlagsCreateInfo;
+
+				VkDescriptorSetLayout DescriptorSetLayout = VK_NULL_HANDLE;
+				VERIFYVULKANRESULT(VulkanRHI::vkCreateDescriptorSetLayout(DeviceHandle, &DescriptorSetLayoutCreateInfo, VULKAN_CPU_ALLOCATOR, &DescriptorSetLayout));
+				return DescriptorSetLayout;
+			}
+		};
+
+		// Uniform buffer descriptor set layout differ from the other resources, we reserve a fixed number of descriptors per stage for each draw/dispatch
+		// todo-jn: this could be compacted..
+		auto CreateShaderStageUniformBufferLayout = [DeviceHandle]() {
+
+			const uint32 NumTotalBindings = VulkanBindless::MaxUniformBuffersPerStage * ShaderStage::MaxNumSets;
+
+			TArray<VkDescriptorSetLayoutBinding> DescriptorSetLayoutBindings;
+			DescriptorSetLayoutBindings.SetNumZeroed(NumTotalBindings);
+			for (uint32 BindingIndex = 0; BindingIndex < NumTotalBindings; ++BindingIndex)
+			{
+				DescriptorSetLayoutBindings[BindingIndex].binding = BindingIndex;
+				DescriptorSetLayoutBindings[BindingIndex].descriptorType = VK_DESCRIPTOR_TYPE_UNIFORM_BUFFER;
+				DescriptorSetLayoutBindings[BindingIndex].descriptorCount = 1;
+				DescriptorSetLayoutBindings[BindingIndex].stageFlags = VK_SHADER_STAGE_ALL;
+			}
+
+			VkDescriptorSetLayoutCreateInfo DescriptorSetLayoutCreateInfo;
+			ZeroVulkanStruct(DescriptorSetLayoutCreateInfo, VK_STRUCTURE_TYPE_DESCRIPTOR_SET_LAYOUT_CREATE_INFO);
+			DescriptorSetLayoutCreateInfo.pBindings = DescriptorSetLayoutBindings.GetData();
+			DescriptorSetLayoutCreateInfo.bindingCount = NumTotalBindings;
+			DescriptorSetLayoutCreateInfo.flags = VK_DESCRIPTOR_SET_LAYOUT_CREATE_DESCRIPTOR_BUFFER_BIT_EXT;
+			DescriptorSetLayoutCreateInfo.pNext = nullptr;
+
+			VkDescriptorSetLayout DescriptorSetLayout = VK_NULL_HANDLE;
+			VERIFYVULKANRESULT(VulkanRHI::vkCreateDescriptorSetLayout(DeviceHandle, &DescriptorSetLayoutCreateInfo, VULKAN_CPU_ALLOCATOR, &DescriptorSetLayout));
+			return DescriptorSetLayout;
+		};
+
+		// Create the descriptor buffer for a BindlessSetState
+		auto CreateDescriptorBuffer = [&](BindlessSetState& InOutState, VkDescriptorBufferBindingInfoEXT& OutBufferBindingInfo) {
+
+			// Skip unsupported descriptors
+			if (InOutState.DescriptorType == VK_DESCRIPTOR_TYPE_MAX_ENUM)
+			{
+				return 0u;
+			}
+
+			const bool IsUniformBufferSet = (InOutState.DescriptorType == VK_DESCRIPTOR_TYPE_UNIFORM_BUFFER);
+			const bool IsSamplerSet = (InOutState.DescriptorType == VK_DESCRIPTOR_TYPE_SAMPLER);
+			const VkBufferUsageFlags BufferUsageFlags =
+				VK_BUFFER_USAGE_SHADER_DEVICE_ADDRESS_BIT |
+				VK_BUFFER_USAGE_TRANSFER_SRC_BIT |
+				VK_BUFFER_USAGE_TRANSFER_DST_BIT |
+				(IsSamplerSet ? VK_BUFFER_USAGE_SAMPLER_DESCRIPTOR_BUFFER_BIT_EXT : VK_BUFFER_USAGE_RESOURCE_DESCRIPTOR_BUFFER_BIT_EXT);
+
+			const uint32 DescriptorBufferSize = InOutState.DescriptorSize * InOutState.MaxDescriptorCount;
+			InOutState.DebugDescriptors.SetNumZeroed(DescriptorBufferSize);
+
+			VkDeviceSize LayoutSizeInBytes = 0;
+			VulkanRHI::vkGetDescriptorSetLayoutSizeEXT(DeviceHandle, InOutState.DescriptorSetLayout, &LayoutSizeInBytes);
+			if (IsUniformBufferSet)
+			{
+				// todo-jn: We're picky about uniform buffers values for now to allow for shortcuts...
+				check(LayoutSizeInBytes == (ShaderStage::MaxNumSets * VulkanBindless::MaxUniformBuffersPerStage * InOutState.DescriptorSize));
+				check((LayoutSizeInBytes % DescriptorBufferProperties.descriptorBufferOffsetAlignment) == 0);
+				check((InOutState.MaxDescriptorCount % VulkanBindless::MaxUniformBuffersPerStage) == 0);
+			}
+			else
+			{
+				// Double check that the layout follows the rules for a single binding with an array of descriptors that are tightly packed
+				check(LayoutSizeInBytes == (InOutState.MaxDescriptorCount * InOutState.DescriptorSize));
+			}
+
+			if (IsSamplerSet)
+			{
+				checkf(DescriptorBufferSize < DescriptorBufferProperties.samplerDescriptorBufferAddressSpaceSize,
+					TEXT("Sampler descriptor buffer size [%u] exceeded maximum [%llu]."),
+					DescriptorBufferSize, DescriptorBufferProperties.samplerDescriptorBufferAddressSpaceSize);
+			}
+
+			// Create descriptor buffer
+			{
+				VkBufferCreateInfo BufferCreateInfo;
+				ZeroVulkanStruct(BufferCreateInfo, VK_STRUCTURE_TYPE_BUFFER_CREATE_INFO);
+				BufferCreateInfo.size = DescriptorBufferSize;
+				BufferCreateInfo.usage = BufferUsageFlags;
+				VERIFYVULKANRESULT(VulkanRHI::vkCreateBuffer(DeviceHandle, &BufferCreateInfo, VULKAN_CPU_ALLOCATOR, &InOutState.BufferHandle));
+			}
+
+			// Allocate buffer memory, bind and map
+			{
+				VkMemoryRequirements BufferMemoryReqs;
+				FMemory::Memzero(BufferMemoryReqs);
+				VulkanRHI::vkGetBufferMemoryRequirements(DeviceHandle, InOutState.BufferHandle, &BufferMemoryReqs);
+				check(BufferMemoryReqs.size >= DescriptorBufferSize);
+
+				uint32 MemoryTypeIndex = 0;
+				VERIFYVULKANRESULT(Device->GetDeviceMemoryManager().GetMemoryTypeFromProperties(BufferMemoryReqs.memoryTypeBits, GetDescriptorBufferMemoryType(Device), &MemoryTypeIndex));
+
+				VkMemoryAllocateFlagsInfo FlagsInfo;
+				ZeroVulkanStruct(FlagsInfo, VK_STRUCTURE_TYPE_MEMORY_ALLOCATE_FLAGS_INFO);
+				FlagsInfo.flags = VK_MEMORY_ALLOCATE_DEVICE_ADDRESS_BIT;
+
+				VkMemoryAllocateInfo AllocateInfo;
+				ZeroVulkanStruct(AllocateInfo, VK_STRUCTURE_TYPE_MEMORY_ALLOCATE_INFO);
+				AllocateInfo.allocationSize = BufferMemoryReqs.size;
+				AllocateInfo.memoryTypeIndex = MemoryTypeIndex;
+				AllocateInfo.pNext = &FlagsInfo;
+
+				VERIFYVULKANRESULT(VulkanRHI::vkAllocateMemory(DeviceHandle, &AllocateInfo, VULKAN_CPU_ALLOCATOR, &InOutState.MemoryHandle));
+				VERIFYVULKANRESULT(VulkanRHI::vkBindBufferMemory(DeviceHandle, InOutState.BufferHandle, InOutState.MemoryHandle, 0));
+				VERIFYVULKANRESULT(VulkanRHI::vkMapMemory(DeviceHandle, InOutState.MemoryHandle, 0, VK_WHOLE_SIZE, 0, (void**)&InOutState.MappedPointer));
+				FMemory::Memzero(InOutState.MappedPointer, AllocateInfo.allocationSize);
+			}
+
+			// Setup the binding info
+			{
+				VkBufferDeviceAddressInfo AddressInfo;
+				ZeroVulkanStruct(AddressInfo, VK_STRUCTURE_TYPE_BUFFER_DEVICE_ADDRESS_INFO);
+				AddressInfo.buffer = InOutState.BufferHandle;
+
+				ZeroVulkanStruct(OutBufferBindingInfo, VK_STRUCTURE_TYPE_DESCRIPTOR_BUFFER_BINDING_INFO_EXT);
+				OutBufferBindingInfo.address = VulkanRHI::vkGetBufferDeviceAddressKHR(DeviceHandle, &AddressInfo);
+				OutBufferBindingInfo.usage = BufferUsageFlags;
+			}
+
+			return IsSamplerSet ? 0u : DescriptorBufferSize;
+		};
+
+		// Fill in one state for each descriptor type
+		uint32 TotalResourceDescriptorBufferSize = 0;
+		for (uint32 SetIndex = 0; SetIndex < VulkanBindless::NumBindlessSets; ++SetIndex)
+		{
+			// Skip anything we don't support
+			if (SetIndex == VulkanBindless::BindlessAccelerationStructureSet)
+			{
+#if VULKAN_RHI_RAYTRACING
+				const bool bHasRaytracingExtensions = Device->GetOptionalExtensions().HasRaytracingExtensions();
+#else
+				const bool bHasRaytracingExtensions = false;
+#endif
+
+				if (!bHasRaytracingExtensions)
+				{
+					continue;
+				}
+			}
+
+			BindlessSetState& State = BindlessSetStates[SetIndex];
+			InitBindlessSetState(GetDescriptorTypeForSetIndex(SetIndex), State);
+			State.DescriptorSetLayout = (SetIndex == VulkanBindless::BindlessUniformBufferSet) ? CreateShaderStageUniformBufferLayout() : CreateDescriptorSetLayout(State);
+			TotalResourceDescriptorBufferSize += CreateDescriptorBuffer(State, BufferBindingInfo[SetIndex]);
+		}
+
+		checkf(TotalResourceDescriptorBufferSize < DescriptorBufferProperties.resourceDescriptorBufferAddressSpaceSize,
+			TEXT("Combined resource descriptor buffer size of [%u] exceeded maximum [%llu]."),
+			TotalResourceDescriptorBufferSize, DescriptorBufferProperties.resourceDescriptorBufferAddressSpaceSize);
+	}
+
+	// Now create the single pipeline layout used by everything
+	{
+		VkDescriptorSetLayout DescriptorSetLayouts[VulkanBindless::NumBindlessSets ];
+		for (int32 LayoutIndex = 0; LayoutIndex < VulkanBindless::NumBindlessSets; ++LayoutIndex)
+		{
+			const BindlessSetState& State = BindlessSetStates[LayoutIndex];
+			DescriptorSetLayouts[LayoutIndex] = State.DescriptorSetLayout;
+		}
+
+		VkPipelineLayoutCreateInfo PipelineLayoutCreateInfo;
+		ZeroVulkanStruct(PipelineLayoutCreateInfo, VK_STRUCTURE_TYPE_PIPELINE_LAYOUT_CREATE_INFO);
+		PipelineLayoutCreateInfo.setLayoutCount = VulkanBindless::NumBindlessSets;
+		PipelineLayoutCreateInfo.pSetLayouts = DescriptorSetLayouts;
+		VERIFYVULKANRESULT(VulkanRHI::vkCreatePipelineLayout(DeviceHandle, &PipelineLayoutCreateInfo, VULKAN_CPU_ALLOCATOR, &BindlessPipelineLayout));
+	}
 }
 
 void FVulkanBindlessDescriptorManager::BindDescriptorBuffers(VkCommandBuffer CommandBuffer, VkPipelineStageFlags SupportedStages)
 {
+	checkf(bIsSupported, TEXT("Trying to BindDescriptorBuffers but bindless is not supported!"));
+
+	VulkanRHI::vkCmdBindDescriptorBuffersEXT(CommandBuffer, VulkanBindless::NumBindlessSets, BufferBindingInfo);
+
+	if (GVulkanBindlessBufferOffsetUpdates != 0)
+	{
+		VkDeviceSize BufferOffsets[VulkanBindless::NumBindlessSets];
+		FMemory::Memzero(BufferOffsets);
+		if (SupportedStages & VK_PIPELINE_STAGE_FRAGMENT_SHADER_BIT)
+		{
+			VulkanRHI::vkCmdSetDescriptorBufferOffsetsEXT(CommandBuffer, VK_PIPELINE_BIND_POINT_GRAPHICS, BindlessPipelineLayout, 0, VulkanBindless::NumBindlessSets, BufferIndices, BufferOffsets);
+		}
+		if (SupportedStages & VK_PIPELINE_STAGE_COMPUTE_SHADER_BIT)
+		{
+			VulkanRHI::vkCmdSetDescriptorBufferOffsetsEXT(CommandBuffer, VK_PIPELINE_BIND_POINT_COMPUTE, BindlessPipelineLayout, 0, VulkanBindless::NumBindlessSets, BufferIndices, BufferOffsets);
+		}
+		if (SupportedStages & VK_PIPELINE_STAGE_RAY_TRACING_SHADER_BIT_KHR)
+		{
+			VulkanRHI::vkCmdSetDescriptorBufferOffsetsEXT(CommandBuffer, VK_PIPELINE_BIND_POINT_RAY_TRACING_KHR, BindlessPipelineLayout, 0, VulkanBindless::NumBindlessSets, BufferIndices, BufferOffsets);
+		}
+	}
 }
 
 void FVulkanBindlessDescriptorManager::RegisterUniformBuffers(VkCommandBuffer CommandBuffer, VkPipelineBindPoint BindPoint, const FUniformBufferDescriptorArrays& StageUBs)
 {
-}
+	checkf(bIsSupported, TEXT("Trying to RegisterUniformBuffers but bindless is not supported!"));
+
+	SCOPED_NAMED_EVENT(FVulkanBindlessDescriptorManager_RegisterUniformBuffers, FColor::Purple);
+
+	const uint8 SetIndex = GetIndexForDescriptorType(VK_DESCRIPTOR_TYPE_UNIFORM_BUFFER);
+	BindlessSetState& BindlessUniformBufferSetState = BindlessSetStates[SetIndex];
+
+	// :todo-jn: Current uniform buffer layout is a bit wasteful with all the skipped bindings...
+	const uint32 BlockDescriptorCount = VulkanBindless::MaxUniformBuffersPerStage * ShaderStage::MaxNumSets;
+	const uint32 BlockSize = BlockDescriptorCount * BindlessUniformBufferSetState.DescriptorSize;
+	// Leave the first block always zeroed for easier debugging
+	const uint32 FirstDescriptorIndex = BlockDescriptorCount + (CurrentUniformBufferDescriptorIndex.fetch_add(BlockDescriptorCount) % (BindlessUniformBufferSetState.MaxDescriptorCount - (2*BlockDescriptorCount)));
+	const VkDeviceSize FirstDescriptorByteOffset = FirstDescriptorIndex * BindlessUniformBufferSetState.DescriptorSize;
+
+	VkDeviceSize BufferOffsets[VulkanBindless::NumBindlessSets];
+	FMemory::Memzero(BufferOffsets);
+	BufferOffsets[VulkanBindless::BindlessUniformBufferSet] = FirstDescriptorByteOffset;
+	checkSlow(FirstDescriptorByteOffset % Device->GetDescriptorBufferProperties().descriptorBufferOffsetAlignment == 0);
+
+	// :todo-jn: Clear them for easier debugging for now
+	FMemory::Memzero(&BindlessUniformBufferSetState.DebugDescriptors[FirstDescriptorByteOffset], BlockSize);
+
+	for (int32 StageIndex = 0; StageIndex < ShaderStage::NumStages; ++StageIndex)
+	{
+		const TArray<VkDescriptorAddressInfoEXT>& DescriptorAddressInfos = StageUBs[StageIndex];
+
+		if (DescriptorAddressInfos.Num())
+		{
+			check(DescriptorAddressInfos.Num() <= VulkanBindless::MaxUniformBuffersPerStage);
+			const int32 StageOffset = StageIndex * VulkanBindless::MaxUniformBuffersPerStage;
+
+			for (int32 DescriptorAddressInfoIndex = 0; DescriptorAddressInfoIndex < DescriptorAddressInfos.Num(); DescriptorAddressInfoIndex++)
+			{
+				const VkDescriptorAddressInfoEXT& DescriptorAddressInfo = DescriptorAddressInfos[DescriptorAddressInfoIndex];
+				checkSlow(DescriptorAddressInfo.sType != 0);  // make sure it was filled
+				checkSlow((DescriptorAddressInfo.range % 16) == 0);  // :todo-jn: make sure we don't trip on driver bug, to be removed on next release
+
+				const int32 BindingIndex = StageOffset + DescriptorAddressInfoIndex;
+				VkDeviceSize BindingByteOffset = 0;
+#if UE_BUILD_DEBUG
+				VulkanRHI::vkGetDescriptorSetLayoutBindingOffsetEXT(Device->GetInstanceHandle(), BindlessUniformBufferSetState.DescriptorSetLayout, BindingIndex, &BindingByteOffset);
+				check(BindingByteOffset == BindingIndex * BindlessUniformBufferSetState.DescriptorSize);
+#else
+				BindingByteOffset = (BindingIndex * BindlessUniformBufferSetState.DescriptorSize);
+#endif
+
+				VkDescriptorGetInfoEXT Info;
+				ZeroVulkanStruct(Info, VK_STRUCTURE_TYPE_DESCRIPTOR_GET_INFO_EXT);
+				Info.type = VK_DESCRIPTOR_TYPE_UNIFORM_BUFFER;
+				Info.data.pUniformBuffer = &DescriptorAddressInfo;
+				VulkanRHI::vkGetDescriptorEXT(Device->GetInstanceHandle(), &Info, BindlessUniformBufferSetState.DescriptorSize, &BindlessUniformBufferSetState.DebugDescriptors[FirstDescriptorByteOffset + BindingByteOffset]);
+			}
+		}
+	}
+
+	// Copy all of them at once
+	FMemory::Memcpy(&BindlessUniformBufferSetState.MappedPointer[FirstDescriptorByteOffset], &BindlessUniformBufferSetState.DebugDescriptors[FirstDescriptorByteOffset], BlockSize);
+
+	{
+		SCOPED_NAMED_EVENT(vkCmdSetDescriptorBufferOffsetsEXT, FColor::Purple);
 
+		if (GVulkanBindlessRebindBuffers)
+		{
+			VulkanRHI::vkCmdBindDescriptorBuffersEXT(CommandBuffer, VulkanBindless::NumBindlessSets, BufferBindingInfo);
+		}
+
+		const bool bSetAllOffsets = (GVulkanBindlessBufferOffsetUpdates == 0);
+		if (bSetAllOffsets)
+		{
+			VulkanRHI::vkCmdSetDescriptorBufferOffsetsEXT(CommandBuffer, BindPoint, BindlessPipelineLayout, 0u, VulkanBindless::NumBindlessSets, BufferIndices, BufferOffsets);
+		}
+		else
+		{
+			VulkanRHI::vkCmdSetDescriptorBufferOffsetsEXT(CommandBuffer, BindPoint, BindlessPipelineLayout, SetIndex, 1u, &BufferIndices[SetIndex], &BufferOffsets[SetIndex]);
+		}
+	}
+}
 
 FRHIDescriptorHandle FVulkanBindlessDescriptorManager::RegisterSampler(VkSampler VulkanSampler)
 {
-	return FRHIDescriptorHandle();
+	if (!bIsSupported)
+	{
+		return FRHIDescriptorHandle();
+	}
+
+	const uint8 SetIndex = GetIndexForDescriptorType(VK_DESCRIPTOR_TYPE_SAMPLER);
+	BindlessSetState& State = BindlessSetStates[SetIndex];
+
+	const uint32 SamplerIndex = State.PeakDescriptorCount++;
+	checkf(SamplerIndex < State.MaxDescriptorCount, TEXT("You need to grow the sampler array size!"));
+	const uint32 ByteOffset = SamplerIndex * State.DescriptorSize;
+
+	VkDescriptorGetInfoEXT Info;
+	ZeroVulkanStruct(Info, VK_STRUCTURE_TYPE_DESCRIPTOR_GET_INFO_EXT);
+	Info.type = State.DescriptorType;
+	Info.data.pSampler = &VulkanSampler;
+	VulkanRHI::vkGetDescriptorEXT(Device->GetInstanceHandle(), &Info, State.DescriptorSize, &State.DebugDescriptors[ByteOffset]);
+
+	FMemory::Memcpy(&State.MappedPointer[ByteOffset], &State.DebugDescriptors[ByteOffset], State.DescriptorSize);
+
+	SET_DWORD_STAT(STAT_VulkanBindlessPeakSampler, State.PeakDescriptorCount);
+	return FRHIDescriptorHandle(SetIndex, SamplerIndex);
 }
 
 FRHIDescriptorHandle FVulkanBindlessDescriptorManager::RegisterImage(VkImageView ImageView, VkDescriptorType DescriptorType, bool bIsDepthStencil)
 {
-	return FRHIDescriptorHandle();
+	if (!bIsSupported)
+	{
+		return FRHIDescriptorHandle();
+	}
+
+	check((DescriptorType == VK_DESCRIPTOR_TYPE_SAMPLED_IMAGE) || (DescriptorType == VK_DESCRIPTOR_TYPE_STORAGE_IMAGE));
+
+	const uint8 SetIndex = GetIndexForDescriptorType(DescriptorType);
+	BindlessSetState& State = BindlessSetStates[SetIndex];
+	const uint32 ResourceIndex = GetFreeResourceIndex(State);
+	const uint32 ByteOffset = ResourceIndex * State.DescriptorSize;
+
+	VkDescriptorImageInfo DescriptorImageInfo;
+	FMemory::Memzero(DescriptorImageInfo);
+	DescriptorImageInfo.imageView = ImageView;
+	DescriptorImageInfo.imageLayout = (DescriptorType == VK_DESCRIPTOR_TYPE_STORAGE_IMAGE) ? VK_IMAGE_LAYOUT_GENERAL :
+		(bIsDepthStencil ? VK_IMAGE_LAYOUT_READ_ONLY_OPTIMAL : VK_IMAGE_LAYOUT_SHADER_READ_ONLY_OPTIMAL);
+
+	VkDescriptorGetInfoEXT Info;
+	ZeroVulkanStruct(Info, VK_STRUCTURE_TYPE_DESCRIPTOR_GET_INFO_EXT);
+	Info.type = State.DescriptorType;
+	Info.data.pSampledImage = &DescriptorImageInfo;  // same pointer for storage, it's a union
+	VulkanRHI::vkGetDescriptorEXT(Device->GetInstanceHandle(), &Info, State.DescriptorSize, &State.DebugDescriptors[ByteOffset]);
+
+	FMemory::Memcpy(&State.MappedPointer[ByteOffset], &State.DebugDescriptors[ByteOffset], State.DescriptorSize);
+
+	SET_DWORD_STAT(STAT_VulkanBindlessPeakImage, State.PeakDescriptorCount);
+	return FRHIDescriptorHandle(SetIndex, ResourceIndex);
 }
 
 FRHIDescriptorHandle FVulkanBindlessDescriptorManager::RegisterBuffer(VkBuffer VulkanBuffer, VkDeviceSize BufferOffset, VkDeviceSize BufferSize, VkDescriptorType DescriptorType)
 {
-	return FRHIDescriptorHandle();
+	if (!bIsSupported)
+	{
+		return FRHIDescriptorHandle();
+	}
+
+	check(/*(DescriptorType == VK_DESCRIPTOR_TYPE_UNIFORM_BUFFER) ||*/ (DescriptorType == VK_DESCRIPTOR_TYPE_STORAGE_BUFFER));
+
+	const uint8 SetIndex = GetIndexForDescriptorType(DescriptorType);
+	BindlessSetState& State = BindlessSetStates[SetIndex];
+	const uint32 ResourceIndex = GetFreeResourceIndex(State);
+	const uint32 ByteOffset = ResourceIndex * State.DescriptorSize;
+
+	VkBufferDeviceAddressInfo BufferInfo;
+	ZeroVulkanStruct(BufferInfo, VK_STRUCTURE_TYPE_BUFFER_DEVICE_ADDRESS_INFO);
+	BufferInfo.buffer = VulkanBuffer;
+	VkDeviceAddress BufferAddress = VulkanRHI::vkGetBufferDeviceAddressKHR(Device->GetInstanceHandle(), &BufferInfo);
+
+	VkDescriptorAddressInfoEXT AddressInfo;
+	ZeroVulkanStruct(AddressInfo, VK_STRUCTURE_TYPE_DESCRIPTOR_ADDRESS_INFO_EXT);
+	AddressInfo.address = BufferAddress + BufferOffset;
+	AddressInfo.range = BufferSize;
+
+	VkDescriptorGetInfoEXT Info;
+	ZeroVulkanStruct(Info, VK_STRUCTURE_TYPE_DESCRIPTOR_GET_INFO_EXT);
+	Info.type = State.DescriptorType;
+	Info.data.pUniformBuffer = &AddressInfo;  // same pointer for storage, it's a union
+	VulkanRHI::vkGetDescriptorEXT(Device->GetInstanceHandle(), &Info, State.DescriptorSize, &State.DebugDescriptors[ByteOffset]);
+
+	FMemory::Memcpy(&State.MappedPointer[ByteOffset], &State.DebugDescriptors[ByteOffset], State.DescriptorSize);
+
+	SET_DWORD_STAT(STAT_VulkanBindlessPeakBuffer, State.PeakDescriptorCount);
+	return FRHIDescriptorHandle(SetIndex, ResourceIndex);
 }
 
 FRHIDescriptorHandle FVulkanBindlessDescriptorManager::RegisterTexelBuffer(const VkBufferViewCreateInfo& ViewInfo, VkDescriptorType DescriptorType)
 {
-	return FRHIDescriptorHandle();
+	if (!bIsSupported)
+	{
+		return FRHIDescriptorHandle();
+	}
+
+	check((DescriptorType == VK_DESCRIPTOR_TYPE_UNIFORM_TEXEL_BUFFER) || (DescriptorType == VK_DESCRIPTOR_TYPE_STORAGE_TEXEL_BUFFER));
+
+	const uint8 SetIndex = GetIndexForDescriptorType(DescriptorType);
+	BindlessSetState& State = BindlessSetStates[SetIndex];
+	const uint32 ResourceIndex = GetFreeResourceIndex(State);
+	const uint32 ByteOffset = ResourceIndex * State.DescriptorSize;
+
+	VkBufferDeviceAddressInfo BufferInfo;
+	ZeroVulkanStruct(BufferInfo, VK_STRUCTURE_TYPE_BUFFER_DEVICE_ADDRESS_INFO);
+	BufferInfo.buffer = ViewInfo.buffer;
+	VkDeviceAddress BufferAddress = VulkanRHI::vkGetBufferDeviceAddressKHR(Device->GetInstanceHandle(), &BufferInfo);
+
+	VkDescriptorAddressInfoEXT AddressInfo;
+	ZeroVulkanStruct(AddressInfo, VK_STRUCTURE_TYPE_DESCRIPTOR_ADDRESS_INFO_EXT);
+	AddressInfo.address = BufferAddress + ViewInfo.offset;
+	AddressInfo.range = ViewInfo.range;
+	AddressInfo.format = ViewInfo.format;
+
+	VkDescriptorGetInfoEXT Info;
+	ZeroVulkanStruct(Info, VK_STRUCTURE_TYPE_DESCRIPTOR_GET_INFO_EXT);
+	Info.type = State.DescriptorType;
+	Info.data.pUniformTexelBuffer = &AddressInfo;  // same pointer for storage, it's a union
+	VulkanRHI::vkGetDescriptorEXT(Device->GetInstanceHandle(), &Info, State.DescriptorSize, &State.DebugDescriptors[ByteOffset]);
+
+	FMemory::Memcpy(&State.MappedPointer[ByteOffset], &State.DebugDescriptors[ByteOffset], State.DescriptorSize);
+
+	SET_DWORD_STAT(STAT_VulkanBindlessPeakTexelBuffer, State.PeakDescriptorCount);
+	return FRHIDescriptorHandle(SetIndex, ResourceIndex);
 }
 
 FRHIDescriptorHandle FVulkanBindlessDescriptorManager::RegisterAccelerationStructure(VkAccelerationStructureKHR AccelerationStructure)
 {
+#if VULKAN_RHI_RAYTRACING
+	if (!bIsSupported)
+	{
+		return FRHIDescriptorHandle();
+	}
+
+	const uint8 SetIndex = GetIndexForDescriptorType(VK_DESCRIPTOR_TYPE_ACCELERATION_STRUCTURE_KHR);
+	BindlessSetState& State = BindlessSetStates[SetIndex];
+	const uint32 ResourceIndex = GetFreeResourceIndex(State);
+	const uint32 ByteOffset = ResourceIndex * State.DescriptorSize;
+
+	VkAccelerationStructureDeviceAddressInfoKHR AccelerationStructureDeviceAddressInfo;
+	ZeroVulkanStruct(AccelerationStructureDeviceAddressInfo, VK_STRUCTURE_TYPE_ACCELERATION_STRUCTURE_DEVICE_ADDRESS_INFO_KHR);
+	AccelerationStructureDeviceAddressInfo.accelerationStructure = AccelerationStructure;
+	VkDeviceAddress BufferAddress = VulkanRHI::vkGetAccelerationStructureDeviceAddressKHR(Device->GetInstanceHandle(), &AccelerationStructureDeviceAddressInfo);
+
+	VkDescriptorGetInfoEXT Info;
+	ZeroVulkanStruct(Info, VK_STRUCTURE_TYPE_DESCRIPTOR_GET_INFO_EXT);
+	Info.type = State.DescriptorType;
+	Info.data.accelerationStructure = BufferAddress;  // same pointer for storage, it's a union
+	VulkanRHI::vkGetDescriptorEXT(Device->GetInstanceHandle(), &Info, State.DescriptorSize, &State.DebugDescriptors[ByteOffset]);
+
+	FMemory::Memcpy(&State.MappedPointer[ByteOffset], &State.DebugDescriptors[ByteOffset], State.DescriptorSize);
+
+	SET_DWORD_STAT(STAT_VulkanBindlessPeakAccelerationStructure, State.PeakDescriptorCount);
+	return FRHIDescriptorHandle(SetIndex, ResourceIndex);
+#else
 	return FRHIDescriptorHandle();
+#endif
 }
 
 
 uint32 FVulkanBindlessDescriptorManager::GetFreeResourceIndex(FVulkanBindlessDescriptorManager::BindlessSetState& State)
 {
+	INC_DWORD_STAT(STAT_VulkanBindlessWritePerFrame);
+
 	{
 		FScopeLock ScopeLock(&State.FreeListCS);
 		if ((State.FreeListHead != MAX_uint32) && (State.PeakDescriptorCount >= State.MaxDescriptorCount)) // todo-jn: temp
@@ -213,6 +865,8 @@ uint32 FVulkanBindlessDescriptorManager::GetFreeResourceIndex(FVulkanBindlessDes
 		}
 	}
 
+	INC_DWORD_STAT(STAT_VulkanBindlessPeakDescriptorCount);
+
 	const uint32 ResourceIndex = State.PeakDescriptorCount++;
 	checkf(ResourceIndex < State.MaxDescriptorCount, TEXT("You need to grow the resource array size for [%s]!"), VK_TYPE_TO_STRING(VkDescriptorType, State.DescriptorType));
 	return ResourceIndex;
@@ -220,8 +874,66 @@ uint32 FVulkanBindlessDescriptorManager::GetFreeResourceIndex(FVulkanBindlessDes
 
 void FVulkanBindlessDescriptorManager::Unregister(FRHIDescriptorHandle DescriptorHandle)
 {
+	if (DescriptorHandle.IsValid())
+	{
+		checkf(bIsSupported, TEXT("Unregistering a valid handle but bindless is not supported!"));
+
+		const uint8 SetIndex = DescriptorHandle.GetRawType();
+		BindlessSetState& State = BindlessSetStates[SetIndex];
+
+		FScopeLock ScopeLock(&State.FreeListCS);
+
+		const uint32 PreviousHead = State.FreeListHead;
+		State.FreeListHead = DescriptorHandle.GetIndex();
+		const uint32 ByteOffset = DescriptorHandle.GetIndex() * State.DescriptorSize;
+		uint32* NewSlot = (uint32*)(&State.DebugDescriptors[ByteOffset]);
+		FMemory::Memzero(NewSlot, State.DescriptorSize); // easier for debugging for now
+		*NewSlot = PreviousHead;
+
+		// Clear the descriptor
+		// todo-jn: invalidate the GPU side?
+	}
 }
 
 void FVulkanBindlessDescriptorManager::CopyDescriptor(VkCommandBuffer CommandBuffer, FRHIDescriptorHandle DstHandle, const FRHIDescriptorHandle SrcHandle)
 {
+	checkf(bIsSupported, TEXT("Trying to CopyDescriptor but bindless is not supported!"));
+	checkf(DstHandle.IsValid(), TEXT("Destination handle for CopyDescriptor not valid!"));
+	checkf(SrcHandle.IsValid(), TEXT("Source handle for CopyDescriptor not valid!"));
+	checkf(DstHandle.GetRawType() == SrcHandle.GetRawType(), TEXT("Handles types for CopyDescriptor do not match!"));
+
+	const uint8 SetIndex = DstHandle.GetRawType();
+	BindlessSetState& State = BindlessSetStates[SetIndex];
+	const uint32 SrcByteOffset = SrcHandle.GetIndex() * State.DescriptorSize;
+	const uint32 DstByteOffset = DstHandle.GetIndex() * State.DescriptorSize;
+
+	FMemory::Memcpy(&State.DebugDescriptors[DstByteOffset], &State.DebugDescriptors[SrcByteOffset], State.DescriptorSize);
+
+	// Can debug using a copy directly instead of using GPU, timing will not be correct however
+	//FMemory::Memcpy(&State.MappedPointer[DstByteOffset], &State.DebugDescriptors[DstByteOffset], State.DescriptorSize);
+
+	VkMemoryBarrier2 MemoryBarrier;
+	ZeroVulkanStruct(MemoryBarrier, VK_STRUCTURE_TYPE_MEMORY_BARRIER_2);
+	MemoryBarrier.srcStageMask = VK_PIPELINE_STAGE_2_ALL_COMMANDS_BIT;
+	MemoryBarrier.srcAccessMask = VK_ACCESS_2_DESCRIPTOR_BUFFER_READ_BIT_EXT;
+	MemoryBarrier.dstStageMask = VK_PIPELINE_STAGE_2_TRANSFER_BIT;
+	MemoryBarrier.dstAccessMask = VK_ACCESS_2_TRANSFER_READ_BIT | VK_ACCESS_2_TRANSFER_WRITE_BIT;
+
+	VkDependencyInfo DependencyInfo;
+	ZeroVulkanStruct(DependencyInfo, VK_STRUCTURE_TYPE_DEPENDENCY_INFO);
+	DependencyInfo.memoryBarrierCount = 1;
+	DependencyInfo.pMemoryBarriers = &MemoryBarrier;
+	VulkanRHI::vkCmdPipelineBarrier2KHR(CommandBuffer, &DependencyInfo);
+
+	VkBufferCopy Region = {};
+	Region.srcOffset = SrcByteOffset;
+	Region.dstOffset = DstByteOffset;
+	Region.size = State.DescriptorSize;
+	VulkanRHI::vkCmdCopyBuffer(CommandBuffer, State.BufferHandle, State.BufferHandle, 1, &Region);
+
+	MemoryBarrier.srcStageMask = VK_PIPELINE_STAGE_2_TRANSFER_BIT;
+	MemoryBarrier.srcAccessMask = VK_ACCESS_2_TRANSFER_READ_BIT | VK_ACCESS_2_TRANSFER_WRITE_BIT;
+	MemoryBarrier.dstStageMask = VK_PIPELINE_STAGE_2_ALL_COMMANDS_BIT;
+	MemoryBarrier.dstAccessMask = VK_ACCESS_2_DESCRIPTOR_BUFFER_READ_BIT_EXT;
+	VulkanRHI::vkCmdPipelineBarrier2KHR(CommandBuffer, &DependencyInfo);
 }
