Fixed: a typo in the log message

diff --git a/Engine/Source/Runtime/VulkanRHI/Private/Android/VulkanAndroidPlatform.cpp b/Engine/Source/Runtime/VulkanRHI/Private/Android/VulkanAndroidPlatform.cpp
--- a/Engine/Source/Runtime/VulkanRHI/Private/Android/VulkanAndroidPlatform.cpp
+++ b/Engine/Source/Runtime/VulkanRHI/Private/Android/VulkanAndroidPlatform.cpp
@@ -1559,7 +1559,7 @@ void FVulkanAndroidPlatform::SetupImageMemoryRequirementWorkaround(const FVulkan
 
 		if (AFBCWorkaroundOption != 0)
 		{
-			UE_LOG(LogRHI, Display, TEXT("Enabling workaround to reduce memory requrement for BGRA textures (%s flag). 128x128 - 8 Mips BGRA texture: %u KiB -> %u KiB"),
+			UE_LOG(LogRHI, Display, TEXT("Enabling workaround to reduce memory requirement for BGRA textures (%s flag). 128x128 - 8 Mips BGRA texture: %u KiB -> %u KiB"),
 				AFBCWorkaroundOption == 1 ? TEXT("MUTABLE") : TEXT("STORAGE"),
 				Image0Mem.size / 1024,
 				AFBCWorkaroundOption == 1 ? ImageMutableMem.size / 1024 : ImageStorageMem.size / 1024
@@ -1596,7 +1596,7 @@ void FVulkanAndroidPlatform::SetupImageMemoryRequirementWorkaround(const FVulkan
 		{
 			ASTCWorkaroundOption = 1;
 
-			UE_LOG(LogRHI, Display, TEXT("Enabling workaround to reduce memory requrement for ASTC textures (VK_IMAGE_TILING_LINEAR). 128x128 - 8 Mips ASTC_8x8 texture: %u KiB -> %u KiB"),
+			UE_LOG(LogRHI, Display, TEXT("Enabling workaround to reduce memory requirement for ASTC textures (VK_IMAGE_TILING_LINEAR). 128x128 - 8 Mips ASTC_8x8 texture: %u KiB -> %u KiB"),
 				ImageOptimalMem_ASTC.size / 1024,
 				ImageLinearMem_ASTC.size / 1024
 			);
