VulkanRHI: Fix issue in FTransitionProcessor::ApplyTracking when GVulkanAutoCorrectUnknownLayouts is enabled.  As observed on the sky cube map being transitioned with multiple transitions: the first transition would read the existing state from the fallback tracker (in the queue) and then write its new state in the cmdbuffer's tracker but only for that subresource (leaving the others in unknown).  Following subresource transitions would discard the data because they would read unknown states from the cmdbuffer's tracker instead of going to the fallback again for the real state.

#rnx

diff --git a/Engine/Source/Runtime/VulkanRHI/Private/VulkanBarriers.h b/Engine/Source/Runtime/VulkanRHI/Private/VulkanBarriers.h
--- a/Engine/Source/Runtime/VulkanRHI/Private/VulkanBarriers.h
+++ b/Engine/Source/Runtime/VulkanRHI/Private/VulkanBarriers.h
@@ -144,6 +144,12 @@ public:
 		if (!Layout && Fallback)
 		{
 			Layout = Fallback->GetFullLayout(VulkanTexture, false);
+
+			// If the layout was found in the fallback, carry it forward to our current manager for future tracking
+			if (Layout)
+			{
+				return &Layouts.Add(VulkanTexture.Image, *Layout);
+			}
 		}
 		
 		if (Layout)
