VulkanRHI:  Clean up: Package extension features and properties we want to keep around into a single struct with a single accessor to reduce clutter as we add more.

#rnx

diff --git a/Engine/Source/Runtime/VulkanRHI/Private/VulkanDescriptorSets.cpp b/Engine/Source/Runtime/VulkanRHI/Private/VulkanDescriptorSets.cpp
--- a/Engine/Source/Runtime/VulkanRHI/Private/VulkanDescriptorSets.cpp
+++ b/Engine/Source/Runtime/VulkanRHI/Private/VulkanDescriptorSets.cpp
@@ -158,7 +158,7 @@ extern TAutoConsoleVariable<int32> GCVarRobustBufferAccess;
 static inline uint32 GetDescriptorTypeSize(FVulkanDevice* Device, VkDescriptorType DescriptorType)
 {
 	const bool bRobustBufferAccess = (GCVarRobustBufferAccess.GetValueOnAnyThread() > 0);
-	const VkPhysicalDeviceDescriptorBufferPropertiesEXT& DescriptorBufferProperties = Device->GetDescriptorBufferProperties();
+	const VkPhysicalDeviceDescriptorBufferPropertiesEXT& DescriptorBufferProperties = Device->GetOptionalExtensionProperties().DescriptorBufferProps;
 
 	switch (DescriptorType)
 	{
@@ -255,7 +255,7 @@ bool FVulkanBindlessDescriptorManager::VerifySupport(FVulkanDevice* Device)
 	{
 		const VkPhysicalDeviceProperties& GpuProps = Device->GetDeviceProperties();
 		const FOptionalVulkanDeviceExtensions& OptionalDeviceExtensions = Device->GetOptionalExtensions();
-		const VkPhysicalDeviceDescriptorBufferPropertiesEXT& DescriptorBufferProperties = Device->GetDescriptorBufferProperties();
+		const VkPhysicalDeviceDescriptorBufferPropertiesEXT& DescriptorBufferProperties = Device->GetOptionalExtensionProperties().DescriptorBufferProps;
 
 		const bool bMeetsExtensionsRequirements =
 			OptionalDeviceExtensions.HasEXTDescriptorIndexing &&
@@ -354,7 +354,7 @@ void FVulkanBindlessDescriptorManager::Init()
 	}
 
 	const VkDevice DeviceHandle = Device->GetInstanceHandle();
-	const VkPhysicalDeviceDescriptorBufferPropertiesEXT& DescriptorBufferProperties = Device->GetDescriptorBufferProperties();
+	const VkPhysicalDeviceDescriptorBufferPropertiesEXT& DescriptorBufferProperties = Device->GetOptionalExtensionProperties().DescriptorBufferProps;
 
 	// Create the dummy layout for unsupported descriptor types
 	{
@@ -623,7 +623,7 @@ void FVulkanBindlessDescriptorManager::RegisterUniformBuffers(VkCommandBuffer Co
 	VkDeviceSize BufferOffsets[VulkanBindless::NumBindlessSets];
 	FMemory::Memzero(BufferOffsets);
 	BufferOffsets[VulkanBindless::BindlessUniformBufferSet] = FirstDescriptorByteOffset;
-	checkSlow(FirstDescriptorByteOffset % Device->GetDescriptorBufferProperties().descriptorBufferOffsetAlignment == 0);
+	checkSlow(FirstDescriptorByteOffset % Device->GetOptionalExtensionProperties().DescriptorBufferProps.descriptorBufferOffsetAlignment == 0);
 
 	// :todo-jn: Clear them for easier debugging for now
 	FMemory::Memzero(&BindlessUniformBufferSetState.DebugDescriptors[FirstDescriptorByteOffset], BlockSize);
diff --git a/Engine/Source/Runtime/VulkanRHI/Private/VulkanDevice.h b/Engine/Source/Runtime/VulkanRHI/Private/VulkanDevice.h
--- a/Engine/Source/Runtime/VulkanRHI/Private/VulkanDevice.h
+++ b/Engine/Source/Runtime/VulkanRHI/Private/VulkanDevice.h
@@ -113,13 +113,22 @@ struct FOptionalVulkanDeviceExtensions
 #endif
 };
 
-#if VULKAN_RHI_RAYTRACING
-struct FRayTracingProperties
+// All the features and properties we need to keep around from extension initialization
+struct FOptionalVulkanDeviceExtensionProperties
 {
-	VkPhysicalDeviceAccelerationStructurePropertiesKHR AccelerationStructure;
-	VkPhysicalDeviceRayTracingPipelinePropertiesKHR RayTracingPipeline;
-};
+	FOptionalVulkanDeviceExtensionProperties()
+	{
+		FMemory::Memzero(*this);
+	}
+
+	VkPhysicalDeviceDescriptorBufferPropertiesEXT DescriptorBufferProps;
+	VkPhysicalDeviceSubgroupSizeControlPropertiesEXT SubgroupSizeControlProperties;
+
+#if VULKAN_RHI_RAYTRACING
+	VkPhysicalDeviceAccelerationStructurePropertiesKHR AccelerationStructureProps;
+	VkPhysicalDeviceRayTracingPipelinePropertiesKHR RayTracingPipelineProps;
 #endif // VULKAN_RHI_RAYTRACING
+};
 
 namespace VulkanRHI
 {
@@ -294,19 +303,7 @@ public:
 		return GpuSubgroupProps;
 	}
 
-	inline const VkPhysicalDeviceDescriptorBufferPropertiesEXT& GetDescriptorBufferProperties() const
-	{
-		check(RHI->GetOptionalExtensions().HasKHRGetPhysicalDeviceProperties2);
-		return DescriptorBufferProps;
-	}
-
 #if VULKAN_RHI_RAYTRACING
-	inline const FRayTracingProperties& GetRayTracingProperties() const
-	{
-		check(OptionalDeviceExtensions.HasRaytracingExtensions() || (Device == VK_NULL_HANDLE));
-		return RayTracingProperties;
-	}
-	
 	FVulkanRayTracingCompactionRequestHandler* GetRayTracingCompactionRequestHandler() { return RayTracingCompactionRequestHandler; }
 
 	void InitializeRayTracing();
@@ -483,6 +480,11 @@ public:
 		return OptionalDeviceExtensions;
 	}
 
+	inline const FOptionalVulkanDeviceExtensionProperties& GetOptionalExtensionProperties() const
+	{
+		return OptionalDeviceExtensionProperties;
+	}
+
 	inline bool SupportsParallelRendering() const
 	{
 		return OptionalDeviceExtensions.HasSeparateDepthStencilLayouts && OptionalDeviceExtensions.HasKHRSynchronization2 && OptionalDeviceExtensions.HasKHRRenderPass2;
@@ -569,10 +571,8 @@ private:
 	// Extension specific properties
 	VkPhysicalDeviceIDPropertiesKHR GpuIdProps;
 	VkPhysicalDeviceSubgroupProperties GpuSubgroupProps;
-	VkPhysicalDeviceDescriptorBufferPropertiesEXT DescriptorBufferProps;
 
 #if VULKAN_RHI_RAYTRACING
-	FRayTracingProperties RayTracingProperties;
 	FVulkanRayTracingCompactionRequestHandler* RayTracingCompactionRequestHandler = nullptr;
 #endif // VULKAN_RHI_RAYTRACING
 
@@ -619,6 +619,7 @@ private:
 	static TArray<const ANSICHAR*> SetupDeviceLayers(VkPhysicalDevice Gpu, FVulkanDeviceExtensionArray& UEExtensions);
 
 	FOptionalVulkanDeviceExtensions	OptionalDeviceExtensions;
+	FOptionalVulkanDeviceExtensionProperties OptionalDeviceExtensionProperties;
 	TArray<const ANSICHAR*>			DeviceExtensions;
 
 	void SetupFormats();
diff --git a/Engine/Source/Runtime/VulkanRHI/Private/VulkanExtensions.cpp b/Engine/Source/Runtime/VulkanRHI/Private/VulkanExtensions.cpp
--- a/Engine/Source/Runtime/VulkanRHI/Private/VulkanExtensions.cpp
+++ b/Engine/Source/Runtime/VulkanRHI/Private/VulkanExtensions.cpp
@@ -79,6 +79,13 @@ static void AddToPNext(ExistingChainType& Existing, NewStructType& Added)
 }
 
 
+struct FOptionalVulkanDeviceExtensionProperties& FVulkanDeviceExtension::GetDeviceExtensionProperties()
+{
+	const FOptionalVulkanDeviceExtensionProperties& ExtensionProperties = Device->GetOptionalExtensionProperties();
+	return const_cast<FOptionalVulkanDeviceExtensionProperties&>(ExtensionProperties);
+}
+
+
 #define VERIFYVULKANRESULT_INIT(VkFunction)	{ const VkResult ScopedResult = VkFunction; \
 												if (ScopedResult == VK_ERROR_INITIALIZATION_FAILED) { \
 													UE_LOG(LogVulkanRHI, Error, \
@@ -718,8 +725,7 @@ public:
 	virtual void PrePhysicalDeviceProperties(VkPhysicalDeviceProperties2KHR& PhysicalDeviceProperties2) override final
 	{
 #if VULKAN_RHI_RAYTRACING
-		const FRayTracingProperties& RayTracingProperties = Device->GetRayTracingProperties();
-		VkPhysicalDeviceAccelerationStructurePropertiesKHR& AccelerationStructure = const_cast<VkPhysicalDeviceAccelerationStructurePropertiesKHR&>(RayTracingProperties.AccelerationStructure);
+		VkPhysicalDeviceAccelerationStructurePropertiesKHR& AccelerationStructure = GetDeviceExtensionProperties().AccelerationStructureProps;
 		ZeroVulkanStruct(AccelerationStructure, VK_STRUCTURE_TYPE_PHYSICAL_DEVICE_ACCELERATION_STRUCTURE_PROPERTIES_KHR);
 		AddToPNext(PhysicalDeviceProperties2, AccelerationStructure);
 #endif
@@ -765,8 +771,7 @@ public:
 	virtual void PrePhysicalDeviceProperties(VkPhysicalDeviceProperties2KHR& PhysicalDeviceProperties2) override final
 	{
 #if VULKAN_RHI_RAYTRACING
-		const FRayTracingProperties& RayTracingProperties = Device->GetRayTracingProperties();
-		VkPhysicalDeviceRayTracingPipelinePropertiesKHR& RayTracingPipeline = const_cast<VkPhysicalDeviceRayTracingPipelinePropertiesKHR&>(RayTracingProperties.RayTracingPipeline);
+		VkPhysicalDeviceRayTracingPipelinePropertiesKHR& RayTracingPipeline = GetDeviceExtensionProperties().RayTracingPipelineProps;
 		ZeroVulkanStruct(RayTracingPipeline, VK_STRUCTURE_TYPE_PHYSICAL_DEVICE_RAY_TRACING_PIPELINE_PROPERTIES_KHR);
 		AddToPNext(PhysicalDeviceProperties2, RayTracingPipeline);
 #endif
@@ -987,6 +992,7 @@ public:
 	{
 		if (bSupportsSubgroupSizeControl)
 		{
+			VkPhysicalDeviceSubgroupSizeControlPropertiesEXT& SubgroupSizeControlProperties = GetDeviceExtensionProperties().SubgroupSizeControlProperties;
 			ZeroVulkanStruct(SubgroupSizeControlProperties, VK_STRUCTURE_TYPE_PHYSICAL_DEVICE_SUBGROUP_SIZE_CONTROL_PROPERTIES);
 			AddToPNext(PhysicalDeviceProperties2, SubgroupSizeControlProperties);
 		}
@@ -996,6 +1002,8 @@ public:
 	{
 		if (bSupportsSubgroupSizeControl)
 		{
+			VkPhysicalDeviceSubgroupSizeControlPropertiesEXT& SubgroupSizeControlProperties = GetDeviceExtensionProperties().SubgroupSizeControlProperties;
+
 			GRHIMinimumWaveSize = SubgroupSizeControlProperties.minSubgroupSize;
 			GRHIMaximumWaveSize = SubgroupSizeControlProperties.maxSubgroupSize;
 
@@ -1005,7 +1013,6 @@ public:
 
 private:
 	VkPhysicalDeviceSubgroupSizeControlFeaturesEXT SubgroupSizeControlFeatures;
-	VkPhysicalDeviceSubgroupSizeControlPropertiesEXT SubgroupSizeControlProperties;
 	bool bSupportsSubgroupSizeControl = false;
 };
 
@@ -1066,7 +1073,7 @@ public:
 	{
 		if (bSupportsDescriptorBuffer)
 		{
-			VkPhysicalDeviceDescriptorBufferPropertiesEXT& DescriptorBufferProperties = const_cast<VkPhysicalDeviceDescriptorBufferPropertiesEXT&>(Device->GetDescriptorBufferProperties());
+			VkPhysicalDeviceDescriptorBufferPropertiesEXT& DescriptorBufferProperties = GetDeviceExtensionProperties().DescriptorBufferProps;
 			ZeroVulkanStruct(DescriptorBufferProperties, VK_STRUCTURE_TYPE_PHYSICAL_DEVICE_DESCRIPTOR_BUFFER_PROPERTIES_EXT);
 			AddToPNext(PhysicalDeviceProperties2, DescriptorBufferProperties);
 		}
@@ -1078,7 +1085,7 @@ public:
 		{
 			AddToPNext(DeviceCreateInfo, DescriptorBufferFeatures);
 
-			VkPhysicalDeviceDescriptorBufferPropertiesEXT& DescriptorBufferProperties = const_cast<VkPhysicalDeviceDescriptorBufferPropertiesEXT&>(Device->GetDescriptorBufferProperties());
+			const VkPhysicalDeviceDescriptorBufferPropertiesEXT& DescriptorBufferProperties = GetDeviceExtensionProperties().DescriptorBufferProps;
 
 			UE_LOG(LogVulkanRHI, Display, TEXT("Enabling Vulkan Descriptor Buffers with: ")
 				TEXT("allowSamplerImageViewPostSubmitCreation=%u, maxDescriptorBufferBindings=%u, ")
diff --git a/Engine/Source/Runtime/VulkanRHI/Private/VulkanExtensions.h b/Engine/Source/Runtime/VulkanRHI/Private/VulkanExtensions.h
--- a/Engine/Source/Runtime/VulkanRHI/Private/VulkanExtensions.h
+++ b/Engine/Source/Runtime/VulkanRHI/Private/VulkanExtensions.h
@@ -113,6 +113,8 @@ public:
 	static TArray<const ANSICHAR*> ExternalExtensions;
 
 protected:
+	struct FOptionalVulkanDeviceExtensionProperties& GetDeviceExtensionProperties();
+
 	FVulkanDevice* Device;
 	TUniqueFunction<void(FOptionalVulkanDeviceExtensions& ExtensionFlags)> FlagSetter = nullptr;
 };
diff --git a/Engine/Source/Runtime/VulkanRHI/Private/VulkanRHI.cpp b/Engine/Source/Runtime/VulkanRHI/Private/VulkanRHI.cpp
--- a/Engine/Source/Runtime/VulkanRHI/Private/VulkanRHI.cpp
+++ b/Engine/Source/Runtime/VulkanRHI/Private/VulkanRHI.cpp
@@ -776,9 +776,8 @@ void FVulkanDynamicRHI::InitInstance()
 			GRHISupportsRayTracingShaders = RHISupportsRayTracingShaders(GMaxRHIShaderPlatform);
 			GRHISupportsInlineRayTracing = RHISupportsInlineRayTracing(GMaxRHIShaderPlatform) && Device->GetOptionalExtensions().HasRayQuery;
 
-			const FRayTracingProperties& RayTracingProps = Device->GetRayTracingProperties();
 			GRHIRayTracingAccelerationStructureAlignment = 256; // TODO (currently handled by FVulkanAccelerationStructureBuffer)
-			GRHIRayTracingScratchBufferAlignment = RayTracingProps.AccelerationStructure.minAccelerationStructureScratchOffsetAlignment;
+			GRHIRayTracingScratchBufferAlignment = Device->GetOptionalExtensionProperties().AccelerationStructureProps.minAccelerationStructureScratchOffsetAlignment;
 			GRHIRayTracingInstanceDescriptorSize = uint32(sizeof(VkAccelerationStructureInstanceKHR));
 		}
 #endif
@@ -1665,7 +1664,7 @@ void FVulkanDescriptorSetsLayout::Compile(FVulkanDescriptorSetLayoutMap& DSetLay
 #if VULKAN_RHI_RAYTRACING
 	if (GRHISupportsRayTracing)
 	{
-		check(LayoutTypes[VK_DESCRIPTOR_TYPE_ACCELERATION_STRUCTURE_KHR] < Device->GetRayTracingProperties().AccelerationStructure.maxDescriptorSetAccelerationStructures);
+		check(LayoutTypes[VK_DESCRIPTOR_TYPE_ACCELERATION_STRUCTURE_KHR] < Device->GetOptionalExtensionProperties().AccelerationStructureProps.maxDescriptorSetAccelerationStructures);
 	}
 #endif
 	
diff --git a/Engine/Source/Runtime/VulkanRHI/Private/VulkanRayTracing.cpp b/Engine/Source/Runtime/VulkanRHI/Private/VulkanRayTracing.cpp
--- a/Engine/Source/Runtime/VulkanRHI/Private/VulkanRayTracing.cpp
+++ b/Engine/Source/Runtime/VulkanRHI/Private/VulkanRayTracing.cpp
@@ -1059,9 +1059,9 @@ FVulkanRayTracingPipelineState::FVulkanRayTracingPipelineState(FVulkanDevice* co
 		delete[] EntryPoint;
 	}
 
-	const FRayTracingProperties& Props = InDevice->GetRayTracingProperties();
-	const uint32 HandleSize = Props.RayTracingPipeline.shaderGroupHandleSize;
-	const uint32 HandleSizeAligned = GetAlignedSize(HandleSize, Props.RayTracingPipeline.shaderGroupHandleAlignment);
+	const VkPhysicalDeviceRayTracingPipelinePropertiesKHR& RayTracingPipelineProps = InDevice->GetOptionalExtensionProperties().RayTracingPipelineProps;
+	const uint32 HandleSize = RayTracingPipelineProps.shaderGroupHandleSize;
+	const uint32 HandleSizeAligned = GetAlignedSize(HandleSize, RayTracingPipelineProps.shaderGroupHandleAlignment);
 	const uint32 GroupCount = ShaderGroups.Num();
 	const uint32 SBTSize = GroupCount * HandleSizeAligned;
 
