VulkanRHI: Apply Wave32 hint in Vulkan when VK_EXT_subgroup_size_control is supported.  Properly check requiredSubgroupSizeStages before applying wave size.  Fixes artifacts on ARC with Vulkan.

#rnx

diff --git a/Engine/Source/Developer/VulkanShaderFormat/Private/VulkanShaderCompiler.cpp b/Engine/Source/Developer/VulkanShaderFormat/Private/VulkanShaderCompiler.cpp
--- a/Engine/Source/Developer/VulkanShaderFormat/Private/VulkanShaderCompiler.cpp
+++ b/Engine/Source/Developer/VulkanShaderFormat/Private/VulkanShaderCompiler.cpp
@@ -2366,6 +2366,12 @@ static bool CompileWithShaderConductor(
 		}
 	}
 
+	// Take note of preferred wave size flag if none was specified in HLSL
+	if ((WaveSize == 0) && Input.Environment.CompilerFlags.Contains(CFLAG_Wave32))
+	{
+		WaveSize = 32;
+	}
+
 	if (bRewriteHlslSource)
 	{
 		// Rewrite HLSL source code to remove unused global resources and variables
diff --git a/Engine/Source/Runtime/VulkanRHI/Private/VulkanPipeline.cpp b/Engine/Source/Runtime/VulkanRHI/Private/VulkanPipeline.cpp
--- a/Engine/Source/Runtime/VulkanRHI/Private/VulkanPipeline.cpp
+++ b/Engine/Source/Runtime/VulkanRHI/Private/VulkanPipeline.cpp
@@ -1271,7 +1271,7 @@ bool FVulkanPipelineStateCacheManager::CreateGfxPipelineFromEntry(FVulkanRHIGrap
 		const ShaderStage::EStage CurrStage = (ShaderStage::EStage)ShaderStage;
 
 		ShaderStages[PipelineInfo.stageCount].sType = VK_STRUCTURE_TYPE_PIPELINE_SHADER_STAGE_CREATE_INFO;
-		VkShaderStageFlagBits Stage = UEFrequencyToVKStageBit(ShaderStage::GetFrequencyForGfxStage(CurrStage));
+		const VkShaderStageFlagBits Stage = UEFrequencyToVKStageBit(ShaderStage::GetFrequencyForGfxStage(CurrStage));
 		ShaderStages[PipelineInfo.stageCount].stage = Stage;
 		ShaderStages[PipelineInfo.stageCount].module = ShaderModules[CurrStage]->GetVkShaderModule();
 		Shaders[ShaderStage]->GetEntryPoint(EntryPoints[PipelineInfo.stageCount], 24);
@@ -1282,10 +1282,13 @@ bool FVulkanPipelineStateCacheManager::CreateGfxPipelineFromEntry(FVulkanRHIGrap
 			const FVulkanShaderHeader& ShaderHeader = Shaders[ShaderStage]->GetCodeHeader();
 			if (ShaderHeader.WaveSize > 0)
 			{
-				check((ShaderHeader.WaveSize >= GRHIMinimumWaveSize) && (ShaderHeader.WaveSize <= GRHIMaximumWaveSize));
-				ZeroVulkanStruct(RequiredSubgroupSizeCreateInfo[PipelineInfo.stageCount], VK_STRUCTURE_TYPE_PIPELINE_SHADER_STAGE_REQUIRED_SUBGROUP_SIZE_CREATE_INFO);
-				RequiredSubgroupSizeCreateInfo[PipelineInfo.stageCount].requiredSubgroupSize = ShaderHeader.WaveSize;
-				ShaderStages[PipelineInfo.stageCount].pNext = &RequiredSubgroupSizeCreateInfo[PipelineInfo.stageCount];
+				if (VKHasAllFlags(Device->GetOptionalExtensionProperties().SubgroupSizeControlProperties.requiredSubgroupSizeStages, Stage))
+				{
+					check((ShaderHeader.WaveSize >= GRHIMinimumWaveSize) && (ShaderHeader.WaveSize <= GRHIMaximumWaveSize));
+					ZeroVulkanStruct(RequiredSubgroupSizeCreateInfo[PipelineInfo.stageCount], VK_STRUCTURE_TYPE_PIPELINE_SHADER_STAGE_REQUIRED_SUBGROUP_SIZE_CREATE_INFO);
+					RequiredSubgroupSizeCreateInfo[PipelineInfo.stageCount].requiredSubgroupSize = ShaderHeader.WaveSize;
+					ShaderStages[PipelineInfo.stageCount].pNext = &RequiredSubgroupSizeCreateInfo[PipelineInfo.stageCount];
+				}
 			}
 		}
 
@@ -2165,10 +2168,13 @@ FVulkanComputePipeline* FVulkanPipelineStateCacheManager::CreateComputePipelineF
 	VkPipelineShaderStageRequiredSubgroupSizeCreateInfo RequiredSubgroupSizeCreateInfo;
 	if ((CSHeader.WaveSize > 0) && Device->GetOptionalExtensions().HasEXTSubgroupSizeControl)
 	{
-		check((CSHeader.WaveSize >= GRHIMinimumWaveSize) && (CSHeader.WaveSize <= GRHIMaximumWaveSize));
-		ZeroVulkanStruct(RequiredSubgroupSizeCreateInfo, VK_STRUCTURE_TYPE_PIPELINE_SHADER_STAGE_REQUIRED_SUBGROUP_SIZE_CREATE_INFO);
-		RequiredSubgroupSizeCreateInfo.requiredSubgroupSize = CSHeader.WaveSize;
-		PipelineInfo.stage.pNext = &RequiredSubgroupSizeCreateInfo;
+		if (VKHasAllFlags(Device->GetOptionalExtensionProperties().SubgroupSizeControlProperties.requiredSubgroupSizeStages, VK_SHADER_STAGE_COMPUTE_BIT))
+		{
+			check((CSHeader.WaveSize >= GRHIMinimumWaveSize) && (CSHeader.WaveSize <= GRHIMaximumWaveSize));
+			ZeroVulkanStruct(RequiredSubgroupSizeCreateInfo, VK_STRUCTURE_TYPE_PIPELINE_SHADER_STAGE_REQUIRED_SUBGROUP_SIZE_CREATE_INFO);
+			RequiredSubgroupSizeCreateInfo.requiredSubgroupSize = CSHeader.WaveSize;
+			PipelineInfo.stage.pNext = &RequiredSubgroupSizeCreateInfo;
+		}
 	}
 
 	VkResult Result;
