Fixed: VulkanRHI was not correctly mapping EVRSShadingRate to a fragmentSize, resulting in non-workng VRS
Pre-compute EVRSShadingRate mapping to a fragmentSize instead of doing it on each PSO create

diff --git a/Engine/Source/Runtime/RHI/Public/RHIDefinitions.h b/Engine/Source/Runtime/RHI/Public/RHIDefinitions.h
--- a/Engine/Source/Runtime/RHI/Public/RHIDefinitions.h
+++ b/Engine/Source/Runtime/RHI/Public/RHIDefinitions.h
@@ -684,6 +684,8 @@ enum EVRSShadingRate : uint8
 	VRSSR_2x4  = (VRSASR_2X << 2) + VRSASR_4X,
 	VRSSR_4x2  = (VRSASR_4X << 2) + VRSASR_2X,
 	VRSSR_4x4  = (VRSASR_4X << 2) + VRSASR_4X,
+	
+	VRSSR_Last  = VRSSR_4x4
 };
 
 enum EVRSRateCombiner : uint8
diff --git a/Engine/Source/Runtime/VulkanRHI/Private/VulkanDevice.cpp b/Engine/Source/Runtime/VulkanRHI/Private/VulkanDevice.cpp
--- a/Engine/Source/Runtime/VulkanRHI/Private/VulkanDevice.cpp
+++ b/Engine/Source/Runtime/VulkanRHI/Private/VulkanDevice.cpp
@@ -167,6 +167,41 @@ static void LoadValidationCache(VkDevice Device, VkValidationCacheEXT& OutValida
 }
 #endif
 
+#if VULKAN_SUPPORTS_FRAGMENT_SHADING_RATE
+static VkExtent2D GetBestMatchedShadingRateExtents(uint32 ShadingRate, const TArray<VkPhysicalDeviceFragmentShadingRateKHR>& FragmentShadingRates)
+{
+	// Given that for Vulkan we need to query available device shading rates, we're not guaranteed to have everything that's in our enum;
+	// This function walks the list of supported fragment rates returned by the device, and returns the closest match to the rate requested.
+	const VkExtent2D DirectMappedExtent = { 
+		1u << (ShadingRate >> 2), 
+		1u << (ShadingRate & 0x03) 
+	};
+	VkExtent2D BestMatchedExtent = { 1, 1 };
+
+	if (BestMatchedExtent.width != DirectMappedExtent.width || 
+		BestMatchedExtent.height != DirectMappedExtent.height)
+	{
+		for (auto const& Rate : FragmentShadingRates)
+		{
+			if (Rate.fragmentSize.width == DirectMappedExtent.width && 
+				Rate.fragmentSize.height == DirectMappedExtent.height)
+			{
+				BestMatchedExtent = DirectMappedExtent;
+				break;
+			}
+
+			if ((Rate.fragmentSize.width >= BestMatchedExtent.width && Rate.fragmentSize.width <= DirectMappedExtent.width && Rate.fragmentSize.height <= DirectMappedExtent.height && Rate.fragmentSize.height >= BestMatchedExtent.height) ||
+				(Rate.fragmentSize.height >= BestMatchedExtent.height && Rate.fragmentSize.height <= DirectMappedExtent.height && Rate.fragmentSize.width <= DirectMappedExtent.width && Rate.fragmentSize.width >= BestMatchedExtent.width))
+			{
+				BestMatchedExtent = Rate.fragmentSize;
+			}
+		}
+	}
+
+	return BestMatchedExtent;
+}
+#endif
+
 
 FVulkanDevice::FVulkanDevice(FVulkanDynamicRHI* InRHI, VkPhysicalDevice InGpu)
 	: Device(VK_NULL_HANDLE)
@@ -441,6 +476,12 @@ void FVulkanDevice::CreateDevice(TArray<const ANSICHAR*>& DeviceLayers, FVulkanD
 				ZeroVulkanStruct(FragmentShadingRates[i], VK_STRUCTURE_TYPE_PHYSICAL_DEVICE_FRAGMENT_SHADING_RATE_KHR);
 			}
 			VulkanRHI::vkGetPhysicalDeviceFragmentShadingRatesKHR(Gpu, &FragmentShadingRateCount, FragmentShadingRates.GetData());
+
+			// Build a map from EVRSShadingRate to fragment size
+			for (uint32 ShadingRate = 0u; ShadingRate < (uint32)FragmentSizeMap.Num(); ++ShadingRate)
+			{
+				FragmentSizeMap[ShadingRate] = GetBestMatchedShadingRateExtents(ShadingRate, FragmentShadingRates);
+			}
 		}
 	}
 #endif // VULKAN_SUPPORTS_FRAGMENT_SHADING_RATE
@@ -1361,36 +1402,6 @@ const VkComponentMapping& FVulkanDevice::GetFormatComponentMapping(EPixelFormat
 	return PixelFormatComponentMapping[UEFormat];
 }
 
-#if VULKAN_SUPPORTS_FRAGMENT_SHADING_RATE
-VkExtent2D FVulkanDevice::GetBestMatchedShadingRateExtents(EVRSShadingRate ShadingRate) const
-{
-	// Given that for Vulkan we need to query available device shading rates, we're not guaranteed to have everything that's in our enum;
-	// This function walks the list of supported fragment rates returned by the device, and returns the closest match to the rate requested.
-	const VkExtent2D DirectMappedExtent = { (uint32)(ShadingRate >> 2) + 1, (uint32)((ShadingRate & 0x03) + 1) };
-	VkExtent2D BestMatchedExtent = { 1, 1 };
-	
-	if (BestMatchedExtent.width != DirectMappedExtent.width && BestMatchedExtent.height != DirectMappedExtent.height)
-	{
-		for (auto const& Rate : FragmentShadingRates)
-		{
-			if (Rate.fragmentSize.width == DirectMappedExtent.width && Rate.fragmentSize.height == DirectMappedExtent.height)
-			{
-				BestMatchedExtent = DirectMappedExtent;
-				break;
-			}
-
-			if ((Rate.fragmentSize.width >= BestMatchedExtent.width && Rate.fragmentSize.width <= DirectMappedExtent.width && Rate.fragmentSize.height <= DirectMappedExtent.height && Rate.fragmentSize.height >= BestMatchedExtent.height) ||
-				(Rate.fragmentSize.height >= BestMatchedExtent.height && Rate.fragmentSize.height <= DirectMappedExtent.height && Rate.fragmentSize.width <= DirectMappedExtent.width && Rate.fragmentSize.width >= BestMatchedExtent.width))
-			{
-				BestMatchedExtent = Rate.fragmentSize;
-			}
-		}
-	}
-
-	return BestMatchedExtent;
-}
-#endif
-
 void FVulkanDevice::NotifyDeletedImage(VkImage Image, bool bRenderTarget)
 {
 	if (bRenderTarget)
diff --git a/Engine/Source/Runtime/VulkanRHI/Private/VulkanDevice.h b/Engine/Source/Runtime/VulkanRHI/Private/VulkanDevice.h
--- a/Engine/Source/Runtime/VulkanRHI/Private/VulkanDevice.h
+++ b/Engine/Source/Runtime/VulkanRHI/Private/VulkanDevice.h
@@ -283,9 +283,12 @@ public:
 	}
 
 #if VULKAN_SUPPORTS_FRAGMENT_SHADING_RATE
-	VkExtent2D GetBestMatchedShadingRateExtents(EVRSShadingRate Rate) const;
+	inline VkExtent2D GetBestMatchedFragmentSize(EVRSShadingRate Rate) const
+	{
+		return FragmentSizeMap[Rate];
+	}
 #endif
-	
+
 	inline const VkPhysicalDeviceLimits& GetLimits() const
 	{
 		return GpuProps.limits;
@@ -566,6 +569,7 @@ private:
 
 #if VULKAN_SUPPORTS_FRAGMENT_SHADING_RATE
 	TArray<VkPhysicalDeviceFragmentShadingRateKHR> FragmentShadingRates;
+	TStaticArray<VkExtent2D, (EVRSShadingRate::VRSSR_Last+1)> FragmentSizeMap;
 #endif
 
 	// Extension specific properties
diff --git a/Engine/Source/Runtime/VulkanRHI/Private/VulkanPipeline.cpp b/Engine/Source/Runtime/VulkanRHI/Private/VulkanPipeline.cpp
--- a/Engine/Source/Runtime/VulkanRHI/Private/VulkanPipeline.cpp
+++ b/Engine/Source/Runtime/VulkanRHI/Private/VulkanPipeline.cpp
@@ -1355,13 +1355,12 @@ bool FVulkanPipelineStateCacheManager::CreateGfxPipelineFromEntry(FVulkanRHIGrap
 	PipelineInfo.pDynamicState = &DynamicState;
 
 #if VULKAN_SUPPORTS_FRAGMENT_SHADING_RATE
-	const VkExtent2D FragmentSize = Device->GetBestMatchedShadingRateExtents(PSO->Desc.ShadingRate);
-	VkFragmentShadingRateCombinerOpKHR PipelineToPrimitiveCombinerOperation = FragmentCombinerOpMap[(uint8)PSO->Desc.Combiner];
-
 	VkPipelineFragmentShadingRateStateCreateInfoKHR PipelineFragmentShadingRate;
-
-	if (GRHISupportsPipelineVariableRateShading && GRHIVariableRateShadingEnabled)
+	if (GRHISupportsPipelineVariableRateShading && GRHIVariableRateShadingEnabled && PSO->Desc.ShadingRate != EVRSShadingRate::VRSSR_1x1)
 	{
+		const VkExtent2D FragmentSize = Device->GetBestMatchedFragmentSize(PSO->Desc.ShadingRate);
+		VkFragmentShadingRateCombinerOpKHR PipelineToPrimitiveCombinerOperation = FragmentCombinerOpMap[(uint8)PSO->Desc.Combiner];
+		
 		ZeroVulkanStruct(PipelineFragmentShadingRate, VK_STRUCTURE_TYPE_PIPELINE_FRAGMENT_SHADING_RATE_STATE_CREATE_INFO_KHR);
 		PipelineFragmentShadingRate.fragmentSize = FragmentSize;
 		PipelineFragmentShadingRate.combinerOps[0] = PipelineToPrimitiveCombinerOperation;
