Add an option to disable use of VRS on Vulkan RHI (r.Vulkan.AllowVariableRateShading). VRS is enabled by default

diff --git a/Engine/Source/Runtime/VulkanRHI/Private/VulkanExtensions.cpp b/Engine/Source/Runtime/VulkanRHI/Private/VulkanExtensions.cpp
--- a/Engine/Source/Runtime/VulkanRHI/Private/VulkanExtensions.cpp
+++ b/Engine/Source/Runtime/VulkanRHI/Private/VulkanExtensions.cpp
@@ -484,6 +484,7 @@ public:
 	FVulkanKHRFragmentShadingRateExtension(FVulkanDevice* InDevice)
 		: FVulkanDeviceExtension(InDevice, VK_KHR_FRAGMENT_SHADING_RATE_EXTENSION_NAME, VULKAN_SUPPORTS_FRAGMENT_SHADING_RATE)
 	{
+		bEnabledInCode = bEnabledInCode && GRHIVariableRateShadingEnabled;
 	}
 
 	virtual void PrePhysicalDeviceProperties(VkPhysicalDeviceProperties2KHR& PhysicalDeviceProperties2) override final
@@ -561,6 +562,7 @@ public:
 	FVulkanEXTFragmentDensityMapExtension(FVulkanDevice* InDevice)
 		: FVulkanDeviceExtension(InDevice, VK_EXT_FRAGMENT_DENSITY_MAP_EXTENSION_NAME, VULKAN_SUPPORTS_FRAGMENT_DENSITY_MAP)
 	{
+		bEnabledInCode = bEnabledInCode && GRHIVariableRateShadingEnabled;
 	}
 
 	virtual void PrePhysicalDeviceProperties(VkPhysicalDeviceProperties2KHR& PhysicalDeviceProperties2) override final
@@ -624,6 +626,7 @@ public:
 	FVulkanEXTFragmentDensityMap2Extension(FVulkanDevice* InDevice)
 		: FVulkanDeviceExtension(InDevice, VK_EXT_FRAGMENT_DENSITY_MAP_2_EXTENSION_NAME, VULKAN_SUPPORTS_FRAGMENT_DENSITY_MAP2)
 	{
+		bEnabledInCode = bEnabledInCode && GRHIVariableRateShadingEnabled;
 	}
 
 	virtual void PrePhysicalDeviceFeatures(VkPhysicalDeviceFeatures2KHR& PhysicalDeviceFeatures2) override final
diff --git a/Engine/Source/Runtime/VulkanRHI/Private/VulkanRHI.cpp b/Engine/Source/Runtime/VulkanRHI/Private/VulkanRHI.cpp
--- a/Engine/Source/Runtime/VulkanRHI/Private/VulkanRHI.cpp
+++ b/Engine/Source/Runtime/VulkanRHI/Private/VulkanRHI.cpp
@@ -90,6 +90,15 @@ static FAutoConsoleVariableRef GCVarEnableTransientResourceAllocator(
 	ECVF_ReadOnly
 );
 
+int32 GVulkanAllowVariableRateShading = 1;
+static FAutoConsoleVariableRef CVarVulkanVariableRateShading(
+	TEXT("r.Vulkan.AllowVariableRateShading"),
+	GVulkanAllowVariableRateShading,
+	TEXT("0 to disable variable rate shading")
+	TEXT("1 to allow use of variable rate shading if available (default)"),
+	ECVF_ReadOnly
+);
+
 bool GGPUCrashDebuggingEnabled = false;
 
 
@@ -216,6 +225,7 @@ FVulkanDynamicRHI::FVulkanDynamicRHI()
 	GRHISupportsMultithreading = true;
 	GRHISupportsMultithreadedResources = true;
 	GRHISupportsPipelineFileCache = true;
+	GRHIVariableRateShadingEnabled &= (GVulkanAllowVariableRateShading != 0); // before extensions setup
 	GRHITransitionPrivateData_SizeInBytes = sizeof(FVulkanPipelineBarrier);
 	GRHITransitionPrivateData_AlignInBytes = alignof(FVulkanPipelineBarrier);
 	GConfig->GetInt(TEXT("TextureStreaming"), TEXT("PoolSizeVRAMPercentage"), GPoolSizeVRAMPercentage, GEngineIni);
