VulkanRHI: Only apply wave size request if supported.

#rnx

diff --git a/Engine/Source/Runtime/VulkanRHI/Private/VulkanPipeline.cpp b/Engine/Source/Runtime/VulkanRHI/Private/VulkanPipeline.cpp
--- a/Engine/Source/Runtime/VulkanRHI/Private/VulkanPipeline.cpp
+++ b/Engine/Source/Runtime/VulkanRHI/Private/VulkanPipeline.cpp
@@ -1282,9 +1282,12 @@ bool FVulkanPipelineStateCacheManager::CreateGfxPipelineFromEntry(FVulkanRHIGrap
 			const FVulkanShaderHeader& ShaderHeader = Shaders[ShaderStage]->GetCodeHeader();
 			if (ShaderHeader.WaveSize > 0)
 			{
-				if (VKHasAllFlags(Device->GetOptionalExtensionProperties().SubgroupSizeControlProperties.requiredSubgroupSizeStages, Stage))
+				// Check if supported by this stage and Check if requested size is supported
+				const VkPhysicalDeviceSubgroupSizeControlPropertiesEXT& SubgroupSizeControlProperties = Device->GetOptionalExtensionProperties().SubgroupSizeControlProperties;
+				const bool bSupportedStage = (VKHasAllFlags(SubgroupSizeControlProperties.requiredSubgroupSizeStages, Stage));
+				const bool bSupportedSize = ((ShaderHeader.WaveSize >= SubgroupSizeControlProperties.minSubgroupSize) && (ShaderHeader.WaveSize <= SubgroupSizeControlProperties.maxSubgroupSize));
+				if (bSupportedStage && bSupportedSize)
 				{
-					check((ShaderHeader.WaveSize >= GRHIMinimumWaveSize) && (ShaderHeader.WaveSize <= GRHIMaximumWaveSize));
 					ZeroVulkanStruct(RequiredSubgroupSizeCreateInfo[PipelineInfo.stageCount], VK_STRUCTURE_TYPE_PIPELINE_SHADER_STAGE_REQUIRED_SUBGROUP_SIZE_CREATE_INFO);
 					RequiredSubgroupSizeCreateInfo[PipelineInfo.stageCount].requiredSubgroupSize = ShaderHeader.WaveSize;
 					ShaderStages[PipelineInfo.stageCount].pNext = &RequiredSubgroupSizeCreateInfo[PipelineInfo.stageCount];
@@ -2167,12 +2170,17 @@ FVulkanComputePipeline* FVulkanPipelineStateCacheManager::CreateComputePipelineF
 	VkPipelineShaderStageRequiredSubgroupSizeCreateInfo RequiredSubgroupSizeCreateInfo;
 	if ((CSHeader.WaveSize > 0) && Device->GetOptionalExtensions().HasEXTSubgroupSizeControl)
 	{
-		if (VKHasAllFlags(Device->GetOptionalExtensionProperties().SubgroupSizeControlProperties.requiredSubgroupSizeStages, VK_SHADER_STAGE_COMPUTE_BIT))
+		// Check if supported by this stage
+		const VkPhysicalDeviceSubgroupSizeControlPropertiesEXT& SubgroupSizeControlProperties = Device->GetOptionalExtensionProperties().SubgroupSizeControlProperties;
+		if (VKHasAllFlags(SubgroupSizeControlProperties.requiredSubgroupSizeStages, VK_SHADER_STAGE_COMPUTE_BIT))
 		{
-			check((CSHeader.WaveSize >= GRHIMinimumWaveSize) && (CSHeader.WaveSize <= GRHIMaximumWaveSize));
-			ZeroVulkanStruct(RequiredSubgroupSizeCreateInfo, VK_STRUCTURE_TYPE_PIPELINE_SHADER_STAGE_REQUIRED_SUBGROUP_SIZE_CREATE_INFO);
-			RequiredSubgroupSizeCreateInfo.requiredSubgroupSize = CSHeader.WaveSize;
-			PipelineInfo.stage.pNext = &RequiredSubgroupSizeCreateInfo;
+			// Check if requested size is supported
+			if ((CSHeader.WaveSize >= SubgroupSizeControlProperties.minSubgroupSize) && (CSHeader.WaveSize <= SubgroupSizeControlProperties.maxSubgroupSize))
+			{
+				ZeroVulkanStruct(RequiredSubgroupSizeCreateInfo, VK_STRUCTURE_TYPE_PIPELINE_SHADER_STAGE_REQUIRED_SUBGROUP_SIZE_CREATE_INFO);
+				RequiredSubgroupSizeCreateInfo.requiredSubgroupSize = CSHeader.WaveSize;
+				PipelineInfo.stage.pNext = &RequiredSubgroupSizeCreateInfo;
+			}
 		}
 	}
 
