Removing /Private include paths from RHI build scripts.

diff --git a/Engine/Platforms/Hololens/Source/Runtime/D3D12RHI/D3D12RHI_HoloLens.build.cs b/Engine/Platforms/Hololens/Source/Runtime/D3D12RHI/D3D12RHI_HoloLens.build.cs
--- a/Engine/Platforms/Hololens/Source/Runtime/D3D12RHI/D3D12RHI_HoloLens.build.cs
+++ b/Engine/Platforms/Hololens/Source/Runtime/D3D12RHI/D3D12RHI_HoloLens.build.cs
@@ -8,7 +8,6 @@ namespace UnrealBuildTool.Rules
 
 		public D3D12RHI_HoloLens(ReadOnlyTargetRules Target) : base(Target)
 		{
-			PrivateIncludePaths.Add("../Platforms/HoloLens/Source/Runtime/D3D12RHI/Private");
 		}
 	}
 }
diff --git a/Engine/Source/Runtime/VulkanRHI/Private/VulkanPlatform.h b/Engine/Source/Runtime/VulkanRHI/Private/VulkanPlatform.h
--- a/Engine/Source/Runtime/VulkanRHI/Private/VulkanPlatform.h
+++ b/Engine/Source/Runtime/VulkanRHI/Private/VulkanPlatform.h
@@ -1,6 +1,14 @@
 // Copyright Epic Games, Inc. All Rights Reserved.
 
-
 #pragma once
 
-#include "VulkanPlatformDefines.h"
\ No newline at end of file
+#include "HAL/PreprocessorHelpers.h"
+
+// Vulkan is currently organized just a little differently than COMPILED_PLATFORM_HEADER does things.
+#if PLATFORM_IS_EXTENSION
+	#define VULKAN_COMPILED_PLATFORM_HEADER(Suffix) PREPROCESSOR_TO_STRING(Suffix)
+#else
+	#define VULKAN_COMPILED_PLATFORM_HEADER(Suffix) PREPROCESSOR_TO_STRING(PLATFORM_HEADER_NAME/Suffix)
+#endif
+
+#include VULKAN_COMPILED_PLATFORM_HEADER(VulkanPlatformDefines.h)
diff --git a/Engine/Source/Runtime/VulkanRHI/VulkanRHI.Build.cs b/Engine/Source/Runtime/VulkanRHI/VulkanRHI.Build.cs
--- a/Engine/Source/Runtime/VulkanRHI/VulkanRHI.Build.cs
+++ b/Engine/Source/Runtime/VulkanRHI/VulkanRHI.Build.cs
@@ -6,8 +6,6 @@ using System.IO;
 
 public class VulkanRHI : ModuleRules
 {
-	protected virtual bool bShouldIncludePlatformPrivate { get { return true; } }
-
 	public VulkanRHI(ReadOnlyTargetRules Target) : base(Target)
 	{
 		IWYUSupport = IWYUSupport.None;
@@ -15,20 +13,8 @@ public class VulkanRHI : ModuleRules
 
 		if (Target.Platform.IsInGroup(UnrealPlatformGroup.Windows))
 		{
-			PrivateIncludePaths.Add("Runtime/VulkanRHI/Private/Windows");
 			AddEngineThirdPartyPrivateStaticDependencies(Target, "AMD_AGS");
 		}
-		else if (Target.IsInPlatformGroup(UnrealPlatformGroup.Unix))
-		{
-			if (Target.IsInPlatformGroup(UnrealPlatformGroup.Linux))
-			{
-				PrivateIncludePaths.Add("Runtime/VulkanRHI/Private/Linux");
-			}
-		}
-		else if (bShouldIncludePlatformPrivate)
-		{
-			PrivateIncludePaths.Add("Runtime/VulkanRHI/Private/" + Target.Platform);
-		}
 
 		if (Target.Platform.IsInGroup(UnrealPlatformGroup.Windows))
 		{
