Missing FWrapLayer defaults for Vulkan RT
Fixed Vulkan RT bug when allocating descriptor pools the AS type was missing

diff --git a/Engine/Source/Runtime/VulkanRHI/Private/VulkanCommandWrappers.h b/Engine/Source/Runtime/VulkanRHI/Private/VulkanCommandWrappers.h
--- a/Engine/Source/Runtime/VulkanRHI/Private/VulkanCommandWrappers.h
+++ b/Engine/Source/Runtime/VulkanRHI/Private/VulkanCommandWrappers.h
@@ -187,16 +187,16 @@ struct FWrapLayer
 	static void CreateWin32SurfaceKHR(VkResult Result, VkInstance instance, const VkWin32SurfaceCreateInfoKHR* pCreateInfo, VkSurfaceKHR* pSurface);
 #endif
 #if VULKAN_RHI_RAYTRACING
-	static void CreateAccelerationStructureKHR(VkResult Result, VkDevice Device, const VkAccelerationStructureCreateInfoKHR* CreateInfo, const VkAllocationCallbacks* Allocator, VkAccelerationStructureKHR* AccelerationStructure);
-	static void DestroyAccelerationStructureKHR(VkResult Result, VkDevice Device, VkAccelerationStructureKHR AccelerationStructure, const VkAllocationCallbacks* Allocator);
-	static void CmdBuildAccelerationStructuresKHR(VkResult Result, VkCommandBuffer CommandBuffer, uint32 InfoCount, const VkAccelerationStructureBuildGeometryInfoKHR* Infos, const VkAccelerationStructureBuildRangeInfoKHR* const* BuildRangeInfos);
-	static void GetAccelerationStructureBuildSizesKHR(VkResult Result, VkDevice Device, VkAccelerationStructureBuildTypeKHR BuildType, const VkAccelerationStructureBuildGeometryInfoKHR* BuildInfo, const uint32* MaxPrimitiveCounts, VkAccelerationStructureBuildSizesInfoKHR* SizeInfo);
-	static void GetAccelerationStructureDeviceAddressKHR(VkResult Result, VkDevice Device, const VkAccelerationStructureDeviceAddressInfoKHR* Info);
-	static void CmdTraceRaysKHR(VkResult Result, VkCommandBuffer CommandBuffer, const VkStridedDeviceAddressRegionKHR* RaygenShaderBindingTable, const VkStridedDeviceAddressRegionKHR* MissShaderBindingTable, const VkStridedDeviceAddressRegionKHR* HitShaderBindingTable, const VkStridedDeviceAddressRegionKHR* CallableShaderBindingTable, uint32 width, uint32 height, uint32 depth);
-	static void CreateRayTracingPipelinesKHR(VkResult Result, VkDevice Device, VkDeferredOperationKHR DeferredOperation, VkPipelineCache PipelineCache, uint32_t CreateInfoCount, const VkRayTracingPipelineCreateInfoKHR* CreateInfos, const VkAllocationCallbacks* Allocator, VkPipeline* Pipelines);
-	static void GetRayTracingShaderGroupHandlesKHR(VkResult Result, VkDevice Device, VkPipeline Pipeline, uint32_t FirstGroup, uint32_t GroupCount, size_t DataSize, void* Data);
-	static void CmdWriteAccelerationStructuresPropertiesKHR(VkResult Result, VkCommandBuffer CommandBuffer, uint32_t AccelerationStructureCount, const VkAccelerationStructureKHR* AccelerationStructures, VkQueryType QueryType, VkQueryPool QueryPool, uint32_t FirstQuery);
-	static void CmdCopyAccelerationStructureKHR(VkResult Result, VkCommandBuffer CommandBuffer, const VkCopyAccelerationStructureInfoKHR* Info);
+	static void CreateAccelerationStructureKHR(VkResult Result, VkDevice Device, const VkAccelerationStructureCreateInfoKHR* CreateInfo, const VkAllocationCallbacks* Allocator, VkAccelerationStructureKHR* AccelerationStructure) VULKAN_LAYER_BODY
+	static void DestroyAccelerationStructureKHR(VkResult Result, VkDevice Device, VkAccelerationStructureKHR AccelerationStructure, const VkAllocationCallbacks* Allocator) VULKAN_LAYER_BODY
+	static void CmdBuildAccelerationStructuresKHR(VkResult Result, VkCommandBuffer CommandBuffer, uint32 InfoCount, const VkAccelerationStructureBuildGeometryInfoKHR* Infos, const VkAccelerationStructureBuildRangeInfoKHR* const* BuildRangeInfos) VULKAN_LAYER_BODY
+	static void GetAccelerationStructureBuildSizesKHR(VkResult Result, VkDevice Device, VkAccelerationStructureBuildTypeKHR BuildType, const VkAccelerationStructureBuildGeometryInfoKHR* BuildInfo, const uint32* MaxPrimitiveCounts, VkAccelerationStructureBuildSizesInfoKHR* SizeInfo) VULKAN_LAYER_BODY
+	static void GetAccelerationStructureDeviceAddressKHR(VkResult Result, VkDevice Device, const VkAccelerationStructureDeviceAddressInfoKHR* Info) VULKAN_LAYER_BODY
+	static void CmdTraceRaysKHR(VkResult Result, VkCommandBuffer CommandBuffer, const VkStridedDeviceAddressRegionKHR* RaygenShaderBindingTable, const VkStridedDeviceAddressRegionKHR* MissShaderBindingTable, const VkStridedDeviceAddressRegionKHR* HitShaderBindingTable, const VkStridedDeviceAddressRegionKHR* CallableShaderBindingTable, uint32 width, uint32 height, uint32 depth) VULKAN_LAYER_BODY
+	static void CreateRayTracingPipelinesKHR(VkResult Result, VkDevice Device, VkDeferredOperationKHR DeferredOperation, VkPipelineCache PipelineCache, uint32_t CreateInfoCount, const VkRayTracingPipelineCreateInfoKHR* CreateInfos, const VkAllocationCallbacks* Allocator, VkPipeline* Pipelines) VULKAN_LAYER_BODY
+	static void GetRayTracingShaderGroupHandlesKHR(VkResult Result, VkDevice Device, VkPipeline Pipeline, uint32_t FirstGroup, uint32_t GroupCount, size_t DataSize, void* Data) VULKAN_LAYER_BODY
+	static void CmdWriteAccelerationStructuresPropertiesKHR(VkResult Result, VkCommandBuffer CommandBuffer, uint32_t AccelerationStructureCount, const VkAccelerationStructureKHR* AccelerationStructures, VkQueryType QueryType, VkQueryPool QueryPool, uint32_t FirstQuery) VULKAN_LAYER_BODY
+	static void CmdCopyAccelerationStructureKHR(VkResult Result, VkCommandBuffer CommandBuffer, const VkCopyAccelerationStructureInfoKHR* Info) VULKAN_LAYER_BODY
 #endif
 	static void GetBufferDeviceAddressKHR(VkResult Result, VkDevice Device, const VkBufferDeviceAddressInfo* Info) VULKAN_LAYER_BODY
 	static void GetDeviceImageMemoryRequirementsKHR(VkResult Result, VkDevice Device, const VkDeviceImageMemoryRequirements* Info, VkMemoryRequirements2* MemoryRequirements) VULKAN_LAYER_BODY
diff --git a/Engine/Source/Runtime/VulkanRHI/Private/VulkanPendingState.cpp b/Engine/Source/Runtime/VulkanRHI/Private/VulkanPendingState.cpp
--- a/Engine/Source/Runtime/VulkanRHI/Private/VulkanPendingState.cpp
+++ b/Engine/Source/Runtime/VulkanRHI/Private/VulkanPendingState.cpp
@@ -38,6 +38,19 @@ FVulkanDescriptorPool::FVulkanDescriptorPool(FVulkanDevice* InDevice, const FVul
 		}
 	}
 
+#if VULKAN_RHI_RAYTRACING
+	{
+		uint32 NumTypesUsed = Layout.GetTypesUsed(VK_DESCRIPTOR_TYPE_ACCELERATION_STRUCTURE_KHR);
+		if (NumTypesUsed > 0)
+		{
+			VkDescriptorPoolSize* Type = new(Types) VkDescriptorPoolSize;
+			FMemory::Memzero(*Type);
+			Type->type = VK_DESCRIPTOR_TYPE_ACCELERATION_STRUCTURE_KHR;
+			Type->descriptorCount = NumTypesUsed * MaxSetsAllocations;
+		}
+	}
+#endif
+
 	VkDescriptorPoolCreateInfo PoolInfo;
 	ZeroVulkanStruct(PoolInfo, VK_STRUCTURE_TYPE_DESCRIPTOR_POOL_CREATE_INFO);
 	// you don't need this flag because pool reset feature. Also this flag increase pool size in memory and vkResetDescriptorPool time.
@@ -74,6 +87,12 @@ void FVulkanDescriptorPool::TrackAddUsage(const FVulkanDescriptorSetsLayout& InL
 		ensure(Layout.GetTypesUsed((VkDescriptorType)TypeIndex) == InLayout.GetTypesUsed((VkDescriptorType)TypeIndex));
 	}
 
+#if VULKAN_RHI_RAYTRACING
+	{
+		ensure(Layout.GetTypesUsed(VK_DESCRIPTOR_TYPE_ACCELERATION_STRUCTURE_KHR) == InLayout.GetTypesUsed(VK_DESCRIPTOR_TYPE_ACCELERATION_STRUCTURE_KHR));
+	}
+#endif
+
 	NumAllocatedDescriptorSets += InLayout.GetLayouts().Num();
 	PeakAllocatedDescriptorSets = FMath::Max(NumAllocatedDescriptorSets, PeakAllocatedDescriptorSets);
 }
@@ -85,6 +104,12 @@ void FVulkanDescriptorPool::TrackRemoveUsage(const FVulkanDescriptorSetsLayout&
 		check(Layout.GetTypesUsed((VkDescriptorType)TypeIndex) == InLayout.GetTypesUsed((VkDescriptorType)TypeIndex));
 	}
 
+#if VULKAN_RHI_RAYTRACING
+	{
+		check(Layout.GetTypesUsed(VK_DESCRIPTOR_TYPE_ACCELERATION_STRUCTURE_KHR) == InLayout.GetTypesUsed(VK_DESCRIPTOR_TYPE_ACCELERATION_STRUCTURE_KHR));
+	}
+#endif
+
 	NumAllocatedDescriptorSets -= InLayout.GetLayouts().Num();
 }
 
