VulkanRHI: Use array for dynamic buffer allocations and removed the hardcoded NUM_BUFFERS=3 value.  Use a fence to manage buffer availability.  Do not rotate dynamic buffer usage when a staging buffer is used for the copy.

diff --git a/Engine/Source/Runtime/VulkanRHI/Private/VulkanBuffer.cpp b/Engine/Source/Runtime/VulkanRHI/Private/VulkanBuffer.cpp
--- a/Engine/Source/Runtime/VulkanRHI/Private/VulkanBuffer.cpp
+++ b/Engine/Source/Runtime/VulkanRHI/Private/VulkanBuffer.cpp
@@ -132,10 +132,6 @@ VkBufferUsageFlags FVulkanResourceMultiBuffer::UEToVKBufferUsageFlags(FVulkanDev
 FVulkanResourceMultiBuffer::FVulkanResourceMultiBuffer(FVulkanDevice* InDevice, uint32 InSize, EBufferUsageFlags InUEUsage, uint32 InStride, FRHIResourceCreateInfo& CreateInfo, FRHICommandListBase* InRHICmdList, const FRHITransientHeapAllocation* InTransientHeapAllocation)
 	: FRHIBuffer(InSize, InUEUsage, InStride)
 	, VulkanRHI::FDeviceChild(InDevice)
-	, BufferUsageFlags(0)
-	, NumBuffers(0)
-	, DynamicBufferIndex(0)
-	, LockStatus(ELockStatus::Unlocked)
 {
 	VULKAN_TRACK_OBJECT_CREATE(FVulkanResourceMultiBuffer, this);
 
@@ -151,6 +147,9 @@ FVulkanResourceMultiBuffer::FVulkanResourceMultiBuffer(FVulkanDevice* InDevice,
 		{
 			check(InRHICmdList);
 
+			// Volatile buffers always work out of the same first slot
+			CurrentBufferIndex = BufferAllocs.Emplace();
+
 			// Get a dummy buffer as sometimes the high-level misbehaves and tries to use SRVs off volatile buffers before filling them in...
 			void* Data = Lock(*InRHICmdList, RLM_WriteOnly, InSize, 0);
 
@@ -168,51 +167,25 @@ FVulkanResourceMultiBuffer::FVulkanResourceMultiBuffer(FVulkanDevice* InDevice,
 		}
 		else
 		{
-			NumBuffers = GetNumBuffersFromUsage(InUEUsage);
-			check(NumBuffers <= UE_ARRAY_COUNT(Buffers));
-
 			const bool bUnifiedMem = InDevice->HasUnifiedMemory();
 			const uint32 BufferAlignment = FMemoryManager::CalculateBufferAlignment(*InDevice, InUEUsage, bZeroSize);
 
+
 			if (InTransientHeapAllocation != nullptr)
 			{
-				const uint32 AlignedSize = Align(InSize, BufferAlignment);
-
-				Buffers[0] = FVulkanTransientHeap::GetVulkanAllocation(*InTransientHeapAllocation);
-				Buffers[0].Size = InSize;
-				check(Buffers[0].Offset % BufferAlignment == 0);
-				for (uint32 Index = 1; Index < NumBuffers; ++Index)
-				{
-					Buffers[Index] = Buffers[Index - 1];
-					Buffers[Index].Offset += AlignedSize;
-					Buffers[Index].Size = InSize;
-					check((Buffers[Index].Offset + InSize) <= InTransientHeapAllocation->Size);
-				}
-				
-				VULKAN_SET_DEBUG_NAME((*InDevice), VK_OBJECT_TYPE_BUFFER, Buffers[0].VulkanHandle, TEXT("%s:(FVulkanResourceMultiBuffer*)0x%p:Buffer[1,2,3]"), CreateInfo.DebugName ? CreateInfo.DebugName : TEXT("?"), this);
+				FBufferAlloc NewBufferAlloc;
+				NewBufferAlloc.Alloc = FVulkanTransientHeap::GetVulkanAllocation(*InTransientHeapAllocation);
+				NewBufferAlloc.HostPtr = bUnifiedMem ? NewBufferAlloc.Alloc.GetMappedPointer(Device) : nullptr;
+				check(NewBufferAlloc.Alloc.Offset % BufferAlignment == 0);
+				check(NewBufferAlloc.Alloc.Size >= InSize);
+				CurrentBufferIndex = BufferAllocs.Add(NewBufferAlloc);
 			}
 			else
 			{
-				VkMemoryPropertyFlags BufferMemFlags = VK_MEMORY_PROPERTY_DEVICE_LOCAL_BIT;
-				if (bUnifiedMem)
-				{
-					BufferMemFlags |= (VK_MEMORY_PROPERTY_HOST_VISIBLE_BIT | VK_MEMORY_PROPERTY_HOST_COHERENT_BIT);
-				}
-
-				for (uint32 Index = 0; Index < NumBuffers; ++Index)
-				{
-					if (!InDevice->GetMemoryManager().AllocateBufferPooled(Buffers[Index], nullptr, InSize, BufferAlignment, BufferUsageFlags, BufferMemFlags, EVulkanAllocationMetaMultiBuffer, __FILE__, __LINE__))
-					{
-						InDevice->GetMemoryManager().HandleOOM();
-					}
-					VULKAN_SET_DEBUG_NAME((*InDevice), VK_OBJECT_TYPE_BUFFER, Buffers[Index].VulkanHandle, TEXT("%s:(FVulkanResourceMultiBuffer*)0x%p:Buffer[%d]"), CreateInfo.DebugName ? CreateInfo.DebugName : TEXT("?"), this, Index);
-				}
+				AdvanceBufferIndex();
 			}
 
-			Current.Alloc.Reference(Buffers[DynamicBufferIndex]);
-			Current.Handle = (VkBuffer)Current.Alloc.VulkanHandle;
-			Current.Offset = Current.Alloc.Offset;
-			Current.Size = InSize;
+			VULKAN_SET_DEBUG_NAME((*InDevice), VK_OBJECT_TYPE_BUFFER, (VkBuffer)BufferAllocs[CurrentBufferIndex].Alloc.VulkanHandle, TEXT("%s"), CreateInfo.DebugName ? CreateInfo.DebugName : TEXT("UnknownBuffer"));
 
 			if (CreateInfo.ResourceArray)
 			{
@@ -222,8 +195,7 @@ FVulkanResourceMultiBuffer::FVulkanResourceMultiBuffer(FVulkanDevice* InDevice,
 				// We know this buffer is not in use by GPU atm. If we do have a direct access initialize it without extra copies
 				if (bUnifiedMem)
 				{
-					void* Data = (uint8*)Buffers[DynamicBufferIndex].GetMappedPointer(Device);
-					FMemory::Memcpy(Data, CreateInfo.ResourceArray->GetResourceData(), CopyDataSize);
+					FMemory::Memcpy(BufferAllocs[CurrentBufferIndex].HostPtr, CreateInfo.ResourceArray->GetResourceData(), CopyDataSize);
 				}
 				else
 				{
@@ -235,7 +207,7 @@ FVulkanResourceMultiBuffer::FVulkanResourceMultiBuffer(FVulkanDevice* InDevice,
 				CreateInfo.ResourceArray->Discard();
 			}
 
-			UpdateVulkanBufferStats(InSize * NumBuffers, BufferUsageFlags, true);
+			UpdateVulkanBufferStats(InSize, BufferUsageFlags, true);
 		}
 	}
 }
@@ -244,14 +216,106 @@ FVulkanResourceMultiBuffer::~FVulkanResourceMultiBuffer()
 {
 	VULKAN_TRACK_OBJECT_DELETE(FVulkanResourceMultiBuffer, this);
 	uint64_t TotalSize = 0;
-	for (uint32 Index = 0; Index < NumBuffers; ++Index)
+	for (int32 Index = 0; Index < BufferAllocs.Num(); ++Index)
 	{
-		TotalSize += Buffers[Index].Size;
-		Device->GetMemoryManager().FreeVulkanAllocation(Buffers[Index]);
+		if (BufferAllocs[Index].Alloc.HasAllocation())
+		{
+			TotalSize += BufferAllocs[Index].Alloc.Size;
+			Device->GetMemoryManager().FreeVulkanAllocation(BufferAllocs[Index].Alloc);
+		}
+		BufferAllocs[Index].Fence = nullptr;
 	}
 	UpdateVulkanBufferStats(TotalSize, BufferUsageFlags, false);
 }
 
+
+void FVulkanResourceMultiBuffer::AdvanceBufferIndex()
+{
+	auto FreePreviousAlloc = [&]() {
+		if (CurrentBufferIndex >= 0)
+		{
+			FBufferAlloc& CurrentBufferAlloc = BufferAllocs[CurrentBufferIndex];
+			CurrentBufferAlloc.AllocStatus = FBufferAlloc::EAllocStatus::NeedsFence;
+			CurrentBufferAlloc.Fence->Clear();
+		}
+	};
+
+	// Try to see if one of the buffers in our pool can be reused
+	if (BufferAllocs.Num() > 1)
+	{
+		for (int32 BufferIndex = 0; BufferIndex < BufferAllocs.Num(); ++BufferIndex)
+		{
+			if (CurrentBufferIndex == BufferIndex)
+			{
+				continue;
+			}
+
+			// Fences are only written on Unlock(), but are polled on lock/unlock
+			FBufferAlloc& BufferAlloc = BufferAllocs[BufferIndex];
+			if ((BufferAlloc.AllocStatus == FBufferAlloc::EAllocStatus::Pending) && BufferAlloc.Fence->Poll())
+			{
+				BufferAlloc.AllocStatus = FBufferAlloc::EAllocStatus::Available;
+				BufferAlloc.Fence->Clear();
+			}
+
+			if (BufferAlloc.AllocStatus == FBufferAlloc::EAllocStatus::Available)
+			{
+				FreePreviousAlloc();
+
+				CurrentBufferIndex = BufferIndex;
+				BufferAllocs[CurrentBufferIndex].AllocStatus = FBufferAlloc::EAllocStatus::InUse;
+				return;
+			}
+		}
+	}
+
+	// Allocate a new buffer
+	{
+		FreePreviousAlloc();
+
+		const bool bUnifiedMem = Device->HasUnifiedMemory();
+		const VkMemoryPropertyFlags BufferMemFlags = VK_MEMORY_PROPERTY_DEVICE_LOCAL_BIT | (bUnifiedMem ? (VK_MEMORY_PROPERTY_HOST_VISIBLE_BIT | VK_MEMORY_PROPERTY_HOST_COHERENT_BIT) : 0);
+		const uint32 BufferSize = GetSize();
+		const uint32 BufferAlignment = FMemoryManager::CalculateBufferAlignment(*Device, GetUsage(), (BufferSize == 0));
+
+		CurrentBufferIndex = BufferAllocs.Emplace();
+		FBufferAlloc& NewBufferAlloc = BufferAllocs[CurrentBufferIndex];
+		if (!Device->GetMemoryManager().AllocateBufferPooled(NewBufferAlloc.Alloc, nullptr, BufferSize, BufferAlignment, BufferUsageFlags, BufferMemFlags, EVulkanAllocationMetaMultiBuffer, __FILE__, __LINE__))
+		{
+			Device->GetMemoryManager().HandleOOM();
+		}
+		NewBufferAlloc.HostPtr = bUnifiedMem ? NewBufferAlloc.Alloc.GetMappedPointer(Device) : nullptr;
+		NewBufferAlloc.Fence = new FVulkanGPUFence(TEXT("VulkanDynamicBuffer"));
+		NewBufferAlloc.AllocStatus = FBufferAlloc::EAllocStatus::InUse;
+	}
+}
+
+void FVulkanResourceMultiBuffer::UpdateBufferAllocStates(FVulkanCommandListContext& Context)
+{
+	for (int32 BufferIndex = 0; BufferIndex < BufferAllocs.Num(); ++BufferIndex)
+	{
+		if (CurrentBufferIndex == BufferIndex)
+		{
+			continue;
+		}
+
+		FBufferAlloc& BufferAlloc = BufferAllocs[BufferIndex];
+		if (BufferAlloc.AllocStatus == FBufferAlloc::EAllocStatus::Pending)
+		{
+			if (BufferAlloc.Fence->Poll())
+			{
+				BufferAlloc.AllocStatus = FBufferAlloc::EAllocStatus::Available;
+				BufferAlloc.Fence->Clear();
+			}
+		}
+		else if (BufferAlloc.AllocStatus == FBufferAlloc::EAllocStatus::NeedsFence)
+		{
+			Context.RHIWriteGPUFence(BufferAlloc.Fence);
+			BufferAlloc.AllocStatus = FBufferAlloc::EAllocStatus::Pending;
+		}
+	}
+}
+
 void* FVulkanResourceMultiBuffer::Lock(FRHICommandListBase& RHICmdList, EResourceLockMode LockMode, uint32 LockSize, uint32 Offset)
 {
 	// Use the immediate context for write operations, since we are only accessing allocators.
@@ -265,47 +329,53 @@ void* FVulkanResourceMultiBuffer::Lock(FVulkanCommandListContext& Context, EReso
 	void* Data = nullptr;
 	uint32 DataOffset = 0;
 
-	const bool bDynamic = EnumHasAnyFlags(GetUsage(), BUF_Dynamic);
 	const bool bVolatile = EnumHasAnyFlags(GetUsage(), BUF_Volatile);
-	const bool bStatic = EnumHasAnyFlags(GetUsage(), BUF_Static) || !(bVolatile || bDynamic);
-	const bool bUAV = EnumHasAnyFlags(GetUsage(), BUF_UnorderedAccess);
-	const bool bSR = EnumHasAnyFlags(GetUsage(), BUF_ShaderResource);
 
 	LockStatus = ELockStatus::Locked;
+	++LockCounter;
 
 	if (bVolatile)
 	{
-		check(NumBuffers == 0);
 		if (LockMode == RLM_ReadOnly)
 		{
 			checkf(0, TEXT("Volatile buffers can't be locked for read."));
 		}
 		else
 		{
-			Context.GetTempFrameAllocationBuffer().Alloc(LockSize + Offset, 256, VolatileLockInfo);
-			Data = VolatileLockInfo.Data;
-			++VolatileLockInfo.LockCounter;
-			check(!VolatileLockInfo.Allocation.HasAllocation());
-			Current.Alloc = VolatileLockInfo.Allocation;
-			Current.Handle = Current.Alloc.GetBufferHandle();
-			Current.Offset = VolatileLockInfo.CurrentOffset + Current.Alloc.Offset;
-			Current.Size = LockSize;
+			FBufferAlloc& BufferAlloc = BufferAllocs[0];
+
+			FTempFrameAllocationBuffer::FTempAllocInfo VolatileAlloc;
+			Context.GetTempFrameAllocationBuffer().Alloc(LockSize + Offset, 256, VolatileAlloc);
+			check(!VolatileAlloc.Allocation.HasAllocation());
+
+			BufferAlloc.Alloc.Reference(VolatileAlloc.Allocation);
+			Data = BufferAlloc.HostPtr = VolatileAlloc.Data;
+			DataOffset = Offset;
+
+			// Patch our alloc to go directly to our offset
+			BufferAlloc.Alloc.Offset += VolatileAlloc.CurrentOffset;
+			BufferAlloc.Alloc.Size = VolatileAlloc.Size;
 		}
 	}
 	else
 	{
+		const bool bDynamic = EnumHasAnyFlags(GetUsage(), BUF_Dynamic);
+		const bool bStatic = EnumHasAnyFlags(GetUsage(), BUF_Static) || !(bVolatile || bDynamic);
+		const bool bUAV = EnumHasAnyFlags(GetUsage(), BUF_UnorderedAccess);
+		const bool bSR = EnumHasAnyFlags(GetUsage(), BUF_ShaderResource);
+		const bool bUnifiedMem = Device->HasUnifiedMemory();
+
 		check(bStatic || bDynamic || bUAV || bSR);
 
 		if (LockMode == RLM_ReadOnly)
 		{
 			check(IsInRenderingThread() && Context.IsImmediate());
 
-			const bool bUnifiedMem = Device->HasUnifiedMemory();
 			if (bUnifiedMem)
 			{
-				Data = Buffers[DynamicBufferIndex].GetMappedPointer(Device);
-				LockStatus = ELockStatus::PersistentMapping;
+				Data = BufferAllocs[CurrentBufferIndex].HostPtr;
 				DataOffset = Offset;
+				LockStatus = ELockStatus::PersistentMapping;
 			}
 			else 
 			{
@@ -314,7 +384,7 @@ void* FVulkanResourceMultiBuffer::Lock(FVulkanCommandListContext& Context, EReso
 				
 				// Make sure any previous tasks have finished on the source buffer.
 				VkMemoryBarrier BarrierBefore = { VK_STRUCTURE_TYPE_MEMORY_BARRIER, nullptr, VK_ACCESS_MEMORY_WRITE_BIT, VK_ACCESS_MEMORY_READ_BIT };
-				VulkanRHI::vkCmdPipelineBarrier(CmdBuffer->GetHandle(), VK_PIPELINE_STAGE_TRANSFER_BIT, VK_PIPELINE_STAGE_TRANSFER_BIT, 0, 1, &BarrierBefore, 0, nullptr, 0, nullptr);
+				VulkanRHI::vkCmdPipelineBarrier(CmdBuffer->GetHandle(), VK_PIPELINE_STAGE_BOTTOM_OF_PIPE_BIT, VK_PIPELINE_STAGE_TRANSFER_BIT, 0, 1, &BarrierBefore, 0, nullptr, 0, nullptr);
 
 				// Create a staging buffer we can use to copy data from device to cpu.
 				VulkanRHI::FStagingBuffer* StagingBuffer = Device->GetStagingManager().AcquireBuffer(LockSize, VK_BUFFER_USAGE_TRANSFER_DST_BIT, VK_MEMORY_PROPERTY_HOST_CACHED_BIT);
@@ -322,10 +392,10 @@ void* FVulkanResourceMultiBuffer::Lock(FVulkanCommandListContext& Context, EReso
 				// Fill the staging buffer with the data on the device.
 				VkBufferCopy Regions;
 				Regions.size = LockSize;
-				Regions.srcOffset = Offset + Buffers[DynamicBufferIndex].Offset;
+				Regions.srcOffset = Offset + BufferAllocs[CurrentBufferIndex].Alloc.Offset;
 				Regions.dstOffset = 0;
 				
-				VulkanRHI::vkCmdCopyBuffer(CmdBuffer->GetHandle(), Buffers[DynamicBufferIndex].GetBufferHandle(), StagingBuffer->GetHandle(), 1, &Regions);
+				VulkanRHI::vkCmdCopyBuffer(CmdBuffer->GetHandle(), BufferAllocs[CurrentBufferIndex].Alloc.GetBufferHandle(), StagingBuffer->GetHandle(), 1, &Regions);
 
 				// Setup barrier.
 				VkMemoryBarrier BarrierAfter = { VK_STRUCTURE_TYPE_MEMORY_BARRIER, nullptr, VK_ACCESS_MEMORY_WRITE_BIT, VK_ACCESS_HOST_READ_BIT };
@@ -360,16 +430,13 @@ void* FVulkanResourceMultiBuffer::Lock(FVulkanCommandListContext& Context, EReso
 		else
 		{
 			check(LockMode == RLM_WriteOnly);
-			DynamicBufferIndex = (DynamicBufferIndex + 1) % NumBuffers;
-			Current.Alloc.Reference(Buffers[DynamicBufferIndex]);		
-			Current.Handle = (VkBuffer)Current.Alloc.VulkanHandle;
-			Current.Offset = Current.Alloc.Offset;
-			Current.Size = LockSize;
 
 			// Always use staging buffers to update 'Static' buffers since they maybe be in use by GPU atm
-			const bool bUseStagingBuffer = (bStatic || !Device->HasUnifiedMemory());
+			const bool bUseStagingBuffer = (bStatic || !bUnifiedMem);
 			if (bUseStagingBuffer)
 			{
+				// NOTE: No need to change the CurrentBufferIndex if we're using a staging buffer for the copy
+
 				VulkanRHI::FPendingBufferLock PendingLock;
 				PendingLock.Offset = Offset;
 				PendingLock.Size = LockSize;
@@ -387,7 +454,9 @@ void* FVulkanResourceMultiBuffer::Lock(FVulkanCommandListContext& Context, EReso
 			}
 			else
 			{
-				Data = Buffers[DynamicBufferIndex].GetMappedPointer(Device);
+				AdvanceBufferIndex();
+
+				Data = BufferAllocs[CurrentBufferIndex].HostPtr;
 				DataOffset = Offset;
 				LockStatus = ELockStatus::PersistentMapping;
 			}
@@ -398,10 +467,10 @@ void* FVulkanResourceMultiBuffer::Lock(FVulkanCommandListContext& Context, EReso
 	return (uint8*)Data + DataOffset;
 }
 
-inline void FVulkanResourceMultiBuffer::InternalUnlock(FVulkanCommandListContext& Context, VulkanRHI::FPendingBufferLock& PendingLock, FVulkanResourceMultiBuffer* MultiBuffer, int32 InDynamicBufferIndex)
+inline void FVulkanResourceMultiBuffer::InternalUnlock(FVulkanCommandListContext& Context, VulkanRHI::FPendingBufferLock& PendingLock, FVulkanResourceMultiBuffer* MultiBuffer, int32 InBufferIndex)
 {
-	uint32 LockSize = PendingLock.Size;
-	uint32 LockOffset = PendingLock.Offset;
+	const uint32 LockSize = PendingLock.Size;
+	const uint32 LockOffset = PendingLock.Offset;
 	VulkanRHI::FStagingBuffer* StagingBuffer = PendingLock.StagingBuffer;
 	PendingLock.StagingBuffer = nullptr;
 
@@ -417,14 +486,16 @@ inline void FVulkanResourceMultiBuffer::InternalUnlock(FVulkanCommandListContext
 	FMemory::Memzero(Region);
 	Region.size = LockSize;
 	//Region.srcOffset = 0;
-	Region.dstOffset = LockOffset + MultiBuffer->Buffers[InDynamicBufferIndex].Offset;
-	VulkanRHI::vkCmdCopyBuffer(CmdBuffer, StagingBuffer->GetHandle(), MultiBuffer->Buffers[InDynamicBufferIndex].GetBufferHandle(), 1, &Region);
+	Region.dstOffset = LockOffset + MultiBuffer->BufferAllocs[InBufferIndex].Alloc.Offset;
+	VulkanRHI::vkCmdCopyBuffer(CmdBuffer, StagingBuffer->GetHandle(), MultiBuffer->BufferAllocs[InBufferIndex].Alloc.GetBufferHandle(), 1, &Region);
 
 	// High level code expects the data in MultiBuffer to be ready to read
 	VkMemoryBarrier BarrierAfter = { VK_STRUCTURE_TYPE_MEMORY_BARRIER, nullptr, VK_ACCESS_TRANSFER_WRITE_BIT, VK_ACCESS_MEMORY_READ_BIT | VK_ACCESS_MEMORY_WRITE_BIT };
 	VulkanRHI::vkCmdPipelineBarrier(CmdBuffer, VK_PIPELINE_STAGE_TRANSFER_BIT, VK_PIPELINE_STAGE_ALL_COMMANDS_BIT, 0, 1, &BarrierAfter, 0, nullptr, 0, nullptr);
 
 	MultiBuffer->GetParent()->GetStagingManager().ReleaseBuffer(Cmd, StagingBuffer);
+
+	MultiBuffer->UpdateBufferAllocStates(Context);
 }
 
 struct FRHICommandMultiBufferUnlock final : public FRHICommand<FRHICommandMultiBufferUnlock>
@@ -432,19 +503,19 @@ struct FRHICommandMultiBufferUnlock final : public FRHICommand<FRHICommandMultiB
 	VulkanRHI::FPendingBufferLock PendingLock;
 	FVulkanResourceMultiBuffer* MultiBuffer;
 	FVulkanDevice* Device;
-	int32 DynamicBufferIndex;
+	int32 BufferIndex;
 
-	FRHICommandMultiBufferUnlock(FVulkanDevice* InDevice, const VulkanRHI::FPendingBufferLock& InPendingLock, FVulkanResourceMultiBuffer* InMultiBuffer, int32 InDynamicBufferIndex)
+	FRHICommandMultiBufferUnlock(FVulkanDevice* InDevice, const VulkanRHI::FPendingBufferLock& InPendingLock, FVulkanResourceMultiBuffer* InMultiBuffer, int32 InBufferIndex)
 		: PendingLock(InPendingLock)
 		, MultiBuffer(InMultiBuffer)
 		, Device(InDevice)
-		, DynamicBufferIndex(InDynamicBufferIndex)
+		, BufferIndex(InBufferIndex)
 	{
 	}
 
 	void Execute(FRHICommandListBase& CmdList)
 	{
-		FVulkanResourceMultiBuffer::InternalUnlock(FVulkanCommandListContext::GetVulkanContext(CmdList.GetContext()), PendingLock, MultiBuffer, DynamicBufferIndex);
+		FVulkanResourceMultiBuffer::InternalUnlock(FVulkanCommandListContext::GetVulkanContext(CmdList.GetContext()), PendingLock, MultiBuffer, BufferIndex);
 	}
 };
 
@@ -459,10 +530,25 @@ void FVulkanResourceMultiBuffer::Unlock(FRHICommandListBase* RHICmdList, FVulkan
 
 	check(LockStatus != ELockStatus::Unlocked);
 
-	if (bVolatile || LockStatus == ELockStatus::PersistentMapping)
+	if (bVolatile)
 	{
 		// Nothing to do here...
 	}
+	else if (LockStatus == ELockStatus::PersistentMapping)
+	{
+		if (Context)
+		{
+			UpdateBufferAllocStates(*Context);
+		}
+		else
+		{
+			RHICmdList->EnqueueLambda([this](FRHICommandListBase& CmdList)
+			{
+				FVulkanCommandListContext& Context = FVulkanCommandListContext::GetVulkanContext(CmdList.GetContext());
+				UpdateBufferAllocStates(Context);
+			});
+		}
+	}
 	else
 	{
 		check(bStatic || bDynamic || bSR);
@@ -485,11 +571,11 @@ void FVulkanResourceMultiBuffer::Unlock(FRHICommandListBase* RHICmdList, FVulkan
 					Context = &FVulkanCommandListContext::GetVulkanContext(RHICmdList->GetContext());
 				}
 
-				FVulkanResourceMultiBuffer::InternalUnlock(*Context, PendingLock, this, DynamicBufferIndex);
+				FVulkanResourceMultiBuffer::InternalUnlock(*Context, PendingLock, this, CurrentBufferIndex);
 			}
 			else
 			{
-				ALLOC_COMMAND_CL(*RHICmdList, FRHICommandMultiBufferUnlock)(Device, PendingLock, this, DynamicBufferIndex);
+				ALLOC_COMMAND_CL(*RHICmdList, FRHICommandMultiBufferUnlock)(Device, PendingLock, this, CurrentBufferIndex);
 			}
 		}
 	}
@@ -507,11 +593,9 @@ void FVulkanResourceMultiBuffer::Swap(FVulkanResourceMultiBuffer& Other)
 	::Swap(Device, Other.Device);
 	
 	::Swap(BufferUsageFlags, Other.BufferUsageFlags);
-	::Swap(NumBuffers, Other.NumBuffers);
-	::Swap(DynamicBufferIndex, Other.DynamicBufferIndex);
-	::Swap(Buffers, Other.Buffers);
-	::Swap(Current, Other.Current);
-	::Swap(VolatileLockInfo, Other.VolatileLockInfo);
+	::Swap(CurrentBufferIndex, Other.CurrentBufferIndex);
+	::Swap(BufferAllocs, Other.BufferAllocs);
+	::Swap(LockCounter, Other.LockCounter);
 }
 
 FBufferRHIRef FVulkanDynamicRHI::RHICreateBuffer(FRHICommandListBase& RHICmdList, uint32 Size, EBufferUsageFlags Usage, uint32 Stride, ERHIAccess ResourceState, FRHIResourceCreateInfo& ResourceCreateInfo)
diff --git a/Engine/Source/Runtime/VulkanRHI/Private/VulkanResources.h b/Engine/Source/Runtime/VulkanRHI/Private/VulkanResources.h
--- a/Engine/Source/Runtime/VulkanRHI/Private/VulkanResources.h
+++ b/Engine/Source/Runtime/VulkanRHI/Private/VulkanResources.h
@@ -910,55 +910,55 @@ public:
 
 	inline const VulkanRHI::FVulkanAllocation& GetCurrentAllocation() const
 	{
-		return Current.Alloc;
+		return BufferAllocs[CurrentBufferIndex].Alloc;
 	}
 
 	inline VkBuffer GetHandle() const
 	{
-		return Current.Handle;
+		return (VkBuffer)GetCurrentAllocation().VulkanHandle;
 	}
 
 	inline bool IsDynamic() const
 	{
-		return NumBuffers > 1;
+		return EnumHasAnyFlags(GetUsage(), BUF_Dynamic);
 	}
 
 	inline int32 GetDynamicIndex() const
 	{
-		return DynamicBufferIndex;
+		return CurrentBufferIndex;
 	}
 
 	inline bool IsVolatile() const
 	{
-		return NumBuffers == 0;
+		return EnumHasAnyFlags(GetUsage(), BUF_Volatile);
 	}
 
 	inline uint32 GetVolatileLockCounter() const
 	{
 		check(IsVolatile());
-		return VolatileLockInfo.LockCounter;
+		return LockCounter;
 	}
 	inline uint32 GetVolatileLockSize() const
 	{
 		check(IsVolatile());
-		return VolatileLockInfo.Size;
+		return GetCurrentAllocation().Size;
 	}
 
 	inline int32 GetNumBuffers() const
 	{
-		return NumBuffers;
+		return BufferAllocs.Num();
 	}
 
 	// Offset used for Binding a VkBuffer
 	inline uint32 GetOffset() const
 	{
-		return Current.Offset;
+		return GetCurrentAllocation().Offset;
 	}
 
 	// Remaining size from the current offset
 	inline uint64 GetCurrentSize() const
 	{
-		return Current.Alloc.Size - (Current.Offset - Current.Alloc.Offset);
+		return GetCurrentAllocation().Size;
 	}
 
 	inline VkBufferUsageFlags GetBufferUsageFlags() const
@@ -998,40 +998,39 @@ public:
 
 	static VkBufferUsageFlags UEToVKBufferUsageFlags(FVulkanDevice* InDevice, EBufferUsageFlags InUEUsage, bool bZeroSize);
 
-	static inline int32 GetNumBuffersFromUsage(EBufferUsageFlags InUEUsage)
-	{
-		const bool bDynamic = EnumHasAnyFlags(InUEUsage, BUF_Dynamic);
-		return bDynamic ? NUM_BUFFERS : 1;
-	}
-
 protected:
-	void Unlock(FRHICommandListBase* RHICmdList, FVulkanCommandListContext* Context);
 
+	void AdvanceBufferIndex();
+	void UpdateBufferAllocStates(FVulkanCommandListContext& Context);
+
+	void Unlock(FRHICommandListBase* RHICmdList, FVulkanCommandListContext* Context);
 
 	VkBufferUsageFlags BufferUsageFlags;
-	uint8 NumBuffers;
-	uint8 DynamicBufferIndex;
+
 	enum class ELockStatus : uint8
 	{
 		Unlocked,
 		Locked,
 		PersistentMapping,
-	} LockStatus;
+	} LockStatus = ELockStatus::Unlocked;
 
-	enum
-	{
-		NUM_BUFFERS = 3,
-	};
-
-	VulkanRHI::FVulkanAllocation Buffers[NUM_BUFFERS];
-	struct
+	struct FBufferAlloc
 	{
 		VulkanRHI::FVulkanAllocation Alloc;
-		VkBuffer Handle = VK_NULL_HANDLE;
-		uint64 Offset = 0;
-		uint64 Size = 0;
-	} Current;
-	VulkanRHI::FTempFrameAllocationBuffer::FTempAllocInfo VolatileLockInfo;
+		void* HostPtr = nullptr;
+		class FVulkanGPUFence* Fence = nullptr;
+
+		enum class EAllocStatus : uint8
+		{
+			Available,	// The allocation is ready to be used
+			InUse,		// CurrentBufferIndex should point to this allocation
+			NeedsFence,	// The allocation was just released and needs a fence to make sure previous commands are done with it
+			Pending,	// Fence was written, we are waiting on it to know that the alloc can be used again
+		} AllocStatus = EAllocStatus::Available;
+	};
+	TArray<FBufferAlloc, TInlineAllocator<3>> BufferAllocs;
+	int32 CurrentBufferIndex = -1;
+	uint32 LockCounter = 0;
 
 	static void InternalUnlock(FVulkanCommandListContext& Context, VulkanRHI::FPendingBufferLock& PendingLock, FVulkanResourceMultiBuffer* MultiBuffer, int32 InDynamicBufferIndex);
 
diff --git a/Engine/Source/Runtime/VulkanRHI/Private/VulkanTransientResourceAllocator.cpp b/Engine/Source/Runtime/VulkanRHI/Private/VulkanTransientResourceAllocator.cpp
--- a/Engine/Source/Runtime/VulkanRHI/Private/VulkanTransientResourceAllocator.cpp
+++ b/Engine/Source/Runtime/VulkanRHI/Private/VulkanTransientResourceAllocator.cpp
@@ -133,7 +133,7 @@ FRHITransientBuffer* FVulkanTransientResourceAllocator::CreateBuffer(const FRHIB
 
 	const bool bZeroSize = (InCreateInfo.Size == 0);
 	const uint32 Alignment = FMemoryManager::CalculateBufferAlignment(*Device, InCreateInfo.Usage, bZeroSize);
-	uint64 Size = Align(InCreateInfo.Size, Alignment) * FVulkanResourceMultiBuffer::GetNumBuffersFromUsage(InCreateInfo.Usage);
+	uint64 Size = Align(InCreateInfo.Size, Alignment);
 
 	return CreateBufferInternal(InCreateInfo, InDebugName, InPassIndex, Size, Alignment,
 		[&](const FRHITransientHeap::FResourceInitializer& Initializer)
diff --git a/Engine/Source/Runtime/VulkanRHI/Private/VulkanUAV.cpp b/Engine/Source/Runtime/VulkanRHI/Private/VulkanUAV.cpp
--- a/Engine/Source/Runtime/VulkanRHI/Private/VulkanUAV.cpp
+++ b/Engine/Source/Runtime/VulkanRHI/Private/VulkanUAV.cpp
@@ -202,6 +202,11 @@ void FVulkanShaderResourceView::UpdateView()
 		else if (SourceBuffer->IsDynamic() && (BufferIndex != SourceBuffer->GetDynamicIndex()))
 		{
 			BufferIndex = SourceBuffer->GetDynamicIndex();
+			if (BufferIndex >= (uint32)BufferViews.Num())
+			{
+				BufferViews.AddZeroed(BufferIndex - BufferViews.Num() + 1);
+			}
+
 			FreeBindlessHandle();
 		}
 
