Disable SN-DBS distribution for VulkanRHI as it cannot resolve a preprocessor in an #include

#rnx

diff --git a/Engine/Source/Runtime/VulkanRHI/VulkanRHI.Build.cs b/Engine/Source/Runtime/VulkanRHI/VulkanRHI.Build.cs
--- a/Engine/Source/Runtime/VulkanRHI/VulkanRHI.Build.cs
+++ b/Engine/Source/Runtime/VulkanRHI/VulkanRHI.Build.cs
@@ -10,6 +10,7 @@ public class VulkanRHI : ModuleRules
 	{
 		IWYUSupport = IWYUSupport.None;
 		bLegalToDistributeObjectCode = true;
+		bBuildLocallyWithSNDBS = false; // VulkanPlatform.h
 
 		if (Target.Platform.IsInGroup(UnrealPlatformGroup.Windows))
 		{
