Fixing vulkan platform header names we can use COMPILED_PLATFORM_HEADER instead of our own version.

diff --git a/Engine/Source/Runtime/VulkanRHI/Private/Android/AndroidVulkanPlatformDefines.h b/Engine/Source/Runtime/VulkanRHI/Private/Android/AndroidVulkanPlatformDefines.h
new file mode 100644
--- /dev/null
+++ b/Engine/Source/Runtime/VulkanRHI/Private/Android/AndroidVulkanPlatformDefines.h
@@ -0,0 +1,5 @@
+// Copyright Epic Games, Inc. All Rights Reserved.
+
+#pragma once
+
+#include "VulkanAndroidPlatform.h"
diff --git a/Engine/Source/Runtime/VulkanRHI/Private/Android/VulkanPlatformDefines.h b/Engine/Source/Runtime/VulkanRHI/Private/Android/VulkanPlatformDefines.h
deleted file mode 100644
diff --git a/Engine/Source/Runtime/VulkanRHI/Private/Linux/LinuxVulkanPlatformDefines.h b/Engine/Source/Runtime/VulkanRHI/Private/Linux/LinuxVulkanPlatformDefines.h
new file mode 100644
--- /dev/null
+++ b/Engine/Source/Runtime/VulkanRHI/Private/Linux/LinuxVulkanPlatformDefines.h
@@ -0,0 +1,5 @@
+// Copyright Epic Games, Inc. All Rights Reserved.
+
+#pragma once
+
+#include "VulkanLinuxPlatform.h"
diff --git a/Engine/Source/Runtime/VulkanRHI/Private/Linux/VulkanPlatformDefines.h b/Engine/Source/Runtime/VulkanRHI/Private/Linux/VulkanPlatformDefines.h
deleted file mode 100644
diff --git a/Engine/Source/Runtime/VulkanRHI/Private/Mac/MacVulkanPlatformDefines.h b/Engine/Source/Runtime/VulkanRHI/Private/Mac/MacVulkanPlatformDefines.h
new file mode 100644
--- /dev/null
+++ b/Engine/Source/Runtime/VulkanRHI/Private/Mac/MacVulkanPlatformDefines.h
@@ -0,0 +1,3 @@
+// Copyright Epic Games, Inc. All Rights Reserved.
+
+#pragma once
diff --git a/Engine/Source/Runtime/VulkanRHI/Private/Mac/VulkanPlatformDefines.h b/Engine/Source/Runtime/VulkanRHI/Private/Mac/VulkanPlatformDefines.h
deleted file mode 100644
diff --git a/Engine/Source/Runtime/VulkanRHI/Private/VulkanPlatform.h b/Engine/Source/Runtime/VulkanRHI/Private/VulkanPlatform.h
--- a/Engine/Source/Runtime/VulkanRHI/Private/VulkanPlatform.h
+++ b/Engine/Source/Runtime/VulkanRHI/Private/VulkanPlatform.h
@@ -4,11 +4,4 @@
 
 #include "HAL/PreprocessorHelpers.h"
 
-// Vulkan is currently organized just a little differently than COMPILED_PLATFORM_HEADER does things.
-#if PLATFORM_IS_EXTENSION
-	#define VULKAN_COMPILED_PLATFORM_HEADER(Suffix) PREPROCESSOR_TO_STRING(Suffix)
-#else
-	#define VULKAN_COMPILED_PLATFORM_HEADER(Suffix) PREPROCESSOR_TO_STRING(PLATFORM_HEADER_NAME/Suffix)
-#endif
-
-#include VULKAN_COMPILED_PLATFORM_HEADER(VulkanPlatformDefines.h)
+#include COMPILED_PLATFORM_HEADER(VulkanPlatformDefines.h)
diff --git a/Engine/Source/Runtime/VulkanRHI/Private/Windows/VulkanPlatformDefines.h b/Engine/Source/Runtime/VulkanRHI/Private/Windows/VulkanPlatformDefines.h
deleted file mode 100644
diff --git a/Engine/Source/Runtime/VulkanRHI/Private/Windows/WindowsVulkanPlatformDefines.h b/Engine/Source/Runtime/VulkanRHI/Private/Windows/WindowsVulkanPlatformDefines.h
new file mode 100644
--- /dev/null
+++ b/Engine/Source/Runtime/VulkanRHI/Private/Windows/WindowsVulkanPlatformDefines.h
@@ -0,0 +1,6 @@
+// Copyright Epic Games, Inc. All Rights Reserved.
+
+#pragma once
+
+#include "VulkanWindowsPlatform.h"
+
