VulkanRHI:  Clean up of Vulkan extension code to make it more consistent.  Add flag to keep track of which extensions were present but failed a feature/property check.  WIP

diff --git a/Engine/Source/Runtime/VulkanRHI/Private/VulkanExtensions.cpp b/Engine/Source/Runtime/VulkanRHI/Private/VulkanExtensions.cpp
--- a/Engine/Source/Runtime/VulkanRHI/Private/VulkanExtensions.cpp
+++ b/Engine/Source/Runtime/VulkanRHI/Private/VulkanExtensions.cpp
@@ -154,7 +154,8 @@ public:
 
 	virtual void PostPhysicalDeviceFeatures(FOptionalVulkanDeviceExtensions& ExtensionFlags) override final
 	{
-		ExtensionFlags.HasKHRMaintenance4 = (Maintenance4Features.maintenance4 == VK_TRUE) ? 1 : 0; 
+		bRequirementsPassed = (Maintenance4Features.maintenance4 == VK_TRUE);
+		ExtensionFlags.HasKHRMaintenance4 = bRequirementsPassed;
 	}
 
 private:
@@ -211,20 +212,16 @@ public:
 
 	virtual void PostPhysicalDeviceFeatures(FOptionalVulkanDeviceExtensions& ExtensionFlags) override final
 	{
-		ExtensionFlags.HasKHRShaderAtomicInt64 = BufferAtomicFeatures.shaderBufferInt64Atomics;
-		GRHISupportsAtomicUInt64 = (BufferAtomicFeatures.shaderBufferInt64Atomics == VK_TRUE);
+		bRequirementsPassed = (BufferAtomicFeatures.shaderBufferInt64Atomics == VK_TRUE);
+		ExtensionFlags.HasKHRShaderAtomicInt64 = bRequirementsPassed;
 	}
 
 	virtual void PreCreateDevice(VkDeviceCreateInfo& DeviceCreateInfo) override final
 	{
-		if (BufferAtomicFeatures.shaderBufferInt64Atomics == VK_TRUE)
+		if (bRequirementsPassed)
 		{
 			AddToPNext(DeviceCreateInfo, BufferAtomicFeatures);
 		}
-		else
-		{
-			GRHISupportsAtomicUInt64 = false;
-		}
 	}
 
 	VkPhysicalDeviceShaderAtomicInt64Features BufferAtomicFeatures;
@@ -253,21 +250,19 @@ public:
 
 	virtual void PostPhysicalDeviceFeatures(FOptionalVulkanDeviceExtensions& ExtensionFlags) override final
 	{
-		ExtensionFlags.HasImageAtomicInt64 = ImageAtomicFeatures.shaderImageInt64Atomics;
-		GRHISupportsAtomicUInt64 = GRHISupportsAtomicUInt64 && (ImageAtomicFeatures.shaderImageInt64Atomics == VK_TRUE); // GRHISupportsAtomicUInt64 already contains FVulkanShaderAtomicInt64Extension support
+		bRequirementsPassed = (ImageAtomicFeatures.shaderImageInt64Atomics == VK_TRUE);
+		ExtensionFlags.HasImageAtomicInt64 = bRequirementsPassed;
 	}
 
 	virtual void PreCreateDevice(VkDeviceCreateInfo& DeviceCreateInfo) override final
 	{
-		const bool bEnable64bitAtomics = (GRHIAllow64bitShaderAtomicsCvar.GetValueOnAnyThread() != 0);
-		if (bEnable64bitAtomics && (ImageAtomicFeatures.shaderImageInt64Atomics == VK_TRUE))
-		{
-			ImageAtomicFeatures.pNext = (void*)DeviceCreateInfo.pNext;
-			DeviceCreateInfo.pNext = &ImageAtomicFeatures;
-		}
-		else
+		// The PreCreateDevice() call is after all extension have gone through PostPhysicalDeviceFeatures(), so ExtensionFlags will be filled for both.
+		const FOptionalVulkanDeviceExtensions& ExtensionFlags = Device->GetOptionalExtensions();
+		GRHISupportsAtomicUInt64 = ExtensionFlags.HasKHRShaderAtomicInt64 && ExtensionFlags.HasImageAtomicInt64;
+
+		if (bRequirementsPassed)
 		{
-			GRHISupportsAtomicUInt64 = false;
+			AddToPNext(DeviceCreateInfo, ImageAtomicFeatures);
 		}
 	}
 
@@ -315,12 +310,13 @@ public:
 
 	virtual void PostPhysicalDeviceFeatures(FOptionalVulkanDeviceExtensions& ExtensionFlags) override final
 	{
-		ExtensionFlags.HasSeparateDepthStencilLayouts = SeparateDepthStencilLayoutsFeatures.separateDepthStencilLayouts;
+		bRequirementsPassed = (SeparateDepthStencilLayoutsFeatures.separateDepthStencilLayouts == VK_TRUE);
+		ExtensionFlags.HasSeparateDepthStencilLayouts = bRequirementsPassed;
 	}
 
 	virtual void PreCreateDevice(VkDeviceCreateInfo& DeviceCreateInfo) override final
 	{
-		if (SeparateDepthStencilLayoutsFeatures.separateDepthStencilLayouts == VK_TRUE)
+		if (bRequirementsPassed)
 		{
 			AddToPNext(DeviceCreateInfo, SeparateDepthStencilLayoutsFeatures);
 		}
@@ -351,12 +347,13 @@ public:
 
 	virtual void PostPhysicalDeviceFeatures(FOptionalVulkanDeviceExtensions& ExtensionFlags) override final
 	{
-		ExtensionFlags.HasKHRSynchronization2 = Synchronization2Features.synchronization2;
+		bRequirementsPassed = (Synchronization2Features.synchronization2 == VK_TRUE);
+		ExtensionFlags.HasKHRSynchronization2 = bRequirementsPassed;
 	}
 
 	virtual void PreCreateDevice(VkDeviceCreateInfo& DeviceCreateInfo) override final
 	{
-		if (Synchronization2Features.synchronization2 == VK_TRUE)
+		if (bRequirementsPassed)
 		{
 			AddToPNext(DeviceCreateInfo, Synchronization2Features);
 		}
@@ -385,12 +382,13 @@ public:
 
 	virtual void PostPhysicalDeviceFeatures(FOptionalVulkanDeviceExtensions& ExtensionFlags) override final 
 	{
-		ExtensionFlags.HasKHRMultiview = DeviceMultiviewFeatures.multiview; 
+		bRequirementsPassed = (DeviceMultiviewFeatures.multiview == VK_TRUE);
+		ExtensionFlags.HasKHRMultiview = bRequirementsPassed;
 	}
 
 	virtual void PreCreateDevice(VkDeviceCreateInfo& DeviceCreateInfo) override final
 	{
-		if (DeviceMultiviewFeatures.multiview == VK_TRUE)
+		if (bRequirementsPassed)
 		{
 			AddToPNext(DeviceCreateInfo, DeviceMultiviewFeatures);
 		}
@@ -419,12 +417,13 @@ public:
 
 	virtual void PostPhysicalDeviceFeatures(FOptionalVulkanDeviceExtensions& ExtensionFlags) override final 
 	{
-		ExtensionFlags.HasEXTScalarBlockLayout = ScalarBlockLayoutFeatures.scalarBlockLayout; 
+		bRequirementsPassed = (ScalarBlockLayoutFeatures.scalarBlockLayout == VK_TRUE);
+		ExtensionFlags.HasEXTScalarBlockLayout = bRequirementsPassed;
 	}
 
 	virtual void PreCreateDevice(VkDeviceCreateInfo& DeviceCreateInfo) override final
 	{
-		if (ScalarBlockLayoutFeatures.scalarBlockLayout == VK_TRUE)
+		if (bRequirementsPassed)
 		{
 			AddToPNext(DeviceCreateInfo, ScalarBlockLayoutFeatures);
 		}
@@ -453,17 +452,17 @@ public:
 
 	virtual void PostPhysicalDeviceFeatures(FOptionalVulkanDeviceExtensions& ExtensionFlags) override final
 	{
-		bFullySupported = (DescriptorIndexingFeatures.runtimeDescriptorArray == VK_TRUE) &&
+		bRequirementsPassed = (DescriptorIndexingFeatures.runtimeDescriptorArray == VK_TRUE) &&
 			(DescriptorIndexingFeatures.descriptorBindingPartiallyBound == VK_TRUE) &&
 			(DescriptorIndexingFeatures.descriptorBindingUpdateUnusedWhilePending == VK_TRUE) &&
 			(DescriptorIndexingFeatures.descriptorBindingVariableDescriptorCount == VK_TRUE);
 
-		ExtensionFlags.HasEXTDescriptorIndexing = bFullySupported ? 1 : 0;
+		ExtensionFlags.HasEXTDescriptorIndexing = bRequirementsPassed;
 	}
 
 	virtual void PreCreateDevice(VkDeviceCreateInfo& DeviceCreateInfo) override final
 	{
-		if (bFullySupported)
+		if (bRequirementsPassed)
 		{
 			AddToPNext(DeviceCreateInfo, DescriptorIndexingFeatures);
 		}
@@ -471,7 +470,6 @@ public:
 
 private:
 	VkPhysicalDeviceDescriptorIndexingFeaturesEXT DescriptorIndexingFeatures;
-	bool bFullySupported = false;
 };
 
 
@@ -573,11 +571,12 @@ public:
 
 	virtual void PostPhysicalDeviceFeatures(FOptionalVulkanDeviceExtensions& ExtensionFlags) override final
 	{
-		ExtensionFlags.HasEXTFragmentDensityMap = 1;
+		bRequirementsPassed = (FragmentDensityMapFeatures.fragmentDensityMap == VK_TRUE);
+		ExtensionFlags.HasEXTFragmentDensityMap = bRequirementsPassed;
 
 		// Use the Fragment Density Map extension if and only if the Fragment Shading Rate extension is not available.
 		// NOTE: FVulkanKHRFragmentShadingRateExtension must be placed before FVulkanEXTFragmentDensityMapExtension for this to work!
-		if (!GRHISupportsAttachmentVariableRateShading && (FragmentDensityMapFeatures.fragmentDensityMap == VK_TRUE))
+		if (!GRHISupportsAttachmentVariableRateShading && bRequirementsPassed)
 		{
 			GRHISupportsAttachmentVariableRateShading = true;
 
@@ -605,7 +604,7 @@ public:
 
 	virtual void PreCreateDevice(VkDeviceCreateInfo& DeviceCreateInfo) override final
 	{
-		if (FragmentDensityMapFeatures.fragmentDensityMap == VK_TRUE)
+		if (bRequirementsPassed)
 		{
 			AddToPNext(DeviceCreateInfo, FragmentDensityMapFeatures);
 		}
@@ -637,14 +636,15 @@ public:
 
 	virtual void PostPhysicalDeviceFeatures(FOptionalVulkanDeviceExtensions& ExtensionFlags) override final
 	{
-		ExtensionFlags.HasEXTFragmentDensityMap2 = FragmentDensityMap2Features.fragmentDensityMapDeferred;
+		bRequirementsPassed = (FragmentDensityMap2Features.fragmentDensityMapDeferred == VK_TRUE);
+		ExtensionFlags.HasEXTFragmentDensityMap2 = bRequirementsPassed;
 
-		GRHISupportsLateVariableRateShadingUpdate = (FragmentDensityMap2Features.fragmentDensityMapDeferred == VK_TRUE);
+		GRHISupportsLateVariableRateShadingUpdate = bRequirementsPassed;
 	}
 
 	virtual void PreCreateDevice(VkDeviceCreateInfo& DeviceCreateInfo) override final
 	{
-		if (FragmentDensityMap2Features.fragmentDensityMapDeferred == VK_TRUE)
+		if (bRequirementsPassed)
 		{
 			AddToPNext(DeviceCreateInfo, FragmentDensityMap2Features);
 		}
@@ -684,12 +684,13 @@ public:
 
 	virtual void PostPhysicalDeviceFeatures(FOptionalVulkanDeviceExtensions& ExtensionFlags) override final 
 	{
-		ExtensionFlags.HasBufferDeviceAddress = BufferDeviceAddressFeatures.bufferDeviceAddress; 
+		bRequirementsPassed = (BufferDeviceAddressFeatures.bufferDeviceAddress == VK_TRUE);
+		ExtensionFlags.HasBufferDeviceAddress = bRequirementsPassed;
 	}
 
 	virtual void PreCreateDevice(VkDeviceCreateInfo& DeviceCreateInfo) override final
 	{
-		if (BufferDeviceAddressFeatures.bufferDeviceAddress == VK_TRUE)
+		if (bRequirementsPassed)
 		{
 			AddToPNext(DeviceCreateInfo, BufferDeviceAddressFeatures);
 		}
@@ -720,9 +721,10 @@ public:
 
 	virtual void PostPhysicalDeviceFeatures(FOptionalVulkanDeviceExtensions& ExtensionFlags) override final
 	{
-		ExtensionFlags.HasAccelerationStructure = 
-			AccelerationStructureFeatures.accelerationStructure & 
-			AccelerationStructureFeatures.descriptorBindingAccelerationStructureUpdateAfterBind;
+		bRequirementsPassed = (AccelerationStructureFeatures.accelerationStructure == VK_TRUE) && 
+			(AccelerationStructureFeatures.descriptorBindingAccelerationStructureUpdateAfterBind == VK_TRUE);
+
+		ExtensionFlags.HasAccelerationStructure = bRequirementsPassed;
 	}
 
 	virtual void PrePhysicalDeviceProperties(VkPhysicalDeviceProperties2KHR& PhysicalDeviceProperties2) override final
@@ -736,8 +738,7 @@ public:
 
 	virtual void PreCreateDevice(VkDeviceCreateInfo& DeviceCreateInfo) override final
 	{
-		if ((AccelerationStructureFeatures.accelerationStructure == VK_TRUE) &&
-			(AccelerationStructureFeatures.descriptorBindingAccelerationStructureUpdateAfterBind == VK_TRUE))
+		if (bRequirementsPassed)
 		{
 			AddToPNext(DeviceCreateInfo, AccelerationStructureFeatures);
 		}
@@ -768,7 +769,10 @@ public:
 
 	virtual void PostPhysicalDeviceFeatures(FOptionalVulkanDeviceExtensions& ExtensionFlags) override final
 	{
-		ExtensionFlags.HasRayTracingPipeline = RayTracingPipelineFeatures.rayTracingPipeline & RayTracingPipelineFeatures.rayTraversalPrimitiveCulling;
+		bRequirementsPassed = (RayTracingPipelineFeatures.rayTracingPipeline == VK_TRUE) &&
+			(RayTracingPipelineFeatures.rayTraversalPrimitiveCulling == VK_TRUE);
+
+		ExtensionFlags.HasRayTracingPipeline = bRequirementsPassed;
 	}
 
 	virtual void PrePhysicalDeviceProperties(VkPhysicalDeviceProperties2KHR& PhysicalDeviceProperties2) override final
@@ -782,8 +786,7 @@ public:
 
 	virtual void PreCreateDevice(VkDeviceCreateInfo& DeviceCreateInfo) override final
 	{
-		if ((RayTracingPipelineFeatures.rayTracingPipeline == VK_TRUE) &&
-			(RayTracingPipelineFeatures.rayTraversalPrimitiveCulling == VK_TRUE))
+		if (bRequirementsPassed)
 		{
 			AddToPNext(DeviceCreateInfo, RayTracingPipelineFeatures);
 		}
@@ -814,12 +817,13 @@ public:
 
 	virtual void PostPhysicalDeviceFeatures(FOptionalVulkanDeviceExtensions& ExtensionFlags) override final
 	{
-		ExtensionFlags.HasRayQuery = RayQueryFeatures.rayQuery; 
+		bRequirementsPassed = (RayQueryFeatures.rayQuery == VK_TRUE);
+		ExtensionFlags.HasRayQuery = bRequirementsPassed;
 	}
 
 	virtual void PreCreateDevice(VkDeviceCreateInfo& DeviceCreateInfo) override final
 	{
-		if (RayQueryFeatures.rayQuery == VK_TRUE)
+		if (bRequirementsPassed)
 		{
 			AddToPNext(DeviceCreateInfo, RayQueryFeatures);
 		}
@@ -914,12 +918,13 @@ public:
 
 	virtual void PostPhysicalDeviceFeatures(FOptionalVulkanDeviceExtensions& ExtensionFlags) override final 
 	{
-		ExtensionFlags.HasNVDeviceDiagnosticConfig = DeviceDiagnosticsConfigFeaturesNV.diagnosticsConfig; 
+		bRequirementsPassed = (DeviceDiagnosticsConfigFeaturesNV.diagnosticsConfig == VK_TRUE);
+		ExtensionFlags.HasNVDeviceDiagnosticConfig = bRequirementsPassed;
 	}
 
 	virtual void PreCreateDevice(VkDeviceCreateInfo& DeviceCreateInfo) override final
 	{
-		if (DeviceDiagnosticsConfigFeaturesNV.diagnosticsConfig == VK_TRUE)
+		if (bRequirementsPassed)
 		{
 			AddToPNext(DeviceCreateInfo, DeviceDiagnosticsConfigFeaturesNV);
 
@@ -954,12 +959,13 @@ public:
 
 	virtual void PostPhysicalDeviceFeatures(FOptionalVulkanDeviceExtensions& ExtensionFlags) override final
 	{
-		ExtensionFlags.HasEXTHostQueryReset = HostQueryResetFeatures.hostQueryReset;
+		bRequirementsPassed = (HostQueryResetFeatures.hostQueryReset == VK_TRUE);
+		ExtensionFlags.HasEXTHostQueryReset = bRequirementsPassed;
 	}
 
 	virtual void PreCreateDevice(VkDeviceCreateInfo& DeviceCreateInfo) override final
 	{
-		if (HostQueryResetFeatures.hostQueryReset == VK_TRUE)
+		if (bRequirementsPassed)
 		{
 			AddToPNext(DeviceCreateInfo, HostQueryResetFeatures);
 		}
@@ -987,13 +993,13 @@ public:
 
 	virtual void PostPhysicalDeviceFeatures(FOptionalVulkanDeviceExtensions& ExtensionFlags) override final
 	{
- 		bSupportsSubgroupSizeControl = (SubgroupSizeControlFeatures.subgroupSizeControl == VK_TRUE);
-		ExtensionFlags.HasEXTSubgroupSizeControl = bSupportsSubgroupSizeControl ? 1 : 0;
+		bRequirementsPassed = (SubgroupSizeControlFeatures.subgroupSizeControl == VK_TRUE);
+		ExtensionFlags.HasEXTSubgroupSizeControl = bRequirementsPassed;
 	}
 
 	virtual void PrePhysicalDeviceProperties(VkPhysicalDeviceProperties2KHR& PhysicalDeviceProperties2) override final
 	{
-		if (bSupportsSubgroupSizeControl)
+		if (bRequirementsPassed)
 		{
 			VkPhysicalDeviceSubgroupSizeControlPropertiesEXT& SubgroupSizeControlProperties = GetDeviceExtensionProperties().SubgroupSizeControlProperties;
 			ZeroVulkanStruct(SubgroupSizeControlProperties, VK_STRUCTURE_TYPE_PHYSICAL_DEVICE_SUBGROUP_SIZE_CONTROL_PROPERTIES);
@@ -1003,7 +1009,7 @@ public:
 
 	virtual void PreCreateDevice(VkDeviceCreateInfo& DeviceCreateInfo) override final
 	{
-		if (bSupportsSubgroupSizeControl)
+		if (bRequirementsPassed)
 		{
 			VkPhysicalDeviceSubgroupSizeControlPropertiesEXT& SubgroupSizeControlProperties = GetDeviceExtensionProperties().SubgroupSizeControlProperties;
 
@@ -1016,7 +1022,6 @@ public:
 
 private:
 	VkPhysicalDeviceSubgroupSizeControlFeaturesEXT SubgroupSizeControlFeatures;
-	bool bSupportsSubgroupSizeControl = false;
 };
 
 
@@ -1068,13 +1073,13 @@ public:
 	virtual void PostPhysicalDeviceFeatures(FOptionalVulkanDeviceExtensions& ExtensionFlags) override final
 	{
 		// only enable descriptor buffers if we also support mutable descriptor types (value filled prior)
-		bSupportsDescriptorBuffer = (DescriptorBufferFeatures.descriptorBuffer == VK_TRUE);
-		ExtensionFlags.HasEXTDescriptorBuffer = bSupportsDescriptorBuffer ? 1 : 0;
+		bRequirementsPassed = (DescriptorBufferFeatures.descriptorBuffer == VK_TRUE);
+		ExtensionFlags.HasEXTDescriptorBuffer = bRequirementsPassed;
 	}
 
 	virtual void PrePhysicalDeviceProperties(VkPhysicalDeviceProperties2KHR& PhysicalDeviceProperties2) override final
 	{
-		if (bSupportsDescriptorBuffer)
+		if (bRequirementsPassed)
 		{
 			VkPhysicalDeviceDescriptorBufferPropertiesEXT& DescriptorBufferProperties = GetDeviceExtensionProperties().DescriptorBufferProps;
 			ZeroVulkanStruct(DescriptorBufferProperties, VK_STRUCTURE_TYPE_PHYSICAL_DEVICE_DESCRIPTOR_BUFFER_PROPERTIES_EXT);
@@ -1084,7 +1089,7 @@ public:
 
 	virtual void PreCreateDevice(VkDeviceCreateInfo& DeviceCreateInfo) override final
 	{
-		if (bSupportsDescriptorBuffer)
+		if (bRequirementsPassed)
 		{
 			AddToPNext(DeviceCreateInfo, DescriptorBufferFeatures);
 
@@ -1108,7 +1113,6 @@ public:
 
 private:
 	VkPhysicalDeviceDescriptorBufferFeaturesEXT DescriptorBufferFeatures;
-	bool bSupportsDescriptorBuffer = false;
 };
 
 
@@ -1173,8 +1177,8 @@ FVulkanDeviceExtensionArray FVulkanDeviceExtension::GetUESupportedDeviceExtensio
 
 	ADD_CUSTOM_EXTENSION(FVulkanKHRDriverPropertiesExtension);
 	ADD_CUSTOM_EXTENSION(FVulkanKHRMaintenance4Extension);
-	ADD_CUSTOM_EXTENSION(FVulkanShaderAtomicInt64Extension);      // must be kept BEFORE ShaderImageAtomicInt64!
-	ADD_CUSTOM_EXTENSION(FVulkanShaderImageAtomicInt64Extension); // must be kept AFTER ShaderAtomicInt64!
+	ADD_CUSTOM_EXTENSION(FVulkanShaderAtomicInt64Extension);
+	ADD_CUSTOM_EXTENSION(FVulkanShaderImageAtomicInt64Extension);
 	ADD_CUSTOM_EXTENSION(FVulkanEXTScalarBlockLayoutExtension);
 	ADD_CUSTOM_EXTENSION(FVulkanEXTShaderViewportIndexLayerExtension);
 	ADD_CUSTOM_EXTENSION(FVulkanKHRSeparateDepthStencilLayoutsExtension);
@@ -1390,6 +1394,3 @@ FVulkanInstanceExtensionArray FVulkanInstanceExtension::GetUESupportedInstanceEx
 	return OutUEInstanceExtensions;
 }
 
-
-
-
diff --git a/Engine/Source/Runtime/VulkanRHI/Private/VulkanExtensions.h b/Engine/Source/Runtime/VulkanRHI/Private/VulkanExtensions.h
--- a/Engine/Source/Runtime/VulkanRHI/Private/VulkanExtensions.h
+++ b/Engine/Source/Runtime/VulkanRHI/Private/VulkanExtensions.h
@@ -88,6 +88,7 @@ public:
 		: FVulkanExtensionBase(InExtensionName, InEnabledInCode, InPromotedVersion, InActivation)
 		, Device(InDevice)
 		, FlagSetter(MoveTemp(InFlagSetter))
+		, bRequirementsPassed(true)
 	{}
 
 	static FVulkanDeviceExtensionArray GetUESupportedDeviceExtensions(FVulkanDevice* InDevice);
@@ -109,14 +110,17 @@ public:
 	virtual void PreCreateDevice(VkDeviceCreateInfo& DeviceInfo) {}
 #endif // VULKAN_SUPPORTS_PHYSICAL_DEVICE_PROPERTIES2
 
-	// Holds extensions requested external (eg plugins)
+	// Holds extensions requested externally (eg plugins)
 	static TArray<const ANSICHAR*> ExternalExtensions;
 
 protected:
 	struct FOptionalVulkanDeviceExtensionProperties& GetDeviceExtensionProperties();
 
 	FVulkanDevice* Device;
-	TUniqueFunction<void(FOptionalVulkanDeviceExtensions& ExtensionFlags)> FlagSetter = nullptr;
+	TUniqueFunction<void(FOptionalVulkanDeviceExtensions& ExtensionFlags)> FlagSetter;
+
+	// Signals that this extension has passed any specific feature and/or property requirements
+	bool bRequirementsPassed;
 };
 
 
@@ -144,7 +148,7 @@ public:
 	static TArray<const ANSICHAR*> ExternalExtensions;
 
 protected:
-	TUniqueFunction<void(FOptionalVulkanInstanceExtensions& ExtensionFlags)> FlagSetter = nullptr;
+	TUniqueFunction<void(FOptionalVulkanInstanceExtensions& ExtensionFlags)> FlagSetter;
 };
 
 // Helpers to create simple extensions that set a flag
