VulkanRHI: Move to UE::RHICore::SetResourcesFromTables to match other RHIs.

#rnx

diff --git a/Engine/Shaders/Public/Platform/Vulkan/VulkanCommon.ush b/Engine/Shaders/Public/Platform/Vulkan/VulkanCommon.ush
--- a/Engine/Shaders/Public/Platform/Vulkan/VulkanCommon.ush
+++ b/Engine/Shaders/Public/Platform/Vulkan/VulkanCommon.ush
@@ -5,7 +5,7 @@
 =============================================================================*/
 
 // Update this GUID to improve shader recompilation for Vulkan only shaders
-// GUID = 93BF8713-DDDD-46C1-A972-1C3534563925
+// GUID = 8dbc67f9-82f8-41dc-8bd9-648d781118e4
 
 #pragma once
 
diff --git a/Engine/Source/Developer/ShaderCompilerCommon/Private/ShaderCompilerCommon.cpp b/Engine/Source/Developer/ShaderCompilerCommon/Private/ShaderCompilerCommon.cpp
--- a/Engine/Source/Developer/ShaderCompilerCommon/Private/ShaderCompilerCommon.cpp
+++ b/Engine/Source/Developer/ShaderCompilerCommon/Private/ShaderCompilerCommon.cpp
@@ -88,9 +88,6 @@ bool BuildResourceTableMapping(
 
 	// Build resource table mapping
 	int32 MaxBoundResourceTable = -1;
-	TArray<uint32> ResourceTableSRVs;
-	TArray<uint32> ResourceTableSamplerStates;
-	TArray<uint32> ResourceTableUAVs;
 
 	// Go through ALL the members of ALL the UB resources
 	for( auto MapIt = ResourceTableMap.CreateConstIterator(); MapIt; ++MapIt )
diff --git a/Engine/Source/Developer/VulkanShaderFormat/Private/VulkanShaderCompiler.cpp b/Engine/Source/Developer/VulkanShaderFormat/Private/VulkanShaderCompiler.cpp
--- a/Engine/Source/Developer/VulkanShaderFormat/Private/VulkanShaderCompiler.cpp
+++ b/Engine/Source/Developer/VulkanShaderFormat/Private/VulkanShaderCompiler.cpp
@@ -590,103 +590,103 @@ static bool IsSamplerState(const CrossCompiler::FHlslccHeader& CCHeader, const F
 static void PrepareGlobals(const FVulkanBindingTable& BindingTable, const FVulkanHlslccHeader& CCHeader, const FShaderCompilerResourceTable& SRT, const TMap<FString, FVulkanShaderHeader::EType>& EntryTypes, const FShaderCompilerInput& ShaderInput, const TArray<FString>& ParameterNames, FShaderParameterMap& ParameterMap, TArray<FString>& OutGlobalNames, FVulkanShaderHeader& OutHeader)
 {
 	// First pass, gather names for all the Globals that are NOT Samplers
-	for (int32 ParameterIndex = 0; ParameterIndex < ParameterNames.Num(); ++ParameterIndex)
 	{
-		uint16 BufferIndex;
-		uint16 BaseIndex;
-		uint16 Size;
-		const FString& ParameterName = *ParameterNames[ParameterIndex];
-		ParameterMap.FindParameterAllocation(*ParameterName, BufferIndex, BaseIndex, Size);
-
-		auto AddGlobalNamesForUB = [&]()
+		auto AddGlobalNamesForUB = [&](const FString& ParameterName)
 		{
+			TOptional<FParameterAllocation> ParameterAllocation = ParameterMap.FindParameterAllocation(*ParameterName);
+			checkf(ParameterAllocation.IsSet(), TEXT("PrepareGlobals failed to find resource ParameterName=%s"), *ParameterName);
+
 			// Add used resources...
-			if (SRT.ResourceTableBits & (1 << BufferIndex))
+			if (SRT.ResourceTableBits & (1 << ParameterAllocation->BufferIndex))
 			{
-				PrepareUBResourceEntryGlobals(CCHeader, SRT.TextureMap, ShaderInput.Environment.ResourceTableMap, BufferIndex, ParameterName, OutGlobalNames, OutHeader);
-				PrepareUBResourceEntryGlobals(CCHeader, SRT.ShaderResourceViewMap, ShaderInput.Environment.ResourceTableMap, BufferIndex, ParameterName, OutGlobalNames, OutHeader);
-				PrepareUBResourceEntryGlobals(CCHeader, SRT.UnorderedAccessViewMap, ShaderInput.Environment.ResourceTableMap, BufferIndex, ParameterName, OutGlobalNames, OutHeader);
+				PrepareUBResourceEntryGlobals(CCHeader, SRT.TextureMap, ShaderInput.Environment.ResourceTableMap, ParameterAllocation->BufferIndex, ParameterName, OutGlobalNames, OutHeader);
+				PrepareUBResourceEntryGlobals(CCHeader, SRT.ShaderResourceViewMap, ShaderInput.Environment.ResourceTableMap, ParameterAllocation->BufferIndex, ParameterName, OutGlobalNames, OutHeader);
+				PrepareUBResourceEntryGlobals(CCHeader, SRT.UnorderedAccessViewMap, ShaderInput.Environment.ResourceTableMap, ParameterAllocation->BufferIndex, ParameterName, OutGlobalNames, OutHeader);
 			}
 		};
 
-		const FVulkanShaderHeader::EType* FoundType = EntryTypes.Find(ParameterName);
-		if (FoundType)
+		for (int32 ParameterIndex = 0; ParameterIndex < ParameterNames.Num(); ++ParameterIndex)
 		{
-			switch (*FoundType)
+			const FString& ParameterName = *ParameterNames[ParameterIndex];
+			const FVulkanShaderHeader::EType* FoundType = EntryTypes.Find(ParameterName);
+			if (FoundType)
 			{
-			case FVulkanShaderHeader::Global:
-				if (!IsSamplerState(CCHeader, ParameterName))
+				switch (*FoundType)
 				{
-					int32 GlobalIndex = DoAddGlobal(ParameterName, OutHeader, OutGlobalNames);
-					if (CCHeader.ExternalTextures.Contains(ParameterName))
+				case FVulkanShaderHeader::Global:
+					if (!IsSamplerState(CCHeader, ParameterName))
 					{
-						AddImmutable(OutHeader, GlobalIndex);
+						int32 GlobalIndex = DoAddGlobal(ParameterName, OutHeader, OutGlobalNames);
+						if (CCHeader.ExternalTextures.Contains(ParameterName))
+						{
+							AddImmutable(OutHeader, GlobalIndex);
+						}
 					}
+					break;
+				case FVulkanShaderHeader::UniformBuffer:
+					AddGlobalNamesForUB(ParameterName);
+					break;
+				case FVulkanShaderHeader::PackedGlobal:
+					// Ignore
+					break;
+				default:
+					check(0);
+					break;
 				}
-				break;
-			case FVulkanShaderHeader::UniformBuffer:
-				AddGlobalNamesForUB();
-				break;
-			case FVulkanShaderHeader::PackedGlobal:
-				// Ignore
-				break;
-			default:
-				check(0);
-				break;
 			}
-		}
-		else
-		{
-			AddGlobalNamesForUB();
+			else
+			{
+				AddGlobalNamesForUB(ParameterName);
+			}
 		}
 	}
 
 	// Second pass, add all samplers
-	for (int32 ParameterIndex = 0; ParameterIndex < ParameterNames.Num(); ++ParameterIndex)
 	{
-		uint16 BufferIndex;
-		uint16 BaseIndex;
-		uint16 Size;
-		const FString& ParameterName = *ParameterNames[ParameterIndex];
-		ParameterMap.FindParameterAllocation(*ParameterName, BufferIndex, BaseIndex, Size);
-
-		auto AddGlobalNamesForUB = [&]()
+		auto AddGlobalNamesForUB = [&](const FString& ParameterName)
 		{
+			TOptional<FParameterAllocation> ParameterAllocation = ParameterMap.FindParameterAllocation(*ParameterName);
+			checkf(ParameterAllocation.IsSet(), TEXT("PrepareGlobals failed to find sampler ParameterName=%s"), *ParameterName);
+
 			// Add used resources...
-			if (SRT.ResourceTableBits & (1 << BufferIndex))
+			if (SRT.ResourceTableBits & (1 << ParameterAllocation->BufferIndex))
 			{
-				PrepareUBResourceEntryGlobals(CCHeader, SRT.SamplerMap, ShaderInput.Environment.ResourceTableMap, BufferIndex, ParameterName, OutGlobalNames, OutHeader);
+				PrepareUBResourceEntryGlobals(CCHeader, SRT.SamplerMap, ShaderInput.Environment.ResourceTableMap, ParameterAllocation->BufferIndex, ParameterName, OutGlobalNames, OutHeader);
 			}
 		};
 
-		const FVulkanShaderHeader::EType* FoundType = EntryTypes.Find(ParameterName);
-		if (FoundType)
+		for (int32 ParameterIndex = 0; ParameterIndex < ParameterNames.Num(); ++ParameterIndex)
 		{
-			switch (*FoundType)
+			const FString& ParameterName = *ParameterNames[ParameterIndex];
+			const FVulkanShaderHeader::EType* FoundType = EntryTypes.Find(ParameterName);
+			if (FoundType)
 			{
-			case FVulkanShaderHeader::Global:
-				if (IsSamplerState(CCHeader, ParameterName))
+				switch (*FoundType)
 				{
-					int32 GlobalIndex = DoAddGlobal(ParameterName, OutHeader, OutGlobalNames);
-					if (CCHeader.ExternalTextures.Contains(ParameterName))
+				case FVulkanShaderHeader::Global:
+					if (IsSamplerState(CCHeader, ParameterName))
 					{
-						AddImmutable(OutHeader, GlobalIndex);
+						int32 GlobalIndex = DoAddGlobal(ParameterName, OutHeader, OutGlobalNames);
+						if (CCHeader.ExternalTextures.Contains(ParameterName))
+						{
+							AddImmutable(OutHeader, GlobalIndex);
+						}
 					}
+					break;
+				case FVulkanShaderHeader::UniformBuffer:
+					AddGlobalNamesForUB(ParameterName);
+					break;
+				case FVulkanShaderHeader::PackedGlobal:
+					break;
+				default:
+					check(0);
+					break;
 				}
-				break;
-			case FVulkanShaderHeader::UniformBuffer:
-				AddGlobalNamesForUB();
-				break;
-			case FVulkanShaderHeader::PackedGlobal:
-				break;
-			default:
-				check(0);
-				break;
 			}
-		}
-		else
-		{
-			AddGlobalNamesForUB();
+			else
+			{
+				AddGlobalNamesForUB(ParameterName);
+			}
 		}
 	}
 
@@ -886,10 +886,70 @@ static void ConvertToNEWHeader(FOLDVulkanCodeHeader& OLDHeader,
 	{
 		OutHeader.RayTracingPayloadSize = FCString::Atoi(**RTPayloadSizePtr);
 	}
+}
+
+
+// Fills the SRT using final values kept in the FVulkanShaderHeader.
+// NOTE: Uses GloalIndex so it can be consumed directly at runtime.
+// NOTE: Keep in sync with BuildResourceTableMapping.
+static FShaderResourceTable BuildSRTFromHeader(const FVulkanShaderHeader& NEWHeader)
+{
+	FShaderResourceTable ShaderResourceTable;
+
+	TArray<uint32> TextureMap;
+	TArray<uint32> ShaderResourceViewMap;
+	TArray<uint32> SamplerMap;
+	TArray<uint32> UnorderedAccessViewMap;
+
+	for (int32 UBIndex = 0; UBIndex < NEWHeader.UniformBuffers.Num(); ++UBIndex)
+	{
+		const FVulkanShaderHeader::FUniformBufferInfo& UBHeader = NEWHeader.UniformBuffers[UBIndex];
+
+		ShaderResourceTable.ResourceTableLayoutHashes.Emplace(UBHeader.LayoutHash);
+		if (UBHeader.ResourceEntries.Num() > 0)
+		{
+			ShaderResourceTable.ResourceTableBits |= 1 << UBIndex;
+
+			for (const FVulkanShaderHeader::FUBResourceInfo& UBRes : UBHeader.ResourceEntries)
+			{
+				uint32 ResourceMap = FRHIResourceTableEntry::Create(UBIndex, UBRes.SourceUBResourceIndex, UBRes.GlobalIndex);
+				switch (UBRes.UBBaseType)
+				{
+				case UBMT_TEXTURE:
+				case UBMT_RDG_TEXTURE:
+					TextureMap.Add(ResourceMap);
+					break;
+				case UBMT_SAMPLER:
+					SamplerMap.Add(ResourceMap);
+					break;
+				case UBMT_SRV:
+				case UBMT_RDG_TEXTURE_SRV:
+				case UBMT_RDG_BUFFER_SRV:
+					ShaderResourceViewMap.Add(ResourceMap);
+					break;
+				case UBMT_UAV:
+				case UBMT_RDG_TEXTURE_UAV:
+				case UBMT_RDG_BUFFER_UAV:
+					UnorderedAccessViewMap.Add(ResourceMap);
+					break;
+				default:
+					check(false);
+				}
+			}
+		}
+	}
+
+	const int32 MaxBoundResourceTable = NEWHeader.UniformBuffers.Num();
+	BuildResourceTableTokenStream(TextureMap, MaxBoundResourceTable, ShaderResourceTable.TextureMap);
+	BuildResourceTableTokenStream(ShaderResourceViewMap, MaxBoundResourceTable, ShaderResourceTable.ShaderResourceViewMap);
+	BuildResourceTableTokenStream(SamplerMap, MaxBoundResourceTable, ShaderResourceTable.SamplerMap);
+	BuildResourceTableTokenStream(UnorderedAccessViewMap, MaxBoundResourceTable, ShaderResourceTable.UnorderedAccessViewMap);
 
+	return ShaderResourceTable;
 }
 
 
+
 static void BuildShaderOutput(
 	FShaderCompilerOutput&		ShaderOutput,
 	const FShaderCompilerInput& ShaderInput,
@@ -1327,12 +1387,16 @@ static void BuildShaderOutput(
 		NEWHeader.DebugName = ShaderInput.GenerateShaderName();
 	}
 
+	// Build the SRT for this shader from the NEWHeader
+	FShaderResourceTable SerializedSRT = BuildSRTFromHeader(NEWHeader);
+
 	// Plug the passed in WaveSize
 	NEWHeader.WaveSize = WaveSize;
 
 	// Write out the header and shader source code.
 	FMemoryWriter Ar(ShaderOutput.ShaderCode.GetWriteAccess(), true);
 	Ar << NEWHeader;
+	Ar << SerializedSRT;
 
 	check(Spirv.Data.Num() != 0);
 	uint32 SpirvCodeSizeBytes = Spirv.Data.Num() * Spirv.Data.GetTypeSize();
diff --git a/Engine/Source/Runtime/VulkanRHI/Private/VulkanCommands.cpp b/Engine/Source/Runtime/VulkanRHI/Private/VulkanCommands.cpp
--- a/Engine/Source/Runtime/VulkanRHI/Private/VulkanCommands.cpp
+++ b/Engine/Source/Runtime/VulkanRHI/Private/VulkanCommands.cpp
@@ -13,6 +13,7 @@
 #include "GlobalRenderResources.h"
 #include "RHIShaderParametersShared.h"
 #include "RHIUtilities.h"
+#include "RHICoreShader.h"
 
 static TAutoConsoleVariable<int32> GCVarSubmitOnDispatch(
 	TEXT("r.Vulkan.SubmitOnDispatch"),
@@ -118,16 +119,157 @@ void FVulkanCommandListContext::RHISetStreamSource(uint32 StreamIndex, FRHIBuffe
 	}
 }
 
+template <class PendingStateType>
+struct FVulkanResourceBinder
+{
+	FVulkanCommandListContext& Context;
+	const EShaderFrequency Frequency;
+	const TArray<FDescriptorSetRemappingInfo::FRemappingInfo>& GlobalRemappingInfo;
+	PendingStateType* PendingState;
+
+	FVulkanResourceBinder(FVulkanCommandListContext& InContext, EShaderFrequency InFrequency, const TArray<FDescriptorSetRemappingInfo::FRemappingInfo>& InGlobalRemappingInfo, PendingStateType* InPendingState)
+		: Context(InContext)
+		, Frequency(InFrequency)
+		, GlobalRemappingInfo(InGlobalRemappingInfo)
+		, PendingState(InPendingState)
+	{
+	}
+
+	void SetUAV(FRHIUnorderedAccessView* UAV, uint8 Index, bool bClearResources = false)
+	{
+		if (bClearResources)
+		{
+			//Context.ClearShaderResources(D3D12UAV);
+		}
+
+		PendingState->SetUAVForUBResource(GlobalRemappingInfo[Index].NewDescriptorSet, GlobalRemappingInfo[Index].NewBindingIndex, ResourceCast(UAV));
+	}
+
+	void SetSRV(FRHIShaderResourceView* SRV, uint8 Index)
+	{
+		PendingState->SetSRVForUBResource(GlobalRemappingInfo[Index].NewDescriptorSet, GlobalRemappingInfo[Index].NewBindingIndex, ResourceCast(SRV));;
+	}
+
+	void SetTexture(FRHITexture* TextureRHI, uint8 Index)
+	{
+		FVulkanTexture* VulkanTexture = FVulkanTexture::Cast(TextureRHI);
+		const ERHIAccess RHIAccess = (Frequency == SF_Compute) ? ERHIAccess::SRVCompute : ERHIAccess::SRVGraphics;
+		const VkImageLayout ExpectedLayout = FVulkanLayoutManager::GetDefaultLayout(Context.GetCommandBufferManager()->GetActiveCmdBuffer(), *VulkanTexture, RHIAccess);
+		PendingState->SetTextureForUBResource(GlobalRemappingInfo[Index].NewDescriptorSet, GlobalRemappingInfo[Index].NewBindingIndex, VulkanTexture, ExpectedLayout);;
+	}
+
+	void SetSampler(FRHISamplerState* Sampler, uint8 Index)
+	{
+		PendingState->SetSamplerStateForUBResource(GlobalRemappingInfo[Index].NewDescriptorSet, GlobalRemappingInfo[Index].NewBindingIndex, ResourceCast(Sampler));
+	}
+};
+
+template <class ShaderType> 
+void FVulkanCommandListContext::SetResourcesFromTables(const ShaderType* Shader)
+{
+	checkSlow(Shader);
+
+	static constexpr EShaderFrequency Frequency = static_cast<EShaderFrequency>(ShaderType::StaticFrequency);
+
+	if (Frequency == SF_Compute)
+	{
+		const TArray<FDescriptorSetRemappingInfo::FRemappingInfo>& GlobalRemappingInfo =
+			PendingComputeState->CurrentState->GetComputePipelineDescriptorInfo().GetGlobalRemappingInfo();
+
+		FVulkanResourceBinder Binder(*this, Frequency, GlobalRemappingInfo, PendingComputeState);
+		UE::RHICore::SetResourcesFromTables(
+			Binder
+			, *Shader
+			, Shader->ShaderResourceTable
+			, DirtyUniformBuffers[Frequency]
+			, BoundUniformBuffers[Frequency]
+#if ENABLE_RHI_VALIDATION
+			, Tracker
+#endif
+		);
+	}
+	else
+	{
+		const TArray<FDescriptorSetRemappingInfo::FRemappingInfo>& GlobalRemappingInfo =
+		PendingGfxState->CurrentState->GetGfxPipelineDescriptorInfo().GetGlobalRemappingInfo(ShaderStage::GetStageForFrequency(Frequency));
+
+		FVulkanResourceBinder Binder(*this, Frequency, GlobalRemappingInfo, PendingGfxState);
+		UE::RHICore::SetResourcesFromTables(
+			Binder
+			, *Shader
+			, Shader->ShaderResourceTable
+			, DirtyUniformBuffers[Frequency]
+			, BoundUniformBuffers[Frequency]
+#if ENABLE_RHI_VALIDATION
+			, Tracker
+#endif
+		);
+	}
+}
+
+void FVulkanCommandListContext::CommitGraphicsResourceTables()
+{
+	checkSlow(PendingGfxState);
+
+	if (const FVulkanShader* Shader = PendingGfxState->GetCurrentShader(SF_Vertex))
+	{
+		checkSlow(Shader->Frequency == SF_Vertex);
+		const FVulkanVertexShader* VertexShader = static_cast<const FVulkanVertexShader*>(Shader);
+		SetResourcesFromTables(VertexShader);
+	}
+
+	if (const FVulkanShader* Shader = PendingGfxState->GetCurrentShader(SF_Pixel))
+	{
+		checkSlow(Shader->Frequency == SF_Pixel);
+		const FVulkanPixelShader* PixelShader = static_cast<const FVulkanPixelShader*>(Shader);
+		SetResourcesFromTables(PixelShader);
+	}
+
+#if PLATFORM_SUPPORTS_MESH_SHADERS
+	// :todo-jn: mesh shaders
+	//if (const FVulkanShader* Shader = PendingGfxState->GetCurrentShader(SF_Mesh))
+	//{
+	//  checkSlow(Shader->Frequency == SF_Mesh);
+	//	const FVulkanMeshShader* MeshShader = static_cast<const FVulkanMeshShader*>(Shader);
+	//	SetResourcesFromTables(MeshShader);
+	//}
+	//if (const FVulkanShader* Shader = PendingGfxState->GetCurrentShader(SF_Amplification))
+	//{
+	//  checkSlow(Shader->Frequency == SF_Amplification);
+	//	const FVulkanAmplificationShader* AmplificationShader = static_cast<const FVulkanAmplificationShader*>(Shader);
+	//	SetResourcesFromTables(AmplificationShader);
+	//}
+#endif
+
+#if PLATFORM_SUPPORTS_GEOMETRY_SHADERS
+	if (const FVulkanShader* Shader = PendingGfxState->GetCurrentShader(SF_Geometry))
+	{
+		checkSlow(Shader->Frequency == SF_Geometry);
+		const FVulkanGeometryShader* GeometryShader = static_cast<const FVulkanGeometryShader*>(Shader);
+		SetResourcesFromTables(GeometryShader);
+	}
+#endif
+}
+
+void FVulkanCommandListContext::CommitComputeResourceTables()
+{
+	SetResourcesFromTables(PendingComputeState->GetCurrentShader());
+}
+
+
 void FVulkanCommandListContext::RHIDispatchComputeShader(uint32 ThreadGroupCountX, uint32 ThreadGroupCountY, uint32 ThreadGroupCountZ)
 {
 #if VULKAN_ENABLE_AGGRESSIVE_STATS
 	SCOPE_CYCLE_COUNTER(STAT_VulkanDispatchCallTime);
 #endif
 
+	CommitComputeResourceTables();
+
 	FVulkanCmdBuffer* Cmd = CommandBufferManager->GetActiveCmdBuffer();
 	ensure(Cmd->IsOutsideRenderPass());
 	VkCommandBuffer CmdBuffer = Cmd->GetHandle();
 	PendingComputeState->PrepareForDispatch(Cmd);
+
 	VulkanRHI::vkCmdDispatch(CmdBuffer, ThreadGroupCountX, ThreadGroupCountY, ThreadGroupCountZ);
 
 	if (GCVarSubmitOnDispatch.GetValueOnRenderThread())
@@ -148,12 +290,13 @@ void FVulkanCommandListContext::RHIDispatchIndirectComputeShader(FRHIBuffer* Arg
 	static_assert(sizeof(FRHIDispatchIndirectParameters) == sizeof(VkDispatchIndirectCommand), "Dispatch indirect doesn't match!");
 	FVulkanResourceMultiBuffer* ArgumentBuffer = ResourceCast(ArgumentBufferRHI);
 
+	CommitComputeResourceTables();
+
 	FVulkanCmdBuffer* Cmd = CommandBufferManager->GetActiveCmdBuffer();
 	ensure(Cmd->IsOutsideRenderPass());
 	VkCommandBuffer CmdBuffer = Cmd->GetHandle();
 	PendingComputeState->PrepareForDispatch(Cmd);
 
-
 	VulkanRHI::vkCmdDispatchIndirect(CmdBuffer, ArgumentBuffer->GetHandle(), ArgumentBuffer->GetOffset() + ArgumentOffset);
 
 	if (GCVarSubmitOnDispatch.GetValueOnRenderThread())
@@ -173,9 +316,9 @@ void FVulkanCommandListContext::RHISetUAVParameter(FRHIPixelShader* PixelShaderR
 {
 	if (UAVRHI)
 	{
-	FVulkanUnorderedAccessView* UAV = ResourceCast(UAVRHI);
-	PendingGfxState->SetUAVForStage(ShaderStage::Pixel, UAVIndex, UAV);
-}
+		FVulkanUnorderedAccessView* UAV = ResourceCast(UAVRHI);
+		PendingGfxState->SetUAVForStage(ShaderStage::Pixel, UAVIndex, UAV);
+	}
 }
 
 void FVulkanCommandListContext::RHISetUAVParameter(FRHIComputeShader* ComputeShaderRHI, uint32 UAVIndex, FRHIUnorderedAccessView* UAVRHI)
@@ -184,9 +327,9 @@ void FVulkanCommandListContext::RHISetUAVParameter(FRHIComputeShader* ComputeSha
 	{
 		check(PendingComputeState->GetCurrentShader() == ResourceCast(ComputeShaderRHI));
 
-	FVulkanUnorderedAccessView* UAV = ResourceCast(UAVRHI);
-	PendingComputeState->SetUAVForStage(UAVIndex, UAV);
-}
+		FVulkanUnorderedAccessView* UAV = ResourceCast(UAVRHI);
+		PendingComputeState->SetUAVForStage(UAVIndex, UAV);
+	}
 }
 
 void FVulkanCommandListContext::RHISetUAVParameter(FRHIComputeShader* ComputeShaderRHI,uint32 UAVIndex, FRHIUnorderedAccessView* UAVRHI, uint32 InitialCount)
@@ -220,10 +363,10 @@ void FVulkanCommandListContext::RHISetShaderResourceViewParameter(FRHIGraphicsSh
 {
 	if (SRVRHI)
 	{
-	ShaderStage::EStage Stage = GetAndVerifyShaderStage(ShaderRHI, PendingGfxState);
-	FVulkanShaderResourceView* SRV = ResourceCast(SRVRHI);
-	PendingGfxState->SetSRVForStage(Stage, TextureIndex, SRV);
-}
+		ShaderStage::EStage Stage = GetAndVerifyShaderStage(ShaderRHI, PendingGfxState);
+		FVulkanShaderResourceView* SRV = ResourceCast(SRVRHI);
+		PendingGfxState->SetSRVForStage(Stage, TextureIndex, SRV);
+	}
 }
 
 void FVulkanCommandListContext::RHISetShaderResourceViewParameter(FRHIComputeShader* ComputeShaderRHI,uint32 TextureIndex, FRHIShaderResourceView* SRVRHI)
@@ -232,9 +375,9 @@ void FVulkanCommandListContext::RHISetShaderResourceViewParameter(FRHIComputeSha
 	{
 		check(PendingComputeState->GetCurrentShader() == ResourceCast(ComputeShaderRHI));
 
-	FVulkanShaderResourceView* SRV = ResourceCast(SRVRHI);
-	PendingComputeState->SetSRVForStage(TextureIndex, SRV);
-}
+		FVulkanShaderResourceView* SRV = ResourceCast(SRVRHI);
+		PendingComputeState->SetSRVForStage(TextureIndex, SRV);
+	}
 }
 
 void FVulkanCommandListContext::RHISetShaderSampler(FRHIGraphicsShader* ShaderRHI, uint32 SamplerIndex, FRHISamplerState* NewStateRHI)
@@ -325,154 +468,32 @@ void FVulkanCommandListContext::RHISetShaderParameters(FRHIComputeShader* Shader
 	);
 }
 
-template <typename TState>
-inline void SetShaderUniformBufferResources(FVulkanCommandListContext* Context, TState* State, const FVulkanShader* Shader, const TArray<FVulkanShaderHeader::FGlobalInfo>& GlobalInfos, const TArray<TEnumAsByte<EVulkanBindingType::EType>>& DescriptorTypes, const FVulkanShaderHeader::FUniformBufferInfo& HeaderUBInfo, const FVulkanUniformBuffer* UniformBuffer, const TArray<FDescriptorSetRemappingInfo::FRemappingInfo>& GlobalRemappingInfo)
+void FVulkanCommandListContext::RHISetStaticUniformBuffers(const FUniformBufferStaticBindings& InUniformBuffers)
 {
-#if ENABLE_RHI_VALIDATION
-	static_assert(std::is_same_v<TState, FVulkanPendingGfxState> || std::is_same_v<TState, FVulkanPendingComputeState>, "TState must be FVulkanPendingGfxState or FVulkanPendingComputeState");
-	constexpr bool bIsGfx = std::is_same_v<TState, FVulkanPendingGfxState>;
-	constexpr ERHIAccess SRVAccess = bIsGfx ? ERHIAccess::SRVGraphics : ERHIAccess::SRVCompute;
-	constexpr ERHIAccess UAVAccess = bIsGfx ? ERHIAccess::UAVGraphics : ERHIAccess::UAVCompute;
-#endif
+	FMemory::Memzero(GlobalUniformBuffers.GetData(), GlobalUniformBuffers.Num() * sizeof(FRHIUniformBuffer*));
 
-	ensure(UniformBuffer->GetLayout().GetHash() == HeaderUBInfo.LayoutHash);
-	float CurrentTime = (float)FPlatformTime::Seconds();
-	const TArray<TRefCountPtr<FRHIResource>>& ResourceArray = UniformBuffer->GetResourceTable();
-	for (int32 Index = 0; Index < HeaderUBInfo.ResourceEntries.Num(); ++Index)
+	for (int32 Index = 0; Index < InUniformBuffers.GetUniformBufferCount(); ++Index)
 	{
-		const FVulkanShaderHeader::FUBResourceInfo& ResourceInfo = HeaderUBInfo.ResourceEntries[Index];
-		switch (ResourceInfo.UBBaseType)
-		{
-		case UBMT_SAMPLER:
-		{
-			uint16 CombinedAlias = GlobalInfos[ResourceInfo.GlobalIndex].CombinedSamplerStateAliasIndex;
-			uint32 GlobalIndex = CombinedAlias == UINT16_MAX ? ResourceInfo.GlobalIndex : CombinedAlias;
-			const VkDescriptorType DescriptorType = BindingToDescriptorType(DescriptorTypes[GlobalInfos[GlobalIndex].TypeIndex]);
-			ensure(DescriptorType == VK_DESCRIPTOR_TYPE_SAMPLER || DescriptorType == VK_DESCRIPTOR_TYPE_COMBINED_IMAGE_SAMPLER);
-			FVulkanSamplerState* CurrSampler = static_cast<FVulkanSamplerState*>(ResourceArray[ResourceInfo.SourceUBResourceIndex].GetReference());
-			if (CurrSampler)
-			{
-				if (CurrSampler->Sampler)
-				{
-					State->SetSamplerStateForUBResource(GlobalRemappingInfo[GlobalIndex].NewDescriptorSet, GlobalRemappingInfo[GlobalIndex].NewBindingIndex, CurrSampler);
-				}
-			}
-			else
-			{
-				UE_LOG(LogVulkanRHI, Warning, TEXT("Invalid sampler in SRT table for shader '%s'"), *Shader->GetDebugName());
-			}
-			break;
-		}
-
-		case UBMT_TEXTURE:
-		case UBMT_RDG_TEXTURE:
-		{
-			const VkDescriptorType DescriptorType = BindingToDescriptorType(DescriptorTypes[GlobalInfos[ResourceInfo.GlobalIndex].TypeIndex]);
-			// Tolerate STORAGE_IMAGE for now, it is sometimes used on formats that don't support sampling
-			ensure(DescriptorType == VK_DESCRIPTOR_TYPE_SAMPLED_IMAGE || DescriptorType == VK_DESCRIPTOR_TYPE_COMBINED_IMAGE_SAMPLER || DescriptorType == VK_DESCRIPTOR_TYPE_STORAGE_IMAGE);
-			FRHITexture* TexRef = (FRHITexture*)(ResourceArray[ResourceInfo.SourceUBResourceIndex].GetReference());
-			if (TexRef)
-			{
-				const FVulkanTexture* VulkanTexture = FVulkanTexture::Cast(TexRef);
-				if (!ensure(VulkanTexture))
-				{
-					VulkanTexture = FVulkanTexture::Cast(GBlackTexture->TextureRHI.GetReference());
-				}
-
-				// If the descriptor is a storage image in a slot expecting to read only, make sure it's because we don't support sampling
-				ensure(DescriptorType != VK_DESCRIPTOR_TYPE_STORAGE_IMAGE || !VulkanTexture->SupportsSampling());
-
-#if ENABLE_RHI_VALIDATION
-				if (Context->Tracker)
-				{
-					Context->Tracker->Assert(TexRef->GetViewIdentity(0, 0, 0, 0, uint32(RHIValidation::EResourcePlane::Common), 1), SRVAccess);
-				}
-#endif
-
-				const VkImageLayout ExpectedLayout = FVulkanLayoutManager::GetDefaultLayout(Context->GetCommandBufferManager()->GetActiveCmdBuffer(), *VulkanTexture, ERHIAccess::SRVMask);
-				State->SetTextureForUBResource(GlobalRemappingInfo[ResourceInfo.GlobalIndex].NewDescriptorSet, GlobalRemappingInfo[ResourceInfo.GlobalIndex].NewBindingIndex, VulkanTexture, ExpectedLayout);
-				TexRef->SetLastRenderTime(CurrentTime);
-			}
-			else
-			{
-				UE_LOG(LogVulkanRHI, Warning, TEXT("Invalid texture in SRT table for shader '%s'"), *Shader->GetDebugName());
-			}
-			break;
-		}
-
-		case UBMT_RDG_TEXTURE_SRV:
-		case UBMT_SRV:
-		case UBMT_RDG_BUFFER_SRV:
-		{
-			const VkDescriptorType DescriptorType = BindingToDescriptorType(DescriptorTypes[GlobalInfos[ResourceInfo.GlobalIndex].TypeIndex]);
-			ensure(DescriptorType == VK_DESCRIPTOR_TYPE_UNIFORM_TEXEL_BUFFER 
-				|| DescriptorType == VK_DESCRIPTOR_TYPE_COMBINED_IMAGE_SAMPLER
-				|| DescriptorType == VK_DESCRIPTOR_TYPE_SAMPLED_IMAGE
-				|| DescriptorType == VK_DESCRIPTOR_TYPE_STORAGE_BUFFER);
-			FRHIShaderResourceView* CurrentSRV = (FRHIShaderResourceView*)(ResourceArray[ResourceInfo.SourceUBResourceIndex].GetReference());
-			if (CurrentSRV)
-			{
-#if ENABLE_RHI_VALIDATION
-				if (Context->Tracker)
-				{
-					Context->Tracker->Assert(CurrentSRV->ViewIdentity, SRVAccess);
-				}
-#endif
-				FVulkanShaderResourceView* SRV = ResourceCast(CurrentSRV);
-				State->SetSRVForUBResource(GlobalRemappingInfo[ResourceInfo.GlobalIndex].NewDescriptorSet, GlobalRemappingInfo[ResourceInfo.GlobalIndex].NewBindingIndex, SRV);
-			}
-			else
-			{
-				UE_LOG(LogVulkanRHI, Warning, TEXT("Invalid texture in SRT table for shader '%s'"), *Shader->GetDebugName());
-			}
-			break;
-		}
-
-		case UBMT_RDG_TEXTURE_UAV:
-		case UBMT_UAV:
-		case UBMT_RDG_BUFFER_UAV:
-		{
-			const VkDescriptorType DescriptorType = BindingToDescriptorType(DescriptorTypes[GlobalInfos[ResourceInfo.GlobalIndex].TypeIndex]);
-			ensure(DescriptorType == VK_DESCRIPTOR_TYPE_STORAGE_BUFFER
-				|| DescriptorType == VK_DESCRIPTOR_TYPE_STORAGE_IMAGE
-				|| DescriptorType == VK_DESCRIPTOR_TYPE_STORAGE_TEXEL_BUFFER);
-			FRHIUnorderedAccessView* CurrentUAV = (FRHIUnorderedAccessView*)(ResourceArray[ResourceInfo.SourceUBResourceIndex].GetReference());
-			if (CurrentUAV)
-			{
-#if ENABLE_RHI_VALIDATION
-				if (Context->Tracker)
-				{
-					Context->Tracker->Assert(CurrentUAV->ViewIdentity, UAVAccess);
-				}
-#endif
-				FVulkanUnorderedAccessView* UAV = ResourceCast(CurrentUAV);
-				State->SetUAVForUBResource(GlobalRemappingInfo[ResourceInfo.GlobalIndex].NewDescriptorSet, GlobalRemappingInfo[ResourceInfo.GlobalIndex].NewBindingIndex, UAV);
-			}
-			else
-			{
-				UE_LOG(LogVulkanRHI, Warning, TEXT("Invalid texture in SRT table for shader '%s'"), *Shader->GetDebugName());
-			}
-			break;
-		}
-
-		default:
-			checkf(0, TEXT("Missing handling for UBMT_ %d"), (int32)ResourceInfo.UBBaseType);
-			break;
-		}
+		GlobalUniformBuffers[InUniformBuffers.GetSlot(Index)] = InUniformBuffers.GetUniformBuffer(Index);
 	}
 }
 
-inline void FVulkanCommandListContext::SetShaderUniformBuffer(ShaderStage::EStage Stage, const FVulkanUniformBuffer* UniformBuffer, int32 BufferIndex, const FVulkanShader* Shader)
+void FVulkanCommandListContext::RHISetShaderUniformBuffer(FRHIGraphicsShader* ShaderRHI, uint32 BufferIndex, FRHIUniformBuffer* BufferRHI)
 {
 #if VULKAN_ENABLE_AGGRESSIVE_STATS
 	SCOPE_CYCLE_COUNTER(STAT_VulkanSetUniformBufferTime);
 #endif
+
+	FVulkanShader* Shader = nullptr;
+	const ShaderStage::EStage Stage = GetAndVerifyShaderStageAndVulkanShader(ShaderRHI, PendingGfxState, Shader);
 	check(Shader->GetShaderKey() == PendingGfxState->GetCurrentShaderKey(Stage));
 
+	FVulkanUniformBuffer* UniformBuffer = ResourceCast(BufferRHI);
 	const FVulkanShaderHeader& CodeHeader = Shader->GetCodeHeader();
 	const FVulkanShaderHeader::FUniformBufferInfo& HeaderUBInfo = CodeHeader.UniformBuffers[BufferIndex];
 	checkfSlow(!HeaderUBInfo.LayoutHash || HeaderUBInfo.LayoutHash == UniformBuffer->GetLayout().GetHash(), TEXT("Mismatched UB layout! Got hash 0x%x, expected 0x%x!"), UniformBuffer->GetLayout().GetHash(), HeaderUBInfo.LayoutHash);
 	const FVulkanGfxPipelineDescriptorInfo& DescriptorInfo = PendingGfxState->CurrentState->GetGfxPipelineDescriptorInfo();
+
 	if (!HeaderUBInfo.bOnlyHasResources)
 	{
 		checkSlow(UniformBuffer->GetLayout().ConstantBufferSize > 0);
@@ -499,7 +520,10 @@ inline void FVulkanCommandListContext::SetShaderUniformBuffer(ShaderStage::EStag
 
 	if (HeaderUBInfo.ResourceEntries.Num())
 	{
-		SetShaderUniformBufferResources(this, PendingGfxState, Shader, CodeHeader.Globals, CodeHeader.GlobalDescriptorTypes, HeaderUBInfo, UniformBuffer, DescriptorInfo.GetGlobalRemappingInfo(Stage));
+		checkSlow(Shader->Frequency < SF_NumStandardFrequencies);
+		check(BufferIndex < MAX_UNIFORM_BUFFERS_PER_SHADER_STAGE);
+		BoundUniformBuffers[Shader->Frequency][BufferIndex] = UniformBuffer;
+		DirtyUniformBuffers[Shader->Frequency] |= (1 << BufferIndex);
 	}
 	else
 	{
@@ -508,24 +532,6 @@ inline void FVulkanCommandListContext::SetShaderUniformBuffer(ShaderStage::EStag
 	}
 }
 
-void FVulkanCommandListContext::RHISetStaticUniformBuffers(const FUniformBufferStaticBindings& InUniformBuffers)
-{
-	FMemory::Memzero(GlobalUniformBuffers.GetData(), GlobalUniformBuffers.Num() * sizeof(FRHIUniformBuffer*));
-
-	for (int32 Index = 0; Index < InUniformBuffers.GetUniformBufferCount(); ++Index)
-	{
-		GlobalUniformBuffers[InUniformBuffers.GetSlot(Index)] = InUniformBuffers.GetUniformBuffer(Index);
-	}
-}
-
-void FVulkanCommandListContext::RHISetShaderUniformBuffer(FRHIGraphicsShader* ShaderRHI, uint32 BufferIndex, FRHIUniformBuffer* BufferRHI)
-{
-	FVulkanShader* Shader = nullptr;
-	ShaderStage::EStage Stage = GetAndVerifyShaderStageAndVulkanShader(ShaderRHI, PendingGfxState, Shader);
-	FVulkanUniformBuffer* UniformBuffer = ResourceCast(BufferRHI);
-	SetShaderUniformBuffer(Stage, UniformBuffer, BufferIndex, Shader);
-}
-
 void FVulkanCommandListContext::RHISetShaderUniformBuffer(FRHIComputeShader* ComputeShaderRHI, uint32 BufferIndex, FRHIUniformBuffer* BufferRHI)
 {
 	FVulkanComputeShader* ComputeShader = ResourceCast(ComputeShaderRHI);
@@ -572,7 +578,10 @@ void FVulkanCommandListContext::RHISetShaderUniformBuffer(FRHIComputeShader* Com
 
 	if (HeaderUBInfo.ResourceEntries.Num())
 	{
-		SetShaderUniformBufferResources(this, PendingComputeState, Shader, Shader->CodeHeader.Globals, Shader->CodeHeader.GlobalDescriptorTypes, HeaderUBInfo, UniformBuffer, DescriptorInfo.GetGlobalRemappingInfo());
+		checkSlow(ComputeShaderRHI->GetFrequency() == SF_Compute);
+		check(BufferIndex < MAX_UNIFORM_BUFFERS_PER_SHADER_STAGE);
+		BoundUniformBuffers[SF_Compute][BufferIndex] = UniformBuffer;
+		DirtyUniformBuffers[SF_Compute] |= (1 << BufferIndex);
 	}
 	else
 	{
@@ -595,6 +604,8 @@ void FVulkanCommandListContext::RHIDrawPrimitive(uint32 BaseVertexIndex, uint32
 
 	RHI_DRAW_CALL_STATS(PendingGfxState->PrimitiveType, NumInstances*NumPrimitives);
 
+	CommitGraphicsResourceTables();
+
 	FVulkanCmdBuffer* CmdBuffer = CommandBufferManager->GetActiveCmdBuffer();
 	PendingGfxState->PrepareForDraw(CmdBuffer);
 	uint32 NumVertices = GetVertexCountForPrimitiveCount(NumPrimitives, PendingGfxState->PrimitiveType);
@@ -615,6 +626,8 @@ void FVulkanCommandListContext::RHIDrawPrimitiveIndirect(FRHIBuffer* ArgumentBuf
 #endif
 	RHI_DRAW_CALL_INC();
 
+	CommitGraphicsResourceTables();
+
 	FVulkanCmdBuffer* Cmd = CommandBufferManager->GetActiveCmdBuffer();
 	VkCommandBuffer CmdBuffer = Cmd->GetHandle();
 	PendingGfxState->PrepareForDraw(Cmd);
@@ -640,6 +653,8 @@ void FVulkanCommandListContext::RHIDrawIndexedPrimitive(FRHIBuffer* IndexBufferR
 	RHI_DRAW_CALL_STATS(PendingGfxState->PrimitiveType, NumInstances*NumPrimitives);
 	checkf(GRHISupportsFirstInstance || FirstInstance == 0, TEXT("FirstInstance must be 0, see GRHISupportsFirstInstance"));
 
+	CommitGraphicsResourceTables();
+
 	FVulkanResourceMultiBuffer* IndexBuffer = ResourceCast(IndexBufferRHI);
 	FVulkanCmdBuffer* Cmd = CommandBufferManager->GetActiveCmdBuffer();
 	VkCommandBuffer CmdBuffer = Cmd->GetHandle();
@@ -662,6 +677,8 @@ void FVulkanCommandListContext::RHIDrawIndexedIndirect(FRHIBuffer* IndexBufferRH
 #endif
 	RHI_DRAW_CALL_INC();
 
+	CommitGraphicsResourceTables();
+
 	FVulkanResourceMultiBuffer* IndexBuffer = ResourceCast(IndexBufferRHI);
 	FVulkanCmdBuffer* Cmd = CommandBufferManager->GetActiveCmdBuffer();
 	VkCommandBuffer CmdBuffer = Cmd->GetHandle();
@@ -687,6 +704,8 @@ void FVulkanCommandListContext::RHIDrawIndexedPrimitiveIndirect(FRHIBuffer* Inde
 #endif
 	RHI_DRAW_CALL_INC();
 
+	CommitGraphicsResourceTables();
+
 	FVulkanResourceMultiBuffer* IndexBuffer = ResourceCast(IndexBufferRHI);
 	FVulkanCmdBuffer* Cmd = CommandBufferManager->GetActiveCmdBuffer();
 	VkCommandBuffer CmdBuffer = Cmd->GetHandle();
diff --git a/Engine/Source/Runtime/VulkanRHI/Private/VulkanContext.h b/Engine/Source/Runtime/VulkanRHI/Private/VulkanContext.h
--- a/Engine/Source/Runtime/VulkanRHI/Private/VulkanContext.h
+++ b/Engine/Source/Runtime/VulkanRHI/Private/VulkanContext.h
@@ -35,6 +35,10 @@ public:
 		return Immediate == nullptr;
 	}
 
+	template <class ShaderType> void SetResourcesFromTables(const ShaderType* RESTRICT);
+	void CommitGraphicsResourceTables();
+	void CommitComputeResourceTables();
+
 	virtual void RHISetStreamSource(uint32 StreamIndex, FRHIBuffer* VertexBuffer, uint32 Offset) final override;
 	virtual void RHISetViewport(float MinX, float MinY, float MinZ, float MaxX, float MaxY, float MaxZ) final override;
 	virtual void RHISetStereoViewport(float LeftMinX, float RightMinX, float LeftMinY, float RightMinY, float MinZ, float LeftMaxX, float RightMaxX, float LeftMaxY, float RightMaxY, float MaxZ) override;
@@ -198,8 +202,6 @@ protected:
 	void BeginOcclusionQueryBatch(FVulkanCmdBuffer* CmdBuffer, uint32 NumQueriesInBatch);
 	void EndOcclusionQueryBatch(FVulkanCmdBuffer* CmdBuffer);
 
-	void SetShaderUniformBuffer(ShaderStage::EStage Stage, const FVulkanUniformBuffer* UniformBuffer, int32 ParameterIndex, const FVulkanShader* Shader);
-
 	VulkanRHI::FTempFrameAllocationBuffer TempFrameAllocationBuffer;
 
 	TArray<FString> EventStack;
@@ -214,6 +216,15 @@ protected:
 	FVulkanPendingGfxState* PendingGfxState;
 	FVulkanPendingComputeState* PendingComputeState;
 
+	// Match the D3D12 maximum of 16 constant buffers per shader stage.
+	enum { MAX_UNIFORM_BUFFERS_PER_SHADER_STAGE = 16 };
+
+	// Track the currently bound uniform buffers.
+	FVulkanUniformBuffer* BoundUniformBuffers[SF_NumStandardFrequencies][MAX_UNIFORM_BUFFERS_PER_SHADER_STAGE] = {};
+
+	// Bit array to track which uniform buffers have changed since the last draw call.
+	uint16 DirtyUniformBuffers[SF_NumStandardFrequencies] = {};
+
 	void PrepareForCPURead();
 	void RequestSubmitCurrentCommands();
 
diff --git a/Engine/Source/Runtime/VulkanRHI/Private/VulkanPendingState.h b/Engine/Source/Runtime/VulkanRHI/Private/VulkanPendingState.h
--- a/Engine/Source/Runtime/VulkanRHI/Private/VulkanPendingState.h
+++ b/Engine/Source/Runtime/VulkanRHI/Private/VulkanPendingState.h
@@ -194,6 +194,11 @@ public:
 		return GetCurrentShaderKey(ShaderStage::GetFrequencyForGfxStage(Stage));
 	}
 
+	const FVulkanShader* GetCurrentShader(EShaderFrequency Frequency) const
+	{
+		return (CurrentPipeline ? CurrentPipeline->GetShader(Frequency) : nullptr);
+	}
+
 	void SetViewport(float MinX, float MinY, float MinZ, float MaxX, float MaxY, float MaxZ)
 	{
 		Viewports.SetNumZeroed(1);
diff --git a/Engine/Source/Runtime/VulkanRHI/Private/VulkanPipeline.h b/Engine/Source/Runtime/VulkanRHI/Private/VulkanPipeline.h
--- a/Engine/Source/Runtime/VulkanRHI/Private/VulkanPipeline.h
+++ b/Engine/Source/Runtime/VulkanRHI/Private/VulkanPipeline.h
@@ -737,6 +737,12 @@ public:
 		return ShaderKeys[Stage];
 	}
 
+	inline const FVulkanShader* GetShader(EShaderFrequency Frequency) const
+	{
+		ShaderStage::EStage Stage = ShaderStage::GetStageForFrequency(Frequency);
+		return VulkanShaders[Stage];
+	}
+
 	inline VkPipeline GetVulkanPipeline() const
 	{
 		return VulkanPipeline;
diff --git a/Engine/Source/Runtime/VulkanRHI/Private/VulkanResources.h b/Engine/Source/Runtime/VulkanRHI/Private/VulkanResources.h
--- a/Engine/Source/Runtime/VulkanRHI/Private/VulkanResources.h
+++ b/Engine/Source/Runtime/VulkanRHI/Private/VulkanResources.h
@@ -191,6 +191,8 @@ protected:
 
 	TArray<FUniformBufferStaticSlot> StaticSlots;
 
+	FShaderResourceTable			ShaderResourceTable;
+
 private:
 	class FSpirvContainer
 	{
@@ -207,7 +209,7 @@ private:
 	static FSpirvCode PatchSpirvInputAttachments(FSpirvCode& SpirvCode);
 
 protected:
-	void Setup(FVulkanShaderHeader&& InCodeHeader, FSpirvContainer&& InSpirvContainer, uint64 InShaderKey);
+	void Setup(FVulkanShaderHeader&& InCodeHeader, FShaderResourceTable&& InSRT, FSpirvContainer&& InSpirvContainer, uint64 InShaderKey);
 
 	FVulkanDevice*					Device;
 
diff --git a/Engine/Source/Runtime/VulkanRHI/Private/VulkanShaders.cpp b/Engine/Source/Runtime/VulkanRHI/Private/VulkanShaders.cpp
--- a/Engine/Source/Runtime/VulkanRHI/Private/VulkanShaders.cpp
+++ b/Engine/Source/Runtime/VulkanRHI/Private/VulkanShaders.cpp
@@ -64,6 +64,8 @@ ShaderType* FVulkanShaderFactory::CreateShader(TArrayView<const uint8> Code, FVu
 		FMemoryReaderView Ar(Code, true);
 		FVulkanShaderHeader CodeHeader;
 		Ar << CodeHeader;
+		FShaderResourceTable SerializedSRT;
+		Ar << SerializedSRT;
 		FVulkanShader::FSpirvContainer SpirvContainer;
 		Ar << SpirvContainer;
 
@@ -77,7 +79,7 @@ ShaderType* FVulkanShaderFactory::CreateShader(TArrayView<const uint8> Code, FVu
 			else
 			{
 				RetShader = new ShaderType(Device);
-				RetShader->Setup(MoveTemp(CodeHeader), MoveTemp(SpirvContainer), ShaderKey);
+				RetShader->Setup(MoveTemp(CodeHeader), MoveTemp(SerializedSRT), MoveTemp(SpirvContainer), ShaderKey);
 				if constexpr (
 					ShaderType::StaticFrequency == SF_RayCallable ||
 					ShaderType::StaticFrequency == SF_RayGen ||
@@ -192,7 +194,7 @@ FVulkanShader::FSpirvCode FVulkanShader::GetSpirvCode()
 }
 
 
-void FVulkanShader::Setup(FVulkanShaderHeader&& InCodeHeader, FSpirvContainer&& InSpirvContainer, uint64 InShaderKey)
+void FVulkanShader::Setup(FVulkanShaderHeader&& InCodeHeader, FShaderResourceTable&& InSRT, FSpirvContainer&& InSpirvContainer, uint64 InShaderKey)
 {
 	LLM_SCOPE_VULKAN(ELLMTagVulkan::VulkanShaders);
 	check(Device);
@@ -201,6 +203,8 @@ void FVulkanShader::Setup(FVulkanShaderHeader&& InCodeHeader, FSpirvContainer&&
 
 	CodeHeader = MoveTemp(InCodeHeader);
 
+	ShaderResourceTable = MoveTemp(InSRT);
+
 	SpirvContainer = MoveTemp(InSpirvContainer);
 
 	checkf(SpirvContainer.GetSizeBytes() != 0, TEXT("Empty SPIR-V! %s"), *CodeHeader.DebugName);
