VulkanRHI: Move Android (and engine) minimum core version to Vulkan 1.1.  Remove need to check for VK_KHR_get_physical_device_properties2.

diff --git a/Engine/Source/Runtime/Core/Private/Android/AndroidPlatformMisc.cpp b/Engine/Source/Runtime/Core/Private/Android/AndroidPlatformMisc.cpp
--- a/Engine/Source/Runtime/Core/Private/Android/AndroidPlatformMisc.cpp
+++ b/Engine/Source/Runtime/Core/Private/Android/AndroidPlatformMisc.cpp
@@ -2114,7 +2114,7 @@ typedef VkResult(VKAPI_PTR *PFN_vkEnumerateDeviceExtensionProperties)(VkPhysical
 
 ///////////////////////////////////////////////////////////////////////////////
 
-#define UE_VK_API_VERSION	VK_MAKE_VERSION(1, 0, 1)
+#define UE_VK_API_VERSION	VK_MAKE_VERSION(1, 1, 0)
 
 enum class EDeviceVulkanSupportStatus
 {
diff --git a/Engine/Source/Runtime/VulkanRHI/Private/Android/VulkanAndroidPlatform.h b/Engine/Source/Runtime/VulkanRHI/Private/Android/VulkanAndroidPlatform.h
--- a/Engine/Source/Runtime/VulkanRHI/Private/Android/VulkanAndroidPlatform.h
+++ b/Engine/Source/Runtime/VulkanRHI/Private/Android/VulkanAndroidPlatform.h
@@ -17,8 +17,7 @@
 #define VULKAN_SUPPORTS_GOOGLE_DISPLAY_TIMING		1
 #define VULKAN_PURGE_SHADER_MODULES					0
 #define VULKAN_SUPPORTS_DEDICATED_ALLOCATION		0
-#define VULKAN_SUPPORTS_PHYSICAL_DEVICE_PROPERTIES2	1
-#define VULKAN_SUPPORTS_ASTC_DECODE_MODE			(VULKAN_SUPPORTS_PHYSICAL_DEVICE_PROPERTIES2)
+#define VULKAN_SUPPORTS_ASTC_DECODE_MODE			1
 #define VULKAN_SUPPORTS_NV_DIAGNOSTIC_CHECKPOINT	0
 #define VULKAN_RHI_RAYTRACING						0
 #define VULKAN_SUPPORTS_TRANSIENT_RESOURCE_ALLOCATOR 0
diff --git a/Engine/Source/Runtime/VulkanRHI/Private/VulkanCommandWrappers.h b/Engine/Source/Runtime/VulkanRHI/Private/VulkanCommandWrappers.h
--- a/Engine/Source/Runtime/VulkanRHI/Private/VulkanCommandWrappers.h
+++ b/Engine/Source/Runtime/VulkanRHI/Private/VulkanCommandWrappers.h
@@ -168,11 +168,9 @@ struct FWrapLayer
 	static void GetImageMemoryRequirements2KHR(VkResult Result, VkDevice Device, const VkImageMemoryRequirementsInfo2KHR* Info, VkMemoryRequirements2KHR* MemoryRequirements) VULKAN_LAYER_BODY
 	static void GetBufferMemoryRequirements2KHR(VkResult Result, VkDevice Device, const VkBufferMemoryRequirementsInfo2KHR* Info, VkMemoryRequirements2KHR* MemoryRequirements) VULKAN_LAYER_BODY
 
-		VULKAN_EXTERN_EXPORT static void GetPhysicalDeviceMemoryProperties2(VkResult Result, VkPhysicalDevice PhysicalDevice, VkPhysicalDeviceMemoryProperties2* MemoryProperties) VULKAN_LAYER_BODY 
-#if VULKAN_SUPPORTS_PHYSICAL_DEVICE_PROPERTIES2
-	VULKAN_EXTERN_EXPORT static void GetPhysicalDeviceProperties2KHR(VkResult Result, VkPhysicalDevice PhysicalDevice, VkPhysicalDeviceProperties2KHR* Properties) VULKAN_LAYER_BODY
-	VULKAN_EXTERN_EXPORT static void GetPhysicalDeviceFeatures2KHR(VkResult Result, VkPhysicalDevice PhysicalDevice, VkPhysicalDeviceFeatures2KHR* Features) VULKAN_LAYER_BODY
-#endif
+	VULKAN_EXTERN_EXPORT static void GetPhysicalDeviceMemoryProperties2(VkResult Result, VkPhysicalDevice PhysicalDevice, VkPhysicalDeviceMemoryProperties2* MemoryProperties) VULKAN_LAYER_BODY 
+	VULKAN_EXTERN_EXPORT static void GetPhysicalDeviceProperties2(VkResult Result, VkPhysicalDevice PhysicalDevice, VkPhysicalDeviceProperties2* Properties) VULKAN_LAYER_BODY
+	VULKAN_EXTERN_EXPORT static void GetPhysicalDeviceFeatures2(VkResult Result, VkPhysicalDevice PhysicalDevice, VkPhysicalDeviceFeatures2* Features) VULKAN_LAYER_BODY
 #if VULKAN_SUPPORTS_FRAGMENT_SHADING_RATE
 	VULKAN_EXTERN_EXPORT static void GetPhysicalDeviceFragmentShadingRatesKHR(VkResult, VkPhysicalDevice PhysicalDevice, uint32* FragmentShadingRateCount, VkPhysicalDeviceFragmentShadingRateKHR* FragmentShadingRates) VULKAN_LAYER_BODY
 #endif
@@ -283,21 +281,19 @@ namespace VulkanRHI
 		FWrapLayer::GetPhysicalDeviceMemoryProperties2(VK_SUCCESS, PhysicalDevice, pMemoryProperties);
 	}
 
-#if VULKAN_SUPPORTS_PHYSICAL_DEVICE_PROPERTIES2
-	static FORCEINLINE_DEBUGGABLE void  vkGetPhysicalDeviceProperties2KHR(VkPhysicalDevice PhysicalDevice, VkPhysicalDeviceProperties2KHR* Properties)
+	static FORCEINLINE_DEBUGGABLE void vkGetPhysicalDeviceProperties2(VkPhysicalDevice PhysicalDevice, VkPhysicalDeviceProperties2* Properties)
 	{
-		FWrapLayer::GetPhysicalDeviceProperties2KHR(VK_RESULT_MAX_ENUM, PhysicalDevice, Properties);
-		VULKANAPINAMESPACE::vkGetPhysicalDeviceProperties2KHR(PhysicalDevice, Properties);
-		FWrapLayer::GetPhysicalDeviceProperties2KHR(VK_SUCCESS, PhysicalDevice, Properties);
+		FWrapLayer::GetPhysicalDeviceProperties2(VK_RESULT_MAX_ENUM, PhysicalDevice, Properties);
+		VULKANAPINAMESPACE::vkGetPhysicalDeviceProperties2(PhysicalDevice, Properties);
+		FWrapLayer::GetPhysicalDeviceProperties2(VK_SUCCESS, PhysicalDevice, Properties);
 	}
 
-	static FORCEINLINE_DEBUGGABLE void  vkGetPhysicalDeviceFeatures2KHR(VkPhysicalDevice PhysicalDevice, VkPhysicalDeviceFeatures2KHR* Features)
+	static FORCEINLINE_DEBUGGABLE void vkGetPhysicalDeviceFeatures2(VkPhysicalDevice PhysicalDevice, VkPhysicalDeviceFeatures2* Features)
 	{
-		FWrapLayer::GetPhysicalDeviceFeatures2KHR(VK_RESULT_MAX_ENUM, PhysicalDevice, Features);
-		VULKANAPINAMESPACE::vkGetPhysicalDeviceFeatures2KHR(PhysicalDevice, Features);
-		FWrapLayer::GetPhysicalDeviceFeatures2KHR(VK_SUCCESS, PhysicalDevice, Features);
+		FWrapLayer::GetPhysicalDeviceFeatures2(VK_RESULT_MAX_ENUM, PhysicalDevice, Features);
+		VULKANAPINAMESPACE::vkGetPhysicalDeviceFeatures2(PhysicalDevice, Features);
+		FWrapLayer::GetPhysicalDeviceFeatures2(VK_SUCCESS, PhysicalDevice, Features);
 	}
-#endif
 
 #if VULKAN_SUPPORTS_FRAGMENT_SHADING_RATE
 	static FORCEINLINE_DEBUGGABLE void vkGetPhysicalDeviceFragmentShadingRatesKHR(VkPhysicalDevice PhysicalDevice, uint32* FragmentShadingRateCount, VkPhysicalDeviceFragmentShadingRateKHR* FragmentShadingRates)
diff --git a/Engine/Source/Runtime/VulkanRHI/Private/VulkanDebug.cpp b/Engine/Source/Runtime/VulkanRHI/Private/VulkanDebug.cpp
--- a/Engine/Source/Runtime/VulkanRHI/Private/VulkanDebug.cpp
+++ b/Engine/Source/Runtime/VulkanRHI/Private/VulkanDebug.cpp
@@ -3858,27 +3858,25 @@ void FWrapLayer::GetPhysicalDeviceMemoryProperties2(VkResult Result, VkPhysicalD
 }
 
 
-#if VULKAN_SUPPORTS_PHYSICAL_DEVICE_PROPERTIES2
-void FWrapLayer::GetPhysicalDeviceProperties2KHR(VkResult Result, VkPhysicalDevice PhysicalDevice, VkPhysicalDeviceProperties2KHR* Properties)
+void FWrapLayer::GetPhysicalDeviceProperties2(VkResult Result, VkPhysicalDevice PhysicalDevice, VkPhysicalDeviceProperties2* Properties)
 {
 	if (Result == VK_RESULT_MAX_ENUM)
 	{
 #if VULKAN_ENABLE_DUMP_LAYER
-		PrintfBegin(FString::Printf(TEXT("vkGetPhysicalDeviceProperties2KHR(PhysicalDevice=0x%p, Properties=0x%p)[...]"), PhysicalDevice, Properties));
+		PrintfBegin(FString::Printf(TEXT("vkGetPhysicalDeviceProperties2(PhysicalDevice=0x%p, Properties=0x%p)[...]"), PhysicalDevice, Properties));
 #endif
 	}
 }
 
-void FWrapLayer::GetPhysicalDeviceFeatures2KHR(VkResult Result, VkPhysicalDevice PhysicalDevice, VkPhysicalDeviceFeatures2KHR* Features)
+void FWrapLayer::GetPhysicalDeviceFeatures2(VkResult Result, VkPhysicalDevice PhysicalDevice, VkPhysicalDeviceFeatures2* Features)
 {
 	if (Result == VK_RESULT_MAX_ENUM)
 	{
 #if VULKAN_ENABLE_DUMP_LAYER
-		PrintfBegin(FString::Printf(TEXT("vkGetPhysicalDeviceFeatures2KHR(PhysicalDevice=0x%p, Features=0x%p)[...]"), PhysicalDevice, Features));
+		PrintfBegin(FString::Printf(TEXT("vkGetPhysicalDeviceFeatures2(PhysicalDevice=0x%p, Features=0x%p)[...]"), PhysicalDevice, Features));
 #endif
 	}
 }
-#endif
 
 #if VULKAN_SUPPORTS_FRAGMENT_SHADING_RATE
 void FWrapLayer::GetPhysicalDeviceFragmentShadingRatesKHR(VkResult Result, VkPhysicalDevice PhysicalDevice, uint32* FragmentShadingRateCount, VkPhysicalDeviceFragmentShadingRateKHR* FragmentShadingRates)
diff --git a/Engine/Source/Runtime/VulkanRHI/Private/VulkanDevice.cpp b/Engine/Source/Runtime/VulkanRHI/Private/VulkanDevice.cpp
--- a/Engine/Source/Runtime/VulkanRHI/Private/VulkanDevice.cpp
+++ b/Engine/Source/Runtime/VulkanRHI/Private/VulkanDevice.cpp
@@ -226,23 +226,15 @@ FVulkanDevice::FVulkanDevice(FVulkanDynamicRHI* InRHI, VkPhysicalDevice InGpu)
 
 	ZeroVulkanStruct(GpuIdProps, VK_STRUCTURE_TYPE_PHYSICAL_DEVICE_ID_PROPERTIES_KHR);
 	ZeroVulkanStruct(GpuSubgroupProps, VK_STRUCTURE_TYPE_PHYSICAL_DEVICE_SUBGROUP_PROPERTIES);
-	if (RHI->GetOptionalExtensions().HasKHRGetPhysicalDeviceProperties2)
+
 	{
 		VkPhysicalDeviceProperties2KHR PhysicalDeviceProperties2;
 		ZeroVulkanStruct(PhysicalDeviceProperties2, VK_STRUCTURE_TYPE_PHYSICAL_DEVICE_PROPERTIES_2_KHR);
 		PhysicalDeviceProperties2.pNext = &GpuIdProps;
-		if (UE_VK_API_VERSION >= VK_API_VERSION_1_1)
-		{
-			// Only consider wave ops on platforms creating a Vulkan 1.1 instance (or greater)
-			GpuIdProps.pNext = &GpuSubgroupProps;
-		}
-		VulkanRHI::vkGetPhysicalDeviceProperties2KHR(Gpu, &PhysicalDeviceProperties2);
+		GpuIdProps.pNext = &GpuSubgroupProps;
+		VulkanRHI::vkGetPhysicalDeviceProperties2(Gpu, &PhysicalDeviceProperties2);
 		GpuProps = PhysicalDeviceProperties2.properties;
 	}
-	else
-	{
-		VulkanRHI::vkGetPhysicalDeviceProperties(Gpu, &GpuProps);
-	}
 
 	// First get the VendorId. We'll have to get properties again after finding out which extensions we want to use
 	VendorId = RHIConvertToGpuVendorId(GpuProps.vendorID);
@@ -1090,55 +1082,52 @@ void FVulkanDevice::InitGPU()
 	FVulkanDeviceExtensionArray UEExtensions = FVulkanDeviceExtension::GetUESupportedDeviceExtensions(this);
 	TArray<const ANSICHAR*> DeviceLayers = FVulkanDevice::SetupDeviceLayers(Gpu, UEExtensions);
 
-	if (RHI->GetOptionalExtensions().HasKHRGetPhysicalDeviceProperties2)
+	// Query advanced features
 	{
-		// Query advanced features
-		{
-			VkPhysicalDeviceFeatures2 PhysicalDeviceFeatures2;
-			ZeroVulkanStruct(PhysicalDeviceFeatures2, VK_STRUCTURE_TYPE_PHYSICAL_DEVICE_FEATURES_2);
+		VkPhysicalDeviceFeatures2 PhysicalDeviceFeatures2;
+		ZeroVulkanStruct(PhysicalDeviceFeatures2, VK_STRUCTURE_TYPE_PHYSICAL_DEVICE_FEATURES_2);
 
-			for (TUniquePtr<FVulkanDeviceExtension>& UEExtension : UEExtensions)
+		for (TUniquePtr<FVulkanDeviceExtension>& UEExtension : UEExtensions)
+		{
+			if (UEExtension->InUse())
 			{
-				if (UEExtension->InUse())
-				{
-					UEExtension->PrePhysicalDeviceFeatures(PhysicalDeviceFeatures2);
-				}
+				UEExtension->PrePhysicalDeviceFeatures(PhysicalDeviceFeatures2);
 			}
+		}
 
-			VulkanRHI::vkGetPhysicalDeviceFeatures2KHR(Gpu, &PhysicalDeviceFeatures2);
+		VulkanRHI::vkGetPhysicalDeviceFeatures2(Gpu, &PhysicalDeviceFeatures2);
 
-			for (TUniquePtr<FVulkanDeviceExtension>& UEExtension : UEExtensions)
+		for (TUniquePtr<FVulkanDeviceExtension>& UEExtension : UEExtensions)
+		{
+			if (UEExtension->InUse())
 			{
-				if (UEExtension->InUse())
-				{
-					UEExtension->PostPhysicalDeviceFeatures(OptionalDeviceExtensions);
-				}
+				UEExtension->PostPhysicalDeviceFeatures(OptionalDeviceExtensions);
 			}
 		}
+	}
 
-		// Query advances properties
-		{
-			VkPhysicalDeviceProperties2KHR PhysicalDeviceProperties2;
-			ZeroVulkanStruct(PhysicalDeviceProperties2, VK_STRUCTURE_TYPE_PHYSICAL_DEVICE_PROPERTIES_2_KHR);
-			PhysicalDeviceProperties2.pNext = &GpuIdProps;
-			ZeroVulkanStruct(GpuIdProps, VK_STRUCTURE_TYPE_PHYSICAL_DEVICE_ID_PROPERTIES_KHR);
+	// Query advances properties
+	{
+		VkPhysicalDeviceProperties2 PhysicalDeviceProperties2;
+		ZeroVulkanStruct(PhysicalDeviceProperties2, VK_STRUCTURE_TYPE_PHYSICAL_DEVICE_PROPERTIES_2);
+		PhysicalDeviceProperties2.pNext = &GpuIdProps;
+		ZeroVulkanStruct(GpuIdProps, VK_STRUCTURE_TYPE_PHYSICAL_DEVICE_ID_PROPERTIES);
 
-			for (TUniquePtr<FVulkanDeviceExtension>& UEExtension : UEExtensions)
+		for (TUniquePtr<FVulkanDeviceExtension>& UEExtension : UEExtensions)
+		{
+			if (UEExtension->InUse())
 			{
-				if (UEExtension->InUse())
-				{
-					UEExtension->PrePhysicalDeviceProperties(PhysicalDeviceProperties2);
-				}
+				UEExtension->PrePhysicalDeviceProperties(PhysicalDeviceProperties2);
 			}
+		}
 
-			VulkanRHI::vkGetPhysicalDeviceProperties2KHR(Gpu, &PhysicalDeviceProperties2);
+		VulkanRHI::vkGetPhysicalDeviceProperties2(Gpu, &PhysicalDeviceProperties2);
 
-			for (TUniquePtr<FVulkanDeviceExtension>& UEExtension : UEExtensions)
+		for (TUniquePtr<FVulkanDeviceExtension>& UEExtension : UEExtensions)
+		{
+			if (UEExtension->InUse())
 			{
-				if (UEExtension->InUse())
-				{
-					UEExtension->PostPhysicalDeviceProperties();
-				}
+				UEExtension->PostPhysicalDeviceProperties();
 			}
 		}
 	}
diff --git a/Engine/Source/Runtime/VulkanRHI/Private/VulkanDevice.h b/Engine/Source/Runtime/VulkanRHI/Private/VulkanDevice.h
--- a/Engine/Source/Runtime/VulkanRHI/Private/VulkanDevice.h
+++ b/Engine/Source/Runtime/VulkanRHI/Private/VulkanDevice.h
@@ -296,13 +296,11 @@ public:
 
 	inline const VkPhysicalDeviceIDPropertiesKHR& GetDeviceIdProperties() const
 	{
-		check(RHI->GetOptionalExtensions().HasKHRGetPhysicalDeviceProperties2);
 		return GpuIdProps;
 	}
 
 	inline const VkPhysicalDeviceSubgroupProperties& GetDeviceSubgroupProperties() const
 	{
-		check(RHI->GetOptionalExtensions().HasKHRGetPhysicalDeviceProperties2);
 		return GpuSubgroupProps;
 	}
 
diff --git a/Engine/Source/Runtime/VulkanRHI/Private/VulkanExtensions.cpp b/Engine/Source/Runtime/VulkanRHI/Private/VulkanExtensions.cpp
--- a/Engine/Source/Runtime/VulkanRHI/Private/VulkanExtensions.cpp
+++ b/Engine/Source/Runtime/VulkanRHI/Private/VulkanExtensions.cpp
@@ -198,12 +198,11 @@ class FVulkanShaderAtomicInt64Extension : public FVulkanDeviceExtension
 public:
 
 	FVulkanShaderAtomicInt64Extension(FVulkanDevice* InDevice)
-		: FVulkanDeviceExtension(InDevice, VK_KHR_SHADER_ATOMIC_INT64_EXTENSION_NAME, VULKAN_SUPPORTS_BUFFER_64BIT_ATOMICS, VK_API_VERSION_1_2)
+		: FVulkanDeviceExtension(InDevice, VK_KHR_SHADER_ATOMIC_INT64_EXTENSION_NAME, VULKAN_EXTENSION_ENABLED, VK_API_VERSION_1_2)
 	{
 		bEnabledInCode = bEnabledInCode && (GRHIAllow64bitShaderAtomicsCvar.GetValueOnAnyThread() != 0);
 	}
 
-#if VULKAN_SUPPORTS_BUFFER_64BIT_ATOMICS
 	virtual void PrePhysicalDeviceFeatures(VkPhysicalDeviceFeatures2KHR& PhysicalDeviceFeatures2) override final
 	{
 		ZeroVulkanStruct(BufferAtomicFeatures, VK_STRUCTURE_TYPE_PHYSICAL_DEVICE_SHADER_ATOMIC_INT64_FEATURES_KHR);
@@ -225,7 +224,6 @@ public:
 	}
 
 	VkPhysicalDeviceShaderAtomicInt64Features BufferAtomicFeatures;
-#endif // VULKAN_SUPPORTS_BUFFER_64BIT_ATOMICS
 };
 
 
@@ -236,12 +234,11 @@ class FVulkanShaderImageAtomicInt64Extension : public FVulkanDeviceExtension
 public:
 
 	FVulkanShaderImageAtomicInt64Extension(FVulkanDevice* InDevice)
-		: FVulkanDeviceExtension(InDevice, VK_EXT_SHADER_IMAGE_ATOMIC_INT64_EXTENSION_NAME, VULKAN_SUPPORTS_IMAGE_64BIT_ATOMICS)
+		: FVulkanDeviceExtension(InDevice, VK_EXT_SHADER_IMAGE_ATOMIC_INT64_EXTENSION_NAME, VULKAN_EXTENSION_ENABLED)
 	{
 		bEnabledInCode = bEnabledInCode && (GRHIAllow64bitShaderAtomicsCvar.GetValueOnAnyThread() != 0);
 	}
 
-#if VULKAN_SUPPORTS_IMAGE_64BIT_ATOMICS
 	virtual void PrePhysicalDeviceFeatures(VkPhysicalDeviceFeatures2KHR& PhysicalDeviceFeatures2) override final
 	{
 		ZeroVulkanStruct(ImageAtomicFeatures, VK_STRUCTURE_TYPE_PHYSICAL_DEVICE_SHADER_IMAGE_ATOMIC_INT64_FEATURES_EXT);
@@ -267,7 +264,6 @@ public:
 	}
 
 	VkPhysicalDeviceShaderImageAtomicInt64FeaturesEXT ImageAtomicFeatures;
-#endif // VULKAN_SUPPORTS_PHYSICAL_DEVICE_PROPERTIES2
 };
 
 
@@ -1337,7 +1333,6 @@ FVulkanInstanceExtensionArray FVulkanInstanceExtension::GetUESupportedInstanceEx
 	FVulkanInstanceExtensionArray OutUEInstanceExtensions;
 
 	// Generic simple extensions :
-	OutUEInstanceExtensions.Add(MakeUnique<FVulkanInstanceExtension>(VK_KHR_GET_PHYSICAL_DEVICE_PROPERTIES_2_EXTENSION_NAME, VULKAN_SUPPORTS_PHYSICAL_DEVICE_PROPERTIES2, VK_API_VERSION_1_1, INSTANCE_EXT_FLAG_SETTER(HasKHRGetPhysicalDeviceProperties2)));
 	OutUEInstanceExtensions.Add(MakeUnique<FVulkanInstanceExtension>(VK_KHR_EXTERNAL_MEMORY_CAPABILITIES_EXTENSION_NAME, VULKAN_SUPPORTS_EXTERNAL_MEMORY, VK_API_VERSION_1_1, INSTANCE_EXT_FLAG_SETTER(HasKHRExternalMemoryCapabilities)));
 	OutUEInstanceExtensions.Add(MakeUnique<FVulkanInstanceExtension>(VK_KHR_SURFACE_EXTENSION_NAME, VULKAN_EXTENSION_ENABLED, VULKAN_EXTENSION_NOT_PROMOTED));
 	OutUEInstanceExtensions.Add(MakeUnique<FVulkanInstanceExtension>(VK_EXT_SWAPCHAIN_COLOR_SPACE_EXTENSION_NAME, VULKAN_EXTENSION_ENABLED, VULKAN_EXTENSION_NOT_PROMOTED));
diff --git a/Engine/Source/Runtime/VulkanRHI/Private/VulkanExtensions.h b/Engine/Source/Runtime/VulkanRHI/Private/VulkanExtensions.h
--- a/Engine/Source/Runtime/VulkanRHI/Private/VulkanExtensions.h
+++ b/Engine/Source/Runtime/VulkanRHI/Private/VulkanExtensions.h
@@ -94,7 +94,6 @@ public:
 	static FVulkanDeviceExtensionArray GetUESupportedDeviceExtensions(FVulkanDevice* InDevice);
 	static TArray<VkExtensionProperties> GetDriverSupportedDeviceExtensions(VkPhysicalDevice Gpu, const ANSICHAR* LayerName = nullptr);
 
-#if VULKAN_SUPPORTS_PHYSICAL_DEVICE_PROPERTIES2
 	virtual void PrePhysicalDeviceProperties(VkPhysicalDeviceProperties2KHR& PhysicalDeviceProperties2) {}
 	virtual void PostPhysicalDeviceProperties() {}
 
@@ -108,7 +107,6 @@ public:
 	}
 
 	virtual void PreCreateDevice(VkDeviceCreateInfo& DeviceInfo) {}
-#endif // VULKAN_SUPPORTS_PHYSICAL_DEVICE_PROPERTIES2
 
 	// Holds extensions requested externally (eg plugins)
 	static TArray<const ANSICHAR*> ExternalExtensions;
diff --git a/Engine/Source/Runtime/VulkanRHI/Private/VulkanLoader.h b/Engine/Source/Runtime/VulkanRHI/Private/VulkanLoader.h
--- a/Engine/Source/Runtime/VulkanRHI/Private/VulkanLoader.h
+++ b/Engine/Source/Runtime/VulkanRHI/Private/VulkanLoader.h
@@ -188,8 +188,8 @@
 // List all optional Vulkan entry points used by Unreal that need to be loaded manually
 #define ENUM_VK_ENTRYPOINTS_OPTIONAL_INSTANCE(EnumMacro) \
 	ENUM_VK_ENTRYPOINTS_RAYTRACING(EnumMacro) \
-	EnumMacro(PFN_vkGetPhysicalDeviceProperties2KHR, vkGetPhysicalDeviceProperties2KHR) \
-	EnumMacro(PFN_vkGetPhysicalDeviceFeatures2KHR, vkGetPhysicalDeviceFeatures2KHR) \
+	EnumMacro(PFN_vkGetPhysicalDeviceProperties2, vkGetPhysicalDeviceProperties2) \
+	EnumMacro(PFN_vkGetPhysicalDeviceFeatures2, vkGetPhysicalDeviceFeatures2) \
 	EnumMacro(PFN_vkGetImageMemoryRequirements2KHR , vkGetImageMemoryRequirements2KHR) \
 	EnumMacro(PFN_vkGetBufferMemoryRequirements2KHR , vkGetBufferMemoryRequirements2KHR) \
 	EnumMacro(PFN_vkGetPhysicalDeviceMemoryProperties2, vkGetPhysicalDeviceMemoryProperties2) \
diff --git a/Engine/Source/Runtime/VulkanRHI/Private/VulkanRHI.cpp b/Engine/Source/Runtime/VulkanRHI/Private/VulkanRHI.cpp
--- a/Engine/Source/Runtime/VulkanRHI/Private/VulkanRHI.cpp
+++ b/Engine/Source/Runtime/VulkanRHI/Private/VulkanRHI.cpp
@@ -556,7 +556,6 @@ void FVulkanDynamicRHI::SelectDevice()
 
 #if VULKAN_ENABLE_DESKTOP_HMD_SUPPORT
 		if (!HmdDevice && HmdGraphicsAdapterLuid != 0 &&
-			GetOptionalExtensions().HasKHRGetPhysicalDeviceProperties2 &&
 			FMemory::Memcmp(&HmdGraphicsAdapterLuid, &NewDevice->GetDeviceIdProperties().deviceLUID, VK_LUID_SIZE_KHR) == 0)
 		{
 			HmdDevice = NewDevice;
@@ -1193,18 +1192,18 @@ uint64 FVulkanDynamicRHI::RHIGetGraphicsAdapterLUID(VkPhysicalDevice InPhysicalD
 {
 	uint64 AdapterLUID = 0;
 #if VULKAN_SUPPORTS_DRIVER_PROPERTIES
-	if (GetOptionalExtensions().HasKHRGetPhysicalDeviceProperties2)
-	{
-		VkPhysicalDeviceIDPropertiesKHR GpuIdProps;
-		VkPhysicalDeviceProperties2KHR GpuProps2;
-		ZeroVulkanStruct(GpuProps2, VK_STRUCTURE_TYPE_PHYSICAL_DEVICE_PROPERTIES_2_KHR);
-		GpuProps2.pNext = &GpuIdProps;
-		ZeroVulkanStruct(GpuIdProps, VK_STRUCTURE_TYPE_PHYSICAL_DEVICE_ID_PROPERTIES_KHR);
-
-		VulkanRHI::vkGetPhysicalDeviceProperties2KHR(InPhysicalDevice, &GpuProps2);
-		check(GpuIdProps.deviceLUIDValid);
-		AdapterLUID = reinterpret_cast<const uint64&>(GpuIdProps.deviceLUID);
-	}
+
+	VkPhysicalDeviceProperties2KHR GpuProps2;
+	ZeroVulkanStruct(GpuProps2, VK_STRUCTURE_TYPE_PHYSICAL_DEVICE_PROPERTIES_2);
+
+	VkPhysicalDeviceIDPropertiesKHR GpuIdProps;
+	ZeroVulkanStruct(GpuIdProps, VK_STRUCTURE_TYPE_PHYSICAL_DEVICE_ID_PROPERTIES);
+
+	GpuProps2.pNext = &GpuIdProps;
+
+	VulkanRHI::vkGetPhysicalDeviceProperties2(InPhysicalDevice, &GpuProps2);
+	check(GpuIdProps.deviceLUIDValid);
+	AdapterLUID = reinterpret_cast<const uint64&>(GpuIdProps.deviceLUID);
 #endif
 	return AdapterLUID;
 }
@@ -1212,13 +1211,10 @@ uint64 FVulkanDynamicRHI::RHIGetGraphicsAdapterLUID(VkPhysicalDevice InPhysicalD
 bool FVulkanDynamicRHI::RHIDoesAdapterMatchDevice(const void* InAdapterId) const
 {
 #if PLATFORM_WINDOWS
-	if (GetOptionalExtensions().HasKHRGetPhysicalDeviceProperties2)
+	const VkPhysicalDeviceIDPropertiesKHR& vkPhysicalDeviceIDProperties = GetDevice()->GetDeviceIdProperties();
+	if (vkPhysicalDeviceIDProperties.deviceLUIDValid)
 	{
-		const VkPhysicalDeviceIDPropertiesKHR& vkPhysicalDeviceIDProperties = GetDevice()->GetDeviceIdProperties();
-		if (vkPhysicalDeviceIDProperties.deviceLUIDValid)
-		{
-			return FMemory::Memcmp(InAdapterId, &vkPhysicalDeviceIDProperties.deviceLUID, sizeof(LUID)) == 0;
-		}
+		return FMemory::Memcmp(InAdapterId, &vkPhysicalDeviceIDProperties.deviceLUID, sizeof(LUID)) == 0;
 	}
 #endif
 
diff --git a/Engine/Source/Runtime/VulkanRHI/Public/VulkanConfiguration.h b/Engine/Source/Runtime/VulkanRHI/Public/VulkanConfiguration.h
--- a/Engine/Source/Runtime/VulkanRHI/Public/VulkanConfiguration.h
+++ b/Engine/Source/Runtime/VulkanRHI/Public/VulkanConfiguration.h
@@ -15,7 +15,7 @@ struct VkAllocationCallbacks;
 
 // API version we want to target.
 #ifndef UE_VK_API_VERSION
-	#define UE_VK_API_VERSION									VK_API_VERSION_1_0
+	#define UE_VK_API_VERSION									VK_API_VERSION_1_1
 #endif
 
 // by default, we enable debugging in Development builds, unless the platform says not to
@@ -198,17 +198,9 @@ struct VkAllocationCallbacks;
 	#endif
 #endif
 
-#ifndef VULKAN_SUPPORTS_PHYSICAL_DEVICE_PROPERTIES2
-	#ifdef VK_KHR_get_physical_device_properties2
-		#define VULKAN_SUPPORTS_PHYSICAL_DEVICE_PROPERTIES2		1
-	#else
-		#define VULKAN_SUPPORTS_PHYSICAL_DEVICE_PROPERTIES2		0
-	#endif
-#endif
-
 #ifndef VULKAN_SUPPORTS_EXTERNAL_MEMORY
 	#ifdef VK_KHR_external_memory_capabilities
-		#define VULKAN_SUPPORTS_EXTERNAL_MEMORY					(VULKAN_SUPPORTS_PHYSICAL_DEVICE_PROPERTIES2)	// Requirement
+		#define VULKAN_SUPPORTS_EXTERNAL_MEMORY					1
 	#else
 		#define VULKAN_SUPPORTS_EXTERNAL_MEMORY					0
 	#endif
@@ -284,24 +276,12 @@ struct VkAllocationCallbacks;
 
 #ifndef VULKAN_SUPPORTS_ASTC_DECODE_MODE
 	#ifdef VK_EXT_astc_decode_mode
-		#define VULKAN_SUPPORTS_ASTC_DECODE_MODE				(VULKAN_SUPPORTS_PHYSICAL_DEVICE_PROPERTIES2)	// Requirement
+		#define VULKAN_SUPPORTS_ASTC_DECODE_MODE				1
 	#else
 		#define VULKAN_SUPPORTS_ASTC_DECODE_MODE				0
 	#endif
 #endif
 
-#ifdef VK_KHR_shader_atomic_int64
-	#define VULKAN_SUPPORTS_BUFFER_64BIT_ATOMICS	(VULKAN_SUPPORTS_PHYSICAL_DEVICE_PROPERTIES2)	// Requirement
-#else
-	#define VULKAN_SUPPORTS_BUFFER_64BIT_ATOMICS	0
-#endif
-
-#ifdef VK_EXT_shader_image_atomic_int64
-	#define VULKAN_SUPPORTS_IMAGE_64BIT_ATOMICS		(VULKAN_SUPPORTS_PHYSICAL_DEVICE_PROPERTIES2)	// Requirement
-#else
-	#define VULKAN_SUPPORTS_IMAGE_64BIT_ATOMICS		0
-#endif
-
 #ifdef VK_EXT_shader_viewport_index_layer
 	#define VULKAN_SUPPORTS_SHADER_VIEWPORT_INDEX_LAYER	1
 #else
@@ -310,7 +290,7 @@ struct VkAllocationCallbacks;
 
 #ifndef VULKAN_SUPPORTS_DESCRIPTOR_INDEXING
 	#ifdef VK_EXT_descriptor_indexing
-		#define VULKAN_SUPPORTS_DESCRIPTOR_INDEXING	(VULKAN_SUPPORTS_PHYSICAL_DEVICE_PROPERTIES2)	// Requirement
+		#define VULKAN_SUPPORTS_DESCRIPTOR_INDEXING	1
 	#else
 		#define VULKAN_SUPPORTS_DESCRIPTOR_INDEXING	0
 	#endif
diff --git a/Engine/Source/Runtime/VulkanRHI/Public/VulkanDynamicRHI.h b/Engine/Source/Runtime/VulkanRHI/Public/VulkanDynamicRHI.h
--- a/Engine/Source/Runtime/VulkanRHI/Public/VulkanDynamicRHI.h
+++ b/Engine/Source/Runtime/VulkanRHI/Public/VulkanDynamicRHI.h
@@ -21,7 +21,6 @@ struct FOptionalVulkanInstanceExtensions
 		struct
 		{
 			uint32 HasKHRExternalMemoryCapabilities : 1;
-			uint32 HasKHRGetPhysicalDeviceProperties2 : 1;
 		};
 		uint32 Packed;
 	};
