VulkanRHI:  Remove VULKAN_SUPPORTS_EXTERNAL_MEMORY which isn't needed now that minimum Vulkan version is 1.1.  Deprecate HasKHRExternalMemoryCapabilities and HasKHRGetPhysicalDeviceProperties2 bits.

#rnx

diff --git a/Engine/Source/Runtime/VulkanRHI/Private/Linux/VulkanLinuxPlatform.cpp b/Engine/Source/Runtime/VulkanRHI/Private/Linux/VulkanLinuxPlatform.cpp
--- a/Engine/Source/Runtime/VulkanRHI/Private/Linux/VulkanLinuxPlatform.cpp
+++ b/Engine/Source/Runtime/VulkanRHI/Private/Linux/VulkanLinuxPlatform.cpp
@@ -197,7 +197,7 @@ void FVulkanLinuxPlatform::GetInstanceExtensions(FVulkanInstanceExtensionArray&
 void FVulkanLinuxPlatform::GetDeviceExtensions(FVulkanDevice* Device, FVulkanDeviceExtensionArray& OutExtensions)
 {
 	// Manually activated extensions
-	OutExtensions.Add(MakeUnique<FVulkanDeviceExtension>(Device, VK_KHR_EXTERNAL_MEMORY_FD_EXTENSION_NAME, VULKAN_SUPPORTS_EXTERNAL_MEMORY, VULKAN_EXTENSION_NOT_PROMOTED, nullptr, FVulkanExtensionBase::ManuallyActivate));
+	OutExtensions.Add(MakeUnique<FVulkanDeviceExtension>(Device, VK_KHR_EXTERNAL_MEMORY_FD_EXTENSION_NAME, VULKAN_EXTENSION_ENABLED, VULKAN_EXTENSION_NOT_PROMOTED, nullptr, FVulkanExtensionBase::ManuallyActivate));
 }
 
 void FVulkanLinuxPlatform::CreateSurface(void* WindowHandle, VkInstance Instance, VkSurfaceKHR* OutSurface)
diff --git a/Engine/Source/Runtime/VulkanRHI/Private/VulkanExtensions.cpp b/Engine/Source/Runtime/VulkanRHI/Private/VulkanExtensions.cpp
--- a/Engine/Source/Runtime/VulkanRHI/Private/VulkanExtensions.cpp
+++ b/Engine/Source/Runtime/VulkanRHI/Private/VulkanExtensions.cpp
@@ -1166,7 +1166,6 @@ FVulkanDeviceExtensionArray FVulkanDeviceExtension::GetUESupportedDeviceExtensio
 
 	// Externally activated extensions (supported by the engine, but enabled externally by plugin or other) :
 
-	ADD_EXTERNAL_EXTENSION(VK_KHR_EXTERNAL_MEMORY_EXTENSION_NAME, VULKAN_SUPPORTS_EXTERNAL_MEMORY, VK_API_VERSION_1_1, nullptr);
 
 
 	// Extensions with custom classes :
@@ -1333,7 +1332,6 @@ FVulkanInstanceExtensionArray FVulkanInstanceExtension::GetUESupportedInstanceEx
 	FVulkanInstanceExtensionArray OutUEInstanceExtensions;
 
 	// Generic simple extensions :
-	OutUEInstanceExtensions.Add(MakeUnique<FVulkanInstanceExtension>(VK_KHR_EXTERNAL_MEMORY_CAPABILITIES_EXTENSION_NAME, VULKAN_SUPPORTS_EXTERNAL_MEMORY, VK_API_VERSION_1_1, INSTANCE_EXT_FLAG_SETTER(HasKHRExternalMemoryCapabilities)));
 	OutUEInstanceExtensions.Add(MakeUnique<FVulkanInstanceExtension>(VK_KHR_SURFACE_EXTENSION_NAME, VULKAN_EXTENSION_ENABLED, VULKAN_EXTENSION_NOT_PROMOTED));
 	OutUEInstanceExtensions.Add(MakeUnique<FVulkanInstanceExtension>(VK_EXT_SWAPCHAIN_COLOR_SPACE_EXTENSION_NAME, VULKAN_EXTENSION_ENABLED, VULKAN_EXTENSION_NOT_PROMOTED));
 
diff --git a/Engine/Source/Runtime/VulkanRHI/Private/VulkanMemory.cpp b/Engine/Source/Runtime/VulkanRHI/Private/VulkanMemory.cpp
--- a/Engine/Source/Runtime/VulkanRHI/Private/VulkanMemory.cpp
+++ b/Engine/Source/Runtime/VulkanRHI/Private/VulkanMemory.cpp
@@ -859,7 +859,6 @@ namespace VulkanRHI
 		}
 #endif
 
-#if VULKAN_SUPPORTS_EXTERNAL_MEMORY
 		VkExportMemoryAllocateInfoKHR VulkanExportMemoryAllocateInfoKHR = {};
 #if PLATFORM_WINDOWS
 		VkExportMemoryWin32HandleInfoKHR VulkanExportMemoryWin32HandleInfoKHR = {};
@@ -881,7 +880,6 @@ namespace VulkanRHI
 #endif // PLATFORM_WINDOWS
 			Info.pNext = &VulkanExportMemoryAllocateInfoKHR;
 		}
-#endif // VULKAN_SUPPORTS_EXTERNAL_MEMORY
 
 #if VULKAN_RHI_RAYTRACING
 		VkMemoryAllocateFlagsInfo MemoryAllocateFlagsInfo;
diff --git a/Engine/Source/Runtime/VulkanRHI/Private/VulkanResources.h b/Engine/Source/Runtime/VulkanRHI/Private/VulkanResources.h
--- a/Engine/Source/Runtime/VulkanRHI/Private/VulkanResources.h
+++ b/Engine/Source/Runtime/VulkanRHI/Private/VulkanResources.h
@@ -434,10 +434,8 @@ public:
 		VkImageCreateInfo ImageCreateInfo;
 		//only used when HasImageFormatListKHR is supported. Otherise VK_IMAGE_CREATE_MUTABLE_FORMAT_BIT is used.
 		VkImageFormatListCreateInfoKHR ImageFormatListCreateInfo;
-#if VULKAN_SUPPORTS_EXTERNAL_MEMORY
 		//used when TexCreate_External is given
 		VkExternalMemoryImageCreateInfoKHR ExternalMemImageCreateInfo;
-#endif // VULKAN_SUPPORTS_EXTERNAL_MEMORY
 		VkFormat FormatsUsed[2];
 	};
 
diff --git a/Engine/Source/Runtime/VulkanRHI/Private/VulkanTexture.cpp b/Engine/Source/Runtime/VulkanRHI/Private/VulkanTexture.cpp
--- a/Engine/Source/Runtime/VulkanRHI/Private/VulkanTexture.cpp
+++ b/Engine/Source/Runtime/VulkanRHI/Private/VulkanTexture.cpp
@@ -369,7 +369,6 @@ void FVulkanTexture::GenerateImageCreateInfo(
 		ImageCreateInfo.usage |= VK_IMAGE_USAGE_STORAGE_BIT;
 	}
 
-	#if VULKAN_SUPPORTS_EXTERNAL_MEMORY
 	if (EnumHasAnyFlags(UEFlags, TexCreate_External))
 	{
 		VkExternalMemoryImageCreateInfoKHR& ExternalMemImageCreateInfo = OutImageCreateInfo.ExternalMemImageCreateInfo;
@@ -382,7 +381,6 @@ void FVulkanTexture::GenerateImageCreateInfo(
 		ExternalMemImageCreateInfo.pNext = ImageCreateInfo.pNext;
     	ImageCreateInfo.pNext = &ExternalMemImageCreateInfo;
 	}
-#endif // VULKAN_SUPPORTS_EXTERNAL_MEMORY
 
 	//#todo-rco: If using CONCURRENT, make sure to NOT do so on render targets as that kills DCC compression
 	ImageCreateInfo.sharingMode = VK_SHARING_MODE_EXCLUSIVE;
diff --git a/Engine/Source/Runtime/VulkanRHI/Private/Windows/VulkanWindowsPlatform.cpp b/Engine/Source/Runtime/VulkanRHI/Private/Windows/VulkanWindowsPlatform.cpp
--- a/Engine/Source/Runtime/VulkanRHI/Private/Windows/VulkanWindowsPlatform.cpp
+++ b/Engine/Source/Runtime/VulkanRHI/Private/Windows/VulkanWindowsPlatform.cpp
@@ -186,7 +186,7 @@ void FVulkanWindowsPlatform::GetDeviceExtensions(FVulkanDevice* Device, FVulkanD
 														VULKAN_EXTENSION_NOT_PROMOTED, DEVICE_EXT_FLAG_SETTER(HasEXTFullscreenExclusive)));
 
 	// Manually activated extensions
-	OutExtensions.Add(MakeUnique<FVulkanDeviceExtension>(Device, VK_KHR_EXTERNAL_MEMORY_WIN32_EXTENSION_NAME, VULKAN_SUPPORTS_EXTERNAL_MEMORY, 
+	OutExtensions.Add(MakeUnique<FVulkanDeviceExtension>(Device, VK_KHR_EXTERNAL_MEMORY_WIN32_EXTENSION_NAME, VULKAN_EXTENSION_ENABLED,
 														VULKAN_EXTENSION_NOT_PROMOTED, nullptr, FVulkanExtensionBase::ManuallyActivate));
 }
 
diff --git a/Engine/Source/Runtime/VulkanRHI/Public/VulkanConfiguration.h b/Engine/Source/Runtime/VulkanRHI/Public/VulkanConfiguration.h
--- a/Engine/Source/Runtime/VulkanRHI/Public/VulkanConfiguration.h
+++ b/Engine/Source/Runtime/VulkanRHI/Public/VulkanConfiguration.h
@@ -198,14 +198,6 @@ struct VkAllocationCallbacks;
 	#endif
 #endif
 
-#ifndef VULKAN_SUPPORTS_EXTERNAL_MEMORY
-	#ifdef VK_KHR_external_memory_capabilities
-		#define VULKAN_SUPPORTS_EXTERNAL_MEMORY					1
-	#else
-		#define VULKAN_SUPPORTS_EXTERNAL_MEMORY					0
-	#endif
-#endif
-
 #ifndef VULKAN_SUPPORTS_DRIVER_PROPERTIES
 	#ifdef VK_KHR_driver_properties
 		#define VULKAN_SUPPORTS_DRIVER_PROPERTIES				1
diff --git a/Engine/Source/Runtime/VulkanRHI/Public/VulkanDynamicRHI.h b/Engine/Source/Runtime/VulkanRHI/Public/VulkanDynamicRHI.h
--- a/Engine/Source/Runtime/VulkanRHI/Public/VulkanDynamicRHI.h
+++ b/Engine/Source/Runtime/VulkanRHI/Public/VulkanDynamicRHI.h
@@ -20,7 +20,11 @@ struct FOptionalVulkanInstanceExtensions
 	{
 		struct
 		{
+			UE_DEPRECATED(5.3, "Vulkan 1.1 is now a requirement so there is no need to check these capabilities.")
 			uint32 HasKHRExternalMemoryCapabilities : 1;
+
+			UE_DEPRECATED(5.3, "Vulkan 1.1 is now a requirement so there is no need to check these capabilities.")
+			uint32 HasKHRGetPhysicalDeviceProperties2 : 1;
 		};
 		uint32 Packed;
 	};
