VulkanRHI: Remove unused code.

#rnx

diff --git a/Engine/Source/Runtime/VulkanRHI/Private/VulkanRHIPrivate.h b/Engine/Source/Runtime/VulkanRHI/Private/VulkanRHIPrivate.h
--- a/Engine/Source/Runtime/VulkanRHI/Private/VulkanRHIPrivate.h
+++ b/Engine/Source/Runtime/VulkanRHI/Private/VulkanRHIPrivate.h
@@ -778,55 +778,6 @@ namespace VulkanRHI
 #endif
 #endif
 
-	// For cases when we want to use DepthRead_StencilDONTCARE
-	inline bool IsDepthReadOnly(FExclusiveDepthStencil DepthStencilAccess)
-	{
-		return DepthStencilAccess.IsUsingDepth() && !DepthStencilAccess.IsDepthWrite();
-	}
-
-	// For cases when we want to use DepthRead_StencilWrite (when we want to read in a shader the current bound depth stencil render target)
-	inline bool IsStencilWrite(FExclusiveDepthStencil DepthStencilAccess)
-	{
-		return DepthStencilAccess.IsUsingStencil() && DepthStencilAccess.IsStencilWrite();
-	}
-
-	inline VkImageLayout GetDepthStencilLayout(FExclusiveDepthStencil RequestedDSAccess, FVulkanDevice& InDevice)
-	{
-		if (!RequestedDSAccess.IsUsingStencil() && InDevice.SupportsParallelRendering())
-		{
-			if (RequestedDSAccess.IsDepthRead())
-			{
-				return VK_IMAGE_LAYOUT_READ_ONLY_OPTIMAL;
-			}
-			else if (RequestedDSAccess.IsDepthWrite())
-			{
-				return VK_IMAGE_LAYOUT_ATTACHMENT_OPTIMAL;
-			}
-		}
-
-		if (RequestedDSAccess == FExclusiveDepthStencil::DepthRead_StencilNop || RequestedDSAccess == FExclusiveDepthStencil::DepthRead_StencilRead)
-		{
-			return VK_IMAGE_LAYOUT_DEPTH_STENCIL_READ_ONLY_OPTIMAL;
-		}
-		else if (RequestedDSAccess == FExclusiveDepthStencil::DepthRead_StencilWrite && InDevice.GetOptionalExtensions().HasKHRMaintenance2)
-		{
-			return VK_IMAGE_LAYOUT_DEPTH_READ_ONLY_STENCIL_ATTACHMENT_OPTIMAL;
-		}
-
-		else if (RequestedDSAccess.IsDepthWrite() && !RequestedDSAccess.IsStencilWrite() && InDevice.SupportsParallelRendering())
-		{
-			return VK_IMAGE_LAYOUT_DEPTH_ATTACHMENT_STENCIL_READ_ONLY_OPTIMAL;
-		}
-
-		if(!RequestedDSAccess.IsUsingDepth() && RequestedDSAccess.IsUsingStencil())  // todo-jn: wussat?  still needs maintenance2?
-		{
-			return VK_IMAGE_LAYOUT_DEPTH_READ_ONLY_STENCIL_ATTACHMENT_OPTIMAL;
-		}
-
-		ensure(RequestedDSAccess.IsDepthWrite() || RequestedDSAccess.IsStencilWrite());
-		return VK_IMAGE_LAYOUT_DEPTH_STENCIL_ATTACHMENT_OPTIMAL;
-	}
-
 	// Merge a depth and a stencil layout for drivers that don't support VK_KHR_separate_depth_stencil_layouts
 	inline VkImageLayout GetMergedDepthStencilLayout(VkImageLayout DepthLayout, VkImageLayout StencilLayout)
 	{
