VulkanRHI: Clean up: Remove extensions that are part of Vulkan 1.1 Core.

#rnx

diff --git a/Engine/Source/Runtime/VulkanRHI/Private/VulkanDebug.cpp b/Engine/Source/Runtime/VulkanRHI/Private/VulkanDebug.cpp
--- a/Engine/Source/Runtime/VulkanRHI/Private/VulkanDebug.cpp
+++ b/Engine/Source/Runtime/VulkanRHI/Private/VulkanDebug.cpp
@@ -876,19 +876,11 @@ namespace VulkanRHI
 			VKSWITCHCASE(VK_ERROR_OUT_OF_DATE_KHR)
 			VKSWITCHCASE(VK_ERROR_INCOMPATIBLE_DISPLAY_KHR)
 			VKSWITCHCASE(VK_ERROR_VALIDATION_FAILED_EXT)
-#if VK_HEADER_VERSION >= 13
 			VKSWITCHCASE(VK_ERROR_INVALID_SHADER_NV)
-#endif
-#if VK_HEADER_VERSION >= 24
 			VKSWITCHCASE(VK_ERROR_FRAGMENTED_POOL)
-#endif
-#if VK_HEADER_VERSION >= 39
 			VKSWITCHCASE(VK_ERROR_OUT_OF_POOL_MEMORY_KHR)
-#endif
-#if VK_HEADER_VERSION >= 65
 			VKSWITCHCASE(VK_ERROR_INVALID_EXTERNAL_HANDLE_KHR)
 			VKSWITCHCASE(VK_ERROR_NOT_PERMITTED_EXT)
-#endif
 #undef VKSWITCHCASE
 		default:
 			break;
diff --git a/Engine/Source/Runtime/VulkanRHI/Private/VulkanDevice.h b/Engine/Source/Runtime/VulkanRHI/Private/VulkanDevice.h
--- a/Engine/Source/Runtime/VulkanRHI/Private/VulkanDevice.h
+++ b/Engine/Source/Runtime/VulkanRHI/Private/VulkanDevice.h
@@ -60,9 +60,6 @@ struct FOptionalVulkanDeviceExtensions
 			uint64 HasQcomRenderPassTransform : 1;
 
 			// Promoted to 1.1
-			uint64 HasKHRMaintenance1 : 1;
-			uint64 HasKHRMaintenance2 : 1;
-			uint64 HasKHRDedicatedAllocation : 1;
 			uint64 HasKHRMultiview : 1;
 
 			// Promoted to 1.2
diff --git a/Engine/Source/Runtime/VulkanRHI/Private/VulkanExtensions.cpp b/Engine/Source/Runtime/VulkanRHI/Private/VulkanExtensions.cpp
--- a/Engine/Source/Runtime/VulkanRHI/Private/VulkanExtensions.cpp
+++ b/Engine/Source/Runtime/VulkanRHI/Private/VulkanExtensions.cpp
@@ -1150,18 +1150,14 @@ FVulkanDeviceExtensionArray FVulkanDeviceExtension::GetUESupportedDeviceExtensio
 	// Generic simple extensions :
 
 	ADD_SIMPLE_EXTENSION(VK_KHR_SWAPCHAIN_EXTENSION_NAME,                VULKAN_EXTENSION_ENABLED,             VULKAN_EXTENSION_NOT_PROMOTED, nullptr);
-	ADD_SIMPLE_EXTENSION(VK_KHR_MAINTENANCE1_EXTENSION_NAME,             VULKAN_EXTENSION_ENABLED,             VK_API_VERSION_1_1,            DEVICE_EXT_FLAG_SETTER(HasKHRMaintenance1));
-	ADD_SIMPLE_EXTENSION(VK_KHR_MAINTENANCE2_EXTENSION_NAME,             VULKAN_EXTENSION_ENABLED,             VK_API_VERSION_1_1,            DEVICE_EXT_FLAG_SETTER(HasKHRMaintenance2));
 	ADD_SIMPLE_EXTENSION(VK_EXT_MEMORY_BUDGET_EXTENSION_NAME,            VULKAN_SUPPORTS_MEMORY_BUDGET,        VULKAN_EXTENSION_NOT_PROMOTED, DEVICE_EXT_FLAG_SETTER(HasMemoryBudget));
 	ADD_SIMPLE_EXTENSION(VK_EXT_MEMORY_PRIORITY_EXTENSION_NAME,          VULKAN_SUPPORTS_MEMORY_PRIORITY,      VULKAN_EXTENSION_NOT_PROMOTED, DEVICE_EXT_FLAG_SETTER(HasMemoryPriority));
 	ADD_SIMPLE_EXTENSION(VK_KHR_CREATE_RENDERPASS_2_EXTENSION_NAME,      VULKAN_SUPPORTS_RENDERPASS2,          VK_API_VERSION_1_2,            DEVICE_EXT_FLAG_SETTER(HasKHRRenderPass2));
-	ADD_SIMPLE_EXTENSION(VK_KHR_DEDICATED_ALLOCATION_EXTENSION_NAME,     VULKAN_SUPPORTS_DEDICATED_ALLOCATION, VK_API_VERSION_1_1,            DEVICE_EXT_FLAG_SETTER(HasKHRDedicatedAllocation));
 	ADD_SIMPLE_EXTENSION(VK_KHR_DEFERRED_HOST_OPERATIONS_EXTENSION_NAME, VULKAN_RHI_RAYTRACING,                VULKAN_EXTENSION_NOT_PROMOTED, DEVICE_EXT_FLAG_SETTER(HasDeferredHostOperations));
 	ADD_SIMPLE_EXTENSION(VK_KHR_SPIRV_1_4_EXTENSION_NAME,                VULKAN_RHI_RAYTRACING,                VK_API_VERSION_1_2,            DEVICE_EXT_FLAG_SETTER(HasSPIRV_14));
 	ADD_SIMPLE_EXTENSION(VK_KHR_SHADER_FLOAT_CONTROLS_EXTENSION_NAME,    VULKAN_RHI_RAYTRACING,                VK_API_VERSION_1_2,            DEVICE_EXT_FLAG_SETTER(HasShaderFloatControls));
 	ADD_SIMPLE_EXTENSION(VK_KHR_IMAGE_FORMAT_LIST_EXTENSION_NAME,        VULKAN_EXTENSION_ENABLED,             VK_API_VERSION_1_2,            DEVICE_EXT_FLAG_SETTER(HasKHRImageFormatList));
 	ADD_SIMPLE_EXTENSION(VK_EXT_VALIDATION_CACHE_EXTENSION_NAME,         VULKAN_SUPPORTS_VALIDATION_CACHE,     VULKAN_EXTENSION_NOT_PROMOTED, DEVICE_EXT_FLAG_SETTER(HasEXTValidationCache));
-	ADD_SIMPLE_EXTENSION(VK_KHR_BIND_MEMORY_2_EXTENSION_NAME,            VULKAN_EXTENSION_ENABLED,             VK_API_VERSION_1_1,            nullptr);
 
 
 	// Externally activated extensions (supported by the engine, but enabled externally by plugin or other) :
diff --git a/Engine/Source/Runtime/VulkanRHI/Private/VulkanTexture.cpp b/Engine/Source/Runtime/VulkanRHI/Private/VulkanTexture.cpp
--- a/Engine/Source/Runtime/VulkanRHI/Private/VulkanTexture.cpp
+++ b/Engine/Source/Runtime/VulkanRHI/Private/VulkanTexture.cpp
@@ -285,7 +285,7 @@ void FVulkanTexture::GenerateImageCreateInfo(
 	ImageCreateInfo.flags = (ResourceType == VK_IMAGE_VIEW_TYPE_CUBE || ResourceType == VK_IMAGE_VIEW_TYPE_CUBE_ARRAY) ? VK_IMAGE_CREATE_CUBE_COMPATIBLE_BIT : 0;
 
 
-	if(EnumHasAllFlags(UEFlags, TexCreate_SRGB))
+	if (EnumHasAllFlags(UEFlags, TexCreate_SRGB))
 	{
 		if(InDevice.GetOptionalExtensions().HasKHRImageFormatList)
 		{
@@ -302,9 +302,9 @@ void FVulkanTexture::GenerateImageCreateInfo(
 		ImageCreateInfo.flags |= VK_IMAGE_CREATE_MUTABLE_FORMAT_BIT;
 	}
 
-	if (InDevice.GetOptionalExtensions().HasKHRMaintenance1 && ImageCreateInfo.imageType == VK_IMAGE_TYPE_3D)
+	if (ImageCreateInfo.imageType == VK_IMAGE_TYPE_3D)
 	{
-		ImageCreateInfo.flags |= VK_IMAGE_CREATE_2D_ARRAY_COMPATIBLE_BIT_KHR;
+		ImageCreateInfo.flags |= VK_IMAGE_CREATE_2D_ARRAY_COMPATIBLE_BIT;
 	}
 
 	ImageCreateInfo.tiling = bForceLinearTexture ? VK_IMAGE_TILING_LINEAR : GVulkanViewTypeTilingMode[ResourceType];
@@ -1512,7 +1512,7 @@ void FVulkanTextureView::Create(FVulkanDevice& Device, VkImage InImage, VkImageV
 
 	// Inform the driver the view will only be used with a subset of usage flags (to help performance and/or compatibility)
 	VkImageViewUsageCreateInfo ImageViewUsageCreateInfo;
-	if ((ImageUsageFlags != 0) && Device.GetOptionalExtensions().HasKHRMaintenance2)
+	if (ImageUsageFlags != 0)
 	{
 		ZeroVulkanStruct(ImageViewUsageCreateInfo, VK_STRUCTURE_TYPE_IMAGE_VIEW_USAGE_CREATE_INFO);
 		ImageViewUsageCreateInfo.usage = ImageUsageFlags;
@@ -1661,7 +1661,7 @@ FVulkanTexture::FVulkanTexture(FVulkanDevice& InDevice, const FRHITextureCreateD
 			extern int32 GVulkanEnableDedicatedImageMemory;
 			// Per https://developer.nvidia.com/what%E2%80%99s-your-vulkan-memory-type
 			VkDeviceSize SizeToBeConsideredForDedicated = 12 * 1024 * 1024;
-			if ((bRenderTarget || MemoryRequirements.size >= SizeToBeConsideredForDedicated) && !bMemoryless && InDevice.GetOptionalExtensions().HasKHRDedicatedAllocation && GVulkanEnableDedicatedImageMemory)
+			if ((bRenderTarget || MemoryRequirements.size >= SizeToBeConsideredForDedicated) && !bMemoryless && GVulkanEnableDedicatedImageMemory)
 			{
 				if (!InDevice.GetMemoryManager().AllocateDedicatedImageMemory(Allocation, this, Image, MemoryRequirements, MemoryFlags, MetaType, bExternal, __FILE__, __LINE__))
 				{
diff --git a/Engine/Source/Runtime/VulkanRHI/Private/VulkanUtil.cpp b/Engine/Source/Runtime/VulkanRHI/Private/VulkanUtil.cpp
--- a/Engine/Source/Runtime/VulkanRHI/Private/VulkanUtil.cpp
+++ b/Engine/Source/Runtime/VulkanRHI/Private/VulkanUtil.cpp
@@ -958,19 +958,11 @@ namespace VulkanRHI
 		VKERRORCASE(VK_ERROR_OUT_OF_DATE_KHR); break;
 		VKERRORCASE(VK_ERROR_INCOMPATIBLE_DISPLAY_KHR); break;
 		VKERRORCASE(VK_ERROR_VALIDATION_FAILED_EXT); break;
-#if VK_HEADER_VERSION >= 13
 		VKERRORCASE(VK_ERROR_INVALID_SHADER_NV); break;
-#endif
-#if VK_HEADER_VERSION >= 24
 		VKERRORCASE(VK_ERROR_FRAGMENTED_POOL); break;
-#endif
-#if VK_HEADER_VERSION >= 39
 		VKERRORCASE(VK_ERROR_OUT_OF_POOL_MEMORY_KHR); break;
-#endif
-#if VK_HEADER_VERSION >= 65
 		VKERRORCASE(VK_ERROR_INVALID_EXTERNAL_HANDLE_KHR); break;
 		VKERRORCASE(VK_ERROR_NOT_PERMITTED_EXT); break;
-#endif
 #undef VKERRORCASE
 		default:
 			break;
diff --git a/Engine/Source/Runtime/VulkanRHI/Public/VulkanConfiguration.h b/Engine/Source/Runtime/VulkanRHI/Public/VulkanConfiguration.h
--- a/Engine/Source/Runtime/VulkanRHI/Public/VulkanConfiguration.h
+++ b/Engine/Source/Runtime/VulkanRHI/Public/VulkanConfiguration.h
@@ -108,11 +108,7 @@ struct VkAllocationCallbacks;
 #endif
 
 #ifndef VULKAN_SUPPORTS_DEDICATED_ALLOCATION
-	#ifdef VK_KHR_dedicated_allocation
-		#define VULKAN_SUPPORTS_DEDICATED_ALLOCATION			1
-	#else
-		#define VULKAN_SUPPORTS_DEDICATED_ALLOCATION			0
-	#endif
+	#define VULKAN_SUPPORTS_DEDICATED_ALLOCATION				1
 #endif
 
 #ifndef VULKAN_SUPPORTS_GOOGLE_DISPLAY_TIMING
@@ -166,44 +162,24 @@ struct VkAllocationCallbacks;
 #endif
 
 #ifndef VULKAN_SUPPORTS_DEBUG_UTILS
-	#ifdef VK_EXT_debug_utils
-		#define VULKAN_SUPPORTS_DEBUG_UTILS						1
-	#else
-		#define VULKAN_SUPPORTS_DEBUG_UTILS						0
-	#endif
+	#define VULKAN_SUPPORTS_DEBUG_UTILS							1
 #endif
 
 #ifndef VULKAN_SUPPORTS_SCALAR_BLOCK_LAYOUT
-	#ifdef VK_EXT_scalar_block_layout
-		#define VULKAN_SUPPORTS_SCALAR_BLOCK_LAYOUT				1
-	#else
-		#define VULKAN_SUPPORTS_SCALAR_BLOCK_LAYOUT				0
-	#endif
+	#define VULKAN_SUPPORTS_SCALAR_BLOCK_LAYOUT					1
 #endif
 
 
 #ifndef VULKAN_SUPPORTS_MEMORY_BUDGET
-	#ifdef VK_EXT_memory_budget
-		#define VULKAN_SUPPORTS_MEMORY_BUDGET					1
-	#else
-		#define VULKAN_SUPPORTS_MEMORY_BUDGET					0
-	#endif
+	#define VULKAN_SUPPORTS_MEMORY_BUDGET						1
 #endif
 
 #ifndef VULKAN_SUPPORTS_MEMORY_PRIORITY
-	#ifdef VK_EXT_memory_priority
-		#define VULKAN_SUPPORTS_MEMORY_PRIORITY					1
-	#else
-		#define VULKAN_SUPPORTS_MEMORY_PRIORITY					0
-	#endif
+	#define VULKAN_SUPPORTS_MEMORY_PRIORITY						1
 #endif
 
 #ifndef VULKAN_SUPPORTS_DRIVER_PROPERTIES
-	#ifdef VK_KHR_driver_properties
-		#define VULKAN_SUPPORTS_DRIVER_PROPERTIES				1
-	#else
-		#define VULKAN_SUPPORTS_DRIVER_PROPERTIES				0
-	#endif
+	#define VULKAN_SUPPORTS_DRIVER_PROPERTIES					1
 #endif
 
 #ifndef VULKAN_SUPPORTS_QCOM_RENDERPASS_TRANSFORM
