VulkanRHI: Add a quick print out of data reported by VK_EXT_device_fault on GPU crashes when supported. WIP

#rnx

diff --git a/Engine/Source/Runtime/VulkanRHI/Private/VulkanCommandWrappers.h b/Engine/Source/Runtime/VulkanRHI/Private/VulkanCommandWrappers.h
--- a/Engine/Source/Runtime/VulkanRHI/Private/VulkanCommandWrappers.h
+++ b/Engine/Source/Runtime/VulkanRHI/Private/VulkanCommandWrappers.h
@@ -214,6 +214,7 @@ struct FWrapLayer
 	static void CmdBindDescriptorBuffersEXT(VkResult Result, VkCommandBuffer CommandBuffer, uint32_t BufferCount, const VkDescriptorBufferBindingInfoEXT* BindingInfos) VULKAN_LAYER_BODY
 	static void CmdSetDescriptorBufferOffsetsEXT(VkResult Result, VkCommandBuffer CommandBuffer, VkPipelineBindPoint PipelineBindPoint, VkPipelineLayout Layout, uint32_t FirstSet, uint32_t SetCount, const uint32_t* BufferIndices, const VkDeviceSize* Offsets) VULKAN_LAYER_BODY
 	static void GetDescriptorEXT(VkResult Result, VkDevice Device, const VkDescriptorGetInfoEXT* DescriptorInfo, size_t DataSize, void* Descriptor) VULKAN_LAYER_BODY
+	static void GetDeviceFaultInfoEXT(VkResult Result, VkDevice Device, VkDeviceFaultCountsEXT* FaultCounts, VkDeviceFaultInfoEXT* FaultInfo) VULKAN_LAYER_BODY
 };
 
 #undef VULKAN_LAYER_BODY
@@ -1618,6 +1619,13 @@ namespace VulkanRHI
 		FWrapLayer::GetDescriptorEXT(VK_SUCCESS, Device, DescriptorInfo,  DataSize, Descriptor);
 	}
 
+	static FORCEINLINE_DEBUGGABLE VkResult vkGetDeviceFaultInfoEXT(VkDevice Device, VkDeviceFaultCountsEXT* FaultCounts, VkDeviceFaultInfoEXT* FaultInfo)
+	{
+		FWrapLayer::GetDeviceFaultInfoEXT(VK_RESULT_MAX_ENUM, Device, FaultCounts, FaultInfo);
+		VkResult Result = VULKANAPINAMESPACE::vkGetDeviceFaultInfoEXT(Device, FaultCounts, FaultInfo);
+		FWrapLayer::GetDeviceFaultInfoEXT(Result, Device, FaultCounts, FaultInfo);
+		return Result;
+	}
 #if VULKAN_ENABLE_IMAGE_TRACKING_LAYER
 	void BindDebugLabelName(VkImage Image, const TCHAR* Name);
 #endif
diff --git a/Engine/Source/Runtime/VulkanRHI/Private/VulkanDebug.cpp b/Engine/Source/Runtime/VulkanRHI/Private/VulkanDebug.cpp
--- a/Engine/Source/Runtime/VulkanRHI/Private/VulkanDebug.cpp
+++ b/Engine/Source/Runtime/VulkanRHI/Private/VulkanDebug.cpp
@@ -4536,6 +4536,16 @@ void FWrapLayer::GetDescriptorEXT(VkResult Result, VkDevice Device, const VkDesc
 	}
 }
 
+void FWrapLayer::GetDeviceFaultInfoEXT(VkResult Result, VkDevice Device, VkDeviceFaultCountsEXT* FaultCounts, VkDeviceFaultInfoEXT* FaultInfo)
+{
+	if (Result == VK_RESULT_MAX_ENUM)
+	{
+#if VULKAN_ENABLE_DUMP_LAYER
+		PrintfBeginResult(FString::Printf(TEXT("GetDeviceFaultInfoEXT(Device=0x%p)"), Device));
+#endif
+	}
+}
+
 void FWrapLayer::CmdPipelineBarrier2KHR(VkResult Result, VkCommandBuffer CommandBuffer, const VkDependencyInfo* DependencyInfo)
 {
 	if (Result == VK_RESULT_MAX_ENUM)
diff --git a/Engine/Source/Runtime/VulkanRHI/Private/VulkanDevice.h b/Engine/Source/Runtime/VulkanRHI/Private/VulkanDevice.h
--- a/Engine/Source/Runtime/VulkanRHI/Private/VulkanDevice.h
+++ b/Engine/Source/Runtime/VulkanRHI/Private/VulkanDevice.h
@@ -52,6 +52,7 @@ struct FOptionalVulkanDeviceExtensions
 			uint64 HasDeferredHostOperations : 1;
 			uint64 HasEXTCalibratedTimestamps : 1;
 			uint64 HasEXTDescriptorBuffer : 1;
+			uint64 HasEXTDeviceFault : 1;
 
 			// Vendor specific
 			uint64 HasAMDBufferMarker : 1;
diff --git a/Engine/Source/Runtime/VulkanRHI/Private/VulkanExtensions.cpp b/Engine/Source/Runtime/VulkanRHI/Private/VulkanExtensions.cpp
--- a/Engine/Source/Runtime/VulkanRHI/Private/VulkanExtensions.cpp
+++ b/Engine/Source/Runtime/VulkanRHI/Private/VulkanExtensions.cpp
@@ -935,6 +935,39 @@ private:
 	VkDeviceDiagnosticsConfigCreateInfoNV DeviceDiagnosticsConfigCreateInfoNV;
 };
 
+// ***** VK_EXT_device_fault
+class FVulkanEXTDeviceFaultExtension : public FVulkanDeviceExtension
+{
+public:
+
+	FVulkanEXTDeviceFaultExtension(FVulkanDevice* InDevice)
+		: FVulkanDeviceExtension(InDevice, VK_EXT_DEVICE_FAULT_EXTENSION_NAME, VULKAN_EXTENSION_ENABLED)
+	{
+	}
+
+	virtual void PrePhysicalDeviceFeatures(VkPhysicalDeviceFeatures2KHR& PhysicalDeviceFeatures2) override final
+	{
+		ZeroVulkanStruct(DeviceFaultFeatures, VK_STRUCTURE_TYPE_PHYSICAL_DEVICE_FAULT_FEATURES_EXT);
+		AddToPNext(PhysicalDeviceFeatures2, DeviceFaultFeatures);
+	}
+
+	virtual void PostPhysicalDeviceFeatures(FOptionalVulkanDeviceExtensions& ExtensionFlags) override final
+	{
+		bRequirementsPassed = (DeviceFaultFeatures.deviceFault == VK_TRUE);
+		ExtensionFlags.HasEXTDeviceFault = bRequirementsPassed;
+	}
+
+	virtual void PreCreateDevice(VkDeviceCreateInfo& DeviceCreateInfo) override final
+	{
+		if (bRequirementsPassed)
+		{
+			AddToPNext(DeviceCreateInfo, DeviceFaultFeatures);
+		}
+	}
+
+private:
+	VkPhysicalDeviceFaultFeaturesEXT DeviceFaultFeatures;
+};
 
 // ***** VK_EXT_host_query_reset
 class FVulkanEXTHostQueryResetExtension : public FVulkanDeviceExtension
@@ -1184,6 +1217,7 @@ FVulkanDeviceExtensionArray FVulkanDeviceExtension::GetUESupportedDeviceExtensio
 	ADD_CUSTOM_EXTENSION(FVulkanEXTSubgroupSizeControlExtension);
 	ADD_CUSTOM_EXTENSION(FVulkanEXTCalibratedTimestampsExtension);
 	ADD_CUSTOM_EXTENSION(FVulkanEXTDescriptorBuffer);
+	ADD_CUSTOM_EXTENSION(FVulkanEXTDeviceFaultExtension);
 
 	// Needed for Raytracing
 	ADD_CUSTOM_EXTENSION(FVulkanKHRBufferDeviceAddressExtension);
diff --git a/Engine/Source/Runtime/VulkanRHI/Private/VulkanLoader.h b/Engine/Source/Runtime/VulkanRHI/Private/VulkanLoader.h
--- a/Engine/Source/Runtime/VulkanRHI/Private/VulkanLoader.h
+++ b/Engine/Source/Runtime/VulkanRHI/Private/VulkanLoader.h
@@ -213,7 +213,8 @@
 	EnumMacro(PFN_vkGetDescriptorSetLayoutBindingOffsetEXT, vkGetDescriptorSetLayoutBindingOffsetEXT) \
 	EnumMacro(PFN_vkCmdBindDescriptorBuffersEXT, vkCmdBindDescriptorBuffersEXT) \
 	EnumMacro(PFN_vkCmdSetDescriptorBufferOffsetsEXT, vkCmdSetDescriptorBufferOffsetsEXT) \
-	EnumMacro(PFN_vkGetDescriptorEXT, vkGetDescriptorEXT)
+	EnumMacro(PFN_vkGetDescriptorEXT, vkGetDescriptorEXT) \
+	EnumMacro(PFN_vkGetDeviceFaultInfoEXT, vkGetDeviceFaultInfoEXT)
 
 
 // List of all Vulkan entry points
diff --git a/Engine/Source/Runtime/VulkanRHI/Private/VulkanUtil.cpp b/Engine/Source/Runtime/VulkanRHI/Private/VulkanUtil.cpp
--- a/Engine/Source/Runtime/VulkanRHI/Private/VulkanUtil.cpp
+++ b/Engine/Source/Runtime/VulkanRHI/Private/VulkanUtil.cpp
@@ -921,6 +921,76 @@ namespace VulkanRHI
 		return Buffer;
 	}
 
+	void CheckDeviceFault(FVulkanDevice* InDevice)
+	{
+		if (InDevice->GetOptionalExtensions().HasEXTDeviceFault)
+		{
+			const VkDevice DeviceHandle = InDevice->GetInstanceHandle();
+			VkResult Result;
+
+			VkDeviceFaultCountsEXT FaultCounts;
+			ZeroVulkanStruct(FaultCounts, VK_STRUCTURE_TYPE_DEVICE_FAULT_COUNTS_EXT);
+			Result = vkGetDeviceFaultInfoEXT(DeviceHandle, &FaultCounts, nullptr);
+			if (Result == VK_SUCCESS)
+			{
+				VkDeviceFaultInfoEXT FaultInfo;
+				ZeroVulkanStruct(FaultInfo, VK_STRUCTURE_TYPE_DEVICE_FAULT_INFO_EXT);
+
+				TArray<VkDeviceFaultAddressInfoEXT> AddressInfos;
+				AddressInfos.SetNumZeroed(FaultCounts.addressInfoCount);
+				FaultInfo.pAddressInfos = AddressInfos.GetData();
+
+				TArray<VkDeviceFaultVendorInfoEXT> VendorInfos;
+				VendorInfos.SetNumZeroed(FaultCounts.vendorInfoCount);
+				FaultInfo.pVendorInfos = VendorInfos.GetData();
+
+				TArray<uint8> VendorBinaryData;
+				VendorBinaryData.SetNumZeroed(FaultCounts.vendorBinarySize);
+				FaultInfo.pVendorBinaryData = VendorBinaryData.GetData();
+
+				Result = vkGetDeviceFaultInfoEXT(DeviceHandle, &FaultCounts, &FaultInfo);
+				if (Result == VK_SUCCESS)
+				{
+					// :todo-jn: match these up to resources
+
+					auto ReportAddrToStr = [&AddressInfos]() {
+						FString AddrStr;
+						for (const VkDeviceFaultAddressInfoEXT& AddrInfo : AddressInfos)
+						{
+							const uint64_t LowerAddress = (AddrInfo.reportedAddress & ~(AddrInfo.addressPrecision - 1));
+							const uint64_t UpperAddress = (AddrInfo.reportedAddress | (AddrInfo.addressPrecision - 1));
+
+							AddrStr += FString::Printf(TEXT("\n    - %s : 0x%016llX (range:0x%016llX-0x%016llX)"), 
+								VK_TYPE_TO_STRING(VkDeviceFaultAddressTypeEXT, AddrInfo.addressType),
+								AddrInfo.reportedAddress, LowerAddress, UpperAddress);
+						}
+						return AddrStr;
+					};
+
+					auto ReportVendorToStr = [&VendorInfos]() {
+						FString VendorStr;
+						for (const VkDeviceFaultVendorInfoEXT& VendorInfo : VendorInfos)
+						{
+							VendorStr += FString::Printf(TEXT("\n    - %s (code:0x%016llX data:0x%016llX)"),
+								StringCast<TCHAR>((UTF8CHAR*)VendorInfo.description).Get(), VendorInfo.vendorFaultCode, VendorInfo.vendorFaultData);
+						}
+						return VendorStr;
+					};
+
+					UE_LOG(LogVulkanRHI, Error, 
+						TEXT("\nDEVICE FAULT REPORT:\n")
+						TEXT("* Description: %s\n")
+						TEXT("* Address Info: %s\n")
+						TEXT("* Vendor Info: %s\n")
+						TEXT("* Vendor Binary Size: %llu\n"),
+
+						StringCast<TCHAR>((UTF8CHAR*)FaultInfo.description).Get(), *ReportAddrToStr(), *ReportVendorToStr(), FaultCounts.vendorBinarySize
+					);
+				}
+			}
+		}
+	}
+
 	/**
 	 * Checks that the given result isn't a failure.  If it is, the application exits with an appropriate error message.
 	 * @param	Result - The result code to check
@@ -981,14 +1051,20 @@ namespace VulkanRHI
 		UE_LOG(LogVulkanRHI, Error, TEXT("%s failed, VkResult=%d\n at %s:%u \n with error %s"),
 			ANSI_TO_TCHAR(VkFunction), (int32)Result, ANSI_TO_TCHAR(Filename), Line, *ErrorString);
 
-#if VULKAN_SUPPORTS_GPU_CRASH_DUMPS
-		if (GIsGPUCrashed && GGPUCrashDebuggingEnabled)
+		if (GIsGPUCrashed)
 		{
 			FVulkanDevice* Device = GVulkanRHI->GetDevice();
-			Device->GetImmediateContext().GetGPUProfiler().DumpCrashMarkers(Device->GetCrashMarkerMappedPointer());
-		}
+
+#if VULKAN_SUPPORTS_GPU_CRASH_DUMPS
+			if (GGPUCrashDebuggingEnabled)
+			{
+				Device->GetImmediateContext().GetGPUProfiler().DumpCrashMarkers(Device->GetCrashMarkerMappedPointer());
+			}
 #endif
 
+			CheckDeviceFault(Device);
+		}
+
 #if UE_BUILD_DEBUG || UE_BUILD_DEVELOPMENT
 		if (bDumpMemory)
 		{
