VulkanRHI: Disable raytracing when 64bit ints aren't supported in shader code.

#rnx
#deadend

diff --git a/Engine/Source/Runtime/VulkanRHI/Private/VulkanRHI.cpp b/Engine/Source/Runtime/VulkanRHI/Private/VulkanRHI.cpp
--- a/Engine/Source/Runtime/VulkanRHI/Private/VulkanRHI.cpp
+++ b/Engine/Source/Runtime/VulkanRHI/Private/VulkanRHI.cpp
@@ -877,7 +877,7 @@ void FVulkanDynamicRHI::InitInstance()
 		GSupportsMobileMultiView = Device->GetOptionalExtensions().HasKHRMultiview ? true : false;
 #endif
 #if VULKAN_RHI_RAYTRACING
-		GRHISupportsRayTracing = RHISupportsRayTracing(GMaxRHIShaderPlatform) && Device->GetOptionalExtensions().HasRaytracingExtensions();
+		GRHISupportsRayTracing = RHISupportsRayTracing(GMaxRHIShaderPlatform) && Device->GetOptionalExtensions().HasRaytracingExtensions() && Device->GetPhysicalFeatures().shaderInt64;
 
 		if (GRHISupportsRayTracing)
 		{
