VulkanRHI: Partial integration from 24868436 to fix flickering in viewport on high CPU workload.

#rnx

diff --git a/Engine/Source/Runtime/VulkanRHI/Private/VulkanUniformBuffer.cpp b/Engine/Source/Runtime/VulkanRHI/Private/VulkanUniformBuffer.cpp
--- a/Engine/Source/Runtime/VulkanRHI/Private/VulkanUniformBuffer.cpp
+++ b/Engine/Source/Runtime/VulkanRHI/Private/VulkanUniformBuffer.cpp
@@ -226,8 +226,12 @@ FVulkanUniformBuffer::FVulkanUniformBuffer(FVulkanDevice& InDevice, const FRHIUn
 	{
 		const bool bInRenderingThread = IsInRenderingThread();
 		const bool bInRHIThread = IsInRHIThread();
+
 PRAGMA_DISABLE_DEPRECATION_WARNINGS
-		if (UseRingBuffer(InUsage) && (bInRenderingThread || bInRHIThread) && !UE::Tasks::Private::IsThreadRetractingTask())
+		if (UseRingBuffer(InUsage) && (bInRenderingThread || bInRHIThread)
+			// :todo-jn:  Temporary check until we have a command list arg passed in to avoid a race where the RenderThread
+			// would pick up other tasks (because of task retraction) and execute them as if on the RenderThread.
+			&& !UE::Tasks::Private::IsThreadRetractingTask())
 PRAGMA_ENABLE_DEPRECATION_WARNINGS
 		{
 			if (Contents)
