VulkanRHI: Remove checkSlow that was firing in debug with new TextureReference code.

#rnx

diff --git a/Engine/Source/Runtime/VulkanRHI/Private/VulkanResources.h b/Engine/Source/Runtime/VulkanRHI/Private/VulkanResources.h
--- a/Engine/Source/Runtime/VulkanRHI/Private/VulkanResources.h
+++ b/Engine/Source/Runtime/VulkanRHI/Private/VulkanResources.h
@@ -417,7 +417,6 @@ public:
 
 	virtual FRHIDescriptorHandle GetDefaultBindlessHandle() const override final
 	{
-		checkSlow(DefaultBindlessHandle.IsValid());
 		return DefaultBindlessHandle;
 	}
 
