VulkanRHI: Fix Vulkan non unity build

#rnx

diff --git a/Engine/Source/Runtime/VulkanRHI/Private/VulkanResources.h b/Engine/Source/Runtime/VulkanRHI/Private/VulkanResources.h
--- a/Engine/Source/Runtime/VulkanRHI/Private/VulkanResources.h
+++ b/Engine/Source/Runtime/VulkanRHI/Private/VulkanResources.h
@@ -6,19 +6,19 @@
 
 #pragma once
 
+#include "CoreMinimal.h"
+#include "VulkanPlatform.h"
 #include "VulkanConfiguration.h"
 #include "VulkanState.h"
 #include "VulkanUtil.h"
 #include "BoundShaderStateCache.h"
 #include "VulkanShaderResources.h"
-#include "VulkanState.h"
 #include "VulkanMemory.h"
 #include "Misc/ScopeRWLock.h"
 
 class FVulkanDevice;
 class FVulkanQueue;
 class FVulkanCmdBuffer;
-class FVulkanBufferCPU;
 class FVulkanTexture;
 struct FVulkanBufferView;
 class FVulkanResourceMultiBuffer;
@@ -31,7 +31,6 @@ struct FRHITransientHeapAllocation;
 namespace VulkanRHI
 {
 	class FDeviceMemoryAllocation;
-	class FOldResourceAllocation;
 	struct FPendingBufferLock;
 	class FVulkanViewBase;
 }
