Update Swappy to the latest from https://android.googlesource.com/platform/frameworks/opt/gamesdk, modified to always use the Java DisplayManager and refresh rate selection
Enable Swappy on Android by default
Fixed: Swappy VK initialization which was missing a WindowHandle
Remove old GoogleGameSDK files that are no longer needed
Removed Java version of MemAdvisor

diff --git a/Engine/Source/Runtime/ApplicationCore/Private/Android/AndroidPlatformFramePacer.cpp b/Engine/Source/Runtime/ApplicationCore/Private/Android/AndroidPlatformFramePacer.cpp
--- a/Engine/Source/Runtime/ApplicationCore/Private/Android/AndroidPlatformFramePacer.cpp
+++ b/Engine/Source/Runtime/ApplicationCore/Private/Android/AndroidPlatformFramePacer.cpp
@@ -18,8 +18,8 @@ TAutoConsoleVariable<int32> FAndroidPlatformRHIFramePacer::CVarUseSwappyForFrame
 
 TAutoConsoleVariable<int32> FAndroidPlatformRHIFramePacer::CVarSupportNonVSyncMultipleFrameRates(
 	TEXT("a.SupportNonVSyncMultipleFrameRates"),
-	1,
-	TEXT("Set to True to support frame rates we cannot vsync at. Requires a.UseSwappyForFramePacing=1"));
+	0,
+	TEXT("Set to True to support frame rates we cannot vsync at. Requires a.UseSwappyForFramePacing=1 for OpenGL ES."));
 
 TAutoConsoleVariable<int32> FAndroidPlatformRHIFramePacer::CVarAllowFrameTimestamps(
 	TEXT("a.AllowFrameTimestamps"),
diff --git a/Engine/Source/Runtime/VulkanRHI/Private/Android/VulkanAndroidPlatform.cpp b/Engine/Source/Runtime/VulkanRHI/Private/Android/VulkanAndroidPlatform.cpp
--- a/Engine/Source/Runtime/VulkanRHI/Private/Android/VulkanAndroidPlatform.cpp
+++ b/Engine/Source/Runtime/VulkanRHI/Private/Android/VulkanAndroidPlatform.cpp
@@ -729,33 +729,16 @@ bool FAndroidVulkanFramePacer::SupportsFramePaceInternal(int32 QueryFramePace, i
 	// check if we want to use CPU frame pacing at less than a multiple of supported refresh rate
 	if (FAndroidPlatformRHIFramePacer::CVarSupportNonVSyncMultipleFrameRates.GetValueOnAnyThread() == 1)
 	{
-		int32 ClosestRefreshRate = 0;
-		
 		for (int32 Rate : RefreshRates)
 		{
 			if (Rate > QueryFramePace)
 			{
-				if (ClosestRefreshRate < QueryFramePace)
-				{
-					ClosestRefreshRate = Rate;
-				}
-
-				// prefer 30hz multiple
-				if ((Rate % 30) == 0)
-				{
-					ClosestRefreshRate = Rate;
-					break;
-				}
+				UE_LOG(LogRHI, Log, TEXT("Supports %d using refresh rate %d with CPU frame pacing"), QueryFramePace, Rate);
+				OutRefreshRate = Rate;
+				OutSyncInterval = 0;
+				return true;
 			}
 		}
-
-		if (ClosestRefreshRate > QueryFramePace)
-		{
-			OutRefreshRate = ClosestRefreshRate;
-			OutSyncInterval = 0;
-			UE_LOG(LogRHI, Log, TEXT("Supports %d using refresh rate %d with CPU frame pacing"), QueryFramePace, OutRefreshRate);
-			return true;
-		}
 	}
 
 	OutRefreshRate = QueryFramePace;
