VulkanRT: Fix a crash when calculating BLAS size when mesh streaming is enabled:
* Don't query IB/VB buffer addresses because those are ignored when estimating size.
* When IB was not streamed in yet assume its format is uint32 (same as we do for other platforms).

diff --git a/Engine/Source/Runtime/VulkanRHI/Private/VulkanRayTracing.cpp b/Engine/Source/Runtime/VulkanRHI/Private/VulkanRayTracing.cpp
--- a/Engine/Source/Runtime/VulkanRHI/Private/VulkanRayTracing.cpp
+++ b/Engine/Source/Runtime/VulkanRHI/Private/VulkanRayTracing.cpp
@@ -45,6 +45,15 @@ bool FVulkanRayTracingPlatform::CheckVulkanInstanceFunctions(VkInstance inInstan
 #pragma warning(pop) // restore 4191
 #endif
 
+enum class EBLASBuildDataUsage
+{	
+	// Uses provided VB/IB when filling out BLAS build data
+	Rendering = 0,
+
+	// Does not use VB/IB. Special mode for estimating BLAS size.
+	Size = 1
+};
+
 static VkDeviceAddress GetDeviceAddress(VkDevice Device, VkBuffer Buffer)
 {
 	VkBufferDeviceAddressInfoKHR DeviceAddressInfo;
@@ -165,17 +174,29 @@ static void GetBLASBuildData(
 	const TArrayView<const FRayTracingGeometrySegment> Segments,
 	const ERayTracingGeometryType GeometryType,
 	const FBufferRHIRef IndexBufferRHI,
-	const uint32 IndexBufferOffset,	
-	const uint32 IndexStrideInBytes,
+	const uint32 IndexBufferOffset,
 	ERayTracingAccelerationStructureFlags BuildFlags,
 	const EAccelerationStructureBuildMode BuildMode,
+	const EBLASBuildDataUsage Usage,
 	FVkRtBLASBuildData& BuildData)
 {
 	static constexpr uint32 IndicesPerPrimitive = 3; // Only triangle meshes are supported
 
 	FVulkanResourceMultiBuffer* const IndexBuffer = ResourceCast(IndexBufferRHI.GetReference());
 	VkDeviceOrHostAddressConstKHR IndexBufferDeviceAddress = {};
-	IndexBufferDeviceAddress.deviceAddress = IndexBufferRHI ? IndexBuffer->GetDeviceAddress() + IndexBufferOffset : 0;
+	
+	// We only need to get IB/VB address when we are getting data for rendering. For estimating BLAS size we set them to 0.
+	// According to vulkan spec any VkDeviceOrHostAddressKHR members are ignored in vkGetAccelerationStructureBuildSizesKHR.
+	uint32 IndexStrideInBytes = 0;
+	if (IndexBufferRHI)
+	{
+		IndexBufferDeviceAddress.deviceAddress = Usage == EBLASBuildDataUsage::Rendering ? IndexBuffer->GetDeviceAddress() + IndexBufferOffset : 0;
+
+		// In case we are just calculating size but index buffer is not yet in valid state we assume the geometry is using uint32 format
+		IndexStrideInBytes = Usage == EBLASBuildDataUsage::Rendering
+			? IndexBuffer->GetStride()
+			: (IndexBuffer->GetSize() > 0) ? IndexBuffer->GetStride() : 4;
+	}
 
 	TArray<uint32, TInlineAllocator<1>> PrimitiveCounts;
 
@@ -186,7 +207,9 @@ static void GetBLASBuildData(
 		FVulkanResourceMultiBuffer* const VertexBuffer = ResourceCast(Segment.VertexBuffer.GetReference());
 
 		VkDeviceOrHostAddressConstKHR VertexBufferDeviceAddress = {};
-		VertexBufferDeviceAddress.deviceAddress = VertexBuffer->GetDeviceAddress() + Segment.VertexBufferOffset;
+		VertexBufferDeviceAddress.deviceAddress = Usage == EBLASBuildDataUsage::Rendering
+			? VertexBuffer->GetDeviceAddress() + Segment.VertexBufferOffset
+			: 0;
 
 		VkAccelerationStructureGeometryKHR SegmentGeometry;
 		ZeroVulkanStruct(SegmentGeometry, VK_STRUCTURE_TYPE_ACCELERATION_STRUCTURE_GEOMETRY_KHR);
@@ -736,8 +759,6 @@ FRayTracingAccelerationStructureSize FVulkanDynamicRHI::RHICalcRayTracingSceneSi
 
 FRayTracingAccelerationStructureSize FVulkanDynamicRHI::RHICalcRayTracingGeometrySize(const FRayTracingGeometryInitializer& Initializer)
 {	
-	const uint32 IndexStrideInBytes = Initializer.IndexBuffer ? Initializer.IndexBuffer->GetStride() : 0;
-
 	FVkRtBLASBuildData BuildData;
 	GetBLASBuildData(
 		Device->GetInstanceHandle(),
@@ -745,9 +766,9 @@ FRayTracingAccelerationStructureSize FVulkanDynamicRHI::RHICalcRayTracingGeometr
 		Initializer.GeometryType,
 		Initializer.IndexBuffer,
 		Initializer.IndexBufferOffset,
-		IndexStrideInBytes,
 		GetRayTracingAccelerationStructureBuildFlags(Initializer),
 		EAccelerationStructureBuildMode::Build,
+		EBLASBuildDataUsage::Size,
 		BuildData);
 
 	FRayTracingAccelerationStructureSize Result;
@@ -841,10 +862,10 @@ void FVulkanCommandListContext::RHIBuildAccelerationStructures(const TArrayView<
 			MakeArrayView(Geometry->Initializer.Segments),
 			Geometry->Initializer.GeometryType,
 			Geometry->Initializer.IndexBuffer,
-			Geometry->Initializer.IndexBufferOffset,			
-			Geometry->Initializer.IndexBuffer ? Geometry->Initializer.IndexBuffer->GetStride() : 0,
+			Geometry->Initializer.IndexBufferOffset,
 			GetRayTracingAccelerationStructureBuildFlags(Geometry->Initializer),
 			P.BuildMode,
+			EBLASBuildDataUsage::Rendering,
 			BuildData);
 
 		check(BuildData.SizesInfo.accelerationStructureSize <= Geometry->AccelerationStructureBuffer->GetSize());
