VulkanRHI: Minor Vulkan 1.1 clean up: move vkGetImageMemoryRequirements2 and vkGetBufferMemoryRequirements2 from optional to required (remove KHR suffix).

#rnx

diff --git a/Engine/Source/Runtime/VulkanRHI/Private/VulkanCommandWrappers.h b/Engine/Source/Runtime/VulkanRHI/Private/VulkanCommandWrappers.h
--- a/Engine/Source/Runtime/VulkanRHI/Private/VulkanCommandWrappers.h
+++ b/Engine/Source/Runtime/VulkanRHI/Private/VulkanCommandWrappers.h
@@ -165,8 +165,8 @@ struct FWrapLayer
 	static void AcquireNextImageKHR(VkResult Result, VkDevice Device, VkSwapchainKHR Swapchain, uint64_t Timeout, VkSemaphore Semaphore, VkFence Fence, uint32_t* ImageIndex) VULKAN_LAYER_BODY
 	static void DestroySurfaceKHR(VkResult Result, VkInstance Instance, VkSurfaceKHR SurfaceKHR) VULKAN_LAYER_BODY
 	static void DestroySwapchainKHR(VkResult Result, VkDevice Device, VkSwapchainKHR Swapchain) VULKAN_LAYER_BODY
-	static void GetImageMemoryRequirements2KHR(VkResult Result, VkDevice Device, const VkImageMemoryRequirementsInfo2KHR* Info, VkMemoryRequirements2KHR* MemoryRequirements) VULKAN_LAYER_BODY
-	static void GetBufferMemoryRequirements2KHR(VkResult Result, VkDevice Device, const VkBufferMemoryRequirementsInfo2KHR* Info, VkMemoryRequirements2KHR* MemoryRequirements) VULKAN_LAYER_BODY
+	static void GetImageMemoryRequirements2(VkResult Result, VkDevice Device, const VkImageMemoryRequirementsInfo2* Info, VkMemoryRequirements2* MemoryRequirements) VULKAN_LAYER_BODY
+	static void GetBufferMemoryRequirements2(VkResult Result, VkDevice Device, const VkBufferMemoryRequirementsInfo2* Info, VkMemoryRequirements2* MemoryRequirements) VULKAN_LAYER_BODY
 
 	VULKAN_EXTERN_EXPORT static void GetPhysicalDeviceMemoryProperties2(VkResult Result, VkPhysicalDevice PhysicalDevice, VkPhysicalDeviceMemoryProperties2* MemoryProperties) VULKAN_LAYER_BODY 
 	VULKAN_EXTERN_EXPORT static void GetPhysicalDeviceProperties2(VkResult Result, VkPhysicalDevice PhysicalDevice, VkPhysicalDeviceProperties2* Properties) VULKAN_LAYER_BODY
@@ -1398,18 +1398,18 @@ namespace VulkanRHI
 		FWrapLayer::DestroySurfaceKHR(VK_SUCCESS, Instance, Surface);
 	}
 
-	static FORCEINLINE_DEBUGGABLE void vkGetImageMemoryRequirements2KHR(VkDevice Device, const VkImageMemoryRequirementsInfo2KHR* Info, VkMemoryRequirements2KHR* MemoryRequirements)
+	static FORCEINLINE_DEBUGGABLE void vkGetImageMemoryRequirements2(VkDevice Device, const VkImageMemoryRequirementsInfo2* Info, VkMemoryRequirements2* MemoryRequirements)
 	{
-		FWrapLayer::GetImageMemoryRequirements2KHR(VK_RESULT_MAX_ENUM, Device, Info, MemoryRequirements);
-		VULKANAPINAMESPACE::vkGetImageMemoryRequirements2KHR(Device, Info, MemoryRequirements);
-		FWrapLayer::GetImageMemoryRequirements2KHR(VK_SUCCESS, Device, Info, MemoryRequirements);
+		FWrapLayer::GetImageMemoryRequirements2(VK_RESULT_MAX_ENUM, Device, Info, MemoryRequirements);
+		VULKANAPINAMESPACE::vkGetImageMemoryRequirements2(Device, Info, MemoryRequirements);
+		FWrapLayer::GetImageMemoryRequirements2(VK_SUCCESS, Device, Info, MemoryRequirements);
 	}
 
-	static FORCEINLINE_DEBUGGABLE void vkGetBufferMemoryRequirements2KHR(VkDevice Device, const VkBufferMemoryRequirementsInfo2KHR* Info, VkMemoryRequirements2KHR* MemoryRequirements)
+	static FORCEINLINE_DEBUGGABLE void vkGetBufferMemoryRequirements2(VkDevice Device, const VkBufferMemoryRequirementsInfo2* Info, VkMemoryRequirements2* MemoryRequirements)
 	{
-		FWrapLayer::GetBufferMemoryRequirements2KHR(VK_RESULT_MAX_ENUM, Device, Info, MemoryRequirements);
-		VULKANAPINAMESPACE::vkGetBufferMemoryRequirements2KHR(Device, Info, MemoryRequirements);
-		FWrapLayer::GetBufferMemoryRequirements2KHR(VK_SUCCESS, Device, Info, MemoryRequirements);
+		FWrapLayer::GetBufferMemoryRequirements2(VK_RESULT_MAX_ENUM, Device, Info, MemoryRequirements);
+		VULKANAPINAMESPACE::vkGetBufferMemoryRequirements2(Device, Info, MemoryRequirements);
+		FWrapLayer::GetBufferMemoryRequirements2(VK_SUCCESS, Device, Info, MemoryRequirements);
 	}
 
 #if VULKAN_RHI_RAYTRACING
diff --git a/Engine/Source/Runtime/VulkanRHI/Private/VulkanDebug.cpp b/Engine/Source/Runtime/VulkanRHI/Private/VulkanDebug.cpp
--- a/Engine/Source/Runtime/VulkanRHI/Private/VulkanDebug.cpp
+++ b/Engine/Source/Runtime/VulkanRHI/Private/VulkanDebug.cpp
@@ -3780,22 +3780,22 @@ void FWrapLayer::GetPhysicalDeviceQueueFamilyProperties(VkResult Result, VkPhysi
 	}
 }
 
-void FWrapLayer::GetImageMemoryRequirements2KHR(VkResult Result, VkDevice Device, const VkImageMemoryRequirementsInfo2KHR* Info, VkMemoryRequirements2KHR* MemoryRequirements)
+void FWrapLayer::GetImageMemoryRequirements2(VkResult Result, VkDevice Device, const VkImageMemoryRequirementsInfo2* Info, VkMemoryRequirements2* MemoryRequirements)
 {
 	if (Result == VK_RESULT_MAX_ENUM)
 	{
 #if VULKAN_ENABLE_DUMP_LAYER
-		DevicePrintfBegin(Device, FString::Printf(TEXT("vkGetImageMemoryRequirements2KHR(Info=0x%p, MemReqs=0x%p)[...]"), Info, MemoryRequirements));
+		DevicePrintfBegin(Device, FString::Printf(TEXT("vkGetImageMemoryRequirements2(Info=0x%p, MemReqs=0x%p)[...]"), Info, MemoryRequirements));
 #endif
 	}
 }
 
-void FWrapLayer::GetBufferMemoryRequirements2KHR(VkResult Result, VkDevice Device, const VkBufferMemoryRequirementsInfo2KHR* Info, VkMemoryRequirements2KHR* MemoryRequirements)
+void FWrapLayer::GetBufferMemoryRequirements2(VkResult Result, VkDevice Device, const VkBufferMemoryRequirementsInfo2* Info, VkMemoryRequirements2* MemoryRequirements)
 {
 	if (Result == VK_RESULT_MAX_ENUM)
 	{
 #if VULKAN_ENABLE_DUMP_LAYER
-		DevicePrintfBegin(Device, FString::Printf(TEXT("vkGetBufferemoryRequirements2KHR(Info=0x%p, MemReqs=0x%p)[...]"), Info, MemoryRequirements));
+		DevicePrintfBegin(Device, FString::Printf(TEXT("vkGetBufferemoryRequirements2(Info=0x%p, MemReqs=0x%p)[...]"), Info, MemoryRequirements));
 #endif
 	}
 }
diff --git a/Engine/Source/Runtime/VulkanRHI/Private/VulkanLoader.h b/Engine/Source/Runtime/VulkanRHI/Private/VulkanLoader.h
--- a/Engine/Source/Runtime/VulkanRHI/Private/VulkanLoader.h
+++ b/Engine/Source/Runtime/VulkanRHI/Private/VulkanLoader.h
@@ -37,6 +37,8 @@
 	EnumMacro(PFN_vkGetBufferMemoryRequirements, vkGetBufferMemoryRequirements) \
 	EnumMacro(PFN_vkGetImageMemoryRequirements, vkGetImageMemoryRequirements) \
 	EnumMacro(PFN_vkGetImageSparseMemoryRequirements, vkGetImageSparseMemoryRequirements) \
+	EnumMacro(PFN_vkGetImageMemoryRequirements2, vkGetImageMemoryRequirements2) \
+	EnumMacro(PFN_vkGetBufferMemoryRequirements2, vkGetBufferMemoryRequirements2) \
 	EnumMacro(PFN_vkGetPhysicalDeviceSparseImageFormatProperties, vkGetPhysicalDeviceSparseImageFormatProperties) \
 	EnumMacro(PFN_vkQueueBindSparse, vkQueueBindSparse) \
 	EnumMacro(PFN_vkCreateFence, vkCreateFence) \
@@ -190,8 +192,6 @@
 	ENUM_VK_ENTRYPOINTS_RAYTRACING(EnumMacro) \
 	EnumMacro(PFN_vkGetPhysicalDeviceProperties2, vkGetPhysicalDeviceProperties2) \
 	EnumMacro(PFN_vkGetPhysicalDeviceFeatures2, vkGetPhysicalDeviceFeatures2) \
-	EnumMacro(PFN_vkGetImageMemoryRequirements2KHR , vkGetImageMemoryRequirements2KHR) \
-	EnumMacro(PFN_vkGetBufferMemoryRequirements2KHR , vkGetBufferMemoryRequirements2KHR) \
 	EnumMacro(PFN_vkGetPhysicalDeviceMemoryProperties2, vkGetPhysicalDeviceMemoryProperties2) \
 	EnumMacro(PFN_vkCreateRenderPass2KHR, vkCreateRenderPass2KHR) \
 	EnumMacro(PFN_vkCmdBeginRenderPass2KHR, vkCmdBeginRenderPass2KHR) \
diff --git a/Engine/Source/Runtime/VulkanRHI/Private/VulkanMemory.cpp b/Engine/Source/Runtime/VulkanRHI/Private/VulkanMemory.cpp
--- a/Engine/Source/Runtime/VulkanRHI/Private/VulkanMemory.cpp
+++ b/Engine/Source/Runtime/VulkanRHI/Private/VulkanMemory.cpp
@@ -2994,11 +2994,11 @@ namespace VulkanRHI
 		VkMemoryDedicatedRequirementsKHR DedMemoryReqs;
 		ZeroVulkanStruct(DedMemoryReqs, VK_STRUCTURE_TYPE_MEMORY_DEDICATED_REQUIREMENTS_KHR);
 
-		VkMemoryRequirements2KHR MemoryReqs2;
+		VkMemoryRequirements2 MemoryReqs2;
 		ZeroVulkanStruct(MemoryReqs2, VK_STRUCTURE_TYPE_MEMORY_REQUIREMENTS_2_KHR);
 		MemoryReqs2.pNext = &DedMemoryReqs;
 
-		VulkanRHI::vkGetImageMemoryRequirements2KHR(Device->GetInstanceHandle(), &ImageMemoryReqs2, &MemoryReqs2);
+		VulkanRHI::vkGetImageMemoryRequirements2(Device->GetInstanceHandle(), &ImageMemoryReqs2, &MemoryReqs2);
 
 		bool bUseDedicated = DedMemoryReqs.prefersDedicatedAllocation != VK_FALSE || DedMemoryReqs.requiresDedicatedAllocation != VK_FALSE;
 		if (bUseDedicated)
diff --git a/Engine/Source/Runtime/VulkanRHI/Private/VulkanTexture.cpp b/Engine/Source/Runtime/VulkanRHI/Private/VulkanTexture.cpp
--- a/Engine/Source/Runtime/VulkanRHI/Private/VulkanTexture.cpp
+++ b/Engine/Source/Runtime/VulkanRHI/Private/VulkanTexture.cpp
@@ -2331,12 +2331,12 @@ FDynamicRHI::FRHICalcTextureSizeResult FVulkanDynamicRHI::RHICalcTexturePlatform
 
 	if (Device->GetOptionalExtensions().HasKHRMaintenance4)
 	{
-		VkDeviceImageMemoryRequirementsKHR ImageMemReq;
+		VkDeviceImageMemoryRequirements ImageMemReq;
 		ZeroVulkanStruct(ImageMemReq, VK_STRUCTURE_TYPE_DEVICE_IMAGE_MEMORY_REQUIREMENTS);
 		ImageMemReq.pCreateInfo = &TmpCreateInfo.ImageCreateInfo;
 		ImageMemReq.planeAspect = (VulkanRHI::GetAspectMaskFromUEFormat(CleanDesc.Format, true, true) == VK_IMAGE_ASPECT_COLOR_BIT) ? VK_IMAGE_ASPECT_COLOR_BIT : VK_IMAGE_ASPECT_DEPTH_BIT;  // should be ignored
 
-		VkMemoryRequirements2KHR MemReq2;
+		VkMemoryRequirements2 MemReq2;
 		ZeroVulkanStruct(MemReq2, VK_STRUCTURE_TYPE_MEMORY_REQUIREMENTS_2);
 
 		VulkanRHI::vkGetDeviceImageMemoryRequirementsKHR(Device->GetInstanceHandle(), &ImageMemReq, &MemReq2);
