VulkanRHI: Minor Vulkan 1.1 clean up: move vkGetPhysicalDeviceFeatures2, vkGetPhysicalDeviceProperties2, vkGetPhysicalDeviceMemoryProperties2, vkBindBufferMemory2 and vkBindImageMemory2 from optional to required (remove KHR suffix).

#rnx

diff --git a/Engine/Source/Runtime/VulkanRHI/Private/VulkanCommandWrappers.h b/Engine/Source/Runtime/VulkanRHI/Private/VulkanCommandWrappers.h
--- a/Engine/Source/Runtime/VulkanRHI/Private/VulkanCommandWrappers.h
+++ b/Engine/Source/Runtime/VulkanRHI/Private/VulkanCommandWrappers.h
@@ -202,8 +202,8 @@ struct FWrapLayer
 	static void ResetQueryPoolEXT(VkResult Result, VkDevice Device, VkQueryPool QueryPool, uint32_t FirstQuery, uint32_t QueryCount) VULKAN_LAYER_BODY
 	static void GetPhysicalDeviceCalibrateableTimeDomainsEXT(VkResult Result, VkPhysicalDevice PhysicalDevice, uint32_t* TimeDomainCount, VkTimeDomainEXT* TimeDomains) VULKAN_LAYER_BODY
 	static void GetCalibratedTimestampsEXT(VkResult Result, VkDevice Device, uint32_t TimestampCount, const VkCalibratedTimestampInfoEXT* TimestampInfos, uint64_t* Timestamps, uint64_t* MaxDeviation) VULKAN_LAYER_BODY
-	static void BindBufferMemory2KHR(VkResult Result, VkDevice Device, uint32_t BindInfoCount, const VkBindBufferMemoryInfo* BindInfos) VULKAN_LAYER_BODY
-	static void BindImageMemory2KHR(VkResult Result, VkDevice Device, uint32_t BindInfoCount, const VkBindImageMemoryInfo* BindInfos) VULKAN_LAYER_BODY
+	static void BindBufferMemory2(VkResult Result, VkDevice Device, uint32_t BindInfoCount, const VkBindBufferMemoryInfo* BindInfos) VULKAN_LAYER_BODY
+	static void BindImageMemory2(VkResult Result, VkDevice Device, uint32_t BindInfoCount, const VkBindImageMemoryInfo* BindInfos) VULKAN_LAYER_BODY
 	static void CmdPipelineBarrier2KHR(VkResult Result, VkCommandBuffer CommandBuffer, const VkDependencyInfo* DependencyInfo) VULKAN_LAYER_BODY
 	static void CmdResetEvent2KHR(VkResult Result, VkCommandBuffer CommandBuffer, VkEvent Event, VkPipelineStageFlags2 StageMask) VULKAN_LAYER_BODY
 	static void CmdSetEvent2KHR(VkResult Result, VkCommandBuffer CommandBuffer, VkEvent Event, const VkDependencyInfo* DependencyInfo) VULKAN_LAYER_BODY
@@ -1533,19 +1533,19 @@ namespace VulkanRHI
 		return Result;
 	}
 
-	static FORCEINLINE_DEBUGGABLE VkResult vkBindBufferMemory2KHR(VkDevice Device, uint32_t BindInfoCount, const VkBindBufferMemoryInfo* BindInfos)
+	static FORCEINLINE_DEBUGGABLE VkResult vkBindBufferMemory2(VkDevice Device, uint32_t BindInfoCount, const VkBindBufferMemoryInfo* BindInfos)
 	{
-		FWrapLayer::BindBufferMemory2KHR(VK_RESULT_MAX_ENUM, Device, BindInfoCount, BindInfos);
-		const VkResult Result = VULKANAPINAMESPACE::vkBindBufferMemory2KHR(Device, BindInfoCount, BindInfos);
-		FWrapLayer::BindBufferMemory2KHR(Result, Device, BindInfoCount, BindInfos);
+		FWrapLayer::BindBufferMemory2(VK_RESULT_MAX_ENUM, Device, BindInfoCount, BindInfos);
+		const VkResult Result = VULKANAPINAMESPACE::vkBindBufferMemory2(Device, BindInfoCount, BindInfos);
+		FWrapLayer::BindBufferMemory2(Result, Device, BindInfoCount, BindInfos);
 		return Result;
 	}
 
-	static FORCEINLINE_DEBUGGABLE VkResult vkBindImageMemory2KHR(VkDevice Device, uint32_t BindInfoCount, const VkBindImageMemoryInfo* BindInfos)
+	static FORCEINLINE_DEBUGGABLE VkResult vkBindImageMemory2(VkDevice Device, uint32_t BindInfoCount, const VkBindImageMemoryInfo* BindInfos)
 	{
-		FWrapLayer::BindImageMemory2KHR(VK_RESULT_MAX_ENUM, Device, BindInfoCount, BindInfos);
-		const VkResult Result = VULKANAPINAMESPACE::vkBindImageMemory2KHR(Device, BindInfoCount, BindInfos);
-		FWrapLayer::BindImageMemory2KHR(Result, Device, BindInfoCount, BindInfos);
+		FWrapLayer::BindImageMemory2(VK_RESULT_MAX_ENUM, Device, BindInfoCount, BindInfos);
+		const VkResult Result = VULKANAPINAMESPACE::vkBindImageMemory2(Device, BindInfoCount, BindInfos);
+		FWrapLayer::BindImageMemory2(Result, Device, BindInfoCount, BindInfos);
 		return Result;
 	}
 
diff --git a/Engine/Source/Runtime/VulkanRHI/Private/VulkanDebug.cpp b/Engine/Source/Runtime/VulkanRHI/Private/VulkanDebug.cpp
--- a/Engine/Source/Runtime/VulkanRHI/Private/VulkanDebug.cpp
+++ b/Engine/Source/Runtime/VulkanRHI/Private/VulkanDebug.cpp
@@ -4457,22 +4457,22 @@ void FWrapLayer::GetCalibratedTimestampsEXT(VkResult Result, VkDevice Device, ui
 	}
 }
 
-void FWrapLayer::BindBufferMemory2KHR(VkResult Result, VkDevice Device, uint32_t BindInfoCount, const VkBindBufferMemoryInfo* BindInfos)
+void FWrapLayer::BindBufferMemory2(VkResult Result, VkDevice Device, uint32_t BindInfoCount, const VkBindBufferMemoryInfo* BindInfos)
 {
 	if (Result == VK_RESULT_MAX_ENUM)
 	{
 #if VULKAN_ENABLE_DUMP_LAYER
-		PrintfBeginResult(FString::Printf(TEXT("BindBufferMemory2KHR(Device=0x%p, BindInfoCount=%u)"), Device, BindInfoCount));
+		PrintfBeginResult(FString::Printf(TEXT("BindBufferMemory2(Device=0x%p, BindInfoCount=%u)"), Device, BindInfoCount));
 #endif
 	}
 }
 
-void FWrapLayer::BindImageMemory2KHR(VkResult Result, VkDevice Device, uint32_t BindInfoCount, const VkBindImageMemoryInfo* BindInfos)
+void FWrapLayer::BindImageMemory2(VkResult Result, VkDevice Device, uint32_t BindInfoCount, const VkBindImageMemoryInfo* BindInfos)
 {
 	if (Result == VK_RESULT_MAX_ENUM)
 	{
 #if VULKAN_ENABLE_DUMP_LAYER
-		PrintfBeginResult(FString::Printf(TEXT("BindImageMemory2KHR(Device=0x%p, BindInfoCount=%u)"), Device, BindInfoCount));
+		PrintfBeginResult(FString::Printf(TEXT("BindImageMemory2(Device=0x%p, BindInfoCount=%u)"), Device, BindInfoCount));
 #endif
 	}
 }
diff --git a/Engine/Source/Runtime/VulkanRHI/Private/VulkanLoader.h b/Engine/Source/Runtime/VulkanRHI/Private/VulkanLoader.h
--- a/Engine/Source/Runtime/VulkanRHI/Private/VulkanLoader.h
+++ b/Engine/Source/Runtime/VulkanRHI/Private/VulkanLoader.h
@@ -17,6 +17,9 @@
 	EnumMacro(PFN_vkGetPhysicalDeviceProperties, vkGetPhysicalDeviceProperties) \
 	EnumMacro(PFN_vkGetPhysicalDeviceQueueFamilyProperties, vkGetPhysicalDeviceQueueFamilyProperties) \
 	EnumMacro(PFN_vkGetPhysicalDeviceMemoryProperties, vkGetPhysicalDeviceMemoryProperties) \
+	EnumMacro(PFN_vkGetPhysicalDeviceProperties2, vkGetPhysicalDeviceProperties2) \
+	EnumMacro(PFN_vkGetPhysicalDeviceFeatures2, vkGetPhysicalDeviceFeatures2) \
+	EnumMacro(PFN_vkGetPhysicalDeviceMemoryProperties2, vkGetPhysicalDeviceMemoryProperties2) \
 	EnumMacro(PFN_vkCreateDevice, vkCreateDevice) \
 	EnumMacro(PFN_vkDestroyDevice, vkDestroyDevice) \
 	EnumMacro(PFN_vkEnumerateDeviceExtensionProperties, vkEnumerateDeviceExtensionProperties) \
@@ -34,6 +37,8 @@
 	EnumMacro(PFN_vkGetDeviceMemoryCommitment, vkGetDeviceMemoryCommitment) \
 	EnumMacro(PFN_vkBindBufferMemory, vkBindBufferMemory) \
 	EnumMacro(PFN_vkBindImageMemory, vkBindImageMemory) \
+	EnumMacro(PFN_vkBindBufferMemory2, vkBindBufferMemory2) \
+	EnumMacro(PFN_vkBindImageMemory2, vkBindImageMemory2) \
 	EnumMacro(PFN_vkGetBufferMemoryRequirements, vkGetBufferMemoryRequirements) \
 	EnumMacro(PFN_vkGetImageMemoryRequirements, vkGetImageMemoryRequirements) \
 	EnumMacro(PFN_vkGetImageSparseMemoryRequirements, vkGetImageSparseMemoryRequirements) \
@@ -190,9 +195,6 @@
 // List all optional Vulkan entry points used by Unreal that need to be loaded manually
 #define ENUM_VK_ENTRYPOINTS_OPTIONAL_INSTANCE(EnumMacro) \
 	ENUM_VK_ENTRYPOINTS_RAYTRACING(EnumMacro) \
-	EnumMacro(PFN_vkGetPhysicalDeviceProperties2, vkGetPhysicalDeviceProperties2) \
-	EnumMacro(PFN_vkGetPhysicalDeviceFeatures2, vkGetPhysicalDeviceFeatures2) \
-	EnumMacro(PFN_vkGetPhysicalDeviceMemoryProperties2, vkGetPhysicalDeviceMemoryProperties2) \
 	EnumMacro(PFN_vkCreateRenderPass2KHR, vkCreateRenderPass2KHR) \
 	EnumMacro(PFN_vkCmdBeginRenderPass2KHR, vkCmdBeginRenderPass2KHR) \
 	EnumMacro(PFN_vkGetDeviceImageMemoryRequirementsKHR, vkGetDeviceImageMemoryRequirementsKHR) \
@@ -206,8 +208,6 @@
 	EnumMacro(PFN_vkCreateSharedSwapchainsKHR, vkCreateSharedSwapchainsKHR) \
 	EnumMacro(PFN_vkGetPhysicalDeviceCalibrateableTimeDomainsEXT, vkGetPhysicalDeviceCalibrateableTimeDomainsEXT) \
 	EnumMacro(PFN_vkGetCalibratedTimestampsEXT, vkGetCalibratedTimestampsEXT) \
-	EnumMacro(PFN_vkBindBufferMemory2KHR, vkBindBufferMemory2KHR) \
-	EnumMacro(PFN_vkBindImageMemory2KHR, vkBindImageMemory2KHR) \
 	EnumMacro(PFN_vkGetBufferDeviceAddressKHR, vkGetBufferDeviceAddressKHR) \
 	EnumMacro(PFN_vkGetDescriptorSetLayoutSizeEXT, vkGetDescriptorSetLayoutSizeEXT) \
 	EnumMacro(PFN_vkGetDescriptorSetLayoutBindingOffsetEXT, vkGetDescriptorSetLayoutBindingOffsetEXT) \
