VulkanRHI: Remove code that was meant to be deadend (wrong tag was used).  shaderInt64 is verified by the vulkan profile now instead.

#rnx

diff --git a/Engine/Source/Runtime/VulkanRHI/Private/VulkanRHI.cpp b/Engine/Source/Runtime/VulkanRHI/Private/VulkanRHI.cpp
--- a/Engine/Source/Runtime/VulkanRHI/Private/VulkanRHI.cpp
+++ b/Engine/Source/Runtime/VulkanRHI/Private/VulkanRHI.cpp
@@ -887,7 +887,7 @@ void FVulkanDynamicRHI::InitInstance()
 		GSupportsMobileMultiView = Device->GetOptionalExtensions().HasKHRMultiview ? true : false;
 #endif
 #if VULKAN_RHI_RAYTRACING
-		GRHISupportsRayTracing = RHISupportsRayTracing(GMaxRHIShaderPlatform) && Device->GetOptionalExtensions().HasRaytracingExtensions() && Device->GetPhysicalFeatures().shaderInt64;
+		GRHISupportsRayTracing = RHISupportsRayTracing(GMaxRHIShaderPlatform) && Device->GetOptionalExtensions().HasRaytracingExtensions();
 
 		if (GRHISupportsRayTracing)
 		{
