VulkanRHI: Remove check in GetBindlessHandle, it's possible to request a handle that isn't valid (like in ReferenceTexture).

#rnx

diff --git a/Engine/Source/Runtime/VulkanRHI/Private/VulkanResources.h b/Engine/Source/Runtime/VulkanRHI/Private/VulkanResources.h
--- a/Engine/Source/Runtime/VulkanRHI/Private/VulkanResources.h
+++ b/Engine/Source/Runtime/VulkanRHI/Private/VulkanResources.h
@@ -458,7 +458,6 @@ public:
 
 	FRHIDescriptorHandle GetBindlessHandle() const
 	{
-		check(BindlessHandle.IsValid());
 		return BindlessHandle;
 	}
 
