Ignore the depth input attachment for subpass if full pre pass is enabled on mobile.

diff --git a/Engine/Source/Runtime/VulkanRHI/Private/VulkanRenderTarget.cpp b/Engine/Source/Runtime/VulkanRHI/Private/VulkanRenderTarget.cpp
--- a/Engine/Source/Runtime/VulkanRHI/Private/VulkanRenderTarget.cpp
+++ b/Engine/Source/Runtime/VulkanRHI/Private/VulkanRenderTarget.cpp
@@ -1354,8 +1354,8 @@ FVulkanRenderTargetLayout::FVulkanRenderTargetLayout(const FGraphicsPipelineStat
 		StencilDesc.stencilFinalLayout = VK_IMAGE_LAYOUT_ATTACHMENT_OPTIMAL;
 
 		DepthReference.attachment = NumAttachmentDescriptions;
-		DepthReference.layout = VK_IMAGE_LAYOUT_ATTACHMENT_OPTIMAL;
-		StencilReference.stencilLayout = VK_IMAGE_LAYOUT_ATTACHMENT_OPTIMAL;
+		DepthReference.layout = Initializer.DepthStencilAccess.IsDepthWrite() ? VK_IMAGE_LAYOUT_ATTACHMENT_OPTIMAL : VK_IMAGE_LAYOUT_READ_ONLY_OPTIMAL;
+		StencilReference.stencilLayout = Initializer.DepthStencilAccess.IsStencilWrite() ? VK_IMAGE_LAYOUT_ATTACHMENT_OPTIMAL : VK_IMAGE_LAYOUT_READ_ONLY_OPTIMAL;
 
 		FullHashInfo.LoadOps[MaxSimultaneousRenderTargets] = CurrDesc.loadOp;
 		FullHashInfo.LoadOps[MaxSimultaneousRenderTargets + 1] = CurrDesc.stencilLoadOp;
diff --git a/Engine/Source/Runtime/VulkanRHI/Private/VulkanRenderpass.h b/Engine/Source/Runtime/VulkanRHI/Private/VulkanRenderpass.h
--- a/Engine/Source/Runtime/VulkanRHI/Private/VulkanRenderpass.h
+++ b/Engine/Source/Runtime/VulkanRHI/Private/VulkanRenderpass.h
@@ -436,10 +436,32 @@ public:
 			}
 		}
 
-
+		uint32_t DepthInputAttachment = VK_ATTACHMENT_UNUSED;
+		VkImageLayout DepthInputAttachmentLayout = VK_IMAGE_LAYOUT_UNDEFINED;
+		VkImageAspectFlags DepthInputAspectMask = 0;
 		if (bHasDepthStencilAttachmentReference)
 		{
 			DepthStencilAttachmentReference.SetDepthStencilAttachment(*RTLayout.GetDepthAttachmentReference(), RTLayout.GetStencilAttachmentReference(), 0, Device.SupportsParallelRendering());
+
+			if (bDepthReadSubpass || bDeferredShadingSubpass)
+			{
+				DepthStencilAttachment.attachment = RTLayout.GetDepthAttachmentReference()->attachment;
+				DepthStencilAttachment.SetAspect(VK_IMAGE_ASPECT_DEPTH_BIT);	// @todo?
+
+				if (RTLayout.GetDepthAttachmentReference()->layout == VK_IMAGE_LAYOUT_READ_ONLY_OPTIMAL)
+				{
+					// Depth is read only and is expected to be sampled as a regular texture
+					DepthStencilAttachment.layout = VK_IMAGE_LAYOUT_READ_ONLY_OPTIMAL;
+				}
+				else
+				{
+					DepthStencilAttachment.layout = VK_IMAGE_LAYOUT_DEPTH_STENCIL_READ_ONLY_OPTIMAL;
+
+					DepthInputAttachment = DepthStencilAttachment.attachment;
+					DepthInputAttachmentLayout = VK_IMAGE_LAYOUT_DEPTH_STENCIL_READ_ONLY_OPTIMAL;
+					DepthInputAspectMask = VK_IMAGE_ASPECT_DEPTH_BIT;
+				}
+			}
 		}
 
 		// main sub-pass
@@ -467,7 +489,6 @@ public:
 		// Color write and depth read sub-pass
 		if (bDepthReadSubpass)
 		{
-			DepthStencilAttachmentOG.SetAttachment(*RTLayout.GetDepthAttachmentReference(), VK_IMAGE_ASPECT_DEPTH_BIT);
 			TSubpassDescriptionClass& SubpassDesc = SubpassDescriptions[NumSubpasses++];
 
 			SubpassDesc.SetColorAttachments(ColorAttachmentReferences);
@@ -476,12 +497,12 @@ public:
 			check(RTLayout.GetDepthAttachmentReference());
 
 			// Depth as Input0
-			InputAttachments1[0].SetAttachment(DepthStencilAttachmentOG, VK_IMAGE_ASPECT_DEPTH_BIT);
-			InputAttachments1[0].layout = VK_IMAGE_LAYOUT_DEPTH_STENCIL_READ_ONLY_OPTIMAL;
-
+			InputAttachments1[0].attachment = DepthInputAttachment;
+			InputAttachments1[0].layout = DepthInputAttachmentLayout;
+			InputAttachments1[0].SetAspect(DepthInputAspectMask);
 			SubpassDesc.SetInputAttachments(InputAttachments1, InputAttachment1Count);
 			// depth attachment is same as input attachment
-			SubpassDesc.SetDepthStencilAttachment(InputAttachments1);
+			SubpassDesc.SetDepthStencilAttachment(&DepthStencilAttachment);
 
 			if (bApplyFragmentShadingRate)
 			{
@@ -502,11 +523,6 @@ public:
 		// Two subpasses for deferred shading
 		if (bDeferredShadingSubpass)
 		{
-			// both sub-passes only test DepthStencil
-			DepthStencilAttachment.attachment = RTLayout.GetDepthAttachmentReference()->attachment;
-			DepthStencilAttachment.layout = VK_IMAGE_LAYOUT_DEPTH_STENCIL_READ_ONLY_OPTIMAL;
-			DepthStencilAttachment.SetAspect(VK_IMAGE_ASPECT_DEPTH_BIT);	// @todo?
-
 			//const VkAttachmentReference* ColorRef = RTLayout.GetColorAttachmentReferences();
 			//uint32 NumColorAttachments = RTLayout.GetNumColorAttachments();
 			//check(RTLayout.GetNumColorAttachments() == 5); //current layout is SceneColor, GBufferA/B/C/D
@@ -516,7 +532,10 @@ public:
 				TSubpassDescriptionClass& SubpassDesc = SubpassDescriptions[NumSubpasses++];
 				SubpassDesc.SetColorAttachments(ColorAttachmentReferences);
 				SubpassDesc.SetDepthStencilAttachment(&DepthStencilAttachment);
-				SubpassDesc.SetInputAttachments(&DepthStencilAttachment, 1);
+				InputAttachments1[0].attachment = DepthInputAttachment;
+				InputAttachments1[0].layout = DepthInputAttachmentLayout;
+				InputAttachments1[0].SetAspect(DepthInputAspectMask);
+				SubpassDesc.SetInputAttachments(InputAttachments1, InputAttachment1Count);
 
 				if (bApplyFragmentShadingRate)
 				{
@@ -542,9 +561,9 @@ public:
 				SubpassDesc.SetDepthStencilAttachment(&DepthStencilAttachment);
 
 				// Depth as Input0
-				InputAttachments2[0].attachment = DepthStencilAttachment.attachment;
-				InputAttachments2[0].layout = DepthStencilAttachment.layout;
-				InputAttachments2[0].SetAspect(VK_IMAGE_ASPECT_DEPTH_BIT);
+				InputAttachments2[0].attachment = DepthInputAttachment;
+				InputAttachments2[0].layout = DepthInputAttachmentLayout;
+				InputAttachments2[0].SetAspect(DepthInputAspectMask);
 
 				// SceneColor write only
 				InputAttachments2[1].attachment = VK_ATTACHMENT_UNUSED;
@@ -677,7 +696,6 @@ private:
 	// Color write and depth read sub-pass
 	static const uint32 InputAttachment1Count = 1;
 	TAttachmentReferenceClass InputAttachments1[InputAttachment1Count];
-	TAttachmentReferenceClass DepthStencilAttachmentOG;
 
 	// Two subpasses for deferred shading
 	TAttachmentReferenceClass InputAttachments2[MaxSimultaneousRenderTargets + 1];
