Enable Android Vulkan SM5 Ray Tracing

diff --git a/Engine/Shaders/Private/LocalVertexFactory.ush b/Engine/Shaders/Private/LocalVertexFactory.ush
--- a/Engine/Shaders/Private/LocalVertexFactory.ush
+++ b/Engine/Shaders/Private/LocalVertexFactory.ush
@@ -231,13 +231,11 @@ FVertexFactoryInput LoadVertexFactoryInputForDynamicUpdate(uint TriangleIndex, i
 	}
 	else
 	{
-	#if  MANUAL_VERTEX_FETCH
 		Input.VertexId = TriangleIndex * 3 + VertexIndex;
 		uint VertexOffset = LocalVF.VertexFetch_Parameters[VF_VertexOffset] + Input.VertexId * 3;
 		Input.Position.x = LocalVF.VertexFetch_PositionBuffer[VertexOffset + 0];
 		Input.Position.y = LocalVF.VertexFetch_PositionBuffer[VertexOffset + 1];
 		Input.Position.z = LocalVF.VertexFetch_PositionBuffer[VertexOffset + 2];
-	#endif
 	}
 
 	VF_GPUSCENE_SET_INPUT_FOR_RT(Input, PrimitiveId, 0U);
diff --git a/Engine/Shaders/Private/RayTracing/TraceRayInlineVulkan.ush b/Engine/Shaders/Private/RayTracing/TraceRayInlineVulkan.ush
--- a/Engine/Shaders/Private/RayTracing/TraceRayInlineVulkan.ush
+++ b/Engine/Shaders/Private/RayTracing/TraceRayInlineVulkan.ush
@@ -14,14 +14,30 @@
 
 #define PLATFORM_SUPPORTS_INLINE_RAY_TRACING_TRIANGLE_ATTRIBUTES 1
 
-uint3 VulkanLoadIndices16Bit(uint64_t SourceBufferBaseAddress, uint OffsetInBytes)
+#if EMULATE_VKDEVICEADRESS
+inline uint2 AddEmulated64(uint2 a, uint b)
+{
+    uint2 c;
+    c.x = a.x + b;
+    c.y = a.y + (c.x < a.x); // Add with carry.
+    return c;
+}
+
+#define VK_ADDR_TYPE uint2
+#define VK_RAW_BUFFER_LOAD(Type, Addr, Offs) vk::RawBufferLoad2<Type>(AddEmulated64(Addr,Offs))
+#else
+#define VK_ADDR_TYPE uint64_t
+#define VK_RAW_BUFFER_LOAD(Type, Addr, Offs) vk::RawBufferLoad<Type>(Addr + Offs)
+#endif 
+
+uint3 VulkanLoadIndices16Bit(VK_ADDR_TYPE SourceBufferBaseAddress, uint OffsetInBytes)
 {
 	uint3 Result;
 
 	// ByteAddressBuffer loads must be aligned to DWORD boundary.
 	// We can load 2 DWORDs (4 SHORTs) at a time and extract 3 SHORTs (first 3 or second 3).
 	uint AlignedOffsetInBytes = OffsetInBytes & (~3);
-	const uint2 PackedIndices = vk::RawBufferLoad<uint2>(SourceBufferBaseAddress + AlignedOffsetInBytes);
+	const uint2 PackedIndices = VK_RAW_BUFFER_LOAD(uint2, SourceBufferBaseAddress, AlignedOffsetInBytes);
 
 	if (AlignedOffsetInBytes == OffsetInBytes)
 	{
@@ -41,15 +57,15 @@ uint3 VulkanLoadIndices16Bit(uint64_t SourceBufferBaseAddress, uint OffsetInByte
 	return Result;
 }
 
-uint3 VulkanLoadIndices32Bit(uint64_t SourceBufferBaseAddress, uint OffsetInBytes)
+uint3 VulkanLoadIndices32Bit(VK_ADDR_TYPE SourceBufferBaseAddress, uint OffsetInBytes)
 {
-	return vk::RawBufferLoad<uint3>(SourceBufferBaseAddress + OffsetInBytes);
+	return VK_RAW_BUFFER_LOAD(uint3, SourceBufferBaseAddress, OffsetInBytes);
 }
 
-float3 VulkanLoadVertexPositionFloat3(uint64_t SourceBufferBaseAddress, uint Index, uint StrideInBytes)
+float3 VulkanLoadVertexPositionFloat3(VK_ADDR_TYPE SourceBufferBaseAddress, uint Index, uint StrideInBytes)
 {
 	uint OffsetInBytes = Index * StrideInBytes;
-	return asfloat(vk::RawBufferLoad<uint3>(SourceBufferBaseAddress + OffsetInBytes));
+	return asfloat(VK_RAW_BUFFER_LOAD(uint3, SourceBufferBaseAddress, OffsetInBytes));
 }
 
 // Counterpart of FVulkanRayTracingGeometryParameters in VulkanRayTracing.cpp
@@ -57,8 +73,8 @@ struct FVulkanRayTracingGeometryParameters
 {
 	uint Config;
 	uint IndexBufferOffsetInBytes;
-	uint64_t IndexBuffer;
-	uint64_t VertexBuffer;
+	VK_ADDR_TYPE IndexBuffer;
+	VK_ADDR_TYPE VertexBuffer;
 };
 
 // Create a generic alias for the Vulkan-specific metadata type
@@ -82,8 +98,6 @@ FTriangleBaseAttributes LoadInlineRayTracingTriangleAttributes(
 	const uint IndexBufferStride = Params.Config & 0xFF;
 	const uint BaseIndex = PrimitiveIndex * 3;
 
-	uint64_t IndexBuffer = Params.IndexBuffer;
-
 	if (IndexBufferStride == 0)
 	{
 		// Non-indexed geometry (implicit triangle list indices)
@@ -91,21 +105,20 @@ FTriangleBaseAttributes LoadInlineRayTracingTriangleAttributes(
 	}
 	else if (IndexBufferStride == 2)
 	{
-		Result.Indices = VulkanLoadIndices16Bit(IndexBuffer, IndexBufferOffsetInBytes + BaseIndex * IndexBufferStride);
+		Result.Indices = VulkanLoadIndices16Bit(Params.IndexBuffer, IndexBufferOffsetInBytes + BaseIndex * IndexBufferStride);
 	}
 	else
 	{
-		Result.Indices = VulkanLoadIndices32Bit(IndexBuffer, IndexBufferOffsetInBytes + BaseIndex * IndexBufferStride);
+		Result.Indices = VulkanLoadIndices32Bit(Params.IndexBuffer, IndexBufferOffsetInBytes + BaseIndex * IndexBufferStride);
 	}
 
 	// Fetch vertex positions (in local space)
 	const uint VertexStride = (Params.Config >> 8) & 0xFF;
 
 	// #dxr_todo: UE-72160 handle various vertex formats, for now only supporting float3
-	uint64_t VertexBuffer = Params.VertexBuffer;
-	Result.LocalPositions[0] = VulkanLoadVertexPositionFloat3(VertexBuffer, Result.Indices[0], VertexStride);
-	Result.LocalPositions[1] = VulkanLoadVertexPositionFloat3(VertexBuffer, Result.Indices[1], VertexStride);
-	Result.LocalPositions[2] = VulkanLoadVertexPositionFloat3(VertexBuffer, Result.Indices[2], VertexStride);
+	Result.LocalPositions[0] = VulkanLoadVertexPositionFloat3(Params.VertexBuffer, Result.Indices[0], VertexStride);
+	Result.LocalPositions[1] = VulkanLoadVertexPositionFloat3(Params.VertexBuffer, Result.Indices[1], VertexStride);
+	Result.LocalPositions[2] = VulkanLoadVertexPositionFloat3(Params.VertexBuffer, Result.Indices[2], VertexStride);
 
 	return Result;
 }
diff --git a/Engine/Source/Developer/Android/AndroidTargetPlatform/Private/AndroidTargetPlatform.h b/Engine/Source/Developer/Android/AndroidTargetPlatform/Private/AndroidTargetPlatform.h
--- a/Engine/Source/Developer/Android/AndroidTargetPlatform/Private/AndroidTargetPlatform.h
+++ b/Engine/Source/Developer/Android/AndroidTargetPlatform/Private/AndroidTargetPlatform.h
@@ -155,6 +155,14 @@ public:
 		InBoolKeys.Add(TEXT("bSaveSymbols")); InStringKeys.Add(TEXT("NDKAPILevel"));
 	}
 
+	virtual bool UsesRayTracing() const override
+	{
+		bool bEnableRayTracing = false;
+		GConfig->GetBool(TEXT("/Script/AndroidRuntimeSettings.AndroidRuntimeSettings"), TEXT("bEnableRayTracing"), bEnableRayTracing, GEngineIni);
+
+		return bEnableRayTracing;
+	}
+
 	virtual bool ShouldExpandTo32Bit(const uint16* Indices, const int32 NumIndices) const override;
 	//~ End ITargetPlatform Interface
 
diff --git a/Engine/Source/Developer/VulkanShaderFormat/Private/VulkanShaderCompiler.cpp b/Engine/Source/Developer/VulkanShaderFormat/Private/VulkanShaderCompiler.cpp
--- a/Engine/Source/Developer/VulkanShaderFormat/Private/VulkanShaderCompiler.cpp
+++ b/Engine/Source/Developer/VulkanShaderFormat/Private/VulkanShaderCompiler.cpp
@@ -1856,7 +1856,7 @@ static bool BuildShaderOutputFromSpirv(
 		check(SpvResult == SPV_REFLECT_RESULT_SUCCESS);
 
 		const FString ResourceName(ANSI_TO_TCHAR(Binding->name));
-		CCHeaderWriter.WriteSRV(*ResourceName, UAVBindings++);
+		CCHeaderWriter.WriteSRV(*ResourceName, SRVBindings++);
 
 		Spirv.ReflectionInfo.Add(FVulkanSpirv::FEntry(ResourceName, BindingIndex));
 	}
@@ -1953,7 +1953,7 @@ static bool BuildShaderOutputFromSpirv(
 	// We have to strip out most debug instructions (except OpName) for Vulkan mobile
 	if (bStripReflect)
 	{
-		const char* OptArgs[] = { "--strip-reflect" };
+		const char* OptArgs[] = { "--strip-reflect", "-O"};
 		if (!CompilerContext.OptimizeSpirv(Spirv.Data, OptArgs, UE_ARRAY_COUNT(OptArgs)))
 		{
 			UE_LOG(LogVulkanShaderCompiler, Error, TEXT("Failed to strip debug instructions from SPIR-V module"));
@@ -2625,6 +2625,12 @@ void DoCompileVulkanShader(const FShaderCompilerInput& Input, FShaderCompilerOut
 	AdditionalDefines.SetDefine(TEXT("VULKAN_BINDLESS_SAMPLER_ARRAY_PREFIX"), VulkanBindless::kBindlessSamplerArrayPrefix);
 	AdditionalDefines.SetDefine(TEXT("VULKAN_BINDLESS_RESOURCE_ARRAY_PREFIX"), VulkanBindless::kBindlessResourceArrayPrefix);
 
+	if (IsAndroidShaderFormat(Input.ShaderFormat))
+	{
+		// On most Android devices uint64_t is unsupported so we emulate as 2 uint32_t's 
+		AdditionalDefines.SetDefine(TEXT("EMULATE_VKDEVICEADRESS"), 1);
+	}
+
 	const double StartPreprocessTime = FPlatformTime::Seconds();
 
 	// Preprocess the shader.
diff --git a/Engine/Source/Programs/UnrealBuildTool/Platform/Android/UEBuildAndroid.cs b/Engine/Source/Programs/UnrealBuildTool/Platform/Android/UEBuildAndroid.cs
--- a/Engine/Source/Programs/UnrealBuildTool/Platform/Android/UEBuildAndroid.cs
+++ b/Engine/Source/Programs/UnrealBuildTool/Platform/Android/UEBuildAndroid.cs
@@ -45,6 +45,12 @@ namespace UnrealBuildTool
 		/// </summary>
 		[CommandLine("-EnableMinUBSan")]
 		public bool bEnableMinimalUndefinedBehaviorSanitizer = false;
+
+		/// <summary>
+		/// Enables runtime ray tracing support.
+		/// </summary>
+		[ConfigFile(ConfigHierarchyType.Engine, "/Script/AndroidTargetPlatform.AndroidTargetSettings")]
+		public bool bEnableRayTracing = false;
 	}
 
 	/// <summary>
@@ -97,6 +103,11 @@ namespace UnrealBuildTool
 			get { return Inner.bEnableMinimalUndefinedBehaviorSanitizer; }
 		}
 
+		public bool bEnableRayTracing
+		{
+			get { return Inner.bEnableRayTracing; }
+		}
+
 		public AndroidTargetRules TargetRules
 		{
 			get { return Inner; }
@@ -506,6 +517,12 @@ namespace UnrealBuildTool
 			CompileEnvironment.Definitions.Add("WITH_EDITOR=0");
 			CompileEnvironment.Definitions.Add("USE_NULL_RHI=0");
 
+			if (Target.AndroidPlatform.bEnableRayTracing)
+			{
+				Logger.LogInformation("Compiling with ray tracing enabled");
+				CompileEnvironment.Definitions.Add("RHI_RAYTRACING=1");
+			}
+
 			if (Target.bPGOOptimize || Target.bPGOProfile)
 			{
 				Logger.LogInformation("PGO {PgoType} build", Target.bPGOOptimize ? "optimize" : "profile");
diff --git a/Engine/Source/Runtime/Core/Public/Android/AndroidPlatform.h b/Engine/Source/Runtime/Core/Public/Android/AndroidPlatform.h
--- a/Engine/Source/Runtime/Core/Public/Android/AndroidPlatform.h
+++ b/Engine/Source/Runtime/Core/Public/Android/AndroidPlatform.h
@@ -58,7 +58,7 @@ typedef FAndroidTypes FPlatformTypes;
 #define PLATFORM_UI_NEEDS_TOOLTIPS					0
 #define PLATFORM_UI_NEEDS_FOCUS_OUTLINES			0
 #define PLATFORM_SUPPORTS_EARLY_MOVIE_PLAYBACK		1 // movies will start before engine is initialized
-#define PLATFORM_SUPPORTS_GEOMETRY_SHADERS			0
+#define PLATFORM_SUPPORTS_GEOMETRY_SHADERS			(RHI_RAYTRACING) // Currently required due to the way we offset ShaderStage::EStage
 #define PLATFORM_SUPPORTS_VIRTUAL_TEXTURE_STREAMING	1
 #define PLATFORM_USE_ANSI_POSIX_MALLOC				1
 #define PLATFORM_USE_MINIMAL_HANG_DETECTION			1
diff --git a/Engine/Source/Runtime/RenderCore/Private/RenderGraphPrivate.cpp b/Engine/Source/Runtime/RenderCore/Private/RenderGraphPrivate.cpp
--- a/Engine/Source/Runtime/RenderCore/Private/RenderGraphPrivate.cpp
+++ b/Engine/Source/Runtime/RenderCore/Private/RenderGraphPrivate.cpp
@@ -634,7 +634,9 @@ bool IsParallelExecuteEnabled()
 		&& !IsImmediateMode()
 		&& !GRDGDebug
 		&& !GRDGTransitionLog
+		&& !IsMobilePlatform(GMaxRHIShaderPlatform)
 		&& !IsOpenGLPlatform(GMaxRHIShaderPlatform)
+		&& !IsVulkanMobileSM5Platform(GMaxRHIShaderPlatform)
 		&& GRHISupportsMultithreadedShaderCreation
 #if WITH_DUMPGPU
 		&& !IsDumpingRDGResources()
@@ -651,7 +653,9 @@ bool IsParallelSetupEnabled()
 		&& !IsImmediateMode()
 		&& !GRDGDebug
 		&& !GRDGTransitionLog
+		&& !IsMobilePlatform(GMaxRHIShaderPlatform)
 		&& !IsOpenGLPlatform(GMaxRHIShaderPlatform)
+		&& !IsVulkanMobileSM5Platform(GMaxRHIShaderPlatform)
 		&& GRHISupportsMultithreadedShaderCreation
 #if WITH_DUMPGPU
 		&& !IsDumpingRDGResources()
diff --git a/Engine/Source/Runtime/RenderCore/Private/RenderUtils.cpp b/Engine/Source/Runtime/RenderCore/Private/RenderUtils.cpp
--- a/Engine/Source/Runtime/RenderCore/Private/RenderUtils.cpp
+++ b/Engine/Source/Runtime/RenderCore/Private/RenderUtils.cpp
@@ -1416,7 +1416,6 @@ bool UseGPUScene(const FStaticShaderPlatform Platform, const FStaticFeatureLevel
 	return FeatureLevel >= ERHIFeatureLevel::SM5
 		//@todo - support GPU Scene management compute shaders on these platforms to get dynamic instancing speedups on the Rendering Thread and RHI Thread
 		&& !IsOpenGLPlatform(Platform)
-		&& !IsVulkanMobileSM5Platform(Platform)
 		&& !IsMetalMobileSM5Platform(Platform)
 		// we only check DDSPI for platforms that have been read in - IsValid() can go away once ALL platforms are converted over to this system
 		&& (!FDataDrivenShaderPlatformInfo::IsValid(Platform) || FDataDrivenShaderPlatformInfo::GetSupportsGPUScene(Platform));
@@ -1713,4 +1712,4 @@ namespace Strata
 	{
 		return CVarSubstrateDebugAdvancedVisualizationShaders.GetValueOnAnyThread() > 0;
 	}
-}
\ No newline at end of file
+}
diff --git a/Engine/Source/Runtime/VulkanRHI/Private/Android/VulkanAndroidPlatform.cpp b/Engine/Source/Runtime/VulkanRHI/Private/Android/VulkanAndroidPlatform.cpp
--- a/Engine/Source/Runtime/VulkanRHI/Private/Android/VulkanAndroidPlatform.cpp
+++ b/Engine/Source/Runtime/VulkanRHI/Private/Android/VulkanAndroidPlatform.cpp
@@ -1,7 +1,7 @@
 // Copyright Epic Games, Inc. All Rights Reserved.
 
 #include "VulkanAndroidPlatform.h"
-#include "../VulkanRHIPrivate.h"
+#include "../VulkanRayTracing.h"
 #include "../VulkanPipeline.h"
 #include "../VulkanRenderpass.h"
 #include <dlfcn.h>
@@ -301,6 +301,14 @@ bool FVulkanAndroidPlatform::LoadVulkanInstanceFunctions(VkInstance inInstance)
 	ENUM_VK_ENTRYPOINTS_PLATFORM_INSTANCE(GETINSTANCE_VK_ENTRYPOINTS);
 	ENUM_VK_ENTRYPOINTS_PLATFORM_INSTANCE(CHECK_VK_ENTRYPOINTS);
 
+#if VULKAN_RHI_RAYTRACING
+	const bool bFoundRayTracingEntries = FVulkanRayTracingPlatform::CheckVulkanInstanceFunctions(inInstance);
+	if (!bFoundRayTracingEntries)
+	{
+		UE_LOG(LogVulkanRHI, Warning, TEXT("Vulkan RHI ray tracing is enabled, but failed to load instance functions."));
+	}
+#endif
+
 	if (!bFoundAllEntryPoints)
 	{
 		return false;
diff --git a/Engine/Source/Runtime/VulkanRHI/Private/Android/VulkanAndroidPlatform.h b/Engine/Source/Runtime/VulkanRHI/Private/Android/VulkanAndroidPlatform.h
--- a/Engine/Source/Runtime/VulkanRHI/Private/Android/VulkanAndroidPlatform.h
+++ b/Engine/Source/Runtime/VulkanRHI/Private/Android/VulkanAndroidPlatform.h
@@ -19,11 +19,12 @@
 #define VULKAN_SUPPORTS_DEDICATED_ALLOCATION		0
 #define VULKAN_SUPPORTS_ASTC_DECODE_MODE			1
 #define VULKAN_SUPPORTS_NV_DIAGNOSTIC_CHECKPOINT	0
-#define VULKAN_RHI_RAYTRACING						0
+#define VULKAN_RHI_RAYTRACING						(RHI_RAYTRACING)
+#define VULKAN_SUPPORTS_SCALAR_BLOCK_LAYOUT			(VULKAN_RHI_RAYTRACING)
 #define VULKAN_SUPPORTS_TRANSIENT_RESOURCE_ALLOCATOR 0
 #define VULKAN_SUPPORTS_DRIVER_PROPERTIES			0
 #define VULKAN_SUPPORTS_MULTIVIEW					1
-#define VULKAN_SUPPORTS_DESCRIPTOR_INDEXING			0
+#define VULKAN_SUPPORTS_DESCRIPTOR_INDEXING			(VULKAN_RHI_RAYTRACING)
 
 // crashing during callback setup on Android, code will fallback to VK_EXT_debug_report instead
 #define VULKAN_SUPPORTS_DEBUG_UTILS					0
@@ -33,6 +34,11 @@
 //	VK_IMAGE_LAYOUT_DEPTH_ATTACHMENT_STENCIL_READ_ONLY_OPTIMAL = 1000117001,
 #define VULKAN_USE_REAL_RENDERPASS_COMPATIBILITY	0
 
+#if VULKAN_RHI_RAYTRACING
+#define UE_VK_API_VERSION							VK_API_VERSION_1_2
+#else
+#define UE_VK_API_VERSION							VK_API_VERSION_1_1
+#endif // VULKAN_RHI_RAYTRACING
 
 #define ENUM_VK_ENTRYPOINTS_PLATFORM_BASE(EnumMacro)
 
diff --git a/Engine/Source/Runtime/VulkanRHI/Private/VulkanRHI.cpp b/Engine/Source/Runtime/VulkanRHI/Private/VulkanRHI.cpp
--- a/Engine/Source/Runtime/VulkanRHI/Private/VulkanRHI.cpp
+++ b/Engine/Source/Runtime/VulkanRHI/Private/VulkanRHI.cpp
@@ -895,7 +895,10 @@ void FVulkanDynamicRHI::InitInstance()
 			GRHISupportsInlineRayTracing = RHISupportsInlineRayTracing(GMaxRHIShaderPlatform) && Device->GetOptionalExtensions().HasRayQuery;
 
 			GRHIRayTracingAccelerationStructureAlignment = 256; // TODO (currently handled by FVulkanAccelerationStructureBuffer)
-			GRHIRayTracingScratchBufferAlignment = Device->GetOptionalExtensionProperties().AccelerationStructureProps.minAccelerationStructureScratchOffsetAlignment;
+			//Some devices have 64 for min AS offset alignment meanwhile engine AS alignment is 256. hence using round up value
+			GRHIRayTracingScratchBufferAlignment = FPlatformMath::Max<uint32>(GRHIRayTracingAccelerationStructureAlignment, 
+													Device->GetOptionalExtensionProperties().AccelerationStructureProps.minAccelerationStructureScratchOffsetAlignment);
+
 			GRHIRayTracingInstanceDescriptorSize = uint32(sizeof(VkAccelerationStructureInstanceKHR));
 		}
 #endif
diff --git a/Engine/Source/Runtime/VulkanRHI/Private/VulkanRHIPrivate.h b/Engine/Source/Runtime/VulkanRHI/Private/VulkanRHIPrivate.h
--- a/Engine/Source/Runtime/VulkanRHI/Private/VulkanRHIPrivate.h
+++ b/Engine/Source/Runtime/VulkanRHI/Private/VulkanRHIPrivate.h
@@ -851,7 +851,11 @@ namespace VulkanRHI
 			VK_ACCESS_TRANSFER_READ_BIT |
 			VK_ACCESS_TRANSFER_WRITE_BIT |
 			VK_ACCESS_HOST_READ_BIT |
-			VK_ACCESS_HOST_WRITE_BIT;
+			VK_ACCESS_HOST_WRITE_BIT |
+			VK_ACCESS_ACCELERATION_STRUCTURE_READ_BIT_KHR |
+			VK_ACCESS_ACCELERATION_STRUCTURE_WRITE_BIT_KHR | 
+			VK_ACCESS_FRAGMENT_DENSITY_MAP_READ_BIT_EXT |
+			VK_ACCESS_FRAGMENT_SHADING_RATE_ATTACHMENT_READ_BIT_KHR;
 		Barrier.dstAccessMask = VK_ACCESS_INDIRECT_COMMAND_READ_BIT |
 			VK_ACCESS_INDEX_READ_BIT |
 			VK_ACCESS_VERTEX_ATTRIBUTE_READ_BIT |
@@ -866,7 +870,11 @@ namespace VulkanRHI
 			VK_ACCESS_TRANSFER_READ_BIT |
 			VK_ACCESS_TRANSFER_WRITE_BIT |
 			VK_ACCESS_HOST_READ_BIT |
-			VK_ACCESS_HOST_WRITE_BIT;
+			VK_ACCESS_HOST_WRITE_BIT |
+			VK_ACCESS_ACCELERATION_STRUCTURE_READ_BIT_KHR |
+			VK_ACCESS_ACCELERATION_STRUCTURE_WRITE_BIT_KHR |
+			VK_ACCESS_FRAGMENT_DENSITY_MAP_READ_BIT_EXT |
+			VK_ACCESS_FRAGMENT_SHADING_RATE_ATTACHMENT_READ_BIT_KHR;
 		VulkanRHI::vkCmdPipelineBarrier(CmdBuffer, VK_PIPELINE_STAGE_ALL_COMMANDS_BIT, VK_PIPELINE_STAGE_ALL_COMMANDS_BIT, 0, 1, &Barrier, 0, nullptr, 0, nullptr);
 	}
 
@@ -883,7 +891,8 @@ namespace VulkanRHI
 
 inline bool UseVulkanDescriptorCache()
 {
-	return (PLATFORM_ANDROID)|| GMaxRHIFeatureLevel <= ERHIFeatureLevel::ES3_1;
+	// Descriptor cache path for WriteAccelerationStructure() is not implemented, so disable if RT is enabled
+	return ((PLATFORM_ANDROID) && !(VULKAN_RHI_RAYTRACING)) || GMaxRHIFeatureLevel <= ERHIFeatureLevel::ES3_1;
 }
 
 inline bool ValidateShadingRateDataType()
diff --git a/Engine/Source/Runtime/VulkanRHI/Private/VulkanRayTracing.cpp b/Engine/Source/Runtime/VulkanRHI/Private/VulkanRayTracing.cpp
--- a/Engine/Source/Runtime/VulkanRHI/Private/VulkanRayTracing.cpp
+++ b/Engine/Source/Runtime/VulkanRHI/Private/VulkanRayTracing.cpp
@@ -28,6 +28,39 @@ static FAutoConsoleVariableRef CVarVulkanRayTracingMaxBatchedCompaction(
 	ECVF_ReadOnly
 );
 
+static bool GVulkanRayTracingTLASPreferFastTrace = true;
+static FAutoConsoleVariableRef CVarVulkanRayTracingTLASPreferFastTraceTLAS(
+	TEXT("r.Vulkan.RayTracing.TLASPreferFastTraceTLAS"),
+	GVulkanRayTracingTLASPreferFastTrace,
+	TEXT("Prefer fast trace for TLAS build (default = true)\n"),
+	ECVF_ReadOnly
+);
+
+// Ray tracing stat counters
+
+DECLARE_STATS_GROUP(TEXT("Vulkan: Ray Tracing"), STATGROUP_VulkanRayTracing, STATCAT_Advanced);
+
+DECLARE_DWORD_ACCUMULATOR_STAT(TEXT("Created pipelines (total)"), STAT_VulkanRayTracingCreatedPipelines, STATGROUP_VulkanRayTracing);
+DECLARE_DWORD_ACCUMULATOR_STAT(TEXT("Compiled shaders (total)"), STAT_VulkanRayTracingCompiledShaders, STATGROUP_VulkanRayTracing);
+
+DECLARE_DWORD_ACCUMULATOR_STAT(TEXT("Allocated bottom level acceleration structures"), STAT_VulkanRayTracingAllocatedBLAS, STATGROUP_VulkanRayTracing);
+DECLARE_DWORD_ACCUMULATOR_STAT(TEXT("Allocated top level acceleration structures"), STAT_VulkanRayTracingAllocatedTLAS, STATGROUP_VulkanRayTracing);
+DECLARE_DWORD_ACCUMULATOR_STAT(TEXT("Triangles in all BL acceleration structures"), STAT_VulkanRayTracingTrianglesBLAS, STATGROUP_VulkanRayTracing);
+
+DECLARE_DWORD_COUNTER_STAT(TEXT("Built BL AS (per frame)"), STAT_VulkanRayTracingBuiltBLAS, STATGROUP_VulkanRayTracing);
+DECLARE_DWORD_COUNTER_STAT(TEXT("Updated BL AS (per frame)"), STAT_VulkanRayTracingUpdatedBLAS, STATGROUP_VulkanRayTracing);
+DECLARE_DWORD_COUNTER_STAT(TEXT("Built TL AS (per frame)"), STAT_VulkanRayTracingBuiltTLAS, STATGROUP_VulkanRayTracing);
+
+DECLARE_MEMORY_STAT(TEXT("Total BL AS Memory"), STAT_VulkanRayTracingBLASMemory, STATGROUP_VulkanRayTracing);
+DECLARE_MEMORY_STAT(TEXT("Static BL AS Memory"), STAT_VulkanRayTracingStaticBLASMemory, STATGROUP_VulkanRayTracing);
+DECLARE_MEMORY_STAT(TEXT("Dynamic BL AS Memory"), STAT_VulkanRayTracingDynamicBLASMemory, STATGROUP_VulkanRayTracing);
+DECLARE_MEMORY_STAT(TEXT("TL AS Memory"), STAT_VulkanRayTracingTLASMemory, STATGROUP_VulkanRayTracing);
+DECLARE_MEMORY_STAT(TEXT("Total Used Video Memory"), STAT_VulkanRayTracingUsedVideoMemory, STATGROUP_VulkanRayTracing);
+
+DECLARE_CYCLE_STAT(TEXT("RTPSO Compile Shader"), STAT_RTPSO_CompileShader, STATGROUP_VulkanRayTracing);
+DECLARE_CYCLE_STAT(TEXT("RTPSO Create Pipeline"), STAT_RTPSO_CreatePipeline, STATGROUP_VulkanRayTracing);
+
+
 #if PLATFORM_WINDOWS
 #pragma warning(push)
 #pragma warning(disable : 4191) // warning C4191: 'type cast': unsafe conversion
@@ -328,6 +361,8 @@ FVulkanRayTracingGeometry::FVulkanRayTracingGeometry(ENoInit)
 FVulkanRayTracingGeometry::FVulkanRayTracingGeometry(const FRayTracingGeometryInitializer& InInitializer, FVulkanDevice* InDevice)
 	: FRHIRayTracingGeometry(InInitializer), Device(InDevice)
 {
+	INC_DWORD_STAT(STAT_VulkanRayTracingAllocatedBLAS);
+
 	uint32 IndexBufferStride = 0;
 	if (Initializer.IndexBuffer)
 	{
@@ -363,18 +398,47 @@ FVulkanRayTracingGeometry::FVulkanRayTracingGeometry(const FRayTracingGeometryIn
 	CreateInfo.type = VK_ACCELERATION_STRUCTURE_TYPE_BOTTOM_LEVEL_KHR;
 	VERIFYVULKANRESULT(VulkanDynamicAPI::vkCreateAccelerationStructureKHR(NativeDevice, &CreateInfo, VULKAN_CPU_ALLOCATOR, &Handle));
 	
+	INC_MEMORY_STAT_BY(STAT_VulkanRayTracingUsedVideoMemory, SizeInfo.ResultSize);
+	INC_MEMORY_STAT_BY(STAT_VulkanRayTracingBLASMemory, SizeInfo.ResultSize);
+	if (Initializer.bAllowUpdate)
+	{
+		INC_MEMORY_STAT_BY(STAT_VulkanRayTracingDynamicBLASMemory, SizeInfo.ResultSize);
+	}
+	else
+	{
+		INC_MEMORY_STAT_BY(STAT_VulkanRayTracingStaticBLASMemory, SizeInfo.ResultSize);
+	}
+
 	VkAccelerationStructureDeviceAddressInfoKHR DeviceAddressInfo;
 	ZeroVulkanStruct(DeviceAddressInfo, VK_STRUCTURE_TYPE_ACCELERATION_STRUCTURE_DEVICE_ADDRESS_INFO_KHR);
 	DeviceAddressInfo.accelerationStructure = Handle;
 	Address = VulkanDynamicAPI::vkGetAccelerationStructureDeviceAddressKHR(NativeDevice, &DeviceAddressInfo);
+
+	INC_DWORD_STAT_BY(STAT_VulkanRayTracingTrianglesBLAS, Initializer.TotalPrimitiveCount);
 }
 
 FVulkanRayTracingGeometry::~FVulkanRayTracingGeometry()
 {
+	DEC_DWORD_STAT(STAT_VulkanRayTracingAllocatedBLAS);
+	DEC_DWORD_STAT_BY(STAT_VulkanRayTracingTrianglesBLAS, Initializer.TotalPrimitiveCount);
 	if (Handle != VK_NULL_HANDLE)
 	{
+		DEC_MEMORY_STAT_BY(STAT_VulkanRayTracingUsedVideoMemory, AccelerationStructureBuffer->GetSize());
+		DEC_MEMORY_STAT_BY(STAT_VulkanRayTracingBLASMemory, AccelerationStructureBuffer->GetSize());
+
+		ERayTracingAccelerationStructureFlags BuildFlags = GetRayTracingAccelerationStructureBuildFlags(Initializer);
+		if (EnumHasAllFlags(BuildFlags, ERayTracingAccelerationStructureFlags::AllowUpdate))
+		{
+			DEC_MEMORY_STAT_BY(STAT_VulkanRayTracingDynamicBLASMemory, AccelerationStructureBuffer->GetSize());
+		}
+		else
+		{
+			DEC_MEMORY_STAT_BY(STAT_VulkanRayTracingStaticBLASMemory, AccelerationStructureBuffer->GetSize());
+		}
+
 		Device->GetDeferredDeletionQueue().EnqueueResource(VulkanRHI::FDeferredDeletionQueue2::EType::AccelerationStructure, Handle);
-	}
+	}	
+
 	RemoveCompactionRequest();
 }
 
@@ -419,6 +483,10 @@ void FVulkanRayTracingGeometry::CompactAccelerationStructure(FVulkanCmdBuffer& C
 		return;
 	}
 
+	DEC_MEMORY_STAT_BY(STAT_VulkanRayTracingUsedVideoMemory, AccelerationStructureBuffer->GetSize());
+	DEC_MEMORY_STAT_BY(STAT_VulkanRayTracingBLASMemory, AccelerationStructureBuffer->GetSize());
+	DEC_MEMORY_STAT_BY(STAT_VulkanRayTracingStaticBLASMemory, AccelerationStructureBuffer->GetSize());
+
 	// Move old AS into this temporary variable which gets released when this function returns	
 	TRefCountPtr<FVulkanResourceMultiBuffer> OldAccelerationStructure = AccelerationStructureBuffer;
 	VkAccelerationStructureKHR OldHandle = Handle;
@@ -427,6 +495,11 @@ void FVulkanRayTracingGeometry::CompactAccelerationStructure(FVulkanCmdBuffer& C
 	FRHIResourceCreateInfo BlasBufferCreateInfo(*DebugNameString);
 	AccelerationStructureBuffer = new FVulkanResourceMultiBuffer(Device, FRHIBufferDesc(InSizeAfterCompaction, 0, BUF_AccelerationStructure), BlasBufferCreateInfo);
 	
+
+	INC_MEMORY_STAT_BY(STAT_VulkanRayTracingUsedVideoMemory, AccelerationStructureBuffer->GetSize());
+	INC_MEMORY_STAT_BY(STAT_VulkanRayTracingBLASMemory, AccelerationStructureBuffer->GetSize());
+	INC_MEMORY_STAT_BY(STAT_VulkanRayTracingStaticBLASMemory, AccelerationStructureBuffer->GetSize());
+
 	VkDevice NativeDevice = Device->GetInstanceHandle();
 
 	VkAccelerationStructureCreateInfoKHR CreateInfo;
@@ -473,7 +546,9 @@ static void GetTLASBuildData(
 
 	BuildData.GeometryInfo.type = VK_ACCELERATION_STRUCTURE_TYPE_TOP_LEVEL_KHR;
 	BuildData.GeometryInfo.mode = VK_BUILD_ACCELERATION_STRUCTURE_MODE_BUILD_KHR;
-	BuildData.GeometryInfo.flags = VK_BUILD_ACCELERATION_STRUCTURE_PREFER_FAST_TRACE_BIT_KHR;
+	BuildData.GeometryInfo.flags = GVulkanRayTracingTLASPreferFastTrace ? 
+										VK_BUILD_ACCELERATION_STRUCTURE_PREFER_FAST_TRACE_BIT_KHR :
+										VK_BUILD_ACCELERATION_STRUCTURE_PREFER_FAST_BUILD_BIT_KHR;
 	BuildData.GeometryInfo.geometryCount = 1;
 	BuildData.GeometryInfo.pGeometries = &BuildData.Geometry;
 
@@ -536,6 +611,8 @@ struct FVulkanRayTracingGeometryParameters
 FVulkanRayTracingScene::FVulkanRayTracingScene(FRayTracingSceneInitializer2 InInitializer, FVulkanDevice* InDevice)
 	: Device(InDevice), Initializer(MoveTemp(InInitializer))
 {
+	INC_DWORD_STAT(STAT_VulkanRayTracingAllocatedTLAS);
+
 	const ERayTracingAccelerationStructureFlags BuildFlags = ERayTracingAccelerationStructureFlags::FastTrace; // #yuriy_todo: pass this in
 
 	SizeInfo = {};
@@ -567,6 +644,12 @@ FVulkanRayTracingScene::FVulkanRayTracingScene(FRayTracingSceneInitializer2 InIn
 
 FVulkanRayTracingScene::~FVulkanRayTracingScene()
 {
+	if (AccelerationStructureBuffer)
+	{
+		DEC_MEMORY_STAT_BY(STAT_VulkanRayTracingUsedVideoMemory, AccelerationStructureBuffer->GetSize());
+		DEC_MEMORY_STAT_BY(STAT_VulkanRayTracingTLASMemory, AccelerationStructureBuffer->GetSize());
+	}
+	DEC_DWORD_STAT(STAT_VulkanRayTracingAllocatedTLAS);
 }
 
 void FVulkanRayTracingScene::BindBuffer(FRHIBuffer* InBuffer, uint32 InBufferOffset)
@@ -575,8 +658,17 @@ void FVulkanRayTracingScene::BindBuffer(FRHIBuffer* InBuffer, uint32 InBufferOff
 
 	check(SizeInfo.ResultSize + InBufferOffset <= InBuffer->GetSize());
 	
+	if (AccelerationStructureBuffer)
+	{
+		DEC_MEMORY_STAT_BY(STAT_VulkanRayTracingUsedVideoMemory, AccelerationStructureBuffer->GetSize());
+		DEC_MEMORY_STAT_BY(STAT_VulkanRayTracingTLASMemory, AccelerationStructureBuffer->GetSize());
+	}
+
 	AccelerationStructureBuffer = ResourceCast(InBuffer);
 
+	INC_MEMORY_STAT_BY(STAT_VulkanRayTracingUsedVideoMemory, AccelerationStructureBuffer->GetSize());
+	INC_MEMORY_STAT_BY(STAT_VulkanRayTracingTLASMemory, AccelerationStructureBuffer->GetSize());
+
 	for (auto& Layer : Layers)
 	{
 		checkf(!Layer.View.IsValid(), TEXT("Binding multiple buffers is not currently supported."));
@@ -654,6 +746,8 @@ void FVulkanRayTracingScene::BuildAccelerationStructure(
 
 		pBuildRanges[LayerIndex] = &TLASBuildRangeInfo;
 
+		INC_DWORD_STAT(STAT_VulkanRayTracingBuiltTLAS);
+
 		InstanceBaseOffset += Initializer.NumNativeInstancesPerLayer[LayerIndex];
 	}
 
@@ -854,6 +948,15 @@ void FVulkanCommandListContext::RHIBuildAccelerationStructures(const TArrayView<
 		FVulkanRayTracingGeometry* const Geometry = ResourceCast(P.Geometry.GetReference());
 		const bool bIsUpdate = P.BuildMode == EAccelerationStructureBuildMode::Update;
 
+		if (bIsUpdate)
+		{
+			INC_DWORD_STAT(STAT_VulkanRayTracingUpdatedBLAS);
+		}
+		else
+		{
+			INC_DWORD_STAT(STAT_VulkanRayTracingBuiltBLAS);
+		}
+
 		uint64 ScratchBufferRequiredSize = bIsUpdate ? Geometry->SizeInfo.UpdateScratchSize : Geometry->SizeInfo.BuildScratchSize;
 		checkf(ScratchBufferRequiredSize + ScratchBufferOffset <= ScratchBufferSize,
 			TEXT("BLAS scratch buffer size is %ld bytes with offset %ld (%ld bytes available), but the build requires %lld bytes. "),
@@ -1110,6 +1213,9 @@ FVulkanRayTracingPipelineState::FVulkanRayTracingPipelineState(FVulkanDevice* co
 		VulkanRHI::vkUnmapMemory(InDevice->GetInstanceHandle(), Allocation.Memory);
 	};
 
+	INC_DWORD_STAT(STAT_VulkanRayTracingCreatedPipelines);
+	INC_DWORD_STAT_BY(STAT_VulkanRayTracingCompiledShaders, 1);
+
 	CopyHandlesToSBT(RayGenShaderBindingTable, 0);
 	CopyHandlesToSBT(MissShaderBindingTable, HandleSizeAligned);
 	CopyHandlesToSBT(HitShaderBindingTable, HandleSizeAligned * 2);
diff --git a/Engine/Source/Runtime/VulkanRHI/Public/VulkanCommon.h b/Engine/Source/Runtime/VulkanRHI/Public/VulkanCommon.h
--- a/Engine/Source/Runtime/VulkanRHI/Public/VulkanCommon.h
+++ b/Engine/Source/Runtime/VulkanRHI/Public/VulkanCommon.h
@@ -7,9 +7,10 @@
 #pragma once
 
 #include "RHIDefinitions.h"
+#include "Logging/LogMacros.h"
 
 #ifndef VULKAN_SUPPORTS_GEOMETRY_SHADERS
-	#define VULKAN_SUPPORTS_GEOMETRY_SHADERS					(!PLATFORM_ANDROID) && PLATFORM_SUPPORTS_GEOMETRY_SHADERS
+	#define VULKAN_SUPPORTS_GEOMETRY_SHADERS					PLATFORM_SUPPORTS_GEOMETRY_SHADERS
 #endif
 
 // This defines controls shader generation (so will cause a format rebuild)
