VulkanRHI: Fix InitAsStructuredBufferView calls that were applying the buffer offset twice.  When possible, remove Allocation and have a single offset parameter to make it clear it should include all offsets (allocation and view).

#rnx

diff --git a/Engine/Source/Runtime/VulkanRHI/Private/VulkanDescriptorSets.h b/Engine/Source/Runtime/VulkanRHI/Private/VulkanDescriptorSets.h
--- a/Engine/Source/Runtime/VulkanRHI/Private/VulkanDescriptorSets.h
+++ b/Engine/Source/Runtime/VulkanRHI/Private/VulkanDescriptorSets.h
@@ -882,14 +882,14 @@ public:
 		return NumWrites;
 	}
 
-	bool WriteUniformBuffer(uint32 DescriptorIndex, const FVulkanAllocation& Allocation, VkDeviceSize Offset, VkDeviceSize Range)
+	bool WriteUniformBuffer(uint32 DescriptorIndex, VkBuffer BufferHandle, uint32 HandleId, VkDeviceSize Offset, VkDeviceSize Range)
 	{
-		return WriteBuffer<VK_DESCRIPTOR_TYPE_UNIFORM_BUFFER>(DescriptorIndex, Allocation, Offset, Range);
+		return WriteBuffer<VK_DESCRIPTOR_TYPE_UNIFORM_BUFFER>(DescriptorIndex, BufferHandle, HandleId, Offset, Range);
 	}
 
-	bool WriteDynamicUniformBuffer(uint32 DescriptorIndex, const FVulkanAllocation& Allocation, VkDeviceSize Offset, VkDeviceSize Range, uint32 DynamicOffset)
+	bool WriteDynamicUniformBuffer(uint32 DescriptorIndex, VkBuffer BufferHandle, uint32 HandleId, VkDeviceSize Offset, VkDeviceSize Range, uint32 DynamicOffset)
 	{
-		return WriteBuffer<VK_DESCRIPTOR_TYPE_UNIFORM_BUFFER_DYNAMIC>(DescriptorIndex, Allocation, Offset, Range, DynamicOffset);
+		return WriteBuffer<VK_DESCRIPTOR_TYPE_UNIFORM_BUFFER_DYNAMIC>(DescriptorIndex, BufferHandle, HandleId, Offset, Range, DynamicOffset);
 	}
 
 	bool WriteSampler(uint32 DescriptorIndex, const FVulkanSamplerState& Sampler)
@@ -942,7 +942,7 @@ public:
 
 	bool WriteStorageBuffer(uint32 DescriptorIndex, const FVulkanView::FStructuredBufferView& View)
 	{
-		return WriteBuffer<VK_DESCRIPTOR_TYPE_STORAGE_BUFFER>(DescriptorIndex, View.Allocation, View.Offset, View.Size);
+		return WriteBuffer<VK_DESCRIPTOR_TYPE_STORAGE_BUFFER>(DescriptorIndex, View.Buffer, View.HandleId, View.Offset, View.Size);
 	}
 
 	bool WriteUniformTexelBuffer(uint32 DescriptorIndex, const FVulkanView::FTypedBufferView& View)
@@ -998,7 +998,7 @@ public:
 
 protected:
 	template <VkDescriptorType DescriptorType>
-	bool WriteBuffer(uint32 DescriptorIndex, const FVulkanAllocation& Allocation, VkDeviceSize Offset, VkDeviceSize Range, uint32 DynamicOffset = 0)
+	bool WriteBuffer(uint32 DescriptorIndex, VkBuffer BufferHandle, uint32 HandleId, VkDeviceSize Offset, VkDeviceSize Range, uint32 DynamicOffset = 0)
 	{
 		check(DescriptorIndex < NumWrites);
 		SetWritten(DescriptorIndex);		
@@ -1021,11 +1021,11 @@ protected:
 		if (UseVulkanDescriptorCache())
 		{
 			FVulkanHashableDescriptorInfo& HashableInfo = HashableDescriptorInfos[DescriptorIndex];
-			check(Allocation.HandleId > 0);
-			if (HashableInfo.Buffer.Id != Allocation.HandleId)
+			check(HandleId > 0);
+			if (HashableInfo.Buffer.Id != HandleId)
 			{
-				HashableInfo.Buffer.Id = Allocation.HandleId;
-				BufferInfo->buffer = (VkBuffer)Allocation.VulkanHandle;
+				HashableInfo.Buffer.Id = HandleId;
+				BufferInfo->buffer = BufferHandle;
 				bChanged = true;
 			}
 			if (HashableInfo.Buffer.Offset != static_cast<uint32>(Offset))
@@ -1044,7 +1044,7 @@ protected:
 		}
 		else
 		{
-			bChanged = CopyAndReturnNotEqual(BufferInfo->buffer, Allocation.GetBufferHandle());
+			bChanged = CopyAndReturnNotEqual(BufferInfo->buffer, BufferHandle);
 			bChanged |= CopyAndReturnNotEqual(BufferInfo->offset, Offset);
 			bChanged |= CopyAndReturnNotEqual(BufferInfo->range, Range);
 		}
diff --git a/Engine/Source/Runtime/VulkanRHI/Private/VulkanPipelineState.cpp b/Engine/Source/Runtime/VulkanRHI/Private/VulkanPipelineState.cpp
--- a/Engine/Source/Runtime/VulkanRHI/Private/VulkanPipelineState.cpp
+++ b/Engine/Source/Runtime/VulkanRHI/Private/VulkanPipelineState.cpp
@@ -743,7 +743,6 @@ void FVulkanCommonPipelineDescriptorState::SetSRV(FVulkanCmdBuffer* CmdBuffer, b
 		break;
 		
 	case FVulkanView::EType::TypedBuffer:
-		check((SRV->GetTypedBufferView().Flags & VK_BUFFER_USAGE_UNIFORM_TEXEL_BUFFER_BIT) == VK_BUFFER_USAGE_UNIFORM_TEXEL_BUFFER_BIT);
 		MarkDirty(DSWriter[DescriptorSet].WriteUniformTexelBuffer(BindingIndex, SRV->GetTypedBufferView()));
 		break;
 
@@ -783,7 +782,6 @@ void FVulkanCommonPipelineDescriptorState::SetUAV(FVulkanCmdBuffer* CmdBuffer, b
 		break;
 
 	case FVulkanView::EType::TypedBuffer:
-		check((UAV->GetTypedBufferView().Flags & VK_BUFFER_USAGE_STORAGE_TEXEL_BUFFER_BIT) == VK_BUFFER_USAGE_STORAGE_TEXEL_BUFFER_BIT);
 		MarkDirty(DSWriter[DescriptorSet].WriteStorageTexelBuffer(BindingIndex, UAV->GetTypedBufferView()));
 		break;
 
diff --git a/Engine/Source/Runtime/VulkanRHI/Private/VulkanPipelineState.h b/Engine/Source/Runtime/VulkanRHI/Private/VulkanPipelineState.h
--- a/Engine/Source/Runtime/VulkanRHI/Private/VulkanPipelineState.h
+++ b/Engine/Source/Runtime/VulkanRHI/Private/VulkanPipelineState.h
@@ -95,13 +95,14 @@ public:
 	template<bool bDynamic>
 	inline void SetUniformBuffer(uint8 DescriptorSet, uint32 BindingIndex, const FVulkanUniformBuffer* UniformBuffer)
 	{
+		const VulkanRHI::FVulkanAllocation& Allocation = UniformBuffer->Allocation;
 		if (bDynamic)
 		{
-			MarkDirty(DSWriter[DescriptorSet].WriteDynamicUniformBuffer(BindingIndex, UniformBuffer->Allocation, 0, UniformBuffer->GetSize(), UniformBuffer->GetOffset()));
+			MarkDirty(DSWriter[DescriptorSet].WriteDynamicUniformBuffer(BindingIndex, Allocation.GetBufferHandle(), Allocation.HandleId, 0, UniformBuffer->GetSize(), UniformBuffer->GetOffset()));
 		}
 		else
 		{
-			MarkDirty(DSWriter[DescriptorSet].WriteUniformBuffer(BindingIndex, UniformBuffer->Allocation, UniformBuffer->GetOffset(), UniformBuffer->GetSize()));
+			MarkDirty(DSWriter[DescriptorSet].WriteUniformBuffer(BindingIndex, Allocation.GetBufferHandle(), Allocation.HandleId, UniformBuffer->GetOffset(), UniformBuffer->GetSize()));
 		}
 	}
 
@@ -313,15 +314,16 @@ static inline bool UpdatePackedUniformBuffers(VkDeviceSize UBOffsetAlignment, co
 			// get location in the ring buffer to use
 			FMemory::Memcpy(CPURingBufferBase + RingBufferOffset, StagedUniformBuffer.GetData(), UBSize);
 
+			const VulkanRHI::FVulkanAllocation& Allocation = UniformBufferUploader->GetCPUBufferAllocation();
 			if (bIsDynamic)
 			{
-				const bool bDirty = DescriptorWriteSet.WriteDynamicUniformBuffer(BindingIndex, UniformBufferUploader->GetCPUBufferAllocation(), UniformBufferUploader->GetCPUBufferOffset(), UBSize, RingBufferOffset);
+				const bool bDirty = DescriptorWriteSet.WriteDynamicUniformBuffer(BindingIndex, Allocation.GetBufferHandle(), Allocation.HandleId, UniformBufferUploader->GetCPUBufferOffset(), UBSize, RingBufferOffset);
 				bAnyUBDirty = bAnyUBDirty || bDirty;
 
 			}
 			else
 			{
-				const bool bDirty = DescriptorWriteSet.WriteUniformBuffer(BindingIndex, UniformBufferUploader->GetCPUBufferAllocation(), RingBufferOffset + UniformBufferUploader->GetCPUBufferOffset(), UBSize);
+				const bool bDirty = DescriptorWriteSet.WriteUniformBuffer(BindingIndex, Allocation.GetBufferHandle(), Allocation.HandleId, RingBufferOffset + UniformBufferUploader->GetCPUBufferOffset(), UBSize);
 				bAnyUBDirty = bAnyUBDirty || bDirty;
 
 			}
diff --git a/Engine/Source/Runtime/VulkanRHI/Private/VulkanResources.h b/Engine/Source/Runtime/VulkanRHI/Private/VulkanResources.h
--- a/Engine/Source/Runtime/VulkanRHI/Private/VulkanResources.h
+++ b/Engine/Source/Runtime/VulkanRHI/Private/VulkanResources.h
@@ -345,20 +345,17 @@ class FVulkanView
 public:
 	struct FTypedBufferView
 	{
-		VkBufferView View   = VK_NULL_HANDLE;
-		uint32       ViewId = 0;
-		VkFlags      Flags  = 0;
-		uint32       Offset = 0;
-		uint32       Size   = 0;
-		// Whether source buffer is volatile
-		bool bVolatile = false;
+		VkBufferView View      = VK_NULL_HANDLE;
+		uint32       ViewId    = 0;
+		bool         bVolatile = false; // Whether source buffer is volatile
 	};
 
 	struct FStructuredBufferView
 	{
-		VulkanRHI::FVulkanAllocation Allocation;
-		uint32 Offset = 0;
-		uint32 Size   = 0;
+		VkBuffer Buffer = VK_NULL_HANDLE;
+		uint32 HandleId = 0;
+		uint32 Offset   = 0;
+		uint32 Size     = 0;
 	};
 
 #if VULKAN_RHI_RAYTRACING
@@ -419,6 +416,7 @@ public:
 	FAccelerationStructureView const& GetAccelerationStructureView() const { return Storage.Get<FAccelerationStructureView>(); }
 #endif
 
+	// NOTE: The InOffset applies to the FVulkanResourceMultiBuffer (it does not include any internal Allocation offsets that may exist)
 	FVulkanView* InitAsTypedBufferView(
 		  FVulkanResourceMultiBuffer* Buffer
 		, EPixelFormat Format
@@ -438,6 +436,7 @@ public:
 		, bool bUseIdentitySwizzle = false
 		, VkImageUsageFlags ImageUsageFlags = 0);
 
+	// NOTE: The InOffset applies to the FVulkanResourceMultiBuffer (it does not include any internal Allocation offsets that may exist)
 	FVulkanView* InitAsStructuredBufferView(
 		  FVulkanResourceMultiBuffer* Buffer
 		, uint32 InOffset
diff --git a/Engine/Source/Runtime/VulkanRHI/Private/VulkanUAV.cpp b/Engine/Source/Runtime/VulkanRHI/Private/VulkanUAV.cpp
--- a/Engine/Source/Runtime/VulkanRHI/Private/VulkanUAV.cpp
+++ b/Engine/Source/Runtime/VulkanRHI/Private/VulkanUAV.cpp
@@ -54,8 +54,7 @@ FVulkanView* FVulkanView::InitAsTypedBufferView(FVulkanResourceMultiBuffer* Buff
 	Storage.Emplace<FTypedBufferView>();
 	FTypedBufferView& TBV = Storage.Get<FTypedBufferView>();
 
-	TBV.Offset = InOffset;
-	TBV.Size = InSize;
+	const uint32 TotalOffset = Buffer->GetOffset() + InOffset;
 
 	check(UEFormat != PF_Unknown);
 	VkFormat Format = GVulkanBufferFormat[UEFormat];
@@ -64,7 +63,7 @@ FVulkanView* FVulkanView::InitAsTypedBufferView(FVulkanResourceMultiBuffer* Buff
 	VkBufferViewCreateInfo ViewInfo;
 	ZeroVulkanStruct(ViewInfo, VK_STRUCTURE_TYPE_BUFFER_VIEW_CREATE_INFO);
 	ViewInfo.buffer = Buffer->GetHandle();
-	ViewInfo.offset = InOffset;
+	ViewInfo.offset = TotalOffset;
 	ViewInfo.format = Format;
 
 	//#todo-rco: Revisit this if buffer views become VK_BUFFER_USAGE_STORAGE_BUFFER_BIT instead of VK_BUFFER_USAGE_UNIFORM_TEXEL_BUFFER_BIT
@@ -73,8 +72,7 @@ FVulkanView* FVulkanView::InitAsTypedBufferView(FVulkanResourceMultiBuffer* Buff
 	ViewInfo.range = FMath::Min<uint64>(InSize, MaxSize);
 	// TODO: add a check() for exceeding MaxSize, to catch code which blindly makes views without checking the platform limits.
 
-	TBV.Flags = Buffer->GetBufferUsageFlags() & (VK_BUFFER_USAGE_UNIFORM_TEXEL_BUFFER_BIT | VK_BUFFER_USAGE_STORAGE_TEXEL_BUFFER_BIT);
-	check(TBV.Flags);
+	check(Buffer->GetBufferUsageFlags() & (VK_BUFFER_USAGE_UNIFORM_TEXEL_BUFFER_BIT | VK_BUFFER_USAGE_STORAGE_TEXEL_BUFFER_BIT));
 	check(IsAligned(InOffset, Limits.minTexelBufferOffsetAlignment));
 
 	VERIFYVULKANRESULT(VulkanRHI::vkCreateBufferView(Device.GetInstanceHandle(), &ViewInfo, VULKAN_CPU_ALLOCATOR, &TBV.View));
@@ -194,13 +192,16 @@ FVulkanView* FVulkanView::InitAsStructuredBufferView(FVulkanResourceMultiBuffer*
 	Storage.Emplace<FStructuredBufferView>();
 	FStructuredBufferView& SBV = Storage.Get<FStructuredBufferView>();
 
-	SBV.Offset = InOffset;
+	const uint32 TotalOffset = Buffer->GetOffset() + InOffset;
+
+	SBV.Buffer = Buffer->GetHandle();
+	SBV.HandleId = Buffer->GetCurrentAllocation().HandleId;
+	SBV.Offset = TotalOffset;
 	SBV.Size = InSize;
-	SBV.Allocation.Reference(Buffer->GetCurrentAllocation());
 
 	BindlessHandle = Device.GetBindlessDescriptorManager()->RegisterBuffer(
 		Buffer->GetHandle(),
-		Buffer->GetOffset() + InOffset,
+		TotalOffset,
 		InSize,
 		VK_DESCRIPTOR_TYPE_STORAGE_BUFFER
 	);
@@ -298,11 +299,12 @@ void FVulkanShaderResourceView::UpdateView()
 			{
 			case FRHIViewDesc::EBufferType::Raw:
 			case FRHIViewDesc::EBufferType::Structured:
-				InitAsStructuredBufferView(Buffer, Buffer->GetOffset() + Info.OffsetInBytes, Info.SizeInBytes);
+				InitAsStructuredBufferView(Buffer, Info.OffsetInBytes, Info.SizeInBytes);
 				break;
 
 			case FRHIViewDesc::EBufferType::Typed:
-				InitAsTypedBufferView(Buffer, Info.Format, Buffer->GetOffset() + Info.OffsetInBytes, Info.SizeInBytes);
+				check(VKHasAllFlags(Buffer->GetBufferUsageFlags(), VK_BUFFER_USAGE_UNIFORM_TEXEL_BUFFER_BIT));
+				InitAsTypedBufferView(Buffer, Info.Format, Info.OffsetInBytes, Info.SizeInBytes);
 				break;
 
 #if VULKAN_RHI_RAYTRACING
@@ -387,11 +389,12 @@ void FVulkanUnorderedAccessView::UpdateView()
 			{
 			case FRHIViewDesc::EBufferType::Raw:
 			case FRHIViewDesc::EBufferType::Structured:
-				InitAsStructuredBufferView(Buffer, Buffer->GetOffset() + Info.OffsetInBytes, Info.SizeInBytes);
+				InitAsStructuredBufferView(Buffer, Info.OffsetInBytes, Info.SizeInBytes);
 				break;
 
 			case FRHIViewDesc::EBufferType::Typed:
-				InitAsTypedBufferView(Buffer, Info.Format, Buffer->GetOffset() + Info.OffsetInBytes, Info.SizeInBytes);
+				check(VKHasAllFlags(Buffer->GetBufferUsageFlags(), VK_BUFFER_USAGE_STORAGE_TEXEL_BUFFER_BIT));
+				InitAsTypedBufferView(Buffer, Info.Format, Info.OffsetInBytes, Info.SizeInBytes);
 				break;
 
 #if VULKAN_RHI_RAYTRACING
