VulkanRHI: Add suport for VK_EXT_shader_demote_to_helper_invocation if present.

#rnx

diff --git a/Engine/Source/Runtime/VulkanRHI/Private/VulkanExtensions.cpp b/Engine/Source/Runtime/VulkanRHI/Private/VulkanExtensions.cpp
--- a/Engine/Source/Runtime/VulkanRHI/Private/VulkanExtensions.cpp
+++ b/Engine/Source/Runtime/VulkanRHI/Private/VulkanExtensions.cpp
@@ -1054,6 +1054,40 @@ private:
 };
 
 
+// ***** VK_EXT_shader_demote_to_helper_invocation
+class FVulkanEXTShaderDemoteToHelperInvocationExtension : public FVulkanDeviceExtension
+{
+public:
+
+	FVulkanEXTShaderDemoteToHelperInvocationExtension(FVulkanDevice* InDevice)
+		: FVulkanDeviceExtension(InDevice, VK_EXT_SHADER_DEMOTE_TO_HELPER_INVOCATION_EXTENSION_NAME, VULKAN_EXTENSION_ENABLED, VK_API_VERSION_1_3)
+	{
+	}
+
+	virtual void PrePhysicalDeviceFeatures(VkPhysicalDeviceFeatures2KHR& PhysicalDeviceFeatures2) override final
+	{
+		ZeroVulkanStruct(ShaderDemoteToHelperInvocationFeatures, VK_STRUCTURE_TYPE_PHYSICAL_DEVICE_SHADER_DEMOTE_TO_HELPER_INVOCATION_FEATURES);
+		AddToPNext(PhysicalDeviceFeatures2, ShaderDemoteToHelperInvocationFeatures);
+	}
+
+	virtual void PostPhysicalDeviceFeatures(FOptionalVulkanDeviceExtensions& ExtensionFlags) override final
+	{
+		bRequirementsPassed = (ShaderDemoteToHelperInvocationFeatures.shaderDemoteToHelperInvocation == VK_TRUE);
+	}
+
+	virtual void PreCreateDevice(VkDeviceCreateInfo& DeviceCreateInfo) override final
+	{
+		if (bRequirementsPassed)
+		{
+			AddToPNext(DeviceCreateInfo, ShaderDemoteToHelperInvocationFeatures);
+		}
+	}
+
+private:
+	VkPhysicalDeviceShaderDemoteToHelperInvocationFeaturesEXT ShaderDemoteToHelperInvocationFeatures;
+};
+
+
 // ***** VK_EXT_calibrated_timestamps
 class FVulkanEXTCalibratedTimestampsExtension : public FVulkanDeviceExtension
 {
@@ -1218,6 +1252,7 @@ FVulkanDeviceExtensionArray FVulkanDeviceExtension::GetUESupportedDeviceExtensio
 	ADD_CUSTOM_EXTENSION(FVulkanEXTCalibratedTimestampsExtension);
 	ADD_CUSTOM_EXTENSION(FVulkanEXTDescriptorBuffer);
 	ADD_CUSTOM_EXTENSION(FVulkanEXTDeviceFaultExtension);
+	ADD_CUSTOM_EXTENSION(FVulkanEXTShaderDemoteToHelperInvocationExtension);
 
 	// Needed for Raytracing
 	ADD_CUSTOM_EXTENSION(FVulkanKHRBufferDeviceAddressExtension);
