VulkanRHI: Split the reservation of a bindless descriptor slot from its update.  Have the VulkanView take a descriptor type on creation so it can reserve its slot right away.

#rnx

diff --git a/Engine/Source/Runtime/VulkanRHI/Private/VulkanDescriptorSets.cpp b/Engine/Source/Runtime/VulkanRHI/Private/VulkanDescriptorSets.cpp
--- a/Engine/Source/Runtime/VulkanRHI/Private/VulkanDescriptorSets.cpp
+++ b/Engine/Source/Runtime/VulkanRHI/Private/VulkanDescriptorSets.cpp
@@ -102,17 +102,21 @@ DECLARE_STATS_GROUP(TEXT("Vulkan Bindless"), STATGROUP_VulkanBindless, STATCAT_A
 
 DECLARE_DWORD_ACCUMULATOR_STAT_EXTERN(TEXT("Num Peak Descriptor Count"), STAT_VulkanBindlessPeakDescriptorCount, STATGROUP_VulkanBindless, );
 DECLARE_DWORD_ACCUMULATOR_STAT_EXTERN(TEXT("Peak Samplers"), STAT_VulkanBindlessPeakSampler, STATGROUP_VulkanBindless, );
-DECLARE_DWORD_ACCUMULATOR_STAT_EXTERN(TEXT("Peak Images"), STAT_VulkanBindlessPeakImage, STATGROUP_VulkanBindless, );
-DECLARE_DWORD_ACCUMULATOR_STAT_EXTERN(TEXT("Peak Buffers"), STAT_VulkanBindlessPeakBuffer, STATGROUP_VulkanBindless, );
-DECLARE_DWORD_ACCUMULATOR_STAT_EXTERN(TEXT("Peak Texel Buffers"), STAT_VulkanBindlessPeakTexelBuffer, STATGROUP_VulkanBindless, );
+DECLARE_DWORD_ACCUMULATOR_STAT_EXTERN(TEXT("Peak Sampled Images"), STAT_VulkanBindlessPeakSampledImage, STATGROUP_VulkanBindless, );
+DECLARE_DWORD_ACCUMULATOR_STAT_EXTERN(TEXT("Peak Storage Images"), STAT_VulkanBindlessPeakStorageImage, STATGROUP_VulkanBindless, );
+DECLARE_DWORD_ACCUMULATOR_STAT_EXTERN(TEXT("Peak Storage Buffers"), STAT_VulkanBindlessPeakStorageBuffer, STATGROUP_VulkanBindless, );
+DECLARE_DWORD_ACCUMULATOR_STAT_EXTERN(TEXT("Peak Uniform Texel Buffers"), STAT_VulkanBindlessPeakUniformTexelBuffer, STATGROUP_VulkanBindless, );
+DECLARE_DWORD_ACCUMULATOR_STAT_EXTERN(TEXT("Peak Storage Texel Buffers"), STAT_VulkanBindlessPeakStorageTexelBuffer, STATGROUP_VulkanBindless, );
 DECLARE_DWORD_ACCUMULATOR_STAT_EXTERN(TEXT("Peak Acceleration Structures"), STAT_VulkanBindlessPeakAccelerationStructure, STATGROUP_VulkanBindless, );
 DECLARE_DWORD_COUNTER_STAT_EXTERN(TEXT("Write Per Frame"), STAT_VulkanBindlessWritePerFrame, STATGROUP_VulkanBindless, );
 
 DEFINE_STAT(STAT_VulkanBindlessPeakDescriptorCount);
 DEFINE_STAT(STAT_VulkanBindlessPeakSampler);
-DEFINE_STAT(STAT_VulkanBindlessPeakImage);
-DEFINE_STAT(STAT_VulkanBindlessPeakBuffer);
-DEFINE_STAT(STAT_VulkanBindlessPeakTexelBuffer);
+DEFINE_STAT(STAT_VulkanBindlessPeakSampledImage);
+DEFINE_STAT(STAT_VulkanBindlessPeakStorageImage);
+DEFINE_STAT(STAT_VulkanBindlessPeakStorageBuffer);
+DEFINE_STAT(STAT_VulkanBindlessPeakUniformTexelBuffer);
+DEFINE_STAT(STAT_VulkanBindlessPeakStorageTexelBuffer);
 DEFINE_STAT(STAT_VulkanBindlessPeakAccelerationStructure);
 DEFINE_STAT(STAT_VulkanBindlessWritePerFrame);
 
@@ -684,167 +688,152 @@ void FVulkanBindlessDescriptorManager::RegisterUniformBuffers(VkCommandBuffer Co
 	}
 }
 
-FRHIDescriptorHandle FVulkanBindlessDescriptorManager::RegisterSampler(VkSampler VulkanSampler)
+void FVulkanBindlessDescriptorManager::UpdateStatsForHandle(FRHIDescriptorHandle DescriptorHandle)
 {
-	if (!bIsSupported)
+	const uint8 SetIndex = DescriptorHandle.GetRawType();
+	const BindlessSetState& State = BindlessSetStates[SetIndex];
+
+	switch (DescriptorHandle.GetRawType())
 	{
-		return FRHIDescriptorHandle();
+	case VulkanBindless::BindlessSamplerSet:                SET_DWORD_STAT(STAT_VulkanBindlessPeakSampler, State.PeakDescriptorCount); break;
+	case VulkanBindless::BindlessSampledImageSet:           SET_DWORD_STAT(STAT_VulkanBindlessPeakSampledImage, State.PeakDescriptorCount); break;
+	case VulkanBindless::BindlessStorageImageSet:           SET_DWORD_STAT(STAT_VulkanBindlessPeakStorageImage, State.PeakDescriptorCount); break;
+	case VulkanBindless::BindlessUniformTexelBufferSet:     SET_DWORD_STAT(STAT_VulkanBindlessPeakUniformTexelBuffer, State.PeakDescriptorCount); break;
+	case VulkanBindless::BindlessStorageTexelBufferSet:     SET_DWORD_STAT(STAT_VulkanBindlessPeakStorageTexelBuffer, State.PeakDescriptorCount); break;
+	case VulkanBindless::BindlessStorageBufferSet:          SET_DWORD_STAT(STAT_VulkanBindlessPeakStorageBuffer, State.PeakDescriptorCount); break;
+	case VulkanBindless::BindlessAccelerationStructureSet:  SET_DWORD_STAT(STAT_VulkanBindlessPeakAccelerationStructure, State.PeakDescriptorCount); break;
+
+	case VulkanBindless::BindlessUniformBufferSet:
+	default: checkNoEntry();
 	}
-
-	const uint8 SetIndex = GetIndexForDescriptorType(VK_DESCRIPTOR_TYPE_SAMPLER);
-	BindlessSetState& State = BindlessSetStates[SetIndex];
-
-	const uint32 SamplerIndex = State.PeakDescriptorCount++;
-	checkf(SamplerIndex < State.MaxDescriptorCount, TEXT("You need to grow the sampler array size!"));
-	const uint32 ByteOffset = SamplerIndex * State.DescriptorSize;
-
-	VkDescriptorGetInfoEXT Info;
-	ZeroVulkanStruct(Info, VK_STRUCTURE_TYPE_DESCRIPTOR_GET_INFO_EXT);
-	Info.type = State.DescriptorType;
-	Info.data.pSampler = &VulkanSampler;
-	VulkanRHI::vkGetDescriptorEXT(Device->GetInstanceHandle(), &Info, State.DescriptorSize, &State.DebugDescriptors[ByteOffset]);
-
-	FMemory::Memcpy(&State.MappedPointer[ByteOffset], &State.DebugDescriptors[ByteOffset], State.DescriptorSize);
-
-	SET_DWORD_STAT(STAT_VulkanBindlessPeakSampler, State.PeakDescriptorCount);
-	return FRHIDescriptorHandle(SetIndex, SamplerIndex);
 }
 
-FRHIDescriptorHandle FVulkanBindlessDescriptorManager::RegisterImage(VkImageView ImageView, VkDescriptorType DescriptorType, bool bIsDepthStencil)
+FRHIDescriptorHandle FVulkanBindlessDescriptorManager::ReserveDescriptor(VkDescriptorType DescriptorType)
 {
-	if (!bIsSupported)
+	if (bIsSupported)
 	{
-		return FRHIDescriptorHandle();
+		const uint8 SetIndex = GetIndexForDescriptorType(DescriptorType);
+		BindlessSetState& State = BindlessSetStates[SetIndex];
+		const uint32 ResourceIndex = GetFreeResourceIndex(State);
+		return FRHIDescriptorHandle(SetIndex, ResourceIndex);
 	}
+	return FRHIDescriptorHandle();
+}
 
-	check((DescriptorType == VK_DESCRIPTOR_TYPE_SAMPLED_IMAGE) || (DescriptorType == VK_DESCRIPTOR_TYPE_STORAGE_IMAGE));
+void FVulkanBindlessDescriptorManager::UpdateDescriptor(FRHIDescriptorHandle DescriptorHandle, VkDescriptorDataEXT DescriptorData)
+{
+	checkf(DescriptorHandle.IsValid(), TEXT("Attemping to update invalid descriptor handle!"));
 
-	const uint8 SetIndex = GetIndexForDescriptorType(DescriptorType);
+	const uint8 SetIndex = DescriptorHandle.GetRawType();
 	BindlessSetState& State = BindlessSetStates[SetIndex];
-	const uint32 ResourceIndex = GetFreeResourceIndex(State);
-	const uint32 ByteOffset = ResourceIndex * State.DescriptorSize;
-
-	VkDescriptorImageInfo DescriptorImageInfo;
-	FMemory::Memzero(DescriptorImageInfo);
-	DescriptorImageInfo.imageView = ImageView;
-	DescriptorImageInfo.imageLayout = (DescriptorType == VK_DESCRIPTOR_TYPE_STORAGE_IMAGE) ? VK_IMAGE_LAYOUT_GENERAL :
-		(bIsDepthStencil ? VK_IMAGE_LAYOUT_READ_ONLY_OPTIMAL : VK_IMAGE_LAYOUT_SHADER_READ_ONLY_OPTIMAL);
+	const uint32 ByteOffset = DescriptorHandle.GetIndex() * State.DescriptorSize;
 
 	VkDescriptorGetInfoEXT Info;
 	ZeroVulkanStruct(Info, VK_STRUCTURE_TYPE_DESCRIPTOR_GET_INFO_EXT);
 	Info.type = State.DescriptorType;
-	Info.data.pSampledImage = &DescriptorImageInfo;  // same pointer for storage, it's a union
+	Info.data = DescriptorData;
 	VulkanRHI::vkGetDescriptorEXT(Device->GetInstanceHandle(), &Info, State.DescriptorSize, &State.DebugDescriptors[ByteOffset]);
 
 	FMemory::Memcpy(&State.MappedPointer[ByteOffset], &State.DebugDescriptors[ByteOffset], State.DescriptorSize);
 
-	SET_DWORD_STAT(STAT_VulkanBindlessPeakImage, State.PeakDescriptorCount);
-	return FRHIDescriptorHandle(SetIndex, ResourceIndex);
+	UpdateStatsForHandle(DescriptorHandle);
 }
 
-FRHIDescriptorHandle FVulkanBindlessDescriptorManager::RegisterBuffer(VkBuffer VulkanBuffer, VkDeviceSize BufferOffset, VkDeviceSize BufferSize, VkDescriptorType DescriptorType)
+void FVulkanBindlessDescriptorManager::UpdateSampler(FRHIDescriptorHandle DescriptorHandle, VkSampler VulkanSampler)
 {
-	if (!bIsSupported)
+	if (bIsSupported)
 	{
-		return FRHIDescriptorHandle();
+		VkDescriptorDataEXT DescriptorData;
+		DescriptorData.pSampler = &VulkanSampler;
+		UpdateDescriptor(DescriptorHandle, DescriptorData);
 	}
-
-	check(/*(DescriptorType == VK_DESCRIPTOR_TYPE_UNIFORM_BUFFER) ||*/ (DescriptorType == VK_DESCRIPTOR_TYPE_STORAGE_BUFFER));
-
-	const uint8 SetIndex = GetIndexForDescriptorType(DescriptorType);
-	BindlessSetState& State = BindlessSetStates[SetIndex];
-	const uint32 ResourceIndex = GetFreeResourceIndex(State);
-	const uint32 ByteOffset = ResourceIndex * State.DescriptorSize;
-
-	VkBufferDeviceAddressInfo BufferInfo;
-	ZeroVulkanStruct(BufferInfo, VK_STRUCTURE_TYPE_BUFFER_DEVICE_ADDRESS_INFO);
-	BufferInfo.buffer = VulkanBuffer;
-	VkDeviceAddress BufferAddress = VulkanRHI::vkGetBufferDeviceAddressKHR(Device->GetInstanceHandle(), &BufferInfo);
-
-	VkDescriptorAddressInfoEXT AddressInfo;
-	ZeroVulkanStruct(AddressInfo, VK_STRUCTURE_TYPE_DESCRIPTOR_ADDRESS_INFO_EXT);
-	AddressInfo.address = BufferAddress + BufferOffset;
-	AddressInfo.range = BufferSize;
-
-	VkDescriptorGetInfoEXT Info;
-	ZeroVulkanStruct(Info, VK_STRUCTURE_TYPE_DESCRIPTOR_GET_INFO_EXT);
-	Info.type = State.DescriptorType;
-	Info.data.pUniformBuffer = &AddressInfo;  // same pointer for storage, it's a union
-	VulkanRHI::vkGetDescriptorEXT(Device->GetInstanceHandle(), &Info, State.DescriptorSize, &State.DebugDescriptors[ByteOffset]);
-
-	FMemory::Memcpy(&State.MappedPointer[ByteOffset], &State.DebugDescriptors[ByteOffset], State.DescriptorSize);
-
-	SET_DWORD_STAT(STAT_VulkanBindlessPeakBuffer, State.PeakDescriptorCount);
-	return FRHIDescriptorHandle(SetIndex, ResourceIndex);
 }
 
-FRHIDescriptorHandle FVulkanBindlessDescriptorManager::RegisterTexelBuffer(const VkBufferViewCreateInfo& ViewInfo, VkDescriptorType DescriptorType)
+void FVulkanBindlessDescriptorManager::UpdateImage(FRHIDescriptorHandle DescriptorHandle, VkImageView ImageView, bool bIsDepthStencil)
 {
-	if (!bIsSupported)
+	if (bIsSupported)
 	{
-		return FRHIDescriptorHandle();
+		const VkDescriptorType DescriptorType = GetDescriptorTypeForSetIndex(DescriptorHandle.GetRawType());
+		check((DescriptorType == VK_DESCRIPTOR_TYPE_SAMPLED_IMAGE) || (DescriptorType == VK_DESCRIPTOR_TYPE_STORAGE_IMAGE));
+
+		VkDescriptorImageInfo DescriptorImageInfo;
+		DescriptorImageInfo.sampler = VK_NULL_HANDLE;
+		DescriptorImageInfo.imageView = ImageView;
+		DescriptorImageInfo.imageLayout = (DescriptorType == VK_DESCRIPTOR_TYPE_STORAGE_IMAGE) ? VK_IMAGE_LAYOUT_GENERAL :
+			(bIsDepthStencil ? VK_IMAGE_LAYOUT_READ_ONLY_OPTIMAL : VK_IMAGE_LAYOUT_SHADER_READ_ONLY_OPTIMAL);
+
+		VkDescriptorDataEXT DescriptorData;
+		DescriptorData.pSampledImage = &DescriptorImageInfo;  // same pointer for storage, it's a union
+		UpdateDescriptor(DescriptorHandle, DescriptorData);
 	}
-
-	check((DescriptorType == VK_DESCRIPTOR_TYPE_UNIFORM_TEXEL_BUFFER) || (DescriptorType == VK_DESCRIPTOR_TYPE_STORAGE_TEXEL_BUFFER));
-
-	const uint8 SetIndex = GetIndexForDescriptorType(DescriptorType);
-	BindlessSetState& State = BindlessSetStates[SetIndex];
-	const uint32 ResourceIndex = GetFreeResourceIndex(State);
-	const uint32 ByteOffset = ResourceIndex * State.DescriptorSize;
-
-	VkBufferDeviceAddressInfo BufferInfo;
-	ZeroVulkanStruct(BufferInfo, VK_STRUCTURE_TYPE_BUFFER_DEVICE_ADDRESS_INFO);
-	BufferInfo.buffer = ViewInfo.buffer;
-	VkDeviceAddress BufferAddress = VulkanRHI::vkGetBufferDeviceAddressKHR(Device->GetInstanceHandle(), &BufferInfo);
-
-	VkDescriptorAddressInfoEXT AddressInfo;
-	ZeroVulkanStruct(AddressInfo, VK_STRUCTURE_TYPE_DESCRIPTOR_ADDRESS_INFO_EXT);
-	AddressInfo.address = BufferAddress + ViewInfo.offset;
-	AddressInfo.range = ViewInfo.range;
-	AddressInfo.format = ViewInfo.format;
-
-	VkDescriptorGetInfoEXT Info;
-	ZeroVulkanStruct(Info, VK_STRUCTURE_TYPE_DESCRIPTOR_GET_INFO_EXT);
-	Info.type = State.DescriptorType;
-	Info.data.pUniformTexelBuffer = &AddressInfo;  // same pointer for storage, it's a union
-	VulkanRHI::vkGetDescriptorEXT(Device->GetInstanceHandle(), &Info, State.DescriptorSize, &State.DebugDescriptors[ByteOffset]);
-
-	FMemory::Memcpy(&State.MappedPointer[ByteOffset], &State.DebugDescriptors[ByteOffset], State.DescriptorSize);
-
-	SET_DWORD_STAT(STAT_VulkanBindlessPeakTexelBuffer, State.PeakDescriptorCount);
-	return FRHIDescriptorHandle(SetIndex, ResourceIndex);
 }
 
-FRHIDescriptorHandle FVulkanBindlessDescriptorManager::RegisterAccelerationStructure(VkAccelerationStructureKHR AccelerationStructure)
+void FVulkanBindlessDescriptorManager::UpdateBuffer(FRHIDescriptorHandle DescriptorHandle, VkBuffer VulkanBuffer, VkDeviceSize BufferOffset, VkDeviceSize BufferSize)
 {
-#if VULKAN_RHI_RAYTRACING
-	if (!bIsSupported)
+	if (bIsSupported)
 	{
-		return FRHIDescriptorHandle();
+		const VkDescriptorType DescriptorType = GetDescriptorTypeForSetIndex(DescriptorHandle.GetRawType());
+		check(/*(DescriptorType == VK_DESCRIPTOR_TYPE_UNIFORM_BUFFER) ||*/ (DescriptorType == VK_DESCRIPTOR_TYPE_STORAGE_BUFFER));
+
+		// :todo-jn: start caching buffer addresses in resources to avoid the extra call
+		VkBufferDeviceAddressInfo BufferInfo;
+		ZeroVulkanStruct(BufferInfo, VK_STRUCTURE_TYPE_BUFFER_DEVICE_ADDRESS_INFO);
+		BufferInfo.buffer = VulkanBuffer;
+		const VkDeviceAddress BufferAddress = VulkanRHI::vkGetBufferDeviceAddressKHR(Device->GetInstanceHandle(), &BufferInfo);
+
+		VkDescriptorAddressInfoEXT AddressInfo;
+		ZeroVulkanStruct(AddressInfo, VK_STRUCTURE_TYPE_DESCRIPTOR_ADDRESS_INFO_EXT);
+		AddressInfo.address = BufferAddress + BufferOffset;
+		AddressInfo.range = BufferSize;
+
+		VkDescriptorDataEXT DescriptorData;
+		DescriptorData.pStorageBuffer = &AddressInfo;  // same pointer for uniform, it's a union
+		UpdateDescriptor(DescriptorHandle, DescriptorData);
 	}
+}
 
-	const uint8 SetIndex = GetIndexForDescriptorType(VK_DESCRIPTOR_TYPE_ACCELERATION_STRUCTURE_KHR);
-	BindlessSetState& State = BindlessSetStates[SetIndex];
-	const uint32 ResourceIndex = GetFreeResourceIndex(State);
-	const uint32 ByteOffset = ResourceIndex * State.DescriptorSize;
-
-	VkAccelerationStructureDeviceAddressInfoKHR AccelerationStructureDeviceAddressInfo;
-	ZeroVulkanStruct(AccelerationStructureDeviceAddressInfo, VK_STRUCTURE_TYPE_ACCELERATION_STRUCTURE_DEVICE_ADDRESS_INFO_KHR);
-	AccelerationStructureDeviceAddressInfo.accelerationStructure = AccelerationStructure;
-	VkDeviceAddress BufferAddress = VulkanRHI::vkGetAccelerationStructureDeviceAddressKHR(Device->GetInstanceHandle(), &AccelerationStructureDeviceAddressInfo);
+void FVulkanBindlessDescriptorManager::UpdateTexelBuffer(FRHIDescriptorHandle DescriptorHandle, const VkBufferViewCreateInfo& ViewInfo)
+{
+	if (bIsSupported)
+	{
+		const VkDescriptorType DescriptorType = GetDescriptorTypeForSetIndex(DescriptorHandle.GetRawType());
+		check((DescriptorType == VK_DESCRIPTOR_TYPE_UNIFORM_TEXEL_BUFFER) || (DescriptorType == VK_DESCRIPTOR_TYPE_STORAGE_TEXEL_BUFFER));
+
+		// :todo-jn: start caching buffer addresses in resources to avoid the extra call
+		VkBufferDeviceAddressInfo BufferInfo;
+		ZeroVulkanStruct(BufferInfo, VK_STRUCTURE_TYPE_BUFFER_DEVICE_ADDRESS_INFO);
+		BufferInfo.buffer = ViewInfo.buffer;
+		const VkDeviceAddress BufferAddress = VulkanRHI::vkGetBufferDeviceAddressKHR(Device->GetInstanceHandle(), &BufferInfo);
+
+		VkDescriptorAddressInfoEXT AddressInfo;
+		ZeroVulkanStruct(AddressInfo, VK_STRUCTURE_TYPE_DESCRIPTOR_ADDRESS_INFO_EXT);
+		AddressInfo.address = BufferAddress + ViewInfo.offset;
+		AddressInfo.range = ViewInfo.range;
+		AddressInfo.format = ViewInfo.format;
+
+		VkDescriptorDataEXT DescriptorData;
+		DescriptorData.pUniformTexelBuffer = &AddressInfo;  // same pointer for storage, it's a union
+		UpdateDescriptor(DescriptorHandle, DescriptorData);
+	}
+}
 
-	VkDescriptorGetInfoEXT Info;
-	ZeroVulkanStruct(Info, VK_STRUCTURE_TYPE_DESCRIPTOR_GET_INFO_EXT);
-	Info.type = State.DescriptorType;
-	Info.data.accelerationStructure = BufferAddress;  // same pointer for storage, it's a union
-	VulkanRHI::vkGetDescriptorEXT(Device->GetInstanceHandle(), &Info, State.DescriptorSize, &State.DebugDescriptors[ByteOffset]);
+void FVulkanBindlessDescriptorManager::UpdateAccelerationStructure(FRHIDescriptorHandle DescriptorHandle, VkAccelerationStructureKHR AccelerationStructure)
+{
+#if VULKAN_RHI_RAYTRACING
+	if (bIsSupported)
+	{
+		check(GetDescriptorTypeForSetIndex(DescriptorHandle.GetRawType()) == VK_DESCRIPTOR_TYPE_ACCELERATION_STRUCTURE_KHR);
 
-	FMemory::Memcpy(&State.MappedPointer[ByteOffset], &State.DebugDescriptors[ByteOffset], State.DescriptorSize);
+		// :todo-jn: start caching AccelerationStructure in resources to avoid the extra call
+		VkAccelerationStructureDeviceAddressInfoKHR AccelerationStructureDeviceAddressInfo;
+		ZeroVulkanStruct(AccelerationStructureDeviceAddressInfo, VK_STRUCTURE_TYPE_ACCELERATION_STRUCTURE_DEVICE_ADDRESS_INFO_KHR);
+		AccelerationStructureDeviceAddressInfo.accelerationStructure = AccelerationStructure;
+		const VkDeviceAddress BufferAddress = VulkanRHI::vkGetAccelerationStructureDeviceAddressKHR(Device->GetInstanceHandle(), &AccelerationStructureDeviceAddressInfo);
 
-	SET_DWORD_STAT(STAT_VulkanBindlessPeakAccelerationStructure, State.PeakDescriptorCount);
-	return FRHIDescriptorHandle(SetIndex, ResourceIndex);
-#else
-	return FRHIDescriptorHandle();
+		VkDescriptorDataEXT DescriptorData;
+		DescriptorData.accelerationStructure = BufferAddress;
+		UpdateDescriptor(DescriptorHandle, DescriptorData);
+	}
 #endif
 }
 
diff --git a/Engine/Source/Runtime/VulkanRHI/Private/VulkanDescriptorSets.h b/Engine/Source/Runtime/VulkanRHI/Private/VulkanDescriptorSets.h
--- a/Engine/Source/Runtime/VulkanRHI/Private/VulkanDescriptorSets.h
+++ b/Engine/Source/Runtime/VulkanRHI/Private/VulkanDescriptorSets.h
@@ -1215,11 +1215,14 @@ public:
 
 	void BindDescriptorBuffers(VkCommandBuffer CommandBuffer, VkPipelineStageFlags SupportedStages);
 
-	FRHIDescriptorHandle RegisterSampler(VkSampler VulkanSampler);
-	FRHIDescriptorHandle RegisterImage(VkImageView VulkanImage, VkDescriptorType DescriptorType, bool bIsDepthStencil);
-	FRHIDescriptorHandle RegisterBuffer(VkBuffer VulkanBuffer, VkDeviceSize BufferOffset, VkDeviceSize BufferSize, VkDescriptorType DescriptorType);
-	FRHIDescriptorHandle RegisterTexelBuffer(const VkBufferViewCreateInfo& ViewInfo, VkDescriptorType DescriptorType);
-	FRHIDescriptorHandle RegisterAccelerationStructure(VkAccelerationStructureKHR AccelerationStructure);
+	FRHIDescriptorHandle ReserveDescriptor(VkDescriptorType DescriptorType);
+
+	void UpdateSampler(FRHIDescriptorHandle DescriptorHandle, VkSampler VulkanSampler);
+	void UpdateImage(FRHIDescriptorHandle DescriptorHandle, VkImageView VulkanImage, bool bIsDepthStencil);
+	void UpdateBuffer(FRHIDescriptorHandle DescriptorHandle, VkBuffer VulkanBuffer, VkDeviceSize BufferOffset, VkDeviceSize BufferSize);
+	void UpdateTexelBuffer(FRHIDescriptorHandle DescriptorHandle, const VkBufferViewCreateInfo& ViewInfo);
+	void UpdateAccelerationStructure(FRHIDescriptorHandle DescriptorHandle, VkAccelerationStructureKHR AccelerationStructure);
+
 	void RegisterUniformBuffers(VkCommandBuffer CommandBuffer, VkPipelineBindPoint BindPoint, const FUniformBufferDescriptorArrays& StageUBs);
 
 	void Unregister(FRHIDescriptorHandle DescriptorHandle);
@@ -1258,6 +1261,8 @@ private:
 	VkPipelineLayout BindlessPipelineLayout = VK_NULL_HANDLE;
 
 	uint32 GetFreeResourceIndex(BindlessSetState& Desc);
+	void UpdateStatsForHandle(FRHIDescriptorHandle DescriptorHandle);
+	void UpdateDescriptor(FRHIDescriptorHandle DescriptorHandle, VkDescriptorDataEXT DescriptorData);
 };
 
 
diff --git a/Engine/Source/Runtime/VulkanRHI/Private/VulkanPipelineState.cpp b/Engine/Source/Runtime/VulkanRHI/Private/VulkanPipelineState.cpp
--- a/Engine/Source/Runtime/VulkanRHI/Private/VulkanPipelineState.cpp
+++ b/Engine/Source/Runtime/VulkanRHI/Private/VulkanPipelineState.cpp
@@ -476,17 +476,16 @@ void FVulkanGraphicsPipelineDescriptorState::UpdateBindlessDescriptors(FVulkanCo
 					{
 						const FPackedUniformBuffers::FPackedBuffer& StagedUniformBuffer = PackedUniformBuffers[Stage].GetBuffer(PackedUBIndex);
 						const int32 UBSize = StagedUniformBuffer.Num();
-						const int32 PaddedUBSize = Align<int32>(UBSize, 64);  // :todo-jn: work around for NV driver issue
 						const int32 BindingIndex = StageInfo.PackedUBBindingIndices[PackedUBIndex];
 
-						const uint64 RingBufferOffset = UniformBufferUploader->AllocateMemory(PaddedUBSize, UBOffsetAlignment, CmdBuffer);
+						const uint64 RingBufferOffset = UniformBufferUploader->AllocateMemory(UBSize, UBOffsetAlignment, CmdBuffer);
 
 						// Make sure it wasn't written to already
 						VkDescriptorAddressInfoEXT& DescriptorAddressInfo = DescriptorAddressInfos[BindingIndex];
 						check(DescriptorAddressInfo.sType == 0);
 						DescriptorAddressInfo.sType = VK_STRUCTURE_TYPE_BUFFER_DEVICE_ADDRESS_INFO;
 						DescriptorAddressInfo.address = UniformBufferUploader->GetCPUBufferAddress() + RingBufferOffset;
-						DescriptorAddressInfo.range = PaddedUBSize;
+						DescriptorAddressInfo.range = UBSize;
 
 						// get location in the ring buffer to use
 						FMemory::Memcpy(CPURingBufferBase + RingBufferOffset, StagedUniformBuffer.GetData(), UBSize);
@@ -521,7 +520,7 @@ void FVulkanGraphicsPipelineDescriptorState::UpdateBindlessDescriptors(FVulkanCo
 
 				DescriptorAddressInfo.sType = VK_STRUCTURE_TYPE_BUFFER_DEVICE_ADDRESS_INFO;
 				DescriptorAddressInfo.address = BufferAddress + WriteDescriptorSet.pBufferInfo->offset;
-				DescriptorAddressInfo.range = Align<VkDeviceSize>(WriteDescriptorSet.pBufferInfo->range, 64u); // :todo-jn: work around for NV driver issue
+				DescriptorAddressInfo.range = WriteDescriptorSet.pBufferInfo->range;
 			}
 		}
 	}
diff --git a/Engine/Source/Runtime/VulkanRHI/Private/VulkanRayTracing.cpp b/Engine/Source/Runtime/VulkanRHI/Private/VulkanRayTracing.cpp
--- a/Engine/Source/Runtime/VulkanRHI/Private/VulkanRayTracing.cpp
+++ b/Engine/Source/Runtime/VulkanRHI/Private/VulkanRayTracing.cpp
@@ -676,7 +676,7 @@ void FVulkanRayTracingScene::BindBuffer(FRHIBuffer* InBuffer, uint32 InBufferOff
 		const uint32 LayerOffset = InBufferOffset + Layer.BufferOffset;
 		check(LayerOffset % GRHIRayTracingAccelerationStructureAlignment == 0);
 
-		Layer.View = MakeUnique<FVulkanView>(*Device);
+		Layer.View = MakeUnique<FVulkanView>(*Device, VK_DESCRIPTOR_TYPE_ACCELERATION_STRUCTURE_KHR);
 		Layer.View->InitAsAccelerationStructureView(
 			AccelerationStructureBuffer
 			, LayerOffset
diff --git a/Engine/Source/Runtime/VulkanRHI/Private/VulkanResources.h b/Engine/Source/Runtime/VulkanRHI/Private/VulkanResources.h
--- a/Engine/Source/Runtime/VulkanRHI/Private/VulkanResources.h
+++ b/Engine/Source/Runtime/VulkanRHI/Private/VulkanResources.h
@@ -343,6 +343,12 @@ struct FVulkanCpuReadbackBuffer
 class FVulkanView
 {
 public:
+
+	struct FInvalidatedState
+	{
+		bool bInitialized = false;
+	};
+
 	struct FTypedBufferView
 	{
 		VkBufferView View      = VK_NULL_HANDLE;
@@ -373,7 +379,7 @@ public:
 	};
 
 	typedef TVariant<
-	      FEmptyVariantState
+		  FInvalidatedState
 		, FTypedBufferView
 		, FTextureView
 		, FStructuredBufferView
@@ -384,7 +390,7 @@ public:
 
 	enum EType
 	{
-		Null                  = TStorage::IndexOfType<FEmptyVariantState    >(),
+		Null                  = TStorage::IndexOfType<FInvalidatedState     >(),
 		TypedBuffer           = TStorage::IndexOfType<FTypedBufferView      >(),
 		Texture               = TStorage::IndexOfType<FTextureView          >(),
 		StructuredBuffer      = TStorage::IndexOfType<FStructuredBufferView >(),
@@ -393,14 +399,9 @@ public:
 #endif
 	};
 
-	FVulkanView(FVulkanDevice& Device)
-		: Device(Device)
-	{}
+	FVulkanView(FVulkanDevice& InDevice, VkDescriptorType InDescriptorType);
 
-	~FVulkanView()
-	{
-		Invalidate();
-	}
+	~FVulkanView();
 
 	void Invalidate();
 
@@ -409,6 +410,11 @@ public:
 		return EType(Storage.GetIndex());
 	}
 
+	bool IsInitialized() const
+	{
+		return (GetViewType() != Null) || Storage.Get<FInvalidatedState>().bInitialized;
+	}
+
 	FTypedBufferView           const& GetTypedBufferView          () const { return Storage.Get<FTypedBufferView          >(); }
 	FTextureView               const& GetTextureView              () const { return Storage.Get<FTextureView              >(); }
 	FStructuredBufferView      const& GetStructuredBufferView     () const { return Storage.Get<FStructuredBufferView     >(); }
@@ -469,8 +475,8 @@ private:
 class FVulkanLinkedView : public FVulkanView, public TIntrusiveLinkedList<FVulkanLinkedView>
 {
 protected:
-	FVulkanLinkedView(FVulkanDevice& Device)
-		: FVulkanView(Device)
+	FVulkanLinkedView(FVulkanDevice& Device, VkDescriptorType DescriptorType)
+		: FVulkanView(Device, DescriptorType)
 	{}
 
 	~FVulkanLinkedView()
diff --git a/Engine/Source/Runtime/VulkanRHI/Private/VulkanState.cpp b/Engine/Source/Runtime/VulkanRHI/Private/VulkanState.cpp
--- a/Engine/Source/Runtime/VulkanRHI/Private/VulkanState.cpp
+++ b/Engine/Source/Runtime/VulkanRHI/Private/VulkanState.cpp
@@ -235,7 +235,9 @@ FVulkanSamplerState::FVulkanSamplerState(const VkSamplerCreateInfo& InInfo, FVul
 
 	if (InDevice.SupportsBindless() && (RHIGetRuntimeBindlessSamplersConfiguration(GMaxRHIShaderPlatform) != ERHIBindlessConfiguration::Disabled))
 	{
-		BindlessHandle = InDevice.GetBindlessDescriptorManager()->RegisterSampler(Sampler);
+		FVulkanBindlessDescriptorManager* BindlessDescriptorManager = InDevice.GetBindlessDescriptorManager();
+		BindlessHandle = BindlessDescriptorManager->ReserveDescriptor(VK_DESCRIPTOR_TYPE_SAMPLER);
+		BindlessDescriptorManager->UpdateSampler(BindlessHandle, Sampler);
 	}
 }
 
diff --git a/Engine/Source/Runtime/VulkanRHI/Private/VulkanSwapChain.cpp b/Engine/Source/Runtime/VulkanRHI/Private/VulkanSwapChain.cpp
--- a/Engine/Source/Runtime/VulkanRHI/Private/VulkanSwapChain.cpp
+++ b/Engine/Source/Runtime/VulkanRHI/Private/VulkanSwapChain.cpp
@@ -883,6 +883,7 @@ void FVulkanSwapChain::CreateQCOMDepthStencil(const FVulkanTexture& InSurface) c
 	const FRHITextureDesc& Desc = InSurface.GetDesc();
 	const ETextureCreateFlags UEFlags = Desc.Flags;
 	check(UEFlags & TexCreate_DepthStencilTargetable);
+	const VkDescriptorType DescriptorType = VK_DESCRIPTOR_TYPE_SAMPLED_IMAGE;
 
 	const FRHITextureCreateDesc CreateDesc =
 		FRHITextureCreateDesc::Create2D(TEXT("FVulkanSwapChainQCOM"), Desc.Extent.Y, Desc.Extent.X, Desc.Format) // Desc.Extent.X and Desc.Extent.Y are intentionally swapped.
@@ -897,7 +898,7 @@ void FVulkanSwapChain::CreateQCOMDepthStencil(const FVulkanTexture& InSurface) c
 	check(QCOMDepthStencilSurface->GetViewType() == VK_IMAGE_VIEW_TYPE_2D);
 	check(QCOMDepthStencilSurface->Image != VK_NULL_HANDLE);
 
-	QCOMDepthStencilView = new FVulkanView(*QCOMDepthStencilSurface->Device);
+	QCOMDepthStencilView = new FVulkanView(*QCOMDepthStencilSurface->Device, DescriptorType);
 	QCOMDepthStencilView->InitAsTextureView(
 		  QCOMDepthStencilSurface->Image
 		, QCOMDepthStencilSurface->GetViewType()
@@ -917,7 +918,7 @@ void FVulkanSwapChain::CreateQCOMDepthStencil(const FVulkanTexture& InSurface) c
 	}
 	else
 	{
-		QCOMDepthView = new FVulkanView(*QCOMDepthStencilSurface->Device);
+		QCOMDepthView = new FVulkanView(*QCOMDepthStencilSurface->Device, DescriptorType);
 		QCOMDepthView->InitAsTextureView(
 			  QCOMDepthStencilSurface->Image
 			, QCOMDepthStencilSurface->GetViewType()
diff --git a/Engine/Source/Runtime/VulkanRHI/Private/VulkanTexture.cpp b/Engine/Source/Runtime/VulkanRHI/Private/VulkanTexture.cpp
--- a/Engine/Source/Runtime/VulkanRHI/Private/VulkanTexture.cpp
+++ b/Engine/Source/Runtime/VulkanRHI/Private/VulkanTexture.cpp
@@ -1588,9 +1588,10 @@ FVulkanTexture::FVulkanTexture(FVulkanDevice& InDevice, const FRHITextureCreateD
 		checkf(StorageFormat != VK_FORMAT_UNDEFINED, TEXT("Pixel Format %d not defined!"), (int32)InCreateDesc.Format);
 	}
 
+	const VkDescriptorType DescriptorType = SupportsSampling() ? VK_DESCRIPTOR_TYPE_SAMPLED_IMAGE : VK_DESCRIPTOR_TYPE_STORAGE_IMAGE;
 	if (ViewType != VK_IMAGE_VIEW_TYPE_MAX_ENUM)
 	{
-		DefaultView = (new FVulkanView(InDevice))->InitAsTextureView(
+		DefaultView = (new FVulkanView(InDevice, DescriptorType))->InitAsTextureView(
 			Image
 			, ViewType
 			, GetFullAspectMask()
@@ -1610,7 +1611,7 @@ FVulkanTexture::FVulkanTexture(FVulkanDevice& InDevice, const FRHITextureCreateD
 	}
 	else
 	{
-		PartialView = (new FVulkanView(InDevice))->InitAsTextureView(
+		PartialView = (new FVulkanView(InDevice, DescriptorType))->InitAsTextureView(
 			Image
 			, ViewType
 			, PartialAspectMask
@@ -1771,9 +1772,10 @@ FVulkanTexture::FVulkanTexture(FVulkanDevice& InDevice, const FRHITextureCreateD
 	}
 
 	const VkImageViewType ViewType = GetViewType();
+	const VkDescriptorType DescriptorType = SupportsSampling() ? VK_DESCRIPTOR_TYPE_SAMPLED_IMAGE : VK_DESCRIPTOR_TYPE_STORAGE_IMAGE;
 	if (Image != VK_NULL_HANDLE)
 	{
-		DefaultView = (new FVulkanView(InDevice))->InitAsTextureView(
+		DefaultView = (new FVulkanView(InDevice, DescriptorType))->InitAsTextureView(
 			Image
 			, ViewType
 			, GetFullAspectMask()
@@ -1793,7 +1795,7 @@ FVulkanTexture::FVulkanTexture(FVulkanDevice& InDevice, const FRHITextureCreateD
 	}
 	else
 	{
-		PartialView = (new FVulkanView(InDevice))->InitAsTextureView(
+		PartialView = (new FVulkanView(InDevice, DescriptorType))->InitAsTextureView(
 			Image
 			, ViewType
 			, PartialAspectMask
diff --git a/Engine/Source/Runtime/VulkanRHI/Private/VulkanUAV.cpp b/Engine/Source/Runtime/VulkanRHI/Private/VulkanUAV.cpp
--- a/Engine/Source/Runtime/VulkanRHI/Private/VulkanUAV.cpp
+++ b/Engine/Source/Runtime/VulkanRHI/Private/VulkanUAV.cpp
@@ -10,13 +10,28 @@
 #include "VulkanRayTracing.h"
 #endif // VULKAN_RHI_RAYTRACING
 
-void FVulkanView::Invalidate()
+
+FVulkanView::FVulkanView(FVulkanDevice& InDevice, VkDescriptorType InDescriptorType)
+	: Device(InDevice)
 {
+	BindlessHandle = Device.GetBindlessDescriptorManager()->ReserveDescriptor(InDescriptorType);
+}
+
+FVulkanView::~FVulkanView()
+{
+	Invalidate();
+
 	if (BindlessHandle.IsValid())
 	{
 		Device.GetDeferredDeletionQueue().EnqueueBindlessHandle(BindlessHandle);
 		BindlessHandle = FRHIDescriptorHandle();
 	}
+}
+
+void FVulkanView::Invalidate()
+{
+	// Carry forward its initialized state
+	const bool bIsInitialized = IsInitialized();
 
 	switch (GetViewType())
 	{
@@ -45,11 +60,15 @@ void FVulkanView::Invalidate()
 #endif
 	}
 
-	Storage.Emplace<FEmptyVariantState>();
+	Storage.Emplace<FInvalidatedState>();
+	Storage.Get<FInvalidatedState>().bInitialized = bIsInitialized;
 }
 
 FVulkanView* FVulkanView::InitAsTypedBufferView(FVulkanResourceMultiBuffer* Buffer, EPixelFormat UEFormat, uint32 InOffset, uint32 InSize)
 {
+	// We will need a deferred update if the descriptor was already in use
+	const bool bDeferredUpdate = IsInitialized();
+
 	check(GetViewType() == EType::Null);
 	Storage.Emplace<FTypedBufferView>();
 	FTypedBufferView& TBV = Storage.Get<FTypedBufferView>();
@@ -86,7 +105,7 @@ FVulkanView* FVulkanView::InitAsTypedBufferView(FVulkanResourceMultiBuffer* Buff
 	INC_DWORD_STAT(STAT_VulkanNumBufferViews);
 	// :todo-jn: the buffer view is actually not needed in bindless anymore
 
-	BindlessHandle = Device.GetBindlessDescriptorManager()->RegisterTexelBuffer(ViewInfo, VK_DESCRIPTOR_TYPE_UNIFORM_TEXEL_BUFFER);
+	Device.GetBindlessDescriptorManager()->UpdateTexelBuffer(BindlessHandle, ViewInfo);
 
 	return this;
 }
@@ -104,6 +123,9 @@ FVulkanView* FVulkanView::InitAsTextureView(
 	, bool bUseIdentitySwizzle
 	, VkImageUsageFlags ImageUsageFlags)
 {
+	// We will need a deferred update if the descriptor was already in use
+	const bool bDeferredUpdate = IsInitialized();
+
 	check(GetViewType() == EType::Null);
 	Storage.Emplace<FTextureView>();
 	FTextureView& TV = Storage.Get<FTextureView>();
@@ -179,15 +201,17 @@ FVulkanView* FVulkanView::InitAsTextureView(
 		TV.ViewId = ++GVulkanImageViewHandleIdCounter;
 	}
 
-	bool bDepthOrStencilAspect = (AspectFlags & (VK_IMAGE_ASPECT_DEPTH_BIT | VK_IMAGE_ASPECT_STENCIL_BIT)) != 0;
-	BindlessHandle = Device.GetBindlessDescriptorManager()->RegisterImage(TV.View, VK_DESCRIPTOR_TYPE_SAMPLED_IMAGE, bDepthOrStencilAspect); // @todo sampled or storage for UAVs?
-	//DefaultBindlessHandle = Device->GetBindlessDescriptorManager()->RegisterImage(PartialView->GetTextureView().View, SupportsSampling() ? VK_DESCRIPTOR_TYPE_SAMPLED_IMAGE : VK_DESCRIPTOR_TYPE_STORAGE_IMAGE, IsDepthOrStencilAspect());
+	const bool bDepthOrStencilAspect = (AspectFlags & (VK_IMAGE_ASPECT_DEPTH_BIT | VK_IMAGE_ASPECT_STENCIL_BIT)) != 0;
+	Device.GetBindlessDescriptorManager()->UpdateImage(BindlessHandle, TV.View, bDepthOrStencilAspect);
 
 	return this;
 }
 
 FVulkanView* FVulkanView::InitAsStructuredBufferView(FVulkanResourceMultiBuffer* Buffer, uint32 InOffset, uint32 InSize)
 {
+	// We will need a deferred update if the descriptor was already in use
+	const bool bDeferredUpdate = IsInitialized();
+
 	check(GetViewType() == EType::Null);
 	Storage.Emplace<FStructuredBufferView>();
 	FStructuredBufferView& SBV = Storage.Get<FStructuredBufferView>();
@@ -199,12 +223,8 @@ FVulkanView* FVulkanView::InitAsStructuredBufferView(FVulkanResourceMultiBuffer*
 	SBV.Offset = TotalOffset;
 	SBV.Size = InSize;
 
-	BindlessHandle = Device.GetBindlessDescriptorManager()->RegisterBuffer(
-		Buffer->GetHandle(),
-		TotalOffset,
-		InSize,
-		VK_DESCRIPTOR_TYPE_STORAGE_BUFFER
-	);
+	Device.GetBindlessDescriptorManager()->UpdateBuffer(
+		BindlessHandle, Buffer->GetHandle(), TotalOffset, InSize);
 
 	return this;
 }
@@ -225,7 +245,7 @@ FVulkanView* FVulkanView::InitAsAccelerationStructureView(FVulkanResourceMultiBu
 
 	VERIFYVULKANRESULT(VulkanDynamicAPI::vkCreateAccelerationStructureKHR(Device.GetInstanceHandle(), &CreateInfo, VULKAN_CPU_ALLOCATOR, &ASV.Handle));
 
-	BindlessHandle = Device.GetBindlessDescriptorManager()->RegisterAccelerationStructure(ASV.Handle);
+	Device.GetBindlessDescriptorManager()->UpdateAccelerationStructure(BindlessHandle, ASV.Handle);
 
 	return this;
 }
@@ -260,11 +280,72 @@ static VkImageViewType GetVkImageViewTypeForDimension(FRHIViewDesc::EDimension D
 	return VK_IMAGE_VIEW_TYPE_MAX_ENUM;
 }
 
+static VkDescriptorType GetDescriptorTypeForViewDesc(FRHIViewDesc const& ViewDesc)
+{
+	if (ViewDesc.IsBuffer())
+	{
+		if (ViewDesc.IsSRV())
+		{
+			switch (ViewDesc.Buffer.SRV.BufferType)
+			{
+			case FRHIViewDesc::EBufferType::Raw:
+			case FRHIViewDesc::EBufferType::Structured:
+				return VK_DESCRIPTOR_TYPE_STORAGE_BUFFER;
+
+			case FRHIViewDesc::EBufferType::Typed:
+				return VK_DESCRIPTOR_TYPE_UNIFORM_TEXEL_BUFFER;
+
+			case FRHIViewDesc::EBufferType::AccelerationStructure:
+				return VK_DESCRIPTOR_TYPE_ACCELERATION_STRUCTURE_KHR;
+
+			default:
+				checkNoEntry();
+				return VK_DESCRIPTOR_TYPE_MAX_ENUM;
+			}
+		}
+		else
+		{
+			switch (ViewDesc.Buffer.UAV.BufferType)
+			{
+			case FRHIViewDesc::EBufferType::Raw:
+			case FRHIViewDesc::EBufferType::Structured:
+				return VK_DESCRIPTOR_TYPE_STORAGE_BUFFER;
+
+			case FRHIViewDesc::EBufferType::Typed:
+				return VK_DESCRIPTOR_TYPE_STORAGE_TEXEL_BUFFER;
+
+			case FRHIViewDesc::EBufferType::AccelerationStructure:
+				return VK_DESCRIPTOR_TYPE_ACCELERATION_STRUCTURE_KHR;
+
+			default:
+				checkNoEntry();
+				return VK_DESCRIPTOR_TYPE_MAX_ENUM;
+			}
+		}
+	}
+	else
+	{
+		if (ViewDesc.IsSRV())
+		{
+			// Sampled images aren't supported in R64, shadercompiler patches them to storage image
+			if (ViewDesc.Texture.SRV.Format == PF_R64_UINT)
+			{
+				return VK_DESCRIPTOR_TYPE_STORAGE_IMAGE;
+			}
+
+			return VK_DESCRIPTOR_TYPE_SAMPLED_IMAGE;
+		}
+		else
+		{
+			return VK_DESCRIPTOR_TYPE_STORAGE_IMAGE;
+		}
+	}
+}
 
 
 FVulkanShaderResourceView::FVulkanShaderResourceView(FRHICommandListImmediate& RHICmdList, FVulkanDevice& InDevice, FRHIViewableResource* InResource, FRHIViewDesc const& InViewDesc)
 	: FRHIShaderResourceView(InResource, InViewDesc)
-	, FVulkanLinkedView(InDevice)
+	, FVulkanLinkedView(InDevice, GetDescriptorTypeForViewDesc(InViewDesc))
 {
 	RHICmdList.EnqueueLambda([this](FRHICommandListImmediate&)
 	{
@@ -352,7 +433,7 @@ void FVulkanShaderResourceView::UpdateView()
 
 FVulkanUnorderedAccessView::FVulkanUnorderedAccessView(FRHICommandListImmediate& RHICmdList, FVulkanDevice& InDevice, FRHIViewableResource* InResource, FRHIViewDesc const& InViewDesc)
 	: FRHIUnorderedAccessView(InResource, InViewDesc)
-	, FVulkanLinkedView(InDevice)
+	, FVulkanLinkedView(InDevice, GetDescriptorTypeForViewDesc(InViewDesc))
 {
 	RHICmdList.EnqueueLambda([this](FRHICommandListImmediate&)
 	{
diff --git a/Engine/Source/Runtime/VulkanRHI/Private/VulkanViewport.cpp b/Engine/Source/Runtime/VulkanRHI/Private/VulkanViewport.cpp
--- a/Engine/Source/Runtime/VulkanRHI/Private/VulkanViewport.cpp
+++ b/Engine/Source/Runtime/VulkanRHI/Private/VulkanViewport.cpp
@@ -306,7 +306,8 @@ FVulkanFramebuffer::FVulkanFramebuffer(FVulkanDevice& Device, const FRHISetRende
 
 	auto CreateOwnedView = [&]()
 	{
-		FVulkanView* View = new FVulkanView(Device);
+		const VkDescriptorType DescriptorType = VK_DESCRIPTOR_TYPE_SAMPLED_IMAGE;
+		FVulkanView* View = new FVulkanView(Device, DescriptorType);
 		AttachmentTextureViews.Add(View);
 		OwnedTextureViews.Add(View);
 		return View;
@@ -720,7 +721,9 @@ void FVulkanViewport::CreateSwapchain(FVulkanSwapChainRecreateInfo* RecreateInfo
 		for (int32 Index = 0; Index < Images.Num(); ++Index)
 		{
 			BackBufferImages[Index] = Images[Index];
-			TextureViews.Add((new FVulkanView(*Device))->InitAsTextureView(Images[Index], VK_IMAGE_VIEW_TYPE_2D, VK_IMAGE_ASPECT_COLOR_BIT, PixelFormat, UEToVkTextureFormat(PixelFormat, false), 0, 1, 0, 1, false));
+			const VkDescriptorType DescriptorType = VK_DESCRIPTOR_TYPE_SAMPLED_IMAGE;
+			TextureViews.Add((new FVulkanView(*Device, DescriptorType))->InitAsTextureView(
+				Images[Index], VK_IMAGE_VIEW_TYPE_2D, VK_IMAGE_ASPECT_COLOR_BIT, PixelFormat, UEToVkTextureFormat(PixelFormat, false), 0, 1, 0, 1, false));
 
 			// Clear the swapchain to avoid a validation warning, and transition to ColorAttachment
 			{
