Fixed issue where UAV Cube maps were incorrectly create as Cube view  instead of Texture2DArray

diff --git a/Engine/Source/Runtime/VulkanRHI/Private/VulkanUAV.cpp b/Engine/Source/Runtime/VulkanRHI/Private/VulkanUAV.cpp
--- a/Engine/Source/Runtime/VulkanRHI/Private/VulkanUAV.cpp
+++ b/Engine/Source/Runtime/VulkanRHI/Private/VulkanUAV.cpp
@@ -263,7 +263,7 @@ void FVulkanViewableResource::UpdateLinkedViews()
 }
 
 
-static VkImageViewType GetVkImageViewTypeForDimension(FRHIViewDesc::EDimension DescDimension, VkImageViewType TextureViewType)
+static VkImageViewType GetVkImageViewTypeForDimensionSRV(FRHIViewDesc::EDimension DescDimension, VkImageViewType TextureViewType)
 {
 	switch (DescDimension)
 	{
@@ -280,6 +280,23 @@ static VkImageViewType GetVkImageViewTypeForDimension(FRHIViewDesc::EDimension D
 	return VK_IMAGE_VIEW_TYPE_MAX_ENUM;
 }
 
+static VkImageViewType GetVkImageViewTypeForDimensionUAV(FRHIViewDesc::EDimension DescDimension, VkImageViewType TextureViewType)
+{
+	switch (DescDimension)
+	{
+	case FRHIViewDesc::EDimension::Texture2D:			return VK_IMAGE_VIEW_TYPE_2D; break;
+	case FRHIViewDesc::EDimension::Texture2DArray:		return VK_IMAGE_VIEW_TYPE_2D_ARRAY; break;
+	case FRHIViewDesc::EDimension::Texture3D:			return VK_IMAGE_VIEW_TYPE_3D; break;
+	case FRHIViewDesc::EDimension::TextureCube:			return VK_IMAGE_VIEW_TYPE_2D_ARRAY; break;
+	case FRHIViewDesc::EDimension::TextureCubeArray:	return VK_IMAGE_VIEW_TYPE_2D_ARRAY; break;
+	case FRHIViewDesc::EDimension::Unknown:				return TextureViewType; break;
+	default: break;
+	}
+
+	checkf(false, TEXT("Unknown texture dimension value!"));
+	return VK_IMAGE_VIEW_TYPE_MAX_ENUM;
+
+}
 static VkDescriptorType GetDescriptorTypeForViewDesc(FRHIViewDesc const& ViewDesc)
 {
 	if (ViewDesc.IsBuffer())
@@ -416,7 +433,7 @@ void FVulkanShaderResourceView::UpdateView()
 
 		InitAsTextureView(
 			  Texture->Image
-			, GetVkImageViewTypeForDimension(Info.Dimension, Texture->GetViewType())
+			, GetVkImageViewTypeForDimensionSRV(Info.Dimension, Texture->GetViewType())
 			, Texture->GetPartialAspectMask()
 			, Info.Format
 			, UEToVkTextureFormat(Info.Format, Info.bSRGB)
@@ -506,7 +523,7 @@ void FVulkanUnorderedAccessView::UpdateView()
 
 		InitAsTextureView(
 			  Texture->Image
-			, GetVkImageViewTypeForDimension(Info.Dimension, Texture->GetViewType())
+			, GetVkImageViewTypeForDimensionUAV(Info.Dimension, Texture->GetViewType())
 			, Texture->GetPartialAspectMask()
 			, Info.Format
 			, UEToVkTextureFormat(Info.Format, false)
