VulkanRHI: Set debug name on shader modules when available to improve validation layer messages.

#rnx

diff --git a/Engine/Source/Runtime/VulkanRHI/Private/VulkanShaders.cpp b/Engine/Source/Runtime/VulkanRHI/Private/VulkanShaders.cpp
--- a/Engine/Source/Runtime/VulkanRHI/Private/VulkanShaders.cpp
+++ b/Engine/Source/Runtime/VulkanRHI/Private/VulkanShaders.cpp
@@ -378,6 +378,10 @@ TRefCountPtr<FVulkanShaderModule> FVulkanShader::CreateHandle(const FVulkanLayou
 	Layout->PatchSpirvBindings(Spirv, Frequency, CodeHeader, StageFlag);
 	TRefCountPtr<FVulkanShaderModule> Module = CreateShaderModule(Device, Spirv);
 	ShaderModules.Add(LayoutHash, Module);
+	if (!CodeHeader.DebugName.IsEmpty())
+	{
+		VULKAN_SET_DEBUG_NAME((*Device), VK_OBJECT_TYPE_SHADER_MODULE, Module->GetVkShaderModule(), TEXT("%s : (FVulkanShader*)0x%p"), *CodeHeader.DebugName, this);
+	}
 	return Module;
 }
 
