VulkanRHI: Bring back RHIUpdateTextureReference for bindless.  Use staging buffers for deferred descriptor updates for now.

#rnx

diff --git a/Engine/Source/Runtime/VulkanRHI/Private/VulkanDescriptorSets.cpp b/Engine/Source/Runtime/VulkanRHI/Private/VulkanDescriptorSets.cpp
--- a/Engine/Source/Runtime/VulkanRHI/Private/VulkanDescriptorSets.cpp
+++ b/Engine/Source/Runtime/VulkanRHI/Private/VulkanDescriptorSets.cpp
@@ -6,6 +6,7 @@
 
 #include "VulkanRHIPrivate.h"
 #include "VulkanDescriptorSets.h"
+#include "VulkanContext.h"
 
 
 int32 GVulkanBindlessEnabled = 0;
@@ -720,7 +721,7 @@ FRHIDescriptorHandle FVulkanBindlessDescriptorManager::ReserveDescriptor(VkDescr
 	return FRHIDescriptorHandle();
 }
 
-void FVulkanBindlessDescriptorManager::UpdateDescriptor(FRHIDescriptorHandle DescriptorHandle, VkDescriptorDataEXT DescriptorData)
+void FVulkanBindlessDescriptorManager::UpdateDescriptor(FRHIDescriptorHandle DescriptorHandle, VkDescriptorDataEXT DescriptorData, bool bImmediateUpdate)
 {
 	checkf(DescriptorHandle.IsValid(), TEXT("Attemping to update invalid descriptor handle!"));
 
@@ -734,7 +735,53 @@ void FVulkanBindlessDescriptorManager::UpdateDescriptor(FRHIDescriptorHandle Des
 	Info.data = DescriptorData;
 	VulkanRHI::vkGetDescriptorEXT(Device->GetInstanceHandle(), &Info, State.DescriptorSize, &State.DebugDescriptors[ByteOffset]);
 
-	FMemory::Memcpy(&State.MappedPointer[ByteOffset], &State.DebugDescriptors[ByteOffset], State.DescriptorSize);
+	if (bImmediateUpdate)
+	{
+		FMemory::Memcpy(&State.MappedPointer[ByteOffset], &State.DebugDescriptors[ByteOffset], State.DescriptorSize);
+	}
+	else
+	{
+		check(!IsInRenderingThread() || FRHICommandListExecutor::GetImmediateCommandList().Bypass() || !IsRunningRHIInSeparateThread());
+		FVulkanCmdBuffer* CmdBuffer = Device->GetImmediateContext().GetCommandBufferManager()->GetActiveCmdBuffer();
+
+		// :todo-jn:  Hack to avoid barriers/copies in renderpasses
+		if (CmdBuffer->IsInsideRenderPass())
+		{
+			FMemory::Memcpy(&State.MappedPointer[ByteOffset], &State.DebugDescriptors[ByteOffset], State.DescriptorSize);
+		}
+		else
+		{
+			FStagingBuffer* StagingBuffer = Device->GetStagingManager().AcquireBuffer(State.DescriptorSize);
+			FMemory::Memcpy(StagingBuffer->GetMappedPointer(), &State.DebugDescriptors[ByteOffset], State.DescriptorSize);
+			{
+				VkMemoryBarrier2 MemoryBarrier;
+				ZeroVulkanStruct(MemoryBarrier, VK_STRUCTURE_TYPE_MEMORY_BARRIER_2);
+				MemoryBarrier.srcStageMask = VK_PIPELINE_STAGE_2_ALL_COMMANDS_BIT;
+				MemoryBarrier.srcAccessMask = VK_ACCESS_2_DESCRIPTOR_BUFFER_READ_BIT_EXT | VK_ACCESS_2_MEMORY_READ_BIT | VK_ACCESS_2_MEMORY_WRITE_BIT;
+				MemoryBarrier.dstStageMask = VK_PIPELINE_STAGE_2_TRANSFER_BIT;
+				MemoryBarrier.dstAccessMask = VK_ACCESS_2_TRANSFER_READ_BIT | VK_ACCESS_2_TRANSFER_WRITE_BIT;
+
+				VkDependencyInfo DependencyInfo;
+				ZeroVulkanStruct(DependencyInfo, VK_STRUCTURE_TYPE_DEPENDENCY_INFO);
+				DependencyInfo.memoryBarrierCount = 1;
+				DependencyInfo.pMemoryBarriers = &MemoryBarrier;
+				VulkanRHI::vkCmdPipelineBarrier2KHR(CmdBuffer->GetHandle(), &DependencyInfo);
+
+				VkBufferCopy Region = {};
+				Region.srcOffset = 0;
+				Region.dstOffset = ByteOffset;
+				Region.size = State.DescriptorSize;
+				VulkanRHI::vkCmdCopyBuffer(CmdBuffer->GetHandle(), StagingBuffer->GetHandle(), State.BufferHandle, 1, &Region);
+
+				MemoryBarrier.srcStageMask = VK_PIPELINE_STAGE_2_TRANSFER_BIT;
+				MemoryBarrier.srcAccessMask = VK_ACCESS_2_TRANSFER_READ_BIT | VK_ACCESS_2_TRANSFER_WRITE_BIT;
+				MemoryBarrier.dstStageMask = VK_PIPELINE_STAGE_2_ALL_COMMANDS_BIT;
+				MemoryBarrier.dstAccessMask = VK_ACCESS_2_DESCRIPTOR_BUFFER_READ_BIT_EXT;
+				VulkanRHI::vkCmdPipelineBarrier2KHR(CmdBuffer->GetHandle(), &DependencyInfo);
+			}
+			Device->GetStagingManager().ReleaseBuffer(CmdBuffer, StagingBuffer);
+		}
+	}
 
 	UpdateStatsForHandle(DescriptorHandle);
 }
@@ -745,11 +792,11 @@ void FVulkanBindlessDescriptorManager::UpdateSampler(FRHIDescriptorHandle Descri
 	{
 		VkDescriptorDataEXT DescriptorData;
 		DescriptorData.pSampler = &VulkanSampler;
-		UpdateDescriptor(DescriptorHandle, DescriptorData);
+		UpdateDescriptor(DescriptorHandle, DescriptorData, true);
 	}
 }
 
-void FVulkanBindlessDescriptorManager::UpdateImage(FRHIDescriptorHandle DescriptorHandle, VkImageView ImageView, bool bIsDepthStencil)
+void FVulkanBindlessDescriptorManager::UpdateImage(FRHIDescriptorHandle DescriptorHandle, VkImageView ImageView, bool bIsDepthStencil, bool bImmediateUpdate)
 {
 	if (bIsSupported)
 	{
@@ -764,11 +811,11 @@ void FVulkanBindlessDescriptorManager::UpdateImage(FRHIDescriptorHandle Descript
 
 		VkDescriptorDataEXT DescriptorData;
 		DescriptorData.pSampledImage = &DescriptorImageInfo;  // same pointer for storage, it's a union
-		UpdateDescriptor(DescriptorHandle, DescriptorData);
+		UpdateDescriptor(DescriptorHandle, DescriptorData, bImmediateUpdate);
 	}
 }
 
-void FVulkanBindlessDescriptorManager::UpdateBuffer(FRHIDescriptorHandle DescriptorHandle, VkBuffer VulkanBuffer, VkDeviceSize BufferOffset, VkDeviceSize BufferSize)
+void FVulkanBindlessDescriptorManager::UpdateBuffer(FRHIDescriptorHandle DescriptorHandle, VkBuffer VulkanBuffer, VkDeviceSize BufferOffset, VkDeviceSize BufferSize, bool bImmediateUpdate)
 {
 	if (bIsSupported)
 	{
@@ -788,11 +835,11 @@ void FVulkanBindlessDescriptorManager::UpdateBuffer(FRHIDescriptorHandle Descrip
 
 		VkDescriptorDataEXT DescriptorData;
 		DescriptorData.pStorageBuffer = &AddressInfo;  // same pointer for uniform, it's a union
-		UpdateDescriptor(DescriptorHandle, DescriptorData);
+		UpdateDescriptor(DescriptorHandle, DescriptorData, bImmediateUpdate);
 	}
 }
 
-void FVulkanBindlessDescriptorManager::UpdateTexelBuffer(FRHIDescriptorHandle DescriptorHandle, const VkBufferViewCreateInfo& ViewInfo)
+void FVulkanBindlessDescriptorManager::UpdateTexelBuffer(FRHIDescriptorHandle DescriptorHandle, const VkBufferViewCreateInfo& ViewInfo, bool bImmediateUpdate)
 {
 	if (bIsSupported)
 	{
@@ -813,11 +860,11 @@ void FVulkanBindlessDescriptorManager::UpdateTexelBuffer(FRHIDescriptorHandle De
 
 		VkDescriptorDataEXT DescriptorData;
 		DescriptorData.pUniformTexelBuffer = &AddressInfo;  // same pointer for storage, it's a union
-		UpdateDescriptor(DescriptorHandle, DescriptorData);
+		UpdateDescriptor(DescriptorHandle, DescriptorData, bImmediateUpdate);
 	}
 }
 
-void FVulkanBindlessDescriptorManager::UpdateAccelerationStructure(FRHIDescriptorHandle DescriptorHandle, VkAccelerationStructureKHR AccelerationStructure)
+void FVulkanBindlessDescriptorManager::UpdateAccelerationStructure(FRHIDescriptorHandle DescriptorHandle, VkAccelerationStructureKHR AccelerationStructure, bool bImmediateUpdate)
 {
 #if VULKAN_RHI_RAYTRACING
 	if (bIsSupported)
@@ -832,7 +879,7 @@ void FVulkanBindlessDescriptorManager::UpdateAccelerationStructure(FRHIDescripto
 
 		VkDescriptorDataEXT DescriptorData;
 		DescriptorData.accelerationStructure = BufferAddress;
-		UpdateDescriptor(DescriptorHandle, DescriptorData);
+		UpdateDescriptor(DescriptorHandle, DescriptorData, bImmediateUpdate);
 	}
 #endif
 }
@@ -884,45 +931,3 @@ void FVulkanBindlessDescriptorManager::Unregister(FRHIDescriptorHandle Descripto
 	}
 }
 
-void FVulkanBindlessDescriptorManager::CopyDescriptor(VkCommandBuffer CommandBuffer, FRHIDescriptorHandle DstHandle, const FRHIDescriptorHandle SrcHandle)
-{
-	checkf(bIsSupported, TEXT("Trying to CopyDescriptor but bindless is not supported!"));
-	checkf(DstHandle.IsValid(), TEXT("Destination handle for CopyDescriptor not valid!"));
-	checkf(SrcHandle.IsValid(), TEXT("Source handle for CopyDescriptor not valid!"));
-	checkf(DstHandle.GetRawType() == SrcHandle.GetRawType(), TEXT("Handles types for CopyDescriptor do not match!"));
-
-	const uint8 SetIndex = DstHandle.GetRawType();
-	BindlessSetState& State = BindlessSetStates[SetIndex];
-	const uint32 SrcByteOffset = SrcHandle.GetIndex() * State.DescriptorSize;
-	const uint32 DstByteOffset = DstHandle.GetIndex() * State.DescriptorSize;
-
-	FMemory::Memcpy(&State.DebugDescriptors[DstByteOffset], &State.DebugDescriptors[SrcByteOffset], State.DescriptorSize);
-
-	// Can debug using a copy directly instead of using GPU, timing will not be correct however
-	//FMemory::Memcpy(&State.MappedPointer[DstByteOffset], &State.DebugDescriptors[DstByteOffset], State.DescriptorSize);
-
-	VkMemoryBarrier2 MemoryBarrier;
-	ZeroVulkanStruct(MemoryBarrier, VK_STRUCTURE_TYPE_MEMORY_BARRIER_2);
-	MemoryBarrier.srcStageMask = VK_PIPELINE_STAGE_2_ALL_COMMANDS_BIT;
-	MemoryBarrier.srcAccessMask = VK_ACCESS_2_DESCRIPTOR_BUFFER_READ_BIT_EXT;
-	MemoryBarrier.dstStageMask = VK_PIPELINE_STAGE_2_TRANSFER_BIT;
-	MemoryBarrier.dstAccessMask = VK_ACCESS_2_TRANSFER_READ_BIT | VK_ACCESS_2_TRANSFER_WRITE_BIT;
-
-	VkDependencyInfo DependencyInfo;
-	ZeroVulkanStruct(DependencyInfo, VK_STRUCTURE_TYPE_DEPENDENCY_INFO);
-	DependencyInfo.memoryBarrierCount = 1;
-	DependencyInfo.pMemoryBarriers = &MemoryBarrier;
-	VulkanRHI::vkCmdPipelineBarrier2KHR(CommandBuffer, &DependencyInfo);
-
-	VkBufferCopy Region = {};
-	Region.srcOffset = SrcByteOffset;
-	Region.dstOffset = DstByteOffset;
-	Region.size = State.DescriptorSize;
-	VulkanRHI::vkCmdCopyBuffer(CommandBuffer, State.BufferHandle, State.BufferHandle, 1, &Region);
-
-	MemoryBarrier.srcStageMask = VK_PIPELINE_STAGE_2_TRANSFER_BIT;
-	MemoryBarrier.srcAccessMask = VK_ACCESS_2_TRANSFER_READ_BIT | VK_ACCESS_2_TRANSFER_WRITE_BIT;
-	MemoryBarrier.dstStageMask = VK_PIPELINE_STAGE_2_ALL_COMMANDS_BIT;
-	MemoryBarrier.dstAccessMask = VK_ACCESS_2_DESCRIPTOR_BUFFER_READ_BIT_EXT;
-	VulkanRHI::vkCmdPipelineBarrier2KHR(CommandBuffer, &DependencyInfo);
-}
diff --git a/Engine/Source/Runtime/VulkanRHI/Private/VulkanDescriptorSets.h b/Engine/Source/Runtime/VulkanRHI/Private/VulkanDescriptorSets.h
--- a/Engine/Source/Runtime/VulkanRHI/Private/VulkanDescriptorSets.h
+++ b/Engine/Source/Runtime/VulkanRHI/Private/VulkanDescriptorSets.h
@@ -1218,16 +1218,15 @@ public:
 	FRHIDescriptorHandle ReserveDescriptor(VkDescriptorType DescriptorType);
 
 	void UpdateSampler(FRHIDescriptorHandle DescriptorHandle, VkSampler VulkanSampler);
-	void UpdateImage(FRHIDescriptorHandle DescriptorHandle, VkImageView VulkanImage, bool bIsDepthStencil);
-	void UpdateBuffer(FRHIDescriptorHandle DescriptorHandle, VkBuffer VulkanBuffer, VkDeviceSize BufferOffset, VkDeviceSize BufferSize);
-	void UpdateTexelBuffer(FRHIDescriptorHandle DescriptorHandle, const VkBufferViewCreateInfo& ViewInfo);
-	void UpdateAccelerationStructure(FRHIDescriptorHandle DescriptorHandle, VkAccelerationStructureKHR AccelerationStructure);
+	void UpdateImage(FRHIDescriptorHandle DescriptorHandle, VkImageView VulkanImage, bool bIsDepthStencil, bool bImmediateUpdate = true);
+	void UpdateBuffer(FRHIDescriptorHandle DescriptorHandle, VkBuffer VulkanBuffer, VkDeviceSize BufferOffset, VkDeviceSize BufferSize, bool bImmediateUpdate = true);
+	void UpdateTexelBuffer(FRHIDescriptorHandle DescriptorHandle, const VkBufferViewCreateInfo& ViewInfo, bool bImmediateUpdate = true);
+	void UpdateAccelerationStructure(FRHIDescriptorHandle DescriptorHandle, VkAccelerationStructureKHR AccelerationStructure, bool bImmediateUpdate = true);
 
 	void RegisterUniformBuffers(VkCommandBuffer CommandBuffer, VkPipelineBindPoint BindPoint, const FUniformBufferDescriptorArrays& StageUBs);
 
 	void Unregister(FRHIDescriptorHandle DescriptorHandle);
 
-	void CopyDescriptor(VkCommandBuffer CommandBuffer, FRHIDescriptorHandle DstHandle, const FRHIDescriptorHandle SrcHandle);
 
 private:
 	const bool bIsSupported;
@@ -1262,7 +1261,7 @@ private:
 
 	uint32 GetFreeResourceIndex(BindlessSetState& Desc);
 	void UpdateStatsForHandle(FRHIDescriptorHandle DescriptorHandle);
-	void UpdateDescriptor(FRHIDescriptorHandle DescriptorHandle, VkDescriptorDataEXT DescriptorData);
+	void UpdateDescriptor(FRHIDescriptorHandle DescriptorHandle, VkDescriptorDataEXT DescriptorData, bool bImmediateUpdate);
 };
 
 
diff --git a/Engine/Source/Runtime/VulkanRHI/Private/VulkanResources.h b/Engine/Source/Runtime/VulkanRHI/Private/VulkanResources.h
--- a/Engine/Source/Runtime/VulkanRHI/Private/VulkanResources.h
+++ b/Engine/Source/Runtime/VulkanRHI/Private/VulkanResources.h
@@ -570,9 +570,6 @@ public:
 		return DefaultView->GetBindlessHandle();
 	}
 
-	// @todo - SRV/UAV refactor - fix this
-	//void CopyBindlessHandle(FVulkanCommandListContext& Context, FRHIDescriptorHandle DestHandle);
-
 	struct FImageCreateInfo
 	{
 		VkImageCreateInfo ImageCreateInfo;
diff --git a/Engine/Source/Runtime/VulkanRHI/Private/VulkanTexture.cpp b/Engine/Source/Runtime/VulkanRHI/Private/VulkanTexture.cpp
--- a/Engine/Source/Runtime/VulkanRHI/Private/VulkanTexture.cpp
+++ b/Engine/Source/Runtime/VulkanRHI/Private/VulkanTexture.cpp
@@ -870,17 +870,6 @@ void FVulkanTexture::SetInitialImageState(FVulkanCommandListContext& Context, Vk
 	CmdBuffer->GetLayoutManager().SetFullLayout(*this, InitialLayout);
 }
 
-// @todo - SRV/UAV refactor - fix this
-/*void FVulkanTexture::CopyBindlessHandle(FVulkanCommandListContext& Context, FRHIDescriptorHandle DestHandle)
-{
-	FVulkanCmdBuffer* CmdBuffer = Context.GetCommandBufferManager()->GetActiveCmdBuffer();
-
-	// Copy the our descriptor in the old texture's slot
-	checkSlow(DefaultBindlessHandle.IsValid());
-	checkSlow(DestHandle.IsValid());
-	Device->GetBindlessDescriptorManager()->CopyDescriptor(CmdBuffer->GetHandle(), DestHandle, DefaultBindlessHandle);
-}*/
-
 
 
 /*-----------------------------------------------------------------------------
@@ -921,20 +910,42 @@ uint32 FVulkanDynamicRHI::RHIComputeMemorySize(FRHITexture* TextureRHI)
 	return ResourceCast(TextureRHI)->GetMemorySize();
 }
 
-void FVulkanDynamicRHI::RHIUpdateTextureReference(FRHITextureReference* TextureRef, FRHITexture* NewTexture)
+void FVulkanDynamicRHI::RHIUpdateTextureReference(FRHITextureReference* TextureRef, FRHITexture* InNewTexture)
 {
 	if (Device->SupportsBindless())
 	{
-		checkNoEntry(); // @todo - SRV/UAV refactor - fix this
-		/*if (TextureRef->GetReferencedTexture() && NewTexture)
+		FRHITexture* NewTexture = InNewTexture ? InNewTexture : FRHITextureReference::GetDefaultTexture();
+
+		if (FRHIShaderResourceView* TextureRefSRV = TextureRef ? TextureRef->GetBindlessView() : nullptr)
 		{
-			// Continue using the old slot for the new texture, this way indexes stored in buffers are still valid
-			FVulkanTexture* NextTexture = ResourceCast(NewTexture);
-			NextTexture->CopyBindlessHandle(Device->GetImmediateContext(), TextureRef->GetDefaultBindlessHandle());
-		}*/
+			FVulkanShaderResourceView* VulkanTextureRefSRV = ResourceCast(TextureRefSRV);
+			FRHIDescriptorHandle DestHandle = VulkanTextureRefSRV->GetBindlessHandle();
+
+			if (DestHandle.IsValid())
+			{
+				checkf(VulkanTextureRefSRV->IsInitialized(), TEXT("TextureReference should always be created with a view of the default texture at least"));
+
+				FVulkanTexture* NewVulkanTexture = ResourceCast(NewTexture);
+				const FRHITextureDesc& Desc = NewVulkanTexture->GetDesc();
+
+				VulkanTextureRefSRV->Invalidate();
+				VulkanTextureRefSRV->InitAsTextureView(
+					  NewVulkanTexture->Image
+					, NewVulkanTexture->GetViewType()
+					, NewVulkanTexture->GetPartialAspectMask()
+					, Desc.Format
+					, NewVulkanTexture->ViewFormat
+					, 0u
+					, FMath::Max(Desc.NumMips, (uint8)1u)
+					, 0u
+					, NewVulkanTexture->GetNumberOfArrayLevels()
+					, !NewVulkanTexture->SupportsSampling());
+			}
+		}
+
 	}
 
-	FDynamicRHI::RHIUpdateTextureReference(TextureRef, NewTexture);
+	FDynamicRHI::RHIUpdateTextureReference(TextureRef, InNewTexture);
 }
 
 /*-----------------------------------------------------------------------------
diff --git a/Engine/Source/Runtime/VulkanRHI/Private/VulkanUAV.cpp b/Engine/Source/Runtime/VulkanRHI/Private/VulkanUAV.cpp
--- a/Engine/Source/Runtime/VulkanRHI/Private/VulkanUAV.cpp
+++ b/Engine/Source/Runtime/VulkanRHI/Private/VulkanUAV.cpp
@@ -67,7 +67,7 @@ void FVulkanView::Invalidate()
 FVulkanView* FVulkanView::InitAsTypedBufferView(FVulkanResourceMultiBuffer* Buffer, EPixelFormat UEFormat, uint32 InOffset, uint32 InSize)
 {
 	// We will need a deferred update if the descriptor was already in use
-	const bool bDeferredUpdate = IsInitialized();
+	const bool bImmediateUpdate = !IsInitialized();
 
 	check(GetViewType() == EType::Null);
 	Storage.Emplace<FTypedBufferView>();
@@ -105,7 +105,7 @@ FVulkanView* FVulkanView::InitAsTypedBufferView(FVulkanResourceMultiBuffer* Buff
 	INC_DWORD_STAT(STAT_VulkanNumBufferViews);
 	// :todo-jn: the buffer view is actually not needed in bindless anymore
 
-	Device.GetBindlessDescriptorManager()->UpdateTexelBuffer(BindlessHandle, ViewInfo);
+	Device.GetBindlessDescriptorManager()->UpdateTexelBuffer(BindlessHandle, ViewInfo, bImmediateUpdate);
 
 	return this;
 }
@@ -124,7 +124,7 @@ FVulkanView* FVulkanView::InitAsTextureView(
 	, VkImageUsageFlags ImageUsageFlags)
 {
 	// We will need a deferred update if the descriptor was already in use
-	const bool bDeferredUpdate = IsInitialized();
+	const bool bImmediateUpdate = !IsInitialized();
 
 	check(GetViewType() == EType::Null);
 	Storage.Emplace<FTextureView>();
@@ -202,7 +202,7 @@ FVulkanView* FVulkanView::InitAsTextureView(
 	}
 
 	const bool bDepthOrStencilAspect = (AspectFlags & (VK_IMAGE_ASPECT_DEPTH_BIT | VK_IMAGE_ASPECT_STENCIL_BIT)) != 0;
-	Device.GetBindlessDescriptorManager()->UpdateImage(BindlessHandle, TV.View, bDepthOrStencilAspect);
+	Device.GetBindlessDescriptorManager()->UpdateImage(BindlessHandle, TV.View, bDepthOrStencilAspect, bImmediateUpdate);
 
 	return this;
 }
@@ -210,7 +210,7 @@ FVulkanView* FVulkanView::InitAsTextureView(
 FVulkanView* FVulkanView::InitAsStructuredBufferView(FVulkanResourceMultiBuffer* Buffer, uint32 InOffset, uint32 InSize)
 {
 	// We will need a deferred update if the descriptor was already in use
-	const bool bDeferredUpdate = IsInitialized();
+	const bool bImmediateUpdate = !IsInitialized();
 
 	check(GetViewType() == EType::Null);
 	Storage.Emplace<FStructuredBufferView>();
@@ -223,8 +223,7 @@ FVulkanView* FVulkanView::InitAsStructuredBufferView(FVulkanResourceMultiBuffer*
 	SBV.Offset = TotalOffset;
 	SBV.Size = InSize;
 
-	Device.GetBindlessDescriptorManager()->UpdateBuffer(
-		BindlessHandle, Buffer->GetHandle(), TotalOffset, InSize);
+	Device.GetBindlessDescriptorManager()->UpdateBuffer(BindlessHandle, Buffer->GetHandle(), TotalOffset, InSize, bImmediateUpdate);
 
 	return this;
 }
