VulkanRHI: vkCmdFillBuffer needs transfer operation synchronization.

diff --git a/Engine/Source/Runtime/VulkanRHI/Private/VulkanUAV.cpp b/Engine/Source/Runtime/VulkanRHI/Private/VulkanUAV.cpp
--- a/Engine/Source/Runtime/VulkanRHI/Private/VulkanUAV.cpp
+++ b/Engine/Source/Runtime/VulkanRHI/Private/VulkanUAV.cpp
@@ -566,6 +566,13 @@ void FVulkanUnorderedAccessView::Clear(TRHICommandList_RecursiveHazardous<FVulka
 			{
 				FVulkanCmdBuffer* CmdBuffer = Context.GetCommandBufferManager()->GetActiveCmdBuffer();
 
+				// vkCmdFillBuffer is treated as a transfer operation for the purposes of synchronization barriers.
+				{
+					FVulkanPipelineBarrier BeforeBarrier;
+					BeforeBarrier.AddMemoryBarrier(VK_ACCESS_MEMORY_READ_BIT | VK_ACCESS_MEMORY_WRITE_BIT, VK_ACCESS_TRANSFER_WRITE_BIT, VK_PIPELINE_STAGE_ALL_COMMANDS_BIT, VK_PIPELINE_STAGE_TRANSFER_BIT);
+					BeforeBarrier.Execute(CmdBuffer);
+				}
+
 				VulkanRHI::vkCmdFillBuffer(
 					  CmdBuffer->GetHandle()
 					, Buffer->GetHandle()
@@ -573,6 +580,12 @@ void FVulkanUnorderedAccessView::Clear(TRHICommandList_RecursiveHazardous<FVulka
 					, Info.SizeInBytes
 					, ClearValue
 				);
+
+				{
+					FVulkanPipelineBarrier AfterBarrier;
+					AfterBarrier.AddMemoryBarrier(VK_ACCESS_TRANSFER_WRITE_BIT, VK_ACCESS_MEMORY_READ_BIT | VK_ACCESS_MEMORY_WRITE_BIT, VK_PIPELINE_STAGE_TRANSFER_BIT, VK_PIPELINE_STAGE_ALL_COMMANDS_BIT);
+					AfterBarrier.Execute(CmdBuffer);
+				}
 			});
 			break;
 
