Added vulkan bits for PSO precaching

diff --git a/Engine/Source/Runtime/VulkanRHI/Private/VulkanPipeline.cpp b/Engine/Source/Runtime/VulkanRHI/Private/VulkanPipeline.cpp
--- a/Engine/Source/Runtime/VulkanRHI/Private/VulkanPipeline.cpp
+++ b/Engine/Source/Runtime/VulkanRHI/Private/VulkanPipeline.cpp
@@ -1484,6 +1484,8 @@ VkResult FVulkanPipelineStateCacheManager::CreateVKPipeline(FVulkanRHIGraphicsPi
 	{
 		if (bWantPSOSize)
 		{
+			QUICK_SCOPE_CYCLE_COUNTER(STAT_VulkanPSOCreationTimeLRU);
+
 			// We create a single pipeline cache for this create so we can observe the size for LRU cache's accounting.
 			// measuring deltas from the global PipelineCache is not thread safe.
 			VkPipelineCacheCreateInfo PipelineCacheInfo;
@@ -1931,9 +1933,10 @@ FVulkanRHIGraphicsPipelineState* FVulkanPipelineStateCacheManager::RHICreateGrap
 	// We stall precompile PSOs which increases the likelihood for non-precompile PSO to jump the queue.
 	// Not using GraphicsPSOLockedCS as the create could take a long time on some platforms, holding GraphicsPSOLockedCS the whole time could cause hitching.
 	const ESingleThreadedPSOCreateMode ThreadingMode = (ESingleThreadedPSOCreateMode)GVulkanPSOForceSingleThreaded;
+	const bool bIsPrecache = Initializer.bFromPSOFileCache || Initializer.bPSOPrecache;
 	bool bShouldLock = ThreadingMode == ESingleThreadedPSOCreateMode::All
-		|| (ThreadingMode == ESingleThreadedPSOCreateMode::Precompile && Initializer.bFromPSOFileCache)
-		|| (ThreadingMode == ESingleThreadedPSOCreateMode::NonPrecompiled && !Initializer.bFromPSOFileCache);
+		|| (ThreadingMode == ESingleThreadedPSOCreateMode::Precompile && bIsPrecache)
+		|| (ThreadingMode == ESingleThreadedPSOCreateMode::NonPrecompiled && !bIsPrecache);
 
 	FPSOOptionalLock PSOSingleThreadedLock(bShouldLock ? &CreateGraphicsPSOMutex : nullptr);
 
@@ -1957,7 +1960,7 @@ FVulkanRHIGraphicsPipelineState* FVulkanPipelineStateCacheManager::RHICreateGrap
 			if(PSO)
 			{
 				check(*PSO);
-				if(!Initializer.bFromPSOFileCache)
+				if(!bIsPrecache)
 				{
 					LRUTouch(*PSO);
 				}
@@ -2001,7 +2004,7 @@ FVulkanRHIGraphicsPipelineState* FVulkanPipelineStateCacheManager::RHICreateGrap
 			const FVulkanShaderHeader& VSHeader = VS->GetCodeHeader();
 			NewPSO->VertexInputState.Generate(ResourceCast(Initializer.BoundShaderState.VertexDeclarationRHI), VSHeader.InOutMask);
 
-			if((!Initializer.bFromPSOFileCache || !LRUEvictImmediately()) 
+			if((!bIsPrecache || !LRUEvictImmediately()) 
 	#if !UE_BUILD_SHIPPING
 				&& 0 == CVarPipelineDebugForceEvictImmediately.GetValueOnAnyThread()
 	#endif
@@ -2022,7 +2025,7 @@ FVulkanRHIGraphicsPipelineState* FVulkanPipelineStateCacheManager::RHICreateGrap
 			
 				QUICK_SCOPE_CYCLE_COUNTER(STAT_Vulkan_RHICreateGraphicsPipelineState_CREATE_PART0);
 
-				if(!CreateGfxPipelineFromEntry(NewPSO, VulkanShaders, Initializer.bFromPSOFileCache))
+				if(!CreateGfxPipelineFromEntry(NewPSO, VulkanShaders, bIsPrecache))
 				{
 					DeleteNewPSO(NewPSO);
 					return nullptr;
diff --git a/Engine/Source/Runtime/VulkanRHI/Private/VulkanRHI.cpp b/Engine/Source/Runtime/VulkanRHI/Private/VulkanRHI.cpp
--- a/Engine/Source/Runtime/VulkanRHI/Private/VulkanRHI.cpp
+++ b/Engine/Source/Runtime/VulkanRHI/Private/VulkanRHI.cpp
@@ -2007,4 +2007,112 @@ IRHITransientResourceAllocator* FVulkanDynamicRHI::RHICreateTransientResourceAll
 	return nullptr;
 }
 
+uint64 FVulkanDynamicRHI::RHIComputePrecachePSOHash(const FGraphicsPipelineStateInitializer& Initializer)
+{
+	// When compute precache PSO hash we assume a valid state precache PSO hash is already provided
+	checkf(Initializer.StatePrecachePSOHash != 0, TEXT("Initializer should have a valid state precache PSO hash set when computing the full initializer PSO hash"));
+
+	// All members which are not part of the state objects
+	struct FNonStateHashKey
+	{
+		uint64							StatePrecachePSOHash;
+
+		EPrimitiveType					PrimitiveType;
+		uint32							RenderTargetsEnabled;
+		FGraphicsPipelineStateInitializer::TRenderTargetFormats	RenderTargetFormats;
+		FGraphicsPipelineStateInitializer::TRenderTargetFlags RenderTargetFlags;
+		EPixelFormat					DepthStencilTargetFormat;
+		ETextureCreateFlags				DepthStencilTargetFlag;
+		FExclusiveDepthStencil			DepthStencilAccess;
+		uint16							NumSamples;
+		ESubpassHint					SubpassHint;
+		uint8							SubpassIndex;
+		EConservativeRasterization		ConservativeRasterization;
+		uint8							MultiViewCount;
+		EVRSShadingRate					ShadingRate;
+		bool							bDepthBounds;
+		bool							bHasFragmentDensityAttachment;
+	} HashKey;
+
+	FMemory::Memzero(&HashKey, sizeof(FNonStateHashKey));
+
+	HashKey.StatePrecachePSOHash = Initializer.StatePrecachePSOHash;
+
+	HashKey.PrimitiveType = Initializer.PrimitiveType;
+	HashKey.RenderTargetsEnabled = Initializer.RenderTargetsEnabled;
+	HashKey.RenderTargetFormats = Initializer.RenderTargetFormats;
+	HashKey.RenderTargetFlags = Initializer.RenderTargetFlags;
+	HashKey.DepthStencilTargetFormat = Initializer.DepthStencilTargetFormat;
+	HashKey.DepthStencilTargetFlag = Initializer.DepthStencilTargetFlag;
+ 	HashKey.DepthStencilAccess = Initializer.DepthStencilAccess;
+	HashKey.NumSamples = Initializer.NumSamples;
+	HashKey.SubpassHint = Initializer.SubpassHint;
+	HashKey.SubpassIndex = Initializer.SubpassIndex;
+	HashKey.ConservativeRasterization = Initializer.ConservativeRasterization;
+	HashKey.MultiViewCount = Initializer.MultiViewCount;
+	HashKey.ShadingRate = Initializer.ShadingRate;
+	HashKey.bDepthBounds = Initializer.bDepthBounds;
+	HashKey.bHasFragmentDensityAttachment = Initializer.bHasFragmentDensityAttachment;
+	
+	// TODO: check if any RT flags actually affect PSO in VK
+	for (ETextureCreateFlags& Flags : HashKey.RenderTargetFlags)
+	{
+		Flags = Flags & FGraphicsPipelineStateInitializer::RelevantRenderTargetFlagMask;
+	}
+	HashKey.DepthStencilTargetFlag = (HashKey.DepthStencilTargetFlag & FGraphicsPipelineStateInitializer::RelevantDepthStencilFlagMask);
+		 
+	return CityHash64((const char*)&HashKey, sizeof(FNonStateHashKey));
+}
+
+bool FVulkanDynamicRHI::RHIMatchPrecachePSOInitializers(const FGraphicsPipelineStateInitializer& LHS, const FGraphicsPipelineStateInitializer& RHS)
+{
+	// first check non pointer objects
+	if (LHS.ImmutableSamplerState != RHS.ImmutableSamplerState ||
+		LHS.PrimitiveType != RHS.PrimitiveType ||
+		LHS.bDepthBounds != RHS.bDepthBounds ||
+		LHS.MultiViewCount != RHS.MultiViewCount ||
+		LHS.ShadingRate != RHS.ShadingRate ||
+		LHS.bHasFragmentDensityAttachment != RHS.bHasFragmentDensityAttachment ||
+		LHS.RenderTargetsEnabled != RHS.RenderTargetsEnabled ||
+		LHS.RenderTargetFormats != RHS.RenderTargetFormats ||
+		!FGraphicsPipelineStateInitializer::RelevantRenderTargetFlagsEqual(LHS.RenderTargetFlags, RHS.RenderTargetFlags) ||
+		LHS.DepthStencilTargetFormat != RHS.DepthStencilTargetFormat ||
+		!FGraphicsPipelineStateInitializer::RelevantDepthStencilFlagsEqual(LHS.DepthStencilTargetFlag, RHS.DepthStencilTargetFlag) ||
+		LHS.NumSamples != RHS.NumSamples ||
+		LHS.SubpassHint != RHS.SubpassHint ||
+		LHS.SubpassIndex != RHS.SubpassIndex ||
+		LHS.StatePrecachePSOHash != RHS.StatePrecachePSOHash ||
+		LHS.ConservativeRasterization != RHS.ConservativeRasterization)
+	{
+		return false;
+	}
+
+	// check the RHI shaders (pointer check for shaders should be fine)
+	if (LHS.BoundShaderState.VertexShaderRHI != RHS.BoundShaderState.VertexShaderRHI ||
+		LHS.BoundShaderState.PixelShaderRHI != RHS.BoundShaderState.PixelShaderRHI ||
+		LHS.BoundShaderState.GetMeshShader() != RHS.BoundShaderState.GetMeshShader() ||
+		LHS.BoundShaderState.GetAmplificationShader() != RHS.BoundShaderState.GetAmplificationShader() ||
+		LHS.BoundShaderState.GetGeometryShader() != RHS.BoundShaderState.GetGeometryShader())
+	{
+		return false;
+	}
+
+	// Full compare the of the vertex declaration
+	if (!MatchRHIState<FRHIVertexDeclaration, FVertexDeclarationElementList>(LHS.BoundShaderState.VertexDeclarationRHI, RHS.BoundShaderState.VertexDeclarationRHI))
+	{
+		return false;
+	}
+
+	// Check actual state content (each initializer can have it's own state and not going through a factory)
+	if (!MatchRHIState<FRHIBlendState, FBlendStateInitializerRHI>(LHS.BlendState, RHS.BlendState)
+		|| !MatchRHIState<FRHIRasterizerState, FRasterizerStateInitializerRHI>(LHS.RasterizerState, RHS.RasterizerState)
+		|| !MatchRHIState<FRHIDepthStencilState, FDepthStencilStateInitializerRHI>(LHS.DepthStencilState, RHS.DepthStencilState)
+		)
+	{
+		return false;
+	}
+
+	return true;
+}
+
 #undef LOCTEXT_NAMESPACE
diff --git a/Engine/Source/Runtime/VulkanRHI/Private/VulkanResources.h b/Engine/Source/Runtime/VulkanRHI/Private/VulkanResources.h
--- a/Engine/Source/Runtime/VulkanRHI/Private/VulkanResources.h
+++ b/Engine/Source/Runtime/VulkanRHI/Private/VulkanResources.h
@@ -66,8 +66,9 @@ class FVulkanVertexDeclaration : public FRHIVertexDeclaration
 {
 public:
 	FVertexDeclarationElementList Elements;
+	uint32 Hash;
 
-	FVulkanVertexDeclaration(const FVertexDeclarationElementList& InElements);
+	FVulkanVertexDeclaration(const FVertexDeclarationElementList& InElements, uint32 InHash);
 
 	virtual bool GetInitializer(FVertexDeclarationElementList& Out) final override
 	{
@@ -76,6 +77,8 @@ public:
 	}
 
 	static void EmptyCache();
+
+	virtual uint32 GetPrecachePSOHash() const final override { return Hash; }
 };
 
 struct FGfxPipelineDesc;
diff --git a/Engine/Source/Runtime/VulkanRHI/Private/VulkanVertexDeclaration.cpp b/Engine/Source/Runtime/VulkanRHI/Private/VulkanVertexDeclaration.cpp
--- a/Engine/Source/Runtime/VulkanRHI/Private/VulkanVertexDeclaration.cpp
+++ b/Engine/Source/Runtime/VulkanRHI/Private/VulkanVertexDeclaration.cpp
@@ -30,8 +30,9 @@ bool operator==(const FVulkanVertexDeclarationKey& A, const FVulkanVertexDeclara
 			&& !memcmp(A.VertexElements.GetData(), B.VertexElements.GetData(), A.VertexElements.Num() * sizeof(FVertexElement)));
 }
 
-FVulkanVertexDeclaration::FVulkanVertexDeclaration(const FVertexDeclarationElementList& InElements) :
-	Elements(InElements)
+FVulkanVertexDeclaration::FVulkanVertexDeclaration(const FVertexDeclarationElementList& InElements, uint32 InHash) 
+	: Elements(InElements)
+	, Hash(InHash)
 {
 }
 
@@ -51,7 +52,7 @@ FVertexDeclarationRHIRef FVulkanDynamicRHI::RHICreateVertexDeclaration(const FVe
 	FVertexDeclarationRHIRef* VertexDeclarationRefPtr = GVertexDeclarationCache.Find(Key);
 	if (VertexDeclarationRefPtr == nullptr)
 	{
-		VertexDeclarationRefPtr = &GVertexDeclarationCache.Add(Key, new FVulkanVertexDeclaration(Elements));
+		VertexDeclarationRefPtr = &GVertexDeclarationCache.Add(Key, new FVulkanVertexDeclaration(Elements, Key.Hash));
 	}
 
 	check(VertexDeclarationRefPtr);
diff --git a/Engine/Source/Runtime/VulkanRHI/Public/VulkanDynamicRHI.h b/Engine/Source/Runtime/VulkanRHI/Public/VulkanDynamicRHI.h
--- a/Engine/Source/Runtime/VulkanRHI/Public/VulkanDynamicRHI.h
+++ b/Engine/Source/Runtime/VulkanRHI/Public/VulkanDynamicRHI.h
@@ -194,6 +194,10 @@ public:
 	virtual void RHISubmitCommandLists(TArrayView<IRHIPlatformCommandList*> CommandLists) final override;
 	virtual uint64 RHIGetMinimumAlignmentForBufferBackedSRV(EPixelFormat Format) final override;
 
+	virtual uint64 RHIComputePrecachePSOHash(const FGraphicsPipelineStateInitializer& Initializer) final override;
+	virtual bool RHIMatchPrecachePSOInitializers(const FGraphicsPipelineStateInitializer& LHS, const FGraphicsPipelineStateInitializer& RHS) final override;
+
+
 	// Transient Resource Functions and Helpers
 	virtual IRHITransientResourceAllocator* RHICreateTransientResourceAllocator() final override;
 
