VulkanRHI: Fixed bug where TArray index was copied by reference then accessed after add.

diff --git a/Engine/Source/Runtime/VulkanRHI/Private/VulkanBarriers.cpp b/Engine/Source/Runtime/VulkanRHI/Private/VulkanBarriers.cpp
--- a/Engine/Source/Runtime/VulkanRHI/Private/VulkanBarriers.cpp
+++ b/Engine/Source/Runtime/VulkanRHI/Private/VulkanBarriers.cpp
@@ -861,7 +861,7 @@ void FVulkanDynamicRHI::RHIReleaseTransition(FRHITransition* Transition)
 	Transition->GetPrivateData<FVulkanPipelineBarrier>()->~FVulkanPipelineBarrier();
 }
 
-static void AddSubresourceTransitions(TArray<VkImageMemoryBarrier>& Barriers, VkPipelineStageFlags& SrcStageMask, const VkImageMemoryBarrier& TemplateBarrier, VkImage ImageHandle, FVulkanImageLayout& CurrentLayout, VkImageLayout DstLayout)
+static void AddSubresourceTransitions(TArray<VkImageMemoryBarrier>& Barriers, VkPipelineStageFlags& SrcStageMask, const VkImageMemoryBarrier TemplateBarrier, VkImage ImageHandle, FVulkanImageLayout& CurrentLayout, VkImageLayout DstLayout)
 {
 	const uint32 FirstPlane = 0;
 	const uint32 LastPlane = CurrentLayout.NumPlanes;
@@ -1035,7 +1035,7 @@ public:
 				if ((RemainingBarriers->image != VK_NULL_HANDLE) && (RemainingBarriers->subresourceRange.aspectMask != 0))
 				{
 					const int32 RemainingCount = Data->ImageBarrierExtras.Num() - Index;
-					ApplyTracking(RemainingBarriers, RemainingExtras, RemainingCount, IsCrossPipe);
+					ApplyTracking(TargetIndex, RemainingExtras, RemainingCount, IsCrossPipe);
 				}
 			}
 
@@ -1063,7 +1063,7 @@ public:
 
 private:
 	void AppendBarriers(const FVulkanPipelineBarrier* Data);
-	void ApplyTracking(ImageBarrierType* RemainingBarriers, const FVulkanPipelineBarrier::ImageBarrierExtraData* RemainingExtras, int32 RemainingCount, bool IsCrossPipe);
+	void ApplyTracking(uint32 TargetIndex, const FVulkanPipelineBarrier::ImageBarrierExtraData* RemainingExtras, int32 RemainingCount, bool IsCrossPipe);
 	void FinishBatch(ERHIPipeline SrcPipelines, ERHIPipeline DstPipelines, int32 BufferBatchStartIndex, int32 ImageBatchStartIndex);
 	void FinishAll();
 
@@ -1166,42 +1166,42 @@ void FTransitionProcessor<VkMemoryBarrier, VkBufferMemoryBarrier, VkImageMemoryB
 }
 
 template <>
-void FTransitionProcessor<VkMemoryBarrier, VkBufferMemoryBarrier, VkImageMemoryBarrier>::ApplyTracking(VkImageMemoryBarrier* RemainingBarriers, const FVulkanPipelineBarrier::ImageBarrierExtraData* RemainingExtras, int32 RemainingCount, bool IsCrossPipe)
+void FTransitionProcessor<VkMemoryBarrier, VkBufferMemoryBarrier, VkImageMemoryBarrier>::ApplyTracking(uint32_t TargetIndex, const FVulkanPipelineBarrier::ImageBarrierExtraData* RemainingExtras, int32 RemainingCount, bool IsCrossPipe)
 {
 	FVulkanLayoutManager& LayoutManager = Context.GetCommandBufferManager()->GetActiveCmdBuffer()->GetLayoutManager();
 
-	VkImageMemoryBarrier& ImageBarrier = *RemainingBarriers;
+	VkImageMemoryBarrier* ImageBarrier = &ImageBarriers[TargetIndex];
 	FVulkanTexture* Texture = RemainingExtras->BaseTexture;
-	check(Texture->Image == ImageBarrier.image);
+	check(Texture->Image == ImageBarrier->image);
 
 	// Get a copy of the current layout, we will alter it and feed it back in the manager
 	FVulkanImageLayout Layout = *LayoutManager.GetFullLayout(*Texture, true);
 
 	check((Layout.NumMips == Texture->GetNumMips()) && (Layout.NumLayers == Texture->GetNumberOfArrayLevels()));
-	check((ImageBarrier.subresourceRange.baseArrayLayer + Layout.GetSubresRangeLayerCount(ImageBarrier.subresourceRange)) <= Layout.NumLayers);
-	check((ImageBarrier.subresourceRange.baseMipLevel + Layout.GetSubresRangeMipCount(ImageBarrier.subresourceRange)) <= Layout.NumMips);
-	check(ImageBarrier.newLayout != VK_IMAGE_LAYOUT_UNDEFINED);
+	check((ImageBarrier->subresourceRange.baseArrayLayer + Layout.GetSubresRangeLayerCount(ImageBarrier->subresourceRange)) <= Layout.NumLayers);
+	check((ImageBarrier->subresourceRange.baseMipLevel + Layout.GetSubresRangeMipCount(ImageBarrier->subresourceRange)) <= Layout.NumMips);
+	check(ImageBarrier->newLayout != VK_IMAGE_LAYOUT_UNDEFINED);
 
 	// For Depth/Stencil formats where only one of the aspects is transitioned, look ahead for other barriers on the same resource
-	if (Texture->IsDepthOrStencilAspect() && (Texture->GetFullAspectMask() != ImageBarrier.subresourceRange.aspectMask))
+	if (Texture->IsDepthOrStencilAspect() && (Texture->GetFullAspectMask() != ImageBarrier->subresourceRange.aspectMask))
 	{
-		check((ImageBarrier.subresourceRange.aspectMask == VK_IMAGE_ASPECT_DEPTH_BIT) || (ImageBarrier.subresourceRange.aspectMask == VK_IMAGE_ASPECT_STENCIL_BIT));
-		if ((ImageBarrier.oldLayout == VK_IMAGE_LAYOUT_UNDEFINED) && (!RemainingExtras->IsAliasingBarrier))
+		check((ImageBarrier->subresourceRange.aspectMask == VK_IMAGE_ASPECT_DEPTH_BIT) || (ImageBarrier->subresourceRange.aspectMask == VK_IMAGE_ASPECT_STENCIL_BIT));
+		if ((ImageBarrier->oldLayout == VK_IMAGE_LAYOUT_UNDEFINED) && (!RemainingExtras->IsAliasingBarrier))
 		{
-			ImageBarrier.oldLayout = Layout.GetSubresLayout(ImageBarrier.subresourceRange.baseArrayLayer, ImageBarrier.subresourceRange.baseMipLevel, (VkImageAspectFlagBits)ImageBarrier.subresourceRange.aspectMask);
+			ImageBarrier->oldLayout = Layout.GetSubresLayout(ImageBarrier->subresourceRange.baseArrayLayer, ImageBarrier->subresourceRange.baseMipLevel, (VkImageAspectFlagBits)ImageBarrier->subresourceRange.aspectMask);
 		}
 
-		const VkImageAspectFlagBits OtherAspectMask = (VkImageAspectFlagBits)(Texture->GetFullAspectMask() ^ ImageBarrier.subresourceRange.aspectMask);
-		VkImageLayout OtherAspectOldLayout = Layout.GetSubresLayout(ImageBarrier.subresourceRange.baseArrayLayer, ImageBarrier.subresourceRange.baseMipLevel, OtherAspectMask);
+		const VkImageAspectFlagBits OtherAspectMask = (VkImageAspectFlagBits)(Texture->GetFullAspectMask() ^ ImageBarrier->subresourceRange.aspectMask);
+		VkImageLayout OtherAspectOldLayout = Layout.GetSubresLayout(ImageBarrier->subresourceRange.baseArrayLayer, ImageBarrier->subresourceRange.baseMipLevel, OtherAspectMask);
 		VkImageLayout OtherAspectNewLayout = OtherAspectOldLayout;
 		for (int32 OtherBarrierIndex = 0; OtherBarrierIndex < RemainingCount; ++OtherBarrierIndex)
 		{
-			VkImageMemoryBarrier& OtherImageBarrier = RemainingBarriers[OtherBarrierIndex];
+			VkImageMemoryBarrier& OtherImageBarrier = ImageBarriers[TargetIndex+OtherBarrierIndex];
 			FVulkanTexture* OtherTexture = RemainingExtras[OtherBarrierIndex].BaseTexture;
-			if ((OtherTexture->Image == ImageBarrier.image) && (OtherImageBarrier.subresourceRange.aspectMask == OtherAspectMask))
+			if ((OtherTexture->Image == ImageBarrier->image) && (OtherImageBarrier.subresourceRange.aspectMask == OtherAspectMask))
 			{
-				check(ImageBarrier.subresourceRange.baseArrayLayer == OtherImageBarrier.subresourceRange.baseArrayLayer);
-				check(ImageBarrier.subresourceRange.baseMipLevel == OtherImageBarrier.subresourceRange.baseMipLevel);
+				check(ImageBarrier->subresourceRange.baseArrayLayer == OtherImageBarrier.subresourceRange.baseArrayLayer);
+				check(ImageBarrier->subresourceRange.baseMipLevel == OtherImageBarrier.subresourceRange.baseMipLevel);
 
 				OtherAspectNewLayout = OtherImageBarrier.newLayout;
 				if ((OtherImageBarrier.oldLayout != VK_IMAGE_LAYOUT_UNDEFINED) || RemainingExtras[OtherBarrierIndex].IsAliasingBarrier)
@@ -1218,40 +1218,42 @@ void FTransitionProcessor<VkMemoryBarrier, VkBufferMemoryBarrier, VkImageMemoryB
 			}
 		}
 
-		Layout.Set(ImageBarrier.newLayout, ImageBarrier.subresourceRange);
+		Layout.Set(ImageBarrier->newLayout, ImageBarrier->subresourceRange);
 
 		// Fix up the barrier to include both layouts
-		ImageBarrier.oldLayout = GetMergedDepthStencilLayout(ImageBarrier.oldLayout, OtherAspectOldLayout);
-		ImageBarrier.newLayout = GetMergedDepthStencilLayout(ImageBarrier.newLayout, OtherAspectNewLayout);
-		ImageBarrier.subresourceRange.aspectMask |= OtherAspectMask;
+		ImageBarrier->oldLayout = GetMergedDepthStencilLayout(ImageBarrier->oldLayout, OtherAspectOldLayout);
+        ImageBarrier->newLayout = GetMergedDepthStencilLayout(ImageBarrier->newLayout, OtherAspectNewLayout);
+        ImageBarrier->subresourceRange.aspectMask |= OtherAspectMask;
 	}
 	else
 	{
-		if ((ImageBarrier.oldLayout == VK_IMAGE_LAYOUT_UNDEFINED) && (!RemainingExtras->IsAliasingBarrier))
+		if ((ImageBarrier->oldLayout == VK_IMAGE_LAYOUT_UNDEFINED) && (!RemainingExtras->IsAliasingBarrier))
 		{
 			if (Layout.AreAllSubresourcesSameLayout())
 			{
-				ImageBarrier.oldLayout = Layout.MainLayout;
-				ImageBarrier.srcAccessMask = GetVkAccessMaskForLayout(ImageBarrier.oldLayout);
-				MergedSrcStageMask |= GetVkStageFlagsForLayout(ImageBarrier.oldLayout);
+				ImageBarrier->oldLayout = Layout.MainLayout;
+				ImageBarrier->srcAccessMask = GetVkAccessMaskForLayout(ImageBarrier->oldLayout);
+				MergedSrcStageMask |= GetVkStageFlagsForLayout(ImageBarrier->oldLayout);
 			}
 			else
 			{
 				// Slow path, adds one transition per subresource.
 				check(!IsCrossPipe);
-				AddSubresourceTransitions(ImageBarriers, MergedSrcStageMask, ImageBarrier, Texture->Image, Layout, ImageBarrier.newLayout);
+				AddSubresourceTransitions(ImageBarriers, MergedSrcStageMask, *ImageBarrier, Texture->Image, Layout, ImageBarrier->newLayout);
+                // Update the pointer to the current barrier as AddSubresourceTransitions can invalidate it when adding
+                ImageBarrier = &ImageBarriers[TargetIndex];
 			}
 		}
 		else
 		{
-			checkSlow(RemainingExtras->IsAliasingBarrier || Layout.AreSubresourcesSameLayout(ImageBarrier.oldLayout, ImageBarrier.subresourceRange));
+			checkSlow(RemainingExtras->IsAliasingBarrier || Layout.AreSubresourcesSameLayout(ImageBarrier->oldLayout, ImageBarrier->subresourceRange));
 		}
 
-		MergedDstStageMask |= GetVkStageFlagsForLayout(ImageBarrier.newLayout);
+		MergedDstStageMask |= GetVkStageFlagsForLayout(ImageBarriers[TargetIndex].newLayout);
 
 		if (!IsCrossPipe)
 		{
-			if (ImageBarrier.oldLayout == ImageBarrier.newLayout)
+			if (ImageBarrier->oldLayout == ImageBarrier->newLayout)
 			{
 				// It turns out that we don't need a layout transition after all. We may still need a memory barrier if the
 				// previous access was writable.
@@ -1262,12 +1264,12 @@ void FTransitionProcessor<VkMemoryBarrier, VkBufferMemoryBarrier, VkImageMemoryB
 					ZeroVulkanStruct(MemoryBarriers[0], VK_STRUCTURE_TYPE_MEMORY_BARRIER);
 				}
 
-				MemoryBarriers[0].srcAccessMask |= ImageBarrier.srcAccessMask;
-				MemoryBarriers[0].dstAccessMask |= ImageBarrier.dstAccessMask;
+				MemoryBarriers[0].srcAccessMask |= ImageBarrier->srcAccessMask;
+				MemoryBarriers[0].dstAccessMask |= ImageBarrier->dstAccessMask;
 			}
 		}
 
-		Layout.Set(ImageBarrier.newLayout, ImageBarrier.subresourceRange);
+		Layout.Set(ImageBarrier->newLayout, ImageBarrier->subresourceRange);
 	}
 
 	LayoutManager.SetFullLayout(*Texture, Layout);
@@ -1349,11 +1351,11 @@ void FTransitionProcessor<VkMemoryBarrier2, VkBufferMemoryBarrier2, VkImageMemor
 }
 
 template <>
-void FTransitionProcessor<VkMemoryBarrier2, VkBufferMemoryBarrier2, VkImageMemoryBarrier2>::ApplyTracking(VkImageMemoryBarrier2* RemainingBarriers, const FVulkanPipelineBarrier::ImageBarrierExtraData* RemainingExtras, int32 RemainingCount, bool IsCrossPipe)
+void FTransitionProcessor<VkMemoryBarrier2, VkBufferMemoryBarrier2, VkImageMemoryBarrier2>::ApplyTracking(uint32 TargetIndex, const FVulkanPipelineBarrier::ImageBarrierExtraData* RemainingExtras, int32 RemainingCount, bool IsCrossPipe)
 {
 	if (GVulkanAutoCorrectUnknownLayouts)
 	{
-		VkImageMemoryBarrier2& ImageBarrier = *RemainingBarriers;
+		VkImageMemoryBarrier2& ImageBarrier = ImageBarriers[TargetIndex];
 
 		if ((ImageBarrier.oldLayout == VK_IMAGE_LAYOUT_UNDEFINED) && (!RemainingExtras->IsAliasingBarrier))
 		{
@@ -1392,7 +1394,7 @@ void FTransitionProcessor<VkMemoryBarrier2, VkBufferMemoryBarrier2, VkImageMemor
 	}
 
 	// Sync2 only needs to inform tracking of the new layout
-	Context.GetCommandBufferManager()->GetActiveCmdBuffer()->GetLayoutManager().SetLayout(*RemainingExtras->BaseTexture, RemainingBarriers->subresourceRange, RemainingBarriers->newLayout);
+	Context.GetCommandBufferManager()->GetActiveCmdBuffer()->GetLayoutManager().SetLayout(*RemainingExtras->BaseTexture, ImageBarriers[TargetIndex].subresourceRange, ImageBarriers[TargetIndex].newLayout);
 
 	// No needs to worry about 'surprise' subresource transition
 	// No need to worry about partial depth/stencil layout changes
