VulkanRHI: Change which view's bindless handle is returned by default.

#rnx

diff --git a/Engine/Source/Runtime/VulkanRHI/Private/VulkanResources.h b/Engine/Source/Runtime/VulkanRHI/Private/VulkanResources.h
--- a/Engine/Source/Runtime/VulkanRHI/Private/VulkanResources.h
+++ b/Engine/Source/Runtime/VulkanRHI/Private/VulkanResources.h
@@ -569,8 +569,8 @@ public:
 
 	virtual FRHIDescriptorHandle GetDefaultBindlessHandle() const override final
 	{
-		check(DefaultView);
-		return DefaultView->GetBindlessHandle();
+		check(PartialView);
+		return PartialView->GetBindlessHandle();
 	}
 
 	struct FImageCreateInfo
