Removing D3D12 comment from Vulkan code.

diff --git a/Engine/Source/Runtime/VulkanRHI/Private/VulkanRHI.cpp b/Engine/Source/Runtime/VulkanRHI/Private/VulkanRHI.cpp
--- a/Engine/Source/Runtime/VulkanRHI/Private/VulkanRHI.cpp
+++ b/Engine/Source/Runtime/VulkanRHI/Private/VulkanRHI.cpp
@@ -1075,7 +1075,7 @@ void FVulkanCommandListContext::RHIBeginFrame()
 
 #if VULKAN_RHI_RAYTRACING
 	Device->GetRayTracingCompactionRequestHandler()->Update(*this);
-#endif // D3D12_RHI_RAYTRACING
+#endif
 }
 
 
