VulkanRHI:  Fill in AftermathResolveMarkerCallback function.

#rnx

diff --git a/Engine/Source/Runtime/VulkanRHI/Private/VulkanUtil.cpp b/Engine/Source/Runtime/VulkanRHI/Private/VulkanUtil.cpp
--- a/Engine/Source/Runtime/VulkanRHI/Private/VulkanUtil.cpp
+++ b/Engine/Source/Runtime/VulkanRHI/Private/VulkanUtil.cpp
@@ -915,7 +915,21 @@ void AftermathCrashDumpDescriptionCallback(PFN_GFSDK_Aftermath_AddGpuCrashDumpDe
 
 void AftermathResolveMarkerCallback(const void* pMarker, void* pUserData, void** resolvedMarkerData, uint32_t* markerSize)
 {
-	// @todo: this needs implementing
+#if VULKAN_SUPPORTS_NV_DIAGNOSTICS
+	FVulkanDevice* VulkanDevice = (FVulkanDevice*)pUserData;
+	if (VulkanDevice->GetOptionalExtensions().HasNVDiagnosticCheckpoints)
+	{
+		const uint32 Value = (uint32)(size_t)pMarker;
+		const FString* MarkerName = VulkanDevice->GetImmediateContext().GetGPUProfiler().CachedStrings.Find(Value);
+		UE_LOG(LogVulkanRHI, Display, TEXT("[AftermathResolveMarkerCallback] Requested %u [%s]"), Value, MarkerName ? *(*MarkerName) : TEXT("<undefined>"));
+		if (MarkerName && !MarkerName->IsEmpty() && resolvedMarkerData && markerSize)
+		{
+			const TArray<TCHAR, FString::AllocatorType>& CharArray = MarkerName->GetCharArray();
+			(*resolvedMarkerData) = (void*)CharArray.GetData();
+			(*markerSize) = CharArray.Num() * CharArray.GetTypeSize();
+		}
+	}
+#endif // VULKAN_SUPPORTS_NV_DIAGNOSTICS
 }
 #endif
 
