Limit VulkanRHI to Desktop and Android, especially when building with -allmodules.

diff --git a/Engine/Source/Runtime/VulkanRHI/VulkanRHI.Build.cs b/Engine/Source/Runtime/VulkanRHI/VulkanRHI.Build.cs
--- a/Engine/Source/Runtime/VulkanRHI/VulkanRHI.Build.cs
+++ b/Engine/Source/Runtime/VulkanRHI/VulkanRHI.Build.cs
@@ -1,9 +1,8 @@
 // Copyright Epic Games, Inc. All Rights Reserved.
 
 using UnrealBuildTool;
-using System;
-using System.IO;
 
+[SupportedPlatformGroups("Desktop", "Android")]
 public class VulkanRHI : ModuleRules
 {
 	public VulkanRHI(ReadOnlyTargetRules Target) : base(Target)
