VulkanRHI:  Make Core 1.1 feature flags accessible from everywhere now that it's the baseline. Load PhysicalDeviceFeatures for every Core API version supported and appropriatly set an IsCore flag in extensions (can eventually be skipped).

#rnx

diff --git a/Engine/Source/Runtime/VulkanRHI/Private/VulkanDevice.cpp b/Engine/Source/Runtime/VulkanRHI/Private/VulkanDevice.cpp
--- a/Engine/Source/Runtime/VulkanRHI/Private/VulkanDevice.cpp
+++ b/Engine/Source/Runtime/VulkanRHI/Private/VulkanDevice.cpp
@@ -203,6 +203,40 @@ static VkExtent2D GetBestMatchedShadingRateExtents(uint32 ShadingRate, const TAr
 #endif
 
 
+void FVulkanPhysicalDeviceFeatures::Query(VkPhysicalDevice PhysicalDevice, uint32 APIVersion)
+{
+	VkPhysicalDeviceFeatures2 PhysicalDeviceFeatures2;
+	ZeroVulkanStruct(PhysicalDeviceFeatures2, VK_STRUCTURE_TYPE_PHYSICAL_DEVICE_FEATURES_2);
+
+	PhysicalDeviceFeatures2.pNext = &Core_1_1;
+	Core_1_1.sType = VK_STRUCTURE_TYPE_PHYSICAL_DEVICE_VULKAN_1_1_FEATURES;
+
+	if (APIVersion >= VK_API_VERSION_1_2)
+	{
+		Core_1_1.pNext = &Core_1_2;
+		Core_1_2.sType = VK_STRUCTURE_TYPE_PHYSICAL_DEVICE_VULKAN_1_2_FEATURES;
+	}
+
+	if (APIVersion >= VK_API_VERSION_1_3)
+	{
+		Core_1_2.pNext = &Core_1_3;
+		Core_1_3.sType = VK_STRUCTURE_TYPE_PHYSICAL_DEVICE_VULKAN_1_3_FEATURES;
+	}
+
+	VulkanRHI::vkGetPhysicalDeviceFeatures2(PhysicalDevice, &PhysicalDeviceFeatures2);
+
+	// Copy features into old struct for convenience
+	Core_1_0 = PhysicalDeviceFeatures2.features;
+
+	// Apply config modifications
+	Core_1_0.robustBufferAccess = GCVarRobustBufferAccess.GetValueOnAnyThread() > 0 ? VK_TRUE : VK_FALSE;
+
+	// Apply platform restrictions
+	FVulkanPlatform::RestrictEnabledPhysicalDeviceFeatures(this);
+}
+
+
+
 FVulkanDevice::FVulkanDevice(FVulkanDynamicRHI* InRHI, VkPhysicalDevice InGpu)
 	: Device(VK_NULL_HANDLE)
 	, MemoryManager(this)
@@ -220,7 +254,6 @@ FVulkanDevice::FVulkanDevice(FVulkanDynamicRHI* InRHI, VkPhysicalDevice InGpu)
 {
 	RHI = InRHI;
 	FMemory::Memzero(GpuProps);
-	FMemory::Memzero(PhysicalFeatures);
 	FMemory::Memzero(FormatProperties);
 	FMemory::Memzero(PixelFormatComponentMapping);
 
@@ -286,7 +319,9 @@ void FVulkanDevice::CreateDevice(TArray<const ANSICHAR*>& DeviceLayers, FVulkanD
 	// Setup extension and layer info
 	VkDeviceCreateInfo DeviceInfo;
 	ZeroVulkanStruct(DeviceInfo, VK_STRUCTURE_TYPE_DEVICE_CREATE_INFO);
-	
+
+	DeviceInfo.pEnabledFeatures = &PhysicalDeviceFeatures.Core_1_0;
+
 	for (TUniquePtr<FVulkanDeviceExtension>& UEExtension : UEExtensions)
 	{
 		if (UEExtension->InUse())
@@ -381,12 +416,6 @@ void FVulkanDevice::CreateDevice(TArray<const ANSICHAR*>& DeviceLayers, FVulkanD
 	DeviceInfo.queueCreateInfoCount = QueueFamilyInfos.Num();
 	DeviceInfo.pQueueCreateInfos = QueueFamilyInfos.GetData();
 
-	PhysicalFeatures.robustBufferAccess = GCVarRobustBufferAccess.GetValueOnAnyThread() > 0 ? VK_TRUE : VK_FALSE;
-	FVulkanPlatform::RestrictEnabledPhysicalDeviceFeatures(PhysicalFeatures);
-	DeviceInfo.pEnabledFeatures = &PhysicalFeatures;
-
-	FVulkanPlatform::EnablePhysicalDeviceFeatureExtensions(DeviceInfo, *this);
-
 #if NV_AFTERMATH && VULKAN_SUPPORTS_NV_DIAGNOSTICS
 	if (GGPUCrashDebuggingEnabled && GVulkanNVAftermathModuleLoaded)
 	{
@@ -986,7 +1015,7 @@ void FVulkanDevice::MapImageFormatSupport(FPixelFormatInfo& PixelFormatInfo, con
 		if (EnumHasAllFlags(Capabilities, EPixelFormatCapabilities::AnyTexture))
 		{
 			// We support gather, but some of our shaders assume offsets so check against features
-			if (GetPhysicalFeatures().shaderImageGatherExtended)
+			if (GetPhysicalDeviceFeatures().Core_1_0.shaderImageGatherExtended)
 			{
 				EnumAddFlags(Capabilities, EPixelFormatCapabilities::TextureGather);
 			}
@@ -1076,10 +1105,10 @@ void FVulkanDevice::InitGPU()
 	VulkanRHI::vkGetPhysicalDeviceQueueFamilyProperties(Gpu, &QueueCount, QueueFamilyProps.GetData());
 
 	// Query base features
-	VulkanRHI::vkGetPhysicalDeviceFeatures(Gpu, &PhysicalFeatures);
+	PhysicalDeviceFeatures.Query(Gpu, RHI->GetApiVersion());
 
 	// Setup layers and extensions
-	FVulkanDeviceExtensionArray UEExtensions = FVulkanDeviceExtension::GetUESupportedDeviceExtensions(this);
+	FVulkanDeviceExtensionArray UEExtensions = FVulkanDeviceExtension::GetUESupportedDeviceExtensions(this, RHI->GetApiVersion());
 	TArray<const ANSICHAR*> DeviceLayers = FVulkanDevice::SetupDeviceLayers(Gpu, UEExtensions);
 
 	// Query advanced features
@@ -1132,7 +1161,8 @@ void FVulkanDevice::InitGPU()
 		}
 	}
 
-	UE_LOG(LogVulkanRHI, Display, TEXT("Device properties: Geometry %d BufferAtomic64 %d ImageAtomic64 %d"), PhysicalFeatures.geometryShader, OptionalDeviceExtensions.HasKHRShaderAtomicInt64, OptionalDeviceExtensions.HasImageAtomicInt64);
+	UE_LOG(LogVulkanRHI, Display, TEXT("Device properties: Geometry %d BufferAtomic64 %d ImageAtomic64 %d"), 
+		PhysicalDeviceFeatures.Core_1_0.geometryShader, OptionalDeviceExtensions.HasKHRShaderAtomicInt64, OptionalDeviceExtensions.HasImageAtomicInt64);
 
 	CreateDevice(DeviceLayers, UEExtensions);
 
diff --git a/Engine/Source/Runtime/VulkanRHI/Private/VulkanDevice.h b/Engine/Source/Runtime/VulkanRHI/Private/VulkanDevice.h
--- a/Engine/Source/Runtime/VulkanRHI/Private/VulkanDevice.h
+++ b/Engine/Source/Runtime/VulkanRHI/Private/VulkanDevice.h
@@ -128,6 +128,27 @@ struct FOptionalVulkanDeviceExtensionProperties
 #endif // VULKAN_RHI_RAYTRACING
 };
 
+class FVulkanPhysicalDeviceFeatures
+{
+public:
+	FVulkanPhysicalDeviceFeatures()
+	{
+		FMemory::Memzero(*this);
+	}
+
+	void Query(VkPhysicalDevice PhysicalDevice, uint32 APIVersion);
+
+	VkPhysicalDeviceFeatures	     Core_1_0;
+	VkPhysicalDeviceVulkan11Features Core_1_1;
+private:
+	// Anything above Core 1.1 cannot be assumed, they should only be used by the device at init time
+	VkPhysicalDeviceVulkan12Features Core_1_2;
+	VkPhysicalDeviceVulkan13Features Core_1_3;
+
+	friend class FVulkanDevice;
+};
+
+
 namespace VulkanRHI
 {
 	class FDeferredDeletionQueue2 : public FDeviceChild
@@ -316,9 +337,9 @@ public:
 	}
 #endif
 
-	inline const VkPhysicalDeviceFeatures& GetPhysicalFeatures() const
+	inline const FVulkanPhysicalDeviceFeatures& GetPhysicalDeviceFeatures() const
 	{
-		return PhysicalFeatures;
+		return PhysicalDeviceFeatures;
 	}
 
 	inline bool HasUnifiedMemory() const
@@ -576,7 +597,7 @@ private:
 	FVulkanRayTracingCompactionRequestHandler* RayTracingCompactionRequestHandler = nullptr;
 #endif // VULKAN_RHI_RAYTRACING
 
-	VkPhysicalDeviceFeatures PhysicalFeatures;
+	FVulkanPhysicalDeviceFeatures PhysicalDeviceFeatures;
 
 	TArray<VkQueueFamilyProperties> QueueFamilyProps;
 	VkFormatProperties FormatProperties[VK_FORMAT_RANGE_SIZE];
diff --git a/Engine/Source/Runtime/VulkanRHI/Private/VulkanExtensions.cpp b/Engine/Source/Runtime/VulkanRHI/Private/VulkanExtensions.cpp
--- a/Engine/Source/Runtime/VulkanRHI/Private/VulkanExtensions.cpp
+++ b/Engine/Source/Runtime/VulkanRHI/Private/VulkanExtensions.cpp
@@ -1182,7 +1182,7 @@ private:
 
 
 template <typename ExtensionType>
-static void FlagExtensionSupport(const TArray<VkExtensionProperties>& ExtensionProperties, TArray<TUniquePtr<ExtensionType>>& UEExtensions, const TCHAR* ExtensionTypeName)
+static void FlagExtensionSupport(const TArray<VkExtensionProperties>& ExtensionProperties, TArray<TUniquePtr<ExtensionType>>& UEExtensions, uint32 ApiVersion, const TCHAR* ExtensionTypeName)
 {
 	// Flag the extension support
 	UE_LOG(LogVulkanRHI, Display, TEXT("Found %d available %s extensions :"), ExtensionProperties.Num(), ExtensionTypeName);
@@ -1190,17 +1190,22 @@ static void FlagExtensionSupport(const TArray<VkExtensionProperties>& ExtensionP
 	{
 		const int32 ExtensionIndex = ExtensionType::FindExtension(UEExtensions, Extension.extensionName);
 		const bool bFound = (ExtensionIndex != INDEX_NONE);
+		bool bIsCore = false;
 		if (bFound)
 		{
 			UEExtensions[ExtensionIndex]->SetSupported();
+
+			// Set the core flag if the extension was promoted for our current api version
+			bIsCore = UEExtensions[ExtensionIndex]->SetCore(ApiVersion);
 		}
-		UE_LOG(LogVulkanRHI, Display, TEXT("  %s %s"), bFound ? TEXT("+") : TEXT("-"), ANSI_TO_TCHAR(Extension.extensionName));
+
+		UE_LOG(LogVulkanRHI, Display, TEXT("  %s %s"), bIsCore ? TEXT("*") : bFound ? TEXT("+") : TEXT("-"), ANSI_TO_TCHAR(Extension.extensionName));
 	}
 }
 
 
 
-FVulkanDeviceExtensionArray FVulkanDeviceExtension::GetUESupportedDeviceExtensions(FVulkanDevice* InDevice)
+FVulkanDeviceExtensionArray FVulkanDeviceExtension::GetUESupportedDeviceExtensions(FVulkanDevice* InDevice, uint32 ApiVersion)
 {
 	FVulkanDeviceExtensionArray OutUEDeviceExtensions;
 
@@ -1306,7 +1311,7 @@ FVulkanDeviceExtensionArray FVulkanDeviceExtension::GetUESupportedDeviceExtensio
 	ActivateExternalExtensions(ExternalExtensions, TEXT("Externally"));
 
 	// Now that all the extensions are listed, update their support flags
-	FlagExtensionSupport(GetDriverSupportedDeviceExtensions(InDevice->GetPhysicalHandle()), OutUEDeviceExtensions, TEXT("device"));
+	FlagExtensionSupport(GetDriverSupportedDeviceExtensions(InDevice->GetPhysicalHandle()), OutUEDeviceExtensions, ApiVersion, TEXT("device"));
 
 	#undef ADD_SIMPLE_EXTENSION
 	#undef ADD_EXTERNAL_EXTENSION
@@ -1392,7 +1397,7 @@ private:
 
 
 
-FVulkanInstanceExtensionArray FVulkanInstanceExtension::GetUESupportedInstanceExtensions()
+FVulkanInstanceExtensionArray FVulkanInstanceExtension::GetUESupportedInstanceExtensions(uint32 ApiVersion)
 {
 	FVulkanInstanceExtensionArray OutUEInstanceExtensions;
 
@@ -1447,7 +1452,7 @@ FVulkanInstanceExtensionArray FVulkanInstanceExtension::GetUESupportedInstanceEx
 	ActivateExternalExtensions(ExternalExtensions, TEXT("Externally"));
 
 	// Now that all the extensions are listed, update their support flags
-	FlagExtensionSupport(GetDriverSupportedInstanceExtensions(), OutUEInstanceExtensions, TEXT("instance"));
+	FlagExtensionSupport(GetDriverSupportedInstanceExtensions(), OutUEInstanceExtensions, ApiVersion, TEXT("instance"));
 
 	return OutUEInstanceExtensions;
 }
diff --git a/Engine/Source/Runtime/VulkanRHI/Private/VulkanExtensions.h b/Engine/Source/Runtime/VulkanRHI/Private/VulkanExtensions.h
--- a/Engine/Source/Runtime/VulkanRHI/Private/VulkanExtensions.h
+++ b/Engine/Source/Runtime/VulkanRHI/Private/VulkanExtensions.h
@@ -9,7 +9,7 @@
 #include "CoreMinimal.h"
 #include "VulkanConfiguration.h"
 
-#define VULKAN_EXTENSION_NOT_PROMOTED 0
+#define VULKAN_EXTENSION_NOT_PROMOTED UINT32_MAX
 #define VULKAN_EXTENSION_ENABLED      1
 #define VULKAN_EXTENSION_DISABLED     0
 
@@ -35,16 +35,19 @@ public:
 		, bEnabledInCode(InEnabledInCode == VULKAN_EXTENSION_ENABLED)
 		, bSupported(false)
 		, bActivated(InActivation == EExtensionActivation::AutoActivate)
+		, bCore(false)
 	{}
 
 	virtual ~FVulkanExtensionBase() {}
 
 	inline void SetSupported() { bSupported = true; }
 	inline void SetActivated() { bActivated = true; }
+	inline bool SetCore(uint32 ApiVersion) { return (bCore = (ApiVersion >= PromotedVersion)); }
 
 	inline const ANSICHAR* GetExtensionName() const { return ExtensionName; }
 	inline bool IsEnabled() const { return bEnabledInCode; }
 	inline bool IsSupported() const { return bSupported; }
+	inline bool IsCore() { return bCore; }
 
 	inline bool InUse() const
 	{
@@ -76,6 +79,9 @@ protected:
 
 	// Signals if the runtime currently wants this extension to be loaded.
 	bool bActivated;
+
+	// Signals if the extension has been promoted to core Vulkan for the current instance version
+	bool bCore;
 };
 
 
@@ -91,7 +97,7 @@ public:
 		, bRequirementsPassed(true)
 	{}
 
-	static FVulkanDeviceExtensionArray GetUESupportedDeviceExtensions(FVulkanDevice* InDevice);
+	static FVulkanDeviceExtensionArray GetUESupportedDeviceExtensions(FVulkanDevice* InDevice, uint32 ApiVersion);
 	static TArray<VkExtensionProperties> GetDriverSupportedDeviceExtensions(VkPhysicalDevice Gpu, const ANSICHAR* LayerName = nullptr);
 
 	virtual void PrePhysicalDeviceProperties(VkPhysicalDeviceProperties2KHR& PhysicalDeviceProperties2) {}
@@ -132,7 +138,7 @@ public:
 		, FlagSetter(MoveTemp(InFlagSetter))
 	{}
 
-	static FVulkanInstanceExtensionArray GetUESupportedInstanceExtensions();
+	static FVulkanInstanceExtensionArray GetUESupportedInstanceExtensions(uint32 ApiVersion);
 	static TArray<VkExtensionProperties> GetDriverSupportedInstanceExtensions(const ANSICHAR* LayerName = nullptr);
 
 	virtual void PreCreateInstance(VkInstanceCreateInfo& CreateInfo, FOptionalVulkanInstanceExtensions& ExtensionFlags)
diff --git a/Engine/Source/Runtime/VulkanRHI/Private/VulkanGenericPlatform.cpp b/Engine/Source/Runtime/VulkanRHI/Private/VulkanGenericPlatform.cpp
--- a/Engine/Source/Runtime/VulkanRHI/Private/VulkanGenericPlatform.cpp
+++ b/Engine/Source/Runtime/VulkanRHI/Private/VulkanGenericPlatform.cpp
@@ -86,6 +86,21 @@ TArray<FString> FVulkanGenericPlatform::GetPSOCacheFilenames()
 	return CacheFilenames;
 }
 
+void FVulkanGenericPlatform::RestrictEnabledPhysicalDeviceFeatures(FVulkanPhysicalDeviceFeatures* InOutFeaturesToEnable)
+{
+	// Disable everything sparse-related
+	InOutFeaturesToEnable->Core_1_0.shaderResourceResidency = VK_FALSE;
+	InOutFeaturesToEnable->Core_1_0.shaderResourceMinLod = VK_FALSE;
+	InOutFeaturesToEnable->Core_1_0.sparseBinding = VK_FALSE;
+	InOutFeaturesToEnable->Core_1_0.sparseResidencyBuffer = VK_FALSE;
+	InOutFeaturesToEnable->Core_1_0.sparseResidencyImage2D = VK_FALSE;
+	InOutFeaturesToEnable->Core_1_0.sparseResidencyImage3D = VK_FALSE;
+	InOutFeaturesToEnable->Core_1_0.sparseResidency2Samples = VK_FALSE;
+	InOutFeaturesToEnable->Core_1_0.sparseResidency4Samples = VK_FALSE;
+	InOutFeaturesToEnable->Core_1_0.sparseResidency8Samples = VK_FALSE;
+	InOutFeaturesToEnable->Core_1_0.sparseResidencyAliased = VK_FALSE;
+}
+
 VkResult FVulkanGenericPlatform::Present(VkQueue Queue, VkPresentInfoKHR& PresentInfo)
 {
 	return VulkanRHI::vkQueuePresentKHR(Queue, &PresentInfo);
@@ -140,7 +155,8 @@ void FVulkanGenericPlatform::ClearVulkanInstanceFunctions()
 
 bool FVulkanGenericPlatform::SupportsProfileChecks()
 {
-	return (CVarVulkanUseProfileCheck.GetValueOnAnyThread() != 0);
+	return (CVarVulkanUseProfileCheck.GetValueOnAnyThread() != 0) && 
+		!FParse::Param(FCommandLine::Get(), TEXT("SkipVulkanProfileCheck"));
 }
 
 FString FVulkanGenericPlatform::GetVulkanProfileNameForFeatureLevel(ERHIFeatureLevel::Type FeatureLevel, bool bRaytracing)
diff --git a/Engine/Source/Runtime/VulkanRHI/Private/VulkanGenericPlatform.h b/Engine/Source/Runtime/VulkanRHI/Private/VulkanGenericPlatform.h
--- a/Engine/Source/Runtime/VulkanRHI/Private/VulkanGenericPlatform.h
+++ b/Engine/Source/Runtime/VulkanRHI/Private/VulkanGenericPlatform.h
@@ -12,6 +12,7 @@ struct FOptionalVulkanDeviceExtensions;
 class FVulkanDevice;
 class FVulkanRenderTargetLayout;
 struct FGfxPipelineDesc;
+class FVulkanPhysicalDeviceFeatures;
 
 using FVulkanDeviceExtensionArray = TArray<TUniquePtr<class FVulkanDeviceExtension>>;
 using FVulkanInstanceExtensionArray = TArray<TUniquePtr<class FVulkanInstanceExtension>>;
@@ -74,20 +75,7 @@ public:
 	static void WriteCrashMarker(const FOptionalVulkanDeviceExtensions& OptionalExtensions, VkCommandBuffer CmdBuffer, VkBuffer DestBuffer, const TArrayView<uint32>& Entries, bool bAdding) {}
 
 	// Allow the platform code to restrict the device features
-	static void RestrictEnabledPhysicalDeviceFeatures(VkPhysicalDeviceFeatures& InOutFeaturesToEnable)
-	{ 
-		// Disable everything sparse-related
-		InOutFeaturesToEnable.shaderResourceResidency	= VK_FALSE;
-		InOutFeaturesToEnable.shaderResourceMinLod		= VK_FALSE;
-		InOutFeaturesToEnable.sparseBinding				= VK_FALSE;
-		InOutFeaturesToEnable.sparseResidencyBuffer		= VK_FALSE;
-		InOutFeaturesToEnable.sparseResidencyImage2D	= VK_FALSE;
-		InOutFeaturesToEnable.sparseResidencyImage3D	= VK_FALSE;
-		InOutFeaturesToEnable.sparseResidency2Samples	= VK_FALSE;
-		InOutFeaturesToEnable.sparseResidency4Samples	= VK_FALSE;
-		InOutFeaturesToEnable.sparseResidency8Samples	= VK_FALSE;
-		InOutFeaturesToEnable.sparseResidencyAliased	= VK_FALSE;
-	}
+	static void RestrictEnabledPhysicalDeviceFeatures(FVulkanPhysicalDeviceFeatures* InOutFeaturesToEnable);
 
 	static bool SupportParallelRenderingTasks() { return true; }
 
@@ -96,9 +84,6 @@ public:
 
 	static bool SupportsVolumeTextureRendering() { return true; }
 
-	// Allow platforms to add extension features to the DeviceInfo pNext chain
-	static void EnablePhysicalDeviceFeatureExtensions(VkDeviceCreateInfo& DeviceInfo, FVulkanDevice& Device) {}
-
 	static bool RequiresSwapchainGeneralInitialLayout() { return false; }
 
 	// Allow platforms to perform their own frame pacing, called before Present. Returns true if the platform has done framepacing, false otherwise.
diff --git a/Engine/Source/Runtime/VulkanRHI/Private/VulkanQueue.cpp b/Engine/Source/Runtime/VulkanRHI/Private/VulkanQueue.cpp
--- a/Engine/Source/Runtime/VulkanRHI/Private/VulkanQueue.cpp
+++ b/Engine/Source/Runtime/VulkanRHI/Private/VulkanQueue.cpp
@@ -155,7 +155,7 @@ void FVulkanQueue::FillSupportedStageBits()
 			VK_PIPELINE_STAGE_TRANSFER_BIT |
 			VK_PIPELINE_STAGE_ALL_GRAPHICS_BIT;
 
-		if (Device->GetPhysicalFeatures().geometryShader)
+		if (Device->GetPhysicalDeviceFeatures().Core_1_0.geometryShader)
 		{
 			SupportedStages |= VK_PIPELINE_STAGE_GEOMETRY_SHADER_BIT;
 		}
diff --git a/Engine/Source/Runtime/VulkanRHI/Private/VulkanRHI.cpp b/Engine/Source/Runtime/VulkanRHI/Private/VulkanRHI.cpp
--- a/Engine/Source/Runtime/VulkanRHI/Private/VulkanRHI.cpp
+++ b/Engine/Source/Runtime/VulkanRHI/Private/VulkanRHI.cpp
@@ -248,18 +248,11 @@ static VkPhysicalDevice SelectPhysicalDevice(VkInstance InInstance)
 
 static uint32 GetVulkanApiVersionForFeatureLevel(ERHIFeatureLevel::Type FeatureLevel, bool bRaytracing)
 {
-	if (FVulkanPlatform::SupportsProfileChecks())
+	const FString ProfileName = FVulkanPlatform::GetVulkanProfileNameForFeatureLevel(FeatureLevel, bRaytracing);
+	const detail::VpProfileDesc* ProfileDesc = detail::vpGetProfileDesc(TCHAR_TO_ANSI(*ProfileName));
+	if (ProfileDesc)
 	{
-		const FString ProfileName = FVulkanPlatform::GetVulkanProfileNameForFeatureLevel(FeatureLevel, bRaytracing);
-		const detail::VpProfileDesc* ProfileDesc = detail::vpGetProfileDesc(TCHAR_TO_ANSI(*ProfileName));
-		if (ProfileDesc)
-		{
-			return ProfileDesc->minApiVersion;
-		}
-
-		UE_LOG(LogVulkanRHI, Warning,
-			TEXT("Profile called [%s] could not be found for feature level [%s]!"),
-			*ProfileName, *LexToString(FeatureLevel));
+		return ProfileDesc->minApiVersion;
 	}
 
 	UE_LOG(LogVulkanRHI, Log, TEXT("Using default apiVersion for platform..."));
@@ -271,6 +264,12 @@ static bool CheckVulkanProfile(ERHIFeatureLevel::Type FeatureLevel, bool bRaytra
 {
 	const FString ProfileName = FVulkanPlatform::GetVulkanProfileNameForFeatureLevel(FeatureLevel, bRaytracing);
 
+	if (!FVulkanGenericPlatform::SupportsProfileChecks())
+	{
+		UE_LOG(LogVulkanRHI, Log, TEXT("Skipping Vulkan Profile check for %s:"), *ProfileName);
+		return true;
+	}
+
 	UE_LOG(LogVulkanRHI, Log, TEXT("Starting Vulkan Profile check for %s:"), *ProfileName);
 	ON_SCOPE_EXIT
 	{
@@ -681,7 +680,7 @@ void FVulkanDynamicRHI::CreateInstance()
 	ZeroVulkanStruct(InstInfo, VK_STRUCTURE_TYPE_INSTANCE_CREATE_INFO);
 	InstInfo.pApplicationInfo = &AppInfo;
 
-	FVulkanInstanceExtensionArray UEInstanceExtensions = FVulkanInstanceExtension::GetUESupportedInstanceExtensions();
+	FVulkanInstanceExtensionArray UEInstanceExtensions = FVulkanInstanceExtension::GetUESupportedInstanceExtensions(ApiVersion);
 	InstanceLayers = SetupInstanceLayers(UEInstanceExtensions);
 	for (TUniquePtr<FVulkanInstanceExtension>& Extension : UEInstanceExtensions)
 	{
@@ -879,7 +878,7 @@ void FVulkanDynamicRHI::InitInstance()
 		GRHISupportsFirstInstance = true;
 		GRHISupportsDynamicResolution = FVulkanPlatform::SupportsDynamicResolution();
 		GRHISupportsFrameCyclesBubblesRemoval = true;
-		GSupportsDepthBoundsTest = Device->GetPhysicalFeatures().depthBounds != 0;
+		GSupportsDepthBoundsTest = Device->GetPhysicalDeviceFeatures().Core_1_0.depthBounds != 0;
 		GSupportsRenderTargetFormat_PF_G8 = false;	// #todo-rco
 		GRHISupportsTextureStreaming = true;
 		GSupportsTimestampRenderQueries = FVulkanPlatform::SupportsTimestampRenderQueries();
@@ -956,7 +955,7 @@ PRAGMA_ENABLE_DEPRECATION_WARNINGS
 		uint32 VulkanDeviceShaderStageBits = VK_SHADER_STAGE_VERTEX_BIT | VK_SHADER_STAGE_FRAGMENT_BIT | VK_SHADER_STAGE_COMPUTE_BIT;
 
 		// optional shader stages
-		if (Device->GetPhysicalFeatures().geometryShader) 
+		if (Device->GetPhysicalDeviceFeatures().Core_1_0.geometryShader)
 		{
 			GVulkanDevicePipelineStageBits |= VK_PIPELINE_STAGE_GEOMETRY_SHADER_BIT;
 			VulkanDeviceShaderStageBits |= VK_SHADER_STAGE_GEOMETRY_BIT;
diff --git a/Engine/Source/Runtime/VulkanRHI/Public/VulkanDynamicRHI.h b/Engine/Source/Runtime/VulkanRHI/Public/VulkanDynamicRHI.h
--- a/Engine/Source/Runtime/VulkanRHI/Public/VulkanDynamicRHI.h
+++ b/Engine/Source/Runtime/VulkanRHI/Public/VulkanDynamicRHI.h
@@ -285,6 +285,11 @@ public:
 		return Viewports;
 	}
 
+	inline uint32 GetApiVersion() const
+	{
+		return ApiVersion;
+	}
+
 public:
 	static void SavePipelineCache();
 
