VulkanRHI: Quick fix: Volatile buffers use frame temp allocations smaller than the actual buffer size (since only the locked region is made available).

#rnx

diff --git a/Engine/Source/Runtime/VulkanRHI/Private/VulkanUAV.cpp b/Engine/Source/Runtime/VulkanRHI/Private/VulkanUAV.cpp
--- a/Engine/Source/Runtime/VulkanRHI/Private/VulkanUAV.cpp
+++ b/Engine/Source/Runtime/VulkanRHI/Private/VulkanUAV.cpp
@@ -85,6 +85,12 @@ FVulkanView* FVulkanView::InitAsTypedBufferView(FVulkanResourceMultiBuffer* Buff
 	ViewInfo.offset = TotalOffset;
 	ViewInfo.format = Format;
 
+	// :todo-jn: Volatile buffers use temporary allocations that can be smaller than the buffer creation size.  Check if the savings are still worth it.
+	if (Buffer->IsVolatile())
+	{
+		InSize = FMath::Min<uint64>(InSize, Buffer->GetCurrentSize());
+	}
+
 	//#todo-rco: Revisit this if buffer views become VK_BUFFER_USAGE_STORAGE_BUFFER_BIT instead of VK_BUFFER_USAGE_UNIFORM_TEXEL_BUFFER_BIT
 	const VkPhysicalDeviceLimits& Limits = Device.GetLimits();
 	const uint64 MaxSize = (uint64)Limits.maxTexelBufferElements * GetNumBitsPerPixel(Format) / 8;
@@ -221,6 +227,13 @@ FVulkanView* FVulkanView::InitAsStructuredBufferView(FVulkanResourceMultiBuffer*
 	SBV.Buffer = Buffer->GetHandle();
 	SBV.HandleId = Buffer->GetCurrentAllocation().HandleId;
 	SBV.Offset = TotalOffset;
+
+	// :todo-jn: Volatile buffers use temporary allocations that can be smaller than the buffer creation size.  Check if the savings are still worth it.
+	if (Buffer->IsVolatile())
+	{
+		InSize = FMath::Min<uint64>(InSize, Buffer->GetCurrentSize());
+	}
+
 	SBV.Size = InSize;
 
 	Device.GetBindlessDescriptorManager()->UpdateBuffer(BindlessHandle, Buffer->GetHandle(), TotalOffset, InSize, bImmediateUpdate);
