PSO precaching: Remove DepthStencilAccess from VK precaching hash as it does not contribute to a final PSO

diff --git a/Engine/Source/Runtime/VulkanRHI/Private/VulkanRHI.cpp b/Engine/Source/Runtime/VulkanRHI/Private/VulkanRHI.cpp
--- a/Engine/Source/Runtime/VulkanRHI/Private/VulkanRHI.cpp
+++ b/Engine/Source/Runtime/VulkanRHI/Private/VulkanRHI.cpp
@@ -2022,7 +2022,6 @@ uint64 FVulkanDynamicRHI::RHIComputePrecachePSOHash(const FGraphicsPipelineState
 		FGraphicsPipelineStateInitializer::TRenderTargetFlags RenderTargetFlags;
 		EPixelFormat					DepthStencilTargetFormat;
 		ETextureCreateFlags				DepthStencilTargetFlag;
-		FExclusiveDepthStencil			DepthStencilAccess;
 		uint16							NumSamples;
 		ESubpassHint					SubpassHint;
 		uint8							SubpassIndex;
@@ -2043,7 +2042,6 @@ uint64 FVulkanDynamicRHI::RHIComputePrecachePSOHash(const FGraphicsPipelineState
 	HashKey.RenderTargetFlags = Initializer.RenderTargetFlags;
 	HashKey.DepthStencilTargetFormat = Initializer.DepthStencilTargetFormat;
 	HashKey.DepthStencilTargetFlag = Initializer.DepthStencilTargetFlag;
- 	HashKey.DepthStencilAccess = Initializer.DepthStencilAccess;
 	HashKey.NumSamples = Initializer.NumSamples;
 	HashKey.SubpassHint = Initializer.SubpassHint;
 	HashKey.SubpassIndex = Initializer.SubpassIndex;
