PSO precaching: More const correctness around VKCreatePipeline

diff --git a/Engine/Source/Runtime/VulkanRHI/Private/Android/VulkanAndroidPlatform.cpp b/Engine/Source/Runtime/VulkanRHI/Private/Android/VulkanAndroidPlatform.cpp
--- a/Engine/Source/Runtime/VulkanRHI/Private/Android/VulkanAndroidPlatform.cpp
+++ b/Engine/Source/Runtime/VulkanRHI/Private/Android/VulkanAndroidPlatform.cpp
@@ -800,7 +800,7 @@ void GetVKStructsFromPNext(const void* pNext, TMap<uint32_t, const void*>& VkStr
 	}
 }
 
-void HandleGraphicsPipelineCreatePNext(VkGraphicsPipelineCreateInfo* PipelineCreateInfo, TArray<char>& MemoryStream)
+void HandleGraphicsPipelineCreatePNext(const VkGraphicsPipelineCreateInfo* PipelineCreateInfo, TArray<char>& MemoryStream)
 {
 	TMap<uint32_t, const void*> VkStructs;
 	GetVKStructsFromPNext(PipelineCreateInfo->pNext, VkStructs);
@@ -859,7 +859,7 @@ void HandleSubpassDescriptionPNext(VkSubpassDescription2* SubpassDescription, TA
 	check(HandledCount == VkStructs.Num());
 }
 
-void PipelineToBinary(FVulkanDevice* Device, VkGraphicsPipelineCreateInfo* PipelineInfo, FGfxPipelineDesc* GfxEntry, const FVulkanRenderTargetLayout* RTLayout, TArray<char>& MemoryStream)
+void PipelineToBinary(FVulkanDevice* Device, const VkGraphicsPipelineCreateInfo* PipelineInfo, FGfxPipelineDesc* GfxEntry, const FVulkanRenderTargetLayout* RTLayout, TArray<char>& MemoryStream)
 {
 	static const unsigned int INITIAL_PSO_STREAM_SIZE = 64 * 1024;
 	MemoryStream.Reserve(INITIAL_PSO_STREAM_SIZE);
@@ -1368,7 +1368,7 @@ namespace AndroidVulkanService
 	std::atomic<bool> bOneTimeErrorEncountered = false;
 }
 
-VkPipelineCache FVulkanAndroidPlatform::PrecompilePSO(FVulkanDevice* Device, VkGraphicsPipelineCreateInfo* PipelineInfo, FGfxPipelineDesc* GfxEntry, const FVulkanRenderTargetLayout* RTLayout, TArrayView<uint32_t> VS, TArrayView<uint32_t> PS, size_t& AfterSize)
+VkPipelineCache FVulkanAndroidPlatform::PrecompilePSO(FVulkanDevice* Device, const VkGraphicsPipelineCreateInfo* PipelineInfo, FGfxPipelineDesc* GfxEntry, const FVulkanRenderTargetLayout* RTLayout, TArrayView<uint32_t> VS, TArrayView<uint32_t> PS, size_t& AfterSize)
 {
 	FString FailureMessageOUT;
 	
diff --git a/Engine/Source/Runtime/VulkanRHI/Private/Android/VulkanAndroidPlatform.h b/Engine/Source/Runtime/VulkanRHI/Private/Android/VulkanAndroidPlatform.h
--- a/Engine/Source/Runtime/VulkanRHI/Private/Android/VulkanAndroidPlatform.h
+++ b/Engine/Source/Runtime/VulkanRHI/Private/Android/VulkanAndroidPlatform.h
@@ -142,7 +142,7 @@ public:
 	static void DestroySwapchainKHR(VkDevice Device, VkSwapchainKHR Swapchain, const VkAllocationCallbacks* Allocator);
 
 	// handle precompile of PSOs, send to an android specific precompile external process.
-	static VkPipelineCache PrecompilePSO(FVulkanDevice* Device, VkGraphicsPipelineCreateInfo* PipelineInfo, FGfxPipelineDesc* GfxEntry, const FVulkanRenderTargetLayout* RTLayout, TArrayView<uint32_t> VS, TArrayView<uint32_t> PS, size_t& AfterSize);
+	static VkPipelineCache PrecompilePSO(FVulkanDevice* Device, const VkGraphicsPipelineCreateInfo* PipelineInfo, FGfxPipelineDesc* GfxEntry, const FVulkanRenderTargetLayout* RTLayout, TArrayView<uint32_t> VS, TArrayView<uint32_t> PS, size_t& AfterSize);
 
 	static bool AreRemoteCompileServicesActive();
 	static bool StartAndWaitForRemoteCompileServices(int NumServices);
diff --git a/Engine/Source/Runtime/VulkanRHI/Private/VulkanPipeline.cpp b/Engine/Source/Runtime/VulkanRHI/Private/VulkanPipeline.cpp
--- a/Engine/Source/Runtime/VulkanRHI/Private/VulkanPipeline.cpp
+++ b/Engine/Source/Runtime/VulkanRHI/Private/VulkanPipeline.cpp
@@ -1303,7 +1303,7 @@ bool FVulkanPipelineStateCacheManager::CreateGfxPipelineFromEntry(FVulkanRHIGrap
 	// Vertex Input. The structure is mandatory even without vertex attributes.
 	VkPipelineVertexInputStateCreateInfo VBInfo;
 	ZeroVulkanStruct(VBInfo, VK_STRUCTURE_TYPE_PIPELINE_VERTEX_INPUT_STATE_CREATE_INFO);
-	TArray<VkVertexInputBindingDescription> VBBindings;
+	TArray<VkVertexInputBindingDescription, TInlineAllocator<32>> VBBindings;
 	for (const FGfxPipelineDesc::FVertexBinding& SourceBinding : GfxEntry->VertexBindings)
 	{
 		VkVertexInputBindingDescription* Binding = new(VBBindings) VkVertexInputBindingDescription;
@@ -1311,7 +1311,7 @@ bool FVulkanPipelineStateCacheManager::CreateGfxPipelineFromEntry(FVulkanRHIGrap
 	}
 	VBInfo.vertexBindingDescriptionCount = VBBindings.Num();
 	VBInfo.pVertexBindingDescriptions = VBBindings.GetData();
-	TArray<VkVertexInputAttributeDescription> VBAttributes;
+	TArray<VkVertexInputAttributeDescription, TInlineAllocator<32>> VBAttributes;
 	for (const FGfxPipelineDesc::FVertexAttribute& SourceAttr : GfxEntry->VertexAttributes)
 	{
 		VkVertexInputAttributeDescription* Attr = new(VBAttributes) VkVertexInputAttributeDescription;
@@ -1373,6 +1373,11 @@ bool FVulkanPipelineStateCacheManager::CreateGfxPipelineFromEntry(FVulkanRHIGrap
 	}
 #endif
 
+	if (Device->SupportsBindless())
+	{
+		PipelineInfo.flags |= VK_PIPELINE_CREATE_DESCRIPTOR_BUFFER_BIT_EXT;
+	}
+
 	VkResult Result = VK_ERROR_INITIALIZATION_FAILED;
 
 	double BeginTime = FPlatformTime::Seconds();
@@ -1398,17 +1403,12 @@ bool FVulkanPipelineStateCacheManager::CreateGfxPipelineFromEntry(FVulkanRHIGrap
 	return true;
 }
 
-VkResult FVulkanPipelineStateCacheManager::CreateVKPipeline(FVulkanRHIGraphicsPipelineState* PSO, FVulkanShader* Shaders[ShaderStage::NumStages], VkGraphicsPipelineCreateInfo PipelineInfo, bool bIsPrecompileJob)
+VkResult FVulkanPipelineStateCacheManager::CreateVKPipeline(FVulkanRHIGraphicsPipelineState* PSO, FVulkanShader* Shaders[ShaderStage::NumStages], const VkGraphicsPipelineCreateInfo& PipelineInfo, bool bIsPrecompileJob)
 {
 	VkPipeline* Pipeline = &PSO->VulkanPipeline;
 
 	FPipelineCache& Cache = bIsPrecompileJob ? CurrentPrecompilingPSOCache : GlobalPSOCache;
 
-	if (Device->SupportsBindless())
-	{
-		PipelineInfo.flags |= VK_PIPELINE_CREATE_DESCRIPTOR_BUFFER_BIT_EXT;
-	}
-
 	VkPipelineCache LocalPipelineCache = VK_NULL_HANDLE;
 	VkResult Result = VK_ERROR_INITIALIZATION_FAILED;
 	uint32 PSOSize = 0;
diff --git a/Engine/Source/Runtime/VulkanRHI/Private/VulkanPipeline.h b/Engine/Source/Runtime/VulkanRHI/Private/VulkanPipeline.h
--- a/Engine/Source/Runtime/VulkanRHI/Private/VulkanPipeline.h
+++ b/Engine/Source/Runtime/VulkanRHI/Private/VulkanPipeline.h
@@ -525,7 +525,7 @@ private:
 	void NotifyDeletedGraphicsPSO(FRHIGraphicsPipelineState* PSO);
 	bool CreateGfxPipelineFromEntry(FVulkanRHIGraphicsPipelineState* PSO, FVulkanShader* Shaders[ShaderStage::NumStages], bool bPrecompile);
 
-	VkResult CreateVKPipeline(FVulkanRHIGraphicsPipelineState* PSO, FVulkanShader* Shaders[ShaderStage::NumStages], VkGraphicsPipelineCreateInfo PipelineInfo, bool bIsPrecompileJob);
+	VkResult CreateVKPipeline(FVulkanRHIGraphicsPipelineState* PSO, FVulkanShader* Shaders[ShaderStage::NumStages], const VkGraphicsPipelineCreateInfo& PipelineInfo, bool bIsPrecompileJob);
 	static FString ShaderHashesToString(FVulkanShader* Shaders[ShaderStage::NumStages]);
 
 	FVulkanLayout* FindOrAddLayout(const FVulkanDescriptorSetsLayoutInfo& DescriptorSetLayoutInfo, bool bGfxLayout);
