VulkanRHI:  VMA WIP:  Add new AllocateBufferMemory function that takes in a VkBuffer instead of VkMemoryRequirements.  Gather up all the memory types (hostvisible, hostcached, etc), params (memoryless, external, dedicated) and functionality (bind, no errors, etc) into a single collection of flags passed in a single argument instead of a collection of bools.

#rnx

diff --git a/Engine/Source/Runtime/VulkanRHI/Private/VulkanMemory.cpp b/Engine/Source/Runtime/VulkanRHI/Private/VulkanMemory.cpp
--- a/Engine/Source/Runtime/VulkanRHI/Private/VulkanMemory.cpp
+++ b/Engine/Source/Runtime/VulkanRHI/Private/VulkanMemory.cpp
@@ -2681,6 +2681,43 @@ namespace VulkanRHI
 		return Priority;
 	}
 
+	static VkMemoryPropertyFlags GetMemoryPropertyFlags(EVulkanAllocationFlags AllocFlags, bool bHasUnifiedMemory)
+	{
+		VkMemoryPropertyFlags MemFlags = 0;
+
+		checkf(!(EnumHasAnyFlags(AllocFlags, EVulkanAllocationFlags::PreferBAR) && !EnumHasAnyFlags(AllocFlags, EVulkanAllocationFlags::HostVisible)), TEXT("PreferBAR should always be used with HostVisible."));
+
+		if (EnumHasAnyFlags(AllocFlags, EVulkanAllocationFlags::HostCached))
+		{
+			MemFlags = VK_MEMORY_PROPERTY_HOST_VISIBLE_BIT | VK_MEMORY_PROPERTY_HOST_CACHED_BIT;
+		}
+		else if (EnumHasAnyFlags(AllocFlags, EVulkanAllocationFlags::HostVisible))
+		{
+			MemFlags = VK_MEMORY_PROPERTY_HOST_VISIBLE_BIT | VK_MEMORY_PROPERTY_HOST_COHERENT_BIT;
+
+			if (EnumHasAnyFlags(AllocFlags, EVulkanAllocationFlags::PreferBAR))
+			{
+				MemFlags |= VK_MEMORY_PROPERTY_DEVICE_LOCAL_BIT;
+			}
+		}
+		else
+		{
+			MemFlags = VK_MEMORY_PROPERTY_DEVICE_LOCAL_BIT;
+
+			if (bHasUnifiedMemory)
+			{
+				MemFlags |= (VK_MEMORY_PROPERTY_HOST_VISIBLE_BIT | VK_MEMORY_PROPERTY_HOST_COHERENT_BIT);
+			}
+		}
+
+		if (EnumHasAnyFlags(AllocFlags, EVulkanAllocationFlags::Memoryless))
+		{
+			MemFlags = VK_MEMORY_PROPERTY_LAZILY_ALLOCATED_BIT;
+		}
+
+		return MemFlags;
+	}
+
 	bool FMemoryManager::AllocateBufferPooled(FVulkanAllocation& OutAllocation, FVulkanEvictable* AllocationOwner, uint32 Size, uint32 MinAlignment, VkBufferUsageFlags BufferUsageFlags, VkMemoryPropertyFlags MemoryPropertyFlags, EVulkanAllocationMetaType MetaType, const char* File, uint32 Line)
 	{
 		SCOPED_NAMED_EVENT(FResourceHeapManager_AllocateBufferPooled, FColor::Cyan);
@@ -2928,7 +2965,7 @@ namespace VulkanRHI
 		return true;
 	}
 
-	bool FMemoryManager::AllocateBufferMemory(FVulkanAllocation& OutAllocation, FVulkanEvictable* AllocationOwner, const VkMemoryRequirements& MemoryReqs, VkMemoryPropertyFlags MemoryPropertyFlags, EVulkanAllocationMetaType MetaType, bool bExternal, const char* File, uint32 Line)
+	bool FMemoryManager::AllocateBufferMemory(FVulkanAllocation& OutAllocation, FVulkanEvictable* AllocationOwner, const VkMemoryRequirements& MemoryReqs, VkMemoryPropertyFlags MemoryPropertyFlags, EVulkanAllocationMetaType MetaType, bool bExternal, bool bForceSeparateAllocation, const char* File, uint32 Line)
 	{
 		SCOPED_NAMED_EVENT(FResourceHeapManager_AllocateBufferMemory, FColor::Cyan);
 		uint32 TypeIndex = 0;
@@ -2964,7 +3001,7 @@ namespace VulkanRHI
 
 		check(MemoryReqs.size <= (uint64)MAX_uint32);
 
-		if (!ResourceTypeHeaps[TypeIndex]->AllocateResource(OutAllocation, AllocationOwner, EType::Buffer, MemoryReqs.size, MemoryReqs.alignment, bMapped, false, MetaType, bExternal, File, Line))
+		if (!ResourceTypeHeaps[TypeIndex]->AllocateResource(OutAllocation, AllocationOwner, EType::Buffer, MemoryReqs.size, MemoryReqs.alignment, bMapped, bForceSeparateAllocation, MetaType, bExternal, File, Line))
 		{
 			// Try another memory type if the allocation failed
 			VERIFYVULKANRESULT(DeviceMemoryManager->GetMemoryTypeFromPropertiesExcluding(MemoryReqs.memoryTypeBits, MemoryPropertyFlags, TypeIndex, &TypeIndex));
@@ -2973,7 +3010,7 @@ namespace VulkanRHI
 			{
 				UE_LOG(LogVulkanRHI, Fatal, TEXT("Missing memory type index %d, MemSize %d, MemPropTypeBits %u, MemPropertyFlags %u, %s(%d)"), TypeIndex, (uint32)MemoryReqs.size, (uint32)MemoryReqs.memoryTypeBits, (uint32)MemoryPropertyFlags, ANSI_TO_TCHAR(File), Line);
 			}
-			if (!ResourceTypeHeaps[TypeIndex]->AllocateResource(OutAllocation, AllocationOwner, EType::Buffer, MemoryReqs.size, MemoryReqs.alignment, bMapped, false, MetaType, bExternal, File, Line))
+			if (!ResourceTypeHeaps[TypeIndex]->AllocateResource(OutAllocation, AllocationOwner, EType::Buffer, MemoryReqs.size, MemoryReqs.alignment, bMapped, bForceSeparateAllocation, MetaType, bExternal, File, Line))
 			{
 				DumpMemory();
 				UE_LOG(LogVulkanRHI, Fatal, TEXT("Out Of Memory, trying to allocate %d bytes\n"), MemoryReqs.size);
@@ -2983,6 +3020,67 @@ namespace VulkanRHI
 		return true;
 	}
 
+	bool FMemoryManager::AllocateBufferMemory(FVulkanAllocation& OutAllocation, VkBuffer InBuffer, EVulkanAllocationFlags InAllocFlags, const TCHAR* InDebugName, uint32 InForceMinAlignment)
+	{
+		SCOPED_NAMED_EVENT(FMemoryManager_AllocateBufferMemory, FColor::Cyan);
+
+		VkBufferMemoryRequirementsInfo2 BufferMemoryRequirementsInfo;
+		ZeroVulkanStruct(BufferMemoryRequirementsInfo, VK_STRUCTURE_TYPE_BUFFER_MEMORY_REQUIREMENTS_INFO_2);
+		BufferMemoryRequirementsInfo.buffer = InBuffer;
+
+		VkMemoryDedicatedRequirements DedicatedRequirements;
+		ZeroVulkanStruct(DedicatedRequirements, VK_STRUCTURE_TYPE_MEMORY_DEDICATED_REQUIREMENTS);
+
+		VkMemoryRequirements2 MemoryRequirements;
+		ZeroVulkanStruct(MemoryRequirements, VK_STRUCTURE_TYPE_MEMORY_REQUIREMENTS_2);
+		MemoryRequirements.pNext = &DedicatedRequirements;
+
+		VulkanRHI::vkGetBufferMemoryRequirements2(Device->GetInstanceHandle(), &BufferMemoryRequirementsInfo, &MemoryRequirements);
+
+		// Allow caller to force his own alignment requirements
+		MemoryRequirements.memoryRequirements.alignment = FMath::Max(MemoryRequirements.memoryRequirements.alignment, (VkDeviceSize)InForceMinAlignment);
+
+		if (DedicatedRequirements.requiresDedicatedAllocation || DedicatedRequirements.prefersDedicatedAllocation)
+		{
+			InAllocFlags |= EVulkanAllocationFlags::Dedicated;
+		}
+
+		// For now, translate all the flags into a call to the legacy AllocateBufferMemory() function
+		const VkMemoryPropertyFlags MemoryPropertyFlags = GetMemoryPropertyFlags(InAllocFlags, DeviceMemoryManager->HasUnifiedMemory());
+		const bool bExternal = EnumHasAllFlags(InAllocFlags, EVulkanAllocationFlags::External);
+		const bool bForceSeparateAllocation = EnumHasAllFlags(InAllocFlags, EVulkanAllocationFlags::Dedicated);
+		AllocateBufferMemory(OutAllocation, nullptr, MemoryRequirements.memoryRequirements, MemoryPropertyFlags, EVulkanAllocationMetaBufferOther, bExternal, bForceSeparateAllocation, __FILE__, __LINE__);
+
+		if (OutAllocation.IsValid())
+		{
+			if (EnumHasAllFlags(InAllocFlags, EVulkanAllocationFlags::AutoBind))
+			{
+				VkBindBufferMemoryInfo BindBufferMemoryInfo;
+				ZeroVulkanStruct(BindBufferMemoryInfo, VK_STRUCTURE_TYPE_BIND_BUFFER_MEMORY_INFO);
+				BindBufferMemoryInfo.buffer = InBuffer;
+				BindBufferMemoryInfo.memory = OutAllocation.GetDeviceMemoryHandle(Device);
+				BindBufferMemoryInfo.memoryOffset = OutAllocation.Offset;
+
+				VERIFYVULKANRESULT(VulkanRHI::vkBindBufferMemory2(Device->GetInstanceHandle(), 1, &BindBufferMemoryInfo));
+			}
+
+			if (InDebugName)
+			{
+				VULKAN_SET_DEBUG_NAME((*Device), VK_OBJECT_TYPE_BUFFER, InBuffer, TEXT("%s"), InDebugName);
+			}
+		}
+		else
+		{
+			if (!EnumHasAllFlags(InAllocFlags, EVulkanAllocationFlags::NoError))
+			{
+				const bool IsHostMemory = EnumHasAnyFlags(InAllocFlags, EVulkanAllocationFlags::HostVisible | EVulkanAllocationFlags::HostCached);
+				HandleOOM(false, IsHostMemory ? VK_ERROR_OUT_OF_HOST_MEMORY : VK_ERROR_OUT_OF_DEVICE_MEMORY, MemoryRequirements.memoryRequirements.size);
+			}
+		}
+
+		return OutAllocation.IsValid();
+	}
+
 	bool FMemoryManager::AllocateDedicatedImageMemory(FVulkanAllocation& OutAllocation, FVulkanEvictable* AllocationOwner, VkImage Image, const VkMemoryRequirements& MemoryReqs, VkMemoryPropertyFlags MemoryPropertyFlags, EVulkanAllocationMetaType MetaType, bool bExternal, const char* File, uint32 Line)
 	{
 #if VULKAN_SUPPORTS_DEDICATED_ALLOCATION
@@ -4189,10 +4287,11 @@ namespace VulkanRHI
 		SCOPE_CYCLE_COUNTER(STAT_VulkanStagingBuffer);
 #endif
 		LLM_SCOPE_VULKAN(ELLMTagVulkan::VulkanStagingBuffers);
-		if (InMemoryReadFlags == VK_MEMORY_PROPERTY_HOST_CACHED_BIT)
+
+		const bool IsHostCached = (InMemoryReadFlags == VK_MEMORY_PROPERTY_HOST_CACHED_BIT);
+		if (IsHostCached)
 		{
-			uint64 NonCoherentAtomSize = (uint64)Device->GetLimits().nonCoherentAtomSize;
-			Size = AlignArbitrary(Size, NonCoherentAtomSize);
+			Size = AlignArbitrary(Size, (uint32)Device->GetLimits().nonCoherentAtomSize);
 		}
 
 		// Add both source and dest flags
@@ -4201,6 +4300,12 @@ namespace VulkanRHI
 			InUsageFlags |= (VK_BUFFER_USAGE_TRANSFER_SRC_BIT | VK_BUFFER_USAGE_TRANSFER_DST_BIT);
 		}
 
+		// For descriptors buffers
+		if (Device->SupportsBindless())
+		{
+			InUsageFlags |= VK_BUFFER_USAGE_SHADER_DEVICE_ADDRESS_BIT;
+		}
+
 		//#todo-rco: Better locking!
 		{
 			FScopeLock Lock(&StagingLock);
@@ -4219,43 +4324,27 @@ namespace VulkanRHI
 		}
 
 		FStagingBuffer* StagingBuffer = new FStagingBuffer(Device);
+		StagingBuffer->MemoryReadFlags = InMemoryReadFlags;
+		StagingBuffer->BufferSize = Size;
 
 		VkBufferCreateInfo StagingBufferCreateInfo;
 		ZeroVulkanStruct(StagingBufferCreateInfo, VK_STRUCTURE_TYPE_BUFFER_CREATE_INFO);
 		StagingBufferCreateInfo.size = Size;
 		StagingBufferCreateInfo.usage = InUsageFlags;
-		// For descriptors buffers
-		if (Device->SupportsBindless())
-		{
-			StagingBufferCreateInfo.usage |= VK_BUFFER_USAGE_SHADER_DEVICE_ADDRESS_BIT;
-		}
-
-		VkDevice VulkanDevice = Device->GetInstanceHandle();
-
-		VERIFYVULKANRESULT(VulkanRHI::vkCreateBuffer(VulkanDevice, &StagingBufferCreateInfo, VULKAN_CPU_ALLOCATOR, &StagingBuffer->Buffer));
-
-		VkMemoryRequirements MemReqs;
-		VulkanRHI::vkGetBufferMemoryRequirements(VulkanDevice, StagingBuffer->Buffer, &MemReqs);
-		ensure(MemReqs.size >= Size);
+		VERIFYVULKANRESULT(VulkanRHI::vkCreateBuffer(Device->GetInstanceHandle(), &StagingBufferCreateInfo, VULKAN_CPU_ALLOCATOR, &StagingBuffer->Buffer));
 
 		// Set minimum alignment to 16 bytes, as some buffers are used with CPU SIMD instructions
-		MemReqs.alignment = FMath::Max<VkDeviceSize>(16, MemReqs.alignment);
-		static const bool bIsAmd = Device->GetDeviceProperties().vendorID == 0x1002;
-		if (InMemoryReadFlags == VK_MEMORY_PROPERTY_HOST_CACHED_BIT || bIsAmd)
+		uint32 ForcedMinAlignment = 16u;
+		static const bool bIsAmd = (Device->GetDeviceProperties().vendorID == (uint32)EGpuVendorId::Amd);
+		if (IsHostCached || bIsAmd)
 		{
-			uint64 NonCoherentAtomSize = (uint64)Device->GetLimits().nonCoherentAtomSize;
-			MemReqs.alignment = AlignArbitrary(MemReqs.alignment, NonCoherentAtomSize);
+			ForcedMinAlignment = AlignArbitrary(ForcedMinAlignment, (uint32)Device->GetLimits().nonCoherentAtomSize);
 		}
 
-		VkMemoryPropertyFlags readTypeFlags = InMemoryReadFlags;
-		if (!Device->GetMemoryManager().AllocateBufferMemory(StagingBuffer->Allocation, nullptr, MemReqs, VK_MEMORY_PROPERTY_HOST_VISIBLE_BIT | readTypeFlags, EVulkanAllocationMetaBufferStaging, false, __FILE__, __LINE__))
-		{
-			check(0);
-		}
-		StagingBuffer->MemoryReadFlags = InMemoryReadFlags;
-		StagingBuffer->BufferSize = Size;
-		StagingBuffer->Allocation.BindBuffer(Device, StagingBuffer->Buffer);
-		//StagingBuffer->ResourceAllocation->BindBuffer(Device, StagingBuffer->Buffer);
+		const EVulkanAllocationFlags AllocFlags = EVulkanAllocationFlags::AutoBind |
+			(IsHostCached ? EVulkanAllocationFlags::HostCached : EVulkanAllocationFlags::HostVisible);
+
+		Device->GetMemoryManager().AllocateBufferMemory(StagingBuffer->Allocation, StagingBuffer->Buffer, AllocFlags, TEXT("StagingBuffer"), ForcedMinAlignment);
 
 		{
 			FScopeLock Lock(&StagingLock);
diff --git a/Engine/Source/Runtime/VulkanRHI/Private/VulkanMemory.h b/Engine/Source/Runtime/VulkanRHI/Private/VulkanMemory.h
--- a/Engine/Source/Runtime/VulkanRHI/Private/VulkanMemory.h
+++ b/Engine/Source/Runtime/VulkanRHI/Private/VulkanMemory.h
@@ -147,7 +147,7 @@ namespace VulkanRHI
 		EVulkanAllocationMetaBufferOther,
 		EVulkanAllocationMetaSize, // keep last
 	};
-	enum EVulkanAllocationFlags
+	enum ELegacyVulkanAllocationFlags
 	{
 		VulkanAllocationFlagsMapped = 0x1,
 		VulkanAllocationFlagsCanEvict= 0x2,
@@ -739,6 +739,24 @@ namespace VulkanRHI
 	};
 
 
+	// Allocation flags to account for conditions not covered by vkGet*MemoryRequirements or resource usage flags
+	enum class EVulkanAllocationFlags : uint16
+	{
+		None         = 0x0000,
+
+		HostVisible  = 0x0001,    // Will be written from CPU, will likely contain the HOST_VISIBLE + HOST_COHERENT flags
+		HostCached   = 0x0002,    // Will be used for readback, will likely contain the HOST_CACHED flag.  Implies HostVisible.
+		PreferBAR    = 0x0004,    // Will be allocated from a HOST_VISIBLE + DEVICE_LOCAL heap if available and possible (HOST_VISIBLE if not)
+
+		Dedicated    = 0x0008,	  // Will not share a memory block with other resources
+		External     = 0x0010,    // To be used with VK_KHR_external_memory
+		Memoryless   = 0x0020,	  // Will use LAZILY_ALLOCATED
+
+		NoError      = 0x0040,    // OOM is not fatal, return an invalid allocation
+		AutoBind     = 0x0080,	  // Will automatically bind the allocation to the supplied VkBuffer/VkImage, avoids an extra lock to bind separately
+	};
+	ENUM_CLASS_FLAGS(EVulkanAllocationFlags);
+
 	// Manages heaps and their interactions
 	class FMemoryManager : public FDeviceChild
 	{
@@ -760,11 +778,16 @@ namespace VulkanRHI
 		void FreeVulkanAllocationImage(FVulkanAllocation& Allocation);
 		void FreeVulkanAllocationImageDedicated(FVulkanAllocation& Allocation);
 
-
+		// Legacy calls
 		bool AllocateBufferPooled(FVulkanAllocation& Allocation, FVulkanEvictable* AllocationOwner, uint32 Size, uint32 MinAlignment, VkBufferUsageFlags BufferUsageFlags, VkMemoryPropertyFlags MemoryPropertyFlags, EVulkanAllocationMetaType MetaType, const char* File, uint32 Line);
 		bool AllocateImageMemory(FVulkanAllocation& Allocation, FVulkanEvictable* AllocationOwner, const VkMemoryRequirements& MemoryReqs, VkMemoryPropertyFlags MemoryPropertyFlags, EVulkanAllocationMetaType MetaType, bool bExternal, const char* File, uint32 Line);
-		bool AllocateBufferMemory(FVulkanAllocation& Allocation, FVulkanEvictable* AllocationOwner, const VkMemoryRequirements& MemoryReqs, VkMemoryPropertyFlags MemoryPropertyFlags, EVulkanAllocationMetaType MetaType, bool bExternal, const char* File, uint32 Line);
 		bool AllocateDedicatedImageMemory(FVulkanAllocation& Allocation, FVulkanEvictable* AllocationOwner, VkImage Image, const VkMemoryRequirements& MemoryReqs, VkMemoryPropertyFlags MemoryPropertyFlags, EVulkanAllocationMetaType MetaType, bool bExternal, const char* File, uint32 Line);
+	private:
+		bool AllocateBufferMemory(FVulkanAllocation& Allocation, FVulkanEvictable* AllocationOwner, const VkMemoryRequirements& MemoryReqs, VkMemoryPropertyFlags MemoryPropertyFlags, EVulkanAllocationMetaType MetaType, bool bExternal, bool bForceSeparateAllocation, const char* File, uint32 Line);
+
+		// New calls
+	public:
+		bool AllocateBufferMemory(FVulkanAllocation& OutAllocation, VkBuffer InBuffer, EVulkanAllocationFlags InAllocFlags, const TCHAR* InDebugName, uint32 InForceMinAlignment = 1);
 
 		void RegisterSubresourceAllocator(FVulkanSubresourceAllocator* SubresourceAllocator);
 		void UnregisterSubresourceAllocator(FVulkanSubresourceAllocator* SubresourceAllocator);
diff --git a/Engine/Source/Runtime/VulkanRHI/Private/VulkanTexture.cpp b/Engine/Source/Runtime/VulkanRHI/Private/VulkanTexture.cpp
--- a/Engine/Source/Runtime/VulkanRHI/Private/VulkanTexture.cpp
+++ b/Engine/Source/Runtime/VulkanRHI/Private/VulkanTexture.cpp
@@ -1453,9 +1453,10 @@ FVulkanTexture::FVulkanTexture(FVulkanDevice& InDevice, const FRHITextureCreateD
 	{
 		if (EnumHasAnyFlags(InCreateDesc.Flags, TexCreate_CPUReadback))
 		{
-			check(InCreateDesc.NumSamples == 1);	//not implemented
-			check(InCreateDesc.Depth == 1);		//not implemented
-			check(InCreateDesc.ArraySize == 1);	//not implemented
+			check(InCreateDesc.NumSamples == 1); //not implemented
+			check(InCreateDesc.Depth == 1);      //not implemented
+			check(InCreateDesc.ArraySize == 1);  //not implemented
+
 			CpuReadbackBuffer = new FVulkanCpuReadbackBuffer;
 			uint32 Size = 0;
 			for (uint32 Mip = 0; Mip < InCreateDesc.NumMips; ++Mip)
@@ -1467,29 +1468,26 @@ FVulkanTexture::FVulkanTexture(FVulkanDevice& InDevice, const FRHITextureCreateD
 				Size += LocalSize;
 			}
 
-			VkDevice VulkanDevice = InDevice.GetInstanceHandle();
-			const VkMemoryPropertyFlags BufferMemFlags = (VK_MEMORY_PROPERTY_HOST_VISIBLE_BIT | VK_MEMORY_PROPERTY_HOST_COHERENT_BIT | VK_MEMORY_PROPERTY_HOST_CACHED_BIT);
-
 			VkBufferCreateInfo BufferCreateInfo;
 			ZeroVulkanStruct(BufferCreateInfo, VK_STRUCTURE_TYPE_BUFFER_CREATE_INFO);
 			BufferCreateInfo.size = Size;
 			BufferCreateInfo.usage = VK_BUFFER_USAGE_TRANSFER_DST_BIT;
+			VERIFYVULKANRESULT(VulkanRHI::vkCreateBuffer(InDevice.GetInstanceHandle(), &BufferCreateInfo, VULKAN_CPU_ALLOCATOR, &CpuReadbackBuffer->Buffer));
 
-			VERIFYVULKANRESULT(VulkanRHI::vkCreateBuffer(VulkanDevice, &BufferCreateInfo, VULKAN_CPU_ALLOCATOR, &CpuReadbackBuffer->Buffer));
-			VulkanRHI::vkGetBufferMemoryRequirements(VulkanDevice, CpuReadbackBuffer->Buffer, &MemoryRequirements);
 			// Set minimum alignment to 16 bytes, as some buffers are used with CPU SIMD instructions
-			MemoryRequirements.alignment = FMath::Max<VkDeviceSize>(16, MemoryRequirements.alignment);
-			if (!InDevice.GetMemoryManager().AllocateBufferMemory(Allocation, this, MemoryRequirements, BufferMemFlags, EVulkanAllocationMetaBufferStaging, false, __FILE__, __LINE__))
-			{
-				InDevice.GetMemoryManager().HandleOOM();
-			}
-			Allocation.BindBuffer(Device, CpuReadbackBuffer->Buffer);
+			const uint32 ForcedMinAlignment = 16u;
+			const EVulkanAllocationFlags AllocFlags = EVulkanAllocationFlags::HostCached | EVulkanAllocationFlags::AutoBind;
+			InDevice.GetMemoryManager().AllocateBufferMemory(Allocation, CpuReadbackBuffer->Buffer, AllocFlags, InCreateDesc.DebugName, ForcedMinAlignment);
+
 			void* Memory = Allocation.GetMappedPointer(Device);
-			FMemory::Memzero(Memory, MemoryRequirements.size);
+			FMemory::Memzero(Memory, Size);
 
 			ImageOwnerType = EImageOwnerType::None;
 			ViewFormat = StorageFormat = UEToVkTextureFormat(InCreateDesc.Format, false);
 
+			// :todo-jn: Kept around temporarily for legacy defrag/eviction/stats
+			VulkanRHI::vkGetBufferMemoryRequirements(InDevice.GetInstanceHandle(), CpuReadbackBuffer->Buffer, &MemoryRequirements);
+
 			return;
 		}
 
diff --git a/Engine/Source/Runtime/VulkanRHI/Private/VulkanTransientResourceAllocator.cpp b/Engine/Source/Runtime/VulkanRHI/Private/VulkanTransientResourceAllocator.cpp
--- a/Engine/Source/Runtime/VulkanRHI/Private/VulkanTransientResourceAllocator.cpp
+++ b/Engine/Source/Runtime/VulkanRHI/Private/VulkanTransientResourceAllocator.cpp
@@ -23,37 +23,17 @@ FVulkanTransientHeap::FVulkanTransientHeap(const FInitializer& Initializer, FVul
 	const bool bZeroSize = false;
 	VkBufferUsageFlags BufferUsageFlags = FVulkanResourceMultiBuffer::UEToVKBufferUsageFlags(InDevice, UEBufferUsageFlags, bZeroSize);
 
-	// :TODO: VK_KHR_maintenance4...
-	{
-		VkBufferCreateInfo BufferCreateInfo;
-		ZeroVulkanStruct(BufferCreateInfo, VK_STRUCTURE_TYPE_BUFFER_CREATE_INFO);
-		BufferCreateInfo.size = Initializer.Size;
-		BufferCreateInfo.usage = BufferUsageFlags;
-
-		const VkDevice VulkanDevice = InDevice->GetInstanceHandle();
-
-		VERIFYVULKANRESULT(VulkanRHI::vkCreateBuffer(VulkanDevice, &BufferCreateInfo, VULKAN_CPU_ALLOCATOR, &VulkanBuffer));
-		VulkanRHI::vkGetBufferMemoryRequirements(VulkanDevice, VulkanBuffer, &MemoryRequirements);
-
-		// Find the alignment that works for everyone
-		const uint32 MinBufferAlignment = FMemoryManager::CalculateBufferAlignment(*InDevice, UEBufferUsageFlags, bZeroSize);
-		MemoryRequirements.alignment = FMath::Max<VkDeviceSize>(Initializer.Alignment, MemoryRequirements.alignment);
-		MemoryRequirements.alignment = FMath::Max<VkDeviceSize>(MinBufferAlignment, MemoryRequirements.alignment);
-	}
+	VkBufferCreateInfo BufferCreateInfo;
+	ZeroVulkanStruct(BufferCreateInfo, VK_STRUCTURE_TYPE_BUFFER_CREATE_INFO);
+	BufferCreateInfo.size = Initializer.Size;
+	BufferCreateInfo.usage = BufferUsageFlags;
+	VERIFYVULKANRESULT(VulkanRHI::vkCreateBuffer(InDevice->GetInstanceHandle(), &BufferCreateInfo, VULKAN_CPU_ALLOCATOR, &VulkanBuffer));
 
-	VkMemoryPropertyFlags BufferMemFlags = VK_MEMORY_PROPERTY_DEVICE_LOCAL_BIT;
-	const bool bUnifiedMem = InDevice->HasUnifiedMemory();
-	if (bUnifiedMem)
-	{
-		BufferMemFlags |= (VK_MEMORY_PROPERTY_HOST_VISIBLE_BIT | VK_MEMORY_PROPERTY_HOST_COHERENT_BIT);
-	}
-
-	if (!InDevice->GetMemoryManager().AllocateBufferMemory(InternalAllocation, nullptr, MemoryRequirements, BufferMemFlags, EVulkanAllocationMetaBufferOther, false, __FILE__, __LINE__))
-	{
-		InDevice->GetMemoryManager().HandleOOM();
-	}
+	// Find the alignment that works for everyone
+	const uint32 MinBufferAlignment = FMath::Max<uint32>(Initializer.Alignment, FMemoryManager::CalculateBufferAlignment(*InDevice, UEBufferUsageFlags, bZeroSize));
 
-	InternalAllocation.BindBuffer(InDevice, VulkanBuffer);
+	const EVulkanAllocationFlags AllocFlags = EVulkanAllocationFlags::Dedicated | EVulkanAllocationFlags::AutoBind;
+	InDevice->GetMemoryManager().AllocateBufferMemory(InternalAllocation, VulkanBuffer, AllocFlags, TEXT("VulkanTransientHeap"), MinBufferAlignment);
 }
 
 FVulkanTransientHeap::~FVulkanTransientHeap()
