VulkanRHI: Fix for Compaction running when ray tracing is not supported on a device

diff --git a/Engine/Source/Runtime/VulkanRHI/Private/VulkanDevice.cpp b/Engine/Source/Runtime/VulkanRHI/Private/VulkanDevice.cpp
--- a/Engine/Source/Runtime/VulkanRHI/Private/VulkanDevice.cpp
+++ b/Engine/Source/Runtime/VulkanRHI/Private/VulkanDevice.cpp
@@ -1281,8 +1281,11 @@ void FVulkanDevice::InitGPU()
 	}
 
 #if VULKAN_RHI_RAYTRACING
-	check(RayTracingCompactionRequestHandler == nullptr);
-	RayTracingCompactionRequestHandler = new FVulkanRayTracingCompactionRequestHandler(this);
+	if (RHISupportsRayTracing(GMaxRHIShaderPlatform) && GetOptionalExtensions().HasRaytracingExtensions())
+	{
+		check(RayTracingCompactionRequestHandler == nullptr);
+		RayTracingCompactionRequestHandler = new FVulkanRayTracingCompactionRequestHandler(this);
+	}
 #endif
 
 	FVulkanPlatform::SetupImageMemoryRequirementWorkaround(*this);
diff --git a/Engine/Source/Runtime/VulkanRHI/Private/VulkanRHI.cpp b/Engine/Source/Runtime/VulkanRHI/Private/VulkanRHI.cpp
--- a/Engine/Source/Runtime/VulkanRHI/Private/VulkanRHI.cpp
+++ b/Engine/Source/Runtime/VulkanRHI/Private/VulkanRHI.cpp
@@ -1073,7 +1073,10 @@ void FVulkanCommandListContext::RHIBeginFrame()
 	GpuProfiler.BeginFrame();
 
 #if VULKAN_RHI_RAYTRACING
-	Device->GetRayTracingCompactionRequestHandler()->Update(*this);
+	if (GRHISupportsRayTracing)
+	{
+		Device->GetRayTracingCompactionRequestHandler()->Update(*this);
+	}
 #endif
 }
 
