VulkanRHI: Use a ::Swap to avoid calling the destructor on duplicated allocations (that weren't freed).  Fix assert on call to RHITransferBufferUnderlyingResource.

#rnx

diff --git a/Engine/Source/Runtime/VulkanRHI/Private/VulkanBuffer.cpp b/Engine/Source/Runtime/VulkanRHI/Private/VulkanBuffer.cpp
--- a/Engine/Source/Runtime/VulkanRHI/Private/VulkanBuffer.cpp
+++ b/Engine/Source/Runtime/VulkanRHI/Private/VulkanBuffer.cpp
@@ -595,12 +595,13 @@ void FVulkanResourceMultiBuffer::TakeOwnership(FVulkanResourceMultiBuffer& Other
 	FRHIBuffer::TakeOwnership(Other);
 
 	BufferUsageFlags   = Other.BufferUsageFlags;
-	BufferAllocs       = Other.BufferAllocs;
 	CurrentBufferIndex = Other.CurrentBufferIndex;
 
 	Other.BufferUsageFlags   = {};
-	Other.BufferAllocs       = {};
 	Other.CurrentBufferIndex = -1;
+
+	// Swap the empty array from the ReleaseOwnership with the Other allocations
+	::Swap(BufferAllocs, Other.BufferAllocs);
 }
 
 void FVulkanResourceMultiBuffer::ReleaseOwnership()
@@ -619,6 +620,7 @@ void FVulkanResourceMultiBuffer::ReleaseOwnership()
 		}
 		BufferAllocs[Index].Fence = nullptr;
 	}
+	BufferAllocs.Empty();
 
 	if (TotalSize > 0)
 	{
