Support VRS for desktop Vulkan via the VK_KHR_fragment_shading_rate extension.

diff --git a/Engine/Config/VulkanPC/DataDrivenPlatformInfo.ini b/Engine/Config/VulkanPC/DataDrivenPlatformInfo.ini
--- a/Engine/Config/VulkanPC/DataDrivenPlatformInfo.ini
+++ b/Engine/Config/VulkanPC/DataDrivenPlatformInfo.ini
@@ -112,7 +112,7 @@ bSupportsGen5TemporalAA=true
 bSupportsPerPixelDBufferMask = false
 
 
-bSupportsVariableRateShading = false
+bSupportsVariableRateShading = true
 
 NumberOfComputeThreads = 64
 
@@ -228,7 +228,7 @@ bSupportsGen5TemporalAA=true
 ; //#todo-vulkanddpi
 bSupportsPerPixelDBufferMask = false
 
-bSupportsVariableRateShading = false
+bSupportsVariableRateShading = true
 
 NumberOfComputeThreads = 64
 
diff --git a/Engine/Source/Runtime/VulkanRHI/Private/VulkanPipeline.cpp b/Engine/Source/Runtime/VulkanRHI/Private/VulkanPipeline.cpp
--- a/Engine/Source/Runtime/VulkanRHI/Private/VulkanPipeline.cpp
+++ b/Engine/Source/Runtime/VulkanRHI/Private/VulkanPipeline.cpp
@@ -1359,7 +1359,7 @@ bool FVulkanPipelineStateCacheManager::CreateGfxPipelineFromEntry(FVulkanRHIGrap
 
 #if VULKAN_SUPPORTS_FRAGMENT_SHADING_RATE
 	VkPipelineFragmentShadingRateStateCreateInfoKHR PipelineFragmentShadingRate;
-	if (GRHISupportsPipelineVariableRateShading && GRHIVariableRateShadingEnabled && PSO->Desc.ShadingRate != EVRSShadingRate::VRSSR_1x1)
+	if (GRHISupportsPipelineVariableRateShading && GRHIVariableRateShadingEnabled)
 	{
 		const VkExtent2D FragmentSize = Device->GetBestMatchedFragmentSize(PSO->Desc.ShadingRate);
 		VkFragmentShadingRateCombinerOpKHR PipelineToPrimitiveCombinerOperation = FragmentCombinerOpMap[(uint8)PSO->Desc.Combiner];
diff --git a/Engine/Source/Runtime/VulkanRHI/Private/VulkanRenderTarget.cpp b/Engine/Source/Runtime/VulkanRHI/Private/VulkanRenderTarget.cpp
--- a/Engine/Source/Runtime/VulkanRHI/Private/VulkanRenderTarget.cpp
+++ b/Engine/Source/Runtime/VulkanRHI/Private/VulkanRenderTarget.cpp
@@ -964,7 +964,7 @@ FVulkanRenderTargetLayout::FVulkanRenderTargetLayout(FVulkanDevice& InDevice, co
 		CurrDesc.flags = 0;
 		CurrDesc.format = UEToVkTextureFormat(RTInfo.ShadingRateTexture->GetFormat(), false);
 		CurrDesc.samples = static_cast<VkSampleCountFlagBits>(RTInfo.ShadingRateTexture->GetNumSamples());
-		CurrDesc.loadOp = VK_ATTACHMENT_LOAD_OP_DONT_CARE;
+		CurrDesc.loadOp = VK_ATTACHMENT_LOAD_OP_LOAD;
 		CurrDesc.storeOp = VK_ATTACHMENT_STORE_OP_DONT_CARE;
 		CurrDesc.stencilLoadOp = VK_ATTACHMENT_LOAD_OP_DONT_CARE;
 		CurrDesc.stencilStoreOp = VK_ATTACHMENT_STORE_OP_DONT_CARE;
@@ -1218,7 +1218,7 @@ FVulkanRenderTargetLayout::FVulkanRenderTargetLayout(FVulkanDevice& InDevice, co
 		CurrDesc.flags = 0;
 		CurrDesc.format = UEToVkTextureFormat(RPInfo.ShadingRateTexture->GetFormat(), false);
 		CurrDesc.samples = static_cast<VkSampleCountFlagBits>(RPInfo.ShadingRateTexture->GetNumSamples());
-		CurrDesc.loadOp = VK_ATTACHMENT_LOAD_OP_DONT_CARE;
+		CurrDesc.loadOp = VK_ATTACHMENT_LOAD_OP_LOAD;
 		CurrDesc.storeOp = VK_ATTACHMENT_STORE_OP_DONT_CARE;
 		CurrDesc.stencilLoadOp = VK_ATTACHMENT_LOAD_OP_DONT_CARE;
 		CurrDesc.stencilStoreOp = VK_ATTACHMENT_STORE_OP_DONT_CARE;
@@ -1383,7 +1383,7 @@ FVulkanRenderTargetLayout::FVulkanRenderTargetLayout(const FGraphicsPipelineStat
 		CurrDesc.flags = 0;
 		CurrDesc.format = UEToVkTextureFormat(GRHIVariableRateShadingImageFormat, false);
 		CurrDesc.samples = VK_SAMPLE_COUNT_1_BIT;
-		CurrDesc.loadOp = VK_ATTACHMENT_LOAD_OP_DONT_CARE;
+		CurrDesc.loadOp = VK_ATTACHMENT_LOAD_OP_LOAD;
 		CurrDesc.storeOp = VK_ATTACHMENT_STORE_OP_DONT_CARE;
 		CurrDesc.stencilLoadOp = VK_ATTACHMENT_LOAD_OP_DONT_CARE;
 		CurrDesc.stencilStoreOp = VK_ATTACHMENT_STORE_OP_DONT_CARE;
