VulkanRHI: Clean up: Remove redundant StageFlag in FVulkanShader.  Remove unused VkShaderStageFlagBits param in PatchSpirvBindings call.

#rnx

diff --git a/Engine/Source/Runtime/VulkanRHI/Private/VulkanDescriptorSets.h b/Engine/Source/Runtime/VulkanRHI/Private/VulkanDescriptorSets.h
--- a/Engine/Source/Runtime/VulkanRHI/Private/VulkanDescriptorSets.h
+++ b/Engine/Source/Runtime/VulkanRHI/Private/VulkanDescriptorSets.h
@@ -1295,7 +1295,7 @@ public:
 		return DescriptorSetLayout.GetHash();
 	}
 
-	void PatchSpirvBindings(FVulkanShader::FSpirvCode& SpirvCode, EShaderFrequency Frequency, const FVulkanShaderHeader& CodeHeader, VkShaderStageFlagBits InStageFlag) const;
+	void PatchSpirvBindings(FVulkanShader::FSpirvCode& SpirvCode, EShaderFrequency Frequency, const FVulkanShaderHeader& CodeHeader) const;
 
 protected:
 	FVulkanDescriptorSetsLayout	DescriptorSetLayout;
diff --git a/Engine/Source/Runtime/VulkanRHI/Private/VulkanPipeline.cpp b/Engine/Source/Runtime/VulkanRHI/Private/VulkanPipeline.cpp
--- a/Engine/Source/Runtime/VulkanRHI/Private/VulkanPipeline.cpp
+++ b/Engine/Source/Runtime/VulkanRHI/Private/VulkanPipeline.cpp
@@ -1181,16 +1181,16 @@ static const TMap<uint8, VkFragmentShadingRateCombinerOpKHR> FragmentCombinerOpM
 FString FVulkanPipelineStateCacheManager::ShaderHashesToString(FVulkanShader* Shaders[ShaderStage::NumStages])
 {
 	FString ShaderHashes = "";
-	if (Shaders[ShaderStage::Vertex] && Shaders[ShaderStage::Vertex]->StageFlag == VK_SHADER_STAGE_VERTEX_BIT)
+	if (Shaders[ShaderStage::Vertex] && Shaders[ShaderStage::Vertex]->Frequency == SF_Vertex)
 	{
 		ShaderHashes += TEXT("VS: ") + static_cast<FVulkanVertexShader*>(Shaders[ShaderStage::Vertex])->GetHash().ToString() + TEXT("\n");
 	}
-	if (Shaders[ShaderStage::Pixel] && Shaders[ShaderStage::Pixel]->StageFlag == VK_SHADER_STAGE_FRAGMENT_BIT)
+	if (Shaders[ShaderStage::Pixel] && Shaders[ShaderStage::Pixel]->Frequency == SF_Pixel)
 	{
 		ShaderHashes += TEXT("PS: ") + static_cast<FVulkanPixelShader*>(Shaders[ShaderStage::Pixel])->GetHash().ToString() + TEXT("\n");
 	}
 #if VULKAN_SUPPORTS_GEOMETRY_SHADERS
-	if (Shaders[ShaderStage::Geometry] && Shaders[ShaderStage::Geometry]->StageFlag == VK_SHADER_STAGE_GEOMETRY_BIT)
+	if (Shaders[ShaderStage::Geometry] && Shaders[ShaderStage::Geometry]->Frequency == SF_Geometry)
 	{
 		ShaderHashes += TEXT("GS: ") + static_cast<FVulkanGeometryShader*>(Shaders[ShaderStage::Geometry])->GetHash().ToString() + TEXT("\n");
 	}
diff --git a/Engine/Source/Runtime/VulkanRHI/Private/VulkanResources.h b/Engine/Source/Runtime/VulkanRHI/Private/VulkanResources.h
--- a/Engine/Source/Runtime/VulkanRHI/Private/VulkanResources.h
+++ b/Engine/Source/Runtime/VulkanRHI/Private/VulkanResources.h
@@ -102,9 +102,8 @@ class FVulkanShader : public IRefCountedObject
 	static FCriticalSection VulkanShaderModulesMapCS;
 
 public:
-	FVulkanShader(FVulkanDevice* InDevice, EShaderFrequency InFrequency, VkShaderStageFlagBits InStageFlag)
+	FVulkanShader(FVulkanDevice* InDevice, EShaderFrequency InFrequency)
 		: ShaderKey(0)
-		, StageFlag(InStageFlag)
 		, Frequency(InFrequency)
 		, Device(InDevice)
 	{
@@ -190,8 +189,7 @@ protected:
 	/** External bindings for this shader. */
 	FVulkanShaderHeader				CodeHeader;
 	TMap<uint32, TRefCountPtr<FVulkanShaderModule>>	ShaderModules;
-	const VkShaderStageFlagBits		StageFlag;
-	EShaderFrequency				Frequency;
+	const EShaderFrequency			Frequency;
 
 	TArray<FUniformBufferStaticSlot> StaticSlots;
 
@@ -230,12 +228,12 @@ protected:
 };
 
 /** This represents a vertex shader that hasn't been combined with a specific declaration to create a bound shader. */
-template<typename BaseResourceType, EShaderFrequency ShaderType, VkShaderStageFlagBits StageFlagBits>
+template<typename BaseResourceType, EShaderFrequency ShaderType>
 class TVulkanBaseShader : public BaseResourceType, public FVulkanShader
 {
 private:
 	TVulkanBaseShader(FVulkanDevice* InDevice) :
-		FVulkanShader(InDevice, ShaderType, StageFlagBits)
+		FVulkanShader(InDevice, ShaderType)
 	{
 	}
 	friend class FVulkanShaderFactory;
@@ -257,16 +255,16 @@ public:
 	}
 };
 
-typedef TVulkanBaseShader<FRHIVertexShader, SF_Vertex, VK_SHADER_STAGE_VERTEX_BIT>					FVulkanVertexShader;
-typedef TVulkanBaseShader<FRHIPixelShader, SF_Pixel, VK_SHADER_STAGE_FRAGMENT_BIT>					FVulkanPixelShader;
-typedef TVulkanBaseShader<FRHIComputeShader, SF_Compute, VK_SHADER_STAGE_COMPUTE_BIT>				FVulkanComputeShader;
-typedef TVulkanBaseShader<FRHIGeometryShader, SF_Geometry, VK_SHADER_STAGE_GEOMETRY_BIT>			FVulkanGeometryShader;
+typedef TVulkanBaseShader<FRHIVertexShader, SF_Vertex>				FVulkanVertexShader;
+typedef TVulkanBaseShader<FRHIPixelShader, SF_Pixel>				FVulkanPixelShader;
+typedef TVulkanBaseShader<FRHIComputeShader, SF_Compute>			FVulkanComputeShader;
+typedef TVulkanBaseShader<FRHIGeometryShader, SF_Geometry>			FVulkanGeometryShader;
 
 #if VULKAN_RHI_RAYTRACING
-typedef TVulkanBaseShader<FRHIRayGenShader, SF_RayGen, VK_SHADER_STAGE_RAYGEN_BIT_KHR>					FVulkanRayGenShader;
-typedef TVulkanBaseShader<FRHIRayMissShader, SF_RayMiss, VK_SHADER_STAGE_MISS_BIT_KHR>					FVulkanRayMissShader;
-typedef TVulkanBaseShader<FRHIRayCallableShader, SF_RayCallable, VK_SHADER_STAGE_CALLABLE_BIT_KHR>		FVulkanRayCallableShader;
-typedef TVulkanBaseShader<FRHIRayHitGroupShader, SF_RayHitGroup, VK_SHADER_STAGE_CLOSEST_HIT_BIT_KHR>	FVulkanRayHitGroupShader; // vkrt todo: How to handle VK_SHADER_STAGE_ANY_HIT_BIT_KHR?
+typedef TVulkanBaseShader<FRHIRayGenShader, SF_RayGen>				FVulkanRayGenShader;
+typedef TVulkanBaseShader<FRHIRayMissShader, SF_RayMiss>			FVulkanRayMissShader;
+typedef TVulkanBaseShader<FRHIRayCallableShader, SF_RayCallable>	FVulkanRayCallableShader;
+typedef TVulkanBaseShader<FRHIRayHitGroupShader, SF_RayHitGroup>	FVulkanRayHitGroupShader; // vkrt todo: How to handle VK_SHADER_STAGE_ANY_HIT_BIT_KHR?
 #endif // VULKAN_RHI_RAYTRACING
 
 class FVulkanShaderFactory
diff --git a/Engine/Source/Runtime/VulkanRHI/Private/VulkanShaders.cpp b/Engine/Source/Runtime/VulkanRHI/Private/VulkanShaders.cpp
--- a/Engine/Source/Runtime/VulkanRHI/Private/VulkanShaders.cpp
+++ b/Engine/Source/Runtime/VulkanRHI/Private/VulkanShaders.cpp
@@ -361,7 +361,7 @@ FVulkanShader::FSpirvCode FVulkanShader::GetPatchedSpirvCode(const FGfxPipelineD
 {
 	FSpirvCode Spirv = GetSpirvCode();
 
-	Layout->PatchSpirvBindings(Spirv, Frequency, CodeHeader, StageFlag);
+	Layout->PatchSpirvBindings(Spirv, Frequency, CodeHeader);
 	if (NeedsSpirvInputAttachmentPatching(Desc))
 	{
 		Spirv = PatchSpirvInputAttachments(Spirv);
@@ -375,7 +375,7 @@ TRefCountPtr<FVulkanShaderModule> FVulkanShader::CreateHandle(const FVulkanLayou
 	FScopeLock Lock(&VulkanShaderModulesMapCS);
 	FSpirvCode Spirv = GetSpirvCode();
 
-	Layout->PatchSpirvBindings(Spirv, Frequency, CodeHeader, StageFlag);
+	Layout->PatchSpirvBindings(Spirv, Frequency, CodeHeader);
 	TRefCountPtr<FVulkanShaderModule> Module = CreateShaderModule(Device, Spirv);
 	ShaderModules.Add(LayoutHash, Module);
 	if (!CodeHeader.DebugName.IsEmpty())
@@ -397,7 +397,7 @@ void FVulkanShader::PurgeShaderModules()
 	ShaderModules.Empty(0);
 }
 
-void FVulkanLayout::PatchSpirvBindings(FVulkanShader::FSpirvCode& SprivCode, EShaderFrequency Frequency, const FVulkanShaderHeader& CodeHeader, VkShaderStageFlagBits InStageFlag) const
+void FVulkanLayout::PatchSpirvBindings(FVulkanShader::FSpirvCode& SprivCode, EShaderFrequency Frequency, const FVulkanShaderHeader& CodeHeader) const
 {
 	// Bindless shader compilation already places descriptors and bindings in their fixed values based on stage and descriptor type
 	if (Device->SupportsBindless())
