VulkanRHI: Fixes for PSO compilation

Added VkPipelineShaderStageRequiredSubgroupSizeCreateInfo
Added VkAttachmentReferenceStencilLayout
Added more nullptr checks for pNext

diff --git a/Engine/Build/Android/Prebuilt/PSOService/Source/PSOService.cpp b/Engine/Build/Android/Prebuilt/PSOService/Source/PSOService.cpp
--- a/Engine/Build/Android/Prebuilt/PSOService/Source/PSOService.cpp
+++ b/Engine/Build/Android/Prebuilt/PSOService/Source/PSOService.cpp
@@ -316,12 +316,23 @@ public:
 		// VkPipelineShaderStageCreateInfo
 		for (int32_t Idx = 0; Idx < PipelineCreateInfo.StageCount; ++Idx)
 		{
+			bool bHasSubGroupSizeInfo = false;
+			COPY_FROM_BUFFER(&bHasSubGroupSizeInfo, PSO, MemoryOffset, sizeof(bool));
+
+			void* PipelineShaderStageCreatePNext = nullptr;
+			if (bHasSubGroupSizeInfo)
+			{
+				PipelineShaderStageCreatePNext = (void*)&PSO[MemoryOffset];
+				MemoryOffset += sizeof(VkPipelineShaderStageRequiredSubgroupSizeCreateInfo);
+			}
+
 			COPY_FROM_BUFFER(&ShaderStages[Idx], PSO, MemoryOffset, sizeof(VkPipelineShaderStageCreateInfo));
 
 			uint32_t NameLength;
 			COPY_FROM_BUFFER(&NameLength, PSO, MemoryOffset, sizeof(uint32_t));
 
 			ShaderStages[Idx].pName = (const char*)&PSO[MemoryOffset];
+			ShaderStages[Idx].pNext = PipelineShaderStageCreatePNext;
 			MemoryOffset += NameLength;
 		}
 		CreateInfo.pStages = ShaderStages;
@@ -508,6 +519,10 @@ public:
 
 			VkSubpassDescription2KHR* SubpassDescriptions = new VkSubpassDescription2KHR[RenderPassCreateInfo.subpassCount];
 			std::vector<VkFragmentShadingRateAttachmentInfoKHR> FSRAttachmentInfos;
+			std::vector<VkAttachmentReference2KHR> DepthStencilAttachments;
+
+			FSRAttachmentInfos.resize(RenderPassCreateInfo.subpassCount);
+			DepthStencilAttachments.resize(RenderPassCreateInfo.subpassCount);
 
 			for (uint32_t Idx = 0; Idx < RenderPassCreateInfo.subpassCount; ++Idx)
 			{
@@ -520,8 +535,8 @@ public:
 
 				if (bHasFSRAttachmentInfo)
 				{
-					FSRAttachmentInfos.push_back(VkFragmentShadingRateAttachmentInfoKHR());
-					auto& FSRAttachmentInfo = FSRAttachmentInfos.back();
+					FSRAttachmentInfos[Idx] = VkFragmentShadingRateAttachmentInfoKHR();
+					auto& FSRAttachmentInfo = FSRAttachmentInfos[Idx];
 					FSRAttachmentInfo.pNext = nullptr;
 					FSRAttachmentInfo.sType = VK_STRUCTURE_TYPE_FRAGMENT_SHADING_RATE_ATTACHMENT_INFO_KHR;
 					FSRAttachmentInfo.pFragmentShadingRateAttachment = (VkAttachmentReference2KHR*)&PSO[MemoryOffset];
@@ -562,8 +577,23 @@ public:
 
 				if (bHasDepthStencilAttachment)
 				{
-					SubpassDescriptions[Idx].pDepthStencilAttachment = (VkAttachmentReference2KHR*)&PSO[MemoryOffset];
+					bool bHasStencilLayout;
+					COPY_FROM_BUFFER(&bHasStencilLayout, PSO, MemoryOffset, sizeof(bool));
+
+					void* pDepthStencilAttachmentPNext = nullptr;
+					if(bHasStencilLayout)
+					{
+						pDepthStencilAttachmentPNext = (void*) & PSO[MemoryOffset];
+						MemoryOffset += sizeof(VkAttachmentReferenceStencilLayout);
+					}
+
+					DepthStencilAttachments[Idx] = *(VkAttachmentReference2KHR*)&PSO[MemoryOffset];
 					MemoryOffset += sizeof(VkAttachmentReference2KHR);
+
+					auto& DepthStencilAttachment = DepthStencilAttachments.back();
+					DepthStencilAttachments[Idx].pNext = pDepthStencilAttachmentPNext;
+
+					SubpassDescriptions[Idx].pDepthStencilAttachment = &DepthStencilAttachments[Idx];
 				}
 			}
 			RenderPassCreateInfo.pSubpasses = SubpassDescriptions;
diff --git a/Engine/Source/Runtime/VulkanRHI/Private/Android/VulkanAndroidPlatform.cpp b/Engine/Source/Runtime/VulkanRHI/Private/Android/VulkanAndroidPlatform.cpp
--- a/Engine/Source/Runtime/VulkanRHI/Private/Android/VulkanAndroidPlatform.cpp
+++ b/Engine/Source/Runtime/VulkanRHI/Private/Android/VulkanAndroidPlatform.cpp
@@ -787,23 +787,25 @@ void CharArrayToBuffer(const TArray<const ANSICHAR*>& CharArray, TArray<char>& M
 	}
 }
 
-void GetVKStructsFromPNext(const void* pNext, TMap<uint32_t, const void*>& VkStructs)
+void GetVKStructsFromPNext(const void* InNext, TMap<VkStructureType, const void*>& VkStructs, const TArray<VkStructureType>& ValidTypes)
 {
-	const void* PointerNext = pNext;
-
-	while (PointerNext)
+	const VkBaseInStructure* Next = reinterpret_cast<const VkBaseInStructure*>(InNext);
+	while (Next != nullptr)
 	{
-		const VkStructureType* NextType = (VkStructureType*)PointerNext;
-		VkStructs.FindOrAdd((uint32_t)*NextType) = PointerNext;
+		if (!ValidTypes.Contains(Next->sType))
+		{
+			UE_LOG(LogRHI, Warning, TEXT("GetVKStructsFromPNext: Unexpected type found when reading pNext->sType %d, Valid Types: "), (uint32_t)Next->sType);
+		}
 
-		PointerNext = *((void**)(((uint8*)NextType) + sizeof(size_t)));
+		VkStructs.FindOrAdd(Next->sType) = Next;
+		Next = reinterpret_cast<const VkBaseInStructure*>(Next->pNext);
 	}
 }
 
 void HandleGraphicsPipelineCreatePNext(const VkGraphicsPipelineCreateInfo* PipelineCreateInfo, TArray<char>& MemoryStream)
 {
-	TMap<uint32_t, const void*> VkStructs;
-	GetVKStructsFromPNext(PipelineCreateInfo->pNext, VkStructs);
+	TMap<VkStructureType, const void*> VkStructs;
+	GetVKStructsFromPNext(PipelineCreateInfo->pNext, VkStructs, { VK_STRUCTURE_TYPE_PIPELINE_FRAGMENT_SHADING_RATE_STATE_CREATE_INFO_KHR });
 
 	int32_t HandledCount = 0;
 
@@ -813,16 +815,17 @@ void HandleGraphicsPipelineCreatePNext(const VkGraphicsPipelineCreateInfo* Pipel
 #if !VULKAN_SUPPORTS_FRAGMENT_SHADING_RATE
 	COPY_TO_BUFFER(MemoryStream, &bHasFSRCreateInfo, sizeof(bool));
 #else
-	const void** Struct = VkStructs.Find((uint32_t)VK_STRUCTURE_TYPE_PIPELINE_FRAGMENT_SHADING_RATE_STATE_CREATE_INFO_KHR);
+	const void** Struct = VkStructs.Find(VK_STRUCTURE_TYPE_PIPELINE_FRAGMENT_SHADING_RATE_STATE_CREATE_INFO_KHR);
 	bHasFSRCreateInfo = Struct != nullptr;
 	COPY_TO_BUFFER(MemoryStream, &bHasFSRCreateInfo, sizeof(bool));
 
 	if (bHasFSRCreateInfo)
 	{
-		VkPipelineFragmentShadingRateStateCreateInfoKHR FSRCreateInfo = *(VkPipelineFragmentShadingRateStateCreateInfoKHR*)*Struct;
-		FSRCreateInfo.pNext = nullptr;
+		VkPipelineFragmentShadingRateStateCreateInfoKHR* FSRCreateInfo = (VkPipelineFragmentShadingRateStateCreateInfoKHR*)*Struct;
+		check(FSRCreateInfo->pNext == nullptr);
+		FSRCreateInfo->pNext = nullptr;
 
-		COPY_TO_BUFFER(MemoryStream, &FSRCreateInfo, sizeof(VkPipelineFragmentShadingRateStateCreateInfoKHR));
+		COPY_TO_BUFFER(MemoryStream, FSRCreateInfo, sizeof(VkPipelineFragmentShadingRateStateCreateInfoKHR));
 		HandledCount++;
 	} 
 #endif
@@ -830,10 +833,36 @@ void HandleGraphicsPipelineCreatePNext(const VkGraphicsPipelineCreateInfo* Pipel
 	check(HandledCount == VkStructs.Num());
 }
 
+void HandlePipelineShaderStagePNext(const VkPipelineShaderStageCreateInfo* CreateInfo, TArray<char>& MemoryStream)
+{
+	TMap<VkStructureType, const void*> VkStructs;
+	GetVKStructsFromPNext(CreateInfo->pNext, VkStructs, { VK_STRUCTURE_TYPE_PIPELINE_SHADER_STAGE_REQUIRED_SUBGROUP_SIZE_CREATE_INFO });
+
+	int32_t HandledCount = 0;
+
+	// Subgroup Size Info
+	bool bHasSubGroupSizeInfo = false;
+
+	const void** Struct = VkStructs.Find(VK_STRUCTURE_TYPE_PIPELINE_SHADER_STAGE_REQUIRED_SUBGROUP_SIZE_CREATE_INFO);
+	bHasSubGroupSizeInfo = Struct != nullptr;
+	COPY_TO_BUFFER(MemoryStream, &bHasSubGroupSizeInfo, sizeof(bool));
+
+	if (bHasSubGroupSizeInfo)
+	{
+		VkPipelineShaderStageRequiredSubgroupSizeCreateInfo* SubgroupSizeCreateInfo = (VkPipelineShaderStageRequiredSubgroupSizeCreateInfo*)*Struct;
+		check(SubgroupSizeCreateInfo->pNext == nullptr);
+
+		COPY_TO_BUFFER(MemoryStream, SubgroupSizeCreateInfo, sizeof(VkPipelineShaderStageRequiredSubgroupSizeCreateInfo));
+		HandledCount++;
+	}
+
+	check(HandledCount == VkStructs.Num());
+}
+
 void HandleSubpassDescriptionPNext(VkSubpassDescription2* SubpassDescription, TArray<char>& MemoryStream)
 {
-	TMap<uint32_t, const void*> VkStructs;
-	GetVKStructsFromPNext(SubpassDescription->pNext, VkStructs);
+	TMap<VkStructureType, const void*> VkStructs;
+	GetVKStructsFromPNext(SubpassDescription->pNext, VkStructs, { VK_STRUCTURE_TYPE_FRAGMENT_SHADING_RATE_ATTACHMENT_INFO_KHR });
 
 	int32_t HandledCount = 0;
 
@@ -843,13 +872,15 @@ void HandleSubpassDescriptionPNext(VkSubpassDescription2* SubpassDescription, TA
 #if !VULKAN_SUPPORTS_FRAGMENT_SHADING_RATE
 	COPY_TO_BUFFER(MemoryStream, &bHasFSRAttachmentCreateInfo, sizeof(bool));
 #else
-	const void** Struct = VkStructs.Find((uint32_t)VK_STRUCTURE_TYPE_FRAGMENT_SHADING_RATE_ATTACHMENT_INFO_KHR);
+	const void** Struct = VkStructs.Find(VK_STRUCTURE_TYPE_FRAGMENT_SHADING_RATE_ATTACHMENT_INFO_KHR);
 	bHasFSRAttachmentCreateInfo = Struct != nullptr;
 	COPY_TO_BUFFER(MemoryStream, &bHasFSRAttachmentCreateInfo, sizeof(bool));
 
 	if (bHasFSRAttachmentCreateInfo)
 	{
 		VkFragmentShadingRateAttachmentInfoKHR* FragmentShadingRateCreateInfo = (VkFragmentShadingRateAttachmentInfoKHR*)*Struct;
+
+		check(FragmentShadingRateCreateInfo->pFragmentShadingRateAttachment->pNext == nullptr);
 		COPY_TO_BUFFER(MemoryStream, FragmentShadingRateCreateInfo->pFragmentShadingRateAttachment, sizeof(VkAttachmentReference2));
 		COPY_TO_BUFFER(MemoryStream, &FragmentShadingRateCreateInfo->shadingRateAttachmentTexelSize, sizeof(VkExtent2D));
 		HandledCount++;
@@ -859,6 +890,30 @@ void HandleSubpassDescriptionPNext(VkSubpassDescription2* SubpassDescription, TA
 	check(HandledCount == VkStructs.Num());
 }
 
+void HandleDepthStencilAttachmentPNext(const VkAttachmentReference2* Attachment, TArray<char>& MemoryStream)
+{
+	TMap<VkStructureType, const void*> VkStructs;
+	GetVKStructsFromPNext(Attachment->pNext, VkStructs, { VK_STRUCTURE_TYPE_ATTACHMENT_REFERENCE_STENCIL_LAYOUT });
+
+	int32_t HandledCount = 0;
+
+	bool bHasStencilLayout = false;
+	
+	const void** Struct = VkStructs.Find(VK_STRUCTURE_TYPE_ATTACHMENT_REFERENCE_STENCIL_LAYOUT);
+	bHasStencilLayout = Struct != nullptr;
+	COPY_TO_BUFFER(MemoryStream, &bHasStencilLayout, sizeof(bool));
+
+	if (bHasStencilLayout)
+	{
+		VkAttachmentReferenceStencilLayout* StencilLayout = (VkAttachmentReferenceStencilLayout*)*Struct;
+		check(StencilLayout->pNext == nullptr);
+		COPY_TO_BUFFER(MemoryStream, StencilLayout, sizeof(VkAttachmentReferenceStencilLayout));
+		HandledCount++;
+	}
+
+	check(HandledCount == VkStructs.Num());
+}
+
 void PipelineToBinary(FVulkanDevice* Device, const VkGraphicsPipelineCreateInfo* PipelineInfo, FGfxPipelineDesc* GfxEntry, const FVulkanRenderTargetLayout* RTLayout, TArray<char>& MemoryStream)
 {
 	static const unsigned int INITIAL_PSO_STREAM_SIZE = 64 * 1024;
@@ -901,6 +956,8 @@ void PipelineToBinary(FVulkanDevice* Device, const VkGraphicsPipelineCreateInfo*
 		ShaderStage.flags = PipelineInfo->pStages[Idx].flags;
 		ShaderStage.stage = PipelineInfo->pStages[Idx].stage;
 
+		HandlePipelineShaderStagePNext(&PipelineInfo->pStages[Idx], MemoryStream);
+
 		COPY_TO_BUFFER(MemoryStream, &ShaderStage, sizeof(VkPipelineShaderStageCreateInfo));
 
 		uint32_t NameLength = static_cast<uint32_t>(strlen(PipelineInfo->pStages[Idx].pName)+1);
@@ -1099,6 +1156,8 @@ void PipelineToBinary(FVulkanDevice* Device, const VkGraphicsPipelineCreateInfo*
 
 			VkRenderPassFragmentDensityMapCreateInfoEXT* FragmentDensityMap = (VkRenderPassFragmentDensityMapCreateInfoEXT*)RenderPassCreateInfo.pNext;
 			COPY_TO_BUFFER(MemoryStream, FragmentDensityMap, sizeof(VkRenderPassFragmentDensityMapCreateInfoEXT));
+
+			check(FragmentDensityMap->pNext == nullptr);
 		}
 
 		if(RenderPassCreateInfo.attachmentCount > 0)
@@ -1120,17 +1179,25 @@ void PipelineToBinary(FVulkanDevice* Device, const VkGraphicsPipelineCreateInfo*
 			SubpassDescription.pInputAttachments = nullptr;
 			SubpassDescription.pResolveAttachments = nullptr;
 			
-			COPY_TO_BUFFER(MemoryStream, &CreateInfo.pSubpasses[Idx], sizeof(VkSubpassDescription2));
+			COPY_TO_BUFFER(MemoryStream, &SubpassDescription, sizeof(VkSubpassDescription2));
 
 			HandleSubpassDescriptionPNext(&SubpassDescription, MemoryStream);
 
 			if(SubpassDescription.colorAttachmentCount > 0)
 			{
+				for (uint32_t n = 0; n < SubpassDescription.colorAttachmentCount; ++n)
+				{
+					check(CreateInfo.pSubpasses[Idx].pColorAttachments[n].pNext == nullptr);
+				}
 				COPY_TO_BUFFER(MemoryStream, CreateInfo.pSubpasses[Idx].pColorAttachments, sizeof(VkAttachmentReference2) * SubpassDescription.colorAttachmentCount);
 			}
 
 			if(SubpassDescription.inputAttachmentCount > 0)
 			{
+				for (uint32_t n = 0; n < SubpassDescription.inputAttachmentCount; ++n)
+				{
+					check(CreateInfo.pSubpasses[Idx].pInputAttachments[n].pNext == nullptr);
+				}
 				COPY_TO_BUFFER(MemoryStream, CreateInfo.pSubpasses[Idx].pInputAttachments, sizeof(VkAttachmentReference2) * SubpassDescription.inputAttachmentCount);
 			}
 
@@ -1141,6 +1208,11 @@ void PipelineToBinary(FVulkanDevice* Device, const VkGraphicsPipelineCreateInfo*
 			{
 				if(SubpassDescription.colorAttachmentCount > 0)
 				{
+					for (uint32_t n = 0; n < SubpassDescription.colorAttachmentCount; ++n)
+					{
+						check(CreateInfo.pSubpasses[Idx].pResolveAttachments[n].pNext == nullptr);
+					}
+					check(CreateInfo.pSubpasses[Idx].pResolveAttachments == nullptr);
 					COPY_TO_BUFFER(MemoryStream, CreateInfo.pSubpasses[Idx].pResolveAttachments, sizeof(VkAttachmentReference2)* SubpassDescription.colorAttachmentCount);
 				}
 			}
@@ -1150,6 +1222,7 @@ void PipelineToBinary(FVulkanDevice* Device, const VkGraphicsPipelineCreateInfo*
 
 			if (bHasDepthStencilAttachment)
 			{
+				HandleDepthStencilAttachmentPNext(CreateInfo.pSubpasses[Idx].pDepthStencilAttachment, MemoryStream);
 				COPY_TO_BUFFER(MemoryStream, CreateInfo.pSubpasses[Idx].pDepthStencilAttachment, sizeof(VkAttachmentReference2));
 			}
 		}
@@ -1192,7 +1265,7 @@ void PipelineToBinary(FVulkanDevice* Device, const VkGraphicsPipelineCreateInfo*
 
 			check(FragmentDensityMap->pNext == nullptr);
 			// TODO: Support Multiview create info
-		}
+		} 
 
 		if(RenderPassCreateInfo.attachmentCount > 0)
 		{
@@ -1213,7 +1286,7 @@ void PipelineToBinary(FVulkanDevice* Device, const VkGraphicsPipelineCreateInfo*
 			SubpassDescription.pInputAttachments = nullptr;
 			SubpassDescription.pResolveAttachments = nullptr;
 			
-			COPY_TO_BUFFER(MemoryStream, &CreateInfo.pSubpasses[Idx], sizeof(VkSubpassDescription));
+			COPY_TO_BUFFER(MemoryStream, &SubpassDescription, sizeof(VkSubpassDescription));
 
 			if(SubpassDescription.colorAttachmentCount > 0)
 			{
