VulkanRHI: VKRT WIP: Add VkDeviceAddress to FVulkanResourceMultiBuffer allocations (move accessor out of raytracing specific code).  Add accessor for FVulkanUniformBuffer's buffer handle.

diff --git a/Engine/Source/Runtime/VulkanRHI/Private/VulkanBuffer.cpp b/Engine/Source/Runtime/VulkanRHI/Private/VulkanBuffer.cpp
--- a/Engine/Source/Runtime/VulkanRHI/Private/VulkanBuffer.cpp
+++ b/Engine/Source/Runtime/VulkanRHI/Private/VulkanBuffer.cpp
@@ -82,6 +82,17 @@ static FORCEINLINE void UpdateVulkanBufferStats(uint64_t Size, VkBufferUsageFlag
 	}
 }
 
+static VkDeviceAddress GetBufferDeviceAddress(FVulkanDevice* Device, VkBuffer Buffer)
+{
+	if (Device->GetOptionalExtensions().HasBufferDeviceAddress)
+	{
+		VkBufferDeviceAddressInfoKHR DeviceAddressInfo;
+		ZeroVulkanStruct(DeviceAddressInfo, VK_STRUCTURE_TYPE_BUFFER_DEVICE_ADDRESS_INFO);
+		DeviceAddressInfo.buffer = Buffer;
+		return VulkanRHI::vkGetBufferDeviceAddressKHR(Device->GetInstanceHandle(), &DeviceAddressInfo);
+	}
+	return 0;
+}
 
 VkBufferUsageFlags FVulkanResourceMultiBuffer::UEToVKBufferUsageFlags(FVulkanDevice* InDevice, EBufferUsageFlags InUEUsage, bool bZeroSize)
 {
@@ -175,6 +186,7 @@ FVulkanResourceMultiBuffer::FVulkanResourceMultiBuffer(FVulkanDevice* InDevice,
 				FBufferAlloc NewBufferAlloc;
 				NewBufferAlloc.Alloc = FVulkanTransientHeap::GetVulkanAllocation(*InTransientHeapAllocation);
 				NewBufferAlloc.HostPtr = bUnifiedMem ? NewBufferAlloc.Alloc.GetMappedPointer(Device) : nullptr;
+				NewBufferAlloc.DeviceAddress = GetBufferDeviceAddress(InDevice, NewBufferAlloc.Alloc.GetBufferHandle()) + NewBufferAlloc.Alloc.Offset;
 				check(NewBufferAlloc.Alloc.Offset % BufferAlignment == 0);
 				check(NewBufferAlloc.Alloc.Size >= InBufferDesc.Size);
 				CurrentBufferIndex = BufferAllocs.Add(NewBufferAlloc);
@@ -184,7 +196,7 @@ FVulkanResourceMultiBuffer::FVulkanResourceMultiBuffer(FVulkanDevice* InDevice,
 				AdvanceBufferIndex();
 			}
 
-			VULKAN_SET_DEBUG_NAME((*InDevice), VK_OBJECT_TYPE_BUFFER, (VkBuffer)BufferAllocs[CurrentBufferIndex].Alloc.VulkanHandle, TEXT("%s"), CreateInfo.DebugName ? CreateInfo.DebugName : TEXT("UnknownBuffer"));
+			VULKAN_SET_DEBUG_NAME((*InDevice), VK_OBJECT_TYPE_BUFFER, BufferAllocs[CurrentBufferIndex].Alloc.GetBufferHandle(), TEXT("%s"), CreateInfo.DebugName ? CreateInfo.DebugName : TEXT("UnknownBuffer"));
 
 			if (CreateInfo.ResourceArray)
 			{
@@ -276,6 +288,7 @@ void FVulkanResourceMultiBuffer::AdvanceBufferIndex()
 		NewBufferAlloc.HostPtr = bUnifiedMem ? NewBufferAlloc.Alloc.GetMappedPointer(Device) : nullptr;
 		NewBufferAlloc.Fence = new FVulkanGPUFence(TEXT("VulkanDynamicBuffer"));
 		NewBufferAlloc.AllocStatus = FBufferAlloc::EAllocStatus::InUse;
+		NewBufferAlloc.DeviceAddress = GetBufferDeviceAddress(Device, NewBufferAlloc.Alloc.GetBufferHandle()) + NewBufferAlloc.Alloc.Offset;
 	}
 }
 
@@ -344,6 +357,8 @@ void* FVulkanResourceMultiBuffer::Lock(FVulkanCommandListContext& Context, EReso
 			// Patch our alloc to go directly to our offset
 			BufferAlloc.Alloc.Offset += VolatileAlloc.CurrentOffset;
 			BufferAlloc.Alloc.Size = VolatileAlloc.Size;
+
+			BufferAlloc.DeviceAddress = GetBufferDeviceAddress(Device, BufferAlloc.Alloc.GetBufferHandle()) + BufferAlloc.Alloc.Offset;
 		}
 	}
 	else
diff --git a/Engine/Source/Runtime/VulkanRHI/Private/VulkanRayTracing.cpp b/Engine/Source/Runtime/VulkanRHI/Private/VulkanRayTracing.cpp
--- a/Engine/Source/Runtime/VulkanRHI/Private/VulkanRayTracing.cpp
+++ b/Engine/Source/Runtime/VulkanRHI/Private/VulkanRayTracing.cpp
@@ -87,19 +87,6 @@ enum class EBLASBuildDataUsage
 	Size = 1
 };
 
-static VkDeviceAddress GetDeviceAddress(VkDevice Device, VkBuffer Buffer)
-{
-	VkBufferDeviceAddressInfoKHR DeviceAddressInfo;
-	ZeroVulkanStruct(DeviceAddressInfo, VK_STRUCTURE_TYPE_BUFFER_DEVICE_ADDRESS_INFO);
-	DeviceAddressInfo.buffer = Buffer;
-	return VulkanDynamicAPI::vkGetBufferDeviceAddressKHR(Device, &DeviceAddressInfo);
-}
-
-VkDeviceAddress FVulkanResourceMultiBuffer::GetDeviceAddress() const
-{
-	return ::GetDeviceAddress(Device->GetInstanceHandle(), GetHandle()) + GetOffset();
-}
-
 // Temporary brute force allocation helper, this should be handled by the memory sub-allocator
 static uint32 FindMemoryType(VkPhysicalDevice Gpu, uint32 Filter, VkMemoryPropertyFlags RequestedProperties)
 {
diff --git a/Engine/Source/Runtime/VulkanRHI/Private/VulkanResources.h b/Engine/Source/Runtime/VulkanRHI/Private/VulkanResources.h
--- a/Engine/Source/Runtime/VulkanRHI/Private/VulkanResources.h
+++ b/Engine/Source/Runtime/VulkanRHI/Private/VulkanResources.h
@@ -1047,6 +1047,11 @@ public:
 		return GetCurrentAllocation().Size;
 	}
 
+	inline VkDeviceAddress GetDeviceAddress() const
+	{
+		return BufferAllocs[CurrentBufferIndex].DeviceAddress;
+	}
+
 	inline VkBufferUsageFlags GetBufferUsageFlags() const
 	{
 		return BufferUsageFlags;
@@ -1078,10 +1083,6 @@ public:
 		Callback(TEXT("FVulkanResourceMultiBuffer"), FName(), this, 0, GetCurrentSize() * GetNumBuffers(), 1, 1, VK_FORMAT_UNDEFINED);
 	}
 
-#if VULKAN_RHI_RAYTRACING
-	VkDeviceAddress GetDeviceAddress() const;
-#endif
-
 	static VkBufferUsageFlags UEToVKBufferUsageFlags(FVulkanDevice* InDevice, EBufferUsageFlags InUEUsage, bool bZeroSize);
 
 protected:
@@ -1105,6 +1106,7 @@ protected:
 		VulkanRHI::FVulkanAllocation Alloc;
 		void* HostPtr = nullptr;
 		class FVulkanGPUFence* Fence = nullptr;
+		VkDeviceAddress DeviceAddress = 0;
 
 		enum class EAllocStatus : uint8
 		{
@@ -1135,6 +1137,11 @@ public:
 	void UpdateResourceTable(const FRHIUniformBufferLayout& InLayout, const void* Contents, int32 ResourceNum);
 	void UpdateResourceTable(FRHIResource** Resources, int32 ResourceNum);
 
+	inline VkBuffer GetBufferHandle() const
+	{
+		return Allocation.GetBufferHandle();
+	}
+
 	inline uint32 GetOffset() const
 	{
 		return Allocation.Offset;
