VulkanRHI: VKRT WIP: Use a single FVulkanRayTracingShader class for raytracing shaders to bring Vulkan closer to D3D.  Use ResourceCast when FRHIRayTracingShader is passed in.  Add checks to ensure we have the proper RT shader type.

#rnx

diff --git a/Engine/Source/Runtime/VulkanRHI/Private/VulkanRayTracing.cpp b/Engine/Source/Runtime/VulkanRHI/Private/VulkanRayTracing.cpp
--- a/Engine/Source/Runtime/VulkanRHI/Private/VulkanRayTracing.cpp
+++ b/Engine/Source/Runtime/VulkanRHI/Private/VulkanRayTracing.cpp
@@ -1052,19 +1052,19 @@ FVulkanRayTracingPipelineState::FVulkanRayTracingPipelineState(FVulkanDevice* co
 	
 	for (FRHIRayTracingShader* RayGenShader : InitializerRayGenShaders)
 	{
-		const FVulkanShaderHeader& Header = static_cast<FVulkanRayGenShader*>(RayGenShader)->GetCodeHeader();
+		const FVulkanShaderHeader& Header = ResourceCast(RayGenShader)->GetCodeHeader();
 		DescriptorSetLayoutInfo.ProcessBindingsForStage(VK_SHADER_STAGE_RAYGEN_BIT_KHR, ShaderStage::RayGen, Header, UBGatherInfo);
 	}
 
 	for (FRHIRayTracingShader* MissShader : InitializerMissShaders)
 	{
-		const FVulkanShaderHeader& Header = static_cast<FVulkanRayMissShader*>(MissShader)->GetCodeHeader();
+		const FVulkanShaderHeader& Header = ResourceCast(MissShader)->GetCodeHeader();
 		DescriptorSetLayoutInfo.ProcessBindingsForStage(VK_SHADER_STAGE_MISS_BIT_KHR, ShaderStage::RayMiss, Header, UBGatherInfo);
 	}
 
 	for (FRHIRayTracingShader* HitGroupShader : InitializerHitGroupShaders)
 	{
-		const FVulkanShaderHeader& Header = static_cast<FVulkanRayHitGroupShader*>(HitGroupShader)->GetCodeHeader();
+		const FVulkanShaderHeader& Header = ResourceCast(HitGroupShader)->GetCodeHeader();
 		DescriptorSetLayoutInfo.ProcessBindingsForStage(VK_SHADER_STAGE_CLOSEST_HIT_BIT_KHR, ShaderStage::RayHitGroup, Header, UBGatherInfo);
 		// vkrt todo: How to handle any hit for hit group?
 	}
@@ -1083,13 +1083,16 @@ FVulkanRayTracingPipelineState::FVulkanRayTracingPipelineState(FVulkanDevice* co
 
 	for (FRHIRayTracingShader* const RayGenShaderRHI : InitializerRayGenShaders)
 	{
+		checkSlow(RayGenShaderRHI->GetFrequency() == SF_RayGen);
+		FVulkanRayTracingShader* const RayGenShader = ResourceCast(RayGenShaderRHI);
+
 		VkPipelineShaderStageCreateInfo ShaderStage;
 		ZeroVulkanStruct(ShaderStage, VK_STRUCTURE_TYPE_PIPELINE_SHADER_STAGE_CREATE_INFO);
-		ShaderStage.module = static_cast<FVulkanRayGenShader*>(RayGenShaderRHI)->GetOrCreateHandle(Layout, Layout->GetDescriptorSetLayoutHash())->GetVkShaderModule();
+		ShaderStage.module = RayGenShader->GetOrCreateHandle()->GetVkShaderModule();
 		ShaderStage.stage = VK_SHADER_STAGE_RAYGEN_BIT_KHR;
 			
 		ANSICHAR* const EntryPoint = new ANSICHAR[EntryPointNameMaxLength];
-		static_cast<FVulkanRayGenShader*>(RayGenShaderRHI)->GetEntryPoint(EntryPoint, EntryPointNameMaxLength);
+		RayGenShader->GetEntryPoint(EntryPoint, EntryPointNameMaxLength);
 		EntryPointNames.Add(EntryPoint);
 		ShaderStage.pName = EntryPoint;
 		ShaderStages.Add(ShaderStage);
@@ -1106,13 +1109,16 @@ FVulkanRayTracingPipelineState::FVulkanRayTracingPipelineState(FVulkanDevice* co
 
 	for (FRHIRayTracingShader* const MissShaderRHI : InitializerMissShaders)
 	{
+		checkSlow(MissShaderRHI->GetFrequency() == SF_RayMiss);
+		FVulkanRayTracingShader* const MissShader = ResourceCast(MissShaderRHI);
+
 		VkPipelineShaderStageCreateInfo ShaderStage;
 		ZeroVulkanStruct(ShaderStage, VK_STRUCTURE_TYPE_PIPELINE_SHADER_STAGE_CREATE_INFO);
-		ShaderStage.module = static_cast<FVulkanRayMissShader*>(MissShaderRHI)->GetOrCreateHandle(Layout, Layout->GetDescriptorSetLayoutHash())->GetVkShaderModule();
+		ShaderStage.module = MissShader->GetOrCreateHandle()->GetVkShaderModule();
 		ShaderStage.stage = VK_SHADER_STAGE_MISS_BIT_KHR;
 
 		ANSICHAR* const EntryPoint = new char[EntryPointNameMaxLength];
-		static_cast<FVulkanRayGenShader*>(MissShaderRHI)->GetEntryPoint(EntryPoint, EntryPointNameMaxLength);
+		MissShader->GetEntryPoint(EntryPoint, EntryPointNameMaxLength);
 		EntryPointNames.Add(EntryPoint);
 		ShaderStage.pName = EntryPoint;
 		ShaderStages.Add(ShaderStage);
@@ -1129,13 +1135,16 @@ FVulkanRayTracingPipelineState::FVulkanRayTracingPipelineState(FVulkanDevice* co
 
 	for (FRHIRayTracingShader* const HitGroupShaderRHI : InitializerHitGroupShaders)
 	{
+		checkSlow(HitGroupShaderRHI->GetFrequency() == SF_RayHitGroup);
+		FVulkanRayTracingShader* const HitGroupShader = ResourceCast(HitGroupShaderRHI);
+
 		VkPipelineShaderStageCreateInfo ShaderStage;
 		ZeroVulkanStruct(ShaderStage, VK_STRUCTURE_TYPE_PIPELINE_SHADER_STAGE_CREATE_INFO);
-		ShaderStage.module = static_cast<FVulkanRayHitGroupShader*>(HitGroupShaderRHI)->GetOrCreateHandle(Layout, Layout->GetDescriptorSetLayoutHash())->GetVkShaderModule();
+		ShaderStage.module = HitGroupShader->GetOrCreateHandle()->GetVkShaderModule();
 		ShaderStage.stage = VK_SHADER_STAGE_CLOSEST_HIT_BIT_KHR;
 
 		ANSICHAR* const EntryPoint = new char[EntryPointNameMaxLength];
-		static_cast<FVulkanRayHitGroupShader*>(HitGroupShaderRHI)->GetEntryPoint(EntryPoint, EntryPointNameMaxLength);
+		HitGroupShader->GetEntryPoint(EntryPoint, EntryPointNameMaxLength);
 		EntryPointNames.Add(EntryPoint);
 		ShaderStage.pName = EntryPoint;
 		ShaderStages.Add(ShaderStage);
diff --git a/Engine/Source/Runtime/VulkanRHI/Private/VulkanResources.h b/Engine/Source/Runtime/VulkanRHI/Private/VulkanResources.h
--- a/Engine/Source/Runtime/VulkanRHI/Private/VulkanResources.h
+++ b/Engine/Source/Runtime/VulkanRHI/Private/VulkanResources.h
@@ -113,6 +113,8 @@ public:
 
 	void PurgeShaderModules();
 
+	TRefCountPtr<FVulkanShaderModule> GetOrCreateHandle();
+
 	TRefCountPtr<FVulkanShaderModule> GetOrCreateHandle(const FVulkanLayout* Layout, uint32 LayoutHash)
 	{
 		FScopeLock Lock(&VulkanShaderModulesMapCS);
@@ -261,10 +263,31 @@ typedef TVulkanBaseShader<FRHIComputeShader, SF_Compute>			FVulkanComputeShader;
 typedef TVulkanBaseShader<FRHIGeometryShader, SF_Geometry>			FVulkanGeometryShader;
 
 #if VULKAN_RHI_RAYTRACING
-typedef TVulkanBaseShader<FRHIRayGenShader, SF_RayGen>				FVulkanRayGenShader;
-typedef TVulkanBaseShader<FRHIRayMissShader, SF_RayMiss>			FVulkanRayMissShader;
-typedef TVulkanBaseShader<FRHIRayCallableShader, SF_RayCallable>	FVulkanRayCallableShader;
-typedef TVulkanBaseShader<FRHIRayHitGroupShader, SF_RayHitGroup>	FVulkanRayHitGroupShader; // vkrt todo: How to handle VK_SHADER_STAGE_ANY_HIT_BIT_KHR?
+class FVulkanRayTracingShader : public FRHIRayTracingShader, public FVulkanShader
+{
+private:
+	FVulkanRayTracingShader(FVulkanDevice* InDevice, EShaderFrequency InFrequency)
+		: FRHIRayTracingShader(InFrequency)
+		, FVulkanShader(InDevice, InFrequency)
+	{
+	}
+	friend class FVulkanShaderFactory;
+
+public:
+	// IRefCountedObject interface.
+	virtual uint32 AddRef() const override final
+	{
+		return FRHIResource::AddRef();
+	}
+	virtual uint32 Release() const override final
+	{
+		return FRHIResource::Release();
+	}
+	virtual uint32 GetRefCount() const override final
+	{
+		return FRHIResource::GetRefCount();
+	}
+};
 #endif // VULKAN_RHI_RAYTRACING
 
 class FVulkanShaderFactory
@@ -290,6 +313,9 @@ public:
 		return nullptr;
 	}
 
+	template <EShaderFrequency ShaderFrequency>
+	FVulkanRayTracingShader* CreateRayTracingShader(TArrayView<const uint8> Code, FVulkanDevice* Device);
+
 	void LookupShaders(const uint64 InShaderKeys[ShaderStage::NumStages], FVulkanShader* OutShaders[ShaderStage::NumStages]) const;
 
 	void OnDeleteShader(const FVulkanShader& Shader);
@@ -1459,6 +1485,8 @@ static FORCEINLINE FVulkanTexture* ResourceCast(FRHITexture* Texture)
 
 #if VULKAN_RHI_RAYTRACING
 class FVulkanRayTracingScene;
+class FVulkanRayTracingGeometry;
+class FVulkanRayTracingPipelineState;
 template<>
 struct TVulkanResourceTraits<FRHIRayTracingScene>
 {
@@ -1470,4 +1498,14 @@ struct TVulkanResourceTraits<FRHIRayTracingGeometry>
 {
 	typedef FVulkanRayTracingGeometry TConcreteType;
 };
+template<>
+struct TVulkanResourceTraits<FRHIRayTracingPipelineState>
+{
+	typedef FVulkanRayTracingPipelineState TConcreteType;
+};
+template<>
+struct TVulkanResourceTraits<FRHIRayTracingShader>
+{
+	typedef FVulkanRayTracingShader TConcreteType;
+};
 #endif // VULKAN_RHI_RAYTRACING
diff --git a/Engine/Source/Runtime/VulkanRHI/Private/VulkanShaders.cpp b/Engine/Source/Runtime/VulkanRHI/Private/VulkanShaders.cpp
--- a/Engine/Source/Runtime/VulkanRHI/Private/VulkanShaders.cpp
+++ b/Engine/Source/Runtime/VulkanRHI/Private/VulkanShaders.cpp
@@ -52,6 +52,8 @@ FVulkanShaderFactory::~FVulkanShaderFactory()
 template <typename ShaderType> 
 ShaderType* FVulkanShaderFactory::CreateShader(TArrayView<const uint8> Code, FVulkanDevice* Device)
 {
+	static_assert(ShaderType::StaticFrequency != SF_RayCallable && ShaderType::StaticFrequency != SF_RayGen && ShaderType::StaticFrequency != SF_RayHitGroup && ShaderType::StaticFrequency != SF_RayMiss);
+
 	const uint32 ShaderCodeLen = Code.Num();
 	const uint32 ShaderCodeCRC = FCrc::MemCrc32(Code.GetData(), Code.Num());
 	const uint64 ShaderKey = ((uint64)ShaderCodeLen | ((uint64)ShaderCodeCRC << 32));
@@ -80,16 +82,6 @@ ShaderType* FVulkanShaderFactory::CreateShader(TArrayView<const uint8> Code, FVu
 			{
 				RetShader = new ShaderType(Device);
 				RetShader->Setup(MoveTemp(CodeHeader), MoveTemp(SerializedSRT), MoveTemp(SpirvContainer), ShaderKey);
-				if constexpr (
-					ShaderType::StaticFrequency == SF_RayCallable ||
-					ShaderType::StaticFrequency == SF_RayGen ||
-					ShaderType::StaticFrequency == SF_RayHitGroup ||
-					ShaderType::StaticFrequency == SF_RayMiss)
-				{
-					RetShader->RayTracingPayloadType = RetShader->CodeHeader.RayTracingPayloadType;
-					RetShader->RayTracingPayloadSize = RetShader->CodeHeader.RayTracingPayloadSize;
-				}
-
 				ShaderMap[ShaderType::StaticFrequency].Add(ShaderKey, RetShader);
 
 #if !(UE_BUILD_SHIPPING || UE_BUILD_TEST)
@@ -103,6 +95,70 @@ ShaderType* FVulkanShaderFactory::CreateShader(TArrayView<const uint8> Code, FVu
 	return RetShader;
 }
 
+template <EShaderFrequency ShaderFrequency>
+FVulkanRayTracingShader* FVulkanShaderFactory::CreateRayTracingShader(TArrayView<const uint8> Code, FVulkanDevice* Device)
+{
+	static_assert(ShaderFrequency == SF_RayCallable || ShaderFrequency == SF_RayGen || ShaderFrequency == SF_RayHitGroup || ShaderFrequency == SF_RayMiss);
+
+	auto LookupRayTracingShader = [this](uint64 ShaderKey)
+	{
+		FVulkanRayTracingShader* RTShader = nullptr;
+		if (ShaderKey)
+		{
+			FRWScopeLock ScopedLock(RWLock[ShaderFrequency], SLT_ReadOnly);
+			FVulkanShader* const* FoundShaderPtr = ShaderMap[ShaderFrequency].Find(ShaderKey);
+			if (FoundShaderPtr)
+			{
+				RTShader = static_cast<FVulkanRayTracingShader*>(*FoundShaderPtr);
+			}
+		}
+		return RTShader;
+	};
+
+	const uint32 ShaderCodeLen = Code.Num();
+	const uint32 ShaderCodeCRC = FCrc::MemCrc32(Code.GetData(), Code.Num());
+	const uint64 ShaderKey = ((uint64)ShaderCodeLen | ((uint64)ShaderCodeCRC << 32));
+
+	FVulkanRayTracingShader* RetShader = LookupRayTracingShader(ShaderKey);
+
+	if (RetShader == nullptr)
+	{
+		// Do serialize outside of lock
+		FMemoryReaderView Ar(Code, true);
+		FVulkanShaderHeader CodeHeader;
+		Ar << CodeHeader;
+		FShaderResourceTable SerializedSRT;
+		Ar << SerializedSRT;
+		FVulkanShader::FSpirvContainer SpirvContainer;
+		Ar << SpirvContainer;
+
+		{
+			FRWScopeLock ScopedLock(RWLock[ShaderFrequency], SLT_Write);
+			FVulkanShader* const* FoundShaderPtr = ShaderMap[ShaderFrequency].Find(ShaderKey);
+			if (FoundShaderPtr)
+			{
+				RetShader = static_cast<FVulkanRayTracingShader*>(*FoundShaderPtr);
+			}
+			else
+			{
+				RetShader = new FVulkanRayTracingShader(Device, ShaderFrequency);
+				RetShader->Setup(MoveTemp(CodeHeader), MoveTemp(SerializedSRT), MoveTemp(SpirvContainer), ShaderKey);
+				RetShader->RayTracingPayloadType = RetShader->CodeHeader.RayTracingPayloadType;
+				RetShader->RayTracingPayloadSize = RetShader->CodeHeader.RayTracingPayloadSize;
+
+				ShaderMap[ShaderFrequency].Add(ShaderKey, RetShader);
+
+#if !(UE_BUILD_SHIPPING || UE_BUILD_TEST)
+				FShaderCodeReader ShaderCode(Code);
+				RetShader->Debug.ShaderName = ShaderCode.FindOptionalData(FShaderCodeName::Key);
+#endif
+			}
+		}
+	}
+
+	return RetShader;
+}
+
 void FVulkanShaderFactory::LookupShaders(const uint64 InShaderKeys[ShaderStage::NumStages], FVulkanShader* OutShaders[ShaderStage::NumStages]) const
 {
 	for (int32 Idx = 0; Idx < ShaderStage::NumStages; ++Idx)
@@ -370,6 +426,22 @@ FVulkanShader::FSpirvCode FVulkanShader::GetPatchedSpirvCode(const FGfxPipelineD
 	return Spirv;
 }
 
+// Bindless variant of function that does not require layout for patching
+TRefCountPtr<FVulkanShaderModule> FVulkanShader::GetOrCreateHandle()
+{
+	check(Device->SupportsBindless());
+	FScopeLock Lock(&VulkanShaderModulesMapCS);
+	FSpirvCode Spirv = GetSpirvCode();
+
+	TRefCountPtr<FVulkanShaderModule> Module = CreateShaderModule(Device, Spirv);
+	ShaderModules.Add(0, Module);
+	if (!CodeHeader.DebugName.IsEmpty())
+	{
+		VULKAN_SET_DEBUG_NAME((*Device), VK_OBJECT_TYPE_SHADER_MODULE, Module->GetVkShaderModule(), TEXT("%s : (FVulkanShader*)0x%p"), *CodeHeader.DebugName, this);
+	}
+	return Module;
+}
+
 TRefCountPtr<FVulkanShaderModule> FVulkanShader::CreateHandle(const FVulkanLayout* Layout, uint32 LayoutHash)
 {
 	FScopeLock Lock(&VulkanShaderModulesMapCS);
@@ -473,16 +545,16 @@ FRayTracingShaderRHIRef FVulkanDynamicRHI::RHICreateRayTracingShader(TArrayView<
 	switch (ShaderFrequency)
 	{
 	case EShaderFrequency::SF_RayGen:
-		 return Device->GetShaderFactory().CreateShader<FVulkanRayGenShader>(Code, Device);
+		 return Device->GetShaderFactory().CreateRayTracingShader<SF_RayGen>(Code, Device);
 
 	case EShaderFrequency::SF_RayMiss:
-		return Device->GetShaderFactory().CreateShader<FVulkanRayMissShader>(Code, Device);
+		return Device->GetShaderFactory().CreateRayTracingShader<SF_RayMiss>(Code, Device);
 
 	case EShaderFrequency::SF_RayCallable:
-		return Device->GetShaderFactory().CreateShader<FVulkanRayCallableShader>(Code, Device);
+		return Device->GetShaderFactory().CreateRayTracingShader<SF_RayCallable>(Code, Device);
 
 	case EShaderFrequency::SF_RayHitGroup:
-		return Device->GetShaderFactory().CreateShader<FVulkanRayHitGroupShader>(Code, Device);
+		return Device->GetShaderFactory().CreateRayTracingShader<SF_RayHitGroup>(Code, Device);
 
 	default:
 		check(false);
