Vulkan: print out unknown vendor ids so they can be added.

diff --git a/Engine/Source/Runtime/VulkanRHI/Private/VulkanDevice.cpp b/Engine/Source/Runtime/VulkanRHI/Private/VulkanDevice.cpp
--- a/Engine/Source/Runtime/VulkanRHI/Private/VulkanDevice.cpp
+++ b/Engine/Source/Runtime/VulkanRHI/Private/VulkanDevice.cpp
@@ -271,12 +271,13 @@ FVulkanDevice::FVulkanDevice(FVulkanDynamicRHI* InRHI, VkPhysicalDevice InGpu)
 
 	// First get the VendorId. We'll have to get properties again after finding out which extensions we want to use
 	VendorId = RHIConvertToGpuVendorId(GpuProps.vendorID);
-	ensure(VendorId != EGpuVendorId::Unknown);
 
 	UE_LOG(LogVulkanRHI, Display, TEXT("- DeviceName: %s"), ANSI_TO_TCHAR(GpuProps.deviceName));
 	UE_LOG(LogVulkanRHI, Display, TEXT("- API=%d.%d.%d (0x%x) Driver=0x%x VendorId=0x%x"), VK_VERSION_MAJOR(GpuProps.apiVersion), VK_VERSION_MINOR(GpuProps.apiVersion), VK_VERSION_PATCH(GpuProps.apiVersion), GpuProps.apiVersion, GpuProps.driverVersion, GpuProps.vendorID);
 	UE_LOG(LogVulkanRHI, Display, TEXT("- DeviceID=0x%x Type=%s"), GpuProps.deviceID, VK_TYPE_TO_STRING(VkPhysicalDeviceType, GpuProps.deviceType));
 	UE_LOG(LogVulkanRHI, Display, TEXT("- Max Descriptor Sets Bound %d, Timestamps %d"), GpuProps.limits.maxBoundDescriptorSets, GpuProps.limits.timestampComputeAndGraphics);
+
+	ensureMsgf(VendorId != EGpuVendorId::Unknown, TEXT("Unknown vendor ID 0x%x"), GpuProps.vendorID);
 }
 
 FVulkanDevice::~FVulkanDevice()
