VulkanRHI: Name the BindlessPipelineLayout for validation.

#rnx

diff --git a/Engine/Source/Runtime/VulkanRHI/Private/VulkanDescriptorSets.cpp b/Engine/Source/Runtime/VulkanRHI/Private/VulkanDescriptorSets.cpp
--- a/Engine/Source/Runtime/VulkanRHI/Private/VulkanDescriptorSets.cpp
+++ b/Engine/Source/Runtime/VulkanRHI/Private/VulkanDescriptorSets.cpp
@@ -581,6 +581,7 @@ void FVulkanBindlessDescriptorManager::Init()
 		PipelineLayoutCreateInfo.setLayoutCount = VulkanBindless::NumBindlessSets;
 		PipelineLayoutCreateInfo.pSetLayouts = DescriptorSetLayouts;
 		VERIFYVULKANRESULT(VulkanRHI::vkCreatePipelineLayout(DeviceHandle, &PipelineLayoutCreateInfo, VULKAN_CPU_ALLOCATOR, &BindlessPipelineLayout));
+		VULKAN_SET_DEBUG_NAME((*Device), VK_OBJECT_TYPE_PIPELINE_LAYOUT, BindlessPipelineLayout, TEXT("BindlessPipelineLayout(SetCount=%d)"), VulkanBindless::NumBindlessSets);
 	}
 }
 
