VulkanRHI: Add DEVICE_ADDRESS_BIT to allocations when BufferDeviceAddress extensions is loaded (not only in raytracing).

#rnx

diff --git a/Engine/Source/Runtime/VulkanRHI/Private/VulkanMemory.cpp b/Engine/Source/Runtime/VulkanRHI/Private/VulkanMemory.cpp
--- a/Engine/Source/Runtime/VulkanRHI/Private/VulkanMemory.cpp
+++ b/Engine/Source/Runtime/VulkanRHI/Private/VulkanMemory.cpp
@@ -881,16 +881,14 @@ namespace VulkanRHI
 			Info.pNext = &VulkanExportMemoryAllocateInfoKHR;
 		}
 
-#if VULKAN_RHI_RAYTRACING
 		VkMemoryAllocateFlagsInfo MemoryAllocateFlagsInfo;
-		if (Device->GetOptionalExtensions().HasRaytracingExtensions())
+		if (Device->GetOptionalExtensions().HasBufferDeviceAddress)
 		{
 			ZeroVulkanStruct(MemoryAllocateFlagsInfo, VK_STRUCTURE_TYPE_MEMORY_ALLOCATE_FLAGS_INFO);
 			MemoryAllocateFlagsInfo.flags = VK_MEMORY_ALLOCATE_DEVICE_ADDRESS_BIT_KHR;
 			MemoryAllocateFlagsInfo.pNext = Info.pNext;
 			Info.pNext = &MemoryAllocateFlagsInfo;
 		}
-#endif
 
 		VkDeviceMemory Handle;
 		VkResult Result;
