VulkanRHI: VKRT WIP: Add trace ray indirect functions.

#rnx

diff --git a/Engine/Source/Runtime/VulkanRHI/Private/VulkanCommandWrappers.h b/Engine/Source/Runtime/VulkanRHI/Private/VulkanCommandWrappers.h
--- a/Engine/Source/Runtime/VulkanRHI/Private/VulkanCommandWrappers.h
+++ b/Engine/Source/Runtime/VulkanRHI/Private/VulkanCommandWrappers.h
@@ -191,6 +191,8 @@ struct FWrapLayer
 	static void GetAccelerationStructureBuildSizesKHR(VkResult Result, VkDevice Device, VkAccelerationStructureBuildTypeKHR BuildType, const VkAccelerationStructureBuildGeometryInfoKHR* BuildInfo, const uint32* MaxPrimitiveCounts, VkAccelerationStructureBuildSizesInfoKHR* SizeInfo) VULKAN_LAYER_BODY
 	static void GetAccelerationStructureDeviceAddressKHR(VkResult Result, VkDevice Device, const VkAccelerationStructureDeviceAddressInfoKHR* Info) VULKAN_LAYER_BODY
 	static void CmdTraceRaysKHR(VkResult Result, VkCommandBuffer CommandBuffer, const VkStridedDeviceAddressRegionKHR* RaygenShaderBindingTable, const VkStridedDeviceAddressRegionKHR* MissShaderBindingTable, const VkStridedDeviceAddressRegionKHR* HitShaderBindingTable, const VkStridedDeviceAddressRegionKHR* CallableShaderBindingTable, uint32 width, uint32 height, uint32 depth) VULKAN_LAYER_BODY
+	static void CmdTraceRaysIndirectKHR(VkResult Result, VkCommandBuffer CommandBuffer, const VkStridedDeviceAddressRegionKHR* RaygenShaderBindingTable, const VkStridedDeviceAddressRegionKHR* MissShaderBindingTable, const VkStridedDeviceAddressRegionKHR* HitShaderBindingTable, const VkStridedDeviceAddressRegionKHR* CallableShaderBindingTable, VkDeviceAddress IndirectDeviceAddress) VULKAN_LAYER_BODY
+	static void CmdTraceRaysIndirect2KHR(VkResult Result, VkCommandBuffer CommandBuffer, VkDeviceAddress IndirectDeviceAddress) VULKAN_LAYER_BODY
 	static void CreateRayTracingPipelinesKHR(VkResult Result, VkDevice Device, VkDeferredOperationKHR DeferredOperation, VkPipelineCache PipelineCache, uint32_t CreateInfoCount, const VkRayTracingPipelineCreateInfoKHR* CreateInfos, const VkAllocationCallbacks* Allocator, VkPipeline* Pipelines) VULKAN_LAYER_BODY
 	static void GetRayTracingShaderGroupHandlesKHR(VkResult Result, VkDevice Device, VkPipeline Pipeline, uint32_t FirstGroup, uint32_t GroupCount, size_t DataSize, void* Data) VULKAN_LAYER_BODY
 	static void CmdWriteAccelerationStructuresPropertiesKHR(VkResult Result, VkCommandBuffer CommandBuffer, uint32_t AccelerationStructureCount, const VkAccelerationStructureKHR* AccelerationStructures, VkQueryType QueryType, VkQueryPool QueryPool, uint32_t FirstQuery) VULKAN_LAYER_BODY
@@ -1457,6 +1459,20 @@ namespace VulkanRHI
 		FWrapLayer::CmdTraceRaysKHR(VK_SUCCESS, CommandBuffer, RaygenShaderBindingTable, MissShaderBindingTable, HitShaderBindingTable, CallableShaderBindingTable, width, height, depth);
 	}
 
+	static FORCEINLINE_DEBUGGABLE void vkCmdTraceRaysIndirectKHR(VkCommandBuffer CommandBuffer, const VkStridedDeviceAddressRegionKHR* RaygenShaderBindingTable, const VkStridedDeviceAddressRegionKHR* MissShaderBindingTable, const VkStridedDeviceAddressRegionKHR* HitShaderBindingTable, const VkStridedDeviceAddressRegionKHR* CallableShaderBindingTable, VkDeviceAddress IndirectDeviceAddress)
+	{
+		FWrapLayer::CmdTraceRaysIndirectKHR(VK_RESULT_MAX_ENUM, CommandBuffer, RaygenShaderBindingTable, MissShaderBindingTable, HitShaderBindingTable, CallableShaderBindingTable, IndirectDeviceAddress);
+		VULKANAPINAMESPACE::vkCmdTraceRaysIndirectKHR(CommandBuffer, RaygenShaderBindingTable, MissShaderBindingTable, HitShaderBindingTable, CallableShaderBindingTable, IndirectDeviceAddress);
+		FWrapLayer::CmdTraceRaysIndirectKHR(VK_SUCCESS, CommandBuffer, RaygenShaderBindingTable, MissShaderBindingTable, HitShaderBindingTable, CallableShaderBindingTable, IndirectDeviceAddress);
+	}
+
+	static FORCEINLINE_DEBUGGABLE void vkCmdTraceRaysIndirect2KHR(VkCommandBuffer CommandBuffer, VkDeviceAddress IndirectDeviceAddress)
+	{
+		FWrapLayer::CmdTraceRaysIndirect2KHR(VK_RESULT_MAX_ENUM, CommandBuffer, IndirectDeviceAddress);
+		VULKANAPINAMESPACE::vkCmdTraceRaysIndirect2KHR(CommandBuffer, IndirectDeviceAddress);
+		FWrapLayer::CmdTraceRaysIndirect2KHR(VK_SUCCESS, CommandBuffer, IndirectDeviceAddress);
+	}
+
 	static FORCEINLINE_DEBUGGABLE VkResult vkCreateRayTracingPipelinesKHR(VkDevice Device, VkDeferredOperationKHR DeferredOperation, VkPipelineCache PipelineCache, uint32_t CreateInfoCount, const VkRayTracingPipelineCreateInfoKHR* CreateInfos, const VkAllocationCallbacks* Allocator, VkPipeline* Pipelines)
 	{
 		FWrapLayer::CreateRayTracingPipelinesKHR(VK_RESULT_MAX_ENUM, Device, DeferredOperation, PipelineCache, CreateInfoCount, CreateInfos, Allocator, Pipelines);
diff --git a/Engine/Source/Runtime/VulkanRHI/Private/VulkanDebug.cpp b/Engine/Source/Runtime/VulkanRHI/Private/VulkanDebug.cpp
--- a/Engine/Source/Runtime/VulkanRHI/Private/VulkanDebug.cpp
+++ b/Engine/Source/Runtime/VulkanRHI/Private/VulkanDebug.cpp
@@ -4326,6 +4326,38 @@ void FWrapLayer::CmdTraceRaysKHR(VkResult Result, VkCommandBuffer CommandBuffer,
 	}
 }
 
+void FWrapLayer::CmdTraceRaysIndirectKHR(VkResult Result, VkCommandBuffer CommandBuffer, const VkStridedDeviceAddressRegionKHR* RaygenShaderBindingTable, const VkStridedDeviceAddressRegionKHR* MissShaderBindingTable, const VkStridedDeviceAddressRegionKHR* HitShaderBindingTable, const VkStridedDeviceAddressRegionKHR* CallableShaderBindingTable, VkDeviceAddress IndirectDeviceAddress)
+{
+	if (Result == VK_RESULT_MAX_ENUM)
+	{
+#if VULKAN_ENABLE_DUMP_LAYER
+		PrintfBeginResult(FString::Printf(TEXT("vkCmdTraceRaysIndirectKHR")));
+#endif
+	}
+	else
+	{
+#if VULKAN_ENABLE_DUMP_LAYER
+		PrintResult(Result);
+#endif
+	}
+}
+
+void FWrapLayer::CmdTraceRaysIndirect2KHR(VkResult Result, VkDeviceAddress IndirectDeviceAddress)
+{
+	if (Result == VK_RESULT_MAX_ENUM)
+	{
+#if VULKAN_ENABLE_DUMP_LAYER
+		PrintfBeginResult(FString::Printf(TEXT("vkCmdTraceRaysIndirect2KHR")));
+#endif
+	}
+	else
+	{
+#if VULKAN_ENABLE_DUMP_LAYER
+		PrintResult(Result);
+#endif
+	}
+}
+
 void FWrapLayer::CreateRayTracingPipelinesKHR(VkResult Result, VkDevice Device, VkDeferredOperationKHR DeferredOperation, VkPipelineCache PipelineCache, uint32_t CreateInfoCount, const VkRayTracingPipelineCreateInfoKHR* CreateInfos, const VkAllocationCallbacks* Allocator, VkPipeline* Pipelines)
 {
 	if (Result == VK_RESULT_MAX_ENUM)
diff --git a/Engine/Source/Runtime/VulkanRHI/Private/VulkanLoader.h b/Engine/Source/Runtime/VulkanRHI/Private/VulkanLoader.h
--- a/Engine/Source/Runtime/VulkanRHI/Private/VulkanLoader.h
+++ b/Engine/Source/Runtime/VulkanRHI/Private/VulkanLoader.h
@@ -187,6 +187,8 @@
 	EnumMacro(PFN_vkGetAccelerationStructureBuildSizesKHR, vkGetAccelerationStructureBuildSizesKHR) \
 	EnumMacro(PFN_vkGetAccelerationStructureDeviceAddressKHR, vkGetAccelerationStructureDeviceAddressKHR) \
 	EnumMacro(PFN_vkCmdTraceRaysKHR, vkCmdTraceRaysKHR) \
+	EnumMacro(PFN_vkCmdTraceRaysIndirectKHR, vkCmdTraceRaysIndirectKHR) \
+	EnumMacro(PFN_vkCmdTraceRaysIndirect2KHR, vkCmdTraceRaysIndirect2KHR) \
 	EnumMacro(PFN_vkCreateRayTracingPipelinesKHR, vkCreateRayTracingPipelinesKHR) \
 	EnumMacro(PFN_vkGetRayTracingShaderGroupHandlesKHR, vkGetRayTracingShaderGroupHandlesKHR) \
 	EnumMacro(PFN_vkCmdWriteAccelerationStructuresPropertiesKHR, vkCmdWriteAccelerationStructuresPropertiesKHR) \
