Enable native MMV for preview platforms when running editor on Vulkan (UE-147903)

diff --git a/Engine/Config/Android/DataDrivenPlatformInfo.ini b/Engine/Config/Android/DataDrivenPlatformInfo.ini
--- a/Engine/Config/Android/DataDrivenPlatformInfo.ini
+++ b/Engine/Config/Android/DataDrivenPlatformInfo.ini
@@ -73,6 +73,8 @@ ShaderFormat=SF_VULKAN_ES31_ANDROID
 bIsMobile=true
 bTargetsTiledGPU=true
 bSupportsMobileMultiView=true
+; The support really is conditional, the runtime code needs to also test GRHISupportsArrayIndexFromAnyShader
+bSupportsVertexShaderLayer=true
 bNeedsOfflineCompiler=true
 bSupportsIndexBufferUAVs=true
 bSupportsManualVertexFetch=false
diff --git a/Engine/Config/VulkanPC/DataDrivenPlatformInfo.ini b/Engine/Config/VulkanPC/DataDrivenPlatformInfo.ini
--- a/Engine/Config/VulkanPC/DataDrivenPlatformInfo.ini
+++ b/Engine/Config/VulkanPC/DataDrivenPlatformInfo.ini
@@ -261,6 +261,8 @@ bIsMobile=true
 bSupportsMobileMultiView=true
 bSupportsIndexBufferUAVs=true
 bSupportsManualVertexFetch=false
+; The support really is conditional, the runtime code needs to also test GRHISupportsArrayIndexFromAnyShader
+bSupportsVertexShaderLayer=true
 bSupportsSceneDataCompressedTransforms=true
 FriendlyName=LOCTEXT("FriendlyShaderPlatformName_VulkanPC_Mobile", "PC Vulkan Mobile")
 
diff --git a/Engine/Source/Runtime/VulkanRHI/Private/Linux/VulkanLinuxPlatform.h b/Engine/Source/Runtime/VulkanRHI/Private/Linux/VulkanLinuxPlatform.h
--- a/Engine/Source/Runtime/VulkanRHI/Private/Linux/VulkanLinuxPlatform.h
+++ b/Engine/Source/Runtime/VulkanRHI/Private/Linux/VulkanLinuxPlatform.h
@@ -14,7 +14,7 @@
 
 #define VULKAN_RHI_RAYTRACING 						(RHI_RAYTRACING)
 #define VULKAN_SUPPORTS_SCALAR_BLOCK_LAYOUT			(VULKAN_RHI_RAYTRACING)
-#define VULKAN_SUPPORTS_MULTIVIEW					0
+#define VULKAN_SUPPORTS_MULTIVIEW					1	// needed for VULKAN_PCES31
 
 #if VULKAN_RHI_RAYTRACING
 #	define UE_VK_API_VERSION						VK_API_VERSION_1_2
diff --git a/Engine/Source/Runtime/VulkanRHI/Private/VulkanPendingState.cpp b/Engine/Source/Runtime/VulkanRHI/Private/VulkanPendingState.cpp
--- a/Engine/Source/Runtime/VulkanRHI/Private/VulkanPendingState.cpp
+++ b/Engine/Source/Runtime/VulkanRHI/Private/VulkanPendingState.cpp
@@ -501,7 +501,8 @@ void FVulkanPendingGfxState::InternalUpdateDynamicStates(FVulkanCmdBuffer* Cmd)
 	// Validate and update Viewport
 	if (bNeedsUpdateViewport)
 	{
-		ensure(Viewports[0].width > 0 || Viewports[0].height > 0);
+		// it is legal to pass a zero-area viewport, and the higher level expectation (see e.g. FProjectedShadowInfo::SetupProjectionStencilMask()) is that
+		// such viewport is going to be essentially disabled.
 
 		// Flip viewport on Y-axis to be uniform between DXC generated SPIR-V shaders (requires VK_KHR_maintenance1 extension)
 		TArray<VkViewport, TInlineAllocator<2>> FlippedViewports = Viewports;
diff --git a/Engine/Source/Runtime/VulkanRHI/Private/VulkanViewport.cpp b/Engine/Source/Runtime/VulkanRHI/Private/VulkanViewport.cpp
--- a/Engine/Source/Runtime/VulkanRHI/Private/VulkanViewport.cpp
+++ b/Engine/Source/Runtime/VulkanRHI/Private/VulkanViewport.cpp
@@ -12,7 +12,7 @@
 #include "GlobalShader.h"
 #include "HAL/PlatformAtomics.h"
 #include "Engine/RendererSettings.h"
-
+#include "StereoRenderUtils.h"
 
 FVulkanBackBuffer::FVulkanBackBuffer(FVulkanDevice& Device, FVulkanViewport* InViewport, EPixelFormat Format, uint32 SizeX, uint32 SizeY, ETextureCreateFlags UEFlags)
 	: FVulkanTexture(Device, FRHITextureCreateDesc::Create2D(TEXT("FVulkanBackBuffer"), SizeX, SizeY, Format).SetFlags(UEFlags).DetermineInititialState(), VK_NULL_HANDLE, false)
@@ -768,8 +768,12 @@ void FVulkanViewport::CreateSwapchain(FVulkanSwapChainRecreateInfo* RecreateInfo
 		uint32 BackBufferSizeX = RequiresRenderingBackBuffer() ? SizeX : 1;
 		uint32 BackBufferSizeY = RequiresRenderingBackBuffer() ? SizeY : 1;
 
-		const FRHITextureCreateDesc Desc =
-			FRHITextureCreateDesc::Create2D(TEXT("RenderingBackBuffer"), BackBufferSizeX, BackBufferSizeY, PixelFormat)
+		const UE::StereoRenderUtils::FStereoShaderAspects Aspects(GMaxRHIShaderPlatform);
+
+		const int kMultiViewCount = 2; // TODO: number of subresources may change in the future
+		const FRHITextureCreateDesc Desc = (Aspects.IsMobileMultiViewEnabled() ?			
+			FRHITextureCreateDesc::Create2DArray(TEXT("RenderingBackBufferArr"), BackBufferSizeX, BackBufferSizeY, kMultiViewCount, PixelFormat) :
+			FRHITextureCreateDesc::Create2D(TEXT("RenderingBackBuffer"), BackBufferSizeX, BackBufferSizeY, PixelFormat))
 			.SetClearValue(FClearValueBinding::None)
 			.SetFlags(ETextureCreateFlags::RenderTargetable | ETextureCreateFlags::ShaderResource | ETextureCreateFlags::ResolveTargetable)
 			.SetInitialState(ERHIAccess::Present);
diff --git a/Engine/Source/Runtime/VulkanRHI/Private/Windows/VulkanWindowsPlatform.h b/Engine/Source/Runtime/VulkanRHI/Private/Windows/VulkanWindowsPlatform.h
--- a/Engine/Source/Runtime/VulkanRHI/Private/Windows/VulkanWindowsPlatform.h
+++ b/Engine/Source/Runtime/VulkanRHI/Private/Windows/VulkanWindowsPlatform.h
@@ -18,7 +18,7 @@
 
 #define VULKAN_RHI_RAYTRACING 						(RHI_RAYTRACING)
 #define VULKAN_SUPPORTS_SCALAR_BLOCK_LAYOUT			(VULKAN_RHI_RAYTRACING)
-#define VULKAN_SUPPORTS_MULTIVIEW					0   // TODO-XR: UE-147795
+#define VULKAN_SUPPORTS_MULTIVIEW					1	// needed for VULKAN_PCES31
 
 #if VULKAN_RHI_RAYTRACING
 #	define UE_VK_API_VERSION						VK_API_VERSION_1_2
