Fix structure init

#vulkan

diff --git a/Engine/Source/Runtime/VulkanRHI/Private/VulkanRHI.cpp b/Engine/Source/Runtime/VulkanRHI/Private/VulkanRHI.cpp
--- a/Engine/Source/Runtime/VulkanRHI/Private/VulkanRHI.cpp
+++ b/Engine/Source/Runtime/VulkanRHI/Private/VulkanRHI.cpp
@@ -288,7 +288,7 @@ static bool CheckVulkanProfile(ERHIFeatureLevel::Type FeatureLevel, bool bRaytra
 		ZeroVulkanStruct(InstanceCreateInfo, VK_STRUCTURE_TYPE_INSTANCE_CREATE_INFO);
 
 		VpInstanceCreateInfo ProfileInstanceCreateInfo;
-		FMemory::Memzero(InstanceCreateInfo);
+		FMemory::Memzero(ProfileInstanceCreateInfo);
 		ProfileInstanceCreateInfo.pProfile = &ProfileProperties;
 		ProfileInstanceCreateInfo.pCreateInfo = &InstanceCreateInfo;
 
