VulkanRHI: VKRT WIP: Make bindless a requirement of Vulkan ray tracing pipelines.  Use the single bindless pipeline layout (one descriptor set per resource type) and always add DESCRIPTOR_BUFFER_BIT.

#rnx

diff --git a/Engine/Source/Runtime/VulkanRHI/Private/VulkanDescriptorSets.h b/Engine/Source/Runtime/VulkanRHI/Private/VulkanDescriptorSets.h
--- a/Engine/Source/Runtime/VulkanRHI/Private/VulkanDescriptorSets.h
+++ b/Engine/Source/Runtime/VulkanRHI/Private/VulkanDescriptorSets.h
@@ -1376,22 +1376,6 @@ protected:
 	friend class FVulkanPipelineStateCacheManager;
 };
 
-#if VULKAN_RHI_RAYTRACING
-class FVulkanRayTracingLayout : public FVulkanLayout
-{
-public:
-	FVulkanRayTracingLayout(FVulkanDevice* InDevice)
-		: FVulkanLayout(InDevice)
-	{
-	}
-
-	virtual bool IsGfxLayout() const final override
-	{
-		return false;
-	}
-};
-#endif // VULKAN_RHI_RAYTRACING
-
 
 
 class FVulkanGenericDescriptorPool : FNoncopyable
diff --git a/Engine/Source/Runtime/VulkanRHI/Private/VulkanRayTracing.cpp b/Engine/Source/Runtime/VulkanRHI/Private/VulkanRayTracing.cpp
--- a/Engine/Source/Runtime/VulkanRHI/Private/VulkanRayTracing.cpp
+++ b/Engine/Source/Runtime/VulkanRHI/Private/VulkanRayTracing.cpp
@@ -1039,43 +1039,14 @@ static uint32 GetAlignedSize(uint32 Value, uint32 Alignment)
 }
 
 FVulkanRayTracingPipelineState::FVulkanRayTracingPipelineState(FVulkanDevice* const InDevice, const FRayTracingPipelineStateInitializer& Initializer)
-{	
-	check(Layout == nullptr);
+{
+	checkf(InDevice->SupportsBindless(), TEXT("Vulkan ray tracing pipelines are only supported in bindless."));
 
 	TArrayView<FRHIRayTracingShader*> InitializerRayGenShaders = Initializer.GetRayGenTable();
 	TArrayView<FRHIRayTracingShader*> InitializerMissShaders = Initializer.GetMissTable();
 	TArrayView<FRHIRayTracingShader*> InitializerHitGroupShaders = Initializer.GetHitGroupTable();
 	// vkrt todo: Callable shader support
 
-	FVulkanDescriptorSetsLayoutInfo DescriptorSetLayoutInfo;
-	FUniformBufferGatherInfo UBGatherInfo;
-	
-	for (FRHIRayTracingShader* RayGenShader : InitializerRayGenShaders)
-	{
-		const FVulkanShaderHeader& Header = ResourceCast(RayGenShader)->GetCodeHeader();
-		DescriptorSetLayoutInfo.ProcessBindingsForStage(VK_SHADER_STAGE_RAYGEN_BIT_KHR, ShaderStage::RayGen, Header, UBGatherInfo);
-	}
-
-	for (FRHIRayTracingShader* MissShader : InitializerMissShaders)
-	{
-		const FVulkanShaderHeader& Header = ResourceCast(MissShader)->GetCodeHeader();
-		DescriptorSetLayoutInfo.ProcessBindingsForStage(VK_SHADER_STAGE_MISS_BIT_KHR, ShaderStage::RayMiss, Header, UBGatherInfo);
-	}
-
-	for (FRHIRayTracingShader* HitGroupShader : InitializerHitGroupShaders)
-	{
-		const FVulkanShaderHeader& Header = ResourceCast(HitGroupShader)->GetCodeHeader();
-		DescriptorSetLayoutInfo.ProcessBindingsForStage(VK_SHADER_STAGE_CLOSEST_HIT_BIT_KHR, ShaderStage::RayHitGroup, Header, UBGatherInfo);
-		// vkrt todo: How to handle any hit for hit group?
-	}
-
-	DescriptorSetLayoutInfo.FinalizeBindings<false>(*InDevice, UBGatherInfo, TArrayView<FRHISamplerState*>());
-
-	Layout = new FVulkanRayTracingLayout(InDevice);
-	Layout->DescriptorSetLayout.CopyFrom(DescriptorSetLayoutInfo);
-	FVulkanDescriptorSetLayoutMap DSetLayoutMap;
-	Layout->Compile(DSetLayoutMap);
-
 	TArray<VkPipelineShaderStageCreateInfo> ShaderStages;
 	TArray<VkRayTracingShaderGroupCreateInfoKHR> ShaderGroups;
 	TArray<ANSICHAR*> EntryPointNames;
@@ -1166,7 +1137,8 @@ FVulkanRayTracingPipelineState::FVulkanRayTracingPipelineState(FVulkanDevice* co
 	RayTracingPipelineCreateInfo.groupCount = ShaderGroups.Num();
 	RayTracingPipelineCreateInfo.pGroups = ShaderGroups.GetData();
 	RayTracingPipelineCreateInfo.maxPipelineRayRecursionDepth = 1;
-	RayTracingPipelineCreateInfo.layout = Layout->GetPipelineLayout();
+	RayTracingPipelineCreateInfo.layout = InDevice->GetBindlessDescriptorManager()->GetPipelineLayout();
+	RayTracingPipelineCreateInfo.flags = VK_PIPELINE_CREATE_DESCRIPTOR_BUFFER_BIT_EXT;
 	
 	VERIFYVULKANRESULT(VulkanDynamicAPI::vkCreateRayTracingPipelinesKHR(
 		InDevice->GetInstanceHandle(), 
@@ -1222,12 +1194,6 @@ FVulkanRayTracingPipelineState::~FVulkanRayTracingPipelineState()
 	FVulkanRayTracingAllocator::Free(RayGenShaderBindingTable);
 	FVulkanRayTracingAllocator::Free(MissShaderBindingTable);
 	FVulkanRayTracingAllocator::Free(HitShaderBindingTable);
-
-	if (Layout != nullptr)
-	{
-		delete Layout;
-		Layout = nullptr;
-	}
 }
 
 void FVulkanRayTracingCompactedSizeQueryPool::EndBatch(FVulkanCmdBuffer* InCmdBuffer)
diff --git a/Engine/Source/Runtime/VulkanRHI/Private/VulkanRayTracing.h b/Engine/Source/Runtime/VulkanRHI/Private/VulkanRayTracing.h
--- a/Engine/Source/Runtime/VulkanRHI/Private/VulkanRayTracing.h
+++ b/Engine/Source/Runtime/VulkanRHI/Private/VulkanRayTracing.h
@@ -8,7 +8,6 @@
 
 class FVulkanCommandListContext;
 class FVulkanResourceMultiBuffer;
-class FVulkanRayTracingLayout;
 
 class FVulkanRayTracingPlatform
 {
@@ -149,8 +148,6 @@ public:
 	~FVulkanRayTracingPipelineState();
 
 private:
-
-	FVulkanRayTracingLayout* Layout = nullptr;
 	VkPipeline Pipeline = VK_NULL_HANDLE;
 	FVkRtAllocation RayGenShaderBindingTable;
 	FVkRtAllocation MissShaderBindingTable;
