VulkanRHI: initialize LLM tags earlier

We'll use FVulkanDynamicRHIModule module startup to initialize LLM, so we can get in before Vulkan profiles are setup!

diff --git a/Engine/Source/Runtime/VulkanRHI/Private/VulkanRHI.cpp b/Engine/Source/Runtime/VulkanRHI/Private/VulkanRHI.cpp
--- a/Engine/Source/Runtime/VulkanRHI/Private/VulkanRHI.cpp
+++ b/Engine/Source/Runtime/VulkanRHI/Private/VulkanRHI.cpp
@@ -325,6 +325,13 @@ static bool CheckVulkanProfile(ERHIFeatureLevel::Type FeatureLevel, bool bRaytra
 	return false;
 }
 
+void FVulkanDynamicRHIModule::StartupModule()
+{
+#if VULKAN_USE_LLM
+	LLM(VulkanLLM::Initialize());
+#endif
+}
+
 bool FVulkanDynamicRHIModule::IsSupported()
 {
 	if (FVulkanPlatform::IsSupported())
@@ -493,9 +500,7 @@ FVulkanDynamicRHI::FVulkanDynamicRHI()
 		GGPUCrashDebuggingEnabled = (GPUCrashDebuggingCVar && GPUCrashDebuggingCVar->GetInt() != 0) || FParse::Param(FCommandLine::Get(), TEXT("gpucrashdebugging"));
 	}
 
-#if VULKAN_USE_LLM
-	LLM(VulkanLLM::Initialize());
-#endif
+
 
 	CreateInstance();
 	SelectDevice();
diff --git a/Engine/Source/Runtime/VulkanRHI/Public/VulkanDynamicRHI.h b/Engine/Source/Runtime/VulkanRHI/Public/VulkanDynamicRHI.h
--- a/Engine/Source/Runtime/VulkanRHI/Public/VulkanDynamicRHI.h
+++ b/Engine/Source/Runtime/VulkanRHI/Public/VulkanDynamicRHI.h
@@ -398,6 +398,9 @@ PRAGMA_ENABLE_DEPRECATION_WARNINGS
 class FVulkanDynamicRHIModule : public IDynamicRHIModule
 {
 public:
+	// IModuleInterface
+	virtual void StartupModule() override;
+
 	// IDynamicRHIModule
 	virtual bool IsSupported() override;
 	virtual bool IsSupported(ERHIFeatureLevel::Type RequestedFeatureLevel) override;
