VulkanRHI: VKRT: Add missing ifdef for raytracing code.  Add function to get uniform buffer device address.

#rnx

diff --git a/Engine/Source/Runtime/VulkanRHI/Private/VulkanResources.h b/Engine/Source/Runtime/VulkanRHI/Private/VulkanResources.h
--- a/Engine/Source/Runtime/VulkanRHI/Private/VulkanResources.h
+++ b/Engine/Source/Runtime/VulkanRHI/Private/VulkanResources.h
@@ -313,8 +313,10 @@ public:
 		return nullptr;
 	}
 
+#if VULKAN_RHI_RAYTRACING
 	template <EShaderFrequency ShaderFrequency>
 	FVulkanRayTracingShader* CreateRayTracingShader(TArrayView<const uint8> Code, FVulkanDevice* Device);
+#endif
 
 	void LookupShaders(const uint64 InShaderKeys[ShaderStage::NumStages], FVulkanShader* OutShaders[ShaderStage::NumStages]) const;
 
@@ -1178,6 +1180,8 @@ public:
 		NewAlloc.Swap(Allocation);
 	}
 	
+	VkDeviceAddress GetDeviceAddress() const;
+
 public:
 	FVulkanDevice* Device;
 	VulkanRHI::FVulkanAllocation Allocation;
diff --git a/Engine/Source/Runtime/VulkanRHI/Private/VulkanShaders.cpp b/Engine/Source/Runtime/VulkanRHI/Private/VulkanShaders.cpp
--- a/Engine/Source/Runtime/VulkanRHI/Private/VulkanShaders.cpp
+++ b/Engine/Source/Runtime/VulkanRHI/Private/VulkanShaders.cpp
@@ -95,6 +95,7 @@ ShaderType* FVulkanShaderFactory::CreateShader(TArrayView<const uint8> Code, FVu
 	return RetShader;
 }
 
+#if VULKAN_RHI_RAYTRACING
 template <EShaderFrequency ShaderFrequency>
 FVulkanRayTracingShader* FVulkanShaderFactory::CreateRayTracingShader(TArrayView<const uint8> Code, FVulkanDevice* Device)
 {
@@ -158,6 +159,7 @@ FVulkanRayTracingShader* FVulkanShaderFactory::CreateRayTracingShader(TArrayView
 
 	return RetShader;
 }
+#endif // VULKAN_RHI_RAYTRACING
 
 void FVulkanShaderFactory::LookupShaders(const uint64 InShaderKeys[ShaderStage::NumStages], FVulkanShader* OutShaders[ShaderStage::NumStages]) const
 {
diff --git a/Engine/Source/Runtime/VulkanRHI/Private/VulkanUniformBuffer.cpp b/Engine/Source/Runtime/VulkanRHI/Private/VulkanUniformBuffer.cpp
--- a/Engine/Source/Runtime/VulkanRHI/Private/VulkanUniformBuffer.cpp
+++ b/Engine/Source/Runtime/VulkanRHI/Private/VulkanUniformBuffer.cpp
@@ -212,6 +212,17 @@ void FVulkanUniformBuffer::UpdateResourceTable(FRHIResource** Resources, int32 R
 	}
 }
 
+VkDeviceAddress FVulkanUniformBuffer::GetDeviceAddress() const
+{
+	// :todo-jn: there will be more and more churn on this, cache the value
+	VkBufferDeviceAddressInfo BufferInfo;
+	ZeroVulkanStruct(BufferInfo, VK_STRUCTURE_TYPE_BUFFER_DEVICE_ADDRESS_INFO);
+	BufferInfo.buffer = GetBufferHandle();
+	VkDeviceAddress BufferAddress = VulkanRHI::vkGetBufferDeviceAddressKHR(Device->GetInstanceHandle(), &BufferInfo);
+	BufferAddress += GetOffset();
+	return BufferAddress;
+}
+
 
 FUniformBufferRHIRef FVulkanDynamicRHI::RHICreateUniformBuffer(const void* Contents, const FRHIUniformBufferLayout* Layout, EUniformBufferUsage Usage, EUniformBufferValidation Validation)
 {
