VulkanRHI: VKRT WIP: Add bindless uniform buffer support needed by RT.  Rename existing uniform buffers to "single use uniform buffers".

#rnx

diff --git a/Engine/Source/Developer/VulkanShaderFormat/Private/VulkanShaderCompiler.cpp b/Engine/Source/Developer/VulkanShaderFormat/Private/VulkanShaderCompiler.cpp
--- a/Engine/Source/Developer/VulkanShaderFormat/Private/VulkanShaderCompiler.cpp
+++ b/Engine/Source/Developer/VulkanShaderFormat/Private/VulkanShaderCompiler.cpp
@@ -416,7 +416,7 @@ static void AddUBResources(FOLDVulkanCodeHeader& OLDHeader,
 				UBResourceInfo.GlobalIndex = HeaderGlobalIndex;
 				UBResourceInfo.UBBaseType = (EUniformBufferBaseType)ResourceTableEntry.Type;
 #if VULKAN_ENABLE_BINDING_DEBUG_NAMES
-				UBResourceInfo.DebugName = MemberName;
+				UBResourceInfo.DebugName = ResourceTableEntry.UniformBufferMemberName;
 #endif
 				// Iterate to next info
 				ResourceInfo = *ResourceInfos++;
@@ -438,15 +438,27 @@ static void AddUniformBuffer(FOLDVulkanCodeHeader& OLDHeader,
 	TArray<FPatchType>& OutTypePatch,
 	TArray<FString>& GlobalNames)
 {
-	int32 HeaderUBIndex = -1;
-	HeaderUBIndex = OutHeader.UniformBuffers.AddZeroed();
-
+	const int32 HeaderUBIndex = OutHeader.UniformBuffers.AddZeroed();
 	FVulkanShaderHeader::FUniformBufferInfo& UBInfo = OutHeader.UniformBuffers[HeaderUBIndex];
+
 	const FUniformBufferEntry* UniformBufferEntry = ShaderInput.Environment.UniformBufferMap.Find(UBName);
-	UBInfo.LayoutHash = UniformBufferEntry ? UniformBufferEntry->LayoutHash : 0;
+	if (UniformBufferEntry)
+	{
+		UBInfo.LayoutHash = UniformBufferEntry->LayoutHash;
+	}
+	else if ((UBName == FShaderParametersMetadata::kRootUniformBufferBindingName) && ShaderInput.RootParametersStructure)
+	{
+		UBInfo.LayoutHash = ShaderInput.RootParametersStructure->GetLayoutHash();
+	}
+	else
+	{
+		UBInfo.LayoutHash = 0;
+	}
+
 #if VULKAN_ENABLE_BINDING_DEBUG_NAMES
 	UBInfo.DebugName = UBName;
 #endif
+
 	const FVulkanSpirv::FEntry* Entry = Spirv.GetEntry(UBName);
 	if (Entry)
 	{
@@ -1479,7 +1491,7 @@ static void GatherSpirvReflectionBindings(
 			for (int32 Index = BindingArray.Num() - 1; Index >= 0; --Index)
 			{
 				const SpvReflectDescriptorBinding* pBinding = BindingArray[Index];
-				FString BindingName(ANSI_TO_TCHAR(pBinding->name));
+				const FString BindingName(ANSI_TO_TCHAR(pBinding->name));
 				if (BindingName.StartsWith(HeapPrefix))
 				{
 					const uint32 Binding = 0;  // single bindless heap per descriptor set
@@ -1503,11 +1515,10 @@ static void GatherSpirvReflectionBindings(
 		// Move uniform buffers to the correct set
 		{
 			const uint32 BindingOffset = (StageIndex * VulkanBindless::MaxUniformBuffersPerStage);
-			const uint32 BinldessDescSetNo = VulkanBindless::BindlessUniformBufferSet;
 			for (int32 Index = OutBindings.UniformBuffers.Num() - 1; Index >= 0; --Index)
 			{
 				const SpvReflectDescriptorBinding* pBinding = OutBindings.UniformBuffers[Index];
-				Reflection.ChangeDescriptorBindingNumbers(pBinding, BindingOffset + pBinding->binding, BinldessDescSetNo);
+				Reflection.ChangeDescriptorBindingNumbers(pBinding, BindingOffset + pBinding->binding, VulkanBindless::BindlessSingleUseUniformBufferSet);
 			}
 		}
 	}
@@ -1618,9 +1629,10 @@ static bool BuildShaderOutputFromSpirv(
 	}
 
 	const bool bSupportsBindless = Input.Environment.CompilerFlags.Contains(CFLAG_BindlessResources) || Input.Environment.CompilerFlags.Contains(CFLAG_BindlessSamplers);
+	const EShaderFrequency Frequency = static_cast<EShaderFrequency>(Input.Target.Frequency);
 
 	FSpirvReflectBindings Bindings;
-	GatherSpirvReflectionBindings(Reflection, Bindings, static_cast<EShaderFrequency>(Input.Target.Frequency), bSupportsBindless);
+	GatherSpirvReflectionBindings(Reflection, Bindings, Frequency, bSupportsBindless);
 
 	// Register how often a sampler-state is used
 	for (const SpvReflectDescriptorBinding* Binding : Bindings.TextureSRVs)
@@ -1752,9 +1764,9 @@ static bool BuildShaderOutputFromSpirv(
 		if (ResourceName == UBOGlobalsNameSpv)
 		{
 			// Register binding for uniform buffer
-			const int32 StageOffset = bSupportsBindless ? (ShaderStage::GetStageForFrequency(static_cast<EShaderFrequency>(Input.Target.Frequency)) * VulkanBindless::MaxUniformBuffersPerStage) : 0;
+			const int32 StageOffset = bSupportsBindless ? (ShaderStage::GetStageForFrequency(Frequency) * VulkanBindless::MaxUniformBuffersPerStage) : 0;
 			const int32 BindingIndex = MapDescriptorBindingToIndex(Binding) + StageOffset;
-			const uint32_t DescSetNumber = bSupportsBindless ? (uint32_t)VulkanBindless::BindlessUniformBufferSet : (uint32_t)SPV_REFLECT_SET_NUMBER_DONT_CHANGE;
+			const uint32_t DescSetNumber = bSupportsBindless ? (uint32_t)VulkanBindless::BindlessSingleUseUniformBufferSet : (uint32_t)SPV_REFLECT_SET_NUMBER_DONT_CHANGE;
 
 			SpvResult = Reflection.ChangeDescriptorBindingNumbers(Binding, BindingIndex, DescSetNumber);
 			check(SpvResult == SPV_REFLECT_RESULT_SUCCESS);
@@ -1784,9 +1796,9 @@ static bool BuildShaderOutputFromSpirv(
 		if (ResourceName != UBOGlobalsNameSpv)
 		{
 			// Register uniform buffer
-			const int32 StageOffset = bSupportsBindless ? (ShaderStage::GetStageForFrequency(static_cast<EShaderFrequency>(Input.Target.Frequency)) * VulkanBindless::MaxUniformBuffersPerStage) : 0;
+			const int32 StageOffset = bSupportsBindless ? (ShaderStage::GetStageForFrequency(Frequency) * VulkanBindless::MaxUniformBuffersPerStage) : 0;
 			const int32 BindingIndex = MapDescriptorBindingToIndex(Binding) + StageOffset;
-			const uint32_t DescSetNumber = bSupportsBindless ? (uint32_t)VulkanBindless::BindlessUniformBufferSet : (uint32_t)SPV_REFLECT_SET_NUMBER_DONT_CHANGE;
+			const uint32_t DescSetNumber = bSupportsBindless ? (uint32_t)VulkanBindless::BindlessSingleUseUniformBufferSet : (uint32_t)SPV_REFLECT_SET_NUMBER_DONT_CHANGE;
 			SpvResult = Reflection.ChangeDescriptorBindingNumbers(Binding, BindingIndex, DescSetNumber);
 			check(SpvResult == SPV_REFLECT_RESULT_SUCCESS);
 
@@ -2484,8 +2496,8 @@ static bool CompileWithShaderConductor(
 class FVulkanShaderParameterParser : public FShaderParameterParser
 {
 public:
-	FVulkanShaderParameterParser(FShaderCompilerFlags CompilerFlags)
-		: FShaderParameterParser(CompilerFlags)
+	FVulkanShaderParameterParser(FShaderCompilerFlags CompilerFlags, const TCHAR* InConstantBufferType)
+		: FShaderParameterParser(CompilerFlags, InConstantBufferType)
 	{}
 
 protected:
@@ -2494,7 +2506,7 @@ protected:
 		if (bBindlessResources || bBindlessSamplers)
 		{
 			const bool IsSampler = (ParsedParameter.BindlessConversionType == EBindlessConversionType::Sampler);
-			const TCHAR* IndexPrefix = IsSampler ? TEXT("BindlessSampler_") : TEXT("BindlessResource_");
+			const TCHAR* IndexPrefix = IsSampler ? FShaderParameterParser::kBindlessSamplerPrefix : FShaderParameterParser::kBindlessResourcePrefix;
 			const TCHAR* HeapPrefix = IsSampler ? VulkanBindless::kBindlessSamplerArrayPrefix : VulkanBindless::kBindlessResourceArrayPrefix;
 
 			const TCHAR* StorageClass = ParsedParameter.bGloballyCoherent ? TEXT("globallycoherent ") : TEXT("");
@@ -2654,14 +2666,14 @@ void DoCompileVulkanShader(const FShaderCompilerInput& Input, FShaderCompilerOut
 		}
 	}
 
-	FVulkanShaderParameterParser ShaderParameterParser(Input.Environment.CompilerFlags);
+	FVulkanShaderParameterParser ShaderParameterParser(Input.Environment.CompilerFlags, nullptr);
 	if (!ShaderParameterParser.ParseAndModify(Input, Output.Errors, PreprocessedShaderSource))
 	{
 		// The FShaderParameterParser will add any relevant errors.
 		return;
 	}
 
-	FString EntryPointName = Input.EntryPointName;
+	const FString EntryPointName = Input.EntryPointName;
 
 	RemoveUniformBuffersFromSource(Input.Environment, PreprocessedShaderSource);
 
diff --git a/Engine/Source/Runtime/VulkanRHI/Private/VulkanDescriptorSets.cpp b/Engine/Source/Runtime/VulkanRHI/Private/VulkanDescriptorSets.cpp
--- a/Engine/Source/Runtime/VulkanRHI/Private/VulkanDescriptorSets.cpp
+++ b/Engine/Source/Runtime/VulkanRHI/Private/VulkanDescriptorSets.cpp
@@ -57,7 +57,7 @@ static FAutoConsoleVariableRef CVarVulkanBindlessMaxStorageTexelBufferCount(
 	ECVF_ReadOnly
 );
 
-int32 GVulkanBindlessMaxUniformBufferDescriptorCount = 2 * 1024 * 1024;
+int32 GVulkanBindlessMaxUniformBufferDescriptorCount = 768 * 1024;
 static FAutoConsoleVariableRef CVarVulkanBindlessMaxUniformBufferCount(
 	TEXT("r.Vulkan.Bindless.MaxResourceUniformBufferCount"),
 	GVulkanBindlessMaxUniformBufferDescriptorCount,
@@ -105,6 +105,7 @@ DECLARE_DWORD_ACCUMULATOR_STAT_EXTERN(TEXT("Num Peak Descriptor Count"), STAT_Vu
 DECLARE_DWORD_ACCUMULATOR_STAT_EXTERN(TEXT("Peak Samplers"), STAT_VulkanBindlessPeakSampler, STATGROUP_VulkanBindless, );
 DECLARE_DWORD_ACCUMULATOR_STAT_EXTERN(TEXT("Peak Sampled Images"), STAT_VulkanBindlessPeakSampledImage, STATGROUP_VulkanBindless, );
 DECLARE_DWORD_ACCUMULATOR_STAT_EXTERN(TEXT("Peak Storage Images"), STAT_VulkanBindlessPeakStorageImage, STATGROUP_VulkanBindless, );
+DECLARE_DWORD_ACCUMULATOR_STAT_EXTERN(TEXT("Peak Uniform Buffers"), STAT_VulkanBindlessPeakUniformBuffer, STATGROUP_VulkanBindless, );
 DECLARE_DWORD_ACCUMULATOR_STAT_EXTERN(TEXT("Peak Storage Buffers"), STAT_VulkanBindlessPeakStorageBuffer, STATGROUP_VulkanBindless, );
 DECLARE_DWORD_ACCUMULATOR_STAT_EXTERN(TEXT("Peak Uniform Texel Buffers"), STAT_VulkanBindlessPeakUniformTexelBuffer, STATGROUP_VulkanBindless, );
 DECLARE_DWORD_ACCUMULATOR_STAT_EXTERN(TEXT("Peak Storage Texel Buffers"), STAT_VulkanBindlessPeakStorageTexelBuffer, STATGROUP_VulkanBindless, );
@@ -115,6 +116,7 @@ DEFINE_STAT(STAT_VulkanBindlessPeakDescriptorCount);
 DEFINE_STAT(STAT_VulkanBindlessPeakSampler);
 DEFINE_STAT(STAT_VulkanBindlessPeakSampledImage);
 DEFINE_STAT(STAT_VulkanBindlessPeakStorageImage);
+DEFINE_STAT(STAT_VulkanBindlessPeakUniformBuffer);
 DEFINE_STAT(STAT_VulkanBindlessPeakStorageBuffer);
 DEFINE_STAT(STAT_VulkanBindlessPeakUniformTexelBuffer);
 DEFINE_STAT(STAT_VulkanBindlessPeakStorageTexelBuffer);
@@ -152,6 +154,7 @@ static inline VkDescriptorType GetDescriptorTypeForSetIndex(uint8 SetIndex)
 	case VulkanBindless::BindlessStorageTexelBufferSet:     return VK_DESCRIPTOR_TYPE_STORAGE_TEXEL_BUFFER;
 	case VulkanBindless::BindlessStorageBufferSet:          return VK_DESCRIPTOR_TYPE_STORAGE_BUFFER;
 	case VulkanBindless::BindlessUniformBufferSet:          return VK_DESCRIPTOR_TYPE_UNIFORM_BUFFER;
+	case VulkanBindless::BindlessSingleUseUniformBufferSet: return VK_DESCRIPTOR_TYPE_UNIFORM_BUFFER;
 	case VulkanBindless::BindlessAccelerationStructureSet:  return VK_DESCRIPTOR_TYPE_ACCELERATION_STRUCTURE_KHR;
 	default: checkNoEntry();
 	}
@@ -447,7 +450,7 @@ void FVulkanBindlessDescriptorManager::Init()
 		};
 
 		// Create the descriptor buffer for a BindlessSetState
-		auto CreateDescriptorBuffer = [&](BindlessSetState& InOutState, VkDescriptorBufferBindingInfoEXT& OutBufferBindingInfo) {
+		auto CreateDescriptorBuffer = [&](BindlessSetState& InOutState, VkDescriptorBufferBindingInfoEXT& OutBufferBindingInfo, bool IsSingleUseUniformBufferSet) {
 
 			// Skip unsupported descriptors
 			if (InOutState.DescriptorType == VK_DESCRIPTOR_TYPE_MAX_ENUM)
@@ -455,7 +458,6 @@ void FVulkanBindlessDescriptorManager::Init()
 				return 0u;
 			}
 
-			const bool IsUniformBufferSet = (InOutState.DescriptorType == VK_DESCRIPTOR_TYPE_UNIFORM_BUFFER);
 			const bool IsSamplerSet = (InOutState.DescriptorType == VK_DESCRIPTOR_TYPE_SAMPLER);
 			const VkBufferUsageFlags BufferUsageFlags =
 				VK_BUFFER_USAGE_SHADER_DEVICE_ADDRESS_BIT |
@@ -468,7 +470,7 @@ void FVulkanBindlessDescriptorManager::Init()
 
 			VkDeviceSize LayoutSizeInBytes = 0;
 			VulkanRHI::vkGetDescriptorSetLayoutSizeEXT(DeviceHandle, InOutState.DescriptorSetLayout, &LayoutSizeInBytes);
-			if (IsUniformBufferSet)
+			if (IsSingleUseUniformBufferSet)
 			{
 				// todo-jn: We're picky about uniform buffers values for now to allow for shortcuts...
 				check(LayoutSizeInBytes == (ShaderStage::MaxNumSets * VulkanBindless::MaxUniformBuffersPerStage * InOutState.DescriptorSize));
@@ -558,8 +560,9 @@ void FVulkanBindlessDescriptorManager::Init()
 
 			BindlessSetState& State = BindlessSetStates[SetIndex];
 			InitBindlessSetState(GetDescriptorTypeForSetIndex(SetIndex), State);
-			State.DescriptorSetLayout = (SetIndex == VulkanBindless::BindlessUniformBufferSet) ? CreateShaderStageUniformBufferLayout() : CreateDescriptorSetLayout(State);
-			TotalResourceDescriptorBufferSize += CreateDescriptorBuffer(State, BufferBindingInfo[SetIndex]);
+			const bool IsSingleUseUniformBufferSet = (SetIndex == VulkanBindless::BindlessSingleUseUniformBufferSet);
+			State.DescriptorSetLayout = IsSingleUseUniformBufferSet ? CreateShaderStageUniformBufferLayout() : CreateDescriptorSetLayout(State);
+			TotalResourceDescriptorBufferSize += CreateDescriptorBuffer(State, BufferBindingInfo[SetIndex], IsSingleUseUniformBufferSet);
 		}
 
 		checkf(TotalResourceDescriptorBufferSize < DescriptorBufferProperties.resourceDescriptorBufferAddressSpaceSize,
@@ -616,8 +619,7 @@ void FVulkanBindlessDescriptorManager::RegisterUniformBuffers(VkCommandBuffer Co
 
 	SCOPED_NAMED_EVENT(FVulkanBindlessDescriptorManager_RegisterUniformBuffers, FColor::Purple);
 
-	const uint8 SetIndex = GetIndexForDescriptorType(VK_DESCRIPTOR_TYPE_UNIFORM_BUFFER);
-	BindlessSetState& BindlessUniformBufferSetState = BindlessSetStates[SetIndex];
+	BindlessSetState& BindlessUniformBufferSetState = BindlessSetStates[VulkanBindless::BindlessSingleUseUniformBufferSet];
 
 	// :todo-jn: Current uniform buffer layout is a bit wasteful with all the skipped bindings...
 	const uint32 BlockDescriptorCount = VulkanBindless::MaxUniformBuffersPerStage * ShaderStage::MaxNumSets;
@@ -628,7 +630,7 @@ void FVulkanBindlessDescriptorManager::RegisterUniformBuffers(VkCommandBuffer Co
 
 	VkDeviceSize BufferOffsets[VulkanBindless::NumBindlessSets];
 	FMemory::Memzero(BufferOffsets);
-	BufferOffsets[VulkanBindless::BindlessUniformBufferSet] = FirstDescriptorByteOffset;
+	BufferOffsets[VulkanBindless::BindlessSingleUseUniformBufferSet] = FirstDescriptorByteOffset;
 	checkSlow(FirstDescriptorByteOffset % Device->GetOptionalExtensionProperties().DescriptorBufferProps.descriptorBufferOffsetAlignment == 0);
 
 	// :todo-jn: Clear them for easier debugging for now
@@ -685,6 +687,7 @@ void FVulkanBindlessDescriptorManager::RegisterUniformBuffers(VkCommandBuffer Co
 		}
 		else
 		{
+			const uint8 SetIndex = VulkanBindless::BindlessSingleUseUniformBufferSet;
 			VulkanRHI::vkCmdSetDescriptorBufferOffsetsEXT(CommandBuffer, BindPoint, BindlessPipelineLayout, SetIndex, 1u, &BufferIndices[SetIndex], &BufferOffsets[SetIndex]);
 		}
 	}
@@ -702,10 +705,11 @@ void FVulkanBindlessDescriptorManager::UpdateStatsForHandle(FRHIDescriptorHandle
 	case VulkanBindless::BindlessStorageImageSet:           SET_DWORD_STAT(STAT_VulkanBindlessPeakStorageImage, State.PeakDescriptorCount); break;
 	case VulkanBindless::BindlessUniformTexelBufferSet:     SET_DWORD_STAT(STAT_VulkanBindlessPeakUniformTexelBuffer, State.PeakDescriptorCount); break;
 	case VulkanBindless::BindlessStorageTexelBufferSet:     SET_DWORD_STAT(STAT_VulkanBindlessPeakStorageTexelBuffer, State.PeakDescriptorCount); break;
+	case VulkanBindless::BindlessUniformBufferSet:          SET_DWORD_STAT(STAT_VulkanBindlessPeakUniformBuffer, State.PeakDescriptorCount); break;
 	case VulkanBindless::BindlessStorageBufferSet:          SET_DWORD_STAT(STAT_VulkanBindlessPeakStorageBuffer, State.PeakDescriptorCount); break;
 	case VulkanBindless::BindlessAccelerationStructureSet:  SET_DWORD_STAT(STAT_VulkanBindlessPeakAccelerationStructure, State.PeakDescriptorCount); break;
 
-	case VulkanBindless::BindlessUniformBufferSet:
+	case VulkanBindless::BindlessSingleUseUniformBufferSet:
 	default: checkNoEntry();
 	}
 }
@@ -821,7 +825,7 @@ void FVulkanBindlessDescriptorManager::UpdateBuffer(FRHIDescriptorHandle Descrip
 	if (bIsSupported)
 	{
 		const VkDescriptorType DescriptorType = GetDescriptorTypeForSetIndex(DescriptorHandle.GetRawType());
-		check(/*(DescriptorType == VK_DESCRIPTOR_TYPE_UNIFORM_BUFFER) ||*/ (DescriptorType == VK_DESCRIPTOR_TYPE_STORAGE_BUFFER));
+		check((DescriptorType == VK_DESCRIPTOR_TYPE_UNIFORM_BUFFER) || (DescriptorType == VK_DESCRIPTOR_TYPE_STORAGE_BUFFER));
 
 		// :todo-jn: start caching buffer addresses in resources to avoid the extra call
 		VkBufferDeviceAddressInfo BufferInfo;
diff --git a/Engine/Source/Runtime/VulkanRHI/Public/VulkanCommon.h b/Engine/Source/Runtime/VulkanRHI/Public/VulkanCommon.h
--- a/Engine/Source/Runtime/VulkanRHI/Public/VulkanCommon.h
+++ b/Engine/Source/Runtime/VulkanRHI/Public/VulkanCommon.h
@@ -121,6 +121,7 @@ namespace VulkanBindless
 		BindlessSamplerSet = 0,
 
 		BindlessStorageBufferSet,
+		BindlessUniformBufferSet,
 
 		BindlessStorageImageSet,
 		BindlessSampledImageSet,
@@ -130,7 +131,7 @@ namespace VulkanBindless
 
 		BindlessAccelerationStructureSet,
 
-		BindlessUniformBufferSet,  // Keep last
+		BindlessSingleUseUniformBufferSet,  // Keep last
 		NumBindlessSets,
 		MaxNumSets = NumBindlessSets
 	};
