Fix use-after-free with VkRenderPassMultiviewCreateInfo

PR #10420: Fix use-after-free with VkRenderPassMultiviewCreateInfo.

diff --git a/Engine/Source/Runtime/VulkanRHI/Private/VulkanRenderpass.h b/Engine/Source/Runtime/VulkanRHI/Private/VulkanRenderpass.h
--- a/Engine/Source/Runtime/VulkanRHI/Private/VulkanRenderpass.h
+++ b/Engine/Source/Runtime/VulkanRHI/Private/VulkanRenderpass.h
@@ -355,6 +355,14 @@ struct FVulkanRenderPassFragmentDensityMapCreateInfoEXT
 	}
 };
 
+struct FVulkanRenderPassMultiviewCreateInfo
+	: public VkRenderPassMultiviewCreateInfo
+{
+	FVulkanRenderPassMultiviewCreateInfo()
+	{
+		ZeroVulkanStruct(*this, VK_STRUCTURE_TYPE_RENDER_PASS_MULTIVIEW_CREATE_INFO);
+	}
+};
 
 template<>
 struct FVulkanRenderPassCreateInfo<VkRenderPassCreateInfo2>
@@ -633,7 +641,6 @@ public:
 		*/
 		CorrelationMask = MultiviewMask;
 
-		VkRenderPassMultiviewCreateInfo MultiviewInfo;
 		if (RTLayout.GetIsMultiView())
 		{
 			if (Device.GetOptionalExtensions().HasKHRRenderPass2)
@@ -642,8 +649,7 @@ public:
 			}
 			else
 			{
-				checkf(Device.GetOptionalExtensions().HasKHRMultiview, TEXT("Layout is multiview but extension is not supported!"))
-				ZeroVulkanStruct(MultiviewInfo, VK_STRUCTURE_TYPE_RENDER_PASS_MULTIVIEW_CREATE_INFO);
+				checkf(Device.GetOptionalExtensions().HasKHRMultiview, TEXT("Layout is multiview but extension is not supported!"));
 				MultiviewInfo.subpassCount = NumSubpasses;
 				MultiviewInfo.pViewMasks = ViewMask;
 				MultiviewInfo.dependencyCount = 0;
@@ -708,6 +714,7 @@ private:
 	FVulkanFragmentShadingRateAttachmentInfo FragmentShadingRateAttachmentInfo;
 
 	FVulkanRenderPassFragmentDensityMapCreateInfoEXT FragDensityCreateInfo;
+	FVulkanRenderPassMultiviewCreateInfo MultiviewInfo;
 
 	TRenderPassCreateInfoClass CreateInfo;
 	FVulkanDevice& Device;
