VulkanRHI: VKRT WIP:
- Add empty stubs for all the ray tracing RHI calls in VulkanDynamicRHI and Context.
- Use VulkanMemory for allocations in RT.
- Add a FVulkanHitGroupSystemParameters class that includes the common FHitGroupSystemRootConstants as well as all the bindless handles for index/vertex buffers and uniform buffers.
- Add HitGroupSystemParameters to FVulkanRayTracingGeometry using bindless handles.
- Keep pointers to the shaders in the pipeline.  Include callable shaders in pipeline.  Rewrite how shader handles for the pipelines are stored (dedicated SRT class to come).

#rnx

diff --git a/Engine/Source/Runtime/VulkanRHI/Private/VulkanContext.h b/Engine/Source/Runtime/VulkanRHI/Private/VulkanContext.h
--- a/Engine/Source/Runtime/VulkanRHI/Private/VulkanContext.h
+++ b/Engine/Source/Runtime/VulkanRHI/Private/VulkanContext.h
@@ -111,7 +111,36 @@ public:
 	virtual void RHIBindAccelerationStructureMemory(FRHIRayTracingScene* Scene, FRHIBuffer* Buffer, uint32 BufferOffset) final override;
 	virtual void RHIBuildAccelerationStructures(const TArrayView<const FRayTracingGeometryBuildParams> Params, const FRHIBufferRange& ScratchBufferRange) final override;
 	virtual void RHIBuildAccelerationStructure(const FRayTracingSceneBuildParams& SceneBuildParams) final override;
-#endif
+
+	virtual void RHIRayTraceDispatch(FRHIRayTracingPipelineState* RayTracingPipelineState, FRHIRayTracingShader* RayGenShader,
+		FRHIRayTracingScene* Scene,
+		const FRayTracingShaderBindings& GlobalResourceBindings,
+		uint32 Width, uint32 Height) final override;
+	virtual void RHIRayTraceDispatchIndirect(FRHIRayTracingPipelineState* RayTracingPipelineState, FRHIRayTracingShader* RayGenShader,
+		FRHIRayTracingScene* Scene,
+		const FRayTracingShaderBindings& GlobalResourceBindings,
+		FRHIBuffer* ArgumentBuffer, uint32 ArgumentOffset) final override;
+	virtual void RHISetRayTracingHitGroup(
+		FRHIRayTracingScene* Scene, uint32 InstanceIndex, uint32 SegmentIndex, uint32 ShaderSlot,
+		FRHIRayTracingPipelineState* Pipeline, uint32 HitGroupIndex,
+		uint32 NumUniformBuffers, FRHIUniformBuffer* const* UniformBuffers,
+		uint32 LooseParameterDataSize, const void* LooseParameterData,
+		uint32 UserData) final override;
+	virtual void RHISetRayTracingCallableShader(
+		FRHIRayTracingScene* Scene, uint32 ShaderSlotInScene,
+		FRHIRayTracingPipelineState* Pipeline, uint32 ShaderIndexInPipeline,
+		uint32 NumUniformBuffers, FRHIUniformBuffer* const* UniformBuffers,
+		uint32 UserData) final override;
+	virtual void RHISetRayTracingMissShader(
+		FRHIRayTracingScene* Scene, uint32 ShaderSlotInScene,
+		FRHIRayTracingPipelineState* Pipeline, uint32 ShaderIndexInPipeline,
+		uint32 NumUniformBuffers, FRHIUniformBuffer* const* UniformBuffers,
+		uint32 UserData) final override;
+	virtual void RHISetRayTracingBindings(
+		FRHIRayTracingScene* Scene, FRHIRayTracingPipelineState* Pipeline,
+		uint32 NumBindings, const FRayTracingLocalShaderBindings* Bindings,
+		ERayTracingBindingType BindingType) final override;
+#endif // VULKAN_RHI_RAYTRACING
 
 	inline FVulkanCommandBufferManager* GetCommandBufferManager()
 	{
diff --git a/Engine/Source/Runtime/VulkanRHI/Private/VulkanRayTracing.cpp b/Engine/Source/Runtime/VulkanRHI/Private/VulkanRayTracing.cpp
--- a/Engine/Source/Runtime/VulkanRHI/Private/VulkanRayTracing.cpp
+++ b/Engine/Source/Runtime/VulkanRHI/Private/VulkanRayTracing.cpp
@@ -109,45 +109,6 @@ static uint32 FindMemoryType(VkPhysicalDevice Gpu, uint32 Filter, VkMemoryProper
 	return Result;
 }
 
-// Temporary brute force allocation
-void FVulkanRayTracingAllocator::Allocate(FVulkanDevice* Device, VkDeviceSize Size, VkBufferUsageFlags UsageFlags, VkMemoryPropertyFlags MemoryFlags, FVkRtAllocation& Result)
-{
-	VkMemoryRequirements MemoryRequirements;
-	Result.Buffer = VulkanRHI::CreateBuffer(Device, Size, UsageFlags, MemoryRequirements);
-
-	VkDevice DeviceHandle = Device->GetInstanceHandle();
-	VkPhysicalDevice Gpu = Device->GetPhysicalHandle();
-
-	VkMemoryAllocateFlagsInfo MemoryAllocateFlagsInfo;
-	ZeroVulkanStruct(MemoryAllocateFlagsInfo, VK_STRUCTURE_TYPE_MEMORY_ALLOCATE_FLAGS_INFO);
-	MemoryAllocateFlagsInfo.flags = VK_MEMORY_ALLOCATE_DEVICE_ADDRESS_BIT_KHR;
-
-	VkMemoryAllocateInfo MemoryAllocateInfo;
-	ZeroVulkanStruct(MemoryAllocateInfo, VK_STRUCTURE_TYPE_MEMORY_ALLOCATE_INFO);
-	MemoryAllocateInfo.pNext = &MemoryAllocateFlagsInfo;
-	MemoryAllocateInfo.allocationSize = MemoryRequirements.size;
-	MemoryAllocateInfo.memoryTypeIndex = FindMemoryType(Gpu, MemoryRequirements.memoryTypeBits, MemoryFlags);
-	VERIFYVULKANRESULT(VulkanRHI::vkAllocateMemory(DeviceHandle, &MemoryAllocateInfo, VULKAN_CPU_ALLOCATOR, &Result.Memory));
-	VERIFYVULKANRESULT(VulkanRHI::vkBindBufferMemory(DeviceHandle, Result.Buffer, Result.Memory, 0));
-
-	Result.Device = DeviceHandle;
-}
-
-// Temporary brute force deallocation
-void FVulkanRayTracingAllocator::Free(FVkRtAllocation& Allocation)
-{
-	if (Allocation.Buffer != VK_NULL_HANDLE)
-	{
-		VulkanRHI::vkDestroyBuffer(Allocation.Device, Allocation.Buffer, VULKAN_CPU_ALLOCATOR);
-		Allocation.Buffer = VK_NULL_HANDLE;
-	}
-	if (Allocation.Memory != VK_NULL_HANDLE)
-	{
-		VulkanRHI::vkFreeMemory(Allocation.Device, Allocation.Memory, VULKAN_CPU_ALLOCATOR);
-		Allocation.Memory = VK_NULL_HANDLE;
-	}
-}
-
 static void AddAccelerationStructureBuildBarrier(VkCommandBuffer CommandBuffer)
 {
 	VkMemoryBarrier Barrier;
@@ -200,8 +161,6 @@ static void GetBLASBuildData(
 	const EBLASBuildDataUsage Usage,
 	FVkRtBLASBuildData& BuildData)
 {
-	static constexpr uint32 IndicesPerPrimitive = 3; // Only triangle meshes are supported
-
 	FVulkanResourceMultiBuffer* const IndexBuffer = ResourceCast(IndexBufferRHI.GetReference());
 	VkDeviceOrHostAddressConstKHR IndexBufferDeviceAddress = {};
 	
@@ -277,13 +236,13 @@ static void GetBLASBuildData(
 				{
 					SegmentGeometry.geometry.triangles.indexType = (IndexStrideInBytes == 2) ? VK_INDEX_TYPE_UINT16 : VK_INDEX_TYPE_UINT32;
 					// offset in bytes into the index buffer where primitive data for the current segment is defined
-					PrimitiveOffset = Segment.FirstPrimitive * IndicesPerPrimitive * IndexStrideInBytes;
+					PrimitiveOffset = Segment.FirstPrimitive * FVulkanRayTracingGeometry::IndicesPerPrimitive * IndexStrideInBytes;
 				}
 				else
 				{
 					SegmentGeometry.geometry.triangles.indexType = VK_INDEX_TYPE_NONE_KHR;
 					// for non-indexed geometry, primitiveOffset is applied when reading from vertex buffer
-					PrimitiveOffset = Segment.FirstPrimitive * IndicesPerPrimitive * Segment.VertexBufferStride;
+					PrimitiveOffset = Segment.FirstPrimitive * FVulkanRayTracingGeometry::IndicesPerPrimitive * Segment.VertexBufferStride;
 				}
 
 				break;
@@ -350,6 +309,9 @@ FVulkanRayTracingGeometry::FVulkanRayTracingGeometry(FRHICommandListBase& RHICmd
 {
 	INC_DWORD_STAT(STAT_VulkanRayTracingAllocatedBLAS);
 
+	DebugName = !Initializer.DebugName.IsNone() ? Initializer.DebugName : FDebugName(FName(TEXT("BLAS")));
+	OwnerName = Initializer.OwnerName;
+
 	uint32 IndexBufferStride = 0;
 	if (Initializer.IndexBuffer)
 	{
@@ -406,6 +368,8 @@ FVulkanRayTracingGeometry::FVulkanRayTracingGeometry(FRHICommandListBase& RHICmd
 
 FVulkanRayTracingGeometry::~FVulkanRayTracingGeometry()
 {
+	ReleaseBindlessHandles();
+
 	DEC_DWORD_STAT(STAT_VulkanRayTracingAllocatedBLAS);
 	DEC_DWORD_STAT_BY(STAT_VulkanRayTracingTrianglesBLAS, Initializer.TotalPrimitiveCount);
 	if (Handle != VK_NULL_HANDLE)
@@ -517,6 +481,69 @@ void FVulkanRayTracingGeometry::CompactAccelerationStructure(FVulkanCmdBuffer& C
 	Device->GetDeferredDeletionQueue().EnqueueResource(VulkanRHI::FDeferredDeletionQueue2::EType::AccelerationStructure, OldHandle);
 }
 
+
+void FVulkanRayTracingGeometry::SetupHitGroupSystemParameters()
+{
+	const bool bIsTriangles = (Initializer.GeometryType == RTGT_Triangles);
+
+	FVulkanBindlessDescriptorManager* BindlessDescriptorManager = Device->GetBindlessDescriptorManager();
+	auto GetBindlessHandle = [BindlessDescriptorManager](FVulkanResourceMultiBuffer* Buffer, uint32 ExtraOffset) 
+	{
+		if (Buffer)
+		{
+			FRHIDescriptorHandle BindlessHandle = BindlessDescriptorManager->ReserveDescriptor(VK_DESCRIPTOR_TYPE_STORAGE_BUFFER);
+			BindlessDescriptorManager->UpdateBuffer(BindlessHandle, Buffer->GetHandle(), Buffer->GetOffset() + ExtraOffset, Buffer->GetCurrentSize() - ExtraOffset);
+			return BindlessHandle;
+		}
+		return FRHIDescriptorHandle();
+	};
+
+	ReleaseBindlessHandles();
+
+	HitGroupSystemParameters.Reset(Initializer.Segments.Num());
+
+	FVulkanResourceMultiBuffer* IndexBuffer = ResourceCast(Initializer.IndexBuffer.GetReference());
+	const uint32 IndexStride = IndexBuffer ? IndexBuffer->GetStride() : 0;
+	HitGroupSystemIndexView = GetBindlessHandle(IndexBuffer, 0);
+
+	for (const FRayTracingGeometrySegment& Segment : Initializer.Segments)
+	{
+		FVulkanResourceMultiBuffer* VertexBuffer = ResourceCast(Segment.VertexBuffer.GetReference());
+		const FRHIDescriptorHandle VBHandle = GetBindlessHandle(VertexBuffer, Segment.VertexBufferOffset);
+		HitGroupSystemVertexViews.Add(VBHandle);
+
+		FVulkanHitGroupSystemParameters SystemParameters = {};
+		SystemParameters.RootConstants.SetVertexAndIndexStride(Segment.VertexBufferStride, IndexStride);
+		SystemParameters.BindlessHitGroupSystemVertexBuffer = VBHandle.GetIndex();
+
+		if (bIsTriangles && (IndexBuffer != nullptr))
+		{
+			SystemParameters.BindlessHitGroupSystemIndexBuffer = HitGroupSystemIndexView.GetIndex();
+			SystemParameters.RootConstants.IndexBufferOffsetInBytes = Initializer.IndexBufferOffset + IndexStride * Segment.FirstPrimitive * FVulkanRayTracingGeometry::IndicesPerPrimitive;
+			SystemParameters.RootConstants.FirstPrimitive = Segment.FirstPrimitive;
+		}
+
+		HitGroupSystemParameters.Add(SystemParameters);
+	}
+}
+
+void FVulkanRayTracingGeometry::ReleaseBindlessHandles()
+{
+	FVulkanBindlessDescriptorManager* BindlessDescriptorManager = Device->GetBindlessDescriptorManager();
+
+	for (FRHIDescriptorHandle BindlesHandle : HitGroupSystemVertexViews)
+	{
+		BindlessDescriptorManager->Unregister(BindlesHandle);
+	}
+	HitGroupSystemVertexViews.Reset(Initializer.Segments.Num());
+
+	if (HitGroupSystemIndexView.IsValid())
+	{
+		BindlessDescriptorManager->Unregister(HitGroupSystemIndexView);
+		HitGroupSystemIndexView = FRHIDescriptorHandle();
+	}
+}
+
 static void GetTLASBuildData(
 	const VkDevice Device,
 	const uint32 NumInstances,
@@ -596,7 +623,9 @@ struct FVulkanRayTracingGeometryParameters
 };
 
 FVulkanRayTracingScene::FVulkanRayTracingScene(FRayTracingSceneInitializer2 InInitializer, FVulkanDevice* InDevice)
-	: Device(InDevice), Initializer(MoveTemp(InInitializer))
+	: FDeviceChild(InDevice)
+	, Initializer(MoveTemp(InInitializer))
+
 {
 	INC_DWORD_STAT(STAT_VulkanRayTracingAllocatedTLAS);
 
@@ -751,6 +780,8 @@ void FVulkanRayTracingScene::BuildAccelerationStructure(
 
 	CommandBufferManager.SubmitActiveCmdBuffer();
 	CommandBufferManager.PrepareForNewActiveCommandBuffer();
+
+	bBuilt = true;
 }
 
 void FVulkanRayTracingScene::BuildPerInstanceGeometryParameterBuffer(FVulkanCommandListContext& CommandContext)
@@ -873,6 +904,11 @@ FRayTracingGeometryRHIRef FVulkanDynamicRHI::RHICreateRayTracingGeometry(FRHICom
 	return new FVulkanRayTracingGeometry(RHICmdList, Initializer, GetDevice());
 }
 
+FRayTracingPipelineStateRHIRef FVulkanDynamicRHI::RHICreateRayTracingPipelineState(const FRayTracingPipelineStateInitializer& Initializer)
+{
+	return new FVulkanRayTracingPipelineState(GetDevice(), Initializer);
+}
+
 void FVulkanCommandListContext::RHIClearRayTracingBindings(FRHIRayTracingScene* Scene)
 {
 	 // TODO
@@ -979,6 +1015,8 @@ void FVulkanCommandListContext::RHIBuildAccelerationStructures(const TArrayView<
 
 		BuildGeometryInfos.Add(BuildData.GeometryInfo);
 		BuildRangeInfos.Add(pBuildRanges);
+
+		Geometry->SetupHitGroupSystemParameters();
 	}
 	
 	FVulkanCmdBuffer* const CmdBuffer = CommandBufferManager->GetActiveCmdBuffer();
@@ -1033,25 +1071,23 @@ void FVulkanDevice::CleanUpRayTracing()
 {
 }
 
-static uint32 GetAlignedSize(uint32 Value, uint32 Alignment)
-{
-	return (Value + Alignment - 1) & ~(Alignment - 1);
-}
 
 FVulkanRayTracingPipelineState::FVulkanRayTracingPipelineState(FVulkanDevice* const InDevice, const FRayTracingPipelineStateInitializer& Initializer)
+	: FDeviceChild(InDevice)
 {
 	checkf(InDevice->SupportsBindless(), TEXT("Vulkan ray tracing pipelines are only supported in bindless."));
 
 	TArrayView<FRHIRayTracingShader*> InitializerRayGenShaders = Initializer.GetRayGenTable();
 	TArrayView<FRHIRayTracingShader*> InitializerMissShaders = Initializer.GetMissTable();
 	TArrayView<FRHIRayTracingShader*> InitializerHitGroupShaders = Initializer.GetHitGroupTable();
-	// vkrt todo: Callable shader support
+	TArrayView<FRHIRayTracingShader*> InitializerCallableShaders = Initializer.GetCallableTable();
 
 	TArray<VkPipelineShaderStageCreateInfo> ShaderStages;
 	TArray<VkRayTracingShaderGroupCreateInfoKHR> ShaderGroups;
 	TArray<ANSICHAR*> EntryPointNames;
 	const uint32 EntryPointNameMaxLength = 24;
 
+	RayGen.Shaders.Reserve(InitializerRayGenShaders.Num());
 	for (FRHIRayTracingShader* const RayGenShaderRHI : InitializerRayGenShaders)
 	{
 		checkSlow(RayGenShaderRHI->GetFrequency() == SF_RayGen);
@@ -1076,8 +1112,11 @@ FVulkanRayTracingPipelineState::FVulkanRayTracingPipelineState(FVulkanDevice* co
 		ShaderGroup.anyHitShader = VK_SHADER_UNUSED_KHR;
 		ShaderGroup.intersectionShader = VK_SHADER_UNUSED_KHR;
 		ShaderGroups.Add(ShaderGroup);
+
+		RayGen.Shaders.Add(RayGenShader);
 	}
 
+	Miss.Shaders.Reserve(InitializerMissShaders.Num());
 	for (FRHIRayTracingShader* const MissShaderRHI : InitializerMissShaders)
 	{
 		checkSlow(MissShaderRHI->GetFrequency() == SF_RayMiss);
@@ -1102,8 +1141,11 @@ FVulkanRayTracingPipelineState::FVulkanRayTracingPipelineState(FVulkanDevice* co
 		ShaderGroup.anyHitShader = VK_SHADER_UNUSED_KHR;
 		ShaderGroup.intersectionShader = VK_SHADER_UNUSED_KHR;
 		ShaderGroups.Add(ShaderGroup);
+
+		Miss.Shaders.Add(MissShader);
 	}
 
+	HitGroup.Shaders.Reserve(InitializerHitGroupShaders.Num());
 	for (FRHIRayTracingShader* const HitGroupShaderRHI : InitializerHitGroupShaders)
 	{
 		checkSlow(HitGroupShaderRHI->GetFrequency() == SF_RayHitGroup);
@@ -1128,6 +1170,37 @@ FVulkanRayTracingPipelineState::FVulkanRayTracingPipelineState(FVulkanDevice* co
 		ShaderGroup.anyHitShader = VK_SHADER_UNUSED_KHR; // vkrt: todo
 		ShaderGroup.intersectionShader = VK_SHADER_UNUSED_KHR;
 		ShaderGroups.Add(ShaderGroup);
+
+		HitGroup.Shaders.Add(HitGroupShader);
+	}
+
+	Callable.Shaders.Reserve(InitializerCallableShaders.Num());
+	for (FRHIRayTracingShader* const CallableShaderRHI : InitializerCallableShaders)
+	{
+		checkSlow(CallableShaderRHI->GetFrequency() == SF_RayCallable);
+		FVulkanRayTracingShader* const CallableShader = ResourceCast(CallableShaderRHI);
+
+		VkPipelineShaderStageCreateInfo ShaderStage;
+		ZeroVulkanStruct(ShaderStage, VK_STRUCTURE_TYPE_PIPELINE_SHADER_STAGE_CREATE_INFO);
+		ShaderStage.module = CallableShader->GetOrCreateHandle()->GetVkShaderModule();
+		ShaderStage.stage = VK_SHADER_STAGE_CALLABLE_BIT_KHR;
+
+		ANSICHAR* const EntryPoint = new char[EntryPointNameMaxLength];
+		CallableShader->GetEntryPoint(EntryPoint, EntryPointNameMaxLength);
+		EntryPointNames.Add(EntryPoint);
+		ShaderStage.pName = EntryPoint;
+		ShaderStages.Add(ShaderStage);
+
+		VkRayTracingShaderGroupCreateInfoKHR ShaderGroup;
+		ZeroVulkanStruct(ShaderGroup, VK_STRUCTURE_TYPE_RAY_TRACING_SHADER_GROUP_CREATE_INFO_KHR);
+		ShaderGroup.type = VK_RAY_TRACING_SHADER_GROUP_TYPE_GENERAL_KHR;
+		ShaderGroup.generalShader = ShaderStages.Num() - 1;
+		ShaderGroup.closestHitShader = VK_SHADER_UNUSED_KHR;
+		ShaderGroup.anyHitShader = VK_SHADER_UNUSED_KHR;
+		ShaderGroup.intersectionShader = VK_SHADER_UNUSED_KHR;
+		ShaderGroups.Add(ShaderGroup);
+
+		Callable.Shaders.Add(CallableShader);
 	}
 
 	VkRayTracingPipelineCreateInfoKHR RayTracingPipelineCreateInfo;
@@ -1154,46 +1227,98 @@ FVulkanRayTracingPipelineState::FVulkanRayTracingPipelineState(FVulkanDevice* co
 		delete[] EntryPoint;
 	}
 
-	const VkPhysicalDeviceRayTracingPipelinePropertiesKHR& RayTracingPipelineProps = InDevice->GetOptionalExtensionProperties().RayTracingPipelineProps;
-	const uint32 HandleSize = RayTracingPipelineProps.shaderGroupHandleSize;
-	const uint32 HandleSizeAligned = GetAlignedSize(HandleSize, RayTracingPipelineProps.shaderGroupHandleAlignment);
-	const uint32 GroupCount = ShaderGroups.Num();
-	const uint32 SBTSize = GroupCount * HandleSizeAligned;
+	// Grab all shader handles for each stage
+	{
+		const VkPhysicalDeviceRayTracingPipelinePropertiesKHR& RayTracingPipelineProps = InDevice->GetOptionalExtensionProperties().RayTracingPipelineProps;
+		const uint32 HandleSize = RayTracingPipelineProps.shaderGroupHandleSize;
+
+		uint32 HandleOffset = 0;
+		auto FetchShaderHandles = [&HandleOffset, HandleSize, InDevice](VkPipeline RTPipeline, uint32 HandleCount) 
+		{
+			TArray<uint8> OutHandleStorage;
+
+			if (HandleCount)
+			{
+				const uint32 ShaderHandleStorageSize = HandleCount * HandleSize;
+				OutHandleStorage.AddUninitialized(ShaderHandleStorageSize);
+
+				VERIFYVULKANRESULT(VulkanDynamicAPI::vkGetRayTracingShaderGroupHandlesKHR(InDevice->GetInstanceHandle(), RTPipeline, HandleOffset, HandleCount, ShaderHandleStorageSize, OutHandleStorage.GetData()));
+
+				HandleOffset += HandleCount;
+			}
+
+			return OutHandleStorage;
+		};
+
+		// NOTE: Must be filled in the same order as created above
+		RayGen.ShaderHandles = FetchShaderHandles(Pipeline, InitializerRayGenShaders.Num());
+		Miss.ShaderHandles = FetchShaderHandles(Pipeline, InitializerMissShaders.Num());
+		HitGroup.ShaderHandles = FetchShaderHandles(Pipeline, InitializerHitGroupShaders.Num());
+		Callable.ShaderHandles = FetchShaderHandles(Pipeline, InitializerCallableShaders.Num());
+	}
+
+	// If no custom hit groups were provided, then disable SBT indexing and force default shader on all primitives
+	bAllowHitGroupIndexing = Initializer.GetHitGroupTable().Num() ? Initializer.bAllowHitGroupIndexing : false;
 
-	TArray<uint8> ShaderHandleStorage;
-	ShaderHandleStorage.AddUninitialized(SBTSize);
-	VERIFYVULKANRESULT(VulkanDynamicAPI::vkGetRayTracingShaderGroupHandlesKHR(InDevice->GetInstanceHandle(), Pipeline, 0, GroupCount, SBTSize, ShaderHandleStorage.GetData()));
+	INC_DWORD_STAT(STAT_VulkanRayTracingCreatedPipelines);
+	INC_DWORD_STAT_BY(STAT_VulkanRayTracingCompiledShaders, 1);
+}
 
-	auto CopyHandlesToSBT = [InDevice, HandleSize, ShaderHandleStorage](FVkRtAllocation& Allocation, uint32 Offset)
+FVulkanRayTracingPipelineState::~FVulkanRayTracingPipelineState()
+{
+	if (Pipeline != VK_NULL_HANDLE)
+	{
+		VulkanRHI::vkDestroyPipeline(Device->GetInstanceHandle(), Pipeline, VULKAN_CPU_ALLOCATOR);
+		Pipeline = VK_NULL_HANDLE;
+	}
+}
+
+const FVulkanRayTracingPipelineState::ShaderData& FVulkanRayTracingPipelineState::GetShaderData(EShaderFrequency Frequency) const
+{
+	switch (Frequency)
 	{
-		FVulkanRayTracingAllocator::Allocate(
-			InDevice,
-			HandleSize,
-			VK_BUFFER_USAGE_SHADER_BINDING_TABLE_BIT_KHR | VK_BUFFER_USAGE_SHADER_DEVICE_ADDRESS_BIT,
-			VK_MEMORY_PROPERTY_HOST_VISIBLE_BIT | VK_MEMORY_PROPERTY_HOST_COHERENT_BIT,
-			Allocation);
-
-		void* pMappedBufferMemory = nullptr;
-		VERIFYVULKANRESULT(VulkanRHI::vkMapMemory(InDevice->GetInstanceHandle(), Allocation.Memory, 0, VK_WHOLE_SIZE, 0, &pMappedBufferMemory));
+	case SF_RayGen:
+		return RayGen;
+	case SF_RayMiss:
+		return Miss;
+	case SF_RayHitGroup:
+		return HitGroup;
+	case SF_RayCallable:
+		return Callable;
+
+	default:
+		checkf(false, TEXT("Only usable with RayTracing shaders."));
+	};
+
+	static ShaderData EmptyShaderData;
+	return EmptyShaderData;
+}
+
+int32 FVulkanRayTracingPipelineState::GetShaderIndex(const FVulkanRayTracingShader* Shader) const
+{
+	const FSHAHash Hash = Shader->GetHash();
+
+	const TArray<TRefCountPtr<FVulkanRayTracingShader>>& ShaderArray = GetShaderData(Shader->GetFrequency()).Shaders;
+	for (int32 Index = 0; Index < ShaderArray.Num(); ++Index)
+	{
+		if (Hash == ShaderArray[Index]->GetHash())
 		{
-			FMemory::Memcpy(pMappedBufferMemory, ShaderHandleStorage.GetData() + Offset, HandleSize);
+			return Index;
 		}
-		VulkanRHI::vkUnmapMemory(InDevice->GetInstanceHandle(), Allocation.Memory);
-	};
+	}
 
-	INC_DWORD_STAT(STAT_VulkanRayTracingCreatedPipelines);
-	INC_DWORD_STAT_BY(STAT_VulkanRayTracingCompiledShaders, 1);
+	checkf(false, TEXT("RayTracing shader is not present in the given ray tracing pipeline. "));
+	return INDEX_NONE;
+}
 
-	CopyHandlesToSBT(RayGenShaderBindingTable, 0);
-	CopyHandlesToSBT(MissShaderBindingTable, HandleSizeAligned);
-	CopyHandlesToSBT(HitShaderBindingTable, HandleSizeAligned * 2);
+const FVulkanRayTracingShader* FVulkanRayTracingPipelineState::GetShader(EShaderFrequency Frequency, int32 ShaderIndex) const
+{
+	return GetShaderData(Frequency).Shaders[ShaderIndex].GetReference();
 }
 
-FVulkanRayTracingPipelineState::~FVulkanRayTracingPipelineState()
+const TArray<uint8>& FVulkanRayTracingPipelineState::GetShaderHandles(EShaderFrequency Frequency) const
 {
-	FVulkanRayTracingAllocator::Free(RayGenShaderBindingTable);
-	FVulkanRayTracingAllocator::Free(MissShaderBindingTable);
-	FVulkanRayTracingAllocator::Free(HitShaderBindingTable);
+	return GetShaderData(Frequency).ShaderHandles;
 }
 
 void FVulkanRayTracingCompactedSizeQueryPool::EndBatch(FVulkanCmdBuffer* InCmdBuffer)
@@ -1342,4 +1467,65 @@ void FVulkanRayTracingCompactionRequestHandler::Update(FVulkanCommandListContext
 	}
 }
 
+void FVulkanCommandListContext::RHIRayTraceDispatch(
+	FRHIRayTracingPipelineState* InRayTracingPipelineState, 
+	FRHIRayTracingShader* InRayGenShader,
+	FRHIRayTracingScene* InScene,
+	const FRayTracingShaderBindings& InGlobalResourceBindings, // :todo-jn:
+	uint32 InWidth, uint32 InHeight)
+{
+	checkNoEntry();
+}
+
+void FVulkanCommandListContext::RHIRayTraceDispatchIndirect(
+	FRHIRayTracingPipelineState* InRayTracingPipelineState, 
+	FRHIRayTracingShader* InRayGenShader,
+	FRHIRayTracingScene* InScene,
+	const FRayTracingShaderBindings& InGlobalResourceBindings, // :todo-jn:
+	FRHIBuffer* InArgumentBuffer, uint32 InArgumentOffset)
+{
+	checkNoEntry();
+}
+
+void FVulkanCommandListContext::RHISetRayTracingHitGroup(
+	FRHIRayTracingScene* InScene, uint32 InstanceIndex, uint32 SegmentIndex, uint32 ShaderSlot,
+	FRHIRayTracingPipelineState* InPipeline, uint32 HitGroupIndex,
+	uint32 NumUniformBuffers, FRHIUniformBuffer* const* UniformBuffers,
+	uint32 LooseParameterDataSize, const void* LooseParameterData,
+	uint32 UserData)
+{
+	checkNoEntry();
+}
+
+void FVulkanCommandListContext::RHISetRayTracingCallableShader(
+	FRHIRayTracingScene* InScene, uint32 ShaderSlotInScene,
+	FRHIRayTracingPipelineState* InPipeline, uint32 ShaderIndexInPipeline,
+	uint32 NumUniformBuffers, FRHIUniformBuffer* const* UniformBuffers,
+	uint32 UserData)
+{
+	checkNoEntry();
+}
+
+void FVulkanCommandListContext::RHISetRayTracingMissShader(
+	FRHIRayTracingScene* InScene, uint32 ShaderSlotInScene,
+	FRHIRayTracingPipelineState* InPipeline, uint32 ShaderIndexInPipeline,
+	uint32 NumUniformBuffers, FRHIUniformBuffer* const* UniformBuffers,
+	uint32 UserData)
+{
+	checkNoEntry();
+}
+
+
+void FVulkanCommandListContext::RHISetRayTracingBindings(
+	FRHIRayTracingScene* InScene, FRHIRayTracingPipelineState* InPipeline,
+	uint32 NumBindings, const FRayTracingLocalShaderBindings* Bindings,
+	ERayTracingBindingType BindingType)
+{
+	FVulkanRayTracingScene* Scene = ResourceCast(InScene);
+
+	checkf(Scene->IsBuilt(), TEXT("Ray tracing scene must be built before any shaders can be bound to it. Make sure that RHIBuildAccelerationStructure() command has been executed."));
+
+	checkNoEntry();
+}
+
 #endif // VULKAN_RHI_RAYTRACING
diff --git a/Engine/Source/Runtime/VulkanRHI/Private/VulkanRayTracing.h b/Engine/Source/Runtime/VulkanRHI/Private/VulkanRayTracing.h
--- a/Engine/Source/Runtime/VulkanRHI/Private/VulkanRayTracing.h
+++ b/Engine/Source/Runtime/VulkanRHI/Private/VulkanRayTracing.h
@@ -6,6 +6,8 @@
 
 #if VULKAN_RHI_RAYTRACING
 
+#include "RayTracingBuiltInResources.h"
+
 class FVulkanCommandListContext;
 class FVulkanResourceMultiBuffer;
 
@@ -15,20 +17,20 @@ public:
 	static bool CheckVulkanInstanceFunctions(VkInstance inInstance);
 };
 
-struct FVkRtAllocation
-{
-	VkDevice Device = VK_NULL_HANDLE;
-	VkDeviceMemory Memory = VK_NULL_HANDLE;
-	VkBuffer Buffer = VK_NULL_HANDLE;
-};
 
-class FVulkanRayTracingAllocator
+// Built-in local root parameters that are always bound to all hit shaders
+// :todo-jn: NOTE: Keep in sync with VulkanCommon.ush decl until it's put in a common header
+struct FVulkanHitGroupSystemParameters
 {
-public:
-	static void Allocate(FVulkanDevice* Device, VkDeviceSize Size, VkBufferUsageFlags UsageFlags, VkMemoryPropertyFlags MemoryFlags, FVkRtAllocation& Result);
-	static void Free(FVkRtAllocation& Allocation);
+	FHitGroupSystemRootConstants RootConstants;
+
+	uint32 BindlessHitGroupSystemIndexBuffer;
+	uint32 BindlessHitGroupSystemVertexBuffer;
+
+	uint32 BindlessUniformBuffers[16];
 };
 
+
 struct FVkRtTLASBuildData
 {
 	FVkRtTLASBuildData()
@@ -60,6 +62,8 @@ struct FVkRtBLASBuildData
 class FVulkanRayTracingGeometry : public FRHIRayTracingGeometry
 {
 public:
+	static constexpr uint32 IndicesPerPrimitive = 3; // Only triangle meshes are supported
+
 	FVulkanRayTracingGeometry(ENoInit);
 	FVulkanRayTracingGeometry(FRHICommandListBase& RHICmdList, const FRayTracingGeometryInitializer& Initializer, FVulkanDevice* InDevice);
 	~FVulkanRayTracingGeometry();
@@ -75,16 +79,26 @@ public:
 	void RemoveCompactionRequest();
 	void CompactAccelerationStructure(FVulkanCmdBuffer& CmdBuffer, uint64 InSizeAfterCompaction);
 
+	void SetupHitGroupSystemParameters();
+	void ReleaseBindlessHandles();
+
 	VkAccelerationStructureKHR Handle = VK_NULL_HANDLE;
 	VkDeviceAddress Address = 0;
 	TRefCountPtr<FVulkanResourceMultiBuffer> AccelerationStructureBuffer;
 	bool bHasPendingCompactionRequests = false;
 	uint64 AccelerationStructureCompactedSize = 0;
-private:
+
 	FVulkanDevice* const Device = nullptr;
+
+	TArray<FVulkanHitGroupSystemParameters> HitGroupSystemParameters;
+	TArray<FRHIDescriptorHandle> HitGroupSystemVertexViews;  // :todo-jn: use views?
+	FRHIDescriptorHandle HitGroupSystemIndexView;
+
+	FDebugName DebugName;
+	FName OwnerName;		// Store the path name of the owner object for resource tracking
 };
 
-class FVulkanRayTracingScene : public FRHIRayTracingScene
+class FVulkanRayTracingScene : public FRHIRayTracingScene, public VulkanRHI::FDeviceChild
 {
 public:
 	FVulkanRayTracingScene(FRayTracingSceneInitializer2 Initializer, FVulkanDevice* InDevice);
@@ -109,9 +123,12 @@ public:
 		return PerInstanceGeometryParameterSRV.GetReference();
 	}
 
-private:
-	FVulkanDevice* const Device = nullptr;
+	inline bool IsBuilt() const
+	{
+		return bBuilt;
+	}
 
+private:
 	const FRayTracingSceneInitializer2 Initializer;
 
 	// Native TLAS handles are owned by SRV objects in Vulkan RHI.
@@ -137,9 +154,12 @@ private:
 	TRefCountPtr<FVulkanResourceMultiBuffer> PerInstanceGeometryParameterBuffer;
 	FShaderResourceViewRHIRef PerInstanceGeometryParameterSRV;
 	void BuildPerInstanceGeometryParameterBuffer(FVulkanCommandListContext& CommandContext);
+
+	bool bBuilt = false;
 };
 
-class FVulkanRayTracingPipelineState : public FRHIRayTracingPipelineState
+
+class FVulkanRayTracingPipelineState : public FRHIRayTracingPipelineState, public VulkanRHI::FDeviceChild
 {
 public:
 
@@ -147,11 +167,37 @@ public:
 	FVulkanRayTracingPipelineState(FVulkanDevice* const InDevice, const FRayTracingPipelineStateInitializer& Initializer);
 	~FVulkanRayTracingPipelineState();
 
+	inline VkPipeline GetPipeline() const
+	{
+		return Pipeline;
+	}
+
+	int32 GetShaderIndex(const FVulkanRayTracingShader* Shader) const;
+	const FVulkanRayTracingShader* GetShader(EShaderFrequency Frequency, int32 ShaderIndex) const;
+	const TArray<uint8>& GetShaderHandles(EShaderFrequency Frequency) const;
+
 private:
+
+	struct ShaderData
+	{
+		TArray<TRefCountPtr<FVulkanRayTracingShader>> Shaders;
+		TArray<uint8> ShaderHandles;
+	};
+
+	const ShaderData& GetShaderData(EShaderFrequency Frequency) const;
+
+	ShaderData RayGen;
+	ShaderData Miss;
+	ShaderData HitGroup;
+	ShaderData Callable;
+
 	VkPipeline Pipeline = VK_NULL_HANDLE;
-	FVkRtAllocation RayGenShaderBindingTable;
-	FVkRtAllocation MissShaderBindingTable;
-	FVkRtAllocation HitShaderBindingTable;
+
+public:
+	bool bAllowHitGroupIndexing = true;
+
+	friend FVulkanCommandListContext;
+	friend FVulkanRayTracingScene;
 };
 
 class FVulkanRayTracingCompactedSizeQueryPool : public FVulkanQueryPool
diff --git a/Engine/Source/Runtime/VulkanRHI/Public/VulkanDynamicRHI.h b/Engine/Source/Runtime/VulkanRHI/Public/VulkanDynamicRHI.h
--- a/Engine/Source/Runtime/VulkanRHI/Public/VulkanDynamicRHI.h
+++ b/Engine/Source/Runtime/VulkanRHI/Public/VulkanDynamicRHI.h
@@ -234,6 +234,7 @@ public:
 	virtual FRayTracingGeometryRHIRef RHICreateRayTracingGeometry(FRHICommandListBase& RHICmdList, const FRayTracingGeometryInitializer& Initializer) final override;
 	virtual FRayTracingSceneRHIRef RHICreateRayTracingScene(FRayTracingSceneInitializer2 Initializer) final override;
 	virtual FRayTracingShaderRHIRef RHICreateRayTracingShader(TArrayView<const uint8> Code, const FSHAHash& Hash, EShaderFrequency ShaderFrequency) final override;
+	virtual FRayTracingPipelineStateRHIRef RHICreateRayTracingPipelineState(const FRayTracingPipelineStateInitializer& Initializer) final override;
 	virtual void RHITransferRayTracingGeometryUnderlyingResource(FRHIRayTracingGeometry* DestGeometry, FRHIRayTracingGeometry* SrcGeometry) final override;
 #endif
 
