VulkanRHI: Always enable VK_KHR_buffer_device_address if it's available (not only when raytracing).

#rnx

diff --git a/Engine/Source/Runtime/VulkanRHI/Private/VulkanExtensions.cpp b/Engine/Source/Runtime/VulkanRHI/Private/VulkanExtensions.cpp
--- a/Engine/Source/Runtime/VulkanRHI/Private/VulkanExtensions.cpp
+++ b/Engine/Source/Runtime/VulkanRHI/Private/VulkanExtensions.cpp
@@ -670,7 +670,7 @@ class FVulkanKHRBufferDeviceAddressExtension : public FVulkanDeviceExtension
 public:
 
 	FVulkanKHRBufferDeviceAddressExtension(FVulkanDevice* InDevice)
-		: FVulkanDeviceExtension(InDevice, VK_KHR_BUFFER_DEVICE_ADDRESS_EXTENSION_NAME, VULKAN_RHI_RAYTRACING, VK_API_VERSION_1_2) {}
+		: FVulkanDeviceExtension(InDevice, VK_KHR_BUFFER_DEVICE_ADDRESS_EXTENSION_NAME, VULKAN_EXTENSION_ENABLED, VK_API_VERSION_1_2) {}
 
 	virtual void PrePhysicalDeviceFeatures(VkPhysicalDeviceFeatures2KHR& PhysicalDeviceFeatures2) override final
 	{
