VulkanRHI: VKRT WIP: Add basic shader table with shader handles and shader records to Vulkan ray tracing.

#rnx

diff --git a/Engine/Source/Runtime/VulkanRHI/Private/VulkanRayTracing.cpp b/Engine/Source/Runtime/VulkanRHI/Private/VulkanRayTracing.cpp
--- a/Engine/Source/Runtime/VulkanRHI/Private/VulkanRayTracing.cpp
+++ b/Engine/Source/Runtime/VulkanRHI/Private/VulkanRayTracing.cpp
@@ -660,6 +660,11 @@ FVulkanRayTracingScene::FVulkanRayTracingScene(FRayTracingSceneInitializer2 InIn
 
 FVulkanRayTracingScene::~FVulkanRayTracingScene()
 {
+	for (auto& Item : ShaderTables)
+	{
+		delete Item.Value;
+	}
+
 	if (AccelerationStructureBuffer)
 	{
 		DEC_MEMORY_STAT_BY(STAT_VulkanRayTracingUsedVideoMemory, AccelerationStructureBuffer->GetSize());
@@ -841,6 +846,142 @@ void FVulkanRayTracingScene::BuildPerInstanceGeometryParameterBuffer(FVulkanComm
 	PerInstanceGeometryParameterBuffer->Unlock(CommandContext);
 }
 
+
+FVulkanRayTracingShaderTable::FVulkanRayTracingShaderTable(FVulkanDevice* Device)
+	: FDeviceChild(Device)
+	, HandleSize(Device->GetOptionalExtensionProperties().RayTracingPipelineProps.shaderGroupHandleSize)
+	, HandleSizeAligned(Align(HandleSize, Device->GetOptionalExtensionProperties().RayTracingPipelineProps.shaderGroupHandleAlignment))
+{
+}
+
+FVulkanRayTracingShaderTable::~FVulkanRayTracingShaderTable()
+{
+	auto FreeAlloc = [Device = Device](FVulkanShaderTableAllocation& Alloc) {
+		if (Alloc.Buffer != VK_NULL_HANDLE)
+		{
+			VulkanRHI::vkDestroyBuffer(Device->GetInstanceHandle(), Alloc.Buffer, VULKAN_CPU_ALLOCATOR);
+			Alloc.Buffer = VK_NULL_HANDLE;
+		}
+
+		if (Alloc.Allocation.IsValid())
+		{
+			Device->GetMemoryManager().FreeVulkanAllocation(Alloc.Allocation);
+		}
+	};
+
+	FreeAlloc(Raygen);
+	FreeAlloc(Miss);
+	FreeAlloc(HitGroup);
+	FreeAlloc(Callable);
+}
+
+void FVulkanRayTracingShaderTable::Init(const FVulkanRayTracingScene* Scene, const FVulkanRayTracingPipelineState* Pipeline)
+{
+	auto InitAlloc = [Device = Device, HandleSize = HandleSize, HandleSizeAligned = HandleSizeAligned](FVulkanShaderTableAllocation& Alloc, uint32 InHandleCount, bool InUseLocalRecord) {
+
+		Alloc.HandleCount = InHandleCount;
+		Alloc.bUseLocalRecord = InUseLocalRecord;
+
+		if (Alloc.HandleCount > 0)
+		{
+			VkDevice DeviceHandle = Device->GetInstanceHandle();
+			VkPhysicalDevice Gpu = Device->GetPhysicalHandle();
+
+			const VkPhysicalDeviceRayTracingPipelinePropertiesKHR& RayTracingPipelineProps = Device->GetOptionalExtensionProperties().RayTracingPipelineProps;
+			Alloc.Region.stride = InUseLocalRecord ? FMath::Min<VkDeviceSize>(RayTracingPipelineProps.maxShaderGroupStride, 4096) : HandleSizeAligned; // :todo-jn: shrink stride to necessary amount
+			Alloc.Region.size = Alloc.HandleCount * Alloc.Region.stride;
+
+			{
+				VkBufferCreateInfo BufferCreateInfo;
+				ZeroVulkanStruct(BufferCreateInfo, VK_STRUCTURE_TYPE_BUFFER_CREATE_INFO);
+				BufferCreateInfo.size = Alloc.Region.size;
+				BufferCreateInfo.usage = VK_BUFFER_USAGE_SHADER_BINDING_TABLE_BIT_KHR | VK_BUFFER_USAGE_SHADER_DEVICE_ADDRESS_BIT;;
+				VERIFYVULKANRESULT(VulkanRHI::vkCreateBuffer(DeviceHandle, &BufferCreateInfo, VULKAN_CPU_ALLOCATOR, &Alloc.Buffer));
+			}
+
+			const EVulkanAllocationFlags AllocFlags = EVulkanAllocationFlags::HostVisible | EVulkanAllocationFlags::AutoBind | EVulkanAllocationFlags::Dedicated; // :todo-jn: Dedicated for now... preferBAR?
+			Device->GetMemoryManager().AllocateBufferMemory(Alloc.Allocation, Alloc.Buffer, AllocFlags, TEXT("FVulkanShaderTableAllocation"), RayTracingPipelineProps.shaderGroupBaseAlignment);
+
+			VkMemoryAllocateFlagsInfo MemoryAllocateFlagsInfo;
+			ZeroVulkanStruct(MemoryAllocateFlagsInfo, VK_STRUCTURE_TYPE_MEMORY_ALLOCATE_FLAGS_INFO);
+			MemoryAllocateFlagsInfo.flags = VK_MEMORY_ALLOCATE_DEVICE_ADDRESS_BIT_KHR;
+
+			VkBufferDeviceAddressInfoKHR DeviceAddressInfo;
+			ZeroVulkanStruct(DeviceAddressInfo, VK_STRUCTURE_TYPE_BUFFER_DEVICE_ADDRESS_INFO);
+			DeviceAddressInfo.buffer = Alloc.Buffer;
+			Alloc.Region.deviceAddress = VulkanRHI::vkGetBufferDeviceAddressKHR(DeviceHandle, &DeviceAddressInfo);
+
+			Alloc.MappedBufferMemory = (uint8*)Alloc.Allocation.GetMappedPointer(Device);
+		}
+	};
+
+	const FRayTracingSceneInitializer2& SceneInitializer = Scene->GetInitializer();
+
+	InitAlloc(Raygen, 1, false);
+	InitAlloc(Miss, SceneInitializer.NumMissShaderSlots, true);
+	InitAlloc(HitGroup, Pipeline->bAllowHitGroupIndexing ? SceneInitializer.NumTotalSegments * SceneInitializer.ShaderSlotsPerGeometrySegment : 1, true);
+	InitAlloc(Callable, SceneInitializer.NumCallableShaderSlots, true);
+}
+
+FVulkanRayTracingShaderTable::FVulkanShaderTableAllocation& FVulkanRayTracingShaderTable::GetAlloc(EShaderFrequency Frequency)
+{
+	switch (Frequency)
+	{
+	case SF_RayGen:       return Raygen;
+	case SF_RayMiss:      return Miss;
+	case SF_RayHitGroup:  return HitGroup;
+	case SF_RayCallable:  return Callable;
+
+	default:
+		checkf(false, TEXT("Only usable with RayTracing shaders."));
+		break;
+	}
+
+	static FVulkanShaderTableAllocation EmptyAlloc;
+	return EmptyAlloc;
+}
+
+const VkStridedDeviceAddressRegionKHR* FVulkanRayTracingShaderTable::GetRegion(EShaderFrequency Frequency)
+{
+	return &GetAlloc(Frequency).Region;
+}
+
+void FVulkanRayTracingShaderTable::SetSlot(EShaderFrequency Frequency, uint32 DstSlot, uint32 SrcHandleIndex, TConstArrayView<uint8> SrcHandleData)
+{
+	FVulkanShaderTableAllocation& Alloc = GetAlloc(Frequency);
+	FMemory::Memcpy(&Alloc.MappedBufferMemory[DstSlot * Alloc.Region.stride], &SrcHandleData[SrcHandleIndex * HandleSize], HandleSize);
+}
+
+void FVulkanRayTracingShaderTable::SetLocalShaderParameters(EShaderFrequency Frequency, uint32 RecordIndex, uint32 OffsetWithinRecord, const void* InData, uint32 InDataSize)
+{
+	FVulkanShaderTableAllocation& Alloc = GetAlloc(Frequency);
+
+	checkfSlow(OffsetWithinRecord % 4 == 0, TEXT("SBT record parameters must be written on DWORD-aligned boundary"));
+	checkfSlow(InDataSize % 4 == 0, TEXT("SBT record parameters must be DWORD-aligned"));
+	checkf(OffsetWithinRecord + InDataSize <= Alloc.Region.size, TEXT("SBT record write request is out of bounds"));
+
+	const uint32 WriteOffset = HandleSizeAligned + (Alloc.Region.stride * RecordIndex) + OffsetWithinRecord;
+	FMemory::Memcpy(Alloc.MappedBufferMemory + WriteOffset, InData, InDataSize);
+}
+
+FVulkanRayTracingShaderTable* FVulkanRayTracingScene::FindOrCreateShaderTable(const FVulkanRayTracingPipelineState* Pipeline)
+{
+	// Find existing table
+	{
+		FVulkanRayTracingShaderTable* const* FoundShaderTable = ShaderTables.Find(Pipeline);
+		if (FoundShaderTable)
+		{
+			return *FoundShaderTable;
+		}
+	}
+
+	// Create new table
+	FVulkanRayTracingShaderTable* CreatedShaderTable = new FVulkanRayTracingShaderTable(Device);
+	CreatedShaderTable->Init(this, Pipeline);
+	ShaderTables.Add(Pipeline, CreatedShaderTable);
+	return CreatedShaderTable;
+}
+
 void FVulkanDynamicRHI::RHITransferRayTracingGeometryUnderlyingResource(FRHIRayTracingGeometry* DestGeometry, FRHIRayTracingGeometry* SrcGeometry)
 {
 	check(DestGeometry);
@@ -1467,6 +1608,31 @@ void FVulkanRayTracingCompactionRequestHandler::Update(FVulkanCommandListContext
 	}
 }
 
+
+
+
+FVulkanBindlessDescriptorManager::FUniformBufferDescriptorArrays GetStageUBs(FVulkanDevice* Device, const FRayTracingShaderBindings& InGlobalResourceBindings)
+{
+	FVulkanBindlessDescriptorManager::FUniformBufferDescriptorArrays StageUBs;
+	for (uint32 UBIndex = 0; UBIndex < 16; ++UBIndex)
+	{
+		const FVulkanUniformBuffer* UniformBuffer = ResourceCast(InGlobalResourceBindings.UniformBuffers[UBIndex]);
+		if (UniformBuffer)
+		{
+			VkDescriptorAddressInfoEXT& DescriptorAddressInfo = StageUBs[ShaderStage::EStage::RayGen].AddZeroed_GetRef();
+			DescriptorAddressInfo.sType = VK_STRUCTURE_TYPE_BUFFER_DEVICE_ADDRESS_INFO;
+			DescriptorAddressInfo.address = UniformBuffer->GetDeviceAddress();
+			DescriptorAddressInfo.range = UniformBuffer->GetSize();
+		}
+		else
+		{
+			break;
+		}
+	}
+	return StageUBs;
+}
+
+
 void FVulkanCommandListContext::RHIRayTraceDispatch(
 	FRHIRayTracingPipelineState* InRayTracingPipelineState, 
 	FRHIRayTracingShader* InRayGenShader,
@@ -1474,7 +1640,27 @@ void FVulkanCommandListContext::RHIRayTraceDispatch(
 	const FRayTracingShaderBindings& InGlobalResourceBindings, // :todo-jn:
 	uint32 InWidth, uint32 InHeight)
 {
-	checkNoEntry();
+	const FVulkanRayTracingPipelineState* Pipeline = ResourceCast(InRayTracingPipelineState);
+	FVulkanRayTracingScene* Scene = ResourceCast(InScene);
+	FVulkanRayTracingShader* RayGenShader = ResourceCast(InRayGenShader);
+	FVulkanRayTracingShaderTable* ShaderTable = Scene->FindOrCreateShaderTable(Pipeline);
+
+	FVulkanCmdBuffer* const CmdBuffer = GetCommandBufferManager()->GetActiveCmdBuffer();
+	VulkanRHI::vkCmdBindPipeline(CmdBuffer->GetHandle(), VK_PIPELINE_BIND_POINT_RAY_TRACING_KHR, Pipeline->GetPipeline());
+
+	// :todo-jn: set real uniform buffers, check for gaps, process all stages
+	FVulkanBindlessDescriptorManager::FUniformBufferDescriptorArrays StageUBs = GetStageUBs(Device, InGlobalResourceBindings);
+	Device->GetBindlessDescriptorManager()->RegisterUniformBuffers(CmdBuffer->GetHandle(), VK_PIPELINE_BIND_POINT_RAY_TRACING_KHR, StageUBs);
+
+	ShaderTable->SetSlot(InRayGenShader->GetFrequency(), 0, Pipeline->GetShaderIndex(RayGenShader), Pipeline->GetShaderHandles(SF_RayGen));
+
+	VulkanRHI::vkCmdTraceRaysKHR(
+		CmdBuffer->GetHandle(),
+		ShaderTable->GetRegion(SF_RayGen),
+		ShaderTable->GetRegion(SF_RayMiss),
+		ShaderTable->GetRegion(SF_RayHitGroup),
+		ShaderTable->GetRegion(SF_RayCallable),
+		InWidth, InHeight, 1);
 }
 
 void FVulkanCommandListContext::RHIRayTraceDispatchIndirect(
@@ -1484,7 +1670,32 @@ void FVulkanCommandListContext::RHIRayTraceDispatchIndirect(
 	const FRayTracingShaderBindings& InGlobalResourceBindings, // :todo-jn:
 	FRHIBuffer* InArgumentBuffer, uint32 InArgumentOffset)
 {
-	checkNoEntry();
+	checkf(GRHISupportsRayTracingDispatchIndirect, TEXT("RHIRayTraceDispatchIndirect may not be used because it is not supported on this machine."));
+
+	const FVulkanRayTracingPipelineState* Pipeline = ResourceCast(InRayTracingPipelineState);
+	FVulkanRayTracingScene* Scene = ResourceCast(InScene);
+	FVulkanRayTracingShader* RayGenShader = ResourceCast(InRayGenShader);
+	FVulkanRayTracingShaderTable* ShaderTable = Scene->FindOrCreateShaderTable(Pipeline);
+
+	FVulkanCmdBuffer* const CmdBuffer = GetCommandBufferManager()->GetActiveCmdBuffer();
+	VulkanRHI::vkCmdBindPipeline(CmdBuffer->GetHandle(), VK_PIPELINE_BIND_POINT_RAY_TRACING_KHR, Pipeline->GetPipeline());
+
+	// :todo-jn: set real uniform buffers
+	FVulkanBindlessDescriptorManager::FUniformBufferDescriptorArrays StageUBs = GetStageUBs(Device, InGlobalResourceBindings);
+	Device->GetBindlessDescriptorManager()->RegisterUniformBuffers(CmdBuffer->GetHandle(), VK_PIPELINE_BIND_POINT_RAY_TRACING_KHR, StageUBs);
+
+	FVulkanResourceMultiBuffer* ArgumentBuffer = ResourceCast(InArgumentBuffer);
+	const VkDeviceAddress IndirectDeviceAddress = ArgumentBuffer->GetDeviceAddress() + InArgumentOffset;
+
+	ShaderTable->SetSlot(InRayGenShader->GetFrequency(), 0, Pipeline->GetShaderIndex(RayGenShader), Pipeline->GetShaderHandles(SF_RayGen));
+
+	VulkanRHI::vkCmdTraceRaysIndirectKHR(
+		CmdBuffer->GetHandle(),
+		ShaderTable->GetRegion(SF_RayGen),
+		ShaderTable->GetRegion(SF_RayMiss),
+		ShaderTable->GetRegion(SF_RayHitGroup),
+		ShaderTable->GetRegion(SF_RayCallable),
+		IndirectDeviceAddress);
 }
 
 void FVulkanCommandListContext::RHISetRayTracingHitGroup(
@@ -1522,6 +1733,8 @@ void FVulkanCommandListContext::RHISetRayTracingBindings(
 	ERayTracingBindingType BindingType)
 {
 	FVulkanRayTracingScene* Scene = ResourceCast(InScene);
+	FVulkanRayTracingPipelineState* Pipeline = ResourceCast(InPipeline);
+	FVulkanRayTracingShaderTable* ShaderTable = Scene->FindOrCreateShaderTable(Pipeline);
 
 	checkf(Scene->IsBuilt(), TEXT("Ray tracing scene must be built before any shaders can be bound to it. Make sure that RHIBuildAccelerationStructure() command has been executed."));
 
diff --git a/Engine/Source/Runtime/VulkanRHI/Private/VulkanRayTracing.h b/Engine/Source/Runtime/VulkanRHI/Private/VulkanRayTracing.h
--- a/Engine/Source/Runtime/VulkanRHI/Private/VulkanRayTracing.h
+++ b/Engine/Source/Runtime/VulkanRHI/Private/VulkanRayTracing.h
@@ -59,6 +59,56 @@ struct FVkRtBLASBuildData
 	VkAccelerationStructureBuildSizesInfoKHR SizesInfo;
 };
 
+class FVulkanRayTracingShaderTable : VulkanRHI::FDeviceChild
+{
+public:
+	FVulkanRayTracingShaderTable(FVulkanDevice* Device);
+	~FVulkanRayTracingShaderTable();
+
+	void Init(const FVulkanRayTracingScene* Scene, const FVulkanRayTracingPipelineState* Pipeline);
+
+	const VkStridedDeviceAddressRegionKHR* GetRegion(EShaderFrequency Frequency);
+
+	void SetSlot(EShaderFrequency Frequency, uint32 DstSlot, uint32 SrcHandleIndex, TConstArrayView<uint8> SrcHandleData);
+
+	template <typename T>
+	void SetLocalShaderParameters(EShaderFrequency Frequency, uint32 RecordIndex, uint32 InOffsetWithinRootSignature, const T& Parameters)
+	{
+		SetLocalShaderParameters(Frequency, RecordIndex, InOffsetWithinRootSignature, &Parameters, sizeof(Parameters));
+	}
+
+	void SetLocalShaderParameters(EShaderFrequency Frequency, uint32 RecordIndex, uint32 OffsetWithinRecord, const void* InData, uint32 InDataSize);
+
+private:
+
+	struct FVulkanShaderTableAllocation
+	{
+		FVulkanShaderTableAllocation()
+		{
+			FMemory::Memzero(Region);
+		}
+
+		uint32 HandleCount = 0;
+		bool bUseLocalRecord = false;
+
+		VkBuffer Buffer = VK_NULL_HANDLE;
+		VulkanRHI::FVulkanAllocation Allocation;
+		VkStridedDeviceAddressRegionKHR Region;
+		uint8* MappedBufferMemory = nullptr;
+	};
+
+	FVulkanShaderTableAllocation& GetAlloc(EShaderFrequency Frequency);
+
+	FVulkanShaderTableAllocation Raygen;
+	FVulkanShaderTableAllocation Miss;
+	FVulkanShaderTableAllocation HitGroup;
+	FVulkanShaderTableAllocation Callable;
+
+	// Convenience
+	const uint32 HandleSize;
+	const uint32 HandleSizeAligned;
+};
+
 class FVulkanRayTracingGeometry : public FRHIRayTracingGeometry
 {
 public:
@@ -123,6 +173,13 @@ public:
 		return PerInstanceGeometryParameterSRV.GetReference();
 	}
 
+	FVulkanRayTracingShaderTable* FindOrCreateShaderTable(const FVulkanRayTracingPipelineState* Pipeline);
+
+	inline uint32 GetHitRecordBaseIndex(uint32 InstanceIndex, uint32 SegmentIndex) const 
+	{
+		return (Initializer.SegmentPrefixSum[InstanceIndex] + SegmentIndex) * Initializer.ShaderSlotsPerGeometrySegment;
+	}
+
 	inline bool IsBuilt() const
 	{
 		return bBuilt;
@@ -153,6 +210,9 @@ private:
 	// Buffer that contains per-instance index and vertex buffer binding data
 	TRefCountPtr<FVulkanResourceMultiBuffer> PerInstanceGeometryParameterBuffer;
 	FShaderResourceViewRHIRef PerInstanceGeometryParameterSRV;
+	
+	TMap<const FVulkanRayTracingPipelineState*, FVulkanRayTracingShaderTable*> ShaderTables;
+
 	void BuildPerInstanceGeometryParameterBuffer(FVulkanCommandListContext& CommandContext);
 
 	bool bBuilt = false;
