VulkanRHI: Make Vulkan bindless more convenient to turn on by automatically disabling dynamic buffers instead of telling the user to disable it.

diff --git a/Engine/Source/Runtime/VulkanRHI/Private/VulkanDescriptorSets.cpp b/Engine/Source/Runtime/VulkanRHI/Private/VulkanDescriptorSets.cpp
--- a/Engine/Source/Runtime/VulkanRHI/Private/VulkanDescriptorSets.cpp
+++ b/Engine/Source/Runtime/VulkanRHI/Private/VulkanDescriptorSets.cpp
@@ -241,20 +241,6 @@ bool FVulkanBindlessDescriptorManager::VerifySupport(FVulkanDevice* Device)
 		return false;
 	}
 
-	extern TAutoConsoleVariable<int32> GDynamicGlobalUBs;
-	if (GDynamicGlobalUBs->GetInt() != 0)  // :todo-jn:
-	{
-		UE_LOG(LogRHI, Warning, TEXT("Please disable Dynamic Uniform Buffers (r.Vulkan.DynamicGlobalUBs=0) if you want to use Bindless in Vulkan."));
-		return false;
-	}
-
-	extern int32 GVulkanEnableDefrag;
-	if (GVulkanEnableDefrag != 0)  // :todo-jn:
-	{
-		UE_LOG(LogRHI, Warning, TEXT("Please disable memory defrag (r.Vulkan.EnableDefrag=0) if you want to use Bindless in Vulkan."));
-		return false;
-	}
-
 	const bool bFullyEnabled =
 		(RHIGetRuntimeBindlessResourcesConfiguration(GMaxRHIShaderPlatform) == ERHIBindlessConfiguration::AllShaders) &&
 		(RHIGetRuntimeBindlessSamplersConfiguration(GMaxRHIShaderPlatform) == ERHIBindlessConfiguration::AllShaders);
@@ -281,6 +267,19 @@ bool FVulkanBindlessDescriptorManager::VerifySupport(FVulkanDevice* Device)
 
 			if (bMeetsPropertiesRequirements)
 			{
+				extern TAutoConsoleVariable<int32> GDynamicGlobalUBs;
+				if (GDynamicGlobalUBs->GetInt() != 0)
+				{
+					UE_LOG(LogRHI, Warning, TEXT("Dynamic Uniform Buffers are enabled, but they will not be used with Vulkan bindless."));
+				}
+
+				extern int32 GVulkanEnableDefrag;
+				if (GVulkanEnableDefrag != 0)  // :todo-jn: to be turned back on with new defragger
+				{
+					UE_LOG(LogRHI, Warning, TEXT("Memory defrag is enabled, but it will not be used with Vulkan bindless."));
+					GVulkanEnableDefrag = 0;
+				}
+
 				return true;
 			}
 			else
diff --git a/Engine/Source/Runtime/VulkanRHI/Private/VulkanShaders.cpp b/Engine/Source/Runtime/VulkanRHI/Private/VulkanShaders.cpp
--- a/Engine/Source/Runtime/VulkanRHI/Private/VulkanShaders.cpp
+++ b/Engine/Source/Runtime/VulkanRHI/Private/VulkanShaders.cpp
@@ -745,8 +745,8 @@ void FVulkanDescriptorSetsLayoutInfo::FinalizeBindings(const FVulkanDevice& Devi
 	FMemory::Memzero(Binding);
 	Binding.descriptorCount = 1;
 
-	const bool bConvertAllUBsToDynamic = (GDynamicGlobalUBs.GetValueOnAnyThread() > 1);
-	const bool bConvertPackedUBsToDynamic = bConvertAllUBsToDynamic || (GDynamicGlobalUBs.GetValueOnAnyThread() == 1);
+	const bool bConvertAllUBsToDynamic = !Device.SupportsBindless() && (GDynamicGlobalUBs.GetValueOnAnyThread() > 1);
+	const bool bConvertPackedUBsToDynamic = !Device.SupportsBindless() && (bConvertAllUBsToDynamic || (GDynamicGlobalUBs.GetValueOnAnyThread() == 1));
 	const bool bConsolidateAllIntoOneSet = GDescriptorSetLayoutMode.GetValueOnAnyThread() == 2;
 	const uint32 MaxDescriptorSetUniformBuffersDynamic = Device.GetLimits().maxDescriptorSetUniformBuffersDynamic;
 
