VulkanRHI: VKRT WIP: Add the possibility to update bindless buffer descriptors with only an address (instead of VkBuffer + Offset) so we can start transitioning to callers caching the address instead of making an API call each time.

#rnx

diff --git a/Engine/Source/Runtime/VulkanRHI/Private/VulkanDescriptorSets.cpp b/Engine/Source/Runtime/VulkanRHI/Private/VulkanDescriptorSets.cpp
--- a/Engine/Source/Runtime/VulkanRHI/Private/VulkanDescriptorSets.cpp
+++ b/Engine/Source/Runtime/VulkanRHI/Private/VulkanDescriptorSets.cpp
@@ -826,15 +826,25 @@ void FVulkanBindlessDescriptorManager::UpdateBuffer(FRHIDescriptorHandle Descrip
 		const VkDescriptorType DescriptorType = GetDescriptorTypeForSetIndex(DescriptorHandle.GetRawType());
 		check((DescriptorType == VK_DESCRIPTOR_TYPE_UNIFORM_BUFFER) || (DescriptorType == VK_DESCRIPTOR_TYPE_STORAGE_BUFFER));
 
-		// :todo-jn: start caching buffer addresses in resources to avoid the extra call
 		VkBufferDeviceAddressInfo BufferInfo;
 		ZeroVulkanStruct(BufferInfo, VK_STRUCTURE_TYPE_BUFFER_DEVICE_ADDRESS_INFO);
 		BufferInfo.buffer = VulkanBuffer;
 		const VkDeviceAddress BufferAddress = VulkanRHI::vkGetBufferDeviceAddressKHR(Device->GetInstanceHandle(), &BufferInfo);
 
+		UpdateBuffer(DescriptorHandle, BufferAddress + BufferOffset, BufferSize, bImmediateUpdate);
+	}
+}
+
+void FVulkanBindlessDescriptorManager::UpdateBuffer(FRHIDescriptorHandle DescriptorHandle, VkDeviceAddress BufferAddress, VkDeviceSize BufferSize, bool bImmediateUpdate)
+{
+	if (bIsSupported)
+	{
+		const VkDescriptorType DescriptorType = GetDescriptorTypeForSetIndex(DescriptorHandle.GetRawType());
+		check((DescriptorType == VK_DESCRIPTOR_TYPE_UNIFORM_BUFFER) || (DescriptorType == VK_DESCRIPTOR_TYPE_STORAGE_BUFFER));
+
 		VkDescriptorAddressInfoEXT AddressInfo;
 		ZeroVulkanStruct(AddressInfo, VK_STRUCTURE_TYPE_DESCRIPTOR_ADDRESS_INFO_EXT);
-		AddressInfo.address = BufferAddress + BufferOffset;
+		AddressInfo.address = BufferAddress;
 		AddressInfo.range = BufferSize;
 
 		VkDescriptorDataEXT DescriptorData;
diff --git a/Engine/Source/Runtime/VulkanRHI/Private/VulkanDescriptorSets.h b/Engine/Source/Runtime/VulkanRHI/Private/VulkanDescriptorSets.h
--- a/Engine/Source/Runtime/VulkanRHI/Private/VulkanDescriptorSets.h
+++ b/Engine/Source/Runtime/VulkanRHI/Private/VulkanDescriptorSets.h
@@ -1220,6 +1220,7 @@ public:
 	void UpdateSampler(FRHIDescriptorHandle DescriptorHandle, VkSampler VulkanSampler);
 	void UpdateImage(FRHIDescriptorHandle DescriptorHandle, VkImageView VulkanImage, bool bIsDepthStencil, bool bImmediateUpdate = true);
 	void UpdateBuffer(FRHIDescriptorHandle DescriptorHandle, VkBuffer VulkanBuffer, VkDeviceSize BufferOffset, VkDeviceSize BufferSize, bool bImmediateUpdate = true);
+	void UpdateBuffer(FRHIDescriptorHandle DescriptorHandle, VkDeviceAddress BufferAddress, VkDeviceSize BufferSize, bool bImmediateUpdate = true);
 	void UpdateTexelBuffer(FRHIDescriptorHandle DescriptorHandle, const VkBufferViewCreateInfo& ViewInfo, bool bImmediateUpdate = true);
 	void UpdateAccelerationStructure(FRHIDescriptorHandle DescriptorHandle, VkAccelerationStructureKHR AccelerationStructure, bool bImmediateUpdate = true);
 
