VulkanRHI: Add VK_BUFFER_USAGE_SHADER_DEVICE_ADDRESS_BIT when the extension is present (instead of checking for Bindless or Raytracing).

#rnx

diff --git a/Engine/Source/Runtime/VulkanRHI/Private/VulkanBuffer.cpp b/Engine/Source/Runtime/VulkanRHI/Private/VulkanBuffer.cpp
--- a/Engine/Source/Runtime/VulkanRHI/Private/VulkanBuffer.cpp
+++ b/Engine/Source/Runtime/VulkanRHI/Private/VulkanBuffer.cpp
@@ -131,7 +131,7 @@ VkBufferUsageFlags FVulkanResourceMultiBuffer::UEToVKBufferUsageFlags(FVulkanDev
 		}
 #endif
 		// For descriptors buffers
-		if (InDevice->SupportsBindless())
+		if (InDevice->GetOptionalExtensions().HasBufferDeviceAddress)
 		{
 			OutVkUsage |= VK_BUFFER_USAGE_SHADER_DEVICE_ADDRESS_BIT;
 		}
diff --git a/Engine/Source/Runtime/VulkanRHI/Private/VulkanMemory.cpp b/Engine/Source/Runtime/VulkanRHI/Private/VulkanMemory.cpp
--- a/Engine/Source/Runtime/VulkanRHI/Private/VulkanMemory.cpp
+++ b/Engine/Source/Runtime/VulkanRHI/Private/VulkanMemory.cpp
@@ -2772,7 +2772,7 @@ namespace VulkanRHI
 		BufferCreateInfo.size = BufferSize;
 		BufferCreateInfo.usage = BufferUsageFlags;
 		// For descriptors buffers
-		if (Device->SupportsBindless())
+		if (Device->GetOptionalExtensions().HasBufferDeviceAddress)
 		{
 			BufferCreateInfo.usage |= VK_BUFFER_USAGE_SHADER_DEVICE_ADDRESS_BIT;
 		}
@@ -4299,7 +4299,7 @@ namespace VulkanRHI
 		}
 
 		// For descriptors buffers
-		if (Device->SupportsBindless())
+		if (Device->GetOptionalExtensions().HasBufferDeviceAddress)
 		{
 			InUsageFlags |= VK_BUFFER_USAGE_SHADER_DEVICE_ADDRESS_BIT;
 		}
