VulkanRHI: slightly improve Vulkan buffer memory tracking by excluding transient memory and including all of the allocations of dynamic buffers (instead of only the first one).

#rnx

diff --git a/Engine/Source/Runtime/VulkanRHI/Private/VulkanBuffer.cpp b/Engine/Source/Runtime/VulkanRHI/Private/VulkanBuffer.cpp
--- a/Engine/Source/Runtime/VulkanRHI/Private/VulkanBuffer.cpp
+++ b/Engine/Source/Runtime/VulkanRHI/Private/VulkanBuffer.cpp
@@ -217,8 +217,6 @@ FVulkanResourceMultiBuffer::FVulkanResourceMultiBuffer(FVulkanDevice* InDevice,
 
 				CreateInfo.ResourceArray->Discard();
 			}
-
-			UpdateVulkanBufferStats(InBufferDesc.Size, BufferUsageFlags, true);
 		}
 	}
 }
@@ -289,6 +287,8 @@ void FVulkanResourceMultiBuffer::AdvanceBufferIndex()
 		NewBufferAlloc.Fence = new FVulkanGPUFence(TEXT("VulkanDynamicBuffer"));
 		NewBufferAlloc.AllocStatus = FBufferAlloc::EAllocStatus::InUse;
 		NewBufferAlloc.DeviceAddress = GetBufferDeviceAddress(Device, NewBufferAlloc.Alloc.GetBufferHandle()) + NewBufferAlloc.Alloc.Offset;
+
+		UpdateVulkanBufferStats(BufferSize, BufferUsageFlags, true);
 	}
 }
 
