Vulkan RHI: Update buffer views when using mapped locks.

diff --git a/Engine/Source/Runtime/VulkanRHI/Private/VulkanBuffer.cpp b/Engine/Source/Runtime/VulkanRHI/Private/VulkanBuffer.cpp
--- a/Engine/Source/Runtime/VulkanRHI/Private/VulkanBuffer.cpp
+++ b/Engine/Source/Runtime/VulkanRHI/Private/VulkanBuffer.cpp
@@ -15,6 +15,15 @@
 static TMap<FVulkanResourceMultiBuffer*, VulkanRHI::FPendingBufferLock> GPendingLockIBs;
 static FCriticalSection GPendingLockIBsMutex;
 
+int32 GVulkanForceStagingBufferOnLock = 0;
+static FAutoConsoleVariableRef CVarVulkanForceStagingBufferOnLock(
+	TEXT("r.Vulkan.ForceStagingBufferOnLock"),
+	GVulkanForceStagingBufferOnLock,
+	TEXT("When nonzero, non-volatile buffer locks will always use staging buffers. Useful for debugging.\n")
+	TEXT("default: 0"),
+	ECVF_RenderThreadSafe
+);
+
 static FORCEINLINE VulkanRHI::FPendingBufferLock GetPendingBufferLock(FVulkanResourceMultiBuffer* Buffer)
 {
 	VulkanRHI::FPendingBufferLock PendingLock;
@@ -332,6 +341,7 @@ void* FVulkanResourceMultiBuffer::Lock(FVulkanCommandListContext& Context, EReso
 	uint32 DataOffset = 0;
 
 	const bool bVolatile = EnumHasAnyFlags(GetUsage(), BUF_Volatile);
+	check(LockStatus == ELockStatus::Unlocked);
 
 	LockStatus = ELockStatus::Locked;
 	++LockCounter;
@@ -436,7 +446,7 @@ void* FVulkanResourceMultiBuffer::Lock(FVulkanCommandListContext& Context, EReso
 			check(LockMode == RLM_WriteOnly);
 
 			// Always use staging buffers to update 'Static' buffers since they maybe be in use by GPU atm
-			const bool bUseStagingBuffer = (bStatic || !bUnifiedMem);
+			const bool bUseStagingBuffer = (bStatic || !bUnifiedMem) || GVulkanForceStagingBufferOnLock;
 			if (bUseStagingBuffer)
 			{
 				// NOTE: No need to change the CurrentBufferIndex if we're using a staging buffer for the copy
@@ -554,6 +564,7 @@ void FVulkanResourceMultiBuffer::Unlock(FRHICommandListBase* RHICmdList, FVulkan
 		if (Context)
 		{
 			UpdateBufferAllocStates(*Context);
+			UpdateLinkedViews();
 		}
 		else
 		{
@@ -561,6 +572,7 @@ void FVulkanResourceMultiBuffer::Unlock(FRHICommandListBase* RHICmdList, FVulkan
 			{
 				FVulkanCommandListContext& Context = FVulkanCommandListContext::GetVulkanContext(CmdList.GetContext());
 				UpdateBufferAllocStates(Context);
+				UpdateLinkedViews();
 			});
 		}
 	}
