VulkanRHI: VKRT WIP: Add a quick configurable impl of DeferredOperation for pipeline creation.  Split up the cached host copy of the SRT and the gpu copy.  Add basic code for the RHI context calls (trace rays, set miss, set callable, set hitgroup).

#rnx

diff --git a/Engine/Source/Runtime/VulkanRHI/Private/VulkanCommandWrappers.h b/Engine/Source/Runtime/VulkanRHI/Private/VulkanCommandWrappers.h
--- a/Engine/Source/Runtime/VulkanRHI/Private/VulkanCommandWrappers.h
+++ b/Engine/Source/Runtime/VulkanRHI/Private/VulkanCommandWrappers.h
@@ -216,6 +216,11 @@ struct FWrapLayer
 	static void CmdBindDescriptorBuffersEXT(VkResult Result, VkCommandBuffer CommandBuffer, uint32_t BufferCount, const VkDescriptorBufferBindingInfoEXT* BindingInfos) VULKAN_LAYER_BODY
 	static void CmdSetDescriptorBufferOffsetsEXT(VkResult Result, VkCommandBuffer CommandBuffer, VkPipelineBindPoint PipelineBindPoint, VkPipelineLayout Layout, uint32_t FirstSet, uint32_t SetCount, const uint32_t* BufferIndices, const VkDeviceSize* Offsets) VULKAN_LAYER_BODY
 	static void GetDescriptorEXT(VkResult Result, VkDevice Device, const VkDescriptorGetInfoEXT* DescriptorInfo, size_t DataSize, void* Descriptor) VULKAN_LAYER_BODY
+	static void CreateDeferredOperationKHR(VkResult Result, VkDevice Device, const VkAllocationCallbacks* Allocator, VkDeferredOperationKHR* DeferredOperation) VULKAN_LAYER_BODY
+	static void DestroyDeferredOperationKHR(VkResult Result, VkDevice Device, VkDeferredOperationKHR DeferredOperation, const VkAllocationCallbacks* Allocator) VULKAN_LAYER_BODY
+	static void DeferredOperationJoinKHR(VkResult Result, VkDevice Device, VkDeferredOperationKHR DeferredOperation) VULKAN_LAYER_BODY
+	static void GetDeferredOperationMaxConcurrencyKHR(VkResult Result, VkDevice Device, VkDeferredOperationKHR DeferredOperation) VULKAN_LAYER_BODY
+	static void GetDeferredOperationResultKHR(VkResult Result, VkDevice Device, VkDeferredOperationKHR DeferredOperation) VULKAN_LAYER_BODY
 	static void GetDeviceFaultInfoEXT(VkResult Result, VkDevice Device, VkDeviceFaultCountsEXT* FaultCounts, VkDeviceFaultInfoEXT* FaultInfo) VULKAN_LAYER_BODY
 };
 
@@ -1635,6 +1640,45 @@ namespace VulkanRHI
 		FWrapLayer::GetDescriptorEXT(VK_SUCCESS, Device, DescriptorInfo,  DataSize, Descriptor);
 	}
 
+	static FORCEINLINE_DEBUGGABLE VkResult vkCreateDeferredOperationKHR(VkDevice Device, const VkAllocationCallbacks* Allocator, VkDeferredOperationKHR* DeferredOperation)
+	{
+		FWrapLayer::CreateDeferredOperationKHR(VK_RESULT_MAX_ENUM, Device, Allocator, DeferredOperation);
+		VkResult Result = VULKANAPINAMESPACE::vkCreateDeferredOperationKHR(Device, Allocator, DeferredOperation);
+		FWrapLayer::CreateDeferredOperationKHR(Result, Device, Allocator, DeferredOperation);
+		return Result;
+	}
+
+	static FORCEINLINE_DEBUGGABLE void vkDestroyDeferredOperationKHR(VkDevice Device, VkDeferredOperationKHR DeferredOperation, const VkAllocationCallbacks* Allocator)
+	{
+		FWrapLayer::DestroyDeferredOperationKHR(VK_RESULT_MAX_ENUM, Device, DeferredOperation, Allocator);
+		VULKANAPINAMESPACE::vkDestroyDeferredOperationKHR(Device, DeferredOperation, Allocator);
+		FWrapLayer::DestroyDeferredOperationKHR(VK_SUCCESS, Device, DeferredOperation, Allocator);
+	}
+
+	static FORCEINLINE_DEBUGGABLE VkResult vkDeferredOperationJoinKHR(VkDevice Device, VkDeferredOperationKHR DeferredOperation)
+	{
+		FWrapLayer::DeferredOperationJoinKHR(VK_RESULT_MAX_ENUM, Device, DeferredOperation);
+		VkResult Result = VULKANAPINAMESPACE::vkDeferredOperationJoinKHR(Device, DeferredOperation);
+		FWrapLayer::DeferredOperationJoinKHR(Result, Device, DeferredOperation);
+		return Result;
+	}
+
+	static FORCEINLINE_DEBUGGABLE uint32_t vkGetDeferredOperationMaxConcurrencyKHR(VkDevice Device, VkDeferredOperationKHR DeferredOperation)
+	{
+		FWrapLayer::GetDeferredOperationMaxConcurrencyKHR(VK_RESULT_MAX_ENUM, Device, DeferredOperation);
+		uint32_t Result = VULKANAPINAMESPACE::vkGetDeferredOperationMaxConcurrencyKHR(Device, DeferredOperation);
+		FWrapLayer::GetDeferredOperationMaxConcurrencyKHR(VK_SUCCESS, Device, DeferredOperation);
+		return Result;
+	}
+
+	static FORCEINLINE_DEBUGGABLE VkResult vkGetDeferredOperationResultKHR(VkDevice Device, VkDeferredOperationKHR DeferredOperation)
+	{
+		FWrapLayer::GetDeferredOperationResultKHR(VK_RESULT_MAX_ENUM, Device, DeferredOperation);
+		VkResult Result = VULKANAPINAMESPACE::vkGetDeferredOperationResultKHR(Device, DeferredOperation);
+		FWrapLayer::GetDeferredOperationResultKHR(Result, Device, DeferredOperation);
+		return Result;
+	}
+
 	static FORCEINLINE_DEBUGGABLE VkResult vkGetDeviceFaultInfoEXT(VkDevice Device, VkDeviceFaultCountsEXT* FaultCounts, VkDeviceFaultInfoEXT* FaultInfo)
 	{
 		FWrapLayer::GetDeviceFaultInfoEXT(VK_RESULT_MAX_ENUM, Device, FaultCounts, FaultInfo);
@@ -1642,6 +1686,7 @@ namespace VulkanRHI
 		FWrapLayer::GetDeviceFaultInfoEXT(Result, Device, FaultCounts, FaultInfo);
 		return Result;
 	}
+
 #if VULKAN_ENABLE_IMAGE_TRACKING_LAYER
 	void BindDebugLabelName(VkImage Image, const TCHAR* Name);
 #endif
diff --git a/Engine/Source/Runtime/VulkanRHI/Private/VulkanDebug.cpp b/Engine/Source/Runtime/VulkanRHI/Private/VulkanDebug.cpp
--- a/Engine/Source/Runtime/VulkanRHI/Private/VulkanDebug.cpp
+++ b/Engine/Source/Runtime/VulkanRHI/Private/VulkanDebug.cpp
@@ -4568,6 +4568,56 @@ void FWrapLayer::GetDescriptorEXT(VkResult Result, VkDevice Device, const VkDesc
 	}
 }
 
+void FWrapLayer::CreateDeferredOperationKHR(VkResult Result, VkDevice Device, const VkAllocationCallbacks* Allocator, VkDeferredOperationKHR* DeferredOperation)
+{
+	if (Result == VK_RESULT_MAX_ENUM)
+	{
+#if VULKAN_ENABLE_DUMP_LAYER
+		PrintfBeginResult(FString::Printf(TEXT("CreateDeferredOperationKHR(Device=0x%p, DeferredOperation=0x%p)"), Device, DeferredOperation));
+#endif
+	}
+}
+
+void FWrapLayer::DestroyDeferredOperationKHR(VkResult Result, VkDevice Device, VkDeferredOperationKHR DeferredOperation, const VkAllocationCallbacks* Allocator)
+{
+	if (Result == VK_RESULT_MAX_ENUM)
+	{
+#if VULKAN_ENABLE_DUMP_LAYER
+		PrintfBeginResult(FString::Printf(TEXT("DestroyDeferredOperationKHR(Device=0x%p, DeferredOperation=0x%p)"), Device, DeferredOperation));
+#endif
+	}
+}
+
+void FWrapLayer::DeferredOperationJoinKHR(VkResult Result, VkDevice Device, VkDeferredOperationKHR DeferredOperation)
+{
+	if (Result == VK_RESULT_MAX_ENUM)
+	{
+#if VULKAN_ENABLE_DUMP_LAYER
+		PrintfBeginResult(FString::Printf(TEXT("DeferredOperationJoinKHR(Device=0x%p, DeferredOperation=0x%p)"), Device, DeferredOperation));
+#endif
+	}
+}
+
+void FWrapLayer::GetDeferredOperationMaxConcurrencyKHR(VkResult Result, VkDevice Device, VkDeferredOperationKHR DeferredOperation)
+{
+	if (Result == VK_RESULT_MAX_ENUM)
+	{
+#if VULKAN_ENABLE_DUMP_LAYER
+		PrintfBeginResult(FString::Printf(TEXT("GetDeferredOperationMaxConcurrencyKHR(Device=0x%p, DeferredOperation=0x%p)"), Device, DeferredOperation));
+#endif
+	}
+}
+
+void FWrapLayer::GetDeferredOperationResultKHR(VkResult Result, VkDevice Device, VkDeferredOperationKHR DeferredOperation)
+{
+	if (Result == VK_RESULT_MAX_ENUM)
+	{
+#if VULKAN_ENABLE_DUMP_LAYER
+		PrintfBeginResult(FString::Printf(TEXT("GetDeferredOperationResultKHR(Device=0x%p, DeferredOperation=0x%p)"), Device, DeferredOperation));
+#endif
+	}
+}
+
 void FWrapLayer::GetDeviceFaultInfoEXT(VkResult Result, VkDevice Device, VkDeviceFaultCountsEXT* FaultCounts, VkDeviceFaultInfoEXT* FaultInfo)
 {
 	if (Result == VK_RESULT_MAX_ENUM)
diff --git a/Engine/Source/Runtime/VulkanRHI/Private/VulkanLoader.h b/Engine/Source/Runtime/VulkanRHI/Private/VulkanLoader.h
--- a/Engine/Source/Runtime/VulkanRHI/Private/VulkanLoader.h
+++ b/Engine/Source/Runtime/VulkanRHI/Private/VulkanLoader.h
@@ -216,6 +216,11 @@
 	EnumMacro(PFN_vkCmdBindDescriptorBuffersEXT, vkCmdBindDescriptorBuffersEXT) \
 	EnumMacro(PFN_vkCmdSetDescriptorBufferOffsetsEXT, vkCmdSetDescriptorBufferOffsetsEXT) \
 	EnumMacro(PFN_vkGetDescriptorEXT, vkGetDescriptorEXT) \
+	EnumMacro(PFN_vkCreateDeferredOperationKHR, vkCreateDeferredOperationKHR) \
+	EnumMacro(PFN_vkDestroyDeferredOperationKHR, vkDestroyDeferredOperationKHR) \
+	EnumMacro(PFN_vkDeferredOperationJoinKHR, vkDeferredOperationJoinKHR) \
+	EnumMacro(PFN_vkGetDeferredOperationMaxConcurrencyKHR, vkGetDeferredOperationMaxConcurrencyKHR) \
+	EnumMacro(PFN_vkGetDeferredOperationResultKHR, vkGetDeferredOperationResultKHR) \
 	EnumMacro(PFN_vkGetDeviceFaultInfoEXT, vkGetDeviceFaultInfoEXT)
 
 
diff --git a/Engine/Source/Runtime/VulkanRHI/Private/VulkanRayTracing.cpp b/Engine/Source/Runtime/VulkanRHI/Private/VulkanRayTracing.cpp
--- a/Engine/Source/Runtime/VulkanRHI/Private/VulkanRayTracing.cpp
+++ b/Engine/Source/Runtime/VulkanRHI/Private/VulkanRayTracing.cpp
@@ -36,6 +36,18 @@ static FAutoConsoleVariableRef CVarVulkanRayTracingTLASPreferFastTraceTLAS(
 	ECVF_ReadOnly
 );
 
+static int32 GVulkanRayTracingAllowDeferredOperation = -1;
+static FAutoConsoleVariableRef CVarVulkanRayTracingAllowDeferredOperation(
+	TEXT("r.Vulkan.RayTracing.AllowDeferredOperation"),
+	GVulkanRayTracingAllowDeferredOperation,
+	TEXT("Whether to use Vulkan Deferred Operation for RT pipeline creation. (default = -1)\n")
+	TEXT(" <0: Disabled\n")
+	TEXT(" 0: Enabled, auto detect the maximum number of threads")
+	TEXT(" >0: Enabled, use the specified number of threads"),
+	ECVF_ReadOnly
+);
+
+
 // Ray tracing stat counters
 
 DECLARE_STATS_GROUP(TEXT("Vulkan: Ray Tracing"), STATGROUP_VulkanRayTracing, STATCAT_Advanced);
@@ -856,23 +868,26 @@ FVulkanRayTracingShaderTable::FVulkanRayTracingShaderTable(FVulkanDevice* Device
 
 FVulkanRayTracingShaderTable::~FVulkanRayTracingShaderTable()
 {
-	auto FreeAlloc = [Device = Device](FVulkanShaderTableAllocation& Alloc) {
-		if (Alloc.Buffer != VK_NULL_HANDLE)
-		{
-			VulkanRHI::vkDestroyBuffer(Device->GetInstanceHandle(), Alloc.Buffer, VULKAN_CPU_ALLOCATOR);
-			Alloc.Buffer = VK_NULL_HANDLE;
-		}
+	ReleaseLocalBuffer(Device, Raygen);
+	ReleaseLocalBuffer(Device, Miss);
+	ReleaseLocalBuffer(Device, HitGroup);
+	ReleaseLocalBuffer(Device, Callable);
+}
 
-		if (Alloc.Allocation.IsValid())
-		{
-			Device->GetMemoryManager().FreeVulkanAllocation(Alloc.Allocation);
-		}
-	};
+void FVulkanRayTracingShaderTable::ReleaseLocalBuffer(FVulkanDevice* Device, FVulkanShaderTableAllocation& Alloc)
+{
+	if (Alloc.LocalBuffer != VK_NULL_HANDLE)
+	{
+		Device->GetDeferredDeletionQueue().EnqueueResource(VulkanRHI::FDeferredDeletionQueue2::EType::Buffer, Alloc.LocalBuffer);
+		Alloc.LocalBuffer = VK_NULL_HANDLE;
+	}
+
+	if (Alloc.LocalAllocation.IsValid())
+	{
+		Device->GetMemoryManager().FreeVulkanAllocation(Alloc.LocalAllocation);
+	}
 
-	FreeAlloc(Raygen);
-	FreeAlloc(Miss);
-	FreeAlloc(HitGroup);
-	FreeAlloc(Callable);
+	Alloc.Region.deviceAddress = 0;
 }
 
 void FVulkanRayTracingShaderTable::Init(const FVulkanRayTracingScene* Scene, const FVulkanRayTracingPipelineState* Pipeline)
@@ -891,27 +906,8 @@ void FVulkanRayTracingShaderTable::Init(const FVulkanRayTracingScene* Scene, con
 			Alloc.Region.stride = InUseLocalRecord ? FMath::Min<VkDeviceSize>(RayTracingPipelineProps.maxShaderGroupStride, 4096) : HandleSizeAligned; // :todo-jn: shrink stride to necessary amount
 			Alloc.Region.size = Alloc.HandleCount * Alloc.Region.stride;
 
-			{
-				VkBufferCreateInfo BufferCreateInfo;
-				ZeroVulkanStruct(BufferCreateInfo, VK_STRUCTURE_TYPE_BUFFER_CREATE_INFO);
-				BufferCreateInfo.size = Alloc.Region.size;
-				BufferCreateInfo.usage = VK_BUFFER_USAGE_SHADER_BINDING_TABLE_BIT_KHR | VK_BUFFER_USAGE_SHADER_DEVICE_ADDRESS_BIT;;
-				VERIFYVULKANRESULT(VulkanRHI::vkCreateBuffer(DeviceHandle, &BufferCreateInfo, VULKAN_CPU_ALLOCATOR, &Alloc.Buffer));
-			}
-
-			const EVulkanAllocationFlags AllocFlags = EVulkanAllocationFlags::HostVisible | EVulkanAllocationFlags::AutoBind | EVulkanAllocationFlags::Dedicated; // :todo-jn: Dedicated for now... preferBAR?
-			Device->GetMemoryManager().AllocateBufferMemory(Alloc.Allocation, Alloc.Buffer, AllocFlags, TEXT("FVulkanShaderTableAllocation"), RayTracingPipelineProps.shaderGroupBaseAlignment);
-
-			VkMemoryAllocateFlagsInfo MemoryAllocateFlagsInfo;
-			ZeroVulkanStruct(MemoryAllocateFlagsInfo, VK_STRUCTURE_TYPE_MEMORY_ALLOCATE_FLAGS_INFO);
-			MemoryAllocateFlagsInfo.flags = VK_MEMORY_ALLOCATE_DEVICE_ADDRESS_BIT_KHR;
-
-			VkBufferDeviceAddressInfoKHR DeviceAddressInfo;
-			ZeroVulkanStruct(DeviceAddressInfo, VK_STRUCTURE_TYPE_BUFFER_DEVICE_ADDRESS_INFO);
-			DeviceAddressInfo.buffer = Alloc.Buffer;
-			Alloc.Region.deviceAddress = VulkanRHI::vkGetBufferDeviceAddressKHR(DeviceHandle, &DeviceAddressInfo);
-
-			Alloc.MappedBufferMemory = (uint8*)Alloc.Allocation.GetMappedPointer(Device);
+			// Host buffer
+			Alloc.HostBuffer.SetNumUninitialized(Alloc.Region.size);
 		}
 	};
 
@@ -943,13 +939,16 @@ FVulkanRayTracingShaderTable::FVulkanShaderTableAllocation& FVulkanRayTracingSha
 
 const VkStridedDeviceAddressRegionKHR* FVulkanRayTracingShaderTable::GetRegion(EShaderFrequency Frequency)
 {
-	return &GetAlloc(Frequency).Region;
+	const FVulkanShaderTableAllocation& Alloc = GetAlloc(Frequency);
+	check(!Alloc.bIsDirty);
+	return &Alloc.Region;
 }
 
 void FVulkanRayTracingShaderTable::SetSlot(EShaderFrequency Frequency, uint32 DstSlot, uint32 SrcHandleIndex, TConstArrayView<uint8> SrcHandleData)
 {
 	FVulkanShaderTableAllocation& Alloc = GetAlloc(Frequency);
-	FMemory::Memcpy(&Alloc.MappedBufferMemory[DstSlot * Alloc.Region.stride], &SrcHandleData[SrcHandleIndex * HandleSize], HandleSize);
+	FMemory::Memcpy(&Alloc.HostBuffer[DstSlot * Alloc.Region.stride], &SrcHandleData[SrcHandleIndex * HandleSize], HandleSize);
+	Alloc.bIsDirty = true;
 }
 
 void FVulkanRayTracingShaderTable::SetLocalShaderParameters(EShaderFrequency Frequency, uint32 RecordIndex, uint32 OffsetWithinRecord, const void* InData, uint32 InDataSize)
@@ -961,7 +960,73 @@ void FVulkanRayTracingShaderTable::SetLocalShaderParameters(EShaderFrequency Fre
 	checkf(OffsetWithinRecord + InDataSize <= Alloc.Region.size, TEXT("SBT record write request is out of bounds"));
 
 	const uint32 WriteOffset = HandleSizeAligned + (Alloc.Region.stride * RecordIndex) + OffsetWithinRecord;
-	FMemory::Memcpy(Alloc.MappedBufferMemory + WriteOffset, InData, InDataSize);
+	FMemory::Memcpy(&Alloc.HostBuffer[WriteOffset], InData, InDataSize);
+
+	Alloc.bIsDirty = true;
+}
+
+void FVulkanRayTracingShaderTable::Commit(FVulkanCommandListContext& Context)
+{
+	FVulkanCommandBufferManager* CommandBufferManager = Context.GetCommandBufferManager();
+	FVulkanCmdBuffer* const CmdBuffer = CommandBufferManager->GetActiveCmdBuffer();
+
+	VkMemoryBarrier BarrierBefore = { VK_STRUCTURE_TYPE_MEMORY_BARRIER, nullptr, VK_ACCESS_MEMORY_READ_BIT | VK_ACCESS_MEMORY_WRITE_BIT, VK_ACCESS_TRANSFER_READ_BIT | VK_ACCESS_TRANSFER_WRITE_BIT };
+	VulkanRHI::vkCmdPipelineBarrier(CmdBuffer->GetHandle(), VK_PIPELINE_STAGE_BOTTOM_OF_PIPE_BIT, VK_PIPELINE_STAGE_TRANSFER_BIT, 0, 1, &BarrierBefore, 0, nullptr, 0, nullptr);
+
+	auto CommitBuffer = [Device = Device, CmdBuffer](FVulkanShaderTableAllocation& Alloc)
+	{
+		if (Alloc.bIsDirty)
+		{
+			if (!Alloc.HostBuffer.IsEmpty())
+			{
+				ReleaseLocalBuffer(Device, Alloc);
+
+				const VkDevice DeviceHandle = Device->GetInstanceHandle();
+				const VkPhysicalDeviceRayTracingPipelinePropertiesKHR& RayTracingPipelineProps = Device->GetOptionalExtensionProperties().RayTracingPipelineProps;
+
+				// Fetch staging buffer and fill it
+				VulkanRHI::FStagingBuffer* StagingBuffer = Device->GetStagingManager().AcquireBuffer(Alloc.Region.size);
+				FMemory::Memcpy(StagingBuffer->GetMappedPointer(), Alloc.HostBuffer.GetData(), Alloc.Region.size);
+
+				// Alloc a new Local buffer
+				{
+					VkBufferCreateInfo BufferCreateInfo;
+					ZeroVulkanStruct(BufferCreateInfo, VK_STRUCTURE_TYPE_BUFFER_CREATE_INFO);
+					BufferCreateInfo.size = Alloc.Region.size;
+					BufferCreateInfo.usage = VK_BUFFER_USAGE_TRANSFER_DST_BIT | VK_BUFFER_USAGE_SHADER_BINDING_TABLE_BIT_KHR | VK_BUFFER_USAGE_SHADER_DEVICE_ADDRESS_BIT;
+					VERIFYVULKANRESULT(VulkanRHI::vkCreateBuffer(DeviceHandle, &BufferCreateInfo, VULKAN_CPU_ALLOCATOR, &Alloc.LocalBuffer));
+
+					const EVulkanAllocationFlags AllocFlags = EVulkanAllocationFlags::AutoBind | EVulkanAllocationFlags::Dedicated;
+					Device->GetMemoryManager().AllocateBufferMemory(Alloc.LocalAllocation, Alloc.LocalBuffer, AllocFlags, TEXT("LocalShaderTableAllocation"), RayTracingPipelineProps.shaderGroupBaseAlignment);
+
+					VkBufferDeviceAddressInfoKHR DeviceAddressInfo;
+					ZeroVulkanStruct(DeviceAddressInfo, VK_STRUCTURE_TYPE_BUFFER_DEVICE_ADDRESS_INFO);
+					DeviceAddressInfo.buffer = Alloc.LocalBuffer;
+					Alloc.Region.deviceAddress = VulkanRHI::vkGetBufferDeviceAddressKHR(DeviceHandle, &DeviceAddressInfo);
+				}
+
+				VkBufferCopy RegionInfo;
+				RegionInfo.srcOffset = 0;
+				RegionInfo.dstOffset = 0;
+				RegionInfo.size = Alloc.Region.size;
+				VulkanRHI::vkCmdCopyBuffer(CmdBuffer->GetHandle(), StagingBuffer->GetHandle(), Alloc.LocalBuffer, 1, &RegionInfo);
+			}
+			else
+			{
+				checkSlow(Alloc.LocalBuffer == VK_NULL_HANDLE);
+			}
+
+			Alloc.bIsDirty = false;
+		}
+	};
+
+	CommitBuffer(Raygen);
+	CommitBuffer(Miss);
+	CommitBuffer(HitGroup);
+	CommitBuffer(Callable);
+
+	VkMemoryBarrier BarrierAfter = { VK_STRUCTURE_TYPE_MEMORY_BARRIER, nullptr, VK_ACCESS_TRANSFER_WRITE_BIT, VK_ACCESS_MEMORY_READ_BIT };  // :todo-jn: VK_ACCESS_2_SHADER_BINDING_TABLE_READ_BIT_KHR
+	VulkanRHI::vkCmdPipelineBarrier(CmdBuffer->GetHandle(), VK_PIPELINE_STAGE_TRANSFER_BIT, VK_PIPELINE_STAGE_RAY_TRACING_SHADER_BIT_KHR, 0, 1, &BarrierAfter, 0, nullptr, 0, nullptr);
 }
 
 FVulkanRayTracingShaderTable* FVulkanRayTracingScene::FindOrCreateShaderTable(const FVulkanRayTracingPipelineState* Pipeline)
@@ -1354,15 +1419,57 @@ FVulkanRayTracingPipelineState::FVulkanRayTracingPipelineState(FVulkanDevice* co
 	RayTracingPipelineCreateInfo.layout = InDevice->GetBindlessDescriptorManager()->GetPipelineLayout();
 	RayTracingPipelineCreateInfo.flags = VK_PIPELINE_CREATE_DESCRIPTOR_BUFFER_BIT_EXT;
 	
-	VERIFYVULKANRESULT(VulkanDynamicAPI::vkCreateRayTracingPipelinesKHR(
-		InDevice->GetInstanceHandle(), 
-		VK_NULL_HANDLE, // Deferred Operation 
+	VkDeferredOperationKHR DeferredOp = VK_NULL_HANDLE; // :todo-jn: more speed
+	if (GVulkanRayTracingAllowDeferredOperation >= 0)
+	{
+		VERIFYVULKANRESULT(VulkanRHI::vkCreateDeferredOperationKHR(
+			InDevice->GetInstanceHandle(),
+			VULKAN_CPU_ALLOCATOR,
+			&DeferredOp));
+	}
+
+	VERIFYVULKANRESULT_EXPANDED(VulkanDynamicAPI::vkCreateRayTracingPipelinesKHR(
+		InDevice->GetInstanceHandle(),
+		DeferredOp,
 		VK_NULL_HANDLE, // Pipeline Cache 
-		1, 
-		&RayTracingPipelineCreateInfo, 
-		VULKAN_CPU_ALLOCATOR, 
+		1,
+		&RayTracingPipelineCreateInfo,
+		VULKAN_CPU_ALLOCATOR,
 		&Pipeline));
 
+	if (DeferredOp != VK_NULL_HANDLE)
+	{
+		int32 MaxConcurrency = FMath::Min(
+			(int32)VulkanRHI::vkGetDeferredOperationMaxConcurrencyKHR(InDevice->GetInstanceHandle(), DeferredOp),
+			FTaskGraphInterface::Get().GetNumWorkerThreads());
+
+		if (GVulkanRayTracingAllowDeferredOperation > 0)
+		{
+			MaxConcurrency = FMath::Min(MaxConcurrency, GVulkanRayTracingAllowDeferredOperation);
+		}
+
+		bool bCompleted = false;
+		ParallelFor(MaxConcurrency, [DeferredOp, InDevice, &bCompleted](int32 Unused)
+			{
+				VkResult Result = VulkanRHI::vkDeferredOperationJoinKHR(InDevice->GetInstanceHandle(), DeferredOp);
+				while (Result == VK_THREAD_IDLE_KHR)
+				{
+					FPlatformProcess::Sleep(0.01f);
+					Result = VulkanRHI::vkDeferredOperationJoinKHR(InDevice->GetInstanceHandle(), DeferredOp);
+				}
+
+				if (Result == VK_SUCCESS)
+				{
+					bCompleted = true;
+				}
+			});
+		checkf(bCompleted, TEXT("ParallelFor returned but Deferred Operation not complete!"));
+
+		VERIFYVULKANRESULT(VulkanRHI::vkGetDeferredOperationResultKHR(InDevice->GetInstanceHandle(), DeferredOp));
+
+		VulkanRHI::vkDestroyDeferredOperationKHR(InDevice->GetInstanceHandle(), DeferredOp, VULKAN_CPU_ALLOCATOR);
+	}
+
 	for (ANSICHAR* const EntryPoint : EntryPointNames)
 	{
 		delete[] EntryPoint;
@@ -1611,25 +1718,44 @@ void FVulkanRayTracingCompactionRequestHandler::Update(FVulkanCommandListContext
 
 
 
-FVulkanBindlessDescriptorManager::FUniformBufferDescriptorArrays GetStageUBs(FVulkanDevice* Device, const FRayTracingShaderBindings& InGlobalResourceBindings)
+void SetRayGenResources(FVulkanDevice* Device, FVulkanCmdBuffer* const CmdBuffer, const FRayTracingShaderBindings& InGlobalResourceBindings)
 {
-	FVulkanBindlessDescriptorManager::FUniformBufferDescriptorArrays StageUBs;
-	for (uint32 UBIndex = 0; UBIndex < 16; ++UBIndex)
+	// Uniform buffers
 	{
-		const FVulkanUniformBuffer* UniformBuffer = ResourceCast(InGlobalResourceBindings.UniformBuffers[UBIndex]);
-		if (UniformBuffer)
-		{
-			VkDescriptorAddressInfoEXT& DescriptorAddressInfo = StageUBs[ShaderStage::EStage::RayGen].AddZeroed_GetRef();
-			DescriptorAddressInfo.sType = VK_STRUCTURE_TYPE_BUFFER_DEVICE_ADDRESS_INFO;
-			DescriptorAddressInfo.address = UniformBuffer->GetDeviceAddress();
-			DescriptorAddressInfo.range = UniformBuffer->GetSize();
-		}
-		else
+		uint32 NumSkippedSlots = 0;
+		FVulkanBindlessDescriptorManager::FUniformBufferDescriptorArrays StageUBs;
+		const uint32 MaxUniformBuffers = UE_ARRAY_COUNT(InGlobalResourceBindings.UniformBuffers);
+		for (uint32 UBIndex = 0; UBIndex < MaxUniformBuffers; ++UBIndex)
 		{
-			break;
+			const FVulkanUniformBuffer* UniformBuffer = ResourceCast(InGlobalResourceBindings.UniformBuffers[UBIndex]);
+			if (UniformBuffer)
+			{
+				if (NumSkippedSlots > 0)
+				{
+					UE_LOG(LogRHI, Warning, TEXT("Skipping %u Uniform Buffer bindings, this isn't normal!"), NumSkippedSlots);
+
+					for (uint32 SkipIndex = 0; SkipIndex < NumSkippedSlots; ++SkipIndex)
+					{
+						VkDescriptorAddressInfoEXT& DescriptorAddressInfo = StageUBs[ShaderStage::EStage::RayGen].AddZeroed_GetRef();
+						DescriptorAddressInfo.sType = VK_STRUCTURE_TYPE_BUFFER_DEVICE_ADDRESS_INFO;
+					}
+
+					NumSkippedSlots = 0;
+				}
+				
+				VkDescriptorAddressInfoEXT& DescriptorAddressInfo = StageUBs[ShaderStage::EStage::RayGen].AddZeroed_GetRef();
+				DescriptorAddressInfo.sType = VK_STRUCTURE_TYPE_BUFFER_DEVICE_ADDRESS_INFO;
+				DescriptorAddressInfo.address = UniformBuffer->GetDeviceAddress();
+				DescriptorAddressInfo.range = UniformBuffer->GetSize();
+			}
+			else
+			{
+				// :todo-jn: There might be unused indices (see USE_INSTANCE_CULLING_DATA issue), just skip them with a warning for now.
+				NumSkippedSlots++;
+			}
 		}
+		Device->GetBindlessDescriptorManager()->RegisterUniformBuffers(CmdBuffer->GetHandle(), VK_PIPELINE_BIND_POINT_RAY_TRACING_KHR, StageUBs);
 	}
-	return StageUBs;
 }
 
 
@@ -1648,11 +1774,10 @@ void FVulkanCommandListContext::RHIRayTraceDispatch(
 	FVulkanCmdBuffer* const CmdBuffer = GetCommandBufferManager()->GetActiveCmdBuffer();
 	VulkanRHI::vkCmdBindPipeline(CmdBuffer->GetHandle(), VK_PIPELINE_BIND_POINT_RAY_TRACING_KHR, Pipeline->GetPipeline());
 
-	// :todo-jn: set real uniform buffers, check for gaps, process all stages
-	FVulkanBindlessDescriptorManager::FUniformBufferDescriptorArrays StageUBs = GetStageUBs(Device, InGlobalResourceBindings);
-	Device->GetBindlessDescriptorManager()->RegisterUniformBuffers(CmdBuffer->GetHandle(), VK_PIPELINE_BIND_POINT_RAY_TRACING_KHR, StageUBs);
+	SetRayGenResources(Device, CmdBuffer, InGlobalResourceBindings);
 
 	ShaderTable->SetSlot(InRayGenShader->GetFrequency(), 0, Pipeline->GetShaderIndex(RayGenShader), Pipeline->GetShaderHandles(SF_RayGen));
+	ShaderTable->Commit(*this);
 
 	VulkanRHI::vkCmdTraceRaysKHR(
 		CmdBuffer->GetHandle(),
@@ -1680,15 +1805,14 @@ void FVulkanCommandListContext::RHIRayTraceDispatchIndirect(
 	FVulkanCmdBuffer* const CmdBuffer = GetCommandBufferManager()->GetActiveCmdBuffer();
 	VulkanRHI::vkCmdBindPipeline(CmdBuffer->GetHandle(), VK_PIPELINE_BIND_POINT_RAY_TRACING_KHR, Pipeline->GetPipeline());
 
-	// :todo-jn: set real uniform buffers
-	FVulkanBindlessDescriptorManager::FUniformBufferDescriptorArrays StageUBs = GetStageUBs(Device, InGlobalResourceBindings);
-	Device->GetBindlessDescriptorManager()->RegisterUniformBuffers(CmdBuffer->GetHandle(), VK_PIPELINE_BIND_POINT_RAY_TRACING_KHR, StageUBs);
+	SetRayGenResources(Device, CmdBuffer, InGlobalResourceBindings);
+
+	ShaderTable->SetSlot(InRayGenShader->GetFrequency(), 0, Pipeline->GetShaderIndex(RayGenShader), Pipeline->GetShaderHandles(SF_RayGen));
+	ShaderTable->Commit(*this);
 
 	FVulkanResourceMultiBuffer* ArgumentBuffer = ResourceCast(InArgumentBuffer);
 	const VkDeviceAddress IndirectDeviceAddress = ArgumentBuffer->GetDeviceAddress() + InArgumentOffset;
 
-	ShaderTable->SetSlot(InRayGenShader->GetFrequency(), 0, Pipeline->GetShaderIndex(RayGenShader), Pipeline->GetShaderHandles(SF_RayGen));
-
 	VulkanRHI::vkCmdTraceRaysIndirectKHR(
 		CmdBuffer->GetHandle(),
 		ShaderTable->GetRegion(SF_RayGen),
@@ -1698,6 +1822,77 @@ void FVulkanCommandListContext::RHIRayTraceDispatchIndirect(
 		IndirectDeviceAddress);
 }
 
+static void SetSystemParametersUB(FVulkanHitGroupSystemParameters& OutSystemParameters, uint32 InNumUniformBuffers, FRHIUniformBuffer* const* InUniformBuffers, const FVulkanRayTracingShader* InShader)
+{
+	// Plug the shaders in the right slots using LayoutHash comparisons
+	// :todo-jn: Find a more direct approach
+	for (uint32 UBIndex = 0; UBIndex < InNumUniformBuffers; ++UBIndex)
+	{
+		FVulkanUniformBuffer* UniformBuffer = ResourceCast(InUniformBuffers[UBIndex]);
+		if (UniformBuffer)
+		{
+			const uint32 LayoutHash = UniformBuffer->GetLayout().GetHash();
+
+			for (int32 TargetIndex = 0; TargetIndex < InShader->GetCodeHeader().UniformBuffers.Num(); ++TargetIndex)
+			{
+				if (InShader->GetCodeHeader().UniformBuffers[TargetIndex].LayoutHash == LayoutHash)
+				{
+					const FRHIDescriptorHandle BindlessHandle = UniformBuffer->GetBindlessHandle();
+					check(BindlessHandle.IsValid());
+					OutSystemParameters.BindlessUniformBuffers[TargetIndex] = BindlessHandle.GetIndex();
+					break;
+				}
+			}
+		}
+	}
+}
+
+
+static void SetRayTracingHitGroup(
+	FVulkanCommandListContext& CommandContext,
+	FVulkanDevice* Device,
+	FVulkanRayTracingShaderTable* ShaderTable,
+	FVulkanRayTracingScene* Scene,
+	FVulkanRayTracingPipelineState* Pipeline,
+	uint32 InstanceIndex, uint32 SegmentIndex, uint32 ShaderSlot, uint32 HitGroupIndex,
+	uint32 NumUniformBuffers, FRHIUniformBuffer* const* UniformBuffers,
+	uint32 LooseParameterDataSize, const void* LooseParameterData,
+	uint32 UserData,
+	uint32 WorkerIndex)
+{
+	const FRayTracingSceneInitializer2& SceneInitializer = Scene->GetInitializer();
+
+	checkf(ShaderSlot < SceneInitializer.ShaderSlotsPerGeometrySegment, TEXT("Shader slot is invalid. Make sure that ShaderSlotsPerGeometrySegment is correct on FRayTracingSceneInitializer."));
+
+	const uint32 RecordIndex = Scene->GetHitRecordBaseIndex(InstanceIndex, SegmentIndex) + ShaderSlot;
+
+#if DO_CHECK
+	{
+		const uint32 NumSceneInstances = (uint32)SceneInitializer.PerInstanceGeometries.Num();
+		checkf(InstanceIndex < NumSceneInstances, TEXT("Instance index %d is out of range for the scene that contains %d instances"), InstanceIndex, NumSceneInstances);
+
+		const FVulkanRayTracingGeometry* Geometry = ResourceCast(SceneInitializer.PerInstanceGeometries[InstanceIndex]);
+		const uint32 NumGeometrySegments = Geometry->GetNumSegments();
+		checkf(SegmentIndex < NumGeometrySegments, TEXT("Segment %d is out of range for ray tracing geometry '%s' that contains %d segments"),
+			SegmentIndex, Geometry->DebugName.IsNone() ? TEXT("UNKNOWN") : *Geometry->DebugName.ToString(), NumGeometrySegments);
+	}
+#endif // DO_CHECK
+
+	const uint32 PrefixedSegmentIndex = SceneInitializer.SegmentPrefixSum[InstanceIndex];
+	const FVulkanRayTracingShader* Shader = Pipeline->GetShader(SF_RayHitGroup, HitGroupIndex);
+
+	const FVulkanRayTracingGeometry* Geometry = ResourceCast(SceneInitializer.PerInstanceGeometries[InstanceIndex]);
+	FVulkanHitGroupSystemParameters SystemParameters = Geometry->HitGroupSystemParameters[SegmentIndex];
+	SystemParameters.RootConstants.BaseInstanceIndex = SceneInitializer.BaseInstancePrefixSum[InstanceIndex];
+	SystemParameters.RootConstants.UserData = UserData;
+	SetSystemParametersUB(SystemParameters, NumUniformBuffers, UniformBuffers, Shader);
+
+	ShaderTable->SetLocalShaderParameters(SF_RayHitGroup, RecordIndex, 0, SystemParameters);
+
+	ShaderTable->SetSlot(SF_RayHitGroup, RecordIndex, HitGroupIndex, Pipeline->GetShaderHandles(SF_RayHitGroup));
+}
+
+
 void FVulkanCommandListContext::RHISetRayTracingHitGroup(
 	FRHIRayTracingScene* InScene, uint32 InstanceIndex, uint32 SegmentIndex, uint32 ShaderSlot,
 	FRHIRayTracingPipelineState* InPipeline, uint32 HitGroupIndex,
@@ -1705,7 +1900,48 @@ void FVulkanCommandListContext::RHISetRayTracingHitGroup(
 	uint32 LooseParameterDataSize, const void* LooseParameterData,
 	uint32 UserData)
 {
-	checkNoEntry();
+	FVulkanRayTracingScene* Scene = ResourceCast(InScene);
+	FVulkanRayTracingPipelineState* Pipeline = ResourceCast(InPipeline);
+	FVulkanRayTracingShaderTable* ShaderTable = Scene->FindOrCreateShaderTable(Pipeline);
+
+	checkf(ShaderSlot < Scene->GetInitializer().ShaderSlotsPerGeometrySegment, TEXT("Shader slot is invalid. Make sure that ShaderSlotsPerGeometrySegment is correct on FRayTracingSceneInitializer."));
+
+	const uint32 WorkerIndex = 0;
+
+	SetRayTracingHitGroup(*this, Device, ShaderTable, Scene, Pipeline,
+		InstanceIndex,
+		SegmentIndex,
+		ShaderSlot,
+		HitGroupIndex,
+		NumUniformBuffers,
+		UniformBuffers,
+		LooseParameterDataSize,
+		LooseParameterData,
+		UserData,
+		WorkerIndex);
+}
+
+
+
+static void SetGenericSystemParameters(
+	FRHIRayTracingScene* InScene, uint32 ShaderSlotInScene,
+	FRHIRayTracingPipelineState* InPipeline, uint32 ShaderIndexInPipeline,
+	uint32 NumUniformBuffers, FRHIUniformBuffer* const* UniformBuffers,
+	uint32 UserData, const EShaderFrequency ShaderFrequency)
+{
+	FVulkanRayTracingScene* Scene = ResourceCast(InScene);
+	FVulkanRayTracingPipelineState* Pipeline = ResourceCast(InPipeline);
+	const uint32 WorkerIndex = 0;
+	FVulkanRayTracingShaderTable* ShaderTable = Scene->FindOrCreateShaderTable(Pipeline);
+	const FVulkanRayTracingShader* Shader = Pipeline->GetShader(ShaderFrequency, ShaderIndexInPipeline);
+
+	FVulkanHitGroupSystemParameters SystemParameters;
+	FMemory::Memzero(SystemParameters);
+	SystemParameters.RootConstants.UserData = UserData;
+	SetSystemParametersUB(SystemParameters, NumUniformBuffers, UniformBuffers, Shader);
+	ShaderTable->SetLocalShaderParameters(ShaderFrequency, ShaderSlotInScene, 0, SystemParameters);
+
+	ShaderTable->SetSlot(ShaderFrequency, ShaderSlotInScene, ShaderIndexInPipeline, Pipeline->GetShaderHandles(ShaderFrequency));
 }
 
 void FVulkanCommandListContext::RHISetRayTracingCallableShader(
@@ -1714,7 +1950,8 @@ void FVulkanCommandListContext::RHISetRayTracingCallableShader(
 	uint32 NumUniformBuffers, FRHIUniformBuffer* const* UniformBuffers,
 	uint32 UserData)
 {
-	checkNoEntry();
+	SetGenericSystemParameters(InScene, ShaderSlotInScene, InPipeline, ShaderIndexInPipeline,
+		NumUniformBuffers, UniformBuffers, UserData, SF_RayCallable);
 }
 
 void FVulkanCommandListContext::RHISetRayTracingMissShader(
@@ -1723,7 +1960,8 @@ void FVulkanCommandListContext::RHISetRayTracingMissShader(
 	uint32 NumUniformBuffers, FRHIUniformBuffer* const* UniformBuffers,
 	uint32 UserData)
 {
-	checkNoEntry();
+	SetGenericSystemParameters(InScene, ShaderSlotInScene, InPipeline, ShaderIndexInPipeline,
+		NumUniformBuffers, UniformBuffers, UserData, SF_RayMiss);
 }
 
 
@@ -1738,7 +1976,70 @@ void FVulkanCommandListContext::RHISetRayTracingBindings(
 
 	checkf(Scene->IsBuilt(), TEXT("Ray tracing scene must be built before any shaders can be bound to it. Make sure that RHIBuildAccelerationStructure() command has been executed."));
 
-	checkNoEntry();
+	FGraphEventArray TaskList;
+
+	const uint32 NumWorkerThreads = FTaskGraphInterface::Get().GetNumWorkerThreads();
+	const uint32 MaxTasks = FApp::ShouldUseThreadingForPerformance() ? FMath::Min<uint32>(NumWorkerThreads, FVulkanRayTracingScene::MaxBindingWorkers) : 1;
+
+	struct FTaskContext
+	{
+		uint32 WorkerIndex = 0;
+	};
+
+	TArray<FTaskContext, TInlineAllocator<FVulkanRayTracingScene::MaxBindingWorkers>> TaskContexts;
+	for (uint32 WorkerIndex = 0; WorkerIndex < MaxTasks; ++WorkerIndex)
+	{
+		TaskContexts.Add(FTaskContext{ WorkerIndex });
+	}
+
+	auto BindingTask = [this, Bindings, Device = Device, Scene, Pipeline, ShaderTable, BindingType](const FTaskContext& Context, int32 CurrentIndex)
+	{
+		const FRayTracingLocalShaderBindings& Binding = Bindings[CurrentIndex];
+
+		if (BindingType == ERayTracingBindingType::HitGroup)
+		{
+			SetRayTracingHitGroup(*this, Device, ShaderTable, Scene, Pipeline,
+				Binding.InstanceIndex,
+				Binding.SegmentIndex,
+				Binding.ShaderSlot,
+				Binding.ShaderIndexInPipeline,
+				Binding.NumUniformBuffers,
+				Binding.UniformBuffers,
+				Binding.LooseParameterDataSize,
+				Binding.LooseParameterData,
+				Binding.UserData,
+				Context.WorkerIndex);
+		}
+		else if (BindingType == ERayTracingBindingType::CallableShader)
+		{
+			SetGenericSystemParameters(
+				Scene, Binding.ShaderSlot,
+				Pipeline, Binding.ShaderIndexInPipeline,
+				Binding.NumUniformBuffers, Binding.UniformBuffers,
+				Binding.UserData,
+				SF_RayCallable);
+		}
+		else if (BindingType == ERayTracingBindingType::MissShader)
+		{
+			SetGenericSystemParameters(
+				Scene, Binding.ShaderSlot,
+				Pipeline, Binding.ShaderIndexInPipeline,
+				Binding.NumUniformBuffers, Binding.UniformBuffers,
+				Binding.UserData,
+				SF_RayMiss);
+		}
+		else
+		{
+			checkNoEntry();
+		}
+	};
+
+	// One helper worker task will be created at most per this many work items, plus one worker for current thread (unless running on a task thread),
+	// up to a hard maximum of FD3D12RayTracingScene::MaxBindingWorkers.
+	// Internally, parallel for tasks still subdivide the work into smaller chunks and perform fine-grained load-balancing.
+	const int32 ItemsPerTask = 1024;
+
+	ParallelForWithExistingTaskContext(TEXT("SetRayTracingBindings"), MakeArrayView(TaskContexts), NumBindings, ItemsPerTask, BindingTask);
 }
 
 #endif // VULKAN_RHI_RAYTRACING
diff --git a/Engine/Source/Runtime/VulkanRHI/Private/VulkanRayTracing.h b/Engine/Source/Runtime/VulkanRHI/Private/VulkanRayTracing.h
--- a/Engine/Source/Runtime/VulkanRHI/Private/VulkanRayTracing.h
+++ b/Engine/Source/Runtime/VulkanRHI/Private/VulkanRayTracing.h
@@ -79,6 +79,8 @@ public:
 
 	void SetLocalShaderParameters(EShaderFrequency Frequency, uint32 RecordIndex, uint32 OffsetWithinRecord, const void* InData, uint32 InDataSize);
 
+	void Commit(FVulkanCommandListContext& Context);
+
 private:
 
 	struct FVulkanShaderTableAllocation
@@ -91,13 +93,19 @@ private:
 		uint32 HandleCount = 0;
 		bool bUseLocalRecord = false;
 
-		VkBuffer Buffer = VK_NULL_HANDLE;
-		VulkanRHI::FVulkanAllocation Allocation;
+		// Host memory copy
+		TArray<uint8> HostBuffer;
+
+		// GPU Local memory copy
+		VkBuffer LocalBuffer = VK_NULL_HANDLE;
+		VulkanRHI::FVulkanAllocation LocalAllocation;
 		VkStridedDeviceAddressRegionKHR Region;
-		uint8* MappedBufferMemory = nullptr;
+
+		bool bIsDirty = true;
 	};
 
 	FVulkanShaderTableAllocation& GetAlloc(EShaderFrequency Frequency);
+	static void ReleaseLocalBuffer(FVulkanDevice* Device, FVulkanShaderTableAllocation& Alloc);
 
 	FVulkanShaderTableAllocation Raygen;
 	FVulkanShaderTableAllocation Miss;
@@ -151,6 +159,11 @@ public:
 class FVulkanRayTracingScene : public FRHIRayTracingScene, public VulkanRHI::FDeviceChild
 {
 public:
+	// Ray tracing shader bindings can be processed in parallel.
+	// Each concurrent worker gets its own dedicated descriptor cache instance to avoid contention or locking.
+	// Scaling beyond 5 total threads does not yield any speedup in practice (RHI thread + 4 parallel workers).
+	static constexpr uint32 MaxBindingWorkers = 1; // :todo-jn:
+
 	FVulkanRayTracingScene(FRayTracingSceneInitializer2 Initializer, FVulkanDevice* InDevice);
 	~FVulkanRayTracingScene();
 
diff --git a/Engine/Source/Runtime/VulkanRHI/Private/VulkanResources.h b/Engine/Source/Runtime/VulkanRHI/Private/VulkanResources.h
--- a/Engine/Source/Runtime/VulkanRHI/Private/VulkanResources.h
+++ b/Engine/Source/Runtime/VulkanRHI/Private/VulkanResources.h
@@ -1180,12 +1180,16 @@ public:
 		NewAlloc.Swap(Allocation);
 	}
 	
+	FRHIDescriptorHandle GetBindlessHandle();
 	VkDeviceAddress GetDeviceAddress() const;
 
 public:
 	FVulkanDevice* Device;
 	VulkanRHI::FVulkanAllocation Allocation;
 	EUniformBufferUsage Usage;
+
+	FRHIDescriptorHandle BindlessHandle;
+	VkDeviceAddress CachedDeviceAddress = 0;
 };
 
 class FVulkanUnorderedAccessView final : public FRHIUnorderedAccessView, public FVulkanLinkedView
diff --git a/Engine/Source/Runtime/VulkanRHI/Private/VulkanUniformBuffer.cpp b/Engine/Source/Runtime/VulkanRHI/Private/VulkanUniformBuffer.cpp
--- a/Engine/Source/Runtime/VulkanRHI/Private/VulkanUniformBuffer.cpp
+++ b/Engine/Source/Runtime/VulkanRHI/Private/VulkanUniformBuffer.cpp
@@ -9,6 +9,7 @@
 #include "VulkanLLM.h"
 #include "ShaderParameterStruct.h"
 #include "RHIUniformBufferDataShared.h"
+#include "VulkanDescriptorSets.h"
 
 static int32 GVulkanAllowUniformUpload = 1;
 static FAutoConsoleVariableRef CVarVulkanAllowUniformUpload(
@@ -188,6 +189,11 @@ PRAGMA_ENABLE_DEPRECATION_WARNINGS
 
 FVulkanUniformBuffer::~FVulkanUniformBuffer()
 {
+	if (BindlessHandle.IsValid())
+	{
+		Device->GetBindlessDescriptorManager()->Unregister(BindlessHandle);
+	}
+
 	Device->GetMemoryManager().FreeUniformBuffer(Allocation);
 }
 
@@ -212,6 +218,26 @@ void FVulkanUniformBuffer::UpdateResourceTable(FRHIResource** Resources, int32 R
 	}
 }
 
+FRHIDescriptorHandle FVulkanUniformBuffer::GetBindlessHandle()
+{
+	// :todo-jn: temporary code to refresh as needed, only used by raytracing
+	const VkDeviceAddress CurrentAddress = GetDeviceAddress();
+	if (!BindlessHandle.IsValid() || (CachedDeviceAddress == 0) || (CurrentAddress != CachedDeviceAddress))
+	{
+		if (BindlessHandle.IsValid())
+		{
+			Device->GetBindlessDescriptorManager()->Unregister(BindlessHandle);
+		}
+
+		BindlessHandle = Device->GetBindlessDescriptorManager()->ReserveDescriptor(VK_DESCRIPTOR_TYPE_UNIFORM_BUFFER);
+		Device->GetBindlessDescriptorManager()->UpdateBuffer(BindlessHandle, CurrentAddress, GetSize(), true);
+
+		CachedDeviceAddress = CurrentAddress;
+	}
+
+	return BindlessHandle;
+}
+
 VkDeviceAddress FVulkanUniformBuffer::GetDeviceAddress() const
 {
 	// :todo-jn: there will be more and more churn on this, cache the value
@@ -223,7 +249,6 @@ VkDeviceAddress FVulkanUniformBuffer::GetDeviceAddress() const
 	return BufferAddress;
 }
 
-
 FUniformBufferRHIRef FVulkanDynamicRHI::RHICreateUniformBuffer(const void* Contents, const FRHIUniformBufferLayout* Layout, EUniformBufferUsage Usage, EUniformBufferValidation Validation)
 {
 	LLM_SCOPE_VULKAN(ELLMTagVulkan::VulkanUniformBuffers);
