Adds bSupportsRealTypes=RuntimeDependent in Vulkan SM5 and SM6

diff --git a/Engine/Config/VulkanPC/DataDrivenPlatformInfo.ini b/Engine/Config/VulkanPC/DataDrivenPlatformInfo.ini
--- a/Engine/Config/VulkanPC/DataDrivenPlatformInfo.ini
+++ b/Engine/Config/VulkanPC/DataDrivenPlatformInfo.ini
@@ -134,6 +134,7 @@ FriendlyName=LOCTEXT("FriendlyShaderPlatformName_VulkanPC_SM5", "PC Vulkan SM5")
 
 BindlessSupport=Unsupported
 bSupportsNNEShaders=true
+bSupportsRealTypes=RuntimeDependent
 
 
 
@@ -250,6 +251,7 @@ FriendlyName=LOCTEXT("FriendlyShaderPlatformName_VulkanPC_SM6", "PC Vulkan SM6")
 
 BindlessSupport=AllShaderTypes
 bSupportsNNEShaders=true
+bSupportsRealTypes=RuntimeDependent
 
 
 [ShaderPlatform VULKAN_PCES3_1]
diff --git a/Engine/Source/Developer/ShaderCompilerCommon/Private/ShaderConductorContext.cpp b/Engine/Source/Developer/ShaderCompilerCommon/Private/ShaderConductorContext.cpp
--- a/Engine/Source/Developer/ShaderCompilerCommon/Private/ShaderConductorContext.cpp
+++ b/Engine/Source/Developer/ShaderCompilerCommon/Private/ShaderConductorContext.cpp
@@ -389,22 +389,29 @@ namespace CrossCompiler
 		}
 
 		using ETargetEnvironment = CrossCompiler::FShaderConductorOptions::ETargetEnvironment;
-		switch (InOptions.TargetEnvironment)
+		if (OutOptions.enable16bitTypes)
 		{
-		default:
-			checkf(false, TEXT("Unexpected SPIR-V target environment: %d"), (uint32)InOptions.TargetEnvironment);
-		case ETargetEnvironment::Vulkan_1_0:
-			DxcArgRefs.Add("-fspv-target-env=vulkan1.0");
-			break;
-		case ETargetEnvironment::Vulkan_1_1:
-			DxcArgRefs.Add("-fspv-target-env=vulkan1.1");
-			break;
-		case ETargetEnvironment::Vulkan_1_2:
-			DxcArgRefs.Add("-fspv-target-env=vulkan1.2");
-			break;
-		case ETargetEnvironment::Vulkan_1_3:
-			DxcArgRefs.Add("-fspv-target-env=vulkan1.3");
-			break;
+			DxcArgRefs.Add("-fspv-target-env=universal1.5");
+		}
+		else
+		{
+			switch (InOptions.TargetEnvironment)
+			{
+			default:
+				checkf(false, TEXT("Unexpected SPIR-V target environment: %d"), (uint32)InOptions.TargetEnvironment);
+			case ETargetEnvironment::Vulkan_1_0:
+				DxcArgRefs.Add("-fspv-target-env=vulkan1.0");
+				break;
+			case ETargetEnvironment::Vulkan_1_1:
+				DxcArgRefs.Add("-fspv-target-env=vulkan1.1");
+				break;
+			case ETargetEnvironment::Vulkan_1_2:
+				DxcArgRefs.Add("-fspv-target-env=vulkan1.2");
+				break;
+			case ETargetEnvironment::Vulkan_1_3:
+				DxcArgRefs.Add("-fspv-target-env=vulkan1.3");
+				break;
+			}
 		}
 
 		if (DxcArgRefs.Num() > 0)
@@ -547,6 +554,17 @@ namespace CrossCompiler
 		ShaderConductor::Compiler::Options ScOptions;
 		ConvertScOptions(Options, ScOptions, {}, Intermediates->DxcArgRefs);
 
+		// ShaderConductor.cpp forgot to pipedown enable16bitTypes, so work arround by adding the parameter manually in here.
+		TArray<ANSICHAR const*> PatchedDxcArgRefs;
+		if (ScOptions.enable16bitTypes)
+		{
+			PatchedDxcArgRefs = Intermediates->DxcArgRefs;
+			PatchedDxcArgRefs.Add("-enable-16bit-types");
+
+			ScOptions.numDXCArgs = PatchedDxcArgRefs.Num();
+			ScOptions.DXCArgs = (const char**)PatchedDxcArgRefs.GetData();
+		}
+
 		// Rewrite HLSL with wrapper function to catch exceptions from ShaderConductor
 		bool bSucceeded = false;
 		ShaderConductor::Compiler::ResultDesc ResultDesc;
diff --git a/Engine/Source/Developer/VulkanShaderFormat/Private/VulkanShaderCompiler.cpp b/Engine/Source/Developer/VulkanShaderFormat/Private/VulkanShaderCompiler.cpp
--- a/Engine/Source/Developer/VulkanShaderFormat/Private/VulkanShaderCompiler.cpp
+++ b/Engine/Source/Developer/VulkanShaderFormat/Private/VulkanShaderCompiler.cpp
@@ -2374,6 +2374,11 @@ static bool CompileWithShaderConductor(
 	Options.bDisableScalarBlockLayout = !(bIsRayTracingShader || 
 		MinTargetEnvironment >= CrossCompiler::FShaderConductorOptions::ETargetEnvironment::Vulkan_1_3);
 
+	if (Input.Environment.CompilerFlags.Contains(CFLAG_AllowRealTypes))
+	{
+		Options.bEnable16bitTypes = true;
+	}
+
 	// Enable HLSL 2021 if specified
 	if (Input.Environment.CompilerFlags.Contains(CFLAG_HLSL2021))
 	{
@@ -2623,6 +2628,11 @@ void DoCompileVulkanShader(const FShaderCompilerInput& Input, FShaderCompilerOut
 		AdditionalDefines.SetDefine(TEXT("PLATFORM_SUPPORTS_INLINE_RAY_TRACING"), 1);
 	}
 
+	if (Input.Environment.CompilerFlags.Contains(CFLAG_AllowRealTypes))
+	{
+		AdditionalDefines.SetDefine(TEXT("PLATFORM_SUPPORTS_REAL_TYPES"), 1);
+	}
+
 	if (MinTargetEnvironment >= CrossCompiler::FShaderConductorOptions::ETargetEnvironment::Vulkan_1_1)
 	{
 		AdditionalDefines.SetDefine(TEXT("PLATFORM_SUPPORTS_SM6_0_WAVE_OPERATIONS"), 1);
diff --git a/Engine/Source/Runtime/VulkanRHI/Private/VulkanDevice.h b/Engine/Source/Runtime/VulkanRHI/Private/VulkanDevice.h
--- a/Engine/Source/Runtime/VulkanRHI/Private/VulkanDevice.h
+++ b/Engine/Source/Runtime/VulkanRHI/Private/VulkanDevice.h
@@ -62,6 +62,7 @@ struct FOptionalVulkanDeviceExtensions
 
 			// Promoted to 1.1
 			uint64 HasKHRMultiview : 1;
+			uint64 HasKHR16bitStorage : 1;
 
 			// Promoted to 1.2
 			uint64 HasKHRRenderPass2 : 1;
@@ -71,6 +72,7 @@ struct FOptionalVulkanDeviceExtensions
 			uint64 HasBufferDeviceAddress : 1;
 			uint64 HasSPIRV_14 : 1;
 			uint64 HasShaderFloatControls : 1;
+			uint64 HasKHRShaderFloat16 : 1;
 			uint64 HasEXTDescriptorIndexing : 1;
 			uint64 HasEXTShaderViewportIndexLayer : 1;
 			uint64 HasSeparateDepthStencilLayouts : 1;
diff --git a/Engine/Source/Runtime/VulkanRHI/Private/VulkanExtensions.cpp b/Engine/Source/Runtime/VulkanRHI/Private/VulkanExtensions.cpp
--- a/Engine/Source/Runtime/VulkanRHI/Private/VulkanExtensions.cpp
+++ b/Engine/Source/Runtime/VulkanRHI/Private/VulkanExtensions.cpp
@@ -35,6 +35,15 @@ TAutoConsoleVariable<int32> GRHIAllow64bitShaderAtomicsCvar(
 	ECVF_ReadOnly
 );
 
+TAutoConsoleVariable<int32> GRHIAllow16bitOps(
+	TEXT("r.Vulkan.Allow16bitOps"),
+	0,
+	TEXT("Whether to enable 16bit ops to speeds up TSR\n")
+	TEXT("0 to disable (default)\n")
+	TEXT("1 to enable"),
+	ECVF_ReadOnly
+);
+
 TAutoConsoleVariable<int32> GVulkanRayTracingCVar(
 	TEXT("r.Vulkan.RayTracing"),
 	1,
@@ -1178,6 +1187,81 @@ private:
 	VkPhysicalDeviceDescriptorBufferFeaturesEXT DescriptorBufferFeatures;
 };
 
+// ***** VK_KHR_16bit_storage
+class FVulkanKHR16BitStorageExtension : public FVulkanDeviceExtension
+{
+public:
+
+	FVulkanKHR16BitStorageExtension(FVulkanDevice* InDevice)
+		: FVulkanDeviceExtension(InDevice, VK_KHR_16BIT_STORAGE_EXTENSION_NAME, VULKAN_EXTENSION_ENABLED)
+	{
+		bEnabledInCode = bEnabledInCode && (GRHIAllow16bitOps.GetValueOnAnyThread() != 0);
+	}
+
+	virtual void PrePhysicalDeviceFeatures(VkPhysicalDeviceFeatures2KHR& PhysicalDeviceFeatures2) override final
+	{
+		ZeroVulkanStruct(Device16BitStorageFeatures, VK_STRUCTURE_TYPE_PHYSICAL_DEVICE_16BIT_STORAGE_FEATURES_KHR);
+		AddToPNext(PhysicalDeviceFeatures2, Device16BitStorageFeatures);
+	}
+
+	virtual void PostPhysicalDeviceFeatures(FOptionalVulkanDeviceExtensions& ExtensionFlags) override final
+	{
+		bRequirementsPassed = (
+			Device16BitStorageFeatures.storageBuffer16BitAccess == VK_TRUE &&
+			Device16BitStorageFeatures.uniformAndStorageBuffer16BitAccess == VK_TRUE &&
+			Device16BitStorageFeatures.storagePushConstant16 == VK_TRUE);
+		ExtensionFlags.HasKHR16bitStorage = bRequirementsPassed;
+	}
+
+	virtual void PreCreateDevice(VkDeviceCreateInfo& DeviceCreateInfo) override final
+	{
+		if (bRequirementsPassed)
+		{
+			AddToPNext(DeviceCreateInfo, Device16BitStorageFeatures);
+		}
+	}
+
+	VkPhysicalDevice16BitStorageFeaturesKHR Device16BitStorageFeatures;
+};
+
+// ***** VK_KHR_shader_float16_int8
+class FVulkanKHRShaderFloat16Int8Extension : public FVulkanDeviceExtension
+{
+public:
+
+	FVulkanKHRShaderFloat16Int8Extension(FVulkanDevice* InDevice)
+		: FVulkanDeviceExtension(InDevice, VK_KHR_SHADER_FLOAT16_INT8_EXTENSION_NAME, VULKAN_EXTENSION_ENABLED)
+	{
+		bEnabledInCode = bEnabledInCode && (GRHIAllow16bitOps.GetValueOnAnyThread() != 0);
+	}
+
+	virtual void PrePhysicalDeviceFeatures(VkPhysicalDeviceFeatures2KHR& PhysicalDeviceFeatures2) override final
+	{
+		ZeroVulkanStruct(ShaderFloat16Int8Features, VK_STRUCTURE_TYPE_PHYSICAL_DEVICE_SHADER_FLOAT16_INT8_FEATURES_KHR);
+		AddToPNext(PhysicalDeviceFeatures2, ShaderFloat16Int8Features);
+	}
+
+	virtual void PostPhysicalDeviceFeatures(FOptionalVulkanDeviceExtensions& ExtensionFlags) override final
+	{
+		bRequirementsPassed = (ShaderFloat16Int8Features.shaderFloat16 == VK_TRUE);
+		ExtensionFlags.HasKHRShaderFloat16 = bRequirementsPassed;
+	}
+
+	virtual void PreCreateDevice(VkDeviceCreateInfo& DeviceCreateInfo) override final
+	{
+		// The PreCreateDevice() call is after all extension have gone through PostPhysicalDeviceFeatures(), so ExtensionFlags will be filled for both.
+		const FOptionalVulkanDeviceExtensions& ExtensionFlags = Device->GetOptionalExtensions();
+		GRHIGlobals.SupportsNative16BitOps = ExtensionFlags.HasKHR16bitStorage && ExtensionFlags.HasKHRShaderFloat16;
+
+		if (bRequirementsPassed)
+		{
+			AddToPNext(DeviceCreateInfo, ShaderFloat16Int8Features);
+		}
+	}
+
+	VkPhysicalDeviceShaderFloat16Int8Features  ShaderFloat16Int8Features;
+};
+
 
 
 
@@ -1258,6 +1342,8 @@ FVulkanDeviceExtensionArray FVulkanDeviceExtension::GetUESupportedDeviceExtensio
 	ADD_CUSTOM_EXTENSION(FVulkanEXTDescriptorBuffer);
 	ADD_CUSTOM_EXTENSION(FVulkanEXTDeviceFaultExtension);
 	ADD_CUSTOM_EXTENSION(FVulkanEXTShaderDemoteToHelperInvocationExtension);
+	ADD_CUSTOM_EXTENSION(FVulkanKHR16BitStorageExtension);
+	ADD_CUSTOM_EXTENSION(FVulkanKHRShaderFloat16Int8Extension);
 
 	// Needed for Raytracing
 	ADD_CUSTOM_EXTENSION(FVulkanKHRBufferDeviceAddressExtension);
