VulkanRHI: Fix GPU crash in RT on Vulkan with missing barrier to compute

diff --git a/Engine/Source/Runtime/VulkanRHI/Private/VulkanRayTracing.cpp b/Engine/Source/Runtime/VulkanRHI/Private/VulkanRayTracing.cpp
--- a/Engine/Source/Runtime/VulkanRHI/Private/VulkanRayTracing.cpp
+++ b/Engine/Source/Runtime/VulkanRHI/Private/VulkanRayTracing.cpp
@@ -127,7 +127,12 @@ static void AddAccelerationStructureBuildBarrier(VkCommandBuffer CommandBuffer)
 	ZeroVulkanStruct(Barrier, VK_STRUCTURE_TYPE_MEMORY_BARRIER);
 	Barrier.srcAccessMask = VK_ACCESS_ACCELERATION_STRUCTURE_WRITE_BIT_KHR;
 	Barrier.dstAccessMask = VK_ACCESS_ACCELERATION_STRUCTURE_READ_BIT_KHR | VK_ACCESS_ACCELERATION_STRUCTURE_WRITE_BIT_KHR;
-	VulkanRHI::vkCmdPipelineBarrier(CommandBuffer, VK_PIPELINE_STAGE_ACCELERATION_STRUCTURE_BUILD_BIT_KHR, VK_PIPELINE_STAGE_ACCELERATION_STRUCTURE_BUILD_BIT_KHR, 0, 1, &Barrier, 0, nullptr, 0, nullptr);
+    
+    // TODO: Revisit the compute stages here as we don't always need barrier to compute
+    VkPipelineStageFlags srcStage = VK_PIPELINE_STAGE_ACCELERATION_STRUCTURE_BUILD_BIT_KHR | VK_PIPELINE_STAGE_COMPUTE_SHADER_BIT;
+    VkPipelineStageFlags dstStage = VK_PIPELINE_STAGE_ACCELERATION_STRUCTURE_BUILD_BIT_KHR | VK_PIPELINE_STAGE_COMPUTE_SHADER_BIT;
+    
+	VulkanRHI::vkCmdPipelineBarrier(CommandBuffer, srcStage, dstStage, 0, 1, &Barrier, 0, nullptr, 0, nullptr);
 }
 
 static bool ShouldCompactAfterBuild(ERayTracingAccelerationStructureFlags BuildFlags)
