Add InstanceProcAddr to Public Vulkan Interface

diff --git a/Engine/Source/Runtime/VulkanRHI/Private/VulkanRHI.cpp b/Engine/Source/Runtime/VulkanRHI/Private/VulkanRHI.cpp
--- a/Engine/Source/Runtime/VulkanRHI/Private/VulkanRHI.cpp
+++ b/Engine/Source/Runtime/VulkanRHI/Private/VulkanRHI.cpp
@@ -1364,6 +1364,11 @@ void* FVulkanDynamicRHI::RHIGetVkDeviceProcAddr(const char* InName) const
 	return (void*)VulkanRHI::vkGetDeviceProcAddr(Device->GetInstanceHandle(), InName);
 }
 
+void* FVulkanDynamicRHI::RHIGetVkInstanceProcAddr(const char* InName) const
+{
+	return (void*)VulkanRHI::vkGetInstanceProcAddr(Instance, InName);
+}
+
 VkFormat FVulkanDynamicRHI::RHIGetSwapChainVkFormat(EPixelFormat InFormat) const
 {
 	// UE renders a gamma-corrected image so we need to use an sRGB format if available
diff --git a/Engine/Source/Runtime/VulkanRHI/Public/IVulkanDynamicRHI.h b/Engine/Source/Runtime/VulkanRHI/Public/IVulkanDynamicRHI.h
--- a/Engine/Source/Runtime/VulkanRHI/Public/IVulkanDynamicRHI.h
+++ b/Engine/Source/Runtime/VulkanRHI/Public/IVulkanDynamicRHI.h
@@ -54,6 +54,7 @@ struct IVulkanDynamicRHI : public FDynamicRHI
 	virtual uint64           RHIGetGraphicsAdapterLUID(VkPhysicalDevice InPhysicalDevice) const = 0;
 	virtual bool             RHIDoesAdapterMatchDevice(const void* InAdapterId) const = 0;
 	virtual void*            RHIGetVkDeviceProcAddr(const char* InName) const = 0;
+	virtual void*            RHIGetVkInstanceProcAddr(const char* InName) const = 0;
 	virtual VkFormat         RHIGetSwapChainVkFormat(EPixelFormat InFormat) const = 0;
 	virtual bool             RHISupportsEXTFragmentDensityMap2() const = 0;
 
diff --git a/Engine/Source/Runtime/VulkanRHI/Public/VulkanDynamicRHI.h b/Engine/Source/Runtime/VulkanRHI/Public/VulkanDynamicRHI.h
--- a/Engine/Source/Runtime/VulkanRHI/Public/VulkanDynamicRHI.h
+++ b/Engine/Source/Runtime/VulkanRHI/Public/VulkanDynamicRHI.h
@@ -66,6 +66,7 @@ public:
 	virtual uint64 RHIGetGraphicsAdapterLUID(VkPhysicalDevice InPhysicalDevice) const final override;
 	virtual bool RHIDoesAdapterMatchDevice(const void* InAdapterId) const final override;
 	virtual void* RHIGetVkDeviceProcAddr(const char* InName) const final override;
+	virtual void* RHIGetVkInstanceProcAddr(const char* InName) const final override;
 	virtual VkFormat RHIGetSwapChainVkFormat(EPixelFormat InFormat) const final override;
 	virtual bool RHISupportsEXTFragmentDensityMap2() const final override;
 	virtual TArray<VkExtensionProperties> RHIGetAllInstanceExtensions() const final override;
