VulkanRHI: Improve UB memory tracking in Vulkan

diff --git a/Engine/Source/Runtime/VulkanRHI/Private/VulkanMemory.cpp b/Engine/Source/Runtime/VulkanRHI/Private/VulkanMemory.cpp
--- a/Engine/Source/Runtime/VulkanRHI/Private/VulkanMemory.cpp
+++ b/Engine/Source/Runtime/VulkanRHI/Private/VulkanMemory.cpp
@@ -3638,16 +3638,20 @@ namespace VulkanRHI
 
 	void FMemoryManager::AllocUniformBuffer(FVulkanAllocation& OutAllocation, uint32 Size)
 	{
-		if(!AllocateBufferPooled(OutAllocation, nullptr, Size, 0, VK_BUFFER_USAGE_UNIFORM_BUFFER_BIT | VK_BUFFER_USAGE_TRANSFER_DST_BIT, VK_MEMORY_PROPERTY_HOST_VISIBLE_BIT | VK_MEMORY_PROPERTY_HOST_COHERENT_BIT, EVulkanAllocationMetaUniformBuffer, __FILE__, __LINE__))
+		if (!AllocateBufferPooled(OutAllocation, nullptr, Size, 0, VK_BUFFER_USAGE_UNIFORM_BUFFER_BIT | VK_BUFFER_USAGE_TRANSFER_DST_BIT, VK_MEMORY_PROPERTY_HOST_VISIBLE_BIT | VK_MEMORY_PROPERTY_HOST_COHERENT_BIT, EVulkanAllocationMetaUniformBuffer, __FILE__, __LINE__))
 		{
 			HandleOOM(false);
 			checkNoEntry();
 		}
+
+		INC_MEMORY_STAT_BY(STAT_UniformBufferMemory, OutAllocation.Size);
 	}
 	void FMemoryManager::FreeUniformBuffer(FVulkanAllocation& InAllocation)
 	{
-		if(InAllocation.HasAllocation())
+		if (InAllocation.HasAllocation())
 		{
+			DEC_MEMORY_STAT_BY(STAT_UniformBufferMemory, InAllocation.Size);
+
 			FScopeLock ScopeLock(&UBAllocations.CS);
 			ProcessPendingUBFreesNoLock(false);
 			FUBPendingFree& Pending = UBAllocations.PendingFree.AddDefaulted_GetRef();
diff --git a/Engine/Source/Runtime/VulkanRHI/Private/VulkanUniformBuffer.cpp b/Engine/Source/Runtime/VulkanRHI/Private/VulkanUniformBuffer.cpp
--- a/Engine/Source/Runtime/VulkanRHI/Private/VulkanUniformBuffer.cpp
+++ b/Engine/Source/Runtime/VulkanRHI/Private/VulkanUniformBuffer.cpp
@@ -403,11 +403,15 @@ FVulkanUniformBufferUploader::FVulkanUniformBufferUploader(FVulkanDevice* InDevi
 			CPUBuffer = new FVulkanRingBuffer(InDevice, PackedUniformsRingBufferSize, VK_BUFFER_USAGE_UNIFORM_BUFFER_BIT, VK_MEMORY_PROPERTY_HOST_VISIBLE_BIT | VK_MEMORY_PROPERTY_HOST_COHERENT_BIT);
 		}
 	}
+
+	INC_MEMORY_STAT_BY(STAT_UniformBufferMemory, PackedUniformsRingBufferSize);
 }
 
 FVulkanUniformBufferUploader::~FVulkanUniformBufferUploader()
 {
 	delete CPUBuffer;
+
+	DEC_MEMORY_STAT_BY(STAT_UniformBufferMemory, PackedUniformsRingBufferSize);
 }
 
 
