VulkanRHI: Generate unique filenames for Aftermath dumps.  Include "NotStarted" state in wait loop.

#rnx

diff --git a/Engine/Source/Runtime/VulkanRHI/Private/VulkanUtil.cpp b/Engine/Source/Runtime/VulkanRHI/Private/VulkanUtil.cpp
--- a/Engine/Source/Runtime/VulkanRHI/Private/VulkanUtil.cpp
+++ b/Engine/Source/Runtime/VulkanRHI/Private/VulkanUtil.cpp
@@ -832,6 +832,8 @@ void FVulkanGPUProfiler::DumpCrashMarkers(void* BufferData)
 #if NV_AFTERMATH
 void AftermathGpuCrashDumpCallback(const void* CrashDump, const uint32 CrashDumpSize, void* UserData)
 {
+	const FString DateTimeString = FDateTime::Now().ToString();
+
 	// Create a GPU crash dump decoder object for the GPU crash dump.
 	GFSDK_Aftermath_GpuCrashDump_Decoder Decoder = {};
 	{
@@ -854,7 +856,7 @@ void AftermathGpuCrashDumpCallback(const void* CrashDump, const uint32 CrashDump
 	}
 
 	{
-		FString Filename = FPaths::ProjectLogDir() / TEXT("vulkan.nv-gpudmp");
+		const FString Filename = FPaths::ProjectLogDir() / FString::Printf(TEXT("vulkan.%s.nv-gpudmp"), *DateTimeString);
 		FArchive* Writer = IFileManager::Get().CreateFileWriter(*Filename);
 		if (Writer)
 		{
@@ -886,7 +888,7 @@ void AftermathGpuCrashDumpCallback(const void* CrashDump, const uint32 CrashDump
 				GFSDK_Aftermath_Result ResultJson = GFSDK_Aftermath_GpuCrashDump_GetJSON(Decoder, (uint32)Json.Num(), Json.GetData());
 				if (ResultJson == GFSDK_Aftermath_Result_Success)
 				{
-					FString Filename = FPaths::ProjectLogDir() / TEXT("vulkan.nv-gpudmp.json");
+					const FString Filename = FPaths::ProjectLogDir() / FString::Printf(TEXT("vulkan.%s.nv-gpudmp.json"), *DateTimeString);
 					FArchive* Writer = IFileManager::Get().CreateFileWriter(*Filename);
 					if (Writer)
 					{
@@ -1110,18 +1112,21 @@ namespace VulkanRHI
 			{
 				GFSDK_Aftermath_CrashDump_Status AftermathStatus{};
 				GFSDK_Aftermath_GetCrashDumpStatus(&AftermathStatus);
-				if (AftermathStatus != GFSDK_Aftermath_CrashDump_Status_Unknown && AftermathStatus != GFSDK_Aftermath_CrashDump_Status_NotStarted)
+				if (AftermathStatus != GFSDK_Aftermath_CrashDump_Status_Unknown)
 				{
 					const float StartTime = FPlatformTime::Seconds();
 					const float EndTime = StartTime + GVulkanNVAfterMathDumpWaitTime;
-					while (AftermathStatus != GFSDK_Aftermath_CrashDump_Status_CollectingDataFailed
-						&& AftermathStatus != GFSDK_Aftermath_CrashDump_Status_Finished
+					while (
+						((AftermathStatus == GFSDK_Aftermath_CrashDump_Status_NotStarted)
+						|| (AftermathStatus != GFSDK_Aftermath_CrashDump_Status_CollectingDataFailed
+						&& AftermathStatus != GFSDK_Aftermath_CrashDump_Status_Finished))
 						&& FPlatformTime::Seconds() < EndTime)
 					{
 						FPlatformProcess::Sleep(0.01f);
 						GFSDK_Aftermath_GetCrashDumpStatus(&AftermathStatus);
 					}
 				}
+				UE_LOG(LogVulkanRHI, Warning, TEXT("Final Aftermath status was %d."), (int32)AftermathStatus);
 			}
 #endif
 		}
