Prevent PSO precaching being used with the vulkan RHI (support is still WIP)

diff --git a/Engine/Source/Runtime/VulkanRHI/Private/VulkanRHI.cpp b/Engine/Source/Runtime/VulkanRHI/Private/VulkanRHI.cpp
--- a/Engine/Source/Runtime/VulkanRHI/Private/VulkanRHI.cpp
+++ b/Engine/Source/Runtime/VulkanRHI/Private/VulkanRHI.cpp
@@ -487,6 +487,14 @@ FVulkanDynamicRHI::FVulkanDynamicRHI()
 	GRHITransitionPrivateData_AlignInBytes = alignof(FVulkanPipelineBarrier);
 	GConfig->GetInt(TEXT("TextureStreaming"), TEXT("PoolSizeVRAMPercentage"), GPoolSizeVRAMPercentage, GEngineIni);
 
+	// VulkanRHI does not yet support PSO precaching, we disable the cvar here.
+	static const auto CVarPSOPrecaching = IConsoleManager::Get().FindConsoleVariable(TEXT("r.PSOPrecaching"));
+	if (CVarPSOPrecaching)
+	{
+		UE_CLOG(CVarPSOPrecaching->GetInt()!=0, LogVulkanRHI, Display, TEXT("r.PSOPrecaching was enabled but is not supported by VulkanRHI. setting r.PSOPrecaching=0"));
+		CVarPSOPrecaching->Set(false);
+	}
+
 	// Copy source requires its own image layout.
 	EnumRemoveFlags(GRHIMergeableAccessMask, ERHIAccess::CopySrc);
 
