VulkanRHI: Allow the use of AMDBufferMarkers and NVCheckpoints simultaneously (fix crash on drivers that support both).

diff --git a/Engine/Source/Runtime/VulkanRHI/Private/Android/VulkanAndroidPlatform.h b/Engine/Source/Runtime/VulkanRHI/Private/Android/VulkanAndroidPlatform.h
--- a/Engine/Source/Runtime/VulkanRHI/Private/Android/VulkanAndroidPlatform.h
+++ b/Engine/Source/Runtime/VulkanRHI/Private/Android/VulkanAndroidPlatform.h
@@ -43,8 +43,7 @@
 
 #define ENUM_VK_ENTRYPOINTS_OPTIONAL_PLATFORM_INSTANCE(EnumMacro) \
 	EnumMacro(PFN_vkGetRefreshCycleDurationGOOGLE, vkGetRefreshCycleDurationGOOGLE) \
-	EnumMacro(PFN_vkGetPastPresentationTimingGOOGLE, vkGetPastPresentationTimingGOOGLE) \
-	EnumMacro(PFN_vkGetPhysicalDeviceFragmentShadingRatesKHR, vkGetPhysicalDeviceFragmentShadingRatesKHR)
+	EnumMacro(PFN_vkGetPastPresentationTimingGOOGLE, vkGetPastPresentationTimingGOOGLE)
 
 // and now, include the GenericPlatform class
 #include "../VulkanGenericPlatform.h"
diff --git a/Engine/Source/Runtime/VulkanRHI/Private/Linux/VulkanLinuxPlatform.cpp b/Engine/Source/Runtime/VulkanRHI/Private/Linux/VulkanLinuxPlatform.cpp
--- a/Engine/Source/Runtime/VulkanRHI/Private/Linux/VulkanLinuxPlatform.cpp
+++ b/Engine/Source/Runtime/VulkanRHI/Private/Linux/VulkanLinuxPlatform.cpp
@@ -267,7 +267,8 @@ void FVulkanLinuxPlatform::WriteCrashMarker(const FOptionalVulkanDeviceExtension
 			VulkanDynamicAPI::vkCmdWriteBufferMarkerAMD(CmdBuffer, VK_PIPELINE_STAGE_BOTTOM_OF_PIPE_BIT, DestBuffer, (1 + LastIndex) * sizeof(uint32), Entries[LastIndex]);
 		}
 	}
-	else if (OptionalExtensions.HasNVDiagnosticCheckpoints)
+
+	if (OptionalExtensions.HasNVDiagnosticCheckpoints)
 	{
 		if (bAdding)
 		{
diff --git a/Engine/Source/Runtime/VulkanRHI/Private/Linux/VulkanLinuxPlatform.h b/Engine/Source/Runtime/VulkanRHI/Private/Linux/VulkanLinuxPlatform.h
--- a/Engine/Source/Runtime/VulkanRHI/Private/Linux/VulkanLinuxPlatform.h
+++ b/Engine/Source/Runtime/VulkanRHI/Private/Linux/VulkanLinuxPlatform.h
@@ -30,13 +30,9 @@
 
 #define ENUM_VK_ENTRYPOINTS_PLATFORM_BASE(EnumMacro)
 
-#define ENUM_VK_ENTRYPOINTS_PLATFORM_INSTANCE(EnumMacro) \
+#define ENUM_VK_ENTRYPOINTS_PLATFORM_INSTANCE(EnumMacro)
 
-#define ENUM_VK_ENTRYPOINTS_OPTIONAL_PLATFORM_INSTANCE(EnumMacro) \
-    EnumMacro(PFN_vkCmdWriteBufferMarkerAMD, vkCmdWriteBufferMarkerAMD) \
-    EnumMacro(PFN_vkCmdSetCheckpointNV, vkCmdSetCheckpointNV) \
-    EnumMacro(PFN_vkGetQueueCheckpointDataNV, vkGetQueueCheckpointDataNV) \
-	EnumMacro(PFN_vkGetPhysicalDeviceFragmentShadingRatesKHR, vkGetPhysicalDeviceFragmentShadingRatesKHR)
+#define ENUM_VK_ENTRYPOINTS_OPTIONAL_PLATFORM_INSTANCE(EnumMacro)
 
 // and now, include the GenericPlatform class
 #include "../VulkanGenericPlatform.h"
diff --git a/Engine/Source/Runtime/VulkanRHI/Private/VulkanDevice.cpp b/Engine/Source/Runtime/VulkanRHI/Private/VulkanDevice.cpp
--- a/Engine/Source/Runtime/VulkanRHI/Private/VulkanDevice.cpp
+++ b/Engine/Source/Runtime/VulkanRHI/Private/VulkanDevice.cpp
@@ -1203,15 +1203,6 @@ void FVulkanDevice::InitGPU()
 			*Entry = 0;
 			VERIFYVULKANRESULT(VulkanRHI::vkBindBufferMemory(Device, CrashMarker.Buffer, CrashMarker.Allocation->GetHandle(), 0));
 		}
-		else if (OptionalDeviceExtensions.HasNVDiagnosticCheckpoints)
-		{
-			CrashMarker.Allocation = DeviceMemoryManager.Alloc(false, GMaxCrashBufferEntries * sizeof(uint32_t), UINT32_MAX, VK_MEMORY_PROPERTY_HOST_VISIBLE_BIT |
-				VK_MEMORY_PROPERTY_HOST_COHERENT_BIT, nullptr, VULKAN_MEMORY_MEDIUM_PRIORITY, false, __FILE__, __LINE__);
-			uint32* Entry = (uint32*)CrashMarker.Allocation->Map(VK_WHOLE_SIZE, 0);
-			check(Entry);
-			// Start with 0 entries
-			*Entry = 0;
-		}
 	}
 #endif
 
@@ -1378,26 +1369,22 @@ void FVulkanDevice::Destroy()
 	PipelineStateCache = nullptr;
 	StagingManager.Deinit();
 
+#if VULKAN_SUPPORTS_GPU_CRASH_DUMPS
 	if (GGPUCrashDebuggingEnabled)
 	{
-#if VULKAN_SUPPORTS_AMD_BUFFER_MARKER
-		if (OptionalDeviceExtensions.HasAMDBufferMarker)
+		if (CrashMarker.Buffer != VK_NULL_HANDLE)
 		{
-			CrashMarker.Allocation->Unmap();
 			VulkanRHI::vkDestroyBuffer(Device, CrashMarker.Buffer, VULKAN_CPU_ALLOCATOR);
 			CrashMarker.Buffer = VK_NULL_HANDLE;
-
-			DeviceMemoryManager.Free(CrashMarker.Allocation);
 		}
-#endif
-#if VULKAN_SUPPORTS_NV_DIAGNOSTIC_CHECKPOINT
-		if (OptionalDeviceExtensions.HasNVDiagnosticCheckpoints)
+
+		if (CrashMarker.Allocation)
 		{
 			CrashMarker.Allocation->Unmap();
 			DeviceMemoryManager.Free(CrashMarker.Allocation);
 		}
-#endif
 	}
+#endif // VULKAN_SUPPORTS_GPU_CRASH_DUMPS
 
 	DeferredDeletionQueue.Clear();
 
diff --git a/Engine/Source/Runtime/VulkanRHI/Private/VulkanLoader.h b/Engine/Source/Runtime/VulkanRHI/Private/VulkanLoader.h
--- a/Engine/Source/Runtime/VulkanRHI/Private/VulkanLoader.h
+++ b/Engine/Source/Runtime/VulkanRHI/Private/VulkanLoader.h
@@ -221,6 +221,10 @@
 	EnumMacro(PFN_vkDeferredOperationJoinKHR, vkDeferredOperationJoinKHR) \
 	EnumMacro(PFN_vkGetDeferredOperationMaxConcurrencyKHR, vkGetDeferredOperationMaxConcurrencyKHR) \
 	EnumMacro(PFN_vkGetDeferredOperationResultKHR, vkGetDeferredOperationResultKHR) \
+    EnumMacro(PFN_vkCmdWriteBufferMarkerAMD, vkCmdWriteBufferMarkerAMD) \
+    EnumMacro(PFN_vkCmdSetCheckpointNV, vkCmdSetCheckpointNV) \
+    EnumMacro(PFN_vkGetQueueCheckpointDataNV, vkGetQueueCheckpointDataNV) \
+	EnumMacro(PFN_vkGetPhysicalDeviceFragmentShadingRatesKHR, vkGetPhysicalDeviceFragmentShadingRatesKHR) \
 	EnumMacro(PFN_vkGetDeviceFaultInfoEXT, vkGetDeviceFaultInfoEXT)
 
 
diff --git a/Engine/Source/Runtime/VulkanRHI/Private/VulkanUtil.cpp b/Engine/Source/Runtime/VulkanRHI/Private/VulkanUtil.cpp
--- a/Engine/Source/Runtime/VulkanRHI/Private/VulkanUtil.cpp
+++ b/Engine/Source/Runtime/VulkanRHI/Private/VulkanUtil.cpp
@@ -764,7 +764,6 @@ void FVulkanGPUProfiler::PopMarkerForCrash(VkCommandBuffer CmdBuffer, VkBuffer D
 
 void FVulkanGPUProfiler::DumpCrashMarkers(void* BufferData)
 {
-#if VULKAN_SUPPORTS_AMD_BUFFER_MARKER
 	if (Device->GetOptionalExtensions().HasAMDBufferMarker)
 	{
 		uint32* Entries = (uint32*)BufferData;
@@ -776,40 +775,35 @@ void FVulkanGPUProfiler::DumpCrashMarkers(void* BufferData)
 			++Entries;
 		}
 	}
-	else
-#endif
+
+	if (Device->GetOptionalExtensions().HasNVDiagnosticCheckpoints)
 	{
-#if VULKAN_SUPPORTS_NV_DIAGNOSTICS
-		if (Device->GetOptionalExtensions().HasNVDiagnosticCheckpoints)
+		struct FCheckpointDataNV : public VkCheckpointDataNV
 		{
-			struct FCheckpointDataNV : public VkCheckpointDataNV
+			FCheckpointDataNV()
 			{
-				FCheckpointDataNV()
-				{
-					ZeroVulkanStruct(*this, VK_STRUCTURE_TYPE_CHECKPOINT_DATA_NV);
-				}
-			};
-			TArray<FCheckpointDataNV> Data;
-			uint32 Num = 0;
-			VkQueue QueueHandle = Device->GetGraphicsQueue()->GetHandle();
-			VulkanDynamicAPI::vkGetQueueCheckpointDataNV(QueueHandle, &Num, nullptr);
-			if (Num > 0)
+				ZeroVulkanStruct(*this, VK_STRUCTURE_TYPE_CHECKPOINT_DATA_NV);
+			}
+		};
+		TArray<FCheckpointDataNV> Data;
+		uint32 Num = 0;
+		VkQueue QueueHandle = Device->GetGraphicsQueue()->GetHandle();
+		VulkanDynamicAPI::vkGetQueueCheckpointDataNV(QueueHandle, &Num, nullptr);
+		if (Num > 0)
+		{
+			Data.AddDefaulted(Num);
+			VulkanDynamicAPI::vkGetQueueCheckpointDataNV(QueueHandle, &Num, &Data[0]);
+			check(Num == Data.Num());
+			for (uint32 Index = 0; Index < Num; ++Index)
 			{
-				Data.AddDefaulted(Num);
-				VulkanDynamicAPI::vkGetQueueCheckpointDataNV(QueueHandle, &Num, &Data[0]);
-				check(Num == Data.Num());
-				for (uint32 Index = 0; Index < Num; ++Index)
-				{
-					check(Data[Index].sType == VK_STRUCTURE_TYPE_CHECKPOINT_DATA_NV);
-					uint32 Value = (uint32)(size_t)Data[Index].pCheckpointMarker;
-					const FString* Frame = CachedStrings.Find(Value);
-					UE_LOG(LogVulkanRHI, Error, TEXT("[VK_NV_device_diagnostic_checkpoints] %i: Stage %s (0x%08x), %s (CRC 0x%x)"), 
-						Index, VK_TYPE_TO_STRING(VkPipelineStageFlagBits, Data[Index].stage), Data[Index].stage, Frame ? *(*Frame) : TEXT("<undefined>"), Value);
-				}
-				GLog->Panic();
+				check(Data[Index].sType == VK_STRUCTURE_TYPE_CHECKPOINT_DATA_NV);
+				uint32 Value = (uint32)(size_t)Data[Index].pCheckpointMarker;
+				const FString* Frame = CachedStrings.Find(Value);
+				UE_LOG(LogVulkanRHI, Error, TEXT("[VK_NV_device_diagnostic_checkpoints] %i: Stage %s (0x%08x), %s (CRC 0x%x)"), 
+					Index, VK_TYPE_TO_STRING(VkPipelineStageFlagBits, Data[Index].stage), Data[Index].stage, Frame ? *(*Frame) : TEXT("<undefined>"), Value);
 			}
+			GLog->Panic();
 		}
-#endif
 	}
 
 	if (!Device->GetOptionalExtensions().HasGPUCrashDumpExtensions())
@@ -827,7 +821,7 @@ void FVulkanGPUProfiler::DumpCrashMarkers(void* BufferData)
 		GLog->Panic();
 	}
 }
-#endif
+#endif // VULKAN_SUPPORTS_GPU_CRASH_DUMPS
 
 #if NV_AFTERMATH
 void AftermathGpuCrashDumpCallback(const void* CrashDump, const uint32 CrashDumpSize, void* UserData)
diff --git a/Engine/Source/Runtime/VulkanRHI/Private/Windows/VulkanWindowsPlatform.cpp b/Engine/Source/Runtime/VulkanRHI/Private/Windows/VulkanWindowsPlatform.cpp
--- a/Engine/Source/Runtime/VulkanRHI/Private/Windows/VulkanWindowsPlatform.cpp
+++ b/Engine/Source/Runtime/VulkanRHI/Private/Windows/VulkanWindowsPlatform.cpp
@@ -236,7 +236,8 @@ void FVulkanWindowsPlatform::WriteCrashMarker(const FOptionalVulkanDeviceExtensi
 			VulkanDynamicAPI::vkCmdWriteBufferMarkerAMD(CmdBuffer, VK_PIPELINE_STAGE_BOTTOM_OF_PIPE_BIT, DestBuffer, (1 + LastIndex) * sizeof(uint32), Entries[LastIndex]);
 		}
 	}
-	else if (OptionalExtensions.HasNVDiagnosticCheckpoints)
+
+	if (OptionalExtensions.HasNVDiagnosticCheckpoints)
 	{
 		if (bAdding)
 		{
diff --git a/Engine/Source/Runtime/VulkanRHI/Private/Windows/VulkanWindowsPlatform.h b/Engine/Source/Runtime/VulkanRHI/Private/Windows/VulkanWindowsPlatform.h
--- a/Engine/Source/Runtime/VulkanRHI/Private/Windows/VulkanWindowsPlatform.h
+++ b/Engine/Source/Runtime/VulkanRHI/Private/Windows/VulkanWindowsPlatform.h
@@ -38,11 +38,7 @@
 #define ENUM_VK_ENTRYPOINTS_PLATFORM_BASE(EnumMacro)
 
 #define ENUM_VK_ENTRYPOINTS_PLATFORM_INSTANCE(EnumMacro)	\
-	EnumMacro(PFN_vkCreateWin32SurfaceKHR, vkCreateWin32SurfaceKHR) \
-	EnumMacro(PFN_vkCmdWriteBufferMarkerAMD, vkCmdWriteBufferMarkerAMD) \
-	EnumMacro(PFN_vkCmdSetCheckpointNV, vkCmdSetCheckpointNV) \
-	EnumMacro(PFN_vkGetQueueCheckpointDataNV, vkGetQueueCheckpointDataNV) \
-	EnumMacro(PFN_vkGetPhysicalDeviceFragmentShadingRatesKHR, vkGetPhysicalDeviceFragmentShadingRatesKHR)
+	EnumMacro(PFN_vkCreateWin32SurfaceKHR, vkCreateWin32SurfaceKHR)
 
 #define ENUM_VK_ENTRYPOINTS_OPTIONAL_PLATFORM_INSTANCE(EnumMacro)
 
