Remove unused VULKAN_SUPPORTS_FRAGMENT_DENSITY_MAP and VULKAN_SUPPORTS_FRAGMENT_DENSITY_MAP2

diff --git a/Engine/Source/Runtime/VulkanRHI/Private/VulkanBarriers.cpp b/Engine/Source/Runtime/VulkanRHI/Private/VulkanBarriers.cpp
--- a/Engine/Source/Runtime/VulkanRHI/Private/VulkanBarriers.cpp
+++ b/Engine/Source/Runtime/VulkanRHI/Private/VulkanBarriers.cpp
@@ -416,14 +416,12 @@ static void GetVkStageAndAccessFlags(ERHIAccess RHIAccess, FRHITransitionInfo::E
 		}
 #endif
 
-#if VULKAN_SUPPORTS_FRAGMENT_DENSITY_MAP
 		if (GRHIVariableRateShadingImageDataType == VRSImage_Fractional)
 		{
 			StageFlags = VK_PIPELINE_STAGE_FRAGMENT_DENSITY_PROCESS_BIT_EXT;
 			AccessFlags = VK_ACCESS_FRAGMENT_DENSITY_MAP_READ_BIT_EXT;
 			Layout = VK_IMAGE_LAYOUT_FRAGMENT_DENSITY_MAP_OPTIMAL_EXT;
 		}
-#endif
 
 		ProcessedRHIFlags |= (uint32)ERHIAccess::ShadingRateSource;
 	}
diff --git a/Engine/Source/Runtime/VulkanRHI/Private/VulkanExtensions.cpp b/Engine/Source/Runtime/VulkanRHI/Private/VulkanExtensions.cpp
--- a/Engine/Source/Runtime/VulkanRHI/Private/VulkanExtensions.cpp
+++ b/Engine/Source/Runtime/VulkanRHI/Private/VulkanExtensions.cpp
@@ -563,7 +563,7 @@ class FVulkanEXTFragmentDensityMapExtension : public FVulkanDeviceExtension
 public:
 
 	FVulkanEXTFragmentDensityMapExtension(FVulkanDevice* InDevice)
-		: FVulkanDeviceExtension(InDevice, VK_EXT_FRAGMENT_DENSITY_MAP_EXTENSION_NAME, VULKAN_SUPPORTS_FRAGMENT_DENSITY_MAP)
+		: FVulkanDeviceExtension(InDevice, VK_EXT_FRAGMENT_DENSITY_MAP_EXTENSION_NAME, VULKAN_EXTENSION_ENABLED)
 	{
 		bEnabledInCode = bEnabledInCode && GRHIVariableRateShadingEnabled;
 	}
@@ -628,7 +628,7 @@ class FVulkanEXTFragmentDensityMap2Extension : public FVulkanDeviceExtension
 public:
 
 	FVulkanEXTFragmentDensityMap2Extension(FVulkanDevice* InDevice)
-		: FVulkanDeviceExtension(InDevice, VK_EXT_FRAGMENT_DENSITY_MAP_2_EXTENSION_NAME, VULKAN_SUPPORTS_FRAGMENT_DENSITY_MAP2)
+		: FVulkanDeviceExtension(InDevice, VK_EXT_FRAGMENT_DENSITY_MAP_2_EXTENSION_NAME, VULKAN_EXTENSION_ENABLED)
 	{
 		bEnabledInCode = bEnabledInCode && GRHIVariableRateShadingEnabled;
 	}
diff --git a/Engine/Source/Runtime/VulkanRHI/Private/VulkanQueue.cpp b/Engine/Source/Runtime/VulkanRHI/Private/VulkanQueue.cpp
--- a/Engine/Source/Runtime/VulkanRHI/Private/VulkanQueue.cpp
+++ b/Engine/Source/Runtime/VulkanRHI/Private/VulkanQueue.cpp
@@ -163,10 +163,10 @@ void FVulkanQueue::FillSupportedStageBits()
 #if VULKAN_SUPPORTS_FRAGMENT_SHADING_RATE
 		SupportedStages |= VK_PIPELINE_STAGE_FRAGMENT_SHADING_RATE_ATTACHMENT_BIT_KHR;
 #endif
-
-#if VULKAN_SUPPORTS_FRAGMENT_DENSITY_MAP
-		SupportedStages |= VK_PIPELINE_STAGE_FRAGMENT_DENSITY_PROCESS_BIT_EXT;
-#endif
+		if (Device->GetOptionalExtensions().HasEXTFragmentDensityMap)
+		{
+			SupportedStages |= VK_PIPELINE_STAGE_FRAGMENT_DENSITY_PROCESS_BIT_EXT;
+		}
 	}
 
 	if (VKHasAnyFlags(QueueProps.queueFlags, VK_QUEUE_COMPUTE_BIT))
diff --git a/Engine/Source/Runtime/VulkanRHI/Private/VulkanRHIPrivate.h b/Engine/Source/Runtime/VulkanRHI/Private/VulkanRHIPrivate.h
--- a/Engine/Source/Runtime/VulkanRHI/Private/VulkanRHIPrivate.h
+++ b/Engine/Source/Runtime/VulkanRHI/Private/VulkanRHIPrivate.h
@@ -908,12 +908,7 @@ inline bool ValidateShadingRateDataType()
 #endif
 
 	case VRSImage_Fractional:
-#if VULKAN_SUPPORTS_FRAGMENT_DENSITY_MAP
 		return true;
-#else
-		checkf(false, TEXT("GRHIVariableRateShadingImageDataType was specified as VRSImage_Fractional, but the VK_EXT_fragment_density_map extension is not supported on this platform."));
-		break;
-#endif
 
 	case VRSImage_NotSupported:
 		checkf(false, TEXT("A texture was marked as a shading rate source but attachment VRS is not supported on this device. Ensure GRHISupportsAttachmentVariableRateShading and GRHIAttachmentVariableRateShadingEnabled are true before specifying a shading rate attachment."));
diff --git a/Engine/Source/Runtime/VulkanRHI/Private/VulkanRenderTarget.cpp b/Engine/Source/Runtime/VulkanRHI/Private/VulkanRenderTarget.cpp
--- a/Engine/Source/Runtime/VulkanRHI/Private/VulkanRenderTarget.cpp
+++ b/Engine/Source/Runtime/VulkanRHI/Private/VulkanRenderTarget.cpp
@@ -765,12 +765,10 @@ VkImageLayout FVulkanRenderTargetLayout::GetVRSImageLayout() const
 			return VK_IMAGE_LAYOUT_FRAGMENT_SHADING_RATE_ATTACHMENT_OPTIMAL_KHR;
 		}
 #endif
-#if VULKAN_SUPPORTS_FRAGMENT_DENSITY_MAP
 		if (GRHIVariableRateShadingImageDataType == VRSImage_Fractional)
 		{
 			return VK_IMAGE_LAYOUT_FRAGMENT_DENSITY_MAP_OPTIMAL_EXT;
 		}
-#endif
 	}
 
 	return VK_IMAGE_LAYOUT_UNDEFINED;
diff --git a/Engine/Source/Runtime/VulkanRHI/Private/VulkanTexture.cpp b/Engine/Source/Runtime/VulkanRHI/Private/VulkanTexture.cpp
--- a/Engine/Source/Runtime/VulkanRHI/Private/VulkanTexture.cpp
+++ b/Engine/Source/Runtime/VulkanRHI/Private/VulkanTexture.cpp
@@ -340,12 +340,10 @@ void FVulkanTexture::GenerateImageCreateInfo(
 		}
 #endif
 
-#if VULKAN_SUPPORTS_FRAGMENT_DENSITY_MAP
 		if (GRHIVariableRateShadingImageDataType == VRSImage_Fractional)
 		{
 			ImageCreateInfo.usage |= VK_IMAGE_USAGE_FRAGMENT_DENSITY_MAP_BIT_EXT;
 		}
-#endif
 	}
 	
 	if (EnumHasAnyFlags(UEFlags, TexCreate_UAV))
@@ -1726,12 +1724,10 @@ FVulkanTexture::FVulkanTexture(FVulkanDevice& InDevice, const FRHITextureCreateD
 				}
 #endif
 
-#if VULKAN_SUPPORTS_FRAGMENT_DENSITY_MAP
 				if (GRHIVariableRateShadingImageDataType == VRSImage_Fractional)
 				{
 					InitialLayout = VK_IMAGE_LAYOUT_FRAGMENT_DENSITY_MAP_OPTIMAL_EXT;
 				}
-#endif
 
 				bOnlyAddToLayoutManager = true;
 				bDoInitialClear = false;
diff --git a/Engine/Source/Runtime/VulkanRHI/Public/VulkanConfiguration.h b/Engine/Source/Runtime/VulkanRHI/Public/VulkanConfiguration.h
--- a/Engine/Source/Runtime/VulkanRHI/Public/VulkanConfiguration.h
+++ b/Engine/Source/Runtime/VulkanRHI/Public/VulkanConfiguration.h
@@ -202,22 +202,6 @@ struct VkAllocationCallbacks;
 #endif
 #endif
 
-#ifndef VULKAN_SUPPORTS_FRAGMENT_DENSITY_MAP
-	#ifdef VK_EXT_fragment_density_map
-		#define VULKAN_SUPPORTS_FRAGMENT_DENSITY_MAP			1
-	#else
-		#define VULKAN_SUPPORTS_FRAGMENT_DENSITY_MAP			0
-	#endif
-#endif
-
-#ifndef VULKAN_SUPPORTS_FRAGMENT_DENSITY_MAP2
-	#ifdef VK_EXT_fragment_density_map2
-		#define VULKAN_SUPPORTS_FRAGMENT_DENSITY_MAP2			1
-	#else
-		#define VULKAN_SUPPORTS_FRAGMENT_DENSITY_MAP2			0
-	#endif
-#endif
-
 #ifndef VULKAN_SUPPORTS_FRAGMENT_SHADING_RATE
 	#ifdef VK_KHR_fragment_shading_rate
 		#define VULKAN_SUPPORTS_FRAGMENT_SHADING_RATE 1
