Removed unused VULKAN_SUPPORTS_FRAGMENT_SHADING_RATE

diff --git a/Engine/Source/Runtime/VulkanRHI/Private/Android/VulkanAndroidPlatform.cpp b/Engine/Source/Runtime/VulkanRHI/Private/Android/VulkanAndroidPlatform.cpp
--- a/Engine/Source/Runtime/VulkanRHI/Private/Android/VulkanAndroidPlatform.cpp
+++ b/Engine/Source/Runtime/VulkanRHI/Private/Android/VulkanAndroidPlatform.cpp
@@ -812,9 +812,6 @@ void HandleGraphicsPipelineCreatePNext(const VkGraphicsPipelineCreateInfo* Pipel
 	// FSR Create Info
 	bool bHasFSRCreateInfo = false;
 
-#if !VULKAN_SUPPORTS_FRAGMENT_SHADING_RATE
-	COPY_TO_BUFFER(MemoryStream, &bHasFSRCreateInfo, sizeof(bool));
-#else
 	const void** Struct = VkStructs.Find(VK_STRUCTURE_TYPE_PIPELINE_FRAGMENT_SHADING_RATE_STATE_CREATE_INFO_KHR);
 	bHasFSRCreateInfo = Struct != nullptr;
 	COPY_TO_BUFFER(MemoryStream, &bHasFSRCreateInfo, sizeof(bool));
@@ -828,7 +825,6 @@ void HandleGraphicsPipelineCreatePNext(const VkGraphicsPipelineCreateInfo* Pipel
 		COPY_TO_BUFFER(MemoryStream, FSRCreateInfo, sizeof(VkPipelineFragmentShadingRateStateCreateInfoKHR));
 		HandledCount++;
 	} 
-#endif
 
 	check(HandledCount == VkStructs.Num());
 }
@@ -869,9 +865,6 @@ void HandleSubpassDescriptionPNext(VkSubpassDescription2* SubpassDescription, TA
 	// FSR Create Info 
 	bool bHasFSRAttachmentCreateInfo = false;
 
-#if !VULKAN_SUPPORTS_FRAGMENT_SHADING_RATE
-	COPY_TO_BUFFER(MemoryStream, &bHasFSRAttachmentCreateInfo, sizeof(bool));
-#else
 	const void** Struct = VkStructs.Find(VK_STRUCTURE_TYPE_FRAGMENT_SHADING_RATE_ATTACHMENT_INFO_KHR);
 	bHasFSRAttachmentCreateInfo = Struct != nullptr;
 	COPY_TO_BUFFER(MemoryStream, &bHasFSRAttachmentCreateInfo, sizeof(bool));
@@ -885,7 +878,6 @@ void HandleSubpassDescriptionPNext(VkSubpassDescription2* SubpassDescription, TA
 		COPY_TO_BUFFER(MemoryStream, &FragmentShadingRateCreateInfo->shadingRateAttachmentTexelSize, sizeof(VkExtent2D));
 		HandledCount++;
 	}
-#endif
 
 	check(HandledCount == VkStructs.Num());
 }
diff --git a/Engine/Source/Runtime/VulkanRHI/Private/VulkanBarriers.cpp b/Engine/Source/Runtime/VulkanRHI/Private/VulkanBarriers.cpp
--- a/Engine/Source/Runtime/VulkanRHI/Private/VulkanBarriers.cpp
+++ b/Engine/Source/Runtime/VulkanRHI/Private/VulkanBarriers.cpp
@@ -407,14 +407,12 @@ static void GetVkStageAndAccessFlags(ERHIAccess RHIAccess, FRHITransitionInfo::E
 	{
 		checkf(ResourceType == FRHITransitionInfo::EType::Texture, TEXT("A non-texture resource was tagged as a shading rate source; only textures (Texture2D and Texture2DArray) can be used for this purpose."));
 		
-#if VULKAN_SUPPORTS_FRAGMENT_SHADING_RATE
 		if (GRHIVariableRateShadingImageDataType == VRSImage_Palette)
 		{
 			StageFlags = VK_PIPELINE_STAGE_FRAGMENT_SHADING_RATE_ATTACHMENT_BIT_KHR;
 			AccessFlags = VK_ACCESS_FRAGMENT_SHADING_RATE_ATTACHMENT_READ_BIT_KHR;
 			Layout = VK_IMAGE_LAYOUT_FRAGMENT_SHADING_RATE_ATTACHMENT_OPTIMAL_KHR;
 		}
-#endif
 
 		if (GRHIVariableRateShadingImageDataType == VRSImage_Fractional)
 		{
diff --git a/Engine/Source/Runtime/VulkanRHI/Private/VulkanCommandWrappers.h b/Engine/Source/Runtime/VulkanRHI/Private/VulkanCommandWrappers.h
--- a/Engine/Source/Runtime/VulkanRHI/Private/VulkanCommandWrappers.h
+++ b/Engine/Source/Runtime/VulkanRHI/Private/VulkanCommandWrappers.h
@@ -171,9 +171,7 @@ struct FWrapLayer
 	VULKAN_EXTERN_EXPORT static void GetPhysicalDeviceMemoryProperties2(VkResult Result, VkPhysicalDevice PhysicalDevice, VkPhysicalDeviceMemoryProperties2* MemoryProperties) VULKAN_LAYER_BODY 
 	VULKAN_EXTERN_EXPORT static void GetPhysicalDeviceProperties2(VkResult Result, VkPhysicalDevice PhysicalDevice, VkPhysicalDeviceProperties2* Properties) VULKAN_LAYER_BODY
 	VULKAN_EXTERN_EXPORT static void GetPhysicalDeviceFeatures2(VkResult Result, VkPhysicalDevice PhysicalDevice, VkPhysicalDeviceFeatures2* Features) VULKAN_LAYER_BODY
-#if VULKAN_SUPPORTS_FRAGMENT_SHADING_RATE
 	VULKAN_EXTERN_EXPORT static void GetPhysicalDeviceFragmentShadingRatesKHR(VkResult, VkPhysicalDevice PhysicalDevice, uint32* FragmentShadingRateCount, VkPhysicalDeviceFragmentShadingRateKHR* FragmentShadingRates) VULKAN_LAYER_BODY
-#endif
 	static void GetPhysicalDeviceSurfaceCapabilitiesKHR(VkResult Result, VkPhysicalDevice PhysicalDevice, VkSurfaceKHR Surface, VkSurfaceCapabilitiesKHR* SurfaceCapabilities) VULKAN_LAYER_BODY
 	static void GetPhysicalDeviceSurfaceFormatsKHR(VkResult Result, VkPhysicalDevice PhysicalDevice, VkSurfaceKHR Surface, uint32_t* SurfaceFormatCountPtr, VkSurfaceFormatKHR* SurfaceFormats) VULKAN_LAYER_BODY
 	static void GetPhysicalDeviceSurfaceSupportKHR(VkResult Result, VkPhysicalDevice PhysicalDevice, uint32_t QueueFamilyIndex, VkSurfaceKHR Surface, VkBool32* SupportedPtr) VULKAN_LAYER_BODY
@@ -303,14 +301,12 @@ namespace VulkanRHI
 		FWrapLayer::GetPhysicalDeviceFeatures2(VK_SUCCESS, PhysicalDevice, Features);
 	}
 
-#if VULKAN_SUPPORTS_FRAGMENT_SHADING_RATE
 	static FORCEINLINE_DEBUGGABLE void vkGetPhysicalDeviceFragmentShadingRatesKHR(VkPhysicalDevice PhysicalDevice, uint32* FragmentShadingRateCount, VkPhysicalDeviceFragmentShadingRateKHR* FragmentShadingRates)
 	{
 		FWrapLayer::GetPhysicalDeviceFragmentShadingRatesKHR(VK_RESULT_MAX_ENUM, PhysicalDevice, FragmentShadingRateCount, FragmentShadingRates);
 		VULKANAPINAMESPACE::vkGetPhysicalDeviceFragmentShadingRatesKHR(PhysicalDevice, FragmentShadingRateCount, FragmentShadingRates);
 		FWrapLayer::GetPhysicalDeviceFragmentShadingRatesKHR(VK_SUCCESS, PhysicalDevice, FragmentShadingRateCount, FragmentShadingRates);
 	}
-#endif
 
 	static FORCEINLINE_DEBUGGABLE void  vkGetPhysicalDeviceQueueFamilyProperties(VkPhysicalDevice PhysicalDevice, uint32* QueueFamilyPropertyCount, VkQueueFamilyProperties* QueueFamilyProperties)
 	{
diff --git a/Engine/Source/Runtime/VulkanRHI/Private/VulkanDebug.cpp b/Engine/Source/Runtime/VulkanRHI/Private/VulkanDebug.cpp
--- a/Engine/Source/Runtime/VulkanRHI/Private/VulkanDebug.cpp
+++ b/Engine/Source/Runtime/VulkanRHI/Private/VulkanDebug.cpp
@@ -3870,7 +3870,6 @@ void FWrapLayer::GetPhysicalDeviceFeatures2(VkResult Result, VkPhysicalDevice Ph
 	}
 }
 
-#if VULKAN_SUPPORTS_FRAGMENT_SHADING_RATE
 void FWrapLayer::GetPhysicalDeviceFragmentShadingRatesKHR(VkResult Result, VkPhysicalDevice PhysicalDevice, uint32* FragmentShadingRateCount, VkPhysicalDeviceFragmentShadingRateKHR* FragmentShadingRates)
 {
 	if (Result == VK_RESULT_MAX_ENUM)
@@ -3880,7 +3879,6 @@ void FWrapLayer::GetPhysicalDeviceFragmentShadingRatesKHR(VkResult Result, VkPhy
 #endif
 	}
 }
-#endif
 
 void FWrapLayer::CmdSetDepthBias(VkResult Result, VkCommandBuffer CommandBuffer, float DepthBiasConstantFactor, float DepthBiasClamp, float DepthBiasSlopeFactor)
 {
diff --git a/Engine/Source/Runtime/VulkanRHI/Private/VulkanDevice.cpp b/Engine/Source/Runtime/VulkanRHI/Private/VulkanDevice.cpp
--- a/Engine/Source/Runtime/VulkanRHI/Private/VulkanDevice.cpp
+++ b/Engine/Source/Runtime/VulkanRHI/Private/VulkanDevice.cpp
@@ -167,7 +167,6 @@ static void LoadValidationCache(VkDevice Device, VkValidationCacheEXT& OutValida
 }
 #endif
 
-#if VULKAN_SUPPORTS_FRAGMENT_SHADING_RATE
 static VkExtent2D GetBestMatchedShadingRateExtents(uint32 ShadingRate, const TArray<VkPhysicalDeviceFragmentShadingRateKHR>& FragmentShadingRates)
 {
 	// Given that for Vulkan we need to query available device shading rates, we're not guaranteed to have everything that's in our enum;
@@ -200,7 +199,6 @@ static VkExtent2D GetBestMatchedShadingRateExtents(uint32 ShadingRate, const TAr
 
 	return BestMatchedExtent;
 }
-#endif
 
 
 void FVulkanPhysicalDeviceFeatures::Query(VkPhysicalDevice PhysicalDevice, uint32 APIVersion)
@@ -484,8 +482,7 @@ void FVulkanDevice::CreateDevice(TArray<const ANSICHAR*>& DeviceLayers, FVulkanD
 		}
 	}
 
-#if VULKAN_SUPPORTS_FRAGMENT_SHADING_RATE
-	// Enumerate the available shading rates if it's supported
+	// Enumerate the available shading rates
 	if (OptionalDeviceExtensions.HasKHRFragmentShadingRate)
 	{
 		uint32 FragmentShadingRateCount = 0;
@@ -506,7 +503,6 @@ void FVulkanDevice::CreateDevice(TArray<const ANSICHAR*>& DeviceLayers, FVulkanD
 			}
 		}
 	}
-#endif // VULKAN_SUPPORTS_FRAGMENT_SHADING_RATE
 
 	UE_LOG(LogVulkanRHI, Display, TEXT("Using %d device layers%s"), DeviceLayers.Num(), DeviceLayers.Num() ? TEXT(":") : TEXT("."));
 	for (const ANSICHAR* Layer : DeviceLayers)
diff --git a/Engine/Source/Runtime/VulkanRHI/Private/VulkanDevice.h b/Engine/Source/Runtime/VulkanRHI/Private/VulkanDevice.h
--- a/Engine/Source/Runtime/VulkanRHI/Private/VulkanDevice.h
+++ b/Engine/Source/Runtime/VulkanRHI/Private/VulkanDevice.h
@@ -303,12 +303,10 @@ public:
 		return GpuProps;
 	}
 
-#if VULKAN_SUPPORTS_FRAGMENT_SHADING_RATE
 	inline VkExtent2D GetBestMatchedFragmentSize(EVRSShadingRate Rate) const
 	{
 		return FragmentSizeMap[Rate];
 	}
-#endif
 
 	inline const VkPhysicalDeviceLimits& GetLimits() const
 	{
@@ -586,10 +584,8 @@ private:
 	VkPhysicalDevice Gpu;
 	VkPhysicalDeviceProperties GpuProps;
 
-#if VULKAN_SUPPORTS_FRAGMENT_SHADING_RATE
 	TArray<VkPhysicalDeviceFragmentShadingRateKHR> FragmentShadingRates;
 	TStaticArray<VkExtent2D, (EVRSShadingRate::VRSSR_Last+1)> FragmentSizeMap;
-#endif
 
 	// Extension specific properties
 	VkPhysicalDeviceIDPropertiesKHR GpuIdProps;
diff --git a/Engine/Source/Runtime/VulkanRHI/Private/VulkanExtensions.cpp b/Engine/Source/Runtime/VulkanRHI/Private/VulkanExtensions.cpp
--- a/Engine/Source/Runtime/VulkanRHI/Private/VulkanExtensions.cpp
+++ b/Engine/Source/Runtime/VulkanRHI/Private/VulkanExtensions.cpp
@@ -485,7 +485,7 @@ class FVulkanKHRFragmentShadingRateExtension : public FVulkanDeviceExtension
 public:
 
 	FVulkanKHRFragmentShadingRateExtension(FVulkanDevice* InDevice)
-		: FVulkanDeviceExtension(InDevice, VK_KHR_FRAGMENT_SHADING_RATE_EXTENSION_NAME, VULKAN_SUPPORTS_FRAGMENT_SHADING_RATE)
+		: FVulkanDeviceExtension(InDevice, VK_KHR_FRAGMENT_SHADING_RATE_EXTENSION_NAME, VULKAN_EXTENSION_ENABLED)
 	{
 		bEnabledInCode = bEnabledInCode && GRHIVariableRateShadingEnabled;
 	}
diff --git a/Engine/Source/Runtime/VulkanRHI/Private/VulkanPipeline.cpp b/Engine/Source/Runtime/VulkanRHI/Private/VulkanPipeline.cpp
--- a/Engine/Source/Runtime/VulkanRHI/Private/VulkanPipeline.cpp
+++ b/Engine/Source/Runtime/VulkanRHI/Private/VulkanPipeline.cpp
@@ -1128,13 +1128,11 @@ FArchive& operator << (FArchive& Ar, FGfxPipelineDesc& Entry)
 #endif
 	Ar << Entry.RenderTargets;
 
-#if VULKAN_SUPPORTS_FRAGMENT_SHADING_RATE
 	uint8 ShadingRate = static_cast<uint8>(Entry.ShadingRate);
 	uint8 Combiner = static_cast<uint8>(Entry.Combiner);
 	
 	Ar << ShadingRate;
 	Ar << Combiner;
-#endif
 
 	Ar << Entry.UseAlphaToCoverage;
 
@@ -1165,7 +1163,6 @@ FVulkanPSOKey FGfxPipelineDesc::CreateKey2() const
 	return Result;
 }
 
-#if VULKAN_SUPPORTS_FRAGMENT_SHADING_RATE
 // Map Unreal VRS combiner operation enums to Vulkan enums.
 static const TMap<uint8, VkFragmentShadingRateCombinerOpKHR> FragmentCombinerOpMap
 {
@@ -1176,7 +1173,6 @@ static const TMap<uint8, VkFragmentShadingRateCombinerOpKHR> FragmentCombinerOpM
 	{ VRSRB_Sum,			VK_FRAGMENT_SHADING_RATE_COMBINER_OP_MAX_KHR },		// No concept of Sum in Vulkan - fall back to max.
 	// @todo: Add "VK_FRAGMENT_SHADING_RATE_COMBINER_OP_MUL_KHR"?
 };
-#endif
 
 FString FVulkanPipelineStateCacheManager::ShaderHashesToString(FVulkanShader* Shaders[ShaderStage::NumStages])
 {
@@ -1357,9 +1353,8 @@ bool FVulkanPipelineStateCacheManager::CreateGfxPipelineFromEntry(FVulkanRHIGrap
 
 	PipelineInfo.pDynamicState = &DynamicState;
 
-#if VULKAN_SUPPORTS_FRAGMENT_SHADING_RATE
 	VkPipelineFragmentShadingRateStateCreateInfoKHR PipelineFragmentShadingRate;
-	if (GRHISupportsPipelineVariableRateShading && GRHIVariableRateShadingEnabled)
+	if (GRHISupportsPipelineVariableRateShading && GRHIVariableRateShadingEnabled && GRHIVariableRateShadingImageDataType == VRSImage_Palette)
 	{
 		const VkExtent2D FragmentSize = Device->GetBestMatchedFragmentSize(PSO->Desc.ShadingRate);
 		VkFragmentShadingRateCombinerOpKHR PipelineToPrimitiveCombinerOperation = FragmentCombinerOpMap[(uint8)PSO->Desc.Combiner];
@@ -1371,7 +1366,6 @@ bool FVulkanPipelineStateCacheManager::CreateGfxPipelineFromEntry(FVulkanRHIGrap
 
 		PipelineInfo.pNext = (void*)&PipelineFragmentShadingRate;
 	}
-#endif
 
 	if (Device->SupportsBindless())
 	{
diff --git a/Engine/Source/Runtime/VulkanRHI/Private/VulkanPipeline.h b/Engine/Source/Runtime/VulkanRHI/Private/VulkanPipeline.h
--- a/Engine/Source/Runtime/VulkanRHI/Private/VulkanPipeline.h
+++ b/Engine/Source/Runtime/VulkanRHI/Private/VulkanPipeline.h
@@ -473,7 +473,6 @@ struct FGfxPipelineDesc
 			return false;
 		}
 
-#if VULKAN_SUPPORTS_FRAGMENT_SHADING_RATE
 		if (ShadingRate != In.ShadingRate)
 		{
 			return false;
@@ -483,7 +482,6 @@ struct FGfxPipelineDesc
 		{
 			return false;
 		}
-#endif
 
 		return true;
 	}
diff --git a/Engine/Source/Runtime/VulkanRHI/Private/VulkanQueue.cpp b/Engine/Source/Runtime/VulkanRHI/Private/VulkanQueue.cpp
--- a/Engine/Source/Runtime/VulkanRHI/Private/VulkanQueue.cpp
+++ b/Engine/Source/Runtime/VulkanRHI/Private/VulkanQueue.cpp
@@ -159,10 +159,10 @@ void FVulkanQueue::FillSupportedStageBits()
 		{
 			SupportedStages |= VK_PIPELINE_STAGE_GEOMETRY_SHADER_BIT;
 		}
-
-#if VULKAN_SUPPORTS_FRAGMENT_SHADING_RATE
-		SupportedStages |= VK_PIPELINE_STAGE_FRAGMENT_SHADING_RATE_ATTACHMENT_BIT_KHR;
-#endif
+		if (Device->GetOptionalExtensions().HasKHRFragmentShadingRate)
+		{
+			SupportedStages |= VK_PIPELINE_STAGE_FRAGMENT_SHADING_RATE_ATTACHMENT_BIT_KHR;
+		}
 		if (Device->GetOptionalExtensions().HasEXTFragmentDensityMap)
 		{
 			SupportedStages |= VK_PIPELINE_STAGE_FRAGMENT_DENSITY_PROCESS_BIT_EXT;
diff --git a/Engine/Source/Runtime/VulkanRHI/Private/VulkanRHIPrivate.h b/Engine/Source/Runtime/VulkanRHI/Private/VulkanRHIPrivate.h
--- a/Engine/Source/Runtime/VulkanRHI/Private/VulkanRHIPrivate.h
+++ b/Engine/Source/Runtime/VulkanRHI/Private/VulkanRHIPrivate.h
@@ -900,12 +900,7 @@ inline bool ValidateShadingRateDataType()
 	switch (GRHIVariableRateShadingImageDataType)
 	{
 	case VRSImage_Palette:
-#if VULKAN_SUPPORTS_FRAGMENT_SHADING_RATE
 		return true;
-#else
-		checkf(false, TEXT("GRHIVariableRateShadingImageDataType was specified as VRSImage_Palette, but the VK_KHR_fragment_shading_rate extension is not supported on this platform."));
-		break;
-#endif
 
 	case VRSImage_Fractional:
 		return true;
diff --git a/Engine/Source/Runtime/VulkanRHI/Private/VulkanRenderTarget.cpp b/Engine/Source/Runtime/VulkanRHI/Private/VulkanRenderTarget.cpp
--- a/Engine/Source/Runtime/VulkanRHI/Private/VulkanRenderTarget.cpp
+++ b/Engine/Source/Runtime/VulkanRHI/Private/VulkanRenderTarget.cpp
@@ -759,12 +759,10 @@ VkImageLayout FVulkanRenderTargetLayout::GetVRSImageLayout() const
 {
 	if (ValidateShadingRateDataType())
 	{
-#if VULKAN_SUPPORTS_FRAGMENT_SHADING_RATE
 		if (GRHIVariableRateShadingImageDataType == VRSImage_Palette)
 		{
 			return VK_IMAGE_LAYOUT_FRAGMENT_SHADING_RATE_ATTACHMENT_OPTIMAL_KHR;
 		}
-#endif
 		if (GRHIVariableRateShadingImageDataType == VRSImage_Fractional)
 		{
 			return VK_IMAGE_LAYOUT_FRAGMENT_DENSITY_MAP_OPTIMAL_EXT;
diff --git a/Engine/Source/Runtime/VulkanRHI/Private/VulkanTexture.cpp b/Engine/Source/Runtime/VulkanRHI/Private/VulkanTexture.cpp
--- a/Engine/Source/Runtime/VulkanRHI/Private/VulkanTexture.cpp
+++ b/Engine/Source/Runtime/VulkanRHI/Private/VulkanTexture.cpp
@@ -333,12 +333,10 @@ void FVulkanTexture::GenerateImageCreateInfo(
 
 	if (EnumHasAnyFlags(UEFlags, TexCreate_Foveation) && ValidateShadingRateDataType())
 	{
-#if VULKAN_SUPPORTS_FRAGMENT_SHADING_RATE
 		if (GRHIVariableRateShadingImageDataType == VRSImage_Palette)
 		{
 			ImageCreateInfo.usage |= VK_IMAGE_USAGE_FRAGMENT_SHADING_RATE_ATTACHMENT_BIT_KHR;
 		}
-#endif
 
 		if (GRHIVariableRateShadingImageDataType == VRSImage_Fractional)
 		{
@@ -1717,12 +1715,10 @@ FVulkanTexture::FVulkanTexture(FVulkanDevice& InDevice, const FRHITextureCreateD
 			else if (EnumHasAnyFlags(InCreateDesc.Flags, TexCreate_Foveation) && ValidateShadingRateDataType())
 			{
 				// If it's a shading rate attachment texture, do not clear but add to layoutmgr, and set correct shading rate layout.
-#if VULKAN_SUPPORTS_FRAGMENT_SHADING_RATE
 				if (GRHIVariableRateShadingImageDataType == VRSImage_Palette)
 				{
 					InitialLayout = VK_IMAGE_LAYOUT_FRAGMENT_SHADING_RATE_ATTACHMENT_OPTIMAL_KHR;
 				}
-#endif
 
 				if (GRHIVariableRateShadingImageDataType == VRSImage_Fractional)
 				{
diff --git a/Engine/Source/Runtime/VulkanRHI/Public/VulkanConfiguration.h b/Engine/Source/Runtime/VulkanRHI/Public/VulkanConfiguration.h
--- a/Engine/Source/Runtime/VulkanRHI/Public/VulkanConfiguration.h
+++ b/Engine/Source/Runtime/VulkanRHI/Public/VulkanConfiguration.h
@@ -202,14 +202,6 @@ struct VkAllocationCallbacks;
 #endif
 #endif
 
-#ifndef VULKAN_SUPPORTS_FRAGMENT_SHADING_RATE
-	#ifdef VK_KHR_fragment_shading_rate
-		#define VULKAN_SUPPORTS_FRAGMENT_SHADING_RATE 1
-	#else
-		#define VULKAN_SUPPORTS_FRAGMENT_SHADING_RATE 0
-	#endif
-#endif
-
 #ifndef VULKAN_SUPPORTS_RENDERPASS2
 	#ifdef VK_KHR_create_renderpass2
 		#define VULKAN_SUPPORTS_RENDERPASS2 1
