Fixed: Blocky rendering artifacts on Xclipse 920 based devices under Vulkan
Also fixed vulkan crash with .Vulkan.DelayAcquireBackBuffer=0

diff --git a/Engine/Source/Runtime/VulkanRHI/Private/VulkanViewport.cpp b/Engine/Source/Runtime/VulkanRHI/Private/VulkanViewport.cpp
--- a/Engine/Source/Runtime/VulkanRHI/Private/VulkanViewport.cpp
+++ b/Engine/Source/Runtime/VulkanRHI/Private/VulkanViewport.cpp
@@ -24,7 +24,7 @@ void FVulkanBackBuffer::ReleaseAcquiredImage()
 {
 	if (DefaultView)
 	{
-		DefaultView->Invalidate();
+		// Do not invalidate view here, just remove a reference to it
 		DefaultView = nullptr;
 	}
 
@@ -717,7 +717,7 @@ void FVulkanViewport::CreateSwapchain(FVulkanSwapChainRecreateInfo* RecreateInfo
 		FMemory::Memzero(ClearColor);
 
 		const VkImageSubresourceRange Range = FVulkanPipelineBarrier::MakeSubresourceRange(VK_IMAGE_ASPECT_COLOR_BIT);
-		const FVulkanImageLayout InitialLayout(VK_IMAGE_LAYOUT_COLOR_ATTACHMENT_OPTIMAL, 1, 1, VK_IMAGE_ASPECT_COLOR_BIT);
+		const FVulkanImageLayout InitialLayout(VK_IMAGE_LAYOUT_PRESENT_SRC_KHR, 1, 1, VK_IMAGE_ASPECT_COLOR_BIT);
 		for (int32 Index = 0; Index < Images.Num(); ++Index)
 		{
 			BackBufferImages[Index] = Images[Index];
@@ -725,11 +725,11 @@ void FVulkanViewport::CreateSwapchain(FVulkanSwapChainRecreateInfo* RecreateInfo
 			TextureViews.Add((new FVulkanView(*Device, DescriptorType))->InitAsTextureView(
 				Images[Index], VK_IMAGE_VIEW_TYPE_2D, VK_IMAGE_ASPECT_COLOR_BIT, PixelFormat, UEToVkTextureFormat(PixelFormat, false), 0, 1, 0, 1, false));
 
-			// Clear the swapchain to avoid a validation warning, and transition to ColorAttachment
+			// Clear the swapchain to avoid a validation warning, and transition to PresentSrc
 			{
 				VulkanSetImageLayout(CmdBuffer, Images[Index], VK_IMAGE_LAYOUT_UNDEFINED, VK_IMAGE_LAYOUT_TRANSFER_DST_OPTIMAL, Range);
 				VulkanRHI::vkCmdClearColorImage(CmdBuffer->GetHandle(), Images[Index], VK_IMAGE_LAYOUT_TRANSFER_DST_OPTIMAL, &ClearColor, 1, &Range);
-				VulkanSetImageLayout(CmdBuffer, Images[Index], VK_IMAGE_LAYOUT_TRANSFER_DST_OPTIMAL, VK_IMAGE_LAYOUT_COLOR_ATTACHMENT_OPTIMAL, Range);
+				VulkanSetImageLayout(CmdBuffer, Images[Index], VK_IMAGE_LAYOUT_TRANSFER_DST_OPTIMAL, InitialLayout.MainLayout, Range);
 				CmdBuffer->GetLayoutManager().SetFullLayout(Images[Index], InitialLayout);
 			}
 
@@ -833,7 +833,7 @@ inline static void CopyImageToBackBuffer(FVulkanCommandListContext* Context, FVu
 
 	{
 		FVulkanPipelineBarrier Barrier;
-		Barrier.AddImageLayoutTransition(DstSurface, VK_IMAGE_LAYOUT_UNDEFINED, VK_IMAGE_LAYOUT_TRANSFER_DST_OPTIMAL, FVulkanPipelineBarrier::MakeSubresourceRange(VK_IMAGE_ASPECT_COLOR_BIT, 0, 1));
+		Barrier.AddImageLayoutTransition(DstSurface, VK_IMAGE_LAYOUT_PRESENT_SRC_KHR, VK_IMAGE_LAYOUT_TRANSFER_DST_OPTIMAL, FVulkanPipelineBarrier::MakeSubresourceRange(VK_IMAGE_ASPECT_COLOR_BIT, 0, 1));
 		Barrier.Execute(CmdBuffer);
 	}
 
