Refactor VulkanTexture constructor.

Constructor was not up to date with how the initial layout was being set.

diff --git a/Engine/Source/Runtime/VulkanRHI/Private/VulkanTexture.cpp b/Engine/Source/Runtime/VulkanRHI/Private/VulkanTexture.cpp
--- a/Engine/Source/Runtime/VulkanRHI/Private/VulkanTexture.cpp
+++ b/Engine/Source/Runtime/VulkanRHI/Private/VulkanTexture.cpp
@@ -1705,38 +1705,9 @@ FVulkanTexture::FVulkanTexture(FVulkanDevice& InDevice, const FRHITextureCreateD
 #endif
 			VULKAN_SET_DEBUG_NAME(InDevice, VK_OBJECT_TYPE_IMAGE, Image, TEXT("%s:(FVulkanTexture*)0x%p"), InCreateDesc.DebugName ? InCreateDesc.DebugName : TEXT("?"), this);
 
-			VkImageLayout InitialLayout = VK_IMAGE_LAYOUT_UNDEFINED;
-			bool bOnlyAddToLayoutManager, bDoInitialClear;
-			if (EnumHasAnyFlags(InCreateDesc.Flags, TexCreate_RenderTargetable | TexCreate_DepthStencilTargetable))
-			{
-				InitialLayout = EnumHasAnyFlags(InCreateDesc.Flags, TexCreate_DepthStencilTargetable) ? VK_IMAGE_LAYOUT_ATTACHMENT_OPTIMAL : VK_IMAGE_LAYOUT_COLOR_ATTACHMENT_OPTIMAL;
-				bOnlyAddToLayoutManager = false;
-				bDoInitialClear = true;
-			}
-			else if (EnumHasAnyFlags(InCreateDesc.Flags, TexCreate_Foveation) && ValidateShadingRateDataType())
-			{
-				// If it's a shading rate attachment texture, do not clear but add to layoutmgr, and set correct shading rate layout.
-				if (GRHIVariableRateShadingImageDataType == VRSImage_Palette)
-				{
-					InitialLayout = VK_IMAGE_LAYOUT_FRAGMENT_SHADING_RATE_ATTACHMENT_OPTIMAL_KHR;
-				}
-
-				if (GRHIVariableRateShadingImageDataType == VRSImage_Fractional)
-				{
-					InitialLayout = VK_IMAGE_LAYOUT_FRAGMENT_DENSITY_MAP_OPTIMAL_EXT;
-				}
-
-				bOnlyAddToLayoutManager = true;
-				bDoInitialClear = false;
-			}
-			else
-			{
-				// If we haven't seen this image before, we assume it's an SRV (VK_IMAGE_LAYOUT_SHADER_READ_ONLY_OPTIMAL) and the call below
-				// tells the layout manager about it. If we've seen it before, the call won't do anything, since the manager already knows the layout.
-				InitialLayout = VK_IMAGE_LAYOUT_SHADER_READ_ONLY_OPTIMAL;
-				bOnlyAddToLayoutManager = true;
-				bDoInitialClear = false;
-			}
+			const VkImageLayout InitialLayout = GetInitialLayoutFromRHIAccess(InCreateDesc.InitialState, IsDepthOrStencilAspect(), SupportsSampling());
+			const bool bDoInitialClear = EnumHasAnyFlags(InCreateDesc.Flags, TexCreate_RenderTargetable | TexCreate_DepthStencilTargetable);
+			const bool bOnlyAddToLayoutManager = !EnumHasAnyFlags(InCreateDesc.Flags, TexCreate_RenderTargetable | TexCreate_DepthStencilTargetable);
 
 			DefaultLayout = InitialLayout;
 
@@ -1749,9 +1720,9 @@ FVulkanTexture::FVulkanTexture(FVulkanDevice& InDevice, const FRHITextureCreateD
 					FVulkanCmdBuffer* CmdBuffer = Context.GetCommandBufferManager()->GetActiveCmdBuffer();
 					CmdBuffer->GetLayoutManager().SetFullLayout(*this, InitialLayout, true);
 				}
-				else
+				else if (InitialLayout != VK_IMAGE_LAYOUT_UNDEFINED || bDoInitialClear)
 				{
-					SetInitialImageState(Context, InitialLayout, true, InCreateDesc.ClearValue, false);
+					SetInitialImageState(Context, InitialLayout, bDoInitialClear, InCreateDesc.ClearValue, false);
 				}
 			}
 			else
