Allow developers to choose between FSR and FDM if both are available.

diff --git a/Engine/Source/Runtime/VulkanRHI/Private/VulkanDevice.cpp b/Engine/Source/Runtime/VulkanRHI/Private/VulkanDevice.cpp
--- a/Engine/Source/Runtime/VulkanRHI/Private/VulkanDevice.cpp
+++ b/Engine/Source/Runtime/VulkanRHI/Private/VulkanDevice.cpp
@@ -1158,6 +1158,8 @@ void FVulkanDevice::InitGPU()
 		}
 	}
 
+	ChooseVariableRateShadingMethod(OptionalDeviceExtensions, GetOptionalExtensionProperties().FragmentShadingRateFeatures);
+
 	UE_LOG(LogVulkanRHI, Display, TEXT("Device properties: Geometry %d BufferAtomic64 %d ImageAtomic64 %d"), 
 		PhysicalDeviceFeatures.Core_1_0.geometryShader, OptionalDeviceExtensions.HasKHRShaderAtomicInt64, OptionalDeviceExtensions.HasImageAtomicInt64);
 
diff --git a/Engine/Source/Runtime/VulkanRHI/Private/VulkanDevice.h b/Engine/Source/Runtime/VulkanRHI/Private/VulkanDevice.h
--- a/Engine/Source/Runtime/VulkanRHI/Private/VulkanDevice.h
+++ b/Engine/Source/Runtime/VulkanRHI/Private/VulkanDevice.h
@@ -128,6 +128,8 @@ struct FOptionalVulkanDeviceExtensionProperties
 	VkPhysicalDeviceAccelerationStructurePropertiesKHR AccelerationStructureProps;
 	VkPhysicalDeviceRayTracingPipelinePropertiesKHR RayTracingPipelineProps;
 #endif // VULKAN_RHI_RAYTRACING
+
+	VkPhysicalDeviceFragmentShadingRateFeaturesKHR FragmentShadingRateFeatures;
 };
 
 class FVulkanPhysicalDeviceFeatures
diff --git a/Engine/Source/Runtime/VulkanRHI/Private/VulkanExtensions.cpp b/Engine/Source/Runtime/VulkanRHI/Private/VulkanExtensions.cpp
--- a/Engine/Source/Runtime/VulkanRHI/Private/VulkanExtensions.cpp
+++ b/Engine/Source/Runtime/VulkanRHI/Private/VulkanExtensions.cpp
@@ -26,6 +26,16 @@
 
 
 
+enum class EVulkanVariableRateShadingPreference : uint8
+{
+	PreferFSR = 0, // This can be used to prefer FSR over FDM when both are available.
+	UseFSROnlyIfAvailable, // This will only print a message if FSR is not available.
+	RequireFSR, // This will print an error if FSR is not available.
+	PreferFDM, // This can be used to prefer FDM over FSR when both are available.
+	UseFDMOnlyIfAvailable, // This will only print a message if FDM is not available.
+	RequireFDM, // This will print an error if FDM is not available.
+};
+
 TAutoConsoleVariable<int32> GRHIAllow64bitShaderAtomicsCvar(
 	TEXT("r.Vulkan.Allow64bitShaderAtomics"),
 	1,
@@ -71,6 +81,19 @@ TAutoConsoleVariable<int32> GVulkaAllowSync2BarriersCVar(
 	ECVF_ReadOnly
 );
 
+TAutoConsoleVariable<int32> GVulkanVariableRateShadingFormatCVar(
+	TEXT("r.Vulkan.VRSFormat"),
+	0,
+	TEXT("Allows to choose the preferred Variable Rate Shading option. \n")
+	TEXT("0: Prefer Fragment Shading Rate if both Fragment Shading Rate and Fragment Density Map are available.\n")
+	TEXT("1: Use Fragment Shading Rate if available. A message will be reported if not available. \n")
+	TEXT("2: Require Fragment Shading Rate. Will generate an error if the extension is not available. \n")
+	TEXT("3: Prefer Fragment Density Map if both Fragment Shading Rate and Fragment Density Map are available.\n")
+	TEXT("4: Use Fragment Density Map if available. A message will be reported if not available.\n")
+	TEXT("5: Require Fragment Density Map. Will generate an error if the extension is not available."),
+	ECVF_ReadOnly
+);
+
 #if VULKAN_HAS_DEBUGGING_ENABLED
 extern TAutoConsoleVariable<int32> GGPUValidationCvar;
 #endif
@@ -487,7 +510,10 @@ public:
 	FVulkanKHRFragmentShadingRateExtension(FVulkanDevice* InDevice)
 		: FVulkanDeviceExtension(InDevice, VK_KHR_FRAGMENT_SHADING_RATE_EXTENSION_NAME, VULKAN_EXTENSION_ENABLED)
 	{
+		int32 VRSFormatPreference = GVulkanVariableRateShadingFormatCVar->GetInt();
 		bEnabledInCode = bEnabledInCode && GRHIVariableRateShadingEnabled;
+		// FSR should be enabled even if FDM is preferred because it could be not available.
+		bEnabledInCode &= (VRSFormatPreference <= (uint8) EVulkanVariableRateShadingPreference::RequireFSR || VRSFormatPreference == (uint8)EVulkanVariableRateShadingPreference::PreferFDM);
 	}
 
 	virtual void PrePhysicalDeviceProperties(VkPhysicalDeviceProperties2KHR& PhysicalDeviceProperties2) override final
@@ -500,6 +526,7 @@ public:
 	{
 		ExtensionFlags.HasKHRFragmentShadingRate = 1;
 
+		VkPhysicalDeviceFragmentShadingRateFeaturesKHR& FragmentShadingRateFeatures = GetDeviceExtensionProperties().FragmentShadingRateFeatures;
 		GRHISupportsAttachmentVariableRateShading = (FragmentShadingRateFeatures.attachmentFragmentShadingRate == VK_TRUE);
 		GRHISupportsPipelineVariableRateShading = (FragmentShadingRateFeatures.pipelineFragmentShadingRate == VK_TRUE);
 
@@ -517,6 +544,7 @@ public:
 
 	virtual void PostPhysicalDeviceProperties() override final
 	{
+		VkPhysicalDeviceFragmentShadingRateFeaturesKHR& FragmentShadingRateFeatures = GetDeviceExtensionProperties().FragmentShadingRateFeatures;
 		if (FragmentShadingRateFeatures.attachmentFragmentShadingRate == VK_TRUE)
 		{
 			GRHIVariableRateShadingImageTileMinWidth = FragmentShadingRateProperties.minFragmentShadingRateAttachmentTexelSize.width;
@@ -538,12 +566,14 @@ public:
 
 	virtual void PrePhysicalDeviceFeatures(VkPhysicalDeviceFeatures2KHR& PhysicalDeviceFeatures2) override final
 	{
+		VkPhysicalDeviceFragmentShadingRateFeaturesKHR& FragmentShadingRateFeatures = GetDeviceExtensionProperties().FragmentShadingRateFeatures;
 		ZeroVulkanStruct(FragmentShadingRateFeatures, VK_STRUCTURE_TYPE_PHYSICAL_DEVICE_FRAGMENT_SHADING_RATE_FEATURES_KHR);
 		AddToPNext(PhysicalDeviceFeatures2, FragmentShadingRateFeatures);
 	}
 
 	virtual void PreCreateDevice(VkDeviceCreateInfo& DeviceCreateInfo) override final
 	{
+		VkPhysicalDeviceFragmentShadingRateFeaturesKHR& FragmentShadingRateFeatures = GetDeviceExtensionProperties().FragmentShadingRateFeatures;
 		if (FragmentShadingRateFeatures.attachmentFragmentShadingRate == VK_TRUE || FragmentShadingRateFeatures.pipelineFragmentShadingRate == VK_TRUE)
 		{
 			AddToPNext(DeviceCreateInfo, FragmentShadingRateFeatures);
@@ -552,7 +582,6 @@ public:
 
 private:
 	VkPhysicalDeviceFragmentShadingRatePropertiesKHR FragmentShadingRateProperties;
-	VkPhysicalDeviceFragmentShadingRateFeaturesKHR FragmentShadingRateFeatures;
 };
 
 
@@ -565,7 +594,10 @@ public:
 	FVulkanEXTFragmentDensityMapExtension(FVulkanDevice* InDevice)
 		: FVulkanDeviceExtension(InDevice, VK_EXT_FRAGMENT_DENSITY_MAP_EXTENSION_NAME, VULKAN_EXTENSION_ENABLED)
 	{
+		int32 VRSFormatPreference = GVulkanVariableRateShadingFormatCVar->GetInt();
 		bEnabledInCode = bEnabledInCode && GRHIVariableRateShadingEnabled;
+		// FDM should be enabled even if the preferred choice is FSR because that might not be available.
+		bEnabledInCode &= (VRSFormatPreference >= (uint8) EVulkanVariableRateShadingPreference::PreferFDM || VRSFormatPreference == (uint8) EVulkanVariableRateShadingPreference::PreferFSR);
 	}
 
 	virtual void PrePhysicalDeviceProperties(VkPhysicalDeviceProperties2KHR& PhysicalDeviceProperties2) override final
@@ -578,10 +610,20 @@ public:
 	{
 		bRequirementsPassed = (FragmentDensityMapFeatures.fragmentDensityMap == VK_TRUE);
 		ExtensionFlags.HasEXTFragmentDensityMap = bRequirementsPassed;
+	}
 
-		// Use the Fragment Density Map extension if and only if the Fragment Shading Rate extension is not available.
+	virtual void PrePhysicalDeviceFeatures(VkPhysicalDeviceFeatures2KHR& PhysicalDeviceFeatures2) override final
+	{
+		ZeroVulkanStruct(FragmentDensityMapFeatures, VK_STRUCTURE_TYPE_PHYSICAL_DEVICE_FRAGMENT_DENSITY_MAP_FEATURES_EXT);
+		AddToPNext(PhysicalDeviceFeatures2, FragmentDensityMapFeatures);
+	}
+
+	virtual void PostPhysicalDeviceProperties() override final
+	{
+		// Use the Fragment Density Map extension if the Fragment Shading Rate extension is not available or if both are available but 
+		// Fragment Density Map is the preferred user choice.
 		// NOTE: FVulkanKHRFragmentShadingRateExtension must be placed before FVulkanEXTFragmentDensityMapExtension for this to work!
-		if (!GRHISupportsAttachmentVariableRateShading && bRequirementsPassed)
+		if ((!GRHISupportsAttachmentVariableRateShading || (GVulkanVariableRateShadingFormatCVar->GetInt() >= (uint8) EVulkanVariableRateShadingPreference::PreferFDM)) && bRequirementsPassed)
 		{
 			GRHISupportsAttachmentVariableRateShading = true;
 
@@ -601,12 +643,6 @@ public:
 		}
 	}
 
-	virtual void PrePhysicalDeviceFeatures(VkPhysicalDeviceFeatures2KHR& PhysicalDeviceFeatures2) override final
-	{
-		ZeroVulkanStruct(FragmentDensityMapFeatures, VK_STRUCTURE_TYPE_PHYSICAL_DEVICE_FRAGMENT_DENSITY_MAP_FEATURES_EXT);
-		AddToPNext(PhysicalDeviceFeatures2, FragmentDensityMapFeatures);
-	}
-
 	virtual void PreCreateDevice(VkDeviceCreateInfo& DeviceCreateInfo) override final
 	{
 		if (bRequirementsPassed)
@@ -1543,3 +1579,39 @@ FVulkanInstanceExtensionArray FVulkanInstanceExtension::GetUESupportedInstanceEx
 	return OutUEInstanceExtensions;
 }
 
+void ChooseVariableRateShadingMethod(FOptionalVulkanDeviceExtensions& ExtensionFlags, const VkPhysicalDeviceFragmentShadingRateFeaturesKHR& FragmentShadingRateFeatures)
+{
+	int32 VRSFormatPreference = GVulkanVariableRateShadingFormatCVar->GetInt();
+
+	if (VRSFormatPreference == (uint8)EVulkanVariableRateShadingPreference::PreferFSR && (FragmentShadingRateFeatures.attachmentFragmentShadingRate == VK_TRUE))
+	{
+		ExtensionFlags.HasEXTFragmentDensityMap = 0;
+		return;
+	}
+	if (VRSFormatPreference == (uint8)EVulkanVariableRateShadingPreference::PreferFDM && ExtensionFlags.HasEXTFragmentDensityMap)
+	{
+		ExtensionFlags.HasKHRFragmentShadingRate = 0;
+		return;
+	}
+	if (VRSFormatPreference == (uint8)EVulkanVariableRateShadingPreference::UseFSROnlyIfAvailable && !(FragmentShadingRateFeatures.attachmentFragmentShadingRate == VK_TRUE))
+	{
+		UE_LOG(LogVulkanRHI, Display, TEXT("Fragment Shading Rate was requested but is not available."));
+		return;
+	}
+	if (VRSFormatPreference == (uint8)EVulkanVariableRateShadingPreference::UseFDMOnlyIfAvailable && !ExtensionFlags.HasEXTFragmentDensityMap)
+	{
+		UE_LOG(LogVulkanRHI, Display, TEXT("Fragment Density Map was requested but is not available."));
+		return;
+	}
+	if (VRSFormatPreference == (uint8)EVulkanVariableRateShadingPreference::RequireFSR && !(FragmentShadingRateFeatures.attachmentFragmentShadingRate == VK_TRUE))
+	{
+		UE_LOG(LogVulkanRHI, Error, TEXT("Fragment Shading Rate was required but is not available."));
+		return;
+	}
+	if (VRSFormatPreference == (uint8)EVulkanVariableRateShadingPreference::RequireFDM && !ExtensionFlags.HasEXTFragmentDensityMap)
+	{
+		UE_LOG(LogVulkanRHI, Error, TEXT("Fragment Density Map was required but is not available."));
+		return;
+	}
+}
+
diff --git a/Engine/Source/Runtime/VulkanRHI/Private/VulkanExtensions.h b/Engine/Source/Runtime/VulkanRHI/Private/VulkanExtensions.h
--- a/Engine/Source/Runtime/VulkanRHI/Private/VulkanExtensions.h
+++ b/Engine/Source/Runtime/VulkanRHI/Private/VulkanExtensions.h
@@ -157,4 +157,6 @@ protected:
 
 // Helpers to create simple extensions that set a flag
 #define DEVICE_EXT_FLAG_SETTER(FLAG_NAME) [](FOptionalVulkanDeviceExtensions& ExtensionFlags) { ExtensionFlags.FLAG_NAME = 1; }
-#define INSTANCE_EXT_FLAG_SETTER(FLAG_NAME) [](FOptionalVulkanInstanceExtensions& ExtensionFlags) { ExtensionFlags.FLAG_NAME = 1; }
\ No newline at end of file
+#define INSTANCE_EXT_FLAG_SETTER(FLAG_NAME) [](FOptionalVulkanInstanceExtensions& ExtensionFlags) { ExtensionFlags.FLAG_NAME = 1; }
+
+void ChooseVariableRateShadingMethod(FOptionalVulkanDeviceExtensions& ExtensionFlags, const VkPhysicalDeviceFragmentShadingRateFeaturesKHR& FragmentShadingRateFeatures);
\ No newline at end of file
diff --git a/Engine/Source/Runtime/VulkanRHI/Private/VulkanRenderpass.h b/Engine/Source/Runtime/VulkanRHI/Private/VulkanRenderpass.h
--- a/Engine/Source/Runtime/VulkanRHI/Private/VulkanRenderpass.h
+++ b/Engine/Source/Runtime/VulkanRHI/Private/VulkanRenderpass.h
@@ -425,7 +425,12 @@ public:
 
 		const bool bDeferredShadingSubpass = RTLayout.GetSubpassHint() == ESubpassHint::DeferredShadingSubpass;
 		const bool bDepthReadSubpass = RTLayout.GetSubpassHint() == ESubpassHint::DepthReadSubpass;
-		const bool bApplyFragmentShadingRate =  GRHISupportsAttachmentVariableRateShading && GRHIVariableRateShadingEnabled && GRHIAttachmentVariableRateShadingEnabled && RTLayout.GetFragmentDensityAttachmentReference() != nullptr;
+		const bool bApplyFragmentShadingRate =  GRHISupportsAttachmentVariableRateShading 
+												&& GRHIVariableRateShadingEnabled 
+												&& GRHIAttachmentVariableRateShadingEnabled 
+												&& RTLayout.GetFragmentDensityAttachmentReference() != nullptr
+												&& Device.GetOptionalExtensions().HasKHRFragmentShadingRate 
+												&& Device.GetOptionalExtensionProperties().FragmentShadingRateFeatures.attachmentFragmentShadingRate == VK_TRUE;
 		const bool bHasDepthStencilAttachmentReference = (RTLayout.GetDepthAttachmentReference() != nullptr);
 
 		if (bApplyFragmentShadingRate)
diff --git a/Engine/Source/Runtime/VulkanRHI/Private/VulkanTexture.cpp b/Engine/Source/Runtime/VulkanRHI/Private/VulkanTexture.cpp
--- a/Engine/Source/Runtime/VulkanRHI/Private/VulkanTexture.cpp
+++ b/Engine/Source/Runtime/VulkanRHI/Private/VulkanTexture.cpp
@@ -508,6 +508,21 @@ struct FRHICommandOnDestroyImage final : public FRHICommand<FRHICommandOnDestroy
 	}
 };
 
+static VkImageLayout ChooseVRSLayout()
+{
+	if(GRHIVariableRateShadingImageDataType == VRSImage_Palette)
+	{
+		return VK_IMAGE_LAYOUT_FRAGMENT_SHADING_RATE_ATTACHMENT_OPTIMAL_KHR;
+	}
+	else if(GRHIVariableRateShadingImageDataType == VRSImage_Fractional)
+	{
+		return VK_IMAGE_LAYOUT_FRAGMENT_DENSITY_MAP_OPTIMAL_EXT;
+	}
+
+	checkNoEntry();
+	return VK_IMAGE_LAYOUT_UNDEFINED;
+}
+
 static VkImageLayout GetInitialLayoutFromRHIAccess(ERHIAccess RHIAccess, bool bIsDepthStencil, bool bSupportReadOnlyOptimal)
 {
 	if (EnumHasAnyFlags(RHIAccess, ERHIAccess::RTV) || RHIAccess == ERHIAccess::Present)
@@ -546,7 +561,7 @@ static VkImageLayout GetInitialLayoutFromRHIAccess(ERHIAccess RHIAccess, bool bI
 		case ERHIAccess::Discard:	return VK_IMAGE_LAYOUT_UNDEFINED;
 		case ERHIAccess::CopySrc:	return VK_IMAGE_LAYOUT_TRANSFER_SRC_OPTIMAL;
 		case ERHIAccess::CopyDest:	return VK_IMAGE_LAYOUT_TRANSFER_DST_OPTIMAL;
-		case ERHIAccess::ShadingRateSource:	return VK_IMAGE_LAYOUT_FRAGMENT_DENSITY_MAP_OPTIMAL_EXT; // TODO: This should be set to FSR or FDM depending on the active extension. Default to FDM for the moment.
+		case ERHIAccess::ShadingRateSource:	return ChooseVRSLayout();
 	}
 
 	checkf(false, TEXT("Invalid initial access %d"), RHIAccess);
