Fixed: vulkan crash related r.Vulkan.DelayAcquireBackBuffer=0

diff --git a/Engine/Source/Runtime/VulkanRHI/Private/VulkanViewport.cpp b/Engine/Source/Runtime/VulkanRHI/Private/VulkanViewport.cpp
--- a/Engine/Source/Runtime/VulkanRHI/Private/VulkanViewport.cpp
+++ b/Engine/Source/Runtime/VulkanRHI/Private/VulkanViewport.cpp
@@ -26,6 +26,7 @@ void FVulkanBackBuffer::ReleaseAcquiredImage()
 	{
 		// Do not invalidate view here, just remove a reference to it
 		DefaultView = nullptr;
+		PartialView = nullptr;
 	}
 
 	Image = VK_NULL_HANDLE;
@@ -76,6 +77,7 @@ void FVulkanBackBuffer::AcquireBackBufferImage(FVulkanCommandListContext& Contex
 
 			Image = ImageView.GetTextureView().Image;
 			DefaultView = &ImageView;
+			PartialView = &ImageView;
 
 			FVulkanCommandBufferManager* CmdBufferManager = Context.GetCommandBufferManager();
 			FVulkanCmdBuffer* CmdBuffer = CmdBufferManager->GetActiveCmdBuffer();
@@ -96,6 +98,7 @@ void FVulkanBackBuffer::AcquireBackBufferImage(FVulkanCommandListContext& Contex
 			FVulkanView* DummyView = Viewport->RenderingBackBuffer->DefaultView;
 			Image = DummyView->GetTextureView().Image;
 			DefaultView = DummyView;
+			PartialView = DummyView;
 		}
 	}
 }
