VK RHI: Correctly check support for ASTC HDR extension

diff --git a/Engine/Source/Runtime/VulkanRHI/Private/Android/VulkanAndroidPlatform.cpp b/Engine/Source/Runtime/VulkanRHI/Private/Android/VulkanAndroidPlatform.cpp
--- a/Engine/Source/Runtime/VulkanRHI/Private/Android/VulkanAndroidPlatform.cpp
+++ b/Engine/Source/Runtime/VulkanRHI/Private/Android/VulkanAndroidPlatform.cpp
@@ -446,7 +446,7 @@ void FVulkanAndroidPlatform::GetDeviceExtensions(FVulkanDevice* Device, FVulkanD
 	{
 		OutExtensions.Add(MakeUnique<FVulkanDeviceExtension>(Device, VK_EXT_ASTC_DECODE_MODE_EXTENSION_NAME, VULKAN_SUPPORTS_ASTC_DECODE_MODE, VULKAN_EXTENSION_NOT_PROMOTED, DEVICE_EXT_FLAG_SETTER(HasEXTASTCDecodeMode)));
 	}
-	//OutExtensions.Add(MakeUnique<FVulkanDeviceExtension>(Device, VK_EXT_TEXTURE_COMPRESSION_ASTC_HDR_EXTENSION_NAME, VULKAN_SUPPORTS_TEXTURE_COMPRESSION_ASTC_HDR, VK_API_VERSION_1_3, DEVICE_EXT_FLAG_SETTER(HasEXTTextureCompressionASTCHDR)));
+	OutExtensions.Add(MakeUnique<FVulkanDeviceExtension>(Device, VK_EXT_TEXTURE_COMPRESSION_ASTC_HDR_EXTENSION_NAME, VULKAN_SUPPORTS_TEXTURE_COMPRESSION_ASTC_HDR, VK_API_VERSION_1_3, DEVICE_EXT_FLAG_SETTER(HasEXTTextureCompressionASTCHDR)));
 
 	if (GVulkanQcomRenderPassTransform)
 	{
diff --git a/Engine/Source/Runtime/VulkanRHI/Private/VulkanDevice.cpp b/Engine/Source/Runtime/VulkanRHI/Private/VulkanDevice.cpp
--- a/Engine/Source/Runtime/VulkanRHI/Private/VulkanDevice.cpp
+++ b/Engine/Source/Runtime/VulkanRHI/Private/VulkanDevice.cpp
@@ -733,7 +733,7 @@ void FVulkanDevice::SetupFormats()
 		MapFormatSupport(PF_G16, { VK_FORMAT_R16_SFLOAT, VK_FORMAT_R16_UNORM }, ComponentMappingR001);
 	}
 
-	if (GetOptionalExtensions().HasEXTASTCDecodeMode)
+	if (GetOptionalExtensions().HasEXTTextureCompressionASTCHDR)
 	{
 		MapFormatSupport(PF_ASTC_4x4_HDR,   { VK_FORMAT_ASTC_4x4_SFLOAT_BLOCK_EXT },   ComponentMappingRGBA);
 		MapFormatSupport(PF_ASTC_6x6_HDR,   { VK_FORMAT_ASTC_6x6_SFLOAT_BLOCK_EXT },   ComponentMappingRGBA);
