Add alternative 'chunked' vulkan PSO cache.
 - Breaks PSO cache up into chunks of any specified PSO count (controlled via cvar)
 - All encountered PSOs are recorded in the cache.
 - a 'chunk' represents a single vulkan PSO cache object.
 - vulkan caches can be unloaded and reloaded when required. (saving memory but potentially hitches.)
 - If the cache grows beyond a threshold size it is deleted and rebuilt.
 - Enable vulkan support for PSO precaching when r.PSOPrecaching + chunked PSO cache is enabled.
 - by default vulkan precaching support disables the PSO filecache, this can be overridden with r.Vulkan.EnablePSOFileCacheWhenPrecachingActive

diff --git a/Engine/Source/Runtime/VulkanRHI/Private/VulkanChunkedPipelineCache.cpp b/Engine/Source/Runtime/VulkanRHI/Private/VulkanChunkedPipelineCache.cpp
new file mode 100644
--- /dev/null
+++ b/Engine/Source/Runtime/VulkanRHI/Private/VulkanChunkedPipelineCache.cpp
@@ -0,0 +1,1157 @@
+// Copyright Epic Games, Inc. All Rights Reserved.
+
+/*=============================================================================
+	VulkanChunkedPipelineCache.cpp: 
+=============================================================================*/
+
+#include "VulkanChunkedPipelineCache.h"
+#include "VulkanRHIPrivate.h"
+#include "VulkanPipeline.h"
+#include "HAL/PlatformFileManager.h"
+#include "HAL/FileManager.h"
+#include "Misc/Paths.h"
+#include "Misc/FileHelper.h"
+#include "Serialization/MemoryReader.h"
+#include "Serialization/MemoryWriter.h"
+#include "VulkanPendingState.h"
+#include "VulkanContext.h"
+#include "GlobalShader.h"
+#include "VulkanLLM.h"
+#include "Misc/ScopeRWLock.h"
+#include "HAL/PlatformFramePacer.h"
+#include "Templates/Greater.h"
+#include "Async/MappedFileHandle.h"
+
+/*
+* Chunked PSO cache is a collection of multiple VK PSO cache objects (referred to as a cache chunks)
+* Cache chunks accumulate only, we do not currently delete them except at startup when the entire cache is removed (when it is deemed over sized).
+* 
+* FVulkanChunkedPipelineCacheManagerImpl implements the external API, it contains:
+* a map of PSO precache hash -> cache chunk key, a map of chunkkey->FVulkanPipelineCacheChunk,
+* and an LRU of recently used cache chunks, (used when cache eviction is active)
+* 
+* when a PSO is created its hash is checked against the PSOPrecacheHash->cache chunk map, if found FVulkanPipelineCacheChunk is responsible for creating and managing the VK PSO object.
+* If not found the most recent FVulkanPipelineCacheChunk is queried for capacity, if space is available the new PSO is added to the chunk.
+* If the most recent chunk is full, it is closed and its data is appended to the permanent storage file. a new chunk is created and the PSO is added to it.
+* 
+* At startup the entire cache file is scanned to build the PSOPrecacheHash->chunk map, the entire cache is mmapped but the PSO cache data is not initially loaded.
+* FVulkanPipelineCacheChunk will load its data from the mmapped cache file on demand. the cache is then considered 'resident'
+* 
+* If the sum of the resident cache chunks goes beyond a cvar specified memory threshold then the least recently used cache chunk is evicted. (eviction means the VK PSO cache object is destroyed and its memory freed up.)
+* Only chunks that are mmapped can be evicted.
+* 
+*/
+
+#define LOGCACHEINFO 0
+
+#if LOGCACHEINFO
+#include "ProfilingDebugging/ScopedTimers.h"
+#define FScopedTimeToLog FScopedDurationTimeLogger
+#else
+class FScopedTimeToLog
+{
+public:
+	FScopedTimeToLog(FString, class FOutputDevice* output = nullptr) {};
+};
+#endif
+#include "../../Core/Public/Containers/LockFreeList.h"
+
+#if PLATFORM_ANDROID && USE_ANDROID_FILE
+// TODO:
+extern FString GExternalFilePath;
+#endif
+namespace UE
+{
+	namespace Vulkan
+	{
+		static int32 GUseChunkedPSOCache = 0;
+		static FAutoConsoleVariableRef GVulkanGUseNewCacheCodeCVar(
+			TEXT("r.Vulkan.UseChunkedPSOCache"),
+			GUseChunkedPSOCache,
+			TEXT("\n")
+			TEXT("")
+			, ECVF_RenderThreadSafe | ECVF_ReadOnly
+		);
+
+		static FVulkanDevice* GetVulkanDevice()
+		{
+			return GVulkanRHI->GetDevice();
+		}
+
+		static VkDevice GetVulkanDeviceHandle()
+		{
+			return GetVulkanDevice()->GetInstanceHandle();
+		}
+
+		int32 GMaxPSOsPerChunk				= 20;
+		int32 GTargetResidentCacheSizeMb	= 40;
+		int32 GMaxTotalCacheSizeMb			= 300;
+		int32 GUntouchedChunkEvictTimeSeconds = 60; 
+
+		static FAutoConsoleVariableRef GMaxSingleCachePSOCountCVar(
+			TEXT("r.Vulkan.ChunkedPSOCache.MaxSingleCachePSOCount"),
+			GMaxPSOsPerChunk,
+			TEXT("The target PSO count for an individual PSO cache.\n")
+			TEXT("existing caches with different PSO counts are discarded at startup.\n")
+			TEXT("(default) 20")
+			, ECVF_RenderThreadSafe | ECVF_ReadOnly
+		);
+
+		static FAutoConsoleVariableRef GTargetResidentCacheSizeCVar(
+			TEXT("r.Vulkan.ChunkedPSOCache.TargetResidentCacheSizeMb"),
+			GTargetResidentCacheSizeMb,
+			TEXT("A target resident cache size in MB, if the combined memory usage of all the currently loaded cache chunks is above this threshold\n")
+			TEXT("the least recently used chunks will be considered for eviction.\n")
+			TEXT("(default) 40")
+			, ECVF_RenderThreadSafe | ECVF_ReadOnly
+		);
+
+		static FAutoConsoleVariableRef GMaxTotalCacheSizeMbCVar(
+			TEXT("r.Vulkan.ChunkedPSOCache.MaxTotalCacheSizeMb"),
+			GMaxTotalCacheSizeMb,
+			TEXT("At startup, if the entire cache is above this threshold the cache will be deleted\n")
+			TEXT("and rebuilt during the subsequent run.\n")			
+			TEXT("(default) 300\n")
+			TEXT("0 to disable cache size limit, note that the cache will grow indefinitely.")
+			, ECVF_RenderThreadSafe | ECVF_ReadOnly
+		);
+
+		static bool GMemoryMapChunkedPSOCache = true;
+		static FAutoConsoleVariableRef CVarMemoryMapChunkedPSOCache(
+			TEXT("r.Vulkan.MemoryMapChunkedPSOCache"),
+			GMemoryMapChunkedPSOCache,
+			TEXT("If true enabled memory mapping of the chunked vulkan PSO cache. (default)\n")
+			TEXT("\n")
+			TEXT("")
+			,
+			ECVF_ReadOnly | ECVF_RenderThreadSafe
+		);
+
+		static FAutoConsoleVariableRef GChunkEvictTimeCVar(
+			TEXT("r.Vulkan.ChunkedPSOCache.ChunkEvictTime"),
+			GUntouchedChunkEvictTimeSeconds,
+			TEXT("Time in seconds for a cache chunk to be unused before it can be evicted from ram.\n")
+			TEXT("(default) 60")
+			, ECVF_RenderThreadSafe | ECVF_ReadOnly
+		);
+
+		bool CanMemoryMapChunkedPSOCache()
+		{
+			return FPlatformProperties::SupportsMemoryMappedFiles() && GMemoryMapChunkedPSOCache;
+		}
+
+		static FCriticalSection CacheVersionedFolderCriticalSection;
+		static FString GetPSOBinaryCacheVersionedFolder()
+		{
+			static FString BinaryCacheVersionKey;
+			if (BinaryCacheVersionKey.IsEmpty())
+			{
+				FScopeLock Lock(&CacheVersionedFolderCriticalSection);
+				if (BinaryCacheVersionKey.IsEmpty())
+				{
+					BinaryCacheVersionKey.Append(LegacyShaderPlatformToShaderFormat(GMaxRHIShaderPlatform).ToString());
+					const VkPhysicalDeviceProperties& DeviceProperties = GetVulkanDevice()->GetDeviceProperties();
+
+					BinaryCacheVersionKey.Append(FString::Printf(TEXT(".%x.%x.%x"), FCrc::MemCrc32(DeviceProperties.pipelineCacheUUID, VK_UUID_SIZE), DeviceProperties.vendorID, DeviceProperties.deviceID));
+#if PLATFORM_ANDROID
+					// Apparently we can't rely on version alone to assume binary compatibility.
+					// Some devices have reported binary compatibility errors after minor OS updates even though the driver version number does not change.
+					const FString BuildNumber = FAndroidMisc::GetDeviceBuildNumber();
+					BinaryCacheVersionKey.Append(BuildNumber);
+
+					// Optional configrule variable for triggering a rebuild of the cache.
+					const FString* ConfigRulesVulkanProgramKey = FAndroidMisc::GetConfigRulesVariable(TEXT("VulkanProgramCacheKey"));
+					if (ConfigRulesVulkanProgramKey && !ConfigRulesVulkanProgramKey->IsEmpty())
+					{
+						BinaryCacheVersionKey.Append(*ConfigRulesVulkanProgramKey);
+					}
+#endif
+				}
+			}
+			check(!BinaryCacheVersionKey.IsEmpty());
+
+			return GVulkanRHI->GetName() / BinaryCacheVersionKey;
+		}
+
+		static FString GetRHICacheRootFolder()
+		{
+#if PLATFORM_ANDROID && USE_ANDROID_FILE
+			static FString RHICacheTopFolderPath = GExternalFilePath / TEXT("RHICache");
+#else
+			static FString RHICacheTopFolderPath = FPaths::ProjectSavedDir() / TEXT("RHICache");
+#endif
+			return RHICacheTopFolderPath;
+		}
+
+		using FPipelineCacheChunkKey = uint32;
+		using FVulkanRHIGraphicsPipelineStateLRU = TDoubleLinkedList<FPipelineCacheChunkKey>;
+		using FVulkanPipelineCacheChunkLRUNode = FVulkanRHIGraphicsPipelineStateLRU::TDoubleLinkedListNode;	
+
+		std::atomic<int32> TotalResidentCacheSize = 0;
+	}
+}
+
+// this class manages a file that combines all of the cache chunks
+// It manages access to each chunk via a single mmap alloc. mmap support is relied on for perf, as a fallback where it's not supported synchronous file access is used.
+class FVulkanCombinedChunkCacheFile
+{
+	TUniquePtr<IMappedFileHandle> MappedCacheFile;
+	
+	mutable FRWLock MappingLock;
+	TUniquePtr<IMappedFileRegion> MappedRegion;
+
+	TUniquePtr<FArchive> PSOFileWriter = nullptr;
+
+	static const uint32 CacheFileVersion = 2;
+	constexpr static TCHAR FileName[] = TEXT("VulkanPSOChunks");
+
+	void UpdateMapping(uint32 Size)
+	{
+		FRWScopeLock Lock(MappingLock, SLT_Write);
+		MappedRegion = TUniquePtr<IMappedFileRegion>(MappedCacheFile->MapRegion(0, Size));
+	}
+
+public:
+
+	static FVulkanCombinedChunkCacheFile& Get()
+	{
+		static FVulkanCombinedChunkCacheFile Impl;
+		return Impl;
+	}
+
+	const TCHAR* GetFilename() const { return FileName; }
+	
+	const FString& GetFullCachePath()
+	{
+		static FString FullCachePath;
+		if(FullCachePath.IsEmpty())
+		{
+			const FString RootCacheFolder = UE::Vulkan::GetRHICacheRootFolder();
+			const FString CacheSubDir = UE::Vulkan::GetPSOBinaryCacheVersionedFolder();
+			const FString CombinedCacheSubDir = FPaths::Combine(RootCacheFolder, CacheSubDir);
+			FullCachePath = FPaths::Combine(CombinedCacheSubDir, FileName);
+		}
+		return FullCachePath;
+	}
+
+	// hash of parameters used while building the cache
+	// a clash would mean some cache chunks would not be honoring the cvar size limits.
+	static uint32 GetCacheBuildingParamHash()
+	{
+		uint32 ParamHash = 0;
+		ParamHash = FCrc::MemCrc32(&UE::Vulkan::GMaxPSOsPerChunk, sizeof(UE::Vulkan::GMaxPSOsPerChunk), ParamHash);
+		return ParamHash;
+	}
+
+	static void WriteFileHeader(FArchive& Archive, uint32 LastValidOffset)
+	{
+		uint32 Version = CacheFileVersion;
+		uint32 ParamHash = GetCacheBuildingParamHash();
+		Archive << Version;
+		Archive << ParamHash;
+		Archive << LastValidOffset;
+	}
+
+	class FPSOArchiveReader
+	{
+		const uint8* PSOBytes;
+		TUniquePtr<FArchive> Archive;
+		TUniquePtr<FRWScopeLock> MappingLock;
+	public:
+		explicit FPSOArchiveReader(TUniquePtr<FArchive>&& InArchive, const uint8* MappedBytes, TUniquePtr<FRWScopeLock>&& LockIn)
+			: PSOBytes(MappedBytes), Archive(MoveTemp(InArchive)), MappingLock(MoveTemp(LockIn))
+		{ }
+		explicit FPSOArchiveReader(TUniquePtr<FArchive>&& InArchive)
+			: PSOBytes(nullptr), Archive(MoveTemp(InArchive))
+		{ }
+
+		bool IsValid() const { return Archive.IsValid(); }
+		const uint8* GetData() const { return PSOBytes; }
+		FArchive* GetArchive() { return Archive.Get(); }
+	};
+
+	bool ReadAllCacheChunks(FPSOArchiveReader& ArchiveReader, TUniqueFunction<void(FPSOArchiveReader& ArchiveReader)> OnFoundCacheChunk)
+	{
+		const FString FullCachePath = GetFullCachePath();
+
+		FArchive& Archive = *ArchiveReader.GetArchive();
+		uint32 Version;
+		Archive << Version;
+		if (Version != CacheFileVersion)
+		{
+			UE_LOG(LogVulkanRHI, Log, TEXT("FVulkanChunkedPipelineCacheManager: incorrect Cache file version (%d, expected %d)"), Version, CacheFileVersion);
+			return false;
+		}
+
+		uint32 ParamHash;
+		uint32 LastValidOffset;
+		Archive << Version;
+		Archive << ParamHash;
+		Archive << LastValidOffset;
+		if (LastValidOffset == 0 || ParamHash != GetCacheBuildingParamHash())
+		{
+			UE_LOG(LogVulkanRHI, Log, TEXT("FVulkanChunkedPipelineCacheManager: Cache file params have changed (%d, %x=%x)"), LastValidOffset, ParamHash, GetCacheBuildingParamHash());
+			return false;
+		}
+
+		if (UE::Vulkan::CanMemoryMapChunkedPSOCache())
+		{
+			MappedCacheFile = TUniquePtr<IMappedFileHandle>(FPlatformFileManager::Get().GetPlatformFile().OpenMapped(*FullCachePath));
+			UpdateMapping(LastValidOffset);
+		}
+
+		for (int ChunkIdx = 0; Archive.Tell() < LastValidOffset; ChunkIdx++)
+		{
+			OnFoundCacheChunk(ArchiveReader);
+		}
+
+		check(Archive.Tell() == LastValidOffset);
+		return true;
+	}
+
+	// open the file and callback OnFoundCacheChunk for each contained cache chunk
+	bool LoadAllCacheChunks(TUniqueFunction<void(FPSOArchiveReader& ArchiveReader)> OnFoundCacheChunk)
+	{
+		const FString FullCachePath = GetFullCachePath();
+		
+		IPlatformFile& PlatformFile = FPlatformFileManager::Get().GetPlatformFile();
+		FFileStatData Stats = PlatformFile.GetStatData(*FullCachePath);
+
+		bool bReadSuccess = false;
+
+		TArray<int32> OffsetsToChunkData;
+		// TODO: simplistic 'GC', blow cache away and rebuild.
+		// This should be acceptable as the cache rebuild needs to be as transparent as possible..
+		if (Stats.FileSize < (UE::Vulkan::GMaxTotalCacheSizeMb * 1024 * 1024))
+		{
+			FPSOArchiveReader ArchiveReader(TUniquePtr<FArchive>(IFileManager::Get().CreateFileReader(*FullCachePath, FILEREAD_AllowWrite)), nullptr, nullptr);
+
+			if (ArchiveReader.IsValid())
+			{
+				UE_LOG(LogVulkanRHI, Log, TEXT("FVulkanChunkedPipelineCacheManager: loading %s"), *FullCachePath);
+				bReadSuccess = ReadAllCacheChunks(ArchiveReader, MoveTemp(OnFoundCacheChunk));
+			}
+			else
+			{
+				UE_LOG(LogVulkanRHI, Log, TEXT("FVulkanChunkedPipelineCacheManager: Cache file could not open %s for reading"), *FullCachePath);
+			}
+		}
+		else
+		{
+			UE_LOG(LogVulkanRHI, Log, TEXT("FVulkanChunkedPipelineCacheManager: deleted oversized cache (%d, %d)"), Stats.FileSize , (UE::Vulkan::GMaxTotalCacheSizeMb * 1024 * 1024));
+		}
+
+		// prepare the output handle.
+		bool bWriteSuccess = OpenWriteHandle(bReadSuccess);
+		check(bWriteSuccess);
+		if (!bReadSuccess && bWriteSuccess)
+		{
+			// fill in the new header and flush to ensure there's actually something for 'OpenMapped' to use.
+			WriteFileHeader(*PSOFileWriter, 0);
+			PSOFileWriter->Flush();
+		}
+
+		if (UE::Vulkan::CanMemoryMapChunkedPSOCache() && !MappedCacheFile.IsValid())
+		{
+			MappedCacheFile = TUniquePtr<IMappedFileHandle>(FPlatformFileManager::Get().GetPlatformFile().OpenMapped(*FullCachePath));
+		}
+
+		return bWriteSuccess;
+	}
+
+	FPSOArchiveReader GetReader(uint32 Offset)
+	{
+		if (UE::Vulkan::CanMemoryMapChunkedPSOCache())
+		{
+			check(MappedRegion);
+			TUniquePtr<FRWScopeLock> MappingScopeLock(new FRWScopeLock(FVulkanCombinedChunkCacheFile::Get().MappingLock, SLT_ReadOnly));
+			FMemoryView MemView(MappedRegion->GetMappedPtr() + Offset, MappedRegion->GetMappedSize() - Offset);
+			return FPSOArchiveReader(TUniquePtr<FArchive>(new FMemoryReaderView(MemView)), (const uint8*)MemView.GetData(), MoveTemp(MappingScopeLock));
+		}
+		else
+		{
+			FPSOArchiveReader Reader(TUniquePtr<FArchive>(IFileManager::Get().CreateFileReader(*GetFullCachePath(), FILEREAD_AllowWrite)));
+			if (Reader.GetArchive())
+			{
+				Reader.GetArchive()->Seek(Offset);
+			}
+			return Reader;
+		}
+	}
+
+	bool OpenWriteHandle(bool bAppend)
+	{
+		check(PSOFileWriter == nullptr);
+		const FString& FullCachePath = GetFullCachePath();
+		uint32 WriteFlags = EFileWrite::FILEWRITE_AllowRead  | (bAppend ? EFileWrite::FILEWRITE_Append : 0);
+		PSOFileWriter = TUniquePtr<FArchive>(IFileManager::Get().CreateFileWriter(*FullCachePath, WriteFlags ));
+		UE_CLOG(PSOFileWriter, LogRHI, Log, TEXT("Opened binary cache for write (%s)"), *FullCachePath);
+		UE_CLOG(PSOFileWriter == nullptr, LogRHI, Warning, TEXT("Failed to open OGL binary cache output file. (%s)"), *FullCachePath);
+		UE_CLOG(PSOFileWriter && (PSOFileWriter->IsError() || PSOFileWriter->IsCriticalError()), LogRHI, Error, TEXT("OGL binary cache output archive error (%s, %d,%d)"), *FullCachePath, PSOFileWriter->IsError(), PSOFileWriter->IsCriticalError());
+
+		return PSOFileWriter != nullptr;
+	}
+
+	FArchive& GetWriter() 
+	{
+		check(PSOFileWriter.IsValid());
+		return *PSOFileWriter.Get(); 
+	}
+
+	void FlushWriteHandle()
+	{
+		FArchive& Archive = *PSOFileWriter;
+		Archive.Flush();
+
+		uint32 CurrentOffsetPos = Archive.Tell();
+		Archive.Seek(0);
+		WriteFileHeader(Archive, CurrentOffsetPos);
+		Archive.Seek(CurrentOffsetPos);
+		Archive.Flush();
+
+		// Bring the mmap up the new write position.
+		UpdateMapping(CurrentOffsetPos);
+	}
+};
+
+
+using FPSOArchiveReader = FVulkanCombinedChunkCacheFile::FPSOArchiveReader;
+
+// A FVulkanPipelineCacheChunk represents a single instance of the VK PSO cache object.
+// GMaxPSOsPerChunk is used to limit the size of a chunk during creation.
+class FVulkanPipelineCacheChunk
+{
+	using FVulkanPipelineCacheChunkLRUNode = UE::Vulkan::FVulkanPipelineCacheChunkLRUNode;
+
+	FVulkanPipelineCacheChunkLRUNode* LRUNode;
+public:
+
+	explicit FVulkanPipelineCacheChunk(FVulkanPipelineCacheChunkLRUNode* LRUNodeIn = nullptr) : LRUNode(LRUNodeIn) { }
+
+	~FVulkanPipelineCacheChunk()
+	{
+		if (PipelineCacheObj != VK_NULL_HANDLE)
+		{
+			VulkanRHI::vkDestroyPipelineCache(UE::Vulkan::GetVulkanDeviceHandle(), PipelineCacheObj, VULKAN_CPU_ALLOCATOR);
+		}
+	}
+
+	void SetLRUNode(FVulkanPipelineCacheChunkLRUNode* InLRUNode) { LRUNode = InLRUNode; }
+	FVulkanPipelineCacheChunkLRUNode* GetLRUNode() { return LRUNode; }
+
+	uint32 OffsetWithinFile = 0; // Location of the binary data as required by 
+	void SetCacheOffset(uint32 OffsetWithinFileIn) { OffsetWithinFile = OffsetWithinFileIn; }
+	uint32 GetCacheOffset() { return OffsetWithinFile; }
+
+	void InitNewCache()
+	{
+		FRWScopeLock Lock(CacheChunkLock, SLT_Write);
+		VkPipelineCacheCreateInfo PipelineCacheInfo;
+		ZeroVulkanStruct(PipelineCacheInfo, VK_STRUCTURE_TYPE_PIPELINE_CACHE_CREATE_INFO);
+		VERIFYVULKANRESULT(VulkanRHI::vkCreatePipelineCache(UE::Vulkan::GetVulkanDeviceHandle(), &PipelineCacheInfo, VULKAN_CPU_ALLOCATOR, &PipelineCacheObj));
+		CacheState = ECacheState::Building;
+		check(LastSaveSize == 0);
+		check(BinaryCacheFileInfo.Filename.IsEmpty());
+		BinaryCacheFileInfo.Filename = GenerateNewFileName();
+	}
+
+	void Touch()
+	{
+		LastUsedFrame = GFrameNumber;
+
+		{
+			FRWScopeLock Lock(CacheChunkLock, SLT_ReadOnly);
+			if (CacheState != ECacheState::FinalizedEvicted)
+			{
+				return;
+			}
+		}
+		{
+			FRWScopeLock Lock(CacheChunkLock, SLT_Write);
+			if(CacheState == ECacheState::FinalizedEvicted)
+			{
+				ReinstateDriverBlobInternal();
+			}
+		}
+	}
+
+	enum class EPSOCacheAccessType
+	{
+		AddTo,		// Use to add new PSOs 
+		ReadFrom	// PSO is expected to be in the cache.
+	};
+
+	template<class TPipelineState>
+	VkResult CreatePSO(TPipelineState* Initializer, EPSOCacheAccessType PSOCacheAccessMode, TUniqueFunction<VkResult(TPipelineState*, VkPipelineCache)> PSOCreateFunc)
+	{
+		FScopedTimeToLog Timer(FString::Printf(TEXT("FVulkanChunkedPipelineCacheManager: FVulkanPipelineCacheChunk.CreatePSO  tot %d "), FPlatformTLS::GetCurrentThreadId()));
+		Touch();
+		VkResult retcode;
+		uint64 PSOHash;
+		{
+			FRWScopeLock Lock(CacheChunkLock, SLT_ReadOnly);
+			FScopedTimeToLog Timer2(FString::Printf(TEXT("FVulkanChunkedPipelineCacheManager: FVulkanPipelineCacheChunk.CreatePSO(lock  tot %d "), FPlatformTLS::GetCurrentThreadId()));
+			PSOHash = GetPrecacheKey(Initializer);
+			if (PSOCacheAccessMode == EPSOCacheAccessType::ReadFrom)
+			{
+				// it's possible to still add a new PSO if we have hash collisions or imperfect PSO hash calc.
+				// the size will be inaccurate.
+				QUICK_SCOPE_CYCLE_COUNTER(STAT_Vulkan_CacheManagerCreatePSO_WARM);
+				retcode = PSOCreateFunc(Initializer, PipelineCacheObj);
+			}
+			else
+			{
+				check(CacheState == ECacheState::Building || CacheState == ECacheState::Closing);
+				// Even though we know we're modifying the cache, we're not taking the write lock.
+				// Holding the write lock for the duration of the create is too costly.
+				// however, other threads can add to the cache during this time.
+				// The size is currently only used to determine if a new cache should be started.
+				// we do not need to ensure an accurate or coherent size the PSO.
+				{
+					QUICK_SCOPE_CYCLE_COUNTER(STAT_Vulkan_CacheManagerCreatePSO_COLD);
+					retcode = PSOCreateFunc(Initializer, PipelineCacheObj);
+				}
+			}
+		}
+
+		if (PSOCacheAccessMode == EPSOCacheAccessType::AddTo)
+		{
+			FScopedTimeToLog Timer3(FString::Printf(TEXT("FVulkanChunkedPipelineCacheManager: FVulkanPipelineCacheChunk.CreatePSO AddTo tot %d "), FPlatformTLS::GetCurrentThreadId()));
+			FRWScopeLock Lock(CacheChunkLock, SLT_Write);
+			size_t EndSize = 0;
+			{
+				FScopedTimeToLog Timer4(FString::Printf(TEXT("FVulkanChunkedPipelineCacheManager: FVulkanPipelineCacheChunk.CreatePSO AddTo getdat tot %d "), FPlatformTLS::GetCurrentThreadId()));
+				VulkanRHI::vkGetPipelineCacheData(UE::Vulkan::GetVulkanDeviceHandle(), PipelineCacheObj, &EndSize, nullptr);
+			}
+
+			PSOsToBeFlushed.Add(PSOHash);
+			TotalNumPSOs++;
+			CacheSize = EndSize;
+
+			check(PendingAddToCompiles.load() != 0);
+			--PendingAddToCompiles;
+			if (CacheState == ECacheState::Closing && PendingAddToCompiles.load() == 0)
+			{
+#if LOGCACHEINFO
+				UE_LOG(LogVulkanRHI, Log, TEXT("FVulkanChunkedPipelineCacheManager: create pso - Chunk Capacity reached %s, Finalizing as No pending jobs remain.."), *BinaryCacheFileInfo.Filename);
+#endif
+				MarkPendingFlushInternal();
+			}
+		}
+		return retcode;
+	}
+	
+	void ReservePendingPSO()
+	{
+		check(CacheState == ECacheState::Building);
+		FRWScopeLock Lock(CacheChunkLock, SLT_ReadOnly);
+		++PendingAddToCompiles;
+		if ((TotalNumPSOs + PendingAddToCompiles) >= UE::Vulkan::GMaxPSOsPerChunk)
+		{
+			Lock.ReleaseReadOnlyLockAndAcquireWriteLock_USE_WITH_CAUTION();
+			if ((TotalNumPSOs + PendingAddToCompiles) >= UE::Vulkan::GMaxPSOsPerChunk)
+			{
+				// become closed..
+#if LOGCACHEINFO
+				UE_LOG(LogVulkanRHI, Log, TEXT("FVulkanChunkedPipelineCacheManager: Chunk Capacity reached , cachestate %d, CacheSize %d, TotalNumPSOs %d, PendingAddToCompiles %d, %s, pending finalize.."), CacheState, CacheSize, TotalNumPSOs, PendingAddToCompiles.load(), *BinaryCacheFileInfo.Filename);
+#endif
+				CacheState = ECacheState::Closing;
+			}
+		}
+	}
+
+	void LogStats(FString&& LogInfo)
+	{
+		FRWScopeLock Lock(CacheChunkLock, SLT_ReadOnly);
+		UE_LOG(LogVulkanRHI, Log, TEXT("FVulkanChunkedPipelineCacheManager: %s Cache name %s num PSOs %d, cache size %d, last used %d, resident %d, state %d"), *LogInfo, *BinaryCacheFileInfo.Filename, TotalNumPSOs, CacheSize, LastUsedFrame.load(), CacheState != ECacheState::FinalizedEvicted, (int)CacheState);
+	}
+
+	bool CanBeEvicted() const
+	{
+		FRWScopeLock Lock(CacheChunkLock, SLT_ReadOnly);
+		return CacheState == ECacheState::Finalized;
+	}
+
+	bool CheckCapacityReached()
+	{
+		FRWScopeLock Lock(CacheChunkLock, SLT_ReadOnly);
+		return CacheState != ECacheState::Building;
+	}
+
+	uint32 GetResidentSize() const
+	{
+		FRWScopeLock Lock(CacheChunkLock, SLT_ReadOnly);
+		uint32 Size = CacheState == ECacheState::FinalizedEvicted ? 0 : CacheSize;
+		return Size;
+	}
+
+	uint64 GetLastUsedFrame() const
+	{
+		return LastUsedFrame.load();
+	}
+
+	void Unload()
+	{
+		FRWScopeLock Lock(CacheChunkLock, SLT_Write);
+		check(CacheState == ECacheState::Finalized);
+#if LOGCACHEINFO
+		UE_LOG(LogVulkanRHI, Log, TEXT("FVulkanChunkedPipelineCacheManager: Evicting cache %s, %d bytes"), *GetCombinedFilePath(), CacheSize);
+#endif
+		// remove vulkan object from mem
+		VulkanRHI::vkDestroyPipelineCache(UE::Vulkan::GetVulkanDeviceHandle(), PipelineCacheObj, VULKAN_CPU_ALLOCATOR);
+		PipelineCacheObj = VK_NULL_HANDLE;
+
+		CacheState = ECacheState::FinalizedEvicted;
+		UE::Vulkan::TotalResidentCacheSize -= CacheSize;
+	}
+
+	enum class ECacheChunkLoadType { LoadAsEvicted, LoadAllData };
+	void Load(FPSOArchiveReader& ArchiveReader, TArray<uint64>* PSOsFoundOUT, ECacheChunkLoadType LoadType)
+	{
+		FRWScopeLock Lock(CacheChunkLock, SLT_Write);
+		FArchive& Archive = *ArchiveReader.GetArchive();
+		check(PSOsFoundOUT->IsEmpty());
+		check(CacheState == ECacheState::Initialized);
+		SetCacheOffset(Archive.Tell());
+
+		LoadInternal(ArchiveReader, PSOsFoundOUT, LoadType);
+		// if we're loading as evicted then we're in a completed (RO) state.
+		CacheState = LoadType == ECacheChunkLoadType::LoadAsEvicted ? ECacheState::FinalizedEvicted : ECacheState::Building;
+	}
+
+	private:
+		mutable FRWLock CacheChunkLock;
+
+		VkPipelineCache PipelineCacheObj = VK_NULL_HANDLE;
+
+		TArray<uint64> PSOsToBeFlushed;
+		int32 PSOsNotFlushed = 0;
+		int32 CacheSize = 0;
+		int32 LastSaveSize = 0;
+		int32 TotalNumPSOs = 0;
+
+		static inline const TCHAR TempFileSuffix[] = TEXT("write");
+
+		struct FFileInfo
+		{
+			FString Filename;
+			uint32 RawDataOffset = 0xffffffff;
+		} BinaryCacheFileInfo;
+
+		std::atomic<uint64> LastUsedFrame = 0;
+
+		std::atomic<uint64> PendingAddToCompiles = 0;
+
+		FString GetCombinedFilePath() const
+		{
+			const FString RootCacheFolder = UE::Vulkan::GetRHICacheRootFolder();
+			const FString CacheSubDir = UE::Vulkan::GetPSOBinaryCacheVersionedFolder();
+			const FString CombinedCacheSubDir = FPaths::Combine(RootCacheFolder, CacheSubDir);
+			return FPaths::Combine(CombinedCacheSubDir, BinaryCacheFileInfo.Filename);
+		}
+
+		static FString GenerateNewFileName()
+		{
+			IPlatformFile& PlatformFile = FPlatformFileManager::Get().GetPlatformFile();
+			return FString(TEXT("VulkanPSO_")) + FGuid::NewGuid().ToString();
+		}
+
+		void ReinstateDriverBlobInternal()
+		{
+			FScopedTimeToLog Timer(FString::Printf(TEXT("FVulkanChunkedPipelineCacheManager: ReinstateDriverBlobInternal tot %d "), FPlatformTLS::GetCurrentThreadId()));
+
+			check(CacheState == ECacheState::FinalizedEvicted);		
+
+			FPSOArchiveReader ArchiveReader = FVulkanCombinedChunkCacheFile::Get().GetReader(GetCacheOffset());
+
+			if(ensure(ArchiveReader.IsValid()))
+			{
+				LoadInternal(ArchiveReader, nullptr, ECacheChunkLoadType::LoadAllData);
+				CacheState = ECacheState::Finalized;
+			}
+		}
+
+		void CreateVKCacheInternal(TConstArrayView<uint8> CacheBytes)
+		{
+			FScopedTimeToLog Timer(FString::Printf(TEXT("FVulkanChunkedPipelineCacheManager: CreateVKCacheInternal %s "), *BinaryCacheFileInfo.Filename));
+			check(CacheState == ECacheState::FinalizedEvicted || CacheState == ECacheState::Initialized);
+			QUICK_SCOPE_CYCLE_COUNTER(STAT_Vulkan_CacheManagerLoadCache);
+			VkPipelineCacheCreateInfo PipelineCacheInfo;
+			ZeroVulkanStruct(PipelineCacheInfo, VK_STRUCTURE_TYPE_PIPELINE_CACHE_CREATE_INFO);
+			PipelineCacheInfo.pInitialData = CacheBytes.GetData();
+			PipelineCacheInfo.initialDataSize = CacheBytes.Num();
+			VERIFYVULKANRESULT(VulkanRHI::vkCreatePipelineCache(UE::Vulkan::GetVulkanDeviceHandle(), &PipelineCacheInfo, VULKAN_CPU_ALLOCATOR, &PipelineCacheObj));
+		}
+
+		uint64 GetPrecacheKey(const FVulkanRHIGraphicsPipelineState* GraphicsPipelineState)
+		{
+			return GraphicsPipelineState->PrecacheKey;
+		}
+
+		void MarkPendingFlushInternal()
+		{
+			CacheState = ECacheState::PendingFlush;
+
+			FArchive& Archive = FVulkanCombinedChunkCacheFile::Get().GetWriter();
+			SavePSOCacheInternal(Archive);
+			FVulkanCombinedChunkCacheFile::Get().FlushWriteHandle();
+		}
+		enum class ECacheState
+		{
+			Unknown,
+			Initialized,
+			Building,			// still accumulating PSOs
+			Closing,			// cache reached the size limit waiting for remaining PSOs to complete.
+			PendingFlush,		// capacity reached and all pending PSOs are stored. Cache has not been flushed to storage.
+			Finalized,			// a finished and resident cache. Finalized caches are always backed by storage.
+			FinalizedEvicted,	// a finished cache that is not resident in RAM. 
+		};
+		ECacheState CacheState = ECacheState::Initialized;
+
+		void LoadInternal(FPSOArchiveReader& ArchiveReader, TArray<uint64>* PSOsInFileOUT, ECacheChunkLoadType LoadType)
+		{
+			FScopedTimeToLog Timer1(FString::Printf(TEXT("FVulkanChunkedPipelineCacheManager: LoadInternal TOT  ")));
+ 			FArchive& Archive = *ArchiveReader.GetArchive();
+
+			Archive << CacheSize;
+			LastSaveSize = CacheSize;
+			BinaryCacheFileInfo.RawDataOffset = Archive.Tell();
+			if (LoadType == ECacheChunkLoadType::LoadAllData)
+			{
+				if (ArchiveReader.GetData())
+				{
+					TConstArrayView<uint8> CacheBytes;
+					{
+						FScopedTimeToLog Timer3(FString::Printf(TEXT("FVulkanChunkedPipelineCacheManager: LoadInternal MappedSerialize %d"), CacheSize));
+						// this avoids a pointless memcpy when mmapping is in use.
+						const uint8* RawData = ArchiveReader.GetData() + Archive.Tell();
+						CacheBytes = MakeArrayView(RawData, CacheSize);
+						Archive.Seek(Archive.Tell() + CacheSize);
+					}
+
+					CreateVKCacheInternal(CacheBytes);
+				}
+				else
+				{
+					TArray<uint8> CacheBytes;
+					{
+						FScopedTimeToLog Timer3(FString::Printf(TEXT("FVulkanChunkedPipelineCacheManager: LoadInternal Serialize %d"), CacheSize));
+						CacheBytes.SetNumUninitialized(CacheSize);
+						Archive.Serialize(CacheBytes.GetData(), CacheSize);;
+						check(!Archive.IsError());
+					}
+					CreateVKCacheInternal(CacheBytes);
+				}
+				UE::Vulkan::TotalResidentCacheSize += CacheSize;
+			}
+			else
+			{
+				Archive.Seek(Archive.Tell() + CacheSize);
+			}
+
+			if (PSOsInFileOUT)
+			{
+				Archive << (*PSOsInFileOUT);
+				TotalNumPSOs = PSOsInFileOUT->Num();
+			}
+		}
+
+		// Save this cache chunk to the archive.
+		void SavePSOCacheInternal(FArchive& Archive)
+		{
+			FScopedTimeToLog Timer(FString::Printf(TEXT("FVulkanChunkedPipelineCacheManager: SavePSOCacheInternal TOT %s "), *BinaryCacheFileInfo.Filename));
+
+			check(CacheState == ECacheState::PendingFlush || CacheState == ECacheState::Building);
+
+			QUICK_SCOPE_CYCLE_COUNTER(STAT_Vulkan_CacheManagerWriteCache);
+
+			TArray<uint8> CacheBytes;
+			size_t TotalSize = 0;
+			{
+				FScopedTimeToLog Timer2(FString::Printf(TEXT("FVulkanChunkedPipelineCacheManager: SavePSOCacheInternal get vk data %s "), *BinaryCacheFileInfo.Filename));
+				VulkanRHI::vkGetPipelineCacheData(UE::Vulkan::GetVulkanDeviceHandle(), PipelineCacheObj, &TotalSize, nullptr);
+				CacheBytes.SetNumUninitialized(TotalSize);
+				VulkanRHI::vkGetPipelineCacheData(UE::Vulkan::GetVulkanDeviceHandle(), PipelineCacheObj, &TotalSize, CacheBytes.GetData());
+			}
+			CacheSize = (uint32)TotalSize;
+
+			{
+				FScopedTimeToLog Timer4(FString::Printf(TEXT("FVulkanChunkedPipelineCacheManager: SavePSOCacheInternal serialize %s %d "), *BinaryCacheFileInfo.Filename, TotalSize));
+
+				SetCacheOffset((uint32)Archive.Tell());
+
+				uint32 RawDataSize = CacheBytes.Num();
+				Archive << RawDataSize;
+				BinaryCacheFileInfo.RawDataOffset = Archive.Tell();
+				Archive.Serialize(CacheBytes.GetData(), RawDataSize);
+
+				Archive << PSOsToBeFlushed; // TODO: !Discuss! - if the PSO precaching algo changes the whole file will be incorrect, we've no way to check if the hash changes... discuss this in review.
+				check(!Archive.IsError());
+			}
+
+			if (CacheState == ECacheState::PendingFlush)
+			{
+				CacheState = ECacheState::Finalized;
+				UE::Vulkan::TotalResidentCacheSize += CacheSize;
+				PSOsToBeFlushed.Empty();
+			}
+		}
+};
+
+//////////////////////////////////////////////////////////////////////////
+
+class FVulkanChunkedPipelineCacheManagerImpl
+{
+	mutable FRWLock ChunkedPipelineCacheLock; // This guards access to the manager but not individual chunks within the map.
+
+	using FCacheChunkKey = UE::Vulkan::FPipelineCacheChunkKey;
+	using FCacheChunksMap = TMap<FCacheChunkKey, TUniquePtr<FVulkanPipelineCacheChunk>>;
+
+	uint32 NumChunks = 0;
+	FCacheChunksMap CacheChunksMap;
+
+	// track number of PSOs created/found. logging use only.
+	std::atomic<int32> FoundPSOs = 0;
+
+	TMap<uint64, FCacheChunkKey> PrecachePSOToCacheChunkMap;
+
+	FCriticalSection LRUCS;
+	using FVulkanPipelineCacheChunkLRUNode = UE::Vulkan::FVulkanPipelineCacheChunkLRUNode;
+	using FVulkanRHIGraphicsPipelineStateLRU = UE::Vulkan::FVulkanRHIGraphicsPipelineStateLRU;
+
+	UE::Vulkan::FVulkanRHIGraphicsPipelineStateLRU CacheChunkLRU;
+
+public:
+	static FVulkanChunkedPipelineCacheManagerImpl& Get()
+	{
+		static FVulkanChunkedPipelineCacheManagerImpl VulkanPipelineCacheManager;
+		return VulkanPipelineCacheManager;
+	}
+
+	FVulkanChunkedPipelineCacheManagerImpl()
+	{
+		LoadAllCaches();
+	}
+
+	template<class TPipelineState>
+	VkResult CreatePSO(TPipelineState* GraphicsPipelineState, bool bIsPrecompileJob, TUniqueFunction<VkResult(TPipelineState*, VkPipelineCache)> PSOCreateFunc)
+	{
+		FScopedTimeToLog Timer(FString::Printf(TEXT("FVulkanChunkedPipelineCacheManager: CreatePSO (precomp %d) tot %d "), bIsPrecompileJob, FPlatformTLS::GetCurrentThreadId()));
+
+		FCacheChunkKey ChunkKey;
+		bool bWasFoundInCache;
+		FVulkanPipelineCacheChunk* Chunk = GetOrAddCache(GraphicsPipelineState, bWasFoundInCache, ChunkKey);
+
+		// dont need to lock ChunkedPipelineCacheLock, we never remove an Chunk once it's added. There should be no PSO create tasks during cache shutdown.
+		// Do not create cached precompile PSOs
+		if (!bIsPrecompileJob || !bWasFoundInCache)
+		{
+			FScopedTimeToLog Timer4(FString::Printf(TEXT("FVulkanChunkedPipelineCacheManager: CreatePSO actualcreate %d "), FPlatformTLS::GetCurrentThreadId()));
+			uint64 Before = Chunk->GetLastUsedFrame();
+			using EPSOCacheAccessType = FVulkanPipelineCacheChunk::EPSOCacheAccessType;
+			VkResult Result = Chunk->CreatePSO(GraphicsPipelineState, bWasFoundInCache ? EPSOCacheAccessType::ReadFrom : EPSOCacheAccessType::AddTo, MoveTemp(PSOCreateFunc));
+			uint64 After = Chunk->GetLastUsedFrame();
+
+			if (Before != After)
+			{
+				FScopedTimeToLog Timer3(FString::Printf(TEXT("FVulkanChunkedPipelineCacheManager: CreatePSO lru %d "), FPlatformTLS::GetCurrentThreadId()));
+				FScopeLock Lock(&LRUCS);
+				FVulkanPipelineCacheChunkLRUNode* LRUNode = Chunk->GetLRUNode();
+				if (LRUNode->GetNextNode() || LRUNode->GetPrevNode())
+				{
+					// if evicted it may not be in the lru yet..
+					CacheChunkLRU.RemoveNode(LRUNode, false);
+				}
+				CacheChunkLRU.AddHead(LRUNode);
+			}
+			return Result;
+		}
+
+		return VK_SUCCESS;
+	}
+
+	void Tick()
+	{
+#if LOGCACHEINFO
+		{ // Logging
+			FRWScopeLock Lock(ChunkedPipelineCacheLock, SLT_ReadOnly);
+			static int32 CachedChunkCount = 0;
+			bool bLogMe = CachedChunkCount != PrecachePSOToCacheChunkMap.Num();
+			if (bLogMe)
+			{
+				UE_LOG(LogVulkanRHI, Log, TEXT("FVulkanChunkedPipelineCacheManager: Total precache PSOs stored %d, Num Chunks %d"), PrecachePSOToCacheChunkMap.Num(), CacheChunksMap.Num());
+				for (auto& ChunkPair : CacheChunksMap)
+				{
+					ChunkPair.Value.Get()->LogStats(FString::Printf(TEXT("Chunk id %d"), ChunkPair.Key));
+				}
+				UE_LOG(LogVulkanRHI, Log, TEXT("FVulkanChunkedPipelineCacheManager: PSOs created from cache %d "), FoundPSOs.load());
+
+				CachedChunkCount = PrecachePSOToCacheChunkMap.Num();
+			}
+		}
+#endif
+		if (UE::Vulkan::GTargetResidentCacheSizeMb)
+		{
+			TryUnloadCacheChunks();
+		}
+	}
+
+private:
+
+	bool LoadAllCaches()
+	{
+		FRWScopeLock WriteLock(ChunkedPipelineCacheLock, SLT_Write);
+
+		check(CacheChunksMap.IsEmpty() && PrecachePSOToCacheChunkMap.IsEmpty());
+
+		// the CacheSubDir is specific to the device+driver+shaderplatform.
+		// We clean out everything else from the cache folder.
+		const FString RootCacheFolder = UE::Vulkan::GetRHICacheRootFolder();
+		const FString CacheSubDir = UE::Vulkan::GetPSOBinaryCacheVersionedFolder();
+		const FString CombinedCacheSubDir = FPaths::Combine(RootCacheFolder, CacheSubDir);
+
+		// delete anything unexpected from the RHI cache root folder.
+		IPlatformFile& PlatformFile = FPlatformFileManager::Get().GetPlatformFile();
+		TArray<FString> FoundFiles;
+		IFileManager::Get().FindFiles(FoundFiles, *(RootCacheFolder / TEXT("*")), true, true);
+		for (FString& FoundFile : FoundFiles)
+		{
+			const FString FullPath = (RootCacheFolder / FoundFile);
+			const bool bIsDir = PlatformFile.DirectoryExists(*FullPath);
+			if (FoundFile != CacheSubDir || !bIsDir)
+			{
+				bool bSuccess;
+				if (bIsDir)
+				{
+					bSuccess = PlatformFile.DeleteDirectoryRecursively(*FullPath);
+				}
+				else
+				{
+					bSuccess = PlatformFile.DeleteFile(*FullPath);
+				}
+				UE_LOG(LogRHI, Verbose, TEXT("FVulkanChunkedPipelineCacheManagerImpl: Deleting %s %s"), bIsDir ? TEXT("dir") : TEXT("file"), *FullPath);
+				UE_CLOG(!bSuccess, LogRHI, Warning, TEXT("FVulkanChunkedPipelineCacheManagerImpl: Failed to delete %s"), *FullPath);
+			}
+		}
+
+		FoundFiles.Reset();
+		IFileManager::Get().FindFiles(FoundFiles, *(CombinedCacheSubDir / TEXT("*")), true, false);
+
+		for (FString& FoundFile : FoundFiles)
+		{
+			if (!FoundFile.Equals(FVulkanCombinedChunkCacheFile::Get().GetFilename()))
+			{
+				FString FullPath = FPaths::Combine(CombinedCacheSubDir, FoundFile);
+				UE_LOG(LogVulkanRHI, Log, TEXT("FVulkanChunkedPipelineCacheManager: deleting unexpected Cache %s."), *FullPath);
+				PlatformFile.DeleteFile(*FullPath);
+			}
+		}
+
+		bool bReadFailed = false;
+		TUniqueFunction<void(FPSOArchiveReader& ArchiveReader)> OnFoundCacheChunk = [this, &bReadFailed](FPSOArchiveReader& ArchiveReader)
+		{
+			TUniquePtr<FVulkanPipelineCacheChunk> PendingChunk = MakeUnique<FVulkanPipelineCacheChunk>();
+			TArray<uint64> PSOsOut;
+			PendingChunk->Load(ArchiveReader, &PSOsOut, FVulkanPipelineCacheChunk::ECacheChunkLoadType::LoadAsEvicted);
+
+			FCacheChunkKey NewChunkKey = ++NumChunks;
+			check(!CacheChunksMap.Contains(NewChunkKey));
+
+			FVulkanPipelineCacheChunkLRUNode* NewLRUNode = new FVulkanPipelineCacheChunkLRUNode(NewChunkKey);
+			{
+				FScopeLock Lock(&LRUCS);
+				CacheChunkLRU.AddHead(NewLRUNode);
+				PendingChunk->SetLRUNode(NewLRUNode);
+			}
+
+			CacheChunksMap.Emplace(NewChunkKey, MoveTemp(PendingChunk));
+			for (uint64 PSOHash : PSOsOut)
+			{
+				PrecachePSOToCacheChunkMap.Add(PSOHash, NewChunkKey);
+			}
+		};
+
+		bool bSuccess = FVulkanCombinedChunkCacheFile::Get().LoadAllCacheChunks(MoveTemp(OnFoundCacheChunk));
+		if (!bSuccess)
+		{
+			UE_LOG(LogVulkanRHI, Log, TEXT("FVulkanChunkedPipelineCacheManager: failed to open chunked PSO cache file, PSO caching is disabled."));
+			UE::Vulkan::GUseChunkedPSOCache = 0;
+		}
+		return bSuccess;
+	}
+
+	static uint64 GetPrecacheHash(const FVulkanRHIGraphicsPipelineState* GFXState) { return GFXState->PrecacheKey; }
+
+	template<class TInitializer>
+	FVulkanPipelineCacheChunk* GetChunk(const TInitializer& Initializer, bool& bWasFound, FCacheChunkKey& ChunkKeyOUT, FRWScopeLock& Lock, const bool bTryAdd = false)
+	{
+		FVulkanPipelineCacheChunk* ReturnChunk = nullptr;
+		uint64 PSOPrecacheKey = GetPrecacheHash(Initializer);
+
+		const FCacheChunkKey* FoundChunkKey = PrecachePSOToCacheChunkMap.Find(PSOPrecacheKey);
+		if (FoundChunkKey)
+		{
+			// we have a cache for this PSO
+			ReturnChunk = CacheChunksMap.FindChecked(*FoundChunkKey).Get();
+			FoundPSOs++;
+			bWasFound = true;
+			ChunkKeyOUT = *FoundChunkKey;
+		}
+		else if (!bTryAdd)
+		{
+			bWasFound = false;
+			// Try again with the write lock, add if it's still missing.
+			Lock.ReleaseReadOnlyLockAndAcquireWriteLock_USE_WITH_CAUTION();
+			return GetChunk(Initializer, bWasFound, ChunkKeyOUT, Lock, true);
+		}
+
+		if (!ReturnChunk)
+		{
+			check(bTryAdd);
+
+			auto FindChunk = [this]
+			{
+				FVulkanPipelineCacheChunk* ReturnChunk = nullptr;
+
+				TUniquePtr<FVulkanPipelineCacheChunk>* Found = CacheChunksMap.Find(NumChunks);
+				if (Found)
+				{
+					ReturnChunk = (*Found).Get();
+				}
+				if ((!ReturnChunk) || (ReturnChunk && ReturnChunk->CheckCapacityReached()))
+				{
+					// create and return a new cache chunk.
+					FCacheChunkKey NewChunk = ++NumChunks;
+					check(!CacheChunksMap.Contains(NewChunk));
+					FVulkanPipelineCacheChunkLRUNode* NewLRUNode = new FVulkanPipelineCacheChunkLRUNode(NewChunk);
+					{
+						FScopeLock Lock(&LRUCS);
+						CacheChunkLRU.AddHead(NewLRUNode);
+						ReturnChunk = CacheChunksMap.Emplace(NewChunk, MakeUnique<FVulkanPipelineCacheChunk>(NewLRUNode)).Get();
+					}
+					ReturnChunk->InitNewCache();
+				}
+
+				ReturnChunk->ReservePendingPSO();
+				return ReturnChunk;
+			};
+
+			uint32 CurrentChunk = NumChunks;
+			ReturnChunk = FindChunk();
+			ChunkKeyOUT = NumChunks;
+			PrecachePSOToCacheChunkMap.Add(PSOPrecacheKey, NumChunks);
+		}
+
+		check(ReturnChunk);
+		return ReturnChunk;
+	}
+
+	template<class TInitializer>
+	FVulkanPipelineCacheChunk* GetOrAddCache(const TInitializer& Initializer, bool& bWasFoundInCache, FCacheChunkKey& ChunkKeyOUT)
+	{
+		FScopedTimeToLog Timer(FString::Printf(TEXT("FVulkanChunkedPipelineCacheManager: GetOrAddCache tot %d "), FPlatformTLS::GetCurrentThreadId()));
+		FRWScopeLock Lock(ChunkedPipelineCacheLock, SLT_ReadOnly);
+		return GetChunk(Initializer, bWasFoundInCache, ChunkKeyOUT, Lock);
+	}
+
+	void TryUnloadCacheChunks()
+	{
+		const uint32 FramePace = FPlatformRHIFramePacer::GetFramePace();
+		const uint32 LastFrameRequired = GFrameNumber - FMath::Min(GFrameNumber, (uint32)(FramePace * UE::Vulkan::GUntouchedChunkEvictTimeSeconds));
+
+		FRWScopeLock CacheLock(ChunkedPipelineCacheLock, SLT_ReadOnly);
+		FScopeLock LRULock(&LRUCS);
+
+		int32 CurrentResidentSize = UE::Vulkan::TotalResidentCacheSize;
+
+		const int32 TargetResidentCacheSizeBytes = (UE::Vulkan::GTargetResidentCacheSizeMb * 1024 * 1024);
+		const int32 MaxToUnloadPerTick = 3;
+
+		// unload oldest
+		FVulkanPipelineCacheChunkLRUNode* CurrentNode = CacheChunkLRU.GetTail();
+		for (int32 UnloadCount = 0; CurrentNode && UnloadCount < MaxToUnloadPerTick && CurrentResidentSize > TargetResidentCacheSizeBytes; UnloadCount++)
+		{
+			FVulkanPipelineCacheChunkLRUNode* NextNode = CurrentNode->GetNextNode();
+
+			UE::Vulkan::FPipelineCacheChunkKey ChunkKey = CurrentNode->GetValue();
+			FVulkanPipelineCacheChunk* FoundCacheChunk = CacheChunksMap.FindChecked(ChunkKey).Get();
+
+			int CacheSize = FoundCacheChunk->GetResidentSize();
+			if (FoundCacheChunk->CanBeEvicted())
+			{
+				if (LastFrameRequired < FoundCacheChunk->GetLastUsedFrame())
+				{
+					// exit, everything else will be too recent
+					break;
+				}
+
+				FoundCacheChunk->Unload();
+				CacheChunkLRU.RemoveNode(CurrentNode, false);
+				CurrentResidentSize -= CacheSize;
+				UnloadCount++;
+			}
+			CurrentNode = NextNode;
+		}
+	}
+};
+
+//////////////////////////////////////////////////////////////////////////
+// public interface:
+static FVulkanChunkedPipelineCacheManager VulkanPipelineCacheManager;
+
+bool FVulkanChunkedPipelineCacheManager::IsEnabled()
+{
+	return UE::Vulkan::GUseChunkedPSOCache != 0;
+}
+
+void FVulkanChunkedPipelineCacheManager::Init()
+{
+	if (UE::Vulkan::GUseChunkedPSOCache == 0)
+	{
+		return;
+	}
+
+	check(!VulkanPipelineCacheManager.VulkanPipelineCacheManagerImpl.IsValid());
+	VulkanPipelineCacheManager.VulkanPipelineCacheManagerImpl = MakeUnique<FVulkanChunkedPipelineCacheManagerImpl>();
+}
+
+void FVulkanChunkedPipelineCacheManager::Shutdown()
+{
+	if (UE::Vulkan::GUseChunkedPSOCache == 0)
+	{
+		return;
+	}
+
+	if (VulkanPipelineCacheManager.VulkanPipelineCacheManagerImpl.IsValid())
+	{
+		VulkanPipelineCacheManager.VulkanPipelineCacheManagerImpl = nullptr;
+	}
+}
+
+FVulkanChunkedPipelineCacheManager& FVulkanChunkedPipelineCacheManager::Get()
+{
+	check(UE::Vulkan::GUseChunkedPSOCache);
+	check(VulkanPipelineCacheManager.VulkanPipelineCacheManagerImpl.IsValid());
+	return VulkanPipelineCacheManager;
+}
+
+
+template<class TPipelineState>
+VkResult FVulkanChunkedPipelineCacheManager::CreatePSO(TPipelineState* GraphicsPipelineState, bool bIsPrecompileJob, TUniqueFunction<VkResult(TPipelineState*, VkPipelineCache)> PSOCreateFunc)
+{
+	check(UE::Vulkan::GUseChunkedPSOCache);
+	return VulkanPipelineCacheManagerImpl->CreatePSO(GraphicsPipelineState, bIsPrecompileJob, MoveTemp(PSOCreateFunc));
+}
+
+void FVulkanChunkedPipelineCacheManager::Tick()
+{
+	if (UE::Vulkan::GUseChunkedPSOCache == 0)
+	{
+		return;
+	}
+
+	VulkanPipelineCacheManagerImpl->Tick();
+}
+
+template VkResult FVulkanChunkedPipelineCacheManager::CreatePSO(FVulkanRHIGraphicsPipelineState* GraphicsPipelineState, bool bIsPrecompileJob, TUniqueFunction<VkResult(FVulkanRHIGraphicsPipelineState*, VkPipelineCache)> PSOCreateFunc);
+
diff --git a/Engine/Source/Runtime/VulkanRHI/Private/VulkanChunkedPipelineCache.h b/Engine/Source/Runtime/VulkanRHI/Private/VulkanChunkedPipelineCache.h
new file mode 100644
--- /dev/null
+++ b/Engine/Source/Runtime/VulkanRHI/Private/VulkanChunkedPipelineCache.h
@@ -0,0 +1,28 @@
+// Copyright Epic Games, Inc. All Rights Reserved..
+
+/*=============================================================================
+	VulkanPipelineCache.h: Private Vulkan RHI definitions.
+=============================================================================*/
+
+#pragma once
+#include "VulkanRHIPrivate.h"
+#include "VulkanResources.h"
+#include "VulkanShaderResources.h"
+#include "VulkanDescriptorSets.h"
+#include "ShaderPipelineCache.h"
+#include "Templates/UniquePtr.h"
+
+class FVulkanChunkedPipelineCacheManager
+{
+	TUniquePtr<class FVulkanChunkedPipelineCacheManagerImpl> VulkanPipelineCacheManagerImpl;
+	friend class FVulkanPipelineCacheEntry;
+public:
+	static void Init();
+	static void Shutdown();
+	static bool IsEnabled();
+	static FVulkanChunkedPipelineCacheManager& Get();
+
+	template<class TPipelineState>
+	VkResult CreatePSO(TPipelineState* GraphicsPipelineState, bool bIsPrecompileJob, TUniqueFunction<VkResult(TPipelineState*, VkPipelineCache)> PSOCreateFunc);
+	void Tick();
+};
diff --git a/Engine/Source/Runtime/VulkanRHI/Private/VulkanDevice.cpp b/Engine/Source/Runtime/VulkanRHI/Private/VulkanDevice.cpp
--- a/Engine/Source/Runtime/VulkanRHI/Private/VulkanDevice.cpp
+++ b/Engine/Source/Runtime/VulkanRHI/Private/VulkanDevice.cpp
@@ -18,6 +18,7 @@
 #include "VulkanRenderpass.h"
 #include "VulkanRayTracing.h"
 #include "VulkanDescriptorSets.h"
+#include "VulkanChunkedPipelineCache.h"
 
 TAutoConsoleVariable<int32> GRHIAllowAsyncComputeCvar(
 	TEXT("r.Vulkan.AllowAsyncCompute"),
@@ -1254,6 +1255,8 @@ void FVulkanDevice::InitGPU()
 	}
 #endif
 
+	FVulkanChunkedPipelineCacheManager::Init();
+
 	PipelineStateCache->InitAndLoad(CacheFilenames);
 
 	// Setup default resource
@@ -1398,6 +1401,7 @@ void FVulkanDevice::Destroy()
 
 	FenceManager.Deinit();
 	DeviceMemoryManager.Deinit();
+	FVulkanChunkedPipelineCacheManager::Shutdown();
 
 	VulkanRHI::vkDestroyDevice(Device, VULKAN_CPU_ALLOCATOR);
 	Device = VK_NULL_HANDLE;
diff --git a/Engine/Source/Runtime/VulkanRHI/Private/VulkanPipeline.cpp b/Engine/Source/Runtime/VulkanRHI/Private/VulkanPipeline.cpp
--- a/Engine/Source/Runtime/VulkanRHI/Private/VulkanPipeline.cpp
+++ b/Engine/Source/Runtime/VulkanRHI/Private/VulkanPipeline.cpp
@@ -17,6 +17,7 @@
 #include "GlobalShader.h"
 #include "VulkanLLM.h"
 #include "Misc/ScopeRWLock.h"
+#include "VulkanChunkedPipelineCache.h"
 
 #define LRU_DEBUG 0
 #if !UE_BUILD_SHIPPING
@@ -1399,6 +1400,19 @@ bool FVulkanPipelineStateCacheManager::CreateGfxPipelineFromEntry(FVulkanRHIGrap
 
 VkResult FVulkanPipelineStateCacheManager::CreateVKPipeline(FVulkanRHIGraphicsPipelineState* PSO, FVulkanShader* Shaders[ShaderStage::NumStages], const VkGraphicsPipelineCreateInfo& PipelineInfo, bool bIsPrecompileJob)
 {
+	if(FVulkanChunkedPipelineCacheManager::IsEnabled())
+	{
+		// Use chunk caching and bypass FVulkanPipelineStateCacheManager's PSO caching
+		// Placeholder PSO size - TODO: remove pipeline cache size stuff.	
+		PSO->PipelineCacheSize = 20 * 1024; // This is only required bUseLRU == true.
+		return FVulkanChunkedPipelineCacheManager::Get().CreatePSO(PSO, bIsPrecompileJob, TUniqueFunction< VkResult(FVulkanRHIGraphicsPipelineState*, VkPipelineCache)>(
+			[&](FVulkanRHIGraphicsPipelineState* PSO, VkPipelineCache PipelineCache)
+			{
+				QUICK_SCOPE_CYCLE_COUNTER(STAT_Vulkan_vkCreateGraphicsPipeline);
+				return VulkanRHI::vkCreateGraphicsPipelines(Device->GetInstanceHandle(), PipelineCache, 1, &PipelineInfo, VULKAN_CPU_ALLOCATOR, &PSO->VulkanPipeline);
+			}));
+	}
+
 	VkPipeline* Pipeline = &PSO->VulkanPipeline;
 
 	FPipelineCache& Cache = bIsPrecompileJob ? CurrentPrecompilingPSOCache : GlobalPSOCache;
@@ -1855,6 +1869,7 @@ FVulkanRHIGraphicsPipelineState::FVulkanRHIGraphicsPipelineState(FVulkanDevice*
 
 	PSOInitializer = PSOInitializer_;
 #endif
+	PrecacheKey = RHIComputePrecachePSOHash(PSOInitializer_);
 	INC_DWORD_STAT(STAT_VulkanNumGraphicsPSOs);
 	INC_DWORD_STAT_BY(STAT_VulkanPSOKeyMemory, this->VulkanKey.GetDataRef().Num());
 }
@@ -2308,6 +2323,11 @@ void GetVulkanShaders(FVulkanDevice* Device, const FVulkanRHIGraphicsPipelineSta
 
 void FVulkanPipelineStateCacheManager::TickLRU()
 {
+	if(FVulkanChunkedPipelineCacheManager::IsEnabled())
+	{
+		FVulkanChunkedPipelineCacheManager::Get().Tick();
+	}
+
 	if (!bUseLRU || GVulkanPSOLRUEvictAfterUnusedFrames == 0)
 	{
 		return;
diff --git a/Engine/Source/Runtime/VulkanRHI/Private/VulkanPipeline.h b/Engine/Source/Runtime/VulkanRHI/Private/VulkanPipeline.h
--- a/Engine/Source/Runtime/VulkanRHI/Private/VulkanPipeline.h
+++ b/Engine/Source/Runtime/VulkanRHI/Private/VulkanPipeline.h
@@ -771,6 +771,7 @@ public:
 	FVulkanRHIGraphicsPipelineStateLRUNode* LRUNode = nullptr;
 	uint32 LRUFrame = UINT32_MAX;
 	uint32 PipelineCacheSize = UINT32_MAX;
+	uint64	PrecacheKey;// hash of elements relevant to the PSO cache
 	FVulkanPSOKey							VulkanKey;
 
 #if VULKAN_PSO_CACHE_DEBUG
diff --git a/Engine/Source/Runtime/VulkanRHI/Private/VulkanRHI.cpp b/Engine/Source/Runtime/VulkanRHI/Private/VulkanRHI.cpp
--- a/Engine/Source/Runtime/VulkanRHI/Private/VulkanRHI.cpp
+++ b/Engine/Source/Runtime/VulkanRHI/Private/VulkanRHI.cpp
@@ -27,7 +27,7 @@
 #include "VulkanTransientResourceAllocator.h"
 #include "VulkanExtensions.h"
 #include "VulkanRayTracing.h"
-
+#include "VulkanChunkedPipelineCache.h"
 
 
 // Use Vulkan Profiles to verify feature level support on startup
@@ -98,6 +98,15 @@ static FAutoConsoleVariableRef CVarVulkanVariableRateShading(
 	ECVF_ReadOnly
 );
 
+// If precaching is active we should not need the file cache.
+// however, precaching and filecache are compatible with each other, there maybe some scenarios in which both could be used.
+static TAutoConsoleVariable<bool> CVarEnableVulkanPSOFileCacheWhenPrecachingActive(
+	TEXT("r.Vulkan.EnablePSOFileCacheWhenPrecachingActive"),
+	false,
+	TEXT("false: If precaching is available (r.PSOPrecaching=1, r.Vulkan.UseChunkedPSOCache=1) then disable the PSO filecache. (default)\n")
+	TEXT("true: Allow both PSO file cache and precaching."),
+	ECVF_RenderThreadSafe | ECVF_ReadOnly);
+
 bool GGPUCrashDebuggingEnabled = false;
 
 extern TAutoConsoleVariable<int32> GVulkanRayTracingCVar;
@@ -487,8 +496,10 @@ FVulkanDynamicRHI::FVulkanDynamicRHI()
 	GRHITransitionPrivateData_AlignInBytes = alignof(FVulkanPipelineBarrier);
 	GConfig->GetInt(TEXT("TextureStreaming"), TEXT("PoolSizeVRAMPercentage"), GPoolSizeVRAMPercentage, GEngineIni);
 
-	// VulkanRHI does not yet support PSO precaching, we ensure it is disabled.
-	GRHISupportsPSOPrecaching = false;
+	static const auto CVarPSOPrecaching = IConsoleManager::Get().FindConsoleVariable(TEXT("r.PSOPrecaching"));
+
+	GRHISupportsPSOPrecaching = FVulkanChunkedPipelineCacheManager::IsEnabled() && (CVarPSOPrecaching && CVarPSOPrecaching->GetInt() != 0);
+	GRHISupportsPipelineFileCache = !GRHISupportsPSOPrecaching || CVarEnableVulkanPSOFileCacheWhenPrecachingActive.GetValueOnAnyThread();
 
 	// Copy source requires its own image layout.
 	EnumRemoveFlags(GRHIMergeableAccessMask, ERHIAccess::CopySrc);
@@ -2008,110 +2019,4 @@ IRHITransientResourceAllocator* FVulkanDynamicRHI::RHICreateTransientResourceAll
 	return nullptr;
 }
 
-uint64 FVulkanDynamicRHI::RHIComputePrecachePSOHash(const FGraphicsPipelineStateInitializer& Initializer)
-{
-	// When compute precache PSO hash we assume a valid state precache PSO hash is already provided
-	checkf(Initializer.StatePrecachePSOHash != 0, TEXT("Initializer should have a valid state precache PSO hash set when computing the full initializer PSO hash"));
-
-	// All members which are not part of the state objects
-	struct FNonStateHashKey
-	{
-		uint64							StatePrecachePSOHash;
-
-		EPrimitiveType					PrimitiveType;
-		uint32							RenderTargetsEnabled;
-		FGraphicsPipelineStateInitializer::TRenderTargetFormats	RenderTargetFormats;
-		FGraphicsPipelineStateInitializer::TRenderTargetFlags RenderTargetFlags;
-		EPixelFormat					DepthStencilTargetFormat;
-		ETextureCreateFlags				DepthStencilTargetFlag;
-		uint16							NumSamples;
-		ESubpassHint					SubpassHint;
-		uint8							SubpassIndex;
-		EConservativeRasterization		ConservativeRasterization;
-		uint8							MultiViewCount;
-		EVRSShadingRate					ShadingRate;
-		bool							bDepthBounds;
-		bool							bHasFragmentDensityAttachment;
-	} HashKey;
-
-	FMemory::Memzero(&HashKey, sizeof(FNonStateHashKey));
-
-	HashKey.StatePrecachePSOHash = Initializer.StatePrecachePSOHash;
-
-	HashKey.PrimitiveType = Initializer.PrimitiveType;
-	HashKey.RenderTargetsEnabled = Initializer.RenderTargetsEnabled;
-	HashKey.RenderTargetFormats = Initializer.RenderTargetFormats;
-	HashKey.RenderTargetFlags = Initializer.RenderTargetFlags;
-	HashKey.DepthStencilTargetFormat = Initializer.DepthStencilTargetFormat;
-	HashKey.DepthStencilTargetFlag = Initializer.DepthStencilTargetFlag;
-	HashKey.NumSamples = Initializer.NumSamples;
-	HashKey.SubpassHint = Initializer.SubpassHint;
-	HashKey.SubpassIndex = Initializer.SubpassIndex;
-	HashKey.ConservativeRasterization = Initializer.ConservativeRasterization;
-	HashKey.MultiViewCount = Initializer.MultiViewCount;
-	HashKey.ShadingRate = Initializer.ShadingRate;
-	HashKey.bDepthBounds = Initializer.bDepthBounds;
-	HashKey.bHasFragmentDensityAttachment = Initializer.bHasFragmentDensityAttachment;
-	
-	// TODO: check if any RT flags actually affect PSO in VK
-	for (ETextureCreateFlags& Flags : HashKey.RenderTargetFlags)
-	{
-		Flags = Flags & FGraphicsPipelineStateInitializer::RelevantRenderTargetFlagMask;
-	}
-	HashKey.DepthStencilTargetFlag = (HashKey.DepthStencilTargetFlag & FGraphicsPipelineStateInitializer::RelevantDepthStencilFlagMask);
-		 
-	return CityHash64((const char*)&HashKey, sizeof(FNonStateHashKey));
-}
-
-bool FVulkanDynamicRHI::RHIMatchPrecachePSOInitializers(const FGraphicsPipelineStateInitializer& LHS, const FGraphicsPipelineStateInitializer& RHS)
-{
-	// first check non pointer objects
-	if (LHS.ImmutableSamplerState != RHS.ImmutableSamplerState ||
-		LHS.PrimitiveType != RHS.PrimitiveType ||
-		LHS.bDepthBounds != RHS.bDepthBounds ||
-		LHS.MultiViewCount != RHS.MultiViewCount ||
-		LHS.ShadingRate != RHS.ShadingRate ||
-		LHS.bHasFragmentDensityAttachment != RHS.bHasFragmentDensityAttachment ||
-		LHS.RenderTargetsEnabled != RHS.RenderTargetsEnabled ||
-		LHS.RenderTargetFormats != RHS.RenderTargetFormats ||
-		!FGraphicsPipelineStateInitializer::RelevantRenderTargetFlagsEqual(LHS.RenderTargetFlags, RHS.RenderTargetFlags) ||
-		LHS.DepthStencilTargetFormat != RHS.DepthStencilTargetFormat ||
-		!FGraphicsPipelineStateInitializer::RelevantDepthStencilFlagsEqual(LHS.DepthStencilTargetFlag, RHS.DepthStencilTargetFlag) ||
-		LHS.NumSamples != RHS.NumSamples ||
-		LHS.SubpassHint != RHS.SubpassHint ||
-		LHS.SubpassIndex != RHS.SubpassIndex ||
-		LHS.StatePrecachePSOHash != RHS.StatePrecachePSOHash ||
-		LHS.ConservativeRasterization != RHS.ConservativeRasterization)
-	{
-		return false;
-	}
-
-	// check the RHI shaders (pointer check for shaders should be fine)
-	if (LHS.BoundShaderState.VertexShaderRHI != RHS.BoundShaderState.VertexShaderRHI ||
-		LHS.BoundShaderState.PixelShaderRHI != RHS.BoundShaderState.PixelShaderRHI ||
-		LHS.BoundShaderState.GetMeshShader() != RHS.BoundShaderState.GetMeshShader() ||
-		LHS.BoundShaderState.GetAmplificationShader() != RHS.BoundShaderState.GetAmplificationShader() ||
-		LHS.BoundShaderState.GetGeometryShader() != RHS.BoundShaderState.GetGeometryShader())
-	{
-		return false;
-	}
-
-	// Full compare the of the vertex declaration
-	if (!MatchRHIState<FRHIVertexDeclaration, FVertexDeclarationElementList>(LHS.BoundShaderState.VertexDeclarationRHI, RHS.BoundShaderState.VertexDeclarationRHI))
-	{
-		return false;
-	}
-
-	// Check actual state content (each initializer can have it's own state and not going through a factory)
-	if (!MatchRHIState<FRHIBlendState, FBlendStateInitializerRHI>(LHS.BlendState, RHS.BlendState)
-		|| !MatchRHIState<FRHIRasterizerState, FRasterizerStateInitializerRHI>(LHS.RasterizerState, RHS.RasterizerState)
-		|| !MatchRHIState<FRHIDepthStencilState, FDepthStencilStateInitializerRHI>(LHS.DepthStencilState, RHS.DepthStencilState)
-		)
-	{
-		return false;
-	}
-
-	return true;
-}
-
 #undef LOCTEXT_NAMESPACE
diff --git a/Engine/Source/Runtime/VulkanRHI/Public/VulkanDynamicRHI.h b/Engine/Source/Runtime/VulkanRHI/Public/VulkanDynamicRHI.h
--- a/Engine/Source/Runtime/VulkanRHI/Public/VulkanDynamicRHI.h
+++ b/Engine/Source/Runtime/VulkanRHI/Public/VulkanDynamicRHI.h
@@ -191,9 +191,6 @@ public:
 	virtual void RHISubmitCommandLists(TArrayView<IRHIPlatformCommandList*> CommandLists, bool bFlushResources) final override;
 	virtual uint64 RHIGetMinimumAlignmentForBufferBackedSRV(EPixelFormat Format) final override;
 
-	virtual uint64 RHIComputePrecachePSOHash(const FGraphicsPipelineStateInitializer& Initializer) final override;
-	virtual bool RHIMatchPrecachePSOInitializers(const FGraphicsPipelineStateInitializer& LHS, const FGraphicsPipelineStateInitializer& RHS) final override;
-
 
 	// Transient Resource Functions and Helpers
 	virtual IRHITransientResourceAllocator* RHICreateTransientResourceAllocator() final override;
