Fixed: Incorrect PSO state in Vulkan (Android) related to depth pre-pass detection during render-pass setup

diff --git a/Engine/Source/Runtime/VulkanRHI/Private/VulkanRenderpass.h b/Engine/Source/Runtime/VulkanRHI/Private/VulkanRenderpass.h
--- a/Engine/Source/Runtime/VulkanRHI/Private/VulkanRenderpass.h
+++ b/Engine/Source/Runtime/VulkanRHI/Private/VulkanRenderpass.h
@@ -461,7 +461,9 @@ public:
 				DepthStencilAttachment.attachment = RTLayout.GetDepthAttachmentReference()->attachment;
 				DepthStencilAttachment.SetAspect(VK_IMAGE_ASPECT_DEPTH_BIT);	// @todo?
 
-				if (RTLayout.GetDepthAttachmentReference()->layout == VK_IMAGE_LAYOUT_READ_ONLY_OPTIMAL)
+				// FIXME: checking a Depth layout is not correct in all cases
+				// PSO cache can create a PSO for subpass 1 or 2 first, where depth is read-only but that does not mean depth pre-pass is enabled
+				if (false && RTLayout.GetDepthAttachmentReference()->layout == VK_IMAGE_LAYOUT_READ_ONLY_OPTIMAL)
 				{
 					// Depth is read only and is expected to be sampled as a regular texture
 					DepthStencilAttachment.layout = VK_IMAGE_LAYOUT_READ_ONLY_OPTIMAL;
