VulkanRHI: Give unique name to api_dump filename.

#rnx

diff --git a/Engine/Source/Runtime/VulkanRHI/Private/VulkanLayers.cpp b/Engine/Source/Runtime/VulkanRHI/Private/VulkanLayers.cpp
--- a/Engine/Source/Runtime/VulkanRHI/Private/VulkanLayers.cpp
+++ b/Engine/Source/Runtime/VulkanRHI/Private/VulkanLayers.cpp
@@ -379,7 +379,8 @@ void FVulkanIntanceSetupHelper::AddDebugLayers(const TArray<FLayerWithExtensions
 			const bool bApiDumpFound = AddRequestedLayer(VkApiDumpName, LayerProperties, UEExtensions, OutLayers);
 			if (bApiDumpFound)
 			{
-				FPlatformMisc::SetEnvironmentVar(TEXT("VK_APIDUMP_LOG_FILENAME"), TEXT("vk_apidump.txt"));
+				const FString ApiDumpFileName = FString::Printf(TEXT("vk_apidump.%s.txt"), *FDateTime::Now().ToString());
+				FPlatformMisc::SetEnvironmentVar(TEXT("VK_APIDUMP_LOG_FILENAME"), *ApiDumpFileName);
 				FPlatformMisc::SetEnvironmentVar(TEXT("VK_APIDUMP_DETAILED"), TEXT("true"));
 				FPlatformMisc::SetEnvironmentVar(TEXT("VK_APIDUMP_FLUSH"), TEXT("true"));
 				FPlatformMisc::SetEnvironmentVar(TEXT("VK_APIDUMP_OUTPUT_FORMAT"), TEXT("text"));
