VulkanRHI: Reformat comments in VulkanCommon.h to include more hlsl types

#rnx

diff --git a/Engine/Source/Runtime/VulkanRHI/Public/VulkanCommon.h b/Engine/Source/Runtime/VulkanRHI/Public/VulkanCommon.h
--- a/Engine/Source/Runtime/VulkanRHI/Public/VulkanCommon.h
+++ b/Engine/Source/Runtime/VulkanRHI/Public/VulkanCommon.h
@@ -144,22 +144,20 @@ namespace EVulkanBindingType
 		PackedUniformBuffer,	//VK_DESCRIPTOR_TYPE_UNIFORM_BUFFER
 		UniformBuffer,			//VK_DESCRIPTOR_TYPE_UNIFORM_BUFFER
 
-		CombinedImageSampler,	//VK_DESCRIPTOR_TYPE_COMBINED_IMAGE_SAMPLER
-		Sampler,				//VK_DESCRIPTOR_TYPE_SAMPLER
-		Image,					//VK_DESCRIPTOR_TYPE_SAMPLED_IMAGE
+		CombinedImageSampler,	//VK_DESCRIPTOR_TYPE_COMBINED_IMAGE_SAMPLER	*not used*
+		Sampler,				//VK_DESCRIPTOR_TYPE_SAMPLER				(HLSL: SamplerState/SamplerComparisonState)
+		Image,					//VK_DESCRIPTOR_TYPE_SAMPLED_IMAGE			(HLSL: Texture2D/3D/Cube)
 
-		UniformTexelBuffer,		//VK_DESCRIPTOR_TYPE_UNIFORM_TEXEL_BUFFER	Buffer<>
+		UniformTexelBuffer,		//VK_DESCRIPTOR_TYPE_UNIFORM_TEXEL_BUFFER	(HLSL: Buffer)
 
-		//A storage image (VK_DESCRIPTOR_TYPE_STORAGE_IMAGE) is a descriptor type that is used for load, store, and atomic operations on image memory from within shaders bound to pipelines.
-		StorageImage,			//VK_DESCRIPTOR_TYPE_STORAGE_IMAGE		RWTexture
+		// A storage image is a descriptor type that is used for load, store, and atomic operations on image memory from within shaders bound to pipelines.
+		StorageImage,			//VK_DESCRIPTOR_TYPE_STORAGE_IMAGE			(HLSL: RWTexture2D/3D/Cube)
 
-		//RWBuffer/RWTexture?
-		//A storage texel buffer (VK_DESCRIPTOR_TYPE_STORAGE_TEXEL_BUFFER) represents a tightly packed array of homogeneous formatted data that is stored in a buffer and is made accessible to shaders. Storage texel buffers differ from uniform texel buffers in that they support stores and atomic operations in shaders, may support a different maximum length, and may have different performance characteristics.
-		StorageTexelBuffer,
+		//A storage texel buffer represents a tightly packed array of homogeneous formatted data that is stored in a buffer and is made accessible to shaders. Storage texel buffers differ from uniform texel buffers in that they support stores and atomic operations in shaders, may support a different maximum length, and may have different performance characteristics.
+		StorageTexelBuffer,		//VK_DESCRIPTOR_TYPE_STORAGE_TEXEL_BUFFER	(HLSL: RWBuffer)
 
-		// UAV/RWBuffer
-		//A storage buffer(VK_DESCRIPTOR_TYPE_STORAGE_BUFFER) is a region of structured storage that supports both read and write access for shaders.In addition to general read and write operations, some members of storage buffers can be used as the target of atomic operations.In general, atomic operations are only supported on members that have unsigned integer formats.
-		StorageBuffer,
+		// A storage buffer is a region of structured storage that supports both read and write access for shaders. In addition to general read and write operations, some members of storage buffers can be used as the target of atomic operations. In general, atomic operations are only supported on members that have unsigned integer formats.
+		StorageBuffer,			//VK_DESCRIPTOR_TYPE_STORAGE_BUFFER			(HLSL: StructuredBuffer/RWStructureBuffer/ByteAddressBuffer/RWByteAddressBuffer)
 
 		InputAttachment,
 
