VulkanRHI: VK RT WIP: Add CVar option to flush commands after each RayTraceDispatch.  Add missing release on staging buffer.

#rnx

diff --git a/Engine/Source/Runtime/VulkanRHI/Private/VulkanRayTracing.cpp b/Engine/Source/Runtime/VulkanRHI/Private/VulkanRayTracing.cpp
--- a/Engine/Source/Runtime/VulkanRHI/Private/VulkanRayTracing.cpp
+++ b/Engine/Source/Runtime/VulkanRHI/Private/VulkanRayTracing.cpp
@@ -47,6 +47,15 @@ static FAutoConsoleVariableRef CVarVulkanRayTracingAllowDeferredOperation(
 	ECVF_ReadOnly
 );
 
+static int32 GVulkanSubmitOnTraceRays = 0;
+static FAutoConsoleVariableRef GCVarSubmitOnTraceRays(
+	TEXT("r.Vulkan.SubmitOnTraceRays"),
+	GVulkanSubmitOnTraceRays,
+	TEXT("0 to not do anything special on trace rays (default)\n")\
+	TEXT("1 to submit the cmd buffer after each trace rays"),
+	ECVF_ReadOnly
+);
+
 
 // Ray tracing stat counters
 
@@ -735,14 +744,14 @@ void FVulkanRayTracingScene::BuildAccelerationStructure(
 	check(AccelerationStructureBuffer.IsValid());
 	check(InInstanceBuffer != nullptr);
 
-	TRefCountPtr<FVulkanResourceMultiBuffer> ScratchBuffer;
+	FBufferRHIRef ScratchBuffer;
 
 	if (InScratchBuffer == nullptr)
 	{
 		TRHICommandList_RecursiveHazardous<FVulkanCommandListContext> RHICmdList(&CommandContext);
 		FRHIResourceCreateInfo ScratchBufferCreateInfo(TEXT("BuildScratchTLAS"));
-		ScratchBuffer = ResourceCast(RHICmdList.CreateBuffer(SizeInfo.BuildScratchSize, BUF_StructuredBuffer | BUF_RayTracingScratch, 0, ERHIAccess::UAVCompute, ScratchBufferCreateInfo).GetReference());
-		InScratchBuffer = ScratchBuffer.GetReference();
+		ScratchBuffer = RHICmdList.CreateBuffer(SizeInfo.BuildScratchSize, BUF_StructuredBuffer | BUF_RayTracingScratch, 0, ERHIAccess::UAVCompute, ScratchBufferCreateInfo);
+		InScratchBuffer = ResourceCast(ScratchBuffer.GetReference());
 		InScratchOffset = 0;
 	}
 
@@ -760,14 +769,13 @@ void FVulkanRayTracingScene::BuildAccelerationStructure(
 	TArray<VkAccelerationStructureBuildRangeInfoKHR*> pBuildRanges;
 	pBuildRanges.SetNum(NumLayers);
 
-	uint32 InstanceBaseOffset = 0;
+	const VkDeviceAddress InstanceBufferAddress = InInstanceBuffer->GetDeviceAddress() + InInstanceOffset;
 
+	uint32 InstanceBaseOffset = 0;
 	for (uint32 LayerIndex = 0; LayerIndex < NumLayers; ++LayerIndex)
 	{
 		const FLayerData& Layer = Layers[LayerIndex];
 
-		VkDeviceAddress InstanceBufferAddress = InInstanceBuffer->GetDeviceAddress() + InInstanceOffset;
-
 		FVkRtTLASBuildData& BuildData = BuildDatas[LayerIndex];
 		GetTLASBuildData(Device->GetInstanceHandle(), Initializer.NumNativeInstancesPerLayer[LayerIndex], InstanceBufferAddress, BuildData);
 
@@ -1016,6 +1024,8 @@ void FVulkanRayTracingShaderTable::Commit(FVulkanCommandListContext& Context)
 				RegionInfo.dstOffset = 0;
 				RegionInfo.size = Alloc.Region.size;
 				VulkanRHI::vkCmdCopyBuffer(CmdBuffer->GetHandle(), StagingBuffer->GetHandle(), Alloc.LocalBuffer, 1, &RegionInfo);
+
+				Device->GetStagingManager().ReleaseBuffer(CmdBuffer, StagingBuffer);
 			}
 			else
 			{
@@ -1780,11 +1790,11 @@ void FVulkanCommandListContext::RHIRayTraceDispatch(
 	FVulkanCmdBuffer* const CmdBuffer = GetCommandBufferManager()->GetActiveCmdBuffer();
 	VulkanRHI::vkCmdBindPipeline(CmdBuffer->GetHandle(), VK_PIPELINE_BIND_POINT_RAY_TRACING_KHR, Pipeline->GetPipeline());
 
-	SetRayGenResources(Device, CmdBuffer, InGlobalResourceBindings);
-
 	ShaderTable->SetSlot(InRayGenShader->GetFrequency(), 0, Pipeline->GetShaderIndex(RayGenShader), Pipeline->GetShaderHandles(SF_RayGen));
 	ShaderTable->Commit(*this);
 
+	SetRayGenResources(Device, CmdBuffer, InGlobalResourceBindings);
+
 	VulkanRHI::vkCmdTraceRaysKHR(
 		CmdBuffer->GetHandle(),
 		ShaderTable->GetRegion(SF_RayGen),
@@ -1792,6 +1802,11 @@ void FVulkanCommandListContext::RHIRayTraceDispatch(
 		ShaderTable->GetRegion(SF_RayHitGroup),
 		ShaderTable->GetRegion(SF_RayCallable),
 		InWidth, InHeight, 1);
+
+	if (GVulkanSubmitOnTraceRays)
+	{
+		InternalSubmitActiveCmdBuffer();
+	}
 }
 
 void FVulkanCommandListContext::RHIRayTraceDispatchIndirect(
@@ -1811,11 +1826,11 @@ void FVulkanCommandListContext::RHIRayTraceDispatchIndirect(
 	FVulkanCmdBuffer* const CmdBuffer = GetCommandBufferManager()->GetActiveCmdBuffer();
 	VulkanRHI::vkCmdBindPipeline(CmdBuffer->GetHandle(), VK_PIPELINE_BIND_POINT_RAY_TRACING_KHR, Pipeline->GetPipeline());
 
-	SetRayGenResources(Device, CmdBuffer, InGlobalResourceBindings);
-
 	ShaderTable->SetSlot(InRayGenShader->GetFrequency(), 0, Pipeline->GetShaderIndex(RayGenShader), Pipeline->GetShaderHandles(SF_RayGen));
 	ShaderTable->Commit(*this);
 
+	SetRayGenResources(Device, CmdBuffer, InGlobalResourceBindings);
+
 	FVulkanResourceMultiBuffer* ArgumentBuffer = ResourceCast(InArgumentBuffer);
 	const VkDeviceAddress IndirectDeviceAddress = ArgumentBuffer->GetDeviceAddress() + InArgumentOffset;
 
@@ -1826,6 +1841,11 @@ void FVulkanCommandListContext::RHIRayTraceDispatchIndirect(
 		ShaderTable->GetRegion(SF_RayHitGroup),
 		ShaderTable->GetRegion(SF_RayCallable),
 		IndirectDeviceAddress);
+
+	if (GVulkanSubmitOnTraceRays)
+	{
+		InternalSubmitActiveCmdBuffer();
+	}
 }
 
 static void SetSystemParametersUB(FVulkanHitGroupSystemParameters& OutSystemParameters, uint32 InNumUniformBuffers, FRHIUniformBuffer* const* InUniformBuffers, const FVulkanRayTracingShader* InShader)
