VulkanRHI: VK RT WIP: Add new values in shader header to indicate the presence of extra shaders for hit groups.

#rnx

diff --git a/Engine/Shaders/Public/Platform/Vulkan/VulkanCommon.ush b/Engine/Shaders/Public/Platform/Vulkan/VulkanCommon.ush
--- a/Engine/Shaders/Public/Platform/Vulkan/VulkanCommon.ush
+++ b/Engine/Shaders/Public/Platform/Vulkan/VulkanCommon.ush
@@ -5,7 +5,7 @@
 =============================================================================*/
 
 // Update this GUID to improve shader recompilation for Vulkan only shaders
-// GUID = 4BB75D9E-9229-4BB3-959A-36A81475A4FE
+// GUID = 77A59EF7-19F3-4EE5-9058-F1E42AAAE164
 
 #pragma once
 
diff --git a/Engine/Source/Runtime/VulkanRHI/Public/VulkanShaderResources.h b/Engine/Source/Runtime/VulkanRHI/Public/VulkanShaderResources.h
--- a/Engine/Source/Runtime/VulkanRHI/Public/VulkanShaderResources.h
+++ b/Engine/Source/Runtime/VulkanRHI/Public/VulkanShaderResources.h
@@ -187,11 +187,27 @@ struct FVulkanShaderHeader
 	uint32                                  RayTracingPayloadType = 0;
 	uint32                                  RayTracingPayloadSize = 0;
 
-
 	FSHAHash								SourceHash;
 	uint32									SpirvCRC = 0;
 	uint8									WaveSize = 0;
 
+	// For RayHitGroup shaders
+	enum class ERayHitGroupEntrypoint : uint8
+	{
+		NotPresent = 0,
+
+		// Hit group types are all stored in a single spirv blob
+		// and each have different entry point names
+		// NOTE: Not used yet because of compiler issues
+		CommonBlob,
+
+		// Hit group types are each stored in a different spirv blob
+		// to circumvent DXC compilation issues
+		SeparateBlob
+	};
+	ERayHitGroupEntrypoint RayGroupAnyHit = ERayHitGroupEntrypoint::NotPresent;
+	ERayHitGroupEntrypoint RayGroupIntersection = ERayHitGroupEntrypoint::NotPresent;
+
 	TArray<FSpirvInfo>						UniformBufferSpirvInfos;
 	TArray<FSpirvInfo>						GlobalSpirvInfos;
 
@@ -296,6 +312,8 @@ inline FArchive& operator<<(FArchive& Ar, FVulkanShaderHeader& Header)
 	Ar << Header.SourceHash;
 	Ar << Header.SpirvCRC;
 	Ar << Header.WaveSize;
+	Ar << Header.RayGroupAnyHit;
+	Ar << Header.RayGroupIntersection;
 	Ar << Header.UniformBufferSpirvInfos;
 	Ar << Header.GlobalSpirvInfos;
 	Ar << Header.DebugName;
