VulkanRHI: (UDN) Fix VulkanLLM stats, did not match up with existing llm scopes in VulkanBuffers.cpp.

#rnx

diff --git a/Engine/Source/Runtime/VulkanRHI/Private/VulkanLLM.cpp b/Engine/Source/Runtime/VulkanRHI/Private/VulkanLLM.cpp
--- a/Engine/Source/Runtime/VulkanRHI/Private/VulkanLLM.cpp
+++ b/Engine/Source/Runtime/VulkanRHI/Private/VulkanLLM.cpp
@@ -22,8 +22,7 @@ struct FLLMTagInfoVulkan
 
 DECLARE_LLM_MEMORY_STAT(TEXT("VulkanMisc"), STAT_VulkanMiscLLM, STATGROUP_LLMPlatform);
 DECLARE_LLM_MEMORY_STAT(TEXT("VulkanUniformBuffers"), STAT_VulkanUniformBuffersLLM, STATGROUP_LLMPlatform);
-DECLARE_LLM_MEMORY_STAT(TEXT("VulkanIndexBuffers"), STAT_VulkanIndexBuffersLLM, STATGROUP_LLMPlatform);
-DECLARE_LLM_MEMORY_STAT(TEXT("VulkanVertexBuffers"), STAT_VulkanVertexBuffersLLM, STATGROUP_LLMPlatform);
+DECLARE_LLM_MEMORY_STAT(TEXT("VulkanBuffers"), STAT_VulkanBuffersLLM, STATGROUP_LLMPlatform);
 DECLARE_LLM_MEMORY_STAT(TEXT("VulkanTextures"), STAT_VulkanTexturesLLM, STATGROUP_LLMPlatform);
 DECLARE_LLM_MEMORY_STAT(TEXT("VulkanRenderTargets"), STAT_VulkanRenderTargetsLLM, STATGROUP_LLMPlatform);
 DECLARE_LLM_MEMORY_STAT(TEXT("VulkanSpareMemoryGPU"), STAT_VulkanSpareMemoryGPULLM, STATGROUP_LLMPlatform);
@@ -40,8 +39,7 @@ static const FLLMTagInfoVulkan ELLMTagNamesVulkan[] =
 	// csv name							// stat name												// summary stat name						// enum value
 	{ TEXT("VulkanMisc"),				GET_STATFNAME(STAT_VulkanMiscLLM),							GET_STATFNAME(STAT_EngineSummaryLLM) },		// ELLMTagVulkan::VulkanMisc
 	{ TEXT("VulkanUniformBuffers"),		GET_STATFNAME(STAT_VulkanUniformBuffersLLM),				GET_STATFNAME(STAT_EngineSummaryLLM) },		// ELLMTagVulkan::VulkanUniformBuffers
-	{ TEXT("VulkanIndexBuffers"),		GET_STATFNAME(STAT_VulkanIndexBuffersLLM),					GET_STATFNAME(STAT_EngineSummaryLLM) },		// ELLMTagVulkan::VulkanIndexBuffers
-	{ TEXT("VulkanVertexBuffers"),		GET_STATFNAME(STAT_VulkanVertexBuffersLLM),					GET_STATFNAME(STAT_EngineSummaryLLM) },		// ELLMTagVulkan::VulkanVertexBuffers
+	{ TEXT("VulkanBuffers"),			GET_STATFNAME(STAT_VulkanBuffersLLM),						GET_STATFNAME(STAT_EngineSummaryLLM) },		// ELLMTagVulkan::VulkanBuffers
 	{ TEXT("VulkanTextures"),			GET_STATFNAME(STAT_VulkanTexturesLLM),						GET_STATFNAME(STAT_EngineSummaryLLM) },		// ELLMTagVulkan::VulkanTextures
 	{ TEXT("VulkanRenderTargets"),		GET_STATFNAME(STAT_VulkanRenderTargetsLLM),					GET_STATFNAME(STAT_EngineSummaryLLM) },		// ELLMTagVulkan::VulkanRenderTargets
 	{ TEXT("VulkanSpareMemoryGPU"),		GET_STATFNAME(STAT_VulkanSpareMemoryGPULLM),				GET_STATFNAME(STAT_EngineSummaryLLM) },		// ELLMTagVulkan::VulkanSpareMemoryGPU
@@ -59,7 +57,7 @@ namespace VulkanLLM
 {
 	void Initialize()
 	{
-		int32 TagCount = sizeof(ELLMTagNamesVulkan) / sizeof(FLLMTagInfoVulkan);
+		const int32 TagCount = sizeof(ELLMTagNamesVulkan) / sizeof(FLLMTagInfoVulkan);
 
 		for (int32 Index = 0; Index < TagCount; ++Index)
 		{
