fix for missing items when looking down on Android Vulkan devices
#tests preflight embedded qa

diff --git a/Engine/Source/Runtime/VulkanRHI/Private/VulkanRenderpass.h b/Engine/Source/Runtime/VulkanRHI/Private/VulkanRenderpass.h
--- a/Engine/Source/Runtime/VulkanRHI/Private/VulkanRenderpass.h
+++ b/Engine/Source/Runtime/VulkanRHI/Private/VulkanRenderpass.h
@@ -470,10 +470,10 @@ public:
 				}
 				else
 				{
-					DepthStencilAttachment.layout = VK_IMAGE_LAYOUT_DEPTH_STENCIL_READ_ONLY_OPTIMAL;
-
+					// lights write to stencil for culling, so stencil is expected to be writebale while depth is read-only
+					DepthStencilAttachment.layout = VK_IMAGE_LAYOUT_DEPTH_READ_ONLY_STENCIL_ATTACHMENT_OPTIMAL;
 					DepthInputAttachment = DepthStencilAttachment.attachment;
-					DepthInputAttachmentLayout = VK_IMAGE_LAYOUT_DEPTH_STENCIL_READ_ONLY_OPTIMAL;
+					DepthInputAttachmentLayout = DepthStencilAttachment.layout;
 					DepthInputAspectMask = VK_IMAGE_ASPECT_DEPTH_BIT;
 				}
 			}
