Fix checkf's printf argument errors in VULKAN_SIGNAL_UNIMPLEMENTED macro

#rnx

diff --git a/Engine/Source/Runtime/VulkanRHI/Private/Linux/VulkanLinuxPlatform.h b/Engine/Source/Runtime/VulkanRHI/Private/Linux/VulkanLinuxPlatform.h
--- a/Engine/Source/Runtime/VulkanRHI/Private/Linux/VulkanLinuxPlatform.h
+++ b/Engine/Source/Runtime/VulkanRHI/Private/Linux/VulkanLinuxPlatform.h
@@ -9,7 +9,7 @@
 #define VULKAN_ENABLE_DUMP_LAYER					0
 #define VULKAN_SHOULD_DEBUG_IN_DEVELOPMENT			1
 #define VULKAN_SHOULD_ENABLE_DRAW_MARKERS			(UE_BUILD_DEBUG || UE_BUILD_DEVELOPMENT)
-#define VULKAN_SIGNAL_UNIMPLEMENTED()				checkf(false, TEXT("Unimplemented vulkan functionality: %s"), __PRETTY_FUNCTION__)
+#define VULKAN_SIGNAL_UNIMPLEMENTED()				checkf(false, TEXT("Unimplemented vulkan functionality: %hs"), __PRETTY_FUNCTION__)
 #define VULKAN_SUPPORTS_AMD_BUFFER_MARKER			1
 
 #define VULKAN_RHI_RAYTRACING 						(RHI_RAYTRACING)
