VulkanRHI: If VK_KHR_synchronization2 is manually disabled, then also disable VK_EXT_descriptor_buffer (since sync2 is a prereq anyhow).

#rnx

diff --git a/Engine/Source/Runtime/VulkanRHI/Private/VulkanExtensions.cpp b/Engine/Source/Runtime/VulkanRHI/Private/VulkanExtensions.cpp
--- a/Engine/Source/Runtime/VulkanRHI/Private/VulkanExtensions.cpp
+++ b/Engine/Source/Runtime/VulkanRHI/Private/VulkanExtensions.cpp
@@ -70,7 +70,7 @@ TAutoConsoleVariable<int32> GVulkanAllowHostQueryResetCVar(
 	ECVF_ReadOnly
 );
 
-TAutoConsoleVariable<int32> GVulkaAllowSync2BarriersCVar(
+TAutoConsoleVariable<int32> GVulkanAllowSync2BarriersCVar(
 	TEXT("r.Vulkan.AllowSynchronization2"),
 	1,
 	TEXT("Enables the use of advanced barriers that combine the use of the VK_KHR_separate_depth_stencil_layouts \n")
@@ -327,7 +327,7 @@ public:
 	FVulkanKHRSeparateDepthStencilLayoutsExtension(FVulkanDevice* InDevice)
 		: FVulkanDeviceExtension(InDevice, VK_KHR_SEPARATE_DEPTH_STENCIL_LAYOUTS_EXTENSION_NAME, VULKAN_EXTENSION_ENABLED, VK_API_VERSION_1_2)
 	{
-		bEnabledInCode = bEnabledInCode && (GVulkaAllowSync2BarriersCVar.GetValueOnAnyThread() != 0);
+		bEnabledInCode = bEnabledInCode && (GVulkanAllowSync2BarriersCVar.GetValueOnAnyThread() != 0);
 	}
 
 	virtual void PrePhysicalDeviceFeatures(VkPhysicalDeviceFeatures2KHR& PhysicalDeviceFeatures2) override final
@@ -364,7 +364,7 @@ public:
 	FVulkanKHRSynchronization2(FVulkanDevice* InDevice)
 		: FVulkanDeviceExtension(InDevice, VK_KHR_SYNCHRONIZATION_2_EXTENSION_NAME, VULKAN_EXTENSION_ENABLED, VK_API_VERSION_1_3)
 	{
-		bEnabledInCode = bEnabledInCode && (GVulkaAllowSync2BarriersCVar.GetValueOnAnyThread() != 0);
+		bEnabledInCode = bEnabledInCode && (GVulkanAllowSync2BarriersCVar.GetValueOnAnyThread() != 0);
 	}
 
 	virtual void PrePhysicalDeviceFeatures(VkPhysicalDeviceFeatures2KHR& PhysicalDeviceFeatures2) override final
@@ -1172,6 +1172,8 @@ public:
 	FVulkanEXTDescriptorBuffer(FVulkanDevice* InDevice)
 		: FVulkanDeviceExtension(InDevice, VK_EXT_DESCRIPTOR_BUFFER_EXTENSION_NAME, VULKAN_EXTENSION_ENABLED, VULKAN_EXTENSION_NOT_PROMOTED)
 	{
+		// Sync2 is a prereq
+		bEnabledInCode = bEnabledInCode && (GVulkanAllowSync2BarriersCVar.GetValueOnAnyThread() != 0);
 	}
 
 	virtual void PrePhysicalDeviceFeatures(VkPhysicalDeviceFeatures2KHR& PhysicalDeviceFeatures2) override final
