VulkanRHI: send api_dump files to the log folder with everything else.

#rnx

diff --git a/Engine/Source/Runtime/VulkanRHI/Private/VulkanLayers.cpp b/Engine/Source/Runtime/VulkanRHI/Private/VulkanLayers.cpp
--- a/Engine/Source/Runtime/VulkanRHI/Private/VulkanLayers.cpp
+++ b/Engine/Source/Runtime/VulkanRHI/Private/VulkanLayers.cpp
@@ -379,7 +379,7 @@ void FVulkanIntanceSetupHelper::AddDebugLayers(const TArray<FLayerWithExtensions
 			const bool bApiDumpFound = AddRequestedLayer(VkApiDumpName, LayerProperties, UEExtensions, OutLayers);
 			if (bApiDumpFound)
 			{
-				const FString ApiDumpFileName = FString::Printf(TEXT("vk_apidump.%s.txt"), *FDateTime::Now().ToString());
+				const FString ApiDumpFileName = FString::Printf(TEXT("%s/vk_apidump.%s.txt"), *FPaths::ProjectLogDir(), *FDateTime::Now().ToString());
 				FPlatformMisc::SetEnvironmentVar(TEXT("VK_APIDUMP_LOG_FILENAME"), *ApiDumpFileName);
 				FPlatformMisc::SetEnvironmentVar(TEXT("VK_APIDUMP_DETAILED"), TEXT("true"));
 				FPlatformMisc::SetEnvironmentVar(TEXT("VK_APIDUMP_FLUSH"), TEXT("true"));
