Back out changelist 27535773

diff --git a/Engine/Source/Runtime/VulkanRHI/Private/VulkanRenderpass.h b/Engine/Source/Runtime/VulkanRHI/Private/VulkanRenderpass.h
--- a/Engine/Source/Runtime/VulkanRHI/Private/VulkanRenderpass.h
+++ b/Engine/Source/Runtime/VulkanRHI/Private/VulkanRenderpass.h
@@ -470,10 +470,10 @@ public:
 				}
 				else
 				{
-					// lights write to stencil for culling, so stencil is expected to be writebale while depth is read-only
-					DepthStencilAttachment.layout = VK_IMAGE_LAYOUT_DEPTH_READ_ONLY_STENCIL_ATTACHMENT_OPTIMAL;
+					DepthStencilAttachment.layout = VK_IMAGE_LAYOUT_DEPTH_STENCIL_READ_ONLY_OPTIMAL;
+
 					DepthInputAttachment = DepthStencilAttachment.attachment;
-					DepthInputAttachmentLayout = DepthStencilAttachment.layout;
+					DepthInputAttachmentLayout = VK_IMAGE_LAYOUT_DEPTH_STENCIL_READ_ONLY_OPTIMAL;
 					DepthInputAspectMask = VK_IMAGE_ASPECT_DEPTH_BIT;
 				}
 			}
