Fixed depth stencil layout transitions on vulkan devices with disabled or unsupported VK_KHR_synchronization2

diff --git a/Engine/Source/Runtime/VulkanRHI/Private/VulkanRHIPrivate.h b/Engine/Source/Runtime/VulkanRHI/Private/VulkanRHIPrivate.h
--- a/Engine/Source/Runtime/VulkanRHI/Private/VulkanRHIPrivate.h
+++ b/Engine/Source/Runtime/VulkanRHI/Private/VulkanRHIPrivate.h
@@ -803,7 +803,7 @@ namespace VulkanRHI
 			return VK_IMAGE_LAYOUT_UNDEFINED;
 		}
 
-		if (DepthLayout == VK_IMAGE_LAYOUT_ATTACHMENT_OPTIMAL)
+		if (DepthLayout == VK_IMAGE_LAYOUT_ATTACHMENT_OPTIMAL || DepthLayout == VK_IMAGE_LAYOUT_DEPTH_STENCIL_ATTACHMENT_OPTIMAL)
 		{
 			// :todo-jn:  barrier downgrades from sync2 to sync1 happen before the tracking is updated sometimes, which causes the StencilLayout to contain an older sync1 value. 
 			// temporarily accept those values until we fix it at the source.
@@ -818,9 +818,9 @@ namespace VulkanRHI
 				return VK_IMAGE_LAYOUT_DEPTH_ATTACHMENT_STENCIL_READ_ONLY_OPTIMAL;
 			}
 		}
-		else if (DepthLayout == VK_IMAGE_LAYOUT_READ_ONLY_OPTIMAL)
+		else if (DepthLayout == VK_IMAGE_LAYOUT_READ_ONLY_OPTIMAL || DepthLayout == VK_IMAGE_LAYOUT_DEPTH_STENCIL_READ_ONLY_OPTIMAL)
 		{
-			if ((StencilLayout == VK_IMAGE_LAYOUT_READ_ONLY_OPTIMAL) || (StencilLayout == VK_IMAGE_LAYOUT_UNDEFINED))
+			if ((StencilLayout == VK_IMAGE_LAYOUT_READ_ONLY_OPTIMAL) || (StencilLayout == VK_IMAGE_LAYOUT_UNDEFINED) || (StencilLayout == VK_IMAGE_LAYOUT_DEPTH_STENCIL_READ_ONLY_OPTIMAL))
 			{
 				return VK_IMAGE_LAYOUT_DEPTH_STENCIL_READ_ONLY_OPTIMAL;
 			}
diff --git a/Engine/Source/Runtime/VulkanRHI/Private/VulkanRenderpass.h b/Engine/Source/Runtime/VulkanRHI/Private/VulkanRenderpass.h
--- a/Engine/Source/Runtime/VulkanRHI/Private/VulkanRenderpass.h
+++ b/Engine/Source/Runtime/VulkanRHI/Private/VulkanRenderpass.h
@@ -470,10 +470,10 @@ public:
 				}
 				else
 				{
-					DepthStencilAttachment.layout = VK_IMAGE_LAYOUT_DEPTH_STENCIL_READ_ONLY_OPTIMAL;
-
+					// lights write to stencil for culling, so stencil is expected to be writebale while depth is read-only
+					DepthStencilAttachment.layout = VK_IMAGE_LAYOUT_DEPTH_READ_ONLY_STENCIL_ATTACHMENT_OPTIMAL;
 					DepthInputAttachment = DepthStencilAttachment.attachment;
-					DepthInputAttachmentLayout = VK_IMAGE_LAYOUT_DEPTH_STENCIL_READ_ONLY_OPTIMAL;
+					DepthInputAttachmentLayout = DepthStencilAttachment.layout;
 					DepthInputAspectMask = VK_IMAGE_ASPECT_DEPTH_BIT;
 				}
 			}
