VulkanRHI:  Fix build after bad merge.

#rnx

diff --git a/Engine/Source/Runtime/VulkanRHI/Private/VulkanTexture.cpp b/Engine/Source/Runtime/VulkanRHI/Private/VulkanTexture.cpp
--- a/Engine/Source/Runtime/VulkanRHI/Private/VulkanTexture.cpp
+++ b/Engine/Source/Runtime/VulkanRHI/Private/VulkanTexture.cpp
@@ -1568,7 +1568,11 @@ FVulkanTexture::FVulkanTexture(FRHICommandListBase* RHICmdList, FVulkanDevice& I
 		if (InitialLayout != VK_IMAGE_LAYOUT_UNDEFINED || bDoInitialClear)
 		{
 			if (RHICmdList && RHICmdList->IsTopOfPipe())
+			{
 				ALLOC_COMMAND_CL(*RHICmdList, FRHICommandSetInitialImageState)(this, InitialLayout, false, bDoInitialClear, InCreateDesc.ClearValue, bIsTransientResource);
+			}
+			else
+			{
 				RHICmdList = &FRHICommandListExecutor::GetImmediateCommandList();
 				if (!IsInRenderingThread() || (RHICmdList->Bypass() || !IsRunningRHIInSeparateThread()))
 				{
