Enable Vulkan LLM in Test configuration

diff --git a/Engine/Source/Runtime/VulkanRHI/Private/Android/VulkanAndroidPlatform.h b/Engine/Source/Runtime/VulkanRHI/Private/Android/VulkanAndroidPlatform.h
--- a/Engine/Source/Runtime/VulkanRHI/Private/Android/VulkanAndroidPlatform.h
+++ b/Engine/Source/Runtime/VulkanRHI/Private/Android/VulkanAndroidPlatform.h
@@ -11,7 +11,7 @@
 #define VULKAN_SHOULD_ENABLE_DRAW_MARKERS			(UE_BUILD_DEVELOPMENT || UE_BUILD_DEBUG)
 #define VULKAN_USE_IMAGE_ACQUIRE_FENCES				0
 #define VULKAN_USE_CREATE_ANDROID_SURFACE			1
-#define VULKAN_SHOULD_USE_LLM						(UE_BUILD_DEBUG || UE_BUILD_DEVELOPMENT)
+#define VULKAN_SHOULD_USE_LLM						(UE_BUILD_DEBUG || UE_BUILD_DEVELOPMENT || UE_BUILD_TEST)
 #define VULKAN_SHOULD_USE_COMMANDWRAPPERS			VULKAN_SHOULD_USE_LLM //LLM on Vulkan needs command wrappers to account for vkallocs
 #define VULKAN_ENABLE_LRU_CACHE						1
 #define VULKAN_SUPPORTS_GOOGLE_DISPLAY_TIMING		1
