Test for VK_EXT_pipeline_creation_cache_control Vulkan extension.

diff --git a/Engine/Source/Runtime/VulkanRHI/Private/VulkanDevice.h b/Engine/Source/Runtime/VulkanRHI/Private/VulkanDevice.h
--- a/Engine/Source/Runtime/VulkanRHI/Private/VulkanDevice.h
+++ b/Engine/Source/Runtime/VulkanRHI/Private/VulkanDevice.h
@@ -83,6 +83,7 @@ struct FOptionalVulkanDeviceExtensions
 			uint64 HasKHRMaintenance4 : 1;
 			uint64 HasKHRSynchronization2 : 1;
 			uint64 HasEXTSubgroupSizeControl : 1;
+			uint64 HasEXTPipelineCreationCacheControl : 1;
 		};
 		uint64 Packed;
 	};
diff --git a/Engine/Source/Runtime/VulkanRHI/Private/VulkanExtensions.cpp b/Engine/Source/Runtime/VulkanRHI/Private/VulkanExtensions.cpp
--- a/Engine/Source/Runtime/VulkanRHI/Private/VulkanExtensions.cpp
+++ b/Engine/Source/Runtime/VulkanRHI/Private/VulkanExtensions.cpp
@@ -1302,7 +1302,37 @@ public:
 	VkPhysicalDeviceShaderFloat16Int8Features  ShaderFloat16Int8Features;
 };
 
+// ***** VK_EXT_pipeline_creation_cache_control
+class FVulkanEXTPipelineCreationCacheControlExtension : public FVulkanDeviceExtension
+{
+public:
+
+	FVulkanEXTPipelineCreationCacheControlExtension(FVulkanDevice* InDevice)
+		: FVulkanDeviceExtension(InDevice, VK_EXT_PIPELINE_CREATION_CACHE_CONTROL_EXTENSION_NAME, VULKAN_EXTENSION_ENABLED, VK_API_VERSION_1_3)
+	{}
+	
+	virtual void PrePhysicalDeviceFeatures(VkPhysicalDeviceFeatures2KHR& PhysicalDeviceFeatures2) override final
+	{
+		ZeroVulkanStruct(PhysicalDevicePipelineCreationCacheControlFeatures, VK_STRUCTURE_TYPE_PHYSICAL_DEVICE_PIPELINE_CREATION_CACHE_CONTROL_FEATURES);
+		AddToPNext(PhysicalDeviceFeatures2, PhysicalDevicePipelineCreationCacheControlFeatures);
+	}
 
+	virtual void PostPhysicalDeviceFeatures(FOptionalVulkanDeviceExtensions& ExtensionFlags) override final
+	{
+		bRequirementsPassed = (PhysicalDevicePipelineCreationCacheControlFeatures.pipelineCreationCacheControl == VK_TRUE);
+		ExtensionFlags.HasEXTPipelineCreationCacheControl = bRequirementsPassed;
+	}
+
+	virtual void PreCreateDevice(VkDeviceCreateInfo& DeviceCreateInfo) override final
+	{
+		if (bRequirementsPassed)
+		{
+			AddToPNext(DeviceCreateInfo, PhysicalDevicePipelineCreationCacheControlFeatures);
+		}
+	}
+
+	VkPhysicalDevicePipelineCreationCacheControlFeatures PhysicalDevicePipelineCreationCacheControlFeatures;
+};
 
 
 template <typename ExtensionType>
@@ -1384,6 +1414,7 @@ FVulkanDeviceExtensionArray FVulkanDeviceExtension::GetUESupportedDeviceExtensio
 	ADD_CUSTOM_EXTENSION(FVulkanEXTShaderDemoteToHelperInvocationExtension);
 	ADD_CUSTOM_EXTENSION(FVulkanKHR16BitStorageExtension);
 	ADD_CUSTOM_EXTENSION(FVulkanKHRShaderFloat16Int8Extension);
+	ADD_CUSTOM_EXTENSION(FVulkanEXTPipelineCreationCacheControlExtension);
 
 	// Needed for Raytracing
 	ADD_CUSTOM_EXTENSION(FVulkanKHRBufferDeviceAddressExtension);
