Fixed "UniformView" setup issues in Vulkan and GL RHI

diff --git a/Engine/Source/Runtime/OpenGLDrv/Private/OpenGLUniformBuffer.cpp b/Engine/Source/Runtime/OpenGLDrv/Private/OpenGLUniformBuffer.cpp
--- a/Engine/Source/Runtime/OpenGLDrv/Private/OpenGLUniformBuffer.cpp
+++ b/Engine/Source/Runtime/OpenGLDrv/Private/OpenGLUniformBuffer.cpp
@@ -419,7 +419,7 @@ FOpenGLUniformBuffer::FOpenGLUniformBuffer(const FRHIUniformBufferLayout* InLayo
 	, bStreamDraw(false)
 	, bOwnsResource(true)
 {
-	bIsEmulatedUniformBuffer = GUseEmulatedUniformBuffers && !InLayout->bNoEmulatedUniformBuffer;
+	bIsEmulatedUniformBuffer = GUseEmulatedUniformBuffers && !(InLayout->bNoEmulatedUniformBuffer || InLayout->bUniformView);
 	RangeSize = InLayout->ConstantBufferSize;
 }
 
diff --git a/Engine/Source/Runtime/VulkanRHI/Private/VulkanUniformBuffer.cpp b/Engine/Source/Runtime/VulkanRHI/Private/VulkanUniformBuffer.cpp
--- a/Engine/Source/Runtime/VulkanRHI/Private/VulkanUniformBuffer.cpp
+++ b/Engine/Source/Runtime/VulkanRHI/Private/VulkanUniformBuffer.cpp
@@ -119,17 +119,23 @@ bool FVulkanUniformBuffer::SetupUniformBufferView(const FRHIUniformBufferLayout*
 		}
 
 		check(UniformViewSRV)
-				
-		FVulkanResourceMultiBuffer* Buffer = ResourceCast(UniformViewSRV->GetBuffer());
-		const FRHIViewDesc::FBufferSRV& SRVInfo = UniformViewSRV->GetDesc().Buffer.SRV;
+
+		FRHICommandListImmediate& RHICmdList = FRHICommandListExecutor::GetImmediateCommandList();
+		RHICmdList.EnqueueLambda([this, UniformViewSRV](FRHICommandList& CmdList)
+		{
+			FVulkanResourceMultiBuffer* Buffer = ResourceCast(UniformViewSRV->GetBuffer());
+			const FRHIViewDesc::FBufferSRV& SRVInfo = UniformViewSRV->GetDesc().Buffer.SRV;
+			Allocation.Reference(Buffer->GetCurrentAllocation());
+			check(Allocation.Size >= PLATFORM_MAX_UNIFORM_BUFFER_RANGE);
+			//Adjust Allocation.Size ???
+			Allocation.Offset += SRVInfo.OffsetInBytes;
+			bUniformView = true;
+		});
 		
-		Allocation.Reference(Buffer->GetCurrentAllocation());
-		//Adjust Allocation.Size ???
-		Allocation.Offset += SRVInfo.OffsetInBytes;
-		bUniformView = true;
+		return true;
 	}
 	
-	return bUniformView;
+	return false;
 }
 
 FVulkanUniformBuffer::FVulkanUniformBuffer(FVulkanDevice& InDevice, const FRHIUniformBufferLayout* InLayout, const void* Contents, EUniformBufferUsage InUsage, EUniformBufferValidation Validation)
