Spport for generic LLM in Vulkan when Vulkan-specific one is turned off

diff --git a/Engine/Source/Runtime/VulkanRHI/Private/Android/VulkanAndroidPlatform.h b/Engine/Source/Runtime/VulkanRHI/Private/Android/VulkanAndroidPlatform.h
--- a/Engine/Source/Runtime/VulkanRHI/Private/Android/VulkanAndroidPlatform.h
+++ b/Engine/Source/Runtime/VulkanRHI/Private/Android/VulkanAndroidPlatform.h
@@ -11,7 +11,7 @@
 #define VULKAN_SHOULD_ENABLE_DRAW_MARKERS			(UE_BUILD_DEVELOPMENT || UE_BUILD_DEBUG)
 #define VULKAN_USE_IMAGE_ACQUIRE_FENCES				0
 #define VULKAN_USE_CREATE_ANDROID_SURFACE			1
-#define VULKAN_SHOULD_USE_LLM						(UE_BUILD_DEBUG || UE_BUILD_DEVELOPMENT || UE_BUILD_TEST)
+#define VULKAN_SHOULD_USE_LLM						(UE_BUILD_DEBUG || UE_BUILD_DEVELOPMENT) // If enabled Vulkan will report detailed allocation statistics, overrides some tags with custom ones
 #define VULKAN_SHOULD_USE_COMMANDWRAPPERS			VULKAN_SHOULD_USE_LLM //LLM on Vulkan needs command wrappers to account for vkallocs
 #define VULKAN_ENABLE_LRU_CACHE						1
 #define VULKAN_SUPPORTS_GOOGLE_DISPLAY_TIMING		1
diff --git a/Engine/Source/Runtime/VulkanRHI/Private/VulkanBuffer.cpp b/Engine/Source/Runtime/VulkanRHI/Private/VulkanBuffer.cpp
--- a/Engine/Source/Runtime/VulkanRHI/Private/VulkanBuffer.cpp
+++ b/Engine/Source/Runtime/VulkanRHI/Private/VulkanBuffer.cpp
@@ -609,7 +609,11 @@ void FVulkanResourceMultiBuffer::ReleaseOwnership()
 
 FBufferRHIRef FVulkanDynamicRHI::RHICreateBuffer(FRHICommandListBase& RHICmdList, FRHIBufferDesc const& Desc, ERHIAccess ResourceState, FRHIResourceCreateInfo& CreateInfo)
 {
+#if VULKAN_USE_LLM
 	LLM_SCOPE_VULKAN(ELLMTagVulkan::VulkanBuffers);
+#else
+	LLM_SCOPE(EnumHasAnyFlags(Desc.Usage, EBufferUsageFlags::VertexBuffer | EBufferUsageFlags::IndexBuffer) ? ELLMTag::Meshes : ELLMTag::RHIMisc);
+#endif
 	return new FVulkanResourceMultiBuffer(Device, Desc, CreateInfo, &RHICmdList);
 }
 
diff --git a/Engine/Source/Runtime/VulkanRHI/Private/VulkanLLM.h b/Engine/Source/Runtime/VulkanRHI/Private/VulkanLLM.h
--- a/Engine/Source/Runtime/VulkanRHI/Private/VulkanLLM.h
+++ b/Engine/Source/Runtime/VulkanRHI/Private/VulkanLLM.h
@@ -9,16 +9,11 @@
 
 #include "HAL/LowLevelMemTracker.h"
 
-#if VULKAN_USE_LLM
+#if ENABLE_LOW_LEVEL_MEM_TRACKER
 
-#define LLM_SCOPE_VULKAN(Tag)			LLM_SCOPE((ELLMTag)Tag)
-#define LLM_PLATFORM_SCOPE_VULKAN(Tag)	LLM_PLATFORM_SCOPE((ELLMTag)Tag)
-extern uint64 GVulkanLLMAllocationID;
-#define LLM_TRACK_VULKAN_HIGH_LEVEL_ALLOC(AllocObj, Size) { AllocObj.SetLLMTrackerID(0xDEAD | (++GVulkanLLMAllocationID << 16)); FLowLevelMemTracker::Get().OnLowLevelAlloc(ELLMTracker::Default, (void*)AllocObj.GetLLMTrackerID(), Size, (ELLMTag)ELLMTagVulkan::VulkanMisc); }
-#define LLM_TRACK_VULKAN_HIGH_LEVEL_FREE(AllocObj) { FLowLevelMemTracker::Get().OnLowLevelFree(ELLMTracker::Default, (void*)AllocObj.GetLLMTrackerID()); }
-#define LLM_TRACK_VULKAN_SPARE_MEMORY_GPU(Size)	{ LLM_SCOPED_PAUSE_TRACKING_WITH_ENUM_AND_AMOUNT((ELLMTag)ELLMTagVulkan::VulkanSpareMemoryGPU, Size, ELLMTracker::Default, ELLMAllocType::None); }
+#if VULKAN_USE_LLM
 
-#define LLM_TRACK_VULKAN_HIGH_LEVEL_ALLOCATION
+// Add custom Tags
 enum class ELLMTagVulkan : LLM_TAG_TYPE
 {
 	VulkanMisc = (LLM_TAG_TYPE)ELLMTag::PlatformTagStart,
@@ -35,6 +30,13 @@ enum class ELLMTagVulkan : LLM_TAG_TYPE
 	Count,
 };
 
+extern uint64 GVulkanLLMAllocationID;
+#define LLM_TRACK_VULKAN_HIGH_LEVEL_ALLOC(AllocObj, Size) { AllocObj.SetLLMTrackerID(0xDEAD | (++GVulkanLLMAllocationID << 16)); FLowLevelMemTracker::Get().OnLowLevelAlloc(ELLMTracker::Default, (void*)AllocObj.GetLLMTrackerID(), Size, (ELLMTag)ELLMTagVulkan::VulkanMisc); }
+#define LLM_TRACK_VULKAN_HIGH_LEVEL_FREE(AllocObj) { FLowLevelMemTracker::Get().OnLowLevelFree(ELLMTracker::Default, (void*)AllocObj.GetLLMTrackerID()); }
+#define LLM_TRACK_VULKAN_SPARE_MEMORY_GPU(Size)	{ LLM_SCOPED_PAUSE_TRACKING_WITH_ENUM_AND_AMOUNT((ELLMTag)ELLMTagVulkan::VulkanSpareMemoryGPU, Size, ELLMTracker::Default, ELLMAllocType::None); }
+
+#define LLM_TRACK_VULKAN_HIGH_LEVEL_ALLOCATION
+
 static_assert((int32)ELLMTagVulkan::Count <= (int32)ELLMTag::PlatformTagEnd, "too many ELLMTagVulkan tags");
 
 namespace VulkanLLM
@@ -42,7 +44,36 @@ namespace VulkanLLM
 	void Initialize();
 }
 
-#else
+#else		// #if VULKAN_USE_LLM
+
+// Translate VulkanTag to regular Tag
+enum class ELLMTagVulkan : LLM_TAG_TYPE
+{
+	VulkanMisc				= (LLM_TAG_TYPE)ELLMTag::RHIMisc,
+	VulkanUniformBuffers	= (LLM_TAG_TYPE)ELLMTag::RHIMisc,
+	VulkanBuffers			= (LLM_TAG_TYPE)ELLMTag::RHIMisc,
+	VulkanTextures			= (LLM_TAG_TYPE)ELLMTag::Textures,
+	VulkanRenderTargets		= (LLM_TAG_TYPE)ELLMTag::RHIMisc,
+	VulkanSpareMemoryGPU	= (LLM_TAG_TYPE)ELLMTag::RHIMisc,
+	VulkanShaders			= (LLM_TAG_TYPE)ELLMTag::Shaders,
+	VulkanFrameTemp			= (LLM_TAG_TYPE)ELLMTag::RHIMisc,
+	VulkanStagingBuffers	= (LLM_TAG_TYPE)ELLMTag::RHIMisc,
+	VulkanDriverMemoryCPU	= (LLM_TAG_TYPE)ELLMTag::RHIMisc,
+	VulkanDriverMemoryGPU	= (LLM_TAG_TYPE)ELLMTag::RHIMisc,
+	Count					= (LLM_TAG_TYPE)ELLMTag::PlatformTagStart + 11,
+};
+
+#define LLM_TRACK_VULKAN_HIGH_LEVEL_ALLOC(...)
+#define LLM_TRACK_VULKAN_HIGH_LEVEL_FREE(...)
+#define LLM_TRACK_VULKAN_SPARE_MEMORY_GPU(...)
+
+#endif		// #if VULKAN_USE_LLM
+
+#define LLM_SCOPE_VULKAN(Tag)			LLM_SCOPE((ELLMTag)Tag)
+#define LLM_PLATFORM_SCOPE_VULKAN(Tag)	LLM_PLATFORM_SCOPE((ELLMTag)Tag)
+static_assert((int32)ELLMTagVulkan::Count == (int32)ELLMTag::PlatformTagStart + 11, "There needs to be a 1 to 1 mapping between VulkanTag and Translation");
+
+#else		// #if ENABLE_LOW_LEVEL_MEM_TRACKER
 
 #define LLM_SCOPE_VULKAN(...)
 #define LLM_PLATFORM_SCOPE_VULKAN(...)
@@ -50,4 +81,4 @@ namespace VulkanLLM
 #define LLM_TRACK_VULKAN_HIGH_LEVEL_FREE(...)
 #define LLM_TRACK_VULKAN_SPARE_MEMORY_GPU(...)
 
-#endif		// #if VULKAN_USE_LLM
+#endif		// #if ENABLE_LOW_LEVEL_MEM_TRACKER
\ No newline at end of file
diff --git a/Engine/Source/Runtime/VulkanRHI/Private/VulkanMemory.cpp b/Engine/Source/Runtime/VulkanRHI/Private/VulkanMemory.cpp
--- a/Engine/Source/Runtime/VulkanRHI/Private/VulkanMemory.cpp
+++ b/Engine/Source/Runtime/VulkanRHI/Private/VulkanMemory.cpp
@@ -961,6 +961,9 @@ namespace VulkanRHI
 		LLM_PLATFORM_SCOPE_VULKAN(ELLMTagVulkan::VulkanDriverMemoryGPU);
 		LLM(FLowLevelMemTracker::Get().OnLowLevelAlloc(ELLMTracker::Platform, (void*)NewAllocation->Handle, AllocationSize, ELLMTag::GraphicsPlatform, ELLMAllocType::System));
 		LLM_TRACK_VULKAN_SPARE_MEMORY_GPU((int64)AllocationSize);
+#else
+		LLM(FLowLevelMemTracker::Get().OnLowLevelAlloc(ELLMTracker::Platform, (void*)NewAllocation->Handle, AllocationSize, ELLMTag::GraphicsPlatform));
+		LLM(FLowLevelMemTracker::Get().OnLowLevelAlloc(ELLMTracker::Default, (void*)NewAllocation->Handle, AllocationSize, ELLMTag::Untagged));
 #endif
 
 		INC_DWORD_STAT(STAT_VulkanNumPhysicalMemAllocations);
@@ -1041,6 +1044,9 @@ namespace VulkanRHI
 #if VULKAN_USE_LLM
 		LLM(FLowLevelMemTracker::Get().OnLowLevelFree(ELLMTracker::Platform, (void*)Allocation->Handle, ELLMAllocType::System));
 		LLM_TRACK_VULKAN_SPARE_MEMORY_GPU(-(int64)Allocation->Size);
+#else
+		LLM(FLowLevelMemTracker::Get().OnLowLevelFree(ELLMTracker::Platform, (void*)Allocation->Handle));
+		LLM(FLowLevelMemTracker::Get().OnLowLevelFree(ELLMTracker::Default, (void*)Allocation->Handle));
 #endif
 
 		--NumAllocations;
@@ -5045,6 +5051,8 @@ void* FVulkanCustomMemManager::Alloc(void* UserData, size_t Size, size_t Alignme
 {
 	check(AllocScope < VK_SYSTEM_ALLOCATION_SCOPE_RANGE_SIZE);
 	LLM_SCOPE_VULKAN(ELLMTagVulkan::VulkanDriverMemoryCPU);
+	LLM_PLATFORM_SCOPE(ELLMTag::GraphicsPlatform);
+
 	FScopeLock Lock(&GMemMgrCS);
 	void* Data = FMemory::Malloc(Size, Alignment);
 	FType& Type = GetType(UserData, AllocScope);
@@ -5057,6 +5065,8 @@ void* FVulkanCustomMemManager::Alloc(void* UserData, size_t Size, size_t Alignme
 void FVulkanCustomMemManager::Free(void* UserData, void* Mem)
 {
 	LLM_SCOPE_VULKAN(ELLMTagVulkan::VulkanDriverMemoryCPU);
+	LLM_PLATFORM_SCOPE(ELLMTag::GraphicsPlatform);
+
 	FScopeLock Lock(&GMemMgrCS);
 	FMemory::Free(Mem);
 	for (int32 Index = 0; Index < GVulkanInstrumentedMemMgr.Types.Num(); ++Index)
@@ -5075,6 +5085,8 @@ void* FVulkanCustomMemManager::Realloc(void* UserData, void* Original, size_t Si
 {
 	check(AllocScope < VK_SYSTEM_ALLOCATION_SCOPE_RANGE_SIZE);
 	LLM_SCOPE_VULKAN(ELLMTagVulkan::VulkanDriverMemoryCPU);
+	LLM_PLATFORM_SCOPE(ELLMTag::GraphicsPlatform);
+
 	FScopeLock Lock(&GMemMgrCS);
 	void* Data = FMemory::Realloc(Original, Size, Alignment);
 	FType& Type = GetType(UserData, AllocScope);
diff --git a/Engine/Source/Runtime/VulkanRHI/Private/VulkanTexture.cpp b/Engine/Source/Runtime/VulkanRHI/Private/VulkanTexture.cpp
--- a/Engine/Source/Runtime/VulkanRHI/Private/VulkanTexture.cpp
+++ b/Engine/Source/Runtime/VulkanRHI/Private/VulkanTexture.cpp
@@ -52,13 +52,13 @@ struct FTextureLock
 	}
 };
 
-#if VULKAN_USE_LLM
+#if ENABLE_LOW_LEVEL_MEM_TRACKER
 inline ELLMTagVulkan GetMemoryTagForTextureFlags(ETextureCreateFlags UEFlags)
 {
 	bool bRenderTarget = EnumHasAnyFlags(UEFlags, TexCreate_RenderTargetable | TexCreate_ResolveTargetable | TexCreate_DepthStencilTargetable);
 	return bRenderTarget ? ELLMTagVulkan::VulkanRenderTargets : ELLMTagVulkan::VulkanTextures;
 }
-#endif // VULKAN_USE_LLM
+#endif // ENABLE_LOW_LEVEL_MEM_TRACKER
 
 inline bool operator == (const FTextureLock& A, const FTextureLock& B)
 {
