Ensure vulkan's FVulkanPipelineStateCacheManager::RHICreateGraphicsPipelineState return value is ref counted.

diff --git a/Engine/Source/Runtime/VulkanRHI/Private/VulkanPipeline.cpp b/Engine/Source/Runtime/VulkanRHI/Private/VulkanPipeline.cpp
--- a/Engine/Source/Runtime/VulkanRHI/Private/VulkanPipeline.cpp
+++ b/Engine/Source/Runtime/VulkanRHI/Private/VulkanPipeline.cpp
@@ -1933,7 +1933,7 @@ struct FPSOOptionalLock
 
 static FCriticalSection CreateGraphicsPSOMutex;
 
-FVulkanRHIGraphicsPipelineState* FVulkanPipelineStateCacheManager::RHICreateGraphicsPipelineState(const FGraphicsPipelineStateInitializer& Initializer)
+FGraphicsPipelineStateRHIRef FVulkanPipelineStateCacheManager::RHICreateGraphicsPipelineState(const FGraphicsPipelineStateInitializer& Initializer)
 {
 	QUICK_SCOPE_CYCLE_COUNTER(STAT_Vulkan_RHICreateGraphicsPipelineState_NEW);
 
diff --git a/Engine/Source/Runtime/VulkanRHI/Private/VulkanPipeline.h b/Engine/Source/Runtime/VulkanRHI/Private/VulkanPipeline.h
--- a/Engine/Source/Runtime/VulkanRHI/Private/VulkanPipeline.h
+++ b/Engine/Source/Runtime/VulkanRHI/Private/VulkanPipeline.h
@@ -518,7 +518,7 @@ private:
 	void SavePSOCache(const FString& CacheFilename, FPipelineCache& Cache);
 	void DestroyCache();
 
-	FVulkanRHIGraphicsPipelineState* RHICreateGraphicsPipelineState(const FGraphicsPipelineStateInitializer& Initializer);
+	FGraphicsPipelineStateRHIRef RHICreateGraphicsPipelineState(const FGraphicsPipelineStateInitializer& Initializer);
 	FVulkanComputePipeline* RHICreateComputePipelineState(FRHIComputeShader* ComputeShaderRHI);
 	void NotifyDeletedGraphicsPSO(FRHIGraphicsPipelineState* PSO);
 	bool CreateGfxPipelineFromEntry(FVulkanRHIGraphicsPipelineState* PSO, FVulkanShader* Shaders[ShaderStage::NumStages], bool bPrecompile);
