VulkanRHI: Track refs to UB used in ShaderTable.  Add transitions for all resources in the UB ResourceTable.  Fix issue with TMaterialCHS and TPathTracingMaterial where SPIRV doesn't not detect and remove the dead code where an InstanceCulling UB is used (was previously patched in CS using USE_INSTANCE_CULLING_DATA to forcefully remove the buffer), plug in a dummy buffer at binding time for now.

#rnx

diff --git a/Engine/Source/Runtime/VulkanRHI/Private/VulkanRayTracing.cpp b/Engine/Source/Runtime/VulkanRHI/Private/VulkanRayTracing.cpp
--- a/Engine/Source/Runtime/VulkanRHI/Private/VulkanRayTracing.cpp
+++ b/Engine/Source/Runtime/VulkanRHI/Private/VulkanRayTracing.cpp
@@ -538,7 +538,8 @@ void FVulkanRayTracingGeometry::SetupHitGroupSystemParameters()
 		const FRHIDescriptorHandle VBHandle = GetBindlessHandle(VertexBuffer, Segment.VertexBufferOffset);
 		HitGroupSystemVertexViews.Add(VBHandle);
 
-		FVulkanHitGroupSystemParameters SystemParameters = {};
+
+		FVulkanHitGroupSystemParameters& SystemParameters = HitGroupSystemParameters.AddZeroed_GetRef();
 		SystemParameters.RootConstants.SetVertexAndIndexStride(Segment.VertexBufferStride, IndexStride);
 		SystemParameters.BindlessHitGroupSystemVertexBuffer = VBHandle.GetIndex();
 
@@ -548,8 +549,6 @@ void FVulkanRayTracingGeometry::SetupHitGroupSystemParameters()
 			SystemParameters.RootConstants.IndexBufferOffsetInBytes = Initializer.IndexBufferOffset + IndexStride * Segment.FirstPrimitive * FVulkanRayTracingGeometry::IndicesPerPrimitive;
 			SystemParameters.RootConstants.FirstPrimitive = Segment.FirstPrimitive;
 		}
-
-		HitGroupSystemParameters.Add(SystemParameters);
 	}
 }
 
@@ -1294,8 +1293,40 @@ void FVulkanDevice::InitializeRayTracing()
 {
 }
 
+// Temporary code to generate dummy UBs to bind when none is provided to prevent bindless code from crashing
+// NOTE: Should currently only be used by InstanceCulling due to a binding that isn't stripped by DXC. See also USE_INSTANCE_CULLING_DATA for same issue in CS.
+static FRWLock DummyUBLock;
+static TMap<uint32, FUniformBufferRHIRef> DummyUBs;
+static FVulkanUniformBuffer* GetDummyUB(FVulkanDevice* Device, uint32 UBLayoutHash)
+{
+	{
+		FRWScopeLock ScopedReadLock(DummyUBLock, SLT_ReadOnly);
+		FUniformBufferRHIRef* UBRef = DummyUBs.Find(UBLayoutHash);
+		if (UBRef)
+		{
+			return ResourceCast(UBRef->GetReference());
+		}
+	}
+
+	FRWScopeLock ScopedReadLock(DummyUBLock, SLT_Write);
+	const FShaderParametersMetadata* DummyMetadata = FindUniformBufferStructByLayoutHash(UBLayoutHash);
+	if (DummyMetadata && DummyMetadata->GetLayoutPtr())
+	{
+		const FRHIUniformBufferLayout* DummyLayout = DummyMetadata->GetLayoutPtr();
+		TArray<uint8> DummyContent;
+		DummyContent.SetNumZeroed(DummyLayout->ConstantBufferSize);
+		FVulkanUniformBuffer* DummyUB = new FVulkanUniformBuffer(*Device, DummyLayout, DummyContent.GetData(), UniformBuffer_MultiFrame, EUniformBufferValidation::None);
+		DummyUBs.Add(UBLayoutHash, DummyUB);
+		const FString& LayoutName = DummyLayout->GetDebugName();
+		UE_LOG(LogRHI, Warning, TEXT("Vulkan ray tracing using DummyUB for %s."), LayoutName.IsEmpty() ? TEXT("<unknown>") : *LayoutName);
+		return DummyUB;
+	}
+	return nullptr;
+}
+
 void FVulkanDevice::CleanUpRayTracing()
 {
+	DummyUBs.Empty();
 }
 
 
@@ -1770,8 +1801,11 @@ void FVulkanRayTracingCompactionRequestHandler::Update(FVulkanCommandListContext
 
 
 
-void SetRayGenResources(FVulkanDevice* Device, FVulkanCmdBuffer* const CmdBuffer, const FRayTracingShaderBindings& InGlobalResourceBindings)
+static FVulkanPipelineBarrier SetRayGenResources(FVulkanDevice* Device, FVulkanCmdBuffer* const CmdBuffer, const FRayTracingShaderBindings& InGlobalResourceBindings, FVulkanRayTracingShaderTable* ShaderTable)
 {
+	TArray<const FVulkanUniformBuffer*> UniformBuffers;
+	UniformBuffers.Reserve(UE_ARRAY_COUNT(InGlobalResourceBindings.UniformBuffers));
+
 	// Uniform buffers
 	{
 		uint32 NumSkippedSlots = 0;
@@ -1799,6 +1833,8 @@ void SetRayGenResources(FVulkanDevice* Device, FVulkanCmdBuffer* const CmdBuffer
 				DescriptorAddressInfo.sType = VK_STRUCTURE_TYPE_BUFFER_DEVICE_ADDRESS_INFO;
 				DescriptorAddressInfo.address = UniformBuffer->GetDeviceAddress();
 				DescriptorAddressInfo.range = UniformBuffer->GetSize();
+
+				UniformBuffers.AddUnique(UniformBuffer);
 			}
 			else
 			{
@@ -1808,8 +1844,157 @@ void SetRayGenResources(FVulkanDevice* Device, FVulkanCmdBuffer* const CmdBuffer
 		}
 		Device->GetBindlessDescriptorManager()->RegisterUniformBuffers(CmdBuffer->GetHandle(), VK_PIPELINE_BIND_POINT_RAY_TRACING_KHR, StageUBs);
 	}
-}
 
+	// Add all the UBs references by the shader table
+	TArrayView<TRefCountPtr<FRHIUniformBuffer>> ShaderTableUBs = ShaderTable->GetUBRefs();
+	for (TRefCountPtr<FRHIUniformBuffer>& UniformBuffer : ShaderTableUBs)
+	{
+		const FVulkanUniformBuffer* VulkanUniformBuffer = ResourceCast(UniformBuffer.GetReference());
+		UniformBuffers.AddUnique(VulkanUniformBuffer);
+	}
+
+	// Track all the missing transitions for the dispatch to be able to bring it back afterwards (will not touch tracking)
+	FVulkanPipelineBarrier PreDispatch, PostDispatch;
+	{
+		auto TransitionBuffer = [&PreDispatch, &PostDispatch](bool bReadOnly)
+		{
+			// :todo-jn: tighten these barriers
+			const VkAccessFlags RWAccessFlags = VK_ACCESS_MEMORY_READ_BIT | VK_ACCESS_MEMORY_WRITE_BIT;
+			const VkAccessFlags DesiredAccessFlags = bReadOnly ? VK_ACCESS_MEMORY_READ_BIT : VK_ACCESS_MEMORY_READ_BIT | VK_ACCESS_MEMORY_WRITE_BIT;
+			PreDispatch.AddMemoryBarrier(RWAccessFlags, DesiredAccessFlags, VK_PIPELINE_STAGE_ALL_COMMANDS_BIT, VK_PIPELINE_STAGE_RAY_TRACING_SHADER_BIT_KHR);
+			PostDispatch.AddMemoryBarrier(DesiredAccessFlags, RWAccessFlags, VK_PIPELINE_STAGE_RAY_TRACING_SHADER_BIT_KHR, VK_PIPELINE_STAGE_ALL_COMMANDS_BIT);
+		};
+
+		// Make sure we only transition textures once, accumulate them in sets
+		TSet<FRHITexture*> SRVTransitions;
+		TSet<FRHITexture*> UAVTransitions;
+		for (const FVulkanUniformBuffer* UniformBuffer : UniformBuffers)
+		{
+			const TArray<TRefCountPtr<FRHIResource>>& ResourceTable = UniformBuffer->GetResourceTable();
+			for (const TRefCountPtr<FRHIResource>& RHIResourceRef : ResourceTable)
+			{
+				const FRHIResource* RHIResource = RHIResourceRef.GetReference();
+				if (!RHIResource)
+				{
+					continue;
+				}
+
+				switch (RHIResource->GetType())
+				{
+				case RRT_Texture:
+				case RRT_Texture2D:
+				case RRT_Texture2DArray:
+				case RRT_Texture3D:
+				case RRT_TextureCube:
+					SRVTransitions.Add((FRHITexture*)RHIResource);
+					break;
+
+				case RRT_TextureReference:
+					SRVTransitions.Add(((FRHITextureReference*)RHIResource)->GetReferencedTexture());
+					break;
+
+				case RRT_UnorderedAccessView:
+				{
+					const FRHIUnorderedAccessView* RHIUnorderedAccessView = (FRHIUnorderedAccessView*)RHIResource;
+					if (RHIUnorderedAccessView->IsTexture())
+					{
+						UAVTransitions.Add(RHIUnorderedAccessView->GetTexture());
+					}
+					else
+					{
+						TransitionBuffer(false);
+					}
+					break;
+				}
+
+				case RRT_ShaderResourceView:
+				{
+					const FRHIShaderResourceView* RHIShaderResourceView = (FRHIShaderResourceView*)RHIResource;
+					if (RHIShaderResourceView->IsTexture())
+					{
+						SRVTransitions.Add(RHIShaderResourceView->GetTexture());
+					}
+					else
+					{
+						TransitionBuffer(true);
+					}
+					break;
+				}
+
+				case RRT_RayTracingAccelerationStructure:
+				case RRT_StagingBuffer:
+				case RRT_Buffer:
+					TransitionBuffer(true);
+					break;
+
+				case RRT_SamplerState: [[fallthrough]];
+				default:
+					// Do nothing
+					break;
+				};
+			}
+		}
+
+		auto TransitionTexture = [&PreDispatch, &PostDispatch, CmdBuffer](FRHITexture* RHITexture, bool bReadOnly)
+		{
+			const FVulkanTexture* Texture = ResourceCast(RHITexture);
+
+			// Because Sync2 is a prereq to ray tracing, use the conveniently generic layout
+			const VkImageLayout TargetLayout = bReadOnly ? VK_IMAGE_LAYOUT_READ_ONLY_OPTIMAL : VK_IMAGE_LAYOUT_GENERAL;
+			const FVulkanImageLayout* OriginalLayout = CmdBuffer->GetLayoutManager().GetFullLayout(Texture->Image);
+			check(OriginalLayout);
+
+			// If all the subresource are already in a correct layout for the desired RendOnly state, then skip the barrier
+			if (!OriginalLayout->AreAllSubresourcesSameLayout() ||
+				(
+					(bReadOnly && (OriginalLayout->MainLayout != VK_IMAGE_LAYOUT_READ_ONLY_OPTIMAL) && (OriginalLayout->MainLayout != VK_IMAGE_LAYOUT_SHADER_READ_ONLY_OPTIMAL))
+					|| (!bReadOnly && (OriginalLayout->MainLayout != VK_IMAGE_LAYOUT_GENERAL)) // :todo-jn: prevent overlap?
+					))
+			{
+				PreDispatch.AddImageLayoutTransition(Texture->Image, Texture->GetFullAspectMask(), *OriginalLayout, TargetLayout);
+
+				// Transition back to where it was, leaving any undefined transitions to whatever we set them to
+				{
+					FVulkanImageLayout FinalLayout = *OriginalLayout;
+					if (FinalLayout.AreAllSubresourcesSameLayout())
+					{
+						if (FinalLayout.MainLayout == VK_IMAGE_LAYOUT_UNDEFINED)
+						{
+							FinalLayout.MainLayout = TargetLayout;
+						}
+					}
+					else
+					{
+						for (int32 SubResIndex = 0; SubResIndex < FinalLayout.SubresLayouts.Num(); ++SubResIndex)
+						{
+							if (FinalLayout.SubresLayouts[SubResIndex] == VK_IMAGE_LAYOUT_UNDEFINED)
+							{
+								FinalLayout.SubresLayouts[SubResIndex] = TargetLayout;
+							}
+						}
+					}
+					PostDispatch.AddImageLayoutTransition(Texture->Image, Texture->GetFullAspectMask(), TargetLayout, FinalLayout);
+				}
+			}
+		};
+
+		for (FRHITexture* RHITexture : UAVTransitions)
+		{
+			TransitionTexture(RHITexture, false);
+
+			// If a resource shows up as both, use it in VK_IMAGE_LAYOUT_GENERAL
+			SRVTransitions.Remove(RHITexture);
+		}
+
+		for (FRHITexture* RHITexture : SRVTransitions)
+		{
+			TransitionTexture(RHITexture, true);
+		}
+	}
+
+	PreDispatch.Execute(CmdBuffer);
+	return PostDispatch;
+}
 
 void FVulkanCommandListContext::RHIRayTraceDispatch(
 	FRHIRayTracingPipelineState* InRayTracingPipelineState, 
@@ -1829,7 +2014,7 @@ void FVulkanCommandListContext::RHIRayTraceDispatch(
 	ShaderTable->SetSlot(InRayGenShader->GetFrequency(), 0, Pipeline->GetShaderIndex(RayGenShader), Pipeline->GetShaderHandles(SF_RayGen));
 	ShaderTable->Commit(*this);
 
-	SetRayGenResources(Device, CmdBuffer, InGlobalResourceBindings);
+	FVulkanPipelineBarrier PostDispatch = SetRayGenResources(Device, CmdBuffer, InGlobalResourceBindings, ShaderTable);
 
 	VulkanRHI::vkCmdTraceRaysKHR(
 		CmdBuffer->GetHandle(),
@@ -1839,6 +2024,8 @@ void FVulkanCommandListContext::RHIRayTraceDispatch(
 		ShaderTable->GetRegion(SF_RayCallable),
 		InWidth, InHeight, 1);
 
+	PostDispatch.Execute(CmdBuffer);
+
 	if (GVulkanSubmitOnTraceRays)
 	{
 		InternalSubmitActiveCmdBuffer();
@@ -1865,7 +2052,7 @@ void FVulkanCommandListContext::RHIRayTraceDispatchIndirect(
 	ShaderTable->SetSlot(InRayGenShader->GetFrequency(), 0, Pipeline->GetShaderIndex(RayGenShader), Pipeline->GetShaderHandles(SF_RayGen));
 	ShaderTable->Commit(*this);
 
-	SetRayGenResources(Device, CmdBuffer, InGlobalResourceBindings);
+	FVulkanPipelineBarrier PostDispatch = SetRayGenResources(Device, CmdBuffer, InGlobalResourceBindings, ShaderTable);
 
 	FVulkanResourceMultiBuffer* ArgumentBuffer = ResourceCast(InArgumentBuffer);
 	const VkDeviceAddress IndirectDeviceAddress = ArgumentBuffer->GetDeviceAddress() + InArgumentOffset;
@@ -1878,29 +2065,42 @@ void FVulkanCommandListContext::RHIRayTraceDispatchIndirect(
 		ShaderTable->GetRegion(SF_RayCallable),
 		IndirectDeviceAddress);
 
+	PostDispatch.Execute(CmdBuffer);
+
 	if (GVulkanSubmitOnTraceRays)
 	{
 		InternalSubmitActiveCmdBuffer();
 	}
 }
 
-static void SetSystemParametersUB(FVulkanHitGroupSystemParameters& OutSystemParameters, uint32 InNumUniformBuffers, FRHIUniformBuffer* const* InUniformBuffers, const FVulkanRayTracingShader* InShader)
+
+
+static void SetSystemParametersUB(FVulkanHitGroupSystemParameters& OutSystemParameters, FVulkanDevice* Device, FVulkanRayTracingShaderTable* ShaderTable, uint32 InNumUniformBuffers, FRHIUniformBuffer* const* InUniformBuffers, const FVulkanRayTracingShader* InShader)
 {
 	// Plug the shaders in the right slots using LayoutHash comparisons
-	check(InShader->GetCodeHeader().UniformBuffers.Num() == InNumUniformBuffers);
-	for (uint32 UBIndex = 0; UBIndex < InNumUniformBuffers; ++UBIndex)
+	check(InShader->GetCodeHeader().UniformBuffers.Num() <= (int32)InNumUniformBuffers);
+	for (int32 UBIndex = 0; UBIndex < InShader->GetCodeHeader().UniformBuffers.Num(); ++UBIndex)
 	{
 		FVulkanUniformBuffer* UniformBuffer = ResourceCast(InUniformBuffers[UBIndex]);
-		check(UniformBuffer);
 
 		const FVulkanShaderHeader::FUniformBufferInfo& UniformBufferInfo = InShader->GetCodeHeader().UniformBuffers[UBIndex];
+
+		// :todo-jn: Hack to force in a DummyCullingBuffer in cases where it should have been culled from source (see SPIRV-Tools Issue 4902).
+		if (!UniformBuffer)
+		{
+			UniformBuffer = GetDummyUB(Device, UniformBufferInfo.LayoutHash);
+		}
+
+		check(UniformBuffer);
+
 		check((UniformBufferInfo.LayoutHash == 0) || (UniformBufferInfo.LayoutHash == UniformBuffer->GetLayout().GetHash()));
 		check(UniformBufferInfo.ConstantDataOriginalBindingIndex != UINT16_MAX);
 
-
 		const FRHIDescriptorHandle BindlessHandle = UniformBuffer->GetBindlessHandle();
 		check(BindlessHandle.IsValid());
 		OutSystemParameters.BindlessUniformBuffers[UniformBufferInfo.ConstantDataOriginalBindingIndex] = BindlessHandle.GetIndex();
+
+		ShaderTable->AddUBRef(UniformBuffer);
 	}
 }
 
@@ -1942,7 +2142,7 @@ static void SetRayTracingHitGroup(
 	FVulkanHitGroupSystemParameters SystemParameters = Geometry->HitGroupSystemParameters[SegmentIndex];
 	SystemParameters.RootConstants.BaseInstanceIndex = SceneInitializer.BaseInstancePrefixSum[InstanceIndex];
 	SystemParameters.RootConstants.UserData = UserData;
-	SetSystemParametersUB(SystemParameters, NumUniformBuffers, UniformBuffers, Shader);
+	SetSystemParametersUB(SystemParameters, Device, ShaderTable, NumUniformBuffers, UniformBuffers, Shader);
 
 	ShaderTable->SetLocalShaderParameters(SF_RayHitGroup, RecordIndex, 0, SystemParameters);
 
@@ -1995,7 +2195,7 @@ static void SetGenericSystemParameters(
 	FVulkanHitGroupSystemParameters SystemParameters;
 	FMemory::Memzero(SystemParameters);
 	SystemParameters.RootConstants.UserData = UserData;
-	SetSystemParametersUB(SystemParameters, NumUniformBuffers, UniformBuffers, Shader);
+	SetSystemParametersUB(SystemParameters, Scene->GetParent(), ShaderTable, NumUniformBuffers, UniformBuffers, Shader);
 	ShaderTable->SetLocalShaderParameters(ShaderFrequency, ShaderSlotInScene, 0, SystemParameters);
 
 	ShaderTable->SetSlot(ShaderFrequency, ShaderSlotInScene, ShaderIndexInPipeline, Pipeline->GetShaderHandles(ShaderFrequency));
diff --git a/Engine/Source/Runtime/VulkanRHI/Private/VulkanRayTracing.h b/Engine/Source/Runtime/VulkanRHI/Private/VulkanRayTracing.h
--- a/Engine/Source/Runtime/VulkanRHI/Private/VulkanRayTracing.h
+++ b/Engine/Source/Runtime/VulkanRHI/Private/VulkanRayTracing.h
@@ -81,6 +81,16 @@ public:
 
 	void Commit(FVulkanCommandListContext& Context);
 
+	void AddUBRef(FRHIUniformBuffer* UB)
+	{
+		ReferencedUniformBuffers.AddUnique(UB);
+	}
+
+	TArrayView<TRefCountPtr<FRHIUniformBuffer>> GetUBRefs()
+	{
+		return ReferencedUniformBuffers;
+	}
+
 private:
 
 	struct FVulkanShaderTableAllocation
@@ -112,6 +122,8 @@ private:
 	FVulkanShaderTableAllocation HitGroup;
 	FVulkanShaderTableAllocation Callable;
 
+	TArray<TRefCountPtr<FRHIUniformBuffer>> ReferencedUniformBuffers;
+
 	// Convenience
 	const uint32 HandleSize;
 	const uint32 HandleSizeAligned;
