Changing loaded spirv compression type

diff --git a/Engine/Source/Runtime/VulkanRHI/Private/VulkanShaders.cpp b/Engine/Source/Runtime/VulkanRHI/Private/VulkanShaders.cpp
--- a/Engine/Source/Runtime/VulkanRHI/Private/VulkanShaders.cpp
+++ b/Engine/Source/Runtime/VulkanRHI/Private/VulkanShaders.cpp
@@ -226,7 +226,7 @@ FArchive& operator<<(FArchive& Ar, FVulkanShader::FSpirvContainer& SpirvContaine
 	}
 	else
 	{
-		const int32 CompressedUpperBound = FCompression::CompressMemoryBound(NAME_Zlib, SpirvCodeSizeInBytes);
+		const int32 CompressedUpperBound = FCompression::CompressMemoryBound(NAME_Oodle, SpirvCodeSizeInBytes);
 		SpirvCode.Reserve(CompressedUpperBound);
 		SpirvCode.SetNumUninitialized(CompressedUpperBound);
 
@@ -235,7 +235,7 @@ FArchive& operator<<(FArchive& Ar, FVulkanShader::FSpirvContainer& SpirvContaine
 		Ar.Serialize(UncompressedSpirv.GetData(), SpirvCodeSizeInBytes);
 
 		int32 CompressedSizeBytes = CompressedUpperBound;
-		if (FCompression::CompressMemory(NAME_Zlib, SpirvCode.GetData(), CompressedSizeBytes, UncompressedSpirv.GetData(), UncompressedSpirv.GetTypeSize() * UncompressedSpirv.Num()))
+		if (FCompression::CompressMemory(NAME_Oodle, SpirvCode.GetData(), CompressedSizeBytes, UncompressedSpirv.GetData(), UncompressedSpirv.GetTypeSize() * UncompressedSpirv.Num()))
 		{
 			SpirvContainer.UncompressedSizeBytes = SpirvCodeSizeInBytes;
 			SpirvCode.SetNumUninitialized(CompressedSizeBytes);
@@ -264,7 +264,7 @@ FVulkanShader::FSpirvCode FVulkanShader::GetSpirvCode(const FSpirvContainer& Con
 		const size_t ElementSize = UncompressedSpirv.GetTypeSize();
 		UncompressedSpirv.Reserve(Container.GetSizeBytes() / ElementSize);
 		UncompressedSpirv.SetNumUninitialized(Container.GetSizeBytes() / ElementSize);
-		FCompression::UncompressMemory(NAME_Zlib, UncompressedSpirv.GetData(), Container.GetSizeBytes(), Container.SpirvCode.GetData(), Container.SpirvCode.Num());
+		FCompression::UncompressMemory(NAME_Oodle, UncompressedSpirv.GetData(), Container.GetSizeBytes(), Container.SpirvCode.GetData(), Container.SpirvCode.Num());
 
 		return FSpirvCode(MoveTemp(UncompressedSpirv));
 	}
