Added speed bias flag for spirv compression

diff --git a/Engine/Source/Runtime/VulkanRHI/Private/VulkanShaders.cpp b/Engine/Source/Runtime/VulkanRHI/Private/VulkanShaders.cpp
--- a/Engine/Source/Runtime/VulkanRHI/Private/VulkanShaders.cpp
+++ b/Engine/Source/Runtime/VulkanRHI/Private/VulkanShaders.cpp
@@ -235,7 +235,7 @@ FArchive& operator<<(FArchive& Ar, FVulkanShader::FSpirvContainer& SpirvContaine
 		Ar.Serialize(UncompressedSpirv.GetData(), SpirvCodeSizeInBytes);
 
 		int32 CompressedSizeBytes = CompressedUpperBound;
-		if (FCompression::CompressMemory(NAME_Oodle, SpirvCode.GetData(), CompressedSizeBytes, UncompressedSpirv.GetData(), UncompressedSpirv.GetTypeSize() * UncompressedSpirv.Num()))
+		if (FCompression::CompressMemory(NAME_Oodle, SpirvCode.GetData(), CompressedSizeBytes, UncompressedSpirv.GetData(), UncompressedSpirv.GetTypeSize() * UncompressedSpirv.Num(), ECompressionFlags::COMPRESS_BiasSpeed))
 		{
 			SpirvContainer.UncompressedSizeBytes = SpirvCodeSizeInBytes;
 			SpirvCode.SetNumUninitialized(CompressedSizeBytes);
