Report gl version string together with vulkan API version string on Android

diff --git a/Engine/Source/Runtime/VulkanRHI/Private/VulkanRHI.cpp b/Engine/Source/Runtime/VulkanRHI/Private/VulkanRHI.cpp
--- a/Engine/Source/Runtime/VulkanRHI/Private/VulkanRHI.cpp
+++ b/Engine/Source/Runtime/VulkanRHI/Private/VulkanRHI.cpp
@@ -28,7 +28,9 @@
 #include "VulkanExtensions.h"
 #include "VulkanRayTracing.h"
 #include "VulkanChunkedPipelineCache.h"
-
+#if PLATFORM_ANDROID
+#include "Android/AndroidPlatformMisc.h"
+#endif
 
 // Use Vulkan Profiles to verify feature level support on startup
 void VulkanProfilePrint(const char* Msg)
@@ -857,7 +859,14 @@ void FVulkanDynamicRHI::SelectDevice()
 	if (PLATFORM_ANDROID)
 	{
 		GRHIAdapterName.Append(TEXT(" Vulkan"));
-		GRHIAdapterInternalDriverVersion = FString::Printf(TEXT("%d.%d.%d"), VK_VERSION_MAJOR(Props.apiVersion), VK_VERSION_MINOR(Props.apiVersion), VK_VERSION_PATCH(Props.apiVersion));
+		// On Android GL version string often contains extra information such as an actual driver version on the device.
+#if PLATFORM_ANDROID
+		FString GLVersion = FAndroidMisc::GetGLVersion();
+#else
+		FString GLVersion = "";
+#endif
+		GRHIAdapterInternalDriverVersion = FString::Printf(TEXT("%d.%d.%d|%s"), VK_VERSION_MAJOR(Props.apiVersion), VK_VERSION_MINOR(Props.apiVersion), VK_VERSION_PATCH(Props.apiVersion), *GLVersion);
+		UE_LOG(LogVulkanRHI, Log, TEXT("API Version: %s"), *GRHIAdapterInternalDriverVersion);
 	}
 	else if (PLATFORM_WINDOWS)
 	{
