Removed VULKAN_SUPPORTS_MULTIVIEW since it's no longer needed.

diff --git a/Engine/Source/Runtime/VulkanRHI/Private/Android/VulkanAndroidPlatform.h b/Engine/Source/Runtime/VulkanRHI/Private/Android/VulkanAndroidPlatform.h
--- a/Engine/Source/Runtime/VulkanRHI/Private/Android/VulkanAndroidPlatform.h
+++ b/Engine/Source/Runtime/VulkanRHI/Private/Android/VulkanAndroidPlatform.h
@@ -23,7 +23,6 @@
 #define VULKAN_SUPPORTS_SCALAR_BLOCK_LAYOUT			(VULKAN_RHI_RAYTRACING)
 #define VULKAN_SUPPORTS_TRANSIENT_RESOURCE_ALLOCATOR 0
 #define VULKAN_SUPPORTS_DRIVER_PROPERTIES			0
-#define VULKAN_SUPPORTS_MULTIVIEW					1
 #define VULKAN_SUPPORTS_DESCRIPTOR_INDEXING			(VULKAN_RHI_RAYTRACING)
 
 // crashing during callback setup on Android, code will fallback to VK_EXT_debug_report instead
diff --git a/Engine/Source/Runtime/VulkanRHI/Private/Linux/VulkanLinuxPlatform.h b/Engine/Source/Runtime/VulkanRHI/Private/Linux/VulkanLinuxPlatform.h
--- a/Engine/Source/Runtime/VulkanRHI/Private/Linux/VulkanLinuxPlatform.h
+++ b/Engine/Source/Runtime/VulkanRHI/Private/Linux/VulkanLinuxPlatform.h
@@ -14,7 +14,6 @@
 
 #define VULKAN_RHI_RAYTRACING 						(RHI_RAYTRACING)
 #define VULKAN_SUPPORTS_SCALAR_BLOCK_LAYOUT			(VULKAN_RHI_RAYTRACING)
-#define VULKAN_SUPPORTS_MULTIVIEW					1	// needed for VULKAN_PCES31
 
 #if VULKAN_RHI_RAYTRACING
 #	define UE_VK_API_VERSION						VK_API_VERSION_1_2
diff --git a/Engine/Source/Runtime/VulkanRHI/Private/VulkanExtensions.cpp b/Engine/Source/Runtime/VulkanRHI/Private/VulkanExtensions.cpp
--- a/Engine/Source/Runtime/VulkanRHI/Private/VulkanExtensions.cpp
+++ b/Engine/Source/Runtime/VulkanRHI/Private/VulkanExtensions.cpp
@@ -399,7 +399,7 @@ class FVulkanKHRMultiviewExtension : public FVulkanDeviceExtension
 public:
 
 	FVulkanKHRMultiviewExtension(FVulkanDevice* InDevice)
-		: FVulkanDeviceExtension(InDevice, VK_KHR_MULTIVIEW_EXTENSION_NAME, VULKAN_SUPPORTS_MULTIVIEW, VK_API_VERSION_1_1)
+		: FVulkanDeviceExtension(InDevice, VK_KHR_MULTIVIEW_EXTENSION_NAME, VULKAN_EXTENSION_ENABLED, VK_API_VERSION_1_1)
 	{}
 
 	virtual void PrePhysicalDeviceFeatures(VkPhysicalDeviceFeatures2KHR& PhysicalDeviceFeatures2) override final
diff --git a/Engine/Source/Runtime/VulkanRHI/Private/VulkanRHI.cpp b/Engine/Source/Runtime/VulkanRHI/Private/VulkanRHI.cpp
--- a/Engine/Source/Runtime/VulkanRHI/Private/VulkanRHI.cpp
+++ b/Engine/Source/Runtime/VulkanRHI/Private/VulkanRHI.cpp
@@ -930,9 +930,7 @@ void FVulkanDynamicRHI::InitInstance()
 		GSupportsRenderTargetFormat_PF_G8 = false;	// #todo-rco
 		GRHISupportsTextureStreaming = true;
 		GSupportsTimestampRenderQueries = FVulkanPlatform::SupportsTimestampRenderQueries();
-#if VULKAN_SUPPORTS_MULTIVIEW
 		GSupportsMobileMultiView = Device->GetOptionalExtensions().HasKHRMultiview ? true : false;
-#endif
 #if VULKAN_RHI_RAYTRACING
 		GRHISupportsRayTracing = RHISupportsRayTracing(GMaxRHIShaderPlatform) && Device->GetOptionalExtensions().HasRaytracingExtensions();
 
diff --git a/Engine/Source/Runtime/VulkanRHI/Private/Windows/VulkanWindowsPlatform.h b/Engine/Source/Runtime/VulkanRHI/Private/Windows/VulkanWindowsPlatform.h
--- a/Engine/Source/Runtime/VulkanRHI/Private/Windows/VulkanWindowsPlatform.h
+++ b/Engine/Source/Runtime/VulkanRHI/Private/Windows/VulkanWindowsPlatform.h
@@ -18,7 +18,6 @@
 
 #define VULKAN_RHI_RAYTRACING 						(RHI_RAYTRACING)
 #define VULKAN_SUPPORTS_SCALAR_BLOCK_LAYOUT			(VULKAN_RHI_RAYTRACING)
-#define VULKAN_SUPPORTS_MULTIVIEW					1	// needed for VULKAN_PCES31
 
 #if VULKAN_RHI_RAYTRACING
 #	define UE_VK_API_VERSION						VK_API_VERSION_1_2
diff --git a/Engine/Source/Runtime/VulkanRHI/Public/VulkanConfiguration.h b/Engine/Source/Runtime/VulkanRHI/Public/VulkanConfiguration.h
--- a/Engine/Source/Runtime/VulkanRHI/Public/VulkanConfiguration.h
+++ b/Engine/Source/Runtime/VulkanRHI/Public/VulkanConfiguration.h
@@ -210,14 +210,6 @@ struct VkAllocationCallbacks;
 	#endif
 #endif
 
-#ifndef VULKAN_SUPPORTS_MULTIVIEW
-	#ifdef VK_KHR_multiview
-		#define VULKAN_SUPPORTS_MULTIVIEW						1
-	#else
-		#define VULKAN_SUPPORTS_MULTIVIEW						0
-	#endif
-#endif
-
 #ifndef VULKAN_SUPPORTS_ASTC_DECODE_MODE
 	#ifdef VK_EXT_astc_decode_mode
 		#define VULKAN_SUPPORTS_ASTC_DECODE_MODE				1
