Change VulkanRHI PSO precaching to default off.

diff --git a/Engine/Source/Runtime/VulkanRHI/Private/VulkanRHI.cpp b/Engine/Source/Runtime/VulkanRHI/Private/VulkanRHI.cpp
--- a/Engine/Source/Runtime/VulkanRHI/Private/VulkanRHI.cpp
+++ b/Engine/Source/Runtime/VulkanRHI/Private/VulkanRHI.cpp
@@ -102,9 +102,9 @@ static FAutoConsoleVariableRef CVarVulkanVariableRateShading(
 
 static TAutoConsoleVariable<bool> CVarAllowVulkanPSOPrecache(
 	TEXT("r.Vulkan.AllowPSOPrecaching"),
-	true,
-	TEXT("true: if r.PSOPrecaching=1 Vulkan RHI will use precaching. (default)\n")
-	TEXT("false: Vulkan RHI will disable precaching (even if r.PSOPrecaching=1). "),
+	false,
+	TEXT("true: if r.PSOPrecaching=1 Vulkan RHI will use precaching.\n")
+	TEXT("false: Vulkan RHI will disable precaching (even if r.PSOPrecaching=1). (default)"),
 	ECVF_RenderThreadSafe | ECVF_ReadOnly);
 
 // If precaching is active we should not need the file cache.
