VulkanRHI: Remove redundant cvar check and rely only on bindless config values to enable bindless in Vulkan.

#rnx

diff --git a/Engine/Source/Runtime/VulkanRHI/Private/VulkanDescriptorSets.cpp b/Engine/Source/Runtime/VulkanRHI/Private/VulkanDescriptorSets.cpp
--- a/Engine/Source/Runtime/VulkanRHI/Private/VulkanDescriptorSets.cpp
+++ b/Engine/Source/Runtime/VulkanRHI/Private/VulkanDescriptorSets.cpp
@@ -9,14 +9,6 @@
 #include "VulkanContext.h"
 
 
-int32 GVulkanBindlessEnabled = 0;
-static FAutoConsoleVariableRef CVarVulkanBindlessEnabled(
-	TEXT("r.Vulkan.Bindless.Enabled"),
-	GVulkanBindlessEnabled,
-	TEXT("Enable the use of bindless if all conditions are met to support it"),
-	ECVF_ReadOnly
-);
-
 int32 GVulkanBindlessMaxSamplerDescriptorCount = 2048;
 static FAutoConsoleVariableRef CVarVulkanBindlessMaxSamplerDescriptorCount(
 	TEXT("r.Vulkan.Bindless.MaxSamplerDescriptorCount"),
@@ -236,7 +228,7 @@ bool FVulkanBindlessDescriptorManager::VerifySupport(FVulkanDevice* Device)
 		(RHIGetRuntimeBindlessResourcesConfiguration(GMaxRHIShaderPlatform) == ERHIBindlessConfiguration::Disabled) &&
 		(RHIGetRuntimeBindlessSamplersConfiguration(GMaxRHIShaderPlatform) == ERHIBindlessConfiguration::Disabled);
 
-	if ((GVulkanBindlessEnabled == 0) || bFullyDisabled)
+	if (bFullyDisabled)
 	{
 		return false;
 	}
