VulkanRHI: Temp fix: Print leaked pipelines on shutdown instead of a hard assert which is failing all our automated tests.

#rnx

diff --git a/Engine/Source/Runtime/VulkanRHI/Private/VulkanPipeline.cpp b/Engine/Source/Runtime/VulkanRHI/Private/VulkanPipeline.cpp
--- a/Engine/Source/Runtime/VulkanRHI/Private/VulkanPipeline.cpp
+++ b/Engine/Source/Runtime/VulkanRHI/Private/VulkanPipeline.cpp
@@ -1656,7 +1656,7 @@ void FVulkanPipelineStateCacheManager::DestroyCache()
 	for (auto& Pair : GraphicsPSOLockedMap)
 	{
 		FVulkanRHIGraphicsPipelineState* Pipeline = Pair.Value;
-		LRUPRINT_DEBUG(TEXT("LIVE PSO %05d  REF %d %p :: %p\n"), idx++, Pipeline->GetRefCount(), Pipeline, Pipeline->LRUNode);
+		FPlatformMisc::LowLevelOutputDebugStringf(TEXT("Leaked PSO %05d: RefCount=%d Handle=0x%p\n"), idx++, Pipeline->GetRefCount(), Pipeline);
 	}
 	LRU2SizeList.Reset();
 
@@ -1666,7 +1666,6 @@ void FVulkanPipelineStateCacheManager::DestroyCache()
 
 	// Compute pipelines already deleted...
 	ComputePipelineEntries.Reset();
-	check(GraphicsPSOLockedMap.Num() == 0);
 }
 
 void FVulkanPipelineStateCacheManager::RebuildCache()
