Android Vulkan - Fix for missing PSO cache synchronisation when merging the result of external compile back into main PSO cache.

diff --git a/Engine/Source/Runtime/VulkanRHI/Private/VulkanChunkedPipelineCache.cpp b/Engine/Source/Runtime/VulkanRHI/Private/VulkanChunkedPipelineCache.cpp
--- a/Engine/Source/Runtime/VulkanRHI/Private/VulkanChunkedPipelineCache.cpp
+++ b/Engine/Source/Runtime/VulkanRHI/Private/VulkanChunkedPipelineCache.cpp
@@ -502,10 +502,8 @@ public:
 	};
 
 	template<class TPipelineState>
-	bool PSORequiresCompile(TPipelineState* Initializer, TUniqueFunction<VkResult(TPipelineState*, VkPipelineCache, EPSOOperation)>& PSOCreateFunc)
+	bool PSORequiresCompile(TPipelineState* Initializer, FVulkanChunkedPipelineCacheManager::FPSOCreateCallbackFunc<TPipelineState>& PSOCreateFunc)
 	{
-		FRWScopeLock PipelineLock(PipelineCacheObjLock, SLT_ReadOnly);
-
 		const bool bCanTestForExistence = UE::Vulkan::GetVulkanDevice()->GetOptionalExtensions().HasEXTPipelineCreationCacheControl;
 	
 		if(!bCanTestForExistence)
@@ -513,19 +511,19 @@ public:
 			return true;
 		}
 
-		VkResult Result = PSOCreateFunc(Initializer, PipelineCacheObj, EPSOOperation::CreateIfPresent);
+		FVulkanChunkedPipelineCacheManager::FPSOCreateFuncParams<TPipelineState> Params(Initializer, PipelineCacheObj, EPSOOperation::CreateIfPresent, PipelineCacheObjLock);
+		VkResult Result = PSOCreateFunc(Params);
 		check(Result == VK_SUCCESS || Result == VK_PIPELINE_COMPILE_REQUIRED_EXT);
 		return Result != VK_SUCCESS;
 	}
 
 	template<class TPipelineState>
-	VkResult CreatePSO(TPipelineState* Initializer, EPSOCacheFindResult PSOCacheFindResult, TUniqueFunction<VkResult(TPipelineState*, VkPipelineCache, EPSOOperation)> PSOCreateFunc)
+	VkResult CreatePSO(TPipelineState* Initializer, EPSOCacheFindResult PSOCacheFindResult, FVulkanChunkedPipelineCacheManager::FPSOCreateCallbackFunc<TPipelineState> PSOCreateFunc)
 	{
 		FScopedTimeToLog Timer(FString::Printf(TEXT("FVulkanChunkedPipelineCacheManager: FVulkanPipelineCacheChunk.CreatePSO  tot %d "), FPlatformTLS::GetCurrentThreadId()));
 		Touch();
 
 		QUICK_SCOPE_CYCLE_COUNTER(STAT_Vulkan_RHICreateGraphicsPipelineState_CREATE_CreatePSO);
-		FRWScopeLock PipelineLock(PipelineCacheObjLock, SLT_ReadOnly);
 
 		VkResult retcode;
 		uint64 PSOHash;
@@ -537,7 +535,8 @@ public:
 			{
 				// it's possible to still add a new PSO if we have hash collisions or imperfect PSO hash calc.
 				QUICK_SCOPE_CYCLE_COUNTER(STAT_Vulkan_CacheManagerCreatePSO_WARM);
-				retcode = PSOCreateFunc(Initializer, PipelineCacheObj, EPSOOperation::CreateAndStorePSO);
+				FVulkanChunkedPipelineCacheManager::FPSOCreateFuncParams<TPipelineState> Params (Initializer, PipelineCacheObj, EPSOOperation::CreateAndStorePSO, PipelineCacheObjLock);
+				retcode = PSOCreateFunc(Params);
 			}
 			else if(PSOCacheFindResult == EPSOCacheFindResult::NotFound)
 			{
@@ -550,7 +549,8 @@ public:
 				// Holding the write lock for the duration of the create is too costly, better to let the driver manage this.
 				{
 					QUICK_SCOPE_CYCLE_COUNTER(STAT_Vulkan_CacheManagerCreatePSO_COLD);
-					retcode = PSOCreateFunc(Initializer, PipelineCacheObj, EPSOOperation::CreateAndStorePSO);
+					FVulkanChunkedPipelineCacheManager::FPSOCreateFuncParams<TPipelineState> Params(Initializer, PipelineCacheObj, EPSOOperation::CreateAndStorePSO, PipelineCacheObjLock);
+					retcode = PSOCreateFunc(Params);
 				}
 			}
 			else
@@ -561,14 +561,6 @@ public:
 
 		if (PSOCacheFindResult != EPSOCacheFindResult::Found)
 		{
-			size_t LocalCacheSize = 0;
-			if (PSOCacheFindResult == EPSOCacheFindResult::NotFound) // we can avoid cachestatelock for this.
-			{
-				// read the current size without the lock (i.e. expect only an approximation)
-				FScopedTimeToLog Timer4(FString::Printf(TEXT("FVulkanChunkedPipelineCacheManager: FVulkanPipelineCacheChunk.CreatePSO AddTo getdat tot %d "), FPlatformTLS::GetCurrentThreadId()));
-				VulkanRHI::vkGetPipelineCacheData(UE::Vulkan::GetVulkanDeviceHandle(), PipelineCacheObj, &LocalCacheSize, nullptr);
-			}
-
 			// if EPSOCacheFindResult::MatchedExisting we record the hash only. 
 			// 'existing' pso's will not contribute to the cache, we dont consider them for pending compiles.
 			FRWScopeLock Lock(CacheStateLock, SLT_Write);
@@ -582,7 +574,6 @@ public:
 				FScopedTimeToLog Timer3(FString::Printf(TEXT("FVulkanChunkedPipelineCacheManager: FVulkanPipelineCacheChunk.CreatePSO AddTo tot %d "), FPlatformTLS::GetCurrentThreadId()));
 
 				TotalNumUniquePSOs++;
-				CacheSize = LocalCacheSize;
 
 				check(PendingAddToCompiles.load() != 0);
 				--PendingAddToCompiles;
@@ -621,7 +612,7 @@ public:
 				CacheState = ECacheState::Closing;
 			}
 		}
-	}
+	} 
 
 	void LogStats(FString&& LogInfo)
 	{
@@ -924,7 +915,7 @@ public:
 	}
 
 	template<class TPipelineState>
-	VkResult CreatePSO(TPipelineState* GraphicsPipelineState, bool bIsPrecompileJob, TUniqueFunction<VkResult(TPipelineState*, VkPipelineCache, EPSOOperation)> PSOCreateFunc)
+	VkResult CreatePSO(TPipelineState* GraphicsPipelineState, bool bIsPrecompileJob, FVulkanChunkedPipelineCacheManager::FPSOCreateCallbackFunc<TPipelineState> PSOCreateFunc)
 	{
 		FScopedTimeToLog Timer(FString::Printf(TEXT("FVulkanChunkedPipelineCacheManager: CreatePSO (precomp %d) tot %d "), bIsPrecompileJob, FPlatformTLS::GetCurrentThreadId()));
 
@@ -1073,7 +1064,7 @@ private:
 	static uint64 GetPrecacheHash(const FVulkanRHIGraphicsPipelineState* GFXState) { return GFXState->PrecacheKey; }
 
 	template<class TPipelineState, class TInitializer>
-	FVulkanPipelineCacheChunk* GetChunk(TUniqueFunction<VkResult(TPipelineState*, VkPipelineCache, EPSOOperation)>& PSOCreateFunc, const TInitializer& Initializer, FVulkanPipelineCacheChunk::EPSOCacheFindResult& FindResult, FCacheChunkKey& ChunkKeyOUT, FRWScopeLock& Lock, const bool bTryAdd = false)
+	FVulkanPipelineCacheChunk* GetChunk(FVulkanChunkedPipelineCacheManager::FPSOCreateCallbackFunc<TPipelineState>& PSOCreateFunc, const TInitializer& Initializer, FVulkanPipelineCacheChunk::EPSOCacheFindResult& FindResult, FCacheChunkKey& ChunkKeyOUT, FRWScopeLock& Lock, const bool bTryAdd = false)
 	{
 		FVulkanPipelineCacheChunk* ReturnChunk = nullptr;
 		uint64 PSOPrecacheKey = GetPrecacheHash(Initializer);
@@ -1148,7 +1139,7 @@ private:
 	}
 
 	template<class TPipelineState, class TInitializer>
-	FVulkanPipelineCacheChunk* GetOrAddCache(TUniqueFunction<VkResult(TPipelineState*, VkPipelineCache, EPSOOperation)>& PSOCreateFunc, const TInitializer& Initializer, FVulkanPipelineCacheChunk::EPSOCacheFindResult& FindResult, FCacheChunkKey& ChunkKeyOUT)
+	FVulkanPipelineCacheChunk* GetOrAddCache(FVulkanChunkedPipelineCacheManager::FPSOCreateCallbackFunc<TPipelineState>& PSOCreateFunc, const TInitializer& Initializer, FVulkanPipelineCacheChunk::EPSOCacheFindResult& FindResult, FCacheChunkKey& ChunkKeyOUT)
 	{
 		FScopedTimeToLog Timer(FString::Printf(TEXT("FVulkanChunkedPipelineCacheManager: GetOrAddCache tot %d "), FPlatformTLS::GetCurrentThreadId()));
 		FRWScopeLock Lock(ChunkedPipelineCacheLock, SLT_ReadOnly);
@@ -1238,7 +1229,7 @@ FVulkanChunkedPipelineCacheManager& FVulkanChunkedPipelineCacheManager::Get()
 
 
 template<class TPipelineState>
-VkResult FVulkanChunkedPipelineCacheManager::CreatePSO(TPipelineState* GraphicsPipelineState, bool bIsPrecompileJob, TUniqueFunction<VkResult(TPipelineState*, VkPipelineCache, EPSOOperation)> PSOCreateFunc)
+VkResult FVulkanChunkedPipelineCacheManager::CreatePSO(TPipelineState* GraphicsPipelineState, bool bIsPrecompileJob, FPSOCreateCallbackFunc<TPipelineState> PSOCreateFunc)
 {
 	check(UE::Vulkan::GUseChunkedPSOCache);
 	return VulkanPipelineCacheManagerImpl->CreatePSO(GraphicsPipelineState, bIsPrecompileJob, MoveTemp(PSOCreateFunc));
@@ -1254,5 +1245,5 @@ void FVulkanChunkedPipelineCacheManager::Tick()
 	VulkanPipelineCacheManagerImpl->Tick();
 }
 
-template VkResult FVulkanChunkedPipelineCacheManager::CreatePSO(FVulkanRHIGraphicsPipelineState* GraphicsPipelineState, bool bIsPrecompileJob, TUniqueFunction<VkResult(FVulkanRHIGraphicsPipelineState*, VkPipelineCache, EPSOOperation)> PSOCreateFunc);
+template VkResult FVulkanChunkedPipelineCacheManager::CreatePSO(FVulkanRHIGraphicsPipelineState* GraphicsPipelineState, bool bIsPrecompileJob, FPSOCreateCallbackFunc<FVulkanRHIGraphicsPipelineState> PSOCreateFunc);
 
diff --git a/Engine/Source/Runtime/VulkanRHI/Private/VulkanChunkedPipelineCache.h b/Engine/Source/Runtime/VulkanRHI/Private/VulkanChunkedPipelineCache.h
--- a/Engine/Source/Runtime/VulkanRHI/Private/VulkanChunkedPipelineCache.h
+++ b/Engine/Source/Runtime/VulkanRHI/Private/VulkanChunkedPipelineCache.h
@@ -28,7 +28,25 @@ public:
 		CreateAndStorePSO,	// Create and store the PSO in the cache.
 	};
 
+
 	template<class TPipelineState>
-	VkResult CreatePSO(TPipelineState* GraphicsPipelineState, bool bIsPrecompileJob, TUniqueFunction<VkResult(TPipelineState*, VkPipelineCache, EPSOOperation)> PSOCreateFunc);
+	struct FPSOCreateFuncParams
+	{
+		FPSOCreateFuncParams(TPipelineState* PSOIn, VkPipelineCache DestPipelineCacheIn, FVulkanChunkedPipelineCacheManager::EPSOOperation PSOOperationIn, FRWLock& DestPipelineCacheLockIn)
+			: PSO(PSOIn), DestPipelineCache(DestPipelineCacheIn), PSOOperation(PSOOperationIn), DestPipelineCacheLock(DestPipelineCacheLockIn)
+		{}
+
+		TPipelineState* PSO;
+		VkPipelineCache DestPipelineCache;
+		FVulkanChunkedPipelineCacheManager::EPSOOperation PSOOperation;
+		FRWLock& DestPipelineCacheLock; // ensure lock is acquired before use of DestPipelineCache.
+	};
+
+	template<class TPipelineState>
+	using FPSOCreateCallbackFunc = TUniqueFunction<VkResult(FPSOCreateFuncParams<TPipelineState>& Params)>;
+
+	template<class TPipelineState>
+	VkResult CreatePSO(TPipelineState* GraphicsPipelineState, bool bIsPrecompileJob, FPSOCreateCallbackFunc<TPipelineState> PSOCreateFunc);
+
 	void Tick();
 };
diff --git a/Engine/Source/Runtime/VulkanRHI/Private/VulkanPipeline.cpp b/Engine/Source/Runtime/VulkanRHI/Private/VulkanPipeline.cpp
--- a/Engine/Source/Runtime/VulkanRHI/Private/VulkanPipeline.cpp
+++ b/Engine/Source/Runtime/VulkanRHI/Private/VulkanPipeline.cpp
@@ -1399,9 +1399,8 @@ bool FVulkanPipelineStateCacheManager::CreateGfxPipelineFromEntry(FVulkanRHIGrap
 }
 
 #if PLATFORM_ANDROID
-static FCriticalSection ExternalService;
 
-VkResult CreatePSOWithExternalService(FVulkanDevice* Device, FVulkanRHIGraphicsPipelineState* PSO, FVulkanShader* Shaders[ShaderStage::NumStages], const VkGraphicsPipelineCreateInfo& PipelineInfo, VkPipelineCache DestPipelineCache)
+VkResult CreatePSOWithExternalService(FVulkanDevice* Device, FVulkanRHIGraphicsPipelineState* PSO, FVulkanShader* Shaders[ShaderStage::NumStages], const VkGraphicsPipelineCreateInfo& PipelineInfo, VkPipelineCache DestPipelineCache, FRWLock& PipelineLock)
 {
 	VkResult Result = VK_ERROR_INITIALIZATION_FAILED;
 	FVulkanShader::FSpirvCode VS = PSO->GetPatchedSpirvCode(Shaders[ShaderStage::Vertex]);
@@ -1415,7 +1414,8 @@ VkResult CreatePSOWithExternalService(FVulkanDevice* Device, FVulkanRHIGraphicsP
 
 	TArray<uint8> InitialCacheData;
 	{
-		FScopeLock Lock(&ExternalService);
+		QUICK_SCOPE_CYCLE_COUNTER(STAT_Vulkan_RHICreateGraphicsPipelineState_ExternalInitialCache);
+		FRWScopeLock Lock(PipelineLock, SLT_Write);
 		size_t InitialCacheSize = 0;
 		VulkanRHI::vkGetPipelineCacheData(Device->GetInstanceHandle(), DestPipelineCache, &InitialCacheSize, nullptr);
 		InitialCacheData.SetNumUninitialized(InitialCacheSize);
@@ -1426,9 +1426,12 @@ VkResult CreatePSOWithExternalService(FVulkanDevice* Device, FVulkanRHIGraphicsP
 
 	if (ensure(LocalPipelineCache != VK_NULL_HANDLE))
 	{
-		Result = VK_SUCCESS;
-		FScopeLock Lock(&ExternalService);
-		VERIFYVULKANRESULT(VulkanRHI::vkMergePipelineCaches(Device->GetInstanceHandle(), DestPipelineCache, 1, &LocalPipelineCache));
+		{
+			QUICK_SCOPE_CYCLE_COUNTER(STAT_Vulkan_RHICreateGraphicsPipelineState_ExternalMergeResult);
+			FRWScopeLock Lock(PipelineLock, SLT_Write);
+			Result = VK_SUCCESS;
+			VERIFYVULKANRESULT(VulkanRHI::vkMergePipelineCaches(Device->GetInstanceHandle(), DestPipelineCache, 1, &LocalPipelineCache));
+		}
 		VulkanRHI::vkDestroyPipelineCache(Device->GetInstanceHandle(), LocalPipelineCache, VULKAN_CPU_ALLOCATOR);
 	}
 	else
@@ -1448,31 +1451,33 @@ VkResult FVulkanPipelineStateCacheManager::CreateVKPipeline(FVulkanRHIGraphicsPi
 		// Use chunk caching and bypass FVulkanPipelineStateCacheManager's PSO caching
 		// Placeholder PSO size - TODO: remove pipeline cache size stuff.	
 		PSO->PipelineCacheSize = 20 * 1024; // This is only required bUseLRU == true.
-		return FVulkanChunkedPipelineCacheManager::Get().CreatePSO(PSO, bIsPrecompileJob, TUniqueFunction< VkResult(FVulkanRHIGraphicsPipelineState*, VkPipelineCache, FVulkanChunkedPipelineCacheManager::EPSOOperation)>(
-			[&](FVulkanRHIGraphicsPipelineState* PSO, VkPipelineCache PipelineCache, FVulkanChunkedPipelineCacheManager::EPSOOperation PSOOperation)
+		return FVulkanChunkedPipelineCacheManager::Get().CreatePSO(PSO, bIsPrecompileJob, FVulkanChunkedPipelineCacheManager::FPSOCreateCallbackFunc<FVulkanRHIGraphicsPipelineState>(
+				[&](FVulkanChunkedPipelineCacheManager::FPSOCreateFuncParams<FVulkanRHIGraphicsPipelineState>& Params)
 			{
+				FVulkanRHIGraphicsPipelineState* PSO = Params.PSO;
+				VkPipelineCache& PipelineCache = Params.DestPipelineCache;
+				FVulkanChunkedPipelineCacheManager::EPSOOperation PSOOperation = Params.PSOOperation;
+				
 				check(PSO->VulkanPipeline == 0);
 				check(PSOOperation == FVulkanChunkedPipelineCacheManager::EPSOOperation::CreateAndStorePSO || PSOOperation == FVulkanChunkedPipelineCacheManager::EPSOOperation::CreateIfPresent);
-				VkResult Result = VK_ERROR_INITIALIZATION_FAILED;
-
-				const bool bCanTestForExistence = Device->GetOptionalExtensions().HasEXTPipelineCreationCacheControl;				
-
-				if (bCanTestForExistence)
-				{
-					VkGraphicsPipelineCreateInfo TestPipelineInfo = PipelineInfo;
-					TestPipelineInfo.flags |= VK_PIPELINE_CREATE_FAIL_ON_PIPELINE_COMPILE_REQUIRED_BIT_EXT;
-					Result = VulkanRHI::vkCreateGraphicsPipelines(Device->GetInstanceHandle(), PipelineCache, 1, &TestPipelineInfo, VULKAN_CPU_ALLOCATOR, &PSO->VulkanPipeline);
-				}
-
-				if (PSOOperation == FVulkanChunkedPipelineCacheManager::EPSOOperation::CreateIfPresent)
-				{
-					// Do not do a full create, we're only testing for existence here.
-					return bCanTestForExistence ? Result : VK_PIPELINE_COMPILE_REQUIRED_EXT;
-				}
+				VkResult Result;
 
-				if (Result == VK_SUCCESS)
+				if(PSOOperation == FVulkanChunkedPipelineCacheManager::EPSOOperation::CreateIfPresent)
 				{
-					UE_CLOG(bIsPrecompileJob, LogVulkanRHI, Log, TEXT("redundant precompile PSO request"));
+					const bool bCanTestForExistence = Device->GetOptionalExtensions().HasEXTPipelineCreationCacheControl;
+					if (bCanTestForExistence)
+					{
+						QUICK_SCOPE_CYCLE_COUNTER(STAT_Vulkan_RHICreateGraphicsPipelineState_TestCreate);
+						FRWScopeLock Lock(Params.DestPipelineCacheLock, SLT_ReadOnly);
+						VkGraphicsPipelineCreateInfo TestPipelineInfo = PipelineInfo;
+						TestPipelineInfo.flags |= VK_PIPELINE_CREATE_FAIL_ON_PIPELINE_COMPILE_REQUIRED_BIT_EXT;
+						Result = VulkanRHI::vkCreateGraphicsPipelines(Device->GetInstanceHandle(), PipelineCache, 1, &TestPipelineInfo, VULKAN_CPU_ALLOCATOR, &PSO->VulkanPipeline);
+					}
+					else
+					{
+						// if we cant test we must create.
+						Result = VK_PIPELINE_COMPILE_REQUIRED_EXT;
+					}
 					return Result;
 				}
 
@@ -1480,15 +1485,16 @@ VkResult FVulkanPipelineStateCacheManager::CreateVKPipeline(FVulkanRHIGraphicsPi
 #if PLATFORM_ANDROID
 				if( !FVulkanAndroidPlatform::AreRemoteCompileServicesActive() || !bIsPrecompileJob)
 				{
-					// enabled this to detect PSOs that the driver had to compile. (useful for PSO precaching analysis)
-					//UE_CLOG(Result == VK_PIPELINE_COMPILE_REQUIRED_EXT && !bIsPrecompileJob, LogVulkanRHI, Log, TEXT("missed PSO create.."));
 #endif
+					QUICK_SCOPE_CYCLE_COUNTER(STAT_Vulkan_RHICreateGraphicsPipelineState_vkCreate);
+					FRWScopeLock Lock(Params.DestPipelineCacheLock, SLT_ReadOnly);
 					Result = VulkanRHI::vkCreateGraphicsPipelines(Device->GetInstanceHandle(), PipelineCache, 1, &PipelineInfo, VULKAN_CPU_ALLOCATOR, &PSO->VulkanPipeline);
 #if PLATFORM_ANDROID
 				}
 				else
 				{
-					Result = CreatePSOWithExternalService(Device, PSO, Shaders, PipelineInfo, PipelineCache);
+					QUICK_SCOPE_CYCLE_COUNTER(STAT_Vulkan_RHICreateGraphicsPipelineState_ExternalCreate);
+					Result = CreatePSOWithExternalService(Device, PSO, Shaders, PipelineInfo, PipelineCache, Params.DestPipelineCacheLock);
 				}
 #endif
 				return Result;
