Android Vulkan: Enable GPU timing stats, available when swappy is used.

diff --git a/Engine/Source/Runtime/VulkanRHI/Private/Android/VulkanAndroidPlatform.cpp b/Engine/Source/Runtime/VulkanRHI/Private/Android/VulkanAndroidPlatform.cpp
--- a/Engine/Source/Runtime/VulkanRHI/Private/Android/VulkanAndroidPlatform.cpp
+++ b/Engine/Source/Runtime/VulkanRHI/Private/Android/VulkanAndroidPlatform.cpp
@@ -24,6 +24,27 @@
 #include "Android/AndroidJNI.h"
 #include "Android/AndroidApplication.h"
 #include "swappy/swappyVk.h"
+#include "EngineGlobals.h"
+
+namespace AndroidVulkan
+{
+	void VKSwappyPostWaitCallback(void*, int64_t cpu_time_ns, int64_t gpu_time_ns)
+	{
+		const double Frequency = 1.0;// FGPUTiming::GetTimingFrequency();
+		const double CyclesPerSecond = 1.0 / (Frequency * FPlatformTime::GetSecondsPerCycle64());
+		const double GPUTimeInSeconds = (double)gpu_time_ns / 1000000000.0;
+
+		GGPUFrameTime = CyclesPerSecond * GPUTimeInSeconds;
+	}
+
+	void SetSwappyPostWaitCallback()
+	{
+		SwappyTracer Tracer = { 0 };
+		Tracer.postWait = VKSwappyPostWaitCallback;
+		SwappyVk_injectTracer(&Tracer);
+	}
+};
+
 #endif
 
 // From VulklanSwapChain.cpp
@@ -343,6 +364,13 @@ void FVulkanAndroidPlatform::FreeVulkanLibrary()
 
 #undef CHECK_VK_ENTRYPOINTS
 
+bool FVulkanAndroidPlatform::HasCustomFrameTiming()
+{
+#if USE_ANDROID_SWAPPY
+	return FAndroidPlatformRHIFramePacer::CVarUseSwappyForFramePacing.GetValueOnAnyThread() != 0;
+#endif
+	return false;
+}
 
 void FVulkanAndroidPlatform::InitDevice(FVulkanDevice* InDevice)
 {
@@ -673,7 +701,7 @@ VkResult FVulkanAndroidPlatform::CreateSwapchainKHR(void* WindowHandle, VkPhysic
 				SwappyVk_initAndGetRefreshCycleDuration(Env, FJavaWrapper::GameActivityThis, PhysicalDevice, Device, *Swapchain, &RefreshDuration);
 				SwappyVk_setWindow(Device, *Swapchain, (ANativeWindow*)WindowHandle);	
 				SwappyVk_setAutoSwapInterval(false);
-				
+				AndroidVulkan::SetSwappyPostWaitCallback();
 				UE_LOG(LogVulkanRHI, Log, TEXT("SwappyVk_initAndGetRefreshCycleDuration: %ull"), RefreshDuration);
 			}
 
diff --git a/Engine/Source/Runtime/VulkanRHI/Private/Android/VulkanAndroidPlatform.h b/Engine/Source/Runtime/VulkanRHI/Private/Android/VulkanAndroidPlatform.h
--- a/Engine/Source/Runtime/VulkanRHI/Private/Android/VulkanAndroidPlatform.h
+++ b/Engine/Source/Runtime/VulkanRHI/Private/Android/VulkanAndroidPlatform.h
@@ -104,6 +104,8 @@ public:
 		return !FAndroidMisc::ShouldUseDesktopVulkan();
 	}
 
+	static bool HasCustomFrameTiming();
+
 	static bool SupportsVolumeTextureRendering() { return false; }
 
 	static void OverridePlatformHandlers(bool bInit);
diff --git a/Engine/Source/Runtime/VulkanRHI/Private/VulkanGenericPlatform.h b/Engine/Source/Runtime/VulkanRHI/Private/VulkanGenericPlatform.h
--- a/Engine/Source/Runtime/VulkanRHI/Private/VulkanGenericPlatform.h
+++ b/Engine/Source/Runtime/VulkanRHI/Private/VulkanGenericPlatform.h
@@ -57,6 +57,8 @@ public:
 
 	static bool SupportsTimestampRenderQueries() { return true; }
 
+	static bool HasCustomFrameTiming() { return false; }
+
 	static bool RequiresMobileRenderer() { return false; }
 
 	// bInit=1 called at RHI init time, bInit=0 at RHI deinit time
diff --git a/Engine/Source/Runtime/VulkanRHI/Private/VulkanQuery.cpp b/Engine/Source/Runtime/VulkanRHI/Private/VulkanQuery.cpp
--- a/Engine/Source/Runtime/VulkanRHI/Private/VulkanQuery.cpp
+++ b/Engine/Source/Runtime/VulkanRHI/Private/VulkanQuery.cpp
@@ -406,7 +406,7 @@ void FVulkanCommandListContext::ReadAndCalculateGPUFrameTime()
 		const double Frequency = double(FVulkanGPUTiming::GetTimingFrequency());
 		GGPUFrameTime = FMath::TruncToInt(double(Delta) / Frequency / SecondsPerCycle);
 	}
-	else
+	else if(!FVulkanPlatform::HasCustomFrameTiming())
 	{
 		GGPUFrameTime = 0;
 	}
