Android Vulkan: Fix bug that caused any existing PSO cache to be rejected.

diff --git a/Engine/Source/Runtime/VulkanRHI/Private/VulkanChunkedPipelineCache.cpp b/Engine/Source/Runtime/VulkanRHI/Private/VulkanChunkedPipelineCache.cpp
--- a/Engine/Source/Runtime/VulkanRHI/Private/VulkanChunkedPipelineCache.cpp
+++ b/Engine/Source/Runtime/VulkanRHI/Private/VulkanChunkedPipelineCache.cpp
@@ -297,7 +297,6 @@ public:
 
 		uint32 ParamHash;
 		uint32 LastValidOffset;
-		Archive << Version;
 		Archive << ParamHash;
 		Archive << LastValidOffset;
 		if (LastValidOffset == 0 || ParamHash != GetCacheBuildingParamHash())
