Fixed crash in FVulkanPipelineStateCacheManager related to PSO pre-caching

diff --git a/Engine/Source/Runtime/VulkanRHI/Private/VulkanPipeline.cpp b/Engine/Source/Runtime/VulkanRHI/Private/VulkanPipeline.cpp
--- a/Engine/Source/Runtime/VulkanRHI/Private/VulkanPipeline.cpp
+++ b/Engine/Source/Runtime/VulkanRHI/Private/VulkanPipeline.cpp
@@ -1638,10 +1638,13 @@ VkResult FVulkanPipelineStateCacheManager::CreateVKPipeline(FVulkanRHIGraphicsPi
 				PSOSize = Found->PipelineSize;
 			}
 		}
-
-		QUICK_SCOPE_CYCLE_COUNTER(STAT_VulkanPSOCacheMerge);
+				
 		FScopedPipelineCache PipelineCacheExclusive = Cache.Get(EPipelineCacheAccess::Exclusive);
-		VERIFYVULKANRESULT(VulkanRHI::vkMergePipelineCaches(Device->GetInstanceHandle(), PipelineCacheExclusive.Get(), 1, &LocalPipelineCache));
+		if (PipelineCacheExclusive.Get() != VK_NULL_HANDLE)
+		{
+			QUICK_SCOPE_CYCLE_COUNTER(STAT_VulkanPSOCacheMerge);
+			VERIFYVULKANRESULT(VulkanRHI::vkMergePipelineCaches(Device->GetInstanceHandle(), PipelineCacheExclusive.Get(), 1, &LocalPipelineCache));
+		}
 		VulkanRHI::vkDestroyPipelineCache(Device->GetInstanceHandle(), LocalPipelineCache, VULKAN_CPU_ALLOCATOR);
 	}
 
