Completely turn off FSR or FDM when the other is available.
This avoids vulkan validation layers warnings.

diff --git a/Engine/Source/Runtime/VulkanRHI/Private/VulkanDevice.cpp b/Engine/Source/Runtime/VulkanRHI/Private/VulkanDevice.cpp
--- a/Engine/Source/Runtime/VulkanRHI/Private/VulkanDevice.cpp
+++ b/Engine/Source/Runtime/VulkanRHI/Private/VulkanDevice.cpp
@@ -1091,6 +1091,81 @@ bool FVulkanDevice::SupportsBindless() const
 	return BindlessDescriptorManager->IsSupported();
 }
 
+void FVulkanDevice::ChooseVariableRateShadingMethod()
+{
+	auto IsFragmentShadingRateAvailable = [](VkPhysicalDeviceFragmentShadingRateFeaturesKHR& FragmentShadingRateFeatures)
+	{
+		return FragmentShadingRateFeatures.attachmentFragmentShadingRate == VK_TRUE;
+	};
+
+	auto IsFragmentDensityMapAvailable = [](FOptionalVulkanDeviceExtensions& ExtensionFlags)
+	{
+		return ExtensionFlags.HasEXTFragmentDensityMap;
+	};
+
+	auto TurnOffFragmentShadingRate = [](VkPhysicalDeviceFragmentShadingRateFeaturesKHR& FragmentShadingRateFeatures)
+	{
+		FragmentShadingRateFeatures.primitiveFragmentShadingRate = VK_FALSE;
+		FragmentShadingRateFeatures.attachmentFragmentShadingRate = VK_FALSE;
+		FragmentShadingRateFeatures.pipelineFragmentShadingRate = VK_FALSE;
+		GRHISupportsPipelineVariableRateShading = false;
+		GRHISupportsLargerVariableRateShadingSizes = false;
+	};
+
+	auto TurnOffFragmentDensityMap = [](FOptionalVulkanDeviceExtensions& ExtensionFlags, VkPhysicalDeviceFragmentDensityMapFeaturesEXT& FragmentDensityMapFeatures, VkPhysicalDeviceFragmentDensityMap2FeaturesEXT& FragmentDensityMap2Features)
+	{
+		ExtensionFlags.HasEXTFragmentDensityMap = 0;
+		FragmentDensityMapFeatures.fragmentDensityMap = VK_FALSE;
+		FragmentDensityMapFeatures.fragmentDensityMapDynamic = VK_FALSE;
+		FragmentDensityMapFeatures.fragmentDensityMapNonSubsampledImages = VK_FALSE;
+		ExtensionFlags.HasEXTFragmentDensityMap2 = 0;
+		FragmentDensityMap2Features.fragmentDensityMapDeferred = VK_FALSE;
+	};
+
+	int32 VRSFormatPreference = IConsoleManager::Get().FindTConsoleVariableDataInt(TEXT("r.Vulkan.VRSFormat"))->GetValueOnAnyThread();
+	UE_LOG(LogVulkanRHI, Display, TEXT("Vulkan Variable Rate Shading choice: %d."), VRSFormatPreference);
+
+	// If both FSR and FDM are available we turn off the one that we're not using to prevent Vulkan validation layers warnings.
+	if (IsFragmentDensityMapAvailable(OptionalDeviceExtensions) && IsFragmentShadingRateAvailable(OptionalDeviceExtensionProperties.FragmentShadingRateFeatures))
+	{
+		if (VRSFormatPreference <= (uint8)EVulkanVariableRateShadingPreference::RequireFSR)
+		{
+			TurnOffFragmentDensityMap(OptionalDeviceExtensions, OptionalDeviceExtensionProperties.FragmentDensityMapFeatures, OptionalDeviceExtensionProperties.FragmentDensityMap2Features);
+		}
+		else
+		{
+			TurnOffFragmentShadingRate(OptionalDeviceExtensionProperties.FragmentShadingRateFeatures);
+		}
+		return;
+	}
+	// When only FSR is available.
+	if (IsFragmentShadingRateAvailable(OptionalDeviceExtensionProperties.FragmentShadingRateFeatures))
+	{
+		if (VRSFormatPreference == (uint8)EVulkanVariableRateShadingPreference::UseFDMOnlyIfAvailable)
+		{
+			UE_LOG(LogVulkanRHI, Display, TEXT("Fragment Density Map was requested but is not available."));
+		}
+		else if (VRSFormatPreference == (uint8)EVulkanVariableRateShadingPreference::RequireFDM)
+		{
+			UE_LOG(LogVulkanRHI, Error, TEXT("Fragment Density Map was required but is not available."));
+		}
+		TurnOffFragmentDensityMap(OptionalDeviceExtensions, OptionalDeviceExtensionProperties.FragmentDensityMapFeatures, OptionalDeviceExtensionProperties.FragmentDensityMap2Features);
+	}
+	// When only FDM is available.
+	if (IsFragmentDensityMapAvailable(OptionalDeviceExtensions))
+	{
+		if (VRSFormatPreference == (uint8)EVulkanVariableRateShadingPreference::UseFSROnlyIfAvailable)
+		{
+			UE_LOG(LogVulkanRHI, Display, TEXT("Fragment Shading Rate was requested but is not available."));
+		}
+		else if (VRSFormatPreference == (uint8)EVulkanVariableRateShadingPreference::RequireFSR)
+		{
+			UE_LOG(LogVulkanRHI, Error, TEXT("Fragment Shading Rate was required but is not available."));
+		}
+		TurnOffFragmentShadingRate(OptionalDeviceExtensionProperties.FragmentShadingRateFeatures);
+	}
+}
+
 void FVulkanDevice::InitGPU()
 {
 	LLM_SCOPE_VULKAN(ELLMTagVulkan::VulkanMisc);
@@ -1159,7 +1234,7 @@ void FVulkanDevice::InitGPU()
 		}
 	}
 
-	ChooseVariableRateShadingMethod(OptionalDeviceExtensions, GetOptionalExtensionProperties().FragmentShadingRateFeatures);
+	ChooseVariableRateShadingMethod();
 
 	UE_LOG(LogVulkanRHI, Display, TEXT("Device properties: Geometry %d BufferAtomic64 %d ImageAtomic64 %d"), 
 		PhysicalDeviceFeatures.Core_1_0.geometryShader, OptionalDeviceExtensions.HasKHRShaderAtomicInt64, OptionalDeviceExtensions.HasImageAtomicInt64);
diff --git a/Engine/Source/Runtime/VulkanRHI/Private/VulkanDevice.h b/Engine/Source/Runtime/VulkanRHI/Private/VulkanDevice.h
--- a/Engine/Source/Runtime/VulkanRHI/Private/VulkanDevice.h
+++ b/Engine/Source/Runtime/VulkanRHI/Private/VulkanDevice.h
@@ -131,6 +131,8 @@ struct FOptionalVulkanDeviceExtensionProperties
 #endif // VULKAN_RHI_RAYTRACING
 
 	VkPhysicalDeviceFragmentShadingRateFeaturesKHR FragmentShadingRateFeatures;
+	VkPhysicalDeviceFragmentDensityMapFeaturesEXT FragmentDensityMapFeatures;
+	VkPhysicalDeviceFragmentDensityMap2FeaturesEXT FragmentDensityMap2Features;
 };
 
 class FVulkanPhysicalDeviceFeatures
@@ -244,6 +246,7 @@ public:
 	void InitGPU();
 
 	void CreateDevice(TArray<const ANSICHAR*>& DeviceLayers, FVulkanDeviceExtensionArray& UEExtensions);
+	void ChooseVariableRateShadingMethod();
 
 	void PrepareForDestroy();
 	void Destroy();
diff --git a/Engine/Source/Runtime/VulkanRHI/Private/VulkanExtensions.cpp b/Engine/Source/Runtime/VulkanRHI/Private/VulkanExtensions.cpp
--- a/Engine/Source/Runtime/VulkanRHI/Private/VulkanExtensions.cpp
+++ b/Engine/Source/Runtime/VulkanRHI/Private/VulkanExtensions.cpp
@@ -24,18 +24,6 @@
 // - To add engine support for a complex extension, but require some kind of external activation (eg plugin) you can use its EExtensionActivation state (see header definition).
 // - If an extension is supported on multiple platforms, it may be cleaner to include it here and simply disable its VULKAN_SUPPORTS_* value in the Vulkan platform header where it's not supported.
 
-
-
-enum class EVulkanVariableRateShadingPreference : uint8
-{
-	PreferFSR = 0, // This can be used to prefer FSR over FDM when both are available.
-	UseFSROnlyIfAvailable, // This will only print a message if FSR is not available.
-	RequireFSR, // This will print an error if FSR is not available.
-	PreferFDM, // This can be used to prefer FDM over FSR when both are available.
-	UseFDMOnlyIfAvailable, // This will only print a message if FDM is not available.
-	RequireFDM, // This will print an error if FDM is not available.
-};
-
 TAutoConsoleVariable<int32> GRHIAllow64bitShaderAtomicsCvar(
 	TEXT("r.Vulkan.Allow64bitShaderAtomics"),
 	1,
@@ -608,12 +596,14 @@ public:
 
 	virtual void PostPhysicalDeviceFeatures(FOptionalVulkanDeviceExtensions& ExtensionFlags) override final
 	{
+		VkPhysicalDeviceFragmentDensityMapFeaturesEXT& FragmentDensityMapFeatures = GetDeviceExtensionProperties().FragmentDensityMapFeatures;
 		bRequirementsPassed = (FragmentDensityMapFeatures.fragmentDensityMap == VK_TRUE);
 		ExtensionFlags.HasEXTFragmentDensityMap = bRequirementsPassed;
 	}
 
 	virtual void PrePhysicalDeviceFeatures(VkPhysicalDeviceFeatures2KHR& PhysicalDeviceFeatures2) override final
 	{
+		VkPhysicalDeviceFragmentDensityMapFeaturesEXT& FragmentDensityMapFeatures = GetDeviceExtensionProperties().FragmentDensityMapFeatures;
 		ZeroVulkanStruct(FragmentDensityMapFeatures, VK_STRUCTURE_TYPE_PHYSICAL_DEVICE_FRAGMENT_DENSITY_MAP_FEATURES_EXT);
 		AddToPNext(PhysicalDeviceFeatures2, FragmentDensityMapFeatures);
 	}
@@ -647,13 +637,13 @@ public:
 	{
 		if (bRequirementsPassed)
 		{
+			VkPhysicalDeviceFragmentDensityMapFeaturesEXT& FragmentDensityMapFeatures = GetDeviceExtensionProperties().FragmentDensityMapFeatures;
 			AddToPNext(DeviceCreateInfo, FragmentDensityMapFeatures);
 		}
 	}
 
 private:
 	VkPhysicalDeviceFragmentDensityMapPropertiesEXT FragmentDensityMapProperties;
-	VkPhysicalDeviceFragmentDensityMapFeaturesEXT FragmentDensityMapFeatures;
 };
 
 
@@ -671,12 +661,14 @@ public:
 
 	virtual void PrePhysicalDeviceFeatures(VkPhysicalDeviceFeatures2KHR& PhysicalDeviceFeatures2) override final
 	{
+		VkPhysicalDeviceFragmentDensityMap2FeaturesEXT& FragmentDensityMap2Features = GetDeviceExtensionProperties().FragmentDensityMap2Features;
 		ZeroVulkanStruct(FragmentDensityMap2Features, VK_STRUCTURE_TYPE_PHYSICAL_DEVICE_FRAGMENT_DENSITY_MAP_2_FEATURES_EXT);
 		AddToPNext(PhysicalDeviceFeatures2, FragmentDensityMap2Features);
 	}
 
 	virtual void PostPhysicalDeviceFeatures(FOptionalVulkanDeviceExtensions& ExtensionFlags) override final
 	{
+		VkPhysicalDeviceFragmentDensityMap2FeaturesEXT& FragmentDensityMap2Features = GetDeviceExtensionProperties().FragmentDensityMap2Features;
 		bRequirementsPassed = (FragmentDensityMap2Features.fragmentDensityMapDeferred == VK_TRUE);
 		ExtensionFlags.HasEXTFragmentDensityMap2 = bRequirementsPassed;
 
@@ -687,12 +679,10 @@ public:
 	{
 		if (bRequirementsPassed)
 		{
+			VkPhysicalDeviceFragmentDensityMap2FeaturesEXT& FragmentDensityMap2Features = GetDeviceExtensionProperties().FragmentDensityMap2Features;
 			AddToPNext(DeviceCreateInfo, FragmentDensityMap2Features);
 		}
 	}
-
-private:
-	VkPhysicalDeviceFragmentDensityMap2FeaturesEXT FragmentDensityMap2Features;
 };
 
 
@@ -1614,39 +1604,3 @@ FVulkanInstanceExtensionArray FVulkanInstanceExtension::GetUESupportedInstanceEx
 	return OutUEInstanceExtensions;
 }
 
-void ChooseVariableRateShadingMethod(FOptionalVulkanDeviceExtensions& ExtensionFlags, const VkPhysicalDeviceFragmentShadingRateFeaturesKHR& FragmentShadingRateFeatures)
-{
-	int32 VRSFormatPreference = GVulkanVariableRateShadingFormatCVar->GetInt();
-
-	if (VRSFormatPreference == (uint8)EVulkanVariableRateShadingPreference::PreferFSR && (FragmentShadingRateFeatures.attachmentFragmentShadingRate == VK_TRUE))
-	{
-		ExtensionFlags.HasEXTFragmentDensityMap = 0;
-		return;
-	}
-	if (VRSFormatPreference == (uint8)EVulkanVariableRateShadingPreference::PreferFDM && ExtensionFlags.HasEXTFragmentDensityMap)
-	{
-		ExtensionFlags.HasKHRFragmentShadingRate = 0;
-		return;
-	}
-	if (VRSFormatPreference == (uint8)EVulkanVariableRateShadingPreference::UseFSROnlyIfAvailable && !(FragmentShadingRateFeatures.attachmentFragmentShadingRate == VK_TRUE))
-	{
-		UE_LOG(LogVulkanRHI, Display, TEXT("Fragment Shading Rate was requested but is not available."));
-		return;
-	}
-	if (VRSFormatPreference == (uint8)EVulkanVariableRateShadingPreference::UseFDMOnlyIfAvailable && !ExtensionFlags.HasEXTFragmentDensityMap)
-	{
-		UE_LOG(LogVulkanRHI, Display, TEXT("Fragment Density Map was requested but is not available."));
-		return;
-	}
-	if (VRSFormatPreference == (uint8)EVulkanVariableRateShadingPreference::RequireFSR && !(FragmentShadingRateFeatures.attachmentFragmentShadingRate == VK_TRUE))
-	{
-		UE_LOG(LogVulkanRHI, Error, TEXT("Fragment Shading Rate was required but is not available."));
-		return;
-	}
-	if (VRSFormatPreference == (uint8)EVulkanVariableRateShadingPreference::RequireFDM && !ExtensionFlags.HasEXTFragmentDensityMap)
-	{
-		UE_LOG(LogVulkanRHI, Error, TEXT("Fragment Density Map was required but is not available."));
-		return;
-	}
-}
-
diff --git a/Engine/Source/Runtime/VulkanRHI/Private/VulkanExtensions.h b/Engine/Source/Runtime/VulkanRHI/Private/VulkanExtensions.h
--- a/Engine/Source/Runtime/VulkanRHI/Private/VulkanExtensions.h
+++ b/Engine/Source/Runtime/VulkanRHI/Private/VulkanExtensions.h
@@ -15,6 +15,16 @@
 
 class FVulkanDevice;
 
+enum class EVulkanVariableRateShadingPreference : uint8
+{
+	PreferFSR = 0, // This can be used to prefer FSR over FDM when both are available.
+	UseFSROnlyIfAvailable, // This will only print a message if FSR is not available.
+	RequireFSR, // This will print an error if FSR is not available.
+	PreferFDM, // This can be used to prefer FDM over FSR when both are available.
+	UseFDMOnlyIfAvailable, // This will only print a message if FDM is not available.
+	RequireFDM, // This will print an error if FDM is not available.
+};
+
 class FVulkanExtensionBase
 {
 public:
@@ -157,6 +167,4 @@ protected:
 
 // Helpers to create simple extensions that set a flag
 #define DEVICE_EXT_FLAG_SETTER(FLAG_NAME) [](FOptionalVulkanDeviceExtensions& ExtensionFlags) { ExtensionFlags.FLAG_NAME = 1; }
-#define INSTANCE_EXT_FLAG_SETTER(FLAG_NAME) [](FOptionalVulkanInstanceExtensions& ExtensionFlags) { ExtensionFlags.FLAG_NAME = 1; }
-
-void ChooseVariableRateShadingMethod(FOptionalVulkanDeviceExtensions& ExtensionFlags, const VkPhysicalDeviceFragmentShadingRateFeaturesKHR& FragmentShadingRateFeatures);
\ No newline at end of file
+#define INSTANCE_EXT_FLAG_SETTER(FLAG_NAME) [](FOptionalVulkanInstanceExtensions& ExtensionFlags) { ExtensionFlags.FLAG_NAME = 1; }
\ No newline at end of file
