Fix Vulkan validation error on Android when ASTC is not supported with linear tiling.
Hit this on ChromeOS device with Mesa driver on Intel GPU.

diff --git a/Engine/Source/Runtime/VulkanRHI/Private/Android/VulkanAndroidPlatform.cpp b/Engine/Source/Runtime/VulkanRHI/Private/Android/VulkanAndroidPlatform.cpp
--- a/Engine/Source/Runtime/VulkanRHI/Private/Android/VulkanAndroidPlatform.cpp
+++ b/Engine/Source/Runtime/VulkanRHI/Private/Android/VulkanAndroidPlatform.cpp
@@ -1724,6 +1724,9 @@ void FVulkanAndroidPlatform::SetupImageMemoryRequirementWorkaround(const FVulkan
 	}
 
 	// ASTC workarounds
+	VkFormatProperties formatProperties{};
+	VulkanRHI::vkGetPhysicalDeviceFormatProperties(InDevice.GetPhysicalHandle(), VK_FORMAT_ASTC_8x8_UNORM_BLOCK, &formatProperties);
+	if ((formatProperties.linearTilingFeatures & (VK_FORMAT_FEATURE_SAMPLED_IMAGE_FILTER_LINEAR_BIT | VK_FORMAT_FEATURE_TRANSFER_DST_BIT)) == (VK_FORMAT_FEATURE_SAMPLED_IMAGE_FILTER_LINEAR_BIT | VK_FORMAT_FEATURE_TRANSFER_DST_BIT))
 	{
 		ImageCreateInfo.flags = 0;
 		ImageCreateInfo.format = VK_FORMAT_ASTC_8x8_UNORM_BLOCK;
