VulkanRHI: Remove VK_IMAGE_USAGE_STORAGE_BIT from views in formats that might not support it (like srgb).

#rnx

diff --git a/Engine/Source/Runtime/VulkanRHI/Private/VulkanTexture.cpp b/Engine/Source/Runtime/VulkanRHI/Private/VulkanTexture.cpp
--- a/Engine/Source/Runtime/VulkanRHI/Private/VulkanTexture.cpp
+++ b/Engine/Source/Runtime/VulkanRHI/Private/VulkanTexture.cpp
@@ -1440,165 +1440,165 @@ FVulkanTexture::FVulkanTexture(FRHICommandListBase* RHICmdList, FVulkanDevice& I
 {
 	VULKAN_TRACK_OBJECT_CREATE(FVulkanTexture, this);
 
+	if (EnumHasAnyFlags(InCreateDesc.Flags, TexCreate_CPUReadback))
 	{
-		if (EnumHasAnyFlags(InCreateDesc.Flags, TexCreate_CPUReadback))
-		{
-			check(InCreateDesc.NumSamples == 1); //not implemented
-			check(InCreateDesc.ArraySize == 1);  //not implemented
-
-			CpuReadbackBuffer = new FVulkanCpuReadbackBuffer;
-			uint32 Size = 0;
-			for (uint32 Mip = 0; Mip < InCreateDesc.NumMips; ++Mip)
-			{
-				uint32 LocalSize;
-				GetMipSize(Mip, LocalSize);
-				CpuReadbackBuffer->MipOffsets[Mip] = Size;
-				CpuReadbackBuffer->MipSize[Mip] = LocalSize;
-				Size += LocalSize;
-			}
+		check(InCreateDesc.NumSamples == 1); //not implemented
+		check(InCreateDesc.ArraySize == 1);  //not implemented
 
-			VkBufferCreateInfo BufferCreateInfo;
-			ZeroVulkanStruct(BufferCreateInfo, VK_STRUCTURE_TYPE_BUFFER_CREATE_INFO);
-			BufferCreateInfo.size = Size;
-			BufferCreateInfo.usage = VK_BUFFER_USAGE_TRANSFER_DST_BIT;
-			VERIFYVULKANRESULT(VulkanRHI::vkCreateBuffer(InDevice.GetInstanceHandle(), &BufferCreateInfo, VULKAN_CPU_ALLOCATOR, &CpuReadbackBuffer->Buffer));
+		CpuReadbackBuffer = new FVulkanCpuReadbackBuffer;
+		uint32 Size = 0;
+		for (uint32 Mip = 0; Mip < InCreateDesc.NumMips; ++Mip)
+		{
+			uint32 LocalSize;
+			GetMipSize(Mip, LocalSize);
+			CpuReadbackBuffer->MipOffsets[Mip] = Size;
+			CpuReadbackBuffer->MipSize[Mip] = LocalSize;
+			Size += LocalSize;
+		}
 
-			// Set minimum alignment to 16 bytes, as some buffers are used with CPU SIMD instructions
-			const uint32 ForcedMinAlignment = 16u;
-			const EVulkanAllocationFlags AllocFlags = EVulkanAllocationFlags::HostCached | EVulkanAllocationFlags::AutoBind;
-			InDevice.GetMemoryManager().AllocateBufferMemory(Allocation, CpuReadbackBuffer->Buffer, AllocFlags, InCreateDesc.DebugName, ForcedMinAlignment);
+		VkBufferCreateInfo BufferCreateInfo;
+		ZeroVulkanStruct(BufferCreateInfo, VK_STRUCTURE_TYPE_BUFFER_CREATE_INFO);
+		BufferCreateInfo.size = Size;
+		BufferCreateInfo.usage = VK_BUFFER_USAGE_TRANSFER_DST_BIT;
+		VERIFYVULKANRESULT(VulkanRHI::vkCreateBuffer(InDevice.GetInstanceHandle(), &BufferCreateInfo, VULKAN_CPU_ALLOCATOR, &CpuReadbackBuffer->Buffer));
 
-			void* Memory = Allocation.GetMappedPointer(Device);
-			FMemory::Memzero(Memory, Size);
+		// Set minimum alignment to 16 bytes, as some buffers are used with CPU SIMD instructions
+		const uint32 ForcedMinAlignment = 16u;
+		const EVulkanAllocationFlags AllocFlags = EVulkanAllocationFlags::HostCached | EVulkanAllocationFlags::AutoBind;
+		InDevice.GetMemoryManager().AllocateBufferMemory(Allocation, CpuReadbackBuffer->Buffer, AllocFlags, InCreateDesc.DebugName, ForcedMinAlignment);
 
-			ImageOwnerType = EImageOwnerType::None;
-			ViewFormat = StorageFormat = UEToVkTextureFormat(InCreateDesc.Format, false);
+		void* Memory = Allocation.GetMappedPointer(Device);
+		FMemory::Memzero(Memory, Size);
 
-			// :todo-jn: Kept around temporarily for legacy defrag/eviction/stats
-			VulkanRHI::vkGetBufferMemoryRequirements(InDevice.GetInstanceHandle(), CpuReadbackBuffer->Buffer, &MemoryRequirements);
+		ImageOwnerType = EImageOwnerType::None;
+		ViewFormat = StorageFormat = UEToVkTextureFormat(InCreateDesc.Format, false);
 
-			return;
-		}
+		// :todo-jn: Kept around temporarily for legacy defrag/eviction/stats
+		VulkanRHI::vkGetBufferMemoryRequirements(InDevice.GetInstanceHandle(), CpuReadbackBuffer->Buffer, &MemoryRequirements);
 
-		ImageOwnerType = EImageOwnerType::LocalOwner;
+		return;
+	}
 
-		FImageCreateInfo ImageCreateInfo;
-		FVulkanTexture::GenerateImageCreateInfo(ImageCreateInfo, InDevice, InCreateDesc, &StorageFormat, &ViewFormat);
+	ImageOwnerType = EImageOwnerType::LocalOwner;
 
-		VERIFYVULKANRESULT(VulkanRHI::vkCreateImage(InDevice.GetInstanceHandle(), &ImageCreateInfo.ImageCreateInfo, VULKAN_CPU_ALLOCATOR, &Image));
+	FImageCreateInfo ImageCreateInfo;
+	FVulkanTexture::GenerateImageCreateInfo(ImageCreateInfo, InDevice, InCreateDesc, &StorageFormat, &ViewFormat);
 
-		// Fetch image size
-		VulkanRHI::vkGetImageMemoryRequirements(InDevice.GetInstanceHandle(), Image, &MemoryRequirements);
+	VERIFYVULKANRESULT(VulkanRHI::vkCreateImage(InDevice.GetInstanceHandle(), &ImageCreateInfo.ImageCreateInfo, VULKAN_CPU_ALLOCATOR, &Image));
 
-		VULKAN_SET_DEBUG_NAME(InDevice, VK_OBJECT_TYPE_IMAGE, Image, TEXT("%s:(FVulkanTexture*)0x%p"), InCreateDesc.DebugName ? InCreateDesc.DebugName : TEXT("?"), this);
+	// Fetch image size
+	VulkanRHI::vkGetImageMemoryRequirements(InDevice.GetInstanceHandle(), Image, &MemoryRequirements);
 
-		FullAspectMask = VulkanRHI::GetAspectMaskFromUEFormat(InCreateDesc.Format, true, true);
-		PartialAspectMask = VulkanRHI::GetAspectMaskFromUEFormat(InCreateDesc.Format, false, true);
+	VULKAN_SET_DEBUG_NAME(InDevice, VK_OBJECT_TYPE_IMAGE, Image, TEXT("%s:(FVulkanTexture*)0x%p"), InCreateDesc.DebugName ? InCreateDesc.DebugName : TEXT("?"), this);
 
-		// If VK_IMAGE_TILING_OPTIMAL is specified,
-		// memoryTypeBits in vkGetImageMemoryRequirements will become 1
-		// which does not support VK_MEMORY_PROPERTY_HOST_VISIBLE_BIT.
-		if (ImageCreateInfo.ImageCreateInfo.tiling != VK_IMAGE_TILING_OPTIMAL)
-		{
-			MemProps |= VK_MEMORY_PROPERTY_HOST_VISIBLE_BIT;
-		}
+	FullAspectMask = VulkanRHI::GetAspectMaskFromUEFormat(InCreateDesc.Format, true, true);
+	PartialAspectMask = VulkanRHI::GetAspectMaskFromUEFormat(InCreateDesc.Format, false, true);
 
-		const bool bRenderTarget = EnumHasAnyFlags(InCreateDesc.Flags, TexCreate_RenderTargetable | TexCreate_DepthStencilTargetable | TexCreate_ResolveTargetable);
-		const bool bUAV = EnumHasAnyFlags(InCreateDesc.Flags, TexCreate_UAV);
-		const bool bDynamic = EnumHasAnyFlags(InCreateDesc.Flags, TexCreate_Dynamic);
-		const bool bExternal = EnumHasAnyFlags(InCreateDesc.Flags, TexCreate_External);
+	// If VK_IMAGE_TILING_OPTIMAL is specified,
+	// memoryTypeBits in vkGetImageMemoryRequirements will become 1
+	// which does not support VK_MEMORY_PROPERTY_HOST_VISIBLE_BIT.
+	if (ImageCreateInfo.ImageCreateInfo.tiling != VK_IMAGE_TILING_OPTIMAL)
+	{
+		MemProps |= VK_MEMORY_PROPERTY_HOST_VISIBLE_BIT;
+	}
 
-		VkMemoryPropertyFlags MemoryFlags = VK_MEMORY_PROPERTY_DEVICE_LOCAL_BIT;
+	const bool bRenderTarget = EnumHasAnyFlags(InCreateDesc.Flags, TexCreate_RenderTargetable | TexCreate_DepthStencilTargetable | TexCreate_ResolveTargetable);
+	const bool bUAV = EnumHasAnyFlags(InCreateDesc.Flags, TexCreate_UAV);
+	const bool bDynamic = EnumHasAnyFlags(InCreateDesc.Flags, TexCreate_Dynamic);
+	const bool bExternal = EnumHasAnyFlags(InCreateDesc.Flags, TexCreate_External);
 
-		bool bMemoryless = EnumHasAnyFlags(InCreateDesc.Flags, TexCreate_Memoryless) && InDevice.GetDeviceMemoryManager().SupportsMemoryless();
-		if (bMemoryless)
-		{
-			if (ensureMsgf(bRenderTarget, TEXT("Memoryless surfaces can only be used for render targets")))
-			{
-				MemoryFlags |= VK_MEMORY_PROPERTY_LAZILY_ALLOCATED_BIT;
-			}
-			else
-			{
-				bMemoryless = false;
-			}
-		}
+	VkMemoryPropertyFlags MemoryFlags = VK_MEMORY_PROPERTY_DEVICE_LOCAL_BIT;
 
-		const bool bIsTransientResource = (InTransientHeapAllocation != nullptr);
-		if (bIsTransientResource)
+	bool bMemoryless = EnumHasAnyFlags(InCreateDesc.Flags, TexCreate_Memoryless) && InDevice.GetDeviceMemoryManager().SupportsMemoryless();
+	if (bMemoryless)
+	{
+		if (ensureMsgf(bRenderTarget, TEXT("Memoryless surfaces can only be used for render targets")))
 		{
-			check(!bMemoryless);
-			check(InTransientHeapAllocation->Offset % MemoryRequirements.alignment == 0);
-			check(InTransientHeapAllocation->Size >= MemoryRequirements.size);
-			Allocation = FVulkanTransientHeap::GetVulkanAllocation(*InTransientHeapAllocation);
+			MemoryFlags |= VK_MEMORY_PROPERTY_LAZILY_ALLOCATED_BIT;
 		}
 		else
 		{
-			EVulkanAllocationMetaType MetaType = (bRenderTarget || bUAV) ? EVulkanAllocationMetaImageRenderTarget : EVulkanAllocationMetaImageOther;
+			bMemoryless = false;
+		}
+	}
+
+	const bool bIsTransientResource = (InTransientHeapAllocation != nullptr);
+	if (bIsTransientResource)
+	{
+		check(!bMemoryless);
+		check(InTransientHeapAllocation->Offset % MemoryRequirements.alignment == 0);
+		check(InTransientHeapAllocation->Size >= MemoryRequirements.size);
+		Allocation = FVulkanTransientHeap::GetVulkanAllocation(*InTransientHeapAllocation);
+	}
+	else
+	{
+		EVulkanAllocationMetaType MetaType = (bRenderTarget || bUAV) ? EVulkanAllocationMetaImageRenderTarget : EVulkanAllocationMetaImageOther;
 #if VULKAN_SUPPORTS_DEDICATED_ALLOCATION
-			extern int32 GVulkanEnableDedicatedImageMemory;
-			// Per https://developer.nvidia.com/what%E2%80%99s-your-vulkan-memory-type
-			VkDeviceSize SizeToBeConsideredForDedicated = 12 * 1024 * 1024;
-			if ((bRenderTarget || MemoryRequirements.size >= SizeToBeConsideredForDedicated) && !bMemoryless && GVulkanEnableDedicatedImageMemory)
+		extern int32 GVulkanEnableDedicatedImageMemory;
+		// Per https://developer.nvidia.com/what%E2%80%99s-your-vulkan-memory-type
+		VkDeviceSize SizeToBeConsideredForDedicated = 12 * 1024 * 1024;
+		if ((bRenderTarget || MemoryRequirements.size >= SizeToBeConsideredForDedicated) && !bMemoryless && GVulkanEnableDedicatedImageMemory)
+		{
+			if (!InDevice.GetMemoryManager().AllocateDedicatedImageMemory(Allocation, this, Image, MemoryRequirements, MemoryFlags, MetaType, bExternal, __FILE__, __LINE__))
 			{
-				if (!InDevice.GetMemoryManager().AllocateDedicatedImageMemory(Allocation, this, Image, MemoryRequirements, MemoryFlags, MetaType, bExternal, __FILE__, __LINE__))
-				{
-					checkNoEntry();
-				}
+				checkNoEntry();
 			}
-			else
+		}
+		else
 #endif
+		{
+			if (!InDevice.GetMemoryManager().AllocateImageMemory(Allocation, this, MemoryRequirements, MemoryFlags, MetaType, bExternal, __FILE__, __LINE__))
 			{
-				if (!InDevice.GetMemoryManager().AllocateImageMemory(Allocation, this, MemoryRequirements, MemoryFlags, MetaType, bExternal, __FILE__, __LINE__))
-				{
-					checkNoEntry();
-				}
+				checkNoEntry();
 			}
-
-			// update rhi stats
-			VulkanTextureAllocated(GetDesc(), Allocation.Size);
 		}
-		Allocation.BindImage(Device, Image);
 
-		Tiling = ImageCreateInfo.ImageCreateInfo.tiling;
-		check(Tiling == VK_IMAGE_TILING_LINEAR || Tiling == VK_IMAGE_TILING_OPTIMAL);
+		// update rhi stats
+		VulkanTextureAllocated(GetDesc(), Allocation.Size);
+	}
+	Allocation.BindImage(Device, Image);
+
+	Tiling = ImageCreateInfo.ImageCreateInfo.tiling;
+	check(Tiling == VK_IMAGE_TILING_LINEAR || Tiling == VK_IMAGE_TILING_OPTIMAL);
 
-		const VkImageLayout InitialLayout = GetInitialLayoutFromRHIAccess(InCreateDesc.InitialState, bRenderTarget && IsDepthOrStencilAspect(), SupportsSampling());
-		const bool bDoInitialClear = VKHasAnyFlags(ImageCreateInfo.ImageCreateInfo.usage, VK_IMAGE_USAGE_SAMPLED_BIT) && EnumHasAnyFlags(InCreateDesc.Flags, TexCreate_RenderTargetable | TexCreate_DepthStencilTargetable);
+	const VkImageLayout InitialLayout = GetInitialLayoutFromRHIAccess(InCreateDesc.InitialState, bRenderTarget && IsDepthOrStencilAspect(), SupportsSampling());
+	const bool bDoInitialClear = VKHasAnyFlags(ImageCreateInfo.ImageCreateInfo.usage, VK_IMAGE_USAGE_SAMPLED_BIT) && EnumHasAnyFlags(InCreateDesc.Flags, TexCreate_RenderTargetable | TexCreate_DepthStencilTargetable);
 
-		if (InitialLayout != VK_IMAGE_LAYOUT_UNDEFINED || bDoInitialClear)
+	if (InitialLayout != VK_IMAGE_LAYOUT_UNDEFINED || bDoInitialClear)
+	{
+		if (RHICmdList && RHICmdList->IsTopOfPipe())
+		{
+			ALLOC_COMMAND_CL(*RHICmdList, FRHICommandSetInitialImageState)(this, InitialLayout, false, bDoInitialClear, InCreateDesc.ClearValue, bIsTransientResource);
+		}
+		else
 		{
-			if (RHICmdList && RHICmdList->IsTopOfPipe())
+			RHICmdList = &FRHICommandListExecutor::GetImmediateCommandList();
+			if (!IsInRenderingThread() || (RHICmdList->Bypass() || !IsRunningRHIInSeparateThread()))
 			{
-				ALLOC_COMMAND_CL(*RHICmdList, FRHICommandSetInitialImageState)(this, InitialLayout, false, bDoInitialClear, InCreateDesc.ClearValue, bIsTransientResource);
+				SetInitialImageState(Device->GetImmediateContext(), InitialLayout, bDoInitialClear, InCreateDesc.ClearValue, bIsTransientResource);
 			}
 			else
 			{
-				RHICmdList = &FRHICommandListExecutor::GetImmediateCommandList();
-				if (!IsInRenderingThread() || (RHICmdList->Bypass() || !IsRunningRHIInSeparateThread()))
-				{
-					SetInitialImageState(Device->GetImmediateContext(), InitialLayout, bDoInitialClear, InCreateDesc.ClearValue, bIsTransientResource);
-				}
-				else
-				{
-					check(IsInRenderingThread());
-					ALLOC_COMMAND_CL(*RHICmdList, FRHICommandSetInitialImageState)(this, InitialLayout, false, bDoInitialClear, InCreateDesc.ClearValue, bIsTransientResource);
-				}
+				check(IsInRenderingThread());
+				ALLOC_COMMAND_CL(*RHICmdList, FRHICommandSetInitialImageState)(this, InitialLayout, false, bDoInitialClear, InCreateDesc.ClearValue, bIsTransientResource);
 			}
 		}
-
-		DefaultLayout = InitialLayout;
 	}
 
+	DefaultLayout = InitialLayout;
+
 	LLM_SCOPE_VULKAN(ELLMTagVulkan::VulkanTextures);
 	const VkImageViewType ViewType = GetViewType();
+	const bool bIsSRGB = EnumHasAllFlags(InCreateDesc.Flags, TexCreate_SRGB);
 	if (ViewFormat == VK_FORMAT_UNDEFINED)
 	{
 		StorageFormat = UEToVkTextureFormat(InCreateDesc.Format, false);
-		ViewFormat = UEToVkTextureFormat(InCreateDesc.Format, EnumHasAllFlags(InCreateDesc.Flags, TexCreate_SRGB));
+		ViewFormat = UEToVkTextureFormat(InCreateDesc.Format, bIsSRGB);
 		checkf(StorageFormat != VK_FORMAT_UNDEFINED, TEXT("Pixel Format %d not defined!"), (int32)InCreateDesc.Format);
 	}
 
 	const VkDescriptorType DescriptorType = SupportsSampling() ? VK_DESCRIPTOR_TYPE_SAMPLED_IMAGE : VK_DESCRIPTOR_TYPE_STORAGE_IMAGE;
+	const VkImageUsageFlags SRVUsage = bIsSRGB ? (ImageCreateInfo.ImageCreateInfo.usage & ~VK_IMAGE_USAGE_STORAGE_BIT) : ImageCreateInfo.ImageCreateInfo.usage;
 	if (ViewType != VK_IMAGE_VIEW_TYPE_MAX_ENUM)
 	{
 		DefaultView = (new FVulkanView(InDevice, DescriptorType))->InitAsTextureView(
@@ -1612,6 +1612,7 @@ FVulkanTexture::FVulkanTexture(FRHICommandListBase* RHICmdList, FVulkanDevice& I
 			, 0
 			, GetNumberOfArrayLevels()
 			, !SupportsSampling()
+			, SRVUsage
 		);
 	}
 
diff --git a/Engine/Source/Runtime/VulkanRHI/Private/VulkanViewport.cpp b/Engine/Source/Runtime/VulkanRHI/Private/VulkanViewport.cpp
--- a/Engine/Source/Runtime/VulkanRHI/Private/VulkanViewport.cpp
+++ b/Engine/Source/Runtime/VulkanRHI/Private/VulkanViewport.cpp
@@ -373,6 +373,7 @@ FVulkanFramebuffer::FVulkanFramebuffer(FVulkanDevice& Device, const FRHISetRende
 				, ArraySliceIndex
 				, NumArraySlices
 				, true
+				, VK_IMAGE_USAGE_COLOR_ATTACHMENT_BIT
 			);
 		}
 		else if (Texture->GetViewType() == VK_IMAGE_VIEW_TYPE_CUBE)
@@ -391,6 +392,7 @@ FVulkanFramebuffer::FVulkanFramebuffer(FVulkanDevice& Device, const FRHISetRende
 				, InRTInfo.ColorRenderTarget[Index].ArraySliceIndex
 				, 1
 				, true
+				, VK_IMAGE_USAGE_COLOR_ATTACHMENT_BIT
 			);
 		}
 		else if (Texture->GetViewType() == VK_IMAGE_VIEW_TYPE_3D)
