Vulkan extra copy

diff --git a/Engine/Source/Runtime/VulkanRHI/Private/VulkanBarriers.cpp b/Engine/Source/Runtime/VulkanRHI/Private/VulkanBarriers.cpp
--- a/Engine/Source/Runtime/VulkanRHI/Private/VulkanBarriers.cpp
+++ b/Engine/Source/Runtime/VulkanRHI/Private/VulkanBarriers.cpp
@@ -1036,7 +1036,7 @@ public:
 			{
 				const FVulkanPipelineBarrier::ImageBarrierExtraData* RemainingExtras = &Data->ImageBarrierExtras[Index];
 				FVulkanTexture* Texture = RemainingExtras->BaseTexture;
-				check(Texture->Image != VK_NULL_HANDLE);  // coming in, we should always have an image
+				check(Texture->Image != VK_NULL_HANDLE || !Texture->IsImageOwner());  // coming in, we should always have an image, except backbuffer with r.Vulkan.DelayAcquireBackBuffer=2
 
 				const int32 TargetIndex = ImageBatchStartIndex + Index;
 				ImageBarrierType* RemainingBarriers = &ImageBarriers[TargetIndex];
