Improvements to Vulkan Android PSO precaching hit rate.

#localization none
#tests local client, Android QA passes

diff --git a/Engine/Source/Runtime/Engine/Private/Components/StaticMeshComponent.cpp b/Engine/Source/Runtime/Engine/Private/Components/StaticMeshComponent.cpp
--- a/Engine/Source/Runtime/Engine/Private/Components/StaticMeshComponent.cpp
+++ b/Engine/Source/Runtime/Engine/Private/Components/StaticMeshComponent.cpp
@@ -1670,13 +1670,12 @@ void UStaticMeshComponent::CollectPSOPrecacheData(const FPSOPrecacheParams& Base
 	}
 
 	int32 LightMapCoordinateIndex = StaticMesh->GetLightMapCoordinateIndex();
-	// FIXME: Need a precise per-LOD test
-	bool bOverrideColorVertexBuffer = LODData.Num() != 0 && LODData[0].OverrideVertexColors != nullptr;
-	
-	auto SMC_GetElements = [LightMapCoordinateIndex, bOverrideColorVertexBuffer](const FStaticMeshLODResources& LODRenderData, int32 LODIndex, bool bSupportsManualVertexFetch, FVertexDeclarationElementList& Elements)
+
+	auto SMC_GetElements = [LightMapCoordinateIndex, &LODData = this->LODData](const FStaticMeshLODResources& LODRenderData, int32 LODIndex, bool bSupportsManualVertexFetch, FVertexDeclarationElementList& Elements)
 	{
 		int32 NumTexCoords = (int32)LODRenderData.VertexBuffers.StaticMeshVertexBuffer.GetNumTexCoords();
 		int32 LODLightMapCoordinateIndex = LightMapCoordinateIndex < NumTexCoords ? LightMapCoordinateIndex : NumTexCoords - 1;
+		bool bOverrideColorVertexBuffer = LODIndex < LODData.Num() && LODData[LODIndex].OverrideVertexColors != nullptr;
 		FLocalVertexFactory::FDataType Data;
 		InitStaticMeshVertexFactoryComponents(LODRenderData.VertexBuffers, nullptr /*VertexFactory*/, LODLightMapCoordinateIndex, bOverrideColorVertexBuffer, Data);
 		FLocalVertexFactory::GetVertexElements(GMaxRHIFeatureLevel, EVertexInputStreamType::Default, bSupportsManualVertexFetch, Data, Elements);
diff --git a/Engine/Source/Runtime/VulkanRHI/Private/VulkanRHI.cpp b/Engine/Source/Runtime/VulkanRHI/Private/VulkanRHI.cpp
--- a/Engine/Source/Runtime/VulkanRHI/Private/VulkanRHI.cpp
+++ b/Engine/Source/Runtime/VulkanRHI/Private/VulkanRHI.cpp
@@ -2046,7 +2046,7 @@ IRHITransientResourceAllocator* FVulkanDynamicRHI::RHICreateTransientResourceAll
 
 uint32 FVulkanDynamicRHI::GetPrecachePSOHashVersion()
 {
-	static const uint32 PrecacheHashVersion = 1;
+	static const uint32 PrecacheHashVersion = 2;
 	return PrecacheHashVersion;
 }
 
@@ -2141,8 +2141,9 @@ uint64 FVulkanDynamicRHI::RHIComputePrecachePSOHash(const FGraphicsPipelineState
 		uint32							RenderTargetsEnabled;
 		FGraphicsPipelineStateInitializer::TRenderTargetFormats	RenderTargetFormats;
 		FGraphicsPipelineStateInitializer::TRenderTargetFlags RenderTargetFlags;
-		EPixelFormat					DepthStencilTargetFormat;
-		ETextureCreateFlags				DepthStencilTargetFlag;
+// AJB: temporarily disabling depth stencil properties as they do not appear to be required and it causes us to miss some permutations.
+//		EPixelFormat					DepthStencilTargetFormat;
+//		ETextureCreateFlags				DepthStencilTargetFlag;
 		uint16							NumSamples;
 		ESubpassHint					SubpassHint;
 		uint8							SubpassIndex;
@@ -2161,8 +2162,8 @@ uint64 FVulkanDynamicRHI::RHIComputePrecachePSOHash(const FGraphicsPipelineState
 	HashKey.RenderTargetsEnabled = Initializer.RenderTargetsEnabled;
 	HashKey.RenderTargetFormats = Initializer.RenderTargetFormats;
 	HashKey.RenderTargetFlags = Initializer.RenderTargetFlags;
-	HashKey.DepthStencilTargetFormat = Initializer.DepthStencilTargetFormat;
-	HashKey.DepthStencilTargetFlag = Initializer.DepthStencilTargetFlag;
+//	HashKey.DepthStencilTargetFormat = Initializer.DepthStencilTargetFormat;
+//	HashKey.DepthStencilTargetFlag = Initializer.DepthStencilTargetFlag;
 	HashKey.NumSamples = Initializer.NumSamples;
 	HashKey.SubpassHint = Initializer.SubpassHint;
 	HashKey.SubpassIndex = Initializer.SubpassIndex;
@@ -2177,7 +2178,7 @@ uint64 FVulkanDynamicRHI::RHIComputePrecachePSOHash(const FGraphicsPipelineState
 	{
 		Flags = Flags & FGraphicsPipelineStateInitializer::RelevantRenderTargetFlagMask;
 	}
-	HashKey.DepthStencilTargetFlag = (HashKey.DepthStencilTargetFlag & FGraphicsPipelineStateInitializer::RelevantDepthStencilFlagMask);
+// 	HashKey.DepthStencilTargetFlag = (HashKey.DepthStencilTargetFlag & FGraphicsPipelineStateInitializer::RelevantDepthStencilFlagMask);
 
 	return CityHash64((const char*)&HashKey, sizeof(FNonStateHashKey));
 }
diff --git a/Engine/Source/Runtime/VulkanRHI/Private/VulkanResources.h b/Engine/Source/Runtime/VulkanRHI/Private/VulkanResources.h
--- a/Engine/Source/Runtime/VulkanRHI/Private/VulkanResources.h
+++ b/Engine/Source/Runtime/VulkanRHI/Private/VulkanResources.h
@@ -67,8 +67,9 @@ class FVulkanVertexDeclaration : public FRHIVertexDeclaration
 public:
 	FVertexDeclarationElementList Elements;
 	uint32 Hash;
+	uint32 HashNoStrides;
 
-	FVulkanVertexDeclaration(const FVertexDeclarationElementList& InElements, uint32 InHash);
+	FVulkanVertexDeclaration(const FVertexDeclarationElementList& InElements, uint32 InHash, uint32 InHashNoStrides);
 
 	virtual bool GetInitializer(FVertexDeclarationElementList& Out) final override
 	{
@@ -78,7 +79,7 @@ public:
 
 	static void EmptyCache();
 
-	virtual uint32 GetPrecachePSOHash() const final override { return Hash; }
+	virtual uint32 GetPrecachePSOHash() const final override { return HashNoStrides; }
 };
 
 struct FGfxPipelineDesc;
diff --git a/Engine/Source/Runtime/VulkanRHI/Private/VulkanVertexDeclaration.cpp b/Engine/Source/Runtime/VulkanRHI/Private/VulkanVertexDeclaration.cpp
--- a/Engine/Source/Runtime/VulkanRHI/Private/VulkanVertexDeclaration.cpp
+++ b/Engine/Source/Runtime/VulkanRHI/Private/VulkanVertexDeclaration.cpp
@@ -11,11 +11,20 @@ struct FVulkanVertexDeclarationKey
 	FVertexDeclarationElementList VertexElements;
 
 	uint32 Hash;
+	uint32 HashNoStride;
 
 	explicit FVulkanVertexDeclarationKey(const FVertexDeclarationElementList& InElements)
 		: VertexElements(InElements)
 	{
 		Hash = FCrc::MemCrc_DEPRECATED(VertexElements.GetData(), VertexElements.Num()*sizeof(FVertexElement));
+
+		HashNoStride = 0;
+		for (int32 ElementIndex = 0; ElementIndex < VertexElements.Num(); ElementIndex++)
+		{
+			FVertexElement Element = VertexElements[ElementIndex];
+			Element.Stride = 0;
+			HashNoStride = FCrc::MemCrc32(&Element, sizeof(Element), HashNoStride);
+		}
 	}
 };
 
@@ -30,9 +39,10 @@ bool operator==(const FVulkanVertexDeclarationKey& A, const FVulkanVertexDeclara
 			&& !memcmp(A.VertexElements.GetData(), B.VertexElements.GetData(), A.VertexElements.Num() * sizeof(FVertexElement)));
 }
 
-FVulkanVertexDeclaration::FVulkanVertexDeclaration(const FVertexDeclarationElementList& InElements, uint32 InHash) 
+FVulkanVertexDeclaration::FVulkanVertexDeclaration(const FVertexDeclarationElementList& InElements, uint32 InHash, uint32 InHashNoStrides) 
 	: Elements(InElements)
 	, Hash(InHash)
+	, HashNoStrides(InHashNoStrides)
 {
 }
 
@@ -52,7 +62,7 @@ FVertexDeclarationRHIRef FVulkanDynamicRHI::RHICreateVertexDeclaration(const FVe
 	FVertexDeclarationRHIRef* VertexDeclarationRefPtr = GVertexDeclarationCache.Find(Key);
 	if (VertexDeclarationRefPtr == nullptr)
 	{
-		VertexDeclarationRefPtr = &GVertexDeclarationCache.Add(Key, new FVulkanVertexDeclaration(Elements, Key.Hash));
+		VertexDeclarationRefPtr = &GVertexDeclarationCache.Add(Key, new FVulkanVertexDeclaration(Elements, Key.Hash, Key.HashNoStride));
 	}
 
 	check(VertexDeclarationRefPtr);
