Vulkan also supports barycentrics.

diff --git a/Engine/Config/VulkanPC/DataDrivenPlatformInfo.ini b/Engine/Config/VulkanPC/DataDrivenPlatformInfo.ini
--- a/Engine/Config/VulkanPC/DataDrivenPlatformInfo.ini
+++ b/Engine/Config/VulkanPC/DataDrivenPlatformInfo.ini
@@ -255,6 +255,8 @@ bSupportsNNEShaders=true
 bSupportsUniformBufferObjects = true
 bSupportsRealTypes=Unsupported
 
+SupportsBarycentricsSemantic=RuntimeGuaranteed
+
 
 [ShaderPlatform VULKAN_PCES3_1]
 Language=Vulkan
diff --git a/Engine/Source/Runtime/VulkanRHI/Private/VulkanRHI.cpp b/Engine/Source/Runtime/VulkanRHI/Private/VulkanRHI.cpp
--- a/Engine/Source/Runtime/VulkanRHI/Private/VulkanRHI.cpp
+++ b/Engine/Source/Runtime/VulkanRHI/Private/VulkanRHI.cpp
@@ -513,6 +513,8 @@ FVulkanDynamicRHI::FVulkanDynamicRHI()
 	GRHITransitionPrivateData_AlignInBytes = alignof(FVulkanPipelineBarrier);
 	GConfig->GetInt(TEXT("TextureStreaming"), TEXT("PoolSizeVRAMPercentage"), GPoolSizeVRAMPercentage, GEngineIni);
 
+	GRHIGlobals.SupportsBarycentricsSemantic = true;
+
 	static const auto CVarPSOPrecaching = IConsoleManager::Get().FindConsoleVariable(TEXT("r.PSOPrecaching"));
 
 	GRHISupportsPSOPrecaching = FVulkanChunkedPipelineCacheManager::IsEnabled() && (CVarPSOPrecaching && CVarPSOPrecaching->GetInt() != 0) && CVarAllowVulkanPSOPrecache.GetValueOnAnyThread();
