VulkanRHI: Handle the ERHIAccess of "CopySrc|SRVCompute" by using the General layout. (only used in WaitForRayTracingScene where nothing is touched anyhow)

#rnx

diff --git a/Engine/Source/Runtime/VulkanRHI/Private/VulkanBarriers.cpp b/Engine/Source/Runtime/VulkanRHI/Private/VulkanBarriers.cpp
--- a/Engine/Source/Runtime/VulkanRHI/Private/VulkanBarriers.cpp
+++ b/Engine/Source/Runtime/VulkanRHI/Private/VulkanBarriers.cpp
@@ -387,17 +387,31 @@ static void GetVkStageAndAccessFlags(ERHIAccess RHIAccess, FRHITransitionInfo::E
 		ProcessedRHIFlags |= (uint32)ERHIAccess::UAVCompute;
 	}
 
+	// ResolveSrc is used when doing a resolve via RHICopyToResolveTarget. For us, it's the same as CopySrc.
 	if (EnumHasAnyFlags(RHIAccess, ERHIAccess::CopySrc | ERHIAccess::ResolveSrc))
 	{
-		// ResolveSrc is used when doing a resolve via RHICopyToResolveTarget. For us, it's the same as CopySrc.
-		StageFlags = VK_PIPELINE_STAGE_TRANSFER_BIT;
-		AccessFlags = VK_ACCESS_TRANSFER_READ_BIT;
+		// If this is requested for a texture, behavior will depend on if we're combined with other flags
 		if (ResourceType == FRHITransitionInfo::EType::Texture)
 		{
-			// If this is requested for a texture, make sure it's not combined with other access flags which require a different layout. It's important
-			// that this block is last, so that if any other flags set the layout before, we trigger the assert below.
-			check(Layout == VK_IMAGE_LAYOUT_UNDEFINED);
-			Layout = VK_IMAGE_LAYOUT_TRANSFER_SRC_OPTIMAL;
+			// If no other RHIAccess is mixed in with our CopySrc, then use proper TRANSFER_SRC layout
+			if (Layout == VK_IMAGE_LAYOUT_UNDEFINED)
+			{
+				Layout = VK_IMAGE_LAYOUT_TRANSFER_SRC_OPTIMAL;
+				StageFlags = VK_PIPELINE_STAGE_TRANSFER_BIT;
+				AccessFlags = VK_ACCESS_TRANSFER_READ_BIT;
+			}
+			else
+			{
+				// If anything else is mixed in with the CopySrc, then go to the "catch all" GENERAL layout
+				Layout = VK_IMAGE_LAYOUT_GENERAL;
+				StageFlags |= VK_PIPELINE_STAGE_TRANSFER_BIT;
+				AccessFlags |= VK_ACCESS_TRANSFER_READ_BIT;
+			}
+		}
+		else
+		{
+			StageFlags = VK_PIPELINE_STAGE_TRANSFER_BIT;
+			AccessFlags = VK_ACCESS_TRANSFER_READ_BIT;
 		}
 
 		ProcessedRHIFlags |= (uint32)(ERHIAccess::CopySrc | ERHIAccess::ResolveSrc);
