- Fix validation error about invalidate with size that is not multiple of nonCoherentAtomSize
- Fix validation error on image copies for low mips
- Filter some false-positive validation errors for dual source blending
- Update some existing filter to changes strings in validation layers

diff --git a/Engine/Source/Runtime/VulkanRHI/Private/VulkanDebug.cpp b/Engine/Source/Runtime/VulkanRHI/Private/VulkanDebug.cpp
--- a/Engine/Source/Runtime/VulkanRHI/Private/VulkanDebug.cpp
+++ b/Engine/Source/Runtime/VulkanRHI/Private/VulkanDebug.cpp
@@ -94,6 +94,12 @@ static VkBool32 VKAPI_PTR DebugReportFunction(
 				//return VK_FALSE;
 			}
 		}
+		if (FCStringAnsi::Strstr(Msg, "VUID-RuntimeSpirv-Fragment-06427"))
+		{
+			// "vkCmdDrawIndexed: Dual source blend mode is used, but the number of written fragment shader output attachment (2) is greater than maxFragmentDualSrcAttachments (1)"
+			// Incorrectly generated in Validation layers 1.3.250.1, no longer shows up in 1.3.268.0
+			return VK_FALSE;
+		}
 
 		MsgPrefix = "ERROR";
 	}
@@ -252,7 +258,8 @@ static VkBool32 DebugUtilsCallback(VkDebugUtilsMessageSeverityFlagBitsEXT MsgSev
 		return VK_FALSE;
 	}
 
-	if (!FCStringAnsi::Strcmp(CallbackData->pMessageIdName, "UNASSIGNED-CoreValidation-Shader-OutputNotConsumed"))
+	if (!FCStringAnsi::Strcmp(CallbackData->pMessageIdName, "UNASSIGNED-CoreValidation-Shader-OutputNotConsumed") ||
+		!FCStringAnsi::Strcmp(CallbackData->pMessageIdName, "Undefined-Value-ShaderOutputNotConsumed"))
 	{
 		// Warning: *** [Warning:Validation-1(UNASSIGNED-CoreValidation-Shader-OutputNotConsumed)] fragment shader writes to output location 0 with no matching attachment
 		return VK_FALSE;
@@ -272,12 +279,12 @@ static VkBool32 DebugUtilsCallback(VkDebugUtilsMessageSeverityFlagBitsEXT MsgSev
 		// *** CreateDevice(): Attempting to enable deprecated extension VK_KHR_get_memory_requirements2, but this extension has been promoted to VK_VERSION_1_1.
 		return VK_FALSE;
 	}
-	else if (FCStringAnsi::Strstr(CallbackData->pMessage, "The SPIR-V Extension (SPV_GOOGLE_hlsl_functionality1) was declared") != nullptr)
+	else if (FCStringAnsi::Strstr(CallbackData->pMessage, "SPV_GOOGLE_hlsl_functionality1") != nullptr)
 	{
 		// *** [Error:Validation(VUID-VkShaderModuleCreateInfo-pCode-04147)] vkCreateShaderModule(): The SPIR-V Extension (SPV_GOOGLE_hlsl_functionality1) was declared, but none of the requirements were met to use it.
 		return VK_FALSE;
 	 }
-	else if (FCStringAnsi::Strstr(CallbackData->pMessage, "The SPIR-V Extension (SPV_GOOGLE_user_type) was declared") != nullptr)
+	else if (FCStringAnsi::Strstr(CallbackData->pMessage, "SPV_GOOGLE_user_type") != nullptr)
 	{
 		// *** [Error:Validation(VUID-VkShaderModuleCreateInfo-pCode-04147)] vkCreateShaderModule(): The SPIR-V Extension (SPV_GOOGLE_user_type) was declared, but none of the requirements were met to use it.
 		return VK_FALSE;
diff --git a/Engine/Source/Runtime/VulkanRHI/Private/VulkanMemory.cpp b/Engine/Source/Runtime/VulkanRHI/Private/VulkanMemory.cpp
--- a/Engine/Source/Runtime/VulkanRHI/Private/VulkanMemory.cpp
+++ b/Engine/Source/Runtime/VulkanRHI/Private/VulkanMemory.cpp
@@ -3978,13 +3978,33 @@ namespace VulkanRHI
 		return false;
 	}
 
+	static VkDeviceSize AdjustToNonCoherentAtomSize(const VkDeviceSize RequestedOffset, const VkDeviceSize RequestedSize, const VkDeviceSize WholeAllocationSize, const VkDeviceSize NonCoherentAtomSize)
+	{
+		if (RequestedSize == VK_WHOLE_SIZE)
+		{
+			return RequestedSize;
+		}
+		
+		check(RequestedOffset + RequestedSize <= WholeAllocationSize);
+		
+		return FMath::Min(WholeAllocationSize - RequestedOffset, AlignArbitrary(RequestedSize, NonCoherentAtomSize));
+	}
+
 	void FVulkanSubresourceAllocator::Flush(VkDeviceSize Offset, VkDeviceSize AllocationSize)
 	{
-		MemoryAllocation->FlushMappedMemory(Offset, AllocationSize);
+		if (!MemoryAllocation->IsCoherent() || GForceCoherent != 0)
+		{
+			const VkDeviceSize NonCoherentAtomSize = Owner->GetParent()->GetLimits().nonCoherentAtomSize;
+			MemoryAllocation->FlushMappedMemory(Offset, AdjustToNonCoherentAtomSize(Offset, AllocationSize, MemoryAllocation->GetSize(), NonCoherentAtomSize));
+		}
 	}
 	void FVulkanSubresourceAllocator::Invalidate(VkDeviceSize Offset, VkDeviceSize AllocationSize)
 	{
-		MemoryAllocation->InvalidateMappedMemory(Offset, AllocationSize);
+		if (!MemoryAllocation->IsCoherent() || GForceCoherent != 0)
+		{
+			const VkDeviceSize NonCoherentAtomSize = Owner->GetParent()->GetLimits().nonCoherentAtomSize;
+			MemoryAllocation->InvalidateMappedMemory(Offset, AdjustToNonCoherentAtomSize(Offset, AllocationSize, MemoryAllocation->GetSize(), NonCoherentAtomSize));
+		}
 	}
 
 	TArrayView<uint32> FVulkanSubresourceAllocator::GetMemoryUsed()
diff --git a/Engine/Source/Runtime/VulkanRHI/Private/VulkanRenderTarget.cpp b/Engine/Source/Runtime/VulkanRHI/Private/VulkanRenderTarget.cpp
--- a/Engine/Source/Runtime/VulkanRHI/Private/VulkanRenderTarget.cpp
+++ b/Engine/Source/Runtime/VulkanRHI/Private/VulkanRenderTarget.cpp
@@ -312,14 +312,14 @@ void FVulkanDynamicRHI::RHIReadSurfaceFloatData(FRHITexture* TextureRHI, FIntRec
 			VkBufferImageCopy CopyRegion;
 			FMemory::Memzero(CopyRegion);
 			//Region.bufferOffset = 0;
-			CopyRegion.bufferRowLength = Desc.Extent.X >> InMipIndex;
-			CopyRegion.bufferImageHeight = Desc.Extent.Y >> InMipIndex;
+			CopyRegion.bufferRowLength = FMath::Max(1, Desc.Extent.X >> InMipIndex);
+			CopyRegion.bufferImageHeight = FMath::Max(1, Desc.Extent.Y >> InMipIndex);
 			CopyRegion.imageSubresource.aspectMask = Surface.GetFullAspectMask();
 			CopyRegion.imageSubresource.mipLevel = InMipIndex;
 			CopyRegion.imageSubresource.baseArrayLayer = SrcBaseArrayLayer;
 			CopyRegion.imageSubresource.layerCount = 1;
-			CopyRegion.imageExtent.width = Desc.Extent.X >> InMipIndex;
-			CopyRegion.imageExtent.height = Desc.Extent.Y >> InMipIndex;
+			CopyRegion.imageExtent.width = FMath::Max(1, Desc.Extent.X >> InMipIndex);
+			CopyRegion.imageExtent.height = FMath::Max(1, Desc.Extent.Y >> InMipIndex);
 			CopyRegion.imageExtent.depth = 1;
 
 			const VkImageLayout OriginalLayout = FVulkanLayoutManager::SetExpectedLayout(InCmdBuffer, Surface, ERHIAccess::CopySrc);
diff --git a/Engine/Source/Runtime/VulkanRHI/Private/VulkanTexture.cpp b/Engine/Source/Runtime/VulkanRHI/Private/VulkanTexture.cpp
--- a/Engine/Source/Runtime/VulkanRHI/Private/VulkanTexture.cpp
+++ b/Engine/Source/Runtime/VulkanRHI/Private/VulkanTexture.cpp
@@ -1241,8 +1241,8 @@ void FVulkanDynamicRHI::RHIUnlockTexture2DArray(FRHITexture2DArray* TextureRHI,
 
 	const FRHITextureDesc& Desc = Texture->GetDesc();
 	const EPixelFormat Format = Desc.Format;
-	const uint32 MipWidth = FMath::Max<uint32>(Desc.Extent.X >> MipIndex, GPixelFormats[Format].BlockSizeX);
-	const uint32 MipHeight = FMath::Max<uint32>(Desc.Extent.Y >> MipIndex, GPixelFormats[Format].BlockSizeY);
+	const uint32 MipWidth = FMath::Max<uint32>(Desc.Extent.X >> MipIndex, 1);
+	const uint32 MipHeight = FMath::Max<uint32>(Desc.Extent.Y >> MipIndex, 1);
 
 	VkBufferImageCopy Region;
 	FMemory::Memzero(Region);
@@ -1305,14 +1305,15 @@ void FVulkanDynamicRHI::InternalUpdateTexture2D(FRHICommandListBase& RHICmdList,
 		CopyDst += StagingPitch;
 	}
 
+	const FIntVector MipDimensions = TextureRHI->GetMipDimensions(MipIndex);
 	VkBufferImageCopy Region{};
 	Region.imageSubresource.aspectMask = VK_IMAGE_ASPECT_COLOR_BIT;
 	Region.imageSubresource.mipLevel = MipIndex;
 	Region.imageSubresource.layerCount = 1;
 	Region.imageOffset.x = UpdateRegion.DestX;
 	Region.imageOffset.y = UpdateRegion.DestY;
-	Region.imageExtent.width = UpdateRegion.Width;
-	Region.imageExtent.height = UpdateRegion.Height;
+	Region.imageExtent.width = FMath::Min(UpdateRegion.Width, static_cast<uint32>(MipDimensions.X) - UpdateRegion.DestX);
+	Region.imageExtent.height = FMath::Min(UpdateRegion.Height, static_cast<uint32>(MipDimensions.Y) - UpdateRegion.DestY);
 	Region.imageExtent.depth = 1;
 
 	FVulkanTexture* Texture = ResourceCast(TextureRHI);
@@ -2239,17 +2240,17 @@ void FVulkanCommandListContext::RHICopyTexture(FRHITexture* SourceTexture, FRHIT
 		if (CopyInfo.Size == FIntVector::ZeroValue)
 		{
 			// Copy whole texture when zero vector is specified for region size
-			Region.extent.width  = FMath::Max<uint32>(PixelFormatInfo.BlockSizeX, SourceXYZ.X >> CopyInfo.SourceMipIndex);
-			Region.extent.height = FMath::Max<uint32>(PixelFormatInfo.BlockSizeY, SourceXYZ.Y >> CopyInfo.SourceMipIndex);
-			Region.extent.depth  = FMath::Max<uint32>(PixelFormatInfo.BlockSizeZ, SourceXYZ.Z >> CopyInfo.SourceMipIndex);
+			Region.extent.width  = FMath::Max<uint32>(1u, SourceXYZ.X >> CopyInfo.SourceMipIndex);
+			Region.extent.height = FMath::Max<uint32>(1u, SourceXYZ.Y >> CopyInfo.SourceMipIndex);
+			Region.extent.depth  = FMath::Max<uint32>(1u, SourceXYZ.Z >> CopyInfo.SourceMipIndex);
 			ensure(Region.extent.width <= (uint32)DestXYZ.X && Region.extent.height <= (uint32)DestXYZ.Y);
 		}
 		else
 		{
 			ensure(CopyInfo.Size.X > 0 && CopyInfo.Size.X <= DestXYZ.X && CopyInfo.Size.Y > 0 && CopyInfo.Size.Y <= DestXYZ.Y);
-			Region.extent.width  = FMath::Max(PixelFormatInfo.BlockSizeX, CopyInfo.Size.X);
-			Region.extent.height = FMath::Max(PixelFormatInfo.BlockSizeY, CopyInfo.Size.Y);
-			Region.extent.depth  = FMath::Max(PixelFormatInfo.BlockSizeZ, CopyInfo.Size.Z);
+			Region.extent.width  = FMath::Max(1, CopyInfo.Size.X);
+			Region.extent.height = FMath::Max(1, CopyInfo.Size.Y);
+			Region.extent.depth  = FMath::Max(1, CopyInfo.Size.Z);
 		}
 		Region.srcSubresource.aspectMask = Source->GetFullAspectMask();
 		Region.srcSubresource.baseArrayLayer = CopyInfo.SourceSliceIndex;
@@ -2287,9 +2288,9 @@ void FVulkanCommandListContext::RHICopyTexture(FRHITexture* SourceTexture, FRHIT
 				Region.dstOffset.y /= 2;
 				Region.dstOffset.z /= 2;
 
-				Region.extent.width  = FMath::Max<uint32>(Region.extent.width  / 2, PixelFormatInfo.BlockSizeX);
-				Region.extent.height = FMath::Max<uint32>(Region.extent.height / 2, PixelFormatInfo.BlockSizeY);
-				Region.extent.depth  = FMath::Max<uint32>(Region.extent.depth  / 2, PixelFormatInfo.BlockSizeZ);
+				Region.extent.width  = FMath::Max<uint32>(Region.extent.width  / 2, 1u);
+				Region.extent.height = FMath::Max<uint32>(Region.extent.height / 2, 1u);
+				Region.extent.depth  = FMath::Max<uint32>(Region.extent.depth  / 2, 1u);
 
 				// RHICopyTexture is allowed to copy mip regions only if are aligned on the block size to prevent unexpected / inconsistent results.
 				ensure(Region.srcOffset.x % PixelFormatInfo.BlockSizeX == 0 && Region.srcOffset.y % PixelFormatInfo.BlockSizeY == 0 && Region.srcOffset.z % PixelFormatInfo.BlockSizeZ == 0);
