Remove VULKAN_USE_REAL_RENDERPASS_COMPATIBILITY

It is causing validation errors

diff --git a/Engine/Source/Runtime/VulkanRHI/Private/Android/VulkanAndroidPlatform.h b/Engine/Source/Runtime/VulkanRHI/Private/Android/VulkanAndroidPlatform.h
--- a/Engine/Source/Runtime/VulkanRHI/Private/Android/VulkanAndroidPlatform.h
+++ b/Engine/Source/Runtime/VulkanRHI/Private/Android/VulkanAndroidPlatform.h
@@ -28,11 +28,6 @@
 // crashing during callback setup on Android, code will fallback to VK_EXT_debug_report instead
 #define VULKAN_SUPPORTS_DEBUG_UTILS					0
 
-// Android's hashes currently work fine as the problematic cases are:
-//	VK_IMAGE_LAYOUT_DEPTH_READ_ONLY_STENCIL_ATTACHMENT_OPTIMAL = 1000117000,
-//	VK_IMAGE_LAYOUT_DEPTH_ATTACHMENT_STENCIL_READ_ONLY_OPTIMAL = 1000117001,
-#define VULKAN_USE_REAL_RENDERPASS_COMPATIBILITY	0
-
 #define UE_VK_API_VERSION							VK_API_VERSION_1_1
 
 #define ENUM_VK_ENTRYPOINTS_PLATFORM_BASE(EnumMacro)
diff --git a/Engine/Source/Runtime/VulkanRHI/Private/VulkanRenderTarget.cpp b/Engine/Source/Runtime/VulkanRHI/Private/VulkanRenderTarget.cpp
--- a/Engine/Source/Runtime/VulkanRHI/Private/VulkanRenderTarget.cpp
+++ b/Engine/Source/Runtime/VulkanRHI/Private/VulkanRenderTarget.cpp
@@ -751,12 +751,10 @@ struct FRenderPassFullHashableStruct
 	// +1 for Depth, +1 for Stencil, +1 for Fragment Density
 	TEnumAsByte<VkAttachmentLoadOp>		LoadOps[MaxSimultaneousRenderTargets + 3];
 	TEnumAsByte<VkAttachmentStoreOp>	StoreOps[MaxSimultaneousRenderTargets + 3];
-#if VULKAN_USE_REAL_RENDERPASS_COMPATIBILITY
 	// If the initial != final we need to add FinalLayout and potentially RefLayout
 	VkImageLayout						InitialLayout[MaxSimultaneousRenderTargets + 3];
 	//VkImageLayout						FinalLayout[MaxSimultaneousRenderTargets + 3];
 	//VkImageLayout						RefLayout[MaxSimultaneousRenderTargets + 3];
-#endif
 };
 
 VkImageLayout FVulkanRenderTargetLayout::GetVRSImageLayout() const
@@ -863,9 +861,7 @@ FVulkanRenderTargetLayout::FVulkanRenderTargetLayout(FVulkanDevice& InDevice, co
 			CompatibleHashInfo.Formats[NumColorAttachments] = CurrDesc.format;
 			FullHashInfo.LoadOps[NumColorAttachments] = CurrDesc.loadOp;
 			FullHashInfo.StoreOps[NumColorAttachments] = CurrDesc.storeOp;
-#if VULKAN_USE_REAL_RENDERPASS_COMPATIBILITY
 			FullHashInfo.InitialLayout[NumColorAttachments] = VK_IMAGE_LAYOUT_COLOR_ATTACHMENT_OPTIMAL;
-#endif
 			++CompatibleHashInfo.NumAttachments;
 
 			++NumAttachmentDescriptions;
@@ -925,10 +921,8 @@ FVulkanRenderTargetLayout::FVulkanRenderTargetLayout(FVulkanDevice& InDevice, co
 		FullHashInfo.LoadOps[MaxSimultaneousRenderTargets + 1] = CurrDesc.stencilLoadOp;
 		FullHashInfo.StoreOps[MaxSimultaneousRenderTargets] = CurrDesc.storeOp;
 		FullHashInfo.StoreOps[MaxSimultaneousRenderTargets + 1] = CurrDesc.stencilStoreOp;
-#if VULKAN_USE_REAL_RENDERPASS_COMPATIBILITY
 		FullHashInfo.InitialLayout[MaxSimultaneousRenderTargets] = DepthLayout;
 		FullHashInfo.InitialLayout[MaxSimultaneousRenderTargets + 1] = StencilLayout;
-#endif
 		CompatibleHashInfo.Formats[MaxSimultaneousRenderTargets] = CurrDesc.format;
 
 		++NumAttachmentDescriptions;
@@ -975,9 +969,7 @@ FVulkanRenderTargetLayout::FVulkanRenderTargetLayout(FVulkanDevice& InDevice, co
 
 		FullHashInfo.LoadOps[MaxSimultaneousRenderTargets + 2] = CurrDesc.stencilLoadOp;
 		FullHashInfo.StoreOps[MaxSimultaneousRenderTargets + 2] = CurrDesc.stencilStoreOp;
-#if VULKAN_USE_REAL_RENDERPASS_COMPATIBILITY
 		FullHashInfo.InitialLayout[MaxSimultaneousRenderTargets + 2] = VRSLayout;
-#endif
 		CompatibleHashInfo.Formats[MaxSimultaneousRenderTargets + 1] = CurrDesc.format;
 
 		++NumAttachmentDescriptions;
@@ -1088,9 +1080,7 @@ FVulkanRenderTargetLayout::FVulkanRenderTargetLayout(FVulkanDevice& InDevice, co
 
 		CompatibleHashInfo.Formats[NumColorAttachments] = CurrDesc.format;
 		FullHashInfo.LoadOps[NumColorAttachments] = CurrDesc.loadOp;
-#if VULKAN_USE_REAL_RENDERPASS_COMPATIBILITY
 		FullHashInfo.InitialLayout[NumColorAttachments] = VK_IMAGE_LAYOUT_COLOR_ATTACHMENT_OPTIMAL;
-#endif
 		FullHashInfo.StoreOps[NumColorAttachments] = CurrDesc.storeOp;
 		++CompatibleHashInfo.NumAttachments;
 
@@ -1168,10 +1158,8 @@ FVulkanRenderTargetLayout::FVulkanRenderTargetLayout(FVulkanDevice& InDevice, co
 		FullHashInfo.LoadOps[MaxSimultaneousRenderTargets + 1] = CurrDesc.stencilLoadOp;
 		FullHashInfo.StoreOps[MaxSimultaneousRenderTargets] = CurrDesc.storeOp;
 		FullHashInfo.StoreOps[MaxSimultaneousRenderTargets + 1] = CurrDesc.stencilStoreOp;
-#if VULKAN_USE_REAL_RENDERPASS_COMPATIBILITY
 		FullHashInfo.InitialLayout[MaxSimultaneousRenderTargets] = CurrentDepthLayout;
 		FullHashInfo.InitialLayout[MaxSimultaneousRenderTargets + 1] = CurrentStencilLayout;
-#endif
 		CompatibleHashInfo.Formats[MaxSimultaneousRenderTargets] = CurrDesc.format;
 
 
@@ -1229,9 +1217,7 @@ FVulkanRenderTargetLayout::FVulkanRenderTargetLayout(FVulkanDevice& InDevice, co
 
 		FullHashInfo.LoadOps[MaxSimultaneousRenderTargets + 2] = CurrDesc.stencilLoadOp;
 		FullHashInfo.StoreOps[MaxSimultaneousRenderTargets + 2] = CurrDesc.stencilStoreOp;
-#if VULKAN_USE_REAL_RENDERPASS_COMPATIBILITY
 		FullHashInfo.InitialLayout[MaxSimultaneousRenderTargets + 2] = VRSLayout;
-#endif
 		CompatibleHashInfo.Formats[MaxSimultaneousRenderTargets + 1] = CurrDesc.format;
 
 		++NumAttachmentDescriptions;
@@ -1311,9 +1297,7 @@ FVulkanRenderTargetLayout::FVulkanRenderTargetLayout(const FGraphicsPipelineStat
 			CompatibleHashInfo.Formats[NumColorAttachments] = CurrDesc.format;
 			FullHashInfo.LoadOps[NumColorAttachments] = CurrDesc.loadOp;
 			FullHashInfo.StoreOps[NumColorAttachments] = CurrDesc.storeOp;
-#if VULKAN_USE_REAL_RENDERPASS_COMPATIBILITY
 			FullHashInfo.InitialLayout[NumColorAttachments] = VK_IMAGE_LAYOUT_COLOR_ATTACHMENT_OPTIMAL;
-#endif
 			++CompatibleHashInfo.NumAttachments;
 
 			++NumAttachmentDescriptions;
@@ -1360,10 +1344,8 @@ FVulkanRenderTargetLayout::FVulkanRenderTargetLayout(const FGraphicsPipelineStat
 		FullHashInfo.LoadOps[MaxSimultaneousRenderTargets + 1] = CurrDesc.stencilLoadOp;
 		FullHashInfo.StoreOps[MaxSimultaneousRenderTargets] = CurrDesc.storeOp;
 		FullHashInfo.StoreOps[MaxSimultaneousRenderTargets + 1] = CurrDesc.stencilStoreOp;
-#if VULKAN_USE_REAL_RENDERPASS_COMPATIBILITY
 		FullHashInfo.InitialLayout[MaxSimultaneousRenderTargets] = VK_IMAGE_LAYOUT_ATTACHMENT_OPTIMAL;
 		FullHashInfo.InitialLayout[MaxSimultaneousRenderTargets + 1] = VK_IMAGE_LAYOUT_ATTACHMENT_OPTIMAL;
-#endif
 		CompatibleHashInfo.Formats[MaxSimultaneousRenderTargets] = CurrDesc.format;
 
 		++NumAttachmentDescriptions;
@@ -1394,9 +1376,7 @@ FVulkanRenderTargetLayout::FVulkanRenderTargetLayout(const FGraphicsPipelineStat
 
 		FullHashInfo.LoadOps[MaxSimultaneousRenderTargets + 2] = CurrDesc.stencilLoadOp;
 		FullHashInfo.StoreOps[MaxSimultaneousRenderTargets + 2] = CurrDesc.stencilStoreOp;
-#if VULKAN_USE_REAL_RENDERPASS_COMPATIBILITY
 		FullHashInfo.InitialLayout[MaxSimultaneousRenderTargets + 2] = VRSLayout;
-#endif
 		CompatibleHashInfo.Formats[MaxSimultaneousRenderTargets + 1] = CurrDesc.format;
 
 		++NumAttachmentDescriptions;
diff --git a/Engine/Source/Runtime/VulkanRHI/Public/VulkanConfiguration.h b/Engine/Source/Runtime/VulkanRHI/Public/VulkanConfiguration.h
--- a/Engine/Source/Runtime/VulkanRHI/Public/VulkanConfiguration.h
+++ b/Engine/Source/Runtime/VulkanRHI/Public/VulkanConfiguration.h
@@ -136,10 +136,6 @@ struct VkAllocationCallbacks;
 	#define VULKAN_USE_CREATE_WIN32_SURFACE						0
 #endif
 
-#ifndef VULKAN_USE_REAL_RENDERPASS_COMPATIBILITY
-	#define VULKAN_USE_REAL_RENDERPASS_COMPATIBILITY			1
-#endif
-
 #ifndef VULKAN_USE_DIFFERENT_POOL_CMDBUFFERS
 	#define VULKAN_USE_DIFFERENT_POOL_CMDBUFFERS				1
 #endif
