Fix VK_EXT_debug_markers.
When Validation was enabled and debug_utils are not used then debug_markers were disabled. With markers you get object names in validation error messages

diff --git a/Engine/Source/Runtime/VulkanRHI/Private/Linux/VulkanLinuxPlatform.cpp b/Engine/Source/Runtime/VulkanRHI/Private/Linux/VulkanLinuxPlatform.cpp
--- a/Engine/Source/Runtime/VulkanRHI/Private/Linux/VulkanLinuxPlatform.cpp
+++ b/Engine/Source/Runtime/VulkanRHI/Private/Linux/VulkanLinuxPlatform.cpp
@@ -14,7 +14,6 @@
 ENUM_VK_ENTRYPOINTS_ALL(DEFINE_VK_ENTRYPOINTS)
 
 static bool GRenderOffScreen = false;
-static bool GForceEnableDebugMarkers = false;
 
 void* FVulkanLinuxPlatform::VulkanLib = nullptr;
 bool FVulkanLinuxPlatform::bAttemptedLoad = false;
@@ -124,16 +123,9 @@ bool FVulkanLinuxPlatform::LoadVulkanLibrary()
 
 #undef GET_VK_ENTRYPOINTS
 
-	// Check for force enabling debug markers
-	GForceEnableDebugMarkers = FParse::Param(FCommandLine::Get(), TEXT("vulkandebugmarkers"));
 	return true;
 }
 
-bool FVulkanLinuxPlatform::ForceEnableDebugMarkers()
-{
-	return GForceEnableDebugMarkers;
-}
-
 bool FVulkanLinuxPlatform::LoadVulkanInstanceFunctions(VkInstance inInstance)
 {
 	bool bFoundAllEntryPoints = true;
diff --git a/Engine/Source/Runtime/VulkanRHI/Private/Linux/VulkanLinuxPlatform.h b/Engine/Source/Runtime/VulkanRHI/Private/Linux/VulkanLinuxPlatform.h
--- a/Engine/Source/Runtime/VulkanRHI/Private/Linux/VulkanLinuxPlatform.h
+++ b/Engine/Source/Runtime/VulkanRHI/Private/Linux/VulkanLinuxPlatform.h
@@ -52,8 +52,6 @@ public:
 
 	static void CreateSurface(void* WindowHandle, VkInstance Instance, VkSurfaceKHR* OutSurface);
 
-	static bool ForceEnableDebugMarkers();
-
 	static void WriteCrashMarker(const FOptionalVulkanDeviceExtensions& OptionalExtensions, VkCommandBuffer CmdBuffer, VkBuffer DestBuffer, const TArrayView<uint32>& Entries, bool bAdding);
 
 protected:
diff --git a/Engine/Source/Runtime/VulkanRHI/Private/VulkanDevice.cpp b/Engine/Source/Runtime/VulkanRHI/Private/VulkanDevice.cpp
--- a/Engine/Source/Runtime/VulkanRHI/Private/VulkanDevice.cpp
+++ b/Engine/Source/Runtime/VulkanRHI/Private/VulkanDevice.cpp
@@ -539,7 +539,7 @@ void FVulkanDevice::SetupDrawMarkers()
 	}
 	else
 #endif	// VULKAN_SUPPORTS_DEBUG_UTILS
-	if (bDebugMarkersFound || FVulkanPlatform::ForceEnableDebugMarkers())
+	if (bDebugMarkersFound)
 	{
 		DebugMarkers.CmdBegin = (PFN_vkCmdDebugMarkerBeginEXT)(void*)VulkanRHI::vkGetDeviceProcAddr(Device, "vkCmdDebugMarkerBeginEXT");
 		DebugMarkers.CmdEnd = (PFN_vkCmdDebugMarkerEndEXT)(void*)VulkanRHI::vkGetDeviceProcAddr(Device, "vkCmdDebugMarkerEndEXT");
diff --git a/Engine/Source/Runtime/VulkanRHI/Private/VulkanExtensions.cpp b/Engine/Source/Runtime/VulkanRHI/Private/VulkanExtensions.cpp
--- a/Engine/Source/Runtime/VulkanRHI/Private/VulkanExtensions.cpp
+++ b/Engine/Source/Runtime/VulkanRHI/Private/VulkanExtensions.cpp
@@ -878,7 +878,7 @@ public:
 	{
 #if VULKAN_HAS_DEBUGGING_ENABLED
 		const int32 VulkanValidationOption = GValidationCvar.GetValueOnAnyThread();
-		bEnabledInCode = bEnabledInCode && ((GRenderDocFound || VulkanValidationOption == 0) || FVulkanPlatform::ForceEnableDebugMarkers());
+		bEnabledInCode = bEnabledInCode && (GRenderDocFound || VulkanValidationOption != 0);
 #endif
 	}
 
diff --git a/Engine/Source/Runtime/VulkanRHI/Private/VulkanGenericPlatform.h b/Engine/Source/Runtime/VulkanRHI/Private/VulkanGenericPlatform.h
--- a/Engine/Source/Runtime/VulkanRHI/Private/VulkanGenericPlatform.h
+++ b/Engine/Source/Runtime/VulkanRHI/Private/VulkanGenericPlatform.h
@@ -69,8 +69,6 @@ public:
 	// Some platforms have issues with the access flags for the Present layout
 	static bool RequiresPresentLayoutFix() { return false; }
 
-	static bool ForceEnableDebugMarkers() { return false; }
-
 	static bool SupportsDeviceLocalHostVisibleWithNoPenalty(EGpuVendorId VendorId) { return false; }
 
 	static bool HasUnifiedMemory() { return false; }
diff --git a/Engine/Source/Runtime/VulkanRHI/Private/VulkanLayers.cpp b/Engine/Source/Runtime/VulkanRHI/Private/VulkanLayers.cpp
--- a/Engine/Source/Runtime/VulkanRHI/Private/VulkanLayers.cpp
+++ b/Engine/Source/Runtime/VulkanRHI/Private/VulkanLayers.cpp
@@ -501,8 +501,18 @@ void FVulkanDeviceSetupHelper::AddDebugLayers(const TArray<FLayerWithExtensions>
 	// Verify that all requested debugging device-layers are available. Skip validation layers under RenderDoc
 	if (!GRenderDocFound && GRHIGlobals.IsDebugLayerEnabled && (GStandardValidationCvar.GetValueOnAnyThread() != 0))
 	{
-		// Path for older drivers
-		AddRequestedLayer(STANDARD_VALIDATION_LAYER_NAME, LayerProperties, UEExtensions, OutLayers);
+		switch (GStandardValidationCvar.GetValueOnAnyThread())
+		{
+		case 1:
+			AddRequestedLayer(STANDARD_VALIDATION_LAYER_NAME, LayerProperties, UEExtensions, OutLayers);
+			break;
+		case 2:
+			AddRequestedLayer(KHRONOS_STANDARD_VALIDATION_LAYER_NAME, LayerProperties, UEExtensions, OutLayers);
+			break;
+		default:
+			break;
+		}
+		
 	}
 #endif	// VULKAN_HAS_DEBUGGING_ENABLED
 }
