Remove support for legacy Vulkan debug extensions VK_EXT_debug_report/debug_marker.
Now only VK_EXT_debug_utils is supported.

diff --git a/Engine/Source/Runtime/VulkanRHI/Private/Android/VulkanAndroidPlatform.h b/Engine/Source/Runtime/VulkanRHI/Private/Android/VulkanAndroidPlatform.h
--- a/Engine/Source/Runtime/VulkanRHI/Private/Android/VulkanAndroidPlatform.h
+++ b/Engine/Source/Runtime/VulkanRHI/Private/Android/VulkanAndroidPlatform.h
@@ -25,9 +25,6 @@
 #define VULKAN_SUPPORTS_DRIVER_PROPERTIES			0
 #define VULKAN_SUPPORTS_DESCRIPTOR_INDEXING			(VULKAN_RHI_RAYTRACING)
 
-// crashing during callback setup on Android, code will fallback to VK_EXT_debug_report instead
-#define VULKAN_SUPPORTS_DEBUG_UTILS					0
-
 #define UE_VK_API_VERSION							VK_API_VERSION_1_1
 
 #define ENUM_VK_ENTRYPOINTS_PLATFORM_BASE(EnumMacro)
diff --git a/Engine/Source/Runtime/VulkanRHI/Private/VulkanDebug.cpp b/Engine/Source/Runtime/VulkanRHI/Private/VulkanDebug.cpp
--- a/Engine/Source/Runtime/VulkanRHI/Private/VulkanDebug.cpp
+++ b/Engine/Source/Runtime/VulkanRHI/Private/VulkanDebug.cpp
@@ -33,150 +33,8 @@ DEFINE_LOG_CATEGORY(LogVulkanRHI);
 
 extern TAutoConsoleVariable<int32> GValidationCvar;
 
-static VkBool32 VKAPI_PTR DebugReportFunction(
-	VkDebugReportFlagsEXT			MsgFlags,
-	VkDebugReportObjectTypeEXT		ObjType,
-	uint64_t						SrcObject,
-	size_t							Location,
-	int32							MsgCode,
-	const ANSICHAR*					LayerPrefix,
-	const ANSICHAR*					Msg,
-	void*							UserData)
-{
-#if VULKAN_ENABLE_DUMP_LAYER
-	VulkanRHI::FlushDebugWrapperLog();
-#endif
-
-	const char* MsgPrefix = "UNKNOWN";
-	if (MsgFlags & VK_DEBUG_REPORT_ERROR_BIT_EXT)
-	{
-		// Ignore some errors we might not fix...
-		if (!FCStringAnsi::Strcmp(LayerPrefix, "Validation"))
-		{
-			if (MsgCode == 0x4c00264)
-			{
-				// Unable to allocate 1 descriptorSets from pool 0x8cb8. This pool only has N descriptorSets remaining. The spec valid usage text states
-				// 'descriptorSetCount must not be greater than the number of sets that are currently available for allocation in descriptorPool'
-				// (https://www.khronos.org/registry/vulkan/specs/1.0/html/vkspec.html#VUID-VkDescriptorSetAllocateInfo-descriptorSetCount-00306)
-				return VK_FALSE;
-			}
-			else if (MsgCode == 0x4c00266)
-			{
-				// Unable to allocate 1 descriptors of type VK_DESCRIPTOR_TYPE_UNIFORM_BUFFER from pool 0x89f4. This pool only has 0 descriptors of this type
-				// remaining.The spec valid usage text states 'descriptorPool must have enough free descriptor capacity remaining to allocate the descriptor sets of
-				// the specified layouts' (https://www.khronos.org/registry/vulkan/specs/1.0/html/vkspec.html#VUID-VkDescriptorSetAllocateInfo-descriptorPool-00307)
-				return VK_FALSE;
-			}
-		}
-		if (!FCStringAnsi::Strcmp(LayerPrefix, "SC"))
-		{
-			if (MsgCode == 3)
-			{
-				// Attachment N not written by fragment shader
-				return VK_FALSE;
-			}
-		}
-		if (!FCStringAnsi::Strcmp(LayerPrefix, "DS"))
-		{
-			if (MsgCode == 6)
-			{
-				auto* Found = FCStringAnsi::Strstr(Msg, " array layer ");
-				if (Found && Found[13] >= '1' && Found[13] <= '9')
-				{
-					//#todo-rco: Remove me?
-					// Potential bug in the validation layers for slice > 1 on 3d textures
-					return VK_FALSE;
-				}
-			}
-			else if (MsgCode == 15)
-			{
-				// Cannot get query results on queryPool 0x327 with index 193 as data has not been collected for this index.
-				//return VK_FALSE;
-			}
-		}
-		if (FCStringAnsi::Strstr(Msg, "VUID-RuntimeSpirv-Fragment-06427"))
-		{
-			// "vkCmdDrawIndexed: Dual source blend mode is used, but the number of written fragment shader output attachment (2) is greater than maxFragmentDualSrcAttachments (1)"
-			// Incorrectly generated in Validation layers 1.3.250.1, no longer shows up in 1.3.268.0
-			return VK_FALSE;
-		}
-
-		MsgPrefix = "ERROR";
-	}
-	else if (MsgFlags & VK_DEBUG_REPORT_WARNING_BIT_EXT)
-	{
-		MsgPrefix = "WARN";
-
-		// Ignore some warnings we might not fix...
-		// Ignore some errors we might not fix...
-		if (!FCStringAnsi::Strcmp(LayerPrefix, "Validation"))
-		{
-			if (MsgCode == 2)
-			{
-				// fragment shader writes to output location 0 with no matching attachment
-				return VK_FALSE;
-			}
-		}
-
-		if (!FCStringAnsi::Strcmp(LayerPrefix, "SC"))
-		{
-			if (MsgCode == 2)
-			{
-				// fragment shader writes to output location 0 with no matching attachment
-				return VK_FALSE;
-			}
-		}
-	}
-	else if (MsgFlags & VK_DEBUG_REPORT_PERFORMANCE_WARNING_BIT_EXT)
-	{
-		MsgPrefix = "PERF";
-		// Ignore some errors we might not fix...
-		if (!FCStringAnsi::Strcmp(LayerPrefix, "SC"))
-		{
-			if (MsgCode == 2)
-			{
-				// vertex shader outputs unused interpolator
-				return VK_FALSE;
-			}
-		}
-		else if (!FCStringAnsi::Strcmp(LayerPrefix, "DS"))
-		{
-			if (MsgCode == 15)
-			{
-				// DescriptorSet previously bound is incompatible with set newly bound as set #0 so set #1 and any subsequent sets were disturbed by newly bound pipelineLayout
-				return VK_FALSE;
-			}
-		}
-	}
-	else if (MsgFlags & VK_DEBUG_REPORT_INFORMATION_BIT_EXT)
-	{
-		MsgPrefix = "INFO";
-	}
-	else if (MsgFlags & VK_DEBUG_REPORT_DEBUG_BIT_EXT)
-	{
-		MsgPrefix = "DEBUG";
-	}
-	else
-	{
-		ensure(0);
-	}
-
-	const VkBool32 bPlatformPrintLog = FVulkanPlatform::DebugReportFunction(MsgFlags, ObjType, SrcObject, Location, MsgCode, LayerPrefix, Msg, UserData);
-	if (bPlatformPrintLog == VK_FALSE)
-	{
-		// Early out if platform wants to suppress this debug report.
-		return VK_FALSE;
-	}
-
-	// MsgCode seem to be always 0
-	VULKAN_REPORT_LOG(TEXT("*** [%s] Obj 0x%p Loc %d %s"), ANSI_TO_TCHAR(MsgPrefix), (void*)SrcObject, (uint32)Location, ANSI_TO_TCHAR(Msg));
-
-	return VK_FALSE;
-}
-
-#if VULKAN_SUPPORTS_DEBUG_UTILS
-static VkBool32 DebugUtilsCallback(VkDebugUtilsMessageSeverityFlagBitsEXT MsgSeverity, VkDebugUtilsMessageTypeFlagsEXT MsgType,
-	const VkDebugUtilsMessengerCallbackDataEXT* CallbackData, void* UserData)
+static VKAPI_ATTR VkBool32 VKAPI_CALL DebugUtilsCallback(VkDebugUtilsMessageSeverityFlagBitsEXT MsgSeverity, 
+	VkDebugUtilsMessageTypeFlagsEXT MsgType, const VkDebugUtilsMessengerCallbackDataEXT* CallbackData, void* UserData)
 {
 	const bool bError = (MsgSeverity & VK_DEBUG_UTILS_MESSAGE_SEVERITY_ERROR_BIT_EXT) != 0;
 	const bool bWarning = (MsgSeverity & VK_DEBUG_UTILS_MESSAGE_SEVERITY_WARNING_BIT_EXT) != 0;
@@ -274,6 +132,12 @@ static VkBool32 DebugUtilsCallback(VkDebugUtilsMessageSeverityFlagBitsEXT MsgSev
 		// *** [Error:Validation(UNASSIGNED-GPU-Assisted Validation Setup Error.)] Unable to reserve descriptor binding slot on a device with only one slot.
 		return VK_FALSE;
 	}
+	else if (!FCStringAnsi::Strcmp(CallbackData->pMessageIdName, "VUID-RuntimeSpirv-Fragment-06427"))
+	{
+		// Warning: *** [Error:Validation(VUID-RuntimeSpirv-Fragment-06427)]
+		// False-positive in validation layers 1.3.250.1, known to be fixed in 1.3.268.0
+		return VK_FALSE;
+	}
 	else if (!FCStringAnsi::Strcmp(CallbackData->pMessageIdName, "UNASSIGNED-BestPractices-vkCreateDevice-deprecated-extension"))
 	{
 		// *** CreateDevice(): Attempting to enable deprecated extension VK_KHR_get_memory_requirements2, but this extension has been promoted to VK_VERSION_1_1.
@@ -311,11 +175,9 @@ static VkBool32 DebugUtilsCallback(VkDebugUtilsMessageSeverityFlagBitsEXT MsgSev
 
 	return VK_FALSE;
 }
-#endif
 
 void FVulkanDynamicRHI::SetupDebugLayerCallback()
 {
-#if VULKAN_SUPPORTS_DEBUG_UTILS
 	if (ActiveDebugLayerExtension == EActiveDebugLayerExtension::DebugUtilsExtension)
 	{
 		PFN_vkCreateDebugUtilsMessengerEXT CreateDebugUtilsMessengerEXT = (PFN_vkCreateDebugUtilsMessengerEXT)(void*)VulkanRHI::vkGetInstanceProcAddr(Instance, "vkCreateDebugUtilsMessengerEXT");
@@ -329,67 +191,15 @@ void FVulkanDynamicRHI::SetupDebugLayerCallback()
 				(CVar >= 2 ? VK_DEBUG_UTILS_MESSAGE_SEVERITY_WARNING_BIT_EXT : 0) | (CVar >= 3 ? VK_DEBUG_UTILS_MESSAGE_SEVERITY_INFO_BIT_EXT : 0);
 			CreateInfo.messageType = (CVar >= 1 ? (VK_DEBUG_UTILS_MESSAGE_TYPE_GENERAL_BIT_EXT | VK_DEBUG_UTILS_MESSAGE_TYPE_VALIDATION_BIT_EXT) : 0) |
 				(CVar >= 3 ? VK_DEBUG_UTILS_MESSAGE_TYPE_PERFORMANCE_BIT_EXT : 0);
-			CreateInfo.pfnUserCallback = (PFN_vkDebugUtilsMessengerCallbackEXT)(void*)DebugUtilsCallback;
+			CreateInfo.pfnUserCallback = DebugUtilsCallback;
 			VkResult Result = (*CreateDebugUtilsMessengerEXT)(Instance, &CreateInfo, nullptr, &Messenger);
 			ensure(Result == VK_SUCCESS);
 		}
 	}
-	else
-#endif
-	if (ActiveDebugLayerExtension == EActiveDebugLayerExtension::DebugReportExtension)
-	{
-		PFN_vkCreateDebugReportCallbackEXT CreateMsgCallback = (PFN_vkCreateDebugReportCallbackEXT)(void*)VulkanRHI::vkGetInstanceProcAddr(Instance, CREATE_MSG_CALLBACK);
-		if (CreateMsgCallback)
-		{
-			VkDebugReportCallbackCreateInfoEXT CreateInfo;
-			ZeroVulkanStruct(CreateInfo, VK_STRUCTURE_TYPE_DEBUG_REPORT_CREATE_INFO_EXT);
-			CreateInfo.pfnCallback = DebugReportFunction;
-
-			const int32 CVar = GValidationCvar.GetValueOnRenderThread();
-			switch (CVar)
-			{
-			default:
-				CreateInfo.flags |= VK_DEBUG_REPORT_DEBUG_BIT_EXT;
-				// Fall-through...
-			case 4:
-				CreateInfo.flags |= VK_DEBUG_REPORT_INFORMATION_BIT_EXT;
-				// Fall-through...
-			case 3:
-				CreateInfo.flags |= VK_DEBUG_REPORT_PERFORMANCE_WARNING_BIT_EXT;
-				// Fall-through...
-			case 2:
-				CreateInfo.flags |= VK_DEBUG_REPORT_WARNING_BIT_EXT;
-				// Fall-through...
-			case 1:
-				CreateInfo.flags |= VK_DEBUG_REPORT_ERROR_BIT_EXT;
-				break;
-			case 0:
-				// Nothing to do!
-				break;
-			}
-			VkResult Result = CreateMsgCallback(Instance, &CreateInfo, nullptr, &MsgCallback);
-			switch (Result)
-			{
-			case VK_SUCCESS:
-				break;
-			case VK_ERROR_OUT_OF_HOST_MEMORY:
-				UE_LOG(LogVulkanRHI, Warning, TEXT("CreateMsgCallback: out of host memory/CreateMsgCallback Failure; debug reporting skipped"));
-				break;
-			default:
-				UE_LOG(LogVulkanRHI, Warning, TEXT("CreateMsgCallback: unknown failure %d/CreateMsgCallback Failure; debug reporting skipped"), (int32)Result);
-				break;
-			}
-		}
-		else
-		{
-			UE_LOG(LogVulkanRHI, Warning, TEXT("GetProcAddr: Unable to find vkDbgCreateMsgCallback/vkGetInstanceProcAddr; debug reporting skipped!"));
-		}
-	}
 }
 
 void FVulkanDynamicRHI::RemoveDebugLayerCallback()
 {
-#if VULKAN_SUPPORTS_DEBUG_UTILS
 	if (Messenger != VK_NULL_HANDLE)
 	{
 		PFN_vkDestroyDebugUtilsMessengerEXT DestroyDebugUtilsMessengerEXT = (PFN_vkDestroyDebugUtilsMessengerEXT)(void*)VulkanRHI::vkGetInstanceProcAddr(Instance, "vkDestroyDebugUtilsMessengerEXT");
@@ -398,14 +208,6 @@ void FVulkanDynamicRHI::RemoveDebugLayerCallback()
 			(*DestroyDebugUtilsMessengerEXT)(Instance, Messenger, nullptr);
 		}
 	}
-	else
-#endif
-	if (MsgCallback != VK_NULL_HANDLE)
-	{
-		PFN_vkDestroyDebugReportCallbackEXT DestroyMsgCallback = (PFN_vkDestroyDebugReportCallbackEXT)(void*)VulkanRHI::vkGetInstanceProcAddr(Instance, DESTROY_MSG_CALLBACK);
-		checkf(DestroyMsgCallback, TEXT("GetProcAddr: Unable to find vkDbgCreateMsgCallback\vkGetInstanceProcAddr Failure"));
-		DestroyMsgCallback(Instance, MsgCallback, nullptr);
-	}
 }
 
 
diff --git a/Engine/Source/Runtime/VulkanRHI/Private/VulkanDevice.cpp b/Engine/Source/Runtime/VulkanRHI/Private/VulkanDevice.cpp
--- a/Engine/Source/Runtime/VulkanRHI/Private/VulkanDevice.cpp
+++ b/Engine/Source/Runtime/VulkanRHI/Private/VulkanDevice.cpp
@@ -525,9 +525,7 @@ void FVulkanDevice::CreateDevice(TArray<const ANSICHAR*>& DeviceLayers, FVulkanD
 void FVulkanDevice::SetupDrawMarkers()
 {
 #if VULKAN_ENABLE_DRAW_MARKERS
-#if 0//VULKAN_SUPPORTS_DEBUG_UTILS
-	FVulkanDynamicRHI* RHI = GVulkanRHI;
-	if (RHI->SupportsDebugUtilsExt() && GRenderDocFound)
+	if (RHI->SupportsDebugUtilsExt())
 	{
 		DebugMarkers.CmdBeginDebugLabel = (PFN_vkCmdBeginDebugUtilsLabelEXT)(void*)VulkanRHI::vkGetInstanceProcAddr(RHI->GetInstance(), "vkCmdBeginDebugUtilsLabelEXT");
 		DebugMarkers.CmdEndDebugLabel = (PFN_vkCmdEndDebugUtilsLabelEXT)(void*)VulkanRHI::vkGetInstanceProcAddr(RHI->GetInstance(), "vkCmdEndDebugUtilsLabelEXT");
@@ -537,42 +535,8 @@ void FVulkanDevice::SetupDrawMarkers()
 			bDebugMarkersFound = true;
 		}
 	}
-	else
-#endif	// VULKAN_SUPPORTS_DEBUG_UTILS
-	if (bDebugMarkersFound)
-	{
-		DebugMarkers.CmdBegin = (PFN_vkCmdDebugMarkerBeginEXT)(void*)VulkanRHI::vkGetDeviceProcAddr(Device, "vkCmdDebugMarkerBeginEXT");
-		DebugMarkers.CmdEnd = (PFN_vkCmdDebugMarkerEndEXT)(void*)VulkanRHI::vkGetDeviceProcAddr(Device, "vkCmdDebugMarkerEndEXT");
-		DebugMarkers.CmdSetObjectName = (PFN_vkDebugMarkerSetObjectNameEXT)(void*)VulkanRHI::vkGetDeviceProcAddr(Device, "vkDebugMarkerSetObjectNameEXT");
-
-		if (DebugMarkers.CmdBegin && DebugMarkers.CmdEnd && DebugMarkers.CmdSetObjectName)
-		{
-			bDebugMarkersFound = true;
-		}
-
-		if (!DebugMarkers.CmdBegin || !DebugMarkers.CmdEnd || !DebugMarkers.CmdSetObjectName)
-		{
-			UE_LOG(LogVulkanRHI, Warning, TEXT("Extension found, but entry points for vkCmdDebugMarker(Begin|End)EXT NOT found!"));
-			bDebugMarkersFound = false;
-			DebugMarkers.CmdBegin = nullptr;
-			DebugMarkers.CmdEnd = nullptr;
-			DebugMarkers.CmdSetObjectName = nullptr;
-		}
-	}
-	else
-	{
-		if (DebugMarkers.CmdBegin && DebugMarkers.CmdEnd && DebugMarkers.CmdSetObjectName)
-		{
-			UE_LOG(LogVulkanRHI, Warning, TEXT("Extension not found, but entry points for vkCmdDebugMarker(Begin|End)EXT found!"));
-			bDebugMarkersFound = true;
-		}
-	}
-	if(GVulkanRHI->SupportsDebugUtilsExt())
-	{
-		DebugMarkers.SetDebugName = (PFN_vkSetDebugUtilsObjectNameEXT)(void*)VulkanRHI::vkGetInstanceProcAddr(RHI->GetInstance(), "vkSetDebugUtilsObjectNameEXT");
-	}
 
-	if (bDebugMarkersFound)
+	if (bDebugMarkersFound && GRenderDocFound)
 	{
 		// We're running under RenderDoc or other trace tool, so enable capturing mode
 		EnableDrawMarkers();
diff --git a/Engine/Source/Runtime/VulkanRHI/Private/VulkanDevice.h b/Engine/Source/Runtime/VulkanRHI/Private/VulkanDevice.h
--- a/Engine/Source/Runtime/VulkanRHI/Private/VulkanDevice.h
+++ b/Engine/Source/Runtime/VulkanRHI/Private/VulkanDevice.h
@@ -449,22 +449,6 @@ public:
 	void NotifyDeletedImage(VkImage Image, bool bRenderTarget);
 
 #if VULKAN_ENABLE_DRAW_MARKERS
-	inline PFN_vkCmdDebugMarkerBeginEXT GetCmdDbgMarkerBegin() const
-	{
-		return DebugMarkers.CmdBegin;
-	}
-
-	inline PFN_vkCmdDebugMarkerEndEXT GetCmdDbgMarkerEnd() const
-	{
-		return DebugMarkers.CmdEnd;
-	}
-
-	inline PFN_vkDebugMarkerSetObjectNameEXT GetDebugMarkerSetObjectName() const
-	{
-		return DebugMarkers.CmdSetObjectName;
-	}
-
-#if 0//VULKAN_SUPPORTS_DEBUG_UTILS
 	inline PFN_vkCmdBeginDebugUtilsLabelEXT GetCmdBeginDebugLabel() const
 	{
 		return DebugMarkers.CmdBeginDebugLabel;
@@ -479,8 +463,6 @@ public:
 	{
 		return DebugMarkers.SetDebugName;
 	}
-#endif
-
 #endif
 
 	void PrepareForCPURead();
@@ -539,11 +521,6 @@ public:
 
 	const TArray<const ANSICHAR*>& GetDeviceExtensions() { return DeviceExtensions; }
 
-	inline void SetDebugMarkersFound()
-	{
-		bDebugMarkersFound = true;
-	}
-
 	// Performs a GPU and CPU timestamp at nearly the same time.
 	// This allows aligning GPU and CPU events on the same timeline in profile visualization.
 	FGPUTimingCalibrationTimestamp GetCalibrationTimestamp();
@@ -656,15 +633,9 @@ private:
 #if VULKAN_ENABLE_DRAW_MARKERS
 	struct
 	{
-		PFN_vkCmdDebugMarkerBeginEXT		CmdBegin = nullptr;
-		PFN_vkCmdDebugMarkerEndEXT			CmdEnd = nullptr;
-		PFN_vkDebugMarkerSetObjectNameEXT	CmdSetObjectName = nullptr;
 		PFN_vkSetDebugUtilsObjectNameEXT	SetDebugName = nullptr;
-
-#if 0//VULKAN_SUPPORTS_DEBUG_UTILS
 		PFN_vkCmdBeginDebugUtilsLabelEXT	CmdBeginDebugLabel = nullptr;
 		PFN_vkCmdEndDebugUtilsLabelEXT		CmdEndDebugLabel = nullptr;
-#endif
 	} DebugMarkers;
 	friend class FVulkanCommandListContext;
 #endif
diff --git a/Engine/Source/Runtime/VulkanRHI/Private/VulkanExtensions.cpp b/Engine/Source/Runtime/VulkanRHI/Private/VulkanExtensions.cpp
--- a/Engine/Source/Runtime/VulkanRHI/Private/VulkanExtensions.cpp
+++ b/Engine/Source/Runtime/VulkanRHI/Private/VulkanExtensions.cpp
@@ -868,28 +868,6 @@ private:
 
 
 
-// ***** VK_EXT_debug_marker
-class FVulkanEXTDebugMarkerExtension : public FVulkanDeviceExtension
-{
-public:
-
-	FVulkanEXTDebugMarkerExtension(FVulkanDevice* InDevice)
-		: FVulkanDeviceExtension(InDevice, VK_EXT_DEBUG_MARKER_EXTENSION_NAME, (VULKAN_ENABLE_DRAW_MARKERS & VULKAN_HAS_DEBUGGING_ENABLED))
-	{
-#if VULKAN_HAS_DEBUGGING_ENABLED
-		const int32 VulkanValidationOption = GValidationCvar.GetValueOnAnyThread();
-		bEnabledInCode = bEnabledInCode && (GRenderDocFound || VulkanValidationOption != 0);
-#endif
-	}
-
-	virtual void PostPhysicalDeviceProperties() override final
-	{
-		Device->SetDebugMarkersFound();
-	}
-};
-
-
-
 // ***** VK_AMD_buffer_marker (vendor)
 class FVulkanAMDBufferMarkerExtension : public FVulkanDeviceExtension
 {
@@ -1417,9 +1395,6 @@ FVulkanDeviceExtensionArray FVulkanDeviceExtension::GetUESupportedDeviceExtensio
 	ADD_CUSTOM_EXTENSION(FVulkanNVDeviceDiagnosticCheckpointsExtension);
 	ADD_CUSTOM_EXTENSION(FVulkanNVDeviceDiagnosticConfigExtension);
 
-	// Debug
-	ADD_CUSTOM_EXTENSION(FVulkanEXTDebugMarkerExtension);
-
 	// Add in platform specific extensions
 	FVulkanPlatform::GetDeviceExtensions(InDevice, OutUEDeviceExtensions);
 
@@ -1553,8 +1528,7 @@ FVulkanInstanceExtensionArray FVulkanInstanceExtension::GetUESupportedInstanceEx
 	OutUEInstanceExtensions.Add(MakeUnique<FVulkanInstanceExtension>(VK_EXT_SWAPCHAIN_COLOR_SPACE_EXTENSION_NAME, VULKAN_EXTENSION_ENABLED, VULKAN_EXTENSION_NOT_PROMOTED));
 
 	// Debug extensions :
-	OutUEInstanceExtensions.Add(MakeUnique<FVulkanInstanceExtension>(VK_EXT_DEBUG_UTILS_EXTENSION_NAME, VULKAN_HAS_DEBUGGING_ENABLED & VULKAN_SUPPORTS_DEBUG_UTILS, VULKAN_EXTENSION_NOT_PROMOTED, nullptr, FVulkanExtensionBase::ManuallyActivate));
-	OutUEInstanceExtensions.Add(MakeUnique<FVulkanInstanceExtension>(VK_EXT_DEBUG_REPORT_EXTENSION_NAME, VULKAN_HAS_DEBUGGING_ENABLED, VULKAN_EXTENSION_NOT_PROMOTED, nullptr, FVulkanExtensionBase::ManuallyActivate));
+	OutUEInstanceExtensions.Add(MakeUnique<FVulkanInstanceExtension>(VK_EXT_DEBUG_UTILS_EXTENSION_NAME, VULKAN_HAS_DEBUGGING_ENABLED, VULKAN_EXTENSION_NOT_PROMOTED, nullptr, FVulkanExtensionBase::ManuallyActivate));
 
 	// Extensions with custom classes :
 	OutUEInstanceExtensions.Add(MakeUnique<FVulkanEXTValidationFeaturesExtension>());
diff --git a/Engine/Source/Runtime/VulkanRHI/Private/VulkanLayers.cpp b/Engine/Source/Runtime/VulkanRHI/Private/VulkanLayers.cpp
--- a/Engine/Source/Runtime/VulkanRHI/Private/VulkanLayers.cpp
+++ b/Engine/Source/Runtime/VulkanRHI/Private/VulkanLayers.cpp
@@ -30,15 +30,6 @@ TAutoConsoleVariable<int32> GValidationCvar(
 	ECVF_ReadOnly | ECVF_RenderThreadSafe
 );
 
-static TAutoConsoleVariable<int32> GStandardValidationCvar(
-	TEXT("r.Vulkan.StandardValidation"),
-	2,
-	TEXT("2 to use VK_LAYER_KHRONOS_validation (default) if available\n")
-	TEXT("1 to use VK_LAYER_LUNARG_standard_validation if available, or \n")
-	TEXT("0 to use individual validation layers (removed)"),
-	ECVF_ReadOnly | ECVF_RenderThreadSafe
-);
-
 TAutoConsoleVariable<int32> GGPUValidationCvar(
 	TEXT("r.Vulkan.GPUValidation"),
 	0,
@@ -53,7 +44,6 @@ TAutoConsoleVariable<int32> GGPUValidationCvar(
 #endif
 
 #define KHRONOS_STANDARD_VALIDATION_LAYER_NAME	"VK_LAYER_KHRONOS_validation"
-#define STANDARD_VALIDATION_LAYER_NAME			"VK_LAYER_LUNARG_standard_validation"
 
 #endif // VULKAN_HAS_DEBUGGING_ENABLED
 
@@ -396,27 +386,14 @@ void FVulkanIntanceSetupHelper::AddDebugLayers(const TArray<FLayerWithExtensions
 	const bool bUseVulkanValidation = GRHIGlobals.IsDebugLayerEnabled;
 	if (!bGfxReconstructOrVkTrace && bUseVulkanValidation)
 	{
-		if (GStandardValidationCvar.GetValueOnAnyThread() != 0)
+		if (!AddRequestedLayer(KHRONOS_STANDARD_VALIDATION_LAYER_NAME, LayerProperties, UEExtensions, OutLayers))
 		{
-			if (GStandardValidationCvar.GetValueOnAnyThread() == 2)
-			{
-				if (!AddRequestedLayer(KHRONOS_STANDARD_VALIDATION_LAYER_NAME, LayerProperties, UEExtensions, OutLayers))
-				{
 #if PLATFORM_WINDOWS || PLATFORM_LINUX
-					//#todo-rco: We don't package DLLs so if this fails it means no DLL was found anywhere, so don't try to load standard validation layers
-					UE_LOG(LogVulkanRHI, Warning, TEXT("Unable to find Vulkan instance validation layer %s;  Do you have the Vulkan SDK Installed?"), TEXT(STANDARD_VALIDATION_LAYER_NAME));
+			//#todo-rco: We don't package DLLs so if this fails it means no DLL was found anywhere, so don't try to load standard validation layers
+			UE_LOG(LogVulkanRHI, Warning, TEXT("Unable to find Vulkan instance validation layer %s;  Do you have the Vulkan SDK Installed?"), TEXT(KHRONOS_STANDARD_VALIDATION_LAYER_NAME));
 #else
-					UE_LOG(LogVulkanRHI, Warning, TEXT("Unable to find Vulkan instance validation layer %s"), TEXT(STANDARD_VALIDATION_LAYER_NAME));
+			UE_LOG(LogVulkanRHI, Warning, TEXT("Unable to find Vulkan instance validation layer %s"), TEXT(KHRONOS_STANDARD_VALIDATION_LAYER_NAME));
 #endif
-				}
-			}
-			else
-			{
-				if (!AddRequestedLayer(STANDARD_VALIDATION_LAYER_NAME, LayerProperties, UEExtensions, OutLayers))
-				{
-					UE_LOG(LogVulkanRHI, Warning, TEXT("Unable to find Vulkan instance validation layer %s"), TEXT(STANDARD_VALIDATION_LAYER_NAME));
-				}
-			}
 		}
 	}
 
@@ -463,17 +440,10 @@ void FVulkanIntanceSetupHelper::AddDebugLayers(const TArray<FLayerWithExtensions
 			return false;
 		};
 
-#if VULKAN_SUPPORTS_DEBUG_UTILS
 		if (ActivateDebuggingExtension(VK_EXT_DEBUG_UTILS_EXTENSION_NAME))
 		{
 			ActiveDebugLayerExtension = FVulkanDynamicRHI::EActiveDebugLayerExtension::DebugUtilsExtension;
 		}
-		else
-#endif
-		if (ActivateDebuggingExtension(VK_EXT_DEBUG_REPORT_EXTENSION_NAME))
-		{
-			ActiveDebugLayerExtension = FVulkanDynamicRHI::EActiveDebugLayerExtension::DebugReportExtension;
-		}
 
 		const bool bRequiresValidationFeatures = (GGPUValidationCvar.GetValueOnAnyThread() != 0) ||
 			(FParse::Param(FCommandLine::Get(), TEXT("vulkanbestpractices"))) ||
@@ -486,37 +456,8 @@ void FVulkanIntanceSetupHelper::AddDebugLayers(const TArray<FLayerWithExtensions
 #endif	// VULKAN_HAS_DEBUGGING_ENABLED
 }
 
-
-
 // Return a list of debug layers to activate
 void FVulkanDeviceSetupHelper::AddDebugLayers(const TArray<FLayerWithExtensions>& LayerProperties, FVulkanDeviceExtensionArray& UEExtensions, TArray<const ANSICHAR*>& OutLayers)
 {
-#if VULKAN_HAS_DEBUGGING_ENABLED
-#if VULKAN_ENABLE_DRAW_MARKERS
-	GRenderDocFound = (FindLayerIndexInList(RENDERDOC_LAYER_NAME, LayerProperties) != INDEX_NONE);
-#else
-	GRenderDocFound = false;
-#endif
 
-	// Verify that all requested debugging device-layers are available. Skip validation layers under RenderDoc
-	if (!GRenderDocFound && GRHIGlobals.IsDebugLayerEnabled && (GStandardValidationCvar.GetValueOnAnyThread() != 0))
-	{
-		switch (GStandardValidationCvar.GetValueOnAnyThread())
-		{
-		case 1:
-			AddRequestedLayer(STANDARD_VALIDATION_LAYER_NAME, LayerProperties, UEExtensions, OutLayers);
-			break;
-		case 2:
-			AddRequestedLayer(KHRONOS_STANDARD_VALIDATION_LAYER_NAME, LayerProperties, UEExtensions, OutLayers);
-			break;
-		default:
-			break;
-		}
-		
-	}
-#endif	// VULKAN_HAS_DEBUGGING_ENABLED
 }
-
-
-
-
diff --git a/Engine/Source/Runtime/VulkanRHI/Private/VulkanRHI.cpp b/Engine/Source/Runtime/VulkanRHI/Private/VulkanRHI.cpp
--- a/Engine/Source/Runtime/VulkanRHI/Private/VulkanRHI.cpp
+++ b/Engine/Source/Runtime/VulkanRHI/Private/VulkanRHI.cpp
@@ -1208,7 +1208,6 @@ void FVulkanCommandListContext::RHIEndFrame()
 void FVulkanCommandListContext::RHIPushEvent(const TCHAR* Name, FColor Color)
 {
 #if VULKAN_ENABLE_DRAW_MARKERS
-#if 0//VULKAN_SUPPORTS_DEBUG_UTILS
 	if (auto CmdBeginLabel = Device->GetCmdBeginDebugLabel())
 	{
 		FTCHARToUTF8 Converter(Name);
@@ -1222,21 +1221,6 @@ void FVulkanCommandListContext::RHIPushEvent(const TCHAR* Name, FColor Color)
 		Label.color[3] = LColor.A;
 		CmdBeginLabel(GetCommandBufferManager()->GetActiveCmdBuffer()->GetHandle(), &Label);
 	}
-	else
-#endif
-	if (auto CmdDbgMarkerBegin = Device->GetCmdDbgMarkerBegin())
-	{
-		FTCHARToUTF8 Converter(Name);
-		VkDebugMarkerMarkerInfoEXT Info;
-		ZeroVulkanStruct(Info, VK_STRUCTURE_TYPE_DEBUG_MARKER_MARKER_INFO_EXT);
-		Info.pMarkerName = Converter.Get();
-		FLinearColor LColor(Color);
-		Info.color[0] = LColor.R;
-		Info.color[1] = LColor.G;
-		Info.color[2] = LColor.B;
-		Info.color[3] = LColor.A;
-		CmdDbgMarkerBegin(GetCommandBufferManager()->GetActiveCmdBuffer()->GetHandle(), &Info);
-	}
 #endif
 
 #if VULKAN_SUPPORTS_GPU_CRASH_DUMPS
@@ -1260,17 +1244,10 @@ void FVulkanCommandListContext::RHIPushEvent(const TCHAR* Name, FColor Color)
 void FVulkanCommandListContext::RHIPopEvent()
 {
 #if VULKAN_ENABLE_DRAW_MARKERS
-#if 0//VULKAN_SUPPORTS_DEBUG_UTILS
 	if (auto CmdEndLabel = Device->GetCmdEndDebugLabel())
 	{
 		CmdEndLabel(GetCommandBufferManager()->GetActiveCmdBuffer()->GetHandle());
 	}
-	else
-#endif
-	if (auto CmdDbgMarkerEnd = Device->GetCmdDbgMarkerEnd())
-	{
-		CmdDbgMarkerEnd(GetCommandBufferManager()->GetActiveCmdBuffer()->GetHandle());
-	}
 #endif
 
 #if VULKAN_SUPPORTS_GPU_CRASH_DUMPS
diff --git a/Engine/Source/Runtime/VulkanRHI/Private/VulkanRHIPrivate.h b/Engine/Source/Runtime/VulkanRHI/Private/VulkanRHIPrivate.h
--- a/Engine/Source/Runtime/VulkanRHI/Private/VulkanRHIPrivate.h
+++ b/Engine/Source/Runtime/VulkanRHI/Private/VulkanRHIPrivate.h
@@ -753,28 +753,15 @@ namespace VulkanRHI
 	}
 
 #if VULKAN_ENABLE_DRAW_MARKERS
-	inline void SetDebugMarkerName(PFN_vkDebugMarkerSetObjectNameEXT DebugMarkerSetObjectName, VkDevice VulkanDevice, VkImage Image, const char* ObjectName)
-	{
-		VkDebugMarkerObjectNameInfoEXT Info;
-		ZeroVulkanStruct(Info, VK_STRUCTURE_TYPE_DEBUG_MARKER_OBJECT_NAME_INFO_EXT);
-		Info.objectType = VK_DEBUG_REPORT_OBJECT_TYPE_IMAGE_EXT;
-		Info.object = (uint64)Image;
-		Info.pObjectName = ObjectName;
-		DebugMarkerSetObjectName(VulkanDevice, &Info);
-};
-
-#if 0//VULKAN_SUPPORTS_DEBUG_UTILS
 	inline void SetDebugName(PFN_vkSetDebugUtilsObjectNameEXT SetDebugName, VkDevice Device, VkImage Image, const char* Name)
 	{
-		FTCHARToUTF8 Converter(Name);
 		VkDebugUtilsObjectNameInfoEXT Info;
 		ZeroVulkanStruct(Info, VK_STRUCTURE_TYPE_DEBUG_UTILS_OBJECT_NAME_INFO_EXT);
 		Info.objectType = VK_OBJECT_TYPE_IMAGE;
 		Info.objectHandle = (uint64)Image;
-		Info.pObjectName = Converter.Get();
+		Info.pObjectName = Name;
 		SetDebugName(Device, &Info);
 }
-#endif
 #endif
 
 	// Merge a depth and a stencil layout for drivers that don't support VK_KHR_separate_depth_stencil_layouts
diff --git a/Engine/Source/Runtime/VulkanRHI/Private/VulkanTexture.cpp b/Engine/Source/Runtime/VulkanRHI/Private/VulkanTexture.cpp
--- a/Engine/Source/Runtime/VulkanRHI/Private/VulkanTexture.cpp
+++ b/Engine/Source/Runtime/VulkanRHI/Private/VulkanTexture.cpp
@@ -2060,21 +2060,12 @@ void FVulkanDynamicRHI::RHIBindDebugLabelName(FRHICommandListBase& RHICmdList, F
 #endif
 
 #if VULKAN_ENABLE_DRAW_MARKERS
-#if 0//VULKAN_SUPPORTS_DEBUG_UTILS
 	if (auto* SetDebugName = Device->GetSetDebugName())
 	{
 		FVulkanTexture* VulkanTexture = ResourceCast(TextureRHI);
 		FTCHARToUTF8 Converter(Name);
 		VulkanRHI::SetDebugName(SetDebugName, Device->GetInstanceHandle(), VulkanTexture->Image, Converter.Get());
 	}
-	else
-#endif
-	if (auto* SetObjectName = Device->GetDebugMarkerSetObjectName())
-	{
-		FVulkanTexture* VulkanTexture = ResourceCast(TextureRHI);
-		FTCHARToUTF8 Converter(Name);
-		VulkanRHI::SetDebugMarkerName(SetObjectName, Device->GetInstanceHandle(), VulkanTexture->Image, Converter.Get());
-	}
 #endif
 	FName DebugName(Name);
 	TextureRHI->SetName(DebugName);
diff --git a/Engine/Source/Runtime/VulkanRHI/Private/VulkanViewport.cpp b/Engine/Source/Runtime/VulkanRHI/Private/VulkanViewport.cpp
--- a/Engine/Source/Runtime/VulkanRHI/Private/VulkanViewport.cpp
+++ b/Engine/Source/Runtime/VulkanRHI/Private/VulkanViewport.cpp
@@ -739,9 +739,9 @@ void FVulkanViewport::CreateSwapchain(FVulkanSwapChainRecreateInfo* RecreateInfo
 			}
 
 #if VULKAN_ENABLE_DRAW_MARKERS
-			if (Device->GetDebugMarkerSetObjectName())
+			if (Device->GetSetDebugName())
 			{
-				VulkanRHI::SetDebugMarkerName(Device->GetDebugMarkerSetObjectName(), Device->GetInstanceHandle(), BackBufferImages[Index], "RenderingBackBuffer");
+				VulkanRHI::SetDebugName(Device->GetSetDebugName(), Device->GetInstanceHandle(), BackBufferImages[Index], "RenderingBackBuffer");
 			}
 #endif
 		}
@@ -785,9 +785,9 @@ void FVulkanViewport::CreateSwapchain(FVulkanSwapChainRecreateInfo* RecreateInfo
 
 		RenderingBackBuffer = new FVulkanTexture(*Device, Desc, nullptr);
 #if VULKAN_ENABLE_DRAW_MARKERS
-		if (Device->GetDebugMarkerSetObjectName())
+		if (Device->GetSetDebugName())
 		{
-			VulkanRHI::SetDebugMarkerName(Device->GetDebugMarkerSetObjectName(), Device->GetInstanceHandle(), RenderingBackBuffer->Image, "RenderingBackBuffer");
+			VulkanRHI::SetDebugName(Device->GetSetDebugName(), Device->GetInstanceHandle(), RenderingBackBuffer->Image, "RenderingBackBuffer");
 		}
 #endif
 	}
diff --git a/Engine/Source/Runtime/VulkanRHI/Public/VulkanConfiguration.h b/Engine/Source/Runtime/VulkanRHI/Public/VulkanConfiguration.h
--- a/Engine/Source/Runtime/VulkanRHI/Public/VulkanConfiguration.h
+++ b/Engine/Source/Runtime/VulkanRHI/Public/VulkanConfiguration.h
@@ -170,10 +170,6 @@ struct VkAllocationCallbacks;
 	#define VULKAN_SUPPORTS_GPU_CRASH_DUMPS						(VULKAN_SUPPORTS_AMD_BUFFER_MARKER || VULKAN_SUPPORTS_NV_DIAGNOSTIC_CONFIG)
 #endif
 
-#ifndef VULKAN_SUPPORTS_DEBUG_UTILS
-	#define VULKAN_SUPPORTS_DEBUG_UTILS							1
-#endif
-
 #ifndef VULKAN_SUPPORTS_SCALAR_BLOCK_LAYOUT
 	#define VULKAN_SUPPORTS_SCALAR_BLOCK_LAYOUT					1
 #endif
diff --git a/Engine/Source/Runtime/VulkanRHI/Public/VulkanDynamicRHI.h b/Engine/Source/Runtime/VulkanRHI/Public/VulkanDynamicRHI.h
--- a/Engine/Source/Runtime/VulkanRHI/Public/VulkanDynamicRHI.h
+++ b/Engine/Source/Runtime/VulkanRHI/Public/VulkanDynamicRHI.h
@@ -352,8 +352,7 @@ public:
 		None,
 		GfxReconstructLayer,
 		VkTraceLayer,
-		DebugUtilsExtension,
-		DebugReportExtension
+		DebugUtilsExtension
 	};
 
 protected:
@@ -362,10 +361,7 @@ protected:
 	EActiveDebugLayerExtension ActiveDebugLayerExtension = EActiveDebugLayerExtension::None;
 
 #if VULKAN_HAS_DEBUGGING_ENABLED
-#if VULKAN_SUPPORTS_DEBUG_UTILS
 	VkDebugUtilsMessengerEXT Messenger = VK_NULL_HANDLE;
-#endif
-	VkDebugReportCallbackEXT MsgCallback = VK_NULL_HANDLE;
 
 	void SetupDebugLayerCallback();
 	void RemoveDebugLayerCallback();
