Moving vulkan incudes to VulkanThirdParty.h

diff --git a/Engine/Source/Runtime/VulkanRHI/Private/VulkanLoader.h b/Engine/Source/Runtime/VulkanRHI/Private/VulkanLoader.h
--- a/Engine/Source/Runtime/VulkanRHI/Private/VulkanLoader.h
+++ b/Engine/Source/Runtime/VulkanRHI/Private/VulkanLoader.h
@@ -2,10 +2,7 @@
 
 #pragma once
 
-#include "CoreTypes.h"
-
-#define VK_NO_PROTOTYPES
-#include "vulkan.h"
+#include "VulkanThirdParty.h"
 
 // List all instance Vulkan entry points used by Unreal that need to be loaded manually
 #define ENUM_VK_ENTRYPOINTS_INSTANCE(EnumMacro) \
diff --git a/Engine/Source/Runtime/VulkanRHI/Public/IVulkanDynamicRHI.h b/Engine/Source/Runtime/VulkanRHI/Public/IVulkanDynamicRHI.h
--- a/Engine/Source/Runtime/VulkanRHI/Public/IVulkanDynamicRHI.h
+++ b/Engine/Source/Runtime/VulkanRHI/Public/IVulkanDynamicRHI.h
@@ -3,17 +3,7 @@
 #pragma once
 
 #include "RHI.h"
-
-#if defined(VK_USE_PLATFORM_WIN32_KHR)
-	#include "Microsoft/AllowMicrosoftPlatformTypes.h"
-	#include <windows.h>
-	#include "Microsoft/HideMicrosoftPlatformTypes.h"
-#endif
-
-THIRD_PARTY_INCLUDES_START
-	#define VK_NO_PROTOTYPES
-	#include <vulkan.h>
-THIRD_PARTY_INCLUDES_END
+#include "VulkanThirdParty.h"
 
 struct FVulkanRHIAllocationInfo
 {
diff --git a/Engine/Source/Runtime/VulkanRHI/Public/VulkanShaderResources.h b/Engine/Source/Runtime/VulkanRHI/Public/VulkanShaderResources.h
--- a/Engine/Source/Runtime/VulkanRHI/Public/VulkanShaderResources.h
+++ b/Engine/Source/Runtime/VulkanRHI/Public/VulkanShaderResources.h
@@ -8,6 +8,8 @@
 
 #include "BoundShaderStateCache.h"
 #include "CrossCompilerCommon.h"
+#include "VulkanCommon.h"
+#include "VulkanThirdParty.h"
 
 static inline VkDescriptorType BindingToDescriptorType(EVulkanBindingType::EType Type)
 {
diff --git a/Engine/Source/Runtime/VulkanRHI/Public/VulkanThirdParty.h b/Engine/Source/Runtime/VulkanRHI/Public/VulkanThirdParty.h
new file mode 100644
--- /dev/null
+++ b/Engine/Source/Runtime/VulkanRHI/Public/VulkanThirdParty.h
@@ -0,0 +1,14 @@
+// Copyright Epic Games, Inc. All Rights Reserved.
+
+#pragma once
+
+#if defined(VK_USE_PLATFORM_WIN32_KHR)
+	#include "Microsoft/AllowMicrosoftPlatformTypes.h"
+	#include <windows.h>
+	#include "Microsoft/HideMicrosoftPlatformTypes.h"
+#endif
+
+THIRD_PARTY_INCLUDES_START
+	#define VK_NO_PROTOTYPES
+	#include <vulkan.h>
+THIRD_PARTY_INCLUDES_END
