Compilation error for program targets with engine dependency

diff --git a/Engine/Source/Runtime/VulkanRHI/Private/VulkanDevice.cpp b/Engine/Source/Runtime/VulkanRHI/Private/VulkanDevice.cpp
--- a/Engine/Source/Runtime/VulkanRHI/Private/VulkanDevice.cpp
+++ b/Engine/Source/Runtime/VulkanRHI/Private/VulkanDevice.cpp
@@ -537,12 +537,14 @@ void FVulkanDevice::SetupDrawMarkers()
 		}
 	}
 
+#if VULKAN_HAS_DEBUGGING_ENABLED
 	if (bDebugMarkersFound && GRenderDocFound)
 	{
 		// We're running under RenderDoc or other trace tool, so enable capturing mode
 		EnableDrawMarkers();
 	}
 #endif
+#endif
 
 #if VULKAN_ENABLE_DUMP_LAYER
 	EnableDrawMarkers();
