Fix Vulkan validation error related to view usage

In cases where VkImageViewUsageCreateInfo is added to the VkImageViewCreateInfo we only used VK_IMAGE_USAGE_COLOR_ATTACHMENT_BIT.
This change preserves

diff --git a/Engine/Source/Runtime/VulkanRHI/Private/VulkanResources.h b/Engine/Source/Runtime/VulkanRHI/Private/VulkanResources.h
--- a/Engine/Source/Runtime/VulkanRHI/Private/VulkanResources.h
+++ b/Engine/Source/Runtime/VulkanRHI/Private/VulkanResources.h
@@ -742,6 +742,7 @@ public:
 
 	FVulkanDevice* Device;
 	VkImage Image;
+	VkImageUsageFlags ImageUsageFlags;
 	VkFormat StorageFormat;  // Removes SRGB if requested, used to upload data
 	VkFormat ViewFormat;  // Format for SRVs, render targets
 	VkMemoryPropertyFlags MemProps;
diff --git a/Engine/Source/Runtime/VulkanRHI/Private/VulkanTexture.cpp b/Engine/Source/Runtime/VulkanRHI/Private/VulkanTexture.cpp
--- a/Engine/Source/Runtime/VulkanRHI/Private/VulkanTexture.cpp
+++ b/Engine/Source/Runtime/VulkanRHI/Private/VulkanTexture.cpp
@@ -1474,6 +1474,7 @@ FVulkanTexture::FVulkanTexture(FRHICommandListBase* RHICmdList, FVulkanDevice& I
 	, PartialView(nullptr)
 	, Device(&InDevice)
 	, Image(VK_NULL_HANDLE)
+	, ImageUsageFlags(0)
 	, StorageFormat(VK_FORMAT_UNDEFINED)
 	, ViewFormat(VK_FORMAT_UNDEFINED)
 	, MemProps(VK_MEMORY_PROPERTY_DEVICE_LOCAL_BIT)
@@ -1605,6 +1606,7 @@ FVulkanTexture::FVulkanTexture(FRHICommandListBase* RHICmdList, FVulkanDevice& I
 
 	Tiling = ImageCreateInfo.ImageCreateInfo.tiling;
 	check(Tiling == VK_IMAGE_TILING_LINEAR || Tiling == VK_IMAGE_TILING_OPTIMAL);
+	ImageUsageFlags = ImageCreateInfo.ImageCreateInfo.usage;
 
 	const VkImageLayout InitialLayout = GetInitialLayoutFromRHIAccess(InCreateDesc.InitialState, bRenderTarget && IsDepthOrStencilAspect(), SupportsSampling());
 	const bool bDoInitialClear = VKHasAnyFlags(ImageCreateInfo.ImageCreateInfo.usage, VK_IMAGE_USAGE_SAMPLED_BIT) && EnumHasAnyFlags(InCreateDesc.Flags, TexCreate_RenderTargetable | TexCreate_DepthStencilTargetable);
diff --git a/Engine/Source/Runtime/VulkanRHI/Private/VulkanViewport.cpp b/Engine/Source/Runtime/VulkanRHI/Private/VulkanViewport.cpp
--- a/Engine/Source/Runtime/VulkanRHI/Private/VulkanViewport.cpp
+++ b/Engine/Source/Runtime/VulkanRHI/Private/VulkanViewport.cpp
@@ -373,7 +373,7 @@ FVulkanFramebuffer::FVulkanFramebuffer(FVulkanDevice& Device, const FRHISetRende
 				, ArraySliceIndex
 				, NumArraySlices
 				, true
-				, VK_IMAGE_USAGE_COLOR_ATTACHMENT_BIT
+				, VK_IMAGE_USAGE_COLOR_ATTACHMENT_BIT | (Texture->ImageUsageFlags & VK_IMAGE_USAGE_INPUT_ATTACHMENT_BIT)
 			);
 		}
 		else if (Texture->GetViewType() == VK_IMAGE_VIEW_TYPE_CUBE)
@@ -392,7 +392,7 @@ FVulkanFramebuffer::FVulkanFramebuffer(FVulkanDevice& Device, const FRHISetRende
 				, InRTInfo.ColorRenderTarget[Index].ArraySliceIndex
 				, 1
 				, true
-				, VK_IMAGE_USAGE_COLOR_ATTACHMENT_BIT
+				, VK_IMAGE_USAGE_COLOR_ATTACHMENT_BIT | (Texture->ImageUsageFlags & VK_IMAGE_USAGE_INPUT_ATTACHMENT_BIT)
 			);
 		}
 		else if (Texture->GetViewType() == VK_IMAGE_VIEW_TYPE_3D)
@@ -408,7 +408,7 @@ FVulkanFramebuffer::FVulkanFramebuffer(FVulkanDevice& Device, const FRHISetRende
 				, 0
 				, Desc.Depth
 				, true
-				, VK_IMAGE_USAGE_COLOR_ATTACHMENT_BIT
+				, VK_IMAGE_USAGE_COLOR_ATTACHMENT_BIT | (Texture->ImageUsageFlags & VK_IMAGE_USAGE_INPUT_ATTACHMENT_BIT)
 			);
 		}
 		else
