Fix Vulkan debug markers in RenderDoc

diff --git a/Engine/Source/Runtime/VulkanRHI/Private/VulkanLayers.cpp b/Engine/Source/Runtime/VulkanRHI/Private/VulkanLayers.cpp
--- a/Engine/Source/Runtime/VulkanRHI/Private/VulkanLayers.cpp
+++ b/Engine/Source/Runtime/VulkanRHI/Private/VulkanLayers.cpp
@@ -460,5 +460,11 @@ void FVulkanIntanceSetupHelper::AddDebugLayers(const TArray<FLayerWithExtensions
 // Return a list of debug layers to activate
 void FVulkanDeviceSetupHelper::AddDebugLayers(const TArray<FLayerWithExtensions>& LayerProperties, FVulkanDeviceExtensionArray& UEExtensions, TArray<const ANSICHAR*>& OutLayers)
 {
-
+#if VULKAN_HAS_DEBUGGING_ENABLED
+#if VULKAN_ENABLE_DRAW_MARKERS
+	GRenderDocFound = (FindLayerIndexInList(RENDERDOC_LAYER_NAME, LayerProperties) != INDEX_NONE);
+#else
+	GRenderDocFound = false;
+#endif // VULKAN_ENABLE_DRAW_MARKERS
+#endif // VULKAN_HAS_DEBUGGING_ENABLED
 }
