VulkanRHI: Fix wrong sType for VkDescriptorAddressInfoEXT structs.

#rnx

diff --git a/Engine/Source/Runtime/VulkanRHI/Private/VulkanPipelineState.cpp b/Engine/Source/Runtime/VulkanRHI/Private/VulkanPipelineState.cpp
--- a/Engine/Source/Runtime/VulkanRHI/Private/VulkanPipelineState.cpp
+++ b/Engine/Source/Runtime/VulkanRHI/Private/VulkanPipelineState.cpp
@@ -257,7 +257,7 @@ void FVulkanComputePipelineDescriptorState::UpdateBindlessDescriptors(FVulkanCom
 				// Make sure it wasn't written to already
 				VkDescriptorAddressInfoEXT& DescriptorAddressInfo = DescriptorAddressInfos[BindingIndex];
 				check(DescriptorAddressInfo.sType == 0);
-				DescriptorAddressInfo.sType = VK_STRUCTURE_TYPE_BUFFER_DEVICE_ADDRESS_INFO;
+				DescriptorAddressInfo.sType = VK_STRUCTURE_TYPE_DESCRIPTOR_ADDRESS_INFO_EXT;
 				DescriptorAddressInfo.address = UniformBufferUploader->GetCPUBufferAddress() + RingBufferOffset;
 				DescriptorAddressInfo.range = UBSize;
 
@@ -292,7 +292,7 @@ void FVulkanComputePipelineDescriptorState::UpdateBindlessDescriptors(FVulkanCom
 		BufferInfo.buffer = WriteDescriptorSet.pBufferInfo->buffer;
 		VkDeviceAddress BufferAddress = VulkanRHI::vkGetBufferDeviceAddressKHR(Device->GetInstanceHandle(), &BufferInfo);
 
-		DescriptorAddressInfo.sType = VK_STRUCTURE_TYPE_BUFFER_DEVICE_ADDRESS_INFO;
+		DescriptorAddressInfo.sType = VK_STRUCTURE_TYPE_DESCRIPTOR_ADDRESS_INFO_EXT;
 		DescriptorAddressInfo.address = BufferAddress + WriteDescriptorSet.pBufferInfo->offset;
 		DescriptorAddressInfo.range = WriteDescriptorSet.pBufferInfo->range;
 	}
@@ -483,7 +483,7 @@ void FVulkanGraphicsPipelineDescriptorState::UpdateBindlessDescriptors(FVulkanCo
 						// Make sure it wasn't written to already
 						VkDescriptorAddressInfoEXT& DescriptorAddressInfo = DescriptorAddressInfos[BindingIndex];
 						check(DescriptorAddressInfo.sType == 0);
-						DescriptorAddressInfo.sType = VK_STRUCTURE_TYPE_BUFFER_DEVICE_ADDRESS_INFO;
+						DescriptorAddressInfo.sType = VK_STRUCTURE_TYPE_DESCRIPTOR_ADDRESS_INFO_EXT;
 						DescriptorAddressInfo.address = UniformBufferUploader->GetCPUBufferAddress() + RingBufferOffset;
 						DescriptorAddressInfo.range = UBSize;
 
@@ -518,7 +518,7 @@ void FVulkanGraphicsPipelineDescriptorState::UpdateBindlessDescriptors(FVulkanCo
 				BufferInfo.buffer = WriteDescriptorSet.pBufferInfo->buffer;
 				VkDeviceAddress BufferAddress = VulkanRHI::vkGetBufferDeviceAddressKHR(Device->GetInstanceHandle(), &BufferInfo);
 
-				DescriptorAddressInfo.sType = VK_STRUCTURE_TYPE_BUFFER_DEVICE_ADDRESS_INFO;
+				DescriptorAddressInfo.sType = VK_STRUCTURE_TYPE_DESCRIPTOR_ADDRESS_INFO_EXT;
 				DescriptorAddressInfo.address = BufferAddress + WriteDescriptorSet.pBufferInfo->offset;
 				DescriptorAddressInfo.range = WriteDescriptorSet.pBufferInfo->range;
 			}
