VulkanRHI: Fix sType on VkDescriptorAddressInfoEXT structures in ray tracing.

#rnx

diff --git a/Engine/Source/Runtime/VulkanRHI/Private/VulkanRayTracing.cpp b/Engine/Source/Runtime/VulkanRHI/Private/VulkanRayTracing.cpp
--- a/Engine/Source/Runtime/VulkanRHI/Private/VulkanRayTracing.cpp
+++ b/Engine/Source/Runtime/VulkanRHI/Private/VulkanRayTracing.cpp
@@ -1823,14 +1823,14 @@ static FVulkanPipelineBarrier SetRayGenResources(FVulkanDevice* Device, FVulkanC
 					for (uint32 SkipIndex = 0; SkipIndex < NumSkippedSlots; ++SkipIndex)
 					{
 						VkDescriptorAddressInfoEXT& DescriptorAddressInfo = StageUBs[ShaderStage::EStage::RayGen].AddZeroed_GetRef();
-						DescriptorAddressInfo.sType = VK_STRUCTURE_TYPE_BUFFER_DEVICE_ADDRESS_INFO;
+						DescriptorAddressInfo.sType = VK_STRUCTURE_TYPE_DESCRIPTOR_ADDRESS_INFO_EXT;
 					}
 
 					NumSkippedSlots = 0;
 				}
 				
 				VkDescriptorAddressInfoEXT& DescriptorAddressInfo = StageUBs[ShaderStage::EStage::RayGen].AddZeroed_GetRef();
-				DescriptorAddressInfo.sType = VK_STRUCTURE_TYPE_BUFFER_DEVICE_ADDRESS_INFO;
+				DescriptorAddressInfo.sType = VK_STRUCTURE_TYPE_DESCRIPTOR_ADDRESS_INFO_EXT;
 				DescriptorAddressInfo.address = UniformBuffer->GetDeviceAddress();
 				DescriptorAddressInfo.range = UniformBuffer->GetSize();
 
