Add support to render to cubemap array face on vulkan

diff --git a/Engine/Source/Runtime/VulkanRHI/Private/VulkanViewport.cpp b/Engine/Source/Runtime/VulkanRHI/Private/VulkanViewport.cpp
--- a/Engine/Source/Runtime/VulkanRHI/Private/VulkanViewport.cpp
+++ b/Engine/Source/Runtime/VulkanRHI/Private/VulkanViewport.cpp
@@ -376,7 +376,7 @@ FVulkanFramebuffer::FVulkanFramebuffer(FVulkanDevice& Device, const FRHISetRende
 				, VK_IMAGE_USAGE_COLOR_ATTACHMENT_BIT | (Texture->ImageUsageFlags & VK_IMAGE_USAGE_INPUT_ATTACHMENT_BIT)
 			);
 		}
-		else if (Texture->GetViewType() == VK_IMAGE_VIEW_TYPE_CUBE)
+		else if (Texture->GetViewType() == VK_IMAGE_VIEW_TYPE_CUBE || Texture->GetViewType() == VK_IMAGE_VIEW_TYPE_CUBE_ARRAY)
 		{
 			// Cube always renders one face at a time
 			INC_DWORD_STAT(STAT_VulkanNumImageViews);
