VulkanRHI: Allow R64 format to be viewed as R32G32 for utility functions like clears.  Fixes validation error about component width mismatch (32bit vs 64bit).

#rnx

diff --git a/Engine/Source/Runtime/RenderCore/Private/RenderGraphUtils.cpp b/Engine/Source/Runtime/RenderCore/Private/RenderGraphUtils.cpp
--- a/Engine/Source/Runtime/RenderCore/Private/RenderGraphUtils.cpp
+++ b/Engine/Source/Runtime/RenderCore/Private/RenderGraphUtils.cpp
@@ -579,6 +579,15 @@ void AddClearUAVPass(FRDGBuilder& GraphBuilder, ERHIFeatureLevel::Type FeatureLe
 
 	check(TextureUAV && RectCoordBufferSRV);
 
+	const FRDGTextureRef Texture = TextureUAV->GetParent();
+	const FIntPoint TextureSize = Texture->Desc.Extent;
+
+	// Create a R32G32 view of the R64 instead of adding a permutation to the clear shader
+	if (Texture->Desc.Format == PF_R64_UINT)
+	{
+		TextureUAV = GraphBuilder.CreateUAV(Texture, ERDGUnorderedAccessViewFlags::None, PF_R32G32_UINT);
+	}
+
 	FClearUAVRectsParameters* PassParameters = GraphBuilder.AllocParameters<FClearUAVRectsParameters>();
 
 	PassParameters->PS.ClearValue.X = ClearValues[0];
@@ -590,9 +599,6 @@ void AddClearUAVPass(FRDGBuilder& GraphBuilder, ERHIFeatureLevel::Type FeatureLe
 	auto* ShaderMap = GetGlobalShaderMap(FeatureLevel);
 	auto PixelShader = ShaderMap->GetShader<FClearUAVRectsPS>();
 
-	const FRDGTextureRef Texture = TextureUAV->GetParent();
-	const FIntPoint TextureSize = Texture->Desc.Extent;
-
 	FPixelShaderUtils::AddRasterizeToRectsPass<FClearUAVRectsPS>(GraphBuilder,
 		ShaderMap,
 		RDG_EVENT_NAME("ClearTextureRects(%s %s %dx%d Mip=%d)",
diff --git a/Engine/Source/Runtime/VulkanRHI/Private/VulkanResources.h b/Engine/Source/Runtime/VulkanRHI/Private/VulkanResources.h
--- a/Engine/Source/Runtime/VulkanRHI/Private/VulkanResources.h
+++ b/Engine/Source/Runtime/VulkanRHI/Private/VulkanResources.h
@@ -632,7 +632,8 @@ public:
 		VkImageFormatListCreateInfoKHR ImageFormatListCreateInfo;
 		//used when TexCreate_External is given
 		VkExternalMemoryImageCreateInfoKHR ExternalMemImageCreateInfo;
-		VkFormat FormatsUsed[2];
+		// Array of formats used for mutable formats
+		TArray<VkFormat, TInlineAllocator<2>> FormatsUsed;
 	};
 
 	// Seperate method for creating VkImageCreateInfo
diff --git a/Engine/Source/Runtime/VulkanRHI/Private/VulkanTexture.cpp b/Engine/Source/Runtime/VulkanRHI/Private/VulkanTexture.cpp
--- a/Engine/Source/Runtime/VulkanRHI/Private/VulkanTexture.cpp
+++ b/Engine/Source/Runtime/VulkanRHI/Private/VulkanTexture.cpp
@@ -270,19 +270,32 @@ void FVulkanTexture::GenerateImageCreateInfo(
 
 	ImageCreateInfo.flags = (ResourceType == VK_IMAGE_VIEW_TYPE_CUBE || ResourceType == VK_IMAGE_VIEW_TYPE_CUBE_ARRAY) ? VK_IMAGE_CREATE_CUBE_COMPATIBLE_BIT : 0;
 
-
-	if (EnumHasAllFlags(UEFlags, TexCreate_SRGB))
+	const bool bNeedsMutableFormat = (EnumHasAllFlags(UEFlags, TexCreate_SRGB) || (InDesc.Format == PF_R64_UINT));
+	if (bNeedsMutableFormat)
 	{
-		if(InDevice.GetOptionalExtensions().HasKHRImageFormatList)
+		if (InDevice.GetOptionalExtensions().HasKHRImageFormatList)
 		{
 			VkImageFormatListCreateInfoKHR& ImageFormatListCreateInfo = OutImageCreateInfo.ImageFormatListCreateInfo;
 			ZeroVulkanStruct(ImageFormatListCreateInfo, VK_STRUCTURE_TYPE_IMAGE_FORMAT_LIST_CREATE_INFO_KHR);
 			ImageFormatListCreateInfo.pNext = ImageCreateInfo.pNext;
 			ImageCreateInfo.pNext = &ImageFormatListCreateInfo;
-			ImageFormatListCreateInfo.viewFormatCount = 2;
-			ImageFormatListCreateInfo.pViewFormats = OutImageCreateInfo.FormatsUsed;
-			OutImageCreateInfo.FormatsUsed[0] = nonSrgbFormat;
-			OutImageCreateInfo.FormatsUsed[1] = srgbFormat;
+
+			// Allow non-SRGB views to be created for SRGB textures
+			if (EnumHasAllFlags(UEFlags, TexCreate_SRGB))
+			{
+				OutImageCreateInfo.FormatsUsed.Add(nonSrgbFormat);
+				OutImageCreateInfo.FormatsUsed.Add(srgbFormat);
+			}
+
+			// Make it possible to create R32G32 views of R64 images for utilities like clears
+			if (InDesc.Format == PF_R64_UINT)
+			{
+				OutImageCreateInfo.FormatsUsed.Add(nonSrgbFormat);
+				OutImageCreateInfo.FormatsUsed.Add(UEToVkTextureFormat(PF_R32G32_UINT, false));
+			}
+
+			ImageFormatListCreateInfo.pViewFormats = OutImageCreateInfo.FormatsUsed.GetData();
+			ImageFormatListCreateInfo.viewFormatCount = OutImageCreateInfo.FormatsUsed.Num();
 		}
 
 		ImageCreateInfo.flags |= VK_IMAGE_CREATE_MUTABLE_FORMAT_BIT;
