Works arround a AMD driver bug on memory order of InterlockedMax on a Texture2D corrupting TSR's parallax disocclusion mask

diff --git a/Engine/Source/Runtime/VulkanRHI/Private/VulkanTexture.cpp b/Engine/Source/Runtime/VulkanRHI/Private/VulkanTexture.cpp
--- a/Engine/Source/Runtime/VulkanRHI/Private/VulkanTexture.cpp
+++ b/Engine/Source/Runtime/VulkanRHI/Private/VulkanTexture.cpp
@@ -212,6 +212,12 @@ void FVulkanTexture::GenerateImageCreateInfo(
 		bForceLinearTexture = true;
 	}
 
+	// Works arround an AMD driver bug where InterlockedMax() on a R32 Texture2D ends up with incorrect memory order swizzling
+	if (IsRHIDeviceAMD() && (InDesc.Format == PF_R32_UINT && UEFlags == (TexCreate_ShaderResource | TexCreate_UAV | TexCreate_AtomicCompatible)))
+	{
+		bForceLinearTexture = true;
+	}
+
 	checkf(TextureFormat != VK_FORMAT_UNDEFINED, TEXT("PixelFormat %d, is not supported for images"), (int32)InDesc.Format);
 	VkImageCreateInfo& ImageCreateInfo = OutImageCreateInfo.ImageCreateInfo;
 	ZeroVulkanStruct(ImageCreateInfo, VK_STRUCTURE_TYPE_IMAGE_CREATE_INFO);
