Fix incorrect memory regions used for RT acceleration structures in Vulkan RHI

Also set object name for AS handles

diff --git a/Engine/Source/Runtime/VulkanRHI/Private/VulkanDevice.h b/Engine/Source/Runtime/VulkanRHI/Private/VulkanDevice.h
--- a/Engine/Source/Runtime/VulkanRHI/Private/VulkanDevice.h
+++ b/Engine/Source/Runtime/VulkanRHI/Private/VulkanDevice.h
@@ -25,7 +25,7 @@ class FVulkanRayTracingCompactionRequestHandler;
 #define VULKAN_USE_DEBUG_NAMES 1
 
 #if VULKAN_USE_DEBUG_NAMES
-#define VULKAN_SET_DEBUG_NAME(Device, Type, Handle, Format, ...) Device.VulkanSetObjectName(Type, (uint64)Handle, *FString::Printf(Format, __VA_ARGS__))
+#define VULKAN_SET_DEBUG_NAME(Device, Type, Handle, Format, ...) (Device).VulkanSetObjectName(Type, (uint64)Handle, *FString::Printf(Format, __VA_ARGS__))
 #else
 #define VULKAN_SET_DEBUG_NAME(Device, Type, Handle, Format, ...) do{}while(0)
 #endif
diff --git a/Engine/Source/Runtime/VulkanRHI/Private/VulkanRayTracing.cpp b/Engine/Source/Runtime/VulkanRHI/Private/VulkanRayTracing.cpp
--- a/Engine/Source/Runtime/VulkanRHI/Private/VulkanRayTracing.cpp
+++ b/Engine/Source/Runtime/VulkanRHI/Private/VulkanRayTracing.cpp
@@ -372,7 +372,8 @@ FVulkanRayTracingGeometry::FVulkanRayTracingGeometry(FRHICommandListBase& RHICmd
 	CreateInfo.size = SizeInfo.ResultSize;
 	CreateInfo.type = VK_ACCELERATION_STRUCTURE_TYPE_BOTTOM_LEVEL_KHR;
 	VERIFYVULKANRESULT(VulkanDynamicAPI::vkCreateAccelerationStructureKHR(NativeDevice, &CreateInfo, VULKAN_CPU_ALLOCATOR, &Handle));
-	
+	VULKAN_SET_DEBUG_NAME(*Device, VK_OBJECT_TYPE_ACCELERATION_STRUCTURE_KHR, Handle, TEXT("%s"), *DebugName.ToString());
+
 	INC_MEMORY_STAT_BY(STAT_VulkanRayTracingUsedVideoMemory, SizeInfo.ResultSize);
 	INC_MEMORY_STAT_BY(STAT_VulkanRayTracingBLASMemory, SizeInfo.ResultSize);
 	if (Initializer.bAllowUpdate)
@@ -486,7 +487,7 @@ void FVulkanRayTracingGeometry::CompactAccelerationStructure(FVulkanCmdBuffer& C
 	CreateInfo.size = InSizeAfterCompaction;
 	CreateInfo.type = VK_ACCELERATION_STRUCTURE_TYPE_BOTTOM_LEVEL_KHR;
 	VERIFYVULKANRESULT(VulkanDynamicAPI::vkCreateAccelerationStructureKHR(NativeDevice, &CreateInfo, VULKAN_CPU_ALLOCATOR, &Handle));
-
+	VULKAN_SET_DEBUG_NAME(*Device, VK_OBJECT_TYPE_ACCELERATION_STRUCTURE_KHR, Handle, TEXT("%s (compact)"), *DebugName.ToString());
 	VkAccelerationStructureDeviceAddressInfoKHR DeviceAddressInfo;
 	ZeroVulkanStruct(DeviceAddressInfo, VK_STRUCTURE_TYPE_ACCELERATION_STRUCTURE_DEVICE_ADDRESS_INFO_KHR);
 	DeviceAddressInfo.accelerationStructure = Handle;
@@ -667,13 +668,12 @@ FVulkanRayTracingScene::FVulkanRayTracingScene(FRayTracingSceneInitializer2 InIn
 	{
 		FLayerData& Layer = Layers[LayerIndex];
 
-		FRayTracingAccelerationStructureSize LayerSizeInfo = RHICalcRayTracingSceneSize(Initializer.NumNativeInstancesPerLayer[LayerIndex], BuildFlags);
-
+		Layer.SizeInfo = RHICalcRayTracingSceneSize(Initializer.NumNativeInstancesPerLayer[LayerIndex], BuildFlags);
 		Layer.BufferOffset = Align(SizeInfo.ResultSize, GRHIRayTracingAccelerationStructureAlignment);
 		Layer.ScratchBufferOffset = Align(SizeInfo.BuildScratchSize, GRHIRayTracingScratchBufferAlignment);
 
-		SizeInfo.ResultSize = Layer.BufferOffset + LayerSizeInfo.ResultSize;
-		SizeInfo.BuildScratchSize = Layer.ScratchBufferOffset + LayerSizeInfo.BuildScratchSize;
+		SizeInfo.ResultSize = Layer.BufferOffset + Layer.SizeInfo.ResultSize;
+		SizeInfo.BuildScratchSize = Layer.ScratchBufferOffset + Layer.SizeInfo.BuildScratchSize;
 	}
 
 	const uint32 ParameterBufferSize = FMath::Max<uint32>(1, Initializer.NumTotalSegments) * sizeof(FVulkanRayTracingGeometryParameters);
@@ -723,11 +723,15 @@ void FVulkanRayTracingScene::BindBuffer(FRHIBuffer* InBuffer, uint32 InBufferOff
 		check(LayerOffset % GRHIRayTracingAccelerationStructureAlignment == 0);
 
 		Layer.View = MakeUnique<FVulkanView>(*Device, VK_DESCRIPTOR_TYPE_ACCELERATION_STRUCTURE_KHR);
-		Layer.View->InitAsAccelerationStructureView(
+		VkAccelerationStructureKHR NativeAccelerationStructureHandle = Layer.View->InitAsAccelerationStructureView(
 			AccelerationStructureBuffer
 			, LayerOffset
-			, InBuffer->GetSize() - LayerOffset
-		);
+			, Layer.SizeInfo.ResultSize
+		)->GetAccelerationStructureView().Handle;
+
+		FString DebugNameString = Initializer.DebugName.ToString();
+		DebugNameString = (DebugNameString.IsEmpty()) ? TEXT("TLAS") : DebugNameString;
+		VULKAN_SET_DEBUG_NAME(*Device, VK_OBJECT_TYPE_ACCELERATION_STRUCTURE_KHR, NativeAccelerationStructureHandle, TEXT("%s"), *DebugNameString);
 	}
 }
 
diff --git a/Engine/Source/Runtime/VulkanRHI/Private/VulkanRayTracing.h b/Engine/Source/Runtime/VulkanRHI/Private/VulkanRayTracing.h
--- a/Engine/Source/Runtime/VulkanRHI/Private/VulkanRayTracing.h
+++ b/Engine/Source/Runtime/VulkanRHI/Private/VulkanRayTracing.h
@@ -224,6 +224,7 @@ private:
 	struct FLayerData
 	{
 		TUniquePtr<FVulkanView> View;
+		FRayTracingAccelerationStructureSize SizeInfo;
 		uint32 BufferOffset;
 		uint32 ScratchBufferOffset;
 	};
