Fixing includes in VulkanRHI so that all headers can be included independently.
Removing 'using namespace VulkanRHI;'

diff --git a/Engine/Source/Runtime/VulkanRHI/Private/Linux/VulkanLinuxPlatform.cpp b/Engine/Source/Runtime/VulkanRHI/Private/Linux/VulkanLinuxPlatform.cpp
--- a/Engine/Source/Runtime/VulkanRHI/Private/Linux/VulkanLinuxPlatform.cpp
+++ b/Engine/Source/Runtime/VulkanRHI/Private/Linux/VulkanLinuxPlatform.cpp
@@ -1,9 +1,10 @@
 // Copyright Epic Games, Inc. All Rights Reserved.
 
 #include "VulkanLinuxPlatform.h"
-#include "../VulkanRHIPrivate.h"
-#include "../VulkanRayTracing.h"
-#include "../VulkanExtensions.h"
+#include "VulkanDevice.h"
+#include "VulkanRHIPrivate.h"
+#include "VulkanRayTracing.h"
+#include "VulkanExtensions.h"
 #include <dlfcn.h>
 #include <SDL.h>
 #include <SDL_vulkan.h>
diff --git a/Engine/Source/Runtime/VulkanRHI/Private/VulkanBarriers.cpp b/Engine/Source/Runtime/VulkanRHI/Private/VulkanBarriers.cpp
--- a/Engine/Source/Runtime/VulkanRHI/Private/VulkanBarriers.cpp
+++ b/Engine/Source/Runtime/VulkanRHI/Private/VulkanBarriers.cpp
@@ -895,7 +895,7 @@ static void AddSubresourceTransitions(TArray<VkImageMemoryBarrier>& Barriers, Vk
 			if (bIsDepthStencil)
 			{
 				const VkImageLayout OtherLayout = (CurrentLayout.NumPlanes == 1) ? SrcLayout : CurrentLayout.GetSubresLayout(LayerIdx, MipIdx, 1);
-				SrcLayout = GetMergedDepthStencilLayout(SrcLayout, OtherLayout);
+				SrcLayout = VulkanRHI::GetMergedDepthStencilLayout(SrcLayout, OtherLayout);
 			}
 
 			// Merge with the previous transition if the previous mip was in the same state as this mip.
@@ -916,7 +916,7 @@ static void AddSubresourceTransitions(TArray<VkImageMemoryBarrier>& Barriers, Vk
 				Barrier = TemplateBarrier;
 				Barrier.srcAccessMask = GetVkAccessMaskForLayout(SrcLayout);
 				Barrier.oldLayout = SrcLayout;
-				Barrier.newLayout = bIsDepthStencil ? GetMergedDepthStencilLayout(DstLayout, DstLayout) : DstLayout;
+				Barrier.newLayout = bIsDepthStencil ? VulkanRHI::GetMergedDepthStencilLayout(DstLayout, DstLayout) : DstLayout;
 				Barrier.image = ImageHandle;
 				Barrier.subresourceRange.baseMipLevel = MipIdx;
 				Barrier.subresourceRange.levelCount = 1;
@@ -1254,13 +1254,13 @@ void FTransitionProcessor<VkMemoryBarrier, VkBufferMemoryBarrier, VkImageMemoryB
 		// Merge the layout with its other half and set it in the barrier
 		if (OtherAspectMask == VK_IMAGE_ASPECT_STENCIL_BIT)
 		{
-			ImageBarrier->oldLayout = GetMergedDepthStencilLayout(ImageBarrier->oldLayout, OtherAspectOldLayout);
-			ImageBarrier->newLayout = GetMergedDepthStencilLayout(ImageBarrier->newLayout, OtherAspectNewLayout);
+			ImageBarrier->oldLayout = VulkanRHI::GetMergedDepthStencilLayout(ImageBarrier->oldLayout, OtherAspectOldLayout);
+			ImageBarrier->newLayout = VulkanRHI::GetMergedDepthStencilLayout(ImageBarrier->newLayout, OtherAspectNewLayout);
 		}
 		else
 		{
-			ImageBarrier->oldLayout = GetMergedDepthStencilLayout(OtherAspectOldLayout, ImageBarrier->oldLayout);
-			ImageBarrier->newLayout = GetMergedDepthStencilLayout(OtherAspectNewLayout, ImageBarrier->newLayout);
+			ImageBarrier->oldLayout = VulkanRHI::GetMergedDepthStencilLayout(OtherAspectOldLayout, ImageBarrier->oldLayout);
+			ImageBarrier->newLayout = VulkanRHI::GetMergedDepthStencilLayout(OtherAspectNewLayout, ImageBarrier->newLayout);
 		}
         ImageBarrier->subresourceRange.aspectMask |= OtherAspectMask;
 
@@ -1297,8 +1297,8 @@ void FTransitionProcessor<VkMemoryBarrier, VkBufferMemoryBarrier, VkImageMemoryB
 		{
 			// The only way we end up here is if the barrier transitions every aspect of the depth(-stencil) texture
 			check(Texture->GetFullAspectMask() == ImageBarrier->subresourceRange.aspectMask);
-			ImageBarrier->oldLayout = GetMergedDepthStencilLayout(ImageBarrier->oldLayout, ImageBarrier->oldLayout);
-			ImageBarrier->newLayout = GetMergedDepthStencilLayout(ImageBarrier->newLayout, ImageBarrier->newLayout);
+			ImageBarrier->oldLayout = VulkanRHI::GetMergedDepthStencilLayout(ImageBarrier->oldLayout, ImageBarrier->oldLayout);
+			ImageBarrier->newLayout = VulkanRHI::GetMergedDepthStencilLayout(ImageBarrier->newLayout, ImageBarrier->newLayout);
 		}
 
 		MergedDstStageMask |= GetVkStageFlagsForLayout(ImageBarriers[TargetIndex].newLayout);
@@ -1581,8 +1581,8 @@ void FVulkanPipelineBarrier::AddFullImageLayoutTransition(const FVulkanTexture&
 	const VkImageSubresourceRange SubresourceRange = MakeSubresourceRange(Texture.GetFullAspectMask());
 	if (Texture.IsDepthOrStencilAspect())
 	{
-		SrcLayout = GetMergedDepthStencilLayout(SrcLayout, SrcLayout);
-		DstLayout = GetMergedDepthStencilLayout(DstLayout, DstLayout);
+		SrcLayout = VulkanRHI::GetMergedDepthStencilLayout(SrcLayout, SrcLayout);
+		DstLayout = VulkanRHI::GetMergedDepthStencilLayout(DstLayout, DstLayout);
 	}
 
 	VkImageMemoryBarrier2& ImgBarrier = ImageBarriers.AddDefaulted_GetRef();
@@ -1945,7 +1945,7 @@ VkImageLayout FVulkanLayoutManager::GetDefaultLayout(FVulkanCmdBuffer* CmdBuffer
 				FVulkanLayoutManager& LayoutMgr = CmdBuffer->GetLayoutManager();
 				const VkImageLayout DepthLayout = LayoutMgr.GetDepthStencilHint(VulkanTexture, VK_IMAGE_ASPECT_DEPTH_BIT);
 				const VkImageLayout StencilLayout = LayoutMgr.GetDepthStencilHint(VulkanTexture, VK_IMAGE_ASPECT_STENCIL_BIT);
-				return GetMergedDepthStencilLayout(DepthLayout, StencilLayout);
+				return VulkanRHI::GetMergedDepthStencilLayout(DepthLayout, StencilLayout);
 			}
 		}
 		else
@@ -2005,7 +2005,7 @@ VkImageLayout FVulkanLayoutManager::SetExpectedLayout(FVulkanCmdBuffer* CmdBuffe
 			FVulkanPipelineBarrier Barrier;
 			if (VulkanTexture.IsDepthOrStencilAspect() && (SrcLayout->NumPlanes > 1) && !LayoutMgr.bWriteOnly)
 			{
-				const VkImageLayout MergedLayout = GetMergedDepthStencilLayout(SrcLayout->GetSubresLayout(0, 0, VK_IMAGE_ASPECT_DEPTH_BIT), SrcLayout->GetSubresLayout(0, 0, VK_IMAGE_ASPECT_STENCIL_BIT));
+				const VkImageLayout MergedLayout = VulkanRHI::GetMergedDepthStencilLayout(SrcLayout->GetSubresLayout(0, 0, VK_IMAGE_ASPECT_DEPTH_BIT), SrcLayout->GetSubresLayout(0, 0, VK_IMAGE_ASPECT_STENCIL_BIT));
 				Barrier.AddImageLayoutTransition(VulkanTexture.Image, MergedLayout, ExpectedLayout, FVulkanPipelineBarrier::MakeSubresourceRange(VulkanTexture.GetFullAspectMask()));
 				PreviousLayout = MergedLayout;
 			}
diff --git a/Engine/Source/Runtime/VulkanRHI/Private/VulkanBarriers.h b/Engine/Source/Runtime/VulkanRHI/Private/VulkanBarriers.h
--- a/Engine/Source/Runtime/VulkanRHI/Private/VulkanBarriers.h
+++ b/Engine/Source/Runtime/VulkanRHI/Private/VulkanBarriers.h
@@ -2,8 +2,11 @@
 
 #pragma once
 
-#include "CoreMinimal.h"
-#include "VulkanCommon.h"
+#include "Containers/Array.h"
+#include "RHIPipeline.h"
+#include "VulkanMemory.h"
+#include "VulkanResources.h"
+#include "VulkanThirdParty.h"
 
 class FVulkanCmdBuffer;
 
diff --git a/Engine/Source/Runtime/VulkanRHI/Private/VulkanBuffer.cpp b/Engine/Source/Runtime/VulkanRHI/Private/VulkanBuffer.cpp
--- a/Engine/Source/Runtime/VulkanRHI/Private/VulkanBuffer.cpp
+++ b/Engine/Source/Runtime/VulkanRHI/Private/VulkanBuffer.cpp
@@ -140,7 +140,7 @@ FVulkanResourceMultiBuffer::FVulkanResourceMultiBuffer(FVulkanDevice* InDevice,
 		else
 		{
 			const bool bUnifiedMem = InDevice->HasUnifiedMemory();
-			const uint32 BufferAlignment = FMemoryManager::CalculateBufferAlignment(*InDevice, InBufferDesc.Usage, bZeroSize);
+			const uint32 BufferAlignment = VulkanRHI::FMemoryManager::CalculateBufferAlignment(*InDevice, InBufferDesc.Usage, bZeroSize);
 
 			if (InTransientHeapAllocation != nullptr)
 			{
@@ -236,11 +236,11 @@ void FVulkanResourceMultiBuffer::AdvanceBufferIndex()
 		const bool bUnifiedMem = Device->HasUnifiedMemory();
 		const VkMemoryPropertyFlags BufferMemFlags = VK_MEMORY_PROPERTY_DEVICE_LOCAL_BIT | (bUnifiedMem ? (VK_MEMORY_PROPERTY_HOST_VISIBLE_BIT | VK_MEMORY_PROPERTY_HOST_COHERENT_BIT) : 0);
 		const uint32 BufferSize = GetSize();
-		const uint32 BufferAlignment = FMemoryManager::CalculateBufferAlignment(*Device, GetUsage(), (BufferSize == 0));
+		const uint32 BufferAlignment = VulkanRHI::FMemoryManager::CalculateBufferAlignment(*Device, GetUsage(), (BufferSize == 0));
 
 		CurrentBufferIndex = BufferAllocs.Emplace();
 		FBufferAlloc& NewBufferAlloc = BufferAllocs[CurrentBufferIndex];
-		if (!Device->GetMemoryManager().AllocateBufferPooled(NewBufferAlloc.Alloc, nullptr, BufferSize, BufferAlignment, BufferUsageFlags, BufferMemFlags, EVulkanAllocationMetaMultiBuffer, __FILE__, __LINE__))
+		if (!Device->GetMemoryManager().AllocateBufferPooled(NewBufferAlloc.Alloc, nullptr, BufferSize, BufferAlignment, BufferUsageFlags, BufferMemFlags, VulkanRHI::EVulkanAllocationMetaMultiBuffer, __FILE__, __LINE__))
 		{
 			Device->GetMemoryManager().HandleOOM();
 		}
@@ -308,7 +308,7 @@ void* FVulkanResourceMultiBuffer::Lock(FVulkanCommandListContext& Context, EReso
 		{
 			FBufferAlloc& BufferAlloc = BufferAllocs[0];
 
-			FTempFrameAllocationBuffer::FTempAllocInfo VolatileAlloc;
+			VulkanRHI::FTempFrameAllocationBuffer::FTempAllocInfo VolatileAlloc;
 			Context.GetTempFrameAllocationBuffer().Alloc(LockSize + Offset, 256, VolatileAlloc);
 			check(!VolatileAlloc.Allocation.HasAllocation());
 
diff --git a/Engine/Source/Runtime/VulkanRHI/Private/VulkanCommandBuffer.cpp b/Engine/Source/Runtime/VulkanRHI/Private/VulkanCommandBuffer.cpp
--- a/Engine/Source/Runtime/VulkanRHI/Private/VulkanCommandBuffer.cpp
+++ b/Engine/Source/Runtime/VulkanRHI/Private/VulkanCommandBuffer.cpp
@@ -4,9 +4,10 @@
 	VulkanCommandBuffer.cpp: Vulkan device RHI implementation.
 =============================================================================*/
 
-#include "VulkanRHIPrivate.h"
+#include "VulkanCommandBuffer.h"
 #include "VulkanContext.h"
 #include "VulkanDescriptorSets.h"
+#include "VulkanMemory.h"
 
 static int32 GUseSingleQueue = 0;
 static FAutoConsoleVariableRef CVarVulkanUseSingleQueue(
@@ -564,7 +565,7 @@ void FVulkanCommandBufferManager::SubmitUploadCmdBuffer(uint32 NumSignalSemaphor
 		{
 			// Add semaphores associated with the recent active cmdbuf(s), if any. That will prevent
 			// the overlap, delaying execution of this cmdbuf until the graphics one(s) is complete.
-			for (FSemaphore* WaitForThis : RenderingCompletedSemaphores)
+			for (VulkanRHI::FSemaphore* WaitForThis : RenderingCompletedSemaphores)
 			{
 				UploadCmdBuffer->AddWaitSemaphore(VK_PIPELINE_STAGE_TOP_OF_PIPE_BIT, WaitForThis);
 			}
@@ -630,7 +631,7 @@ void FVulkanCommandBufferManager::SubmitActiveCmdBuffer(TArrayView<VulkanRHI::FS
 
 			// Add semaphores associated with the recent upload cmdbuf(s), if any. That will prevent
 			// the overlap, delaying execution of this cmdbuf until upload one(s) are complete.
-			for (FSemaphore* UploadCompleteSema : UploadCompletedSemaphores)
+			for (VulkanRHI::FSemaphore* UploadCompleteSema : UploadCompletedSemaphores)
 			{
 				ActiveCmdBuffer->AddWaitSemaphore(VK_PIPELINE_STAGE_TOP_OF_PIPE_BIT, UploadCompleteSema);
 			}
@@ -856,7 +857,7 @@ void FVulkanCommandBufferPool::FreeUnusedCmdBuffers(FVulkanQueue* InQueue)
 	InQueue->GetLastSubmittedInfo(LastSubmittedCmdBuffer, LastSubmittedFenceCounter);
 
 	// Deferred deletion queue caches pointers to cmdbuffers
-	FDeferredDeletionQueue2& DeferredDeletionQueue = Device->GetDeferredDeletionQueue();
+	VulkanRHI::FDeferredDeletionQueue2& DeferredDeletionQueue = Device->GetDeferredDeletionQueue();
 
 	for (int32 Index = CmdBuffers.Num() - 1; Index >= 0; --Index)
 	{
diff --git a/Engine/Source/Runtime/VulkanRHI/Private/VulkanCommandWrappers.h b/Engine/Source/Runtime/VulkanRHI/Private/VulkanCommandWrappers.h
--- a/Engine/Source/Runtime/VulkanRHI/Private/VulkanCommandWrappers.h
+++ b/Engine/Source/Runtime/VulkanRHI/Private/VulkanCommandWrappers.h
@@ -4,7 +4,34 @@
 	VulkanCommandWrappers.h: Wrap all Vulkan API functions so we can add our own 'layers'
 =============================================================================*/
 
-#pragma once 
+#pragma once
+
+#include "VulkanConfiguration.h"
+#include "VulkanLoader.h"
+#include "VulkanThirdParty.h"
+
+#if !VULKAN_COMMANDWRAPPERS_ENABLE
+	#if VULKAN_DYNAMICALLYLOADED
+		// Bring functions from VulkanDynamicAPI to VulkanRHI
+		#define VK_DYNAMICAPI_TO_VULKANRHI(Type,Func) using VulkanDynamicAPI::Func;
+
+		namespace VulkanRHI
+		{
+			ENUM_VK_ENTRYPOINTS_ALL(VK_DYNAMICAPI_TO_VULKANRHI);
+		}
+	#else // VULKAN_DYNAMICALLYLOADED
+		#error "Statically linked vulkan api must be wrapped!"
+	#endif // VULKAN_DYNAMICALLYLOADED
+
+#else // VULKAN_COMMANDWRAPPERS_ENABLE
+
+#if VULKAN_DYNAMICALLYLOADED
+	// Vulkan API is defined in VulkanDynamicAPI namespace.
+	#define VULKANAPINAMESPACE VulkanDynamicAPI
+#else
+	// Vulkan API is in the global namespace.
+	#define VULKANAPINAMESPACE
+#endif
 
 #if VULKAN_ENABLE_WRAP_LAYER
 	#define  VULKAN_LAYER_BODY		;
@@ -1693,3 +1720,5 @@ namespace VulkanRHI
 	void PrintfBegin(const FString& String);
 #endif
 }
+
+#endif // VULKAN_COMMANDWRAPPERS_ENABLE
diff --git a/Engine/Source/Runtime/VulkanRHI/Private/VulkanCommandsDirect.h b/Engine/Source/Runtime/VulkanRHI/Private/VulkanCommandsDirect.h
--- a/Engine/Source/Runtime/VulkanRHI/Private/VulkanCommandsDirect.h
+++ b/Engine/Source/Runtime/VulkanRHI/Private/VulkanCommandsDirect.h
@@ -4,20 +4,6 @@
 	VulkanCommandsDirect.h: Enumerate the VulkanDynamicAPI.  Doing this in a header to avoid static analysis warnings.
 =============================================================================*/
 
-#pragma once 
+#pragma once
 
-#include "VulkanConfiguration.h"
-
-#if VULKAN_COMMANDWRAPPERS_ENABLE
-	#error "VulkanCommandsDirect should not be used when CommandWrappers are enabled."
-#endif
-#if !VULKAN_DYNAMICALLYLOADED
-	#error "VulkanCommandsDirect is for the Dynamic api."
-#endif
-
-// Bring functions from VulkanDynamicAPI to VulkanRHI
-#define VK_DYNAMICAPI_TO_VULKANRHI(Type,Func) using VulkanDynamicAPI::Func;
-namespace VulkanRHI
-{
-	ENUM_VK_ENTRYPOINTS_ALL(VK_DYNAMICAPI_TO_VULKANRHI);
-}
+#include "VulkanCommandWrappers.h"
diff --git a/Engine/Source/Runtime/VulkanRHI/Private/VulkanContext.h b/Engine/Source/Runtime/VulkanRHI/Private/VulkanContext.h
--- a/Engine/Source/Runtime/VulkanRHI/Private/VulkanContext.h
+++ b/Engine/Source/Runtime/VulkanRHI/Private/VulkanContext.h
@@ -7,14 +7,16 @@
 #pragma once 
 
 #include "VulkanResources.h"
+#include "VulkanRHIPrivate.h"
 #include "VulkanGPUProfiler.h"
 
-class FVulkanDevice;
 class FVulkanCommandBufferManager;
+class FVulkanDevice;
+class FVulkanDynamicRHI;
+class FVulkanOcclusionQueryPool;
 class FVulkanPendingGfxState;
 class FVulkanPendingComputeState;
 class FVulkanQueue;
-class FVulkanOcclusionQueryPool;
 class FVulkanSwapChain;
 
 struct FInputAttachmentData;
@@ -333,8 +335,3 @@ private:
 	friend class FVulkanDevice;
 };
 #endif
-
-inline FVulkanCommandListContextImmediate& FVulkanDevice::GetImmediateContext()
-{
-	return *ImmediateContext;
-}
diff --git a/Engine/Source/Runtime/VulkanRHI/Private/VulkanDescriptorSets.cpp b/Engine/Source/Runtime/VulkanRHI/Private/VulkanDescriptorSets.cpp
--- a/Engine/Source/Runtime/VulkanRHI/Private/VulkanDescriptorSets.cpp
+++ b/Engine/Source/Runtime/VulkanRHI/Private/VulkanDescriptorSets.cpp
@@ -4,8 +4,8 @@
 	VulkanDescriptorSets.cpp: Vulkan descriptor set RHI implementation.
 =============================================================================*/
 
-#include "VulkanRHIPrivate.h"
 #include "VulkanDescriptorSets.h"
+#include "VulkanRHIPrivate.h"
 #include "VulkanContext.h"
 
 
@@ -747,7 +747,7 @@ void FVulkanBindlessDescriptorManager::UpdateDescriptor(FRHIDescriptorHandle Des
 		}
 		else
 		{
-			FStagingBuffer* StagingBuffer = Device->GetStagingManager().AcquireBuffer(State.DescriptorSize);
+			VulkanRHI::FStagingBuffer* StagingBuffer = Device->GetStagingManager().AcquireBuffer(State.DescriptorSize);
 			FMemory::Memcpy(StagingBuffer->GetMappedPointer(), &State.DebugDescriptors[ByteOffset], State.DescriptorSize);
 			{
 				VkMemoryBarrier2 MemoryBarrier;
diff --git a/Engine/Source/Runtime/VulkanRHI/Private/VulkanDescriptorSets.h b/Engine/Source/Runtime/VulkanRHI/Private/VulkanDescriptorSets.h
--- a/Engine/Source/Runtime/VulkanRHI/Private/VulkanDescriptorSets.h
+++ b/Engine/Source/Runtime/VulkanRHI/Private/VulkanDescriptorSets.h
@@ -8,7 +8,9 @@
 
 #include "Async/AsyncWork.h"
 #include "VulkanConfiguration.h"
+#include "VulkanDevice.h"
 #include "VulkanMemory.h"
+#include "VulkanRHIPrivate.h"
 #include "VulkanShaderResources.h"
 
 #if (UE_BUILD_SHIPPING || UE_BUILD_TEST)
diff --git a/Engine/Source/Runtime/VulkanRHI/Private/VulkanDevice.h b/Engine/Source/Runtime/VulkanRHI/Private/VulkanDevice.h
--- a/Engine/Source/Runtime/VulkanRHI/Private/VulkanDevice.h
+++ b/Engine/Source/Runtime/VulkanRHI/Private/VulkanDevice.h
@@ -7,7 +7,11 @@
 #pragma once
 
 #include "VulkanMemory.h"
+#include "VulkanResources.h"
+#include "GPUProfiler.h"
 
+class FVulkanSamplerState;
+class FVulkanDynamicRHI;
 class FVulkanDescriptorSetCache;
 class FVulkanDescriptorPool;
 class FVulkanDescriptorPoolsManager;
@@ -159,7 +163,7 @@ private:
 
 namespace VulkanRHI
 {
-	class FDeferredDeletionQueue2 : public FDeviceChild
+	class FDeferredDeletionQueue2 : public VulkanRHI::FDeviceChild
 	{
 
 	public:
@@ -440,7 +444,10 @@ public:
 		return ShaderFactory;
 	}
 
-	FVulkanCommandListContextImmediate& GetImmediateContext();
+	FVulkanCommandListContextImmediate& GetImmediateContext()
+	{
+		return *ImmediateContext;
+	}
 
 	void NotifyDeletedImage(VkImage Image, bool bRenderTarget);
 
diff --git a/Engine/Source/Runtime/VulkanRHI/Private/VulkanExtensions.cpp b/Engine/Source/Runtime/VulkanRHI/Private/VulkanExtensions.cpp
--- a/Engine/Source/Runtime/VulkanRHI/Private/VulkanExtensions.cpp
+++ b/Engine/Source/Runtime/VulkanRHI/Private/VulkanExtensions.cpp
@@ -1,7 +1,8 @@
 // Copyright Epic Games, Inc. All Rights Reserved.
 
-#include "VulkanRHIPrivate.h"
 #include "VulkanExtensions.h"
+#include "VulkanDevice.h"
+#include "VulkanRHIPrivate.h"
 
 #include "IHeadMountedDisplayModule.h"
 #include "IHeadMountedDisplayVulkanExtensions.h"
diff --git a/Engine/Source/Runtime/VulkanRHI/Private/VulkanExtensions.h b/Engine/Source/Runtime/VulkanRHI/Private/VulkanExtensions.h
--- a/Engine/Source/Runtime/VulkanRHI/Private/VulkanExtensions.h
+++ b/Engine/Source/Runtime/VulkanRHI/Private/VulkanExtensions.h
@@ -6,9 +6,11 @@
 
 #pragma once
 
-#include "CoreMinimal.h"
 #include "VulkanConfiguration.h"
 
+struct FOptionalVulkanInstanceExtensions;
+
+
 #define VULKAN_EXTENSION_NOT_PROMOTED UINT32_MAX
 #define VULKAN_EXTENSION_ENABLED      1
 #define VULKAN_EXTENSION_DISABLED     0
diff --git a/Engine/Source/Runtime/VulkanRHI/Private/VulkanGPUProfiler.h b/Engine/Source/Runtime/VulkanRHI/Private/VulkanGPUProfiler.h
--- a/Engine/Source/Runtime/VulkanRHI/Private/VulkanGPUProfiler.h
+++ b/Engine/Source/Runtime/VulkanRHI/Private/VulkanGPUProfiler.h
@@ -6,10 +6,13 @@
 
 #pragma once
 
+#include "Containers/Queue.h"
 #include "GPUProfiler.h"
+#include "VulkanConfiguration.h"
 
 class FVulkanCmdBuffer;
 class FVulkanCommandListContext;
+class FVulkanDevice;
 class FVulkanTimingQueryPool;
 
 class FVulkanGPUTiming : public FGPUTiming
diff --git a/Engine/Source/Runtime/VulkanRHI/Private/VulkanGenericPlatform.cpp b/Engine/Source/Runtime/VulkanRHI/Private/VulkanGenericPlatform.cpp
--- a/Engine/Source/Runtime/VulkanRHI/Private/VulkanGenericPlatform.cpp
+++ b/Engine/Source/Runtime/VulkanRHI/Private/VulkanGenericPlatform.cpp
@@ -1,7 +1,8 @@
 // Copyright Epic Games, Inc. All Rights Reserved.
 
-#include "VulkanRHIPrivate.h"
 #include "VulkanGenericPlatform.h"
+#include "VulkanCommandWrappers.h"
+#include "VulkanDevice.h"
 #include "HAL/FileManager.h"
 #include "Misc/CommandLine.h"
 
diff --git a/Engine/Source/Runtime/VulkanRHI/Private/VulkanGenericPlatform.h b/Engine/Source/Runtime/VulkanRHI/Private/VulkanGenericPlatform.h
--- a/Engine/Source/Runtime/VulkanRHI/Private/VulkanGenericPlatform.h
+++ b/Engine/Source/Runtime/VulkanRHI/Private/VulkanGenericPlatform.h
@@ -1,12 +1,12 @@
 // Copyright Epic Games, Inc. All Rights Reserved.
 
 #pragma once
-#include "CoreMinimal.h"
-#include "PixelFormat.h"
-#include "Containers/ArrayView.h"
-#include "RHI.h"	// for GShaderPlatformForFeatureLevel and its friends
 
-#include "VulkanLoader.h"
+#include "Containers/Array.h"
+#include "Containers/ArrayView.h"
+#include "PixelFormat.h"
+#include "Templates/UniquePtr.h"
+#include "VulkanThirdParty.h"
 
 struct FOptionalVulkanDeviceExtensions;
 class FVulkanDevice;
@@ -14,6 +14,10 @@ class FVulkanRenderTargetLayout;
 struct FGfxPipelineDesc;
 class FVulkanPhysicalDeviceFeatures;
 
+enum EShaderPlatform : uint16;
+enum class EGpuVendorId : uint32;
+namespace ERHIFeatureLevel { enum Type : int; }
+
 using FVulkanDeviceExtensionArray = TArray<TUniquePtr<class FVulkanDeviceExtension>>;
 using FVulkanInstanceExtensionArray = TArray<TUniquePtr<class FVulkanInstanceExtension>>;
 
diff --git a/Engine/Source/Runtime/VulkanRHI/Private/VulkanLLM.h b/Engine/Source/Runtime/VulkanRHI/Private/VulkanLLM.h
--- a/Engine/Source/Runtime/VulkanRHI/Private/VulkanLLM.h
+++ b/Engine/Source/Runtime/VulkanRHI/Private/VulkanLLM.h
@@ -4,10 +4,10 @@
 	VulkanLLM.h: Vulkan LLM definitions.
 =============================================================================*/
 
-
 #pragma once
 
 #include "HAL/LowLevelMemTracker.h"
+#include "VulkanConfiguration.h"
 
 #if ENABLE_LOW_LEVEL_MEM_TRACKER
 
diff --git a/Engine/Source/Runtime/VulkanRHI/Private/VulkanLayers.cpp b/Engine/Source/Runtime/VulkanRHI/Private/VulkanLayers.cpp
--- a/Engine/Source/Runtime/VulkanRHI/Private/VulkanLayers.cpp
+++ b/Engine/Source/Runtime/VulkanRHI/Private/VulkanLayers.cpp
@@ -5,6 +5,7 @@
 =============================================================================*/
 
 #include "VulkanRHIPrivate.h"
+#include "VulkanDevice.h"
 #include "VulkanExtensions.h"
 #include "IHeadMountedDisplayModule.h"
 #include "IHeadMountedDisplayVulkanExtensions.h"
diff --git a/Engine/Source/Runtime/VulkanRHI/Private/VulkanLoader.h b/Engine/Source/Runtime/VulkanRHI/Private/VulkanLoader.h
--- a/Engine/Source/Runtime/VulkanRHI/Private/VulkanLoader.h
+++ b/Engine/Source/Runtime/VulkanRHI/Private/VulkanLoader.h
@@ -2,6 +2,7 @@
 
 #pragma once
 
+#include "VulkanPlatform.h"
 #include "VulkanThirdParty.h"
 
 // List all instance Vulkan entry points used by Unreal that need to be loaded manually
diff --git a/Engine/Source/Runtime/VulkanRHI/Private/VulkanMemory.cpp b/Engine/Source/Runtime/VulkanRHI/Private/VulkanMemory.cpp
--- a/Engine/Source/Runtime/VulkanRHI/Private/VulkanMemory.cpp
+++ b/Engine/Source/Runtime/VulkanRHI/Private/VulkanMemory.cpp
@@ -5130,7 +5130,7 @@ void FVulkanCustomMemManager::InternalFreeNotification(void* UserData, size_t Si
 #endif
 
 
-VkResult FDeviceMemoryManager::GetMemoryTypeFromProperties(uint32 TypeBits, VkMemoryPropertyFlags Properties, uint32* OutTypeIndex)
+VkResult VulkanRHI::FDeviceMemoryManager::GetMemoryTypeFromProperties(uint32 TypeBits, VkMemoryPropertyFlags Properties, uint32* OutTypeIndex)
 {
 	//#todo-rco: Might need to revisit based on https://gitlab.khronos.org/vulkan/vulkan/merge_requests/1165
 	// Search memtypes to find first index with those properties
@@ -5152,7 +5152,7 @@ VkResult FDeviceMemoryManager::GetMemoryTypeFromProperties(uint32 TypeBits, VkMe
 	return VK_ERROR_FEATURE_NOT_PRESENT;
 }
 
-VkResult FDeviceMemoryManager::GetMemoryTypeFromPropertiesExcluding(uint32 TypeBits, VkMemoryPropertyFlags Properties, uint32 ExcludeTypeIndex, uint32* OutTypeIndex)
+VkResult VulkanRHI::FDeviceMemoryManager::GetMemoryTypeFromPropertiesExcluding(uint32 TypeBits, VkMemoryPropertyFlags Properties, uint32 ExcludeTypeIndex, uint32* OutTypeIndex)
 {
 	// Search memtypes to find first index with those properties
 	for (uint32 i = 0; i < MemoryProperties.memoryTypeCount && TypeBits; i++)
@@ -5173,7 +5173,7 @@ VkResult FDeviceMemoryManager::GetMemoryTypeFromPropertiesExcluding(uint32 TypeB
 	return VK_ERROR_FEATURE_NOT_PRESENT;
 }
 
-const VkPhysicalDeviceMemoryProperties& FDeviceMemoryManager::GetMemoryProperties() const
+const VkPhysicalDeviceMemoryProperties& VulkanRHI::FDeviceMemoryManager::GetMemoryProperties() const
 {
 	return MemoryProperties;
 }
diff --git a/Engine/Source/Runtime/VulkanRHI/Private/VulkanMemory.h b/Engine/Source/Runtime/VulkanRHI/Private/VulkanMemory.h
--- a/Engine/Source/Runtime/VulkanRHI/Private/VulkanMemory.h
+++ b/Engine/Source/Runtime/VulkanRHI/Private/VulkanMemory.h
@@ -6,7 +6,10 @@
 
 #pragma once 
 
+#include "Containers/Map.h"
 #include "Misc/ScopeRWLock.h"
+#include "VulkanConfiguration.h"
+#include "VulkanThirdParty.h"
 
 //enable to store FILE/LINE, and optionally a stacktrace via r.vulkan.backtrace
 #if !UE_BUILD_SHIPPING
diff --git a/Engine/Source/Runtime/VulkanRHI/Private/VulkanPendingState.cpp b/Engine/Source/Runtime/VulkanRHI/Private/VulkanPendingState.cpp
--- a/Engine/Source/Runtime/VulkanRHI/Private/VulkanPendingState.cpp
+++ b/Engine/Source/Runtime/VulkanRHI/Private/VulkanPendingState.cpp
@@ -5,6 +5,7 @@
 =============================================================================*/
 
 #include "VulkanRHIPrivate.h"
+#include "VulkanCommandWrappers.h"
 #include "VulkanPendingState.h"
 #include "VulkanPipeline.h"
 #include "VulkanContext.h"
@@ -296,7 +297,7 @@ void FVulkanDescriptorPoolsManager::GC()
 	for (int32 Index = PoolSets.Num() - 1; Index >= 0; Index--)
 	{
 		auto* PoolSet = PoolSets[Index];
-		if (PoolSet->IsUnused() && GFrameNumberRenderThread - PoolSet->GetLastFrameUsed() > NUM_FRAMES_TO_WAIT_BEFORE_RELEASING_TO_OS)
+		if (PoolSet->IsUnused() && GFrameNumberRenderThread - PoolSet->GetLastFrameUsed() > VulkanRHI::NUM_FRAMES_TO_WAIT_BEFORE_RELEASING_TO_OS)
 		{
 			PoolSets.RemoveAtSwap(Index, 1, EAllowShrinking::Yes);
 
@@ -853,7 +854,7 @@ void FVulkanDescriptorSetCache::FCachedPool::Reset()
 
 bool FVulkanDescriptorSetCache::FCachedPool::CanGC() const
 {
-	constexpr uint32 FramesBeforeGC = NUM_FRAMES_TO_WAIT_BEFORE_RELEASING_TO_OS;
+	constexpr uint32 FramesBeforeGC = VulkanRHI::NUM_FRAMES_TO_WAIT_BEFORE_RELEASING_TO_OS;
 	return ((GFrameNumberRenderThread - RecentFrame) > FramesBeforeGC);
 }
 
diff --git a/Engine/Source/Runtime/VulkanRHI/Private/VulkanQuery.cpp b/Engine/Source/Runtime/VulkanRHI/Private/VulkanQuery.cpp
--- a/Engine/Source/Runtime/VulkanRHI/Private/VulkanQuery.cpp
+++ b/Engine/Source/Runtime/VulkanRHI/Private/VulkanQuery.cpp
@@ -5,6 +5,7 @@
 =============================================================================*/
 
 #include "VulkanRHIPrivate.h"
+#include "VulkanDevice.h"
 #include "VulkanResources.h"
 #include "VulkanContext.h"
 #include "VulkanCommandBuffer.h"
diff --git a/Engine/Source/Runtime/VulkanRHI/Private/VulkanQueue.cpp b/Engine/Source/Runtime/VulkanRHI/Private/VulkanQueue.cpp
--- a/Engine/Source/Runtime/VulkanRHI/Private/VulkanQueue.cpp
+++ b/Engine/Source/Runtime/VulkanRHI/Private/VulkanQueue.cpp
@@ -4,8 +4,8 @@
 	VulkanQueue.cpp: Vulkan Queue implementation.
 =============================================================================*/
 
-#include "VulkanRHIPrivate.h"
 #include "VulkanQueue.h"
+#include "VulkanDevice.h"
 #include "VulkanMemory.h"
 #include "VulkanContext.h"
 
diff --git a/Engine/Source/Runtime/VulkanRHI/Private/VulkanRHI.cpp b/Engine/Source/Runtime/VulkanRHI/Private/VulkanRHI.cpp
--- a/Engine/Source/Runtime/VulkanRHI/Private/VulkanRHI.cpp
+++ b/Engine/Source/Runtime/VulkanRHI/Private/VulkanRHI.cpp
@@ -527,6 +527,8 @@ FVulkanDynamicRHI::FVulkanDynamicRHI()
 	SelectDevice();
 }
 
+FVulkanDynamicRHI::~FVulkanDynamicRHI() = default;
+
 void FVulkanDynamicRHI::Init()
 {
 	InitInstance();
@@ -1864,13 +1866,13 @@ FVulkanRenderPass::~FVulkanRenderPass()
 {
 	DEC_DWORD_STAT(STAT_VulkanNumRenderPasses);
 
-	Device.GetDeferredDeletionQueue().EnqueueResource(FDeferredDeletionQueue2::EType::RenderPass, RenderPass);
+	Device.GetDeferredDeletionQueue().EnqueueResource(VulkanRHI::FDeferredDeletionQueue2::EType::RenderPass, RenderPass);
 	RenderPass = VK_NULL_HANDLE;
 }
 
 void FVulkanDynamicRHI::SavePipelineCache()
 {
-	FString CacheFile = GetPipelineCacheFilename();
+	FString CacheFile = VulkanRHI::GetPipelineCacheFilename();
 
 	GVulkanRHI->Device->PipelineStateCache->Save(CacheFile);
 }
@@ -1900,7 +1902,7 @@ void FVulkanDynamicRHI::SaveValidationCache()
 				Result = vkGetValidationCacheData(Device, ValidationCache, &CacheSize, Data.GetData());
 				if (Result == VK_SUCCESS)
 				{
-					FString CacheFilename = GetValidationCacheFilename();
+					FString CacheFilename = VulkanRHI::GetValidationCacheFilename();
 					if (FFileHelper::SaveArrayToFile(Data, *CacheFilename))
 					{
 						UE_LOG(LogVulkanRHI, Display, TEXT("Saved validation cache file '%s', %d bytes"), *CacheFilename, Data.Num());
diff --git a/Engine/Source/Runtime/VulkanRHI/Private/VulkanRHIPrivate.h b/Engine/Source/Runtime/VulkanRHI/Private/VulkanRHIPrivate.h
--- a/Engine/Source/Runtime/VulkanRHI/Private/VulkanRHIPrivate.h
+++ b/Engine/Source/Runtime/VulkanRHI/Private/VulkanRHIPrivate.h
@@ -19,22 +19,7 @@
 // the configuration will set up anything not set up by the platform
 #include "VulkanConfiguration.h"
 
-#if VULKAN_COMMANDWRAPPERS_ENABLE
-	#if VULKAN_DYNAMICALLYLOADED
-		// Vulkan API is defined in VulkanDynamicAPI namespace.
-		#define VULKANAPINAMESPACE VulkanDynamicAPI
-	#else
-		// Vulkan API is in the global namespace.
-		#define VULKANAPINAMESPACE
-	#endif
-	#include "VulkanCommandWrappers.h"
-#else
-	#if VULKAN_DYNAMICALLYLOADED
-		#include "VulkanCommandsDirect.h"
-	#else
-		#error "Statically linked vulkan api must be wrapped!"
-	#endif
-#endif
+#include "VulkanCommandWrappers.h"
 
 #include "VulkanState.h"
 #include "VulkanResources.h"
@@ -43,27 +28,11 @@
 #include "VulkanDynamicRHI.h"
 #include "RHI.h"
 
-#if VK_HEADER_VERSION >= 141
-//workaround for removed defines in sdk 141
-#define VK_DESCRIPTOR_TYPE_BEGIN_RANGE (VK_DESCRIPTOR_TYPE_SAMPLER)
-#define VK_DESCRIPTOR_TYPE_END_RANGE (VK_DESCRIPTOR_TYPE_INPUT_ATTACHMENT)
-#define VK_DESCRIPTOR_TYPE_RANGE_SIZE (VK_DESCRIPTOR_TYPE_INPUT_ATTACHMENT - VK_DESCRIPTOR_TYPE_SAMPLER + 1)
-#define VK_IMAGE_VIEW_TYPE_RANGE_SIZE (VK_IMAGE_VIEW_TYPE_CUBE_ARRAY - VK_IMAGE_VIEW_TYPE_1D + 1)
-#define VK_DYNAMIC_STATE_BEGIN_RANGE (VK_DYNAMIC_STATE_VIEWPORT)
-#define VK_DYNAMIC_STATE_END_RANGE (VK_DYNAMIC_STATE_STENCIL_REFERENCE)
-#define VK_DYNAMIC_STATE_RANGE_SIZE (VK_DYNAMIC_STATE_STENCIL_REFERENCE - VK_DYNAMIC_STATE_VIEWPORT + 1)
-#define VK_FORMAT_RANGE_SIZE (VK_FORMAT_ASTC_12x12_SRGB_BLOCK - VK_FORMAT_UNDEFINED + 1)
-
-#endif
-
 #include "GPUProfiler.h"
-#include "VulkanDevice.h"
 #include "VulkanQueue.h"
 #include "VulkanCommandBuffer.h"
 #include "Stats/Stats2.h"
 
-using namespace VulkanRHI;
-
 class FVulkanQueue;
 class FVulkanCmdBuffer;
 class FVulkanShader;
diff --git a/Engine/Source/Runtime/VulkanRHI/Private/VulkanRayTracing.cpp b/Engine/Source/Runtime/VulkanRHI/Private/VulkanRayTracing.cpp
--- a/Engine/Source/Runtime/VulkanRHI/Private/VulkanRayTracing.cpp
+++ b/Engine/Source/Runtime/VulkanRHI/Private/VulkanRayTracing.cpp
@@ -1018,7 +1018,7 @@ void FVulkanRayTracingShaderTable::Commit(FVulkanCommandListContext& Context)
 					BufferCreateInfo.usage = VK_BUFFER_USAGE_TRANSFER_DST_BIT | VK_BUFFER_USAGE_SHADER_BINDING_TABLE_BIT_KHR | VK_BUFFER_USAGE_SHADER_DEVICE_ADDRESS_BIT;
 					VERIFYVULKANRESULT(VulkanRHI::vkCreateBuffer(DeviceHandle, &BufferCreateInfo, VULKAN_CPU_ALLOCATOR, &Alloc.LocalBuffer));
 
-					const EVulkanAllocationFlags AllocFlags = EVulkanAllocationFlags::AutoBind | EVulkanAllocationFlags::Dedicated;
+					const VulkanRHI::EVulkanAllocationFlags AllocFlags = VulkanRHI::EVulkanAllocationFlags::AutoBind | VulkanRHI::EVulkanAllocationFlags::Dedicated;
 					Device->GetMemoryManager().AllocateBufferMemory(Alloc.LocalAllocation, Alloc.LocalBuffer, AllocFlags, TEXT("LocalShaderTableAllocation"), RayTracingPipelineProps.shaderGroupBaseAlignment);
 
 					VkBufferDeviceAddressInfoKHR DeviceAddressInfo;
diff --git a/Engine/Source/Runtime/VulkanRHI/Private/VulkanRenderTarget.cpp b/Engine/Source/Runtime/VulkanRHI/Private/VulkanRenderTarget.cpp
--- a/Engine/Source/Runtime/VulkanRHI/Private/VulkanRenderTarget.cpp
+++ b/Engine/Source/Runtime/VulkanRHI/Private/VulkanRenderTarget.cpp
@@ -162,7 +162,7 @@ void FVulkanDynamicRHI::RHIReadSurfaceData(FRHITexture* TextureRHI, FIntRect Rec
 	Device->PrepareForCPURead();
 
 	// Figure out the size of the buffer required to hold the requested pixels
-	const uint32 PixelByteSize = GetNumBitsPerPixel(Surface.StorageFormat) / 8;
+	const uint32 PixelByteSize = VulkanRHI::GetNumBitsPerPixel(Surface.StorageFormat) / 8;
 	checkf(GPixelFormats[TextureRHI->GetFormat()].Supported && (PixelByteSize > 0), TEXT("Trying to read from unsupported format."));
 	const uint32 BufferSize = NumRequestedPixels * PixelByteSize;
 
diff --git a/Engine/Source/Runtime/VulkanRHI/Private/VulkanRenderpass.h b/Engine/Source/Runtime/VulkanRHI/Private/VulkanRenderpass.h
--- a/Engine/Source/Runtime/VulkanRHI/Private/VulkanRenderpass.h
+++ b/Engine/Source/Runtime/VulkanRHI/Private/VulkanRenderpass.h
@@ -68,7 +68,7 @@ inline void FVulkanAttachmentReference<VkAttachmentReference>::SetDepthStencilAt
 {
 	attachment = AttachmentReferenceIn.attachment;
 	const VkImageLayout StencilLayout = StencilReference ? StencilReference->stencilLayout : VK_IMAGE_LAYOUT_UNDEFINED;
-	layout = GetMergedDepthStencilLayout(AttachmentReferenceIn.layout, StencilLayout);
+	layout = VulkanRHI::GetMergedDepthStencilLayout(AttachmentReferenceIn.layout, StencilLayout);
 }
 
 template <>
@@ -78,7 +78,7 @@ inline void FVulkanAttachmentReference<VkAttachmentReference2>::SetDepthStencilA
 	sType = VK_STRUCTURE_TYPE_ATTACHMENT_REFERENCE_2;
 	pNext = (bSupportsParallelRendering && StencilReference && StencilReference->stencilLayout != VK_IMAGE_LAYOUT_UNDEFINED) ? StencilReference : nullptr;
 	attachment = AttachmentReferenceIn.attachment;
-	layout = bSupportsParallelRendering ? AttachmentReferenceIn.layout : GetMergedDepthStencilLayout(AttachmentReferenceIn.layout, StencilReference->stencilLayout);
+	layout = bSupportsParallelRendering ? AttachmentReferenceIn.layout : VulkanRHI::GetMergedDepthStencilLayout(AttachmentReferenceIn.layout, StencilReference->stencilLayout);
 	aspectMask = AspectMask;
 }
 
@@ -270,11 +270,11 @@ struct FVulkanAttachmentDescription<VkAttachmentDescription>
 		stencilLoadOp = InDesc.stencilLoadOp;
 		stencilStoreOp = InDesc.stencilStoreOp;
 
-		const bool bHasStencilLayout = VulkanFormatHasStencil(InDesc.format) && (InStencilDesc != nullptr);
+		const bool bHasStencilLayout = VulkanRHI::VulkanFormatHasStencil(InDesc.format) && (InStencilDesc != nullptr);
 		const VkImageLayout StencilInitialLayout = bHasStencilLayout ? InStencilDesc->stencilInitialLayout : VK_IMAGE_LAYOUT_UNDEFINED;
-		initialLayout = GetMergedDepthStencilLayout(InDesc.initialLayout, StencilInitialLayout);
+		initialLayout = VulkanRHI::GetMergedDepthStencilLayout(InDesc.initialLayout, StencilInitialLayout);
 		const VkImageLayout StencilFinalLayout = bHasStencilLayout ? InStencilDesc->stencilFinalLayout : VK_IMAGE_LAYOUT_UNDEFINED;
-		finalLayout = GetMergedDepthStencilLayout(InDesc.finalLayout, StencilFinalLayout);
+		finalLayout = VulkanRHI::GetMergedDepthStencilLayout(InDesc.finalLayout, StencilFinalLayout);
 	}
 };
 
@@ -304,7 +304,7 @@ struct FVulkanAttachmentDescription<VkAttachmentDescription2>
 
 	FVulkanAttachmentDescription(const VkAttachmentDescription& InDesc, const VkAttachmentDescriptionStencilLayout* InStencilDesc, bool bSupportsParallelRendering)
 	{
-		const bool bHasStencilLayout = bSupportsParallelRendering && VulkanFormatHasStencil(InDesc.format) && (InStencilDesc != nullptr);
+		const bool bHasStencilLayout = bSupportsParallelRendering && VulkanRHI::VulkanFormatHasStencil(InDesc.format) && (InStencilDesc != nullptr);
 
 		sType = VK_STRUCTURE_TYPE_ATTACHMENT_DESCRIPTION_2;
 		pNext = (bHasStencilLayout && (InStencilDesc->stencilFinalLayout != VK_IMAGE_LAYOUT_UNDEFINED)) ? InStencilDesc : nullptr;
@@ -315,8 +315,8 @@ struct FVulkanAttachmentDescription<VkAttachmentDescription2>
 		storeOp = InDesc.storeOp;
 		stencilLoadOp = InDesc.stencilLoadOp;
 		stencilStoreOp = InDesc.stencilStoreOp;
-		initialLayout = bSupportsParallelRendering ? InDesc.initialLayout : GetMergedDepthStencilLayout(InDesc.initialLayout, InStencilDesc->stencilInitialLayout);
-		finalLayout = bSupportsParallelRendering ? InDesc.finalLayout : GetMergedDepthStencilLayout(InDesc.finalLayout, InStencilDesc->stencilFinalLayout);
+		initialLayout = bSupportsParallelRendering ? InDesc.initialLayout : VulkanRHI::GetMergedDepthStencilLayout(InDesc.initialLayout, InStencilDesc->stencilInitialLayout);
+		finalLayout = bSupportsParallelRendering ? InDesc.finalLayout : VulkanRHI::GetMergedDepthStencilLayout(InDesc.finalLayout, InStencilDesc->stencilFinalLayout);
 	}
 };
 
diff --git a/Engine/Source/Runtime/VulkanRHI/Private/VulkanState.cpp b/Engine/Source/Runtime/VulkanRHI/Private/VulkanState.cpp
--- a/Engine/Source/Runtime/VulkanRHI/Private/VulkanState.cpp
+++ b/Engine/Source/Runtime/VulkanRHI/Private/VulkanState.cpp
@@ -4,7 +4,7 @@
 	VulkanState.cpp: Vulkan state implementation.
 =============================================================================*/
 
-#include "VulkanRHIPrivate.h"
+#include "VulkanState.h"
 #include "VulkanPendingState.h"
 #include "VulkanDescriptorSets.h"
 #include "RHIUtilities.h"
diff --git a/Engine/Source/Runtime/VulkanRHI/Private/VulkanState.h b/Engine/Source/Runtime/VulkanRHI/Private/VulkanState.h
--- a/Engine/Source/Runtime/VulkanRHI/Private/VulkanState.h
+++ b/Engine/Source/Runtime/VulkanRHI/Private/VulkanState.h
@@ -6,8 +6,10 @@
 
 #pragma once
 
-#include "CoreMinimal.h"
+#include "RHI.h"
 #include "RHIResources.h"
+#include "VulkanCommon.h"
+#include "VulkanThirdParty.h"
 
 class FVulkanDevice;
 
diff --git a/Engine/Source/Runtime/VulkanRHI/Private/VulkanSwapChain.cpp b/Engine/Source/Runtime/VulkanRHI/Private/VulkanSwapChain.cpp
--- a/Engine/Source/Runtime/VulkanRHI/Private/VulkanSwapChain.cpp
+++ b/Engine/Source/Runtime/VulkanRHI/Private/VulkanSwapChain.cpp
@@ -1,12 +1,14 @@
 // Copyright Epic Games, Inc. All Rights Reserved.
 
 /*=============================================================================
-	VulkanSwapChain.h: Vulkan viewport RHI definitions.
+	VulkanSwapChain.cpp: Vulkan viewport RHI definitions.
 =============================================================================*/
 
-#include "VulkanRHIPrivate.h"
 #include "VulkanSwapChain.h"
+#include "VulkanCommandWrappers.h"
+#include "VulkanDevice.h"
 #include "VulkanPlatform.h"
+#include "VulkanRHIPrivate.h"
 #include "Engine/RendererSettings.h"
 #include "HAL/PlatformFramePacer.h"
 #include "IHeadMountedDisplayModule.h"
diff --git a/Engine/Source/Runtime/VulkanRHI/Private/VulkanSwapChain.h b/Engine/Source/Runtime/VulkanRHI/Private/VulkanSwapChain.h
--- a/Engine/Source/Runtime/VulkanRHI/Private/VulkanSwapChain.h
+++ b/Engine/Source/Runtime/VulkanRHI/Private/VulkanSwapChain.h
@@ -6,15 +6,23 @@
 
 #pragma once
 
+#include "Containers/Array.h"
+#include "VulkanConfiguration.h"
+#include "VulkanThirdParty.h"
+
+class FVulkanDevice;
+class FVulkanQueue;
 class FVulkanTexture;
+class FVulkanView;
+enum EPixelFormat : uint8;
+
 
 namespace VulkanRHI
 {
 	class FFence;
+	class FSemaphore;
 }
 
-class FVulkanQueue;
-
 struct FVulkanSwapChainRecreateInfo
 {
 	VkSwapchainKHR SwapChain;
diff --git a/Engine/Source/Runtime/VulkanRHI/Private/VulkanTexture.cpp b/Engine/Source/Runtime/VulkanRHI/Private/VulkanTexture.cpp
--- a/Engine/Source/Runtime/VulkanRHI/Private/VulkanTexture.cpp
+++ b/Engine/Source/Runtime/VulkanRHI/Private/VulkanTexture.cpp
@@ -587,7 +587,7 @@ static VkImageLayout GetInitialLayoutFromRHIAccess(ERHIAccess RHIAccess, bool bI
 	return VK_IMAGE_LAYOUT_UNDEFINED;
 }
 
-void FVulkanTexture::InternalMoveSurface(FVulkanDevice& InDevice, FVulkanCommandListContext& Context, FVulkanAllocation& DestAllocation, VkImageLayout OriginalLayout)
+void FVulkanTexture::InternalMoveSurface(FVulkanDevice& InDevice, FVulkanCommandListContext& Context, VulkanRHI::FVulkanAllocation& DestAllocation, VkImageLayout OriginalLayout)
 {
 	FImageCreateInfo ImageCreateInfo;
 	const FRHITextureDesc& Desc = GetDesc();
@@ -1529,7 +1529,7 @@ FVulkanTexture::FVulkanTexture(FRHICommandListBase* RHICmdList, FVulkanDevice& I
 
 		// Set minimum alignment to 16 bytes, as some buffers are used with CPU SIMD instructions
 		const uint32 ForcedMinAlignment = 16u;
-		const EVulkanAllocationFlags AllocFlags = EVulkanAllocationFlags::HostCached | EVulkanAllocationFlags::AutoBind;
+		const VulkanRHI::EVulkanAllocationFlags AllocFlags = VulkanRHI::EVulkanAllocationFlags::HostCached | VulkanRHI::EVulkanAllocationFlags::AutoBind;
 		InDevice.GetMemoryManager().AllocateBufferMemory(Allocation, CpuReadbackBuffer->Buffer, AllocFlags, InCreateDesc.DebugName, ForcedMinAlignment);
 
 		void* Memory = Allocation.GetMappedPointer(Device);
@@ -1597,7 +1597,7 @@ FVulkanTexture::FVulkanTexture(FRHICommandListBase* RHICmdList, FVulkanDevice& I
 	}
 	else
 	{
-		EVulkanAllocationMetaType MetaType = (bRenderTarget || bUAV) ? EVulkanAllocationMetaImageRenderTarget : EVulkanAllocationMetaImageOther;
+		VulkanRHI::EVulkanAllocationMetaType MetaType = (bRenderTarget || bUAV) ? VulkanRHI::EVulkanAllocationMetaImageRenderTarget : VulkanRHI::EVulkanAllocationMetaImageOther;
 #if VULKAN_SUPPORTS_DEDICATED_ALLOCATION
 		extern int32 GVulkanEnableDedicatedImageMemory;
 		// Per https://developer.nvidia.com/what%E2%80%99s-your-vulkan-memory-type
@@ -1938,7 +1938,7 @@ void FVulkanTexture::UpdateLinkedViews()
 	FVulkanViewableResource::UpdateLinkedViews();
 }
 
-void FVulkanTexture::Move(FVulkanDevice& InDevice, FVulkanCommandListContext& Context, FVulkanAllocation& NewAllocation)
+void FVulkanTexture::Move(FVulkanDevice& InDevice, FVulkanCommandListContext& Context, VulkanRHI::FVulkanAllocation& NewAllocation)
 {
 	const uint64 Size = GetMemorySize();
 	static uint64 TotalSize = 0;
@@ -2001,8 +2001,8 @@ void FVulkanTexture::Evict(FVulkanDevice& InDevice, FVulkanCommandListContext& C
 		MemProps = InDevice.GetDeviceMemoryManager().GetEvictedMemoryProperties();
 
 		// Create a new host allocation to move the surface to
-		FVulkanAllocation HostAllocation;
-		const EVulkanAllocationMetaType MetaType = EVulkanAllocationMetaImageOther;
+		VulkanRHI::FVulkanAllocation HostAllocation;
+		const VulkanRHI::EVulkanAllocationMetaType MetaType = VulkanRHI::EVulkanAllocationMetaImageOther;
 		if (!InDevice.GetMemoryManager().AllocateImageMemory(HostAllocation, this, MemoryRequirements, MemProps, MetaType, false, __FILE__, __LINE__))
 		{
 			InDevice.GetMemoryManager().HandleOOM();
diff --git a/Engine/Source/Runtime/VulkanRHI/Private/VulkanTransientResourceAllocator.cpp b/Engine/Source/Runtime/VulkanRHI/Private/VulkanTransientResourceAllocator.cpp
--- a/Engine/Source/Runtime/VulkanRHI/Private/VulkanTransientResourceAllocator.cpp
+++ b/Engine/Source/Runtime/VulkanRHI/Private/VulkanTransientResourceAllocator.cpp
@@ -1,8 +1,9 @@
 // Copyright Epic Games, Inc. All Rights Reserved.
 
-#include "VulkanRHIPrivate.h"
 #include "VulkanTransientResourceAllocator.h"
-
+#include "VulkanCommandWrappers.h"
+#include "VulkanDevice.h"
+#include "VulkanMemory.h"
 
 FVulkanTransientHeap::FVulkanTransientHeap(const FInitializer& Initializer, FVulkanDevice* InDevice)
 	: FRHITransientHeap(Initializer)
@@ -30,9 +31,9 @@ FVulkanTransientHeap::FVulkanTransientHeap(const FInitializer& Initializer, FVul
 	VERIFYVULKANRESULT(VulkanRHI::vkCreateBuffer(InDevice->GetInstanceHandle(), &BufferCreateInfo, VULKAN_CPU_ALLOCATOR, &VulkanBuffer));
 
 	// Find the alignment that works for everyone
-	const uint32 MinBufferAlignment = FMath::Max<uint32>(Initializer.Alignment, FMemoryManager::CalculateBufferAlignment(*InDevice, UEBufferUsageFlags, bZeroSize));
+	const uint32 MinBufferAlignment = FMath::Max<uint32>(Initializer.Alignment, VulkanRHI::FMemoryManager::CalculateBufferAlignment(*InDevice, UEBufferUsageFlags, bZeroSize));
 
-	const EVulkanAllocationFlags AllocFlags = EVulkanAllocationFlags::Dedicated | EVulkanAllocationFlags::AutoBind;
+	const VulkanRHI::EVulkanAllocationFlags AllocFlags = VulkanRHI::EVulkanAllocationFlags::Dedicated | VulkanRHI::EVulkanAllocationFlags::AutoBind;
 	InDevice->GetMemoryManager().AllocateBufferMemory(InternalAllocation, VulkanBuffer, AllocFlags, TEXT("VulkanTransientHeap"), MinBufferAlignment);
 }
 
@@ -48,12 +49,12 @@ VkDeviceMemory FVulkanTransientHeap::GetMemoryHandle()
 	return InternalAllocation.GetDeviceMemoryHandle(Device);
 }
 
-FVulkanAllocation FVulkanTransientHeap::GetVulkanAllocation(const FRHITransientHeapAllocation& HeapAllocation)
+VulkanRHI::FVulkanAllocation FVulkanTransientHeap::GetVulkanAllocation(const FRHITransientHeapAllocation& HeapAllocation)
 {
 	FVulkanTransientHeap* Heap = static_cast<FVulkanTransientHeap*>(HeapAllocation.Heap);
 	check(Heap);
 
-	FVulkanAllocation TransientAlloc;
+	VulkanRHI::FVulkanAllocation TransientAlloc;
 	TransientAlloc.Reference(Heap->InternalAllocation);
 	TransientAlloc.VulkanHandle = (uint64)Heap->VulkanBuffer;
 	TransientAlloc.Offset += HeapAllocation.Offset;
@@ -95,7 +96,7 @@ FVulkanTransientResourceAllocator::FVulkanTransientResourceAllocator(FVulkanTran
 
 FRHITransientTexture* FVulkanTransientResourceAllocator::CreateTexture(const FRHITextureCreateInfo& InCreateInfo, const TCHAR* InDebugName, uint32 InPassIndex)
 {
-	FDynamicRHI::FRHICalcTextureSizeResult MemReq = GVulkanRHI->RHICalcTexturePlatformSize(InCreateInfo, 0);
+	FDynamicRHI::FRHICalcTextureSizeResult MemReq = GDynamicRHI->RHICalcTexturePlatformSize(InCreateInfo, 0);
 
 	return CreateTextureInternal(InCreateInfo, InDebugName, InPassIndex, MemReq.Size, MemReq.Align,
 		[&](const FRHITransientHeap::FResourceInitializer& Initializer)
@@ -112,7 +113,7 @@ FRHITransientBuffer* FVulkanTransientResourceAllocator::CreateBuffer(const FRHIB
 	checkf(!EnumHasAnyFlags(InCreateInfo.Usage, BUF_Volatile), TEXT("The volatile flag is not supported for transient resources."));
 
 	const bool bZeroSize = (InCreateInfo.Size == 0);
-	const uint32 Alignment = FMemoryManager::CalculateBufferAlignment(*Device, InCreateInfo.Usage, bZeroSize);
+	const uint32 Alignment = VulkanRHI::FMemoryManager::CalculateBufferAlignment(*Device, InCreateInfo.Usage, bZeroSize);
 	uint64 Size = Align(InCreateInfo.Size, Alignment);
 
 	return CreateBufferInternal(InCreateInfo, InDebugName, InPassIndex, Size, Alignment,
diff --git a/Engine/Source/Runtime/VulkanRHI/Private/VulkanTransientResourceAllocator.h b/Engine/Source/Runtime/VulkanRHI/Private/VulkanTransientResourceAllocator.h
--- a/Engine/Source/Runtime/VulkanRHI/Private/VulkanTransientResourceAllocator.h
+++ b/Engine/Source/Runtime/VulkanRHI/Private/VulkanTransientResourceAllocator.h
@@ -2,32 +2,32 @@
 
 #pragma once
 
+#include "VulkanMemory.h"
 #include "VulkanResources.h"
 #include "RHICoreTransientResourceAllocator.h"
 
-
 class FVulkanTransientHeap final
 	: public FRHITransientHeap
 	, public FRefCountBase
-	, public FDeviceChild
+	, public VulkanRHI::FDeviceChild
 {
 public:
 	FVulkanTransientHeap(const FInitializer& Initializer, FVulkanDevice* InDevice);
 	~FVulkanTransientHeap();
 
 	VkDeviceMemory GetMemoryHandle();
-	static FVulkanAllocation GetVulkanAllocation(const FRHITransientHeapAllocation& TransientInitializer);
+	static VulkanRHI::FVulkanAllocation GetVulkanAllocation(const FRHITransientHeapAllocation& TransientInitializer);
 
 private:
 	VkBuffer VulkanBuffer;
 	VkMemoryRequirements MemoryRequirements;
-	FVulkanAllocation InternalAllocation;
+	VulkanRHI::FVulkanAllocation InternalAllocation;
 };
 
 
 class FVulkanTransientHeapCache final
 	: public FRHITransientHeapCache
-	, public FDeviceChild
+	, public VulkanRHI::FDeviceChild
 {
 public:
 	static FVulkanTransientHeapCache* Create(FVulkanDevice* InDevice);
@@ -41,7 +41,7 @@ private:
 
 class FVulkanTransientResourceAllocator final
 	: public FRHITransientResourceHeapAllocator
-	, public FDeviceChild
+	, public VulkanRHI::FDeviceChild
 {
 public:
 	FVulkanTransientResourceAllocator(FVulkanTransientHeapCache& InHeapCache);
diff --git a/Engine/Source/Runtime/VulkanRHI/Private/VulkanUAV.cpp b/Engine/Source/Runtime/VulkanRHI/Private/VulkanUAV.cpp
--- a/Engine/Source/Runtime/VulkanRHI/Private/VulkanUAV.cpp
+++ b/Engine/Source/Runtime/VulkanRHI/Private/VulkanUAV.cpp
@@ -41,7 +41,7 @@ void FVulkanView::Invalidate()
 
 	case EType::TypedBuffer:
 		DEC_DWORD_STAT(STAT_VulkanNumBufferViews);
-		Device.GetDeferredDeletionQueue().EnqueueResource(FDeferredDeletionQueue2::EType::BufferView, Storage.Get<FTypedBufferView>().View);
+		Device.GetDeferredDeletionQueue().EnqueueResource(VulkanRHI::FDeferredDeletionQueue2::EType::BufferView, Storage.Get<FTypedBufferView>().View);
 		break;
 
 	case EType::Texture:
@@ -93,7 +93,7 @@ FVulkanView* FVulkanView::InitAsTypedBufferView(FVulkanResourceMultiBuffer* Buff
 
 	//#todo-rco: Revisit this if buffer views become VK_BUFFER_USAGE_STORAGE_BUFFER_BIT instead of VK_BUFFER_USAGE_UNIFORM_TEXEL_BUFFER_BIT
 	const VkPhysicalDeviceLimits& Limits = Device.GetLimits();
-	const uint64 MaxSize = (uint64)Limits.maxTexelBufferElements * GetNumBitsPerPixel(Format) / 8;
+	const uint64 MaxSize = (uint64)Limits.maxTexelBufferElements * VulkanRHI::GetNumBitsPerPixel(Format) / 8;
 	ViewInfo.range = FMath::Min<uint64>(InSize, MaxSize);
 	// TODO: add a check() for exceeding MaxSize, to catch code which blindly makes views without checking the platform limits.
 
diff --git a/Engine/Source/Runtime/VulkanRHI/Private/VulkanUniformBuffer.cpp b/Engine/Source/Runtime/VulkanRHI/Private/VulkanUniformBuffer.cpp
--- a/Engine/Source/Runtime/VulkanRHI/Private/VulkanUniformBuffer.cpp
+++ b/Engine/Source/Runtime/VulkanRHI/Private/VulkanUniformBuffer.cpp
@@ -66,12 +66,12 @@ static void UpdateUniformBufferHelper(FVulkanCommandListContext& Context, FVulka
 	{
 		FVulkanUniformBufferUploader* UniformBufferUploader = Context.GetUniformBufferUploader();
 		const VkDeviceSize UBOffsetAlignment = Device->GetLimits().minUniformBufferOffsetAlignment;
-		const FVulkanAllocation& RingBufferAllocation = UniformBufferUploader->GetCPUBufferAllocation();
+		const VulkanRHI::FVulkanAllocation& RingBufferAllocation = UniformBufferUploader->GetCPUBufferAllocation();
 		uint64 RingBufferOffset = UniformBufferUploader->AllocateMemory(DataSize, UBOffsetAlignment, CmdBuffer);
 
 		VulkanUniformBuffer->Allocation.Init(
-			EVulkanAllocationEmpty, 
-			EVulkanAllocationMetaUnknown, 
+			VulkanRHI::EVulkanAllocationEmpty,
+			VulkanRHI::EVulkanAllocationMetaUnknown,
 			RingBufferAllocation.VulkanHandle, 
 			DataSize,
 			RingBufferOffset,
@@ -308,7 +308,7 @@ inline void FVulkanDynamicRHI::UpdateUniformBuffer(FRHICommandListBase& RHICmdLi
 	const int32 ConstantBufferSize = Layout.ConstantBufferSize;
 	const int32 NumResources = Layout.Resources.Num();
 
-	FVulkanAllocation NewUBAlloc;
+	VulkanRHI::FVulkanAllocation NewUBAlloc;
 	bool bUseUpload = GVulkanAllowUniformUpload && !RHICmdList.IsInsideRenderPass(); //inside renderpasses, a rename is enforced.
 	const bool bUseRingBuffer = UseRingBuffer(UniformBuffer->Usage);
 
@@ -377,7 +377,7 @@ inline void FVulkanDynamicRHI::UpdateUniformBuffer(FRHICommandListBase& RHICmdLi
 			NewUBAlloc.Disown(); //this releases ownership while its put into the lambda
 			RHICmdList.EnqueueLambda([UniformBuffer, NewUBAlloc, CmdListResources, NumResources](FRHICommandListBase& CmdList)
 			{
-				FVulkanAllocation Alloc;
+				VulkanRHI::FVulkanAllocation Alloc;
 				Alloc.Reference(NewUBAlloc);
 				Alloc.Own(); //this takes ownership of the allocation
 				UniformBuffer->UpdateAllocation(Alloc);
@@ -443,7 +443,7 @@ FVulkanRingBuffer::FVulkanRingBuffer(FVulkanDevice* InDevice, uint64 TotalSize,
 	}
 
 	check(TotalSize <= (uint64)MAX_uint32);
-	InDevice->GetMemoryManager().AllocateBufferPooled(Allocation, nullptr, TotalSize, 0, Usage, MemPropertyFlags, EVulkanAllocationMetaRingBuffer, __FILE__, __LINE__);
+	InDevice->GetMemoryManager().AllocateBufferPooled(Allocation, nullptr, TotalSize, 0, Usage, MemPropertyFlags, VulkanRHI::EVulkanAllocationMetaRingBuffer, __FILE__, __LINE__);
 	MinAlignment = Allocation.GetBufferAlignment(Device);
 	// Start by wrapping around to set up the correct fence
 	BufferOffset = TotalSize;
diff --git a/Engine/Source/Runtime/VulkanRHI/Public/VulkanConfiguration.h b/Engine/Source/Runtime/VulkanRHI/Public/VulkanConfiguration.h
--- a/Engine/Source/Runtime/VulkanRHI/Public/VulkanConfiguration.h
+++ b/Engine/Source/Runtime/VulkanRHI/Public/VulkanConfiguration.h
@@ -10,6 +10,7 @@
 
 #include "HAL/LowLevelMemTracker.h"
 #include "VulkanCommon.h"
+#include "VulkanPlatform.h"
 
 struct VkAllocationCallbacks;
 
diff --git a/Engine/Source/Runtime/VulkanRHI/Public/VulkanDynamicRHI.h b/Engine/Source/Runtime/VulkanRHI/Public/VulkanDynamicRHI.h
--- a/Engine/Source/Runtime/VulkanRHI/Public/VulkanDynamicRHI.h
+++ b/Engine/Source/Runtime/VulkanRHI/Public/VulkanDynamicRHI.h
@@ -7,10 +7,12 @@
 #pragma once
 
 #include "IVulkanDynamicRHI.h"
+#include "VulkanConfiguration.h"
 
 class FVulkanFramebuffer;
 class FVulkanDevice;
 class FVulkanQueue;
+class FVulkanUniformBuffer;
 class FVulkanViewport;
 class IHeadMountedDisplayVulkanExtensions;
 struct FRHITransientHeapAllocation;
@@ -50,7 +52,7 @@ public:
 	FVulkanDynamicRHI();
 
 	/** Destructor */
-	~FVulkanDynamicRHI() {}
+	~FVulkanDynamicRHI();
 
 	// IVulkanDynamicRHI interface
 	virtual uint32 RHIGetVulkanVersion() const final override;
diff --git a/Engine/Source/Runtime/VulkanRHI/Public/VulkanThirdParty.h b/Engine/Source/Runtime/VulkanRHI/Public/VulkanThirdParty.h
--- a/Engine/Source/Runtime/VulkanRHI/Public/VulkanThirdParty.h
+++ b/Engine/Source/Runtime/VulkanRHI/Public/VulkanThirdParty.h
@@ -2,6 +2,8 @@
 
 #pragma once
 
+#include "HAL/Platform.h"
+
 #if defined(VK_USE_PLATFORM_WIN32_KHR)
 	#include "Microsoft/AllowMicrosoftPlatformTypes.h"
 	#include <windows.h>
@@ -39,4 +41,16 @@ THIRD_PARTY_INCLUDES_START
 	#include <vulkan.h>
 THIRD_PARTY_INCLUDES_END
 
+#if VK_HEADER_VERSION >= 141
+	//workaround for removed defines in sdk 141
+	#define VK_DESCRIPTOR_TYPE_BEGIN_RANGE (VK_DESCRIPTOR_TYPE_SAMPLER)
+	#define VK_DESCRIPTOR_TYPE_END_RANGE (VK_DESCRIPTOR_TYPE_INPUT_ATTACHMENT)
+	#define VK_DESCRIPTOR_TYPE_RANGE_SIZE (VK_DESCRIPTOR_TYPE_INPUT_ATTACHMENT - VK_DESCRIPTOR_TYPE_SAMPLER + 1)
+	#define VK_IMAGE_VIEW_TYPE_RANGE_SIZE (VK_IMAGE_VIEW_TYPE_CUBE_ARRAY - VK_IMAGE_VIEW_TYPE_1D + 1)
+	#define VK_DYNAMIC_STATE_BEGIN_RANGE (VK_DYNAMIC_STATE_VIEWPORT)
+	#define VK_DYNAMIC_STATE_END_RANGE (VK_DYNAMIC_STATE_STENCIL_REFERENCE)
+	#define VK_DYNAMIC_STATE_RANGE_SIZE (VK_DYNAMIC_STATE_STENCIL_REFERENCE - VK_DYNAMIC_STATE_VIEWPORT + 1)
+	#define VK_FORMAT_RANGE_SIZE (VK_FORMAT_ASTC_12x12_SRGB_BLOCK - VK_FORMAT_UNDEFINED + 1)
+#endif
+
 #endif
diff --git a/Engine/Source/Runtime/VulkanRHI/Public/VulkanUtil.h b/Engine/Source/Runtime/VulkanRHI/Public/VulkanUtil.h
--- a/Engine/Source/Runtime/VulkanRHI/Public/VulkanUtil.h
+++ b/Engine/Source/Runtime/VulkanRHI/Public/VulkanUtil.h
@@ -6,7 +6,12 @@
 
 #pragma once
 
+#include "HAL/Platform.h"
 #include "Serialization/MemoryWriter.h"
+#include "Templates/UniquePtr.h"
+#include "VulkanThirdParty.h"
+
+class FVulkanDevice;
 
 namespace VulkanRHI
 {
