VulkanRHI: Added a vkTrimCommandPool call between map loads to free up extra memory reserved for command buffers

diff --git a/Engine/Source/Runtime/VulkanRHI/Private/VulkanCommandBuffer.cpp b/Engine/Source/Runtime/VulkanRHI/Private/VulkanCommandBuffer.cpp
--- a/Engine/Source/Runtime/VulkanRHI/Private/VulkanCommandBuffer.cpp
+++ b/Engine/Source/Runtime/VulkanRHI/Private/VulkanCommandBuffer.cpp
@@ -830,22 +830,24 @@ struct FRHICommandFreeUnusedCmdBuffers final : public FRHICommand<FRHICommandFre
 {
 	FVulkanCommandBufferPool* Pool;
 	FVulkanQueue* Queue;
+	bool bTrimMemory;
 
-	FRHICommandFreeUnusedCmdBuffers(FVulkanCommandBufferPool* InPool, FVulkanQueue* InQueue)
+	FRHICommandFreeUnusedCmdBuffers(FVulkanCommandBufferPool* InPool, FVulkanQueue* InQueue, bool bInTrimMemory)
 		: Pool(InPool)
 		, Queue(InQueue)
+		, bTrimMemory(bInTrimMemory)
 	{
 	}
 
 	void Execute(FRHICommandListBase& CmdList)
 	{
-		Pool->FreeUnusedCmdBuffers(Queue);
+		Pool->FreeUnusedCmdBuffers(Queue, bTrimMemory);
 	}
 };
 #endif
 
 
-void FVulkanCommandBufferPool::FreeUnusedCmdBuffers(FVulkanQueue* InQueue)
+void FVulkanCommandBufferPool::FreeUnusedCmdBuffers(FVulkanQueue* InQueue, bool bTrimMemory)
 {
 #if VULKAN_DELETE_STALE_CMDBUFFERS
 	FScopeLock ScopeLock(&CS);
@@ -873,21 +875,33 @@ void FVulkanCommandBufferPool::FreeUnusedCmdBuffers(FVulkanQueue* InQueue)
 			FreeCmdBuffers.Add(CmdBuffer);
 		}
 	}
+
+	if (bTrimMemory)
+	{
+		VulkanRHI::vkTrimCommandPool(Device->GetInstanceHandle(), Handle, 0);
+		
+		for (int32 Index = 0; Index < FreeCmdBuffers.Num(); ++Index)
+		{
+			FVulkanCmdBuffer* CmdBuffer = FreeCmdBuffers[Index];
+			delete CmdBuffer;
+		}
+		FreeCmdBuffers.Empty(8);
+	}
 #endif
 }
 
-void FVulkanCommandBufferManager::FreeUnusedCmdBuffers()
+void FVulkanCommandBufferManager::FreeUnusedCmdBuffers(bool bTrimMemory)
 {
 #if VULKAN_DELETE_STALE_CMDBUFFERS
 	FRHICommandList& RHICmdList = FRHICommandListExecutor::GetImmediateCommandList();
 	if (!IsInRenderingThread() || (RHICmdList.Bypass() || !IsRunningRHIInSeparateThread()))
 	{
-		Pool.FreeUnusedCmdBuffers(Queue);
+		Pool.FreeUnusedCmdBuffers(Queue, bTrimMemory);
 	}
 	else
 	{
 		check(IsInRenderingThread());
-		ALLOC_COMMAND_CL(RHICmdList, FRHICommandFreeUnusedCmdBuffers)(&Pool, Queue);
+		ALLOC_COMMAND_CL(RHICmdList, FRHICommandFreeUnusedCmdBuffers)(&Pool, Queue, bTrimMemory);
 	}
 #endif
 }
diff --git a/Engine/Source/Runtime/VulkanRHI/Private/VulkanCommandBuffer.h b/Engine/Source/Runtime/VulkanRHI/Private/VulkanCommandBuffer.h
--- a/Engine/Source/Runtime/VulkanRHI/Private/VulkanCommandBuffer.h
+++ b/Engine/Source/Runtime/VulkanRHI/Private/VulkanCommandBuffer.h
@@ -277,7 +277,7 @@ public:
 		return &CS;
 	}
 
-	void FreeUnusedCmdBuffers(FVulkanQueue* Queue);
+	void FreeUnusedCmdBuffers(FVulkanQueue* Queue, bool bTrimMemory);
 
 	inline FVulkanCommandBufferManager& GetMgr()
 	{
@@ -376,7 +376,7 @@ public:
 
 	uint32 CalculateGPUTime();
 
-	void FreeUnusedCmdBuffers();
+	void FreeUnusedCmdBuffers(bool bTrimMemory);
 
 	inline FVulkanCommandListContext* GetCommandListContext()
 	{
diff --git a/Engine/Source/Runtime/VulkanRHI/Private/VulkanCommandWrappers.h b/Engine/Source/Runtime/VulkanRHI/Private/VulkanCommandWrappers.h
--- a/Engine/Source/Runtime/VulkanRHI/Private/VulkanCommandWrappers.h
+++ b/Engine/Source/Runtime/VulkanRHI/Private/VulkanCommandWrappers.h
@@ -134,6 +134,7 @@ struct FWrapLayer
 	static void CreateCommandPool(VkResult Result, VkDevice Device, const VkCommandPoolCreateInfo* CreateInfo, VkCommandPool* CommandPool) VULKAN_LAYER_BODY
 	static void DestroyCommandPool(VkResult Result, VkDevice Device, VkCommandPool CommandPool) VULKAN_LAYER_BODY
 	static void ResetCommandPool(VkResult Result, VkDevice Device, VkCommandPool CommandPool, VkCommandPoolResetFlags Flags) VULKAN_LAYER_BODY
+	static void TrimCommandPool(VkResult Result, VkDevice Device, VkCommandPool CommandPool, VkCommandPoolTrimFlags Flags) VULKAN_LAYER_BODY
 	static void AllocateCommandBuffers(VkResult Result, VkDevice Device, const VkCommandBufferAllocateInfo* AllocateInfo, VkCommandBuffer* CommandBuffers) VULKAN_LAYER_BODY
 	static void FreeCommandBuffers(VkResult Result, VkDevice Device, VkCommandPool CommandPool, uint32 CommandBufferCount, const VkCommandBuffer* CommandBuffers) VULKAN_LAYER_BODY
 	static void BeginCommandBuffer(VkResult Result, VkCommandBuffer CommandBuffer, const VkCommandBufferBeginInfo* BeginInfo) VULKAN_LAYER_BODY
@@ -960,6 +961,13 @@ namespace VulkanRHI
 		return Result;
 	}
 
+	static FORCEINLINE_DEBUGGABLE void  vkTrimCommandPool(VkDevice Device, VkCommandPool CommandPool, VkCommandPoolTrimFlags Flags)
+	{
+		FWrapLayer::TrimCommandPool(VK_RESULT_MAX_ENUM, Device, CommandPool, Flags);
+		VULKANAPINAMESPACE::vkTrimCommandPool(Device, CommandPool, Flags);
+		FWrapLayer::TrimCommandPool(VK_SUCCESS, Device, CommandPool, Flags);
+	}
+
 	static FORCEINLINE_DEBUGGABLE VkResult  vkAllocateCommandBuffers(VkDevice Device, const VkCommandBufferAllocateInfo* AllocateInfo, VkCommandBuffer* CommandBuffers)
 	{
 		FWrapLayer::AllocateCommandBuffers(VK_RESULT_MAX_ENUM, Device, AllocateInfo, CommandBuffers);
diff --git a/Engine/Source/Runtime/VulkanRHI/Private/VulkanDebug.cpp b/Engine/Source/Runtime/VulkanRHI/Private/VulkanDebug.cpp
--- a/Engine/Source/Runtime/VulkanRHI/Private/VulkanDebug.cpp
+++ b/Engine/Source/Runtime/VulkanRHI/Private/VulkanDebug.cpp
@@ -3148,6 +3148,22 @@ void FWrapLayer::ResetCommandPool(VkResult Result, VkDevice Device, VkCommandPoo
 	}
 }
 
+void FWrapLayer::TrimCommandPool(VkResult Result, VkDevice Device, VkCommandPool CommandPool, VkCommandPoolTrimFlags Flags)
+{
+	if (Result == VK_RESULT_MAX_ENUM)
+	{
+#if VULKAN_ENABLE_DUMP_LAYER
+		DevicePrintfBeginResult(Device, FString::Printf(TEXT("vkTrimCommandPool(CommandPool=0x%p, Flags=0x%08X)"), CommandPool, Flags));
+#endif
+	}
+	else
+	{
+#if VULKAN_ENABLE_DUMP_LAYER
+		PrintResult(Result);
+#endif
+	}
+}
+
 void FWrapLayer::AllocateCommandBuffers(VkResult Result, VkDevice Device, const VkCommandBufferAllocateInfo* AllocateInfo, VkCommandBuffer* CommandBuffers)
 {
 	if (Result == VK_RESULT_MAX_ENUM)
diff --git a/Engine/Source/Runtime/VulkanRHI/Private/VulkanLoader.h b/Engine/Source/Runtime/VulkanRHI/Private/VulkanLoader.h
--- a/Engine/Source/Runtime/VulkanRHI/Private/VulkanLoader.h
+++ b/Engine/Source/Runtime/VulkanRHI/Private/VulkanLoader.h
@@ -97,6 +97,7 @@
 	EnumMacro(PFN_vkCreateCommandPool, vkCreateCommandPool) \
 	EnumMacro(PFN_vkDestroyCommandPool, vkDestroyCommandPool) \
 	EnumMacro(PFN_vkResetCommandPool, vkResetCommandPool) \
+	EnumMacro(PFN_vkTrimCommandPool, vkTrimCommandPool) \
 	EnumMacro(PFN_vkAllocateCommandBuffers, vkAllocateCommandBuffers) \
 	EnumMacro(PFN_vkFreeCommandBuffers, vkFreeCommandBuffers) \
 	EnumMacro(PFN_vkBeginCommandBuffer, vkBeginCommandBuffer) \
diff --git a/Engine/Source/Runtime/VulkanRHI/Private/VulkanRHI.cpp b/Engine/Source/Runtime/VulkanRHI/Private/VulkanRHI.cpp
--- a/Engine/Source/Runtime/VulkanRHI/Private/VulkanRHI.cpp
+++ b/Engine/Source/Runtime/VulkanRHI/Private/VulkanRHI.cpp
@@ -1167,7 +1167,8 @@ void FVulkanCommandListContext::RHIEndFrame()
 
 	GetGPUProfiler().EndFrame();
 
-	GetCommandBufferManager()->FreeUnusedCmdBuffers();
+	bool bTrimMemory = false;
+	GetCommandBufferManager()->FreeUnusedCmdBuffers(bTrimMemory);
 
 	Device->GetStagingManager().ProcessPendingFree(false, true);
 	Device->GetMemoryManager().ReleaseFreedPages(*this);
@@ -1292,6 +1293,9 @@ bool FVulkanDynamicRHI::RHIGetAvailableResolutions(FScreenResolutionArray& Resol
 
 void FVulkanDynamicRHI::RHIFlushResources()
 {
+	FVulkanCommandListContextImmediate& ImmediateContext = GetDevice()->GetImmediateContext();
+	bool bTrimMemory = true;
+	ImmediateContext.GetCommandBufferManager()->FreeUnusedCmdBuffers(bTrimMemory);
 }
 
 // IVulkanDynamicRHI interface
