PR #10926: Add support for decoding R11G11B10UF vulkan targets

#rnx

diff --git a/Engine/Source/Runtime/VulkanRHI/Private/VulkanRenderTarget.cpp b/Engine/Source/Runtime/VulkanRHI/Private/VulkanRenderTarget.cpp
--- a/Engine/Source/Runtime/VulkanRHI/Private/VulkanRenderTarget.cpp
+++ b/Engine/Source/Runtime/VulkanRHI/Private/VulkanRenderTarget.cpp
@@ -87,6 +87,10 @@ static void ConvertRawDataToFColor(VkFormat VulkanFormat, uint32 DestWidth, uint
 		ConvertRawR16G16B16A16FDataToFColor(DestWidth, DestHeight, In, SrcPitch, Dest, bLinearToGamma);
 		break;
 
+	case VK_FORMAT_B10G11R11_UFLOAT_PACK32:
+		ConvertRawR11G11B10DataToFColor(DestWidth, DestHeight, In, SrcPitch, Dest, bLinearToGamma);
+		break;
+
 	case VK_FORMAT_A2B10G10R10_UNORM_PACK32:
 		ConvertRawR10G10B10A2DataToFColor(DestWidth, DestHeight, In, SrcPitch, Dest);
 		break;
