Attempt to fix vulkan RHI assert CmdBuffer->IsSubmitted() on running engine rendering tests

diff --git a/Engine/Source/Runtime/VulkanRHI/Private/VulkanCommandBuffer.cpp b/Engine/Source/Runtime/VulkanRHI/Private/VulkanCommandBuffer.cpp
--- a/Engine/Source/Runtime/VulkanRHI/Private/VulkanCommandBuffer.cpp
+++ b/Engine/Source/Runtime/VulkanRHI/Private/VulkanCommandBuffer.cpp
@@ -851,6 +851,13 @@ void FVulkanCommandBufferPool::FreeUnusedCmdBuffers(FVulkanQueue* InQueue, bool
 {
 #if VULKAN_DELETE_STALE_CMDBUFFERS
 	FScopeLock ScopeLock(&CS);
+	
+	if (bTrimMemory)
+	{
+		VulkanRHI::vkTrimCommandPool(Device->GetInstanceHandle(), Handle, 0);
+		return;
+	}
+
 	const double CurrentTime = FPlatformTime::Seconds();
 
 	// In case Queue stores pointer to a cmdbuffer, do not delete it
@@ -875,18 +882,6 @@ void FVulkanCommandBufferPool::FreeUnusedCmdBuffers(FVulkanQueue* InQueue, bool
 			FreeCmdBuffers.Add(CmdBuffer);
 		}
 	}
-
-	if (bTrimMemory)
-	{
-		VulkanRHI::vkTrimCommandPool(Device->GetInstanceHandle(), Handle, 0);
-		
-		for (int32 Index = 0; Index < FreeCmdBuffers.Num(); ++Index)
-		{
-			FVulkanCmdBuffer* CmdBuffer = FreeCmdBuffers[Index];
-			delete CmdBuffer;
-		}
-		FreeCmdBuffers.Empty(8);
-	}
 #endif
 }
 
