Deprecate:

diff --git a/Engine/Source/Runtime/Renderer/Private/VT/AdaptiveVirtualTexture.cpp b/Engine/Source/Runtime/Renderer/Private/VT/AdaptiveVirtualTexture.cpp
--- a/Engine/Source/Runtime/Renderer/Private/VT/AdaptiveVirtualTexture.cpp
+++ b/Engine/Source/Runtime/Renderer/Private/VT/AdaptiveVirtualTexture.cpp
@@ -752,7 +752,7 @@ void FAdaptiveVirtualTexture::UpdateAllocations(FVirtualTextureSystem* InSystem,
 		for (FIndirectionTextureUpdate& TextureUpdate : TextureUpdates)
 		{
 			const FUpdateTextureRegion2D Region(TextureUpdate.X, TextureUpdate.Y, 0, 0, 1, 1);
-			RHIUpdateTexture2D((FRHITexture2D*)Texture, 0, Region, 4, (uint8*)&TextureUpdate.Value);
+			RHIUpdateTexture2D((FRHITexture*)Texture, 0, Region, 4, (uint8*)&TextureUpdate.Value);
 		}
 		RHICmdList.Transition(FRHITransitionInfo(Texture, ERHIAccess::UAVCompute, ERHIAccess::SRVMask));
 	}
diff --git a/Engine/Source/Runtime/Renderer/Private/VT/RuntimeVirtualTextureRender.cpp b/Engine/Source/Runtime/Renderer/Private/VT/RuntimeVirtualTextureRender.cpp
--- a/Engine/Source/Runtime/Renderer/Private/VT/RuntimeVirtualTextureRender.cpp
+++ b/Engine/Source/Runtime/Renderer/Private/VT/RuntimeVirtualTextureRender.cpp
@@ -1331,7 +1331,7 @@ namespace RuntimeVirtualTexture
 	END_SHADER_PARAMETER_STRUCT()
 
 	/** Set up the copy to final output physical texture. */
-	void AddCopyToOutputPass(FRDGBuilder& GraphBuilder, FRDGTextureRef InputTexture, FRHITexture2D* OutputTexture, FBox2D const& DestBox)
+	void AddCopyToOutputPass(FRDGBuilder& GraphBuilder, FRDGTextureRef InputTexture, FRHITexture* OutputTexture, FBox2D const& DestBox)
 	{
 		FRHICopyTextureInfo CopyInfo;
 		CopyInfo.Size = InputTexture->Desc.GetSize();
@@ -1375,13 +1375,13 @@ namespace RuntimeVirtualTexture
 		ERuntimeVirtualTextureMaterialType MaterialType,
 		bool bClearTextures,
 		bool bIsThumbnails,
-		FRHITexture2D* OutputTexture0,		// todo[vt]: Only use IPooledRenderTarget or FRDGTextureRef, not raw RHI textures.
+		FRHITexture* OutputTexture0,		// todo[vt]: Only use IPooledRenderTarget or FRDGTextureRef, not raw RHI textures.
 		IPooledRenderTarget* OutputTarget0,
 		FBox2D const& DestBox0,
-		FRHITexture2D* OutputTexture1,
+		FRHITexture* OutputTexture1,
 		IPooledRenderTarget* OutputTarget1,
 		FBox2D const& DestBox1,
-		FRHITexture2D* OutputTexture2, 
+		FRHITexture* OutputTexture2, 
 		IPooledRenderTarget* OutputTarget2,
 		FBox2D const& DestBox2,
 		FTransform const& UVToWorld,
diff --git a/Engine/Source/Runtime/SlateCore/Public/Rendering/SlateRenderer.h b/Engine/Source/Runtime/SlateCore/Public/Rendering/SlateRenderer.h
--- a/Engine/Source/Runtime/SlateCore/Public/Rendering/SlateRenderer.h
+++ b/Engine/Source/Runtime/SlateCore/Public/Rendering/SlateRenderer.h
@@ -24,7 +24,7 @@ struct Rect;
 class FSceneInterface;
 struct FSlateBrush;
 
-typedef TRefCountPtr<FRHITexture> FTexture2DRHIRef;
+typedef TRefCountPtr<FRHITexture> FTextureRHIRef;
 
 /**
  * Update context for deferred drawing of widgets to render targets
@@ -297,7 +297,7 @@ public:
 	FOnPostResizeWindowBackbuffer& OnPostResizeWindowBackBuffer() { return PostResizeBackBufferDelegate; }
 
 	/** Callback on the render thread after slate rendering finishes and right before present is called */
-	DECLARE_TS_MULTICAST_DELEGATE_TwoParams(FOnBackBufferReadyToPresent, SWindow&, const FTexture2DRHIRef&);
+	DECLARE_TS_MULTICAST_DELEGATE_TwoParams(FOnBackBufferReadyToPresent, SWindow&, const FTextureRHIRef&);
 	FOnBackBufferReadyToPresent& OnBackBufferReadyToPresent() { return OnBackBufferReadyToPresentDelegate; }
 
 	/** 
diff --git a/Engine/Source/Runtime/SlateRHIRenderer/Private/FX/SlatePostBufferBlur.cpp b/Engine/Source/Runtime/SlateRHIRenderer/Private/FX/SlatePostBufferBlur.cpp
--- a/Engine/Source/Runtime/SlateRHIRenderer/Private/FX/SlatePostBufferBlur.cpp
+++ b/Engine/Source/Runtime/SlateRHIRenderer/Private/FX/SlatePostBufferBlur.cpp
@@ -78,13 +78,13 @@ void USlatePostBufferBlur::PostProcess(FRenderResource* InViewInfo, FRenderResou
 	ENQUEUE_RENDER_COMMAND(FUpdateSlatePostBuffersWithFX_Blur)([InViewInfo, InViewportTexture, InElementWindowSize, InRenderingPolicy, InSlatePostBuffer, GaussianBlurStrengthCopy](FRHICommandListImmediate& RHICmdList)
 	{
 		// Get Backbuffer, which can vary between PIE or standalone
-		FTexture2DRHIRef BackBuffer = USlateRHIPostBufferProcessor::GetBackbuffer_RenderThread(InViewInfo, InViewportTexture, InElementWindowSize, RHICmdList);
+		FTextureRHIRef BackBuffer = USlateRHIPostBufferProcessor::GetBackbuffer_RenderThread(InViewInfo, InViewportTexture, InElementWindowSize, RHICmdList);
 
 		if (BackBuffer)
 		{
 			// Get Src / Dst textures & their rects, again may vary between PIE or standalone
 			// Here we can simply use Src Rect since the Src texture in PIE is the 'BufferedRT' scene backbuffer without the editor
-			FTexture2DRHIRef Src = USlateRHIPostBufferProcessor::GetSrcTexture_RenderThread(BackBuffer, InViewportTexture);
+			FTextureRHIRef Src = USlateRHIPostBufferProcessor::GetSrcTexture_RenderThread(BackBuffer, InViewportTexture);
 			FTextureReferenceRHIRef& Dst = USlateRHIPostBufferProcessor::GetDstTexture_RenderThread(InSlatePostBuffer);
 			FIntPoint DstExtent = USlateRHIPostBufferProcessor::GetDstExtent_RenderThread(BackBuffer, InViewportTexture);
 			FIntRect SrcRect = FIntRect(0, 0, Src->GetSizeX(), Src->GetSizeY());
diff --git a/Engine/Source/Runtime/SlateRHIRenderer/Private/FX/SlateRHIPostBufferProcessor.cpp b/Engine/Source/Runtime/SlateRHIRenderer/Private/FX/SlateRHIPostBufferProcessor.cpp
--- a/Engine/Source/Runtime/SlateRHIRenderer/Private/FX/SlateRHIPostBufferProcessor.cpp
+++ b/Engine/Source/Runtime/SlateRHIRenderer/Private/FX/SlateRHIPostBufferProcessor.cpp
@@ -9,7 +9,7 @@
 
 #include UE_INLINE_GENERATED_CPP_BY_NAME(SlateRHIPostBufferProcessor)
 
-FTexture2DRHIRef USlateRHIPostBufferProcessor::GetBackbuffer_RenderThread(FRenderResource* InViewInfo, FRenderResource* InViewportTexture, FVector2D InElementWindowSize, FRHICommandListImmediate& InRHICmdList)
+FTextureRHIRef USlateRHIPostBufferProcessor::GetBackbuffer_RenderThread(FRenderResource* InViewInfo, FRenderResource* InViewportTexture, FVector2D InElementWindowSize, FRHICommandListImmediate& InRHICmdList)
 {
 	FViewportInfo* ViewInfo = static_cast<FViewportInfo*>(InViewInfo);
 	FSlateRenderTargetRHI* ViewportTexture = static_cast<FSlateRenderTargetRHI*>(InViewportTexture);
@@ -21,13 +21,13 @@ FTexture2DRHIRef USlateRHIPostBufferProcessor::GetBackbuffer_RenderThread(FRende
 		bRenderedStereo = true;
 	}
 
-	FTexture2DRHIRef ViewportRT = bRenderedStereo ? nullptr : ViewInfo->GetRenderTargetTexture();
-	FTexture2DRHIRef BackBuffer = (ViewportRT) ? ViewportRT : RHIGetViewportBackBuffer(ViewInfo->ViewportRHI);
+	FTextureRHIRef ViewportRT = bRenderedStereo ? nullptr : ViewInfo->GetRenderTargetTexture();
+	FTextureRHIRef BackBuffer = (ViewportRT) ? ViewportRT : RHIGetViewportBackBuffer(ViewInfo->ViewportRHI);
 
 	return BackBuffer;
 }
 
-FTexture2DRHIRef USlateRHIPostBufferProcessor::GetSrcTexture_RenderThread(FTexture2DRHIRef InBackBuffer, FRenderResource* InViewportTexture)
+FTextureRHIRef USlateRHIPostBufferProcessor::GetSrcTexture_RenderThread(FTextureRHIRef InBackBuffer, FRenderResource* InViewportTexture)
 {
 	FSlateRenderTargetRHI* ViewportTexture = static_cast<FSlateRenderTargetRHI*>(InViewportTexture);
 	return GIsEditor ? ViewportTexture->GetRHIRef() : InBackBuffer;
@@ -38,7 +38,7 @@ FTextureReferenceRHIRef& USlateRHIPostBufferProcessor::GetDstTexture_RenderThrea
 	return InSlatePostBuffer->TextureReference.TextureReferenceRHI;
 }
 
-FIntPoint USlateRHIPostBufferProcessor::GetDstExtent_RenderThread(FTexture2DRHIRef InBackBuffer, FRenderResource* InViewportTexture)
+FIntPoint USlateRHIPostBufferProcessor::GetDstExtent_RenderThread(FTextureRHIRef InBackBuffer, FRenderResource* InViewportTexture)
 {
 	FSlateRenderTargetRHI* ViewportTexture = static_cast<FSlateRenderTargetRHI*>(InViewportTexture);
 	return GIsEditor ? FIntPoint(ViewportTexture->GetWidth(), ViewportTexture->GetHeight()) : InBackBuffer->GetDesc().Extent;
diff --git a/Engine/Source/Runtime/SlateRHIRenderer/Private/Slate3DRenderer.cpp b/Engine/Source/Runtime/SlateRHIRenderer/Private/Slate3DRenderer.cpp
--- a/Engine/Source/Runtime/SlateRHIRenderer/Private/Slate3DRenderer.cpp
+++ b/Engine/Source/Runtime/SlateRHIRenderer/Private/Slate3DRenderer.cpp
@@ -186,7 +186,7 @@ void FSlate3DRenderer::DrawWindowToTarget_RenderThread(FRHICommandListImmediate&
 			// The scene renderer will handle it in this case
 			DrawOptions.ViewOffset = UE::Slate::CastToVector2f(DrawOffset);
 
-			FTexture2DRHIRef ColorTarget = Context.RenderTarget->GetRenderTargetTexture();
+			FTextureRHIRef ColorTarget = Context.RenderTarget->GetRenderTargetTexture();
 
 			if (BatchData.IsStencilClippingRequired())
 			{
diff --git a/Engine/Source/Runtime/SlateRHIRenderer/Private/Slate3DRenderer.h b/Engine/Source/Runtime/SlateRHIRenderer/Private/Slate3DRenderer.h
--- a/Engine/Source/Runtime/SlateRHIRenderer/Private/Slate3DRenderer.h
+++ b/Engine/Source/Runtime/SlateRHIRenderer/Private/Slate3DRenderer.h
@@ -53,7 +53,7 @@ private:
 	uint8 FreeBufferIndex;
 
 	/** The depth buffer texture if any */
-	FTexture2DRHIRef DepthStencil;
+	FTextureRHIRef DepthStencil;
 
 	/** Set to true when the render target was cleared and prevent another clear to be called for nothing.**/
 	bool bRenderTargetWasCleared;
diff --git a/Engine/Source/Runtime/SlateRHIRenderer/Private/SlatePostProcessResource.h b/Engine/Source/Runtime/SlateRHIRenderer/Private/SlatePostProcessResource.h
--- a/Engine/Source/Runtime/SlateRHIRenderer/Private/SlatePostProcessResource.h
+++ b/Engine/Source/Runtime/SlateRHIRenderer/Private/SlatePostProcessResource.h
@@ -19,7 +19,7 @@ public:
 	FSlatePostProcessResource(int32 InRenderTargetCount);
 	~FSlatePostProcessResource();
 
-	const FTexture2DRHIRef& GetRenderTarget(int32 Index)
+	const FTextureRHIRef& GetRenderTarget(int32 Index)
 	{
 		return RenderTargets[Index]; 
 	}
@@ -46,7 +46,7 @@ private:
 	void ResizeTargets(const FIntPoint& NewSize, EPixelFormat RequestedPixelFormat);
 
 private:
-	TArray<FTexture2DRHIRef, TInlineAllocator<2>> RenderTargets;
+	TArray<FTextureRHIRef, TInlineAllocator<2>> RenderTargets;
 	EPixelFormat PixelFormat;
 	FIntPoint RenderTargetSize;
 	int32 RenderTargetCount;
diff --git a/Engine/Source/Runtime/SlateRHIRenderer/Private/SlatePostProcessor.cpp b/Engine/Source/Runtime/SlateRHIRenderer/Private/SlatePostProcessor.cpp
--- a/Engine/Source/Runtime/SlateRHIRenderer/Private/SlatePostProcessor.cpp
+++ b/Engine/Source/Runtime/SlateRHIRenderer/Private/SlatePostProcessor.cpp
@@ -147,7 +147,7 @@ public:
 IMPLEMENT_SHADER_TYPE(template<>, FBlitUIToHDRPS<0>, FBlitUIToHDRPS::GetSourceFilename(), FBlitUIToHDRPS::GetFunctionName(), SF_Pixel);
 IMPLEMENT_SHADER_TYPE(template<>, FBlitUIToHDRPS<1>, FBlitUIToHDRPS::GetSourceFilename(), FBlitUIToHDRPS::GetFunctionName(), SF_Pixel);
 
-static void BlitUIToHDRScene(FRHICommandListImmediate& RHICmdList, IRendererModule& RendererModule, const FPostProcessRectParams& RectParams, FTexture2DRHIRef DestTexture, const FIntPoint& AllocatedSize)
+static void BlitUIToHDRScene(FRHICommandListImmediate& RHICmdList, IRendererModule& RendererModule, const FPostProcessRectParams& RectParams, FTextureRHIRef DestTexture, const FIntPoint& AllocatedSize)
 {
 	SCOPED_DRAW_EVENT(RHICmdList, SlatePostProcessBlitUIToHDR);
 
@@ -337,8 +337,8 @@ void FSlatePostProcessor::BlurRect(FRHICommandListImmediate& RHICmdList, IRender
 		// First pass render to the render target with the post process fx
 		if (PassIndex == 0)
 		{
-			FTexture2DRHIRef SourceTexture = (bDownsample || bIsHDRSource) ? IntermediateTargets->GetRenderTarget(0) : RectParams.SourceTexture;
-			FTexture2DRHIRef DestTexture = IntermediateTargets->GetRenderTarget(1);
+			FTextureRHIRef SourceTexture = (bDownsample || bIsHDRSource) ? IntermediateTargets->GetRenderTarget(0) : RectParams.SourceTexture;
+			FTextureRHIRef DestTexture = IntermediateTargets->GetRenderTarget(1);
 
 			RHICmdList.Transition(FRHITransitionInfo(SourceTexture, ERHIAccess::Unknown, ERHIAccess::SRVGraphics));
 			RHICmdList.Transition(FRHITransitionInfo(DestTexture, ERHIAccess::Unknown, ERHIAccess::RTV));
@@ -411,8 +411,8 @@ void FSlatePostProcessor::BlurRect(FRHICommandListImmediate& RHICmdList, IRender
 		}
 		else
 		{
-			FTexture2DRHIRef SourceTexture = IntermediateTargets->GetRenderTarget(1);
-			FTexture2DRHIRef DestTexture = IntermediateTargets->GetRenderTarget(0);
+			FTextureRHIRef SourceTexture = IntermediateTargets->GetRenderTarget(1);
+			FTextureRHIRef DestTexture = IntermediateTargets->GetRenderTarget(0);
 
 			RHICmdList.Transition(FRHITransitionInfo(SourceTexture, ERHIAccess::Unknown, ERHIAccess::SRVGraphics));
 			RHICmdList.Transition(FRHITransitionInfo(DestTexture, ERHIAccess::Unknown, ERHIAccess::RTV));
@@ -499,8 +499,8 @@ void FSlatePostProcessor::ColorDeficiency(FRHICommandListImmediate& RHICmdList,
 
 	// 
 	{
-		FTexture2DRHIRef SourceTexture = RectParams.SourceTexture;
-		FTexture2DRHIRef DestTexture = IntermediateTargets->GetRenderTarget(0);
+		FTextureRHIRef SourceTexture = RectParams.SourceTexture;
+		FTextureRHIRef DestTexture = IntermediateTargets->GetRenderTarget(0);
 
 		RHICmdList.Transition(FRHITransitionInfo(SourceTexture, ERHIAccess::Unknown, ERHIAccess::SRVGraphics));
 		RHICmdList.Transition(FRHITransitionInfo(DestTexture, ERHIAccess::Unknown, ERHIAccess::RTV));
@@ -578,7 +578,7 @@ void FSlatePostProcessor::DownsampleRect(FRHICommandListImmediate& RHICmdList, I
 
 	FSamplerStateRHIRef BilinearClamp = TStaticSamplerState<SF_Bilinear, AM_Clamp, AM_Clamp, AM_Clamp>::GetRHI();
 
-	FTexture2DRHIRef DestTexture = IntermediateTargets->GetRenderTarget(0);
+	FTextureRHIRef DestTexture = IntermediateTargets->GetRenderTarget(0);
 
 	// Downsample and store in intermediate texture
 	{
@@ -656,7 +656,7 @@ void FSlatePostProcessor::UpsampleRect(FRHICommandListImmediate& RHICmdList, IRe
 	GraphicsPSOInit.DepthStencilState = TStaticDepthStencilState<false, CF_Always>::GetRHI();
 
 	// Original source texture is now the destination texture
-	FTexture2DRHIRef DestTexture = Params.PostProcessDest == EPostProcessDestination::DestTexture && Params.DestTexture 
+	FTextureRHIRef DestTexture = Params.PostProcessDest == EPostProcessDestination::DestTexture && Params.DestTexture 
 		? Params.DestTexture 
 		: Params.SourceTexture;
 	const int32 DestTextureWidth = DestTexture->GetSizeX();
@@ -666,7 +666,7 @@ void FSlatePostProcessor::UpsampleRect(FRHICommandListImmediate& RHICmdList, IRe
 	const int32 DownsampledHeight = DownsampleSize.Y;
 
 	// Source texture is the texture that was originally downsampled
-	FTexture2DRHIRef SrcTexture = IntermediateTargets->GetRenderTarget(0);
+	FTextureRHIRef SrcTexture = IntermediateTargets->GetRenderTarget(0);
 	const int32 SrcTextureWidth = IntermediateTargets->GetWidth();
 	const int32 SrcTextureHeight = IntermediateTargets->GetHeight();
 
diff --git a/Engine/Source/Runtime/SlateRHIRenderer/Private/SlateRHIFontTexture.h b/Engine/Source/Runtime/SlateRHIRenderer/Private/SlateRHIFontTexture.h
--- a/Engine/Source/Runtime/SlateRHIRenderer/Private/SlateRHIFontTexture.h
+++ b/Engine/Source/Runtime/SlateRHIRenderer/Private/SlateRHIFontTexture.h
@@ -12,7 +12,7 @@
  * ensure all characters in the font texture exist if the rendering resource has to be recreated 
  * between caching new characters
  */
-class FSlateFontTextureRHIResource : public TSlateTexture<FTexture2DRHIRef>, public FTextureResource
+class FSlateFontTextureRHIResource : public TSlateTexture<FTextureRHIRef>, public FTextureResource
 {
 public:
 	/** Constructor.  Initializes the texture
diff --git a/Engine/Source/Runtime/SlateRHIRenderer/Private/SlateRHIRenderer.cpp b/Engine/Source/Runtime/SlateRHIRenderer/Private/SlateRHIRenderer.cpp
--- a/Engine/Source/Runtime/SlateRHIRenderer/Private/SlateRHIRenderer.cpp
+++ b/Engine/Source/Runtime/SlateRHIRenderer/Private/SlateRHIRenderer.cpp
@@ -218,7 +218,7 @@ void FViewportInfo::RecreateDepthBuffer_RenderThread()
 	}
 }
 
-bool IsMemorylessTexture(const FTexture2DRHIRef& Tex)
+bool IsMemorylessTexture(const FTextureRHIRef& Tex)
 {
 	if (Tex)
 	{
@@ -969,9 +969,9 @@ IMPLEMENT_SHADER_TYPE(template<>, FCompositeCS_SCRGB_CDV, FCompositeCSBase::GetS
 int32 SlateWireFrame = 0;
 static FAutoConsoleVariableRef CVarSlateWireframe(TEXT("Slate.ShowWireFrame"), SlateWireFrame, TEXT(""), ECVF_Default);
 
-void RenderSlateBatch(FTexture2DRHIRef SlateRenderTarget, bool bClear, bool bIsHDR, FViewportInfo& ViewportInfo, const FMatrix& ViewMatrix, FSlateBatchData& BatchData, FRHICommandListImmediate& RHICmdList,
+void RenderSlateBatch(FTextureRHIRef SlateRenderTarget, bool bClear, bool bIsHDR, FViewportInfo& ViewportInfo, const FMatrix& ViewMatrix, FSlateBatchData& BatchData, FRHICommandListImmediate& RHICmdList,
 				      const uint32 ViewportWidth, const uint32 ViewportHeight, const struct FSlateDrawWindowCommandParams& DrawCommandParams, TSharedPtr<FSlateRHIRenderingPolicy> RenderingPolicy, 
-					  FTexture2DRHIRef PostProcessBuffer)
+					  FTextureRHIRef PostProcessBuffer)
 {
 	FRHIRenderPassInfo RPInfo(SlateRenderTarget, ERenderTargetActions::Load_Store);
 
@@ -1028,7 +1028,7 @@ void RenderSlateBatch(FTexture2DRHIRef SlateRenderTarget, bool bClear, bool bIsH
 				}
 				RenderingPolicy->SetUseGammaCorrection(!bIsHDR);
 
-				FTexture2DRHIRef EmptyTarget;
+				FTextureRHIRef EmptyTarget;
 
 				RenderingPolicy->DrawElements
 				(
@@ -1132,9 +1132,9 @@ void FSlateRHIRenderer::DrawWindow_RenderThread(FRHICommandListImmediate& RHICmd
 			// should have been created by the game thread
 			check(IsValidRef(ViewportInfo.ViewportRHI));
 
-			FTexture2DRHIRef ViewportRT = bRenderedStereo ? nullptr : ViewportInfo.GetRenderTargetTexture();
-			FTexture2DRHIRef BackBuffer = (ViewportRT) ? ViewportRT : RHIGetViewportBackBuffer(ViewportInfo.ViewportRHI);
-			FTexture2DRHIRef PostProcessBuffer = BackBuffer;	// If compositing UI then this will be different to the back buffer
+			FTextureRHIRef ViewportRT = bRenderedStereo ? nullptr : ViewportInfo.GetRenderTargetTexture();
+			FTextureRHIRef BackBuffer = (ViewportRT) ? ViewportRT : RHIGetViewportBackBuffer(ViewportInfo.ViewportRHI);
+			FTextureRHIRef PostProcessBuffer = BackBuffer;	// If compositing UI then this will be different to the back buffer
 
 			const uint32 ViewportWidth = (ViewportRT) ? ViewportRT->GetSizeX() : ViewportInfo.Width;
 			const uint32 ViewportHeight = (ViewportRT) ? ViewportRT->GetSizeY() : ViewportInfo.Height;
@@ -1176,7 +1176,7 @@ void FSlateRHIRenderer::DrawWindow_RenderThread(FRHICommandListImmediate& RHICmd
 				bLUTStale = true;
 			}
 
-			FTexture2DRHIRef FinalBuffer = BackBuffer;
+			FTextureRHIRef FinalBuffer = BackBuffer;
 
 			bool bClear = DrawCommandParams.bClear;
 			if (bCompositeUI)
@@ -1237,7 +1237,7 @@ void FSlateRHIRenderer::DrawWindow_RenderThread(FRHICommandListImmediate& RHICmd
 				if (bHasBatches)
 				{
 					RenderingPolicy->BuildRenderingBuffers(RHICmdList, BatchDataHDR);
-					FTexture2DRHIRef UITargetHDRRTRHI(FinalBuffer);
+					FTextureRHIRef UITargetHDRRTRHI(FinalBuffer);
 					RenderSlateBatch(UITargetHDRRTRHI, /*bClear*/ false, /*bIsHDR*/ true, ViewportInfo, ViewMatrix, BatchDataHDR, RHICmdList, ViewportWidth, ViewportHeight, DrawCommandParams, RenderingPolicy, PostProcessBuffer);
 			    }
 
@@ -1519,7 +1519,7 @@ void FSlateRHIRenderer::DrawWindow_RenderThread(FRHICommandListImmediate& RHICmd
 	if (bTakingAScreenShot && ScreenshotViewportInfo != nullptr && ScreenshotViewportInfo == &ViewportInfo)
 	{
 		// take screenshot before swapbuffer
-		FTexture2DRHIRef BackBuffer = RHIGetViewportBackBuffer(ViewportInfo.ViewportRHI);
+		FTextureRHIRef BackBuffer = RHIGetViewportBackBuffer(ViewportInfo.ViewportRHI);
 
 		// Sanity check to make sure the user specified a valid screenshot rect.
 		FIntRect ClampedScreenshotRect;
@@ -1845,8 +1845,8 @@ void FSlateRHIRenderer::DrawWindows_Private(FSlateDrawBuffer& WindowDrawBuffer)
 											bRenderedStereo = true;
 										}
 
-										FTexture2DRHIRef ViewportRT = bRenderedStereo ? nullptr : ViewInfo->GetRenderTargetTexture();
-										FTexture2DRHIRef BackBuffer = (ViewportRT) ? ViewportRT : RHIGetViewportBackBuffer(ViewInfo->ViewportRHI);
+										FTextureRHIRef ViewportRT = bRenderedStereo ? nullptr : ViewInfo->GetRenderTargetTexture();
+										FTextureRHIRef BackBuffer = (ViewportRT) ? ViewportRT : RHIGetViewportBackBuffer(ViewInfo->ViewportRHI);
 
 										if (BackBuffer)
 										{
diff --git a/Engine/Source/Runtime/SlateRHIRenderer/Private/SlateRHIRenderer.h b/Engine/Source/Runtime/SlateRHIRenderer/Private/SlateRHIRenderer.h
--- a/Engine/Source/Runtime/SlateRHIRenderer/Private/SlateRHIRenderer.h
+++ b/Engine/Source/Runtime/SlateRHIRenderer/Private/SlateRHIRenderer.h
@@ -43,7 +43,7 @@ extern TAutoConsoleVariable<int32> CVarShowSlateBatching;
 class FSlateBackBuffer : public FRenderTarget
 {
 public:
-	FSlateBackBuffer(FTexture2DRHIRef InRenderTargetTexture, FIntPoint InSizeXY)
+	FSlateBackBuffer(FTextureRHIRef InRenderTargetTexture, FIntPoint InSizeXY)
 		: SizeXY(InSizeXY)
 	{
 		RenderTargetTextureRHI = InRenderTargetTexture;
@@ -62,7 +62,7 @@ struct FViewportInfo : public FRenderResource
 	/** The viewport rendering handle */
 	FViewportRHIRef ViewportRHI;
 	/** The depth buffer texture if any */
-	FTexture2DRHIRef DepthStencil;
+	FTextureRHIRef DepthStencil;
 
 	// Buffers used in HDR compositing
 	/** sRGB UI render target */
@@ -74,7 +74,7 @@ struct FViewportInfo : public FRenderResource
 	int32 ColorSpaceLUTOutputDevice;
 	int32 ColorSpaceLUTOutputGamut;
 		
-	//FTexture2DRHIRef RenderTargetTexture;
+	//FTextureRHIRef RenderTargetTexture;
 	/** The OS Window handle (for recreating the viewport) */
 	void* OSWindow;
 	/** The actual width of the viewport */
@@ -140,7 +140,7 @@ struct FViewportInfo : public FRenderResource
 	void ConditionallyUpdateDepthBuffer(bool bInRequiresStencilTest, uint32 Width, uint32 Height);
 	void RecreateDepthBuffer_RenderThread();
 
-	FTexture2DRHIRef GetRenderTargetTexture() const
+	FTextureRHIRef GetRenderTargetTexture() const
 	{
 		if (RTProvider)
 		{
@@ -155,7 +155,7 @@ struct FViewportInfo : public FRenderResource
 	}	
 };
 
-bool IsMemorylessTexture(const FTexture2DRHIRef& Tex);
+bool IsMemorylessTexture(const FTextureRHIRef& Tex);
 
 struct FFastPathRenderingDataCleanupList
 {
diff --git a/Engine/Source/Runtime/SlateRHIRenderer/Private/SlateRHIRenderingPolicy.cpp b/Engine/Source/Runtime/SlateRHIRenderer/Private/SlateRHIRenderingPolicy.cpp
--- a/Engine/Source/Runtime/SlateRHIRenderer/Private/SlateRHIRenderingPolicy.cpp
+++ b/Engine/Source/Runtime/SlateRHIRenderer/Private/SlateRHIRenderingPolicy.cpp
@@ -276,7 +276,7 @@ static bool UpdateScissorRect(
 	FSlateBackBuffer& BackBuffer,
 	const FSlateRenderBatch& RenderBatch, 
 	FRHIRenderPassInfo& RPInfo,
-	FTexture2DRHIRef& DepthStencilTarget,
+	FTextureRHIRef& DepthStencilTarget,
 	const FSlateClippingState*& LastClippingState,
 	const FVector2f ViewTranslation2D, 
 	FGraphicsPipelineStateInitializer& InGraphicsPSOInit,
@@ -566,8 +566,8 @@ static bool UpdateScissorRect(
 	uint32& MaskingID,
 	FSlateBackBuffer& BackBuffer,
 	const FSlateRenderBatch& RenderBatch,
-	FTexture2DRHIRef& ColorTarget,
-	FTexture2DRHIRef& DepthStencilTarget,
+	FTextureRHIRef& ColorTarget,
+	FTextureRHIRef& DepthStencilTarget,
 	const FSlateClippingState*& LastClippingState,
 	const FVector2f ViewTranslation2D,
 	FGraphicsPipelineStateInitializer& InGraphicsPSOInit,
@@ -690,9 +690,9 @@ static FRHISamplerState* GetSamplerState(ESlateBatchDrawFlag DrawFlags, ETexture
 void FSlateRHIRenderingPolicy::DrawElements(
 	FRHICommandListImmediate& RHICmdList,
 	FSlateBackBuffer& BackBuffer,
-	FTexture2DRHIRef& ColorTarget,
-	FTexture2DRHIRef& PostProcessTexture,
-	FTexture2DRHIRef& DepthStencilTarget,
+	FTextureRHIRef& ColorTarget,
+	FTextureRHIRef& PostProcessTexture,
+	FTextureRHIRef& DepthStencilTarget,
 	int32 FirstBatchIndex,
 	const TArray<FSlateRenderBatch>& RenderBatches,
 	const FSlateRenderingParams& Params)
@@ -1037,7 +1037,7 @@ void FSlateRHIRenderingPolicy::DrawElements(
 					}
 					else
 					{
-						FRHITexture* NativeTextureRHI = ((TSlateTexture<FTexture2DRHIRef>*)ShaderResource)->GetTypedResource();
+						FRHITexture* NativeTextureRHI = ((TSlateTexture<FTextureRHIRef>*)ShaderResource)->GetTypedResource();
 						// Atlas textures that have no content are never initialized but null textures are invalid on many platforms.
 						TextureRHI = NativeTextureRHI ? NativeTextureRHI : (FRHITexture*)GWhiteTexture->TextureRHI;
 					}
@@ -1199,8 +1199,8 @@ void FSlateRHIRenderingPolicy::DrawElements(
 
 								if (MaskResource)
 								{
-									FTexture2DRHIRef TextureRHI;
-									TextureRHI = ((TSlateTexture<FTexture2DRHIRef>*)MaskResource)->GetTypedResource();
+									FTextureRHIRef TextureRHI;
+									TextureRHI = ((TSlateTexture<FTextureRHIRef>*)MaskResource)->GetTypedResource();
 
 									PixelShader->SetAdditionalTexture(BatchedParameters, TextureRHI, BilinearClamp);
 								}
diff --git a/Engine/Source/Runtime/SlateRHIRenderer/Private/SlateRHIRenderingPolicy.h b/Engine/Source/Runtime/SlateRHIRenderer/Private/SlateRHIRenderingPolicy.h
--- a/Engine/Source/Runtime/SlateRHIRenderer/Private/SlateRHIRenderingPolicy.h
+++ b/Engine/Source/Runtime/SlateRHIRenderer/Private/SlateRHIRenderingPolicy.h
@@ -61,8 +61,8 @@ public:
 
 	void BuildRenderingBuffers(FRHICommandListImmediate& RHICmdList, FSlateBatchData& InBatchData);
 
-	void DrawElements(FRHICommandListImmediate& RHICmdList, class FSlateBackBuffer& BackBuffer, FTexture2DRHIRef& ColorTarget, FTexture2DRHIRef& PostProcessBuffer,
-		FTexture2DRHIRef& DepthStencilTarget, int32 FirstBatchIndex,  const TArray<FSlateRenderBatch>& RenderBatches, const FSlateRenderingParams& Params);
+	void DrawElements(FRHICommandListImmediate& RHICmdList, class FSlateBackBuffer& BackBuffer, FTextureRHIRef& ColorTarget, FTextureRHIRef& PostProcessBuffer,
+		FTextureRHIRef& DepthStencilTarget, int32 FirstBatchIndex,  const TArray<FSlateRenderBatch>& RenderBatches, const FSlateRenderingParams& Params);
 
 	virtual TSharedRef<FSlateShaderResourceManager> GetResourceManager() const override { return ResourceManager; }
 	virtual bool IsVertexColorInLinearSpace() const override { return false; }
diff --git a/Engine/Source/Runtime/SlateRHIRenderer/Public/FX/SlateRHIPostBufferProcessor.h b/Engine/Source/Runtime/SlateRHIRenderer/Public/FX/SlateRHIPostBufferProcessor.h
--- a/Engine/Source/Runtime/SlateRHIRenderer/Public/FX/SlateRHIPostBufferProcessor.h
+++ b/Engine/Source/Runtime/SlateRHIRenderer/Public/FX/SlateRHIPostBufferProcessor.h
@@ -118,7 +118,7 @@ protected:
 	 * @param InElementWindowSize	Size of window being rendered, used to determine if using stereo rendering or not.
 	 * @param InRHICmdList			RHI command list to queue commands on
 	 */
-	static FTexture2DRHIRef GetBackbuffer_RenderThread(FRenderResource* InViewInfo, FRenderResource* InViewportTexture, FVector2D InElementWindowSize, FRHICommandListImmediate& InRHICmdList);
+	static FTextureRHIRef GetBackbuffer_RenderThread(FRenderResource* InViewInfo, FRenderResource* InViewportTexture, FVector2D InElementWindowSize, FRHICommandListImmediate& InRHICmdList);
 
 	/**
 	 * Gets 'Src' texture for post process command. Typically the scenebuffer.
@@ -126,7 +126,7 @@ protected:
 	 * @param InBackBuffer			Backbuffer used in standalone
 	 * @param InViewportTexture		'FSlateRenderTargetRHI' resource used for 'BufferedRT' viewport texture in PIE
 	 */
-	static FTexture2DRHIRef GetSrcTexture_RenderThread(FTexture2DRHIRef InBackBuffer, FRenderResource* InViewportTexture);
+	static FTextureRHIRef GetSrcTexture_RenderThread(FTextureRHIRef InBackBuffer, FRenderResource* InViewportTexture);
 
 	/**
 	 * Gets 'Dst' texture for post process command. Convience method, this should be possible through the direct resource.
@@ -141,5 +141,5 @@ protected:
 	 * @param InBackBuffer			Backbuffer used for size in standalone
 	 * @param InViewportTexture		'FSlateRenderTargetRHI' resource used for 'BufferedRT' viewport texture size in PIE
 	 */
-	static FIntPoint GetDstExtent_RenderThread(FTexture2DRHIRef InBackBuffer, FRenderResource* InViewportTexture);
+	static FIntPoint GetDstExtent_RenderThread(FTextureRHIRef InBackBuffer, FRenderResource* InViewportTexture);
 };
\ No newline at end of file
diff --git a/Engine/Source/Runtime/StreamingPauseRendering/Private/StreamingPauseRendering.cpp b/Engine/Source/Runtime/StreamingPauseRendering/Private/StreamingPauseRendering.cpp
--- a/Engine/Source/Runtime/StreamingPauseRendering/Private/StreamingPauseRendering.cpp
+++ b/Engine/Source/Runtime/StreamingPauseRendering/Private/StreamingPauseRendering.cpp
@@ -22,7 +22,7 @@ IMPLEMENT_MODULE(FStreamingPauseRenderingModule, StreamingPauseRendering);
 class FBackgroundView : public ISlateViewport
 {
 public:
-	FBackgroundView(FTexture2DRHIRef RenderTargetTexture, FIntPoint InSize)
+	FBackgroundView(FTextureRHIRef RenderTargetTexture, FIntPoint InSize)
 		: RenderTarget( new FSlateRenderTargetRHI(RenderTargetTexture, InSize.X, InSize.Y))
 		, Size(InSize)
 	{
diff --git a/Engine/Source/Runtime/VulkanRHI/Private/VulkanRHI.cpp b/Engine/Source/Runtime/VulkanRHI/Private/VulkanRHI.cpp
--- a/Engine/Source/Runtime/VulkanRHI/Private/VulkanRHI.cpp
+++ b/Engine/Source/Runtime/VulkanRHI/Private/VulkanRHI.cpp
@@ -1562,7 +1562,7 @@ uint64 FVulkanDynamicRHI::RHIGetMinimumAlignmentForBufferBackedSRV(EPixelFormat
 	return Limits.minTexelBufferOffsetAlignment;
 }
 
-FTexture2DRHIRef FVulkanDynamicRHI::RHICreateTexture2DFromResource(EPixelFormat Format, uint32 SizeX, uint32 SizeY, uint32 NumMips, uint32 NumSamples, VkImage Resource, ETextureCreateFlags Flags, const FClearValueBinding& ClearValueBinding)
+FTextureRHIRef FVulkanDynamicRHI::RHICreateTexture2DFromResource(EPixelFormat Format, uint32 SizeX, uint32 SizeY, uint32 NumMips, uint32 NumSamples, VkImage Resource, ETextureCreateFlags Flags, const FClearValueBinding& ClearValueBinding)
 {
 	const FRHITextureCreateDesc Desc =
 		FRHITextureCreateDesc::Create2D(TEXT("VulkanTexture2DFromResource"), SizeX, SizeY, Format)
@@ -1575,7 +1575,7 @@ FTexture2DRHIRef FVulkanDynamicRHI::RHICreateTexture2DFromResource(EPixelFormat
 	return new FVulkanTexture(*Device, Desc, Resource, false);
 }
 
-FTexture2DArrayRHIRef FVulkanDynamicRHI::RHICreateTexture2DArrayFromResource(EPixelFormat Format, uint32 SizeX, uint32 SizeY, uint32 ArraySize, uint32 NumMips, uint32 NumSamples, VkImage Resource, ETextureCreateFlags Flags, const FClearValueBinding& ClearValueBinding)
+FTextureRHIRef FVulkanDynamicRHI::RHICreateTexture2DArrayFromResource(EPixelFormat Format, uint32 SizeX, uint32 SizeY, uint32 ArraySize, uint32 NumMips, uint32 NumSamples, VkImage Resource, ETextureCreateFlags Flags, const FClearValueBinding& ClearValueBinding)
 {
 	const FRHITextureCreateDesc Desc =
 		FRHITextureCreateDesc::Create2DArray(TEXT("VulkanTextureArrayFromResource"), SizeX, SizeY, ArraySize, Format)
@@ -1588,7 +1588,7 @@ FTexture2DArrayRHIRef FVulkanDynamicRHI::RHICreateTexture2DArrayFromResource(EPi
 	return new FVulkanTexture(*Device, Desc, Resource, false);
 }
 
-FTextureCubeRHIRef FVulkanDynamicRHI::RHICreateTextureCubeFromResource(EPixelFormat Format, uint32 Size, bool bArray, uint32 ArraySize, uint32 NumMips, VkImage Resource, ETextureCreateFlags Flags, const FClearValueBinding& ClearValueBinding)
+FTextureRHIRef FVulkanDynamicRHI::RHICreateTextureCubeFromResource(EPixelFormat Format, uint32 Size, bool bArray, uint32 ArraySize, uint32 NumMips, VkImage Resource, ETextureCreateFlags Flags, const FClearValueBinding& ClearValueBinding)
 {
 	const FRHITextureCreateDesc Desc =
 		FRHITextureCreateDesc::Create(TEXT("VulkanTextureCubeFromResource"), ArraySize > 1 ? ETextureDimension::TextureCubeArray : ETextureDimension::TextureCube)
diff --git a/Engine/Source/Runtime/VulkanRHI/Private/VulkanTexture.cpp b/Engine/Source/Runtime/VulkanRHI/Private/VulkanTexture.cpp
--- a/Engine/Source/Runtime/VulkanRHI/Private/VulkanTexture.cpp
+++ b/Engine/Source/Runtime/VulkanRHI/Private/VulkanTexture.cpp
@@ -1124,7 +1124,7 @@ struct FRHICommandVulkanAsyncReallocateTexture2D final : public FRHICommand<FRHI
 	}
 };
 
-FTexture2DRHIRef FVulkanDynamicRHI::AsyncReallocateTexture2D_RenderThread(FRHICommandListImmediate& RHICmdList, FRHITexture2D* OldTextureRHI, int32 NewMipCount, int32 NewSizeX, int32 NewSizeY, FThreadSafeCounter* RequestStatus)
+FTextureRHIRef FVulkanDynamicRHI::AsyncReallocateTexture2D_RenderThread(FRHICommandListImmediate& RHICmdList, FRHITexture* OldTextureRHI, int32 NewMipCount, int32 NewSizeX, int32 NewSizeY, FThreadSafeCounter* RequestStatus)
 {
 	LLM_SCOPE_VULKAN(ELLMTagVulkan::VulkanTextures);
 	if (RHICmdList.Bypass())
@@ -1149,7 +1149,7 @@ FTexture2DRHIRef FVulkanDynamicRHI::AsyncReallocateTexture2D_RenderThread(FRHICo
 	return NewTexture;
 }
 
-FTexture2DRHIRef FVulkanDynamicRHI::RHIAsyncReallocateTexture2D(FRHITexture2D* OldTextureRHI, int32 NewMipCount, int32 NewSizeX, int32 NewSizeY, FThreadSafeCounter* RequestStatus)
+FTextureRHIRef FVulkanDynamicRHI::RHIAsyncReallocateTexture2D(FRHITexture* OldTextureRHI, int32 NewMipCount, int32 NewSizeX, int32 NewSizeY, FThreadSafeCounter* RequestStatus)
 {
 	LLM_SCOPE_VULKAN(ELLMTagVulkan::VulkanTextures);
 	FVulkanTexture* OldTexture = ResourceCast(OldTextureRHI);
@@ -1170,17 +1170,17 @@ FTexture2DRHIRef FVulkanDynamicRHI::RHIAsyncReallocateTexture2D(FRHITexture2D* O
 	return NewTexture;
 }
 
-ETextureReallocationStatus FVulkanDynamicRHI::RHIFinalizeAsyncReallocateTexture2D(FRHITexture2D* Texture2D, bool bBlockUntilCompleted)
+ETextureReallocationStatus FVulkanDynamicRHI::RHIFinalizeAsyncReallocateTexture2D(FRHITexture* Texture2D, bool bBlockUntilCompleted)
 {
 	return TexRealloc_Succeeded;
 }
 
-ETextureReallocationStatus FVulkanDynamicRHI::RHICancelAsyncReallocateTexture2D(FRHITexture2D* Texture2D, bool bBlockUntilCompleted)
+ETextureReallocationStatus FVulkanDynamicRHI::RHICancelAsyncReallocateTexture2D(FRHITexture* Texture2D, bool bBlockUntilCompleted)
 {
 	return TexRealloc_Succeeded;
 }
 
-void* FVulkanDynamicRHI::RHILockTexture2D(FRHITexture2D* TextureRHI, uint32 MipIndex, EResourceLockMode LockMode, uint32& DestStride, bool bLockWithinMiptail, uint64* OutLockedByteCount)
+void* FVulkanDynamicRHI::RHILockTexture2D(FRHITexture* TextureRHI, uint32 MipIndex, EResourceLockMode LockMode, uint32& DestStride, bool bLockWithinMiptail, uint64* OutLockedByteCount)
 {
 	LLM_SCOPE_VULKAN(ELLMTagVulkan::VulkanTextures);
 	FVulkanTexture* Texture = ResourceCast(TextureRHI);
@@ -1211,7 +1211,7 @@ void* FVulkanDynamicRHI::RHILockTexture2D(FRHITexture2D* TextureRHI, uint32 MipI
 	return Data;
 }
 
-void FVulkanDynamicRHI::InternalUnlockTexture2D(bool bFromRenderingThread, FRHITexture2D* TextureRHI, uint32 MipIndex, bool bLockWithinMiptail)
+void FVulkanDynamicRHI::InternalUnlockTexture2D(bool bFromRenderingThread, FRHITexture* TextureRHI, uint32 MipIndex, bool bLockWithinMiptail)
 {
 	LLM_SCOPE_VULKAN(ELLMTagVulkan::VulkanTextures);
 	FVulkanTexture* Texture = ResourceCast(TextureRHI);
@@ -1261,7 +1261,7 @@ void FVulkanDynamicRHI::InternalUnlockTexture2D(bool bFromRenderingThread, FRHIT
 	}
 }
 
-void* FVulkanDynamicRHI::RHILockTexture2DArray(FRHITexture2DArray* TextureRHI,uint32 TextureIndex,uint32 MipIndex,EResourceLockMode LockMode,uint32& DestStride,bool bLockWithinMiptail)
+void* FVulkanDynamicRHI::RHILockTexture2DArray(FRHITexture* TextureRHI,uint32 TextureIndex,uint32 MipIndex,EResourceLockMode LockMode,uint32& DestStride,bool bLockWithinMiptail)
 {
 	LLM_SCOPE_VULKAN(ELLMTagVulkan::VulkanTextures);
 	FVulkanTexture* Texture = ResourceCast(TextureRHI);
@@ -1287,7 +1287,7 @@ void* FVulkanDynamicRHI::RHILockTexture2DArray(FRHITexture2DArray* TextureRHI,ui
 	return Data;
 }
 
-void FVulkanDynamicRHI::RHIUnlockTexture2DArray(FRHITexture2DArray* TextureRHI, uint32 TextureIndex, uint32 MipIndex, bool bLockWithinMiptail)
+void FVulkanDynamicRHI::RHIUnlockTexture2DArray(FRHITexture* TextureRHI, uint32 TextureIndex, uint32 MipIndex, bool bLockWithinMiptail)
 {
 	LLM_SCOPE_VULKAN(ELLMTagVulkan::VulkanTextures);
 	FVulkanTexture* Texture = ResourceCast(TextureRHI);
@@ -1333,7 +1333,7 @@ void FVulkanDynamicRHI::RHIUnlockTexture2DArray(FRHITexture2DArray* TextureRHI,
 	}
 }
 
-void FVulkanDynamicRHI::InternalUpdateTexture2D(FRHICommandListBase& RHICmdList, FRHITexture2D* TextureRHI, uint32 MipIndex, const FUpdateTextureRegion2D& UpdateRegion, uint32 SourcePitch, const uint8* SourceData)
+void FVulkanDynamicRHI::InternalUpdateTexture2D(FRHICommandListBase& RHICmdList, FRHITexture* TextureRHI, uint32 MipIndex, const FUpdateTextureRegion2D& UpdateRegion, uint32 SourcePitch, const uint8* SourceData)
 {
 	LLM_SCOPE_VULKAN(ELLMTagVulkan::VulkanTextures);
 
@@ -1391,7 +1391,7 @@ void FVulkanDynamicRHI::InternalUpdateTexture2D(FRHICommandListBase& RHICmdList,
 	}
 }
 
-FUpdateTexture3DData FVulkanDynamicRHI::RHIBeginUpdateTexture3D(FRHICommandListBase& RHICmdList, FRHITexture3D* Texture, uint32 MipIndex, const struct FUpdateTextureRegion3D& UpdateRegion)
+FUpdateTexture3DData FVulkanDynamicRHI::RHIBeginUpdateTexture3D(FRHICommandListBase& RHICmdList, FRHITexture* Texture, uint32 MipIndex, const struct FUpdateTextureRegion3D& UpdateRegion)
 {
 	const int32 FormatSize = PixelFormatBlockBytes[Texture->GetFormat()];
 	const int32 RowPitch = UpdateRegion.Width * FormatSize;
@@ -1414,7 +1414,7 @@ void FVulkanDynamicRHI::RHIEndUpdateTexture3D(FRHICommandListBase& RHICmdList, F
 	UpdateData.Data = nullptr;
 }
 
-void FVulkanDynamicRHI::InternalUpdateTexture3D(FRHICommandListBase& RHICmdList, FRHITexture3D* TextureRHI, uint32 MipIndex, const FUpdateTextureRegion3D& UpdateRegion, uint32 SourceRowPitch, uint32 SourceDepthPitch, const uint8* SourceData)
+void FVulkanDynamicRHI::InternalUpdateTexture3D(FRHICommandListBase& RHICmdList, FRHITexture* TextureRHI, uint32 MipIndex, const FUpdateTextureRegion3D& UpdateRegion, uint32 SourceRowPitch, uint32 SourceDepthPitch, const uint8* SourceData)
 {
 	LLM_SCOPE_VULKAN(ELLMTagVulkan::VulkanTextures);
 	FVulkanTexture* Texture = ResourceCast(TextureRHI);
@@ -2032,7 +2032,7 @@ bool FVulkanTexture::GetTextureResourceInfo(FRHIResourceInfo& OutResourceInfo) c
 	Cubemap texture support.
 -----------------------------------------------------------------------------*/
 
-void* FVulkanDynamicRHI::RHILockTextureCubeFace(FRHITextureCube* TextureCubeRHI, uint32 FaceIndex, uint32 ArrayIndex, uint32 MipIndex, EResourceLockMode LockMode, uint32& DestStride, bool bLockWithinMiptail)
+void* FVulkanDynamicRHI::RHILockTextureCubeFace(FRHITexture* TextureCubeRHI, uint32 FaceIndex, uint32 ArrayIndex, uint32 MipIndex, EResourceLockMode LockMode, uint32& DestStride, bool bLockWithinMiptail)
 {
 	LLM_SCOPE_VULKAN(ELLMTagVulkan::VulkanTextures);
 	FVulkanTexture* Texture = ResourceCast(TextureCubeRHI);
@@ -2055,7 +2055,7 @@ void* FVulkanDynamicRHI::RHILockTextureCubeFace(FRHITextureCube* TextureCubeRHI,
 	return Data;
 }
 
-void FVulkanDynamicRHI::RHIUnlockTextureCubeFace(FRHITextureCube* TextureCubeRHI, uint32 FaceIndex, uint32 ArrayIndex, uint32 MipIndex, bool bLockWithinMiptail)
+void FVulkanDynamicRHI::RHIUnlockTextureCubeFace(FRHITexture* TextureCubeRHI, uint32 FaceIndex, uint32 ArrayIndex, uint32 MipIndex, bool bLockWithinMiptail)
 {
 	LLM_SCOPE_VULKAN(ELLMTagVulkan::VulkanTextures);
 	FVulkanTexture* Texture = ResourceCast(TextureCubeRHI);
@@ -2141,7 +2141,7 @@ void FVulkanDynamicRHI::RHIBindDebugLabelName(FRHICommandListBase& RHICmdList, F
 #if VULKAN_ENABLE_DUMP_LAYER
 	//if (Device->SupportsDebugMarkers())
 	{
-		//if (FRHITexture2D* Tex2d = UnorderedAccessViewRHI->GetTexture2D())
+		//if (FRHITexture* Tex2d = UnorderedAccessViewRHI->GetTexture2D())
 		//{
 		//	FVulkanTexture2D* VulkanTexture = (FVulkanTexture2D*)Tex2d;
 		//	VkDebugMarkerObjectTagInfoEXT Info;
diff --git a/Engine/Source/Runtime/VulkanRHI/Private/VulkanViewport.cpp b/Engine/Source/Runtime/VulkanRHI/Private/VulkanViewport.cpp
--- a/Engine/Source/Runtime/VulkanRHI/Private/VulkanViewport.cpp
+++ b/Engine/Source/Runtime/VulkanRHI/Private/VulkanViewport.cpp
@@ -223,7 +223,7 @@ bool FVulkanViewport::TryAcquireImageIndex()
 	return false;
 }
 
-FTexture2DRHIRef FVulkanViewport::GetBackBuffer(FRHICommandListImmediate& RHICmdList)
+FTextureRHIRef FVulkanViewport::GetBackBuffer(FRHICommandListImmediate& RHICmdList)
 {
 	check(IsInRenderingThread());
 
@@ -1178,7 +1178,7 @@ void FVulkanDynamicRHI::RHITick(float DeltaTime)
 	});
 }
 
-FTexture2DRHIRef FVulkanDynamicRHI::RHIGetViewportBackBuffer(FRHIViewport* ViewportRHI)
+FTextureRHIRef FVulkanDynamicRHI::RHIGetViewportBackBuffer(FRHIViewport* ViewportRHI)
 {
 	FRHICommandListImmediate& RHICmdList = FRHICommandListImmediate::Get();
 
diff --git a/Engine/Source/Runtime/VulkanRHI/Private/VulkanViewport.h b/Engine/Source/Runtime/VulkanRHI/Private/VulkanViewport.h
--- a/Engine/Source/Runtime/VulkanRHI/Private/VulkanViewport.h
+++ b/Engine/Source/Runtime/VulkanRHI/Private/VulkanViewport.h
@@ -49,7 +49,7 @@ public:
 	FVulkanViewport(FVulkanDevice* InDevice, void* InWindowHandle, uint32 InSizeX,uint32 InSizeY,bool bInIsFullscreen, EPixelFormat InPreferredPixelFormat);
 	~FVulkanViewport();
 
-	FTexture2DRHIRef GetBackBuffer(FRHICommandListImmediate& RHICmdList);
+	FTextureRHIRef GetBackBuffer(FRHICommandListImmediate& RHICmdList);
 	void AdvanceBackBufferFrame(FRHICommandListImmediate& RHICmdList);
 
 	virtual void WaitForFrameEventCompletion() override;
diff --git a/Engine/Source/Runtime/VulkanRHI/Public/IVulkanDynamicRHI.h b/Engine/Source/Runtime/VulkanRHI/Public/IVulkanDynamicRHI.h
--- a/Engine/Source/Runtime/VulkanRHI/Public/IVulkanDynamicRHI.h
+++ b/Engine/Source/Runtime/VulkanRHI/Public/IVulkanDynamicRHI.h
@@ -51,9 +51,9 @@ struct IVulkanDynamicRHI : public FDynamicRHI
 	virtual TArray<VkExtensionProperties> RHIGetAllInstanceExtensions() const = 0;
 	virtual TArray<VkExtensionProperties> RHIGetAllDeviceExtensions(VkPhysicalDevice InPhysicalDevice) const = 0;
 
-	virtual FTexture2DRHIRef      RHICreateTexture2DFromResource(EPixelFormat Format, uint32 SizeX, uint32 SizeY, uint32 NumMips, uint32 NumSamples, VkImage Resource, ETextureCreateFlags Flags, const FClearValueBinding& ClearValueBinding = FClearValueBinding::Transparent) = 0;
-	virtual FTexture2DArrayRHIRef RHICreateTexture2DArrayFromResource(EPixelFormat Format, uint32 SizeX, uint32 SizeY, uint32 ArraySize, uint32 NumMips, uint32 NumSamples, VkImage Resource, ETextureCreateFlags Flags, const FClearValueBinding& ClearValueBinding = FClearValueBinding::Transparent) = 0;
-	virtual FTextureCubeRHIRef    RHICreateTextureCubeFromResource(EPixelFormat Format, uint32 Size, bool bArray, uint32 ArraySize, uint32 NumMips, VkImage Resource, ETextureCreateFlags Flags, const FClearValueBinding& ClearValueBinding = FClearValueBinding::Transparent) = 0;
+	virtual FTextureRHIRef   RHICreateTexture2DFromResource(EPixelFormat Format, uint32 SizeX, uint32 SizeY, uint32 NumMips, uint32 NumSamples, VkImage Resource, ETextureCreateFlags Flags, const FClearValueBinding& ClearValueBinding = FClearValueBinding::Transparent) = 0;
+	virtual FTextureRHIRef   RHICreateTexture2DArrayFromResource(EPixelFormat Format, uint32 SizeX, uint32 SizeY, uint32 ArraySize, uint32 NumMips, uint32 NumSamples, VkImage Resource, ETextureCreateFlags Flags, const FClearValueBinding& ClearValueBinding = FClearValueBinding::Transparent) = 0;
+	virtual FTextureRHIRef   RHICreateTextureCubeFromResource(EPixelFormat Format, uint32 Size, bool bArray, uint32 ArraySize, uint32 NumMips, VkImage Resource, ETextureCreateFlags Flags, const FClearValueBinding& ClearValueBinding = FClearValueBinding::Transparent) = 0;
 
 	virtual VkImage                  RHIGetVkImage(FRHITexture* InTexture) const = 0;
 	virtual VkFormat                 RHIGetViewVkFormat(FRHITexture* InTexture) const = 0;
diff --git a/Engine/Source/Runtime/VulkanRHI/Public/VulkanDynamicRHI.h b/Engine/Source/Runtime/VulkanRHI/Public/VulkanDynamicRHI.h
--- a/Engine/Source/Runtime/VulkanRHI/Public/VulkanDynamicRHI.h
+++ b/Engine/Source/Runtime/VulkanRHI/Public/VulkanDynamicRHI.h
@@ -73,9 +73,9 @@ public:
 	virtual bool RHISupportsEXTFragmentDensityMap2() const final override;
 	virtual TArray<VkExtensionProperties> RHIGetAllInstanceExtensions() const final override;
 	virtual TArray<VkExtensionProperties> RHIGetAllDeviceExtensions(VkPhysicalDevice InPhysicalDevice) const final override;
-	virtual FTexture2DRHIRef RHICreateTexture2DFromResource(EPixelFormat Format, uint32 SizeX, uint32 SizeY, uint32 NumMips, uint32 NumSamples, VkImage Resource, ETextureCreateFlags Flags, const FClearValueBinding& ClearValueBinding = FClearValueBinding::Transparent) final override;
-	virtual FTexture2DArrayRHIRef RHICreateTexture2DArrayFromResource(EPixelFormat Format, uint32 SizeX, uint32 SizeY, uint32 ArraySize, uint32 NumMips, uint32 NumSamples, VkImage Resource, ETextureCreateFlags Flags, const FClearValueBinding& ClearValueBinding = FClearValueBinding::Transparent) final override;
-	virtual FTextureCubeRHIRef RHICreateTextureCubeFromResource(EPixelFormat Format, uint32 Size, bool bArray, uint32 ArraySize, uint32 NumMips, VkImage Resource, ETextureCreateFlags Flags, const FClearValueBinding& ClearValueBinding = FClearValueBinding::Transparent) final override;
+	virtual FTextureRHIRef RHICreateTexture2DFromResource(EPixelFormat Format, uint32 SizeX, uint32 SizeY, uint32 NumMips, uint32 NumSamples, VkImage Resource, ETextureCreateFlags Flags, const FClearValueBinding& ClearValueBinding = FClearValueBinding::Transparent) final override;
+	virtual FTextureRHIRef RHICreateTexture2DArrayFromResource(EPixelFormat Format, uint32 SizeX, uint32 SizeY, uint32 ArraySize, uint32 NumMips, uint32 NumSamples, VkImage Resource, ETextureCreateFlags Flags, const FClearValueBinding& ClearValueBinding = FClearValueBinding::Transparent) final override;
+	virtual FTextureRHIRef RHICreateTextureCubeFromResource(EPixelFormat Format, uint32 Size, bool bArray, uint32 ArraySize, uint32 NumMips, VkImage Resource, ETextureCreateFlags Flags, const FClearValueBinding& ClearValueBinding = FClearValueBinding::Transparent) final override;
 	virtual VkImage RHIGetVkImage(FRHITexture* InTexture) const final override;
 	virtual VkFormat RHIGetViewVkFormat(FRHITexture* InTexture) const final override;
 	virtual FVulkanRHIAllocationInfo RHIGetAllocationInfo(FRHITexture* InTexture) const final override;
@@ -127,31 +127,31 @@ public:
 	virtual FTextureRHIRef RHICreateTexture(FRHICommandListBase& RHICmdList, const FRHITextureCreateDesc& CreateDesc) final override;
 	virtual FTextureRHIRef RHIAsyncCreateTexture2D(uint32 SizeX, uint32 SizeY, uint8 Format, uint32 NumMips, ETextureCreateFlags Flags, ERHIAccess InResourceState, void** InitialMipData, uint32 NumInitialMips, const TCHAR* DebugName, FGraphEventRef& OutCompletionEvent) final override;
 	virtual uint32 RHIComputeMemorySize(FRHITexture* TextureRHI) final override;
-	virtual FTexture2DRHIRef RHIAsyncReallocateTexture2D(FRHITexture2D* Texture2D, int32 NewMipCount, int32 NewSizeX, int32 NewSizeY, FThreadSafeCounter* RequestStatus) final override;
-	virtual ETextureReallocationStatus RHIFinalizeAsyncReallocateTexture2D(FRHITexture2D* Texture2D, bool bBlockUntilCompleted) final override;
-	virtual ETextureReallocationStatus FinalizeAsyncReallocateTexture2D_RenderThread(class FRHICommandListImmediate& RHICmdList, FRHITexture2D* Texture2D, bool bBlockUntilCompleted) final override
+	virtual FTextureRHIRef RHIAsyncReallocateTexture2D(FRHITexture* Texture2D, int32 NewMipCount, int32 NewSizeX, int32 NewSizeY, FThreadSafeCounter* RequestStatus) final override;
+	virtual ETextureReallocationStatus RHIFinalizeAsyncReallocateTexture2D(FRHITexture* Texture2D, bool bBlockUntilCompleted) final override;
+	virtual ETextureReallocationStatus FinalizeAsyncReallocateTexture2D_RenderThread(class FRHICommandListImmediate& RHICmdList, FRHITexture* Texture2D, bool bBlockUntilCompleted) final override
 	{
 		return this->RHIFinalizeAsyncReallocateTexture2D(Texture2D, bBlockUntilCompleted);
 	}
 
-	virtual ETextureReallocationStatus RHICancelAsyncReallocateTexture2D(FRHITexture2D* Texture2D, bool bBlockUntilCompleted) final override;
-	virtual void* RHILockTexture2D(FRHITexture2D* Texture, uint32 MipIndex, EResourceLockMode LockMode, uint32& DestStride, bool bLockWithinMiptail, uint64* OutLockedByteCount = nullptr) final override;
-	virtual void RHIUnlockTexture2D(FRHITexture2D* Texture, uint32 MipIndex, bool bLockWithinMiptail) final override
+	virtual ETextureReallocationStatus RHICancelAsyncReallocateTexture2D(FRHITexture* Texture2D, bool bBlockUntilCompleted) final override;
+	virtual void* RHILockTexture2D(FRHITexture* Texture, uint32 MipIndex, EResourceLockMode LockMode, uint32& DestStride, bool bLockWithinMiptail, uint64* OutLockedByteCount = nullptr) final override;
+	virtual void RHIUnlockTexture2D(FRHITexture* Texture, uint32 MipIndex, bool bLockWithinMiptail) final override
 	{
 		InternalUnlockTexture2D(false, Texture, MipIndex, bLockWithinMiptail);
 	}
-	virtual void* RHILockTexture2DArray(FRHITexture2DArray* Texture, uint32 TextureIndex, uint32 MipIndex, EResourceLockMode LockMode, uint32& DestStride, bool bLockWithinMiptail) final override;
-	virtual void RHIUnlockTexture2DArray(FRHITexture2DArray* Texture, uint32 TextureIndex, uint32 MipIndex, bool bLockWithinMiptail) final override;
-	virtual void RHIUpdateTexture2D(FRHICommandListBase& RHICmdList, FRHITexture2D* Texture, uint32 MipIndex, const struct FUpdateTextureRegion2D& UpdateRegion, uint32 SourcePitch, const uint8* SourceData) final override
+	virtual void* RHILockTexture2DArray(FRHITexture* Texture, uint32 TextureIndex, uint32 MipIndex, EResourceLockMode LockMode, uint32& DestStride, bool bLockWithinMiptail) final override;
+	virtual void RHIUnlockTexture2DArray(FRHITexture* Texture, uint32 TextureIndex, uint32 MipIndex, bool bLockWithinMiptail) final override;
+	virtual void RHIUpdateTexture2D(FRHICommandListBase& RHICmdList, FRHITexture* Texture, uint32 MipIndex, const struct FUpdateTextureRegion2D& UpdateRegion, uint32 SourcePitch, const uint8* SourceData) final override
 	{
 		InternalUpdateTexture2D(RHICmdList, Texture, MipIndex, UpdateRegion, SourcePitch, SourceData);
 	}
-	virtual void RHIUpdateTexture3D(FRHICommandListBase& RHICmdList, FRHITexture3D* Texture, uint32 MipIndex, const struct FUpdateTextureRegion3D& UpdateRegion, uint32 SourceRowPitch, uint32 SourceDepthPitch, const uint8* SourceData) final override
+	virtual void RHIUpdateTexture3D(FRHICommandListBase& RHICmdList, FRHITexture* Texture, uint32 MipIndex, const struct FUpdateTextureRegion3D& UpdateRegion, uint32 SourceRowPitch, uint32 SourceDepthPitch, const uint8* SourceData) final override
 	{
 		InternalUpdateTexture3D(RHICmdList, Texture, MipIndex, UpdateRegion, SourceRowPitch, SourceDepthPitch, SourceData);
 	}
-	virtual void* RHILockTextureCubeFace(FRHITextureCube* Texture, uint32 FaceIndex, uint32 ArrayIndex, uint32 MipIndex, EResourceLockMode LockMode, uint32& DestStride, bool bLockWithinMiptail) final override;
-	virtual void RHIUnlockTextureCubeFace(FRHITextureCube* Texture, uint32 FaceIndex, uint32 ArrayIndex, uint32 MipIndex, bool bLockWithinMiptail) final override;
+	virtual void* RHILockTextureCubeFace(FRHITexture* Texture, uint32 FaceIndex, uint32 ArrayIndex, uint32 MipIndex, EResourceLockMode LockMode, uint32& DestStride, bool bLockWithinMiptail) final override;
+	virtual void RHIUnlockTextureCubeFace(FRHITexture* Texture, uint32 FaceIndex, uint32 ArrayIndex, uint32 MipIndex, bool bLockWithinMiptail) final override;
 	virtual void RHIBindDebugLabelName(FRHICommandListBase& RHICmdList, FRHITexture* Texture, const TCHAR* Name) final override;
 	virtual void RHIBindDebugLabelName(FRHICommandListBase& RHICmdList, FRHIUnorderedAccessView* UnorderedAccessViewRHI, const TCHAR* Name) final override;
 	virtual void RHIReadSurfaceData(FRHITexture* Texture, FIntRect Rect, TArray<FColor>& OutData, FReadSurfaceDataFlags InFlags) final override;
@@ -163,7 +163,7 @@ public:
 	virtual void RHIRead3DSurfaceFloatData(FRHITexture* Texture, FIntRect Rect, FIntPoint ZMinMax, TArray<FFloat16Color>& OutData) final override;
 	virtual FRenderQueryRHIRef RHICreateRenderQuery(ERenderQueryType QueryType) final override;
 	virtual bool RHIGetRenderQueryResult(FRHIRenderQuery* RenderQuery, uint64& OutResult, bool bWait, uint32 GPUIndex = INDEX_NONE) final override;
-	virtual FTexture2DRHIRef RHIGetViewportBackBuffer(FRHIViewport* Viewport) final override;
+	virtual FTextureRHIRef RHIGetViewportBackBuffer(FRHIViewport* Viewport) final override;
 	virtual void RHIAliasTextureResources(FTextureRHIRef& DestTexture, FTextureRHIRef& SrcTexture) final override;
 	virtual FTextureRHIRef RHICreateAliasedTexture(FTextureRHIRef& SourceTexture) final override;
 	virtual void RHIAdvanceFrameForGetViewportBackBuffer(FRHIViewport* Viewport) final override;
@@ -200,19 +200,19 @@ public:
 	virtual uint64 RHIComputeStatePrecachePSOHash(const FGraphicsPipelineStateInitializer& Initializer) final override;
 	virtual bool RHIMatchPrecachePSOInitializers(const FGraphicsPipelineStateInitializer& LHS, const FGraphicsPipelineStateInitializer& RHS) final override;
 
-	virtual FTexture2DRHIRef AsyncReallocateTexture2D_RenderThread(class FRHICommandListImmediate& RHICmdList, FRHITexture2D* Texture2D, int32 NewMipCount, int32 NewSizeX, int32 NewSizeY, FThreadSafeCounter* RequestStatus) override final;
-	//virtual ETextureReallocationStatus CancelAsyncReallocateTexture2D_RenderThread(class FRHICommandListImmediate& RHICmdList, FRHITexture2D* Texture2D, bool bBlockUntilCompleted) override final;
+	virtual FTextureRHIRef AsyncReallocateTexture2D_RenderThread(class FRHICommandListImmediate& RHICmdList, FRHITexture* Texture2D, int32 NewMipCount, int32 NewSizeX, int32 NewSizeY, FThreadSafeCounter* RequestStatus) override final;
+	//virtual ETextureReallocationStatus CancelAsyncReallocateTexture2D_RenderThread(class FRHICommandListImmediate& RHICmdList, FRHITexture* Texture2D, bool bBlockUntilCompleted) override final;
 	
-	virtual void* LockTexture2D_RenderThread(class FRHICommandListImmediate& RHICmdList, FRHITexture2D* Texture, uint32 MipIndex, EResourceLockMode LockMode, uint32& DestStride, bool bLockWithinMiptail, bool bNeedsDefaultRHIFlush = true, uint64* OutLockedByteCount = nullptr) override final
+	virtual void* LockTexture2D_RenderThread(class FRHICommandListImmediate& RHICmdList, FRHITexture* Texture, uint32 MipIndex, EResourceLockMode LockMode, uint32& DestStride, bool bLockWithinMiptail, bool bNeedsDefaultRHIFlush = true, uint64* OutLockedByteCount = nullptr) override final
 	{
 		return RHILockTexture2D(Texture, MipIndex, LockMode, DestStride, bLockWithinMiptail, OutLockedByteCount);
 	}
-	virtual void UnlockTexture2D_RenderThread(class FRHICommandListImmediate& RHICmdList, FRHITexture2D* Texture, uint32 MipIndex, bool bLockWithinMiptail, bool bNeedsDefaultRHIFlush = true) override final
+	virtual void UnlockTexture2D_RenderThread(class FRHICommandListImmediate& RHICmdList, FRHITexture* Texture, uint32 MipIndex, bool bLockWithinMiptail, bool bNeedsDefaultRHIFlush = true) override final
 	{
 		InternalUnlockTexture2D(true, Texture, MipIndex, bLockWithinMiptail);
 	}
 	
-	virtual FUpdateTexture3DData RHIBeginUpdateTexture3D(FRHICommandListBase& RHICmdList, FRHITexture3D* Texture, uint32 MipIndex, const struct FUpdateTextureRegion3D& UpdateRegion) override final;
+	virtual FUpdateTexture3DData RHIBeginUpdateTexture3D(FRHICommandListBase& RHICmdList, FRHITexture* Texture, uint32 MipIndex, const struct FUpdateTextureRegion3D& UpdateRegion) override final;
 	virtual void RHIEndUpdateTexture3D(FRHICommandListBase& RHICmdList, FUpdateTexture3DData& UpdateData) override final;
 
 	// SRV / UAV creation functions
@@ -366,9 +366,9 @@ protected:
 
 	FCriticalSection LockBufferCS;
 
-	void InternalUnlockTexture2D(bool bFromRenderingThread, FRHITexture2D* Texture, uint32 MipIndex, bool bLockWithinMiptail);
-	void InternalUpdateTexture2D(FRHICommandListBase& RHICmdList, FRHITexture2D* TextureRHI, uint32 MipIndex, const struct FUpdateTextureRegion2D& UpdateRegion, uint32 SourcePitch, const uint8* SourceData);
-	void InternalUpdateTexture3D(FRHICommandListBase& RHICmdList, FRHITexture3D* TextureRHI, uint32 MipIndex, const struct FUpdateTextureRegion3D& UpdateRegion, uint32 SourceRowPitch, uint32 SourceDepthPitch, const uint8* SourceData);
+	void InternalUnlockTexture2D(bool bFromRenderingThread, FRHITexture* Texture, uint32 MipIndex, bool bLockWithinMiptail);
+	void InternalUpdateTexture2D(FRHICommandListBase& RHICmdList, FRHITexture* TextureRHI, uint32 MipIndex, const struct FUpdateTextureRegion2D& UpdateRegion, uint32 SourcePitch, const uint8* SourceData);
+	void InternalUpdateTexture3D(FRHICommandListBase& RHICmdList, FRHITexture* TextureRHI, uint32 MipIndex, const struct FUpdateTextureRegion3D& UpdateRegion, uint32 SourceRowPitch, uint32 SourceDepthPitch, const uint8* SourceData);
 
 	void UpdateUniformBuffer(FRHICommandListBase& RHICmdList, FVulkanUniformBuffer* UniformBuffer, const void* Contents);
 
diff --git a/Engine/Source/Runtime/WebBrowser/Private/Android/AndroidWebBrowserWidget.cpp b/Engine/Source/Runtime/WebBrowser/Private/Android/AndroidWebBrowserWidget.cpp
--- a/Engine/Source/Runtime/WebBrowser/Private/Android/AndroidWebBrowserWidget.cpp
+++ b/Engine/Source/Runtime/WebBrowser/Private/Android/AndroidWebBrowserWidget.cpp
@@ -438,7 +438,7 @@ void SAndroidWebBrowserWidget::Tick(const FGeometry& AllottedGeometry, const dou
 					}
 
 					// write frame into texture
-					FRHITexture2D* Texture = Params.NewTextureSamplePtr->InitializeTexture();
+					FRHITexture* Texture = Params.NewTextureSamplePtr->InitializeTexture();
 
 					if (Texture != nullptr)
 					{
diff --git a/Engine/Source/Runtime/WebBrowserTexture/Private/WebBrowserTextureResource.cpp b/Engine/Source/Runtime/WebBrowserTexture/Private/WebBrowserTextureResource.cpp
--- a/Engine/Source/Runtime/WebBrowserTexture/Private/WebBrowserTextureResource.cpp
+++ b/Engine/Source/Runtime/WebBrowserTexture/Private/WebBrowserTextureResource.cpp
@@ -183,7 +183,7 @@ void FWebBrowserTextureResource::ClearTexture(const FLinearColor& ClearColor)
 void FWebBrowserTextureResource::CopySample(const TSharedPtr<FWebBrowserTextureSample, ESPMode::ThreadSafe>& Sample, const FLinearColor& ClearColor)
 {
 	FRHITexture* SampleTexture = Sample->GetTexture();
-	FRHITexture2D* SampleTexture2D = (SampleTexture != nullptr) ? SampleTexture->GetTexture2D() : nullptr;
+	FRHITexture* SampleTexture2D = (SampleTexture != nullptr) ? SampleTexture->GetTexture2D() : nullptr;
 	// If the sample already provides a texture resource, we simply use that
 	// as the output render target. If the sample only provides raw data, then
 	// we create our own output render target and copy the data into it.
@@ -260,7 +260,7 @@ void FWebBrowserTextureResource::UpdateResourceSize()
 }
 
 
-void FWebBrowserTextureResource::UpdateTextureReference(FRHITexture2D* NewTexture)
+void FWebBrowserTextureResource::UpdateTextureReference(FRHITexture* NewTexture)
 {
 	TextureRHI = NewTexture;
 	RenderTargetTextureRHI = NewTexture;
diff --git a/Engine/Source/Runtime/WebBrowserTexture/Public/WebBrowserTextureResource.h b/Engine/Source/Runtime/WebBrowserTexture/Public/WebBrowserTextureResource.h
--- a/Engine/Source/Runtime/WebBrowserTexture/Public/WebBrowserTextureResource.h
+++ b/Engine/Source/Runtime/WebBrowserTexture/Public/WebBrowserTextureResource.h
@@ -107,7 +107,7 @@ protected:
 	 *
 	 * @param NewTexture The texture to set.
 	 */
-	void UpdateTextureReference(FRHITexture2D* NewTexture);
+	void UpdateTextureReference(FRHITexture* NewTexture);
 
 private:
 
@@ -118,10 +118,10 @@ private:
 	FLinearColor CurrentClearColor;
 
 	/** Input render target if the texture samples don't provide one (for conversions). */
-	TRefCountPtr<FRHITexture2D> InputTarget;
+	TRefCountPtr<FRHITexture> InputTarget;
 
 	/** Output render target if the texture samples don't provide one. */
-	TRefCountPtr<FRHITexture2D> OutputTarget;
+	TRefCountPtr<FRHITexture> OutputTarget;
 
 	/** The webbrowser texture that owns this resource. */
 	UWebBrowserTexture& Owner;
diff --git a/Engine/Source/Runtime/WebBrowserTexture/Public/WebBrowserTextureSample.h b/Engine/Source/Runtime/WebBrowserTexture/Public/WebBrowserTextureSample.h
--- a/Engine/Source/Runtime/WebBrowserTexture/Public/WebBrowserTextureSample.h
+++ b/Engine/Source/Runtime/WebBrowserTexture/Public/WebBrowserTextureSample.h
@@ -141,7 +141,7 @@ public:
 	* @note This method must be called on the render thread.
 	* @see InitializeBuffer
 	*/
-	FRHITexture2D* InitializeTexture()
+	FRHITexture* InitializeTexture()
 	{
 		check(IsInRenderingThread());
 
@@ -232,7 +232,7 @@ private:
 #if WITH_ENGINE
 
 	/** Texture resource. */
-	TRefCountPtr<FRHITexture2D> Texture;
+	TRefCountPtr<FRHITexture> Texture;
 
 #endif //WITH_ENGINE
 };
diff --git a/Engine/Source/Runtime/Windows/D3D11RHI/Private/D3D11RHIPrivate.h b/Engine/Source/Runtime/Windows/D3D11RHI/Private/D3D11RHIPrivate.h
--- a/Engine/Source/Runtime/Windows/D3D11RHI/Private/D3D11RHIPrivate.h
+++ b/Engine/Source/Runtime/Windows/D3D11RHI/Private/D3D11RHIPrivate.h
@@ -452,18 +452,18 @@ public:
 	virtual FTextureRHIRef RHIAsyncCreateTexture2D(uint32 SizeX, uint32 SizeY, uint8 Format, uint32 NumMips, ETextureCreateFlags Flags, ERHIAccess InResourceState, void** InitialMipData, uint32 NumInitialMips, const TCHAR* DebugName, FGraphEventRef& OutCompletionEvent) final override;
 	virtual void RHIGenerateMips(FRHITexture* Texture) final override;
 	virtual uint32 RHIComputeMemorySize(FRHITexture* TextureRHI) final override;
-	virtual void RHIAsyncCopyTexture2DCopy(FRHITexture2D* NewTexture2DRHI, FRHITexture2D* Texture2DRHI, int32 NewMipCount, int32 NewSizeX, int32 NewSizeY, FThreadSafeCounter* RequestStatus);
-	virtual FTexture2DRHIRef RHIAsyncReallocateTexture2D(FRHITexture2D* Texture2D, int32 NewMipCount, int32 NewSizeX, int32 NewSizeY, FThreadSafeCounter* RequestStatus) final override;
-	virtual ETextureReallocationStatus RHIFinalizeAsyncReallocateTexture2D(FRHITexture2D* Texture2D, bool bBlockUntilCompleted) final override;
-	virtual ETextureReallocationStatus RHICancelAsyncReallocateTexture2D(FRHITexture2D* Texture2D, bool bBlockUntilCompleted) final override;
-	virtual void* RHILockTexture2D(FRHITexture2D* Texture, uint32 MipIndex, EResourceLockMode LockMode, uint32& DestStride, bool bLockWithinMiptail, uint64* OutLockedByteCount) final override;
-	virtual void RHIUnlockTexture2D(FRHITexture2D* Texture, uint32 MipIndex, bool bLockWithinMiptail) final override;
-	virtual void* RHILockTexture2DArray(FRHITexture2DArray* Texture, uint32 TextureIndex, uint32 MipIndex, EResourceLockMode LockMode, uint32& DestStride, bool bLockWithinMiptail) final override;
-	virtual void RHIUnlockTexture2DArray(FRHITexture2DArray* Texture, uint32 TextureIndex, uint32 MipIndex, bool bLockWithinMiptail) final override;
-	virtual void RHIUpdateTexture2D(FRHICommandListBase& RHICmdList, FRHITexture2D* Texture, uint32 MipIndex, const struct FUpdateTextureRegion2D& UpdateRegion, uint32 SourcePitch, const uint8* SourceData) final override;
-	virtual void RHIUpdateTexture3D(FRHICommandListBase& RHICmdList, FRHITexture3D* Texture, uint32 MipIndex, const struct FUpdateTextureRegion3D& UpdateRegion, uint32 SourceRowPitch, uint32 SourceDepthPitch, const uint8* SourceData) final override;
-	virtual void* RHILockTextureCubeFace(FRHITextureCube* Texture, uint32 FaceIndex, uint32 ArrayIndex, uint32 MipIndex, EResourceLockMode LockMode, uint32& DestStride, bool bLockWithinMiptail) final override;
-	virtual void RHIUnlockTextureCubeFace(FRHITextureCube* Texture, uint32 FaceIndex, uint32 ArrayIndex, uint32 MipIndex, bool bLockWithinMiptail) final override;
+	virtual void RHIAsyncCopyTexture2DCopy(FRHITexture* NewTexture2DRHI, FRHITexture* Texture2DRHI, int32 NewMipCount, int32 NewSizeX, int32 NewSizeY, FThreadSafeCounter* RequestStatus);
+	virtual FTextureRHIRef RHIAsyncReallocateTexture2D(FRHITexture* Texture2D, int32 NewMipCount, int32 NewSizeX, int32 NewSizeY, FThreadSafeCounter* RequestStatus) final override;
+	virtual ETextureReallocationStatus RHIFinalizeAsyncReallocateTexture2D(FRHITexture* Texture2D, bool bBlockUntilCompleted) final override;
+	virtual ETextureReallocationStatus RHICancelAsyncReallocateTexture2D(FRHITexture* Texture2D, bool bBlockUntilCompleted) final override;
+	virtual void* RHILockTexture2D(FRHITexture* Texture, uint32 MipIndex, EResourceLockMode LockMode, uint32& DestStride, bool bLockWithinMiptail, uint64* OutLockedByteCount) final override;
+	virtual void RHIUnlockTexture2D(FRHITexture* Texture, uint32 MipIndex, bool bLockWithinMiptail) final override;
+	virtual void* RHILockTexture2DArray(FRHITexture* Texture, uint32 TextureIndex, uint32 MipIndex, EResourceLockMode LockMode, uint32& DestStride, bool bLockWithinMiptail) final override;
+	virtual void RHIUnlockTexture2DArray(FRHITexture* Texture, uint32 TextureIndex, uint32 MipIndex, bool bLockWithinMiptail) final override;
+	virtual void RHIUpdateTexture2D(FRHICommandListBase& RHICmdList, FRHITexture* Texture, uint32 MipIndex, const struct FUpdateTextureRegion2D& UpdateRegion, uint32 SourcePitch, const uint8* SourceData) final override;
+	virtual void RHIUpdateTexture3D(FRHICommandListBase& RHICmdList, FRHITexture* Texture, uint32 MipIndex, const struct FUpdateTextureRegion3D& UpdateRegion, uint32 SourceRowPitch, uint32 SourceDepthPitch, const uint8* SourceData) final override;
+	virtual void* RHILockTextureCubeFace(FRHITexture* Texture, uint32 FaceIndex, uint32 ArrayIndex, uint32 MipIndex, EResourceLockMode LockMode, uint32& DestStride, bool bLockWithinMiptail) final override;
+	virtual void RHIUnlockTextureCubeFace(FRHITexture* Texture, uint32 FaceIndex, uint32 ArrayIndex, uint32 MipIndex, bool bLockWithinMiptail) final override;
 	virtual void RHIBindDebugLabelName(FRHICommandListBase& RHICmdList, FRHITexture* Texture, const TCHAR* Name) final override;
 	virtual void RHIBindDebugLabelName(FRHICommandListBase& RHICmdList, FRHIBuffer* Buffer, const TCHAR* Name) final override;
 	virtual void RHIBindDebugLabelName(FRHICommandListBase& RHICmdList, FRHIUnorderedAccessView* UnorderedAccessViewRHI, const TCHAR* Name) final override;
@@ -475,7 +475,7 @@ public:
 	virtual void RHIRead3DSurfaceFloatData(FRHITexture* Texture,FIntRect Rect,FIntPoint ZMinMax,TArray<FFloat16Color>& OutData) final override;
 	virtual FRenderQueryRHIRef RHICreateRenderQuery(ERenderQueryType QueryType) final override;
 	virtual bool RHIGetRenderQueryResult(FRHIRenderQuery* RenderQuery, uint64& OutResult, bool bWait, uint32 GPUIndex = INDEX_NONE) final override;
-	virtual FTexture2DRHIRef RHIGetViewportBackBuffer(FRHIViewport* Viewport) final override;
+	virtual FTextureRHIRef RHIGetViewportBackBuffer(FRHIViewport* Viewport) final override;
 	virtual void RHIAliasTextureResources(FTextureRHIRef& DestTexture, FTextureRHIRef& SrcTexture) final override;
 	virtual FTextureRHIRef RHICreateAliasedTexture(FTextureRHIRef& SourceTexture) final override;
 	virtual void RHIAdvanceFrameForGetViewportBackBuffer(FRHIViewport* Viewport) final override;
@@ -578,14 +578,14 @@ public:
 	// *_RenderThread functions. Command lists call these functions on RT. You can implement your own behavior inside these functions.
 	// For example, deferring the actual creation to RHI thread by sending an RHI command.
 	// For D3D11, these functions mainly just remove RHIT stalls because ID3D11Device is thread safe.
-	virtual FTexture2DRHIRef AsyncReallocateTexture2D_RenderThread(class FRHICommandListImmediate& RHICmdList, FRHITexture2D* Texture2D, int32 NewMipCount, int32 NewSizeX, int32 NewSizeY, FThreadSafeCounter* RequestStatus) final override;
-	virtual ETextureReallocationStatus FinalizeAsyncReallocateTexture2D_RenderThread(class FRHICommandListImmediate& RHICmdList, FRHITexture2D* Texture2D, bool bBlockUntilCompleted) final override;
-	virtual ETextureReallocationStatus CancelAsyncReallocateTexture2D_RenderThread(class FRHICommandListImmediate& RHICmdList, FRHITexture2D* Texture2D, bool bBlockUntilCompleted) final override;
-	virtual void* LockTexture2D_RenderThread(class FRHICommandListImmediate& RHICmdList, FRHITexture2D* Texture, uint32 MipIndex, EResourceLockMode LockMode, uint32& DestStride, bool bLockWithinMiptail, bool bNeedsDefaultRHIFlush = true, uint64* OutLockedByteCount = nullptr) final override;
-	virtual void UnlockTexture2D_RenderThread(class FRHICommandListImmediate& RHICmdList, FRHITexture2D* Texture, uint32 MipIndex, bool bLockWithinMiptail, bool bNeedsDefaultRHIFlush = true) final override;
+	virtual FTextureRHIRef AsyncReallocateTexture2D_RenderThread(class FRHICommandListImmediate& RHICmdList, FRHITexture* Texture2D, int32 NewMipCount, int32 NewSizeX, int32 NewSizeY, FThreadSafeCounter* RequestStatus) final override;
+	virtual ETextureReallocationStatus FinalizeAsyncReallocateTexture2D_RenderThread(class FRHICommandListImmediate& RHICmdList, FRHITexture* Texture2D, bool bBlockUntilCompleted) final override;
+	virtual ETextureReallocationStatus CancelAsyncReallocateTexture2D_RenderThread(class FRHICommandListImmediate& RHICmdList, FRHITexture* Texture2D, bool bBlockUntilCompleted) final override;
+	virtual void* LockTexture2D_RenderThread(class FRHICommandListImmediate& RHICmdList, FRHITexture* Texture, uint32 MipIndex, EResourceLockMode LockMode, uint32& DestStride, bool bLockWithinMiptail, bool bNeedsDefaultRHIFlush = true, uint64* OutLockedByteCount = nullptr) final override;
+	virtual void UnlockTexture2D_RenderThread(class FRHICommandListImmediate& RHICmdList, FRHITexture* Texture, uint32 MipIndex, bool bLockWithinMiptail, bool bNeedsDefaultRHIFlush = true) final override;
 	virtual void RHIEndUpdateTexture3D(FRHICommandListBase& RHICmdList, FUpdateTexture3DData& UpdateData) final override;
-	virtual void* RHILockTextureCubeFace_RenderThread(class FRHICommandListImmediate& RHICmdList, FRHITextureCube* Texture, uint32 FaceIndex, uint32 ArrayIndex, uint32 MipIndex, EResourceLockMode LockMode, uint32& DestStride, bool bLockWithinMiptail) final override;
-	virtual void RHIUnlockTextureCubeFace_RenderThread(class FRHICommandListImmediate& RHICmdList, FRHITextureCube* Texture, uint32 FaceIndex, uint32 ArrayIndex, uint32 MipIndex, bool bLockWithinMiptail) final override;
+	virtual void* RHILockTextureCubeFace_RenderThread(class FRHICommandListImmediate& RHICmdList, FRHITexture* Texture, uint32 FaceIndex, uint32 ArrayIndex, uint32 MipIndex, EResourceLockMode LockMode, uint32& DestStride, bool bLockWithinMiptail) final override;
+	virtual void RHIUnlockTextureCubeFace_RenderThread(class FRHICommandListImmediate& RHICmdList, FRHITexture* Texture, uint32 FaceIndex, uint32 ArrayIndex, uint32 MipIndex, bool bLockWithinMiptail) final override;
 
 	virtual void RHIBeginRenderPass(const FRHIRenderPassInfo& InInfo, const TCHAR* InName) final override;
 	virtual void RHIEndRenderPass() final override;
@@ -598,9 +598,9 @@ public:
 	virtual IDXGIAdapter*         RHIGetAdapter() const final override;
 	virtual IDXGISwapChain*       RHIGetSwapChain(FRHIViewport* InViewport) const final override;
 	virtual DXGI_FORMAT           RHIGetSwapChainFormat(EPixelFormat InFormat) const final override;
-	virtual FTexture2DRHIRef      RHICreateTexture2DFromResource(EPixelFormat Format, ETextureCreateFlags TexCreateFlags, const FClearValueBinding& ClearValueBinding, ID3D11Texture2D* Resource) final override;
-	virtual FTexture2DArrayRHIRef RHICreateTexture2DArrayFromResource(EPixelFormat Format, ETextureCreateFlags TexCreateFlags, const FClearValueBinding&, ID3D11Texture2D* Resource) final override;
-	virtual FTextureCubeRHIRef    RHICreateTextureCubeFromResource(EPixelFormat Format, ETextureCreateFlags TexCreateFlags, const FClearValueBinding& ClearValueBinding, ID3D11Texture2D* Resource) final override;
+	virtual FTextureRHIRef        RHICreateTexture2DFromResource(EPixelFormat Format, ETextureCreateFlags TexCreateFlags, const FClearValueBinding& ClearValueBinding, ID3D11Texture2D* Resource) final override;
+	virtual FTextureRHIRef        RHICreateTexture2DArrayFromResource(EPixelFormat Format, ETextureCreateFlags TexCreateFlags, const FClearValueBinding&, ID3D11Texture2D* Resource) final override;
+	virtual FTextureRHIRef        RHICreateTextureCubeFromResource(EPixelFormat Format, ETextureCreateFlags TexCreateFlags, const FClearValueBinding& ClearValueBinding, ID3D11Texture2D* Resource) final override;
 	virtual ID3D11Buffer*         RHIGetResource(FRHIBuffer* InBuffer) const final override;
 	virtual ID3D11Resource*       RHIGetResource(FRHITexture* InTexture) const final override;
 	virtual int64                 RHIGetResourceMemorySize(FRHITexture* InTexture) const final override;
diff --git a/Engine/Source/Runtime/Windows/D3D11RHI/Private/D3D11Texture.cpp b/Engine/Source/Runtime/Windows/D3D11RHI/Private/D3D11Texture.cpp
--- a/Engine/Source/Runtime/Windows/D3D11RHI/Private/D3D11Texture.cpp
+++ b/Engine/Source/Runtime/Windows/D3D11RHI/Private/D3D11Texture.cpp
@@ -908,7 +908,7 @@ uint32 FD3D11DynamicRHI::RHIComputeMemorySize(FRHITexture* TextureRHI)
  * @param RequestStatus			- Will be decremented by 1 when the reallocation is complete (success or failure).
  * @return						- New reference to the texture, or an invalid reference upon failure
  */
-void FD3D11DynamicRHI::RHIAsyncCopyTexture2DCopy(FRHITexture2D* NewTexture2DRHI, FRHITexture2D* Texture2DRHI, int32 NewMipCount, int32 NewSizeX, int32 NewSizeY, FThreadSafeCounter* RequestStatus)
+void FD3D11DynamicRHI::RHIAsyncCopyTexture2DCopy(FRHITexture* NewTexture2DRHI, FRHITexture* Texture2DRHI, int32 NewMipCount, int32 NewSizeX, int32 NewSizeY, FThreadSafeCounter* RequestStatus)
 {
 	FD3D11Texture* Texture2D = ResourceCast(Texture2DRHI);
 	FD3D11Texture* NewTexture2D = ResourceCast(NewTexture2DRHI);
@@ -954,7 +954,7 @@ void FD3D11DynamicRHI::RHIAsyncCopyTexture2DCopy(FRHITexture2D* NewTexture2DRHI,
  * @param RequestStatus	- Will be decremented by 1 when the reallocation is complete (success or failure).
  * @return				- New reference to the texture, or an invalid reference upon failure
  */
-FTexture2DRHIRef FD3D11DynamicRHI::RHIAsyncReallocateTexture2D(FRHITexture2D* Texture2DRHI, int32 NewMipCount, int32 NewSizeX, int32 NewSizeY, FThreadSafeCounter* RequestStatus)
+FTextureRHIRef FD3D11DynamicRHI::RHIAsyncReallocateTexture2D(FRHITexture* Texture2DRHI, int32 NewMipCount, int32 NewSizeX, int32 NewSizeY, FThreadSafeCounter* RequestStatus)
 {
 	FD3D11Texture* Texture2D = ResourceCast(Texture2DRHI);
 
@@ -974,15 +974,15 @@ FTexture2DRHIRef FD3D11DynamicRHI::RHIAsyncReallocateTexture2D(FRHITexture2D* Te
 	return NewTexture2D;
 }
 
-FTexture2DRHIRef FD3D11DynamicRHI::AsyncReallocateTexture2D_RenderThread(
+FTextureRHIRef FD3D11DynamicRHI::AsyncReallocateTexture2D_RenderThread(
 	class FRHICommandListImmediate& RHICmdList,
-	FRHITexture2D* Texture2D,
+	FRHITexture* Texture2D,
 	int32 NewMipCount,
 	int32 NewSizeX,
 	int32 NewSizeY,
 	FThreadSafeCounter* RequestStatus)
 {
-	FTexture2DRHIRef NewTexture2D;
+	FTextureRHIRef NewTexture2D;
 
 	if (ShouldNotEnqueueRHICommand())
 	{
@@ -1018,14 +1018,14 @@ FTexture2DRHIRef FD3D11DynamicRHI::AsyncReallocateTexture2D_RenderThread(
  * @param Texture2D		- Texture to check the reallocation status for
  * @return				- Current reallocation status
  */
-ETextureReallocationStatus FD3D11DynamicRHI::RHIFinalizeAsyncReallocateTexture2D(FRHITexture2D* Texture2D, bool bBlockUntilCompleted )
+ETextureReallocationStatus FD3D11DynamicRHI::RHIFinalizeAsyncReallocateTexture2D(FRHITexture* Texture2D, bool bBlockUntilCompleted )
 {
 	return TexRealloc_Succeeded;
 }
 
 ETextureReallocationStatus FD3D11DynamicRHI::FinalizeAsyncReallocateTexture2D_RenderThread(
 	class FRHICommandListImmediate& RHICmdList,
-	FRHITexture2D* Texture2D,
+	FRHITexture* Texture2D,
 	bool bBlockUntilCompleted)
 {
 	return RHIFinalizeAsyncReallocateTexture2D(Texture2D, bBlockUntilCompleted);
@@ -1039,14 +1039,14 @@ ETextureReallocationStatus FD3D11DynamicRHI::FinalizeAsyncReallocateTexture2D_Re
  * @param bBlockUntilCompleted	If true, blocks until the cancellation is fully completed
  * @return						Reallocation status
  */
-ETextureReallocationStatus FD3D11DynamicRHI::RHICancelAsyncReallocateTexture2D(FRHITexture2D* Texture2D, bool bBlockUntilCompleted )
+ETextureReallocationStatus FD3D11DynamicRHI::RHICancelAsyncReallocateTexture2D(FRHITexture* Texture2D, bool bBlockUntilCompleted )
 {
 	return TexRealloc_Succeeded;
 }
 
 ETextureReallocationStatus FD3D11DynamicRHI::CancelAsyncReallocateTexture2D_RenderThread(
 	class FRHICommandListImmediate& RHICmdList,
-	FRHITexture2D* Texture2D,
+	FRHITexture* Texture2D,
 	bool bBlockUntilCompleted)
 {
 	return RHICancelAsyncReallocateTexture2D(Texture2D, bBlockUntilCompleted);
@@ -1188,7 +1188,7 @@ void FD3D11Texture::Unlock(FD3D11DynamicRHI* D3DRHI, uint32 MipIndex, uint32 Arr
 	}
 }
 
-void* FD3D11DynamicRHI::RHILockTexture2D(FRHITexture2D* TextureRHI, uint32 MipIndex, EResourceLockMode LockMode, uint32& DestStride, bool bLockWithinMiptail, uint64* OutLockedByteCount)
+void* FD3D11DynamicRHI::RHILockTexture2D(FRHITexture* TextureRHI, uint32 MipIndex, EResourceLockMode LockMode, uint32& DestStride, bool bLockWithinMiptail, uint64* OutLockedByteCount)
 {
 	check(TextureRHI);
 	FD3D11Texture* Texture = ResourceCast(TextureRHI);
@@ -1198,7 +1198,7 @@ void* FD3D11DynamicRHI::RHILockTexture2D(FRHITexture2D* TextureRHI, uint32 MipIn
 
 void* FD3D11DynamicRHI::LockTexture2D_RenderThread(
 	class FRHICommandListImmediate& RHICmdList,
-	FRHITexture2D* Texture,
+	FRHITexture* Texture,
 	uint32 MipIndex,
 	EResourceLockMode LockMode,
 	uint32& DestStride,
@@ -1225,7 +1225,7 @@ void* FD3D11DynamicRHI::LockTexture2D_RenderThread(
 	return LockedTexture;
 }
 
-void FD3D11DynamicRHI::RHIUnlockTexture2D(FRHITexture2D* TextureRHI, uint32 MipIndex, bool bLockWithinMiptail)
+void FD3D11DynamicRHI::RHIUnlockTexture2D(FRHITexture* TextureRHI, uint32 MipIndex, bool bLockWithinMiptail)
 {
 	check(TextureRHI);
 	FD3D11Texture* Texture = ResourceCast(TextureRHI);
@@ -1234,7 +1234,7 @@ void FD3D11DynamicRHI::RHIUnlockTexture2D(FRHITexture2D* TextureRHI, uint32 MipI
 
 void FD3D11DynamicRHI::UnlockTexture2D_RenderThread(
 	class FRHICommandListImmediate& RHICmdList,
-	FRHITexture2D* Texture,
+	FRHITexture* Texture,
 	uint32 MipIndex,
 	bool bLockWithinMiptail,
 	bool bNeedsDefaultRHIFlush)
@@ -1252,20 +1252,20 @@ void FD3D11DynamicRHI::UnlockTexture2D_RenderThread(
 	}
 }
 
-void* FD3D11DynamicRHI::RHILockTexture2DArray(FRHITexture2DArray* TextureRHI, uint32 TextureIndex, uint32 MipIndex, EResourceLockMode LockMode, uint32& DestStride, bool bLockWithinMiptail)
+void* FD3D11DynamicRHI::RHILockTexture2DArray(FRHITexture* TextureRHI, uint32 TextureIndex, uint32 MipIndex, EResourceLockMode LockMode, uint32& DestStride, bool bLockWithinMiptail)
 {
 	FD3D11Texture* Texture = ResourceCast(TextureRHI);
 	ConditionalClearShaderResource(Texture, false);
 	return Texture->Lock(this, MipIndex, TextureIndex, LockMode, DestStride);
 }
 
-void FD3D11DynamicRHI::RHIUnlockTexture2DArray(FRHITexture2DArray* TextureRHI, uint32 TextureIndex, uint32 MipIndex, bool bLockWithinMiptail)
+void FD3D11DynamicRHI::RHIUnlockTexture2DArray(FRHITexture* TextureRHI, uint32 TextureIndex, uint32 MipIndex, bool bLockWithinMiptail)
 {
 	FD3D11Texture* Texture = ResourceCast(TextureRHI);
 	Texture->Unlock(this, MipIndex, TextureIndex);
 }
 
-void FD3D11DynamicRHI::RHIUpdateTexture2D(FRHICommandListBase& RHICmdList, FRHITexture2D* TextureRHI, uint32 MipIndex, const FUpdateTextureRegion2D& UpdateRegion, uint32 SourcePitch, const uint8* SourceData)
+void FD3D11DynamicRHI::RHIUpdateTexture2D(FRHICommandListBase& RHICmdList, FRHITexture* TextureRHI, uint32 MipIndex, const FUpdateTextureRegion2D& UpdateRegion, uint32 SourcePitch, const uint8* SourceData)
 {
 	const FPixelFormatInfo& FormatInfo = GPixelFormats[TextureRHI->GetFormat()];
 
@@ -1325,7 +1325,7 @@ void FD3D11DynamicRHI::RHIUpdateTexture2D(FRHICommandListBase& RHICmdList, FRHIT
 	});
 }
 
-void FD3D11DynamicRHI::RHIUpdateTexture3D(FRHICommandListBase& RHICmdList, FRHITexture3D* TextureRHI,uint32 MipIndex,const FUpdateTextureRegion3D& UpdateRegion,uint32 SourceRowPitch,uint32 SourceDepthPitch,const uint8* SourceData)
+void FD3D11DynamicRHI::RHIUpdateTexture3D(FRHICommandListBase& RHICmdList, FRHITexture* TextureRHI,uint32 MipIndex,const FUpdateTextureRegion3D& UpdateRegion,uint32 SourceRowPitch,uint32 SourceDepthPitch,const uint8* SourceData)
 {
 	const SIZE_T SourceDataSize = static_cast<SIZE_T>(SourceDepthPitch) * UpdateRegion.Depth;
 	uint8* SourceDataCopy = (uint8*)FMemory::Malloc(SourceDataSize);
@@ -1365,7 +1365,7 @@ void FD3D11DynamicRHI::RHIEndUpdateTexture3D(FRHICommandListBase& RHICmdList, FU
 	Cubemap texture support.
 -----------------------------------------------------------------------------*/
 
-void* FD3D11DynamicRHI::RHILockTextureCubeFace(FRHITextureCube* TextureCubeRHI, uint32 FaceIndex, uint32 ArrayIndex, uint32 MipIndex, EResourceLockMode LockMode, uint32& DestStride, bool bLockWithinMiptail)
+void* FD3D11DynamicRHI::RHILockTextureCubeFace(FRHITexture* TextureCubeRHI, uint32 FaceIndex, uint32 ArrayIndex, uint32 MipIndex, EResourceLockMode LockMode, uint32& DestStride, bool bLockWithinMiptail)
 {
 	FD3D11Texture* TextureCube = ResourceCast(TextureCubeRHI);
 	ConditionalClearShaderResource(TextureCube, false);
@@ -1375,7 +1375,7 @@ void* FD3D11DynamicRHI::RHILockTextureCubeFace(FRHITextureCube* TextureCubeRHI,
 
 void* FD3D11DynamicRHI::RHILockTextureCubeFace_RenderThread(
 	class FRHICommandListImmediate& RHICmdList,
-	FRHITextureCube* Texture,
+	FRHITexture* Texture,
 	uint32 FaceIndex,
 	uint32 ArrayIndex,
 	uint32 MipIndex,
@@ -1403,7 +1403,7 @@ void* FD3D11DynamicRHI::RHILockTextureCubeFace_RenderThread(
 	return LockedTexture;
 }
 
-void FD3D11DynamicRHI::RHIUnlockTextureCubeFace(FRHITextureCube* TextureCubeRHI, uint32 FaceIndex, uint32 ArrayIndex, uint32 MipIndex, bool bLockWithinMiptail)
+void FD3D11DynamicRHI::RHIUnlockTextureCubeFace(FRHITexture* TextureCubeRHI, uint32 FaceIndex, uint32 ArrayIndex, uint32 MipIndex, bool bLockWithinMiptail)
 {
 	FD3D11Texture* TextureCube = ResourceCast(TextureCubeRHI);
 	uint32 D3DFace = GetD3D11CubeFace((ECubeFace)FaceIndex);
@@ -1412,7 +1412,7 @@ void FD3D11DynamicRHI::RHIUnlockTextureCubeFace(FRHITextureCube* TextureCubeRHI,
 
 void FD3D11DynamicRHI::RHIUnlockTextureCubeFace_RenderThread(
 	class FRHICommandListImmediate& RHICmdList,
-	FRHITextureCube* Texture,
+	FRHITexture* Texture,
 	uint32 FaceIndex,
 	uint32 ArrayIndex,
 	uint32 MipIndex,
@@ -1697,17 +1697,17 @@ FD3D11Texture* FD3D11DynamicRHI::CreateTextureFromResource(bool bTextureArray, b
 	return Texture2D;
 }
 
-FTexture2DRHIRef FD3D11DynamicRHI::RHICreateTexture2DFromResource(EPixelFormat Format, ETextureCreateFlags TexCreateFlags, const FClearValueBinding& ClearValueBinding, ID3D11Texture2D* TextureResource)
+FTextureRHIRef FD3D11DynamicRHI::RHICreateTexture2DFromResource(EPixelFormat Format, ETextureCreateFlags TexCreateFlags, const FClearValueBinding& ClearValueBinding, ID3D11Texture2D* TextureResource)
 {
 	return CreateTextureFromResource(false, false, Format, TexCreateFlags, ClearValueBinding, TextureResource);
 }
 
-FTexture2DArrayRHIRef FD3D11DynamicRHI::RHICreateTexture2DArrayFromResource(EPixelFormat Format, ETextureCreateFlags TexCreateFlags, const FClearValueBinding& ClearValueBinding, ID3D11Texture2D* TextureResource)
+FTextureRHIRef FD3D11DynamicRHI::RHICreateTexture2DArrayFromResource(EPixelFormat Format, ETextureCreateFlags TexCreateFlags, const FClearValueBinding& ClearValueBinding, ID3D11Texture2D* TextureResource)
 {
 	return CreateTextureFromResource(true, false, Format, TexCreateFlags, ClearValueBinding, TextureResource);
 }
 
-FTextureCubeRHIRef FD3D11DynamicRHI::RHICreateTextureCubeFromResource(EPixelFormat Format, ETextureCreateFlags TexCreateFlags, const FClearValueBinding& ClearValueBinding, ID3D11Texture2D* TextureResource)
+FTextureRHIRef FD3D11DynamicRHI::RHICreateTextureCubeFromResource(EPixelFormat Format, ETextureCreateFlags TexCreateFlags, const FClearValueBinding& ClearValueBinding, ID3D11Texture2D* TextureResource)
 {
 	return CreateTextureFromResource(false, true, Format, TexCreateFlags, ClearValueBinding, TextureResource);
 }
diff --git a/Engine/Source/Runtime/Windows/D3D11RHI/Private/D3D11Viewport.cpp b/Engine/Source/Runtime/Windows/D3D11RHI/Private/D3D11Viewport.cpp
--- a/Engine/Source/Runtime/Windows/D3D11RHI/Private/D3D11Viewport.cpp
+++ b/Engine/Source/Runtime/Windows/D3D11RHI/Private/D3D11Viewport.cpp
@@ -839,7 +839,7 @@ void FD3D11DynamicRHI::RHIAdvanceFrameForGetViewportBackBuffer(FRHIViewport* Vie
 {
 }
 
-FTexture2DRHIRef FD3D11DynamicRHI::RHIGetViewportBackBuffer(FRHIViewport* ViewportRHI)
+FTextureRHIRef FD3D11DynamicRHI::RHIGetViewportBackBuffer(FRHIViewport* ViewportRHI)
 {
 	FD3D11Viewport* Viewport = ResourceCast(ViewportRHI);
 
diff --git a/Engine/Source/Runtime/Windows/D3D11RHI/Public/ID3D11DynamicRHI.h b/Engine/Source/Runtime/Windows/D3D11RHI/Public/ID3D11DynamicRHI.h
--- a/Engine/Source/Runtime/Windows/D3D11RHI/Public/ID3D11DynamicRHI.h
+++ b/Engine/Source/Runtime/Windows/D3D11RHI/Public/ID3D11DynamicRHI.h
@@ -22,9 +22,9 @@ struct ID3D11DynamicRHI : public FDynamicRHIPSOFallback
 	virtual IDXGISwapChain*       RHIGetSwapChain(FRHIViewport* InViewport) const = 0;
 	virtual DXGI_FORMAT           RHIGetSwapChainFormat(EPixelFormat InFormat) const = 0;
 
-	virtual FTexture2DRHIRef      RHICreateTexture2DFromResource(EPixelFormat Format, ETextureCreateFlags TexCreateFlags, const FClearValueBinding& ClearValueBinding, ID3D11Texture2D* Resource) = 0;
-	virtual FTexture2DArrayRHIRef RHICreateTexture2DArrayFromResource(EPixelFormat Format, ETextureCreateFlags TexCreateFlags, const FClearValueBinding& ClearValueBinding, ID3D11Texture2D* Resource) = 0;
-	virtual FTextureCubeRHIRef    RHICreateTextureCubeFromResource(EPixelFormat Format, ETextureCreateFlags TexCreateFlags, const FClearValueBinding& ClearValueBinding, ID3D11Texture2D* Resource) = 0;
+	virtual FTextureRHIRef        RHICreateTexture2DFromResource(EPixelFormat Format, ETextureCreateFlags TexCreateFlags, const FClearValueBinding& ClearValueBinding, ID3D11Texture2D* Resource) = 0;
+	virtual FTextureRHIRef        RHICreateTexture2DArrayFromResource(EPixelFormat Format, ETextureCreateFlags TexCreateFlags, const FClearValueBinding& ClearValueBinding, ID3D11Texture2D* Resource) = 0;
+	virtual FTextureRHIRef        RHICreateTextureCubeFromResource(EPixelFormat Format, ETextureCreateFlags TexCreateFlags, const FClearValueBinding& ClearValueBinding, ID3D11Texture2D* Resource) = 0;
 
 	virtual ID3D11Buffer*         RHIGetResource(FRHIBuffer* InBuffer) const = 0;
 	virtual ID3D11Resource*       RHIGetResource(FRHITexture* InTexture) const = 0;
