VulkanRHI:  add initialization for extensions: VK_KHR_fragment_shader_barycentric and VK_NV_compute_shader_derivatives

#rnx

diff --git a/Engine/Source/Runtime/VulkanRHI/Private/VulkanDevice.h b/Engine/Source/Runtime/VulkanRHI/Private/VulkanDevice.h
--- a/Engine/Source/Runtime/VulkanRHI/Private/VulkanDevice.h
+++ b/Engine/Source/Runtime/VulkanRHI/Private/VulkanDevice.h
@@ -48,6 +48,7 @@ struct FOptionalVulkanDeviceExtensions
 			uint64 HasEXTFragmentDensityMap : 1;
 			uint64 HasEXTFragmentDensityMap2 : 1;
 			uint64 HasKHRFragmentShadingRate : 1;
+			uint64 HasKHRFragmentShaderBarycentric : 1;
 			uint64 HasEXTFullscreenExclusive : 1;
 			uint64 HasImageAtomicInt64 : 1;
 			uint64 HasAccelerationStructure : 1;
@@ -138,6 +139,9 @@ struct FOptionalVulkanDeviceExtensionProperties
 	VkPhysicalDeviceFragmentShadingRateFeaturesKHR FragmentShadingRateFeatures;
 	VkPhysicalDeviceFragmentDensityMapFeaturesEXT FragmentDensityMapFeatures;
 	VkPhysicalDeviceFragmentDensityMap2FeaturesEXT FragmentDensityMap2Features;
+
+	VkPhysicalDeviceFragmentShaderBarycentricPropertiesKHR FragmentShaderBarycentricProps;
+	VkPhysicalDeviceComputeShaderDerivativesFeaturesNV ComputeShaderDerivativesFeatures;
 };
 
 class FVulkanPhysicalDeviceFeatures
diff --git a/Engine/Source/Runtime/VulkanRHI/Private/VulkanExtensions.cpp b/Engine/Source/Runtime/VulkanRHI/Private/VulkanExtensions.cpp
--- a/Engine/Source/Runtime/VulkanRHI/Private/VulkanExtensions.cpp
+++ b/Engine/Source/Runtime/VulkanRHI/Private/VulkanExtensions.cpp
@@ -691,6 +691,96 @@ public:
 
 
 
+// ***** VK_KHR_fragment_shader_barycentric
+class FVulkanKHRFragmentShaderBarycentricExtension : public FVulkanDeviceExtension
+{
+public:
+
+	FVulkanKHRFragmentShaderBarycentricExtension(FVulkanDevice* InDevice)
+		: FVulkanDeviceExtension(InDevice, VK_KHR_FRAGMENT_SHADER_BARYCENTRIC_EXTENSION_NAME, VULKAN_EXTENSION_ENABLED)
+	{
+	}
+
+	virtual void PrePhysicalDeviceProperties(VkPhysicalDeviceProperties2KHR& PhysicalDeviceProperties2) override final
+	{
+		VkPhysicalDeviceFragmentShaderBarycentricPropertiesKHR& FragmentShaderBarycentricProps = GetDeviceExtensionProperties().FragmentShaderBarycentricProps;
+		ZeroVulkanStruct(FragmentShaderBarycentricProps, VK_STRUCTURE_TYPE_PHYSICAL_DEVICE_FRAGMENT_SHADER_BARYCENTRIC_PROPERTIES_KHR);
+		AddToPNext(PhysicalDeviceProperties2, FragmentShaderBarycentricProps);
+	}
+
+	virtual void PrePhysicalDeviceFeatures(VkPhysicalDeviceFeatures2KHR& PhysicalDeviceFeatures2) override final
+	{
+		ZeroVulkanStruct(FragmentShaderBarycentricFeatures, VK_STRUCTURE_TYPE_PHYSICAL_DEVICE_FRAGMENT_SHADER_BARYCENTRIC_FEATURES_KHR);
+		AddToPNext(PhysicalDeviceFeatures2, FragmentShaderBarycentricFeatures);
+	}
+
+	virtual void PostPhysicalDeviceFeatures(FOptionalVulkanDeviceExtensions& ExtensionFlags) override final
+	{
+		bRequirementsPassed = (FragmentShaderBarycentricFeatures.fragmentShaderBarycentric == VK_TRUE);
+
+		// Should be runtime guaranteed through SM6 profile
+		ExtensionFlags.HasKHRFragmentShaderBarycentric = bRequirementsPassed;
+		GRHIGlobals.SupportsBarycentricsSemantic = bRequirementsPassed;
+	}
+
+	virtual void PostPhysicalDeviceProperties() override final
+	{
+		VkPhysicalDeviceFragmentShaderBarycentricPropertiesKHR& FragmentShaderBarycentricProps = GetDeviceExtensionProperties().FragmentShaderBarycentricProps;
+		//UE_LOG(LogVulkanRHI, Verbose, TEXT("triStripVertexOrderIndependentOfProvokingVertex is %s"), FragmentShaderBarycentricProps.triStripVertexOrderIndependentOfProvokingVertex ? TEXT("true") : TEXT("false"));
+	}
+
+	virtual void PreCreateDevice(VkDeviceCreateInfo& DeviceCreateInfo) override final
+	{
+		// fragmentShaderBarycentric indicates that the implementation supports the BaryCoordKHR and BaryCoordNoPerspKHR SPIR - V fragment shader built - ins 
+		// and supports the PerVertexKHR SPIR - V decoration on fragment shader input variables.
+		if (bRequirementsPassed)
+		{
+			AddToPNext(DeviceCreateInfo, FragmentShaderBarycentricFeatures);
+		}
+	}
+
+	VkPhysicalDeviceFragmentShaderBarycentricFeaturesKHR FragmentShaderBarycentricFeatures;
+};
+
+
+
+// ***** VK_NV_compute_shader_derivatives
+class FVulkanNVComputeShaderDerivatives : public FVulkanDeviceExtension
+{
+public:
+
+	FVulkanNVComputeShaderDerivatives(FVulkanDevice* InDevice)
+		: FVulkanDeviceExtension(InDevice, VK_NV_COMPUTE_SHADER_DERIVATIVES_EXTENSION_NAME, VULKAN_EXTENSION_ENABLED)
+	{
+	}
+
+	virtual void PrePhysicalDeviceFeatures(VkPhysicalDeviceFeatures2KHR& PhysicalDeviceFeatures2) override final
+	{
+		VkPhysicalDeviceComputeShaderDerivativesFeaturesNV& ComputeShaderDerivativesFeatures = GetDeviceExtensionProperties().ComputeShaderDerivativesFeatures;
+		ZeroVulkanStruct(ComputeShaderDerivativesFeatures, VK_STRUCTURE_TYPE_PHYSICAL_DEVICE_COMPUTE_SHADER_DERIVATIVES_FEATURES_NV);
+		AddToPNext(PhysicalDeviceFeatures2, ComputeShaderDerivativesFeatures);
+	}
+
+	virtual void PostPhysicalDeviceFeatures(FOptionalVulkanDeviceExtensions& ExtensionFlags) override final
+	{
+		VkPhysicalDeviceComputeShaderDerivativesFeaturesNV& ComputeShaderDerivativesFeatures = GetDeviceExtensionProperties().ComputeShaderDerivativesFeatures;
+		bRequirementsPassed = (ComputeShaderDerivativesFeatures.computeDerivativeGroupLinear == VK_TRUE);
+		ComputeShaderDerivativesFeatures.computeDerivativeGroupQuads = VK_FALSE;  // disable the unused quad mode
+		// Should be runtime guaranteed through SM6 profile
+	}
+
+	virtual void PreCreateDevice(VkDeviceCreateInfo& DeviceCreateInfo) override final
+	{
+		if (bRequirementsPassed)
+		{
+			VkPhysicalDeviceComputeShaderDerivativesFeaturesNV& ComputeShaderDerivativesFeatures = GetDeviceExtensionProperties().ComputeShaderDerivativesFeatures;
+			AddToPNext(DeviceCreateInfo, ComputeShaderDerivativesFeatures);
+		}
+	}
+};
+
+
+
 // ***** VK_KHR_get_memory_requirements2
 class FVulkanKHRGetMemoryRequirements2Extension : public FVulkanDeviceExtension
 {
@@ -1394,6 +1484,8 @@ FVulkanDeviceExtensionArray FVulkanDeviceExtension::GetUESupportedDeviceExtensio
 	ADD_CUSTOM_EXTENSION(FVulkanKHR16BitStorageExtension);
 	ADD_CUSTOM_EXTENSION(FVulkanKHRShaderFloat16Int8Extension);
 	ADD_CUSTOM_EXTENSION(FVulkanEXTPipelineCreationCacheControlExtension);
+	ADD_CUSTOM_EXTENSION(FVulkanKHRFragmentShaderBarycentricExtension);
+	ADD_CUSTOM_EXTENSION(FVulkanNVComputeShaderDerivatives);
 
 	// Needed for Raytracing
 	ADD_CUSTOM_EXTENSION(FVulkanKHRBufferDeviceAddressExtension);
