Fix hash collision between compute and graphics descriptor set layouts in Vulkan.

The layout info didn't differentiate between compute/graphics, so when a compute shader (e.g. FShaderPrintUploadCS) had the same layout as a graphics shader (e.g. FPostProcessBenchmarkPS), they mapped to the same entry in cache manager layout map, leading to a crash because we ended up using an FVulkanComputeLayout object where FVulkanGfxLayout was expected.

#rnx

diff --git a/Engine/Source/Runtime/VulkanRHI/Private/VulkanDescriptorSets.h b/Engine/Source/Runtime/VulkanRHI/Private/VulkanDescriptorSets.h
--- a/Engine/Source/Runtime/VulkanRHI/Private/VulkanDescriptorSets.h
+++ b/Engine/Source/Runtime/VulkanRHI/Private/VulkanDescriptorSets.h
@@ -357,7 +357,7 @@ public:
 	template<bool bIsCompute>
 	void FinalizeBindings(const FVulkanDevice& Device, const FUniformBufferGatherInfo& UBGatherInfo, const TArrayView<FRHISamplerState*>& ImmutableSamplers);
 
-	void GenerateHash(const TArrayView<FRHISamplerState*>& ImmutableSamplers);
+	void GenerateHash(const TArrayView<FRHISamplerState*>& ImmutableSamplers, VkPipelineBindPoint InBindPoint);
 
 	friend uint32 GetTypeHash(const FVulkanDescriptorSetsLayoutInfo& In)
 	{
@@ -371,6 +371,11 @@ public:
 			return false;
 		}
 
+		if (In.BindPoint != BindPoint)
+		{
+			return false;
+		}
+
 		if (In.SetLayouts.Num() != SetLayouts.Num())
 		{
 			return false;
@@ -429,6 +434,8 @@ protected:
 
 	uint32 TypesUsageID = ~0;
 
+	VkPipelineBindPoint BindPoint = VK_PIPELINE_BIND_POINT_MAX_ENUM;
+
 	void CompileTypesUsageID();
 
 	void AddDescriptor(int32 DescriptorSetIndex, const VkDescriptorSetLayoutBinding& Descriptor);
diff --git a/Engine/Source/Runtime/VulkanRHI/Private/VulkanRHI.cpp b/Engine/Source/Runtime/VulkanRHI/Private/VulkanRHI.cpp
--- a/Engine/Source/Runtime/VulkanRHI/Private/VulkanRHI.cpp
+++ b/Engine/Source/Runtime/VulkanRHI/Private/VulkanRHI.cpp
@@ -1692,7 +1692,7 @@ void FVulkanDescriptorSetsLayoutInfo::AddDescriptor(int32 DescriptorSetIndex, co
 	}
 }
 
-void FVulkanDescriptorSetsLayoutInfo::GenerateHash(const TArrayView<FRHISamplerState*>& InImmutableSamplers)
+void FVulkanDescriptorSetsLayoutInfo::GenerateHash(const TArrayView<FRHISamplerState*>& InImmutableSamplers, VkPipelineBindPoint InBindPoint)
 {
 	const int32 LayoutCount = SetLayouts.Num();
 	Hash = FCrc::MemCrc32(&TypesUsageID, sizeof(uint32), LayoutCount);
@@ -1717,6 +1717,12 @@ void FVulkanDescriptorSetsLayoutInfo::GenerateHash(const TArrayView<FRHISamplerS
 		TArray<uint16>& PackedUBBindingIndices = RemappingInfo.StageInfos[RemapingIndex].PackedUBBindingIndices;
 		Hash = FCrc::MemCrc32(PackedUBBindingIndices.GetData(), sizeof(uint16) * PackedUBBindingIndices.Num(), Hash);
 	}
+
+	// It would be better to store this when the object is created, but it's not available at that time, so we'll do it here.
+	BindPoint = InBindPoint;
+
+	// Include the bind point in the hash, because we can have graphics and compute PSOs with the same descriptor info, and we don't want them to collide.
+	Hash = FCrc::MemCrc32(&BindPoint, sizeof(BindPoint), Hash);
 }
 
 static FCriticalSection GTypesUsageCS;
diff --git a/Engine/Source/Runtime/VulkanRHI/Private/VulkanShaders.cpp b/Engine/Source/Runtime/VulkanRHI/Private/VulkanShaders.cpp
--- a/Engine/Source/Runtime/VulkanRHI/Private/VulkanShaders.cpp
+++ b/Engine/Source/Runtime/VulkanRHI/Private/VulkanShaders.cpp
@@ -1001,7 +1001,7 @@ void FVulkanDescriptorSetsLayoutInfo::FinalizeBindings(const FVulkanDevice& Devi
 	}
 
 	CompileTypesUsageID();
-	GenerateHash(ImmutableSamplers);
+	GenerateHash(ImmutableSamplers, bIsCompute ? VK_PIPELINE_BIND_POINT_COMPUTE : VK_PIPELINE_BIND_POINT_GRAPHICS);
 
 	// If we are consolidating and no uniforms are present in the shader, then strip the empty set data
 	if (bConsolidateAllIntoOneSet)
