Fix use of dual src blending in Vulkan mobile forward renderer.
This now does no longer bind the SceneDepthAux color attachment for the second subpass.
This fixes thin translucent materials on xclipse in mobile forward.

diff --git a/Engine/Source/Runtime/VulkanRHI/Private/VulkanPipeline.cpp b/Engine/Source/Runtime/VulkanRHI/Private/VulkanPipeline.cpp
--- a/Engine/Source/Runtime/VulkanRHI/Private/VulkanPipeline.cpp
+++ b/Engine/Source/Runtime/VulkanRHI/Private/VulkanPipeline.cpp
@@ -1822,6 +1822,12 @@ void FVulkanPipelineStateCacheManager::CreateGfxEntry(const FGraphicsPipelineSta
 		NumRenderTargets = 1;
 	}
 
+	if (PSOInitializer.SubpassHint == ESubpassHint::DepthReadSubpass && PSOInitializer.SubpassIndex >= 1)
+	{
+		// Only SceneColor is used as a color attachment after the first subpass (not SceneDepthAux)
+		NumRenderTargets = 1;
+	}
+
 	if (PSOInitializer.SubpassHint == ESubpassHint::CustomResolveSubpass)
 	{
 		NumRenderTargets = 1; // This applies to base and depth passes as well. One render target for base and depth, another one for custom resolve.
diff --git a/Engine/Source/Runtime/VulkanRHI/Private/VulkanRenderpass.h b/Engine/Source/Runtime/VulkanRHI/Private/VulkanRenderpass.h
--- a/Engine/Source/Runtime/VulkanRHI/Private/VulkanRenderpass.h
+++ b/Engine/Source/Runtime/VulkanRHI/Private/VulkanRenderpass.h
@@ -510,7 +510,7 @@ public:
 		{
 			TSubpassDescriptionClass& SubpassDesc = SubpassDescriptions[NumSubpasses++];
 
-			SubpassDesc.SetColorAttachments(ColorAttachmentReferences, NumColorAttachments);
+			SubpassDesc.SetColorAttachments(ColorAttachmentReferences, 1);
 
 			check(RTLayout.GetDepthAttachmentReference());
 
