Support Vulkan pre-rotation on mobile.

- Create an intermediate image as the BackBuffer and render normally, then copy the image to the actual BackBuffer in the low level.
- Remove the QCom pre-rotation implementation.

diff --git a/Engine/Shaders/Private/Tools/FullscreenVertexShader.usf b/Engine/Shaders/Private/Tools/FullscreenVertexShader.usf
--- a/Engine/Shaders/Private/Tools/FullscreenVertexShader.usf
+++ b/Engine/Shaders/Private/Tools/FullscreenVertexShader.usf
@@ -30,3 +30,12 @@ void MainVS(
 	OutViewId = InstanceId;
 #endif
 }
+
+float4 PreTransform;
+void MainForPreTransform(in uint VertexId : SV_VertexID, out FScreenVertexOutput Output)
+{
+	Output.UV = float2(((VertexId + 1) / 3) & 1, VertexId & 1);
+	Output.Position.xy = float2(Output.UV.x, 1.f - Output.UV.y) * 2.f - 1.f;
+	Output.Position.zw = float2(0, 1);
+	Output.Position.xy = float2(Output.Position.x * PreTransform.x + Output.Position.y * PreTransform.z, Output.Position.x * PreTransform.y + Output.Position.y * PreTransform.w);
+}
\ No newline at end of file
diff --git a/Engine/Source/Runtime/RenderCore/Private/CommonRenderResources.cpp b/Engine/Source/Runtime/RenderCore/Private/CommonRenderResources.cpp
--- a/Engine/Source/Runtime/RenderCore/Private/CommonRenderResources.cpp
+++ b/Engine/Source/Runtime/RenderCore/Private/CommonRenderResources.cpp
@@ -18,6 +18,7 @@ TGlobalResource<FScreenRectangleIndexBuffer, FRenderResource::EInitPhase::Pre> G
 IMPLEMENT_GLOBAL_SHADER(FScreenVertexShaderVS, "/Engine/Private/Tools/FullscreenVertexShader.usf", "MainVS", SF_Vertex);
 IMPLEMENT_GLOBAL_SHADER(FInstancedScreenVertexShaderVS, "/Engine/Private/Tools/FullscreenVertexShader.usf", "MainVS", SF_Vertex);
 IMPLEMENT_GLOBAL_SHADER(FCopyRectPS, "/Engine/Private/ScreenPass.usf", "CopyRectPS", SF_Pixel);
+IMPLEMENT_GLOBAL_SHADER(FImagePreTransformVS, "/Engine/Private/Tools/FullscreenVertexShader.usf", "MainForPreTransform", SF_Vertex);
 
 bool FInstancedScreenVertexShaderVS::ShouldCompilePermutation(const FGlobalShaderPermutationParameters& Parameters)
 {
diff --git a/Engine/Source/Runtime/RenderCore/Public/CommonRenderResources.h b/Engine/Source/Runtime/RenderCore/Public/CommonRenderResources.h
--- a/Engine/Source/Runtime/RenderCore/Public/CommonRenderResources.h
+++ b/Engine/Source/Runtime/RenderCore/Public/CommonRenderResources.h
@@ -19,6 +19,7 @@
 #include "ShaderParameterMacros.h"
 #include "ShaderParameterStruct.h"
 #include "Templates/UnrealTemplate.h"
+#include "DataDrivenShaderPlatformInfo.h"
 
 class FPointerTableBase;
 class FRDGTexture;
@@ -146,3 +147,19 @@ class FCopyRectPS : public FGlobalShader
 		RENDER_TARGET_BINDING_SLOTS()
 	END_SHADER_PARAMETER_STRUCT()
 };
+
+/** Vertex shader to perform a screen rotation for Vulkan pre-rotation on mobile. */
+class FImagePreTransformVS : public FGlobalShader
+{
+	DECLARE_EXPORTED_GLOBAL_SHADER(FImagePreTransformVS, RENDERCORE_API);
+	SHADER_USE_PARAMETER_STRUCT(FImagePreTransformVS, FGlobalShader);
+
+	static bool ShouldCompilePermutation(const FGlobalShaderPermutationParameters& Parameters)
+	{
+		return IsVulkanMobilePlatform(Parameters.Platform) || IsVulkanMobileSM5Platform(Parameters.Platform);
+	}
+
+	BEGIN_SHADER_PARAMETER_STRUCT(FParameters, )
+		SHADER_PARAMETER(FVector4f, PreTransform)
+	END_SHADER_PARAMETER_STRUCT()
+};
diff --git a/Engine/Source/Runtime/VulkanRHI/Private/Android/VulkanAndroidPlatform.cpp b/Engine/Source/Runtime/VulkanRHI/Private/Android/VulkanAndroidPlatform.cpp
--- a/Engine/Source/Runtime/VulkanRHI/Private/Android/VulkanAndroidPlatform.cpp
+++ b/Engine/Source/Runtime/VulkanRHI/Private/Android/VulkanAndroidPlatform.cpp
@@ -465,15 +465,6 @@ void FVulkanAndroidPlatform::GetInstanceLayers(TArray<const ANSICHAR*>& OutLayer
 #endif
 }
 
-
-static int32 GVulkanQcomRenderPassTransform = 0;
-static FAutoConsoleVariableRef CVarVulkanQcomRenderPassTransform(
-	TEXT("r.Vulkan.UseQcomRenderPassTransform"),
-	GVulkanQcomRenderPassTransform,
-	TEXT("UseQcomRenderPassTransform\n"),
-	ECVF_ReadOnly
-);
-
 static int32 GVulkanUseASTCDecodeMode = 1;
 static FAutoConsoleVariableRef CVarVulkanUseASTCDecodeMode(
 	TEXT("r.Vulkan.UseASTCDecodeMode"),
@@ -492,11 +483,6 @@ void FVulkanAndroidPlatform::GetDeviceExtensions(FVulkanDevice* Device, FVulkanD
 	}
 	OutExtensions.Add(MakeUnique<FVulkanDeviceExtension>(Device, VK_EXT_TEXTURE_COMPRESSION_ASTC_HDR_EXTENSION_NAME, VULKAN_SUPPORTS_TEXTURE_COMPRESSION_ASTC_HDR, VK_API_VERSION_1_3, DEVICE_EXT_FLAG_SETTER(HasEXTTextureCompressionASTCHDR)));
 
-	if (GVulkanQcomRenderPassTransform)
-	{
-		OutExtensions.Add(MakeUnique<FVulkanDeviceExtension>(Device, VK_QCOM_RENDER_PASS_TRANSFORM_EXTENSION_NAME, VULKAN_EXTENSION_ENABLED, VK_API_VERSION_1_3, DEVICE_EXT_FLAG_SETTER(HasQcomRenderPassTransform)));
-	}
-
 #if !UE_BUILD_SHIPPING
 	// Layer name as extension
 	OutExtensions.Add(MakeUnique<FVulkanDeviceExtension>(Device, VULKAN_MALI_LAYER_NAME, VULKAN_EXTENSION_ENABLED, VULKAN_EXTENSION_NOT_PROMOTED));
diff --git a/Engine/Source/Runtime/VulkanRHI/Private/Android/VulkanAndroidPlatform.h b/Engine/Source/Runtime/VulkanRHI/Private/Android/VulkanAndroidPlatform.h
--- a/Engine/Source/Runtime/VulkanRHI/Private/Android/VulkanAndroidPlatform.h
+++ b/Engine/Source/Runtime/VulkanRHI/Private/Android/VulkanAndroidPlatform.h
@@ -244,27 +244,3 @@ private:
 #endif //VULKAN_SUPPORTS_GOOGLE_DISPLAY_TIMING
 
 typedef FVulkanAndroidPlatform FVulkanPlatform;
-
-/* VK_QCOM_render_pass_transform */
-#ifndef VK_QCOM_render_pass_transform
-#define VK_QCOM_render_pass_transform 1
-#define VK_QCOM_RENDER_PASS_TRANSFORM_SPEC_VERSION 1
-#define VK_QCOM_RENDER_PASS_TRANSFORM_EXTENSION_NAME "VK_QCOM_render_pass_transform"
-#define VK_STRUCTURE_TYPE_RENDER_PASS_TRANSFORM_BEGIN_INFO_QCOM 1000282000
-#define VK_STRUCTURE_TYPE_COMMAND_BUFFER_INHERITANCE_RENDER_PASS_TRANSFORM_INFO_QCOM 1000282001
-#define VK_RENDER_PASS_CREATE_TRANSFORM_BIT_QCOM 0x00000002
-typedef struct VkRenderPassTransformBeginInfoQCOM {
-	VkStructureType sType;
-	void* pNext;
-	VkSurfaceTransformFlagBitsKHR transform;
-} VkRenderPassTransformBeginInfoQCOM;
-
-typedef struct VkCommandBufferInheritanceRenderPassTransformInfoQCOM {
-	VkStructureType sType;
-	void* pNext;
-	VkSurfaceTransformFlagBitsKHR transform;
-	VkRect2D renderArea;
-} VkCommandBufferInheritanceRenderPassTransformInfoQCOM;
-
-#endif //VK_QCOM_render_pass_transform
-#define VULKAN_SUPPORTS_QCOM_RENDERPASS_TRANSFORM			1
diff --git a/Engine/Source/Runtime/VulkanRHI/Private/VulkanCommandBuffer.cpp b/Engine/Source/Runtime/VulkanRHI/Private/VulkanCommandBuffer.cpp
--- a/Engine/Source/Runtime/VulkanRHI/Private/VulkanCommandBuffer.cpp
+++ b/Engine/Source/Runtime/VulkanRHI/Private/VulkanCommandBuffer.cpp
@@ -192,19 +192,6 @@ void FVulkanCmdBuffer::BeginRenderPass(const FVulkanRenderTargetLayout& Layout,
 	Info.clearValueCount = Layout.GetNumUsedClearValues();
 	Info.pClearValues = AttachmentClearValues;
 
-#if VULKAN_SUPPORTS_QCOM_RENDERPASS_TRANSFORM
-	VkRenderPassTransformBeginInfoQCOM RPTransformBeginInfoQCOM;
-	VkSurfaceTransformFlagBitsKHR QCOMTransform = Layout.GetQCOMRenderPassTransform();
-
-	if (QCOMTransform != VK_SURFACE_TRANSFORM_IDENTITY_BIT_KHR)
-	{
-		ZeroVulkanStruct(RPTransformBeginInfoQCOM, (VkStructureType)VK_STRUCTURE_TYPE_RENDER_PASS_TRANSFORM_BEGIN_INFO_QCOM);
-
-		RPTransformBeginInfoQCOM.transform = QCOMTransform;
-		Info.pNext = &RPTransformBeginInfoQCOM;
-	}
-#endif
-
 	if (Device->GetOptionalExtensions().HasKHRRenderPass2)
 	{
 		VkSubpassBeginInfo SubpassInfo;
diff --git a/Engine/Source/Runtime/VulkanRHI/Private/VulkanContext.h b/Engine/Source/Runtime/VulkanRHI/Private/VulkanContext.h
--- a/Engine/Source/Runtime/VulkanRHI/Private/VulkanContext.h
+++ b/Engine/Source/Runtime/VulkanRHI/Private/VulkanContext.h
@@ -253,7 +253,6 @@ protected:
 
 public:
 	bool IsSwapchainImage(FRHITexture* InTexture) const;
-	VkSurfaceTransformFlagBitsKHR GetSwapchainQCOMRenderPassTransform() const;
 	VkFormat GetSwapchainImageFormat() const;
 	FVulkanSwapChain* GetSwapChain() const;
 
diff --git a/Engine/Source/Runtime/VulkanRHI/Private/VulkanDevice.h b/Engine/Source/Runtime/VulkanRHI/Private/VulkanDevice.h
--- a/Engine/Source/Runtime/VulkanRHI/Private/VulkanDevice.h
+++ b/Engine/Source/Runtime/VulkanRHI/Private/VulkanDevice.h
@@ -63,7 +63,6 @@ struct FOptionalVulkanDeviceExtensions
 			uint64 HasAMDBufferMarker : 1;
 			uint64 HasNVDiagnosticCheckpoints : 1;
 			uint64 HasNVDeviceDiagnosticConfig : 1;
-			uint64 HasQcomRenderPassTransform : 1;
 
 			// Promoted to 1.1
 			uint64 HasKHRMultiview : 1;
diff --git a/Engine/Source/Runtime/VulkanRHI/Private/VulkanRHIPrivate.h b/Engine/Source/Runtime/VulkanRHI/Private/VulkanRHIPrivate.h
--- a/Engine/Source/Runtime/VulkanRHI/Private/VulkanRHIPrivate.h
+++ b/Engine/Source/Runtime/VulkanRHI/Private/VulkanRHIPrivate.h
@@ -151,13 +151,11 @@ public:
 	inline const VkAttachmentDescriptionStencilLayout* GetStencilDesc() const { return bHasDepthStencil ? &StencilDesc : nullptr; }
 
 	inline const ESubpassHint GetSubpassHint() const { return SubpassHint; }
-	inline const VkSurfaceTransformFlagBitsKHR GetQCOMRenderPassTransform() const { return QCOMRenderPassTransform; }
 
 protected:
 	VkImageLayout GetVRSImageLayout() const;
 
 protected:
-	VkSurfaceTransformFlagBitsKHR QCOMRenderPassTransform = VK_SURFACE_TRANSFORM_IDENTITY_BIT_KHR;
 	VkAttachmentReference ColorReferences[MaxSimultaneousRenderTargets];
 	VkAttachmentReference DepthReference;
 	VkAttachmentReferenceStencilLayout StencilReference;
diff --git a/Engine/Source/Runtime/VulkanRHI/Private/VulkanRenderTarget.cpp b/Engine/Source/Runtime/VulkanRHI/Private/VulkanRenderTarget.cpp
--- a/Engine/Source/Runtime/VulkanRHI/Private/VulkanRenderTarget.cpp
+++ b/Engine/Source/Runtime/VulkanRHI/Private/VulkanRenderTarget.cpp
@@ -522,18 +522,6 @@ void FVulkanDynamicRHI::RHIRead3DSurfaceFloatData(FRHITexture* TextureRHI, FIntR
 	Device->GetImmediateContext().GetCommandBufferManager()->PrepareForNewActiveCommandBuffer();
 }
 
-VkSurfaceTransformFlagBitsKHR FVulkanCommandListContext::GetSwapchainQCOMRenderPassTransform() const
-{
-	TArray<FVulkanViewport*>& viewports = RHI->GetViewports();
-	if (viewports.Num() == 0)
-	{
-		return VK_SURFACE_TRANSFORM_IDENTITY_BIT_KHR;
-	}
-
-	// check(viewports.Num() == 1);
-	return viewports[0]->GetSwapchainQCOMRenderPassTransform();
-}
-
 VkFormat FVulkanCommandListContext::GetSwapchainImageFormat() const
 {
 	TArray<FVulkanViewport*>& viewports = RHI->GetViewports();
@@ -736,7 +724,6 @@ struct FRenderPassCompatibleHashableStruct
 	uint8							MultiViewCount;
 	uint8							NumSamples;
 	uint8							SubpassHint;
-	VkSurfaceTransformFlagBitsKHR	QCOMRenderPassTransform;
 	// +1 for Depth, +1 for Stencil, +1 for Fragment Density
 	VkFormat						Formats[MaxSimultaneousRenderTargets + 3];
 	uint16							AttachmentsToResolve;
@@ -802,11 +789,6 @@ FVulkanRenderTargetLayout::FVulkanRenderTargetLayout(FVulkanDevice& InDevice, co
 			check(Texture);
 			const FRHITextureDesc& TextureDesc = Texture->GetDesc();
 
-			if (InDevice.GetImmediateContext().IsSwapchainImage(RTView.Texture))
-			{
-				QCOMRenderPassTransform = InDevice.GetImmediateContext().GetSwapchainQCOMRenderPassTransform();
-			}
-
 			if (bSetExtent)
 			{
 				ensure(Extent.Extent3D.width == FMath::Max(1, TextureDesc.Extent.X >> RTView.MipIndex));
@@ -981,8 +963,6 @@ FVulkanRenderTargetLayout::FVulkanRenderTargetLayout(FVulkanDevice& InDevice, co
 	SubpassHint = ESubpassHint::None;
 	CompatibleHashInfo.SubpassHint = 0;
 
-	CompatibleHashInfo.QCOMRenderPassTransform = QCOMRenderPassTransform;
-
 	CompatibleHashInfo.NumSamples = NumSamples;
 	CompatibleHashInfo.MultiViewCount = MultiViewCount;
 
@@ -1020,12 +1000,6 @@ FVulkanRenderTargetLayout::FVulkanRenderTargetLayout(FVulkanDevice& InDevice, co
 		check(Texture);
 		const FRHITextureDesc& TextureDesc = Texture->GetDesc();
 
-		if (InDevice.GetImmediateContext().IsSwapchainImage(ColorEntry.RenderTarget))
-		{
-			QCOMRenderPassTransform = InDevice.GetImmediateContext().GetSwapchainQCOMRenderPassTransform();
-		}
-		check(QCOMRenderPassTransform == VK_SURFACE_TRANSFORM_IDENTITY_BIT_KHR || NumAttachmentDescriptions == 0);
-
 		if (bSetExtent)
 		{
 			ensure(Extent.Extent3D.width == FMath::Max(1, TextureDesc.Extent.X >> ColorEntry.MipIndex));
@@ -1233,8 +1207,6 @@ FVulkanRenderTargetLayout::FVulkanRenderTargetLayout(FVulkanDevice& InDevice, co
 	SubpassHint = RPInfo.SubpassHint;
 	CompatibleHashInfo.SubpassHint = (uint8)RPInfo.SubpassHint;
 
-	CompatibleHashInfo.QCOMRenderPassTransform = QCOMRenderPassTransform;
-
 	CompatibleHashInfo.NumSamples = NumSamples;
 	CompatibleHashInfo.MultiViewCount = MultiViewCount;
 
@@ -1395,26 +1367,6 @@ FVulkanRenderTargetLayout::FVulkanRenderTargetLayout(const FGraphicsPipelineStat
 	SubpassHint = Initializer.SubpassHint;
 	CompatibleHashInfo.SubpassHint = (uint8)Initializer.SubpassHint;
 
-	FVulkanCommandListContext& ImmediateContext = GVulkanRHI->GetDevice()->GetImmediateContext();
-
-	if (GVulkanRHI->GetDevice()->GetOptionalExtensions().HasQcomRenderPassTransform)
-	{
-		VkFormat SwapchainImageFormat = ImmediateContext.GetSwapchainImageFormat();
-		if (Desc[0].format == SwapchainImageFormat)
-		{
-			// Potential Swapchain RenderPass
-			QCOMRenderPassTransform = ImmediateContext.GetSwapchainQCOMRenderPassTransform();
-		}
-		// TODO: add some checks to detect potential Swapchain pass
-		else if (SwapchainImageFormat == VK_FORMAT_UNDEFINED)
-		{
-			// WA: to have compatible RP created with VK_RENDER_PASS_CREATE_TRANSFORM_BIT_QCOM flag
-			QCOMRenderPassTransform = VK_SURFACE_TRANSFORM_ROTATE_90_BIT_KHR;
-		}
-	}
-
-	CompatibleHashInfo.QCOMRenderPassTransform = QCOMRenderPassTransform;
-
 	CompatibleHashInfo.NumSamples = NumSamples;
 	CompatibleHashInfo.MultiViewCount = MultiViewCount;
 
diff --git a/Engine/Source/Runtime/VulkanRHI/Private/VulkanRenderpass.h b/Engine/Source/Runtime/VulkanRHI/Private/VulkanRenderpass.h
--- a/Engine/Source/Runtime/VulkanRHI/Private/VulkanRenderpass.h
+++ b/Engine/Source/Runtime/VulkanRHI/Private/VulkanRenderpass.h
@@ -720,13 +720,6 @@ public:
 			FragDensityCreateInfo.pNext = CreateInfo.pNext;
 			CreateInfo.pNext = &FragDensityCreateInfo;
 		}
-
-#if VULKAN_SUPPORTS_QCOM_RENDERPASS_TRANSFORM
-		if (RTLayout.GetQCOMRenderPassTransform() != VK_SURFACE_TRANSFORM_IDENTITY_BIT_KHR)
-		{
-			CreateInfo.flags = VK_RENDER_PASS_CREATE_TRANSFORM_BIT_QCOM;
-		}
-#endif
 	}
 
 	VkRenderPass Create(const FVulkanRenderTargetLayout& RTLayout)
diff --git a/Engine/Source/Runtime/VulkanRHI/Private/VulkanSwapChain.cpp b/Engine/Source/Runtime/VulkanRHI/Private/VulkanSwapChain.cpp
--- a/Engine/Source/Runtime/VulkanRHI/Private/VulkanSwapChain.cpp
+++ b/Engine/Source/Runtime/VulkanRHI/Private/VulkanSwapChain.cpp
@@ -70,6 +70,13 @@ static FAutoConsoleVariableRef CVarVulkanForcePacingWithoutVSync(
 	ECVF_RenderThreadSafe
 );
 
+static TAutoConsoleVariable<int32> CVarVulkanPreTransform(
+	TEXT("r.Vulkan.PreTransform"),
+	1,
+	TEXT("0: Surface transformed in Vulkan Presentation Engine")
+	TEXT("1: Surface transformed in App to save bandwidth (Default)"),
+	ECVF_ReadOnly);
+
 int32 GPrintVulkanVsyncDebug = 0;
 #if !(UE_BUILD_SHIPPING)
 static FAutoConsoleVariableRef CVarVulkanDebugVsync(
@@ -132,6 +139,14 @@ VkResult SimulateErrors(VkResult Result)
 extern TAutoConsoleVariable<int32> GAllowPresentOnComputeQueue;
 static TSet<EPixelFormat> GPixelFormatNotSupportedWarning;
 
+bool IsVulkanPreTransformEnabled(const FStaticShaderPlatform Platform)
+{
+	return (IsVulkanMobilePlatform(Platform) || IsVulkanMobileSM5Platform(Platform))
+			// Request an intermediate image as the BackBuffer
+			&& GVulkanDelayAcquireImage == EDelayAcquireImageType::DelayAcquire
+			&& CVarVulkanPreTransform.GetValueOnAnyThread() > 0;
+}
+
 FVulkanSwapChain::FVulkanSwapChain(VkInstance InInstance, FVulkanDevice& InDevice, void* InWindowHandle, EPixelFormat& InOutPixelFormat, uint32 Width, uint32 Height, bool bIsFullScreen,
 	uint32* InOutDesiredNumBackBuffers, TArray<VkImage>& OutImages, int8 InLockToVsync, FVulkanSwapChainRecreateInfo* RecreateInfo)
 	: SwapChain(VK_NULL_HANDLE)
@@ -439,8 +454,8 @@ FVulkanSwapChain::FVulkanSwapChain(VkInstance InInstance, FVulkanDevice& InDevic
 	VERIFYVULKANRESULT_EXPANDED(VulkanRHI::vkGetPhysicalDeviceSurfaceCapabilitiesKHR(Device.GetPhysicalHandle(),
 		Surface,
 		&SurfProperties));
-	VkSurfaceTransformFlagBitsKHR PreTransform;
-	if (SurfProperties.supportedTransforms & VK_SURFACE_TRANSFORM_IDENTITY_BIT_KHR)
+
+	if (!IsVulkanPreTransformEnabled(GMaxRHIShaderPlatform))
 	{
 		PreTransform = VK_SURFACE_TRANSFORM_IDENTITY_BIT_KHR;
 	}
@@ -495,18 +510,6 @@ FVulkanSwapChain::FVulkanSwapChain(VkInstance InInstance, FVulkanDevice& InDevic
 		}
 	}
 
-	if (Device.GetOptionalExtensions().HasQcomRenderPassTransform)
-	{
-		QCOMRenderPassTransform = SurfProperties.currentTransform;
-		SwapChainInfo.preTransform = QCOMRenderPassTransform;
-		ImageFormat = SwapChainInfo.imageFormat;
-		if (SwapChainInfo.preTransform == VK_SURFACE_TRANSFORM_ROTATE_90_BIT_KHR ||
-			SwapChainInfo.preTransform == VK_SURFACE_TRANSFORM_ROTATE_270_BIT_KHR)
-		{
-			Swap(SwapChainInfo.imageExtent.width, SwapChainInfo.imageExtent.height);
-		}
-	}
-
 	VkBool32 bSupportsPresent;
 	VERIFYVULKANRESULT(VulkanRHI::vkGetPhysicalDeviceSurfaceSupportKHR(Device.GetPhysicalHandle(), Device.GetPresentQueue()->GetFamilyIndex(), Surface, &bSupportsPresent));
 	ensure(bSupportsPresent);
@@ -535,7 +538,20 @@ FVulkanSwapChain::FVulkanSwapChain(VkInstance InInstance, FVulkanDevice& InDevic
 	}
 #endif
 
+	const bool bSwapChainNeedsTransform = PreTransform == VK_SURFACE_TRANSFORM_ROTATE_90_BIT_KHR || PreTransform == VK_SURFACE_TRANSFORM_ROTATE_270_BIT_KHR;
+	if (bSwapChainNeedsTransform)
+	{
+		Swap(SwapChainInfo.imageExtent.width, SwapChainInfo.imageExtent.height);
+	}
+
 	VkResult Result = FVulkanPlatform::CreateSwapchainKHR(WindowHandle, Device.GetPhysicalHandle(), Device.GetInstanceHandle(), &SwapChainInfo, VULKAN_CPU_ALLOCATOR, &SwapChain);
+
+	// Swap the aspect ratio back to make it transparent to high level
+	if (bSwapChainNeedsTransform)
+	{
+		Swap(SwapChainInfo.imageExtent.width, SwapChainInfo.imageExtent.height);
+	}
+
 #if VULKAN_SUPPORTS_FULLSCREEN_EXCLUSIVE
 	if (Device.GetOptionalExtensions().HasEXTFullscreenExclusive && Result == VK_ERROR_INITIALIZATION_FAILED)
 	{
@@ -633,19 +649,6 @@ void FVulkanSwapChain::Destroy(FVulkanSwapChainRecreateInfo* RecreateInfo)
 		VulkanRHI::vkDestroySurfaceKHR(Instance, Surface, VULKAN_CPU_ALLOCATOR);
 	}
 
-	if (QCOMDepthView && QCOMDepthView != QCOMDepthStencilView)
-	{
-		delete QCOMDepthView;
-		QCOMDepthView = nullptr;
-	}
-
-	if (QCOMDepthStencilView)
-	{
-		delete QCOMDepthStencilView;
-		QCOMDepthStencilView = nullptr;
-		QCOMDepthView = nullptr;
-	}
-
 	Surface = VK_NULL_HANDLE;
 }
 
@@ -877,95 +880,6 @@ FVulkanSwapChain::EStatus FVulkanSwapChain::Present(FVulkanQueue* GfxQueue, FVul
 	return EStatus::Healthy;
 }
 
-void FVulkanSwapChain::CreateQCOMDepthStencil(const FVulkanTexture& InSurface) const
-{
-	check(!QCOMDepthStencilSurface);
-	check(!QCOMDepthStencilView);
-	check(!QCOMDepthView);
-
-	const FRHITextureDesc& Desc = InSurface.GetDesc();
-	const ETextureCreateFlags UEFlags = Desc.Flags;
-	check(UEFlags & TexCreate_DepthStencilTargetable);
-	const VkDescriptorType DescriptorType = VK_DESCRIPTOR_TYPE_SAMPLED_IMAGE;
-
-	const FRHITextureCreateDesc CreateDesc =
-		FRHITextureCreateDesc::Create2D(TEXT("FVulkanSwapChainQCOM"), Desc.Extent.Y, Desc.Extent.X, Desc.Format) // Desc.Extent.X and Desc.Extent.Y are intentionally swapped.
-		.SetClearValue(FClearValueBinding::None)
-		.SetFlags(UEFlags)
-		.SetNumMips(Desc.NumMips)
-		.SetNumSamples(Desc.NumSamples)
-		.DetermineInititialState();
-
-	QCOMDepthStencilSurface = new FVulkanTexture(Device, CreateDesc, nullptr);
-
-	check(QCOMDepthStencilSurface->GetViewType() == VK_IMAGE_VIEW_TYPE_2D);
-	check(QCOMDepthStencilSurface->Image != VK_NULL_HANDLE);
-
-	QCOMDepthStencilView = new FVulkanView(*QCOMDepthStencilSurface->Device, DescriptorType);
-	QCOMDepthStencilView->InitAsTextureView(
-		  QCOMDepthStencilSurface->Image
-		, QCOMDepthStencilSurface->GetViewType()
-		, QCOMDepthStencilSurface->GetFullAspectMask()
-		, QCOMDepthStencilSurface->GetDesc().Format
-		, QCOMDepthStencilSurface->ViewFormat
-		, 0
-		, FMath::Max(QCOMDepthStencilSurface->GetNumMips(), 1u)
-		, 0
-		, 1u
-		, false
-	);
-
-	if (QCOMDepthStencilSurface->GetFullAspectMask() == QCOMDepthStencilSurface->GetPartialAspectMask())
-	{
-		QCOMDepthView = QCOMDepthStencilView;
-	}
-	else
-	{
-		QCOMDepthView = new FVulkanView(*QCOMDepthStencilSurface->Device, DescriptorType);
-		QCOMDepthView->InitAsTextureView(
-			  QCOMDepthStencilSurface->Image
-			, QCOMDepthStencilSurface->GetViewType()
-			, QCOMDepthStencilSurface->GetPartialAspectMask()
-			, QCOMDepthStencilSurface->GetDesc().Format
-			, QCOMDepthStencilSurface->ViewFormat
-			, 0
-			, FMath::Max(QCOMDepthStencilSurface->GetNumMips(), 1u)
-			, 0
-			, 1u
-			, false
-		);
-	}
-}
-
-const FVulkanView* FVulkanSwapChain::GetOrCreateQCOMDepthStencilView(const FVulkanTexture& InSurface) const
-{
-	if (QCOMDepthStencilView)
-	{
-		return QCOMDepthStencilView;
-	}
-
-	CreateQCOMDepthStencil(InSurface);
-
-	return QCOMDepthStencilView;
-}
-
-const FVulkanView* FVulkanSwapChain::GetOrCreateQCOMDepthView(const FVulkanTexture& InSurface) const
-{
-	if (QCOMDepthView)
-	{
-		return QCOMDepthView;
-	}
-
-	CreateQCOMDepthStencil(InSurface);
-
-	return QCOMDepthView;
-}
-
-const FVulkanTexture* FVulkanSwapChain::GetQCOMDepthStencilSurface() const
-{
-	return QCOMDepthStencilSurface;
-}
-
 void FVulkanDevice::SetupPresentQueue(VkSurfaceKHR Surface)
 {
 	if (!PresentQueue)
diff --git a/Engine/Source/Runtime/VulkanRHI/Private/VulkanSwapChain.h b/Engine/Source/Runtime/VulkanRHI/Private/VulkanSwapChain.h
--- a/Engine/Source/Runtime/VulkanRHI/Private/VulkanSwapChain.h
+++ b/Engine/Source/Runtime/VulkanRHI/Private/VulkanSwapChain.h
@@ -50,12 +50,10 @@ public:
 	void RenderThreadPacing();
 	inline int8 DoesLockToVsync() { return LockToVsync; }
 
-	const FVulkanView* GetOrCreateQCOMDepthStencilView(const FVulkanTexture& InSurface) const;
-	const FVulkanView* GetOrCreateQCOMDepthView(const FVulkanTexture& InSurface) const;
-	const FVulkanTexture* GetQCOMDepthStencilSurface() const;
+	inline VkSurfaceTransformFlagBitsKHR GetCachedSurfaceTransform() const { return PreTransform; }
 
 protected:
-	VkSurfaceTransformFlagBitsKHR QCOMRenderPassTransform = VK_SURFACE_TRANSFORM_IDENTITY_BIT_KHR;
+	VkSurfaceTransformFlagBitsKHR PreTransform;
 	VkFormat ImageFormat = VK_FORMAT_UNDEFINED;
 
 	VkSwapchainKHR SwapChain;
@@ -89,14 +87,6 @@ protected:
 
 	int32 AcquireImageIndex(VulkanRHI::FSemaphore** OutSemaphore);
 
-	// WA: if the swapchain pass uses a depth target, it must have same size as the swapchain images.
-	// For example in case if QCOMRenderPassTransform is VK_SURFACE_TRANSFORM_ROTATE_90_BIT_KHR we must swap width/height of depth target.
-	// But probably engine can use same depth texture for swapchain and NON swapchain passes. So it is why we have this addditional surface.
-	// With this approach we should be careful in case if depth in swapchain pass is used as attachement and fetched in shader in same time
-	void CreateQCOMDepthStencil(const FVulkanTexture& InSurface) const;
-	mutable FVulkanTexture* QCOMDepthStencilSurface = nullptr;
-	mutable FVulkanView* QCOMDepthStencilView = nullptr;
-	mutable FVulkanView* QCOMDepthView = nullptr;
 
 	friend class FVulkanViewport;
 	friend class FVulkanQueue;
diff --git a/Engine/Source/Runtime/VulkanRHI/Private/VulkanTexture.cpp b/Engine/Source/Runtime/VulkanRHI/Private/VulkanTexture.cpp
--- a/Engine/Source/Runtime/VulkanRHI/Private/VulkanTexture.cpp
+++ b/Engine/Source/Runtime/VulkanRHI/Private/VulkanTexture.cpp
@@ -725,6 +725,10 @@ void FVulkanTexture::DestroySurface()
 				Image = VK_NULL_HANDLE;
 			}
 		}
+		else
+		{
+			Image = VK_NULL_HANDLE;
+		}
 
 		ImageOwnerType = EImageOwnerType::None;
 	}
diff --git a/Engine/Source/Runtime/VulkanRHI/Private/VulkanViewport.cpp b/Engine/Source/Runtime/VulkanRHI/Private/VulkanViewport.cpp
--- a/Engine/Source/Runtime/VulkanRHI/Private/VulkanViewport.cpp
+++ b/Engine/Source/Runtime/VulkanRHI/Private/VulkanViewport.cpp
@@ -13,6 +13,9 @@
 #include "HAL/PlatformAtomics.h"
 #include "Engine/RendererSettings.h"
 #include "StereoRenderUtils.h"
+#include "CommonRenderResources.h"
+#include "ScreenRendering.h"
+#include "RHIStaticStates.h"
 
 FVulkanBackBuffer::FVulkanBackBuffer(FVulkanDevice& Device, FVulkanViewport* InViewport, EPixelFormat Format, uint32 SizeX, uint32 SizeY, ETextureCreateFlags UEFlags)
 	: FVulkanTexture(Device, FRHITextureCreateDesc::Create2D(TEXT("FVulkanBackBuffer"), SizeX, SizeY, Format).SetFlags(UEFlags).DetermineInititialState(), VK_NULL_HANDLE, false)
@@ -163,8 +166,8 @@ FVulkanViewport::~FVulkanViewport()
 			RenderingDoneSemaphores[Index]->Release();
 
 			// FIXME: race condition on TransitionAndLayoutManager, could this be called from RT while RHIT is active?
-			Device->NotifyDeletedImage(BackBufferImages[Index], true);
-			BackBufferImages[Index] = VK_NULL_HANDLE;
+			Device->NotifyDeletedImage(BackBufferImages[Index]->Image, true);
+			BackBufferImages[Index] = nullptr;
 		}
 
 		SwapChain->Destroy(nullptr);
@@ -453,8 +456,6 @@ FVulkanFramebuffer::FVulkanFramebuffer(FVulkanDevice& Device, const FRHISetRende
 		}
 	}
 
-	VkSurfaceTransformFlagBitsKHR QCOMRenderPassTransform = RTLayout.GetQCOMRenderPassTransform();
-
 	if (RTLayout.GetHasDepthStencil())
 	{
 		FVulkanTexture* Texture = ResourceCast(InRTInfo.DepthStencilRenderTarget.Texture);
@@ -499,13 +500,6 @@ FVulkanFramebuffer::FVulkanFramebuffer(FVulkanDevice& Device, const FRHISetRende
 				, true
 			);
 		}
-		else if (QCOMRenderPassTransform != VK_SURFACE_TRANSFORM_IDENTITY_BIT_KHR && Desc.Extent.X == RTExtents.width && Desc.Extent.Y == RTExtents.height)
-		{
-			FVulkanSwapChain* SwapChain = Device.GetImmediateContext().GetSwapChain();
-
-			PartialDepthTextureView = SwapChain->GetOrCreateQCOMDepthView(*Texture);
-			AddExternalView(SwapChain->GetOrCreateQCOMDepthStencilView(*Texture));
-		}
 		else
 		{
 			AddExternalView(Texture->DefaultView);
@@ -549,12 +543,6 @@ FVulkanFramebuffer::FVulkanFramebuffer(FVulkanDevice& Device, const FRHISetRende
 	CreateInfo.height = RTExtents.height;
 	CreateInfo.layers = NumLayers;
 
-	if (QCOMRenderPassTransform == VK_SURFACE_TRANSFORM_ROTATE_90_BIT_KHR ||
-		QCOMRenderPassTransform == VK_SURFACE_TRANSFORM_ROTATE_270_BIT_KHR)
-	{
-		Swap(CreateInfo.width, CreateInfo.height);
-	}
-
 	VERIFYVULKANRESULT_EXPANDED(VulkanRHI::vkCreateFramebuffer(Device.GetInstanceHandle(), &CreateInfo, VULKAN_CPU_ALLOCATOR, &Framebuffer));
 
 	RenderArea.offset.x = 0;
@@ -690,6 +678,7 @@ void FVulkanViewport::RecreateSwapchainFromRT(EPixelFormat PreferredPixelFormat)
 	check(IsInRenderingThread());
 
 	// TODO: should flush RHIT commands here?
+	FRHICommandListExecutor::GetImmediateCommandList().ImmediateFlush(EImmediateFlushType::FlushRHIThread);
 	
 	FVulkanSwapChainRecreateInfo RecreateInfo = { VK_NULL_HANDLE, VK_NULL_HANDLE};
 	DestroySwapchain(&RecreateInfo);
@@ -733,7 +722,14 @@ void FVulkanViewport::CreateSwapchain(FVulkanSwapChainRecreateInfo* RecreateInfo
 		const FVulkanImageLayout InitialLayout(VK_IMAGE_LAYOUT_PRESENT_SRC_KHR, 1, 1, VK_IMAGE_ASPECT_COLOR_BIT);
 		for (int32 Index = 0; Index < Images.Num(); ++Index)
 		{
-			BackBufferImages[Index] = Images[Index];
+			uint32 ImageSizeX = SizeX;
+			uint32 ImageSizeY = SizeY;
+			VkSurfaceTransformFlagBitsKHR CachedSurfaceTransform = SwapChain->GetCachedSurfaceTransform();
+			if (CachedSurfaceTransform == VK_SURFACE_TRANSFORM_ROTATE_90_BIT_KHR || CachedSurfaceTransform == VK_SURFACE_TRANSFORM_ROTATE_270_BIT_KHR)
+			{
+				Swap(ImageSizeX, ImageSizeY);
+			}
+			BackBufferImages[Index] = (FVulkanTexture*)FVulkanDynamicRHI::Get().RHICreateTexture2DFromResource(PixelFormat, ImageSizeX, ImageSizeY, 1, 1, Images[Index], ETextureCreateFlags::RenderTargetable | ETextureCreateFlags::Presentable).GetReference();
 			const VkDescriptorType DescriptorType = VK_DESCRIPTOR_TYPE_SAMPLED_IMAGE;
 			TextureViews.Add((new FVulkanView(*Device, DescriptorType))->InitAsTextureView(
 				Images[Index], VK_IMAGE_VIEW_TYPE_2D, VK_IMAGE_ASPECT_COLOR_BIT, PixelFormat, UEToVkTextureFormat(PixelFormat, false), 0, 1, 0, 1, false));
@@ -749,7 +745,7 @@ void FVulkanViewport::CreateSwapchain(FVulkanSwapChainRecreateInfo* RecreateInfo
 #if VULKAN_ENABLE_DRAW_MARKERS
 			if (Device->GetSetDebugName())
 			{
-				VulkanRHI::SetDebugName(Device->GetSetDebugName(), Device->GetInstanceHandle(), BackBufferImages[Index], "RenderingBackBuffer");
+				VulkanRHI::SetDebugName(Device->GetSetDebugName(), Device->GetInstanceHandle(), BackBufferImages[Index]->Image, "RenderingBackBuffer");
 			}
 #endif
 		}
@@ -823,8 +819,8 @@ void FVulkanViewport::DestroySwapchain(FVulkanSwapChainRecreateInfo* RecreateInf
 		TextureViews.Empty();
 		for (int32 Index = 0, NumBuffers = BackBufferImages.Num(); Index < NumBuffers; ++Index)
 		{
-			Device->NotifyDeletedImage(BackBufferImages[Index], true);
-			BackBufferImages[Index] = VK_NULL_HANDLE;
+			Device->NotifyDeletedImage(BackBufferImages[Index]->Image, true);
+			BackBufferImages[Index] = nullptr;
 		}
 		
 		Device->GetDeferredDeletionQueue().ReleaseResources(true);
@@ -839,72 +835,121 @@ void FVulkanViewport::DestroySwapchain(FVulkanSwapChainRecreateInfo* RecreateInf
 	AcquiredImageIndex = -1;
 }
 
-inline static void CopyImageToBackBuffer(FVulkanCommandListContext* Context, FVulkanCmdBuffer* CmdBuffer, FVulkanTexture& SrcSurface, VkImage DstSurface, int32 SizeX, int32 SizeY, int32 WindowSizeX, int32 WindowSizeY)
+inline static void CopyImageToBackBuffer(FVulkanCommandListContext* Context, FVulkanCmdBuffer* CmdBuffer, FVulkanTexture& SrcSurface, FVulkanTexture& DstSurface, int32 SizeX, int32 SizeY, int32 WindowSizeX, int32 WindowSizeY, VkSurfaceTransformFlagBitsKHR CachedSurfaceTransform)
 {
 	RHI_BREADCRUMB_EVENT(*Context, CopyImageToBackBuffer);
+	const bool bNeedsVulkanPreTransform = CachedSurfaceTransform != VK_SURFACE_TRANSFORM_IDENTITY_BIT_KHR;
 
-	FVulkanLayoutManager& LayoutManager = CmdBuffer->GetLayoutManager();
-	const VkImageLayout PreviousSrcLayout = FVulkanLayoutManager::SetExpectedLayout(CmdBuffer, SrcSurface, ERHIAccess::CopySrc);
+	VkImageLayout SrcSurfaceLayout = VK_IMAGE_LAYOUT_TRANSFER_SRC_OPTIMAL;
+	VkImageLayout DstSurfaceLayout = VK_IMAGE_LAYOUT_TRANSFER_DST_OPTIMAL;
+
+	if (bNeedsVulkanPreTransform)
+	{
+		SrcSurfaceLayout = VK_IMAGE_LAYOUT_SHADER_READ_ONLY_OPTIMAL;
+		DstSurfaceLayout = VK_IMAGE_LAYOUT_COLOR_ATTACHMENT_OPTIMAL;
+	}
 
 	{
 		FVulkanPipelineBarrier Barrier;
-		Barrier.AddImageLayoutTransition(DstSurface, VK_IMAGE_LAYOUT_PRESENT_SRC_KHR, VK_IMAGE_LAYOUT_TRANSFER_DST_OPTIMAL, FVulkanPipelineBarrier::MakeSubresourceRange(VK_IMAGE_ASPECT_COLOR_BIT, 0, 1));
+		Barrier.AddImageLayoutTransition(SrcSurface.Image, VK_IMAGE_LAYOUT_COLOR_ATTACHMENT_OPTIMAL, SrcSurfaceLayout, FVulkanPipelineBarrier::MakeSubresourceRange(VK_IMAGE_ASPECT_COLOR_BIT, 0, 1));
+		Barrier.AddImageLayoutTransition(DstSurface.Image, VK_IMAGE_LAYOUT_PRESENT_SRC_KHR, DstSurfaceLayout, FVulkanPipelineBarrier::MakeSubresourceRange(VK_IMAGE_ASPECT_COLOR_BIT, 0, 1));
 		Barrier.Execute(CmdBuffer);
 	}
 
 	VulkanRHI::DebugHeavyWeightBarrier(CmdBuffer->GetHandle(), 32);
 
-	if (SizeX != WindowSizeX || SizeY != WindowSizeY)
+	// Copy and rotate the intermediate image to the BackBuffer with a pixel shader
+	if (bNeedsVulkanPreTransform)
 	{
-		VkImageBlit Region;
-		FMemory::Memzero(Region);
-		Region.srcOffsets[0].x = 0;
-		Region.srcOffsets[0].y = 0;
-		Region.srcOffsets[0].z = 0;
-		Region.srcOffsets[1].x = SizeX;
-		Region.srcOffsets[1].y = SizeY;
-		Region.srcOffsets[1].z = 1;
-		Region.srcSubresource.aspectMask = VK_IMAGE_ASPECT_COLOR_BIT;
-		Region.srcSubresource.layerCount = 1;
-		Region.dstOffsets[0].x = 0;
-		Region.dstOffsets[0].y = 0;
-		Region.dstOffsets[0].z = 0;
-		Region.dstOffsets[1].x = WindowSizeX;
-		Region.dstOffsets[1].y = WindowSizeY;
-		Region.dstOffsets[1].z = 1;
-		Region.dstSubresource.aspectMask = VK_IMAGE_ASPECT_COLOR_BIT;
-		Region.dstSubresource.baseArrayLayer = 0;
-		Region.dstSubresource.layerCount = 1;
-		VulkanRHI::vkCmdBlitImage(CmdBuffer->GetHandle(),
-			SrcSurface.Image, VK_IMAGE_LAYOUT_TRANSFER_SRC_OPTIMAL,
-			DstSurface, VK_IMAGE_LAYOUT_TRANSFER_DST_OPTIMAL,
-			1, &Region, VK_FILTER_LINEAR);
+		FGraphicsPipelineStateInitializer GraphicsPSOInit;
+		// No alpha blending, no depth tests or writes, no stencil tests or writes, no backface culling.
+		GraphicsPSOInit.BlendState = TStaticBlendState<>::GetRHI();
+		GraphicsPSOInit.RasterizerState = TStaticRasterizerState<>::GetRHI();
+		GraphicsPSOInit.DepthStencilState = TStaticDepthStencilState<false>::GetRHI();
+
+		TRHICommandList_RecursiveHazardous<FVulkanCommandListContext> RHICmdList(Context);
+
+		RHICmdList.BeginRenderPass(FRHIRenderPassInfo(&DstSurface, ERenderTargetActions::DontLoad_Store), TEXT("SurfaceTransform"));
+
+		auto ShaderMap = GetGlobalShaderMap(GMaxRHIFeatureLevel);
+		TShaderMapRef<FImagePreTransformVS> VertexShader(ShaderMap);
+		TShaderMapRef<FScreenPS> PixelShader(ShaderMap);
+
+		RHICmdList.ApplyCachedRenderTargets(GraphicsPSOInit);
+		GraphicsPSOInit.BoundShaderState.VertexShaderRHI = VertexShader.GetVertexShader();
+		GraphicsPSOInit.BoundShaderState.PixelShaderRHI = PixelShader.GetPixelShader();
+		GraphicsPSOInit.BoundShaderState.VertexDeclarationRHI = GFilterVertexDeclaration.VertexDeclarationRHI;
+		GraphicsPSOInit.PrimitiveType = PT_TriangleStrip;
+		SetGraphicsPipelineState(RHICmdList, GraphicsPSOInit, 0);
+
+		FImagePreTransformVS::FParameters VSParameters;
+		FMatrix44f RenderPassTransformMatrix = FRotationMatrix44f(FRotator3f(0.0f, -180.0f * (FMath::Log2(static_cast<float>(CachedSurfaceTransform))) / 2, 0.0f));
+		VSParameters.PreTransform.X = RenderPassTransformMatrix.M[0][0];
+		VSParameters.PreTransform.Y = RenderPassTransformMatrix.M[0][1];
+		VSParameters.PreTransform.Z = RenderPassTransformMatrix.M[1][0];
+		VSParameters.PreTransform.W = RenderPassTransformMatrix.M[1][1];
+
+		SetShaderParameters(RHICmdList, VertexShader, VertexShader.GetVertexShader(), VSParameters);
+		SetShaderParametersLegacyPS(RHICmdList, PixelShader, TStaticSamplerState<SF_Point>::GetRHI(), &SrcSurface);
+
+		RHICmdList.DrawPrimitive(0, 2, 1);
+
+		RHICmdList.EndRenderPass();
 	}
 	else
 	{
-		VkImageCopy Region;
-		FMemory::Memzero(Region);
-		Region.extent.width = SizeX;
-		Region.extent.height = SizeY;
-		Region.extent.depth = 1;
-		Region.srcSubresource.aspectMask = VK_IMAGE_ASPECT_COLOR_BIT;
-		//Region.srcSubresource.baseArrayLayer = 0;
-		Region.srcSubresource.layerCount = 1;
-		//Region.srcSubresource.mipLevel = 0;
-		Region.dstSubresource.aspectMask = VK_IMAGE_ASPECT_COLOR_BIT;
-		//Region.dstSubresource.baseArrayLayer = 0;
-		Region.dstSubresource.layerCount = 1;
-		//Region.dstSubresource.mipLevel = 0;
-		VulkanRHI::vkCmdCopyImage(CmdBuffer->GetHandle(),
-			SrcSurface.Image, VK_IMAGE_LAYOUT_TRANSFER_SRC_OPTIMAL,
-			DstSurface, VK_IMAGE_LAYOUT_TRANSFER_DST_OPTIMAL,
-			1, &Region);
+		if (SizeX != WindowSizeX || SizeY != WindowSizeY)
+		{
+			VkImageBlit Region;
+			FMemory::Memzero(Region);
+			Region.srcOffsets[0].x = 0;
+			Region.srcOffsets[0].y = 0;
+			Region.srcOffsets[0].z = 0;
+			Region.srcOffsets[1].x = SizeX;
+			Region.srcOffsets[1].y = SizeY;
+			Region.srcOffsets[1].z = 1;
+			Region.srcSubresource.aspectMask = VK_IMAGE_ASPECT_COLOR_BIT;
+			Region.srcSubresource.layerCount = 1;
+			Region.dstOffsets[0].x = 0;
+			Region.dstOffsets[0].y = 0;
+			Region.dstOffsets[0].z = 0;
+			Region.dstOffsets[1].x = WindowSizeX;
+			Region.dstOffsets[1].y = WindowSizeY;
+			Region.dstOffsets[1].z = 1;
+			Region.dstSubresource.aspectMask = VK_IMAGE_ASPECT_COLOR_BIT;
+			Region.dstSubresource.baseArrayLayer = 0;
+			Region.dstSubresource.layerCount = 1;
+			VulkanRHI::vkCmdBlitImage(CmdBuffer->GetHandle(),
+				SrcSurface.Image, VK_IMAGE_LAYOUT_TRANSFER_SRC_OPTIMAL,
+				DstSurface.Image, VK_IMAGE_LAYOUT_TRANSFER_DST_OPTIMAL,
+				1, &Region, VK_FILTER_LINEAR);
+		}
+		else
+		{
+			VkImageCopy Region;
+			FMemory::Memzero(Region);
+			Region.extent.width = SizeX;
+			Region.extent.height = SizeY;
+			Region.extent.depth = 1;
+			Region.srcSubresource.aspectMask = VK_IMAGE_ASPECT_COLOR_BIT;
+			//Region.srcSubresource.baseArrayLayer = 0;
+			Region.srcSubresource.layerCount = 1;
+			//Region.srcSubresource.mipLevel = 0;
+			Region.dstSubresource.aspectMask = VK_IMAGE_ASPECT_COLOR_BIT;
+			//Region.dstSubresource.baseArrayLayer = 0;
+			Region.dstSubresource.layerCount = 1;
+			//Region.dstSubresource.mipLevel = 0;
+			VulkanRHI::vkCmdCopyImage(CmdBuffer->GetHandle(),
+				SrcSurface.Image, VK_IMAGE_LAYOUT_TRANSFER_SRC_OPTIMAL,
+				DstSurface.Image, VK_IMAGE_LAYOUT_TRANSFER_DST_OPTIMAL,
+				1, &Region);
+		}
 	}
 
 	{
 		FVulkanPipelineBarrier Barrier;
-		Barrier.AddImageLayoutTransition(SrcSurface.Image, VK_IMAGE_LAYOUT_TRANSFER_SRC_OPTIMAL, PreviousSrcLayout, FVulkanPipelineBarrier::MakeSubresourceRange(VK_IMAGE_ASPECT_COLOR_BIT, 0, 1));
-		Barrier.AddImageLayoutTransition(DstSurface, VK_IMAGE_LAYOUT_TRANSFER_DST_OPTIMAL, VK_IMAGE_LAYOUT_PRESENT_SRC_KHR, FVulkanPipelineBarrier::MakeSubresourceRange(VK_IMAGE_ASPECT_COLOR_BIT, 0, 1));
+		Barrier.AddImageLayoutTransition(SrcSurface.Image, SrcSurfaceLayout, VK_IMAGE_LAYOUT_COLOR_ATTACHMENT_OPTIMAL, FVulkanPipelineBarrier::MakeSubresourceRange(VK_IMAGE_ASPECT_COLOR_BIT, 0, 1));
+		Barrier.AddImageLayoutTransition(DstSurface.Image, DstSurfaceLayout, VK_IMAGE_LAYOUT_PRESENT_SRC_KHR, FVulkanPipelineBarrier::MakeSubresourceRange(VK_IMAGE_ASPECT_COLOR_BIT, 0, 1));
 		Barrier.Execute(CmdBuffer);
 	}
 }
@@ -928,7 +973,7 @@ bool FVulkanViewport::Present(FVulkanCommandListContext* Context, FVulkanCmdBuff
 				uint32 WindowSizeX = FMath::Min(SizeX, SwapChain->InternalWidth);
 				uint32 WindowSizeY = FMath::Min(SizeY, SwapChain->InternalHeight);
 
-				CopyImageToBackBuffer(Context, CmdBuffer, *RenderingBackBuffer.GetReference(), BackBufferImages[AcquiredImageIndex], SizeX, SizeY, WindowSizeX, WindowSizeY);
+				CopyImageToBackBuffer(Context, CmdBuffer, *RenderingBackBuffer.GetReference(), *BackBufferImages[AcquiredImageIndex].GetReference(), SizeX, SizeY, WindowSizeX, WindowSizeY, SwapChain->GetCachedSurfaceTransform());
 			}
 			else
 			{
@@ -939,22 +984,22 @@ bool FVulkanViewport::Present(FVulkanCommandListContext* Context, FVulkanCmdBuff
 		{
 			if (AcquiredImageIndex != -1)
 			{
-				check(RHIBackBuffer != nullptr && RHIBackBuffer->Image == BackBufferImages[AcquiredImageIndex]);
+				check(RHIBackBuffer != nullptr && RHIBackBuffer->Image == BackBufferImages[AcquiredImageIndex]->Image);
 
 				FVulkanLayoutManager& LayoutManager = CmdBuffer->GetLayoutManager();
-				const FVulkanImageLayout* TrackedLayout = LayoutManager.GetFullLayout(BackBufferImages[AcquiredImageIndex]);
+				const FVulkanImageLayout* TrackedLayout = LayoutManager.GetFullLayout(BackBufferImages[AcquiredImageIndex]->Image);
 
 				// One of the rare cases we'll let parallel rendering check the tracking (legacy path already checked as a fallback)
 				if (!TrackedLayout && Device->SupportsParallelRendering())
 				{
-					TrackedLayout = Context->GetQueue()->GetLayoutManager().GetFullLayout(BackBufferImages[AcquiredImageIndex]);
+					TrackedLayout = Context->GetQueue()->GetLayoutManager().GetFullLayout(BackBufferImages[AcquiredImageIndex]->Image);
 				}
 
 				VkImageLayout LastLayout = TrackedLayout ? TrackedLayout->MainLayout : VK_IMAGE_LAYOUT_UNDEFINED;
-				VulkanSetImageLayout(CmdBuffer, BackBufferImages[AcquiredImageIndex], LastLayout, VK_IMAGE_LAYOUT_PRESENT_SRC_KHR, FVulkanPipelineBarrier::MakeSubresourceRange(VK_IMAGE_ASPECT_COLOR_BIT));
+				VulkanSetImageLayout(CmdBuffer, BackBufferImages[AcquiredImageIndex]->Image, LastLayout, VK_IMAGE_LAYOUT_PRESENT_SRC_KHR, FVulkanPipelineBarrier::MakeSubresourceRange(VK_IMAGE_ASPECT_COLOR_BIT));
 
 				const FVulkanImageLayout UndefinedLayout(VK_IMAGE_LAYOUT_UNDEFINED, 1, 1, VK_IMAGE_ASPECT_COLOR_BIT);
-				LayoutManager.SetFullLayout(BackBufferImages[AcquiredImageIndex], UndefinedLayout);
+				LayoutManager.SetFullLayout(BackBufferImages[AcquiredImageIndex]->Image, UndefinedLayout);
 			}
 			else
 			{
@@ -1081,11 +1126,6 @@ bool FVulkanViewport::Present(FVulkanCommandListContext* Context, FVulkanCmdBuff
 	return bResult;
 }
 
-VkSurfaceTransformFlagBitsKHR FVulkanViewport::GetSwapchainQCOMRenderPassTransform() const
-{
-	return SwapChain->QCOMRenderPassTransform;
-}
-
 VkFormat FVulkanViewport::GetSwapchainImageFormat() const
 {
 	return SwapChain->ImageFormat;
diff --git a/Engine/Source/Runtime/VulkanRHI/Private/VulkanViewport.h b/Engine/Source/Runtime/VulkanRHI/Private/VulkanViewport.h
--- a/Engine/Source/Runtime/VulkanRHI/Private/VulkanViewport.h
+++ b/Engine/Source/Runtime/VulkanRHI/Private/VulkanViewport.h
@@ -94,7 +94,7 @@ public:
 	{
 		if (BackBufferImages.Num() > 0)
 		{
-			return BackBufferImages[Index];
+			return BackBufferImages[Index]->Image;
 		}
 		else
 		{
@@ -107,12 +107,11 @@ public:
 		return SwapChain;
 	}
 
-	VkSurfaceTransformFlagBitsKHR GetSwapchainQCOMRenderPassTransform() const;
 	VkFormat GetSwapchainImageFormat() const;
 
 protected:
 	// NUM_BUFFERS don't have to match exactly as the driver can require a minimum number larger than NUM_BUFFERS. Provide some slack
-	TArray<VkImage, TInlineAllocator<NUM_BUFFERS*2>> BackBufferImages;
+	TArray<TRefCountPtr<FVulkanTexture>, TInlineAllocator<NUM_BUFFERS*2>> BackBufferImages;
 	TArray<VulkanRHI::FSemaphore*, TInlineAllocator<NUM_BUFFERS*2>> RenderingDoneSemaphores;
 	TIndirectArray<FVulkanView, TInlineAllocator<NUM_BUFFERS*2>> TextureViews;
 	TRefCountPtr<FVulkanBackBuffer> RHIBackBuffer;
diff --git a/Engine/Source/Runtime/VulkanRHI/Public/VulkanConfiguration.h b/Engine/Source/Runtime/VulkanRHI/Public/VulkanConfiguration.h
--- a/Engine/Source/Runtime/VulkanRHI/Public/VulkanConfiguration.h
+++ b/Engine/Source/Runtime/VulkanRHI/Public/VulkanConfiguration.h
@@ -188,10 +188,6 @@ struct VkAllocationCallbacks;
 	#define VULKAN_SUPPORTS_DRIVER_PROPERTIES					1
 #endif
 
-#ifndef VULKAN_SUPPORTS_QCOM_RENDERPASS_TRANSFORM
-	#define VULKAN_SUPPORTS_QCOM_RENDERPASS_TRANSFORM			0
-#endif
-
 #ifndef VULKAN_SUPPORTS_QCOM_RENDERPASS_SHADER_RESOLVE
 	#ifdef VK_QCOM_render_pass_shader_resolve
 		#define VULKAN_SUPPORTS_QCOM_RENDERPASS_SHADER_RESOLVE	1
