Vulkan: fixed validation warning when underlying buffer size is not a multiple of a buffer view element size

diff --git a/Engine/Source/Runtime/VulkanRHI/Private/VulkanBuffer.cpp b/Engine/Source/Runtime/VulkanRHI/Private/VulkanBuffer.cpp
--- a/Engine/Source/Runtime/VulkanRHI/Private/VulkanBuffer.cpp
+++ b/Engine/Source/Runtime/VulkanRHI/Private/VulkanBuffer.cpp
@@ -292,6 +292,9 @@ void* FVulkanResourceMultiBuffer::Lock(FVulkanCommandListContext& Context, EReso
 	void* Data = nullptr;
 	uint32 DataOffset = 0;
 
+	// High level code might create SRVs on an allocated buffer, make sure size of the buffer is at least multiple of float4 (16bytes)
+	LockSize = Align(LockSize, 16u);
+
 	const bool bVolatile = EnumHasAnyFlags(GetUsage(), BUF_Volatile);
 	check(LockStatus == ELockStatus::Unlocked);
 
diff --git a/Engine/Source/Runtime/VulkanRHI/Private/VulkanUAV.cpp b/Engine/Source/Runtime/VulkanRHI/Private/VulkanUAV.cpp
--- a/Engine/Source/Runtime/VulkanRHI/Private/VulkanUAV.cpp
+++ b/Engine/Source/Runtime/VulkanRHI/Private/VulkanUAV.cpp
@@ -91,9 +91,13 @@ FVulkanView* FVulkanView::InitAsTypedBufferView(FVulkanResourceMultiBuffer* Buff
 		InSize = FMath::Min<uint64>(InSize, Buffer->GetCurrentSize());
 	}
 
+	const uint32 TypeSize =  VulkanRHI::GetNumBitsPerPixel(Format) / 8u;
+	// view size has to be a multiple of element size
+	check(IsAligned(InSize, TypeSize));
+
 	//#todo-rco: Revisit this if buffer views become VK_BUFFER_USAGE_STORAGE_BUFFER_BIT instead of VK_BUFFER_USAGE_UNIFORM_TEXEL_BUFFER_BIT
 	const VkPhysicalDeviceLimits& Limits = Device.GetLimits();
-	const uint64 MaxSize = (uint64)Limits.maxTexelBufferElements * VulkanRHI::GetNumBitsPerPixel(Format) / 8;
+	const uint64 MaxSize = (uint64)Limits.maxTexelBufferElements * TypeSize;
 	ViewInfo.range = FMath::Min<uint64>(InSize, MaxSize);
 	// TODO: add a check() for exceeding MaxSize, to catch code which blindly makes views without checking the platform limits.
 
