Fix crash when usingVulkan and  MSAA with mobile forward renderer

diff --git a/Engine/Source/Runtime/VulkanRHI/Private/VulkanRenderpass.h b/Engine/Source/Runtime/VulkanRHI/Private/VulkanRenderpass.h
--- a/Engine/Source/Runtime/VulkanRHI/Private/VulkanRenderpass.h
+++ b/Engine/Source/Runtime/VulkanRHI/Private/VulkanRenderpass.h
@@ -127,7 +127,7 @@ struct FVulkanSubpassDescription<VkSubpassDescription>
 		pColorAttachments = ColorAttachmentReferences.GetData();
 	}
 
-	void SetResolveAttachments(const TArray<FVulkanAttachmentReference<VkAttachmentReference>>& ResolveAttachmentReferences)
+	void SetResolveAttachments(const TArrayView<FVulkanAttachmentReference<VkAttachmentReference>>& ResolveAttachmentReferences)
 	{
 		if (ResolveAttachmentReferences.Num() > 0)
 		{
@@ -175,7 +175,7 @@ struct FVulkanSubpassDescription<VkSubpassDescription2>
 		pColorAttachments = ColorAttachmentReferences.GetData();
 	}
 
-	void SetResolveAttachments(const TArray<FVulkanAttachmentReference<VkAttachmentReference2>>& ResolveAttachmentReferences)
+	void SetResolveAttachments(const TArrayView<FVulkanAttachmentReference<VkAttachmentReference2>>& ResolveAttachmentReferences)
 	{
 		if (ResolveAttachmentReferences.Num() > 0)
 		{
@@ -652,11 +652,25 @@ public:
 			SubpassDep.dependencyFlags = VK_DEPENDENCY_BY_REGION_BIT;
 		}
 
-		// Only set resolve attachment on the last subpass
 		// CustomResolveSubpass does a custom resolve into second color target and does not need resolve attachments
 		if (!bCustomResolveSubpass)
 		{
-			SubpassDescriptions[NumSubpasses - 1].SetResolveAttachments(ResolveAttachmentReferences);
+			if (bDepthReadSubpass && ResolveAttachmentReferences.Num() > 1)
+			{
+				// Handle SceneDepthAux resolve:
+				// The depth read subpass has only a single color attachment (using more would not be compatible dual source blending), so make SceneDepthAux a resolve attachment of the first subpass instead
+				ResolveAttachmentReferences.Add(TAttachmentReferenceClass{});
+				ResolveAttachmentReferences.Last().attachment = VK_ATTACHMENT_UNUSED;
+				Swap(ResolveAttachmentReferences.Last(), ResolveAttachmentReferences[0]);
+
+				SubpassDescriptions[0].SetResolveAttachments(TArrayView<TAttachmentReferenceClass>(ResolveAttachmentReferences.GetData(), ResolveAttachmentReferences.Num() - 1));
+				SubpassDescriptions[NumSubpasses - 1].SetResolveAttachments(TArrayView<TAttachmentReferenceClass>(&ResolveAttachmentReferences.Last(), 1));
+			}
+			else
+			{
+				// Only set resolve attachment on the last subpass
+				SubpassDescriptions[NumSubpasses - 1].SetResolveAttachments(ResolveAttachmentReferences);
+			}
 		}
 
 		for (uint32 Attachment = 0; Attachment < RTLayout.GetNumAttachmentDescriptions(); ++Attachment)
