Fix bug in vulkan pipeline cache that could cause the LRU's head node to be added multiple times.

diff --git a/Engine/Source/Runtime/VulkanRHI/Private/VulkanChunkedPipelineCache.cpp b/Engine/Source/Runtime/VulkanRHI/Private/VulkanChunkedPipelineCache.cpp
--- a/Engine/Source/Runtime/VulkanRHI/Private/VulkanChunkedPipelineCache.cpp
+++ b/Engine/Source/Runtime/VulkanRHI/Private/VulkanChunkedPipelineCache.cpp
@@ -937,12 +937,15 @@ public:
 				FScopedTimeToLog Timer3(FString::Printf(TEXT("FVulkanChunkedPipelineCacheManager: CreatePSO lru %d "), FPlatformTLS::GetCurrentThreadId()));
 				FScopeLock Lock(&LRUCS);
 				FVulkanPipelineCacheChunkLRUNode* LRUNode = Chunk->GetLRUNode();
-				if (LRUNode->GetNextNode() || LRUNode->GetPrevNode())
+				if(CacheChunkLRU.GetHead() != LRUNode)
 				{
-					// if evicted it may not be in the lru yet..
-					CacheChunkLRU.RemoveNode(LRUNode, false);
+					if (LRUNode->GetNextNode() || LRUNode->GetPrevNode() )
+					{
+						// if evicted it may not be in the lru yet..
+						CacheChunkLRU.RemoveNode(LRUNode, false);
+					}
+					CacheChunkLRU.AddHead(LRUNode);
 				}
-				CacheChunkLRU.AddHead(LRUNode);
 			}
 			return Result;
 		}
