VulkanRHI: Add command line to force debug utils to be enabled
#rnx

diff --git a/Engine/Source/Runtime/VulkanRHI/Private/VulkanLayers.cpp b/Engine/Source/Runtime/VulkanRHI/Private/VulkanLayers.cpp
--- a/Engine/Source/Runtime/VulkanRHI/Private/VulkanLayers.cpp
+++ b/Engine/Source/Runtime/VulkanRHI/Private/VulkanLayers.cpp
@@ -399,7 +399,8 @@ void FVulkanIntanceSetupHelper::AddDebugLayers(const TArray<FLayerWithExtensions
 		}
 	}
 
-	if (bUseVulkanValidation && (ActiveDebugLayerExtension == FVulkanDynamicRHI::EActiveDebugLayerExtension::None))
+	const bool bForceDebugUtils = FParse::Param(FCommandLine::Get(), TEXT("vulkandebugutils"));
+	if ((bUseVulkanValidation || bForceDebugUtils) && (ActiveDebugLayerExtension == FVulkanDynamicRHI::EActiveDebugLayerExtension::None))
 	{
 		auto FindLayerContainingExtension = [](const ANSICHAR* ExtensionName, const TArray<FLayerWithExtensions>& LayerProperties)
 		{
