Fix RHI validation error in Vulkan pre-transform blit

Error was "Graphics PSO is using depth but it's not enabled on the RenderPass."

diff --git a/Engine/Source/Runtime/VulkanRHI/Private/VulkanViewport.cpp b/Engine/Source/Runtime/VulkanRHI/Private/VulkanViewport.cpp
--- a/Engine/Source/Runtime/VulkanRHI/Private/VulkanViewport.cpp
+++ b/Engine/Source/Runtime/VulkanRHI/Private/VulkanViewport.cpp
@@ -865,7 +865,7 @@ inline static void CopyImageToBackBuffer(FVulkanCommandListContext* Context, FVu
 		// No alpha blending, no depth tests or writes, no stencil tests or writes, no backface culling.
 		GraphicsPSOInit.BlendState = TStaticBlendState<>::GetRHI();
 		GraphicsPSOInit.RasterizerState = TStaticRasterizerState<>::GetRHI();
-		GraphicsPSOInit.DepthStencilState = TStaticDepthStencilState<false>::GetRHI();
+		GraphicsPSOInit.DepthStencilState = TStaticDepthStencilState<false, CF_Always>::GetRHI();
 
 		TRHICommandList_RecursiveHazardous<FVulkanCommandListContext> RHICmdList(Context);
 
