Updated deprecated headers to use UE_DEPRECATED_HEADER.

diff --git a/Engine/Source/Runtime/D3D12RHI/Public/D3D12RHIBridge.h b/Engine/Source/Runtime/D3D12RHI/Public/D3D12RHIBridge.h
--- a/Engine/Source/Runtime/D3D12RHI/Public/D3D12RHIBridge.h
+++ b/Engine/Source/Runtime/D3D12RHI/Public/D3D12RHIBridge.h
@@ -2,4 +2,4 @@
 
 #pragma once
 
-#pragma message("D3D12RHIBridge.h is deprecated. Please switch to using ID3D12DynamicRHI.h")
+UE_DEPRECATED_HEADER(5.3, "D3D12RHIBridge.h is deprecated. Please switch to using ID3D12DynamicRHI.h")
diff --git a/Engine/Source/Runtime/VulkanRHI/Public/VulkanRHIBridge.h b/Engine/Source/Runtime/VulkanRHI/Public/VulkanRHIBridge.h
--- a/Engine/Source/Runtime/VulkanRHI/Public/VulkanRHIBridge.h
+++ b/Engine/Source/Runtime/VulkanRHI/Public/VulkanRHIBridge.h
@@ -2,4 +2,4 @@
 
 #pragma once
 
-#pragma message("VulkanRHIBridge.h is deprecated. Please switch to using IVulkanDynamicRHI.h")
+UE_DEPRECATED_HEADER(5.3, "VulkanRHIBridge.h is deprecated. Please switch to using IVulkanDynamicRHI.h")
