VulkanRT: Respect r.RayTracing.RequireSM6 when checking raytracing support.

diff --git a/Engine/Source/Runtime/VulkanRHI/Private/VulkanRHI.cpp b/Engine/Source/Runtime/VulkanRHI/Private/VulkanRHI.cpp
--- a/Engine/Source/Runtime/VulkanRHI/Private/VulkanRHI.cpp
+++ b/Engine/Source/Runtime/VulkanRHI/Private/VulkanRHI.cpp
@@ -667,7 +667,11 @@ void FVulkanDynamicRHI::CreateInstance()
 	// Run a profile check to see if this device can support our raytacing requirements since it might change the required API version of the instance
 	if (FVulkanPlatform::SupportsProfileChecks() && GVulkanRayTracingCVar.GetValueOnAnyThread())
 	{
-		if (CheckVulkanProfile(GMaxRHIFeatureLevel, true))
+		static IConsoleVariable* RequireSM6CVar = IConsoleManager::Get().FindConsoleVariable(TEXT("r.RayTracing.RequireSM6"));
+		const bool bRequireSM6 = RequireSM6CVar && RequireSM6CVar->GetBool();
+		const bool bRayTracingAllowedOnCurrentShaderPlatform = (bRequireSM6 == false) || (GMaxRHIShaderPlatform == SP_VULKAN_SM6 || IsVulkanMobileSM5Platform(GMaxRHIShaderPlatform));
+
+		if (CheckVulkanProfile(GMaxRHIFeatureLevel, true) && bRayTracingAllowedOnCurrentShaderPlatform)
 		{
 			// Raytracing is supported, update the required API version
 			ApiVersion = GetVulkanApiVersionForFeatureLevel(GMaxRHIFeatureLevel, true);
@@ -676,7 +680,15 @@ void FVulkanDynamicRHI::CreateInstance()
 		{
 			// Raytracing is not supported, disable it completely instead of only loading parts of it
 			GVulkanRayTracingCVar->Set(0, ECVF_SetByCode);
-			UE_LOG(LogVulkanRHI, Display, TEXT("Vulkan RayTracing disabled because of failed profile check."));
+
+			if (!bRayTracingAllowedOnCurrentShaderPlatform)
+			{
+				UE_LOG(LogVulkanRHI, Display, TEXT("Vulkan RayTracing disabled because SM6 shader platform is required (r.RayTracing.RequireSM6=1)."));
+			}
+			else
+			{
+				UE_LOG(LogVulkanRHI, Display, TEXT("Vulkan RayTracing disabled because of failed profile check."));
+			}
 		}
 	}
 #endif // VULKAN_RHI_RAYTRACING
