Temporary workaround to fix regression on Adreno when using ray tracing.

This reintroduces a validation error that was fixed earlier.

diff --git a/Engine/Source/Runtime/VulkanRHI/Private/VulkanRayTracing.cpp b/Engine/Source/Runtime/VulkanRHI/Private/VulkanRayTracing.cpp
--- a/Engine/Source/Runtime/VulkanRHI/Private/VulkanRayTracing.cpp
+++ b/Engine/Source/Runtime/VulkanRHI/Private/VulkanRayTracing.cpp
@@ -726,7 +726,9 @@ void FVulkanRayTracingScene::BindBuffer(FRHIBuffer* InBuffer, uint32 InBufferOff
 		VkAccelerationStructureKHR NativeAccelerationStructureHandle = Layer.View->InitAsAccelerationStructureView(
 			AccelerationStructureBuffer
 			, LayerOffset
-			, Layer.SizeInfo.ResultSize
+			//, Layer.SizeInfo.ResultSize
+			// TODO: Using whole remaining size instead of Layer.SizeInfo.ResultSize reintruduces a validation error but use of Layer.SizeInfo.ResultSize broke RT on Adreno.
+			, InBuffer->GetSize() - LayerOffset 
 		)->GetAccelerationStructureView().Handle;
 
 		FString DebugNameString = Initializer.DebugName.ToString();
