VulkanRHI: Set debug name for Vulkan buffers.

#rnx

diff --git a/Engine/Source/Runtime/VulkanRHI/Private/VulkanBuffer.cpp b/Engine/Source/Runtime/VulkanRHI/Private/VulkanBuffer.cpp
--- a/Engine/Source/Runtime/VulkanRHI/Private/VulkanBuffer.cpp
+++ b/Engine/Source/Runtime/VulkanRHI/Private/VulkanBuffer.cpp
@@ -110,6 +110,13 @@ FVulkanResourceMultiBuffer::FVulkanResourceMultiBuffer(FVulkanDevice* InDevice,
 	const bool bZeroSize = (InBufferDesc.Size == 0);
 	BufferUsageFlags = UEToVKBufferUsageFlags(InDevice, InBufferDesc.Usage, bZeroSize);
 	
+	#if !(UE_BUILD_SHIPPING || UE_BUILD_TEST)
+	if (CreateInfo.DebugName)
+	{
+		SetName(CreateInfo.DebugName);
+	}
+	#endif
+
 	if (!bZeroSize)
 	{
 		check(InDevice);
