Comment out assert which validates that buffers are created with a size that's a multiple of the stride.

This check was introduced in CL 32498769 as an attempt to catch bad usage of create/lock/unlock in the higher level code (that CL also fixes the one case which was found at the time). However, it turns out there multiple places in the high level rendering code which re-purpose buffers for a new format while there are still views with the old format lying around, and then lock them with a size computed based on the new stride, triggering the assert when the old views are re-created. This problem is benign as long as the old views aren't used again. It's too late in the current release cycle to fix the high level code, so we'll disable the check until the correct fix can be applied in Main.

#rnx

diff --git a/Engine/Source/Runtime/VulkanRHI/Private/VulkanUAV.cpp b/Engine/Source/Runtime/VulkanRHI/Private/VulkanUAV.cpp
--- a/Engine/Source/Runtime/VulkanRHI/Private/VulkanUAV.cpp
+++ b/Engine/Source/Runtime/VulkanRHI/Private/VulkanUAV.cpp
@@ -92,8 +92,11 @@ FVulkanView* FVulkanView::InitAsTypedBufferView(FVulkanResourceMultiBuffer* Buff
 	}
 
 	const uint32 TypeSize =  VulkanRHI::GetNumBitsPerPixel(Format) / 8u;
-	// view size has to be a multiple of element size
-	check(IsAligned(InSize, TypeSize));
+	// View size has to be a multiple of element size
+	// Commented out because there are multiple places in the high level rendering code which re-purpose buffers for a new format while there are still
+	// views with the old format lying around, and then lock them with a size computed based on the new stride, triggering this assert when the old views
+	// are re-created. These places need to be fixed before re-enabling this check (UE-211785).
+	//check(IsAligned(InSize, TypeSize));
 
 	//#todo-rco: Revisit this if buffer views become VK_BUFFER_USAGE_STORAGE_BUFFER_BIT instead of VK_BUFFER_USAGE_UNIFORM_TEXEL_BUFFER_BIT
 	const VkPhysicalDeviceLimits& Limits = Device.GetLimits();
