HWRT: Add TLAS build mode
* Controlled by r.RayTracing.Scene.BuildMode
* Defaults to FastTrace like before.

#tests CitySample

diff --git a/Engine/Source/Runtime/Apple/MetalRHI/Private/MetalRayTracing.cpp b/Engine/Source/Runtime/Apple/MetalRHI/Private/MetalRayTracing.cpp
--- a/Engine/Source/Runtime/Apple/MetalRHI/Private/MetalRayTracing.cpp
+++ b/Engine/Source/Runtime/Apple/MetalRHI/Private/MetalRayTracing.cpp
@@ -435,8 +435,6 @@ void FMetalRayTracingGeometry::RebuildDescriptors()
 FMetalRayTracingScene::FMetalRayTracingScene(FRayTracingSceneInitializer2 InInitializer)
 	: Initializer(MoveTemp(InInitializer))
 {
-	const ERayTracingAccelerationStructureFlags BuildFlags = ERayTracingAccelerationStructureFlags::FastTrace; // #yuriy_todo: pass this in
-
 	SizeInfo = {};
 
 	const uint32 NumLayers = Initializer.NumNativeInstancesPerLayer.Num();
diff --git a/Engine/Source/Runtime/D3D12RHI/Private/D3D12RayTracing.cpp b/Engine/Source/Runtime/D3D12RHI/Private/D3D12RayTracing.cpp
--- a/Engine/Source/Runtime/D3D12RHI/Private/D3D12RayTracing.cpp
+++ b/Engine/Source/Runtime/D3D12RHI/Private/D3D12RayTracing.cpp
@@ -3355,8 +3355,6 @@ FD3D12RayTracingScene::FD3D12RayTracingScene(FD3D12Adapter* Adapter, FRayTracing
 	checkf(Initializer.NumMissShaderSlots >= 1, TEXT("Need at least 1 miss shader slot."));
 	checkf(Initializer.Lifetime == RTSL_SingleFrame, TEXT("Only single-frame ray tracing scenes are currently implemented."));
 
-	const ERayTracingAccelerationStructureFlags BuildFlags = ERayTracingAccelerationStructureFlags::FastTrace; // #yuriy_todo: pass this in
-
 	SizeInfo = {};
 
 	const uint32 NumLayers = Initializer.NumNativeInstancesPerLayer.Num();
@@ -3372,10 +3370,10 @@ FD3D12RayTracingScene::FD3D12RayTracingScene(FD3D12Adapter* Adapter, FRayTracing
 		Layer.BuildInputs.Type = D3D12_RAYTRACING_ACCELERATION_STRUCTURE_TYPE_TOP_LEVEL;
 		Layer.BuildInputs.DescsLayout = D3D12_ELEMENTS_LAYOUT_ARRAY;
 		Layer.BuildInputs.NumDescs = Initializer.NumNativeInstancesPerLayer[LayerIndex];
-		Layer.BuildInputs.Flags = TranslateRayTracingAccelerationStructureFlags(BuildFlags);
+		Layer.BuildInputs.Flags = TranslateRayTracingAccelerationStructureFlags(Initializer.BuildFlags);
 
 		// Get maximum buffer sizes for all GPUs in the system
-		Layer.SizeInfo = RHICalcRayTracingSceneSize(Layer.BuildInputs.NumDescs, BuildFlags);
+		Layer.SizeInfo = RHICalcRayTracingSceneSize(Layer.BuildInputs.NumDescs, Initializer.BuildFlags);
 
 		Layer.BufferOffset = Align(SizeInfo.ResultSize, GRHIRayTracingAccelerationStructureAlignment);
 		Layer.ScratchBufferOffset = Align(SizeInfo.BuildScratchSize, GRHIRayTracingScratchBufferAlignment);
diff --git a/Engine/Source/Runtime/RHI/Public/RHIResources.h b/Engine/Source/Runtime/RHI/Public/RHIResources.h
--- a/Engine/Source/Runtime/RHI/Public/RHIResources.h
+++ b/Engine/Source/Runtime/RHI/Public/RHIResources.h
@@ -3197,6 +3197,9 @@ struct FRayTracingSceneInitializer2
 	// Defines whether data in this scene should persist between frames.
 	// Currently only single-frame lifetime is supported.
 	ERayTracingSceneLifetime Lifetime = RTSL_SingleFrame;
+	
+	// Controls the flags of the ray tracing scene build.
+	ERayTracingAccelerationStructureFlags BuildFlags = ERayTracingAccelerationStructureFlags::FastTrace;
 
 	FName DebugName;
 
diff --git a/Engine/Source/Runtime/Renderer/Private/RayTracing/RayTracingInstanceBufferUtil.cpp b/Engine/Source/Runtime/Renderer/Private/RayTracing/RayTracingInstanceBufferUtil.cpp
--- a/Engine/Source/Runtime/Renderer/Private/RayTracing/RayTracingInstanceBufferUtil.cpp
+++ b/Engine/Source/Runtime/Renderer/Private/RayTracing/RayTracingInstanceBufferUtil.cpp
@@ -30,7 +30,8 @@ FRayTracingSceneWithGeometryInstances CreateRayTracingSceneWithGeometryInstances
 	uint8 NumLayers,
 	uint32 NumShaderSlotsPerGeometrySegment,
 	uint32 NumMissShaderSlots,
-	uint32 NumCallableShaderSlots)
+	uint32 NumCallableShaderSlots,
+	ERayTracingAccelerationStructureFlags BuildFlags)
 {
 	const uint32 NumSceneInstances = Instances.Num();
 
@@ -50,6 +51,7 @@ FRayTracingSceneWithGeometryInstances CreateRayTracingSceneWithGeometryInstances
 	Initializer.SegmentPrefixSum.SetNumUninitialized(NumSceneInstances);
 	Initializer.NumNativeInstancesPerLayer.SetNumZeroed(NumLayers);
 	Initializer.NumTotalSegments = 0;
+	Initializer.BuildFlags = BuildFlags;
 
 	Experimental::TSherwoodMap<FRHIRayTracingGeometry*, uint32> UniqueGeometries;
 
@@ -115,14 +117,6 @@ FRayTracingSceneWithGeometryInstances CreateRayTracingSceneWithGeometryInstances
 	return MoveTemp(Output);
 }
 
-FRayTracingSceneWithGeometryInstances CreateRayTracingSceneWithGeometryInstances(
-	TArrayView<FRayTracingGeometryInstance> Instances,
-	uint32 NumShaderSlotsPerGeometrySegment,
-	uint32 NumMissShaderSlots)
-{
-	return CreateRayTracingSceneWithGeometryInstances(Instances, 1, NumShaderSlotsPerGeometrySegment, NumMissShaderSlots);
-}
-
 void FillRayTracingInstanceUploadBuffer(
 	FRayTracingSceneRHIRef RayTracingSceneRHI,
 	FVector PreViewTranslation,
diff --git a/Engine/Source/Runtime/Renderer/Private/RayTracing/RayTracingScene.cpp b/Engine/Source/Runtime/Renderer/Private/RayTracing/RayTracingScene.cpp
--- a/Engine/Source/Runtime/Renderer/Private/RayTracing/RayTracingScene.cpp
+++ b/Engine/Source/Runtime/Renderer/Private/RayTracing/RayTracingScene.cpp
@@ -17,6 +17,15 @@
 
 DECLARE_GPU_STAT(RayTracingScene);
 
+static TAutoConsoleVariable<int32> CVarRayTracingSceneBuildMode(
+	TEXT("r.RayTracing.Scene.BuildMode"),
+	1,
+	TEXT("Controls the mode in which ray tracing scene is built:\n")
+	TEXT(" 0: Fast build\n")
+	TEXT(" 1: Fast trace (default)\n"),
+	ECVF_RenderThreadSafe | ECVF_Scalability
+);
+
 BEGIN_SHADER_PARAMETER_STRUCT(FBuildInstanceBufferPassParams, )
 	SHADER_PARAMETER_RDG_BUFFER_UAV(RWStructuredBuffer, InstanceBuffer)
 	SHADER_PARAMETER_RDG_BUFFER_UAV(RWStructuredBuffer, DebugInstanceGPUSceneIndexBuffer)
@@ -36,12 +45,17 @@ FRayTracingScene::~FRayTracingScene()
 
 FRayTracingSceneWithGeometryInstances FRayTracingScene::BuildInitializationData() const
 {
+	ERayTracingAccelerationStructureFlags BuildFlags = CVarRayTracingSceneBuildMode.GetValueOnRenderThread()
+		? ERayTracingAccelerationStructureFlags::FastTrace
+		: ERayTracingAccelerationStructureFlags::FastBuild;
+
 	return CreateRayTracingSceneWithGeometryInstances(
 		Instances,
 		uint8(ERayTracingSceneLayer::NUM),
 		RAY_TRACING_NUM_SHADER_SLOTS,
 		NumMissShaderSlots,
-		NumCallableShaderSlots);
+		NumCallableShaderSlots,
+		BuildFlags);
 }
 
 void FRayTracingScene::Create(FRDGBuilder& GraphBuilder, const FViewInfo& View, const FGPUScene* GPUScene)
diff --git a/Engine/Source/Runtime/Renderer/Public/RayTracingInstanceBufferUtil.h b/Engine/Source/Runtime/Renderer/Public/RayTracingInstanceBufferUtil.h
--- a/Engine/Source/Runtime/Renderer/Public/RayTracingInstanceBufferUtil.h
+++ b/Engine/Source/Runtime/Renderer/Public/RayTracingInstanceBufferUtil.h
@@ -74,13 +74,8 @@ RENDERER_API FRayTracingSceneWithGeometryInstances CreateRayTracingSceneWithGeom
 	uint8 NumLayers,
 	uint32 NumShaderSlotsPerGeometrySegment,
 	uint32 NumMissShaderSlots,
-	uint32 NumCallableShaderSlots = 0);
-
-UE_DEPRECATED(5.1, "Specify NumLayers instead.")
-RENDERER_API FRayTracingSceneWithGeometryInstances CreateRayTracingSceneWithGeometryInstances(
-	TArrayView<FRayTracingGeometryInstance> Instances,
-	uint32 NumShaderSlotsPerGeometrySegment,
-	uint32 NumMissShaderSlots);
+	uint32 NumCallableShaderSlots = 0,
+	ERayTracingAccelerationStructureFlags BuildFlags = ERayTracingAccelerationStructureFlags::FastTrace);
 
 // Helper function to fill upload buffers required by BuildRayTracingInstanceBuffer with instance descriptors
 // Transforms of CPU instances are copied to OutTransformData
diff --git a/Engine/Source/Runtime/VulkanRHI/Private/VulkanRayTracing.cpp b/Engine/Source/Runtime/VulkanRHI/Private/VulkanRayTracing.cpp
--- a/Engine/Source/Runtime/VulkanRHI/Private/VulkanRayTracing.cpp
+++ b/Engine/Source/Runtime/VulkanRHI/Private/VulkanRayTracing.cpp
@@ -28,14 +28,6 @@ static FAutoConsoleVariableRef CVarVulkanRayTracingMaxBatchedCompaction(
 	ECVF_ReadOnly
 );
 
-static bool GVulkanRayTracingTLASPreferFastTrace = true;
-static FAutoConsoleVariableRef CVarVulkanRayTracingTLASPreferFastTraceTLAS(
-	TEXT("r.Vulkan.RayTracing.TLASPreferFastTraceTLAS"),
-	GVulkanRayTracingTLASPreferFastTrace,
-	TEXT("Prefer fast trace for TLAS build (default = true)\n"),
-	ECVF_ReadOnly
-);
-
 static int32 GVulkanRayTracingAllowDeferredOperation = -1;
 static FAutoConsoleVariableRef CVarVulkanRayTracingAllowDeferredOperation(
 	TEXT("r.Vulkan.RayTracing.AllowDeferredOperation"),
@@ -150,6 +142,30 @@ static bool ShouldCompactAfterBuild(ERayTracingAccelerationStructureFlags BuildF
 		&& !EnumHasAnyFlags(BuildFlags, ERayTracingAccelerationStructureFlags::AllowUpdate);
 }
 
+static VkBuildAccelerationStructureFlagBitsKHR TranslateRayTracingAccelerationStructureFlags(ERayTracingAccelerationStructureFlags Flags)
+{
+	uint32 Result = {};
+
+	auto HandleFlag = [&Flags, &Result](ERayTracingAccelerationStructureFlags Engine, VkBuildAccelerationStructureFlagBitsKHR Native)
+		{
+			if (EnumHasAllFlags(Flags, Engine))
+			{
+				Result |= (uint32)Native;
+				EnumRemoveFlags(Flags, Engine);
+			}
+		};
+
+	HandleFlag(ERayTracingAccelerationStructureFlags::AllowUpdate, VK_BUILD_ACCELERATION_STRUCTURE_ALLOW_UPDATE_BIT_KHR);
+	HandleFlag(ERayTracingAccelerationStructureFlags::AllowCompaction, VK_BUILD_ACCELERATION_STRUCTURE_ALLOW_COMPACTION_BIT_KHR);
+	HandleFlag(ERayTracingAccelerationStructureFlags::FastTrace, VK_BUILD_ACCELERATION_STRUCTURE_PREFER_FAST_TRACE_BIT_KHR);
+	HandleFlag(ERayTracingAccelerationStructureFlags::FastBuild, VK_BUILD_ACCELERATION_STRUCTURE_PREFER_FAST_BUILD_BIT_KHR);
+	HandleFlag(ERayTracingAccelerationStructureFlags::MinimizeMemory, VK_BUILD_ACCELERATION_STRUCTURE_LOW_MEMORY_BIT_KHR);
+
+	checkf(!EnumHasAnyFlags(Flags, Flags), TEXT("Some ERayTracingAccelerationStructureFlags entries were not handled"));
+
+	return VkBuildAccelerationStructureFlagBitsKHR(Result);
+}
+
 static ERayTracingAccelerationStructureFlags GetRayTracingAccelerationStructureBuildFlags(const FRayTracingGeometryInitializer& Initializer)
 {
 	ERayTracingAccelerationStructureFlags BuildFlags = ERayTracingAccelerationStructureFlags::None;
@@ -574,6 +590,7 @@ static void GetTLASBuildData(
 	const VkDevice Device,
 	const uint32 NumInstances,
 	const VkDeviceAddress InstanceBufferAddress,
+	ERayTracingAccelerationStructureFlags BuildFlags,
 	FVkRtTLASBuildData& BuildData)
 {
 	VkDeviceOrHostAddressConstKHR InstanceBufferDeviceAddress = {};
@@ -586,9 +603,7 @@ static void GetTLASBuildData(
 
 	BuildData.GeometryInfo.type = VK_ACCELERATION_STRUCTURE_TYPE_TOP_LEVEL_KHR;
 	BuildData.GeometryInfo.mode = VK_BUILD_ACCELERATION_STRUCTURE_MODE_BUILD_KHR;
-	BuildData.GeometryInfo.flags = GVulkanRayTracingTLASPreferFastTrace ? 
-										VK_BUILD_ACCELERATION_STRUCTURE_PREFER_FAST_TRACE_BIT_KHR :
-										VK_BUILD_ACCELERATION_STRUCTURE_PREFER_FAST_BUILD_BIT_KHR;
+	BuildData.GeometryInfo.flags = TranslateRayTracingAccelerationStructureFlags(BuildFlags);
 	BuildData.GeometryInfo.geometryCount = 1;
 	BuildData.GeometryInfo.pGeometries = &BuildData.Geometry;
 
@@ -655,8 +670,6 @@ FVulkanRayTracingScene::FVulkanRayTracingScene(FRayTracingSceneInitializer2 InIn
 {
 	INC_DWORD_STAT(STAT_VulkanRayTracingAllocatedTLAS);
 
-	const ERayTracingAccelerationStructureFlags BuildFlags = ERayTracingAccelerationStructureFlags::FastTrace; // #yuriy_todo: pass this in
-
 	SizeInfo = {};
 
 	const uint32 NumLayers = Initializer.NumNativeInstancesPerLayer.Num();
@@ -668,7 +681,7 @@ FVulkanRayTracingScene::FVulkanRayTracingScene(FRayTracingSceneInitializer2 InIn
 	{
 		FLayerData& Layer = Layers[LayerIndex];
 
-		Layer.SizeInfo = RHICalcRayTracingSceneSize(Initializer.NumNativeInstancesPerLayer[LayerIndex], BuildFlags);
+		Layer.SizeInfo = RHICalcRayTracingSceneSize(Initializer.NumNativeInstancesPerLayer[LayerIndex], Initializer.BuildFlags);
 		Layer.BufferOffset = Align(SizeInfo.ResultSize, GRHIRayTracingAccelerationStructureAlignment);
 		Layer.ScratchBufferOffset = Align(SizeInfo.BuildScratchSize, GRHIRayTracingScratchBufferAlignment);
 
@@ -782,7 +795,7 @@ void FVulkanRayTracingScene::BuildAccelerationStructure(
 		const FLayerData& Layer = Layers[LayerIndex];
 
 		FVkRtTLASBuildData& BuildData = BuildDatas[LayerIndex];
-		GetTLASBuildData(Device->GetInstanceHandle(), Initializer.NumNativeInstancesPerLayer[LayerIndex], InstanceBufferAddress, BuildData);
+		GetTLASBuildData(Device->GetInstanceHandle(), Initializer.NumNativeInstancesPerLayer[LayerIndex], InstanceBufferAddress, Initializer.BuildFlags, BuildData);
 
 		checkf(Layer.View.IsValid(), TEXT("A buffer must be bound to the ray tracing scene before it can be built."));
 		BuildData.GeometryInfo.dstAccelerationStructure = Layer.View->GetAccelerationStructureView().Handle;
@@ -1077,7 +1090,7 @@ FRayTracingAccelerationStructureSize FVulkanDynamicRHI::RHICalcRayTracingSceneSi
 {
 	FVkRtTLASBuildData BuildData;
 	const VkDeviceAddress InstanceBufferAddress = 0; // No device address available when only querying TLAS size
-	GetTLASBuildData(Device->GetInstanceHandle(), MaxInstances, InstanceBufferAddress, BuildData);
+	GetTLASBuildData(Device->GetInstanceHandle(), MaxInstances, InstanceBufferAddress, Flags, BuildData);
 
 	FRayTracingAccelerationStructureSize Result;
 	Result.ResultSize = BuildData.SizesInfo.accelerationStructureSize;
