VulkanRHI: Add two more Vulkan 1.1 functions to loader header (vkGetPhysicalDeviceFormatProperties2 and vkGetPhysicalDeviceQueueFamilyProperties2)

#rnx

diff --git a/Engine/Source/Runtime/VulkanRHI/Private/VulkanLoader.h b/Engine/Source/Runtime/VulkanRHI/Private/VulkanLoader.h
--- a/Engine/Source/Runtime/VulkanRHI/Private/VulkanLoader.h
+++ b/Engine/Source/Runtime/VulkanRHI/Private/VulkanLoader.h
@@ -11,9 +11,11 @@
 	EnumMacro(PFN_vkEnumeratePhysicalDevices, vkEnumeratePhysicalDevices) \
 	EnumMacro(PFN_vkGetPhysicalDeviceFeatures, vkGetPhysicalDeviceFeatures) \
 	EnumMacro(PFN_vkGetPhysicalDeviceFormatProperties, vkGetPhysicalDeviceFormatProperties) \
+	EnumMacro(PFN_vkGetPhysicalDeviceFormatProperties2, vkGetPhysicalDeviceFormatProperties2) \
 	EnumMacro(PFN_vkGetPhysicalDeviceImageFormatProperties, vkGetPhysicalDeviceImageFormatProperties) \
 	EnumMacro(PFN_vkGetPhysicalDeviceProperties, vkGetPhysicalDeviceProperties) \
 	EnumMacro(PFN_vkGetPhysicalDeviceQueueFamilyProperties, vkGetPhysicalDeviceQueueFamilyProperties) \
+	EnumMacro(PFN_vkGetPhysicalDeviceQueueFamilyProperties2, vkGetPhysicalDeviceQueueFamilyProperties2) \
 	EnumMacro(PFN_vkGetPhysicalDeviceMemoryProperties, vkGetPhysicalDeviceMemoryProperties) \
 	EnumMacro(PFN_vkGetPhysicalDeviceProperties2, vkGetPhysicalDeviceProperties2) \
 	EnumMacro(PFN_vkGetPhysicalDeviceFeatures2, vkGetPhysicalDeviceFeatures2) \
