VulkanRHI: Add new functions in command wrapper.

diff --git a/Engine/Source/Runtime/VulkanRHI/Private/VulkanCommandWrappers.h b/Engine/Source/Runtime/VulkanRHI/Private/VulkanCommandWrappers.h
--- a/Engine/Source/Runtime/VulkanRHI/Private/VulkanCommandWrappers.h
+++ b/Engine/Source/Runtime/VulkanRHI/Private/VulkanCommandWrappers.h
@@ -50,9 +50,11 @@ struct FWrapLayer
 	static void EnumeratePhysicalDevices(VkResult Result, VkInstance Instance, uint32* PhysicalDeviceCount, VkPhysicalDevice* PhysicalDevices) VULKAN_LAYER_BODY
 	static void GetPhysicalDeviceFeatures(VkResult Result, VkPhysicalDevice PhysicalDevice, VkPhysicalDeviceFeatures* Features) VULKAN_LAYER_BODY
 	static void GetPhysicalDeviceFormatProperties(VkResult Result, VkPhysicalDevice PhysicalDevice, VkFormat Format, VkFormatProperties* FormatProperties) VULKAN_LAYER_BODY
+	static void GetPhysicalDeviceFormatProperties2(VkResult Result, VkPhysicalDevice PhysicalDevice, VkFormat Format, VkFormatProperties2* FormatProperties) VULKAN_LAYER_BODY
 	static void GetPhysicalDeviceImageFormatProperties(VkResult Result, VkPhysicalDevice PhysicalDevice, VkFormat Format, VkImageType Type, VkImageTiling Tiling, VkImageUsageFlags Usage, VkImageCreateFlags Flags, VkImageFormatProperties* pImageFormatProperties) VULKAN_LAYER_BODY
 	static void GetPhysicalDeviceProperties(VkResult Result, VkPhysicalDevice PhysicalDevice, VkPhysicalDeviceProperties* FormatProperties) VULKAN_LAYER_BODY
 	static void GetPhysicalDeviceQueueFamilyProperties(VkResult Result, VkPhysicalDevice PhysicalDevice, uint32* QueueFamilyPropertyCount, VkQueueFamilyProperties* QueueFamilyProperties) VULKAN_LAYER_BODY
+	static void GetPhysicalDeviceQueueFamilyProperties2(VkResult Result, VkPhysicalDevice PhysicalDevice, uint32* QueueFamilyPropertyCount, VkQueueFamilyProperties2* QueueFamilyProperties) VULKAN_LAYER_BODY
 	static void GetPhysicalDeviceMemoryProperties(VkResult Result, VkPhysicalDevice PhysicalDevice, VkPhysicalDeviceMemoryProperties* Properties) VULKAN_LAYER_BODY
 	static void GetInstanceProcAddr(VkResult Result, VkInstance Instance, const char* Name, PFN_vkVoidFunction VoidFunction) VULKAN_LAYER_BODY
 	VULKAN_EXTERN_EXPORT static void GetDeviceProcAddr(VkResult Result, VkDevice Device, const char* Name, PFN_vkVoidFunction VoidFunction) VULKAN_LAYER_BODY
@@ -293,6 +295,13 @@ namespace VulkanRHI
 		FWrapLayer::GetPhysicalDeviceFormatProperties(VK_SUCCESS, PhysicalDevice, Format, FormatProperties);
 	}
 
+	static FORCEINLINE_DEBUGGABLE void  vkGetPhysicalDeviceFormatProperties2(VkPhysicalDevice PhysicalDevice, VkFormat Format, VkFormatProperties2* FormatProperties)
+	{
+		FWrapLayer::GetPhysicalDeviceFormatProperties2(VK_RESULT_MAX_ENUM, PhysicalDevice, Format, FormatProperties);
+		VULKANAPINAMESPACE::vkGetPhysicalDeviceFormatProperties2(PhysicalDevice, Format, FormatProperties);
+		FWrapLayer::GetPhysicalDeviceFormatProperties2(VK_SUCCESS, PhysicalDevice, Format, FormatProperties);
+	}
+
 	static FORCEINLINE_DEBUGGABLE VkResult  vkGetPhysicalDeviceImageFormatProperties(VkPhysicalDevice PhysicalDevice, VkFormat Format, VkImageType Type, VkImageTiling Tiling, VkImageUsageFlags Usage, VkImageCreateFlags Flags, VkImageFormatProperties* pImageFormatProperties)
 	{
 		FWrapLayer::GetPhysicalDeviceImageFormatProperties(VK_RESULT_MAX_ENUM, PhysicalDevice, Format, Type, Tiling, Usage, Flags, pImageFormatProperties);
@@ -343,6 +352,13 @@ namespace VulkanRHI
 		FWrapLayer::GetPhysicalDeviceQueueFamilyProperties(VK_SUCCESS, PhysicalDevice, QueueFamilyPropertyCount, QueueFamilyProperties);
 	}
 
+	static FORCEINLINE_DEBUGGABLE void  vkGetPhysicalDeviceQueueFamilyProperties2(VkPhysicalDevice PhysicalDevice, uint32* QueueFamilyPropertyCount, VkQueueFamilyProperties2* QueueFamilyProperties)
+	{
+		FWrapLayer::GetPhysicalDeviceQueueFamilyProperties2(VK_RESULT_MAX_ENUM, PhysicalDevice, QueueFamilyPropertyCount, QueueFamilyProperties);
+		VULKANAPINAMESPACE::vkGetPhysicalDeviceQueueFamilyProperties2(PhysicalDevice, QueueFamilyPropertyCount, QueueFamilyProperties);
+		FWrapLayer::GetPhysicalDeviceQueueFamilyProperties2(VK_SUCCESS, PhysicalDevice, QueueFamilyPropertyCount, QueueFamilyProperties);
+	}
+
 	static FORCEINLINE_DEBUGGABLE void  vkGetPhysicalDeviceMemoryProperties(VkPhysicalDevice PhysicalDevice, VkPhysicalDeviceMemoryProperties* MemoryProperties)
 	{
 		FWrapLayer::GetPhysicalDeviceMemoryProperties(VK_RESULT_MAX_ENUM, PhysicalDevice, MemoryProperties);
diff --git a/Engine/Source/Runtime/VulkanRHI/Private/VulkanDebug.cpp b/Engine/Source/Runtime/VulkanRHI/Private/VulkanDebug.cpp
--- a/Engine/Source/Runtime/VulkanRHI/Private/VulkanDebug.cpp
+++ b/Engine/Source/Runtime/VulkanRHI/Private/VulkanDebug.cpp
@@ -2244,6 +2244,16 @@ void FWrapLayer::GetPhysicalDeviceFormatProperties(VkResult Result, VkPhysicalDe
 	}
 }
 
+void FWrapLayer::GetPhysicalDeviceFormatProperties2(VkResult Result, VkPhysicalDevice PhysicalDevice, VkFormat Format, VkFormatProperties2* FormatProperties)
+{
+	if (Result == VK_RESULT_MAX_ENUM)
+	{
+#if VULKAN_ENABLE_DUMP_LAYER
+		PrintfBegin(FString::Printf(TEXT("vkGetPhysicalDeviceFormatProperties2(PhysicalDevice=0x%p, Format=%d, FormatProperties2=0x%p)[...]"), PhysicalDevice, (int32)Format, FormatProperties));
+#endif
+	}
+}
+
 #if VULKAN_ENABLE_DUMP_LAYER
 static void DumpImageFormatProperties(const VkImageFormatProperties* pProps)
 {
@@ -3605,6 +3615,16 @@ void FWrapLayer::GetPhysicalDeviceQueueFamilyProperties(VkResult Result, VkPhysi
 	}
 }
 
+void FWrapLayer::GetPhysicalDeviceQueueFamilyProperties2(VkResult Result, VkPhysicalDevice PhysicalDevice, uint32* QueueFamilyPropertyCount, VkQueueFamilyProperties2* QueueFamilyProperties)
+{
+	if (Result == VK_RESULT_MAX_ENUM)
+	{
+#if VULKAN_ENABLE_DUMP_LAYER
+		PrintfBeginResult(FString::Printf(TEXT("vkGetPhysicalDeviceQueueFamilyProperties2(PhysicalDevice=0x%p, QueueFamilyPropertyCount=0x%p, QueueFamilyProperties2=0x%p)[...]"), PhysicalDevice, QueueFamilyPropertyCount, QueueFamilyProperties));
+#endif
+	}
+}
+
 void FWrapLayer::GetImageMemoryRequirements2(VkResult Result, VkDevice Device, const VkImageMemoryRequirementsInfo2* Info, VkMemoryRequirements2* MemoryRequirements)
 {
 	if (Result == VK_RESULT_MAX_ENUM)
diff --git a/Engine/Source/Runtime/VulkanRHI/Private/VulkanUtil.cpp b/Engine/Source/Runtime/VulkanRHI/Private/VulkanUtil.cpp
--- a/Engine/Source/Runtime/VulkanRHI/Private/VulkanUtil.cpp
+++ b/Engine/Source/Runtime/VulkanRHI/Private/VulkanUtil.cpp
@@ -849,7 +849,7 @@ namespace VulkanRHI
 
 			VkDeviceFaultCountsEXT FaultCounts;
 			ZeroVulkanStruct(FaultCounts, VK_STRUCTURE_TYPE_DEVICE_FAULT_COUNTS_EXT);
-			Result = vkGetDeviceFaultInfoEXT(DeviceHandle, &FaultCounts, nullptr);
+			Result = VulkanRHI::vkGetDeviceFaultInfoEXT(DeviceHandle, &FaultCounts, nullptr);
 			if (Result == VK_SUCCESS)
 			{
 				VkDeviceFaultInfoEXT FaultInfo;
@@ -867,7 +867,7 @@ namespace VulkanRHI
 				VendorBinaryData.SetNumZeroed(FaultCounts.vendorBinarySize);
 				FaultInfo.pVendorBinaryData = VendorBinaryData.GetData();
 
-				Result = vkGetDeviceFaultInfoEXT(DeviceHandle, &FaultCounts, &FaultInfo);
+				Result = VulkanRHI::vkGetDeviceFaultInfoEXT(DeviceHandle, &FaultCounts, &FaultInfo);
 				if (Result == VK_SUCCESS)
 				{
 					// :todo-jn: match these up to resources
