Recreate the swap chain after screen rotation on Android Vulkan (Required by Vulkan preTransform).

diff --git a/Engine/Source/Runtime/OpenGLDrv/Private/Android/AndroidEGL.cpp b/Engine/Source/Runtime/OpenGLDrv/Private/Android/AndroidEGL.cpp
--- a/Engine/Source/Runtime/OpenGLDrv/Private/Android/AndroidEGL.cpp
+++ b/Engine/Source/Runtime/OpenGLDrv/Private/Android/AndroidEGL.cpp
@@ -1127,6 +1127,8 @@ void FAndroidAppEntry::ReInitWindow(void* NewNativeWindowHandle)
 	// Window creation is now handled by BlockRendering, when it resumes after a new window is created.
 	FPlatformMisc::LowLevelOutputDebugString(TEXT("AndroidEGL::ReInitWindow()"));
 
+	GSystemResolution.bForceRefresh = true;
+
 	// It isn't safe to call ShouldUseVulkan if AndroidEGL is not initialized.
 	// However, since we don't need to ReInit the window in that case anyways we
 	// can return early.
diff --git a/Engine/Source/Runtime/VulkanRHI/Private/VulkanViewport.cpp b/Engine/Source/Runtime/VulkanRHI/Private/VulkanViewport.cpp
--- a/Engine/Source/Runtime/VulkanRHI/Private/VulkanViewport.cpp
+++ b/Engine/Source/Runtime/VulkanRHI/Private/VulkanViewport.cpp
@@ -144,6 +144,8 @@ FVulkanViewport::FVulkanViewport(FVulkanDevice* InDevice, void* InWindowHandle,
 			RenderingDoneSemaphores[Index] = new VulkanRHI::FSemaphore(*InDevice);
 			RenderingDoneSemaphores[Index]->AddRef();
 		}
+
+		FCoreDelegates::OnSystemResolutionChanged.AddRaw(this, &FVulkanViewport::OnSystemResolutionChanged);
 	}
 }
 
@@ -173,6 +175,8 @@ FVulkanViewport::~FVulkanViewport()
 		SwapChain->Destroy(nullptr);
 		delete SwapChain;
 		SwapChain = nullptr;
+
+		FCoreDelegates::OnSystemResolutionChanged.RemoveAll(this);
 	}
 
 	FVulkanDynamicRHI::Get().Viewports.Remove(this);
@@ -1154,6 +1158,31 @@ EPixelFormat FVulkanViewport::GetPixelFormatForNonDefaultSwapchain()
 	}
 }
 
+void FVulkanViewport::OnSystemResolutionChanged(uint32 ResX, uint32 ResY)
+{
+	EDeviceScreenOrientation CurrentOrientation = FPlatformMisc::GetDeviceOrientation();
+
+	// The swap chain needs to be recreated after a rotation
+	// Only 180-degree rotations need to be handled here because 90-degree rotations will resize the viewport and recreate the swap chain.
+	if ((CachedOrientation == EDeviceScreenOrientation::Portrait && CurrentOrientation == EDeviceScreenOrientation::PortraitUpsideDown)
+		|| (CachedOrientation == EDeviceScreenOrientation::PortraitUpsideDown && CurrentOrientation == EDeviceScreenOrientation::Portrait)
+		|| (CachedOrientation == EDeviceScreenOrientation::LandscapeRight && CurrentOrientation == EDeviceScreenOrientation::LandscapeLeft)
+		|| (CachedOrientation == EDeviceScreenOrientation::LandscapeLeft && CurrentOrientation == EDeviceScreenOrientation::LandscapeRight))
+	{
+		check(IsInGameThread());
+
+		FlushRenderingCommands();
+
+		ENQUEUE_RENDER_COMMAND(RecreateSwapchain)(
+			[this](FRHICommandListImmediate& RHICmdList)
+			{
+				RecreateSwapchainFromRT(PixelFormat);
+			});
+		FlushRenderingCommands();
+	}
+	CachedOrientation = CurrentOrientation;
+}
+
 /*=============================================================================
  *	The following RHI functions must be called from the main thread.
  *=============================================================================*/
diff --git a/Engine/Source/Runtime/VulkanRHI/Private/VulkanViewport.h b/Engine/Source/Runtime/VulkanRHI/Private/VulkanViewport.h
--- a/Engine/Source/Runtime/VulkanRHI/Private/VulkanViewport.h
+++ b/Engine/Source/Runtime/VulkanRHI/Private/VulkanViewport.h
@@ -142,6 +142,9 @@ protected:
 	FVulkanCmdBuffer* LastFrameCommandBuffer = nullptr;
 	uint64 LastFrameFenceCounter = 0;
 
+	EDeviceScreenOrientation CachedOrientation = EDeviceScreenOrientation::Unknown;
+	void OnSystemResolutionChanged(uint32 ResX, uint32 ResY);
+
 	void CreateSwapchain(struct FVulkanSwapChainRecreateInfo* RecreateInfo);
 	void DestroySwapchain(struct FVulkanSwapChainRecreateInfo* RecreateInfo);
 	bool TryAcquireImageIndex();
