Vulkan GPU Breadcrumbs for Android, outside of RenderPass only

This change allows RHI Breadcrumbs to be enabled in all build configurations on Android

diff --git a/Engine/Source/Programs/UnrealBuildTool/Platform/Android/UEBuildAndroid.cs b/Engine/Source/Programs/UnrealBuildTool/Platform/Android/UEBuildAndroid.cs
--- a/Engine/Source/Programs/UnrealBuildTool/Platform/Android/UEBuildAndroid.cs
+++ b/Engine/Source/Programs/UnrealBuildTool/Platform/Android/UEBuildAndroid.cs
@@ -503,6 +503,7 @@ namespace UnrealBuildTool
 
 			CompileEnvironment.Definitions.Add("WITH_EDITOR=0");
 			CompileEnvironment.Definitions.Add("USE_NULL_RHI=0");
+ 			CompileEnvironment.Definitions.Add("WITH_RHI_BREADCRUMBS=1");
 
 			if (Target.AndroidPlatform.bEnableRayTracing)
 			{
diff --git a/Engine/Source/Runtime/VulkanRHI/Private/Android/VulkanAndroidPlatform.cpp b/Engine/Source/Runtime/VulkanRHI/Private/Android/VulkanAndroidPlatform.cpp
--- a/Engine/Source/Runtime/VulkanRHI/Private/Android/VulkanAndroidPlatform.cpp
+++ b/Engine/Source/Runtime/VulkanRHI/Private/Android/VulkanAndroidPlatform.cpp
@@ -1799,4 +1799,11 @@ FString FVulkanAndroidPlatform::GetVulkanProfileNameForFeatureLevel(ERHIFeatureL
 		ProfileName += TEXT("_RT");
 	}
 	return ProfileName;
+} 
+
+void FVulkanAndroidPlatform::WriteCrashMarker(const FOptionalVulkanDeviceExtensions& OptionalExtensions, FVulkanCmdBuffer* CmdBuffer, VkBuffer DestBuffer, const TArrayView<uint32>& Entries, bool bAdding)
+{
+	ensure(Entries.Num() <= GMaxCrashBufferEntries);
+
+	WriteCrashMarkerWithoutExtensions(CmdBuffer, DestBuffer, Entries, bAdding);
 }
diff --git a/Engine/Source/Runtime/VulkanRHI/Private/Android/VulkanAndroidPlatform.h b/Engine/Source/Runtime/VulkanRHI/Private/Android/VulkanAndroidPlatform.h
--- a/Engine/Source/Runtime/VulkanRHI/Private/Android/VulkanAndroidPlatform.h
+++ b/Engine/Source/Runtime/VulkanRHI/Private/Android/VulkanAndroidPlatform.h
@@ -24,6 +24,7 @@
 #define VULKAN_SUPPORTS_TRANSIENT_RESOURCE_ALLOCATOR 0
 #define VULKAN_SUPPORTS_DRIVER_PROPERTIES			0
 #define VULKAN_SUPPORTS_DESCRIPTOR_INDEXING			(VULKAN_RHI_RAYTRACING)
+#define VULKAN_SUPPORTS_GPU_CRASH_DUMPS				1
 
 #define UE_VK_API_VERSION							VK_API_VERSION_1_1
 
@@ -174,6 +175,8 @@ public:
 		return VK_SHADER_STAGE_FRAGMENT_BIT | VK_SHADER_STAGE_COMPUTE_BIT;
 	}
 
+	static void WriteCrashMarker(const FOptionalVulkanDeviceExtensions& OptionalExtensions, FVulkanCmdBuffer* CmdBuffer, VkBuffer DestBuffer, const TArrayView<uint32>& Entries, bool bAdding);
+
 protected:
 	static void* VulkanLib;
 	static bool bAttemptedLoad;
diff --git a/Engine/Source/Runtime/VulkanRHI/Private/Linux/VulkanLinuxPlatform.cpp b/Engine/Source/Runtime/VulkanRHI/Private/Linux/VulkanLinuxPlatform.cpp
--- a/Engine/Source/Runtime/VulkanRHI/Private/Linux/VulkanLinuxPlatform.cpp
+++ b/Engine/Source/Runtime/VulkanRHI/Private/Linux/VulkanLinuxPlatform.cpp
@@ -245,21 +245,25 @@ void FVulkanLinuxPlatform::CreateSurface(void* WindowHandle, VkInstance Instance
 	}
 }
 
-void FVulkanLinuxPlatform::WriteCrashMarker(const FOptionalVulkanDeviceExtensions& OptionalExtensions, VkCommandBuffer CmdBuffer, VkBuffer DestBuffer, const TArrayView<uint32>& Entries, bool bAdding)
+void FVulkanLinuxPlatform::WriteCrashMarker(const FOptionalVulkanDeviceExtensions& OptionalExtensions, FVulkanCmdBuffer* CmdBuffer, VkBuffer DestBuffer, const TArrayView<uint32>& Entries, bool bAdding)
 {
 	ensure(Entries.Num() <= GMaxCrashBufferEntries);
 
 	if (OptionalExtensions.HasAMDBufferMarker)
 	{
 		// AMD API only allows updating one entry at a time. Assume buffer has entry 0 as num entries
-		VulkanDynamicAPI::vkCmdWriteBufferMarkerAMD(CmdBuffer, VK_PIPELINE_STAGE_BOTTOM_OF_PIPE_BIT, DestBuffer, 0, Entries.Num());
+		VulkanDynamicAPI::vkCmdWriteBufferMarkerAMD(CmdBuffer->GetHandle(), VK_PIPELINE_STAGE_BOTTOM_OF_PIPE_BIT, DestBuffer, 0, Entries.Num());
 		if (bAdding)
 		{
 			int32 LastIndex = Entries.Num() - 1;
 			// +1 size as entries start at index 1
-			VulkanDynamicAPI::vkCmdWriteBufferMarkerAMD(CmdBuffer, VK_PIPELINE_STAGE_BOTTOM_OF_PIPE_BIT, DestBuffer, (1 + LastIndex) * sizeof(uint32), Entries[LastIndex]);
+			VulkanDynamicAPI::vkCmdWriteBufferMarkerAMD(CmdBuffer->GetHandle(), VK_PIPELINE_STAGE_BOTTOM_OF_PIPE_BIT, DestBuffer, (1 + LastIndex) * sizeof(uint32), Entries[LastIndex]);
 		}
 	}
+	else
+	{
+		WriteCrashMarkerWithoutExtensions(CmdBuffer, DestBuffer, Entries, bAdding);
+	}
 
 	if (OptionalExtensions.HasNVDiagnosticCheckpoints)
 	{
@@ -267,7 +271,7 @@ void FVulkanLinuxPlatform::WriteCrashMarker(const FOptionalVulkanDeviceExtension
 		{
 			int32 LastIndex = Entries.Num() - 1;
 			uint32 Value = Entries[LastIndex];
-			VulkanDynamicAPI::vkCmdSetCheckpointNV(CmdBuffer, (void*)(size_t)Value);
+			VulkanDynamicAPI::vkCmdSetCheckpointNV(CmdBuffer->GetHandle(), (void*)(size_t)Value);
 		}
 	}
 }
diff --git a/Engine/Source/Runtime/VulkanRHI/Private/Linux/VulkanLinuxPlatform.h b/Engine/Source/Runtime/VulkanRHI/Private/Linux/VulkanLinuxPlatform.h
--- a/Engine/Source/Runtime/VulkanRHI/Private/Linux/VulkanLinuxPlatform.h
+++ b/Engine/Source/Runtime/VulkanRHI/Private/Linux/VulkanLinuxPlatform.h
@@ -52,7 +52,7 @@ public:
 
 	static void CreateSurface(void* WindowHandle, VkInstance Instance, VkSurfaceKHR* OutSurface);
 
-	static void WriteCrashMarker(const FOptionalVulkanDeviceExtensions& OptionalExtensions, VkCommandBuffer CmdBuffer, VkBuffer DestBuffer, const TArrayView<uint32>& Entries, bool bAdding);
+	static void WriteCrashMarker(const FOptionalVulkanDeviceExtensions& OptionalExtensions, FVulkanCmdBuffer* CmdBuffer, VkBuffer DestBuffer, const TArrayView<uint32>& Entries, bool bAdding);
 
 protected:
 	static void* VulkanLib;
diff --git a/Engine/Source/Runtime/VulkanRHI/Private/VulkanDevice.cpp b/Engine/Source/Runtime/VulkanRHI/Private/VulkanDevice.cpp
--- a/Engine/Source/Runtime/VulkanRHI/Private/VulkanDevice.cpp
+++ b/Engine/Source/Runtime/VulkanRHI/Private/VulkanDevice.cpp
@@ -1176,27 +1176,24 @@ void FVulkanDevice::InitGPU()
 #if VULKAN_SUPPORTS_GPU_CRASH_DUMPS
 	if (UE::RHI::UseGPUCrashDebugging())
 	{
-		if (OptionalDeviceExtensions.HasAMDBufferMarker)
-		{
-			VkBufferCreateInfo CreateInfo;
-			ZeroVulkanStruct(CreateInfo, VK_STRUCTURE_TYPE_BUFFER_CREATE_INFO);
-			CreateInfo.size = GMaxCrashBufferEntries * sizeof(uint32_t);
-			CreateInfo.usage = VK_BUFFER_USAGE_TRANSFER_DST_BIT;
-			VERIFYVULKANRESULT(VulkanRHI::vkCreateBuffer(Device, &CreateInfo, VULKAN_CPU_ALLOCATOR, &CrashMarker.Buffer));
-
-			VkMemoryRequirements MemReq;
-			FMemory::Memzero(MemReq);
-			VulkanRHI::vkGetBufferMemoryRequirements(Device, CrashMarker.Buffer, &MemReq);
-
-			CrashMarker.Allocation = DeviceMemoryManager.Alloc(false, CreateInfo.size, MemReq.memoryTypeBits, VK_MEMORY_PROPERTY_HOST_VISIBLE_BIT |
-				VK_MEMORY_PROPERTY_HOST_COHERENT_BIT, nullptr, VULKAN_MEMORY_MEDIUM_PRIORITY, false, __FILE__, __LINE__);
-
-			uint32* Entry = (uint32*)CrashMarker.Allocation->Map(VK_WHOLE_SIZE, 0);
-			check(Entry);
-			// Start with 0 entries
-			*Entry = 0;
-			VERIFYVULKANRESULT(VulkanRHI::vkBindBufferMemory(Device, CrashMarker.Buffer, CrashMarker.Allocation->GetHandle(), 0));
-		}
+		VkBufferCreateInfo CreateInfo;
+		ZeroVulkanStruct(CreateInfo, VK_STRUCTURE_TYPE_BUFFER_CREATE_INFO);
+		CreateInfo.size = GMaxCrashBufferEntries * sizeof(uint32_t);
+		CreateInfo.usage = VK_BUFFER_USAGE_TRANSFER_DST_BIT;
+		VERIFYVULKANRESULT(VulkanRHI::vkCreateBuffer(Device, &CreateInfo, VULKAN_CPU_ALLOCATOR, &CrashMarker.Buffer));
+
+		VkMemoryRequirements MemReq;
+		FMemory::Memzero(MemReq);
+		VulkanRHI::vkGetBufferMemoryRequirements(Device, CrashMarker.Buffer, &MemReq);
+
+		CrashMarker.Allocation = DeviceMemoryManager.Alloc(false, CreateInfo.size, MemReq.memoryTypeBits, VK_MEMORY_PROPERTY_HOST_VISIBLE_BIT |
+			VK_MEMORY_PROPERTY_HOST_COHERENT_BIT, nullptr, VULKAN_MEMORY_MEDIUM_PRIORITY, false, __FILE__, __LINE__);
+
+		uint32* Entry = (uint32*)CrashMarker.Allocation->Map(VK_WHOLE_SIZE, 0);
+		check(Entry);
+		// Start with 0 entries
+		*Entry = 0;
+		VERIFYVULKANRESULT(VulkanRHI::vkBindBufferMemory(Device, CrashMarker.Buffer, CrashMarker.Allocation->GetHandle(), 0));
 	}
 #endif
 
diff --git a/Engine/Source/Runtime/VulkanRHI/Private/VulkanGPUProfiler.h b/Engine/Source/Runtime/VulkanRHI/Private/VulkanGPUProfiler.h
--- a/Engine/Source/Runtime/VulkanRHI/Private/VulkanGPUProfiler.h
+++ b/Engine/Source/Runtime/VulkanRHI/Private/VulkanGPUProfiler.h
@@ -192,8 +192,9 @@ struct FVulkanGPUProfiler : public FGPUProfiler
 	FVulkanCommandListContext* CmdContext;
 
 #if VULKAN_SUPPORTS_GPU_CRASH_DUMPS
-	void PushMarkerForCrash(VkCommandBuffer CmdBuffer, VkBuffer DestBuffer, const TCHAR* Name);
-	void PopMarkerForCrash(VkCommandBuffer CmdBuffer, VkBuffer DestBuffer);
+	void PrepareCrashMarkerkBuffer();
+	void PushMarkerForCrash(FVulkanCmdBuffer* CmdBuffer, VkBuffer DestBuffer, const TCHAR* Name);
+	void PopMarkerForCrash(FVulkanCmdBuffer* CmdBuffer, VkBuffer DestBuffer);
 	void DumpCrashMarkers(void* BufferData);
 #endif
 
@@ -201,7 +202,6 @@ struct FVulkanGPUProfiler : public FGPUProfiler
 	TMap<uint32, FString> CachedStrings;
 	TArray<uint32> PushPopStack;
 
-	FVulkanTimingQueryPool* LocalTracePointsQueryPool;
 	TArray<uint64> CrashMarkers;
 	bool bBeginFrame;
 };
diff --git a/Engine/Source/Runtime/VulkanRHI/Private/VulkanGenericPlatform.cpp b/Engine/Source/Runtime/VulkanRHI/Private/VulkanGenericPlatform.cpp
--- a/Engine/Source/Runtime/VulkanRHI/Private/VulkanGenericPlatform.cpp
+++ b/Engine/Source/Runtime/VulkanRHI/Private/VulkanGenericPlatform.cpp
@@ -5,6 +5,7 @@
 #include "VulkanDevice.h"
 #include "HAL/FileManager.h"
 #include "Misc/CommandLine.h"
+#include "VulkanCommandBuffer.h"
 
 static TAutoConsoleVariable<int32> CVarVulkanUseProfileCheck(
 	TEXT("r.Vulkan.UseProfileCheck"),
@@ -178,4 +179,24 @@ FString FVulkanGenericPlatform::GetVulkanProfileNameForFeatureLevel(ERHIFeatureL
 		ProfileName += TEXT("_RT");
 	}
 	return ProfileName;
-}
\ No newline at end of file
+}
+
+void FVulkanGenericPlatform::WriteCrashMarkerWithoutExtensions(FVulkanCmdBuffer* CmdBuffer, VkBuffer DestBuffer, const TArrayView<uint32>& Entries, bool bAdding)
+{
+	// Buffer writes cannot be recorded inside of a render pass
+	if (CmdBuffer->IsOutsideRenderPass())
+	{
+		// execution barrier
+		VulkanDynamicAPI::vkCmdPipelineBarrier(CmdBuffer->GetHandle(), VK_PIPELINE_STAGE_BOTTOM_OF_PIPE_BIT, VK_PIPELINE_STAGE_TOP_OF_PIPE_BIT, 0, 0, nullptr, 0, nullptr, 0, nullptr);
+
+		VulkanDynamicAPI::vkCmdFillBuffer(CmdBuffer->GetHandle(), DestBuffer, 0, sizeof(uint32), Entries.Num());
+		if (bAdding)
+		{
+			int32 LastIndex = Entries.Num() - 1;
+			VulkanDynamicAPI::vkCmdFillBuffer(CmdBuffer->GetHandle(), DestBuffer, (1 + LastIndex) * sizeof(uint32), sizeof(uint32), Entries[LastIndex]);
+		}
+
+		// execution barrier
+		VulkanDynamicAPI::vkCmdPipelineBarrier(CmdBuffer->GetHandle(), VK_PIPELINE_STAGE_BOTTOM_OF_PIPE_BIT, VK_PIPELINE_STAGE_TOP_OF_PIPE_BIT, 0, 0, nullptr, 0, nullptr, 0, nullptr);
+	}
+}
diff --git a/Engine/Source/Runtime/VulkanRHI/Private/VulkanGenericPlatform.h b/Engine/Source/Runtime/VulkanRHI/Private/VulkanGenericPlatform.h
--- a/Engine/Source/Runtime/VulkanRHI/Private/VulkanGenericPlatform.h
+++ b/Engine/Source/Runtime/VulkanRHI/Private/VulkanGenericPlatform.h
@@ -13,6 +13,7 @@ class FVulkanDevice;
 class FVulkanRenderTargetLayout;
 struct FGfxPipelineDesc;
 class FVulkanPhysicalDeviceFeatures;
+class FVulkanCmdBuffer;
 
 enum EShaderPlatform : uint16;
 enum class EGpuVendorId : uint32;
@@ -79,7 +80,9 @@ public:
 
 	static bool RegisterGPUWork() { return true; }
 
-	static void WriteCrashMarker(const FOptionalVulkanDeviceExtensions& OptionalExtensions, VkCommandBuffer CmdBuffer, VkBuffer DestBuffer, const TArrayView<uint32>& Entries, bool bAdding) {}
+	static void WriteCrashMarker(const FOptionalVulkanDeviceExtensions& OptionalExtensions, FVulkanCmdBuffer* CmdBuffer, VkBuffer DestBuffer, const TArrayView<uint32>& Entries, bool bAdding) {}
+
+	static void WriteCrashMarkerWithoutExtensions(FVulkanCmdBuffer* CmdBuffer, VkBuffer DestBuffer, const TArrayView<uint32>& Entries, bool bAdding);
 
 	// Allow the platform code to restrict the device features
 	static void RestrictEnabledPhysicalDeviceFeatures(FVulkanPhysicalDeviceFeatures* InOutFeaturesToEnable);
diff --git a/Engine/Source/Runtime/VulkanRHI/Private/VulkanRHI.cpp b/Engine/Source/Runtime/VulkanRHI/Private/VulkanRHI.cpp
--- a/Engine/Source/Runtime/VulkanRHI/Private/VulkanRHI.cpp
+++ b/Engine/Source/Runtime/VulkanRHI/Private/VulkanRHI.cpp
@@ -1030,11 +1030,6 @@ void FVulkanDynamicRHI::InitInstance()
 			UE_LOG(LogVulkanRHI, Display, TEXT("Wave Operations have been DISABLED (missing stages=0x%x operations=0x%x)."), MissingStageFlags, MissingOperationFlags);
 		}
 
-		if (UE::RHI::UseGPUCrashDebugging() && !Device->GetOptionalExtensions().HasGPUCrashDumpExtensions())
-		{
-			UE_LOG(LogVulkanRHI, Warning, TEXT("Tried to enable GPU crash debugging but no extension found! Will use local tracepoints."));
-		}
-
 		FHardwareInfo::RegisterHardwareInfo(NAME_RHI, TEXT("Vulkan"));
 
 		SavePipelineCacheCmd = IConsoleManager::Get().RegisterConsoleCommand(
@@ -1253,7 +1248,7 @@ void FVulkanCommandListContext::RHIEndFrame()
 		#if VULKAN_SUPPORTS_GPU_CRASH_DUMPS
 			if (GpuProfiler.bTrackingGPUCrashData)
 			{
-				GpuProfiler.PushMarkerForCrash(GetCommandBufferManager()->GetActiveCmdBuffer()->GetHandle(), Device->GetCrashMarkerBuffer(), GetNameStr());
+				GpuProfiler.PushMarkerForCrash(GetCommandBufferManager()->GetActiveCmdBuffer(), Device->GetCrashMarkerBuffer(), GetNameStr());
 			}
 		#endif
 			if (GpuProfiler.IsProfilingGPU())
@@ -1276,7 +1271,7 @@ void FVulkanCommandListContext::RHIEndFrame()
 		#if VULKAN_SUPPORTS_GPU_CRASH_DUMPS
 			if (GpuProfiler.bTrackingGPUCrashData)
 			{
-				GpuProfiler.PopMarkerForCrash(GetCommandBufferManager()->GetActiveCmdBuffer()->GetHandle(), Device->GetCrashMarkerBuffer());
+				GpuProfiler.PopMarkerForCrash(GetCommandBufferManager()->GetActiveCmdBuffer(), Device->GetCrashMarkerBuffer());
 			}
 		#endif
 		}
diff --git a/Engine/Source/Runtime/VulkanRHI/Private/VulkanUtil.cpp b/Engine/Source/Runtime/VulkanRHI/Private/VulkanUtil.cpp
--- a/Engine/Source/Runtime/VulkanRHI/Private/VulkanUtil.cpp
+++ b/Engine/Source/Runtime/VulkanRHI/Private/VulkanUtil.cpp
@@ -519,17 +519,12 @@ FVulkanGPUProfiler::FVulkanGPUProfiler(FVulkanCommandListContext* InCmd, FVulkan
 	: bCommandlistSubmitted(false)
 	, Device(InDevice)
 	, CmdContext(InCmd)
-	, LocalTracePointsQueryPool(nullptr)
 	, bBeginFrame(false)
 {
 }
 
 FVulkanGPUProfiler::~FVulkanGPUProfiler()
 {
-	if (LocalTracePointsQueryPool != nullptr)
-	{
-		delete LocalTracePointsQueryPool;
-	}
 }
 
 void FVulkanGPUProfiler::BeginFrame()
@@ -543,22 +538,13 @@ void FVulkanGPUProfiler::BeginFrame()
 		GPUCrashDataDepth = CrashCollectionDataDepth ? CrashCollectionDataDepth->GetValueOnRenderThread() : -1;
 		if (GPUCrashDataDepth == -1 || GPUCrashDataDepth > GMaxCrashBufferEntries)
 		{
-			if (Device->GetOptionalExtensions().HasAMDBufferMarker)
+			static bool bChecked = false;
+			if (!bChecked)
 			{
-				static bool bChecked = false;
-				if (!bChecked)
-				{
-					bChecked = true;
-					UE_LOG(LogVulkanRHI, Warning, TEXT("Clamping r.gpucrash.datadepth to %d"), GMaxCrashBufferEntries);
-				}
-				GPUCrashDataDepth = GMaxCrashBufferEntries;
+				bChecked = true;
+				UE_LOG(LogVulkanRHI, Warning, TEXT("Clamping r.gpucrash.datadepth to %d"), GMaxCrashBufferEntries);
 			}
-		}
-
-		// Use local tracepoints if no extension is available
-		if (!Device->GetOptionalExtensions().HasGPUCrashDumpExtensions() && LocalTracePointsQueryPool == nullptr)
-		{
-			LocalTracePointsQueryPool = new FVulkanTimingQueryPool(Device, CmdContext->GetCommandBufferManager(), GMaxCrashBufferEntries);
+			GPUCrashDataDepth = GMaxCrashBufferEntries;
 		}
 	}
 #endif
@@ -568,15 +554,6 @@ void FVulkanGPUProfiler::BeginFrame()
 	check(!bTrackingEvents);
 	check(!CurrentEventNodeFrame); // this should have already been cleaned up and the end of the previous frame
 
-	if (UE::RHI::UseGPUCrashDebugging() && !Device->GetOptionalExtensions().HasGPUCrashDumpExtensions())
-	{
-		VulkanRHI::vkCmdResetQueryPool(Device->GetImmediateContext().GetCommandBufferManager()->GetActiveCmdBuffer()->GetHandle(), LocalTracePointsQueryPool->GetHandle(), 0, GMaxCrashBufferEntries);
-
-		PushPopStack.Reset();
-		CrashMarkers.Reset();
-		CrashMarkers.AddZeroed(GMaxCrashBufferEntries);
-	}
-
 	bBeginFrame = true;
 
 	// latch the bools from the game thread into our private copy
@@ -680,14 +657,9 @@ void FVulkanGPUProfiler::EndFrame()
 }
 
 #if VULKAN_SUPPORTS_GPU_CRASH_DUMPS
-void FVulkanGPUProfiler::PushMarkerForCrash(VkCommandBuffer CmdBuffer, VkBuffer DestBuffer, const TCHAR* Name)
-{
-	if (!Device->GetOptionalExtensions().HasGPUCrashDumpExtensions() && !bBeginFrame)
-	{
-		// If using local trace points, ignore any markers pushed before begin frame or after end frame.
-		return;
-	}
 
+void FVulkanGPUProfiler::PushMarkerForCrash(FVulkanCmdBuffer* CmdBuffer, VkBuffer DestBuffer, const TCHAR* Name)
+{
 	uint32 CRC = 0;
 	if (GPUCrashDataDepth < 0 || PushPopStack.Num() < GPUCrashDataDepth)
 	{
@@ -711,45 +683,35 @@ void FVulkanGPUProfiler::PushMarkerForCrash(VkCommandBuffer CmdBuffer, VkBuffer
 
 	PushPopStack.Push(CRC);
 	FVulkanPlatform::WriteCrashMarker(Device->GetOptionalExtensions(), CmdBuffer, DestBuffer, TArrayView<uint32>(PushPopStack), true);
-
-	if (UE::RHI::UseGPUCrashDebugging() && !Device->GetOptionalExtensions().HasGPUCrashDumpExtensions())
-	{
-		VulkanRHI::vkCmdWriteTimestamp(CmdBuffer, VK_PIPELINE_STAGE_TOP_OF_PIPE_BIT, LocalTracePointsQueryPool->GetHandle(), PushPopStack.Num() - 1);
-	}
 }
 
-void FVulkanGPUProfiler::PopMarkerForCrash(VkCommandBuffer CmdBuffer, VkBuffer DestBuffer)
+void FVulkanGPUProfiler::PopMarkerForCrash(FVulkanCmdBuffer* CmdBuffer, VkBuffer DestBuffer)
 {
-	if (!Device->GetOptionalExtensions().HasGPUCrashDumpExtensions() && !bBeginFrame)
+	if (PushPopStack.Num() > 0)
 	{
-		// If using local trace points, ignore any markers popped before begin frame or after end frame.
-		return;
+		PushPopStack.Pop(EAllowShrinking::No);
+		FVulkanPlatform::WriteCrashMarker(Device->GetOptionalExtensions(), CmdBuffer, DestBuffer, TArrayView<uint32>(PushPopStack), false);
 	}
+}
 
-	if (PushPopStack.Num() > 0)
+void FVulkanGPUProfiler::DumpCrashMarkers(void* BufferData)
+{
 	{
-		if (Device->GetOptionalExtensions().HasGPUCrashDumpExtensions())
+		if (Device->GetOptionalExtensions().HasAMDBufferMarker)
 		{
-			PushPopStack.Pop(EAllowShrinking::No);
-			FVulkanPlatform::WriteCrashMarker(Device->GetOptionalExtensions(), CmdBuffer, DestBuffer, TArrayView<uint32>(PushPopStack), false);
+			UE_LOG(LogVulkanRHI, Error, TEXT("Breadcrumbs using ADM extension"));
 		}
-		else if (UE::RHI::UseGPUCrashDebugging())
+		else
 		{
-			VulkanRHI::vkGetQueryPoolResults(Device->GetInstanceHandle(), LocalTracePointsQueryPool->GetHandle(), 0, PushPopStack.Num(), sizeof(uint64) * GMaxCrashBufferEntries, CrashMarkers.GetData(), sizeof(uint64), VK_QUERY_RESULT_64_BIT | VK_QUERY_RESULT_WITH_AVAILABILITY_BIT);
+			UE_LOG(LogVulkanRHI, Error, TEXT("Breadcrumbs without extensions"));
 		}
-	}
-}
 
-void FVulkanGPUProfiler::DumpCrashMarkers(void* BufferData)
-{
-	if (Device->GetOptionalExtensions().HasAMDBufferMarker)
-	{
 		uint32* Entries = (uint32*)BufferData;
 		uint32 NumCRCs = *Entries++;
 		for (uint32 Index = 0; Index < NumCRCs; ++Index)
 		{
 			const FString* Frame = CachedStrings.Find(*Entries);
-			UE_LOG(LogVulkanRHI, Error, TEXT("[VK_AMD_buffer_info] %i: %s (CRC 0x%x)"), Index, Frame ? *(*Frame) : TEXT("<undefined>"), *Entries);
+			UE_LOG(LogVulkanRHI, Error, TEXT("[GPU Breadcrumb] %i: %s (CRC 0x%x)"), Index, Frame ? *(*Frame) : TEXT("<undefined>"), *Entries);
 			++Entries;
 		}
 	}
@@ -783,21 +745,6 @@ void FVulkanGPUProfiler::DumpCrashMarkers(void* BufferData)
 			GLog->Panic();
 		}
 	}
-
-	if (!Device->GetOptionalExtensions().HasGPUCrashDumpExtensions())
-	{
-		UE_LOG(LogVulkanRHI, Warning, TEXT("Printing trace points."));
-
-		for (int32 i = 0; i < PushPopStack.Num(); ++i)
-		{
-			const FString* InsertedFrame = CachedStrings.Find(PushPopStack[i]);
-			const FString FrameName = InsertedFrame ? *InsertedFrame : TEXT("<undefined>");
-
-			UE_LOG(LogVulkanRHI, Warning, TEXT("[gpu_crash_markers] %s"), (CrashMarkers[i] != 0) ? *FrameName : TEXT("unavailable"));
-		}
-
-		GLog->Panic();
-	}
 }
 #endif // VULKAN_SUPPORTS_GPU_CRASH_DUMPS
 
diff --git a/Engine/Source/Runtime/VulkanRHI/Private/Windows/VulkanWindowsPlatform.cpp b/Engine/Source/Runtime/VulkanRHI/Private/Windows/VulkanWindowsPlatform.cpp
--- a/Engine/Source/Runtime/VulkanRHI/Private/Windows/VulkanWindowsPlatform.cpp
+++ b/Engine/Source/Runtime/VulkanRHI/Private/Windows/VulkanWindowsPlatform.cpp
@@ -197,22 +197,25 @@ bool FVulkanWindowsPlatform::SupportsDeviceLocalHostVisibleWithNoPenalty(EGpuVen
 	return (VendorId == EGpuVendorId::Amd && bIsWin10);
 }
 
-
-void FVulkanWindowsPlatform::WriteCrashMarker(const FOptionalVulkanDeviceExtensions& OptionalExtensions, VkCommandBuffer CmdBuffer, VkBuffer DestBuffer, const TArrayView<uint32>& Entries, bool bAdding)
+void FVulkanWindowsPlatform::WriteCrashMarker(const FOptionalVulkanDeviceExtensions& OptionalExtensions, FVulkanCmdBuffer* CmdBuffer, VkBuffer DestBuffer, const TArrayView<uint32>& Entries, bool bAdding)
 {
 	ensure(Entries.Num() <= GMaxCrashBufferEntries);
 
 	if (OptionalExtensions.HasAMDBufferMarker)
 	{
 		// AMD API only allows updating one entry at a time. Assume buffer has entry 0 as num entries
-		VulkanDynamicAPI::vkCmdWriteBufferMarkerAMD(CmdBuffer, VK_PIPELINE_STAGE_BOTTOM_OF_PIPE_BIT, DestBuffer, 0, Entries.Num());
+		VulkanDynamicAPI::vkCmdWriteBufferMarkerAMD(CmdBuffer->GetHandle(), VK_PIPELINE_STAGE_BOTTOM_OF_PIPE_BIT, DestBuffer, 0, Entries.Num());
 		if (bAdding)
 		{
 			int32 LastIndex = Entries.Num() - 1;
 			// +1 size as entries start at index 1
-			VulkanDynamicAPI::vkCmdWriteBufferMarkerAMD(CmdBuffer, VK_PIPELINE_STAGE_BOTTOM_OF_PIPE_BIT, DestBuffer, (1 + LastIndex) * sizeof(uint32), Entries[LastIndex]);
+			VulkanDynamicAPI::vkCmdWriteBufferMarkerAMD(CmdBuffer->GetHandle(), VK_PIPELINE_STAGE_BOTTOM_OF_PIPE_BIT, DestBuffer, (1 + LastIndex) * sizeof(uint32), Entries[LastIndex]);
 		}
 	}
+	else
+	{
+		WriteCrashMarkerWithoutExtensions(CmdBuffer, DestBuffer, Entries, bAdding);
+	}
 
 	if (OptionalExtensions.HasNVDiagnosticCheckpoints)
 	{
@@ -220,7 +223,7 @@ void FVulkanWindowsPlatform::WriteCrashMarker(const FOptionalVulkanDeviceExtensi
 		{
 			int32 LastIndex = Entries.Num() - 1;
 			uint32 Value = Entries[LastIndex];
-			VulkanDynamicAPI::vkCmdSetCheckpointNV(CmdBuffer, (void*)(size_t)Value);
+			VulkanDynamicAPI::vkCmdSetCheckpointNV(CmdBuffer->GetHandle(), (void*)(size_t)Value);
 		}
 	}
 }
diff --git a/Engine/Source/Runtime/VulkanRHI/Private/Windows/VulkanWindowsPlatform.h b/Engine/Source/Runtime/VulkanRHI/Private/Windows/VulkanWindowsPlatform.h
--- a/Engine/Source/Runtime/VulkanRHI/Private/Windows/VulkanWindowsPlatform.h
+++ b/Engine/Source/Runtime/VulkanRHI/Private/Windows/VulkanWindowsPlatform.h
@@ -60,7 +60,7 @@ public:
 
 	static bool SupportsDeviceLocalHostVisibleWithNoPenalty(EGpuVendorId VendorId);
 
-	static void WriteCrashMarker(const FOptionalVulkanDeviceExtensions& OptionalExtensions, VkCommandBuffer CmdBuffer, VkBuffer DestBuffer, const TArrayView<uint32>& Entries, bool bAdding);
+	static void WriteCrashMarker(const FOptionalVulkanDeviceExtensions& OptionalExtensions, FVulkanCmdBuffer* CmdBuffer, VkBuffer DestBuffer, const TArrayView<uint32>& Entries, bool bAdding);
 
 private:
 	static bool bAttemptedLoad;
