Adding HW Depth Resolve for Vulkan.

diff --git a/Engine/Source/Runtime/RHI/Private/RHI.cpp b/Engine/Source/Runtime/RHI/Private/RHI.cpp
--- a/Engine/Source/Runtime/RHI/Private/RHI.cpp
+++ b/Engine/Source/Runtime/RHI/Private/RHI.cpp
@@ -1528,6 +1528,12 @@ void FRHIRenderPassInfo::ConvertToRenderTargetsInfo(FRHISetRenderTargetsInfo& Ou
 	OutRTInfo.bClearDepth = (DepthLoadAction == ERenderTargetLoadAction::EClear);
 	OutRTInfo.bClearStencil = (StencilLoadAction == ERenderTargetLoadAction::EClear);
 
+	if (OutRTInfo.bHasResolveAttachments && DepthStencilRenderTarget.ResolveTarget && DepthStencilRenderTarget.ResolveTarget != DepthStencilRenderTarget.DepthStencilTarget)
+	{
+		OutRTInfo.DepthStencilResolveRenderTarget = OutRTInfo.DepthStencilRenderTarget;
+		OutRTInfo.DepthStencilResolveRenderTarget.Texture = DepthStencilRenderTarget.ResolveTarget;
+	}
+
 	OutRTInfo.ShadingRateTexture = ShadingRateTexture;
 	OutRTInfo.ShadingRateTextureCombiner = ShadingRateTextureCombiner;
 	OutRTInfo.MultiViewCount = MultiViewCount;
@@ -1625,6 +1631,20 @@ void FRHIRenderPassInfo::Validate() const
 			// 2. depth target must support InputAttachement
 			ensure(EnumHasAnyFlags(DepthStencilRenderTarget.DepthStencilTarget->GetFlags(), TexCreate_InputAttachmentRead));
 		}
+
+		if (DepthStencilRenderTarget.ResolveTarget && DepthStencilRenderTarget.ResolveTarget != DepthStencilRenderTarget.DepthStencilTarget)
+		{
+			// For depth resolve
+			// 1. RHI must support depth stencil resolve
+			// 2. Must be using MSAA resolve
+			// 3. Resolve target sample count must be 1
+			// 4. Resolve target format must be the same as the MSAA target format
+			ensureMsgf(GRHISupportsDepthStencilResolve, TEXT("Attempted to resolve depth/stencil target but feature is not supported."));
+			ensureMsgf(bIsMSAAResolve, TEXT("Depth/stencil resolve target is bound but resolve was not requested."));
+			ensureMsgf(DepthStencilRenderTarget.ResolveTarget->GetNumSamples() == 1, TEXT("Depth/stencil resolve targets must have a sample count of 1."));
+			ensureMsgf(DepthStencilRenderTarget.ResolveTarget->GetFormat() == DepthStencilRenderTarget.DepthStencilTarget->GetFormat(),
+				TEXT("Depth/stencil resolve targets must have the same format as the MSAA target."));
+		}
 	}
 	else
 	{
diff --git a/Engine/Source/Runtime/RHI/Public/RHIGlobals.h b/Engine/Source/Runtime/RHI/Public/RHIGlobals.h
--- a/Engine/Source/Runtime/RHI/Public/RHIGlobals.h
+++ b/Engine/Source/Runtime/RHI/Public/RHIGlobals.h
@@ -658,6 +658,9 @@ struct FRHIGlobals
 
 	/** True if RHI supports MSAA resolve with a custom shader */
 	bool SupportsMSAAShaderResolve = false;
+
+	/** Whether Depth Stencil MSAA Resolve Targets are supported. */
+	bool SupportsDepthStencilResolve = false;
 };
 
 extern RHI_API FRHIGlobals GRHIGlobals;
@@ -828,6 +831,7 @@ extern RHI_API FRHIGlobals GRHIGlobals;
 #define GRHIIsDebugLayerEnabled                                GRHIGlobals.IsDebugLayerEnabled
 #define GRHIHDRNeedsVendorExtensions						   GRHIGlobals.HDRNeedsVendorExtensions
 #define GRHISupportsMSAAShaderResolve						   GRHIGlobals.SupportsMSAAShaderResolve
+#define GRHISupportsDepthStencilResolve                        GRHIGlobals.SupportsDepthStencilResolve
 
 // Utility Getters
 
diff --git a/Engine/Source/Runtime/RHI/Public/RHIResources.h b/Engine/Source/Runtime/RHI/Public/RHIResources.h
--- a/Engine/Source/Runtime/RHI/Public/RHIResources.h
+++ b/Engine/Source/Runtime/RHI/Public/RHIResources.h
@@ -3505,6 +3505,8 @@ public:
 
 	// Depth/Stencil Render Target Info
 	FRHIDepthRenderTargetView DepthStencilRenderTarget;	
+	// Used when depth resolve is enabled.
+	FRHIDepthRenderTargetView DepthStencilResolveRenderTarget;
 	bool bClearDepth;
 	bool bClearStencil;
 
@@ -3557,8 +3559,11 @@ public:
 		// Need a separate struct so we can memzero/remove dependencies on reference counts
 		struct FHashableStruct
 		{
-			// *2 for color and resolves, depth goes in the second-to-last slot, shading rate goes in the last slot
-			FRHITexture* Texture[MaxSimultaneousRenderTargets*2 + 2];
+			// *2 for color and resolves
+			// depth goes in the third-to-last slot
+			// depth resolve goes in the second-to-last slot
+			// shading rate goes in the last slot
+			FRHITexture* Texture[MaxSimultaneousRenderTargets*2 + 3];
 			uint32 MipIndex[MaxSimultaneousRenderTargets];
 			uint32 ArraySliceIndex[MaxSimultaneousRenderTargets];
 			ERenderTargetLoadAction LoadAction[MaxSimultaneousRenderTargets];
@@ -3590,8 +3595,9 @@ public:
 					StoreAction[Index] = RTInfo.ColorRenderTarget[Index].StoreAction;
 				}
 
-				Texture[MaxSimultaneousRenderTargets] = RTInfo.DepthStencilRenderTarget.Texture;
-				Texture[MaxSimultaneousRenderTargets + 1] = RTInfo.ShadingRateTexture;
+				Texture[MaxSimultaneousRenderTargets*2] = RTInfo.DepthStencilRenderTarget.Texture;
+				Texture[MaxSimultaneousRenderTargets*2 + 1] = RTInfo.DepthStencilResolveRenderTarget.Texture;
+				Texture[MaxSimultaneousRenderTargets*2 + 2] = RTInfo.ShadingRateTexture;
 				DepthLoadAction = RTInfo.DepthStencilRenderTarget.DepthLoadAction;
 				DepthStoreAction = RTInfo.DepthStencilRenderTarget.DepthStoreAction;
 				StencilLoadAction = RTInfo.DepthStencilRenderTarget.StencilLoadAction;
diff --git a/Engine/Source/Runtime/RenderCore/Private/RenderGraphBuilder.cpp b/Engine/Source/Runtime/RenderCore/Private/RenderGraphBuilder.cpp
--- a/Engine/Source/Runtime/RenderCore/Private/RenderGraphBuilder.cpp
+++ b/Engine/Source/Runtime/RenderCore/Private/RenderGraphBuilder.cpp
@@ -181,6 +181,7 @@ void EnumerateTextureAccess(FRDGParameterStruct PassParameters, ERDGPassFlags Pa
 
 			if (FRDGTextureRef Texture = DepthStencil.GetTexture())
 			{
+				FRDGTextureRef ResolveTexture = DepthStencil.GetResolveTexture();
 				DepthStencil.GetDepthStencilAccess().EnumerateSubresources([&](ERHIAccess NewAccess, uint32 PlaneSlice)
 				{
 					FRDGTextureSubresourceRange Range = Texture->GetSubresourceRange();
@@ -193,6 +194,12 @@ void EnumerateTextureAccess(FRDGParameterStruct PassParameters, ERDGPassFlags Pa
 					}
 
 					AccessFunction(nullptr, Texture, NewAccess, RenderTargetAccess, Range);
+
+					if (ResolveTexture && ResolveTexture != Texture)
+					{
+						// If we're resolving depth stencil, it must be DSVWrite and ResolveDst
+						AccessFunction(nullptr, ResolveTexture, ERHIAccess::DSVWrite | ERHIAccess::ResolveDst, RenderTargetAccess, Range);
+					}
 				});
 			}
 
diff --git a/Engine/Source/Runtime/RenderCore/Private/RenderGraphPass.cpp b/Engine/Source/Runtime/RenderCore/Private/RenderGraphPass.cpp
--- a/Engine/Source/Runtime/RenderCore/Private/RenderGraphPass.cpp
+++ b/Engine/Source/Runtime/RenderCore/Private/RenderGraphPass.cpp
@@ -72,12 +72,27 @@ FRHIRenderPassInfo FRDGParameterStruct::GetRenderPassInfo() const
 	if (FRDGTextureRef Texture = DepthStencil.GetTexture())
 	{
 		const FExclusiveDepthStencil ExclusiveDepthStencil = DepthStencil.GetDepthStencilAccess();
-		const ERenderTargetStoreAction StoreAction = EnumHasAnyFlags(Texture->Desc.Flags, TexCreate_Memoryless) ? ERenderTargetStoreAction::ENoAction : ERenderTargetStoreAction::EStore;
+		ERenderTargetStoreAction StoreAction = EnumHasAnyFlags(Texture->Desc.Flags, TexCreate_Memoryless) ? ERenderTargetStoreAction::ENoAction : ERenderTargetStoreAction::EStore;
+		FRDGTextureRef ResolveTexture = DepthStencil.GetResolveTexture();
+		if (ResolveTexture)
+		{
+			// Silently skip the resolve if the resolve texture is the same as the render target texture.
+			if (ResolveTexture != Texture)
+			{
+				StoreAction = ERenderTargetStoreAction::EMultisampleResolve;
+			}
+			else
+			{
+				ResolveTexture = nullptr;
+			}
+		}
+
 		const ERenderTargetStoreAction DepthStoreAction = ExclusiveDepthStencil.IsUsingDepth() ? StoreAction : ERenderTargetStoreAction::ENoAction;
 		const ERenderTargetStoreAction StencilStoreAction = ExclusiveDepthStencil.IsUsingStencil() ? StoreAction : ERenderTargetStoreAction::ENoAction;
 
 		auto& DepthStencilTarget = RenderPassInfo.DepthStencilRenderTarget;
 		DepthStencilTarget.DepthStencilTarget = Texture->GetRHI();
+		DepthStencilTarget.ResolveTarget = ResolveTexture ? ResolveTexture->GetRHI() : nullptr;
 		DepthStencilTarget.Action = MakeDepthStencilTargetActions(
 			MakeRenderTargetActions(DepthStencil.GetDepthLoadAction(), DepthStoreAction),
 			MakeRenderTargetActions(DepthStencil.GetStencilLoadAction(), StencilStoreAction));
diff --git a/Engine/Source/Runtime/RenderCore/Private/RenderGraphValidation.cpp b/Engine/Source/Runtime/RenderCore/Private/RenderGraphValidation.cpp
--- a/Engine/Source/Runtime/RenderCore/Private/RenderGraphValidation.cpp
+++ b/Engine/Source/Runtime/RenderCore/Private/RenderGraphValidation.cpp
@@ -1315,6 +1315,11 @@ void FRDGUserValidation::SetAllowRHIAccess(const FRDGPass* Pass, bool bAllowAcce
 				Texture->GetDebugData().bAllowRHIAccess = bAllowAccess;
 			}
 
+			if (FRDGTextureRef ResolveTexture = RenderTargets.DepthStencil.GetResolveTexture())
+			{
+				ResolveTexture->GetDebugData().bAllowRHIAccess = bAllowAccess;
+			}
+
 			if (FRDGTexture* Texture = RenderTargets.ShadingRateTexture)
 			{
 				Texture->GetDebugData().bAllowRHIAccess = bAllowAccess;
diff --git a/Engine/Source/Runtime/RenderCore/Private/ShaderParameterStruct.cpp b/Engine/Source/Runtime/RenderCore/Private/ShaderParameterStruct.cpp
--- a/Engine/Source/Runtime/RenderCore/Private/ShaderParameterStruct.cpp
+++ b/Engine/Source/Runtime/RenderCore/Private/ShaderParameterStruct.cpp
@@ -451,6 +451,8 @@ bool FDepthStencilBinding::Validate() const
 	}
 	else
 	{
+		checkf(!ResolveTexture,
+			TEXT("Can't have a depth stencil resolve texture when no depth stencil target texture is bound."));
 		checkf(DepthLoadAction == ERenderTargetLoadAction::ENoAction,
 			TEXT("Can't have a depth load action when no texture is bound."));
 		checkf(StencilLoadAction == ERenderTargetLoadAction::ENoAction,
diff --git a/Engine/Source/Runtime/RenderCore/Public/RenderGraphResources.inl b/Engine/Source/Runtime/RenderCore/Public/RenderGraphResources.inl
--- a/Engine/Source/Runtime/RenderCore/Public/RenderGraphResources.inl
+++ b/Engine/Source/Runtime/RenderCore/Public/RenderGraphResources.inl
@@ -242,7 +242,11 @@ inline FGraphicsPipelineRenderTargetsInfo ExtractRenderTargetsInfo(const FRender
 		RenderTargetsInfo.StencilTargetLoadAction = DepthStencil.GetStencilLoadAction();
 
 		RenderTargetsInfo.DepthStencilAccess = DepthStencil.GetDepthStencilAccess();
-		const ERenderTargetStoreAction StoreAction = EnumHasAnyFlags(DepthTexture->Desc.Flags, TexCreate_Memoryless) ? ERenderTargetStoreAction::ENoAction : ERenderTargetStoreAction::EStore;
+		ERenderTargetStoreAction StoreAction = EnumHasAnyFlags(DepthTexture->Desc.Flags, TexCreate_Memoryless) ? ERenderTargetStoreAction::ENoAction : ERenderTargetStoreAction::EStore;
+		if (DepthStencil.GetResolveTexture())
+		{
+			StoreAction = ERenderTargetStoreAction::EMultisampleResolve;
+		}
 		RenderTargetsInfo.DepthTargetStoreAction = RenderTargetsInfo.DepthStencilAccess.IsUsingDepth() ? StoreAction : ERenderTargetStoreAction::ENoAction;
 		RenderTargetsInfo.StencilTargetStoreAction = RenderTargetsInfo.DepthStencilAccess.IsUsingStencil() ? StoreAction : ERenderTargetStoreAction::ENoAction;
 	}
diff --git a/Engine/Source/Runtime/RenderCore/Public/ShaderParameterMacros.h b/Engine/Source/Runtime/RenderCore/Public/ShaderParameterMacros.h
--- a/Engine/Source/Runtime/RenderCore/Public/ShaderParameterMacros.h
+++ b/Engine/Source/Runtime/RenderCore/Public/ShaderParameterMacros.h
@@ -634,6 +634,7 @@ struct FDepthStencilBinding
 		ERenderTargetLoadAction InStencilLoadAction,
 		FExclusiveDepthStencil InDepthStencilAccess)
 		: Texture(InTexture)
+		, ResolveTexture(nullptr)
 		, DepthLoadAction(InDepthLoadAction)
 		, StencilLoadAction(InStencilLoadAction)
 		, DepthStencilAccess(InDepthStencilAccess)
@@ -643,9 +644,38 @@ struct FDepthStencilBinding
 
 	FORCEINLINE FDepthStencilBinding(
 		FRDGTexture* InTexture,
+		FRDGTexture* InResolveTexture,
+		ERenderTargetLoadAction InDepthLoadAction,
+		ERenderTargetLoadAction InStencilLoadAction,
+		FExclusiveDepthStencil InDepthStencilAccess)
+		: Texture(InTexture)
+		, ResolveTexture(InResolveTexture)
+		, DepthLoadAction(InDepthLoadAction)
+		, StencilLoadAction(InStencilLoadAction)
+		, DepthStencilAccess(InDepthStencilAccess)
+	{
+		check(Validate());
+	}
+
+	FORCEINLINE FDepthStencilBinding(
+		FRDGTexture* InTexture,
+		ERenderTargetLoadAction InDepthLoadAction,
+		FExclusiveDepthStencil InDepthStencilAccess)
+		: Texture(InTexture)
+		, ResolveTexture(nullptr)
+		, DepthLoadAction(InDepthLoadAction)
+		, DepthStencilAccess(InDepthStencilAccess)
+	{
+		check(Validate());
+	}
+
+	FORCEINLINE FDepthStencilBinding(
+		FRDGTexture* InTexture,
+		FRDGTexture* InResolveTexture,
 		ERenderTargetLoadAction InDepthLoadAction,
 		FExclusiveDepthStencil InDepthStencilAccess)
 		: Texture(InTexture)
+		, ResolveTexture(InResolveTexture)
 		, DepthLoadAction(InDepthLoadAction)
 		, DepthStencilAccess(InDepthStencilAccess)
 	{
@@ -656,6 +686,10 @@ struct FDepthStencilBinding
 	{
 		return Texture;
 	}
+	FORCEINLINE FRDGTexture* GetResolveTexture() const
+	{
+		return ResolveTexture;
+	}
 	FORCEINLINE ERenderTargetLoadAction GetDepthLoadAction() const
 	{
 		return DepthLoadAction;
@@ -674,6 +708,7 @@ struct FDepthStencilBinding
 	{
 		return
 			Texture == Other.Texture &&
+			ResolveTexture == Other.ResolveTexture &&
 			Other.DepthLoadAction != ERenderTargetLoadAction::EClear &&
 			Other.StencilLoadAction != ERenderTargetLoadAction::EClear &&
 			DepthStencilAccess == Other.DepthStencilAccess;
@@ -685,6 +720,12 @@ struct FDepthStencilBinding
 		check(Validate());
 	}
 
+	void SetResolveTexture(FRDGTexture* InTexture)
+	{
+		ResolveTexture = InTexture;
+		check(Validate());
+	}
+
 	void SetDepthStencilAccess(FExclusiveDepthStencil InAccess)
 	{
 		DepthStencilAccess = InAccess;
@@ -709,6 +750,7 @@ private:
 	 * force the user to call FDepthStencilBinding() constructors. No defaults allowed.
 	 */
 	TAlignedShaderParameterPtr<FRDGTexture*> Texture = nullptr;
+	TAlignedShaderParameterPtr<FRDGTexture*> ResolveTexture = nullptr;
 	ERenderTargetLoadAction		DepthLoadAction		= ERenderTargetLoadAction::ENoAction;
 	ERenderTargetLoadAction		StencilLoadAction	= ERenderTargetLoadAction::ENoAction;
 	FExclusiveDepthStencil		DepthStencilAccess	= FExclusiveDepthStencil::DepthNop_StencilNop;
@@ -823,7 +865,7 @@ struct alignas(SHADER_PARAMETER_STRUCT_ALIGNMENT) FRenderTargetBindingSlots
 	};
 };
 
-static_assert(sizeof(FRenderTargetBindingSlots) == 240, "FRenderTargetBindingSlots needs to be same size on all platforms.");
+static_assert(sizeof(FRenderTargetBindingSlots) == 256, "FRenderTargetBindingSlots needs to be same size on all platforms.");
 
 /** Static array of shader resource shader that is initialized to nullptr. */
 template<typename TElement, uint32 NumElements>
diff --git a/Engine/Source/Runtime/Renderer/Private/MobileShadingRenderer.cpp b/Engine/Source/Runtime/Renderer/Private/MobileShadingRenderer.cpp
--- a/Engine/Source/Runtime/Renderer/Private/MobileShadingRenderer.cpp
+++ b/Engine/Source/Runtime/Renderer/Private/MobileShadingRenderer.cpp
@@ -594,7 +594,6 @@ void FMobileSceneRenderer::InitViews(
 	{
 		bKeepDepthContent = true;
 	}
-
 	// Update the bKeepDepthContent based on the mobile renderer status.
 	SceneTexturesConfig.bKeepDepthContent = bKeepDepthContent;
 	// If we render in a single pass MSAA targets can be memoryless
@@ -1418,6 +1417,7 @@ FRenderTargetBindingSlots FMobileSceneRenderer::InitRenderTargetBindings_Forward
 	FRDGTextureRef SceneColor = nullptr;
 	FRDGTextureRef SceneColorResolve = nullptr;
 	FRDGTextureRef SceneDepth = nullptr;
+	FRDGTextureRef SceneDepthResolve = nullptr;
 
 	// Verify using both MSAA sample count AND the scene color surface sample count, since on GLES you can't have MSAA color targets,
 	// so the color target would be created without MSAA, and MSAA is achieved through magical means (the framebuffer, being MSAA,
@@ -1435,14 +1435,14 @@ FRenderTargetBindingSlots FMobileSceneRenderer::InitRenderTargetBindings_Forward
 		{
 			SceneColor = ViewFamilyTexture;
 		}
-		SceneDepth = SceneTextures.Depth.Target;
 	}
 	else
 	{
 		SceneColor = SceneTextures.Color.Target;
 		SceneColorResolve = bMobileMSAA ? SceneTextures.Color.Resolve : nullptr;
-		SceneDepth = SceneTextures.Depth.Target;
 	}
+	SceneDepth = SceneTextures.Depth.Target;
+	SceneDepthResolve = GRHISupportsDepthStencilResolve && bMobileMSAA && SceneTextures.Depth.IsSeparate() ? SceneTextures.Depth.Resolve : nullptr;
 
 	FRenderTargetBindingSlots BasePassRenderTargets;
 	BasePassRenderTargets[0] = FRenderTargetBinding(SceneColor, SceneColorResolve, ERenderTargetLoadAction::EClear);
@@ -1459,9 +1459,9 @@ FRenderTargetBindingSlots FMobileSceneRenderer::InitRenderTargetBindings_Forward
 		BasePassRenderTargets[1] = FRenderTargetBinding(ViewFamilyTexture, nullptr, ERenderTargetLoadAction::EClear);
 	}
 	
-	BasePassRenderTargets.DepthStencil = bIsFullDepthPrepassEnabled ? 
-		FDepthStencilBinding(SceneDepth, ERenderTargetLoadAction::ELoad, ERenderTargetLoadAction::ELoad, FExclusiveDepthStencil::DepthRead_StencilWrite) : 
-		FDepthStencilBinding(SceneDepth, ERenderTargetLoadAction::EClear, ERenderTargetLoadAction::EClear, FExclusiveDepthStencil::DepthWrite_StencilWrite);
+	BasePassRenderTargets.DepthStencil = bIsFullDepthPrepassEnabled ?
+		FDepthStencilBinding(SceneDepth, SceneDepthResolve, ERenderTargetLoadAction::ELoad, ERenderTargetLoadAction::ELoad, FExclusiveDepthStencil::DepthRead_StencilWrite) :
+		FDepthStencilBinding(SceneDepth, SceneDepthResolve, ERenderTargetLoadAction::EClear, ERenderTargetLoadAction::EClear, FExclusiveDepthStencil::DepthWrite_StencilWrite);
 	BasePassRenderTargets.SubpassHint = ESubpassHint::None;
 	BasePassRenderTargets.NumOcclusionQueries = 0u;
 
@@ -1617,7 +1617,7 @@ void FMobileSceneRenderer::RenderForwardSinglePass(FRDGBuilder& GraphBuilder, FM
 	});
 	
 	// resolve MSAA depth
-	if (!bIsFullDepthPrepassEnabled)
+	if (!GRHISupportsDepthStencilResolve && !bIsFullDepthPrepassEnabled)
 	{
 		AddResolveSceneDepthPass(GraphBuilder, *ViewContext.ViewInfo, SceneTextures.Depth);
 	}
diff --git a/Engine/Source/Runtime/Renderer/Private/SceneTextures.cpp b/Engine/Source/Runtime/Renderer/Private/SceneTextures.cpp
--- a/Engine/Source/Runtime/Renderer/Private/SceneTextures.cpp
+++ b/Engine/Source/Runtime/Renderer/Private/SceneTextures.cpp
@@ -469,7 +469,7 @@ void FMinimalSceneTextures::InitializeViewFamily(FRDGBuilder& GraphBuilder, FVie
 					Desc.ArraySize, StereoDepthRHI->GetDesc().ArraySize);
 				SceneTextures.Depth.Resolve = RegisterExternalTexture(GraphBuilder, StereoDepthRHI, TEXT("SceneDepthZ"));
 			}
-			else
+			else if (Config.bKeepDepthContent)
 			{
 				SceneTextures.Depth.Resolve = GraphBuilder.CreateTexture(Desc, TEXT("SceneDepthZ"));
 			}
diff --git a/Engine/Source/Runtime/VulkanRHI/Private/VulkanBarriers.cpp b/Engine/Source/Runtime/VulkanRHI/Private/VulkanBarriers.cpp
--- a/Engine/Source/Runtime/VulkanRHI/Private/VulkanBarriers.cpp
+++ b/Engine/Source/Runtime/VulkanRHI/Private/VulkanBarriers.cpp
@@ -510,6 +510,14 @@ static void GetDepthOrStencilStageAndAccessFlags(ERHIAccess Access, VkPipelineSt
 	AccessFlags = 0;
 	uint32 ProcessedRHIFlags = 0;
 
+	if (EnumHasAllFlags(Access, ERHIAccess::ResolveDst))
+	{
+		// Despite being a depth/stencil target, resolve operations are part of the color attachment output stage
+		StageFlags |= VK_PIPELINE_STAGE_COLOR_ATTACHMENT_OUTPUT_BIT;
+		AccessFlags |= VK_ACCESS_COLOR_ATTACHMENT_WRITE_BIT;
+		ProcessedRHIFlags |= (uint32)ERHIAccess::ResolveDst;
+	}
+
 	if (EnumHasAnyFlags(Access, ERHIAccess::DSVWrite))
 	{
 		StageFlags |= VK_PIPELINE_STAGE_EARLY_FRAGMENT_TESTS_BIT | VK_PIPELINE_STAGE_LATE_FRAGMENT_TESTS_BIT;
diff --git a/Engine/Source/Runtime/VulkanRHI/Private/VulkanDevice.h b/Engine/Source/Runtime/VulkanRHI/Private/VulkanDevice.h
--- a/Engine/Source/Runtime/VulkanRHI/Private/VulkanDevice.h
+++ b/Engine/Source/Runtime/VulkanRHI/Private/VulkanDevice.h
@@ -82,6 +82,7 @@ struct FOptionalVulkanDeviceExtensions
 			uint64 HasSeparateDepthStencilLayouts : 1;
 			uint64 HasEXTHostQueryReset : 1;
 			uint64 HasQcomRenderPassShaderResolve : 1;
+			uint64 HasKHRDepthStencilResolve : 1;
 
 			// Promoted to 1.3
 			uint64 HasEXTTextureCompressionASTCHDR : 1;
diff --git a/Engine/Source/Runtime/VulkanRHI/Private/VulkanExtensions.cpp b/Engine/Source/Runtime/VulkanRHI/Private/VulkanExtensions.cpp
--- a/Engine/Source/Runtime/VulkanRHI/Private/VulkanExtensions.cpp
+++ b/Engine/Source/Runtime/VulkanRHI/Private/VulkanExtensions.cpp
@@ -1371,6 +1371,23 @@ public:
 	VkPhysicalDeviceShaderFloat16Int8Features  ShaderFloat16Int8Features;
 };
 
+// ***** VK_KHR_depth_stencil_resolve
+class FVulkanKHRDepthStencilResolveExtension : public FVulkanDeviceExtension
+{
+public:
+
+	FVulkanKHRDepthStencilResolveExtension(FVulkanDevice* InDevice)
+		: FVulkanDeviceExtension(InDevice, VK_KHR_DEPTH_STENCIL_RESOLVE_EXTENSION_NAME, VULKAN_EXTENSION_ENABLED)
+	{
+	}
+
+	virtual void PostPhysicalDeviceFeatures(FOptionalVulkanDeviceExtensions& ExtensionFlags) override final
+	{
+		ExtensionFlags.HasKHRDepthStencilResolve = 1;
+		GRHISupportsDepthStencilResolve = ExtensionFlags.HasKHRDepthStencilResolve;
+	}
+};
+
 // ***** VK_EXT_pipeline_creation_cache_control
 class FVulkanEXTPipelineCreationCacheControlExtension : public FVulkanDeviceExtension
 {
@@ -1483,6 +1500,7 @@ FVulkanDeviceExtensionArray FVulkanDeviceExtension::GetUESupportedDeviceExtensio
 	ADD_CUSTOM_EXTENSION(FVulkanEXTShaderDemoteToHelperInvocationExtension);
 	ADD_CUSTOM_EXTENSION(FVulkanKHR16BitStorageExtension);
 	ADD_CUSTOM_EXTENSION(FVulkanKHRShaderFloat16Int8Extension);
+	ADD_CUSTOM_EXTENSION(FVulkanKHRDepthStencilResolveExtension);
 	ADD_CUSTOM_EXTENSION(FVulkanEXTPipelineCreationCacheControlExtension);
 	ADD_CUSTOM_EXTENSION(FVulkanKHRFragmentShaderBarycentricExtension);
 	ADD_CUSTOM_EXTENSION(FVulkanNVComputeShaderDerivatives);
diff --git a/Engine/Source/Runtime/VulkanRHI/Private/VulkanPipeline.cpp b/Engine/Source/Runtime/VulkanRHI/Private/VulkanPipeline.cpp
--- a/Engine/Source/Runtime/VulkanRHI/Private/VulkanPipeline.cpp
+++ b/Engine/Source/Runtime/VulkanRHI/Private/VulkanPipeline.cpp
@@ -1006,6 +1006,7 @@ void FGfxPipelineDesc::FRenderTargets::ReadFrom(const FVulkanRenderTargetLayout&
 
 	bHasDepthStencil =			RTLayout.bHasDepthStencil != 0;
 	bHasResolveAttachments =	RTLayout.bHasResolveAttachments != 0;
+	bHasDepthStencilResolve =	RTLayout.bHasDepthStencilResolve != 0;
 	bHasFragmentDensityAttachment =	RTLayout.bHasFragmentDensityAttachment != 0;
 	NumUsedClearValues =		RTLayout.NumUsedClearValues;
 
@@ -1044,6 +1045,7 @@ void FGfxPipelineDesc::FRenderTargets::WriteInto(FVulkanRenderTargetLayout& Out)
 
 	Out.bHasDepthStencil =			bHasDepthStencil;
 	Out.bHasResolveAttachments =	bHasResolveAttachments;
+	Out.bHasDepthStencilResolve =	bHasDepthStencilResolve;
 	Out.bHasFragmentDensityAttachment =	bHasFragmentDensityAttachment;
 	Out.NumUsedClearValues =		NumUsedClearValues;
 
@@ -1092,6 +1094,7 @@ FArchive& operator << (FArchive& Ar, FGfxPipelineDesc::FRenderTargets& RTs)
 
 	Ar << RTs.bHasDepthStencil;
 	Ar << RTs.bHasResolveAttachments;
+	Ar << RTs.bHasDepthStencilResolve;
 	Ar << RTs.RenderPassCompatibleHash;
 	Ar << RTs.Extent3D;
 
diff --git a/Engine/Source/Runtime/VulkanRHI/Private/VulkanPipeline.h b/Engine/Source/Runtime/VulkanRHI/Private/VulkanPipeline.h
--- a/Engine/Source/Runtime/VulkanRHI/Private/VulkanPipeline.h
+++ b/Engine/Source/Runtime/VulkanRHI/Private/VulkanPipeline.h
@@ -364,6 +364,7 @@ struct FGfxPipelineDesc
 		uint8 NumColorAttachments;
 		uint8 bHasDepthStencil;
 		uint8 bHasResolveAttachments;
+		uint8 bHasDepthStencilResolve;
 		uint8 bHasFragmentDensityAttachment;
 		uint8 NumUsedClearValues;
 		uint32 RenderPassCompatibleHash;
@@ -385,6 +386,7 @@ struct FGfxPipelineDesc
 				NumColorAttachments == In.NumColorAttachments &&
 				bHasDepthStencil == In.bHasDepthStencil &&
 				bHasResolveAttachments == In.bHasResolveAttachments &&
+				bHasDepthStencilResolve == In.bHasDepthStencilResolve &&
 				bHasFragmentDensityAttachment == In.bHasFragmentDensityAttachment &&
 				NumUsedClearValues == In.NumUsedClearValues &&
 				RenderPassCompatibleHash == In.RenderPassCompatibleHash &&
diff --git a/Engine/Source/Runtime/VulkanRHI/Private/VulkanRHIPrivate.h b/Engine/Source/Runtime/VulkanRHI/Private/VulkanRHIPrivate.h
--- a/Engine/Source/Runtime/VulkanRHI/Private/VulkanRHIPrivate.h
+++ b/Engine/Source/Runtime/VulkanRHI/Private/VulkanRHIPrivate.h
@@ -134,6 +134,7 @@ public:
 	inline uint32 GetNumColorAttachments() const { return NumColorAttachments; }
 	inline bool GetHasDepthStencil() const { return bHasDepthStencil != 0; }
 	inline bool GetHasResolveAttachments() const { return bHasResolveAttachments != 0; }
+	inline bool GetHasDepthStencilResolve() const { return bHasDepthStencilResolve != 0; }
 	inline bool GetHasFragmentDensityAttachment() const { return bHasFragmentDensityAttachment != 0; }
 	inline uint32 GetNumAttachmentDescriptions() const { return NumAttachmentDescriptions; }
 	inline uint32 GetNumSamples() const { return NumSamples; }
@@ -146,6 +147,7 @@ public:
 	inline const VkAttachmentReference* GetResolveAttachmentReferences() const { return bHasResolveAttachments ? ResolveReferences : nullptr; }
 	inline const VkAttachmentReference* GetDepthAttachmentReference() const { return bHasDepthStencil ? &DepthReference : nullptr; }
 	inline const VkAttachmentReferenceStencilLayout* GetStencilAttachmentReference() const { return bHasDepthStencil ? &StencilReference : nullptr; }
+	inline const VkAttachmentReference* GetDepthStencilResolveAttachmentReference() const { return bHasDepthStencilResolve ? &DepthStencilResolveReference : nullptr; }
 	inline const VkAttachmentReference* GetFragmentDensityAttachmentReference() const { return bHasFragmentDensityAttachment ? &FragmentDensityReference : nullptr; }
 
 	inline const VkAttachmentDescriptionStencilLayout* GetStencilDesc() const { return bHasDepthStencil ? &StencilDesc : nullptr; }
@@ -161,9 +163,10 @@ protected:
 	VkAttachmentReferenceStencilLayout StencilReference;
 	VkAttachmentReference FragmentDensityReference;
 	VkAttachmentReference ResolveReferences[MaxSimultaneousRenderTargets];
+	VkAttachmentReference DepthStencilResolveReference;
 
-	// Depth goes in the "+1" slot and the Shading Rate texture goes in the "+2" slot.
-	VkAttachmentDescription Desc[MaxSimultaneousRenderTargets * 2 + 2];
+	// Depth goes in the "+1" slot, Depth resolve goes in the "+2 slot", and the Shading Rate texture goes in the "+3" slot.
+	VkAttachmentDescription Desc[MaxSimultaneousRenderTargets * 2 + 3];
 	VkAttachmentDescriptionStencilLayout StencilDesc;
 
 	uint8 NumAttachmentDescriptions;
@@ -171,6 +174,7 @@ protected:
 	uint8 NumInputAttachments = 0;
 	uint8 bHasDepthStencil;
 	uint8 bHasResolveAttachments;
+	uint8 bHasDepthStencilResolve;
 	uint8 bHasFragmentDensityAttachment;
 	uint8 NumSamples;
 	uint8 NumUsedClearValues;
@@ -200,6 +204,7 @@ protected:
 		FMemory::Memzero(DepthReference);
 		FMemory::Memzero(FragmentDensityReference);
 		FMemory::Memzero(ResolveReferences);
+		FMemory::Memzero(DepthStencilResolveReference);
 		FMemory::Memzero(Desc);
 		FMemory::Memzero(Offset);
 		FMemory::Memzero(Extent);
@@ -214,6 +219,7 @@ protected:
 		NumColorAttachments = 0;
 		bHasDepthStencil = 0;
 		bHasResolveAttachments = 0;
+		bHasDepthStencilResolve = 0;
 		bHasFragmentDensityAttachment = 0;
 		NumSamples = 0;
 		NumUsedClearValues = 0;
@@ -299,6 +305,7 @@ private:
 	VkImage ColorRenderTargetImages[MaxSimultaneousRenderTargets];
 	VkImage ColorResolveTargetImages[MaxSimultaneousRenderTargets];
 	VkImage DepthStencilRenderTargetImage;
+	VkImage DepthStencilResolveRenderTargetImage;
 	VkImage FragmentDensityImage;
 
 	// Predefined set of barriers, when executes ensuring all writes are finished
diff --git a/Engine/Source/Runtime/VulkanRHI/Private/VulkanRenderTarget.cpp b/Engine/Source/Runtime/VulkanRHI/Private/VulkanRenderTarget.cpp
--- a/Engine/Source/Runtime/VulkanRHI/Private/VulkanRenderTarget.cpp
+++ b/Engine/Source/Runtime/VulkanRHI/Private/VulkanRenderTarget.cpp
@@ -768,6 +768,7 @@ FVulkanRenderTargetLayout::FVulkanRenderTargetLayout(FVulkanDevice& InDevice, co
 	, NumColorAttachments(0)
 	, bHasDepthStencil(false)
 	, bHasResolveAttachments(false)
+	, bHasDepthStencilResolve(false)
 	, bHasFragmentDensityAttachment(false)
 	, NumSamples(0)
 	, NumUsedClearValues(0)
@@ -869,24 +870,15 @@ FVulkanRenderTargetLayout::FVulkanRenderTargetLayout(FVulkanDevice& InDevice, co
 		CurrDesc.loadOp = RenderTargetLoadActionToVulkan(RTInfo.DepthStencilRenderTarget.DepthLoadAction);
 		CurrDesc.stencilLoadOp = RenderTargetLoadActionToVulkan(RTInfo.DepthStencilRenderTarget.StencilLoadAction);
 		bFoundClearOp = bFoundClearOp || (CurrDesc.loadOp == VK_ATTACHMENT_LOAD_OP_CLEAR || CurrDesc.stencilLoadOp == VK_ATTACHMENT_LOAD_OP_CLEAR);
-		if (CurrDesc.samples == VK_SAMPLE_COUNT_1_BIT)
-		{
-			CurrDesc.storeOp = RenderTargetStoreActionToVulkan(RTInfo.DepthStencilRenderTarget.DepthStoreAction);
-			CurrDesc.stencilStoreOp = RenderTargetStoreActionToVulkan(RTInfo.DepthStencilRenderTarget.GetStencilStoreAction());
+		CurrDesc.storeOp = RenderTargetStoreActionToVulkan(RTInfo.DepthStencilRenderTarget.DepthStoreAction);
+		CurrDesc.stencilStoreOp = RenderTargetStoreActionToVulkan(RTInfo.DepthStencilRenderTarget.GetStencilStoreAction());
 
-			// Removed this temporarily as we need a way to determine if the target is actually memoryless
-			/*if (EnumHasAllFlags(Texture->UEFlags, TexCreate_Memoryless))
-			{
-				ensure(CurrDesc.storeOp == VK_ATTACHMENT_STORE_OP_DONT_CARE);
-				ensure(CurrDesc.stencilStoreOp == VK_ATTACHMENT_STORE_OP_DONT_CARE);
-			}*/
-		}
-		else
+		// Removed this temporarily as we need a way to determine if the target is actually memoryless
+		/*if (EnumHasAllFlags(Texture->UEFlags, TexCreate_Memoryless))
 		{
-			// Never want to store MSAA depth/stencil
-			CurrDesc.storeOp = VK_ATTACHMENT_STORE_OP_DONT_CARE;
-			CurrDesc.stencilStoreOp = VK_ATTACHMENT_STORE_OP_DONT_CARE;
-		}
+			ensure(CurrDesc.storeOp == VK_ATTACHMENT_STORE_OP_DONT_CARE);
+			ensure(CurrDesc.stencilStoreOp == VK_ATTACHMENT_STORE_OP_DONT_CARE);
+		}*/
 
 		const VkImageLayout DepthLayout = RTInfo.DepthStencilRenderTarget.GetDepthStencilAccess().IsDepthWrite() ? VK_IMAGE_LAYOUT_ATTACHMENT_OPTIMAL : VK_IMAGE_LAYOUT_READ_ONLY_OPTIMAL;
 		const VkImageLayout StencilLayout = RTInfo.DepthStencilRenderTarget.GetDepthStencilAccess().IsStencilWrite() ? VK_IMAGE_LAYOUT_ATTACHMENT_OPTIMAL : VK_IMAGE_LAYOUT_READ_ONLY_OPTIMAL;
@@ -901,6 +893,25 @@ FVulkanRenderTargetLayout::FVulkanRenderTargetLayout(FVulkanDevice& InDevice, co
 		DepthReference.layout = DepthLayout;
 		StencilReference.stencilLayout = StencilLayout;
 
+		// Use depth/stencil resolve target only if we're MSAA and a RT has a resolve attachment
+		const bool bHasValidResolveAttachment = RTInfo.bHasResolveAttachments && RTInfo.DepthStencilResolveRenderTarget.Texture;
+		if (GRHISupportsDepthStencilResolve && CurrDesc.samples > VK_SAMPLE_COUNT_1_BIT && bHasValidResolveAttachment)
+		{
+			Desc[NumAttachmentDescriptions + 1] = Desc[NumAttachmentDescriptions];
+			Desc[NumAttachmentDescriptions + 1].samples = VK_SAMPLE_COUNT_1_BIT;
+			Desc[NumAttachmentDescriptions + 1].loadOp = VK_ATTACHMENT_LOAD_OP_DONT_CARE;
+			Desc[NumAttachmentDescriptions + 1].storeOp = VK_ATTACHMENT_STORE_OP_STORE;
+			Desc[NumAttachmentDescriptions + 1].stencilLoadOp = VK_ATTACHMENT_LOAD_OP_DONT_CARE;
+			Desc[NumAttachmentDescriptions + 1].stencilStoreOp = VK_ATTACHMENT_STORE_OP_STORE;
+			DepthStencilResolveReference.attachment = NumAttachmentDescriptions + 1;
+			DepthStencilResolveReference.layout = DepthLayout;
+			// NumColorAttachments was incremented after the last color attachment
+			ensureMsgf(NumColorAttachments < 16, TEXT("Must have room for depth resolve bit"));
+			CompatibleHashInfo.AttachmentsToResolve |= (uint16)(1 << NumColorAttachments);
+			++NumAttachmentDescriptions;
+			bHasDepthStencilResolve = true;
+		}
+
 		FullHashInfo.LoadOps[MaxSimultaneousRenderTargets] = CurrDesc.loadOp;
 		FullHashInfo.LoadOps[MaxSimultaneousRenderTargets + 1] = CurrDesc.stencilLoadOp;
 		FullHashInfo.StoreOps[MaxSimultaneousRenderTargets] = CurrDesc.storeOp;
@@ -978,6 +989,7 @@ FVulkanRenderTargetLayout::FVulkanRenderTargetLayout(FVulkanDevice& InDevice, co
 	, NumColorAttachments(0)
 	, bHasDepthStencil(false)
 	, bHasResolveAttachments(false)
+	, bHasDepthStencilResolve(false)
 	, bHasFragmentDensityAttachment(false)
 	, NumSamples(0)
 	, NumUsedClearValues(0)
@@ -1083,12 +1095,6 @@ FVulkanRenderTargetLayout::FVulkanRenderTargetLayout(FVulkanDevice& InDevice, co
 		CurrDesc.loadOp = RenderTargetLoadActionToVulkan(GetLoadAction(GetDepthActions(RPInfo.DepthStencilRenderTarget.Action)));
 		CurrDesc.stencilLoadOp = RenderTargetLoadActionToVulkan(GetLoadAction(GetStencilActions(RPInfo.DepthStencilRenderTarget.Action)));
 		bFoundClearOp = bFoundClearOp || (CurrDesc.loadOp == VK_ATTACHMENT_LOAD_OP_CLEAR || CurrDesc.stencilLoadOp == VK_ATTACHMENT_LOAD_OP_CLEAR);
-		if (CurrDesc.samples != VK_SAMPLE_COUNT_1_BIT)
-		{
-			// Can't resolve MSAA depth/stencil
-			ensure(GetStoreAction(GetDepthActions(RPInfo.DepthStencilRenderTarget.Action)) != ERenderTargetStoreAction::EMultisampleResolve);
-			ensure(GetStoreAction(GetStencilActions(RPInfo.DepthStencilRenderTarget.Action)) != ERenderTargetStoreAction::EMultisampleResolve);
-		}
 
 		CurrDesc.storeOp = RenderTargetStoreActionToVulkan(GetStoreAction(GetDepthActions(RPInfo.DepthStencilRenderTarget.Action)));
 		CurrDesc.stencilStoreOp = RenderTargetStoreActionToVulkan(GetStoreAction(GetStencilActions(RPInfo.DepthStencilRenderTarget.Action)));
@@ -1132,7 +1138,22 @@ FVulkanRenderTargetLayout::FVulkanRenderTargetLayout(FVulkanDevice& InDevice, co
 		DepthReference.layout = CurrentDepthLayout;
 		StencilReference.stencilLayout = CurrentStencilLayout;
 
-		++NumAttachmentDescriptions;
+		if (GRHISupportsDepthStencilResolve && CurrDesc.samples > VK_SAMPLE_COUNT_1_BIT && RPInfo.DepthStencilRenderTarget.ResolveTarget)
+		{
+			Desc[NumAttachmentDescriptions + 1] = Desc[NumAttachmentDescriptions];
+			Desc[NumAttachmentDescriptions + 1].samples = VK_SAMPLE_COUNT_1_BIT;
+			Desc[NumAttachmentDescriptions + 1].loadOp = VK_ATTACHMENT_LOAD_OP_DONT_CARE;
+			Desc[NumAttachmentDescriptions + 1].storeOp = VK_ATTACHMENT_STORE_OP_STORE;
+			Desc[NumAttachmentDescriptions + 1].stencilLoadOp = VK_ATTACHMENT_LOAD_OP_DONT_CARE;
+			Desc[NumAttachmentDescriptions + 1].stencilStoreOp = VK_ATTACHMENT_STORE_OP_STORE;
+			DepthStencilResolveReference.attachment = NumAttachmentDescriptions + 1;
+			DepthStencilResolveReference.layout = CurrentDepthLayout;
+			// NumColorAttachments was incremented after the last color attachment
+			ensureMsgf(NumColorAttachments < 16, TEXT("Must have room for depth resolve bit"));
+			CompatibleHashInfo.AttachmentsToResolve |= (uint16)(1 << NumColorAttachments);
+			++NumAttachmentDescriptions;
+			bHasDepthStencilResolve = true;
+		}
 
 		FullHashInfo.LoadOps[MaxSimultaneousRenderTargets] = CurrDesc.loadOp;
 		FullHashInfo.LoadOps[MaxSimultaneousRenderTargets + 1] = CurrDesc.stencilLoadOp;
@@ -1142,6 +1163,7 @@ FVulkanRenderTargetLayout::FVulkanRenderTargetLayout(FVulkanDevice& InDevice, co
 		FullHashInfo.InitialLayout[MaxSimultaneousRenderTargets + 1] = CurrentStencilLayout;
 		CompatibleHashInfo.Formats[MaxSimultaneousRenderTargets] = CurrDesc.format;
 
+		++NumAttachmentDescriptions;
 
 		bHasDepthStencil = true;
 
@@ -1226,6 +1248,7 @@ FVulkanRenderTargetLayout::FVulkanRenderTargetLayout(const FGraphicsPipelineStat
 	, NumColorAttachments(0)
 	, bHasDepthStencil(false)
 	, bHasResolveAttachments(false)
+	, bHasDepthStencilResolve(false)
 	, bHasFragmentDensityAttachment(false)
 	, NumSamples(0)
 	, NumUsedClearValues(0)
@@ -1299,18 +1322,8 @@ FVulkanRenderTargetLayout::FVulkanRenderTargetLayout(const FGraphicsPipelineStat
 		{
 			bFoundClearOp = true;
 		}
-		if (CurrDesc.samples == VK_SAMPLE_COUNT_1_BIT)
-		{
-			CurrDesc.storeOp = RenderTargetStoreActionToVulkan(Initializer.DepthTargetStoreAction);
-			CurrDesc.stencilStoreOp = RenderTargetStoreActionToVulkan(Initializer.StencilTargetStoreAction);
-		}
-		else
-		{
-			// Never want to store MSAA depth/stencil
-			CurrDesc.storeOp = VK_ATTACHMENT_STORE_OP_DONT_CARE;
-			CurrDesc.stencilStoreOp = VK_ATTACHMENT_STORE_OP_DONT_CARE;
-		}
-
+		CurrDesc.storeOp = RenderTargetStoreActionToVulkan(Initializer.DepthTargetStoreAction);
+		CurrDesc.stencilStoreOp = RenderTargetStoreActionToVulkan(Initializer.StencilTargetStoreAction);
 		// If the initial != final we need to change the FullHashInfo and use FinalLayout
 		CurrDesc.initialLayout = VK_IMAGE_LAYOUT_ATTACHMENT_OPTIMAL;
 		CurrDesc.finalLayout = VK_IMAGE_LAYOUT_ATTACHMENT_OPTIMAL;
@@ -1321,6 +1334,24 @@ FVulkanRenderTargetLayout::FVulkanRenderTargetLayout(const FGraphicsPipelineStat
 		DepthReference.layout = Initializer.DepthStencilAccess.IsDepthWrite() ? VK_IMAGE_LAYOUT_ATTACHMENT_OPTIMAL : VK_IMAGE_LAYOUT_READ_ONLY_OPTIMAL;
 		StencilReference.stencilLayout = Initializer.DepthStencilAccess.IsStencilWrite() ? VK_IMAGE_LAYOUT_ATTACHMENT_OPTIMAL : VK_IMAGE_LAYOUT_READ_ONLY_OPTIMAL;
 
+		const bool bDepthStencilResolve = (Initializer.DepthTargetStoreAction == ERenderTargetStoreAction::EMultisampleResolve) || (Initializer.StencilTargetStoreAction == ERenderTargetStoreAction::EMultisampleResolve);
+		if (bDepthStencilResolve && GRHISupportsDepthStencilResolve && CurrDesc.samples > VK_SAMPLE_COUNT_1_BIT)
+		{
+			Desc[NumAttachmentDescriptions + 1] = Desc[NumAttachmentDescriptions];
+			Desc[NumAttachmentDescriptions + 1].samples = VK_SAMPLE_COUNT_1_BIT;
+			Desc[NumAttachmentDescriptions + 1].loadOp = VK_ATTACHMENT_LOAD_OP_DONT_CARE;
+			Desc[NumAttachmentDescriptions + 1].storeOp = VK_ATTACHMENT_STORE_OP_STORE;
+			Desc[NumAttachmentDescriptions + 1].stencilLoadOp = VK_ATTACHMENT_LOAD_OP_DONT_CARE;
+			Desc[NumAttachmentDescriptions + 1].stencilStoreOp = VK_ATTACHMENT_STORE_OP_STORE;
+			DepthStencilResolveReference.attachment = NumAttachmentDescriptions + 1;
+			DepthStencilResolveReference.layout = VK_IMAGE_LAYOUT_DEPTH_STENCIL_ATTACHMENT_OPTIMAL;
+			// NumColorAttachments was incremented after the last color attachment
+			ensureMsgf(NumColorAttachments < 16, TEXT("Must have room for depth resolve bit"));
+			CompatibleHashInfo.AttachmentsToResolve |= (uint16)(1 << NumColorAttachments);
+			++NumAttachmentDescriptions;
+			bHasDepthStencilResolve = true;
+		}
+
 		FullHashInfo.LoadOps[MaxSimultaneousRenderTargets] = CurrDesc.loadOp;
 		FullHashInfo.LoadOps[MaxSimultaneousRenderTargets + 1] = CurrDesc.stencilLoadOp;
 		FullHashInfo.StoreOps[MaxSimultaneousRenderTargets] = CurrDesc.storeOp;
diff --git a/Engine/Source/Runtime/VulkanRHI/Private/VulkanRenderpass.h b/Engine/Source/Runtime/VulkanRHI/Private/VulkanRenderpass.h
--- a/Engine/Source/Runtime/VulkanRHI/Private/VulkanRenderpass.h
+++ b/Engine/Source/Runtime/VulkanRHI/Private/VulkanRenderpass.h
@@ -147,7 +147,12 @@ struct FVulkanSubpassDescription<VkSubpassDescription>
 		inputAttachmentCount = NumInputAttachmentReferences;
 	}
 
-	void SetShadingRateAttachment(void* /* ShadingRateAttachmentInfo */)
+	void SetDepthStencilResolveAttachment(VkSubpassDescriptionDepthStencilResolveKHR* DepthStencilResolveAttachmentDesc)
+	{
+		// No-op without VK_KHR_create_renderpass2
+	}
+
+	void SetShadingRateAttachment(VkFragmentShadingRateAttachmentInfoKHR* /* ShadingRateAttachmentInfo */)
 	{
 		// No-op without VK_KHR_create_renderpass2
 	}
@@ -195,9 +200,18 @@ struct FVulkanSubpassDescription<VkSubpassDescription2>
 		inputAttachmentCount = NumInputAttachmentReferences;
 	}
 
-	void SetShadingRateAttachment(void* ShadingRateAttachmentInfo)
+	void SetDepthStencilResolveAttachment(VkSubpassDescriptionDepthStencilResolveKHR* DepthStencilResolveAttachmentDesc)
 	{
+		const void* Next = pNext;
+		pNext = DepthStencilResolveAttachmentDesc;
+		DepthStencilResolveAttachmentDesc->pNext = Next;
+	}
+
+	void SetShadingRateAttachment(VkFragmentShadingRateAttachmentInfoKHR* ShadingRateAttachmentInfo)
+	{
+		const void* Next = pNext;
 		pNext = ShadingRateAttachmentInfo;
+		ShadingRateAttachmentInfo->pNext = Next;
 	}
 
 	void SetMultiViewMask(uint32_t Mask)
@@ -404,6 +418,26 @@ struct FVulkanFragmentShadingRateAttachmentInfo
 	}
 };
 
+struct FVulkanDepthStencilResolveSubpassDesc
+	: public VkSubpassDescriptionDepthStencilResolveKHR
+{
+	FVulkanDepthStencilResolveSubpassDesc()
+	{
+		ZeroVulkanStruct(*this, VK_STRUCTURE_TYPE_SUBPASS_DESCRIPTION_DEPTH_STENCIL_RESOLVE);
+	}
+
+	void SetResolveModes(VkResolveModeFlagBits DepthMode, VkResolveModeFlagBits StencilMode)
+	{
+		depthResolveMode = DepthMode;
+		stencilResolveMode = StencilMode;
+	}
+
+	void SetReference(FVulkanAttachmentReference<VkAttachmentReference2>* AttachmentReference)
+	{
+		pDepthStencilResolveAttachment = AttachmentReference;
+	}
+};
+
 extern int32 GVulkanInputAttachmentShaderRead;
 
 template <typename TSubpassDescriptionClass, typename TSubpassDependencyClass, typename TAttachmentReferenceClass, typename TAttachmentDescriptionClass, typename TRenderPassCreateInfoClass>
@@ -430,6 +464,10 @@ public:
 												&& RTLayout.GetFragmentDensityAttachmentReference() != nullptr
 												&& Device.GetOptionalExtensions().HasKHRFragmentShadingRate 
 												&& Device.GetOptionalExtensionProperties().FragmentShadingRateFeatures.attachmentFragmentShadingRate == VK_TRUE;
+		const bool bResolveDepth =  GRHISupportsDepthStencilResolve &&
+									Device.GetOptionalExtensions().HasKHRDepthStencilResolve &&
+									RTLayout.GetHasDepthStencilResolve();
+
 		const bool bCustomResolveSubpass = RTLayout.GetSubpassHint() == ESubpassHint::CustomResolveSubpass;
 		const bool bDepthReadSubpass = bCustomResolveSubpass || (RTLayout.GetSubpassHint() == ESubpassHint::DepthReadSubpass);
 		const bool bHasDepthStencilAttachmentReference = (RTLayout.GetDepthAttachmentReference() != nullptr);
@@ -440,6 +478,15 @@ public:
 			FragmentShadingRateAttachmentInfo.SetReference(&ShadingRateAttachmentReference);
 		}
 
+		if (bResolveDepth)
+		{
+			DepthStencilResolveAttachmentReference.SetAttachment(*RTLayout.GetDepthStencilResolveAttachmentReference(), VkImageAspectFlagBits::VK_IMAGE_ASPECT_NONE);
+			// Using zero bit because it is always supported if the extension is supported, from spec: "The VK_RESOLVE_MODE_SAMPLE_ZERO_BIT mode
+			// is the only mode that is required of all implementations (that support the extension or support Vulkan 1.2 or higher)."
+			DepthStencilResolveSubpassDesc.SetResolveModes(VK_RESOLVE_MODE_SAMPLE_ZERO_BIT, VK_RESOLVE_MODE_SAMPLE_ZERO_BIT);
+			DepthStencilResolveSubpassDesc.SetReference(&DepthStencilResolveAttachmentReference);
+		}
+
 		// Grab (and optionally convert) attachment references.
 		uint32 NumColorAttachments = RTLayout.GetNumColorAttachments();
 		for (uint32 ColorAttachment = 0; ColorAttachment < NumColorAttachments; ++ColorAttachment)
@@ -498,6 +545,11 @@ public:
 				SubpassDesc.SetDepthStencilAttachment(&DepthStencilAttachmentReference);
 			}
 
+			if (!bDepthReadSubpass && bResolveDepth)
+			{
+				SubpassDesc.SetDepthStencilResolveAttachment(&DepthStencilResolveSubpassDesc);
+			}
+
 			if (bApplyFragmentShadingRate)
 			{
 				SubpassDesc.SetShadingRateAttachment(&FragmentShadingRateAttachmentInfo);
@@ -522,6 +574,11 @@ public:
 			// depth attachment is same as input attachment
 			SubpassDesc.SetDepthStencilAttachment(&DepthStencilAttachment);
 
+			if (bResolveDepth)
+			{
+				SubpassDesc.SetDepthStencilResolveAttachment(&DepthStencilResolveSubpassDesc);
+			}
+
 			if (bApplyFragmentShadingRate)
 			{
 				SubpassDesc.SetShadingRateAttachment(&FragmentShadingRateAttachmentInfo);
@@ -773,6 +830,10 @@ private:
 	FVulkanAttachmentReference<VkAttachmentReference2> ShadingRateAttachmentReference;
 	FVulkanFragmentShadingRateAttachmentInfo FragmentShadingRateAttachmentInfo;
 
+	// Depth stencil resolve
+	FVulkanAttachmentReference<VkAttachmentReference2> DepthStencilResolveAttachmentReference;
+	FVulkanDepthStencilResolveSubpassDesc DepthStencilResolveSubpassDesc;
+
 	FVulkanRenderPassFragmentDensityMapCreateInfoEXT FragDensityCreateInfo;
 	FVulkanRenderPassMultiviewCreateInfo MultiviewInfo;
 
diff --git a/Engine/Source/Runtime/VulkanRHI/Private/VulkanViewport.cpp b/Engine/Source/Runtime/VulkanRHI/Private/VulkanViewport.cpp
--- a/Engine/Source/Runtime/VulkanRHI/Private/VulkanViewport.cpp
+++ b/Engine/Source/Runtime/VulkanRHI/Private/VulkanViewport.cpp
@@ -508,6 +508,30 @@ FVulkanFramebuffer::FVulkanFramebuffer(FVulkanDevice& Device, const FRHISetRende
 		{
 			AddExternalView(Texture->DefaultView);
 		}
+
+		if (InRTInfo.bHasResolveAttachments && RTLayout.GetHasDepthStencilResolve() && RTLayout.GetDepthStencilResolveAttachmentReference()->layout != VK_IMAGE_LAYOUT_UNDEFINED)
+		{
+			FRHITexture* ResolveRHITexture = InRTInfo.DepthStencilResolveRenderTarget.Texture;
+			FVulkanTexture* ResolveTexture = ResourceCast(ResolveRHITexture);
+			DepthStencilResolveRenderTargetImage = ResolveTexture->Image;
+
+			// Resolve attachments only supported for 2d/2d array textures
+			if (ResolveTexture->GetViewType() == VK_IMAGE_VIEW_TYPE_2D || ResolveTexture->GetViewType() == VK_IMAGE_VIEW_TYPE_2D_ARRAY)
+			{
+				CreateOwnedView()->InitAsTextureView(
+					ResolveTexture->Image
+					, ResolveTexture->GetViewType()
+					, ResolveTexture->GetFullAspectMask()
+					, ResolveTexture->GetDesc().Format
+					, ResolveTexture->ViewFormat
+					, MipIndex
+					, 1
+					, 0
+					, ResolveTexture->GetNumberOfArrayLevels()
+					, true
+				);
+			}
+		}
 	}
 
 	if (GRHISupportsAttachmentVariableRateShading && GRHIVariableRateShadingEnabled && GRHIAttachmentVariableRateShadingEnabled && RTLayout.GetHasFragmentDensityAttachment())
@@ -593,6 +617,19 @@ bool FVulkanFramebuffer::Matches(const FRHISetRenderTargetsInfo& InRTInfo) const
 		}
 	}
 
+	{
+		const FRHIDepthRenderTargetView& R = InRTInfo.DepthStencilResolveRenderTarget;
+		if (R.Texture)
+		{
+			VkImage AImage = DepthStencilResolveRenderTargetImage;
+			VkImage BImage = ResourceCast(R.Texture)->Image;
+			if (AImage != BImage)
+			{
+				return false;
+			}
+		}
+	}
+
 	{
 		FRHITexture* Texture = InRTInfo.ShadingRateTexture;
 		if (Texture)
