Make the compute shader pointer part of FRHIComputePipelineState.

All the RHI implementations had a shader member in their compute PSO derived type, so it's better if this member lives in the base class, to be accessible to other systems (such as RHI validation).

#rnx
#tests Ran the editor in D3D11, D3D12 and Vulkan and verified that things still render correctly.

diff --git a/Engine/Source/Runtime/Apple/MetalRHI/Private/MetalCommands.cpp b/Engine/Source/Runtime/Apple/MetalRHI/Private/MetalCommands.cpp
--- a/Engine/Source/Runtime/Apple/MetalRHI/Private/MetalCommands.cpp
+++ b/Engine/Source/Runtime/Apple/MetalRHI/Private/MetalCommands.cpp
@@ -7,7 +7,6 @@
 #include "MetalRHIPrivate.h"
 #include "MetalShaderTypes.h"
 #include "MetalGraphicsPipelineState.h"
-#include "MetalComputePipelineState.h"
 #include "MetalRHIContext.h"
 #include "GlobalShader.h"
 #include "OneColorShader.h"
@@ -183,13 +182,13 @@ void FMetalRHICommandContext::RHISetComputePipelineState(FRHIComputePipelineStat
 {
     MTL_SCOPED_AUTORELEASE_POOL;
 	
-    FMetalComputePipelineState* ComputePipeline = ResourceCast(ComputePipelineState);
+	FMetalComputeShader* ComputeShader = ResourceCast(ComputePipelineState->GetComputeShader());
 	
 	// cache this for Dispatch
 	// sets this compute shader pipeline as the current (this resets all state, so we need to set all resources after calling this)
-	Context->GetCurrentState().SetComputeShader(ComputePipeline->GetComputeShader());
+	Context->GetCurrentState().SetComputeShader(ComputeShader);
 
-    ApplyStaticUniformBuffersOnContext(*this, ComputePipeline->ComputeShader);
+    ApplyStaticUniformBuffersOnContext(*this, ComputeShader, static_cast<FMetalShaderData*>(ComputeShader));
 }
 
 void FMetalRHICommandContext::RHIDispatchComputeShader(uint32 ThreadGroupCountX, uint32 ThreadGroupCountY, uint32 ThreadGroupCountZ)
diff --git a/Engine/Source/Runtime/Apple/MetalRHI/Private/MetalComputePipelineState.h b/Engine/Source/Runtime/Apple/MetalRHI/Private/MetalComputePipelineState.h
deleted file mode 100644
diff --git a/Engine/Source/Runtime/Apple/MetalRHI/Private/MetalDynamicRHI.cpp b/Engine/Source/Runtime/Apple/MetalRHI/Private/MetalDynamicRHI.cpp
--- a/Engine/Source/Runtime/Apple/MetalRHI/Private/MetalDynamicRHI.cpp
+++ b/Engine/Source/Runtime/Apple/MetalRHI/Private/MetalDynamicRHI.cpp
@@ -11,7 +11,6 @@
 #include "MetalShaderTypes.h"
 #include "MetalVertexDeclaration.h"
 #include "MetalGraphicsPipelineState.h"
-#include "MetalComputePipelineState.h"
 #include "MetalTransitionData.h"
 
 //------------------------------------------------------------------------------
@@ -108,7 +107,7 @@ FGraphicsPipelineStateRHIRef FMetalDynamicRHI::RHICreateGraphicsPipelineState(co
 TRefCountPtr<FRHIComputePipelineState> FMetalDynamicRHI::RHICreateComputePipelineState(FRHIComputeShader* ComputeShader)
 {
     MTL_SCOPED_AUTORELEASE_POOL;
-    return new FMetalComputePipelineState(ResourceCast(ComputeShader));
+    return new FRHIComputePipelineState(ComputeShader);
 }
 
 
diff --git a/Engine/Source/Runtime/Apple/MetalRHI/Private/MetalPipeline.cpp b/Engine/Source/Runtime/Apple/MetalRHI/Private/MetalPipeline.cpp
--- a/Engine/Source/Runtime/Apple/MetalRHI/Private/MetalPipeline.cpp
+++ b/Engine/Source/Runtime/Apple/MetalRHI/Private/MetalPipeline.cpp
@@ -9,7 +9,6 @@
 #include "MetalVertexDeclaration.h"
 #include "MetalShaderTypes.h"
 #include "MetalGraphicsPipelineState.h"
-#include "MetalComputePipelineState.h"
 #include "MetalPipeline.h"
 #include "MetalShaderResources.h"
 #include "MetalProfiler.h"
diff --git a/Engine/Source/Runtime/Apple/MetalRHI/Private/MetalRayTracing.cpp b/Engine/Source/Runtime/Apple/MetalRHI/Private/MetalRayTracing.cpp
--- a/Engine/Source/Runtime/Apple/MetalRHI/Private/MetalRayTracing.cpp
+++ b/Engine/Source/Runtime/Apple/MetalRHI/Private/MetalRayTracing.cpp
@@ -10,7 +10,6 @@
 
 #include "MetalRHIContext.h"
 #include "MetalShaderTypes.h"
-#include "MetalComputePipelineState.h"
 #include "BuiltInRayTracingShaders.h"
 #include "RayTracingValidationShaders.h"
 #include "RayTracingBuiltInResources.h"
diff --git a/Engine/Source/Runtime/Apple/MetalRHI/Public/MetalResources.h b/Engine/Source/Runtime/Apple/MetalRHI/Public/MetalResources.h
--- a/Engine/Source/Runtime/Apple/MetalRHI/Public/MetalResources.h
+++ b/Engine/Source/Runtime/Apple/MetalRHI/Public/MetalResources.h
@@ -585,7 +585,6 @@ private:
 
 class FMetalShaderLibrary;
 class FMetalGraphicsPipelineState;
-class FMetalComputePipelineState;
 class FMetalVertexDeclaration;
 class FMetalVertexShader;
 class FMetalGeometryShader;
@@ -678,11 +677,6 @@ struct TMetalResourceTraits<FRHIGraphicsPipelineState>
 	typedef FMetalGraphicsPipelineState TConcreteType;
 };
 template<>
-struct TMetalResourceTraits<FRHIComputePipelineState>
-{
-	typedef FMetalComputePipelineState TConcreteType;
-};
-template<>
 struct TMetalResourceTraits<FRHIGPUFence>
 {
 	typedef FMetalGPUFence TConcreteType;
diff --git a/Engine/Source/Runtime/D3D12RHI/Private/D3D12Commands.cpp b/Engine/Source/Runtime/D3D12RHI/Private/D3D12Commands.cpp
--- a/Engine/Source/Runtime/D3D12RHI/Private/D3D12Commands.cpp
+++ b/Engine/Source/Runtime/D3D12RHI/Private/D3D12Commands.cpp
@@ -819,7 +819,7 @@ void FD3D12CommandContext::RHISetComputePipelineState(FRHIComputePipelineState*
 
 	StateCache.SetComputePipelineState(ComputePipelineState);
 
-	ApplyStaticUniformBuffersOnContext(*this, ComputePipelineState->ComputeShader.GetReference());
+	ApplyStaticUniformBuffersOnContext(*this, FD3D12DynamicRHI::ResourceCast(ComputePipelineState->GetComputeShader()));
 }
 
 void FD3D12CommandContext::SetUAVParameter(EShaderFrequency Frequency, uint32 UAVIndex, FD3D12UnorderedAccessView* UAV)
@@ -1485,7 +1485,7 @@ void FD3D12CommandContext::CommitComputeResourceTables()
 
 	const FD3D12ComputePipelineState* const RESTRICT ComputePSO = StateCache.GetComputePipelineState();
 
-	SetResourcesFromTables(ComputePSO->GetComputeShader());
+	SetResourcesFromTables(FD3D12DynamicRHI::ResourceCast(ComputePSO->GetComputeShader()));
 }
 
 void FD3D12CommandContext::RHISetShaderRootConstants(const FUint32Vector4& Constants)
diff --git a/Engine/Source/Runtime/D3D12RHI/Private/D3D12DescriptorCache.cpp b/Engine/Source/Runtime/D3D12RHI/Private/D3D12DescriptorCache.cpp
--- a/Engine/Source/Runtime/D3D12RHI/Private/D3D12DescriptorCache.cpp
+++ b/Engine/Source/Runtime/D3D12RHI/Private/D3D12DescriptorCache.cpp
@@ -663,7 +663,11 @@ void FD3D12DescriptorCache::SetRootConstantBuffers(EShaderFrequency ShaderStage,
 		if (FD3D12ConstantBufferCache::IsSlotDirty(RDCBVSlotsNeededMask, SlotIndex))
 		{
 			const D3D12_GPU_VIRTUAL_ADDRESS CurrentGPUVirtualAddress = Cache.CurrentGPUVirtualAddress[ShaderStage][SlotIndex];
-			check(CurrentGPUVirtualAddress != 0);
+			if (CurrentGPUVirtualAddress == 0)
+			{
+				UE_LOG(LogD3D12RHI, Fatal, TEXT("Missing uniform buffer at slot %u, stage %s. Please check the high level drawing code."), SlotIndex, GetShaderFrequencyString(ShaderStage));
+			}
+
 			if (ShaderStage == SF_Compute)
 			{
 				Context.GraphicsCommandList()->SetComputeRootConstantBufferView(BaseIndex + SlotIndex, CurrentGPUVirtualAddress);
diff --git a/Engine/Source/Runtime/D3D12RHI/Private/D3D12PipelineState.h b/Engine/Source/Runtime/D3D12RHI/Private/D3D12PipelineState.h
--- a/Engine/Source/Runtime/D3D12RHI/Private/D3D12PipelineState.h
+++ b/Engine/Source/Runtime/D3D12RHI/Private/D3D12PipelineState.h
@@ -395,9 +395,6 @@ struct FD3D12ComputePipelineState : public FRHIComputePipelineState, FD3D12Pipel
 	FD3D12ComputePipelineState(FD3D12ComputeShader* InComputeShader, const FD3D12RootSignature* InRootSignature, FD3D12PipelineState* InPipelineState);
 	~FD3D12ComputePipelineState();
 
-	FORCEINLINE FD3D12ComputeShader* GetComputeShader() const { return ComputeShader; }
-
-	TRefCountPtr<FD3D12ComputeShader> ComputeShader;
 	bool bShaderNeedsGlobalConstantBuffer;
 };
 
diff --git a/Engine/Source/Runtime/D3D12RHI/Private/D3D12Pipelinestate.cpp b/Engine/Source/Runtime/D3D12RHI/Private/D3D12Pipelinestate.cpp
--- a/Engine/Source/Runtime/D3D12RHI/Private/D3D12Pipelinestate.cpp
+++ b/Engine/Source/Runtime/D3D12RHI/Private/D3D12Pipelinestate.cpp
@@ -449,8 +449,8 @@ FD3D12GraphicsPipelineState::~FD3D12GraphicsPipelineState()
 }
 
 FD3D12ComputePipelineState::FD3D12ComputePipelineState(FD3D12ComputeShader* InComputeShader, const FD3D12RootSignature* InRootSignature, FD3D12PipelineState* InPipelineState)
-	: FD3D12PipelineStateCommonData(InRootSignature, InPipelineState)
-	, ComputeShader(InComputeShader)
+	: FRHIComputePipelineState(InComputeShader)
+	, FD3D12PipelineStateCommonData(InRootSignature, InPipelineState)
 {
 	bShaderNeedsGlobalConstantBuffer = InComputeShader && InComputeShader->UsesGlobalUniformBuffer();
 }
diff --git a/Engine/Source/Runtime/D3D12RHI/Private/D3D12StateCache.cpp b/Engine/Source/Runtime/D3D12RHI/Private/D3D12StateCache.cpp
--- a/Engine/Source/Runtime/D3D12RHI/Private/D3D12StateCache.cpp
+++ b/Engine/Source/Runtime/D3D12RHI/Private/D3D12StateCache.cpp
@@ -1377,9 +1377,9 @@ void FD3D12StateCache::SetComputePipelineState(FD3D12ComputePipelineState* Compu
 			PipelineState.Compute.bNeedSetRootSignature = true;
 		}
 
-		if (bForceSet || CurrentComputePipelineState->ComputeShader != ComputePipelineState->ComputeShader)
+		if (bForceSet || CurrentComputePipelineState->GetComputeShader() != ComputePipelineState->GetComputeShader())
 		{
-			SetNewShaderData(SF_Compute, ComputePipelineState->ComputeShader);
+			SetNewShaderData(SF_Compute, FD3D12DynamicRHI::ResourceCast(ComputePipelineState->GetComputeShader()));
 		}
 
 		// Save the PSO
diff --git a/Engine/Source/Runtime/RHI/Public/RHIResources.h b/Engine/Source/Runtime/RHI/Public/RHIResources.h
--- a/Engine/Source/Runtime/RHI/Public/RHIResources.h
+++ b/Engine/Source/Runtime/RHI/Public/RHIResources.h
@@ -1072,11 +1072,24 @@ private:
 class FRHIComputePipelineState : public FRHIResource
 {
 public:
-	FRHIComputePipelineState() : FRHIResource(RRT_ComputePipelineState) {}
+	FRHIComputePipelineState(FRHIComputeShader* InComputeShader) :
+		FRHIResource(RRT_ComputePipelineState)
+		, ComputeShader(InComputeShader)
+	{
+		check(InComputeShader);
+	}
 
 	inline void SetValid(bool InIsValid) { bIsValid = InIsValid; }
 	inline bool IsValid() const { return bIsValid; }
 
+	FORCEINLINE FRHIComputeShader* GetComputeShader() const
+	{
+		return ComputeShader;
+	}
+
+protected:
+	TRefCountPtr<FRHIComputeShader> ComputeShader;
+
 private:
 	bool bIsValid = true;
 };
@@ -4313,19 +4326,9 @@ public:
 class FRHIComputePipelineStateFallback : public FRHIComputePipelineState
 {
 public:
-	FRHIComputePipelineStateFallback(FRHIComputeShader* InComputeShader)
-		: ComputeShader(InComputeShader)
+	FRHIComputePipelineStateFallback(FRHIComputeShader* InComputeShader) : FRHIComputePipelineState(InComputeShader)
 	{
-		check(InComputeShader);
 	}
-
-	FRHIComputeShader* GetComputeShader()
-	{
-		return ComputeShader;
-	}
-
-protected:
-	TRefCountPtr<FRHIComputeShader> ComputeShader;
 };
 
 enum class ERenderTargetActions : uint8
diff --git a/Engine/Source/Runtime/VulkanRHI/Private/VulkanPendingState.h b/Engine/Source/Runtime/VulkanRHI/Private/VulkanPendingState.h
--- a/Engine/Source/Runtime/VulkanRHI/Private/VulkanPendingState.h
+++ b/Engine/Source/Runtime/VulkanRHI/Private/VulkanPendingState.h
@@ -59,7 +59,7 @@ public:
 
 	inline const FVulkanComputeShader* GetCurrentShader() const
 	{
-		return CurrentPipeline ? CurrentPipeline->GetShader() : nullptr;
+		return CurrentPipeline ? ResourceCast(CurrentPipeline->GetComputeShader()) : nullptr;
 	}
 
 	void SetUAVForUBResource(uint32 DescriptorSet, uint32 BindingIndex, FVulkanUnorderedAccessView* UAV);
diff --git a/Engine/Source/Runtime/VulkanRHI/Private/VulkanPipeline.cpp b/Engine/Source/Runtime/VulkanRHI/Private/VulkanPipeline.cpp
--- a/Engine/Source/Runtime/VulkanRHI/Private/VulkanPipeline.cpp
+++ b/Engine/Source/Runtime/VulkanRHI/Private/VulkanPipeline.cpp
@@ -217,9 +217,9 @@ FVulkanPipeline::~FVulkanPipeline()
 	/* we do NOT own Layout !*/
 }
 
-FVulkanComputePipeline::FVulkanComputePipeline(FVulkanDevice* InDevice)
+FVulkanComputePipeline::FVulkanComputePipeline(FVulkanDevice* InDevice, FVulkanComputeShader* InShader)
 	: FVulkanPipeline(InDevice)
-	, ComputeShader(nullptr)
+	, FRHIComputePipelineState(InShader)
 {
 	INC_DWORD_STAT(STAT_VulkanNumComputePSOs);
 }
@@ -227,16 +227,9 @@ FVulkanComputePipeline::FVulkanComputePipeline(FVulkanDevice* InDevice)
 FVulkanComputePipeline::~FVulkanComputePipeline()
 {
 	Device->NotifyDeletedComputePipeline(this);
-	
-	if (ComputeShader)
-	{
-		ComputeShader->Release();
-	}
-	
 	DEC_DWORD_STAT(STAT_VulkanNumComputePSOs);
 }
 
-
 FVulkanRHIGraphicsPipelineState::~FVulkanRHIGraphicsPipelineState()
 {
 #if !UE_BUILD_SHIPPING
@@ -2252,10 +2245,7 @@ FVulkanComputePipeline* FVulkanPipelineStateCacheManager::GetOrCreateComputePipe
 
 FVulkanComputePipeline* FVulkanPipelineStateCacheManager::CreateComputePipelineFromShader(FVulkanComputeShader* Shader)
 {
-	FVulkanComputePipeline* Pipeline = new FVulkanComputePipeline(Device);
-
-	Pipeline->ComputeShader = Shader;
-	Pipeline->ComputeShader->AddRef();
+	FVulkanComputePipeline* Pipeline = new FVulkanComputePipeline(Device, Shader);
 
 	FVulkanDescriptorSetsLayoutInfo DescriptorSetLayoutInfo;
 	const FVulkanShaderHeader& CSHeader = Shader->GetCodeHeader();
@@ -2327,9 +2317,10 @@ FVulkanComputePipeline* FVulkanPipelineStateCacheManager::CreateComputePipelineF
 
 void FVulkanPipelineStateCacheManager::NotifyDeletedComputePipeline(FVulkanComputePipeline* Pipeline)
 {
-	if (Pipeline->ComputeShader)
+	FVulkanComputeShader* ComputeShader = ResourceCast(Pipeline->GetComputeShader());
+	if (ComputeShader)
 	{
-		const uint64 Key = Pipeline->ComputeShader->GetShaderKey();
+		const uint64 Key = ComputeShader->GetShaderKey();
 		FRWScopeLock ScopeLock(ComputePipelineLock, SLT_Write); 
 		ComputePipelineEntries.Remove(Key);
 	}
diff --git a/Engine/Source/Runtime/VulkanRHI/Private/VulkanPipeline.h b/Engine/Source/Runtime/VulkanRHI/Private/VulkanPipeline.h
--- a/Engine/Source/Runtime/VulkanRHI/Private/VulkanPipeline.h
+++ b/Engine/Source/Runtime/VulkanRHI/Private/VulkanPipeline.h
@@ -676,17 +676,12 @@ protected:
 class FVulkanComputePipeline : public FVulkanPipeline, public FRHIComputePipelineState
 {
 public:
-	FVulkanComputePipeline(FVulkanDevice* InDevice);
+	FVulkanComputePipeline(FVulkanDevice* InDevice, FVulkanComputeShader* InShader);
 	virtual ~FVulkanComputePipeline();
 
 	inline const FVulkanShaderHeader& GetShaderCodeHeader() const
 	{
-		return ComputeShader->GetCodeHeader();
-	}
-
-	inline const FVulkanComputeShader* GetShader() const
-	{
-		return ComputeShader;
+		return ResourceCast(GetComputeShader())->GetCodeHeader();
 	}
 
 	inline void Bind(VkCommandBuffer CmdBuffer)
@@ -698,11 +693,6 @@ public:
 	{
 		return *(FVulkanComputeLayout*)Layout;
 	}
-
-protected:
-	FVulkanComputeShader*	ComputeShader;
-
-	friend class FVulkanPipelineStateCacheManager;
 };
 
 class FVulkanRHIGraphicsPipelineState : public FRHIGraphicsPipelineState
diff --git a/Engine/Source/Runtime/VulkanRHI/Private/VulkanPipelineState.cpp b/Engine/Source/Runtime/VulkanRHI/Private/VulkanPipelineState.cpp
--- a/Engine/Source/Runtime/VulkanRHI/Private/VulkanPipelineState.cpp
+++ b/Engine/Source/Runtime/VulkanRHI/Private/VulkanPipelineState.cpp
@@ -610,7 +610,7 @@ void FVulkanCommandListContext::RHISetComputePipelineState(FRHIComputePipelineSt
 	FVulkanComputePipeline* ComputePipeline = ResourceCast(ComputePipelineState);
 	PendingComputeState->SetComputePipeline(ComputePipeline);
 
-	ApplyStaticUniformBuffers(const_cast<FVulkanComputeShader*>(ComputePipeline->GetShader()));
+	ApplyStaticUniformBuffers(ResourceCast(ComputePipeline->GetComputeShader()));
 }
 
 
