HWRT: First pass at reducing RHI_RAYTRACING throughout the engine - Vulkan RHI.

diff --git a/Engine/Source/Runtime/VulkanRHI/Private/Android/VulkanAndroidPlatform.cpp b/Engine/Source/Runtime/VulkanRHI/Private/Android/VulkanAndroidPlatform.cpp
--- a/Engine/Source/Runtime/VulkanRHI/Private/Android/VulkanAndroidPlatform.cpp
+++ b/Engine/Source/Runtime/VulkanRHI/Private/Android/VulkanAndroidPlatform.cpp
@@ -333,13 +333,11 @@ bool FVulkanAndroidPlatform::LoadVulkanInstanceFunctions(VkInstance inInstance)
 	ENUM_VK_ENTRYPOINTS_PLATFORM_INSTANCE(GETINSTANCE_VK_ENTRYPOINTS);
 	ENUM_VK_ENTRYPOINTS_PLATFORM_INSTANCE(CHECK_VK_ENTRYPOINTS);
 
-#if VULKAN_RHI_RAYTRACING
 	const bool bFoundRayTracingEntries = FVulkanRayTracingPlatform::CheckVulkanInstanceFunctions(inInstance);
 	if (!bFoundRayTracingEntries)
 	{
 		UE_LOG(LogVulkanRHI, Warning, TEXT("Vulkan RHI ray tracing is enabled, but failed to load instance functions."));
 	}
-#endif
 
 	if (!bFoundAllEntryPoints)
 	{
diff --git a/Engine/Source/Runtime/VulkanRHI/Private/Android/VulkanAndroidPlatform.h b/Engine/Source/Runtime/VulkanRHI/Private/Android/VulkanAndroidPlatform.h
--- a/Engine/Source/Runtime/VulkanRHI/Private/Android/VulkanAndroidPlatform.h
+++ b/Engine/Source/Runtime/VulkanRHI/Private/Android/VulkanAndroidPlatform.h
@@ -19,11 +19,10 @@
 #define VULKAN_SUPPORTS_DEDICATED_ALLOCATION		0
 #define VULKAN_SUPPORTS_ASTC_DECODE_MODE			1
 #define VULKAN_SUPPORTS_NV_DIAGNOSTIC_CHECKPOINT	0
-#define VULKAN_RHI_RAYTRACING						(RHI_RAYTRACING)
-#define VULKAN_SUPPORTS_SCALAR_BLOCK_LAYOUT			(VULKAN_RHI_RAYTRACING)
+#define VULKAN_SUPPORTS_SCALAR_BLOCK_LAYOUT			1
 #define VULKAN_SUPPORTS_TRANSIENT_RESOURCE_ALLOCATOR 0
 #define VULKAN_SUPPORTS_DRIVER_PROPERTIES			0
-#define VULKAN_SUPPORTS_DESCRIPTOR_INDEXING			(VULKAN_RHI_RAYTRACING)
+#define VULKAN_SUPPORTS_DESCRIPTOR_INDEXING			1
 #define VULKAN_SUPPORTS_GPU_CRASH_DUMPS				1
 
 #define UE_VK_API_VERSION							VK_API_VERSION_1_1
diff --git a/Engine/Source/Runtime/VulkanRHI/Private/Linux/VulkanLinuxPlatform.cpp b/Engine/Source/Runtime/VulkanRHI/Private/Linux/VulkanLinuxPlatform.cpp
--- a/Engine/Source/Runtime/VulkanRHI/Private/Linux/VulkanLinuxPlatform.cpp
+++ b/Engine/Source/Runtime/VulkanRHI/Private/Linux/VulkanLinuxPlatform.cpp
@@ -143,13 +143,11 @@ bool FVulkanLinuxPlatform::LoadVulkanInstanceFunctions(VkInstance inInstance)
 		return false;
 	}
 
-#if VULKAN_RHI_RAYTRACING
 	const bool bFoundRayTracingEntries = FVulkanRayTracingPlatform::CheckVulkanInstanceFunctions(inInstance);
 	if (!bFoundRayTracingEntries)
 	{
 		UE_LOG(LogVulkanRHI, Warning, TEXT("Vulkan RHI ray tracing is enabled, but failed to load instance functions."));
 	}
-#endif
 
 	ENUM_VK_ENTRYPOINTS_OPTIONAL_INSTANCE(GETINSTANCE_VK_ENTRYPOINTS);
 	ENUM_VK_ENTRYPOINTS_OPTIONAL_PLATFORM_INSTANCE(GETINSTANCE_VK_ENTRYPOINTS);
diff --git a/Engine/Source/Runtime/VulkanRHI/Private/Linux/VulkanLinuxPlatform.h b/Engine/Source/Runtime/VulkanRHI/Private/Linux/VulkanLinuxPlatform.h
--- a/Engine/Source/Runtime/VulkanRHI/Private/Linux/VulkanLinuxPlatform.h
+++ b/Engine/Source/Runtime/VulkanRHI/Private/Linux/VulkanLinuxPlatform.h
@@ -11,15 +11,9 @@
 #define VULKAN_SHOULD_ENABLE_DRAW_MARKERS			(UE_BUILD_DEBUG || UE_BUILD_DEVELOPMENT)
 #define VULKAN_SIGNAL_UNIMPLEMENTED()				checkf(false, TEXT("Unimplemented vulkan functionality: %hs"), __PRETTY_FUNCTION__)
 #define VULKAN_SUPPORTS_AMD_BUFFER_MARKER			1
+#define VULKAN_SUPPORTS_SCALAR_BLOCK_LAYOUT			1
 
-#define VULKAN_RHI_RAYTRACING 						(RHI_RAYTRACING)
-#define VULKAN_SUPPORTS_SCALAR_BLOCK_LAYOUT			(VULKAN_RHI_RAYTRACING)
-
-#if VULKAN_RHI_RAYTRACING
-#	define UE_VK_API_VERSION						VK_API_VERSION_1_2
-#else
-#	define UE_VK_API_VERSION						VK_API_VERSION_1_1
-#endif // VULKAN_RHI_RAYTRACING
+#define UE_VK_API_VERSION							VK_API_VERSION_1_1
 
 #if UE_BUILD_DEBUG || UE_BUILD_DEVELOPMENT
 #	include "vk_enum_string_helper.h"
diff --git a/Engine/Source/Runtime/VulkanRHI/Private/VulkanBarriers.cpp b/Engine/Source/Runtime/VulkanRHI/Private/VulkanBarriers.cpp
--- a/Engine/Source/Runtime/VulkanRHI/Private/VulkanBarriers.cpp
+++ b/Engine/Source/Runtime/VulkanRHI/Private/VulkanBarriers.cpp
@@ -207,13 +207,11 @@ static void GetVkStageAndAccessFlags(ERHIAccess RHIAccess, FRHITransitionInfo::E
 		RHIAccess = ERHIAccess::RTV;
 	}
 
-#if VULKAN_RHI_RAYTRACING
 	// BVHRead state may be combined with SRV, but we always treat this as just BVHRead by clearing the SRV mask
 	if (EnumHasAnyFlags(RHIAccess, ERHIAccess::BVHRead))
 	{
 		RHIAccess &= ~ERHIAccess::SRVMask;
 	}
-#endif
 
 	Layout = VK_IMAGE_LAYOUT_UNDEFINED;
 
diff --git a/Engine/Source/Runtime/VulkanRHI/Private/VulkanBuffer.cpp b/Engine/Source/Runtime/VulkanRHI/Private/VulkanBuffer.cpp
--- a/Engine/Source/Runtime/VulkanRHI/Private/VulkanBuffer.cpp
+++ b/Engine/Source/Runtime/VulkanRHI/Private/VulkanBuffer.cpp
@@ -69,10 +69,7 @@ VkBufferUsageFlags FVulkanResourceMultiBuffer::UEToVKBufferUsageFlags(FVulkanDev
 	TranslateFlag(BUF_IndexBuffer, VK_BUFFER_USAGE_INDEX_BUFFER_BIT);
 	TranslateFlag(BUF_StructuredBuffer, VK_BUFFER_USAGE_STORAGE_BUFFER_BIT);
 	TranslateFlag(BUF_UniformBuffer, VK_BUFFER_USAGE_UNIFORM_BUFFER_BIT);
-
-#if VULKAN_RHI_RAYTRACING
 	TranslateFlag(BUF_AccelerationStructure, VK_BUFFER_USAGE_ACCELERATION_STRUCTURE_STORAGE_BIT_KHR);
-#endif
 
 	if (!bZeroSize)
 	{
@@ -83,14 +80,13 @@ VkBufferUsageFlags FVulkanResourceMultiBuffer::UEToVKBufferUsageFlags(FVulkanDev
 
 		TranslateFlag(BUF_Volatile, 0, VK_BUFFER_USAGE_TRANSFER_DST_BIT);
 
-#if VULKAN_RHI_RAYTRACING
 		if (InDevice->GetOptionalExtensions().HasRaytracingExtensions())
 		{
 			OutVkUsage |= VK_BUFFER_USAGE_SHADER_DEVICE_ADDRESS_BIT;
 
 			TranslateFlag(BUF_AccelerationStructure, 0, VK_BUFFER_USAGE_ACCELERATION_STRUCTURE_BUILD_INPUT_READ_ONLY_BIT_KHR);
 		}
-#endif
+
 		// For descriptors buffers
 		if (InDevice->GetOptionalExtensions().HasBufferDeviceAddress)
 		{
diff --git a/Engine/Source/Runtime/VulkanRHI/Private/VulkanCommandWrappers.h b/Engine/Source/Runtime/VulkanRHI/Private/VulkanCommandWrappers.h
--- a/Engine/Source/Runtime/VulkanRHI/Private/VulkanCommandWrappers.h
+++ b/Engine/Source/Runtime/VulkanRHI/Private/VulkanCommandWrappers.h
@@ -212,7 +212,7 @@ struct FWrapLayer
 #if VULKAN_USE_CREATE_WIN32_SURFACE
 	static void CreateWin32SurfaceKHR(VkResult Result, VkInstance instance, const VkWin32SurfaceCreateInfoKHR* pCreateInfo, VkSurfaceKHR* pSurface);
 #endif
-#if VULKAN_RHI_RAYTRACING
+
 	static void CreateAccelerationStructureKHR(VkResult Result, VkDevice Device, const VkAccelerationStructureCreateInfoKHR* CreateInfo, const VkAllocationCallbacks* Allocator, VkAccelerationStructureKHR* AccelerationStructure) VULKAN_LAYER_BODY
 	static void DestroyAccelerationStructureKHR(VkResult Result, VkDevice Device, VkAccelerationStructureKHR AccelerationStructure, const VkAllocationCallbacks* Allocator) VULKAN_LAYER_BODY
 	static void CmdBuildAccelerationStructuresKHR(VkResult Result, VkCommandBuffer CommandBuffer, uint32 InfoCount, const VkAccelerationStructureBuildGeometryInfoKHR* Infos, const VkAccelerationStructureBuildRangeInfoKHR* const* BuildRangeInfos) VULKAN_LAYER_BODY
@@ -225,7 +225,7 @@ struct FWrapLayer
 	static void GetRayTracingShaderGroupHandlesKHR(VkResult Result, VkDevice Device, VkPipeline Pipeline, uint32_t FirstGroup, uint32_t GroupCount, size_t DataSize, void* Data) VULKAN_LAYER_BODY
 	static void CmdWriteAccelerationStructuresPropertiesKHR(VkResult Result, VkCommandBuffer CommandBuffer, uint32_t AccelerationStructureCount, const VkAccelerationStructureKHR* AccelerationStructures, VkQueryType QueryType, VkQueryPool QueryPool, uint32_t FirstQuery) VULKAN_LAYER_BODY
 	static void CmdCopyAccelerationStructureKHR(VkResult Result, VkCommandBuffer CommandBuffer, const VkCopyAccelerationStructureInfoKHR* Info) VULKAN_LAYER_BODY
-#endif
+
 	static void GetBufferDeviceAddressKHR(VkResult Result, VkDevice Device, const VkBufferDeviceAddressInfo* Info) VULKAN_LAYER_BODY
 	static void GetDeviceImageMemoryRequirementsKHR(VkResult Result, VkDevice Device, const VkDeviceImageMemoryRequirements* Info, VkMemoryRequirements2* MemoryRequirements) VULKAN_LAYER_BODY
 	static void GetDeviceBufferMemoryRequirementsKHR(VkResult Result, VkDevice Device, const VkDeviceBufferMemoryRequirements* Info, VkMemoryRequirements2* MemoryRequirements) VULKAN_LAYER_BODY
@@ -1466,7 +1466,6 @@ namespace VulkanRHI
 		FWrapLayer::GetBufferMemoryRequirements2(VK_SUCCESS, Device, Info, MemoryRequirements);
 	}
 
-#if VULKAN_RHI_RAYTRACING
 	static FORCEINLINE_DEBUGGABLE VkResult vkCreateAccelerationStructureKHR(VkDevice Device, const VkAccelerationStructureCreateInfoKHR* CreateInfo, const VkAllocationCallbacks* Allocator, VkAccelerationStructureKHR* AccelerationStructure)
 	{
 		FWrapLayer::CreateAccelerationStructureKHR(VK_RESULT_MAX_ENUM, Device, CreateInfo, Allocator, AccelerationStructure);
@@ -1554,7 +1553,6 @@ namespace VulkanRHI
 		VULKANAPINAMESPACE::vkCmdCopyAccelerationStructureKHR(CommandBuffer, Info);
 		FWrapLayer::CmdCopyAccelerationStructureKHR(VK_SUCCESS, CommandBuffer, Info);
 	}
-#endif
 
 	static FORCEINLINE_DEBUGGABLE VkDeviceAddress vkGetBufferDeviceAddressKHR(VkDevice Device, const VkBufferDeviceAddressInfo* Info)
 	{
diff --git a/Engine/Source/Runtime/VulkanRHI/Private/VulkanContext.h b/Engine/Source/Runtime/VulkanRHI/Private/VulkanContext.h
--- a/Engine/Source/Runtime/VulkanRHI/Private/VulkanContext.h
+++ b/Engine/Source/Runtime/VulkanRHI/Private/VulkanContext.h
@@ -111,7 +111,6 @@ public:
 	virtual void RHIEndRenderPass() final override;
 	virtual void RHINextSubpass() final override;
 
-#if VULKAN_RHI_RAYTRACING
 	virtual void RHIClearRayTracingBindings(FRHIRayTracingScene* Scene) final override;
 	virtual void RHICommitRayTracingBindings(FRHIRayTracingScene* Scene) final override;
 	virtual void RHIBindAccelerationStructureMemory(FRHIRayTracingScene* Scene, FRHIBuffer* Buffer, uint32 BufferOffset) final override;
@@ -130,7 +129,6 @@ public:
 		FRHIRayTracingScene* Scene, FRHIRayTracingPipelineState* Pipeline,
 		uint32 NumBindings, const FRayTracingLocalShaderBindings* Bindings,
 		ERayTracingBindingType BindingType) final override;
-#endif // VULKAN_RHI_RAYTRACING
 
 	inline FVulkanCommandBufferManager* GetCommandBufferManager()
 	{
diff --git a/Engine/Source/Runtime/VulkanRHI/Private/VulkanDebug.cpp b/Engine/Source/Runtime/VulkanRHI/Private/VulkanDebug.cpp
--- a/Engine/Source/Runtime/VulkanRHI/Private/VulkanDebug.cpp
+++ b/Engine/Source/Runtime/VulkanRHI/Private/VulkanDebug.cpp
@@ -4072,7 +4072,6 @@ void FWrapLayer::CreateWin32SurfaceKHR(VkResult Result, VkInstance Instance, con
 }
 #endif
 
-#if VULKAN_RHI_RAYTRACING
 void FWrapLayer::CreateAccelerationStructureKHR(VkResult Result, VkDevice Device, const VkAccelerationStructureCreateInfoKHR* CreateInfo, const VkAllocationCallbacks* Allocator, VkAccelerationStructureKHR* AccelerationStructure)
 {
 	if (Result == VK_RESULT_MAX_ENUM)
@@ -4264,7 +4263,7 @@ void FWrapLayer::CmdCopyAccelerationStructureKHR(VkResult Result, VkCommandBuffe
 #endif
 	}
 }
-#endif // VULKAN_RHI_RAYTRACING
+
 
 void FWrapLayer::GetBufferDeviceAddressKHR(VkResult Result, VkDevice Device, const VkBufferDeviceAddressInfo* Info)
 {
diff --git a/Engine/Source/Runtime/VulkanRHI/Private/VulkanDescriptorSets.cpp b/Engine/Source/Runtime/VulkanRHI/Private/VulkanDescriptorSets.cpp
--- a/Engine/Source/Runtime/VulkanRHI/Private/VulkanDescriptorSets.cpp
+++ b/Engine/Source/Runtime/VulkanRHI/Private/VulkanDescriptorSets.cpp
@@ -537,11 +537,7 @@ void FVulkanBindlessDescriptorManager::Init()
 			// Skip anything we don't support
 			if (SetIndex == VulkanBindless::BindlessAccelerationStructureSet)
 			{
-#if VULKAN_RHI_RAYTRACING
 				const bool bHasRaytracingExtensions = Device->GetOptionalExtensions().HasRaytracingExtensions();
-#else
-				const bool bHasRaytracingExtensions = false;
-#endif
 
 				if (!bHasRaytracingExtensions)
 				{
@@ -872,7 +868,6 @@ void FVulkanBindlessDescriptorManager::UpdateTexelBuffer(FRHIDescriptorHandle De
 
 void FVulkanBindlessDescriptorManager::UpdateAccelerationStructure(FRHIDescriptorHandle DescriptorHandle, VkAccelerationStructureKHR AccelerationStructure, bool bImmediateUpdate)
 {
-#if VULKAN_RHI_RAYTRACING
 	if (bIsSupported)
 	{
 		check(GetDescriptorTypeForSetIndex(DescriptorHandle.GetRawType()) == VK_DESCRIPTOR_TYPE_ACCELERATION_STRUCTURE_KHR);
@@ -887,7 +882,6 @@ void FVulkanBindlessDescriptorManager::UpdateAccelerationStructure(FRHIDescripto
 		DescriptorData.accelerationStructure = BufferAddress;
 		UpdateDescriptor(DescriptorHandle, DescriptorData, bImmediateUpdate);
 	}
-#endif
 }
 
 
diff --git a/Engine/Source/Runtime/VulkanRHI/Private/VulkanDescriptorSets.h b/Engine/Source/Runtime/VulkanRHI/Private/VulkanDescriptorSets.h
--- a/Engine/Source/Runtime/VulkanRHI/Private/VulkanDescriptorSets.h
+++ b/Engine/Source/Runtime/VulkanRHI/Private/VulkanDescriptorSets.h
@@ -70,9 +70,7 @@ struct FDescriptorSetRemappingInfo
 		TArray<VkDescriptorType>	Types;
 		uint16						NumImageInfos = 0;
 		uint16						NumBufferInfos = 0;
-#if VULKAN_RHI_RAYTRACING
 		uint8						NumAccelerationStructures = 0;
-#endif // VULKAN_RHI_RAYTRACING
 	};
 	TArray<FSetInfo>	SetInfos;
 
@@ -130,9 +128,7 @@ struct FDescriptorSetRemappingInfo
 			int32 SetInfosNums = SetInfos[SetInfosIndex].Types.Num();
 			if (SetInfos[SetInfosIndex].NumBufferInfos != In.SetInfos[SetInfosIndex].NumBufferInfos ||
 				SetInfos[SetInfosIndex].NumImageInfos != In.SetInfos[SetInfosIndex].NumImageInfos ||
-#if VULKAN_RHI_RAYTRACING
 				SetInfos[SetInfosIndex].NumAccelerationStructures != In.SetInfos[SetInfosIndex].NumAccelerationStructures ||
-#endif // VULKAN_RHI_RAYTRACING
 				SetInfosNums != In.SetInfos[SetInfosIndex].Types.Num() ||
 				(SetInfosNums != 0 && FMemory::Memcmp(SetInfos[SetInfosIndex].Types.GetData(), In.SetInfos[SetInfosIndex].Types.GetData(), sizeof(VkDescriptorType) * SetInfosNums)))
 			{
@@ -288,9 +284,7 @@ public:
 			LayoutTypes.Add(static_cast<VkDescriptorType>(i), 0);
 		}
 
-#if VULKAN_RHI_RAYTRACING
 		LayoutTypes.Add(VK_DESCRIPTOR_TYPE_ACCELERATION_STRUCTURE_KHR, 0);
-#endif
 	}
 
 	inline uint32 GetTypesUsed(VkDescriptorType Type) const
@@ -705,10 +699,9 @@ struct FVulkanDescriptorSetWriteContainer
 	TArray<VkDescriptorImageInfo> DescriptorImageInfo;
 	TArray<VkDescriptorBufferInfo> DescriptorBufferInfo;
 	TArray<VkWriteDescriptorSet> DescriptorWrites;
-#if VULKAN_RHI_RAYTRACING
 	TArray<VkAccelerationStructureKHR> AccelerationStructures;
 	TArray<VkWriteDescriptorSetAccelerationStructureKHR> AccelerationStructureWrites;
-#endif // VULKAN_RHI_RAYTRACING
+
 	TArray<uint8> BindingToDynamicOffsetMap;
 };
 
@@ -960,7 +953,6 @@ public:
 	}
 
 
-#if VULKAN_RHI_RAYTRACING
 	bool WriteAccelerationStructure(uint32 DescriptorIndex, VkAccelerationStructureKHR InAccelerationStructure)
 	{
 		checkf(!UseVulkanDescriptorCache(), TEXT("Descriptor cache path for WriteAccelerationStructure() is not implemented"));
@@ -995,7 +987,6 @@ public:
 
 		return bChanged;
 	}
-#endif // VULKAN_RHI_RAYTRACING
 
 	void SetDescriptorSet(VkDescriptorSet DescriptorSet)
 	{
@@ -1176,10 +1167,8 @@ protected:
 		FVulkanHashableDescriptorInfo* InHashableDescriptorInfos,
 		VkWriteDescriptorSet* InWriteDescriptors, VkDescriptorImageInfo* InImageInfo,
 		VkDescriptorBufferInfo* InBufferInfo, uint8* InBindingToDynamicOffsetMap,
-#if VULKAN_RHI_RAYTRACING
 		VkWriteDescriptorSetAccelerationStructureKHR* InAccelerationStructuresWriteDescriptors,
 		VkAccelerationStructureKHR* InAccelerationStructures,
-#endif // VULKAN_RHI_RAYTRACING
 		const FVulkanSamplerState& DefaultSampler, const FVulkanView::FTextureView& DefaultImageView);
 
 	friend class FVulkanCommonPipelineDescriptorState;
@@ -1334,9 +1323,7 @@ protected:
 	friend class FVulkanComputePipeline;
 	friend class FVulkanGfxPipeline;
 	friend class FVulkanPipelineStateCacheManager;
-#if VULKAN_RHI_RAYTRACING
 	friend class FVulkanRayTracingPipelineState;
-#endif
 };
 
 class FVulkanGfxLayout : public FVulkanLayout
diff --git a/Engine/Source/Runtime/VulkanRHI/Private/VulkanDevice.cpp b/Engine/Source/Runtime/VulkanRHI/Private/VulkanDevice.cpp
--- a/Engine/Source/Runtime/VulkanRHI/Private/VulkanDevice.cpp
+++ b/Engine/Source/Runtime/VulkanRHI/Private/VulkanDevice.cpp
@@ -1264,13 +1264,11 @@ void FVulkanDevice::InitGPU()
 		DefaultTexture = new FVulkanTexture(*this, Desc, nullptr);
 	}
 
-#if VULKAN_RHI_RAYTRACING
 	if (RHISupportsRayTracing(GMaxRHIShaderPlatform) && GetOptionalExtensions().HasRaytracingExtensions())
 	{
 		check(RayTracingCompactionRequestHandler == nullptr);
 		RayTracingCompactionRequestHandler = new FVulkanRayTracingCompactionRequestHandler(this);
 	}
-#endif
 
 	FVulkanPlatform::PostInitGPU(*this);
 }
diff --git a/Engine/Source/Runtime/VulkanRHI/Private/VulkanDevice.h b/Engine/Source/Runtime/VulkanRHI/Private/VulkanDevice.h
--- a/Engine/Source/Runtime/VulkanRHI/Private/VulkanDevice.h
+++ b/Engine/Source/Runtime/VulkanRHI/Private/VulkanDevice.h
@@ -21,10 +21,7 @@ class FVulkanTransientHeapCache;
 class FVulkanDeviceExtension;
 class FVulkanOcclusionQueryPool;
 class FVulkanRenderPassManager;
-
-#if VULKAN_RHI_RAYTRACING
 class FVulkanRayTracingCompactionRequestHandler;
-#endif
 
 #define VULKAN_USE_DEBUG_NAMES 1
 
@@ -105,7 +102,6 @@ struct FOptionalVulkanDeviceExtensions
 		return HasAMDBufferMarker || HasNVDiagnosticCheckpoints;
 	}
 
-#if VULKAN_RHI_RAYTRACING
 	inline bool HasRaytracingExtensions() const
 	{
 		return 
@@ -117,7 +113,6 @@ struct FOptionalVulkanDeviceExtensions
 			HasSPIRV_14 && 
 			HasShaderFloatControls;
 	}
-#endif
 };
 
 // All the features and properties we need to keep around from extension initialization
@@ -131,10 +126,8 @@ struct FOptionalVulkanDeviceExtensionProperties
 	VkPhysicalDeviceDescriptorBufferPropertiesEXT DescriptorBufferProps;
 	VkPhysicalDeviceSubgroupSizeControlPropertiesEXT SubgroupSizeControlProperties;
 
-#if VULKAN_RHI_RAYTRACING
 	VkPhysicalDeviceAccelerationStructurePropertiesKHR AccelerationStructureProps;
 	VkPhysicalDeviceRayTracingPipelinePropertiesKHR RayTracingPipelineProps;
-#endif // VULKAN_RHI_RAYTRACING
 
 	VkPhysicalDeviceFragmentShadingRateFeaturesKHR FragmentShadingRateFeatures;
 	VkPhysicalDeviceFragmentDensityMapFeaturesEXT FragmentDensityMapFeatures;
@@ -338,12 +331,10 @@ public:
 		return GpuSubgroupProps;
 	}
 
-#if VULKAN_RHI_RAYTRACING
 	FVulkanRayTracingCompactionRequestHandler* GetRayTracingCompactionRequestHandler() { return RayTracingCompactionRequestHandler; }
 
 	void InitializeRayTracing();
 	void CleanUpRayTracing();
-#endif // VULKAN_RHI_RAYTRACING
 
 #if VULKAN_SUPPORTS_VALIDATION_CACHE
 	inline VkValidationCacheEXT GetValidationCache() const
@@ -583,9 +574,7 @@ private:
 	VkPhysicalDeviceIDPropertiesKHR GpuIdProps;
 	VkPhysicalDeviceSubgroupProperties GpuSubgroupProps;
 
-#if VULKAN_RHI_RAYTRACING
 	FVulkanRayTracingCompactionRequestHandler* RayTracingCompactionRequestHandler = nullptr;
-#endif // VULKAN_RHI_RAYTRACING
 
 	FVulkanPhysicalDeviceFeatures PhysicalDeviceFeatures;
 
diff --git a/Engine/Source/Runtime/VulkanRHI/Private/VulkanExtensions.cpp b/Engine/Source/Runtime/VulkanRHI/Private/VulkanExtensions.cpp
--- a/Engine/Source/Runtime/VulkanRHI/Private/VulkanExtensions.cpp
+++ b/Engine/Source/Runtime/VulkanRHI/Private/VulkanExtensions.cpp
@@ -833,7 +833,7 @@ class FVulkanKHRAccelerationStructureExtension : public FVulkanDeviceExtension
 public:
 
 	FVulkanKHRAccelerationStructureExtension(FVulkanDevice* InDevice)
-		: FVulkanDeviceExtension(InDevice, VK_KHR_ACCELERATION_STRUCTURE_EXTENSION_NAME, VULKAN_RHI_RAYTRACING)
+		: FVulkanDeviceExtension(InDevice, VK_KHR_ACCELERATION_STRUCTURE_EXTENSION_NAME, VULKAN_EXTENSION_ENABLED)
 	{
 		bEnabledInCode = bEnabledInCode && GVulkanRayTracingCVar.GetValueOnAnyThread() && !FParse::Param(FCommandLine::Get(), TEXT("noraytracing"));
 	}
@@ -854,11 +854,9 @@ public:
 
 	virtual void PrePhysicalDeviceProperties(VkPhysicalDeviceProperties2KHR& PhysicalDeviceProperties2) override final
 	{
-#if VULKAN_RHI_RAYTRACING
 		VkPhysicalDeviceAccelerationStructurePropertiesKHR& AccelerationStructure = GetDeviceExtensionProperties().AccelerationStructureProps;
 		ZeroVulkanStruct(AccelerationStructure, VK_STRUCTURE_TYPE_PHYSICAL_DEVICE_ACCELERATION_STRUCTURE_PROPERTIES_KHR);
 		AddToPNext(PhysicalDeviceProperties2, AccelerationStructure);
-#endif
 	}
 
 	virtual void PreCreateDevice(VkDeviceCreateInfo& DeviceCreateInfo) override final
@@ -881,7 +879,7 @@ class FVulkanKHRRayTracingPipelineExtension : public FVulkanDeviceExtension
 public:
 
 	FVulkanKHRRayTracingPipelineExtension(FVulkanDevice* InDevice) 
-		: FVulkanDeviceExtension(InDevice, VK_KHR_RAY_TRACING_PIPELINE_EXTENSION_NAME, VULKAN_RHI_RAYTRACING)
+		: FVulkanDeviceExtension(InDevice, VK_KHR_RAY_TRACING_PIPELINE_EXTENSION_NAME, VULKAN_EXTENSION_ENABLED)
 	{
 		bEnabledInCode = bEnabledInCode && GVulkanRayTracingCVar.GetValueOnAnyThread() && !FParse::Param(FCommandLine::Get(), TEXT("noraytracing"));
 	}
@@ -902,11 +900,9 @@ public:
 
 	virtual void PrePhysicalDeviceProperties(VkPhysicalDeviceProperties2KHR& PhysicalDeviceProperties2) override final
 	{
-#if VULKAN_RHI_RAYTRACING
 		VkPhysicalDeviceRayTracingPipelinePropertiesKHR& RayTracingPipeline = GetDeviceExtensionProperties().RayTracingPipelineProps;
 		ZeroVulkanStruct(RayTracingPipeline, VK_STRUCTURE_TYPE_PHYSICAL_DEVICE_RAY_TRACING_PIPELINE_PROPERTIES_KHR);
 		AddToPNext(PhysicalDeviceProperties2, RayTracingPipeline);
-#endif
 	}
 
 	virtual void PreCreateDevice(VkDeviceCreateInfo& DeviceCreateInfo) override final
@@ -931,7 +927,7 @@ class FVulkanKHRRayQueryExtension : public FVulkanDeviceExtension
 public:
 
 	FVulkanKHRRayQueryExtension(FVulkanDevice* InDevice) 
-		: FVulkanDeviceExtension(InDevice, VK_KHR_RAY_QUERY_EXTENSION_NAME, VULKAN_RHI_RAYTRACING)
+		: FVulkanDeviceExtension(InDevice, VK_KHR_RAY_QUERY_EXTENSION_NAME, VULKAN_EXTENSION_ENABLED)
 	{
 		bEnabledInCode = bEnabledInCode && GVulkanRayTracingCVar.GetValueOnAnyThread() && !FParse::Param(FCommandLine::Get(), TEXT("noraytracing"));
 	}
@@ -1465,9 +1461,9 @@ FVulkanDeviceExtensionArray FVulkanDeviceExtension::GetUESupportedDeviceExtensio
 	ADD_SIMPLE_EXTENSION(VK_EXT_MEMORY_BUDGET_EXTENSION_NAME,            VULKAN_SUPPORTS_MEMORY_BUDGET,        VULKAN_EXTENSION_NOT_PROMOTED, DEVICE_EXT_FLAG_SETTER(HasMemoryBudget));
 	ADD_SIMPLE_EXTENSION(VK_EXT_MEMORY_PRIORITY_EXTENSION_NAME,          VULKAN_SUPPORTS_MEMORY_PRIORITY,      VULKAN_EXTENSION_NOT_PROMOTED, DEVICE_EXT_FLAG_SETTER(HasMemoryPriority));
 	ADD_SIMPLE_EXTENSION(VK_KHR_CREATE_RENDERPASS_2_EXTENSION_NAME,      VULKAN_SUPPORTS_RENDERPASS2,          VK_API_VERSION_1_2,            DEVICE_EXT_FLAG_SETTER(HasKHRRenderPass2));
-	ADD_SIMPLE_EXTENSION(VK_KHR_DEFERRED_HOST_OPERATIONS_EXTENSION_NAME, VULKAN_RHI_RAYTRACING,                VULKAN_EXTENSION_NOT_PROMOTED, DEVICE_EXT_FLAG_SETTER(HasDeferredHostOperations));
-	ADD_SIMPLE_EXTENSION(VK_KHR_SPIRV_1_4_EXTENSION_NAME,                VULKAN_RHI_RAYTRACING,                VK_API_VERSION_1_2,            DEVICE_EXT_FLAG_SETTER(HasSPIRV_14));
-	ADD_SIMPLE_EXTENSION(VK_KHR_SHADER_FLOAT_CONTROLS_EXTENSION_NAME,    VULKAN_RHI_RAYTRACING,                VK_API_VERSION_1_2,            DEVICE_EXT_FLAG_SETTER(HasShaderFloatControls));
+	ADD_SIMPLE_EXTENSION(VK_KHR_DEFERRED_HOST_OPERATIONS_EXTENSION_NAME, VULKAN_EXTENSION_ENABLED,             VULKAN_EXTENSION_NOT_PROMOTED, DEVICE_EXT_FLAG_SETTER(HasDeferredHostOperations));
+	ADD_SIMPLE_EXTENSION(VK_KHR_SPIRV_1_4_EXTENSION_NAME,				 VULKAN_EXTENSION_ENABLED,             VK_API_VERSION_1_2,            DEVICE_EXT_FLAG_SETTER(HasSPIRV_14));
+	ADD_SIMPLE_EXTENSION(VK_KHR_SHADER_FLOAT_CONTROLS_EXTENSION_NAME,	 VULKAN_EXTENSION_ENABLED,             VK_API_VERSION_1_2,            DEVICE_EXT_FLAG_SETTER(HasShaderFloatControls));
 	ADD_SIMPLE_EXTENSION(VK_KHR_IMAGE_FORMAT_LIST_EXTENSION_NAME,        VULKAN_EXTENSION_ENABLED,             VK_API_VERSION_1_2,            DEVICE_EXT_FLAG_SETTER(HasKHRImageFormatList));
 	ADD_SIMPLE_EXTENSION(VK_EXT_VALIDATION_CACHE_EXTENSION_NAME,         VULKAN_SUPPORTS_VALIDATION_CACHE,     VULKAN_EXTENSION_NOT_PROMOTED, DEVICE_EXT_FLAG_SETTER(HasEXTValidationCache));
 	ADD_SIMPLE_EXTENSION(VK_QCOM_RENDER_PASS_SHADER_RESOLVE_EXTENSION_NAME, VULKAN_SUPPORTS_QCOM_RENDERPASS_SHADER_RESOLVE, VULKAN_EXTENSION_NOT_PROMOTED, DEVICE_EXT_FLAG_SETTER(HasQcomRenderPassShaderResolve));
diff --git a/Engine/Source/Runtime/VulkanRHI/Private/VulkanMemory.cpp b/Engine/Source/Runtime/VulkanRHI/Private/VulkanMemory.cpp
--- a/Engine/Source/Runtime/VulkanRHI/Private/VulkanMemory.cpp
+++ b/Engine/Source/Runtime/VulkanRHI/Private/VulkanMemory.cpp
@@ -21,9 +21,7 @@ THIRD_PARTY_INCLUDES_END
 #include "Windows/HideWindowsPlatformTypes.h"
 #endif
 
-#if VULKAN_RHI_RAYTRACING
 #include "VulkanRayTracing.h"
-#endif
 
 // This 'frame number' should only be used for the deletion queue
 uint32 GVulkanRHIDeletionFrameNumber = 0;
@@ -2639,12 +2637,11 @@ namespace VulkanRHI
 
 		uint32 Alignment = CalculateBufferAlignmentFromVKUsageFlags(InDevice, VulkanBufferUsage);
 
-#if VULKAN_RHI_RAYTRACING
 		if (EnumHasAllFlags(InUEUsage, BUF_RayTracingScratch))
 		{
 			Alignment = GRHIRayTracingScratchBufferAlignment;
 		}
-#endif
+
 		return Alignment;
 	}
 
@@ -2660,7 +2657,6 @@ namespace VulkanRHI
 		{
 			Priority = VULKAN_MEMORY_MEDIUM_PRIORITY;
 		}
-#if VULKAN_RHI_RAYTRACING
 		else if (VKHasAnyFlags(BufferUsageFlags, 
 			(VK_BUFFER_USAGE_ACCELERATION_STRUCTURE_STORAGE_BIT_KHR |
 			VK_BUFFER_USAGE_SHADER_BINDING_TABLE_BIT_KHR |
@@ -2668,7 +2664,6 @@ namespace VulkanRHI
 		{
 			Priority = VULKAN_MEMORY_MEDIUM_PRIORITY;
 		}
-#endif // VULKAN_RHI_RAYTRACING
 		else if (VKHasAnyFlags(BufferUsageFlags, (VK_BUFFER_USAGE_VERTEX_BUFFER_BIT | VK_BUFFER_USAGE_INDEX_BUFFER_BIT)))
 		{
 			Priority = VULKAN_MEMORY_LOW_PRIORITY;
@@ -4857,13 +4852,11 @@ namespace VulkanRHI
 					Device->GetDeviceMemoryManager().Free(Entry->DeviceMemoryAllocation);
 					break;
 				}
-#if VULKAN_RHI_RAYTRACING
 				case EType::AccelerationStructure:
 				{
 					VulkanDynamicAPI::vkDestroyAccelerationStructureKHR(DeviceHandle, (VkAccelerationStructureKHR)Entry->Handle, VULKAN_CPU_ALLOCATOR);
 					break;
 				}
-#endif // VULKAN_RHI_RAYTRACING
 				case EType::BindlessHandle:
 				{
 					check(Device->SupportsBindless());
diff --git a/Engine/Source/Runtime/VulkanRHI/Private/VulkanPendingState.cpp b/Engine/Source/Runtime/VulkanRHI/Private/VulkanPendingState.cpp
--- a/Engine/Source/Runtime/VulkanRHI/Private/VulkanPendingState.cpp
+++ b/Engine/Source/Runtime/VulkanRHI/Private/VulkanPendingState.cpp
@@ -39,7 +39,6 @@ FVulkanDescriptorPool::FVulkanDescriptorPool(FVulkanDevice* InDevice, const FVul
 		}
 	}
 
-#if VULKAN_RHI_RAYTRACING
 	{
 		uint32 NumTypesUsed = Layout.GetTypesUsed(VK_DESCRIPTOR_TYPE_ACCELERATION_STRUCTURE_KHR);
 		if (NumTypesUsed > 0)
@@ -50,7 +49,6 @@ FVulkanDescriptorPool::FVulkanDescriptorPool(FVulkanDevice* InDevice, const FVul
 			Type.descriptorCount = NumTypesUsed * MaxSetsAllocations;
 		}
 	}
-#endif
 
 	VkDescriptorPoolCreateInfo PoolInfo;
 	ZeroVulkanStruct(PoolInfo, VK_STRUCTURE_TYPE_DESCRIPTOR_POOL_CREATE_INFO);
@@ -88,11 +86,9 @@ void FVulkanDescriptorPool::TrackAddUsage(const FVulkanDescriptorSetsLayout& InL
 		ensure(Layout.GetTypesUsed((VkDescriptorType)TypeIndex) == InLayout.GetTypesUsed((VkDescriptorType)TypeIndex));
 	}
 
-#if VULKAN_RHI_RAYTRACING
 	{
 		ensure(Layout.GetTypesUsed(VK_DESCRIPTOR_TYPE_ACCELERATION_STRUCTURE_KHR) == InLayout.GetTypesUsed(VK_DESCRIPTOR_TYPE_ACCELERATION_STRUCTURE_KHR));
 	}
-#endif
 
 	NumAllocatedDescriptorSets += InLayout.GetLayouts().Num();
 	PeakAllocatedDescriptorSets = FMath::Max(NumAllocatedDescriptorSets, PeakAllocatedDescriptorSets);
@@ -105,11 +101,9 @@ void FVulkanDescriptorPool::TrackRemoveUsage(const FVulkanDescriptorSetsLayout&
 		check(Layout.GetTypesUsed((VkDescriptorType)TypeIndex) == InLayout.GetTypesUsed((VkDescriptorType)TypeIndex));
 	}
 
-#if VULKAN_RHI_RAYTRACING
 	{
 		check(Layout.GetTypesUsed(VK_DESCRIPTOR_TYPE_ACCELERATION_STRUCTURE_KHR) == InLayout.GetTypesUsed(VK_DESCRIPTOR_TYPE_ACCELERATION_STRUCTURE_KHR));
 	}
-#endif
 
 	NumAllocatedDescriptorSets -= InLayout.GetLayouts().Num();
 }
diff --git a/Engine/Source/Runtime/VulkanRHI/Private/VulkanPipelineState.cpp b/Engine/Source/Runtime/VulkanRHI/Private/VulkanPipelineState.cpp
--- a/Engine/Source/Runtime/VulkanRHI/Private/VulkanPipelineState.cpp
+++ b/Engine/Source/Runtime/VulkanRHI/Private/VulkanPipelineState.cpp
@@ -77,11 +77,8 @@ void FVulkanCommonPipelineDescriptorState::CreateDescriptorWriteInfos()
 		DSWriteContainer.DescriptorWrites.AddZeroed(SetInfo.Types.Num());
 		DSWriteContainer.DescriptorImageInfo.AddZeroed(SetInfo.NumImageInfos);
 		DSWriteContainer.DescriptorBufferInfo.AddZeroed(SetInfo.NumBufferInfos);
-
-#if VULKAN_RHI_RAYTRACING
 		DSWriteContainer.AccelerationStructureWrites.AddZeroed(SetInfo.NumAccelerationStructures);
 		DSWriteContainer.AccelerationStructures.AddZeroed(SetInfo.NumAccelerationStructures);
-#endif // VULKAN_RHI_RAYTRACING
 
 		checkf(SetInfo.Types.Num() < 255, TEXT("Need more bits for BindingToDynamicOffsetMap (currently 8)! Requires %d descriptor bindings in a set!"), SetInfo.Types.Num());
 		DSWriteContainer.BindingToDynamicOffsetMap.AddUninitialized(SetInfo.Types.Num());
@@ -103,11 +100,8 @@ void FVulkanCommonPipelineDescriptorState::CreateDescriptorWriteInfos()
 	VkWriteDescriptorSet* CurrentDescriptorWrite = DSWriteContainer.DescriptorWrites.GetData();
 	VkDescriptorImageInfo* CurrentImageInfo = DSWriteContainer.DescriptorImageInfo.GetData();
 	VkDescriptorBufferInfo* CurrentBufferInfo = DSWriteContainer.DescriptorBufferInfo.GetData();
-
-#if VULKAN_RHI_RAYTRACING
 	VkWriteDescriptorSetAccelerationStructureKHR* CurrentAccelerationStructuresWriteDescriptors = DSWriteContainer.AccelerationStructureWrites.GetData();
 	VkAccelerationStructureKHR* CurrentAccelerationStructures = DSWriteContainer.AccelerationStructures.GetData();
-#endif // VULKAN_RHI_RAYTRACING
 
 	uint8* CurrentBindingToDynamicOffsetMap = DSWriteContainer.BindingToDynamicOffsetMap.GetData();
 	TArray<uint32> DynamicOffsetsStart;
@@ -123,10 +117,8 @@ void FVulkanCommonPipelineDescriptorState::CreateDescriptorWriteInfos()
 		uint32 NumDynamicOffsets = DSWriter[Set].SetupDescriptorWrites(
 			SetInfo.Types, CurrentHashableDescriptorInfo,
 			CurrentDescriptorWrite, CurrentImageInfo, CurrentBufferInfo, CurrentBindingToDynamicOffsetMap,
-#if VULKAN_RHI_RAYTRACING
 			CurrentAccelerationStructuresWriteDescriptors,
 			CurrentAccelerationStructures,
-#endif // VULKAN_RHI_RAYTRACING
 			DefaultSampler, DefaultImageView);
 
 		TotalNumDynamicOffsets += NumDynamicOffsets;
@@ -143,11 +135,8 @@ void FVulkanCommonPipelineDescriptorState::CreateDescriptorWriteInfos()
 		CurrentDescriptorWrite += SetInfo.Types.Num();
 		CurrentImageInfo += SetInfo.NumImageInfos;
 		CurrentBufferInfo += SetInfo.NumBufferInfos;
-
-#if VULKAN_RHI_RAYTRACING
 		CurrentAccelerationStructuresWriteDescriptors += SetInfo.NumAccelerationStructures;
 		CurrentAccelerationStructures += SetInfo.NumAccelerationStructures;
-#endif // VULKAN_RHI_RAYTRACING
 
 		CurrentBindingToDynamicOffsetMap += SetInfo.Types.Num();
 	}
@@ -758,11 +747,9 @@ void FVulkanCommonPipelineDescriptorState::SetSRV(FVulkanCmdBuffer* CmdBuffer, b
 		MarkDirty(DSWriter[DescriptorSet].WriteStorageBuffer(BindingIndex, SRV->GetStructuredBufferView()));
 		break;
 
-#if VULKAN_RHI_RAYTRACING
 	case FVulkanView::EType::AccelerationStructure:
 		MarkDirty(DSWriter[DescriptorSet].WriteAccelerationStructure(BindingIndex, SRV->GetAccelerationStructureView().Handle));
 		break;
-#endif
 	}
 }
 
@@ -798,10 +785,8 @@ void FVulkanCommonPipelineDescriptorState::SetUAV(FVulkanCmdBuffer* CmdBuffer, b
 		MarkDirty(DSWriter[DescriptorSet].WriteStorageBuffer(BindingIndex, UAV->GetStructuredBufferView()));
 		break;
 
-#if VULKAN_RHI_RAYTRACING
 	case FVulkanView::EType::AccelerationStructure:
 		MarkDirty(DSWriter[DescriptorSet].WriteAccelerationStructure(BindingIndex, UAV->GetAccelerationStructureView().Handle));
 		break;
-#endif
 	}
 }
diff --git a/Engine/Source/Runtime/VulkanRHI/Private/VulkanQueue.cpp b/Engine/Source/Runtime/VulkanRHI/Private/VulkanQueue.cpp
--- a/Engine/Source/Runtime/VulkanRHI/Private/VulkanQueue.cpp
+++ b/Engine/Source/Runtime/VulkanRHI/Private/VulkanQueue.cpp
@@ -174,11 +174,9 @@ void FVulkanQueue::FillSupportedStageBits()
 		SupportedStages |=
 			VK_PIPELINE_STAGE_COMPUTE_SHADER_BIT |
 			VK_PIPELINE_STAGE_DRAW_INDIRECT_BIT |
-			VK_PIPELINE_STAGE_TRANSFER_BIT;
-
-#if VULKAN_RHI_RAYTRACING
-		SupportedStages |= VK_PIPELINE_STAGE_ACCELERATION_STRUCTURE_BUILD_BIT_KHR | VK_PIPELINE_STAGE_RAY_TRACING_SHADER_BIT_KHR;
-#endif
+			VK_PIPELINE_STAGE_TRANSFER_BIT |
+			VK_PIPELINE_STAGE_ACCELERATION_STRUCTURE_BUILD_BIT_KHR |
+			VK_PIPELINE_STAGE_RAY_TRACING_SHADER_BIT_KHR;
 	}
 
 	if (VKHasAnyFlags(QueueProps.queueFlags, VK_QUEUE_TRANSFER_BIT))
diff --git a/Engine/Source/Runtime/VulkanRHI/Private/VulkanRHI.cpp b/Engine/Source/Runtime/VulkanRHI/Private/VulkanRHI.cpp
--- a/Engine/Source/Runtime/VulkanRHI/Private/VulkanRHI.cpp
+++ b/Engine/Source/Runtime/VulkanRHI/Private/VulkanRHI.cpp
@@ -582,12 +582,10 @@ void FVulkanDynamicRHI::Init()
 
 void FVulkanDynamicRHI::PostInit()
 {
-#if VULKAN_RHI_RAYTRACING
 	if (GRHISupportsRayTracing)
 	{
 		Device->InitializeRayTracing();
 	}
-#endif // VULKAN_RHI_RAYTRACING
 }
 
 void FVulkanDynamicRHI::Shutdown()
@@ -629,9 +627,7 @@ void FVulkanDynamicRHI::Shutdown()
 			Device->SamplerMap.Empty();
 		}
 
-#if VULKAN_RHI_RAYTRACING
 		Device->CleanUpRayTracing();
-#endif // VULKAN_RHI_RAYTRACING
 
 		// Flush all pending deletes before destroying the device.
 		FRHICommandListImmediate::Get().ImmediateFlush(EImmediateFlushType::FlushRHIThreadFlushResources);
@@ -679,7 +675,6 @@ void FVulkanDynamicRHI::CreateInstance()
 	// Use the API version stored in the profile
 	ApiVersion = GetVulkanApiVersionForFeatureLevel(GMaxRHIFeatureLevel, false);
 
-#if VULKAN_RHI_RAYTRACING
 	// Run a profile check to see if this device can support our raytacing requirements since it might change the required API version of the instance
 	if (FVulkanPlatform::SupportsProfileChecks() && GVulkanRayTracingCVar.GetValueOnAnyThread())
 	{
@@ -707,7 +702,6 @@ void FVulkanDynamicRHI::CreateInstance()
 			}
 		}
 	}
-#endif // VULKAN_RHI_RAYTRACING
 
 	UE_LOG(LogVulkanRHI, Log, TEXT("Using API Version %u.%u."), VK_API_VERSION_MAJOR(ApiVersion), VK_API_VERSION_MINOR(ApiVersion));
 
@@ -939,7 +933,6 @@ void FVulkanDynamicRHI::InitInstance()
 		GSupportsTimestampRenderQueries = FVulkanPlatform::SupportsTimestampRenderQueries();
 		GSupportsMobileMultiView = Device->GetOptionalExtensions().HasKHRMultiview ? true : false;
 		GRHISupportsMSAAShaderResolve = Device->GetOptionalExtensions().HasQcomRenderPassShaderResolve ? true : false;
-#if VULKAN_RHI_RAYTRACING
 		GRHISupportsRayTracing = RHISupportsRayTracing(GMaxRHIShaderPlatform) && Device->GetOptionalExtensions().HasRaytracingExtensions();
 
 		if (GRHISupportsRayTracing)
@@ -954,7 +947,6 @@ void FVulkanDynamicRHI::InitInstance()
 
 			GRHIRayTracingInstanceDescriptorSize = uint32(sizeof(VkAccelerationStructureInstanceKHR));
 		}
-#endif
 #if VULKAN_ENABLE_DUMP_LAYER
 		// Disable RHI thread by default if the dump layer is enabled
 		GRHISupportsRHIThread = false;
@@ -1121,12 +1113,10 @@ void FVulkanCommandListContext::RHIBeginFrame()
 	GpuProfiler.BeginFrame();
 #endif
 
-#if VULKAN_RHI_RAYTRACING
 	if (GRHISupportsRayTracing)
 	{
 		Device->GetRayTracingCompactionRequestHandler()->Update(*this);
 	}
-#endif
 }
 
 
@@ -1708,11 +1698,9 @@ void FVulkanDescriptorSetsLayoutInfo::AddDescriptor(int32 DescriptorSetIndex, co
 	case VK_DESCRIPTOR_TYPE_UNIFORM_BUFFER_DYNAMIC:
 		IncrementChecked(RemappingInfo.SetInfos[DescriptorSetIndex].NumBufferInfos);
 		break;
-#if VULKAN_RHI_RAYTRACING
 	case VK_DESCRIPTOR_TYPE_ACCELERATION_STRUCTURE_KHR:
 		IncrementChecked(RemappingInfo.SetInfos[DescriptorSetIndex].NumAccelerationStructures);
 		break;
-#endif
 	case VK_DESCRIPTOR_TYPE_STORAGE_TEXEL_BUFFER:
 	case VK_DESCRIPTOR_TYPE_UNIFORM_TEXEL_BUFFER:
 		break;
@@ -1829,12 +1817,10 @@ void FVulkanDescriptorSetsLayout::Compile(FVulkanDescriptorSetLayoutMap& DSetLay
 
 	check(LayoutTypes[VK_DESCRIPTOR_TYPE_INPUT_ATTACHMENT] <= Limits.maxDescriptorSetInputAttachments);
 
-#if VULKAN_RHI_RAYTRACING
 	if (GRHISupportsRayTracing)
 	{
 		check(LayoutTypes[VK_DESCRIPTOR_TYPE_ACCELERATION_STRUCTURE_KHR] < Device->GetOptionalExtensionProperties().AccelerationStructureProps.maxDescriptorSetAccelerationStructures);
 	}
-#endif
 	
 	LayoutHandles.Empty(SetLayouts.Num());
 
@@ -2282,7 +2268,6 @@ void FVulkanDynamicRHI::RHIReplaceResources(FRHICommandListBase& RHICmdList, TAr
 					}
 					break;
 
-#if VULKAN_RHI_RAYTRACING
 				case FRHIResourceReplaceInfo::EType::RTGeometry:
 					{
 						FVulkanRayTracingGeometry* Src = ResourceCast(Info.GetRTGeometry().Src);
@@ -2300,7 +2285,6 @@ void FVulkanDynamicRHI::RHIReplaceResources(FRHICommandListBase& RHICmdList, TAr
 						}
 					}
 					break;
-#endif // VULKAN_RHI_RAYTRACING
 				}
 			}
 		}
diff --git a/Engine/Source/Runtime/VulkanRHI/Private/VulkanRHIPrivate.h b/Engine/Source/Runtime/VulkanRHI/Private/VulkanRHIPrivate.h
--- a/Engine/Source/Runtime/VulkanRHI/Private/VulkanRHIPrivate.h
+++ b/Engine/Source/Runtime/VulkanRHI/Private/VulkanRHIPrivate.h
@@ -64,14 +64,10 @@ inline VkShaderStageFlagBits UEFrequencyToVKStageBit(EShaderFrequency InStage)
 	case SF_Pixel:			return VK_SHADER_STAGE_FRAGMENT_BIT;
 	case SF_Geometry:		return VK_SHADER_STAGE_GEOMETRY_BIT;
 	case SF_Compute:		return VK_SHADER_STAGE_COMPUTE_BIT;
-
-#if VULKAN_RHI_RAYTRACING
 	case SF_RayGen:			return VK_SHADER_STAGE_RAYGEN_BIT_KHR;
 	case SF_RayMiss:		return VK_SHADER_STAGE_MISS_BIT_KHR;
 	case SF_RayHitGroup:	return VK_SHADER_STAGE_CLOSEST_HIT_BIT_KHR; // vkrt todo: How to handle VK_SHADER_STAGE_ANY_HIT_BIT_KHR?
 	case SF_RayCallable:	return VK_SHADER_STAGE_CALLABLE_BIT_KHR;
-#endif // VULKAN_RHI_RAYTRACING
-
 	default:
 		checkf(false, TEXT("Undefined shader stage %d"), (int32)InStage);
 		break;
@@ -89,7 +85,6 @@ inline EShaderFrequency VkStageBitToUEFrequency(VkShaderStageFlagBits FlagBits)
 	case VK_SHADER_STAGE_GEOMETRY_BIT:					return SF_Geometry;
 	case VK_SHADER_STAGE_COMPUTE_BIT:					return SF_Compute;
 
-#if VULKAN_RHI_RAYTRACING
 	case VK_SHADER_STAGE_RAYGEN_BIT_KHR:				return SF_RayGen;
 	case VK_SHADER_STAGE_MISS_BIT_KHR:					return SF_RayMiss;
 	case VK_SHADER_STAGE_CALLABLE_BIT_KHR:				return SF_RayCallable;
@@ -99,7 +94,6 @@ inline EShaderFrequency VkStageBitToUEFrequency(VkShaderStageFlagBits FlagBits)
 	case VK_SHADER_STAGE_CLOSEST_HIT_BIT_KHR:
 	case VK_SHADER_STAGE_INTERSECTION_BIT_KHR:
 		return SF_RayHitGroup;
-#endif // VULKAN_RHI_RAYTRACING
 
 	default:
 		checkf(false, TEXT("Undefined VkShaderStageFlagBits %d"), (int32)FlagBits);
@@ -875,7 +869,7 @@ namespace VulkanRHI
 inline bool UseVulkanDescriptorCache()
 {
 	// Descriptor cache path for WriteAccelerationStructure() is not implemented, so disable if RT is enabled
-	return ((PLATFORM_ANDROID) && !(VULKAN_RHI_RAYTRACING)) || GMaxRHIFeatureLevel <= ERHIFeatureLevel::ES3_1;
+	return GMaxRHIFeatureLevel <= ERHIFeatureLevel::ES3_1;
 }
 
 inline bool ValidateShadingRateDataType()
diff --git a/Engine/Source/Runtime/VulkanRHI/Private/VulkanRayTracing.cpp b/Engine/Source/Runtime/VulkanRHI/Private/VulkanRayTracing.cpp
--- a/Engine/Source/Runtime/VulkanRHI/Private/VulkanRayTracing.cpp
+++ b/Engine/Source/Runtime/VulkanRHI/Private/VulkanRayTracing.cpp
@@ -2,8 +2,6 @@
 
 #include "VulkanRayTracing.h"
 
-#if VULKAN_RHI_RAYTRACING
-
 #include "VulkanContext.h"
 #include "VulkanDescriptorSets.h"
 #include "BuiltInRayTracingShaders.h"
@@ -2282,5 +2280,3 @@ void FVulkanCommandListContext::RHISetRayTracingBindings(
 
 	ParallelForWithExistingTaskContext(TEXT("SetRayTracingBindings"), MakeArrayView(TaskContexts), NumBindings, ItemsPerTask, BindingTask);
 }
-
-#endif // VULKAN_RHI_RAYTRACING
diff --git a/Engine/Source/Runtime/VulkanRHI/Private/VulkanRayTracing.h b/Engine/Source/Runtime/VulkanRHI/Private/VulkanRayTracing.h
--- a/Engine/Source/Runtime/VulkanRHI/Private/VulkanRayTracing.h
+++ b/Engine/Source/Runtime/VulkanRHI/Private/VulkanRayTracing.h
@@ -4,8 +4,6 @@
 
 #include "VulkanRHIPrivate.h"
 
-#if VULKAN_RHI_RAYTRACING
-
 #include "RayTracingBuiltInResources.h"
 
 class FVulkanCommandListContext;
@@ -343,4 +341,4 @@ private:
 	FVulkanRayTracingCompactedSizeQueryPool* QueryPool = nullptr;
 };
 
-#endif // VULKAN_RHI_RAYTRACING
+
diff --git a/Engine/Source/Runtime/VulkanRHI/Private/VulkanResources.h b/Engine/Source/Runtime/VulkanRHI/Private/VulkanResources.h
--- a/Engine/Source/Runtime/VulkanRHI/Private/VulkanResources.h
+++ b/Engine/Source/Runtime/VulkanRHI/Private/VulkanResources.h
@@ -267,7 +267,6 @@ typedef TVulkanBaseShader<FRHIPixelShader, SF_Pixel>				FVulkanPixelShader;
 typedef TVulkanBaseShader<FRHIComputeShader, SF_Compute>			FVulkanComputeShader;
 typedef TVulkanBaseShader<FRHIGeometryShader, SF_Geometry>			FVulkanGeometryShader;
 
-#if VULKAN_RHI_RAYTRACING
 class FVulkanRayTracingShader : public FRHIRayTracingShader, public FVulkanShader
 {
 private:
@@ -305,7 +304,6 @@ public:
 		return FRHIResource::GetRefCount();
 	}
 };
-#endif // VULKAN_RHI_RAYTRACING
 
 class FVulkanShaderFactory
 {
@@ -330,10 +328,8 @@ public:
 		return nullptr;
 	}
 
-#if VULKAN_RHI_RAYTRACING
 	template <EShaderFrequency ShaderFrequency>
 	FVulkanRayTracingShader* CreateRayTracingShader(TArrayView<const uint8> Code, FVulkanDevice* Device);
-#endif
 
 	void LookupShaders(const uint64 InShaderKeys[ShaderStage::NumStages], FVulkanShader* OutShaders[ShaderStage::NumStages]) const;
 
@@ -410,12 +406,10 @@ public:
 		uint32 Size     = 0;
 	};
 
-#if VULKAN_RHI_RAYTRACING
 	struct FAccelerationStructureView
 	{
 		VkAccelerationStructureKHR Handle = VK_NULL_HANDLE;
 	};
-#endif
 
 	struct FTextureView
 	{
@@ -429,9 +423,7 @@ public:
 		, FTypedBufferView
 		, FTextureView
 		, FStructuredBufferView
-#if VULKAN_RHI_RAYTRACING
 		, FAccelerationStructureView
-#endif
 	> TStorage;
 
 	enum EType
@@ -440,9 +432,7 @@ public:
 		TypedBuffer           = TStorage::IndexOfType<FTypedBufferView      >(),
 		Texture               = TStorage::IndexOfType<FTextureView          >(),
 		StructuredBuffer      = TStorage::IndexOfType<FStructuredBufferView >(),
-#if VULKAN_RHI_RAYTRACING
 		AccelerationStructure = TStorage::IndexOfType<FAccelerationStructureView>(),
-#endif
 	};
 
 	FVulkanView(FVulkanDevice& InDevice, VkDescriptorType InDescriptorType);
@@ -464,9 +454,7 @@ public:
 	FTypedBufferView           const& GetTypedBufferView          () const { return Storage.Get<FTypedBufferView          >(); }
 	FTextureView               const& GetTextureView              () const { return Storage.Get<FTextureView              >(); }
 	FStructuredBufferView      const& GetStructuredBufferView     () const { return Storage.Get<FStructuredBufferView     >(); }
-#if VULKAN_RHI_RAYTRACING
 	FAccelerationStructureView const& GetAccelerationStructureView() const { return Storage.Get<FAccelerationStructureView>(); }
-#endif
 
 	// NOTE: The InOffset applies to the FVulkanResourceMultiBuffer (it does not include any internal Allocation offsets that may exist)
 	FVulkanView* InitAsTypedBufferView(
@@ -494,12 +482,10 @@ public:
 		, uint32 InOffset
 		, uint32 InSize);
 
-#if VULKAN_RHI_RAYTRACING
 	FVulkanView* InitAsAccelerationStructureView(
 		  FVulkanResourceMultiBuffer* Buffer
 		, uint32 Offset
 		, uint32 Size);
-#endif
 
 	// No moving or copying
 	FVulkanView(FVulkanView     &&) = delete;
@@ -1518,7 +1504,6 @@ static FORCEINLINE FVulkanTexture* ResourceCast(FRHITexture* Texture)
 	return static_cast<FVulkanTexture*>(Texture->GetTextureBaseRHI());
 }
 
-#if VULKAN_RHI_RAYTRACING
 class FVulkanRayTracingScene;
 class FVulkanRayTracingGeometry;
 class FVulkanRayTracingPipelineState;
@@ -1543,4 +1528,3 @@ struct TVulkanResourceTraits<FRHIRayTracingShader>
 {
 	typedef FVulkanRayTracingShader TConcreteType;
 };
-#endif // VULKAN_RHI_RAYTRACING
diff --git a/Engine/Source/Runtime/VulkanRHI/Private/VulkanShaders.cpp b/Engine/Source/Runtime/VulkanRHI/Private/VulkanShaders.cpp
--- a/Engine/Source/Runtime/VulkanRHI/Private/VulkanShaders.cpp
+++ b/Engine/Source/Runtime/VulkanRHI/Private/VulkanShaders.cpp
@@ -97,7 +97,6 @@ ShaderType* FVulkanShaderFactory::CreateShader(TArrayView<const uint8> Code, FVu
 	return RetShader;
 }
 
-#if VULKAN_RHI_RAYTRACING
 template <EShaderFrequency ShaderFrequency>
 FVulkanRayTracingShader* FVulkanShaderFactory::CreateRayTracingShader(TArrayView<const uint8> Code, FVulkanDevice* Device)
 {
@@ -182,7 +181,6 @@ FVulkanRayTracingShader* FVulkanShaderFactory::CreateRayTracingShader(TArrayView
 
 	return RetShader;
 }
-#endif // VULKAN_RHI_RAYTRACING
 
 void FVulkanShaderFactory::LookupShaders(const uint64 InShaderKeys[ShaderStage::NumStages], FVulkanShader* OutShaders[ShaderStage::NumStages]) const
 {
@@ -473,7 +471,6 @@ TRefCountPtr<FVulkanShaderModule> FVulkanShader::GetOrCreateHandle()
 	return Module;
 }
 
-#if VULKAN_RHI_RAYTRACING
 TRefCountPtr<FVulkanShaderModule> FVulkanRayTracingShader::GetOrCreateHandle(uint32 ModuleIdentifier)
 {
 	check(Device->SupportsBindless());
@@ -541,7 +538,6 @@ TRefCountPtr<FVulkanShaderModule> FVulkanRayTracingShader::GetOrCreateHandle(uin
 
 	return Module;
 }
-#endif // VULKAN_RHI_RAYTRACING
 
 
 TRefCountPtr<FVulkanShaderModule> FVulkanShader::CreateHandle(const FVulkanLayout* Layout, uint32 LayoutHash)
@@ -641,7 +637,6 @@ FComputeShaderRHIRef FVulkanDynamicRHI::RHICreateComputeShader(TArrayView<const
 	return Device->GetShaderFactory().CreateShader<FVulkanComputeShader>(Code, Device);
 }
 
-#if VULKAN_RHI_RAYTRACING
 FRayTracingShaderRHIRef FVulkanDynamicRHI::RHICreateRayTracingShader(TArrayView<const uint8> Code, const FSHAHash& Hash, EShaderFrequency ShaderFrequency)
 {
 	switch (ShaderFrequency)
@@ -663,7 +658,6 @@ FRayTracingShaderRHIRef FVulkanDynamicRHI::RHICreateRayTracingShader(TArrayView<
 		return nullptr;
 	}
 }
-#endif // VULKAN_RHI_RAYTRACING
 
 FVulkanLayout::FVulkanLayout(FVulkanDevice* InDevice)
 	: VulkanRHI::FDeviceChild(InDevice)
@@ -717,10 +711,8 @@ bool FVulkanGfxLayout::UsesInputAttachment(FVulkanShaderHeader::EAttachmentType
 uint32 FVulkanDescriptorSetWriter::SetupDescriptorWrites(
 	const TArray<VkDescriptorType>& Types, FVulkanHashableDescriptorInfo* InHashableDescriptorInfos,
 	VkWriteDescriptorSet* InWriteDescriptors, VkDescriptorImageInfo* InImageInfo, VkDescriptorBufferInfo* InBufferInfo, uint8* InBindingToDynamicOffsetMap,
-#if VULKAN_RHI_RAYTRACING
 	VkWriteDescriptorSetAccelerationStructureKHR* InAccelerationStructuresWriteDescriptors,
 	VkAccelerationStructureKHR* InAccelerationStructures,
-#endif // VULKAN_RHI_RAYTRACING
 	const FVulkanSamplerState& DefaultSampler, const FVulkanView::FTextureView& DefaultImageView)
 {
 	HashableDescriptorInfos = InHashableDescriptorInfos;
@@ -772,7 +764,6 @@ uint32 FVulkanDescriptorSetWriter::SetupDescriptorWrites(
 		case VK_DESCRIPTOR_TYPE_STORAGE_TEXEL_BUFFER:
 		case VK_DESCRIPTOR_TYPE_UNIFORM_TEXEL_BUFFER:
 			break;
-#if VULKAN_RHI_RAYTRACING
 		case VK_DESCRIPTOR_TYPE_ACCELERATION_STRUCTURE_KHR:
 			InAccelerationStructuresWriteDescriptors->sType = VK_STRUCTURE_TYPE_WRITE_DESCRIPTOR_SET_ACCELERATION_STRUCTURE_KHR;
 			InAccelerationStructuresWriteDescriptors->pNext = nullptr;
@@ -780,7 +771,6 @@ uint32 FVulkanDescriptorSetWriter::SetupDescriptorWrites(
 			InAccelerationStructuresWriteDescriptors->pAccelerationStructures = InAccelerationStructures++;
 			InWriteDescriptors->pNext = InAccelerationStructuresWriteDescriptors++;
 			break;
-#endif // VULKAN_RHI_RAYTRACING
 		default:
 			checkf(0, TEXT("Unsupported descriptor type %d"), (int32)Types[Index]);
 			break;
diff --git a/Engine/Source/Runtime/VulkanRHI/Private/VulkanTransientResourceAllocator.cpp b/Engine/Source/Runtime/VulkanRHI/Private/VulkanTransientResourceAllocator.cpp
--- a/Engine/Source/Runtime/VulkanRHI/Private/VulkanTransientResourceAllocator.cpp
+++ b/Engine/Source/Runtime/VulkanRHI/Private/VulkanTransientResourceAllocator.cpp
@@ -13,14 +13,13 @@ FVulkanTransientHeap::FVulkanTransientHeap(const FInitializer& Initializer, FVul
 	EBufferUsageFlags UEBufferUsageFlags = BUF_VertexBuffer | BUF_IndexBuffer | BUF_DrawIndirect 
 		| BUF_UnorderedAccess | BUF_StructuredBuffer | BUF_ShaderResource | BUF_KeepCPUAccessible;
 
-#if VULKAN_RHI_RAYTRACING
 	if (InDevice->GetOptionalExtensions().HasRaytracingExtensions())
 	{
 		UEBufferUsageFlags |= BUF_RayTracingScratch;
 		// AccelerationStructure not yet supported as TransientResource see FVulkanTransientResourceAllocator::CreateBuffer
 		//UEBufferUsageFlags |= BUF_AccelerationStructure;
 	}
-#endif
+
 	const bool bZeroSize = false;
 	VkBufferUsageFlags BufferUsageFlags = FVulkanResourceMultiBuffer::UEToVKBufferUsageFlags(InDevice, UEBufferUsageFlags, bZeroSize);
 
diff --git a/Engine/Source/Runtime/VulkanRHI/Private/VulkanUAV.cpp b/Engine/Source/Runtime/VulkanRHI/Private/VulkanUAV.cpp
--- a/Engine/Source/Runtime/VulkanRHI/Private/VulkanUAV.cpp
+++ b/Engine/Source/Runtime/VulkanRHI/Private/VulkanUAV.cpp
@@ -5,11 +5,7 @@
 #include "VulkanDescriptorSets.h"
 #include "VulkanLLM.h"
 #include "ClearReplacementShaders.h"
-
-#if VULKAN_RHI_RAYTRACING
 #include "VulkanRayTracing.h"
-#endif // VULKAN_RHI_RAYTRACING
-
 
 FVulkanView::FVulkanView(FVulkanDevice& InDevice, VkDescriptorType InDescriptorType)
 	: Device(InDevice)
@@ -53,11 +49,9 @@ void FVulkanView::Invalidate()
 		// Nothing to do
 		break;
 
-#if VULKAN_RHI_RAYTRACING
 	case EType::AccelerationStructure:
 		Device.GetDeferredDeletionQueue().EnqueueResource(VulkanRHI::FDeferredDeletionQueue2::EType::AccelerationStructure, Storage.Get<FAccelerationStructureView>().Handle);
 		break;
-#endif
 	}
 
 	Storage.Emplace<FInvalidatedState>();
@@ -248,7 +242,6 @@ FVulkanView* FVulkanView::InitAsStructuredBufferView(FVulkanResourceMultiBuffer*
 	return this;
 }
 
-#if VULKAN_RHI_RAYTRACING
 FVulkanView* FVulkanView::InitAsAccelerationStructureView(FVulkanResourceMultiBuffer* Buffer, uint32 Offset, uint32 Size)
 {
 	check(GetViewType() == EType::Null);
@@ -268,7 +261,6 @@ FVulkanView* FVulkanView::InitAsAccelerationStructureView(FVulkanResourceMultiBu
 
 	return this;
 }
-#endif
 
 
 
@@ -424,11 +416,9 @@ void FVulkanShaderResourceView::UpdateView()
 				InitAsTypedBufferView(Buffer, Info.Format, Info.OffsetInBytes, Info.SizeInBytes);
 				break;
 
-#if VULKAN_RHI_RAYTRACING
 			case FRHIViewDesc::EBufferType::AccelerationStructure:
 				InitAsAccelerationStructureView(Buffer, Info.OffsetInBytes, Info.SizeInBytes);
 				break;
-#endif
 
 			default:
 				checkNoEntry();
@@ -514,11 +504,9 @@ void FVulkanUnorderedAccessView::UpdateView()
 				InitAsTypedBufferView(Buffer, Info.Format, Info.OffsetInBytes, Info.SizeInBytes);
 				break;
 
-#if VULKAN_RHI_RAYTRACING
 			case FRHIViewDesc::EBufferType::AccelerationStructure:
 				checkNoEntry(); // @todo implement
 				break;
-#endif
 
 			default:
 				checkNoEntry();
diff --git a/Engine/Source/Runtime/VulkanRHI/Private/Windows/VulkanWindowsPlatform.cpp b/Engine/Source/Runtime/VulkanRHI/Private/Windows/VulkanWindowsPlatform.cpp
--- a/Engine/Source/Runtime/VulkanRHI/Private/Windows/VulkanWindowsPlatform.cpp
+++ b/Engine/Source/Runtime/VulkanRHI/Private/Windows/VulkanWindowsPlatform.cpp
@@ -135,13 +135,11 @@ bool FVulkanWindowsPlatform::LoadVulkanInstanceFunctions(VkInstance inInstance)
 	ENUM_VK_ENTRYPOINTS_OPTIONAL_PLATFORM_INSTANCE(CHECK_VK_ENTRYPOINTS);
 #endif
 
-#if VULKAN_RHI_RAYTRACING
 	const bool bFoundRayTracingEntries = FVulkanRayTracingPlatform::CheckVulkanInstanceFunctions(inInstance);
 	if (!bFoundRayTracingEntries)
 	{
 		UE_LOG(LogVulkanRHI, Warning, TEXT("Vulkan RHI ray tracing is enabled, but failed to load instance functions."));
 	}
-#endif
 	
 	ENUM_VK_ENTRYPOINTS_PLATFORM_INSTANCE(GETINSTANCE_VK_ENTRYPOINTS);
 	ENUM_VK_ENTRYPOINTS_PLATFORM_INSTANCE(CHECK_VK_ENTRYPOINTS);
diff --git a/Engine/Source/Runtime/VulkanRHI/Private/Windows/VulkanWindowsPlatform.h b/Engine/Source/Runtime/VulkanRHI/Private/Windows/VulkanWindowsPlatform.h
--- a/Engine/Source/Runtime/VulkanRHI/Private/Windows/VulkanWindowsPlatform.h
+++ b/Engine/Source/Runtime/VulkanRHI/Private/Windows/VulkanWindowsPlatform.h
@@ -15,15 +15,10 @@
 #define VULKAN_SHOULD_ENABLE_DESKTOP_HMD_SUPPORT	1
 #define VULKAN_SIGNAL_UNIMPLEMENTED()				checkf(false, TEXT("Unimplemented vulkan functionality: %s"), StringCast<TCHAR>(__FUNCTION__).Get())
 #define VULKAN_SUPPORTS_AMD_BUFFER_MARKER			1
+#define VULKAN_SUPPORTS_SCALAR_BLOCK_LAYOUT			1
 
-#define VULKAN_RHI_RAYTRACING 						(RHI_RAYTRACING)
-#define VULKAN_SUPPORTS_SCALAR_BLOCK_LAYOUT			(VULKAN_RHI_RAYTRACING)
+#define UE_VK_API_VERSION							VK_API_VERSION_1_1
 
-#if VULKAN_RHI_RAYTRACING
-#	define UE_VK_API_VERSION						VK_API_VERSION_1_2
-#else
-#	define UE_VK_API_VERSION						VK_API_VERSION_1_1
-#endif // VULKAN_RHI_RAYTRACING
 
 #if UE_BUILD_DEBUG || UE_BUILD_DEVELOPMENT
 #	include "vk_enum_string_helper.h"
diff --git a/Engine/Source/Runtime/VulkanRHI/Public/VulkanDynamicRHI.h b/Engine/Source/Runtime/VulkanRHI/Public/VulkanDynamicRHI.h
--- a/Engine/Source/Runtime/VulkanRHI/Public/VulkanDynamicRHI.h
+++ b/Engine/Source/Runtime/VulkanRHI/Public/VulkanDynamicRHI.h
@@ -223,14 +223,12 @@ public:
 	virtual FRHIResourceCollectionRef RHICreateResourceCollection(FRHICommandListBase& RHICmdList, TConstArrayView<FRHIResourceCollectionMember> InMembers) final;
 #endif
 
-#if VULKAN_RHI_RAYTRACING
 	virtual FRayTracingAccelerationStructureSize RHICalcRayTracingSceneSize(uint32 MaxInstances, ERayTracingAccelerationStructureFlags Flags) final override;
 	virtual FRayTracingAccelerationStructureSize RHICalcRayTracingGeometrySize(const FRayTracingGeometryInitializer& Initializer) final override;
 	virtual FRayTracingGeometryRHIRef RHICreateRayTracingGeometry(FRHICommandListBase& RHICmdList, const FRayTracingGeometryInitializer& Initializer) final override;
 	virtual FRayTracingSceneRHIRef RHICreateRayTracingScene(FRayTracingSceneInitializer2 Initializer) final override;
 	virtual FRayTracingShaderRHIRef RHICreateRayTracingShader(TArrayView<const uint8> Code, const FSHAHash& Hash, EShaderFrequency ShaderFrequency) final override;
 	virtual FRayTracingPipelineStateRHIRef RHICreateRayTracingPipelineState(const FRayTracingPipelineStateInitializer& Initializer) final override;
-#endif
 
 	// Historical number of times we've presented any and all viewports
 	uint32 TotalPresentCount = 0;
