HWRT: Removing RHI_RAYTRACING that was missed for Vulkan in previous change.

diff --git a/Engine/Source/Runtime/VulkanRHI/Private/VulkanPipeline.h b/Engine/Source/Runtime/VulkanRHI/Private/VulkanPipeline.h
--- a/Engine/Source/Runtime/VulkanRHI/Private/VulkanPipeline.h
+++ b/Engine/Source/Runtime/VulkanRHI/Private/VulkanPipeline.h
@@ -44,13 +44,11 @@ inline uint64 GetShaderKeyForGfxStage(const FBoundShaderStateInput& BSI, ShaderS
 	case ShaderStage::Geometry:
 		return GetShaderKey<FVulkanGeometryShader>(BSI.GetGeometryShader());
 #endif
-#if RHI_RAYTRACING
 	case ShaderStage::RayGen:
 	case ShaderStage::RayHitGroup:
 	case ShaderStage::RayMiss:
 	case ShaderStage::RayCallable:
 		return 0; // VKRT todo
-#endif
 	default:
 		check(0);
 	}
diff --git a/Engine/Source/Runtime/VulkanRHI/Public/VulkanCommon.h b/Engine/Source/Runtime/VulkanRHI/Public/VulkanCommon.h
--- a/Engine/Source/Runtime/VulkanRHI/Public/VulkanCommon.h
+++ b/Engine/Source/Runtime/VulkanRHI/Public/VulkanCommon.h
@@ -31,12 +31,10 @@ namespace ShaderStage
 		Geometry = 2,
 #endif
 
-#if RHI_RAYTRACING
 		RayGen = 3,
 		RayMiss = 4,
 		RayHitGroup = 5,
 		RayCallable = 6,
-#endif
 
 #if VULKAN_SUPPORTS_GEOMETRY_SHADERS
 		NumGeometryStages = 1,
@@ -44,22 +42,14 @@ namespace ShaderStage
 		NumGeometryStages = 0,
 #endif
 
-#if RHI_RAYTRACING
 		NumRayTracingStages = 4,
-#else
-		NumRayTracingStages = 0,
-#endif
 
 		NumStages = (2 + NumGeometryStages + NumRayTracingStages),
 
 		// Compute is its own pipeline, so it can all live as set 0
 		Compute = 0,
 
-#if VULKAN_SUPPORTS_GEOMETRY_SHADERS || RHI_RAYTRACING
 		MaxNumSets = 8,
-#else
-		MaxNumSets = 4,
-#endif
 
 		Invalid = -1,
 	};
@@ -73,12 +63,10 @@ namespace ShaderStage
 #if VULKAN_SUPPORTS_GEOMETRY_SHADERS
 		case SF_Geometry:	return Geometry;
 #endif
-#if RHI_RAYTRACING
 		case SF_RayGen:			return RayGen;
 		case SF_RayMiss:		return RayMiss;
 		case SF_RayHitGroup:	return RayHitGroup;
 		case SF_RayCallable:	return RayCallable;
-#endif // RHI_RAYTRACING
 		case SF_Compute:	return Compute;
 		default:
 			checkf(0, TEXT("Invalid shader Stage %d"), (int32)Stage);
@@ -97,12 +85,10 @@ namespace ShaderStage
 #if VULKAN_SUPPORTS_GEOMETRY_SHADERS
 		case EStage::Geometry:	return SF_Geometry;
 #endif
-#if RHI_RAYTRACING
 		case EStage::RayGen:		return SF_RayGen;
 		case EStage::RayMiss:		return SF_RayMiss;
 		case EStage::RayHitGroup:	return SF_RayHitGroup;
 		case EStage::RayCallable:	return SF_RayCallable;
-#endif //	RHI_RAYTRACING
 		default:
 			checkf(0, TEXT("Invalid shader Stage %d"), (int32)Stage);
 			break;
diff --git a/Engine/Source/Runtime/VulkanRHI/Public/VulkanShaderResources.h b/Engine/Source/Runtime/VulkanRHI/Public/VulkanShaderResources.h
--- a/Engine/Source/Runtime/VulkanRHI/Public/VulkanShaderResources.h
+++ b/Engine/Source/Runtime/VulkanRHI/Public/VulkanShaderResources.h
@@ -26,9 +26,7 @@ static inline VkDescriptorType BindingToDescriptorType(EVulkanBindingType::EType
 	case EVulkanBindingType::StorageTexelBuffer:	return VK_DESCRIPTOR_TYPE_STORAGE_TEXEL_BUFFER;
 	case EVulkanBindingType::StorageBuffer:			return VK_DESCRIPTOR_TYPE_STORAGE_BUFFER;
 	case EVulkanBindingType::InputAttachment:		return VK_DESCRIPTOR_TYPE_INPUT_ATTACHMENT;
-#if RHI_RAYTRACING
 	case EVulkanBindingType::AccelerationStructure:	return VK_DESCRIPTOR_TYPE_ACCELERATION_STRUCTURE_KHR;
-#endif
 	default:
 		check(0);
 		break;
@@ -50,9 +48,7 @@ static inline EVulkanBindingType::EType DescriptorTypeToBinding(VkDescriptorType
 	case VK_DESCRIPTOR_TYPE_STORAGE_TEXEL_BUFFER:		return EVulkanBindingType::StorageTexelBuffer;
 	case VK_DESCRIPTOR_TYPE_STORAGE_BUFFER:				return EVulkanBindingType::StorageBuffer;
 	case VK_DESCRIPTOR_TYPE_INPUT_ATTACHMENT:			return EVulkanBindingType::InputAttachment;
-#if RHI_RAYTRACING
 	case VK_DESCRIPTOR_TYPE_ACCELERATION_STRUCTURE_KHR: return EVulkanBindingType::AccelerationStructure;
-#endif
 	default:
 		check(0);
 		break;
