VulkanRHI: fix stack value being used in renderpass CreateInfo

#rnx

diff --git a/Engine/Source/Runtime/VulkanRHI/Private/VulkanRenderpass.h b/Engine/Source/Runtime/VulkanRHI/Private/VulkanRenderpass.h
--- a/Engine/Source/Runtime/VulkanRHI/Private/VulkanRenderpass.h
+++ b/Engine/Source/Runtime/VulkanRHI/Private/VulkanRenderpass.h
@@ -753,7 +753,8 @@ public:
 		Bit mask that specifies which view rendering is broadcast to
 		0011 = Broadcast to first and second view (layer)
 		*/
-		const uint32_t ViewMask[2] = { MultiviewMask, MultiviewMask };
+		ViewMask[0] = MultiviewMask;
+		ViewMask[1] = MultiviewMask;
 
 		/*
 		Bit mask that specifices correlation between views
@@ -840,7 +841,8 @@ private:
 	TRenderPassCreateInfoClass CreateInfo;
 	FVulkanDevice& Device;
 
-	uint32_t CorrelationMask;
+	uint32 ViewMask[2];
+	uint32 CorrelationMask;
 };
 
 VkRenderPass CreateVulkanRenderPass(FVulkanDevice& Device, const FVulkanRenderTargetLayout& RTLayout);
