Fix reads of uninitialized FVulkanTexture::ImageUsageFlags.

diff --git a/Engine/Source/Runtime/VulkanRHI/Private/VulkanTexture.cpp b/Engine/Source/Runtime/VulkanRHI/Private/VulkanTexture.cpp
--- a/Engine/Source/Runtime/VulkanRHI/Private/VulkanTexture.cpp
+++ b/Engine/Source/Runtime/VulkanRHI/Private/VulkanTexture.cpp
@@ -1753,6 +1753,7 @@ FVulkanTexture::FVulkanTexture(FVulkanDevice& InDevice, const FRHITextureCreateD
 	: FRHITexture(InCreateDesc)
 	, Device(&InDevice)
 	, Image(InImage)
+	, ImageUsageFlags(0)
 	, StorageFormat(VK_FORMAT_UNDEFINED)
 	, ViewFormat(VK_FORMAT_UNDEFINED)
 	, MemProps(VK_MEMORY_PROPERTY_DEVICE_LOCAL_BIT)
@@ -1782,9 +1783,10 @@ FVulkanTexture::FVulkanTexture(FVulkanDevice& InDevice, const FRHITextureCreateD
 
 		if (Image != VK_NULL_HANDLE)
 		{
-#if VULKAN_ENABLE_WRAP_LAYER
 			FImageCreateInfo ImageCreateInfo;
 			FVulkanTexture::GenerateImageCreateInfo(ImageCreateInfo, InDevice, InCreateDesc, &StorageFormat, &ViewFormat);
+			ImageUsageFlags = ImageCreateInfo.ImageCreateInfo.usage;
+#if VULKAN_ENABLE_WRAP_LAYER
 			FWrapLayer::CreateImage(VK_SUCCESS, InDevice.GetInstanceHandle(), &ImageCreateInfo.ImageCreateInfo, &Image);
 #endif
 			VULKAN_SET_DEBUG_NAME(InDevice, VK_OBJECT_TYPE_IMAGE, Image, TEXT("%s:(FVulkanTexture*)0x%p"), InCreateDesc.DebugName ? InCreateDesc.DebugName : TEXT("?"), this);
@@ -1861,6 +1863,7 @@ FVulkanTexture::FVulkanTexture(FVulkanDevice& InDevice, const FRHITextureCreateD
 	: FRHITexture(InCreateDesc)
 	, Device(&InDevice)
 	, Image(VK_NULL_HANDLE)
+	, ImageUsageFlags(0)
 	, StorageFormat(VK_FORMAT_UNDEFINED)
 	, ViewFormat(VK_FORMAT_UNDEFINED)
 	, MemProps(VK_MEMORY_PROPERTY_DEVICE_LOCAL_BIT)
@@ -1887,6 +1890,11 @@ FVulkanTexture::FVulkanTexture(FVulkanDevice& InDevice, const FRHITextureCreateD
 		{
 			Tiling = VK_IMAGE_TILING_OPTIMAL;
 		}
+
+		FImageCreateInfo ImageCreateInfo;
+		FVulkanTexture::GenerateImageCreateInfo(ImageCreateInfo, InDevice, InCreateDesc, &StorageFormat, &ViewFormat);
+
+		ImageUsageFlags = ImageCreateInfo.ImageCreateInfo.usage;
 	}
 
 	AliasTextureResources(SrcTextureRHI);
