VulkanRHI: Improve performance of Vulkan ray tracing on the first run by enabling DeferredOperation by default.

#rnx

diff --git a/Engine/Source/Runtime/VulkanRHI/Private/VulkanRayTracing.cpp b/Engine/Source/Runtime/VulkanRHI/Private/VulkanRayTracing.cpp
--- a/Engine/Source/Runtime/VulkanRHI/Private/VulkanRayTracing.cpp
+++ b/Engine/Source/Runtime/VulkanRHI/Private/VulkanRayTracing.cpp
@@ -26,7 +26,7 @@ static FAutoConsoleVariableRef CVarVulkanRayTracingMaxBatchedCompaction(
 	ECVF_ReadOnly
 );
 
-static int32 GVulkanRayTracingAllowDeferredOperation = -1;
+static int32 GVulkanRayTracingAllowDeferredOperation = 0;
 static FAutoConsoleVariableRef CVarVulkanRayTracingAllowDeferredOperation(
 	TEXT("r.Vulkan.RayTracing.AllowDeferredOperation"),
 	GVulkanRayTracingAllowDeferredOperation,
@@ -1565,6 +1565,11 @@ FVulkanRayTracingPipelineState::FVulkanRayTracingPipelineState(FVulkanDevice* co
 		{
 			MaxConcurrency = FMath::Min(MaxConcurrency, GVulkanRayTracingAllowDeferredOperation);
 		}
+		else if (MaxConcurrency > 3)
+		{
+			// Prevent automatic detection from completely clogging the machine
+			MaxConcurrency -= 2;
+		}
 
 		bool bCompleted = false;
 		ParallelFor(MaxConcurrency, [DeferredOp, InDevice, &bCompleted](int32 Unused)
