VulkanRHI: start removing code related to r.Vulkan.DescriptorSetLayoutMode

#rnx

diff --git a/Engine/Source/Runtime/VulkanRHI/Private/VulkanDescriptorSets.h b/Engine/Source/Runtime/VulkanRHI/Private/VulkanDescriptorSets.h
--- a/Engine/Source/Runtime/VulkanRHI/Private/VulkanDescriptorSets.h
+++ b/Engine/Source/Runtime/VulkanRHI/Private/VulkanDescriptorSets.h
@@ -33,10 +33,6 @@ struct FUniformBufferGatherInfo
 		FMemory::Memzero(CodeHeaders);
 	}
 
-	// These maps are used to find UBs that are used on multiple stages
-	TMap<uint32, VkShaderStageFlags>	UBLayoutsToUsedStageMap;
-	TMap<uint32, VkShaderStageFlags>	CommonUBLayoutsToStageMap;
-
 	const FVulkanShaderHeader*	CodeHeaders[ShaderStage::NumStages];
 };
 
diff --git a/Engine/Source/Runtime/VulkanRHI/Private/VulkanShaders.cpp b/Engine/Source/Runtime/VulkanRHI/Private/VulkanShaders.cpp
--- a/Engine/Source/Runtime/VulkanRHI/Private/VulkanShaders.cpp
+++ b/Engine/Source/Runtime/VulkanRHI/Private/VulkanShaders.cpp
@@ -22,16 +22,6 @@ TAutoConsoleVariable<int32> GDynamicGlobalUBs(
 	ECVF_ReadOnly | ECVF_RenderThreadSafe
 );
 
-
-static TAutoConsoleVariable<int32> GDescriptorSetLayoutMode(
-	TEXT("r.Vulkan.DescriptorSetLayoutMode"),
-	0,
-	TEXT("0 to not change layouts (eg Set 0 = Vertex, 1 = Pixel, etc\n")\
-	TEXT("1 to use a new set for common Uniform Buffers\n")\
-	TEXT("2 to collapse all sets into Set 0"),
-	ECVF_ReadOnly | ECVF_RenderThreadSafe
-);
-
 static int32 GVulkanCompressSPIRV = 0;
 static FAutoConsoleVariableRef GVulkanCompressSPIRVCVar(
 	TEXT("r.Vulkan.CompressSPIRV"),
@@ -783,43 +773,6 @@ uint32 FVulkanDescriptorSetWriter::SetupDescriptorWrites(
 
 void FVulkanDescriptorSetsLayoutInfo::ProcessBindingsForStage(VkShaderStageFlagBits StageFlags, ShaderStage::EStage DescSetStage, const FVulkanShaderHeader& CodeHeader, FUniformBufferGatherInfo& OutUBGatherInfo) const
 {
-	const bool bMoveCommonUBsToExtraSet = GDescriptorSetLayoutMode.GetValueOnAnyThread() == 1 || GDescriptorSetLayoutMode.GetValueOnAnyThread() == 2;
-
-	// Find all common UBs from different stages
-	for (const FVulkanShaderHeader::FUniformBufferInfo& UBInfo : CodeHeader.UniformBuffers)
-	{
-		if (bMoveCommonUBsToExtraSet)
-		{
-			VkShaderStageFlags* Found = OutUBGatherInfo.CommonUBLayoutsToStageMap.Find(UBInfo.LayoutHash);
-			if (Found)
-			{
-				*Found = *Found | StageFlags;
-			}
-			else
-			{
-				//#todo-rco: Only process constant data part of the UB
-				Found = (UBInfo.ConstantDataOriginalBindingIndex == UINT16_MAX) ? nullptr : OutUBGatherInfo.UBLayoutsToUsedStageMap.Find(UBInfo.LayoutHash);
-				if (Found && *Found != StageFlags)
-				{
-					// Move from per stage to common UBs
-					VkShaderStageFlags PrevStage = (VkShaderStageFlags)0;
-					bool bFound = OutUBGatherInfo.UBLayoutsToUsedStageMap.RemoveAndCopyValue(UBInfo.LayoutHash, PrevStage);
-					check(bFound);
-					check(OutUBGatherInfo.CommonUBLayoutsToStageMap.Find(UBInfo.LayoutHash) == nullptr);
-					OutUBGatherInfo.CommonUBLayoutsToStageMap.Add(UBInfo.LayoutHash, PrevStage | (VkShaderStageFlags)StageFlags);
-				}
-				else
-				{
-					OutUBGatherInfo.UBLayoutsToUsedStageMap.Add(UBInfo.LayoutHash, (VkShaderStageFlags)StageFlags);
-				}
-			}
-		}
-		else
-		{
-			OutUBGatherInfo.UBLayoutsToUsedStageMap.Add(UBInfo.LayoutHash, (VkShaderStageFlags)StageFlags);
-		}
-	}
-
 	OutUBGatherInfo.CodeHeaders[DescSetStage] = &CodeHeader;
 }
 
@@ -837,22 +790,13 @@ void FVulkanDescriptorSetsLayoutInfo::FinalizeBindings(const FVulkanDevice& Devi
 
 	const bool bConvertAllUBsToDynamic = !Device.SupportsBindless() && (GDynamicGlobalUBs.GetValueOnAnyThread() > 1);
 	const bool bConvertPackedUBsToDynamic = !Device.SupportsBindless() && (bConvertAllUBsToDynamic || (GDynamicGlobalUBs.GetValueOnAnyThread() == 1));
-	const bool bConsolidateAllIntoOneSet = GDescriptorSetLayoutMode.GetValueOnAnyThread() == 2;
 	const uint32 MaxDescriptorSetUniformBuffersDynamic = Device.GetLimits().maxDescriptorSetUniformBuffersDynamic;
 
 	uint8	DescriptorStageToSetMapping[ShaderStage::NumStages];
 	FMemory::Memset(DescriptorStageToSetMapping, UINT8_MAX);
 
-	const bool bMoveCommonUBsToExtraSet = (UBGatherInfo.CommonUBLayoutsToStageMap.Num() > 0) || bConsolidateAllIntoOneSet;
-	const uint32 CommonUBDescriptorSet = bMoveCommonUBsToExtraSet ? RemappingInfo.SetInfos.AddDefaulted() : UINT32_MAX;
-
 	auto FindOrAddDescriptorSet = [&](int32 Stage) -> uint8
 	{
-		if (bConsolidateAllIntoOneSet)
-		{
-			return 0;
-		}
-
 		if (DescriptorStageToSetMapping[Stage] == UINT8_MAX)
 		{
 			uint32 NewSet = RemappingInfo.SetInfos.AddDefaulted();
@@ -901,28 +845,6 @@ void FVulkanDescriptorSetsLayoutInfo::FinalizeBindings(const FVulkanDevice& Devi
 				if (bUBHasConstantData)
 				{
 					bool bProcessRegularUB = true;
-					const VkShaderStageFlags* FoundFlags = bMoveCommonUBsToExtraSet ? UBGatherInfo.CommonUBLayoutsToStageMap.Find(LayoutHash) : nullptr;
-					if (FoundFlags)
-					{
-						if (const FDescriptorSetRemappingInfo::FUBRemappingInfo* UBRemapInfo = AlreadyProcessedUBs.Find(LayoutHash))
-						{
-							RemappingInfo.AddRedundantUB(Stage, Index, UBRemapInfo);
-						}
-						else
-						{
-							//#todo-rco: Only process constant data part of the UB
-							check(bUBHasConstantData);
-
-							Binding.stageFlags = *FoundFlags;
-							uint32 NewBindingIndex;
-							AlreadyProcessedUBs.Add(LayoutHash, RemappingInfo.AddUBWithData(Stage, Index, CommonUBDescriptorSet, Type, NewBindingIndex));
-							Binding.binding = NewBindingIndex;
-
-							AddDescriptor(CommonUBDescriptorSet, Binding);
-						}
-						bProcessRegularUB = false;
-					}
-
 					if (bProcessRegularUB)
 					{
 						int32 DescriptorSet = FindOrAddDescriptorSet(Stage);
@@ -994,18 +916,6 @@ void FVulkanDescriptorSetsLayoutInfo::FinalizeBindings(const FVulkanDevice& Devi
 	GenerateHash(ImmutableSamplers, bIsCompute ? VK_PIPELINE_BIND_POINT_COMPUTE : VK_PIPELINE_BIND_POINT_GRAPHICS);
 
 	// If we are consolidating and no uniforms are present in the shader, then strip the empty set data
-	if (bConsolidateAllIntoOneSet)
-	{
-		for (int32 Index = 0; Index < RemappingInfo.SetInfos.Num(); ++Index)
-		{
-			if (RemappingInfo.SetInfos[Index].Types.Num() == 0)
-			{
-				RemappingInfo.SetInfos.RemoveAt(Index);
-			}
-		}
-		check(RemappingInfo.SetInfos.Num() <= 1);
-	}
-	else
 	{
 		for (int32 Index = 0; Index < RemappingInfo.SetInfos.Num(); ++Index)
 		{
