PR #11825: Optional limit on Vulkan barriers per batch to circumvent issues with internal limits on some drivers

#11825
#rnx

diff --git a/Engine/Config/BaseDeviceProfiles.ini b/Engine/Config/BaseDeviceProfiles.ini
--- a/Engine/Config/BaseDeviceProfiles.ini
+++ b/Engine/Config/BaseDeviceProfiles.ini
@@ -1208,3 +1208,4 @@ BaseProfileName=Android_Mid
 +CVars=xr.DisableOpenXROnAndroidWithoutOculus=0
 +CVars=r.Android.DisableVulkanSupport=0
 +CVars=r.Vulkan.SupportsBCTextureFormats=1
++CVars=r.Vulkan.MaxBarriersPerBatch=128
\ No newline at end of file
diff --git a/Engine/Source/Runtime/VulkanRHI/Private/VulkanBarriers.cpp b/Engine/Source/Runtime/VulkanRHI/Private/VulkanBarriers.cpp
--- a/Engine/Source/Runtime/VulkanRHI/Private/VulkanBarriers.cpp
+++ b/Engine/Source/Runtime/VulkanRHI/Private/VulkanBarriers.cpp
@@ -44,6 +44,17 @@ static FAutoConsoleVariableRef CVarVulkanAutoCorrectUnknownLayouts(
 	ECVF_Default
 );
 
+int32 GVulkanMaxBarriersPerBatch = -1;
+static FAutoConsoleVariableRef CVarVulkanMaxBarriersPerBatch(
+	TEXT("r.Vulkan.MaxBarriersPerBatch"),
+	GVulkanMaxBarriersPerBatch,
+	TEXT("Will limit the number of barriers sent per batch\n")
+	TEXT(" <=0: Do not limit (default)\n")
+	TEXT(" >0: Limit to the specified number\n"),
+	ECVF_Default
+);
+
+
 //
 // The following two functions are used when the RHI needs to do image layout transitions internally.
 // They are not used for the transitions requested through the public API (RHICreate/Begin/EndTransition)
@@ -1383,6 +1394,19 @@ void FTransitionProcessor<VkMemoryBarrier2, VkBufferMemoryBarrier2, VkImageMemor
 	// Sync2 does not submit batches, only one big clump at the end
 }
 
+template<typename BarrierType>
+static void SendBatchedBarriers(VkCommandBuffer CommandBuffer, VkDependencyInfo& BatchDependencyInfo, BarrierType*& BarrierPtr, uint32_t& BarrierCountRef, int32 TotalBarrierCount)
+{
+	for (int32 BatchStartIndex = 0; BatchStartIndex < TotalBarrierCount; BatchStartIndex += GVulkanMaxBarriersPerBatch)
+	{
+		BarrierCountRef = FMath::Min((TotalBarrierCount - BatchStartIndex), GVulkanMaxBarriersPerBatch);
+		VulkanRHI::vkCmdPipelineBarrier2KHR(CommandBuffer, &BatchDependencyInfo);
+		BarrierPtr += BarrierCountRef;
+	}
+	BarrierPtr = nullptr;
+	BarrierCountRef = 0;
+}
+
 template <>
 void FTransitionProcessor<VkMemoryBarrier2, VkBufferMemoryBarrier2, VkImageMemoryBarrier2>::FinishAll()
 {
@@ -1401,7 +1425,23 @@ void FTransitionProcessor<VkMemoryBarrier2, VkBufferMemoryBarrier2, VkImageMemor
 	DependencyInfo.pBufferMemoryBarriers = BufferBarriers.GetData();
 	DependencyInfo.imageMemoryBarrierCount = ImageBarriers.Num();
 	DependencyInfo.pImageMemoryBarriers = ImageBarriers.GetData();
-	VulkanRHI::vkCmdPipelineBarrier2KHR(Context.GetCommandBufferManager()->GetActiveCmdBuffer()->GetHandle(), &DependencyInfo);
+
+	if ((GVulkanMaxBarriersPerBatch <= 0) || ((MemoryBarriers.Num() + BufferBarriers.Num() + ImageBarriers.Num()) < GVulkanMaxBarriersPerBatch))
+	{
+		VulkanRHI::vkCmdPipelineBarrier2KHR(Context.GetCommandBufferManager()->GetActiveCmdBuffer()->GetHandle(), &DependencyInfo);
+	}
+	else
+	{
+		VkDependencyInfo BatchDependencyInfo = DependencyInfo;
+		BatchDependencyInfo.memoryBarrierCount = 0;
+		BatchDependencyInfo.bufferMemoryBarrierCount = 0;
+		BatchDependencyInfo.imageMemoryBarrierCount = 0;
+
+		VkCommandBuffer CommandBuffer = Context.GetCommandBufferManager()->GetActiveCmdBuffer()->GetHandle();
+		SendBatchedBarriers(CommandBuffer, BatchDependencyInfo, BatchDependencyInfo.pMemoryBarriers, BatchDependencyInfo.memoryBarrierCount, DependencyInfo.memoryBarrierCount);
+		SendBatchedBarriers(CommandBuffer, BatchDependencyInfo, BatchDependencyInfo.pBufferMemoryBarriers, BatchDependencyInfo.bufferMemoryBarrierCount, DependencyInfo.bufferMemoryBarrierCount);
+		SendBatchedBarriers(CommandBuffer, BatchDependencyInfo, BatchDependencyInfo.pImageMemoryBarriers, BatchDependencyInfo.imageMemoryBarrierCount, DependencyInfo.imageMemoryBarrierCount);
+	}
 }
 
 template <>
