Adding a method for retrieving shaders from FRHIGraphicsPipelineState.

#rnx

diff --git a/Engine/Source/Runtime/Apple/MetalRHI/Private/MetalGraphicsPipelineState.h b/Engine/Source/Runtime/Apple/MetalRHI/Private/MetalGraphicsPipelineState.h
--- a/Engine/Source/Runtime/Apple/MetalRHI/Private/MetalGraphicsPipelineState.h
+++ b/Engine/Source/Runtime/Apple/MetalRHI/Private/MetalGraphicsPipelineState.h
@@ -21,6 +21,38 @@ class FMetalGraphicsPipelineState : public FRHIGraphicsPipelineState
 public:
 	virtual ~FMetalGraphicsPipelineState();
 
+	FRHIGraphicsShader* GetShader(EShaderFrequency Frequency) const override
+	{
+		switch (Frequency)
+		{
+		case SF_Vertex: return VertexShader;
+		case SF_Pixel: return PixelShader;
+
+		case SF_Geometry:
+#if PLATFORM_SUPPORTS_GEOMETRY_SHADERS
+			return GeometryShader;
+#else
+			return nullptr;
+#endif
+
+		case SF_Mesh:
+#if PLATFORM_SUPPORTS_MESH_SHADERS
+			return MeshShader;
+#else
+			return nullptr;
+#endif
+
+		case SF_Amplification:
+#if PLATFORM_SUPPORTS_MESH_SHADERS
+			return AmplificationShader;
+#else
+			return nullptr;
+#endif
+
+		default: return nullptr;
+		}
+	}
+
 	FMetalShaderPipelinePtr GetPipeline();
 
 	/** Cached vertex structure */
diff --git a/Engine/Source/Runtime/D3D12RHI/Private/D3D12PipelineState.h b/Engine/Source/Runtime/D3D12RHI/Private/D3D12PipelineState.h
--- a/Engine/Source/Runtime/D3D12RHI/Private/D3D12PipelineState.h
+++ b/Engine/Source/Runtime/D3D12RHI/Private/D3D12PipelineState.h
@@ -384,6 +384,19 @@ struct FD3D12GraphicsPipelineState : public FRHIGraphicsPipelineState, FD3D12Pip
 	FORCEINLINE FD3D12AmplificationShader* GetAmplificationShader() const { return (FD3D12AmplificationShader*)PipelineStateInitializer.BoundShaderState.GetAmplificationShader(); }
 	FORCEINLINE FD3D12GeometryShader*      GetGeometryShader() const      { return (FD3D12GeometryShader*)PipelineStateInitializer.BoundShaderState.GetGeometryShader(); }
 
+	FRHIGraphicsShader* GetShader(EShaderFrequency Frequency) const override
+	{
+		switch (Frequency)
+		{
+		case SF_Vertex: return PipelineStateInitializer.BoundShaderState.GetVertexShader();
+		case SF_Mesh: return PipelineStateInitializer.BoundShaderState.GetMeshShader();
+		case SF_Amplification: return PipelineStateInitializer.BoundShaderState.GetAmplificationShader();
+		case SF_Pixel: return PipelineStateInitializer.BoundShaderState.GetPixelShader();
+		case SF_Geometry: return PipelineStateInitializer.BoundShaderState.GetGeometryShader();
+		default: return nullptr;
+		}
+	}
+
 	FGraphicsPipelineStateInitializer PipelineStateInitializer;
 	TStaticArray<uint16, MaxVertexElementCount> StreamStrides;
 	bool bShaderNeedsGlobalConstantBuffer[SF_NumStandardFrequencies];
diff --git a/Engine/Source/Runtime/RHI/Public/RHIResources.h b/Engine/Source/Runtime/RHI/Public/RHIResources.h
--- a/Engine/Source/Runtime/RHI/Public/RHIResources.h
+++ b/Engine/Source/Runtime/RHI/Public/RHIResources.h
@@ -1059,6 +1059,8 @@ public:
 	inline void SetSortKey(uint64 InSortKey) { SortKey = InSortKey; }
 	inline uint64 GetSortKey() const { return SortKey; }
 
+	virtual FRHIGraphicsShader* GetShader(EShaderFrequency Frequency) const = 0;
+
 private:
 	uint64 SortKey = 0;
 
@@ -4427,6 +4429,19 @@ public:
 	{
 	}
 
+	FRHIGraphicsShader* GetShader(EShaderFrequency Frequency) const override
+	{
+		switch (Frequency)
+		{
+		case SF_Vertex: return Initializer.BoundShaderState.GetVertexShader();
+		case SF_Mesh: return Initializer.BoundShaderState.GetMeshShader();
+		case SF_Amplification: return Initializer.BoundShaderState.GetAmplificationShader();
+		case SF_Pixel: return Initializer.BoundShaderState.GetPixelShader();
+		case SF_Geometry: return Initializer.BoundShaderState.GetGeometryShader();
+		default: return nullptr;
+		}
+	}
+
 	FGraphicsPipelineStateInitializer Initializer;
 };
 
diff --git a/Engine/Source/Runtime/VulkanRHI/Private/VulkanPendingState.h b/Engine/Source/Runtime/VulkanRHI/Private/VulkanPendingState.h
--- a/Engine/Source/Runtime/VulkanRHI/Private/VulkanPendingState.h
+++ b/Engine/Source/Runtime/VulkanRHI/Private/VulkanPendingState.h
@@ -196,7 +196,7 @@ public:
 
 	const FVulkanShader* GetCurrentShader(EShaderFrequency Frequency) const
 	{
-		return (CurrentPipeline ? CurrentPipeline->GetShader(Frequency) : nullptr);
+		return (CurrentPipeline ? CurrentPipeline->GetVulkanShader(Frequency) : nullptr);
 	}
 
 	void SetViewport(float MinX, float MinY, float MinZ, float MaxX, float MaxY, float MaxZ)
diff --git a/Engine/Source/Runtime/VulkanRHI/Private/VulkanPipeline.h b/Engine/Source/Runtime/VulkanRHI/Private/VulkanPipeline.h
--- a/Engine/Source/Runtime/VulkanRHI/Private/VulkanPipeline.h
+++ b/Engine/Source/Runtime/VulkanRHI/Private/VulkanPipeline.h
@@ -725,7 +725,7 @@ public:
 		return ShaderKeys[Stage];
 	}
 
-	inline const FVulkanShader* GetShader(EShaderFrequency Frequency) const
+	inline const FVulkanShader* GetVulkanShader(EShaderFrequency Frequency) const
 	{
 		ShaderStage::EStage Stage = ShaderStage::GetStageForFrequency(Frequency);
 		return VulkanShaders[Stage];
@@ -736,6 +736,19 @@ public:
 		return VulkanPipeline;
 	}
 
+	FRHIGraphicsShader* GetShader(EShaderFrequency Frequency) const override
+	{
+		switch (Frequency)
+		{
+		case SF_Vertex: return static_cast<FVulkanVertexShader*>(VulkanShaders[ShaderStage::Vertex]);
+		case SF_Pixel: return static_cast<FVulkanPixelShader*>(VulkanShaders[ShaderStage::Pixel]);
+#if VULKAN_SUPPORTS_GEOMETRY_SHADERS
+		case SF_Geometry: return static_cast<FVulkanGeometryShader*>(VulkanShaders[ShaderStage::Geometry]);
+#endif
+		default: return nullptr;
+		}
+	}
+
 	void DeleteVkPipeline(bool bImmediate);
 	void GetOrCreateShaderModules(TRefCountPtr<FVulkanShaderModule>(&ShaderModulesOUT)[ShaderStage::NumStages], FVulkanShader* const* Shaders);
 	FVulkanShader::FSpirvCode GetPatchedSpirvCode(FVulkanShader* Shader);
diff --git a/Engine/Source/Runtime/VulkanRHI/Private/VulkanRayTracing.cpp b/Engine/Source/Runtime/VulkanRHI/Private/VulkanRayTracing.cpp
--- a/Engine/Source/Runtime/VulkanRHI/Private/VulkanRayTracing.cpp
+++ b/Engine/Source/Runtime/VulkanRHI/Private/VulkanRayTracing.cpp
@@ -1681,7 +1681,7 @@ int32 FVulkanRayTracingPipelineState::GetShaderIndex(const FVulkanRayTracingShad
 	return INDEX_NONE;
 }
 
-const FVulkanRayTracingShader* FVulkanRayTracingPipelineState::GetShader(EShaderFrequency Frequency, int32 ShaderIndex) const
+const FVulkanRayTracingShader* FVulkanRayTracingPipelineState::GetVulkanShader(EShaderFrequency Frequency, int32 ShaderIndex) const
 {
 	return GetShaderData(Frequency).Shaders[ShaderIndex].GetReference();
 }
@@ -2173,7 +2173,7 @@ static void SetRayTracingHitGroup(
 #endif // DO_CHECK
 
 	const uint32 PrefixedSegmentIndex = SceneInitializer.SegmentPrefixSum[InstanceIndex];
-	const FVulkanRayTracingShader* Shader = Pipeline->GetShader(SF_RayHitGroup, HitGroupIndex);
+	const FVulkanRayTracingShader* Shader = Pipeline->GetVulkanShader(SF_RayHitGroup, HitGroupIndex);
 
 	const FVulkanRayTracingGeometry* Geometry = ResourceCast(SceneInitializer.PerInstanceGeometries[InstanceIndex]);
 	FVulkanHitGroupSystemParameters SystemParameters = Geometry->HitGroupSystemParameters[SegmentIndex];
@@ -2196,7 +2196,7 @@ static void SetGenericSystemParameters(
 	FVulkanRayTracingPipelineState* Pipeline = ResourceCast(InPipeline);
 	const uint32 WorkerIndex = 0;
 	FVulkanRayTracingShaderTable* ShaderTable = Scene->FindOrCreateShaderTable(Pipeline);
-	const FVulkanRayTracingShader* Shader = Pipeline->GetShader(ShaderFrequency, ShaderIndexInPipeline);
+	const FVulkanRayTracingShader* Shader = Pipeline->GetVulkanShader(ShaderFrequency, ShaderIndexInPipeline);
 
 	FVulkanHitGroupSystemParameters SystemParameters;
 	FMemory::Memzero(SystemParameters);
diff --git a/Engine/Source/Runtime/VulkanRHI/Private/VulkanRayTracing.h b/Engine/Source/Runtime/VulkanRHI/Private/VulkanRayTracing.h
--- a/Engine/Source/Runtime/VulkanRHI/Private/VulkanRayTracing.h
+++ b/Engine/Source/Runtime/VulkanRHI/Private/VulkanRayTracing.h
@@ -271,7 +271,7 @@ public:
 	}
 
 	int32 GetShaderIndex(const FVulkanRayTracingShader* Shader) const;
-	const FVulkanRayTracingShader* GetShader(EShaderFrequency Frequency, int32 ShaderIndex) const;
+	const FVulkanRayTracingShader* GetVulkanShader(EShaderFrequency Frequency, int32 ShaderIndex) const;
 	const TArray<uint8>& GetShaderHandles(EShaderFrequency Frequency) const;
 
 private:
