VulkanRHI: Update Vulkan RT to support the latest shader table changes.  Keep references on shader tables and pipelines.

#rnx

diff --git a/Engine/Source/Runtime/VulkanRHI/Private/VulkanRayTracing.cpp b/Engine/Source/Runtime/VulkanRHI/Private/VulkanRayTracing.cpp
--- a/Engine/Source/Runtime/VulkanRHI/Private/VulkanRayTracing.cpp
+++ b/Engine/Source/Runtime/VulkanRHI/Private/VulkanRayTracing.cpp
@@ -916,6 +916,11 @@ FVulkanRayTracingShaderTable::FVulkanRayTracingShaderTable(FVulkanDevice* Device
 }
 
 FVulkanRayTracingShaderTable::~FVulkanRayTracingShaderTable()
+{
+	ReleaseLocalBuffers();
+}
+
+void FVulkanRayTracingShaderTable::ReleaseLocalBuffers()
 {
 	ReleaseLocalBuffer(Device, Raygen);
 	ReleaseLocalBuffer(Device, Miss);
@@ -971,6 +976,17 @@ void FVulkanRayTracingShaderTable::Init(const FVulkanRayTracingPipelineState* Pi
 	}
 }
 
+void FVulkanRayTracingShaderTable::SetRayTracingPipelineState(const FVulkanRayTracingPipelineState* InPipeline)
+{
+	checkf(RayTracingPipelineState == nullptr || RayTracingPipelineState == InPipeline, TEXT("FRHIShaderBindingTable can't currently be used with multiple different RTPSOs"));
+
+	if (RayTracingPipelineState == nullptr)
+	{
+		Init(InPipeline);
+	}
+	RayTracingPipelineState = InPipeline;
+}
+
 FVulkanRayTracingShaderTable::FVulkanShaderTableAllocation& FVulkanRayTracingShaderTable::GetAlloc(EShaderFrequency Frequency)
 {
 	switch (Frequency)
@@ -1098,10 +1114,10 @@ FRHIShaderBindingTable* FVulkanRayTracingScene::FindOrCreateShaderBindingTable(c
 
 	// Find existing table
 	{
-		FVulkanRayTracingShaderTable* const* FoundShaderTable = ShaderTables.Find(Pipeline);
+		TRefCountPtr<FVulkanRayTracingShaderTable>* FoundShaderTable = ShaderTables.Find(Pipeline);
 		if (FoundShaderTable)
 		{
-			return *FoundShaderTable;
+			return FoundShaderTable->GetReference();
 		}
 	}
 
@@ -1115,7 +1131,7 @@ FRHIShaderBindingTable* FVulkanRayTracingScene::FindOrCreateShaderBindingTable(c
 
 	// Create new table
 	FVulkanRayTracingShaderTable* CreatedShaderTable = new FVulkanRayTracingShaderTable(Device, MoveTemp(SBTInitializer));
-	CreatedShaderTable->Init(Pipeline);
+	CreatedShaderTable->SetRayTracingPipelineState(Pipeline);
 
 	ShaderTables.Add(Pipeline, CreatedShaderTable);
 
@@ -1178,9 +1194,15 @@ FShaderBindingTableRHIRef FVulkanDynamicRHI::RHICreateShaderBindingTable(const F
 	return new FVulkanRayTracingShaderTable(GetDevice(), Initializer);
 }
 
-void FVulkanCommandListContext::RHIClearRayTracingBindings(FRHIRayTracingScene* Scene)
+void FVulkanCommandListContext::RHIClearRayTracingBindings(FRHIRayTracingScene* InScene)
 {
-	 // TODO
+	FVulkanRayTracingScene* Scene = ResourceCast(InScene);
+	check(Scene);
+
+	for (auto& Table : Scene->ShaderTables)
+	{
+		Table.Value->ReleaseLocalBuffers();
+	}
 }
 
 void FVulkanCommandListContext::RHICommitRayTracingBindings(FRHIRayTracingScene* InScene)
@@ -1189,9 +1211,10 @@ void FVulkanCommandListContext::RHICommitRayTracingBindings(FRHIRayTracingScene*
 	Scene->CommitShaderTables(*this);
 }
 
-void FVulkanCommandListContext::RHIClearShaderBindingTable(FRHIShaderBindingTable* SBT)
+void FVulkanCommandListContext::RHIClearShaderBindingTable(FRHIShaderBindingTable* InSBT)
 {
-	// TODO
+	FVulkanRayTracingShaderTable* SBT = ResourceCast(InSBT);
+	SBT->ReleaseLocalBuffers();
 }
 
 void FVulkanCommandListContext::RHICommitShaderBindingTable(FRHIShaderBindingTable* InSBT)
@@ -2246,6 +2269,8 @@ void FVulkanCommandListContext::RHISetBindingsOnShaderBindingTable(FRHIShaderBin
 
 	checkf(Scene->IsBuilt(), TEXT("Ray tracing scene must be built before any shaders can be bound to it. Make sure that RHIBuildAccelerationStructure() command has been executed."));
 
+	ShaderTable->SetRayTracingPipelineState(Pipeline);
+
 	FGraphEventArray TaskList;
 
 	const uint32 NumWorkerThreads = FTaskGraphInterface::Get().GetNumWorkerThreads();
diff --git a/Engine/Source/Runtime/VulkanRHI/Private/VulkanRayTracing.h b/Engine/Source/Runtime/VulkanRHI/Private/VulkanRayTracing.h
--- a/Engine/Source/Runtime/VulkanRHI/Private/VulkanRayTracing.h
+++ b/Engine/Source/Runtime/VulkanRHI/Private/VulkanRayTracing.h
@@ -64,6 +64,7 @@ public:
 	~FVulkanRayTracingShaderTable();
 
 	void Init(const FVulkanRayTracingPipelineState* Pipeline);
+	void ReleaseLocalBuffers();
 
 	const VkStridedDeviceAddressRegionKHR* GetRegion(EShaderFrequency Frequency);
 
@@ -91,6 +92,8 @@ public:
 		return ReferencedUniformBuffers;
 	}
 
+	void SetRayTracingPipelineState(const FVulkanRayTracingPipelineState* InPipeline);
+
 	// Ray tracing shader bindings can be processed in parallel.
 	// Each concurrent worker gets its own dedicated descriptor cache instance to avoid contention or locking.
 	// Scaling beyond 5 total threads does not yield any speedup in practice (RHI thread + 4 parallel workers).
@@ -130,6 +133,8 @@ private:
 
 	TArray<TRefCountPtr<FRHIUniformBuffer>> ReferencedUniformBuffers;
 
+	TRefCountPtr<const FVulkanRayTracingPipelineState> RayTracingPipelineState; // TODO: Is holding reference necessary?
+
 	// Convenience
 	const uint32 HandleSize;
 	const uint32 HandleSizeAligned;
@@ -176,8 +181,9 @@ public:
 
 class FVulkanRayTracingScene : public FRHIRayTracingScene, public VulkanRHI::FDeviceChild
 {
-public:
+	friend FVulkanCommandListContext;
 
+public:
 	FVulkanRayTracingScene(FRayTracingSceneInitializer2 Initializer, FVulkanDevice* InDevice);
 	~FVulkanRayTracingScene();
 
@@ -246,7 +252,7 @@ private:
 	TRefCountPtr<FVulkanResourceMultiBuffer> PerInstanceGeometryParameterBuffer;
 	FShaderResourceViewRHIRef PerInstanceGeometryParameterSRV;
 	
-	TMap<const FVulkanRayTracingPipelineState*, FVulkanRayTracingShaderTable*> ShaderTables;
+	TMap<const FVulkanRayTracingPipelineState*, TRefCountPtr<FVulkanRayTracingShaderTable>> ShaderTables;
 
 	void BuildPerInstanceGeometryParameterBuffer(FVulkanCommandListContext& CommandContext);
 
