VulkanRHI: make sure buffer creation in ray tracing properly initializes the tracking in RHI validation.

#rnx

diff --git a/Engine/Source/Runtime/VulkanRHI/Private/VulkanRayTracing.cpp b/Engine/Source/Runtime/VulkanRHI/Private/VulkanRayTracing.cpp
--- a/Engine/Source/Runtime/VulkanRHI/Private/VulkanRayTracing.cpp
+++ b/Engine/Source/Runtime/VulkanRHI/Private/VulkanRayTracing.cpp
@@ -692,9 +692,8 @@ FVulkanRayTracingScene::FVulkanRayTracingScene(FRayTracingSceneInitializer2 InIn
 
 	const uint32 ParameterBufferSize = FMath::Max<uint32>(1, Initializer.NumTotalSegments) * sizeof(FVulkanRayTracingGeometryParameters);
 	FRHIResourceCreateInfo ParameterBufferCreateInfo(TEXT("RayTracingSceneMetadata"));
-	PerInstanceGeometryParameterBuffer = new FVulkanResourceMultiBuffer(Device,
-		FRHIBufferDesc(ParameterBufferSize, sizeof(FVulkanRayTracingGeometryParameters), BUF_StructuredBuffer | BUF_ShaderResource),
-		ParameterBufferCreateInfo);
+	// Use FRHICommandListExecutor::GetImmediateCommandList as a temporary patch for 5.4 to avoid issues with RHI Validation
+	PerInstanceGeometryParameterBuffer = ResourceCast(FRHICommandListExecutor::GetImmediateCommandList().CreateBuffer(ParameterBufferSize, BUF_StructuredBuffer | BUF_ShaderResource, sizeof(FVulkanRayTracingGeometryParameters), ERHIAccess::SRVMask, ParameterBufferCreateInfo).GetReference());
 }
 
 FVulkanRayTracingScene::~FVulkanRayTracingScene()
