VulkanRHI: Keep references on geometry going through compaction (essentially disables cancellation of compaction requests until lifetime issues are resolved).

#rnx

diff --git a/Engine/Source/Runtime/VulkanRHI/Private/VulkanCommandBuffer.cpp b/Engine/Source/Runtime/VulkanRHI/Private/VulkanCommandBuffer.cpp
--- a/Engine/Source/Runtime/VulkanRHI/Private/VulkanCommandBuffer.cpp
+++ b/Engine/Source/Runtime/VulkanRHI/Private/VulkanCommandBuffer.cpp
@@ -8,6 +8,7 @@
 #include "VulkanContext.h"
 #include "VulkanDescriptorSets.h"
 #include "VulkanMemory.h"
+#include "VulkanRayTracing.h"
 
 static int32 GUseSingleQueue = 0;
 static FAutoConsoleVariableRef CVarVulkanUseSingleQueue(
@@ -875,6 +876,10 @@ void FVulkanCommandBufferPool::FreeUnusedCmdBuffers(FVulkanQueue* InQueue, bool
 			(CurrentTime - CmdBuffer->SubmittedTime) > CMD_BUFFER_TIME_TO_WAIT_BEFORE_DELETING)
 		{
 			DeferredDeletionQueue.OnCmdBufferDeleted(CmdBuffer);
+			if (Device->GetRayTracingCompactionRequestHandler())
+			{
+				Device->GetRayTracingCompactionRequestHandler()->OnCmdBufferDeleted(CmdBuffer);
+			}
 
 			CmdBuffer->FreeMemory();
 			CmdBuffers.RemoveAtSwap(Index, EAllowShrinking::No);
diff --git a/Engine/Source/Runtime/VulkanRHI/Private/VulkanRayTracing.cpp b/Engine/Source/Runtime/VulkanRHI/Private/VulkanRayTracing.cpp
--- a/Engine/Source/Runtime/VulkanRHI/Private/VulkanRayTracing.cpp
+++ b/Engine/Source/Runtime/VulkanRHI/Private/VulkanRayTracing.cpp
@@ -1758,10 +1758,13 @@ void FVulkanRayTracingCompactedSizeQueryPool::Reset(FVulkanCmdBuffer* InCmdBuffe
 
 bool FVulkanRayTracingCompactedSizeQueryPool::TryGetResults(uint32 NumResults)
 {
-	if (CmdBuffer == nullptr) return false;
+	if (CmdBuffer == nullptr)
+	{
+		return false;
+	}
 
 	const uint64 FenceCurrentSignaledCounter = CmdBuffer->GetFenceSignaledCounter();
-	if (FenceSignaledCounter > FenceCurrentSignaledCounter)
+	if (FenceSignaledCounter >= FenceCurrentSignaledCounter)
 	{
 		return false;
 	}
@@ -1778,6 +1781,9 @@ FVulkanRayTracingCompactionRequestHandler::FVulkanRayTracingCompactionRequestHan
 	: VulkanRHI::FDeviceChild(InDevice)
 {
 	QueryPool = new FVulkanRayTracingCompactedSizeQueryPool(InDevice, GVulkanRayTracingMaxBatchedCompaction);
+
+	ActiveRequests.Reserve(GVulkanRayTracingMaxBatchedCompaction);
+	ActiveBLASes.Reserve(GVulkanRayTracingMaxBatchedCompaction);
 }
 
 void FVulkanRayTracingCompactionRequestHandler::RequestCompact(FVulkanRayTracingGeometry* InRTGeometry)
@@ -1822,6 +1828,7 @@ void FVulkanRayTracingCompactionRequestHandler::Update(FVulkanCommandListContext
 	LLM_SCOPE_BYNAME(TEXT("FVulkanRT/Compaction"));
 	FScopeLock Lock(&CS);
 
+	// If we have an active batch, wait on those queries and launch compaction when the complete
 	if (ActiveBLASes.Num() > 0)
 	{		
 		FVulkanCommandBufferManager& CommandBufferManager = *InCommandContext.GetCommandBufferManager();
@@ -1840,10 +1847,34 @@ void FVulkanRayTracingCompactionRequestHandler::Update(FVulkanCommandListContext
 
 			QueryPool->Reset(CmdBuffer);
 
-			ActiveRequests.Empty(ActiveRequests.Num());
-			ActiveBLASes.Empty(ActiveBLASes.Num());
+			ActiveBLASes.Empty(GVulkanRayTracingMaxBatchedCompaction);
+
+			ActiveRequestsCmdBuffer = CmdBuffer;
+			ActiveRequestsFenceCounter = CmdBuffer->GetFenceSignaledCounter();
 		}
+
+		// Only one active batch at a time (otherwise track the offset for when we launch queries)
+		return;
 	}
+	// If we have an active batch, wait until the compaction went through to launch another batch
+	else if (ActiveRequests.Num() > 0)
+	{
+		if (ActiveRequestsCmdBuffer)
+		{
+			if (ActiveRequestsFenceCounter >= ActiveRequestsCmdBuffer->GetFenceSignaledCounter())
+			{
+				return;
+			}
+
+			ActiveRequestsCmdBuffer = nullptr;
+			ActiveRequestsFenceCounter = MAX_uint64;
+		}
+
+		ActiveRequests.Empty(GVulkanRayTracingMaxBatchedCompaction);
+	}
+
+	check(ActiveBLASes.Num() == 0);
+	check(ActiveRequests.Num() == 0);
 
 	// build a new set of build requests to extract the build data	
 	for (FVulkanRayTracingGeometry* RTGeometry : PendingRequests)
@@ -1879,13 +1910,26 @@ void FVulkanRayTracingCompactionRequestHandler::Update(FVulkanCommandListContext
 			0
 		);
 
+		QueryPool->EndBatch(CmdBuffer);
+
 		CommandBufferManager.SubmitActiveCmdBuffer();
 		CommandBufferManager.PrepareForNewActiveCommandBuffer();
-
-		QueryPool->EndBatch(CmdBuffer);
 	}
 }
 
+void FVulkanRayTracingCompactionRequestHandler::OnCmdBufferDeleted(FVulkanCmdBuffer* DeletedCmdBuffer)
+{
+	FScopeLock Lock(&CS);
+
+	// If we are deleting it, it means all its commands went through and we don't need to wait on it anymore
+	if (DeletedCmdBuffer == ActiveRequestsCmdBuffer)
+	{
+		ActiveRequestsCmdBuffer = nullptr;
+		ActiveRequestsFenceCounter = MAX_uint64;
+	}
+
+	check(DeletedCmdBuffer != QueryPool->CmdBuffer);
+}
 
 
 
diff --git a/Engine/Source/Runtime/VulkanRHI/Private/VulkanRayTracing.h b/Engine/Source/Runtime/VulkanRHI/Private/VulkanRayTracing.h
--- a/Engine/Source/Runtime/VulkanRHI/Private/VulkanRayTracing.h
+++ b/Engine/Source/Runtime/VulkanRHI/Private/VulkanRayTracing.h
@@ -335,13 +335,19 @@ public:
 
 	void Update(FVulkanCommandListContext& InCommandContext);
 
+	void OnCmdBufferDeleted(FVulkanCmdBuffer* DeletedCmdBuffer);
+
 private:
 
 	FCriticalSection CS;
 	TArray<FVulkanRayTracingGeometry*> PendingRequests;
-	TArray<FVulkanRayTracingGeometry*> ActiveRequests;
 	TArray<VkAccelerationStructureKHR> ActiveBLASes;
 
+	// Keep references on FVulkanRayTracingGeometry until lifetime issue is found (this prevents cancellation)
+	TArray<TRefCountPtr<FVulkanRayTracingGeometry>> ActiveRequests;
+	FVulkanCmdBuffer* ActiveRequestsCmdBuffer = nullptr;
+	uint64 ActiveRequestsFenceCounter = MAX_uint64;
+
 	FVulkanRayTracingCompactedSizeQueryPool* QueryPool = nullptr;
 };
 
