VulkanRHI: Clear pointer on validation cache when it is destroyed.

#rnx

diff --git a/Engine/Source/Runtime/VulkanRHI/Private/VulkanDevice.cpp b/Engine/Source/Runtime/VulkanRHI/Private/VulkanDevice.cpp
--- a/Engine/Source/Runtime/VulkanRHI/Private/VulkanDevice.cpp
+++ b/Engine/Source/Runtime/VulkanRHI/Private/VulkanDevice.cpp
@@ -1287,6 +1287,7 @@ void FVulkanDevice::Destroy()
 		if (vkDestroyValidationCache)
 		{
 			vkDestroyValidationCache(Device, ValidationCache, VULKAN_CPU_ALLOCATOR);
+			ValidationCache = VK_NULL_HANDLE;
 		}
 	}
 #endif
