Remove commented-out logging in RHIBeginScene/RHIEndScene of Vulkan RHI
 - Trivial cleanup in preparation to remove both these functions in a future CL.

diff --git a/Engine/Source/Runtime/VulkanRHI/Private/VulkanRHI.cpp b/Engine/Source/Runtime/VulkanRHI/Private/VulkanRHI.cpp
--- a/Engine/Source/Runtime/VulkanRHI/Private/VulkanRHI.cpp
+++ b/Engine/Source/Runtime/VulkanRHI/Private/VulkanRHI.cpp
@@ -1112,12 +1112,10 @@ void FVulkanCommandListContext::RHIBeginFrame()
 
 void FVulkanCommandListContext::RHIBeginScene()
 {
-	//FRCLog::Printf(FString::Printf(TEXT("FVulkanCommandListContext::RHIBeginScene()")));
 }
 
 void FVulkanCommandListContext::RHIEndScene()
 {
-	//FRCLog::Printf(FString::Printf(TEXT("FVulkanCommandListContext::RHIEndScene()")));
 }
 
 void FVulkanCommandListContext::RHIBeginDrawingViewport(FRHIViewport* ViewportRHI, FRHITexture* RenderTargetRHI)
