VulkanRHI: Lock descriptor sets on stages as generated by the shader compiler to remove the need to overwrite them at runtime.  Also removes the need for FSetInfo.

#rnx

diff --git a/Engine/Source/Runtime/VulkanRHI/Private/VulkanDescriptorSets.h b/Engine/Source/Runtime/VulkanRHI/Private/VulkanDescriptorSets.h
--- a/Engine/Source/Runtime/VulkanRHI/Private/VulkanDescriptorSets.h
+++ b/Engine/Source/Runtime/VulkanRHI/Private/VulkanDescriptorSets.h
@@ -51,6 +51,7 @@ struct FDescriptorSetRemappingInfo
 		uint16	NewDescriptorSet = UINT16_MAX;
 		uint16	NewBindingIndex = UINT16_MAX;
 	};
+
 	struct FUBRemappingInfo
 	{
 		// Remapping is only valid if there is constant data
@@ -60,22 +61,16 @@ struct FDescriptorSetRemappingInfo
 		const bool				bPadding = false;	// padding is need on memcmp/MemCrc to make sure mem align
 	};
 
-	struct FSetInfo
-	{
-		TArray<VkDescriptorType>	Types;
-		uint16						NumImageInfos = 0;
-		uint16						NumBufferInfos = 0;
-		uint8						NumAccelerationStructures = 0;
-	};
-	TArray<FSetInfo>	SetInfos;
-
 	struct FStageInfo
 	{
 		TArray<FRemappingInfo>		Globals;
 		TArray<FUBRemappingInfo>	UniformBuffers;
+		TArray<VkDescriptorType>	Types;
 		TArray<uint16>				PackedUBBindingIndices;
 		uint16						PackedUBDescriptorSet = UINT16_MAX;
-		uint16						Pad0 = 0;
+		uint16						NumImageInfos = 0;
+		uint16						NumBufferInfos = 0;
+		uint16						NumAccelerationStructures = 0;
 
 		inline bool IsEmpty() const
 		{
@@ -108,36 +103,22 @@ struct FDescriptorSetRemappingInfo
 			return false;
 		}
 
-		if (SetInfos.Num() != In.SetInfos.Num())
-		{
-			return false;
-		}
-
 		if (FMemory::Memcmp(InputAttachmentData.GetData(), In.InputAttachmentData.GetData(), sizeof(FInputAttachmentData) * InputAttachmentData.Num()))
 		{
 			return false;
 		}
 
-		for (int32 SetInfosIndex = 0; SetInfosIndex < SetInfos.Num(); ++SetInfosIndex)
-		{
-			int32 SetInfosNums = SetInfos[SetInfosIndex].Types.Num();
-			if (SetInfos[SetInfosIndex].NumBufferInfos != In.SetInfos[SetInfosIndex].NumBufferInfos ||
-				SetInfos[SetInfosIndex].NumImageInfos != In.SetInfos[SetInfosIndex].NumImageInfos ||
-				SetInfos[SetInfosIndex].NumAccelerationStructures != In.SetInfos[SetInfosIndex].NumAccelerationStructures ||
-				SetInfosNums != In.SetInfos[SetInfosIndex].Types.Num() ||
-				(SetInfosNums != 0 && FMemory::Memcmp(SetInfos[SetInfosIndex].Types.GetData(), In.SetInfos[SetInfosIndex].Types.GetData(), sizeof(VkDescriptorType) * SetInfosNums)))
-			{
-				return false;
-			}
-		}
-
 		for (uint32 StageInfosIndex = 0; StageInfosIndex < ShaderStage::NumStages; ++StageInfosIndex)
 		{
 			if (StageInfos[StageInfosIndex].PackedUBDescriptorSet != In.StageInfos[StageInfosIndex].PackedUBDescriptorSet ||
-				StageInfos[StageInfosIndex].Pad0 != In.StageInfos[StageInfosIndex].Pad0 ||
+				StageInfos[StageInfosIndex].NumBufferInfos != In.StageInfos[StageInfosIndex].NumBufferInfos ||
+				StageInfos[StageInfosIndex].NumImageInfos != In.StageInfos[StageInfosIndex].NumImageInfos ||
+				StageInfos[StageInfosIndex].NumAccelerationStructures != In.StageInfos[StageInfosIndex].NumAccelerationStructures ||
 				StageInfos[StageInfosIndex].Globals.Num() != In.StageInfos[StageInfosIndex].Globals.Num() ||
 				StageInfos[StageInfosIndex].PackedUBBindingIndices.Num() != In.StageInfos[StageInfosIndex].PackedUBBindingIndices.Num() ||
 				StageInfos[StageInfosIndex].UniformBuffers.Num() != In.StageInfos[StageInfosIndex].UniformBuffers.Num() ||
+				StageInfos[StageInfosIndex].Types.Num() != In.StageInfos[StageInfosIndex].Types.Num() ||
+				FMemory::Memcmp(StageInfos[StageInfosIndex].Types.GetData(), In.StageInfos[StageInfosIndex].Types.GetData(), sizeof(VkDescriptorType) * StageInfos[StageInfosIndex].Types.Num()) ||
 				FMemory::Memcmp(StageInfos[StageInfosIndex].Globals.GetData(), In.StageInfos[StageInfosIndex].Globals.GetData(), sizeof(FRemappingInfo) * StageInfos[StageInfosIndex].Globals.Num()) ||
 				FMemory::Memcmp(StageInfos[StageInfosIndex].PackedUBBindingIndices.GetData(), In.StageInfos[StageInfosIndex].PackedUBBindingIndices.GetData(), sizeof(uint16) * StageInfos[StageInfosIndex].PackedUBBindingIndices.Num()) ||
 				FMemory::Memcmp(StageInfos[StageInfosIndex].UniformBuffers.GetData(), In.StageInfos[StageInfosIndex].UniformBuffers.GetData(), sizeof(FUBRemappingInfo) * StageInfos[StageInfosIndex].UniformBuffers.Num()))
@@ -156,29 +137,32 @@ struct FDescriptorSetRemappingInfo
 
 	inline bool IsEmpty() const
 	{
-		if (SetInfos.Num() == 0)
+		for (int32 Index = 0; Index < ShaderStage::NumStages; ++Index)
 		{
-			for (int32 Index = 0; Index < ShaderStage::NumStages; ++Index)
+			const FStageInfo& StageInfo = StageInfos[Index];
+
+			if (StageInfo.Types.Num() || StageInfo.NumBufferInfos || StageInfo.NumImageInfos || StageInfo.NumAccelerationStructures)
 			{
-				if (!StageInfos[Index].IsEmpty())
-				{
-					return false;
-				}
+				return false;
 			}
 
-			return true;
+			if (!StageInfo.IsEmpty())
+			{
+				return false;
+			}
 		}
-
-		return false;
+		return true;
 	}
 
 	uint32 AddGlobal(uint32 Stage, int32 GlobalIndex, uint32 NewDescriptorSet, VkDescriptorType InType)
 	{
-		const uint32 NewBindingIndex = SetInfos[NewDescriptorSet].Types.Add(InType);
+		FStageInfo& StageInfo = StageInfos[Stage];
+
+		const uint32 NewBindingIndex = StageInfo.Types.Add(InType);
 
-		int32 RemappingIndex = StageInfos[Stage].Globals.AddDefaulted();
+		const int32 RemappingIndex = StageInfo.Globals.AddDefaulted();
 		check(RemappingIndex == GlobalIndex);
-		FDescriptorSetRemappingInfo::FRemappingInfo& Remapping = StageInfos[Stage].Globals[RemappingIndex];
+		FDescriptorSetRemappingInfo::FRemappingInfo& Remapping = StageInfo.Globals[RemappingIndex];
 		Remapping.NewDescriptorSet = NewDescriptorSet;
 		Remapping.NewBindingIndex = NewBindingIndex;
 
@@ -187,51 +171,49 @@ struct FDescriptorSetRemappingInfo
 
 	uint32 AddPackedUB(uint32 Stage, int32 PackUBIndex, uint32 NewDescriptorSet, VkDescriptorType InType)
 	{
-		uint32 NewBindingIndex = SetInfos[NewDescriptorSet].Types.Add(InType);
-		if (StageInfos[Stage].PackedUBDescriptorSet == UINT16_MAX)
+		FStageInfo& StageInfo = StageInfos[Stage];
+
+		const uint32 NewBindingIndex = StageInfo.Types.Add(InType);
+		if (StageInfo.PackedUBDescriptorSet == UINT16_MAX)
 		{
-			StageInfos[Stage].PackedUBDescriptorSet = NewDescriptorSet;
+			StageInfo.PackedUBDescriptorSet = NewDescriptorSet;
 		}
 		else
 		{
-			ensure(StageInfos[Stage].PackedUBDescriptorSet == NewDescriptorSet);
+			ensure(StageInfo.PackedUBDescriptorSet == NewDescriptorSet);
 		}
-		int32 RemappingIndex = StageInfos[Stage].PackedUBBindingIndices.Add(NewBindingIndex);
+		const int32 RemappingIndex = StageInfo.PackedUBBindingIndices.Add(NewBindingIndex);
 		check(RemappingIndex == PackUBIndex);
 
 		return NewBindingIndex;
 	}
 
-	FDescriptorSetRemappingInfo::FUBRemappingInfo AddUBWithData(uint32 Stage, int32 UniformBufferIndex, uint32 NewDescriptorSet, VkDescriptorType InType, uint32& OutNewBindingIndex)
+	void AddUBWithData(uint32 Stage, int32 UniformBufferIndex, uint32 NewDescriptorSet, VkDescriptorType InType, uint32& OutNewBindingIndex)
 	{
-		OutNewBindingIndex = SetInfos[NewDescriptorSet].Types.Add(InType);
+		FStageInfo& StageInfo = StageInfos[Stage];
+
+		OutNewBindingIndex = StageInfo.Types.Add(InType);
 
-		int32 UBRemappingIndex = StageInfos[Stage].UniformBuffers.AddDefaulted();
+		const int32 UBRemappingIndex = StageInfo.UniformBuffers.AddDefaulted();
 		check(UBRemappingIndex == UniformBufferIndex);
-		FDescriptorSetRemappingInfo::FUBRemappingInfo& UBRemapping = StageInfos[Stage].UniformBuffers[UBRemappingIndex];
+		FDescriptorSetRemappingInfo::FUBRemappingInfo& UBRemapping = StageInfo.UniformBuffers[UBRemappingIndex];
 		UBRemapping.bHasConstantData = true;
 		UBRemapping.Remapping.NewDescriptorSet = NewDescriptorSet;
 		UBRemapping.Remapping.NewBindingIndex = OutNewBindingIndex;
-
-		return UBRemapping;
 	}
 
 	void AddRedundantUB(uint32 Stage, int32 UniformBufferIndex, const FDescriptorSetRemappingInfo::FUBRemappingInfo* InExistingUBInfo)
 	{
-		int32 UBRemappingIndex = StageInfos[Stage].UniformBuffers.AddDefaulted();
+		const int32 UBRemappingIndex = StageInfos[Stage].UniformBuffers.AddDefaulted();
 		check(UBRemappingIndex == UniformBufferIndex);
 		FDescriptorSetRemappingInfo::FUBRemappingInfo& UBRemapping = StageInfos[Stage].UniformBuffers[UBRemappingIndex];
-		//UBRemapping.bIsRedundant = true;
 		UBRemapping.bHasConstantData = InExistingUBInfo->bHasConstantData;
-/*
-		UBRemapping.EntriesRemappingInfo = InExistingUBInfo->EntriesRemappingInfo;
-*/
 		UBRemapping.Remapping = InExistingUBInfo->Remapping;
 	}
 
 	void AddUBResourceOnly(uint32 Stage, int32 UniformBufferIndex)
 	{
-		int32 UBRemappingIndex = StageInfos[Stage].UniformBuffers.AddDefaulted();
+		const int32 UBRemappingIndex = StageInfos[Stage].UniformBuffers.AddDefaulted();
 		check(UBRemappingIndex == UniformBufferIndex);
 		FDescriptorSetRemappingInfo::FUBRemappingInfo& UBRemapping = StageInfos[Stage].UniformBuffers[UBRemappingIndex];
 		UBRemapping.bHasConstantData = false;
@@ -720,12 +702,13 @@ public:
 	inline const TArray<FDescriptorSetRemappingInfo::FRemappingInfo>& GetGlobalRemappingInfo() const
 	{
 		//OutDescriptorSet = RemappingUBInfos[Stage][ParameterIndex].Remapping.NewDescriptorSet;
-		return RemappingInfo->StageInfos[0].Globals;
+		return RemappingInfo->StageInfos[ShaderStage::Compute].Globals;
 	}
 	
 	inline VkDescriptorType GetDescriptorType(uint8 DescriptorSet, int32 DescriptorIndex) const
 	{
-		return RemappingInfo->SetInfos[DescriptorSet].Types[DescriptorIndex];
+		check(DescriptorSet == ShaderStage::Compute);
+		return RemappingInfo->StageInfos[DescriptorSet].Types[DescriptorIndex];
 	}
 
 	inline bool IsInitialized() const
@@ -785,7 +768,7 @@ public:
 
 	inline VkDescriptorType GetDescriptorType(uint8 DescriptorSet, int32 DescriptorIndex) const
 	{
-		return RemappingInfo->SetInfos[DescriptorSet].Types[DescriptorIndex];
+		return RemappingInfo->StageInfos[DescriptorSet].Types[DescriptorIndex];
 	}
 
 	inline bool IsInitialized() const
diff --git a/Engine/Source/Runtime/VulkanRHI/Private/VulkanPipelineState.cpp b/Engine/Source/Runtime/VulkanRHI/Private/VulkanPipelineState.cpp
--- a/Engine/Source/Runtime/VulkanRHI/Private/VulkanPipelineState.cpp
+++ b/Engine/Source/Runtime/VulkanRHI/Private/VulkanPipelineState.cpp
@@ -38,7 +38,7 @@ static bool ShouldAlwaysWriteDescriptors()
 }
 
 FVulkanComputePipelineDescriptorState::FVulkanComputePipelineDescriptorState(FVulkanDevice* InDevice, FVulkanComputePipeline* InComputePipeline)
-	: FVulkanCommonPipelineDescriptorState(InDevice)
+	: FVulkanCommonPipelineDescriptorState(InDevice, 1)
 	, PackedUniformBuffersMask(0)
 	, PackedUniformBuffersDirty(0)
 	, ComputePipeline(InComputePipeline)
@@ -62,32 +62,30 @@ FVulkanComputePipelineDescriptorState::FVulkanComputePipelineDescriptorState(FVu
 void FVulkanCommonPipelineDescriptorState::CreateDescriptorWriteInfos()
 {
 	check(DSWriteContainer.DescriptorWrites.Num() == 0);
+	check(UsedSetsMask <= (uint32)(((uint32)1 << MaxNumSets) - 1));
 
-	const int32 NumSets = DescriptorSetsLayout->RemappingInfo.SetInfos.Num();
-	check(UsedSetsMask <= (uint32)(((uint32)1 << NumSets) - 1));
-
-	for (int32 Set = 0; Set < NumSets; ++Set)
+	for (uint32 Set = 0; Set < MaxNumSets; ++Set)
 	{
-		const FDescriptorSetRemappingInfo::FSetInfo& SetInfo = DescriptorSetsLayout->RemappingInfo.SetInfos[Set];
+		const FDescriptorSetRemappingInfo::FStageInfo& StageInfo = DescriptorSetsLayout->RemappingInfo.StageInfos[Set];
 		
 		if (UseVulkanDescriptorCache())
 		{
-			DSWriteContainer.HashableDescriptorInfo.AddZeroed(SetInfo.Types.Num() + 1); // Add 1 for the Layout
+			DSWriteContainer.HashableDescriptorInfo.AddZeroed(StageInfo.Types.Num() + 1); // Add 1 for the Layout
 		}
-		DSWriteContainer.DescriptorWrites.AddZeroed(SetInfo.Types.Num());
-		DSWriteContainer.DescriptorImageInfo.AddZeroed(SetInfo.NumImageInfos);
-		DSWriteContainer.DescriptorBufferInfo.AddZeroed(SetInfo.NumBufferInfos);
-		DSWriteContainer.AccelerationStructureWrites.AddZeroed(SetInfo.NumAccelerationStructures);
-		DSWriteContainer.AccelerationStructures.AddZeroed(SetInfo.NumAccelerationStructures);
-
-		checkf(SetInfo.Types.Num() < 255, TEXT("Need more bits for BindingToDynamicOffsetMap (currently 8)! Requires %d descriptor bindings in a set!"), SetInfo.Types.Num());
-		DSWriteContainer.BindingToDynamicOffsetMap.AddUninitialized(SetInfo.Types.Num());
+		DSWriteContainer.DescriptorWrites.AddZeroed(StageInfo.Types.Num());
+		DSWriteContainer.DescriptorImageInfo.AddZeroed(StageInfo.NumImageInfos);
+		DSWriteContainer.DescriptorBufferInfo.AddZeroed(StageInfo.NumBufferInfos);
+		DSWriteContainer.AccelerationStructureWrites.AddZeroed(StageInfo.NumAccelerationStructures);
+		DSWriteContainer.AccelerationStructures.AddZeroed(StageInfo.NumAccelerationStructures);
+
+		checkf(StageInfo.Types.Num() < 255, TEXT("Need more bits for BindingToDynamicOffsetMap (currently 8)! Requires %d descriptor bindings in a set!"), StageInfo.Types.Num());
+		DSWriteContainer.BindingToDynamicOffsetMap.AddUninitialized(StageInfo.Types.Num());
 	}
 
 	FMemory::Memset(DSWriteContainer.BindingToDynamicOffsetMap.GetData(), 255, DSWriteContainer.BindingToDynamicOffsetMap.Num());
 
 	check(DSWriter.Num() == 0);
-	DSWriter.AddDefaulted(NumSets);
+	DSWriter.AddDefaulted(MaxNumSets);
 
 	const FVulkanSamplerState& DefaultSampler = Device->GetDefaultSampler();
 	const FVulkanView::FTextureView& DefaultImageView = Device->GetDefaultImageView();
@@ -105,17 +103,17 @@ void FVulkanCommonPipelineDescriptorState::CreateDescriptorWriteInfos()
 
 	uint8* CurrentBindingToDynamicOffsetMap = DSWriteContainer.BindingToDynamicOffsetMap.GetData();
 	TArray<uint32> DynamicOffsetsStart;
-	DynamicOffsetsStart.AddZeroed(NumSets);
+	DynamicOffsetsStart.AddZeroed(MaxNumSets);
 	uint32 TotalNumDynamicOffsets = 0;
 
-	for (int32 Set = 0; Set < NumSets; ++Set)
+	for (uint32 Set = 0; Set < MaxNumSets; ++Set)
 	{
-		const FDescriptorSetRemappingInfo::FSetInfo& SetInfo = DescriptorSetsLayout->RemappingInfo.SetInfos[Set];
+		const FDescriptorSetRemappingInfo::FStageInfo& StageInfo = DescriptorSetsLayout->RemappingInfo.StageInfos[Set];
 
 		DynamicOffsetsStart[Set] = TotalNumDynamicOffsets;
 
 		uint32 NumDynamicOffsets = DSWriter[Set].SetupDescriptorWrites(
-			SetInfo.Types, CurrentHashableDescriptorInfo,
+			StageInfo.Types, CurrentHashableDescriptorInfo,
 			CurrentDescriptorWrite, CurrentImageInfo, CurrentBufferInfo, CurrentBindingToDynamicOffsetMap,
 			CurrentAccelerationStructuresWriteDescriptors,
 			CurrentAccelerationStructures,
@@ -125,29 +123,29 @@ void FVulkanCommonPipelineDescriptorState::CreateDescriptorWriteInfos()
 
 		if (CurrentHashableDescriptorInfo) // UseVulkanDescriptorCache()
 		{
-			CurrentHashableDescriptorInfo += SetInfo.Types.Num();
+			CurrentHashableDescriptorInfo += StageInfo.Types.Num();
 			CurrentHashableDescriptorInfo->Layout.Max0 = UINT32_MAX;
 			CurrentHashableDescriptorInfo->Layout.Max1 = UINT32_MAX;
 			CurrentHashableDescriptorInfo->Layout.LayoutId = DescriptorSetsLayout->GetHandleIds()[Set];
 			++CurrentHashableDescriptorInfo;
 		}
 
-		CurrentDescriptorWrite += SetInfo.Types.Num();
-		CurrentImageInfo += SetInfo.NumImageInfos;
-		CurrentBufferInfo += SetInfo.NumBufferInfos;
-		CurrentAccelerationStructuresWriteDescriptors += SetInfo.NumAccelerationStructures;
-		CurrentAccelerationStructures += SetInfo.NumAccelerationStructures;
+		CurrentDescriptorWrite += StageInfo.Types.Num();
+		CurrentImageInfo += StageInfo.NumImageInfos;
+		CurrentBufferInfo += StageInfo.NumBufferInfos;
+		CurrentAccelerationStructuresWriteDescriptors += StageInfo.NumAccelerationStructures;
+		CurrentAccelerationStructures += StageInfo.NumAccelerationStructures;
 
-		CurrentBindingToDynamicOffsetMap += SetInfo.Types.Num();
+		CurrentBindingToDynamicOffsetMap += StageInfo.Types.Num();
 	}
 
 	DynamicOffsets.AddZeroed(TotalNumDynamicOffsets);
-	for (int32 Set = 0; Set < NumSets; ++Set)
+	for (uint32 Set = 0; Set < MaxNumSets; ++Set)
 	{
 		DSWriter[Set].DynamicOffsets = DynamicOffsetsStart[Set] + DynamicOffsets.GetData();
 	}
 
-	DescriptorSetHandles.AddZeroed(NumSets);
+	DescriptorSetHandles.AddZeroed(MaxNumSets);
 }
 
 template<bool bUseDynamicGlobalUBs>
@@ -291,7 +289,7 @@ void FVulkanComputePipelineDescriptorState::UpdateBindlessDescriptors(FVulkanCom
 }
 
 FVulkanGraphicsPipelineDescriptorState::FVulkanGraphicsPipelineDescriptorState(FVulkanDevice* InDevice, FVulkanRHIGraphicsPipelineState* InGfxPipeline)
-	: FVulkanCommonPipelineDescriptorState(InDevice)
+	: FVulkanCommonPipelineDescriptorState(InDevice, ShaderStage::NumGraphicsStages)
 	, GfxPipeline(InGfxPipeline)
 {
 	LLM_SCOPE_VULKAN(ELLMTagVulkan::VulkanShaders);
diff --git a/Engine/Source/Runtime/VulkanRHI/Private/VulkanPipelineState.h b/Engine/Source/Runtime/VulkanRHI/Private/VulkanPipelineState.h
--- a/Engine/Source/Runtime/VulkanRHI/Private/VulkanPipelineState.h
+++ b/Engine/Source/Runtime/VulkanRHI/Private/VulkanPipelineState.h
@@ -22,8 +22,9 @@ extern TAutoConsoleVariable<int32> GDynamicGlobalUBs;
 class FVulkanCommonPipelineDescriptorState : public VulkanRHI::FDeviceChild
 {
 public:
-	FVulkanCommonPipelineDescriptorState(FVulkanDevice* InDevice)
+	FVulkanCommonPipelineDescriptorState(FVulkanDevice* InDevice, uint32 InMaxNumSets)
 		: VulkanRHI::FDeviceChild(InDevice)
+		, MaxNumSets(InMaxNumSets)
 		, bUseBindless(InDevice->SupportsBindless())
 	{
 	}
@@ -137,6 +138,8 @@ protected:
 
 	void CreateDescriptorWriteInfos();
 
+	const uint32 MaxNumSets;
+
 	//#todo-rco: Won't work multithreaded!
 	FVulkanDescriptorSetWriteContainer DSWriteContainer;
 	const FVulkanDescriptorSetsLayout* DescriptorSetsLayout = nullptr;
diff --git a/Engine/Source/Runtime/VulkanRHI/Private/VulkanRHI.cpp b/Engine/Source/Runtime/VulkanRHI/Private/VulkanRHI.cpp
--- a/Engine/Source/Runtime/VulkanRHI/Private/VulkanRHI.cpp
+++ b/Engine/Source/Runtime/VulkanRHI/Private/VulkanRHI.cpp
@@ -1670,7 +1670,7 @@ void FVulkanDescriptorSetsLayoutInfo::AddDescriptor(int32 DescriptorSetIndex, co
 	VkDescriptorSetLayoutBinding* Binding = new(DescSetLayout.LayoutBindings) VkDescriptorSetLayoutBinding;
 	*Binding = Descriptor;
 
-	const FDescriptorSetRemappingInfo::FSetInfo& SetInfo = RemappingInfo.SetInfos[DescriptorSetIndex];
+	const FDescriptorSetRemappingInfo::FStageInfo& SetInfo = RemappingInfo.StageInfos[DescriptorSetIndex];
 	check(SetInfo.Types[Descriptor.binding] == Descriptor.descriptorType);
 	switch (Descriptor.descriptorType)
 	{
@@ -1679,15 +1679,15 @@ void FVulkanDescriptorSetsLayoutInfo::AddDescriptor(int32 DescriptorSetIndex, co
 	case VK_DESCRIPTOR_TYPE_SAMPLED_IMAGE:
 	case VK_DESCRIPTOR_TYPE_STORAGE_IMAGE:
 	case VK_DESCRIPTOR_TYPE_INPUT_ATTACHMENT:
-		IncrementChecked(RemappingInfo.SetInfos[DescriptorSetIndex].NumImageInfos);
+		IncrementChecked(RemappingInfo.StageInfos[DescriptorSetIndex].NumImageInfos);
 		break;
 	case VK_DESCRIPTOR_TYPE_UNIFORM_BUFFER:
 	case VK_DESCRIPTOR_TYPE_STORAGE_BUFFER:
 	case VK_DESCRIPTOR_TYPE_UNIFORM_BUFFER_DYNAMIC:
-		IncrementChecked(RemappingInfo.SetInfos[DescriptorSetIndex].NumBufferInfos);
+		IncrementChecked(RemappingInfo.StageInfos[DescriptorSetIndex].NumBufferInfos);
 		break;
 	case VK_DESCRIPTOR_TYPE_ACCELERATION_STRUCTURE_KHR:
-		IncrementChecked(RemappingInfo.SetInfos[DescriptorSetIndex].NumAccelerationStructures);
+		IncrementChecked(RemappingInfo.StageInfos[DescriptorSetIndex].NumAccelerationStructures);
 		break;
 	case VK_DESCRIPTOR_TYPE_STORAGE_TEXEL_BUFFER:
 	case VK_DESCRIPTOR_TYPE_UNIFORM_TEXEL_BUFFER:
@@ -1711,16 +1711,23 @@ void FVulkanDescriptorSetsLayoutInfo::GenerateHash(const TArrayView<FRHISamplerS
 
 	for (uint32 RemapingIndex = 0; RemapingIndex < ShaderStage::NumStages; ++RemapingIndex)
 	{
-		Hash = FCrc::MemCrc32(&RemappingInfo.StageInfos[RemapingIndex].PackedUBDescriptorSet, sizeof(uint16), Hash);
-		Hash = FCrc::MemCrc32(&RemappingInfo.StageInfos[RemapingIndex].Pad0, sizeof(uint16), Hash);
+		const FDescriptorSetRemappingInfo::FStageInfo& StageInfo = RemappingInfo.StageInfos[RemapingIndex];
 
-		TArray<FDescriptorSetRemappingInfo::FRemappingInfo>& Globals = RemappingInfo.StageInfos[RemapingIndex].Globals;
+		Hash = FCrc::MemCrc32(&StageInfo.PackedUBDescriptorSet, sizeof(uint16), Hash);
+		Hash = FCrc::MemCrc32(&StageInfo.NumImageInfos, sizeof(uint16), Hash);
+		Hash = FCrc::MemCrc32(&StageInfo.NumBufferInfos, sizeof(uint16), Hash);
+		Hash = FCrc::MemCrc32(&StageInfo.NumAccelerationStructures, sizeof(uint16), Hash);
+
+		const TArray<FDescriptorSetRemappingInfo::FRemappingInfo>& Globals = StageInfo.Globals;
 		Hash = FCrc::MemCrc32(Globals.GetData(), sizeof(FDescriptorSetRemappingInfo::FRemappingInfo) * Globals.Num(), Hash);
 
-		TArray<FDescriptorSetRemappingInfo::FUBRemappingInfo>& UniformBuffers = RemappingInfo.StageInfos[RemapingIndex].UniformBuffers;
+		const TArray<FDescriptorSetRemappingInfo::FUBRemappingInfo>& UniformBuffers = StageInfo.UniformBuffers;
 		Hash = FCrc::MemCrc32(UniformBuffers.GetData(), sizeof(FDescriptorSetRemappingInfo::FUBRemappingInfo) * UniformBuffers.Num(), Hash);
 
-		TArray<uint16>& PackedUBBindingIndices = RemappingInfo.StageInfos[RemapingIndex].PackedUBBindingIndices;
+		const TArray<VkDescriptorType>& Types = StageInfo.Types;
+		Hash = FCrc::MemCrc32(Types.GetData(), sizeof(VkDescriptorType) * Types.Num(), Hash);
+
+		const TArray<uint16>& PackedUBBindingIndices = StageInfo.PackedUBBindingIndices;
 		Hash = FCrc::MemCrc32(PackedUBBindingIndices.GetData(), sizeof(uint16) * PackedUBBindingIndices.Num(), Hash);
 	}
 
diff --git a/Engine/Source/Runtime/VulkanRHI/Private/VulkanShaders.cpp b/Engine/Source/Runtime/VulkanRHI/Private/VulkanShaders.cpp
--- a/Engine/Source/Runtime/VulkanRHI/Private/VulkanShaders.cpp
+++ b/Engine/Source/Runtime/VulkanRHI/Private/VulkanShaders.cpp
@@ -577,9 +577,11 @@ void FVulkanLayout::PatchSpirvBindings(FVulkanShader::FSpirvCode& SprivCode, ESh
 			const uint32 OffsetDescriptorSet = CodeHeader.UniformBufferSpirvInfos[Index].DescriptorSetOffset;
 			const uint32 OffsetBindingIndex = CodeHeader.UniformBufferSpirvInfos[Index].BindingIndexOffset;
 			check(OffsetDescriptorSet != UINT32_MAX && OffsetBindingIndex != UINT32_MAX);
-			uint16 NewDescriptorSet = StageInfo.UniformBuffers[Index].Remapping.NewDescriptorSet;
-			Spirv[OffsetDescriptorSet] = NewDescriptorSet;
-			uint16 NewBindingIndex = StageInfo.UniformBuffers[Index].Remapping.NewBindingIndex;
+			const uint16 NewDescriptorSet = StageInfo.UniformBuffers[Index].Remapping.NewDescriptorSet;
+			checkf(Spirv[OffsetDescriptorSet] == NewDescriptorSet,
+				TEXT("Descriptor set in SPIRV (%u) differs from expected stage value (%u)"),
+				Spirv[OffsetDescriptorSet], NewDescriptorSet);
+			const uint16 NewBindingIndex = StageInfo.UniformBuffers[Index].Remapping.NewBindingIndex;
 			Spirv[OffsetBindingIndex] = NewBindingIndex;
 		}
 	}
@@ -590,9 +592,11 @@ void FVulkanLayout::PatchSpirvBindings(FVulkanShader::FSpirvCode& SprivCode, ESh
 		const uint32 OffsetDescriptorSet = CodeHeader.GlobalSpirvInfos[Index].DescriptorSetOffset;
 		const uint32 OffsetBindingIndex = CodeHeader.GlobalSpirvInfos[Index].BindingIndexOffset;
 		check(OffsetDescriptorSet != UINT32_MAX && OffsetBindingIndex != UINT32_MAX);
-		uint16 NewDescriptorSet = StageInfo.Globals[Index].NewDescriptorSet;
-		Spirv[OffsetDescriptorSet] = NewDescriptorSet;
-		uint16 NewBindingIndex = StageInfo.Globals[Index].NewBindingIndex;
+		const uint16 NewDescriptorSet = StageInfo.Globals[Index].NewDescriptorSet;
+		checkf(Spirv[OffsetDescriptorSet] == NewDescriptorSet,
+			TEXT("Descriptor set in SPIRV (%u) differs from expected stage value (%u)"),
+			Spirv[OffsetDescriptorSet], NewDescriptorSet);
+		const uint16 NewBindingIndex = StageInfo.Globals[Index].NewBindingIndex;
 		Spirv[OffsetBindingIndex] = NewBindingIndex;
 	}
 
@@ -602,7 +606,9 @@ void FVulkanLayout::PatchSpirvBindings(FVulkanShader::FSpirvCode& SprivCode, ESh
 		const uint32 OffsetDescriptorSet = CodeHeader.PackedUBs[Index].SPIRVDescriptorSetOffset;
 		const uint32 OffsetBindingIndex = CodeHeader.PackedUBs[Index].SPIRVBindingIndexOffset;
 		check(OffsetDescriptorSet != UINT32_MAX && OffsetBindingIndex != UINT32_MAX);
-		Spirv[OffsetDescriptorSet] = StageInfo.PackedUBDescriptorSet;
+		checkf(Spirv[OffsetDescriptorSet] == StageInfo.PackedUBDescriptorSet, 
+			TEXT("Descriptor set in SPIRV (%u) differs from expected stage value (%u)"), 
+			Spirv[OffsetDescriptorSet], StageInfo.PackedUBDescriptorSet);
 		Spirv[OffsetBindingIndex] = StageInfo.PackedUBBindingIndices[Index];
 	}
 }
@@ -781,8 +787,6 @@ void FVulkanDescriptorSetsLayoutInfo::FinalizeBindings(const FVulkanDevice& Devi
 {
 	checkSlow(RemappingInfo.IsEmpty());
 
-	TMap<uint32, FDescriptorSetRemappingInfo::FUBRemappingInfo> AlreadyProcessedUBs;
-
 	// We'll be reusing this struct
 	VkDescriptorSetLayoutBinding Binding;
 	FMemory::Memzero(Binding);
@@ -792,35 +796,22 @@ void FVulkanDescriptorSetsLayoutInfo::FinalizeBindings(const FVulkanDevice& Devi
 	const bool bConvertPackedUBsToDynamic = !Device.SupportsBindless() && (bConvertAllUBsToDynamic || (GDynamicGlobalUBs.GetValueOnAnyThread() == 1));
 	const uint32 MaxDescriptorSetUniformBuffersDynamic = Device.GetLimits().maxDescriptorSetUniformBuffersDynamic;
 
-	uint8	DescriptorStageToSetMapping[ShaderStage::NumStages];
-	FMemory::Memset(DescriptorStageToSetMapping, UINT8_MAX);
-
-	auto FindOrAddDescriptorSet = [&](int32 Stage) -> uint8
-	{
-		if (DescriptorStageToSetMapping[Stage] == UINT8_MAX)
-		{
-			uint32 NewSet = RemappingInfo.SetInfos.AddDefaulted();
-			DescriptorStageToSetMapping[Stage] = (uint8)NewSet;
-			return NewSet;
-		}
-
-		return DescriptorStageToSetMapping[Stage];
-	};
-
 	int32 CurrentImmutableSampler = 0;
 	for (int32 Stage = 0; Stage < (bIsCompute ? 1 : ShaderStage::NumStages); ++Stage)
 	{
 		if (const FVulkanShaderHeader* ShaderHeader = UBGatherInfo.CodeHeaders[Stage])
 		{
+			// Shaders are compiles with the value of Stage as the descriptor set
+			const int32 DescriptorSet = Stage;
+
 			const VkShaderStageFlags StageFlags = UEFrequencyToVKStageBit(bIsCompute ? SF_Compute : ShaderStage::GetFrequencyForGfxStage((ShaderStage::EStage)Stage));
 			Binding.stageFlags = StageFlags;
 
 			RemappingInfo.StageInfos[Stage].PackedUBBindingIndices.Reserve(ShaderHeader->PackedUBs.Num());
 			for (int32 Index = 0; Index < ShaderHeader->PackedUBs.Num(); ++Index)
 			{
-				int32 DescriptorSet = FindOrAddDescriptorSet(Stage);
-				VkDescriptorType Type = bConvertPackedUBsToDynamic ? VK_DESCRIPTOR_TYPE_UNIFORM_BUFFER_DYNAMIC : VK_DESCRIPTOR_TYPE_UNIFORM_BUFFER;
-				uint32 NewBindingIndex = RemappingInfo.AddPackedUB(Stage, Index, DescriptorSet, Type);
+				const VkDescriptorType Type = bConvertPackedUBsToDynamic ? VK_DESCRIPTOR_TYPE_UNIFORM_BUFFER_DYNAMIC : VK_DESCRIPTOR_TYPE_UNIFORM_BUFFER;
+				const uint32 NewBindingIndex = RemappingInfo.AddPackedUB(Stage, Index, DescriptorSet, Type);
 
 				Binding.binding = NewBindingIndex;
 				Binding.descriptorType = Type;
@@ -847,12 +838,11 @@ void FVulkanDescriptorSetsLayoutInfo::FinalizeBindings(const FVulkanDevice& Devi
 					bool bProcessRegularUB = true;
 					if (bProcessRegularUB)
 					{
-						int32 DescriptorSet = FindOrAddDescriptorSet(Stage);
 						uint32 NewBindingIndex;
 						RemappingInfo.AddUBWithData(Stage, Index, DescriptorSet, Type, NewBindingIndex);
 						Binding.binding = NewBindingIndex;
 
-						AddDescriptor(FindOrAddDescriptorSet(Stage), Binding);
+						AddDescriptor(DescriptorSet, Binding);
 					}
 				}
 				else
@@ -866,17 +856,15 @@ void FVulkanDescriptorSetsLayoutInfo::FinalizeBindings(const FVulkanDevice& Devi
 			for (int32 Index = 0; Index < ShaderHeader->Globals.Num(); ++Index)
 			{
 				const FVulkanShaderHeader::FGlobalInfo& GlobalInfo = ShaderHeader->Globals[Index];
-				int32 DescriptorSet = FindOrAddDescriptorSet(Stage);
 				const VkDescriptorType Type = (VkDescriptorType)GlobalInfo.TypeIndex;
 				uint32 NewBindingIndex = RemappingInfo.AddGlobal(Stage, Index, DescriptorSet, Type);
 				Binding.binding = NewBindingIndex;
 				Binding.descriptorType = Type;
-					AddDescriptor(DescriptorSet, Binding);
-				}
+				AddDescriptor(DescriptorSet, Binding);
+			}
 
 			if (ShaderHeader->InputAttachments.Num())
 			{
-				int32 DescriptorSet = FindOrAddDescriptorSet(Stage);
 				check(Stage == ShaderStage::Pixel);
 				for (int32 SrcIndex = 0; SrcIndex < ShaderHeader->InputAttachments.Num(); ++SrcIndex)
 				{
@@ -895,38 +883,23 @@ void FVulkanDescriptorSetsLayoutInfo::FinalizeBindings(const FVulkanDevice& Devi
 
 	CompileTypesUsageID();
 	GenerateHash(ImmutableSamplers, bIsCompute ? VK_PIPELINE_BIND_POINT_COMPUTE : VK_PIPELINE_BIND_POINT_GRAPHICS);
-
-	// If we are consolidating and no uniforms are present in the shader, then strip the empty set data
-	{
-		for (int32 Index = 0; Index < RemappingInfo.SetInfos.Num(); ++Index)
-		{
-			check(RemappingInfo.SetInfos[Index].Types.Num() > 0);
-		}
-	}	
 }
 
 void FVulkanComputePipelineDescriptorInfo::Initialize(const FDescriptorSetRemappingInfo& InRemappingInfo)
 {
 	check(!bInitialized);
 
-	RemappingGlobalInfos = InRemappingInfo.StageInfos[0].Globals;
-	RemappingUBInfos = InRemappingInfo.StageInfos[0].UniformBuffers;
-	RemappingPackedUBInfos = InRemappingInfo.StageInfos[0].PackedUBBindingIndices;
+	const FDescriptorSetRemappingInfo::FStageInfo& StageInfo = InRemappingInfo.StageInfos[ShaderStage::Compute];
+
+	RemappingGlobalInfos = StageInfo.Globals;
+	RemappingUBInfos = StageInfo.UniformBuffers;
+	RemappingPackedUBInfos = StageInfo.PackedUBBindingIndices;
 
 	RemappingInfo = &InRemappingInfo;
 
-	for (int32 Index = 0; Index < InRemappingInfo.SetInfos.Num(); ++Index)
+	if (StageInfo.Types.Num() > 0)
 	{
-		const FDescriptorSetRemappingInfo::FSetInfo& SetInfo = InRemappingInfo.SetInfos[Index];
-		if (SetInfo.Types.Num() > 0)
-		{
-			check(Index < sizeof(HasDescriptorsInSetMask) * 8);
-			HasDescriptorsInSetMask = HasDescriptorsInSetMask | (1 << Index);
-		}
-		else
-		{
-			ensure(0);
-		}
+		HasDescriptorsInSetMask = HasDescriptorsInSetMask | (1 << 0);
 	}
 
 	bInitialized = true;
@@ -938,28 +911,22 @@ void FVulkanGfxPipelineDescriptorInfo::Initialize(const FDescriptorSetRemappingI
 
 	for (int32 StageIndex = 0; StageIndex < ShaderStage::NumStages; ++StageIndex)
 	{
-		//#todo-rco: Enable this!
-		RemappingUBInfos[StageIndex] = InRemappingInfo.StageInfos[StageIndex].UniformBuffers;
-		RemappingGlobalInfos[StageIndex] = InRemappingInfo.StageInfos[StageIndex].Globals;
-		RemappingPackedUBInfos[StageIndex] = InRemappingInfo.StageInfos[StageIndex].PackedUBBindingIndices;
-	}
+		const FDescriptorSetRemappingInfo::FStageInfo& StageInfo = InRemappingInfo.StageInfos[StageIndex];
 
-	RemappingInfo = &InRemappingInfo;
+		//#todo-rco: Enable this!
+		RemappingUBInfos[StageIndex] = StageInfo.UniformBuffers;
+		RemappingGlobalInfos[StageIndex] = StageInfo.Globals;
+		RemappingPackedUBInfos[StageIndex] = StageInfo.PackedUBBindingIndices;
 
-	for (int32 Index = 0; Index < InRemappingInfo.SetInfos.Num(); ++Index)
-	{
-		const FDescriptorSetRemappingInfo::FSetInfo& SetInfo = InRemappingInfo.SetInfos[Index];
-		if (SetInfo.Types.Num() > 0)
-		{
-			check(Index < sizeof(HasDescriptorsInSetMask) * 8);
-			HasDescriptorsInSetMask = HasDescriptorsInSetMask | (1 << Index);
-		}
-		else
+		if (StageInfo.Types.Num() > 0)
 		{
-			ensure(0);
+			check(StageIndex < sizeof(HasDescriptorsInSetMask) * 8);
+			HasDescriptorsInSetMask = HasDescriptorsInSetMask | (1 << StageIndex);
 		}
 	}
 
+	RemappingInfo = &InRemappingInfo;
+
 	bInitialized = true;
 }
 
