VulkanRHI: Fix compilation issue (declaration shadows a field)

#rnx

diff --git a/Engine/Source/Runtime/VulkanRHI/Private/VulkanRayTracing.cpp b/Engine/Source/Runtime/VulkanRHI/Private/VulkanRayTracing.cpp
--- a/Engine/Source/Runtime/VulkanRHI/Private/VulkanRayTracing.cpp
+++ b/Engine/Source/Runtime/VulkanRHI/Private/VulkanRayTracing.cpp
@@ -945,16 +945,13 @@ void FVulkanRayTracingShaderTable::ReleaseLocalBuffer(FVulkanDevice* Device, FVu
 
 void FVulkanRayTracingShaderTable::Init(const FVulkanRayTracingPipelineState* Pipeline)
 {
-	auto InitAlloc = [Device = Device, HandleSize = HandleSize, HandleSizeAligned = HandleSizeAligned](FVulkanShaderTableAllocation& Alloc, uint32 InHandleCount, bool InUseLocalRecord) {
+	auto InitAlloc = [Device = Device, HandleSizeAligned = HandleSizeAligned](FVulkanShaderTableAllocation& Alloc, uint32 InHandleCount, bool InUseLocalRecord) {
 
 		Alloc.HandleCount = InHandleCount;
 		Alloc.bUseLocalRecord = InUseLocalRecord;
 
 		if (Alloc.HandleCount > 0)
 		{
-			VkDevice DeviceHandle = Device->GetInstanceHandle();
-			VkPhysicalDevice Gpu = Device->GetPhysicalHandle();
-
 			const VkPhysicalDeviceRayTracingPipelinePropertiesKHR& RayTracingPipelineProps = Device->GetOptionalExtensionProperties().RayTracingPipelineProps;
 			Alloc.Region.stride = InUseLocalRecord ? FMath::Min<VkDeviceSize>(RayTracingPipelineProps.maxShaderGroupStride, 4096) : HandleSizeAligned; // :todo-jn: shrink stride to necessary amount
 			Alloc.Region.size = Alloc.HandleCount * Alloc.Region.stride;
@@ -2309,7 +2306,7 @@ void FVulkanCommandListContext::RHISetBindingsOnShaderBindingTable(FRHIShaderBin
 		TaskContexts.Add(FTaskContext{ WorkerIndex });
 	}
 
-	auto BindingTask = [this, Bindings, Device = Device, Pipeline, ShaderTable, BindingType](const FTaskContext& Context, int32 CurrentIndex)
+	auto BindingTask = [this, Bindings, Pipeline, ShaderTable, BindingType](const FTaskContext& Context, int32 CurrentIndex)
 	{
 		const FRayTracingLocalShaderBindings& Binding = Bindings[CurrentIndex];
 
