VulkanRHI: Quick fix Vulkan descriptor cache, skip empty stages.

#rnx

diff --git a/Engine/Source/Runtime/VulkanRHI/Private/VulkanPendingState.cpp b/Engine/Source/Runtime/VulkanRHI/Private/VulkanPendingState.cpp
--- a/Engine/Source/Runtime/VulkanRHI/Private/VulkanPendingState.cpp
+++ b/Engine/Source/Runtime/VulkanRHI/Private/VulkanPendingState.cpp
@@ -789,12 +789,12 @@ bool FVulkanDescriptorSetCache::FCachedPool::CreateDescriptorSets(
 
 	NewSetEntry.NumSets = DSWriters.Num();
 	check(NewSetEntry.NumSets <= NewSetEntry.Sets.Num());
-	check(NewSetEntry.NumSets == SetsLayout.GetHandles().Num());
+	check(SetsLayout.GetHandles().Num() <= NewSetEntry.NumSets);
 
 	for (int32 Index = 0; Index < NewSetEntry.NumSets; ++Index)
 	{
 		FVulkanDescriptorSetWriter& DSWriter = DSWriters[Index];
-		if (DSWriter.GetNumWrites() == 0) // Should not normally happen
+		if (DSWriter.GetNumWrites() == 0)
 		{
 			NewSetEntry.Sets[Index] = VK_NULL_HANDLE;
 			continue;
diff --git a/Engine/Source/Runtime/VulkanRHI/Private/VulkanPipelineState.cpp b/Engine/Source/Runtime/VulkanRHI/Private/VulkanPipelineState.cpp
--- a/Engine/Source/Runtime/VulkanRHI/Private/VulkanPipelineState.cpp
+++ b/Engine/Source/Runtime/VulkanRHI/Private/VulkanPipelineState.cpp
@@ -67,6 +67,10 @@ void FVulkanCommonPipelineDescriptorState::CreateDescriptorWriteInfos()
 	for (uint32 Set = 0; Set < MaxNumSets; ++Set)
 	{
 		const FDescriptorSetRemappingInfo::FStageInfo& StageInfo = DescriptorSetsLayout->RemappingInfo.StageInfos[Set];
+		if (!StageInfo.Types.Num())
+		{
+			continue;
+		}
 		
 		if (UseVulkanDescriptorCache())
 		{
@@ -109,10 +113,14 @@ void FVulkanCommonPipelineDescriptorState::CreateDescriptorWriteInfos()
 	for (uint32 Set = 0; Set < MaxNumSets; ++Set)
 	{
 		const FDescriptorSetRemappingInfo::FStageInfo& StageInfo = DescriptorSetsLayout->RemappingInfo.StageInfos[Set];
+		if (!StageInfo.Types.Num())
+		{
+			continue;
+		}
 
 		DynamicOffsetsStart[Set] = TotalNumDynamicOffsets;
 
-		uint32 NumDynamicOffsets = DSWriter[Set].SetupDescriptorWrites(
+		const uint32 NumDynamicOffsets = DSWriter[Set].SetupDescriptorWrites(
 			StageInfo.Types, CurrentHashableDescriptorInfo,
 			CurrentDescriptorWrite, CurrentImageInfo, CurrentBufferInfo, CurrentBindingToDynamicOffsetMap,
 			CurrentAccelerationStructuresWriteDescriptors,
