VulkanRHI: Fix validation errors about null descriptor sets being bound.

#rnx

diff --git a/Engine/Source/Runtime/VulkanRHI/Private/VulkanPipelineState.h b/Engine/Source/Runtime/VulkanRHI/Private/VulkanPipelineState.h
--- a/Engine/Source/Runtime/VulkanRHI/Private/VulkanPipelineState.h
+++ b/Engine/Source/Runtime/VulkanRHI/Private/VulkanPipelineState.h
@@ -126,12 +126,16 @@ protected:
 	inline void Bind(VkCommandBuffer CmdBuffer, VkPipelineLayout PipelineLayout, VkPipelineBindPoint BindPoint)
 	{
 		// Bindless will replace with global sets
-		if (!bUseBindless)
+		if (!bUseBindless && UsedSetsMask)
 		{
+			const uint32 FirstSet = FMath::CountTrailingZeros(UsedSetsMask);
+			const uint32 NumSets = 32 - FMath::CountLeadingZeros(UsedSetsMask) - FirstSet;
+			check(FirstSet + NumSets <= (uint32)DescriptorSetHandles.Num());
+
 			VulkanRHI::vkCmdBindDescriptorSets(CmdBuffer,
 				BindPoint,
 				PipelineLayout,
-				0, DescriptorSetHandles.Num(), DescriptorSetHandles.GetData(),
+				FirstSet, NumSets, &DescriptorSetHandles[FirstSet],
 				(uint32)DynamicOffsets.Num(), DynamicOffsets.GetData());
 		}
 	}
