Vulkan: Fixed missing debug markers

diff --git a/Engine/Source/Runtime/VulkanRHI/Private/VulkanLayers.cpp b/Engine/Source/Runtime/VulkanRHI/Private/VulkanLayers.cpp
--- a/Engine/Source/Runtime/VulkanRHI/Private/VulkanLayers.cpp
+++ b/Engine/Source/Runtime/VulkanRHI/Private/VulkanLayers.cpp
@@ -399,7 +399,7 @@ void FVulkanIntanceSetupHelper::AddDebugLayers(const TArray<FLayerWithExtensions
 		}
 	}
 
-	const bool bForceDebugUtils = FParse::Param(FCommandLine::Get(), TEXT("vulkandebugutils"));
+	const bool bForceDebugUtils = VULKAN_ENABLE_DRAW_MARKERS || FParse::Param(FCommandLine::Get(), TEXT("vulkandebugutils"));
 	if ((bUseVulkanValidation || bForceDebugUtils) && (ActiveDebugLayerExtension == FVulkanDynamicRHI::EActiveDebugLayerExtension::None))
 	{
 		auto FindLayerContainingExtension = [](const ANSICHAR* ExtensionName, const TArray<FLayerWithExtensions>& LayerProperties)
