Add precache priority to RHI gfx PSO creation info.
Pass through PSO compile priorities to Android external compilers.
Add parameter to optionally control Android external PSO compiler affinity and cpu scheduling settings.
Add ability to re-compile a vulkan PSO in the engine's process if the external compiler fails.
Add option to entirely disable all external PSO compiles if the count of failures exceeds a certain threshold.

diff --git a/Engine/Build/Android/Java/src/com/epicgames/unreal/psoservices/PSOProgramService.java b/Engine/Build/Android/Java/src/com/epicgames/unreal/psoservices/PSOProgramService.java
--- a/Engine/Build/Android/Java/src/com/epicgames/unreal/psoservices/PSOProgramService.java
+++ b/Engine/Build/Android/Java/src/com/epicgames/unreal/psoservices/PSOProgramService.java
@@ -59,6 +59,7 @@ public class PSOProgramService extends Service implements Logger.ILoggerCallback
 	static final String PSOCacheData_Key = "psocache";
 	static final String JobContext_Key = "jtx";
 	static final String JobID_Key = "jid";
+	static final String Priority_Key = "pri";
 	static final String ServiceID_Key = "sid";
 	static final String CompiledProgram_Key = "cpg";
 	static final String SHMem_Key = "shm";
@@ -263,6 +264,8 @@ public class PSOProgramService extends Service implements Logger.ILoggerCallback
 		return false;
 	}
 
+	public native void NativeSetThreadPriority(long PriorityInfo);
+
 	@Nullable
 	@Override
 	public IBinder onBind(Intent intent)
@@ -271,6 +274,8 @@ public class PSOProgramService extends Service implements Logger.ILoggerCallback
 
 		if(!bGFXInitialized)
 		{
+			System.loadLibrary("psoservice");
+
 			if(!UseVulkan())
 			{
 				Bundle extras =intent.getExtras();
@@ -316,16 +321,18 @@ public class PSOProgramService extends Service implements Logger.ILoggerCallback
 					beginTrace("PSOProgramService.handleMessage");
 					if(UseVulkan())
 					{
-						byte[] VS = msg.getData().getByteArray(VS_Key);
-						byte[] PS = msg.getData().getByteArray(PS_Key);
-						byte[] PSOData = msg.getData().getByteArray(PSOData_Key);
+						byte[] VS			= msg.getData().getByteArray(VS_Key);
+						byte[] PS			= msg.getData().getByteArray(PS_Key);
+						byte[] PSOData		= msg.getData().getByteArray(PSOData_Key);
 						byte[] PSOCacheData = msg.getData().getByteArray(PSOCacheData_Key);
+						long PriorityInfo	= msg.getData().getLong(Priority_Key);
 
-						int JobID = msg.getData().getInt(JobID_Key);
-						int ServiceID = msg.getData().getInt(ServiceID_Key);
+						int JobID			= msg.getData().getInt(JobID_Key);
+						int ServiceID		= msg.getData().getInt(ServiceID_Key);
 
 						byte[] JobContext = msg.getData().getByteArray(JobContext_Key);
 
+						NativeSetThreadPriority(PriorityInfo);
 						//logger.verbose("Processing program job "+JobID);
 						ProcessVulkanProgramRequest(msg.replyTo, JobID, ServiceID, JobContext, VS, PS, PSOData, PSOCacheData);
 					}
@@ -336,9 +343,10 @@ public class PSOProgramService extends Service implements Logger.ILoggerCallback
 						String CS = msg.getData().getString(CS_Key);
 						int JobID = msg.getData().getInt(JobID_Key);
 						int ServiceID = msg.getData().getInt(ServiceID_Key);
-
+						long PriorityInfo = msg.getData().getLong(Priority_Key);
 						byte[] JobContext = msg.getData().getByteArray(JobContext_Key);
 
+						NativeSetThreadPriority(PriorityInfo);
 						//logger.verbose("Processing program job "+JobID);
 						ProcessGLProgramRequest(msg.replyTo, JobID, ServiceID, JobContext, VS, PS, CS);
 					}
@@ -355,10 +363,12 @@ public class PSOProgramService extends Service implements Logger.ILoggerCallback
 						long PSODataSize				= msg.getData().getLong(PSOData_Key);
 						long PSOCacheDataSize			= msg.getData().getLong(PSOCacheData_Key);
 						int JobID						= msg.getData().getInt(JobID_Key);
+						long PriorityInfo				= msg.getData().getLong(Priority_Key);
 						int ServiceID					= msg.getData().getInt(ServiceID_Key);
 						byte[] JobContext				= msg.getData().getByteArray(JobContext_Key);
 						int SHMemFD						= SharedFD.getFd();
 
+						NativeSetThreadPriority(PriorityInfo);
 						//logger.verbose("Processing program job "+JobID);
 						ProcessVulkanProgramRequestSHM(msg.replyTo, JobID, ServiceID, JobContext, SHMemFD, VSSize, PSSize, PSODataSize, PSOCacheDataSize);
 						try
@@ -586,8 +596,6 @@ public class PSOProgramService extends Service implements Logger.ILoggerCallback
 	private void initVulkanContext()
 	{
 		beginTrace("PSOProgramService.InitVulkanContext");
-		System.loadLibrary("psoservice");
-
 		InitVKDevice();
 		endTrace();
 	}
diff --git a/Engine/Build/Android/Java/src/com/epicgames/unreal/psoservices/PSOProgramServiceAccessor.java b/Engine/Build/Android/Java/src/com/epicgames/unreal/psoservices/PSOProgramServiceAccessor.java
--- a/Engine/Build/Android/Java/src/com/epicgames/unreal/psoservices/PSOProgramServiceAccessor.java
+++ b/Engine/Build/Android/Java/src/com/epicgames/unreal/psoservices/PSOProgramServiceAccessor.java
@@ -34,9 +34,12 @@ import java.util.HashSet;
 import java.util.Map;
 import java.util.Set;
 import java.util.concurrent.ConcurrentHashMap;
+import java.util.concurrent.Semaphore;
 import java.util.concurrent.atomic.AtomicInteger;
+import java.util.concurrent.atomic.AtomicBoolean;
 import java.util.concurrent.locks.ReadWriteLock;
 import java.util.concurrent.locks.ReentrantReadWriteLock;
+import java.util.concurrent.TimeUnit;
 
 public class PSOProgramServiceAccessor
 {
@@ -160,6 +163,9 @@ public class PSOProgramServiceAccessor
 
 	static HandlerThread PSOProgramAccessorHandlerThread;
 
+	static AtomicBoolean bIsReady = new AtomicBoolean(false);
+	static Semaphore MaxLiveRequestsSemaphore;
+	static int MaxLiveRequestPermitCount = 0;
 	static Messenger mReplyToMe;
 
 	final AtomicInteger LastServiceIdx = new AtomicInteger(0);
@@ -220,6 +226,10 @@ public class PSOProgramServiceAccessor
 				PSOProgramAccessorHandlerThread.start();
 				mReplyToMe = new Messenger(new IncomingHandler(PSOProgramAccessorHandlerThread.getLooper()));
 
+				// rate limit the total number of inflight PSO compile requests. shut down waits for all jobs to complete and this reduces the potential wait time.
+				// calling threads wait for a request permit to become available or shutdown is issued.
+				MaxLiveRequestPermitCount = numServices * 3;
+				MaxLiveRequestsSemaphore = new Semaphore(MaxLiveRequestPermitCount); 
 				if(bUseVulkan)
 				{
 					bSuccess = _PSOProgramServiceAccessor.StartVulkanServices(numServices);
@@ -243,6 +253,7 @@ public class PSOProgramServiceAccessor
 		{
 			AndroidThunkJava_StopRemoteProgramLink();
 		}
+		bIsReady.set(bSuccess);
 		return bSuccess;
 	}
 
@@ -250,6 +261,9 @@ public class PSOProgramServiceAccessor
 	{
 		if( _PSOProgramServiceAccessor != null)
 		{
+			bIsReady.set(false);
+			// acquire all of the requests, this causes us to wait for all in progress compiles.
+			MaxLiveRequestsSemaphore.acquireUninterruptibly(MaxLiveRequestPermitCount);
 			try
 			{
 				ProgramServiceAccessorlock.writeLock().lock();
@@ -350,23 +364,41 @@ public class PSOProgramServiceAccessor
 	}
 
 	static final private Object ProgramLinkLock = new Object();
-	public static JNIProgramLinkResponse AndroidThunkJava_OGLRemoteProgramLink(byte[] ContextData, String VertexShader, String PixelShader, String ComputeShader, boolean bAllowTimeOuts)
+	public static JNIProgramLinkResponse AndroidThunkJava_OGLRemoteProgramLink(byte[] ContextData, long PriorityInfo, String VertexShader, String PixelShader, String ComputeShader, boolean bAllowTimeOuts)
 	{
-		if(GameActivity.IsActivityPaused())
+		while( bIsReady.get())
 		{
-			// one at a time when backgrounded.
-			synchronized (ProgramLinkLock)
+			try
 			{
-				return OGLRemoteProgramLink_internal(ContextData, VertexShader, PixelShader, ComputeShader, bAllowTimeOuts);
+				if( MaxLiveRequestsSemaphore.tryAcquire(1, 100, TimeUnit.MILLISECONDS) )
+				{
+					JNIProgramLinkResponse returnResponse = null;
+					if(GameActivity.IsActivityPaused())
+					{
+						// one at a time when backgrounded.
+						synchronized (ProgramLinkLock)
+						{
+							returnResponse = OGLRemoteProgramLink_internal(ContextData, PriorityInfo, VertexShader, PixelShader, ComputeShader, bAllowTimeOuts);
+						}
+					}
+					else
+					{
+						returnResponse = OGLRemoteProgramLink_internal(ContextData, PriorityInfo, VertexShader, PixelShader, ComputeShader, bAllowTimeOuts);
+					}
+					MaxLiveRequestsSemaphore.release(1);
+					return returnResponse;
+				}
+			}
+			catch (InterruptedException e)
+			{
+				Log.error( "interrupted " + e);
 			}
 		}
-		else
-		{
-			return OGLRemoteProgramLink_internal(ContextData, VertexShader, PixelShader, ComputeShader, bAllowTimeOuts);
-		}
+		Log.error("Failed to compile PSO services not ready.");
+		return null;
 	}
 
-	private static JNIProgramLinkResponse OGLRemoteProgramLink_internal(byte[] ContextData, String VertexShader, String PixelShader, String ComputeShader, boolean bAllowTimeOuts)
+	private static JNIProgramLinkResponse OGLRemoteProgramLink_internal(byte[] ContextData, long PriorityInfo, String VertexShader, String PixelShader, String ComputeShader, boolean bAllowTimeOuts)
 	{
 		try
 		{
@@ -411,6 +443,8 @@ public class PSOProgramServiceAccessor
 			params.putByteArray(PSOProgramService.JobContext_Key, ContextData);
 
 			params.putInt(PSOProgramService.JobID_Key, ThisJobID);
+			params.putLong(PSOProgramService.Priority_Key, PriorityInfo);
+
 			msg.replyTo = mReplyToMe;
 
 			JobResponse pendingResponse = new JobResponse();
@@ -518,39 +552,76 @@ public class PSOProgramServiceAccessor
 		return null;
 	}
 
-	public static JNIProgramLinkResponse AndroidThunkJava_VKPSOGFXCompile(byte[] ContextData, byte[] VertexShader, byte[] PixelShader, byte[] PSOData, byte[] PSOCacheData, boolean bAllowTimeOuts)
+	public static JNIProgramLinkResponse AndroidThunkJava_VKPSOGFXCompile(byte[] ContextData, long PriorityInfo, byte[] VertexShader, byte[] PixelShader, byte[] PSOData, byte[] PSOCacheData, boolean bAllowTimeOuts)
 	{
-		if(GameActivity.IsActivityPaused())
+		while( bIsReady.get() )
 		{
-			// one at a time when backgrounded.
-			synchronized (ProgramLinkLock)
+			try
+			{
+				if( MaxLiveRequestsSemaphore.tryAcquire(1, 100, TimeUnit.MILLISECONDS) )
+				{
+					JNIProgramLinkResponse returnResponse = null;
+
+					if(GameActivity.IsActivityPaused())
+					{
+						// one at a time when backgrounded.
+						synchronized (ProgramLinkLock)
+						{
+							returnResponse = VKPSOGFXCompile_internal(ContextData, PriorityInfo, VertexShader, PixelShader, PSOData, PSOCacheData, bAllowTimeOuts);
+						}
+					}
+					else
+					{
+						returnResponse = VKPSOGFXCompile_internal(ContextData, PriorityInfo, VertexShader, PixelShader, PSOData, PSOCacheData, bAllowTimeOuts);
+					}
+					MaxLiveRequestsSemaphore.release(1);
+					return returnResponse;
+				}
+			}
+			catch (InterruptedException e)
 			{
-				return VKPSOGFXCompile_internal(ContextData, VertexShader, PixelShader, PSOData, PSOCacheData, bAllowTimeOuts);
+				Log.error( "interrupted " + e);
 			}
 		}
-		else
-		{
-			return VKPSOGFXCompile_internal(ContextData, VertexShader, PixelShader, PSOData, PSOCacheData, bAllowTimeOuts);
-		}
+		Log.error("Failed to compile PSO services not ready.");
+		return null;
 	}
 
-	public static JNIProgramLinkResponse AndroidThunkJava_VKPSOGFXCompileShm(byte[] ContextData, int SharedMemFD, long VertexShaderSize, long PixelShaderSize, long PSODataSize, long PSOCacheDataSize, boolean bAllowTimeOuts)
+	public static JNIProgramLinkResponse AndroidThunkJava_VKPSOGFXCompileShm(byte[] ContextData, long PriorityInfo, int SharedMemFD, long VertexShaderSize, long PixelShaderSize, long PSODataSize, long PSOCacheDataSize, boolean bAllowTimeOuts)
 	{
-		if(GameActivity.IsActivityPaused())
+		while( bIsReady.get() )
 		{
-			// one at a time when backgrounded.
-			synchronized (ProgramLinkLock)
+			try
+			{
+				if( MaxLiveRequestsSemaphore.tryAcquire(1, 100, TimeUnit.MILLISECONDS) )
+				{
+					JNIProgramLinkResponse returnResponse = null;
+					if(GameActivity.IsActivityPaused())
+					{
+						// one at a time when backgrounded.
+						synchronized (ProgramLinkLock)
+						{
+							returnResponse = VKPSOGFXCompileShm_internal(ContextData, PriorityInfo, SharedMemFD, VertexShaderSize, PixelShaderSize, PSODataSize, PSOCacheDataSize, bAllowTimeOuts);
+						}
+					}
+					else
+					{
+						returnResponse = VKPSOGFXCompileShm_internal(ContextData, PriorityInfo, SharedMemFD, VertexShaderSize, PixelShaderSize, PSODataSize, PSOCacheDataSize, bAllowTimeOuts);
+					}
+					MaxLiveRequestsSemaphore.release(1);
+					return returnResponse;
+				}
+			}
+			catch (InterruptedException e)
 			{
-				return VKPSOGFXCompileShm_internal(ContextData, SharedMemFD, VertexShaderSize, PixelShaderSize, PSODataSize, PSOCacheDataSize, bAllowTimeOuts);
+				Log.error( "interrupted " + e);
 			}
 		}
-		else
-		{
-			return VKPSOGFXCompileShm_internal(ContextData, SharedMemFD, VertexShaderSize, PixelShaderSize, PSODataSize, PSOCacheDataSize, bAllowTimeOuts);
-		}
+		Log.error("Failed to compile PSO services not ready.");
+		return null;
 	}
 
-	private static JNIProgramLinkResponse VKPSOGFXCompile_internal(byte[] ContextData, byte[] VertexShader, byte[] PixelShader, byte[] PSOData, byte[] PSOCacheData, boolean bAllowTimeOuts)
+	private static JNIProgramLinkResponse VKPSOGFXCompile_internal(byte[] ContextData, long PriorityInfo, byte[] VertexShader, byte[] PixelShader, byte[] PSOData, byte[] PSOCacheData, boolean bAllowTimeOuts)
 	{
 		try
 		{
@@ -591,6 +662,7 @@ public class PSOProgramServiceAccessor
 			params.putByteArray(PSOProgramService.JobContext_Key, ContextData);
 
 			params.putInt(PSOProgramService.JobID_Key, ThisJobID);
+			params.putLong(PSOProgramService.Priority_Key, PriorityInfo);
 			msg.replyTo = mReplyToMe;
 
 			JobResponse pendingResponse = new JobResponse();
@@ -698,7 +770,7 @@ public class PSOProgramServiceAccessor
 		return null;
 	}
 
-	private static JNIProgramLinkResponse VKPSOGFXCompileShm_internal(byte[] ContextData, int SharedMemFD, long VertexShaderSize, long PixelShaderSize, long PSODataSize, long PSOCacheDataSize, boolean bAllowTimeOuts)
+	private static JNIProgramLinkResponse VKPSOGFXCompileShm_internal(byte[] ContextData, long PriorityInfo, int SharedMemFD, long VertexShaderSize, long PixelShaderSize, long PSODataSize, long PSOCacheDataSize, boolean bAllowTimeOuts)
 	{
 		ParcelFileDescriptor parcelFD = null;
 		try
@@ -744,6 +816,8 @@ public class PSOProgramServiceAccessor
 			params.putByteArray(PSOProgramService.JobContext_Key, ContextData);
 
 			params.putInt(PSOProgramService.JobID_Key, ThisJobID);
+			params.putLong(PSOProgramService.Priority_Key, PriorityInfo);
+			
 			msg.replyTo = mReplyToMe;
 
 			JobResponse pendingResponse = new JobResponse();
diff --git a/Engine/Build/Android/Prebuilt/PSOService/Source/PSOService.cpp b/Engine/Build/Android/Prebuilt/PSOService/Source/PSOService.cpp
--- a/Engine/Build/Android/Prebuilt/PSOService/Source/PSOService.cpp
+++ b/Engine/Build/Android/Prebuilt/PSOService/Source/PSOService.cpp
@@ -13,6 +13,9 @@
 #include <android/sharedmem.h>
 #include <sys/mman.h>
 #include <unistd.h>
+#include <pthread.h>
+#include <sys/resource.h>
+#include <dirent.h>
 
 #define ENABLETRACING 0
 
@@ -839,6 +842,128 @@ public:
 
 };
 
+static void SetAffinity(pid_t ThreadId, const cpu_set_t& DesiredAffinitySet)
+{
+	int rescode = sched_setaffinity(ThreadId, sizeof(DesiredAffinitySet), &DesiredAffinitySet);
+	if (rescode)
+	{
+		LOG_ERROR("set affinity %d, %d, %x, errno %d", rescode, ThreadId, *((int*)&DesiredAffinitySet), errno);
+	}
+#ifndef NDEBUG
+	cpu_set_t TestAffinitySet;
+	CPU_ZERO(&TestAffinitySet);
+	rescode = sched_getaffinity(ThreadId, sizeof(TestAffinitySet), &TestAffinitySet);
+	LOG_VERBOSE("affinity Info: tid %d, desired %x, set %x, rescode %d, errno %d", ThreadId, *((int*)&DesiredAffinitySet), *((int*)&TestAffinitySet), rescode, errno);
+#endif
+}
+
+static void SetAffinityAllThreads(const cpu_set_t& DesiredAffinity)
+{
+	// this is required as some drivers have additional threads which need the same treatment.
+	// we dont know what they are so all threads get hit, any new threads inherit the current settings.
+	DIR* SelfTaskDirectory;
+	struct dirent* Entry;
+
+	static const char ThreadDir[] = "/proc/self/task";
+	SelfTaskDirectory = opendir(ThreadDir);
+	if (SelfTaskDirectory != NULL)
+	{
+		while ((Entry = readdir(SelfTaskDirectory)))
+		{
+			pid_t tid = strtol(Entry->d_name, nullptr, 10);
+			if (tid)
+			{
+				SetAffinity(tid, DesiredAffinity);
+			}
+		}
+		closedir(SelfTaskDirectory);
+	}
+	else
+	{
+		LOG_ERROR("set affinity failed to find thread dir %s", ThreadDir);
+		SetAffinity(0, DesiredAffinity);
+	}
+}
+
+JNI_METHOD void Java_com_epicgames_unreal_psoservices_PSOProgramService_NativeSetThreadPriority(JNIEnv* jenv, jobject thiz, jlong PriInfoIn)
+{
+	struct PrecompilePriInfo
+	{
+		PrecompilePriInfo(uint64_t InfoIn) : PriInfo(InfoIn) {}
+		bool ShouldSetSchedPolicy() const	{ return PriInfo & (1 << 0); }
+		bool ShouldSetNice() const			{ return PriInfo & (1 << 1); }
+		bool ShouldSetAffinity() const		{ return PriInfo & (1 << 2); }
+
+		char GetSchedPolicy() const			{ return (PriInfo << 8) & 0xff; }
+		char GetSchedPolicyPri() const		{ return ((PriInfo << 16) & 0xff) - 128; }
+		char GetNice() const				{ return ((PriInfo << 24) & 0xff) - 128; }
+		uint32_t GetAffinity() const		{ return (PriInfo >> 32) & 0xFFFFFFFF; }
+
+		uint64_t PriInfo = 0;
+	};
+
+	PrecompilePriInfo PriInfo(PriInfoIn);
+
+	if(PriInfo.ShouldSetSchedPolicy())
+	{
+		int InitialPolicy;
+		int NewPolicy = PriInfo.GetSchedPolicy();
+		int SchedPri = PriInfo.GetSchedPolicyPri();
+
+		struct sched_param Sched = { };
+		pthread_t InThread = pthread_self();
+		int getres = pthread_getschedparam(InThread, &InitialPolicy, &Sched);
+
+		int primax = sched_get_priority_max(NewPolicy);
+		int primin = sched_get_priority_min(NewPolicy);
+
+		Sched.sched_priority = SchedPri < primin ? primin : (SchedPri > primax ? primax : SchedPri);
+
+		LOG_VERBOSE("tinfo initial policy %d, desired %d, getres %d, errno %d, pridesired %d, primin %d primax %d", InitialPolicy, NewPolicy, getres, errno, Sched.sched_priority, primin, primax);
+
+		int rescode = sched_setscheduler(0, NewPolicy, &Sched);
+		if (rescode)
+		{
+			LOG_ERROR("setsched error %d, errno %d", rescode, errno);
+		}
+	}
+
+	if (PriInfo.ShouldSetNice())
+	{
+		int Nice = PriInfo.GetNice();
+		int InitialNice = getpriority(PRIO_PROCESS, 0);
+		int rescode = setpriority(PRIO_PROCESS, 0, Nice);
+		int resultNice = getpriority(PRIO_PROCESS, 0);
+		if (rescode)
+		{
+			LOG_ERROR("setpriority failed. initial nice %d, desired %d, res %d, errno %d, result %d ", InitialNice, Nice, rescode, errno, resultNice);
+		}
+	}
+
+	if (PriInfo.ShouldSetAffinity())
+	{
+		const uint32_t AffinityMask = PriInfo.GetAffinity();
+
+		cpu_set_t DesiredAffinitySet;
+		CPU_ZERO(&DesiredAffinitySet);
+		if (AffinityMask == 0xFFFFFFFF)
+		{
+			memset(&DesiredAffinitySet, 0xff, sizeof(DesiredAffinitySet));
+		}
+		else
+		{
+			for (int i = 0; i < 32; i++)
+			{
+				if (AffinityMask & (1 << i))
+				{
+					CPU_SET(i, &DesiredAffinitySet);
+				}
+			}
+		}
+		SetAffinityAllThreads(DesiredAffinitySet);
+	}
+}
+
 JNI_METHOD void Java_com_epicgames_unreal_psoservices_PSOProgramService_InitVKDevice(JNIEnv* jenv, jobject thiz)
 {
 
diff --git a/Engine/Source/Runtime/OpenGLDrv/Private/Android/AndroidOpenGL.cpp b/Engine/Source/Runtime/OpenGLDrv/Private/Android/AndroidOpenGL.cpp
--- a/Engine/Source/Runtime/OpenGLDrv/Private/Android/AndroidOpenGL.cpp
+++ b/Engine/Source/Runtime/OpenGLDrv/Private/Android/AndroidOpenGL.cpp
@@ -15,6 +15,7 @@
 #include "Misc/ConfigCacheIni.h"
 #include "String/Find.h"
 #include "String/LexFromString.h"
+#include "Android/AndroidDynamicRHI.h"
 
 int32 FAndroidOpenGL::GLMajorVerion = 0;
 int32 FAndroidOpenGL::GLMinorVersion = 0;
@@ -117,7 +118,7 @@ struct FOpenGLRemoteGLProgramCompileJNI
 		CHECK_JNI_EXCEPTIONS(Env);
 		if(OGLServiceAccessor)
 		{
-			DispatchProgramLink = FJavaWrapper::FindStaticMethod(Env, OGLServiceAccessor, "AndroidThunkJava_OGLRemoteProgramLink", "([BLjava/lang/String;Ljava/lang/String;Ljava/lang/String;Z)Lcom/epicgames/unreal/psoservices/PSOProgramServiceAccessor$JNIProgramLinkResponse;", false);
+			DispatchProgramLink = FJavaWrapper::FindStaticMethod(Env, OGLServiceAccessor, "AndroidThunkJava_OGLRemoteProgramLink", "([BJLjava/lang/String;Ljava/lang/String;Ljava/lang/String;Z)Lcom/epicgames/unreal/psoservices/PSOProgramServiceAccessor$JNIProgramLinkResponse;", false);
 			CHECK_JNI_EXCEPTIONS(Env);
 			StartRemoteProgramLink = FJavaWrapper::FindStaticMethod(Env, OGLServiceAccessor, "AndroidThunkJava_StartRemoteProgramLink", "(IZZ)Z", false);
 			CHECK_JNI_EXCEPTIONS(Env);
@@ -1078,19 +1079,27 @@ void FAndroidOpenGL::ProcessExtensions(const FString& ExtensionsString)
 	}
 }
 
-static bool GRemoteCompileServicesStarted = false;
-static bool GRemoteCompileServicesActive = false;
+namespace AndroidOGLService
+{
+	std::atomic<bool> GRemoteCompileServicesStarted = false;
+	std::atomic<bool> GRemoteCompileServicesActive = false;
+	std::atomic<bool> bOneTimeErrorEncountered = false;
+	std::atomic<int> TotalErrors = 0;
+}
 extern bool AreAndroidOpenGLRemoteCompileServicesAvailable();
 
 bool FAndroidOpenGL::AreRemoteCompileServicesActive()
 {
-	if (GRemoteCompileServicesStarted && AreAndroidOpenGLRemoteCompileServicesAvailable())
+	// The services could be stopped at any point elsewhere, the return value is not guaranteed to be correct.
+	// it does not need to be exact as the PSO service will reject any new requests after service stop has been encountered.
+	// any existing PSOservice jobs will complete as normal.
+	if (AndroidOGLService::GRemoteCompileServicesStarted && AreAndroidOpenGLRemoteCompileServicesAvailable())
 	{
-		if (!GRemoteCompileServicesActive)
+		if (!AndroidOGLService::GRemoteCompileServicesActive)
 		{
 			JNIEnv* Env = FAndroidApplication::GetJavaEnv();
-			GRemoteCompileServicesActive = (bool)Env->CallStaticBooleanMethod(OpenGLRemoteGLProgramCompileJNI.OGLServiceAccessor, OpenGLRemoteGLProgramCompileJNI.AreProgramServicesReady);
-			if (!GRemoteCompileServicesActive)
+			AndroidOGLService::GRemoteCompileServicesActive = (bool)Env->CallStaticBooleanMethod(OpenGLRemoteGLProgramCompileJNI.OGLServiceAccessor, OpenGLRemoteGLProgramCompileJNI.AreProgramServicesReady);
+			if (!AndroidOGLService::GRemoteCompileServicesActive)
 			{
 				if ((bool)Env->CallStaticBooleanMethod(OpenGLRemoteGLProgramCompileJNI.OGLServiceAccessor, OpenGLRemoteGLProgramCompileJNI.HaveServicesFailed))
 				{
@@ -1098,8 +1107,12 @@ bool FAndroidOpenGL::AreRemoteCompileServicesActive()
 					StopRemoteCompileServices();
 				}
 			}
+			else
+			{
+				UE_LOG(LogRHI, Log, TEXT("Remote compile services are active."));
+			}
 		}
-		return GRemoteCompileServicesActive;
+		return AndroidOGLService::GRemoteCompileServicesActive;
 	}
 	return false;
 }
@@ -1108,40 +1121,39 @@ bool FAndroidOpenGL::StartRemoteCompileServices(int NumServices)
 {
 	JNIEnv* Env = FAndroidApplication::GetJavaEnv();
 
-	if (Env && AreAndroidOpenGLRemoteCompileServicesAvailable() && !GRemoteCompileServicesStarted)
+	if (Env && AreAndroidOpenGLRemoteCompileServicesAvailable() && !AndroidOGLService::GRemoteCompileServicesStarted)
 	{
 		bool bUseRobustContexts = AndroidEGL::GetInstance()->IsUsingRobustContext();
-		GRemoteCompileServicesStarted = (bool)Env->CallStaticBooleanMethod(OpenGLRemoteGLProgramCompileJNI.OGLServiceAccessor, OpenGLRemoteGLProgramCompileJNI.StartRemoteProgramLink, (jint)NumServices, (jboolean)bUseRobustContexts, /*bUseVulkan*/(jboolean)false);
+		AndroidOGLService::GRemoteCompileServicesStarted = (bool)Env->CallStaticBooleanMethod(OpenGLRemoteGLProgramCompileJNI.OGLServiceAccessor, OpenGLRemoteGLProgramCompileJNI.StartRemoteProgramLink, (jint)NumServices, (jboolean)bUseRobustContexts, /*bUseVulkan*/(jboolean)false);
 	}
 
-	return GRemoteCompileServicesStarted;
+	return AndroidOGLService::GRemoteCompileServicesStarted;
 }
 
 void FAndroidOpenGL::StopRemoteCompileServices()
 {
-	GRemoteCompileServicesActive = false;
-	GRemoteCompileServicesStarted = false;
-	JNIEnv* Env = FAndroidApplication::GetJavaEnv();
-
-	if (Env && ensure(AreAndroidOpenGLRemoteCompileServicesAvailable()))
+	bool bExpected = true;
+	if (AndroidOGLService::GRemoteCompileServicesStarted.compare_exchange_strong(bExpected, false))
 	{
-		Env->CallStaticVoidMethod(OpenGLRemoteGLProgramCompileJNI.OGLServiceAccessor, OpenGLRemoteGLProgramCompileJNI.StopRemoteProgramLink);
-	}
-}
+		UE_LOG(LogRHI, Log, TEXT("Stopping Remote Compile Services"));
+		AndroidOGLService::GRemoteCompileServicesActive = false;
+		JNIEnv* Env = FAndroidApplication::GetJavaEnv();
 
-namespace AndroidOGLService
-{
-	std::atomic<bool> bOneTimeErrorEncountered = false;
+		if (Env && ensure(AreAndroidOpenGLRemoteCompileServicesAvailable()))
+		{
+			Env->CallStaticVoidMethod(OpenGLRemoteGLProgramCompileJNI.OGLServiceAccessor, OpenGLRemoteGLProgramCompileJNI.StopRemoteProgramLink);
+		}
+	}
 }
 
-TArray<uint8> FAndroidOpenGL::DispatchAndWaitForRemoteGLProgramCompile(const TArrayView<uint8> ContextData, const TArray<ANSICHAR>& VertexGlslCode, const TArray<ANSICHAR>& PixelGlslCode, const TArray<ANSICHAR>& ComputeGlslCode, FString& FailureMessageOUT)
+TArray<uint8> FAndroidOpenGL::DispatchAndWaitForRemoteGLProgramCompile(FGraphicsPipelineStateInitializer::EPSOPrecacheCompileType PSOCompileType, const TArrayView<uint8> ContextData, const TArray<ANSICHAR>& VertexGlslCode, const TArray<ANSICHAR>& PixelGlslCode, const TArray<ANSICHAR>& ComputeGlslCode, FString& FailureMessageOUT)
 {
 	bool bResult = false;
 	JNIEnv* Env = FAndroidApplication::GetJavaEnv();
 	TArray<uint8> CompiledProgramBinary;
 	FString ErrorMessage;
 
-	if (Env && ensure(GRemoteCompileServicesActive) && ensure(AreAndroidOpenGLRemoteCompileServicesAvailable()))
+	if (Env && ensure(AndroidOGLService::GRemoteCompileServicesActive) && ensure(AreAndroidOpenGLRemoteCompileServicesAvailable()))
 	{
 		// todo: double conversion :(
 		auto jVS = NewScopedJavaObject(Env, Env->NewStringUTF(TCHAR_TO_UTF8(ANSI_TO_TCHAR(VertexGlslCode.IsEmpty() ? "" : VertexGlslCode.GetData()))));
@@ -1151,7 +1163,8 @@ TArray<uint8> FAndroidOpenGL::DispatchAndWaitForRemoteGLProgramCompile(const TAr
 		Env->SetByteArrayRegion(*ProgramKeyBuffer, 0, ContextData.Num(), reinterpret_cast<const jbyte*>(ContextData.GetData()));
 		// dont time out if the debugger is attached.
 		bool bEnableTimeOuts = !FPlatformMisc::IsDebuggerPresent();
-		auto ProgramResponseObj = NewScopedJavaObject(Env, Env->CallStaticObjectMethod(OpenGLRemoteGLProgramCompileJNI.OGLServiceAccessor, OpenGLRemoteGLProgramCompileJNI.DispatchProgramLink, *ProgramKeyBuffer, *jVS, *jPS, *jCS, bEnableTimeOuts));
+		FPlatformDynamicRHI::FPSOServicePriInfo PriorityInfo(PSOCompileType);
+		auto ProgramResponseObj = NewScopedJavaObject(Env, Env->CallStaticObjectMethod(OpenGLRemoteGLProgramCompileJNI.OGLServiceAccessor, OpenGLRemoteGLProgramCompileJNI.DispatchProgramLink, *ProgramKeyBuffer, PriorityInfo.GetPriorityInfo(), *jVS, *jPS, *jCS, bEnableTimeOuts));
  		CHECK_JNI_EXCEPTIONS(Env);
 
 		if(*ProgramResponseObj)
@@ -1183,6 +1196,16 @@ TArray<uint8> FAndroidOpenGL::DispatchAndWaitForRemoteGLProgramCompile(const TAr
 			}
 			FailureMessageOUT = TEXT("Remote compiler failed.");
 		}
+
+		if(CompiledProgramBinary.IsEmpty())
+		{
+			check(!FailureMessageOUT.IsEmpty());
+			if ((AndroidOGLService::TotalErrors++) == FPlatformDynamicRHI::GetPSOServiceFailureThreshold())
+			{
+				StopRemoteCompileServices();
+				FailureMessageOUT = TEXT("Remote compiler passed failure threshold, disabling further remote compiles.");
+			}
+		}
 	}
 	return CompiledProgramBinary;
 }
diff --git a/Engine/Source/Runtime/OpenGLDrv/Private/Android/AndroidOpenGL.h b/Engine/Source/Runtime/OpenGLDrv/Private/Android/AndroidOpenGL.h
--- a/Engine/Source/Runtime/OpenGLDrv/Private/Android/AndroidOpenGL.h
+++ b/Engine/Source/Runtime/OpenGLDrv/Private/Android/AndroidOpenGL.h
@@ -140,7 +140,7 @@ struct FAndroidOpenGL : public FOpenGLES
 	static EImageExternalType ImageExternalType;
 
 	/* interface to remote GLES program compiler */
-	static TArray<uint8> DispatchAndWaitForRemoteGLProgramCompile(const TArrayView<uint8> ContextData, const TArray<ANSICHAR>& VertexGlslCode, const TArray<ANSICHAR>& PixelGlslCode, const TArray<ANSICHAR>& ComputeGlslCode, FString& FailureMessageOUT);
+	static TArray<uint8> DispatchAndWaitForRemoteGLProgramCompile(FGraphicsPipelineStateInitializer::EPSOPrecacheCompileType PSOCompileType,const TArrayView<uint8> ContextData, const TArray<ANSICHAR>& VertexGlslCode, const TArray<ANSICHAR>& PixelGlslCode, const TArray<ANSICHAR>& ComputeGlslCode, FString& FailureMessageOUT);
 	static bool AreRemoteCompileServicesActive();
 	static bool StartRemoteCompileServices(int NumServices);
 	static void StopRemoteCompileServices();
diff --git a/Engine/Source/Runtime/OpenGLDrv/Private/OpenGLShaders.cpp b/Engine/Source/Runtime/OpenGLDrv/Private/OpenGLShaders.cpp
--- a/Engine/Source/Runtime/OpenGLDrv/Private/OpenGLShaders.cpp
+++ b/Engine/Source/Runtime/OpenGLDrv/Private/OpenGLShaders.cpp
@@ -2661,7 +2661,7 @@ static bool CanCreateExternally(bool bIsFromPSO)
 	return false;
 }
 
-static FOpenGLProgramBinary ExternalProgramCompile(const FOpenGLProgramKey& ProgramKey, FRHIVertexShader* VertexShaderRHI, FRHIPixelShader* PixelShaderRHI)
+static FOpenGLProgramBinary ExternalProgramCompile(const FOpenGLProgramKey& ProgramKey, FGraphicsPipelineStateInitializer::EPSOPrecacheCompileType PSOCompileType, FRHIVertexShader* VertexShaderRHI, FRHIPixelShader* PixelShaderRHI)
 {
 	FOpenGLProgramBinary CompiledProgram;
 #if PLATFORM_ANDROID
@@ -2679,7 +2679,7 @@ static FOpenGLProgramBinary ExternalProgramCompile(const FOpenGLProgramKey& Prog
 	}
 
 	FString FailLog;
-	TArray<uint8> CompiledProgramBytes = FAndroidOpenGL::DispatchAndWaitForRemoteGLProgramCompile(TArrayView<uint8>((uint8*)&ProgramKey, sizeof(ProgramKey)), VSCode, PSCode, ComputeGlslCode, FailLog);
+	TArray<uint8> CompiledProgramBytes = FAndroidOpenGL::DispatchAndWaitForRemoteGLProgramCompile(PSOCompileType, TArrayView<uint8>((uint8*)&ProgramKey, sizeof(ProgramKey)), VSCode, PSCode, ComputeGlslCode, FailLog);
 
 	if (FailLog.IsEmpty())
 	{
@@ -2752,7 +2752,7 @@ void FOpenGLDynamicRHI::PrepareGFXBoundShaderState(const FGraphicsPipelineStateI
 		if (FOpenGLProgramBinaryCache::IsBuildingCache())
 		{
 			OGL_BINARYCACHE_STATS_MARKBEGINCOMPILE(ProgramKey);
-			FOpenGLProgramBinary CompiledProgram = ExternalProgramCompile(ProgramKey, VertexShaderRHI, PixelShaderRHI);
+			FOpenGLProgramBinary CompiledProgram = ExternalProgramCompile(ProgramKey, Initializer.GetPSOPrecacheCompileType(), VertexShaderRHI, PixelShaderRHI);
 
 			if (CompiledProgram.IsValid())
 			{
diff --git a/Engine/Source/Runtime/RHI/Private/Android/AndroidDynamicRHI.cpp b/Engine/Source/Runtime/RHI/Private/Android/AndroidDynamicRHI.cpp
--- a/Engine/Source/Runtime/RHI/Private/Android/AndroidDynamicRHI.cpp
+++ b/Engine/Source/Runtime/RHI/Private/Android/AndroidDynamicRHI.cpp
@@ -1,10 +1,100 @@
 // Copyright Epic Games, Inc. All Rights Reserved.
 
-#include "RHI.h"
+#include "Android/AndroidDynamicRHI.h"
 #include "Modules/ModuleManager.h"
 #include "Android/AndroidApplication.h"
 #include "Misc/App.h"
 
+namespace UE
+{
+	namespace FAndroidPlatformDynamicRHI
+	{
+		using EPSOPrecacheCompileType = FGraphicsPipelineStateInitializer::EPSOPrecacheCompileType;
+		static_assert((int)EPSOPrecacheCompileType::MaxPri == 3
+			&& (int)EPSOPrecacheCompileType::MinPri == 1
+			&& (int)EPSOPrecacheCompileType::NumTypes == 4, "Modifications maybe required if the number precache priorities changes.");
+
+		static bool bEnablePSOSchedulingParams = true;
+		static FAutoConsoleVariableRef CVarDisablePSOSchedulingParams(
+			TEXT("android.PSOService.EnableSchedulingParams"),
+			bEnablePSOSchedulingParams,
+			TEXT("Whether to set scheduler values (such as nice and scheduler policy) for each precache process job\n")
+			TEXT(" 1 (default)")
+			,
+			ECVF_RenderThreadSafe
+		);
+		static int32 PSOPrecacheAffinity[(int)EPSOPrecacheCompileType::NumTypes] = {};
+		static FAutoConsoleVariableRef CVarMinPriPSOPrecacheAffinity(
+			TEXT("android.PSOService.MinPriPSOPrecacheAffinity"),
+			PSOPrecacheAffinity[(int)EPSOPrecacheCompileType::MinPri],
+			TEXT("CPU affinity to use when compiling low priority precache PSOs via external compilers\n")
+			TEXT(" 0: all cpus (default)\n")
+			TEXT(" all other values represent a 32 bit mask of cpu affinity.")
+			,
+			ECVF_RenderThreadSafe
+		);
+		static FAutoConsoleVariableRef CVarNormalPriPSOPrecacheAffinity(
+			TEXT("android.PSOService.NormalPriPSOPrecacheAffinity"),
+			PSOPrecacheAffinity[(int)EPSOPrecacheCompileType::NormalPri],
+			TEXT("CPU affinity to use when compiling normal priority precache PSOs via external compilers\n")
+			TEXT(" 0: all cpus (default)\n")
+			TEXT(" all other values represent a 32 bit mask of cpu affinity.")
+			,
+			ECVF_RenderThreadSafe
+		);
+		static FAutoConsoleVariableRef CVarMaxPriPSOPrecacheAffinity(
+			TEXT("android.PSOService.MaxPriPSOPrecacheAffinity"),
+			PSOPrecacheAffinity[(int)EPSOPrecacheCompileType::MaxPri],
+			TEXT("CPU affinity to use when compiling high priority precache PSOs via external compilers\n")
+			TEXT(" 0: all cpus (default)\n")
+			TEXT(" all other values represent a 32 bit mask of cpu affinity.")
+			,
+			ECVF_RenderThreadSafe
+		);
+
+		static int32 GExternalCompilerFailureThreshold = 5;
+		static FAutoConsoleVariableRef GExternalCompilerFailureThresholdCvar(
+			TEXT("android.PSOService.ExternalCompilerFailureThreshold"),
+			GExternalCompilerFailureThreshold,
+			TEXT("Number of external PSO compiler failures to ignore before disabling external PSO compiling altogether.\n")
+			TEXT("  default: 5"),
+			ECVF_Default | ECVF_RenderThreadSafe
+		);
+
+		static bool ArePSOPrecacheAffinitiesSet()
+		{
+			return PSOPrecacheAffinity[(int)EPSOPrecacheCompileType::MinPri] != 0
+				|| PSOPrecacheAffinity[(int)EPSOPrecacheCompileType::NormalPri] != 0
+				|| PSOPrecacheAffinity[(int)EPSOPrecacheCompileType::MaxPri] != 0;
+		}
+
+		FPSOServicePriInfo::FPSOServicePriInfo(EPSOPrecacheCompileType PSOCompileType)
+		{
+			if (ArePSOPrecacheAffinitiesSet() && ensure(PSOCompileType >= EPSOPrecacheCompileType::MinPri && PSOCompileType <= EPSOPrecacheCompileType::MaxPri))
+			{
+				const uint32 CVarValue = PSOPrecacheAffinity[(int)PSOCompileType];
+				const uint32 Affinity = CVarValue ? CVarValue : 0xFFFFFFFF;
+
+				SetAffinity(Affinity);
+
+				if (bEnablePSOSchedulingParams)
+				{
+					static const int32 NiceValuesPerPri[(int)EPSOPrecacheCompileType::NumTypes] = { 0,-19,0,10 };
+					SetNice(NiceValuesPerPri[(int)PSOCompileType]);
+
+					static const char SchedPolicy[(int)EPSOPrecacheCompileType::NumTypes] = { SCHED_IDLE, SCHED_IDLE, SCHED_NORMAL, SCHED_NORMAL };
+					SetSchedPolicy(SchedPolicy[(int)PSOCompileType], 0); // sched_priority=0 as generally not used with scheduling policies we care about.
+				}
+			}
+		}
+
+		int32 GetPSOServiceFailureThreshold()
+		{
+			return GExternalCompilerFailureThreshold;
+		}
+	} // namespace FAndroidPlatformDynamicRHI
+} // namespace UE
+
 FDynamicRHI* PlatformCreateDynamicRHI()
 {
 	FDynamicRHI* DynamicRHI = NULL;
diff --git a/Engine/Source/Runtime/RHI/Private/PipelineStateCache.cpp b/Engine/Source/Runtime/RHI/Private/PipelineStateCache.cpp
--- a/Engine/Source/Runtime/RHI/Private/PipelineStateCache.cpp
+++ b/Engine/Source/Runtime/RHI/Private/PipelineStateCache.cpp
@@ -281,7 +281,7 @@ static FAutoConsoleVariableRef GPSOPrecompileThreadPoolSizeMaxVar(
 );
 
 int32 GPSOPrecompileThreadPoolThreadPriority = (int32)EThreadPriority::TPri_BelowNormal;
-static FAutoConsoleVariableRef CVarStreamingTextureIOPriority(
+static FAutoConsoleVariableRef CVarPrecompileThreadPoolThreadPriority(
 	TEXT("r.pso.PrecompileThreadPoolThreadPriority"),
 	GPSOPrecompileThreadPoolThreadPriority,
 	TEXT("Thread priority for the PSO precompile pool"),
@@ -2179,6 +2179,13 @@ public:
 					}
 				}
 
+			    const bool bAbortPSOCompileDueToShutdown = IsEngineExitRequested() && Initializer.bPSOPrecache;
+			    if (bAbortPSOCompileDueToShutdown)
+			    {
+				    UE_LOG(LogRHI, Verbose, TEXT("Skipping a precache compile due to engine shutdown."));
+				    bSkipCreation = true;
+			    }
+
 				FGraphicsPipelineState* GfxPipeline = static_cast<FGraphicsPipelineState*>(Pipeline);
 
 				uint64 StartTime = FPlatformTime::Cycles64();
@@ -2189,7 +2196,7 @@ public:
 				{
 					GfxPipeline->SortKey = GfxPipeline->RHIPipeline->GetSortKey();
 				}
-				else
+				else if(!bAbortPSOCompileDueToShutdown)
 				{
 					HandlePipelineCreationFailure(Initializer);
 				}
@@ -3149,6 +3156,7 @@ void FPrecacheComputePipelineCache::OnNewPipelineStateCreated(const FPrecacheCom
 		check(CachedState->CompletionEvent != nullptr);
 		FGraphicsPipelineStateInitializer GraphicsPipelineStateInitializer;
 		GraphicsPipelineStateInitializer.bPSOPrecache = true;
+		GraphicsPipelineStateInitializer.SetPSOPrecacheCompileType(FGraphicsPipelineStateInitializer::EPSOPrecacheCompileType::NormalPri);
 
 		FPSOCompilationDebugData PSOCompilationDebugData;
 		PSOCompilationDebugData.PSOCompilationEventName = PSOCompilationEventName;
@@ -3221,16 +3229,15 @@ void FPrecacheGraphicsPipelineCache::OnNewPipelineStateCreated(const FGraphicsPi
 	check((NewGraphicsPipelineState->CompletionEvent != nullptr) == bDoAsyncCompile);
 
 	// Mark as precache so it will try and use the background thread pool if available
-	bool bPSOPrecache = true;
-
 	FGraphicsPipelineStateInitializer InitializerCopy(Initializer);
-	InitializerCopy.bPSOPrecache = bPSOPrecache;
+	InitializerCopy.SetPSOPrecacheCompileType(FGraphicsPipelineStateInitializer::EPSOPrecacheCompileType::NormalPri);
+	InitializerCopy.bPSOPrecache = true;
 
 	FPSOCompilationDebugData PSOCompilationDebugData;
 	PSOCompilationDebugData.PSOCompilationEventName = PSOCompilationEventName;
 
 	// Start the precache task	
-	InternalCreateGraphicsPipelineState(InitializerCopy, EPSOPrecacheResult::Active, bDoAsyncCompile, bPSOPrecache, NewGraphicsPipelineState, PSOCompilationDebugData, false);
+	InternalCreateGraphicsPipelineState(InitializerCopy, EPSOPrecacheResult::Active, bDoAsyncCompile, InitializerCopy.bPSOPrecache, NewGraphicsPipelineState, PSOCompilationDebugData, false);
 }
 
 FPSOPrecacheRequestResult PipelineStateCache::PrecacheGraphicsPipelineState(const FGraphicsPipelineStateInitializer& Initializer)
diff --git a/Engine/Source/Runtime/RHI/Public/Android/AndroidDynamicRHI.h b/Engine/Source/Runtime/RHI/Public/Android/AndroidDynamicRHI.h
new file mode 100644
--- /dev/null
+++ b/Engine/Source/Runtime/RHI/Public/Android/AndroidDynamicRHI.h
@@ -0,0 +1,28 @@
+// Copyright Epic Games, Inc. All Rights Reserved.
+#pragma once
+
+#include "RHI.h"
+
+namespace UE
+{
+	namespace FAndroidPlatformDynamicRHI
+	{
+		// Helper class to convert FGraphicsPipelineStateInitializer::EPSOPrecacheCompileType to the uint64 required by android's external PSO service.
+		class FPSOServicePriInfo
+		{
+			using EPSOPrecacheCompileType = FGraphicsPipelineStateInitializer::EPSOPrecacheCompileType;
+		public:
+			FPSOServicePriInfo(EPSOPrecacheCompileType PSOCompileType);
+			uint64 GetPriorityInfo() const { return PriInfo; }
+		private:
+			void SetSchedPolicy(char SchedPolicy, char SchedPri) { check(SchedPri < 128); PriInfo = PriInfo | 1 << 0 | (SchedPolicy << 8) | ((SchedPri + 128) & 0xFF) << 16; }
+			void SetNice(char Nice) { check(Nice < 128); PriInfo = PriInfo | 1 << 1 | ((Nice + 128) & 0xFF) << 24; }
+			void SetAffinity(uint32 AffinityMask) { PriInfo = PriInfo | 1 << 2 | ((uint64)AffinityMask << 32); }
+			uint64 PriInfo = 0;
+		};
+
+		int32 GetPSOServiceFailureThreshold();
+	} // namespace FAndroidPlatformDynamicRHI
+} // namespace UE
+
+namespace FPlatformDynamicRHI = UE::FAndroidPlatformDynamicRHI;
diff --git a/Engine/Source/Runtime/RHI/Public/RHIResources.h b/Engine/Source/Runtime/RHI/Public/RHIResources.h
--- a/Engine/Source/Runtime/RHI/Public/RHIResources.h
+++ b/Engine/Source/Runtime/RHI/Public/RHIResources.h
@@ -4273,13 +4273,31 @@ public:
 	{
 		struct
 		{
-			uint16					Reserved			: 14;
-			uint16					bPSOPrecache			: 1;
+			uint16					Reserved			: 11;
+			uint16					bPSOPrecache		: 1;
 			uint16					bFromPSOFileCache	: 1;
+			uint16					PrecacheCompileType : 3;
 		};
 		uint16						Flags;
 	};
 
+	enum class EPSOPrecacheCompileType : uint8
+	{
+		NotSet = 0,
+		MinPri = 1,
+		NormalPri = 2,
+		MaxPri = 3,
+
+		NumTypes = 4,
+	};
+	static_assert((int)EPSOPrecacheCompileType::MaxPri < (1<<3) ); // ensure MaxPri fits within PrecacheCompileType
+	void SetPSOPrecacheCompileType(EPSOPrecacheCompileType PrecacheCompileTypeIN) 
+	{ 
+		check(PrecacheCompileTypeIN <= EPSOPrecacheCompileType::MaxPri && PrecacheCompileTypeIN >= EPSOPrecacheCompileType::MinPri);
+		PrecacheCompileType = (uint16)PrecacheCompileTypeIN;
+	}
+	EPSOPrecacheCompileType GetPSOPrecacheCompileType() const {	return (EPSOPrecacheCompileType)PrecacheCompileType; }
+
 	// Cached hash off all state data provided at creation time (Only contains hash of data which influences the PSO precaching for the current platform)
 	// Created from hashing the state data instead of the pointers which are used during fast runtime cache checking and compares
 	uint64							StatePrecachePSOHash;
diff --git a/Engine/Source/Runtime/RenderCore/Private/ShaderPipelineCache.cpp b/Engine/Source/Runtime/RenderCore/Private/ShaderPipelineCache.cpp
--- a/Engine/Source/Runtime/RenderCore/Private/ShaderPipelineCache.cpp
+++ b/Engine/Source/Runtime/RenderCore/Private/ShaderPipelineCache.cpp
@@ -1088,6 +1088,7 @@ bool FShaderPipelineCacheTask::Precompile(FRHICommandListImmediate& RHICmdList,
 				// This indicates we do not want a fatal error if this compilation fails
 				// (ie, if this entry in the file cache is bad)
 				GraphicsInitializer.bFromPSOFileCache = true;
+				GraphicsInitializer.SetPSOPrecacheCompileType(FGraphicsPipelineStateInitializer::EPSOPrecacheCompileType::NormalPri);
 
 #if !UE_BUILD_SHIPPING
 				// dump to log to describe
diff --git a/Engine/Source/Runtime/VulkanRHI/Private/Android/VulkanAndroidPlatform.cpp b/Engine/Source/Runtime/VulkanRHI/Private/Android/VulkanAndroidPlatform.cpp
--- a/Engine/Source/Runtime/VulkanRHI/Private/Android/VulkanAndroidPlatform.cpp
+++ b/Engine/Source/Runtime/VulkanRHI/Private/Android/VulkanAndroidPlatform.cpp
@@ -18,6 +18,7 @@
 #include <android/sharedmem_jni.h>
 #include <sys/mman.h>
 #include "ProfilingDebugging/ScopedTimers.h"
+#include "Android/AndroidDynamicRHI.h"
 
 #if USE_ANDROID_SWAPPY
 #undef VK_NO_PROTOTYPES
@@ -1346,8 +1347,13 @@ void PipelineToBinary(FVulkanDevice* Device, const VkGraphicsPipelineCreateInfo*
 #define CHECK_JNI_EXCEPTIONS(env)  if (env->ExceptionCheck()) {env->ExceptionDescribe();env->ExceptionClear();}
 #endif
 
-static bool GRemoteCompileServicesStarted = false;
-static bool GRemoteCompileServicesActive = false;
+namespace AndroidVulkanService
+{
+	std::atomic<bool> GRemoteCompileServicesStarted = false;
+	std::atomic<bool> GRemoteCompileServicesActive = false;
+	std::atomic<bool> bOneTimeErrorEncountered = false;
+	std::atomic<int> TotalErrors = 0;
+}
 
 struct FVKRemoteProgramCompileJNI
 {
@@ -1385,9 +1391,9 @@ struct FVKRemoteProgramCompileJNI
 		CHECK_JNI_EXCEPTIONS(Env);
 		if (PSOServiceAccessor)
 		{
-			DispatchPSOCompile = FJavaWrapper::FindStaticMethod(Env, PSOServiceAccessor, "AndroidThunkJava_VKPSOGFXCompile", "([B[B[B[B[BZ)Lcom/epicgames/unreal/psoservices/PSOProgramServiceAccessor$JNIProgramLinkResponse;", false);
+			DispatchPSOCompile = FJavaWrapper::FindStaticMethod(Env, PSOServiceAccessor, "AndroidThunkJava_VKPSOGFXCompile", "([BJ[B[B[B[BZ)Lcom/epicgames/unreal/psoservices/PSOProgramServiceAccessor$JNIProgramLinkResponse;", false);
 			CHECK_JNI_EXCEPTIONS(Env);
-			DispatchPSOCompileShm = FJavaWrapper::FindStaticMethod(Env, PSOServiceAccessor, "AndroidThunkJava_VKPSOGFXCompileShm", "([BIJJJJZ)Lcom/epicgames/unreal/psoservices/PSOProgramServiceAccessor$JNIProgramLinkResponse;", false);
+			DispatchPSOCompileShm = FJavaWrapper::FindStaticMethod(Env, PSOServiceAccessor, "AndroidThunkJava_VKPSOGFXCompileShm", "([BJIJJJJZ)Lcom/epicgames/unreal/psoservices/PSOProgramServiceAccessor$JNIProgramLinkResponse;", false);
 			CHECK_JNI_EXCEPTIONS(Env);
 			StartRemoteProgramLink = FJavaWrapper::FindStaticMethod(Env, PSOServiceAccessor, "AndroidThunkJava_StartRemoteProgramLink", "(IZZ)Z", false);
 			CHECK_JNI_EXCEPTIONS(Env);
@@ -1434,13 +1440,16 @@ static bool AreAndroidVulkanRemoteCompileServicesAvailable()
 
 bool FVulkanAndroidPlatform::AreRemoteCompileServicesActive()
 {
-	if (GRemoteCompileServicesStarted && AreAndroidVulkanRemoteCompileServicesAvailable())
+	// The services could be stopped at any point elsewhere, the return value is not guaranteed to be correct.
+	// it does not need to be exact as the PSO service will reject any new requests after service stop has been encountered.
+	// any existing PSOservice jobs will complete as normal.
+	if (AndroidVulkanService::GRemoteCompileServicesStarted && AreAndroidVulkanRemoteCompileServicesAvailable())
 	{
-		if (!GRemoteCompileServicesActive)
+		if (!AndroidVulkanService::GRemoteCompileServicesActive)
 		{
 			JNIEnv* Env = FAndroidApplication::GetJavaEnv();
-			GRemoteCompileServicesActive = (bool)Env->CallStaticBooleanMethod(VKRemoteProgramCompileJNI.PSOServiceAccessor, VKRemoteProgramCompileJNI.AreProgramServicesReady);
-			if (!GRemoteCompileServicesActive)
+			AndroidVulkanService::GRemoteCompileServicesActive = (bool)Env->CallStaticBooleanMethod(VKRemoteProgramCompileJNI.PSOServiceAccessor, VKRemoteProgramCompileJNI.AreProgramServicesReady);
+			if (!AndroidVulkanService::GRemoteCompileServicesActive)
 			{
 				if ((bool)Env->CallStaticBooleanMethod(VKRemoteProgramCompileJNI.PSOServiceAccessor, VKRemoteProgramCompileJNI.HaveServicesFailed))
 				{
@@ -1448,8 +1457,12 @@ bool FVulkanAndroidPlatform::AreRemoteCompileServicesActive()
 					StopRemoteCompileServices();
 				}
 			}
+			else
+			{
+				UE_LOG(LogRHI, Log, TEXT("Remote compile services are active."));
+			}
 		}
-		return GRemoteCompileServicesActive;
+		return AndroidVulkanService::GRemoteCompileServicesActive;
 	}
 	return false;
 }
@@ -1460,41 +1473,54 @@ bool FVulkanAndroidPlatform::StartRemoteCompileServices(int NumServices)
 	
 	VKRemoteProgramCompileJNI.Init(Env);
 
-	if (Env && AreAndroidVulkanRemoteCompileServicesAvailable() && !GRemoteCompileServicesStarted)
+	if (Env && AreAndroidVulkanRemoteCompileServicesAvailable() && !AndroidVulkanService::GRemoteCompileServicesStarted)
 	{
-		GRemoteCompileServicesStarted = (bool)Env->CallStaticBooleanMethod(VKRemoteProgramCompileJNI.PSOServiceAccessor, VKRemoteProgramCompileJNI.StartRemoteProgramLink, (jint)NumServices, /*bUseRobustEGLContext*/(jboolean)false, /*bUseVulkan*/(jboolean)true);
+		AndroidVulkanService::GRemoteCompileServicesStarted = (bool)Env->CallStaticBooleanMethod(VKRemoteProgramCompileJNI.PSOServiceAccessor, VKRemoteProgramCompileJNI.StartRemoteProgramLink, (jint)NumServices, /*bUseRobustEGLContext*/(jboolean)false, /*bUseVulkan*/(jboolean)true);
 	}
 
-	return GRemoteCompileServicesStarted;
+	return AndroidVulkanService::GRemoteCompileServicesStarted;
 }
 
 void FVulkanAndroidPlatform::StopRemoteCompileServices()
 {
-	GRemoteCompileServicesStarted = false;
-	GRemoteCompileServicesActive = false;
-	JNIEnv* Env = FAndroidApplication::GetJavaEnv();
-
-	if (Env && ensure(AreAndroidVulkanRemoteCompileServicesAvailable()))
+	bool bExpected = true;
+	if(AndroidVulkanService::GRemoteCompileServicesStarted.compare_exchange_strong(bExpected, false) )
 	{
-		Env->CallStaticVoidMethod(VKRemoteProgramCompileJNI.PSOServiceAccessor, VKRemoteProgramCompileJNI.StopRemoteProgramLink);
-	}
-}
+		UE_LOG(LogVulkanRHI, Log, TEXT("Stopping Remote Compile Services"));
+		AndroidVulkanService::GRemoteCompileServicesActive = false;
+		JNIEnv* Env = FAndroidApplication::GetJavaEnv();
 
-namespace AndroidVulkanService
-{
-	std::atomic<bool> bOneTimeErrorEncountered = false;
+		if (Env && ensure(AreAndroidVulkanRemoteCompileServicesAvailable()))
+		{
+			Env->CallStaticVoidMethod(VKRemoteProgramCompileJNI.PSOServiceAccessor, VKRemoteProgramCompileJNI.StopRemoteProgramLink);
+		}
+	}
 }
 
-VkPipelineCache FVulkanAndroidPlatform::PrecompilePSO(FVulkanDevice* Device, const TArrayView<uint8> OptionalPSOCacheData, const VkGraphicsPipelineCreateInfo* PipelineInfo, FGfxPipelineDesc* GfxEntry, const FVulkanRenderTargetLayout* RTLayout, TArrayView<uint32_t> VS, TArrayView<uint32_t> PS, size_t& AfterSize)
+VkPipelineCache FVulkanAndroidPlatform::PrecompilePSO(
+		FVulkanDevice* Device, 
+		const TArrayView<uint8> OptionalPSOCacheData, 
+		FGraphicsPipelineStateInitializer::EPSOPrecacheCompileType PSOCompileType,
+		const VkGraphicsPipelineCreateInfo* PipelineInfo, 
+		FGfxPipelineDesc* GfxEntry, 
+		const FVulkanRenderTargetLayout* RTLayout, 
+		TArrayView<uint32_t> VS, 
+		TArrayView<uint32_t> PS, 
+		size_t& AfterSize,
+		FString* FailureMessageOUT
+	)
 {
-	FString FailureMessageOUT;
-	
-	if (!AreRemoteCompileServicesActive())
+	if (!ensure(AreRemoteCompileServicesActive()))
 	{
 		return VK_NULL_HANDLE;
 	}
 	QUICK_SCOPE_CYCLE_COUNTER(STAT_VulkanAndroid_PrecompilePSO);
 
+	auto SetFailureMessage = [FailureMessageOUT](FString&& Msg) { if (FailureMessageOUT) { *FailureMessageOUT = Msg; }};
+
+	VkPipelineCache ReturnPipelineCache = VK_NULL_HANDLE;
+	FPlatformDynamicRHI::FPSOServicePriInfo PriorityInfo(PSOCompileType);
+
 // 	FScopedDurationTimeLogger Timer(TEXT("FVulkanAndroidPlatform::PrecompilePSO"));
 
 	TArray<char> MemoryStream;
@@ -1556,7 +1582,7 @@ VkPipelineCache FVulkanAndroidPlatform::PrecompilePSO(FVulkanDevice* Device, con
 					bool bEnableTimeOuts = !FPlatformMisc::IsDebuggerPresent();
 					{
 						QUICK_SCOPE_CYCLE_COUNTER(STAT_VulkanAndroid_PrecompilePSOJAVA);
-						ProgramResponseObj = NewScopedJavaObject(Env, Env->CallStaticObjectMethod(VKRemoteProgramCompileJNI.PSOServiceAccessor, VKRemoteProgramCompileJNI.DispatchPSOCompileShm, *ProgramKeyBuffer, SharedMemFD, VSSize, PSSize, PSOParamsSize, PreSuppliedCacheSize, bEnableTimeOuts));
+						ProgramResponseObj = NewScopedJavaObject(Env, Env->CallStaticObjectMethod(VKRemoteProgramCompileJNI.PSOServiceAccessor, VKRemoteProgramCompileJNI.DispatchPSOCompileShm, *ProgramKeyBuffer, PriorityInfo.GetPriorityInfo(), SharedMemFD, VSSize, PSSize, PSOParamsSize, PreSuppliedCacheSize, bEnableTimeOuts));
 					}
 					CHECK_JNI_EXCEPTIONS(Env);
 					munmap(SharedBuffer, memSize);
@@ -1600,21 +1626,16 @@ VkPipelineCache FVulkanAndroidPlatform::PrecompilePSO(FVulkanDevice* Device, con
 						{
 							QUICK_SCOPE_CYCLE_COUNTER(STAT_VulkanAndroid_PrecompilePSOCreateCache);
 							VkPipelineCacheCreateInfo PipelineCacheCreateInfo;
-							VkPipelineCache PipelineCache;
 							memset(&PipelineCacheCreateInfo, 0, sizeof(VkPipelineCacheCreateInfo));
 							PipelineCacheCreateInfo.sType = VK_STRUCTURE_TYPE_PIPELINE_CACHE_CREATE_INFO;
 							PipelineCacheCreateInfo.flags = 0;
 							PipelineCacheCreateInfo.pInitialData = ResultSharedBuffer + sizeof(ResultMemSize);
 							PipelineCacheCreateInfo.initialDataSize = ResultSize;
-							VERIFYVULKANRESULT(VulkanRHI::vkCreatePipelineCache(Device->GetInstanceHandle(), &PipelineCacheCreateInfo, VULKAN_CPU_ALLOCATOR, &PipelineCache));
+							VERIFYVULKANRESULT(VulkanRHI::vkCreatePipelineCache(Device->GetInstanceHandle(), &PipelineCacheCreateInfo, VULKAN_CPU_ALLOCATOR, &ReturnPipelineCache));
 							AfterSize = ResultSize - OptionalPSOCacheData.Num();
-
-							return PipelineCache;
 						}
 					}
 				}
-
- 				return VK_NULL_HANDLE;
 			}
 			else
 			{
@@ -1623,8 +1644,7 @@ VkPipelineCache FVulkanAndroidPlatform::PrecompilePSO(FVulkanDevice* Device, con
 					FGenericCrashContext::SetEngineData(TEXT("Android.PSOService"), TEXT("ec"));
 				}
 
-				FailureMessageOUT = FJavaHelper::FStringFromLocalRef(Env, (jstring)Env->GetObjectField(*ProgramResponseObj, VKRemoteProgramCompileJNI.ProgramResponse_ErrorField));
-				check(!FailureMessageOUT.IsEmpty());
+				SetFailureMessage( FJavaHelper::FStringFromLocalRef(Env, (jstring)Env->GetObjectField(*ProgramResponseObj, VKRemoteProgramCompileJNI.ProgramResponse_ErrorField)));
 			}
 		}
 		else
@@ -1633,14 +1653,29 @@ VkPipelineCache FVulkanAndroidPlatform::PrecompilePSO(FVulkanDevice* Device, con
 			{
 				FGenericCrashContext::SetEngineData(TEXT("Android.PSOService"), TEXT("es"));
 			}
-			FailureMessageOUT = TEXT("Remote compiler failed.");
+			SetFailureMessage(TEXT("Remote PSO compiler failed."));
+		}
+	}
+	else
+	{
+		if (AndroidVulkanService::bOneTimeErrorEncountered.exchange(true) == false)
+		{
+			FGenericCrashContext::SetEngineData(TEXT("Android.PSOService"), TEXT("ejni"));
 		}
+		SetFailureMessage(TEXT("Remote PSO compiler JNI error."));
 	}
 
-	return VK_NULL_HANDLE;
+	if (ReturnPipelineCache == VK_NULL_HANDLE)
+	{
+		if ((AndroidVulkanService::TotalErrors++) == FPlatformDynamicRHI::GetPSOServiceFailureThreshold())
+		{
+			FVulkanAndroidPlatform::StopRemoteCompileServices();
+			SetFailureMessage(TEXT("Remote PSO compiler failed, error count has passed threshold. Future compiles will be in-process."));
+		}
+	}
+	return ReturnPipelineCache;
 }
 
-
 bool FAndroidVulkanFramePacer::SupportsFramePace(int32 QueryFramePace)
 {
 	int32 TempRefreshRate, TempSyncInterval;
diff --git a/Engine/Source/Runtime/VulkanRHI/Private/Android/VulkanAndroidPlatform.h b/Engine/Source/Runtime/VulkanRHI/Private/Android/VulkanAndroidPlatform.h
--- a/Engine/Source/Runtime/VulkanRHI/Private/Android/VulkanAndroidPlatform.h
+++ b/Engine/Source/Runtime/VulkanRHI/Private/Android/VulkanAndroidPlatform.h
@@ -145,7 +145,7 @@ public:
 	static void DestroySwapchainKHR(VkDevice Device, VkSwapchainKHR Swapchain, const VkAllocationCallbacks* Allocator);
 
 	// handle precompile of PSOs, send to an android specific precompile external process.
-	static VkPipelineCache PrecompilePSO(FVulkanDevice* Device, const TArrayView<uint8> OptionalPSOCacheData, const VkGraphicsPipelineCreateInfo* PipelineInfo, FGfxPipelineDesc* GfxEntry, const FVulkanRenderTargetLayout* RTLayout, TArrayView<uint32_t> VS, TArrayView<uint32_t> PS, size_t& AfterSize);
+	static VkPipelineCache PrecompilePSO(FVulkanDevice* Device, const TArrayView<uint8> OptionalPSOCacheData, FGraphicsPipelineStateInitializer::EPSOPrecacheCompileType PSOCompileType, const VkGraphicsPipelineCreateInfo* PipelineInfo, FGfxPipelineDesc* GfxEntry, const FVulkanRenderTargetLayout* RTLayout, TArrayView<uint32_t> VS, TArrayView<uint32_t> PS, size_t& AfterSize, FString* FailureMessageOUT = nullptr);
 
 	static bool AreRemoteCompileServicesActive();
 	static bool StartRemoteCompileServices(int NumServices);
diff --git a/Engine/Source/Runtime/VulkanRHI/Private/VulkanGenericPlatform.h b/Engine/Source/Runtime/VulkanRHI/Private/VulkanGenericPlatform.h
--- a/Engine/Source/Runtime/VulkanRHI/Private/VulkanGenericPlatform.h
+++ b/Engine/Source/Runtime/VulkanRHI/Private/VulkanGenericPlatform.h
@@ -130,7 +130,7 @@ public:
 	static TArray<FString> GetPSOCacheFilenames();
 
 	// Gives platform a chance to handle precompile of PSOs, returns nullptr if unsupported
-	static VkPipelineCache PrecompilePSO(FVulkanDevice* Device, const uint8* OptionalPSOCacheData,VkGraphicsPipelineCreateInfo* PipelineInfo, FGfxPipelineDesc* GfxEntry, const FVulkanRenderTargetLayout* RTLayout, TArrayView<uint32_t> VS, TArrayView<uint32_t> PS, size_t& AfterSize) { return VK_NULL_HANDLE; }
+	static VkPipelineCache PrecompilePSO(FVulkanDevice* Device, const uint8* OptionalPSOCacheData,VkGraphicsPipelineCreateInfo* PipelineInfo, FGfxPipelineDesc* GfxEntry, const FVulkanRenderTargetLayout* RTLayout, TArrayView<uint32_t> VS, TArrayView<uint32_t> PS, size_t& AfterSize, FString* FailureMessageOUT = nullptr) { return VK_NULL_HANDLE; }
 
 	// Return VK_FALSE if platform wants to suppress the given debug report from the validation layers, VK_TRUE to print it.
 	static VkBool32 DebugReportFunction(VkDebugReportFlagsEXT MsgFlags, VkDebugReportObjectTypeEXT ObjType, uint64_t SrcObject, size_t Location, int32 MsgCode, const ANSICHAR* LayerPrefix, const ANSICHAR* Msg, void* UserData) { return VK_TRUE; }
diff --git a/Engine/Source/Runtime/VulkanRHI/Private/VulkanPipeline.cpp b/Engine/Source/Runtime/VulkanRHI/Private/VulkanPipeline.cpp
--- a/Engine/Source/Runtime/VulkanRHI/Private/VulkanPipeline.cpp
+++ b/Engine/Source/Runtime/VulkanRHI/Private/VulkanPipeline.cpp
@@ -1171,19 +1171,19 @@ static const TMap<uint8, VkFragmentShadingRateCombinerOpKHR> FragmentCombinerOpM
 	// @todo: Add "VK_FRAGMENT_SHADING_RATE_COMBINER_OP_MUL_KHR"?
 };
 
-FString FVulkanPipelineStateCacheManager::ShaderHashesToString(FVulkanShader* Shaders[ShaderStage::NumStages])
+static FString ShaderHashesToString(FVulkanShader* Shaders[ShaderStage::NumStages])
 {
 	FString ShaderHashes = "";
-	if (Shaders[ShaderStage::Vertex] && Shaders[ShaderStage::Vertex]->Frequency == SF_Vertex)
+	if (Shaders[ShaderStage::Vertex])
 	{
 		ShaderHashes += TEXT("VS: ") + static_cast<FVulkanVertexShader*>(Shaders[ShaderStage::Vertex])->GetHash().ToString() + TEXT("\n");
 	}
-	if (Shaders[ShaderStage::Pixel] && Shaders[ShaderStage::Pixel]->Frequency == SF_Pixel)
+	if (Shaders[ShaderStage::Pixel] )
 	{
 		ShaderHashes += TEXT("PS: ") + static_cast<FVulkanPixelShader*>(Shaders[ShaderStage::Pixel])->GetHash().ToString() + TEXT("\n");
 	}
 #if VULKAN_SUPPORTS_GEOMETRY_SHADERS
-	if (Shaders[ShaderStage::Geometry] && Shaders[ShaderStage::Geometry]->Frequency == SF_Geometry)
+	if (Shaders[ShaderStage::Geometry])
 	{
 		ShaderHashes += TEXT("GS: ") + static_cast<FVulkanGeometryShader*>(Shaders[ShaderStage::Geometry])->GetHash().ToString() + TEXT("\n");
 	}
@@ -1191,7 +1191,7 @@ FString FVulkanPipelineStateCacheManager::ShaderHashesToString(FVulkanShader* Sh
 	return ShaderHashes;
 }
 
-bool FVulkanPipelineStateCacheManager::CreateGfxPipelineFromEntry(FVulkanRHIGraphicsPipelineState* PSO, FVulkanShader* Shaders[ShaderStage::NumStages], bool bPrecompile)
+bool FVulkanPipelineStateCacheManager::CreateGfxPipelineFromEntry(FVulkanRHIGraphicsPipelineState* PSO, FVulkanShader* Shaders[ShaderStage::NumStages], FGraphicsPipelineStateInitializer::EPSOPrecacheCompileType PSOCompileType)
 {
 	VkPipeline* Pipeline = &PSO->VulkanPipeline;
 	FGfxPipelineDesc* GfxEntry = &PSO->Desc;
@@ -1375,7 +1375,7 @@ bool FVulkanPipelineStateCacheManager::CreateGfxPipelineFromEntry(FVulkanRHIGrap
 
 	double BeginTime = FPlatformTime::Seconds();
 
-	Result = CreateVKPipeline(PSO, Shaders, PipelineInfo, bPrecompile);
+	Result = CreateVKPipeline(PSO, Shaders, PipelineInfo, PSOCompileType);
 
 	if (Result != VK_SUCCESS)
 	{
@@ -1398,7 +1398,7 @@ bool FVulkanPipelineStateCacheManager::CreateGfxPipelineFromEntry(FVulkanRHIGrap
 
 #if PLATFORM_ANDROID
 
-VkResult CreatePSOWithExternalService(FVulkanDevice* Device, FVulkanRHIGraphicsPipelineState* PSO, FVulkanShader* Shaders[ShaderStage::NumStages], const VkGraphicsPipelineCreateInfo& PipelineInfo, VkPipelineCache DestPipelineCache, FRWLock& PipelineLock)
+static VkResult CreatePSOWithExternalService(FVulkanDevice* Device, FGraphicsPipelineStateInitializer::EPSOPrecacheCompileType PSOCompileType, FVulkanRHIGraphicsPipelineState* PSO, FVulkanShader* Shaders[ShaderStage::NumStages], const VkGraphicsPipelineCreateInfo& PipelineInfo, VkPipelineCache DestPipelineCache, FRWLock& PipelineLock)
 {
 	VkResult Result = VK_ERROR_INITIALIZATION_FAILED;
 	FVulkanShader::FSpirvCode VS = PSO->GetPatchedSpirvCode(Shaders[ShaderStage::Vertex]);
@@ -1411,6 +1411,8 @@ VkResult CreatePSOWithExternalService(FVulkanDevice* Device, FVulkanRHIGraphicsP
 	FGfxPipelineDesc* GfxEntry = &PSO->Desc;
 
 	TArray<uint8> InitialCacheData;
+	bool bSupplyDestPSOCacheData = false; // this can be an optimization, but only if the PSO compile is able to use content from an existing cache.
+	if( bSupplyDestPSOCacheData )
 	{
 		QUICK_SCOPE_CYCLE_COUNTER(STAT_Vulkan_RHICreateGraphicsPipelineState_ExternalInitialCache);
 		FRWScopeLock Lock(PipelineLock, SLT_Write);
@@ -1419,8 +1421,8 @@ VkResult CreatePSOWithExternalService(FVulkanDevice* Device, FVulkanRHIGraphicsP
 		InitialCacheData.SetNumUninitialized(InitialCacheSize);
 		VulkanRHI::vkGetPipelineCacheData(Device->GetInstanceHandle(), DestPipelineCache, &InitialCacheSize, InitialCacheData.GetData());
 	}
-
-	LocalPipelineCache = FVulkanPlatform::PrecompilePSO(Device, InitialCacheData, &PipelineInfo, GfxEntry, &PSO->RenderPass->GetLayout(), VSCode, PSCode, AfterSize);
+	FString FailLog;
+	LocalPipelineCache = FVulkanPlatform::PrecompilePSO(Device, InitialCacheData, PSOCompileType, &PipelineInfo, GfxEntry, &PSO->RenderPass->GetLayout(), VSCode, PSCode, AfterSize,&FailLog);
 
 	if (ensure(LocalPipelineCache != VK_NULL_HANDLE))
 	{
@@ -1434,14 +1436,15 @@ VkResult CreatePSOWithExternalService(FVulkanDevice* Device, FVulkanRHIGraphicsP
 	}
 	else
 	{
-  		UE_LOG(LogVulkanRHI, Error, TEXT("Android RemoteCompileServices Failed to create graphics pipeline.\nShaders in pipeline"));
+  		UE_LOG(LogVulkanRHI, Error, TEXT("Android RemoteCompileServices Failed to create graphics pipeline (%s).\nShaders in pipeline %s"),*FailLog, *ShaderHashesToString(Shaders));
 	}
 	return Result;
 }
 #endif
 
-VkResult FVulkanPipelineStateCacheManager::CreateVKPipeline(FVulkanRHIGraphicsPipelineState* PSO, FVulkanShader* Shaders[ShaderStage::NumStages], const VkGraphicsPipelineCreateInfo& PipelineInfo, bool bIsPrecompileJob)
+VkResult FVulkanPipelineStateCacheManager::CreateVKPipeline(FVulkanRHIGraphicsPipelineState* PSO, FVulkanShader* Shaders[ShaderStage::NumStages], const VkGraphicsPipelineCreateInfo& PipelineInfo, FGraphicsPipelineStateInitializer::EPSOPrecacheCompileType PSOCompileType)
 {
+	bool bIsPrecompileJob = PSOCompileType != FGraphicsPipelineStateInitializer::EPSOPrecacheCompileType::NotSet;
 	if(FVulkanChunkedPipelineCacheManager::IsEnabled())
 	{
 		QUICK_SCOPE_CYCLE_COUNTER(STAT_Vulkan_RHICreateGraphicsPipelineState_CREATE_VKPIPELINE);
@@ -1458,7 +1461,7 @@ VkResult FVulkanPipelineStateCacheManager::CreateVKPipeline(FVulkanRHIGraphicsPi
 				
 				check(PSO->VulkanPipeline == 0);
 				check(PSOOperation == FVulkanChunkedPipelineCacheManager::EPSOOperation::CreateAndStorePSO || PSOOperation == FVulkanChunkedPipelineCacheManager::EPSOOperation::CreateIfPresent);
-				VkResult Result;
+				VkResult Result = VK_ERROR_UNKNOWN;
 
 				if(PSOOperation == FVulkanChunkedPipelineCacheManager::EPSOOperation::CreateIfPresent)
 				{
@@ -1481,20 +1484,19 @@ VkResult FVulkanPipelineStateCacheManager::CreateVKPipeline(FVulkanRHIGraphicsPi
 
 				QUICK_SCOPE_CYCLE_COUNTER(STAT_Vulkan_vkCreateGraphicsPipeline);
 #if PLATFORM_ANDROID
-				if( !FVulkanAndroidPlatform::AreRemoteCompileServicesActive() || !bIsPrecompileJob)
+				if (FVulkanAndroidPlatform::AreRemoteCompileServicesActive() && bIsPrecompileJob)
 				{
+					QUICK_SCOPE_CYCLE_COUNTER(STAT_Vulkan_RHICreateGraphicsPipelineState_ExternalCreate);
+					Result = CreatePSOWithExternalService(Device, PSOCompileType, PSO, Shaders, PipelineInfo, PipelineCache, Params.DestPipelineCacheLock);
+				}
+				// if the external service did not produce a result the following will create it in-process as a best effort fallback.
 #endif
+				if (Result != VK_SUCCESS)
+				{
 					QUICK_SCOPE_CYCLE_COUNTER(STAT_Vulkan_RHICreateGraphicsPipelineState_vkCreate);
 					FRWScopeLock Lock(Params.DestPipelineCacheLock, SLT_ReadOnly);
 					Result = VulkanRHI::vkCreateGraphicsPipelines(Device->GetInstanceHandle(), PipelineCache, 1, &PipelineInfo, VULKAN_CPU_ALLOCATOR, &PSO->VulkanPipeline);
-#if PLATFORM_ANDROID
 				}
-				else
-				{
-					QUICK_SCOPE_CYCLE_COUNTER(STAT_Vulkan_RHICreateGraphicsPipelineState_ExternalCreate);
-					Result = CreatePSOWithExternalService(Device, PSO, Shaders, PipelineInfo, PipelineCache, Params.DestPipelineCacheLock);
-				}
-#endif
 				return Result;
 			}));
 	}
@@ -1542,8 +1544,8 @@ VkResult FVulkanPipelineStateCacheManager::CreateVKPipeline(FVulkanRHIGraphicsPi
 		TArrayView<uint32_t> VSCode = VS.GetCodeView();
 		TArrayView<uint32_t> PSCode = PS.GetCodeView();
 		size_t AfterSize = 0;
-
-		LocalPipelineCache = FVulkanPlatform::PrecompilePSO(Device, MakeArrayView<uint8>(nullptr,0), &PipelineInfo, GfxEntry, &PSO->RenderPass->GetLayout(), VSCode, PSCode, AfterSize);
+		FString FailLog;
+		LocalPipelineCache = FVulkanPlatform::PrecompilePSO(Device, MakeArrayView<uint8>(nullptr,0), PSOCompileType, &PipelineInfo, GfxEntry, &PSO->RenderPass->GetLayout(), VSCode, PSCode, AfterSize,&FailLog);
 
 		if (ensure(LocalPipelineCache != VK_NULL_HANDLE))
 		{
@@ -1568,7 +1570,7 @@ VkResult FVulkanPipelineStateCacheManager::CreateVKPipeline(FVulkanRHIGraphicsPi
 		else
 		{
 			FString ShaderHashes = ShaderHashesToString(Shaders);
-			UE_LOG(LogVulkanRHI, Error, TEXT("Android RemoteCompileServices Failed to create graphics pipeline.\nShaders in pipeline: %s"), *ShaderHashes);
+			UE_LOG(LogVulkanRHI, Error, TEXT("Android RemoteCompileServices Failed to create graphics pipeline (%s).\nShaders in pipeline: %s"), *FailLog, *ShaderHashes);
 		}
 	}
 #endif
@@ -2138,7 +2140,7 @@ FGraphicsPipelineStateRHIRef FVulkanPipelineStateCacheManager::RHICreateGraphics
 			
 				QUICK_SCOPE_CYCLE_COUNTER(STAT_Vulkan_RHICreateGraphicsPipelineState_CREATE_PART0);
 
-				if(!CreateGfxPipelineFromEntry(NewPSO, VulkanShaders, bIsPrecache))
+				if(!CreateGfxPipelineFromEntry(NewPSO, VulkanShaders, Initializer.GetPSOPrecacheCompileType()))
 				{
 					DeleteNewPSO(NewPSO);
 					return nullptr;
@@ -2169,6 +2171,13 @@ FGraphicsPipelineStateRHIRef FVulkanPipelineStateCacheManager::RHICreateGraphics
 					NewPSO->bIsRegistered = true;
 					LRUTrim(NewPSO->PipelineCacheSize);
 					LRUAdd(NewPSO);
+					if(bIsPrecache)
+					{
+						// immediately evict precache PSOs from the LRU.
+						// precache PSOs can saturate the LRU. precache PSOs can end up being trimmed/evicted in the same frame which LRUTrim does not expect.
+						// This means we are LRU-ing rendered PSOs only.
+						LRURemove(NewPSO);
+					}
 				}
 				else
 				{
@@ -2571,7 +2580,7 @@ void FVulkanPipelineStateCacheManager::LRUTouch(FVulkanRHIGraphicsPipelineState*
 
 			QUICK_SCOPE_CYCLE_COUNTER(STAT_Vulkan_LRUMiss);
 
-			if (!CreateGfxPipelineFromEntry(PSO, VulkanShaders, false))
+			if (!CreateGfxPipelineFromEntry(PSO, VulkanShaders, FGraphicsPipelineStateInitializer::EPSOPrecacheCompileType::NotSet))
 			{
 				check(0);
 			}
diff --git a/Engine/Source/Runtime/VulkanRHI/Private/VulkanPipeline.h b/Engine/Source/Runtime/VulkanRHI/Private/VulkanPipeline.h
--- a/Engine/Source/Runtime/VulkanRHI/Private/VulkanPipeline.h
+++ b/Engine/Source/Runtime/VulkanRHI/Private/VulkanPipeline.h
@@ -521,10 +521,9 @@ private:
 	FGraphicsPipelineStateRHIRef RHICreateGraphicsPipelineState(const FGraphicsPipelineStateInitializer& Initializer);
 	FVulkanComputePipeline* RHICreateComputePipelineState(FRHIComputeShader* ComputeShaderRHI);
 	void NotifyDeletedGraphicsPSO(FRHIGraphicsPipelineState* PSO);
-	bool CreateGfxPipelineFromEntry(FVulkanRHIGraphicsPipelineState* PSO, FVulkanShader* Shaders[ShaderStage::NumStages], bool bPrecompile);
+	bool CreateGfxPipelineFromEntry(FVulkanRHIGraphicsPipelineState* PSO, FVulkanShader* Shaders[ShaderStage::NumStages], FGraphicsPipelineStateInitializer::EPSOPrecacheCompileType PSOCompileType);
 
-	VkResult CreateVKPipeline(FVulkanRHIGraphicsPipelineState* PSO, FVulkanShader* Shaders[ShaderStage::NumStages], const VkGraphicsPipelineCreateInfo& PipelineInfo, bool bIsPrecompileJob);
-	static FString ShaderHashesToString(FVulkanShader* Shaders[ShaderStage::NumStages]);
+	VkResult CreateVKPipeline(FVulkanRHIGraphicsPipelineState* PSO, FVulkanShader* Shaders[ShaderStage::NumStages], const VkGraphicsPipelineCreateInfo& PipelineInfo, FGraphicsPipelineStateInitializer::EPSOPrecacheCompileType PSOCompileType);
 
 	FVulkanLayout* FindOrAddLayout(const FVulkanDescriptorSetsLayoutInfo& DescriptorSetLayoutInfo, bool bGfxLayout);
 	FVulkanComputePipeline* CreateComputePipelineFromShader(FVulkanComputeShader* Shader);
