VulkanRHI: Fix Vulkan buffer lock for readback where FRHIGPUBufferReadback isn't used by overloading RHILockBuffer to use existing buffer read-only lock code.

#rnx

diff --git a/Engine/Source/Runtime/VulkanRHI/Private/VulkanBuffer.cpp b/Engine/Source/Runtime/VulkanRHI/Private/VulkanBuffer.cpp
--- a/Engine/Source/Runtime/VulkanRHI/Private/VulkanBuffer.cpp
+++ b/Engine/Source/Runtime/VulkanRHI/Private/VulkanBuffer.cpp
@@ -639,8 +639,28 @@ void FVulkanDynamicRHI::RHICopyBuffer(FRHIBuffer* SourceBufferRHI, FRHIBuffer* D
 	VULKAN_SIGNAL_UNIMPLEMENTED();
 }
 
+void* FVulkanDynamicRHI::RHILockBuffer(FRHICommandListBase& RHICmdList, FRHIBuffer* BufferRHI, uint32 Offset, uint32 Size, EResourceLockMode LockMode)
+{
+	QUICK_SCOPE_CYCLE_COUNTER(STAT_LockBuffer_RenderThread);
+	LLM_SCOPE_VULKAN(ELLMTagVulkan::VulkanBuffers);
+	FVulkanResourceMultiBuffer* Buffer = ResourceCast(BufferRHI);
+
+	if (LockMode == RLM_ReadOnly)
+	{
+		// Lock in ReadOnly will call PrepareForCPURead on Device, so we can use the immediate context directly.
+		// This is a legacy path though, FRHIGPUBufferReadback should be used for reading back buffers.
+		return Buffer->Lock(Device->GetImmediateContext(), LockMode, Size, Offset);
+	}
+	else
+	{
+		return Buffer->Lock(RHICmdList, LockMode, Size, Offset);
+	}
+}
+
 void FVulkanDynamicRHI::RHIUnlockBuffer(FRHICommandListBase& RHICmdList, FRHIBuffer* BufferRHI)
 {
-	QUICK_SCOPE_CYCLE_COUNTER(STAT_FDynamicRHI_UnlockBuffer_RenderThread);
-	FDynamicRHI::RHIUnlockBuffer(RHICmdList, BufferRHI);
+	QUICK_SCOPE_CYCLE_COUNTER(STAT_UnlockBuffer_RenderThread);
+	LLM_SCOPE_VULKAN(ELLMTagVulkan::VulkanBuffers);
+	FVulkanResourceMultiBuffer* Buffer = ResourceCast(BufferRHI);
+	Buffer->Unlock(RHICmdList);
 }
diff --git a/Engine/Source/Runtime/VulkanRHI/Public/VulkanDynamicRHI.h b/Engine/Source/Runtime/VulkanRHI/Public/VulkanDynamicRHI.h
--- a/Engine/Source/Runtime/VulkanRHI/Public/VulkanDynamicRHI.h
+++ b/Engine/Source/Runtime/VulkanRHI/Public/VulkanDynamicRHI.h
@@ -118,6 +118,7 @@ public:
 	virtual void RHIReplaceResources(FRHICommandListBase& RHICmdList, TArray<FRHIResourceReplaceInfo>&& ReplaceInfos) final override;
 	virtual void* LockBuffer_BottomOfPipe(FRHICommandListBase& RHICmdList, FRHIBuffer* Buffer, uint32 Offset, uint32 Size, EResourceLockMode LockMode) final override;
 	virtual void UnlockBuffer_BottomOfPipe(FRHICommandListBase& RHICmdList, FRHIBuffer* Buffer) final override;
+	virtual void* RHILockBuffer(FRHICommandListBase& RHICmdList, FRHIBuffer* BufferRHI, uint32 Offset, uint32 Size, EResourceLockMode LockMode) final override;
 	virtual void RHIUnlockBuffer(FRHICommandListBase& RHICmdList, FRHIBuffer* Buffer) final override;
 	virtual FTextureReferenceRHIRef RHICreateTextureReference(FRHICommandListBase& RHICmdList, FRHITexture* InReferencedTexture) final override;
 	virtual void RHIUpdateTextureReference(FRHICommandListBase& RHICmdList, FRHITextureReference* TextureRef, FRHITexture* NewTexture) final override;
