VulkanRHI: Break out the generation of VkImageUsageFlags from GenerateImageCreateInfo so that we do only what is necessary when the VkImage is provided.

#rnx

diff --git a/Engine/Source/Runtime/VulkanRHI/Private/VulkanDebug.cpp b/Engine/Source/Runtime/VulkanRHI/Private/VulkanDebug.cpp
--- a/Engine/Source/Runtime/VulkanRHI/Private/VulkanDebug.cpp
+++ b/Engine/Source/Runtime/VulkanRHI/Private/VulkanDebug.cpp
@@ -1206,11 +1206,14 @@ void FWrapLayer::CreateImage(VkResult Result, VkDevice Device, const VkImageCrea
 	if (Result == VK_RESULT_MAX_ENUM)
 	{
 #if VULKAN_ENABLE_DUMP_LAYER
-		DevicePrintfBegin(Device, FString::Printf(TEXT("vkCreateImage(Info=0x%p, OutImage=0x%p)"), CreateInfo, Image));
-		DebugLog += FString::Printf(TEXT("%sVkImageCreateInfo: Flags=%d, ImageType=%s, Format=%s, MipLevels=%d, ArrayLayers=%d, Samples=%s\n"), Tabs, CreateInfo->flags, *GetVkImageTypeString(CreateInfo->imageType),
-			*GetVkFormatString(CreateInfo->format), CreateInfo->mipLevels, CreateInfo->arrayLayers, *GetSampleCountString(CreateInfo->samples));
-		DebugLog += FString::Printf(TEXT("%s\tExtent=(%s) Tiling=%s, Usage=%s, Initial=%s\n"), Tabs, *GetExtentString(CreateInfo->extent),
-			*GetVkImageTilingString(CreateInfo->tiling), *GetImageUsageString(CreateInfo->usage), *GetVkImageLayoutString(CreateInfo->initialLayout));
+		DevicePrintfBegin(Device, FString::Printf(TEXT("vkCreateImage(CreateInfo=0x%p, OutImage=0x%p)"), CreateInfo, Image));
+		if (CreateInfo)
+		{
+			DebugLog += FString::Printf(TEXT("%sVkImageCreateInfo: Flags=%d, ImageType=%s, Format=%s, MipLevels=%d, ArrayLayers=%d, Samples=%s\n"), Tabs, CreateInfo->flags, *GetVkImageTypeString(CreateInfo->imageType),
+				*GetVkFormatString(CreateInfo->format), CreateInfo->mipLevels, CreateInfo->arrayLayers, *GetSampleCountString(CreateInfo->samples));
+			DebugLog += FString::Printf(TEXT("%s\tExtent=(%s) Tiling=%s, Usage=%s, Initial=%s\n"), Tabs, *GetExtentString(CreateInfo->extent),
+				*GetVkImageTilingString(CreateInfo->tiling), *GetImageUsageString(CreateInfo->usage), *GetVkImageLayoutString(CreateInfo->initialLayout));
+		}
 		FlushDebugWrapperLog();
 #endif
 	}
diff --git a/Engine/Source/Runtime/VulkanRHI/Private/VulkanTexture.cpp b/Engine/Source/Runtime/VulkanRHI/Private/VulkanTexture.cpp
--- a/Engine/Source/Runtime/VulkanRHI/Private/VulkanTexture.cpp
+++ b/Engine/Source/Runtime/VulkanRHI/Private/VulkanTexture.cpp
@@ -194,6 +194,62 @@ struct FRHICommandLockWriteTexture final : public FRHICommand<FRHICommandLockWri
 	}
 };
 
+static VkImageUsageFlags GetUsageFlagsFromCreateFlags(FVulkanDevice& InDevice, const ETextureCreateFlags& UEFlags)
+{
+	VkImageUsageFlags UsageFlags = VK_IMAGE_USAGE_TRANSFER_SRC_BIT | VK_IMAGE_USAGE_TRANSFER_DST_BIT | VK_IMAGE_USAGE_SAMPLED_BIT;
+
+	if (EnumHasAnyFlags(UEFlags, TexCreate_Presentable))
+	{
+		UsageFlags |= VK_IMAGE_USAGE_STORAGE_BIT;
+	}
+	else if (EnumHasAnyFlags(UEFlags, TexCreate_RenderTargetable | TexCreate_DepthStencilTargetable))
+	{
+		if (EnumHasAllFlags(UEFlags, TexCreate_InputAttachmentRead))
+		{
+			UsageFlags |= VK_IMAGE_USAGE_INPUT_ATTACHMENT_BIT;
+		}
+		UsageFlags |= (EnumHasAnyFlags(UEFlags, TexCreate_RenderTargetable) ? VK_IMAGE_USAGE_COLOR_ATTACHMENT_BIT : VK_IMAGE_USAGE_DEPTH_STENCIL_ATTACHMENT_BIT);
+		if (EnumHasAllFlags(UEFlags, TexCreate_Memoryless) && InDevice.GetDeviceMemoryManager().SupportsMemoryless())
+		{
+			UsageFlags |= VK_IMAGE_USAGE_TRANSIENT_ATTACHMENT_BIT;
+			// Remove the transfer and sampled bits, as they are incompatible with the transient bit.
+			UsageFlags &= ~(VK_IMAGE_USAGE_TRANSFER_DST_BIT | VK_IMAGE_USAGE_TRANSFER_SRC_BIT | VK_IMAGE_USAGE_SAMPLED_BIT);
+		}
+	}
+	else if (EnumHasAnyFlags(UEFlags, TexCreate_DepthStencilResolveTarget))
+	{
+		UsageFlags |= VK_IMAGE_USAGE_DEPTH_STENCIL_ATTACHMENT_BIT;
+	}
+	else if (EnumHasAnyFlags(UEFlags, TexCreate_ResolveTargetable))
+	{
+		UsageFlags |= VK_IMAGE_USAGE_COLOR_ATTACHMENT_BIT;
+	}
+
+	if (EnumHasAnyFlags(UEFlags, TexCreate_Foveation) && ValidateShadingRateDataType())
+	{
+		if (GRHIVariableRateShadingImageDataType == VRSImage_Palette)
+		{
+			UsageFlags |= VK_IMAGE_USAGE_FRAGMENT_SHADING_RATE_ATTACHMENT_BIT_KHR;
+		}
+
+		if (GRHIVariableRateShadingImageDataType == VRSImage_Fractional)
+		{
+			UsageFlags |= VK_IMAGE_USAGE_FRAGMENT_DENSITY_MAP_BIT_EXT;
+		}
+	}
+
+	if (EnumHasAnyFlags(UEFlags, TexCreate_UAV))
+	{
+		//cannot have the storage bit on a memoryless texture
+		ensure(!EnumHasAnyFlags(UEFlags, TexCreate_Memoryless));
+		UsageFlags |= VK_IMAGE_USAGE_STORAGE_BIT;
+	}
+
+	return UsageFlags;
+}
+
+
+
 void FVulkanTexture::GenerateImageCreateInfo(
 	FImageCreateInfo& OutImageCreateInfo,
 	FVulkanDevice& InDevice,
@@ -313,62 +369,12 @@ void FVulkanTexture::GenerateImageCreateInfo(
 	}
 
 	ImageCreateInfo.tiling = bForceLinearTexture ? VK_IMAGE_TILING_LINEAR : GVulkanViewTypeTilingMode[ResourceType];
-
-	ImageCreateInfo.usage = 0;
-	ImageCreateInfo.usage |= VK_IMAGE_USAGE_TRANSFER_DST_BIT;
-	//@TODO: should everything be created with the source bit?
-	ImageCreateInfo.usage |= VK_IMAGE_USAGE_TRANSFER_SRC_BIT;
-	ImageCreateInfo.usage |= VK_IMAGE_USAGE_SAMPLED_BIT;
-
-	if (EnumHasAnyFlags(UEFlags, TexCreate_Presentable))
-	{
-		ImageCreateInfo.usage |= VK_IMAGE_USAGE_STORAGE_BIT;		
-	}
-	else if (EnumHasAnyFlags(UEFlags, TexCreate_RenderTargetable | TexCreate_DepthStencilTargetable))
-	{
-		if (EnumHasAllFlags(UEFlags, TexCreate_InputAttachmentRead))
-		{
-			ImageCreateInfo.usage |= VK_IMAGE_USAGE_INPUT_ATTACHMENT_BIT;
-		}
-		ImageCreateInfo.usage |= (EnumHasAnyFlags(UEFlags, TexCreate_RenderTargetable) ? VK_IMAGE_USAGE_COLOR_ATTACHMENT_BIT : VK_IMAGE_USAGE_DEPTH_STENCIL_ATTACHMENT_BIT);
-		ImageCreateInfo.tiling = VK_IMAGE_TILING_OPTIMAL;
-		if (EnumHasAllFlags(UEFlags, TexCreate_Memoryless) && InDevice.GetDeviceMemoryManager().SupportsMemoryless())
-		{
-			ImageCreateInfo.usage |= VK_IMAGE_USAGE_TRANSIENT_ATTACHMENT_BIT;
-			// Remove the transfer and sampled bits, as they are incompatible with the transient bit.
-			ImageCreateInfo.usage &= ~(VK_IMAGE_USAGE_TRANSFER_DST_BIT | VK_IMAGE_USAGE_TRANSFER_SRC_BIT | VK_IMAGE_USAGE_SAMPLED_BIT);
-		}
-	}
-	else if (EnumHasAnyFlags(UEFlags, TexCreate_DepthStencilResolveTarget))
-	{
-		ImageCreateInfo.usage |= VK_IMAGE_USAGE_DEPTH_STENCIL_ATTACHMENT_BIT;
-		ImageCreateInfo.tiling = VK_IMAGE_TILING_OPTIMAL;
-	}
-	else if (EnumHasAnyFlags(UEFlags, TexCreate_ResolveTargetable))
+	if (EnumHasAnyFlags(UEFlags, TexCreate_RenderTargetable | TexCreate_DepthStencilTargetable | TexCreate_ResolveTargetable | TexCreate_DepthStencilResolveTarget))
 	{
-		ImageCreateInfo.usage |= VK_IMAGE_USAGE_COLOR_ATTACHMENT_BIT;
 		ImageCreateInfo.tiling = VK_IMAGE_TILING_OPTIMAL;
 	}
 
-	if (EnumHasAnyFlags(UEFlags, TexCreate_Foveation) && ValidateShadingRateDataType())
-	{
-		if (GRHIVariableRateShadingImageDataType == VRSImage_Palette)
-		{
-			ImageCreateInfo.usage |= VK_IMAGE_USAGE_FRAGMENT_SHADING_RATE_ATTACHMENT_BIT_KHR;
-		}
-
-		if (GRHIVariableRateShadingImageDataType == VRSImage_Fractional)
-		{
-			ImageCreateInfo.usage |= VK_IMAGE_USAGE_FRAGMENT_DENSITY_MAP_BIT_EXT;
-		}
-	}
-	
-	if (EnumHasAnyFlags(UEFlags, TexCreate_UAV))
-	{
-		//cannot have the storage bit on a memoryless texture
-		ensure(!EnumHasAnyFlags(UEFlags, TexCreate_Memoryless));
-		ImageCreateInfo.usage |= VK_IMAGE_USAGE_STORAGE_BIT;
-	}
+	ImageCreateInfo.usage = GetUsageFlagsFromCreateFlags(InDevice, UEFlags);
 
 	if (EnumHasAnyFlags(UEFlags, TexCreate_External))
 	{
@@ -1783,11 +1789,9 @@ FVulkanTexture::FVulkanTexture(FVulkanDevice& InDevice, const FRHITextureCreateD
 
 		if (Image != VK_NULL_HANDLE)
 		{
-			FImageCreateInfo ImageCreateInfo;
-			FVulkanTexture::GenerateImageCreateInfo(ImageCreateInfo, InDevice, InCreateDesc, &StorageFormat, &ViewFormat);
-			ImageUsageFlags = ImageCreateInfo.ImageCreateInfo.usage;
+			ImageUsageFlags = GetUsageFlagsFromCreateFlags(InDevice, InCreateDesc.Flags);
 #if VULKAN_ENABLE_WRAP_LAYER
-			FWrapLayer::CreateImage(VK_SUCCESS, InDevice.GetInstanceHandle(), &ImageCreateInfo.ImageCreateInfo, &Image);
+			FWrapLayer::CreateImage(VK_SUCCESS, InDevice.GetInstanceHandle(), nullptr, &Image);
 #endif
 			VULKAN_SET_DEBUG_NAME(InDevice, VK_OBJECT_TYPE_IMAGE, Image, TEXT("%s:(FVulkanTexture*)0x%p"), InCreateDesc.DebugName ? InCreateDesc.DebugName : TEXT("?"), this);
 
