VulkanRHI: Fix typo in extension name

#rnx

diff --git a/Engine/Source/Runtime/VulkanRHI/Private/VulkanUtil.cpp b/Engine/Source/Runtime/VulkanRHI/Private/VulkanUtil.cpp
--- a/Engine/Source/Runtime/VulkanRHI/Private/VulkanUtil.cpp
+++ b/Engine/Source/Runtime/VulkanRHI/Private/VulkanUtil.cpp
@@ -627,7 +627,7 @@ void FVulkanGPUProfiler::DumpCrashMarkers(void* BufferData)
 	{
 		if (Device->GetOptionalExtensions().HasAMDBufferMarker)
 		{
-			UE_LOG(LogVulkanRHI, Error, TEXT("Breadcrumbs using ADM extension"));
+			UE_LOG(LogVulkanRHI, Error, TEXT("Breadcrumbs using VK_AMD_buffer_marker extension"));
 		}
 		else
 		{
