Silence PVS static analysis warning.

diff --git a/Engine/Source/Runtime/VulkanRHI/Private/VulkanPipeline.cpp b/Engine/Source/Runtime/VulkanRHI/Private/VulkanPipeline.cpp
--- a/Engine/Source/Runtime/VulkanRHI/Private/VulkanPipeline.cpp
+++ b/Engine/Source/Runtime/VulkanRHI/Private/VulkanPipeline.cpp
@@ -1491,7 +1491,7 @@ VkResult FVulkanPipelineStateCacheManager::CreateVKPipeline(FVulkanRHIGraphicsPi
 				}
 				// if the external service did not produce a result the following will create it in-process as a best effort fallback.
 #endif
-				if (Result != VK_SUCCESS)
+				if (Result != VK_SUCCESS) //-V547
 				{
 					QUICK_SCOPE_CYCLE_COUNTER(STAT_Vulkan_RHICreateGraphicsPipelineState_vkCreate);
 					FRWScopeLock Lock(Params.DestPipelineCacheLock, SLT_ReadOnly);
