VulkanRHI: add requested Vulkan RHI hooks for plugins.

#rnx

diff --git a/Engine/Source/Runtime/VulkanRHI/Private/VulkanRHI.cpp b/Engine/Source/Runtime/VulkanRHI/Private/VulkanRHI.cpp
--- a/Engine/Source/Runtime/VulkanRHI/Private/VulkanRHI.cpp
+++ b/Engine/Source/Runtime/VulkanRHI/Private/VulkanRHI.cpp
@@ -1432,6 +1432,18 @@ TArray<VkExtensionProperties> FVulkanDynamicRHI::RHIGetAllDeviceExtensions(VkPhy
 	return Extensions;
 }
 
+TArray<FAnsiString> FVulkanDynamicRHI::RHIGetLoadedDeviceExtensions() const
+{
+	// Create copies to prevent issues
+	TArray<FAnsiString> OutExtensions;
+	const TArray<const ANSICHAR*>& DeviceExtensions = GetDevice()->DeviceExtensions;
+	for (const ANSICHAR* ExtensionName : DeviceExtensions)
+	{
+		OutExtensions.Emplace(ExtensionName);
+	}
+	return OutExtensions;
+}
+
 VkImage FVulkanDynamicRHI::RHIGetVkImage(FRHITexture* InTexture) const
 {
 	FVulkanTexture* VulkanTexture = ResourceCast(InTexture);
@@ -1482,6 +1494,19 @@ FVulkanRHIImageViewInfo FVulkanDynamicRHI::RHIGetImageViewInfo(FRHITexture* InTe
 	return Info;
 }
 
+FVulkanRHIAllocationInfo FVulkanDynamicRHI::RHIGetAllocationInfo(FRHIBuffer* InBuffer) const
+{
+	FVulkanResourceMultiBuffer* VulkanBuffer = ResourceCast(InBuffer);
+	const VulkanRHI::FVulkanAllocation& Allocation = VulkanBuffer->GetCurrentAllocation();
+
+	FVulkanRHIAllocationInfo NewInfo{};
+	NewInfo.Handle = Allocation.GetDeviceMemoryHandle(GetDevice());
+	NewInfo.Offset = Allocation.Offset;
+	NewInfo.Size = Allocation.Size;
+
+	return NewInfo;
+}
+
 void FVulkanDynamicRHI::RHISetImageLayout(VkImage Image, VkImageLayout OldLayout, VkImageLayout NewLayout, const VkImageSubresourceRange& SubresourceRange)
 {
 	FVulkanCommandListContext& ImmediateContext = GetDevice()->GetImmediateContext();
diff --git a/Engine/Source/Runtime/VulkanRHI/Public/IVulkanDynamicRHI.h b/Engine/Source/Runtime/VulkanRHI/Public/IVulkanDynamicRHI.h
--- a/Engine/Source/Runtime/VulkanRHI/Public/IVulkanDynamicRHI.h
+++ b/Engine/Source/Runtime/VulkanRHI/Public/IVulkanDynamicRHI.h
@@ -50,6 +50,7 @@ struct IVulkanDynamicRHI : public FDynamicRHI
 
 	virtual TArray<VkExtensionProperties> RHIGetAllInstanceExtensions() const = 0;
 	virtual TArray<VkExtensionProperties> RHIGetAllDeviceExtensions(VkPhysicalDevice InPhysicalDevice) const = 0;
+	virtual TArray<FAnsiString> RHIGetLoadedDeviceExtensions() const = 0;
 
 	virtual FTextureRHIRef   RHICreateTexture2DFromResource(EPixelFormat Format, uint32 SizeX, uint32 SizeY, uint32 NumMips, uint32 NumSamples, VkImage Resource, ETextureCreateFlags Flags, const FClearValueBinding& ClearValueBinding = FClearValueBinding::Transparent) = 0;
 	virtual FTextureRHIRef   RHICreateTexture2DArrayFromResource(EPixelFormat Format, uint32 SizeX, uint32 SizeY, uint32 ArraySize, uint32 NumMips, uint32 NumSamples, VkImage Resource, ETextureCreateFlags Flags, const FClearValueBinding& ClearValueBinding = FClearValueBinding::Transparent) = 0;
@@ -59,6 +60,7 @@ struct IVulkanDynamicRHI : public FDynamicRHI
 	virtual VkFormat                 RHIGetViewVkFormat(FRHITexture* InTexture) const = 0;
 	virtual FVulkanRHIAllocationInfo RHIGetAllocationInfo(FRHITexture* InTexture) const = 0;
 	virtual FVulkanRHIImageViewInfo  RHIGetImageViewInfo(FRHITexture* InTexture) const = 0;
+	virtual FVulkanRHIAllocationInfo RHIGetAllocationInfo(FRHIBuffer* InBuffer) const = 0;
 
 	virtual void           RHISetImageLayout(VkImage Image, VkImageLayout OldLayout, VkImageLayout NewLayout, const VkImageSubresourceRange& SubresourceRange) = 0;
 	virtual void           RHISetUploadImageLayout(VkImage Image, VkImageLayout OldLayout, VkImageLayout NewLayout, const VkImageSubresourceRange& SubresourceRange) = 0;
diff --git a/Engine/Source/Runtime/VulkanRHI/Public/VulkanDynamicRHI.h b/Engine/Source/Runtime/VulkanRHI/Public/VulkanDynamicRHI.h
--- a/Engine/Source/Runtime/VulkanRHI/Public/VulkanDynamicRHI.h
+++ b/Engine/Source/Runtime/VulkanRHI/Public/VulkanDynamicRHI.h
@@ -73,6 +73,7 @@ public:
 	virtual bool RHISupportsEXTFragmentDensityMap2() const final override;
 	virtual TArray<VkExtensionProperties> RHIGetAllInstanceExtensions() const final override;
 	virtual TArray<VkExtensionProperties> RHIGetAllDeviceExtensions(VkPhysicalDevice InPhysicalDevice) const final override;
+	virtual TArray<FAnsiString> RHIGetLoadedDeviceExtensions() const final override;
 	virtual FTextureRHIRef RHICreateTexture2DFromResource(EPixelFormat Format, uint32 SizeX, uint32 SizeY, uint32 NumMips, uint32 NumSamples, VkImage Resource, ETextureCreateFlags Flags, const FClearValueBinding& ClearValueBinding = FClearValueBinding::Transparent) final override;
 	virtual FTextureRHIRef RHICreateTexture2DArrayFromResource(EPixelFormat Format, uint32 SizeX, uint32 SizeY, uint32 ArraySize, uint32 NumMips, uint32 NumSamples, VkImage Resource, ETextureCreateFlags Flags, const FClearValueBinding& ClearValueBinding = FClearValueBinding::Transparent) final override;
 	virtual FTextureRHIRef RHICreateTextureCubeFromResource(EPixelFormat Format, uint32 Size, bool bArray, uint32 ArraySize, uint32 NumMips, VkImage Resource, ETextureCreateFlags Flags, const FClearValueBinding& ClearValueBinding = FClearValueBinding::Transparent) final override;
@@ -80,6 +81,7 @@ public:
 	virtual VkFormat RHIGetViewVkFormat(FRHITexture* InTexture) const final override;
 	virtual FVulkanRHIAllocationInfo RHIGetAllocationInfo(FRHITexture* InTexture) const final override;
 	virtual FVulkanRHIImageViewInfo RHIGetImageViewInfo(FRHITexture* InTexture) const final override;
+	virtual FVulkanRHIAllocationInfo RHIGetAllocationInfo(FRHIBuffer* InBuffer) const final override;
 	virtual void RHISetImageLayout(VkImage Image, VkImageLayout OldLayout, VkImageLayout NewLayout, const VkImageSubresourceRange& SubresourceRange) final override;
 	virtual void RHISetUploadImageLayout(VkImage Image, VkImageLayout OldLayout, VkImageLayout NewLayout, const VkImageSubresourceRange& SubresourceRange) final override;
 	virtual void RHIFinishExternalComputeWork(VkCommandBuffer InCommandBuffer) final override;
