VulkanRHI: Remove duplicate InputAttachmentInfo array.  Remove unused code.

#rnx

diff --git a/Engine/Source/Runtime/VulkanRHI/Private/VulkanCommands.cpp b/Engine/Source/Runtime/VulkanRHI/Private/VulkanCommands.cpp
--- a/Engine/Source/Runtime/VulkanRHI/Private/VulkanCommands.cpp
+++ b/Engine/Source/Runtime/VulkanRHI/Private/VulkanCommands.cpp
@@ -455,7 +455,6 @@ void FVulkanCommandListContext::RHISetUniformBufferDynamicOffset(FUniformBufferS
 	check(IsAligned(InOffset, Device->GetLimits().minUniformBufferOffsetAlignment));
 
 	FVulkanUniformBuffer* UniformBuffer = ResourceCast(GlobalUniformBuffers[InSlot]);
-	const FVulkanGfxPipelineDescriptorInfo& DescriptorInfo = PendingGfxState->CurrentState->GetGfxPipelineDescriptorInfo();
 
 	static const ShaderStage::EStage Stages[2] = 
 	{
diff --git a/Engine/Source/Runtime/VulkanRHI/Private/VulkanDescriptorSets.h b/Engine/Source/Runtime/VulkanRHI/Private/VulkanDescriptorSets.h
--- a/Engine/Source/Runtime/VulkanRHI/Private/VulkanDescriptorSets.h
+++ b/Engine/Source/Runtime/VulkanRHI/Private/VulkanDescriptorSets.h
@@ -615,11 +615,6 @@ public:
 		return bInitialized;
 	}
 
-	inline const TArray<FVulkanShaderHeader::FInputAttachmentInfo>& GetInputAttachmentData() const
-	{
-		return RemappingInfo->InputAttachmentData;
-	}
-
 	void Initialize(const FDescriptorSetRemappingInfo& InRemappingInfo);
 
 protected:
@@ -1129,8 +1124,6 @@ public:
 		return GfxPipelineDescriptorInfo;
 	}
 
-	bool UsesInputAttachment(FVulkanShaderHeader::EAttachmentType AttachmentType) const;
-
 protected:
 	FVulkanGfxPipelineDescriptorInfo		GfxPipelineDescriptorInfo;
 	friend class FVulkanPipelineStateCacheManager;
diff --git a/Engine/Source/Runtime/VulkanRHI/Private/VulkanPendingState.cpp b/Engine/Source/Runtime/VulkanRHI/Private/VulkanPendingState.cpp
--- a/Engine/Source/Runtime/VulkanRHI/Private/VulkanPendingState.cpp
+++ b/Engine/Source/Runtime/VulkanRHI/Private/VulkanPendingState.cpp
@@ -533,34 +533,38 @@ void FVulkanPendingGfxState::InternalUpdateDynamicStates(FVulkanCmdBuffer* Cmd)
 
 void FVulkanPendingGfxState::UpdateInputAttachments(FVulkanFramebuffer* Framebuffer)
 {
-	const FVulkanGfxPipelineDescriptorInfo& GfxDescriptorInfo = CurrentState->GetGfxPipelineDescriptorInfo();
-	const TArray<FVulkanShaderHeader::FInputAttachmentInfo>& InputAttachmentData = GfxDescriptorInfo.GetInputAttachmentData();
-
-	for (int32 Index = 0; Index < InputAttachmentData.Num(); ++Index)
+	const FVulkanShader* PixelShader = CurrentPipeline->GetVulkanShader(SF_Pixel);
+	if (PixelShader)
 	{
-		const FVulkanShaderHeader::FInputAttachmentInfo& AttachmentData = InputAttachmentData[Index];
+		const FVulkanShaderHeader& Header = PixelShader->GetCodeHeader();
+		const TArray<FVulkanShaderHeader::FInputAttachmentInfo>& InputAttachmentData = Header.InputAttachmentInfos;
 	
-		switch (AttachmentData.Type)
+		for (int32 Index = 0; Index < InputAttachmentData.Num(); ++Index)
 		{
-		case FVulkanShaderHeader::EAttachmentType::Color0:
-		case FVulkanShaderHeader::EAttachmentType::Color1:
-		case FVulkanShaderHeader::EAttachmentType::Color2:
-		case FVulkanShaderHeader::EAttachmentType::Color3:
-		case FVulkanShaderHeader::EAttachmentType::Color4:
-		case FVulkanShaderHeader::EAttachmentType::Color5:
-		case FVulkanShaderHeader::EAttachmentType::Color6:
-		case FVulkanShaderHeader::EAttachmentType::Color7:
+			const FVulkanShaderHeader::FInputAttachmentInfo& AttachmentData = InputAttachmentData[Index];
+
+			switch (AttachmentData.Type)
+			{
+			case FVulkanShaderHeader::EAttachmentType::Color0:
+			case FVulkanShaderHeader::EAttachmentType::Color1:
+			case FVulkanShaderHeader::EAttachmentType::Color2:
+			case FVulkanShaderHeader::EAttachmentType::Color3:
+			case FVulkanShaderHeader::EAttachmentType::Color4:
+			case FVulkanShaderHeader::EAttachmentType::Color5:
+			case FVulkanShaderHeader::EAttachmentType::Color6:
+			case FVulkanShaderHeader::EAttachmentType::Color7:
 			{
 				const int32 ColorIndex = static_cast<int32>(AttachmentData.Type) - (int32)FVulkanShaderHeader::EAttachmentType::Color0;
 				check((ColorIndex >= 0) && (ColorIndex < (int32)Framebuffer->GetNumColorAttachments()));
 				CurrentState->SetInputAttachment(ShaderStage::Pixel, AttachmentData.BindingIndex, Framebuffer->AttachmentTextureViews[ColorIndex]->GetTextureView(), VK_IMAGE_LAYOUT_SHADER_READ_ONLY_OPTIMAL);
 			}
 			break;
-		case FVulkanShaderHeader::EAttachmentType::Depth:
-			CurrentState->SetInputAttachment(ShaderStage::Pixel, AttachmentData.BindingIndex, Framebuffer->GetPartialDepthTextureView(), VK_IMAGE_LAYOUT_DEPTH_STENCIL_READ_ONLY_OPTIMAL);
-			break;
-		default:
-			check(0);
+			case FVulkanShaderHeader::EAttachmentType::Depth:
+				CurrentState->SetInputAttachment(ShaderStage::Pixel, AttachmentData.BindingIndex, Framebuffer->GetPartialDepthTextureView(), VK_IMAGE_LAYOUT_DEPTH_STENCIL_READ_ONLY_OPTIMAL);
+				break;
+			default:
+				check(0);
+			}
 		}
 	}
 }
diff --git a/Engine/Source/Runtime/VulkanRHI/Private/VulkanPipeline.cpp b/Engine/Source/Runtime/VulkanRHI/Private/VulkanPipeline.cpp
--- a/Engine/Source/Runtime/VulkanRHI/Private/VulkanPipeline.cpp
+++ b/Engine/Source/Runtime/VulkanRHI/Private/VulkanPipeline.cpp
@@ -1781,6 +1781,13 @@ void FVulkanPipelineStateCacheManager::CreateGfxEntry(const FGraphicsPipelineSta
 		{
 			const FVulkanShaderHeader& PSHeader = Shaders[ShaderStage::Pixel]->GetCodeHeader();
 			DescriptorSetLayoutInfo.ProcessBindingsForStage(VK_SHADER_STAGE_FRAGMENT_BIT, ShaderStage::Pixel, PSHeader, UBGatherInfo);
+
+			if (PSHeader.InputAttachmentInfos.Num())
+			{
+				// input attachements can't exist in a first sub-pass
+				check(PSOInitializer.SubpassHint != ESubpassHint::None);
+				check(PSOInitializer.SubpassIndex != 0);
+			}
 		}
 
 #if VULKAN_SUPPORTS_GEOMETRY_SHADERS
@@ -1797,14 +1804,6 @@ void FVulkanPipelineStateCacheManager::CreateGfxEntry(const FGraphicsPipelineSta
 		DescriptorSetLayoutInfo.FinalizeBindings<false>(*Device, UBGatherInfo, ImmutableSamplers);
 	}
 
-	FDescriptorSetRemappingInfo& RemappingInfo = DescriptorSetLayoutInfo.RemappingInfo;
-
-	if (RemappingInfo.InputAttachmentData.Num())
-	{
-		// input attachements can't exist in a first sub-pass
-		check(PSOInitializer.SubpassHint != ESubpassHint::None); 
-		check(PSOInitializer.SubpassIndex != 0);
-	}
 	OutGfxEntry->SubpassIndex = PSOInitializer.SubpassIndex;
 
 	FVulkanBlendState* BlendState = ResourceCast(PSOInitializer.BlendState);
diff --git a/Engine/Source/Runtime/VulkanRHI/Private/VulkanShaders.cpp b/Engine/Source/Runtime/VulkanRHI/Private/VulkanShaders.cpp
--- a/Engine/Source/Runtime/VulkanRHI/Private/VulkanShaders.cpp
+++ b/Engine/Source/Runtime/VulkanRHI/Private/VulkanShaders.cpp
@@ -628,20 +628,6 @@ void FVulkanLayout::Compile(FVulkanDescriptorSetLayoutMap& DSetLayoutMap)
 	}
 }
 
-bool FVulkanGfxLayout::UsesInputAttachment(FVulkanShaderHeader::EAttachmentType AttachmentType) const
-{
-	const TArray<FVulkanShaderHeader::FInputAttachmentInfo>& InputAttachmentData = GfxPipelineDescriptorInfo.GetInputAttachmentData();
-	for (const FVulkanShaderHeader::FInputAttachmentInfo& Input : InputAttachmentData)
-	{
-		if (Input.Type == AttachmentType)
-		{
-			return true;
-		}
-	}
-
-	return false;
-}
-
 uint32 FVulkanDescriptorSetWriter::SetupDescriptorWrites(
 	const TArray<VkDescriptorType>& Types, FVulkanHashableDescriptorInfo* InHashableDescriptorInfos,
 	VkWriteDescriptorSet* InWriteDescriptors, VkDescriptorImageInfo* InImageInfo, VkDescriptorBufferInfo* InBufferInfo, uint8* InBindingToDynamicOffsetMap,
@@ -793,13 +779,6 @@ void FVulkanDescriptorSetsLayoutInfo::FinalizeBindings(const FVulkanDevice& Devi
 					AddDescriptor(Stage, Binding);
 				}
 			}
-
-			if (ShaderHeader->InputAttachmentInfos.Num())
-			{
-				check(Stage == ShaderStage::Pixel);
-				check(RemappingInfo.InputAttachmentData.Num() == 0);
-				RemappingInfo.InputAttachmentData = ShaderHeader->InputAttachmentInfos;
-			}
 		}
 	}
 
