VulkanRHI: Actually remove unused InputAttachmentData array.

#rnx

diff --git a/Engine/Source/Runtime/VulkanRHI/Private/VulkanDescriptorSets.h b/Engine/Source/Runtime/VulkanRHI/Private/VulkanDescriptorSets.h
--- a/Engine/Source/Runtime/VulkanRHI/Private/VulkanDescriptorSets.h
+++ b/Engine/Source/Runtime/VulkanRHI/Private/VulkanDescriptorSets.h
@@ -71,20 +71,8 @@ struct FDescriptorSetRemappingInfo
 	};
 	TStaticArray<FStageInfo, ShaderStage::NumStages>	StageInfos;
 
-	TArray<FVulkanShaderHeader::FInputAttachmentInfo>	InputAttachmentData;
-
 	inline bool operator==(const FDescriptorSetRemappingInfo& In) const
 	{
-		if (InputAttachmentData.Num() != In.InputAttachmentData.Num())
-		{
-			return false;
-		}
-
-		if (FMemory::Memcmp(InputAttachmentData.GetData(), In.InputAttachmentData.GetData(), sizeof(FVulkanShaderHeader::FInputAttachmentInfo) * InputAttachmentData.Num()))
-		{
-			return false;
-		}
-
 		for (uint32 StageInfosIndex = 0; StageInfosIndex < ShaderStage::NumStages; ++StageInfosIndex)
 		{
 			if (StageInfos[StageInfosIndex].PackedGlobalsSize != In.StageInfos[StageInfosIndex].PackedGlobalsSize ||
