VulkanRHI: Fix RHI validation error about StaticSlot count.

#rnx

diff --git a/Engine/Source/Runtime/VulkanRHI/Private/VulkanShaders.cpp b/Engine/Source/Runtime/VulkanRHI/Private/VulkanShaders.cpp
--- a/Engine/Source/Runtime/VulkanRHI/Private/VulkanShaders.cpp
+++ b/Engine/Source/Runtime/VulkanRHI/Private/VulkanShaders.cpp
@@ -276,17 +276,21 @@ FVulkanShader::FVulkanShader(FVulkanDevice* InDevice, EShaderFrequency InFrequen
 
 	checkf(SpirvContainer.GetSizeBytes() != 0, TEXT("Empty SPIR-V! %s"), *CodeHeader.DebugName);
 
-	StaticSlots.Reserve(CodeHeader.UniformBufferInfos.Num());
-
-	for (const FVulkanShaderHeader::FUniformBufferInfo& UBInfo : CodeHeader.UniformBufferInfos)
+	const int32 NumGlobalPackedBuffer = (CodeHeader.PackedGlobalsSize > 0) ? 1 : 0;
+	if (CodeHeader.UniformBufferInfos.Num() > NumGlobalPackedBuffer)
 	{
-		if (const FShaderParametersMetadata* Metadata = FindUniformBufferStructByLayoutHash(UBInfo.LayoutHash))
-		{
-			StaticSlots.Add(Metadata->GetLayout().StaticSlot);
-		}
-		else
+		StaticSlots.Reserve(CodeHeader.UniformBufferInfos.Num());
+
+		for (const FVulkanShaderHeader::FUniformBufferInfo& UBInfo : CodeHeader.UniformBufferInfos)
 		{
-			StaticSlots.Add(MAX_UNIFORM_BUFFER_STATIC_SLOTS);
+			if (const FShaderParametersMetadata* Metadata = FindUniformBufferStructByLayoutHash(UBInfo.LayoutHash))
+			{
+				StaticSlots.Add(Metadata->GetLayout().StaticSlot);
+			}
+			else
+			{
+				StaticSlots.Add(MAX_UNIFORM_BUFFER_STATIC_SLOTS);
+			}
 		}
 	}
 
