Fixed Vulkan Validation Error for SceneColorMS texture not getting
VK_IMAGE_USAGE_SAMPLED_BIT because it's memoryless

diff --git a/Engine/Source/Runtime/VulkanRHI/Private/VulkanResources.h b/Engine/Source/Runtime/VulkanRHI/Private/VulkanResources.h
--- a/Engine/Source/Runtime/VulkanRHI/Private/VulkanResources.h
+++ b/Engine/Source/Runtime/VulkanRHI/Private/VulkanResources.h
@@ -705,7 +705,8 @@ public:
 
 	inline bool SupportsSampling() const
 	{
-		return EnumHasAllFlags(GPixelFormats[GetDesc().Format].Capabilities, EPixelFormatCapabilities::TextureSample);
+		return EnumHasAllFlags(GPixelFormats[GetDesc().Format].Capabilities, EPixelFormatCapabilities::TextureSample) &&
+			((ImageUsageFlags & VK_IMAGE_USAGE_SAMPLED_BIT) != 0);
 	}
 
 	inline VkImageLayout GetDefaultLayout() const
