VulkanRHI: Fix Vulkan bindless and ray tracing with latest descriptor remapping changes.

#rnx

diff --git a/Engine/Source/Developer/VulkanShaderFormat/Private/VulkanShaderCompiler.cpp b/Engine/Source/Developer/VulkanShaderFormat/Private/VulkanShaderCompiler.cpp
--- a/Engine/Source/Developer/VulkanShaderFormat/Private/VulkanShaderCompiler.cpp
+++ b/Engine/Source/Developer/VulkanShaderFormat/Private/VulkanShaderCompiler.cpp
@@ -473,10 +473,26 @@ static void BuildShaderOutput(
 				FVulkanShaderHeader::FUniformBufferInfo& UniformBufferInfo = Header.UniformBufferInfos[UniformBufferIndex];
 				UniformBufferInfo.bHasResources = 1;
 
-				const uint32 UBLayoutHash = SerializedOutput.CompilerSRT.ResourceTableLayoutHashes[UniformBufferIndex];
-				checkf(!UniformBufferInfo.LayoutHash || (UniformBufferInfo.LayoutHash == UBLayoutHash), 
-					TEXT("Layout hash should be unset (resource only UB) or identical (UB with constants and resources)!"));
-				UniformBufferInfo.LayoutHash = UBLayoutHash;
+				const bool bIsRootParamStructure = (ParameterName == FShaderParametersMetadata::kRootUniformBufferBindingName) && ShaderInput.RootParametersStructure;
+				if (bIsRootParamStructure)
+				{
+					check(UniformBufferIndex == FShaderParametersMetadata::kRootCBufferBindingIndex);
+					const uint32 UBLayoutHash = SerializedOutput.CompilerSRT.ResourceTableLayoutHashes[UniformBufferIndex];
+
+					checkf(!UBLayoutHash || (UBLayoutHash == ShaderInput.RootParametersStructure->GetLayoutHash()),
+						TEXT("Resource table layout hash for RootParametersStructure (0x%08X) should be unset (0x0) or identical to shader input (0x%08X)!"),
+						UBLayoutHash, ShaderInput.RootParametersStructure->GetLayoutHash());
+
+					SerializedOutput.CompilerSRT.ResourceTableLayoutHashes[UniformBufferIndex] = ShaderInput.RootParametersStructure->GetLayoutHash();
+				}
+				else
+				{
+					const uint32 UBLayoutHash = SerializedOutput.CompilerSRT.ResourceTableLayoutHashes[UniformBufferIndex];
+					checkf(!UniformBufferInfo.LayoutHash || (UniformBufferInfo.LayoutHash == UBLayoutHash),
+						TEXT("Existing layout hash (0x%08X) should be unset (resource only UB) or identical to resource table (0x%08X)!"),
+						UniformBufferInfo.LayoutHash, UBLayoutHash);
+					UniformBufferInfo.LayoutHash = UBLayoutHash;
+				}
 			}
 		}
 	}
diff --git a/Engine/Source/Runtime/VulkanRHI/Private/VulkanPipelineState.cpp b/Engine/Source/Runtime/VulkanRHI/Private/VulkanPipelineState.cpp
--- a/Engine/Source/Runtime/VulkanRHI/Private/VulkanPipelineState.cpp
+++ b/Engine/Source/Runtime/VulkanRHI/Private/VulkanPipelineState.cpp
@@ -230,9 +230,10 @@ void FVulkanComputePipelineDescriptorState::UpdateBindlessDescriptors(FVulkanCom
 	const FDescriptorSetRemappingInfo::FStageInfo& StageInfo = RemappingInfo->StageInfos[Stage];
 
 	TArray<VkDescriptorAddressInfoEXT>& DescriptorAddressInfos = StageUBs[Stage];
-	DescriptorAddressInfos.SetNumZeroed((StageInfo.PackedGlobalsSize ? 1 : 0) + StageInfo.NumBoundUniformBuffers);
+	DescriptorAddressInfos.SetNumZeroed(StageInfo.NumBoundUniformBuffers);
+	uint32 UBIndex = 0;
 
-	// PackedUniformBuffersDirty ?
+	// UBs are currently set from a fresh batch of descriptors for every call, so ignore PackedUniformBuffersDirty
 	check(PackedUniformBuffersMask <= 1);
 	if (PackedUniformBuffersMask != 0)
 	{
@@ -240,6 +241,7 @@ void FVulkanComputePipelineDescriptorState::UpdateBindlessDescriptors(FVulkanCom
 		const int32 UBSize = StagedUniformBuffer.Num();
 		const int32 BindingIndex = 0;
 
+		// :todo-jn: Use PackedUniformBuffersDirty and cache last RingBufferOffset to avoid copy when there are no changes (while updating lifetime in UniformBufferUploader)
 		const uint64 RingBufferOffset = UniformBufferUploader->AllocateMemory(UBSize, UBOffsetAlignment, CmdBuffer);
 
 		// Make sure it wasn't written to already
@@ -253,9 +255,10 @@ void FVulkanComputePipelineDescriptorState::UpdateBindlessDescriptors(FVulkanCom
 		FMemory::Memcpy(CPURingBufferBase + RingBufferOffset, StagedUniformBuffer.GetData(), UBSize);
 
 		PackedUniformBuffersDirty = 0;
+		++UBIndex;
 	}
 
-	for (uint32 UBIndex = 0; UBIndex < StageInfo.NumBoundUniformBuffers; ++UBIndex)
+	for (;UBIndex < StageInfo.NumBoundUniformBuffers; ++UBIndex)
 	{
 		VkDescriptorAddressInfoEXT& DescriptorAddressInfo = DescriptorAddressInfos[UBIndex];
 		check(DescriptorAddressInfo.sType == 0);
@@ -438,9 +441,10 @@ void FVulkanGraphicsPipelineDescriptorState::UpdateBindlessDescriptors(FVulkanCo
 			const FDescriptorSetRemappingInfo::FStageInfo& StageInfo = RemappingInfo->StageInfos[Stage];
 
 			TArray<VkDescriptorAddressInfoEXT>& DescriptorAddressInfos = StageUBs[Stage];
-			DescriptorAddressInfos.SetNumZeroed((StageInfo.PackedGlobalsSize ? 1 : 0) + StageInfo.NumBoundUniformBuffers);
+			DescriptorAddressInfos.SetNumZeroed(StageInfo.NumBoundUniformBuffers);
+			uint32 UBIndex = 0;
 
-			// PackedUniformBuffersDirty ?
+			// UBs are currently set from a fresh batch of descriptors for every call, so ignore PackedUniformBuffersDirty
 			check(PackedUniformBuffersMask[Stage] <= 1);
 			if (PackedUniformBuffersMask[Stage] != 0)
 			{
@@ -448,6 +452,7 @@ void FVulkanGraphicsPipelineDescriptorState::UpdateBindlessDescriptors(FVulkanCo
 				const int32 UBSize = StagedUniformBuffer.Num();
 				const int32 BindingIndex = 0;
 
+				// :todo-jn: Use PackedUniformBuffersDirty and cache last RingBufferOffset to avoid copy when there are no changes (while updating lifetime in UniformBufferUploader)
 				const uint64 RingBufferOffset = UniformBufferUploader->AllocateMemory(UBSize, UBOffsetAlignment, CmdBuffer);
 
 				// Make sure it wasn't written to already
@@ -461,9 +466,10 @@ void FVulkanGraphicsPipelineDescriptorState::UpdateBindlessDescriptors(FVulkanCo
 				FMemory::Memcpy(CPURingBufferBase + RingBufferOffset, StagedUniformBuffer.GetData(), UBSize);
 
 				PackedUniformBuffersDirty[Stage] = 0;
+				++UBIndex;
 			}
 
-			for (uint32 UBIndex = 0; UBIndex < StageInfo.NumBoundUniformBuffers; ++UBIndex)
+			for (; UBIndex < StageInfo.NumBoundUniformBuffers; ++UBIndex)
 			{
 				VkDescriptorAddressInfoEXT& DescriptorAddressInfo = DescriptorAddressInfos[UBIndex];
 				check(DescriptorAddressInfo.sType == 0);
