VulkanRHI: remove the need for HasDescriptorsInSetMask, get the same information from the shader code header.

#rnx

diff --git a/Engine/Source/Runtime/VulkanRHI/Private/VulkanDescriptorSets.h b/Engine/Source/Runtime/VulkanRHI/Private/VulkanDescriptorSets.h
--- a/Engine/Source/Runtime/VulkanRHI/Private/VulkanDescriptorSets.h
+++ b/Engine/Source/Runtime/VulkanRHI/Private/VulkanDescriptorSets.h
@@ -554,9 +554,7 @@ class FVulkanComputePipelineDescriptorInfo
 {
 public:
 	FVulkanComputePipelineDescriptorInfo()
-		: HasDescriptorsInSetMask(0)
-		, RemappingInfo(nullptr)
-		, bInitialized(false)
+		: RemappingInfo(nullptr)
 	{
 	}
 	
@@ -568,16 +566,17 @@ public:
 
 	inline bool IsInitialized() const
 	{
-		return bInitialized;
+		return (RemappingInfo != nullptr);
 	}
 
-	void Initialize(const FDescriptorSetRemappingInfo& InRemappingInfo);
+	void Initialize(const FDescriptorSetRemappingInfo& InRemappingInfo)
+	{
+		check(!RemappingInfo);
+		RemappingInfo = &InRemappingInfo;
+	}
 
 protected:
-	// Cached data from FDescriptorSetRemappingInfo
-	uint32															HasDescriptorsInSetMask;
 	const FDescriptorSetRemappingInfo*								RemappingInfo;
-	bool															bInitialized;
 
 	friend class FVulkanComputePipelineDescriptorState;
 };
@@ -587,9 +586,7 @@ class FVulkanGfxPipelineDescriptorInfo
 {
 public:
 	FVulkanGfxPipelineDescriptorInfo()
-		: HasDescriptorsInSetMask(0)
-		, RemappingInfo(nullptr)
-		, bInitialized(false)
+		: RemappingInfo(nullptr)
 	{
 	}
 
@@ -600,17 +597,17 @@ public:
 
 	inline bool IsInitialized() const
 	{
-		return bInitialized;
+		return (RemappingInfo != nullptr);
 	}
 
-	void Initialize(const FDescriptorSetRemappingInfo& InRemappingInfo);
+	void Initialize(const FDescriptorSetRemappingInfo& InRemappingInfo)
+	{
+		check(!RemappingInfo);
+		RemappingInfo = &InRemappingInfo;
+	}
 
 protected:
-	// Cached data from FDescriptorSetRemappingInfo
-	uint32													HasDescriptorsInSetMask;
-
 	const FDescriptorSetRemappingInfo*						RemappingInfo;
-	bool													bInitialized;
 
 	friend class FVulkanGraphicsPipelineDescriptorState;
 };
diff --git a/Engine/Source/Runtime/VulkanRHI/Private/VulkanPipelineState.cpp b/Engine/Source/Runtime/VulkanRHI/Private/VulkanPipelineState.cpp
--- a/Engine/Source/Runtime/VulkanRHI/Private/VulkanPipelineState.cpp
+++ b/Engine/Source/Runtime/VulkanRHI/Private/VulkanPipelineState.cpp
@@ -51,7 +51,7 @@ FVulkanComputePipelineDescriptorState::FVulkanComputePipelineDescriptorState(FVu
 	DescriptorSetsLayout = &InComputePipeline->GetLayout().GetDescriptorSetsLayout();
 	PipelineDescriptorInfo = &InComputePipeline->GetComputeLayout().GetComputePipelineDescriptorInfo();
 
-	UsedSetsMask = PipelineDescriptorInfo->HasDescriptorsInSetMask;
+	UsedSetsMask = (CodeHeader.Bindings.Num() > 0) ? 1 : 0;
 
 	CreateDescriptorWriteInfos();
 	InComputePipeline->AddRef();
@@ -300,12 +300,14 @@ FVulkanGraphicsPipelineDescriptorState::FVulkanGraphicsPipelineDescriptorState(F
 		FVulkanGfxLayout& GfxLayout  = *(FVulkanGfxLayout*)InGfxPipeline->Layout;
 		PipelineDescriptorInfo = &GfxLayout.GetGfxPipelineDescriptorInfo();
 
-		UsedSetsMask = PipelineDescriptorInfo->HasDescriptorsInSetMask;
+		UsedSetsMask = 0;
+
 		const FVulkanShaderFactory& ShaderFactory = Device->GetShaderFactory();
 
 		const FVulkanVertexShader* VertexShader = ShaderFactory.LookupShader<FVulkanVertexShader>(InGfxPipeline->GetShaderKey(SF_Vertex));
 		check(VertexShader);
 		PackedUniformBuffers[ShaderStage::Vertex].Init(VertexShader->GetCodeHeader(), PackedUniformBuffersMask[ShaderStage::Vertex]);
+		UsedSetsMask |= VertexShader->GetCodeHeader().Bindings.Num() ? (1u << ShaderStage::Vertex) : 0u;
 
 		uint64 PixelShaderKey = InGfxPipeline->GetShaderKey(SF_Pixel);
 		if (PixelShaderKey)
@@ -314,6 +316,7 @@ FVulkanGraphicsPipelineDescriptorState::FVulkanGraphicsPipelineDescriptorState(F
 			check(PixelShader);
 
 			PackedUniformBuffers[ShaderStage::Pixel].Init(PixelShader->GetCodeHeader(), PackedUniformBuffersMask[ShaderStage::Pixel]);
+			UsedSetsMask |= PixelShader->GetCodeHeader().Bindings.Num() ? (1u << ShaderStage::Pixel) : 0u;
 		}
 
 #if VULKAN_SUPPORTS_GEOMETRY_SHADERS
@@ -324,6 +327,7 @@ FVulkanGraphicsPipelineDescriptorState::FVulkanGraphicsPipelineDescriptorState(F
 			check(GeometryShader);
 
 			PackedUniformBuffers[ShaderStage::Geometry].Init(GeometryShader->GetCodeHeader(), PackedUniformBuffersMask[ShaderStage::Geometry]);
+			UsedSetsMask |= GeometryShader->GetCodeHeader().Bindings.Num() ? (1u << ShaderStage::Geometry) : 0u;
 		}
 #endif
 
diff --git a/Engine/Source/Runtime/VulkanRHI/Private/VulkanShaders.cpp b/Engine/Source/Runtime/VulkanRHI/Private/VulkanShaders.cpp
--- a/Engine/Source/Runtime/VulkanRHI/Private/VulkanShaders.cpp
+++ b/Engine/Source/Runtime/VulkanRHI/Private/VulkanShaders.cpp
@@ -790,43 +790,6 @@ void FVulkanDescriptorSetsLayoutInfo::FinalizeBindings(const FVulkanDevice& Devi
 	GenerateHash(ImmutableSamplers, bIsCompute ? VK_PIPELINE_BIND_POINT_COMPUTE : VK_PIPELINE_BIND_POINT_GRAPHICS);
 }
 
-void FVulkanComputePipelineDescriptorInfo::Initialize(const FDescriptorSetRemappingInfo& InRemappingInfo)
-{
-	check(!bInitialized);
-
-	const FDescriptorSetRemappingInfo::FStageInfo& StageInfo = InRemappingInfo.StageInfos[ShaderStage::Compute];
-
-	RemappingInfo = &InRemappingInfo;
-
-	if (StageInfo.Types.Num() > 0)
-	{
-		HasDescriptorsInSetMask = HasDescriptorsInSetMask | (1 << 0);
-	}
-
-	bInitialized = true;
-}
-
-void FVulkanGfxPipelineDescriptorInfo::Initialize(const FDescriptorSetRemappingInfo& InRemappingInfo)
-{
-	check(!bInitialized);
-
-	for (int32 StageIndex = 0; StageIndex < ShaderStage::NumStages; ++StageIndex)
-	{
-		const FDescriptorSetRemappingInfo::FStageInfo& StageInfo = InRemappingInfo.StageInfos[StageIndex];
-
-		if (StageInfo.Types.Num() > 0)
-		{
-			check(StageIndex < sizeof(HasDescriptorsInSetMask) * 8);
-			HasDescriptorsInSetMask = HasDescriptorsInSetMask | (1 << StageIndex);
-		}
-	}
-
-	RemappingInfo = &InRemappingInfo;
-
-	bInitialized = true;
-}
-
-
 FVulkanBoundShaderState::FVulkanBoundShaderState(FRHIVertexDeclaration* InVertexDeclarationRHI, FRHIVertexShader* InVertexShaderRHI,
 	FRHIPixelShader* InPixelShaderRHI, FRHIGeometryShader* InGeometryShaderRHI)
 	: CacheLink(InVertexDeclarationRHI, InVertexShaderRHI, InPixelShaderRHI, InGeometryShaderRHI, this)
