VulkanRHI: Remove FVulkanGfxPipelineDescriptorInfo and FVulkanComputePipelineDescriptorInfo classes, pull the same information from the shader code header and the layout.

#rnx

diff --git a/Engine/Source/Runtime/VulkanRHI/Private/VulkanCommands.cpp b/Engine/Source/Runtime/VulkanRHI/Private/VulkanCommands.cpp
--- a/Engine/Source/Runtime/VulkanRHI/Private/VulkanCommands.cpp
+++ b/Engine/Source/Runtime/VulkanRHI/Private/VulkanCommands.cpp
@@ -503,14 +503,13 @@ void FVulkanCommandListContext::RHISetShaderUniformBuffer(FRHIGraphicsShader* Sh
 	const FVulkanShaderHeader& CodeHeader = Shader->GetCodeHeader();
 	checkfSlow(!CodeHeader.UniformBufferInfos[BufferIndex].LayoutHash || (CodeHeader.UniformBufferInfos[BufferIndex].LayoutHash == UniformBuffer->GetLayout().GetHash()),
 		TEXT("Mismatched UB layout! Got hash 0x%x, expected 0x%x!"), UniformBuffer->GetLayout().GetHash(), CodeHeader.UniformBufferInfos[BufferIndex].LayoutHash);
-	const FVulkanGfxPipelineDescriptorInfo& DescriptorInfo = PendingGfxState->CurrentState->GetGfxPipelineDescriptorInfo();
 
 	bool bHasResources = false;
 	if (BufferIndex < CodeHeader.NumBoundUniformBuffers)
 	{
 		checkSlow(UniformBuffer->GetLayout().ConstantBufferSize > 0);
-
-		const VkDescriptorType DescriptorType = DescriptorInfo.GetDescriptorType(Stage, BufferIndex);
+		
+		const VkDescriptorType DescriptorType = PendingGfxState->CurrentState->GetDescriptorType(Stage, BufferIndex);
 
 		if (DescriptorType == VK_DESCRIPTOR_TYPE_UNIFORM_BUFFER_DYNAMIC)
 		{
@@ -553,7 +552,6 @@ void FVulkanCommandListContext::RHISetShaderUniformBuffer(FRHIComputeShader* Com
 	FVulkanComputeShader* Shader = ResourceCast(ComputeShaderRHI);
 	FVulkanUniformBuffer* UniformBuffer = ResourceCast(BufferRHI);
 
-	const FVulkanComputePipelineDescriptorInfo& DescriptorInfo = PendingComputeState->CurrentState->GetComputePipelineDescriptorInfo();
 	const FVulkanShaderHeader& CodeHeader = Shader->GetCodeHeader();
 	checkfSlow(!CodeHeader.UniformBufferInfos[BufferIndex].LayoutHash || (CodeHeader.UniformBufferInfos[BufferIndex].LayoutHash == UniformBuffer->GetLayout().GetHash()), 
 		TEXT("Mismatched UB layout! Got hash 0x%x, expected 0x%x!"), UniformBuffer->GetLayout().GetHash(), CodeHeader.UniformBufferInfos[BufferIndex].LayoutHash);
@@ -563,8 +561,8 @@ void FVulkanCommandListContext::RHISetShaderUniformBuffer(FRHIComputeShader* Com
 	if (BufferIndex < CodeHeader.NumBoundUniformBuffers)
 	{
 		checkSlow(UniformBuffer->GetLayout().ConstantBufferSize > 0);
-		
-		const VkDescriptorType DescriptorType = DescriptorInfo.GetDescriptorType(ShaderStage::Compute, BufferIndex);
+
+		const VkDescriptorType DescriptorType = State.GetDescriptorType(ShaderStage::Compute, BufferIndex);
 
 		if (DescriptorType == VK_DESCRIPTOR_TYPE_UNIFORM_BUFFER_DYNAMIC)
 		{
diff --git a/Engine/Source/Runtime/VulkanRHI/Private/VulkanDescriptorSets.h b/Engine/Source/Runtime/VulkanRHI/Private/VulkanDescriptorSets.h
--- a/Engine/Source/Runtime/VulkanRHI/Private/VulkanDescriptorSets.h
+++ b/Engine/Source/Runtime/VulkanRHI/Private/VulkanDescriptorSets.h
@@ -281,8 +281,6 @@ protected:
 
 	friend class FVulkanPipelineStateCacheManager;
 	friend class FVulkanCommonPipelineDescriptorState;
-	friend class FVulkanGfxPipelineDescriptorInfo;
-	friend class FVulkanComputePipelineDescriptorInfo;
 	friend class FVulkanLayout;
 };
 
@@ -549,72 +547,6 @@ struct FVulkanDescriptorSetWriteContainer
 };
 
 
-// Smaller data structure for runtime information about descriptor sets and bindings for a pipeline
-class FVulkanComputePipelineDescriptorInfo
-{
-public:
-	FVulkanComputePipelineDescriptorInfo()
-		: RemappingInfo(nullptr)
-	{
-	}
-	
-	inline VkDescriptorType GetDescriptorType(uint8 DescriptorSet, int32 DescriptorIndex) const
-	{
-		check(DescriptorSet == ShaderStage::Compute);
-		return RemappingInfo->StageInfos[DescriptorSet].Types[DescriptorIndex];
-	}
-
-	inline bool IsInitialized() const
-	{
-		return (RemappingInfo != nullptr);
-	}
-
-	void Initialize(const FDescriptorSetRemappingInfo& InRemappingInfo)
-	{
-		check(!RemappingInfo);
-		RemappingInfo = &InRemappingInfo;
-	}
-
-protected:
-	const FDescriptorSetRemappingInfo*								RemappingInfo;
-
-	friend class FVulkanComputePipelineDescriptorState;
-};
-
-// Smaller data structure for runtime information about descriptor sets and bindings for a pipeline
-class FVulkanGfxPipelineDescriptorInfo
-{
-public:
-	FVulkanGfxPipelineDescriptorInfo()
-		: RemappingInfo(nullptr)
-	{
-	}
-
-	inline VkDescriptorType GetDescriptorType(uint8 DescriptorSet, int32 DescriptorIndex) const
-	{
-		return RemappingInfo->StageInfos[DescriptorSet].Types[DescriptorIndex];
-	}
-
-	inline bool IsInitialized() const
-	{
-		return (RemappingInfo != nullptr);
-	}
-
-	void Initialize(const FDescriptorSetRemappingInfo& InRemappingInfo)
-	{
-		check(!RemappingInfo);
-		RemappingInfo = &InRemappingInfo;
-	}
-
-protected:
-	const FDescriptorSetRemappingInfo*						RemappingInfo;
-
-	friend class FVulkanGraphicsPipelineDescriptorState;
-};
-
-
-
-
 // This class encapsulates updating VkWriteDescriptorSet structures (but doesn't own them), and their flags for dirty ranges; it is intended
 // to be used to access a sub-region of a long array of VkWriteDescriptorSet (ie FVulkanDescriptorSetWriteContainer)
 class FVulkanDescriptorSetWriter
@@ -1104,13 +1036,6 @@ public:
 		return true;
 	}
 
-	inline const FVulkanGfxPipelineDescriptorInfo& GetGfxPipelineDescriptorInfo() const
-	{
-		return GfxPipelineDescriptorInfo;
-	}
-
-protected:
-	FVulkanGfxPipelineDescriptorInfo		GfxPipelineDescriptorInfo;
 	friend class FVulkanPipelineStateCacheManager;
 };
 
@@ -1127,13 +1052,6 @@ public:
 		return false;
 	}
 
-	inline const FVulkanComputePipelineDescriptorInfo& GetComputePipelineDescriptorInfo() const
-	{
-		return ComputePipelineDescriptorInfo;
-	}
-
-protected:
-	FVulkanComputePipelineDescriptorInfo		ComputePipelineDescriptorInfo;
 	friend class FVulkanPipelineStateCacheManager;
 };
 
diff --git a/Engine/Source/Runtime/VulkanRHI/Private/VulkanPendingState.h b/Engine/Source/Runtime/VulkanRHI/Private/VulkanPendingState.h
--- a/Engine/Source/Runtime/VulkanRHI/Private/VulkanPendingState.h
+++ b/Engine/Source/Runtime/VulkanRHI/Private/VulkanPendingState.h
@@ -280,7 +280,6 @@ public:
 
 	inline void SetPackedGlobalShaderParameter(ShaderStage::EStage Stage, uint32 BufferIndex, uint32 Offset, uint32 NumBytes, const void* NewValue)
 	{
-		const FVulkanGfxPipelineDescriptorInfo& DescriptorInfo = CurrentState->GetGfxPipelineDescriptorInfo();
 		CurrentState->SetPackedGlobalShaderParameter(Stage, BufferIndex, Offset, NumBytes, NewValue);
 	}
 
diff --git a/Engine/Source/Runtime/VulkanRHI/Private/VulkanPipeline.cpp b/Engine/Source/Runtime/VulkanRHI/Private/VulkanPipeline.cpp
--- a/Engine/Source/Runtime/VulkanRHI/Private/VulkanPipeline.cpp
+++ b/Engine/Source/Runtime/VulkanRHI/Private/VulkanPipeline.cpp
@@ -1727,12 +1727,6 @@ FVulkanLayout* FVulkanPipelineStateCacheManager::FindOrAddLayout(const FVulkanDe
 	Layout->DescriptorSetLayout.CopyFrom(DescriptorSetLayoutInfo);
 	Layout->Compile(DSetLayoutMap);
 
-	if(GfxLayout)
-	{
-		GfxLayout->GfxPipelineDescriptorInfo.Initialize(GfxLayout->GetDescriptorSetsLayout().RemappingInfo);
-	}
-
-
 	LayoutMap.Add(DescriptorSetLayoutInfo, Layout);
 	return Layout;
 }
@@ -2101,7 +2095,6 @@ FGraphicsPipelineStateRHIRef FVulkanPipelineStateCacheManager::RHICreateGraphics
 
 			FVulkanLayout* Layout = FindOrAddLayout(DescriptorSetLayoutInfo, true);
 			FVulkanGfxLayout* GfxLayout = (FVulkanGfxLayout*)Layout;
-			check(GfxLayout->GfxPipelineDescriptorInfo.IsInitialized());
 			NewPSO->Layout = GfxLayout;
 			NewPSO->bHasInputAttachments = GfxLayout->GetDescriptorSetsLayout().HasInputAttachments();
 		}
@@ -2264,10 +2257,6 @@ FVulkanComputePipeline* FVulkanPipelineStateCacheManager::CreateComputePipelineF
 	DescriptorSetLayoutInfo.FinalizeBindings<true>(*Device, UBGatherInfo, TArrayView<FRHISamplerState*>());
 	FVulkanLayout* Layout = FindOrAddLayout(DescriptorSetLayoutInfo, false);
 	FVulkanComputeLayout* ComputeLayout = (FVulkanComputeLayout*)Layout;
-	if (!ComputeLayout->ComputePipelineDescriptorInfo.IsInitialized())
-	{
-		ComputeLayout->ComputePipelineDescriptorInfo.Initialize(Layout->GetDescriptorSetsLayout().RemappingInfo);
-	}
 
 	TRefCountPtr<FVulkanShaderModule> ShaderModule = Shader->GetOrCreateHandle(Layout, Layout->GetDescriptorSetLayoutHash());
 
diff --git a/Engine/Source/Runtime/VulkanRHI/Private/VulkanPipelineState.cpp b/Engine/Source/Runtime/VulkanRHI/Private/VulkanPipelineState.cpp
--- a/Engine/Source/Runtime/VulkanRHI/Private/VulkanPipelineState.cpp
+++ b/Engine/Source/Runtime/VulkanRHI/Private/VulkanPipelineState.cpp
@@ -49,7 +49,6 @@ FVulkanComputePipelineDescriptorState::FVulkanComputePipelineDescriptorState(FVu
 	PackedUniformBuffers.Init(CodeHeader, PackedUniformBuffersMask);
 
 	DescriptorSetsLayout = &InComputePipeline->GetLayout().GetDescriptorSetsLayout();
-	PipelineDescriptorInfo = &InComputePipeline->GetComputeLayout().GetComputePipelineDescriptorInfo();
 
 	UsedSetsMask = (CodeHeader.Bindings.Num() > 0) ? 1 : 0;
 
@@ -225,12 +224,10 @@ void FVulkanComputePipelineDescriptorState::UpdateBindlessDescriptors(FVulkanCom
 
 	FVulkanBindlessDescriptorManager::FUniformBufferDescriptorArrays StageUBs;
 
-	const int32 Stage = (int32)ShaderStage::EStage::Compute;
-	const FDescriptorSetRemappingInfo* RESTRICT RemappingInfo = PipelineDescriptorInfo->RemappingInfo;
-	const FDescriptorSetRemappingInfo::FStageInfo& StageInfo = RemappingInfo->StageInfos[Stage];
+	const FVulkanShaderHeader& Header = ComputePipeline->GetShaderCodeHeader();
 
-	TArray<VkDescriptorAddressInfoEXT>& DescriptorAddressInfos = StageUBs[Stage];
-	DescriptorAddressInfos.SetNumZeroed(StageInfo.NumBoundUniformBuffers);
+	TArray<VkDescriptorAddressInfoEXT>& DescriptorAddressInfos = StageUBs[ShaderStage::EStage::Compute];
+	DescriptorAddressInfos.SetNumZeroed(Header.NumBoundUniformBuffers);
 	uint32 UBIndex = 0;
 
 	// UBs are currently set from a fresh batch of descriptors for every call, so ignore PackedUniformBuffersDirty
@@ -258,12 +255,12 @@ void FVulkanComputePipelineDescriptorState::UpdateBindlessDescriptors(FVulkanCom
 		++UBIndex;
 	}
 
-	for (;UBIndex < StageInfo.NumBoundUniformBuffers; ++UBIndex)
+	for (;UBIndex < Header.NumBoundUniformBuffers; ++UBIndex)
 	{
 		VkDescriptorAddressInfoEXT& DescriptorAddressInfo = DescriptorAddressInfos[UBIndex];
 		check(DescriptorAddressInfo.sType == 0);
 
-		VkWriteDescriptorSet& WriteDescriptorSet = DSWriter[Stage].WriteDescriptors[UBIndex];
+		VkWriteDescriptorSet& WriteDescriptorSet = DSWriter[ShaderStage::EStage::Compute].WriteDescriptors[UBIndex];
 		check(WriteDescriptorSet.dstBinding == UBIndex);
 		check(WriteDescriptorSet.dstArrayElement == 0);
 		check(WriteDescriptorSet.descriptorCount == 1);
@@ -298,7 +295,6 @@ FVulkanGraphicsPipelineDescriptorState::FVulkanGraphicsPipelineDescriptorState(F
 		check(InGfxPipeline->Layout);
 		DescriptorSetsLayout = &InGfxPipeline->Layout->GetDescriptorSetsLayout();
 		FVulkanGfxLayout& GfxLayout  = *(FVulkanGfxLayout*)InGfxPipeline->Layout;
-		PipelineDescriptorInfo = &GfxLayout.GetGfxPipelineDescriptorInfo();
 
 		UsedSetsMask = 0;
 
@@ -431,8 +427,6 @@ void FVulkanGraphicsPipelineDescriptorState::UpdateBindlessDescriptors(FVulkanCo
 	uint8* CPURingBufferBase = (uint8*)UniformBufferUploader->GetCPUMappedPointer();
 	const VkDeviceSize UBOffsetAlignment = Device->GetLimits().minUniformBufferOffsetAlignment;
 
-	const FDescriptorSetRemappingInfo* RESTRICT RemappingInfo = PipelineDescriptorInfo->RemappingInfo;
-
 	FVulkanBindlessDescriptorManager::FUniformBufferDescriptorArrays StageUBs;
 
 	// Process updates
@@ -442,10 +436,16 @@ void FVulkanGraphicsPipelineDescriptorState::UpdateBindlessDescriptors(FVulkanCo
 #endif
 		for (int32 Stage = 0; Stage < ShaderStage::NumStages; ++Stage)
 		{
-			const FDescriptorSetRemappingInfo::FStageInfo& StageInfo = RemappingInfo->StageInfos[Stage];
+			const FVulkanShader* VulkanShader = GfxPipeline->GetVulkanShader(GetFrequencyForGfxStage((ShaderStage::EStage)Stage));
+			if (!VulkanShader)
+			{
+				continue;
+			}
+
+			const FVulkanShaderHeader& Header = VulkanShader->GetCodeHeader();
 
 			TArray<VkDescriptorAddressInfoEXT>& DescriptorAddressInfos = StageUBs[Stage];
-			DescriptorAddressInfos.SetNumZeroed(StageInfo.NumBoundUniformBuffers);
+			DescriptorAddressInfos.SetNumZeroed(Header.NumBoundUniformBuffers);
 			uint32 UBIndex = 0;
 
 			// UBs are currently set from a fresh batch of descriptors for every call, so ignore PackedUniformBuffersDirty
@@ -473,7 +473,7 @@ void FVulkanGraphicsPipelineDescriptorState::UpdateBindlessDescriptors(FVulkanCo
 				++UBIndex;
 			}
 
-			for (; UBIndex < StageInfo.NumBoundUniformBuffers; ++UBIndex)
+			for (; UBIndex < Header.NumBoundUniformBuffers; ++UBIndex)
 			{
 				VkDescriptorAddressInfoEXT& DescriptorAddressInfo = DescriptorAddressInfos[UBIndex];
 				check(DescriptorAddressInfo.sType == 0);
diff --git a/Engine/Source/Runtime/VulkanRHI/Private/VulkanPipelineState.h b/Engine/Source/Runtime/VulkanRHI/Private/VulkanPipelineState.h
--- a/Engine/Source/Runtime/VulkanRHI/Private/VulkanPipelineState.h
+++ b/Engine/Source/Runtime/VulkanRHI/Private/VulkanPipelineState.h
@@ -115,6 +115,12 @@ public:
 		DSWriter[DescriptorSet].DynamicOffsets[DynamicOffsetIndex] = DynamicOffset;
 	}
 
+	VkDescriptorType GetDescriptorType(uint8 DescriptorSet, uint32 BindingIndex) const
+	{
+		const TArray<FVulkanDescriptorSetsLayout::FSetLayout>& Layouts = DescriptorSetsLayout->GetLayouts();
+		return Layouts[DescriptorSet].LayoutBindings[BindingIndex].descriptorType;
+	}
+
 protected:
 	void Reset()
 	{
@@ -212,15 +218,7 @@ public:
 		Bind(CmdBuffer, ComputePipeline->GetLayout().GetPipelineLayout(), VK_PIPELINE_BIND_POINT_COMPUTE);
 	}
 
-	inline const FVulkanComputePipelineDescriptorInfo& GetComputePipelineDescriptorInfo() const
-	{
-		return *PipelineDescriptorInfo;
-		//return GfxPipeline->Pipeline->GetGfxLayout().GetGfxPipelineDescriptorInfo();
-	}
-
 protected:
-	const FVulkanComputePipelineDescriptorInfo* PipelineDescriptorInfo;
-
 	FPackedUniformBuffers PackedUniformBuffers;
 	uint32 PackedUniformBuffersMask;
 	uint32 PackedUniformBuffersDirty;
@@ -278,14 +276,7 @@ public:
 		bIsResourcesDirty = true;
 	}
 
-	inline const FVulkanGfxPipelineDescriptorInfo& GetGfxPipelineDescriptorInfo() const
-	{
-		return *PipelineDescriptorInfo;
-	}
-
 protected:
-	const FVulkanGfxPipelineDescriptorInfo* PipelineDescriptorInfo;
-
 	TStaticArray<FPackedUniformBuffers, ShaderStage::NumStages> PackedUniformBuffers;
 	TStaticArray<uint32, ShaderStage::NumStages> PackedUniformBuffersMask;
 	TStaticArray<uint32, ShaderStage::NumStages> PackedUniformBuffersDirty;
