VulkanRHI: Remove FVulkanComputeLayout and FVulkanGfxLayout and track type with a bool in FVulkanLayout.

#rnx

diff --git a/Engine/Source/Runtime/VulkanRHI/Private/VulkanDescriptorSets.h b/Engine/Source/Runtime/VulkanRHI/Private/VulkanDescriptorSets.h
--- a/Engine/Source/Runtime/VulkanRHI/Private/VulkanDescriptorSets.h
+++ b/Engine/Source/Runtime/VulkanRHI/Private/VulkanDescriptorSets.h
@@ -969,10 +969,13 @@ private:
 class FVulkanLayout : public VulkanRHI::FDeviceChild
 {
 public:
-	FVulkanLayout(FVulkanDevice* InDevice);
+	FVulkanLayout(FVulkanDevice* InDevice, bool InGfxLayout);
 	virtual ~FVulkanLayout();
 
-	virtual bool IsGfxLayout() const = 0;
+	bool IsGfxLayout() const
+	{
+		return bIsGfxLayout;
+	}
 
 	inline const FVulkanDescriptorSetsLayout& GetDescriptorSetsLayout() const
 	{
@@ -995,6 +998,7 @@ public:
 	}
 
 protected:
+	const bool bIsGfxLayout;
 	FVulkanDescriptorSetsLayout	DescriptorSetLayout;
 	VkPipelineLayout			PipelineLayout;
 
@@ -1023,39 +1027,6 @@ protected:
 	friend class FVulkanRayTracingPipelineState;
 };
 
-class FVulkanGfxLayout : public FVulkanLayout
-{
-public:
-	FVulkanGfxLayout(FVulkanDevice* InDevice)
-		: FVulkanLayout(InDevice)
-	{
-	}
-
-	virtual bool IsGfxLayout() const final override
-	{
-		return true;
-	}
-
-	friend class FVulkanPipelineStateCacheManager;
-};
-
-class FVulkanComputeLayout : public FVulkanLayout
-{
-public:
-	FVulkanComputeLayout(FVulkanDevice* InDevice)
-		: FVulkanLayout(InDevice)
-	{
-	}
-
-	virtual bool IsGfxLayout() const final override
-	{
-		return false;
-	}
-
-	friend class FVulkanPipelineStateCacheManager;
-};
-
-
 
 class FVulkanGenericDescriptorPool : FNoncopyable
 {
diff --git a/Engine/Source/Runtime/VulkanRHI/Private/VulkanPipeline.cpp b/Engine/Source/Runtime/VulkanRHI/Private/VulkanPipeline.cpp
--- a/Engine/Source/Runtime/VulkanRHI/Private/VulkanPipeline.cpp
+++ b/Engine/Source/Runtime/VulkanRHI/Private/VulkanPipeline.cpp
@@ -1711,19 +1711,7 @@ FVulkanLayout* FVulkanPipelineStateCacheManager::FindOrAddLayout(const FVulkanDe
 		return *FoundLayout;
 	}
 
-	FVulkanLayout* Layout = nullptr;
-	FVulkanGfxLayout* GfxLayout = nullptr;
-
-	if (bGfxLayout)
-	{
-		GfxLayout = new FVulkanGfxLayout(Device);
-		Layout = GfxLayout;
-	}
-	else
-	{
-		Layout = new FVulkanComputeLayout(Device);
-	}
-
+	FVulkanLayout* Layout = new FVulkanLayout(Device, bGfxLayout);
 	Layout->DescriptorSetLayout.CopyFrom(DescriptorSetLayoutInfo);
 	Layout->Compile(DSetLayoutMap);
 
@@ -2094,9 +2082,8 @@ FGraphicsPipelineStateRHIRef FVulkanPipelineStateCacheManager::RHICreateGraphics
 		{
 
 			FVulkanLayout* Layout = FindOrAddLayout(DescriptorSetLayoutInfo, true);
-			FVulkanGfxLayout* GfxLayout = (FVulkanGfxLayout*)Layout;
-			NewPSO->Layout = GfxLayout;
-			NewPSO->bHasInputAttachments = GfxLayout->GetDescriptorSetsLayout().HasInputAttachments();
+			NewPSO->Layout = Layout;
+			NewPSO->bHasInputAttachments = Layout->GetDescriptorSetsLayout().HasInputAttachments();
 		}
 		NewPSO->RenderPass = Device->GetImmediateContext().PrepareRenderPassForPSOCreation(Initializer);
 		{
@@ -2256,7 +2243,7 @@ FVulkanComputePipeline* FVulkanPipelineStateCacheManager::CreateComputePipelineF
 	DescriptorSetLayoutInfo.ProcessBindingsForStage(VK_SHADER_STAGE_COMPUTE_BIT, ShaderStage::Compute, CSHeader, UBGatherInfo);
 	DescriptorSetLayoutInfo.FinalizeBindings<true>(*Device, UBGatherInfo, TArrayView<FRHISamplerState*>());
 	FVulkanLayout* Layout = FindOrAddLayout(DescriptorSetLayoutInfo, false);
-	FVulkanComputeLayout* ComputeLayout = (FVulkanComputeLayout*)Layout;
+	checkSlow(!Layout->IsGfxLayout());
 
 	TRefCountPtr<FVulkanShaderModule> ShaderModule = Shader->GetOrCreateHandle(Layout, Layout->GetDescriptorSetLayoutHash());
 
@@ -2269,7 +2256,7 @@ FVulkanComputePipeline* FVulkanPipelineStateCacheManager::CreateComputePipelineF
 	ANSICHAR EntryPoint[24];
 	Shader->GetEntryPoint(EntryPoint, 24);
 	PipelineInfo.stage.pName = EntryPoint;
-	PipelineInfo.layout = ComputeLayout->GetPipelineLayout();
+	PipelineInfo.layout = Layout->GetPipelineLayout();
 
 	if (Device->SupportsBindless())
 	{
@@ -2307,7 +2294,7 @@ FVulkanComputePipeline* FVulkanPipelineStateCacheManager::CreateComputePipelineF
 		Pipeline->SetValid(false);
 	}
 
-	Pipeline->Layout = ComputeLayout;
+	Pipeline->Layout = Layout;
 
 	INC_DWORD_STAT(STAT_VulkanNumPSOs);
 
diff --git a/Engine/Source/Runtime/VulkanRHI/Private/VulkanPipeline.h b/Engine/Source/Runtime/VulkanRHI/Private/VulkanPipeline.h
--- a/Engine/Source/Runtime/VulkanRHI/Private/VulkanPipeline.h
+++ b/Engine/Source/Runtime/VulkanRHI/Private/VulkanPipeline.h
@@ -681,14 +681,14 @@ public:
 		return ResourceCast(GetComputeShader())->GetCodeHeader();
 	}
 
-	inline void Bind(VkCommandBuffer CmdBuffer)
+	inline const FVulkanLayout& GetLayout() const
 	{
-		VulkanRHI::vkCmdBindPipeline(CmdBuffer, VK_PIPELINE_BIND_POINT_COMPUTE, Pipeline);
+		return *Layout;
 	}
 
-	inline const FVulkanComputeLayout& GetComputeLayout() const
+	inline void Bind(VkCommandBuffer CmdBuffer)
 	{
-		return *(FVulkanComputeLayout*)Layout;
+		VulkanRHI::vkCmdBindPipeline(CmdBuffer, VK_PIPELINE_BIND_POINT_COMPUTE, Pipeline);
 	}
 };
 
@@ -708,11 +708,6 @@ public:
 		return *Layout;
 	}
 
-	inline const FVulkanGfxLayout& GetGfxLayout() const
-	{
-		return *(FVulkanGfxLayout*)&GetLayout();
-	}
-
 	inline void Bind(VkCommandBuffer CmdBuffer)
 	{
 		VulkanRHI::vkCmdBindPipeline(CmdBuffer, VK_PIPELINE_BIND_POINT_GRAPHICS, VulkanPipeline);
diff --git a/Engine/Source/Runtime/VulkanRHI/Private/VulkanPipelineState.cpp b/Engine/Source/Runtime/VulkanRHI/Private/VulkanPipelineState.cpp
--- a/Engine/Source/Runtime/VulkanRHI/Private/VulkanPipelineState.cpp
+++ b/Engine/Source/Runtime/VulkanRHI/Private/VulkanPipelineState.cpp
@@ -289,50 +289,45 @@ FVulkanGraphicsPipelineDescriptorState::FVulkanGraphicsPipelineDescriptorState(F
 	FMemory::Memzero(PackedUniformBuffersMask);
 	FMemory::Memzero(PackedUniformBuffersDirty);
 
-	check(InGfxPipeline);
-	{
-		
-		check(InGfxPipeline->Layout);
-		DescriptorSetsLayout = &InGfxPipeline->Layout->GetDescriptorSetsLayout();
-		FVulkanGfxLayout& GfxLayout  = *(FVulkanGfxLayout*)InGfxPipeline->Layout;
+	check(InGfxPipeline && InGfxPipeline->Layout && InGfxPipeline->Layout->IsGfxLayout());
+	DescriptorSetsLayout = &InGfxPipeline->Layout->GetDescriptorSetsLayout();
 
-		UsedSetsMask = 0;
+	UsedSetsMask = 0;
 
-		const FVulkanShaderFactory& ShaderFactory = Device->GetShaderFactory();
+	const FVulkanShaderFactory& ShaderFactory = Device->GetShaderFactory();
 
-		const FVulkanVertexShader* VertexShader = ShaderFactory.LookupShader<FVulkanVertexShader>(InGfxPipeline->GetShaderKey(SF_Vertex));
-		check(VertexShader);
-		PackedUniformBuffers[ShaderStage::Vertex].Init(VertexShader->GetCodeHeader(), PackedUniformBuffersMask[ShaderStage::Vertex]);
-		UsedSetsMask |= VertexShader->GetCodeHeader().Bindings.Num() ? (1u << ShaderStage::Vertex) : 0u;
+	const FVulkanVertexShader* VertexShader = ShaderFactory.LookupShader<FVulkanVertexShader>(InGfxPipeline->GetShaderKey(SF_Vertex));
+	check(VertexShader);
+	PackedUniformBuffers[ShaderStage::Vertex].Init(VertexShader->GetCodeHeader(), PackedUniformBuffersMask[ShaderStage::Vertex]);
+	UsedSetsMask |= VertexShader->GetCodeHeader().Bindings.Num() ? (1u << ShaderStage::Vertex) : 0u;
 
-		uint64 PixelShaderKey = InGfxPipeline->GetShaderKey(SF_Pixel);
-		if (PixelShaderKey)
-		{
-			const FVulkanPixelShader* PixelShader = ShaderFactory.LookupShader<FVulkanPixelShader>(PixelShaderKey);
-			check(PixelShader);
+	uint64 PixelShaderKey = InGfxPipeline->GetShaderKey(SF_Pixel);
+	if (PixelShaderKey)
+	{
+		const FVulkanPixelShader* PixelShader = ShaderFactory.LookupShader<FVulkanPixelShader>(PixelShaderKey);
+		check(PixelShader);
 
-			PackedUniformBuffers[ShaderStage::Pixel].Init(PixelShader->GetCodeHeader(), PackedUniformBuffersMask[ShaderStage::Pixel]);
-			UsedSetsMask |= PixelShader->GetCodeHeader().Bindings.Num() ? (1u << ShaderStage::Pixel) : 0u;
-		}
+		PackedUniformBuffers[ShaderStage::Pixel].Init(PixelShader->GetCodeHeader(), PackedUniformBuffersMask[ShaderStage::Pixel]);
+		UsedSetsMask |= PixelShader->GetCodeHeader().Bindings.Num() ? (1u << ShaderStage::Pixel) : 0u;
+	}
 
 #if VULKAN_SUPPORTS_GEOMETRY_SHADERS
-		uint64 GeometryShaderKey = InGfxPipeline->GetShaderKey(SF_Geometry);
-		if (GeometryShaderKey)
-		{
-			const FVulkanGeometryShader* GeometryShader = ShaderFactory.LookupShader<FVulkanGeometryShader>(GeometryShaderKey);
-			check(GeometryShader);
+	uint64 GeometryShaderKey = InGfxPipeline->GetShaderKey(SF_Geometry);
+	if (GeometryShaderKey)
+	{
+		const FVulkanGeometryShader* GeometryShader = ShaderFactory.LookupShader<FVulkanGeometryShader>(GeometryShaderKey);
+		check(GeometryShader);
 
-			PackedUniformBuffers[ShaderStage::Geometry].Init(GeometryShader->GetCodeHeader(), PackedUniformBuffersMask[ShaderStage::Geometry]);
-			UsedSetsMask |= GeometryShader->GetCodeHeader().Bindings.Num() ? (1u << ShaderStage::Geometry) : 0u;
-		}
+		PackedUniformBuffers[ShaderStage::Geometry].Init(GeometryShader->GetCodeHeader(), PackedUniformBuffersMask[ShaderStage::Geometry]);
+		UsedSetsMask |= GeometryShader->GetCodeHeader().Bindings.Num() ? (1u << ShaderStage::Geometry) : 0u;
+	}
 #endif
 
-		CreateDescriptorWriteInfos();
+	CreateDescriptorWriteInfos();
 
-		//UE_LOG(LogVulkanRHI, Warning, TEXT("GfxPSOState %p For PSO %p Writes:%d"), this, InGfxPipeline, DSWriteContainer.DescriptorWrites.Num());
+	//UE_LOG(LogVulkanRHI, Warning, TEXT("GfxPSOState %p For PSO %p Writes:%d"), this, InGfxPipeline, DSWriteContainer.DescriptorWrites.Num());
 
-		InGfxPipeline->AddRef();
-	}
+	InGfxPipeline->AddRef();
 }
 
 template<bool bUseDynamicGlobalUBs>
diff --git a/Engine/Source/Runtime/VulkanRHI/Private/VulkanShaders.cpp b/Engine/Source/Runtime/VulkanRHI/Private/VulkanShaders.cpp
--- a/Engine/Source/Runtime/VulkanRHI/Private/VulkanShaders.cpp
+++ b/Engine/Source/Runtime/VulkanRHI/Private/VulkanShaders.cpp
@@ -597,8 +597,9 @@ FRayTracingShaderRHIRef FVulkanDynamicRHI::RHICreateRayTracingShader(TArrayView<
 	}
 }
 
-FVulkanLayout::FVulkanLayout(FVulkanDevice* InDevice)
+FVulkanLayout::FVulkanLayout(FVulkanDevice* InDevice, bool InGfxLayout)
 	: VulkanRHI::FDeviceChild(InDevice)
+	, bIsGfxLayout(InGfxLayout)
 	, DescriptorSetLayout(Device)
 	, PipelineLayout(VK_NULL_HANDLE)
 {
