Enabling DynamicResolution on Android devices.
Fixing GPUTime reports using Swappy.
Setting the GPU time Fence value through a CVar, defaulted to 200ms.

diff --git a/Engine/Source/Runtime/ApplicationCore/Private/Android/AndroidApplication.cpp b/Engine/Source/Runtime/ApplicationCore/Private/Android/AndroidApplication.cpp
--- a/Engine/Source/Runtime/ApplicationCore/Private/Android/AndroidApplication.cpp
+++ b/Engine/Source/Runtime/ApplicationCore/Private/Android/AndroidApplication.cpp
@@ -16,6 +16,23 @@ bool FAndroidApplication::bWindowSizeChanged = false;
 
 FAndroidApplication* FAndroidApplication::_application = nullptr;
 
+TAutoConsoleVariable<bool> CVarAndroidSupportsTimestampQueries(
+	TEXT("r.Android.SupportsTimestampQueries"),
+	0,
+	TEXT("State of Android (GLES and Vulkan) timestamp queries support on an Android device\n")
+	TEXT("  0 = unsupported\n")
+	TEXT("  1 = supported."),
+	ECVF_SetByDeviceProfile);
+
+TAutoConsoleVariable<bool> CVarAndroidSupportsDynamicResolution(
+	TEXT("r.Android.SupportsDynamicResolution"),
+	0,
+	TEXT("State of DynamicResolution (GLES and Vulkan) support on an Android device\n")
+	TEXT("  0 = unsupported\n")
+	TEXT("  1 = supported."),
+	ECVF_SetByDeviceProfile
+);
+
 FAndroidApplication* FAndroidApplication::CreateAndroidApplication()
 {
 	return new FAndroidApplication();
diff --git a/Engine/Source/Runtime/ApplicationCore/Private/Android/AndroidPlatformFramePacer.cpp b/Engine/Source/Runtime/ApplicationCore/Private/Android/AndroidPlatformFramePacer.cpp
--- a/Engine/Source/Runtime/ApplicationCore/Private/Android/AndroidPlatformFramePacer.cpp
+++ b/Engine/Source/Runtime/ApplicationCore/Private/Android/AndroidPlatformFramePacer.cpp
@@ -65,6 +65,12 @@ TAutoConsoleVariable<int32> FAndroidPlatformRHIFramePacer::CVarDisableOpenGLGPUS
 	TEXT("When true, android OpenGL will not prevent the GPU from running more than one frame behind. This will allow higher performance on some devices but increase input latency."),
 	ECVF_RenderThreadSafe);
 
+TAutoConsoleVariable<int32> FAndroidPlatformRHIFramePacer::CVarSwappyGPUFrameTimeFence(
+	TEXT("r.Android.SwappyGPUFrameTimeFence"),
+	200,
+	TEXT("Android Swappy GPU Frame Time Fence in milliseconds. Defaulted to 200 ms"),
+	ECVF_Default);
+
 IAndroidFramePacer* FAndroidPlatformRHIFramePacer::FramePacer = nullptr;
 int32 FAndroidPlatformRHIFramePacer::InternalFramePace = 0;
 
diff --git a/Engine/Source/Runtime/ApplicationCore/Public/Android/AndroidApplication.h b/Engine/Source/Runtime/ApplicationCore/Public/Android/AndroidApplication.h
--- a/Engine/Source/Runtime/ApplicationCore/Public/Android/AndroidApplication.h
+++ b/Engine/Source/Runtime/ApplicationCore/Public/Android/AndroidApplication.h
@@ -104,6 +104,9 @@ public:
 
 	virtual bool IsGamepadAttached() const override;
 
+	static TAutoConsoleVariable<bool> CVarAndroidSupportsTimestampQueries;
+	static TAutoConsoleVariable<bool> CVarAndroidSupportsDynamicResolution;
+
 protected:
 
 	FAndroidApplication();
diff --git a/Engine/Source/Runtime/ApplicationCore/Public/Android/AndroidPlatformFramePacer.h b/Engine/Source/Runtime/ApplicationCore/Public/Android/AndroidPlatformFramePacer.h
--- a/Engine/Source/Runtime/ApplicationCore/Public/Android/AndroidPlatformFramePacer.h
+++ b/Engine/Source/Runtime/ApplicationCore/Public/Android/AndroidPlatformFramePacer.h
@@ -73,6 +73,7 @@ struct FAndroidPlatformRHIFramePacer : public FGenericPlatformRHIFramePacer
 	static void SwapBuffers(bool bLockToVsync) { if (ensure(FramePacer)) { FramePacer->SwapBuffers(bLockToVsync); } }
 
 	static TAutoConsoleVariable<int32> CVarUseSwappyForFramePacing;
+	static TAutoConsoleVariable<int32> CVarSwappyGPUFrameTimeFence;
 	static TAutoConsoleVariable<int32> CVarSupportNonVSyncMultipleFrameRates;
 
 	// Legacy pacer stuff
diff --git a/Engine/Source/Runtime/OpenGLDrv/Private/Android/AndroidOpenGL.cpp b/Engine/Source/Runtime/OpenGLDrv/Private/Android/AndroidOpenGL.cpp
--- a/Engine/Source/Runtime/OpenGLDrv/Private/Android/AndroidOpenGL.cpp
+++ b/Engine/Source/Runtime/OpenGLDrv/Private/Android/AndroidOpenGL.cpp
@@ -192,6 +192,24 @@ void FPlatformOpenGLDevice::Init()
 	// AsyncPipelinePrecompile can be enabled on android GL, precompiles are compiled via separate processes and the result is stored in GL's LRU cache as an evicted binary.
 	// The lru cache is a requirement as the precompile produces binary program data only.
 	GRHISupportsAsyncPipelinePrecompile = AreAndroidOpenGLRemoteCompileServicesAvailable();
+
+#if USE_ANDROID_OPENGL_SWAPPY
+	bool bIsSwappyEnabled = FAndroidPlatformRHIFramePacer::CVarUseSwappyForFramePacing.GetValueOnAnyThread() == 1;
+	
+	// don't even initialize this if swappy is not enabled
+	if (bIsSwappyEnabled)
+	{
+		IConsoleVariable* CVarAndroidSupportsTimestampQueries = IConsoleManager::Get().FindConsoleVariable(TEXT("r.Android.SupportsTimestampQueries"));
+		IConsoleVariable* CVarAndroidSupportsDynamicResolution = IConsoleManager::Get().FindConsoleVariable(TEXT("r.Android.SupportsDynamicResolution"));
+
+		bool bSupportsTimestampQueries = CVarAndroidSupportsTimestampQueries != nullptr && CVarAndroidSupportsTimestampQueries->GetBool();
+		bool bSupportsDynamicResolution = CVarAndroidSupportsDynamicResolution != nullptr && CVarAndroidSupportsDynamicResolution->GetBool();
+
+		GRHISupportsDynamicResolution = bSupportsDynamicResolution;
+		GSupportsTimestampRenderQueries = bSupportsTimestampQueries;
+		GRHISupportsGPUTimestampBubblesRemoval = true;
+	}
+#endif
 }
 
 FPlatformOpenGLDevice* PlatformCreateOpenGLDevice()
diff --git a/Engine/Source/Runtime/OpenGLDrv/Private/Android/AndroidOpenGLFramePacer.cpp b/Engine/Source/Runtime/OpenGLDrv/Private/Android/AndroidOpenGLFramePacer.cpp
--- a/Engine/Source/Runtime/OpenGLDrv/Private/Android/AndroidOpenGLFramePacer.cpp
+++ b/Engine/Source/Runtime/OpenGLDrv/Private/Android/AndroidOpenGLFramePacer.cpp
@@ -94,11 +94,11 @@ namespace AndroidGL
 {
 	void SwappyPostWaitCallback(void*, int64_t cpu_time_ns, int64_t gpu_time_ns)
 	{
-		const double Frequency = 1.0;// FGPUTiming::GetTimingFrequency();
-		const double CyclesPerSecond = 1.0 / (Frequency * FPlatformTime::GetSecondsPerCycle64());
 		const double GPUTimeInSeconds = (double)gpu_time_ns / 1000000000.0;
 
-		GetDynamicRHI<FOpenGLDynamicRHI>()->RHISetExternalGPUTime(CyclesPerSecond * GPUTimeInSeconds);
+		GGPUFrameTime = GPUTimeInSeconds / FPlatformTime::GetSecondsPerCycle64();
+
+		GetDynamicRHI<FOpenGLDynamicRHI>()->RHISetExternalGPUTime(GGPUFrameTime);
 	}
 
 	void SetSwappyPostWaitCallback()
@@ -106,6 +106,10 @@ namespace AndroidGL
 		SwappyTracer Tracer = { 0 };
 		Tracer.postWait = AndroidGL::SwappyPostWaitCallback;
 		SwappyGL_injectTracer(&Tracer);
+
+		int32 FrameTimeFenceInMillis = FAndroidPlatformRHIFramePacer::CVarSwappyGPUFrameTimeFence.GetValueOnAnyThread();
+
+		SwappyGL_setFenceTimeoutNS(FrameTimeFenceInMillis * 1000000); // millis to ns (ms * 1000000)
 	}
 };
 
diff --git a/Engine/Source/Runtime/VulkanRHI/Private/Android/VulkanAndroidPlatform.cpp b/Engine/Source/Runtime/VulkanRHI/Private/Android/VulkanAndroidPlatform.cpp
--- a/Engine/Source/Runtime/VulkanRHI/Private/Android/VulkanAndroidPlatform.cpp
+++ b/Engine/Source/Runtime/VulkanRHI/Private/Android/VulkanAndroidPlatform.cpp
@@ -31,11 +31,9 @@ namespace AndroidVulkan
 {
 	void VKSwappyPostWaitCallback(void*, int64_t cpu_time_ns, int64_t gpu_time_ns)
 	{
-		const double Frequency = 1.0;// FGPUTiming::GetTimingFrequency();
-		const double CyclesPerSecond = 1.0 / (Frequency * FPlatformTime::GetSecondsPerCycle64());
 		const double GPUTimeInSeconds = (double)gpu_time_ns / 1000000000.0;
 
-		GGPUFrameTime = CyclesPerSecond * GPUTimeInSeconds;
+		GGPUFrameTime = GPUTimeInSeconds / FPlatformTime::GetSecondsPerCycle64();
 	}
 
 	void SetSwappyPostWaitCallback()
@@ -43,10 +41,14 @@ namespace AndroidVulkan
 		SwappyTracer Tracer = { 0 };
 		Tracer.postWait = VKSwappyPostWaitCallback;
 		SwappyVk_injectTracer(&Tracer);
+
+		int32 FrameTimeFenceInMillis = FAndroidPlatformRHIFramePacer::CVarSwappyGPUFrameTimeFence.GetValueOnAnyThread();
+
+		SwappyVk_setFenceTimeoutNS(FrameTimeFenceInMillis * 1000000); // millis to ns (ms * 1000000)
 	}
 };
 
-#endif
+#endif // #if USE_ANDROID_SWAPPY
 
 // From VulklanSwapChain.cpp
 extern int32 GVulkanCPURenderThreadFramePacer;
@@ -60,15 +62,6 @@ static FAutoConsoleVariableRef CVarVulkanExtensionFramePacer(
 	ECVF_RenderThreadSafe
 );
 
-static TAutoConsoleVariable<int32> CVarVulkanSupportsTimestampQueries(
-	TEXT("r.Vulkan.SupportsTimestampQueries"),
-	0,
-	TEXT("State of Vulkan timestamp queries support on an Android device\n")
-	TEXT("  0 = unsupported\n")
-	TEXT("  1 = supported."),
-	ECVF_SetByDeviceProfile
-);
-
 static TAutoConsoleVariable<int32> CVarVulkanSupportsBCTextureFormats(
 	TEXT("r.Vulkan.SupportsBCTextureFormats"),
 	0,
@@ -501,7 +494,7 @@ void FVulkanAndroidPlatform::GetDeviceExtensions(FVulkanDevice* Device, FVulkanD
 			OutExtensions.Add(MakeUnique<FVulkanDeviceExtension>(Device, SwappyRequiredExtension.GetData(), VULKAN_EXTENSION_ENABLED, VULKAN_EXTENSION_NOT_PROMOTED));
 		}
 	}
-	#endif
+#endif
 
 }
 
@@ -598,8 +591,32 @@ void FVulkanAndroidPlatform::NotifyFoundDeviceLayersAndExtensions(VkPhysicalDevi
 
 bool FVulkanAndroidPlatform::SupportsTimestampRenderQueries()
 {
+	IConsoleVariable* CVarAndroidSupportsTimestampQueries = IConsoleManager::Get().FindConsoleVariable(TEXT("r.Android.SupportsTimestampQueries"));
+
 	// standalone devices have newer drivers where timestamp render queries work.
-	return (CVarVulkanSupportsTimestampQueries.GetValueOnAnyThread() == 1);
+	return CVarAndroidSupportsTimestampQueries != nullptr &&
+		CVarAndroidSupportsTimestampQueries->GetBool();
+}
+
+bool FVulkanAndroidPlatform::SupportsDynamicResolution()
+{
+	// separating render timestamp queries from dynres availability
+
+#if USE_ANDROID_SWAPPY
+	
+	IConsoleVariable* CVarAndroidSupportsDynamicResolution = IConsoleManager::Get().FindConsoleVariable(TEXT("r.Android.SupportsDynamicResolution"));
+
+	bool bIsSwappyEnabled = FAndroidPlatformRHIFramePacer::CVarUseSwappyForFramePacing.GetValueOnAnyThread() == 1;
+	return bIsSwappyEnabled && 
+		CVarAndroidSupportsDynamicResolution != nullptr &&
+		CVarAndroidSupportsDynamicResolution->GetBool(); // is supported
+
+#else // USE_ANDROID_SWAPPY
+
+	// defaulted to the previous code
+	return SupportsTimestampRenderQueries();
+
+#endif
 }
 
 void FVulkanAndroidPlatform::OverridePlatformHandlers(bool bInit)
diff --git a/Engine/Source/Runtime/VulkanRHI/Private/Android/VulkanAndroidPlatform.h b/Engine/Source/Runtime/VulkanRHI/Private/Android/VulkanAndroidPlatform.h
--- a/Engine/Source/Runtime/VulkanRHI/Private/Android/VulkanAndroidPlatform.h
+++ b/Engine/Source/Runtime/VulkanRHI/Private/Android/VulkanAndroidPlatform.h
@@ -89,7 +89,7 @@ public:
 
 	static bool SupportsTimestampRenderQueries();
 
-	static bool SupportsDynamicResolution() { return SupportsTimestampRenderQueries(); }
+	static bool SupportsDynamicResolution();
 
 	static bool RequiresMobileRenderer()
 	{
diff --git a/Engine/Source/Runtime/VulkanRHI/Private/VulkanRHI.cpp b/Engine/Source/Runtime/VulkanRHI/Private/VulkanRHI.cpp
--- a/Engine/Source/Runtime/VulkanRHI/Private/VulkanRHI.cpp
+++ b/Engine/Source/Runtime/VulkanRHI/Private/VulkanRHI.cpp
@@ -921,6 +921,7 @@ void FVulkanDynamicRHI::InitInstance()
 		GSupportsRenderTargetFormat_PF_G8 = false;	// #todo-rco
 		GRHISupportsTextureStreaming = true;
 		GSupportsTimestampRenderQueries = FVulkanPlatform::SupportsTimestampRenderQueries();
+		GRHISupportsGPUTimestampBubblesRemoval = true;
 		GSupportsMobileMultiView = Device->GetOptionalExtensions().HasKHRMultiview ? true : false;
 		GRHISupportsMSAAShaderResolve = Device->GetOptionalExtensions().HasQcomRenderPassShaderResolve ? true : false;
 		GRHISupportsRayTracing = RHISupportsRayTracing(GMaxRHIShaderPlatform) && Device->GetOptionalExtensions().HasRaytracingExtensions();
