Fixed an assert in Vulkan editor when previewing mobile feature level. Related to incorrect setup of unfiorm buffers views.

diff --git a/Engine/Source/Runtime/VulkanRHI/Private/VulkanCommands.cpp b/Engine/Source/Runtime/VulkanRHI/Private/VulkanCommands.cpp
--- a/Engine/Source/Runtime/VulkanRHI/Private/VulkanCommands.cpp
+++ b/Engine/Source/Runtime/VulkanRHI/Private/VulkanCommands.cpp
@@ -504,6 +504,11 @@ void FVulkanCommandListContext::RHISetShaderUniformBuffer(FRHIGraphicsShader* Sh
 	checkfSlow(!CodeHeader.UniformBufferInfos[BufferIndex].LayoutHash || (CodeHeader.UniformBufferInfos[BufferIndex].LayoutHash == UniformBuffer->GetLayout().GetHash()),
 		TEXT("Mismatched UB layout! Got hash 0x%x, expected 0x%x!"), UniformBuffer->GetLayout().GetHash(), CodeHeader.UniformBufferInfos[BufferIndex].LayoutHash);
 
+	if (UniformBuffer->IsUniformView())
+	{
+		UniformBuffer->SetupUniformBufferView();
+	}
+
 	bool bHasResources = false;
 	if (BufferIndex < CodeHeader.NumBoundUniformBuffers)
 	{
diff --git a/Engine/Source/Runtime/VulkanRHI/Private/VulkanPipelineState.h b/Engine/Source/Runtime/VulkanRHI/Private/VulkanPipelineState.h
--- a/Engine/Source/Runtime/VulkanRHI/Private/VulkanPipelineState.h
+++ b/Engine/Source/Runtime/VulkanRHI/Private/VulkanPipelineState.h
@@ -97,7 +97,7 @@ public:
 	inline void SetUniformBuffer(uint8 DescriptorSet, uint32 BindingIndex, const FVulkanUniformBuffer* UniformBuffer)
 	{
 		const VulkanRHI::FVulkanAllocation& Allocation = UniformBuffer->Allocation;
-		VkDeviceSize Range = UniformBuffer->bUniformView ? PLATFORM_MAX_UNIFORM_BUFFER_RANGE : UniformBuffer->GetSize();
+		VkDeviceSize Range = UniformBuffer->IsUniformView() ? PLATFORM_MAX_UNIFORM_BUFFER_RANGE : UniformBuffer->GetSize();
 
 		if (bDynamic)
 		{
diff --git a/Engine/Source/Runtime/VulkanRHI/Private/VulkanResources.h b/Engine/Source/Runtime/VulkanRHI/Private/VulkanResources.h
--- a/Engine/Source/Runtime/VulkanRHI/Private/VulkanResources.h
+++ b/Engine/Source/Runtime/VulkanRHI/Private/VulkanResources.h
@@ -1189,12 +1189,16 @@ public:
 	{
 		NewAlloc.Swap(Allocation);
 	}
+
+	inline bool IsUniformView() const
+	{
+		return UniformViewSRV != nullptr;
+	}
 	
 	FRHIDescriptorHandle GetBindlessHandle();
 	VkDeviceAddress GetDeviceAddress() const;
 
-protected:
-	bool SetupUniformBufferView(const FRHIUniformBufferLayout* InLayout, const void* Contents);
+	void SetupUniformBufferView();
 
 public:
 	FVulkanDevice* Device;
@@ -1203,7 +1207,7 @@ public:
 
 	FRHIDescriptorHandle BindlessHandle;
 	VkDeviceAddress CachedDeviceAddress = 0;
-	bool bUniformView = false;
+	FRHIShaderResourceView* UniformViewSRV;
 };
 
 class FVulkanUnorderedAccessView final : public FRHIUnorderedAccessView, public FVulkanLinkedView
diff --git a/Engine/Source/Runtime/VulkanRHI/Private/VulkanUniformBuffer.cpp b/Engine/Source/Runtime/VulkanRHI/Private/VulkanUniformBuffer.cpp
--- a/Engine/Source/Runtime/VulkanRHI/Private/VulkanUniformBuffer.cpp
+++ b/Engine/Source/Runtime/VulkanRHI/Private/VulkanUniformBuffer.cpp
@@ -101,48 +101,24 @@ static void UpdateUniformBufferHelper(FVulkanCommandListContext& Context, FVulka
 	}
 };
 
-bool FVulkanUniformBuffer::SetupUniformBufferView(const FRHIUniformBufferLayout* InLayout, const void* Contents)
+void FVulkanUniformBuffer::SetupUniformBufferView()
 {
-	bUniformView = false;
-
-	if (InLayout->bUniformView)
+	if (UniformViewSRV && GetBufferHandle() == VK_NULL_HANDLE)
 	{
-		FRHIShaderResourceView* UniformViewSRV = nullptr;
-		for (int32 Index = 0; Index < InLayout->Resources.Num() && !UniformViewSRV; ++Index)
-		{
-			EUniformBufferBaseType ResourceBaseType = InLayout->Resources[Index].MemberType;
-			if (ResourceBaseType == UBMT_SRV || 
-				ResourceBaseType == UBMT_RDG_BUFFER_SRV)
-			{
-				UniformViewSRV = (FRHIShaderResourceView*)GetShaderParameterResourceRHI(Contents, InLayout->Resources[Index].MemberOffset, ResourceBaseType);
-			}
-		}
-
-		check(UniformViewSRV)
-
-		FRHICommandListImmediate& RHICmdList = FRHICommandListExecutor::GetImmediateCommandList();
-		RHICmdList.EnqueueLambda([this, UniformViewSRV](FRHICommandList& CmdList)
-		{
-			FVulkanResourceMultiBuffer* Buffer = ResourceCast(UniformViewSRV->GetBuffer());
-			const FRHIViewDesc::FBufferSRV& SRVInfo = UniformViewSRV->GetDesc().Buffer.SRV;
-			Allocation.Reference(Buffer->GetCurrentAllocation());
-			check(Allocation.Size >= PLATFORM_MAX_UNIFORM_BUFFER_RANGE);
-			//Adjust Allocation.Size ???
-			Allocation.Offset += SRVInfo.OffsetInBytes;
-			bUniformView = true;
-		});
-		
-		return true;
+		const FRHIViewDesc::FBufferSRV& SRVInfo = UniformViewSRV->GetDesc().Buffer.SRV;
+		FVulkanResourceMultiBuffer* Buffer = ResourceCast(UniformViewSRV->GetBuffer());
+		Allocation.Reference(Buffer->GetCurrentAllocation());
+		check(Allocation.Size >= PLATFORM_MAX_UNIFORM_BUFFER_RANGE);
+		//Adjust Allocation.Size ???
+		Allocation.Offset += SRVInfo.OffsetInBytes;
 	}
-	
-	return false;
 }
 
 FVulkanUniformBuffer::FVulkanUniformBuffer(FVulkanDevice& InDevice, const FRHIUniformBufferLayout* InLayout, const void* Contents, EUniformBufferUsage InUsage, EUniformBufferValidation Validation)
 	: FRHIUniformBuffer(InLayout)
 	, Device(&InDevice)
 	, Usage(InUsage)
-	, bUniformView(false)
+	, UniformViewSRV(nullptr)
 {
 #if VULKAN_ENABLE_AGGRESSIVE_STATS
 	SCOPE_CYCLE_COUNTER(STAT_VulkanUniformBufferCreateTime);
@@ -170,8 +146,16 @@ FVulkanUniformBuffer::FVulkanUniformBuffer(FVulkanDevice& InDevice, const FRHIUn
 		}
 	}
 
-	if (SetupUniformBufferView(InLayout, Contents))
+	if (InLayout->bUniformView)
 	{
+		// For uniform view we expect an buffer SRV as a first resource
+		check(InLayout->Resources.Num() > 0);
+		EUniformBufferBaseType ResourceBaseType = InLayout->Resources[0].MemberType;
+		if (ResourceBaseType == UBMT_SRV || ResourceBaseType == UBMT_RDG_BUFFER_SRV)
+		{
+			UniformViewSRV = (FRHIShaderResourceView*)GetShaderParameterResourceRHI(Contents, InLayout->Resources[0].MemberOffset, ResourceBaseType);
+		}
+		check(UniformViewSRV)
 		return;
 	}
 
