VulkanRHI: Also reflect members from non-global ub

#rnx

diff --git a/Engine/Source/Developer/VulkanShaderFormat/Private/VulkanShaderCompiler.cpp b/Engine/Source/Developer/VulkanShaderFormat/Private/VulkanShaderCompiler.cpp
--- a/Engine/Source/Developer/VulkanShaderFormat/Private/VulkanShaderCompiler.cpp
+++ b/Engine/Source/Developer/VulkanShaderFormat/Private/VulkanShaderCompiler.cpp
@@ -821,7 +821,7 @@ static bool BuildShaderOutputFromSpirv(
 							{
 								const SpvReflectBlockVariable& Member = Binding->block.members[MemberIndex];
 
-								FString MemberName(ANSI_TO_TCHAR(Member.name));
+								FString MemberName(Member.name);
 								FStringView AdjustedMemberName(MemberName);
 
 								const EShaderParameterType BindlessParameterType = FShaderParameterParser::ParseAndRemoveBindlessParameterPrefix(AdjustedMemberName);
@@ -849,6 +849,26 @@ static bool BuildShaderOutputFromSpirv(
 							check(!UsedUniformBufferSlots[ReflectionSlot]);
 							HandleReflectedUniformBuffer(ResourceName, ReflectionSlot, Output);
 							AddShaderValidationUBSize(BindingTypeCount, Binding->block.padded_size, Output);
+
+							// Register uniform buffer members that are in use
+							for (uint32 MemberIndex = 0; MemberIndex < Binding->block.member_count; ++MemberIndex)
+							{
+								const SpvReflectBlockVariable& Member = Binding->block.members[MemberIndex];
+
+								if ((Member.flags & SPV_REFLECT_VARIABLE_FLAGS_UNUSED) != 0)
+								{
+									continue;
+								}
+
+								const FString MemberName(Member.name);
+								HandleReflectedUniformBufferConstantBufferMember(
+									ReflectionSlot,
+									MemberName,
+									Member.absolute_offset,
+									Member.size,
+									Output
+								);
+							}
 						}
 
 						check(!UsedUniformBufferSlots[ReflectionSlot]);
diff --git a/Engine/Source/Runtime/VulkanRHI/Private/VulkanCommands.cpp b/Engine/Source/Runtime/VulkanRHI/Private/VulkanCommands.cpp
--- a/Engine/Source/Runtime/VulkanRHI/Private/VulkanCommands.cpp
+++ b/Engine/Source/Runtime/VulkanRHI/Private/VulkanCommands.cpp
@@ -128,11 +128,15 @@ struct FVulkanResourceBinder
 	const ShaderStage::EStage Stage;
 	PendingStateType* PendingState;
 
+	// Even if no resource are bound when bindless is enabled, we still need to process the Resource Table to go through proper validation
+	const bool bBindless;
+
 	FVulkanResourceBinder(FVulkanCommandListContext& InContext, EShaderFrequency InFrequency, PendingStateType* InPendingState)
 		: Context(InContext)
 		, Frequency(InFrequency)
 		, Stage((InFrequency == SF_Compute) ? ShaderStage::Compute : ShaderStage::GetStageForFrequency(InFrequency))
 		, PendingState(InPendingState)
+		, bBindless(InContext.GetDevice()->SupportsBindless())
 	{
 	}
 
@@ -143,25 +147,37 @@ struct FVulkanResourceBinder
 			//Context.ClearShaderResources(UAV);
 		}
 
-		PendingState->SetUAVForUBResource(Stage, Index, ResourceCast(UAV));
+		if (!bBindless)
+		{
+			PendingState->SetUAVForUBResource(Stage, Index, ResourceCast(UAV));
+		}
 	}
 
 	void SetSRV(FRHIShaderResourceView* SRV, uint16 Index)
 	{
-		PendingState->SetSRVForUBResource(Stage, Index, ResourceCast(SRV));
+		if (!bBindless)
+		{
+			PendingState->SetSRVForUBResource(Stage, Index, ResourceCast(SRV));
+		}
 	}
 
 	void SetTexture(FRHITexture* TextureRHI, uint16 Index)
 	{
-		FVulkanTexture* VulkanTexture = ResourceCast(TextureRHI);
-		const ERHIAccess RHIAccess = (Frequency == SF_Compute) ? ERHIAccess::SRVCompute : ERHIAccess::SRVGraphics;
-		const VkImageLayout ExpectedLayout = FVulkanLayoutManager::GetDefaultLayout(Context.GetCommandBufferManager()->GetActiveCmdBuffer(), *VulkanTexture, RHIAccess);
-		PendingState->SetTextureForUBResource(Stage, Index, VulkanTexture, ExpectedLayout);
+		if (!bBindless)
+		{
+			FVulkanTexture* VulkanTexture = ResourceCast(TextureRHI);
+			const ERHIAccess RHIAccess = (Frequency == SF_Compute) ? ERHIAccess::SRVCompute : ERHIAccess::SRVGraphics;
+			const VkImageLayout ExpectedLayout = FVulkanLayoutManager::GetDefaultLayout(Context.GetCommandBufferManager()->GetActiveCmdBuffer(), *VulkanTexture, RHIAccess);
+			PendingState->SetTextureForUBResource(Stage, Index, VulkanTexture, ExpectedLayout);
+		}
 	}
 
 	void SetSampler(FRHISamplerState* Sampler, uint16 Index)
 	{
-		PendingState->SetSamplerStateForUBResource(Stage, Index, ResourceCast(Sampler));
+		if (!bBindless)
+		{
+			PendingState->SetSamplerStateForUBResource(Stage, Index, ResourceCast(Sampler));
+		}
 	}
 
 #if PLATFORM_SUPPORTS_BINDLESS_RENDERING
