VulkanRHI: Add missing const on HasAllocation() call

#rnx

diff --git a/Engine/Source/Runtime/VulkanRHI/Private/VulkanMemory.cpp b/Engine/Source/Runtime/VulkanRHI/Private/VulkanMemory.cpp
--- a/Engine/Source/Runtime/VulkanRHI/Private/VulkanMemory.cpp
+++ b/Engine/Source/Runtime/VulkanRHI/Private/VulkanMemory.cpp
@@ -3810,7 +3810,7 @@ namespace VulkanRHI
 		FMemory::Memcpy(*this, Other);
 		bHasOwnership = 0;
 	}
-	bool FVulkanAllocation::HasAllocation()
+	bool FVulkanAllocation::HasAllocation() const
 	{
 		return Type != EVulkanAllocationEmpty && bHasOwnership;
 	}
diff --git a/Engine/Source/Runtime/VulkanRHI/Private/VulkanMemory.h b/Engine/Source/Runtime/VulkanRHI/Private/VulkanMemory.h
--- a/Engine/Source/Runtime/VulkanRHI/Private/VulkanMemory.h
+++ b/Engine/Source/Runtime/VulkanRHI/Private/VulkanMemory.h
@@ -472,7 +472,7 @@ namespace VulkanRHI
 		void Free(FVulkanDevice& Device);
 		void Swap(FVulkanAllocation& Other);
 		void Reference(const FVulkanAllocation& Other); //point to other, but don't take ownership
-		bool HasAllocation();
+		bool HasAllocation() const;
 
 		void Disown(); //disown & own should be used if ownership is transferred.
 		void Own();
