Compile Fix

diff --git a/Engine/Source/Runtime/VulkanRHI/Private/VulkanBarriers.cpp b/Engine/Source/Runtime/VulkanRHI/Private/VulkanBarriers.cpp
--- a/Engine/Source/Runtime/VulkanRHI/Private/VulkanBarriers.cpp
+++ b/Engine/Source/Runtime/VulkanRHI/Private/VulkanBarriers.cpp
@@ -1976,7 +1976,6 @@ VkImageLayout FVulkanLayoutManager::GetDefaultLayout(FVulkanCmdBuffer* CmdBuffer
 	case ERHIAccess::SRVCompute:
 	case ERHIAccess::SRVGraphics:
 	case ERHIAccess::SRVGraphicsNonPixel:
-	case ERHIAccess::SRVGraphicsPixel:
 	case ERHIAccess::SRVMask:
 	{
 		if (VulkanTexture.IsDepthOrStencilAspect())
