VulkanRHI: Fix crash on startup when geometry shaders not enabled

#rnx

diff --git a/Engine/Source/Runtime/VulkanRHI/Private/VulkanPipeline.h b/Engine/Source/Runtime/VulkanRHI/Private/VulkanPipeline.h
--- a/Engine/Source/Runtime/VulkanRHI/Private/VulkanPipeline.h
+++ b/Engine/Source/Runtime/VulkanRHI/Private/VulkanPipeline.h
@@ -40,10 +40,8 @@ inline uint64 GetShaderKeyForGfxStage(const FBoundShaderStateInput& BSI, ShaderS
 		return GetShaderKey<FVulkanVertexShader>(BSI.VertexShaderRHI);
 	case ShaderStage::Pixel:
 		return GetShaderKey<FVulkanPixelShader>(BSI.PixelShaderRHI);
-#if VULKAN_SUPPORTS_GEOMETRY_SHADERS
 	case ShaderStage::Geometry:
 		return GetShaderKey<FVulkanGeometryShader>(BSI.GetGeometryShader());
-#endif
 	case ShaderStage::RayGen:
 	case ShaderStage::RayHitGroup:
 	case ShaderStage::RayMiss:
