VulkanRHI: Use Host memory for dynamic buffers (can be copied from host and locked during render passes).  Make Volatile buffers use the dynamic buffer code path for now.  Remove array of allocations and clean up top vs bottom of pipe code paths.  Do not set debug utils name on suballocated buffers.

diff --git a/Engine/Source/Runtime/VulkanRHI/Private/VulkanBuffer.cpp b/Engine/Source/Runtime/VulkanRHI/Private/VulkanBuffer.cpp
--- a/Engine/Source/Runtime/VulkanRHI/Private/VulkanBuffer.cpp
+++ b/Engine/Source/Runtime/VulkanRHI/Private/VulkanBuffer.cpp
@@ -13,8 +13,16 @@
 #include "VulkanTransientResourceAllocator.h"
 #include "RHICoreStats.h"
 
-static TMap<FVulkanResourceMultiBuffer*, VulkanRHI::FPendingBufferLock> GPendingLockIBs;
-static FCriticalSection GPendingLockIBsMutex;
+struct FVulkanPendingBufferLock
+{
+	VulkanRHI::FStagingBuffer* StagingBuffer = nullptr;
+	uint32 Offset = 0;
+	uint32 Size = 0;
+	EResourceLockMode LockMode = RLM_Num;
+};
+
+static TMap<FVulkanResourceMultiBuffer*, FVulkanPendingBufferLock> GPendingLocks;
+static FCriticalSection GPendingLockMutex;
 
 int32 GVulkanForceStagingBufferOnLock = 0;
 static FAutoConsoleVariableRef CVarVulkanForceStagingBufferOnLock(
@@ -25,18 +33,25 @@ static FAutoConsoleVariableRef CVarVulkanForceStagingBufferOnLock(
 	ECVF_RenderThreadSafe
 );
 
-static FORCEINLINE VulkanRHI::FPendingBufferLock GetPendingBufferLock(FVulkanResourceMultiBuffer* Buffer)
+static FORCEINLINE FVulkanPendingBufferLock GetPendingBufferLock(FVulkanResourceMultiBuffer* Buffer)
 {
-	VulkanRHI::FPendingBufferLock PendingLock;
+	FVulkanPendingBufferLock PendingLock;
 
 	// Found only if it was created for Write
-	FScopeLock ScopeLock(&GPendingLockIBsMutex);
-	const bool bFound = GPendingLockIBs.RemoveAndCopyValue(Buffer, PendingLock);
+	FScopeLock ScopeLock(&GPendingLockMutex);
+	const bool bFound = GPendingLocks.RemoveAndCopyValue(Buffer, PendingLock);
 
 	checkf(bFound, TEXT("Mismatched Buffer Lock/Unlock!"));
 	return PendingLock;
 }
 
+static FORCEINLINE void AddPendingBufferLock(FVulkanResourceMultiBuffer* Buffer, FVulkanPendingBufferLock& PendingLock)
+{
+	FScopeLock ScopeLock(&GPendingLockMutex);
+	check(!GPendingLocks.Contains(Buffer));
+	GPendingLocks.Add(Buffer, PendingLock);
+}
+
 static void UpdateVulkanBufferStats(const FRHIBufferDesc& BufferDesc, int64 BufferSize, bool bAllocating)
 {
 	UE::RHICore::UpdateGlobalBufferStats(BufferDesc, BufferSize, bAllocating);
@@ -117,70 +132,41 @@ FVulkanResourceMultiBuffer::FVulkanResourceMultiBuffer(FVulkanDevice* InDevice,
 	{
 		check(InDevice);
 
-		const bool bVolatile = EnumHasAnyFlags(InBufferDesc.Usage, BUF_Volatile);
-		if (bVolatile)
-		{
-			check(InRHICmdList);
-
-			// Volatile buffers always work out of the same first slot
-			CurrentBufferIndex = BufferAllocs.Emplace();
-
-			// Get a dummy buffer as sometimes the high-level misbehaves and tries to use SRVs off volatile buffers before filling them in...
-			void* Data = Lock(*InRHICmdList, RLM_WriteOnly, InBufferDesc.Size, 0);
-
-			if (CreateInfo.ResourceArray)
-			{
-				uint32 CopyDataSize = FMath::Min(InBufferDesc.Size, CreateInfo.ResourceArray->GetResourceDataSize());
-				FMemory::Memcpy(Data, CreateInfo.ResourceArray->GetResourceData(), CopyDataSize);
-			}
-			else
-			{
-				FMemory::Memzero(Data, InBufferDesc.Size);
-			}
+		const bool bUnifiedMem = InDevice->HasUnifiedMemory();
+		const uint32 BufferAlignment = VulkanRHI::FMemoryManager::CalculateBufferAlignment(*InDevice, InBufferDesc.Usage, bZeroSize);
 
-			Unlock(*InRHICmdList);
+		if (InTransientHeapAllocation != nullptr)
+		{
+			CurrentBufferAlloc.Alloc = FVulkanTransientHeap::GetVulkanAllocation(*InTransientHeapAllocation);
+			CurrentBufferAlloc.HostPtr = bUnifiedMem ? CurrentBufferAlloc.Alloc.GetMappedPointer(Device) : nullptr;
+			CurrentBufferAlloc.DeviceAddress = GetBufferDeviceAddress(InDevice, CurrentBufferAlloc.Alloc.GetBufferHandle()) + CurrentBufferAlloc.Alloc.Offset;
+			check(CurrentBufferAlloc.Alloc.Offset % BufferAlignment == 0);
+			check(CurrentBufferAlloc.Alloc.Size >= InBufferDesc.Size);
 		}
 		else
 		{
-			const bool bUnifiedMem = InDevice->HasUnifiedMemory();
-			const uint32 BufferAlignment = VulkanRHI::FMemoryManager::CalculateBufferAlignment(*InDevice, InBufferDesc.Usage, bZeroSize);
+			AllocateMemory(CurrentBufferAlloc);
+		}
+
+		if (CreateInfo.ResourceArray)
+		{
+			const uint32 CopyDataSize = FMath::Min(InBufferDesc.Size, CreateInfo.ResourceArray->GetResourceDataSize());
 
-			if (InTransientHeapAllocation != nullptr)
+			// We know this buffer is not in use by GPU atm. If we do have a direct access initialize it without extra copies
+			if (CurrentBufferAlloc.HostPtr)
 			{
-				FBufferAlloc NewBufferAlloc;
-				NewBufferAlloc.Alloc = FVulkanTransientHeap::GetVulkanAllocation(*InTransientHeapAllocation);
-				NewBufferAlloc.HostPtr = bUnifiedMem ? NewBufferAlloc.Alloc.GetMappedPointer(Device) : nullptr;
-				NewBufferAlloc.DeviceAddress = GetBufferDeviceAddress(InDevice, NewBufferAlloc.Alloc.GetBufferHandle()) + NewBufferAlloc.Alloc.Offset;
-				check(NewBufferAlloc.Alloc.Offset % BufferAlignment == 0);
-				check(NewBufferAlloc.Alloc.Size >= InBufferDesc.Size);
-				CurrentBufferIndex = BufferAllocs.Add(NewBufferAlloc);
+				FMemory::Memcpy(CurrentBufferAlloc.HostPtr, CreateInfo.ResourceArray->GetResourceData(), CopyDataSize);
+				++LockCounter;
 			}
 			else
-			{
-				AdvanceBufferIndex();
-			}
-
-			VULKAN_SET_DEBUG_NAME((*InDevice), VK_OBJECT_TYPE_BUFFER, BufferAllocs[CurrentBufferIndex].Alloc.GetBufferHandle(), TEXT("%s"), CreateInfo.DebugName ? CreateInfo.DebugName : TEXT("UnknownBuffer"));
-
-			if (CreateInfo.ResourceArray)
 			{
 				check(InRHICmdList);
-
-				uint32 CopyDataSize = FMath::Min(InBufferDesc.Size, CreateInfo.ResourceArray->GetResourceDataSize());
-				// We know this buffer is not in use by GPU atm. If we do have a direct access initialize it without extra copies
-				if (bUnifiedMem)
-				{
-					FMemory::Memcpy(BufferAllocs[CurrentBufferIndex].HostPtr, CreateInfo.ResourceArray->GetResourceData(), CopyDataSize);
-				}
-				else
-				{
-					void* Data = Lock(*InRHICmdList, RLM_WriteOnly, CopyDataSize, 0);
-					FMemory::Memcpy(Data, CreateInfo.ResourceArray->GetResourceData(), CopyDataSize);
-					Unlock(*InRHICmdList);
-				}
-
-				CreateInfo.ResourceArray->Discard();
+				void* Data = Lock(*InRHICmdList, RLM_WriteOnly, CopyDataSize, 0);
+				FMemory::Memcpy(Data, CreateInfo.ResourceArray->GetResourceData(), CopyDataSize);
+				Unlock(*InRHICmdList);
 			}
+
+			CreateInfo.ResourceArray->Discard();
 		}
 	}
 }
@@ -192,243 +178,160 @@ FVulkanResourceMultiBuffer::~FVulkanResourceMultiBuffer()
 }
 
 
-void FVulkanResourceMultiBuffer::AdvanceBufferIndex()
+void FVulkanResourceMultiBuffer::AllocateMemory(FBufferAlloc& OutAlloc)
 {
-	auto FreePreviousAlloc = [&]() {
-		if (CurrentBufferIndex >= 0)
-		{
-			FBufferAlloc& CurrentBufferAlloc = BufferAllocs[CurrentBufferIndex];
-			CurrentBufferAlloc.AllocStatus = FBufferAlloc::EAllocStatus::NeedsFence;
-			CurrentBufferAlloc.Fence->Clear();
-		}
-	};
-
-	// Try to see if one of the buffers in our pool can be reused
-	if (BufferAllocs.Num() > 1)
+	VkMemoryPropertyFlags BufferMemFlags = 0;
+	const bool bUnifiedMem = Device->HasUnifiedMemory();
+	const bool bDynamic = EnumHasAnyFlags(GetUsage(), BUF_Dynamic) || EnumHasAnyFlags(GetUsage(), BUF_Volatile);
+	if (bUnifiedMem)
 	{
-		for (int32 BufferIndex = 0; BufferIndex < BufferAllocs.Num(); ++BufferIndex)
-		{
-			if (CurrentBufferIndex == BufferIndex)
-			{
-				continue;
-			}
-
-			// Fences are only written on Unlock(), but are polled on lock/unlock
-			FBufferAlloc& BufferAlloc = BufferAllocs[BufferIndex];
-			if ((BufferAlloc.AllocStatus == FBufferAlloc::EAllocStatus::Pending) && BufferAlloc.Fence->Poll())
-			{
-				BufferAlloc.AllocStatus = FBufferAlloc::EAllocStatus::Available;
-				BufferAlloc.Fence->Clear();
-			}
-
-			if (BufferAlloc.AllocStatus == FBufferAlloc::EAllocStatus::Available)
-			{
-				FreePreviousAlloc();
-
-				CurrentBufferIndex = BufferIndex;
-				BufferAllocs[CurrentBufferIndex].AllocStatus = FBufferAlloc::EAllocStatus::InUse;
-				return;
-			}
-		}
+		BufferMemFlags = VK_MEMORY_PROPERTY_DEVICE_LOCAL_BIT | VK_MEMORY_PROPERTY_HOST_VISIBLE_BIT | VK_MEMORY_PROPERTY_HOST_COHERENT_BIT;
 	}
-
-	// Allocate a new buffer
+	else if (bDynamic)
 	{
-		FreePreviousAlloc();
-
-		const bool bUnifiedMem = Device->HasUnifiedMemory();
-		const VkMemoryPropertyFlags BufferMemFlags = VK_MEMORY_PROPERTY_DEVICE_LOCAL_BIT | (bUnifiedMem ? (VK_MEMORY_PROPERTY_HOST_VISIBLE_BIT | VK_MEMORY_PROPERTY_HOST_COHERENT_BIT) : 0);
-		const uint32 BufferSize = GetSize();
-		const uint32 BufferAlignment = VulkanRHI::FMemoryManager::CalculateBufferAlignment(*Device, GetUsage(), (BufferSize == 0));
-
-		CurrentBufferIndex = BufferAllocs.Emplace();
-		FBufferAlloc& NewBufferAlloc = BufferAllocs[CurrentBufferIndex];
-		if (!Device->GetMemoryManager().AllocateBufferPooled(NewBufferAlloc.Alloc, nullptr, BufferSize, BufferAlignment, BufferUsageFlags, BufferMemFlags, VulkanRHI::EVulkanAllocationMetaMultiBuffer, __FILE__, __LINE__))
-		{
-			Device->GetMemoryManager().HandleOOM();
-		}
-		NewBufferAlloc.HostPtr = bUnifiedMem ? NewBufferAlloc.Alloc.GetMappedPointer(Device) : nullptr;
-		NewBufferAlloc.Fence = new FVulkanGPUFence(TEXT("VulkanDynamicBuffer"));
-		NewBufferAlloc.AllocStatus = FBufferAlloc::EAllocStatus::InUse;
-		NewBufferAlloc.DeviceAddress = GetBufferDeviceAddress(Device, NewBufferAlloc.Alloc.GetBufferHandle()) + NewBufferAlloc.Alloc.Offset;
-
-		UpdateVulkanBufferStats(GetDesc(), BufferSize, true);
+		BufferMemFlags = VK_MEMORY_PROPERTY_HOST_VISIBLE_BIT | VK_MEMORY_PROPERTY_HOST_COHERENT_BIT;
 	}
-}
-
-void FVulkanResourceMultiBuffer::UpdateBufferAllocStates(FVulkanCommandListContext& Context)
-{
-	for (int32 BufferIndex = 0; BufferIndex < BufferAllocs.Num(); ++BufferIndex)
+	else
 	{
-		if (CurrentBufferIndex == BufferIndex)
-		{
-			continue;
-		}
-
-		FBufferAlloc& BufferAlloc = BufferAllocs[BufferIndex];
-		if (BufferAlloc.AllocStatus == FBufferAlloc::EAllocStatus::Pending)
-		{
-			if (BufferAlloc.Fence->Poll())
-			{
-				BufferAlloc.AllocStatus = FBufferAlloc::EAllocStatus::Available;
-				BufferAlloc.Fence->Clear();
-			}
-		}
-		else if (BufferAlloc.AllocStatus == FBufferAlloc::EAllocStatus::NeedsFence)
-		{
-			Context.RHIWriteGPUFence(BufferAlloc.Fence);
-			BufferAlloc.AllocStatus = FBufferAlloc::EAllocStatus::Pending;
-		}
+		BufferMemFlags = VK_MEMORY_PROPERTY_DEVICE_LOCAL_BIT;
 	}
-}
+	const uint32 BufferSize = GetSize();
+	const uint32 BufferAlignment = VulkanRHI::FMemoryManager::CalculateBufferAlignment(*Device, GetUsage(), (BufferSize == 0));
 
-void* FVulkanResourceMultiBuffer::Lock(FRHICommandListBase& RHICmdList, EResourceLockMode LockMode, uint32 LockSize, uint32 Offset)
-{
-	// Use the immediate context for write operations, since we are only accessing allocators.
-	FVulkanCommandListContext& Context = FVulkanCommandListContext::GetVulkanContext(LockMode == RLM_WriteOnly ? *RHIGetDefaultContext() : RHICmdList.GetContext());
+	FBufferAlloc& NewBufferAlloc = OutAlloc;
+	if (!Device->GetMemoryManager().AllocateBufferPooled(NewBufferAlloc.Alloc, nullptr, BufferSize, BufferAlignment, BufferUsageFlags, BufferMemFlags, VulkanRHI::EVulkanAllocationMetaMultiBuffer, __FILE__, __LINE__))
+	{
+		Device->GetMemoryManager().HandleOOM();
+	}
+	NewBufferAlloc.HostPtr = (bUnifiedMem || bDynamic) ? NewBufferAlloc.Alloc.GetMappedPointer(Device) : nullptr;
+	NewBufferAlloc.DeviceAddress = GetBufferDeviceAddress(Device, NewBufferAlloc.Alloc.GetBufferHandle()) + NewBufferAlloc.Alloc.Offset;
 
-	return Lock(Context, LockMode, LockSize, Offset);
+	UpdateVulkanBufferStats(GetDesc(), BufferSize, true);
 }
 
-void* FVulkanResourceMultiBuffer::Lock(FVulkanCommandListContext& Context, EResourceLockMode LockMode, uint32 LockSize, uint32 Offset)
+void* FVulkanResourceMultiBuffer::Lock(FRHICommandListBase& RHICmdList, EResourceLockMode LockMode, uint32 LockSize, uint32 Offset)
 {
 	void* Data = nullptr;
 	uint32 DataOffset = 0;
 
-	const bool bVolatile = EnumHasAnyFlags(GetUsage(), BUF_Volatile);
 	check(LockStatus == ELockStatus::Unlocked);
 
 	LockStatus = ELockStatus::Locked;
-	++LockCounter;
+	const bool bIsFirstLock = (0 == LockCounter++);
+
+	// Dynamic:    Allocate a new Host_Visible buffer, swap this new buffer in on RHI thread and update views.  
+	//             GPU reads directly from host memory, but no copy is required so it can be used in render passes.
+	// Static:     A single Device_Local buffer is allocated at creation.  For Lock/Unlock, use a staging buffer for the upload:
+	//             host writes to staging buffer on lock, a copy on GPU is issued on unlock to update the device_local memory.
 
-	if (bVolatile)
+	const bool bUnifiedMem = Device->HasUnifiedMemory();
+	const bool bDynamic = EnumHasAnyFlags(GetUsage(), BUF_Dynamic) || EnumHasAnyFlags(GetUsage(), BUF_Volatile);
+	const bool bStatic = EnumHasAnyFlags(GetUsage(), BUF_Static) || !bDynamic;
+	const bool bUAV = EnumHasAnyFlags(GetUsage(), BUF_UnorderedAccess);
+	const bool bSR = EnumHasAnyFlags(GetUsage(), BUF_ShaderResource);
+
+	check(bStatic || bDynamic || bUAV || bSR);
+
+	if (LockMode == RLM_ReadOnly)
 	{
-		if (LockMode == RLM_ReadOnly)
+		check(IsInRenderingThread());
+
+		if (bUnifiedMem)
 		{
-			checkf(0, TEXT("Volatile buffers can't be locked for read."));
+			Data = CurrentBufferAlloc.HostPtr;
+			DataOffset = Offset;
+			LockStatus = ELockStatus::PersistentMapping;
 		}
-		else
+		else 
 		{
-			FBufferAlloc& BufferAlloc = BufferAllocs[0];
+			Device->PrepareForCPURead();
+		
+			FVulkanCommandListContextImmediate& Context = Device->GetImmediateContext();
+			FVulkanCommandBufferManager* CommandBufferManager = Context.GetCommandBufferManager();
+			FVulkanCmdBuffer* CmdBuffer = CommandBufferManager->GetUploadCmdBuffer();
+				
+			// Make sure any previous tasks have finished on the source buffer.
+			VkMemoryBarrier BarrierBefore = { VK_STRUCTURE_TYPE_MEMORY_BARRIER, nullptr, VK_ACCESS_MEMORY_WRITE_BIT, VK_ACCESS_MEMORY_READ_BIT };
+			VulkanRHI::vkCmdPipelineBarrier(CmdBuffer->GetHandle(), VK_PIPELINE_STAGE_BOTTOM_OF_PIPE_BIT, VK_PIPELINE_STAGE_TRANSFER_BIT, 0, 1, &BarrierBefore, 0, nullptr, 0, nullptr);
+
+			// Create a staging buffer we can use to copy data from device to cpu.
+			VulkanRHI::FStagingBuffer* StagingBuffer = Device->GetStagingManager().AcquireBuffer(LockSize, VK_BUFFER_USAGE_TRANSFER_DST_BIT, VK_MEMORY_PROPERTY_HOST_CACHED_BIT);
+
+			// Fill the staging buffer with the data on the device.
+			VkBufferCopy Regions;
+			Regions.size = LockSize;
+			Regions.srcOffset = Offset + CurrentBufferAlloc.Alloc.Offset;
+			Regions.dstOffset = 0;
+				
+			VulkanRHI::vkCmdCopyBuffer(CmdBuffer->GetHandle(), CurrentBufferAlloc.Alloc.GetBufferHandle(), StagingBuffer->GetHandle(), 1, &Regions);
 
-			VulkanRHI::FTempFrameAllocationBuffer::FTempAllocInfo VolatileAlloc;
-			Context.GetTempFrameAllocationBuffer().Alloc(LockSize + Offset, 256, VolatileAlloc);
-			check(!VolatileAlloc.Allocation.HasAllocation());
+			// Setup barrier.
+			VkMemoryBarrier BarrierAfter = { VK_STRUCTURE_TYPE_MEMORY_BARRIER, nullptr, VK_ACCESS_MEMORY_WRITE_BIT, VK_ACCESS_HOST_READ_BIT };
+			VulkanRHI::vkCmdPipelineBarrier(CmdBuffer->GetHandle(), VK_PIPELINE_STAGE_TRANSFER_BIT, VK_PIPELINE_STAGE_HOST_BIT, 0, 1, &BarrierAfter, 0, nullptr, 0, nullptr);
+				
+			// Force upload.
+			CommandBufferManager->SubmitUploadCmdBuffer();
+			Device->WaitUntilIdle();
 
-			BufferAlloc.Alloc.Reference(VolatileAlloc.Allocation);
-			Data = BufferAlloc.HostPtr = VolatileAlloc.Data;
-			DataOffset = Offset;
+			// Flush.
+			StagingBuffer->FlushMappedMemory();
+
+			// Get mapped pointer. 
+			Data = StagingBuffer->GetMappedPointer();
 
-			// Patch our alloc to go directly to our offset
-			BufferAlloc.Alloc.Offset += VolatileAlloc.CurrentOffset;
-			BufferAlloc.Alloc.Size = VolatileAlloc.Size;
+			// Release temp staging buffer during unlock.
+			FVulkanPendingBufferLock PendingLock;
+			PendingLock.Offset = 0;
+			PendingLock.Size = LockSize;
+			PendingLock.LockMode = LockMode;
+			PendingLock.StagingBuffer = StagingBuffer;
+			AddPendingBufferLock(this, PendingLock);
 
-			BufferAlloc.DeviceAddress = GetBufferDeviceAddress(Device, BufferAlloc.Alloc.GetBufferHandle()) + BufferAlloc.Alloc.Offset;
+			CommandBufferManager->PrepareForNewActiveCommandBuffer();
 		}
 	}
 	else
 	{
-		const bool bDynamic = EnumHasAnyFlags(GetUsage(), BUF_Dynamic);
-		const bool bStatic = EnumHasAnyFlags(GetUsage(), BUF_Static) || !(bVolatile || bDynamic);
-		const bool bUAV = EnumHasAnyFlags(GetUsage(), BUF_UnorderedAccess);
-		const bool bSR = EnumHasAnyFlags(GetUsage(), BUF_ShaderResource);
-		const bool bUnifiedMem = Device->HasUnifiedMemory();
+		check(LockMode == RLM_WriteOnly);
 
-		check(bStatic || bDynamic || bUAV || bSR);
-
-		if (LockMode == RLM_ReadOnly)
+		// If this is the first lock on host visible memory, then the memory is still untouched so use it directly
+		if ((bUnifiedMem || bDynamic) && bIsFirstLock)
+		{
+			check(CurrentBufferAlloc.HostPtr);
+			Data = CurrentBufferAlloc.HostPtr;
+			DataOffset = Offset;
+			LockStatus = ELockStatus::PersistentMapping;
+		}
+		else if (bStatic || GVulkanForceStagingBufferOnLock)
 		{
-			check(IsInRenderingThread() && Context.IsImmediate());
+			FVulkanPendingBufferLock PendingLock;
+			PendingLock.Offset = Offset;
+			PendingLock.Size = LockSize;
+			PendingLock.LockMode = LockMode;
 
-			if (bUnifiedMem)
-			{
-				Data = BufferAllocs[CurrentBufferIndex].HostPtr;
-				DataOffset = Offset;
-				LockStatus = ELockStatus::PersistentMapping;
-			}
-			else 
-			{
-				Device->PrepareForCPURead();
-				FVulkanCmdBuffer* CmdBuffer = Context.GetCommandBufferManager()->GetUploadCmdBuffer();
-				
-				// Make sure any previous tasks have finished on the source buffer.
-				VkMemoryBarrier BarrierBefore = { VK_STRUCTURE_TYPE_MEMORY_BARRIER, nullptr, VK_ACCESS_MEMORY_WRITE_BIT, VK_ACCESS_MEMORY_READ_BIT };
-				VulkanRHI::vkCmdPipelineBarrier(CmdBuffer->GetHandle(), VK_PIPELINE_STAGE_BOTTOM_OF_PIPE_BIT, VK_PIPELINE_STAGE_TRANSFER_BIT, 0, 1, &BarrierBefore, 0, nullptr, 0, nullptr);
-
-				// Create a staging buffer we can use to copy data from device to cpu.
-				VulkanRHI::FStagingBuffer* StagingBuffer = Device->GetStagingManager().AcquireBuffer(LockSize, VK_BUFFER_USAGE_TRANSFER_DST_BIT, VK_MEMORY_PROPERTY_HOST_CACHED_BIT);
-
-				// Fill the staging buffer with the data on the device.
-				VkBufferCopy Regions;
-				Regions.size = LockSize;
-				Regions.srcOffset = Offset + BufferAllocs[CurrentBufferIndex].Alloc.Offset;
-				Regions.dstOffset = 0;
-				
-				VulkanRHI::vkCmdCopyBuffer(CmdBuffer->GetHandle(), BufferAllocs[CurrentBufferIndex].Alloc.GetBufferHandle(), StagingBuffer->GetHandle(), 1, &Regions);
+			VulkanRHI::FStagingBuffer* StagingBuffer = Device->GetStagingManager().AcquireBuffer(LockSize, VK_BUFFER_USAGE_TRANSFER_SRC_BIT);
+			PendingLock.StagingBuffer = StagingBuffer;
+			Data = StagingBuffer->GetMappedPointer();
 
-				// Setup barrier.
-				VkMemoryBarrier BarrierAfter = { VK_STRUCTURE_TYPE_MEMORY_BARRIER, nullptr, VK_ACCESS_MEMORY_WRITE_BIT, VK_ACCESS_HOST_READ_BIT };
-				VulkanRHI::vkCmdPipelineBarrier(CmdBuffer->GetHandle(), VK_PIPELINE_STAGE_TRANSFER_BIT, VK_PIPELINE_STAGE_HOST_BIT, 0, 1, &BarrierAfter, 0, nullptr, 0, nullptr);
-				
-				// Force upload.
-				Context.GetCommandBufferManager()->SubmitUploadCmdBuffer();
-				Device->WaitUntilIdle();
-
-				// Flush.
-				StagingBuffer->FlushMappedMemory();
-
-				// Get mapped pointer. 
-				Data = StagingBuffer->GetMappedPointer();
-
-				// Release temp staging buffer during unlock.
-				VulkanRHI::FPendingBufferLock PendingLock;
-				PendingLock.Offset = 0;
-				PendingLock.Size = LockSize;
-				PendingLock.LockMode = LockMode;
-				PendingLock.StagingBuffer = StagingBuffer;
-
-				{
-					FScopeLock ScopeLock(&GPendingLockIBsMutex);
-					check(!GPendingLockIBs.Contains(this));
-					GPendingLockIBs.Add(this, PendingLock);
-				}
-
-				Context.GetCommandBufferManager()->PrepareForNewActiveCommandBuffer();
-			}
+			AddPendingBufferLock(this, PendingLock);
 		}
 		else
 		{
-			check(LockMode == RLM_WriteOnly);
+			FBufferAlloc NewAlloc;
+			AllocateMemory(NewAlloc);
+			NewAlloc.Alloc.Disown();
 
-			// Always use staging buffers to update 'Static' buffers since they maybe be in use by GPU atm
-			const bool bUseStagingBuffer = (bStatic || !bUnifiedMem) || GVulkanForceStagingBufferOnLock;
-			if (bUseStagingBuffer)
+			RHICmdList.EnqueueLambda(TEXT("FVulkanBuffer::Lock"), [Buffer = this, NewAlloc](FRHICommandListBase& CmdList)
 			{
-				// NOTE: No need to change the CurrentBufferIndex if we're using a staging buffer for the copy
-
-				VulkanRHI::FPendingBufferLock PendingLock;
-				PendingLock.Offset = Offset;
-				PendingLock.Size = LockSize;
-				PendingLock.LockMode = LockMode;
-
-				VulkanRHI::FStagingBuffer* StagingBuffer = Device->GetStagingManager().AcquireBuffer(LockSize, VK_BUFFER_USAGE_TRANSFER_SRC_BIT);
-				PendingLock.StagingBuffer = StagingBuffer;
-				Data = StagingBuffer->GetMappedPointer();
-
-				{
-					FScopeLock ScopeLock(&GPendingLockIBsMutex);
-					check(!GPendingLockIBs.Contains(this));
-					GPendingLockIBs.Add(this, PendingLock);
-				}
-			}
-			else
-			{
-				AdvanceBufferIndex();
+				Buffer->CurrentBufferAlloc.Alloc.Free(*Buffer->GetParent());
+				Buffer->CurrentBufferAlloc = NewAlloc;
+				Buffer->CurrentBufferAlloc.Alloc.Own();
+				Buffer->UpdateLinkedViews();
+			});
 
-				Data = BufferAllocs[CurrentBufferIndex].HostPtr;
-				DataOffset = Offset;
-				LockStatus = ELockStatus::PersistentMapping;
-			}
+			Data = NewAlloc.HostPtr;
+			DataOffset = Offset;
+			LockStatus = ELockStatus::PersistentMapping;
 		}
 	}
 
@@ -436,130 +339,65 @@ void* FVulkanResourceMultiBuffer::Lock(FVulkanCommandListContext& Context, EReso
 	return (uint8*)Data + DataOffset;
 }
 
-inline void FVulkanResourceMultiBuffer::InternalUnlock(FVulkanCommandListContext& Context, VulkanRHI::FPendingBufferLock& PendingLock, FVulkanResourceMultiBuffer* MultiBuffer, int32 InBufferIndex)
-{
-	const uint32 LockSize = PendingLock.Size;
-	const uint32 LockOffset = PendingLock.Offset;
-	VulkanRHI::FStagingBuffer* StagingBuffer = PendingLock.StagingBuffer;
-	PendingLock.StagingBuffer = nullptr;
-
-	// We need to do this on the active command buffer instead of using an upload command buffer. The high level code sometimes reuses the same
-	// buffer in sequences of upload / dispatch, upload / dispatch, so we need to order the copy commands correctly with respect to the dispatches.
-	FVulkanCmdBuffer* Cmd = Context.GetCommandBufferManager()->GetActiveCmdBuffer();
-	check(Cmd && Cmd->IsOutsideRenderPass());
-	VkCommandBuffer CmdBuffer = Cmd->GetHandle();
-
-	VulkanRHI::DebugHeavyWeightBarrier(CmdBuffer, 16);
-
-	VkBufferCopy Region;
-	FMemory::Memzero(Region);
-	Region.size = LockSize;
-	//Region.srcOffset = 0;
-	Region.dstOffset = LockOffset + MultiBuffer->BufferAllocs[InBufferIndex].Alloc.Offset;
-	VulkanRHI::vkCmdCopyBuffer(CmdBuffer, StagingBuffer->GetHandle(), MultiBuffer->BufferAllocs[InBufferIndex].Alloc.GetBufferHandle(), 1, &Region);
-
-	// High level code expects the data in MultiBuffer to be ready to read
-	VkMemoryBarrier BarrierAfter = { VK_STRUCTURE_TYPE_MEMORY_BARRIER, nullptr, VK_ACCESS_TRANSFER_WRITE_BIT, VK_ACCESS_MEMORY_READ_BIT | VK_ACCESS_MEMORY_WRITE_BIT };
-	VulkanRHI::vkCmdPipelineBarrier(CmdBuffer, VK_PIPELINE_STAGE_TRANSFER_BIT, VK_PIPELINE_STAGE_ALL_COMMANDS_BIT, 0, 1, &BarrierAfter, 0, nullptr, 0, nullptr);
-
-	MultiBuffer->GetParent()->GetStagingManager().ReleaseBuffer(Cmd, StagingBuffer);
-
-	MultiBuffer->UpdateBufferAllocStates(Context);
-	MultiBuffer->UpdateLinkedViews();
-}
 
-struct FRHICommandMultiBufferUnlock final : public FRHICommand<FRHICommandMultiBufferUnlock>
+void FVulkanResourceMultiBuffer::Unlock(FRHICommandListBase& RHICmdList)
 {
-	VulkanRHI::FPendingBufferLock PendingLock;
-	FVulkanResourceMultiBuffer* MultiBuffer;
-	FVulkanDevice* Device;
-	int32 BufferIndex;
-
-	FRHICommandMultiBufferUnlock(FVulkanDevice* InDevice, const VulkanRHI::FPendingBufferLock& InPendingLock, FVulkanResourceMultiBuffer* InMultiBuffer, int32 InBufferIndex)
-		: PendingLock(InPendingLock)
-		, MultiBuffer(InMultiBuffer)
-		, Device(InDevice)
-		, BufferIndex(InBufferIndex)
-	{
-	}
-
-	void Execute(FRHICommandListBase& CmdList)
-	{
-		FVulkanResourceMultiBuffer::InternalUnlock(FVulkanCommandListContext::GetVulkanContext(CmdList.GetContext()), PendingLock, MultiBuffer, BufferIndex);
-	}
-};
-
-void FVulkanResourceMultiBuffer::Unlock(FRHICommandListBase* RHICmdList, FVulkanCommandListContext* Context)
-{
-	check(RHICmdList || Context);
-
-	const bool bDynamic = EnumHasAnyFlags(GetUsage(), BUF_Dynamic);
-	const bool bVolatile = EnumHasAnyFlags(GetUsage(), BUF_Volatile);
-	const bool bStatic = EnumHasAnyFlags(GetUsage(), BUF_Static) || !(bVolatile || bDynamic);
+	const bool bUnifiedMem = Device->HasUnifiedMemory();
+	const bool bDynamic = EnumHasAnyFlags(GetUsage(), BUF_Dynamic) || EnumHasAnyFlags(GetUsage(), BUF_Volatile);
+	const bool bStatic = EnumHasAnyFlags(GetUsage(), BUF_Static) || !bDynamic;
 	const bool bSR = EnumHasAnyFlags(GetUsage(), BUF_ShaderResource);
 
 	check(LockStatus != ELockStatus::Unlocked);
 
-	if (bVolatile)
-	{
-		if (RHICmdList && RHICmdList->IsTopOfPipe())
-		{
-			RHICmdList->EnqueueLambda([this](FRHICommandListBase&)
-			{
-				UpdateLinkedViews();
-			});
-		}
-		else
-		{
-			UpdateLinkedViews();
-		}
-	}
-	else if (LockStatus == ELockStatus::PersistentMapping)
+	if (LockStatus == ELockStatus::PersistentMapping)
 	{
-		if (Context)
-		{
-			UpdateBufferAllocStates(*Context);
-			UpdateLinkedViews();
-		}
-		else
-		{
-			RHICmdList->EnqueueLambda([this](FRHICommandListBase& CmdList)
-			{
-				FVulkanCommandListContext& Context = FVulkanCommandListContext::GetVulkanContext(CmdList.GetContext());
-				UpdateBufferAllocStates(Context);
-				UpdateLinkedViews();
-			});
-		}
+		// Do nothing
 	}
 	else
 	{
 		check(bStatic || bDynamic || bSR);
 
-		VulkanRHI::FPendingBufferLock PendingLock = GetPendingBufferLock(this);
+		FVulkanPendingBufferLock PendingLock = GetPendingBufferLock(this);
 
-		PendingLock.StagingBuffer->FlushMappedMemory();
-
-		if (PendingLock.LockMode == RLM_ReadOnly)
-		{
-			// Just remove the staging buffer here.
-			Device->GetStagingManager().ReleaseBuffer(0, PendingLock.StagingBuffer);
-		}
-		else if (PendingLock.LockMode == RLM_WriteOnly)
+		RHICmdList.EnqueueLambda(TEXT("FVulkanBuffer::Unlock"), [Buffer=this, PendingLock](FRHICommandListBase& CmdList)
 		{
-			if (Context || (RHICmdList && RHICmdList->IsBottomOfPipe()))
-			{
-				if (!Context)
-				{
-					Context = &FVulkanCommandListContext::GetVulkanContext(RHICmdList->GetContext());
-				}
+			VulkanRHI::FStagingBuffer* StagingBuffer = PendingLock.StagingBuffer;
+			check(StagingBuffer);
+			StagingBuffer->FlushMappedMemory();
 
-				FVulkanResourceMultiBuffer::InternalUnlock(*Context, PendingLock, this, CurrentBufferIndex);
+			if (PendingLock.LockMode == RLM_ReadOnly)
+			{
+				// Just remove the staging buffer here.
+				Buffer->Device->GetStagingManager().ReleaseBuffer(nullptr, StagingBuffer);
 			}
-			else
+			else if (PendingLock.LockMode == RLM_WriteOnly)
 			{
-				ALLOC_COMMAND_CL(*RHICmdList, FRHICommandMultiBufferUnlock)(Device, PendingLock, this, CurrentBufferIndex);
+				FVulkanCommandListContext& Context = FVulkanCommandListContext::GetVulkanContext(CmdList.GetContext());
+
+				// We need to do this on the active command buffer instead of using an upload command buffer. The high level code sometimes reuses the same
+				// buffer in sequences of upload / dispatch, upload / dispatch, so we need to order the copy commands correctly with respect to the dispatches.
+				FVulkanCmdBuffer* Cmd = Context.GetCommandBufferManager()->GetActiveCmdBuffer();
+				check(Cmd && Cmd->IsOutsideRenderPass());
+				VkCommandBuffer CmdBuffer = Cmd->GetHandle();
+
+				VulkanRHI::DebugHeavyWeightBarrier(CmdBuffer, 16);
+
+				VkBufferCopy Region;
+				FMemory::Memzero(Region);
+				Region.size = PendingLock.Size;
+				//Region.srcOffset = 0;
+				Region.dstOffset = PendingLock.Offset + Buffer->CurrentBufferAlloc.Alloc.Offset;
+				VulkanRHI::vkCmdCopyBuffer(CmdBuffer, StagingBuffer->GetHandle(), Buffer->CurrentBufferAlloc.Alloc.GetBufferHandle(), 1, &Region);
+
+				// High level code expects the data in Buffer to be ready to read
+				VkMemoryBarrier BarrierAfter = { VK_STRUCTURE_TYPE_MEMORY_BARRIER, nullptr, VK_ACCESS_TRANSFER_WRITE_BIT, VK_ACCESS_MEMORY_READ_BIT | VK_ACCESS_MEMORY_WRITE_BIT };
+				VulkanRHI::vkCmdPipelineBarrier(CmdBuffer, VK_PIPELINE_STAGE_TRANSFER_BIT, VK_PIPELINE_STAGE_ALL_COMMANDS_BIT, 0, 1, &BarrierAfter, 0, nullptr, 0, nullptr);
+
+				Buffer->GetParent()->GetStagingManager().ReleaseBuffer(Cmd, StagingBuffer);
+
+				Buffer->UpdateLinkedViews(); // :todo-jn:  not needed?  Same buffer when we use staging?
 			}
-		}
+		});
 	}
 
 	LockStatus = ELockStatus::Unlocked;
@@ -577,34 +415,20 @@ void FVulkanResourceMultiBuffer::TakeOwnership(FVulkanResourceMultiBuffer& Other
 	FRHIBuffer::TakeOwnership(Other);
 
 	BufferUsageFlags   = Other.BufferUsageFlags;
-	CurrentBufferIndex = Other.CurrentBufferIndex;
+	CurrentBufferAlloc = Other.CurrentBufferAlloc;
 
 	Other.BufferUsageFlags   = {};
-	Other.CurrentBufferIndex = -1;
-
-	// Swap the empty array from the ReleaseOwnership with the Other allocations
-	::Swap(BufferAllocs, Other.BufferAllocs);
+	Other.CurrentBufferAlloc = {};
 }
 
 void FVulkanResourceMultiBuffer::ReleaseOwnership()
 {
 	check(LockStatus == ELockStatus::Unlocked);
 
-	uint64 TotalSize = 0;
-	for (int32 Index = 0; Index < BufferAllocs.Num(); ++Index)
+	if (CurrentBufferAlloc.Alloc.HasAllocation())
 	{
-		if (BufferAllocs[Index].Alloc.HasAllocation())
-		{
-			TotalSize += BufferAllocs[Index].Alloc.Size;
-			Device->GetMemoryManager().FreeVulkanAllocation(BufferAllocs[Index].Alloc);
-		}
-		BufferAllocs[Index].Fence = nullptr;
-	}
-	BufferAllocs.Empty();
-
-	if (TotalSize > 0)
-	{
-		UpdateVulkanBufferStats(GetDesc(), TotalSize, false);
+		UpdateVulkanBufferStats(GetDesc(), CurrentBufferAlloc.Alloc.Size, false);
+		Device->GetMemoryManager().FreeVulkanAllocation(CurrentBufferAlloc.Alloc);
 	}
 
 	FRHIBuffer::ReleaseOwnership();
@@ -644,17 +468,7 @@ void* FVulkanDynamicRHI::RHILockBuffer(FRHICommandListBase& RHICmdList, FRHIBuff
 	QUICK_SCOPE_CYCLE_COUNTER(STAT_LockBuffer_RenderThread);
 	LLM_SCOPE_VULKAN(ELLMTagVulkan::VulkanBuffers);
 	FVulkanResourceMultiBuffer* Buffer = ResourceCast(BufferRHI);
-
-	if (LockMode == RLM_ReadOnly)
-	{
-		// Lock in ReadOnly will call PrepareForCPURead on Device, so we can use the immediate context directly.
-		// This is a legacy path though, FRHIGPUBufferReadback should be used for reading back buffers.
-		return Buffer->Lock(Device->GetImmediateContext(), LockMode, Size, Offset);
-	}
-	else
-	{
-		return Buffer->Lock(RHICmdList, LockMode, Size, Offset);
-	}
+	return Buffer->Lock(RHICmdList, LockMode, Size, Offset);
 }
 
 void FVulkanDynamicRHI::RHIUnlockBuffer(FRHICommandListBase& RHICmdList, FRHIBuffer* BufferRHI)
diff --git a/Engine/Source/Runtime/VulkanRHI/Private/VulkanRHIPrivate.h b/Engine/Source/Runtime/VulkanRHI/Private/VulkanRHIPrivate.h
--- a/Engine/Source/Runtime/VulkanRHI/Private/VulkanRHIPrivate.h
+++ b/Engine/Source/Runtime/VulkanRHI/Private/VulkanRHIPrivate.h
@@ -433,14 +433,6 @@ DECLARE_DWORD_ACCUMULATOR_STAT_EXTERN(TEXT("Max Potential Desc Sets"), STAT_Vulk
 
 namespace VulkanRHI
 {
-	struct FPendingBufferLock
-	{
-		FStagingBuffer* StagingBuffer;
-		uint32 Offset;
-		uint32 Size;
-		EResourceLockMode LockMode;
-	};
-
 	static uint32 GetNumBitsPerPixel(VkFormat Format)
 	{
 		switch (Format)
diff --git a/Engine/Source/Runtime/VulkanRHI/Private/VulkanRayTracing.cpp b/Engine/Source/Runtime/VulkanRHI/Private/VulkanRayTracing.cpp
--- a/Engine/Source/Runtime/VulkanRHI/Private/VulkanRayTracing.cpp
+++ b/Engine/Source/Runtime/VulkanRHI/Private/VulkanRayTracing.cpp
@@ -762,24 +762,27 @@ void FVulkanRayTracingScene::BuildAccelerationStructure(
 {
 	const bool bIsUpdate = BuildMode == EAccelerationStructureBuildMode::Update;
 
-	// Build a metadata buffer	that contains VulkanRHI-specific per-geometry parameters that allow us to access
-	// vertex and index buffers from shaders that use inline ray tracing.
-	BuildPerInstanceGeometryParameterBuffer(CommandContext);
+	FBufferRHIRef ScratchBuffer;
+	{
+		TRHICommandList_RecursiveHazardous<FVulkanCommandListContext> RHICmdList(&CommandContext);
 
-	check(AccelerationStructureBuffer.IsValid());
-	check(InInstanceBuffer != nullptr);
+		// Build a metadata buffer	that contains VulkanRHI-specific per-geometry parameters that allow us to access
+		// vertex and index buffers from shaders that use inline ray tracing.
+		BuildPerInstanceGeometryParameterBuffer(RHICmdList);
 
-	FBufferRHIRef ScratchBuffer;
+		check(AccelerationStructureBuffer.IsValid());
+		check(InInstanceBuffer != nullptr);
 
-	if (InScratchBuffer == nullptr)
-	{
-		const uint64 ScratchBufferSize = bIsUpdate ? SizeInfo.UpdateScratchSize : SizeInfo.BuildScratchSize;
 
-		TRHICommandList_RecursiveHazardous<FVulkanCommandListContext> RHICmdList(&CommandContext);
-		FRHIResourceCreateInfo ScratchBufferCreateInfo(TEXT("BuildScratchTLAS"));
-		ScratchBuffer = RHICmdList.CreateBuffer(ScratchBufferSize, BUF_StructuredBuffer | BUF_RayTracingScratch, 0, ERHIAccess::UAVCompute, ScratchBufferCreateInfo);
-		InScratchBuffer = ResourceCast(ScratchBuffer.GetReference());
-		InScratchOffset = 0;
+		if (InScratchBuffer == nullptr)
+		{
+			const uint64 ScratchBufferSize = bIsUpdate ? SizeInfo.UpdateScratchSize : SizeInfo.BuildScratchSize;
+
+			FRHIResourceCreateInfo ScratchBufferCreateInfo(TEXT("BuildScratchTLAS"));
+			ScratchBuffer = RHICmdList.CreateBuffer(ScratchBufferSize, BUF_StructuredBuffer | BUF_RayTracingScratch, 0, ERHIAccess::UAVCompute, ScratchBufferCreateInfo);
+			InScratchBuffer = ResourceCast(ScratchBuffer.GetReference());
+			InScratchOffset = 0;
+		}
 	}
 
 	PRAGMA_DISABLE_DEPRECATION_WARNINGS
@@ -860,7 +863,7 @@ void FVulkanRayTracingScene::BuildAccelerationStructure(
 	bBuilt = true;
 }
 
-void FVulkanRayTracingScene::BuildPerInstanceGeometryParameterBuffer(FVulkanCommandListContext& CommandContext)
+void FVulkanRayTracingScene::BuildPerInstanceGeometryParameterBuffer(FRHICommandListBase& RHICmdList)
 {
 	// TODO: we could cache parameters in the geometry object to avoid some of the pointer chasing (if this is measured to be a performance issue)
 
@@ -869,7 +872,7 @@ void FVulkanRayTracingScene::BuildPerInstanceGeometryParameterBuffer(FVulkanComm
 
 	check(IsInRHIThread() || !IsRunningRHIInSeparateThread());
 
-	void* MappedBuffer = PerInstanceGeometryParameterBuffer->Lock(CommandContext, RLM_WriteOnly, ParameterBufferSize, 0);
+	void* MappedBuffer = PerInstanceGeometryParameterBuffer->Lock(RHICmdList, RLM_WriteOnly, ParameterBufferSize, 0);
 	FVulkanRayTracingGeometryParameters* MappedParameters = reinterpret_cast<FVulkanRayTracingGeometryParameters*>(MappedBuffer);
 	uint32 ParameterIndex = 0;
 
@@ -914,7 +917,7 @@ void FVulkanRayTracingScene::BuildPerInstanceGeometryParameterBuffer(FVulkanComm
 
 	check(ParameterIndex == Initializer.NumTotalSegments);
 
-	PerInstanceGeometryParameterBuffer->Unlock(CommandContext);
+	PerInstanceGeometryParameterBuffer->Unlock(RHICmdList);
 }
 
 
@@ -1347,7 +1350,7 @@ void FVulkanCommandListContext::RHIBuildAccelerationStructure(const FRayTracingS
 	FVulkanResourceMultiBuffer* const ScratchBuffer = ResourceCast(SceneBuildParams.ScratchBuffer);
 	FVulkanResourceMultiBuffer* const InstanceBuffer = ResourceCast(SceneBuildParams.InstanceBuffer);
 	Scene->BuildAccelerationStructure(
-		*this, 
+		*this,
 		ScratchBuffer, SceneBuildParams.ScratchBufferOffset, 
 		InstanceBuffer, SceneBuildParams.InstanceBufferOffset,
 		SceneBuildParams.BuildMode);
diff --git a/Engine/Source/Runtime/VulkanRHI/Private/VulkanRayTracing.h b/Engine/Source/Runtime/VulkanRHI/Private/VulkanRayTracing.h
--- a/Engine/Source/Runtime/VulkanRHI/Private/VulkanRayTracing.h
+++ b/Engine/Source/Runtime/VulkanRHI/Private/VulkanRayTracing.h
@@ -254,7 +254,7 @@ private:
 	
 	TMap<const FVulkanRayTracingPipelineState*, TRefCountPtr<FVulkanRayTracingShaderTable>> ShaderTables;
 
-	void BuildPerInstanceGeometryParameterBuffer(FVulkanCommandListContext& CommandContext);
+	void BuildPerInstanceGeometryParameterBuffer(FRHICommandListBase& RHICmdList);
 
 	UE::FMutex Mutex;
 	bool bBuilt = false;
diff --git a/Engine/Source/Runtime/VulkanRHI/Private/VulkanResources.h b/Engine/Source/Runtime/VulkanRHI/Private/VulkanResources.h
--- a/Engine/Source/Runtime/VulkanRHI/Private/VulkanResources.h
+++ b/Engine/Source/Runtime/VulkanRHI/Private/VulkanResources.h
@@ -25,6 +25,7 @@ class FVulkanLayout;
 class FVulkanOcclusionQuery;
 class FVulkanCommandBufferManager;
 struct FRHITransientHeapAllocation;
+struct FVulkanPendingBufferLock;
 
 class FVulkanView;
 class FVulkanViewableResource;
@@ -34,7 +35,6 @@ class FVulkanUnorderedAccessView;
 namespace VulkanRHI
 {
 	class FDeviceMemoryAllocation;
-	struct FPendingBufferLock;
 }
 
 enum
@@ -1055,7 +1055,7 @@ public:
 
 	inline const VulkanRHI::FVulkanAllocation& GetCurrentAllocation() const
 	{
-		return BufferAllocs[CurrentBufferIndex].Alloc;
+		return CurrentBufferAlloc.Alloc;
 	}
 
 	inline VkBuffer GetHandle() const
@@ -1068,11 +1068,6 @@ public:
 		return EnumHasAnyFlags(GetUsage(), BUF_Volatile);
 	}
 
-	inline int32 GetNumBuffers() const
-	{
-		return BufferAllocs.Num();
-	}
-
 	// Offset used for Binding a VkBuffer
 	inline uint32 GetOffset() const
 	{
@@ -1087,7 +1082,7 @@ public:
 
 	inline VkDeviceAddress GetDeviceAddress() const
 	{
-		return BufferAllocs[CurrentBufferIndex].DeviceAddress;
+		return CurrentBufferAlloc.DeviceAddress;
 	}
 
 	inline VkBufferUsageFlags GetBufferUsageFlags() const
@@ -1101,16 +1096,7 @@ public:
 	}
 
 	void* Lock(FRHICommandListBase& RHICmdList, EResourceLockMode LockMode, uint32 Size, uint32 Offset);
-	void* Lock(FVulkanCommandListContext& Context, EResourceLockMode LockMode, uint32 Size, uint32 Offset);
-
-	inline void Unlock(FRHICommandListBase& RHICmdList)
-	{
-		Unlock(&RHICmdList, nullptr);
-	}
-	inline void Unlock(FVulkanCommandListContext& Context)
-	{
-		Unlock(nullptr, &Context);
-	}
+	void Unlock(FRHICommandListBase& RHICmdList);
 
 	void TakeOwnership(FVulkanResourceMultiBuffer& Other);
 	void ReleaseOwnership();
@@ -1118,17 +1104,22 @@ public:
 	template<typename T>
 	void DumpMemory(T Callback)
 	{
-		Callback(TEXT("FVulkanResourceMultiBuffer"), FName(), this, 0, GetCurrentSize() * GetNumBuffers(), 1, 1, VK_FORMAT_UNDEFINED);
+		Callback(TEXT("FVulkanResourceMultiBuffer"), FName(), this, 0, GetCurrentSize(), 1, 1, VK_FORMAT_UNDEFINED);
 	}
 
 	static VkBufferUsageFlags UEToVKBufferUsageFlags(FVulkanDevice* InDevice, EBufferUsageFlags InUEUsage, bool bZeroSize);
 
-protected:
+	struct FBufferAlloc
+	{
+		VulkanRHI::FVulkanAllocation Alloc;
+		void* HostPtr = nullptr;
+		VkDeviceAddress DeviceAddress = 0;
+	};
 
-	void AdvanceBufferIndex();
-	void UpdateBufferAllocStates(FVulkanCommandListContext& Context);
+protected:
 
-	void Unlock(FRHICommandListBase* RHICmdList, FVulkanCommandListContext* Context);
+	// Will return a new allocation that can be used for this buffer (proper memory type and size)
+	void AllocateMemory(FBufferAlloc& OutAlloc);
 
 	VkBufferUsageFlags BufferUsageFlags;
 
@@ -1139,27 +1130,9 @@ protected:
 		PersistentMapping,
 	} LockStatus = ELockStatus::Unlocked;
 
-	struct FBufferAlloc
-	{
-		VulkanRHI::FVulkanAllocation Alloc;
-		void* HostPtr = nullptr;
-		class FVulkanGPUFence* Fence = nullptr;
-		VkDeviceAddress DeviceAddress = 0;
-
-		enum class EAllocStatus : uint8
-		{
-			Available,	// The allocation is ready to be used
-			InUse,		// CurrentBufferIndex should point to this allocation
-			NeedsFence,	// The allocation was just released and needs a fence to make sure previous commands are done with it
-			Pending,	// Fence was written, we are waiting on it to know that the alloc can be used again
-		} AllocStatus = EAllocStatus::Available;
-	};
-	TArray<FBufferAlloc, TInlineAllocator<3>> BufferAllocs;
-	int32 CurrentBufferIndex = -1;
+	FBufferAlloc CurrentBufferAlloc;
 	uint32 LockCounter = 0;
 
-	static void InternalUnlock(FVulkanCommandListContext& Context, VulkanRHI::FPendingBufferLock& PendingLock, FVulkanResourceMultiBuffer* MultiBuffer, int32 InDynamicBufferIndex);
-
 	friend class FVulkanCommandListContext;
 	friend struct FRHICommandMultiBufferUnlock;
 };
diff --git a/Engine/Source/Runtime/VulkanRHI/Private/VulkanUAV.cpp b/Engine/Source/Runtime/VulkanRHI/Private/VulkanUAV.cpp
--- a/Engine/Source/Runtime/VulkanRHI/Private/VulkanUAV.cpp
+++ b/Engine/Source/Runtime/VulkanRHI/Private/VulkanUAV.cpp
@@ -79,12 +79,6 @@ FVulkanView* FVulkanView::InitAsTypedBufferView(FVulkanResourceMultiBuffer* Buff
 	ViewInfo.offset = TotalOffset;
 	ViewInfo.format = Format;
 
-	// :todo-jn: Volatile buffers use temporary allocations that can be smaller than the buffer creation size.  Check if the savings are still worth it.
-	if (Buffer->IsVolatile())
-	{
-		InSize = FMath::Min<uint64>(InSize, Buffer->GetCurrentSize());
-	}
-
 	const uint32 TypeSize =  VulkanRHI::GetNumBitsPerPixel(Format) / 8u;
 	// View size has to be a multiple of element size
 	// Commented out because there are multiple places in the high level rendering code which re-purpose buffers for a new format while there are still
