VulkanRHI: add messageId for existing skipped validation message (changed in newer validation layers)

#rnx

diff --git a/Engine/Source/Runtime/VulkanRHI/Private/VulkanDebug.cpp b/Engine/Source/Runtime/VulkanRHI/Private/VulkanDebug.cpp
--- a/Engine/Source/Runtime/VulkanRHI/Private/VulkanDebug.cpp
+++ b/Engine/Source/Runtime/VulkanRHI/Private/VulkanDebug.cpp
@@ -122,7 +122,8 @@ static VKAPI_ATTR VkBool32 VKAPI_CALL DebugUtilsCallback(VkDebugUtilsMessageSeve
 		// Warning: *** [Warning:Validation-1(UNASSIGNED-CoreValidation-Shader-OutputNotConsumed)] fragment shader writes to output location 0 with no matching attachment
 		return VK_FALSE;
 	}
-	else if (!FCStringAnsi::Strcmp(CallbackData->pMessageIdName, "VUID-VkSwapchainCreateInfoKHR-imageExtent-01274"))
+	else if (!FCStringAnsi::Strcmp(CallbackData->pMessageIdName, "VUID-VkSwapchainCreateInfoKHR-imageExtent-01274")||
+		     !FCStringAnsi::Strcmp(CallbackData->pMessageIdName, "VUID-VkSwapchainCreateInfoKHR-pNext-07781"))
 	{
 		// Warning: *** [Error:Validation341838324(VUID-VkSwapchainCreateInfoKHR-imageExtent-01274)] vkCreateSwapChainKHR() called with imageExtent = (8,8), which is outside the bounds returned by vkGetPhysicalDeviceSurfaceCapabilitiesKHR(): currentExtent = (0,0), minImageExtent = (0,0), maxImageExtent = (0,0).
 		return VK_FALSE;
