- Adding support for the VK_ANDROID_external_memory_android_hardware_buffer extension
- Added support for sampling from external memory / external format textures (that will have format PF_Unknown)

diff --git a/Engine/Source/Runtime/VulkanRHI/Private/Android/VulkanAndroidPlatform.cpp b/Engine/Source/Runtime/VulkanRHI/Private/Android/VulkanAndroidPlatform.cpp
--- a/Engine/Source/Runtime/VulkanRHI/Private/Android/VulkanAndroidPlatform.cpp
+++ b/Engine/Source/Runtime/VulkanRHI/Private/Android/VulkanAndroidPlatform.cpp
@@ -474,6 +474,7 @@ void FVulkanAndroidPlatform::GetDeviceExtensions(FVulkanDevice* Device, FVulkanD
 		OutExtensions.Add(MakeUnique<FVulkanDeviceExtension>(Device, VK_EXT_ASTC_DECODE_MODE_EXTENSION_NAME, VULKAN_SUPPORTS_ASTC_DECODE_MODE, VULKAN_EXTENSION_NOT_PROMOTED, DEVICE_EXT_FLAG_SETTER(HasEXTASTCDecodeMode)));
 	}
 	OutExtensions.Add(MakeUnique<FVulkanDeviceExtension>(Device, VK_EXT_TEXTURE_COMPRESSION_ASTC_HDR_EXTENSION_NAME, VULKAN_SUPPORTS_TEXTURE_COMPRESSION_ASTC_HDR, VK_API_VERSION_1_3, DEVICE_EXT_FLAG_SETTER(HasEXTTextureCompressionASTCHDR)));
+	OutExtensions.Add(MakeUnique<FVulkanDeviceExtension>(Device, VK_ANDROID_EXTERNAL_MEMORY_ANDROID_HARDWARE_BUFFER_EXTENSION_NAME, VULKAN_SUPPORTS_ANDROID_EXTERNAL_MEMORY_ANDROID_HARDWARE_BUFFER, VULKAN_EXTENSION_NOT_PROMOTED, DEVICE_EXT_FLAG_SETTER(HasANDROIDExternalMemoryHardwareBuffer)));
 
 #if !UE_BUILD_SHIPPING
 	// Layer name as extension
diff --git a/Engine/Source/Runtime/VulkanRHI/Private/VulkanDevice.cpp b/Engine/Source/Runtime/VulkanRHI/Private/VulkanDevice.cpp
--- a/Engine/Source/Runtime/VulkanRHI/Private/VulkanDevice.cpp
+++ b/Engine/Source/Runtime/VulkanRHI/Private/VulkanDevice.cpp
@@ -663,6 +663,12 @@ void FVulkanDevice::SetupFormats()
 		MapFormatSupport(PF_ASTC_12x12_HDR, { VK_FORMAT_ASTC_12x12_SFLOAT_BLOCK_EXT }, ComponentMappingRGBA);
 	}
 
+	// When this extension is available, PF_Unknown texture can have an external buffer attached which has an
+	// internal format, that can be sampled. If it really can be sampled depends on the VK_IMAGE_USAGE_SAMPLED_BIT
+	if (GetOptionalExtensions().HasANDROIDExternalMemoryHardwareBuffer)
+	{
+		GPixelFormats[PF_Unknown].Capabilities |= EPixelFormatCapabilities::TextureSample;
+	}
 
 	// Verify available Vertex Formats
 	{
diff --git a/Engine/Source/Runtime/VulkanRHI/Private/VulkanDevice.h b/Engine/Source/Runtime/VulkanRHI/Private/VulkanDevice.h
--- a/Engine/Source/Runtime/VulkanRHI/Private/VulkanDevice.h
+++ b/Engine/Source/Runtime/VulkanRHI/Private/VulkanDevice.h
@@ -60,6 +60,7 @@ struct FOptionalVulkanDeviceExtensions
 			uint64 HasAMDBufferMarker : 1;
 			uint64 HasNVDiagnosticCheckpoints : 1;
 			uint64 HasNVDeviceDiagnosticConfig : 1;
+			uint64 HasANDROIDExternalMemoryHardwareBuffer : 1;
 
 			// Promoted to 1.1
 			uint64 HasKHRMultiview : 1;
diff --git a/Engine/Source/Runtime/VulkanRHI/Private/VulkanTexture.cpp b/Engine/Source/Runtime/VulkanRHI/Private/VulkanTexture.cpp
--- a/Engine/Source/Runtime/VulkanRHI/Private/VulkanTexture.cpp
+++ b/Engine/Source/Runtime/VulkanRHI/Private/VulkanTexture.cpp
@@ -1826,6 +1826,9 @@ FVulkanTexture::FVulkanTexture(FVulkanDevice& InDevice, const FRHITextureCreateD
 
 	const VkImageViewType ViewType = GetViewType();
 	const VkDescriptorType DescriptorType = SupportsSampling() ? VK_DESCRIPTOR_TYPE_SAMPLED_IMAGE : VK_DESCRIPTOR_TYPE_STORAGE_IMAGE;
+	const bool bUseIdentitySwizzle = (DescriptorType != VK_DESCRIPTOR_TYPE_SAMPLED_IMAGE) ||
+									(ViewFormat == VK_FORMAT_UNDEFINED); // External buffer textures also require identity swizzle
+
 	if (Image != VK_NULL_HANDLE)
 	{
 		DefaultView = (new FVulkanView(InDevice, DescriptorType))->InitAsTextureView(
@@ -1838,7 +1841,7 @@ FVulkanTexture::FVulkanTexture(FVulkanDevice& InDevice, const FRHITextureCreateD
 			, FMath::Max(InCreateDesc.NumMips, (uint8)1u)
 			, 0
 			, GetNumberOfArrayLevels()
-			, !SupportsSampling()
+			, bUseIdentitySwizzle
 		);
 	}
 
diff --git a/Engine/Source/Runtime/VulkanRHI/Public/VulkanConfiguration.h b/Engine/Source/Runtime/VulkanRHI/Public/VulkanConfiguration.h
--- a/Engine/Source/Runtime/VulkanRHI/Public/VulkanConfiguration.h
+++ b/Engine/Source/Runtime/VulkanRHI/Public/VulkanConfiguration.h
@@ -212,6 +212,10 @@ struct VkAllocationCallbacks;
 #endif
 #endif
 
+#ifndef VULKAN_SUPPORTS_ANDROID_EXTERNAL_MEMORY_ANDROID_HARDWARE_BUFFER
+	#define VULKAN_SUPPORTS_ANDROID_EXTERNAL_MEMORY_ANDROID_HARDWARE_BUFFER	1
+#endif
+
 #ifndef VULKAN_SUPPORTS_RENDERPASS2
 	#ifdef VK_KHR_create_renderpass2
 		#define VULKAN_SUPPORTS_RENDERPASS2 1
