VulkanRHI: Temporarily skip the profile check on build machines until loader is updated.

#rnx

diff --git a/Engine/Source/Runtime/VulkanRHI/Private/VulkanGenericPlatform.cpp b/Engine/Source/Runtime/VulkanRHI/Private/VulkanGenericPlatform.cpp
--- a/Engine/Source/Runtime/VulkanRHI/Private/VulkanGenericPlatform.cpp
+++ b/Engine/Source/Runtime/VulkanRHI/Private/VulkanGenericPlatform.cpp
@@ -168,7 +168,8 @@ void FVulkanGenericPlatform::ClearVulkanInstanceFunctions()
 bool FVulkanGenericPlatform::SupportsProfileChecks()
 {
 	return (CVarVulkanUseProfileCheck.GetValueOnAnyThread() != 0) && 
-		!FParse::Param(FCommandLine::Get(), TEXT("SkipVulkanProfileCheck"));
+		!FParse::Param(FCommandLine::Get(), TEXT("SkipVulkanProfileCheck")) &&
+		!FParse::Param(FCommandLine::Get(), TEXT("BUILDMACHINE"));
 }
 
 FString FVulkanGenericPlatform::GetVulkanProfileNameForFeatureLevel(ERHIFeatureLevel::Type FeatureLevel, bool bRaytracing)
