Fixed performance warning on Android related to too much queries to r.Android.SupportsTimestampQueries cvar

diff --git a/Engine/Source/Runtime/VulkanRHI/Private/Android/VulkanAndroidPlatform.cpp b/Engine/Source/Runtime/VulkanRHI/Private/Android/VulkanAndroidPlatform.cpp
--- a/Engine/Source/Runtime/VulkanRHI/Private/Android/VulkanAndroidPlatform.cpp
+++ b/Engine/Source/Runtime/VulkanRHI/Private/Android/VulkanAndroidPlatform.cpp
@@ -592,7 +592,7 @@ void FVulkanAndroidPlatform::NotifyFoundDeviceLayersAndExtensions(VkPhysicalDevi
 
 bool FVulkanAndroidPlatform::SupportsTimestampRenderQueries()
 {
-	IConsoleVariable* CVarAndroidSupportsTimestampQueries = IConsoleManager::Get().FindConsoleVariable(TEXT("r.Android.SupportsTimestampQueries"));
+	static const auto CVarAndroidSupportsTimestampQueries = IConsoleManager::Get().FindConsoleVariable(TEXT("r.Android.SupportsTimestampQueries"));
 
 	// standalone devices have newer drivers where timestamp render queries work.
 	return CVarAndroidSupportsTimestampQueries != nullptr &&
@@ -605,7 +605,7 @@ bool FVulkanAndroidPlatform::SupportsDynamicResolution()
 
 #if USE_ANDROID_SWAPPY
 	
-	IConsoleVariable* CVarAndroidSupportsDynamicResolution = IConsoleManager::Get().FindConsoleVariable(TEXT("r.Android.SupportsDynamicResolution"));
+	static const auto CVarAndroidSupportsDynamicResolution = IConsoleManager::Get().FindConsoleVariable(TEXT("r.Android.SupportsDynamicResolution"));
 
 	bool bIsSwappyEnabled = FAndroidPlatformRHIFramePacer::CVarUseSwappyForFramePacing.GetValueOnAnyThread() == 1;
 	return bIsSwappyEnabled && 
diff --git a/Engine/Source/Runtime/VulkanRHI/Private/VulkanQuery.cpp b/Engine/Source/Runtime/VulkanRHI/Private/VulkanQuery.cpp
--- a/Engine/Source/Runtime/VulkanRHI/Private/VulkanQuery.cpp
+++ b/Engine/Source/Runtime/VulkanRHI/Private/VulkanQuery.cpp
@@ -390,7 +390,7 @@ void FVulkanCommandListContext::ReadAndCalculateGPUFrameTime()
 {
 	check(IsImmediate());
 
-	if (FVulkanPlatform::SupportsTimestampRenderQueries() && FrameTiming)
+	if (GSupportsTimestampRenderQueries && FrameTiming)
 	{
 		uint64 Delta = 0;
 		
diff --git a/Engine/Source/Runtime/VulkanRHI/Private/VulkanRHI.cpp b/Engine/Source/Runtime/VulkanRHI/Private/VulkanRHI.cpp
--- a/Engine/Source/Runtime/VulkanRHI/Private/VulkanRHI.cpp
+++ b/Engine/Source/Runtime/VulkanRHI/Private/VulkanRHI.cpp
@@ -469,7 +469,7 @@ void FVulkanCommandListContext::ReleasePendingState()
 
 FVulkanCommandListContext::~FVulkanCommandListContext()
 {
-	if (FVulkanPlatform::SupportsTimestampRenderQueries())
+	if (GSupportsTimestampRenderQueries)
 	{
 #if (RHI_NEW_GPU_PROFILER == 0)
 		FrameTiming->Release();
