Fixed vulkan validation warnings about invalid image layout for a RenderingBackBuffer

diff --git a/Engine/Source/Runtime/VulkanRHI/Private/VulkanViewport.cpp b/Engine/Source/Runtime/VulkanRHI/Private/VulkanViewport.cpp
--- a/Engine/Source/Runtime/VulkanRHI/Private/VulkanViewport.cpp
+++ b/Engine/Source/Runtime/VulkanRHI/Private/VulkanViewport.cpp
@@ -786,7 +786,7 @@ void FVulkanViewport::CreateSwapchain(FVulkanSwapChainRecreateInfo* RecreateInfo
 #if VULKAN_ENABLE_DRAW_MARKERS
 			if (Device->GetSetDebugName())
 			{
-				VulkanRHI::SetDebugName(Device->GetSetDebugName(), Device->GetInstanceHandle(), BackBufferImages[Index]->Image, "RenderingBackBuffer");
+				VulkanRHI::SetDebugName(Device->GetSetDebugName(), Device->GetInstanceHandle(), BackBufferImages[Index]->Image, "VulkanBackBuffer");
 			}
 #endif
 		}
@@ -881,6 +881,9 @@ inline static void CopyImageToBackBuffer(FVulkanCommandListContext* Context, FVu
 	RHI_BREADCRUMB_EVENT(*Context, CopyImageToBackBuffer);
 	const bool bNeedsVulkanPreTransform = CachedSurfaceTransform != VK_SURFACE_TRANSFORM_IDENTITY_BIT_KHR;
 
+	FVulkanLayoutManager& LayoutManager = CmdBuffer->GetLayoutManager();
+	const FVulkanImageLayout* SrcImageLayout = LayoutManager.GetFullLayout(SrcSurface);
+
 	VkImageLayout SrcSurfaceLayout = VK_IMAGE_LAYOUT_TRANSFER_SRC_OPTIMAL;
 	VkImageLayout DstSurfaceLayout = VK_IMAGE_LAYOUT_TRANSFER_DST_OPTIMAL;
 
@@ -892,7 +895,11 @@ inline static void CopyImageToBackBuffer(FVulkanCommandListContext* Context, FVu
 
 	{
 		FVulkanPipelineBarrier Barrier;
-		Barrier.AddImageLayoutTransition(SrcSurface.Image, VK_IMAGE_LAYOUT_COLOR_ATTACHMENT_OPTIMAL, SrcSurfaceLayout, FVulkanPipelineBarrier::MakeSubresourceRange(VK_IMAGE_ASPECT_COLOR_BIT, 0, 1));
+		if (!SrcImageLayout || SrcImageLayout->GetSubresLayout(0, 0, VK_IMAGE_ASPECT_COLOR_BIT) != SrcSurfaceLayout)
+		{
+			Barrier.AddImageLayoutTransition(SrcSurface.Image, VK_IMAGE_LAYOUT_UNDEFINED, SrcSurfaceLayout, FVulkanPipelineBarrier::MakeSubresourceRange(VK_IMAGE_ASPECT_COLOR_BIT, 0, 1));
+			LayoutManager.SetFullLayout(SrcSurface, SrcSurfaceLayout);
+		}
 		Barrier.AddImageLayoutTransition(DstSurface.Image, VK_IMAGE_LAYOUT_PRESENT_SRC_KHR, DstSurfaceLayout, FVulkanPipelineBarrier::MakeSubresourceRange(VK_IMAGE_ASPECT_COLOR_BIT, 0, 1));
 		Barrier.Execute(CmdBuffer);
 	}
@@ -989,7 +996,7 @@ inline static void CopyImageToBackBuffer(FVulkanCommandListContext* Context, FVu
 
 	{
 		FVulkanPipelineBarrier Barrier;
-		Barrier.AddImageLayoutTransition(SrcSurface.Image, SrcSurfaceLayout, VK_IMAGE_LAYOUT_COLOR_ATTACHMENT_OPTIMAL, FVulkanPipelineBarrier::MakeSubresourceRange(VK_IMAGE_ASPECT_COLOR_BIT, 0, 1));
+		// We do not need to transition SrcSurface here, render passes that use it will do transition
 		Barrier.AddImageLayoutTransition(DstSurface.Image, DstSurfaceLayout, VK_IMAGE_LAYOUT_PRESENT_SRC_KHR, FVulkanPipelineBarrier::MakeSubresourceRange(VK_IMAGE_ASPECT_COLOR_BIT, 0, 1));
 		Barrier.Execute(CmdBuffer);
 	}
