VulkanRHI: Move the rest of RemappingInfo into FVulkanDescriptorSetsLayoutInfo directly.

#rnx

diff --git a/Engine/Source/Runtime/VulkanRHI/Private/VulkanDescriptorSets.h b/Engine/Source/Runtime/VulkanRHI/Private/VulkanDescriptorSets.h
--- a/Engine/Source/Runtime/VulkanRHI/Private/VulkanDescriptorSets.h
+++ b/Engine/Source/Runtime/VulkanRHI/Private/VulkanDescriptorSets.h
@@ -37,84 +37,6 @@ struct FUniformBufferGatherInfo
 };
 
 
-// Information for remapping descriptor sets when combining layouts
-struct FDescriptorSetRemappingInfo
-{
-	struct FStageInfo
-	{
-		TArray<VkDescriptorType>	Types;
-		uint32						PackedGlobalsSize = 0;
-		uint32						NumBoundUniformBuffers = 0;
-		uint16						NumImageInfos = 0;
-		uint16						NumBufferInfos = 0;
-		uint16						NumAccelerationStructures = 0;
-
-		inline bool IsEmpty() const
-		{
-			if (Types.Num() != 0)
-			{
-				return false;
-			}
-
-			if (PackedGlobalsSize != 0)
-			{
-				return false;
-			}
-
-			if (NumBoundUniformBuffers != 0)
-			{
-				return false;
-			}
-
-			return true;
-		}
-	};
-	TStaticArray<FStageInfo, ShaderStage::NumStages>	StageInfos;
-
-	inline bool operator==(const FDescriptorSetRemappingInfo& In) const
-	{
-		for (uint32 StageInfosIndex = 0; StageInfosIndex < ShaderStage::NumStages; ++StageInfosIndex)
-		{
-			if (StageInfos[StageInfosIndex].PackedGlobalsSize != In.StageInfos[StageInfosIndex].PackedGlobalsSize ||
-				StageInfos[StageInfosIndex].NumBoundUniformBuffers != In.StageInfos[StageInfosIndex].NumBoundUniformBuffers ||
-				StageInfos[StageInfosIndex].NumBufferInfos != In.StageInfos[StageInfosIndex].NumBufferInfos ||
-				StageInfos[StageInfosIndex].NumImageInfos != In.StageInfos[StageInfosIndex].NumImageInfos ||
-				StageInfos[StageInfosIndex].NumAccelerationStructures != In.StageInfos[StageInfosIndex].NumAccelerationStructures ||
-				StageInfos[StageInfosIndex].Types.Num() != In.StageInfos[StageInfosIndex].Types.Num() ||
-				FMemory::Memcmp(StageInfos[StageInfosIndex].Types.GetData(), In.StageInfos[StageInfosIndex].Types.GetData(), sizeof(VkDescriptorType) * StageInfos[StageInfosIndex].Types.Num()))
-			{
-				return false;
-			}
-		}
-
-		return true;
-	}
-
-	inline bool operator != (const FDescriptorSetRemappingInfo& In) const
-	{
-		return !(*this == In);
-	}
-
-	inline bool IsEmpty() const
-	{
-		for (int32 Index = 0; Index < ShaderStage::NumStages; ++Index)
-		{
-			const FStageInfo& StageInfo = StageInfos[Index];
-
-			if (StageInfo.Types.Num() || StageInfo.NumBufferInfos || StageInfo.NumImageInfos || StageInfo.NumAccelerationStructures)
-			{
-				return false;
-			}
-
-			if (!StageInfo.IsEmpty())
-			{
-				return false;
-			}
-		}
-		return true;
-	}
-};
-
 // Information for the layout of descriptor sets; does not hold runtime objects
 class FVulkanDescriptorSetsLayoutInfo
 {
@@ -231,7 +153,7 @@ public:
 			}
 		}
 
-		if (RemappingInfo != In.RemappingInfo)
+		if (StageInfos != In.StageInfos)
 		{
 			return false;
 		}
@@ -245,7 +167,7 @@ public:
 		Hash = Info.Hash;
 		TypesUsageID = Info.TypesUsageID;
 		SetLayouts = Info.SetLayouts;
-		RemappingInfo = Info.RemappingInfo;
+		StageInfos = Info.StageInfos;
 	}
 
 	inline const TMap<VkDescriptorType, uint32>& GetLayoutTypes() const
@@ -263,6 +185,43 @@ public:
 		return GetTypesUsed(VK_DESCRIPTOR_TYPE_INPUT_ATTACHMENT) > 0;
 	}
 
+	struct FStageInfo
+	{
+		TArray<VkDescriptorType>	Types;
+		uint32						PackedGlobalsSize = 0;
+		uint32						NumBoundUniformBuffers = 0;
+		uint16						NumImageInfos = 0;
+		uint16						NumBufferInfos = 0;
+		uint16						NumAccelerationStructures = 0;
+
+		inline bool IsEmpty() const
+		{
+			if ((Types.Num() != 0) || (PackedGlobalsSize != 0) || (NumBoundUniformBuffers != 0))
+			{
+				return false;
+			}
+
+			return true;
+		}
+
+		inline bool operator==(const FStageInfo& In) const
+		{
+			if (PackedGlobalsSize != In.PackedGlobalsSize ||
+				NumBoundUniformBuffers != In.NumBoundUniformBuffers ||
+				NumBufferInfos != In.NumBufferInfos ||
+				NumImageInfos != In.NumImageInfos ||
+				NumAccelerationStructures != In.NumAccelerationStructures ||
+				Types.Num() != In.Types.Num() ||
+				FMemory::Memcmp(Types.GetData(), In.Types.GetData(), Types.NumBytes()))
+			{
+				return false;
+			}
+
+			return true;
+		}
+	};
+	TStaticArray<FStageInfo, ShaderStage::NumStages> StageInfos;
+
 protected:
 	TMap<VkDescriptorType, uint32> LayoutTypes;
 	TArray<FSetLayout> SetLayouts;
@@ -277,8 +236,6 @@ protected:
 
 	void AddDescriptor(int32 DescriptorSetIndex, const VkDescriptorSetLayoutBinding& Descriptor);
 
-	FDescriptorSetRemappingInfo	RemappingInfo;
-
 	friend class FVulkanPipelineStateCacheManager;
 	friend class FVulkanCommonPipelineDescriptorState;
 	friend class FVulkanLayout;
diff --git a/Engine/Source/Runtime/VulkanRHI/Private/VulkanPipelineState.cpp b/Engine/Source/Runtime/VulkanRHI/Private/VulkanPipelineState.cpp
--- a/Engine/Source/Runtime/VulkanRHI/Private/VulkanPipelineState.cpp
+++ b/Engine/Source/Runtime/VulkanRHI/Private/VulkanPipelineState.cpp
@@ -65,7 +65,7 @@ void FVulkanCommonPipelineDescriptorState::CreateDescriptorWriteInfos()
 
 	for (uint32 Set = 0; Set < MaxNumSets; ++Set)
 	{
-		const FDescriptorSetRemappingInfo::FStageInfo& StageInfo = DescriptorSetsLayout->RemappingInfo.StageInfos[Set];
+		const FVulkanDescriptorSetsLayoutInfo::FStageInfo& StageInfo = DescriptorSetsLayout->StageInfos[Set];
 		if (!StageInfo.Types.Num())
 		{
 			continue;
@@ -111,7 +111,7 @@ void FVulkanCommonPipelineDescriptorState::CreateDescriptorWriteInfos()
 
 	for (uint32 Set = 0; Set < MaxNumSets; ++Set)
 	{
-		const FDescriptorSetRemappingInfo::FStageInfo& StageInfo = DescriptorSetsLayout->RemappingInfo.StageInfos[Set];
+		const FVulkanDescriptorSetsLayoutInfo::FStageInfo& StageInfo = DescriptorSetsLayout->StageInfos[Set];
 		if (!StageInfo.Types.Num())
 		{
 			continue;
diff --git a/Engine/Source/Runtime/VulkanRHI/Private/VulkanRHI.cpp b/Engine/Source/Runtime/VulkanRHI/Private/VulkanRHI.cpp
--- a/Engine/Source/Runtime/VulkanRHI/Private/VulkanRHI.cpp
+++ b/Engine/Source/Runtime/VulkanRHI/Private/VulkanRHI.cpp
@@ -1690,7 +1690,7 @@ void FVulkanDescriptorSetsLayoutInfo::AddDescriptor(int32 DescriptorSetIndex, co
 	VkDescriptorSetLayoutBinding* Binding = new(DescSetLayout.LayoutBindings) VkDescriptorSetLayoutBinding;
 	*Binding = Descriptor;
 
-	const FDescriptorSetRemappingInfo::FStageInfo& SetInfo = RemappingInfo.StageInfos[DescriptorSetIndex];
+	const FStageInfo& SetInfo = StageInfos[DescriptorSetIndex];
 	check(SetInfo.Types[Descriptor.binding] == Descriptor.descriptorType);
 	switch (Descriptor.descriptorType)
 	{
@@ -1699,15 +1699,15 @@ void FVulkanDescriptorSetsLayoutInfo::AddDescriptor(int32 DescriptorSetIndex, co
 	case VK_DESCRIPTOR_TYPE_SAMPLED_IMAGE:
 	case VK_DESCRIPTOR_TYPE_STORAGE_IMAGE:
 	case VK_DESCRIPTOR_TYPE_INPUT_ATTACHMENT:
-		IncrementChecked(RemappingInfo.StageInfos[DescriptorSetIndex].NumImageInfos);
+		IncrementChecked(StageInfos[DescriptorSetIndex].NumImageInfos);
 		break;
 	case VK_DESCRIPTOR_TYPE_UNIFORM_BUFFER:
 	case VK_DESCRIPTOR_TYPE_STORAGE_BUFFER:
 	case VK_DESCRIPTOR_TYPE_UNIFORM_BUFFER_DYNAMIC:
-		IncrementChecked(RemappingInfo.StageInfos[DescriptorSetIndex].NumBufferInfos);
+		IncrementChecked(StageInfos[DescriptorSetIndex].NumBufferInfos);
 		break;
 	case VK_DESCRIPTOR_TYPE_ACCELERATION_STRUCTURE_KHR:
-		IncrementChecked(RemappingInfo.StageInfos[DescriptorSetIndex].NumAccelerationStructures);
+		IncrementChecked(StageInfos[DescriptorSetIndex].NumAccelerationStructures);
 		break;
 	case VK_DESCRIPTOR_TYPE_STORAGE_TEXEL_BUFFER:
 	case VK_DESCRIPTOR_TYPE_UNIFORM_TEXEL_BUFFER:
@@ -1731,7 +1731,7 @@ void FVulkanDescriptorSetsLayoutInfo::GenerateHash(const TArrayView<FRHISamplerS
 
 	for (uint32 RemapingIndex = 0; RemapingIndex < ShaderStage::NumStages; ++RemapingIndex)
 	{
-		const FDescriptorSetRemappingInfo::FStageInfo& StageInfo = RemappingInfo.StageInfos[RemapingIndex];
+		const FStageInfo& StageInfo = StageInfos[RemapingIndex];
 
 		Hash = FCrc::TypeCrc32(StageInfo.PackedGlobalsSize, Hash);
 		Hash = FCrc::TypeCrc32(StageInfo.NumBoundUniformBuffers, Hash);
diff --git a/Engine/Source/Runtime/VulkanRHI/Private/VulkanShaders.cpp b/Engine/Source/Runtime/VulkanRHI/Private/VulkanShaders.cpp
--- a/Engine/Source/Runtime/VulkanRHI/Private/VulkanShaders.cpp
+++ b/Engine/Source/Runtime/VulkanRHI/Private/VulkanShaders.cpp
@@ -714,8 +714,6 @@ void FVulkanDescriptorSetsLayoutInfo::ProcessBindingsForStage(VkShaderStageFlagB
 template<bool bIsCompute>
 void FVulkanDescriptorSetsLayoutInfo::FinalizeBindings(const FVulkanDevice& Device, const FUniformBufferGatherInfo& UBGatherInfo, const TArrayView<FRHISamplerState*>& ImmutableSamplers)
 {
-	checkSlow(RemappingInfo.IsEmpty());
-
 	// We'll be reusing this struct
 	VkDescriptorSetLayoutBinding Binding;
 	FMemory::Memzero(Binding);
@@ -728,9 +726,11 @@ void FVulkanDescriptorSetsLayoutInfo::FinalizeBindings(const FVulkanDevice& Devi
 	int32 CurrentImmutableSampler = 0;
 	for (int32 Stage = 0; Stage < (bIsCompute ? 1 : ShaderStage::NumStages); ++Stage)
 	{
+		checkSlow(StageInfos[Stage].IsEmpty());
+
 		if (const FVulkanShaderHeader* ShaderHeader = UBGatherInfo.CodeHeaders[Stage])
 		{
-			FDescriptorSetRemappingInfo::FStageInfo& StageInfo = RemappingInfo.StageInfos[Stage];
+			FStageInfo& StageInfo = StageInfos[Stage];
 
 			Binding.stageFlags = UEFrequencyToVKStageBit(bIsCompute ? SF_Compute : ShaderStage::GetFrequencyForGfxStage((ShaderStage::EStage)Stage));
 
