VulkanRHI: Change the shader group type when an intersection shader is present

#rnx

diff --git a/Engine/Source/Runtime/VulkanRHI/Private/VulkanRayTracing.cpp b/Engine/Source/Runtime/VulkanRHI/Private/VulkanRayTracing.cpp
--- a/Engine/Source/Runtime/VulkanRHI/Private/VulkanRayTracing.cpp
+++ b/Engine/Source/Runtime/VulkanRHI/Private/VulkanRayTracing.cpp
@@ -1537,6 +1537,9 @@ FVulkanRayTracingPipelineState::FVulkanRayTracingPipelineState(FVulkanDevice* co
 			ShaderStage.stage = VK_SHADER_STAGE_INTERSECTION_BIT_KHR;
 			ShaderStage.pName = "main_00000000_00000000"; // :todo-jn: patch in the size_crc
 			ShaderGroup.intersectionShader = ShaderStages.Add(ShaderStage);
+
+			// Switch the shader group type given the presence of an intersection shader
+			ShaderGroup.type = VK_RAY_TRACING_SHADER_GROUP_TYPE_PROCEDURAL_HIT_GROUP_KHR;
 		}
 		else
 		{
