VulkanRHI:  Add support for VK_EXT_load_store_op_none.  Use it on unused aspect of depth attachments with an undefined layout.

#rnx

diff --git a/Engine/Source/Runtime/VulkanRHI/Private/VulkanDevice.h b/Engine/Source/Runtime/VulkanRHI/Private/VulkanDevice.h
--- a/Engine/Source/Runtime/VulkanRHI/Private/VulkanDevice.h
+++ b/Engine/Source/Runtime/VulkanRHI/Private/VulkanDevice.h
@@ -55,6 +55,7 @@ struct FOptionalVulkanDeviceExtensions
 			uint64 HasEXTCalibratedTimestamps : 1;
 			uint64 HasEXTDescriptorBuffer : 1;
 			uint64 HasEXTDeviceFault : 1;
+			uint64 HasEXTLoadStoreOpNone : 1;
 
 			// Vendor specific
 			uint64 HasAMDBufferMarker : 1;
diff --git a/Engine/Source/Runtime/VulkanRHI/Private/VulkanExtensions.cpp b/Engine/Source/Runtime/VulkanRHI/Private/VulkanExtensions.cpp
--- a/Engine/Source/Runtime/VulkanRHI/Private/VulkanExtensions.cpp
+++ b/Engine/Source/Runtime/VulkanRHI/Private/VulkanExtensions.cpp
@@ -1468,6 +1468,7 @@ FVulkanDeviceExtensionArray FVulkanDeviceExtension::GetUESupportedDeviceExtensio
 	ADD_SIMPLE_EXTENSION(VK_KHR_SHADER_FLOAT_CONTROLS_EXTENSION_NAME,	 VULKAN_EXTENSION_ENABLED,             VK_API_VERSION_1_2,            DEVICE_EXT_FLAG_SETTER(HasShaderFloatControls));
 	ADD_SIMPLE_EXTENSION(VK_KHR_IMAGE_FORMAT_LIST_EXTENSION_NAME,        VULKAN_EXTENSION_ENABLED,             VK_API_VERSION_1_2,            DEVICE_EXT_FLAG_SETTER(HasKHRImageFormatList));
 	ADD_SIMPLE_EXTENSION(VK_EXT_VALIDATION_CACHE_EXTENSION_NAME,         VULKAN_SUPPORTS_VALIDATION_CACHE,     VULKAN_EXTENSION_NOT_PROMOTED, DEVICE_EXT_FLAG_SETTER(HasEXTValidationCache));
+	ADD_SIMPLE_EXTENSION(VK_EXT_LOAD_STORE_OP_NONE_EXTENSION_NAME,       VULKAN_EXTENSION_ENABLED,             VULKAN_EXTENSION_NOT_PROMOTED, DEVICE_EXT_FLAG_SETTER(HasEXTLoadStoreOpNone));
 	ADD_SIMPLE_EXTENSION(VK_QCOM_RENDER_PASS_SHADER_RESOLVE_EXTENSION_NAME, VULKAN_SUPPORTS_QCOM_RENDERPASS_SHADER_RESOLVE, VULKAN_EXTENSION_NOT_PROMOTED, DEVICE_EXT_FLAG_SETTER(HasQcomRenderPassShaderResolve));
 
 	// Externally activated extensions (supported by the engine, but enabled externally by plugin or other) :
diff --git a/Engine/Source/Runtime/VulkanRHI/Private/VulkanRenderTarget.cpp b/Engine/Source/Runtime/VulkanRHI/Private/VulkanRenderTarget.cpp
--- a/Engine/Source/Runtime/VulkanRHI/Private/VulkanRenderTarget.cpp
+++ b/Engine/Source/Runtime/VulkanRHI/Private/VulkanRenderTarget.cpp
@@ -1125,11 +1125,27 @@ FVulkanRenderTargetLayout::FVulkanRenderTargetLayout(FVulkanDevice& InDevice, co
 		// happen from UNDEFINED anyhow.
 		if (CurrentDepthLayout == VK_IMAGE_LAYOUT_UNDEFINED)
 		{
+			// Unused image aspects with a LoadOp but undefined layout should just remain untouched
+			if (!RPInfo.DepthStencilRenderTarget.ExclusiveDepthStencil.IsUsingDepth() &&
+				InDevice.GetOptionalExtensions().HasEXTLoadStoreOpNone &&
+				(CurrDesc.loadOp == VK_ATTACHMENT_LOAD_OP_LOAD))
+			{
+				CurrDesc.loadOp = VK_ATTACHMENT_LOAD_OP_NONE_KHR;
+			}
+
 			check(CurrDesc.storeOp == VK_ATTACHMENT_STORE_OP_DONT_CARE);
 			CurrDesc.finalLayout = VK_IMAGE_LAYOUT_ATTACHMENT_OPTIMAL;
 		}
 		if (CurrentStencilLayout == VK_IMAGE_LAYOUT_UNDEFINED)
 		{
+			// Unused image aspects with a LoadOp but undefined layout should just remain untouched
+			if (!RPInfo.DepthStencilRenderTarget.ExclusiveDepthStencil.IsUsingStencil() &&
+				InDevice.GetOptionalExtensions().HasEXTLoadStoreOpNone &&
+				(CurrDesc.stencilLoadOp == VK_ATTACHMENT_LOAD_OP_LOAD))
+			{
+				CurrDesc.stencilLoadOp = VK_ATTACHMENT_LOAD_OP_NONE_KHR;
+			}
+
 			check(CurrDesc.stencilStoreOp == VK_ATTACHMENT_STORE_OP_DONT_CARE);
 			StencilDesc.stencilFinalLayout = VK_IMAGE_LAYOUT_ATTACHMENT_OPTIMAL;
 		}
