VulkanRHI: support VK_NV_ray_tracing_validation if available when validation is enabled.

#rnx

diff --git a/Engine/Source/Runtime/VulkanRHI/Private/VulkanExtensions.cpp b/Engine/Source/Runtime/VulkanRHI/Private/VulkanExtensions.cpp
--- a/Engine/Source/Runtime/VulkanRHI/Private/VulkanExtensions.cpp
+++ b/Engine/Source/Runtime/VulkanRHI/Private/VulkanExtensions.cpp
@@ -1049,6 +1049,44 @@ private:
 	VkDeviceDiagnosticsConfigCreateInfoNV DeviceDiagnosticsConfigCreateInfoNV;
 };
 
+// ***** VK_NV_ray_tracing_validation (vendor)
+class FVulkanNVRayTracingValidationExtension : public FVulkanDeviceExtension
+{
+public:
+
+	FVulkanNVRayTracingValidationExtension(FVulkanDevice* InDevice)
+		: FVulkanDeviceExtension(InDevice, VK_NV_RAY_TRACING_VALIDATION_EXTENSION_NAME, VULKAN_EXTENSION_ENABLED)
+	{
+#if VULKAN_HAS_DEBUGGING_ENABLED
+		bEnabledInCode = bEnabledInCode && (GValidationCvar.GetValueOnAnyThread() > 0) && UE::RHICore::AllowVendorDevice();
+#else
+		bEnabledInCode = false;
+#endif
+	}
+
+	virtual void PrePhysicalDeviceFeatures(VkPhysicalDeviceFeatures2KHR& PhysicalDeviceFeatures2) override final
+	{
+		ZeroVulkanStruct(RayTracingValidationFeaturesNV, VK_STRUCTURE_TYPE_PHYSICAL_DEVICE_RAY_TRACING_VALIDATION_FEATURES_NV);
+		AddToPNext(PhysicalDeviceFeatures2, RayTracingValidationFeaturesNV);
+	}
+
+	virtual void PostPhysicalDeviceFeatures(FOptionalVulkanDeviceExtensions& ExtensionFlags) override final
+	{
+		bRequirementsPassed = (RayTracingValidationFeaturesNV.rayTracingValidation == VK_TRUE);
+	}
+
+	virtual void PreCreateDevice(VkDeviceCreateInfo& DeviceCreateInfo) override final
+	{
+		if (bRequirementsPassed)
+		{
+			AddToPNext(DeviceCreateInfo, RayTracingValidationFeaturesNV);
+		}
+	}
+
+private:
+	VkPhysicalDeviceRayTracingValidationFeaturesNV RayTracingValidationFeaturesNV;
+};
+
 // ***** VK_EXT_device_fault
 class FVulkanEXTDeviceFaultExtension : public FVulkanDeviceExtension
 {
@@ -1514,6 +1552,7 @@ FVulkanDeviceExtensionArray FVulkanDeviceExtension::GetUESupportedDeviceExtensio
 	ADD_CUSTOM_EXTENSION(FVulkanAMDBufferMarkerExtension);
 	ADD_CUSTOM_EXTENSION(FVulkanNVDeviceDiagnosticCheckpointsExtension);
 	ADD_CUSTOM_EXTENSION(FVulkanNVDeviceDiagnosticConfigExtension);
+	ADD_CUSTOM_EXTENSION(FVulkanNVRayTracingValidationExtension);
 
 	// Add in platform specific extensions
 	FVulkanPlatform::GetDeviceExtensions(InDevice, OutUEDeviceExtensions);
