VulkanRHI: Set the record stride to 0 when bAllowHitGroupIndexing is false to force the use of the first record (prevent oob SBT access).  Avoid assert during very long frames on startup by delaying clean up of command buffers that contain unresolved compaction queries.

#rnx

diff --git a/Engine/Source/Runtime/VulkanRHI/Private/VulkanCommandBuffer.cpp b/Engine/Source/Runtime/VulkanRHI/Private/VulkanCommandBuffer.cpp
--- a/Engine/Source/Runtime/VulkanRHI/Private/VulkanCommandBuffer.cpp
+++ b/Engine/Source/Runtime/VulkanRHI/Private/VulkanCommandBuffer.cpp
@@ -873,14 +873,16 @@ void FVulkanCommandBufferPool::FreeUnusedCmdBuffers(FVulkanQueue* InQueue, bool
 		FVulkanCmdBuffer* CmdBuffer = CmdBuffers[Index];
 		if (CmdBuffer != LastSubmittedCmdBuffer &&
 			(CmdBuffer->State == FVulkanCmdBuffer::EState::ReadyForBegin || CmdBuffer->State == FVulkanCmdBuffer::EState::NeedReset) &&
-			(CurrentTime - CmdBuffer->SubmittedTime) > CMD_BUFFER_TIME_TO_WAIT_BEFORE_DELETING)
+			((CurrentTime - CmdBuffer->SubmittedTime) > CMD_BUFFER_TIME_TO_WAIT_BEFORE_DELETING))
 		{
-			DeferredDeletionQueue.OnCmdBufferDeleted(CmdBuffer);
-			if (Device->GetRayTracingCompactionRequestHandler())
+			// Skip command buffer that contain unresolved compaction queries
+			if (Device->GetRayTracingCompactionRequestHandler() && Device->GetRayTracingCompactionRequestHandler()->IsUsingCmdBuffer(CmdBuffer))
 			{
-				Device->GetRayTracingCompactionRequestHandler()->OnCmdBufferDeleted(CmdBuffer);
+				continue;
 			}
 
+			DeferredDeletionQueue.OnCmdBufferDeleted(CmdBuffer);
+
 			CmdBuffer->FreeMemory();
 			CmdBuffers.RemoveAtSwap(Index, EAllowShrinking::No);
 			FreeCmdBuffers.Add(CmdBuffer);
diff --git a/Engine/Source/Runtime/VulkanRHI/Private/VulkanRayTracing.cpp b/Engine/Source/Runtime/VulkanRHI/Private/VulkanRayTracing.cpp
--- a/Engine/Source/Runtime/VulkanRHI/Private/VulkanRayTracing.cpp
+++ b/Engine/Source/Runtime/VulkanRHI/Private/VulkanRayTracing.cpp
@@ -942,8 +942,18 @@ FVulkanRayTracingShaderTable::FVulkanRayTracingShaderTable(FVulkanDevice* Device
 
 		if (Alloc.HandleCount > 0)
 		{
-			Alloc.Region.stride = InUseLocalRecord ? (uint32)GVulkanRayTracingMaxShaderGroupStride : HandleSizeAligned;
-			Alloc.Region.size = Alloc.HandleCount * Alloc.Region.stride;
+			if (InUseLocalRecord)
+			{
+				Alloc.Region.stride = (Alloc.HandleCount > 1) ? (uint32)GVulkanRayTracingMaxShaderGroupStride : 0;
+				Alloc.Region.size = Alloc.HandleCount * (uint32)GVulkanRayTracingMaxShaderGroupStride;
+
+			}
+			else
+			{
+				checkSlow(InHandleCount == 1);
+				Alloc.Region.stride = HandleSizeAligned;
+				Alloc.Region.size = HandleSizeAligned;
+			}
 
 			// Host buffer
 			Alloc.HostBuffer.SetNumUninitialized(Alloc.Region.size);
@@ -1013,6 +1023,7 @@ const VkStridedDeviceAddressRegionKHR* FVulkanRayTracingShaderTable::GetRegion(E
 void FVulkanRayTracingShaderTable::SetSlot(EShaderFrequency Frequency, uint32 DstSlot, uint32 SrcHandleIndex, TConstArrayView<uint8> SrcHandleData)
 {
 	FVulkanShaderTableAllocation& Alloc = GetAlloc(Frequency);
+	checkf((DstSlot == 0) || (Alloc.Region.stride != 0), TEXT("Attempting to index a record in a region without stride"));
 	FMemory::Memcpy(&Alloc.HostBuffer[DstSlot * Alloc.Region.stride], &SrcHandleData[SrcHandleIndex * HandleSize], HandleSize);
 	Alloc.bIsDirty = true;
 }
@@ -1030,7 +1041,8 @@ void FVulkanRayTracingShaderTable::SetLocalShaderParameters(EShaderFrequency Fre
 
 	checkfSlow(OffsetWithinRecord % 4 == 0, TEXT("SBT record parameters must be written on DWORD-aligned boundary"));
 	checkfSlow(InDataSize % 4 == 0, TEXT("SBT record parameters must be DWORD-aligned"));
-	checkf(OffsetWithinRecord + InDataSize <= Alloc.Region.stride, TEXT("SBT record write request is out of bounds"));
+	checkf(OffsetWithinRecord + InDataSize <= Alloc.Region.stride ? Alloc.Region.stride : Alloc.Region.size, TEXT("SBT record write request is out of bounds"));
+	checkf((RecordIndex == 0) || (Alloc.Region.stride != 0), TEXT("Attempting to index a record in a region without stride"));
 
 	const uint32 WriteOffset = HandleSizeAligned + (Alloc.Region.stride * RecordIndex) + OffsetWithinRecord;
 	FMemory::Memcpy(&Alloc.HostBuffer[WriteOffset], InData, InDataSize);
@@ -1924,22 +1936,13 @@ void FVulkanRayTracingCompactionRequestHandler::Update(FVulkanCommandListContext
 	}
 }
 
-void FVulkanRayTracingCompactionRequestHandler::OnCmdBufferDeleted(FVulkanCmdBuffer* DeletedCmdBuffer)
+bool FVulkanRayTracingCompactionRequestHandler::IsUsingCmdBuffer(FVulkanCmdBuffer* CmdBuffer)
 {
 	FScopeLock Lock(&CS);
-
-	// If we are deleting it, it means all its commands went through and we don't need to wait on it anymore
-	if (DeletedCmdBuffer == ActiveRequestsCmdBuffer)
-	{
-		ActiveRequestsCmdBuffer = nullptr;
-		ActiveRequestsFenceCounter = MAX_uint64;
-	}
-
-	check(DeletedCmdBuffer != QueryPool->CmdBuffer);
+	return (CmdBuffer == ActiveRequestsCmdBuffer) || (CmdBuffer == QueryPool->CmdBuffer);
 }
 
 
-
 static FVulkanPipelineBarrier SetRayGenResources(FVulkanDevice* Device, FVulkanCmdBuffer* const CmdBuffer, const FRayTracingShaderBindings& InGlobalResourceBindings, FVulkanRayTracingShaderTable* ShaderTable)
 {
 	TArray<const FVulkanUniformBuffer*> UniformBuffers;
diff --git a/Engine/Source/Runtime/VulkanRHI/Private/VulkanRayTracing.h b/Engine/Source/Runtime/VulkanRHI/Private/VulkanRayTracing.h
--- a/Engine/Source/Runtime/VulkanRHI/Private/VulkanRayTracing.h
+++ b/Engine/Source/Runtime/VulkanRHI/Private/VulkanRayTracing.h
@@ -348,7 +348,7 @@ public:
 
 	void Update(FVulkanCommandListContext& InCommandContext);
 
-	void OnCmdBufferDeleted(FVulkanCmdBuffer* DeletedCmdBuffer);
+	bool IsUsingCmdBuffer(FVulkanCmdBuffer* CmdBuffer);
 
 private:
 
