VulkanRHI: Fix bad layout transition in CopyImageToBackbuffer that wiped image contents on AMD.

#rnx

diff --git a/Engine/Source/Runtime/VulkanRHI/Private/VulkanViewport.cpp b/Engine/Source/Runtime/VulkanRHI/Private/VulkanViewport.cpp
--- a/Engine/Source/Runtime/VulkanRHI/Private/VulkanViewport.cpp
+++ b/Engine/Source/Runtime/VulkanRHI/Private/VulkanViewport.cpp
@@ -895,9 +895,10 @@ inline static void CopyImageToBackBuffer(FVulkanCommandListContext* Context, FVu
 
 	{
 		FVulkanPipelineBarrier Barrier;
-		if (!SrcImageLayout || SrcImageLayout->GetSubresLayout(0, 0, VK_IMAGE_ASPECT_COLOR_BIT) != SrcSurfaceLayout)
+		const VkImageLayout PreviousSrcSurfaceLayout = SrcImageLayout ? SrcImageLayout->GetSubresLayout(0, 0, VK_IMAGE_ASPECT_COLOR_BIT) : VK_IMAGE_LAYOUT_UNDEFINED;
+		if (PreviousSrcSurfaceLayout != SrcSurfaceLayout)
 		{
-			Barrier.AddImageLayoutTransition(SrcSurface.Image, VK_IMAGE_LAYOUT_UNDEFINED, SrcSurfaceLayout, FVulkanPipelineBarrier::MakeSubresourceRange(VK_IMAGE_ASPECT_COLOR_BIT, 0, 1));
+			Barrier.AddImageLayoutTransition(SrcSurface.Image, PreviousSrcSurfaceLayout, SrcSurfaceLayout, FVulkanPipelineBarrier::MakeSubresourceRange(VK_IMAGE_ASPECT_COLOR_BIT, 0, 1));
 			LayoutManager.SetFullLayout(SrcSurface, SrcSurfaceLayout);
 		}
 		Barrier.AddImageLayoutTransition(DstSurface.Image, VK_IMAGE_LAYOUT_PRESENT_SRC_KHR, DstSurfaceLayout, FVulkanPipelineBarrier::MakeSubresourceRange(VK_IMAGE_ASPECT_COLOR_BIT, 0, 1));
