PR #12155: Fix ray query on devices without pipeline

#rnx

diff --git a/Engine/Source/Runtime/VulkanRHI/Private/VulkanBarriers.cpp b/Engine/Source/Runtime/VulkanRHI/Private/VulkanBarriers.cpp
--- a/Engine/Source/Runtime/VulkanRHI/Private/VulkanBarriers.cpp
+++ b/Engine/Source/Runtime/VulkanRHI/Private/VulkanBarriers.cpp
@@ -281,13 +281,19 @@ static void GetVkStageAndAccessFlags(ERHIAccess RHIAccess, FRHITransitionInfo::E
 
 		case ERHIAccess::BVHRead:
 			// vkrt todo: Finer grain stage flags would be ideal here.
-			StageFlags = VK_PIPELINE_STAGE_ACCELERATION_STRUCTURE_BUILD_BIT_KHR | VK_PIPELINE_STAGE_RAY_TRACING_SHADER_BIT_KHR | VK_PIPELINE_STAGE_COMPUTE_SHADER_BIT;
+			StageFlags = VK_PIPELINE_STAGE_ACCELERATION_STRUCTURE_BUILD_BIT_KHR | VK_PIPELINE_STAGE_COMPUTE_SHADER_BIT;
+			if (GRHISupportsRayTracingShaders) {
+				StageFlags |= VK_PIPELINE_STAGE_RAY_TRACING_SHADER_BIT_KHR;
+			}
 			AccessFlags = VK_ACCESS_ACCELERATION_STRUCTURE_READ_BIT_KHR;
 			return;
 
 		case ERHIAccess::BVHWrite:
 			// vkrt todo: Finer grain stage flags would be ideal here.
-			StageFlags = VK_PIPELINE_STAGE_ACCELERATION_STRUCTURE_BUILD_BIT_KHR | VK_PIPELINE_STAGE_RAY_TRACING_SHADER_BIT_KHR | VK_PIPELINE_STAGE_COMPUTE_SHADER_BIT;
+			StageFlags = VK_PIPELINE_STAGE_ACCELERATION_STRUCTURE_BUILD_BIT_KHR | VK_PIPELINE_STAGE_COMPUTE_SHADER_BIT;
+			if (GRHISupportsRayTracingShaders) {
+				StageFlags |= VK_PIPELINE_STAGE_RAY_TRACING_SHADER_BIT_KHR;
+			}
 			AccessFlags = VK_ACCESS_ACCELERATION_STRUCTURE_WRITE_BIT_KHR;
 			return;
 	}
diff --git a/Engine/Source/Runtime/VulkanRHI/Private/VulkanRHI.cpp b/Engine/Source/Runtime/VulkanRHI/Private/VulkanRHI.cpp
--- a/Engine/Source/Runtime/VulkanRHI/Private/VulkanRHI.cpp
+++ b/Engine/Source/Runtime/VulkanRHI/Private/VulkanRHI.cpp
@@ -926,7 +926,7 @@ void FVulkanDynamicRHI::InitInstance()
 
 		if (GRHISupportsRayTracing)
 		{
-			GRHISupportsRayTracingShaders = RHISupportsRayTracingShaders(GMaxRHIShaderPlatform);
+			GRHISupportsRayTracingShaders = RHISupportsRayTracingShaders(GMaxRHIShaderPlatform) && Device->GetOptionalExtensions().HasRayTracingPipeline;
 			GRHISupportsInlineRayTracing = RHISupportsInlineRayTracing(GMaxRHIShaderPlatform) && Device->GetOptionalExtensions().HasRayQuery;
 
 			GRHIRayTracingAccelerationStructureAlignment = 256; // TODO (currently handled by FVulkanAccelerationStructureBuffer)
diff --git a/Engine/Source/Runtime/VulkanRHI/Private/VulkanRayTracing.cpp b/Engine/Source/Runtime/VulkanRHI/Private/VulkanRayTracing.cpp
--- a/Engine/Source/Runtime/VulkanRHI/Private/VulkanRayTracing.cpp
+++ b/Engine/Source/Runtime/VulkanRHI/Private/VulkanRayTracing.cpp
@@ -1384,10 +1384,15 @@ static FRHIRayTracingShader* GetBuiltInRayTracingShader()
 
 void FVulkanDevice::InitializeRayTracing()
 {
-	const VkPhysicalDeviceRayTracingPipelinePropertiesKHR& RayTracingPipelineProps = GetOptionalExtensionProperties().RayTracingPipelineProps;
-	checkf((uint32)GVulkanRayTracingMaxShaderGroupStride <= RayTracingPipelineProps.maxShaderGroupStride, 
-		TEXT("Specified value for r.Vulkan.RayTracing.MaxShaderGroupStride is too large for this device! It will be capped."));
-	GVulkanRayTracingMaxShaderGroupStride = FMath::Min<VkDeviceSize>(RayTracingPipelineProps.maxShaderGroupStride, GVulkanRayTracingMaxShaderGroupStride);
+	if (GRHISupportsRayTracingShaders)
+	{
+		const VkPhysicalDeviceRayTracingPipelinePropertiesKHR& RayTracingPipelineProps = GetOptionalExtensionProperties().RayTracingPipelineProps;
+		if ((uint32)GVulkanRayTracingMaxShaderGroupStride > RayTracingPipelineProps.maxShaderGroupStride)
+		{
+			UE_LOG(LogRHI, Warning, TEXT("Specified value for r.Vulkan.RayTracing.MaxShaderGroupStride is too large for this device! It will be capped."));
+		}
+		GVulkanRayTracingMaxShaderGroupStride = FMath::Min<VkDeviceSize>(RayTracingPipelineProps.maxShaderGroupStride, GVulkanRayTracingMaxShaderGroupStride);
+	}
 }
 
 // Temporary code to generate dummy UBs to bind when none is provided to prevent bindless code from crashing
