VulkanRHI: appropriately fill descriptor buffer arrays when "-noraytracing" is used and accel structures aren't supported.

#rnx

diff --git a/Engine/Source/Runtime/VulkanRHI/Private/VulkanDescriptorSets.cpp b/Engine/Source/Runtime/VulkanRHI/Private/VulkanDescriptorSets.cpp
--- a/Engine/Source/Runtime/VulkanRHI/Private/VulkanDescriptorSets.cpp
+++ b/Engine/Source/Runtime/VulkanRHI/Private/VulkanDescriptorSets.cpp
@@ -359,6 +359,7 @@ void FVulkanBindlessDescriptorManager::Init()
 	{
 		VkDescriptorSetLayoutCreateInfo EmptyDescriptorSetLayoutCreateInfo;
 		ZeroVulkanStruct(EmptyDescriptorSetLayoutCreateInfo, VK_STRUCTURE_TYPE_DESCRIPTOR_SET_LAYOUT_CREATE_INFO);
+		EmptyDescriptorSetLayoutCreateInfo.flags = VK_DESCRIPTOR_SET_LAYOUT_CREATE_DESCRIPTOR_BUFFER_BIT_EXT;
 		VERIFYVULKANRESULT(VulkanRHI::vkCreateDescriptorSetLayout(DeviceHandle, &EmptyDescriptorSetLayoutCreateInfo, VULKAN_CPU_ALLOCATOR, &EmptyDescriptorSetLayout));
 	}
 
@@ -534,19 +535,21 @@ void FVulkanBindlessDescriptorManager::Init()
 		uint32 TotalResourceDescriptorBufferSize = 0;
 		for (uint32 SetIndex = 0; SetIndex < VulkanBindless::NumBindlessSets; ++SetIndex)
 		{
-			// Skip anything we don't support
-			if (SetIndex == VulkanBindless::BindlessAccelerationStructureSet)
-			{
-				const bool bHasRaytracingExtensions = Device->GetOptionalExtensions().HasRaytracingExtensions();
+			BindlessSetState& State = BindlessSetStates[SetIndex];
 
-				if (!bHasRaytracingExtensions)
-				{
-					continue;
-				}
+			// Create a dummy buffer for acceleration structures when they aren't supported (or ray tracing is disabled)
+			const bool bSupported = (SetIndex != VulkanBindless::BindlessAccelerationStructureSet) || Device->GetOptionalExtensions().HasRaytracingExtensions();
+			if (bSupported)
+			{
+				InitBindlessSetState(GetDescriptorTypeForSetIndex(SetIndex), State);
+			}
+			else
+			{
+				State.DescriptorType = VK_DESCRIPTOR_TYPE_STORAGE_BUFFER;
+				State.DescriptorSize = GetDescriptorTypeSize(Device, VK_DESCRIPTOR_TYPE_STORAGE_BUFFER);
+				State.MaxDescriptorCount = 16;
 			}
 
-			BindlessSetState& State = BindlessSetStates[SetIndex];
-			InitBindlessSetState(GetDescriptorTypeForSetIndex(SetIndex), State);
 			const bool IsSingleUseUniformBufferSet = (SetIndex == VulkanBindless::BindlessSingleUseUniformBufferSet);
 			State.DescriptorSetLayout = IsSingleUseUniformBufferSet ? CreateShaderStageUniformBufferLayout() : CreateDescriptorSetLayout(State);
 			TotalResourceDescriptorBufferSize += CreateDescriptorBuffer(State, BufferBindingInfo[SetIndex], IsSingleUseUniformBufferSet);
