Skip PSO pre-caching for shadows if they are disabled
Remove VertexDecl from PSO hash on GPUs that do not need it

diff --git a/Engine/Source/Runtime/Renderer/Private/ShadowDepthRendering.cpp b/Engine/Source/Runtime/Renderer/Private/ShadowDepthRendering.cpp
--- a/Engine/Source/Runtime/Renderer/Private/ShadowDepthRendering.cpp
+++ b/Engine/Source/Runtime/Renderer/Private/ShadowDepthRendering.cpp
@@ -2189,7 +2189,8 @@ void FShadowDepthPassMeshProcessor::CollectPSOInitializers(const FSceneTexturesC
 {
 	// Early out if possible
 	if (!PreCacheParams.bDefaultMaterial && 
-		(!Material.ShouldCastDynamicShadows() ||
+		(GetShadowQuality() == 0 ||
+		!Material.ShouldCastDynamicShadows() ||
 		!ShouldIncludeDomainInMeshPass(Material.GetMaterialDomain()) ||
 		!ShouldIncludeMaterialInDefaultOpaquePass(Material) ||
 		!EnumHasAnyFlags(MeshSelectionMask, VertexFactoryData.VertexFactoryType->SupportsPrimitiveIdStream() ? EShadowMeshSelection::VSM : EShadowMeshSelection::SM)))
diff --git a/Engine/Source/Runtime/VulkanRHI/Private/VulkanRHI.cpp b/Engine/Source/Runtime/VulkanRHI/Private/VulkanRHI.cpp
--- a/Engine/Source/Runtime/VulkanRHI/Private/VulkanRHI.cpp
+++ b/Engine/Source/Runtime/VulkanRHI/Private/VulkanRHI.cpp
@@ -2101,7 +2101,12 @@ uint64 FVulkanDynamicRHI::RHIComputeStatePrecachePSOHash(const FGraphicsPipeline
 
 	FMemory::Memzero(&HashKey, sizeof(FHashKey));
 
-	HashKey.VertexDeclaration = Initializer.BoundShaderState.VertexDeclarationRHI ? Initializer.BoundShaderState.VertexDeclarationRHI->GetPrecachePSOHash() : 0;
+	// We know for sure that on ARM MALI GPUs vertex decl does not affect PSO
+	const bool bVertexDeclAffectsPSO = (GRHIVendorId != (uint32)EGpuVendorId::Arm);
+	if (bVertexDeclAffectsPSO)
+	{ 
+		HashKey.VertexDeclaration = Initializer.BoundShaderState.VertexDeclarationRHI ? Initializer.BoundShaderState.VertexDeclarationRHI->GetPrecachePSOHash() : 0;
+	}
 	HashKey.VertexShader = Initializer.BoundShaderState.GetVertexShader() ? GetTypeHash(Initializer.BoundShaderState.GetVertexShader()->GetHash()) : 0;
 	HashKey.PixelShader = Initializer.BoundShaderState.GetPixelShader() ? GetTypeHash(Initializer.BoundShaderState.GetPixelShader()->GetHash()) : 0;
 #if PLATFORM_SUPPORTS_GEOMETRY_SHADERS
@@ -2159,7 +2164,6 @@ uint64 FVulkanDynamicRHI::RHIComputePrecachePSOHash(const FGraphicsPipelineState
 	{
 		uint64							StatePrecachePSOHash;
 
-		EPrimitiveType					PrimitiveType;
 		uint32							RenderTargetsEnabled;
 		FGraphicsPipelineStateInitializer::TRenderTargetFormats	RenderTargetFormats;
 		FGraphicsPipelineStateInitializer::TRenderTargetFlags RenderTargetFlags;
@@ -2170,18 +2174,12 @@ uint64 FVulkanDynamicRHI::RHIComputePrecachePSOHash(const FGraphicsPipelineState
 		ESubpassHint					SubpassHint;
 		uint8							SubpassIndex;
 		EConservativeRasterization		ConservativeRasterization;
-		uint8							MultiViewCount;
-		EVRSShadingRate					ShadingRate;
-		bool							bDepthBounds;
-		bool							bHasFragmentDensityAttachment;
-		bool							bAllowVariableRateShading;
 	} HashKey;
 
 	FMemory::Memzero(&HashKey, sizeof(FNonStateHashKey));
 
 	HashKey.StatePrecachePSOHash = StatePrecachePSOHash;
 
-	HashKey.PrimitiveType = Initializer.PrimitiveType;
 	HashKey.RenderTargetsEnabled = Initializer.RenderTargetsEnabled;
 	HashKey.RenderTargetFormats = Initializer.RenderTargetFormats;
 	HashKey.RenderTargetFlags = Initializer.RenderTargetFlags;
@@ -2191,12 +2189,7 @@ uint64 FVulkanDynamicRHI::RHIComputePrecachePSOHash(const FGraphicsPipelineState
 	HashKey.SubpassHint = Initializer.SubpassHint;
 	HashKey.SubpassIndex = Initializer.SubpassIndex;
 	HashKey.ConservativeRasterization = Initializer.ConservativeRasterization;
-	HashKey.MultiViewCount = Initializer.MultiViewCount;
-	HashKey.ShadingRate = Initializer.ShadingRate;
-	HashKey.bDepthBounds = Initializer.bDepthBounds;
-	HashKey.bHasFragmentDensityAttachment = Initializer.bHasFragmentDensityAttachment;
-	HashKey.bAllowVariableRateShading = Initializer.bAllowVariableRateShading;
-
+	
 	// TODO: check if any RT flags actually affect PSO in VK
 	for (ETextureCreateFlags& Flags : HashKey.RenderTargetFlags)
 	{
