Added option to provide callbacks for resource cleanup for VulkanTextures with user provided VkImage

diff --git a/Engine/Source/Runtime/VulkanRHI/Private/VulkanRHI.cpp b/Engine/Source/Runtime/VulkanRHI/Private/VulkanRHI.cpp
--- a/Engine/Source/Runtime/VulkanRHI/Private/VulkanRHI.cpp
+++ b/Engine/Source/Runtime/VulkanRHI/Private/VulkanRHI.cpp
@@ -1603,7 +1603,7 @@ uint64 FVulkanDynamicRHI::RHIGetMinimumAlignmentForBufferBackedSRV(EPixelFormat
 	return Limits.minTexelBufferOffsetAlignment;
 }
 
-FTextureRHIRef FVulkanDynamicRHI::RHICreateTexture2DFromResource(EPixelFormat Format, uint32 SizeX, uint32 SizeY, uint32 NumMips, uint32 NumSamples, VkImage Resource, ETextureCreateFlags Flags, const FClearValueBinding& ClearValueBinding)
+FTextureRHIRef FVulkanDynamicRHI::RHICreateTexture2DFromResource(EPixelFormat Format, uint32 SizeX, uint32 SizeY, uint32 NumMips, uint32 NumSamples, VkImage Resource, ETextureCreateFlags Flags, const FClearValueBinding& ClearValueBinding, const FVulkanRHIExternalImageDeleteCallbackInfo& ExternalImageDeleteCallbackInfo)
 {
 	const FRHITextureCreateDesc Desc =
 		FRHITextureCreateDesc::Create2D(TEXT("VulkanTexture2DFromResource"), SizeX, SizeY, Format)
@@ -1613,7 +1613,7 @@ FTextureRHIRef FVulkanDynamicRHI::RHICreateTexture2DFromResource(EPixelFormat Fo
 		.SetNumSamples(NumSamples)
 		.DetermineInititialState();
 
-	return new FVulkanTexture(*Device, Desc, Resource, false);
+	return new FVulkanTexture(*Device, Desc, Resource, ExternalImageDeleteCallbackInfo);
 }
 
 FTextureRHIRef FVulkanDynamicRHI::RHICreateTexture2DArrayFromResource(EPixelFormat Format, uint32 SizeX, uint32 SizeY, uint32 ArraySize, uint32 NumMips, uint32 NumSamples, VkImage Resource, ETextureCreateFlags Flags, const FClearValueBinding& ClearValueBinding)
@@ -1626,7 +1626,7 @@ FTextureRHIRef FVulkanDynamicRHI::RHICreateTexture2DArrayFromResource(EPixelForm
 		.SetNumSamples(NumSamples)
 		.DetermineInititialState();
 
-	return new FVulkanTexture(*Device, Desc, Resource, false);
+	return new FVulkanTexture(*Device, Desc, Resource, {});
 }
 
 FTextureRHIRef FVulkanDynamicRHI::RHICreateTextureCubeFromResource(EPixelFormat Format, uint32 Size, bool bArray, uint32 ArraySize, uint32 NumMips, VkImage Resource, ETextureCreateFlags Flags, const FClearValueBinding& ClearValueBinding)
@@ -1641,7 +1641,7 @@ FTextureRHIRef FVulkanDynamicRHI::RHICreateTextureCubeFromResource(EPixelFormat
 		.SetNumMips(NumMips)
 		.DetermineInititialState();
 
-	return new FVulkanTexture(*Device, Desc, Resource, false);
+	return new FVulkanTexture(*Device, Desc, Resource, {});
 }
 
 void FVulkanDynamicRHI::RHIAliasTextureResources(FTextureRHIRef& DestTextureRHI, FTextureRHIRef& SrcTextureRHI)
diff --git a/Engine/Source/Runtime/VulkanRHI/Private/VulkanResources.h b/Engine/Source/Runtime/VulkanRHI/Private/VulkanResources.h
--- a/Engine/Source/Runtime/VulkanRHI/Private/VulkanResources.h
+++ b/Engine/Source/Runtime/VulkanRHI/Private/VulkanResources.h
@@ -15,6 +15,7 @@
 #include "VulkanShaderResources.h"
 #include "VulkanMemory.h"
 #include "Misc/ScopeRWLock.h"
+#include "IVulkanDynamicRHI.h"
 
 class FVulkanDevice;
 class FVulkanQueue;
@@ -563,7 +564,7 @@ public:
 
 	// Construct from external resource.
 	// FIXME: HUGE HACK: the bUnused argument is there to disambiguate this overload from the one above when passing nullptr, since nullptr is a valid VkImage. Get rid of this code smell when unifying FVulkanSurface and FVulkanTexture.
-	FVulkanTexture(FVulkanDevice& InDevice, const FRHITextureCreateDesc& InCreateDesc, VkImage InImage, bool bUnused);
+	FVulkanTexture(FVulkanDevice& InDevice, const FRHITextureCreateDesc& InCreateDesc, VkImage InImage, const FVulkanRHIExternalImageDeleteCallbackInfo& InExternalImageDeleteCallbackInfo);
 
 	// Aliasing constructor.
 	FVulkanTexture(FVulkanDevice& InDevice, const FRHITextureCreateDesc& InCreateDesc, FTextureRHIRef& SrcTextureRHI);
@@ -731,6 +732,8 @@ public:
 	VkMemoryPropertyFlags MemProps;
 	VkMemoryRequirements MemoryRequirements;
 
+	FVulkanRHIExternalImageDeleteCallbackInfo ExternalImageDeleteCallbackInfo;
+
 private:
 	void SetInitialImageState(FVulkanCommandListContext& Context, VkImageLayout InitialLayout, bool bClear, const FClearValueBinding& ClearValueBinding, bool bIsTransientResource);
 	void InternalMoveSurface(FVulkanDevice& InDevice, FVulkanCommandListContext& Context, VulkanRHI::FVulkanAllocation& DestAllocation, VkImageLayout OriginalLayout);
diff --git a/Engine/Source/Runtime/VulkanRHI/Private/VulkanTexture.cpp b/Engine/Source/Runtime/VulkanRHI/Private/VulkanTexture.cpp
--- a/Engine/Source/Runtime/VulkanRHI/Private/VulkanTexture.cpp
+++ b/Engine/Source/Runtime/VulkanRHI/Private/VulkanTexture.cpp
@@ -734,6 +734,10 @@ void FVulkanTexture::DestroySurface()
 		else
 		{
 			Image = VK_NULL_HANDLE;
+			if (ExternalImageDeleteCallbackInfo.Function)
+			{
+				ExternalImageDeleteCallbackInfo.Function(ExternalImageDeleteCallbackInfo.UserData);
+			}
 		}
 
 		ImageOwnerType = EImageOwnerType::None;
@@ -1755,7 +1759,7 @@ FVulkanTexture::FVulkanTexture(FRHICommandListBase* RHICmdList, FVulkanDevice& I
 	}
 }
 
-FVulkanTexture::FVulkanTexture(FVulkanDevice& InDevice, const FRHITextureCreateDesc& InCreateDesc, VkImage InImage, bool /*bUnused*/)
+FVulkanTexture::FVulkanTexture(FVulkanDevice& InDevice, const FRHITextureCreateDesc& InCreateDesc, VkImage InImage, const FVulkanRHIExternalImageDeleteCallbackInfo& InExternalImageDeleteCallbackInfo)
 	: FRHITexture(InCreateDesc)
 	, Device(&InDevice)
 	, Image(InImage)
@@ -1763,6 +1767,7 @@ FVulkanTexture::FVulkanTexture(FVulkanDevice& InDevice, const FRHITextureCreateD
 	, StorageFormat(VK_FORMAT_UNDEFINED)
 	, ViewFormat(VK_FORMAT_UNDEFINED)
 	, MemProps(VK_MEMORY_PROPERTY_DEVICE_LOCAL_BIT)
+	, ExternalImageDeleteCallbackInfo(InExternalImageDeleteCallbackInfo)
 	, Tiling(VK_IMAGE_TILING_MAX_ENUM)	// Can be expanded to a per-platform definition
 	, FullAspectMask(0)
 	, PartialAspectMask(0)
diff --git a/Engine/Source/Runtime/VulkanRHI/Private/VulkanViewport.cpp b/Engine/Source/Runtime/VulkanRHI/Private/VulkanViewport.cpp
--- a/Engine/Source/Runtime/VulkanRHI/Private/VulkanViewport.cpp
+++ b/Engine/Source/Runtime/VulkanRHI/Private/VulkanViewport.cpp
@@ -18,7 +18,7 @@
 #include "RHIStaticStates.h"
 
 FVulkanBackBuffer::FVulkanBackBuffer(FVulkanDevice& Device, FVulkanViewport* InViewport, EPixelFormat Format, uint32 SizeX, uint32 SizeY, ETextureCreateFlags UEFlags)
-	: FVulkanTexture(Device, FRHITextureCreateDesc::Create2D(TEXT("FVulkanBackBuffer"), SizeX, SizeY, Format).SetFlags(UEFlags).DetermineInititialState(), VK_NULL_HANDLE, false)
+	: FVulkanTexture(Device, FRHITextureCreateDesc::Create2D(TEXT("FVulkanBackBuffer"), SizeX, SizeY, Format).SetFlags(UEFlags).DetermineInititialState(), VK_NULL_HANDLE, {})
 	, Viewport(InViewport)
 {
 }
diff --git a/Engine/Source/Runtime/VulkanRHI/Public/IVulkanDynamicRHI.h b/Engine/Source/Runtime/VulkanRHI/Public/IVulkanDynamicRHI.h
--- a/Engine/Source/Runtime/VulkanRHI/Public/IVulkanDynamicRHI.h
+++ b/Engine/Source/Runtime/VulkanRHI/Public/IVulkanDynamicRHI.h
@@ -24,6 +24,12 @@ struct FVulkanRHIImageViewInfo
 	ETextureCreateFlags UEFlags;
 };
 
+struct FVulkanRHIExternalImageDeleteCallbackInfo
+{
+	void* UserData;
+	void (*Function)(void* UserData);
+};
+
 struct IVulkanDynamicRHI : public FDynamicRHI
 {
 	virtual ERHIInterfaceType GetInterfaceType() const override { return ERHIInterfaceType::Vulkan; }
@@ -52,7 +58,7 @@ struct IVulkanDynamicRHI : public FDynamicRHI
 	virtual TArray<VkExtensionProperties> RHIGetAllDeviceExtensions(VkPhysicalDevice InPhysicalDevice) const = 0;
 	virtual TArray<FAnsiString> RHIGetLoadedDeviceExtensions() const = 0;
 
-	virtual FTextureRHIRef   RHICreateTexture2DFromResource(EPixelFormat Format, uint32 SizeX, uint32 SizeY, uint32 NumMips, uint32 NumSamples, VkImage Resource, ETextureCreateFlags Flags, const FClearValueBinding& ClearValueBinding = FClearValueBinding::Transparent) = 0;
+	virtual FTextureRHIRef   RHICreateTexture2DFromResource(EPixelFormat Format, uint32 SizeX, uint32 SizeY, uint32 NumMips, uint32 NumSamples, VkImage Resource, ETextureCreateFlags Flags, const FClearValueBinding& ClearValueBinding = FClearValueBinding::Transparent, const FVulkanRHIExternalImageDeleteCallbackInfo& ExternalImageDeleteCallbackInfo = {}) = 0;
 	virtual FTextureRHIRef   RHICreateTexture2DArrayFromResource(EPixelFormat Format, uint32 SizeX, uint32 SizeY, uint32 ArraySize, uint32 NumMips, uint32 NumSamples, VkImage Resource, ETextureCreateFlags Flags, const FClearValueBinding& ClearValueBinding = FClearValueBinding::Transparent) = 0;
 	virtual FTextureRHIRef   RHICreateTextureCubeFromResource(EPixelFormat Format, uint32 Size, bool bArray, uint32 ArraySize, uint32 NumMips, VkImage Resource, ETextureCreateFlags Flags, const FClearValueBinding& ClearValueBinding = FClearValueBinding::Transparent) = 0;
 
diff --git a/Engine/Source/Runtime/VulkanRHI/Public/VulkanDynamicRHI.h b/Engine/Source/Runtime/VulkanRHI/Public/VulkanDynamicRHI.h
--- a/Engine/Source/Runtime/VulkanRHI/Public/VulkanDynamicRHI.h
+++ b/Engine/Source/Runtime/VulkanRHI/Public/VulkanDynamicRHI.h
@@ -74,7 +74,7 @@ public:
 	virtual TArray<VkExtensionProperties> RHIGetAllInstanceExtensions() const final override;
 	virtual TArray<VkExtensionProperties> RHIGetAllDeviceExtensions(VkPhysicalDevice InPhysicalDevice) const final override;
 	virtual TArray<FAnsiString> RHIGetLoadedDeviceExtensions() const final override;
-	virtual FTextureRHIRef RHICreateTexture2DFromResource(EPixelFormat Format, uint32 SizeX, uint32 SizeY, uint32 NumMips, uint32 NumSamples, VkImage Resource, ETextureCreateFlags Flags, const FClearValueBinding& ClearValueBinding = FClearValueBinding::Transparent) final override;
+	virtual FTextureRHIRef RHICreateTexture2DFromResource(EPixelFormat Format, uint32 SizeX, uint32 SizeY, uint32 NumMips, uint32 NumSamples, VkImage Resource, ETextureCreateFlags Flags, const FClearValueBinding& ClearValueBinding = FClearValueBinding::Transparent, const FVulkanRHIExternalImageDeleteCallbackInfo& ExternalImageDeleteCallbackInfo = {}) final override;
 	virtual FTextureRHIRef RHICreateTexture2DArrayFromResource(EPixelFormat Format, uint32 SizeX, uint32 SizeY, uint32 ArraySize, uint32 NumMips, uint32 NumSamples, VkImage Resource, ETextureCreateFlags Flags, const FClearValueBinding& ClearValueBinding = FClearValueBinding::Transparent) final override;
 	virtual FTextureRHIRef RHICreateTextureCubeFromResource(EPixelFormat Format, uint32 Size, bool bArray, uint32 ArraySize, uint32 NumMips, VkImage Resource, ETextureCreateFlags Flags, const FClearValueBinding& ClearValueBinding = FClearValueBinding::Transparent) final override;
 	virtual VkImage RHIGetVkImage(FRHITexture* InTexture) const final override;
