VulkanRHI: fix validation error where a renderpass with ReadOnly depth ref layout was used with pipeline that require DepthWrite.  Match other constructors for FVulkanRenderTargetLayout and always use the same layout for Initial->Ref->Final (the full hash only includes Initial, so we currently can't have two different Ref layouts for the same Initial layout unless we add it to the hash).

#rnx

diff --git a/Engine/Source/Runtime/VulkanRHI/Private/Android/VulkanAndroidPlatform.cpp b/Engine/Source/Runtime/VulkanRHI/Private/Android/VulkanAndroidPlatform.cpp
--- a/Engine/Source/Runtime/VulkanRHI/Private/Android/VulkanAndroidPlatform.cpp
+++ b/Engine/Source/Runtime/VulkanRHI/Private/Android/VulkanAndroidPlatform.cpp
@@ -953,7 +953,7 @@ void HandleDepthStencilAttachmentPNext(const VkAttachmentReference2* Attachment,
 	check(HandledCount == VkStructs.Num());
 }
 
-void PipelineToBinary(FVulkanDevice* Device, const VkGraphicsPipelineCreateInfo* PipelineInfo, FGfxPipelineDesc* GfxEntry, const FVulkanRenderTargetLayout* RTLayout, TArray<char>& MemoryStream)
+void PipelineToBinary(FVulkanDevice* Device, const VkGraphicsPipelineCreateInfo* PipelineInfo, const FGfxPipelineDesc* GfxEntry, const FVulkanRenderTargetLayout* RTLayout, TArray<char>& MemoryStream)
 {
 	static const unsigned int INITIAL_PSO_STREAM_SIZE = 64 * 1024;
 	MemoryStream.Reserve(INITIAL_PSO_STREAM_SIZE);
@@ -1520,7 +1520,7 @@ VkPipelineCache FVulkanAndroidPlatform::PrecompilePSO(
 		const TArrayView<uint8> OptionalPSOCacheData, 
 		FGraphicsPipelineStateInitializer::EPSOPrecacheCompileType PSOCompileType,
 		const VkGraphicsPipelineCreateInfo* PipelineInfo, 
-		FGfxPipelineDesc* GfxEntry, 
+		const FGfxPipelineDesc* GfxEntry, 
 		const FVulkanRenderTargetLayout* RTLayout, 
 		TArrayView<uint32_t> VS, 
 		TArrayView<uint32_t> PS, 
diff --git a/Engine/Source/Runtime/VulkanRHI/Private/Android/VulkanAndroidPlatform.h b/Engine/Source/Runtime/VulkanRHI/Private/Android/VulkanAndroidPlatform.h
--- a/Engine/Source/Runtime/VulkanRHI/Private/Android/VulkanAndroidPlatform.h
+++ b/Engine/Source/Runtime/VulkanRHI/Private/Android/VulkanAndroidPlatform.h
@@ -145,7 +145,7 @@ public:
 	static void DestroySwapchainKHR(VkDevice Device, VkSwapchainKHR Swapchain, const VkAllocationCallbacks* Allocator);
 
 	// handle precompile of PSOs, send to an android specific precompile external process.
-	static VkPipelineCache PrecompilePSO(FVulkanDevice* Device, const TArrayView<uint8> OptionalPSOCacheData, FGraphicsPipelineStateInitializer::EPSOPrecacheCompileType PSOCompileType, const VkGraphicsPipelineCreateInfo* PipelineInfo, FGfxPipelineDesc* GfxEntry, const FVulkanRenderTargetLayout* RTLayout, TArrayView<uint32_t> VS, TArrayView<uint32_t> PS, size_t& AfterSize, FString* FailureMessageOUT = nullptr);
+	static VkPipelineCache PrecompilePSO(FVulkanDevice* Device, const TArrayView<uint8> OptionalPSOCacheData, FGraphicsPipelineStateInitializer::EPSOPrecacheCompileType PSOCompileType, const VkGraphicsPipelineCreateInfo* PipelineInfo, const FGfxPipelineDesc* GfxEntry, const FVulkanRenderTargetLayout* RTLayout, TArrayView<uint32_t> VS, TArrayView<uint32_t> PS, size_t& AfterSize, FString* FailureMessageOUT = nullptr);
 
 	static bool AreRemoteCompileServicesActive();
 	static bool StartRemoteCompileServices(int NumServices);
diff --git a/Engine/Source/Runtime/VulkanRHI/Private/VulkanGenericPlatform.h b/Engine/Source/Runtime/VulkanRHI/Private/VulkanGenericPlatform.h
--- a/Engine/Source/Runtime/VulkanRHI/Private/VulkanGenericPlatform.h
+++ b/Engine/Source/Runtime/VulkanRHI/Private/VulkanGenericPlatform.h
@@ -130,7 +130,7 @@ public:
 	static TArray<FString> GetPSOCacheFilenames();
 
 	// Gives platform a chance to handle precompile of PSOs, returns nullptr if unsupported
-	static VkPipelineCache PrecompilePSO(FVulkanDevice* Device, const uint8* OptionalPSOCacheData,VkGraphicsPipelineCreateInfo* PipelineInfo, FGfxPipelineDesc* GfxEntry, const FVulkanRenderTargetLayout* RTLayout, TArrayView<uint32_t> VS, TArrayView<uint32_t> PS, size_t& AfterSize, FString* FailureMessageOUT = nullptr) { return VK_NULL_HANDLE; }
+	static VkPipelineCache PrecompilePSO(FVulkanDevice* Device, const uint8* OptionalPSOCacheData,VkGraphicsPipelineCreateInfo* PipelineInfo, const FGfxPipelineDesc* GfxEntry, const FVulkanRenderTargetLayout* RTLayout, TArrayView<uint32_t> VS, TArrayView<uint32_t> PS, size_t& AfterSize, FString* FailureMessageOUT = nullptr) { return VK_NULL_HANDLE; }
 
 	// Return VK_FALSE if platform wants to suppress the given debug report from the validation layers, VK_TRUE to print it.
 	static VkBool32 DebugReportFunction(VkDebugReportFlagsEXT MsgFlags, VkDebugReportObjectTypeEXT ObjType, uint64_t SrcObject, size_t Location, int32 MsgCode, const ANSICHAR* LayerPrefix, const ANSICHAR* Msg, void* UserData) { return VK_TRUE; }
diff --git a/Engine/Source/Runtime/VulkanRHI/Private/VulkanPipeline.cpp b/Engine/Source/Runtime/VulkanRHI/Private/VulkanPipeline.cpp
--- a/Engine/Source/Runtime/VulkanRHI/Private/VulkanPipeline.cpp
+++ b/Engine/Source/Runtime/VulkanRHI/Private/VulkanPipeline.cpp
@@ -1194,7 +1194,7 @@ static FString ShaderHashesToString(FVulkanShader* Shaders[ShaderStage::NumStage
 bool FVulkanPipelineStateCacheManager::CreateGfxPipelineFromEntry(FVulkanRHIGraphicsPipelineState* PSO, FVulkanShader* Shaders[ShaderStage::NumStages], FGraphicsPipelineStateInitializer::EPSOPrecacheCompileType PSOCompileType)
 {
 	VkPipeline* Pipeline = &PSO->VulkanPipeline;
-	FGfxPipelineDesc* GfxEntry = &PSO->Desc;
+	const FGfxPipelineDesc* GfxEntry = &PSO->Desc;
 	if (Shaders[ShaderStage::Pixel] == nullptr && !FVulkanPlatform::SupportsNullPixelShader())
 	{
 		Shaders[ShaderStage::Pixel] = ResourceCast(TShaderMapRef<FNULLPS>(GetGlobalShaderMap(GMaxRHIFeatureLevel)).GetPixelShader());
@@ -1408,7 +1408,7 @@ static VkResult CreatePSOWithExternalService(FVulkanDevice* Device, FGraphicsPip
 	size_t AfterSize = 0;
 
 	VkPipelineCache LocalPipelineCache = VK_NULL_HANDLE;
-	FGfxPipelineDesc* GfxEntry = &PSO->Desc;
+	const FGfxPipelineDesc* GfxEntry = &PSO->Desc;
 
 	TArray<uint8> InitialCacheData;
 	bool bSupplyDestPSOCacheData = false; // this can be an optimization, but only if the PSO compile is able to use content from an existing cache.
@@ -1509,7 +1509,7 @@ VkResult FVulkanPipelineStateCacheManager::CreateVKPipeline(FVulkanRHIGraphicsPi
 	VkResult Result = VK_ERROR_INITIALIZATION_FAILED;
 	uint32 PSOSize = 0;
 	bool bWantPSOSize = false;
-	FGfxPipelineDesc* GfxEntry = &PSO->Desc;
+	const FGfxPipelineDesc* GfxEntry = &PSO->Desc;
 	uint64 ShaderHash = 0;
 	bool bValidateServicePSO = false;
 	if (bUseLRU)
@@ -1918,12 +1918,12 @@ void FVulkanPipelineStateCacheManager::CreateGfxEntry(const FGraphicsPipelineSta
 
 
 
-FVulkanRHIGraphicsPipelineState::FVulkanRHIGraphicsPipelineState(FVulkanDevice* Device, const FGraphicsPipelineStateInitializer& PSOInitializer_, FGfxPipelineDesc& Desc, FVulkanPSOKey* VulkanKey)
+FVulkanRHIGraphicsPipelineState::FVulkanRHIGraphicsPipelineState(FVulkanDevice* InDevice, const FGraphicsPipelineStateInitializer& PSOInitializer_, const FGfxPipelineDesc& InDesc, FVulkanPSOKey* VulkanKey)
 	: bIsRegistered(false)
 	, PrimitiveType(PSOInitializer_.PrimitiveType)
 	, VulkanPipeline(0)
-	, Device(Device)
-	, Desc(Desc)
+	, Device(InDevice)
+	, Desc(InDesc)
 	, VulkanKey(VulkanKey->CopyDeep())
 {
 #if !UE_BUILD_SHIPPING
@@ -2080,7 +2080,6 @@ FGraphicsPipelineStateRHIRef FVulkanPipelineStateCacheManager::RHICreateGraphics
 		SCOPE_CYCLE_COUNTER(STAT_VulkanPSOCreationTime);
 		NewPSO = new FVulkanRHIGraphicsPipelineState(Device, Initializer, Desc, &Key);
 		{
-
 			FVulkanLayout* Layout = FindOrAddLayout(DescriptorSetLayoutInfo, true);
 			NewPSO->Layout = Layout;
 			NewPSO->bHasInputAttachments = Layout->GetDescriptorSetsLayout().HasInputAttachments();
diff --git a/Engine/Source/Runtime/VulkanRHI/Private/VulkanPipeline.h b/Engine/Source/Runtime/VulkanRHI/Private/VulkanPipeline.h
--- a/Engine/Source/Runtime/VulkanRHI/Private/VulkanPipeline.h
+++ b/Engine/Source/Runtime/VulkanRHI/Private/VulkanPipeline.h
@@ -693,7 +693,7 @@ public:
 class FVulkanRHIGraphicsPipelineState : public FRHIGraphicsPipelineState
 {
 public:
-	FVulkanRHIGraphicsPipelineState(FVulkanDevice* Device, const FGraphicsPipelineStateInitializer& PSOInitializer, FGfxPipelineDesc& Desc, FVulkanPSOKey* Key);
+	FVulkanRHIGraphicsPipelineState(FVulkanDevice* Device, const FGraphicsPipelineStateInitializer& PSOInitializer, const FGfxPipelineDesc& Desc, FVulkanPSOKey* Key);
 	~FVulkanRHIGraphicsPipelineState();
 
 	inline const FVulkanVertexInputStateInfo& GetVertexInputState() const
@@ -758,7 +758,7 @@ public:
 	FVulkanVertexInputStateInfo VertexInputState;
 	FVulkanLayout* Layout;
 	FVulkanDevice* Device;
-	FGfxPipelineDesc Desc;
+	const FGfxPipelineDesc Desc;
 	FVulkanShader* VulkanShaders[ShaderStage::NumStages];
 	const FVulkanRenderPass* RenderPass;
 
diff --git a/Engine/Source/Runtime/VulkanRHI/Private/VulkanRenderTarget.cpp b/Engine/Source/Runtime/VulkanRHI/Private/VulkanRenderTarget.cpp
--- a/Engine/Source/Runtime/VulkanRHI/Private/VulkanRenderTarget.cpp
+++ b/Engine/Source/Runtime/VulkanRHI/Private/VulkanRenderTarget.cpp
@@ -1340,15 +1340,19 @@ FVulkanRenderTargetLayout::FVulkanRenderTargetLayout(const FGraphicsPipelineStat
 		}
 		CurrDesc.storeOp = RenderTargetStoreActionToVulkan(Initializer.DepthTargetStoreAction);
 		CurrDesc.stencilStoreOp = RenderTargetStoreActionToVulkan(Initializer.StencilTargetStoreAction);
+
+		const VkImageLayout DepthLayout = Initializer.DepthStencilAccess.IsDepthWrite() ? VK_IMAGE_LAYOUT_ATTACHMENT_OPTIMAL : VK_IMAGE_LAYOUT_READ_ONLY_OPTIMAL;
+		const VkImageLayout StencilLayout = Initializer.DepthStencilAccess.IsStencilWrite() ? VK_IMAGE_LAYOUT_ATTACHMENT_OPTIMAL : VK_IMAGE_LAYOUT_READ_ONLY_OPTIMAL;
+
 		// If the initial != final we need to change the FullHashInfo and use FinalLayout
-		CurrDesc.initialLayout = VK_IMAGE_LAYOUT_ATTACHMENT_OPTIMAL;
-		CurrDesc.finalLayout = VK_IMAGE_LAYOUT_ATTACHMENT_OPTIMAL;
-		StencilDesc.stencilInitialLayout = VK_IMAGE_LAYOUT_ATTACHMENT_OPTIMAL;
-		StencilDesc.stencilFinalLayout = VK_IMAGE_LAYOUT_ATTACHMENT_OPTIMAL;
+		CurrDesc.initialLayout = DepthLayout;
+		CurrDesc.finalLayout = DepthLayout;
+		StencilDesc.stencilInitialLayout = StencilLayout;
+		StencilDesc.stencilFinalLayout = StencilLayout;
 
 		DepthReference.attachment = NumAttachmentDescriptions;
-		DepthReference.layout = Initializer.DepthStencilAccess.IsDepthWrite() ? VK_IMAGE_LAYOUT_ATTACHMENT_OPTIMAL : VK_IMAGE_LAYOUT_READ_ONLY_OPTIMAL;
-		StencilReference.stencilLayout = Initializer.DepthStencilAccess.IsStencilWrite() ? VK_IMAGE_LAYOUT_ATTACHMENT_OPTIMAL : VK_IMAGE_LAYOUT_READ_ONLY_OPTIMAL;
+		DepthReference.layout = DepthLayout;
+		StencilReference.stencilLayout = StencilLayout;
 
 		const bool bDepthStencilResolve = (Initializer.DepthTargetStoreAction == ERenderTargetStoreAction::EMultisampleResolve) || (Initializer.StencilTargetStoreAction == ERenderTargetStoreAction::EMultisampleResolve);
 		if (bDepthStencilResolve && GRHISupportsDepthStencilResolve && CurrDesc.samples > VK_SAMPLE_COUNT_1_BIT)
@@ -1372,8 +1376,8 @@ FVulkanRenderTargetLayout::FVulkanRenderTargetLayout(const FGraphicsPipelineStat
 		FullHashInfo.LoadOps[MaxSimultaneousRenderTargets + 1] = CurrDesc.stencilLoadOp;
 		FullHashInfo.StoreOps[MaxSimultaneousRenderTargets] = CurrDesc.storeOp;
 		FullHashInfo.StoreOps[MaxSimultaneousRenderTargets + 1] = CurrDesc.stencilStoreOp;
-		FullHashInfo.InitialLayout[MaxSimultaneousRenderTargets] = VK_IMAGE_LAYOUT_ATTACHMENT_OPTIMAL;
-		FullHashInfo.InitialLayout[MaxSimultaneousRenderTargets + 1] = VK_IMAGE_LAYOUT_ATTACHMENT_OPTIMAL;
+		FullHashInfo.InitialLayout[MaxSimultaneousRenderTargets] = DepthLayout;
+		FullHashInfo.InitialLayout[MaxSimultaneousRenderTargets + 1] = StencilLayout;
 		CompatibleHashInfo.Formats[MaxSimultaneousRenderTargets] = CurrDesc.format;
 
 		++NumAttachmentDescriptions;
