VulkanRHI: Disable draw markers and vulkan debug names until crashes and performance issues can be identified.

#rnx

diff --git a/Engine/Source/Runtime/VulkanRHI/Private/VulkanDevice.cpp b/Engine/Source/Runtime/VulkanRHI/Private/VulkanDevice.cpp
--- a/Engine/Source/Runtime/VulkanRHI/Private/VulkanDevice.cpp
+++ b/Engine/Source/Runtime/VulkanRHI/Private/VulkanDevice.cpp
@@ -481,9 +481,14 @@ void FVulkanDevice::SetupDrawMarkers()
 #if VULKAN_ENABLE_DRAW_MARKERS
 	if (RHI->SupportsDebugUtilsExt())
 	{
-		DebugMarkers.CmdBeginDebugLabel = (PFN_vkCmdBeginDebugUtilsLabelEXT)(void*)VulkanRHI::vkGetInstanceProcAddr(RHI->GetInstance(), "vkCmdBeginDebugUtilsLabelEXT");
-		DebugMarkers.CmdEndDebugLabel = (PFN_vkCmdEndDebugUtilsLabelEXT)(void*)VulkanRHI::vkGetInstanceProcAddr(RHI->GetInstance(), "vkCmdEndDebugUtilsLabelEXT");
-		DebugMarkers.SetDebugName = (PFN_vkSetDebugUtilsObjectNameEXT)(void*)VulkanRHI::vkGetInstanceProcAddr(RHI->GetInstance(), "vkSetDebugUtilsObjectNameEXT");
+		// HOTFIX for UE-218250: Disable vulkan draw markers to get around crash/performance issues
+		if (FParse::Param(FCommandLine::Get(), TEXT("forcevulkanddrawmarkers")))
+		{
+			DebugMarkers.CmdBeginDebugLabel = (PFN_vkCmdBeginDebugUtilsLabelEXT)(void*)VulkanRHI::vkGetInstanceProcAddr(RHI->GetInstance(), "vkCmdBeginDebugUtilsLabelEXT");
+			DebugMarkers.CmdEndDebugLabel = (PFN_vkCmdEndDebugUtilsLabelEXT)(void*)VulkanRHI::vkGetInstanceProcAddr(RHI->GetInstance(), "vkCmdEndDebugUtilsLabelEXT");
+			DebugMarkers.SetDebugName = (PFN_vkSetDebugUtilsObjectNameEXT)(void*)VulkanRHI::vkGetInstanceProcAddr(RHI->GetInstance(), "vkSetDebugUtilsObjectNameEXT");
+		}
+
 		if (DebugMarkers.CmdBeginDebugLabel && DebugMarkers.CmdEndDebugLabel && DebugMarkers.SetDebugName)
 		{
 			bDebugMarkersFound = true;
diff --git a/Engine/Source/Runtime/VulkanRHI/Private/VulkanDevice.h b/Engine/Source/Runtime/VulkanRHI/Private/VulkanDevice.h
--- a/Engine/Source/Runtime/VulkanRHI/Private/VulkanDevice.h
+++ b/Engine/Source/Runtime/VulkanRHI/Private/VulkanDevice.h
@@ -23,7 +23,8 @@ class FVulkanOcclusionQueryPool;
 class FVulkanRenderPassManager;
 class FVulkanRayTracingCompactionRequestHandler;
 
-#define VULKAN_USE_DEBUG_NAMES 1
+// HOTFIX for UE-218250: Disable vulkan debug names to get around crash/performance issues
+#define VULKAN_USE_DEBUG_NAMES 0
 
 #if VULKAN_USE_DEBUG_NAMES
 #define VULKAN_SET_DEBUG_NAME(Device, Type, Handle, Format, ...) (Device).VulkanSetObjectName(Type, (uint64)Handle, *FString::Printf(Format, __VA_ARGS__))
