Fixed typo in a Vulkan PSO file cache

diff --git a/Engine/Source/Runtime/VulkanRHI/Private/VulkanChunkedPipelineCache.cpp b/Engine/Source/Runtime/VulkanRHI/Private/VulkanChunkedPipelineCache.cpp
--- a/Engine/Source/Runtime/VulkanRHI/Private/VulkanChunkedPipelineCache.cpp
+++ b/Engine/Source/Runtime/VulkanRHI/Private/VulkanChunkedPipelineCache.cpp
@@ -397,8 +397,8 @@ public:
 		uint32 WriteFlags = EFileWrite::FILEWRITE_AllowRead  | (bAppend ? EFileWrite::FILEWRITE_Append : 0);
 		PSOFileWriter = TUniquePtr<FArchive>(IFileManager::Get().CreateFileWriter(*FullCachePath, WriteFlags ));
 		UE_CLOG(PSOFileWriter, LogRHI, Log, TEXT("Opened binary cache for write (%s)"), *FullCachePath);
-		UE_CLOG(PSOFileWriter == nullptr, LogRHI, Warning, TEXT("Failed to open OGL binary cache output file. (%s)"), *FullCachePath);
-		UE_CLOG(PSOFileWriter && (PSOFileWriter->IsError() || PSOFileWriter->IsCriticalError()), LogRHI, Error, TEXT("OGL binary cache output archive error (%s, %d,%d)"), *FullCachePath, PSOFileWriter->IsError(), PSOFileWriter->IsCriticalError());
+		UE_CLOG(PSOFileWriter == nullptr, LogRHI, Warning, TEXT("Failed to open Vulkan binary cache output file. (%s)"), *FullCachePath);
+		UE_CLOG(PSOFileWriter && (PSOFileWriter->IsError() || PSOFileWriter->IsCriticalError()), LogRHI, Error, TEXT("Vulkan binary cache output archive error (%s, %d,%d)"), *FullCachePath, PSOFileWriter->IsError(), PSOFileWriter->IsCriticalError());
 
 		return PSOFileWriter != nullptr;
 	}
