VulkanRHI: Use precise stage counts for gfx shaders.  Reduce single use bindless UB usage in SM6.

#rnx

diff --git a/Engine/Shaders/Public/Platform/Vulkan/VulkanCommon.ush b/Engine/Shaders/Public/Platform/Vulkan/VulkanCommon.ush
--- a/Engine/Shaders/Public/Platform/Vulkan/VulkanCommon.ush
+++ b/Engine/Shaders/Public/Platform/Vulkan/VulkanCommon.ush
@@ -7,7 +7,7 @@
 #pragma once
 
 // Update this GUID to invalidate shader recompilation for only Vulkan shaders
-#pragma message("UESHADERMETADATA_VERSION E4EECD20-AD45-4120-9A4E-FBCA3FFB670D")
+#pragma message("UESHADERMETADATA_VERSION AE91ACB9-5CCF-4845-8BE4-6887EE1241EE")
 
 #define COMPILER_SUPPORTS_ULONG_TYPES 1
 #define COMPILER_SUPPORTS_UINT64_IMAGE_ATOMICS (1 && COMPILER_SUPPORTS_ULONG_TYPES)
diff --git a/Engine/Source/Runtime/VulkanRHI/Private/VulkanDescriptorSets.cpp b/Engine/Source/Runtime/VulkanRHI/Private/VulkanDescriptorSets.cpp
--- a/Engine/Source/Runtime/VulkanRHI/Private/VulkanDescriptorSets.cpp
+++ b/Engine/Source/Runtime/VulkanRHI/Private/VulkanDescriptorSets.cpp
@@ -417,7 +417,7 @@ void FVulkanBindlessDescriptorManager::Init()
 		// todo-jn: this could be compacted..
 		auto CreateShaderStageUniformBufferLayout = [DeviceHandle]() {
 
-			const uint32 NumTotalBindings = VulkanBindless::MaxUniformBuffersPerStage * ShaderStage::MaxNumSets;
+			const uint32 NumTotalBindings = VulkanBindless::MaxUniformBuffersPerStage * ShaderStage::MaxNumStages;
 
 			TArray<VkDescriptorSetLayoutBinding> DescriptorSetLayoutBindings;
 			DescriptorSetLayoutBindings.SetNumZeroed(NumTotalBindings);
@@ -465,9 +465,9 @@ void FVulkanBindlessDescriptorManager::Init()
 			if (IsSingleUseUniformBufferSet)
 			{
 				// todo-jn: We're picky about uniform buffers values for now to allow for shortcuts...
-				check(LayoutSizeInBytes == (ShaderStage::MaxNumSets * VulkanBindless::MaxUniformBuffersPerStage * InOutState.DescriptorSize));
+				check(LayoutSizeInBytes == (ShaderStage::MaxNumStages * VulkanBindless::MaxUniformBuffersPerStage * InOutState.DescriptorSize));
 				check((LayoutSizeInBytes % DescriptorBufferProperties.descriptorBufferOffsetAlignment) == 0);
-				check((InOutState.MaxDescriptorCount % VulkanBindless::MaxUniformBuffersPerStage) == 0);
+				check((InOutState.MaxDescriptorCount % (ShaderStage::MaxNumStages * VulkanBindless::MaxUniformBuffersPerStage)) == 0);
 			}
 			else
 			{
@@ -607,12 +607,10 @@ void FVulkanBindlessDescriptorManager::RegisterUniformBuffers(VkCommandBuffer Co
 {
 	checkf(bIsSupported, TEXT("Trying to RegisterUniformBuffers but bindless is not supported!"));
 
-	SCOPED_NAMED_EVENT(FVulkanBindlessDescriptorManager_RegisterUniformBuffers, FColor::Purple);
-
 	BindlessSetState& BindlessUniformBufferSetState = BindlessSetStates[VulkanBindless::BindlessSingleUseUniformBufferSet];
 
 	// :todo-jn: Current uniform buffer layout is a bit wasteful with all the skipped bindings...
-	const uint32 BlockDescriptorCount = VulkanBindless::MaxUniformBuffersPerStage * ShaderStage::MaxNumSets;
+	const uint32 BlockDescriptorCount = VulkanBindless::MaxUniformBuffersPerStage * ShaderStage::MaxNumStages;
 	const uint32 BlockSize = BlockDescriptorCount * BindlessUniformBufferSetState.DescriptorSize;
 	// Leave the first block always zeroed for easier debugging
 	const uint32 FirstDescriptorIndex = BlockDescriptorCount + (CurrentUniformBufferDescriptorIndex.fetch_add(BlockDescriptorCount) % (BindlessUniformBufferSetState.MaxDescriptorCount - (2*BlockDescriptorCount)));
@@ -626,12 +624,13 @@ void FVulkanBindlessDescriptorManager::RegisterUniformBuffers(VkCommandBuffer Co
 	// :todo-jn: Clear them for easier debugging for now
 	FMemory::Memzero(&BindlessUniformBufferSetState.DebugDescriptors[FirstDescriptorByteOffset], BlockSize);
 
-	for (int32 StageIndex = 0; StageIndex < ShaderStage::NumStages; ++StageIndex)
+	for (int32 StageIndex = 0; StageIndex < ShaderStage::MaxNumStages; ++StageIndex)
 	{
 		const TArray<VkDescriptorAddressInfoEXT>& DescriptorAddressInfos = StageUBs[StageIndex];
 
 		if (DescriptorAddressInfos.Num())
 		{
+			checkSlow(StageIndex < GetNumStagesForBindPoint(BindPoint));
 			check(DescriptorAddressInfos.Num() <= VulkanBindless::MaxUniformBuffersPerStage);
 			const int32 StageOffset = StageIndex * VulkanBindless::MaxUniformBuffersPerStage;
 
diff --git a/Engine/Source/Runtime/VulkanRHI/Private/VulkanDescriptorSets.h b/Engine/Source/Runtime/VulkanRHI/Private/VulkanDescriptorSets.h
--- a/Engine/Source/Runtime/VulkanRHI/Private/VulkanDescriptorSets.h
+++ b/Engine/Source/Runtime/VulkanRHI/Private/VulkanDescriptorSets.h
@@ -33,7 +33,7 @@ struct FUniformBufferGatherInfo
 		FMemory::Memzero(CodeHeaders);
 	}
 
-	const FVulkanShaderHeader*	CodeHeaders[ShaderStage::NumStages];
+	const FVulkanShaderHeader* CodeHeaders[ShaderStage::MaxNumStages];
 };
 
 
@@ -220,7 +220,7 @@ public:
 			return true;
 		}
 	};
-	TStaticArray<FStageInfo, ShaderStage::NumStages> StageInfos;
+	TStaticArray<FStageInfo, ShaderStage::MaxNumStages> StageInfos;
 
 protected:
 	TMap<VkDescriptorType, uint32> LayoutTypes;
@@ -852,7 +852,7 @@ public:
 	FVulkanBindlessDescriptorManager(FVulkanDevice* InDevice);
 	~FVulkanBindlessDescriptorManager();
 
-	typedef TStaticArray<TArray<VkDescriptorAddressInfoEXT>, ShaderStage::NumStages> FUniformBufferDescriptorArrays;
+	typedef TStaticArray<TArray<VkDescriptorAddressInfoEXT>, ShaderStage::MaxNumStages> FUniformBufferDescriptorArrays;
 
 	void Init();
 	void Deinit();
@@ -1029,7 +1029,7 @@ private:
 private:
 	struct FSetsEntry
 	{
-		TStaticArray<VkDescriptorSet, ShaderStage::MaxNumSets> Sets;
+		TStaticArray<VkDescriptorSet, ShaderStage::MaxNumStages> Sets;
 		int32 NumSets;
 	};
 
diff --git a/Engine/Source/Runtime/VulkanRHI/Private/VulkanPipeline.cpp b/Engine/Source/Runtime/VulkanRHI/Private/VulkanPipeline.cpp
--- a/Engine/Source/Runtime/VulkanRHI/Private/VulkanPipeline.cpp
+++ b/Engine/Source/Runtime/VulkanRHI/Private/VulkanPipeline.cpp
@@ -239,7 +239,7 @@ FVulkanRHIGraphicsPipelineState::~FVulkanRHIGraphicsPipelineState()
 
 	Device->PipelineStateCache->NotifyDeletedGraphicsPSO(this);
 
-	for (int ShaderStageIndex = 0; ShaderStageIndex < ShaderStage::NumStages; ShaderStageIndex++)
+	for (int ShaderStageIndex = 0; ShaderStageIndex < ShaderStage::NumGraphicsStages; ShaderStageIndex++)
 	{
 		if (VulkanShaders[ShaderStageIndex] != nullptr)
 		{
@@ -248,9 +248,9 @@ FVulkanRHIGraphicsPipelineState::~FVulkanRHIGraphicsPipelineState()
 	}
 }
 
-void FVulkanRHIGraphicsPipelineState::GetOrCreateShaderModules(TRefCountPtr<FVulkanShaderModule> (&ShaderModulesOUT)[ShaderStage::NumStages], FVulkanShader*const* Shaders)
+void FVulkanRHIGraphicsPipelineState::GetOrCreateShaderModules(TRefCountPtr<FVulkanShaderModule> (&ShaderModulesOUT)[ShaderStage::NumGraphicsStages], FVulkanShader*const* Shaders)
 {
-	for (int32 Index = 0; Index < ShaderStage::NumStages; ++Index)
+	for (int32 Index = 0; Index < ShaderStage::NumGraphicsStages; ++Index)
 	{
 		check(!ShaderModulesOUT[Index].IsValid());
 		FVulkanShader* Shader = Shaders[Index];
@@ -269,7 +269,7 @@ FVulkanShader::FSpirvCode FVulkanRHIGraphicsPipelineState::GetPatchedSpirvCode(F
 
 void FVulkanRHIGraphicsPipelineState::PurgeShaderModules(FVulkanShader*const* Shaders)
 {
-	for (int32 Index = 0; Index < ShaderStage::NumStages; ++Index)
+	for (int32 Index = 0; Index < ShaderStage::NumGraphicsStages; ++Index)
 	{
 		FVulkanShader* Shader = Shaders[Index];
 		if (Shader)
@@ -1171,7 +1171,7 @@ static const TMap<uint8, VkFragmentShadingRateCombinerOpKHR> FragmentCombinerOpM
 	// @todo: Add "VK_FRAGMENT_SHADING_RATE_COMBINER_OP_MUL_KHR"?
 };
 
-static FString ShaderHashesToString(FVulkanShader* Shaders[ShaderStage::NumStages])
+static FString GfxShaderHashesToString(FVulkanShader* Shaders[ShaderStage::NumGraphicsStages])
 {
 	FString ShaderHashes = "";
 	if (Shaders[ShaderStage::Vertex])
@@ -1191,7 +1191,7 @@ static FString ShaderHashesToString(FVulkanShader* Shaders[ShaderStage::NumStage
 	return ShaderHashes;
 }
 
-bool FVulkanPipelineStateCacheManager::CreateGfxPipelineFromEntry(FVulkanRHIGraphicsPipelineState* PSO, FVulkanShader* Shaders[ShaderStage::NumStages], FGraphicsPipelineStateInitializer::EPSOPrecacheCompileType PSOCompileType)
+bool FVulkanPipelineStateCacheManager::CreateGfxPipelineFromEntry(FVulkanRHIGraphicsPipelineState* PSO, FVulkanShader* Shaders[ShaderStage::NumGraphicsStages], FGraphicsPipelineStateInitializer::EPSOPrecacheCompileType PSOCompileType)
 {
 	VkPipeline* Pipeline = &PSO->VulkanPipeline;
 	const FGfxPipelineDesc* GfxEntry = &PSO->Desc;
@@ -1200,7 +1200,7 @@ bool FVulkanPipelineStateCacheManager::CreateGfxPipelineFromEntry(FVulkanRHIGrap
 		Shaders[ShaderStage::Pixel] = ResourceCast(TShaderMapRef<FNULLPS>(GetGlobalShaderMap(GMaxRHIFeatureLevel)).GetPixelShader());
 	}
 
-	TRefCountPtr<FVulkanShaderModule> ShaderModules[ShaderStage::NumStages];
+	TRefCountPtr<FVulkanShaderModule> ShaderModules[ShaderStage::NumGraphicsStages];
 
 	PSO->GetOrCreateShaderModules(ShaderModules, Shaders);
 
@@ -1248,14 +1248,14 @@ bool FVulkanPipelineStateCacheManager::CreateGfxPipelineFromEntry(FVulkanRHIGrap
 	MSInfo.rasterizationSamples = (VkSampleCountFlagBits)FMath::Max<uint16>(1u, GfxEntry->RasterizationSamples);
 	MSInfo.alphaToCoverageEnable = GfxEntry->UseAlphaToCoverage;
 
-	VkPipelineShaderStageCreateInfo ShaderStages[ShaderStage::NumStages];
+	VkPipelineShaderStageCreateInfo ShaderStages[ShaderStage::NumGraphicsStages];
 	FMemory::Memzero(ShaderStages);
 	PipelineInfo.stageCount = 0;
 	PipelineInfo.pStages = ShaderStages;
 	// main_00000000_00000000
-	ANSICHAR EntryPoints[ShaderStage::NumStages][24];
-	VkPipelineShaderStageRequiredSubgroupSizeCreateInfo RequiredSubgroupSizeCreateInfo[ShaderStage::NumStages];
-	for (int32 ShaderStage = 0; ShaderStage < ShaderStage::NumStages; ++ShaderStage)
+	ANSICHAR EntryPoints[ShaderStage::NumGraphicsStages][24];
+	VkPipelineShaderStageRequiredSubgroupSizeCreateInfo RequiredSubgroupSizeCreateInfo[ShaderStage::NumGraphicsStages];
+	for (int32 ShaderStage = 0; ShaderStage < ShaderStage::NumGraphicsStages; ++ShaderStage)
 	{
 		if (!ShaderModules[ShaderStage].IsValid() || (Shaders[ShaderStage] == nullptr))
 		{
@@ -1379,7 +1379,7 @@ bool FVulkanPipelineStateCacheManager::CreateGfxPipelineFromEntry(FVulkanRHIGrap
 
 	if (Result != VK_SUCCESS)
 	{
-		FString ShaderHashes = ShaderHashesToString(Shaders);
+		FString ShaderHashes = GfxShaderHashesToString(Shaders);
 
 		UE_LOG(LogVulkanRHI, Error, TEXT("Failed to create graphics pipeline.\nShaders in pipeline: %s"), *ShaderHashes);
 		return false;
@@ -1398,7 +1398,7 @@ bool FVulkanPipelineStateCacheManager::CreateGfxPipelineFromEntry(FVulkanRHIGrap
 
 #if PLATFORM_ANDROID
 
-static VkResult CreatePSOWithExternalService(FVulkanDevice* Device, FGraphicsPipelineStateInitializer::EPSOPrecacheCompileType PSOCompileType, FVulkanRHIGraphicsPipelineState* PSO, FVulkanShader* Shaders[ShaderStage::NumStages], const VkGraphicsPipelineCreateInfo& PipelineInfo, VkPipelineCache DestPipelineCache, FRWLock& PipelineLock)
+static VkResult CreatePSOWithExternalService(FVulkanDevice* Device, FGraphicsPipelineStateInitializer::EPSOPrecacheCompileType PSOCompileType, FVulkanRHIGraphicsPipelineState* PSO, FVulkanShader* Shaders[ShaderStage::NumGraphicsStages], const VkGraphicsPipelineCreateInfo& PipelineInfo, VkPipelineCache DestPipelineCache, FRWLock& PipelineLock)
 {
 	VkResult Result = VK_ERROR_INITIALIZATION_FAILED;
 	FVulkanShader::FSpirvCode VS = PSO->GetPatchedSpirvCode(Shaders[ShaderStage::Vertex]);
@@ -1436,13 +1436,13 @@ static VkResult CreatePSOWithExternalService(FVulkanDevice* Device, FGraphicsPip
 	}
 	else
 	{
-  		UE_LOG(LogVulkanRHI, Error, TEXT("Android RemoteCompileServices Failed to create graphics pipeline (%s).\nShaders in pipeline %s"),*FailLog, *ShaderHashesToString(Shaders));
+  		UE_LOG(LogVulkanRHI, Error, TEXT("Android RemoteCompileServices Failed to create graphics pipeline (%s).\nShaders in pipeline %s"),*FailLog, *GfxShaderHashesToString(Shaders));
 	}
 	return Result;
 }
 #endif
 
-VkResult FVulkanPipelineStateCacheManager::CreateVKPipeline(FVulkanRHIGraphicsPipelineState* PSO, FVulkanShader* Shaders[ShaderStage::NumStages], const VkGraphicsPipelineCreateInfo& PipelineInfo, FGraphicsPipelineStateInitializer::EPSOPrecacheCompileType PSOCompileType)
+VkResult FVulkanPipelineStateCacheManager::CreateVKPipeline(FVulkanRHIGraphicsPipelineState* PSO, FVulkanShader* Shaders[ShaderStage::NumGraphicsStages], const VkGraphicsPipelineCreateInfo& PipelineInfo, FGraphicsPipelineStateInitializer::EPSOPrecacheCompileType PSOCompileType)
 {
 	bool bIsPrecompileJob = PSOCompileType != FGraphicsPipelineStateInitializer::EPSOPrecacheCompileType::NotSet;
 	if(FVulkanChunkedPipelineCacheManager::IsEnabled())
@@ -1569,7 +1569,7 @@ VkResult FVulkanPipelineStateCacheManager::CreateVKPipeline(FVulkanRHIGraphicsPi
 		}
 		else
 		{
-			FString ShaderHashes = ShaderHashesToString(Shaders);
+			FString ShaderHashes = GfxShaderHashesToString(Shaders);
 			UE_LOG(LogVulkanRHI, Error, TEXT("Android RemoteCompileServices Failed to create graphics pipeline (%s).\nShaders in pipeline: %s"), *FailLog, *ShaderHashes);
 		}
 	}
@@ -1744,8 +1744,8 @@ void FVulkanPipelineStateCacheManager::CreateGfxEntry(const FGraphicsPipelineSta
 {
 	FGfxPipelineDesc* OutGfxEntry = Desc;
 
-	FVulkanShader* Shaders[ShaderStage::NumStages];
-	GetVulkanShaders(PSOInitializer.BoundShaderState, Shaders);
+	FVulkanShader* Shaders[ShaderStage::NumGraphicsStages];
+	GetVulkanGfxShaders(PSOInitializer.BoundShaderState, Shaders);
 
 	FVulkanVertexInputStateInfo VertexInputState;
 	
@@ -1870,8 +1870,8 @@ void FVulkanPipelineStateCacheManager::CreateGfxEntry(const FGraphicsPipelineSta
 		1263027877466626099llu,
 		2698115308251696101llu,
 	};
-	static_assert(sizeof(Primes) / sizeof(Primes[0]) >= ShaderStage::NumStages);
-	for (int32 Index = 0; Index < ShaderStage::NumStages; ++Index)
+	static_assert(sizeof(Primes) / sizeof(Primes[0]) >= ShaderStage::NumGraphicsStages);
+	for (int32 Index = 0; Index < ShaderStage::NumGraphicsStages; ++Index)
 	{
 		FVulkanShader* Shader = Shaders[Index];
 		uint64 Key = 0;
@@ -1885,7 +1885,7 @@ void FVulkanPipelineStateCacheManager::CreateGfxEntry(const FGraphicsPipelineSta
 	}
 	OutGfxEntry->ShaderKeyShared = SharedKey;
 #else
-	for (int32 Index = 0; Index < ShaderStage::NumStages; ++Index)
+	for (int32 Index = 0; Index < ShaderStage::NumGraphicsStages; ++Index)
 	{
 		FVulkanShader* Shader = Shaders[Index];
 		if (Shader)
@@ -1937,7 +1937,7 @@ FVulkanRHIGraphicsPipelineState::FVulkanRHIGraphicsPipelineState(FVulkanDevice*
 #endif
 	VulkanShaders[ShaderStage::Pixel] = static_cast<FVulkanPixelShader*>(PSOInitializer_.BoundShaderState.PixelShaderRHI);
 
-	for (int ShaderStageIndex = 0; ShaderStageIndex < ShaderStage::NumStages; ShaderStageIndex++)
+	for (int ShaderStageIndex = 0; ShaderStageIndex < ShaderStage::NumGraphicsStages; ShaderStageIndex++)
 	{
 		if (VulkanShaders[ShaderStageIndex] != nullptr)
 		{
@@ -2087,7 +2087,7 @@ FGraphicsPipelineStateRHIRef FVulkanPipelineStateCacheManager::RHICreateGraphics
 		NewPSO->RenderPass = Device->GetImmediateContext().PrepareRenderPassForPSOCreation(Initializer);
 		{
 			const FBoundShaderStateInput& BSI = Initializer.BoundShaderState;
-			for (int32 StageIdx = 0; StageIdx < ShaderStage::NumStages; ++StageIdx)
+			for (int32 StageIdx = 0; StageIdx < ShaderStage::NumGraphicsStages; ++StageIdx)
 			{
 				NewPSO->ShaderKeys[StageIdx] = GetShaderKeyForGfxStage(BSI, (ShaderStage::EStage)StageIdx);
 			}
@@ -2106,10 +2106,10 @@ FGraphicsPipelineStateRHIRef FVulkanPipelineStateCacheManager::RHICreateGraphics
 
 				// Create the pipeline
 				double BeginTime = FPlatformTime::Seconds();
-				FVulkanShader* VulkanShaders[ShaderStage::NumStages];
-				GetVulkanShaders(Initializer.BoundShaderState, VulkanShaders);
+				FVulkanShader* VulkanShaders[ShaderStage::NumGraphicsStages];
+				GetVulkanGfxShaders(Initializer.BoundShaderState, VulkanShaders);
 
-				for (int32 StageIdx = 0; StageIdx < ShaderStage::NumStages; ++StageIdx)
+				for (int32 StageIdx = 0; StageIdx < ShaderStage::NumGraphicsStages; ++StageIdx)
 				{
 					uint64 key = GetShaderKeyForGfxStage(BSI, (ShaderStage::EStage)StageIdx);
 					check(key == NewPSO->ShaderKeys[StageIdx]);
@@ -2379,9 +2379,9 @@ bool FVulkanPipelineStateCacheManager::FVulkanLRUCacheFile::Load(FArchive& Ar)
 
 
 
-void GetVulkanShaders(const FBoundShaderStateInput& BSI, FVulkanShader* OutShaders[ShaderStage::NumStages])
+void GetVulkanGfxShaders(const FBoundShaderStateInput& BSI, FVulkanShader* OutShaders[ShaderStage::NumGraphicsStages])
 {
-	FMemory::Memzero(OutShaders, ShaderStage::NumStages * sizeof(*OutShaders));
+	FMemory::Memzero(OutShaders, ShaderStage::NumGraphicsStages * sizeof(*OutShaders));
 
 	OutShaders[ShaderStage::Vertex] = ResourceCast(BSI.VertexShaderRHI);
 
@@ -2400,10 +2400,10 @@ void GetVulkanShaders(const FBoundShaderStateInput& BSI, FVulkanShader* OutShade
 	}
 }
 
-void GetVulkanShaders(FVulkanDevice* Device, const FVulkanRHIGraphicsPipelineState& GfxPipelineState, FVulkanShader* OutShaders[ShaderStage::NumStages])
+void GetVulkanGfxShaders(FVulkanDevice* Device, const FVulkanRHIGraphicsPipelineState& GfxPipelineState, FVulkanShader* OutShaders[ShaderStage::NumGraphicsStages])
 {
-	FMemory::Memzero(OutShaders, ShaderStage::NumStages * sizeof(*OutShaders));
-	Device->GetShaderFactory().LookupShaders(GfxPipelineState.ShaderKeys, OutShaders);
+	FMemory::Memzero(OutShaders, ShaderStage::NumGraphicsStages * sizeof(*OutShaders));
+	Device->GetShaderFactory().LookupGfxShaders(GfxPipelineState.ShaderKeys, OutShaders);
 }
 
 void FVulkanPipelineStateCacheManager::TickLRU()
@@ -2548,9 +2548,9 @@ void FVulkanPipelineStateCacheManager::LRUTouch(FVulkanRHIGraphicsPipelineState*
 		{
 			// Create the pipeline
 			double BeginTime = FPlatformTime::Seconds();
-			FVulkanShader* VulkanShaders[ShaderStage::NumStages];
+			FVulkanShader* VulkanShaders[ShaderStage::NumGraphicsStages];
 
-			GetVulkanShaders(Device, *PSO, VulkanShaders);
+			GetVulkanGfxShaders(Device, *PSO, VulkanShaders);
 
 			QUICK_SCOPE_CYCLE_COUNTER(STAT_Vulkan_LRUMiss);
 
@@ -2633,7 +2633,7 @@ void FVulkanPipelineStateCacheManager::LRURemove(FVulkanRHIGraphicsPipelineState
 		PSO->DeleteVkPipeline(bImmediate);
 		if (GVulkanReleaseShaderModuleWhenEvictingPSO)
 		{
-	        for (int ShaderStageIndex = 0; ShaderStageIndex < ShaderStage::NumStages; ShaderStageIndex++)
+	        for (int ShaderStageIndex = 0; ShaderStageIndex < ShaderStage::NumGraphicsStages; ShaderStageIndex++)
 	        {
 				if (PSO->VulkanShaders[ShaderStageIndex] != nullptr)
 				{
diff --git a/Engine/Source/Runtime/VulkanRHI/Private/VulkanPipeline.h b/Engine/Source/Runtime/VulkanRHI/Private/VulkanPipeline.h
--- a/Engine/Source/Runtime/VulkanRHI/Private/VulkanPipeline.h
+++ b/Engine/Source/Runtime/VulkanRHI/Private/VulkanPipeline.h
@@ -42,11 +42,10 @@ inline uint64 GetShaderKeyForGfxStage(const FBoundShaderStateInput& BSI, ShaderS
 		return GetShaderKey<FVulkanPixelShader>(BSI.PixelShaderRHI);
 	case ShaderStage::Geometry:
 		return GetShaderKey<FVulkanGeometryShader>(BSI.GetGeometryShader());
-	case ShaderStage::RayGen:
-	case ShaderStage::RayHitGroup:
-	case ShaderStage::RayMiss:
-	case ShaderStage::RayCallable:
-		return 0; // VKRT todo
+	case ShaderStage::Mesh:
+		return 0; // GetShaderKey<FVulkanMeshShader>(BSI.GetMeshShader());
+	case ShaderStage::Task:
+		return 0; // GetShaderKey<FVulkanTaskShader>(BSI.GetAmplificationShader());
 	default:
 		check(0);
 	}
@@ -55,15 +54,15 @@ inline uint64 GetShaderKeyForGfxStage(const FBoundShaderStateInput& BSI, ShaderS
 }
 
 
-void GetVulkanShaders(const FBoundShaderStateInput& BSI, FVulkanShader* OutShaders[ShaderStage::NumStages]);
-void GetVulkanShaders(FVulkanDevice* Device, const FVulkanRHIGraphicsPipelineState& GfxPipelineState, FVulkanShader* OutShaders[ShaderStage::NumStages]);
+void GetVulkanGfxShaders(const FBoundShaderStateInput& BSI, FVulkanShader* OutShaders[ShaderStage::NumGraphicsStages]);
+void GetVulkanGfxShaders(FVulkanDevice* Device, const FVulkanRHIGraphicsPipelineState& GfxPipelineState, FVulkanShader* OutShaders[ShaderStage::NumGraphicsStages]);
 
 
 
 struct FVulkanShaderHashes
 {
 	uint32 Hash;
-	FSHAHash Stages[ShaderStage::NumStages];
+	FSHAHash Stages[ShaderStage::MaxNumStages];
 
 	FVulkanShaderHashes();
 	FVulkanShaderHashes(const FGraphicsPipelineStateInitializer& PSOInitializer);
@@ -80,7 +79,7 @@ struct FVulkanShaderHashes
 
 	friend inline bool operator == (const FVulkanShaderHashes& A, const FVulkanShaderHashes& B)
 	{
-		for (int32 Index = 0; Index < ShaderStage::NumStages; ++Index)
+		for (int32 Index = 0; Index < ShaderStage::MaxNumStages; ++Index)
 		{
 			if (A.Stages[Index] != B.Stages[Index])
 			{
@@ -270,7 +269,7 @@ struct FGfxPipelineDesc
 	FDepthStencil DepthStencil;
 
 #if VULKAN_USE_SHADERKEYS
-	uint64 ShaderKeys[ShaderStage::NumStages];
+	uint64 ShaderKeys[ShaderStage::NumGraphicsStages];
 	uint64 ShaderKeyShared;
 #else
 	FVulkanShaderHashes ShaderHashes;
@@ -519,9 +518,9 @@ private:
 	FGraphicsPipelineStateRHIRef RHICreateGraphicsPipelineState(const FGraphicsPipelineStateInitializer& Initializer);
 	FVulkanComputePipeline* RHICreateComputePipelineState(FRHIComputeShader* ComputeShaderRHI);
 	void NotifyDeletedGraphicsPSO(FRHIGraphicsPipelineState* PSO);
-	bool CreateGfxPipelineFromEntry(FVulkanRHIGraphicsPipelineState* PSO, FVulkanShader* Shaders[ShaderStage::NumStages], FGraphicsPipelineStateInitializer::EPSOPrecacheCompileType PSOCompileType);
+	bool CreateGfxPipelineFromEntry(FVulkanRHIGraphicsPipelineState* PSO, FVulkanShader* Shaders[ShaderStage::NumGraphicsStages], FGraphicsPipelineStateInitializer::EPSOPrecacheCompileType PSOCompileType);
 
-	VkResult CreateVKPipeline(FVulkanRHIGraphicsPipelineState* PSO, FVulkanShader* Shaders[ShaderStage::NumStages], const VkGraphicsPipelineCreateInfo& PipelineInfo, FGraphicsPipelineStateInitializer::EPSOPrecacheCompileType PSOCompileType);
+	VkResult CreateVKPipeline(FVulkanRHIGraphicsPipelineState* PSO, FVulkanShader* Shaders[ShaderStage::NumGraphicsStages], const VkGraphicsPipelineCreateInfo& PipelineInfo, FGraphicsPipelineStateInitializer::EPSOPrecacheCompileType PSOCompileType);
 
 	FVulkanLayout* FindOrAddLayout(const FVulkanDescriptorSetsLayoutInfo& DescriptorSetLayoutInfo, bool bGfxLayout);
 	FVulkanComputePipeline* CreateComputePipelineFromShader(FVulkanComputeShader* Shader);
@@ -742,7 +741,7 @@ public:
 	}
 
 	void DeleteVkPipeline(bool bImmediate);
-	void GetOrCreateShaderModules(TRefCountPtr<FVulkanShaderModule>(&ShaderModulesOUT)[ShaderStage::NumStages], FVulkanShader* const* Shaders);
+	void GetOrCreateShaderModules(TRefCountPtr<FVulkanShaderModule>(&ShaderModulesOUT)[ShaderStage::NumGraphicsStages], FVulkanShader* const* Shaders);
 	FVulkanShader::FSpirvCode GetPatchedSpirvCode(FVulkanShader* Shader);
 	void PurgeShaderModules(FVulkanShader*const* Shaders);
 
@@ -750,7 +749,7 @@ public:
 	bool								bHasInputAttachments = false;
 	bool								bIsRegistered;
 
-	uint64								ShaderKeys[ShaderStage::NumStages];
+	uint64								ShaderKeys[ShaderStage::NumGraphicsStages];
 	TEnumAsByte<EPrimitiveType>			PrimitiveType;
 
 
@@ -759,7 +758,7 @@ public:
 	FVulkanLayout* Layout;
 	FVulkanDevice* Device;
 	const FGfxPipelineDesc Desc;
-	FVulkanShader* VulkanShaders[ShaderStage::NumStages];
+	FVulkanShader* VulkanShaders[ShaderStage::NumGraphicsStages];
 	const FVulkanRenderPass* RenderPass;
 
 
diff --git a/Engine/Source/Runtime/VulkanRHI/Private/VulkanPipelineState.cpp b/Engine/Source/Runtime/VulkanRHI/Private/VulkanPipelineState.cpp
--- a/Engine/Source/Runtime/VulkanRHI/Private/VulkanPipelineState.cpp
+++ b/Engine/Source/Runtime/VulkanRHI/Private/VulkanPipelineState.cpp
@@ -38,7 +38,7 @@ static bool ShouldAlwaysWriteDescriptors()
 }
 
 FVulkanComputePipelineDescriptorState::FVulkanComputePipelineDescriptorState(FVulkanDevice* InDevice, FVulkanComputePipeline* InComputePipeline)
-	: FVulkanCommonPipelineDescriptorState(InDevice, 1)
+	: FVulkanCommonPipelineDescriptorState(InDevice, ShaderStage::NumComputeStages)
 	, PackedUniformBuffersMask(0)
 	, PackedUniformBuffersDirty(0)
 	, ComputePipeline(InComputePipeline)
@@ -352,7 +352,7 @@ bool FVulkanGraphicsPipelineDescriptorState::InternalUpdateDescriptorSets(FVulka
 #if VULKAN_ENABLE_AGGRESSIVE_STATS
 		SCOPE_CYCLE_COUNTER(STAT_VulkanApplyPackedUniformBuffers);
 #endif
-		for (int32 Stage = 0; Stage < ShaderStage::NumStages; ++Stage)
+		for (int32 Stage = 0; Stage < ShaderStage::NumGraphicsStages; ++Stage)
 		{
 			if (PackedUniformBuffersDirty[Stage] != 0)
 			{
@@ -429,7 +429,7 @@ void FVulkanGraphicsPipelineDescriptorState::UpdateBindlessDescriptors(FVulkanCo
 #if VULKAN_ENABLE_AGGRESSIVE_STATS
 		SCOPE_CYCLE_COUNTER(STAT_VulkanApplyPackedUniformBuffers);
 #endif
-		for (int32 Stage = 0; Stage < ShaderStage::NumStages; ++Stage)
+		for (int32 Stage = 0; Stage < ShaderStage::NumGraphicsStages; ++Stage)
 		{
 			const FVulkanShader* VulkanShader = GfxPipeline->GetVulkanShader(GetFrequencyForGfxStage((ShaderStage::EStage)Stage));
 			if (!VulkanShader)
diff --git a/Engine/Source/Runtime/VulkanRHI/Private/VulkanPipelineState.h b/Engine/Source/Runtime/VulkanRHI/Private/VulkanPipelineState.h
--- a/Engine/Source/Runtime/VulkanRHI/Private/VulkanPipelineState.h
+++ b/Engine/Source/Runtime/VulkanRHI/Private/VulkanPipelineState.h
@@ -277,9 +277,9 @@ public:
 	}
 
 protected:
-	TStaticArray<FPackedUniformBuffers, ShaderStage::NumStages> PackedUniformBuffers;
-	TStaticArray<uint32, ShaderStage::NumStages> PackedUniformBuffersMask;
-	TStaticArray<uint32, ShaderStage::NumStages> PackedUniformBuffersDirty;
+	TStaticArray<FPackedUniformBuffers, ShaderStage::NumGraphicsStages> PackedUniformBuffers;
+	TStaticArray<uint32, ShaderStage::NumGraphicsStages> PackedUniformBuffersMask;
+	TStaticArray<uint32, ShaderStage::NumGraphicsStages> PackedUniformBuffersDirty;
 
 	FVulkanRHIGraphicsPipelineState* GfxPipeline;
 
diff --git a/Engine/Source/Runtime/VulkanRHI/Private/VulkanRHI.cpp b/Engine/Source/Runtime/VulkanRHI/Private/VulkanRHI.cpp
--- a/Engine/Source/Runtime/VulkanRHI/Private/VulkanRHI.cpp
+++ b/Engine/Source/Runtime/VulkanRHI/Private/VulkanRHI.cpp
@@ -1742,7 +1742,8 @@ void FVulkanDescriptorSetsLayoutInfo::GenerateHash(const TArrayView<FRHISamplerS
 		Hash = FCrc::MemCrc32(&SetLayouts[layoutIndex].Hash, sizeof(uint32), Hash);
 	}
 
-	for (uint32 RemapingIndex = 0; RemapingIndex < ShaderStage::NumStages; ++RemapingIndex)
+	const uint32 NumStages = GetNumStagesForBindPoint(InBindPoint);
+	for (uint32 RemapingIndex = 0; RemapingIndex < NumStages; ++RemapingIndex)
 	{
 		const FStageInfo& StageInfo = StageInfos[RemapingIndex];
 
diff --git a/Engine/Source/Runtime/VulkanRHI/Private/VulkanRHIPrivate.h b/Engine/Source/Runtime/VulkanRHI/Private/VulkanRHIPrivate.h
--- a/Engine/Source/Runtime/VulkanRHI/Private/VulkanRHIPrivate.h
+++ b/Engine/Source/Runtime/VulkanRHI/Private/VulkanRHIPrivate.h
@@ -102,6 +102,23 @@ inline EShaderFrequency VkStageBitToUEFrequency(VkShaderStageFlagBits FlagBits)
 	return SF_NumFrequencies;
 }
 
+static constexpr int32 GetNumStagesForBindPoint(VkPipelineBindPoint BindPoint)
+{
+	switch (BindPoint)
+	{
+	case VK_PIPELINE_BIND_POINT_GRAPHICS:
+		return ShaderStage::NumGraphicsStages;
+
+	case VK_PIPELINE_BIND_POINT_COMPUTE:
+		return ShaderStage::NumComputeStages;
+
+	case VK_PIPELINE_BIND_POINT_RAY_TRACING_KHR:
+		return ShaderStage::NumRayTracingStages;
+	};
+
+	return 0;
+}
+
 class FVulkanRenderTargetLayout
 {
 public:
diff --git a/Engine/Source/Runtime/VulkanRHI/Private/VulkanResources.h b/Engine/Source/Runtime/VulkanRHI/Private/VulkanResources.h
--- a/Engine/Source/Runtime/VulkanRHI/Private/VulkanResources.h
+++ b/Engine/Source/Runtime/VulkanRHI/Private/VulkanResources.h
@@ -332,7 +332,7 @@ public:
 	template <EShaderFrequency ShaderFrequency>
 	FVulkanRayTracingShader* CreateRayTracingShader(TArrayView<const uint8> Code, FVulkanDevice* Device);
 
-	void LookupShaders(const uint64 InShaderKeys[ShaderStage::NumStages], FVulkanShader* OutShaders[ShaderStage::NumStages]) const;
+	void LookupGfxShaders(const uint64 InShaderKeys[ShaderStage::NumGraphicsStages], FVulkanShader* OutShaders[ShaderStage::NumGraphicsStages]) const;
 
 	void OnDeleteShader(const FVulkanShader& Shader);
 
diff --git a/Engine/Source/Runtime/VulkanRHI/Private/VulkanShaders.cpp b/Engine/Source/Runtime/VulkanRHI/Private/VulkanShaders.cpp
--- a/Engine/Source/Runtime/VulkanRHI/Private/VulkanShaders.cpp
+++ b/Engine/Source/Runtime/VulkanRHI/Private/VulkanShaders.cpp
@@ -185,9 +185,9 @@ FVulkanRayTracingShader* FVulkanShaderFactory::CreateRayTracingShader(TArrayView
 	return RetShader;
 }
 
-void FVulkanShaderFactory::LookupShaders(const uint64 InShaderKeys[ShaderStage::NumStages], FVulkanShader* OutShaders[ShaderStage::NumStages]) const
+void FVulkanShaderFactory::LookupGfxShaders(const uint64 InShaderKeys[ShaderStage::NumGraphicsStages], FVulkanShader* OutShaders[ShaderStage::NumGraphicsStages]) const
 {
-	for (int32 Idx = 0; Idx < ShaderStage::NumStages; ++Idx)
+	for (int32 Idx = 0; Idx < ShaderStage::NumGraphicsStages; ++Idx)
 	{
 		uint64 ShaderKey = InShaderKeys[Idx];
 		if (ShaderKey)
@@ -737,7 +737,7 @@ void FVulkanDescriptorSetsLayoutInfo::FinalizeBindings(const FVulkanDevice& Devi
 	const uint32 MaxDescriptorSetUniformBuffersDynamic = Device.GetLimits().maxDescriptorSetUniformBuffersDynamic;
 
 	int32 CurrentImmutableSampler = 0;
-	for (int32 Stage = 0; Stage < (bIsCompute ? 1 : ShaderStage::NumStages); ++Stage)
+	for (int32 Stage = 0; Stage < (bIsCompute ? ShaderStage::NumComputeStages : ShaderStage::NumGraphicsStages); ++Stage)
 	{
 		checkSlow(StageInfos[Stage].IsEmpty());
 
diff --git a/Engine/Source/Runtime/VulkanRHI/Public/VulkanCommon.h b/Engine/Source/Runtime/VulkanRHI/Public/VulkanCommon.h
--- a/Engine/Source/Runtime/VulkanRHI/Public/VulkanCommon.h
+++ b/Engine/Source/Runtime/VulkanRHI/Public/VulkanCommon.h
@@ -20,44 +20,50 @@
 
 namespace ShaderStage
 {
+	// Adjusting these requires a full shader rebuild (ie modify the guid on VulkanCommon.usf)
+	// Keep the values in sync with EShaderFrequency
 	enum EStage
 	{
-		// Adjusting these requires a full shader rebuild (ie modify the guid on VulkanCommon.usf)
-		// Keep the values in sync with EShaderFrequency
 		Vertex = 0,
 		Pixel = 1,
 		Geometry = 2,
+		Mesh = 3,
+		Task = 4,
 
-		RayGen = 3,
-		RayMiss = 4,
-		RayHitGroup = 5,
-		RayCallable = 6,
+		NumGraphicsStages = 5,
 
-		NumGraphicsStages = 3,
-		NumRayTracingStages = 4,
+		RayGen = 0,
+		RayMiss = 1,
+		RayHitGroup = 2,
+		RayCallable = 3,
 
-		NumStages = (NumGraphicsStages + NumRayTracingStages),
+		NumRayTracingStages = 4,
 
-		// Compute is its own pipeline, so it can all live as set 0
 		Compute = 0,
 
-		MaxNumSets = 8,
+		NumComputeStages = 1,
+
+		MaxNumStages = 6, // work with even count to simplify bindless alignment requirements
 
 		Invalid = -1,
 	};
 
+	static_assert(MaxNumStages >= FMath::Max(NumComputeStages, FMath::Max(NumGraphicsStages, NumRayTracingStages)), "MaxNumStages too small!");
+
 	inline EStage GetStageForFrequency(EShaderFrequency Stage)
 	{
 		switch (Stage)
 		{
-		case SF_Vertex:		return Vertex;
-		case SF_Pixel:		return Pixel;
-		case SF_Geometry:	return Geometry;
+		case SF_Vertex:			return Vertex;
+		case SF_Mesh:			return Mesh;
+		case SF_Amplification:	return Task;
+		case SF_Pixel:			return Pixel;
+		case SF_Geometry:		return Geometry;
 		case SF_RayGen:			return RayGen;
 		case SF_RayMiss:		return RayMiss;
 		case SF_RayHitGroup:	return RayHitGroup;
 		case SF_RayCallable:	return RayCallable;
-		case SF_Compute:	return Compute;
+		case SF_Compute:		return Compute;
 		default:
 			checkf(0, TEXT("Invalid shader Stage %d"), (int32)Stage);
 			break;
@@ -73,12 +79,10 @@ namespace ShaderStage
 		case EStage::Vertex:	return SF_Vertex;
 		case EStage::Pixel:		return SF_Pixel;
 		case EStage::Geometry:	return SF_Geometry;
-		case EStage::RayGen:		return SF_RayGen;
-		case EStage::RayMiss:		return SF_RayMiss;
-		case EStage::RayHitGroup:	return SF_RayHitGroup;
-		case EStage::RayCallable:	return SF_RayCallable;
+		case EStage::Mesh:		return SF_Mesh;
+		case EStage::Task:		return SF_Amplification;
 		default:
-			checkf(0, TEXT("Invalid shader Stage %d"), (int32)Stage);
+			checkf(0, TEXT("Invalid graphic shader stage: %d"), (int32)Stage);
 			break;
 		}
 
