Fix swappy issue in non-unity Android build

diff --git a/Engine/Source/Runtime/Core/Private/Android/AndroidSwappyLoader.cpp b/Engine/Source/Runtime/Core/Private/Android/AndroidSwappyLoader.cpp
--- a/Engine/Source/Runtime/Core/Private/Android/AndroidSwappyLoader.cpp
+++ b/Engine/Source/Runtime/Core/Private/Android/AndroidSwappyLoader.cpp
@@ -2,6 +2,8 @@
 
 #include "AndroidSwappyLoader.h"
 
+#if USE_ANDROID_OPENGL_SWAPPY || USE_ANDROID_VULKAN_SWAPPY
+
 #undef VK_NO_PROTOTYPES
 #include "swappy/swappyVk.h"
 #include "swappy/swappyGL.h"
@@ -356,3 +358,5 @@ void SwappyGL_uninjectTracer(const SwappyTracer* t)
 {
 	fpSwappyGL_uninjectTracer(t);
 }
+
+#endif // USE_ANDROID_OPENGL_SWAPPY || USE_ANDROID_VULKAN_SWAPPY
\ No newline at end of file
diff --git a/Engine/Source/Runtime/Core/Private/Android/AndroidSwappyLoader.h b/Engine/Source/Runtime/Core/Private/Android/AndroidSwappyLoader.h
--- a/Engine/Source/Runtime/Core/Private/Android/AndroidSwappyLoader.h
+++ b/Engine/Source/Runtime/Core/Private/Android/AndroidSwappyLoader.h
@@ -1,6 +1,17 @@
 // Copyright Epic Games, Inc. All Rights Reserved.
-
 #pragma once
 
+#include "CoreMinimal.h"
+#include "CoreTypes.h"
+
+#ifndef USE_ANDROID_OPENGL_SWAPPY
+#define USE_ANDROID_OPENGL_SWAPPY 0
+#endif
+#ifndef USE_ANDROID_VULKAN_SWAPPY
+#define USE_ANDROID_VULKAN_SWAPPY 0
+#endif
+
+#if USE_ANDROID_OPENGL_SWAPPY || USE_ANDROID_VULKAN_SWAPPY
 void LoadSwappy();
 void UnloadSwappy();
+#endif
\ No newline at end of file
diff --git a/Engine/Source/Runtime/VulkanRHI/Private/Android/VulkanAndroidPlatform.cpp b/Engine/Source/Runtime/VulkanRHI/Private/Android/VulkanAndroidPlatform.cpp
--- a/Engine/Source/Runtime/VulkanRHI/Private/Android/VulkanAndroidPlatform.cpp
+++ b/Engine/Source/Runtime/VulkanRHI/Private/Android/VulkanAndroidPlatform.cpp
@@ -20,7 +20,7 @@
 #include "ProfilingDebugging/ScopedTimers.h"
 #include "Android/AndroidDynamicRHI.h"
 
-#if USE_ANDROID_SWAPPY
+#if USE_ANDROID_VULKAN_SWAPPY
 #undef VK_NO_PROTOTYPES
 #include "Android/AndroidJNI.h"
 #include "Android/AndroidApplication.h"
@@ -48,7 +48,7 @@ namespace AndroidVulkan
 	}
 };
 
-#endif // #if USE_ANDROID_SWAPPY
+#endif // #if USE_ANDROID_VULKAN_SWAPPY
 
 // From VulklanSwapChain.cpp
 extern int32 GVulkanCPURenderThreadFramePacer;
@@ -78,7 +78,7 @@ ENUM_VK_ENTRYPOINTS_ALL(DEFINE_VK_ENTRYPOINTS)
 #define VULKAN_MALI_LAYER_NAME "VK_LAYER_ARM_AGA"
 
 
-#if USE_ANDROID_SWAPPY
+#if USE_ANDROID_VULKAN_SWAPPY
 bool FVulkanAndroidPlatform::bSwappyEnabledAtRHIInit = false;
 #endif
 void* FVulkanAndroidPlatform::VulkanLib = nullptr;
@@ -371,7 +371,7 @@ void FVulkanAndroidPlatform::FreeVulkanLibrary()
 
 bool FVulkanAndroidPlatform::HasCustomFrameTiming()
 {
-#if USE_ANDROID_SWAPPY
+#if USE_ANDROID_VULKAN_SWAPPY
 	return FVulkanAndroidPlatform::bSwappyEnabledAtRHIInit;
 #endif
 	return false;
@@ -379,7 +379,7 @@ bool FVulkanAndroidPlatform::HasCustomFrameTiming()
 
 void FVulkanAndroidPlatform::InitDevice(FVulkanDevice* InDevice)
 {
-#if USE_ANDROID_SWAPPY
+#if USE_ANDROID_VULKAN_SWAPPY
 	if (FVulkanAndroidPlatform::bSwappyEnabledAtRHIInit)
 	{
 		FVulkanQueue* GfxQueue = InDevice->GetGraphicsQueue();
@@ -485,7 +485,7 @@ void FVulkanAndroidPlatform::GetDeviceExtensions(FVulkanDevice* Device, FVulkanD
 	OutExtensions.Add(MakeUnique<FVulkanDeviceExtension>(Device, VULKAN_MALI_LAYER_NAME, VULKAN_EXTENSION_ENABLED, VULKAN_EXTENSION_NOT_PROMOTED));
 #endif
 
-#if USE_ANDROID_SWAPPY
+#if USE_ANDROID_VULKAN_SWAPPY
 	if (FVulkanAndroidPlatform::bSwappyEnabledAtRHIInit)
 	{
 		// make sure any extensions swappy requires are included
@@ -531,7 +531,7 @@ void FVulkanAndroidPlatform::GetDeviceLayers(TArray<const ANSICHAR*>& OutLayers)
 
 void FVulkanAndroidPlatform::NotifyFoundDeviceLayersAndExtensions(VkPhysicalDevice PhysicalDevice, const TArray<const ANSICHAR*>& Layers, const TArray<const ANSICHAR*>& Extensions)
 {
-#if USE_ANDROID_SWAPPY
+#if USE_ANDROID_VULKAN_SWAPPY
 	if (FVulkanAndroidPlatform::bSwappyEnabledAtRHIInit)
 	{
 		// Determine extensions required by Swappy
@@ -607,7 +607,7 @@ bool FVulkanAndroidPlatform::SupportsDynamicResolution()
 {
 	// separating render timestamp queries from dynres availability
 
-#if USE_ANDROID_SWAPPY
+#if USE_ANDROID_VULKAN_SWAPPY
 	
 	static const auto CVarAndroidSupportsDynamicResolution = IConsoleManager::Get().FindConsoleVariable(TEXT("r.Android.SupportsDynamicResolution"));
 	
@@ -615,7 +615,7 @@ bool FVulkanAndroidPlatform::SupportsDynamicResolution()
 		CVarAndroidSupportsDynamicResolution != nullptr &&
 		CVarAndroidSupportsDynamicResolution->GetBool(); // is supported
 
-#else // USE_ANDROID_SWAPPY
+#else // USE_ANDROID_VULKAN_SWAPPY
 
 	// defaulted to the previous code
 	return SupportsTimestampRenderQueries();
@@ -646,7 +646,7 @@ bool FVulkanAndroidPlatform::FramePace(FVulkanDevice& Device, void* WindowHandle
 	bool bVsyncMultiple = (CachedSyncInterval != 0);
 	int32 CurrentFramePace = FAndroidPlatformRHIFramePacer::GetFramePace();
 
-#if USE_ANDROID_SWAPPY
+#if USE_ANDROID_VULKAN_SWAPPY
 	if (FVulkanAndroidPlatform::bSwappyEnabledAtRHIInit && CurrentFramePace != 0)
 	{
 		// cache refresh rate and sync interval
@@ -689,7 +689,7 @@ bool FVulkanAndroidPlatform::FramePace(FVulkanDevice& Device, void* WindowHandle
 
 VkResult FVulkanAndroidPlatform::Present(VkQueue Queue, VkPresentInfoKHR& PresentInfo)
 {
-#if USE_ANDROID_SWAPPY
+#if USE_ANDROID_VULKAN_SWAPPY
 	if (FVulkanAndroidPlatform::bSwappyEnabledAtRHIInit)
 	{
 		return SwappyVk_queuePresent(Queue, &PresentInfo);
@@ -707,7 +707,7 @@ VkResult FVulkanAndroidPlatform::CreateSwapchainKHR(void* WindowHandle, VkPhysic
 
 	if (Result == VK_SUCCESS)
 	{
-#if USE_ANDROID_SWAPPY
+#if USE_ANDROID_VULKAN_SWAPPY
 		if (FVulkanAndroidPlatform::bSwappyEnabledAtRHIInit)
 		{
 			JNIEnv* Env = FAndroidApplication::GetJavaEnv();
@@ -745,7 +745,7 @@ VkResult FVulkanAndroidPlatform::CreateSwapchainKHR(void* WindowHandle, VkPhysic
 
 void FVulkanAndroidPlatform::DestroySwapchainKHR(VkDevice Device, VkSwapchainKHR Swapchain, const VkAllocationCallbacks* Allocator)
 {
-#if USE_ANDROID_SWAPPY
+#if USE_ANDROID_VULKAN_SWAPPY
 	if (FVulkanAndroidPlatform::bSwappyEnabledAtRHIInit)
 	{
 		SwappyVk_destroySwapchain(Device, Swapchain);
@@ -1699,7 +1699,7 @@ VkPipelineCache FVulkanAndroidPlatform::PrecompilePSO(
 
 void FAndroidVulkanFramePacer::Init()
 {
-#if USE_ANDROID_SWAPPY
+#if USE_ANDROID_VULKAN_SWAPPY
 	if (FAndroidPlatformRHIFramePacer::CVarUseSwappyForFramePacing.GetValueOnAnyThread() != 0)
 	{
 		FVulkanAndroidPlatform::bSwappyEnabledAtRHIInit = true;
diff --git a/Engine/Source/Runtime/VulkanRHI/Private/Android/VulkanAndroidPlatform.h b/Engine/Source/Runtime/VulkanRHI/Private/Android/VulkanAndroidPlatform.h
--- a/Engine/Source/Runtime/VulkanRHI/Private/Android/VulkanAndroidPlatform.h
+++ b/Engine/Source/Runtime/VulkanRHI/Private/Android/VulkanAndroidPlatform.h
@@ -176,7 +176,7 @@ public:
 
 	static void WriteCrashMarker(const FOptionalVulkanDeviceExtensions& OptionalExtensions, FVulkanCmdBuffer* CmdBuffer, VkBuffer DestBuffer, const TArrayView<uint32>& Entries, bool bAdding);
 
-#if USE_ANDROID_SWAPPY
+#if USE_ANDROID_VULKAN_SWAPPY
 	static bool bSwappyEnabledAtRHIInit;
 #endif
 
diff --git a/Engine/Source/Runtime/VulkanRHI/VulkanRHI.Build.cs b/Engine/Source/Runtime/VulkanRHI/VulkanRHI.Build.cs
--- a/Engine/Source/Runtime/VulkanRHI/VulkanRHI.Build.cs
+++ b/Engine/Source/Runtime/VulkanRHI/VulkanRHI.Build.cs
@@ -53,7 +53,7 @@ public class VulkanRHI : ModuleRules
 			);
 
 			// for Swappy
-			PublicDefinitions.Add("USE_ANDROID_SWAPPY=1");
+			PublicDefinitions.Add("USE_ANDROID_VULKAN_SWAPPY=1");
 			PrivateDependencyModuleNames.AddRange(
 				new string[]
 				{
