Force process termination when a Vulkan GPU crash is encountered.

Logging a fatal error doesn't immediately terminate the process, so we can end up logging the GPU crash multiple times if the failing function is being called in a tight loop. It's better to explicitly terminate the process after we report the problem.

#rnx

diff --git a/Engine/Source/Runtime/VulkanRHI/Private/VulkanUtil.cpp b/Engine/Source/Runtime/VulkanRHI/Private/VulkanUtil.cpp
--- a/Engine/Source/Runtime/VulkanRHI/Private/VulkanUtil.cpp
+++ b/Engine/Source/Runtime/VulkanRHI/Private/VulkanUtil.cpp
@@ -948,8 +948,8 @@ namespace VulkanRHI
 		}
 #endif
 
-		UE_LOG(LogVulkanRHI, Fatal, TEXT("%s failed, VkResult=%d\n at %s:%u \n with error %s"),
-			ANSI_TO_TCHAR(VkFunction), (int32)Result, ANSI_TO_TCHAR(Filename), Line, *ErrorString);
+		// Force shutdown, we can't do anything useful anymore.
+		FPlatformMisc::RequestExit(true, TEXT("VerifyVulkanResult"));
 	}
 }
 
