Fixed tonemap subpass making apps crash.
To check we were in MSAA, we were using the presence of color resolve attachments as a gate to bind the depth resolve attachment. This doesn't work with tonemap subpass as the color resolve attachment is used as another color attachment which led to a mismatch in the number of attachments between render pass and framebuffer.

diff --git a/Engine/Source/Runtime/RHI/Private/RHI.cpp b/Engine/Source/Runtime/RHI/Private/RHI.cpp
--- a/Engine/Source/Runtime/RHI/Private/RHI.cpp
+++ b/Engine/Source/Runtime/RHI/Private/RHI.cpp
@@ -1532,7 +1532,7 @@ void FRHIRenderPassInfo::ConvertToRenderTargetsInfo(FRHISetRenderTargetsInfo& Ou
 	OutRTInfo.bClearDepth = (DepthLoadAction == ERenderTargetLoadAction::EClear);
 	OutRTInfo.bClearStencil = (StencilLoadAction == ERenderTargetLoadAction::EClear);
 
-	if (OutRTInfo.bHasResolveAttachments && DepthStencilRenderTarget.ResolveTarget && DepthStencilRenderTarget.ResolveTarget != DepthStencilRenderTarget.DepthStencilTarget)
+	if (DepthStencilRenderTarget.ResolveTarget && DepthStencilRenderTarget.ResolveTarget != DepthStencilRenderTarget.DepthStencilTarget)
 	{
 		OutRTInfo.DepthStencilResolveRenderTarget = OutRTInfo.DepthStencilRenderTarget;
 		OutRTInfo.DepthStencilResolveRenderTarget.Texture = DepthStencilRenderTarget.ResolveTarget;
diff --git a/Engine/Source/Runtime/VulkanRHI/Private/VulkanRenderTarget.cpp b/Engine/Source/Runtime/VulkanRHI/Private/VulkanRenderTarget.cpp
--- a/Engine/Source/Runtime/VulkanRHI/Private/VulkanRenderTarget.cpp
+++ b/Engine/Source/Runtime/VulkanRHI/Private/VulkanRenderTarget.cpp
@@ -893,9 +893,9 @@ FVulkanRenderTargetLayout::FVulkanRenderTargetLayout(FVulkanDevice& InDevice, co
 		DepthReference.layout = DepthLayout;
 		StencilReference.stencilLayout = StencilLayout;
 
-		// Use depth/stencil resolve target only if we're MSAA and a RT has a resolve attachment
-		const bool bHasValidResolveAttachment = RTInfo.bHasResolveAttachments && RTInfo.DepthStencilResolveRenderTarget.Texture;
-		if (GRHISupportsDepthStencilResolve && CurrDesc.samples > VK_SAMPLE_COUNT_1_BIT && bHasValidResolveAttachment)
+		// Use depth/stencil resolve target only if we're MSAA
+		const bool bDepthStencilResolve = (RTInfo.DepthStencilRenderTarget.DepthStoreAction == ERenderTargetStoreAction::EMultisampleResolve) || (RTInfo.DepthStencilRenderTarget.GetStencilStoreAction() == ERenderTargetStoreAction::EMultisampleResolve);
+		if (GRHISupportsDepthStencilResolve && bDepthStencilResolve && CurrDesc.samples > VK_SAMPLE_COUNT_1_BIT && RTInfo.DepthStencilResolveRenderTarget.Texture)
 		{
 			Desc[NumAttachmentDescriptions + 1] = Desc[NumAttachmentDescriptions];
 			Desc[NumAttachmentDescriptions + 1].samples = VK_SAMPLE_COUNT_1_BIT;
diff --git a/Engine/Source/Runtime/VulkanRHI/Private/VulkanRenderpass.h b/Engine/Source/Runtime/VulkanRHI/Private/VulkanRenderpass.h
--- a/Engine/Source/Runtime/VulkanRHI/Private/VulkanRenderpass.h
+++ b/Engine/Source/Runtime/VulkanRHI/Private/VulkanRenderpass.h
@@ -572,7 +572,7 @@ public:
 			// depth attachment is same as input attachment
 			SubpassDesc.SetDepthStencilAttachment(&DepthStencilAttachment);
 
-			if (bResolveDepth)
+			if (bResolveDepth && !bCustomResolveSubpass)
 			{
 				SubpassDesc.SetDepthStencilResolveAttachment(&DepthStencilResolveSubpassDesc);
 			}
@@ -690,6 +690,13 @@ public:
 			InputAttachments3[1].attachment = ColorAttachmentReferences[0].attachment; // SceneColor as input
 			InputAttachments3[1].layout = VK_IMAGE_LAYOUT_SHADER_READ_ONLY_OPTIMAL;
 			InputAttachments3[1].SetAspect(VK_IMAGE_ASPECT_COLOR_BIT);
+
+			SubpassDesc.SetDepthStencilAttachment(&DepthStencilAttachment);
+
+			if (bResolveDepth)
+			{
+				SubpassDesc.SetDepthStencilResolveAttachment(&DepthStencilResolveSubpassDesc);
+			}
 		
 			SubpassDesc.SetInputAttachments(InputAttachments3, 2);
 			SubpassDesc.colorAttachmentCount = 1;
diff --git a/Engine/Source/Runtime/VulkanRHI/Private/VulkanViewport.cpp b/Engine/Source/Runtime/VulkanRHI/Private/VulkanViewport.cpp
--- a/Engine/Source/Runtime/VulkanRHI/Private/VulkanViewport.cpp
+++ b/Engine/Source/Runtime/VulkanRHI/Private/VulkanViewport.cpp
@@ -509,7 +509,7 @@ FVulkanFramebuffer::FVulkanFramebuffer(FVulkanDevice& Device, const FRHISetRende
 			AddExternalView(Texture->DefaultView);
 		}
 
-		if (InRTInfo.bHasResolveAttachments && RTLayout.GetHasDepthStencilResolve() && RTLayout.GetDepthStencilResolveAttachmentReference()->layout != VK_IMAGE_LAYOUT_UNDEFINED)
+		if (RTLayout.GetHasDepthStencilResolve() && RTLayout.GetDepthStencilResolveAttachmentReference()->layout != VK_IMAGE_LAYOUT_UNDEFINED)
 		{
 			FRHITexture* ResolveRHITexture = InRTInfo.DepthStencilResolveRenderTarget.Texture;
 			FVulkanTexture* ResolveTexture = ResourceCast(ResolveRHITexture);
