VulkanRHI: Add support for VK_KHR_ray_tracing_position_fetch and require it on Windows/Linux.  Defer metadata buffer creation until needed.

#rnx

diff --git a/Engine/Shaders/Private/RayTracing/TraceRayInline.ush b/Engine/Shaders/Private/RayTracing/TraceRayInline.ush
--- a/Engine/Shaders/Private/RayTracing/TraceRayInline.ush
+++ b/Engine/Shaders/Private/RayTracing/TraceRayInline.ush
@@ -33,7 +33,15 @@
 // #include "/Engine/Private/RayTracing/TraceRayInlineMetal.ush"
 #endif // COMPILER_VULKAN
 
-#define PLATFORM_SUPPORTS_INLINE_RAY_TRACING_TRIANGLE_NORMALS 0
+#ifndef PLATFORM_SUPPORTS_INLINE_RAY_TRACING_TRIANGLE_NORMALS
+#	define PLATFORM_SUPPORTS_INLINE_RAY_TRACING_TRIANGLE_NORMALS 0
+#endif // PLATFORM_SUPPORTS_INLINE_RAY_TRACING_TRIANGLE_NORMALS
+
+#if !PLATFORM_SUPPORTS_INLINE_RAY_TRACING_TRIANGLE_NORMALS
+// placeholder since platform is not expected to fill WorldGeometryNormal
+#define GetInlineRayTracingWorldGeometryNormal(RayQ, Mat, Committed) RayQ.WorldRayDirection()
+#endif
+
 #define PLATFORM_SUPPORTS_INLINE_RAY_TRACING_MAX_ITERATIONS 0
 #define PLATFORM_SUPPORTS_INLINE_RAY_TRACING_INSTANCE_MATRIX 1
 
@@ -105,8 +113,7 @@ FTraceRayInlineResult TraceRayInlineWithCallback(RaytracingAccelerationStructure
 			float3x4 WorldToObject3x4 = RayQ.CandidateWorldToObject3x4();
 			CandidateTraceRayResult.WorldToObject3x4 = WorldToObject3x4;
 
-			// placeholder
-			CandidateTraceRayResult.WorldGeometryNormal = RayQ.WorldRayDirection();
+			CandidateTraceRayResult.WorldGeometryNormal = GetInlineRayTracingWorldGeometryNormal(RayQ, TraceRayResult.WorldToObject3x4, false);
 
 			CandidateTraceRayResult.Bookmark[0] = asuint(TraceRayResult.HitT);
 			CandidateTraceRayResult.Bookmark[1] = asuint(TraceRayResult.HitT);
@@ -132,9 +139,6 @@ FTraceRayInlineResult TraceRayInlineWithCallback(RaytracingAccelerationStructure
 		TraceRayResult.bIsFrontFace = RayQ.CommittedTriangleFrontFace();
 		TraceRayResult.InstanceContributionToHitGroupIndex = RayQ.CommittedInstanceContributionToHitGroupIndex();
 
-		// placeholder
-		TraceRayResult.WorldGeometryNormal = RayQ.WorldRayDirection();
-
 		TraceRayResult.Bookmark[0] = asuint(TraceRayResult.HitT);
 		TraceRayResult.Bookmark[1] = asuint(TraceRayResult.HitT);
 
@@ -142,6 +146,8 @@ FTraceRayInlineResult TraceRayInlineWithCallback(RaytracingAccelerationStructure
 		TraceRayResult.WorldToObject3x4 = RayQ.CommittedWorldToObject3x4();
 #endif // PLATFORM_SUPPORTS_INLINE_RAY_TRACING_INSTANCE_MATRIX
 
+		TraceRayResult.WorldGeometryNormal = GetInlineRayTracingWorldGeometryNormal(RayQ, TraceRayResult.WorldToObject3x4, true);
+
 #if ENABLE_TRACE_RAY_INLINE_TRAVERSAL_STATISTICS
 		// statistics?
 #endif // ENABLE_TRACE_RAY_INLINE_TRAVERSAL_STATISTICS
@@ -160,15 +166,15 @@ FTraceRayInlineResult TraceRayInlineWithCallback(RaytracingAccelerationStructure
 		TraceRayResult.bIsFrontFace = true; // ignored for procedural
 		TraceRayResult.InstanceContributionToHitGroupIndex = RayQ.CommittedInstanceContributionToHitGroupIndex();
 
-		// placeholder
-		TraceRayResult.WorldGeometryNormal = RayQ.WorldRayDirection();
-
 		TraceRayResult.Bookmark[0] = asuint(TraceRayResult.HitT);
 		TraceRayResult.Bookmark[1] = asuint(TraceRayResult.HitT);
 
 #if PLATFORM_SUPPORTS_INLINE_RAY_TRACING_INSTANCE_MATRIX
 		TraceRayResult.WorldToObject3x4 = RayQ.CommittedWorldToObject3x4();
 #endif // PLATFORM_SUPPORTS_INLINE_RAY_TRACING_INSTANCE_MATRIX
+
+		TraceRayResult.WorldGeometryNormal = GetInlineRayTracingWorldGeometryNormal(RayQ, TraceRayResult.WorldToObject3x4, true);
+
 		break;
 	}
 #endif // ENABLE_TRACE_RAY_INLINE_PROCEDURAL_PRIMITIVE
@@ -190,10 +196,6 @@ FTraceRayInlineResult TraceRayInline(RaytracingAccelerationStructure Acceleratio
 
 #endif // OVERRIDE_TRACERAYINLINE_USH
 
-#ifndef PLATFORM_SUPPORTS_INLINE_RAY_TRACING_TRIANGLE_NORMALS
-#define PLATFORM_SUPPORTS_INLINE_RAY_TRACING_TRIANGLE_NORMALS 0
-#endif // PLATFORM_SUPPORTS_INLINE_RAY_TRACING_TRIANGLE_NORMALS
-
 #ifndef PLATFORM_SUPPORTS_INLINE_RAY_TRACING_MAX_ITERATIONS
 #define PLATFORM_SUPPORTS_INLINE_RAY_TRACING_MAX_ITERATIONS 0
 #endif // PLATFORM_SUPPORTS_INLINE_RAY_TRACING_MAX_ITERATIONS
diff --git a/Engine/Shaders/Private/RayTracing/TraceRayInlineVulkan.ush b/Engine/Shaders/Private/RayTracing/TraceRayInlineVulkan.ush
--- a/Engine/Shaders/Private/RayTracing/TraceRayInlineVulkan.ush
+++ b/Engine/Shaders/Private/RayTracing/TraceRayInlineVulkan.ush
@@ -12,7 +12,12 @@
 
 #if COMPILER_DXC && COMPILER_VULKAN && PLATFORM_SUPPORTS_INLINE_RAY_TRACING
 
-#define PLATFORM_SUPPORTS_INLINE_RAY_TRACING_TRIANGLE_ATTRIBUTES 1
+#ifndef VULKAN_SUPPORTS_RAY_TRACING_POSITION_FETCH
+#define VULKAN_SUPPORTS_RAY_TRACING_POSITION_FETCH 0
+#endif
+
+// Metadata buffer is not allocated when position_fetch is available
+#define PLATFORM_SUPPORTS_INLINE_RAY_TRACING_TRIANGLE_ATTRIBUTES !VULKAN_SUPPORTS_RAY_TRACING_POSITION_FETCH
 
 #if EMULATE_VKDEVICEADRESS
 inline uint2 AddEmulated64(uint2 a, uint b)
@@ -68,6 +73,50 @@ float3 VulkanLoadVertexPositionFloat3(VK_ADDR_TYPE SourceBufferBaseAddress, uint
 	return asfloat(VK_RAW_BUFFER_LOAD(uint3, SourceBufferBaseAddress, OffsetInBytes));
 }
 
+#if VULKAN_SUPPORTS_RAY_TRACING_POSITION_FETCH
+
+#define PLATFORM_SUPPORTS_INLINE_RAY_TRACING_TRIANGLE_NORMALS 1
+
+[[vk::ext_capability(5391 /* RayQueryPositionFetchKHR */)]]
+[[vk::ext_extension("SPV_KHR_ray_tracing_position_fetch")]]
+[[vk::ext_instruction(5340 /* OpRayQueryGetIntersectionTriangleVertexPositionsKHR */)]]
+float3 RayQueryGetIntersectionTriangleVertexPositionsKHR(
+#if ENABLE_TRACE_RAY_INLINE_PROCEDURAL_PRIMITIVE
+	[[vk::ext_reference]] RayQuery<RAY_FLAG_NONE> RayQ,
+#else
+	[[vk::ext_reference]] RayQuery<RAY_FLAG_SKIP_PROCEDURAL_PRIMITIVES> RayQ,
+#endif // ENABLE_TRACE_RAY_INLINE_PROCEDURAL_PRIMITIVE
+int committed)[3];
+
+float3 GetInlineRayTracingWorldGeometryNormal(
+#if ENABLE_TRACE_RAY_INLINE_PROCEDURAL_PRIMITIVE
+	RayQuery<RAY_FLAG_NONE> RayQ,
+#else
+	RayQuery<RAY_FLAG_SKIP_PROCEDURAL_PRIMITIVES> RayQ,
+#endif
+	float3x4 WorldToObject3x4,
+	bool Committed)
+{
+	FTriangleBaseAttributes Result = (FTriangleBaseAttributes)0;
+
+	float3 LocalNormal;
+	if (Committed)
+	{
+		float3 LocalPositions[3] = RayQueryGetIntersectionTriangleVertexPositionsKHR(RayQ, 1 /* RayQueryCommittedIntersectionKHR */);
+		LocalNormal = cross(LocalPositions[2] - LocalPositions[0], LocalPositions[1] - LocalPositions[0]);
+	}
+	else
+	{
+		float3 LocalPositions[3] = RayQueryGetIntersectionTriangleVertexPositionsKHR(RayQ, 0 /* RayQueryCandidateIntersectionKHR */);
+		LocalNormal = cross(LocalPositions[2] - LocalPositions[0], LocalPositions[1] - LocalPositions[0]);
+	}
+
+	float3 WorldNormal = normalize(mul(LocalNormal, (float3x3)WorldToObject3x4));
+	return WorldNormal;
+}
+
+#else // VULKAN_SUPPORTS_RAY_TRACING_POSITION_FETCH
+
 // Counterpart of FVulkanRayTracingGeometryParameters in VulkanRayTracing.cpp
 struct FVulkanRayTracingGeometryParameters
 {
@@ -123,4 +172,6 @@ FTriangleBaseAttributes LoadInlineRayTracingTriangleAttributes(
 	return Result;
 }
 
+#endif // VULKAN_SUPPORTS_RAY_TRACING_POSITION_FETCH
+
 #endif // COMPILER_DXC && COMPILER_VULKAN && PLATFORM_SUPPORTS_INLINE_RAY_TRACING
diff --git a/Engine/Source/Developer/VulkanShaderFormat/Private/SpirVShaderCompiler.inl b/Engine/Source/Developer/VulkanShaderFormat/Private/SpirVShaderCompiler.inl
--- a/Engine/Source/Developer/VulkanShaderFormat/Private/SpirVShaderCompiler.inl
+++ b/Engine/Source/Developer/VulkanShaderFormat/Private/SpirVShaderCompiler.inl
@@ -1383,6 +1383,9 @@ static void ModifyCompilerInput(FSpirvShaderCompilerInternalState& InternalState
 	if (Input.Environment.CompilerFlags.Contains(CFLAG_InlineRayTracing))
 	{
 		Input.Environment.SetDefine(TEXT("PLATFORM_SUPPORTS_INLINE_RAY_TRACING"), 1);
+
+		// Support is only garanteed on desktop currently
+		Input.Environment.SetDefine(TEXT("VULKAN_SUPPORTS_RAY_TRACING_POSITION_FETCH"), InternalState.IsAndroid() ? 0 : 1);
 	}
 
 	if (Input.Environment.CompilerFlags.Contains(CFLAG_AllowRealTypes))
diff --git a/Engine/Source/Runtime/VulkanRHI/Private/Android/VulkanAndroidPlatform.h b/Engine/Source/Runtime/VulkanRHI/Private/Android/VulkanAndroidPlatform.h
--- a/Engine/Source/Runtime/VulkanRHI/Private/Android/VulkanAndroidPlatform.h
+++ b/Engine/Source/Runtime/VulkanRHI/Private/Android/VulkanAndroidPlatform.h
@@ -24,6 +24,7 @@
 #define VULKAN_SUPPORTS_DRIVER_PROPERTIES			0
 #define VULKAN_SUPPORTS_DESCRIPTOR_INDEXING			1
 #define VULKAN_SUPPORTS_GPU_CRASH_DUMPS				1
+#define VULKAN_SUPPORTS_RAY_TRACING_POSITION_FETCH	0
 
 #define UE_VK_API_VERSION							VK_API_VERSION_1_1
 
diff --git a/Engine/Source/Runtime/VulkanRHI/Private/Linux/VulkanLinuxPlatform.h b/Engine/Source/Runtime/VulkanRHI/Private/Linux/VulkanLinuxPlatform.h
--- a/Engine/Source/Runtime/VulkanRHI/Private/Linux/VulkanLinuxPlatform.h
+++ b/Engine/Source/Runtime/VulkanRHI/Private/Linux/VulkanLinuxPlatform.h
@@ -12,6 +12,7 @@
 #define VULKAN_SIGNAL_UNIMPLEMENTED()				checkf(false, TEXT("Unimplemented vulkan functionality: %hs"), __PRETTY_FUNCTION__)
 #define VULKAN_SUPPORTS_AMD_BUFFER_MARKER			1
 #define VULKAN_SUPPORTS_SCALAR_BLOCK_LAYOUT			1
+#define VULKAN_SUPPORTS_RAY_TRACING_POSITION_FETCH	1
 
 #define UE_VK_API_VERSION							VK_API_VERSION_1_1
 
diff --git a/Engine/Source/Runtime/VulkanRHI/Private/VulkanExtensions.cpp b/Engine/Source/Runtime/VulkanRHI/Private/VulkanExtensions.cpp
--- a/Engine/Source/Runtime/VulkanRHI/Private/VulkanExtensions.cpp
+++ b/Engine/Source/Runtime/VulkanRHI/Private/VulkanExtensions.cpp
@@ -959,6 +959,40 @@ private:
 };
 
 
+// ***** VK_KHR_ray_tracing_position_fetch
+class FVulkanKHRRayTracingPositionFetchExtension : public FVulkanDeviceExtension
+{
+public:
+
+	FVulkanKHRRayTracingPositionFetchExtension(FVulkanDevice* InDevice)
+		: FVulkanDeviceExtension(InDevice, VK_KHR_RAY_TRACING_POSITION_FETCH_EXTENSION_NAME, VULKAN_EXTENSION_ENABLED)
+	{
+		bEnabledInCode = bEnabledInCode && GVulkanRayTracingCVar.GetValueOnAnyThread() && !FParse::Param(FCommandLine::Get(), TEXT("noraytracing"));
+	}
+
+	virtual void PrePhysicalDeviceFeatures(VkPhysicalDeviceFeatures2KHR& PhysicalDeviceFeatures2) override final
+	{
+		ZeroVulkanStruct(RayTracingPositionFetchFeatures, VK_STRUCTURE_TYPE_PHYSICAL_DEVICE_RAY_TRACING_POSITION_FETCH_FEATURES_KHR);
+		AddToPNext(PhysicalDeviceFeatures2, RayTracingPositionFetchFeatures);
+	}
+
+	virtual void PostPhysicalDeviceFeatures(FOptionalVulkanDeviceExtensions& ExtensionFlags) override final
+	{
+		bRequirementsPassed = (RayTracingPositionFetchFeatures.rayTracingPositionFetch == VK_TRUE);
+	}
+
+	virtual void PreCreateDevice(VkDeviceCreateInfo& DeviceCreateInfo) override final
+	{
+		if (bRequirementsPassed)
+		{
+			AddToPNext(DeviceCreateInfo, RayTracingPositionFetchFeatures);
+		}
+	}
+
+private:
+	VkPhysicalDeviceRayTracingPositionFetchFeaturesKHR  RayTracingPositionFetchFeatures;
+};
+
 
 // ***** VK_AMD_buffer_marker (vendor)
 class FVulkanAMDBufferMarkerExtension : public FVulkanDeviceExtension
@@ -1547,6 +1581,7 @@ FVulkanDeviceExtensionArray FVulkanDeviceExtension::GetUESupportedDeviceExtensio
 	ADD_CUSTOM_EXTENSION(FVulkanKHRAccelerationStructureExtension);
 	ADD_CUSTOM_EXTENSION(FVulkanKHRRayTracingPipelineExtension);
 	ADD_CUSTOM_EXTENSION(FVulkanKHRRayQueryExtension);
+	ADD_CUSTOM_EXTENSION(FVulkanKHRRayTracingPositionFetchExtension);
 
 	// Vendor extensions
 	ADD_CUSTOM_EXTENSION(FVulkanAMDBufferMarkerExtension);
diff --git a/Engine/Source/Runtime/VulkanRHI/Private/VulkanRayTracing.cpp b/Engine/Source/Runtime/VulkanRHI/Private/VulkanRayTracing.cpp
--- a/Engine/Source/Runtime/VulkanRHI/Private/VulkanRayTracing.cpp
+++ b/Engine/Source/Runtime/VulkanRHI/Private/VulkanRayTracing.cpp
@@ -169,6 +169,10 @@ static VkBuildAccelerationStructureFlagBitsKHR TranslateRayTracingAccelerationSt
 
 	checkf(!EnumHasAnyFlags(Flags, Flags), TEXT("Some ERayTracingAccelerationStructureFlags entries were not handled"));
 
+#if VULKAN_SUPPORTS_RAY_TRACING_POSITION_FETCH
+	Result |= (uint32)VK_BUILD_ACCELERATION_STRUCTURE_ALLOW_DATA_ACCESS_KHR;
+#endif
+
 	return VkBuildAccelerationStructureFlagBitsKHR(Result);
 }
 
@@ -337,6 +341,10 @@ static void GetBLASBuildData(
 	{
 		BuildData.GeometryInfo.flags |= VK_BUILD_ACCELERATION_STRUCTURE_ALLOW_COMPACTION_BIT_KHR;
 	}
+#if VULKAN_SUPPORTS_RAY_TRACING_POSITION_FETCH
+	BuildData.GeometryInfo.flags |= VK_BUILD_ACCELERATION_STRUCTURE_ALLOW_DATA_ACCESS_KHR;
+#endif
+
 	BuildData.GeometryInfo.mode = (BuildMode == EAccelerationStructureBuildMode::Build) ? VK_BUILD_ACCELERATION_STRUCTURE_MODE_BUILD_KHR : VK_BUILD_ACCELERATION_STRUCTURE_MODE_UPDATE_KHR;
 	BuildData.GeometryInfo.geometryCount = BuildData.Segments.Num();
 	BuildData.GeometryInfo.pGeometries = BuildData.Segments.GetData();
@@ -677,11 +685,6 @@ FVulkanRayTracingScene::FVulkanRayTracingScene(FRayTracingSceneInitializer2 InIn
 	INC_DWORD_STAT(STAT_VulkanRayTracingAllocatedTLAS);
 
 	SizeInfo = RHICalcRayTracingSceneSize(Initializer.MaxNumInstances, Initializer.BuildFlags);
-
-	const uint32 ParameterBufferSize = FMath::Max<uint32>(1, Initializer.NumTotalSegments) * sizeof(FVulkanRayTracingGeometryParameters);
-	FRHIResourceCreateInfo ParameterBufferCreateInfo(TEXT("RayTracingSceneMetadata"));
-	// Use FRHICommandListExecutor::GetImmediateCommandList as a temporary patch for 5.4 to avoid issues with RHI Validation
-	PerInstanceGeometryParameterBuffer = ResourceCast(FRHICommandListExecutor::GetImmediateCommandList().CreateBuffer(ParameterBufferSize, BUF_StructuredBuffer | BUF_ShaderResource, sizeof(FVulkanRayTracingGeometryParameters), ERHIAccess::SRVMask, ParameterBufferCreateInfo).GetReference());
 }
 
 FVulkanRayTracingScene::~FVulkanRayTracingScene()
@@ -837,16 +840,50 @@ void BuildAccelerationStructure(
 	Scene.bBuilt = true;
 }
 
+
+#if VULKAN_SUPPORTS_RAY_TRACING_POSITION_FETCH
+
+// Metadata buffer is unused with the support of position_fetch
+FVulkanResourceMultiBuffer* FVulkanRayTracingScene::GetOrCreateMetadataBuffer(FRHICommandListBase& RHICmdList) { return nullptr; }
+FRHIShaderResourceView* FVulkanRayTracingScene::GetOrCreateMetadataBufferSRV(FRHICommandListImmediate& RHICmdList) { return nullptr; }
+void FVulkanRayTracingScene::BuildPerInstanceGeometryParameterBuffer(FRHICommandListBase& RHICmdList) {}
+
+#else
+
+FVulkanResourceMultiBuffer* FVulkanRayTracingScene::GetOrCreateMetadataBuffer(FRHICommandListBase& RHICmdList)
+{
+	UE::TScopeLock Lock(Mutex);
+	if (!PerInstanceGeometryParameterBuffer.IsValid())
+	{
+		const uint32 ParameterBufferSize = FMath::Max<uint32>(1, Initializer.NumTotalSegments) * sizeof(FVulkanRayTracingGeometryParameters);
+		FRHIResourceCreateInfo ParameterBufferCreateInfo(TEXT("RayTracingSceneMetadata"));
+		PerInstanceGeometryParameterBuffer = ResourceCast(RHICmdList.CreateBuffer(ParameterBufferSize, BUF_Dynamic | BUF_StructuredBuffer | BUF_ShaderResource, sizeof(FVulkanRayTracingGeometryParameters), ERHIAccess::SRVMask, ParameterBufferCreateInfo).GetReference());
+	}
+	return PerInstanceGeometryParameterBuffer.GetReference();
+}
+
+FRHIShaderResourceView* FVulkanRayTracingScene::GetOrCreateMetadataBufferSRV(FRHICommandListImmediate& RHICmdList)
+{
+	if (!PerInstanceGeometryParameterSRV.IsValid())
+	{
+		PerInstanceGeometryParameterSRV = RHICmdList.CreateShaderResourceView(GetOrCreateMetadataBuffer(RHICmdList), FRHIViewDesc::CreateBufferSRV().SetType(FRHIViewDesc::EBufferType::Structured));
+	}
+
+	return PerInstanceGeometryParameterSRV.GetReference();
+}
+
 void FVulkanRayTracingScene::BuildPerInstanceGeometryParameterBuffer(FRHICommandListBase& RHICmdList)
 {
 	// TODO: we could cache parameters in the geometry object to avoid some of the pointer chasing (if this is measured to be a performance issue)
 
+	FVulkanResourceMultiBuffer* MetadataBuffer = GetOrCreateMetadataBuffer(RHICmdList);
+
 	const uint32 ParameterBufferSize = FMath::Max<uint32>(1, Initializer.NumTotalSegments) * sizeof(FVulkanRayTracingGeometryParameters);
-	check(PerInstanceGeometryParameterBuffer->GetSize() >= ParameterBufferSize);
+	check(MetadataBuffer->GetSize() >= ParameterBufferSize);
 
 	check(IsInRHIThread() || !IsRunningRHIInSeparateThread());
 
-	void* MappedBuffer = PerInstanceGeometryParameterBuffer->Lock(RHICmdList, RLM_WriteOnly, ParameterBufferSize, 0);
+	void* MappedBuffer = MetadataBuffer->Lock(RHICmdList, RLM_WriteOnly, ParameterBufferSize, 0);
 	FVulkanRayTracingGeometryParameters* MappedParameters = reinterpret_cast<FVulkanRayTracingGeometryParameters*>(MappedBuffer);
 	uint32 ParameterIndex = 0;
 
@@ -891,9 +928,10 @@ void FVulkanRayTracingScene::BuildPerInstanceGeometryParameterBuffer(FRHICommand
 
 	check(ParameterIndex == Initializer.NumTotalSegments);
 
-	PerInstanceGeometryParameterBuffer->Unlock(RHICmdList);
+	MetadataBuffer->Unlock(RHICmdList);
 }
 
+#endif // VULKAN_SUPPORTS_RAY_TRACING_POSITION_FETCH
 
 FVulkanRayTracingShaderTable::FVulkanRayTracingShaderTable(FVulkanDevice* Device, const FRayTracingShaderBindingTableInitializer& InInitializer)
 	: FRHIShaderBindingTable(InInitializer)
diff --git a/Engine/Source/Runtime/VulkanRHI/Private/VulkanRayTracing.h b/Engine/Source/Runtime/VulkanRHI/Private/VulkanRayTracing.h
--- a/Engine/Source/Runtime/VulkanRHI/Private/VulkanRayTracing.h
+++ b/Engine/Source/Runtime/VulkanRHI/Private/VulkanRayTracing.h
@@ -213,15 +213,8 @@ public:
 		}
 	}
 
-	virtual FRHIShaderResourceView* GetOrCreateMetadataBufferSRV(FRHICommandListImmediate& RHICmdList) override final
-	{
-		if (!PerInstanceGeometryParameterSRV.IsValid())
-		{
-			PerInstanceGeometryParameterSRV = RHICmdList.CreateShaderResourceView(PerInstanceGeometryParameterBuffer, FRHIViewDesc::CreateBufferSRV().SetType(FRHIViewDesc::EBufferType::Structured));
-		}
-
-		return PerInstanceGeometryParameterSRV.GetReference();
-	}
+	FVulkanResourceMultiBuffer* GetOrCreateMetadataBuffer(FRHICommandListBase& RHICmdList);
+	virtual FRHIShaderResourceView* GetOrCreateMetadataBufferSRV(FRHICommandListImmediate& RHICmdList) override final;
 
 	FRHIShaderBindingTable* FindOrCreateShaderBindingTable(const FRHIRayTracingPipelineState* Pipeline);
 
@@ -360,4 +353,3 @@ private:
 	FVulkanRayTracingCompactedSizeQueryPool* QueryPool = nullptr;
 };
 
-
diff --git a/Engine/Source/Runtime/VulkanRHI/Private/Windows/VulkanWindowsPlatform.h b/Engine/Source/Runtime/VulkanRHI/Private/Windows/VulkanWindowsPlatform.h
--- a/Engine/Source/Runtime/VulkanRHI/Private/Windows/VulkanWindowsPlatform.h
+++ b/Engine/Source/Runtime/VulkanRHI/Private/Windows/VulkanWindowsPlatform.h
@@ -16,6 +16,7 @@
 #define VULKAN_SIGNAL_UNIMPLEMENTED()				checkf(false, TEXT("Unimplemented vulkan functionality: %s"), StringCast<TCHAR>(__FUNCTION__).Get())
 #define VULKAN_SUPPORTS_AMD_BUFFER_MARKER			1
 #define VULKAN_SUPPORTS_SCALAR_BLOCK_LAYOUT			1
+#define VULKAN_SUPPORTS_RAY_TRACING_POSITION_FETCH	1
 
 #define UE_VK_API_VERSION							VK_API_VERSION_1_1
 
