VulkanRHI: exit query timeout loops and handle DEVICE_LOST when fetching query results

#rnx

diff --git a/Engine/Source/Runtime/VulkanRHI/Private/VulkanQuery.cpp b/Engine/Source/Runtime/VulkanRHI/Private/VulkanQuery.cpp
--- a/Engine/Source/Runtime/VulkanRHI/Private/VulkanQuery.cpp
+++ b/Engine/Source/Runtime/VulkanRHI/Private/VulkanQuery.cpp
@@ -135,7 +135,8 @@ bool FVulkanOcclusionQueryPool::InternalTryGetResults(bool bWait)
 					FTaskGraphInterface::Get().ProcessThreadUntilIdle(RenderThread_Local);
 				}
 
-				if (VulkanRHI::vkGetEventStatus(Device->GetInstanceHandle(), ResetEvent) == VK_EVENT_SET)
+				Result = VulkanRHI::vkGetEventStatus(Device->GetInstanceHandle(), ResetEvent);
+				if (Result == VK_EVENT_SET)
 				{
 					Result = VulkanRHI::vkGetQueryPoolResults(Device->GetInstanceHandle(), QueryPool, 0, NumUsedQueries, NumUsedQueries * sizeof(uint64), QueryOutput.GetData(), sizeof(uint64), VK_QUERY_RESULT_64_BIT);
 				}
@@ -145,7 +146,7 @@ bool FVulkanOcclusionQueryPool::InternalTryGetResults(bool bWait)
 					bSuccess = true;
 					break;
 				}
-				else if (Result == VK_NOT_READY)
+				else if ((Result == VK_NOT_READY) || (Result == VK_EVENT_RESET))
 				{
 					bSuccess = false;
 				}
