VulkanRHI: Clean Up: Pull the single use ub out of the states used to track samplers/resources.

#rnx

diff --git a/Engine/Source/Developer/VulkanShaderFormat/Private/SpirVShaderCompiler.inl b/Engine/Source/Developer/VulkanShaderFormat/Private/SpirVShaderCompiler.inl
--- a/Engine/Source/Developer/VulkanShaderFormat/Private/SpirVShaderCompiler.inl
+++ b/Engine/Source/Developer/VulkanShaderFormat/Private/SpirVShaderCompiler.inl
@@ -504,7 +504,7 @@ static void GatherSpirvReflectionBindings(
 	// If bindless is supported, then offset the descriptor set to fit the bindless heaps at the beginning
 	const EShaderFrequency ShaderFrequency = InternalState.GetShaderFrequency();
 	const uint32 StageIndex = (uint32)ShaderStage::GetStageForFrequency(ShaderFrequency);
-	const uint32 DescSetNo = InternalState.bSupportsBindless ? VulkanBindless::NumBindlessSets + StageIndex : StageIndex;
+	const uint32 DescSetNo = InternalState.bSupportsBindless ? VulkanBindless::MaxNumSets + StageIndex : StageIndex;
 
 	SpvReflectResult SpvResult = Reflection.EnumerateDescriptorSets(&NumDescriptorSets, nullptr);
 	check(SpvResult == SPV_REFLECT_RESULT_SUCCESS);
diff --git a/Engine/Source/Runtime/VulkanRHI/Private/VulkanDescriptorSets.cpp b/Engine/Source/Runtime/VulkanRHI/Private/VulkanDescriptorSets.cpp
--- a/Engine/Source/Runtime/VulkanRHI/Private/VulkanDescriptorSets.cpp
+++ b/Engine/Source/Runtime/VulkanRHI/Private/VulkanDescriptorSets.cpp
@@ -116,7 +116,7 @@ DEFINE_STAT(STAT_VulkanBindlessPeakAccelerationStructure);
 DEFINE_STAT(STAT_VulkanBindlessWritePerFrame);
 
 
-static inline uint8 GetIndexForDescriptorType(VkDescriptorType DescriptorType)
+static constexpr uint8 GetIndexForDescriptorType(VkDescriptorType DescriptorType)
 {
 	switch (DescriptorType)
 	{
@@ -131,11 +131,11 @@ static inline uint8 GetIndexForDescriptorType(VkDescriptorType DescriptorType)
 	default: checkNoEntry();
 	}
 
-	return VulkanBindless::NumBindlessSets;
+	return VulkanBindless::MaxNumSets;
 }
 
 
-static inline VkDescriptorType GetDescriptorTypeForSetIndex(uint8 SetIndex)
+static constexpr VkDescriptorType GetDescriptorTypeForSetIndex(uint8 SetIndex)
 {
 	switch (SetIndex)
 	{
@@ -251,8 +251,8 @@ bool FVulkanBindlessDescriptorManager::VerifySupport(FVulkanDevice* Device)
 		if (bMeetsExtensionsRequirements)
 		{
 			const bool bMeetsPropertiesRequirements =
-				(GpuProps.limits.maxBoundDescriptorSets >= VulkanBindless::NumBindlessSets) &&
-				(DescriptorBufferProperties.maxDescriptorBufferBindings >= VulkanBindless::NumBindlessSets) &&
+				(GpuProps.limits.maxBoundDescriptorSets >= VulkanBindless::MaxNumSets) &&
+				(DescriptorBufferProperties.maxDescriptorBufferBindings >= VulkanBindless::MaxNumSets) &&
 				(DescriptorBufferProperties.maxResourceDescriptorBufferBindings >= VulkanBindless::NumBindlessSets) &&
 				(DescriptorBufferProperties.maxSamplerDescriptorBufferBindings >= 1) &&
 				Device->GetDeviceMemoryManager().SupportsMemoryType(GetDescriptorBufferMemoryType(Device));
@@ -298,7 +298,7 @@ FVulkanBindlessDescriptorManager::FVulkanBindlessDescriptorManager(FVulkanDevice
 	, bIsSupported(VerifySupport(InDevice))
 {
 	FMemory::Memzero(BufferBindingInfo);
-	for (uint32 Index = 0; Index < VulkanBindless::NumBindlessSets; Index++)
+	for (uint32 Index = 0; Index < VulkanBindless::MaxNumSets; Index++)
 	{
 		BufferIndices[Index] = Index;
 	}
@@ -332,13 +332,15 @@ void FVulkanBindlessDescriptorManager::Deinit()
 		};
 
 		for (uint32 SetIndex = 0; SetIndex < VulkanBindless::NumBindlessSets; ++SetIndex)
-	{
+		{
 			BindlessSetState& State = BindlessSetStates[SetIndex];
 			if (State.DescriptorType != VK_DESCRIPTOR_TYPE_MAX_ENUM)
 			{
 				DestroyBindlessState(State);
 			}
 		}
+		
+		DestroyBindlessState(SingleUseUBState);
 
 		VulkanRHI::vkDestroyDescriptorSetLayout(DeviceHandle, EmptyDescriptorSetLayout, VULKAN_CPU_ALLOCATOR);
 		EmptyDescriptorSetLayout = VK_NULL_HANDLE;
@@ -555,6 +557,13 @@ void FVulkanBindlessDescriptorManager::Init()
 			TotalResourceDescriptorBufferSize += CreateDescriptorBuffer(State, BufferBindingInfo[SetIndex], IsSingleUseUniformBufferSet);
 		}
 
+		// Fill in the state for single-use UB
+		{
+			InitBindlessSetState(GetDescriptorTypeForSetIndex(VulkanBindless::BindlessSingleUseUniformBufferSet), SingleUseUBState);
+			SingleUseUBState.DescriptorSetLayout = CreateShaderStageUniformBufferLayout();
+			TotalResourceDescriptorBufferSize += CreateDescriptorBuffer(SingleUseUBState, BufferBindingInfo[VulkanBindless::BindlessSingleUseUniformBufferSet], true);
+		}
+
 		checkf(TotalResourceDescriptorBufferSize < DescriptorBufferProperties.resourceDescriptorBufferAddressSpaceSize,
 			TEXT("Combined resource descriptor buffer size of [%u] exceeded maximum [%llu]."),
 			TotalResourceDescriptorBufferSize, DescriptorBufferProperties.resourceDescriptorBufferAddressSpaceSize);
@@ -562,19 +571,20 @@ void FVulkanBindlessDescriptorManager::Init()
 
 	// Now create the single pipeline layout used by everything
 	{
-		VkDescriptorSetLayout DescriptorSetLayouts[VulkanBindless::NumBindlessSets ];
+		VkDescriptorSetLayout DescriptorSetLayouts[VulkanBindless::MaxNumSets];
 		for (int32 LayoutIndex = 0; LayoutIndex < VulkanBindless::NumBindlessSets; ++LayoutIndex)
 		{
 			const BindlessSetState& State = BindlessSetStates[LayoutIndex];
 			DescriptorSetLayouts[LayoutIndex] = State.DescriptorSetLayout;
 		}
+		DescriptorSetLayouts[VulkanBindless::BindlessSingleUseUniformBufferSet] = SingleUseUBState.DescriptorSetLayout;
 
 		VkPipelineLayoutCreateInfo PipelineLayoutCreateInfo;
 		ZeroVulkanStruct(PipelineLayoutCreateInfo, VK_STRUCTURE_TYPE_PIPELINE_LAYOUT_CREATE_INFO);
-		PipelineLayoutCreateInfo.setLayoutCount = VulkanBindless::NumBindlessSets;
+		PipelineLayoutCreateInfo.setLayoutCount = VulkanBindless::MaxNumSets;
 		PipelineLayoutCreateInfo.pSetLayouts = DescriptorSetLayouts;
 		VERIFYVULKANRESULT(VulkanRHI::vkCreatePipelineLayout(DeviceHandle, &PipelineLayoutCreateInfo, VULKAN_CPU_ALLOCATOR, &BindlessPipelineLayout));
-		VULKAN_SET_DEBUG_NAME((*Device), VK_OBJECT_TYPE_PIPELINE_LAYOUT, BindlessPipelineLayout, TEXT("BindlessPipelineLayout(SetCount=%d)"), VulkanBindless::NumBindlessSets);
+		VULKAN_SET_DEBUG_NAME((*Device), VK_OBJECT_TYPE_PIPELINE_LAYOUT, BindlessPipelineLayout, TEXT("BindlessPipelineLayout(SetCount=%d)"), VulkanBindless::MaxNumSets);
 	}
 }
 
@@ -582,23 +592,23 @@ void FVulkanBindlessDescriptorManager::BindDescriptorBuffers(VkCommandBuffer Com
 {
 	checkf(bIsSupported, TEXT("Trying to BindDescriptorBuffers but bindless is not supported!"));
 
-	VulkanRHI::vkCmdBindDescriptorBuffersEXT(CommandBuffer, VulkanBindless::NumBindlessSets, BufferBindingInfo);
+	VulkanRHI::vkCmdBindDescriptorBuffersEXT(CommandBuffer, VulkanBindless::MaxNumSets, BufferBindingInfo);
 
 	if (GVulkanBindlessBufferOffsetUpdates != 0)
 	{
-		VkDeviceSize BufferOffsets[VulkanBindless::NumBindlessSets];
+		VkDeviceSize BufferOffsets[VulkanBindless::MaxNumSets];
 		FMemory::Memzero(BufferOffsets);
 		if (SupportedStages & VK_PIPELINE_STAGE_FRAGMENT_SHADER_BIT)
 		{
-			VulkanRHI::vkCmdSetDescriptorBufferOffsetsEXT(CommandBuffer, VK_PIPELINE_BIND_POINT_GRAPHICS, BindlessPipelineLayout, 0, VulkanBindless::NumBindlessSets, BufferIndices, BufferOffsets);
+			VulkanRHI::vkCmdSetDescriptorBufferOffsetsEXT(CommandBuffer, VK_PIPELINE_BIND_POINT_GRAPHICS, BindlessPipelineLayout, 0, VulkanBindless::MaxNumSets, BufferIndices, BufferOffsets);
 		}
 		if (SupportedStages & VK_PIPELINE_STAGE_COMPUTE_SHADER_BIT)
 		{
-			VulkanRHI::vkCmdSetDescriptorBufferOffsetsEXT(CommandBuffer, VK_PIPELINE_BIND_POINT_COMPUTE, BindlessPipelineLayout, 0, VulkanBindless::NumBindlessSets, BufferIndices, BufferOffsets);
+			VulkanRHI::vkCmdSetDescriptorBufferOffsetsEXT(CommandBuffer, VK_PIPELINE_BIND_POINT_COMPUTE, BindlessPipelineLayout, 0, VulkanBindless::MaxNumSets, BufferIndices, BufferOffsets);
 		}
 		if (SupportedStages & VK_PIPELINE_STAGE_RAY_TRACING_SHADER_BIT_KHR)
 		{
-			VulkanRHI::vkCmdSetDescriptorBufferOffsetsEXT(CommandBuffer, VK_PIPELINE_BIND_POINT_RAY_TRACING_KHR, BindlessPipelineLayout, 0, VulkanBindless::NumBindlessSets, BufferIndices, BufferOffsets);
+			VulkanRHI::vkCmdSetDescriptorBufferOffsetsEXT(CommandBuffer, VK_PIPELINE_BIND_POINT_RAY_TRACING_KHR, BindlessPipelineLayout, 0, VulkanBindless::MaxNumSets, BufferIndices, BufferOffsets);
 		}
 	}
 }
@@ -607,22 +617,20 @@ void FVulkanBindlessDescriptorManager::RegisterUniformBuffers(VkCommandBuffer Co
 {
 	checkf(bIsSupported, TEXT("Trying to RegisterUniformBuffers but bindless is not supported!"));
 
-	BindlessSetState& BindlessUniformBufferSetState = BindlessSetStates[VulkanBindless::BindlessSingleUseUniformBufferSet];
-
 	// :todo-jn: Current uniform buffer layout is a bit wasteful with all the skipped bindings...
 	const uint32 BlockDescriptorCount = VulkanBindless::MaxUniformBuffersPerStage * ShaderStage::MaxNumStages;
-	const uint32 BlockSize = BlockDescriptorCount * BindlessUniformBufferSetState.DescriptorSize;
+	const uint32 BlockSize = BlockDescriptorCount * SingleUseUBState.DescriptorSize;
 	// Leave the first block always zeroed for easier debugging
-	const uint32 FirstDescriptorIndex = BlockDescriptorCount + (CurrentUniformBufferDescriptorIndex.fetch_add(BlockDescriptorCount) % (BindlessUniformBufferSetState.MaxDescriptorCount - (2*BlockDescriptorCount)));
-	const VkDeviceSize FirstDescriptorByteOffset = FirstDescriptorIndex * BindlessUniformBufferSetState.DescriptorSize;
+	const uint32 FirstDescriptorIndex = BlockDescriptorCount + (CurrentUniformBufferDescriptorIndex.fetch_add(BlockDescriptorCount) % (SingleUseUBState.MaxDescriptorCount - (2*BlockDescriptorCount)));
+	const VkDeviceSize FirstDescriptorByteOffset = FirstDescriptorIndex * SingleUseUBState.DescriptorSize;
 
-	VkDeviceSize BufferOffsets[VulkanBindless::NumBindlessSets];
+	VkDeviceSize BufferOffsets[VulkanBindless::MaxNumSets];
 	FMemory::Memzero(BufferOffsets);
 	BufferOffsets[VulkanBindless::BindlessSingleUseUniformBufferSet] = FirstDescriptorByteOffset;
 	checkSlow(FirstDescriptorByteOffset % Device->GetOptionalExtensionProperties().DescriptorBufferProps.descriptorBufferOffsetAlignment == 0);
 
 	// :todo-jn: Clear them for easier debugging for now
-	FMemory::Memzero(&BindlessUniformBufferSetState.DebugDescriptors[FirstDescriptorByteOffset], BlockSize);
+	FMemory::Memzero(&SingleUseUBState.DebugDescriptors[FirstDescriptorByteOffset], BlockSize);
 
 	for (int32 StageIndex = 0; StageIndex < ShaderStage::MaxNumStages; ++StageIndex)
 	{
@@ -643,36 +651,36 @@ void FVulkanBindlessDescriptorManager::RegisterUniformBuffers(VkCommandBuffer Co
 				const int32 BindingIndex = StageOffset + DescriptorAddressInfoIndex;
 				VkDeviceSize BindingByteOffset = 0;
 #if UE_BUILD_DEBUG
-				VulkanRHI::vkGetDescriptorSetLayoutBindingOffsetEXT(Device->GetInstanceHandle(), BindlessUniformBufferSetState.DescriptorSetLayout, BindingIndex, &BindingByteOffset);
-				check(BindingByteOffset == BindingIndex * BindlessUniformBufferSetState.DescriptorSize);
+				VulkanRHI::vkGetDescriptorSetLayoutBindingOffsetEXT(Device->GetInstanceHandle(), SingleUseUBState.DescriptorSetLayout, BindingIndex, &BindingByteOffset);
+				check(BindingByteOffset == BindingIndex * SingleUseUBState.DescriptorSize);
 #else
-				BindingByteOffset = (BindingIndex * BindlessUniformBufferSetState.DescriptorSize);
+				BindingByteOffset = (BindingIndex * SingleUseUBState.DescriptorSize);
 #endif
 
 				VkDescriptorGetInfoEXT Info;
 				ZeroVulkanStruct(Info, VK_STRUCTURE_TYPE_DESCRIPTOR_GET_INFO_EXT);
 				Info.type = VK_DESCRIPTOR_TYPE_UNIFORM_BUFFER;
 				Info.data.pUniformBuffer = &DescriptorAddressInfo;
-				VulkanRHI::vkGetDescriptorEXT(Device->GetInstanceHandle(), &Info, BindlessUniformBufferSetState.DescriptorSize, &BindlessUniformBufferSetState.DebugDescriptors[FirstDescriptorByteOffset + BindingByteOffset]);
+				VulkanRHI::vkGetDescriptorEXT(Device->GetInstanceHandle(), &Info, SingleUseUBState.DescriptorSize, &SingleUseUBState.DebugDescriptors[FirstDescriptorByteOffset + BindingByteOffset]);
 			}
 		}
 	}
 
 	// Copy all of them at once
-	FMemory::Memcpy(&BindlessUniformBufferSetState.MappedPointer[FirstDescriptorByteOffset], &BindlessUniformBufferSetState.DebugDescriptors[FirstDescriptorByteOffset], BlockSize);
+	FMemory::Memcpy(&SingleUseUBState.MappedPointer[FirstDescriptorByteOffset], &SingleUseUBState.DebugDescriptors[FirstDescriptorByteOffset], BlockSize);
 
 	{
 		SCOPED_NAMED_EVENT(vkCmdSetDescriptorBufferOffsetsEXT, FColor::Purple);
 
 		if (GVulkanBindlessRebindBuffers)
 		{
-			VulkanRHI::vkCmdBindDescriptorBuffersEXT(CommandBuffer, VulkanBindless::NumBindlessSets, BufferBindingInfo);
+			VulkanRHI::vkCmdBindDescriptorBuffersEXT(CommandBuffer, VulkanBindless::MaxNumSets, BufferBindingInfo);
 		}
 
 		const bool bSetAllOffsets = (GVulkanBindlessBufferOffsetUpdates == 0);
 		if (bSetAllOffsets)
 		{
-			VulkanRHI::vkCmdSetDescriptorBufferOffsetsEXT(CommandBuffer, BindPoint, BindlessPipelineLayout, 0u, VulkanBindless::NumBindlessSets, BufferIndices, BufferOffsets);
+			VulkanRHI::vkCmdSetDescriptorBufferOffsetsEXT(CommandBuffer, BindPoint, BindlessPipelineLayout, 0u, VulkanBindless::MaxNumSets, BufferIndices, BufferOffsets);
 		}
 		else
 		{
diff --git a/Engine/Source/Runtime/VulkanRHI/Private/VulkanDescriptorSets.h b/Engine/Source/Runtime/VulkanRHI/Private/VulkanDescriptorSets.h
--- a/Engine/Source/Runtime/VulkanRHI/Private/VulkanDescriptorSets.h
+++ b/Engine/Source/Runtime/VulkanRHI/Private/VulkanDescriptorSets.h
@@ -908,10 +908,11 @@ private:
 	};
 	BindlessSetState BindlessSetStates[VulkanBindless::NumBindlessSets];
 
+	BindlessSetState SingleUseUBState;
 	std::atomic<uint32> CurrentUniformBufferDescriptorIndex = 0;
 
-	VkDescriptorBufferBindingInfoEXT BufferBindingInfo[VulkanBindless::NumBindlessSets];
-	uint32_t BufferIndices[VulkanBindless::NumBindlessSets];
+	VkDescriptorBufferBindingInfoEXT BufferBindingInfo[VulkanBindless::MaxNumSets];
+	uint32_t BufferIndices[VulkanBindless::MaxNumSets];
 
 	VkPipelineLayout BindlessPipelineLayout = VK_NULL_HANDLE;
 
diff --git a/Engine/Source/Runtime/VulkanRHI/Public/VulkanCommon.h b/Engine/Source/Runtime/VulkanRHI/Public/VulkanCommon.h
--- a/Engine/Source/Runtime/VulkanRHI/Public/VulkanCommon.h
+++ b/Engine/Source/Runtime/VulkanRHI/Public/VulkanCommon.h
@@ -109,9 +109,14 @@ namespace VulkanBindless
 
 		BindlessAccelerationStructureSet,
 
-		BindlessSingleUseUniformBufferSet,  // Keep last
+		// Number of sets reserved for samplers/resources
 		NumBindlessSets,
-		MaxNumSets = NumBindlessSets
+
+		// Index of the descriptor set used for single use ub (like globals)
+		BindlessSingleUseUniformBufferSet = NumBindlessSets,
+
+		// Total number of descriptor sets used in a bindless pipeline
+		MaxNumSets = NumBindlessSets + 1
 	};
 };
 
