VulkanRHI: allocate Vulkan bindless single use UB from a list of blocks that can grow to accomodate larger maps.

#rnx

diff --git a/Engine/Source/Runtime/VulkanRHI/Private/VulkanDescriptorSets.cpp b/Engine/Source/Runtime/VulkanRHI/Private/VulkanDescriptorSets.cpp
--- a/Engine/Source/Runtime/VulkanRHI/Private/VulkanDescriptorSets.cpp
+++ b/Engine/Source/Runtime/VulkanRHI/Private/VulkanDescriptorSets.cpp
@@ -49,7 +49,7 @@ static FAutoConsoleVariableRef CVarVulkanBindlessMaxStorageTexelBufferCount(
 	ECVF_ReadOnly
 );
 
-int32 GVulkanBindlessMaxUniformBufferDescriptorCount = 768 * 1024;
+int32 GVulkanBindlessMaxUniformBufferDescriptorCount = 32 * 1024;
 static FAutoConsoleVariableRef CVarVulkanBindlessMaxUniformBufferCount(
 	TEXT("r.Vulkan.Bindless.MaxResourceUniformBufferCount"),
 	GVulkanBindlessMaxUniformBufferDescriptorCount,
@@ -73,23 +73,14 @@ static FAutoConsoleVariableRef CVarVulkanBindlessMaxAccelerationStructureCount(
 	ECVF_ReadOnly
 );
 
-int32 GVulkanBindlessRebindBuffers = 1;
-static FAutoConsoleVariableRef CVarVulkanBindlessRebindBuffers(
-	TEXT("r.Vulkan.Bindless.RebindBuffers"),
-	GVulkanBindlessRebindBuffers,
-	TEXT("Rebind buffers for every draw or dispatch.  Handy for debugging but not great for performance."),
+int32 GVulkanBindlessBlockSize = 1024 * 1024;
+static FAutoConsoleVariableRef CVarVulkanBindlessBlockSize(
+	TEXT("r.Vulkan.Bindless.BlockSize"),
+	GVulkanBindlessBlockSize,
+	TEXT("Block size to use for single use ub. (default: 1MB)"),
 	ECVF_RenderThreadSafe
 );
 
-int32 GVulkanBindlessBufferOffsetUpdates = 0;
-static FAutoConsoleVariableRef CVarVulkanBindlessBufferOffsetUpdates(
-	TEXT("r.Vulkan.Bindless.BufferOffsetUpdates"),
-	GVulkanBindlessBufferOffsetUpdates,
-	TEXT("0 to set all offsets for each draw/dispatch\n")\
-	TEXT("1 to set resource descriptor buffer offsets once, and only update for uniform buffer offsets on draw/dispatch\n"),
-	ECVF_ReadOnly
-);
-
 
 DECLARE_STATS_GROUP(TEXT("Vulkan Bindless"), STATGROUP_VulkanBindless, STATCAT_Advanced);
 
@@ -340,10 +331,14 @@ void FVulkanBindlessDescriptorManager::Deinit()
 			}
 		}
 		
-		DestroyBindlessState(SingleUseUBState);
+		VulkanRHI::vkDestroyDescriptorSetLayout(DeviceHandle, SingleUseUBDescriptorSetLayout, VULKAN_CPU_ALLOCATOR);
+		SingleUseUBDescriptorSetLayout = VK_NULL_HANDLE;
 
 		VulkanRHI::vkDestroyDescriptorSetLayout(DeviceHandle, EmptyDescriptorSetLayout, VULKAN_CPU_ALLOCATOR);
 		EmptyDescriptorSetLayout = VK_NULL_HANDLE;
+
+		delete SingleUseUBAllocator;
+		SingleUseUBAllocator = nullptr;
 	}
 }
 
@@ -357,6 +352,9 @@ void FVulkanBindlessDescriptorManager::Init()
 	const VkDevice DeviceHandle = Device->GetInstanceHandle();
 	const VkPhysicalDeviceDescriptorBufferPropertiesEXT& DescriptorBufferProperties = Device->GetOptionalExtensionProperties().DescriptorBufferProps;
 
+	const VkBufferUsageFlags BufferUsageFlags = VK_BUFFER_USAGE_SHADER_DEVICE_ADDRESS_BIT | VK_BUFFER_USAGE_RESOURCE_DESCRIPTOR_BUFFER_BIT_EXT;
+	SingleUseUBAllocator = new VulkanRHI::FTempBlockAllocator(Device, GVulkanBindlessBlockSize, DescriptorBufferProperties.descriptorBufferOffsetAlignment, BufferUsageFlags);
+
 	// Create the dummy layout for unsupported descriptor types
 	{
 		VkDescriptorSetLayoutCreateInfo EmptyDescriptorSetLayoutCreateInfo;
@@ -559,9 +557,7 @@ void FVulkanBindlessDescriptorManager::Init()
 
 		// Fill in the state for single-use UB
 		{
-			InitBindlessSetState(GetDescriptorTypeForSetIndex(VulkanBindless::BindlessSingleUseUniformBufferSet), SingleUseUBState);
-			SingleUseUBState.DescriptorSetLayout = CreateShaderStageUniformBufferLayout();
-			TotalResourceDescriptorBufferSize += CreateDescriptorBuffer(SingleUseUBState, BufferBindingInfo[VulkanBindless::BindlessSingleUseUniformBufferSet], true);
+			SingleUseUBDescriptorSetLayout = CreateShaderStageUniformBufferLayout();
 		}
 
 		checkf(TotalResourceDescriptorBufferSize < DescriptorBufferProperties.resourceDescriptorBufferAddressSpaceSize,
@@ -577,7 +573,7 @@ void FVulkanBindlessDescriptorManager::Init()
 			const BindlessSetState& State = BindlessSetStates[LayoutIndex];
 			DescriptorSetLayouts[LayoutIndex] = State.DescriptorSetLayout;
 		}
-		DescriptorSetLayouts[VulkanBindless::BindlessSingleUseUniformBufferSet] = SingleUseUBState.DescriptorSetLayout;
+		DescriptorSetLayouts[VulkanBindless::BindlessSingleUseUniformBufferSet] = SingleUseUBDescriptorSetLayout;
 
 		VkPipelineLayoutCreateInfo PipelineLayoutCreateInfo;
 		ZeroVulkanStruct(PipelineLayoutCreateInfo, VK_STRUCTURE_TYPE_PIPELINE_LAYOUT_CREATE_INFO);
@@ -592,45 +588,41 @@ void FVulkanBindlessDescriptorManager::BindDescriptorBuffers(VkCommandBuffer Com
 {
 	checkf(bIsSupported, TEXT("Trying to BindDescriptorBuffers but bindless is not supported!"));
 
-	VulkanRHI::vkCmdBindDescriptorBuffersEXT(CommandBuffer, VulkanBindless::MaxNumSets, BufferBindingInfo);
+	VulkanRHI::vkCmdBindDescriptorBuffersEXT(CommandBuffer, VulkanBindless::NumBindlessSets, BufferBindingInfo);
 
-	if (GVulkanBindlessBufferOffsetUpdates != 0)
+	VkDeviceSize BufferOffsets[VulkanBindless::NumBindlessSets];
+	FMemory::Memzero(BufferOffsets);
+	if (SupportedStages & VK_PIPELINE_STAGE_FRAGMENT_SHADER_BIT)
 	{
-		VkDeviceSize BufferOffsets[VulkanBindless::MaxNumSets];
-		FMemory::Memzero(BufferOffsets);
-		if (SupportedStages & VK_PIPELINE_STAGE_FRAGMENT_SHADER_BIT)
-		{
-			VulkanRHI::vkCmdSetDescriptorBufferOffsetsEXT(CommandBuffer, VK_PIPELINE_BIND_POINT_GRAPHICS, BindlessPipelineLayout, 0, VulkanBindless::MaxNumSets, BufferIndices, BufferOffsets);
-		}
-		if (SupportedStages & VK_PIPELINE_STAGE_COMPUTE_SHADER_BIT)
-		{
-			VulkanRHI::vkCmdSetDescriptorBufferOffsetsEXT(CommandBuffer, VK_PIPELINE_BIND_POINT_COMPUTE, BindlessPipelineLayout, 0, VulkanBindless::MaxNumSets, BufferIndices, BufferOffsets);
-		}
-		if (SupportedStages & VK_PIPELINE_STAGE_RAY_TRACING_SHADER_BIT_KHR)
-		{
-			VulkanRHI::vkCmdSetDescriptorBufferOffsetsEXT(CommandBuffer, VK_PIPELINE_BIND_POINT_RAY_TRACING_KHR, BindlessPipelineLayout, 0, VulkanBindless::MaxNumSets, BufferIndices, BufferOffsets);
-		}
+		VulkanRHI::vkCmdSetDescriptorBufferOffsetsEXT(CommandBuffer, VK_PIPELINE_BIND_POINT_GRAPHICS, BindlessPipelineLayout, 0, VulkanBindless::NumBindlessSets, BufferIndices, BufferOffsets);
+	}
+	if (SupportedStages & VK_PIPELINE_STAGE_COMPUTE_SHADER_BIT)
+	{
+		VulkanRHI::vkCmdSetDescriptorBufferOffsetsEXT(CommandBuffer, VK_PIPELINE_BIND_POINT_COMPUTE, BindlessPipelineLayout, 0, VulkanBindless::NumBindlessSets, BufferIndices, BufferOffsets);
+	}
+	if (SupportedStages & VK_PIPELINE_STAGE_RAY_TRACING_SHADER_BIT_KHR)
+	{
+		VulkanRHI::vkCmdSetDescriptorBufferOffsetsEXT(CommandBuffer, VK_PIPELINE_BIND_POINT_RAY_TRACING_KHR, BindlessPipelineLayout, 0, VulkanBindless::NumBindlessSets, BufferIndices, BufferOffsets);
 	}
 }
 
-void FVulkanBindlessDescriptorManager::RegisterUniformBuffers(VkCommandBuffer CommandBuffer, VkPipelineBindPoint BindPoint, const FUniformBufferDescriptorArrays& StageUBs)
+void FVulkanBindlessDescriptorManager::RegisterUniformBuffers(FVulkanCmdBuffer* CommandBuffer, VkPipelineBindPoint BindPoint, const FUniformBufferDescriptorArrays& StageUBs)
 {
 	checkf(bIsSupported, TEXT("Trying to RegisterUniformBuffers but bindless is not supported!"));
 
 	// :todo-jn: Current uniform buffer layout is a bit wasteful with all the skipped bindings...
+	const uint32 UBDescriptorSize = GetDescriptorTypeSize(Device, VK_DESCRIPTOR_TYPE_UNIFORM_BUFFER);
 	const uint32 BlockDescriptorCount = VulkanBindless::MaxUniformBuffersPerStage * ShaderStage::MaxNumStages;
-	const uint32 BlockSize = BlockDescriptorCount * SingleUseUBState.DescriptorSize;
-	// Leave the first block always zeroed for easier debugging
-	const uint32 FirstDescriptorIndex = BlockDescriptorCount + (CurrentUniformBufferDescriptorIndex.fetch_add(BlockDescriptorCount) % (SingleUseUBState.MaxDescriptorCount - (2*BlockDescriptorCount)));
-	const VkDeviceSize FirstDescriptorByteOffset = FirstDescriptorIndex * SingleUseUBState.DescriptorSize;
+	const uint32 BlockSize = BlockDescriptorCount * UBDescriptorSize;
+
+	VkDescriptorBufferBindingInfoEXT LocalBufferBindingInfo[VulkanBindless::MaxNumSets];
+	FMemory::Memcpy(LocalBufferBindingInfo, BufferBindingInfo, VulkanBindless::NumBindlessSets * sizeof(VkDescriptorBufferBindingInfoEXT));
 
 	VkDeviceSize BufferOffsets[VulkanBindless::MaxNumSets];
 	FMemory::Memzero(BufferOffsets);
-	BufferOffsets[VulkanBindless::BindlessSingleUseUniformBufferSet] = FirstDescriptorByteOffset;
-	checkSlow(FirstDescriptorByteOffset % Device->GetOptionalExtensionProperties().DescriptorBufferProps.descriptorBufferOffsetAlignment == 0);
 
-	// :todo-jn: Clear them for easier debugging for now
-	FMemory::Memzero(&SingleUseUBState.DebugDescriptors[FirstDescriptorByteOffset], BlockSize);
+	uint8* MappedPointer = SingleUseUBAllocator->Alloc(BlockSize, CommandBuffer, 
+		LocalBufferBindingInfo[VulkanBindless::BindlessSingleUseUniformBufferSet], BufferOffsets[VulkanBindless::BindlessSingleUseUniformBufferSet]);
 
 	for (int32 StageIndex = 0; StageIndex < ShaderStage::MaxNumStages; ++StageIndex)
 	{
@@ -651,43 +643,24 @@ void FVulkanBindlessDescriptorManager::RegisterUniformBuffers(VkCommandBuffer Co
 				const int32 BindingIndex = StageOffset + DescriptorAddressInfoIndex;
 				VkDeviceSize BindingByteOffset = 0;
 #if UE_BUILD_DEBUG
-				VulkanRHI::vkGetDescriptorSetLayoutBindingOffsetEXT(Device->GetInstanceHandle(), SingleUseUBState.DescriptorSetLayout, BindingIndex, &BindingByteOffset);
-				check(BindingByteOffset == BindingIndex * SingleUseUBState.DescriptorSize);
+				VulkanRHI::vkGetDescriptorSetLayoutBindingOffsetEXT(Device->GetInstanceHandle(), SingleUseUBDescriptorSetLayout, BindingIndex, &BindingByteOffset);
+				check(BindingByteOffset == BindingIndex * UBDescriptorSize);
 #else
-				BindingByteOffset = (BindingIndex * SingleUseUBState.DescriptorSize);
+				BindingByteOffset = (BindingIndex * UBDescriptorSize);
 #endif
 
 				VkDescriptorGetInfoEXT Info;
 				ZeroVulkanStruct(Info, VK_STRUCTURE_TYPE_DESCRIPTOR_GET_INFO_EXT);
 				Info.type = VK_DESCRIPTOR_TYPE_UNIFORM_BUFFER;
 				Info.data.pUniformBuffer = &DescriptorAddressInfo;
-				VulkanRHI::vkGetDescriptorEXT(Device->GetInstanceHandle(), &Info, SingleUseUBState.DescriptorSize, &SingleUseUBState.DebugDescriptors[FirstDescriptorByteOffset + BindingByteOffset]);
+				VulkanRHI::vkGetDescriptorEXT(Device->GetInstanceHandle(), &Info, UBDescriptorSize, &MappedPointer[BindingByteOffset]);
 			}
 		}
 	}
 
-	// Copy all of them at once
-	FMemory::Memcpy(&SingleUseUBState.MappedPointer[FirstDescriptorByteOffset], &SingleUseUBState.DebugDescriptors[FirstDescriptorByteOffset], BlockSize);
-
-	{
-		SCOPED_NAMED_EVENT(vkCmdSetDescriptorBufferOffsetsEXT, FColor::Purple);
-
-		if (GVulkanBindlessRebindBuffers)
-		{
-			VulkanRHI::vkCmdBindDescriptorBuffersEXT(CommandBuffer, VulkanBindless::MaxNumSets, BufferBindingInfo);
-		}
-
-		const bool bSetAllOffsets = (GVulkanBindlessBufferOffsetUpdates == 0);
-		if (bSetAllOffsets)
-		{
-			VulkanRHI::vkCmdSetDescriptorBufferOffsetsEXT(CommandBuffer, BindPoint, BindlessPipelineLayout, 0u, VulkanBindless::MaxNumSets, BufferIndices, BufferOffsets);
-		}
-		else
-		{
-			const uint8 SetIndex = VulkanBindless::BindlessSingleUseUniformBufferSet;
-			VulkanRHI::vkCmdSetDescriptorBufferOffsetsEXT(CommandBuffer, BindPoint, BindlessPipelineLayout, SetIndex, 1u, &BufferIndices[SetIndex], &BufferOffsets[SetIndex]);
-		}
-	}
+	// todo-jn: cache these states and only repeat when necessary
+	VulkanRHI::vkCmdBindDescriptorBuffersEXT(CommandBuffer->GetHandle(), VulkanBindless::MaxNumSets, LocalBufferBindingInfo);
+	VulkanRHI::vkCmdSetDescriptorBufferOffsetsEXT(CommandBuffer->GetHandle(), BindPoint, BindlessPipelineLayout, 0u, VulkanBindless::MaxNumSets, BufferIndices, BufferOffsets);
 }
 
 void FVulkanBindlessDescriptorManager::UpdateStatsForHandle(FRHIDescriptorHandle DescriptorHandle)
@@ -941,3 +914,10 @@ void FVulkanBindlessDescriptorManager::Unregister(FRHIDescriptorHandle Descripto
 	}
 }
 
+void FVulkanBindlessDescriptorManager::UpdateUBAllocator()
+{
+	if (bIsSupported && SingleUseUBAllocator)
+	{
+		SingleUseUBAllocator->UpdateBlocks();
+	}
+}
diff --git a/Engine/Source/Runtime/VulkanRHI/Private/VulkanDescriptorSets.h b/Engine/Source/Runtime/VulkanRHI/Private/VulkanDescriptorSets.h
--- a/Engine/Source/Runtime/VulkanRHI/Private/VulkanDescriptorSets.h
+++ b/Engine/Source/Runtime/VulkanRHI/Private/VulkanDescriptorSets.h
@@ -879,10 +879,11 @@ public:
 	void UpdateTexelBuffer(FRHIDescriptorHandle DescriptorHandle, const VkBufferViewCreateInfo& ViewInfo, bool bImmediateUpdate = true);
 	void UpdateAccelerationStructure(FRHIDescriptorHandle DescriptorHandle, VkAccelerationStructureKHR AccelerationStructure, bool bImmediateUpdate = true);
 
-	void RegisterUniformBuffers(VkCommandBuffer CommandBuffer, VkPipelineBindPoint BindPoint, const FUniformBufferDescriptorArrays& StageUBs);
+	void RegisterUniformBuffers(FVulkanCmdBuffer* CommandBuffer, VkPipelineBindPoint BindPoint, const FUniformBufferDescriptorArrays& StageUBs);
 
 	void Unregister(FRHIDescriptorHandle DescriptorHandle);
 
+	void UpdateUBAllocator();
 
 private:
 	const bool bIsSupported;
@@ -908,10 +909,10 @@ private:
 	};
 	BindlessSetState BindlessSetStates[VulkanBindless::NumBindlessSets];
 
-	BindlessSetState SingleUseUBState;
-	std::atomic<uint32> CurrentUniformBufferDescriptorIndex = 0;
+	VkDescriptorSetLayout SingleUseUBDescriptorSetLayout = VK_NULL_HANDLE;
+	VulkanRHI::FTempBlockAllocator* SingleUseUBAllocator = nullptr;
 
-	VkDescriptorBufferBindingInfoEXT BufferBindingInfo[VulkanBindless::MaxNumSets];
+	VkDescriptorBufferBindingInfoEXT BufferBindingInfo[VulkanBindless::NumBindlessSets];
 	uint32_t BufferIndices[VulkanBindless::MaxNumSets];
 
 	VkPipelineLayout BindlessPipelineLayout = VK_NULL_HANDLE;
diff --git a/Engine/Source/Runtime/VulkanRHI/Private/VulkanMemory.cpp b/Engine/Source/Runtime/VulkanRHI/Private/VulkanMemory.cpp
--- a/Engine/Source/Runtime/VulkanRHI/Private/VulkanMemory.cpp
+++ b/Engine/Source/Runtime/VulkanRHI/Private/VulkanMemory.cpp
@@ -90,6 +90,7 @@ DECLARE_MEMORY_STAT_EXTERN(TEXT("ImageOther"), STAT_VulkanAllocation_ImageOther,
 DECLARE_MEMORY_STAT_EXTERN(TEXT("BufferUAV"), STAT_VulkanAllocation_BufferUAV, STATGROUP_VulkanMemory, );
 DECLARE_MEMORY_STAT_EXTERN(TEXT("BufferStaging"), STAT_VulkanAllocation_BufferStaging, STATGROUP_VulkanMemory, );
 DECLARE_MEMORY_STAT_EXTERN(TEXT("BufferOther"), STAT_VulkanAllocation_BufferOther, STATGROUP_VulkanMemory, );
+DECLARE_MEMORY_STAT_EXTERN(TEXT("TempBlocks"), STAT_VulkanAllocation_TempBlocks, STATGROUP_VulkanMemory, );
 DECLARE_MEMORY_STAT_EXTERN(TEXT("_Total"), STAT_VulkanAllocation_Allocated, STATGROUP_VulkanMemory, );
 
 DEFINE_STAT(STAT_VulkanAllocation_UniformBuffer);
@@ -101,6 +102,7 @@ DEFINE_STAT(STAT_VulkanAllocation_ImageOther);
 DEFINE_STAT(STAT_VulkanAllocation_BufferUAV);
 DEFINE_STAT(STAT_VulkanAllocation_BufferStaging);
 DEFINE_STAT(STAT_VulkanAllocation_BufferOther);
+DEFINE_STAT(STAT_VulkanAllocation_TempBlocks);
 DEFINE_STAT(STAT_VulkanAllocation_Allocated);
 
 int32 GVulkanLogDefrag = 0;
@@ -4893,6 +4895,178 @@ namespace VulkanRHI
 		}
 	}
 
+
+
+
+	FTempBlockAllocator::FTempBlockAllocator(FVulkanDevice* InDevice, uint32 InBlockSize, uint32 InBlockAlignment, VkBufferUsageFlags InBufferUsage)
+		: FDeviceChild(InDevice)
+		, BlockSize(InBlockSize)
+		, BlockAlignment(InBlockAlignment)
+		, BufferUsage(InBufferUsage)
+	{
+		CurrentBlock = AllocBlock();
+	}
+
+	FTempBlockAllocator::~FTempBlockAllocator()
+	{
+		auto FreeBlock = [](FVulkanDevice* InDevice, FTempMemoryBlock* Block)
+		{
+			VulkanRHI::vkDestroyBuffer(InDevice->GetInstanceHandle(), Block->Buffer, VULKAN_CPU_ALLOCATOR);
+			InDevice->GetMemoryManager().FreeVulkanAllocation(Block->Allocation, VulkanRHI::EVulkanFreeFlag_DontDefer);
+			delete Block;
+		};
+
+		for (FTempMemoryBlock* Block : AvailableBlocks)
+		{
+			FreeBlock(Device, Block);
+		}
+		AvailableBlocks.Empty();
+
+		for (FTempMemoryBlock* Block : BusyBlocks)
+		{
+			FreeBlock(Device, Block);
+		}
+		BusyBlocks.Empty();
+
+		FreeBlock(Device, CurrentBlock);
+		CurrentBlock = nullptr;
+	}
+
+	FTempBlockAllocator::FTempMemoryBlock* FTempBlockAllocator::AllocBlock()
+	{
+		FTempMemoryBlock* NewBlock = new FTempMemoryBlock;
+
+		// Create buffer
+		{
+			VkBufferCreateInfo BufferCreateInfo;
+			ZeroVulkanStruct(BufferCreateInfo, VK_STRUCTURE_TYPE_BUFFER_CREATE_INFO);
+			BufferCreateInfo.size = BlockSize;
+			BufferCreateInfo.usage = BufferUsage;
+			VERIFYVULKANRESULT(VulkanRHI::vkCreateBuffer(Device->GetInstanceHandle(), &BufferCreateInfo, VULKAN_CPU_ALLOCATOR, &NewBlock->Buffer));
+		}
+
+		// Allocate memory
+		{
+			const VulkanRHI::EVulkanAllocationFlags AllocFlags =
+				VulkanRHI::EVulkanAllocationFlags::HostVisible |
+				VulkanRHI::EVulkanAllocationFlags::PreferBAR |
+				VulkanRHI::EVulkanAllocationFlags::AutoBind |
+				VulkanRHI::EVulkanAllocationFlags::Dedicated;
+			Device->GetMemoryManager().AllocateBufferMemory(NewBlock->Allocation, NewBlock->Buffer, AllocFlags, TEXT("FTempBlockAllocator"), BlockAlignment);
+
+			// Pull the stat of the generic "buffer other" to put it in the temp block bucket (will get cleaned up soon)
+			DEC_DWORD_STAT_BY(STAT_VulkanAllocation_BufferOther, NewBlock->Allocation.Size);
+			INC_DWORD_STAT_BY(STAT_VulkanAllocation_TempBlocks, NewBlock->Allocation.Size);
+
+			NewBlock->MappedPointer = (uint8*)NewBlock->Allocation.GetMappedPointer(Device);
+		}
+
+		// Get the device addr if needed
+		if (VKHasAllFlags(BufferUsage, VK_BUFFER_USAGE_SHADER_DEVICE_ADDRESS_BIT) && Device->GetOptionalExtensions().HasBufferDeviceAddress)
+		{
+			VkBufferDeviceAddressInfo BufferInfo;
+			ZeroVulkanStruct(BufferInfo, VK_STRUCTURE_TYPE_BUFFER_DEVICE_ADDRESS_INFO);
+			BufferInfo.buffer = NewBlock->Buffer;
+			NewBlock->BufferAddress = VulkanRHI::vkGetBufferDeviceAddressKHR(Device->GetInstanceHandle(), &BufferInfo);
+		}
+
+		return NewBlock;
+	}
+
+	uint8* FTempBlockAllocator::Alloc(uint32 InSize, FVulkanCmdBuffer* CmdBuffer, VkDescriptorBufferBindingInfoEXT& OutBindingInfo, VkDeviceSize& OutOffset)
+	{
+		const uint32 AlignedSize = Align(InSize, BlockAlignment);
+		checkfSlow(BlockAlignment < BlockSize, TEXT("Requested size of %d (%d aligned) is too large for block size of %d"), InSize, AlignedSize, BlockSize);
+
+		auto ReturnAllocationFromBlock = [AlignedSize, BufferUsageFlags=BufferUsage, &OutBindingInfo, &OutOffset](FTempMemoryBlock* Block, uint32 Offset)
+		{
+			ZeroVulkanStruct(OutBindingInfo, VK_STRUCTURE_TYPE_DESCRIPTOR_BUFFER_BINDING_INFO_EXT);
+			OutBindingInfo.address = Block->BufferAddress;
+			OutBindingInfo.usage = BufferUsageFlags;
+			OutOffset = Offset;
+			return Block->MappedPointer + Offset;
+		};
+
+		// Parallel scope: allocate in existing block from existing command buffer
+		{
+			FReadScopeLock ScopeLock(RWLock);
+
+			if (CurrentBlock->Fences.Contains(CmdBuffer))
+			{
+				CurrentBlock->Fences[CmdBuffer] = CmdBuffer->GetFenceSignaledCounter();
+
+				const uint32 AllocOffset = CurrentBlock->CurrentOffset.fetch_add(AlignedSize);
+				if (AllocOffset + InSize < BlockSize)
+				{
+					return ReturnAllocationFromBlock(CurrentBlock, AllocOffset);
+				}
+			}
+		}
+
+		// Locked path (allocate a new block or from a new command buffer)
+		{
+			FWriteScopeLock ScopeLock(RWLock);
+
+			// Make sure someone else didn't swap the block before it was our turn
+			uint32 AllocOffset = CurrentBlock->CurrentOffset.fetch_add(AlignedSize);
+			if (AllocOffset + InSize < BlockSize)
+			{
+				CurrentBlock->Fences.FindOrAdd(CmdBuffer) = CmdBuffer->GetFenceSignaledCounter();
+				return ReturnAllocationFromBlock(CurrentBlock, AllocOffset);
+			}
+
+			// Get a new block
+			BusyBlocks.Add(CurrentBlock);
+			if (AvailableBlocks.Num())
+			{
+				CurrentBlock = AvailableBlocks.Pop(EAllowShrinking::No);
+			}
+			else
+			{
+				CurrentBlock = AllocBlock();
+			}
+
+			CurrentBlock->Fences.FindOrAdd(CmdBuffer) = CmdBuffer->GetFenceSignaledCounter();
+
+			AllocOffset = CurrentBlock->CurrentOffset.fetch_add(AlignedSize);
+			checkSlow(AllocOffset == 0);
+			checkSlow(AllocOffset + InSize < BlockSize);
+			return ReturnAllocationFromBlock(CurrentBlock, AllocOffset);
+		}	
+	}
+
+	void FTempBlockAllocator::UpdateBlocks()
+	{
+		FWriteScopeLock ScopeLock(RWLock);
+
+		for (int32 Index = BusyBlocks.Num() - 1; Index >= 0; --Index)
+		{
+			FTempMemoryBlock* Block = BusyBlocks[Index];
+
+			bool BlockReady = true;
+			for (auto& Pair : Block->Fences)
+			{
+				FVulkanCmdBuffer* CmdBuffer = Pair.Key;
+				if (Pair.Value >= CmdBuffer->GetFenceSignaledCounter())
+				{
+					BlockReady = false;
+					break;
+				}
+			}
+
+			if (BlockReady)
+			{
+				Block->Fences.Empty();
+				Block->CurrentOffset = 0;
+				BusyBlocks.RemoveAtSwap(Index, EAllowShrinking::No);
+				AvailableBlocks.Add(Block);
+			}
+		}
+	}
+
+
+
+
 	FTempFrameAllocationBuffer::FTempFrameAllocationBuffer(FVulkanDevice* InDevice)
 		: FDeviceChild(InDevice)
 		, BufferIndex(0)
diff --git a/Engine/Source/Runtime/VulkanRHI/Private/VulkanMemory.h b/Engine/Source/Runtime/VulkanRHI/Private/VulkanMemory.h
--- a/Engine/Source/Runtime/VulkanRHI/Private/VulkanMemory.h
+++ b/Engine/Source/Runtime/VulkanRHI/Private/VulkanMemory.h
@@ -1118,6 +1118,43 @@ namespace VulkanRHI
 	};
 
 
+	// Simple temp allocation blocks used for volatile allocations
+	class FTempBlockAllocator : public FDeviceChild
+	{
+	public:
+		struct FTempMemoryBlock
+		{
+			FVulkanAllocation Allocation;
+			VkBuffer Buffer = 0;
+			uint8* MappedPointer = 0;
+			VkDeviceAddress BufferAddress = 0;
+
+			// Used by owners to suballocate
+			TMap<FVulkanCmdBuffer*, uint64> Fences;
+			std::atomic<uint32> CurrentOffset = 0;
+		};
+
+		FTempBlockAllocator(FVulkanDevice* InDevice, uint32 InBlockSize, uint32 InBlockAlignment, VkBufferUsageFlags InBufferUsage);
+		virtual ~FTempBlockAllocator();
+
+		uint8* Alloc(uint32 InSize, FVulkanCmdBuffer* CmdBuffer, VkDescriptorBufferBindingInfoEXT& OutBindingInfo, VkDeviceSize& OutOffset);
+
+		void UpdateBlocks();
+
+	protected:
+		FTempMemoryBlock* AllocBlock();
+
+		const uint32 BlockSize;
+		const uint32 BlockAlignment;
+		const VkBufferUsageFlags BufferUsage;
+
+		FTempMemoryBlock* CurrentBlock = nullptr;
+		TArray<FTempMemoryBlock*> BusyBlocks;
+		TArray<FTempMemoryBlock*> AvailableBlocks;
+		FRWLock RWLock;
+	};
+
+
 	// Simple tape allocation per frame for a VkBuffer, used for Volatile allocations
 	class FTempFrameAllocationBuffer : public FDeviceChild
 	{
diff --git a/Engine/Source/Runtime/VulkanRHI/Private/VulkanPipelineState.cpp b/Engine/Source/Runtime/VulkanRHI/Private/VulkanPipelineState.cpp
--- a/Engine/Source/Runtime/VulkanRHI/Private/VulkanPipelineState.cpp
+++ b/Engine/Source/Runtime/VulkanRHI/Private/VulkanPipelineState.cpp
@@ -278,7 +278,7 @@ void FVulkanComputePipelineDescriptorState::UpdateBindlessDescriptors(FVulkanCom
 	}
 
 	// Send to descriptor manager
-	Device->GetBindlessDescriptorManager()->RegisterUniformBuffers(CmdBuffer->GetHandle(), VK_PIPELINE_BIND_POINT_COMPUTE, StageUBs);
+	Device->GetBindlessDescriptorManager()->RegisterUniformBuffers(CmdBuffer, VK_PIPELINE_BIND_POINT_COMPUTE, StageUBs);
 }
 
 FVulkanGraphicsPipelineDescriptorState::FVulkanGraphicsPipelineDescriptorState(FVulkanDevice* InDevice, FVulkanRHIGraphicsPipelineState* InGfxPipeline)
@@ -493,7 +493,7 @@ void FVulkanGraphicsPipelineDescriptorState::UpdateBindlessDescriptors(FVulkanCo
 	}
 
 	// Send to descriptor manager
-	Device->GetBindlessDescriptorManager()->RegisterUniformBuffers(CmdBuffer->GetHandle(), VK_PIPELINE_BIND_POINT_GRAPHICS, StageUBs);
+	Device->GetBindlessDescriptorManager()->RegisterUniformBuffers(CmdBuffer, VK_PIPELINE_BIND_POINT_GRAPHICS, StageUBs);
 }
 
 
diff --git a/Engine/Source/Runtime/VulkanRHI/Private/VulkanRHI.cpp b/Engine/Source/Runtime/VulkanRHI/Private/VulkanRHI.cpp
--- a/Engine/Source/Runtime/VulkanRHI/Private/VulkanRHI.cpp
+++ b/Engine/Source/Runtime/VulkanRHI/Private/VulkanRHI.cpp
@@ -1128,6 +1128,8 @@ void FVulkanDynamicRHI::RHIEndFrame_RenderThread(FRHICommandListImmediate& RHICm
 
 		Context.Device->GetPipelineStateCache()->TickLRU();
 
+		Context.Device->GetBindlessDescriptorManager()->UpdateUBAllocator();
+
 		++Context.FrameCounter;
 	});
 
diff --git a/Engine/Source/Runtime/VulkanRHI/Private/VulkanRayTracing.cpp b/Engine/Source/Runtime/VulkanRHI/Private/VulkanRayTracing.cpp
--- a/Engine/Source/Runtime/VulkanRHI/Private/VulkanRayTracing.cpp
+++ b/Engine/Source/Runtime/VulkanRHI/Private/VulkanRayTracing.cpp
@@ -1992,7 +1992,7 @@ static FVulkanPipelineBarrier SetRayGenResources(FVulkanDevice* Device, FVulkanC
 				NumSkippedSlots++;
 			}
 		}
-		Device->GetBindlessDescriptorManager()->RegisterUniformBuffers(CmdBuffer->GetHandle(), VK_PIPELINE_BIND_POINT_RAY_TRACING_KHR, StageUBs);
+		Device->GetBindlessDescriptorManager()->RegisterUniformBuffers(CmdBuffer, VK_PIPELINE_BIND_POINT_RAY_TRACING_KHR, StageUBs);
 	}
 
 	// Add all the UBs references by the shader table
