VulkenRHI: Remove FTempFrameAllocationBuffer (minimum 12MB) and FVulkanUniformBufferUploader (fixed 16MB).  Replace both with same fenced blocks allocator (currently 4MB per block by default) that offer no risk of stomping and grow to required size.

#rnx

diff --git a/Engine/Source/Runtime/VulkanRHI/Private/VulkanCommands.cpp b/Engine/Source/Runtime/VulkanRHI/Private/VulkanCommands.cpp
--- a/Engine/Source/Runtime/VulkanRHI/Private/VulkanCommands.cpp
+++ b/Engine/Source/Runtime/VulkanRHI/Private/VulkanCommands.cpp
@@ -977,16 +977,16 @@ void FVulkanDynamicRHI::RHISubmitCommandLists(FRHISubmitCommandListsArgs&& Args)
 	for (IRHIPlatformCommandList* Ptr : Args.CommandLists)
 	{
 		FVulkanPlatformCommandList* PlatformCmdList = ResourceCast(Ptr);
+		FVulkanCommandListContext* CurrentContext = PlatformCmdList->CmdContext;
+		FVulkanCommandBufferManager* CmdBufMgr = CurrentContext->GetCommandBufferManager();
 
-		FVulkanCommandBufferManager* CmdBufMgr = PlatformCmdList->CmdContext->GetCommandBufferManager();
-
-		if (PlatformCmdList->CmdContext->IsImmediate())
+		if (CurrentContext->IsImmediate())
 		{
-			PlatformCmdList->CmdContext->RequestSubmitCurrentCommands();
+			CurrentContext->RequestSubmitCurrentCommands();
 			FVulkanCmdBuffer* CmdBuffer = CmdBufMgr->GetActiveCmdBuffer();
 			if (CmdBuffer && CmdBuffer->HasBegun() && CmdBuffer->IsOutsideRenderPass())
 			{
-				PlatformCmdList->CmdContext->SafePointSubmit();
+				CurrentContext->SafePointSubmit();
 			}
 			CmdBufMgr->RefreshFenceStatus();
 		}
@@ -997,7 +997,7 @@ void FVulkanDynamicRHI::RHISubmitCommandLists(FRHISubmitCommandListsArgs&& Args)
 			check(!CmdBuffer->IsInsideRenderPass());
 			CmdBufMgr->SubmitActiveCmdBuffer();
 
-			Device->ReleaseDeferredContext(PlatformCmdList->CmdContext);
+			Device->ReleaseDeferredContext(CurrentContext);
 		}
 
 		delete PlatformCmdList;
diff --git a/Engine/Source/Runtime/VulkanRHI/Private/VulkanContext.h b/Engine/Source/Runtime/VulkanRHI/Private/VulkanContext.h
--- a/Engine/Source/Runtime/VulkanRHI/Private/VulkanContext.h
+++ b/Engine/Source/Runtime/VulkanRHI/Private/VulkanContext.h
@@ -129,9 +129,9 @@ public:
 		return CommandBufferManager;
 	}
 
-	inline VulkanRHI::FTempFrameAllocationBuffer& GetTempFrameAllocationBuffer()
+	inline VulkanRHI::FTempBlockAllocator& GetTempBlockAllocator()
 	{
-		return TempFrameAllocationBuffer;
+		return *TempBlockAllocator;
 	}
 
 	inline FVulkanPendingGfxState* GetPendingGfxState()
@@ -173,11 +173,6 @@ public:
 		return FrameCounter;
 	}
 
-	inline FVulkanUniformBufferUploader* GetUniformBufferUploader()
-	{
-		return UniformBufferUploader;
-	}
-
 	inline FVulkanQueue* GetQueue()
 	{
 		return Queue;
@@ -216,12 +211,11 @@ protected:
 	FVulkanQueue* Queue;
 	bool bSubmitAtNextSafePoint;
 	bool bUniformBufferUploadRenderPassDirty = true;
-	FVulkanUniformBufferUploader* UniformBufferUploader;
 
 	void BeginOcclusionQueryBatch(FVulkanCmdBuffer* CmdBuffer, uint32 NumQueriesInBatch);
 	void EndOcclusionQueryBatch(FVulkanCmdBuffer* CmdBuffer);
 
-	VulkanRHI::FTempFrameAllocationBuffer TempFrameAllocationBuffer;
+	VulkanRHI::FTempBlockAllocator* TempBlockAllocator = nullptr;
 
 	TArray<FString> EventStack;
 
diff --git a/Engine/Source/Runtime/VulkanRHI/Private/VulkanMemory.cpp b/Engine/Source/Runtime/VulkanRHI/Private/VulkanMemory.cpp
--- a/Engine/Source/Runtime/VulkanRHI/Private/VulkanMemory.cpp
+++ b/Engine/Source/Runtime/VulkanRHI/Private/VulkanMemory.cpp
@@ -4973,20 +4973,11 @@ namespace VulkanRHI
 		return NewBlock;
 	}
 
-	uint8* FTempBlockAllocator::Alloc(uint32 InSize, FVulkanCmdBuffer* CmdBuffer, VkDescriptorBufferBindingInfoEXT& OutBindingInfo, VkDeviceSize& OutOffset)
+	FTempBlockAllocator::FInternalAlloc FTempBlockAllocator::InternalAlloc(uint32 InSize, FVulkanCmdBuffer* CmdBuffer)
 	{
 		const uint32 AlignedSize = Align(InSize, BlockAlignment);
 		checkfSlow(BlockAlignment < BlockSize, TEXT("Requested size of %d (%d aligned) is too large for block size of %d"), InSize, AlignedSize, BlockSize);
 
-		auto ReturnAllocationFromBlock = [AlignedSize, BufferUsageFlags=BufferUsage, &OutBindingInfo, &OutOffset](FTempMemoryBlock* Block, uint32 Offset)
-		{
-			ZeroVulkanStruct(OutBindingInfo, VK_STRUCTURE_TYPE_DESCRIPTOR_BUFFER_BINDING_INFO_EXT);
-			OutBindingInfo.address = Block->BufferAddress;
-			OutBindingInfo.usage = BufferUsageFlags;
-			OutOffset = Offset;
-			return Block->MappedPointer + Offset;
-		};
-
 		// Parallel scope: allocate in existing block from existing command buffer
 		{
 			FReadScopeLock ScopeLock(RWLock);
@@ -4998,7 +4989,7 @@ namespace VulkanRHI
 				const uint32 AllocOffset = CurrentBlock->CurrentOffset.fetch_add(AlignedSize);
 				if (AllocOffset + InSize < BlockSize)
 				{
-					return ReturnAllocationFromBlock(CurrentBlock, AllocOffset);
+					return FInternalAlloc{ CurrentBlock, AllocOffset };
 				}
 			}
 		}
@@ -5012,7 +5003,7 @@ namespace VulkanRHI
 			if (AllocOffset + InSize < BlockSize)
 			{
 				CurrentBlock->Fences.FindOrAdd(CmdBuffer) = CmdBuffer->GetFenceSignaledCounter();
-				return ReturnAllocationFromBlock(CurrentBlock, AllocOffset);
+				return FInternalAlloc{ CurrentBlock, AllocOffset };
 			}
 
 			// Get a new block
@@ -5031,8 +5022,41 @@ namespace VulkanRHI
 			AllocOffset = CurrentBlock->CurrentOffset.fetch_add(AlignedSize);
 			checkSlow(AllocOffset == 0);
 			checkSlow(AllocOffset + InSize < BlockSize);
-			return ReturnAllocationFromBlock(CurrentBlock, AllocOffset);
-		}	
+			return FInternalAlloc{ CurrentBlock, AllocOffset };
+		}
+	}
+
+	uint8* FTempBlockAllocator::Alloc(uint32 InSize, FVulkanCmdBuffer* CmdBuffer, VkDescriptorBufferBindingInfoEXT& OutBindingInfo, VkDeviceSize& OutOffset)
+	{
+		FInternalAlloc Alloc = InternalAlloc(InSize, CmdBuffer);
+
+		ZeroVulkanStruct(OutBindingInfo, VK_STRUCTURE_TYPE_DESCRIPTOR_BUFFER_BINDING_INFO_EXT);
+		OutBindingInfo.address = Alloc.Block->BufferAddress;
+		OutBindingInfo.usage = BufferUsage;
+		OutOffset = Alloc.Offset;
+		return Alloc.Block->MappedPointer + Alloc.Offset;
+	}
+
+	uint8* FTempBlockAllocator::Alloc(uint32 InSize, uint32 InAlignment, FVulkanCmdBuffer* CmdBuffer, FVulkanAllocation& OutAllocation, VkDescriptorAddressInfoEXT* OutDescriptorAddressInfo)
+	{
+		const uint32 AlignedSize = Align(InSize, BlockAlignment);
+		checkfSlow(BlockAlignment < BlockSize, TEXT("Requested size of %d (%d aligned) is too large for block size of %d"), InSize, AlignedSize, BlockSize);
+
+		FInternalAlloc Alloc = InternalAlloc(InSize, CmdBuffer);
+
+		OutAllocation.Reference(Alloc.Block->Allocation);
+		OutAllocation.VulkanHandle = (uint64)Alloc.Block->Buffer;
+		OutAllocation.Size = InSize;
+		OutAllocation.Offset += Alloc.Offset;
+
+		if (OutDescriptorAddressInfo)
+		{
+			ZeroVulkanStruct(*OutDescriptorAddressInfo, VK_STRUCTURE_TYPE_DESCRIPTOR_ADDRESS_INFO_EXT);
+			OutDescriptorAddressInfo->address = Alloc.Block->BufferAddress + Alloc.Offset;
+			OutDescriptorAddressInfo->range = InSize;
+		}
+
+		return Alloc.Block->MappedPointer + Alloc.Offset;
 	}
 
 	void FTempBlockAllocator::UpdateBlocks()
@@ -5067,118 +5091,6 @@ namespace VulkanRHI
 
 
 
-	FTempFrameAllocationBuffer::FTempFrameAllocationBuffer(FVulkanDevice* InDevice)
-		: FDeviceChild(InDevice)
-		, BufferIndex(0)
-	{
-		for (int32 Index = 0; Index < NUM_BUFFERS; ++Index)
-		{
-			INC_MEMORY_STAT_BY(STAT_VulkanTempFrameAllocationBuffer, ALLOCATION_SIZE);
-			Entries[Index].InitBuffer(Device, ALLOCATION_SIZE);
-		}
-	}
-
-	FTempFrameAllocationBuffer::~FTempFrameAllocationBuffer()
-	{
-		Destroy();
-	}
-
-	void FTempFrameAllocationBuffer::FFrameEntry::InitBuffer(FVulkanDevice* InDevice, uint32 InSize)
-	{
-		LLM_SCOPE_VULKAN(ELLMTagVulkan::VulkanFrameTemp);
-		Size = InSize;
-		PeakUsed = 0;
-		FMemoryManager& ResourceHeapManager = InDevice->GetMemoryManager();
-		check(Allocation.Type == EVulkanAllocationEmpty);
-		if(ResourceHeapManager.AllocateBufferPooled(Allocation, nullptr, InSize, 0,
-			VK_BUFFER_USAGE_TRANSFER_SRC_BIT | VK_BUFFER_USAGE_VERTEX_BUFFER_BIT |
-			VK_BUFFER_USAGE_INDEX_BUFFER_BIT | VK_BUFFER_USAGE_UNIFORM_BUFFER_BIT |
-			VK_BUFFER_USAGE_UNIFORM_TEXEL_BUFFER_BIT | VK_BUFFER_USAGE_STORAGE_BUFFER_BIT,
-			VK_MEMORY_PROPERTY_HOST_VISIBLE_BIT | VK_MEMORY_PROPERTY_HOST_COHERENT_BIT,
-			EVulkanAllocationMetaFrameTempBuffer,
-			__FILE__, __LINE__))
-		{
-			MappedData = (uint8*)Allocation.GetMappedPointer(InDevice);
-			CurrentData = MappedData;
-		}
-		else
-		{
-			ResourceHeapManager.HandleOOM(true);
-		}
-	}
-
-	void FTempFrameAllocationBuffer::Destroy()
-	{
-		FMemoryManager& MemoryManager = Device->GetMemoryManager();
-		for (int32 Index = 0; Index < NUM_BUFFERS; ++Index)
-		{
-			Entries[Index].Reset(Device);
-			MemoryManager.FreeVulkanAllocation(Entries[Index].Allocation); ;			
-		}
-	}
-
-	bool FTempFrameAllocationBuffer::FFrameEntry::TryAlloc(uint32 InSize, uint32 InAlignment, FTempAllocInfo& OutInfo)
-	{
-		uint8* AlignedData = (uint8*)Align((uintptr_t)CurrentData, (uintptr_t)InAlignment);
-		if (AlignedData + InSize <= MappedData + Size)
-		{
-			OutInfo.Data = AlignedData;
-			OutInfo.Allocation.Reference(Allocation);
-			OutInfo.CurrentOffset = (uint32)(AlignedData - MappedData);
-			OutInfo.Size = InSize;
-			CurrentData = AlignedData + InSize;
-			PeakUsed = FMath::Max(PeakUsed, (uint32)(CurrentData - MappedData));
-			return true;
-		}
-
-		return false;
-	}
-
-	void FTempFrameAllocationBuffer::Alloc(uint32 InSize, uint32 InAlignment, FTempAllocInfo& OutInfo)
-	{
-		FScopeLock ScopeLock(&CS);
-
-		if (Entries[BufferIndex].TryAlloc(InSize, InAlignment, OutInfo))
-		{
-			return;
-		}
-
-		// Couldn't fit in the current buffers; allocate a new bigger one and schedule the current one for deletion
-		uint32 NewSize = Align(ALLOCATION_SIZE + InSize + InAlignment, ALLOCATION_SIZE);
-		DEC_MEMORY_STAT_BY(STAT_VulkanTempFrameAllocationBuffer, Entries[BufferIndex].Allocation.Size);
-		INC_MEMORY_STAT_BY(STAT_VulkanTempFrameAllocationBuffer, NewSize);
-		FVulkanAllocation& PendingDelete = Entries[BufferIndex].PendingDeletionList.AddDefaulted_GetRef();
-		PendingDelete.Swap(Entries[BufferIndex].Allocation);
-		Entries[BufferIndex].InitBuffer(Device, NewSize);
-		if (!Entries[BufferIndex].TryAlloc(InSize, InAlignment, OutInfo))
-		{
-			checkf(0, TEXT("Internal Error trying to allocate %d Align %d on TempFrameBuffer, size %d"), InSize, InAlignment, NewSize);
-		}
-	}
-
-	void FTempFrameAllocationBuffer::Reset()
-	{
-		FScopeLock ScopeLock(&CS);
-		BufferIndex = (BufferIndex + 1) % NUM_BUFFERS;
-		Entries[BufferIndex].Reset(Device);
-	}
-
-	void FTempFrameAllocationBuffer::FFrameEntry::Reset(FVulkanDevice* InDevice)
-	{
-		CurrentData = MappedData;
-		FMemoryManager& MemoryManager = InDevice->GetMemoryManager();
-		for(FVulkanAllocation& Alloc : PendingDeletionList)
-		{
-			if(Alloc.HasAllocation())
-			{
-				MemoryManager.FreeVulkanAllocation(Alloc);
-			}
-			check(!Alloc.HasAllocation());
-
-		}
-		PendingDeletionList.SetNum(0);
-	}
-
 	FSemaphore::FSemaphore(FVulkanDevice& InDevice) :
 		Device(InDevice),
 		SemaphoreHandle(VK_NULL_HANDLE),
diff --git a/Engine/Source/Runtime/VulkanRHI/Private/VulkanMemory.h b/Engine/Source/Runtime/VulkanRHI/Private/VulkanMemory.h
--- a/Engine/Source/Runtime/VulkanRHI/Private/VulkanMemory.h
+++ b/Engine/Source/Runtime/VulkanRHI/Private/VulkanMemory.h
@@ -1134,15 +1134,23 @@ namespace VulkanRHI
 			std::atomic<uint32> CurrentOffset = 0;
 		};
 
+		struct FInternalAlloc
+		{
+			FTempMemoryBlock* Block = nullptr;
+			uint32 Offset = 0;
+		};
+
 		FTempBlockAllocator(FVulkanDevice* InDevice, uint32 InBlockSize, uint32 InBlockAlignment, VkBufferUsageFlags InBufferUsage);
 		virtual ~FTempBlockAllocator();
 
 		uint8* Alloc(uint32 InSize, FVulkanCmdBuffer* CmdBuffer, VkDescriptorBufferBindingInfoEXT& OutBindingInfo, VkDeviceSize& OutOffset);
+		uint8* Alloc(uint32 InSize, uint32 InAlignment, FVulkanCmdBuffer* CmdBuffer, FVulkanAllocation& OutAllocation, VkDescriptorAddressInfoEXT* OutDescriptorAddressInfo = nullptr);
 
 		void UpdateBlocks();
 
 	protected:
 		FTempMemoryBlock* AllocBlock();
+		FInternalAlloc InternalAlloc(uint32 InSize, FVulkanCmdBuffer* CmdBuffer);
 
 		const uint32 BlockSize;
 		const uint32 BlockAlignment;
@@ -1155,65 +1163,6 @@ namespace VulkanRHI
 	};
 
 
-	// Simple tape allocation per frame for a VkBuffer, used for Volatile allocations
-	class FTempFrameAllocationBuffer : public FDeviceChild
-	{
-		enum
-		{
-			ALLOCATION_SIZE = (4 * 1024 * 1024),
-		};
-
-	public:
-		FTempFrameAllocationBuffer(FVulkanDevice* InDevice);
-		virtual ~FTempFrameAllocationBuffer();
-		void Destroy();
-
-		struct FTempAllocInfo
-		{
-			FVulkanAllocation Allocation;
-			void* Data = 0;
-			uint32 CurrentOffset = 0;
-			uint32 Size = 0;
-			uint32 LockCounter = 0;
-
-			uint32 GetBindOffset()
-			{
-				checkNoEntry();
-				return 0;
-			}
-
-		};
-
-		void Alloc(uint32 InSize, uint32 InAlignment, FTempAllocInfo& OutInfo);
-		void Reset();
-
-	protected:
-		uint32 BufferIndex;
-
-		enum
-		{
-			NUM_BUFFERS = 3,
-		};
-
-		struct FFrameEntry
-		{
-			FVulkanAllocation Allocation;
-			TArray<FVulkanAllocation> PendingDeletionList;
-			uint8* MappedData = nullptr;
-			uint8* CurrentData = nullptr;
-			uint32 Size = 0;
-			uint32 PeakUsed = 0;
-
-			void InitBuffer(FVulkanDevice* InDevice, uint32 InSize);
-			void Reset(FVulkanDevice* Device);
-			bool TryAlloc(uint32 InSize, uint32 InAlignment, FTempAllocInfo& OutInfo);
-		};
-		FFrameEntry Entries[NUM_BUFFERS];
-		FCriticalSection CS;
-
-		friend class FVulkanCommandListContext;
-	};
-
 	class VULKANRHI_API FSemaphore : public FRefCount
 	{
 	public:
diff --git a/Engine/Source/Runtime/VulkanRHI/Private/VulkanPipelineState.cpp b/Engine/Source/Runtime/VulkanRHI/Private/VulkanPipelineState.cpp
--- a/Engine/Source/Runtime/VulkanRHI/Private/VulkanPipelineState.cpp
+++ b/Engine/Source/Runtime/VulkanRHI/Private/VulkanPipelineState.cpp
@@ -155,6 +155,23 @@ void FVulkanCommonPipelineDescriptorState::CreateDescriptorWriteInfos()
 	DescriptorSetHandles.AddZeroed(MaxNumSets);
 }
 
+
+static inline VulkanRHI::FVulkanAllocation UpdatePackedUniformBuffers(const FPackedUniformBuffers& PackedUniformBuffers, FVulkanCommandListContext* InContext, FVulkanCmdBuffer* InCmdBuffer)
+{
+	const FPackedUniformBuffers::FPackedBuffer& StagedUniformBuffer = PackedUniformBuffers.GetBuffer();
+
+	const uint32 UBSize = (uint32)StagedUniformBuffer.Num();
+	const uint32 UBAlign = (uint32)InCmdBuffer->GetDevice()->GetLimits().minUniformBufferOffsetAlignment;
+
+	VulkanRHI::FVulkanAllocation TempAllocation;
+	uint8* MappedPointer = InContext->GetTempBlockAllocator().Alloc(UBSize, UBAlign, InCmdBuffer, TempAllocation);
+
+	FMemory::Memcpy(MappedPointer, StagedUniformBuffer.GetData(), UBSize);
+
+	return TempAllocation;
+}
+
+
 template<bool bUseDynamicGlobalUBs>
 bool FVulkanComputePipelineDescriptorState::InternalUpdateDescriptorSets(FVulkanCommandListContext* CmdListContext, FVulkanCmdBuffer* CmdBuffer)
 {
@@ -168,16 +185,12 @@ bool FVulkanComputePipelineDescriptorState::InternalUpdateDescriptorSets(FVulkan
 		return false;
 	}
 
-	FVulkanUniformBufferUploader* UniformBufferUploader = CmdListContext->GetUniformBufferUploader();
-	uint8* CPURingBufferBase = (uint8*)UniformBufferUploader->GetCPUMappedPointer();
-	const VkDeviceSize UBOffsetAlignment = Device->GetLimits().minUniformBufferOffsetAlignment;
-
 	if (PackedUniformBuffersDirty != 0)
 	{
 #if VULKAN_ENABLE_AGGRESSIVE_STATS
 		SCOPE_CYCLE_COUNTER(STAT_VulkanApplyPackedUniformBuffers);
 #endif
-		UpdatePackedUniformBuffers<bUseDynamicGlobalUBs>(UBOffsetAlignment, PackedUniformBuffers, DSWriter[ShaderStage::Compute], UniformBufferUploader, CPURingBufferBase, CmdBuffer);
+		SubmitPackedUniformBuffers<bUseDynamicGlobalUBs>(DSWriter[ShaderStage::Compute], UpdatePackedUniformBuffers(PackedUniformBuffers, CmdListContext, CmdBuffer));
 		PackedUniformBuffersDirty = 0;
 	}
 
@@ -218,10 +231,6 @@ void FVulkanComputePipelineDescriptorState::UpdateBindlessDescriptors(FVulkanCom
 	check(DSWriteContainer.DescriptorBufferInfo.Num() == DSWriteContainer.DescriptorWrites.Num());
 	check(DSWriteContainer.DescriptorImageInfo.Num() == 0);
 
-	FVulkanUniformBufferUploader* UniformBufferUploader = CmdListContext->GetUniformBufferUploader();
-	uint8* CPURingBufferBase = (uint8*)UniformBufferUploader->GetCPUMappedPointer();
-	const VkDeviceSize UBOffsetAlignment = Device->GetLimits().minUniformBufferOffsetAlignment;
-
 	FVulkanBindlessDescriptorManager::FUniformBufferDescriptorArrays StageUBs;
 
 	const FVulkanShaderHeader& Header = ComputePipeline->GetShaderCodeHeader();
@@ -237,19 +246,11 @@ void FVulkanComputePipelineDescriptorState::UpdateBindlessDescriptors(FVulkanCom
 		const FPackedUniformBuffers::FPackedBuffer& StagedUniformBuffer = PackedUniformBuffers.GetBuffer();
 		const int32 UBSize = StagedUniformBuffer.Num();
 		const int32 BindingIndex = 0;
+		const VkDeviceSize UBOffsetAlignment = Device->GetLimits().minUniformBufferOffsetAlignment;
 
-		// :todo-jn: Use PackedUniformBuffersDirty and cache last RingBufferOffset to avoid copy when there are no changes (while updating lifetime in UniformBufferUploader)
-		const uint64 RingBufferOffset = UniformBufferUploader->AllocateMemory(UBSize, UBOffsetAlignment, CmdBuffer);
-
-		// Make sure it wasn't written to already
-		VkDescriptorAddressInfoEXT& DescriptorAddressInfo = DescriptorAddressInfos[BindingIndex];
-		check(DescriptorAddressInfo.sType == 0);
-		DescriptorAddressInfo.sType = VK_STRUCTURE_TYPE_DESCRIPTOR_ADDRESS_INFO_EXT;
-		DescriptorAddressInfo.address = UniformBufferUploader->GetCPUBufferAddress() + RingBufferOffset;
-		DescriptorAddressInfo.range = UBSize;
-
-		// get location in the ring buffer to use
-		FMemory::Memcpy(CPURingBufferBase + RingBufferOffset, StagedUniformBuffer.GetData(), UBSize);
+		VulkanRHI::FVulkanAllocation TempAllocation;
+		uint8* MappedPointer = CmdListContext->GetTempBlockAllocator().Alloc(UBSize, UBOffsetAlignment, CmdBuffer, TempAllocation, &DescriptorAddressInfos[BindingIndex]);
+		FMemory::Memcpy(MappedPointer, StagedUniformBuffer.GetData(), UBSize);
 
 		PackedUniformBuffersDirty = 0;
 		++UBIndex;
@@ -343,10 +344,6 @@ bool FVulkanGraphicsPipelineDescriptorState::InternalUpdateDescriptorSets(FVulka
 		return false;
 	}
 
-	FVulkanUniformBufferUploader* UniformBufferUploader = CmdListContext->GetUniformBufferUploader();
-	uint8* CPURingBufferBase = (uint8*)UniformBufferUploader->GetCPUMappedPointer();
-	const VkDeviceSize UBOffsetAlignment = Device->GetLimits().minUniformBufferOffsetAlignment;
-
 	// Process updates
 	{
 #if VULKAN_ENABLE_AGGRESSIVE_STATS
@@ -356,7 +353,7 @@ bool FVulkanGraphicsPipelineDescriptorState::InternalUpdateDescriptorSets(FVulka
 		{
 			if (PackedUniformBuffersDirty[Stage] != 0)
 			{
-				MarkDirty(UpdatePackedUniformBuffers<bUseDynamicGlobalUBs>(UBOffsetAlignment, PackedUniformBuffers[Stage], DSWriter[Stage], UniformBufferUploader, CPURingBufferBase, CmdBuffer));
+				MarkDirty(SubmitPackedUniformBuffers<bUseDynamicGlobalUBs>(DSWriter[Stage], UpdatePackedUniformBuffers(PackedUniformBuffers[Stage], CmdListContext, CmdBuffer)));
 				PackedUniformBuffersDirty[Stage] = 0;
 			}
 		}
@@ -418,8 +415,6 @@ void FVulkanGraphicsPipelineDescriptorState::UpdateBindlessDescriptors(FVulkanCo
 	check(DSWriteContainer.DescriptorBufferInfo.Num() == DSWriteContainer.DescriptorWrites.Num());
 	check(DSWriteContainer.DescriptorImageInfo.Num() == 0);
 
-	FVulkanUniformBufferUploader* UniformBufferUploader = CmdListContext->GetUniformBufferUploader();
-	uint8* CPURingBufferBase = (uint8*)UniformBufferUploader->GetCPUMappedPointer();
 	const VkDeviceSize UBOffsetAlignment = Device->GetLimits().minUniformBufferOffsetAlignment;
 
 	FVulkanBindlessDescriptorManager::FUniformBufferDescriptorArrays StageUBs;
@@ -451,18 +446,9 @@ void FVulkanGraphicsPipelineDescriptorState::UpdateBindlessDescriptors(FVulkanCo
 				const int32 UBSize = StagedUniformBuffer.Num();
 				const int32 BindingIndex = 0;
 
-				// :todo-jn: Use PackedUniformBuffersDirty and cache last RingBufferOffset to avoid copy when there are no changes (while updating lifetime in UniformBufferUploader)
-				const uint64 RingBufferOffset = UniformBufferUploader->AllocateMemory(UBSize, UBOffsetAlignment, CmdBuffer);
-
-				// Make sure it wasn't written to already
-				VkDescriptorAddressInfoEXT& DescriptorAddressInfo = DescriptorAddressInfos[BindingIndex];
-				check(DescriptorAddressInfo.sType == 0);
-				DescriptorAddressInfo.sType = VK_STRUCTURE_TYPE_DESCRIPTOR_ADDRESS_INFO_EXT;
-				DescriptorAddressInfo.address = UniformBufferUploader->GetCPUBufferAddress() + RingBufferOffset;
-				DescriptorAddressInfo.range = UBSize;
-
-				// get location in the ring buffer to use
-				FMemory::Memcpy(CPURingBufferBase + RingBufferOffset, StagedUniformBuffer.GetData(), UBSize);
+				VulkanRHI::FVulkanAllocation TempAllocation;
+				uint8* MappedPointer = CmdListContext->GetTempBlockAllocator().Alloc(UBSize, UBOffsetAlignment, CmdBuffer, TempAllocation, &DescriptorAddressInfos[BindingIndex]);
+				FMemory::Memcpy(MappedPointer, StagedUniformBuffer.GetData(), UBSize);
 
 				PackedUniformBuffersDirty[Stage] = 0;
 				++UBIndex;
diff --git a/Engine/Source/Runtime/VulkanRHI/Private/VulkanPipelineState.h b/Engine/Source/Runtime/VulkanRHI/Private/VulkanPipelineState.h
--- a/Engine/Source/Runtime/VulkanRHI/Private/VulkanPipelineState.h
+++ b/Engine/Source/Runtime/VulkanRHI/Private/VulkanPipelineState.h
@@ -291,28 +291,15 @@ protected:
 };
 
 template <bool bIsDynamic>
-static inline bool UpdatePackedUniformBuffers(VkDeviceSize UBOffsetAlignment, const FPackedUniformBuffers& PackedUniformBuffers,
-	FVulkanDescriptorSetWriter& DescriptorWriteSet, FVulkanUniformBufferUploader* UniformBufferUploader, uint8* RESTRICT CPURingBufferBase,
-	FVulkanCmdBuffer* InCmdBuffer)
+static inline bool SubmitPackedUniformBuffers(FVulkanDescriptorSetWriter& DescriptorWriteSet, const VulkanRHI::FVulkanAllocation& TempAllocation)
 {
-	const FPackedUniformBuffers::FPackedBuffer& StagedUniformBuffer = PackedUniformBuffers.GetBuffer();
-	const int32 BindingIndex = 0;  // Packed unifrom buffers are only used for globals at binding 0
-
-	const int32 UBSize = StagedUniformBuffer.Num();
-
-	// get offset into the RingBufferBase pointer
-	uint64 RingBufferOffset = UniformBufferUploader->AllocateMemory(UBSize, UBOffsetAlignment, InCmdBuffer);
-
-	// get location in the ring buffer to use
-	FMemory::Memcpy(CPURingBufferBase + RingBufferOffset, StagedUniformBuffer.GetData(), UBSize);
-
-	const VulkanRHI::FVulkanAllocation& Allocation = UniformBufferUploader->GetCPUBufferAllocation();
+	const int32 BindingIndex = 0;  // Packed uniform buffers are only used for globals at binding 0
 	if (bIsDynamic)
 	{
-		return DescriptorWriteSet.WriteDynamicUniformBuffer(BindingIndex, Allocation.GetBufferHandle(), Allocation.HandleId, UniformBufferUploader->GetCPUBufferOffset(), UBSize, RingBufferOffset);
+		return DescriptorWriteSet.WriteDynamicUniformBuffer(BindingIndex, TempAllocation.GetBufferHandle(), TempAllocation.HandleId, 0, TempAllocation.Size, TempAllocation.Offset);
 	}
 	else
 	{
-		return DescriptorWriteSet.WriteUniformBuffer(BindingIndex, Allocation.GetBufferHandle(), Allocation.HandleId, RingBufferOffset + UniformBufferUploader->GetCPUBufferOffset(), UBSize);
+		return DescriptorWriteSet.WriteUniformBuffer(BindingIndex, TempAllocation.GetBufferHandle(), TempAllocation.HandleId, TempAllocation.Offset, TempAllocation.Size);
 	}
 }
diff --git a/Engine/Source/Runtime/VulkanRHI/Private/VulkanRHI.cpp b/Engine/Source/Runtime/VulkanRHI/Private/VulkanRHI.cpp
--- a/Engine/Source/Runtime/VulkanRHI/Private/VulkanRHI.cpp
+++ b/Engine/Source/Runtime/VulkanRHI/Private/VulkanRHI.cpp
@@ -105,6 +105,16 @@ static TAutoConsoleVariable<bool> CVarEnableVulkanPSOFileCacheWhenPrecachingActi
 	TEXT("true: Allow both PSO file cache and precaching."),
 	ECVF_RenderThreadSafe | ECVF_ReadOnly);
 
+int32 GVulkanTempBlockSize = 4 * 1024 * 1024;
+static FAutoConsoleVariableRef GCVarVulkanTempBlockSize(
+	TEXT("r.Vulkan.TempBlockSize"),
+	GVulkanTempBlockSize,
+	TEXT("Size of the temporary blocks allocate by contexts, used for single use ub allocs and copies (default: 4MB)."),
+	ECVF_ReadOnly
+);
+
+
+
 extern TAutoConsoleVariable<int32> GVulkanRayTracingCVar;
 
 // All shader stages supported by VK device - VK_PIPELINE_STAGE_VERTEX_SHADER_BIT, FRAGMENT etc
@@ -416,8 +426,6 @@ FVulkanCommandListContext::FVulkanCommandListContext(FVulkanDynamicRHI* InRHI, F
 	, Device(InDevice)
 	, Queue(InQueue)
 	, bSubmitAtNextSafePoint(false)
-	, UniformBufferUploader(nullptr)
-	, TempFrameAllocationBuffer(InDevice)
 	, CommandBufferManager(nullptr)
 	, PendingGfxState(nullptr)
 	, PendingComputeState(nullptr)
@@ -453,7 +461,15 @@ FVulkanCommandListContext::FVulkanCommandListContext(FVulkanDynamicRHI* InRHI, F
 	PendingGfxState = new FVulkanPendingGfxState(Device, *this);
 	PendingComputeState = new FVulkanPendingComputeState(Device, *this);
 
-	UniformBufferUploader = new FVulkanUniformBufferUploader(Device);
+	// Currently used for UB and copies
+	{
+		const uint32 BlockAlignment = FMath::Max<uint32>(Device->GetLimits().minUniformBufferOffsetAlignment, 16u);
+		const VkBufferUsageFlags BufferUsageFlags = 
+			(Device->GetOptionalExtensions().HasBufferDeviceAddress ? VK_BUFFER_USAGE_SHADER_DEVICE_ADDRESS_BIT : 0) |
+			VK_BUFFER_USAGE_TRANSFER_SRC_BIT |
+			VK_BUFFER_USAGE_UNIFORM_BUFFER_BIT;
+		TempBlockAllocator = new VulkanRHI::FTempBlockAllocator(InDevice, GVulkanTempBlockSize, BlockAlignment, BufferUsageFlags);
+	}
 
 	GlobalUniformBuffers.AddZeroed(FUniformBufferStaticSlotRegistry::Get().GetSlotCount());
 }
@@ -482,11 +498,9 @@ FVulkanCommandListContext::~FVulkanCommandListContext()
 	delete CommandBufferManager;
 	CommandBufferManager = nullptr;
 
-	delete UniformBufferUploader;
+	delete TempBlockAllocator;
 	delete PendingGfxState;
 	delete PendingComputeState;
-
-	TempFrameAllocationBuffer.Destroy();
 }
 
 
@@ -1129,6 +1143,7 @@ void FVulkanDynamicRHI::RHIEndFrame_RenderThread(FRHICommandListImmediate& RHICm
 		Context.Device->GetPipelineStateCache()->TickLRU();
 
 		Context.Device->GetBindlessDescriptorManager()->UpdateUBAllocator();
+		Context.GetTempBlockAllocator().UpdateBlocks();
 
 		++Context.FrameCounter;
 	});
diff --git a/Engine/Source/Runtime/VulkanRHI/Private/VulkanResources.h b/Engine/Source/Runtime/VulkanRHI/Private/VulkanResources.h
--- a/Engine/Source/Runtime/VulkanRHI/Private/VulkanResources.h
+++ b/Engine/Source/Runtime/VulkanRHI/Private/VulkanResources.h
@@ -949,112 +949,6 @@ public:
 	FVulkanTimingQueryPool* Pool = nullptr;
 };
 
-struct FVulkanRingBuffer : public VulkanRHI::FDeviceChild
-{
-public:
-	FVulkanRingBuffer(FVulkanDevice* InDevice, uint64 TotalSize, VkFlags Usage, VkMemoryPropertyFlags MemPropertyFlags);
-	virtual ~FVulkanRingBuffer();
-
-	// Allocate some space in the ring buffer
-	inline uint64 AllocateMemory(uint64 Size, uint32 Alignment, FVulkanCmdBuffer* InCmdBuffer)
-	{
-		Alignment = FMath::Max(Alignment, MinAlignment);
-		uint64 AllocationOffset = Align<uint64>(BufferOffset, Alignment);
-		if (AllocationOffset + Size <= BufferSize)
-		{
-			BufferOffset = AllocationOffset + Size;
-			return AllocationOffset;
-		}
-
-		return WrapAroundAllocateMemory(Size, Alignment, InCmdBuffer);
-	}
-
-	inline uint32 GetBufferOffset() const
-	{
-		return Allocation.Offset;
-	}
-
-	inline VkDeviceAddress GetBufferAddress() const
-	{
-		return BufferAddress;
-	}
-
-	inline VkBuffer GetHandle() const
-	{
-		return Allocation.GetBufferHandle();
-	}
-
-	inline void* GetMappedPointer()
-	{
-		return Allocation.GetMappedPointer(Device);
-	}
-
-	inline VulkanRHI::FVulkanAllocation& GetAllocation()
-	{
-		return Allocation;
-	}
-
-	inline const VulkanRHI::FVulkanAllocation& GetAllocation() const
-	{
-		return Allocation;
-	}
-
-
-protected:
-	uint64 BufferSize;
-	uint64 BufferOffset;
-	VkDeviceAddress BufferAddress;
-	uint32 MinAlignment;
-	VulkanRHI::FVulkanAllocation Allocation;
-
-	// Fence for wrapping around
-	FVulkanCmdBuffer* FenceCmdBuffer = nullptr;
-	uint64 FenceCounter = 0;
-
-	uint64 WrapAroundAllocateMemory(uint64 Size, uint32 Alignment, FVulkanCmdBuffer* InCmdBuffer);
-};
-
-struct FVulkanUniformBufferUploader : public VulkanRHI::FDeviceChild
-{
-public:
-	FVulkanUniformBufferUploader(FVulkanDevice* InDevice);
-	~FVulkanUniformBufferUploader();
-
-	inline uint8* GetCPUMappedPointer()
-	{
-		return (uint8*)CPUBuffer->GetMappedPointer();
-	}
-
-	inline uint64 AllocateMemory(uint64 Size, uint32 Alignment, FVulkanCmdBuffer* InCmdBuffer)
-	{
-		return CPUBuffer->AllocateMemory(Size, Alignment, InCmdBuffer);
-	}
-
-	inline const VulkanRHI::FVulkanAllocation& GetCPUBufferAllocation() const
-	{
-		return CPUBuffer->GetAllocation();
-	}
-
-	inline VkBuffer GetCPUBufferHandle() const
-	{
-		return CPUBuffer->GetHandle();
-	}
-
-	inline uint32 GetCPUBufferOffset() const
-	{
-		return CPUBuffer->GetBufferOffset();
-	}
-
-	inline VkDeviceAddress GetCPUBufferAddress() const
-	{
-		return CPUBuffer->GetBufferAddress();
-	}
-
-protected:
-	FVulkanRingBuffer* CPUBuffer;
-	friend class FVulkanCommandListContext;
-};
-
 class FVulkanResourceMultiBuffer : public FRHIBuffer, public VulkanRHI::FDeviceChild, public FVulkanViewableResource
 {
 public:
diff --git a/Engine/Source/Runtime/VulkanRHI/Private/VulkanUniformBuffer.cpp b/Engine/Source/Runtime/VulkanRHI/Private/VulkanUniformBuffer.cpp
--- a/Engine/Source/Runtime/VulkanRHI/Private/VulkanUniformBuffer.cpp
+++ b/Engine/Source/Runtime/VulkanRHI/Private/VulkanUniformBuffer.cpp
@@ -35,7 +35,7 @@ static void UpdateUniformBufferConstants(FVulkanDevice* Device, void* Destinatio
 	UE::RHICore::UpdateUniformBufferConstants(DestinationData, SourceData, *Layout, Device->SupportsBindless());
 }
 
-static bool UseRingBuffer(EUniformBufferUsage Usage)
+static bool UseTemporaryBuffer(EUniformBufferUsage Usage)
 {
 	// Add a cvar to control this behavior?
 	return (Usage == UniformBuffer_SingleDraw || Usage == UniformBuffer_SingleFrame);
@@ -47,57 +47,46 @@ static void UpdateUniformBufferHelper(FVulkanCommandListContext& Context, FVulka
 
 	FVulkanDevice* Device = Context.GetDevice();
 	const int32 DataSize = VulkanUniformBuffer->GetLayout().ConstantBufferSize;
+	const int32 DataAlignment = FMath::Max<uint32>(Device->GetLimits().minUniformBufferOffsetAlignment, 16u);
 
-	auto CopyUniformBufferData = [&](void* DestinationData, const void* SourceData) {
+	VulkanRHI::FVulkanAllocation TempAllocation;
+	void* DestinationData = Context.GetTempBlockAllocator().Alloc(DataSize, DataAlignment, CmdBuffer, TempAllocation);
 
-		if (bUpdateConstants)
-		{
-			// Update constants as the data is copied
-			UpdateUniformBufferConstants(Device, DestinationData, SourceData, VulkanUniformBuffer->GetLayoutPtr());
-		}
-		else
-		{
-			// Don't touch constant, copy the data as-is
-			FMemory::Memcpy(DestinationData, SourceData, DataSize);
-		}
-	};
-
-	if (UseRingBuffer(VulkanUniformBuffer->Usage))
+	if (bUpdateConstants)
 	{
-		FVulkanUniformBufferUploader* UniformBufferUploader = Context.GetUniformBufferUploader();
-		const VkDeviceSize UBOffsetAlignment = Device->GetLimits().minUniformBufferOffsetAlignment;
-		const VulkanRHI::FVulkanAllocation& RingBufferAllocation = UniformBufferUploader->GetCPUBufferAllocation();
-		uint64 RingBufferOffset = UniformBufferUploader->AllocateMemory(DataSize, UBOffsetAlignment, CmdBuffer);
+		// Update constants as the data is copied
+		UpdateUniformBufferConstants(Device, DestinationData, Data, VulkanUniformBuffer->GetLayoutPtr());
+	}
+	else
+	{
+		// Don't touch constant, copy the data as-is
+		FMemory::Memcpy(DestinationData, Data, DataSize);
+	}
 
+
+	if (UseTemporaryBuffer(VulkanUniformBuffer->Usage))
+	{
 		VulkanUniformBuffer->Allocation.Init(
 			VulkanRHI::EVulkanAllocationEmpty,
 			VulkanRHI::EVulkanAllocationMetaUnknown,
-			RingBufferAllocation.VulkanHandle, 
+			TempAllocation.VulkanHandle,
 			DataSize,
-			RingBufferOffset,
-			RingBufferAllocation.AllocatorIndex,
-			RingBufferAllocation.AllocationIndex,
-			RingBufferAllocation.HandleId);
-		
-		uint8* UploadLocation = UniformBufferUploader->GetCPUMappedPointer() + RingBufferOffset;
-		CopyUniformBufferData(UploadLocation, Data);
+			TempAllocation.Offset,
+			TempAllocation.AllocatorIndex,
+			TempAllocation.AllocationIndex,
+			TempAllocation.HandleId);
 	}
 	else
 	{
 		check(CmdBuffer->IsOutsideRenderPass());
 
-		VulkanRHI::FTempFrameAllocationBuffer::FTempAllocInfo LockInfo;
-		Context.GetTempFrameAllocationBuffer().Alloc(DataSize, 16, LockInfo);
-		CopyUniformBufferData(LockInfo.Data, Data);
-
 		VkBufferCopy Region;
 		Region.size = DataSize;
-		Region.srcOffset = LockInfo.CurrentOffset + LockInfo.Allocation.Offset;
+		Region.srcOffset = TempAllocation.Offset;
 		Region.dstOffset = VulkanUniformBuffer->GetOffset();
 		VkBuffer UBBuffer = VulkanUniformBuffer->Allocation.GetBufferHandle();
-		VkBuffer LockHandle = LockInfo.Allocation.GetBufferHandle();
 
-		VulkanRHI::vkCmdCopyBuffer(CmdBuffer->GetHandle(), LockHandle, UBBuffer, 1, &Region);
+		VulkanRHI::vkCmdCopyBuffer(CmdBuffer->GetHandle(), TempAllocation.GetBufferHandle(), UBBuffer, 1, &Region);
 	}
 };
 
@@ -165,7 +154,7 @@ FVulkanUniformBuffer::FVulkanUniformBuffer(FVulkanDevice& InDevice, const FRHIUn
 		const bool bInRHIThread = IsInRHIThread();
 
 PRAGMA_DISABLE_DEPRECATION_WARNINGS
-		if (UseRingBuffer(InUsage) && (bInRenderingThread || bInRHIThread)
+		if (UseTemporaryBuffer(InUsage) && (bInRenderingThread || bInRHIThread)
 			// :todo-jn:  Temporary check until we have a command list arg passed in to avoid a race where the RenderThread
 			// would pick up other tasks (because of task retraction) and execute them as if on the RenderThread.
 			&& !UE::Tasks::Private::IsThreadRetractingTask())
@@ -294,9 +283,9 @@ inline void FVulkanDynamicRHI::UpdateUniformBuffer(FRHICommandListBase& RHICmdLi
 
 	VulkanRHI::FVulkanAllocation NewUBAlloc;
 	bool bUseUpload = GVulkanAllowUniformUpload && !RHICmdList.IsInsideRenderPass(); //inside renderpasses, a rename is enforced.
-	const bool bUseRingBuffer = UseRingBuffer(UniformBuffer->Usage);
+	const bool bUseTempBuffer = UseTemporaryBuffer(UniformBuffer->Usage);
 
-	if (!bUseUpload && !bUseRingBuffer)
+	if (!bUseUpload && !bUseTempBuffer)
 	{
 		if (ConstantBufferSize > 0)
 		{
@@ -315,7 +304,7 @@ inline void FVulkanDynamicRHI::UpdateUniformBuffer(FRHICommandListBase& RHICmdLi
 	{
 		if (ConstantBufferSize > 0)
 		{
-			if (bUseUpload || bUseRingBuffer)
+			if (bUseUpload || bUseTempBuffer)
 			{			
 				FVulkanCommandListContext& Context = Device->GetImmediateContext();
 				UpdateUniformBufferHelper(Context, UniformBuffer, Contents);
@@ -342,7 +331,7 @@ inline void FVulkanDynamicRHI::UpdateUniformBuffer(FRHICommandListBase& RHICmdLi
 			}
 		}
 
-		if (bUseUpload || bUseRingBuffer)
+		if (bUseUpload || bUseTempBuffer)
 		{
 			void* CmdListConstantBufferData = RHICmdList.Alloc(ConstantBufferSize, 16);
 			FMemory::Memcpy(CmdListConstantBufferData, Contents, ConstantBufferSize);
@@ -381,98 +370,3 @@ void FVulkanDynamicRHI::RHIUpdateUniformBuffer(FRHICommandListBase& RHICmdList,
 	UpdateUniformBuffer(RHICmdList, UniformBuffer, Contents);
 }
 
-FVulkanUniformBufferUploader::FVulkanUniformBufferUploader(FVulkanDevice* InDevice)
-	: VulkanRHI::FDeviceChild(InDevice)
-	, CPUBuffer(nullptr)
-{
-	if (Device->HasUnifiedMemory())
-	{
-		CPUBuffer = new FVulkanRingBuffer(InDevice, PackedUniformsRingBufferSize, VK_BUFFER_USAGE_UNIFORM_BUFFER_BIT, VK_MEMORY_PROPERTY_HOST_VISIBLE_BIT | VK_MEMORY_PROPERTY_HOST_COHERENT_BIT | VK_MEMORY_PROPERTY_DEVICE_LOCAL_BIT);
-	}
-	else
-	{
-		if (FVulkanPlatform::SupportsDeviceLocalHostVisibleWithNoPenalty(InDevice->GetVendorId()) &&
-			InDevice->GetDeviceMemoryManager().SupportsMemoryType(VK_MEMORY_PROPERTY_HOST_VISIBLE_BIT | VK_MEMORY_PROPERTY_HOST_COHERENT_BIT | VK_MEMORY_PROPERTY_DEVICE_LOCAL_BIT))
-		{
-			CPUBuffer = new FVulkanRingBuffer(InDevice, PackedUniformsRingBufferSize, VK_BUFFER_USAGE_UNIFORM_BUFFER_BIT, VK_MEMORY_PROPERTY_HOST_VISIBLE_BIT | VK_MEMORY_PROPERTY_HOST_COHERENT_BIT | VK_MEMORY_PROPERTY_DEVICE_LOCAL_BIT);
-		}
-		else
-		{
-			CPUBuffer = new FVulkanRingBuffer(InDevice, PackedUniformsRingBufferSize, VK_BUFFER_USAGE_UNIFORM_BUFFER_BIT, VK_MEMORY_PROPERTY_HOST_VISIBLE_BIT | VK_MEMORY_PROPERTY_HOST_COHERENT_BIT);
-		}
-	}
-
-	INC_MEMORY_STAT_BY(STAT_UniformBufferMemory, PackedUniformsRingBufferSize);
-}
-
-FVulkanUniformBufferUploader::~FVulkanUniformBufferUploader()
-{
-	delete CPUBuffer;
-
-	DEC_MEMORY_STAT_BY(STAT_UniformBufferMemory, PackedUniformsRingBufferSize);
-}
-
-
-FVulkanRingBuffer::FVulkanRingBuffer(FVulkanDevice* InDevice, uint64 TotalSize, VkFlags Usage, VkMemoryPropertyFlags MemPropertyFlags)
-	: VulkanRHI::FDeviceChild(InDevice)
-	, BufferSize(TotalSize)
-	, BufferOffset(0)
-	, BufferAddress(0)
-	, MinAlignment(0)
-{
-	const bool bHasBufferDeviceAddress = InDevice->GetOptionalExtensions().HasBufferDeviceAddress;
-	if (bHasBufferDeviceAddress)
-	{
-		Usage |= VK_BUFFER_USAGE_SHADER_DEVICE_ADDRESS_BIT;
-	}
-
-	check(TotalSize <= (uint64)MAX_uint32);
-	InDevice->GetMemoryManager().AllocateBufferPooled(Allocation, nullptr, TotalSize, 0, Usage, MemPropertyFlags, VulkanRHI::EVulkanAllocationMetaRingBuffer, __FILE__, __LINE__);
-	MinAlignment = Allocation.GetBufferAlignment(Device);
-	// Start by wrapping around to set up the correct fence
-	BufferOffset = TotalSize;
-
-	if (bHasBufferDeviceAddress)
-	{
-		VkBufferDeviceAddressInfo BufferInfo;
-		ZeroVulkanStruct(BufferInfo, VK_STRUCTURE_TYPE_BUFFER_DEVICE_ADDRESS_INFO);
-		BufferInfo.buffer = Allocation.GetBufferHandle();
-		BufferAddress = VulkanRHI::vkGetBufferDeviceAddressKHR(Device->GetInstanceHandle(), &BufferInfo);
-	}
-}
-
-FVulkanRingBuffer::~FVulkanRingBuffer()
-{
-	Device->GetMemoryManager().FreeVulkanAllocation(Allocation);
-}
-
-uint64 FVulkanRingBuffer::WrapAroundAllocateMemory(uint64 Size, uint32 Alignment, FVulkanCmdBuffer* InCmdBuffer)
-{
-	CA_ASSUME(InCmdBuffer != nullptr); // Suppress static analysis warning
-	uint64 AllocationOffset = Align<uint64>(BufferOffset, Alignment);
-	ensure(AllocationOffset + Size > BufferSize);
-
-	// Check to see if we can wrap around the ring buffer
-	if (FenceCmdBuffer)
-	{
-		if (FenceCounter == FenceCmdBuffer->GetFenceSignaledCounterI())
-		{
-			//if (FenceCounter == FenceCmdBuffer->GetSubmittedFenceCounter())
-			{
-				//UE_LOG(LogVulkanRHI, Error, TEXT("Ringbuffer overflow during the same cmd buffer!"));
-			}
-			//else
-			{
-				//UE_LOG(LogVulkanRHI, Error, TEXT("Wrapped around the ring buffer! Waiting for the GPU..."));
-				//Device->GetImmediateContext().GetCommandBufferManager()->WaitForCmdBuffer(FenceCmdBuffer, 0.5f);
-			}
-		}
-	}
-
-	BufferOffset = Size;
-
-	FenceCmdBuffer = InCmdBuffer;
-	FenceCounter = InCmdBuffer->GetSubmittedFenceCounter();
-
-	return 0;
-}
diff --git a/Engine/Source/Runtime/VulkanRHI/Private/VulkanViewport.cpp b/Engine/Source/Runtime/VulkanRHI/Private/VulkanViewport.cpp
--- a/Engine/Source/Runtime/VulkanRHI/Private/VulkanViewport.cpp
+++ b/Engine/Source/Runtime/VulkanRHI/Private/VulkanViewport.cpp
@@ -1292,13 +1292,6 @@ void FVulkanDynamicRHI::RHIResizeViewport(FRHIViewport* ViewportRHI, uint32 Size
 void FVulkanDynamicRHI::RHITick(float DeltaTime)
 {
 	check(IsInGameThread());
-	FVulkanDevice* VulkanDevice = GetDevice();
-
-	ENQUEUE_RENDER_COMMAND(TempFrameReset)(
-		[VulkanDevice](FRHICommandListImmediate& RHICmdList)
-	{
-		VulkanDevice->GetImmediateContext().GetTempFrameAllocationBuffer().Reset();
-	});
 }
 
 FTextureRHIRef FVulkanDynamicRHI::RHIGetViewportBackBuffer(FRHIViewport* ViewportRHI)
