VulkanRHI: Load up the timeline semaphore extension when present to avoid warnings with gpu assisted validation.

#rnx

diff --git a/Engine/Source/Runtime/VulkanRHI/Private/VulkanExtensions.cpp b/Engine/Source/Runtime/VulkanRHI/Private/VulkanExtensions.cpp
--- a/Engine/Source/Runtime/VulkanRHI/Private/VulkanExtensions.cpp
+++ b/Engine/Source/Runtime/VulkanRHI/Private/VulkanExtensions.cpp
@@ -993,6 +993,38 @@ private:
 	VkPhysicalDeviceRayTracingPositionFetchFeaturesKHR  RayTracingPositionFetchFeatures;
 };
 
+// ***** VK_KHR_timeline_semaphore
+class FVulkanKHRTimelineSemaphoreExtension : public FVulkanDeviceExtension
+{
+public:
+
+	FVulkanKHRTimelineSemaphoreExtension(FVulkanDevice* InDevice)
+		: FVulkanDeviceExtension(InDevice, VK_KHR_TIMELINE_SEMAPHORE_EXTENSION_NAME, VULKAN_EXTENSION_ENABLED)
+	{
+	}
+
+	virtual void PrePhysicalDeviceFeatures(VkPhysicalDeviceFeatures2KHR& PhysicalDeviceFeatures2) override final
+	{
+		ZeroVulkanStruct(TimelineSemaphoreFeatures, VK_STRUCTURE_TYPE_PHYSICAL_DEVICE_TIMELINE_SEMAPHORE_FEATURES);
+		AddToPNext(PhysicalDeviceFeatures2, TimelineSemaphoreFeatures);
+	}
+
+	virtual void PostPhysicalDeviceFeatures(FOptionalVulkanDeviceExtensions& ExtensionFlags) override final
+	{
+		bRequirementsPassed = (TimelineSemaphoreFeatures.timelineSemaphore == VK_TRUE);
+	}
+
+	virtual void PreCreateDevice(VkDeviceCreateInfo& DeviceCreateInfo) override final
+	{
+		if (bRequirementsPassed)
+		{
+			AddToPNext(DeviceCreateInfo, TimelineSemaphoreFeatures);
+		}
+	}
+
+private:
+	VkPhysicalDeviceTimelineSemaphoreFeatures TimelineSemaphoreFeatures;
+};
 
 // ***** VK_AMD_buffer_marker (vendor)
 class FVulkanAMDBufferMarkerExtension : public FVulkanDeviceExtension
@@ -1610,6 +1642,7 @@ FVulkanDeviceExtensionArray FVulkanDeviceExtension::GetUESupportedDeviceExtensio
 	ADD_CUSTOM_EXTENSION(FVulkanKHRFragmentShaderBarycentricExtension);
 	ADD_CUSTOM_EXTENSION(FVulkanNVComputeShaderDerivatives);
 	ADD_CUSTOM_EXTENSION(FVulkanKHRSamplerYcbcrConversionExtension);
+	ADD_CUSTOM_EXTENSION(FVulkanKHRTimelineSemaphoreExtension);
 
 	// Needed for Raytracing
 	ADD_CUSTOM_EXTENSION(FVulkanKHRBufferDeviceAddressExtension);
