VulkanRHI: Add support for VK_EXT_mesh_shader.  Enable it by default in Vulkan SM6.

#rnx

diff --git a/Engine/Config/VulkanPC/DataDrivenPlatformInfo.ini b/Engine/Config/VulkanPC/DataDrivenPlatformInfo.ini
--- a/Engine/Config/VulkanPC/DataDrivenPlatformInfo.ini
+++ b/Engine/Config/VulkanPC/DataDrivenPlatformInfo.ini
@@ -191,6 +191,10 @@ bSupportsNanite = true
 bSupportsSceneDataCompressedTransforms = true
 bSupportsLumenGI = true
 
+bSupportsMeshShadersTier0=true
+bSupportsMeshShadersTier1=true
+MaxMeshShaderThreadGroupSize=128
+
 ; //#todo-vulkanddpi
 bSupportsSSDIndirect = false
 
diff --git a/Engine/Shaders/Public/Platform/Vulkan/VulkanCommon.ush b/Engine/Shaders/Public/Platform/Vulkan/VulkanCommon.ush
--- a/Engine/Shaders/Public/Platform/Vulkan/VulkanCommon.ush
+++ b/Engine/Shaders/Public/Platform/Vulkan/VulkanCommon.ush
@@ -7,7 +7,7 @@
 #pragma once
 
 // Update this GUID to invalidate shader recompilation for only Vulkan shaders
-#pragma message("UESHADERMETADATA_VERSION AE91ACB9-5CCF-4845-8BE4-6887EE1241EE")
+#pragma message("UESHADERMETADATA_VERSION 368AA5D0-57AB-4A02-BE50-0E3432A4820D")
 
 #define COMPILER_SUPPORTS_ULONG_TYPES 1
 #define COMPILER_SUPPORTS_UINT64_IMAGE_ATOMICS (1 && COMPILER_SUPPORTS_ULONG_TYPES)
@@ -115,3 +115,25 @@ void ImageInterlockedAddUInt64(RWTexture2D<UlongType> Dest, uint2 Coord, UlongTy
 	// Will get filled in with the bindless index in VulkanHitGroupSystemParameters.BindlessHitGroupSystemVertexBuffer
 	ByteAddressBuffer HitGroupSystemVertexBuffer;
 #endif
+
+#if PLATFORM_SUPPORTS_MESH_SHADERS_TIER0
+	#define MESH_SHADER_TRIANGLE_ATTRIBUTES(InThreadsX) [numthreads(InThreadsX, 1, 1)][outputtopology("triangle")]
+	#define MESH_SHADER_LINE_ATTRIBUTES(InThreadsX) [numthreads(InThreadsX, 1, 1)][outputtopology("line")]
+
+	#define AMPLIFICATION_SHADER_TRIANGLE_ATTRIBUTES(InThreadsX) [numthreads(InThreadsX, 1, 1)][outputtopology("triangle")]
+	#define AMPLIFICATION_SHADER_LINE_ATTRIBUTES(InThreadsX) [numthreads(InThreadsX, 1, 1)][outputtopology("line")]
+
+	#define MESH_SHADER_VERTEX_EXPORT(VertexType, InMaxVertexCount) out vertices VertexType OutVertices[InMaxVertexCount]
+	#define MESH_SHADER_TRIANGLE_EXPORT(InMaxTriangleCount) out indices uint3 OutTriangles[InMaxTriangleCount]
+	#define MESH_SHADER_LINE_EXPORT(InMaxLineCount) out indices uint2 OutLines[InMaxLineCount]
+
+	#define MESH_SHADER_WRITE_VERTEX(InVertexIndex, Value) OutVertices[InVertexIndex] = Value;
+	#define MESH_SHADER_WRITE_TRIANGLE(InTriangleIndex, Value) OutTriangles[InTriangleIndex] = Value;
+	#define MESH_SHADER_WRITE_LINE(InLineIndex, Value) OutLines[InLineIndex] = Value;
+#endif
+
+#if PLATFORM_SUPPORTS_MESH_SHADERS_TIER1
+	#define MESH_SHADER_PRIMITIVE_EXPORT(PrimitiveType, InMaxPrimitiveCount) out primitives PrimitiveType OutPrimitives[InMaxPrimitiveCount]
+
+	#define MESH_SHADER_WRITE_PRIMITIVE(InPrimitiveIndex, Value) OutPrimitives[InPrimitiveIndex] = Value;
+#endif
diff --git a/Engine/Source/Runtime/Core/Public/Unix/UnixPlatform.h b/Engine/Source/Runtime/Core/Public/Unix/UnixPlatform.h
--- a/Engine/Source/Runtime/Core/Public/Unix/UnixPlatform.h
+++ b/Engine/Source/Runtime/Core/Public/Unix/UnixPlatform.h
@@ -69,6 +69,8 @@ typedef FUnixPlatformTypes FPlatformTypes;
 #define PLATFORM_SUPPORTS_VIRTUAL_TEXTURE_STREAMING		1
 #define PLATFORM_USE_ANSI_POSIX_MALLOC					1
 #define PLATFORM_SUPPORTS_BINDLESS_RENDERING			1
+#define PLATFORM_SUPPORTS_MESH_SHADERS					1
+
 
 #define PLATFORM_RETURN_ADDRESS_FOR_CALLSTACKTRACING	PLATFORM_RETURN_ADDRESS
 
diff --git a/Engine/Source/Runtime/VulkanRHI/Private/VulkanCommands.cpp b/Engine/Source/Runtime/VulkanRHI/Private/VulkanCommands.cpp
--- a/Engine/Source/Runtime/VulkanRHI/Private/VulkanCommands.cpp
+++ b/Engine/Source/Runtime/VulkanRHI/Private/VulkanCommands.cpp
@@ -60,6 +60,14 @@ static FORCEINLINE ShaderStage::EStage GetAndVerifyShaderStage(FRHIGraphicsShade
 	case SF_Vertex:
 		check(PendingGfxState->GetCurrentShaderKey(ShaderStage::Vertex) == GetShaderKey<FVulkanVertexShader>(ShaderRHI));
 		return ShaderStage::Vertex;
+#if PLATFORM_SUPPORTS_MESH_SHADERS
+	case SF_Mesh:
+		check(PendingGfxState->GetCurrentShaderKey(ShaderStage::Mesh) == GetShaderKey<FVulkanMeshShader>(ShaderRHI));
+		return ShaderStage::Mesh;
+	case SF_Amplification:
+		check(PendingGfxState->GetCurrentShaderKey(ShaderStage::Task) == GetShaderKey<FVulkanTaskShader>(ShaderRHI));
+		return ShaderStage::Task;
+#endif
 	case SF_Geometry:
 #if VULKAN_SUPPORTS_GEOMETRY_SHADERS
 		check(PendingGfxState->GetCurrentShaderKey(ShaderStage::Geometry) == GetShaderKey<FVulkanGeometryShader>(ShaderRHI));
@@ -73,7 +81,7 @@ static FORCEINLINE ShaderStage::EStage GetAndVerifyShaderStage(FRHIGraphicsShade
 		check(PendingGfxState->GetCurrentShaderKey(ShaderStage::Pixel) == GetShaderKey<FVulkanPixelShader>(ShaderRHI));
 		return ShaderStage::Pixel;
 	default:
-		checkf(0, TEXT("Undefined FRHIShader Type %d!"), (int32)ShaderRHI->GetFrequency());
+		checkf(0, TEXT("Undefined FRHIShader Frequency %d!"), (int32)ShaderRHI->GetFrequency());
 		break;
 	}
 
@@ -88,6 +96,14 @@ static FORCEINLINE ShaderStage::EStage GetAndVerifyShaderStageAndVulkanShader(FR
 		//check(PendingGfxState->GetCurrentShaderKey(ShaderStage::Vertex) == GetShaderKey<FVulkanVertexShader>(ShaderRHI));
 		OutShader = static_cast<FVulkanVertexShader*>(static_cast<FRHIVertexShader*>(ShaderRHI));
 		return ShaderStage::Vertex;
+#if PLATFORM_SUPPORTS_MESH_SHADERS
+	case SF_Mesh:
+		OutShader = static_cast<FVulkanMeshShader*>(static_cast<FRHIMeshShader*>(ShaderRHI));
+		return ShaderStage::Mesh;
+	case SF_Amplification:
+		OutShader = static_cast<FVulkanTaskShader*>(static_cast<FVulkanTaskShader*>(ShaderRHI));
+		return ShaderStage::Task;
+#endif
 	case SF_Geometry:
 #if VULKAN_SUPPORTS_GEOMETRY_SHADERS
 		//check(PendingGfxState->GetCurrentShaderKey(ShaderStage::Geometry) == GetShaderKey<FVulkanGeometryShader>(ShaderRHI));
@@ -103,7 +119,7 @@ static FORCEINLINE ShaderStage::EStage GetAndVerifyShaderStageAndVulkanShader(FR
 		OutShader = static_cast<FVulkanPixelShader*>(static_cast<FRHIPixelShader*>(ShaderRHI));
 		return ShaderStage::Pixel;
 	default:
-		checkf(0, TEXT("Undefined FRHIShader Type %d!"), (int32)ShaderRHI->GetFrequency());
+		checkf(0, TEXT("Undefined FRHIShader Frequency %d!"), (int32)ShaderRHI->GetFrequency());
 		break;
 	}
 
@@ -243,19 +259,19 @@ void FVulkanCommandListContext::CommitGraphicsResourceTables()
 	}
 
 #if PLATFORM_SUPPORTS_MESH_SHADERS
-	// :todo-jn: mesh shaders
-	//if (const FVulkanShader* Shader = PendingGfxState->GetCurrentShader(SF_Mesh))
-	//{
-	//  checkSlow(Shader->Frequency == SF_Mesh);
-	//	const FVulkanMeshShader* MeshShader = static_cast<const FVulkanMeshShader*>(Shader);
-	//	SetResourcesFromTables(MeshShader);
-	//}
-	//if (const FVulkanShader* Shader = PendingGfxState->GetCurrentShader(SF_Amplification))
-	//{
-	//  checkSlow(Shader->Frequency == SF_Amplification);
-	//	const FVulkanAmplificationShader* AmplificationShader = static_cast<const FVulkanAmplificationShader*>(Shader);
-	//	SetResourcesFromTables(AmplificationShader);
-	//}
+	if (const FVulkanShader* Shader = PendingGfxState->GetCurrentShader(SF_Mesh))
+	{
+		checkSlow(Shader->Frequency == SF_Mesh);
+		const FVulkanMeshShader* MeshShader = static_cast<const FVulkanMeshShader*>(Shader);
+		SetResourcesFromTables(MeshShader);
+	}
+
+	if (const FVulkanShader* Shader = PendingGfxState->GetCurrentShader(SF_Amplification))
+	{
+		checkSlow(Shader->Frequency == SF_Amplification);
+		const FVulkanTaskShader* AmplificationShader = static_cast<const FVulkanTaskShader*>(Shader);
+		SetResourcesFromTables(AmplificationShader);
+	}
 #endif
 
 #if PLATFORM_SUPPORTS_GEOMETRY_SHADERS
@@ -745,6 +761,42 @@ void FVulkanCommandListContext::RHIDrawIndexedPrimitiveIndirect(FRHIBuffer* Inde
 	}
 }
 
+#if PLATFORM_SUPPORTS_MESH_SHADERS
+void FVulkanCommandListContext::RHIDispatchMeshShader(uint32 ThreadGroupCountX, uint32 ThreadGroupCountY, uint32 ThreadGroupCountZ)
+{
+	CommitGraphicsResourceTables();
+
+	FVulkanCmdBuffer* Cmd = CommandBufferManager->GetActiveCmdBuffer();
+
+	PendingGfxState->PrepareForDraw(Cmd);
+
+	VulkanRHI::vkCmdDrawMeshTasksEXT(Cmd->GetHandle(), ThreadGroupCountX, ThreadGroupCountY, ThreadGroupCountZ);
+
+	if (FVulkanPlatform::RegisterGPUWork() && IsImmediate())
+	{
+		RegisterGPUWork(1);
+	}
+}
+
+void FVulkanCommandListContext::RHIDispatchIndirectMeshShader(FRHIBuffer* ArgumentBufferRHI, uint32 ArgumentOffset)
+{
+	CommitGraphicsResourceTables();
+
+	FVulkanCmdBuffer* Cmd = CommandBufferManager->GetActiveCmdBuffer();
+
+	PendingGfxState->PrepareForDraw(Cmd);
+
+	FVulkanResourceMultiBuffer* ArgumentBuffer = ResourceCast(ArgumentBufferRHI);
+
+	VulkanRHI::vkCmdDrawMeshTasksIndirectEXT(Cmd->GetHandle(), ArgumentBuffer->GetHandle(), ArgumentBuffer->GetOffset() + ArgumentOffset, 1, sizeof(VkDrawMeshTasksIndirectCommandEXT));
+
+	if (FVulkanPlatform::RegisterGPUWork() && IsImmediate())
+	{
+		RegisterGPUWork(1);
+	}
+}
+#endif
+
 void FVulkanCommandListContext::RHIClearMRT(bool bClearColor, int32 NumClearColors, const FLinearColor* ClearColorArray, bool bClearDepth, float Depth, bool bClearStencil, uint32 Stencil)
 {
 	if (!(bClearColor || bClearDepth || bClearStencil))
diff --git a/Engine/Source/Runtime/VulkanRHI/Private/VulkanContext.h b/Engine/Source/Runtime/VulkanRHI/Private/VulkanContext.h
--- a/Engine/Source/Runtime/VulkanRHI/Private/VulkanContext.h
+++ b/Engine/Source/Runtime/VulkanRHI/Private/VulkanContext.h
@@ -69,6 +69,10 @@ public:
 	virtual void RHIDrawIndexedIndirect(FRHIBuffer* IndexBufferRHI, FRHIBuffer* ArgumentsBufferRHI, int32 DrawArgumentsIndex, uint32 NumInstances) final override;
 	virtual void RHIDrawIndexedPrimitive(FRHIBuffer* IndexBuffer, int32 BaseVertexIndex, uint32 FirstInstance, uint32 NumVertices, uint32 StartIndex, uint32 NumPrimitives, uint32 NumInstances) final override;
 	virtual void RHIDrawIndexedPrimitiveIndirect(FRHIBuffer* IndexBuffer, FRHIBuffer* ArgumentBuffer, uint32 ArgumentOffset) final override;
+#if PLATFORM_SUPPORTS_MESH_SHADERS
+	virtual void RHIDispatchMeshShader(uint32 ThreadGroupCountX, uint32 ThreadGroupCountY, uint32 ThreadGroupCountZ) final override;
+	virtual void RHIDispatchIndirectMeshShader(FRHIBuffer* ArgumentBuffer, uint32 ArgumentOffset) final override;
+#endif
 	virtual void RHISetDepthBounds(float MinDepth, float MaxDepth) final override;
 #if WITH_RHI_BREADCRUMBS
 	virtual void RHIBeginBreadcrumbGPU(FRHIBreadcrumbNode* Breadcrumb) final override;
diff --git a/Engine/Source/Runtime/VulkanRHI/Private/VulkanDevice.h b/Engine/Source/Runtime/VulkanRHI/Private/VulkanDevice.h
--- a/Engine/Source/Runtime/VulkanRHI/Private/VulkanDevice.h
+++ b/Engine/Source/Runtime/VulkanRHI/Private/VulkanDevice.h
@@ -139,6 +139,8 @@ struct FOptionalVulkanDeviceExtensionProperties
 
 	VkPhysicalDeviceFragmentShaderBarycentricPropertiesKHR FragmentShaderBarycentricProps;
 	VkPhysicalDeviceComputeShaderDerivativesFeaturesNV ComputeShaderDerivativesFeatures;
+
+	VkPhysicalDeviceMeshShaderPropertiesEXT MeshShaderProperties;
 };
 
 class FVulkanPhysicalDeviceFeatures
diff --git a/Engine/Source/Runtime/VulkanRHI/Private/VulkanExtensions.cpp b/Engine/Source/Runtime/VulkanRHI/Private/VulkanExtensions.cpp
--- a/Engine/Source/Runtime/VulkanRHI/Private/VulkanExtensions.cpp
+++ b/Engine/Source/Runtime/VulkanRHI/Private/VulkanExtensions.cpp
@@ -1026,6 +1026,48 @@ private:
 	VkPhysicalDeviceTimelineSemaphoreFeatures TimelineSemaphoreFeatures;
 };
 
+
+// ***** VK_EXT_mesh_shader
+class FVulkanEXTMeshShaderExtension : public FVulkanDeviceExtension
+{
+public:
+
+	FVulkanEXTMeshShaderExtension(FVulkanDevice* InDevice)
+		: FVulkanDeviceExtension(InDevice, VK_EXT_MESH_SHADER_EXTENSION_NAME, VULKAN_EXTENSION_ENABLED)
+	{
+	}
+
+	virtual void PrePhysicalDeviceFeatures(VkPhysicalDeviceFeatures2KHR& PhysicalDeviceFeatures2) override final
+	{
+		ZeroVulkanStruct(MeshShaderFeatures, VK_STRUCTURE_TYPE_PHYSICAL_DEVICE_MESH_SHADER_FEATURES_EXT);
+		AddToPNext(PhysicalDeviceFeatures2, MeshShaderFeatures);
+	}
+
+	virtual void PostPhysicalDeviceFeatures(FOptionalVulkanDeviceExtensions& ExtensionFlags) override final
+	{
+		bRequirementsPassed = (MeshShaderFeatures.meshShader == VK_TRUE) && (MeshShaderFeatures.multiviewMeshShader == VK_TRUE);
+	}
+
+	virtual void PrePhysicalDeviceProperties(VkPhysicalDeviceProperties2KHR& PhysicalDeviceProperties2) override final
+	{
+		VkPhysicalDeviceMeshShaderPropertiesEXT& MeshShaderProperties = GetDeviceExtensionProperties().MeshShaderProperties;
+		ZeroVulkanStruct(MeshShaderProperties, VK_STRUCTURE_TYPE_PHYSICAL_DEVICE_MESH_SHADER_PROPERTIES_EXT);
+		AddToPNext(PhysicalDeviceProperties2, MeshShaderProperties);
+	}
+
+	virtual void PreCreateDevice(VkDeviceCreateInfo& DeviceCreateInfo) override final
+	{
+		if (bRequirementsPassed)
+		{
+			AddToPNext(DeviceCreateInfo, MeshShaderFeatures);
+		}
+	}
+
+private:
+	VkPhysicalDeviceMeshShaderFeaturesEXT MeshShaderFeatures;
+};
+
+
 // ***** VK_AMD_buffer_marker (vendor)
 class FVulkanAMDBufferMarkerExtension : public FVulkanDeviceExtension
 {
@@ -1643,6 +1685,7 @@ FVulkanDeviceExtensionArray FVulkanDeviceExtension::GetUESupportedDeviceExtensio
 	ADD_CUSTOM_EXTENSION(FVulkanNVComputeShaderDerivatives);
 	ADD_CUSTOM_EXTENSION(FVulkanKHRSamplerYcbcrConversionExtension);
 	ADD_CUSTOM_EXTENSION(FVulkanKHRTimelineSemaphoreExtension);
+	ADD_CUSTOM_EXTENSION(FVulkanEXTMeshShaderExtension);
 
 	// Needed for Raytracing
 	ADD_CUSTOM_EXTENSION(FVulkanKHRBufferDeviceAddressExtension);
diff --git a/Engine/Source/Runtime/VulkanRHI/Private/VulkanLoader.h b/Engine/Source/Runtime/VulkanRHI/Private/VulkanLoader.h
--- a/Engine/Source/Runtime/VulkanRHI/Private/VulkanLoader.h
+++ b/Engine/Source/Runtime/VulkanRHI/Private/VulkanLoader.h
@@ -224,6 +224,9 @@
 	EnumMacro(PFN_vkDeferredOperationJoinKHR, vkDeferredOperationJoinKHR) \
 	EnumMacro(PFN_vkGetDeferredOperationMaxConcurrencyKHR, vkGetDeferredOperationMaxConcurrencyKHR) \
 	EnumMacro(PFN_vkGetDeferredOperationResultKHR, vkGetDeferredOperationResultKHR) \
+	EnumMacro(PFN_vkCmdDrawMeshTasksEXT, vkCmdDrawMeshTasksEXT) \
+	EnumMacro(PFN_vkCmdDrawMeshTasksIndirectCountEXT, vkCmdDrawMeshTasksIndirectCountEXT) \
+	EnumMacro(PFN_vkCmdDrawMeshTasksIndirectEXT, vkCmdDrawMeshTasksIndirectEXT) \
 	EnumMacro(PFN_vkCmdWriteBufferMarkerAMD, vkCmdWriteBufferMarkerAMD) \
 	EnumMacro(PFN_vkCmdSetCheckpointNV, vkCmdSetCheckpointNV) \
 	EnumMacro(PFN_vkGetQueueCheckpointDataNV, vkGetQueueCheckpointDataNV) \
diff --git a/Engine/Source/Runtime/VulkanRHI/Private/VulkanPipeline.cpp b/Engine/Source/Runtime/VulkanRHI/Private/VulkanPipeline.cpp
--- a/Engine/Source/Runtime/VulkanRHI/Private/VulkanPipeline.cpp
+++ b/Engine/Source/Runtime/VulkanRHI/Private/VulkanPipeline.cpp
@@ -184,6 +184,10 @@ static inline FSHAHash GetShaderHashForStage(const FGraphicsPipelineStateInitial
 	{
 	case ShaderStage::Vertex:		return GetShaderHash<FRHIVertexShader, FVulkanVertexShader>(PSOInitializer.BoundShaderState.VertexShaderRHI);
 	case ShaderStage::Pixel:		return GetShaderHash<FRHIPixelShader, FVulkanPixelShader>(PSOInitializer.BoundShaderState.PixelShaderRHI);
+#if PLATFORM_SUPPORTS_MESH_SHADERS
+	case ShaderStage::Mesh:			return GetShaderHash<FRHIMeshShader, FVulkanMeshShader>(PSOInitializer.BoundShaderState.GetMeshShader());
+	case ShaderStage::Task:			return GetShaderHash<FRHIAmplificationShader, FVulkanTaskShader>(PSOInitializer.BoundShaderState.GetAmplificationShader());
+#endif
 #if VULKAN_SUPPORTS_GEOMETRY_SHADERS
 	case ShaderStage::Geometry:		return GetShaderHash<FRHIGeometryShader, FVulkanGeometryShader>(PSOInitializer.BoundShaderState.GetGeometryShader());
 #endif
@@ -1182,6 +1186,16 @@ static FString GfxShaderHashesToString(FVulkanShader* Shaders[ShaderStage::NumGr
 	{
 		ShaderHashes += TEXT("PS: ") + static_cast<FVulkanPixelShader*>(Shaders[ShaderStage::Pixel])->GetHash().ToString() + TEXT("\n");
 	}
+#if PLATFORM_SUPPORTS_MESH_SHADERS
+	if (Shaders[ShaderStage::Mesh])
+	{
+		ShaderHashes += TEXT("MS: ") + static_cast<FVulkanMeshShader*>(Shaders[ShaderStage::Mesh])->GetHash().ToString() + TEXT("\n");
+	}
+	if (Shaders[ShaderStage::Task])
+	{
+		ShaderHashes += TEXT("AS: ") + static_cast<FVulkanTaskShader*>(Shaders[ShaderStage::Task])->GetHash().ToString() + TEXT("\n");
+	}
+#endif
 #if VULKAN_SUPPORTS_GEOMETRY_SHADERS
 	if (Shaders[ShaderStage::Geometry])
 	{
@@ -1690,6 +1704,10 @@ FVulkanShaderHashes::FVulkanShaderHashes(const FGraphicsPipelineStateInitializer
 {
 	Stages[ShaderStage::Vertex] = GetShaderHash<FRHIVertexShader, FVulkanVertexShader>(PSOInitializer.BoundShaderState.VertexShaderRHI);
 	Stages[ShaderStage::Pixel] = GetShaderHash<FRHIPixelShader, FVulkanPixelShader>(PSOInitializer.BoundShaderState.PixelShaderRHI);
+#if PLATFORM_SUPPORTS_MESH_SHADERS
+	Stages[ShaderStage::Mesh] = GetShaderHash<FRHIMeshShader, FVulkanMeshShader>(PSOInitializer.BoundShaderState.GetMeshShader());
+	Stages[ShaderStage::Task] = GetShaderHash<FRHIAmplificationShader, FVulkanTaskShader>(PSOInitializer.BoundShaderState.GetAmplificationShader());
+#endif
 #if VULKAN_SUPPORTS_GEOMETRY_SHADERS
 	Stages[ShaderStage::Geometry] = GetShaderHash<FRHIGeometryShader, FVulkanGeometryShader>(PSOInitializer.BoundShaderState.GetGeometryShader());
 #endif
@@ -1752,12 +1770,15 @@ void FVulkanPipelineStateCacheManager::CreateGfxEntry(const FGraphicsPipelineSta
 	{
 		const FBoundShaderStateInput& BSI = PSOInitializer.BoundShaderState;
 
-		const FVulkanShaderHeader& VSHeader = Shaders[ShaderStage::Vertex]->GetCodeHeader();
-		VertexInputState.Generate(ResourceCast(PSOInitializer.BoundShaderState.VertexDeclarationRHI), VSHeader.InOutMask);
-
 		FUniformBufferGatherInfo UBGatherInfo;
 
-		DescriptorSetLayoutInfo.ProcessBindingsForStage(VK_SHADER_STAGE_VERTEX_BIT, ShaderStage::Vertex, VSHeader, UBGatherInfo);
+		if (Shaders[ShaderStage::Vertex])
+		{
+			const FVulkanShaderHeader& VSHeader = Shaders[ShaderStage::Vertex]->GetCodeHeader();
+			VertexInputState.Generate(ResourceCast(PSOInitializer.BoundShaderState.VertexDeclarationRHI), VSHeader.InOutMask);
+
+			DescriptorSetLayoutInfo.ProcessBindingsForStage(VK_SHADER_STAGE_VERTEX_BIT, ShaderStage::Vertex, VSHeader, UBGatherInfo);
+		}
 
 		if (Shaders[ShaderStage::Pixel])
 		{
@@ -1772,6 +1793,20 @@ void FVulkanPipelineStateCacheManager::CreateGfxEntry(const FGraphicsPipelineSta
 			}
 		}
 
+#if PLATFORM_SUPPORTS_MESH_SHADERS
+		if (Shaders[ShaderStage::Mesh])
+		{
+			const FVulkanShaderHeader& MSHeader = Shaders[ShaderStage::Mesh]->GetCodeHeader();
+			DescriptorSetLayoutInfo.ProcessBindingsForStage(VK_SHADER_STAGE_MESH_BIT_EXT, ShaderStage::Mesh, MSHeader, UBGatherInfo);
+		}
+
+		if (Shaders[ShaderStage::Task])
+		{
+			const FVulkanShaderHeader& TSHeader = Shaders[ShaderStage::Task]->GetCodeHeader();
+			DescriptorSetLayoutInfo.ProcessBindingsForStage(VK_SHADER_STAGE_TASK_BIT_EXT, ShaderStage::Task, TSHeader, UBGatherInfo);
+		}
+#endif
+
 #if VULKAN_SUPPORTS_GEOMETRY_SHADERS
 		if (Shaders[ShaderStage::Geometry])
 		{
@@ -1932,6 +1967,10 @@ FVulkanRHIGraphicsPipelineState::FVulkanRHIGraphicsPipelineState(FVulkanDevice*
 
 	FMemory::Memset(VulkanShaders, 0, sizeof(VulkanShaders));
 	VulkanShaders[ShaderStage::Vertex] = static_cast<FVulkanVertexShader*>(PSOInitializer_.BoundShaderState.VertexShaderRHI);
+#if PLATFORM_SUPPORTS_MESH_SHADERS
+	VulkanShaders[ShaderStage::Mesh] = static_cast<FVulkanMeshShader*>(PSOInitializer_.BoundShaderState.GetMeshShader());
+	VulkanShaders[ShaderStage::Task] = static_cast<FVulkanTaskShader*>(PSOInitializer_.BoundShaderState.GetAmplificationShader());
+#endif
 #if PLATFORM_SUPPORTS_GEOMETRY_SHADERS
 	VulkanShaders[ShaderStage::Geometry] = static_cast<FVulkanGeometryShader*>(PSOInitializer_.BoundShaderState.GetGeometryShader());
 #endif
@@ -2092,10 +2131,13 @@ FGraphicsPipelineStateRHIRef FVulkanPipelineStateCacheManager::RHICreateGraphics
 				NewPSO->ShaderKeys[StageIdx] = GetShaderKeyForGfxStage(BSI, (ShaderStage::EStage)StageIdx);
 			}
 
-			check(BSI.VertexShaderRHI);
-			FVulkanVertexShader* VS = ResourceCast(BSI.VertexShaderRHI);
-			const FVulkanShaderHeader& VSHeader = VS->GetCodeHeader();
-			NewPSO->VertexInputState.Generate(ResourceCast(Initializer.BoundShaderState.VertexDeclarationRHI), VSHeader.InOutMask);
+			if (Initializer.BoundShaderState.VertexDeclarationRHI)
+			{
+				check(BSI.VertexShaderRHI);
+				FVulkanVertexShader* VS = ResourceCast(BSI.VertexShaderRHI);
+				const FVulkanShaderHeader& VSHeader = VS->GetCodeHeader();
+				NewPSO->VertexInputState.Generate(ResourceCast(Initializer.BoundShaderState.VertexDeclarationRHI), VSHeader.InOutMask);
+			}
 
 			if((!bIsPrecache || !LRUEvictImmediately()) 
 	#if !UE_BUILD_SHIPPING
@@ -2390,6 +2432,16 @@ void GetVulkanGfxShaders(const FBoundShaderStateInput& BSI, FVulkanShader* OutSh
 		OutShaders[ShaderStage::Pixel] = ResourceCast(BSI.PixelShaderRHI);
 	}
 
+	if (BSI.GetMeshShader())
+	{
+		OutShaders[ShaderStage::Mesh] = ResourceCast(BSI.GetMeshShader());
+	}
+
+	if (BSI.GetAmplificationShader())
+	{
+		OutShaders[ShaderStage::Task] = ResourceCast(BSI.GetAmplificationShader());
+	}
+
 	if (BSI.GetGeometryShader())
 	{
 #if VULKAN_SUPPORTS_GEOMETRY_SHADERS
diff --git a/Engine/Source/Runtime/VulkanRHI/Private/VulkanPipeline.h b/Engine/Source/Runtime/VulkanRHI/Private/VulkanPipeline.h
--- a/Engine/Source/Runtime/VulkanRHI/Private/VulkanPipeline.h
+++ b/Engine/Source/Runtime/VulkanRHI/Private/VulkanPipeline.h
@@ -43,9 +43,9 @@ inline uint64 GetShaderKeyForGfxStage(const FBoundShaderStateInput& BSI, ShaderS
 	case ShaderStage::Geometry:
 		return GetShaderKey<FVulkanGeometryShader>(BSI.GetGeometryShader());
 	case ShaderStage::Mesh:
-		return 0; // GetShaderKey<FVulkanMeshShader>(BSI.GetMeshShader());
+		return GetShaderKey<FVulkanMeshShader>(BSI.GetMeshShader());
 	case ShaderStage::Task:
-		return 0; // GetShaderKey<FVulkanTaskShader>(BSI.GetAmplificationShader());
+		return GetShaderKey<FVulkanTaskShader>(BSI.GetAmplificationShader());
 	default:
 		check(0);
 	}
@@ -733,6 +733,10 @@ public:
 		{
 		case SF_Vertex: return static_cast<FVulkanVertexShader*>(VulkanShaders[ShaderStage::Vertex]);
 		case SF_Pixel: return static_cast<FVulkanPixelShader*>(VulkanShaders[ShaderStage::Pixel]);
+#if PLATFORM_SUPPORTS_MESH_SHADERS
+		case SF_Mesh: return static_cast<FVulkanGeometryShader*>(VulkanShaders[ShaderStage::Mesh]);
+		case SF_Amplification: return static_cast<FVulkanGeometryShader*>(VulkanShaders[ShaderStage::Task]);
+#endif
 #if VULKAN_SUPPORTS_GEOMETRY_SHADERS
 		case SF_Geometry: return static_cast<FVulkanGeometryShader*>(VulkanShaders[ShaderStage::Geometry]);
 #endif
diff --git a/Engine/Source/Runtime/VulkanRHI/Private/VulkanPipelineState.cpp b/Engine/Source/Runtime/VulkanRHI/Private/VulkanPipelineState.cpp
--- a/Engine/Source/Runtime/VulkanRHI/Private/VulkanPipelineState.cpp
+++ b/Engine/Source/Runtime/VulkanRHI/Private/VulkanPipelineState.cpp
@@ -297,10 +297,16 @@ FVulkanGraphicsPipelineDescriptorState::FVulkanGraphicsPipelineDescriptorState(F
 
 	const FVulkanShaderFactory& ShaderFactory = Device->GetShaderFactory();
 
-	const FVulkanVertexShader* VertexShader = ShaderFactory.LookupShader<FVulkanVertexShader>(InGfxPipeline->GetShaderKey(SF_Vertex));
-	check(VertexShader);
-	PackedUniformBuffers[ShaderStage::Vertex].Init(VertexShader->GetCodeHeader(), PackedUniformBuffersMask[ShaderStage::Vertex]);
-	UsedSetsMask |= VertexShader->GetCodeHeader().Bindings.Num() ? (1u << ShaderStage::Vertex) : 0u;
+
+	uint64 VertexShaderKey = InGfxPipeline->GetShaderKey(SF_Vertex);
+	if (VertexShaderKey)
+	{
+		const FVulkanVertexShader* VertexShader = ShaderFactory.LookupShader<FVulkanVertexShader>(InGfxPipeline->GetShaderKey(SF_Vertex));
+		check(VertexShader);
+
+		PackedUniformBuffers[ShaderStage::Vertex].Init(VertexShader->GetCodeHeader(), PackedUniformBuffersMask[ShaderStage::Vertex]);
+		UsedSetsMask |= VertexShader->GetCodeHeader().Bindings.Num() ? (1u << ShaderStage::Vertex) : 0u;
+	}
 
 	uint64 PixelShaderKey = InGfxPipeline->GetShaderKey(SF_Pixel);
 	if (PixelShaderKey)
@@ -312,6 +318,28 @@ FVulkanGraphicsPipelineDescriptorState::FVulkanGraphicsPipelineDescriptorState(F
 		UsedSetsMask |= PixelShader->GetCodeHeader().Bindings.Num() ? (1u << ShaderStage::Pixel) : 0u;
 	}
 
+#if PLATFORM_SUPPORTS_MESH_SHADERS
+	uint64 MeshShaderKey = InGfxPipeline->GetShaderKey(SF_Mesh);
+	if (MeshShaderKey)
+	{
+		const FVulkanMeshShader* MeshShader = ShaderFactory.LookupShader<FVulkanMeshShader>(MeshShaderKey);
+		check(MeshShader);
+
+		PackedUniformBuffers[ShaderStage::Mesh].Init(MeshShader->GetCodeHeader(), PackedUniformBuffersMask[ShaderStage::Mesh]);
+		UsedSetsMask |= MeshShader->GetCodeHeader().Bindings.Num() ? (1u << ShaderStage::Mesh) : 0u;
+	}
+
+	uint64 TaskShaderKey = InGfxPipeline->GetShaderKey(SF_Amplification);
+	if (TaskShaderKey)
+	{
+		const FVulkanTaskShader* TaskShader = ShaderFactory.LookupShader<FVulkanTaskShader>(TaskShaderKey);
+		check(TaskShader);
+
+		PackedUniformBuffers[ShaderStage::Task].Init(TaskShader->GetCodeHeader(), PackedUniformBuffersMask[ShaderStage::Task]);
+		UsedSetsMask |= TaskShader->GetCodeHeader().Bindings.Num() ? (1u << ShaderStage::Task) : 0u;
+	}
+#endif
+
 #if VULKAN_SUPPORTS_GEOMETRY_SHADERS
 	uint64 GeometryShaderKey = InGfxPipeline->GetShaderKey(SF_Geometry);
 	if (GeometryShaderKey)
@@ -534,6 +562,10 @@ void FVulkanCommandListContext::RHISetGraphicsPipelineState(FRHIGraphicsPipeline
 	if (bApplyAdditionalState)
 	{
 		ApplyStaticUniformBuffers(static_cast<FVulkanVertexShader*>(Pipeline->VulkanShaders[ShaderStage::Vertex]));
+#if PLATFORM_SUPPORTS_MESH_SHADERS
+		ApplyStaticUniformBuffers(static_cast<FVulkanMeshShader*>(Pipeline->VulkanShaders[ShaderStage::Mesh]));
+		ApplyStaticUniformBuffers(static_cast<FVulkanTaskShader*>(Pipeline->VulkanShaders[ShaderStage::Task]));
+#endif
 #if PLATFORM_SUPPORTS_GEOMETRY_SHADERS
 		ApplyStaticUniformBuffers(static_cast<FVulkanGeometryShader*>(Pipeline->VulkanShaders[ShaderStage::Geometry]));
 #endif
diff --git a/Engine/Source/Runtime/VulkanRHI/Private/VulkanRHI.cpp b/Engine/Source/Runtime/VulkanRHI/Private/VulkanRHI.cpp
--- a/Engine/Source/Runtime/VulkanRHI/Private/VulkanRHI.cpp
+++ b/Engine/Source/Runtime/VulkanRHI/Private/VulkanRHI.cpp
@@ -998,6 +998,12 @@ void FVulkanDynamicRHI::InitInstance()
 		GRHIMaxDispatchThreadGroupsPerDimension.Y = FMath::Min<uint32>(Limits.maxComputeWorkGroupCount[1], 0x7fffffff);
 		GRHIMaxDispatchThreadGroupsPerDimension.Z = FMath::Min<uint32>(Limits.maxComputeWorkGroupCount[2], 0x7fffffff);
 
+#if PLATFORM_SUPPORTS_MESH_SHADERS
+		// If mesh shaders are enabled in DDPI (currently SM6), then the profile check will ensure it's supported
+		GRHIGlobals.SupportsMeshShadersTier0 = RHISupportsMeshShadersTier0(GMaxRHIShaderPlatform);
+		GRHIGlobals.SupportsMeshShadersTier1 = RHISupportsMeshShadersTier1(GMaxRHIShaderPlatform);
+#endif
+
 		FVulkanPlatform::SetupFeatureLevels(GRHIGlobals.ShaderPlatformForFeatureLevel);
 
 		GRHIRequiresRenderTargetForPixelShaderUAVs = true;
@@ -2123,6 +2129,7 @@ uint64 FVulkanDynamicRHI::RHIComputeStatePrecachePSOHash(const FGraphicsPipeline
 #endif // PLATFORM_SUPPORTS_GEOMETRY_SHADERS
 #if PLATFORM_SUPPORTS_MESH_SHADERS
 		uint32 MeshShader;
+		uint32 TaskShader;
 #endif // PLATFORM_SUPPORTS_MESH_SHADERS
 		uint32 BlendState;
 		uint32 RasterizerState;
@@ -2151,6 +2158,7 @@ uint64 FVulkanDynamicRHI::RHIComputeStatePrecachePSOHash(const FGraphicsPipeline
 #endif
 #if PLATFORM_SUPPORTS_MESH_SHADERS
 	HashKey.MeshShader = Initializer.BoundShaderState.GetMeshShader() ? GetTypeHash(Initializer.BoundShaderState.GetMeshShader()->GetHash()) : 0;
+	HashKey.TaskShader = Initializer.BoundShaderState.GetAmplificationShader() ? GetTypeHash(Initializer.BoundShaderState.GetAmplificationShader()->GetHash()) : 0;
 #endif
 
 	FBlendStateInitializerRHI BlendStateInitializerRHI;
diff --git a/Engine/Source/Runtime/VulkanRHI/Private/VulkanRHIPrivate.h b/Engine/Source/Runtime/VulkanRHI/Private/VulkanRHIPrivate.h
--- a/Engine/Source/Runtime/VulkanRHI/Private/VulkanRHIPrivate.h
+++ b/Engine/Source/Runtime/VulkanRHI/Private/VulkanRHIPrivate.h
@@ -60,6 +60,8 @@ inline VkShaderStageFlagBits UEFrequencyToVKStageBit(EShaderFrequency InStage)
 	switch (InStage)
 	{
 	case SF_Vertex:			return VK_SHADER_STAGE_VERTEX_BIT;
+	case SF_Mesh:			return VK_SHADER_STAGE_MESH_BIT_EXT;
+	case SF_Amplification:	return VK_SHADER_STAGE_TASK_BIT_EXT;
 	case SF_Pixel:			return VK_SHADER_STAGE_FRAGMENT_BIT;
 	case SF_Geometry:		return VK_SHADER_STAGE_GEOMETRY_BIT;
 	case SF_Compute:		return VK_SHADER_STAGE_COMPUTE_BIT;
@@ -81,6 +83,8 @@ inline EShaderFrequency VkStageBitToUEFrequency(VkShaderStageFlagBits FlagBits)
 	{
 	case VK_SHADER_STAGE_VERTEX_BIT:					return SF_Vertex;
 	case VK_SHADER_STAGE_FRAGMENT_BIT:					return SF_Pixel;
+	case VK_SHADER_STAGE_MESH_BIT_EXT:					return SF_Mesh;
+	case VK_SHADER_STAGE_TASK_BIT_EXT:					return SF_Amplification;
 	case VK_SHADER_STAGE_GEOMETRY_BIT:					return SF_Geometry;
 	case VK_SHADER_STAGE_COMPUTE_BIT:					return SF_Compute;
 
diff --git a/Engine/Source/Runtime/VulkanRHI/Private/VulkanResources.h b/Engine/Source/Runtime/VulkanRHI/Private/VulkanResources.h
--- a/Engine/Source/Runtime/VulkanRHI/Private/VulkanResources.h
+++ b/Engine/Source/Runtime/VulkanRHI/Private/VulkanResources.h
@@ -267,6 +267,8 @@ typedef TVulkanBaseShader<FRHIVertexShader, SF_Vertex>				FVulkanVertexShader;
 typedef TVulkanBaseShader<FRHIPixelShader, SF_Pixel>				FVulkanPixelShader;
 typedef TVulkanBaseShader<FRHIComputeShader, SF_Compute>			FVulkanComputeShader;
 typedef TVulkanBaseShader<FRHIGeometryShader, SF_Geometry>			FVulkanGeometryShader;
+typedef TVulkanBaseShader<FRHIMeshShader, SF_Mesh>					FVulkanMeshShader;
+typedef TVulkanBaseShader<FRHIAmplificationShader, SF_Amplification> FVulkanTaskShader;
 
 class FVulkanRayTracingShader : public FRHIRayTracingShader, public FVulkanShader
 {
@@ -355,6 +357,8 @@ public:
 
 	FORCEINLINE FVulkanVertexShader*   GetVertexShader() const { return (FVulkanVertexShader*)CacheLink.GetVertexShader(); }
 	FORCEINLINE FVulkanPixelShader*    GetPixelShader() const { return (FVulkanPixelShader*)CacheLink.GetPixelShader(); }
+	FORCEINLINE FVulkanMeshShader*     GetMeshShader() const { return (FVulkanMeshShader*)CacheLink.GetMeshShader(); }
+	FORCEINLINE FVulkanTaskShader*     GetTaskShader() const { return (FVulkanTaskShader*)CacheLink.GetAmplificationShader(); }
 	FORCEINLINE FVulkanGeometryShader* GetGeometryShader() const { return (FVulkanGeometryShader*)CacheLink.GetGeometryShader(); }
 
 	const FVulkanShader* GetShader(ShaderStage::EStage Stage) const
@@ -363,6 +367,10 @@ public:
 		{
 		case ShaderStage::Vertex:		return GetVertexShader();
 		case ShaderStage::Pixel:		return GetPixelShader();
+#if PLATFORM_SUPPORTS_MESH_SHADERS
+		case ShaderStage::Mesh:			return GetMeshShader();
+		case ShaderStage::Task:			return GetTaskShader();
+#endif
 #if VULKAN_SUPPORTS_GEOMETRY_SHADERS
 		case ShaderStage::Geometry:	return GetGeometryShader();
 #endif
@@ -1263,6 +1271,16 @@ struct TVulkanResourceTraits<FRHIVertexShader>
 	typedef FVulkanVertexShader TConcreteType;
 };
 template<>
+struct TVulkanResourceTraits<FRHIMeshShader>
+{
+	typedef FVulkanMeshShader TConcreteType;
+};
+template<>
+struct TVulkanResourceTraits<FRHIAmplificationShader>
+{
+	typedef FVulkanTaskShader TConcreteType;
+};
+template<>
 struct TVulkanResourceTraits<FRHIGeometryShader>
 {
 	typedef FVulkanGeometryShader TConcreteType;
diff --git a/Engine/Source/Runtime/VulkanRHI/Private/VulkanShaders.cpp b/Engine/Source/Runtime/VulkanRHI/Private/VulkanShaders.cpp
--- a/Engine/Source/Runtime/VulkanRHI/Private/VulkanShaders.cpp
+++ b/Engine/Source/Runtime/VulkanRHI/Private/VulkanShaders.cpp
@@ -578,6 +578,16 @@ FPixelShaderRHIRef FVulkanDynamicRHI::RHICreatePixelShader(TArrayView<const uint
 	return Device->GetShaderFactory().CreateShader<FVulkanPixelShader>(Code, Device);
 }
 
+FMeshShaderRHIRef FVulkanDynamicRHI::RHICreateMeshShader(TArrayView<const uint8> Code, const FSHAHash& Hash)
+{
+	return Device->GetShaderFactory().CreateShader<FVulkanMeshShader>(Code, Device);
+}
+
+FAmplificationShaderRHIRef FVulkanDynamicRHI::RHICreateAmplificationShader(TArrayView<const uint8> Code, const FSHAHash& Hash)
+{
+	return Device->GetShaderFactory().CreateShader<FVulkanTaskShader>(Code, Device);
+}
+
 FGeometryShaderRHIRef FVulkanDynamicRHI::RHICreateGeometryShader(TArrayView<const uint8> Code, const FSHAHash& Hash)
 { 
 	return Device->GetShaderFactory().CreateShader<FVulkanGeometryShader>(Code, Device);
diff --git a/Engine/Source/Runtime/VulkanRHI/Public/VulkanCommon.h b/Engine/Source/Runtime/VulkanRHI/Public/VulkanCommon.h
--- a/Engine/Source/Runtime/VulkanRHI/Public/VulkanCommon.h
+++ b/Engine/Source/Runtime/VulkanRHI/Public/VulkanCommon.h
@@ -20,8 +20,10 @@
 
 namespace ShaderStage
 {
-	// Adjusting these requires a full shader rebuild (ie modify the guid on VulkanCommon.usf)
-	// Keep the values in sync with EShaderFrequency
+	// There should be one value for each value in EShaderFrequency.
+	// These values are meant to be used as indices in contexts where values for different bind points can overlap (Graphics/Compute/RayTracing)
+	// like shader arrays in pipeline states or UB binding indices for Graphics (Vertex==0, Pixel==1) that can overlap with Compute (Compute==0). 
+	// IMPORTANT: Adjusting these requires a full shader rebuild (ie modify the GUID in VulkanCommon.usf)
 	enum EStage
 	{
 		Vertex = 0,
diff --git a/Engine/Source/Runtime/VulkanRHI/Public/VulkanDynamicRHI.h b/Engine/Source/Runtime/VulkanRHI/Public/VulkanDynamicRHI.h
--- a/Engine/Source/Runtime/VulkanRHI/Public/VulkanDynamicRHI.h
+++ b/Engine/Source/Runtime/VulkanRHI/Public/VulkanDynamicRHI.h
@@ -110,6 +110,8 @@ public:
 	virtual FVertexDeclarationRHIRef RHICreateVertexDeclaration(const FVertexDeclarationElementList& Elements) final override;
 	virtual FPixelShaderRHIRef RHICreatePixelShader(TArrayView<const uint8> Code, const FSHAHash& Hash) final override;
 	virtual FVertexShaderRHIRef RHICreateVertexShader(TArrayView<const uint8> Code, const FSHAHash& Hash) final override;
+	virtual FMeshShaderRHIRef RHICreateMeshShader(TArrayView<const uint8> Code, const FSHAHash& Hash) final override;
+	virtual FAmplificationShaderRHIRef RHICreateAmplificationShader(TArrayView<const uint8> Code, const FSHAHash& Hash) final override;
 	virtual FGeometryShaderRHIRef RHICreateGeometryShader(TArrayView<const uint8> Code, const FSHAHash& Hash) final override;
 	virtual FComputeShaderRHIRef RHICreateComputeShader(TArrayView<const uint8> Code, const FSHAHash& Hash) final override;
 	virtual FGPUFenceRHIRef RHICreateGPUFence(const FName &Name) final override;
