VulkanRHI: Add support for VK_KHR_Maintenance5 to fix validation error on point list without point size.

#rnx

diff --git a/Engine/Source/Runtime/VulkanRHI/Private/VulkanExtensions.cpp b/Engine/Source/Runtime/VulkanRHI/Private/VulkanExtensions.cpp
--- a/Engine/Source/Runtime/VulkanRHI/Private/VulkanExtensions.cpp
+++ b/Engine/Source/Runtime/VulkanRHI/Private/VulkanExtensions.cpp
@@ -182,11 +182,51 @@ public:
 		ExtensionFlags.HasKHRMaintenance4 = bRequirementsPassed;
 	}
 
+	virtual void PreCreateDevice(VkDeviceCreateInfo& DeviceCreateInfo) override final
+	{
+		if (bRequirementsPassed)
+		{
+			AddToPNext(DeviceCreateInfo, Maintenance4Features);
+		}
+	}
+
 private:
 	VkPhysicalDeviceMaintenance4FeaturesKHR Maintenance4Features;
 };
 
 
+// ***** VK_KHR_maintenance5
+class FVulkanKHRMaintenance5Extension : public FVulkanDeviceExtension
+{
+public:
+
+	FVulkanKHRMaintenance5Extension(FVulkanDevice* InDevice)
+		: FVulkanDeviceExtension(InDevice, VK_KHR_MAINTENANCE_5_EXTENSION_NAME, VULKAN_EXTENSION_ENABLED)
+	{}
+
+	virtual void PrePhysicalDeviceFeatures(VkPhysicalDeviceFeatures2KHR& PhysicalDeviceFeatures2) override final
+	{
+		ZeroVulkanStruct(Maintenance5Features, VK_STRUCTURE_TYPE_PHYSICAL_DEVICE_MAINTENANCE_5_FEATURES_KHR);
+		AddToPNext(PhysicalDeviceFeatures2, Maintenance5Features);
+	}
+
+	virtual void PostPhysicalDeviceFeatures(FOptionalVulkanDeviceExtensions& ExtensionFlags) override final
+	{
+		bRequirementsPassed = (Maintenance5Features.maintenance5 == VK_TRUE);
+	}
+
+	virtual void PreCreateDevice(VkDeviceCreateInfo& DeviceCreateInfo) override final
+	{
+		if (bRequirementsPassed)
+		{
+			AddToPNext(DeviceCreateInfo, Maintenance5Features);
+		}
+	}
+
+private:
+	VkPhysicalDeviceMaintenance5FeaturesKHR Maintenance5Features;
+};
+
 
 // ***** VK_KHR_driver_properties
 class FVulkanKHRDriverPropertiesExtension : public FVulkanDeviceExtension
@@ -1659,6 +1699,7 @@ FVulkanDeviceExtensionArray FVulkanDeviceExtension::GetUESupportedDeviceExtensio
 
 	ADD_CUSTOM_EXTENSION(FVulkanKHRDriverPropertiesExtension);
 	ADD_CUSTOM_EXTENSION(FVulkanKHRMaintenance4Extension);
+	ADD_CUSTOM_EXTENSION(FVulkanKHRMaintenance5Extension);
 	ADD_CUSTOM_EXTENSION(FVulkanShaderAtomicInt64Extension);
 	ADD_CUSTOM_EXTENSION(FVulkanShaderImageAtomicInt64Extension);
 	ADD_CUSTOM_EXTENSION(FVulkanEXTScalarBlockLayoutExtension);
