VulkanRHI: locked buffer memory is sometimes directly cast to classes with 16byte alignment requirements (like FVector3f).  Fix crash on GPUs with alignment requirement below 16.

#rnx

diff --git a/Engine/Source/Runtime/VulkanRHI/Private/VulkanMemory.cpp b/Engine/Source/Runtime/VulkanRHI/Private/VulkanMemory.cpp
--- a/Engine/Source/Runtime/VulkanRHI/Private/VulkanMemory.cpp
+++ b/Engine/Source/Runtime/VulkanRHI/Private/VulkanMemory.cpp
@@ -2606,7 +2606,8 @@ namespace VulkanRHI
 		const bool bIsAccelerationStructureBuffer = VKHasAnyFlags(BufferUsageFlags, VK_BUFFER_USAGE_ACCELERATION_STRUCTURE_STORAGE_BIT_KHR);
 		const bool bIsUniformBuffer = VKHasAnyFlags(BufferUsageFlags, VK_BUFFER_USAGE_UNIFORM_BUFFER_BIT);
 
-		uint32 Alignment = 1;
+		// Buffers are sometimes directly cast into classes with 16byte alignment expectations (like FVector3f)
+		uint32 Alignment = 16u;
 
 		if (bIsTexelBuffer || bIsStorageBuffer)
 		{
