VulkanRHI: lower eviction threshold to avoid OOM on larger maps.

#rnx

diff --git a/Engine/Source/Runtime/VulkanRHI/Private/VulkanMemory.cpp b/Engine/Source/Runtime/VulkanRHI/Private/VulkanMemory.cpp
--- a/Engine/Source/Runtime/VulkanRHI/Private/VulkanMemory.cpp
+++ b/Engine/Source/Runtime/VulkanRHI/Private/VulkanMemory.cpp
@@ -256,7 +256,7 @@ static FAutoConsoleVariableRef CVarVulkanDefragOnce(
 
 
 
-static float GVulkanEvictionLimitPercentage = 80.f;
+static float GVulkanEvictionLimitPercentage = 70.f;
 static FAutoConsoleVariableRef CVarVulkanEvictionLimitPercentage(
 	TEXT("r.Vulkan.EvictionLimitPercentage"),
 	GVulkanEvictionLimitPercentage,
@@ -265,7 +265,7 @@ static FAutoConsoleVariableRef CVarVulkanEvictionLimitPercentage(
 );
 
 
-static float GVulkanEvictionLimitPercentageReenableLimit = 75.f;
+static float GVulkanEvictionLimitPercentageReenableLimit = 60.f;
 static FAutoConsoleVariableRef CVarVulkanEvictionLimitPercentageReenableLimit(
 	TEXT("r.Vulkan.EvictionLimitPercentageRenableLimit"),
 	GVulkanEvictionLimitPercentageReenableLimit,
