VulkanRHI: Add pipeline stage bits for mesh shaders

#rnx

diff --git a/Engine/Source/Runtime/VulkanRHI/Private/VulkanDevice.h b/Engine/Source/Runtime/VulkanRHI/Private/VulkanDevice.h
--- a/Engine/Source/Runtime/VulkanRHI/Private/VulkanDevice.h
+++ b/Engine/Source/Runtime/VulkanRHI/Private/VulkanDevice.h
@@ -57,6 +57,7 @@ struct FOptionalVulkanDeviceExtensions
 			uint64 HasEXTDescriptorBuffer : 1;
 			uint64 HasEXTDeviceFault : 1;
 			uint64 HasEXTLoadStoreOpNone : 1;
+			uint64 HasEXTMeshShader : 1;
 
 			// Vendor specific
 			uint64 HasAMDBufferMarker : 1;
diff --git a/Engine/Source/Runtime/VulkanRHI/Private/VulkanExtensions.cpp b/Engine/Source/Runtime/VulkanRHI/Private/VulkanExtensions.cpp
--- a/Engine/Source/Runtime/VulkanRHI/Private/VulkanExtensions.cpp
+++ b/Engine/Source/Runtime/VulkanRHI/Private/VulkanExtensions.cpp
@@ -1086,6 +1086,7 @@ public:
 	virtual void PostPhysicalDeviceFeatures(FOptionalVulkanDeviceExtensions& ExtensionFlags) override final
 	{
 		bRequirementsPassed = (MeshShaderFeatures.meshShader == VK_TRUE) && (MeshShaderFeatures.multiviewMeshShader == VK_TRUE);
+		ExtensionFlags.HasEXTMeshShader = bRequirementsPassed;
 	}
 
 	virtual void PrePhysicalDeviceProperties(VkPhysicalDeviceProperties2KHR& PhysicalDeviceProperties2) override final
diff --git a/Engine/Source/Runtime/VulkanRHI/Private/VulkanQueue.cpp b/Engine/Source/Runtime/VulkanRHI/Private/VulkanQueue.cpp
--- a/Engine/Source/Runtime/VulkanRHI/Private/VulkanQueue.cpp
+++ b/Engine/Source/Runtime/VulkanRHI/Private/VulkanQueue.cpp
@@ -167,6 +167,10 @@ void FVulkanQueue::FillSupportedStageBits()
 		{
 			SupportedStages |= VK_PIPELINE_STAGE_FRAGMENT_DENSITY_PROCESS_BIT_EXT;
 		}
+		if (Device->GetOptionalExtensions().HasEXTMeshShader)
+		{
+			SupportedStages |= VK_PIPELINE_STAGE_TASK_SHADER_BIT_EXT | VK_PIPELINE_STAGE_MESH_SHADER_BIT_EXT;
+		}
 	}
 
 	if (VKHasAnyFlags(QueueProps.queueFlags, VK_QUEUE_COMPUTE_BIT))
diff --git a/Engine/Source/Runtime/VulkanRHI/Private/VulkanRHI.cpp b/Engine/Source/Runtime/VulkanRHI/Private/VulkanRHI.cpp
--- a/Engine/Source/Runtime/VulkanRHI/Private/VulkanRHI.cpp
+++ b/Engine/Source/Runtime/VulkanRHI/Private/VulkanRHI.cpp
@@ -998,12 +998,6 @@ void FVulkanDynamicRHI::InitInstance()
 		GRHIMaxDispatchThreadGroupsPerDimension.Y = FMath::Min<uint32>(Limits.maxComputeWorkGroupCount[1], 0x7fffffff);
 		GRHIMaxDispatchThreadGroupsPerDimension.Z = FMath::Min<uint32>(Limits.maxComputeWorkGroupCount[2], 0x7fffffff);
 
-#if PLATFORM_SUPPORTS_MESH_SHADERS
-		// If mesh shaders are enabled in DDPI (currently SM6), then the profile check will ensure it's supported
-		GRHIGlobals.SupportsMeshShadersTier0 = RHISupportsMeshShadersTier0(GMaxRHIShaderPlatform);
-		GRHIGlobals.SupportsMeshShadersTier1 = RHISupportsMeshShadersTier1(GMaxRHIShaderPlatform);
-#endif
-
 		FVulkanPlatform::SetupFeatureLevels(GRHIGlobals.ShaderPlatformForFeatureLevel);
 
 		GRHIRequiresRenderTargetForPixelShaderUAVs = true;
@@ -1021,6 +1015,18 @@ void FVulkanDynamicRHI::InitInstance()
 			VulkanDeviceShaderStageBits |= VK_SHADER_STAGE_GEOMETRY_BIT;
 		}
 
+#if PLATFORM_SUPPORTS_MESH_SHADERS
+		// If mesh shaders are enabled in DDPI (currently SM6), then the profile check will ensure it's supported
+		if (Device->GetOptionalExtensions().HasEXTMeshShader)
+		{
+			GRHIGlobals.SupportsMeshShadersTier0 = RHISupportsMeshShadersTier0(GMaxRHIShaderPlatform);
+			GRHIGlobals.SupportsMeshShadersTier1 = RHISupportsMeshShadersTier1(GMaxRHIShaderPlatform);
+
+			GVulkanDevicePipelineStageBits |= VK_PIPELINE_STAGE_TASK_SHADER_BIT_EXT | VK_PIPELINE_STAGE_MESH_SHADER_BIT_EXT;
+			VulkanDeviceShaderStageBits |= VK_SHADER_STAGE_TASK_BIT_EXT | VK_SHADER_STAGE_MESH_BIT_EXT;
+		}
+#endif
+
 		const VkShaderStageFlags RequiredSubgroupShaderStageFlags = FVulkanPlatform::RequiredWaveOpsShaderStageFlags(VulkanDeviceShaderStageBits);
 		
 		// Check for wave ops support (only filled on platforms creating Vulkan 1.1 or greater instances)
