VulkanRHI: fix crash when using -noraytracing on some vendors

#rnx

diff --git a/Engine/Source/Runtime/VulkanRHI/Private/VulkanQueue.cpp b/Engine/Source/Runtime/VulkanRHI/Private/VulkanQueue.cpp
--- a/Engine/Source/Runtime/VulkanRHI/Private/VulkanQueue.cpp
+++ b/Engine/Source/Runtime/VulkanRHI/Private/VulkanQueue.cpp
@@ -178,9 +178,17 @@ void FVulkanQueue::FillSupportedStageBits()
 		SupportedStages |=
 			VK_PIPELINE_STAGE_COMPUTE_SHADER_BIT |
 			VK_PIPELINE_STAGE_DRAW_INDIRECT_BIT |
-			VK_PIPELINE_STAGE_TRANSFER_BIT |
-			VK_PIPELINE_STAGE_ACCELERATION_STRUCTURE_BUILD_BIT_KHR |
-			VK_PIPELINE_STAGE_RAY_TRACING_SHADER_BIT_KHR;
+			VK_PIPELINE_STAGE_TRANSFER_BIT;
+
+		if (Device->GetOptionalExtensions().HasAccelerationStructure)
+		{
+			SupportedStages |= VK_PIPELINE_STAGE_ACCELERATION_STRUCTURE_BUILD_BIT_KHR;
+		}
+
+		if (Device->GetOptionalExtensions().HasRayTracingPipeline)
+		{
+			SupportedStages |= VK_PIPELINE_STAGE_RAY_TRACING_SHADER_BIT_KHR;
+		}
 	}
 
 	if (VKHasAnyFlags(QueueProps.queueFlags, VK_QUEUE_TRANSFER_BIT))
