VulkanRHI: Fix gap in layout tracking (would result in bad depth reads on some Lumen passes in Vulkan with RT enabled)

#rnx

diff --git a/Engine/Source/Runtime/VulkanRHI/Private/VulkanBarriers.h b/Engine/Source/Runtime/VulkanRHI/Private/VulkanBarriers.h
--- a/Engine/Source/Runtime/VulkanRHI/Private/VulkanBarriers.h
+++ b/Engine/Source/Runtime/VulkanRHI/Private/VulkanBarriers.h
@@ -204,9 +204,29 @@ public:
 		}
 	}
 
+	static bool IsPartialResource(const FVulkanTexture& VulkanTexture, const VkImageSubresourceRange& InSubresourceRange)
+	{
+		return (VulkanTexture.GetFullAspectMask() != InSubresourceRange.aspectMask) ||
+			(InSubresourceRange.baseArrayLayer != 0) ||
+			(InSubresourceRange.baseMipLevel != 0) ||
+			((InSubresourceRange.levelCount != VK_REMAINING_MIP_LEVELS) && (InSubresourceRange.levelCount != VulkanTexture.GetNumMips())) ||
+			((InSubresourceRange.layerCount != VK_REMAINING_ARRAY_LAYERS) && (InSubresourceRange.layerCount != VulkanTexture.GetNumberOfArrayLevels()));
+	}
+
 	VULKANRHI_API void SetLayout(const FVulkanTexture& VulkanTexture, const VkImageSubresourceRange& InSubresourceRange, VkImageLayout InLayout)
 	{
 		FVulkanImageLayout* Layout = Layouts.Find(VulkanTexture.Image);
+
+		// If we're not going to overwrite the entire resource, start from its last known state on the queue
+		if (!Layout && IsPartialResource(VulkanTexture, InSubresourceRange) && Fallback)
+		{
+			const FVulkanImageLayout* FallbackLayout = Fallback->GetFullLayout(VulkanTexture, false);
+			if (FallbackLayout)
+			{
+				Layout = &Layouts.Add(VulkanTexture.Image, *FallbackLayout);
+			}
+		}
+
 		if (Layout)
 		{
 			Layout->Set(InLayout, InSubresourceRange);
