Fix Vulkan debug names and markers

diff --git a/Engine/Source/Runtime/VulkanRHI/Private/VulkanCommandWrappers.h b/Engine/Source/Runtime/VulkanRHI/Private/VulkanCommandWrappers.h
--- a/Engine/Source/Runtime/VulkanRHI/Private/VulkanCommandWrappers.h
+++ b/Engine/Source/Runtime/VulkanRHI/Private/VulkanCommandWrappers.h
@@ -250,6 +250,7 @@ struct FWrapLayer
 	static void GetDeferredOperationMaxConcurrencyKHR(VkResult Result, VkDevice Device, VkDeferredOperationKHR DeferredOperation) VULKAN_LAYER_BODY
 	static void GetDeferredOperationResultKHR(VkResult Result, VkDevice Device, VkDeferredOperationKHR DeferredOperation) VULKAN_LAYER_BODY
 	static void GetDeviceFaultInfoEXT(VkResult Result, VkDevice Device, VkDeviceFaultCountsEXT* FaultCounts, VkDeviceFaultInfoEXT* FaultInfo) VULKAN_LAYER_BODY
+	static void GetPhysicalDeviceToolPropertiesEXT(VkResult Result, VkPhysicalDevice PhysicalDevice, uint32_t* ToolCount, VkPhysicalDeviceToolProperties* ToolProperties) VULKAN_LAYER_BODY
 	static void CreateSamplerYcbcrConversion(VkResult Result, VkDevice Device, const VkSamplerYcbcrConversionCreateInfo* CreateInfo, const VkAllocationCallbacks* Allocator, VkSamplerYcbcrConversion* YcbcrConversion) VULKAN_LAYER_BODY
 	static void DestroySamplerYcbcrConversion(VkResult Result, VkDevice Device, VkSamplerYcbcrConversion YcbcrConversion, const VkAllocationCallbacks* Allocator) VULKAN_LAYER_BODY
 #if PLATFORM_ANDROID
@@ -1738,6 +1739,14 @@ namespace VulkanRHI
 		return Result;
 	}
 
+	static FORCEINLINE_DEBUGGABLE VkResult vkGetPhysicalDeviceToolPropertiesEXT(VkPhysicalDevice PhysicalDevice, uint32_t* ToolCount, VkPhysicalDeviceToolPropertiesEXT* ToolProperties)
+	{
+		FWrapLayer::GetPhysicalDeviceToolPropertiesEXT(VK_RESULT_MAX_ENUM, PhysicalDevice, ToolCount, ToolProperties);
+		VkResult Result = VULKANAPINAMESPACE::vkGetPhysicalDeviceToolPropertiesEXT(PhysicalDevice, ToolCount, ToolProperties);
+		FWrapLayer::GetPhysicalDeviceToolPropertiesEXT(Result, PhysicalDevice, ToolCount, ToolProperties);
+		return Result;
+	}
+
 	static FORCEINLINE_DEBUGGABLE VkResult vkCreateSamplerYcbcrConversion(VkDevice Device, const VkSamplerYcbcrConversionCreateInfo* CreateInfo, const VkAllocationCallbacks* Allocator, VkSamplerYcbcrConversion* YcbcrConversion)
 	{
 		FWrapLayer::CreateSamplerYcbcrConversion(VK_RESULT_MAX_ENUM, Device, CreateInfo, Allocator, YcbcrConversion);
diff --git a/Engine/Source/Runtime/VulkanRHI/Private/VulkanDevice.cpp b/Engine/Source/Runtime/VulkanRHI/Private/VulkanDevice.cpp
--- a/Engine/Source/Runtime/VulkanRHI/Private/VulkanDevice.cpp
+++ b/Engine/Source/Runtime/VulkanRHI/Private/VulkanDevice.cpp
@@ -476,33 +476,85 @@ void FVulkanDevice::CreateDevice(TArray<const ANSICHAR*>& DeviceLayers, FVulkanD
 	SetupDrawMarkers();
 }
 
+#if VULKAN_HAS_DEBUGGING_ENABLED
+TAutoConsoleVariable<int32> CVarVulkanDebugMarkers(
+	TEXT("r.Vulkan.DebugMarkers"),
+	4,
+	TEXT("0 to disable all debug markers\n")
+	TEXT("1 to enable debug names for resources\n")
+	TEXT("2 to enable debug labels for commands\n")
+	TEXT("3 to enable debug resource names command labels\n")
+	TEXT("4 to automatically enable markers depending on tool detection (default)\n"),
+	ECVF_ReadOnly | ECVF_RenderThreadSafe
+);
+#endif // VULKAN_HAS_DEBUGGING_ENABLED
+
 void FVulkanDevice::SetupDrawMarkers()
 {
 #if VULKAN_ENABLE_DRAW_MARKERS
+	bool bTraceToolFound = GRenderDocFound;
 	if (RHI->SupportsDebugUtilsExt())
 	{
-		// HOTFIX for UE-218250: Disable vulkan draw markers to get around crash/performance issues
-		if (FParse::Param(FCommandLine::Get(), TEXT("forcevulkanddrawmarkers")))
+		bool bUseLabel = false;
+		bool bUseName = false;
+
+		if (OptionalDeviceExtensions.HasEXTToolingInfo)
+		{
+			uint32_t ToolCount = 0;
+			VulkanRHI::vkGetPhysicalDeviceToolPropertiesEXT(Gpu, &ToolCount, nullptr);
+			TArray<VkPhysicalDeviceToolPropertiesEXT> ToolProperties;
+			ToolProperties.AddZeroed(ToolCount);
+			VulkanRHI::vkGetPhysicalDeviceToolPropertiesEXT(Gpu, &ToolCount, ToolProperties.GetData());
+			for (VkPhysicalDeviceToolPropertiesEXT const& Tool : ToolProperties)
+			{
+				if (Tool.purposes & VK_TOOL_PURPOSE_DEBUG_MARKERS_BIT_EXT)
+				{
+					bUseName = true;
+					if (Tool.purposes & VK_TOOL_PURPOSE_TRACING_BIT_EXT)
+					{
+						bUseLabel = true;
+					}
+				}
+				if (Tool.purposes & VK_TOOL_PURPOSE_TRACING_BIT_EXT)
+				{
+					bTraceToolFound = true;
+				}
+				UE_LOG(LogVulkanRHI, Display, TEXT("Tool \"%s\" version %s PurposeFlags=0x%x"), ANSI_TO_TCHAR(Tool.name), ANSI_TO_TCHAR(Tool.version), Tool.purposes);
+			}
+		}
+
+		const int VulkanDebugMarkers = FParse::Param(FCommandLine::Get(), TEXT("forcevulkandrawmarkers")) ? 3 : CVarVulkanDebugMarkers.GetValueOnRenderThread();
+		if (VulkanDebugMarkers != 4)
+		{
+			bUseLabel = VulkanDebugMarkers & 2;
+			bUseName = VulkanDebugMarkers & 1;
+		}
+
+		if (bUseLabel)
 		{
 			DebugMarkers.CmdBeginDebugLabel = (PFN_vkCmdBeginDebugUtilsLabelEXT)(void*)VulkanRHI::vkGetInstanceProcAddr(RHI->GetInstance(), "vkCmdBeginDebugUtilsLabelEXT");
 			DebugMarkers.CmdEndDebugLabel = (PFN_vkCmdEndDebugUtilsLabelEXT)(void*)VulkanRHI::vkGetInstanceProcAddr(RHI->GetInstance(), "vkCmdEndDebugUtilsLabelEXT");
-			DebugMarkers.SetDebugName = (PFN_vkSetDebugUtilsObjectNameEXT)(void*)VulkanRHI::vkGetInstanceProcAddr(RHI->GetInstance(), "vkSetDebugUtilsObjectNameEXT");
 		}
-
-		if (DebugMarkers.CmdBeginDebugLabel && DebugMarkers.CmdEndDebugLabel && DebugMarkers.SetDebugName)
+		if (bUseName)
 		{
-			bDebugMarkersFound = true;
+			DebugMarkers.SetDebugName = (PFN_vkSetDebugUtilsObjectNameEXT)(void*)VulkanRHI::vkGetInstanceProcAddr(RHI->GetInstance(), "vkSetDebugUtilsObjectNameEXT");
 		}
+
+		UE_LOG(LogVulkanRHI, Display, TEXT("Vulkan debug markers support: resource names %s, begin/end labels %s."), bUseName ? TEXT("enabled") : TEXT("disabled"), bUseLabel ? TEXT("enabled") : TEXT("disabled"));
 	}
 
 #if VULKAN_HAS_DEBUGGING_ENABLED
-	if (bDebugMarkersFound && GRenderDocFound)
+	if (DebugMarkers.CmdBeginDebugLabel && DebugMarkers.CmdEndDebugLabel && bTraceToolFound)
 	{
 		// We're running under RenderDoc or other trace tool, so enable capturing mode
 		FDynamicRHI::EnableIdealGPUCaptureOptions(true);
 	}
-#endif
-#endif
+	else if(bTraceToolFound)
+	{
+		UE_LOG(LogVulkanRHI, Warning, TEXT("Vulkan API trace tool detected but not running in ideal GPU capture mode."));
+	}
+#endif // VULKAN_HAS_DEBUGGING_ENABLED
+#endif // VULKAN_ENABLE_DRAW_MARKERS
 
 #if VULKAN_ENABLE_DUMP_LAYER
 	FDynamicRHI::EnableIdealGPUCaptureOptions(true);
diff --git a/Engine/Source/Runtime/VulkanRHI/Private/VulkanDevice.h b/Engine/Source/Runtime/VulkanRHI/Private/VulkanDevice.h
--- a/Engine/Source/Runtime/VulkanRHI/Private/VulkanDevice.h
+++ b/Engine/Source/Runtime/VulkanRHI/Private/VulkanDevice.h
@@ -23,9 +23,7 @@ class FVulkanOcclusionQueryPool;
 class FVulkanRenderPassManager;
 class FVulkanRayTracingCompactionRequestHandler;
 
-// HOTFIX for UE-218250: Disable vulkan debug names to get around crash/performance issues
-#define VULKAN_USE_DEBUG_NAMES 0
-
+#define VULKAN_USE_DEBUG_NAMES 1
 #if VULKAN_USE_DEBUG_NAMES
 #define VULKAN_SET_DEBUG_NAME(Device, Type, Handle, Format, ...) (Device).VulkanSetObjectName(Type, (uint64)Handle, *FString::Printf(Format, __VA_ARGS__))
 #else
@@ -58,6 +56,7 @@ struct FOptionalVulkanDeviceExtensions
 			uint64 HasEXTDeviceFault : 1;
 			uint64 HasEXTLoadStoreOpNone : 1;
 			uint64 HasEXTMeshShader : 1;
+			uint64 HasEXTToolingInfo : 1;
 
 			// Vendor specific
 			uint64 HasAMDBufferMarker : 1;
@@ -621,7 +620,6 @@ private:
 	TArray<FVulkanCommandListContext*> CommandContexts;
 
 	FVulkanDynamicRHI* RHI = nullptr;
-	bool bDebugMarkersFound = false;
 	
 	static TArray<const ANSICHAR*> SetupDeviceLayers(VkPhysicalDevice Gpu, FVulkanDeviceExtensionArray& UEExtensions);
 
diff --git a/Engine/Source/Runtime/VulkanRHI/Private/VulkanExtensions.cpp b/Engine/Source/Runtime/VulkanRHI/Private/VulkanExtensions.cpp
--- a/Engine/Source/Runtime/VulkanRHI/Private/VulkanExtensions.cpp
+++ b/Engine/Source/Runtime/VulkanRHI/Private/VulkanExtensions.cpp
@@ -1691,6 +1691,7 @@ FVulkanDeviceExtensionArray FVulkanDeviceExtension::GetUESupportedDeviceExtensio
 	ADD_SIMPLE_EXTENSION(VK_EXT_VALIDATION_CACHE_EXTENSION_NAME,         VULKAN_SUPPORTS_VALIDATION_CACHE,     VULKAN_EXTENSION_NOT_PROMOTED, DEVICE_EXT_FLAG_SETTER(HasEXTValidationCache));
 	ADD_SIMPLE_EXTENSION(VK_EXT_LOAD_STORE_OP_NONE_EXTENSION_NAME,       VULKAN_EXTENSION_ENABLED,             VULKAN_EXTENSION_NOT_PROMOTED, DEVICE_EXT_FLAG_SETTER(HasEXTLoadStoreOpNone));
 	ADD_SIMPLE_EXTENSION(VK_QCOM_RENDER_PASS_SHADER_RESOLVE_EXTENSION_NAME, VULKAN_SUPPORTS_QCOM_RENDERPASS_SHADER_RESOLVE, VULKAN_EXTENSION_NOT_PROMOTED, DEVICE_EXT_FLAG_SETTER(HasQcomRenderPassShaderResolve));
+	ADD_SIMPLE_EXTENSION(VK_EXT_TOOLING_INFO_EXTENSION_NAME, VULKAN_HAS_DEBUGGING_ENABLED || VULKAN_ENABLE_DRAW_MARKERS, VK_API_VERSION_1_3, DEVICE_EXT_FLAG_SETTER(HasEXTToolingInfo));
 
 	// Externally activated extensions (supported by the engine, but enabled externally by plugin or other) :
 
diff --git a/Engine/Source/Runtime/VulkanRHI/Private/VulkanLoader.h b/Engine/Source/Runtime/VulkanRHI/Private/VulkanLoader.h
--- a/Engine/Source/Runtime/VulkanRHI/Private/VulkanLoader.h
+++ b/Engine/Source/Runtime/VulkanRHI/Private/VulkanLoader.h
@@ -231,7 +231,8 @@
 	EnumMacro(PFN_vkCmdSetCheckpointNV, vkCmdSetCheckpointNV) \
 	EnumMacro(PFN_vkGetQueueCheckpointDataNV, vkGetQueueCheckpointDataNV) \
 	EnumMacro(PFN_vkGetPhysicalDeviceFragmentShadingRatesKHR, vkGetPhysicalDeviceFragmentShadingRatesKHR) \
-	EnumMacro(PFN_vkGetDeviceFaultInfoEXT, vkGetDeviceFaultInfoEXT)
+	EnumMacro(PFN_vkGetDeviceFaultInfoEXT, vkGetDeviceFaultInfoEXT) \
+	EnumMacro(PFN_vkGetPhysicalDeviceToolPropertiesEXT, vkGetPhysicalDeviceToolPropertiesEXT)
 
 
 // List of all Vulkan entry points
