Fix compile error related to Vulkan debug markers

diff --git a/Engine/Source/Runtime/VulkanRHI/Private/VulkanDevice.cpp b/Engine/Source/Runtime/VulkanRHI/Private/VulkanDevice.cpp
--- a/Engine/Source/Runtime/VulkanRHI/Private/VulkanDevice.cpp
+++ b/Engine/Source/Runtime/VulkanRHI/Private/VulkanDevice.cpp
@@ -476,7 +476,7 @@ void FVulkanDevice::CreateDevice(TArray<const ANSICHAR*>& DeviceLayers, FVulkanD
 	SetupDrawMarkers();
 }
 
-#if VULKAN_HAS_DEBUGGING_ENABLED
+#if VULKAN_ENABLE_DRAW_MARKERS
 TAutoConsoleVariable<int32> CVarVulkanDebugMarkers(
 	TEXT("r.Vulkan.DebugMarkers"),
 	4,
@@ -487,12 +487,16 @@ TAutoConsoleVariable<int32> CVarVulkanDebugMarkers(
 	TEXT("4 to automatically enable markers depending on tool detection (default)\n"),
 	ECVF_ReadOnly | ECVF_RenderThreadSafe
 );
-#endif // VULKAN_HAS_DEBUGGING_ENABLED
+#endif // VULKAN_ENABLE_DRAW_MARKERS
 
 void FVulkanDevice::SetupDrawMarkers()
 {
 #if VULKAN_ENABLE_DRAW_MARKERS
-	bool bTraceToolFound = GRenderDocFound;
+	bool bTraceToolFound = false;
+#if VULKAN_HAS_DEBUGGING_ENABLED
+	bTraceToolFound = GRenderDocFound;
+#endif // VULKAN_HAS_DEBUGGING_ENABLED
+
 	if (RHI->SupportsDebugUtilsExt())
 	{
 		bool bUseLabel = false;
