VulkanRHI: Create a Vulkan AMD compatibility mode to circumvent issues present on different combinations of drivers and hardware.

#rnx

diff --git a/Engine/Source/Runtime/VulkanRHI/Private/Linux/VulkanLinuxPlatform.cpp b/Engine/Source/Runtime/VulkanRHI/Private/Linux/VulkanLinuxPlatform.cpp
--- a/Engine/Source/Runtime/VulkanRHI/Private/Linux/VulkanLinuxPlatform.cpp
+++ b/Engine/Source/Runtime/VulkanRHI/Private/Linux/VulkanLinuxPlatform.cpp
@@ -15,6 +15,7 @@
 ENUM_VK_ENTRYPOINTS_ALL(DEFINE_VK_ENTRYPOINTS)
 
 static bool GRenderOffScreen = false;
+extern int32 GVulkanAMDCompatibilityMode;
 
 void* FVulkanLinuxPlatform::VulkanLib = nullptr;
 bool FVulkanLinuxPlatform::bAttemptedLoad = false;
@@ -51,6 +52,17 @@ bool FVulkanLinuxPlatform::LoadVulkanLibrary()
 	}
 	bAttemptedLoad = true;
 
+	// Set regardless of GPU being used, it will simply get ignored on other vendors
+	if (GVulkanAMDCompatibilityMode && FPlatformMisc::GetEnvironmentVariable(TEXT("RADV_DEBUG")).IsEmpty())
+	{
+		// Force compiler backend to llvm for better compatibility with our ray tracing pipeline shaders at time of release 5.5 (with Mesa 24.0.9)
+		FPlatformMisc::SetEnvironmentVar(TEXT("RADV_DEBUG"), TEXT("llvm"));
+	}
+	else
+	{
+		UE_LOG(LogVulkanRHI, Display, TEXT("Found existing RADV_DEBUG, it will not be overwritten."));
+	}
+
 #if VULKAN_HAS_DEBUGGING_ENABLED
 	const FString VulkanSDK = FPlatformMisc::GetEnvironmentVariable(TEXT("VULKAN_SDK"));
 	UE_LOG(LogVulkanRHI, Display, TEXT("Found VULKAN_SDK=%s"), *VulkanSDK);
diff --git a/Engine/Source/Runtime/VulkanRHI/Private/VulkanRHI.cpp b/Engine/Source/Runtime/VulkanRHI/Private/VulkanRHI.cpp
--- a/Engine/Source/Runtime/VulkanRHI/Private/VulkanRHI.cpp
+++ b/Engine/Source/Runtime/VulkanRHI/Private/VulkanRHI.cpp
@@ -113,6 +113,13 @@ static FAutoConsoleVariableRef GCVarVulkanTempBlockSize(
 	ECVF_ReadOnly
 );
 
+int32 GVulkanAMDCompatibilityMode = 1;
+static FAutoConsoleVariableRef GCVarVulkanAMDCompatibilityMode(
+	TEXT("r.Vulkan.AMDCompatibilityMode"),
+	GVulkanAMDCompatibilityMode,
+	TEXT("Used to tweak enabled Vulkan feature set in order to ensure wider compatibility with all AMD GPUs on all platforms. (default:1)"),
+	ECVF_ReadOnly
+);
 
 
 extern TAutoConsoleVariable<int32> GVulkanRayTracingCVar;
@@ -936,10 +943,16 @@ void FVulkanDynamicRHI::InitInstance()
 		GRHISupportsMSAAShaderResolve = Device->GetOptionalExtensions().HasQcomRenderPassShaderResolve ? true : false;
 		GRHISupportsRayTracing = RHISupportsRayTracing(GMaxRHIShaderPlatform) && Device->GetOptionalExtensions().HasRaytracingExtensions();
 
+		// Use this compatibility mode avoid known issues at launch time with latest drivers at the time of release 5.5.  This will:
+		// - disable inline ray tracing and use ray tracing pipelines everywhere (instead of a mix of both)
+		// - disable mesh shaders until issues can be resolved (holes in Nanite meshes)
+		// - force llvm compiler backend on Linux (see VulkanLinuxPlatform.cpp) to circumvent raytracing pipeline compilation crash
+		const bool bUseAMDCompatibilityMode = GVulkanAMDCompatibilityMode && (Device->GetVendorId() == EGpuVendorId::Amd);
+
 		if (GRHISupportsRayTracing)
 		{
 			GRHISupportsRayTracingShaders = RHISupportsRayTracingShaders(GMaxRHIShaderPlatform) && Device->GetOptionalExtensions().HasRayTracingPipeline;
-			GRHISupportsInlineRayTracing = RHISupportsInlineRayTracing(GMaxRHIShaderPlatform) && Device->GetOptionalExtensions().HasRayQuery;
+			GRHISupportsInlineRayTracing = !bUseAMDCompatibilityMode && RHISupportsInlineRayTracing(GMaxRHIShaderPlatform) && Device->GetOptionalExtensions().HasRayQuery;
 
 			// Inline RayTracing SBT is needed if raytracing position fetch isn't available
 			GRHIGlobals.RayTracing.RequiresInlineRayTracingSBT = !VULKAN_SUPPORTS_RAY_TRACING_POSITION_FETCH;
@@ -1015,7 +1028,7 @@ void FVulkanDynamicRHI::InitInstance()
 
 #if PLATFORM_SUPPORTS_MESH_SHADERS
 		// If mesh shaders are enabled in DDPI (currently SM6), then the profile check will ensure it's supported
-		if (Device->GetOptionalExtensions().HasEXTMeshShader)
+		if (!bUseAMDCompatibilityMode && Device->GetOptionalExtensions().HasEXTMeshShader)
 		{
 			GRHIGlobals.SupportsMeshShadersTier0 = RHISupportsMeshShadersTier0(GMaxRHIShaderPlatform);
 			GRHIGlobals.SupportsMeshShadersTier1 = RHISupportsMeshShadersTier1(GMaxRHIShaderPlatform);
