VulkanRHI: Use a temporary alloc for RayGen handle (avoid staging buffer and copy).

#rnx

diff --git a/Engine/Source/Runtime/VulkanRHI/Private/VulkanMemory.cpp b/Engine/Source/Runtime/VulkanRHI/Private/VulkanMemory.cpp
--- a/Engine/Source/Runtime/VulkanRHI/Private/VulkanMemory.cpp
+++ b/Engine/Source/Runtime/VulkanRHI/Private/VulkanMemory.cpp
@@ -5041,7 +5041,7 @@ namespace VulkanRHI
 	uint8* FTempBlockAllocator::Alloc(uint32 InSize, uint32 InAlignment, FVulkanCmdBuffer* CmdBuffer, FVulkanAllocation& OutAllocation, VkDescriptorAddressInfoEXT* OutDescriptorAddressInfo)
 	{
 		const uint32 AlignedSize = Align(InSize, BlockAlignment);
-		checkfSlow(BlockAlignment < BlockSize, TEXT("Requested size of %d (%d aligned) is too large for block size of %d"), InSize, AlignedSize, BlockSize);
+		checkfSlow(AlignedSize < BlockSize, TEXT("Requested size of %d (%d aligned) is too large for block size of %d"), InSize, AlignedSize, BlockSize);
 
 		FInternalAlloc Alloc = InternalAlloc(InSize, CmdBuffer);
 
@@ -5060,6 +5060,19 @@ namespace VulkanRHI
 		return Alloc.Block->MappedPointer + Alloc.Offset;
 	}
 
+	uint8* FTempBlockAllocator::Alloc(uint32 InSize, FVulkanCmdBuffer* CmdBuffer, VkStridedDeviceAddressRegionKHR* OutStridedDeviceAddressRegion)
+	{
+		const uint32 AlignedSize = Align(InSize, BlockAlignment);
+		checkfSlow(AlignedSize < BlockSize, TEXT("Requested size of %d (%d aligned) is too large for block size of %d"), InSize, AlignedSize, BlockSize);
+		checkSlow(OutStridedDeviceAddressRegion);
+
+		FInternalAlloc Alloc = InternalAlloc(InSize, CmdBuffer);
+		OutStridedDeviceAddressRegion->deviceAddress = Alloc.Block->BufferAddress + Alloc.Offset;
+		OutStridedDeviceAddressRegion->size = InSize;
+		OutStridedDeviceAddressRegion->stride = OutStridedDeviceAddressRegion->size;
+		return Alloc.Block->MappedPointer + Alloc.Offset;
+	}
+
 	void FTempBlockAllocator::UpdateBlocks()
 	{
 		FWriteScopeLock ScopeLock(RWLock);
diff --git a/Engine/Source/Runtime/VulkanRHI/Private/VulkanMemory.h b/Engine/Source/Runtime/VulkanRHI/Private/VulkanMemory.h
--- a/Engine/Source/Runtime/VulkanRHI/Private/VulkanMemory.h
+++ b/Engine/Source/Runtime/VulkanRHI/Private/VulkanMemory.h
@@ -1145,6 +1145,7 @@ namespace VulkanRHI
 
 		uint8* Alloc(uint32 InSize, FVulkanCmdBuffer* CmdBuffer, VkDescriptorBufferBindingInfoEXT& OutBindingInfo, VkDeviceSize& OutOffset);
 		uint8* Alloc(uint32 InSize, uint32 InAlignment, FVulkanCmdBuffer* CmdBuffer, FVulkanAllocation& OutAllocation, VkDescriptorAddressInfoEXT* OutDescriptorAddressInfo = nullptr);
+		uint8* Alloc(uint32 InSize, FVulkanCmdBuffer* CmdBuffer, VkStridedDeviceAddressRegionKHR* OutStridedDeviceAddressRegion);
 
 		void UpdateBlocks();
 
diff --git a/Engine/Source/Runtime/VulkanRHI/Private/VulkanRHI.cpp b/Engine/Source/Runtime/VulkanRHI/Private/VulkanRHI.cpp
--- a/Engine/Source/Runtime/VulkanRHI/Private/VulkanRHI.cpp
+++ b/Engine/Source/Runtime/VulkanRHI/Private/VulkanRHI.cpp
@@ -470,11 +470,16 @@ FVulkanCommandListContext::FVulkanCommandListContext(FVulkanDynamicRHI* InRHI, F
 
 	// Currently used for UB and copies
 	{
-		const uint32 BlockAlignment = FMath::Max<uint32>(Device->GetLimits().minUniformBufferOffsetAlignment, 16u);
-		const VkBufferUsageFlags BufferUsageFlags = 
+		uint32 BlockAlignment = FMath::Max<uint32>(Device->GetLimits().minUniformBufferOffsetAlignment, 16u);
+		VkBufferUsageFlags BufferUsageFlags =
 			(Device->GetOptionalExtensions().HasBufferDeviceAddress ? VK_BUFFER_USAGE_SHADER_DEVICE_ADDRESS_BIT : 0) |
 			VK_BUFFER_USAGE_TRANSFER_SRC_BIT |
 			VK_BUFFER_USAGE_UNIFORM_BUFFER_BIT;
+		if (Device->GetOptionalExtensions().HasRayTracingPipeline)
+		{
+			BufferUsageFlags |= VK_BUFFER_USAGE_SHADER_BINDING_TABLE_BIT_KHR;
+			BlockAlignment = FMath::Max<uint32>(Device->GetOptionalExtensionProperties().RayTracingPipelineProps.shaderGroupHandleAlignment, BlockAlignment);
+		}
 		TempBlockAllocator = new VulkanRHI::FTempBlockAllocator(InDevice, GVulkanTempBlockSize, BlockAlignment, BufferUsageFlags);
 	}
 
diff --git a/Engine/Source/Runtime/VulkanRHI/Private/VulkanRayTracing.cpp b/Engine/Source/Runtime/VulkanRHI/Private/VulkanRayTracing.cpp
--- a/Engine/Source/Runtime/VulkanRHI/Private/VulkanRayTracing.cpp
+++ b/Engine/Source/Runtime/VulkanRHI/Private/VulkanRayTracing.cpp
@@ -879,25 +879,14 @@ FVulkanRayTracingShaderTable::FVulkanRayTracingShaderTable(FRHICommandListBase&
 
 			if (Alloc.HandleCount > 0)
 			{
-				if (InUseLocalRecord)
-				{
-					Alloc.Region.stride = (Alloc.HandleCount > 1) ? (uint32)GVulkanRayTracingMaxShaderGroupStride : 0;
-					Alloc.Region.size = Alloc.HandleCount * (uint32)GVulkanRayTracingMaxShaderGroupStride;
-
-				}
-				else
-				{
-					checkSlow(InHandleCount == 1);
-					Alloc.Region.stride = HandleSizeAligned;
-					Alloc.Region.size = HandleSizeAligned;
-				}
+				Alloc.Region.stride = (Alloc.HandleCount > 1) ? (uint32)GVulkanRayTracingMaxShaderGroupStride : 0;
+				Alloc.Region.size = Alloc.HandleCount * (uint32)GVulkanRayTracingMaxShaderGroupStride;
 
 				// Host buffer
 				Alloc.HostBuffer.SetNumUninitialized(Alloc.Region.size);
 			}
 		};
 
-		InitAlloc(Raygen, 1, false);
 		InitAlloc(Miss, Initializer.NumMissShaderSlots, true);
 		InitAlloc(Callable, Initializer.NumCallableShaderSlots, true);
 
@@ -925,7 +914,6 @@ FVulkanRayTracingShaderTable::~FVulkanRayTracingShaderTable()
 
 void FVulkanRayTracingShaderTable::ReleaseLocalBuffers()
 {
-	ReleaseLocalBuffer(Device, Raygen);
 	ReleaseLocalBuffer(Device, Miss);
 	ReleaseLocalBuffer(Device, HitGroup);
 	ReleaseLocalBuffer(Device, Callable);
@@ -951,11 +939,14 @@ FVulkanRayTracingShaderTable::FVulkanShaderTableAllocation& FVulkanRayTracingSha
 {
 	switch (Frequency)
 	{
-	case SF_RayGen:       return Raygen;
 	case SF_RayMiss:      return Miss;
 	case SF_RayHitGroup:  return HitGroup;
 	case SF_RayCallable:  return Callable;
 
+	case SF_RayGen:
+		checkf(false, TEXT("RayGen have not ShaderTable allocation."));
+		break;
+
 	default:
 		checkf(false, TEXT("Only usable with RayTracing shaders."));
 		break;
@@ -980,11 +971,16 @@ void FVulkanRayTracingShaderTable::SetSlot(EShaderFrequency Frequency, uint32 Ds
 	Alloc.bIsDirty = true;
 }
 
-void FVulkanRayTracingShaderTable::CommitRayGenShader(FVulkanCommandListContext& Context, EShaderFrequency ShaderFrequency, uint32 SrcHandleIndex, TConstArrayView<uint8> SrcHandleData)
+VkStridedDeviceAddressRegionKHR FVulkanRayTracingShaderTable::CommitRayGenShader(FVulkanCommandListContext& Context, uint32 SrcHandleIndex, TConstArrayView<uint8> SrcHandleData)
 {
-	UE::TScopeLock Lock(RaygenMutex);
-	SetSlot(ShaderFrequency, 0, SrcHandleIndex, SrcHandleData);
+	VkStridedDeviceAddressRegionKHR RayGenRegion;
+	uint8* pMappedMemory = Context.GetTempBlockAllocator().Alloc(HandleSize, Context.GetCommandBufferManager()->GetActiveCmdBuffer(), &RayGenRegion);
+
+	FMemory::Memcpy(pMappedMemory, &SrcHandleData[SrcHandleIndex * HandleSize], HandleSize);
+
 	Commit(Context);
+
+	return RayGenRegion;
 }
 
 void FVulkanRayTracingShaderTable::SetLocalShaderParameters(EShaderFrequency Frequency, uint32 RecordIndex, uint32 OffsetWithinRecord, const void* InData, uint32 InDataSize)
@@ -1084,7 +1080,6 @@ void FVulkanRayTracingShaderTable::Commit(FVulkanCommandListContext& Context)
 		}
 	};
 
-	CommitBuffer(Raygen);
 	CommitBuffer(Miss);
 	CommitBuffer(HitGroup);
 	CommitBuffer(Callable);
@@ -2160,13 +2155,14 @@ void FVulkanCommandListContext::RHIRayTraceDispatch(
 	FVulkanCmdBuffer* const CmdBuffer = GetCommandBufferManager()->GetActiveCmdBuffer();
 	VulkanRHI::vkCmdBindPipeline(CmdBuffer->GetHandle(), VK_PIPELINE_BIND_POINT_RAY_TRACING_KHR, Pipeline->GetPipeline());
 
-	ShaderTable->CommitRayGenShader(*this, InRayGenShader->GetFrequency(), Pipeline->GetShaderIndex(RayGenShader), Pipeline->GetShaderHandles(SF_RayGen));
+	checkSlow(InRayGenShader->GetFrequency() == SF_RayGen);
+	const VkStridedDeviceAddressRegionKHR RayGenRegion = ShaderTable->CommitRayGenShader(*this, Pipeline->GetShaderIndex(RayGenShader), Pipeline->GetShaderHandles(SF_RayGen));
 
 	FVulkanPipelineBarrier PostDispatch = SetRayGenResources(Device, CmdBuffer, InGlobalResourceBindings, ShaderTable);
 
 	VulkanRHI::vkCmdTraceRaysKHR(
 		CmdBuffer->GetHandle(),
-		ShaderTable->GetRegion(SF_RayGen),
+		&RayGenRegion,
 		ShaderTable->GetRegion(SF_RayMiss),
 		ShaderTable->GetRegion(SF_RayHitGroup),
 		ShaderTable->GetRegion(SF_RayCallable),
@@ -2196,7 +2192,8 @@ void FVulkanCommandListContext::RHIRayTraceDispatchIndirect(
 	FVulkanCmdBuffer* const CmdBuffer = GetCommandBufferManager()->GetActiveCmdBuffer();
 	VulkanRHI::vkCmdBindPipeline(CmdBuffer->GetHandle(), VK_PIPELINE_BIND_POINT_RAY_TRACING_KHR, Pipeline->GetPipeline());
 
-	ShaderTable->CommitRayGenShader(*this, InRayGenShader->GetFrequency(), Pipeline->GetShaderIndex(RayGenShader), Pipeline->GetShaderHandles(SF_RayGen));
+	checkSlow(InRayGenShader->GetFrequency() == SF_RayGen);
+	const VkStridedDeviceAddressRegionKHR RayGenRegion = ShaderTable->CommitRayGenShader(*this, Pipeline->GetShaderIndex(RayGenShader), Pipeline->GetShaderHandles(SF_RayGen));
 
 	FVulkanPipelineBarrier PostDispatch = SetRayGenResources(Device, CmdBuffer, InGlobalResourceBindings, ShaderTable);
 
@@ -2205,7 +2202,7 @@ void FVulkanCommandListContext::RHIRayTraceDispatchIndirect(
 
 	VulkanRHI::vkCmdTraceRaysIndirectKHR(
 		CmdBuffer->GetHandle(),
-		ShaderTable->GetRegion(SF_RayGen),
+		&RayGenRegion,
 		ShaderTable->GetRegion(SF_RayMiss),
 		ShaderTable->GetRegion(SF_RayHitGroup),
 		ShaderTable->GetRegion(SF_RayCallable),
diff --git a/Engine/Source/Runtime/VulkanRHI/Private/VulkanRayTracing.h b/Engine/Source/Runtime/VulkanRHI/Private/VulkanRayTracing.h
--- a/Engine/Source/Runtime/VulkanRHI/Private/VulkanRayTracing.h
+++ b/Engine/Source/Runtime/VulkanRHI/Private/VulkanRayTracing.h
@@ -72,7 +72,7 @@ public:
 
 	void SetSlot(EShaderFrequency Frequency, uint32 DstSlot, uint32 SrcHandleIndex, TConstArrayView<uint8> SrcHandleData);
 
-	void CommitRayGenShader(FVulkanCommandListContext& Context, EShaderFrequency ShaderFrequency, uint32 SrcHandleIndex, TConstArrayView<uint8> SrcHandleData);
+	VkStridedDeviceAddressRegionKHR CommitRayGenShader(FVulkanCommandListContext& Context, uint32 SrcHandleIndex, TConstArrayView<uint8> SrcHandleData);
 
 	template <typename T>
 	void SetLocalShaderParameters(EShaderFrequency Frequency, uint32 RecordIndex, uint32 InOffsetWithinRootSignature, const T& Parameters)
@@ -152,8 +152,6 @@ private:
 	ERayTracingShaderBindingMode ShaderBindingMode = ERayTracingShaderBindingMode::Disabled;
 	ERayTracingHitGroupIndexingMode HitGroupIndexingMode = ERayTracingHitGroupIndexingMode::Allow;
 
-	UE::FMutex RaygenMutex;
-	FVulkanShaderTableAllocation Raygen;
 	FVulkanShaderTableAllocation Miss;
 	FVulkanShaderTableAllocation HitGroup;
 	FVulkanShaderTableAllocation Callable;
