VulkanRHI: Fix alignment used for raygen handle.

#rnx

diff --git a/Engine/Source/Runtime/VulkanRHI/Private/VulkanRHI.cpp b/Engine/Source/Runtime/VulkanRHI/Private/VulkanRHI.cpp
--- a/Engine/Source/Runtime/VulkanRHI/Private/VulkanRHI.cpp
+++ b/Engine/Source/Runtime/VulkanRHI/Private/VulkanRHI.cpp
@@ -478,7 +478,7 @@ FVulkanCommandListContext::FVulkanCommandListContext(FVulkanDynamicRHI* InRHI, F
 		if (Device->GetOptionalExtensions().HasRayTracingPipeline)
 		{
 			BufferUsageFlags |= VK_BUFFER_USAGE_SHADER_BINDING_TABLE_BIT_KHR;
-			BlockAlignment = FMath::Max<uint32>(Device->GetOptionalExtensionProperties().RayTracingPipelineProps.shaderGroupHandleAlignment, BlockAlignment);
+			BlockAlignment = FMath::Max<uint32>(Device->GetOptionalExtensionProperties().RayTracingPipelineProps.shaderGroupBaseAlignment, BlockAlignment);
 		}
 		TempBlockAllocator = new VulkanRHI::FTempBlockAllocator(InDevice, GVulkanTempBlockSize, BlockAlignment, BufferUsageFlags);
 	}
