VulkanRHI: Enable labels when debugging GPU crashes

#rnx

diff --git a/Engine/Source/Runtime/VulkanRHI/Private/VulkanDevice.cpp b/Engine/Source/Runtime/VulkanRHI/Private/VulkanDevice.cpp
--- a/Engine/Source/Runtime/VulkanRHI/Private/VulkanDevice.cpp
+++ b/Engine/Source/Runtime/VulkanRHI/Private/VulkanDevice.cpp
@@ -492,14 +492,16 @@ TAutoConsoleVariable<int32> CVarVulkanDebugMarkers(
 void FVulkanDevice::SetupDrawMarkers()
 {
 #if VULKAN_ENABLE_DRAW_MARKERS
-	bool bTraceToolFound = false;
+	const bool bGPUCrashDebugging = UE::RHI::UseGPUCrashDebugging();
+	bool bTraceToolFound = bGPUCrashDebugging;
+
 #if VULKAN_HAS_DEBUGGING_ENABLED
-	bTraceToolFound = GRenderDocFound;
+	bTraceToolFound |= GRenderDocFound;
 #endif // VULKAN_HAS_DEBUGGING_ENABLED
 
 	if (RHI->SupportsDebugUtilsExt())
 	{
-		bool bUseLabel = false;
+		bool bUseLabel = bGPUCrashDebugging;
 		bool bUseName = false;
 
 		if (OptionalDeviceExtensions.HasEXTToolingInfo)
@@ -530,8 +532,8 @@ void FVulkanDevice::SetupDrawMarkers()
 		const int VulkanDebugMarkers = FParse::Param(FCommandLine::Get(), TEXT("forcevulkandrawmarkers")) ? 3 : CVarVulkanDebugMarkers.GetValueOnRenderThread();
 		if (VulkanDebugMarkers != 4)
 		{
-			bUseLabel = VulkanDebugMarkers & 2;
-			bUseName = VulkanDebugMarkers & 1;
+			bUseLabel = bUseLabel || (VulkanDebugMarkers & 2);
+			bUseName = bUseName || (VulkanDebugMarkers & 1);
 		}
 
 		if (bUseLabel)
