Improve r.Vulkan.DepthStencilForceStorageBit workaround
This is part of the PR below:
PR #10547: Avoid VK_IMAGE_USAGE_STORAGE_BIT when not necessary

diff --git a/Engine/Config/BaseDeviceProfiles.ini b/Engine/Config/BaseDeviceProfiles.ini
--- a/Engine/Config/BaseDeviceProfiles.ini
+++ b/Engine/Config/BaseDeviceProfiles.ini
@@ -967,7 +967,8 @@ BaseProfileName=Android_Vulkan_SM5
 [Android_Mali_Vulkan_SM5 DeviceProfile]
 DeviceType=Android
 BaseProfileName=Android_Vulkan_SM5
-+CVars=r.Vulkan.DepthStencilForceStorageBit=1
+; DepthStencilForceStorageBit is a workaround that is only needed for Mali before the G710, so it's not needed based on the current matching rules for Android_Mali_Vulkan_SM5
+;+CVars=r.Vulkan.DepthStencilForceStorageBit=1
 
 [Android_Xclipse_Vulkan_SM5 DeviceProfile]
 DeviceType=Android
diff --git a/Engine/Source/Runtime/VulkanRHI/Private/VulkanTexture.cpp b/Engine/Source/Runtime/VulkanRHI/Private/VulkanTexture.cpp
--- a/Engine/Source/Runtime/VulkanRHI/Private/VulkanTexture.cpp
+++ b/Engine/Source/Runtime/VulkanRHI/Private/VulkanTexture.cpp
@@ -24,6 +24,7 @@ static FAutoConsoleVariableRef CVarVulkanSubmitOnTextureUnlock(
 	ECVF_Default
 );
 
+// This is a workaround for issues with AFBC on Mali GPUs before the G710
 int32 GVulkanDepthStencilForceStorageBit = 0;
 static FAutoConsoleVariableRef CVarVulkanDepthStencilForceStorageBit(
 	TEXT("r.Vulkan.DepthStencilForceStorageBit"),
@@ -473,7 +474,7 @@ void FVulkanTexture::GenerateImageCreateInfo(
 		ImageCreateInfo.usage &= ~VK_IMAGE_USAGE_TRANSFER_DST_BIT;
 	}
 
-	if (EnumHasAnyFlags(UEFlags, TexCreate_DepthStencilTargetable) && GVulkanDepthStencilForceStorageBit)
+	if (GVulkanDepthStencilForceStorageBit && EnumHasAnyFlags(UEFlags, TexCreate_DepthStencilTargetable) && (TextureFormat != VK_FORMAT_D16_UNORM && TextureFormat != VK_FORMAT_D32_SFLOAT))
 	{
 		ImageCreateInfo.usage |= VK_IMAGE_USAGE_STORAGE_BIT;
 	}
