Fix using non multiview texture on multiview layour error log spam.
Also fix shadow rendering using depth stencil with mismatched number of samples

diff --git a/Engine/Source/Runtime/Renderer/Private/ShadowRendering.cpp b/Engine/Source/Runtime/Renderer/Private/ShadowRendering.cpp
--- a/Engine/Source/Runtime/Renderer/Private/ShadowRendering.cpp
+++ b/Engine/Source/Runtime/Renderer/Private/ShadowRendering.cpp
@@ -2060,12 +2060,22 @@ void FSceneRenderer::RenderShadowProjections(
 		{
 			CommonPassParameters.Substrate = Substrate::BindSubstrateGlobalUniformParameters(View);
 		}
-		
+
+		FRDGTextureRef DepthStencilTexture = nullptr;
+		// When using MSAA we need to bind the texture with the correct number of samples
+		if (!(OutputTexture->Desc.NumSamples > 1) && SceneTextures.Depth.IsSeparate() && HasBeenProduced(SceneTextures.Depth.Resolve))
+		{
+			DepthStencilTexture = SceneTextures.Depth.Resolve;
+		}
+		else
+		{
+			DepthStencilTexture = SceneTextures.Depth.Target;
+		}
 		CommonPassParameters.RenderTargets[0] = FRenderTargetBinding(OutputTexture, ERenderTargetLoadAction::ELoad);
 		CommonPassParameters.RenderTargets.DepthStencil =
 			bSubPixelShadow ?
 			FDepthStencilBinding(View.HairStrandsViewData.VisibilityData.HairOnlyDepthTexture, ERenderTargetLoadAction::ELoad, ERenderTargetLoadAction::ELoad, ExclusiveDepthStencil) :
-			FDepthStencilBinding(SceneTextures.Depth.Target, ERenderTargetLoadAction::ELoad, ERenderTargetLoadAction::ELoad, ExclusiveDepthStencil);
+			FDepthStencilBinding(DepthStencilTexture, ERenderTargetLoadAction::ELoad, ERenderTargetLoadAction::ELoad, ExclusiveDepthStencil);
 
 		RDG_GPU_MASK_SCOPE(GraphBuilder, View.GPUMask);
 		RDG_EVENT_SCOPE_CONDITIONAL(GraphBuilder, Views.Num() > 1, "View%d", ViewIndex);
diff --git a/Engine/Source/Runtime/VulkanRHI/Private/VulkanRenderTarget.cpp b/Engine/Source/Runtime/VulkanRHI/Private/VulkanRenderTarget.cpp
--- a/Engine/Source/Runtime/VulkanRHI/Private/VulkanRenderTarget.cpp
+++ b/Engine/Source/Runtime/VulkanRHI/Private/VulkanRenderTarget.cpp
@@ -1076,7 +1076,7 @@ FVulkanRenderTargetLayout::FVulkanRenderTargetLayout(FVulkanDevice& InDevice, co
 		++NumAttachmentDescriptions;
 		++NumColorAttachments;
 	}
-
+	bool bMultiViewDepthStencil = false;
 	if (RPInfo.DepthStencilRenderTarget.DepthStencilTarget)
 	{
 		VkAttachmentDescription& CurrDesc = Desc[NumAttachmentDescriptions];
@@ -1084,7 +1084,7 @@ FVulkanRenderTargetLayout::FVulkanRenderTargetLayout(FVulkanDevice& InDevice, co
 		FVulkanTexture* Texture = ResourceCast(RPInfo.DepthStencilRenderTarget.DepthStencilTarget);
 		check(Texture);
 		const FRHITextureDesc& TextureDesc = Texture->GetDesc();
-
+		bMultiViewDepthStencil = (Texture->GetNumberOfArrayLevels() > 1) && !Texture->GetDesc().IsTextureCube();
 		CurrDesc.samples = static_cast<VkSampleCountFlagBits>(RPInfo.DepthStencilRenderTarget.DepthStencilTarget->GetNumSamples());
 		// CustomResolveSubpass can have targets with a different NumSamples
 		ensure(!NumSamples || CurrDesc.samples == NumSamples || RPInfo.SubpassHint == ESubpassHint::CustomResolveSubpass);
@@ -1245,8 +1245,8 @@ FVulkanRenderTargetLayout::FVulkanRenderTargetLayout(FVulkanDevice& InDevice, co
 
 	CompatibleHashInfo.NumSamples = NumSamples;
 	CompatibleHashInfo.MultiViewCount = MultiViewCount;
-
-	if (MultiViewCount > 1 && !bMultiviewRenderTargets)
+	// Depth prepass has no color RTs but has a depth attachment that must be multiview
+	if (MultiViewCount > 1 && !bMultiviewRenderTargets && !(NumColorRenderTargets == 0 && bMultiViewDepthStencil))
 	{
 		UE_LOG(LogVulkan, Error, TEXT("Non multiview textures on a multiview layout!"));
 	}
