VulkanRHI: Fix validation error for missing sType

#rnx

diff --git a/Engine/Source/Runtime/VulkanRHI/Private/VulkanDevice.cpp b/Engine/Source/Runtime/VulkanRHI/Private/VulkanDevice.cpp
--- a/Engine/Source/Runtime/VulkanRHI/Private/VulkanDevice.cpp
+++ b/Engine/Source/Runtime/VulkanRHI/Private/VulkanDevice.cpp
@@ -509,7 +509,12 @@ void FVulkanDevice::SetupDrawMarkers()
 			uint32_t ToolCount = 0;
 			VulkanRHI::vkGetPhysicalDeviceToolPropertiesEXT(Gpu, &ToolCount, nullptr);
 			TArray<VkPhysicalDeviceToolPropertiesEXT> ToolProperties;
-			ToolProperties.AddZeroed(ToolCount);
+			ToolProperties.SetNumUninitialized(ToolCount);
+			for (VkPhysicalDeviceToolPropertiesEXT& Tool : ToolProperties)
+			{
+				ZeroVulkanStruct(Tool, VK_STRUCTURE_TYPE_PHYSICAL_DEVICE_TOOL_PROPERTIES);
+			}
+
 			VulkanRHI::vkGetPhysicalDeviceToolPropertiesEXT(Gpu, &ToolCount, ToolProperties.GetData());
 			for (VkPhysicalDeviceToolPropertiesEXT const& Tool : ToolProperties)
 			{
