Vulkan UAV format aliasing

diff --git a/Engine/Source/Runtime/VulkanRHI/Private/VulkanRHI.cpp b/Engine/Source/Runtime/VulkanRHI/Private/VulkanRHI.cpp
--- a/Engine/Source/Runtime/VulkanRHI/Private/VulkanRHI.cpp
+++ b/Engine/Source/Runtime/VulkanRHI/Private/VulkanRHI.cpp
@@ -981,6 +981,8 @@ void FVulkanDynamicRHI::InitInstance()
 		GRHISupportsParallelRHIExecute = GRHIThreadCvar->GetInt() > 1;
 #endif
 
+		GRHISupportsUAVFormatAliasing = true;
+
 		// Some platforms might only have CPU for an RHI thread, but not for parallel tasks
 		GSupportsParallelRenderingTasksWithSeparateRHIThread = GRHISupportsRHIThread ? FVulkanPlatform::SupportParallelRenderingTasks() : false;
 
diff --git a/Engine/Source/Runtime/VulkanRHI/Private/VulkanTexture.cpp b/Engine/Source/Runtime/VulkanRHI/Private/VulkanTexture.cpp
--- a/Engine/Source/Runtime/VulkanRHI/Private/VulkanTexture.cpp
+++ b/Engine/Source/Runtime/VulkanRHI/Private/VulkanTexture.cpp
@@ -333,7 +333,8 @@ void FVulkanTexture::GenerateImageCreateInfo(
 
 	ImageCreateInfo.flags = (ResourceType == VK_IMAGE_VIEW_TYPE_CUBE || ResourceType == VK_IMAGE_VIEW_TYPE_CUBE_ARRAY) ? VK_IMAGE_CREATE_CUBE_COMPATIBLE_BIT : 0;
 
-	const bool bNeedsMutableFormat = (EnumHasAllFlags(UEFlags, TexCreate_SRGB) || (InDesc.Format == PF_R64_UINT));
+	const bool bHasUAVFormat = (InDesc.UAVFormat != PF_Unknown && InDesc.UAVFormat != InDesc.Format);
+	const bool bNeedsMutableFormat = (EnumHasAllFlags(UEFlags, TexCreate_SRGB) || (InDesc.Format == PF_R64_UINT) || bHasUAVFormat);
 	if (bNeedsMutableFormat)
 	{
 		if (InDevice.GetOptionalExtensions().HasKHRImageFormatList)
@@ -357,11 +358,21 @@ void FVulkanTexture::GenerateImageCreateInfo(
 				OutImageCreateInfo.FormatsUsed.Add(UEToVkTextureFormat(PF_R32G32_UINT, false));
 			}
 
+			if (bHasUAVFormat)
+			{
+				OutImageCreateInfo.FormatsUsed.AddUnique(nonSrgbFormat);
+				OutImageCreateInfo.FormatsUsed.AddUnique(UEToVkTextureFormat(InDesc.UAVFormat, false));
+			}
+
 			ImageFormatListCreateInfo.pViewFormats = OutImageCreateInfo.FormatsUsed.GetData();
 			ImageFormatListCreateInfo.viewFormatCount = OutImageCreateInfo.FormatsUsed.Num();
 		}
 
 		ImageCreateInfo.flags |= VK_IMAGE_CREATE_MUTABLE_FORMAT_BIT;
+		if (bHasUAVFormat && IsBlockCompressedFormat(InDesc.Format) && !IsBlockCompressedFormat(InDesc.UAVFormat))
+		{
+			ImageCreateInfo.flags |= VK_IMAGE_CREATE_BLOCK_TEXEL_VIEW_COMPATIBLE_BIT;
+		}
 	}
 
 	if (ImageCreateInfo.imageType == VK_IMAGE_TYPE_3D)
@@ -444,9 +455,20 @@ void FVulkanTexture::GenerateImageCreateInfo(
 		ImageCreateInfo.usage &= ~VK_IMAGE_USAGE_SAMPLED_BIT;
 	}
 
+	if (bHasUAVFormat)
+	{
+		const VkFormat UAVFormat = UEToVkTextureFormat(InDesc.UAVFormat, false);
+		const VkFormatProperties& UAVFormatProperties = InDevice.GetFormatProperties(UAVFormat);
+		const VkFormatFeatureFlags UAVFormatFlags = ImageCreateInfo.tiling == VK_IMAGE_TILING_LINEAR ?
+			UAVFormatProperties.linearTilingFeatures :
+			UAVFormatProperties.optimalTilingFeatures;
+
+		ensure((UAVFormatFlags & VK_FORMAT_FEATURE_STORAGE_IMAGE_BIT) != 0);
+	}
+
 	if (!VKHasAnyFlags(FormatFlags, VK_FORMAT_FEATURE_STORAGE_IMAGE_BIT))
 	{
-		ensure((ImageCreateInfo.usage & VK_IMAGE_USAGE_STORAGE_BIT) == 0);
+		ensure((ImageCreateInfo.usage & VK_IMAGE_USAGE_STORAGE_BIT) == 0 || bHasUAVFormat);
 		ImageCreateInfo.usage &= ~VK_IMAGE_USAGE_STORAGE_BIT;
 	}
 
