VulkanRHI: Clean up: rename classes to be more consistent with other RHI:
- FVulkanCmdBuffer to FVulkanCommandBuffer
- FVulkanResourceMultiBuffer to FVulkanBuffer

#rnx

diff --git a/Engine/Source/Runtime/VulkanRHI/Private/Android/VulkanAndroidPlatform.cpp b/Engine/Source/Runtime/VulkanRHI/Private/Android/VulkanAndroidPlatform.cpp
--- a/Engine/Source/Runtime/VulkanRHI/Private/Android/VulkanAndroidPlatform.cpp
+++ b/Engine/Source/Runtime/VulkanRHI/Private/Android/VulkanAndroidPlatform.cpp
@@ -1899,7 +1899,7 @@ FString FVulkanAndroidPlatform::GetVulkanProfileNameForFeatureLevel(ERHIFeatureL
 	return ProfileName;
 } 
 
-void FVulkanAndroidPlatform::WriteCrashMarker(const FOptionalVulkanDeviceExtensions& OptionalExtensions, FVulkanCmdBuffer* CmdBuffer, VkBuffer DestBuffer, const TArrayView<uint32>& Entries, bool bAdding)
+void FVulkanAndroidPlatform::WriteCrashMarker(const FOptionalVulkanDeviceExtensions& OptionalExtensions, FVulkanCommandBuffer* CmdBuffer, VkBuffer DestBuffer, const TArrayView<uint32>& Entries, bool bAdding)
 {
 	ensure(Entries.Num() <= GMaxCrashBufferEntries);
 
diff --git a/Engine/Source/Runtime/VulkanRHI/Private/Android/VulkanAndroidPlatform.h b/Engine/Source/Runtime/VulkanRHI/Private/Android/VulkanAndroidPlatform.h
--- a/Engine/Source/Runtime/VulkanRHI/Private/Android/VulkanAndroidPlatform.h
+++ b/Engine/Source/Runtime/VulkanRHI/Private/Android/VulkanAndroidPlatform.h
@@ -176,7 +176,7 @@ public:
 		return VK_SHADER_STAGE_FRAGMENT_BIT | VK_SHADER_STAGE_COMPUTE_BIT;
 	}
 
-	static void WriteCrashMarker(const FOptionalVulkanDeviceExtensions& OptionalExtensions, FVulkanCmdBuffer* CmdBuffer, VkBuffer DestBuffer, const TArrayView<uint32>& Entries, bool bAdding);
+	static void WriteCrashMarker(const FOptionalVulkanDeviceExtensions& OptionalExtensions, FVulkanCommandBuffer* CmdBuffer, VkBuffer DestBuffer, const TArrayView<uint32>& Entries, bool bAdding);
 
 #if USE_ANDROID_VULKAN_SWAPPY
 	static bool bSwappyEnabledAtRHIInit;
diff --git a/Engine/Source/Runtime/VulkanRHI/Private/Linux/VulkanLinuxPlatform.cpp b/Engine/Source/Runtime/VulkanRHI/Private/Linux/VulkanLinuxPlatform.cpp
--- a/Engine/Source/Runtime/VulkanRHI/Private/Linux/VulkanLinuxPlatform.cpp
+++ b/Engine/Source/Runtime/VulkanRHI/Private/Linux/VulkanLinuxPlatform.cpp
@@ -255,7 +255,7 @@ void FVulkanLinuxPlatform::CreateSurface(void* WindowHandle, VkInstance Instance
 	}
 }
 
-void FVulkanLinuxPlatform::WriteCrashMarker(const FOptionalVulkanDeviceExtensions& OptionalExtensions, FVulkanCmdBuffer* CmdBuffer, VkBuffer DestBuffer, const TArrayView<uint32>& Entries, bool bAdding)
+void FVulkanLinuxPlatform::WriteCrashMarker(const FOptionalVulkanDeviceExtensions& OptionalExtensions, FVulkanCommandBuffer* CmdBuffer, VkBuffer DestBuffer, const TArrayView<uint32>& Entries, bool bAdding)
 {
 	ensure(Entries.Num() <= GMaxCrashBufferEntries);
 
diff --git a/Engine/Source/Runtime/VulkanRHI/Private/Linux/VulkanLinuxPlatform.h b/Engine/Source/Runtime/VulkanRHI/Private/Linux/VulkanLinuxPlatform.h
--- a/Engine/Source/Runtime/VulkanRHI/Private/Linux/VulkanLinuxPlatform.h
+++ b/Engine/Source/Runtime/VulkanRHI/Private/Linux/VulkanLinuxPlatform.h
@@ -47,7 +47,7 @@ public:
 
 	static void CreateSurface(void* WindowHandle, VkInstance Instance, VkSurfaceKHR* OutSurface);
 
-	static void WriteCrashMarker(const FOptionalVulkanDeviceExtensions& OptionalExtensions, FVulkanCmdBuffer* CmdBuffer, VkBuffer DestBuffer, const TArrayView<uint32>& Entries, bool bAdding);
+	static void WriteCrashMarker(const FOptionalVulkanDeviceExtensions& OptionalExtensions, FVulkanCommandBuffer* CmdBuffer, VkBuffer DestBuffer, const TArrayView<uint32>& Entries, bool bAdding);
 
 protected:
 	static void* VulkanLib;
diff --git a/Engine/Source/Runtime/VulkanRHI/Private/VulkanBarriers.cpp b/Engine/Source/Runtime/VulkanRHI/Private/VulkanBarriers.cpp
--- a/Engine/Source/Runtime/VulkanRHI/Private/VulkanBarriers.cpp
+++ b/Engine/Source/Runtime/VulkanRHI/Private/VulkanBarriers.cpp
@@ -766,7 +766,7 @@ void FVulkanDynamicRHI::RHICreateTransition(FRHITransition* Transition, const FR
 
 		checkf(Info.AccessAfter != ERHIAccess::Unknown, TEXT("Transitioning a resource to an unknown state is not allowed."));
 
-		FVulkanResourceMultiBuffer* Buffer = nullptr;
+		FVulkanBuffer* Buffer = nullptr;
 		FVulkanTexture* Texture = nullptr;
 		FRHITransitionInfo::EType UnderlyingType = Info.Type;
 		uint32 UsageFlags = 0;
@@ -1572,7 +1572,7 @@ void FVulkanCommandListContext::RHIEndTransitions(TArrayView<const FRHITransitio
 			CommandBufferManager->PrepareForNewActiveCommandBuffer();
 		}
 
-		FVulkanCmdBuffer* CmdBuffer = CommandBufferManager->GetActiveCmdBuffer();
+		FVulkanCommandBuffer* CmdBuffer = CommandBufferManager->GetActiveCmdBuffer();
 		CmdBuffer->AddWaitSemaphore(VK_PIPELINE_STAGE_TOP_OF_PIPE_BIT, WaitSemaphores);
 	}
 
@@ -1820,7 +1820,7 @@ void FVulkanPipelineBarrier::Execute(VkCommandBuffer CmdBuffer)
 	}
 }
 
-void FVulkanPipelineBarrier::Execute(FVulkanCmdBuffer* CmdBuffer)
+void FVulkanPipelineBarrier::Execute(FVulkanCommandBuffer* CmdBuffer)
 {
 	if (MemoryBarriers.Num() != 0 || BufferBarriers.Num() != 0 || ImageBarriers.Num() != 0)
 	{
@@ -1860,7 +1860,7 @@ VkImageSubresourceRange FVulkanPipelineBarrier::MakeSubresourceRange(VkImageAspe
 //
 // Used when we need to change the layout of a single image. Some plug-ins call this function from outside the RHI (Steam VR, at the time of writing this).
 //
-void VulkanSetImageLayout(FVulkanCmdBuffer* CmdBuffer, VkImage Image, VkImageLayout OldLayout, VkImageLayout NewLayout, const VkImageSubresourceRange& SubresourceRange)
+void VulkanSetImageLayout(FVulkanCommandBuffer* CmdBuffer, VkImage Image, VkImageLayout OldLayout, VkImageLayout NewLayout, const VkImageSubresourceRange& SubresourceRange)
 {
 	FVulkanPipelineBarrier Barrier;
 	Barrier.AddImageLayoutTransition(Image, OldLayout, NewLayout, SubresourceRange);
@@ -1975,7 +1975,7 @@ void FVulkanImageLayout::Set(VkImageLayout Layout, const VkImageSubresourceRange
 	CollapseSubresLayoutsIfSame();
 }
 
-VkImageLayout FVulkanLayoutManager::GetDefaultLayout(FVulkanCmdBuffer* CmdBuffer, const FVulkanTexture& VulkanTexture, ERHIAccess DesiredAccess)
+VkImageLayout FVulkanLayoutManager::GetDefaultLayout(FVulkanCommandBuffer* CmdBuffer, const FVulkanTexture& VulkanTexture, ERHIAccess DesiredAccess)
 {
 	switch (DesiredAccess)
 	{
@@ -2036,7 +2036,7 @@ VkImageLayout FVulkanLayoutManager::GetDefaultLayout(FVulkanCmdBuffer* CmdBuffer
 	}
 }
 
-VkImageLayout FVulkanLayoutManager::SetExpectedLayout(FVulkanCmdBuffer* CmdBuffer, const FVulkanTexture& VulkanTexture, ERHIAccess DesiredAccess)
+VkImageLayout FVulkanLayoutManager::SetExpectedLayout(FVulkanCommandBuffer* CmdBuffer, const FVulkanTexture& VulkanTexture, ERHIAccess DesiredAccess)
 {
 	const VkImageLayout ExpectedLayout = GetDefaultLayout(CmdBuffer, VulkanTexture, DesiredAccess);
 	VkImageLayout PreviousLayout = ExpectedLayout;
diff --git a/Engine/Source/Runtime/VulkanRHI/Private/VulkanBarriers.h b/Engine/Source/Runtime/VulkanRHI/Private/VulkanBarriers.h
--- a/Engine/Source/Runtime/VulkanRHI/Private/VulkanBarriers.h
+++ b/Engine/Source/Runtime/VulkanRHI/Private/VulkanBarriers.h
@@ -8,7 +8,7 @@
 #include "VulkanResources.h"
 #include "VulkanThirdParty.h"
 
-class FVulkanCmdBuffer;
+class FVulkanCommandBuffer;
 
 struct FVulkanPipelineBarrier
 {
@@ -41,7 +41,7 @@ struct FVulkanPipelineBarrier
 	void AddImageLayoutTransition(VkImage Image, VkImageAspectFlags AspectMask, const struct FVulkanImageLayout& SrcLayout, const struct FVulkanImageLayout& DstLayout);
 	void AddImageAccessTransition(const FVulkanTexture& Surface, ERHIAccess SrcAccess, ERHIAccess DstAccess, const VkImageSubresourceRange& SubresourceRange, VkImageLayout& InOutLayout);
 	void Execute(VkCommandBuffer CmdBuffer);
-	void Execute(FVulkanCmdBuffer* CmdBuffer);
+	void Execute(FVulkanCommandBuffer* CmdBuffer);
 
 	static VkImageSubresourceRange MakeSubresourceRange(VkImageAspectFlags AspectMask, uint32 FirstMip = 0, uint32 NumMips = VK_REMAINING_MIP_LEVELS, uint32 FirstLayer = 0, uint32 NumLayers = VK_REMAINING_ARRAY_LAYERS);
 };
@@ -123,10 +123,10 @@ public:
 	void NotifyDeletedImage(VkImage Image);
 
 	// Predetermined layouts for given RHIAccess
-	static VkImageLayout GetDefaultLayout(FVulkanCmdBuffer* CmdBuffer, const FVulkanTexture& VulkanTexture, ERHIAccess DesiredAccess);
+	static VkImageLayout GetDefaultLayout(FVulkanCommandBuffer* CmdBuffer, const FVulkanTexture& VulkanTexture, ERHIAccess DesiredAccess);
 
 	// Expected layouts and Hints are workarounds until we can use 'hardcoded layouts' everywhere.
-	static VkImageLayout SetExpectedLayout(FVulkanCmdBuffer* CmdBuffer, const FVulkanTexture& VulkanTexture, ERHIAccess DesiredAccess);
+	static VkImageLayout SetExpectedLayout(FVulkanCommandBuffer* CmdBuffer, const FVulkanTexture& VulkanTexture, ERHIAccess DesiredAccess);
 	VkImageLayout GetDepthStencilHint(const FVulkanTexture& VulkanTexture, VkImageAspectFlagBits AspectBit);
 
 	VULKANRHI_API const FVulkanImageLayout* GetFullLayout(VkImage Image) const
diff --git a/Engine/Source/Runtime/VulkanRHI/Private/VulkanBuffer.cpp b/Engine/Source/Runtime/VulkanRHI/Private/VulkanBuffer.cpp
--- a/Engine/Source/Runtime/VulkanRHI/Private/VulkanBuffer.cpp
+++ b/Engine/Source/Runtime/VulkanRHI/Private/VulkanBuffer.cpp
@@ -21,7 +21,7 @@ struct FVulkanPendingBufferLock
 	EResourceLockMode LockMode = RLM_Num;
 };
 
-static TMap<FVulkanResourceMultiBuffer*, FVulkanPendingBufferLock> GPendingLocks;
+static TMap<FVulkanBuffer*, FVulkanPendingBufferLock> GPendingLocks;
 static FCriticalSection GPendingLockMutex;
 
 int32 GVulkanForceStagingBufferOnLock = 0;
@@ -33,7 +33,7 @@ static FAutoConsoleVariableRef CVarVulkanForceStagingBufferOnLock(
 	ECVF_RenderThreadSafe
 );
 
-static FORCEINLINE FVulkanPendingBufferLock GetPendingBufferLock(FVulkanResourceMultiBuffer* Buffer)
+static FORCEINLINE FVulkanPendingBufferLock GetPendingBufferLock(FVulkanBuffer* Buffer)
 {
 	FVulkanPendingBufferLock PendingLock;
 
@@ -45,7 +45,7 @@ static FORCEINLINE FVulkanPendingBufferLock GetPendingBufferLock(FVulkanResource
 	return PendingLock;
 }
 
-static FORCEINLINE void AddPendingBufferLock(FVulkanResourceMultiBuffer* Buffer, FVulkanPendingBufferLock& PendingLock)
+static FORCEINLINE void AddPendingBufferLock(FVulkanBuffer* Buffer, FVulkanPendingBufferLock& PendingLock)
 {
 	FScopeLock ScopeLock(&GPendingLockMutex);
 	check(!GPendingLocks.Contains(Buffer));
@@ -69,7 +69,7 @@ static VkDeviceAddress GetBufferDeviceAddress(FVulkanDevice* Device, VkBuffer Bu
 	return 0;
 }
 
-VkBufferUsageFlags FVulkanResourceMultiBuffer::UEToVKBufferUsageFlags(FVulkanDevice* InDevice, EBufferUsageFlags InUEUsage, bool bZeroSize)
+VkBufferUsageFlags FVulkanBuffer::UEToVKBufferUsageFlags(FVulkanDevice* InDevice, EBufferUsageFlags InUEUsage, bool bZeroSize)
 {
 	// Always include TRANSFER_SRC since hardware vendors confirmed it wouldn't have any performance cost and we need it for some debug functionalities.
 	VkBufferUsageFlags OutVkUsage = VK_BUFFER_USAGE_TRANSFER_SRC_BIT;
@@ -112,11 +112,11 @@ VkBufferUsageFlags FVulkanResourceMultiBuffer::UEToVKBufferUsageFlags(FVulkanDev
 	return OutVkUsage;
 }
 
-FVulkanResourceMultiBuffer::FVulkanResourceMultiBuffer(FVulkanDevice* InDevice, FRHIBufferDesc const& InBufferDesc, FRHIResourceCreateInfo& CreateInfo, FRHICommandListBase* InRHICmdList, const FRHITransientHeapAllocation* InTransientHeapAllocation)
+FVulkanBuffer::FVulkanBuffer(FVulkanDevice* InDevice, FRHIBufferDesc const& InBufferDesc, FRHIResourceCreateInfo& CreateInfo, FRHICommandListBase* InRHICmdList, const FRHITransientHeapAllocation* InTransientHeapAllocation)
 	: FRHIBuffer(InBufferDesc)
 	, VulkanRHI::FDeviceChild(InDevice)
 {
-	VULKAN_TRACK_OBJECT_CREATE(FVulkanResourceMultiBuffer, this);
+	VULKAN_TRACK_OBJECT_CREATE(FVulkanBuffer, this);
 
 	const bool bZeroSize = (InBufferDesc.Size == 0);
 	BufferUsageFlags = UEToVKBufferUsageFlags(InDevice, InBufferDesc.Usage, bZeroSize);
@@ -171,14 +171,14 @@ FVulkanResourceMultiBuffer::FVulkanResourceMultiBuffer(FVulkanDevice* InDevice,
 	}
 }
 
-FVulkanResourceMultiBuffer::~FVulkanResourceMultiBuffer()
+FVulkanBuffer::~FVulkanBuffer()
 {
-	VULKAN_TRACK_OBJECT_DELETE(FVulkanResourceMultiBuffer, this);
+	VULKAN_TRACK_OBJECT_DELETE(FVulkanBuffer, this);
 	ReleaseOwnership();
 }
 
 
-void FVulkanResourceMultiBuffer::AllocateMemory(FBufferAlloc& OutAlloc)
+void FVulkanBuffer::AllocateMemory(FBufferAlloc& OutAlloc)
 {
 	VkMemoryPropertyFlags BufferMemFlags = 0;
 	const bool bUnifiedMem = Device->HasUnifiedMemory();
@@ -209,7 +209,7 @@ void FVulkanResourceMultiBuffer::AllocateMemory(FBufferAlloc& OutAlloc)
 	UpdateVulkanBufferStats(GetDesc(), BufferSize, true);
 }
 
-void* FVulkanResourceMultiBuffer::Lock(FRHICommandListBase& RHICmdList, EResourceLockMode LockMode, uint32 LockSize, uint32 Offset)
+void* FVulkanBuffer::Lock(FRHICommandListBase& RHICmdList, EResourceLockMode LockMode, uint32 LockSize, uint32 Offset)
 {
 	void* Data = nullptr;
 	uint32 DataOffset = 0;
@@ -248,7 +248,7 @@ void* FVulkanResourceMultiBuffer::Lock(FRHICommandListBase& RHICmdList, EResourc
 		
 			FVulkanCommandListContextImmediate& Context = Device->GetImmediateContext();
 			FVulkanCommandBufferManager* CommandBufferManager = Context.GetCommandBufferManager();
-			FVulkanCmdBuffer* CmdBuffer = CommandBufferManager->GetUploadCmdBuffer();
+			FVulkanCommandBuffer* CmdBuffer = CommandBufferManager->GetUploadCmdBuffer();
 				
 			// Make sure any previous tasks have finished on the source buffer.
 			VkMemoryBarrier BarrierBefore = { VK_STRUCTURE_TYPE_MEMORY_BARRIER, nullptr, VK_ACCESS_MEMORY_WRITE_BIT, VK_ACCESS_MEMORY_READ_BIT };
@@ -340,7 +340,7 @@ void* FVulkanResourceMultiBuffer::Lock(FRHICommandListBase& RHICmdList, EResourc
 }
 
 
-void FVulkanResourceMultiBuffer::Unlock(FRHICommandListBase& RHICmdList)
+void FVulkanBuffer::Unlock(FRHICommandListBase& RHICmdList)
 {
 	const bool bUnifiedMem = Device->HasUnifiedMemory();
 	const bool bDynamic = EnumHasAnyFlags(GetUsage(), BUF_Dynamic) || EnumHasAnyFlags(GetUsage(), BUF_Volatile);
@@ -376,7 +376,7 @@ void FVulkanResourceMultiBuffer::Unlock(FRHICommandListBase& RHICmdList)
 
 				// We need to do this on the active command buffer instead of using an upload command buffer. The high level code sometimes reuses the same
 				// buffer in sequences of upload / dispatch, upload / dispatch, so we need to order the copy commands correctly with respect to the dispatches.
-				FVulkanCmdBuffer* Cmd = Context.GetCommandBufferManager()->GetActiveCmdBuffer();
+				FVulkanCommandBuffer* Cmd = Context.GetCommandBufferManager()->GetActiveCmdBuffer();
 				check(Cmd && Cmd->IsOutsideRenderPass());
 				VkCommandBuffer CmdBuffer = Cmd->GetHandle();
 
@@ -403,7 +403,7 @@ void FVulkanResourceMultiBuffer::Unlock(FRHICommandListBase& RHICmdList)
 	LockStatus = ELockStatus::Unlocked;
 }
 
-void FVulkanResourceMultiBuffer::TakeOwnership(FVulkanResourceMultiBuffer& Other)
+void FVulkanBuffer::TakeOwnership(FVulkanBuffer& Other)
 {
 	check(Other.LockStatus == ELockStatus::Unlocked);
 	check(GetParent() == Other.GetParent());
@@ -421,7 +421,7 @@ void FVulkanResourceMultiBuffer::TakeOwnership(FVulkanResourceMultiBuffer& Other
 	Other.CurrentBufferAlloc = {};
 }
 
-void FVulkanResourceMultiBuffer::ReleaseOwnership()
+void FVulkanBuffer::ReleaseOwnership()
 {
 	check(LockStatus == ELockStatus::Unlocked);
 
@@ -441,20 +441,20 @@ FBufferRHIRef FVulkanDynamicRHI::RHICreateBuffer(FRHICommandListBase& RHICmdList
 #else
 	LLM_SCOPE(EnumHasAnyFlags(Desc.Usage, EBufferUsageFlags::VertexBuffer | EBufferUsageFlags::IndexBuffer) ? ELLMTag::Meshes : ELLMTag::RHIMisc);
 #endif
-	return new FVulkanResourceMultiBuffer(Device, Desc, CreateInfo, &RHICmdList);
+	return new FVulkanBuffer(Device, Desc, CreateInfo, &RHICmdList);
 }
 
 void* FVulkanDynamicRHI::LockBuffer_BottomOfPipe(FRHICommandListBase& RHICmdList, FRHIBuffer* BufferRHI, uint32 Offset, uint32 Size, EResourceLockMode LockMode)
 {
 	LLM_SCOPE_VULKAN(ELLMTagVulkan::VulkanBuffers);
-	FVulkanResourceMultiBuffer* Buffer = ResourceCast(BufferRHI);
+	FVulkanBuffer* Buffer = ResourceCast(BufferRHI);
 	return Buffer->Lock(RHICmdList, LockMode, Size, Offset);
 }
 
 void FVulkanDynamicRHI::UnlockBuffer_BottomOfPipe(FRHICommandListBase& RHICmdList, FRHIBuffer* BufferRHI)
 {
 	LLM_SCOPE_VULKAN(ELLMTagVulkan::VulkanBuffers);
-	FVulkanResourceMultiBuffer* Buffer = ResourceCast(BufferRHI);
+	FVulkanBuffer* Buffer = ResourceCast(BufferRHI);
 	Buffer->Unlock(RHICmdList);
 }
 
@@ -462,7 +462,7 @@ void* FVulkanDynamicRHI::RHILockBuffer(FRHICommandListBase& RHICmdList, FRHIBuff
 {
 	QUICK_SCOPE_CYCLE_COUNTER(STAT_LockBuffer_RenderThread);
 	LLM_SCOPE_VULKAN(ELLMTagVulkan::VulkanBuffers);
-	FVulkanResourceMultiBuffer* Buffer = ResourceCast(BufferRHI);
+	FVulkanBuffer* Buffer = ResourceCast(BufferRHI);
 	return Buffer->Lock(RHICmdList, LockMode, Size, Offset);
 }
 
@@ -470,6 +470,6 @@ void FVulkanDynamicRHI::RHIUnlockBuffer(FRHICommandListBase& RHICmdList, FRHIBuf
 {
 	QUICK_SCOPE_CYCLE_COUNTER(STAT_UnlockBuffer_RenderThread);
 	LLM_SCOPE_VULKAN(ELLMTagVulkan::VulkanBuffers);
-	FVulkanResourceMultiBuffer* Buffer = ResourceCast(BufferRHI);
+	FVulkanBuffer* Buffer = ResourceCast(BufferRHI);
 	Buffer->Unlock(RHICmdList);
 }
diff --git a/Engine/Source/Runtime/VulkanRHI/Private/VulkanCommandBuffer.cpp b/Engine/Source/Runtime/VulkanRHI/Private/VulkanCommandBuffer.cpp
--- a/Engine/Source/Runtime/VulkanRHI/Private/VulkanCommandBuffer.cpp
+++ b/Engine/Source/Runtime/VulkanRHI/Private/VulkanCommandBuffer.cpp
@@ -51,7 +51,7 @@ static FAutoConsoleVariableRef CVarVulkanPreventOverlapWithUpload(
 const uint32 GNumberOfFramesBeforeDeletingDescriptorPool = 300;
 extern int32 GVulkanAutoCorrectUnknownLayouts;
 
-FVulkanCmdBuffer::FVulkanCmdBuffer(FVulkanDevice* InDevice, FVulkanCommandBufferPool* InCommandBufferPool, bool bInIsUploadOnly)
+FVulkanCommandBuffer::FVulkanCommandBuffer(FVulkanDevice* InDevice, FVulkanCommandBufferPool* InCommandBufferPool, bool bInIsUploadOnly)
 	: CurrentStencilRef(0)
 	, State(EState::NotAllocated)
 	, bNeedsDynamicStateSet(1)
@@ -82,7 +82,7 @@ FVulkanCmdBuffer::FVulkanCmdBuffer(FVulkanDevice* InDevice, FVulkanCommandBuffer
 	Fence = Device->GetFenceManager().AllocateFence();
 }
 
-void FVulkanCmdBuffer::AllocMemory()
+void FVulkanCommandBuffer::AllocMemory()
 {
 	// Assumes we are inside a lock for the pool
 	check(State == EState::NotAllocated);
@@ -107,7 +107,7 @@ void FVulkanCmdBuffer::AllocMemory()
 	INC_DWORD_STAT(STAT_VulkanNumCmdBuffers);
 }
 
-FVulkanCmdBuffer::~FVulkanCmdBuffer()
+FVulkanCommandBuffer::~FVulkanCommandBuffer()
 {
 	VulkanRHI::FFenceManager& FenceManager = Device->GetFenceManager();
 	if (State == EState::Submitted)
@@ -137,7 +137,7 @@ FVulkanCmdBuffer::~FVulkanCmdBuffer()
 #endif
 }
 
-void FVulkanCmdBuffer::FreeMemory()
+void FVulkanCommandBuffer::FreeMemory()
 {
 	// Assumes we are inside a lock for the pool
 	check(State != EState::NotAllocated);
@@ -149,7 +149,7 @@ void FVulkanCmdBuffer::FreeMemory()
 	State = EState::NotAllocated;
 }
 
-void FVulkanCmdBuffer::BeginUniformUpdateBarrier()
+void FVulkanCommandBuffer::BeginUniformUpdateBarrier()
 {
 	if(!bIsUniformBufferBarrierAdded)
 	{
@@ -162,7 +162,7 @@ void FVulkanCmdBuffer::BeginUniformUpdateBarrier()
 	}
 }
 
-void FVulkanCmdBuffer::EndUniformUpdateBarrier()
+void FVulkanCommandBuffer::EndUniformUpdateBarrier()
 {
 	if(bIsUniformBufferBarrierAdded)
 	{
@@ -174,14 +174,14 @@ void FVulkanCmdBuffer::EndUniformUpdateBarrier()
 		bIsUniformBufferBarrierAdded = false;
 	}
 }
-void FVulkanCmdBuffer::EndRenderPass()
+void FVulkanCommandBuffer::EndRenderPass()
 {
 	checkf(IsInsideRenderPass(), TEXT("Can't EndRP as we're NOT inside one! CmdBuffer 0x%p State=%d"), CommandBufferHandle, (int32)State);
 	VulkanRHI::vkCmdEndRenderPass(CommandBufferHandle);
 	State = EState::IsInsideBegin;
 }
 
-void FVulkanCmdBuffer::BeginRenderPass(const FVulkanRenderTargetLayout& Layout, FVulkanRenderPass* RenderPass, FVulkanFramebuffer* Framebuffer, const VkClearValue* AttachmentClearValues)
+void FVulkanCommandBuffer::BeginRenderPass(const FVulkanRenderTargetLayout& Layout, FVulkanRenderPass* RenderPass, FVulkanFramebuffer* Framebuffer, const VkClearValue* AttachmentClearValues)
 {
 	if (bIsUniformBufferBarrierAdded)
 	{
@@ -219,7 +219,7 @@ void FVulkanCmdBuffer::BeginRenderPass(const FVulkanRenderTargetLayout& Layout,
 }
 
 
-void FVulkanCmdBuffer::AddPendingTimestampQuery(uint64 Index, uint64 Count, VkQueryPool PoolHandle, VkBuffer BufferHandle, bool bBlocking)
+void FVulkanCommandBuffer::AddPendingTimestampQuery(uint64 Index, uint64 Count, VkQueryPool PoolHandle, VkBuffer BufferHandle, bool bBlocking)
 {
 	PendingQuery PQ;
 	PQ.Index = Index;
@@ -230,7 +230,7 @@ void FVulkanCmdBuffer::AddPendingTimestampQuery(uint64 Index, uint64 Count, VkQu
 	PendingTimestampQueries.Add(PQ);
 }
 
-void FVulkanCmdBuffer::End()
+void FVulkanCommandBuffer::End()
 {
 	checkf(IsOutsideRenderPass(), TEXT("Can't End as we're inside a render pass! CmdBuffer 0x%p State=%d"), CommandBufferHandle, (int32)State);
 
@@ -267,7 +267,7 @@ void FVulkanCmdBuffer::End()
 	State = EState::HasEnded;
 }
 
-inline void FVulkanCmdBuffer::InitializeTimings(FVulkanCommandListContext* InContext)
+inline void FVulkanCommandBuffer::InitializeTimings(FVulkanCommandListContext* InContext)
 {
 #if (RHI_NEW_GPU_PROFILER == 0)
 	if ((GVulkanProfileCmdBuffers || GVulkanUseCmdBufferTimingForGPUTime) && !Timing)
@@ -285,7 +285,7 @@ inline void FVulkanCmdBuffer::InitializeTimings(FVulkanCommandListContext* InCon
 #endif
 }
 
-void FVulkanCmdBuffer::AddWaitSemaphore(VkPipelineStageFlags InWaitFlags, TArrayView<VulkanRHI::FSemaphore*> InWaitSemaphores)
+void FVulkanCommandBuffer::AddWaitSemaphore(VkPipelineStageFlags InWaitFlags, TArrayView<VulkanRHI::FSemaphore*> InWaitSemaphores)
 {
 	WaitFlags.Reserve(WaitFlags.Num() + InWaitSemaphores.Num());
 
@@ -299,7 +299,7 @@ void FVulkanCmdBuffer::AddWaitSemaphore(VkPipelineStageFlags InWaitFlags, TArray
 	WaitSemaphores.Append(InWaitSemaphores);
 }
 
-void FVulkanCmdBuffer::Begin()
+void FVulkanCommandBuffer::Begin()
 {
 	{
 		FScopeLock ScopeLock(CommandBufferPool->GetCS());
@@ -343,14 +343,14 @@ void FVulkanCmdBuffer::Begin()
 	bNeedsDynamicStateSet = true;
 }
 
-void FVulkanCmdBuffer::AcquirePoolSetContainer()
+void FVulkanCommandBuffer::AcquirePoolSetContainer()
 {
 	check(!CurrentDescriptorPoolSetContainer);
 	CurrentDescriptorPoolSetContainer = &Device->GetDescriptorPoolsManager().AcquirePoolSetContainer();
 	ensure(TypedDescriptorPoolSets.Num() == 0);
 }
 
-bool FVulkanCmdBuffer::AcquirePoolSetAndDescriptorsIfNeeded(const class FVulkanDescriptorSetsLayout& Layout, bool bNeedDescriptors, VkDescriptorSet* OutDescriptors)
+bool FVulkanCommandBuffer::AcquirePoolSetAndDescriptorsIfNeeded(const class FVulkanDescriptorSetsLayout& Layout, bool bNeedDescriptors, VkDescriptorSet* OutDescriptors)
 {
 	//#todo-rco: This only happens when we call draws outside a render pass...
 	if (!CurrentDescriptorPoolSetContainer)
@@ -375,7 +375,7 @@ bool FVulkanCmdBuffer::AcquirePoolSetAndDescriptorsIfNeeded(const class FVulkanD
 	return false;
 }
 
-void FVulkanCmdBuffer::RefreshFenceStatus()
+void FVulkanCommandBuffer::RefreshFenceStatus()
 {
 	if (State == EState::Submitted)
 	{
@@ -438,13 +438,13 @@ FVulkanCommandBufferPool::~FVulkanCommandBufferPool()
 {
 	for (int32 Index = 0; Index < CmdBuffers.Num(); ++Index)
 	{
-		FVulkanCmdBuffer* CmdBuffer = CmdBuffers[Index];
+		FVulkanCommandBuffer* CmdBuffer = CmdBuffers[Index];
 		delete CmdBuffer;
 	}
 
 	for (int32 Index = 0; Index < FreeCmdBuffers.Num(); ++Index)
 	{
-		FVulkanCmdBuffer* CmdBuffer = FreeCmdBuffers[Index];
+		FVulkanCommandBuffer* CmdBuffer = FreeCmdBuffers[Index];
 		delete CmdBuffer;
 	}
 
@@ -494,7 +494,7 @@ FVulkanCommandBufferManager::~FVulkanCommandBufferManager()
 {
 }
 
-void FVulkanCommandBufferManager::WaitForCmdBuffer(FVulkanCmdBuffer* CmdBuffer, float TimeInSecondsToWait)
+void FVulkanCommandBufferManager::WaitForCmdBuffer(FVulkanCommandBuffer* CmdBuffer, float TimeInSecondsToWait)
 {
 	FScopeLock ScopeLock(&Pool.CS);
 	check(CmdBuffer->IsSubmitted());
@@ -514,16 +514,16 @@ void FVulkanCommandBufferManager::FlushResetQueryPools()
 	FScopeLock ScopeLock(&Pool.CS);
 	if(PoolResets.Num())
 	{
-		FVulkanCmdBuffer* PoolResetCommandBuffer = nullptr;
+		FVulkanCommandBuffer* PoolResetCommandBuffer = nullptr;
 		for (int32 Index = 0; Index < Pool.CmdBuffers.Num(); ++Index)
 		{
-			FVulkanCmdBuffer* CmdBuffer = Pool.CmdBuffers[Index];
+			FVulkanCommandBuffer* CmdBuffer = Pool.CmdBuffers[Index];
 			CmdBuffer->RefreshFenceStatus();
 	#if VULKAN_USE_DIFFERENT_POOL_CMDBUFFERS
 			if (!CmdBuffer->bIsUploadOnly)
 	#endif
 			{
-				if (CmdBuffer->State == FVulkanCmdBuffer::EState::ReadyForBegin)
+				if (CmdBuffer->State == FVulkanCommandBuffer::EState::ReadyForBegin)
 				{
 					PoolResetCommandBuffer = CmdBuffer;
 				}
@@ -698,12 +698,12 @@ void FVulkanCommandBufferManager::SubmitActiveCmdBufferFromPresent(VulkanRHI::FS
 	}
 }
 
-FVulkanCmdBuffer* FVulkanCommandBufferPool::Create(bool bIsUploadOnly)
+FVulkanCommandBuffer* FVulkanCommandBufferPool::Create(bool bIsUploadOnly)
 {
 	// Assumes we are inside a lock for the pool
 	for (int32 Index = FreeCmdBuffers.Num() - 1; Index >= 0; --Index)
 	{
-		FVulkanCmdBuffer* CmdBuffer = FreeCmdBuffers[Index];
+		FVulkanCommandBuffer* CmdBuffer = FreeCmdBuffers[Index];
 #if VULKAN_USE_DIFFERENT_POOL_CMDBUFFERS
 		if (CmdBuffer->bIsUploadOnly == bIsUploadOnly)
 #endif
@@ -715,18 +715,18 @@ FVulkanCmdBuffer* FVulkanCommandBufferPool::Create(bool bIsUploadOnly)
 		}
 	}
 
-	FVulkanCmdBuffer* CmdBuffer = new FVulkanCmdBuffer(Device, this, bIsUploadOnly);
+	FVulkanCommandBuffer* CmdBuffer = new FVulkanCommandBuffer(Device, this, bIsUploadOnly);
 	CmdBuffers.Add(CmdBuffer);
 	check(CmdBuffer);
 	return CmdBuffer;
 }
 
-void FVulkanCommandBufferPool::RefreshFenceStatus(FVulkanCmdBuffer* SkipCmdBuffer)
+void FVulkanCommandBufferPool::RefreshFenceStatus(FVulkanCommandBuffer* SkipCmdBuffer)
 {
 	FScopeLock ScopeLock(&CS);
 	for (int32 Index = 0; Index < CmdBuffers.Num(); ++Index)
 	{
-		FVulkanCmdBuffer* CmdBuffer = CmdBuffers[Index];
+		FVulkanCommandBuffer* CmdBuffer = CmdBuffers[Index];
 		if (CmdBuffer != SkipCmdBuffer)
 		{
 			CmdBuffer->RefreshFenceStatus();
@@ -748,13 +748,13 @@ void FVulkanCommandBufferManager::PrepareForNewActiveCommandBuffer()
 
 	for (int32 Index = 0; Index < Pool.CmdBuffers.Num(); ++Index)
 	{
-		FVulkanCmdBuffer* CmdBuffer = Pool.CmdBuffers[Index];
+		FVulkanCommandBuffer* CmdBuffer = Pool.CmdBuffers[Index];
 		CmdBuffer->RefreshFenceStatus();
 #if VULKAN_USE_DIFFERENT_POOL_CMDBUFFERS
 		if (!CmdBuffer->bIsUploadOnly)
 #endif
 		{
-			if (CmdBuffer->State == FVulkanCmdBuffer::EState::ReadyForBegin || CmdBuffer->State == FVulkanCmdBuffer::EState::NeedReset)
+			if (CmdBuffer->State == FVulkanCommandBuffer::EState::ReadyForBegin || CmdBuffer->State == FVulkanCommandBuffer::EState::NeedReset)
 			{
 				ActiveCmdBuffer = CmdBuffer;
 				ActiveCmdBuffer->Begin();
@@ -762,7 +762,7 @@ void FVulkanCommandBufferManager::PrepareForNewActiveCommandBuffer()
 			}
 			else
 			{
-				check(CmdBuffer->State == FVulkanCmdBuffer::EState::Submitted);
+				check(CmdBuffer->State == FVulkanCommandBuffer::EState::Submitted);
 			}
 		}
 	}
@@ -778,7 +778,7 @@ uint32 FVulkanCommandBufferManager::CalculateGPUTime()
 	uint32 Time = 0;
 	for (int32 Index = 0; Index < Pool.CmdBuffers.Num(); ++Index)
 	{
-		FVulkanCmdBuffer* CmdBuffer = Pool.CmdBuffers[Index];
+		FVulkanCommandBuffer* CmdBuffer = Pool.CmdBuffers[Index];
 		if (CmdBuffer->HasValidTiming())
 		{
 			Time += CmdBuffer->Timing->GetTiming(false);
@@ -788,7 +788,7 @@ uint32 FVulkanCommandBufferManager::CalculateGPUTime()
 }
 #endif // (RHI_NEW_GPU_PROFILER == 0)
 
-FVulkanCmdBuffer* FVulkanCommandBufferManager::GetUploadCmdBuffer()
+FVulkanCommandBuffer* FVulkanCommandBufferManager::GetUploadCmdBuffer()
 {
 	FScopeLock ScopeLock(&Pool.CS);
 	if (!UploadCmdBuffer)
@@ -802,13 +802,13 @@ FVulkanCmdBuffer* FVulkanCommandBufferManager::GetUploadCmdBuffer()
 
 		for (int32 Index = 0; Index < Pool.CmdBuffers.Num(); ++Index)
 		{
-			FVulkanCmdBuffer* CmdBuffer = Pool.CmdBuffers[Index];
+			FVulkanCommandBuffer* CmdBuffer = Pool.CmdBuffers[Index];
 			CmdBuffer->RefreshFenceStatus();
 #if VULKAN_USE_DIFFERENT_POOL_CMDBUFFERS
 			if (CmdBuffer->bIsUploadOnly)
 #endif
 			{
-				if (CmdBuffer->State == FVulkanCmdBuffer::EState::ReadyForBegin || CmdBuffer->State == FVulkanCmdBuffer::EState::NeedReset)
+				if (CmdBuffer->State == FVulkanCommandBuffer::EState::ReadyForBegin || CmdBuffer->State == FVulkanCommandBuffer::EState::NeedReset)
 				{
 					UploadCmdBuffer = CmdBuffer;
 					UploadCmdBuffer->Begin();
@@ -861,7 +861,7 @@ void FVulkanCommandBufferPool::FreeUnusedCmdBuffers(FVulkanQueue* InQueue, bool
 	const double CurrentTime = FPlatformTime::Seconds();
 
 	// In case Queue stores pointer to a cmdbuffer, do not delete it
-	FVulkanCmdBuffer* LastSubmittedCmdBuffer = nullptr;
+	FVulkanCommandBuffer* LastSubmittedCmdBuffer = nullptr;
 	uint64 LastSubmittedFenceCounter = 0;
 	InQueue->GetLastSubmittedInfo(LastSubmittedCmdBuffer, LastSubmittedFenceCounter);
 
@@ -870,9 +870,9 @@ void FVulkanCommandBufferPool::FreeUnusedCmdBuffers(FVulkanQueue* InQueue, bool
 
 	for (int32 Index = CmdBuffers.Num() - 1; Index >= 0; --Index)
 	{
-		FVulkanCmdBuffer* CmdBuffer = CmdBuffers[Index];
+		FVulkanCommandBuffer* CmdBuffer = CmdBuffers[Index];
 		if (CmdBuffer != LastSubmittedCmdBuffer &&
-			(CmdBuffer->State == FVulkanCmdBuffer::EState::ReadyForBegin || CmdBuffer->State == FVulkanCmdBuffer::EState::NeedReset) &&
+			(CmdBuffer->State == FVulkanCommandBuffer::EState::ReadyForBegin || CmdBuffer->State == FVulkanCommandBuffer::EState::NeedReset) &&
 			((CurrentTime - CmdBuffer->SubmittedTime) > CMD_BUFFER_TIME_TO_WAIT_BEFORE_DELETING))
 		{
 			// Skip command buffer that contain unresolved compaction queries
diff --git a/Engine/Source/Runtime/VulkanRHI/Private/VulkanCommandBuffer.h b/Engine/Source/Runtime/VulkanRHI/Private/VulkanCommandBuffer.h
--- a/Engine/Source/Runtime/VulkanRHI/Private/VulkanCommandBuffer.h
+++ b/Engine/Source/Runtime/VulkanRHI/Private/VulkanCommandBuffer.h
@@ -27,15 +27,15 @@ namespace VulkanRHI
 	class FSemaphore;
 }
 
-class FVulkanCmdBuffer
+class FVulkanCommandBuffer
 {
 protected:
 	friend class FVulkanCommandBufferManager;
 	friend class FVulkanCommandBufferPool;
 	friend class FVulkanQueue;
 
-	FVulkanCmdBuffer(FVulkanDevice* InDevice, FVulkanCommandBufferPool* InCommandBufferPool, bool bInIsUploadOnly);
-	~FVulkanCmdBuffer();
+	FVulkanCommandBuffer(FVulkanDevice* InDevice, FVulkanCommandBufferPool* InCommandBufferPool, bool bInIsUploadOnly);
+	~FVulkanCommandBuffer();
 
 public:
 	FVulkanCommandBufferPool* GetOwner()
@@ -269,7 +269,7 @@ public:
 	FVulkanCommandBufferPool(FVulkanDevice* InDevice, FVulkanCommandBufferManager& InMgr);
 	~FVulkanCommandBufferPool();
 
-	void RefreshFenceStatus(FVulkanCmdBuffer* SkipCmdBuffer = nullptr);
+	void RefreshFenceStatus(FVulkanCommandBuffer* SkipCmdBuffer = nullptr);
 
 	inline VkCommandPool GetHandle() const
 	{
@@ -291,15 +291,15 @@ public:
 private:
 	VkCommandPool Handle;
 
-	TArray<FVulkanCmdBuffer*> CmdBuffers;
-	TArray<FVulkanCmdBuffer*> FreeCmdBuffers;
+	TArray<FVulkanCommandBuffer*> CmdBuffers;
+	TArray<FVulkanCommandBuffer*> FreeCmdBuffers;
 
 	FCriticalSection CS;
 	FVulkanDevice* Device;
 
 	FVulkanCommandBufferManager& Mgr;
 
-	FVulkanCmdBuffer* Create(bool bIsUploadOnly);
+	FVulkanCommandBuffer* Create(bool bIsUploadOnly);
 
 
 	void Create(uint32 QueueFamilyIndex);
@@ -314,7 +314,7 @@ public:
 
 	void Init(FVulkanCommandListContext* InContext);
 
-	inline FVulkanCmdBuffer* GetActiveCmdBuffer()
+	inline FVulkanCommandBuffer* GetActiveCmdBuffer()
 	{
 		if (UploadCmdBuffer)
 		{
@@ -324,7 +324,7 @@ public:
 		return ActiveCmdBuffer;
 	}
 
-	inline FVulkanCmdBuffer* GetActiveCmdBufferDirect()
+	inline FVulkanCommandBuffer* GetActiveCmdBufferDirect()
 	{
 		return ActiveCmdBuffer;
 	}
@@ -340,7 +340,7 @@ public:
 		return ActiveCmdBuffer != nullptr;
 	}
 
-	VULKANRHI_API FVulkanCmdBuffer* GetUploadCmdBuffer();
+	VULKANRHI_API FVulkanCommandBuffer* GetUploadCmdBuffer();
 
 	VULKANRHI_API void SubmitUploadCmdBuffer(uint32 NumSignalSemaphores = 0, VkSemaphore* SignalSemaphores = nullptr);
 
@@ -359,14 +359,14 @@ public:
 	/** Regular SACB() expects not-ended and would rotate the command buffer immediately, but Present has a special logic */
 	void SubmitActiveCmdBufferFromPresent(VulkanRHI::FSemaphore* SignalSemaphore = nullptr);
 
-	void WaitForCmdBuffer(FVulkanCmdBuffer* CmdBuffer, float TimeInSecondsToWait = 10.0f);
+	void WaitForCmdBuffer(FVulkanCommandBuffer* CmdBuffer, float TimeInSecondsToWait = 10.0f);
 
 
 	void AddQueryPoolForReset(VkQueryPool Pool, uint32 Size);
 	void FlushResetQueryPools();
 
 	// Update the fences of all cmd buffers except SkipCmdBuffer
-	void RefreshFenceStatus(FVulkanCmdBuffer* SkipCmdBuffer = nullptr)
+	void RefreshFenceStatus(FVulkanCommandBuffer* SkipCmdBuffer = nullptr)
 	{
 		Pool.RefreshFenceStatus(SkipCmdBuffer);
 	}
@@ -417,8 +417,8 @@ private:
 	FVulkanCommandListContext* Context;
 	FVulkanCommandBufferPool Pool;
 	FVulkanQueue* Queue;
-	FVulkanCmdBuffer* ActiveCmdBuffer;
-	FVulkanCmdBuffer* UploadCmdBuffer;
+	FVulkanCommandBuffer* ActiveCmdBuffer;
+	FVulkanCommandBuffer* UploadCmdBuffer;
 	TArray<FQueryPoolReset> PoolResets;
 
 	/** This semaphore is used to prevent overlaps between the (current) graphics cmdbuf and next upload cmdbuf. */
diff --git a/Engine/Source/Runtime/VulkanRHI/Private/VulkanCommands.cpp b/Engine/Source/Runtime/VulkanRHI/Private/VulkanCommands.cpp
--- a/Engine/Source/Runtime/VulkanRHI/Private/VulkanCommands.cpp
+++ b/Engine/Source/Runtime/VulkanRHI/Private/VulkanCommands.cpp
@@ -129,7 +129,7 @@ static FORCEINLINE ShaderStage::EStage GetAndVerifyShaderStageAndVulkanShader(FR
 
 void FVulkanCommandListContext::RHISetStreamSource(uint32 StreamIndex, FRHIBuffer* VertexBufferRHI, uint32 Offset)
 {
-	FVulkanResourceMultiBuffer* VertexBuffer = ResourceCast(VertexBufferRHI);
+	FVulkanBuffer* VertexBuffer = ResourceCast(VertexBufferRHI);
 	if (VertexBuffer != nullptr)
 	{
 		PendingGfxState->SetStreamSource(StreamIndex, VertexBuffer->GetHandle(), Offset + VertexBuffer->GetOffset());
@@ -298,7 +298,7 @@ void FVulkanCommandListContext::RHIDispatchComputeShader(uint32 ThreadGroupCount
 
 	CommitComputeResourceTables();
 
-	FVulkanCmdBuffer* Cmd = CommandBufferManager->GetActiveCmdBuffer();
+	FVulkanCommandBuffer* Cmd = CommandBufferManager->GetActiveCmdBuffer();
 	ensure(Cmd->IsOutsideRenderPass());
 	VkCommandBuffer CmdBuffer = Cmd->GetHandle();
 	PendingComputeState->PrepareForDispatch(Cmd);
@@ -321,11 +321,11 @@ void FVulkanCommandListContext::RHIDispatchComputeShader(uint32 ThreadGroupCount
 void FVulkanCommandListContext::RHIDispatchIndirectComputeShader(FRHIBuffer* ArgumentBufferRHI, uint32 ArgumentOffset)
 {
 	static_assert(sizeof(FRHIDispatchIndirectParameters) == sizeof(VkDispatchIndirectCommand), "Dispatch indirect doesn't match!");
-	FVulkanResourceMultiBuffer* ArgumentBuffer = ResourceCast(ArgumentBufferRHI);
+	FVulkanBuffer* ArgumentBuffer = ResourceCast(ArgumentBufferRHI);
 
 	CommitComputeResourceTables();
 
-	FVulkanCmdBuffer* Cmd = CommandBufferManager->GetActiveCmdBuffer();
+	FVulkanCommandBuffer* Cmd = CommandBufferManager->GetActiveCmdBuffer();
 	ensure(Cmd->IsOutsideRenderPass());
 	VkCommandBuffer CmdBuffer = Cmd->GetHandle();
 	PendingComputeState->PrepareForDispatch(Cmd);
@@ -644,7 +644,7 @@ void FVulkanCommandListContext::RHIDrawPrimitive(uint32 BaseVertexIndex, uint32
 
 	CommitGraphicsResourceTables();
 
-	FVulkanCmdBuffer* CmdBuffer = CommandBufferManager->GetActiveCmdBuffer();
+	FVulkanCommandBuffer* CmdBuffer = CommandBufferManager->GetActiveCmdBuffer();
 	PendingGfxState->PrepareForDraw(CmdBuffer);
 	uint32 NumVertices = GetVertexCountForPrimitiveCount(NumPrimitives, PendingGfxState->PrimitiveType);
 	VulkanRHI::vkCmdDraw(CmdBuffer->GetHandle(), NumVertices, NumInstances, BaseVertexIndex, 0);
@@ -666,11 +666,11 @@ void FVulkanCommandListContext::RHIDrawPrimitiveIndirect(FRHIBuffer* ArgumentBuf
 
 	CommitGraphicsResourceTables();
 
-	FVulkanCmdBuffer* Cmd = CommandBufferManager->GetActiveCmdBuffer();
+	FVulkanCommandBuffer* Cmd = CommandBufferManager->GetActiveCmdBuffer();
 	VkCommandBuffer CmdBuffer = Cmd->GetHandle();
 	PendingGfxState->PrepareForDraw(Cmd);
 
-	FVulkanResourceMultiBuffer* ArgumentBuffer = ResourceCast(ArgumentBufferRHI);
+	FVulkanBuffer* ArgumentBuffer = ResourceCast(ArgumentBufferRHI);
 
 
 	VulkanRHI::vkCmdDrawIndirect(CmdBuffer, ArgumentBuffer->GetHandle(), ArgumentBuffer->GetOffset() + ArgumentOffset, 1, sizeof(VkDrawIndirectCommand));
@@ -693,8 +693,8 @@ void FVulkanCommandListContext::RHIDrawIndexedPrimitive(FRHIBuffer* IndexBufferR
 
 	CommitGraphicsResourceTables();
 
-	FVulkanResourceMultiBuffer* IndexBuffer = ResourceCast(IndexBufferRHI);
-	FVulkanCmdBuffer* Cmd = CommandBufferManager->GetActiveCmdBuffer();
+	FVulkanBuffer* IndexBuffer = ResourceCast(IndexBufferRHI);
+	FVulkanCommandBuffer* Cmd = CommandBufferManager->GetActiveCmdBuffer();
 	VkCommandBuffer CmdBuffer = Cmd->GetHandle();
 	PendingGfxState->PrepareForDraw(Cmd);
 	VulkanRHI::vkCmdBindIndexBuffer(CmdBuffer, IndexBuffer->GetHandle(), IndexBuffer->GetOffset(), IndexBuffer->GetIndexType());
@@ -717,13 +717,13 @@ void FVulkanCommandListContext::RHIDrawIndexedIndirect(FRHIBuffer* IndexBufferRH
 
 	CommitGraphicsResourceTables();
 
-	FVulkanResourceMultiBuffer* IndexBuffer = ResourceCast(IndexBufferRHI);
-	FVulkanCmdBuffer* Cmd = CommandBufferManager->GetActiveCmdBuffer();
+	FVulkanBuffer* IndexBuffer = ResourceCast(IndexBufferRHI);
+	FVulkanCommandBuffer* Cmd = CommandBufferManager->GetActiveCmdBuffer();
 	VkCommandBuffer CmdBuffer = Cmd->GetHandle();
 	PendingGfxState->PrepareForDraw(Cmd);
 	VulkanRHI::vkCmdBindIndexBuffer(CmdBuffer, IndexBuffer->GetHandle(), IndexBuffer->GetOffset(), IndexBuffer->GetIndexType());
 
-	FVulkanResourceMultiBuffer* ArgumentBuffer = ResourceCast(ArgumentsBufferRHI);
+	FVulkanBuffer* ArgumentBuffer = ResourceCast(ArgumentsBufferRHI);
 	VkDeviceSize ArgumentOffset = DrawArgumentsIndex * sizeof(VkDrawIndexedIndirectCommand);
 
 
@@ -744,13 +744,13 @@ void FVulkanCommandListContext::RHIDrawIndexedPrimitiveIndirect(FRHIBuffer* Inde
 
 	CommitGraphicsResourceTables();
 
-	FVulkanResourceMultiBuffer* IndexBuffer = ResourceCast(IndexBufferRHI);
-	FVulkanCmdBuffer* Cmd = CommandBufferManager->GetActiveCmdBuffer();
+	FVulkanBuffer* IndexBuffer = ResourceCast(IndexBufferRHI);
+	FVulkanCommandBuffer* Cmd = CommandBufferManager->GetActiveCmdBuffer();
 	VkCommandBuffer CmdBuffer = Cmd->GetHandle();
 	PendingGfxState->PrepareForDraw(Cmd);
 	VulkanRHI::vkCmdBindIndexBuffer(CmdBuffer, IndexBuffer->GetHandle(), IndexBuffer->GetOffset(), IndexBuffer->GetIndexType());
 
-	FVulkanResourceMultiBuffer* ArgumentBuffer = ResourceCast(ArgumentBufferRHI);
+	FVulkanBuffer* ArgumentBuffer = ResourceCast(ArgumentBufferRHI);
 
 
 	VulkanRHI::vkCmdDrawIndexedIndirect(CmdBuffer, ArgumentBuffer->GetHandle(), ArgumentBuffer->GetOffset() + ArgumentOffset, 1, sizeof(VkDrawIndexedIndirectCommand));
@@ -766,7 +766,7 @@ void FVulkanCommandListContext::RHIDispatchMeshShader(uint32 ThreadGroupCountX,
 {
 	CommitGraphicsResourceTables();
 
-	FVulkanCmdBuffer* Cmd = CommandBufferManager->GetActiveCmdBuffer();
+	FVulkanCommandBuffer* Cmd = CommandBufferManager->GetActiveCmdBuffer();
 
 	PendingGfxState->PrepareForDraw(Cmd);
 
@@ -782,11 +782,11 @@ void FVulkanCommandListContext::RHIDispatchIndirectMeshShader(FRHIBuffer* Argume
 {
 	CommitGraphicsResourceTables();
 
-	FVulkanCmdBuffer* Cmd = CommandBufferManager->GetActiveCmdBuffer();
+	FVulkanCommandBuffer* Cmd = CommandBufferManager->GetActiveCmdBuffer();
 
 	PendingGfxState->PrepareForDraw(Cmd);
 
-	FVulkanResourceMultiBuffer* ArgumentBuffer = ResourceCast(ArgumentBufferRHI);
+	FVulkanBuffer* ArgumentBuffer = ResourceCast(ArgumentBufferRHI);
 
 	VulkanRHI::vkCmdDrawMeshTasksIndirectEXT(Cmd->GetHandle(), ArgumentBuffer->GetHandle(), ArgumentBuffer->GetOffset() + ArgumentOffset, 1, sizeof(VkDrawMeshTasksIndirectCommandEXT));
 
@@ -806,7 +806,7 @@ void FVulkanCommandListContext::RHIClearMRT(bool bClearColor, int32 NumClearColo
 
 	check(bClearColor ? NumClearColors > 0 : true);
 
-	FVulkanCmdBuffer* CmdBuffer = CommandBufferManager->GetActiveCmdBuffer();
+	FVulkanCommandBuffer* CmdBuffer = CommandBufferManager->GetActiveCmdBuffer();
 	//FRCLog::Printf(TEXT("RHIClearMRT"));
 
 	const uint32 NumColorAttachments = CurrentFramebuffer->GetNumColorAttachments();
@@ -814,7 +814,7 @@ void FVulkanCommandListContext::RHIClearMRT(bool bClearColor, int32 NumClearColo
 	InternalClearMRT(CmdBuffer, bClearColor, bClearColor ? NumClearColors : 0, ClearColorArray, bClearDepth, Depth, bClearStencil, Stencil);
 }
 
-void FVulkanCommandListContext::InternalClearMRT(FVulkanCmdBuffer* CmdBuffer, bool bClearColor, int32 NumClearColors, const FLinearColor* ClearColorArray, bool bClearDepth, float Depth, bool bClearStencil, uint32 Stencil)
+void FVulkanCommandListContext::InternalClearMRT(FVulkanCommandBuffer* CmdBuffer, bool bClearColor, int32 NumClearColors, const FLinearColor* ClearColorArray, bool bClearDepth, float Depth, bool bClearStencil, uint32 Stencil)
 {
 	if (CurrentRenderPass)
 	{
@@ -888,7 +888,7 @@ uint32 FVulkanDynamicRHI::RHIGetGPUFrameCycles(uint32 GPUIndex)
 
 void FVulkanCommandListContext::RHISetDepthBounds(float MinDepth, float MaxDepth)
 {
-	FVulkanCmdBuffer* CmdBuffer = CommandBufferManager->GetActiveCmdBuffer();
+	FVulkanCommandBuffer* CmdBuffer = CommandBufferManager->GetActiveCmdBuffer();
 	VulkanRHI::vkCmdSetDepthBounds(CmdBuffer->GetHandle(), MinDepth, MaxDepth);
 }
 
@@ -919,7 +919,7 @@ void FVulkanCommandListContext::InternalSubmitActiveCmdBuffer()
 void FVulkanCommandListContext::PrepareForCPURead()
 {
 	ensure(IsImmediate());
-	FVulkanCmdBuffer* CmdBuffer = CommandBufferManager->GetActiveCmdBuffer();
+	FVulkanCommandBuffer* CmdBuffer = CommandBufferManager->GetActiveCmdBuffer();
 	if (CmdBuffer && CmdBuffer->HasBegun())
 	{
 		check(!CmdBuffer->IsInsideRenderPass());
@@ -935,8 +935,8 @@ void FVulkanCommandListContext::PrepareForCPURead()
 
 void FVulkanCommandListContext::RHICopyToStagingBuffer(FRHIBuffer* SourceBufferRHI, FRHIStagingBuffer* StagingBufferRHI, uint32 Offset, uint32 NumBytes)
 {
-	FVulkanCmdBuffer* CmdBuffer = CommandBufferManager->GetActiveCmdBuffer();
-	FVulkanResourceMultiBuffer* VertexBuffer = ResourceCast(SourceBufferRHI);
+	FVulkanCommandBuffer* CmdBuffer = CommandBufferManager->GetActiveCmdBuffer();
+	FVulkanBuffer* VertexBuffer = ResourceCast(SourceBufferRHI);
 
 	ensure(CmdBuffer->IsOutsideRenderPass());
 
@@ -965,7 +965,7 @@ void FVulkanCommandListContext::RHICopyToStagingBuffer(FRHIBuffer* SourceBufferR
 
 void FVulkanCommandListContext::RHIWriteGPUFence(FRHIGPUFence* FenceRHI)
 {
-	FVulkanCmdBuffer* CmdBuffer = CommandBufferManager->GetActiveCmdBuffer();
+	FVulkanCommandBuffer* CmdBuffer = CommandBufferManager->GetActiveCmdBuffer();
 	FVulkanGPUFence* Fence = ResourceCast(FenceRHI);
 
 	Fence->CmdBuffer = CmdBuffer;
@@ -994,7 +994,7 @@ IRHIComputeContext* FVulkanDynamicRHI::RHIGetCommandContext(ERHIPipeline Pipelin
 	FVulkanCommandListContext* CmdContext = Device->AcquireDeferredContext();
 
 	FVulkanCommandBufferManager* CmdMgr = CmdContext->GetCommandBufferManager();
-	FVulkanCmdBuffer* CmdBuffer = CmdMgr->GetActiveCmdBuffer();
+	FVulkanCommandBuffer* CmdBuffer = CmdMgr->GetActiveCmdBuffer();
 	if (!CmdBuffer)
 	{
 		CmdMgr->PrepareForNewActiveCommandBuffer();
@@ -1035,7 +1035,7 @@ void FVulkanDynamicRHI::RHISubmitCommandLists(FRHISubmitCommandListsArgs&& Args)
 		if (CurrentContext->IsImmediate())
 		{
 			CurrentContext->RequestSubmitCurrentCommands();
-			FVulkanCmdBuffer* CmdBuffer = CmdBufMgr->GetActiveCmdBuffer();
+			FVulkanCommandBuffer* CmdBuffer = CmdBufMgr->GetActiveCmdBuffer();
 			if (CmdBuffer && CmdBuffer->HasBegun() && CmdBuffer->IsOutsideRenderPass())
 			{
 				CurrentContext->SafePointSubmit();
@@ -1045,7 +1045,7 @@ void FVulkanDynamicRHI::RHISubmitCommandLists(FRHISubmitCommandListsArgs&& Args)
 		else
 		{
 			check(!CmdBufMgr->HasPendingUploadCmdBuffer());  // todo-jn
-			FVulkanCmdBuffer* CmdBuffer = CmdBufMgr->GetActiveCmdBuffer();
+			FVulkanCommandBuffer* CmdBuffer = CmdBufMgr->GetActiveCmdBuffer();
 			check(!CmdBuffer->IsInsideRenderPass());
 			CmdBufMgr->SubmitActiveCmdBuffer();
 
diff --git a/Engine/Source/Runtime/VulkanRHI/Private/VulkanContext.h b/Engine/Source/Runtime/VulkanRHI/Private/VulkanContext.h
--- a/Engine/Source/Runtime/VulkanRHI/Private/VulkanContext.h
+++ b/Engine/Source/Runtime/VulkanRHI/Private/VulkanContext.h
@@ -182,8 +182,8 @@ public:
 		return Queue;
 	}
 
-	void WriteBeginTimestamp(FVulkanCmdBuffer* CmdBuffer);
-	void WriteEndTimestamp(FVulkanCmdBuffer* CmdBuffer);
+	void WriteBeginTimestamp(FVulkanCommandBuffer* CmdBuffer);
+	void WriteEndTimestamp(FVulkanCommandBuffer* CmdBuffer);
 
 #if (RHI_NEW_GPU_PROFILER == 0)
 	void ReadAndCalculateGPUFrameTime();
@@ -204,7 +204,7 @@ public:
 		// Nothing to do
 	}
 
-	void EndRenderQueryInternal(FVulkanCmdBuffer* CmdBuffer, FVulkanRenderQuery* Query);
+	void EndRenderQueryInternal(FVulkanCommandBuffer* CmdBuffer, FVulkanRenderQuery* Query);
 
 	void ReleasePendingState();
 
@@ -216,8 +216,8 @@ protected:
 	bool bSubmitAtNextSafePoint;
 	bool bUniformBufferUploadRenderPassDirty = true;
 
-	void BeginOcclusionQueryBatch(FVulkanCmdBuffer* CmdBuffer, uint32 NumQueriesInBatch);
-	void EndOcclusionQueryBatch(FVulkanCmdBuffer* CmdBuffer);
+	void BeginOcclusionQueryBatch(FVulkanCommandBuffer* CmdBuffer, uint32 NumQueriesInBatch);
+	void EndOcclusionQueryBatch(FVulkanCommandBuffer* CmdBuffer);
 
 	VulkanRHI::FTempBlockAllocator* TempBlockAllocator = nullptr;
 
@@ -245,7 +245,7 @@ protected:
 	void PrepareForCPURead();
 	void RequestSubmitCurrentCommands();
 
-	void InternalClearMRT(FVulkanCmdBuffer* CmdBuffer, bool bClearColor, int32 NumClearColors, const FLinearColor* ColorArray, bool bClearDepth, float Depth, bool bClearStencil, uint32 Stencil);
+	void InternalClearMRT(FVulkanCommandBuffer* CmdBuffer, bool bClearColor, int32 NumClearColors, const FLinearColor* ColorArray, bool bClearDepth, float Depth, bool bClearStencil, uint32 Stencil);
 
 public:
 	bool IsSwapchainImage(FRHITexture* InTexture) const;
diff --git a/Engine/Source/Runtime/VulkanRHI/Private/VulkanDescriptorSets.cpp b/Engine/Source/Runtime/VulkanRHI/Private/VulkanDescriptorSets.cpp
--- a/Engine/Source/Runtime/VulkanRHI/Private/VulkanDescriptorSets.cpp
+++ b/Engine/Source/Runtime/VulkanRHI/Private/VulkanDescriptorSets.cpp
@@ -606,7 +606,7 @@ void FVulkanBindlessDescriptorManager::BindDescriptorBuffers(VkCommandBuffer Com
 	}
 }
 
-void FVulkanBindlessDescriptorManager::RegisterUniformBuffers(FVulkanCmdBuffer* CommandBuffer, VkPipelineBindPoint BindPoint, const FUniformBufferDescriptorArrays& StageUBs)
+void FVulkanBindlessDescriptorManager::RegisterUniformBuffers(FVulkanCommandBuffer* CommandBuffer, VkPipelineBindPoint BindPoint, const FUniformBufferDescriptorArrays& StageUBs)
 {
 	checkf(bIsSupported, TEXT("Trying to RegisterUniformBuffers but bindless is not supported!"));
 
@@ -717,7 +717,7 @@ void FVulkanBindlessDescriptorManager::UpdateDescriptor(FRHIDescriptorHandle Des
 	else
 	{
 		check(!IsInRenderingThread() || FRHICommandListExecutor::GetImmediateCommandList().Bypass() || !IsRunningRHIInSeparateThread());
-		FVulkanCmdBuffer* CmdBuffer = Device->GetImmediateContext().GetCommandBufferManager()->GetActiveCmdBuffer();
+		FVulkanCommandBuffer* CmdBuffer = Device->GetImmediateContext().GetCommandBufferManager()->GetActiveCmdBuffer();
 
 		// :todo-jn:  Hack to avoid barriers/copies in renderpasses
 		if (CmdBuffer->IsInsideRenderPass())
diff --git a/Engine/Source/Runtime/VulkanRHI/Private/VulkanDescriptorSets.h b/Engine/Source/Runtime/VulkanRHI/Private/VulkanDescriptorSets.h
--- a/Engine/Source/Runtime/VulkanRHI/Private/VulkanDescriptorSets.h
+++ b/Engine/Source/Runtime/VulkanRHI/Private/VulkanDescriptorSets.h
@@ -342,7 +342,7 @@ class FVulkanTypedDescriptorPoolSet
 
 protected:
 	friend class FVulkanDescriptorPoolSetContainer;
-	friend class FVulkanCmdBuffer;
+	friend class FVulkanCommandBuffer;
 
 	FVulkanTypedDescriptorPoolSet(FVulkanDevice* InDevice, const FVulkanDescriptorSetsLayout& InLayout)
 		: Device(InDevice)
@@ -879,7 +879,7 @@ public:
 	void UpdateTexelBuffer(FRHIDescriptorHandle DescriptorHandle, const VkBufferViewCreateInfo& ViewInfo, bool bImmediateUpdate = true);
 	void UpdateAccelerationStructure(FRHIDescriptorHandle DescriptorHandle, VkAccelerationStructureKHR AccelerationStructure, bool bImmediateUpdate = true);
 
-	void RegisterUniformBuffers(FVulkanCmdBuffer* CommandBuffer, VkPipelineBindPoint BindPoint, const FUniformBufferDescriptorArrays& StageUBs);
+	void RegisterUniformBuffers(FVulkanCommandBuffer* CommandBuffer, VkPipelineBindPoint BindPoint, const FUniformBufferDescriptorArrays& StageUBs);
 
 	void Unregister(FRHIDescriptorHandle DescriptorHandle);
 
diff --git a/Engine/Source/Runtime/VulkanRHI/Private/VulkanDevice.h b/Engine/Source/Runtime/VulkanRHI/Private/VulkanDevice.h
--- a/Engine/Source/Runtime/VulkanRHI/Private/VulkanDevice.h
+++ b/Engine/Source/Runtime/VulkanRHI/Private/VulkanDevice.h
@@ -225,7 +225,7 @@ namespace VulkanRHI
 			ReleaseResources(true);
 		}
 
-		void OnCmdBufferDeleted(FVulkanCmdBuffer* CmdBuffer);
+		void OnCmdBufferDeleted(FVulkanCommandBuffer* CmdBuffer);
 	private:
 		void EnqueueGenericResource(EType Type, uint64 Handle);
 
@@ -234,7 +234,7 @@ namespace VulkanRHI
 			EType StructureType;
 			uint32 FrameNumber;
 			uint64 FenceCounter;
-			FVulkanCmdBuffer* CmdBuffer;
+			FVulkanCommandBuffer* CmdBuffer;
 
 			uint64 Handle;
 			FVulkanAllocation Allocation;
diff --git a/Engine/Source/Runtime/VulkanRHI/Private/VulkanGPUProfiler.h b/Engine/Source/Runtime/VulkanRHI/Private/VulkanGPUProfiler.h
--- a/Engine/Source/Runtime/VulkanRHI/Private/VulkanGPUProfiler.h
+++ b/Engine/Source/Runtime/VulkanRHI/Private/VulkanGPUProfiler.h
@@ -10,7 +10,7 @@
 #include "GPUProfiler.h"
 #include "VulkanConfiguration.h"
 
-class FVulkanCmdBuffer;
+class FVulkanCommandBuffer;
 class FVulkanCommandListContext;
 class FVulkanDevice;
 class FVulkanTimingQueryPool;
@@ -33,13 +33,13 @@ public:
 	/**
 	 * Start a GPU timing measurement.
 	 */
-	void StartTiming(FVulkanCmdBuffer* CmdBuffer = nullptr);
+	void StartTiming(FVulkanCommandBuffer* CmdBuffer = nullptr);
 
 	/**
 	 * End a GPU timing measurement.
 	 * The timing for this particular measurement will be resolved at a later time by the GPU.
 	 */
-	void EndTiming(FVulkanCmdBuffer* CmdBuffer = nullptr);
+	void EndTiming(FVulkanCommandBuffer* CmdBuffer = nullptr);
 
 	/**
 	 * Retrieves the most recently resolved timing measurement.
@@ -72,7 +72,7 @@ private:
 	/**
 	 * If results weren't available when first queried, make sure we try again.
 	 */
-	void RefreshPendingQueries(FVulkanCmdBuffer& CmdBuffer);
+	void RefreshPendingQueries(FVulkanCommandBuffer& CmdBuffer);
 
 	/**
 	 * Remove expired queries to make space for new ones
@@ -195,8 +195,8 @@ struct FVulkanGPUProfiler : public FGPUProfiler
 
 #if VULKAN_SUPPORTS_GPU_CRASH_DUMPS
 	void PrepareCrashMarkerkBuffer();
-	void PushMarkerForCrash(FVulkanCmdBuffer* CmdBuffer, VkBuffer DestBuffer, const TCHAR* Name);
-	void PopMarkerForCrash(FVulkanCmdBuffer* CmdBuffer, VkBuffer DestBuffer);
+	void PushMarkerForCrash(FVulkanCommandBuffer* CmdBuffer, VkBuffer DestBuffer, const TCHAR* Name);
+	void PopMarkerForCrash(FVulkanCommandBuffer* CmdBuffer, VkBuffer DestBuffer);
 	void DumpCrashMarkers(void* BufferData);
 #endif
 
diff --git a/Engine/Source/Runtime/VulkanRHI/Private/VulkanGenericPlatform.cpp b/Engine/Source/Runtime/VulkanRHI/Private/VulkanGenericPlatform.cpp
--- a/Engine/Source/Runtime/VulkanRHI/Private/VulkanGenericPlatform.cpp
+++ b/Engine/Source/Runtime/VulkanRHI/Private/VulkanGenericPlatform.cpp
@@ -181,7 +181,7 @@ FString FVulkanGenericPlatform::GetVulkanProfileNameForFeatureLevel(ERHIFeatureL
 	return ProfileName;
 }
 
-void FVulkanGenericPlatform::WriteCrashMarkerWithoutExtensions(FVulkanCmdBuffer* CmdBuffer, VkBuffer DestBuffer, const TArrayView<uint32>& Entries, bool bAdding)
+void FVulkanGenericPlatform::WriteCrashMarkerWithoutExtensions(FVulkanCommandBuffer* CmdBuffer, VkBuffer DestBuffer, const TArrayView<uint32>& Entries, bool bAdding)
 {
 	// Buffer writes cannot be recorded inside of a render pass
 	if (CmdBuffer->IsOutsideRenderPass())
diff --git a/Engine/Source/Runtime/VulkanRHI/Private/VulkanGenericPlatform.h b/Engine/Source/Runtime/VulkanRHI/Private/VulkanGenericPlatform.h
--- a/Engine/Source/Runtime/VulkanRHI/Private/VulkanGenericPlatform.h
+++ b/Engine/Source/Runtime/VulkanRHI/Private/VulkanGenericPlatform.h
@@ -13,7 +13,7 @@ class FVulkanDevice;
 class FVulkanRenderTargetLayout;
 struct FGfxPipelineDesc;
 class FVulkanPhysicalDeviceFeatures;
-class FVulkanCmdBuffer;
+class FVulkanCommandBuffer;
 
 enum EShaderPlatform : uint16;
 enum class EGpuVendorId : uint32;
@@ -80,9 +80,9 @@ public:
 
 	static bool RegisterGPUWork() { return true; }
 
-	static void WriteCrashMarker(const FOptionalVulkanDeviceExtensions& OptionalExtensions, FVulkanCmdBuffer* CmdBuffer, VkBuffer DestBuffer, const TArrayView<uint32>& Entries, bool bAdding) {}
+	static void WriteCrashMarker(const FOptionalVulkanDeviceExtensions& OptionalExtensions, FVulkanCommandBuffer* CmdBuffer, VkBuffer DestBuffer, const TArrayView<uint32>& Entries, bool bAdding) {}
 
-	static void WriteCrashMarkerWithoutExtensions(FVulkanCmdBuffer* CmdBuffer, VkBuffer DestBuffer, const TArrayView<uint32>& Entries, bool bAdding);
+	static void WriteCrashMarkerWithoutExtensions(FVulkanCommandBuffer* CmdBuffer, VkBuffer DestBuffer, const TArrayView<uint32>& Entries, bool bAdding);
 
 	// Allow the platform code to restrict the device features
 	static void RestrictEnabledPhysicalDeviceFeatures(FVulkanPhysicalDeviceFeatures* InOutFeaturesToEnable);
diff --git a/Engine/Source/Runtime/VulkanRHI/Private/VulkanMemory.cpp b/Engine/Source/Runtime/VulkanRHI/Private/VulkanMemory.cpp
--- a/Engine/Source/Runtime/VulkanRHI/Private/VulkanMemory.cpp
+++ b/Engine/Source/Runtime/VulkanRHI/Private/VulkanMemory.cpp
@@ -1257,7 +1257,7 @@ namespace VulkanRHI
 
 
 			TVulkanTrackBase<FVulkanTexture>::CollectAll(Collector);
-			TVulkanTrackBase<FVulkanResourceMultiBuffer>::CollectAll(Collector);
+			TVulkanTrackBase<FVulkanBuffer>::CollectAll(Collector);
 			for(auto& Itr : AllocationBuckets)
 			{
 				VULKAN_LOGMEMORY(TEXT("***** BUCKET < %d kb *****"), Itr.Key/1024);
@@ -2636,7 +2636,7 @@ namespace VulkanRHI
 
 	uint32 FMemoryManager::CalculateBufferAlignment(FVulkanDevice& InDevice, EBufferUsageFlags InUEUsage, bool bZeroSize)
 	{
-		const VkBufferUsageFlags VulkanBufferUsage = FVulkanResourceMultiBuffer::UEToVKBufferUsageFlags(&InDevice, InUEUsage, bZeroSize);
+		const VkBufferUsageFlags VulkanBufferUsage = FVulkanBuffer::UEToVKBufferUsageFlags(&InDevice, InUEUsage, bZeroSize);
 
 		uint32 Alignment = CalculateBufferAlignmentFromVKUsageFlags(InDevice, VulkanBufferUsage);
 
@@ -4383,7 +4383,7 @@ namespace VulkanRHI
 		return StagingBuffer;
 	}
 
-	inline FStagingManager::FPendingItemsPerCmdBuffer* FStagingManager::FindOrAdd(FVulkanCmdBuffer* CmdBuffer)
+	inline FStagingManager::FPendingItemsPerCmdBuffer* FStagingManager::FindOrAdd(FVulkanCommandBuffer* CmdBuffer)
 	{
 		for (int32 Index = 0; Index < PendingFreeStagingBuffers.Num(); ++Index)
 		{
@@ -4413,7 +4413,7 @@ namespace VulkanRHI
 		return &New;
 	}
 
-	void FStagingManager::ReleaseBuffer(FVulkanCmdBuffer* CmdBuffer, FStagingBuffer*& StagingBuffer)
+	void FStagingManager::ReleaseBuffer(FVulkanCommandBuffer* CmdBuffer, FStagingBuffer*& StagingBuffer)
 	{
 #if VULKAN_ENABLE_AGGRESSIVE_STATS
 		SCOPE_CYCLE_COUNTER(STAT_VulkanStagingBuffer);
@@ -4883,7 +4883,7 @@ namespace VulkanRHI
 
 
 
-	void FDeferredDeletionQueue2::OnCmdBufferDeleted(FVulkanCmdBuffer* DeletedCmdBuffer)
+	void FDeferredDeletionQueue2::OnCmdBufferDeleted(FVulkanCommandBuffer* DeletedCmdBuffer)
 	{
 		FScopeLock ScopeLock(&CS);
 		for (int32 Index = 0; Index < Entries.Num(); ++Index)
@@ -4974,7 +4974,7 @@ namespace VulkanRHI
 		return NewBlock;
 	}
 
-	FTempBlockAllocator::FInternalAlloc FTempBlockAllocator::InternalAlloc(uint32 InSize, FVulkanCmdBuffer* CmdBuffer)
+	FTempBlockAllocator::FInternalAlloc FTempBlockAllocator::InternalAlloc(uint32 InSize, FVulkanCommandBuffer* CmdBuffer)
 	{
 		const uint32 AlignedSize = Align(InSize, BlockAlignment);
 		checkfSlow(BlockAlignment < BlockSize, TEXT("Requested size of %d (%d aligned) is too large for block size of %d"), InSize, AlignedSize, BlockSize);
@@ -5027,7 +5027,7 @@ namespace VulkanRHI
 		}
 	}
 
-	uint8* FTempBlockAllocator::Alloc(uint32 InSize, FVulkanCmdBuffer* CmdBuffer, VkDescriptorBufferBindingInfoEXT& OutBindingInfo, VkDeviceSize& OutOffset)
+	uint8* FTempBlockAllocator::Alloc(uint32 InSize, FVulkanCommandBuffer* CmdBuffer, VkDescriptorBufferBindingInfoEXT& OutBindingInfo, VkDeviceSize& OutOffset)
 	{
 		FInternalAlloc Alloc = InternalAlloc(InSize, CmdBuffer);
 
@@ -5038,7 +5038,7 @@ namespace VulkanRHI
 		return Alloc.Block->MappedPointer + Alloc.Offset;
 	}
 
-	uint8* FTempBlockAllocator::Alloc(uint32 InSize, uint32 InAlignment, FVulkanCmdBuffer* CmdBuffer, FVulkanAllocation& OutAllocation, VkDescriptorAddressInfoEXT* OutDescriptorAddressInfo)
+	uint8* FTempBlockAllocator::Alloc(uint32 InSize, uint32 InAlignment, FVulkanCommandBuffer* CmdBuffer, FVulkanAllocation& OutAllocation, VkDescriptorAddressInfoEXT* OutDescriptorAddressInfo)
 	{
 		const uint32 AlignedSize = Align(InSize, BlockAlignment);
 		checkfSlow(AlignedSize < BlockSize, TEXT("Requested size of %d (%d aligned) is too large for block size of %d"), InSize, AlignedSize, BlockSize);
@@ -5060,7 +5060,7 @@ namespace VulkanRHI
 		return Alloc.Block->MappedPointer + Alloc.Offset;
 	}
 
-	uint8* FTempBlockAllocator::Alloc(uint32 InSize, FVulkanCmdBuffer* CmdBuffer, VkStridedDeviceAddressRegionKHR* OutStridedDeviceAddressRegion)
+	uint8* FTempBlockAllocator::Alloc(uint32 InSize, FVulkanCommandBuffer* CmdBuffer, VkStridedDeviceAddressRegionKHR* OutStridedDeviceAddressRegion)
 	{
 		const uint32 AlignedSize = Align(InSize, BlockAlignment);
 		checkfSlow(AlignedSize < BlockSize, TEXT("Requested size of %d (%d aligned) is too large for block size of %d"), InSize, AlignedSize, BlockSize);
@@ -5084,7 +5084,7 @@ namespace VulkanRHI
 			bool BlockReady = true;
 			for (auto& Pair : Block->Fences)
 			{
-				FVulkanCmdBuffer* CmdBuffer = Pair.Key;
+				FVulkanCommandBuffer* CmdBuffer = Pair.Key;
 				if (Pair.Value >= CmdBuffer->GetFenceSignaledCounter())
 				{
 					BlockReady = false;
diff --git a/Engine/Source/Runtime/VulkanRHI/Private/VulkanMemory.h b/Engine/Source/Runtime/VulkanRHI/Private/VulkanMemory.h
--- a/Engine/Source/Runtime/VulkanRHI/Private/VulkanMemory.h
+++ b/Engine/Source/Runtime/VulkanRHI/Private/VulkanMemory.h
@@ -27,7 +27,7 @@
 class FVulkanCommandListContext;
 class FVulkanDevice;
 class FVulkanQueue;
-class FVulkanCmdBuffer;
+class FVulkanCommandBuffer;
 class FVulkanTexture;
 
 #if !VULKAN_OBJECT_TRACKING
@@ -968,7 +968,7 @@ namespace VulkanRHI
 		FStagingBuffer* AcquireBuffer(uint32 Size, VkBufferUsageFlags InUsageFlags = VK_BUFFER_USAGE_TRANSFER_SRC_BIT, VkMemoryPropertyFlagBits InMemoryReadFlags = VK_MEMORY_PROPERTY_HOST_COHERENT_BIT);
 
 		// Sets pointer to nullptr
-		void ReleaseBuffer(FVulkanCmdBuffer* CmdBuffer, FStagingBuffer*& StagingBuffer);
+		void ReleaseBuffer(FVulkanCommandBuffer* CmdBuffer, FStagingBuffer*& StagingBuffer);
 
 		void ProcessPendingFree(bool bImmediately, bool bFreeToOS);
 
@@ -980,7 +980,7 @@ namespace VulkanRHI
 		friend class FMemoryManager;
 		struct FPendingItemsPerCmdBuffer
 		{
-			FVulkanCmdBuffer* CmdBuffer;
+			FVulkanCommandBuffer* CmdBuffer;
 			struct FPendingItems
 			{
 				uint64 FenceCounter;
@@ -1007,7 +1007,7 @@ namespace VulkanRHI
 		uint64 PeakUsedMemory;
 		uint64 UsedMemory;
 
-		FPendingItemsPerCmdBuffer* FindOrAdd(FVulkanCmdBuffer* CmdBuffer);
+		FPendingItemsPerCmdBuffer* FindOrAdd(FVulkanCommandBuffer* CmdBuffer);
 
 		void ProcessPendingFreeNoLock(bool bImmediately, bool bFreeToOS);
 
@@ -1130,7 +1130,7 @@ namespace VulkanRHI
 			VkDeviceAddress BufferAddress = 0;
 
 			// Used by owners to suballocate
-			TMap<FVulkanCmdBuffer*, uint64> Fences;
+			TMap<FVulkanCommandBuffer*, uint64> Fences;
 			std::atomic<uint32> CurrentOffset = 0;
 		};
 
@@ -1143,15 +1143,15 @@ namespace VulkanRHI
 		FTempBlockAllocator(FVulkanDevice* InDevice, uint32 InBlockSize, uint32 InBlockAlignment, VkBufferUsageFlags InBufferUsage);
 		virtual ~FTempBlockAllocator();
 
-		uint8* Alloc(uint32 InSize, FVulkanCmdBuffer* CmdBuffer, VkDescriptorBufferBindingInfoEXT& OutBindingInfo, VkDeviceSize& OutOffset);
-		uint8* Alloc(uint32 InSize, uint32 InAlignment, FVulkanCmdBuffer* CmdBuffer, FVulkanAllocation& OutAllocation, VkDescriptorAddressInfoEXT* OutDescriptorAddressInfo = nullptr);
-		uint8* Alloc(uint32 InSize, FVulkanCmdBuffer* CmdBuffer, VkStridedDeviceAddressRegionKHR* OutStridedDeviceAddressRegion);
+		uint8* Alloc(uint32 InSize, FVulkanCommandBuffer* CmdBuffer, VkDescriptorBufferBindingInfoEXT& OutBindingInfo, VkDeviceSize& OutOffset);
+		uint8* Alloc(uint32 InSize, uint32 InAlignment, FVulkanCommandBuffer* CmdBuffer, FVulkanAllocation& OutAllocation, VkDescriptorAddressInfoEXT* OutDescriptorAddressInfo = nullptr);
+		uint8* Alloc(uint32 InSize, FVulkanCommandBuffer* CmdBuffer, VkStridedDeviceAddressRegionKHR* OutStridedDeviceAddressRegion);
 
 		void UpdateBlocks();
 
 	protected:
 		FTempMemoryBlock* AllocBlock();
-		FInternalAlloc InternalAlloc(uint32 InSize, FVulkanCmdBuffer* CmdBuffer);
+		FInternalAlloc InternalAlloc(uint32 InSize, FVulkanCommandBuffer* CmdBuffer);
 
 		const uint32 BlockSize;
 		const uint32 BlockAlignment;
diff --git a/Engine/Source/Runtime/VulkanRHI/Private/VulkanPendingState.cpp b/Engine/Source/Runtime/VulkanRHI/Private/VulkanPendingState.cpp
--- a/Engine/Source/Runtime/VulkanRHI/Private/VulkanPendingState.cpp
+++ b/Engine/Source/Runtime/VulkanRHI/Private/VulkanPendingState.cpp
@@ -339,7 +339,7 @@ void FVulkanPendingComputeState::SetUAVForUBResource(uint32 DescriptorSet, uint3
 }
 
 
-void FVulkanPendingComputeState::PrepareForDispatch(FVulkanCmdBuffer* InCmdBuffer)
+void FVulkanPendingComputeState::PrepareForDispatch(FVulkanCommandBuffer* InCmdBuffer)
 {
 #if VULKAN_ENABLE_AGGRESSIVE_STATS
 	SCOPE_CYCLE_COUNTER(STAT_VulkanDispatchCallPrepareTime);
@@ -381,7 +381,7 @@ FVulkanPendingGfxState::~FVulkanPendingGfxState()
 	}
 }
 
-void FVulkanPendingGfxState::PrepareForDraw(FVulkanCmdBuffer* CmdBuffer)
+void FVulkanPendingGfxState::PrepareForDraw(FVulkanCommandBuffer* CmdBuffer)
 {
 #if VULKAN_ENABLE_AGGRESSIVE_STATS
 	SCOPE_CYCLE_COUNTER(STAT_VulkanDrawCallPrepareTime);
@@ -490,7 +490,7 @@ void FVulkanPendingGfxState::PrepareForDraw(FVulkanCmdBuffer* CmdBuffer)
 	}
 }
 
-void FVulkanPendingGfxState::InternalUpdateDynamicStates(FVulkanCmdBuffer* Cmd)
+void FVulkanPendingGfxState::InternalUpdateDynamicStates(FVulkanCommandBuffer* Cmd)
 {
 	const bool bNeedsUpdateViewport = !Cmd->bHasViewport || Viewports.Num() != Cmd->CurrentViewports.Num() || (FMemory::Memcmp((const void*)Cmd->CurrentViewports.GetData(), (const void*)Viewports.GetData(), Viewports.Num() * sizeof(VkViewport)) != 0);
 	// Validate and update Viewport
diff --git a/Engine/Source/Runtime/VulkanRHI/Private/VulkanPendingState.h b/Engine/Source/Runtime/VulkanRHI/Private/VulkanPendingState.h
--- a/Engine/Source/Runtime/VulkanRHI/Private/VulkanPendingState.h
+++ b/Engine/Source/Runtime/VulkanRHI/Private/VulkanPendingState.h
@@ -55,7 +55,7 @@ public:
 		}
 	}
 
-	void PrepareForDispatch(FVulkanCmdBuffer* CmdBuffer);
+	void PrepareForDispatch(FVulkanCommandBuffer* CmdBuffer);
 
 	inline const FVulkanComputeShader* GetCurrentShader() const
 	{
@@ -283,11 +283,11 @@ public:
 		CurrentState->SetPackedGlobalShaderParameter(Stage, BufferIndex, Offset, NumBytes, NewValue);
 	}
 
-	void PrepareForDraw(FVulkanCmdBuffer* CmdBuffer);
+	void PrepareForDraw(FVulkanCommandBuffer* CmdBuffer);
 
 	bool SetGfxPipeline(FVulkanRHIGraphicsPipelineState* InGfxPipeline, bool bForceReset);
 
-	inline void UpdateDynamicStates(FVulkanCmdBuffer* Cmd)
+	inline void UpdateDynamicStates(FVulkanCommandBuffer* Cmd)
 	{
 		InternalUpdateDynamicStates(Cmd);
 	}
@@ -338,7 +338,7 @@ protected:
 	FVertexStream PendingStreams[MaxVertexElementCount];
 	bool bDirtyVertexStreams;
 
-	void InternalUpdateDynamicStates(FVulkanCmdBuffer* Cmd);
+	void InternalUpdateDynamicStates(FVulkanCommandBuffer* Cmd);
 	void UpdateInputAttachments(FVulkanFramebuffer* Framebuffer);
 
 	FVulkanCommandListContext& Context;
diff --git a/Engine/Source/Runtime/VulkanRHI/Private/VulkanPipelineState.cpp b/Engine/Source/Runtime/VulkanRHI/Private/VulkanPipelineState.cpp
--- a/Engine/Source/Runtime/VulkanRHI/Private/VulkanPipelineState.cpp
+++ b/Engine/Source/Runtime/VulkanRHI/Private/VulkanPipelineState.cpp
@@ -156,7 +156,7 @@ void FVulkanCommonPipelineDescriptorState::CreateDescriptorWriteInfos()
 }
 
 
-static inline VulkanRHI::FVulkanAllocation UpdatePackedUniformBuffers(const FPackedUniformBuffers& PackedUniformBuffers, FVulkanCommandListContext* InContext, FVulkanCmdBuffer* InCmdBuffer)
+static inline VulkanRHI::FVulkanAllocation UpdatePackedUniformBuffers(const FPackedUniformBuffers& PackedUniformBuffers, FVulkanCommandListContext* InContext, FVulkanCommandBuffer* InCmdBuffer)
 {
 	const FPackedUniformBuffers::FPackedBuffer& StagedUniformBuffer = PackedUniformBuffers.GetBuffer();
 
@@ -173,7 +173,7 @@ static inline VulkanRHI::FVulkanAllocation UpdatePackedUniformBuffers(const FPac
 
 
 template<bool bUseDynamicGlobalUBs>
-bool FVulkanComputePipelineDescriptorState::InternalUpdateDescriptorSets(FVulkanCommandListContext* CmdListContext, FVulkanCmdBuffer* CmdBuffer)
+bool FVulkanComputePipelineDescriptorState::InternalUpdateDescriptorSets(FVulkanCommandListContext* CmdListContext, FVulkanCommandBuffer* CmdBuffer)
 {
 #if VULKAN_ENABLE_AGGRESSIVE_STATS
 	SCOPE_CYCLE_COUNTER(STAT_VulkanUpdateDescriptorSets);
@@ -223,7 +223,7 @@ bool FVulkanComputePipelineDescriptorState::InternalUpdateDescriptorSets(FVulkan
 	return true;
 }
 
-void FVulkanComputePipelineDescriptorState::UpdateBindlessDescriptors(FVulkanCommandListContext* CmdListContext, FVulkanCmdBuffer* CmdBuffer)
+void FVulkanComputePipelineDescriptorState::UpdateBindlessDescriptors(FVulkanCommandListContext* CmdListContext, FVulkanCommandBuffer* CmdBuffer)
 {
 	check(bUseBindless);
 
@@ -360,7 +360,7 @@ FVulkanGraphicsPipelineDescriptorState::FVulkanGraphicsPipelineDescriptorState(F
 }
 
 template<bool bUseDynamicGlobalUBs>
-bool FVulkanGraphicsPipelineDescriptorState::InternalUpdateDescriptorSets(FVulkanCommandListContext* CmdListContext, FVulkanCmdBuffer* CmdBuffer)
+bool FVulkanGraphicsPipelineDescriptorState::InternalUpdateDescriptorSets(FVulkanCommandListContext* CmdListContext, FVulkanCommandBuffer* CmdBuffer)
 {
 #if VULKAN_ENABLE_AGGRESSIVE_STATS
 	SCOPE_CYCLE_COUNTER(STAT_VulkanUpdateDescriptorSets);
@@ -435,7 +435,7 @@ bool FVulkanGraphicsPipelineDescriptorState::InternalUpdateDescriptorSets(FVulka
 	return true;
 }
 
-void FVulkanGraphicsPipelineDescriptorState::UpdateBindlessDescriptors(FVulkanCommandListContext* CmdListContext, FVulkanCmdBuffer* CmdBuffer)
+void FVulkanGraphicsPipelineDescriptorState::UpdateBindlessDescriptors(FVulkanCommandListContext* CmdListContext, FVulkanCommandBuffer* CmdBuffer)
 {
 	check(bUseBindless);
 
@@ -544,7 +544,7 @@ void FVulkanCommandListContext::RHISetGraphicsPipelineState(FRHIGraphicsPipeline
 	FVulkanPipelineStateCacheManager* PipelineStateCache = Device->GetPipelineStateCache();
 	PipelineStateCache->LRUTouch(Pipeline);
 
-	FVulkanCmdBuffer* CmdBuffer = CommandBufferManager->GetActiveCmdBuffer();
+	FVulkanCommandBuffer* CmdBuffer = CommandBufferManager->GetActiveCmdBuffer();
 	bool bForceResetPipeline = !CmdBuffer->bHasPipeline;
 
 	if (PendingGfxState->SetGfxPipeline(Pipeline, bForceResetPipeline))
@@ -575,7 +575,7 @@ void FVulkanCommandListContext::RHISetGraphicsPipelineState(FRHIGraphicsPipeline
 
 void FVulkanCommandListContext::RHISetComputePipelineState(FRHIComputePipelineState* ComputePipelineState)
 {
-	FVulkanCmdBuffer* CmdBuffer = CommandBufferManager->GetActiveCmdBuffer();
+	FVulkanCommandBuffer* CmdBuffer = CommandBufferManager->GetActiveCmdBuffer();
 	if (CmdBuffer->IsInsideRenderPass())
 	{
 		if (GVulkanSubmitAfterEveryEndRenderPass)
@@ -599,10 +599,10 @@ void FVulkanCommandListContext::RHISetComputePipelineState(FRHIComputePipelineSt
 }
 
 
-template bool FVulkanGraphicsPipelineDescriptorState::InternalUpdateDescriptorSets<true>(FVulkanCommandListContext* CmdListContext, FVulkanCmdBuffer* CmdBuffer);
-template bool FVulkanGraphicsPipelineDescriptorState::InternalUpdateDescriptorSets<false>(FVulkanCommandListContext* CmdListContext, FVulkanCmdBuffer* CmdBuffer);
-template bool FVulkanComputePipelineDescriptorState::InternalUpdateDescriptorSets<true>(FVulkanCommandListContext* CmdListContext, FVulkanCmdBuffer* CmdBuffer);
-template bool FVulkanComputePipelineDescriptorState::InternalUpdateDescriptorSets<false>(FVulkanCommandListContext* CmdListContext, FVulkanCmdBuffer* CmdBuffer);
+template bool FVulkanGraphicsPipelineDescriptorState::InternalUpdateDescriptorSets<true>(FVulkanCommandListContext* CmdListContext, FVulkanCommandBuffer* CmdBuffer);
+template bool FVulkanGraphicsPipelineDescriptorState::InternalUpdateDescriptorSets<false>(FVulkanCommandListContext* CmdListContext, FVulkanCommandBuffer* CmdBuffer);
+template bool FVulkanComputePipelineDescriptorState::InternalUpdateDescriptorSets<true>(FVulkanCommandListContext* CmdListContext, FVulkanCommandBuffer* CmdBuffer);
+template bool FVulkanComputePipelineDescriptorState::InternalUpdateDescriptorSets<false>(FVulkanCommandListContext* CmdListContext, FVulkanCommandBuffer* CmdBuffer);
 
 
 void FVulkanDescriptorSetWriter::CheckAllWritten()
@@ -712,7 +712,7 @@ void FVulkanDescriptorSetWriter::InitWrittenMasks(uint32 NumDescriptorWrites)
 #endif
 }
 
-void FVulkanCommonPipelineDescriptorState::SetSRV(FVulkanCmdBuffer* CmdBuffer, bool bCompute, uint8 DescriptorSet, uint32 BindingIndex, FVulkanShaderResourceView* SRV)
+void FVulkanCommonPipelineDescriptorState::SetSRV(FVulkanCommandBuffer* CmdBuffer, bool bCompute, uint8 DescriptorSet, uint32 BindingIndex, FVulkanShaderResourceView* SRV)
 {
 	check(!bUseBindless);
 
@@ -749,7 +749,7 @@ void FVulkanCommonPipelineDescriptorState::SetSRV(FVulkanCmdBuffer* CmdBuffer, b
 	}
 }
 
-void FVulkanCommonPipelineDescriptorState::SetUAV(FVulkanCmdBuffer* CmdBuffer, bool bCompute, uint8 DescriptorSet, uint32 BindingIndex, FVulkanUnorderedAccessView* UAV)
+void FVulkanCommonPipelineDescriptorState::SetUAV(FVulkanCommandBuffer* CmdBuffer, bool bCompute, uint8 DescriptorSet, uint32 BindingIndex, FVulkanUnorderedAccessView* UAV)
 {
 	check(!bUseBindless);
 
diff --git a/Engine/Source/Runtime/VulkanRHI/Private/VulkanPipelineState.h b/Engine/Source/Runtime/VulkanRHI/Private/VulkanPipelineState.h
--- a/Engine/Source/Runtime/VulkanRHI/Private/VulkanPipelineState.h
+++ b/Engine/Source/Runtime/VulkanRHI/Private/VulkanPipelineState.h
@@ -61,8 +61,8 @@ public:
 		bIsDSetsKeyDirty |= bDirty;
 	}
 
-	void SetSRV(FVulkanCmdBuffer* CmdBuffer, bool bCompute, uint8 DescriptorSet, uint32 BindingIndex, FVulkanShaderResourceView* SRV);
-	void SetUAV(FVulkanCmdBuffer* CmdBuffer, bool bCompute, uint8 DescriptorSet, uint32 BindingIndex, FVulkanUnorderedAccessView* UAV);
+	void SetSRV(FVulkanCommandBuffer* CmdBuffer, bool bCompute, uint8 DescriptorSet, uint32 BindingIndex, FVulkanShaderResourceView* SRV);
+	void SetUAV(FVulkanCommandBuffer* CmdBuffer, bool bCompute, uint8 DescriptorSet, uint32 BindingIndex, FVulkanUnorderedAccessView* UAV);
 
 	inline void SetTexture(uint8 DescriptorSet, uint32 BindingIndex, const FVulkanTexture* Texture, VkImageLayout Layout)
 	{
@@ -196,7 +196,7 @@ public:
 		PackedUniformBuffers.SetPackedGlobalParameter(ByteOffset, NumBytes, NewValue, PackedUniformBuffersDirty);
 	}
 
-	bool UpdateDescriptorSets(FVulkanCommandListContext* CmdListContext, FVulkanCmdBuffer* CmdBuffer)
+	bool UpdateDescriptorSets(FVulkanCommandListContext* CmdListContext, FVulkanCommandBuffer* CmdBuffer)
 	{
 		check(!bUseBindless);
 
@@ -211,7 +211,7 @@ public:
 		}
 	}
 
-	void UpdateBindlessDescriptors(FVulkanCommandListContext* CmdListContext, FVulkanCmdBuffer* CmdBuffer);
+	void UpdateBindlessDescriptors(FVulkanCommandListContext* CmdListContext, FVulkanCommandBuffer* CmdBuffer);
 
 	inline void BindDescriptorSets(VkCommandBuffer CmdBuffer)
 	{
@@ -226,7 +226,7 @@ protected:
 	FVulkanComputePipeline* ComputePipeline;
 
 	template<bool bUseDynamicGlobalUBs>
-	bool InternalUpdateDescriptorSets(FVulkanCommandListContext* CmdListContext, FVulkanCmdBuffer* CmdBuffer);
+	bool InternalUpdateDescriptorSets(FVulkanCommandListContext* CmdListContext, FVulkanCommandBuffer* CmdBuffer);
 
 	friend class FVulkanPendingComputeState;
 	friend class FVulkanCommandListContext;
@@ -247,7 +247,7 @@ public:
 		PackedUniformBuffers[Stage].SetPackedGlobalParameter(ByteOffset, NumBytes, NewValue, PackedUniformBuffersDirty[Stage]);
 	}
 
-	bool UpdateDescriptorSets(FVulkanCommandListContext* CmdListContext, FVulkanCmdBuffer* CmdBuffer)
+	bool UpdateDescriptorSets(FVulkanCommandListContext* CmdListContext, FVulkanCommandBuffer* CmdBuffer)
 	{
 		check(!bUseBindless);
 
@@ -262,7 +262,7 @@ public:
 		}
 	}
 
-	void UpdateBindlessDescriptors(FVulkanCommandListContext* CmdListContext, FVulkanCmdBuffer* CmdBuffer);
+	void UpdateBindlessDescriptors(FVulkanCommandListContext* CmdListContext, FVulkanCommandBuffer* CmdBuffer);
 
 	inline void BindDescriptorSets(VkCommandBuffer CmdBuffer)
 	{
@@ -284,7 +284,7 @@ protected:
 	FVulkanRHIGraphicsPipelineState* GfxPipeline;
 
 	template<bool bUseDynamicGlobalUBs>
-	bool InternalUpdateDescriptorSets(FVulkanCommandListContext* CmdListContext, FVulkanCmdBuffer* CmdBuffer);
+	bool InternalUpdateDescriptorSets(FVulkanCommandListContext* CmdListContext, FVulkanCommandBuffer* CmdBuffer);
 
 	friend class FVulkanPendingGfxState;
 	friend class FVulkanCommandListContext;
diff --git a/Engine/Source/Runtime/VulkanRHI/Private/VulkanQuery.cpp b/Engine/Source/Runtime/VulkanRHI/Private/VulkanQuery.cpp
--- a/Engine/Source/Runtime/VulkanRHI/Private/VulkanQuery.cpp
+++ b/Engine/Source/Runtime/VulkanRHI/Private/VulkanQuery.cpp
@@ -187,14 +187,14 @@ bool FVulkanOcclusionQueryPool::InternalTryGetResults(bool bWait)
 	return false;
 }
 
-void FVulkanOcclusionQueryPool::SetFence(FVulkanCmdBuffer* InCmdBuffer)
+void FVulkanOcclusionQueryPool::SetFence(FVulkanCommandBuffer* InCmdBuffer)
 {
 	check(InCmdBuffer);
 	CmdBuffer = InCmdBuffer;
 	FenceCounter = InCmdBuffer->GetFenceSignaledCounter();
 }
 
-void FVulkanOcclusionQueryPool::Reset(FVulkanCmdBuffer* InCmdBuffer, uint32 InFrameNumber)
+void FVulkanOcclusionQueryPool::Reset(FVulkanCommandBuffer* InCmdBuffer, uint32 InFrameNumber)
 {
 /*
 	check(!InCmdBuffer || InCmdBuffer->GetFenceSignaledCounter() > FenceCounter);
@@ -229,7 +229,7 @@ void FVulkanOcclusionQueryPool::Reset(FVulkanCmdBuffer* InCmdBuffer, uint32 InFr
 	State = EState::RHIT_PostBeginBatch;
 }
 
-void FVulkanCommandListContext::BeginOcclusionQueryBatch(FVulkanCmdBuffer* CmdBuffer, uint32 NumQueriesInBatch)
+void FVulkanCommandListContext::BeginOcclusionQueryBatch(FVulkanCommandBuffer* CmdBuffer, uint32 NumQueriesInBatch)
 {
 	ensure(IsImmediate());
 /*
@@ -346,7 +346,7 @@ void FVulkanDevice::ReleaseUnusedOcclusionQueryPools()
 	}
 }
 
-void FVulkanCommandListContext::EndOcclusionQueryBatch(FVulkanCmdBuffer* CmdBuffer)
+void FVulkanCommandListContext::EndOcclusionQueryBatch(FVulkanCommandBuffer* CmdBuffer)
 {
 	ensure(IsImmediate());
 	check(CmdBuffer);
@@ -586,7 +586,7 @@ void FVulkanCommandListContext::RHIBeginRenderQuery(FRHIRenderQuery* QueryRHI)
 		const uint32 IndexInPool = CurrentOcclusionQueryPool->AcquireIndex(Query);
 		Query->Pool = CurrentOcclusionQueryPool;
 		Query->IndexInPool = IndexInPool;
-		FVulkanCmdBuffer* CmdBuffer = CommandBufferManager->GetActiveCmdBuffer();
+		FVulkanCommandBuffer* CmdBuffer = CommandBufferManager->GetActiveCmdBuffer();
 		VulkanRHI::vkCmdBeginQuery(CmdBuffer->GetHandle(), CurrentOcclusionQueryPool->GetHandle(), IndexInPool, VK_QUERY_CONTROL_PRECISE_BIT);
 	}
 	else if (BaseQuery->QueryType == RQT_AbsoluteTime)
@@ -609,7 +609,7 @@ void FVulkanCommandListContext::RHIEndRenderQuery(FRHIRenderQuery* QueryRHI)
 		FVulkanOcclusionQuery* Query = static_cast<FVulkanOcclusionQuery*>(BaseQuery);
 		ensure(Query->State == FVulkanOcclusionQuery::EState::RHI_PostBegin);
 		Query->State = FVulkanOcclusionQuery::EState::RHI_PostEnd;
-		FVulkanCmdBuffer* CmdBuffer = CommandBufferManager->GetActiveCmdBuffer();
+		FVulkanCommandBuffer* CmdBuffer = CommandBufferManager->GetActiveCmdBuffer();
 		VulkanRHI::vkCmdEndQuery(CmdBuffer->GetHandle(), CurrentOcclusionQueryPool->GetHandle(), Query->IndexInPool);
 	}
 	else if (BaseQuery->QueryType == RQT_AbsoluteTime)
@@ -626,7 +626,7 @@ void FVulkanCommandListContext::RHIEndRenderQuery(FRHIRenderQuery* QueryRHI)
 		Query->Pool->CurrentTimestamp = (Query->Pool->CurrentTimestamp + 1) % Query->Pool->BufferSize;
 		const uint32 QueryEndIndex = Query->Pool->CurrentTimestamp;
 		const VkPipelineStageFlagBits QueryPipelineStage = GTimestampQueryStage ? VK_PIPELINE_STAGE_BOTTOM_OF_PIPE_BIT : VK_PIPELINE_STAGE_TOP_OF_PIPE_BIT;
-		FVulkanCmdBuffer* CmdBuffer = CommandBufferManager->GetActiveCmdBuffer();
+		FVulkanCommandBuffer* CmdBuffer = CommandBufferManager->GetActiveCmdBuffer();
 		VulkanRHI::vkCmdWriteTimestamp(CmdBuffer->GetHandle(), QueryPipelineStage, Query->Pool->GetHandle(), QueryEndIndex);
 		CmdBuffer->AddPendingTimestampQuery(QueryEndIndex, 1, Query->Pool->GetHandle(), Query->Pool->ResultsBuffer->GetHandle(), true);
 		Query->Pool->TimestampListHandles[QueryEndIndex].CmdBuffer = CmdBuffer;
@@ -648,14 +648,14 @@ void FVulkanCommandListContext::RHICalibrateTimers(FRHITimestampCalibrationQuery
 }
 #endif
 
-void FVulkanCommandListContext::WriteBeginTimestamp(FVulkanCmdBuffer* CmdBuffer)
+void FVulkanCommandListContext::WriteBeginTimestamp(FVulkanCommandBuffer* CmdBuffer)
 {
 #if (RHI_NEW_GPU_PROFILER == 0)
 	FrameTiming->StartTiming(CmdBuffer);
 #endif
 }
 
-void FVulkanCommandListContext::WriteEndTimestamp(FVulkanCmdBuffer* CmdBuffer)
+void FVulkanCommandListContext::WriteEndTimestamp(FVulkanCommandBuffer* CmdBuffer)
 {
 #if (RHI_NEW_GPU_PROFILER == 0)
 	FrameTiming->EndTiming(CmdBuffer);
diff --git a/Engine/Source/Runtime/VulkanRHI/Private/VulkanQueue.cpp b/Engine/Source/Runtime/VulkanRHI/Private/VulkanQueue.cpp
--- a/Engine/Source/Runtime/VulkanRHI/Private/VulkanQueue.cpp
+++ b/Engine/Source/Runtime/VulkanRHI/Private/VulkanQueue.cpp
@@ -41,7 +41,7 @@ FVulkanQueue::~FVulkanQueue()
 	check(Device);
 }
 
-void FVulkanQueue::Submit(FVulkanCmdBuffer* CmdBuffer, uint32 NumSignalSemaphores, VkSemaphore* SignalSemaphores)
+void FVulkanQueue::Submit(FVulkanCommandBuffer* CmdBuffer, uint32 NumSignalSemaphores, VkSemaphore* SignalSemaphores)
 {
 	check(CmdBuffer->HasEnded());
 
@@ -74,7 +74,7 @@ void FVulkanQueue::Submit(FVulkanCmdBuffer* CmdBuffer, uint32 NumSignalSemaphore
 		VERIFYVULKANRESULT(VulkanRHI::vkQueueSubmit(Queue, 1, &SubmitInfo, Fence->GetHandle()));
 	}
 
-	CmdBuffer->State = FVulkanCmdBuffer::EState::Submitted;
+	CmdBuffer->State = FVulkanCommandBuffer::EState::Submitted;
 	CmdBuffer->MarkSemaphoresAsSubmitted();
 	CmdBuffer->SubmittedFenceCounter = CmdBuffer->FenceSignaledCounter;
 
@@ -121,7 +121,7 @@ void FVulkanQueue::Submit(FVulkanCmdBuffer* CmdBuffer, uint32 NumSignalSemaphore
 	CmdBuffer->GetLayoutManager().TransferTo(LayoutManager);
 }
 
-void FVulkanQueue::UpdateLastSubmittedCommandBuffer(FVulkanCmdBuffer* CmdBuffer)
+void FVulkanQueue::UpdateLastSubmittedCommandBuffer(FVulkanCommandBuffer* CmdBuffer)
 {
 	FScopeLock ScopeLock(&CS);
 	LastSubmittedCmdBuffer = CmdBuffer;
diff --git a/Engine/Source/Runtime/VulkanRHI/Private/VulkanQueue.h b/Engine/Source/Runtime/VulkanRHI/Private/VulkanQueue.h
--- a/Engine/Source/Runtime/VulkanRHI/Private/VulkanQueue.h
+++ b/Engine/Source/Runtime/VulkanRHI/Private/VulkanQueue.h
@@ -10,7 +10,7 @@
 #include "VulkanBarriers.h"
 
 class FVulkanDevice;
-class FVulkanCmdBuffer;
+class FVulkanCommandBuffer;
 class FVulkanCommandListContext;
 class FVulkanLayoutManager;
 
@@ -36,9 +36,9 @@ public:
 		return QueueIndex;
 	}
 
-	void Submit(FVulkanCmdBuffer* CmdBuffer, uint32 NumSignalSemaphores = 0, VkSemaphore* SignalSemaphores = nullptr);
+	void Submit(FVulkanCommandBuffer* CmdBuffer, uint32 NumSignalSemaphores = 0, VkSemaphore* SignalSemaphores = nullptr);
 
-	inline void Submit(FVulkanCmdBuffer* CmdBuffer, VkSemaphore SignalSemaphore)
+	inline void Submit(FVulkanCommandBuffer* CmdBuffer, VkSemaphore SignalSemaphore)
 	{
 		Submit(CmdBuffer, 1, &SignalSemaphore);
 	}
@@ -48,7 +48,7 @@ public:
 		return Queue;
 	}
 
-	inline void GetLastSubmittedInfo(FVulkanCmdBuffer*& OutCmdBuffer, uint64& OutFenceCounter) const
+	inline void GetLastSubmittedInfo(FVulkanCommandBuffer*& OutCmdBuffer, uint64& OutFenceCounter) const
 	{
 		FScopeLock ScopeLock(&CS);
 		OutCmdBuffer = LastSubmittedCmdBuffer;
@@ -79,7 +79,7 @@ private:
 	FVulkanDevice* Device;
 
 	mutable FCriticalSection CS;
-	FVulkanCmdBuffer* LastSubmittedCmdBuffer;
+	FVulkanCommandBuffer* LastSubmittedCmdBuffer;
 	uint64 LastSubmittedCmdBufferFenceCounter;
 	uint64 SubmitCounter;
 	VkPipelineStageFlags SupportedStages;
@@ -87,6 +87,6 @@ private:
 	// Last known layouts submitted on this queue, used for defrag
 	FVulkanLayoutManager LayoutManager;
 
-	void UpdateLastSubmittedCommandBuffer(FVulkanCmdBuffer* CmdBuffer);
+	void UpdateLastSubmittedCommandBuffer(FVulkanCommandBuffer* CmdBuffer);
 	void FillSupportedStageBits();
 };
diff --git a/Engine/Source/Runtime/VulkanRHI/Private/VulkanRHI.cpp b/Engine/Source/Runtime/VulkanRHI/Private/VulkanRHI.cpp
--- a/Engine/Source/Runtime/VulkanRHI/Private/VulkanRHI.cpp
+++ b/Engine/Source/Runtime/VulkanRHI/Private/VulkanRHI.cpp
@@ -1236,7 +1236,7 @@ void FVulkanCommandListContext::RHIEndDrawingViewport(FRHIViewport* ViewportRHI,
 	check(bPresent);
 	RHI->Present();
 */
-	FVulkanCmdBuffer* CmdBuffer = CommandBufferManager->GetActiveCmdBuffer();
+	FVulkanCommandBuffer* CmdBuffer = CommandBufferManager->GetActiveCmdBuffer();
 	check(!CmdBuffer->HasEnded() && !CmdBuffer->IsInsideRenderPass());
 
 	WriteEndTimestamp(CmdBuffer);
@@ -1566,7 +1566,7 @@ FVulkanRHIImageViewInfo FVulkanDynamicRHI::RHIGetImageViewInfo(FRHITexture* InTe
 
 FVulkanRHIAllocationInfo FVulkanDynamicRHI::RHIGetAllocationInfo(FRHIBuffer* InBuffer) const
 {
-	FVulkanResourceMultiBuffer* VulkanBuffer = ResourceCast(InBuffer);
+	FVulkanBuffer* VulkanBuffer = ResourceCast(InBuffer);
 	const VulkanRHI::FVulkanAllocation& Allocation = VulkanBuffer->GetCurrentAllocation();
 
 	FVulkanRHIAllocationInfo NewInfo{};
@@ -1580,14 +1580,14 @@ FVulkanRHIAllocationInfo FVulkanDynamicRHI::RHIGetAllocationInfo(FRHIBuffer* InB
 void FVulkanDynamicRHI::RHISetImageLayout(VkImage Image, VkImageLayout OldLayout, VkImageLayout NewLayout, const VkImageSubresourceRange& SubresourceRange)
 {
 	FVulkanCommandListContext& ImmediateContext = GetDevice()->GetImmediateContext();
-	FVulkanCmdBuffer* CmdBuffer = ImmediateContext.GetCommandBufferManager()->GetActiveCmdBuffer();
+	FVulkanCommandBuffer* CmdBuffer = ImmediateContext.GetCommandBufferManager()->GetActiveCmdBuffer();
 	VulkanSetImageLayout(CmdBuffer->GetHandle(), Image, OldLayout, NewLayout, SubresourceRange);
 }
 
 void FVulkanDynamicRHI::RHISetUploadImageLayout(VkImage Image, VkImageLayout OldLayout, VkImageLayout NewLayout, const VkImageSubresourceRange& SubresourceRange)
 {
 	FVulkanCommandListContext& ImmediateContext = GetDevice()->GetImmediateContext();
-	FVulkanCmdBuffer* CmdBuffer = ImmediateContext.GetCommandBufferManager()->GetUploadCmdBuffer();
+	FVulkanCommandBuffer* CmdBuffer = ImmediateContext.GetCommandBufferManager()->GetUploadCmdBuffer();
 	VulkanSetImageLayout(CmdBuffer->GetHandle(), Image, OldLayout, NewLayout, SubresourceRange);
 }
 
@@ -2341,8 +2341,8 @@ void FVulkanDynamicRHI::RHIReplaceResources(FRHICommandListBase& RHICmdList, TAr
 
 				case FRHIResourceReplaceInfo::EType::Buffer:
 					{
-						FVulkanResourceMultiBuffer* Dst = ResourceCast(Info.GetBuffer().Dst);
-						FVulkanResourceMultiBuffer* Src = ResourceCast(Info.GetBuffer().Src);
+						FVulkanBuffer* Dst = ResourceCast(Info.GetBuffer().Dst);
+						FVulkanBuffer* Src = ResourceCast(Info.GetBuffer().Src);
 
 						if (Src)
 						{
diff --git a/Engine/Source/Runtime/VulkanRHI/Private/VulkanRHIPrivate.h b/Engine/Source/Runtime/VulkanRHI/Private/VulkanRHIPrivate.h
--- a/Engine/Source/Runtime/VulkanRHI/Private/VulkanRHIPrivate.h
+++ b/Engine/Source/Runtime/VulkanRHI/Private/VulkanRHIPrivate.h
@@ -34,7 +34,7 @@
 #include "Stats/Stats2.h"
 
 class FVulkanQueue;
-class FVulkanCmdBuffer;
+class FVulkanCommandBuffer;
 class FVulkanShader;
 class FVulkanDescriptorSetsLayout;
 class FVulkanGfxPipeline;
@@ -380,7 +380,7 @@ union UNvidiaDriverVersion
 };
 
 // Transitions an image to the specified layout. This does not update the layout cached internally by the RHI; the calling code must do that explicitly via FVulkanCommandListContext::GetLayoutManager() if necessary.
-void VulkanSetImageLayout(FVulkanCmdBuffer* CmdBuffer, VkImage Image, VkImageLayout OldLayout, VkImageLayout NewLayout, const VkImageSubresourceRange& SubresourceRange);
+void VulkanSetImageLayout(FVulkanCommandBuffer* CmdBuffer, VkImage Image, VkImageLayout OldLayout, VkImageLayout NewLayout, const VkImageSubresourceRange& SubresourceRange);
 
 
 
diff --git a/Engine/Source/Runtime/VulkanRHI/Private/VulkanRayTracing.cpp b/Engine/Source/Runtime/VulkanRHI/Private/VulkanRayTracing.cpp
--- a/Engine/Source/Runtime/VulkanRHI/Private/VulkanRayTracing.cpp
+++ b/Engine/Source/Runtime/VulkanRHI/Private/VulkanRayTracing.cpp
@@ -213,7 +213,7 @@ static void GetBLASBuildData(
 	const EBLASBuildDataUsage Usage,
 	FVkRtBLASBuildData& BuildData)
 {
-	FVulkanResourceMultiBuffer* const IndexBuffer = ResourceCast(IndexBufferRHI.GetReference());
+	FVulkanBuffer* const IndexBuffer = ResourceCast(IndexBufferRHI.GetReference());
 	VkDeviceOrHostAddressConstKHR IndexBufferDeviceAddress = {};
 	
 	// We only need to get IB/VB address when we are getting data for rendering. For estimating BLAS size we set them to 0.
@@ -235,7 +235,7 @@ static void GetBLASBuildData(
 	{
 		const FRayTracingGeometrySegment& Segment = Segments[SegmentIndex];
 
-		FVulkanResourceMultiBuffer* const VertexBuffer = ResourceCast(Segment.VertexBuffer.GetReference());
+		FVulkanBuffer* const VertexBuffer = ResourceCast(Segment.VertexBuffer.GetReference());
 
 		VkDeviceOrHostAddressConstKHR VertexBufferDeviceAddress = {};
 		VertexBufferDeviceAddress.deviceAddress = Usage == EBLASBuildDataUsage::Rendering
@@ -494,7 +494,7 @@ void FVulkanRayTracingGeometry::RemoveCompactionRequest()
 	}
 }
 
-void FVulkanRayTracingGeometry::CompactAccelerationStructure(FVulkanCmdBuffer& CmdBuffer, uint64 InSizeAfterCompaction)
+void FVulkanRayTracingGeometry::CompactAccelerationStructure(FVulkanCommandBuffer& CmdBuffer, uint64 InSizeAfterCompaction)
 {
 	check(bHasPendingCompactionRequests);
 	bHasPendingCompactionRequests = false;
@@ -510,12 +510,12 @@ void FVulkanRayTracingGeometry::CompactAccelerationStructure(FVulkanCmdBuffer& C
 	DEC_MEMORY_STAT_BY(STAT_VulkanRayTracingStaticBLASMemory, AccelerationStructureBuffer->GetSize());
 
 	// Move old AS into this temporary variable which gets released when this function returns	
-	TRefCountPtr<FVulkanResourceMultiBuffer> OldAccelerationStructure = AccelerationStructureBuffer;
+	TRefCountPtr<FVulkanBuffer> OldAccelerationStructure = AccelerationStructureBuffer;
 	VkAccelerationStructureKHR OldHandle = Handle;
 
 	FString DebugNameString = Initializer.DebugName.ToString();
 	FRHIResourceCreateInfo BlasBufferCreateInfo(*DebugNameString);
-	AccelerationStructureBuffer = new FVulkanResourceMultiBuffer(Device, FRHIBufferDesc(InSizeAfterCompaction, 0, BUF_AccelerationStructure), BlasBufferCreateInfo);
+	AccelerationStructureBuffer = new FVulkanBuffer(Device, FRHIBufferDesc(InSizeAfterCompaction, 0, BUF_AccelerationStructure), BlasBufferCreateInfo);
 	
 
 	INC_MEMORY_STAT_BY(STAT_VulkanRayTracingUsedVideoMemory, AccelerationStructureBuffer->GetSize());
@@ -558,7 +558,7 @@ void FVulkanRayTracingGeometry::SetupHitGroupSystemParameters()
 	const bool bIsTriangles = (Initializer.GeometryType == RTGT_Triangles);
 
 	FVulkanBindlessDescriptorManager* BindlessDescriptorManager = Device->GetBindlessDescriptorManager();
-	auto GetBindlessHandle = [BindlessDescriptorManager](FVulkanResourceMultiBuffer* Buffer, uint32 ExtraOffset) 
+	auto GetBindlessHandle = [BindlessDescriptorManager](FVulkanBuffer* Buffer, uint32 ExtraOffset)
 	{
 		if (Buffer)
 		{
@@ -573,13 +573,13 @@ void FVulkanRayTracingGeometry::SetupHitGroupSystemParameters()
 
 	HitGroupSystemParameters.Reset(Initializer.Segments.Num());
 
-	FVulkanResourceMultiBuffer* IndexBuffer = ResourceCast(Initializer.IndexBuffer.GetReference());
+	FVulkanBuffer* IndexBuffer = ResourceCast(Initializer.IndexBuffer.GetReference());
 	const uint32 IndexStride = IndexBuffer ? IndexBuffer->GetStride() : 0;
 	HitGroupSystemIndexView = GetBindlessHandle(IndexBuffer, 0);
 
 	for (const FRayTracingGeometrySegment& Segment : Initializer.Segments)
 	{
-		FVulkanResourceMultiBuffer* VertexBuffer = ResourceCast(Segment.VertexBuffer.GetReference());
+		FVulkanBuffer* VertexBuffer = ResourceCast(Segment.VertexBuffer.GetReference());
 		const FRHIDescriptorHandle VBHandle = GetBindlessHandle(VertexBuffer, Segment.VertexBufferOffset);
 		HitGroupSystemVertexViews.Add(VBHandle);
 
@@ -617,7 +617,7 @@ void FVulkanRayTracingGeometry::ReleaseBindlessHandles()
 void FVulkanRayTracingGeometry::SetupInlineGeometryParameters(uint32 GeometrySegmentIndex, FVulkanRayTracingGeometryParameters& Parameters) const
 {
 	const FRayTracingGeometryInitializer& GeometryInitializer = GetInitializer();
-	const FVulkanResourceMultiBuffer* IndexBuffer = ResourceCast(GeometryInitializer.IndexBuffer.GetReference());
+	const FVulkanBuffer* IndexBuffer = ResourceCast(GeometryInitializer.IndexBuffer.GetReference());
 
 	const uint32 IndexStride = IndexBuffer ? IndexBuffer->GetStride() : 0;
 	const uint32 IndexOffsetInBytes = GeometryInitializer.IndexBufferOffset;
@@ -625,7 +625,7 @@ void FVulkanRayTracingGeometry::SetupInlineGeometryParameters(uint32 GeometrySeg
 
 	const FRayTracingGeometrySegment& Segment = GeometryInitializer.Segments[GeometrySegmentIndex];
 
-	const FVulkanResourceMultiBuffer* VertexBuffer = ResourceCast(Segment.VertexBuffer.GetReference());
+	const FVulkanBuffer* VertexBuffer = ResourceCast(Segment.VertexBuffer.GetReference());
 	checkf(VertexBuffer, TEXT("All ray tracing geometry segments must have a valid vertex buffer"));
 	const VkDeviceAddress VertexBufferAddress = VertexBuffer->GetDeviceAddress();
 
@@ -766,8 +766,8 @@ void FVulkanRayTracingScene::BindBuffer(FRHIBuffer* InBuffer, uint32 InBufferOff
 void BuildAccelerationStructure(
 	FVulkanCommandListContext& CommandContext,
 	FVulkanRayTracingScene& Scene,
-	FVulkanResourceMultiBuffer* InScratchBuffer, uint32 InScratchOffset,
-	FVulkanResourceMultiBuffer* InInstanceBuffer, uint32 InInstanceOffset,
+	FVulkanBuffer* InScratchBuffer, uint32 InScratchOffset,
+	FVulkanBuffer* InInstanceBuffer, uint32 InInstanceOffset,
 	uint32 NumInstances,
 	EAccelerationStructureBuildMode BuildMode)
 {
@@ -844,7 +844,7 @@ void BuildAccelerationStructure(
 	}
 
 	FVulkanCommandBufferManager& CommandBufferManager = *CommandContext.GetCommandBufferManager();
-	FVulkanCmdBuffer* const CmdBuffer = CommandBufferManager.GetActiveCmdBuffer();
+	FVulkanCommandBuffer* const CmdBuffer = CommandBufferManager.GetActiveCmdBuffer();
 
 	// Force a memory barrier to make sure all previous builds ops are finished before building the TLAS
 	AddAccelerationStructureBuildBarrier(CmdBuffer->GetHandle());
@@ -1026,7 +1026,7 @@ FRHIShaderResourceView* FVulkanRayTracingShaderTable::GetOrCreateInlineBufferSRV
 void FVulkanRayTracingShaderTable::Commit(FVulkanCommandListContext& Context)
 {
 	FVulkanCommandBufferManager* CommandBufferManager = Context.GetCommandBufferManager();
-	FVulkanCmdBuffer* const CmdBuffer = CommandBufferManager->GetActiveCmdBuffer();
+	FVulkanCommandBuffer* const CmdBuffer = CommandBufferManager->GetActiveCmdBuffer();
 
 	VkMemoryBarrier BarrierBefore = { VK_STRUCTURE_TYPE_MEMORY_BARRIER, nullptr, VK_ACCESS_MEMORY_READ_BIT | VK_ACCESS_MEMORY_WRITE_BIT, VK_ACCESS_TRANSFER_READ_BIT | VK_ACCESS_TRANSFER_WRITE_BIT };
 	VulkanRHI::vkCmdPipelineBarrier(CmdBuffer->GetHandle(), VK_PIPELINE_STAGE_BOTTOM_OF_PIPE_BIT, VK_PIPELINE_STAGE_TRANSFER_BIT, 0, 1, &BarrierBefore, 0, nullptr, 0, nullptr);
@@ -1260,7 +1260,7 @@ void FVulkanCommandListContext::RHIBuildAccelerationStructures(const TArrayView<
 		ScratchBufferRange.Size, ScratchBufferRange.Offset, ScratchBufferRange.Buffer->GetSize());
 
 	const uint64 ScratchAlignment = GRHIRayTracingScratchBufferAlignment;
-	FVulkanResourceMultiBuffer* ScratchBuffer = ResourceCast(ScratchBufferRange.Buffer);
+	FVulkanBuffer* ScratchBuffer = ResourceCast(ScratchBufferRange.Buffer);
 	uint32 ScratchBufferOffset = ScratchBufferRange.Offset;
 
 	TArray<FVkRtBLASBuildData, TInlineAllocator<32>> TempBuildData;
@@ -1323,7 +1323,7 @@ void FVulkanCommandListContext::RHIBuildAccelerationStructures(const TArrayView<
 		Geometry->SetupHitGroupSystemParameters();
 	}
 	
-	FVulkanCmdBuffer* const CmdBuffer = CommandBufferManager->GetActiveCmdBuffer();
+	FVulkanCommandBuffer* const CmdBuffer = CommandBufferManager->GetActiveCmdBuffer();
 	VulkanDynamicAPI::vkCmdBuildAccelerationStructuresKHR(CmdBuffer->GetHandle(), Params.Num(), BuildGeometryInfos.GetData(), BuildRangeInfos.GetData());
 
 	// Add an acceleration structure build barrier after each acceleration structure build batch.
@@ -1351,8 +1351,8 @@ void FVulkanCommandListContext::RHIBuildAccelerationStructures(const TArrayView<
 void FVulkanCommandListContext::RHIBuildAccelerationStructure(const FRayTracingSceneBuildParams& SceneBuildParams)
 {
 	FVulkanRayTracingScene* const Scene = ResourceCast(SceneBuildParams.Scene);
-	FVulkanResourceMultiBuffer* const ScratchBuffer = ResourceCast(SceneBuildParams.ScratchBuffer);
-	FVulkanResourceMultiBuffer* const InstanceBuffer = ResourceCast(SceneBuildParams.InstanceBuffer);
+	FVulkanBuffer* const ScratchBuffer = ResourceCast(SceneBuildParams.ScratchBuffer);
+	FVulkanBuffer* const InstanceBuffer = ResourceCast(SceneBuildParams.InstanceBuffer);
 
 	Scene->ReferencedGeometries.Reserve(SceneBuildParams.ReferencedGeometries.Num());
 
@@ -1762,14 +1762,14 @@ const TArray<uint8>& FVulkanRayTracingPipelineState::GetShaderHandles(EShaderFre
 	return GetShaderData(Frequency).ShaderHandles;
 }
 
-void FVulkanRayTracingCompactedSizeQueryPool::EndBatch(FVulkanCmdBuffer* InCmdBuffer)
+void FVulkanRayTracingCompactedSizeQueryPool::EndBatch(FVulkanCommandBuffer* InCmdBuffer)
 {
 	check(CmdBuffer == nullptr);
 	CmdBuffer = InCmdBuffer;
 	FenceSignaledCounter = InCmdBuffer->GetFenceSignaledCounter();
 }
 
-void FVulkanRayTracingCompactedSizeQueryPool::Reset(FVulkanCmdBuffer* InCmdBuffer)
+void FVulkanRayTracingCompactedSizeQueryPool::Reset(FVulkanCommandBuffer* InCmdBuffer)
 {
 	VulkanRHI::vkCmdResetQueryPool(InCmdBuffer->GetHandle(), QueryPool, 0, MaxQueries);
 	FenceSignaledCounter = 0;
@@ -1854,7 +1854,7 @@ void FVulkanRayTracingCompactionRequestHandler::Update(FVulkanCommandListContext
 	if (ActiveBLASes.Num() > 0)
 	{		
 		FVulkanCommandBufferManager& CommandBufferManager = *InCommandContext.GetCommandBufferManager();
-		FVulkanCmdBuffer* const CmdBuffer = CommandBufferManager.GetActiveCmdBuffer();
+		FVulkanCommandBuffer* const CmdBuffer = CommandBufferManager.GetActiveCmdBuffer();
 
 		if (QueryPool->TryGetResults(ActiveBLASes.Num()))
 		{
@@ -1918,7 +1918,7 @@ void FVulkanRayTracingCompactionRequestHandler::Update(FVulkanCommandListContext
 		PendingRequests.RemoveAt(0, ActiveRequests.Num(), EAllowShrinking::No);
 
 		FVulkanCommandBufferManager& CommandBufferManager = *InCommandContext.GetCommandBufferManager();
-		FVulkanCmdBuffer* const CmdBuffer = CommandBufferManager.GetActiveCmdBuffer();
+		FVulkanCommandBuffer* const CmdBuffer = CommandBufferManager.GetActiveCmdBuffer();
 
 		// Barrier here is not stricly necessary as it is added after the build.
 		// AddAccelerationStructureBuildBarrier(CmdBuffer->GetHandle());
@@ -1939,14 +1939,14 @@ void FVulkanRayTracingCompactionRequestHandler::Update(FVulkanCommandListContext
 	}
 }
 
-bool FVulkanRayTracingCompactionRequestHandler::IsUsingCmdBuffer(FVulkanCmdBuffer* CmdBuffer)
+bool FVulkanRayTracingCompactionRequestHandler::IsUsingCmdBuffer(FVulkanCommandBuffer* CmdBuffer)
 {
 	FScopeLock Lock(&CS);
 	return (CmdBuffer == ActiveRequestsCmdBuffer) || (CmdBuffer == QueryPool->CmdBuffer);
 }
 
 
-static FVulkanPipelineBarrier SetRayGenResources(FVulkanDevice* Device, FVulkanCmdBuffer* const CmdBuffer, const FRayTracingShaderBindings& InGlobalResourceBindings, FVulkanRayTracingShaderTable* ShaderTable)
+static FVulkanPipelineBarrier SetRayGenResources(FVulkanDevice* Device, FVulkanCommandBuffer* const CmdBuffer, const FRayTracingShaderBindings& InGlobalResourceBindings, FVulkanRayTracingShaderTable* ShaderTable)
 {
 	TArray<const FVulkanUniformBuffer*> UniformBuffers;
 	UniformBuffers.Reserve(UE_ARRAY_COUNT(InGlobalResourceBindings.UniformBuffers));
@@ -2152,7 +2152,7 @@ void FVulkanCommandListContext::RHIRayTraceDispatch(
 	FVulkanRayTracingShader* RayGenShader = ResourceCast(InRayGenShader);
 	FVulkanRayTracingShaderTable* ShaderTable = ResourceCast(InSBT);
 
-	FVulkanCmdBuffer* const CmdBuffer = GetCommandBufferManager()->GetActiveCmdBuffer();
+	FVulkanCommandBuffer* const CmdBuffer = GetCommandBufferManager()->GetActiveCmdBuffer();
 	VulkanRHI::vkCmdBindPipeline(CmdBuffer->GetHandle(), VK_PIPELINE_BIND_POINT_RAY_TRACING_KHR, Pipeline->GetPipeline());
 
 	checkSlow(InRayGenShader->GetFrequency() == SF_RayGen);
@@ -2189,7 +2189,7 @@ void FVulkanCommandListContext::RHIRayTraceDispatchIndirect(
 	FVulkanRayTracingShader* RayGenShader = ResourceCast(InRayGenShader);
 	FVulkanRayTracingShaderTable* ShaderTable = ResourceCast(InSBT);
 
-	FVulkanCmdBuffer* const CmdBuffer = GetCommandBufferManager()->GetActiveCmdBuffer();
+	FVulkanCommandBuffer* const CmdBuffer = GetCommandBufferManager()->GetActiveCmdBuffer();
 	VulkanRHI::vkCmdBindPipeline(CmdBuffer->GetHandle(), VK_PIPELINE_BIND_POINT_RAY_TRACING_KHR, Pipeline->GetPipeline());
 
 	checkSlow(InRayGenShader->GetFrequency() == SF_RayGen);
@@ -2197,7 +2197,7 @@ void FVulkanCommandListContext::RHIRayTraceDispatchIndirect(
 
 	FVulkanPipelineBarrier PostDispatch = SetRayGenResources(Device, CmdBuffer, InGlobalResourceBindings, ShaderTable);
 
-	FVulkanResourceMultiBuffer* ArgumentBuffer = ResourceCast(InArgumentBuffer);
+	FVulkanBuffer* ArgumentBuffer = ResourceCast(InArgumentBuffer);
 	const VkDeviceAddress IndirectDeviceAddress = ArgumentBuffer->GetDeviceAddress() + InArgumentOffset;
 
 	VulkanRHI::vkCmdTraceRaysIndirectKHR(
diff --git a/Engine/Source/Runtime/VulkanRHI/Private/VulkanRayTracing.h b/Engine/Source/Runtime/VulkanRHI/Private/VulkanRayTracing.h
--- a/Engine/Source/Runtime/VulkanRHI/Private/VulkanRayTracing.h
+++ b/Engine/Source/Runtime/VulkanRHI/Private/VulkanRayTracing.h
@@ -7,7 +7,7 @@
 #include "RayTracingBuiltInResources.h"
 
 class FVulkanCommandListContext;
-class FVulkanResourceMultiBuffer;
+class FVulkanBuffer;
 struct FVulkanRayTracingGeometryParameters;
 
 class FVulkanRayTracingPlatform
@@ -158,7 +158,7 @@ private:
 		
 	// Buffer that contains per-hitrecord index and vertex buffer binding data
 	TArray<uint8> InlineGeometryParameterData;
-	TRefCountPtr<FVulkanResourceMultiBuffer> InlineGeometryParameterBuffer;
+	TRefCountPtr<FVulkanBuffer> InlineGeometryParameterBuffer;
 	FShaderResourceViewRHIRef InlineGeometryParameterSRV;
 
 	TArray<TRefCountPtr<FRHIUniformBuffer>> ReferencedUniformBuffers;
@@ -185,7 +185,7 @@ public:
 	using FRHIRayTracingGeometry::SizeInfo;
 	
 	void RemoveCompactionRequest();
-	void CompactAccelerationStructure(FVulkanCmdBuffer& CmdBuffer, uint64 InSizeAfterCompaction);
+	void CompactAccelerationStructure(FVulkanCommandBuffer& CmdBuffer, uint64 InSizeAfterCompaction);
 
 	void SetupHitGroupSystemParameters();
 	void ReleaseBindlessHandles();
@@ -194,7 +194,7 @@ public:
 
 	VkAccelerationStructureKHR Handle = VK_NULL_HANDLE;
 	VkDeviceAddress Address = 0;
-	TRefCountPtr<FVulkanResourceMultiBuffer> AccelerationStructureBuffer;
+	TRefCountPtr<FVulkanBuffer> AccelerationStructureBuffer;
 	bool bHasPendingCompactionRequests = false;
 	uint64 AccelerationStructureCompactedSize = 0;
 
@@ -257,7 +257,7 @@ public:
 
 	uint32 NumInstances = 0;
 	
-	TRefCountPtr<FVulkanResourceMultiBuffer> AccelerationStructureBuffer;
+	TRefCountPtr<FVulkanBuffer> AccelerationStructureBuffer;
 	
 	TMap<const FVulkanRayTracingPipelineState*, TRefCountPtr<FVulkanRayTracingShaderTable>> ShaderTables;
 
@@ -317,11 +317,11 @@ public:
 		: FVulkanQueryPool(InDevice, nullptr, InMaxQueries, VK_QUERY_TYPE_ACCELERATION_STRUCTURE_COMPACTED_SIZE_KHR, false)
 	{}
 
-	void EndBatch(FVulkanCmdBuffer* InCmdBuffer);
+	void EndBatch(FVulkanCommandBuffer* InCmdBuffer);
 	bool TryGetResults(uint32 NumResults);
-	void Reset(FVulkanCmdBuffer* InCmdBuffer);
+	void Reset(FVulkanCommandBuffer* InCmdBuffer);
 
-	FVulkanCmdBuffer* CmdBuffer = nullptr;
+	FVulkanCommandBuffer* CmdBuffer = nullptr;
 	uint64 FenceSignaledCounter = 0;
 };
 
@@ -343,7 +343,7 @@ public:
 
 	void Update(FVulkanCommandListContext& InCommandContext);
 
-	bool IsUsingCmdBuffer(FVulkanCmdBuffer* CmdBuffer);
+	bool IsUsingCmdBuffer(FVulkanCommandBuffer* CmdBuffer);
 
 private:
 
@@ -353,7 +353,7 @@ private:
 
 	// Keep references on FVulkanRayTracingGeometry until lifetime issue is found (this prevents cancellation)
 	TArray<TRefCountPtr<FVulkanRayTracingGeometry>> ActiveRequests;
-	FVulkanCmdBuffer* ActiveRequestsCmdBuffer = nullptr;
+	FVulkanCommandBuffer* ActiveRequestsCmdBuffer = nullptr;
 	uint64 ActiveRequestsFenceCounter = MAX_uint64;
 
 	FVulkanRayTracingCompactedSizeQueryPool* QueryPool = nullptr;
diff --git a/Engine/Source/Runtime/VulkanRHI/Private/VulkanRenderTarget.cpp b/Engine/Source/Runtime/VulkanRHI/Private/VulkanRenderTarget.cpp
--- a/Engine/Source/Runtime/VulkanRHI/Private/VulkanRenderTarget.cpp
+++ b/Engine/Source/Runtime/VulkanRHI/Private/VulkanRenderTarget.cpp
@@ -179,7 +179,7 @@ void FVulkanDynamicRHI::RHIReadSurfaceData(FRHITexture* TextureRHI, FIntRect Rec
 	FVulkanCommandListContext& ImmediateContext = Device->GetImmediateContext();
 
 	VulkanRHI::FStagingBuffer* StagingBuffer = nullptr;
-	FVulkanCmdBuffer* CmdBuffer = nullptr;
+	FVulkanCommandBuffer* CmdBuffer = nullptr;
 	const bool bCPUReadback = EnumHasAllFlags(Surface.GetDesc().Flags, TexCreate_CPUReadback);
 	if (!bCPUReadback) //this function supports reading back arbitrary rendertargets, so if its not a cpu readback surface, we do a copy.
 	{
@@ -298,7 +298,7 @@ void FVulkanDynamicRHI::RHIUnmapStagingSurface(FRHITexture* TextureRHI, uint32 G
 
 void FVulkanDynamicRHI::RHIReadSurfaceFloatData(FRHITexture* TextureRHI, FIntRect Rect, TArray<FFloat16Color>& OutData, ECubeFace CubeFace, int32 ArrayIndex, int32 MipIndex)
 {
-	auto DoCopyFloat = [](FVulkanDevice* InDevice, FVulkanCmdBuffer* InCmdBuffer, const FVulkanTexture& Surface, uint32 InMipIndex, uint32 SrcBaseArrayLayer, FIntRect InRect, TArray<FFloat16Color>& OutputData)
+	auto DoCopyFloat = [](FVulkanDevice* InDevice, FVulkanCommandBuffer* InCmdBuffer, const FVulkanTexture& Surface, uint32 InMipIndex, uint32 SrcBaseArrayLayer, FIntRect InRect, TArray<FFloat16Color>& OutputData)
 	{
 		ensure(Surface.StorageFormat == VK_FORMAT_R16G16B16A16_SFLOAT);
 
@@ -402,7 +402,7 @@ void FVulkanDynamicRHI::RHIReadSurfaceFloatData(FRHITexture* TextureRHI, FIntRec
 	{
 		Device->PrepareForCPURead();
 
-		FVulkanCmdBuffer* CmdBuffer = Device->GetImmediateContext().GetCommandBufferManager()->GetUploadCmdBuffer();
+		FVulkanCommandBuffer* CmdBuffer = Device->GetImmediateContext().GetCommandBufferManager()->GetUploadCmdBuffer();
 		switch (TextureRHI->GetDesc().Dimension)
 		{
 			case ETextureDimension::TextureCubeArray:
@@ -445,7 +445,7 @@ void FVulkanDynamicRHI::RHIRead3DSurfaceFloatData(FRHITexture* TextureRHI, FIntR
 	}
 
 	Device->PrepareForCPURead();
-	FVulkanCmdBuffer* CmdBuffer = Device->GetImmediateContext().GetCommandBufferManager()->GetUploadCmdBuffer();
+	FVulkanCommandBuffer* CmdBuffer = Device->GetImmediateContext().GetCommandBufferManager()->GetUploadCmdBuffer();
 
 	ensure(Surface.StorageFormat == VK_FORMAT_R16G16B16A16_SFLOAT);
 
@@ -567,7 +567,7 @@ bool FVulkanCommandListContext::IsSwapchainImage(FRHITexture* InTexture) const
 
 void FVulkanCommandListContext::RHIBeginRenderPass(const FRHIRenderPassInfo& InInfo, const TCHAR* InName)
 {
-	FVulkanCmdBuffer* CmdBuffer = CommandBufferManager->GetActiveCmdBuffer();
+	FVulkanCommandBuffer* CmdBuffer = CommandBufferManager->GetActiveCmdBuffer();
 
 	if (GVulkanSubmitAfterEveryEndRenderPass)
 	{
@@ -682,7 +682,7 @@ void FVulkanCommandListContext::RHIBeginRenderPass(const FRHIRenderPassInfo& InI
 void FVulkanCommandListContext::RHIEndRenderPass()
 {
 	const bool bHasOcclusionQueries = (RenderPassInfo.NumOcclusionQueries > 0);
-	FVulkanCmdBuffer* CmdBuffer = CommandBufferManager->GetActiveCmdBuffer();
+	FVulkanCommandBuffer* CmdBuffer = CommandBufferManager->GetActiveCmdBuffer();
 
 	if (bHasOcclusionQueries)
 	{
@@ -705,7 +705,7 @@ void FVulkanCommandListContext::RHIEndRenderPass()
 void FVulkanCommandListContext::RHINextSubpass()
 {
 	check(CurrentRenderPass);
-	FVulkanCmdBuffer* CmdBuffer = CommandBufferManager->GetActiveCmdBuffer();
+	FVulkanCommandBuffer* CmdBuffer = CommandBufferManager->GetActiveCmdBuffer();
 	VkCommandBuffer Cmd = CmdBuffer->GetHandle();
 	VulkanRHI::vkCmdNextSubpass(Cmd, VK_SUBPASS_CONTENTS_INLINE);
 }
diff --git a/Engine/Source/Runtime/VulkanRHI/Private/VulkanRenderpass.cpp b/Engine/Source/Runtime/VulkanRHI/Private/VulkanRenderpass.cpp
--- a/Engine/Source/Runtime/VulkanRHI/Private/VulkanRenderpass.cpp
+++ b/Engine/Source/Runtime/VulkanRHI/Private/VulkanRenderpass.cpp
@@ -126,7 +126,7 @@ FVulkanFramebuffer* FVulkanRenderPassManager::GetOrCreateFramebuffer(const FRHIS
 	return Framebuffer;
 }
 
-void FVulkanRenderPassManager::BeginRenderPass(FVulkanCommandListContext& Context, FVulkanDevice& InDevice, FVulkanCmdBuffer* CmdBuffer, const FRHIRenderPassInfo& RPInfo, const FVulkanRenderTargetLayout& RTLayout, FVulkanRenderPass* RenderPass, FVulkanFramebuffer* Framebuffer)
+void FVulkanRenderPassManager::BeginRenderPass(FVulkanCommandListContext& Context, FVulkanDevice& InDevice, FVulkanCommandBuffer* CmdBuffer, const FRHIRenderPassInfo& RPInfo, const FVulkanRenderTargetLayout& RTLayout, FVulkanRenderPass* RenderPass, FVulkanFramebuffer* Framebuffer)
 {
 	// (NumRT + 1 [Depth] ) * 2 [surface + resolve]
 	VkClearValue ClearValues[(MaxSimultaneousRenderTargets + 1) * 2];
@@ -214,7 +214,7 @@ void FVulkanRenderPassManager::BeginRenderPass(FVulkanCommandListContext& Contex
 	}
 }
 
-void FVulkanRenderPassManager::EndRenderPass(FVulkanCmdBuffer* CmdBuffer)
+void FVulkanRenderPassManager::EndRenderPass(FVulkanCommandBuffer* CmdBuffer)
 {
 	CmdBuffer->EndRenderPass();
 
diff --git a/Engine/Source/Runtime/VulkanRHI/Private/VulkanRenderpass.h b/Engine/Source/Runtime/VulkanRHI/Private/VulkanRenderpass.h
--- a/Engine/Source/Runtime/VulkanRHI/Private/VulkanRenderpass.h
+++ b/Engine/Source/Runtime/VulkanRHI/Private/VulkanRenderpass.h
@@ -888,8 +888,8 @@ public:
 		return RenderPass;
 	}
 
-	void BeginRenderPass(FVulkanCommandListContext& Context, FVulkanDevice& InDevice, FVulkanCmdBuffer* CmdBuffer, const FRHIRenderPassInfo& RPInfo, const FVulkanRenderTargetLayout& RTLayout, FVulkanRenderPass* RenderPass, FVulkanFramebuffer* Framebuffer);
-	void EndRenderPass(FVulkanCmdBuffer* CmdBuffer);
+	void BeginRenderPass(FVulkanCommandListContext& Context, FVulkanDevice& InDevice, FVulkanCommandBuffer* CmdBuffer, const FRHIRenderPassInfo& RPInfo, const FVulkanRenderTargetLayout& RTLayout, FVulkanRenderPass* RenderPass, FVulkanFramebuffer* Framebuffer);
+	void EndRenderPass(FVulkanCommandBuffer* CmdBuffer);
 
 	FRWLock RenderPassesLock;
 	FRWLock FramebuffersLock;
diff --git a/Engine/Source/Runtime/VulkanRHI/Private/VulkanResources.h b/Engine/Source/Runtime/VulkanRHI/Private/VulkanResources.h
--- a/Engine/Source/Runtime/VulkanRHI/Private/VulkanResources.h
+++ b/Engine/Source/Runtime/VulkanRHI/Private/VulkanResources.h
@@ -19,9 +19,9 @@
 
 class FVulkanDevice;
 class FVulkanQueue;
-class FVulkanCmdBuffer;
+class FVulkanCommandBuffer;
 class FVulkanTexture;
-class FVulkanResourceMultiBuffer;
+class FVulkanBuffer;
 class FVulkanLayout;
 class FVulkanOcclusionQuery;
 class FVulkanCommandBufferManager;
@@ -465,9 +465,9 @@ public:
 	FStructuredBufferView      const& GetStructuredBufferView     () const { return Storage.Get<FStructuredBufferView     >(); }
 	FAccelerationStructureView const& GetAccelerationStructureView() const { return Storage.Get<FAccelerationStructureView>(); }
 
-	// NOTE: The InOffset applies to the FVulkanResourceMultiBuffer (it does not include any internal Allocation offsets that may exist)
+	// NOTE: The InOffset applies to the FVulkanBuffer (it does not include any internal Allocation offsets that may exist)
 	FVulkanView* InitAsTypedBufferView(
-		  FVulkanResourceMultiBuffer* Buffer
+		  FVulkanBuffer* Buffer
 		, EPixelFormat Format
 		, uint32 InOffset
 		, uint32 InSize);
@@ -486,14 +486,14 @@ public:
 		, VkImageUsageFlags ImageUsageFlags = 0
 		, VkSamplerYcbcrConversion SamplerYcbcrConversion = nullptr);
 
-	// NOTE: The InOffset applies to the FVulkanResourceMultiBuffer (it does not include any internal Allocation offsets that may exist)
+	// NOTE: The InOffset applies to the FVulkanBuffer (it does not include any internal Allocation offsets that may exist)
 	FVulkanView* InitAsStructuredBufferView(
-		  FVulkanResourceMultiBuffer* Buffer
+		FVulkanBuffer* Buffer
 		, uint32 InOffset
 		, uint32 InSize);
 
 	FVulkanView* InitAsAccelerationStructureView(
-		  FVulkanResourceMultiBuffer* Buffer
+		FVulkanBuffer* Buffer
 		, uint32 Offset
 		, uint32 Size);
 
@@ -832,7 +832,7 @@ public:
 		AllocatedQueries[Index] = nullptr;
 	}
 
-	inline void EndBatch(FVulkanCmdBuffer* InCmdBuffer)
+	inline void EndBatch(FVulkanCommandBuffer* InCmdBuffer)
 	{
 		ensure(State == EState::RHIT_PostBeginBatch);
 		State = EState::RHIT_PostEndBatch;
@@ -856,7 +856,7 @@ public:
 		return false;
 	}
 
-	void Reset(FVulkanCmdBuffer* InCmdBuffer, uint32 InFrameNumber);
+	void Reset(FVulkanCommandBuffer* InCmdBuffer, uint32 InFrameNumber);
 
 	bool IsStalePool() const;
 
@@ -878,9 +878,9 @@ protected:
 	TArray<FVulkanOcclusionQuery*> AllocatedQueries;
 	TArray<uint64> AcquiredIndices;
 	bool InternalTryGetResults(bool bWait);
-	void SetFence(FVulkanCmdBuffer* InCmdBuffer);
+	void SetFence(FVulkanCommandBuffer* InCmdBuffer);
 
-	FVulkanCmdBuffer* CmdBuffer = nullptr;
+	FVulkanCommandBuffer* CmdBuffer = nullptr;
 	uint64 FenceCounter = UINT64_MAX;
 	uint32 FrameNumber = UINT32_MAX;
 };
@@ -901,7 +901,7 @@ public:
 
 	struct FCmdBufferFence
 	{
-		FVulkanCmdBuffer* CmdBuffer;
+		FVulkanCommandBuffer* CmdBuffer;
 		uint64 FenceCounter;
 		uint64 FrameCount = UINT64_MAX;
 		uint32 Attempts = 0;
@@ -959,11 +959,11 @@ public:
 	FVulkanTimingQueryPool* Pool = nullptr;
 };
 
-class FVulkanResourceMultiBuffer : public FRHIBuffer, public VulkanRHI::FDeviceChild, public FVulkanViewableResource
+class FVulkanBuffer : public FRHIBuffer, public VulkanRHI::FDeviceChild, public FVulkanViewableResource
 {
 public:
-	FVulkanResourceMultiBuffer(FVulkanDevice* InDevice, FRHIBufferDesc const& InBufferDesc, FRHIResourceCreateInfo& CreateInfo, class FRHICommandListBase* InRHICmdList = nullptr, const FRHITransientHeapAllocation* InTransientHeapAllocation = nullptr);
-	virtual ~FVulkanResourceMultiBuffer();
+	FVulkanBuffer(FVulkanDevice* InDevice, FRHIBufferDesc const& InBufferDesc, FRHIResourceCreateInfo& CreateInfo, class FRHICommandListBase* InRHICmdList = nullptr, const FRHITransientHeapAllocation* InTransientHeapAllocation = nullptr);
+	virtual ~FVulkanBuffer();
 
 	inline const VulkanRHI::FVulkanAllocation& GetCurrentAllocation() const
 	{
@@ -1010,13 +1010,13 @@ public:
 	void* Lock(FRHICommandListBase& RHICmdList, EResourceLockMode LockMode, uint32 Size, uint32 Offset);
 	void Unlock(FRHICommandListBase& RHICmdList);
 
-	void TakeOwnership(FVulkanResourceMultiBuffer& Other);
+	void TakeOwnership(FVulkanBuffer& Other);
 	void ReleaseOwnership();
 
 	template<typename T>
 	void DumpMemory(T Callback)
 	{
-		Callback(TEXT("FVulkanResourceMultiBuffer"), FName(), this, 0, GetCurrentSize(), 1, 1, VK_FORMAT_UNDEFINED);
+		Callback(TEXT("FVulkanBuffer"), FName(), this, 0, GetCurrentSize(), 1, 1, VK_FORMAT_UNDEFINED);
 	}
 
 	static VkBufferUsageFlags UEToVKBufferUsageFlags(FVulkanDevice* InDevice, EBufferUsageFlags InUEUsage, bool bZeroSize);
@@ -1251,7 +1251,7 @@ public:
 	virtual void Wait(FRHICommandListImmediate& RHICmdList, FRHIGPUMask GPUMask) const override;
 
 protected:
-	FVulkanCmdBuffer*	CmdBuffer = nullptr;
+	FVulkanCommandBuffer*	CmdBuffer = nullptr;
 	uint64				FenceSignaledCounter = MAX_uint64;
 
 	friend class FVulkanCommandListContext;
@@ -1309,7 +1309,7 @@ struct TVulkanResourceTraits<FRHIUniformBuffer>
 template<>
 struct TVulkanResourceTraits<FRHIBuffer>
 {
-	typedef FVulkanResourceMultiBuffer TConcreteType;
+	typedef FVulkanBuffer TConcreteType;
 };
 template<>
 struct TVulkanResourceTraits<FRHIShaderResourceView>
diff --git a/Engine/Source/Runtime/VulkanRHI/Private/VulkanTexture.cpp b/Engine/Source/Runtime/VulkanRHI/Private/VulkanTexture.cpp
--- a/Engine/Source/Runtime/VulkanRHI/Private/VulkanTexture.cpp
+++ b/Engine/Source/Runtime/VulkanRHI/Private/VulkanTexture.cpp
@@ -141,7 +141,7 @@ static void VulkanTextureDestroyed(const FRHITextureDesc& TextureDesc, uint64 Si
 
 inline void FVulkanTexture::InternalLockWrite(FVulkanCommandListContext& Context, FVulkanTexture* Surface, const VkBufferImageCopy& Region, VulkanRHI::FStagingBuffer* StagingBuffer)
 {
-	FVulkanCmdBuffer* CmdBuffer = Context.GetCommandBufferManager()->GetUploadCmdBuffer();
+	FVulkanCommandBuffer* CmdBuffer = Context.GetCommandBufferManager()->GetUploadCmdBuffer();
 	ensure(CmdBuffer->IsOutsideRenderPass());
 	VkCommandBuffer StagingCommandBuffer = CmdBuffer->GetHandle();
 
@@ -566,7 +566,7 @@ struct FRHICommandSetInitialImageState final : public FRHICommand<FRHICommandSet
 		
 		if (bOnlyAddToLayoutManager)
 		{
-			FVulkanCmdBuffer* CmdBuffer = Context.GetCommandBufferManager()->GetActiveCmdBuffer();
+			FVulkanCommandBuffer* CmdBuffer = Context.GetCommandBufferManager()->GetActiveCmdBuffer();
 			CmdBuffer->GetLayoutManager().SetFullLayout(*VulkanTexture, InitialLayout, true);
 		}
 		else
@@ -686,7 +686,7 @@ void FVulkanTexture::InternalMoveSurface(FVulkanDevice& InDevice, FVulkanCommand
 	DestAllocation.BindImage(&InDevice, MovedImage);
 
 	// Copy Original -> Moved
-	FVulkanCmdBuffer* CmdBuffer = Context.GetCommandBufferManager()->GetActiveCmdBuffer();
+	FVulkanCommandBuffer* CmdBuffer = Context.GetCommandBufferManager()->GetActiveCmdBuffer();
 	VkCommandBuffer VkCmdBuffer = CmdBuffer->GetHandle();
 	ensure(CmdBuffer->IsOutsideRenderPass());
 
@@ -914,7 +914,7 @@ void FVulkanTexture::SetInitialImageState(FVulkanCommandListContext& Context, Vk
 	// Can't use TransferQueue as Vulkan requires that queue to also have Gfx or Compute capabilities...
 	//#todo-rco: This function is only used during loading currently, if used for regular RHIClear then use the ActiveCmdBuffer
 	// NOTE: Transient resources' memory might have belonged to another resource earlier in the ActiveCmdBuffer, so we can't use UploadCmdBuffer
-	FVulkanCmdBuffer* CmdBuffer = bIsTransientResource ? Context.GetCommandBufferManager()->GetActiveCmdBuffer() : Context.GetCommandBufferManager()->GetUploadCmdBuffer();
+	FVulkanCommandBuffer* CmdBuffer = bIsTransientResource ? Context.GetCommandBufferManager()->GetActiveCmdBuffer() : Context.GetCommandBufferManager()->GetUploadCmdBuffer();
 	ensure(CmdBuffer->IsOutsideRenderPass());
 
 	VkImageSubresourceRange SubresourceRange = FVulkanPipelineBarrier::MakeSubresourceRange(FullAspectMask);
@@ -1110,7 +1110,7 @@ static void DoAsyncReallocateTexture2D(FVulkanCommandListContext& Context, FVulk
 	const uint32 SourceFirstMip = OldTexture->GetNumMips() - NumSharedMips;
 	const uint32 DestFirstMip = NewTexture->GetNumMips() - NumSharedMips;
 
-	FVulkanCmdBuffer* CmdBuffer = Context.GetCommandBufferManager()->GetUploadCmdBuffer();
+	FVulkanCommandBuffer* CmdBuffer = Context.GetCommandBufferManager()->GetUploadCmdBuffer();
 	ensure(CmdBuffer->IsOutsideRenderPass());
 
 	VkCommandBuffer StagingCommandBuffer = CmdBuffer->GetHandle();
@@ -1865,7 +1865,7 @@ FVulkanTexture::FVulkanTexture(FVulkanDevice& InDevice, const FRHITextureCreateD
 				FVulkanCommandListContext& Context = InDevice.GetImmediateContext();
 				if (bOnlyAddToLayoutManager)
 				{
-					FVulkanCmdBuffer* CmdBuffer = Context.GetCommandBufferManager()->GetActiveCmdBuffer();
+					FVulkanCommandBuffer* CmdBuffer = Context.GetCommandBufferManager()->GetActiveCmdBuffer();
 					CmdBuffer->GetLayoutManager().SetFullLayout(*this, InitialLayout, true);
 				}
 				else if (InitialLayout != VK_IMAGE_LAYOUT_UNDEFINED || bDoInitialClear)
@@ -2122,7 +2122,7 @@ FVulkanTexture::FVulkanTexture(FVulkanDevice& InDevice, const FRHITextureCreateD
 				FVulkanCommandListContext& Context = InDevice.GetImmediateContext();
 				if (bOnlyAddToLayoutManager)
 				{
-					FVulkanCmdBuffer* CmdBuffer = Context.GetCommandBufferManager()->GetActiveCmdBuffer();
+					FVulkanCommandBuffer* CmdBuffer = Context.GetCommandBufferManager()->GetActiveCmdBuffer();
 					CmdBuffer->GetLayoutManager().SetFullLayout(*this, InitialLayout, true);
 				}
 				else if (InitialLayout != VK_IMAGE_LAYOUT_UNDEFINED || bDoInitialClear)
@@ -2559,7 +2559,7 @@ void FVulkanCommandListContext::RHICopyTexture(FRHITexture* SourceTexture, FRHIT
 	FVulkanTexture* Source = ResourceCast(SourceTexture);
 	FVulkanTexture* Dest = ResourceCast(DestTexture);
 
-	FVulkanCmdBuffer* CmdBuffer = CommandBufferManager->GetActiveCmdBuffer();
+	FVulkanCommandBuffer* CmdBuffer = CommandBufferManager->GetActiveCmdBuffer();
 	check(CmdBuffer->IsOutsideRenderPass());
 
 	{
@@ -2700,8 +2700,8 @@ void FVulkanCommandListContext::RHICopyBufferRegion(FRHIBuffer* DstBuffer, uint6
 		return;
 	}
 
-	FVulkanResourceMultiBuffer* DstBufferVk = ResourceCast(DstBuffer);
-	FVulkanResourceMultiBuffer* SrcBufferVk = ResourceCast(SrcBuffer);
+	FVulkanBuffer* DstBufferVk = ResourceCast(DstBuffer);
+	FVulkanBuffer* SrcBufferVk = ResourceCast(SrcBuffer);
 
 	check(DstBufferVk && SrcBufferVk);
 	check(DstOffset + NumBytes <= DstBuffer->GetSize() && SrcOffset + NumBytes <= SrcBuffer->GetSize());
@@ -2709,7 +2709,7 @@ void FVulkanCommandListContext::RHICopyBufferRegion(FRHIBuffer* DstBuffer, uint6
 	uint64 DstOffsetVk = DstBufferVk->GetOffset() + DstOffset;
 	uint64 SrcOffsetVk = SrcBufferVk->GetOffset() + SrcOffset;
 
-	FVulkanCmdBuffer* CmdBuffer = CommandBufferManager->GetActiveCmdBuffer();
+	FVulkanCommandBuffer* CmdBuffer = CommandBufferManager->GetActiveCmdBuffer();
 	check(CmdBuffer->IsOutsideRenderPass());
 	VkCommandBuffer VkCmdBuffer = CmdBuffer->GetHandle();
 
diff --git a/Engine/Source/Runtime/VulkanRHI/Private/VulkanTransientResourceAllocator.cpp b/Engine/Source/Runtime/VulkanRHI/Private/VulkanTransientResourceAllocator.cpp
--- a/Engine/Source/Runtime/VulkanRHI/Private/VulkanTransientResourceAllocator.cpp
+++ b/Engine/Source/Runtime/VulkanRHI/Private/VulkanTransientResourceAllocator.cpp
@@ -21,7 +21,7 @@ FVulkanTransientHeap::FVulkanTransientHeap(const FInitializer& Initializer, FVul
 	}
 
 	const bool bZeroSize = false;
-	VkBufferUsageFlags BufferUsageFlags = FVulkanResourceMultiBuffer::UEToVKBufferUsageFlags(InDevice, UEBufferUsageFlags, bZeroSize);
+	VkBufferUsageFlags BufferUsageFlags = FVulkanBuffer::UEToVKBufferUsageFlags(InDevice, UEBufferUsageFlags, bZeroSize);
 
 	VkBufferCreateInfo BufferCreateInfo;
 	ZeroVulkanStruct(BufferCreateInfo, VK_STRUCTURE_TYPE_BUFFER_CREATE_INFO);
@@ -119,7 +119,7 @@ FRHITransientBuffer* FVulkanTransientResourceAllocator::CreateBuffer(const FRHIB
 		[&](const FRHITransientHeap::FResourceInitializer& Initializer)
 	{
 		FRHIResourceCreateInfo ResourceCreateInfo(InDebugName);
-		FRHIBuffer* Buffer = new FVulkanResourceMultiBuffer(Device, FRHIBufferDesc(InCreateInfo.Size, InCreateInfo.Stride, InCreateInfo.Usage), ResourceCreateInfo, nullptr, &Initializer.Allocation);
+		FRHIBuffer* Buffer = new FVulkanBuffer(Device, FRHIBufferDesc(InCreateInfo.Size, InCreateInfo.Stride, InCreateInfo.Usage), ResourceCreateInfo, nullptr, &Initializer.Allocation);
 		Buffer->SetTrackedAccess_Unsafe(ERHIAccess::Discard);
 		return new FRHITransientBuffer(Buffer, 0/*GpuVirtualAddress*/, Initializer.Hash, Size, ERHITransientAllocationType::Heap, InCreateInfo);
 	});
diff --git a/Engine/Source/Runtime/VulkanRHI/Private/VulkanUAV.cpp b/Engine/Source/Runtime/VulkanRHI/Private/VulkanUAV.cpp
--- a/Engine/Source/Runtime/VulkanRHI/Private/VulkanUAV.cpp
+++ b/Engine/Source/Runtime/VulkanRHI/Private/VulkanUAV.cpp
@@ -58,7 +58,7 @@ void FVulkanView::Invalidate()
 	Storage.Get<FInvalidatedState>().bInitialized = bIsInitialized;
 }
 
-FVulkanView* FVulkanView::InitAsTypedBufferView(FVulkanResourceMultiBuffer* Buffer, EPixelFormat UEFormat, uint32 InOffset, uint32 InSize)
+FVulkanView* FVulkanView::InitAsTypedBufferView(FVulkanBuffer* Buffer, EPixelFormat UEFormat, uint32 InOffset, uint32 InSize)
 {
 	// We will need a deferred update if the descriptor was already in use
 	const bool bImmediateUpdate = !IsInitialized();
@@ -218,7 +218,7 @@ FVulkanView* FVulkanView::InitAsTextureView(
 	return this;
 }
 
-FVulkanView* FVulkanView::InitAsStructuredBufferView(FVulkanResourceMultiBuffer* Buffer, uint32 InOffset, uint32 InSize)
+FVulkanView* FVulkanView::InitAsStructuredBufferView(FVulkanBuffer* Buffer, uint32 InOffset, uint32 InSize)
 {
 	// We will need a deferred update if the descriptor was already in use
 	const bool bImmediateUpdate = !IsInitialized();
@@ -246,7 +246,7 @@ FVulkanView* FVulkanView::InitAsStructuredBufferView(FVulkanResourceMultiBuffer*
 	return this;
 }
 
-FVulkanView* FVulkanView::InitAsAccelerationStructureView(FVulkanResourceMultiBuffer* Buffer, uint32 Offset, uint32 Size)
+FVulkanView* FVulkanView::InitAsAccelerationStructureView(FVulkanBuffer* Buffer, uint32 Offset, uint32 Size)
 {
 	check(GetViewType() == EType::Null);
 	Storage.Emplace<FAccelerationStructureView>();
@@ -403,7 +403,7 @@ void FVulkanShaderResourceView::UpdateView()
 
 	if (IsBuffer())
 	{
-		FVulkanResourceMultiBuffer* Buffer = ResourceCast(GetBuffer());
+		FVulkanBuffer* Buffer = ResourceCast(GetBuffer());
 		auto const Info = ViewDesc.Buffer.SRV.GetViewInfo(Buffer);
 
 		if (!Info.bNullView)
@@ -489,7 +489,7 @@ void FVulkanUnorderedAccessView::UpdateView()
 
 	if (IsBuffer())
 	{
-		FVulkanResourceMultiBuffer* Buffer = ResourceCast(GetBuffer());
+		FVulkanBuffer* Buffer = ResourceCast(GetBuffer());
 		auto const Info = ViewDesc.Buffer.UAV.GetViewInfo(Buffer);
 
 		checkf(!Info.bAppendBuffer && !Info.bAtomicCounter, TEXT("UAV counters not implemented in Vulkan RHI."));
@@ -567,7 +567,7 @@ void FVulkanUnorderedAccessView::Clear(TRHICommandList_RecursiveHazardous<FVulka
 
 	if (IsBuffer())
 	{
-		FVulkanResourceMultiBuffer* Buffer = ResourceCast(GetBuffer());
+		FVulkanBuffer* Buffer = ResourceCast(GetBuffer());
 		auto const Info = ViewDesc.Buffer.UAV.GetViewInfo(Buffer);
 
 		switch (Info.BufferType)
@@ -576,7 +576,7 @@ void FVulkanUnorderedAccessView::Clear(TRHICommandList_RecursiveHazardous<FVulka
 		case FRHIViewDesc::EBufferType::Structured:
 			RHICmdList.RunOnContext([this, Buffer, Info, ClearValue = *static_cast<const uint32*>(ClearValue)](FVulkanCommandListContext& Context)
 			{
-				FVulkanCmdBuffer* CmdBuffer = Context.GetCommandBufferManager()->GetActiveCmdBuffer();
+				FVulkanCommandBuffer* CmdBuffer = Context.GetCommandBufferManager()->GetActiveCmdBuffer();
 
 				// vkCmdFillBuffer is treated as a transfer operation for the purposes of synchronization barriers.
 				{
diff --git a/Engine/Source/Runtime/VulkanRHI/Private/VulkanUniformBuffer.cpp b/Engine/Source/Runtime/VulkanRHI/Private/VulkanUniformBuffer.cpp
--- a/Engine/Source/Runtime/VulkanRHI/Private/VulkanUniformBuffer.cpp
+++ b/Engine/Source/Runtime/VulkanRHI/Private/VulkanUniformBuffer.cpp
@@ -43,7 +43,7 @@ static bool UseTemporaryBuffer(EUniformBufferUsage Usage)
 
 static void UpdateUniformBufferHelper(FVulkanCommandListContext& Context, FVulkanUniformBuffer* VulkanUniformBuffer, const void* Data, bool bUpdateConstants = true)
 {
-	FVulkanCmdBuffer* CmdBuffer = Context.GetCommandBufferManager()->GetActiveCmdBufferDirect();
+	FVulkanCommandBuffer* CmdBuffer = Context.GetCommandBufferManager()->GetActiveCmdBufferDirect();
 
 	FVulkanDevice* Device = Context.GetDevice();
 	const int32 DataSize = VulkanUniformBuffer->GetLayout().ConstantBufferSize;
@@ -95,7 +95,7 @@ void FVulkanUniformBuffer::SetupUniformBufferView()
 	if (UniformViewSRV && GetBufferHandle() == VK_NULL_HANDLE)
 	{
 		const FRHIViewDesc::FBufferSRV& SRVInfo = UniformViewSRV->GetDesc().Buffer.SRV;
-		FVulkanResourceMultiBuffer* Buffer = ResourceCast(UniformViewSRV->GetBuffer());
+		FVulkanBuffer* Buffer = ResourceCast(UniformViewSRV->GetBuffer());
 		Allocation.Reference(Buffer->GetCurrentAllocation());
 		check(Allocation.Size >= PLATFORM_MAX_UNIFORM_BUFFER_RANGE);
 		//Adjust Allocation.Size ???
diff --git a/Engine/Source/Runtime/VulkanRHI/Private/VulkanUtil.cpp b/Engine/Source/Runtime/VulkanRHI/Private/VulkanUtil.cpp
--- a/Engine/Source/Runtime/VulkanRHI/Private/VulkanUtil.cpp
+++ b/Engine/Source/Runtime/VulkanRHI/Private/VulkanUtil.cpp
@@ -103,7 +103,7 @@ void FVulkanGPUTiming::Release()
 /**
  * Start a GPU timing measurement.
  */
-void FVulkanGPUTiming::StartTiming(FVulkanCmdBuffer* CmdBuffer)
+void FVulkanGPUTiming::StartTiming(FVulkanCommandBuffer* CmdBuffer)
 {
 	// Issue a timestamp query for the 'start' time.
 	if (GIsSupported && !bIsTiming)
@@ -153,7 +153,7 @@ void FVulkanGPUTiming::StartTiming(FVulkanCmdBuffer* CmdBuffer)
  * End a GPU timing measurement.
  * The timing for this particular measurement will be resolved at a later time by the GPU.
  */
-void FVulkanGPUTiming::EndTiming(FVulkanCmdBuffer* CmdBuffer)
+void FVulkanGPUTiming::EndTiming(FVulkanCommandBuffer* CmdBuffer)
 {
 	// Issue a timestamp query for the 'end' time.
 	if (GIsSupported && bIsTiming)
@@ -246,7 +246,7 @@ void FVulkanGPUTiming::DiscardExpiredQueries(bool bNeedSpace)
 	}
 }
 
-void FVulkanGPUTiming::RefreshPendingQueries(FVulkanCmdBuffer& CmdBuffer)
+void FVulkanGPUTiming::RefreshPendingQueries(FVulkanCommandBuffer& CmdBuffer)
 {
 	for (int32 QueryIndex=0; QueryIndex < Pool->TimestampListHandles.Num(); ++QueryIndex)
 	{
@@ -582,7 +582,7 @@ void FVulkanGPUProfiler::EndFrame()
 
 #if VULKAN_SUPPORTS_GPU_CRASH_DUMPS
 
-void FVulkanGPUProfiler::PushMarkerForCrash(FVulkanCmdBuffer* CmdBuffer, VkBuffer DestBuffer, const TCHAR* Name)
+void FVulkanGPUProfiler::PushMarkerForCrash(FVulkanCommandBuffer* CmdBuffer, VkBuffer DestBuffer, const TCHAR* Name)
 {
 	uint32 CRC = 0;
 	if (GPUCrashDataDepth < 0 || PushPopStack.Num() < GPUCrashDataDepth)
@@ -609,7 +609,7 @@ void FVulkanGPUProfiler::PushMarkerForCrash(FVulkanCmdBuffer* CmdBuffer, VkBuffe
 	FVulkanPlatform::WriteCrashMarker(Device->GetOptionalExtensions(), CmdBuffer, DestBuffer, TArrayView<uint32>(PushPopStack), true);
 }
 
-void FVulkanGPUProfiler::PopMarkerForCrash(FVulkanCmdBuffer* CmdBuffer, VkBuffer DestBuffer)
+void FVulkanGPUProfiler::PopMarkerForCrash(FVulkanCommandBuffer* CmdBuffer, VkBuffer DestBuffer)
 {
 	if (PushPopStack.Num() > 0)
 	{
diff --git a/Engine/Source/Runtime/VulkanRHI/Private/VulkanViewport.cpp b/Engine/Source/Runtime/VulkanRHI/Private/VulkanViewport.cpp
--- a/Engine/Source/Runtime/VulkanRHI/Private/VulkanViewport.cpp
+++ b/Engine/Source/Runtime/VulkanRHI/Private/VulkanViewport.cpp
@@ -83,7 +83,7 @@ void FVulkanBackBuffer::AcquireBackBufferImage(FVulkanCommandListContext& Contex
 			PartialView = &ImageView;
 
 			FVulkanCommandBufferManager* CmdBufferManager = Context.GetCommandBufferManager();
-			FVulkanCmdBuffer* CmdBuffer = CmdBufferManager->GetActiveCmdBuffer();
+			FVulkanCommandBuffer* CmdBuffer = CmdBufferManager->GetActiveCmdBuffer();
 			check(!CmdBuffer->IsInsideRenderPass());
 
 			// right after acquiring image is in undefined state
@@ -753,7 +753,7 @@ void FVulkanViewport::CreateSwapchain(FVulkanSwapChainRecreateInfo* RecreateInfo
 		BackBufferImages.SetNum(Images.Num());
 		RenderingDoneSemaphores.SetNum(Images.Num());
 
-		FVulkanCmdBuffer* CmdBuffer = Device->GetImmediateContext().GetCommandBufferManager()->GetUploadCmdBuffer();
+		FVulkanCommandBuffer* CmdBuffer = Device->GetImmediateContext().GetCommandBufferManager()->GetUploadCmdBuffer();
 		ensure(CmdBuffer->IsOutsideRenderPass());
 
 		VkClearColorValue ClearColor;
@@ -876,7 +876,7 @@ void FVulkanViewport::DestroySwapchain(FVulkanSwapChainRecreateInfo* RecreateInf
 	AcquiredImageIndex = -1;
 }
 
-inline static void CopyImageToBackBuffer(FVulkanCommandListContext* Context, FVulkanCmdBuffer* CmdBuffer, FVulkanTexture& SrcSurface, FVulkanTexture& DstSurface, int32 SizeX, int32 SizeY, int32 WindowSizeX, int32 WindowSizeY, VkSurfaceTransformFlagBitsKHR CachedSurfaceTransform)
+inline static void CopyImageToBackBuffer(FVulkanCommandListContext* Context, FVulkanCommandBuffer* CmdBuffer, FVulkanTexture& SrcSurface, FVulkanTexture& DstSurface, int32 SizeX, int32 SizeY, int32 WindowSizeX, int32 WindowSizeY, VkSurfaceTransformFlagBitsKHR CachedSurfaceTransform)
 {
 	RHI_BREADCRUMB_EVENT(*Context, "CopyImageToBackBuffer");
 	const bool bNeedsVulkanPreTransform = CachedSurfaceTransform != VK_SURFACE_TRANSFORM_IDENTITY_BIT_KHR;
@@ -1003,7 +1003,7 @@ inline static void CopyImageToBackBuffer(FVulkanCommandListContext* Context, FVu
 	}
 }
 
-bool FVulkanViewport::Present(FVulkanCommandListContext* Context, FVulkanCmdBuffer* CmdBuffer, FVulkanQueue* Queue, FVulkanQueue* PresentQueue, bool bLockToVsync)
+bool FVulkanViewport::Present(FVulkanCommandListContext* Context, FVulkanCommandBuffer* CmdBuffer, FVulkanQueue* Queue, FVulkanQueue* PresentQueue, bool bLockToVsync)
 {
 	FPlatformAtomics::AtomicStore(&LockToVsync, bLockToVsync ? 1 : 0);
 	bool bFailedToDelayAcquireBackbuffer = false;
diff --git a/Engine/Source/Runtime/VulkanRHI/Private/VulkanViewport.h b/Engine/Source/Runtime/VulkanRHI/Private/VulkanViewport.h
--- a/Engine/Source/Runtime/VulkanRHI/Private/VulkanViewport.h
+++ b/Engine/Source/Runtime/VulkanRHI/Private/VulkanViewport.h
@@ -73,7 +73,7 @@ public:
 
 	virtual void Tick(float DeltaTime) override final;
 	
-	bool Present(FVulkanCommandListContext* Context, FVulkanCmdBuffer* CmdBuffer, FVulkanQueue* Queue, FVulkanQueue* PresentQueue, bool bLockToVsync);
+	bool Present(FVulkanCommandListContext* Context, FVulkanCommandBuffer* CmdBuffer, FVulkanQueue* Queue, FVulkanQueue* PresentQueue, bool bLockToVsync);
 
 	inline uint32 GetPresentCount() const
 	{
@@ -139,7 +139,7 @@ protected:
 
 	FCustomPresentRHIRef CustomPresent;
 
-	FVulkanCmdBuffer* LastFrameCommandBuffer = nullptr;
+	FVulkanCommandBuffer* LastFrameCommandBuffer = nullptr;
 	uint64 LastFrameFenceCounter = 0;
 
 	EDeviceScreenOrientation CachedOrientation = EDeviceScreenOrientation::Unknown;
diff --git a/Engine/Source/Runtime/VulkanRHI/Private/Windows/VulkanWindowsPlatform.cpp b/Engine/Source/Runtime/VulkanRHI/Private/Windows/VulkanWindowsPlatform.cpp
--- a/Engine/Source/Runtime/VulkanRHI/Private/Windows/VulkanWindowsPlatform.cpp
+++ b/Engine/Source/Runtime/VulkanRHI/Private/Windows/VulkanWindowsPlatform.cpp
@@ -195,7 +195,7 @@ bool FVulkanWindowsPlatform::SupportsDeviceLocalHostVisibleWithNoPenalty(EGpuVen
 	return (VendorId == EGpuVendorId::Amd && bIsWin10);
 }
 
-void FVulkanWindowsPlatform::WriteCrashMarker(const FOptionalVulkanDeviceExtensions& OptionalExtensions, FVulkanCmdBuffer* CmdBuffer, VkBuffer DestBuffer, const TArrayView<uint32>& Entries, bool bAdding)
+void FVulkanWindowsPlatform::WriteCrashMarker(const FOptionalVulkanDeviceExtensions& OptionalExtensions, FVulkanCommandBuffer* CmdBuffer, VkBuffer DestBuffer, const TArrayView<uint32>& Entries, bool bAdding)
 {
 	ensure(Entries.Num() <= GMaxCrashBufferEntries);
 
diff --git a/Engine/Source/Runtime/VulkanRHI/Private/Windows/VulkanWindowsPlatform.h b/Engine/Source/Runtime/VulkanRHI/Private/Windows/VulkanWindowsPlatform.h
--- a/Engine/Source/Runtime/VulkanRHI/Private/Windows/VulkanWindowsPlatform.h
+++ b/Engine/Source/Runtime/VulkanRHI/Private/Windows/VulkanWindowsPlatform.h
@@ -56,7 +56,7 @@ public:
 
 	static bool SupportsDeviceLocalHostVisibleWithNoPenalty(EGpuVendorId VendorId);
 
-	static void WriteCrashMarker(const FOptionalVulkanDeviceExtensions& OptionalExtensions, FVulkanCmdBuffer* CmdBuffer, VkBuffer DestBuffer, const TArrayView<uint32>& Entries, bool bAdding);
+	static void WriteCrashMarker(const FOptionalVulkanDeviceExtensions& OptionalExtensions, FVulkanCommandBuffer* CmdBuffer, VkBuffer DestBuffer, const TArrayView<uint32>& Entries, bool bAdding);
 
 private:
 	static bool bAttemptedLoad;
