VulkanRHI: Fix RHI stat for buffers

#rnx

diff --git a/Engine/Source/Runtime/VulkanRHI/Private/VulkanBuffer.cpp b/Engine/Source/Runtime/VulkanRHI/Private/VulkanBuffer.cpp
--- a/Engine/Source/Runtime/VulkanRHI/Private/VulkanBuffer.cpp
+++ b/Engine/Source/Runtime/VulkanRHI/Private/VulkanBuffer.cpp
@@ -206,7 +206,7 @@ void FVulkanBuffer::AllocateMemory(FBufferAlloc& OutAlloc)
 	NewBufferAlloc.HostPtr = (bUnifiedMem || bDynamic) ? NewBufferAlloc.Alloc.GetMappedPointer(Device) : nullptr;
 	NewBufferAlloc.DeviceAddress = GetBufferDeviceAddress(Device, NewBufferAlloc.Alloc.GetBufferHandle()) + NewBufferAlloc.Alloc.Offset;
 
-	UpdateVulkanBufferStats(GetDesc(), BufferSize, true);
+	UpdateVulkanBufferStats(GetDesc(), NewBufferAlloc.Alloc.Size, true);
 }
 
 void* FVulkanBuffer::Lock(FRHICommandListBase& RHICmdList, EResourceLockMode LockMode, uint32 LockSize, uint32 Offset)
@@ -323,6 +323,7 @@ void* FVulkanBuffer::Lock(FRHICommandListBase& RHICmdList, EResourceLockMode Loc
 
 			RHICmdList.EnqueueLambda(TEXT("FVulkanBuffer::Lock"), [Buffer = this, NewAlloc](FRHICommandListBase& CmdList)
 			{
+				UpdateVulkanBufferStats(Buffer->GetDesc(), Buffer->CurrentBufferAlloc.Alloc.Size, false);
 				Buffer->CurrentBufferAlloc.Alloc.Free(*Buffer->GetParent());
 				Buffer->CurrentBufferAlloc = NewAlloc;
 				Buffer->CurrentBufferAlloc.Alloc.Own();
