VulkanRHI: Clean up:
- Start using a reference to FVulkanCommandBuffer in places where null pointer isn't possible/handled/accepted.
- Add accessor to current command buffer directly from context (start phasing out CommandBufferManager for new submission pipeline).
- Start using consistent naming for variables: "CommandBuffer" for a FVulkanCommandBuffer, "CommandBufferHandle" for a VkCommandBuffer handle.

#rnx

diff --git a/Engine/Source/Runtime/VulkanRHI/Private/VulkanBarriers.cpp b/Engine/Source/Runtime/VulkanRHI/Private/VulkanBarriers.cpp
--- a/Engine/Source/Runtime/VulkanRHI/Private/VulkanBarriers.cpp
+++ b/Engine/Source/Runtime/VulkanRHI/Private/VulkanBarriers.cpp
@@ -1197,7 +1197,7 @@ void FTransitionProcessor<VkMemoryBarrier, VkBufferMemoryBarrier, VkImageMemoryB
 	}
 
 	// Submit merged stage masks with arrays of barriers
-	VulkanRHI::vkCmdPipelineBarrier(Context.GetCommandBufferManager()->GetActiveCmdBuffer()->GetHandle(), MergedSrcStageMask, MergedDstStageMask, 0, MemoryBarriers.Num(), MemoryBarriers.GetData(), BufferBarriers.Num(), BufferBarriers.GetData(), ImageBarriers.Num(), ImageBarriers.GetData());
+	VulkanRHI::vkCmdPipelineBarrier(Context.GetCommandBuffer().GetHandle(), MergedSrcStageMask, MergedDstStageMask, 0, MemoryBarriers.Num(), MemoryBarriers.GetData(), BufferBarriers.Num(), BufferBarriers.GetData(), ImageBarriers.Num(), ImageBarriers.GetData());
 	MemoryBarriers.Reset();
 	BufferBarriers.Reset();
 	ImageBarriers.Reset();
@@ -1222,7 +1222,7 @@ void FTransitionProcessor<VkMemoryBarrier, VkBufferMemoryBarrier, VkImageMemoryB
 template <>
 void FTransitionProcessor<VkMemoryBarrier, VkBufferMemoryBarrier, VkImageMemoryBarrier>::ApplyTracking(uint32_t TargetIndex, const FVulkanPipelineBarrier::ImageBarrierExtraData* RemainingExtras, int32 RemainingCount, bool IsCrossPipe)
 {
-	FVulkanLayoutManager& LayoutManager = Context.GetCommandBufferManager()->GetActiveCmdBuffer()->GetLayoutManager();
+	FVulkanLayoutManager& LayoutManager = Context.GetCommandBuffer().GetLayoutManager();
 
 	VkImageMemoryBarrier* ImageBarrier = &ImageBarriers[TargetIndex];
 	FVulkanTexture* Texture = RemainingExtras->BaseTexture;
@@ -1434,7 +1434,7 @@ void FTransitionProcessor<VkMemoryBarrier2, VkBufferMemoryBarrier2, VkImageMemor
 
 	if ((GVulkanMaxBarriersPerBatch <= 0) || ((MemoryBarriers.Num() + BufferBarriers.Num() + ImageBarriers.Num()) < GVulkanMaxBarriersPerBatch))
 	{
-		VulkanRHI::vkCmdPipelineBarrier2KHR(Context.GetCommandBufferManager()->GetActiveCmdBuffer()->GetHandle(), &DependencyInfo);
+		VulkanRHI::vkCmdPipelineBarrier2KHR(Context.GetCommandBuffer().GetHandle(), &DependencyInfo);
 	}
 	else
 	{
@@ -1443,7 +1443,7 @@ void FTransitionProcessor<VkMemoryBarrier2, VkBufferMemoryBarrier2, VkImageMemor
 		BatchDependencyInfo.bufferMemoryBarrierCount = 0;
 		BatchDependencyInfo.imageMemoryBarrierCount = 0;
 
-		VkCommandBuffer CommandBuffer = Context.GetCommandBufferManager()->GetActiveCmdBuffer()->GetHandle();
+		VkCommandBuffer CommandBuffer = Context.GetCommandBuffer().GetHandle();
 		SendBatchedBarriers(CommandBuffer, BatchDependencyInfo, BatchDependencyInfo.pMemoryBarriers, BatchDependencyInfo.memoryBarrierCount, DependencyInfo.memoryBarrierCount);
 		SendBatchedBarriers(CommandBuffer, BatchDependencyInfo, BatchDependencyInfo.pBufferMemoryBarriers, BatchDependencyInfo.bufferMemoryBarrierCount, DependencyInfo.bufferMemoryBarrierCount);
 		SendBatchedBarriers(CommandBuffer, BatchDependencyInfo, BatchDependencyInfo.pImageMemoryBarriers, BatchDependencyInfo.imageMemoryBarrierCount, DependencyInfo.imageMemoryBarrierCount);
@@ -1467,7 +1467,7 @@ void FTransitionProcessor<VkMemoryBarrier2, VkBufferMemoryBarrier2, VkImageMemor
 
 		if ((ImageBarrier.oldLayout == VK_IMAGE_LAYOUT_UNDEFINED) && (!RemainingExtras->IsAliasingBarrier))
 		{
-			FVulkanLayoutManager& LayoutManager = Context.GetCommandBufferManager()->GetActiveCmdBuffer()->GetLayoutManager();
+			FVulkanLayoutManager& LayoutManager = Context.GetCommandBuffer().GetLayoutManager();
 			FVulkanTexture* Texture = RemainingExtras->BaseTexture;
 			check(Texture->Image == ImageBarrier.image);
 			const FVulkanImageLayout* Layout = LayoutManager.GetFullLayout(*Texture, false);
@@ -1502,7 +1502,7 @@ void FTransitionProcessor<VkMemoryBarrier2, VkBufferMemoryBarrier2, VkImageMemor
 	}
 
 	// Sync2 only needs to inform tracking of the new layout
-	Context.GetCommandBufferManager()->GetActiveCmdBuffer()->GetLayoutManager().SetLayout(*RemainingExtras->BaseTexture, ImageBarriers[TargetIndex].subresourceRange, ImageBarriers[TargetIndex].newLayout);
+	Context.GetCommandBuffer().GetLayoutManager().SetLayout(*RemainingExtras->BaseTexture, ImageBarriers[TargetIndex].subresourceRange, ImageBarriers[TargetIndex].newLayout);
 
 	// No needs to worry about 'surprise' subresource transition
 	// No need to worry about partial depth/stencil layout changes
@@ -1572,8 +1572,8 @@ void FVulkanCommandListContext::RHIEndTransitions(TArrayView<const FRHITransitio
 			CommandBufferManager->PrepareForNewActiveCommandBuffer();
 		}
 
-		FVulkanCommandBuffer* CmdBuffer = CommandBufferManager->GetActiveCmdBuffer();
-		CmdBuffer->AddWaitSemaphore(VK_PIPELINE_STAGE_TOP_OF_PIPE_BIT, WaitSemaphores);
+		FVulkanCommandBuffer& CmdBuffer = GetCommandBuffer();
+		CmdBuffer.AddWaitSemaphore(VK_PIPELINE_STAGE_TOP_OF_PIPE_BIT, WaitSemaphores);
 	}
 
 	if (Device->SupportsParallelRendering())
diff --git a/Engine/Source/Runtime/VulkanRHI/Private/VulkanBuffer.cpp b/Engine/Source/Runtime/VulkanRHI/Private/VulkanBuffer.cpp
--- a/Engine/Source/Runtime/VulkanRHI/Private/VulkanBuffer.cpp
+++ b/Engine/Source/Runtime/VulkanRHI/Private/VulkanBuffer.cpp
@@ -377,24 +377,24 @@ void FVulkanBuffer::Unlock(FRHICommandListBase& RHICmdList)
 
 				// We need to do this on the active command buffer instead of using an upload command buffer. The high level code sometimes reuses the same
 				// buffer in sequences of upload / dispatch, upload / dispatch, so we need to order the copy commands correctly with respect to the dispatches.
-				FVulkanCommandBuffer* Cmd = Context.GetCommandBufferManager()->GetActiveCmdBuffer();
-				check(Cmd && Cmd->IsOutsideRenderPass());
-				VkCommandBuffer CmdBuffer = Cmd->GetHandle();
+				FVulkanCommandBuffer& CommandBuffer = Context.GetCommandBuffer();
+				check(CommandBuffer.IsOutsideRenderPass());
+				VkCommandBuffer CommandBufferHandle = CommandBuffer.GetHandle();
 
-				VulkanRHI::DebugHeavyWeightBarrier(CmdBuffer, 16);
+				VulkanRHI::DebugHeavyWeightBarrier(CommandBufferHandle, 16);
 
 				VkBufferCopy Region;
 				FMemory::Memzero(Region);
 				Region.size = PendingLock.Size;
 				//Region.srcOffset = 0;
 				Region.dstOffset = PendingLock.Offset + Buffer->CurrentBufferAlloc.Alloc.Offset;
-				VulkanRHI::vkCmdCopyBuffer(CmdBuffer, StagingBuffer->GetHandle(), Buffer->CurrentBufferAlloc.Alloc.GetBufferHandle(), 1, &Region);
+				VulkanRHI::vkCmdCopyBuffer(CommandBufferHandle, StagingBuffer->GetHandle(), Buffer->CurrentBufferAlloc.Alloc.GetBufferHandle(), 1, &Region);
 
 				// High level code expects the data in Buffer to be ready to read
 				VkMemoryBarrier BarrierAfter = { VK_STRUCTURE_TYPE_MEMORY_BARRIER, nullptr, VK_ACCESS_TRANSFER_WRITE_BIT, VK_ACCESS_MEMORY_READ_BIT | VK_ACCESS_MEMORY_WRITE_BIT };
-				VulkanRHI::vkCmdPipelineBarrier(CmdBuffer, VK_PIPELINE_STAGE_TRANSFER_BIT, VK_PIPELINE_STAGE_ALL_COMMANDS_BIT, 0, 1, &BarrierAfter, 0, nullptr, 0, nullptr);
+				VulkanRHI::vkCmdPipelineBarrier(CommandBufferHandle, VK_PIPELINE_STAGE_TRANSFER_BIT, VK_PIPELINE_STAGE_ALL_COMMANDS_BIT, 0, 1, &BarrierAfter, 0, nullptr, 0, nullptr);
 
-				Buffer->GetParent()->GetStagingManager().ReleaseBuffer(Cmd, StagingBuffer);
+				Buffer->GetParent()->GetStagingManager().ReleaseBuffer(&CommandBuffer, StagingBuffer);
 
 				Buffer->UpdateLinkedViews(); // :todo-jn:  not needed?  Same buffer when we use staging?
 			}
diff --git a/Engine/Source/Runtime/VulkanRHI/Private/VulkanCommands.cpp b/Engine/Source/Runtime/VulkanRHI/Private/VulkanCommands.cpp
--- a/Engine/Source/Runtime/VulkanRHI/Private/VulkanCommands.cpp
+++ b/Engine/Source/Runtime/VulkanRHI/Private/VulkanCommands.cpp
@@ -298,12 +298,12 @@ void FVulkanCommandListContext::RHIDispatchComputeShader(uint32 ThreadGroupCount
 
 	CommitComputeResourceTables();
 
-	FVulkanCommandBuffer* Cmd = CommandBufferManager->GetActiveCmdBuffer();
-	ensure(Cmd->IsOutsideRenderPass());
-	VkCommandBuffer CmdBuffer = Cmd->GetHandle();
-	PendingComputeState->PrepareForDispatch(Cmd);
+	FVulkanCommandBuffer& CommandBuffer = GetCommandBuffer();
+	ensure(CommandBuffer.IsOutsideRenderPass());
+	VkCommandBuffer CommandBufferHandle = CommandBuffer.GetHandle();
+	PendingComputeState->PrepareForDispatch(CommandBuffer);
 
-	VulkanRHI::vkCmdDispatch(CmdBuffer, ThreadGroupCountX, ThreadGroupCountY, ThreadGroupCountZ);
+	VulkanRHI::vkCmdDispatch(CommandBufferHandle, ThreadGroupCountX, ThreadGroupCountY, ThreadGroupCountZ);
 
 	if (GCVarSubmitOnDispatch.GetValueOnRenderThread())
 	{
@@ -315,7 +315,7 @@ void FVulkanCommandListContext::RHIDispatchComputeShader(uint32 ThreadGroupCount
 		RegisterGPUDispatch(FIntVector(ThreadGroupCountX, ThreadGroupCountY, ThreadGroupCountZ));	
 	}
 
-	VulkanRHI::DebugHeavyWeightBarrier(CmdBuffer, 2);
+	VulkanRHI::DebugHeavyWeightBarrier(CommandBufferHandle, 2);
 }
 
 void FVulkanCommandListContext::RHIDispatchIndirectComputeShader(FRHIBuffer* ArgumentBufferRHI, uint32 ArgumentOffset)
@@ -325,12 +325,12 @@ void FVulkanCommandListContext::RHIDispatchIndirectComputeShader(FRHIBuffer* Arg
 
 	CommitComputeResourceTables();
 
-	FVulkanCommandBuffer* Cmd = CommandBufferManager->GetActiveCmdBuffer();
-	ensure(Cmd->IsOutsideRenderPass());
-	VkCommandBuffer CmdBuffer = Cmd->GetHandle();
-	PendingComputeState->PrepareForDispatch(Cmd);
+	FVulkanCommandBuffer& CommandBuffer = GetCommandBuffer();
+	ensure(CommandBuffer.IsOutsideRenderPass());
+	VkCommandBuffer CommandBufferHandle = CommandBuffer.GetHandle();
+	PendingComputeState->PrepareForDispatch(CommandBuffer);
 
-	VulkanRHI::vkCmdDispatchIndirect(CmdBuffer, ArgumentBuffer->GetHandle(), ArgumentBuffer->GetOffset() + ArgumentOffset);
+	VulkanRHI::vkCmdDispatchIndirect(CommandBufferHandle, ArgumentBuffer->GetHandle(), ArgumentBuffer->GetOffset() + ArgumentOffset);
 
 	if (GCVarSubmitOnDispatch.GetValueOnRenderThread())
 	{
@@ -342,7 +342,7 @@ void FVulkanCommandListContext::RHIDispatchIndirectComputeShader(FRHIBuffer* Arg
 		RegisterGPUDispatch(FIntVector(1, 1, 1));	
 	}
 
-	VulkanRHI::DebugHeavyWeightBarrier(CmdBuffer, 2);
+	VulkanRHI::DebugHeavyWeightBarrier(CommandBufferHandle, 2);
 }
 
 void FVulkanCommandListContext::RHISetUAVParameter(FRHIPixelShader* PixelShaderRHI, uint32 UAVIndex, FRHIUnorderedAccessView* UAVRHI)
@@ -644,10 +644,10 @@ void FVulkanCommandListContext::RHIDrawPrimitive(uint32 BaseVertexIndex, uint32
 
 	CommitGraphicsResourceTables();
 
-	FVulkanCommandBuffer* CmdBuffer = CommandBufferManager->GetActiveCmdBuffer();
-	PendingGfxState->PrepareForDraw(CmdBuffer);
+	FVulkanCommandBuffer& CommandBuffer = GetCommandBuffer();
+	PendingGfxState->PrepareForDraw(CommandBuffer);
 	uint32 NumVertices = GetVertexCountForPrimitiveCount(NumPrimitives, PendingGfxState->PrimitiveType);
-	VulkanRHI::vkCmdDraw(CmdBuffer->GetHandle(), NumVertices, NumInstances, BaseVertexIndex, 0);
+	VulkanRHI::vkCmdDraw(CommandBuffer.GetHandle(), NumVertices, NumInstances, BaseVertexIndex, 0);
 
 	if (FVulkanPlatform::RegisterGPUWork() && IsImmediate())
 	{
@@ -666,14 +666,14 @@ void FVulkanCommandListContext::RHIDrawPrimitiveIndirect(FRHIBuffer* ArgumentBuf
 
 	CommitGraphicsResourceTables();
 
-	FVulkanCommandBuffer* Cmd = CommandBufferManager->GetActiveCmdBuffer();
-	VkCommandBuffer CmdBuffer = Cmd->GetHandle();
-	PendingGfxState->PrepareForDraw(Cmd);
+	FVulkanCommandBuffer& CommandBuffer = GetCommandBuffer();
+	VkCommandBuffer CommandBufferHandle = CommandBuffer.GetHandle();
+	PendingGfxState->PrepareForDraw(CommandBuffer);
 
 	FVulkanBuffer* ArgumentBuffer = ResourceCast(ArgumentBufferRHI);
 
 
-	VulkanRHI::vkCmdDrawIndirect(CmdBuffer, ArgumentBuffer->GetHandle(), ArgumentBuffer->GetOffset() + ArgumentOffset, 1, sizeof(VkDrawIndirectCommand));
+	VulkanRHI::vkCmdDrawIndirect(CommandBufferHandle, ArgumentBuffer->GetHandle(), ArgumentBuffer->GetOffset() + ArgumentOffset, 1, sizeof(VkDrawIndirectCommand));
 
 	if (FVulkanPlatform::RegisterGPUWork() && IsImmediate())
 	{
@@ -694,13 +694,13 @@ void FVulkanCommandListContext::RHIDrawIndexedPrimitive(FRHIBuffer* IndexBufferR
 	CommitGraphicsResourceTables();
 
 	FVulkanBuffer* IndexBuffer = ResourceCast(IndexBufferRHI);
-	FVulkanCommandBuffer* Cmd = CommandBufferManager->GetActiveCmdBuffer();
-	VkCommandBuffer CmdBuffer = Cmd->GetHandle();
-	PendingGfxState->PrepareForDraw(Cmd);
-	VulkanRHI::vkCmdBindIndexBuffer(CmdBuffer, IndexBuffer->GetHandle(), IndexBuffer->GetOffset(), IndexBuffer->GetIndexType());
+	FVulkanCommandBuffer& CommandBuffer = GetCommandBuffer();
+	VkCommandBuffer CommandBufferHandle = CommandBuffer.GetHandle();
+	PendingGfxState->PrepareForDraw(CommandBuffer);
+	VulkanRHI::vkCmdBindIndexBuffer(CommandBufferHandle, IndexBuffer->GetHandle(), IndexBuffer->GetOffset(), IndexBuffer->GetIndexType());
 
 	uint32 NumIndices = GetVertexCountForPrimitiveCount(NumPrimitives, PendingGfxState->PrimitiveType);
-	VulkanRHI::vkCmdDrawIndexed(CmdBuffer, NumIndices, NumInstances, StartIndex, BaseVertexIndex, FirstInstance);
+	VulkanRHI::vkCmdDrawIndexed(CommandBufferHandle, NumIndices, NumInstances, StartIndex, BaseVertexIndex, FirstInstance);
 
 	if (FVulkanPlatform::RegisterGPUWork() && IsImmediate())
 	{
@@ -718,16 +718,16 @@ void FVulkanCommandListContext::RHIDrawIndexedIndirect(FRHIBuffer* IndexBufferRH
 	CommitGraphicsResourceTables();
 
 	FVulkanBuffer* IndexBuffer = ResourceCast(IndexBufferRHI);
-	FVulkanCommandBuffer* Cmd = CommandBufferManager->GetActiveCmdBuffer();
-	VkCommandBuffer CmdBuffer = Cmd->GetHandle();
-	PendingGfxState->PrepareForDraw(Cmd);
-	VulkanRHI::vkCmdBindIndexBuffer(CmdBuffer, IndexBuffer->GetHandle(), IndexBuffer->GetOffset(), IndexBuffer->GetIndexType());
+	FVulkanCommandBuffer& CommandBuffer = GetCommandBuffer();
+	VkCommandBuffer CommandBufferHandle = CommandBuffer.GetHandle();
+	PendingGfxState->PrepareForDraw(CommandBuffer);
+	VulkanRHI::vkCmdBindIndexBuffer(CommandBufferHandle, IndexBuffer->GetHandle(), IndexBuffer->GetOffset(), IndexBuffer->GetIndexType());
 
 	FVulkanBuffer* ArgumentBuffer = ResourceCast(ArgumentsBufferRHI);
 	VkDeviceSize ArgumentOffset = DrawArgumentsIndex * sizeof(VkDrawIndexedIndirectCommand);
 
 
-	VulkanRHI::vkCmdDrawIndexedIndirect(CmdBuffer, ArgumentBuffer->GetHandle(), ArgumentBuffer->GetOffset() + ArgumentOffset, 1, sizeof(VkDrawIndexedIndirectCommand));
+	VulkanRHI::vkCmdDrawIndexedIndirect(CommandBufferHandle, ArgumentBuffer->GetHandle(), ArgumentBuffer->GetOffset() + ArgumentOffset, 1, sizeof(VkDrawIndexedIndirectCommand));
 
 	if (FVulkanPlatform::RegisterGPUWork() && IsImmediate())
 	{
@@ -745,15 +745,15 @@ void FVulkanCommandListContext::RHIDrawIndexedPrimitiveIndirect(FRHIBuffer* Inde
 	CommitGraphicsResourceTables();
 
 	FVulkanBuffer* IndexBuffer = ResourceCast(IndexBufferRHI);
-	FVulkanCommandBuffer* Cmd = CommandBufferManager->GetActiveCmdBuffer();
-	VkCommandBuffer CmdBuffer = Cmd->GetHandle();
-	PendingGfxState->PrepareForDraw(Cmd);
-	VulkanRHI::vkCmdBindIndexBuffer(CmdBuffer, IndexBuffer->GetHandle(), IndexBuffer->GetOffset(), IndexBuffer->GetIndexType());
+	FVulkanCommandBuffer& CommandBuffer = GetCommandBuffer();
+	VkCommandBuffer CommandBufferHandle = CommandBuffer.GetHandle();
+	PendingGfxState->PrepareForDraw(CommandBuffer);
+	VulkanRHI::vkCmdBindIndexBuffer(CommandBufferHandle, IndexBuffer->GetHandle(), IndexBuffer->GetOffset(), IndexBuffer->GetIndexType());
 
 	FVulkanBuffer* ArgumentBuffer = ResourceCast(ArgumentBufferRHI);
 
 
-	VulkanRHI::vkCmdDrawIndexedIndirect(CmdBuffer, ArgumentBuffer->GetHandle(), ArgumentBuffer->GetOffset() + ArgumentOffset, 1, sizeof(VkDrawIndexedIndirectCommand));
+	VulkanRHI::vkCmdDrawIndexedIndirect(CommandBufferHandle, ArgumentBuffer->GetHandle(), ArgumentBuffer->GetOffset() + ArgumentOffset, 1, sizeof(VkDrawIndexedIndirectCommand));
 
 	if (FVulkanPlatform::RegisterGPUWork() && IsImmediate())
 	{
@@ -766,11 +766,11 @@ void FVulkanCommandListContext::RHIDispatchMeshShader(uint32 ThreadGroupCountX,
 {
 	CommitGraphicsResourceTables();
 
-	FVulkanCommandBuffer* Cmd = CommandBufferManager->GetActiveCmdBuffer();
+	FVulkanCommandBuffer& CommandBuffer = GetCommandBuffer();
 
-	PendingGfxState->PrepareForDraw(Cmd);
+	PendingGfxState->PrepareForDraw(CommandBuffer);
 
-	VulkanRHI::vkCmdDrawMeshTasksEXT(Cmd->GetHandle(), ThreadGroupCountX, ThreadGroupCountY, ThreadGroupCountZ);
+	VulkanRHI::vkCmdDrawMeshTasksEXT(CommandBuffer.GetHandle(), ThreadGroupCountX, ThreadGroupCountY, ThreadGroupCountZ);
 
 	if (FVulkanPlatform::RegisterGPUWork() && IsImmediate())
 	{
@@ -782,13 +782,13 @@ void FVulkanCommandListContext::RHIDispatchIndirectMeshShader(FRHIBuffer* Argume
 {
 	CommitGraphicsResourceTables();
 
-	FVulkanCommandBuffer* Cmd = CommandBufferManager->GetActiveCmdBuffer();
+	FVulkanCommandBuffer& CommandBuffer = GetCommandBuffer();
 
-	PendingGfxState->PrepareForDraw(Cmd);
+	PendingGfxState->PrepareForDraw(CommandBuffer);
 
 	FVulkanBuffer* ArgumentBuffer = ResourceCast(ArgumentBufferRHI);
 
-	VulkanRHI::vkCmdDrawMeshTasksIndirectEXT(Cmd->GetHandle(), ArgumentBuffer->GetHandle(), ArgumentBuffer->GetOffset() + ArgumentOffset, 1, sizeof(VkDrawMeshTasksIndirectCommandEXT));
+	VulkanRHI::vkCmdDrawMeshTasksIndirectEXT(CommandBuffer.GetHandle(), ArgumentBuffer->GetHandle(), ArgumentBuffer->GetOffset() + ArgumentOffset, 1, sizeof(VkDrawMeshTasksIndirectCommandEXT));
 
 	if (FVulkanPlatform::RegisterGPUWork() && IsImmediate())
 	{
@@ -806,15 +806,15 @@ void FVulkanCommandListContext::RHIClearMRT(bool bClearColor, int32 NumClearColo
 
 	check(bClearColor ? NumClearColors > 0 : true);
 
-	FVulkanCommandBuffer* CmdBuffer = CommandBufferManager->GetActiveCmdBuffer();
+	FVulkanCommandBuffer& CommandBuffer = GetCommandBuffer();
 	//FRCLog::Printf(TEXT("RHIClearMRT"));
 
 	const uint32 NumColorAttachments = CurrentFramebuffer->GetNumColorAttachments();
 	check(!bClearColor || (uint32)NumClearColors <= NumColorAttachments);
-	InternalClearMRT(CmdBuffer, bClearColor, bClearColor ? NumClearColors : 0, ClearColorArray, bClearDepth, Depth, bClearStencil, Stencil);
+	InternalClearMRT(CommandBuffer, bClearColor, bClearColor ? NumClearColors : 0, ClearColorArray, bClearDepth, Depth, bClearStencil, Stencil);
 }
 
-void FVulkanCommandListContext::InternalClearMRT(FVulkanCommandBuffer* CmdBuffer, bool bClearColor, int32 NumClearColors, const FLinearColor* ClearColorArray, bool bClearDepth, float Depth, bool bClearStencil, uint32 Stencil)
+void FVulkanCommandListContext::InternalClearMRT(FVulkanCommandBuffer& CommandBuffer, bool bClearColor, int32 NumClearColors, const FLinearColor* ClearColorArray, bool bClearDepth, float Depth, bool bClearStencil, uint32 Stencil)
 {
 	if (CurrentRenderPass)
 	{
@@ -852,12 +852,12 @@ void FVulkanCommandListContext::InternalClearMRT(FVulkanCommandBuffer* CmdBuffer
 			++NumAttachments;
 		}
 
-		VulkanRHI::vkCmdClearAttachments(CmdBuffer->GetHandle(), NumAttachments, Attachments, 1, &Rect);
+		VulkanRHI::vkCmdClearAttachments(CommandBuffer.GetHandle(), NumAttachments, Attachments, 1, &Rect);
 	}
 	else
 	{
 		ensure(0);
-		//VulkanRHI::vkCmdClearColorImage(CmdBuffer->GetHandle(), )
+		//VulkanRHI::vkCmdClearColorImage(CommandBuffer.GetHandle(), )
 	}
 }
 
@@ -888,8 +888,8 @@ uint32 FVulkanDynamicRHI::RHIGetGPUFrameCycles(uint32 GPUIndex)
 
 void FVulkanCommandListContext::RHISetDepthBounds(float MinDepth, float MaxDepth)
 {
-	FVulkanCommandBuffer* CmdBuffer = CommandBufferManager->GetActiveCmdBuffer();
-	VulkanRHI::vkCmdSetDepthBounds(CmdBuffer->GetHandle(), MinDepth, MaxDepth);
+	FVulkanCommandBuffer& CommandBuffer = GetCommandBuffer();
+	VulkanRHI::vkCmdSetDepthBounds(CommandBuffer.GetHandle(), MinDepth, MaxDepth);
 }
 
 void FVulkanCommandListContext::RequestSubmitCurrentCommands()
@@ -935,10 +935,10 @@ void FVulkanCommandListContext::PrepareForCPURead()
 
 void FVulkanCommandListContext::RHICopyToStagingBuffer(FRHIBuffer* SourceBufferRHI, FRHIStagingBuffer* StagingBufferRHI, uint32 Offset, uint32 NumBytes)
 {
-	FVulkanCommandBuffer* CmdBuffer = CommandBufferManager->GetActiveCmdBuffer();
+	FVulkanCommandBuffer& CommandBuffer = GetCommandBuffer();
 	FVulkanBuffer* VertexBuffer = ResourceCast(SourceBufferRHI);
 
-	ensure(CmdBuffer->IsOutsideRenderPass());
+	ensure(CommandBuffer.IsOutsideRenderPass());
 
 	FVulkanStagingBuffer* StagingBuffer = ResourceCast(StagingBufferRHI);
 	if (!StagingBuffer->StagingBuffer || StagingBuffer->StagingBuffer->GetSize() < NumBytes) //-V1051
@@ -960,7 +960,7 @@ void FVulkanCommandListContext::RHICopyToStagingBuffer(FRHIBuffer* SourceBufferR
 	Region.size = NumBytes;
 	Region.srcOffset = Offset + VertexBuffer->GetOffset();
 	//Region.dstOffset = 0;
-	VulkanRHI::vkCmdCopyBuffer(CmdBuffer->GetHandle(), VertexBuffer->GetHandle(), StagingBuffer->StagingBuffer->GetHandle(), 1, &Region);
+	VulkanRHI::vkCmdCopyBuffer(CommandBuffer.GetHandle(), VertexBuffer->GetHandle(), StagingBuffer->StagingBuffer->GetHandle(), 1, &Region);
 }
 
 void FVulkanCommandListContext::RHIWriteGPUFence(FRHIGPUFence* FenceRHI)
diff --git a/Engine/Source/Runtime/VulkanRHI/Private/VulkanContext.h b/Engine/Source/Runtime/VulkanRHI/Private/VulkanContext.h
--- a/Engine/Source/Runtime/VulkanRHI/Private/VulkanContext.h
+++ b/Engine/Source/Runtime/VulkanRHI/Private/VulkanContext.h
@@ -133,6 +133,8 @@ public:
 		return CommandBufferManager;
 	}
 
+	FVulkanCommandBuffer& GetCommandBuffer();
+
 	inline VulkanRHI::FTempBlockAllocator& GetTempBlockAllocator()
 	{
 		return *TempBlockAllocator;
@@ -223,15 +225,15 @@ protected:
 
 	TArray<FString> EventStack;
 
-	FVulkanCommandBufferManager* CommandBufferManager;
+	FVulkanCommandBufferManager* CommandBufferManager = nullptr;
 
 	FVulkanRenderPass* CurrentRenderPass = nullptr;
 	FVulkanFramebuffer* CurrentFramebuffer = nullptr;
 
 	FVulkanOcclusionQueryPool* CurrentOcclusionQueryPool = nullptr;
 
-	FVulkanPendingGfxState* PendingGfxState;
-	FVulkanPendingComputeState* PendingComputeState;
+	FVulkanPendingGfxState* PendingGfxState = nullptr;
+	FVulkanPendingComputeState* PendingComputeState = nullptr;
 
 	// Match the D3D12 maximum of 16 constant buffers per shader stage.
 	enum { MAX_UNIFORM_BUFFERS_PER_SHADER_STAGE = 16 };
@@ -245,7 +247,7 @@ protected:
 	void PrepareForCPURead();
 	void RequestSubmitCurrentCommands();
 
-	void InternalClearMRT(FVulkanCommandBuffer* CmdBuffer, bool bClearColor, int32 NumClearColors, const FLinearColor* ColorArray, bool bClearDepth, float Depth, bool bClearStencil, uint32 Stencil);
+	void InternalClearMRT(FVulkanCommandBuffer& CommandBuffer, bool bClearColor, int32 NumClearColors, const FLinearColor* ColorArray, bool bClearDepth, float Depth, bool bClearStencil, uint32 Stencil);
 
 public:
 	bool IsSwapchainImage(FRHITexture* InTexture) const;
diff --git a/Engine/Source/Runtime/VulkanRHI/Private/VulkanDescriptorSets.cpp b/Engine/Source/Runtime/VulkanRHI/Private/VulkanDescriptorSets.cpp
--- a/Engine/Source/Runtime/VulkanRHI/Private/VulkanDescriptorSets.cpp
+++ b/Engine/Source/Runtime/VulkanRHI/Private/VulkanDescriptorSets.cpp
@@ -606,7 +606,7 @@ void FVulkanBindlessDescriptorManager::BindDescriptorBuffers(VkCommandBuffer Com
 	}
 }
 
-void FVulkanBindlessDescriptorManager::RegisterUniformBuffers(FVulkanCommandBuffer* CommandBuffer, VkPipelineBindPoint BindPoint, const FUniformBufferDescriptorArrays& StageUBs)
+void FVulkanBindlessDescriptorManager::RegisterUniformBuffers(FVulkanCommandBuffer& CommandBuffer, VkPipelineBindPoint BindPoint, const FUniformBufferDescriptorArrays& StageUBs)
 {
 	checkf(bIsSupported, TEXT("Trying to RegisterUniformBuffers but bindless is not supported!"));
 
@@ -621,7 +621,7 @@ void FVulkanBindlessDescriptorManager::RegisterUniformBuffers(FVulkanCommandBuff
 	VkDeviceSize BufferOffsets[VulkanBindless::MaxNumSets];
 	FMemory::Memzero(BufferOffsets);
 
-	uint8* MappedPointer = SingleUseUBAllocator->Alloc(BlockSize, CommandBuffer, 
+	uint8* MappedPointer = SingleUseUBAllocator->Alloc(BlockSize, &CommandBuffer, 
 		LocalBufferBindingInfo[VulkanBindless::BindlessSingleUseUniformBufferSet], BufferOffsets[VulkanBindless::BindlessSingleUseUniformBufferSet]);
 
 	for (int32 StageIndex = 0; StageIndex < ShaderStage::MaxNumStages; ++StageIndex)
@@ -659,8 +659,8 @@ void FVulkanBindlessDescriptorManager::RegisterUniformBuffers(FVulkanCommandBuff
 	}
 
 	// todo-jn: cache these states and only repeat when necessary
-	VulkanRHI::vkCmdBindDescriptorBuffersEXT(CommandBuffer->GetHandle(), VulkanBindless::MaxNumSets, LocalBufferBindingInfo);
-	VulkanRHI::vkCmdSetDescriptorBufferOffsetsEXT(CommandBuffer->GetHandle(), BindPoint, BindlessPipelineLayout, 0u, VulkanBindless::MaxNumSets, BufferIndices, BufferOffsets);
+	VulkanRHI::vkCmdBindDescriptorBuffersEXT(CommandBuffer.GetHandle(), VulkanBindless::MaxNumSets, LocalBufferBindingInfo);
+	VulkanRHI::vkCmdSetDescriptorBufferOffsetsEXT(CommandBuffer.GetHandle(), BindPoint, BindlessPipelineLayout, 0u, VulkanBindless::MaxNumSets, BufferIndices, BufferOffsets);
 }
 
 void FVulkanBindlessDescriptorManager::UpdateStatsForHandle(FRHIDescriptorHandle DescriptorHandle)
diff --git a/Engine/Source/Runtime/VulkanRHI/Private/VulkanDescriptorSets.h b/Engine/Source/Runtime/VulkanRHI/Private/VulkanDescriptorSets.h
--- a/Engine/Source/Runtime/VulkanRHI/Private/VulkanDescriptorSets.h
+++ b/Engine/Source/Runtime/VulkanRHI/Private/VulkanDescriptorSets.h
@@ -879,7 +879,7 @@ public:
 	void UpdateTexelBuffer(FRHIDescriptorHandle DescriptorHandle, const VkBufferViewCreateInfo& ViewInfo, bool bImmediateUpdate = true);
 	void UpdateAccelerationStructure(FRHIDescriptorHandle DescriptorHandle, VkAccelerationStructureKHR AccelerationStructure, bool bImmediateUpdate = true);
 
-	void RegisterUniformBuffers(FVulkanCommandBuffer* CommandBuffer, VkPipelineBindPoint BindPoint, const FUniformBufferDescriptorArrays& StageUBs);
+	void RegisterUniformBuffers(FVulkanCommandBuffer& CommandBuffer, VkPipelineBindPoint BindPoint, const FUniformBufferDescriptorArrays& StageUBs);
 
 	void Unregister(FRHIDescriptorHandle DescriptorHandle);
 
diff --git a/Engine/Source/Runtime/VulkanRHI/Private/VulkanPendingState.cpp b/Engine/Source/Runtime/VulkanRHI/Private/VulkanPendingState.cpp
--- a/Engine/Source/Runtime/VulkanRHI/Private/VulkanPendingState.cpp
+++ b/Engine/Source/Runtime/VulkanRHI/Private/VulkanPendingState.cpp
@@ -339,7 +339,7 @@ void FVulkanPendingComputeState::SetUAVForUBResource(uint32 DescriptorSet, uint3
 }
 
 
-void FVulkanPendingComputeState::PrepareForDispatch(FVulkanCommandBuffer* InCmdBuffer)
+void FVulkanPendingComputeState::PrepareForDispatch(FVulkanCommandBuffer& InCommandBuffer)
 {
 #if VULKAN_ENABLE_AGGRESSIVE_STATS
 	SCOPE_CYCLE_COUNTER(STAT_VulkanDispatchCallPrepareTime);
@@ -347,25 +347,25 @@ void FVulkanPendingComputeState::PrepareForDispatch(FVulkanCommandBuffer* InCmdB
 
 	check(CurrentState);
 
-	VkCommandBuffer CmdBuffer = InCmdBuffer->GetHandle();
+	VkCommandBuffer CommandBufferHandle = InCommandBuffer.GetHandle();
 
 	if (Device->SupportsBindless())
 	{
-		CurrentState->UpdateBindlessDescriptors(&Context, InCmdBuffer);
-		CurrentPipeline->Bind(CmdBuffer);
+		CurrentState->UpdateBindlessDescriptors(&Context, InCommandBuffer);
+		CurrentPipeline->Bind(CommandBufferHandle);
 	}
 	else
 	{
-		const bool bHasDescriptorSets = CurrentState->UpdateDescriptorSets(&Context, InCmdBuffer);
+		const bool bHasDescriptorSets = CurrentState->UpdateDescriptorSets(&Context, &InCommandBuffer);
 
 		//#todo-rco: Move this to SetComputePipeline()
 #if VULKAN_ENABLE_AGGRESSIVE_STATS
 		SCOPE_CYCLE_COUNTER(STAT_VulkanPipelineBind);
 #endif
-		CurrentPipeline->Bind(CmdBuffer);
+		CurrentPipeline->Bind(CommandBufferHandle);
 		if (bHasDescriptorSets)
 		{
-			CurrentState->BindDescriptorSets(CmdBuffer);
+			CurrentState->BindDescriptorSets(CommandBufferHandle);
 		}
 	}
 }
@@ -381,19 +381,19 @@ FVulkanPendingGfxState::~FVulkanPendingGfxState()
 	}
 }
 
-void FVulkanPendingGfxState::PrepareForDraw(FVulkanCommandBuffer* CmdBuffer)
+void FVulkanPendingGfxState::PrepareForDraw(FVulkanCommandBuffer& CommandBuffer)
 {
 #if VULKAN_ENABLE_AGGRESSIVE_STATS
 	SCOPE_CYCLE_COUNTER(STAT_VulkanDrawCallPrepareTime);
 #endif
 
-	check(CmdBuffer->bHasPipeline);
+	check(CommandBuffer.bHasPipeline);
 
 	if (Device->SupportsBindless())
 	{
 		check(!CurrentPipeline->bHasInputAttachments); // todo-jn: bindless + InputAttachments
-		UpdateDynamicStates(CmdBuffer);
-		CurrentState->UpdateBindlessDescriptors(&Context, CmdBuffer);
+		UpdateDynamicStates(CommandBuffer);
+		CurrentState->UpdateBindlessDescriptors(&Context, CommandBuffer);
 	}
 	else
 	{
@@ -404,13 +404,13 @@ void FVulkanPendingGfxState::PrepareForDraw(FVulkanCommandBuffer* CmdBuffer)
 			UpdateInputAttachments(CurrentFramebuffer);
 		}
 
-		const bool bHasDescriptorSets = CurrentState->UpdateDescriptorSets(&Context, CmdBuffer);
+		const bool bHasDescriptorSets = CurrentState->UpdateDescriptorSets(&Context, &CommandBuffer);
 
-		UpdateDynamicStates(CmdBuffer);
+		UpdateDynamicStates(CommandBuffer);
 
 		if (bHasDescriptorSets)
 		{
-			CurrentState->BindDescriptorSets(CmdBuffer->GetHandle());
+			CurrentState->BindDescriptorSets(CommandBuffer.GetHandle());
 		}
 	}
 
@@ -483,16 +483,17 @@ void FVulkanPendingGfxState::PrepareForDraw(FVulkanCommandBuffer* CmdBuffer)
 			// Incorrect:	1, 0, 2, 3
 			// Incorrect:	0, 2, 3, 5
 			// Reordering and creation of stream binding index is done in "GenerateVertexInputStateInfo()"
-			VulkanRHI::vkCmdBindVertexBuffers(CmdBuffer->GetHandle(), 0, TemporaryIA.NumUsed, TemporaryIA.VertexBuffers, TemporaryIA.VertexOffsets);
+			VulkanRHI::vkCmdBindVertexBuffers(CommandBuffer.GetHandle(), 0, TemporaryIA.NumUsed, TemporaryIA.VertexBuffers, TemporaryIA.VertexOffsets);
 		}
 
 		bDirtyVertexStreams = false;
 	}
 }
 
-void FVulkanPendingGfxState::InternalUpdateDynamicStates(FVulkanCommandBuffer* Cmd)
+void FVulkanPendingGfxState::InternalUpdateDynamicStates(FVulkanCommandBuffer& CommandBuffer)
 {
-	const bool bNeedsUpdateViewport = !Cmd->bHasViewport || Viewports.Num() != Cmd->CurrentViewports.Num() || (FMemory::Memcmp((const void*)Cmd->CurrentViewports.GetData(), (const void*)Viewports.GetData(), Viewports.Num() * sizeof(VkViewport)) != 0);
+	const bool bNeedsUpdateViewport = !CommandBuffer.bHasViewport || Viewports.Num() != CommandBuffer.CurrentViewports.Num() ||
+		(FMemory::Memcmp((const void*)CommandBuffer.CurrentViewports.GetData(), (const void*)Viewports.GetData(), Viewports.Num() * sizeof(VkViewport)) != 0);
 	// Validate and update Viewport
 	if (bNeedsUpdateViewport)
 	{
@@ -506,29 +507,30 @@ void FVulkanPendingGfxState::InternalUpdateDynamicStates(FVulkanCommandBuffer* C
 			FlippedViewport.y += FlippedViewport.height;
 			FlippedViewport.height = -FlippedViewport.height;
 		}
-		VulkanRHI::vkCmdSetViewport(Cmd->GetHandle(), 0, FlippedViewports.Num(), FlippedViewports.GetData());
+		VulkanRHI::vkCmdSetViewport(CommandBuffer.GetHandle(), 0, FlippedViewports.Num(), FlippedViewports.GetData());
 
-		Cmd->CurrentViewports = Viewports;
-		Cmd->bHasViewport = true;
+		CommandBuffer.CurrentViewports = Viewports;
+		CommandBuffer.bHasViewport = true;
 	}
 
-	const bool bNeedsUpdateScissor = !Cmd->bHasScissor || Scissors.Num() != Cmd->CurrentScissors.Num() || (FMemory::Memcmp((const void*)Cmd->CurrentScissors.GetData(), (const void*)Scissors.GetData(), Scissors.Num() * sizeof(VkRect2D)) != 0);
+	const bool bNeedsUpdateScissor = !CommandBuffer.bHasScissor || Scissors.Num() != CommandBuffer.CurrentScissors.Num() || 
+		(FMemory::Memcmp((const void*)CommandBuffer.CurrentScissors.GetData(), (const void*)Scissors.GetData(), Scissors.Num() * sizeof(VkRect2D)) != 0);
 	if (bNeedsUpdateScissor)
 	{
-		VulkanRHI::vkCmdSetScissor(Cmd->GetHandle(), 0, Scissors.Num(), Scissors.GetData());
-		Cmd->CurrentScissors = Scissors;
-		Cmd->bHasScissor = true;
+		VulkanRHI::vkCmdSetScissor(CommandBuffer.GetHandle(), 0, Scissors.Num(), Scissors.GetData());
+		CommandBuffer.CurrentScissors = Scissors;
+		CommandBuffer.bHasScissor = true;
 	}
 
-	const bool bNeedsUpdateStencil = !Cmd->bHasStencilRef || (Cmd->CurrentStencilRef != StencilRef);
+	const bool bNeedsUpdateStencil = !CommandBuffer.bHasStencilRef || (CommandBuffer.CurrentStencilRef != StencilRef);
 	if (bNeedsUpdateStencil)
 	{
-		VulkanRHI::vkCmdSetStencilReference(Cmd->GetHandle(), VK_STENCIL_FRONT_AND_BACK, StencilRef);
-		Cmd->CurrentStencilRef = StencilRef;
-		Cmd->bHasStencilRef = true;
+		VulkanRHI::vkCmdSetStencilReference(CommandBuffer.GetHandle(), VK_STENCIL_FRONT_AND_BACK, StencilRef);
+		CommandBuffer.CurrentStencilRef = StencilRef;
+		CommandBuffer.bHasStencilRef = true;
 	}
 
-	Cmd->bNeedsDynamicStateSet = false;
+	CommandBuffer.bNeedsDynamicStateSet = false;
 }
 
 void FVulkanPendingGfxState::UpdateInputAttachments(FVulkanFramebuffer* Framebuffer)
diff --git a/Engine/Source/Runtime/VulkanRHI/Private/VulkanPendingState.h b/Engine/Source/Runtime/VulkanRHI/Private/VulkanPendingState.h
--- a/Engine/Source/Runtime/VulkanRHI/Private/VulkanPendingState.h
+++ b/Engine/Source/Runtime/VulkanRHI/Private/VulkanPendingState.h
@@ -55,7 +55,7 @@ public:
 		}
 	}
 
-	void PrepareForDispatch(FVulkanCommandBuffer* CmdBuffer);
+	void PrepareForDispatch(FVulkanCommandBuffer& InCommandBuffer);
 
 	inline const FVulkanComputeShader* GetCurrentShader() const
 	{
@@ -283,13 +283,13 @@ public:
 		CurrentState->SetPackedGlobalShaderParameter(Stage, BufferIndex, Offset, NumBytes, NewValue);
 	}
 
-	void PrepareForDraw(FVulkanCommandBuffer* CmdBuffer);
+	void PrepareForDraw(FVulkanCommandBuffer& CommandBuffer);
 
 	bool SetGfxPipeline(FVulkanRHIGraphicsPipelineState* InGfxPipeline, bool bForceReset);
 
-	inline void UpdateDynamicStates(FVulkanCommandBuffer* Cmd)
+	inline void UpdateDynamicStates(FVulkanCommandBuffer& CommandBuffer)
 	{
-		InternalUpdateDynamicStates(Cmd);
+		InternalUpdateDynamicStates(CommandBuffer);
 	}
 
 	inline void SetStencilRef(uint32 InStencilRef)
@@ -338,7 +338,7 @@ protected:
 	FVertexStream PendingStreams[MaxVertexElementCount];
 	bool bDirtyVertexStreams;
 
-	void InternalUpdateDynamicStates(FVulkanCommandBuffer* Cmd);
+	void InternalUpdateDynamicStates(FVulkanCommandBuffer& CommandBuffer);
 	void UpdateInputAttachments(FVulkanFramebuffer* Framebuffer);
 
 	FVulkanCommandListContext& Context;
diff --git a/Engine/Source/Runtime/VulkanRHI/Private/VulkanPipelineState.cpp b/Engine/Source/Runtime/VulkanRHI/Private/VulkanPipelineState.cpp
--- a/Engine/Source/Runtime/VulkanRHI/Private/VulkanPipelineState.cpp
+++ b/Engine/Source/Runtime/VulkanRHI/Private/VulkanPipelineState.cpp
@@ -223,7 +223,7 @@ bool FVulkanComputePipelineDescriptorState::InternalUpdateDescriptorSets(FVulkan
 	return true;
 }
 
-void FVulkanComputePipelineDescriptorState::UpdateBindlessDescriptors(FVulkanCommandListContext* CmdListContext, FVulkanCommandBuffer* CmdBuffer)
+void FVulkanComputePipelineDescriptorState::UpdateBindlessDescriptors(FVulkanCommandListContext* CmdListContext, FVulkanCommandBuffer& CommandBuffer)
 {
 	check(bUseBindless);
 
@@ -249,7 +249,7 @@ void FVulkanComputePipelineDescriptorState::UpdateBindlessDescriptors(FVulkanCom
 		const VkDeviceSize UBOffsetAlignment = Device->GetLimits().minUniformBufferOffsetAlignment;
 
 		VulkanRHI::FVulkanAllocation TempAllocation;
-		uint8* MappedPointer = CmdListContext->GetTempBlockAllocator().Alloc(UBSize, UBOffsetAlignment, CmdBuffer, TempAllocation, &DescriptorAddressInfos[BindingIndex]);
+		uint8* MappedPointer = CmdListContext->GetTempBlockAllocator().Alloc(UBSize, UBOffsetAlignment, &CommandBuffer, TempAllocation, &DescriptorAddressInfos[BindingIndex]);
 		FMemory::Memcpy(MappedPointer, StagedUniformBuffer.GetData(), UBSize);
 
 		PackedUniformBuffersDirty = 0;
@@ -279,7 +279,7 @@ void FVulkanComputePipelineDescriptorState::UpdateBindlessDescriptors(FVulkanCom
 	}
 
 	// Send to descriptor manager
-	Device->GetBindlessDescriptorManager()->RegisterUniformBuffers(CmdBuffer, VK_PIPELINE_BIND_POINT_COMPUTE, StageUBs);
+	Device->GetBindlessDescriptorManager()->RegisterUniformBuffers(CommandBuffer, VK_PIPELINE_BIND_POINT_COMPUTE, StageUBs);
 }
 
 FVulkanGraphicsPipelineDescriptorState::FVulkanGraphicsPipelineDescriptorState(FVulkanDevice* InDevice, FVulkanRHIGraphicsPipelineState* InGfxPipeline)
@@ -435,7 +435,7 @@ bool FVulkanGraphicsPipelineDescriptorState::InternalUpdateDescriptorSets(FVulka
 	return true;
 }
 
-void FVulkanGraphicsPipelineDescriptorState::UpdateBindlessDescriptors(FVulkanCommandListContext* CmdListContext, FVulkanCommandBuffer* CmdBuffer)
+void FVulkanGraphicsPipelineDescriptorState::UpdateBindlessDescriptors(FVulkanCommandListContext* CmdListContext, FVulkanCommandBuffer& CommandBuffer)
 {
 	check(bUseBindless);
 
@@ -475,7 +475,7 @@ void FVulkanGraphicsPipelineDescriptorState::UpdateBindlessDescriptors(FVulkanCo
 				const int32 BindingIndex = 0;
 
 				VulkanRHI::FVulkanAllocation TempAllocation;
-				uint8* MappedPointer = CmdListContext->GetTempBlockAllocator().Alloc(UBSize, UBOffsetAlignment, CmdBuffer, TempAllocation, &DescriptorAddressInfos[BindingIndex]);
+				uint8* MappedPointer = CmdListContext->GetTempBlockAllocator().Alloc(UBSize, UBOffsetAlignment, &CommandBuffer, TempAllocation, &DescriptorAddressInfos[BindingIndex]);
 				FMemory::Memcpy(MappedPointer, StagedUniformBuffer.GetData(), UBSize);
 
 				PackedUniformBuffersDirty[Stage] = 0;
@@ -507,7 +507,7 @@ void FVulkanGraphicsPipelineDescriptorState::UpdateBindlessDescriptors(FVulkanCo
 	}
 
 	// Send to descriptor manager
-	Device->GetBindlessDescriptorManager()->RegisterUniformBuffers(CmdBuffer, VK_PIPELINE_BIND_POINT_GRAPHICS, StageUBs);
+	Device->GetBindlessDescriptorManager()->RegisterUniformBuffers(CommandBuffer, VK_PIPELINE_BIND_POINT_GRAPHICS, StageUBs);
 }
 
 
@@ -544,16 +544,16 @@ void FVulkanCommandListContext::RHISetGraphicsPipelineState(FRHIGraphicsPipeline
 	FVulkanPipelineStateCacheManager* PipelineStateCache = Device->GetPipelineStateCache();
 	PipelineStateCache->LRUTouch(Pipeline);
 
-	FVulkanCommandBuffer* CmdBuffer = CommandBufferManager->GetActiveCmdBuffer();
-	bool bForceResetPipeline = !CmdBuffer->bHasPipeline;
+	FVulkanCommandBuffer& CommandBuffer = GetCommandBuffer();
+	bool bForceResetPipeline = !CommandBuffer.bHasPipeline;
 
 	if (PendingGfxState->SetGfxPipeline(Pipeline, bForceResetPipeline))
 	{
 #if VULKAN_ENABLE_AGGRESSIVE_STATS
 		SCOPE_CYCLE_COUNTER(STAT_VulkanPipelineBind);
 #endif
-		PendingGfxState->Bind(CmdBuffer->GetHandle());
-		CmdBuffer->bHasPipeline = true;
+		PendingGfxState->Bind(CommandBuffer.GetHandle());
+		CommandBuffer.bHasPipeline = true;
 		PendingGfxState->MarkNeedsDynamicStates();
 	}
 
diff --git a/Engine/Source/Runtime/VulkanRHI/Private/VulkanPipelineState.h b/Engine/Source/Runtime/VulkanRHI/Private/VulkanPipelineState.h
--- a/Engine/Source/Runtime/VulkanRHI/Private/VulkanPipelineState.h
+++ b/Engine/Source/Runtime/VulkanRHI/Private/VulkanPipelineState.h
@@ -211,7 +211,7 @@ public:
 		}
 	}
 
-	void UpdateBindlessDescriptors(FVulkanCommandListContext* CmdListContext, FVulkanCommandBuffer* CmdBuffer);
+	void UpdateBindlessDescriptors(FVulkanCommandListContext* CmdListContext, FVulkanCommandBuffer& CommandBuffer);
 
 	inline void BindDescriptorSets(VkCommandBuffer CmdBuffer)
 	{
@@ -262,7 +262,7 @@ public:
 		}
 	}
 
-	void UpdateBindlessDescriptors(FVulkanCommandListContext* CmdListContext, FVulkanCommandBuffer* CmdBuffer);
+	void UpdateBindlessDescriptors(FVulkanCommandListContext* CmdListContext, FVulkanCommandBuffer& CommandBuffer);
 
 	inline void BindDescriptorSets(VkCommandBuffer CmdBuffer)
 	{
diff --git a/Engine/Source/Runtime/VulkanRHI/Private/VulkanQuery.cpp b/Engine/Source/Runtime/VulkanRHI/Private/VulkanQuery.cpp
--- a/Engine/Source/Runtime/VulkanRHI/Private/VulkanQuery.cpp
+++ b/Engine/Source/Runtime/VulkanRHI/Private/VulkanQuery.cpp
@@ -586,8 +586,8 @@ void FVulkanCommandListContext::RHIBeginRenderQuery(FRHIRenderQuery* QueryRHI)
 		const uint32 IndexInPool = CurrentOcclusionQueryPool->AcquireIndex(Query);
 		Query->Pool = CurrentOcclusionQueryPool;
 		Query->IndexInPool = IndexInPool;
-		FVulkanCommandBuffer* CmdBuffer = CommandBufferManager->GetActiveCmdBuffer();
-		VulkanRHI::vkCmdBeginQuery(CmdBuffer->GetHandle(), CurrentOcclusionQueryPool->GetHandle(), IndexInPool, VK_QUERY_CONTROL_PRECISE_BIT);
+		FVulkanCommandBuffer& CommandBuffer = GetCommandBuffer();
+		VulkanRHI::vkCmdBeginQuery(CommandBuffer.GetHandle(), CurrentOcclusionQueryPool->GetHandle(), IndexInPool, VK_QUERY_CONTROL_PRECISE_BIT);
 	}
 	else if (BaseQuery->QueryType == RQT_AbsoluteTime)
 	{
@@ -609,8 +609,7 @@ void FVulkanCommandListContext::RHIEndRenderQuery(FRHIRenderQuery* QueryRHI)
 		FVulkanOcclusionQuery* Query = static_cast<FVulkanOcclusionQuery*>(BaseQuery);
 		ensure(Query->State == FVulkanOcclusionQuery::EState::RHI_PostBegin);
 		Query->State = FVulkanOcclusionQuery::EState::RHI_PostEnd;
-		FVulkanCommandBuffer* CmdBuffer = CommandBufferManager->GetActiveCmdBuffer();
-		VulkanRHI::vkCmdEndQuery(CmdBuffer->GetHandle(), CurrentOcclusionQueryPool->GetHandle(), Query->IndexInPool);
+		VulkanRHI::vkCmdEndQuery(GetCommandBuffer().GetHandle(), CurrentOcclusionQueryPool->GetHandle(), Query->IndexInPool);
 	}
 	else if (BaseQuery->QueryType == RQT_AbsoluteTime)
 	{
@@ -626,11 +625,11 @@ void FVulkanCommandListContext::RHIEndRenderQuery(FRHIRenderQuery* QueryRHI)
 		Query->Pool->CurrentTimestamp = (Query->Pool->CurrentTimestamp + 1) % Query->Pool->BufferSize;
 		const uint32 QueryEndIndex = Query->Pool->CurrentTimestamp;
 		const VkPipelineStageFlagBits QueryPipelineStage = GTimestampQueryStage ? VK_PIPELINE_STAGE_BOTTOM_OF_PIPE_BIT : VK_PIPELINE_STAGE_TOP_OF_PIPE_BIT;
-		FVulkanCommandBuffer* CmdBuffer = CommandBufferManager->GetActiveCmdBuffer();
-		VulkanRHI::vkCmdWriteTimestamp(CmdBuffer->GetHandle(), QueryPipelineStage, Query->Pool->GetHandle(), QueryEndIndex);
-		CmdBuffer->AddPendingTimestampQuery(QueryEndIndex, 1, Query->Pool->GetHandle(), Query->Pool->ResultsBuffer->GetHandle(), true);
-		Query->Pool->TimestampListHandles[QueryEndIndex].CmdBuffer = CmdBuffer;
-		Query->Pool->TimestampListHandles[QueryEndIndex].FenceCounter = CmdBuffer->GetFenceSignaledCounter();
+		FVulkanCommandBuffer& CommandBuffer = GetCommandBuffer();
+		VulkanRHI::vkCmdWriteTimestamp(CommandBuffer.GetHandle(), QueryPipelineStage, Query->Pool->GetHandle(), QueryEndIndex);
+		CommandBuffer.AddPendingTimestampQuery(QueryEndIndex, 1, Query->Pool->GetHandle(), Query->Pool->ResultsBuffer->GetHandle(), true);
+		Query->Pool->TimestampListHandles[QueryEndIndex].CmdBuffer = &CommandBuffer;
+		Query->Pool->TimestampListHandles[QueryEndIndex].FenceCounter = CommandBuffer.GetFenceSignaledCounter();
 		Query->Pool->TimestampListHandles[QueryEndIndex].FrameCount = GetFrameCounter();
 		Query->Pool->NumIssuedTimestamps = FMath::Min<uint32>(Query->Pool->NumIssuedTimestamps + 1, Query->Pool->BufferSize);
 	}
diff --git a/Engine/Source/Runtime/VulkanRHI/Private/VulkanRHI.cpp b/Engine/Source/Runtime/VulkanRHI/Private/VulkanRHI.cpp
--- a/Engine/Source/Runtime/VulkanRHI/Private/VulkanRHI.cpp
+++ b/Engine/Source/Runtime/VulkanRHI/Private/VulkanRHI.cpp
@@ -435,9 +435,6 @@ FVulkanCommandListContext::FVulkanCommandListContext(FVulkanDynamicRHI* InRHI, F
 	, Device(InDevice)
 	, Queue(InQueue)
 	, bSubmitAtNextSafePoint(false)
-	, CommandBufferManager(nullptr)
-	, PendingGfxState(nullptr)
-	, PendingComputeState(nullptr)
 	, FrameCounter(0)
 #if (RHI_NEW_GPU_PROFILER == 0)
 	, GpuProfiler(this, InDevice)
@@ -488,6 +485,13 @@ FVulkanCommandListContext::FVulkanCommandListContext(FVulkanDynamicRHI* InRHI, F
 	GlobalUniformBuffers.AddZeroed(FUniformBufferStaticSlotRegistry::Get().GetSlotCount());
 }
 
+FVulkanCommandBuffer& FVulkanCommandListContext::GetCommandBuffer()
+{
+	FVulkanCommandBuffer* CommandBuffer = CommandBufferManager->GetActiveCmdBuffer();
+	check(CommandBuffer);
+	return *CommandBuffer;
+}
+
 void FVulkanCommandListContext::ReleasePendingState()
 {
 	delete PendingGfxState;
@@ -1236,12 +1240,12 @@ void FVulkanCommandListContext::RHIEndDrawingViewport(FRHIViewport* ViewportRHI,
 	check(bPresent);
 	RHI->Present();
 */
-	FVulkanCommandBuffer* CmdBuffer = CommandBufferManager->GetActiveCmdBuffer();
-	check(!CmdBuffer->HasEnded() && !CmdBuffer->IsInsideRenderPass());
+	FVulkanCommandBuffer& CommandBuffer = GetCommandBuffer();
+	check(!CommandBuffer.HasEnded() && !CommandBuffer.IsInsideRenderPass());
 
-	WriteEndTimestamp(CmdBuffer);
+	WriteEndTimestamp(&CommandBuffer);
 
-	bool bNativePresent = Viewport->Present(this, CmdBuffer, Queue, Device->GetPresentQueue(), bLockToVsync);
+	bool bNativePresent = Viewport->Present(this, &CommandBuffer, Queue, Device->GetPresentQueue(), bLockToVsync);
 	if (bNativePresent)
 	{
 		//#todo-rco: Check for r.FinishCurrentFrame
diff --git a/Engine/Source/Runtime/VulkanRHI/Private/VulkanRayTracing.cpp b/Engine/Source/Runtime/VulkanRHI/Private/VulkanRayTracing.cpp
--- a/Engine/Source/Runtime/VulkanRHI/Private/VulkanRayTracing.cpp
+++ b/Engine/Source/Runtime/VulkanRHI/Private/VulkanRayTracing.cpp
@@ -1025,13 +1025,12 @@ FRHIShaderResourceView* FVulkanRayTracingShaderTable::GetOrCreateInlineBufferSRV
 
 void FVulkanRayTracingShaderTable::Commit(FVulkanCommandListContext& Context)
 {
-	FVulkanCommandBufferManager* CommandBufferManager = Context.GetCommandBufferManager();
-	FVulkanCommandBuffer* const CmdBuffer = CommandBufferManager->GetActiveCmdBuffer();
+	FVulkanCommandBuffer& CommandBuffer = Context.GetCommandBuffer();
 
 	VkMemoryBarrier BarrierBefore = { VK_STRUCTURE_TYPE_MEMORY_BARRIER, nullptr, VK_ACCESS_MEMORY_READ_BIT | VK_ACCESS_MEMORY_WRITE_BIT, VK_ACCESS_TRANSFER_READ_BIT | VK_ACCESS_TRANSFER_WRITE_BIT };
-	VulkanRHI::vkCmdPipelineBarrier(CmdBuffer->GetHandle(), VK_PIPELINE_STAGE_BOTTOM_OF_PIPE_BIT, VK_PIPELINE_STAGE_TRANSFER_BIT, 0, 1, &BarrierBefore, 0, nullptr, 0, nullptr);
+	VulkanRHI::vkCmdPipelineBarrier(CommandBuffer.GetHandle(), VK_PIPELINE_STAGE_BOTTOM_OF_PIPE_BIT, VK_PIPELINE_STAGE_TRANSFER_BIT, 0, 1, &BarrierBefore, 0, nullptr, 0, nullptr);
 
-	auto CommitBuffer = [Device = Device, CmdBuffer](FVulkanShaderTableAllocation& Alloc)
+	auto CommitBuffer = [Device = Device, &CommandBuffer](FVulkanShaderTableAllocation& Alloc)
 	{
 		if (Alloc.bIsDirty)
 		{
@@ -1067,9 +1066,9 @@ void FVulkanRayTracingShaderTable::Commit(FVulkanCommandListContext& Context)
 				RegionInfo.srcOffset = 0;
 				RegionInfo.dstOffset = 0;
 				RegionInfo.size = Alloc.Region.size;
-				VulkanRHI::vkCmdCopyBuffer(CmdBuffer->GetHandle(), StagingBuffer->GetHandle(), Alloc.LocalBuffer, 1, &RegionInfo);
+				VulkanRHI::vkCmdCopyBuffer(CommandBuffer.GetHandle(), StagingBuffer->GetHandle(), Alloc.LocalBuffer, 1, &RegionInfo);
 
-				Device->GetStagingManager().ReleaseBuffer(CmdBuffer, StagingBuffer);
+				Device->GetStagingManager().ReleaseBuffer(&CommandBuffer, StagingBuffer);
 			}
 			else
 			{
@@ -1098,7 +1097,7 @@ void FVulkanRayTracingShaderTable::Commit(FVulkanCommandListContext& Context)
 #endif // !VULKAN_SUPPORTS_RAY_TRACING_POSITION_FETCH
 
 	VkMemoryBarrier BarrierAfter = { VK_STRUCTURE_TYPE_MEMORY_BARRIER, nullptr, VK_ACCESS_TRANSFER_WRITE_BIT, VK_ACCESS_MEMORY_READ_BIT };  // :todo-jn: VK_ACCESS_2_SHADER_BINDING_TABLE_READ_BIT_KHR
-	VulkanRHI::vkCmdPipelineBarrier(CmdBuffer->GetHandle(), VK_PIPELINE_STAGE_TRANSFER_BIT, VK_PIPELINE_STAGE_RAY_TRACING_SHADER_BIT_KHR, 0, 1, &BarrierAfter, 0, nullptr, 0, nullptr);
+	VulkanRHI::vkCmdPipelineBarrier(CommandBuffer.GetHandle(), VK_PIPELINE_STAGE_TRANSFER_BIT, VK_PIPELINE_STAGE_RAY_TRACING_SHADER_BIT_KHR, 0, 1, &BarrierAfter, 0, nullptr, 0, nullptr);
 }
 
 FRHIShaderBindingTable* FVulkanRayTracingScene::FindOrCreateShaderBindingTable(const FRHIRayTracingPipelineState* InPipeline)
@@ -1853,8 +1852,7 @@ void FVulkanRayTracingCompactionRequestHandler::Update(FVulkanCommandListContext
 	// If we have an active batch, wait on those queries and launch compaction when the complete
 	if (ActiveBLASes.Num() > 0)
 	{		
-		FVulkanCommandBufferManager& CommandBufferManager = *InCommandContext.GetCommandBufferManager();
-		FVulkanCommandBuffer* const CmdBuffer = CommandBufferManager.GetActiveCmdBuffer();
+		FVulkanCommandBuffer& CmdBuffer = InCommandContext.GetCommandBuffer();
 
 		if (QueryPool->TryGetResults(ActiveBLASes.Num()))
 		{
@@ -1863,16 +1861,16 @@ void FVulkanRayTracingCompactionRequestHandler::Update(FVulkanCommandListContext
 			{
 				if (ActiveRequests[BLASIndex] != nullptr)
 				{
-					ActiveRequests[BLASIndex]->CompactAccelerationStructure(*CmdBuffer, QueryPool->GetResultValue(BLASIndex));
+					ActiveRequests[BLASIndex]->CompactAccelerationStructure(CmdBuffer, QueryPool->GetResultValue(BLASIndex));
 				}
 			}
 
-			QueryPool->Reset(CmdBuffer);
+			QueryPool->Reset(&CmdBuffer);
 
 			ActiveBLASes.Empty(GVulkanRayTracingMaxBatchedCompaction);
 
-			ActiveRequestsCmdBuffer = CmdBuffer;
-			ActiveRequestsFenceCounter = CmdBuffer->GetFenceSignaledCounter();
+			ActiveRequestsCmdBuffer = &CmdBuffer;
+			ActiveRequestsFenceCounter = CmdBuffer.GetFenceSignaledCounter();
 		}
 
 		// Only one active batch at a time (otherwise track the offset for when we launch queries)
@@ -1987,7 +1985,7 @@ static FVulkanPipelineBarrier SetRayGenResources(FVulkanDevice* Device, FVulkanC
 				NumSkippedSlots++;
 			}
 		}
-		Device->GetBindlessDescriptorManager()->RegisterUniformBuffers(CmdBuffer, VK_PIPELINE_BIND_POINT_RAY_TRACING_KHR, StageUBs);
+		Device->GetBindlessDescriptorManager()->RegisterUniformBuffers(*CmdBuffer, VK_PIPELINE_BIND_POINT_RAY_TRACING_KHR, StageUBs);
 	}
 
 	// Add all the UBs references by the shader table
diff --git a/Engine/Source/Runtime/VulkanRHI/Private/VulkanRenderTarget.cpp b/Engine/Source/Runtime/VulkanRHI/Private/VulkanRenderTarget.cpp
--- a/Engine/Source/Runtime/VulkanRHI/Private/VulkanRenderTarget.cpp
+++ b/Engine/Source/Runtime/VulkanRHI/Private/VulkanRenderTarget.cpp
@@ -705,9 +705,9 @@ void FVulkanCommandListContext::RHIEndRenderPass()
 void FVulkanCommandListContext::RHINextSubpass()
 {
 	check(CurrentRenderPass);
-	FVulkanCommandBuffer* CmdBuffer = CommandBufferManager->GetActiveCmdBuffer();
-	VkCommandBuffer Cmd = CmdBuffer->GetHandle();
-	VulkanRHI::vkCmdNextSubpass(Cmd, VK_SUBPASS_CONTENTS_INLINE);
+	FVulkanCommandBuffer& CommandBuffer = GetCommandBuffer();
+	VkCommandBuffer CommandBufferHandle = CommandBuffer.GetHandle();
+	VulkanRHI::vkCmdNextSubpass(CommandBufferHandle, VK_SUBPASS_CONTENTS_INLINE);
 }
 
 // Need a separate struct so we can memzero/remove dependencies on reference counts
diff --git a/Engine/Source/Runtime/VulkanRHI/Private/VulkanTexture.cpp b/Engine/Source/Runtime/VulkanRHI/Private/VulkanTexture.cpp
--- a/Engine/Source/Runtime/VulkanRHI/Private/VulkanTexture.cpp
+++ b/Engine/Source/Runtime/VulkanRHI/Private/VulkanTexture.cpp
@@ -566,8 +566,8 @@ struct FRHICommandSetInitialImageState final : public FRHICommand<FRHICommandSet
 		
 		if (bOnlyAddToLayoutManager)
 		{
-			FVulkanCommandBuffer* CmdBuffer = Context.GetCommandBufferManager()->GetActiveCmdBuffer();
-			CmdBuffer->GetLayoutManager().SetFullLayout(*VulkanTexture, InitialLayout, true);
+			FVulkanCommandBuffer& CommandBuffer = Context.GetCommandBuffer();
+			CommandBuffer.GetLayoutManager().SetFullLayout(*VulkanTexture, InitialLayout, true);
 		}
 		else
 		{
@@ -686,9 +686,9 @@ void FVulkanTexture::InternalMoveSurface(FVulkanDevice& InDevice, FVulkanCommand
 	DestAllocation.BindImage(&InDevice, MovedImage);
 
 	// Copy Original -> Moved
-	FVulkanCommandBuffer* CmdBuffer = Context.GetCommandBufferManager()->GetActiveCmdBuffer();
-	VkCommandBuffer VkCmdBuffer = CmdBuffer->GetHandle();
-	ensure(CmdBuffer->IsOutsideRenderPass());
+	FVulkanCommandBuffer& CommandBuffer = Context.GetCommandBuffer();
+	VkCommandBuffer CommandBufferHandle = CommandBuffer.GetHandle();
+	ensure(CommandBuffer.IsOutsideRenderPass());
 
 	{
 		const uint32 NumberOfArrayLevels = GetNumberOfArrayLevels();
@@ -699,7 +699,7 @@ void FVulkanTexture::InternalMoveSurface(FVulkanDevice& InDevice, FVulkanCommand
 			FVulkanPipelineBarrier Barrier;
 			Barrier.AddImageLayoutTransition(Image, OriginalLayout, VK_IMAGE_LAYOUT_TRANSFER_SRC_OPTIMAL, FullSubresourceRange);
 			Barrier.AddImageLayoutTransition(MovedImage, VK_IMAGE_LAYOUT_UNDEFINED, VK_IMAGE_LAYOUT_TRANSFER_DST_OPTIMAL, FullSubresourceRange);
-			Barrier.Execute(CmdBuffer);
+			Barrier.Execute(&CommandBuffer);
 		}
 		{
 			VkImageCopy Regions[MAX_TEXTURE_MIP_COUNT];
@@ -720,7 +720,7 @@ void FVulkanTexture::InternalMoveSurface(FVulkanDevice& InDevice, FVulkanCommand
 				Region.srcSubresource.mipLevel = i;
 				Region.dstSubresource.mipLevel = i;
 			}
-			VulkanRHI::vkCmdCopyImage(VkCmdBuffer,
+			VulkanRHI::vkCmdCopyImage(CommandBufferHandle,
 				Image, VK_IMAGE_LAYOUT_TRANSFER_SRC_OPTIMAL,
 				MovedImage, VK_IMAGE_LAYOUT_TRANSFER_DST_OPTIMAL,
 				Desc.NumMips, &Regions[0]);
@@ -731,7 +731,7 @@ void FVulkanTexture::InternalMoveSurface(FVulkanDevice& InDevice, FVulkanCommand
 			FVulkanPipelineBarrier Barrier;
 			Barrier.AddImageLayoutTransition(Image, VK_IMAGE_LAYOUT_TRANSFER_SRC_OPTIMAL, OriginalLayout, FullSubresourceRange);
 			Barrier.AddImageLayoutTransition(MovedImage, VK_IMAGE_LAYOUT_TRANSFER_DST_OPTIMAL, OriginalLayout, FullSubresourceRange);
-			Barrier.Execute(CmdBuffer);
+			Barrier.Execute(&CommandBuffer);
 		}
 	}
 
@@ -1865,8 +1865,8 @@ FVulkanTexture::FVulkanTexture(FVulkanDevice& InDevice, const FRHITextureCreateD
 				FVulkanCommandListContext& Context = InDevice.GetImmediateContext();
 				if (bOnlyAddToLayoutManager)
 				{
-					FVulkanCommandBuffer* CmdBuffer = Context.GetCommandBufferManager()->GetActiveCmdBuffer();
-					CmdBuffer->GetLayoutManager().SetFullLayout(*this, InitialLayout, true);
+					FVulkanCommandBuffer& CommandBuffer = Context.GetCommandBuffer();
+					CommandBuffer.GetLayoutManager().SetFullLayout(*this, InitialLayout, true);
 				}
 				else if (InitialLayout != VK_IMAGE_LAYOUT_UNDEFINED || bDoInitialClear)
 				{
@@ -2122,8 +2122,8 @@ FVulkanTexture::FVulkanTexture(FVulkanDevice& InDevice, const FRHITextureCreateD
 				FVulkanCommandListContext& Context = InDevice.GetImmediateContext();
 				if (bOnlyAddToLayoutManager)
 				{
-					FVulkanCommandBuffer* CmdBuffer = Context.GetCommandBufferManager()->GetActiveCmdBuffer();
-					CmdBuffer->GetLayoutManager().SetFullLayout(*this, InitialLayout, true);
+					FVulkanCommandBuffer& CommandBuffer = Context.GetCommandBuffer();
+					CommandBuffer.GetLayoutManager().SetFullLayout(*this, InitialLayout, true);
 				}
 				else if (InitialLayout != VK_IMAGE_LAYOUT_UNDEFINED || bDoInitialClear)
 				{
@@ -2559,11 +2559,11 @@ void FVulkanCommandListContext::RHICopyTexture(FRHITexture* SourceTexture, FRHIT
 	FVulkanTexture* Source = ResourceCast(SourceTexture);
 	FVulkanTexture* Dest = ResourceCast(DestTexture);
 
-	FVulkanCommandBuffer* CmdBuffer = CommandBufferManager->GetActiveCmdBuffer();
-	check(CmdBuffer->IsOutsideRenderPass());
+	FVulkanCommandBuffer& CommandBuffer = GetCommandBuffer();
+	check(CommandBuffer.IsOutsideRenderPass());
 
 	{
-		const VkImageLayout ExpectedSrcLayout = FVulkanLayoutManager::SetExpectedLayout(CmdBuffer, *Source, ERHIAccess::CopySrc);
+		const VkImageLayout ExpectedSrcLayout = FVulkanLayoutManager::SetExpectedLayout(&CommandBuffer, *Source, ERHIAccess::CopySrc);
 		ensureMsgf((ExpectedSrcLayout == VK_IMAGE_LAYOUT_TRANSFER_SRC_OPTIMAL), TEXT("Expected source texture to be in VK_IMAGE_LAYOUT_TRANSFER_SRC_OPTIMAL, actual layout is %s"), VK_TYPE_TO_STRING(VkImageLayout, ExpectedSrcLayout));
 	}
 
@@ -2614,16 +2614,16 @@ void FVulkanCommandListContext::RHICopyTexture(FRHITexture* SourceTexture, FRHIT
 			Size.Z = FMath::Max(1, Size.Z / 2);
 		}
 
-		VulkanRHI::vkCmdCopyImageToBuffer(CmdBuffer->GetHandle(), Source->Image, VK_IMAGE_LAYOUT_TRANSFER_SRC_OPTIMAL, CpuReadbackBuffer->Buffer, CopyInfo.NumMips, &CopyRegion[0]);
+		VulkanRHI::vkCmdCopyImageToBuffer(CommandBuffer.GetHandle(), Source->Image, VK_IMAGE_LAYOUT_TRANSFER_SRC_OPTIMAL, CpuReadbackBuffer->Buffer, CopyInfo.NumMips, &CopyRegion[0]);
 
 		FVulkanPipelineBarrier BarrierMemory;
 		BarrierMemory.AddMemoryBarrier(VK_ACCESS_TRANSFER_WRITE_BIT, VK_ACCESS_HOST_READ_BIT, VK_PIPELINE_STAGE_ALL_COMMANDS_BIT, VK_PIPELINE_STAGE_HOST_BIT);
-		BarrierMemory.Execute(CmdBuffer);
+		BarrierMemory.Execute(&CommandBuffer);
 	}
 	else
 	{
 		{
-			const VkImageLayout ExpectedDstLayout = FVulkanLayoutManager::SetExpectedLayout(CmdBuffer, *Dest, ERHIAccess::CopyDest);
+			const VkImageLayout ExpectedDstLayout = FVulkanLayoutManager::SetExpectedLayout(&CommandBuffer, *Dest, ERHIAccess::CopyDest);
 			ensureMsgf((ExpectedDstLayout == VK_IMAGE_LAYOUT_TRANSFER_DST_OPTIMAL), TEXT("Expected destination texture to be in VK_IMAGE_LAYOUT_TRANSFER_DST_OPTIMAL, actual layout is %s"), VK_TYPE_TO_STRING(VkImageLayout, ExpectedDstLayout));
 		}
 
@@ -2661,7 +2661,7 @@ void FVulkanCommandListContext::RHICopyTexture(FRHITexture* SourceTexture, FRHIT
 
 		for (uint32 Index = 0; Index < CopyInfo.NumMips; ++Index)
 		{
-			VulkanRHI::vkCmdCopyImage(CmdBuffer->GetHandle(),
+			VulkanRHI::vkCmdCopyImage(CommandBuffer.GetHandle(),
 				Source->Image, VK_IMAGE_LAYOUT_TRANSFER_SRC_OPTIMAL,
 				Dest->Image, VK_IMAGE_LAYOUT_TRANSFER_DST_OPTIMAL,
 				1, &Region);
@@ -2709,19 +2709,19 @@ void FVulkanCommandListContext::RHICopyBufferRegion(FRHIBuffer* DstBuffer, uint6
 	uint64 DstOffsetVk = DstBufferVk->GetOffset() + DstOffset;
 	uint64 SrcOffsetVk = SrcBufferVk->GetOffset() + SrcOffset;
 
-	FVulkanCommandBuffer* CmdBuffer = CommandBufferManager->GetActiveCmdBuffer();
-	check(CmdBuffer->IsOutsideRenderPass());
-	VkCommandBuffer VkCmdBuffer = CmdBuffer->GetHandle();
+	FVulkanCommandBuffer& CommandBuffer = GetCommandBuffer();
+	check(CommandBuffer.IsOutsideRenderPass());
+	VkCommandBuffer CommandBufferHandle = CommandBuffer.GetHandle();
 
 	VkMemoryBarrier BarrierBefore = { VK_STRUCTURE_TYPE_MEMORY_BARRIER, nullptr, VK_ACCESS_MEMORY_READ_BIT|VK_ACCESS_MEMORY_WRITE_BIT, VK_ACCESS_TRANSFER_READ_BIT| VK_ACCESS_TRANSFER_WRITE_BIT };
-	VulkanRHI::vkCmdPipelineBarrier(VkCmdBuffer, VK_PIPELINE_STAGE_ALL_COMMANDS_BIT, VK_PIPELINE_STAGE_TRANSFER_BIT, 0, 1, &BarrierBefore, 0, nullptr, 0, nullptr);
+	VulkanRHI::vkCmdPipelineBarrier(CommandBufferHandle, VK_PIPELINE_STAGE_ALL_COMMANDS_BIT, VK_PIPELINE_STAGE_TRANSFER_BIT, 0, 1, &BarrierBefore, 0, nullptr, 0, nullptr);
 
 	VkBufferCopy Region = {};
 	Region.srcOffset = SrcOffsetVk;
 	Region.dstOffset = DstOffsetVk;
 	Region.size = NumBytes;
-	VulkanRHI::vkCmdCopyBuffer(VkCmdBuffer, SrcBufferVk->GetHandle(), DstBufferVk->GetHandle(), 1, &Region);
+	VulkanRHI::vkCmdCopyBuffer(CommandBufferHandle, SrcBufferVk->GetHandle(), DstBufferVk->GetHandle(), 1, &Region);
 
 	VkMemoryBarrier BarrierAfter = { VK_STRUCTURE_TYPE_MEMORY_BARRIER, nullptr, VK_ACCESS_TRANSFER_READ_BIT | VK_ACCESS_TRANSFER_WRITE_BIT, VK_ACCESS_MEMORY_READ_BIT | VK_ACCESS_MEMORY_WRITE_BIT };
-	VulkanRHI::vkCmdPipelineBarrier(VkCmdBuffer, VK_PIPELINE_STAGE_TRANSFER_BIT, VK_PIPELINE_STAGE_ALL_COMMANDS_BIT, 0, 1, &BarrierAfter, 0, nullptr, 0, nullptr);
+	VulkanRHI::vkCmdPipelineBarrier(CommandBufferHandle, VK_PIPELINE_STAGE_TRANSFER_BIT, VK_PIPELINE_STAGE_ALL_COMMANDS_BIT, 0, 1, &BarrierAfter, 0, nullptr, 0, nullptr);
 }
diff --git a/Engine/Source/Runtime/VulkanRHI/Private/VulkanUAV.cpp b/Engine/Source/Runtime/VulkanRHI/Private/VulkanUAV.cpp
--- a/Engine/Source/Runtime/VulkanRHI/Private/VulkanUAV.cpp
+++ b/Engine/Source/Runtime/VulkanRHI/Private/VulkanUAV.cpp
@@ -576,17 +576,17 @@ void FVulkanUnorderedAccessView::Clear(TRHICommandList_RecursiveHazardous<FVulka
 		case FRHIViewDesc::EBufferType::Structured:
 			RHICmdList.RunOnContext([this, Buffer, Info, ClearValue = *static_cast<const uint32*>(ClearValue)](FVulkanCommandListContext& Context)
 			{
-				FVulkanCommandBuffer* CmdBuffer = Context.GetCommandBufferManager()->GetActiveCmdBuffer();
+				FVulkanCommandBuffer& CmdBuffer = Context.GetCommandBuffer();
 
 				// vkCmdFillBuffer is treated as a transfer operation for the purposes of synchronization barriers.
 				{
 					FVulkanPipelineBarrier BeforeBarrier;
 					BeforeBarrier.AddMemoryBarrier(VK_ACCESS_MEMORY_READ_BIT | VK_ACCESS_MEMORY_WRITE_BIT, VK_ACCESS_TRANSFER_WRITE_BIT, VK_PIPELINE_STAGE_ALL_COMMANDS_BIT, VK_PIPELINE_STAGE_TRANSFER_BIT);
-					BeforeBarrier.Execute(CmdBuffer);
+					BeforeBarrier.Execute(&CmdBuffer);
 				}
 
 				VulkanRHI::vkCmdFillBuffer(
-					  CmdBuffer->GetHandle()
+					  CmdBuffer.GetHandle()
 					, Buffer->GetHandle()
 					, Buffer->GetOffset() + Info.OffsetInBytes
 					, Info.SizeInBytes
@@ -596,7 +596,7 @@ void FVulkanUnorderedAccessView::Clear(TRHICommandList_RecursiveHazardous<FVulka
 				{
 					FVulkanPipelineBarrier AfterBarrier;
 					AfterBarrier.AddMemoryBarrier(VK_ACCESS_TRANSFER_WRITE_BIT, VK_ACCESS_MEMORY_READ_BIT | VK_ACCESS_MEMORY_WRITE_BIT, VK_PIPELINE_STAGE_TRANSFER_BIT, VK_PIPELINE_STAGE_ALL_COMMANDS_BIT);
-					AfterBarrier.Execute(CmdBuffer);
+					AfterBarrier.Execute(&CmdBuffer);
 				}
 			});
 			break;
diff --git a/Engine/Source/Runtime/VulkanRHI/Private/VulkanViewport.cpp b/Engine/Source/Runtime/VulkanRHI/Private/VulkanViewport.cpp
--- a/Engine/Source/Runtime/VulkanRHI/Private/VulkanViewport.cpp
+++ b/Engine/Source/Runtime/VulkanRHI/Private/VulkanViewport.cpp
@@ -82,17 +82,16 @@ void FVulkanBackBuffer::AcquireBackBufferImage(FVulkanCommandListContext& Contex
 			DefaultView = &ImageView;
 			PartialView = &ImageView;
 
-			FVulkanCommandBufferManager* CmdBufferManager = Context.GetCommandBufferManager();
-			FVulkanCommandBuffer* CmdBuffer = CmdBufferManager->GetActiveCmdBuffer();
-			check(!CmdBuffer->IsInsideRenderPass());
+			FVulkanCommandBuffer& CommandBuffer = Context.GetCommandBuffer();
+			check(!CommandBuffer.IsInsideRenderPass());
 
 			// right after acquiring image is in undefined state
-			FVulkanLayoutManager& LayoutMgr = CmdBuffer->GetLayoutManager();
+			FVulkanLayoutManager& LayoutMgr = CommandBuffer.GetLayoutManager();
 			const FVulkanImageLayout CustomLayout(VK_IMAGE_LAYOUT_UNDEFINED, 1, 1, VK_IMAGE_ASPECT_COLOR_BIT);
 			LayoutMgr.SetFullLayout(Image, CustomLayout);
 
 			// Wait for semaphore signal before writing to backbuffer image
-			CmdBuffer->AddWaitSemaphore(VK_PIPELINE_STAGE_COLOR_ATTACHMENT_OUTPUT_BIT, Viewport->AcquiredSemaphore);
+			CommandBuffer.AddWaitSemaphore(VK_PIPELINE_STAGE_COLOR_ATTACHMENT_OUTPUT_BIT, Viewport->AcquiredSemaphore);
 		}
 		else
 		{
