Disable Vulkan UAV aliasing

It is currently breaking Rendering tests

diff --git a/Engine/Source/Runtime/VulkanRHI/Private/VulkanRHI.cpp b/Engine/Source/Runtime/VulkanRHI/Private/VulkanRHI.cpp
--- a/Engine/Source/Runtime/VulkanRHI/Private/VulkanRHI.cpp
+++ b/Engine/Source/Runtime/VulkanRHI/Private/VulkanRHI.cpp
@@ -987,7 +987,8 @@ void FVulkanDynamicRHI::InitInstance()
 		GRHISupportsParallelRHIExecute = GRHIThreadCvar->GetInt() > 1;
 #endif
 
-		GRHISupportsUAVFormatAliasing = true;
+		// Disabled since it causes issues in Lumen tests
+		//GRHISupportsUAVFormatAliasing = true;
 
 		// Some platforms might only have CPU for an RHI thread, but not for parallel tasks
 		GSupportsParallelRenderingTasksWithSeparateRHIThread = GRHISupportsRHIThread ? FVulkanPlatform::SupportParallelRenderingTasks() : false;
