Use R8G8B8_UNORM where supported
When r.Mobile.SceneColorFormat=3 use R8G8B8_UNORM if supported and alpha is not needed
For everyone else use R8G8B8_UNORM in tonemapping if supported and alpha is not needed

diff --git a/Engine/Source/Runtime/Core/Private/Misc/PixelFormat.cpp b/Engine/Source/Runtime/Core/Private/Misc/PixelFormat.cpp
--- a/Engine/Source/Runtime/Core/Private/Misc/PixelFormat.cpp
+++ b/Engine/Source/Runtime/Core/Private/Misc/PixelFormat.cpp
@@ -146,6 +146,7 @@ FPixelFormatInfo    GPixelFormats[PF_MAX] =
 	FPixelFormatInfo(PF_ASTC_12x12_NORM_RG, TEXT("ASTC_12x12_NORM_RG"),    12,         12,         1,          16,         2,              0),
 
 	FPixelFormatInfo(PF_R16G16_SINT,        TEXT("PF_R16G16_SINT"),        1,          1,          1,          4,          4,              1),
+	FPixelFormatInfo(PF_R8G8B8,             TEXT("R8G8B8"),                1,          1,          1,          3,          3,              0),
 };
 
 
diff --git a/Engine/Source/Runtime/Core/Public/PixelFormat.h b/Engine/Source/Runtime/Core/Public/PixelFormat.h
--- a/Engine/Source/Runtime/Core/Public/PixelFormat.h
+++ b/Engine/Source/Runtime/Core/Public/PixelFormat.h
@@ -107,7 +107,8 @@ enum EPixelFormat : uint8
 	PF_ASTC_10x10_NORM_RG	=90,	
 	PF_ASTC_12x12_NORM_RG	=91,	
 	PF_R16G16_SINT			=92,
-	PF_MAX					=93,
+	PF_R8G8B8				=93,
+	PF_MAX					=94,
 };
 #define FOREACH_ENUM_EPIXELFORMAT(op) \
 	op(PF_Unknown) \
@@ -202,6 +203,7 @@ enum EPixelFormat : uint8
 	op(PF_ASTC_8x8_NORM_RG) \
 	op(PF_ASTC_10x10_NORM_RG) \
 	op(PF_ASTC_12x12_NORM_RG) \
+	op(PF_R8G8B8)\
 	op(PF_R16G16_SINT)
 
 // Defines which channel is valid for each pixel format
diff --git a/Engine/Source/Runtime/Engine/Private/SceneTexturesConfig.cpp b/Engine/Source/Runtime/Engine/Private/SceneTexturesConfig.cpp
--- a/Engine/Source/Runtime/Engine/Private/SceneTexturesConfig.cpp
+++ b/Engine/Source/Runtime/Engine/Private/SceneTexturesConfig.cpp
@@ -30,8 +30,10 @@ FSceneTextureShaderParameters GetSceneTextureShaderParameters(TRDGUniformBufferR
 	return Parameters;
 }
 
-static EPixelFormat GetDefaultMobileSceneColorLowPrecisionFormat()
+static EPixelFormat GetDefaultMobileSceneColorLowPrecisionFormat(bool bRequiresAlphaChannel)
 {
+	bool bShouldUseRGBRT = !bRequiresAlphaChannel && GPixelFormats[PF_R8G8B8].Supported && !!(GPixelFormats[PF_R8G8B8].Capabilities & EPixelFormatCapabilities::RenderTarget);
+
 	if (UNLIKELY(GEngine && GEngine->XRSystem.IsValid() && GEngine->StereoRenderingDevice.IsValid()))
 	{
 		IStereoRenderTargetManager* RTManager = GEngine->StereoRenderingDevice->GetRenderTargetManager();
@@ -46,10 +48,10 @@ static EPixelFormat GetDefaultMobileSceneColorLowPrecisionFormat()
 		
 		if (GEngine->StereoRenderingDevice->IsStandaloneStereoOnlyDevice())
 		{
-			return PF_R8G8B8A8;
+			return bShouldUseRGBRT ? PF_R8G8B8 : PF_R8G8B8A8;
 		}
 	}
-	return PF_B8G8R8A8;
+	return  bShouldUseRGBRT ? PF_R8G8B8 : PF_B8G8R8A8;
 }
 
 static EPixelFormat GetMobileSceneColorFormat(bool bRequiresAlphaChannel)
@@ -58,7 +60,7 @@ static EPixelFormat GetMobileSceneColorFormat(bool bRequiresAlphaChannel)
 	const bool bUseLowPrecisionFormat = !IsMobileHDR() || !GSupportsRenderTargetFormat_PF_FloatRGBA;
 	if (bUseLowPrecisionFormat)
 	{
-		DefaultColorFormat = GetDefaultMobileSceneColorLowPrecisionFormat();
+		DefaultColorFormat = GetDefaultMobileSceneColorLowPrecisionFormat(bRequiresAlphaChannel);
 	}
 	else
 	{
@@ -80,12 +82,12 @@ static EPixelFormat GetMobileSceneColorFormat(bool bRequiresAlphaChannel)
 		if (bUseLowPrecisionFormat)
 		{
 			// if bUseLowPrecisionFormat, DefaultColorFormat already contains the value of GetDefaultMobileSceneColorLowPrecisionFormat
-			checkSlow(DefaultColorFormat == GetDefaultMobileSceneColorLowPrecisionFormat());
+			checkSlow(DefaultColorFormat == GetDefaultMobileSceneColorLowPrecisionFormat(bRequiresAlphaChannel));
 			Format = DefaultColorFormat;
 		}
 		else
 		{
-			Format = GetDefaultMobileSceneColorLowPrecisionFormat();
+			Format = GetDefaultMobileSceneColorLowPrecisionFormat(bRequiresAlphaChannel);
 		}
 		break;
 	default:
diff --git a/Engine/Source/Runtime/OpenGLDrv/Private/OpenGLDevice.cpp b/Engine/Source/Runtime/OpenGLDrv/Private/OpenGLDevice.cpp
--- a/Engine/Source/Runtime/OpenGLDrv/Private/OpenGLDevice.cpp
+++ b/Engine/Source/Runtime/OpenGLDrv/Private/OpenGLDevice.cpp
@@ -1072,7 +1072,7 @@ static void InitRHICapabilitiesForGL()
 	SetupTextureFormat( PF_R8G8B8A8,			FOpenGLTextureFormat( GL_RGBA8,					GL_SRGB8_ALPHA8,		GL_RGBA,			GL_UNSIGNED_BYTE,					false, false));
 	SetupTextureFormat( PF_R8G8B8A8_UINT,		FOpenGLTextureFormat( GL_RGBA8UI,				GL_RGBA8UI,				GL_RGBA_INTEGER,	GL_UNSIGNED_BYTE,					false, false));
 	SetupTextureFormat( PF_G8,					FOpenGLTextureFormat( GL_R8,					GL_R8,					GL_RED,				GL_UNSIGNED_BYTE,					false, false));
-
+	SetupTextureFormat( PF_R8G8B8,				FOpenGLTextureFormat( GL_RGB8,					GL_SRGB8,				GL_RGB,				GL_UNSIGNED_BYTE,					false, false));
 #if PLATFORM_DESKTOP
 	CA_SUPPRESS(6286);
 	if (PLATFORM_DESKTOP)
diff --git a/Engine/Source/Runtime/RenderCore/Private/RenderUtils.cpp b/Engine/Source/Runtime/RenderCore/Private/RenderUtils.cpp
--- a/Engine/Source/Runtime/RenderCore/Private/RenderUtils.cpp
+++ b/Engine/Source/Runtime/RenderCore/Private/RenderUtils.cpp
@@ -323,6 +323,7 @@ EPixelFormatChannelFlags GetPixelFormatValidChannels(EPixelFormat InPixelFormat)
 		EPixelFormatChannelFlags::RG,		// PF_ASTC_10x10_NORM_RG
 		EPixelFormatChannelFlags::RG,		// PF_ASTC_12x12_NORM_RG
 		EPixelFormatChannelFlags::RG,		// PF_R16G16_SINT
+		EPixelFormatChannelFlags::RGB,		// PF_R8G8B8
 	};
 	static_assert(UE_ARRAY_COUNT(PixelFormatToChannelFlags) == (uint8)PF_MAX, "Missing pixel format");
 	return (InPixelFormat < PF_MAX) ? PixelFormatToChannelFlags[(uint8)InPixelFormat] : EPixelFormatChannelFlags::None;
diff --git a/Engine/Source/Runtime/Renderer/Private/PostProcess/PostProcessTonemap.cpp b/Engine/Source/Runtime/Renderer/Private/PostProcess/PostProcessTonemap.cpp
--- a/Engine/Source/Runtime/Renderer/Private/PostProcess/PostProcessTonemap.cpp
+++ b/Engine/Source/Runtime/Renderer/Private/PostProcess/PostProcessTonemap.cpp
@@ -541,12 +541,12 @@ FRDGBufferRef BuildFilmGrainConstants(FRDGBuilder& GraphBuilder, const FViewInfo
 	return FilmGrainConstantsBuffer;
 }
 
-bool ShouldWriteAlphaChannel(const FViewInfo& View, const FTonemapInputs& Inputs, const FRDGTextureRef Output)
+bool ShouldWriteAlphaChannel(const FViewInfo& View, const FTonemapInputs& Inputs, EPixelFormat Format)
 {
 	// If this is a stereo view, there's a good chance we need alpha out of the tonemapper
 	// @todo: Remove this once Oculus fix the bug in their runtime that requires alpha here.
 	const bool bIsStereo = IStereoRendering::IsStereoEyeView(View);
-	const bool bFormatNeedsAlphaWrite = Output->Desc.Format == PF_R9G9B9EXP5;
+	const bool bFormatNeedsAlphaWrite = Format == PF_R9G9B9EXP5;
 	return (Inputs.bWriteAlphaChannel || bIsStereo || bFormatNeedsAlphaWrite);
 }
 
@@ -636,13 +636,23 @@ FScreenPassTexture AddTonemapPass(FRDGBuilder& GraphBuilder, const FViewInfo& Vi
 			OutputDesc.Format = PF_A2B10G10R10;
 		}
 
+		if (
+			(OutputDesc.Format == PF_R8G8B8A8 || OutputDesc.Format == PF_B8G8R8A8) && 
+			!ShouldWriteAlphaChannel(View, Inputs, OutputDesc.Format) && 
+			GPixelFormats[PF_R8G8B8].Supported &&
+			!!(GPixelFormats[PF_R8G8B8].Capabilities & EPixelFormatCapabilities::RenderTarget)
+			)
+		{
+			OutputDesc.Format = PF_R8G8B8;
+		}
+
 		Output = FScreenPassRenderTarget(
 			GraphBuilder.CreateTexture(OutputDesc, TEXT("Tonemap")),
 			Inputs.SceneColor.ViewRect,
 			ERenderTargetLoadAction::ENoAction);
 	}
 
-	const bool bShouldWriteAlphaChannel = ShouldWriteAlphaChannel(View, Inputs, Output.Texture);
+	const bool bShouldWriteAlphaChannel = ShouldWriteAlphaChannel(View, Inputs, Output.Texture->Desc.Format);
 	if (ShouldOverrideOutputLoadActionToFastClear(Output, bShouldWriteAlphaChannel))
 	{
 		Output.LoadAction = ERenderTargetLoadAction::EClear;
diff --git a/Engine/Source/Runtime/Renderer/Private/SystemTextures.cpp b/Engine/Source/Runtime/Renderer/Private/SystemTextures.cpp
--- a/Engine/Source/Runtime/Renderer/Private/SystemTextures.cpp
+++ b/Engine/Source/Runtime/Renderer/Private/SystemTextures.cpp
@@ -1335,7 +1335,7 @@ template<typename TInType>
 void InitializeData(const TInType& InData, EPixelFormat InFormat, uint8* OutData, uint32& OutByteCount)
 {
 	// If a new format is added insure that it is either supported here, or at least flagged as not supported
-	static_assert(PF_MAX == 93);
+	static_assert(PF_MAX == 94);
 
 	switch (InFormat)
 	{
@@ -1379,6 +1379,7 @@ void InitializeData(const TInType& InData, EPixelFormat InFormat, uint8* OutData
 		case PF_A8R8G8B8:				{ FormatData<EDefaultInputType::UNorm, TInType,	uint8,    3, 2, 1, 0>	(InData, OutData, OutByteCount); } break;
 		case PF_R8G8B8A8_UINT:			{ FormatData<EDefaultInputType::Typed, TInType,	uint8,    0, 1, 2, 3>	(InData, OutData, OutByteCount); } break;
 		case PF_R8G8B8A8_SNORM:			{ FormatData<EDefaultInputType::SNorm, TInType,	int8,     0, 1, 2, 3>	(InData, OutData, OutByteCount); } break;
+		case PF_R8G8B8:					{ FormatData<EDefaultInputType::UNorm, TInType,	uint8,    0, 1, 2>		(InData, OutData, OutByteCount); } break;
 		case PF_R8G8:					{ FormatData<EDefaultInputType::UNorm, TInType,	uint8,    0, 1>			(InData, OutData, OutByteCount); } break;
 		case PF_R8G8_UINT:				{ FormatData<EDefaultInputType::Typed, TInType, uint8,    0, 1>			(InData, OutData, OutByteCount); } break;
 		case PF_R8_UINT:				{ FormatData<EDefaultInputType::Typed, TInType,	uint8>					(InData, OutData, OutByteCount); } break;
diff --git a/Engine/Source/Runtime/VulkanRHI/Private/VulkanDevice.cpp b/Engine/Source/Runtime/VulkanRHI/Private/VulkanDevice.cpp
--- a/Engine/Source/Runtime/VulkanRHI/Private/VulkanDevice.cpp
+++ b/Engine/Source/Runtime/VulkanRHI/Private/VulkanDevice.cpp
@@ -653,6 +653,7 @@ void FVulkanDevice::SetupFormats()
 	MapFormatSupport(PF_R32G32B32_SINT, { VK_FORMAT_R32G32B32_SINT }, ComponentMappingRGB1);
 	MapFormatSupport(PF_R32G32B32F, { VK_FORMAT_R32G32B32_SFLOAT }, ComponentMappingRGB1);
 	MapFormatSupport(PF_R8_SINT, { VK_FORMAT_R8_SINT }, ComponentMappingR001);
+	MapFormatSupport(PF_R8G8B8, { VK_FORMAT_R8G8B8_UNORM }, ComponentMappingRGB1, ColorRenderTargetRequiredCapabilities);
 
 	// This will be the format used for 64bit image atomics
 	// This format is SM5 only, skip it for mobile to not confuse QA with a logged error about missing pixel format
@@ -974,10 +975,15 @@ void FVulkanDevice::MapImageFormatSupport(FPixelFormatInfo& PixelFormatInfo, con
 	};
 
 	EPixelFormatCapabilities Capabilities = EPixelFormatCapabilities::None;
-	auto ConvertImageCap = [&Capabilities](const VkFormatProperties& InFormatProperties, EPixelFormatCapabilities UnrealCap, VkFormatFeatureFlags InFlag)
+	auto ConvertImageCap = [&Capabilities](const VkFormatProperties& InFormatProperties, EPixelFormatCapabilities UnrealCap, VkFormatFeatureFlags InFlag, bool bOnlyOptimalTiling)
 	{
 		// Do not distinguish between Linear and Optimal for now.
-		const bool HasImageFeature = VKHasAllFlags(InFormatProperties.linearTilingFeatures, InFlag) || VKHasAllFlags(InFormatProperties.optimalTilingFeatures, InFlag);
+		bool HasImageFeature = VKHasAllFlags(InFormatProperties.optimalTilingFeatures, InFlag);
+		if (!bOnlyOptimalTiling)
+		{
+			HasImageFeature |= VKHasAllFlags(InFormatProperties.linearTilingFeatures, InFlag);
+		}
+
 		if (HasImageFeature)
 		{
 			EnumAddFlags(Capabilities, UnrealCap);
@@ -996,14 +1002,14 @@ void FVulkanDevice::MapImageFormatSupport(FPixelFormatInfo& PixelFormatInfo, con
 		const VkFormatProperties& LocalFormatProperties = GetFormatProperties(VulkanFormat);
 
 		// Check for individual texture caps
-		ConvertImageCap(LocalFormatProperties, EPixelFormatCapabilities::AnyTexture | EPixelFormatCapabilities::TextureSample | EPixelFormatCapabilities::TextureLoad, VK_FORMAT_FEATURE_SAMPLED_IMAGE_BIT);
-		ConvertImageCap(LocalFormatProperties, EPixelFormatCapabilities::DepthStencil, VK_FORMAT_FEATURE_DEPTH_STENCIL_ATTACHMENT_BIT);
-		ConvertImageCap(LocalFormatProperties, EPixelFormatCapabilities::RenderTarget, VK_FORMAT_FEATURE_COLOR_ATTACHMENT_BIT);
-		ConvertImageCap(LocalFormatProperties, EPixelFormatCapabilities::TextureBlendable, VK_FORMAT_FEATURE_COLOR_ATTACHMENT_BLEND_BIT);
-		ConvertImageCap(LocalFormatProperties, EPixelFormatCapabilities::AllUAVFlags | EPixelFormatCapabilities::TextureStore, VK_FORMAT_FEATURE_STORAGE_IMAGE_BIT);
-
-		ConvertImageCap(LocalFormatProperties, EPixelFormatCapabilities::TextureAtomics, VK_FORMAT_FEATURE_STORAGE_IMAGE_ATOMIC_BIT);
-		ConvertImageCap(LocalFormatProperties, EPixelFormatCapabilities::TextureFilterable, VK_FORMAT_FEATURE_SAMPLED_IMAGE_FILTER_LINEAR_BIT);
+		ConvertImageCap(LocalFormatProperties, EPixelFormatCapabilities::AnyTexture | EPixelFormatCapabilities::TextureSample | EPixelFormatCapabilities::TextureLoad, VK_FORMAT_FEATURE_SAMPLED_IMAGE_BIT, false);
+		ConvertImageCap(LocalFormatProperties, EPixelFormatCapabilities::DepthStencil, VK_FORMAT_FEATURE_DEPTH_STENCIL_ATTACHMENT_BIT, true);
+		ConvertImageCap(LocalFormatProperties, EPixelFormatCapabilities::RenderTarget, VK_FORMAT_FEATURE_COLOR_ATTACHMENT_BIT, true);
+		ConvertImageCap(LocalFormatProperties, EPixelFormatCapabilities::TextureBlendable, VK_FORMAT_FEATURE_COLOR_ATTACHMENT_BLEND_BIT, true);
+		ConvertImageCap(LocalFormatProperties, EPixelFormatCapabilities::AllUAVFlags | EPixelFormatCapabilities::TextureStore, VK_FORMAT_FEATURE_STORAGE_IMAGE_BIT, false);
+
+		ConvertImageCap(LocalFormatProperties, EPixelFormatCapabilities::TextureAtomics, VK_FORMAT_FEATURE_STORAGE_IMAGE_ATOMIC_BIT, false);
+		ConvertImageCap(LocalFormatProperties, EPixelFormatCapabilities::TextureFilterable, VK_FORMAT_FEATURE_SAMPLED_IMAGE_FILTER_LINEAR_BIT, false);
 
 		if (EnumHasAllFlags(Capabilities, EPixelFormatCapabilities::AnyTexture))
 		{
