Fix Vulkan allocator usage when creating images from Android hardware buffers

diff --git a/Engine/Source/Runtime/VulkanRHI/Private/VulkanTexture.cpp b/Engine/Source/Runtime/VulkanRHI/Private/VulkanTexture.cpp
--- a/Engine/Source/Runtime/VulkanRHI/Private/VulkanTexture.cpp
+++ b/Engine/Source/Runtime/VulkanRHI/Private/VulkanTexture.cpp
@@ -2030,7 +2030,7 @@ FVulkanTexture::FVulkanTexture(FVulkanDevice& InDevice, const FRHITextureCreateD
 	ImageCreateInfo.pNext = &ExternalMemoryImageCreateInfo;
 
 	VkImage VulkanImage;
-	VERIFYVULKANRESULT(VulkanRHI::vkCreateImage(VulkanDevice, &ImageCreateInfo, nullptr, &VulkanImage));
+	VERIFYVULKANRESULT(VulkanRHI::vkCreateImage(VulkanDevice, &ImageCreateInfo, AllocationCallbacks, &VulkanImage));
 
 	VkMemoryDedicatedAllocateInfo MemoryDedicatedAllocateInfo;
 	ZeroVulkanStruct(MemoryDedicatedAllocateInfo, VK_STRUCTURE_TYPE_MEMORY_DEDICATED_ALLOCATE_INFO);
