VulkanRHI: fallback host+local memory to regular host memory when it's not available.

#rnx

diff --git a/Engine/Source/Runtime/VulkanRHI/Private/VulkanMemory.cpp b/Engine/Source/Runtime/VulkanRHI/Private/VulkanMemory.cpp
--- a/Engine/Source/Runtime/VulkanRHI/Private/VulkanMemory.cpp
+++ b/Engine/Source/Runtime/VulkanRHI/Private/VulkanMemory.cpp
@@ -2987,6 +2987,12 @@ namespace VulkanRHI
 				MemoryPropertyFlags = MemoryPropertyFlags & ~VK_MEMORY_PROPERTY_LAZILY_ALLOCATED_BIT;
 			}
 
+			if (VKHasAllFlags(MemoryPropertyFlags, VK_MEMORY_PROPERTY_HOST_VISIBLE_BIT | VK_MEMORY_PROPERTY_HOST_COHERENT_BIT | VK_MEMORY_PROPERTY_DEVICE_LOCAL_BIT))
+			{
+				// Try regular host memory if local+host is not available
+				MemoryPropertyFlags = MemoryPropertyFlags & ~VK_MEMORY_PROPERTY_DEVICE_LOCAL_BIT;
+			}
+
 			// Try another heap type
 			uint32 OriginalTypeIndex = TypeIndex;
 			if (DeviceMemoryManager->GetMemoryTypeFromPropertiesExcluding(MemoryReqs.memoryTypeBits, MemoryPropertyFlags, (Result == VK_SUCCESS) ? TypeIndex : (uint32)-1, &TypeIndex) != VK_SUCCESS)
