Fix Vulkan UAV aliasing

diff --git a/Engine/Source/Runtime/VulkanRHI/Private/VulkanRHI.cpp b/Engine/Source/Runtime/VulkanRHI/Private/VulkanRHI.cpp
--- a/Engine/Source/Runtime/VulkanRHI/Private/VulkanRHI.cpp
+++ b/Engine/Source/Runtime/VulkanRHI/Private/VulkanRHI.cpp
@@ -987,8 +987,7 @@ void FVulkanDynamicRHI::InitInstance()
 		GRHISupportsParallelRHIExecute = GRHIThreadCvar->GetInt() > 1;
 #endif
 
-		// Disabled since it causes issues in Lumen tests
-		//GRHISupportsUAVFormatAliasing = true;
+		GRHISupportsUAVFormatAliasing = true;
 
 		// Some platforms might only have CPU for an RHI thread, but not for parallel tasks
 		GSupportsParallelRenderingTasksWithSeparateRHIThread = GRHISupportsRHIThread ? FVulkanPlatform::SupportParallelRenderingTasks() : false;
diff --git a/Engine/Source/Runtime/VulkanRHI/Private/VulkanTexture.cpp b/Engine/Source/Runtime/VulkanRHI/Private/VulkanTexture.cpp
--- a/Engine/Source/Runtime/VulkanRHI/Private/VulkanTexture.cpp
+++ b/Engine/Source/Runtime/VulkanRHI/Private/VulkanTexture.cpp
@@ -318,7 +318,7 @@ void FVulkanTexture::GenerateImageCreateInfo(
 			ImageCreateInfo.pNext = &ImageFormatListCreateInfo;
 
 			// Allow non-SRGB views to be created for SRGB textures
-			if (EnumHasAllFlags(UEFlags, TexCreate_SRGB))
+			if (EnumHasAllFlags(UEFlags, TexCreate_SRGB) && nonSrgbFormat != srgbFormat)
 			{
 				OutImageCreateInfo.FormatsUsed.Add(nonSrgbFormat);
 				OutImageCreateInfo.FormatsUsed.Add(srgbFormat);
@@ -327,8 +327,8 @@ void FVulkanTexture::GenerateImageCreateInfo(
 			// Make it possible to create R32G32 views of R64 images for utilities like clears
 			if (InDesc.Format == PF_R64_UINT)
 			{
-				OutImageCreateInfo.FormatsUsed.Add(nonSrgbFormat);
-				OutImageCreateInfo.FormatsUsed.Add(UEToVkTextureFormat(PF_R32G32_UINT, false));
+				OutImageCreateInfo.FormatsUsed.AddUnique(nonSrgbFormat);
+				OutImageCreateInfo.FormatsUsed.AddUnique(UEToVkTextureFormat(PF_R32G32_UINT, false));
 			}
 
 			if (bHasUAVFormat)
@@ -442,7 +442,14 @@ void FVulkanTexture::GenerateImageCreateInfo(
 	if (!VKHasAnyFlags(FormatFlags, VK_FORMAT_FEATURE_STORAGE_IMAGE_BIT))
 	{
 		ensure((ImageCreateInfo.usage & VK_IMAGE_USAGE_STORAGE_BIT) == 0 || bHasUAVFormat);
-		ImageCreateInfo.usage &= ~VK_IMAGE_USAGE_STORAGE_BIT;
+		if (bHasUAVFormat)
+		{
+			ImageCreateInfo.flags |= VK_IMAGE_CREATE_EXTENDED_USAGE_BIT;
+		}
+		else
+		{
+			ImageCreateInfo.usage &= ~VK_IMAGE_USAGE_STORAGE_BIT;
+		}
 	}
 
 	if (!VKHasAnyFlags(FormatFlags, VK_FORMAT_FEATURE_COLOR_ATTACHMENT_BIT))
