VulkanRHI: Read and track bindless usage flags on shader load.  Determine bindless state of a pipeline based on the bindless states of the shaders that are used.  Fix Vulkan SM6 while previewing shader platforms without bindless support.

#rnx

diff --git a/Engine/Source/Runtime/VulkanRHI/Private/VulkanCommands.cpp b/Engine/Source/Runtime/VulkanRHI/Private/VulkanCommands.cpp
--- a/Engine/Source/Runtime/VulkanRHI/Private/VulkanCommands.cpp
+++ b/Engine/Source/Runtime/VulkanRHI/Private/VulkanCommands.cpp
@@ -148,12 +148,12 @@ struct FVulkanResourceBinder
 	// Even if no resource are bound when bindless is enabled, we still need to process the Resource Table to go through proper validation
 	const bool bBindless;
 
-	FVulkanResourceBinder(FVulkanCommandListContext& InContext, EShaderFrequency InFrequency, PendingStateType* InPendingState)
+	FVulkanResourceBinder(FVulkanCommandListContext& InContext, EShaderFrequency InFrequency, PendingStateType* InPendingState, bool InUsesBindless)
 		: Context(InContext)
 		, Frequency(InFrequency)
 		, Stage((InFrequency == SF_Compute) ? ShaderStage::Compute : ShaderStage::GetStageForFrequency(InFrequency))
 		, PendingState(InPendingState)
-		, bBindless(InContext.GetDevice()->SupportsBindless())
+		, bBindless(InUsesBindless)
 	{
 	}
 
@@ -215,7 +215,7 @@ void FVulkanCommandListContext::SetResourcesFromTables(const ShaderType* Shader)
 
 	if (Frequency == SF_Compute)
 	{
-		FVulkanResourceBinder Binder(*this, Frequency, PendingComputeState);
+		FVulkanResourceBinder Binder(*this, Frequency, PendingComputeState, Shader->UsesBindless());
 		UE::RHI::Private::SetUniformBufferResourcesFromTables(
 			Binder
 			, *Shader
@@ -228,7 +228,7 @@ void FVulkanCommandListContext::SetResourcesFromTables(const ShaderType* Shader)
 	}
 	else
 	{
-		FVulkanResourceBinder Binder(*this, Frequency, PendingGfxState);
+		FVulkanResourceBinder Binder(*this, Frequency, PendingGfxState, Shader->UsesBindless());
 		UE::RHI::Private::SetUniformBufferResourcesFromTables(
 			Binder
 			, *Shader
diff --git a/Engine/Source/Runtime/VulkanRHI/Private/VulkanDescriptorSets.h b/Engine/Source/Runtime/VulkanRHI/Private/VulkanDescriptorSets.h
--- a/Engine/Source/Runtime/VulkanRHI/Private/VulkanDescriptorSets.h
+++ b/Engine/Source/Runtime/VulkanRHI/Private/VulkanDescriptorSets.h
@@ -114,7 +114,7 @@ public:
 	void ProcessBindingsForStage(VkShaderStageFlagBits StageFlags, ShaderStage::EStage DescSetStage, const FVulkanShaderHeader& CodeHeader, FUniformBufferGatherInfo& OutUBGatherInfo) const;
 
 	template<bool bIsCompute>
-	void FinalizeBindings(const FVulkanDevice& Device, const FUniformBufferGatherInfo& UBGatherInfo, const TArrayView<FRHISamplerState*>& ImmutableSamplers);
+	void FinalizeBindings(const FVulkanDevice& Device, const FUniformBufferGatherInfo& UBGatherInfo, const TArrayView<FRHISamplerState*>& ImmutableSamplers, bool bUsesBindless);
 
 	void GenerateHash(const TArrayView<FRHISamplerState*>& ImmutableSamplers, VkPipelineBindPoint InBindPoint);
 
@@ -928,7 +928,7 @@ private:
 class FVulkanLayout : public VulkanRHI::FDeviceChild
 {
 public:
-	FVulkanLayout(FVulkanDevice* InDevice, bool InGfxLayout);
+	FVulkanLayout(FVulkanDevice* InDevice, bool InGfxLayout, bool InUsesBindless);
 	virtual ~FVulkanLayout();
 
 	bool IsGfxLayout() const
@@ -943,7 +943,7 @@ public:
 
 	inline VkPipelineLayout GetPipelineLayout() const
 	{
-		return Device->SupportsBindless() ? Device->GetBindlessDescriptorManager()->GetPipelineLayout() : PipelineLayout;
+		return bUsesBindless ? Device->GetBindlessDescriptorManager()->GetPipelineLayout() : PipelineLayout;
 	}
 
 	inline bool HasDescriptors() const
@@ -958,6 +958,7 @@ public:
 
 protected:
 	const bool bIsGfxLayout;
+	const bool bUsesBindless;
 	FVulkanDescriptorSetsLayout	DescriptorSetLayout;
 	VkPipelineLayout			PipelineLayout;
 
diff --git a/Engine/Source/Runtime/VulkanRHI/Private/VulkanPendingState.cpp b/Engine/Source/Runtime/VulkanRHI/Private/VulkanPendingState.cpp
--- a/Engine/Source/Runtime/VulkanRHI/Private/VulkanPendingState.cpp
+++ b/Engine/Source/Runtime/VulkanRHI/Private/VulkanPendingState.cpp
@@ -349,7 +349,7 @@ void FVulkanPendingComputeState::PrepareForDispatch(FVulkanCommandBuffer& InComm
 
 	VkCommandBuffer CommandBufferHandle = InCommandBuffer.GetHandle();
 
-	if (Device->SupportsBindless())
+	if (CurrentPipeline->UsesBindless())
 	{
 		CurrentState->UpdateBindlessDescriptors(&Context, InCommandBuffer);
 		CurrentPipeline->Bind(CommandBufferHandle);
@@ -389,7 +389,7 @@ void FVulkanPendingGfxState::PrepareForDraw(FVulkanCommandBuffer& CommandBuffer)
 
 	check(CommandBuffer.bHasPipeline);
 
-	if (Device->SupportsBindless())
+	if (CurrentPipeline->UsesBindless())
 	{
 		check(!CurrentPipeline->bHasInputAttachments); // todo-jn: bindless + InputAttachments
 		UpdateDynamicStates(CommandBuffer);
diff --git a/Engine/Source/Runtime/VulkanRHI/Private/VulkanPipeline.cpp b/Engine/Source/Runtime/VulkanRHI/Private/VulkanPipeline.cpp
--- a/Engine/Source/Runtime/VulkanRHI/Private/VulkanPipeline.cpp
+++ b/Engine/Source/Runtime/VulkanRHI/Private/VulkanPipeline.cpp
@@ -224,6 +224,7 @@ FVulkanPipeline::~FVulkanPipeline()
 FVulkanComputePipeline::FVulkanComputePipeline(FVulkanDevice* InDevice, FVulkanComputeShader* InShader)
 	: FVulkanPipeline(InDevice)
 	, FRHIComputePipelineState(InShader)
+	, bUsesBindless(InShader->UsesBindless())
 {
 	INC_DWORD_STAT(STAT_VulkanNumComputePSOs);
 }
@@ -1380,7 +1381,7 @@ bool FVulkanPipelineStateCacheManager::CreateGfxPipelineFromEntry(FVulkanRHIGrap
 		PipelineInfo.pNext = (void*)&PipelineFragmentShadingRate;
 	}
 
-	if (Device->SupportsBindless())
+	if (PSO->UsesBindless())
 	{
 		PipelineInfo.flags |= VK_PIPELINE_CREATE_DESCRIPTOR_BUFFER_BIT_EXT;
 	}
@@ -1720,7 +1721,7 @@ FVulkanShaderHashes::FVulkanShaderHashes()
 	Hash = 0;
 }
 
-FVulkanLayout* FVulkanPipelineStateCacheManager::FindOrAddLayout(const FVulkanDescriptorSetsLayoutInfo& DescriptorSetLayoutInfo, bool bGfxLayout)
+FVulkanLayout* FVulkanPipelineStateCacheManager::FindOrAddLayout(const FVulkanDescriptorSetsLayoutInfo& DescriptorSetLayoutInfo, bool bGfxLayout, bool bUsesBindless)
 {
 	FScopeLock Lock(&LayoutMapCS);
 	if (FVulkanLayout** FoundLayout = LayoutMap.Find(DescriptorSetLayoutInfo))
@@ -1729,7 +1730,7 @@ FVulkanLayout* FVulkanPipelineStateCacheManager::FindOrAddLayout(const FVulkanDe
 		return *FoundLayout;
 	}
 
-	FVulkanLayout* Layout = new FVulkanLayout(Device, bGfxLayout);
+	FVulkanLayout* Layout = new FVulkanLayout(Device, bGfxLayout, bUsesBindless);
 	Layout->DescriptorSetLayout.CopyFrom(DescriptorSetLayoutInfo);
 	Layout->Compile(DSetLayoutMap);
 
@@ -1771,54 +1772,54 @@ void FVulkanPipelineStateCacheManager::CreateGfxEntry(const FGraphicsPipelineSta
 		const FBoundShaderStateInput& BSI = PSOInitializer.BoundShaderState;
 
 		FUniformBufferGatherInfo UBGatherInfo;
+		uint32 NumActiveShaders = 0;
+		uint32 NumBindlessShaders = 0;
+
+		auto ProcessShaderStage = [&DescriptorSetLayoutInfo, &UBGatherInfo, &NumActiveShaders, &NumBindlessShaders](VkShaderStageFlagBits StageFlag, ShaderStage::EStage Stage, FVulkanShader* Shader)
+		{
+			if (Shader)
+			{
+				const FVulkanShaderHeader& Header = Shader->GetCodeHeader();
+				DescriptorSetLayoutInfo.ProcessBindingsForStage(StageFlag, Stage, Header, UBGatherInfo);
+				NumActiveShaders++;
+				if (Shader->UsesBindless())
+				{
+					NumBindlessShaders++;
+				}
+			}
+		};
 
 		if (Shaders[ShaderStage::Vertex])
 		{
 			const FVulkanShaderHeader& VSHeader = Shaders[ShaderStage::Vertex]->GetCodeHeader();
 			VertexInputState.Generate(ResourceCast(PSOInitializer.BoundShaderState.VertexDeclarationRHI), VSHeader.InOutMask);
-
-			DescriptorSetLayoutInfo.ProcessBindingsForStage(VK_SHADER_STAGE_VERTEX_BIT, ShaderStage::Vertex, VSHeader, UBGatherInfo);
 		}
 
-		if (Shaders[ShaderStage::Pixel])
+		if (Shaders[ShaderStage::Pixel] && Shaders[ShaderStage::Pixel]->GetCodeHeader().InputAttachmentInfos.Num())
 		{
-			const FVulkanShaderHeader& PSHeader = Shaders[ShaderStage::Pixel]->GetCodeHeader();
-			DescriptorSetLayoutInfo.ProcessBindingsForStage(VK_SHADER_STAGE_FRAGMENT_BIT, ShaderStage::Pixel, PSHeader, UBGatherInfo);
-
-			if (PSHeader.InputAttachmentInfos.Num())
-			{
-				// input attachements can't exist in a first sub-pass
-				check(PSOInitializer.SubpassHint != ESubpassHint::None);
-				check(PSOInitializer.SubpassIndex != 0);
-			}
+			// input attachements can't exist in a first sub-pass
+			check(PSOInitializer.SubpassHint != ESubpassHint::None);
+			check(PSOInitializer.SubpassIndex != 0);
 		}
 
-#if PLATFORM_SUPPORTS_MESH_SHADERS
-		if (Shaders[ShaderStage::Mesh])
-		{
-			const FVulkanShaderHeader& MSHeader = Shaders[ShaderStage::Mesh]->GetCodeHeader();
-			DescriptorSetLayoutInfo.ProcessBindingsForStage(VK_SHADER_STAGE_MESH_BIT_EXT, ShaderStage::Mesh, MSHeader, UBGatherInfo);
-		}
+		ProcessShaderStage(VK_SHADER_STAGE_VERTEX_BIT, ShaderStage::Vertex, Shaders[ShaderStage::Vertex]);
+		ProcessShaderStage(VK_SHADER_STAGE_FRAGMENT_BIT, ShaderStage::Pixel, Shaders[ShaderStage::Pixel]);
 
-		if (Shaders[ShaderStage::Task])
-		{
-			const FVulkanShaderHeader& TSHeader = Shaders[ShaderStage::Task]->GetCodeHeader();
-			DescriptorSetLayoutInfo.ProcessBindingsForStage(VK_SHADER_STAGE_TASK_BIT_EXT, ShaderStage::Task, TSHeader, UBGatherInfo);
-		}
+#if PLATFORM_SUPPORTS_MESH_SHADERS
+		ProcessShaderStage(VK_SHADER_STAGE_MESH_BIT_EXT, ShaderStage::Mesh, Shaders[ShaderStage::Mesh]);
+		ProcessShaderStage(VK_SHADER_STAGE_TASK_BIT_EXT, ShaderStage::Task, Shaders[ShaderStage::Task]);
 #endif
 
 #if VULKAN_SUPPORTS_GEOMETRY_SHADERS
-		if (Shaders[ShaderStage::Geometry])
-		{
-			const FVulkanShaderHeader& GSHeader = Shaders[ShaderStage::Geometry]->GetCodeHeader();
-			DescriptorSetLayoutInfo.ProcessBindingsForStage(VK_SHADER_STAGE_GEOMETRY_BIT, ShaderStage::Geometry, GSHeader, UBGatherInfo);
-		}
+		ProcessShaderStage(VK_SHADER_STAGE_GEOMETRY_BIT, ShaderStage::Geometry, Shaders[ShaderStage::Geometry]);
 #endif
 
+		checkf((NumBindlessShaders == 0) || (NumBindlessShaders == NumActiveShaders), TEXT("All shaders must be bindless or non-bindless."));
+
 		// Second pass
 		const int32 NumImmutableSamplers = PSOInitializer.ImmutableSamplerState.ImmutableSamplers.Num();
 		TArrayView<FRHISamplerState*> ImmutableSamplers(NumImmutableSamplers > 0 ? &(FRHISamplerState*&)PSOInitializer.ImmutableSamplerState.ImmutableSamplers[0] : nullptr, NumImmutableSamplers);
-		DescriptorSetLayoutInfo.FinalizeBindings<false>(*Device, UBGatherInfo, ImmutableSamplers);
+		DescriptorSetLayoutInfo.FinalizeBindings<false>(*Device, UBGatherInfo, ImmutableSamplers, (NumBindlessShaders != 0));
 	}
 
 	OutGfxEntry->SubpassIndex = PSOInitializer.SubpassIndex;
@@ -1976,13 +1977,23 @@ FVulkanRHIGraphicsPipelineState::FVulkanRHIGraphicsPipelineState(FVulkanDevice*
 #endif
 	VulkanShaders[ShaderStage::Pixel] = static_cast<FVulkanPixelShader*>(PSOInitializer_.BoundShaderState.PixelShaderRHI);
 
-	for (int ShaderStageIndex = 0; ShaderStageIndex < ShaderStage::NumGraphicsStages; ShaderStageIndex++)
+	uint32 ActiveShaderCount = 0;
+	uint32 BindlessShaderCount = 0;
+	for (int32 ShaderStageIndex = 0; ShaderStageIndex < ShaderStage::NumGraphicsStages; ShaderStageIndex++)
 	{
 		if (VulkanShaders[ShaderStageIndex] != nullptr)
 		{
 			VulkanShaders[ShaderStageIndex]->AddRef();
+
+			ActiveShaderCount++;
+			if (VulkanShaders[ShaderStageIndex]->UsesBindless())
+			{
+				BindlessShaderCount++;
+			}
 		}
 	}
+	checkf((BindlessShaderCount == 0) || (ActiveShaderCount == BindlessShaderCount), TEXT("Pipelines can't be created with mix of bindless and non-bindless shaders."));
+	bUsesBindless = (BindlessShaderCount != 0);
 
 #if VULKAN_PSO_CACHE_DEBUG
 	PixelShaderRHI = PSOInitializer_.BoundShaderState.PixelShaderRHI;
@@ -2119,7 +2130,7 @@ FGraphicsPipelineStateRHIRef FVulkanPipelineStateCacheManager::RHICreateGraphics
 		SCOPE_CYCLE_COUNTER(STAT_VulkanPSOCreationTime);
 		NewPSO = new FVulkanRHIGraphicsPipelineState(Device, Initializer, Desc, &Key);
 		{
-			FVulkanLayout* Layout = FindOrAddLayout(DescriptorSetLayoutInfo, true);
+			FVulkanLayout* Layout = FindOrAddLayout(DescriptorSetLayoutInfo, true, NewPSO->UsesBindless());
 			NewPSO->Layout = Layout;
 			NewPSO->bHasInputAttachments = Layout->GetDescriptorSetsLayout().HasInputAttachments();
 		}
@@ -2282,8 +2293,8 @@ FVulkanComputePipeline* FVulkanPipelineStateCacheManager::CreateComputePipelineF
 	const FVulkanShaderHeader& CSHeader = Shader->GetCodeHeader();
 	FUniformBufferGatherInfo UBGatherInfo;
 	DescriptorSetLayoutInfo.ProcessBindingsForStage(VK_SHADER_STAGE_COMPUTE_BIT, ShaderStage::Compute, CSHeader, UBGatherInfo);
-	DescriptorSetLayoutInfo.FinalizeBindings<true>(*Device, UBGatherInfo, TArrayView<FRHISamplerState*>());
-	FVulkanLayout* Layout = FindOrAddLayout(DescriptorSetLayoutInfo, false);
+	DescriptorSetLayoutInfo.FinalizeBindings<true>(*Device, UBGatherInfo, TArrayView<FRHISamplerState*>(), Shader->UsesBindless());
+	FVulkanLayout* Layout = FindOrAddLayout(DescriptorSetLayoutInfo, false, Shader->UsesBindless());
 	checkSlow(!Layout->IsGfxLayout());
 
 	TRefCountPtr<FVulkanShaderModule> ShaderModule = Shader->GetOrCreateHandle(Layout, Layout->GetDescriptorSetLayoutHash());
@@ -2299,7 +2310,7 @@ FVulkanComputePipeline* FVulkanPipelineStateCacheManager::CreateComputePipelineF
 	PipelineInfo.stage.pName = EntryPoint;
 	PipelineInfo.layout = Layout->GetPipelineLayout();
 
-	if (Device->SupportsBindless())
+	if (Shader->UsesBindless())
 	{
 		PipelineInfo.flags |= VK_PIPELINE_CREATE_DESCRIPTOR_BUFFER_BIT_EXT;
 	}
diff --git a/Engine/Source/Runtime/VulkanRHI/Private/VulkanPipeline.h b/Engine/Source/Runtime/VulkanRHI/Private/VulkanPipeline.h
--- a/Engine/Source/Runtime/VulkanRHI/Private/VulkanPipeline.h
+++ b/Engine/Source/Runtime/VulkanRHI/Private/VulkanPipeline.h
@@ -522,7 +522,7 @@ private:
 
 	VkResult CreateVKPipeline(FVulkanRHIGraphicsPipelineState* PSO, FVulkanShader* Shaders[ShaderStage::NumGraphicsStages], const VkGraphicsPipelineCreateInfo& PipelineInfo, FGraphicsPipelineStateInitializer::EPSOPrecacheCompileType PSOCompileType);
 
-	FVulkanLayout* FindOrAddLayout(const FVulkanDescriptorSetsLayoutInfo& DescriptorSetLayoutInfo, bool bGfxLayout);
+	FVulkanLayout* FindOrAddLayout(const FVulkanDescriptorSetsLayoutInfo& DescriptorSetLayoutInfo, bool bGfxLayout, bool bUsesBindless);
 	FVulkanComputePipeline* CreateComputePipelineFromShader(FVulkanComputeShader* Shader);
 
 	/** LRU Related functions */
@@ -673,20 +673,28 @@ public:
 	FVulkanComputePipeline(FVulkanDevice* InDevice, FVulkanComputeShader* InShader);
 	virtual ~FVulkanComputePipeline();
 
-	inline const FVulkanShaderHeader& GetShaderCodeHeader() const
+	const FVulkanShaderHeader& GetShaderCodeHeader() const
 	{
 		return ResourceCast(GetComputeShader())->GetCodeHeader();
 	}
 
-	inline const FVulkanLayout& GetLayout() const
+	const FVulkanLayout& GetLayout() const
 	{
 		return *Layout;
 	}
 
-	inline void Bind(VkCommandBuffer CmdBuffer)
+	void Bind(VkCommandBuffer CmdBuffer)
 	{
 		VulkanRHI::vkCmdBindPipeline(CmdBuffer, VK_PIPELINE_BIND_POINT_COMPUTE, Pipeline);
 	}
+
+	bool UsesBindless() const
+	{
+		return bUsesBindless;
+	}
+
+protected:
+	const bool bUsesBindless;
 };
 
 class FVulkanRHIGraphicsPipelineState : public FRHIGraphicsPipelineState
@@ -749,9 +757,15 @@ public:
 	FVulkanShader::FSpirvCode GetPatchedSpirvCode(FVulkanShader* Shader);
 	void PurgeShaderModules(FVulkanShader*const* Shaders);
 
+	bool UsesBindless() const
+	{
+		return bUsesBindless;
+	}
+
 
 	bool								bHasInputAttachments = false;
 	bool								bIsRegistered;
+	bool								bUsesBindless = false;
 
 	uint64								ShaderKeys[ShaderStage::NumGraphicsStages];
 	TEnumAsByte<EPrimitiveType>			PrimitiveType;
diff --git a/Engine/Source/Runtime/VulkanRHI/Private/VulkanPipelineState.cpp b/Engine/Source/Runtime/VulkanRHI/Private/VulkanPipelineState.cpp
--- a/Engine/Source/Runtime/VulkanRHI/Private/VulkanPipelineState.cpp
+++ b/Engine/Source/Runtime/VulkanRHI/Private/VulkanPipelineState.cpp
@@ -38,7 +38,7 @@ static bool ShouldAlwaysWriteDescriptors()
 }
 
 FVulkanComputePipelineDescriptorState::FVulkanComputePipelineDescriptorState(FVulkanDevice* InDevice, FVulkanComputePipeline* InComputePipeline)
-	: FVulkanCommonPipelineDescriptorState(InDevice, ShaderStage::NumComputeStages)
+	: FVulkanCommonPipelineDescriptorState(InDevice, ShaderStage::NumComputeStages, InComputePipeline->UsesBindless())
 	, PackedUniformBuffersMask(0)
 	, PackedUniformBuffersDirty(0)
 	, ComputePipeline(InComputePipeline)
@@ -283,7 +283,7 @@ void FVulkanComputePipelineDescriptorState::UpdateBindlessDescriptors(FVulkanCom
 }
 
 FVulkanGraphicsPipelineDescriptorState::FVulkanGraphicsPipelineDescriptorState(FVulkanDevice* InDevice, FVulkanRHIGraphicsPipelineState* InGfxPipeline)
-	: FVulkanCommonPipelineDescriptorState(InDevice, ShaderStage::NumGraphicsStages)
+	: FVulkanCommonPipelineDescriptorState(InDevice, ShaderStage::NumGraphicsStages, InGfxPipeline->UsesBindless())
 	, GfxPipeline(InGfxPipeline)
 {
 	LLM_SCOPE_VULKAN(ELLMTagVulkan::VulkanShaders);
diff --git a/Engine/Source/Runtime/VulkanRHI/Private/VulkanPipelineState.h b/Engine/Source/Runtime/VulkanRHI/Private/VulkanPipelineState.h
--- a/Engine/Source/Runtime/VulkanRHI/Private/VulkanPipelineState.h
+++ b/Engine/Source/Runtime/VulkanRHI/Private/VulkanPipelineState.h
@@ -22,10 +22,10 @@ extern TAutoConsoleVariable<int32> GDynamicGlobalUBs;
 class FVulkanCommonPipelineDescriptorState : public VulkanRHI::FDeviceChild
 {
 public:
-	FVulkanCommonPipelineDescriptorState(FVulkanDevice* InDevice, uint32 InMaxNumSets)
+	FVulkanCommonPipelineDescriptorState(FVulkanDevice* InDevice, uint32 InMaxNumSets, bool InUsesBindless)
 		: VulkanRHI::FDeviceChild(InDevice)
 		, MaxNumSets(InMaxNumSets)
-		, bUseBindless(InDevice->SupportsBindless())
+		, bUseBindless(InUsesBindless)
 	{
 	}
 
diff --git a/Engine/Source/Runtime/VulkanRHI/Private/VulkanResources.h b/Engine/Source/Runtime/VulkanRHI/Private/VulkanResources.h
--- a/Engine/Source/Runtime/VulkanRHI/Private/VulkanResources.h
+++ b/Engine/Source/Runtime/VulkanRHI/Private/VulkanResources.h
@@ -187,6 +187,16 @@ public:
 
 	TArray<FUniformBufferStaticSlot>& StaticSlots;
 
+	void SetUsesBindless(bool bValue)
+	{
+		bUsesBindless = bValue;
+	}
+
+	bool UsesBindless() const
+	{
+		return bUsesBindless;
+	}
+
 protected:
 
 #if UE_BUILD_DEBUG || UE_BUILD_DEVELOPMENT
@@ -198,6 +208,7 @@ protected:
 	FVulkanShaderHeader				CodeHeader;
 	TMap<uint32, TRefCountPtr<FVulkanShaderModule>>	ShaderModules;
 	const EShaderFrequency			Frequency;
+	bool							bUsesBindless = false;
 
 protected:
 	class FSpirvContainer
diff --git a/Engine/Source/Runtime/VulkanRHI/Private/VulkanShaders.cpp b/Engine/Source/Runtime/VulkanRHI/Private/VulkanShaders.cpp
--- a/Engine/Source/Runtime/VulkanRHI/Private/VulkanShaders.cpp
+++ b/Engine/Source/Runtime/VulkanRHI/Private/VulkanShaders.cpp
@@ -40,8 +40,8 @@ FVulkanShaderFactory::~FVulkanShaderFactory()
 	}
 }
 
-
-static void ReadShaderOptionalData(FShaderCodeReader& ShaderCode, FRHIShader* RHIShader)
+template <typename ShaderType>
+static void ReadShaderOptionalData(FShaderCodeReader& ShaderCode, ShaderType* RHIShader)
 {
 	const FShaderCodePackedResourceCounts* PackedResourceCounts = ShaderCode.FindOptionalData<FShaderCodePackedResourceCounts>();
 	if (PackedResourceCounts)
@@ -51,6 +51,7 @@ static void ReadShaderOptionalData(FShaderCodeReader& ShaderCode, FRHIShader* RH
 			RHIShader->SetNoDerivativeOps(EnumHasAnyFlags(PackedResourceCounts->UsageFlags, EShaderResourceUsageFlags::NoDerivativeOps));
 		}
 		RHIShader->SetShaderBundleUsage(EnumHasAnyFlags(PackedResourceCounts->UsageFlags, EShaderResourceUsageFlags::ShaderBundle));
+		RHIShader->SetUsesBindless(EnumHasAnyFlags(PackedResourceCounts->UsageFlags, EShaderResourceUsageFlags::BindlessSamplers | EShaderResourceUsageFlags::BindlessResources));
 	}
 
 #if RHI_INCLUDE_SHADER_DEBUG_DATA
@@ -620,9 +621,10 @@ FRayTracingShaderRHIRef FVulkanDynamicRHI::RHICreateRayTracingShader(TArrayView<
 	}
 }
 
-FVulkanLayout::FVulkanLayout(FVulkanDevice* InDevice, bool InGfxLayout)
+FVulkanLayout::FVulkanLayout(FVulkanDevice* InDevice, bool InGfxLayout, bool InUsesBindless)
 	: VulkanRHI::FDeviceChild(InDevice)
 	, bIsGfxLayout(InGfxLayout)
+	, bUsesBindless(InUsesBindless)
 	, DescriptorSetLayout(Device)
 	, PipelineLayout(VK_NULL_HANDLE)
 {
@@ -643,7 +645,7 @@ void FVulkanLayout::Compile(FVulkanDescriptorSetLayoutMap& DSetLayoutMap)
 
 	DescriptorSetLayout.Compile(DSetLayoutMap);
 
-	if (!Device->SupportsBindless())
+	if (!bUsesBindless)
 	{
 		VkPipelineLayoutCreateInfo PipelineLayoutCreateInfo;
 		ZeroVulkanStruct(PipelineLayoutCreateInfo, VK_STRUCTURE_TYPE_PIPELINE_LAYOUT_CREATE_INFO);
@@ -735,15 +737,15 @@ void FVulkanDescriptorSetsLayoutInfo::ProcessBindingsForStage(VkShaderStageFlagB
 }
 
 template<bool bIsCompute>
-void FVulkanDescriptorSetsLayoutInfo::FinalizeBindings(const FVulkanDevice& Device, const FUniformBufferGatherInfo& UBGatherInfo, const TArrayView<FRHISamplerState*>& ImmutableSamplers)
+void FVulkanDescriptorSetsLayoutInfo::FinalizeBindings(const FVulkanDevice& Device, const FUniformBufferGatherInfo& UBGatherInfo, const TArrayView<FRHISamplerState*>& ImmutableSamplers, bool bUsesBindless)
 {
 	// We'll be reusing this struct
 	VkDescriptorSetLayoutBinding Binding;
 	FMemory::Memzero(Binding);
 	Binding.descriptorCount = 1;
 
-	const bool bConvertAllUBsToDynamic = !Device.SupportsBindless() && (GDynamicGlobalUBs.GetValueOnAnyThread() > 1);
-	const bool bConvertPackedUBsToDynamic = !Device.SupportsBindless() && (bConvertAllUBsToDynamic || (GDynamicGlobalUBs.GetValueOnAnyThread() == 1));
+	const bool bConvertAllUBsToDynamic = !bUsesBindless && (GDynamicGlobalUBs.GetValueOnAnyThread() > 1);
+	const bool bConvertPackedUBsToDynamic = !bUsesBindless && (bConvertAllUBsToDynamic || (GDynamicGlobalUBs.GetValueOnAnyThread() == 1));
 	const uint32 MaxDescriptorSetUniformBuffersDynamic = Device.GetLimits().maxDescriptorSetUniformBuffersDynamic;
 
 	int32 CurrentImmutableSampler = 0;
@@ -850,5 +852,5 @@ FBoundShaderStateRHIRef FVulkanDynamicRHI::RHICreateBoundShaderState(
 }
 
 
-template void FVulkanDescriptorSetsLayoutInfo::FinalizeBindings<true>(const FVulkanDevice& Device, const FUniformBufferGatherInfo& UBGatherInfo, const TArrayView<FRHISamplerState*>& ImmutableSamplers);
-template void FVulkanDescriptorSetsLayoutInfo::FinalizeBindings<false>(const FVulkanDevice& Device, const FUniformBufferGatherInfo& UBGatherInfo, const TArrayView<FRHISamplerState*>& ImmutableSamplers);
+template void FVulkanDescriptorSetsLayoutInfo::FinalizeBindings<true>(const FVulkanDevice& Device, const FUniformBufferGatherInfo& UBGatherInfo, const TArrayView<FRHISamplerState*>& ImmutableSamplers, bool bUsesBindless);
+template void FVulkanDescriptorSetsLayoutInfo::FinalizeBindings<false>(const FVulkanDevice& Device, const FUniformBufferGatherInfo& UBGatherInfo, const TArrayView<FRHISamplerState*>& ImmutableSamplers, bool bUsesBindless);
