move back to a log so that we get the Vulkan crashes in CrashReporter

diff --git a/Engine/Source/Runtime/VulkanRHI/Private/VulkanUtil.cpp b/Engine/Source/Runtime/VulkanRHI/Private/VulkanUtil.cpp
--- a/Engine/Source/Runtime/VulkanRHI/Private/VulkanUtil.cpp
+++ b/Engine/Source/Runtime/VulkanRHI/Private/VulkanUtil.cpp
@@ -951,9 +951,11 @@ namespace VulkanRHI
 			GVulkanRHI->DumpMemory();
 		}
 #endif
+		UE_LOG(LogVulkanRHI, Fatal, TEXT("%s failed, VkResult=%d\n at %s:%u \n with error %s"),
+			ANSI_TO_TCHAR(VkFunction), (int32)Result, ANSI_TO_TCHAR(Filename), Line, *ErrorString);
 
 		// Force shutdown, we can't do anything useful anymore.
-		FPlatformMisc::RequestExit(true, TEXT("VerifyVulkanResult"));
+//		FPlatformMisc::RequestExit(true, TEXT("VerifyVulkanResult"));
 	}
 }
 
