VulkanRHI: Move CommandListManager code to a ContextCommon base class (part of submission refactor).  Add VulkanContext.cpp to start moving Context code together.  Start replacing pointers with references where null is not expected/supported.

#rnx

diff --git a/Engine/Source/Runtime/VulkanRHI/Private/VulkanBarriers.cpp b/Engine/Source/Runtime/VulkanRHI/Private/VulkanBarriers.cpp
--- a/Engine/Source/Runtime/VulkanRHI/Private/VulkanBarriers.cpp
+++ b/Engine/Source/Runtime/VulkanRHI/Private/VulkanBarriers.cpp
@@ -1152,10 +1152,10 @@ private:
 		switch (RequestedPipeline)
 		{
 		case ERHIPipeline::Graphics:
-			checkf(!Context.GetDevice()->IsRealAsyncComputeContext(&Context), TEXT("The %s pipeline for this transition is the Graphics queue, but it's submitted on the AsyncCompute queue."), IsBeginTransition ? TEXT("BEGIN") : TEXT("END"));
+			checkf(!Context.Device.IsRealAsyncComputeContext(&Context), TEXT("The %s pipeline for this transition is the Graphics queue, but it's submitted on the AsyncCompute queue."), IsBeginTransition ? TEXT("BEGIN") : TEXT("END"));
 			break;
 		case ERHIPipeline::AsyncCompute:
-			checkf(Context.GetDevice()->IsRealAsyncComputeContext(&Context), TEXT("The %s pipeline for this transition is the AsyncCompute queue, but it's submitted on the Graphics queue."), IsBeginTransition ? TEXT("BEGIN") : TEXT("END"));
+			checkf(Context.Device.IsRealAsyncComputeContext(&Context), TEXT("The %s pipeline for this transition is the AsyncCompute queue, but it's submitted on the Graphics queue."), IsBeginTransition ? TEXT("BEGIN") : TEXT("END"));
 			break;
 		default:
 			checkNoEntry();
@@ -1189,11 +1189,11 @@ void FTransitionProcessor<VkMemoryBarrier, VkBufferMemoryBarrier, VkImageMemoryB
 
 	if (SrcPipelines == ERHIPipeline::AsyncCompute)
 	{
-		MergedSrcStageMask &= Context.GetDevice()->GetComputeQueue()->GetSupportedStageBits();
+		MergedSrcStageMask &= Context.Device.GetComputeQueue()->GetSupportedStageBits();
 	}
 	if (DstPipelines == ERHIPipeline::AsyncCompute)
 	{
-		MergedDstStageMask &= Context.GetDevice()->GetComputeQueue()->GetSupportedStageBits();
+		MergedDstStageMask &= Context.Device.GetComputeQueue()->GetSupportedStageBits();
 	}
 
 	// Submit merged stage masks with arrays of barriers
@@ -1368,14 +1368,14 @@ void FTransitionProcessor<VkMemoryBarrier2, VkBufferMemoryBarrier2, VkImageMemor
 		for (int32 Index = BufferBatchStartIndex; Index < BufferBarriers.Num(); ++Index)
 		{
 			VkBufferMemoryBarrier2& BufferBarrier = BufferBarriers[Index];
-			BufferBarrier.srcStageMask &= Context.GetDevice()->GetComputeQueue()->GetSupportedStageBits();
+			BufferBarrier.srcStageMask &= Context.Device.GetComputeQueue()->GetSupportedStageBits();
 			BufferBarrier.srcAccessMask &= (VK_ACCESS_SHADER_READ_BIT | VK_ACCESS_SHADER_WRITE_BIT);
 		}
 
 		for (int32 Index = ImageBatchStartIndex; Index < ImageBarriers.Num(); ++Index)
 		{
 			VkImageMemoryBarrier2& Barrier = ImageBarriers[Index];
-			Barrier.srcStageMask &= Context.GetDevice()->GetComputeQueue()->GetSupportedStageBits();
+			Barrier.srcStageMask &= Context.Device.GetComputeQueue()->GetSupportedStageBits();
 			Barrier.srcAccessMask &= (VK_ACCESS_SHADER_READ_BIT | VK_ACCESS_SHADER_WRITE_BIT);
 		}
 	}
@@ -1385,14 +1385,14 @@ void FTransitionProcessor<VkMemoryBarrier2, VkBufferMemoryBarrier2, VkImageMemor
 		for (int32 Index = BufferBatchStartIndex; Index < BufferBarriers.Num(); ++Index)
 		{
 			VkBufferMemoryBarrier2& BufferBarrier = BufferBarriers[Index];
-			BufferBarrier.dstStageMask &= Context.GetDevice()->GetComputeQueue()->GetSupportedStageBits();
+			BufferBarrier.dstStageMask &= Context.Device.GetComputeQueue()->GetSupportedStageBits();
 			BufferBarrier.dstAccessMask &= (VK_ACCESS_SHADER_READ_BIT | VK_ACCESS_SHADER_WRITE_BIT);
 		}
 
 		for (int32 Index = ImageBatchStartIndex; Index < ImageBarriers.Num(); ++Index)
 		{
 			VkImageMemoryBarrier2& Barrier = ImageBarriers[Index];
-			Barrier.dstStageMask &= Context.GetDevice()->GetComputeQueue()->GetSupportedStageBits();
+			Barrier.dstStageMask &= Context.Device.GetComputeQueue()->GetSupportedStageBits();
 			Barrier.dstAccessMask &= (VK_ACCESS_SHADER_READ_BIT | VK_ACCESS_SHADER_WRITE_BIT);
 		}
 	}
@@ -1540,7 +1540,7 @@ static TArray<VulkanRHI::FSemaphore*> ExtractTransitionSemaphores(TArrayView<con
 
 void FVulkanCommandListContext::RHIBeginTransitions(TArrayView<const FRHITransition*> Transitions)
 {
-	if (Device->SupportsParallelRendering())
+	if (Device.SupportsParallelRendering())
 	{
 		FTransitionProcessor<VkMemoryBarrier2, VkBufferMemoryBarrier2, VkImageMemoryBarrier2> Processor(*this, true);
 		Processor.Process(Transitions);
@@ -1555,28 +1555,28 @@ void FVulkanCommandListContext::RHIBeginTransitions(TArrayView<const FRHITransit
 	TArray<VulkanRHI::FSemaphore*> SignalSemaphores = ExtractTransitionSemaphores(Transitions);
 	if (SignalSemaphores.Num() > 0)
 	{
-		CommandBufferManager->SubmitActiveCmdBuffer(SignalSemaphores);
-		CommandBufferManager->PrepareForNewActiveCommandBuffer();
+		SubmitActiveCmdBuffer(SignalSemaphores);
+		PrepareForNewActiveCommandBuffer();
 	}
 }
 
 void FVulkanCommandListContext::RHIEndTransitions(TArrayView<const FRHITransition*> Transitions)
 {
-	const ERHIPipeline CurrentPipeline = Device->IsRealAsyncComputeContext(this) ? ERHIPipeline::AsyncCompute : ERHIPipeline::Graphics;
+	const ERHIPipeline CurrentPipeline = Device.IsRealAsyncComputeContext(this) ? ERHIPipeline::AsyncCompute : ERHIPipeline::Graphics;
 	TArray<VulkanRHI::FSemaphore*> WaitSemaphores = ExtractTransitionSemaphores(Transitions, CurrentPipeline);
 	if (WaitSemaphores.Num() > 0)
 	{
-		if (CommandBufferManager->HasPendingActiveCmdBuffer())
+		if (HasPendingActiveCmdBuffer())
 		{
-			CommandBufferManager->SubmitActiveCmdBuffer();
-			CommandBufferManager->PrepareForNewActiveCommandBuffer();
+			SubmitActiveCmdBuffer();
+			PrepareForNewActiveCommandBuffer();
 		}
 
 		FVulkanCommandBuffer& CmdBuffer = GetCommandBuffer();
 		CmdBuffer.AddWaitSemaphore(VK_PIPELINE_STAGE_TOP_OF_PIPE_BIT, WaitSemaphores);
 	}
 
-	if (Device->SupportsParallelRendering())
+	if (Device.SupportsParallelRendering())
 	{
 		FTransitionProcessor<VkMemoryBarrier2, VkBufferMemoryBarrier2, VkImageMemoryBarrier2> Processor(*this, false);
 		Processor.Process(Transitions);
@@ -1824,7 +1824,7 @@ void FVulkanPipelineBarrier::Execute(FVulkanCommandBuffer* CmdBuffer)
 {
 	if (MemoryBarriers.Num() != 0 || BufferBarriers.Num() != 0 || ImageBarriers.Num() != 0)
 	{
-		if (CmdBuffer->GetDevice()->SupportsParallelRendering())
+		if (CmdBuffer->Device.SupportsParallelRendering())
 		{
 			VkDependencyInfo DependencyInfo;
 			DependencyInfo.sType = VK_STRUCTURE_TYPE_DEPENDENCY_INFO;
@@ -1987,7 +1987,7 @@ VkImageLayout FVulkanLayoutManager::GetDefaultLayout(FVulkanCommandBuffer* CmdBu
 	{
 		if (VulkanTexture.IsDepthOrStencilAspect())
 		{
-			const bool bSupportsParallelRendering = CmdBuffer->GetDevice()->SupportsParallelRendering();
+			const bool bSupportsParallelRendering = CmdBuffer->Device.SupportsParallelRendering();
 			if (bSupportsParallelRendering)
 			{
 				return VK_IMAGE_LAYOUT_READ_ONLY_OPTIMAL;
@@ -2050,7 +2050,7 @@ VkImageLayout FVulkanLayoutManager::SetExpectedLayout(FVulkanCommandBuffer* CmdB
 
 		if (!SrcLayout)
 		{
-			FVulkanLayoutManager& QueueLayoutMgr = CmdBuffer->GetOwner()->GetMgr().GetCommandListContext()->GetQueue()->GetLayoutManager();
+			FVulkanLayoutManager& QueueLayoutMgr = CmdBuffer->GetOwner()->GetContext().Queue.GetLayoutManager();
 			SrcLayout = QueueLayoutMgr.Layouts.Find(VulkanTexture.Image);
 		}
 
diff --git a/Engine/Source/Runtime/VulkanRHI/Private/VulkanBuffer.cpp b/Engine/Source/Runtime/VulkanRHI/Private/VulkanBuffer.cpp
--- a/Engine/Source/Runtime/VulkanRHI/Private/VulkanBuffer.cpp
+++ b/Engine/Source/Runtime/VulkanRHI/Private/VulkanBuffer.cpp
@@ -247,8 +247,7 @@ void* FVulkanBuffer::Lock(FRHICommandListBase& RHICmdList, EResourceLockMode Loc
 			Device->PrepareForCPURead();
 		
 			FVulkanCommandListContextImmediate& Context = Device->GetImmediateContext();
-			FVulkanCommandBufferManager* CommandBufferManager = Context.GetCommandBufferManager();
-			FVulkanCommandBuffer* CmdBuffer = CommandBufferManager->GetUploadCmdBuffer();
+			FVulkanCommandBuffer* CmdBuffer = Context.GetUploadCmdBuffer();
 				
 			// Make sure any previous tasks have finished on the source buffer.
 			VkMemoryBarrier BarrierBefore = { VK_STRUCTURE_TYPE_MEMORY_BARRIER, nullptr, VK_ACCESS_MEMORY_WRITE_BIT, VK_ACCESS_MEMORY_READ_BIT };
@@ -270,7 +269,7 @@ void* FVulkanBuffer::Lock(FRHICommandListBase& RHICmdList, EResourceLockMode Loc
 			VulkanRHI::vkCmdPipelineBarrier(CmdBuffer->GetHandle(), VK_PIPELINE_STAGE_TRANSFER_BIT, VK_PIPELINE_STAGE_HOST_BIT, 0, 1, &BarrierAfter, 0, nullptr, 0, nullptr);
 				
 			// Force upload.
-			CommandBufferManager->SubmitUploadCmdBuffer();
+			Context.SubmitUploadCmdBuffer();
 			Device->WaitUntilIdle();
 
 			// Flush.
@@ -287,7 +286,7 @@ void* FVulkanBuffer::Lock(FRHICommandListBase& RHICmdList, EResourceLockMode Loc
 			PendingLock.StagingBuffer = StagingBuffer;
 			AddPendingBufferLock(this, PendingLock);
 
-			CommandBufferManager->PrepareForNewActiveCommandBuffer();
+			Context.PrepareForNewActiveCommandBuffer();
 		}
 	}
 	else
diff --git a/Engine/Source/Runtime/VulkanRHI/Private/VulkanCommandBuffer.cpp b/Engine/Source/Runtime/VulkanRHI/Private/VulkanCommandBuffer.cpp
--- a/Engine/Source/Runtime/VulkanRHI/Private/VulkanCommandBuffer.cpp
+++ b/Engine/Source/Runtime/VulkanRHI/Private/VulkanCommandBuffer.cpp
@@ -36,22 +36,12 @@ static FAutoConsoleVariableRef CVarVulkanUseCmdBufferTimingForGPUTime(
 	ECVF_Default
 );
 
-static int32 GVulkanUploadCmdBufferSemaphore = 0;	// despite this behavior being more correct, default to off to avoid changing the existing (as of UE 4.24) behavior
-static FAutoConsoleVariableRef CVarVulkanPreventOverlapWithUpload(
-	TEXT("r.Vulkan.UploadCmdBufferSemaphore"),
-	GVulkanUploadCmdBufferSemaphore,
-	TEXT("Whether command buffers for uploads and graphics can be executed simultaneously.\n")
-	TEXT(" 0: The buffers are submitted without any synch(default)\n")
-	TEXT(" 1: Graphics buffers will not overlap with the upload buffers"),
-	ECVF_ReadOnly
-);
-
 #define CMD_BUFFER_TIME_TO_WAIT_BEFORE_DELETING		10
 
 const uint32 GNumberOfFramesBeforeDeletingDescriptorPool = 300;
 extern int32 GVulkanAutoCorrectUnknownLayouts;
 
-FVulkanCommandBuffer::FVulkanCommandBuffer(FVulkanDevice* InDevice, FVulkanCommandBufferPool* InCommandBufferPool, bool bInIsUploadOnly)
+FVulkanCommandBuffer::FVulkanCommandBuffer(FVulkanDevice& InDevice, FVulkanCommandBufferPool* InCommandBufferPool, bool bInIsUploadOnly)
 	: CurrentStencilRef(0)
 	, State(EState::NotAllocated)
 	, bNeedsDynamicStateSet(1)
@@ -71,15 +61,15 @@ FVulkanCommandBuffer::FVulkanCommandBuffer(FVulkanDevice* InDevice, FVulkanComma
 	, Timing(nullptr)
 #endif
 	, LastValidTiming(0)
-	, LayoutManager(InDevice->SupportsParallelRendering() && !GVulkanAutoCorrectUnknownLayouts,
-		&InCommandBufferPool->GetMgr().GetCommandListContext()->GetQueue()->GetLayoutManager())
+	, LayoutManager(InDevice.SupportsParallelRendering() && !GVulkanAutoCorrectUnknownLayouts,
+		&InCommandBufferPool->GetContext().Queue.GetLayoutManager())
 {
 	{
 		FScopeLock ScopeLock(CommandBufferPool->GetCS());
 		AllocMemory();
 	}
 
-	Fence = Device->GetFenceManager().AllocateFence();
+	Fence = Device.GetFenceManager().AllocateFence();
 }
 
 void FVulkanCommandBuffer::AllocMemory()
@@ -95,7 +85,7 @@ void FVulkanCommandBuffer::AllocMemory()
 	CreateCmdBufInfo.commandBufferCount = 1;
 	CreateCmdBufInfo.commandPool = CommandBufferPool->GetHandle();
 
-	VERIFYVULKANRESULT(VulkanRHI::vkAllocateCommandBuffers(Device->GetInstanceHandle(), &CreateCmdBufInfo, &CommandBufferHandle));
+	VERIFYVULKANRESULT(VulkanRHI::vkAllocateCommandBuffers(Device.GetInstanceHandle(), &CreateCmdBufInfo, &CommandBufferHandle));
 
 	bNeedsDynamicStateSet = 1;
 	bHasPipeline = 0;
@@ -109,7 +99,7 @@ void FVulkanCommandBuffer::AllocMemory()
 
 FVulkanCommandBuffer::~FVulkanCommandBuffer()
 {
-	VulkanRHI::FFenceManager& FenceManager = Device->GetFenceManager();
+	VulkanRHI::FFenceManager& FenceManager = Device.GetFenceManager();
 	if (State == EState::Submitted)
 	{
 		// Wait 33ms
@@ -142,7 +132,7 @@ void FVulkanCommandBuffer::FreeMemory()
 	// Assumes we are inside a lock for the pool
 	check(State != EState::NotAllocated);
 	check(CommandBufferHandle != VK_NULL_HANDLE);
-	VulkanRHI::vkFreeCommandBuffers(Device->GetInstanceHandle(), CommandBufferPool->GetHandle(), 1, &CommandBufferHandle);
+	VulkanRHI::vkFreeCommandBuffers(Device.GetInstanceHandle(), CommandBufferPool->GetHandle(), 1, &CommandBufferHandle);
 	CommandBufferHandle = VK_NULL_HANDLE;
 
 	DEC_DWORD_STAT(STAT_VulkanNumCmdBuffers);
@@ -197,7 +187,7 @@ void FVulkanCommandBuffer::BeginRenderPass(const FVulkanRenderTargetLayout& Layo
 	Info.clearValueCount = Layout.GetNumUsedClearValues();
 	Info.pClearValues = AttachmentClearValues;
 
-	if (Device->GetOptionalExtensions().HasKHRRenderPass2)
+	if (Device.GetOptionalExtensions().HasKHRRenderPass2)
 	{
 		VkSubpassBeginInfo SubpassInfo;
 		ZeroVulkanStruct(SubpassInfo, VK_STRUCTURE_TYPE_SUBPASS_BEGIN_INFO);
@@ -267,14 +257,14 @@ void FVulkanCommandBuffer::End()
 	State = EState::HasEnded;
 }
 
-inline void FVulkanCommandBuffer::InitializeTimings(FVulkanCommandListContext* InContext)
+void FVulkanCommandBuffer::InitializeTimings(FVulkanContextCommon* InContext)
 {
 #if (RHI_NEW_GPU_PROFILER == 0)
 	if ((GVulkanProfileCmdBuffers || GVulkanUseCmdBufferTimingForGPUTime) && !Timing)
 	{
 		if (InContext)
 		{
-			Timing = new FVulkanGPUTiming(InContext, Device);
+			Timing = new FVulkanGPUTiming(InContext, &Device);
 
 			// Upload cb's can be submitted multiple times in a single frame, so we use an expanded pool to catch timings
 			// Any overflow will wrap
@@ -322,7 +312,7 @@ void FVulkanCommandBuffer::Begin()
 
 	if (GVulkanProfileCmdBuffers || GVulkanUseCmdBufferTimingForGPUTime)
 	{
-		InitializeTimings(CommandBufferPool->GetMgr().GetCommandListContext());
+		InitializeTimings(&CommandBufferPool->GetContext());
 #if (RHI_NEW_GPU_PROFILER == 0)
 		if (Timing)
 		{
@@ -332,11 +322,11 @@ void FVulkanCommandBuffer::Begin()
 	}
 	check(!CurrentDescriptorPoolSetContainer);
 
-	if (!bIsUploadOnly && Device->SupportsBindless())
+	if (!bIsUploadOnly && Device.SupportsBindless())
 	{
-		FVulkanBindlessDescriptorManager* BindlessDescriptorManager = Device->GetBindlessDescriptorManager();
-		FVulkanQueue* Queue = GetOwner()->GetMgr().GetQueue();
-		const VkPipelineStageFlags SupportedStages = Queue->GetSupportedStageBits();
+		FVulkanBindlessDescriptorManager* BindlessDescriptorManager = Device.GetBindlessDescriptorManager();
+		FVulkanQueue& Queue = GetOwner()->GetContext().Queue;
+		const VkPipelineStageFlags SupportedStages = Queue.GetSupportedStageBits();
 		BindlessDescriptorManager->BindDescriptorBuffers(CommandBufferHandle, SupportedStages);
 	}
 
@@ -346,7 +336,7 @@ void FVulkanCommandBuffer::Begin()
 void FVulkanCommandBuffer::AcquirePoolSetContainer()
 {
 	check(!CurrentDescriptorPoolSetContainer);
-	CurrentDescriptorPoolSetContainer = &Device->GetDescriptorPoolsManager().AcquirePoolSetContainer();
+	CurrentDescriptorPoolSetContainer = &Device.GetDescriptorPoolsManager().AcquirePoolSetContainer();
 	ensure(TypedDescriptorPoolSets.Num() == 0);
 }
 
@@ -409,7 +399,7 @@ void FVulkanCommandBuffer::RefreshFenceStatus()
 			{
 				//#todo-rco: Reset here?
 				TypedDescriptorPoolSets.Reset();
-				Device->GetDescriptorPoolsManager().ReleasePoolSet(*CurrentDescriptorPoolSetContainer);
+				Device.GetDescriptorPoolsManager().ReleasePoolSet(*CurrentDescriptorPoolSetContainer);
 				CurrentDescriptorPoolSetContainer = nullptr;
 			}
 			else
@@ -427,10 +417,10 @@ void FVulkanCommandBuffer::RefreshFenceStatus()
 	}
 }
 
-FVulkanCommandBufferPool::FVulkanCommandBufferPool(FVulkanDevice* InDevice, FVulkanCommandBufferManager& InMgr)
-	: Handle(VK_NULL_HANDLE)
-	, Device(InDevice)
-	, Mgr(InMgr)
+FVulkanCommandBufferPool::FVulkanCommandBufferPool(FVulkanDevice& InDevice, FVulkanContextCommon& InContext)
+	: Device(InDevice)
+	, Context(InContext)
+	, Handle(VK_NULL_HANDLE)
 {
 }
 
@@ -448,7 +438,7 @@ FVulkanCommandBufferPool::~FVulkanCommandBufferPool()
 		delete CmdBuffer;
 	}
 
-	VulkanRHI::vkDestroyCommandPool(Device->GetInstanceHandle(), Handle, VULKAN_CPU_ALLOCATOR);
+	VulkanRHI::vkDestroyCommandPool(Device.GetInstanceHandle(), Handle, VULKAN_CPU_ALLOCATOR);
 	Handle = VK_NULL_HANDLE;
 }
 
@@ -459,243 +449,52 @@ void FVulkanCommandBufferPool::Create(uint32 QueueFamilyIndex)
 	CmdPoolInfo.queueFamilyIndex =  QueueFamilyIndex;
 	//#todo-rco: Should we use VK_COMMAND_POOL_CREATE_TRANSIENT_BIT?
 	CmdPoolInfo.flags = VK_COMMAND_POOL_CREATE_RESET_COMMAND_BUFFER_BIT;
-	VERIFYVULKANRESULT(VulkanRHI::vkCreateCommandPool(Device->GetInstanceHandle(), &CmdPoolInfo, VULKAN_CPU_ALLOCATOR, &Handle));
-}
-
-FVulkanCommandBufferManager::FVulkanCommandBufferManager(FVulkanDevice* InDevice, FVulkanCommandListContext* InContext)
-	: Device(InDevice)
-	, Context(InContext)
-	, Pool(InDevice, *this)
-	, Queue(InContext->GetQueue())
-	, ActiveCmdBuffer(nullptr)
-	, UploadCmdBuffer(nullptr)
-	, ActiveCmdBufferSemaphore(nullptr)
-	, UploadCmdBufferSemaphore(nullptr)
-{
-	check(Device);
-
-	Pool.Create(Queue->GetFamilyIndex());
-
-	if (GVulkanUploadCmdBufferSemaphore)
-	{
-		ActiveCmdBufferSemaphore = new VulkanRHI::FSemaphore(*Device);
-	}
-
-	ActiveCmdBuffer = Pool.Create(false);
-}
-
-void FVulkanCommandBufferManager::Init(FVulkanCommandListContext* InContext)
-{
-	ActiveCmdBuffer->InitializeTimings(InContext);
-	ActiveCmdBuffer->Begin();
-}
-
-FVulkanCommandBufferManager::~FVulkanCommandBufferManager()
-{
+	VERIFYVULKANRESULT(VulkanRHI::vkCreateCommandPool(Device.GetInstanceHandle(), &CmdPoolInfo, VULKAN_CPU_ALLOCATOR, &Handle));
 }
 
-void FVulkanCommandBufferManager::WaitForCmdBuffer(FVulkanCommandBuffer* CmdBuffer, float TimeInSecondsToWait)
-{
-	FScopeLock ScopeLock(&Pool.CS);
-	check(CmdBuffer->IsSubmitted());
-	bool bSuccess = Device->GetFenceManager().WaitForFence(CmdBuffer->Fence, (uint64)(TimeInSecondsToWait * 1e9));
-	check(bSuccess);
-	CmdBuffer->RefreshFenceStatus();
-}
 
-void FVulkanCommandBufferManager::AddQueryPoolForReset(VkQueryPool QueryPool, uint32 Size)
-{
-	FScopeLock ScopeLock(&Pool.CS);
-	FQueryPoolReset QPR = { QueryPool, Size};
-	PoolResets.Add(QPR);
-}
-void FVulkanCommandBufferManager::FlushResetQueryPools()
+void FVulkanCommandBufferPool::FreeUnusedCmdBuffers(FVulkanQueue* InQueue, bool bTrimMemory)
 {
-	FScopeLock ScopeLock(&Pool.CS);
-	if(PoolResets.Num())
+#if VULKAN_DELETE_STALE_CMDBUFFERS
+	FScopeLock ScopeLock(&CS);
+	
+	if (bTrimMemory)
 	{
-		FVulkanCommandBuffer* PoolResetCommandBuffer = nullptr;
-		for (int32 Index = 0; Index < Pool.CmdBuffers.Num(); ++Index)
-		{
-			FVulkanCommandBuffer* CmdBuffer = Pool.CmdBuffers[Index];
-			CmdBuffer->RefreshFenceStatus();
-	#if VULKAN_USE_DIFFERENT_POOL_CMDBUFFERS
-			if (!CmdBuffer->bIsUploadOnly)
-	#endif
-			{
-				if (CmdBuffer->State == FVulkanCommandBuffer::EState::ReadyForBegin)
-				{
-					PoolResetCommandBuffer = CmdBuffer;
-				}
-			}
-		}
-		if (!PoolResetCommandBuffer)
-		{
-			PoolResetCommandBuffer = Pool.Create(false);
-		}
-		PoolResetCommandBuffer->Begin();
-
-		VkCommandBuffer CommandBuffer = PoolResetCommandBuffer->GetHandle();
-		for (FQueryPoolReset& QPR : PoolResets)
-		{
-			VulkanRHI::vkCmdResetQueryPool(CommandBuffer, QPR.Pool, 0, QPR.Size);
-		}
-		PoolResetCommandBuffer->End();
-		Queue->Submit(PoolResetCommandBuffer);
-		PoolResets.Empty();
+		VulkanRHI::vkTrimCommandPool(Device.GetInstanceHandle(), Handle, 0);
+		return;
 	}
-}
-
-void FVulkanCommandBufferManager::SubmitUploadCmdBuffer(uint32 NumSignalSemaphores, VkSemaphore* SignalSemaphores)
-{
-	FScopeLock ScopeLock(&Pool.CS);
-	check(UploadCmdBuffer);
-	check(UploadCmdBuffer->CurrentDescriptorPoolSetContainer == nullptr);
-	if (!UploadCmdBuffer->IsSubmitted() && UploadCmdBuffer->HasBegun())
-	{
-		check(UploadCmdBuffer->IsOutsideRenderPass());
-
-		VulkanRHI::DebugHeavyWeightBarrier(UploadCmdBuffer->GetHandle(), 4);
-
-		UploadCmdBuffer->End();
-
-		if (GVulkanUploadCmdBufferSemaphore)
-		{
-			// Add semaphores associated with the recent active cmdbuf(s), if any. That will prevent
-			// the overlap, delaying execution of this cmdbuf until the graphics one(s) is complete.
-			for (VulkanRHI::FSemaphore* WaitForThis : RenderingCompletedSemaphores)
-			{
-				UploadCmdBuffer->AddWaitSemaphore(VK_PIPELINE_STAGE_TOP_OF_PIPE_BIT, WaitForThis);
-			}
-
-			if (NumSignalSemaphores == 0)
-			{
-				checkf(UploadCmdBufferSemaphore != nullptr, TEXT("FVulkanCommandBufferManager::SubmitUploadCmdBuffer: upload command buffer does not have an associated completion semaphore."));
-				VkSemaphore Sema = UploadCmdBufferSemaphore->GetHandle();
-				Queue->Submit(UploadCmdBuffer, 1, &Sema);
-				UploadCompletedSemaphores.Add(UploadCmdBufferSemaphore);
-				UploadCmdBufferSemaphore = nullptr;
-			}
-			else
-			{
-				TArray<VkSemaphore, TInlineAllocator<16>> CombinedSemaphores;
-				CombinedSemaphores.Add(UploadCmdBufferSemaphore->GetHandle());
-				CombinedSemaphores.Append(SignalSemaphores, NumSignalSemaphores);
-				Queue->Submit(UploadCmdBuffer, CombinedSemaphores.Num(), CombinedSemaphores.GetData());
-			}
-			// the buffer will now hold on to the wait semaphores, so we can clear them here
-			RenderingCompletedSemaphores.Empty();
-		}
-		else
-		{
-			Queue->Submit(UploadCmdBuffer, NumSignalSemaphores, SignalSemaphores);
-		}
 
-		UploadCmdBuffer->SubmittedTime = FPlatformTime::Seconds();
-	}
+	const double CurrentTime = FPlatformTime::Seconds();
 
-	UploadCmdBuffer = nullptr;
-}
+	// In case Queue stores pointer to a cmdbuffer, do not delete it
+	FVulkanCommandBuffer* LastSubmittedCmdBuffer = nullptr;
+	uint64 LastSubmittedFenceCounter = 0;
+	InQueue->GetLastSubmittedInfo(LastSubmittedCmdBuffer, LastSubmittedFenceCounter);
 
-void FVulkanCommandBufferManager::SubmitActiveCmdBuffer(TArrayView<VulkanRHI::FSemaphore*> SignalSemaphores)
-{
-	FScopeLock ScopeLock(&Pool.CS);
-	FlushResetQueryPools();
-	check(!UploadCmdBuffer);
-	check(ActiveCmdBuffer);
-
-	TArray<VkSemaphore, FConcurrentLinearArrayAllocator> SemaphoreHandles;
-	SemaphoreHandles.Reserve(SignalSemaphores.Num() + 1);
-	for (VulkanRHI::FSemaphore* Semaphore : SignalSemaphores)
-	{
-		SemaphoreHandles.Add(Semaphore->GetHandle());
-	}
+	// Deferred deletion queue caches pointers to cmdbuffers
+	VulkanRHI::FDeferredDeletionQueue2& DeferredDeletionQueue = Device.GetDeferredDeletionQueue();
 
-	if (!ActiveCmdBuffer->IsSubmitted() && ActiveCmdBuffer->HasBegun())
+	for (int32 Index = CmdBuffers.Num() - 1; Index >= 0; --Index)
 	{
-		if (!ActiveCmdBuffer->IsOutsideRenderPass())
-		{
-			UE_LOG(LogVulkanRHI, Warning, TEXT("Forcing EndRenderPass() for submission"));
-			ActiveCmdBuffer->EndRenderPass();
-		}
-
-		VulkanRHI::DebugHeavyWeightBarrier(ActiveCmdBuffer->GetHandle(), 8);
-
-		ActiveCmdBuffer->End();
-
-		if (GVulkanUploadCmdBufferSemaphore)
+		FVulkanCommandBuffer* CmdBuffer = CmdBuffers[Index];
+		if (CmdBuffer != LastSubmittedCmdBuffer &&
+			(CmdBuffer->State == FVulkanCommandBuffer::EState::ReadyForBegin || CmdBuffer->State == FVulkanCommandBuffer::EState::NeedReset) &&
+			((CurrentTime - CmdBuffer->SubmittedTime) > CMD_BUFFER_TIME_TO_WAIT_BEFORE_DELETING))
 		{
-			checkf(ActiveCmdBufferSemaphore != nullptr, TEXT("FVulkanCommandBufferManager::SubmitUploadCmdBuffer: graphics command buffer does not have an associated completion semaphore."));
-
-			// Add semaphores associated with the recent upload cmdbuf(s), if any. That will prevent
-			// the overlap, delaying execution of this cmdbuf until upload one(s) are complete.
-			for (VulkanRHI::FSemaphore* UploadCompleteSema : UploadCompletedSemaphores)
+			// Skip command buffer that contain unresolved compaction queries
+			if (Device.GetRayTracingCompactionRequestHandler() && Device.GetRayTracingCompactionRequestHandler()->IsUsingCmdBuffer(CmdBuffer))
 			{
-				ActiveCmdBuffer->AddWaitSemaphore(VK_PIPELINE_STAGE_TOP_OF_PIPE_BIT, UploadCompleteSema);
+				continue;
 			}
 
-			SemaphoreHandles.Add(ActiveCmdBufferSemaphore->GetHandle());
-			Queue->Submit(ActiveCmdBuffer, SemaphoreHandles.Num(), SemaphoreHandles.GetData());
-
-			RenderingCompletedSemaphores.Add(ActiveCmdBufferSemaphore);
-			ActiveCmdBufferSemaphore = nullptr;
-
-			// the buffer will now hold on to the wait semaphores, so we can clear them here
-			UploadCompletedSemaphores.Empty();
-		}
-		else
-		{
-			Queue->Submit(ActiveCmdBuffer, SemaphoreHandles.Num(), SemaphoreHandles.GetData());
-		}
-		ActiveCmdBuffer->SubmittedTime = FPlatformTime::Seconds();
-	}
-
-	ActiveCmdBuffer = nullptr;
-	if (GVulkanUploadCmdBufferSemaphore && ActiveCmdBufferSemaphore != nullptr)
-	{
-		// most likely we didn't submit as it wasn't begun
-		check(ActiveCmdBufferSemaphore->GetRefCount() == 0);
-		delete ActiveCmdBufferSemaphore;
-		ActiveCmdBufferSemaphore = nullptr;
-	}
-}
-
-void FVulkanCommandBufferManager::SubmitActiveCmdBufferFromPresent(VulkanRHI::FSemaphore* SignalSemaphore)
-{
-	if (GVulkanUploadCmdBufferSemaphore)
-	{
-		// unlike more advanced regular SACB(), this is just a wrapper
-		// around Queue->Submit() to avoid rewriting the logic in Present
-		if (SignalSemaphore)
-		{
-			VkSemaphore SignalThis[2] =
-			{
-				SignalSemaphore->GetHandle(),
-				ActiveCmdBufferSemaphore->GetHandle()
-			};
-
-			Queue->Submit(ActiveCmdBuffer, UE_ARRAY_COUNT(SignalThis), SignalThis);
-		}
-		else
-		{
-			VkSemaphore SignalThis = ActiveCmdBufferSemaphore->GetHandle();
-			Queue->Submit(ActiveCmdBuffer, 1, &SignalThis);
-		}
+			DeferredDeletionQueue.OnCmdBufferDeleted(CmdBuffer);
 
-		RenderingCompletedSemaphores.Add(ActiveCmdBufferSemaphore);
-		ActiveCmdBufferSemaphore = nullptr;
-	}
-	else
-	{
-		if (SignalSemaphore)
-		{
-			Queue->Submit(ActiveCmdBuffer, SignalSemaphore->GetHandle());
-		}
-		else
-		{
-			Queue->Submit(ActiveCmdBuffer);
+			CmdBuffer->FreeMemory();
+			CmdBuffers.RemoveAtSwap(Index, EAllowShrinking::No);
+			FreeCmdBuffers.Add(CmdBuffer);
 		}
 	}
+#endif
 }
 
 FVulkanCommandBuffer* FVulkanCommandBufferPool::Create(bool bIsUploadOnly)
@@ -733,176 +532,3 @@ void FVulkanCommandBufferPool::RefreshFenceStatus(FVulkanCommandBuffer* SkipCmdB
 		}
 	}
 }
-
-void FVulkanCommandBufferManager::PrepareForNewActiveCommandBuffer()
-{
-	FScopeLock ScopeLock(&Pool.CS);
-	check(!UploadCmdBuffer);
-
-	if (GVulkanUploadCmdBufferSemaphore)
-	{
-		// create a completion semaphore
-		checkf(ActiveCmdBufferSemaphore == nullptr, TEXT("FVulkanCommandBufferManager::PrepareForNewActiveCommandBuffer: a semaphore from the previous active buffer has been leaked."));
-		ActiveCmdBufferSemaphore = new VulkanRHI::FSemaphore(*Device);
-	}
-
-	for (int32 Index = 0; Index < Pool.CmdBuffers.Num(); ++Index)
-	{
-		FVulkanCommandBuffer* CmdBuffer = Pool.CmdBuffers[Index];
-		CmdBuffer->RefreshFenceStatus();
-#if VULKAN_USE_DIFFERENT_POOL_CMDBUFFERS
-		if (!CmdBuffer->bIsUploadOnly)
-#endif
-		{
-			if (CmdBuffer->State == FVulkanCommandBuffer::EState::ReadyForBegin || CmdBuffer->State == FVulkanCommandBuffer::EState::NeedReset)
-			{
-				ActiveCmdBuffer = CmdBuffer;
-				ActiveCmdBuffer->Begin();
-				return;
-			}
-			else
-			{
-				check(CmdBuffer->State == FVulkanCommandBuffer::EState::Submitted);
-			}
-		}
-	}
-
-	// All cmd buffers are being executed still
-	ActiveCmdBuffer = Pool.Create(false);
-	ActiveCmdBuffer->Begin();
-}
-
-#if (RHI_NEW_GPU_PROFILER == 0)
-uint32 FVulkanCommandBufferManager::CalculateGPUTime()
-{
-	uint32 Time = 0;
-	for (int32 Index = 0; Index < Pool.CmdBuffers.Num(); ++Index)
-	{
-		FVulkanCommandBuffer* CmdBuffer = Pool.CmdBuffers[Index];
-		if (CmdBuffer->HasValidTiming())
-		{
-			Time += CmdBuffer->Timing->GetTiming(false);
-		}
-	}
-	return Time;
-}
-#endif // (RHI_NEW_GPU_PROFILER == 0)
-
-FVulkanCommandBuffer* FVulkanCommandBufferManager::GetUploadCmdBuffer()
-{
-	FScopeLock ScopeLock(&Pool.CS);
-	if (!UploadCmdBuffer)
-	{
-		if (GVulkanUploadCmdBufferSemaphore)
-		{
-			// create a completion semaphore
-			checkf(UploadCmdBufferSemaphore == nullptr, TEXT("FVulkanCommandBufferManager::GetUploadCmdBuffer: a semaphore from the previous upload buffer has been leaked."));
-			UploadCmdBufferSemaphore = new VulkanRHI::FSemaphore(*Device);
-		}
-
-		for (int32 Index = 0; Index < Pool.CmdBuffers.Num(); ++Index)
-		{
-			FVulkanCommandBuffer* CmdBuffer = Pool.CmdBuffers[Index];
-			CmdBuffer->RefreshFenceStatus();
-#if VULKAN_USE_DIFFERENT_POOL_CMDBUFFERS
-			if (CmdBuffer->bIsUploadOnly)
-#endif
-			{
-				if (CmdBuffer->State == FVulkanCommandBuffer::EState::ReadyForBegin || CmdBuffer->State == FVulkanCommandBuffer::EState::NeedReset)
-				{
-					UploadCmdBuffer = CmdBuffer;
-					UploadCmdBuffer->Begin();
-					return UploadCmdBuffer;
-				}
-			}
-		}
-
-		// All cmd buffers are being executed still
-		UploadCmdBuffer = Pool.Create(true);
-		UploadCmdBuffer->Begin();
-	}
-
-	return UploadCmdBuffer;
-}
-
-#if VULKAN_DELETE_STALE_CMDBUFFERS
-struct FRHICommandFreeUnusedCmdBuffers final : public FRHICommand<FRHICommandFreeUnusedCmdBuffers>
-{
-	FVulkanCommandBufferPool* Pool;
-	FVulkanQueue* Queue;
-	bool bTrimMemory;
-
-	FRHICommandFreeUnusedCmdBuffers(FVulkanCommandBufferPool* InPool, FVulkanQueue* InQueue, bool bInTrimMemory)
-		: Pool(InPool)
-		, Queue(InQueue)
-		, bTrimMemory(bInTrimMemory)
-	{
-	}
-
-	void Execute(FRHICommandListBase& CmdList)
-	{
-		Pool->FreeUnusedCmdBuffers(Queue, bTrimMemory);
-	}
-};
-#endif
-
-
-void FVulkanCommandBufferPool::FreeUnusedCmdBuffers(FVulkanQueue* InQueue, bool bTrimMemory)
-{
-#if VULKAN_DELETE_STALE_CMDBUFFERS
-	FScopeLock ScopeLock(&CS);
-	
-	if (bTrimMemory)
-	{
-		VulkanRHI::vkTrimCommandPool(Device->GetInstanceHandle(), Handle, 0);
-		return;
-	}
-
-	const double CurrentTime = FPlatformTime::Seconds();
-
-	// In case Queue stores pointer to a cmdbuffer, do not delete it
-	FVulkanCommandBuffer* LastSubmittedCmdBuffer = nullptr;
-	uint64 LastSubmittedFenceCounter = 0;
-	InQueue->GetLastSubmittedInfo(LastSubmittedCmdBuffer, LastSubmittedFenceCounter);
-
-	// Deferred deletion queue caches pointers to cmdbuffers
-	VulkanRHI::FDeferredDeletionQueue2& DeferredDeletionQueue = Device->GetDeferredDeletionQueue();
-
-	for (int32 Index = CmdBuffers.Num() - 1; Index >= 0; --Index)
-	{
-		FVulkanCommandBuffer* CmdBuffer = CmdBuffers[Index];
-		if (CmdBuffer != LastSubmittedCmdBuffer &&
-			(CmdBuffer->State == FVulkanCommandBuffer::EState::ReadyForBegin || CmdBuffer->State == FVulkanCommandBuffer::EState::NeedReset) &&
-			((CurrentTime - CmdBuffer->SubmittedTime) > CMD_BUFFER_TIME_TO_WAIT_BEFORE_DELETING))
-		{
-			// Skip command buffer that contain unresolved compaction queries
-			if (Device->GetRayTracingCompactionRequestHandler() && Device->GetRayTracingCompactionRequestHandler()->IsUsingCmdBuffer(CmdBuffer))
-			{
-				continue;
-			}
-
-			DeferredDeletionQueue.OnCmdBufferDeleted(CmdBuffer);
-
-			CmdBuffer->FreeMemory();
-			CmdBuffers.RemoveAtSwap(Index, EAllowShrinking::No);
-			FreeCmdBuffers.Add(CmdBuffer);
-		}
-	}
-#endif
-}
-
-void FVulkanCommandBufferManager::FreeUnusedCmdBuffers(bool bTrimMemory)
-{
-#if VULKAN_DELETE_STALE_CMDBUFFERS
-	FRHICommandList& RHICmdList = FRHICommandListExecutor::GetImmediateCommandList();
-	if (!IsInRenderingThread() || (RHICmdList.Bypass() || !IsRunningRHIInSeparateThread()))
-	{
-		Pool.FreeUnusedCmdBuffers(Queue, bTrimMemory);
-	}
-	else
-	{
-		check(IsInRenderingThread());
-		ALLOC_COMMAND_CL(RHICmdList, FRHICommandFreeUnusedCmdBuffers)(&Pool, Queue, bTrimMemory);
-	}
-#endif
-}
diff --git a/Engine/Source/Runtime/VulkanRHI/Private/VulkanCommandBuffer.h b/Engine/Source/Runtime/VulkanRHI/Private/VulkanCommandBuffer.h
--- a/Engine/Source/Runtime/VulkanRHI/Private/VulkanCommandBuffer.h
+++ b/Engine/Source/Runtime/VulkanRHI/Private/VulkanCommandBuffer.h
@@ -15,7 +15,7 @@ extern int32 GVulkanUseCmdBufferTimingForGPUTime;
 
 class FVulkanDevice;
 class FVulkanCommandBufferPool;
-class FVulkanCommandBufferManager;
+class FVulkanContextCommon;
 class FVulkanRenderTargetLayout;
 class FVulkanQueue;
 class FVulkanDescriptorPoolSetContainer;
@@ -30,11 +30,10 @@ namespace VulkanRHI
 class FVulkanCommandBuffer
 {
 protected:
-	friend class FVulkanCommandBufferManager;
 	friend class FVulkanCommandBufferPool;
 	friend class FVulkanQueue;
 
-	FVulkanCommandBuffer(FVulkanDevice* InDevice, FVulkanCommandBufferPool* InCommandBufferPool, bool bInIsUploadOnly);
+	FVulkanCommandBuffer(FVulkanDevice& InDevice, FVulkanCommandBufferPool* InCommandBufferPool, bool bInIsUploadOnly);
 	~FVulkanCommandBuffer();
 
 public:
@@ -43,11 +42,6 @@ public:
 		return CommandBufferPool;
 	}
 
-	FVulkanDevice* GetDevice()
-	{
-		return Device;
-	}
-
 	bool IsUniformBufferBarrierAdded() const
 	{
 		return bIsUniformBufferBarrierAdded;
@@ -212,8 +206,10 @@ public:
 	};
 	void AddPendingTimestampQuery(uint64 Index, uint64 Count, VkQueryPool PoolHandle, VkBuffer BufferHandle, bool bBlocking);
 
+	// Public immutable reference to be accessed directly
+	FVulkanDevice& Device;
+
 private:
-	FVulkanDevice* Device;
 	VkCommandBuffer CommandBufferHandle;
 	double SubmittedTime = 0.0f;
 
@@ -239,7 +235,7 @@ private:
 	volatile uint64 SubmittedFenceCounter;
 
 	void RefreshFenceStatus();
-	void InitializeTimings(FVulkanCommandListContext* InContext);
+	void InitializeTimings(FVulkanContextCommon* InContext);
 
 	FVulkanCommandBufferPool* CommandBufferPool;
 
@@ -260,176 +256,49 @@ public:
 	TMap<uint32, class FVulkanTypedDescriptorPoolSet*> TypedDescriptorPoolSets;
 
 	friend class FVulkanDynamicRHI;
-	friend class FVulkanLayoutManager;
+	friend FVulkanLayoutManager;
+	friend FVulkanContextCommon;
 };
 
 class FVulkanCommandBufferPool
 {
 public:
-	FVulkanCommandBufferPool(FVulkanDevice* InDevice, FVulkanCommandBufferManager& InMgr);
+	FVulkanCommandBufferPool(FVulkanDevice& InDevice, FVulkanContextCommon& InContext);
 	~FVulkanCommandBufferPool();
 
 	void RefreshFenceStatus(FVulkanCommandBuffer* SkipCmdBuffer = nullptr);
 
-	inline VkCommandPool GetHandle() const
+	VkCommandPool GetHandle() const
 	{
 		return Handle;
 	}
 
-	inline FCriticalSection* GetCS()
+	FCriticalSection* GetCS()
 	{
 		return &CS;
 	}
 
 	void FreeUnusedCmdBuffers(FVulkanQueue* Queue, bool bTrimMemory);
 
-	inline FVulkanCommandBufferManager& GetMgr()
+	FVulkanContextCommon& GetContext()
 	{
-		return Mgr;
+		return Context;
 	}
 
 private:
+	FVulkanDevice& Device;
+	FVulkanContextCommon& Context;
+
 	VkCommandPool Handle;
 
 	TArray<FVulkanCommandBuffer*> CmdBuffers;
 	TArray<FVulkanCommandBuffer*> FreeCmdBuffers;
 
 	FCriticalSection CS;
-	FVulkanDevice* Device;
-
-	FVulkanCommandBufferManager& Mgr;
 
 	FVulkanCommandBuffer* Create(bool bIsUploadOnly);
 
-
 	void Create(uint32 QueueFamilyIndex);
-	friend class FVulkanCommandBufferManager;
-};
-
-class FVulkanCommandBufferManager
-{
-public:
-	FVulkanCommandBufferManager(FVulkanDevice* InDevice, FVulkanCommandListContext* InContext);
-	~FVulkanCommandBufferManager();
-
-	void Init(FVulkanCommandListContext* InContext);
-
-	inline FVulkanCommandBuffer* GetActiveCmdBuffer()
-	{
-		if (UploadCmdBuffer)
-		{
-			SubmitUploadCmdBuffer();
-		}
-
-		return ActiveCmdBuffer;
-	}
-
-	inline FVulkanCommandBuffer* GetActiveCmdBufferDirect()
-	{
-		return ActiveCmdBuffer;
-	}
-
-
-	inline bool HasPendingUploadCmdBuffer() const
-	{
-		return UploadCmdBuffer != nullptr;
-	}
-
-	inline bool HasPendingActiveCmdBuffer() const
-	{
-		return ActiveCmdBuffer != nullptr;
-	}
-
-	VULKANRHI_API FVulkanCommandBuffer* GetUploadCmdBuffer();
-
-	VULKANRHI_API void SubmitUploadCmdBuffer(uint32 NumSignalSemaphores = 0, VkSemaphore* SignalSemaphores = nullptr);
-
-	void SubmitActiveCmdBuffer(TArrayView<VulkanRHI::FSemaphore*> SignalSemaphores);
-
-	void SubmitActiveCmdBuffer()
-	{
-		SubmitActiveCmdBuffer(MakeArrayView<VulkanRHI::FSemaphore*>(nullptr, 0));
-	}
-
-	void SubmitActiveCmdBuffer(VulkanRHI::FSemaphore* SignalSemaphore)
-	{
-		SubmitActiveCmdBuffer(MakeArrayView<VulkanRHI::FSemaphore*>(&SignalSemaphore, 1));
-	}
-
-	/** Regular SACB() expects not-ended and would rotate the command buffer immediately, but Present has a special logic */
-	void SubmitActiveCmdBufferFromPresent(VulkanRHI::FSemaphore* SignalSemaphore = nullptr);
-
-	void WaitForCmdBuffer(FVulkanCommandBuffer* CmdBuffer, float TimeInSecondsToWait = 10.0f);
-
-
-	void AddQueryPoolForReset(VkQueryPool Pool, uint32 Size);
-	void FlushResetQueryPools();
-
-	// Update the fences of all cmd buffers except SkipCmdBuffer
-	void RefreshFenceStatus(FVulkanCommandBuffer* SkipCmdBuffer = nullptr)
-	{
-		Pool.RefreshFenceStatus(SkipCmdBuffer);
-	}
-
-	void PrepareForNewActiveCommandBuffer();
-
-	inline VkCommandPool GetHandle() const
-	{
-		return Pool.GetHandle();
-	}
-
-#if (RHI_NEW_GPU_PROFILER == 0)
-	uint32 CalculateGPUTime();
-#endif
-
-	void FreeUnusedCmdBuffers(bool bTrimMemory);
-
-	inline FVulkanCommandListContext* GetCommandListContext()
-	{
-		return Context;
-	}
-
-	inline FVulkanQueue* GetQueue()
-	{
-		return Queue;
-	}
-
-	inline void NotifyDeletedImage(VkImage Image)
-	{
-		if (UploadCmdBuffer)
-		{
-			UploadCmdBuffer->GetLayoutManager().NotifyDeletedImage(Image);
-		}
-		if (ActiveCmdBuffer)
-		{
-			ActiveCmdBuffer->GetLayoutManager().NotifyDeletedImage(Image);
-		}
-	}
 
-private:
-	struct FQueryPoolReset
-	{
-		VkQueryPool Pool; 
-		uint32 Size;
-	};
-	
-	FVulkanDevice* Device;
-	FVulkanCommandListContext* Context;
-	FVulkanCommandBufferPool Pool;
-	FVulkanQueue* Queue;
-	FVulkanCommandBuffer* ActiveCmdBuffer;
-	FVulkanCommandBuffer* UploadCmdBuffer;
-	TArray<FQueryPoolReset> PoolResets;
-
-	/** This semaphore is used to prevent overlaps between the (current) graphics cmdbuf and next upload cmdbuf. */
-	VulkanRHI::FSemaphore* ActiveCmdBufferSemaphore;
-
-	/** Holds semaphores associated with the recent upload cmdbuf(s) - waiting to be added to the next graphics cmdbuf as WaitSemaphores. */
-	TArray<VulkanRHI::FSemaphore*> RenderingCompletedSemaphores;
-	
-	/** This semaphore is used to prevent overlaps between (current) upload cmdbuf and next graphics cmdbuf. */
-	VulkanRHI::FSemaphore* UploadCmdBufferSemaphore;
-	
-	/** Holds semaphores associated with the recent upload cmdbuf(s) - waiting to be added to the next graphics cmdbuf as WaitSemaphores. */
-	TArray<VulkanRHI::FSemaphore*> UploadCompletedSemaphores;
+	friend FVulkanContextCommon;
 };
diff --git a/Engine/Source/Runtime/VulkanRHI/Private/VulkanCommands.cpp b/Engine/Source/Runtime/VulkanRHI/Private/VulkanCommands.cpp
--- a/Engine/Source/Runtime/VulkanRHI/Private/VulkanCommands.cpp
+++ b/Engine/Source/Runtime/VulkanRHI/Private/VulkanCommands.cpp
@@ -184,7 +184,7 @@ struct FVulkanResourceBinder
 		{
 			FVulkanTexture* VulkanTexture = ResourceCast(TextureRHI);
 			const ERHIAccess RHIAccess = (Frequency == SF_Compute) ? ERHIAccess::SRVCompute : ERHIAccess::SRVGraphics;
-			const VkImageLayout ExpectedLayout = FVulkanLayoutManager::GetDefaultLayout(Context.GetCommandBufferManager()->GetActiveCmdBuffer(), *VulkanTexture, RHIAccess);
+			const VkImageLayout ExpectedLayout = FVulkanLayoutManager::GetDefaultLayout(Context.GetActiveCmdBuffer(), *VulkanTexture, RHIAccess);
 			PendingState->SetTextureForUBResource(Stage, Index, VulkanTexture, ExpectedLayout);
 		}
 	}
@@ -375,7 +375,7 @@ void FVulkanCommandListContext::RHISetUAVParameter(FRHIComputeShader* ComputeSha
 void FVulkanCommandListContext::RHISetShaderTexture(FRHIGraphicsShader* ShaderRHI, uint32 TextureIndex, FRHITexture* NewTextureRHI)
 {
 	FVulkanTexture* VulkanTexture = ResourceCast(NewTextureRHI);
-	const VkImageLayout ExpectedLayout = FVulkanLayoutManager::GetDefaultLayout(GetCommandBufferManager()->GetActiveCmdBuffer(), *VulkanTexture, ERHIAccess::SRVGraphics);
+	const VkImageLayout ExpectedLayout = FVulkanLayoutManager::GetDefaultLayout(GetActiveCmdBuffer(), *VulkanTexture, ERHIAccess::SRVGraphics);
 
 	ShaderStage::EStage Stage = GetAndVerifyShaderStage(ShaderRHI, PendingGfxState);
 	PendingGfxState->SetTextureForStage(Stage, TextureIndex, VulkanTexture, ExpectedLayout);
@@ -388,7 +388,7 @@ void FVulkanCommandListContext::RHISetShaderTexture(FRHIComputeShader* ComputeSh
 	check(PendingComputeState->GetCurrentShader() == ComputeShader);
 
 	FVulkanTexture* VulkanTexture = ResourceCast(NewTextureRHI);
-	const VkImageLayout ExpectedLayout = FVulkanLayoutManager::GetDefaultLayout(GetCommandBufferManager()->GetActiveCmdBuffer(), *VulkanTexture, ERHIAccess::SRVCompute);
+	const VkImageLayout ExpectedLayout = FVulkanLayoutManager::GetDefaultLayout(GetActiveCmdBuffer(), *VulkanTexture, ERHIAccess::SRVCompute);
 	PendingComputeState->SetTextureForStage(TextureIndex, VulkanTexture, ExpectedLayout);
 	NewTextureRHI->SetLastRenderTime((float)FPlatformTime::Seconds());
 }
@@ -485,7 +485,7 @@ void FVulkanCommandListContext::RHISetStaticUniformBuffer(FUniformBufferStaticSl
 
 void FVulkanCommandListContext::RHISetUniformBufferDynamicOffset(FUniformBufferStaticSlot InSlot, uint32 InOffset)
 {
-	check(IsAligned(InOffset, Device->GetLimits().minUniformBufferOffsetAlignment));
+	check(IsAligned(InOffset, Device.GetLimits().minUniformBufferOffsetAlignment));
 
 	FVulkanUniformBuffer* UniformBuffer = ResourceCast(GlobalUniformBuffers[InSlot]);
 
@@ -895,11 +895,11 @@ void FVulkanCommandListContext::RHISetDepthBounds(float MinDepth, float MaxDepth
 
 void FVulkanCommandListContext::RequestSubmitCurrentCommands()
 {
-	if (Device->GetComputeQueue() == Queue)
+	if (Device.GetComputeQueue() == &Queue)
 	{
-		if (CommandBufferManager->HasPendingUploadCmdBuffer())
+		if (HasPendingUploadCmdBuffer())
 		{
-			CommandBufferManager->SubmitUploadCmdBuffer();
+			SubmitUploadCmdBuffer();
 		}
 		bSubmitAtNextSafePoint = true;
 		SafePointSubmit();
@@ -913,23 +913,23 @@ void FVulkanCommandListContext::RequestSubmitCurrentCommands()
 
 void FVulkanCommandListContext::InternalSubmitActiveCmdBuffer()
 {
-	CommandBufferManager->SubmitActiveCmdBuffer();
-	CommandBufferManager->PrepareForNewActiveCommandBuffer();
+	SubmitActiveCmdBuffer();
+	PrepareForNewActiveCommandBuffer();
 }
 
 void FVulkanCommandListContext::PrepareForCPURead()
 {
 	ensure(IsImmediate());
-	FVulkanCommandBuffer* CmdBuffer = CommandBufferManager->GetActiveCmdBuffer();
+	FVulkanCommandBuffer* CmdBuffer = GetActiveCmdBuffer();
 	if (CmdBuffer && CmdBuffer->HasBegun())
 	{
 		check(!CmdBuffer->IsInsideRenderPass());
 
-		CommandBufferManager->SubmitActiveCmdBuffer();
+		SubmitActiveCmdBuffer();
 		if (!GWaitForIdleOnSubmit)
 		{
 			// The wait has already happened if GWaitForIdleOnSubmit is set
-			CommandBufferManager->WaitForCmdBuffer(CmdBuffer);
+			WaitForCmdBuffer(CmdBuffer);
 		}
 	}
 }
@@ -946,12 +946,12 @@ void FVulkanCommandListContext::RHICopyToStagingBuffer(FRHIBuffer* SourceBufferR
 	{
 		if (StagingBuffer->StagingBuffer)
 		{
-			Device->GetStagingManager().ReleaseBuffer(nullptr, StagingBuffer->StagingBuffer);
+			Device.GetStagingManager().ReleaseBuffer(nullptr, StagingBuffer->StagingBuffer);
 		}
 
-		VulkanRHI::FStagingBuffer* ReadbackStagingBuffer = Device->GetStagingManager().AcquireBuffer(NumBytes, VK_BUFFER_USAGE_TRANSFER_DST_BIT, VK_MEMORY_PROPERTY_HOST_CACHED_BIT);
+		VulkanRHI::FStagingBuffer* ReadbackStagingBuffer = Device.GetStagingManager().AcquireBuffer(NumBytes, VK_BUFFER_USAGE_TRANSFER_DST_BIT, VK_MEMORY_PROPERTY_HOST_CACHED_BIT);
 		StagingBuffer->StagingBuffer = ReadbackStagingBuffer;
-		StagingBuffer->Device = Device;
+		StagingBuffer->Device = &Device;
 	}
 
 	StagingBuffer->QueuedNumBytes = NumBytes;
@@ -966,7 +966,7 @@ void FVulkanCommandListContext::RHICopyToStagingBuffer(FRHIBuffer* SourceBufferR
 
 void FVulkanCommandListContext::RHIWriteGPUFence(FRHIGPUFence* FenceRHI)
 {
-	FVulkanCommandBuffer* CmdBuffer = CommandBufferManager->GetActiveCmdBuffer();
+	FVulkanCommandBuffer* CmdBuffer = GetActiveCmdBuffer();
 	FVulkanGPUFence* Fence = ResourceCast(FenceRHI);
 
 	Fence->CmdBuffer = CmdBuffer;
@@ -994,17 +994,16 @@ IRHIComputeContext* FVulkanDynamicRHI::RHIGetCommandContext(ERHIPipeline Pipelin
 
 	FVulkanCommandListContext* CmdContext = Device->AcquireDeferredContext();
 
-	FVulkanCommandBufferManager* CmdMgr = CmdContext->GetCommandBufferManager();
-	FVulkanCommandBuffer* CmdBuffer = CmdMgr->GetActiveCmdBuffer();
+	FVulkanCommandBuffer* CmdBuffer = CmdContext->GetActiveCmdBuffer();
 	if (!CmdBuffer)
 	{
-		CmdMgr->PrepareForNewActiveCommandBuffer();
-		CmdBuffer = CmdMgr->GetActiveCmdBuffer();
+		CmdContext->PrepareForNewActiveCommandBuffer();
+		CmdBuffer = CmdContext->GetActiveCmdBuffer();
 	}
 	else if (CmdBuffer->IsSubmitted())
 	{
-		CmdMgr->PrepareForNewActiveCommandBuffer();
-		CmdBuffer = CmdMgr->GetActiveCmdBuffer();
+		CmdContext->PrepareForNewActiveCommandBuffer();
+		CmdBuffer = CmdContext->GetActiveCmdBuffer();
 	}
 	if (!CmdBuffer->HasBegun())
 	{
@@ -1031,24 +1030,23 @@ void FVulkanDynamicRHI::RHISubmitCommandLists(FRHISubmitCommandListsArgs&& Args)
 	{
 		FVulkanPlatformCommandList* PlatformCmdList = ResourceCast(Ptr);
 		FVulkanCommandListContext* CurrentContext = PlatformCmdList->CmdContext;
-		FVulkanCommandBufferManager* CmdBufMgr = CurrentContext->GetCommandBufferManager();
 
 		if (CurrentContext->IsImmediate())
 		{
 			CurrentContext->RequestSubmitCurrentCommands();
-			FVulkanCommandBuffer* CmdBuffer = CmdBufMgr->GetActiveCmdBuffer();
+			FVulkanCommandBuffer* CmdBuffer = CurrentContext->GetActiveCmdBuffer();
 			if (CmdBuffer && CmdBuffer->HasBegun() && CmdBuffer->IsOutsideRenderPass())
 			{
 				CurrentContext->SafePointSubmit();
 			}
-			CmdBufMgr->RefreshFenceStatus();
+			CurrentContext->RefreshFenceStatus();
 		}
 		else
 		{
-			check(!CmdBufMgr->HasPendingUploadCmdBuffer());  // todo-jn
-			FVulkanCommandBuffer* CmdBuffer = CmdBufMgr->GetActiveCmdBuffer();
+			check(!CurrentContext->HasPendingUploadCmdBuffer());  // todo-jn
+			FVulkanCommandBuffer* CmdBuffer = CurrentContext->GetActiveCmdBuffer();
 			check(!CmdBuffer->IsInsideRenderPass());
-			CmdBufMgr->SubmitActiveCmdBuffer();
+			CurrentContext->SubmitActiveCmdBuffer();
 
 			Device->ReleaseDeferredContext(CurrentContext);
 		}
diff --git a/Engine/Source/Runtime/VulkanRHI/Private/VulkanContext.cpp b/Engine/Source/Runtime/VulkanRHI/Private/VulkanContext.cpp
new file mode 100644
--- /dev/null
+++ b/Engine/Source/Runtime/VulkanRHI/Private/VulkanContext.cpp
@@ -0,0 +1,458 @@
+// Copyright Epic Games, Inc. All Rights Reserved.
+
+#include "VulkanContext.h"
+#include "VulkanCommandBuffer.h"
+#include "VulkanMemory.h"
+#include "VulkanResources.h"
+#include "VulkanRHIPrivate.h"
+#include "VulkanDevice.h"
+#include "VulkanPendingState.h"
+
+static int32 GVulkanUploadCmdBufferSemaphore = 0;	// despite this behavior being more correct, default to off to avoid changing the existing (as of UE 4.24) behavior
+static FAutoConsoleVariableRef CVarVulkanPreventOverlapWithUpload(
+	TEXT("r.Vulkan.UploadCmdBufferSemaphore"),
+	GVulkanUploadCmdBufferSemaphore,
+	TEXT("Whether command buffers for uploads and graphics can be executed simultaneously.\n")
+	TEXT(" 0: The buffers are submitted without any synch(default)\n")
+	TEXT(" 1: Graphics buffers will not overlap with the upload buffers"),
+	ECVF_ReadOnly
+);
+
+int32 GVulkanTempBlockSize = 4 * 1024 * 1024;
+static FAutoConsoleVariableRef GCVarVulkanTempBlockSize(
+	TEXT("r.Vulkan.TempBlockSize"),
+	GVulkanTempBlockSize,
+	TEXT("Size of the temporary blocks allocate by contexts, used for single use ub allocs and copies (default: 4MB)."),
+	ECVF_ReadOnly
+);
+
+
+
+
+FVulkanContextCommon::FVulkanContextCommon(FVulkanDevice& InDevice, FVulkanQueue& InQueue)
+	: Device(InDevice)
+	, Queue(InQueue)
+	, Pool(InDevice, *this)
+{
+	Pool.Create(Queue.GetFamilyIndex());
+
+	if (GVulkanUploadCmdBufferSemaphore)
+	{
+		ActiveCmdBufferSemaphore = new VulkanRHI::FSemaphore(Device);
+	}
+
+	ActiveCmdBuffer = Pool.Create(false);
+	ActiveCmdBuffer->InitializeTimings(this);
+	ActiveCmdBuffer->Begin();
+}
+
+FVulkanContextCommon::~FVulkanContextCommon()
+{
+}
+
+void FVulkanContextCommon::WaitForCmdBuffer(FVulkanCommandBuffer* CmdBuffer, float TimeInSecondsToWait)
+{
+	FScopeLock ScopeLock(&Pool.CS);
+	check(CmdBuffer->IsSubmitted());
+	bool bSuccess = Device.GetFenceManager().WaitForFence(CmdBuffer->Fence, (uint64)(TimeInSecondsToWait * 1e9));
+	check(bSuccess);
+	CmdBuffer->RefreshFenceStatus();
+}
+
+void FVulkanContextCommon::AddQueryPoolForReset(VkQueryPool QueryPool, uint32 Size)
+{
+	FScopeLock ScopeLock(&Pool.CS);
+	FQueryPoolReset QPR = { QueryPool, Size };
+	PoolResets.Add(QPR);
+}
+
+void FVulkanContextCommon::FlushResetQueryPools()
+{
+	FScopeLock ScopeLock(&Pool.CS);
+	if (PoolResets.Num())
+	{
+		FVulkanCommandBuffer* PoolResetCommandBuffer = nullptr;
+		for (int32 Index = 0; Index < Pool.CmdBuffers.Num(); ++Index)
+		{
+			FVulkanCommandBuffer* CmdBuffer = Pool.CmdBuffers[Index];
+			CmdBuffer->RefreshFenceStatus();
+#if VULKAN_USE_DIFFERENT_POOL_CMDBUFFERS
+			if (!CmdBuffer->bIsUploadOnly)
+#endif
+			{
+				if (CmdBuffer->State == FVulkanCommandBuffer::EState::ReadyForBegin)
+				{
+					PoolResetCommandBuffer = CmdBuffer;
+				}
+			}
+		}
+		if (!PoolResetCommandBuffer)
+		{
+			PoolResetCommandBuffer = Pool.Create(false);
+		}
+		PoolResetCommandBuffer->Begin();
+
+		VkCommandBuffer CommandBuffer = PoolResetCommandBuffer->GetHandle();
+		for (FQueryPoolReset& QPR : PoolResets)
+		{
+			VulkanRHI::vkCmdResetQueryPool(CommandBuffer, QPR.Pool, 0, QPR.Size);
+		}
+		PoolResetCommandBuffer->End();
+		Queue.Submit(PoolResetCommandBuffer);
+		PoolResets.Empty();
+	}
+}
+
+void FVulkanContextCommon::SubmitUploadCmdBuffer(uint32 NumSignalSemaphores, VkSemaphore* SignalSemaphores)
+{
+	FScopeLock ScopeLock(&Pool.CS);
+	check(UploadCmdBuffer);
+	check(UploadCmdBuffer->CurrentDescriptorPoolSetContainer == nullptr);
+	if (!UploadCmdBuffer->IsSubmitted() && UploadCmdBuffer->HasBegun())
+	{
+		check(UploadCmdBuffer->IsOutsideRenderPass());
+
+		VulkanRHI::DebugHeavyWeightBarrier(UploadCmdBuffer->GetHandle(), 4);
+
+		UploadCmdBuffer->End();
+
+		if (GVulkanUploadCmdBufferSemaphore)
+		{
+			// Add semaphores associated with the recent active cmdbuf(s), if any. That will prevent
+			// the overlap, delaying execution of this cmdbuf until the graphics one(s) is complete.
+			for (VulkanRHI::FSemaphore* WaitForThis : RenderingCompletedSemaphores)
+			{
+				UploadCmdBuffer->AddWaitSemaphore(VK_PIPELINE_STAGE_TOP_OF_PIPE_BIT, WaitForThis);
+			}
+
+			if (NumSignalSemaphores == 0)
+			{
+				checkf(UploadCmdBufferSemaphore != nullptr, TEXT("FVulkanContextCommon::SubmitUploadCmdBuffer: upload command buffer does not have an associated completion semaphore."));
+				VkSemaphore Sema = UploadCmdBufferSemaphore->GetHandle();
+				Queue.Submit(UploadCmdBuffer, 1, &Sema);
+				UploadCompletedSemaphores.Add(UploadCmdBufferSemaphore);
+				UploadCmdBufferSemaphore = nullptr;
+			}
+			else
+			{
+				TArray<VkSemaphore, TInlineAllocator<16>> CombinedSemaphores;
+				CombinedSemaphores.Add(UploadCmdBufferSemaphore->GetHandle());
+				CombinedSemaphores.Append(SignalSemaphores, NumSignalSemaphores);
+				Queue.Submit(UploadCmdBuffer, CombinedSemaphores.Num(), CombinedSemaphores.GetData());
+			}
+			// the buffer will now hold on to the wait semaphores, so we can clear them here
+			RenderingCompletedSemaphores.Empty();
+		}
+		else
+		{
+			Queue.Submit(UploadCmdBuffer, NumSignalSemaphores, SignalSemaphores);
+		}
+
+		UploadCmdBuffer->SubmittedTime = FPlatformTime::Seconds();
+	}
+
+	UploadCmdBuffer = nullptr;
+}
+
+void FVulkanContextCommon::SubmitActiveCmdBuffer(TArrayView<VulkanRHI::FSemaphore*> SignalSemaphores)
+{
+	FScopeLock ScopeLock(&Pool.CS);
+	FlushResetQueryPools();
+	check(!UploadCmdBuffer);
+	check(ActiveCmdBuffer);
+
+	TArray<VkSemaphore, FConcurrentLinearArrayAllocator> SemaphoreHandles;
+	SemaphoreHandles.Reserve(SignalSemaphores.Num() + 1);
+	for (VulkanRHI::FSemaphore* Semaphore : SignalSemaphores)
+	{
+		SemaphoreHandles.Add(Semaphore->GetHandle());
+	}
+
+	if (!ActiveCmdBuffer->IsSubmitted() && ActiveCmdBuffer->HasBegun())
+	{
+		if (!ActiveCmdBuffer->IsOutsideRenderPass())
+		{
+			UE_LOG(LogVulkanRHI, Warning, TEXT("Forcing EndRenderPass() for submission"));
+			ActiveCmdBuffer->EndRenderPass();
+		}
+
+		VulkanRHI::DebugHeavyWeightBarrier(ActiveCmdBuffer->GetHandle(), 8);
+
+		ActiveCmdBuffer->End();
+
+		if (GVulkanUploadCmdBufferSemaphore)
+		{
+			checkf(ActiveCmdBufferSemaphore != nullptr, TEXT("FVulkanContextCommon::SubmitUploadCmdBuffer: graphics command buffer does not have an associated completion semaphore."));
+
+			// Add semaphores associated with the recent upload cmdbuf(s), if any. That will prevent
+			// the overlap, delaying execution of this cmdbuf until upload one(s) are complete.
+			for (VulkanRHI::FSemaphore* UploadCompleteSema : UploadCompletedSemaphores)
+			{
+				ActiveCmdBuffer->AddWaitSemaphore(VK_PIPELINE_STAGE_TOP_OF_PIPE_BIT, UploadCompleteSema);
+			}
+
+			SemaphoreHandles.Add(ActiveCmdBufferSemaphore->GetHandle());
+			Queue.Submit(ActiveCmdBuffer, SemaphoreHandles.Num(), SemaphoreHandles.GetData());
+
+			RenderingCompletedSemaphores.Add(ActiveCmdBufferSemaphore);
+			ActiveCmdBufferSemaphore = nullptr;
+
+			// the buffer will now hold on to the wait semaphores, so we can clear them here
+			UploadCompletedSemaphores.Empty();
+		}
+		else
+		{
+			Queue.Submit(ActiveCmdBuffer, SemaphoreHandles.Num(), SemaphoreHandles.GetData());
+		}
+		ActiveCmdBuffer->SubmittedTime = FPlatformTime::Seconds();
+	}
+
+	ActiveCmdBuffer = nullptr;
+	if (GVulkanUploadCmdBufferSemaphore && ActiveCmdBufferSemaphore != nullptr)
+	{
+		// most likely we didn't submit as it wasn't begun
+		check(ActiveCmdBufferSemaphore->GetRefCount() == 0);
+		delete ActiveCmdBufferSemaphore;
+		ActiveCmdBufferSemaphore = nullptr;
+	}
+}
+
+void FVulkanContextCommon::SubmitActiveCmdBufferFromPresent(VulkanRHI::FSemaphore* SignalSemaphore)
+{
+	if (GVulkanUploadCmdBufferSemaphore)
+	{
+		// unlike more advanced regular SACB(), this is just a wrapper
+		// around Queue.Submit() to avoid rewriting the logic in Present
+		if (SignalSemaphore)
+		{
+			VkSemaphore SignalThis[2] =
+			{
+				SignalSemaphore->GetHandle(),
+				ActiveCmdBufferSemaphore->GetHandle()
+			};
+
+			Queue.Submit(ActiveCmdBuffer, UE_ARRAY_COUNT(SignalThis), SignalThis);
+		}
+		else
+		{
+			VkSemaphore SignalThis = ActiveCmdBufferSemaphore->GetHandle();
+			Queue.Submit(ActiveCmdBuffer, 1, &SignalThis);
+		}
+
+		RenderingCompletedSemaphores.Add(ActiveCmdBufferSemaphore);
+		ActiveCmdBufferSemaphore = nullptr;
+	}
+	else
+	{
+		if (SignalSemaphore)
+		{
+			Queue.Submit(ActiveCmdBuffer, SignalSemaphore->GetHandle());
+		}
+		else
+		{
+			Queue.Submit(ActiveCmdBuffer);
+		}
+	}
+}
+
+#if VULKAN_DELETE_STALE_CMDBUFFERS
+struct FRHICommandFreeUnusedCmdBuffers final : public FRHICommand<FRHICommandFreeUnusedCmdBuffers>
+{
+	FVulkanCommandBufferPool* Pool;
+	FVulkanQueue* Queue;
+	bool bTrimMemory;
+
+	FRHICommandFreeUnusedCmdBuffers(FVulkanCommandBufferPool* InPool, FVulkanQueue* InQueue, bool bInTrimMemory)
+		: Pool(InPool)
+		, Queue(InQueue)
+		, bTrimMemory(bInTrimMemory)
+	{
+	}
+
+	void Execute(FRHICommandListBase& CmdList)
+	{
+		Pool->FreeUnusedCmdBuffers(Queue, bTrimMemory);
+	}
+};
+#endif
+
+void FVulkanContextCommon::FreeUnusedCmdBuffers(bool bTrimMemory)
+{
+#if VULKAN_DELETE_STALE_CMDBUFFERS
+	FRHICommandList& RHICmdList = FRHICommandListExecutor::GetImmediateCommandList();
+	if (!IsInRenderingThread() || (RHICmdList.Bypass() || !IsRunningRHIInSeparateThread()))
+	{
+		Pool.FreeUnusedCmdBuffers(&Queue, bTrimMemory);
+	}
+	else
+	{
+		check(IsInRenderingThread());
+		ALLOC_COMMAND_CL(RHICmdList, FRHICommandFreeUnusedCmdBuffers)(&Pool, &Queue, bTrimMemory);
+	}
+#endif
+}
+
+void FVulkanContextCommon::PrepareForNewActiveCommandBuffer()
+{
+	FScopeLock ScopeLock(&Pool.CS);
+	check(!UploadCmdBuffer);
+
+	if (GVulkanUploadCmdBufferSemaphore)
+	{
+		// create a completion semaphore
+		checkf(ActiveCmdBufferSemaphore == nullptr, TEXT("FVulkanContextCommon::PrepareForNewActiveCommandBuffer: a semaphore from the previous active buffer has been leaked."));
+		ActiveCmdBufferSemaphore = new VulkanRHI::FSemaphore(Device);
+	}
+
+	for (int32 Index = 0; Index < Pool.CmdBuffers.Num(); ++Index)
+	{
+		FVulkanCommandBuffer* CmdBuffer = Pool.CmdBuffers[Index];
+		CmdBuffer->RefreshFenceStatus();
+#if VULKAN_USE_DIFFERENT_POOL_CMDBUFFERS
+		if (!CmdBuffer->bIsUploadOnly)
+#endif
+		{
+			if (CmdBuffer->State == FVulkanCommandBuffer::EState::ReadyForBegin || CmdBuffer->State == FVulkanCommandBuffer::EState::NeedReset)
+			{
+				ActiveCmdBuffer = CmdBuffer;
+				ActiveCmdBuffer->Begin();
+				return;
+			}
+			else
+			{
+				check(CmdBuffer->State == FVulkanCommandBuffer::EState::Submitted);
+			}
+		}
+	}
+
+	// All cmd buffers are being executed still
+	ActiveCmdBuffer = Pool.Create(false);
+	ActiveCmdBuffer->Begin();
+}
+
+#if (RHI_NEW_GPU_PROFILER == 0)
+uint32 FVulkanContextCommon::CalculateGPUTime()
+{
+	uint32 Time = 0;
+	for (int32 Index = 0; Index < Pool.CmdBuffers.Num(); ++Index)
+	{
+		FVulkanCommandBuffer* CmdBuffer = Pool.CmdBuffers[Index];
+		if (CmdBuffer->HasValidTiming())
+		{
+			Time += CmdBuffer->Timing->GetTiming(false);
+		}
+	}
+	return Time;
+}
+#endif // (RHI_NEW_GPU_PROFILER == 0)
+
+FVulkanCommandBuffer* FVulkanContextCommon::GetUploadCmdBuffer()
+{
+	FScopeLock ScopeLock(&Pool.CS);
+	if (!UploadCmdBuffer)
+	{
+		if (GVulkanUploadCmdBufferSemaphore)
+		{
+			// create a completion semaphore
+			checkf(UploadCmdBufferSemaphore == nullptr, TEXT("FVulkanContextCommon::GetUploadCmdBuffer: a semaphore from the previous upload buffer has been leaked."));
+			UploadCmdBufferSemaphore = new VulkanRHI::FSemaphore(Device);
+		}
+
+		for (int32 Index = 0; Index < Pool.CmdBuffers.Num(); ++Index)
+		{
+			FVulkanCommandBuffer* CmdBuffer = Pool.CmdBuffers[Index];
+			CmdBuffer->RefreshFenceStatus();
+#if VULKAN_USE_DIFFERENT_POOL_CMDBUFFERS
+			if (CmdBuffer->bIsUploadOnly)
+#endif
+			{
+				if (CmdBuffer->State == FVulkanCommandBuffer::EState::ReadyForBegin || CmdBuffer->State == FVulkanCommandBuffer::EState::NeedReset)
+				{
+					UploadCmdBuffer = CmdBuffer;
+					UploadCmdBuffer->Begin();
+					return UploadCmdBuffer;
+				}
+			}
+		}
+
+		// All cmd buffers are being executed still
+		UploadCmdBuffer = Pool.Create(true);
+		UploadCmdBuffer->Begin();
+	}
+
+	return UploadCmdBuffer;
+}
+
+
+
+FVulkanCommandListContext::FVulkanCommandListContext(FVulkanDynamicRHI* InRHI, FVulkanDevice& InDevice, FVulkanQueue& InQueue, FVulkanCommandListContext* InImmediate)
+	: FVulkanContextCommon(InDevice, InQueue)
+	, RHI(InRHI)
+	, Immediate(InImmediate)
+	, bSubmitAtNextSafePoint(false)
+#if (RHI_NEW_GPU_PROFILER == 0)
+	, GpuProfiler(this, &InDevice)
+#endif
+{
+#if (RHI_NEW_GPU_PROFILER == 0)
+	FrameTiming = new FVulkanGPUTiming(this, &InDevice);
+	FrameTiming->Initialize();
+#endif
+
+	if (IsImmediate())
+	{
+		// Insert the Begin frame timestamp query. On EndDrawingViewport() we'll insert the End and immediately after a new Begin()
+		WriteBeginTimestamp(GetActiveCmdBuffer());
+
+		// Flush the cmd buffer immediately to ensure a valid
+		// 'Last submitted' cmd buffer exists at frame 0.
+		SubmitActiveCmdBuffer();
+		PrepareForNewActiveCommandBuffer();
+	}
+
+	// Create Pending state, contains pipeline states such as current shader and etc..
+	PendingGfxState = new FVulkanPendingGfxState(&Device, *this);
+	PendingComputeState = new FVulkanPendingComputeState(&Device, *this);
+
+	// Currently used for UB and copies
+	{
+		uint32 BlockAlignment = FMath::Max<uint32>(Device.GetLimits().minUniformBufferOffsetAlignment, 16u);
+		VkBufferUsageFlags BufferUsageFlags =
+			(Device.GetOptionalExtensions().HasBufferDeviceAddress ? VK_BUFFER_USAGE_SHADER_DEVICE_ADDRESS_BIT : 0) |
+			VK_BUFFER_USAGE_TRANSFER_SRC_BIT |
+			VK_BUFFER_USAGE_UNIFORM_BUFFER_BIT;
+		if (Device.GetOptionalExtensions().HasRayTracingPipeline)
+		{
+			BufferUsageFlags |= VK_BUFFER_USAGE_SHADER_BINDING_TABLE_BIT_KHR;
+			BlockAlignment = FMath::Max<uint32>(Device.GetOptionalExtensionProperties().RayTracingPipelineProps.shaderGroupBaseAlignment, BlockAlignment);
+		}
+		TempBlockAllocator = new VulkanRHI::FTempBlockAllocator(&InDevice, GVulkanTempBlockSize, BlockAlignment, BufferUsageFlags);
+	}
+
+	GlobalUniformBuffers.AddZeroed(FUniformBufferStaticSlotRegistry::Get().GetSlotCount());
+}
+
+void FVulkanCommandListContext::ReleasePendingState()
+{
+	delete PendingGfxState;
+	PendingGfxState = nullptr;
+
+	delete PendingComputeState;
+	PendingComputeState = nullptr;
+}
+
+FVulkanCommandListContext::~FVulkanCommandListContext()
+{
+	if (GSupportsTimestampRenderQueries)
+	{
+#if (RHI_NEW_GPU_PROFILER == 0)
+		FrameTiming->Release();
+		delete FrameTiming;
+		FrameTiming = nullptr;
+#endif
+	}
+
+	delete TempBlockAllocator;
+	delete PendingGfxState;
+	delete PendingComputeState;
+}
\ No newline at end of file
diff --git a/Engine/Source/Runtime/VulkanRHI/Private/VulkanContext.h b/Engine/Source/Runtime/VulkanRHI/Private/VulkanContext.h
--- a/Engine/Source/Runtime/VulkanRHI/Private/VulkanContext.h
+++ b/Engine/Source/Runtime/VulkanRHI/Private/VulkanContext.h
@@ -10,7 +10,6 @@
 #include "VulkanRHIPrivate.h"
 #include "VulkanGPUProfiler.h"
 
-class FVulkanCommandBufferManager;
 class FVulkanDevice;
 class FVulkanDynamicRHI;
 class FVulkanOcclusionQueryPool;
@@ -20,10 +19,146 @@ class FVulkanQueue;
 class FVulkanSwapChain;
 
 
-class FVulkanCommandListContext : public IRHICommandContext
+class FVulkanContextCommon
 {
 public:
-	FVulkanCommandListContext(FVulkanDynamicRHI* InRHI, FVulkanDevice* InDevice, FVulkanQueue* InQueue, FVulkanCommandListContext* InImmediate);
+	FVulkanContextCommon(FVulkanDevice& InDevice, FVulkanQueue& InQueue);
+	~FVulkanContextCommon();
+
+	FVulkanCommandBuffer* GetActiveCmdBuffer()
+	{
+		if (UploadCmdBuffer)
+		{
+			SubmitUploadCmdBuffer();
+		}
+
+		return ActiveCmdBuffer;
+	}
+
+	FVulkanCommandBuffer* GetActiveCmdBufferDirect()
+	{
+		return ActiveCmdBuffer;
+	}
+
+	FVulkanCommandBuffer& GetCommandBuffer()
+	{
+		FVulkanCommandBuffer* CommandBuffer = GetActiveCmdBuffer();
+		check(CommandBuffer);
+		return *CommandBuffer;
+	}
+
+	bool HasPendingUploadCmdBuffer() const
+	{
+		return UploadCmdBuffer != nullptr;
+	}
+
+	bool HasPendingActiveCmdBuffer() const
+	{
+		return ActiveCmdBuffer != nullptr;
+	}
+
+	VULKANRHI_API FVulkanCommandBuffer* GetUploadCmdBuffer();
+
+	VULKANRHI_API void SubmitUploadCmdBuffer(uint32 NumSignalSemaphores = 0, VkSemaphore* SignalSemaphores = nullptr);
+
+	void SubmitActiveCmdBuffer(TArrayView<VulkanRHI::FSemaphore*> SignalSemaphores);
+
+	void SubmitActiveCmdBuffer()
+	{
+		SubmitActiveCmdBuffer(MakeArrayView<VulkanRHI::FSemaphore*>(nullptr, 0));
+	}
+
+	void SubmitActiveCmdBuffer(VulkanRHI::FSemaphore* SignalSemaphore)
+	{
+		SubmitActiveCmdBuffer(MakeArrayView<VulkanRHI::FSemaphore*>(&SignalSemaphore, 1));
+	}
+
+	/** Regular SACB() expects not-ended and would rotate the command buffer immediately, but Present has a special logic */
+	void SubmitActiveCmdBufferFromPresent(VulkanRHI::FSemaphore* SignalSemaphore = nullptr);
+
+	void WaitForCmdBuffer(FVulkanCommandBuffer* CmdBuffer, float TimeInSecondsToWait = 10.0f);
+
+
+	void AddQueryPoolForReset(VkQueryPool Pool, uint32 Size);
+	void FlushResetQueryPools();
+
+	// Update the fences of all cmd buffers except SkipCmdBuffer
+	void RefreshFenceStatus(FVulkanCommandBuffer* SkipCmdBuffer = nullptr)
+	{
+		Pool.RefreshFenceStatus(SkipCmdBuffer);
+	}
+
+	void PrepareForNewActiveCommandBuffer();
+
+	VkCommandPool GetHandle() const
+	{
+		return Pool.GetHandle();
+	}
+
+#if (RHI_NEW_GPU_PROFILER == 0)
+	uint32 CalculateGPUTime();
+#endif
+
+	void FreeUnusedCmdBuffers(bool bTrimMemory);
+
+	void NotifyDeletedImage(VkImage Image)
+	{
+		if (UploadCmdBuffer)
+		{
+			UploadCmdBuffer->GetLayoutManager().NotifyDeletedImage(Image);
+		}
+		if (ActiveCmdBuffer)
+		{
+			ActiveCmdBuffer->GetLayoutManager().NotifyDeletedImage(Image);
+		}
+
+		Queue.NotifyDeletedImage(Image);
+	}
+
+	uint64 GetFrameCounter() const
+	{
+		return FrameCounter;
+	}
+
+	// Public immutable references, access them directly
+	FVulkanDevice& Device;
+	FVulkanQueue& Queue;
+
+protected:
+	// Number of times EndFrame() has been called on this context
+	uint64 FrameCounter = 0;
+
+private:
+	FVulkanCommandBufferPool Pool;
+	struct FQueryPoolReset
+	{
+		VkQueryPool Pool;
+		uint32 Size;
+	};
+	TArray<FQueryPoolReset> PoolResets;
+
+	FVulkanCommandBuffer* ActiveCmdBuffer = nullptr;
+	FVulkanCommandBuffer* UploadCmdBuffer = nullptr;
+
+	// This semaphore is used to prevent overlaps between the (current) graphics cmdbuf and next upload cmdbuf.
+	VulkanRHI::FSemaphore* ActiveCmdBufferSemaphore = nullptr;
+
+	// Holds semaphores associated with the recent upload cmdbuf(s) - waiting to be added to the next graphics cmdbuf as WaitSemaphores.
+	TArray<VulkanRHI::FSemaphore*> RenderingCompletedSemaphores;
+
+	// This semaphore is used to prevent overlaps between (current) upload cmdbuf and next graphics cmdbuf.
+	VulkanRHI::FSemaphore* UploadCmdBufferSemaphore = nullptr;
+
+	// Holds semaphores associated with the recent upload cmdbuf(s) - waiting to be added to the next graphics cmdbuf as WaitSemaphores.
+	TArray<VulkanRHI::FSemaphore*> UploadCompletedSemaphores;
+};
+
+
+
+class FVulkanCommandListContext : public FVulkanContextCommon, public IRHICommandContext
+{
+public:
+	FVulkanCommandListContext(FVulkanDynamicRHI* InRHI, FVulkanDevice& InDevice, FVulkanQueue& InQueue, FVulkanCommandListContext* InImmediate);
 	virtual ~FVulkanCommandListContext();
 
 	static inline FVulkanCommandListContext& GetVulkanContext(IRHICommandContext& CmdContext)
@@ -128,13 +263,6 @@ public:
 		uint32 NumBindings, const FRayTracingLocalShaderBindings* Bindings,
 		ERayTracingBindingType BindingType) final override;
 
-	inline FVulkanCommandBufferManager* GetCommandBufferManager()
-	{
-		return CommandBufferManager;
-	}
-
-	FVulkanCommandBuffer& GetCommandBuffer();
-
 	inline VulkanRHI::FTempBlockAllocator& GetTempBlockAllocator()
 	{
 		return *TempBlockAllocator;
@@ -158,12 +286,6 @@ public:
 		}
 	}
 
-	inline void NotifyDeletedImage(VkImage Image)
-	{
-		CommandBufferManager->NotifyDeletedImage(Image);
-		Queue->NotifyDeletedImage(Image);
-	}
-
 	inline FVulkanRenderPass* GetCurrentRenderPass()
 	{
 		return CurrentRenderPass;
@@ -174,16 +296,6 @@ public:
 		return CurrentFramebuffer;
 	}
 
-	inline uint64 GetFrameCounter() const
-	{
-		return FrameCounter;
-	}
-
-	inline FVulkanQueue* GetQueue()
-	{
-		return Queue;
-	}
-
 	void WriteBeginTimestamp(FVulkanCommandBuffer* CmdBuffer);
 	void WriteEndTimestamp(FVulkanCommandBuffer* CmdBuffer);
 
@@ -196,11 +308,6 @@ public:
 	}
 #endif
 
-	inline FVulkanDevice* GetDevice() const
-	{
-		return Device;
-	}
-
 	void BeginRecursiveCommand()
 	{
 		// Nothing to do
@@ -213,8 +320,6 @@ public:
 protected:
 	FVulkanDynamicRHI* RHI;
 	FVulkanCommandListContext* Immediate;
-	FVulkanDevice* Device;
-	FVulkanQueue* Queue;
 	bool bSubmitAtNextSafePoint;
 	bool bUniformBufferUploadRenderPassDirty = true;
 
@@ -225,8 +330,6 @@ protected:
 
 	TArray<FString> EventStack;
 
-	FVulkanCommandBufferManager* CommandBufferManager = nullptr;
-
 	FVulkanRenderPass* CurrentRenderPass = nullptr;
 	FVulkanFramebuffer* CurrentFramebuffer = nullptr;
 
@@ -277,9 +380,6 @@ private:
 	friend class FVulkanDevice;
 	friend class FVulkanDynamicRHI;
 
-	// Number of times EndFrame() has been called on this context
-	uint64 FrameCounter;
-
 #if RHI_NEW_GPU_PROFILER
 	void RegisterGPUWork(uint32 NumPrimitives = 0, uint32 NumVertices = 0)	{ checkNoEntry(); } // @todo - new gpu profiler
 	void RegisterGPUDispatch(FIntVector GroupCount)	                        { checkNoEntry(); } // @todo - new gpu profiler
@@ -302,28 +402,8 @@ private:
 class FVulkanCommandListContextImmediate : public FVulkanCommandListContext
 {
 public:
-	FVulkanCommandListContextImmediate(FVulkanDynamicRHI* InRHI, FVulkanDevice* InDevice, FVulkanQueue* InQueue);
-};
-
-#if 0 // @todo: RHI command list refactor - todo
-struct FVulkanCommandContextContainer : public IRHICommandContextContainer, public VulkanRHI::FDeviceChild
-{
-	FVulkanCommandListContext* CmdContext;
-
-	FVulkanCommandContextContainer(FVulkanDevice* InDevice);
-
-	virtual IRHICommandContext* GetContext() override final;
-	virtual void FinishContext() override final;
-	virtual void SubmitAndFreeContextContainer(int32 Index, int32 Num) override final;
-
-	/** Custom new/delete with recycling */
-	void* operator new(size_t Size);
-	void operator delete(void* RawMemory);
-
-private:
-	friend class FVulkanDevice;
+	FVulkanCommandListContextImmediate(FVulkanDynamicRHI* InRHI, FVulkanDevice& InDevice, FVulkanQueue& InQueue);
 };
-#endif
 
 struct FVulkanContextArray : public TRHIPipelineArray<FVulkanCommandListContext*>
 {
diff --git a/Engine/Source/Runtime/VulkanRHI/Private/VulkanDescriptorSets.cpp b/Engine/Source/Runtime/VulkanRHI/Private/VulkanDescriptorSets.cpp
--- a/Engine/Source/Runtime/VulkanRHI/Private/VulkanDescriptorSets.cpp
+++ b/Engine/Source/Runtime/VulkanRHI/Private/VulkanDescriptorSets.cpp
@@ -717,7 +717,7 @@ void FVulkanBindlessDescriptorManager::UpdateDescriptor(FRHIDescriptorHandle Des
 	else
 	{
 		check(!IsInRenderingThread() || FRHICommandListExecutor::GetImmediateCommandList().Bypass() || !IsRunningRHIInSeparateThread());
-		FVulkanCommandBuffer* CmdBuffer = Device->GetImmediateContext().GetCommandBufferManager()->GetActiveCmdBuffer();
+		FVulkanCommandBuffer* CmdBuffer = Device->GetImmediateContext().GetActiveCmdBuffer();
 
 		// :todo-jn:  Hack to avoid barriers/copies in renderpasses
 		if (CmdBuffer->IsInsideRenderPass())
diff --git a/Engine/Source/Runtime/VulkanRHI/Private/VulkanDevice.cpp b/Engine/Source/Runtime/VulkanRHI/Private/VulkanDevice.cpp
--- a/Engine/Source/Runtime/VulkanRHI/Private/VulkanDevice.cpp
+++ b/Engine/Source/Runtime/VulkanRHI/Private/VulkanDevice.cpp
@@ -1320,11 +1320,11 @@ void FVulkanDevice::InitGPU()
 	// always look in the saved directory (for the cache from previous run that wasn't moved over to stage directory)
 	CacheFilenames.Add(VulkanRHI::GetPipelineCacheFilename());
 
-	ImmediateContext = new FVulkanCommandListContextImmediate(GVulkanRHI, this, GfxQueue);
+	ImmediateContext = new FVulkanCommandListContextImmediate(GVulkanRHI, *this, *GfxQueue);
 
 	if (GfxQueue->GetFamilyIndex() != ComputeQueue->GetFamilyIndex() && GRHIAllowAsyncComputeCvar.GetValueOnAnyThread() != 0)
 	{
-		ComputeContext = new FVulkanCommandListContextImmediate(GVulkanRHI, this, ComputeQueue);
+		ComputeContext = new FVulkanCommandListContextImmediate(GVulkanRHI, *this, *ComputeQueue);
 	}
 	else
 	{
@@ -1337,7 +1337,7 @@ void FVulkanDevice::InitGPU()
 		int32 Num = FTaskGraphInterface::Get().GetNumWorkerThreads();
 		for (int32 Index = 0; Index < Num; Index++)
 		{
-			FVulkanCommandListContext* CmdContext = new FVulkanCommandListContext(GVulkanRHI, this, GfxQueue, ImmediateContext);
+			FVulkanCommandListContext* CmdContext = new FVulkanCommandListContext(GVulkanRHI, *this, *GfxQueue, ImmediateContext);
 			CommandContexts.Add(CmdContext);
 		}
 	}
@@ -1498,7 +1498,7 @@ void FVulkanDevice::WaitUntilIdle()
 	VERIFYVULKANRESULT(VulkanRHI::vkDeviceWaitIdle(Device));
 
 	//#todo-rco: Loop through all contexts!
-	GetImmediateContext().GetCommandBufferManager()->RefreshFenceStatus();
+	GetImmediateContext().RefreshFenceStatus();
 }
 
 const VkComponentMapping& FVulkanDevice::GetFormatComponentMapping(EPixelFormat UEFormat) const
@@ -1529,16 +1529,15 @@ void FVulkanDevice::PrepareForCPURead()
 
 void FVulkanDevice::SubmitCommands(FVulkanCommandListContext* Context)
 {
-	FVulkanCommandBufferManager* CmdMgr = Context->GetCommandBufferManager();
-	if (CmdMgr->HasPendingUploadCmdBuffer())
+	if (Context->HasPendingUploadCmdBuffer())
 	{
-		CmdMgr->SubmitUploadCmdBuffer();
+		Context->SubmitUploadCmdBuffer();
 	}
-	if (CmdMgr->HasPendingActiveCmdBuffer())
+	if (Context->HasPendingActiveCmdBuffer())
 	{
-		CmdMgr->SubmitActiveCmdBuffer();
+		Context->SubmitActiveCmdBuffer();
 	}
-	CmdMgr->PrepareForNewActiveCommandBuffer();
+	Context->PrepareForNewActiveCommandBuffer();
 }
 
 void FVulkanDevice::SubmitCommandsAndFlushGPU()
@@ -1592,7 +1591,7 @@ FVulkanCommandListContext* FVulkanDevice::AcquireDeferredContext()
 	FScopeLock Lock(&GContextCS);
 	if (CommandContexts.Num() == 0)
 	{
-		return new FVulkanCommandListContext(GVulkanRHI, this, GfxQueue, ImmediateContext);
+		return new FVulkanCommandListContext(GVulkanRHI, *this, *GfxQueue, ImmediateContext);
 	}
 	return CommandContexts.Pop(EAllowShrinking::No);
 }
diff --git a/Engine/Source/Runtime/VulkanRHI/Private/VulkanDevice.h b/Engine/Source/Runtime/VulkanRHI/Private/VulkanDevice.h
--- a/Engine/Source/Runtime/VulkanRHI/Private/VulkanDevice.h
+++ b/Engine/Source/Runtime/VulkanRHI/Private/VulkanDevice.h
@@ -475,7 +475,7 @@ public:
 
 	void SubmitCommandsAndFlushGPU();
 
-	FVulkanOcclusionQueryPool* AcquireOcclusionQueryPool(FVulkanCommandBufferManager* CommandBufferManager, uint32 NumQueries);
+	FVulkanOcclusionQueryPool* AcquireOcclusionQueryPool(FVulkanContextCommon* InCommonContext, uint32 NumQueries);
 	void ReleaseUnusedOcclusionQueryPools();
 
 	inline class FVulkanPipelineStateCacheManager* GetPipelineStateCache()
diff --git a/Engine/Source/Runtime/VulkanRHI/Private/VulkanGPUProfiler.h b/Engine/Source/Runtime/VulkanRHI/Private/VulkanGPUProfiler.h
--- a/Engine/Source/Runtime/VulkanRHI/Private/VulkanGPUProfiler.h
+++ b/Engine/Source/Runtime/VulkanRHI/Private/VulkanGPUProfiler.h
@@ -11,7 +11,7 @@
 #include "VulkanConfiguration.h"
 
 class FVulkanCommandBuffer;
-class FVulkanCommandListContext;
+class FVulkanContextCommon;
 class FVulkanDevice;
 class FVulkanTimingQueryPool;
 
@@ -20,11 +20,11 @@ class FVulkanTimingQueryPool;
 class FVulkanGPUTiming : public FGPUTiming
 {
 public:
-	FVulkanGPUTiming(FVulkanCommandListContext* InCmd, FVulkanDevice* InDevice)
+	FVulkanGPUTiming(FVulkanContextCommon* InContext, FVulkanDevice* InDevice)
 		: Device(InDevice)
+		, Context(InContext)
 		, bIsTiming(false)
 		, bEndTimestampIssued(false)
-		, CmdContext(InCmd)
 	{
 	}
 
@@ -84,7 +84,8 @@ private:
 	 */
 	static void PlatformStaticInitialize(void* UserData);
 
-	FVulkanDevice* Device;
+	FVulkanDevice* Device = nullptr;
+	FVulkanContextCommon* Context = nullptr;
 
 	/** Whether we are currently timing the GPU: between StartTiming() and EndTiming(). */
 	bool bIsTiming;
@@ -93,7 +94,6 @@ private:
 	uint64 PreviousTime = 0;
 	uint32 NumPendingQueries = 0;
 	TQueue<uint32> PendingQueries;
-	FVulkanCommandListContext* CmdContext;
 	FVulkanTimingQueryPool* Pool = nullptr;
 };
 
@@ -101,9 +101,9 @@ private:
 class FVulkanEventNode : public FGPUProfilerEventNode
 {
 public:
-	FVulkanEventNode(const TCHAR* InName, FGPUProfilerEventNode* InParent, FVulkanCommandListContext* InCmd, FVulkanDevice* InDevice) :
+	FVulkanEventNode(const TCHAR* InName, FGPUProfilerEventNode* InParent, FVulkanContextCommon* InContext, FVulkanDevice* InDevice) :
 		FGPUProfilerEventNode(InName, InParent),
-		Timing(InCmd, InDevice)
+		Timing(InContext, InDevice)
 	{
 		// Initialize Buffered timestamp queries 
 		Timing.Initialize();
@@ -139,8 +139,8 @@ class FVulkanEventNodeFrame : public FGPUProfilerEventNodeFrame
 {
 public:
 
-	FVulkanEventNodeFrame(FVulkanCommandListContext* InCmd, FVulkanDevice* InDevice)
-		: RootEventTiming(InCmd, InDevice)
+	FVulkanEventNodeFrame(FVulkanContextCommon* InContext, FVulkanDevice* InDevice)
+		: RootEventTiming(InContext, InDevice)
 	{
 		RootEventTiming.Initialize();
 	}
@@ -174,7 +174,7 @@ struct FVulkanGPUProfiler : public FGPUProfiler
 	/** GPU hitch profile histories */
 	TIndirectArray<FVulkanEventNodeFrame> GPUHitchEventNodeFrames;
 
-	FVulkanGPUProfiler(FVulkanCommandListContext* InCmd, FVulkanDevice* InDevice);
+	FVulkanGPUProfiler(FVulkanContextCommon* InContext, FVulkanDevice* InDevice);
 
 	virtual ~FVulkanGPUProfiler();
 
@@ -191,7 +191,7 @@ struct FVulkanGPUProfiler : public FGPUProfiler
 
 	bool bCommandlistSubmitted;
 	FVulkanDevice* Device;
-	FVulkanCommandListContext* CmdContext;
+	FVulkanContextCommon* CmdContext;
 
 #if VULKAN_SUPPORTS_GPU_CRASH_DUMPS
 	void PrepareCrashMarkerkBuffer();
diff --git a/Engine/Source/Runtime/VulkanRHI/Private/VulkanMemory.cpp b/Engine/Source/Runtime/VulkanRHI/Private/VulkanMemory.cpp
--- a/Engine/Source/Runtime/VulkanRHI/Private/VulkanMemory.cpp
+++ b/Engine/Source/Runtime/VulkanRHI/Private/VulkanMemory.cpp
@@ -4169,7 +4169,7 @@ namespace VulkanRHI
 
 						// Defrag works from the queue for layouts
 						// Only work with straightforward targets that are in a single layout
-						const FVulkanImageLayout* OriginalLayout = Context.GetQueue()->GetLayoutManager().GetFullLayout(Texture->Image);
+						const FVulkanImageLayout* OriginalLayout = Context.Queue.GetLayoutManager().GetFullLayout(Texture->Image);
 						if (OriginalLayout && OriginalLayout->AreAllSubresourcesSameLayout() && (OriginalLayout->MainLayout != VK_IMAGE_LAYOUT_UNDEFINED))
 						{
 							FVulkanAllocation Allocation;
diff --git a/Engine/Source/Runtime/VulkanRHI/Private/VulkanPendingState.cpp b/Engine/Source/Runtime/VulkanRHI/Private/VulkanPendingState.cpp
--- a/Engine/Source/Runtime/VulkanRHI/Private/VulkanPendingState.cpp
+++ b/Engine/Source/Runtime/VulkanRHI/Private/VulkanPendingState.cpp
@@ -329,13 +329,13 @@ FVulkanPendingComputeState::~FVulkanPendingComputeState()
 void FVulkanPendingComputeState::SetSRVForUBResource(uint32 DescriptorSet, uint32 BindingIndex, FVulkanShaderResourceView* SRV)
 {
 	check(SRV);
-	CurrentState->SetSRV(Context.GetCommandBufferManager()->GetActiveCmdBuffer(), true, DescriptorSet, BindingIndex, SRV);
+	CurrentState->SetSRV(Context.GetActiveCmdBuffer(), true, DescriptorSet, BindingIndex, SRV);
 }
 
 void FVulkanPendingComputeState::SetUAVForUBResource(uint32 DescriptorSet, uint32 BindingIndex, FVulkanUnorderedAccessView* UAV)
 {
 	check(UAV);
-	CurrentState->SetUAV(Context.GetCommandBufferManager()->GetActiveCmdBuffer(), true, DescriptorSet, BindingIndex, UAV);
+	CurrentState->SetUAV(Context.GetActiveCmdBuffer(), true, DescriptorSet, BindingIndex, UAV);
 }
 
 
@@ -574,13 +574,13 @@ void FVulkanPendingGfxState::UpdateInputAttachments(FVulkanFramebuffer* Framebuf
 void FVulkanPendingGfxState::SetSRVForUBResource(uint8 DescriptorSet, uint32 BindingIndex, FVulkanShaderResourceView* SRV)
 {
 	check(SRV);
-	CurrentState->SetSRV(Context.GetCommandBufferManager()->GetActiveCmdBuffer(), false, DescriptorSet, BindingIndex, SRV);
+	CurrentState->SetSRV(Context.GetActiveCmdBuffer(), false, DescriptorSet, BindingIndex, SRV);
 }
 
 void FVulkanPendingGfxState::SetUAVForUBResource(uint8 DescriptorSet, uint32 BindingIndex, FVulkanUnorderedAccessView* UAV)
 {
 	check(UAV);
-	CurrentState->SetUAV(Context.GetCommandBufferManager()->GetActiveCmdBuffer(), false, DescriptorSet, BindingIndex, UAV);
+	CurrentState->SetUAV(Context.GetActiveCmdBuffer(), false, DescriptorSet, BindingIndex, UAV);
 }
 
 int32 GDSetCacheTargetSetsPerPool = 4096;
diff --git a/Engine/Source/Runtime/VulkanRHI/Private/VulkanPipelineState.cpp b/Engine/Source/Runtime/VulkanRHI/Private/VulkanPipelineState.cpp
--- a/Engine/Source/Runtime/VulkanRHI/Private/VulkanPipelineState.cpp
+++ b/Engine/Source/Runtime/VulkanRHI/Private/VulkanPipelineState.cpp
@@ -161,7 +161,7 @@ static inline VulkanRHI::FVulkanAllocation UpdatePackedUniformBuffers(const FPac
 	const FPackedUniformBuffers::FPackedBuffer& StagedUniformBuffer = PackedUniformBuffers.GetBuffer();
 
 	const uint32 UBSize = (uint32)StagedUniformBuffer.Num();
-	const uint32 UBAlign = (uint32)InCmdBuffer->GetDevice()->GetLimits().minUniformBufferOffsetAlignment;
+	const uint32 UBAlign = (uint32)InCmdBuffer->Device.GetLimits().minUniformBufferOffsetAlignment;
 
 	VulkanRHI::FVulkanAllocation TempAllocation;
 	uint8* MappedPointer = InContext->GetTempBlockAllocator().Alloc(UBSize, UBAlign, InCmdBuffer, TempAllocation);
@@ -541,7 +541,7 @@ void FVulkanCommandListContext::RHISetGraphicsPipelineState(FRHIGraphicsPipeline
 {
 	FVulkanRHIGraphicsPipelineState* Pipeline = ResourceCast(GraphicsState);
 	
-	FVulkanPipelineStateCacheManager* PipelineStateCache = Device->GetPipelineStateCache();
+	FVulkanPipelineStateCacheManager* PipelineStateCache = Device.GetPipelineStateCache();
 	PipelineStateCache->LRUTouch(Pipeline);
 
 	FVulkanCommandBuffer& CommandBuffer = GetCommandBuffer();
@@ -575,20 +575,20 @@ void FVulkanCommandListContext::RHISetGraphicsPipelineState(FRHIGraphicsPipeline
 
 void FVulkanCommandListContext::RHISetComputePipelineState(FRHIComputePipelineState* ComputePipelineState)
 {
-	FVulkanCommandBuffer* CmdBuffer = CommandBufferManager->GetActiveCmdBuffer();
+	FVulkanCommandBuffer* CmdBuffer = GetActiveCmdBuffer();
 	if (CmdBuffer->IsInsideRenderPass())
 	{
 		if (GVulkanSubmitAfterEveryEndRenderPass)
 		{
-			CommandBufferManager->SubmitActiveCmdBuffer();
-			CommandBufferManager->PrepareForNewActiveCommandBuffer();
-			CmdBuffer = CommandBufferManager->GetActiveCmdBuffer();
+			SubmitActiveCmdBuffer();
+			PrepareForNewActiveCommandBuffer();
+			CmdBuffer = GetActiveCmdBuffer();
 		}
 	}
 
 	if (CmdBuffer->CurrentDescriptorPoolSetContainer == nullptr)
 	{
-		CmdBuffer->CurrentDescriptorPoolSetContainer = &Device->GetDescriptorPoolsManager().AcquirePoolSetContainer();
+		CmdBuffer->CurrentDescriptorPoolSetContainer = &Device.GetDescriptorPoolsManager().AcquirePoolSetContainer();
 	}
 
 	//#todo-rco: Set PendingGfx to null
diff --git a/Engine/Source/Runtime/VulkanRHI/Private/VulkanQuery.cpp b/Engine/Source/Runtime/VulkanRHI/Private/VulkanQuery.cpp
--- a/Engine/Source/Runtime/VulkanRHI/Private/VulkanQuery.cpp
+++ b/Engine/Source/Runtime/VulkanRHI/Private/VulkanQuery.cpp
@@ -37,8 +37,8 @@ constexpr int32 NUM_FRAMES_TO_WAIT_REUSE_POOL = 10;
 constexpr uint32 NUM_FRAMES_TO_WAIT_RELEASE_POOL = MAX_uint32; // never release
 #endif
 
-FVulkanQueryPool::FVulkanQueryPool(FVulkanDevice* InDevice, FVulkanCommandBufferManager* CommandBufferManager, uint32 InMaxQueries, VkQueryType InQueryType, bool bInShouldAddReset)
-	: VulkanRHI::FDeviceChild(InDevice)
+FVulkanQueryPool::FVulkanQueryPool(FVulkanDevice& InDevice, FVulkanContextCommon* InCommonContext, uint32 InMaxQueries, VkQueryType InQueryType, bool bInShouldAddReset)
+	: Device(InDevice)
 	, QueryPool(VK_NULL_HANDLE)
 	, ResetEvent(VK_NULL_HANDLE)
 	, MaxQueries(InMaxQueries)
@@ -50,21 +50,21 @@ FVulkanQueryPool::FVulkanQueryPool(FVulkanDevice* InDevice, FVulkanCommandBuffer
 	PoolCreateInfo.queryType = QueryType;
 	PoolCreateInfo.queryCount = MaxQueries;
 
-	VERIFYVULKANRESULT(VulkanRHI::vkCreateQueryPool(Device->GetInstanceHandle(), &PoolCreateInfo, VULKAN_CPU_ALLOCATOR, &QueryPool));
+	VERIFYVULKANRESULT(VulkanRHI::vkCreateQueryPool(Device.GetInstanceHandle(), &PoolCreateInfo, VULKAN_CPU_ALLOCATOR, &QueryPool));
 	
 	// If host query resets are supported, reset all new query pools on creation
-	if (Device->GetOptionalExtensions().HasEXTHostQueryReset)
+	if (Device.GetOptionalExtensions().HasEXTHostQueryReset)
 	{
-		VulkanRHI::vkResetQueryPoolEXT(Device->GetInstanceHandle(), QueryPool, 0, MaxQueries);
+		VulkanRHI::vkResetQueryPoolEXT(Device.GetInstanceHandle(), QueryPool, 0, MaxQueries);
 	}
-	else if (bInShouldAddReset && CommandBufferManager)
+	else if (bInShouldAddReset && InCommonContext)
 	{
-		CommandBufferManager->AddQueryPoolForReset(QueryPool, InMaxQueries);
+		InCommonContext->AddQueryPoolForReset(QueryPool, InMaxQueries);
 	}
 
 	VkEventCreateInfo EventCreateInfo;
 	ZeroVulkanStruct(EventCreateInfo, VK_STRUCTURE_TYPE_EVENT_CREATE_INFO);
-	VERIFYVULKANRESULT(VulkanRHI::vkCreateEvent(Device->GetInstanceHandle(), &EventCreateInfo, VULKAN_CPU_ALLOCATOR, &ResetEvent));
+	VERIFYVULKANRESULT(VulkanRHI::vkCreateEvent(Device.GetInstanceHandle(), &EventCreateInfo, VULKAN_CPU_ALLOCATOR, &ResetEvent));
 	
 	QueryOutput.AddZeroed(MaxQueries);
 }
@@ -72,8 +72,8 @@ FVulkanQueryPool::FVulkanQueryPool(FVulkanDevice* InDevice, FVulkanCommandBuffer
 FVulkanQueryPool::~FVulkanQueryPool()
 {
 	DEC_DWORD_STAT(STAT_VulkanNumQueryPools);
-	VulkanRHI::vkDestroyQueryPool(Device->GetInstanceHandle(), QueryPool, VULKAN_CPU_ALLOCATOR);
-	VulkanRHI::vkDestroyEvent(Device->GetInstanceHandle(), ResetEvent, VULKAN_CPU_ALLOCATOR);
+	VulkanRHI::vkDestroyQueryPool(Device.GetInstanceHandle(), QueryPool, VULKAN_CPU_ALLOCATOR);
+	VulkanRHI::vkDestroyEvent(Device.GetInstanceHandle(), ResetEvent, VULKAN_CPU_ALLOCATOR);
 	
 	QueryPool = VK_NULL_HANDLE;
 	ResetEvent = VK_NULL_HANDLE;
@@ -102,9 +102,9 @@ bool FVulkanOcclusionQueryPool::InternalTryGetResults(bool bWait)
 	check(State == EState::RHIT_PostEndBatch);
 
 	VkResult Result = VK_NOT_READY;
-	if (VulkanRHI::vkGetEventStatus(Device->GetInstanceHandle(), ResetEvent) == VK_EVENT_SET)
+	if (VulkanRHI::vkGetEventStatus(Device.GetInstanceHandle(), ResetEvent) == VK_EVENT_SET)
 	{
-		Result = VulkanRHI::vkGetQueryPoolResults(Device->GetInstanceHandle(), QueryPool, 0, NumUsedQueries, NumUsedQueries * sizeof(uint64), QueryOutput.GetData(), sizeof(uint64), VK_QUERY_RESULT_64_BIT);
+		Result = VulkanRHI::vkGetQueryPoolResults(Device.GetInstanceHandle(), QueryPool, 0, NumUsedQueries, NumUsedQueries * sizeof(uint64), QueryOutput.GetData(), sizeof(uint64), VK_QUERY_RESULT_64_BIT);
 		if (Result == VK_SUCCESS)
 		{
 			State = EState::RT_PostGetResults;
@@ -135,10 +135,10 @@ bool FVulkanOcclusionQueryPool::InternalTryGetResults(bool bWait)
 					FTaskGraphInterface::Get().ProcessThreadUntilIdle(RenderThread_Local);
 				}
 
-				Result = VulkanRHI::vkGetEventStatus(Device->GetInstanceHandle(), ResetEvent);
+				Result = VulkanRHI::vkGetEventStatus(Device.GetInstanceHandle(), ResetEvent);
 				if (Result == VK_EVENT_SET)
 				{
-					Result = VulkanRHI::vkGetQueryPoolResults(Device->GetInstanceHandle(), QueryPool, 0, NumUsedQueries, NumUsedQueries * sizeof(uint64), QueryOutput.GetData(), sizeof(uint64), VK_QUERY_RESULT_64_BIT);
+					Result = VulkanRHI::vkGetQueryPoolResults(Device.GetInstanceHandle(), QueryPool, 0, NumUsedQueries, NumUsedQueries * sizeof(uint64), QueryOutput.GetData(), sizeof(uint64), VK_QUERY_RESULT_64_BIT);
 				}
 
 				if (Result == VK_SUCCESS)
@@ -210,11 +210,11 @@ void FVulkanOcclusionQueryPool::Reset(FVulkanCommandBuffer* InCmdBuffer, uint32
 /*
 	bHasResults = false;
 */
-	VulkanRHI::vkResetEvent(Device->GetInstanceHandle(), ResetEvent);
+	VulkanRHI::vkResetEvent(Device.GetInstanceHandle(), ResetEvent);
 	VulkanRHI::vkCmdResetQueryPool(InCmdBuffer->GetHandle(), QueryPool, 0, MaxQueries);
 	
 	// workaround for apparent cache-flush bug in AMD driver implementation of vkCmdResetQueryPool
-	if (Device->GetVendorId() == EGpuVendorId::Amd)
+	if (Device.GetVendorId() == EGpuVendorId::Amd)
 	{
 		VkMemoryBarrier barrier;
 		ZeroVulkanStruct(barrier, VK_STRUCTURE_TYPE_MEMORY_BARRIER);
@@ -236,7 +236,7 @@ void FVulkanCommandListContext::BeginOcclusionQueryBatch(FVulkanCommandBuffer* C
 	FScopeLock ScopeLock(&GOcclusionQueryCS);
 */
 	checkf(!CurrentOcclusionQueryPool, TEXT("BeginOcclusionQueryBatch called without corresponding EndOcclusionQueryBatch!"));
-	CurrentOcclusionQueryPool = Device->AcquireOcclusionQueryPool(GetCommandBufferManager(), NumQueriesInBatch);
+	CurrentOcclusionQueryPool = Device.AcquireOcclusionQueryPool(this, NumQueriesInBatch);
 	ensure(CmdBuffer->IsOutsideRenderPass());
 	CurrentOcclusionQueryPool->Reset(CmdBuffer, GFrameNumberRenderThread);
 }
@@ -272,7 +272,7 @@ void FVulkanOcclusionQueryPool::FlushAllocatedQueries()
 	}
 }
 
-FVulkanOcclusionQueryPool* FVulkanDevice::AcquireOcclusionQueryPool(FVulkanCommandBufferManager* CommandBufferManager, uint32 NumQueries)
+FVulkanOcclusionQueryPool* FVulkanDevice::AcquireOcclusionQueryPool(FVulkanContextCommon* InCommonContext, uint32 NumQueries)
 {
 	// At least add one query
 	NumQueries = FMath::Max(1u, NumQueries);
@@ -321,7 +321,7 @@ FVulkanOcclusionQueryPool* FVulkanDevice::AcquireOcclusionQueryPool(FVulkanComma
 		}
 	}
 
-	FVulkanOcclusionQueryPool* Pool = new FVulkanOcclusionQueryPool(this, CommandBufferManager, NumQueries);
+	FVulkanOcclusionQueryPool* Pool = new FVulkanOcclusionQueryPool(*this, InCommonContext, NumQueries);
 	UsedOcclusionQueryPools.Add(Pool);
 	return Pool;
 }
@@ -398,7 +398,7 @@ void FVulkanCommandListContext::ReadAndCalculateGPUFrameTime()
 		// If we support profile command buffers then use this timing for GPU
 		if (GVulkanUseCmdBufferTimingForGPUTime)
 		{
-			Delta = CommandBufferManager->CalculateGPUTime();
+			Delta = CalculateGPUTime();
 		}
 		else
 		{
@@ -542,7 +542,7 @@ bool FVulkanDynamicRHI::RHIGetRenderQueryResult(FRHIRenderQuery* QueryRHI, uint6
 					// CPU wait for query results to be ready.
 					if (bWaitForStart && StartQuerySyncPoint.FenceCounter == StartQuerySyncPoint.CmdBuffer->GetFenceSignaledCounter())
 					{
-						Device->GetImmediateContext().GetCommandBufferManager()->WaitForCmdBuffer(StartQuerySyncPoint.CmdBuffer);
+						Device->GetImmediateContext().WaitForCmdBuffer(StartQuerySyncPoint.CmdBuffer);
 					}
 				}
 
@@ -617,8 +617,8 @@ void FVulkanCommandListContext::RHIEndRenderQuery(FRHIRenderQuery* QueryRHI)
 
 		if (Query->Pool == nullptr)
 		{
-			Query->Pool = new FVulkanTimingQueryPool(Device, CommandBufferManager, 4);
-			Query->Pool->ResultsBuffer = Device->GetStagingManager().AcquireBuffer(Query->Pool->BufferSize * sizeof(uint64), VK_BUFFER_USAGE_TRANSFER_DST_BIT, VK_MEMORY_PROPERTY_HOST_COHERENT_BIT);
+			Query->Pool = new FVulkanTimingQueryPool(Device, this, 4);
+			Query->Pool->ResultsBuffer = Device.GetStagingManager().AcquireBuffer(Query->Pool->BufferSize * sizeof(uint64), VK_BUFFER_USAGE_TRANSFER_DST_BIT, VK_MEMORY_PROPERTY_HOST_COHERENT_BIT);
 			Query->Pool->MappedPointer = (uint64*)Query->Pool->ResultsBuffer->GetMappedPointer();
 		}
 
@@ -638,9 +638,9 @@ void FVulkanCommandListContext::RHIEndRenderQuery(FRHIRenderQuery* QueryRHI)
 #if (RHI_NEW_GPU_PROFILER == 0)
 void FVulkanCommandListContext::RHICalibrateTimers(FRHITimestampCalibrationQuery* CalibrationQuery)
 {
-	if (Device->GetOptionalExtensions().HasEXTCalibratedTimestamps)
+	if (Device.GetOptionalExtensions().HasEXTCalibratedTimestamps)
 	{
-		FGPUTimingCalibrationTimestamp CalibrationTimestamp = Device->GetCalibrationTimestamp();
+		FGPUTimingCalibrationTimestamp CalibrationTimestamp = Device.GetCalibrationTimestamp();
 		CalibrationQuery->CPUMicroseconds[0] = CalibrationTimestamp.CPUMicroseconds;
 		CalibrationQuery->GPUMicroseconds[0] = CalibrationTimestamp.GPUMicroseconds;
 	}
@@ -670,7 +670,7 @@ FVulkanTimingQuery::~FVulkanTimingQuery()
 {
 	if (Pool != nullptr)
 	{
-		Pool->GetParent()->GetStagingManager().ReleaseBuffer(nullptr, Pool->ResultsBuffer);
+		Pool->GetDevice().GetStagingManager().ReleaseBuffer(nullptr, Pool->ResultsBuffer);
 		delete Pool;
 	}
 }
diff --git a/Engine/Source/Runtime/VulkanRHI/Private/VulkanQueue.cpp b/Engine/Source/Runtime/VulkanRHI/Private/VulkanQueue.cpp
--- a/Engine/Source/Runtime/VulkanRHI/Private/VulkanQueue.cpp
+++ b/Engine/Source/Runtime/VulkanRHI/Private/VulkanQueue.cpp
@@ -82,7 +82,7 @@ void FVulkanQueue::Submit(FVulkanCommandBuffer* CmdBuffer, uint32 NumSignalSemap
 
 	if (GWaitForIdleOnSubmit != 0)
 	{
-		FVulkanCommandBufferManager* CmdBufferMgr = &CmdBuffer->GetOwner()->GetMgr();
+		FVulkanContextCommon& Context = CmdBuffer->GetOwner()->GetContext();
 
 		switch(GWaitForIdleOnSubmit)
 		{
@@ -93,11 +93,11 @@ void FVulkanQueue::Submit(FVulkanCommandBuffer* CmdBuffer, uint32 NumSignalSemap
 				break;
 
 			case 2:
-				bShouldStall = (CmdBufferMgr->HasPendingUploadCmdBuffer() && CmdBufferMgr->GetUploadCmdBuffer() == CmdBuffer);
+				bShouldStall = (Context.HasPendingUploadCmdBuffer() && Context.GetUploadCmdBuffer() == CmdBuffer);
 				break;
 
 			case 3:
-				bShouldStall = (CmdBufferMgr->HasPendingActiveCmdBuffer() && CmdBufferMgr->GetActiveCmdBufferDirect() == CmdBuffer);
+				bShouldStall = (Context.HasPendingActiveCmdBuffer() && Context.GetActiveCmdBufferDirect() == CmdBuffer);
 				break;
 		}
 	}
diff --git a/Engine/Source/Runtime/VulkanRHI/Private/VulkanRHI.cpp b/Engine/Source/Runtime/VulkanRHI/Private/VulkanRHI.cpp
--- a/Engine/Source/Runtime/VulkanRHI/Private/VulkanRHI.cpp
+++ b/Engine/Source/Runtime/VulkanRHI/Private/VulkanRHI.cpp
@@ -105,14 +105,6 @@ static TAutoConsoleVariable<bool> CVarEnableVulkanPSOFileCacheWhenPrecachingActi
 	TEXT("true: Allow both PSO file cache and precaching."),
 	ECVF_RenderThreadSafe | ECVF_ReadOnly);
 
-int32 GVulkanTempBlockSize = 4 * 1024 * 1024;
-static FAutoConsoleVariableRef GCVarVulkanTempBlockSize(
-	TEXT("r.Vulkan.TempBlockSize"),
-	GVulkanTempBlockSize,
-	TEXT("Size of the temporary blocks allocate by contexts, used for single use ub allocs and copies (default: 4MB)."),
-	ECVF_ReadOnly
-);
-
 int32 GVulkanAMDCompatibilityMode = 1;
 static FAutoConsoleVariableRef GCVarVulkanAMDCompatibilityMode(
 	TEXT("r.Vulkan.AMDCompatibilityMode"),
@@ -429,100 +421,7 @@ FDynamicRHI* FVulkanDynamicRHIModule::CreateRHI(ERHIFeatureLevel::Type InRequest
 IMPLEMENT_MODULE(FVulkanDynamicRHIModule, VulkanRHI);
 
 
-FVulkanCommandListContext::FVulkanCommandListContext(FVulkanDynamicRHI* InRHI, FVulkanDevice* InDevice, FVulkanQueue* InQueue, FVulkanCommandListContext* InImmediate)
-	: RHI(InRHI)
-	, Immediate(InImmediate)
-	, Device(InDevice)
-	, Queue(InQueue)
-	, bSubmitAtNextSafePoint(false)
-	, FrameCounter(0)
-#if (RHI_NEW_GPU_PROFILER == 0)
-	, GpuProfiler(this, InDevice)
-#endif
-{
-#if (RHI_NEW_GPU_PROFILER == 0)
-	FrameTiming = new FVulkanGPUTiming(this, InDevice);
-#endif
-
-	// Create CommandBufferManager, contain all active buffers
-	CommandBufferManager = new FVulkanCommandBufferManager(InDevice, this);
-	CommandBufferManager->Init(this);
-
-#if (RHI_NEW_GPU_PROFILER == 0)
-	FrameTiming->Initialize();
-#endif
-
-	if (IsImmediate())
-	{
-		// Insert the Begin frame timestamp query. On EndDrawingViewport() we'll insert the End and immediately after a new Begin()
-		WriteBeginTimestamp(CommandBufferManager->GetActiveCmdBuffer());
-
-		// Flush the cmd buffer immediately to ensure a valid
-		// 'Last submitted' cmd buffer exists at frame 0.
-		CommandBufferManager->SubmitActiveCmdBuffer();
-		CommandBufferManager->PrepareForNewActiveCommandBuffer();
-	}
-
-	// Create Pending state, contains pipeline states such as current shader and etc..
-	PendingGfxState = new FVulkanPendingGfxState(Device, *this);
-	PendingComputeState = new FVulkanPendingComputeState(Device, *this);
-
-	// Currently used for UB and copies
-	{
-		uint32 BlockAlignment = FMath::Max<uint32>(Device->GetLimits().minUniformBufferOffsetAlignment, 16u);
-		VkBufferUsageFlags BufferUsageFlags =
-			(Device->GetOptionalExtensions().HasBufferDeviceAddress ? VK_BUFFER_USAGE_SHADER_DEVICE_ADDRESS_BIT : 0) |
-			VK_BUFFER_USAGE_TRANSFER_SRC_BIT |
-			VK_BUFFER_USAGE_UNIFORM_BUFFER_BIT;
-		if (Device->GetOptionalExtensions().HasRayTracingPipeline)
-		{
-			BufferUsageFlags |= VK_BUFFER_USAGE_SHADER_BINDING_TABLE_BIT_KHR;
-			BlockAlignment = FMath::Max<uint32>(Device->GetOptionalExtensionProperties().RayTracingPipelineProps.shaderGroupBaseAlignment, BlockAlignment);
-		}
-		TempBlockAllocator = new VulkanRHI::FTempBlockAllocator(InDevice, GVulkanTempBlockSize, BlockAlignment, BufferUsageFlags);
-	}
-
-	GlobalUniformBuffers.AddZeroed(FUniformBufferStaticSlotRegistry::Get().GetSlotCount());
-}
-
-FVulkanCommandBuffer& FVulkanCommandListContext::GetCommandBuffer()
-{
-	FVulkanCommandBuffer* CommandBuffer = CommandBufferManager->GetActiveCmdBuffer();
-	check(CommandBuffer);
-	return *CommandBuffer;
-}
-
-void FVulkanCommandListContext::ReleasePendingState()
-{
-	delete PendingGfxState;
-	PendingGfxState = nullptr;
-	
-	delete PendingComputeState;
-	PendingComputeState = nullptr;
-}
-
-FVulkanCommandListContext::~FVulkanCommandListContext()
-{
-	if (GSupportsTimestampRenderQueries)
-	{
-#if (RHI_NEW_GPU_PROFILER == 0)
-		FrameTiming->Release();
-		delete FrameTiming;
-		FrameTiming = nullptr;
-#endif
-	}
-
-	check(CommandBufferManager != nullptr);
-	delete CommandBufferManager;
-	CommandBufferManager = nullptr;
-
-	delete TempBlockAllocator;
-	delete PendingGfxState;
-	delete PendingComputeState;
-}
-
-
-FVulkanCommandListContextImmediate::FVulkanCommandListContextImmediate(FVulkanDynamicRHI* InRHI, FVulkanDevice* InDevice, FVulkanQueue* InQueue)
+FVulkanCommandListContextImmediate::FVulkanCommandListContextImmediate(FVulkanDynamicRHI* InRHI, FVulkanDevice& InDevice, FVulkanQueue& InQueue)
 	: FVulkanCommandListContext(InRHI, InDevice, InQueue, nullptr)
 {
 }
@@ -1163,23 +1062,23 @@ void FVulkanDynamicRHI::RHIEndFrame_RenderThread(FRHICommandListImmediate& RHICm
 #endif
 
 		bool bTrimMemory = false;
-		Context.GetCommandBufferManager()->FreeUnusedCmdBuffers(bTrimMemory);
+		Context.FreeUnusedCmdBuffers(bTrimMemory);
 
-		Context.Device->GetStagingManager().ProcessPendingFree(false, true);
-		Context.Device->GetMemoryManager().ReleaseFreedPages(Context);
-		Context.Device->GetDeferredDeletionQueue().ReleaseResources();
+		Context.Device.GetStagingManager().ProcessPendingFree(false, true);
+		Context.Device.GetMemoryManager().ReleaseFreedPages(Context);
+		Context.Device.GetDeferredDeletionQueue().ReleaseResources();
 
 		if (UseVulkanDescriptorCache())
 		{
-			Context.Device->GetDescriptorSetCache().GC();
+			Context.Device.GetDescriptorSetCache().GC();
 		}
-		Context.Device->GetDescriptorPoolsManager().GC();
+		Context.Device.GetDescriptorPoolsManager().GC();
 
-		Context.Device->ReleaseUnusedOcclusionQueryPools();
+		Context.Device.ReleaseUnusedOcclusionQueryPools();
 
-		Context.Device->GetPipelineStateCache()->TickLRU();
+		Context.Device.GetPipelineStateCache()->TickLRU();
 
-		Context.Device->GetBindlessDescriptorManager()->UpdateUBAllocator();
+		Context.Device.GetBindlessDescriptorManager()->UpdateUBAllocator();
 		Context.GetTempBlockAllocator().UpdateBlocks();
 
 		++Context.FrameCounter;
@@ -1203,7 +1102,7 @@ void FVulkanDynamicRHI::RHIEndFrame_RenderThread(FRHICommandListImmediate& RHICm
 
 		if (GRHISupportsRayTracing)
 		{
-			Context.Device->GetRayTracingCompactionRequestHandler()->Update(Context);
+			Context.Device.GetRayTracingCompactionRequestHandler()->Update(Context);
 		}
 	});
 }
@@ -1245,7 +1144,7 @@ void FVulkanCommandListContext::RHIEndDrawingViewport(FRHIViewport* ViewportRHI,
 
 	WriteEndTimestamp(&CommandBuffer);
 
-	bool bNativePresent = Viewport->Present(this, &CommandBuffer, Queue, Device->GetPresentQueue(), bLockToVsync);
+	bool bNativePresent = Viewport->Present(this, &CommandBuffer, &Queue, Device.GetPresentQueue(), bLockToVsync);
 	if (bNativePresent)
 	{
 		//#todo-rco: Check for r.FinishCurrentFrame
@@ -1253,7 +1152,7 @@ void FVulkanCommandListContext::RHIEndDrawingViewport(FRHIViewport* ViewportRHI,
 
 	RHI->DrawingViewport = nullptr;
 
-	WriteBeginTimestamp(CommandBufferManager->GetActiveCmdBuffer());
+	WriteBeginTimestamp(GetActiveCmdBuffer());
 }
 
 #if WITH_RHI_BREADCRUMBS
@@ -1275,7 +1174,7 @@ void FVulkanCommandListContext::RHIEndDrawingViewport(FRHIViewport* ViewportRHI,
 		if (ShouldEmitBreadcrumbs())
 		{
 		#if VULKAN_ENABLE_DRAW_MARKERS
-			if (auto CmdBeginLabel = Device->GetCmdBeginDebugLabel())
+			if (auto CmdBeginLabel = Device.GetCmdBeginDebugLabel())
 			{
 				FTCHARToUTF8 Converter(GetNameStr());
 				VkDebugUtilsLabelEXT Label;
@@ -1286,7 +1185,7 @@ void FVulkanCommandListContext::RHIEndDrawingViewport(FRHIViewport* ViewportRHI,
 				Label.color[1] = LColor.G;
 				Label.color[2] = LColor.B;
 				Label.color[3] = LColor.A;
-				CmdBeginLabel(GetCommandBufferManager()->GetActiveCmdBuffer()->GetHandle(), &Label);
+				CmdBeginLabel(GetActiveCmdBuffer()->GetHandle(), &Label);
 			}
 		#endif
 
@@ -1305,7 +1204,7 @@ void FVulkanCommandListContext::RHIEndDrawingViewport(FRHIViewport* ViewportRHI,
 		#if VULKAN_SUPPORTS_GPU_CRASH_DUMPS
 			if (GpuProfiler.bTrackingGPUCrashData)
 			{
-				GpuProfiler.PushMarkerForCrash(GetCommandBufferManager()->GetActiveCmdBuffer(), Device->GetCrashMarkerBuffer(), GetNameStr());
+				GpuProfiler.PushMarkerForCrash(GetActiveCmdBuffer(), Device.GetCrashMarkerBuffer(), GetNameStr());
 			}
 		#endif
 			if (GpuProfiler.IsProfilingGPU())
@@ -1330,7 +1229,7 @@ void FVulkanCommandListContext::RHIEndDrawingViewport(FRHIViewport* ViewportRHI,
 		#if VULKAN_SUPPORTS_GPU_CRASH_DUMPS
 			if (GpuProfiler.bTrackingGPUCrashData)
 			{
-				GpuProfiler.PopMarkerForCrash(GetCommandBufferManager()->GetActiveCmdBuffer(), Device->GetCrashMarkerBuffer());
+				GpuProfiler.PopMarkerForCrash(GetActiveCmdBuffer(), Device.GetCrashMarkerBuffer());
 			}
 		#endif
 		}
@@ -1346,9 +1245,9 @@ void FVulkanCommandListContext::RHIEndDrawingViewport(FRHIViewport* ViewportRHI,
 		#endif
 
 		#if VULKAN_ENABLE_DRAW_MARKERS
-			if (auto CmdEndLabel = Device->GetCmdEndDebugLabel())
+			if (auto CmdEndLabel = Device.GetCmdEndDebugLabel())
 			{
-				CmdEndLabel(GetCommandBufferManager()->GetActiveCmdBuffer()->GetHandle());
+				CmdEndLabel(GetActiveCmdBuffer()->GetHandle());
 			}
 		#endif
 		}
@@ -1368,7 +1267,7 @@ void FVulkanDynamicRHI::RHIFlushResources()
 {
 	FVulkanCommandListContextImmediate& ImmediateContext = GetDevice()->GetImmediateContext();
 	bool bTrimMemory = true;
-	ImmediateContext.GetCommandBufferManager()->FreeUnusedCmdBuffers(bTrimMemory);
+	ImmediateContext.FreeUnusedCmdBuffers(bTrimMemory);
 }
 
 // IVulkanDynamicRHI interface
@@ -1421,7 +1320,7 @@ uint32 FVulkanDynamicRHI::RHIGetGraphicsQueueFamilyIndex() const
 VkCommandBuffer FVulkanDynamicRHI::RHIGetActiveVkCommandBuffer()
 {
 	FVulkanCommandListContextImmediate& ImmediateContext = GetDevice()->GetImmediateContext();
-	VkCommandBuffer VulkanCommandBuffer = ImmediateContext.GetCommandBufferManager()->GetActiveCmdBuffer()->GetHandle();
+	VkCommandBuffer VulkanCommandBuffer = ImmediateContext.GetActiveCmdBuffer()->GetHandle();
 	return VulkanCommandBuffer;
 }
 
@@ -1584,21 +1483,21 @@ FVulkanRHIAllocationInfo FVulkanDynamicRHI::RHIGetAllocationInfo(FRHIBuffer* InB
 void FVulkanDynamicRHI::RHISetImageLayout(VkImage Image, VkImageLayout OldLayout, VkImageLayout NewLayout, const VkImageSubresourceRange& SubresourceRange)
 {
 	FVulkanCommandListContext& ImmediateContext = GetDevice()->GetImmediateContext();
-	FVulkanCommandBuffer* CmdBuffer = ImmediateContext.GetCommandBufferManager()->GetActiveCmdBuffer();
+	FVulkanCommandBuffer* CmdBuffer = ImmediateContext.GetActiveCmdBuffer();
 	VulkanSetImageLayout(CmdBuffer->GetHandle(), Image, OldLayout, NewLayout, SubresourceRange);
 }
 
 void FVulkanDynamicRHI::RHISetUploadImageLayout(VkImage Image, VkImageLayout OldLayout, VkImageLayout NewLayout, const VkImageSubresourceRange& SubresourceRange)
 {
 	FVulkanCommandListContext& ImmediateContext = GetDevice()->GetImmediateContext();
-	FVulkanCommandBuffer* CmdBuffer = ImmediateContext.GetCommandBufferManager()->GetUploadCmdBuffer();
+	FVulkanCommandBuffer* CmdBuffer = ImmediateContext.GetUploadCmdBuffer();
 	VulkanSetImageLayout(CmdBuffer->GetHandle(), Image, OldLayout, NewLayout, SubresourceRange);
 }
 
 void FVulkanDynamicRHI::RHIFinishExternalComputeWork(VkCommandBuffer InCommandBuffer)
 {
 	FVulkanCommandListContextImmediate& ImmediateContext = GetDevice()->GetImmediateContext();
-	check(InCommandBuffer == ImmediateContext.GetCommandBufferManager()->GetActiveCmdBuffer()->GetHandle());
+	check(InCommandBuffer == ImmediateContext.GetActiveCmdBuffer()->GetHandle());
 
 	ImmediateContext.GetPendingComputeState()->Reset();
 	ImmediateContext.GetPendingGfxState()->Reset();
@@ -1616,7 +1515,7 @@ void FVulkanDynamicRHI::RHIRegisterWork(uint32 NumPrimitives)
 void FVulkanDynamicRHI::RHISubmitUploadCommandBuffer()
 {
 	FVulkanCommandListContext& ImmediateContext = GetDevice()->GetImmediateContext();
-	ImmediateContext.GetCommandBufferManager()->SubmitUploadCmdBuffer();
+	ImmediateContext.SubmitUploadCmdBuffer();
 }
 
 void FVulkanDynamicRHI::RHIVerifyResult(VkResult Result, const ANSICHAR* VkFuntion, const ANSICHAR* Filename, uint32 Line)
diff --git a/Engine/Source/Runtime/VulkanRHI/Private/VulkanRHIPrivate.h b/Engine/Source/Runtime/VulkanRHI/Private/VulkanRHIPrivate.h
--- a/Engine/Source/Runtime/VulkanRHI/Private/VulkanRHIPrivate.h
+++ b/Engine/Source/Runtime/VulkanRHI/Private/VulkanRHIPrivate.h
@@ -40,7 +40,6 @@ class FVulkanShader;
 class FVulkanDescriptorSetsLayout;
 class FVulkanGfxPipeline;
 class FVulkanRenderPass;
-class FVulkanCommandBufferManager;
 class FValidationContext;
 
 
diff --git a/Engine/Source/Runtime/VulkanRHI/Private/VulkanRayTracing.cpp b/Engine/Source/Runtime/VulkanRHI/Private/VulkanRayTracing.cpp
--- a/Engine/Source/Runtime/VulkanRHI/Private/VulkanRayTracing.cpp
+++ b/Engine/Source/Runtime/VulkanRHI/Private/VulkanRayTracing.cpp
@@ -843,8 +843,7 @@ void BuildAccelerationStructure(
 		}
 	}
 
-	FVulkanCommandBufferManager& CommandBufferManager = *CommandContext.GetCommandBufferManager();
-	FVulkanCommandBuffer* const CmdBuffer = CommandBufferManager.GetActiveCmdBuffer();
+	FVulkanCommandBuffer* const CmdBuffer = CommandContext.GetActiveCmdBuffer();
 
 	// Force a memory barrier to make sure all previous builds ops are finished before building the TLAS
 	AddAccelerationStructureBuildBarrier(CmdBuffer->GetHandle());
@@ -854,8 +853,8 @@ void BuildAccelerationStructure(
 	// Acceleration structure build barrier is used here to ensure that the acceleration structure build is complete before any rays are traced
 	AddAccelerationStructureBuildBarrier(CmdBuffer->GetHandle());
 
-	CommandBufferManager.SubmitActiveCmdBuffer();
-	CommandBufferManager.PrepareForNewActiveCommandBuffer();
+	CommandContext.SubmitActiveCmdBuffer();
+	CommandContext.PrepareForNewActiveCommandBuffer();
 
 	Scene.bBuilt = true;
 }
@@ -974,7 +973,7 @@ void FVulkanRayTracingShaderTable::SetSlot(EShaderFrequency Frequency, uint32 Ds
 VkStridedDeviceAddressRegionKHR FVulkanRayTracingShaderTable::CommitRayGenShader(FVulkanCommandListContext& Context, uint32 SrcHandleIndex, TConstArrayView<uint8> SrcHandleData)
 {
 	VkStridedDeviceAddressRegionKHR RayGenRegion;
-	uint8* pMappedMemory = Context.GetTempBlockAllocator().Alloc(HandleSize, Context.GetCommandBufferManager()->GetActiveCmdBuffer(), &RayGenRegion);
+	uint8* pMappedMemory = Context.GetTempBlockAllocator().Alloc(HandleSize, Context.GetActiveCmdBuffer(), &RayGenRegion);
 
 	FMemory::Memcpy(pMappedMemory, &SrcHandleData[SrcHandleIndex * HandleSize], HandleSize);
 
@@ -1290,7 +1289,7 @@ void FVulkanCommandListContext::RHIBuildAccelerationStructures(const TArrayView<
 
 		FVkRtBLASBuildData& BuildData = TempBuildData.AddDefaulted_GetRef();
 		GetBLASBuildData(
-			Device->GetInstanceHandle(),
+			Device.GetInstanceHandle(),
 			MakeArrayView(Geometry->Initializer.Segments),
 			Geometry->Initializer.GeometryType,
 			Geometry->Initializer.IndexBuffer,
@@ -1322,7 +1321,7 @@ void FVulkanCommandListContext::RHIBuildAccelerationStructures(const TArrayView<
 		Geometry->SetupHitGroupSystemParameters();
 	}
 	
-	FVulkanCommandBuffer* const CmdBuffer = CommandBufferManager->GetActiveCmdBuffer();
+	FVulkanCommandBuffer* const CmdBuffer = GetActiveCmdBuffer();
 	VulkanDynamicAPI::vkCmdBuildAccelerationStructuresKHR(CmdBuffer->GetHandle(), Params.Num(), BuildGeometryInfos.GetData(), BuildRangeInfos.GetData());
 
 	// Add an acceleration structure build barrier after each acceleration structure build batch.
@@ -1331,8 +1330,8 @@ void FVulkanCommandListContext::RHIBuildAccelerationStructures(const TArrayView<
 	// are complete before BLAS is used again on the GPU.
 	AddAccelerationStructureBuildBarrier(CmdBuffer->GetHandle());
 
-	CommandBufferManager->SubmitActiveCmdBuffer();
-	CommandBufferManager->PrepareForNewActiveCommandBuffer();
+	SubmitActiveCmdBuffer();
+	PrepareForNewActiveCommandBuffer();
 
 	for (const FRayTracingGeometryBuildParams& P : Params)
 	{
@@ -1341,7 +1340,7 @@ void FVulkanCommandListContext::RHIBuildAccelerationStructures(const TArrayView<
 		ERayTracingAccelerationStructureFlags GeometryBuildFlags = GetRayTracingAccelerationStructureBuildFlags(Geometry->Initializer);
 		if (ShouldCompactAfterBuild(GeometryBuildFlags))
 		{
-			Device->GetRayTracingCompactionRequestHandler()->RequestCompact(Geometry);
+			Device.GetRayTracingCompactionRequestHandler()->RequestCompact(Geometry);
 			Geometry->bHasPendingCompactionRequests = true;
 		}
 	}
@@ -1790,7 +1789,7 @@ bool FVulkanRayTracingCompactedSizeQueryPool::TryGetResults(uint32 NumResults)
 		return false;
 	}
 
-	VkResult Result = VulkanRHI::vkGetQueryPoolResults(Device->GetInstanceHandle(), QueryPool, 0, NumResults, NumResults * sizeof(uint64), QueryOutput.GetData(), sizeof(uint64), VK_QUERY_RESULT_64_BIT | VK_QUERY_RESULT_WAIT_BIT);
+	VkResult Result = VulkanRHI::vkGetQueryPoolResults(Device.GetInstanceHandle(), QueryPool, 0, NumResults, NumResults * sizeof(uint64), QueryOutput.GetData(), sizeof(uint64), VK_QUERY_RESULT_64_BIT | VK_QUERY_RESULT_WAIT_BIT);
 	if (Result == VK_SUCCESS)
 	{
 		return true;
@@ -1801,7 +1800,7 @@ bool FVulkanRayTracingCompactedSizeQueryPool::TryGetResults(uint32 NumResults)
 FVulkanRayTracingCompactionRequestHandler::FVulkanRayTracingCompactionRequestHandler(FVulkanDevice* const InDevice)
 	: VulkanRHI::FDeviceChild(InDevice)
 {
-	QueryPool = new FVulkanRayTracingCompactedSizeQueryPool(InDevice, GVulkanRayTracingMaxBatchedCompaction);
+	QueryPool = new FVulkanRayTracingCompactedSizeQueryPool(*InDevice, GVulkanRayTracingMaxBatchedCompaction);
 
 	ActiveRequests.Reserve(GVulkanRayTracingMaxBatchedCompaction);
 	ActiveBLASes.Reserve(GVulkanRayTracingMaxBatchedCompaction);
@@ -1915,8 +1914,7 @@ void FVulkanRayTracingCompactionRequestHandler::Update(FVulkanCommandListContext
 		// clear out all of the pending requests, don't allow the array to shrink
 		PendingRequests.RemoveAt(0, ActiveRequests.Num(), EAllowShrinking::No);
 
-		FVulkanCommandBufferManager& CommandBufferManager = *InCommandContext.GetCommandBufferManager();
-		FVulkanCommandBuffer* const CmdBuffer = CommandBufferManager.GetActiveCmdBuffer();
+		FVulkanCommandBuffer* const CmdBuffer = InCommandContext.GetActiveCmdBuffer();
 
 		// Barrier here is not stricly necessary as it is added after the build.
 		// AddAccelerationStructureBuildBarrier(CmdBuffer->GetHandle());
@@ -1932,8 +1930,8 @@ void FVulkanRayTracingCompactionRequestHandler::Update(FVulkanCommandListContext
 
 		QueryPool->EndBatch(CmdBuffer);
 
-		CommandBufferManager.SubmitActiveCmdBuffer();
-		CommandBufferManager.PrepareForNewActiveCommandBuffer();
+		InCommandContext.SubmitActiveCmdBuffer();
+		InCommandContext.PrepareForNewActiveCommandBuffer();
 	}
 }
 
@@ -1944,7 +1942,7 @@ bool FVulkanRayTracingCompactionRequestHandler::IsUsingCmdBuffer(FVulkanCommandB
 }
 
 
-static FVulkanPipelineBarrier SetRayGenResources(FVulkanDevice* Device, FVulkanCommandBuffer* const CmdBuffer, const FRayTracingShaderBindings& InGlobalResourceBindings, FVulkanRayTracingShaderTable* ShaderTable)
+static FVulkanPipelineBarrier SetRayGenResources(FVulkanDevice& Device, FVulkanCommandBuffer* const CmdBuffer, const FRayTracingShaderBindings& InGlobalResourceBindings, FVulkanRayTracingShaderTable* ShaderTable)
 {
 	TArray<const FVulkanUniformBuffer*> UniformBuffers;
 	UniformBuffers.Reserve(UE_ARRAY_COUNT(InGlobalResourceBindings.UniformBuffers));
@@ -1985,7 +1983,7 @@ static FVulkanPipelineBarrier SetRayGenResources(FVulkanDevice* Device, FVulkanC
 				NumSkippedSlots++;
 			}
 		}
-		Device->GetBindlessDescriptorManager()->RegisterUniformBuffers(*CmdBuffer, VK_PIPELINE_BIND_POINT_RAY_TRACING_KHR, StageUBs);
+		Device.GetBindlessDescriptorManager()->RegisterUniformBuffers(*CmdBuffer, VK_PIPELINE_BIND_POINT_RAY_TRACING_KHR, StageUBs);
 	}
 
 	// Add all the UBs references by the shader table
@@ -2150,7 +2148,7 @@ void FVulkanCommandListContext::RHIRayTraceDispatch(
 	FVulkanRayTracingShader* RayGenShader = ResourceCast(InRayGenShader);
 	FVulkanRayTracingShaderTable* ShaderTable = ResourceCast(InSBT);
 
-	FVulkanCommandBuffer* const CmdBuffer = GetCommandBufferManager()->GetActiveCmdBuffer();
+	FVulkanCommandBuffer* const CmdBuffer = GetActiveCmdBuffer();
 	VulkanRHI::vkCmdBindPipeline(CmdBuffer->GetHandle(), VK_PIPELINE_BIND_POINT_RAY_TRACING_KHR, Pipeline->GetPipeline());
 
 	checkSlow(InRayGenShader->GetFrequency() == SF_RayGen);
@@ -2187,7 +2185,7 @@ void FVulkanCommandListContext::RHIRayTraceDispatchIndirect(
 	FVulkanRayTracingShader* RayGenShader = ResourceCast(InRayGenShader);
 	FVulkanRayTracingShaderTable* ShaderTable = ResourceCast(InSBT);
 
-	FVulkanCommandBuffer* const CmdBuffer = GetCommandBufferManager()->GetActiveCmdBuffer();
+	FVulkanCommandBuffer* const CmdBuffer = GetActiveCmdBuffer();
 	VulkanRHI::vkCmdBindPipeline(CmdBuffer->GetHandle(), VK_PIPELINE_BIND_POINT_RAY_TRACING_KHR, Pipeline->GetPipeline());
 
 	checkSlow(InRayGenShader->GetFrequency() == SF_RayGen);
diff --git a/Engine/Source/Runtime/VulkanRHI/Private/VulkanRayTracing.h b/Engine/Source/Runtime/VulkanRHI/Private/VulkanRayTracing.h
--- a/Engine/Source/Runtime/VulkanRHI/Private/VulkanRayTracing.h
+++ b/Engine/Source/Runtime/VulkanRHI/Private/VulkanRayTracing.h
@@ -313,7 +313,7 @@ public:
 class FVulkanRayTracingCompactedSizeQueryPool : public FVulkanQueryPool
 {
 public:
-	FVulkanRayTracingCompactedSizeQueryPool(FVulkanDevice* InDevice, uint32 InMaxQueries)
+	FVulkanRayTracingCompactedSizeQueryPool(FVulkanDevice& InDevice, uint32 InMaxQueries)
 		: FVulkanQueryPool(InDevice, nullptr, InMaxQueries, VK_QUERY_TYPE_ACCELERATION_STRUCTURE_COMPACTED_SIZE_KHR, false)
 	{}
 
diff --git a/Engine/Source/Runtime/VulkanRHI/Private/VulkanRenderTarget.cpp b/Engine/Source/Runtime/VulkanRHI/Private/VulkanRenderTarget.cpp
--- a/Engine/Source/Runtime/VulkanRHI/Private/VulkanRenderTarget.cpp
+++ b/Engine/Source/Runtime/VulkanRHI/Private/VulkanRenderTarget.cpp
@@ -70,7 +70,7 @@ FVulkanRenderPass* FVulkanCommandListContext::PrepareRenderPassForPSOCreation(co
 FVulkanRenderPass* FVulkanCommandListContext::PrepareRenderPassForPSOCreation(const FVulkanRenderTargetLayout& RTLayout)
 {
 	FVulkanRenderPass* RenderPass = nullptr;
-	RenderPass = Device->GetRenderPassManager().GetOrCreateRenderPass(RTLayout);
+	RenderPass = Device.GetRenderPassManager().GetOrCreateRenderPass(RTLayout);
 	return RenderPass;
 }
 
@@ -183,7 +183,7 @@ void FVulkanDynamicRHI::RHIReadSurfaceData(FRHITexture* TextureRHI, FIntRect Rec
 	const bool bCPUReadback = EnumHasAllFlags(Surface.GetDesc().Flags, TexCreate_CPUReadback);
 	if (!bCPUReadback) //this function supports reading back arbitrary rendertargets, so if its not a cpu readback surface, we do a copy.
 	{
-		CmdBuffer = ImmediateContext.GetCommandBufferManager()->GetUploadCmdBuffer();
+		CmdBuffer = ImmediateContext.GetUploadCmdBuffer();
 		StagingBuffer = Device->GetStagingManager().AcquireBuffer(BufferSize, VK_BUFFER_USAGE_TRANSFER_DST_BIT, VK_MEMORY_PROPERTY_HOST_CACHED_BIT);
 		VkBufferImageCopy CopyRegion;
 		FMemory::Memzero(CopyRegion);
@@ -217,7 +217,7 @@ void FVulkanDynamicRHI::RHIReadSurfaceData(FRHITexture* TextureRHI, FIntRect Rec
 		AfterBarrier.Execute(CmdBuffer);
 
 		// Force upload
-		ImmediateContext.GetCommandBufferManager()->SubmitUploadCmdBuffer();
+		ImmediateContext.SubmitUploadCmdBuffer();
 	}
 
 	Device->WaitUntilIdle();
@@ -246,7 +246,7 @@ void FVulkanDynamicRHI::RHIReadSurfaceData(FRHITexture* TextureRHI, FIntRect Rec
 		Device->GetStagingManager().ReleaseBuffer(CmdBuffer, StagingBuffer);
 	}
 
-	ImmediateContext.GetCommandBufferManager()->PrepareForNewActiveCommandBuffer();
+	ImmediateContext.PrepareForNewActiveCommandBuffer();
 }
 
 void FVulkanDynamicRHI::RHIReadSurfaceData(FRHITexture* TextureRHI, FIntRect Rect, TArray<FLinearColor>& OutData, FReadSurfaceDataFlags InFlags)
@@ -350,7 +350,7 @@ void FVulkanDynamicRHI::RHIReadSurfaceFloatData(FRHITexture* TextureRHI, FIntRec
 		AfterBarrier.Execute(InCmdBuffer);
 
 		// Force upload
-		InDevice->GetImmediateContext().GetCommandBufferManager()->SubmitUploadCmdBuffer();
+		InDevice->GetImmediateContext().SubmitUploadCmdBuffer();
 		InDevice->WaitUntilIdle();
 
 		StagingBuffer->InvalidateMappedMemory();
@@ -402,7 +402,7 @@ void FVulkanDynamicRHI::RHIReadSurfaceFloatData(FRHITexture* TextureRHI, FIntRec
 	{
 		Device->PrepareForCPURead();
 
-		FVulkanCommandBuffer* CmdBuffer = Device->GetImmediateContext().GetCommandBufferManager()->GetUploadCmdBuffer();
+		FVulkanCommandBuffer* CmdBuffer = Device->GetImmediateContext().GetUploadCmdBuffer();
 		switch (TextureRHI->GetDesc().Dimension)
 		{
 			case ETextureDimension::TextureCubeArray:
@@ -419,7 +419,7 @@ void FVulkanDynamicRHI::RHIReadSurfaceFloatData(FRHITexture* TextureRHI, FIntRec
 				checkNoEntry();
 				break;
 		}
-		Device->GetImmediateContext().GetCommandBufferManager()->PrepareForNewActiveCommandBuffer();
+		Device->GetImmediateContext().PrepareForNewActiveCommandBuffer();
 	}
 }
 
@@ -445,7 +445,7 @@ void FVulkanDynamicRHI::RHIRead3DSurfaceFloatData(FRHITexture* TextureRHI, FIntR
 	}
 
 	Device->PrepareForCPURead();
-	FVulkanCommandBuffer* CmdBuffer = Device->GetImmediateContext().GetCommandBufferManager()->GetUploadCmdBuffer();
+	FVulkanCommandBuffer* CmdBuffer = Device->GetImmediateContext().GetUploadCmdBuffer();
 
 	ensure(Surface.StorageFormat == VK_FORMAT_R16G16B16A16_SFLOAT);
 
@@ -496,7 +496,7 @@ void FVulkanDynamicRHI::RHIRead3DSurfaceFloatData(FRHITexture* TextureRHI, FIntR
 	AfterBarrier.Execute(CmdBuffer);
 
 	// Force upload
-	Device->GetImmediateContext().GetCommandBufferManager()->SubmitUploadCmdBuffer();
+	Device->GetImmediateContext().SubmitUploadCmdBuffer();
 	Device->WaitUntilIdle();
 
 	StagingBuffer->InvalidateMappedMemory();
@@ -517,7 +517,7 @@ void FVulkanDynamicRHI::RHIRead3DSurfaceFloatData(FRHITexture* TextureRHI, FIntR
 	checkf(Dest <= End, TEXT("Memory overwrite! Calculated total size %d: SizeX %d SizeY %d SizeZ %d; InRect(%d, %d, %d, %d) InZ(%d, %d)"),
 		Size, SizeX, SizeY, SizeZ, InRect.Min.X, InRect.Min.Y, InRect.Max.X, InRect.Max.Y, ZMinMax.X, ZMinMax.Y);
 	Device->GetStagingManager().ReleaseBuffer(CmdBuffer, StagingBuffer);
-	Device->GetImmediateContext().GetCommandBufferManager()->PrepareForNewActiveCommandBuffer();
+	Device->GetImmediateContext().PrepareForNewActiveCommandBuffer();
 }
 
 VkFormat FVulkanCommandListContext::GetSwapchainImageFormat() const
@@ -567,17 +567,17 @@ bool FVulkanCommandListContext::IsSwapchainImage(FRHITexture* InTexture) const
 
 void FVulkanCommandListContext::RHIBeginRenderPass(const FRHIRenderPassInfo& InInfo, const TCHAR* InName)
 {
-	FVulkanCommandBuffer* CmdBuffer = CommandBufferManager->GetActiveCmdBuffer();
+	FVulkanCommandBuffer* CmdBuffer = GetActiveCmdBuffer();
 
 	if (GVulkanSubmitAfterEveryEndRenderPass)
 	{
-		CommandBufferManager->SubmitActiveCmdBuffer();
-		CommandBufferManager->PrepareForNewActiveCommandBuffer();
-		CmdBuffer = CommandBufferManager->GetActiveCmdBuffer();
+		SubmitActiveCmdBuffer();
+		PrepareForNewActiveCommandBuffer();
+		CmdBuffer = GetActiveCmdBuffer();
 	}
 	else if (SafePointSubmit())
 	{
-		CmdBuffer = CommandBufferManager->GetActiveCmdBuffer();
+		CmdBuffer = GetActiveCmdBuffer();
 	}
 
 	RenderPassInfo = InInfo;
@@ -595,7 +595,7 @@ void FVulkanCommandListContext::RHIBeginRenderPass(const FRHIRenderPassInfo& InI
 		FVulkanTexture& VulkanTexture = *ResourceCast(DSTexture);
 		const VkImageAspectFlags AspectFlags = VulkanTexture.GetFullAspectMask();
 
-		if (GetDevice()->SupportsParallelRendering())
+		if (Device.SupportsParallelRendering())
 		{
 			const FExclusiveDepthStencil ExclusiveDepthStencil = InInfo.DepthStencilRenderTarget.ExclusiveDepthStencil;
 			if (VKHasAnyFlags(AspectFlags, VK_IMAGE_ASPECT_DEPTH_BIT))
@@ -610,7 +610,7 @@ void FVulkanCommandListContext::RHIBeginRenderPass(const FRHIRenderPassInfo& InI
 					CurrentDepthLayout = CmdBuffer->GetLayoutManager().GetDepthStencilHint(VulkanTexture, VK_IMAGE_ASPECT_DEPTH_BIT);
 					if (CurrentDepthLayout == VK_IMAGE_LAYOUT_UNDEFINED)
 					{
-						CurrentDepthLayout = GetQueue()->GetLayoutManager().GetDepthStencilHint(VulkanTexture, VK_IMAGE_ASPECT_DEPTH_BIT);
+						CurrentDepthLayout = Queue.GetLayoutManager().GetDepthStencilHint(VulkanTexture, VK_IMAGE_ASPECT_DEPTH_BIT);
 						if (CurrentDepthLayout == VK_IMAGE_LAYOUT_UNDEFINED)
 						{
 							if (ExclusiveDepthStencil.IsDepthRead())
@@ -642,7 +642,7 @@ void FVulkanCommandListContext::RHIBeginRenderPass(const FRHIRenderPassInfo& InI
 					CurrentStencilLayout = CmdBuffer->GetLayoutManager().GetDepthStencilHint(VulkanTexture, VK_IMAGE_ASPECT_STENCIL_BIT);
 					if (CurrentStencilLayout == VK_IMAGE_LAYOUT_UNDEFINED)
 					{
-						CurrentStencilLayout = GetQueue()->GetLayoutManager().GetDepthStencilHint(VulkanTexture, VK_IMAGE_ASPECT_STENCIL_BIT);
+						CurrentStencilLayout = Queue.GetLayoutManager().GetDepthStencilHint(VulkanTexture, VK_IMAGE_ASPECT_STENCIL_BIT);
 						// If the layout is still UNDEFINED, FVulkanRenderTargetLayout will force it into a known state
 					}
 				}
@@ -663,16 +663,16 @@ void FVulkanCommandListContext::RHIBeginRenderPass(const FRHIRenderPassInfo& InI
 		}
 	}
 
-	FVulkanRenderTargetLayout RTLayout(*Device, InInfo, CurrentDepthLayout, CurrentStencilLayout);
+	FVulkanRenderTargetLayout RTLayout(Device, InInfo, CurrentDepthLayout, CurrentStencilLayout);
 	check(RTLayout.GetExtent2D().width != 0 && RTLayout.GetExtent2D().height != 0);
 
-	FVulkanRenderPass* RenderPass = Device->GetRenderPassManager().GetOrCreateRenderPass(RTLayout);
+	FVulkanRenderPass* RenderPass = Device.GetRenderPassManager().GetOrCreateRenderPass(RTLayout);
 	FRHISetRenderTargetsInfo RTInfo;
 	InInfo.ConvertToRenderTargetsInfo(RTInfo);
 
-	FVulkanFramebuffer* Framebuffer = Device->GetRenderPassManager().GetOrCreateFramebuffer(RTInfo, RTLayout, RenderPass);
+	FVulkanFramebuffer* Framebuffer = Device.GetRenderPassManager().GetOrCreateFramebuffer(RTInfo, RTLayout, RenderPass);
 	checkf(RenderPass != nullptr && Framebuffer != nullptr, TEXT("RenderPass not started! Bad combination of values? Depth %p #Color %d Color0 %p"), (void*)InInfo.DepthStencilRenderTarget.DepthStencilTarget, InInfo.GetNumColorRenderTargets(), (void*)InInfo.ColorRenderTargets[0].RenderTarget);
-	Device->GetRenderPassManager().BeginRenderPass(*this, *Device, CmdBuffer, InInfo, RTLayout, RenderPass, Framebuffer);
+	Device.GetRenderPassManager().BeginRenderPass(*this, Device, CmdBuffer, InInfo, RTLayout, RenderPass, Framebuffer);
 
 	check(!CurrentRenderPass);
 	CurrentRenderPass = RenderPass;
@@ -682,14 +682,14 @@ void FVulkanCommandListContext::RHIBeginRenderPass(const FRHIRenderPassInfo& InI
 void FVulkanCommandListContext::RHIEndRenderPass()
 {
 	const bool bHasOcclusionQueries = (RenderPassInfo.NumOcclusionQueries > 0);
-	FVulkanCommandBuffer* CmdBuffer = CommandBufferManager->GetActiveCmdBuffer();
+	FVulkanCommandBuffer* CmdBuffer = GetActiveCmdBuffer();
 
 	if (bHasOcclusionQueries)
 	{
 		EndOcclusionQueryBatch(CmdBuffer);
 	}
 	
-	Device->GetRenderPassManager().EndRenderPass(CmdBuffer);
+	Device.GetRenderPassManager().EndRenderPass(CmdBuffer);
 
 	check(CurrentRenderPass);
 	CurrentRenderPass = nullptr;
diff --git a/Engine/Source/Runtime/VulkanRHI/Private/VulkanResources.h b/Engine/Source/Runtime/VulkanRHI/Private/VulkanResources.h
--- a/Engine/Source/Runtime/VulkanRHI/Private/VulkanResources.h
+++ b/Engine/Source/Runtime/VulkanRHI/Private/VulkanResources.h
@@ -19,12 +19,12 @@
 
 class FVulkanDevice;
 class FVulkanQueue;
+class FVulkanContextCommon;
 class FVulkanCommandBuffer;
 class FVulkanTexture;
 class FVulkanBuffer;
 class FVulkanLayout;
 class FVulkanOcclusionQuery;
-class FVulkanCommandBufferManager;
 struct FRHITransientHeapAllocation;
 struct FVulkanPendingBufferLock;
 
@@ -771,28 +771,34 @@ protected:
 	EImageOwnerType ImageOwnerType;
 };
 
-class FVulkanQueryPool : public VulkanRHI::FDeviceChild
+class FVulkanQueryPool
 {
 public:
-	FVulkanQueryPool(FVulkanDevice* InDevice, FVulkanCommandBufferManager* CommandBufferManager, uint32 InMaxQueries, VkQueryType InQueryType, bool bInShouldAddReset = true);
+	FVulkanQueryPool(FVulkanDevice& InDevice, FVulkanContextCommon* InCommonContext, uint32 InMaxQueries, VkQueryType InQueryType, bool bInShouldAddReset = true);
 	virtual ~FVulkanQueryPool();
 
-	inline uint32 GetMaxQueries() const
+	uint32 GetMaxQueries() const
 	{
 		return MaxQueries;
 	}
 
-	inline VkQueryPool GetHandle() const
+	VkQueryPool GetHandle() const
 	{
 		return QueryPool;
 	}
 
-	inline uint64 GetResultValue(uint32 Index) const
+	uint64 GetResultValue(uint32 Index) const
 	{
 		return QueryOutput[Index];
 	}
 
+	FVulkanDevice& GetDevice()
+	{
+		return Device;
+	}
+
 protected:
+	FVulkanDevice& Device;
 	VkQueryPool QueryPool;
 	VkEvent ResetEvent;
 	const uint32 MaxQueries;
@@ -803,8 +809,8 @@ protected:
 class FVulkanOcclusionQueryPool : public FVulkanQueryPool
 {
 public:
-	FVulkanOcclusionQueryPool(FVulkanDevice* InDevice, FVulkanCommandBufferManager* CommandBufferManager, uint32 InMaxQueries)
-		: FVulkanQueryPool(InDevice, CommandBufferManager, InMaxQueries, VK_QUERY_TYPE_OCCLUSION)
+	FVulkanOcclusionQueryPool(FVulkanDevice& InDevice, FVulkanContextCommon* InContext, uint32 InMaxQueries)
+		: FVulkanQueryPool(InDevice, InContext, InMaxQueries, VK_QUERY_TYPE_OCCLUSION)
 	{
 		AcquiredIndices.AddZeroed(Align(InMaxQueries, 64) / 64);
 		AllocatedQueries.AddZeroed(InMaxQueries);
@@ -893,8 +899,8 @@ protected:
 class FVulkanTimingQueryPool : public FVulkanQueryPool
 {
 public:
-	FVulkanTimingQueryPool(FVulkanDevice* InDevice, FVulkanCommandBufferManager* CommandBufferManager, uint32 InBufferSize)
-		: FVulkanQueryPool(InDevice, CommandBufferManager, InBufferSize * 2, VK_QUERY_TYPE_TIMESTAMP, false)
+	FVulkanTimingQueryPool(FVulkanDevice& InDevice, FVulkanContextCommon* InContext, uint32 InBufferSize)
+		: FVulkanQueryPool(InDevice, InContext, InBufferSize * 2, VK_QUERY_TYPE_TIMESTAMP, false)
 		, BufferSize(InBufferSize)
 	{
 		TimestampListHandles.AddZeroed(InBufferSize * 2);
diff --git a/Engine/Source/Runtime/VulkanRHI/Private/VulkanTexture.cpp b/Engine/Source/Runtime/VulkanRHI/Private/VulkanTexture.cpp
--- a/Engine/Source/Runtime/VulkanRHI/Private/VulkanTexture.cpp
+++ b/Engine/Source/Runtime/VulkanRHI/Private/VulkanTexture.cpp
@@ -114,7 +114,7 @@ static void VulkanTextureDestroyed(const FRHITextureDesc& TextureDesc, uint64 Si
 
 inline void FVulkanTexture::InternalLockWrite(FVulkanCommandListContext& Context, FVulkanTexture* Surface, const VkBufferImageCopy& Region, VulkanRHI::FStagingBuffer* StagingBuffer)
 {
-	FVulkanCommandBuffer* CmdBuffer = Context.GetCommandBufferManager()->GetUploadCmdBuffer();
+	FVulkanCommandBuffer* CmdBuffer = Context.GetUploadCmdBuffer();
 	ensure(CmdBuffer->IsOutsideRenderPass());
 	VkCommandBuffer StagingCommandBuffer = CmdBuffer->GetHandle();
 
@@ -139,7 +139,7 @@ inline void FVulkanTexture::InternalLockWrite(FVulkanCommandListContext& Context
 
 	if (GVulkanSubmitOnTextureUnlock != 0)
 	{
-		Context.GetCommandBufferManager()->SubmitUploadCmdBuffer();
+		Context.SubmitUploadCmdBuffer();
 	}
 }
 
@@ -729,7 +729,7 @@ void FVulkanTexture::InternalMoveSurface(FVulkanDevice& InDevice, FVulkanCommand
 	Image = MovedImage;
 
 	// Move is used for defrag, which uses layouts stored in the queue, update the layout stored there
-	Context.GetQueue()->GetLayoutManager().SetFullLayout(*this, OriginalLayout);
+	Context.Queue.GetLayoutManager().SetFullLayout(*this, OriginalLayout);
 }
 
 void FVulkanTexture::DestroySurface()
@@ -885,7 +885,7 @@ void FVulkanTexture::SetInitialImageState(FVulkanCommandListContext& Context, Vk
 	// Can't use TransferQueue as Vulkan requires that queue to also have Gfx or Compute capabilities...
 	//#todo-rco: This function is only used during loading currently, if used for regular RHIClear then use the ActiveCmdBuffer
 	// NOTE: Transient resources' memory might have belonged to another resource earlier in the ActiveCmdBuffer, so we can't use UploadCmdBuffer
-	FVulkanCommandBuffer* CmdBuffer = bIsTransientResource ? Context.GetCommandBufferManager()->GetActiveCmdBuffer() : Context.GetCommandBufferManager()->GetUploadCmdBuffer();
+	FVulkanCommandBuffer* CmdBuffer = bIsTransientResource ? Context.GetActiveCmdBuffer() : Context.GetUploadCmdBuffer();
 	ensure(CmdBuffer->IsOutsideRenderPass());
 
 	VkImageSubresourceRange SubresourceRange = FVulkanPipelineBarrier::MakeSubresourceRange(FullAspectMask);
@@ -1081,7 +1081,7 @@ static void DoAsyncReallocateTexture2D(FVulkanCommandListContext& Context, FVulk
 	const uint32 SourceFirstMip = OldTexture->GetNumMips() - NumSharedMips;
 	const uint32 DestFirstMip = NewTexture->GetNumMips() - NumSharedMips;
 
-	FVulkanCommandBuffer* CmdBuffer = Context.GetCommandBufferManager()->GetUploadCmdBuffer();
+	FVulkanCommandBuffer* CmdBuffer = Context.GetUploadCmdBuffer();
 	ensure(CmdBuffer->IsOutsideRenderPass());
 
 	VkCommandBuffer StagingCommandBuffer = CmdBuffer->GetHandle();
@@ -2188,7 +2188,7 @@ void FVulkanTexture::Move(FVulkanDevice& InDevice, FVulkanCommandListContext& Co
 	}
 
 	// Move is used for defrag, which uses layouts stored in the queue
-	const FVulkanImageLayout* OriginalLayout = Context.GetQueue()->GetLayoutManager().GetFullLayout(Image);
+	const FVulkanImageLayout* OriginalLayout = Context.Queue.GetLayoutManager().GetFullLayout(Image);
 	check(OriginalLayout && OriginalLayout->AreAllSubresourcesSameLayout());
 
 	const ETextureCreateFlags UEFlags = GetDesc().Flags;
@@ -2219,7 +2219,7 @@ void FVulkanTexture::Evict(FVulkanDevice& InDevice, FVulkanCommandListContext& C
 	}
 
 	// Eviction layouts are read from the queue
-	const FVulkanImageLayout* OriginalLayout = Context.GetQueue()->GetLayoutManager().GetFullLayout(Image);
+	const FVulkanImageLayout* OriginalLayout = Context.Queue.GetLayoutManager().GetFullLayout(Image);
 	if (OriginalLayout && OriginalLayout->AreAllSubresourcesSameLayout())
 	{
 		check(0 == CpuReadbackBuffer);
diff --git a/Engine/Source/Runtime/VulkanRHI/Private/VulkanUAV.cpp b/Engine/Source/Runtime/VulkanRHI/Private/VulkanUAV.cpp
--- a/Engine/Source/Runtime/VulkanRHI/Private/VulkanUAV.cpp
+++ b/Engine/Source/Runtime/VulkanRHI/Private/VulkanUAV.cpp
@@ -605,8 +605,8 @@ void FVulkanUnorderedAccessView::Clear(TRHICommandList_RecursiveHazardous<FVulka
 			{
 				const uint32 ComputeWorkGroupCount = FMath::DivideAndRoundUp(Info.NumElements, (uint32)ClearReplacementCS::TThreadGroupSize<EClearReplacementResourceType::Buffer>::X);
 
-				FVulkanDevice* TargetDevice = FVulkanCommandListContext::GetVulkanContext(RHICmdList.GetContext()).GetDevice();
-				const bool bOversizedBuffer = (ComputeWorkGroupCount > TargetDevice->GetLimits().maxComputeWorkGroupCount[0]);
+				FVulkanDevice& TargetDevice = FVulkanCommandListContext::GetVulkanContext(RHICmdList.GetContext()).Device;
+				const bool bOversizedBuffer = (ComputeWorkGroupCount > TargetDevice.GetLimits().maxComputeWorkGroupCount[0]);
 
 				if (bOversizedBuffer)
 				{
diff --git a/Engine/Source/Runtime/VulkanRHI/Private/VulkanUniformBuffer.cpp b/Engine/Source/Runtime/VulkanRHI/Private/VulkanUniformBuffer.cpp
--- a/Engine/Source/Runtime/VulkanRHI/Private/VulkanUniformBuffer.cpp
+++ b/Engine/Source/Runtime/VulkanRHI/Private/VulkanUniformBuffer.cpp
@@ -30,9 +30,9 @@ enum
 	Uniform buffer RHI object
 -----------------------------------------------------------------------------*/
 
-static void UpdateUniformBufferConstants(FVulkanDevice* Device, void* DestinationData, const void* SourceData, const FRHIUniformBufferLayout* Layout)
+static void UpdateUniformBufferConstants(FVulkanDevice& Device, void* DestinationData, const void* SourceData, const FRHIUniformBufferLayout* Layout)
 {
-	UE::RHICore::UpdateUniformBufferConstants(DestinationData, SourceData, *Layout, Device->SupportsBindless());
+	UE::RHICore::UpdateUniformBufferConstants(DestinationData, SourceData, *Layout, Device.SupportsBindless());
 }
 
 static bool UseTemporaryBuffer(EUniformBufferUsage Usage)
@@ -43,11 +43,11 @@ static bool UseTemporaryBuffer(EUniformBufferUsage Usage)
 
 static void UpdateUniformBufferHelper(FVulkanCommandListContext& Context, FVulkanUniformBuffer* VulkanUniformBuffer, const void* Data, bool bUpdateConstants = true)
 {
-	FVulkanCommandBuffer* CmdBuffer = Context.GetCommandBufferManager()->GetActiveCmdBufferDirect();
+	FVulkanCommandBuffer* CmdBuffer = Context.GetActiveCmdBufferDirect();
 
-	FVulkanDevice* Device = Context.GetDevice();
+	FVulkanDevice& Device = Context.Device;
 	const int32 DataSize = VulkanUniformBuffer->GetLayout().ConstantBufferSize;
-	const int32 DataAlignment = FMath::Max<uint32>(Device->GetLimits().minUniformBufferOffsetAlignment, 16u);
+	const int32 DataAlignment = FMath::Max<uint32>(Device.GetLimits().minUniformBufferOffsetAlignment, 16u);
 
 	VulkanRHI::FVulkanAllocation TempAllocation;
 	void* DestinationData = Context.GetTempBlockAllocator().Alloc(DataSize, DataAlignment, CmdBuffer, TempAllocation);
@@ -176,7 +176,7 @@ PRAGMA_ENABLE_DEPRECATION_WARNINGS
 				else
 				{
 					void* CmdListConstantBufferData = RHICmdList.Alloc(DataSize, 16);
-					UpdateUniformBufferConstants(Device, CmdListConstantBufferData, Contents, InLayout);
+					UpdateUniformBufferConstants(*Device, CmdListConstantBufferData, Contents, InLayout);
 
 					RHICmdList.EnqueueLambda([UniformBuffer, DataSize, CmdListConstantBufferData](FRHICommandList& CmdList)
 					{
@@ -195,7 +195,7 @@ PRAGMA_ENABLE_DEPRECATION_WARNINGS
 			ResourceMgr.AllocUniformBuffer(Allocation, InLayout->ConstantBufferSize);
 			if (Contents)
 			{
-				UpdateUniformBufferConstants(Device, Allocation.GetMappedPointer(Device), Contents, InLayout);
+				UpdateUniformBufferConstants(*Device, Allocation.GetMappedPointer(Device), Contents, InLayout);
 				Allocation.FlushMappedMemory(Device);
 			}
 		}
@@ -293,7 +293,7 @@ inline void FVulkanDynamicRHI::UpdateUniformBuffer(FRHICommandListBase& RHICmdLi
 			Device->GetMemoryManager().AllocUniformBuffer(NewUBAlloc, ConstantBufferSize);
 			if (Contents)
 			{
-				UpdateUniformBufferConstants(Device, NewUBAlloc.GetMappedPointer(Device), Contents, &Layout);
+				UpdateUniformBufferConstants(*Device, NewUBAlloc.GetMappedPointer(Device), Contents, &Layout);
 				NewUBAlloc.FlushMappedMemory(Device);
 			}
 		}
diff --git a/Engine/Source/Runtime/VulkanRHI/Private/VulkanUtil.cpp b/Engine/Source/Runtime/VulkanRHI/Private/VulkanUtil.cpp
--- a/Engine/Source/Runtime/VulkanRHI/Private/VulkanUtil.cpp
+++ b/Engine/Source/Runtime/VulkanRHI/Private/VulkanUtil.cpp
@@ -80,7 +80,7 @@ void FVulkanGPUTiming::Initialize(uint32 PoolSize)
 	if (FVulkanPlatform::SupportsTimestampRenderQueries() && GIsSupported)
 	{
 		check(!Pool);
-		Pool = new FVulkanTimingQueryPool(Device, CmdContext->GetCommandBufferManager(), PoolSize);
+		Pool = new FVulkanTimingQueryPool(*Device, Context, PoolSize);
 		Pool->ResultsBuffer = Device->GetStagingManager().AcquireBuffer(Pool->GetMaxQueries() * sizeof(uint64) * 2, VK_BUFFER_USAGE_TRANSFER_DST_BIT, VK_MEMORY_PROPERTY_HOST_COHERENT_BIT);
 		Pool->MappedPointer = (uint64*)Pool->ResultsBuffer->GetMappedPointer();
 	}
@@ -110,7 +110,7 @@ void FVulkanGPUTiming::StartTiming(FVulkanCommandBuffer* CmdBuffer)
 	{
 		if (CmdBuffer == nullptr)
 		{
-			CmdBuffer = CmdContext->GetCommandBufferManager()->GetActiveCmdBuffer();
+			CmdBuffer = Context->GetActiveCmdBuffer();
 		}
 
 		// In case we aren't reading queries, remove oldest
@@ -143,7 +143,7 @@ void FVulkanGPUTiming::StartTiming(FVulkanCommandBuffer* CmdBuffer)
 		VulkanRHI::vkCmdWriteTimestamp(CmdBuffer->GetHandle(), VK_PIPELINE_STAGE_BOTTOM_OF_PIPE_BIT, Pool->GetHandle(), QueryStartIndex);
 		Pool->TimestampListHandles[QueryStartIndex].CmdBuffer = CmdBuffer;
 		Pool->TimestampListHandles[QueryStartIndex].FenceCounter = CmdBuffer->GetFenceSignaledCounter();
-		Pool->TimestampListHandles[QueryStartIndex].FrameCount = CmdContext->GetFrameCounter();
+		Pool->TimestampListHandles[QueryStartIndex].FrameCount = Context->GetFrameCounter();
 		Pool->TimestampListHandles[QueryStartIndex].Attempts = 0;
 		bIsTiming = true;
 	}
@@ -160,7 +160,7 @@ void FVulkanGPUTiming::EndTiming(FVulkanCommandBuffer* CmdBuffer)
 	{
 		if (CmdBuffer == nullptr)
 		{
-			CmdBuffer = CmdContext->GetCommandBufferManager()->GetActiveCmdBuffer();
+			CmdBuffer = Context->GetActiveCmdBuffer();
 		}
 		check(Pool->CurrentTimestamp < Pool->BufferSize);
 		const uint32 QueryStartIndex = Pool->CurrentTimestamp * 2;
@@ -176,7 +176,7 @@ void FVulkanGPUTiming::EndTiming(FVulkanCommandBuffer* CmdBuffer)
 		Pool->TimestampListHandles[QueryStartIndex].FenceCounter = CmdBuffer->GetFenceSignaledCounter();
 		Pool->TimestampListHandles[QueryEndIndex].CmdBuffer = CmdBuffer;
 		Pool->TimestampListHandles[QueryEndIndex].FenceCounter = CmdBuffer->GetFenceSignaledCounter();
-		Pool->TimestampListHandles[QueryEndIndex].FrameCount = CmdContext->GetFrameCounter();
+		Pool->TimestampListHandles[QueryEndIndex].FrameCount = Context->GetFrameCounter();
 		Pool->TimestampListHandles[QueryEndIndex].Attempts = 0;
 		Pool->NumIssuedTimestamps = FMath::Min<uint32>(Pool->NumIssuedTimestamps + 1, Pool->BufferSize);
 
@@ -261,7 +261,7 @@ void FVulkanGPUTiming::RefreshPendingQueries(FVulkanCommandBuffer& CmdBuffer)
 			{
 				QuerySyncPoint.CmdBuffer = &CmdBuffer;
 				QuerySyncPoint.FenceCounter = CmdBuffer.GetFenceSignaledCounter();
-				QuerySyncPoint.FrameCount = CmdContext->GetFrameCounter();
+				QuerySyncPoint.FrameCount = Context->GetFrameCounter();
 				++QuerySyncPoint.Attempts;
 
 				CmdBuffer.AddPendingTimestampQuery(
@@ -288,7 +288,7 @@ uint64 FVulkanGPUTiming::GetTiming(bool bGetCurrentResultsAndBlock)
 
 		uint64 TotalTime = 0;
 
-		uint64_t FrameCount = CmdContext->GetFrameCounter();
+		uint64_t FrameCount = Context->GetFrameCounter();
 
 		// If these timings have been processed return the same time
 		if (PreviousFrame == FrameCount)
@@ -371,11 +371,11 @@ uint64 FVulkanGPUTiming::GetTiming(bool bGetCurrentResultsAndBlock)
 					// CPU wait for query results to be ready.
 					if (bWaitForStart && StartQuerySyncPoint.FenceCounter >= StartQuerySyncPoint.CmdBuffer->GetFenceSignaledCounter())
 					{
-						CmdContext->GetCommandBufferManager()->WaitForCmdBuffer(StartQuerySyncPoint.CmdBuffer);
+						Context->WaitForCmdBuffer(StartQuerySyncPoint.CmdBuffer);
 					}
 					if (bWaitForEnd && EndQuerySyncPoint.FenceCounter >= EndQuerySyncPoint.CmdBuffer->GetFenceSignaledCounter())
 					{
-						CmdContext->GetCommandBufferManager()->WaitForCmdBuffer(EndQuerySyncPoint.CmdBuffer);
+						Context->WaitForCmdBuffer(EndQuerySyncPoint.CmdBuffer);
 					}
 				}
 
@@ -443,10 +443,10 @@ float FVulkanEventNode::GetTiming()
 	return Result;
 }
 
-FVulkanGPUProfiler::FVulkanGPUProfiler(FVulkanCommandListContext* InCmd, FVulkanDevice* InDevice)
+FVulkanGPUProfiler::FVulkanGPUProfiler(FVulkanContextCommon* InContext, FVulkanDevice* InDevice)
 	: bCommandlistSubmitted(false)
 	, Device(InDevice)
-	, CmdContext(InCmd)
+	, CmdContext(InContext)
 	, bBeginFrame(false)
 {
 	BeginFrame();
@@ -555,7 +555,7 @@ void FVulkanGPUProfiler::EndFrame()
 	{
 		if (bTrackingEvents)
 		{
-			CmdContext->GetDevice()->SubmitCommandsAndFlushGPU();
+			CmdContext->Device.SubmitCommandsAndFlushGPU();
 
 			SetEmitDrawEvents(bOriginalGEmitDrawEvents);
 			UE_LOG(LogRHI, Warning, TEXT(""));
@@ -763,7 +763,7 @@ void FVulkanGPUFence::Wait(FRHICommandListImmediate& RHICmdList, FRHIGPUMask GPU
 {
 	if (!Poll())
 	{
-		FVulkanDynamicRHI::Get().GetDevice()->GetImmediateContext().GetCommandBufferManager()->WaitForCmdBuffer(CmdBuffer);
+		FVulkanDynamicRHI::Get().GetDevice()->GetImmediateContext().WaitForCmdBuffer(CmdBuffer);
 	}
 }
 
diff --git a/Engine/Source/Runtime/VulkanRHI/Private/VulkanViewport.cpp b/Engine/Source/Runtime/VulkanRHI/Private/VulkanViewport.cpp
--- a/Engine/Source/Runtime/VulkanRHI/Private/VulkanViewport.cpp
+++ b/Engine/Source/Runtime/VulkanRHI/Private/VulkanViewport.cpp
@@ -281,7 +281,7 @@ void FVulkanViewport::WaitForFrameEventCompletion()
 				if (!GWaitForIdleOnSubmit)
 				{
 					// The wait has already happened if GWaitForIdleOnSubmit is set
-					LastFrameCommandBuffer->GetOwner()->GetMgr().WaitForCmdBuffer(LastFrameCommandBuffer);
+					LastFrameCommandBuffer->GetOwner()->GetContext().WaitForCmdBuffer(LastFrameCommandBuffer);
 				}
 			}
 		}
@@ -752,7 +752,7 @@ void FVulkanViewport::CreateSwapchain(FVulkanSwapChainRecreateInfo* RecreateInfo
 		BackBufferImages.SetNum(Images.Num());
 		RenderingDoneSemaphores.SetNum(Images.Num());
 
-		FVulkanCommandBuffer* CmdBuffer = Device->GetImmediateContext().GetCommandBufferManager()->GetUploadCmdBuffer();
+		FVulkanCommandBuffer* CmdBuffer = Device->GetImmediateContext().GetUploadCmdBuffer();
 		ensure(CmdBuffer->IsOutsideRenderPass());
 
 		VkClearColorValue ClearColor;
@@ -790,7 +790,7 @@ void FVulkanViewport::CreateSwapchain(FVulkanSwapChainRecreateInfo* RecreateInfo
 #endif
 		}
 		
-		Device->GetImmediateContext().GetCommandBufferManager()->SubmitUploadCmdBuffer();
+		Device->GetImmediateContext().SubmitUploadCmdBuffer();
 
 		RHIBackBuffer = new FVulkanBackBuffer(*Device, this, PixelFormat, SizeX, SizeY, TexCreate_RenderTargetable | TexCreate_ShaderResource | TexCreate_ResolveTargetable);
 	}
@@ -1040,7 +1040,7 @@ bool FVulkanViewport::Present(FVulkanCommandListContext* Context, FVulkanCommand
 				// One of the rare cases we'll let parallel rendering check the tracking (legacy path already checked as a fallback)
 				if (!TrackedLayout && Device->SupportsParallelRendering())
 				{
-					TrackedLayout = Context->GetQueue()->GetLayoutManager().GetFullLayout(BackBufferImages[AcquiredImageIndex]->Image);
+					TrackedLayout = Context->Queue.GetLayoutManager().GetFullLayout(BackBufferImages[AcquiredImageIndex]->Image);
 				}
 
 				VkImageLayout LastLayout = TrackedLayout ? TrackedLayout->MainLayout : VK_IMAGE_LAYOUT_UNDEFINED;
@@ -1058,9 +1058,9 @@ bool FVulkanViewport::Present(FVulkanCommandListContext* Context, FVulkanCommand
 	}
 
 	CmdBuffer->End();
-	FVulkanCommandBufferManager* ImmediateCmdBufMgr = Device->GetImmediateContext().GetCommandBufferManager();
-	ImmediateCmdBufMgr->FlushResetQueryPools();
-	checkf(ImmediateCmdBufMgr->GetActiveCmdBufferDirect() == CmdBuffer, TEXT("Present() is submitting something else than the active command buffer"));
+	FVulkanCommandListContextImmediate& ImmediateContext = Device->GetImmediateContext();
+	ImmediateContext.FlushResetQueryPools();
+	checkf(ImmediateContext.GetActiveCmdBufferDirect() == CmdBuffer, TEXT("Present() is submitting something else than the active command buffer"));
 	if (SupportsStandardSwapchain())
 	{
 		if (LIKELY(!bFailedToDelayAcquireBackbuffer))
@@ -1072,7 +1072,7 @@ bool FVulkanViewport::Present(FVulkanCommandListContext* Context, FVulkanCommand
 			
 			VulkanRHI::FSemaphore* SignalSemaphore = (AcquiredImageIndex >= 0 ? RenderingDoneSemaphores[AcquiredImageIndex] : nullptr);
 			// submit through the CommandBufferManager as it will add the proper semaphore
-			ImmediateCmdBufMgr->SubmitActiveCmdBufferFromPresent(SignalSemaphore);
+			ImmediateContext.SubmitActiveCmdBufferFromPresent(SignalSemaphore);
 		}
 		else
 		{
@@ -1096,7 +1096,7 @@ bool FVulkanViewport::Present(FVulkanCommandListContext* Context, FVulkanCommand
 	else
 	{
 		// Submit active command buffer if not supporting standard swapchain (e.g. XR devices).
-		ImmediateCmdBufMgr->SubmitActiveCmdBufferFromPresent(nullptr);
+		ImmediateContext.SubmitActiveCmdBufferFromPresent(nullptr);
 	}
 
 	//Flush all commands
@@ -1162,9 +1162,9 @@ bool FVulkanViewport::Present(FVulkanCommandListContext* Context, FVulkanCommand
 	//}
 
 	// PrepareForNewActiveCommandBuffer might be called by swapchain re-creation routine. Skip prepare if we already have an open active buffer.
-	if (ImmediateCmdBufMgr->GetActiveCmdBuffer() && !ImmediateCmdBufMgr->GetActiveCmdBuffer()->HasBegun())
+	if (ImmediateContext.GetActiveCmdBuffer() && !ImmediateContext.GetActiveCmdBuffer()->HasBegun())
 	{
-		ImmediateCmdBufMgr->PrepareForNewActiveCommandBuffer();
+		ImmediateContext.PrepareForNewActiveCommandBuffer();
 	}
 
 	AcquiredImageIndex = -1;
