VulkanRHI: Cap the reported maxPerStageDescriptorSamplers (uint32) to max_int32 for GMaxTextureSamplers (int32).  Fix all preview modes being disabled on AMD gpus in Vulkan.

#rnx

diff --git a/Engine/Source/Runtime/VulkanRHI/Private/VulkanRHI.cpp b/Engine/Source/Runtime/VulkanRHI/Private/VulkanRHI.cpp
--- a/Engine/Source/Runtime/VulkanRHI/Private/VulkanRHI.cpp
+++ b/Engine/Source/Runtime/VulkanRHI/Private/VulkanRHI.cpp
@@ -915,7 +915,7 @@ void FVulkanDynamicRHI::InitInstance()
 		GSupportsDualSrcBlending = Device->GetPhysicalDeviceFeatures().Core_1_0.dualSrcBlend == VK_TRUE;
 		GRHISupportsSeparateDepthStencilCopyAccess = Device->SupportsParallelRendering();
 		GRHIBindlessSupport = Device->SupportsBindless() ? RHIGetBindlessSupport(GMaxRHIShaderPlatform) : ERHIBindlessSupport::Unsupported;
-		GMaxTextureSamplers = Props.limits.maxPerStageDescriptorSamplers;
+		GMaxTextureSamplers = (int32)FMath::Min<uint32>(MAX_int32, Props.limits.maxPerStageDescriptorSamplers);
 		GRHISupportsLossyFramebufferCompression = Device->GetOptionalExtensions().HasEXTImageCompressionControl;
 		
 		GRHIMaxDispatchThreadGroupsPerDimension.X = FMath::Min<uint32>(Limits.maxComputeWorkGroupCount[0], 0x7fffffff);
