VulkanRHI: Remove the storage bit when creating read-only views.  Fix validation warning in uav aliasing.

#rnx

diff --git a/Engine/Source/Runtime/VulkanRHI/Private/VulkanUAV.cpp b/Engine/Source/Runtime/VulkanRHI/Private/VulkanUAV.cpp
--- a/Engine/Source/Runtime/VulkanRHI/Private/VulkanUAV.cpp
+++ b/Engine/Source/Runtime/VulkanRHI/Private/VulkanUAV.cpp
@@ -444,6 +444,9 @@ void FVulkanShaderResourceView::UpdateView()
 			checkf((ArrayFirst + ArrayNum) <= Texture->GetNumberOfArrayLevels(), TEXT("View extends beyond original cube texture level count!"));
 		}
 
+		// Remove storage bit from usage flags for SRVs
+		const VkImageUsageFlags ImageUsageFlags = Texture->ImageUsageFlags & ~VK_IMAGE_USAGE_STORAGE_BIT;
+
 		InitAsTextureView(
 			  Texture->Image
 			, GetVkImageViewTypeForDimensionSRV(Info.Dimension, Texture->GetViewType())
@@ -455,6 +458,7 @@ void FVulkanShaderResourceView::UpdateView()
 			, ArrayFirst
 			, ArrayNum
 			, false
+			, ImageUsageFlags
 		);
 	}
 }
