Fix vulkan lru crash when the lru is empty

diff --git a/Engine/Source/Runtime/VulkanRHI/Private/VulkanPipeline.cpp b/Engine/Source/Runtime/VulkanRHI/Private/VulkanPipeline.cpp
--- a/Engine/Source/Runtime/VulkanRHI/Private/VulkanPipeline.cpp
+++ b/Engine/Source/Runtime/VulkanRHI/Private/VulkanPipeline.cpp
@@ -2485,11 +2485,13 @@ void FVulkanPipelineStateCacheManager::TickLRU()
 	const int MaxEvictsPerTick = 5;
 	for(int i = 0 ; i<MaxEvictsPerTick; i++)
 	{
-		uint32 tid = FPlatformTLS::GetCurrentThreadId();
 		FVulkanRHIGraphicsPipelineStateLRUNode* Node = LRU.GetTail();
-		check(Node != 0);
-		TRefCountPtr<FVulkanRHIGraphicsPipelineState> PSO = Node->GetValue();
+		if (!Node)
+		{
+			return;
+		}
 
+		TRefCountPtr<FVulkanRHIGraphicsPipelineState> PSO = Node->GetValue();
 		bool bTimeToDie = PSO->LRUFrame + GVulkanPSOLRUEvictAfterUnusedFrames < GFrameNumberRenderThread;
 		if (bTimeToDie)
 		{
