VulkanRHI:  Remove VulkanRHI's internal PrepareForCPURead and replace it with proper immediate cmdlist flow. Create temporary implementation of the UploadContext.  Move all previous uses of the upload command buffer to the new UploadContext.  Remove UploadCommandBuffer from Vulkan contexts.  Replace the device's DefaultTexture with GBlackTexture as a default descriptor value.

#rnx

diff --git a/Engine/Source/Runtime/VulkanRHI/Private/VulkanBuffer.cpp b/Engine/Source/Runtime/VulkanRHI/Private/VulkanBuffer.cpp
--- a/Engine/Source/Runtime/VulkanRHI/Private/VulkanBuffer.cpp
+++ b/Engine/Source/Runtime/VulkanRHI/Private/VulkanBuffer.cpp
@@ -19,6 +19,7 @@ struct FVulkanPendingBufferLock
 	uint32 Offset = 0;
 	uint32 Size = 0;
 	EResourceLockMode LockMode = RLM_Num;
+	bool FirstLock = false;
 };
 
 static TMap<FVulkanBuffer*, FVulkanPendingBufferLock> GPendingLocks;
@@ -244,33 +245,37 @@ void* FVulkanBuffer::Lock(FRHICommandListBase& RHICmdList, EResourceLockMode Loc
 		}
 		else 
 		{
-			Device->PrepareForCPURead();
-		
-			FVulkanCommandListContextImmediate& Context = Device->GetImmediateContext();
-			FVulkanCommandBuffer* CmdBuffer = Context.GetUploadCmdBuffer();
-				
-			// Make sure any previous tasks have finished on the source buffer.
-			VkMemoryBarrier BarrierBefore = { VK_STRUCTURE_TYPE_MEMORY_BARRIER, nullptr, VK_ACCESS_MEMORY_WRITE_BIT, VK_ACCESS_MEMORY_READ_BIT };
-			VulkanRHI::vkCmdPipelineBarrier(CmdBuffer->GetHandle(), VK_PIPELINE_STAGE_BOTTOM_OF_PIPE_BIT, VK_PIPELINE_STAGE_TRANSFER_BIT, 0, 1, &BarrierBefore, 0, nullptr, 0, nullptr);
-
 			// Create a staging buffer we can use to copy data from device to cpu.
 			VulkanRHI::FStagingBuffer* StagingBuffer = Device->GetStagingManager().AcquireBuffer(LockSize, VK_BUFFER_USAGE_TRANSFER_DST_BIT, VK_MEMORY_PROPERTY_HOST_CACHED_BIT);
 
-			// Fill the staging buffer with the data on the device.
 			VkBufferCopy Regions;
 			Regions.size = LockSize;
 			Regions.srcOffset = Offset + CurrentBufferAlloc.Alloc.Offset;
 			Regions.dstOffset = 0;
-				
-			VulkanRHI::vkCmdCopyBuffer(CmdBuffer->GetHandle(), CurrentBufferAlloc.Alloc.GetBufferHandle(), StagingBuffer->GetHandle(), 1, &Regions);
 
-			// Setup barrier.
-			VkMemoryBarrier BarrierAfter = { VK_STRUCTURE_TYPE_MEMORY_BARRIER, nullptr, VK_ACCESS_MEMORY_WRITE_BIT, VK_ACCESS_HOST_READ_BIT };
-			VulkanRHI::vkCmdPipelineBarrier(CmdBuffer->GetHandle(), VK_PIPELINE_STAGE_TRANSFER_BIT, VK_PIPELINE_STAGE_HOST_BIT, 0, 1, &BarrierAfter, 0, nullptr, 0, nullptr);
-				
-			// Force upload.
-			Context.SubmitUploadCmdBuffer();
-			Device->WaitUntilIdle();
+			VkBuffer BufferHandle = CurrentBufferAlloc.Alloc.GetBufferHandle();
+
+			FRHICommandListImmediate& ImmCmdList = RHICmdList.GetAsImmediate();
+			ImmCmdList.EnqueueLambda([StagingBuffer, Regions, BufferHandle](FRHICommandListBase& ExecutingCmdList)
+				{
+					FVulkanCommandListContext& Context = FVulkanCommandListContext::Get(ExecutingCmdList);
+					FVulkanCommandBuffer& CommandBuffer = Context.GetCommandBuffer();
+
+					// Make sure any previous tasks have finished on the source buffer.
+					VkMemoryBarrier BarrierBefore = { VK_STRUCTURE_TYPE_MEMORY_BARRIER, nullptr, VK_ACCESS_MEMORY_WRITE_BIT, VK_ACCESS_MEMORY_READ_BIT };
+					VulkanRHI::vkCmdPipelineBarrier(CommandBuffer.GetHandle(), VK_PIPELINE_STAGE_BOTTOM_OF_PIPE_BIT, VK_PIPELINE_STAGE_TRANSFER_BIT, 0, 1, &BarrierBefore, 0, nullptr, 0, nullptr);
+
+					// Fill the staging buffer with the data on the device.
+					VulkanRHI::vkCmdCopyBuffer(CommandBuffer.GetHandle(), BufferHandle, StagingBuffer->GetHandle(), 1, &Regions);
+
+					// Setup barrier.
+					VkMemoryBarrier BarrierAfter = { VK_STRUCTURE_TYPE_MEMORY_BARRIER, nullptr, VK_ACCESS_MEMORY_WRITE_BIT, VK_ACCESS_HOST_READ_BIT };
+					VulkanRHI::vkCmdPipelineBarrier(CommandBuffer.GetHandle(), VK_PIPELINE_STAGE_TRANSFER_BIT, VK_PIPELINE_STAGE_HOST_BIT, 0, 1, &BarrierAfter, 0, nullptr, 0, nullptr);
+
+				});
+
+			// We need to execute the command list so we can read the data from the map below
+			ImmCmdList.SubmitAndBlockUntilGPUIdle();
 
 			// Flush.
 			StagingBuffer->FlushMappedMemory();
@@ -285,8 +290,6 @@ void* FVulkanBuffer::Lock(FRHICommandListBase& RHICmdList, EResourceLockMode Loc
 			PendingLock.LockMode = LockMode;
 			PendingLock.StagingBuffer = StagingBuffer;
 			AddPendingBufferLock(this, PendingLock);
-
-			Context.PrepareForNewActiveCommandBuffer();
 		}
 	}
 	else
@@ -294,7 +297,8 @@ void* FVulkanBuffer::Lock(FRHICommandListBase& RHICmdList, EResourceLockMode Loc
 		check(LockMode == RLM_WriteOnly);
 
 		// If this is the first lock on host visible memory, then the memory is still untouched so use it directly
-		if ((bUnifiedMem || bDynamic) && bIsFirstLock)
+		const bool bIsHostVisible = (bUnifiedMem || bDynamic);
+		if (bIsHostVisible && bIsFirstLock)
 		{
 			check(CurrentBufferAlloc.HostPtr);
 			Data = CurrentBufferAlloc.HostPtr;
@@ -306,7 +310,8 @@ void* FVulkanBuffer::Lock(FRHICommandListBase& RHICmdList, EResourceLockMode Loc
 			FVulkanPendingBufferLock PendingLock;
 			PendingLock.Offset = Offset;
 			PendingLock.Size = LockSize;
-			PendingLock.LockMode = LockMode;
+			PendingLock.LockMode = LockMode; 
+			PendingLock.FirstLock = bIsFirstLock;
 
 			VulkanRHI::FStagingBuffer* StagingBuffer = Device->GetStagingManager().AcquireBuffer(LockSize, VK_BUFFER_USAGE_TRANSFER_SRC_BIT);
 			PendingLock.StagingBuffer = StagingBuffer;
@@ -372,11 +377,12 @@ void FVulkanBuffer::Unlock(FRHICommandListBase& RHICmdList)
 			}
 			else if (PendingLock.LockMode == RLM_WriteOnly)
 			{
-				FVulkanCommandListContext& Context = FVulkanCommandListContext::GetVulkanContext(CmdList.GetContext());
-
 				// We need to do this on the active command buffer instead of using an upload command buffer. The high level code sometimes reuses the same
 				// buffer in sequences of upload / dispatch, upload / dispatch, so we need to order the copy commands correctly with respect to the dispatches.
-				FVulkanCommandBuffer& CommandBuffer = Context.GetCommandBuffer();
+				// Unless this is the first time any data is pushed into this buffer, then ordering doesn't matter and UploadContext can be used
+				FVulkanCommandBuffer& CommandBuffer = PendingLock.FirstLock ? 
+					FVulkanUploadContext::Get(CmdList).GetCommandBuffer() :
+					FVulkanCommandListContext::Get(CmdList).GetCommandBuffer();
 				check(CommandBuffer.IsOutsideRenderPass());
 				VkCommandBuffer CommandBufferHandle = CommandBuffer.GetHandle();
 
@@ -394,8 +400,6 @@ void FVulkanBuffer::Unlock(FRHICommandListBase& RHICmdList)
 				VulkanRHI::vkCmdPipelineBarrier(CommandBufferHandle, VK_PIPELINE_STAGE_TRANSFER_BIT, VK_PIPELINE_STAGE_ALL_COMMANDS_BIT, 0, 1, &BarrierAfter, 0, nullptr, 0, nullptr);
 
 				Buffer->GetParent()->GetStagingManager().ReleaseBuffer(&CommandBuffer, StagingBuffer);
-
-				Buffer->UpdateLinkedViews(); // :todo-jn:  not needed?  Same buffer when we use staging?
 			}
 		});
 	}
diff --git a/Engine/Source/Runtime/VulkanRHI/Private/VulkanCommands.cpp b/Engine/Source/Runtime/VulkanRHI/Private/VulkanCommands.cpp
--- a/Engine/Source/Runtime/VulkanRHI/Private/VulkanCommands.cpp
+++ b/Engine/Source/Runtime/VulkanRHI/Private/VulkanCommands.cpp
@@ -897,10 +897,6 @@ void FVulkanCommandListContext::RequestSubmitCurrentCommands()
 {
 	if (Device.GetComputeQueue() == &Queue)
 	{
-		if (HasPendingUploadCmdBuffer())
-		{
-			SubmitUploadCmdBuffer();
-		}
 		bSubmitAtNextSafePoint = true;
 		SafePointSubmit();
 	}
@@ -917,23 +913,6 @@ void FVulkanCommandListContext::InternalSubmitActiveCmdBuffer()
 	PrepareForNewActiveCommandBuffer();
 }
 
-void FVulkanCommandListContext::PrepareForCPURead()
-{
-	ensure(IsImmediate());
-	FVulkanCommandBuffer* CmdBuffer = GetActiveCmdBuffer();
-	if (CmdBuffer && CmdBuffer->HasBegun())
-	{
-		check(!CmdBuffer->IsInsideRenderPass());
-
-		SubmitActiveCmdBuffer();
-		if (!GWaitForIdleOnSubmit)
-		{
-			// The wait has already happened if GWaitForIdleOnSubmit is set
-			WaitForCmdBuffer(CmdBuffer);
-		}
-	}
-}
-
 void FVulkanCommandListContext::RHICopyToStagingBuffer(FRHIBuffer* SourceBufferRHI, FRHIStagingBuffer* StagingBufferRHI, uint32 Offset, uint32 NumBytes)
 {
 	FVulkanCommandBuffer& CommandBuffer = GetCommandBuffer();
@@ -979,6 +958,7 @@ void FVulkanCommandListContext::RHIWriteGPUFence(FRHIGPUFence* FenceRHI)
 struct FVulkanPlatformCommandList : public IRHIPlatformCommandList
 {
 	FVulkanCommandListContext* CmdContext = nullptr;
+	FVulkanUploadContext* UploadContext = nullptr;
 };
 
 template<>
@@ -1015,11 +995,17 @@ IRHIComputeContext* FVulkanDynamicRHI::RHIGetCommandContext(ERHIPipeline Pipelin
 
 void FVulkanDynamicRHI::RHIFinalizeContext(FRHIFinalizeContextArgs&& Args, TRHIPipelineArray<IRHIPlatformCommandList*>& Output)
 {
-	for(IRHIComputeContext* Context : Args.Contexts)
+	for (IRHIComputeContext* Context : Args.Contexts)
 	{
 		FVulkanPlatformCommandList* PlatformCmdList = new FVulkanPlatformCommandList();
 		PlatformCmdList->CmdContext = static_cast<FVulkanCommandListContext*>(Context);
-		
+
+		// :todo-jn: Temporarily just place it before the graphic pipeline
+		if (Args.UploadContext && (Context->GetPipeline() == ERHIPipeline::Graphics))
+		{
+			PlatformCmdList->UploadContext = static_cast<FVulkanUploadContext*>(Args.UploadContext);
+		}
+
 		Output[Context->GetPipeline()] = PlatformCmdList;
 	}
 }
@@ -1029,8 +1015,16 @@ void FVulkanDynamicRHI::RHISubmitCommandLists(FRHISubmitCommandListsArgs&& Args)
 	for (IRHIPlatformCommandList* Ptr : Args.CommandLists)
 	{
 		FVulkanPlatformCommandList* PlatformCmdList = ResourceCast(Ptr);
-		FVulkanCommandListContext* CurrentContext = PlatformCmdList->CmdContext;
 
+		FVulkanUploadContext* UploadContext = PlatformCmdList->UploadContext;
+		if (UploadContext)
+		{
+			UploadContext->SubmitActiveCmdBuffer();
+			UploadContext->RefreshFenceStatus();
+			FVulkanUploadContext::Pool.Push(PlatformCmdList->UploadContext);
+		}
+
+		FVulkanCommandListContext* CurrentContext = PlatformCmdList->CmdContext;
 		if (CurrentContext->IsImmediate())
 		{
 			CurrentContext->RequestSubmitCurrentCommands();
@@ -1043,7 +1037,6 @@ void FVulkanDynamicRHI::RHISubmitCommandLists(FRHISubmitCommandListsArgs&& Args)
 		}
 		else
 		{
-			check(!CurrentContext->HasPendingUploadCmdBuffer());  // todo-jn
 			FVulkanCommandBuffer* CmdBuffer = CurrentContext->GetActiveCmdBuffer();
 			check(!CmdBuffer->IsInsideRenderPass());
 			CurrentContext->SubmitActiveCmdBuffer();
diff --git a/Engine/Source/Runtime/VulkanRHI/Private/VulkanContext.cpp b/Engine/Source/Runtime/VulkanRHI/Private/VulkanContext.cpp
--- a/Engine/Source/Runtime/VulkanRHI/Private/VulkanContext.cpp
+++ b/Engine/Source/Runtime/VulkanRHI/Private/VulkanContext.cpp
@@ -8,16 +8,6 @@
 #include "VulkanDevice.h"
 #include "VulkanPendingState.h"
 
-static int32 GVulkanUploadCmdBufferSemaphore = 0;	// despite this behavior being more correct, default to off to avoid changing the existing (as of UE 4.24) behavior
-static FAutoConsoleVariableRef CVarVulkanPreventOverlapWithUpload(
-	TEXT("r.Vulkan.UploadCmdBufferSemaphore"),
-	GVulkanUploadCmdBufferSemaphore,
-	TEXT("Whether command buffers for uploads and graphics can be executed simultaneously.\n")
-	TEXT(" 0: The buffers are submitted without any synch(default)\n")
-	TEXT(" 1: Graphics buffers will not overlap with the upload buffers"),
-	ECVF_ReadOnly
-);
-
 int32 GVulkanTempBlockSize = 4 * 1024 * 1024;
 static FAutoConsoleVariableRef GCVarVulkanTempBlockSize(
 	TEXT("r.Vulkan.TempBlockSize"),
@@ -36,11 +26,6 @@ FVulkanContextCommon::FVulkanContextCommon(FVulkanDevice& InDevice, FVulkanQueue
 {
 	Pool.Create(Queue.GetFamilyIndex());
 
-	if (GVulkanUploadCmdBufferSemaphore)
-	{
-		ActiveCmdBufferSemaphore = new VulkanRHI::FSemaphore(Device);
-	}
-
 	ActiveCmdBuffer = Pool.Create(false);
 	ActiveCmdBuffer->InitializeTimings(this);
 	ActiveCmdBuffer->Begin();
@@ -103,62 +88,10 @@ void FVulkanContextCommon::FlushResetQueryPools()
 	}
 }
 
-void FVulkanContextCommon::SubmitUploadCmdBuffer(uint32 NumSignalSemaphores, VkSemaphore* SignalSemaphores)
-{
-	FScopeLock ScopeLock(&Pool.CS);
-	check(UploadCmdBuffer);
-	check(UploadCmdBuffer->CurrentDescriptorPoolSetContainer == nullptr);
-	if (!UploadCmdBuffer->IsSubmitted() && UploadCmdBuffer->HasBegun())
-	{
-		check(UploadCmdBuffer->IsOutsideRenderPass());
-
-		VulkanRHI::DebugHeavyWeightBarrier(UploadCmdBuffer->GetHandle(), 4);
-
-		UploadCmdBuffer->End();
-
-		if (GVulkanUploadCmdBufferSemaphore)
-		{
-			// Add semaphores associated with the recent active cmdbuf(s), if any. That will prevent
-			// the overlap, delaying execution of this cmdbuf until the graphics one(s) is complete.
-			for (VulkanRHI::FSemaphore* WaitForThis : RenderingCompletedSemaphores)
-			{
-				UploadCmdBuffer->AddWaitSemaphore(VK_PIPELINE_STAGE_TOP_OF_PIPE_BIT, WaitForThis);
-			}
-
-			if (NumSignalSemaphores == 0)
-			{
-				checkf(UploadCmdBufferSemaphore != nullptr, TEXT("FVulkanContextCommon::SubmitUploadCmdBuffer: upload command buffer does not have an associated completion semaphore."));
-				VkSemaphore Sema = UploadCmdBufferSemaphore->GetHandle();
-				Queue.Submit(UploadCmdBuffer, 1, &Sema);
-				UploadCompletedSemaphores.Add(UploadCmdBufferSemaphore);
-				UploadCmdBufferSemaphore = nullptr;
-			}
-			else
-			{
-				TArray<VkSemaphore, TInlineAllocator<16>> CombinedSemaphores;
-				CombinedSemaphores.Add(UploadCmdBufferSemaphore->GetHandle());
-				CombinedSemaphores.Append(SignalSemaphores, NumSignalSemaphores);
-				Queue.Submit(UploadCmdBuffer, CombinedSemaphores.Num(), CombinedSemaphores.GetData());
-			}
-			// the buffer will now hold on to the wait semaphores, so we can clear them here
-			RenderingCompletedSemaphores.Empty();
-		}
-		else
-		{
-			Queue.Submit(UploadCmdBuffer, NumSignalSemaphores, SignalSemaphores);
-		}
-
-		UploadCmdBuffer->SubmittedTime = FPlatformTime::Seconds();
-	}
-
-	UploadCmdBuffer = nullptr;
-}
-
 void FVulkanContextCommon::SubmitActiveCmdBuffer(TArrayView<VulkanRHI::FSemaphore*> SignalSemaphores)
 {
 	FScopeLock ScopeLock(&Pool.CS);
 	FlushResetQueryPools();
-	check(!UploadCmdBuffer);
 	check(ActiveCmdBuffer);
 
 	TArray<VkSemaphore, FConcurrentLinearArrayAllocator> SemaphoreHandles;
@@ -180,78 +113,23 @@ void FVulkanContextCommon::SubmitActiveCmdBuffer(TArrayView<VulkanRHI::FSemaphor
 
 		ActiveCmdBuffer->End();
 
-		if (GVulkanUploadCmdBufferSemaphore)
-		{
-			checkf(ActiveCmdBufferSemaphore != nullptr, TEXT("FVulkanContextCommon::SubmitUploadCmdBuffer: graphics command buffer does not have an associated completion semaphore."));
-
-			// Add semaphores associated with the recent upload cmdbuf(s), if any. That will prevent
-			// the overlap, delaying execution of this cmdbuf until upload one(s) are complete.
-			for (VulkanRHI::FSemaphore* UploadCompleteSema : UploadCompletedSemaphores)
-			{
-				ActiveCmdBuffer->AddWaitSemaphore(VK_PIPELINE_STAGE_TOP_OF_PIPE_BIT, UploadCompleteSema);
-			}
+		Queue.Submit(ActiveCmdBuffer, SemaphoreHandles.Num(), SemaphoreHandles.GetData());
 
-			SemaphoreHandles.Add(ActiveCmdBufferSemaphore->GetHandle());
-			Queue.Submit(ActiveCmdBuffer, SemaphoreHandles.Num(), SemaphoreHandles.GetData());
-
-			RenderingCompletedSemaphores.Add(ActiveCmdBufferSemaphore);
-			ActiveCmdBufferSemaphore = nullptr;
-
-			// the buffer will now hold on to the wait semaphores, so we can clear them here
-			UploadCompletedSemaphores.Empty();
-		}
-		else
-		{
-			Queue.Submit(ActiveCmdBuffer, SemaphoreHandles.Num(), SemaphoreHandles.GetData());
-		}
 		ActiveCmdBuffer->SubmittedTime = FPlatformTime::Seconds();
 	}
 
 	ActiveCmdBuffer = nullptr;
-	if (GVulkanUploadCmdBufferSemaphore && ActiveCmdBufferSemaphore != nullptr)
-	{
-		// most likely we didn't submit as it wasn't begun
-		check(ActiveCmdBufferSemaphore->GetRefCount() == 0);
-		delete ActiveCmdBufferSemaphore;
-		ActiveCmdBufferSemaphore = nullptr;
-	}
 }
 
 void FVulkanContextCommon::SubmitActiveCmdBufferFromPresent(VulkanRHI::FSemaphore* SignalSemaphore)
 {
-	if (GVulkanUploadCmdBufferSemaphore)
+	if (SignalSemaphore)
 	{
-		// unlike more advanced regular SACB(), this is just a wrapper
-		// around Queue.Submit() to avoid rewriting the logic in Present
-		if (SignalSemaphore)
-		{
-			VkSemaphore SignalThis[2] =
-			{
-				SignalSemaphore->GetHandle(),
-				ActiveCmdBufferSemaphore->GetHandle()
-			};
-
-			Queue.Submit(ActiveCmdBuffer, UE_ARRAY_COUNT(SignalThis), SignalThis);
-		}
-		else
-		{
-			VkSemaphore SignalThis = ActiveCmdBufferSemaphore->GetHandle();
-			Queue.Submit(ActiveCmdBuffer, 1, &SignalThis);
-		}
-
-		RenderingCompletedSemaphores.Add(ActiveCmdBufferSemaphore);
-		ActiveCmdBufferSemaphore = nullptr;
+		Queue.Submit(ActiveCmdBuffer, SignalSemaphore->GetHandle());
 	}
 	else
 	{
-		if (SignalSemaphore)
-		{
-			Queue.Submit(ActiveCmdBuffer, SignalSemaphore->GetHandle());
-		}
-		else
-		{
-			Queue.Submit(ActiveCmdBuffer);
-		}
+		Queue.Submit(ActiveCmdBuffer);
 	}
 }
 
@@ -295,14 +173,6 @@ void FVulkanContextCommon::FreeUnusedCmdBuffers(bool bTrimMemory)
 void FVulkanContextCommon::PrepareForNewActiveCommandBuffer()
 {
 	FScopeLock ScopeLock(&Pool.CS);
-	check(!UploadCmdBuffer);
-
-	if (GVulkanUploadCmdBufferSemaphore)
-	{
-		// create a completion semaphore
-		checkf(ActiveCmdBufferSemaphore == nullptr, TEXT("FVulkanContextCommon::PrepareForNewActiveCommandBuffer: a semaphore from the previous active buffer has been leaked."));
-		ActiveCmdBufferSemaphore = new VulkanRHI::FSemaphore(Device);
-	}
 
 	for (int32 Index = 0; Index < Pool.CmdBuffers.Num(); ++Index)
 	{
@@ -346,44 +216,6 @@ uint32 FVulkanContextCommon::CalculateGPUTime()
 }
 #endif // (RHI_NEW_GPU_PROFILER == 0)
 
-FVulkanCommandBuffer* FVulkanContextCommon::GetUploadCmdBuffer()
-{
-	FScopeLock ScopeLock(&Pool.CS);
-	if (!UploadCmdBuffer)
-	{
-		if (GVulkanUploadCmdBufferSemaphore)
-		{
-			// create a completion semaphore
-			checkf(UploadCmdBufferSemaphore == nullptr, TEXT("FVulkanContextCommon::GetUploadCmdBuffer: a semaphore from the previous upload buffer has been leaked."));
-			UploadCmdBufferSemaphore = new VulkanRHI::FSemaphore(Device);
-		}
-
-		for (int32 Index = 0; Index < Pool.CmdBuffers.Num(); ++Index)
-		{
-			FVulkanCommandBuffer* CmdBuffer = Pool.CmdBuffers[Index];
-			CmdBuffer->RefreshFenceStatus();
-#if VULKAN_USE_DIFFERENT_POOL_CMDBUFFERS
-			if (CmdBuffer->bIsUploadOnly)
-#endif
-			{
-				if (CmdBuffer->State == FVulkanCommandBuffer::EState::ReadyForBegin || CmdBuffer->State == FVulkanCommandBuffer::EState::NeedReset)
-				{
-					UploadCmdBuffer = CmdBuffer;
-					UploadCmdBuffer->Begin();
-					return UploadCmdBuffer;
-				}
-			}
-		}
-
-		// All cmd buffers are being executed still
-		UploadCmdBuffer = Pool.Create(true);
-		UploadCmdBuffer->Begin();
-	}
-
-	return UploadCmdBuffer;
-}
-
-
 
 FVulkanCommandListContext::FVulkanCommandListContext(FVulkanDynamicRHI* InRHI, FVulkanDevice& InDevice, FVulkanQueue& InQueue, FVulkanCommandListContext* InImmediate)
 	: FVulkanContextCommon(InDevice, InQueue)
@@ -455,4 +287,41 @@ FVulkanCommandListContext::~FVulkanCommandListContext()
 	delete TempBlockAllocator;
 	delete PendingGfxState;
 	delete PendingComputeState;
+}
+
+
+
+TLockFreePointerListUnordered<FVulkanUploadContext, PLATFORM_CACHE_LINE_SIZE> FVulkanUploadContext::Pool;
+void FVulkanUploadContext::DestroyPool()
+{
+	while (FVulkanUploadContext* Context = FVulkanUploadContext::Pool.Pop())
+	{
+		delete Context;
+	}
+}
+
+
+FVulkanUploadContext::FVulkanUploadContext(FVulkanDevice& InDevice, FVulkanQueue& InQueue)
+	: FVulkanContextCommon(InDevice, InQueue)
+{
+
+}
+FVulkanUploadContext::~FVulkanUploadContext()
+{
+
+}
+
+IRHIUploadContext* FVulkanDynamicRHI::RHIGetUploadContext()
+{
+	FVulkanUploadContext* Context = FVulkanUploadContext::Pool.Pop();
+	if (Context)
+	{
+		Context->PrepareForNewActiveCommandBuffer();
+	}
+	else
+	{
+		// :todo-jn: locked to graphics queue for now
+		Context = new FVulkanUploadContext(*Device, *Device->GetGraphicsQueue());
+	}
+	return Context;
 }
\ No newline at end of file
diff --git a/Engine/Source/Runtime/VulkanRHI/Private/VulkanContext.h b/Engine/Source/Runtime/VulkanRHI/Private/VulkanContext.h
--- a/Engine/Source/Runtime/VulkanRHI/Private/VulkanContext.h
+++ b/Engine/Source/Runtime/VulkanRHI/Private/VulkanContext.h
@@ -23,19 +23,9 @@ class FVulkanContextCommon
 {
 public:
 	FVulkanContextCommon(FVulkanDevice& InDevice, FVulkanQueue& InQueue);
-	~FVulkanContextCommon();
+	virtual ~FVulkanContextCommon();
 
 	FVulkanCommandBuffer* GetActiveCmdBuffer()
-	{
-		if (UploadCmdBuffer)
-		{
-			SubmitUploadCmdBuffer();
-		}
-
-		return ActiveCmdBuffer;
-	}
-
-	FVulkanCommandBuffer* GetActiveCmdBufferDirect()
 	{
 		return ActiveCmdBuffer;
 	}
@@ -47,20 +37,11 @@ public:
 		return *CommandBuffer;
 	}
 
-	bool HasPendingUploadCmdBuffer() const
-	{
-		return UploadCmdBuffer != nullptr;
-	}
-
 	bool HasPendingActiveCmdBuffer() const
 	{
 		return ActiveCmdBuffer != nullptr;
 	}
 
-	VULKANRHI_API FVulkanCommandBuffer* GetUploadCmdBuffer();
-
-	VULKANRHI_API void SubmitUploadCmdBuffer(uint32 NumSignalSemaphores = 0, VkSemaphore* SignalSemaphores = nullptr);
-
 	void SubmitActiveCmdBuffer(TArrayView<VulkanRHI::FSemaphore*> SignalSemaphores);
 
 	void SubmitActiveCmdBuffer()
@@ -78,7 +59,6 @@ public:
 
 	void WaitForCmdBuffer(FVulkanCommandBuffer* CmdBuffer, float TimeInSecondsToWait = 10.0f);
 
-
 	void AddQueryPoolForReset(VkQueryPool Pool, uint32 Size);
 	void FlushResetQueryPools();
 
@@ -103,10 +83,6 @@ public:
 
 	void NotifyDeletedImage(VkImage Image)
 	{
-		if (UploadCmdBuffer)
-		{
-			UploadCmdBuffer->GetLayoutManager().NotifyDeletedImage(Image);
-		}
 		if (ActiveCmdBuffer)
 		{
 			ActiveCmdBuffer->GetLayoutManager().NotifyDeletedImage(Image);
@@ -138,19 +114,9 @@ private:
 	TArray<FQueryPoolReset> PoolResets;
 
 	FVulkanCommandBuffer* ActiveCmdBuffer = nullptr;
-	FVulkanCommandBuffer* UploadCmdBuffer = nullptr;
-
-	// This semaphore is used to prevent overlaps between the (current) graphics cmdbuf and next upload cmdbuf.
-	VulkanRHI::FSemaphore* ActiveCmdBufferSemaphore = nullptr;
 
 	// Holds semaphores associated with the recent upload cmdbuf(s) - waiting to be added to the next graphics cmdbuf as WaitSemaphores.
 	TArray<VulkanRHI::FSemaphore*> RenderingCompletedSemaphores;
-
-	// This semaphore is used to prevent overlaps between (current) upload cmdbuf and next graphics cmdbuf.
-	VulkanRHI::FSemaphore* UploadCmdBufferSemaphore = nullptr;
-
-	// Holds semaphores associated with the recent upload cmdbuf(s) - waiting to be added to the next graphics cmdbuf as WaitSemaphores.
-	TArray<VulkanRHI::FSemaphore*> UploadCompletedSemaphores;
 };
 
 
@@ -161,9 +127,9 @@ public:
 	FVulkanCommandListContext(FVulkanDynamicRHI* InRHI, FVulkanDevice& InDevice, FVulkanQueue& InQueue, FVulkanCommandListContext* InImmediate);
 	virtual ~FVulkanCommandListContext();
 
-	static inline FVulkanCommandListContext& GetVulkanContext(IRHICommandContext& CmdContext)
+	static inline FVulkanCommandListContext& Get(FRHICommandListBase& RHICmdList)
 	{
-		return static_cast<FVulkanCommandListContext&>(CmdContext.GetLowestLevelContext());
+		return static_cast<FVulkanCommandListContext&>(RHICmdList.GetContext().GetLowestLevelContext());
 	}
 
 	inline bool IsImmediate() const
@@ -347,7 +313,6 @@ protected:
 	// Bit array to track which uniform buffers have changed since the last draw call.
 	uint16 DirtyUniformBuffers[SF_NumStandardFrequencies] = {};
 
-	void PrepareForCPURead();
 	void RequestSubmitCurrentCommands();
 
 	void InternalClearMRT(FVulkanCommandBuffer& CommandBuffer, bool bClearColor, int32 NumClearColors, const FLinearColor* ColorArray, bool bClearDepth, float Depth, bool bClearStencil, uint32 Stencil);
@@ -402,9 +367,30 @@ private:
 class FVulkanCommandListContextImmediate : public FVulkanCommandListContext
 {
 public:
+	static inline FVulkanCommandListContextImmediate& Get(FRHICommandListBase& RHICmdList)
+	{
+		return static_cast<FVulkanCommandListContextImmediate&>(RHICmdList.GetContext().GetLowestLevelContext());
+	}
+
 	FVulkanCommandListContextImmediate(FVulkanDynamicRHI* InRHI, FVulkanDevice& InDevice, FVulkanQueue& InQueue);
 };
 
+class FVulkanUploadContext : public FVulkanContextCommon, public IRHIUploadContext
+{
+public:
+
+	static inline FVulkanUploadContext& Get(FRHICommandListBase& RHICmdList)
+	{
+		return static_cast<FVulkanUploadContext&>(RHICmdList.GetUploadContext());
+	}
+
+	FVulkanUploadContext(FVulkanDevice& Device, FVulkanQueue& InQueue);
+	virtual ~FVulkanUploadContext();
+
+	static TLockFreePointerListUnordered<FVulkanUploadContext, PLATFORM_CACHE_LINE_SIZE> Pool;
+	static void DestroyPool();
+};
+
 struct FVulkanContextArray : public TRHIPipelineArray<FVulkanCommandListContext*>
 {
 	FVulkanContextArray(FRHIContextArray const& Contexts)
diff --git a/Engine/Source/Runtime/VulkanRHI/Private/VulkanDevice.cpp b/Engine/Source/Runtime/VulkanRHI/Private/VulkanDevice.cpp
--- a/Engine/Source/Runtime/VulkanRHI/Private/VulkanDevice.cpp
+++ b/Engine/Source/Runtime/VulkanRHI/Private/VulkanDevice.cpp
@@ -1479,6 +1479,8 @@ void FVulkanDevice::Destroy()
 	delete BindlessDescriptorManager;
 	BindlessDescriptorManager = nullptr;
 
+	FVulkanUploadContext::DestroyPool();
+
 	MemoryManager.Deinit();
 
 	delete TransferQueue;
@@ -1521,18 +1523,8 @@ void FVulkanDevice::NotifyDeletedImage(VkImage Image, bool bRenderTarget)
 	GetImmediateContext().NotifyDeletedImage(Image);
 }
 
-void FVulkanDevice::PrepareForCPURead()
-{
-	//#todo-rco: Process other contexts first!
-	ImmediateContext->PrepareForCPURead();
-}
-
 void FVulkanDevice::SubmitCommands(FVulkanCommandListContext* Context)
 {
-	if (Context->HasPendingUploadCmdBuffer())
-	{
-		Context->SubmitUploadCmdBuffer();
-	}
 	if (Context->HasPendingActiveCmdBuffer())
 	{
 		Context->SubmitActiveCmdBuffer();
diff --git a/Engine/Source/Runtime/VulkanRHI/Private/VulkanDevice.h b/Engine/Source/Runtime/VulkanRHI/Private/VulkanDevice.h
--- a/Engine/Source/Runtime/VulkanRHI/Private/VulkanDevice.h
+++ b/Engine/Source/Runtime/VulkanRHI/Private/VulkanDevice.h
@@ -380,11 +380,6 @@ public:
 		return *DefaultSampler;
 	}
 
-	inline const FVulkanView::FTextureView& GetDefaultImageView() const
-	{
-		return DefaultTexture->DefaultView->GetTextureView();
-	}
-
 	const VkFormatProperties& GetFormatProperties(VkFormat InFormat) const;
 
 	inline VulkanRHI::FDeviceMemoryManager& GetDeviceMemoryManager()
@@ -471,8 +466,6 @@ public:
 	}
 #endif
 
-	void PrepareForCPURead();
-
 	void SubmitCommandsAndFlushGPU();
 
 	FVulkanOcclusionQueryPool* AcquireOcclusionQueryPool(FVulkanContextCommon* InCommonContext, uint32 NumQueries);
diff --git a/Engine/Source/Runtime/VulkanRHI/Private/VulkanPipelineState.cpp b/Engine/Source/Runtime/VulkanRHI/Private/VulkanPipelineState.cpp
--- a/Engine/Source/Runtime/VulkanRHI/Private/VulkanPipelineState.cpp
+++ b/Engine/Source/Runtime/VulkanRHI/Private/VulkanPipelineState.cpp
@@ -13,6 +13,7 @@
 #include "VulkanPipeline.h"
 #include "VulkanLLM.h"
 #include "RHICoreShader.h"
+#include "GlobalRenderResources.h"  // For GBlackTexture
 
 enum
 {
@@ -91,7 +92,7 @@ void FVulkanCommonPipelineDescriptorState::CreateDescriptorWriteInfos()
 	DSWriter.AddDefaulted(MaxNumSets);
 
 	const FVulkanSamplerState& DefaultSampler = Device->GetDefaultSampler();
-	const FVulkanView::FTextureView& DefaultImageView = Device->GetDefaultImageView();
+	const FVulkanView::FTextureView& DefaultImageView = ResourceCast(GBlackTexture->TextureRHI)->DefaultView->GetTextureView();
 
 	FVulkanHashableDescriptorInfo* CurrentHashableDescriptorInfo = nullptr;
 	if (UseVulkanDescriptorCache())
diff --git a/Engine/Source/Runtime/VulkanRHI/Private/VulkanQueue.cpp b/Engine/Source/Runtime/VulkanRHI/Private/VulkanQueue.cpp
--- a/Engine/Source/Runtime/VulkanRHI/Private/VulkanQueue.cpp
+++ b/Engine/Source/Runtime/VulkanRHI/Private/VulkanQueue.cpp
@@ -14,10 +14,8 @@ FAutoConsoleVariableRef CVarVulkanWaitForIdleOnSubmit(
 	TEXT("r.Vulkan.WaitForIdleOnSubmit"),
 	GWaitForIdleOnSubmit,
 	TEXT("Waits for the GPU to be idle after submitting a command buffer. Useful for tracking GPU hangs.\n")
-	TEXT(" 0: Do not wait(default)\n")
-	TEXT(" 1: Wait on every submit\n")
-	TEXT(" 2: Wait when submitting an upload buffer\n")
-	TEXT(" 3: Wait when submitting an active buffer (one that has gfx commands)\n"),
+	TEXT(" 0: Do not wait (default)\n")
+	TEXT(" 1: Wait on every submit\n"),
 	ECVF_Default
 	);
 
@@ -78,34 +76,10 @@ void FVulkanQueue::Submit(FVulkanCommandBuffer* CmdBuffer, uint32 NumSignalSemap
 	CmdBuffer->MarkSemaphoresAsSubmitted();
 	CmdBuffer->SubmittedFenceCounter = CmdBuffer->FenceSignaledCounter;
 
-	bool bShouldStall = false;
-
 	if (GWaitForIdleOnSubmit != 0)
 	{
-		FVulkanContextCommon& Context = CmdBuffer->GetOwner()->GetContext();
-
-		switch(GWaitForIdleOnSubmit)
-		{
-			default:
-				// intentional fall-through
-			case 1:
-				bShouldStall = true;
-				break;
-
-			case 2:
-				bShouldStall = (Context.HasPendingUploadCmdBuffer() && Context.GetUploadCmdBuffer() == CmdBuffer);
-				break;
-
-			case 3:
-				bShouldStall = (Context.HasPendingActiveCmdBuffer() && Context.GetActiveCmdBufferDirect() == CmdBuffer);
-				break;
-		}
-	}
-
-	if (bShouldStall)
-	{
-		// 200 ms timeout
-		bool bSuccess = Device->GetFenceManager().WaitForFence(CmdBuffer->Fence, 200 * 1000 * 1000);
+		VERIFYVULKANRESULT(VulkanRHI::vkDeviceWaitIdle(Device->GetInstanceHandle()));
+		const bool bSuccess = Device->GetFenceManager().WaitForFence(CmdBuffer->Fence, 200 * 1000 * 1000);
 		ensure(bSuccess);
 		ensure(Device->GetFenceManager().IsFenceSignaled(CmdBuffer->Fence));
 		CmdBuffer->GetOwner()->RefreshFenceStatus();
diff --git a/Engine/Source/Runtime/VulkanRHI/Private/VulkanRHI.cpp b/Engine/Source/Runtime/VulkanRHI/Private/VulkanRHI.cpp
--- a/Engine/Source/Runtime/VulkanRHI/Private/VulkanRHI.cpp
+++ b/Engine/Source/Runtime/VulkanRHI/Private/VulkanRHI.cpp
@@ -1483,15 +1483,15 @@ FVulkanRHIAllocationInfo FVulkanDynamicRHI::RHIGetAllocationInfo(FRHIBuffer* InB
 void FVulkanDynamicRHI::RHISetImageLayout(VkImage Image, VkImageLayout OldLayout, VkImageLayout NewLayout, const VkImageSubresourceRange& SubresourceRange)
 {
 	FVulkanCommandListContext& ImmediateContext = GetDevice()->GetImmediateContext();
-	FVulkanCommandBuffer* CmdBuffer = ImmediateContext.GetActiveCmdBuffer();
-	VulkanSetImageLayout(CmdBuffer->GetHandle(), Image, OldLayout, NewLayout, SubresourceRange);
+	FVulkanCommandBuffer& CommandBuffer = ImmediateContext.GetCommandBuffer();
+	VulkanSetImageLayout(CommandBuffer.GetHandle(), Image, OldLayout, NewLayout, SubresourceRange);
 }
 
 void FVulkanDynamicRHI::RHISetUploadImageLayout(VkImage Image, VkImageLayout OldLayout, VkImageLayout NewLayout, const VkImageSubresourceRange& SubresourceRange)
 {
 	FVulkanCommandListContext& ImmediateContext = GetDevice()->GetImmediateContext();
-	FVulkanCommandBuffer* CmdBuffer = ImmediateContext.GetUploadCmdBuffer();
-	VulkanSetImageLayout(CmdBuffer->GetHandle(), Image, OldLayout, NewLayout, SubresourceRange);
+	FVulkanCommandBuffer& CommandBuffer = ImmediateContext.GetCommandBuffer();
+	VulkanSetImageLayout(CommandBuffer.GetHandle(), Image, OldLayout, NewLayout, SubresourceRange);
 }
 
 void FVulkanDynamicRHI::RHIFinishExternalComputeWork(VkCommandBuffer InCommandBuffer)
@@ -1514,8 +1514,7 @@ void FVulkanDynamicRHI::RHIRegisterWork(uint32 NumPrimitives)
 
 void FVulkanDynamicRHI::RHISubmitUploadCommandBuffer()
 {
-	FVulkanCommandListContext& ImmediateContext = GetDevice()->GetImmediateContext();
-	ImmediateContext.SubmitUploadCmdBuffer();
+
 }
 
 void FVulkanDynamicRHI::RHIVerifyResult(VkResult Result, const ANSICHAR* VkFuntion, const ANSICHAR* Filename, uint32 Line)
diff --git a/Engine/Source/Runtime/VulkanRHI/Private/VulkanRenderTarget.cpp b/Engine/Source/Runtime/VulkanRHI/Private/VulkanRenderTarget.cpp
--- a/Engine/Source/Runtime/VulkanRHI/Private/VulkanRenderTarget.cpp
+++ b/Engine/Source/Runtime/VulkanRHI/Private/VulkanRenderTarget.cpp
@@ -163,7 +163,6 @@ void FVulkanDynamicRHI::RHIReadSurfaceData(FRHITexture* TextureRHI, FIntRect Rec
 
 	FVulkanTexture& Surface = *ResourceCast(TextureRHI);
 
-	Device->PrepareForCPURead();
 
 	// Figure out the size of the buffer required to hold the requested pixels
 	const uint32 PixelByteSize = VulkanRHI::GetNumBitsPerPixel(Surface.StorageFormat) / 8;
@@ -176,15 +175,15 @@ void FVulkanDynamicRHI::RHIReadSurfaceData(FRHITexture* TextureRHI, FIntRect Rec
 	const uint32 MipSizeY = FMath::Max(Desc.Extent.Y >> MipLevel, 1);
 	checkf((Rect.Max.X <= (int32)MipSizeX) && (Rect.Max.Y <= (int32)MipSizeY), TEXT("The specified Rect [%dx%d] extends beyond this Mip [%dx%d]."), Rect.Max.X, Rect.Max.Y, MipSizeX, MipSizeY);
 
-	FVulkanCommandListContext& ImmediateContext = Device->GetImmediateContext();
-
+	FRHICommandListImmediate& RHICmdList = FRHICommandListImmediate::Get();
 	VulkanRHI::FStagingBuffer* StagingBuffer = nullptr;
-	FVulkanCommandBuffer* CmdBuffer = nullptr;
 	const bool bCPUReadback = EnumHasAllFlags(Surface.GetDesc().Flags, TexCreate_CPUReadback);
+
 	if (!bCPUReadback) //this function supports reading back arbitrary rendertargets, so if its not a cpu readback surface, we do a copy.
 	{
-		CmdBuffer = ImmediateContext.GetUploadCmdBuffer();
 		StagingBuffer = Device->GetStagingManager().AcquireBuffer(BufferSize, VK_BUFFER_USAGE_TRANSFER_DST_BIT, VK_MEMORY_PROPERTY_HOST_CACHED_BIT);
+		ensure(StagingBuffer->GetSize() >= BufferSize);
+
 		VkBufferImageCopy CopyRegion;
 		FMemory::Memzero(CopyRegion);
 		// Leave bufferRowLength/bufferImageHeight at 0 for tightly packed
@@ -198,29 +197,32 @@ void FVulkanDynamicRHI::RHIReadSurfaceData(FRHITexture* TextureRHI, FIntRect Rec
 		CopyRegion.imageExtent.height = DestHeight;
 		CopyRegion.imageExtent.depth = 1;
 
-		const VkImageLayout PreviousLayout = FVulkanLayoutManager::SetExpectedLayout(CmdBuffer, Surface, ERHIAccess::CopySrc);
+		RHICmdList.EnqueueLambda([CopyRegion, StagingBuffer, &Surface](FRHICommandListBase& ExecutingCmdList)
+			{
+				FVulkanCommandListContext& Context = FVulkanCommandListContext::Get(ExecutingCmdList);
+				FVulkanCommandBuffer& CommandBuffer = Context.GetCommandBuffer();
 
-		VulkanRHI::vkCmdCopyImageToBuffer(CmdBuffer->GetHandle(), Surface.Image, VK_IMAGE_LAYOUT_TRANSFER_SRC_OPTIMAL, StagingBuffer->GetHandle(), 1, &CopyRegion);
+				const VkImageLayout PreviousLayout = FVulkanLayoutManager::SetExpectedLayout(&CommandBuffer, Surface, ERHIAccess::CopySrc);
 
-		FVulkanPipelineBarrier AfterBarrier;
-		if ((PreviousLayout != VK_IMAGE_LAYOUT_UNDEFINED) && (PreviousLayout != VK_IMAGE_LAYOUT_TRANSFER_SRC_OPTIMAL))
-		{
-			AfterBarrier.AddImageLayoutTransition(Surface.Image, VK_IMAGE_LAYOUT_TRANSFER_SRC_OPTIMAL, PreviousLayout, FVulkanPipelineBarrier::MakeSubresourceRange(Surface.GetFullAspectMask()));
-		}
-		else
-		{
-			CmdBuffer->GetLayoutManager().SetFullLayout(Surface, VK_IMAGE_LAYOUT_TRANSFER_SRC_OPTIMAL);
-		}
-		ensure(StagingBuffer->GetSize() >= BufferSize);
+				VulkanRHI::vkCmdCopyImageToBuffer(CommandBuffer.GetHandle(), Surface.Image, VK_IMAGE_LAYOUT_TRANSFER_SRC_OPTIMAL, StagingBuffer->GetHandle(), 1, &CopyRegion);
 
-		AfterBarrier.AddMemoryBarrier(VK_ACCESS_TRANSFER_WRITE_BIT, VK_ACCESS_HOST_READ_BIT, VK_PIPELINE_STAGE_TRANSFER_BIT, VK_PIPELINE_STAGE_HOST_BIT);
-		AfterBarrier.Execute(CmdBuffer);
+				FVulkanPipelineBarrier AfterBarrier;
+				if ((PreviousLayout != VK_IMAGE_LAYOUT_UNDEFINED) && (PreviousLayout != VK_IMAGE_LAYOUT_TRANSFER_SRC_OPTIMAL))
+				{
+					AfterBarrier.AddImageLayoutTransition(Surface.Image, VK_IMAGE_LAYOUT_TRANSFER_SRC_OPTIMAL, PreviousLayout, FVulkanPipelineBarrier::MakeSubresourceRange(Surface.GetFullAspectMask()));
+				}
+				else
+				{
+					CommandBuffer.GetLayoutManager().SetFullLayout(Surface, VK_IMAGE_LAYOUT_TRANSFER_SRC_OPTIMAL);
+				}
 
-		// Force upload
-		ImmediateContext.SubmitUploadCmdBuffer();
+				AfterBarrier.AddMemoryBarrier(VK_ACCESS_TRANSFER_WRITE_BIT, VK_ACCESS_HOST_READ_BIT, VK_PIPELINE_STAGE_TRANSFER_BIT, VK_PIPELINE_STAGE_HOST_BIT);
+				AfterBarrier.Execute(&CommandBuffer);
+			});
 	}
 
-	Device->WaitUntilIdle();
+	// We need to execute the command list so we can read the data from the map below
+	RHICmdList.SubmitAndBlockUntilGPUIdle();
 
 	uint8* In;
 	uint32 SrcPitch;
@@ -243,10 +245,8 @@ void FVulkanDynamicRHI::RHIReadSurfaceData(FRHITexture* TextureRHI, FIntRect Rec
 
 	if (!bCPUReadback)
 	{
-		Device->GetStagingManager().ReleaseBuffer(CmdBuffer, StagingBuffer);
+		Device->GetStagingManager().ReleaseBuffer(nullptr, StagingBuffer);
 	}
-
-	ImmediateContext.PrepareForNewActiveCommandBuffer();
 }
 
 void FVulkanDynamicRHI::RHIReadSurfaceData(FRHITexture* TextureRHI, FIntRect Rect, TArray<FLinearColor>& OutData, FReadSurfaceDataFlags InFlags)
@@ -298,60 +298,58 @@ void FVulkanDynamicRHI::RHIUnmapStagingSurface(FRHITexture* TextureRHI, uint32 G
 
 void FVulkanDynamicRHI::RHIReadSurfaceFloatData(FRHITexture* TextureRHI, FIntRect Rect, TArray<FFloat16Color>& OutData, ECubeFace CubeFace, int32 ArrayIndex, int32 MipIndex)
 {
-	auto DoCopyFloat = [](FVulkanDevice* InDevice, FVulkanCommandBuffer* InCmdBuffer, const FVulkanTexture& Surface, uint32 InMipIndex, uint32 SrcBaseArrayLayer, FIntRect InRect, TArray<FFloat16Color>& OutputData)
+	auto DoCopyFloat = [](FVulkanDevice* InDevice, const FVulkanTexture& VulkanTexture, uint32 InMipIndex, uint32 SrcBaseArrayLayer, FIntRect InRect, TArray<FFloat16Color>& OutputData)
 	{
-		ensure(Surface.StorageFormat == VK_FORMAT_R16G16B16A16_SFLOAT);
+		ensure(VulkanTexture.StorageFormat == VK_FORMAT_R16G16B16A16_SFLOAT);
 
-		const FRHITextureDesc& Desc = Surface.GetDesc();
+		const FRHITextureDesc& Desc = VulkanTexture.GetDesc();
 
 		const uint32 NumPixels = (Desc.Extent.X >> InMipIndex) * (Desc.Extent.Y >> InMipIndex);
 		const uint32 Size = NumPixels * sizeof(FFloat16Color);
 		VulkanRHI::FStagingBuffer* StagingBuffer = InDevice->GetStagingManager().AcquireBuffer(Size, VK_BUFFER_USAGE_TRANSFER_DST_BIT, VK_MEMORY_PROPERTY_HOST_CACHED_BIT);
 
-		FVulkanPipelineBarrier AfterBarrier;
-		if (GIgnoreCPUReads == 0)
-		{
-			VkBufferImageCopy CopyRegion;
-			FMemory::Memzero(CopyRegion);
-			//Region.bufferOffset = 0;
-			CopyRegion.bufferRowLength = FMath::Max(1, Desc.Extent.X >> InMipIndex);
-			CopyRegion.bufferImageHeight = FMath::Max(1, Desc.Extent.Y >> InMipIndex);
-			CopyRegion.imageSubresource.aspectMask = Surface.GetFullAspectMask();
-			CopyRegion.imageSubresource.mipLevel = InMipIndex;
-			CopyRegion.imageSubresource.baseArrayLayer = SrcBaseArrayLayer;
-			CopyRegion.imageSubresource.layerCount = 1;
-			CopyRegion.imageExtent.width = FMath::Max(1, Desc.Extent.X >> InMipIndex);
-			CopyRegion.imageExtent.height = FMath::Max(1, Desc.Extent.Y >> InMipIndex);
-			CopyRegion.imageExtent.depth = 1;
-
-			const VkImageLayout OriginalLayout = FVulkanLayoutManager::SetExpectedLayout(InCmdBuffer, Surface, ERHIAccess::CopySrc);
-
-			VulkanRHI::vkCmdCopyImageToBuffer(InCmdBuffer->GetHandle(), Surface.Image, VK_IMAGE_LAYOUT_TRANSFER_SRC_OPTIMAL, StagingBuffer->GetHandle(), 1, &CopyRegion);
+		// the staging buffer size may be bigger then the size due to alignment, etc. but it must not be smaller!
+		ensure(StagingBuffer->GetSize() >= Size);
 
-			if ((OriginalLayout != VK_IMAGE_LAYOUT_UNDEFINED) && (OriginalLayout != VK_IMAGE_LAYOUT_TRANSFER_SRC_OPTIMAL))
-			{
-				const VkImageSubresourceRange ImageSubresource = FVulkanPipelineBarrier::MakeSubresourceRange(Surface.GetFullAspectMask());
-				AfterBarrier.AddImageLayoutTransition(Surface.Image, VK_IMAGE_LAYOUT_TRANSFER_SRC_OPTIMAL, OriginalLayout, ImageSubresource);
-			}
-			else
+		VkBufferImageCopy CopyRegion;
+		FMemory::Memzero(CopyRegion);
+		//Region.bufferOffset = 0;
+		CopyRegion.bufferRowLength = FMath::Max(1, Desc.Extent.X >> InMipIndex);
+		CopyRegion.bufferImageHeight = FMath::Max(1, Desc.Extent.Y >> InMipIndex);
+		CopyRegion.imageSubresource.aspectMask = VulkanTexture.GetFullAspectMask();
+		CopyRegion.imageSubresource.mipLevel = InMipIndex;
+		CopyRegion.imageSubresource.baseArrayLayer = SrcBaseArrayLayer;
+		CopyRegion.imageSubresource.layerCount = 1;
+		CopyRegion.imageExtent.width = FMath::Max(1, Desc.Extent.X >> InMipIndex);
+		CopyRegion.imageExtent.height = FMath::Max(1, Desc.Extent.Y >> InMipIndex);
+		CopyRegion.imageExtent.depth = 1;
+
+		FRHICommandListImmediate& RHICmdList = FRHICommandListImmediate::Get();
+		RHICmdList.EnqueueLambda([CopyRegion, StagingBuffer, &VulkanTexture](FRHICommandListBase& ExecutingCmdList)
 			{
-				InCmdBuffer->GetLayoutManager().SetFullLayout(Surface, VK_IMAGE_LAYOUT_TRANSFER_SRC_OPTIMAL);
-			}
-		}
-		else
-		{
-			VulkanRHI::vkCmdFillBuffer(InCmdBuffer->GetHandle(), StagingBuffer->GetHandle(), 0, Size, (FFloat16(1.0).Encoded << 16) + FFloat16(1.0).Encoded);
-		}
+				FVulkanCommandListContext& Context = FVulkanCommandListContext::Get(ExecutingCmdList);
+				FVulkanCommandBuffer& CommandBuffer = Context.GetCommandBuffer();
+				const VkImageLayout OriginalLayout = FVulkanLayoutManager::SetExpectedLayout(&CommandBuffer, VulkanTexture, ERHIAccess::CopySrc);
 
-		// the staging buffer size may be bigger then the size due to alignment, etc. but it must not be smaller!
-		ensure(StagingBuffer->GetSize() >= Size);
+				VulkanRHI::vkCmdCopyImageToBuffer(CommandBuffer.GetHandle(), VulkanTexture.Image, VK_IMAGE_LAYOUT_TRANSFER_SRC_OPTIMAL, StagingBuffer->GetHandle(), 1, &CopyRegion);
 
-		AfterBarrier.AddMemoryBarrier(VK_ACCESS_TRANSFER_WRITE_BIT, VK_ACCESS_HOST_READ_BIT, VK_PIPELINE_STAGE_TRANSFER_BIT, VK_PIPELINE_STAGE_HOST_BIT);
-		AfterBarrier.Execute(InCmdBuffer);
+				FVulkanPipelineBarrier AfterBarrier;
+				if ((OriginalLayout != VK_IMAGE_LAYOUT_UNDEFINED) && (OriginalLayout != VK_IMAGE_LAYOUT_TRANSFER_SRC_OPTIMAL))
+				{
+					const VkImageSubresourceRange ImageSubresource = FVulkanPipelineBarrier::MakeSubresourceRange(VulkanTexture.GetFullAspectMask());
+					AfterBarrier.AddImageLayoutTransition(VulkanTexture.Image, VK_IMAGE_LAYOUT_TRANSFER_SRC_OPTIMAL, OriginalLayout, ImageSubresource);
+				}
+				else
+				{
+					CommandBuffer.GetLayoutManager().SetFullLayout(VulkanTexture, VK_IMAGE_LAYOUT_TRANSFER_SRC_OPTIMAL);
+				}
 
-		// Force upload
-		InDevice->GetImmediateContext().SubmitUploadCmdBuffer();
-		InDevice->WaitUntilIdle();
+				AfterBarrier.AddMemoryBarrier(VK_ACCESS_TRANSFER_WRITE_BIT, VK_ACCESS_HOST_READ_BIT, VK_PIPELINE_STAGE_TRANSFER_BIT, VK_PIPELINE_STAGE_HOST_BIT);
+				AfterBarrier.Execute(&CommandBuffer);
+			});
+
+		// We need to execute the command list so we can read the data from the map below
+		RHICmdList.SubmitAndBlockUntilGPUIdle();
 
 		StagingBuffer->InvalidateMappedMemory();
 
@@ -360,21 +358,23 @@ void FVulkanDynamicRHI::RHIReadSurfaceFloatData(FRHITexture* TextureRHI, FIntRec
 		OutputData.SetNumUninitialized(OutWidth * OutHeight);
 		uint32 OutIndex = 0;
 		FFloat16Color* Dest = OutputData.GetData();
+		void* MappedPointer = StagingBuffer->GetMappedPointer();
 		for (int32 Row = InRect.Min.Y; Row < InRect.Max.Y; ++Row)
 		{
-			FFloat16Color* Src = (FFloat16Color*)StagingBuffer->GetMappedPointer() + Row * (Desc.Extent.X >> InMipIndex) + InRect.Min.X;
+			FFloat16Color* Src = (FFloat16Color*)MappedPointer + Row * (Desc.Extent.X >> InMipIndex) + InRect.Min.X;
 			for (int32 Col = InRect.Min.X; Col < InRect.Max.X; ++Col)
 			{
 				OutputData[OutIndex++] = *Src++;
 			}
 		}
-		InDevice->GetStagingManager().ReleaseBuffer(InCmdBuffer, StagingBuffer);
+
+		InDevice->GetStagingManager().ReleaseBuffer(nullptr, StagingBuffer);
 	};
 
 	FVulkanTexture& Surface = *ResourceCast(TextureRHI);
 	const FRHITextureDesc& Desc = Surface.GetDesc();
 
-	if (GIgnoreCPUReads == 1)
+	if (GIgnoreCPUReads)
 	{
 		// Debug: Fill with CPU
 		uint32 NumPixels = 0;
@@ -400,26 +400,22 @@ void FVulkanDynamicRHI::RHIReadSurfaceFloatData(FRHITexture* TextureRHI, FIntRec
 	}
 	else
 	{
-		Device->PrepareForCPURead();
-
-		FVulkanCommandBuffer* CmdBuffer = Device->GetImmediateContext().GetUploadCmdBuffer();
 		switch (TextureRHI->GetDesc().Dimension)
 		{
 			case ETextureDimension::TextureCubeArray:
 			case ETextureDimension::TextureCube:
-				DoCopyFloat(Device, CmdBuffer, Surface, MipIndex, CubeFace + 6 * ArrayIndex, Rect, OutData);
+				DoCopyFloat(Device, Surface, MipIndex, CubeFace + 6 * ArrayIndex, Rect, OutData);
 				break;
 
 			case ETextureDimension::Texture2DArray:
 			case ETextureDimension::Texture2D:
-				DoCopyFloat(Device, CmdBuffer, Surface, MipIndex, ArrayIndex, Rect, OutData);
+				DoCopyFloat(Device, Surface, MipIndex, ArrayIndex, Rect, OutData);
 				break;
 
 			default:
 				checkNoEntry();
 				break;
 		}
-		Device->GetImmediateContext().PrepareForNewActiveCommandBuffer();
 	}
 }
 
@@ -437,67 +433,62 @@ void FVulkanDynamicRHI::RHIRead3DSurfaceFloatData(FRHITexture* TextureRHI, FIntR
 	// Allocate the output buffer.
 	OutData.SetNumUninitialized(Size);
 
-	if (GIgnoreCPUReads == 1)
+	if (GIgnoreCPUReads)
 	{
 		// Debug: Fill with CPU
 		FMemory::Memzero(OutData.GetData(), Size * sizeof(FFloat16Color));
 		return;
 	}
 
-	Device->PrepareForCPURead();
-	FVulkanCommandBuffer* CmdBuffer = Device->GetImmediateContext().GetUploadCmdBuffer();
-
 	ensure(Surface.StorageFormat == VK_FORMAT_R16G16B16A16_SFLOAT);
 
 	VulkanRHI::FStagingBuffer* StagingBuffer = Device->GetStagingManager().AcquireBuffer(Size, VK_BUFFER_USAGE_TRANSFER_DST_BIT, VK_MEMORY_PROPERTY_HOST_CACHED_BIT);
+	// the staging buffer size may be bigger then the size due to alignment, etc. but it must not be smaller!
+	ensure(StagingBuffer->GetSize() >= Size);
 
-	FVulkanPipelineBarrier AfterBarrier;
-	if (GIgnoreCPUReads == 0)
-	{
-		VkBufferImageCopy CopyRegion;
-		FMemory::Memzero(CopyRegion);
-		//Region.bufferOffset = 0;
-		CopyRegion.bufferRowLength = Desc.Extent.X;
-		CopyRegion.bufferImageHeight = Desc.Extent.Y;
-		CopyRegion.imageSubresource.aspectMask = Surface.GetFullAspectMask();
-		//CopyRegion.imageSubresource.mipLevel = 0;
-		//CopyRegion.imageSubresource.baseArrayLayer = 0;
-		CopyRegion.imageSubresource.layerCount = 1;
-		CopyRegion.imageOffset.x = InRect.Min.X;
-		CopyRegion.imageOffset.y = InRect.Min.Y;
-		CopyRegion.imageOffset.z = ZMinMax.X;
-		CopyRegion.imageExtent.width = SizeX;
-		CopyRegion.imageExtent.height = SizeY;
-		CopyRegion.imageExtent.depth = SizeZ;
+	VkBufferImageCopy CopyRegion;
+	FMemory::Memzero(CopyRegion);
+	//Region.bufferOffset = 0;
+	CopyRegion.bufferRowLength = Desc.Extent.X;
+	CopyRegion.bufferImageHeight = Desc.Extent.Y;
+	CopyRegion.imageSubresource.aspectMask = Surface.GetFullAspectMask();
+	//CopyRegion.imageSubresource.mipLevel = 0;
+	//CopyRegion.imageSubresource.baseArrayLayer = 0;
+	CopyRegion.imageSubresource.layerCount = 1;
+	CopyRegion.imageOffset.x = InRect.Min.X;
+	CopyRegion.imageOffset.y = InRect.Min.Y;
+	CopyRegion.imageOffset.z = ZMinMax.X;
+	CopyRegion.imageExtent.width = SizeX;
+	CopyRegion.imageExtent.height = SizeY;
+	CopyRegion.imageExtent.depth = SizeZ;
+
+	FRHICommandListImmediate& RHICmdList = FRHICommandListImmediate::Get();
+	RHICmdList.EnqueueLambda([CopyRegion, StagingBuffer, &Surface](FRHICommandListBase& ExecutingCmdList)
+		{
+			FVulkanCommandListContext& Context = FVulkanCommandListContext::Get(ExecutingCmdList);
+			FVulkanCommandBuffer& CommandBuffer = Context.GetCommandBuffer();
 
-		//#todo-rco: Multithreaded!
-		const VkImageLayout OriginalLayout = FVulkanLayoutManager::SetExpectedLayout(CmdBuffer, Surface, ERHIAccess::CopySrc);
+			const VkImageLayout OriginalLayout = FVulkanLayoutManager::SetExpectedLayout(&CommandBuffer, Surface, ERHIAccess::CopySrc);
 
-		VulkanRHI::vkCmdCopyImageToBuffer(CmdBuffer->GetHandle(), Surface.Image, VK_IMAGE_LAYOUT_TRANSFER_SRC_OPTIMAL, StagingBuffer->GetHandle(), 1, &CopyRegion);
+			VulkanRHI::vkCmdCopyImageToBuffer(CommandBuffer.GetHandle(), Surface.Image, VK_IMAGE_LAYOUT_TRANSFER_SRC_OPTIMAL, StagingBuffer->GetHandle(), 1, &CopyRegion);
 
-		if ((OriginalLayout != VK_IMAGE_LAYOUT_UNDEFINED) && (OriginalLayout != VK_IMAGE_LAYOUT_TRANSFER_SRC_OPTIMAL))
-		{
-			AfterBarrier.AddImageLayoutTransition(Surface.Image, VK_IMAGE_LAYOUT_TRANSFER_SRC_OPTIMAL, OriginalLayout, FVulkanPipelineBarrier::MakeSubresourceRange(Surface.GetFullAspectMask()));
-		}
-		else
-		{
-			CmdBuffer->GetLayoutManager().SetFullLayout(Surface, VK_IMAGE_LAYOUT_TRANSFER_SRC_OPTIMAL);
-		}
-	}
-	else
-	{
-		VulkanRHI::vkCmdFillBuffer(CmdBuffer->GetHandle(), StagingBuffer->GetHandle(), 0, Size, (FFloat16(1.0).Encoded << 16) + FFloat16(1.0).Encoded);
-	}
+			FVulkanPipelineBarrier AfterBarrier;
+			if ((OriginalLayout != VK_IMAGE_LAYOUT_UNDEFINED) && (OriginalLayout != VK_IMAGE_LAYOUT_TRANSFER_SRC_OPTIMAL))
+			{
+				AfterBarrier.AddImageLayoutTransition(Surface.Image, VK_IMAGE_LAYOUT_TRANSFER_SRC_OPTIMAL, OriginalLayout, FVulkanPipelineBarrier::MakeSubresourceRange(Surface.GetFullAspectMask()));
+			}
+			else
+			{
+				CommandBuffer.GetLayoutManager().SetFullLayout(Surface, VK_IMAGE_LAYOUT_TRANSFER_SRC_OPTIMAL);
+			}
 
-	// the staging buffer size may be bigger then the size due to alignment, etc. but it must not be smaller!
-	ensure(StagingBuffer->GetSize() >= Size);
 
-	AfterBarrier.AddMemoryBarrier(VK_ACCESS_TRANSFER_WRITE_BIT, VK_ACCESS_HOST_READ_BIT, VK_PIPELINE_STAGE_TRANSFER_BIT, VK_PIPELINE_STAGE_HOST_BIT);
-	AfterBarrier.Execute(CmdBuffer);
+			AfterBarrier.AddMemoryBarrier(VK_ACCESS_TRANSFER_WRITE_BIT, VK_ACCESS_HOST_READ_BIT, VK_PIPELINE_STAGE_TRANSFER_BIT, VK_PIPELINE_STAGE_HOST_BIT);
+			AfterBarrier.Execute(&CommandBuffer);
+		});
 
-	// Force upload
-	Device->GetImmediateContext().SubmitUploadCmdBuffer();
-	Device->WaitUntilIdle();
+	// We need to execute the command list so we can read the data from the map below
+	RHICmdList.SubmitAndBlockUntilGPUIdle();
 
 	StagingBuffer->InvalidateMappedMemory();
 
@@ -516,8 +507,7 @@ void FVulkanDynamicRHI::RHIRead3DSurfaceFloatData(FRHITexture* TextureRHI, FIntR
 	FFloat16Color* End = OutData.GetData() + OutData.Num();
 	checkf(Dest <= End, TEXT("Memory overwrite! Calculated total size %d: SizeX %d SizeY %d SizeZ %d; InRect(%d, %d, %d, %d) InZ(%d, %d)"),
 		Size, SizeX, SizeY, SizeZ, InRect.Min.X, InRect.Min.Y, InRect.Max.X, InRect.Max.Y, ZMinMax.X, ZMinMax.Y);
-	Device->GetStagingManager().ReleaseBuffer(CmdBuffer, StagingBuffer);
-	Device->GetImmediateContext().PrepareForNewActiveCommandBuffer();
+	Device->GetStagingManager().ReleaseBuffer(nullptr, StagingBuffer);
 }
 
 VkFormat FVulkanCommandListContext::GetSwapchainImageFormat() const
diff --git a/Engine/Source/Runtime/VulkanRHI/Private/VulkanResources.h b/Engine/Source/Runtime/VulkanRHI/Private/VulkanResources.h
--- a/Engine/Source/Runtime/VulkanRHI/Private/VulkanResources.h
+++ b/Engine/Source/Runtime/VulkanRHI/Private/VulkanResources.h
@@ -738,7 +738,7 @@ public:
 	VULKANRHI_API VkDeviceMemory GetAllocationHandle() const;
 	VULKANRHI_API uint64 GetAllocationOffset() const;
 
-	static void InternalLockWrite(FVulkanCommandListContext& Context, FVulkanTexture* Surface, const VkBufferImageCopy& Region, VulkanRHI::FStagingBuffer* StagingBuffer);
+	static void InternalLockWrite(FVulkanContextCommon& Context, FVulkanTexture* Surface, const VkBufferImageCopy& Region, VulkanRHI::FStagingBuffer* StagingBuffer);
 
 	const FVulkanCpuReadbackBuffer* GetCpuReadbackBuffer() const { return CpuReadbackBuffer; }
 
@@ -755,7 +755,7 @@ public:
 	FVulkanRHIExternalImageDeleteCallbackInfo ExternalImageDeleteCallbackInfo;
 
 private:
-	void SetInitialImageState(FVulkanCommandListContext& Context, VkImageLayout InitialLayout, bool bClear, const FClearValueBinding& ClearValueBinding, bool bIsTransientResource);
+	void SetInitialImageState(FRHICommandListBase& RHICmdList, VkImageLayout InitialLayout, bool bClear, const FClearValueBinding& ClearValueBinding, bool bIsTransientResource);
 	void InternalMoveSurface(FVulkanDevice& InDevice, FVulkanCommandListContext& Context, VulkanRHI::FVulkanAllocation& DestAllocation, VkImageLayout OriginalLayout);
 
 	VkImageTiling Tiling;
diff --git a/Engine/Source/Runtime/VulkanRHI/Private/VulkanTexture.cpp b/Engine/Source/Runtime/VulkanRHI/Private/VulkanTexture.cpp
--- a/Engine/Source/Runtime/VulkanRHI/Private/VulkanTexture.cpp
+++ b/Engine/Source/Runtime/VulkanRHI/Private/VulkanTexture.cpp
@@ -15,15 +15,6 @@
 #include "RHICoreStats.h"
 #include "RHILockTracker.h"
 
-int32 GVulkanSubmitOnTextureUnlock = 1;
-static FAutoConsoleVariableRef CVarVulkanSubmitOnTextureUnlock(
-	TEXT("r.Vulkan.SubmitOnTextureUnlock"),
-	GVulkanSubmitOnTextureUnlock,
-	TEXT("Whether to submit upload cmd buffer on each texture unlock.\n")
-	TEXT(" 0: Do not submit\n")
-	TEXT(" 1: Submit (default)"),
-	ECVF_Default
-);
 
 // This is a workaround for issues with AFBC on Mali GPUs before the G710
 int32 GVulkanDepthStencilForceStorageBit = 0;
@@ -112,9 +103,9 @@ static void VulkanTextureDestroyed(const FRHITextureDesc& TextureDesc, uint64 Si
 	UpdateVulkanTextureStats(TextureDesc, Size, false);
 }
 
-inline void FVulkanTexture::InternalLockWrite(FVulkanCommandListContext& Context, FVulkanTexture* Surface, const VkBufferImageCopy& Region, VulkanRHI::FStagingBuffer* StagingBuffer)
+void FVulkanTexture::InternalLockWrite(FVulkanContextCommon& Context, FVulkanTexture* Surface, const VkBufferImageCopy& Region, VulkanRHI::FStagingBuffer* StagingBuffer)
 {
-	FVulkanCommandBuffer* CmdBuffer = Context.GetUploadCmdBuffer();
+	FVulkanCommandBuffer* CmdBuffer = Context.GetActiveCmdBuffer();
 	ensure(CmdBuffer->IsOutsideRenderPass());
 	VkCommandBuffer StagingCommandBuffer = CmdBuffer->GetHandle();
 
@@ -136,11 +127,6 @@ inline void FVulkanTexture::InternalLockWrite(FVulkanCommandListContext& Context
 	}
 
 	Surface->Device->GetStagingManager().ReleaseBuffer(CmdBuffer, StagingBuffer);
-
-	if (GVulkanSubmitOnTextureUnlock != 0)
-	{
-		Context.SubmitUploadCmdBuffer();
-	}
 }
 
 void FVulkanTexture::ErrorInvalidViewType() const
@@ -149,25 +135,6 @@ void FVulkanTexture::ErrorInvalidViewType() const
 }
 
 
-struct FRHICommandLockWriteTexture final : public FRHICommand<FRHICommandLockWriteTexture>
-{
-	FVulkanTexture* Surface;
-	VkBufferImageCopy Region;
-	VulkanRHI::FStagingBuffer* StagingBuffer;
-
-	FRHICommandLockWriteTexture(FVulkanTexture* InSurface, const VkBufferImageCopy& InRegion, VulkanRHI::FStagingBuffer* InStagingBuffer)
-		: Surface(InSurface)
-		, Region(InRegion)
-		, StagingBuffer(InStagingBuffer)
-	{
-	}
-
-	void Execute(FRHICommandListBase& RHICmdList)
-	{
-		FVulkanTexture::InternalLockWrite(FVulkanCommandListContext::GetVulkanContext(RHICmdList.GetContext()), Surface, Region, StagingBuffer);
-	}
-};
-
 static VkImageUsageFlags GetUsageFlagsFromCreateFlags(FVulkanDevice& InDevice, const ETextureCreateFlags& UEFlags)
 {
 	VkImageUsageFlags UsageFlags = VK_IMAGE_USAGE_TRANSFER_SRC_BIT | VK_IMAGE_USAGE_TRANSFER_DST_BIT | VK_IMAGE_USAGE_SAMPLED_BIT;
@@ -521,60 +488,6 @@ void FVulkanTexture::GenerateImageCreateInfo(
 	}
 }
 
-struct FRHICommandSetInitialImageState final : public FRHICommand<FRHICommandSetInitialImageState>
-{
-	FVulkanTexture* VulkanTexture;
-	VkImageLayout InitialLayout;
-	bool bOnlyAddToLayoutManager;
-	bool bClear;
-	bool bIsTransientResource;
-	FClearValueBinding ClearValueBinding;
-
-	FRHICommandSetInitialImageState(FVulkanTexture* InVulkanTexture, VkImageLayout InInitialLayout, bool bInOnlyAddToLayoutManager, bool bInClear, const FClearValueBinding& InClearValueBinding, bool bInIsTransientResource)
-		: VulkanTexture(InVulkanTexture)
-		, InitialLayout(InInitialLayout)
-		, bOnlyAddToLayoutManager(bInOnlyAddToLayoutManager)
-		, bClear(bInClear)
-		, bIsTransientResource(bInIsTransientResource)
-		, ClearValueBinding(InClearValueBinding)
-	{
-	}
-
-	void Execute(FRHICommandListBase& CmdList)
-	{
-		FVulkanCommandListContext& Context = FVulkanCommandListContext::GetVulkanContext(CmdList.GetContext());
-		
-		if (bOnlyAddToLayoutManager)
-		{
-			FVulkanCommandBuffer& CommandBuffer = Context.GetCommandBuffer();
-			CommandBuffer.GetLayoutManager().SetFullLayout(*VulkanTexture, InitialLayout, true);
-		}
-		else
-		{
-			VulkanTexture->SetInitialImageState(Context, InitialLayout, bClear, ClearValueBinding, bIsTransientResource);
-		}
-	}
-};
-
-struct FRHICommandOnDestroyImage final : public FRHICommand<FRHICommandOnDestroyImage>
-{
-	VkImage Image;
-	FVulkanDevice* Device;
-	bool bRenderTarget;
-
-	FRHICommandOnDestroyImage(VkImage InImage, FVulkanDevice* InDevice, bool bInRenderTarget)
-		: Image(InImage)
-		, Device(InDevice)
-		, bRenderTarget(bInRenderTarget)
-	{
-	}
-
-	void Execute(FRHICommandListBase& RHICmdList)
-	{
-		Device->NotifyDeletedImage(Image, bRenderTarget);
-	}
-};
-
 static VkImageLayout ChooseVRSLayout()
 {
 	if(GRHIVariableRateShadingImageDataType == VRSImage_Palette)
@@ -747,16 +660,7 @@ void FVulkanTexture::DestroySurface()
 	else if (bIsLocalOwner || bHasExternalOwner)
 	{
 		const bool bRenderTarget = EnumHasAnyFlags(GetDesc().Flags, TexCreate_RenderTargetable | TexCreate_DepthStencilTargetable | TexCreate_ResolveTargetable);
-		FRHICommandList& RHICmdList = FRHICommandListExecutor::GetImmediateCommandList();
-		if (!IsInRenderingThread() || (RHICmdList.Bypass() || !IsRunningRHIInSeparateThread()))
-		{
-			Device->NotifyDeletedImage(Image, bRenderTarget);
-		}
-		else
-		{
-			check(IsInRenderingThread());
-			new (RHICmdList.AllocCommand<FRHICommandOnDestroyImage>()) FRHICommandOnDestroyImage(Image, Device, bRenderTarget);
-		}
+		Device->NotifyDeletedImage(Image, bRenderTarget);
 
 		if (bIsLocalOwner)
 		{
@@ -880,58 +784,67 @@ void FVulkanTexture::GetMipSize(uint32 MipIndex, uint64& MipBytes)
 */
 }
 
-void FVulkanTexture::SetInitialImageState(FVulkanCommandListContext& Context, VkImageLayout InitialLayout, bool bClear, const FClearValueBinding& ClearValueBinding, bool bIsTransientResource)
+void FVulkanTexture::SetInitialImageState(FRHICommandListBase& RHICmdList, VkImageLayout InitialLayout, bool bClear, const FClearValueBinding& ClearValueBinding, bool bIsTransientResource)
 {
-	// Can't use TransferQueue as Vulkan requires that queue to also have Gfx or Compute capabilities...
-	//#todo-rco: This function is only used during loading currently, if used for regular RHIClear then use the ActiveCmdBuffer
-	// NOTE: Transient resources' memory might have belonged to another resource earlier in the ActiveCmdBuffer, so we can't use UploadCmdBuffer
-	FVulkanCommandBuffer* CmdBuffer = bIsTransientResource ? Context.GetActiveCmdBuffer() : Context.GetUploadCmdBuffer();
-	ensure(CmdBuffer->IsOutsideRenderPass());
+	RHICmdList.EnqueueLambda(TEXT("FVulkanTexture::SetInitialImageState"),
+		[VulkanTexture = this, InitialLayout, bClear, bIsTransientResource, ClearValueBinding](FRHICommandListBase& ExecutingCmdList)
+		{
+			// Can't use TransferQueue as Vulkan requires that queue to also have Gfx or Compute capabilities...
+			// NOTE: Transient resources' memory might have belonged to another resource earlier in the ActiveCmdBuffer, so we can't use UploadCmdBuffer
+			FVulkanCommandBuffer& CommandBuffer = bIsTransientResource ?
+				FVulkanCommandListContext::Get(ExecutingCmdList).GetCommandBuffer() :
+				FVulkanUploadContext::Get(ExecutingCmdList).GetCommandBuffer();
+			ensure(CommandBuffer.IsOutsideRenderPass());
 
-	VkImageSubresourceRange SubresourceRange = FVulkanPipelineBarrier::MakeSubresourceRange(FullAspectMask);
+			VkImageSubresourceRange SubresourceRange = FVulkanPipelineBarrier::MakeSubresourceRange(VulkanTexture->GetFullAspectMask());
 
-	VkImageLayout CurrentLayout = VK_IMAGE_LAYOUT_UNDEFINED;
-	if (bClear && !bIsTransientResource)
-	{
-		{
-			FVulkanPipelineBarrier Barrier;
-			Barrier.AddImageLayoutTransition(Image, VK_IMAGE_LAYOUT_UNDEFINED, VK_IMAGE_LAYOUT_TRANSFER_DST_OPTIMAL, SubresourceRange);
-			Barrier.Execute(CmdBuffer);
-		}
+			VkImageLayout CurrentLayout = VK_IMAGE_LAYOUT_UNDEFINED;
+			if (bClear && !bIsTransientResource)
+			{
+				{
+					FVulkanPipelineBarrier Barrier;
+					Barrier.AddImageLayoutTransition(VulkanTexture->Image, VK_IMAGE_LAYOUT_UNDEFINED, VK_IMAGE_LAYOUT_TRANSFER_DST_OPTIMAL, SubresourceRange);
+					Barrier.Execute(&CommandBuffer);
+				}
 
-		if (FullAspectMask == VK_IMAGE_ASPECT_COLOR_BIT)
-		{
-			VkClearColorValue Color;
-			FMemory::Memzero(Color);
-			Color.float32[0] = ClearValueBinding.Value.Color[0];
-			Color.float32[1] = ClearValueBinding.Value.Color[1];
-			Color.float32[2] = ClearValueBinding.Value.Color[2];
-			Color.float32[3] = ClearValueBinding.Value.Color[3];
-
-			VulkanRHI::vkCmdClearColorImage(CmdBuffer->GetHandle(), Image, VK_IMAGE_LAYOUT_TRANSFER_DST_OPTIMAL, &Color, 1, &SubresourceRange);
-		}
-		else
-		{
-			check(IsDepthOrStencilAspect());
-			VkClearDepthStencilValue Value;
-			FMemory::Memzero(Value);
-			Value.depth = ClearValueBinding.Value.DSValue.Depth;
-			Value.stencil = ClearValueBinding.Value.DSValue.Stencil;
+				if (VulkanTexture->GetFullAspectMask() == VK_IMAGE_ASPECT_COLOR_BIT)
+				{
+					VkClearColorValue Color;
+					FMemory::Memzero(Color);
+					Color.float32[0] = ClearValueBinding.Value.Color[0];
+					Color.float32[1] = ClearValueBinding.Value.Color[1];
+					Color.float32[2] = ClearValueBinding.Value.Color[2];
+					Color.float32[3] = ClearValueBinding.Value.Color[3];
+
+					VulkanRHI::vkCmdClearColorImage(CommandBuffer.GetHandle(), VulkanTexture->Image, VK_IMAGE_LAYOUT_TRANSFER_DST_OPTIMAL, &Color, 1, &SubresourceRange);
+				}
+				else
+				{
+					check(VulkanTexture->IsDepthOrStencilAspect());
+					VkClearDepthStencilValue Value;
+					FMemory::Memzero(Value);
+					Value.depth = ClearValueBinding.Value.DSValue.Depth;
+					Value.stencil = ClearValueBinding.Value.DSValue.Stencil;
 
-			VulkanRHI::vkCmdClearDepthStencilImage(CmdBuffer->GetHandle(), Image, VK_IMAGE_LAYOUT_TRANSFER_DST_OPTIMAL, &Value, 1, &SubresourceRange);
-		}
+					VulkanRHI::vkCmdClearDepthStencilImage(CommandBuffer.GetHandle(), VulkanTexture->Image, VK_IMAGE_LAYOUT_TRANSFER_DST_OPTIMAL, &Value, 1, &SubresourceRange);
+				}
 
-		CurrentLayout = VK_IMAGE_LAYOUT_TRANSFER_DST_OPTIMAL;
-	}
+				CurrentLayout = VK_IMAGE_LAYOUT_TRANSFER_DST_OPTIMAL;
+			}
 
-	if ((InitialLayout != CurrentLayout) && (InitialLayout != VK_IMAGE_LAYOUT_UNDEFINED))
-	{
-		FVulkanPipelineBarrier Barrier;
-		Barrier.AddFullImageLayoutTransition(*this, CurrentLayout, InitialLayout);
-		Barrier.Execute(CmdBuffer);
-	}
+			if ((InitialLayout != CurrentLayout) && (InitialLayout != VK_IMAGE_LAYOUT_UNDEFINED))
+			{
+				FVulkanPipelineBarrier Barrier;
+				Barrier.AddFullImageLayoutTransition(*VulkanTexture, CurrentLayout, InitialLayout);
+				Barrier.Execute(&CommandBuffer);
+			}
 
-	CmdBuffer->GetLayoutManager().SetFullLayout(*this, InitialLayout);
+			if (ExecutingCmdList.GetPipeline() != ERHIPipeline::None)
+			{
+				FVulkanCommandBuffer& TrackingCommandBuffer = FVulkanCommandListContext::Get(ExecutingCmdList).GetCommandBuffer();
+				TrackingCommandBuffer.GetLayoutManager().SetFullLayout(*VulkanTexture, InitialLayout);
+			}
+		});
 }
 
 
@@ -1070,21 +983,20 @@ FTextureRHIRef FVulkanDynamicRHI::RHIAsyncCreateTexture2D(uint32 SizeX,uint32 Si
 	return FTextureRHIRef();
 }
 
-static void DoAsyncReallocateTexture2D(FVulkanCommandListContext& Context, FVulkanTexture* OldTexture, FVulkanTexture* NewTexture, int32 NewMipCount, int32 NewSizeX, int32 NewSizeY, FThreadSafeCounter* RequestStatus)
+static void DoAsyncReallocateTexture2D(FVulkanContextCommon& Context, FVulkanTexture* OldTexture, FVulkanTexture* NewTexture, int32 NewMipCount, int32 NewSizeX, int32 NewSizeY, FThreadSafeCounter* RequestStatus)
 {
 	LLM_SCOPE_VULKAN(ELLMTagVulkan::VulkanTextures);
 	//QUICK_SCOPE_CYCLE_COUNTER(STAT_FRHICommandGnmAsyncReallocateTexture2D_Execute);
-	check(Context.IsImmediate());
 
 	// figure out what mips to copy from/to
 	const uint32 NumSharedMips = FMath::Min(OldTexture->GetNumMips(), NewTexture->GetNumMips());
 	const uint32 SourceFirstMip = OldTexture->GetNumMips() - NumSharedMips;
 	const uint32 DestFirstMip = NewTexture->GetNumMips() - NumSharedMips;
 
-	FVulkanCommandBuffer* CmdBuffer = Context.GetUploadCmdBuffer();
-	ensure(CmdBuffer->IsOutsideRenderPass());
+	FVulkanCommandBuffer& CommandBuffer = Context.GetCommandBuffer();
+	ensure(CommandBuffer.IsOutsideRenderPass());
 
-	VkCommandBuffer StagingCommandBuffer = CmdBuffer->GetHandle();
+	VkCommandBuffer StagingCommandBuffer = CommandBuffer.GetHandle();
 
 	VkImageCopy Regions[MAX_TEXTURE_MIP_COUNT];
 	check(NumSharedMips <= MAX_TEXTURE_MIP_COUNT);
@@ -1118,7 +1030,7 @@ static void DoAsyncReallocateTexture2D(FVulkanCommandListContext& Context, FVulk
 		FVulkanPipelineBarrier Barrier;
 		Barrier.AddImageLayoutTransition(OldTexture->Image, OldTexture->GetDefaultLayout(), VK_IMAGE_LAYOUT_TRANSFER_SRC_OPTIMAL, SourceSubResourceRange);
 		Barrier.AddImageLayoutTransition(NewTexture->Image, VK_IMAGE_LAYOUT_UNDEFINED, VK_IMAGE_LAYOUT_TRANSFER_DST_OPTIMAL, DestSubResourceRange);
-		Barrier.Execute(CmdBuffer);
+		Barrier.Execute(&CommandBuffer);
 	}
 
 	VulkanRHI::vkCmdCopyImage(StagingCommandBuffer, OldTexture->Image, VK_IMAGE_LAYOUT_TRANSFER_SRC_OPTIMAL, NewTexture->Image, VK_IMAGE_LAYOUT_TRANSFER_DST_OPTIMAL, NumSharedMips, Regions);
@@ -1128,11 +1040,11 @@ static void DoAsyncReallocateTexture2D(FVulkanCommandListContext& Context, FVulk
 		FVulkanPipelineBarrier Barrier;
 		Barrier.AddImageLayoutTransition(OldTexture->Image, VK_IMAGE_LAYOUT_TRANSFER_SRC_OPTIMAL, OldTexture->GetDefaultLayout(), SourceSubResourceRange);
 		Barrier.AddImageLayoutTransition(NewTexture->Image, VK_IMAGE_LAYOUT_TRANSFER_DST_OPTIMAL, NewTexture->GetDefaultLayout(), DestSubResourceRange);
-		Barrier.Execute(CmdBuffer);
+		Barrier.Execute(&CommandBuffer);
 
 		// Add tracking for the appropriate subresources (intentionally leave added mips in VK_IMAGE_LAYOUT_UNDEFINED)
 		// NOTE: Overwriting whatever is contained in the layout manager for the new texture (circumvents issue with stale layouts and pointer reuse)
-		CmdBuffer->GetLayoutManager().SetLayout(*NewTexture, DestSubResourceRange, NewTexture->GetDefaultLayout());
+		CommandBuffer.GetLayoutManager().SetLayout(*NewTexture, DestSubResourceRange, NewTexture->GetDefaultLayout());
 	}
 
 	// request is now complete
@@ -1142,40 +1054,9 @@ static void DoAsyncReallocateTexture2D(FVulkanCommandListContext& Context, FVulk
 	//NewTexture->bSkipBlockOnUnlock = true;
 }
 
-struct FRHICommandVulkanAsyncReallocateTexture2D final : public FRHICommand<FRHICommandVulkanAsyncReallocateTexture2D>
-{
-	FVulkanTexture* OldTexture;
-	FVulkanTexture* NewTexture;
-	int32 NewMipCount;
-	int32 NewSizeX;
-	int32 NewSizeY;
-	FThreadSafeCounter* RequestStatus;
-
-	FORCEINLINE_DEBUGGABLE FRHICommandVulkanAsyncReallocateTexture2D(FVulkanTexture* InOldTexture, FVulkanTexture* InNewTexture, int32 InNewMipCount, int32 InNewSizeX, int32 InNewSizeY, FThreadSafeCounter* InRequestStatus)
-		: OldTexture(InOldTexture)
-		, NewTexture(InNewTexture)
-		, NewMipCount(InNewMipCount)
-		, NewSizeX(InNewSizeX)
-		, NewSizeY(InNewSizeY)
-		, RequestStatus(InRequestStatus)
-	{
-	}
-
-	void Execute(FRHICommandListBase& RHICmdList)
-	{
-		FVulkanCommandListContext& VulkanContext = (FVulkanCommandListContext&)RHICmdList.GetContext().GetLowestLevelContext();
-		check(VulkanContext.IsImmediate());
-		DoAsyncReallocateTexture2D(VulkanContext, OldTexture, NewTexture, NewMipCount, NewSizeX, NewSizeY, RequestStatus);
-	}
-};
-
 FTextureRHIRef FVulkanDynamicRHI::AsyncReallocateTexture2D_RenderThread(FRHICommandListImmediate& RHICmdList, FRHITexture* OldTextureRHI, int32 NewMipCount, int32 NewSizeX, int32 NewSizeY, FThreadSafeCounter* RequestStatus)
 {
 	LLM_SCOPE_VULKAN(ELLMTagVulkan::VulkanTextures);
-	if (RHICmdList.Bypass())
-	{
-		return FDynamicRHI::AsyncReallocateTexture2D_RenderThread(RHICmdList, OldTextureRHI, NewMipCount, NewSizeX, NewSizeY, RequestStatus);
-	}
 
 	FVulkanTexture* OldTexture = ResourceCast(OldTextureRHI);
 	const FRHITextureDesc& OldDesc = OldTexture->GetDesc();
@@ -1189,7 +1070,12 @@ FTextureRHIRef FVulkanDynamicRHI::AsyncReallocateTexture2D_RenderThread(FRHIComm
 		.DetermineInititialState();
 
 	FVulkanTexture* NewTexture = new FVulkanTexture(&RHICmdList, *Device, Desc, nullptr);
-	ALLOC_COMMAND_CL(RHICmdList, FRHICommandVulkanAsyncReallocateTexture2D)(OldTexture, NewTexture, NewMipCount, NewSizeX, NewSizeY, RequestStatus);
+
+	RHICmdList.EnqueueLambda(TEXT("AsyncReallocateTexture2D"), [OldTexture, NewTexture, NewMipCount, NewSizeX, NewSizeY, RequestStatus](FRHICommandListImmediate& ImmCmdList)
+	{
+		FVulkanUploadContext& UploadContext = FVulkanUploadContext::Get(ImmCmdList);
+		DoAsyncReallocateTexture2D(UploadContext, OldTexture, NewTexture, NewMipCount, NewSizeX, NewSizeY, RequestStatus);
+	});
 
 	return NewTexture;
 }
@@ -1279,7 +1165,6 @@ void FVulkanDynamicRHI::RHIUnlockTexture(FRHICommandListImmediate& RHICmdList, c
 	MipHeight = FMath::Max<uint32>(MipHeight, 1);
 
 	VkBufferImageCopy Region{};
-
 	Region.imageSubresource.aspectMask = Texture->GetPartialAspectMask();
 	Region.imageSubresource.mipLevel = Arguments.MipIndex;
 	Region.imageSubresource.baseArrayLayer = ArrayIndex;
@@ -1288,15 +1173,11 @@ void FVulkanDynamicRHI::RHIUnlockTexture(FRHICommandListImmediate& RHICmdList, c
 	Region.imageExtent.height = MipHeight;
 	Region.imageExtent.depth = 1;
 
-	if (RHICmdList.Bypass() || !IsRunningRHIInSeparateThread())
-	{
-		FVulkanTexture::InternalLockWrite(Device->GetImmediateContext(), Texture, Region, StagingBuffer);
-	}
-	else
-	{
-		check(IsInRenderingThread());
-		ALLOC_COMMAND_CL(RHICmdList, FRHICommandLockWriteTexture)(Texture, Region, StagingBuffer);
-	}
+	RHICmdList.EnqueueLambda(TEXT("FVulkanTexture::InternalLockWrite"),
+		[Texture, Region, StagingBuffer](FRHICommandListBase& ExecutingCmdList)
+		{
+			FVulkanTexture::InternalLockWrite(FVulkanUploadContext::Get(ExecutingCmdList), Texture, Region, StagingBuffer);
+		});
 }
 
 
@@ -1347,15 +1228,11 @@ void FVulkanDynamicRHI::InternalUpdateTexture2D(FRHICommandListBase& RHICmdList,
 	Region.imageExtent.depth = 1;
 
 	FVulkanTexture* Texture = ResourceCast(TextureRHI);
-
-	if (RHICmdList.IsBottomOfPipe())
-	{
-		FVulkanTexture::InternalLockWrite(FVulkanCommandListContext::GetVulkanContext(RHICmdList.GetContext()), Texture, Region, StagingBuffer);
-	}
-	else
-	{
-		ALLOC_COMMAND_CL(RHICmdList, FRHICommandLockWriteTexture)(Texture, Region, StagingBuffer);
-	}
+	RHICmdList.EnqueueLambda(TEXT("FVulkanTexture::InternalLockWrite"),
+		[Texture, Region, StagingBuffer](FRHICommandListBase& ExecutingCmdList)
+		{
+			FVulkanTexture::InternalLockWrite(FVulkanUploadContext::Get(ExecutingCmdList), Texture, Region, StagingBuffer);
+		});
 }
 
 FUpdateTexture3DData FVulkanDynamicRHI::RHIBeginUpdateTexture3D(FRHICommandListBase& RHICmdList, FRHITexture* Texture, uint32 MipIndex, const struct FUpdateTextureRegion3D& UpdateRegion)
@@ -1445,14 +1322,11 @@ void FVulkanDynamicRHI::InternalUpdateTexture3D(FRHICommandListBase& RHICmdList,
 	Region.imageExtent.height = (uint32)FMath::Min((int32)(TextureSizeY-UpdateRegion.DestY), (int32)UpdateRegion.Height);
 	Region.imageExtent.depth = (uint32)FMath::Min((int32)(TextureSizeZ-UpdateRegion.DestZ), (int32)UpdateRegion.Depth);
 
-	if (RHICmdList.IsBottomOfPipe())
-	{
-		FVulkanTexture::InternalLockWrite(FVulkanCommandListContext::GetVulkanContext(RHICmdList.GetContext()), Texture, Region, StagingBuffer);
-	}
-	else
-	{
-		ALLOC_COMMAND_CL(RHICmdList, FRHICommandLockWriteTexture)(Texture, Region, StagingBuffer);
-	}
+	RHICmdList.EnqueueLambda(TEXT("FVulkanTexture::InternalLockWrite"),
+		[Texture, Region, StagingBuffer](FRHICommandListBase& ExecutingCmdList)
+		{
+			FVulkanTexture::InternalLockWrite(FVulkanUploadContext::Get(ExecutingCmdList), Texture, Region, StagingBuffer);
+		});
 }
 
 FVulkanTexture::FVulkanTexture(FRHICommandListBase* RHICmdList, FVulkanDevice& InDevice, const FRHITextureCreateDesc& InCreateDesc, const FRHITransientHeapAllocation* InTransientHeapAllocation)
@@ -1598,23 +1472,11 @@ FVulkanTexture::FVulkanTexture(FRHICommandListBase* RHICmdList, FVulkanDevice& I
 
 	if (InitialLayout != VK_IMAGE_LAYOUT_UNDEFINED || bDoInitialClear)
 	{
-		if (RHICmdList && RHICmdList->IsTopOfPipe())
-		{
-			ALLOC_COMMAND_CL(*RHICmdList, FRHICommandSetInitialImageState)(this, InitialLayout, false, bDoInitialClear, InCreateDesc.ClearValue, bIsTransientResource);
-		}
-		else
+		if (!RHICmdList)
 		{
 			RHICmdList = &FRHICommandListExecutor::GetImmediateCommandList();
-			if (!IsInRenderingThread() || (RHICmdList->Bypass() || !IsRunningRHIInSeparateThread()))
-			{
-				SetInitialImageState(Device->GetImmediateContext(), InitialLayout, bDoInitialClear, InCreateDesc.ClearValue, bIsTransientResource);
-			}
-			else
-			{
-				check(IsInRenderingThread());
-				ALLOC_COMMAND_CL(*RHICmdList, FRHICommandSetInitialImageState)(this, InitialLayout, false, bDoInitialClear, InCreateDesc.ClearValue, bIsTransientResource);
-			}
 		}
+		SetInitialImageState(*RHICmdList, InitialLayout, bDoInitialClear, InCreateDesc.ClearValue, bIsTransientResource);
 	}
 
 	DefaultLayout = InitialLayout;
@@ -1702,15 +1564,11 @@ FVulkanTexture::FVulkanTexture(FRHICommandListBase* RHICmdList, FVulkanDevice& I
 	Region.imageExtent.depth = InCreateDesc.Depth;
 
 	checkf(RHICmdList, TEXT("FVulkanTexture requires a command list for creating bulk data."));
-
-	if (RHICmdList->IsTopOfPipe())
-	{
-		ALLOC_COMMAND_CL(*RHICmdList, FRHICommandLockWriteTexture)(this, Region, StagingBuffer);
-	}
-	else
-	{
-		FVulkanTexture::InternalLockWrite(InDevice.GetImmediateContext(), this, Region, StagingBuffer);
-	}
+	RHICmdList->EnqueueLambda(TEXT("FVulkanTexture::InternalLockWrite"),
+		[VulkanTexture = this, Region, StagingBuffer](FRHICommandListBase& ExecutingCmdList)
+		{
+			FVulkanTexture::InternalLockWrite(FVulkanUploadContext::Get(ExecutingCmdList), VulkanTexture, Region, StagingBuffer);
+		});
 }
 
 FVulkanTexture::FVulkanTexture(FVulkanDevice& InDevice, const FRHITextureCreateDesc& InCreateDesc, VkImage InImage, const FVulkanRHIExternalImageDeleteCallbackInfo& InExternalImageDeleteCallbackInfo)
@@ -1755,31 +1613,12 @@ FVulkanTexture::FVulkanTexture(FVulkanDevice& InDevice, const FRHITextureCreateD
 			VULKAN_SET_DEBUG_NAME(InDevice, VK_OBJECT_TYPE_IMAGE, Image, TEXT("%s:(FVulkanTexture*)0x%p"), InCreateDesc.DebugName ? InCreateDesc.DebugName : TEXT("?"), this);
 
 			const bool bRenderTarget = EnumHasAnyFlags(InCreateDesc.Flags, TexCreate_RenderTargetable | TexCreate_DepthStencilTargetable);
-			const VkImageLayout InitialLayout = GetInitialLayoutFromRHIAccess(InCreateDesc.InitialState, bRenderTarget && IsDepthOrStencilAspect(), SupportsSampling());
+			DefaultLayout = GetInitialLayoutFromRHIAccess(InCreateDesc.InitialState, bRenderTarget && IsDepthOrStencilAspect(), SupportsSampling());
 			const bool bDoInitialClear = bRenderTarget;
-			const bool bOnlyAddToLayoutManager = !bRenderTarget;
-
-			DefaultLayout = InitialLayout;
+			const VkImageLayout InitialLayout = VK_IMAGE_LAYOUT_UNDEFINED;  // use undefinied to avoid transitioning the texture when aliasing
 
 			FRHICommandList& RHICmdList = FRHICommandListExecutor::GetImmediateCommandList();
-			if (!IsInRenderingThread() || (RHICmdList.Bypass() || !IsRunningRHIInSeparateThread()))
-			{
-				FVulkanCommandListContext& Context = InDevice.GetImmediateContext();
-				if (bOnlyAddToLayoutManager)
-				{
-					FVulkanCommandBuffer& CommandBuffer = Context.GetCommandBuffer();
-					CommandBuffer.GetLayoutManager().SetFullLayout(*this, InitialLayout, true);
-				}
-				else if (InitialLayout != VK_IMAGE_LAYOUT_UNDEFINED || bDoInitialClear)
-				{
-					SetInitialImageState(Context, InitialLayout, bDoInitialClear, InCreateDesc.ClearValue, false);
-				}
-			}
-			else
-			{
-				check(IsInRenderingThread());
-				ALLOC_COMMAND_CL(RHICmdList, FRHICommandSetInitialImageState)(this, InitialLayout, bOnlyAddToLayoutManager, bDoInitialClear, InCreateDesc.ClearValue, false);
-			}
+			SetInitialImageState(RHICmdList, InitialLayout, bDoInitialClear, InCreateDesc.ClearValue, false);
 		}
 	}
 
@@ -2012,31 +1851,12 @@ FVulkanTexture::FVulkanTexture(FVulkanDevice& InDevice, const FRHITextureCreateD
 			VULKAN_SET_DEBUG_NAME(InDevice, VK_OBJECT_TYPE_IMAGE, Image, TEXT("%s:(FVulkanTexture*)0x%p"), InCreateDesc.DebugName ? InCreateDesc.DebugName : TEXT("?"), this);
 
 			const bool bRenderTarget = EnumHasAnyFlags(InCreateDesc.Flags, TexCreate_RenderTargetable | TexCreate_DepthStencilTargetable);
-			const VkImageLayout InitialLayout = GetInitialLayoutFromRHIAccess(InCreateDesc.InitialState, bRenderTarget && IsDepthOrStencilAspect(), SupportsSampling());
+			DefaultLayout = GetInitialLayoutFromRHIAccess(InCreateDesc.InitialState, bRenderTarget && IsDepthOrStencilAspect(), SupportsSampling());
 			const bool bDoInitialClear = bRenderTarget;
-			const bool bOnlyAddToLayoutManager = !bRenderTarget;
-
-			DefaultLayout = InitialLayout;
+			const VkImageLayout InitialLayout = VK_IMAGE_LAYOUT_UNDEFINED;  // use undefinied to avoid transitioning the texture when aliasing
 
-			FRHICommandList& RHICmdList = FRHICommandListExecutor::GetImmediateCommandList();
-			if (!IsInRenderingThread() || (RHICmdList.Bypass() || !IsRunningRHIInSeparateThread()))
-			{
-				FVulkanCommandListContext& Context = InDevice.GetImmediateContext();
-				if (bOnlyAddToLayoutManager)
-				{
-					FVulkanCommandBuffer& CommandBuffer = Context.GetCommandBuffer();
-					CommandBuffer.GetLayoutManager().SetFullLayout(*this, InitialLayout, true);
-				}
-				else if (InitialLayout != VK_IMAGE_LAYOUT_UNDEFINED || bDoInitialClear)
-				{
-					SetInitialImageState(Context, InitialLayout, bDoInitialClear, InCreateDesc.ClearValue, false);
-				}
-			}
-			else
-			{
-				check(IsInRenderingThread());
-				ALLOC_COMMAND_CL(RHICmdList, FRHICommandSetInitialImageState)(this, InitialLayout, bOnlyAddToLayoutManager, bDoInitialClear, InCreateDesc.ClearValue, false);
-			}
+			FRHICommandList& RHICmdList = FRHICommandListImmediate::Get();
+			SetInitialImageState(RHICmdList, InitialLayout, bDoInitialClear, InCreateDesc.ClearValue, false);
 		}
 	}
 
diff --git a/Engine/Source/Runtime/VulkanRHI/Private/VulkanUAV.cpp b/Engine/Source/Runtime/VulkanRHI/Private/VulkanUAV.cpp
--- a/Engine/Source/Runtime/VulkanRHI/Private/VulkanUAV.cpp
+++ b/Engine/Source/Runtime/VulkanRHI/Private/VulkanUAV.cpp
@@ -609,7 +609,7 @@ void FVulkanUnorderedAccessView::Clear(TRHICommandList_RecursiveHazardous<FVulka
 			{
 				const uint32 ComputeWorkGroupCount = FMath::DivideAndRoundUp(Info.NumElements, (uint32)ClearReplacementCS::TThreadGroupSize<EClearReplacementResourceType::Buffer>::X);
 
-				FVulkanDevice& TargetDevice = FVulkanCommandListContext::GetVulkanContext(RHICmdList.GetContext()).Device;
+				FVulkanDevice& TargetDevice = FVulkanCommandListContext::Get(RHICmdList).Device;
 				const bool bOversizedBuffer = (ComputeWorkGroupCount > TargetDevice.GetLimits().maxComputeWorkGroupCount[0]);
 
 				if (bOversizedBuffer)
diff --git a/Engine/Source/Runtime/VulkanRHI/Private/VulkanUniformBuffer.cpp b/Engine/Source/Runtime/VulkanRHI/Private/VulkanUniformBuffer.cpp
--- a/Engine/Source/Runtime/VulkanRHI/Private/VulkanUniformBuffer.cpp
+++ b/Engine/Source/Runtime/VulkanRHI/Private/VulkanUniformBuffer.cpp
@@ -43,7 +43,7 @@ static bool UseTemporaryBuffer(EUniformBufferUsage Usage)
 
 static void UpdateUniformBufferHelper(FVulkanCommandListContext& Context, FVulkanUniformBuffer* VulkanUniformBuffer, const void* Data, bool bUpdateConstants = true)
 {
-	FVulkanCommandBuffer* CmdBuffer = Context.GetActiveCmdBufferDirect();
+	FVulkanCommandBuffer* CmdBuffer = Context.GetActiveCmdBuffer();
 
 	FVulkanDevice& Device = Context.Device;
 	const int32 DataSize = VulkanUniformBuffer->GetLayout().ConstantBufferSize;
@@ -180,7 +180,7 @@ PRAGMA_ENABLE_DEPRECATION_WARNINGS
 
 					RHICmdList.EnqueueLambda([UniformBuffer, DataSize, CmdListConstantBufferData](FRHICommandList& CmdList)
 					{
-						FVulkanCommandListContext& Context = FVulkanCommandListContext::GetVulkanContext(CmdList.GetContext());
+						FVulkanCommandListContext& Context = FVulkanCommandListContext::Get(CmdList);
 						UpdateUniformBufferHelper(Context, UniformBuffer, CmdListConstantBufferData, false);
 					});
 
diff --git a/Engine/Source/Runtime/VulkanRHI/Private/VulkanViewport.cpp b/Engine/Source/Runtime/VulkanRHI/Private/VulkanViewport.cpp
--- a/Engine/Source/Runtime/VulkanRHI/Private/VulkanViewport.cpp
+++ b/Engine/Source/Runtime/VulkanRHI/Private/VulkanViewport.cpp
@@ -728,6 +728,46 @@ void FVulkanViewport::RecreateSwapchainFromRT(EPixelFormat PreferredPixelFormat)
 	check(RecreateInfo.SwapChain == VK_NULL_HANDLE);
 }
 
+void FVulkanViewport::InitImages(FVulkanContextCommon& Context, TConstArrayView<VkImage> Images)
+{
+	FVulkanCommandBuffer& CommandBuffer = Context.GetCommandBuffer();
+
+	VkClearColorValue ClearColor;
+	FMemory::Memzero(ClearColor);
+
+	const VkImageSubresourceRange Range = FVulkanPipelineBarrier::MakeSubresourceRange(VK_IMAGE_ASPECT_COLOR_BIT);
+	const FVulkanImageLayout InitialLayout(VK_IMAGE_LAYOUT_PRESENT_SRC_KHR, 1, 1, VK_IMAGE_ASPECT_COLOR_BIT);
+	for (int32 Index = 0; Index < Images.Num(); ++Index)
+	{
+		uint32 ImageSizeX = SizeX;
+		uint32 ImageSizeY = SizeY;
+		VkSurfaceTransformFlagBitsKHR CachedSurfaceTransform = SwapChain->GetCachedSurfaceTransform();
+		if (CachedSurfaceTransform == VK_SURFACE_TRANSFORM_ROTATE_90_BIT_KHR || CachedSurfaceTransform == VK_SURFACE_TRANSFORM_ROTATE_270_BIT_KHR)
+		{
+			Swap(ImageSizeX, ImageSizeY);
+		}
+		BackBufferImages[Index] = (FVulkanTexture*)FVulkanDynamicRHI::Get().RHICreateTexture2DFromResource(PixelFormat, ImageSizeX, ImageSizeY, 1, 1, Images[Index], ETextureCreateFlags::RenderTargetable | ETextureCreateFlags::Presentable).GetReference();
+		const VkDescriptorType DescriptorType = VK_DESCRIPTOR_TYPE_SAMPLED_IMAGE;
+		TextureViews.Add((new FVulkanView(*Device, DescriptorType))->InitAsTextureView(
+			Images[Index], VK_IMAGE_VIEW_TYPE_2D, VK_IMAGE_ASPECT_COLOR_BIT, PixelFormat, UEToVkTextureFormat(PixelFormat, false), 0, 1, 0, 1, false));
+
+		// Clear the swapchain to avoid a validation warning, and transition to PresentSrc
+		{
+			VulkanSetImageLayout(&CommandBuffer, Images[Index], VK_IMAGE_LAYOUT_UNDEFINED, VK_IMAGE_LAYOUT_TRANSFER_DST_OPTIMAL, Range);
+			VulkanRHI::vkCmdClearColorImage(CommandBuffer.GetHandle(), Images[Index], VK_IMAGE_LAYOUT_TRANSFER_DST_OPTIMAL, &ClearColor, 1, &Range);
+			VulkanSetImageLayout(&CommandBuffer, Images[Index], VK_IMAGE_LAYOUT_TRANSFER_DST_OPTIMAL, InitialLayout.MainLayout, Range);
+			CommandBuffer.GetLayoutManager().SetFullLayout(Images[Index], InitialLayout);
+		}
+
+#if VULKAN_ENABLE_DRAW_MARKERS
+		if (Device->GetSetDebugName())
+		{
+			VulkanRHI::SetDebugName(Device->GetSetDebugName(), Device->GetInstanceHandle(), BackBufferImages[Index]->Image, "VulkanBackBuffer");
+		}
+#endif
+	}
+}
+
 void FVulkanViewport::CreateSwapchain(FVulkanSwapChainRecreateInfo* RecreateInfo)
 {
 	// Release a previous swapchain 'dummy' and a real backbuffer if any
@@ -752,45 +792,35 @@ void FVulkanViewport::CreateSwapchain(FVulkanSwapChainRecreateInfo* RecreateInfo
 		BackBufferImages.SetNum(Images.Num());
 		RenderingDoneSemaphores.SetNum(Images.Num());
 
-		FVulkanCommandBuffer* CmdBuffer = Device->GetImmediateContext().GetUploadCmdBuffer();
-		ensure(CmdBuffer->IsOutsideRenderPass());
-
-		VkClearColorValue ClearColor;
-		FMemory::Memzero(ClearColor);
-
-		const VkImageSubresourceRange Range = FVulkanPipelineBarrier::MakeSubresourceRange(VK_IMAGE_ASPECT_COLOR_BIT);
-		const FVulkanImageLayout InitialLayout(VK_IMAGE_LAYOUT_PRESENT_SRC_KHR, 1, 1, VK_IMAGE_ASPECT_COLOR_BIT);
-		for (int32 Index = 0; Index < Images.Num(); ++Index)
+		if (IsInGameThread())
 		{
-			uint32 ImageSizeX = SizeX;
-			uint32 ImageSizeY = SizeY;
-			VkSurfaceTransformFlagBitsKHR CachedSurfaceTransform = SwapChain->GetCachedSurfaceTransform();
-			if (CachedSurfaceTransform == VK_SURFACE_TRANSFORM_ROTATE_90_BIT_KHR || CachedSurfaceTransform == VK_SURFACE_TRANSFORM_ROTATE_270_BIT_KHR)
-			{
-				Swap(ImageSizeX, ImageSizeY);
-			}
-			BackBufferImages[Index] = (FVulkanTexture*)FVulkanDynamicRHI::Get().RHICreateTexture2DFromResource(PixelFormat, ImageSizeX, ImageSizeY, 1, 1, Images[Index], ETextureCreateFlags::RenderTargetable | ETextureCreateFlags::Presentable).GetReference();
-			const VkDescriptorType DescriptorType = VK_DESCRIPTOR_TYPE_SAMPLED_IMAGE;
-			TextureViews.Add((new FVulkanView(*Device, DescriptorType))->InitAsTextureView(
-				Images[Index], VK_IMAGE_VIEW_TYPE_2D, VK_IMAGE_ASPECT_COLOR_BIT, PixelFormat, UEToVkTextureFormat(PixelFormat, false), 0, 1, 0, 1, false));
-
-			// Clear the swapchain to avoid a validation warning, and transition to PresentSrc
-			{
-				VulkanSetImageLayout(CmdBuffer, Images[Index], VK_IMAGE_LAYOUT_UNDEFINED, VK_IMAGE_LAYOUT_TRANSFER_DST_OPTIMAL, Range);
-				VulkanRHI::vkCmdClearColorImage(CmdBuffer->GetHandle(), Images[Index], VK_IMAGE_LAYOUT_TRANSFER_DST_OPTIMAL, &ClearColor, 1, &Range);
-				VulkanSetImageLayout(CmdBuffer, Images[Index], VK_IMAGE_LAYOUT_TRANSFER_DST_OPTIMAL, InitialLayout.MainLayout, Range);
-				CmdBuffer->GetLayoutManager().SetFullLayout(Images[Index], InitialLayout);
-			}
+			FlushRenderingCommands();
 
-#if VULKAN_ENABLE_DRAW_MARKERS
-			if (Device->GetSetDebugName())
-			{
-				VulkanRHI::SetDebugName(Device->GetSetDebugName(), Device->GetInstanceHandle(), BackBufferImages[Index]->Image, "VulkanBackBuffer");
-			}
-#endif
+			ENQUEUE_RENDER_COMMAND(CreateSwapchain)(
+				[&](FRHICommandListImmediate& RHICmdList)
+				{
+					RHICmdList.EnqueueLambda(TEXT("FVulkanViewport::CreateSwapchain"), [VulkanViewport=this, Images](FRHICommandListImmediate& ImmCmdList) 
+						{
+							VulkanViewport->InitImages(FVulkanUploadContext::Get(ImmCmdList), Images);
+						});
+					RHICmdList.SubmitAndBlockUntilGPUIdle();
+				});
+
+			FlushRenderingCommands();
+		}
+		else if (IsInRenderingThread())
+		{
+			FRHICommandListImmediate& RHICmdList = FRHICommandListExecutor::GetImmediateCommandList();
+			RHICmdList.EnqueueLambda(TEXT("FVulkanViewport::CreateSwapchain"), [VulkanViewport = this, Images](FRHICommandListImmediate& ImmCmdList) 
+				{
+					VulkanViewport->InitImages(FVulkanUploadContext::Get(ImmCmdList), Images);
+				});
+			RHICmdList.SubmitAndBlockUntilGPUIdle();
+		}
+		else
+		{
+			InitImages(FVulkanUploadContext::Get(FRHICommandListExecutor::GetImmediateCommandList()), Images);
 		}
-		
-		Device->GetImmediateContext().SubmitUploadCmdBuffer();
 
 		RHIBackBuffer = new FVulkanBackBuffer(*Device, this, PixelFormat, SizeX, SizeY, TexCreate_RenderTargetable | TexCreate_ShaderResource | TexCreate_ResolveTargetable);
 	}
@@ -1060,7 +1090,7 @@ bool FVulkanViewport::Present(FVulkanCommandListContext* Context, FVulkanCommand
 	CmdBuffer->End();
 	FVulkanCommandListContextImmediate& ImmediateContext = Device->GetImmediateContext();
 	ImmediateContext.FlushResetQueryPools();
-	checkf(ImmediateContext.GetActiveCmdBufferDirect() == CmdBuffer, TEXT("Present() is submitting something else than the active command buffer"));
+	checkf(ImmediateContext.GetActiveCmdBuffer() == CmdBuffer, TEXT("Present() is submitting something else than the active command buffer"));
 	if (SupportsStandardSwapchain())
 	{
 		if (LIKELY(!bFailedToDelayAcquireBackbuffer))
diff --git a/Engine/Source/Runtime/VulkanRHI/Private/VulkanViewport.h b/Engine/Source/Runtime/VulkanRHI/Private/VulkanViewport.h
--- a/Engine/Source/Runtime/VulkanRHI/Private/VulkanViewport.h
+++ b/Engine/Source/Runtime/VulkanRHI/Private/VulkanViewport.h
@@ -148,6 +148,7 @@ protected:
 	void CreateSwapchain(struct FVulkanSwapChainRecreateInfo* RecreateInfo);
 	void DestroySwapchain(struct FVulkanSwapChainRecreateInfo* RecreateInfo);
 	bool TryAcquireImageIndex();
+	void InitImages(FVulkanContextCommon& Context, TConstArrayView<VkImage> Images);
 
 	void RecreateSwapchain(void* NewNativeWindow);
 	void RecreateSwapchainFromRT(EPixelFormat PreferredPixelFormat);
diff --git a/Engine/Source/Runtime/VulkanRHI/Public/IVulkanDynamicRHI.h b/Engine/Source/Runtime/VulkanRHI/Public/IVulkanDynamicRHI.h
--- a/Engine/Source/Runtime/VulkanRHI/Public/IVulkanDynamicRHI.h
+++ b/Engine/Source/Runtime/VulkanRHI/Public/IVulkanDynamicRHI.h
@@ -72,10 +72,14 @@ struct IVulkanDynamicRHI : public FDynamicRHI
 	virtual FVulkanRHIAllocationInfo RHIGetAllocationInfo(FRHIBuffer* InBuffer) const = 0;
 
 	virtual void           RHISetImageLayout(VkImage Image, VkImageLayout OldLayout, VkImageLayout NewLayout, const VkImageSubresourceRange& SubresourceRange) = 0;
+
+	UE_DEPRECATED(5.5, "Upload command buffers are deprecated. Use RHISetImageLayout().")
 	virtual void           RHISetUploadImageLayout(VkImage Image, VkImageLayout OldLayout, VkImageLayout NewLayout, const VkImageSubresourceRange& SubresourceRange) = 0;
 
 	virtual void           RHIFinishExternalComputeWork(VkCommandBuffer InCommandBuffer) = 0;
 	virtual void           RHIRegisterWork(uint32 NumPrimitives) = 0;
+	
+	UE_DEPRECATED(5.5, "Upload command buffers are deprecated.")
 	virtual void           RHISubmitUploadCommandBuffer() = 0;
 
 	virtual void           RHIVerifyResult(VkResult Result, const ANSICHAR* VkFuntion, const ANSICHAR* Filename, uint32 Line) = 0;
diff --git a/Engine/Source/Runtime/VulkanRHI/Public/VulkanDynamicRHI.h b/Engine/Source/Runtime/VulkanRHI/Public/VulkanDynamicRHI.h
--- a/Engine/Source/Runtime/VulkanRHI/Public/VulkanDynamicRHI.h
+++ b/Engine/Source/Runtime/VulkanRHI/Public/VulkanDynamicRHI.h
@@ -184,6 +184,7 @@ public:
 	virtual void* RHIGetNativeInstance() final override;
 	virtual class IRHICommandContext* RHIGetDefaultContext() final override;
 	virtual IRHIComputeContext* RHIGetCommandContext(ERHIPipeline Pipeline, FRHIGPUMask GPUMask) final override;
+	virtual IRHIUploadContext* RHIGetUploadContext() final override;
 	virtual void RHIFinalizeContext(FRHIFinalizeContextArgs&& Args, TRHIPipelineArray<IRHIPlatformCommandList*>& Output) final override;
 	virtual void RHISubmitCommandLists(FRHISubmitCommandListsArgs&& Args) final override;
 	virtual uint64 RHIGetMinimumAlignmentForBufferBackedSRV(EPixelFormat Format) final override;
