VulkanRHI: Remove unused DefaultTexture in FVulkanDevice.

#rnx

diff --git a/Engine/Source/Runtime/VulkanRHI/Private/VulkanDevice.cpp b/Engine/Source/Runtime/VulkanRHI/Private/VulkanDevice.cpp
--- a/Engine/Source/Runtime/VulkanRHI/Private/VulkanDevice.cpp
+++ b/Engine/Source/Runtime/VulkanRHI/Private/VulkanDevice.cpp
@@ -213,7 +213,6 @@ FVulkanDevice::FVulkanDevice(FVulkanDynamicRHI* InRHI, VkPhysicalDevice InGpu)
 	, MemoryManager(this)
 	, DeferredDeletionQueue(this)
 	, DefaultSampler(nullptr)
-	, DefaultTexture(nullptr)
 	, Gpu(InGpu)
 	, GfxQueue(nullptr)
 	, ComputeQueue(nullptr)
@@ -1357,14 +1356,6 @@ void FVulkanDevice::InitGPU()
 	{
 		FSamplerStateInitializerRHI Default(SF_Point);
 		DefaultSampler = ResourceCast(RHICreateSamplerState(Default).GetReference());
-
-		const FRHITextureCreateDesc Desc =
-			FRHITextureCreateDesc::Create2D(TEXT("FVulkanDevice_DefaultImage"), 1, 1, PF_B8G8R8A8)
-			.SetClearValue(FClearValueBinding::None)
-			.SetFlags(ETextureCreateFlags::RenderTargetable | ETextureCreateFlags::ShaderResource)
-			.SetInitialState(ERHIAccess::SRVMask);
-
-		DefaultTexture = new FVulkanTexture(*this, Desc, nullptr);
 	}
 
 	if (RHISupportsRayTracing(GMaxRHIShaderPlatform) && GetOptionalExtensions().HasRaytracingExtensions())
@@ -1420,9 +1411,6 @@ void FVulkanDevice::Destroy()
 	// No need to delete as it's stored in SamplerMap
 	DefaultSampler = nullptr;
 
-	delete DefaultTexture;
-	DefaultTexture = nullptr;
-
 	for (int32 Index = CommandContexts.Num() - 1; Index >= 0; --Index)
 	{
 		delete CommandContexts[Index];
diff --git a/Engine/Source/Runtime/VulkanRHI/Private/VulkanDevice.h b/Engine/Source/Runtime/VulkanRHI/Private/VulkanDevice.h
--- a/Engine/Source/Runtime/VulkanRHI/Private/VulkanDevice.h
+++ b/Engine/Source/Runtime/VulkanRHI/Private/VulkanDevice.h
@@ -563,7 +563,6 @@ private:
 	FVulkanShaderFactory ShaderFactory;
 
 	FVulkanSamplerState* DefaultSampler;
-	FVulkanTexture* DefaultTexture;
 
 	VkPhysicalDevice Gpu;
 	VkPhysicalDeviceProperties GpuProps;
