VulkanRHI: Clean up: Remove some uses of FRHICommandListExecutor::GetImmediateCommandList() from outside render thread in the VulkanRHI.  Split FVulkanTexture constructors to force a command list reference in the general case (and limit access to constructor without command list).  Add CmdList params to swapchain creation flow in viewport.

#rnx

diff --git a/Engine/Source/Runtime/VulkanRHI/Private/VulkanRHI.cpp b/Engine/Source/Runtime/VulkanRHI/Private/VulkanRHI.cpp
--- a/Engine/Source/Runtime/VulkanRHI/Private/VulkanRHI.cpp
+++ b/Engine/Source/Runtime/VulkanRHI/Private/VulkanRHI.cpp
@@ -1149,7 +1149,7 @@ void FVulkanCommandListContext::RHIEndDrawingViewport(FRHIViewport* ViewportRHI,
 
 	WriteEndTimestamp(&CommandBuffer);
 
-	bool bNativePresent = Viewport->Present(this, &CommandBuffer, &Queue, Device.GetPresentQueue(), bLockToVsync);
+	bool bNativePresent = Viewport->Present(*this, &CommandBuffer, &Queue, Device.GetPresentQueue(), bLockToVsync);
 	if (bNativePresent)
 	{
 		//#todo-rco: Check for r.FinishCurrentFrame
@@ -2012,8 +2012,13 @@ void FVulkanDynamicRHI::RecreateSwapChain(void* NewNativeWindow)
 			
 			for (FVulkanViewport* Viewport : Viewports)
 			{
-				Viewport->RecreateSwapchain(NewNativeWindow);
+				RHICmdList.EnqueueLambda(TEXT("RecreateSwapchain"), [Viewport, NewNativeWindow](FRHICommandListImmediate& ExecutingCmdList)
+				{
+					Viewport->RecreateSwapchain(FVulkanCommandListContext::Get(ExecutingCmdList), NewNativeWindow);
+				});
 			}
+
+			RHICmdList.SubmitAndBlockUntilGPUIdle();
 		});
 
 		if (IsInGameThread())
diff --git a/Engine/Source/Runtime/VulkanRHI/Private/VulkanResources.h b/Engine/Source/Runtime/VulkanRHI/Private/VulkanResources.h
--- a/Engine/Source/Runtime/VulkanRHI/Private/VulkanResources.h
+++ b/Engine/Source/Runtime/VulkanRHI/Private/VulkanResources.h
@@ -573,13 +573,14 @@ enum class EImageOwnerType : uint8
 
 class FVulkanTexture : public FRHITexture, public FVulkanEvictable, public FVulkanViewableResource
 {
+	// Constructor without command list: targets will not be cleared, cannot contain initial data, will be left in VK_IMAGE_LAYOUT_UNDEFINED.
+	// Only used by transient texture allocator.
+	friend class FVulkanTransientResourceAllocator;
+	FVulkanTexture(FVulkanDevice& InDevice, const FRHITextureCreateDesc& InCreateDesc, const FRHITransientHeapAllocation* InTransientHeapAllocation);
+
 public:
 	// Regular constructor.
-	FVulkanTexture(FRHICommandListBase* RHICmdList, FVulkanDevice& InDevice, const FRHITextureCreateDesc& InCreateDesc, const FRHITransientHeapAllocation* InTransientHeapAllocation);
-
-	FVulkanTexture(FVulkanDevice& InDevice, const FRHITextureCreateDesc& InCreateDesc, const FRHITransientHeapAllocation* InTransientHeapAllocation)
-		: FVulkanTexture(nullptr, InDevice, InCreateDesc, InTransientHeapAllocation)
-	{}
+	FVulkanTexture(FRHICommandListBase& RHICmdList, FVulkanDevice& InDevice, const FRHITextureCreateDesc& InCreateDesc, const FRHITransientHeapAllocation* InTransientHeapAllocation);
 
 	// Construct from external resource.
 	FVulkanTexture(FVulkanDevice& InDevice, const FRHITextureCreateDesc& InCreateDesc, VkImage InImage, const FVulkanRHIExternalImageDeleteCallbackInfo& InExternalImageDeleteCallbackInfo);
@@ -765,8 +766,6 @@ private:
 	FVulkanCpuReadbackBuffer* CpuReadbackBuffer;
 	VkImageLayout DefaultLayout;
 
-	friend struct FRHICommandSetInitialImageState;
-
 protected:
 	EImageOwnerType ImageOwnerType;
 };
diff --git a/Engine/Source/Runtime/VulkanRHI/Private/VulkanTexture.cpp b/Engine/Source/Runtime/VulkanRHI/Private/VulkanTexture.cpp
--- a/Engine/Source/Runtime/VulkanRHI/Private/VulkanTexture.cpp
+++ b/Engine/Source/Runtime/VulkanRHI/Private/VulkanTexture.cpp
@@ -973,7 +973,7 @@ void FVulkanDynamicRHI::RHIUpdateTextureReference(FRHICommandListBase& RHICmdLis
 FTextureRHIRef FVulkanDynamicRHI::RHICreateTexture(FRHICommandListBase& RHICmdList, const FRHITextureCreateDesc& CreateDesc)
 {
 	LLM_SCOPE_VULKAN(GetMemoryTagForTextureFlags(CreateDesc.Flags));
-	return new FVulkanTexture(&RHICmdList, *Device, CreateDesc, nullptr);
+	return new FVulkanTexture(RHICmdList, *Device, CreateDesc, nullptr);
 }
 
 FTextureRHIRef FVulkanDynamicRHI::RHIAsyncCreateTexture2D(uint32 SizeX,uint32 SizeY,uint8 Format,uint32 NumMips,ETextureCreateFlags Flags, ERHIAccess InResourceState,void** InitialMipData,uint32 NumInitialMips, const TCHAR* DebugName, FGraphEventRef& OutCompletionEvent)
@@ -1069,7 +1069,7 @@ FTextureRHIRef FVulkanDynamicRHI::AsyncReallocateTexture2D_RenderThread(FRHIComm
 		.SetNumSamples(OldDesc.NumSamples)
 		.DetermineInititialState();
 
-	FVulkanTexture* NewTexture = new FVulkanTexture(&RHICmdList, *Device, Desc, nullptr);
+	FVulkanTexture* NewTexture = new FVulkanTexture(RHICmdList, *Device, Desc, nullptr);
 
 	RHICmdList.EnqueueLambda(TEXT("AsyncReallocateTexture2D"), [OldTexture, NewTexture, NewMipCount, NewSizeX, NewSizeY, RequestStatus](FRHICommandListImmediate& ImmCmdList)
 	{
@@ -1082,23 +1082,7 @@ FTextureRHIRef FVulkanDynamicRHI::AsyncReallocateTexture2D_RenderThread(FRHIComm
 
 FTextureRHIRef FVulkanDynamicRHI::RHIAsyncReallocateTexture2D(FRHITexture* OldTextureRHI, int32 NewMipCount, int32 NewSizeX, int32 NewSizeY, FThreadSafeCounter* RequestStatus)
 {
-	LLM_SCOPE_VULKAN(ELLMTagVulkan::VulkanTextures);
-	FVulkanTexture* OldTexture = ResourceCast(OldTextureRHI);
-	const FRHITextureDesc& OldDesc = OldTexture->GetDesc();
-
-	const FRHITextureCreateDesc Desc =
-		FRHITextureCreateDesc::Create2D(TEXT("AsyncReallocateTexture2D_RenderThread"), NewSizeX, NewSizeY, OldDesc.Format)
-		.SetClearValue(OldDesc.ClearValue)
-		.SetFlags(OldDesc.Flags)
-		.SetNumMips(NewMipCount)
-		.SetNumSamples(OldDesc.NumSamples)
-		.DetermineInititialState();
-
-	FVulkanTexture* NewTexture = new FVulkanTexture(*Device, Desc, nullptr);
-
-	DoAsyncReallocateTexture2D(Device->GetImmediateContext(), OldTexture, NewTexture, NewMipCount, NewSizeX, NewSizeY, RequestStatus);
-
-	return NewTexture;
+	return AsyncReallocateTexture2D_RenderThread(FRHICommandListImmediate::Get(), OldTextureRHI, NewMipCount, NewSizeX, NewSizeY, RequestStatus);
 }
 
 static FCriticalSection GTextureMapLock;
@@ -1329,7 +1313,7 @@ void FVulkanDynamicRHI::InternalUpdateTexture3D(FRHICommandListBase& RHICmdList,
 		});
 }
 
-FVulkanTexture::FVulkanTexture(FRHICommandListBase* RHICmdList, FVulkanDevice& InDevice, const FRHITextureCreateDesc& InCreateDesc, const FRHITransientHeapAllocation* InTransientHeapAllocation)
+FVulkanTexture::FVulkanTexture(FVulkanDevice& InDevice, const FRHITextureCreateDesc& InCreateDesc, const FRHITransientHeapAllocation* InTransientHeapAllocation)
 	: FRHITexture(InCreateDesc)
 	, PartialView(nullptr)
 	, Device(&InDevice)
@@ -1346,6 +1330,8 @@ FVulkanTexture::FVulkanTexture(FRHICommandListBase* RHICmdList, FVulkanDevice& I
 {
 	VULKAN_TRACK_OBJECT_CREATE(FVulkanTexture, this);
 
+	checkf(InCreateDesc.BulkData == nullptr, TEXT("Must call the constructor with a command list to initialize a texture with data."));
+
 	if (EnumHasAnyFlags(InCreateDesc.Flags, TexCreate_CPUReadback))
 	{
 		check(InCreateDesc.NumSamples == 1); //not implemented
@@ -1467,19 +1453,7 @@ FVulkanTexture::FVulkanTexture(FRHICommandListBase* RHICmdList, FVulkanDevice& I
 	check(Tiling == VK_IMAGE_TILING_LINEAR || Tiling == VK_IMAGE_TILING_OPTIMAL);
 	ImageUsageFlags = ImageCreateInfo.ImageCreateInfo.usage;
 
-	const VkImageLayout InitialLayout = GetInitialLayoutFromRHIAccess(InCreateDesc.InitialState, bRenderTarget && IsDepthOrStencilAspect(), SupportsSampling());
-	const bool bDoInitialClear = VKHasAnyFlags(ImageCreateInfo.ImageCreateInfo.usage, VK_IMAGE_USAGE_SAMPLED_BIT) && EnumHasAnyFlags(InCreateDesc.Flags, TexCreate_RenderTargetable | TexCreate_DepthStencilTargetable);
-
-	if (InitialLayout != VK_IMAGE_LAYOUT_UNDEFINED || bDoInitialClear)
-	{
-		if (!RHICmdList)
-		{
-			RHICmdList = &FRHICommandListExecutor::GetImmediateCommandList();
-		}
-		SetInitialImageState(*RHICmdList, InitialLayout, bDoInitialClear, InCreateDesc.ClearValue, bIsTransientResource);
-	}
-
-	DefaultLayout = InitialLayout;
+	DefaultLayout = GetInitialLayoutFromRHIAccess(InCreateDesc.InitialState, bRenderTarget && IsDepthOrStencilAspect(), SupportsSampling());
 
 	LLM_SCOPE_VULKAN(ELLMTagVulkan::VulkanTextures);
 	const VkImageViewType ViewType = GetViewType();
@@ -1529,46 +1503,61 @@ FVulkanTexture::FVulkanTexture(FRHICommandListBase* RHICmdList, FVulkanDevice& I
 			, false
 		);
 	}
+}
 
-	if (!InCreateDesc.BulkData)
+FVulkanTexture::FVulkanTexture(FRHICommandListBase& RHICmdList, FVulkanDevice& InDevice, const FRHITextureCreateDesc& InCreateDesc, const FRHITransientHeapAllocation* InTransientHeapAllocation)
+	: FVulkanTexture(InDevice, FRHITextureCreateDesc(InCreateDesc, InCreateDesc.InitialState, InCreateDesc.DebugName), InTransientHeapAllocation)
+{
+	if (EnumHasAnyFlags(InCreateDesc.Flags, TexCreate_CPUReadback))
 	{
 		return;
 	}
 
-	// InternalLockWrite leaves the image in VK_IMAGE_LAYOUT_SHADER_READ_ONLY_OPTIMAL, so make sure the requested resource state is SRV.
-	check(EnumHasAnyFlags(InCreateDesc.InitialState, ERHIAccess::SRVMask));
-	DefaultLayout = VK_IMAGE_LAYOUT_SHADER_READ_ONLY_OPTIMAL;
+	const bool bIsTransientResource = (InTransientHeapAllocation != nullptr);
+	const bool bDoInitialClear = VKHasAnyFlags(ImageUsageFlags, VK_IMAGE_USAGE_SAMPLED_BIT) &&
+		EnumHasAnyFlags(InCreateDesc.Flags, TexCreate_RenderTargetable | TexCreate_DepthStencilTargetable) &&
+		!bIsTransientResource;
+	if (DefaultLayout != VK_IMAGE_LAYOUT_UNDEFINED || bDoInitialClear)
+	{
+		SetInitialImageState(RHICmdList, DefaultLayout, bDoInitialClear, InCreateDesc.ClearValue, bIsTransientResource);
+	}
 
-	// Transfer bulk data
-	VulkanRHI::FStagingBuffer* StagingBuffer = InDevice.GetStagingManager().AcquireBuffer(InCreateDesc.BulkData->GetResourceBulkDataSize());
-	void* Data = StagingBuffer->GetMappedPointer();
+	if (InCreateDesc.BulkData)
+	{
+		// InternalLockWrite leaves the image in VK_IMAGE_LAYOUT_SHADER_READ_ONLY_OPTIMAL, so make sure the requested resource state is SRV.
+		check(EnumHasAnyFlags(InCreateDesc.InitialState, ERHIAccess::SRVMask));
+		DefaultLayout = VK_IMAGE_LAYOUT_SHADER_READ_ONLY_OPTIMAL;
 
-	// Do copy
-	FMemory::Memcpy(Data, InCreateDesc.BulkData->GetResourceBulkData(), InCreateDesc.BulkData->GetResourceBulkDataSize());
-	InCreateDesc.BulkData->Discard();
+		// Transfer bulk data
+		VulkanRHI::FStagingBuffer* StagingBuffer = InDevice.GetStagingManager().AcquireBuffer(InCreateDesc.BulkData->GetResourceBulkDataSize());
+		void* Data = StagingBuffer->GetMappedPointer();
 
-	VkBufferImageCopy Region;
-	FMemory::Memzero(Region);
-	//#todo-rco: Use real Buffer offset when switching to suballocations!
-	Region.bufferOffset = 0;
-	Region.bufferRowLength = InCreateDesc.Extent.X;
-	Region.bufferImageHeight = InCreateDesc.Extent.Y;
-	
-	Region.imageSubresource.mipLevel = 0;
-	Region.imageSubresource.baseArrayLayer = 0;
-	Region.imageSubresource.layerCount = GetNumberOfArrayLevels();
-	Region.imageSubresource.aspectMask = GetFullAspectMask();
-
-	Region.imageExtent.width = Region.bufferRowLength;
-	Region.imageExtent.height = Region.bufferImageHeight;
-	Region.imageExtent.depth = InCreateDesc.Depth;
-
-	checkf(RHICmdList, TEXT("FVulkanTexture requires a command list for creating bulk data."));
-	RHICmdList->EnqueueLambda(TEXT("FVulkanTexture::InternalLockWrite"),
-		[VulkanTexture = this, Region, StagingBuffer](FRHICommandListBase& ExecutingCmdList)
-		{
-			FVulkanTexture::InternalLockWrite(FVulkanUploadContext::Get(ExecutingCmdList), VulkanTexture, Region, StagingBuffer);
-		});
+		// Do copy
+		FMemory::Memcpy(Data, InCreateDesc.BulkData->GetResourceBulkData(), InCreateDesc.BulkData->GetResourceBulkDataSize());
+		InCreateDesc.BulkData->Discard();
+
+		VkBufferImageCopy Region;
+		FMemory::Memzero(Region);
+		//#todo-rco: Use real Buffer offset when switching to suballocations!
+		Region.bufferOffset = 0;
+		Region.bufferRowLength = InCreateDesc.Extent.X;
+		Region.bufferImageHeight = InCreateDesc.Extent.Y;
+
+		Region.imageSubresource.mipLevel = 0;
+		Region.imageSubresource.baseArrayLayer = 0;
+		Region.imageSubresource.layerCount = GetNumberOfArrayLevels();
+		Region.imageSubresource.aspectMask = GetFullAspectMask();
+
+		Region.imageExtent.width = Region.bufferRowLength;
+		Region.imageExtent.height = Region.bufferImageHeight;
+		Region.imageExtent.depth = InCreateDesc.Depth;
+
+		RHICmdList.EnqueueLambda(TEXT("FVulkanTexture::InternalLockWrite"),
+			[VulkanTexture = this, Region, StagingBuffer](FRHICommandListBase& ExecutingCmdList)
+			{
+				FVulkanTexture::InternalLockWrite(FVulkanUploadContext::Get(ExecutingCmdList), VulkanTexture, Region, StagingBuffer);
+			});
+	}
 }
 
 FVulkanTexture::FVulkanTexture(FVulkanDevice& InDevice, const FRHITextureCreateDesc& InCreateDesc, VkImage InImage, const FVulkanRHIExternalImageDeleteCallbackInfo& InExternalImageDeleteCallbackInfo)
diff --git a/Engine/Source/Runtime/VulkanRHI/Private/VulkanViewport.cpp b/Engine/Source/Runtime/VulkanRHI/Private/VulkanViewport.cpp
--- a/Engine/Source/Runtime/VulkanRHI/Private/VulkanViewport.cpp
+++ b/Engine/Source/Runtime/VulkanRHI/Private/VulkanViewport.cpp
@@ -134,7 +134,15 @@ FVulkanViewport::FVulkanViewport(FVulkanDevice* InDevice, void* InWindowHandle,
 	FVulkanDynamicRHI::Get().InitInstance();
 
 	bRenderOffscreen = FParse::Param(FCommandLine::Get(), TEXT("RenderOffScreen"));
-	CreateSwapchain(nullptr);
+
+	FlushRenderingCommands();
+	ENQUEUE_RENDER_COMMAND(CreateSwapchain)(
+		[VulkanViewport=this](FRHICommandListImmediate& RHICmdList)
+		{
+			VulkanViewport->CreateSwapchain(RHICmdList, nullptr);
+			RHICmdList.SubmitAndBlockUntilGPUIdle();
+		});
+	FlushRenderingCommands();
 
 	if (SupportsStandardSwapchain())
 	{
@@ -181,7 +189,7 @@ FVulkanViewport::~FVulkanViewport()
 	FVulkanDynamicRHI::Get().Viewports.Remove(this);
 }
 
-bool FVulkanViewport::DoCheckedSwapChainJob(TFunction<int32(FVulkanViewport*)> SwapChainJob)
+bool FVulkanViewport::DoCheckedSwapChainJob(FVulkanCommandListContext& Context, TFunction<int32(FVulkanViewport*)> SwapChainJob)
 {
 	int32 AttemptsPending = FVulkanPlatform::RecreateSwapchainOnFail() ? 4 : 0;
 	int32 Status = SwapChainJob(this);
@@ -201,7 +209,7 @@ bool FVulkanViewport::DoCheckedSwapChainJob(TFunction<int32(FVulkanViewport*)> S
 			check(0);
 		}
 
-		RecreateSwapchain(WindowHandle);
+		RecreateSwapchain(Context, WindowHandle);
 
 		// Swapchain creation pushes some commands - flush the command buffers now to begin with a fresh state
 		Device->SubmitCommandsAndFlushGPU();
@@ -676,14 +684,16 @@ bool FVulkanFramebuffer::Matches(const FRHISetRenderTargetsInfo& InRTInfo) const
 }
 
 // Tear down and recreate swapchain and related resources.
-void FVulkanViewport::RecreateSwapchain(void* NewNativeWindow)
+void FVulkanViewport::RecreateSwapchain(FVulkanCommandListContext& Context, void* NewNativeWindow)
 {
 	FScopeLock LockSwapchain(&RecreatingSwapchain);
 
+	TRHICommandList_RecursiveHazardous<FVulkanCommandListContext> RHICmdList(&Context);
+
 	FVulkanSwapChainRecreateInfo RecreateInfo = { VK_NULL_HANDLE, VK_NULL_HANDLE };
 	DestroySwapchain(&RecreateInfo);
 	WindowHandle = NewNativeWindow;
-	CreateSwapchain(&RecreateInfo);
+	CreateSwapchain(RHICmdList, &RecreateInfo);
 	check(RecreateInfo.Surface == VK_NULL_HANDLE);
 	check(RecreateInfo.SwapChain == VK_NULL_HANDLE);
 }
@@ -698,7 +708,7 @@ void FVulkanViewport::Tick(float DeltaTime)
 		ENQUEUE_RENDER_COMMAND(UpdateVsync)(
 			[this](FRHICommandListImmediate& RHICmdList)
 		{
-			RecreateSwapchainFromRT(PixelFormat);
+			RecreateSwapchainFromRT(RHICmdList, PixelFormat);
 		});
 		FlushRenderingCommands();
 	}
@@ -710,20 +720,20 @@ void FVulkanViewport::Resize(uint32 InSizeX, uint32 InSizeY, bool bInIsFullscree
 	SizeY = InSizeY;
 	bIsFullscreen = bInIsFullscreen;
 
-	RecreateSwapchainFromRT(PreferredPixelFormat);
+	RecreateSwapchainFromRT(FRHICommandListImmediate::Get(), PreferredPixelFormat);
 }
 
-void FVulkanViewport::RecreateSwapchainFromRT(EPixelFormat PreferredPixelFormat)
+void FVulkanViewport::RecreateSwapchainFromRT(FRHICommandListImmediate& RHICmdList, EPixelFormat PreferredPixelFormat)
 {
 	check(IsInRenderingThread());
 
 	// TODO: should flush RHIT commands here?
-	FRHICommandListExecutor::GetImmediateCommandList().ImmediateFlush(EImmediateFlushType::FlushRHIThread);
+	RHICmdList.ImmediateFlush(EImmediateFlushType::FlushRHIThread);
 	
 	FVulkanSwapChainRecreateInfo RecreateInfo = { VK_NULL_HANDLE, VK_NULL_HANDLE};
 	DestroySwapchain(&RecreateInfo);
 	PixelFormat = PreferredPixelFormat;
-	CreateSwapchain(&RecreateInfo);
+	CreateSwapchain(RHICmdList, &RecreateInfo);
 	check(RecreateInfo.Surface == VK_NULL_HANDLE);
 	check(RecreateInfo.SwapChain == VK_NULL_HANDLE);
 }
@@ -768,7 +778,7 @@ void FVulkanViewport::InitImages(FVulkanContextCommon& Context, TConstArrayView<
 	}
 }
 
-void FVulkanViewport::CreateSwapchain(FVulkanSwapChainRecreateInfo* RecreateInfo)
+void FVulkanViewport::CreateSwapchain(FRHICommandListBase& RHICmdList, FVulkanSwapChainRecreateInfo* RecreateInfo)
 {
 	// Release a previous swapchain 'dummy' and a real backbuffer if any
 	RenderingBackBuffer = nullptr;
@@ -792,34 +802,13 @@ void FVulkanViewport::CreateSwapchain(FVulkanSwapChainRecreateInfo* RecreateInfo
 		BackBufferImages.SetNum(Images.Num());
 		RenderingDoneSemaphores.SetNum(Images.Num());
 
-		if (IsInGameThread())
-		{
-			FlushRenderingCommands();
-
-			ENQUEUE_RENDER_COMMAND(CreateSwapchain)(
-				[&](FRHICommandListImmediate& RHICmdList)
-				{
-					RHICmdList.EnqueueLambda(TEXT("FVulkanViewport::CreateSwapchain"), [VulkanViewport=this, Images](FRHICommandListImmediate& ImmCmdList) 
-						{
-							VulkanViewport->InitImages(FVulkanUploadContext::Get(ImmCmdList), Images);
-						});
-					RHICmdList.SubmitAndBlockUntilGPUIdle();
-				});
-
-			FlushRenderingCommands();
-		}
-		else if (IsInRenderingThread())
+		RHICmdList.EnqueueLambda(TEXT("FVulkanViewport::CreateSwapchain"), [VulkanViewport = this, Images](FRHICommandListBase& ExecutingCmdList) 
 		{
-			FRHICommandListImmediate& RHICmdList = FRHICommandListExecutor::GetImmediateCommandList();
-			RHICmdList.EnqueueLambda(TEXT("FVulkanViewport::CreateSwapchain"), [VulkanViewport = this, Images](FRHICommandListImmediate& ImmCmdList) 
-				{
-					VulkanViewport->InitImages(FVulkanUploadContext::Get(ImmCmdList), Images);
-				});
-			RHICmdList.SubmitAndBlockUntilGPUIdle();
-		}
-		else
+			VulkanViewport->InitImages(FVulkanUploadContext::Get(ExecutingCmdList), Images);
+		});
+		if (RHICmdList.IsImmediate())
 		{
-			InitImages(FVulkanUploadContext::Get(FRHICommandListExecutor::GetImmediateCommandList()), Images);
+			RHICmdList.GetAsImmediate().SubmitAndBlockUntilGPUIdle();
 		}
 
 		RHIBackBuffer = new FVulkanBackBuffer(*Device, this, PixelFormat, SizeX, SizeY, TexCreate_RenderTargetable | TexCreate_ShaderResource | TexCreate_ResolveTargetable);
@@ -857,7 +846,7 @@ void FVulkanViewport::CreateSwapchain(FVulkanSwapChainRecreateInfo* RecreateInfo
 			.SetFlags(ETextureCreateFlags::RenderTargetable | ETextureCreateFlags::ShaderResource | ETextureCreateFlags::ResolveTargetable)
 			.SetInitialState(ERHIAccess::Present);
 
-		RenderingBackBuffer = new FVulkanTexture(*Device, Desc, nullptr);
+		RenderingBackBuffer = new FVulkanTexture(RHICmdList, *Device, Desc, nullptr);
 #if VULKAN_ENABLE_DRAW_MARKERS
 		if (Device->GetSetDebugName())
 		{
@@ -905,9 +894,9 @@ void FVulkanViewport::DestroySwapchain(FVulkanSwapChainRecreateInfo* RecreateInf
 	AcquiredImageIndex = -1;
 }
 
-inline static void CopyImageToBackBuffer(FVulkanCommandListContext* Context, FVulkanCommandBuffer* CmdBuffer, FVulkanTexture& SrcSurface, FVulkanTexture& DstSurface, int32 SizeX, int32 SizeY, int32 WindowSizeX, int32 WindowSizeY, VkSurfaceTransformFlagBitsKHR CachedSurfaceTransform)
+inline static void CopyImageToBackBuffer(FVulkanCommandListContext& Context, FVulkanCommandBuffer* CmdBuffer, FVulkanTexture& SrcSurface, FVulkanTexture& DstSurface, int32 SizeX, int32 SizeY, int32 WindowSizeX, int32 WindowSizeY, VkSurfaceTransformFlagBitsKHR CachedSurfaceTransform)
 {
-	RHI_BREADCRUMB_EVENT(*Context, "CopyImageToBackBuffer");
+	RHI_BREADCRUMB_EVENT(Context, "CopyImageToBackBuffer");
 	const bool bNeedsVulkanPreTransform = CachedSurfaceTransform != VK_SURFACE_TRANSFORM_IDENTITY_BIT_KHR;
 
 	FVulkanLayoutManager& LayoutManager = CmdBuffer->GetLayoutManager();
@@ -945,7 +934,7 @@ inline static void CopyImageToBackBuffer(FVulkanCommandListContext* Context, FVu
 		GraphicsPSOInit.RasterizerState = TStaticRasterizerState<>::GetRHI();
 		GraphicsPSOInit.DepthStencilState = TStaticDepthStencilState<false, CF_Always>::GetRHI();
 
-		TRHICommandList_RecursiveHazardous<FVulkanCommandListContext> RHICmdList(Context);
+		TRHICommandList_RecursiveHazardous<FVulkanCommandListContext> RHICmdList(&Context);
 
 		RHICmdList.BeginRenderPass(FRHIRenderPassInfo(&DstSurface, ERenderTargetActions::DontLoad_Store), TEXT("SurfaceTransform"));
 
@@ -1032,8 +1021,10 @@ inline static void CopyImageToBackBuffer(FVulkanCommandListContext* Context, FVu
 	}
 }
 
-bool FVulkanViewport::Present(FVulkanCommandListContext* Context, FVulkanCommandBuffer* CmdBuffer, FVulkanQueue* Queue, FVulkanQueue* PresentQueue, bool bLockToVsync)
+bool FVulkanViewport::Present(FVulkanCommandListContext& Context, FVulkanCommandBuffer* CmdBuffer, FVulkanQueue* Queue, FVulkanQueue* PresentQueue, bool bLockToVsync)
 {
+	check(Context.IsImmediate());
+
 	FPlatformAtomics::AtomicStore(&LockToVsync, bLockToVsync ? 1 : 0);
 	bool bFailedToDelayAcquireBackbuffer = false;
 
@@ -1070,7 +1061,7 @@ bool FVulkanViewport::Present(FVulkanCommandListContext* Context, FVulkanCommand
 				// One of the rare cases we'll let parallel rendering check the tracking (legacy path already checked as a fallback)
 				if (!TrackedLayout && Device->SupportsParallelRendering())
 				{
-					TrackedLayout = Context->Queue.GetLayoutManager().GetFullLayout(BackBufferImages[AcquiredImageIndex]->Image);
+					TrackedLayout = Context.Queue.GetLayoutManager().GetFullLayout(BackBufferImages[AcquiredImageIndex]->Image);
 				}
 
 				VkImageLayout LastLayout = TrackedLayout ? TrackedLayout->MainLayout : VK_IMAGE_LAYOUT_UNDEFINED;
@@ -1088,9 +1079,9 @@ bool FVulkanViewport::Present(FVulkanCommandListContext* Context, FVulkanCommand
 	}
 
 	CmdBuffer->End();
-	FVulkanCommandListContextImmediate& ImmediateContext = Device->GetImmediateContext();
-	ImmediateContext.FlushResetQueryPools();
-	checkf(ImmediateContext.GetActiveCmdBuffer() == CmdBuffer, TEXT("Present() is submitting something else than the active command buffer"));
+
+	Context.FlushResetQueryPools();
+	checkf(Context.GetActiveCmdBuffer() == CmdBuffer, TEXT("Present() is submitting something else than the active command buffer"));
 	if (SupportsStandardSwapchain())
 	{
 		if (LIKELY(!bFailedToDelayAcquireBackbuffer))
@@ -1102,7 +1093,7 @@ bool FVulkanViewport::Present(FVulkanCommandListContext* Context, FVulkanCommand
 			
 			VulkanRHI::FSemaphore* SignalSemaphore = (AcquiredImageIndex >= 0 ? RenderingDoneSemaphores[AcquiredImageIndex] : nullptr);
 			// submit through the CommandBufferManager as it will add the proper semaphore
-			ImmediateContext.SubmitActiveCmdBufferFromPresent(SignalSemaphore);
+			Context.SubmitActiveCmdBufferFromPresent(SignalSemaphore);
 		}
 		else
 		{
@@ -1113,7 +1104,7 @@ bool FVulkanViewport::Present(FVulkanCommandListContext* Context, FVulkanCommand
 
 			// cannot just throw out this command buffer (needs to be submitted or other checks fail)
 			Queue->Submit(CmdBuffer);
-			RecreateSwapchain(WindowHandle);
+			RecreateSwapchain(Context, WindowHandle);
 
 			// Swapchain creation pushes some commands - flush the command buffers now to begin with a fresh state
 			Device->SubmitCommandsAndFlushGPU();
@@ -1126,7 +1117,7 @@ bool FVulkanViewport::Present(FVulkanCommandListContext* Context, FVulkanCommand
 	else
 	{
 		// Submit active command buffer if not supporting standard swapchain (e.g. XR devices).
-		ImmediateContext.SubmitActiveCmdBufferFromPresent(nullptr);
+		Context.SubmitActiveCmdBufferFromPresent(nullptr);
 	}
 
 	//Flush all commands
@@ -1140,8 +1131,7 @@ bool FVulkanViewport::Present(FVulkanCommandListContext* Context, FVulkanCommand
 	if (bHasCustomPresent)
 	{
 		SCOPE_CYCLE_COUNTER(STAT_VulkanCustomPresentTime);
-		check(Context);
-		bNeedNativePresent = CustomPresent->Present(*Context, SyncInterval);
+		bNeedNativePresent = CustomPresent->Present(Context, SyncInterval);
 	}
 
 	bool bResult = false;
@@ -1158,7 +1148,7 @@ bool FVulkanViewport::Present(FVulkanCommandListContext* Context, FVulkanCommand
 			}
 			return (int32)Viewport->SwapChain->Present(Queue, PresentQueue, Viewport->RenderingDoneSemaphores[Viewport->AcquiredImageIndex]);
 		};
-		if (SupportsStandardSwapchain() && !DoCheckedSwapChainJob(SwapChainJob))
+		if (SupportsStandardSwapchain() && !DoCheckedSwapChainJob(Context, SwapChainJob))
 		{
 			UE_LOG(LogVulkanRHI, Error, TEXT("Swapchain present failed!"));
 			bResult = false;
@@ -1192,9 +1182,9 @@ bool FVulkanViewport::Present(FVulkanCommandListContext* Context, FVulkanCommand
 	//}
 
 	// PrepareForNewActiveCommandBuffer might be called by swapchain re-creation routine. Skip prepare if we already have an open active buffer.
-	if (ImmediateContext.GetActiveCmdBuffer() && !ImmediateContext.GetActiveCmdBuffer()->HasBegun())
+	if (Context.GetActiveCmdBuffer() && !Context.GetActiveCmdBuffer()->HasBegun())
 	{
-		ImmediateContext.PrepareForNewActiveCommandBuffer();
+		Context.PrepareForNewActiveCommandBuffer();
 	}
 
 	AcquiredImageIndex = -1;
@@ -1251,7 +1241,7 @@ void FVulkanViewport::OnSystemResolutionChanged(uint32 ResX, uint32 ResY)
 		ENQUEUE_RENDER_COMMAND(RecreateSwapchain)(
 			[this](FRHICommandListImmediate& RHICmdList)
 			{
-				RecreateSwapchainFromRT(PixelFormat);
+				RecreateSwapchainFromRT(RHICmdList, PixelFormat);
 			});
 		FlushRenderingCommands();
 	}
diff --git a/Engine/Source/Runtime/VulkanRHI/Private/VulkanViewport.h b/Engine/Source/Runtime/VulkanRHI/Private/VulkanViewport.h
--- a/Engine/Source/Runtime/VulkanRHI/Private/VulkanViewport.h
+++ b/Engine/Source/Runtime/VulkanRHI/Private/VulkanViewport.h
@@ -73,7 +73,7 @@ public:
 
 	virtual void Tick(float DeltaTime) override final;
 	
-	bool Present(FVulkanCommandListContext* Context, FVulkanCommandBuffer* CmdBuffer, FVulkanQueue* Queue, FVulkanQueue* PresentQueue, bool bLockToVsync);
+	bool Present(FVulkanCommandListContext& Context, FVulkanCommandBuffer* CmdBuffer, FVulkanQueue* Queue, FVulkanQueue* PresentQueue, bool bLockToVsync);
 
 	inline uint32 GetPresentCount() const
 	{
@@ -145,16 +145,16 @@ protected:
 	EDeviceScreenOrientation CachedOrientation = EDeviceScreenOrientation::Unknown;
 	void OnSystemResolutionChanged(uint32 ResX, uint32 ResY);
 
-	void CreateSwapchain(struct FVulkanSwapChainRecreateInfo* RecreateInfo);
+	void CreateSwapchain(FRHICommandListBase& RHICmdList, struct FVulkanSwapChainRecreateInfo* RecreateInfo);
 	void DestroySwapchain(struct FVulkanSwapChainRecreateInfo* RecreateInfo);
 	bool TryAcquireImageIndex();
 	void InitImages(FVulkanContextCommon& Context, TConstArrayView<VkImage> Images);
 
-	void RecreateSwapchain(void* NewNativeWindow);
-	void RecreateSwapchainFromRT(EPixelFormat PreferredPixelFormat);
+	void RecreateSwapchain(FVulkanCommandListContext& Context, void* NewNativeWindow);
+	void RecreateSwapchainFromRT(FRHICommandListImmediate& RHICmdList, EPixelFormat PreferredPixelFormat);
 	void Resize(uint32 InSizeX, uint32 InSizeY, bool bIsFullscreen, EPixelFormat PreferredPixelFormat);
 
-	bool DoCheckedSwapChainJob(TFunction<int32(FVulkanViewport*)> SwapChainJob);
+	bool DoCheckedSwapChainJob(FVulkanCommandListContext& Context, TFunction<int32(FVulkanViewport*)> SwapChainJob);
 	bool SupportsStandardSwapchain();
 	bool RequiresRenderingBackBuffer();
 	EPixelFormat GetPixelFormatForNonDefaultSwapchain();
