Don't enable Vulkan validation layers when it's not requested on systems where debug_utils is not supported.

Since this would cause missing markers in tracing tools add support for legacy debug_marker to avoid this.

diff --git a/Engine/Source/Runtime/VulkanRHI/Private/VulkanDevice.cpp b/Engine/Source/Runtime/VulkanRHI/Private/VulkanDevice.cpp
--- a/Engine/Source/Runtime/VulkanRHI/Private/VulkanDevice.cpp
+++ b/Engine/Source/Runtime/VulkanRHI/Private/VulkanDevice.cpp
@@ -498,7 +498,7 @@ void FVulkanDevice::SetupDrawMarkers()
 	bTraceToolFound |= GRenderDocFound;
 #endif // VULKAN_HAS_DEBUGGING_ENABLED
 
-	if (RHI->SupportsDebugUtilsExt())
+	if (RHI->SupportsDebugUtilsExt() || bUseLegacyDebugMarkerExt)
 	{
 		bool bUseLabel = bGPUCrashDebugging;
 		bool bUseName = false;
@@ -540,14 +540,71 @@ void FVulkanDevice::SetupDrawMarkers()
 			bUseName = bUseName || (VulkanDebugMarkers & 1);
 		}
 
-		if (bUseLabel)
+		if (bUseLegacyDebugMarkerExt)
 		{
-			DebugMarkers.CmdBeginDebugLabel = (PFN_vkCmdBeginDebugUtilsLabelEXT)(void*)VulkanRHI::vkGetInstanceProcAddr(RHI->GetInstance(), "vkCmdBeginDebugUtilsLabelEXT");
-			DebugMarkers.CmdEndDebugLabel = (PFN_vkCmdEndDebugUtilsLabelEXT)(void*)VulkanRHI::vkGetInstanceProcAddr(RHI->GetInstance(), "vkCmdEndDebugUtilsLabelEXT");
+			if (bUseLabel)
+			{
+				static PFN_vkCmdDebugMarkerBeginEXT LegacyCmdBeginDebugLabel;
+				static PFN_vkCmdDebugMarkerEndEXT LegacyCmdEndDebugLabel;
+				LegacyCmdBeginDebugLabel = (PFN_vkCmdDebugMarkerBeginEXT)(void*)VulkanRHI::vkGetDeviceProcAddr(Device, "vkCmdDebugMarkerBeginEXT");
+				LegacyCmdEndDebugLabel = (PFN_vkCmdDebugMarkerEndEXT)(void*)VulkanRHI::vkGetDeviceProcAddr(Device, "vkCmdDebugMarkerEndEXT");
+				if (LegacyCmdBeginDebugLabel && LegacyCmdEndDebugLabel)
+				{
+					struct Wrap
+					{
+						static VKAPI_ATTR void VKAPI_CALL CmdBeginDebugLabel(VkCommandBuffer CommandBuffer, const VkDebugUtilsLabelEXT* LabelInfo)
+						{
+							VkDebugMarkerMarkerInfoEXT MarkerInfo{ VK_STRUCTURE_TYPE_DEBUG_MARKER_MARKER_INFO_EXT };
+							MarkerInfo.pMarkerName = LabelInfo->pLabelName;
+							memcpy(MarkerInfo.color, LabelInfo->color, sizeof(LabelInfo->color));
+							LegacyCmdBeginDebugLabel(CommandBuffer, &MarkerInfo);
+						}
+						static VKAPI_ATTR void VKAPI_CALL CmdEndDebugLabel(VkCommandBuffer CommandBuffer)
+						{
+							LegacyCmdEndDebugLabel(CommandBuffer);
+						}
+					};
+					DebugMarkers.CmdBeginDebugLabel = &Wrap::CmdBeginDebugLabel;
+					DebugMarkers.CmdEndDebugLabel = &Wrap::CmdEndDebugLabel;
+				}
+			}
+			if (bUseName)
+			{
+				static PFN_vkDebugMarkerSetObjectNameEXT LegacySetObjectName;
+				LegacySetObjectName = (PFN_vkDebugMarkerSetObjectNameEXT)(void*)VulkanRHI::vkGetDeviceProcAddr(Device, "vkDebugMarkerSetObjectNameEXT");
+				if (LegacySetObjectName)
+				{
+					struct Wrap
+					{
+						static VKAPI_ATTR VkResult VKAPI_CALL SetObjectName(VkDevice Device, const VkDebugUtilsObjectNameInfoEXT* NameInfo)
+						{
+							if (NameInfo->objectType <= 25)
+							{
+								VkDebugMarkerObjectNameInfoEXT MarkerInfo{ VK_STRUCTURE_TYPE_DEBUG_MARKER_OBJECT_NAME_INFO_EXT };
+								MarkerInfo.pObjectName = NameInfo->pObjectName;
+								MarkerInfo.object = NameInfo->objectHandle;
+								MarkerInfo.objectType = static_cast<VkDebugReportObjectTypeEXT>(NameInfo->objectType);
+								LegacySetObjectName(Device, &MarkerInfo);
+							}
+							return VK_SUCCESS;
+						}
+					};
+
+					DebugMarkers.SetDebugName = &Wrap::SetObjectName;
+				}
+			}
 		}
-		if (bUseName)
+		else
 		{
-			DebugMarkers.SetDebugName = (PFN_vkSetDebugUtilsObjectNameEXT)(void*)VulkanRHI::vkGetInstanceProcAddr(RHI->GetInstance(), "vkSetDebugUtilsObjectNameEXT");
+			if (bUseLabel)
+			{
+				DebugMarkers.CmdBeginDebugLabel = (PFN_vkCmdBeginDebugUtilsLabelEXT)(void*)VulkanRHI::vkGetInstanceProcAddr(RHI->GetInstance(), "vkCmdBeginDebugUtilsLabelEXT");
+				DebugMarkers.CmdEndDebugLabel = (PFN_vkCmdEndDebugUtilsLabelEXT)(void*)VulkanRHI::vkGetInstanceProcAddr(RHI->GetInstance(), "vkCmdEndDebugUtilsLabelEXT");
+			}
+			if (bUseName)
+			{
+				DebugMarkers.SetDebugName = (PFN_vkSetDebugUtilsObjectNameEXT)(void*)VulkanRHI::vkGetInstanceProcAddr(RHI->GetInstance(), "vkSetDebugUtilsObjectNameEXT");
+			}
 		}
 
 		UE_LOG(LogVulkanRHI, Display, TEXT("Vulkan debug markers support: resource names %s, begin/end labels %s."), bUseName ? TEXT("enabled") : TEXT("disabled"), bUseLabel ? TEXT("enabled") : TEXT("disabled"));
@@ -1204,7 +1261,7 @@ void FVulkanDevice::InitGPU()
 
 	// Setup layers and extensions
 	FVulkanDeviceExtensionArray UEExtensions = FVulkanDeviceExtension::GetUESupportedDeviceExtensions(this, RHI->GetApiVersion());
-	TArray<const ANSICHAR*> DeviceLayers = FVulkanDevice::SetupDeviceLayers(Gpu, UEExtensions);
+	TArray<const ANSICHAR*> DeviceLayers = SetupDeviceLayers(UEExtensions);
 
 	// Query advanced features
 	{
diff --git a/Engine/Source/Runtime/VulkanRHI/Private/VulkanDevice.h b/Engine/Source/Runtime/VulkanRHI/Private/VulkanDevice.h
--- a/Engine/Source/Runtime/VulkanRHI/Private/VulkanDevice.h
+++ b/Engine/Source/Runtime/VulkanRHI/Private/VulkanDevice.h
@@ -615,7 +615,7 @@ private:
 
 	FVulkanDynamicRHI* RHI = nullptr;
 	
-	static TArray<const ANSICHAR*> SetupDeviceLayers(VkPhysicalDevice Gpu, FVulkanDeviceExtensionArray& UEExtensions);
+	TArray<const ANSICHAR*> SetupDeviceLayers(FVulkanDeviceExtensionArray& UEExtensions);
 
 	FOptionalVulkanDeviceExtensions	OptionalDeviceExtensions;
 	FOptionalVulkanDeviceExtensionProperties OptionalDeviceExtensionProperties;
@@ -628,6 +628,7 @@ private:
 #endif
 
 #if VULKAN_ENABLE_DRAW_MARKERS
+	bool bUseLegacyDebugMarkerExt = false;
 	struct
 	{
 		PFN_vkSetDebugUtilsObjectNameEXT	SetDebugName = nullptr;
diff --git a/Engine/Source/Runtime/VulkanRHI/Private/VulkanExtensions.cpp b/Engine/Source/Runtime/VulkanRHI/Private/VulkanExtensions.cpp
--- a/Engine/Source/Runtime/VulkanRHI/Private/VulkanExtensions.cpp
+++ b/Engine/Source/Runtime/VulkanRHI/Private/VulkanExtensions.cpp
@@ -1737,8 +1737,7 @@ FVulkanDeviceExtensionArray FVulkanDeviceExtension::GetUESupportedDeviceExtensio
 	ADD_SIMPLE_EXTENSION(VK_KHR_DYNAMIC_RENDERING_EXTENSION_NAME,        VULKAN_EXTENSION_ENABLED,             VK_API_VERSION_1_3,            nullptr); // Only enabled because maintenance5 depends on it
 
 	// Externally activated extensions (supported by the engine, but enabled externally by plugin or other) :
-
-
+	ADD_EXTERNAL_EXTENSION(VK_EXT_DEBUG_MARKER_EXTENSION_NAME, VULKAN_HAS_DEBUGGING_ENABLED || VULKAN_ENABLE_DRAW_MARKERS, VULKAN_EXTENSION_NOT_PROMOTED, nullptr);
 
 	// Extensions with custom classes :
 
diff --git a/Engine/Source/Runtime/VulkanRHI/Private/VulkanLayers.cpp b/Engine/Source/Runtime/VulkanRHI/Private/VulkanLayers.cpp
--- a/Engine/Source/Runtime/VulkanRHI/Private/VulkanLayers.cpp
+++ b/Engine/Source/Runtime/VulkanRHI/Private/VulkanLayers.cpp
@@ -87,6 +87,12 @@ public:
 			{
 				OutLayerProperties[i].LayerProperties = TempLayerProperties[i];
 				OutLayerProperties[i].ExtensionProperties = FVulkanInstanceExtension::GetDriverSupportedInstanceExtensions(TempLayerProperties[i].layerName);
+
+				UE_LOG(LogVulkanRHI, Display, TEXT("Intance Layer %s provides %d extensions:"), ANSI_TO_TCHAR(TempLayerProperties[i].layerName), (int)OutLayerProperties[i].ExtensionProperties.Num());
+				for (VkExtensionProperties ep : OutLayerProperties[i].ExtensionProperties)
+				{
+					UE_LOG(LogVulkanRHI, Display, TEXT("  %s"), ANSI_TO_TCHAR(ep.extensionName));
+				}
 			}
 		}
 		return OutLayerProperties;
@@ -127,6 +133,12 @@ public:
 			{
 				OutLayerProperties[i].LayerProperties = TempLayerProperties[i];
 				OutLayerProperties[i].ExtensionProperties = FVulkanDeviceExtension::GetDriverSupportedDeviceExtensions(Gpu, TempLayerProperties[i].layerName);
+
+				UE_LOG(LogVulkanRHI, Display, TEXT("Device Layer %s provides %d extensions:"), ANSI_TO_TCHAR(TempLayerProperties[i].layerName), (int)OutLayerProperties[i].ExtensionProperties.Num());
+				for (VkExtensionProperties ep : OutLayerProperties[i].ExtensionProperties)
+				{
+					UE_LOG(LogVulkanRHI, Display, TEXT("  %s"), ANSI_TO_TCHAR(ep.extensionName));
+				}
 			}
 		}
 		return OutLayerProperties;
@@ -144,6 +156,8 @@ public:
 	const TCHAR* HelperTypeName = TEXT("device");
 	static TArray<const ANSICHAR*> ExternalLayers;
 	VkPhysicalDevice Gpu = VK_NULL_HANDLE;
+	bool bSupportsDebugUtilsExt = false;
+	bool bSupportsDebugMarkerExt = false;
 };
 
 TArray<const ANSICHAR*> FVulkanDeviceSetupHelper::ExternalLayers;
@@ -300,11 +314,16 @@ TArray<const ANSICHAR*> FVulkanDynamicRHI::SetupInstanceLayers(FVulkanInstanceEx
 	return OutInstanceLayers;
 }
 
-TArray<const ANSICHAR*> FVulkanDevice::SetupDeviceLayers(VkPhysicalDevice Gpu, FVulkanDeviceExtensionArray& UEExtensions)
+TArray<const ANSICHAR*> FVulkanDevice::SetupDeviceLayers(FVulkanDeviceExtensionArray& UEExtensions)
 {
 	FVulkanDeviceSetupHelper DeviceHelper;
 	DeviceHelper.Gpu = Gpu;
-	return SetupLayers(DeviceHelper, UEExtensions);
+	DeviceHelper.bSupportsDebugUtilsExt = RHI->SupportsDebugUtilsExt();
+	TArray<const ANSICHAR*> OutDeviceLayers = SetupLayers(DeviceHelper, UEExtensions);
+#if VULKAN_ENABLE_DRAW_MARKERS
+	bUseLegacyDebugMarkerExt = !DeviceHelper.bSupportsDebugUtilsExt && DeviceHelper.bSupportsDebugMarkerExt;
+#endif
+	return OutDeviceLayers;
 }
 
 void FVulkanDynamicRHI::SetupValidationRequests()
@@ -423,18 +442,6 @@ void FVulkanIntanceSetupHelper::AddDebugLayers(const TArray<FLayerWithExtensions
 			const int32 ExtensionIndex = FVulkanInstanceExtension::FindExtension(UEExtensions, ExtensionName);
 			check(ExtensionIndex != INDEX_NONE);
 
-			// If the extension isn't supported out of the box, check to see if an extension can add support
-			if (!UEExtensions[ExtensionIndex]->IsSupported())
-			{
-				const int32 LayerIndex = FindLayerContainingExtension(ExtensionName, LayerProperties);
-				if (LayerIndex != INDEX_NONE)
-				{
-					static VkLayerProperties StaticLayer;
-					StaticLayer = LayerProperties[LayerIndex].LayerProperties;  // TODO: prevent from pointing to temporary memory
-					AddRequestedLayer(StaticLayer.layerName, LayerProperties, UEExtensions, OutLayers);
-				}
-			}
-
 			// If the extension is supported, activate it and set it as our active debugging extension
 			if (UEExtensions[ExtensionIndex]->IsSupported())
 			{
@@ -470,4 +477,20 @@ void FVulkanDeviceSetupHelper::AddDebugLayers(const TArray<FLayerWithExtensions>
 	GRenderDocFound = false;
 #endif // VULKAN_ENABLE_DRAW_MARKERS
 #endif // VULKAN_HAS_DEBUGGING_ENABLED
+
+#if VULKAN_ENABLE_DRAW_MARKERS
+	// If debug_utils is supported we don't need the legacy extension debug_marker
+	if (!bSupportsDebugUtilsExt)
+	{
+		const int32 ExtensionIndex = FVulkanDeviceExtension::FindExtension(UEExtensions, VK_EXT_DEBUG_MARKER_EXTENSION_NAME);
+		check(ExtensionIndex != INDEX_NONE);
+
+		// If the extension is supported, activate it and set it as our active debugging extension
+		if (UEExtensions[ExtensionIndex]->IsSupported())
+		{
+			UEExtensions[ExtensionIndex]->SetActivated();
+			bSupportsDebugMarkerExt = true;
+		}
+	}
+#endif // VULKAN_ENABLE_DRAW_MARKERS
 }
