VulkanRHI: Clean up: Remove unused code

#rnx

diff --git a/Engine/Source/Runtime/VulkanRHI/Private/VulkanCommandBuffer.cpp b/Engine/Source/Runtime/VulkanRHI/Private/VulkanCommandBuffer.cpp
--- a/Engine/Source/Runtime/VulkanRHI/Private/VulkanCommandBuffer.cpp
+++ b/Engine/Source/Runtime/VulkanRHI/Private/VulkanCommandBuffer.cpp
@@ -50,7 +50,6 @@ FVulkanCommandBuffer::FVulkanCommandBuffer(FVulkanDevice& InDevice, FVulkanComma
 	, bHasScissor(0)
 	, bHasStencilRef(0)
 	, bIsUploadOnly(bInIsUploadOnly ? 1 : 0)
-	, bIsUniformBufferBarrierAdded(0)
 	, Device(InDevice)
 	, CommandBufferHandle(VK_NULL_HANDLE)
 	, Fence(nullptr)
@@ -139,31 +138,6 @@ void FVulkanCommandBuffer::FreeMemory()
 	State = EState::NotAllocated;
 }
 
-void FVulkanCommandBuffer::BeginUniformUpdateBarrier()
-{
-	if(!bIsUniformBufferBarrierAdded)
-	{
-		VkMemoryBarrier Barrier;
-		ZeroVulkanStruct(Barrier, VK_STRUCTURE_TYPE_MEMORY_BARRIER);
-		Barrier.srcAccessMask = VK_ACCESS_UNIFORM_READ_BIT;
-		Barrier.dstAccessMask = VK_ACCESS_TRANSFER_WRITE_BIT;
-		VulkanRHI::vkCmdPipelineBarrier(GetHandle(), VK_PIPELINE_STAGE_ALL_COMMANDS_BIT, VK_PIPELINE_STAGE_ALL_COMMANDS_BIT, 0, 1, &Barrier, 0, nullptr, 0, nullptr);
-		bIsUniformBufferBarrierAdded = true;
-	}
-}
-
-void FVulkanCommandBuffer::EndUniformUpdateBarrier()
-{
-	if(bIsUniformBufferBarrierAdded)
-	{
-		VkMemoryBarrier Barrier;
-		ZeroVulkanStruct(Barrier, VK_STRUCTURE_TYPE_MEMORY_BARRIER);
-		Barrier.srcAccessMask = VK_ACCESS_TRANSFER_WRITE_BIT;
-		Barrier.dstAccessMask = VK_ACCESS_UNIFORM_READ_BIT;
-		VulkanRHI::vkCmdPipelineBarrier(GetHandle(), VK_PIPELINE_STAGE_ALL_COMMANDS_BIT, VK_PIPELINE_STAGE_ALL_COMMANDS_BIT, 0, 1, &Barrier, 0, nullptr, 0, nullptr);
-		bIsUniformBufferBarrierAdded = false;
-	}
-}
 void FVulkanCommandBuffer::EndRenderPass()
 {
 	checkf(IsInsideRenderPass(), TEXT("Can't EndRP as we're NOT inside one! CmdBuffer 0x%p State=%d"), CommandBufferHandle, (int32)State);
@@ -173,10 +147,6 @@ void FVulkanCommandBuffer::EndRenderPass()
 
 void FVulkanCommandBuffer::BeginRenderPass(const FVulkanRenderTargetLayout& Layout, FVulkanRenderPass* RenderPass, FVulkanFramebuffer* Framebuffer, const VkClearValue* AttachmentClearValues)
 {
-	if (bIsUniformBufferBarrierAdded)
-	{
-		EndUniformUpdateBarrier();
-	}
 	checkf(IsOutsideRenderPass(), TEXT("Can't BeginRP as already inside one! CmdBuffer 0x%p State=%d"), CommandBufferHandle, (int32)State);
 
 	VkRenderPassBeginInfo Info;
diff --git a/Engine/Source/Runtime/VulkanRHI/Private/VulkanCommandBuffer.h b/Engine/Source/Runtime/VulkanRHI/Private/VulkanCommandBuffer.h
--- a/Engine/Source/Runtime/VulkanRHI/Private/VulkanCommandBuffer.h
+++ b/Engine/Source/Runtime/VulkanRHI/Private/VulkanCommandBuffer.h
@@ -42,11 +42,6 @@ public:
 		return CommandBufferPool;
 	}
 
-	bool IsUniformBufferBarrierAdded() const
-	{
-		return bIsUniformBufferBarrierAdded;
-	}
-
 	inline bool IsInsideRenderPass() const
 	{
 		return State == EState::IsInsideRenderPass;
@@ -181,15 +176,11 @@ public:
 	uint8 bHasScissor					: 1;
 	uint8 bHasStencilRef				: 1;
 	uint8 bIsUploadOnly					: 1;
-	uint8 bIsUniformBufferBarrierAdded	: 1;
 
 	// You never want to call Begin/EndRenderPass directly as it will mess up the layout manager.
 	void BeginRenderPass(const FVulkanRenderTargetLayout& Layout, class FVulkanRenderPass* RenderPass, class FVulkanFramebuffer* Framebuffer, const VkClearValue* AttachmentClearValues);
 	void EndRenderPass();
 
-
-	void BeginUniformUpdateBarrier();
-	void EndUniformUpdateBarrier();
 	//#todo-rco: Hide this
 	FVulkanDescriptorPoolSetContainer* CurrentDescriptorPoolSetContainer = nullptr;
 
