VulkanRHI: Remove some calls to SubmitActiveCmdBuffer from within the RHI (it's not aware of UploadContext and might submit command buffers out of order).  Remove r.Vulkan.SubmitAfterEveryEndRenderPass (and Dispatch and RayTrace), can be accomplished through RDG with r.RDG.Debug.FlushGPU.

#rnx

diff --git a/Engine/Source/Runtime/VulkanRHI/Private/VulkanCommands.cpp b/Engine/Source/Runtime/VulkanRHI/Private/VulkanCommands.cpp
--- a/Engine/Source/Runtime/VulkanRHI/Private/VulkanCommands.cpp
+++ b/Engine/Source/Runtime/VulkanRHI/Private/VulkanCommands.cpp
@@ -17,23 +17,6 @@
 #include "RHICoreShader.h"
 #include "RHIUniformBufferUtilities.h"
 
-static TAutoConsoleVariable<int32> GCVarSubmitOnDispatch(
-	TEXT("r.Vulkan.SubmitOnDispatch"),
-	0,
-	TEXT("0 to not do anything special on dispatch(default)\n")\
-	TEXT("1 to submit the cmd buffer after each dispatch"),
-	ECVF_RenderThreadSafe
-);
-
-int32 GVulkanSubmitAfterEveryEndRenderPass = 0;
-static FAutoConsoleVariableRef CVarVulkanSubmitAfterEveryEndRenderPass(
-	TEXT("r.Vulkan.SubmitAfterEveryEndRenderPass"),
-	GVulkanSubmitAfterEveryEndRenderPass,
-	TEXT("Forces a submit after every end render pass.\n")
-	TEXT(" 0: Don't(default)\n")
-	TEXT(" 1: Enable submitting"),
-	ECVF_Default
-);
 
 // make sure what the hardware expects matches what we give it for indirect arguments
 static_assert(sizeof(FRHIDrawIndirectParameters) == sizeof(VkDrawIndirectCommand), "FRHIDrawIndirectParameters size is wrong.");
@@ -306,11 +289,6 @@ void FVulkanCommandListContext::RHIDispatchComputeShader(uint32 ThreadGroupCount
 
 	VulkanRHI::vkCmdDispatch(CommandBufferHandle, ThreadGroupCountX, ThreadGroupCountY, ThreadGroupCountZ);
 
-	if (GCVarSubmitOnDispatch.GetValueOnRenderThread())
-	{
-		InternalSubmitActiveCmdBuffer();
-	}
-
 	if (FVulkanPlatform::RegisterGPUWork() && IsImmediate())
 	{
 		RegisterGPUDispatch(FIntVector(ThreadGroupCountX, ThreadGroupCountY, ThreadGroupCountZ));	
@@ -333,11 +311,6 @@ void FVulkanCommandListContext::RHIDispatchIndirectComputeShader(FRHIBuffer* Arg
 
 	VulkanRHI::vkCmdDispatchIndirect(CommandBufferHandle, ArgumentBuffer->GetHandle(), ArgumentBuffer->GetOffset() + ArgumentOffset);
 
-	if (GCVarSubmitOnDispatch.GetValueOnRenderThread())
-	{
-		InternalSubmitActiveCmdBuffer();
-	}
-
 	if (FVulkanPlatform::RegisterGPUWork()/* && IsImmediate()*/)
 	{
 		RegisterGPUDispatch(FIntVector(1, 1, 1));	
diff --git a/Engine/Source/Runtime/VulkanRHI/Private/VulkanPipelineState.cpp b/Engine/Source/Runtime/VulkanRHI/Private/VulkanPipelineState.cpp
--- a/Engine/Source/Runtime/VulkanRHI/Private/VulkanPipelineState.cpp
+++ b/Engine/Source/Runtime/VulkanRHI/Private/VulkanPipelineState.cpp
@@ -577,16 +577,6 @@ void FVulkanCommandListContext::RHISetGraphicsPipelineState(FRHIGraphicsPipeline
 void FVulkanCommandListContext::RHISetComputePipelineState(FRHIComputePipelineState* ComputePipelineState)
 {
 	FVulkanCommandBuffer* CmdBuffer = GetActiveCmdBuffer();
-	if (CmdBuffer->IsInsideRenderPass())
-	{
-		if (GVulkanSubmitAfterEveryEndRenderPass)
-		{
-			SubmitActiveCmdBuffer();
-			PrepareForNewActiveCommandBuffer();
-			CmdBuffer = GetActiveCmdBuffer();
-		}
-	}
-
 	if (CmdBuffer->CurrentDescriptorPoolSetContainer == nullptr)
 	{
 		CmdBuffer->CurrentDescriptorPoolSetContainer = &Device.GetDescriptorPoolsManager().AcquirePoolSetContainer();
diff --git a/Engine/Source/Runtime/VulkanRHI/Private/VulkanQuery.cpp b/Engine/Source/Runtime/VulkanRHI/Private/VulkanQuery.cpp
--- a/Engine/Source/Runtime/VulkanRHI/Private/VulkanQuery.cpp
+++ b/Engine/Source/Runtime/VulkanRHI/Private/VulkanQuery.cpp
@@ -534,9 +534,6 @@ bool FVulkanDynamicRHI::RHIGetRenderQueryResult(FRHIRenderQuery* QueryRHI, uint6
 					if (bWaitForStart)
 					{
 						RHICmdList.ImmediateFlush(EImmediateFlushType::FlushRHIThread);
-
-						// Need to submit the open command lists.
-						Device->SubmitCommandsAndFlushGPU();
 					}
 
 					// CPU wait for query results to be ready.
diff --git a/Engine/Source/Runtime/VulkanRHI/Private/VulkanRHIPrivate.h b/Engine/Source/Runtime/VulkanRHI/Private/VulkanRHIPrivate.h
--- a/Engine/Source/Runtime/VulkanRHI/Private/VulkanRHIPrivate.h
+++ b/Engine/Source/Runtime/VulkanRHI/Private/VulkanRHIPrivate.h
@@ -906,7 +906,6 @@ inline bool ValidateShadingRateDataType()
 	return false;
 }
 
-extern int32 GVulkanSubmitAfterEveryEndRenderPass;
 extern int32 GWaitForIdleOnSubmit;
 
 #if VULKAN_HAS_DEBUGGING_ENABLED
diff --git a/Engine/Source/Runtime/VulkanRHI/Private/VulkanRayTracing.cpp b/Engine/Source/Runtime/VulkanRHI/Private/VulkanRayTracing.cpp
--- a/Engine/Source/Runtime/VulkanRHI/Private/VulkanRayTracing.cpp
+++ b/Engine/Source/Runtime/VulkanRHI/Private/VulkanRayTracing.cpp
@@ -37,15 +37,6 @@ static FAutoConsoleVariableRef CVarVulkanRayTracingAllowDeferredOperation(
 	ECVF_ReadOnly
 );
 
-static int32 GVulkanSubmitOnTraceRays = 0;
-static FAutoConsoleVariableRef GCVarSubmitOnTraceRays(
-	TEXT("r.Vulkan.SubmitOnTraceRays"),
-	GVulkanSubmitOnTraceRays,
-	TEXT("0 to not do anything special on trace rays (default)\n")\
-	TEXT("1 to submit the cmd buffer after each trace rays"),
-	ECVF_ReadOnly
-);
-
 static int32 GVulkanRayTracingMaxShaderGroupStride = 4096;
 static FAutoConsoleVariableRef GCVarVulkanRayTracingMaxShaderGroupStride(
 	TEXT("r.Vulkan.RayTracing.MaxShaderGroupStride"),
@@ -853,9 +844,6 @@ void BuildAccelerationStructure(
 	// Acceleration structure build barrier is used here to ensure that the acceleration structure build is complete before any rays are traced
 	AddAccelerationStructureBuildBarrier(CmdBuffer->GetHandle());
 
-	CommandContext.SubmitActiveCmdBuffer();
-	CommandContext.PrepareForNewActiveCommandBuffer();
-
 	Scene.bBuilt = true;
 }
 
@@ -1330,9 +1318,6 @@ void FVulkanCommandListContext::RHIBuildAccelerationStructures(const TArrayView<
 	// are complete before BLAS is used again on the GPU.
 	AddAccelerationStructureBuildBarrier(CmdBuffer->GetHandle());
 
-	SubmitActiveCmdBuffer();
-	PrepareForNewActiveCommandBuffer();
-
 	for (const FRayTracingGeometryBuildParams& P : Params)
 	{
 		FVulkanRayTracingGeometry* const Geometry = ResourceCast(P.Geometry.GetReference());
@@ -1929,9 +1914,6 @@ void FVulkanRayTracingCompactionRequestHandler::Update(FVulkanCommandListContext
 		);
 
 		QueryPool->EndBatch(CmdBuffer);
-
-		InCommandContext.SubmitActiveCmdBuffer();
-		InCommandContext.PrepareForNewActiveCommandBuffer();
 	}
 }
 
@@ -2165,11 +2147,6 @@ void FVulkanCommandListContext::RHIRayTraceDispatch(
 		InWidth, InHeight, 1);
 
 	PostDispatch.Execute(CmdBuffer);
-
-	if (GVulkanSubmitOnTraceRays)
-	{
-		InternalSubmitActiveCmdBuffer();
-	}
 }
 
 void FVulkanCommandListContext::RHIRayTraceDispatchIndirect(
@@ -2205,11 +2182,6 @@ void FVulkanCommandListContext::RHIRayTraceDispatchIndirect(
 		IndirectDeviceAddress);
 
 	PostDispatch.Execute(CmdBuffer);
-
-	if (GVulkanSubmitOnTraceRays)
-	{
-		InternalSubmitActiveCmdBuffer();
-	}
 }
 
 
diff --git a/Engine/Source/Runtime/VulkanRHI/Private/VulkanRenderTarget.cpp b/Engine/Source/Runtime/VulkanRHI/Private/VulkanRenderTarget.cpp
--- a/Engine/Source/Runtime/VulkanRHI/Private/VulkanRenderTarget.cpp
+++ b/Engine/Source/Runtime/VulkanRHI/Private/VulkanRenderTarget.cpp
@@ -267,9 +267,7 @@ void FVulkanDynamicRHI::RHIMapStagingSurface(FRHITexture* TextureRHI, FRHIGPUFen
 	if (FenceRHI && !FenceRHI->Poll())
 	{
 		FRHICommandListImmediate& RHICmdList = FRHICommandListImmediate::Get();
-
 		RHICmdList.ImmediateFlush(EImmediateFlushType::FlushRHIThread);
-		Device->SubmitCommandsAndFlushGPU();
 
 		// SubmitCommandsAndFlushGPU might update fence state if it was tied to a previously submitted command buffer.
 		// Its state will have been updated from Submitted to NeedReset, and would assert in WaitForCmdBuffer (which is not needed in such a case)
@@ -557,19 +555,9 @@ bool FVulkanCommandListContext::IsSwapchainImage(FRHITexture* InTexture) const
 
 void FVulkanCommandListContext::RHIBeginRenderPass(const FRHIRenderPassInfo& InInfo, const TCHAR* InName)
 {
+	SafePointSubmit();
 	FVulkanCommandBuffer* CmdBuffer = GetActiveCmdBuffer();
 
-	if (GVulkanSubmitAfterEveryEndRenderPass)
-	{
-		SubmitActiveCmdBuffer();
-		PrepareForNewActiveCommandBuffer();
-		CmdBuffer = GetActiveCmdBuffer();
-	}
-	else if (SafePointSubmit())
-	{
-		CmdBuffer = GetActiveCmdBuffer();
-	}
-
 	RenderPassInfo = InInfo;
 
 	if (InInfo.NumOcclusionQueries > 0)
diff --git a/Engine/Source/Runtime/VulkanRHI/Private/VulkanTransientResourceAllocator.cpp b/Engine/Source/Runtime/VulkanRHI/Private/VulkanTransientResourceAllocator.cpp
--- a/Engine/Source/Runtime/VulkanRHI/Private/VulkanTransientResourceAllocator.cpp
+++ b/Engine/Source/Runtime/VulkanRHI/Private/VulkanTransientResourceAllocator.cpp
@@ -98,7 +98,7 @@ FRHITransientTexture* FVulkanTransientResourceAllocator::CreateTexture(const FRH
 	FDynamicRHI::FRHICalcTextureSizeResult MemReq = GDynamicRHI->RHICalcTexturePlatformSize(InCreateInfo, 0);
 
 	return CreateTextureInternal(InCreateInfo, InDebugName, Fences, MemReq.Size, MemReq.Align,
-		[&](const FRHITransientHeap::FResourceInitializer& Initializer)
+		[Device=Device, InCreateInfo, InDebugName, MemReq](const FRHITransientHeap::FResourceInitializer& Initializer)
 	{
 		FRHITextureCreateDesc CreateDesc(InCreateInfo, ERHIAccess::Discard, InDebugName);
 		FRHITexture* Texture = new FVulkanTexture(*Device, CreateDesc, &Initializer.Allocation);
@@ -116,7 +116,7 @@ FRHITransientBuffer* FVulkanTransientResourceAllocator::CreateBuffer(const FRHIB
 	uint64 Size = Align(InCreateInfo.Size, Alignment);
 
 	return CreateBufferInternal(InCreateInfo, InDebugName, Fences, Size, Alignment,
-		[&](const FRHITransientHeap::FResourceInitializer& Initializer)
+		[Device=Device, InCreateInfo, InDebugName, Size](const FRHITransientHeap::FResourceInitializer& Initializer)
 	{
 		FRHIResourceCreateInfo ResourceCreateInfo(InDebugName);
 		FRHIBuffer* Buffer = new FVulkanBuffer(Device, FRHIBufferDesc(InCreateInfo.Size, InCreateInfo.Stride, InCreateInfo.Usage), ResourceCreateInfo, nullptr, &Initializer.Allocation);
