Fix Vulkan GPU timestamp query

diff --git a/Engine/Source/Runtime/VulkanRHI/Private/VulkanQuery.cpp b/Engine/Source/Runtime/VulkanRHI/Private/VulkanQuery.cpp
--- a/Engine/Source/Runtime/VulkanRHI/Private/VulkanQuery.cpp
+++ b/Engine/Source/Runtime/VulkanRHI/Private/VulkanQuery.cpp
@@ -404,10 +404,13 @@ void FVulkanCommandListContext::ReadAndCalculateGPUFrameTime()
 		{
 			Delta = FrameTiming->GetTiming(false);
 		}
-	
-		const double SecondsPerCycle = FPlatformTime::GetSecondsPerCycle();
-		const double Frequency = double(FVulkanGPUTiming::GetTimingFrequency());
-		GGPUFrameTime = FMath::TruncToInt(double(Delta) / Frequency / SecondsPerCycle);
+
+		if (Delta > 0)
+		{
+			const double SecondsPerCycle = FPlatformTime::GetSecondsPerCycle();
+			const double Frequency = double(FVulkanGPUTiming::GetTimingFrequency());
+			GGPUFrameTime = FMath::TruncToInt(double(Delta) / Frequency / SecondsPerCycle);
+		}
 	}
 	else if(!FVulkanPlatform::HasCustomFrameTiming())
 	{
diff --git a/Engine/Source/Runtime/VulkanRHI/Private/VulkanUtil.cpp b/Engine/Source/Runtime/VulkanRHI/Private/VulkanUtil.cpp
--- a/Engine/Source/Runtime/VulkanRHI/Private/VulkanUtil.cpp
+++ b/Engine/Source/Runtime/VulkanRHI/Private/VulkanUtil.cpp
@@ -191,16 +191,20 @@ void FVulkanGPUTiming::DiscardExpiredQueries(bool bNeedSpace)
 	auto CanBeDiscarded = [this](uint32 QueryIndex)
 	{
 		FVulkanTimingQueryPool::FCmdBufferFence& QuerySyncPoint = Pool->TimestampListHandles[QueryIndex];
-		if (QuerySyncPoint.CmdBuffer && (QuerySyncPoint.FenceCounter < QuerySyncPoint.CmdBuffer->GetFenceSignaledCounter()))
+
+		if (!QuerySyncPoint.CmdBuffer)
+		{
+			return true;
+		}
+		if (QuerySyncPoint.FenceCounter < QuerySyncPoint.CmdBuffer->GetFenceSignaledCounter())
 		{
 			const uint64 Availability = Pool->MappedPointer[QueryIndex * 2 + 1];
-
-			if (Availability || (QuerySyncPoint.Attempts < TIMING_QUERY_RETRIES))
+			if (!Availability && (QuerySyncPoint.Attempts >= TIMING_QUERY_RETRIES))
 			{
-				return false;
+				return true;
 			}
 		}
-		return true;
+		return false;
 	};
 
 	auto IsQueryExpired = [&](uint32 QueryIndex)
