VulkanRHI: replace check for unfreed staging resources with warning in logs

#rnx

diff --git a/Engine/Source/Runtime/VulkanRHI/Private/VulkanMemory.cpp b/Engine/Source/Runtime/VulkanRHI/Private/VulkanMemory.cpp
--- a/Engine/Source/Runtime/VulkanRHI/Private/VulkanMemory.cpp
+++ b/Engine/Source/Runtime/VulkanRHI/Private/VulkanMemory.cpp
@@ -4299,18 +4299,18 @@ namespace VulkanRHI
 
 	FStagingManager::~FStagingManager()
 	{
-		check(UsedStagingBuffers.Num() == 0);
-		check(PendingFreeStagingBuffers.Num() == 0);
-		check(FreeStagingBuffers.Num() == 0);
 	}
 
 	void FStagingManager::Deinit()
 	{
 		ProcessPendingFree(true, true);
 
-		check(UsedStagingBuffers.Num() == 0);
-		check(PendingFreeStagingBuffers.Num() == 0);
-		check(FreeStagingBuffers.Num() == 0);
+		if ((UsedStagingBuffers.Num() != 0) || (PendingFreeStagingBuffers.Num() != 0) || (FreeStagingBuffers.Num() != 0))
+		{
+			UE_LOG(LogVulkanRHI, Warning,
+				TEXT("Some resources in the FStagingManager were not freed!  (UsedBuffer=%d, PendingFree=%d, FreeBuffer=%d)"),
+				UsedStagingBuffers.Num(), PendingFreeStagingBuffers.Num(), FreeStagingBuffers.Num());
+		}
 	}
 
 	FStagingBuffer* FStagingManager::AcquireBuffer(uint32 Size, VkBufferUsageFlags InUsageFlags, VkMemoryPropertyFlagBits InMemoryReadFlags)
