VulkanRHI: Remove more instances of the RHI submitting internally (causes command buffers to get inserted before the upload context).

#rnx

diff --git a/Engine/Source/Runtime/VulkanRHI/Private/VulkanCommands.cpp b/Engine/Source/Runtime/VulkanRHI/Private/VulkanCommands.cpp
--- a/Engine/Source/Runtime/VulkanRHI/Private/VulkanCommands.cpp
+++ b/Engine/Source/Runtime/VulkanRHI/Private/VulkanCommands.cpp
@@ -858,7 +858,6 @@ bool FVulkanDynamicRHI::RHIIsRenderingSuspended()
 
 void FVulkanDynamicRHI::RHIBlockUntilGPUIdle()
 {
-	Device->SubmitCommandsAndFlushGPU();
 	Device->WaitUntilIdle();
 }
 
@@ -980,30 +979,42 @@ void FVulkanDynamicRHI::RHIFinalizeContext(FRHIFinalizeContextArgs&& Args, TRHIP
 	{
 		FVulkanPlatformCommandList* PlatformCmdList = new FVulkanPlatformCommandList();
 		PlatformCmdList->CmdContext = static_cast<FVulkanCommandListContext*>(Context);
+		Output[Context->GetPipeline()] = PlatformCmdList;
+	}
 
-		// :todo-jn: Temporarily just place it before the graphic pipeline
-		if (Args.UploadContext && (Context->GetPipeline() == ERHIPipeline::Graphics))
+	// :todo-jn: Temporarily just place it before the graphic pipeline
+	if (Args.UploadContext)
+	{
+		if (Output[ERHIPipeline::Graphics] == nullptr)
 		{
-			PlatformCmdList->UploadContext = static_cast<FVulkanUploadContext*>(Args.UploadContext);
+			Output[ERHIPipeline::Graphics] = new FVulkanPlatformCommandList();
 		}
-
-		Output[Context->GetPipeline()] = PlatformCmdList;
+		FVulkanUploadContext* UploadContext = static_cast<FVulkanUploadContext*>(Args.UploadContext);
+		ResourceCast(Output[ERHIPipeline::Graphics])->UploadContext = UploadContext;
+		Args.UploadContext = nullptr;
 	}
 }
 
 void FVulkanDynamicRHI::RHISubmitCommandLists(FRHISubmitCommandListsArgs&& Args)
 {
+	check(Args.UploadContext == nullptr);
+
+	// Submit all the upload context up front
 	for (IRHIPlatformCommandList* Ptr : Args.CommandLists)
 	{
 		FVulkanPlatformCommandList* PlatformCmdList = ResourceCast(Ptr);
-
-		FVulkanUploadContext* UploadContext = PlatformCmdList->UploadContext;
-		if (UploadContext)
+		if (FVulkanUploadContext* UploadContext = PlatformCmdList->UploadContext)
 		{
 			UploadContext->SubmitActiveCmdBuffer();
 			UploadContext->RefreshFenceStatus();
-			FVulkanUploadContext::Pool.Push(PlatformCmdList->UploadContext);
+			FVulkanUploadContext::Pool.Push(UploadContext);
+			PlatformCmdList->UploadContext = nullptr;
 		}
+	}
+
+	for (IRHIPlatformCommandList* Ptr : Args.CommandLists)
+	{
+		FVulkanPlatformCommandList* PlatformCmdList = ResourceCast(Ptr);
 
 		FVulkanCommandListContext* CurrentContext = PlatformCmdList->CmdContext;
 		if (CurrentContext->IsImmediate())
diff --git a/Engine/Source/Runtime/VulkanRHI/Private/VulkanRenderTarget.cpp b/Engine/Source/Runtime/VulkanRHI/Private/VulkanRenderTarget.cpp
--- a/Engine/Source/Runtime/VulkanRHI/Private/VulkanRenderTarget.cpp
+++ b/Engine/Source/Runtime/VulkanRHI/Private/VulkanRenderTarget.cpp
@@ -35,13 +35,6 @@ static FAutoConsoleVariableRef CVarVulkanIgnoreCPUReads(
 	ECVF_Default
 	);
 
-TAutoConsoleVariable<int32> GSubmitOcclusionBatchCmdBufferCVar(
-	TEXT("r.Vulkan.SubmitOcclusionBatchCmdBuffer"),
-	1,
-	TEXT("1 to submit the cmd buffer after end occlusion query batch (default)"),
-	ECVF_RenderThreadSafe
-);
-
 static FCriticalSection GStagingMapLock;
 static TMap<FVulkanTexture*, VulkanRHI::FStagingBuffer*> GPendingLockedStagingBuffers;
 
@@ -671,13 +664,6 @@ void FVulkanCommandListContext::RHIEndRenderPass()
 
 	check(CurrentRenderPass);
 	CurrentRenderPass = nullptr;
-
-	// Sync point for passes with occlusion queries
-	if (bHasOcclusionQueries && GSubmitOcclusionBatchCmdBufferCVar.GetValueOnAnyThread())
-	{
-		RequestSubmitCurrentCommands();
-		SafePointSubmit();
-	}
 }
 
 void FVulkanCommandListContext::RHINextSubpass()
diff --git a/Engine/Source/Runtime/VulkanRHI/Private/VulkanUtil.cpp b/Engine/Source/Runtime/VulkanRHI/Private/VulkanUtil.cpp
--- a/Engine/Source/Runtime/VulkanRHI/Private/VulkanUtil.cpp
+++ b/Engine/Source/Runtime/VulkanRHI/Private/VulkanUtil.cpp
@@ -736,11 +736,12 @@ void* FVulkanDynamicRHI::RHILockStagingBuffer(FRHIStagingBuffer* StagingBufferRH
 
 	if (FenceRHI && !FenceRHI->Poll())
 	{
-		Device->SubmitCommandsAndFlushGPU();
+		FRHICommandListImmediate& RHICmdList = FRHICommandListImmediate::Get();
+		RHICmdList.ImmediateFlush(EImmediateFlushType::FlushRHIThread);
 
 		// SubmitCommandsAndFlushGPU might update fence state if it was tied to a previously submitted command buffer.
 		// Its state will have been updated from Submitted to NeedReset, and would assert in WaitForCmdBuffer (which is not needed in such a case)
-		FenceRHI->Wait(FRHICommandListImmediate::Get(), FRHIGPUMask::All());
+		FenceRHI->Wait(RHICmdList, FRHIGPUMask::All());
 	}
 
 	return StagingBuffer->Lock(Offset, NumBytes);
