VulkanRHI: Add timeline semaphore functions to VulkanLoader.

#rnx

diff --git a/Engine/Source/Runtime/VulkanRHI/Private/VulkanCommandWrappers.h b/Engine/Source/Runtime/VulkanRHI/Private/VulkanCommandWrappers.h
--- a/Engine/Source/Runtime/VulkanRHI/Private/VulkanCommandWrappers.h
+++ b/Engine/Source/Runtime/VulkanRHI/Private/VulkanCommandWrappers.h
@@ -258,7 +258,8 @@ struct FWrapLayer
 #if PLATFORM_ANDROID
 	static void GetAndroidHardwareBufferPropertiesANDROID(VkResult Result, VkDevice Device, const struct AHardwareBuffer* Buffer, VkAndroidHardwareBufferPropertiesANDROID* Properties) VULKAN_LAYER_BODY
 #endif
-
+	static void GetSemaphoreCounterValue(VkResult Result, VkDevice Device, VkSemaphore Semaphore, uint64_t* Value) VULKAN_LAYER_BODY
+	static void WaitSemaphores(VkResult Result, VkDevice Device, const VkSemaphoreWaitInfo* WaitInfo, uint64_t Timeout) VULKAN_LAYER_BODY
 };
 
 #undef VULKAN_LAYER_BODY
@@ -1789,6 +1790,20 @@ namespace VulkanRHI
 	}
 #endif // PLATFORM_ANDROID
 
+	static FORCEINLINE_DEBUGGABLE void vkGetSemaphoreCounterValue(VkDevice Device, VkSemaphore Semaphore, uint64_t* Value)
+	{
+		FWrapLayer::GetSemaphoreCounterValue(VK_RESULT_MAX_ENUM, Device, Semaphore, Value);
+		VkResult Result = VULKANAPINAMESPACE::vkGetSemaphoreCounterValue(Device, Semaphore, Value);
+		FWrapLayer::GetSemaphoreCounterValue(Result, Device, Semaphore, Value);
+	}
+
+	static FORCEINLINE_DEBUGGABLE void vkWaitSemaphores(VkDevice Device, const VkSemaphoreWaitInfo* WaitInfo, uint64_t Timeout)
+	{
+		FWrapLayer::WaitSemaphores(VK_RESULT_MAX_ENUM, Device, WaitInfo, Timeout);
+		VkResult Result = VULKANAPINAMESPACE::vkWaitSemaphores(Device, WaitInfo, Timeout);
+		FWrapLayer::WaitSemaphores(Result, Device, WaitInfo, Timeout);
+	}
+
 #if VULKAN_ENABLE_IMAGE_TRACKING_LAYER
 	void BindDebugLabelName(VkImage Image, const TCHAR* Name);
 #endif
diff --git a/Engine/Source/Runtime/VulkanRHI/Private/VulkanDebug.cpp b/Engine/Source/Runtime/VulkanRHI/Private/VulkanDebug.cpp
--- a/Engine/Source/Runtime/VulkanRHI/Private/VulkanDebug.cpp
+++ b/Engine/Source/Runtime/VulkanRHI/Private/VulkanDebug.cpp
@@ -4638,6 +4638,27 @@ void FWrapLayer::GetAndroidHardwareBufferPropertiesANDROID(VkResult Result, VkDe
 }
 #endif
 
+void FWrapLayer::GetSemaphoreCounterValue(VkResult Result, VkDevice Device, VkSemaphore Semaphore, uint64_t* Value)
+{
+	if (Result == VK_RESULT_MAX_ENUM)
+	{
+#if VULKAN_ENABLE_DUMP_LAYER
+		PrintfBeginResult(FString::Printf(TEXT("GetSemaphoreCounterValue(Device=0x%p, Semaphore=0x%p, Value=0x%p))"), Device, Semaphore, Value));
+#endif
+	}
+}
+
+void FWrapLayer::WaitSemaphores(VkResult Result, VkDevice Device, const VkSemaphoreWaitInfo* WaitInfo, uint64_t Timeout)
+{
+	if (Result == VK_RESULT_MAX_ENUM)
+	{
+#if VULKAN_ENABLE_DUMP_LAYER
+		PrintfBeginResult(FString::Printf(TEXT("WaitSemaphores(Device=0x%p, WaitInfo=0x%p, Timeout=0xllx))"), Device, WaitInfo, Timeout));
+#endif
+	}
+}
+
+
 #if VULKAN_ENABLE_IMAGE_TRACKING_LAYER
 namespace VulkanRHI
 {
diff --git a/Engine/Source/Runtime/VulkanRHI/Private/VulkanLoader.h b/Engine/Source/Runtime/VulkanRHI/Private/VulkanLoader.h
--- a/Engine/Source/Runtime/VulkanRHI/Private/VulkanLoader.h
+++ b/Engine/Source/Runtime/VulkanRHI/Private/VulkanLoader.h
@@ -235,7 +235,10 @@
 	EnumMacro(PFN_vkGetDeviceFaultInfoEXT, vkGetDeviceFaultInfoEXT) \
 	EnumMacro(PFN_vkGetPhysicalDeviceToolPropertiesEXT, vkGetPhysicalDeviceToolPropertiesEXT) \
 	EnumMacro(PFN_vkGetImageSubresourceLayout2EXT, vkGetImageSubresourceLayout2EXT) \
-	EnumMacro(PFN_vkGetImageSubresourceLayout2KHR, vkGetImageSubresourceLayout2KHR)
+	EnumMacro(PFN_vkGetImageSubresourceLayout2KHR, vkGetImageSubresourceLayout2KHR) \
+	EnumMacro(PFN_vkGetSemaphoreCounterValue, vkGetSemaphoreCounterValue) \
+	EnumMacro(PFN_vkWaitSemaphores, vkWaitSemaphores)
+
 
 
 // List of all Vulkan entry points
