VulkanRHI: use extension function names for timeline sempahores.  Add function to signal from host.

#rnx

diff --git a/Engine/Source/Runtime/VulkanRHI/Private/VulkanCommandWrappers.h b/Engine/Source/Runtime/VulkanRHI/Private/VulkanCommandWrappers.h
--- a/Engine/Source/Runtime/VulkanRHI/Private/VulkanCommandWrappers.h
+++ b/Engine/Source/Runtime/VulkanRHI/Private/VulkanCommandWrappers.h
@@ -258,8 +258,9 @@ struct FWrapLayer
 #if PLATFORM_ANDROID
 	static void GetAndroidHardwareBufferPropertiesANDROID(VkResult Result, VkDevice Device, const struct AHardwareBuffer* Buffer, VkAndroidHardwareBufferPropertiesANDROID* Properties) VULKAN_LAYER_BODY
 #endif
-	static void GetSemaphoreCounterValue(VkResult Result, VkDevice Device, VkSemaphore Semaphore, uint64_t* Value) VULKAN_LAYER_BODY
-	static void WaitSemaphores(VkResult Result, VkDevice Device, const VkSemaphoreWaitInfo* WaitInfo, uint64_t Timeout) VULKAN_LAYER_BODY
+	static void GetSemaphoreCounterValueKHR(VkResult Result, VkDevice Device, VkSemaphore Semaphore, uint64_t* Value) VULKAN_LAYER_BODY
+	static void WaitSemaphoresKHR(VkResult Result, VkDevice Device, const VkSemaphoreWaitInfo* WaitInfo, uint64_t Timeout) VULKAN_LAYER_BODY
+	static void SignalSemaphoreKHR(VkResult Result, VkDevice Device, const VkSemaphoreSignalInfo* SignalInfo) VULKAN_LAYER_BODY
 };
 
 #undef VULKAN_LAYER_BODY
@@ -1790,18 +1791,28 @@ namespace VulkanRHI
 	}
 #endif // PLATFORM_ANDROID
 
-	static FORCEINLINE_DEBUGGABLE void vkGetSemaphoreCounterValue(VkDevice Device, VkSemaphore Semaphore, uint64_t* Value)
+	static FORCEINLINE_DEBUGGABLE VkResult vkGetSemaphoreCounterValueKHR(VkDevice Device, VkSemaphore Semaphore, uint64_t* Value)
 	{
-		FWrapLayer::GetSemaphoreCounterValue(VK_RESULT_MAX_ENUM, Device, Semaphore, Value);
-		VkResult Result = VULKANAPINAMESPACE::vkGetSemaphoreCounterValue(Device, Semaphore, Value);
-		FWrapLayer::GetSemaphoreCounterValue(Result, Device, Semaphore, Value);
+		FWrapLayer::GetSemaphoreCounterValueKHR(VK_RESULT_MAX_ENUM, Device, Semaphore, Value);
+		const VkResult Result = VULKANAPINAMESPACE::vkGetSemaphoreCounterValueKHR(Device, Semaphore, Value);
+		FWrapLayer::GetSemaphoreCounterValueKHR(Result, Device, Semaphore, Value);
+		return Result;
+	}
+
+	static FORCEINLINE_DEBUGGABLE VkResult vkWaitSemaphoresKHR(VkDevice Device, const VkSemaphoreWaitInfo* WaitInfo, uint64_t Timeout)
+	{
+		FWrapLayer::WaitSemaphoresKHR(VK_RESULT_MAX_ENUM, Device, WaitInfo, Timeout);
+		const VkResult Result = VULKANAPINAMESPACE::vkWaitSemaphoresKHR(Device, WaitInfo, Timeout);
+		FWrapLayer::WaitSemaphoresKHR(Result, Device, WaitInfo, Timeout);
+		return Result;
 	}
 
-	static FORCEINLINE_DEBUGGABLE void vkWaitSemaphores(VkDevice Device, const VkSemaphoreWaitInfo* WaitInfo, uint64_t Timeout)
+	static FORCEINLINE_DEBUGGABLE VkResult vkSignalSemaphoreKHR(VkDevice Device, const VkSemaphoreSignalInfo* SignalInfo)
 	{
-		FWrapLayer::WaitSemaphores(VK_RESULT_MAX_ENUM, Device, WaitInfo, Timeout);
-		VkResult Result = VULKANAPINAMESPACE::vkWaitSemaphores(Device, WaitInfo, Timeout);
-		FWrapLayer::WaitSemaphores(Result, Device, WaitInfo, Timeout);
+		FWrapLayer::SignalSemaphoreKHR(VK_RESULT_MAX_ENUM, Device, SignalInfo);
+		const VkResult Result = VULKANAPINAMESPACE::vkSignalSemaphoreKHR(Device, SignalInfo);
+		FWrapLayer::SignalSemaphoreKHR(Result, Device, SignalInfo);
+		return Result;
 	}
 
 #if VULKAN_ENABLE_IMAGE_TRACKING_LAYER
diff --git a/Engine/Source/Runtime/VulkanRHI/Private/VulkanDebug.cpp b/Engine/Source/Runtime/VulkanRHI/Private/VulkanDebug.cpp
--- a/Engine/Source/Runtime/VulkanRHI/Private/VulkanDebug.cpp
+++ b/Engine/Source/Runtime/VulkanRHI/Private/VulkanDebug.cpp
@@ -4638,26 +4638,35 @@ void FWrapLayer::GetAndroidHardwareBufferPropertiesANDROID(VkResult Result, VkDe
 }
 #endif
 
-void FWrapLayer::GetSemaphoreCounterValue(VkResult Result, VkDevice Device, VkSemaphore Semaphore, uint64_t* Value)
+void FWrapLayer::GetSemaphoreCounterValueKHR(VkResult Result, VkDevice Device, VkSemaphore Semaphore, uint64_t* Value)
 {
 	if (Result == VK_RESULT_MAX_ENUM)
 	{
 #if VULKAN_ENABLE_DUMP_LAYER
-		PrintfBeginResult(FString::Printf(TEXT("GetSemaphoreCounterValue(Device=0x%p, Semaphore=0x%p, Value=0x%p))"), Device, Semaphore, Value));
+		PrintfBeginResult(FString::Printf(TEXT("GetSemaphoreCounterValueKHR(Device=0x%p, Semaphore=0x%p, Value=0x%p))"), Device, Semaphore, Value));
 #endif
 	}
 }
 
-void FWrapLayer::WaitSemaphores(VkResult Result, VkDevice Device, const VkSemaphoreWaitInfo* WaitInfo, uint64_t Timeout)
+void FWrapLayer::WaitSemaphoresKHR(VkResult Result, VkDevice Device, const VkSemaphoreWaitInfo* WaitInfo, uint64_t Timeout)
 {
 	if (Result == VK_RESULT_MAX_ENUM)
 	{
 #if VULKAN_ENABLE_DUMP_LAYER
-		PrintfBeginResult(FString::Printf(TEXT("WaitSemaphores(Device=0x%p, WaitInfo=0x%p, Timeout=0xllx))"), Device, WaitInfo, Timeout));
+		PrintfBeginResult(FString::Printf(TEXT("WaitSemaphoresKHR(Device=0x%p, WaitInfo=0x%p, Timeout=0xllx))"), Device, WaitInfo, Timeout));
 #endif
 	}
 }
 
+void FWrapLayer::SignalSemaphoreKHR(VkResult Result, VkDevice Device, const VkSemaphoreSignalInfo* SignalInfo)
+{
+	if (Result == VK_RESULT_MAX_ENUM)
+	{
+#if VULKAN_ENABLE_DUMP_LAYER
+		PrintfBeginResult(FString::Printf(TEXT("SignalSemaphoreKHR(Device=0x%p, SignalInfo=0x%p))"), Device, SignalInfo));
+#endif
+	}
+}
 
 #if VULKAN_ENABLE_IMAGE_TRACKING_LAYER
 namespace VulkanRHI
diff --git a/Engine/Source/Runtime/VulkanRHI/Private/VulkanLoader.h b/Engine/Source/Runtime/VulkanRHI/Private/VulkanLoader.h
--- a/Engine/Source/Runtime/VulkanRHI/Private/VulkanLoader.h
+++ b/Engine/Source/Runtime/VulkanRHI/Private/VulkanLoader.h
@@ -236,8 +236,10 @@
 	EnumMacro(PFN_vkGetPhysicalDeviceToolPropertiesEXT, vkGetPhysicalDeviceToolPropertiesEXT) \
 	EnumMacro(PFN_vkGetImageSubresourceLayout2EXT, vkGetImageSubresourceLayout2EXT) \
 	EnumMacro(PFN_vkGetImageSubresourceLayout2KHR, vkGetImageSubresourceLayout2KHR) \
-	EnumMacro(PFN_vkGetSemaphoreCounterValue, vkGetSemaphoreCounterValue) \
-	EnumMacro(PFN_vkWaitSemaphores, vkWaitSemaphores)
+	EnumMacro(PFN_vkGetSemaphoreCounterValue, vkGetSemaphoreCounterValueKHR) \
+	EnumMacro(PFN_vkWaitSemaphores, vkWaitSemaphoresKHR) \
+	EnumMacro(PFN_vkSignalSemaphore, vkSignalSemaphoreKHR)
+
 
 
 
