VulkanRHI: Accumulate command buffers and synchronisation in Payloads.  Handle all of submission at a single point in RHISubmitCommandLists (moved to VulkanSubmission.h/cpp).
Remove more submits called from inside the RHI.
Remove the option to UseCmdBufferTimingForGPUTime (new profiler coming)
Remove r.Vulkan.SubmitAfterEveryEndRenderPass/Dispatch, can be accomplished through RDG (r.RDG.Debug.FlushGPU).
Pull Semaphores and Fences out of the Command Buffer and handle them at the Payload level.
Add a streamlined path for submission when timeline semaphores are supported.
Remove all the different GetFenceSignaledCounter* functions once used for debugging
Remove bIsUploadOnly flag from CommandBuffers now that we use UploadContext.
Replace bools with Flags to configure VulkanRHI::FSemaphore

#rnx

diff --git a/Engine/Source/Runtime/VulkanRHI/Private/VulkanBarriers.cpp b/Engine/Source/Runtime/VulkanRHI/Private/VulkanBarriers.cpp
--- a/Engine/Source/Runtime/VulkanRHI/Private/VulkanBarriers.cpp
+++ b/Engine/Source/Runtime/VulkanRHI/Private/VulkanBarriers.cpp
@@ -1555,8 +1555,7 @@ void FVulkanCommandListContext::RHIBeginTransitions(TArrayView<const FRHITransit
 	TArray<VulkanRHI::FSemaphore*> SignalSemaphores = ExtractTransitionSemaphores(Transitions);
 	if (SignalSemaphores.Num() > 0)
 	{
-		SubmitActiveCmdBuffer(SignalSemaphores);
-		PrepareForNewActiveCommandBuffer();
+		AddSignalSemaphores(SignalSemaphores);
 	}
 }
 
@@ -1566,14 +1565,7 @@ void FVulkanCommandListContext::RHIEndTransitions(TArrayView<const FRHITransitio
 	TArray<VulkanRHI::FSemaphore*> WaitSemaphores = ExtractTransitionSemaphores(Transitions, CurrentPipeline);
 	if (WaitSemaphores.Num() > 0)
 	{
-		if (HasPendingActiveCmdBuffer())
-		{
-			SubmitActiveCmdBuffer();
-			PrepareForNewActiveCommandBuffer();
-		}
-
-		FVulkanCommandBuffer& CmdBuffer = GetCommandBuffer();
-		CmdBuffer.AddWaitSemaphore(VK_PIPELINE_STAGE_TOP_OF_PIPE_BIT, WaitSemaphores);
+		AddWaitSemaphores(VK_PIPELINE_STAGE_TOP_OF_PIPE_BIT, WaitSemaphores);
 	}
 
 	if (Device.SupportsParallelRendering())
diff --git a/Engine/Source/Runtime/VulkanRHI/Private/VulkanCommandBuffer.cpp b/Engine/Source/Runtime/VulkanRHI/Private/VulkanCommandBuffer.cpp
--- a/Engine/Source/Runtime/VulkanRHI/Private/VulkanCommandBuffer.cpp
+++ b/Engine/Source/Runtime/VulkanRHI/Private/VulkanCommandBuffer.cpp
@@ -20,28 +20,12 @@ static FAutoConsoleVariableRef CVarVulkanUseSingleQueue(
 	ECVF_Default
 );
 
-int32 GVulkanProfileCmdBuffers = 0;
-static FAutoConsoleVariableRef CVarVulkanProfileCmdBuffers(
-	TEXT("r.Vulkan.ProfileCmdBuffers"),
-	GVulkanProfileCmdBuffers,
-	TEXT("Insert GPU timing queries in every cmd buffer\n"),
-	ECVF_Default
-);
-
-int32 GVulkanUseCmdBufferTimingForGPUTime = 0;
-static FAutoConsoleVariableRef CVarVulkanUseCmdBufferTimingForGPUTime(
-	TEXT("r.Vulkan.UseCmdBufferTimingForGPUTime"),
-	GVulkanUseCmdBufferTimingForGPUTime,
-	TEXT("Use the profile command buffers for GPU time\n"),
-	ECVF_Default
-);
-
 #define CMD_BUFFER_TIME_TO_WAIT_BEFORE_DELETING		10
 
 const uint32 GNumberOfFramesBeforeDeletingDescriptorPool = 300;
 extern int32 GVulkanAutoCorrectUnknownLayouts;
 
-FVulkanCommandBuffer::FVulkanCommandBuffer(FVulkanDevice& InDevice, FVulkanCommandBufferPool* InCommandBufferPool, bool bInIsUploadOnly)
+FVulkanCommandBuffer::FVulkanCommandBuffer(FVulkanDevice& InDevice, FVulkanCommandBufferPool* InCommandBufferPool)
 	: CurrentStencilRef(0)
 	, State(EState::NotAllocated)
 	, bNeedsDynamicStateSet(1)
@@ -49,16 +33,9 @@ FVulkanCommandBuffer::FVulkanCommandBuffer(FVulkanDevice& InDevice, FVulkanComma
 	, bHasViewport(0)
 	, bHasScissor(0)
 	, bHasStencilRef(0)
-	, bIsUploadOnly(bInIsUploadOnly ? 1 : 0)
 	, Device(InDevice)
 	, CommandBufferHandle(VK_NULL_HANDLE)
-	, Fence(nullptr)
-	, FenceSignaledCounter(0)
-	, SubmittedFenceCounter(0)
 	, CommandBufferPool(InCommandBufferPool)
-#if (RHI_NEW_GPU_PROFILER == 0)
-	, Timing(nullptr)
-#endif
 	, LastValidTiming(0)
 	, LayoutManager(InDevice.SupportsParallelRendering() && !GVulkanAutoCorrectUnknownLayouts,
 		&InCommandBufferPool->GetContext().Queue.GetLayoutManager())
@@ -67,8 +44,6 @@ FVulkanCommandBuffer::FVulkanCommandBuffer(FVulkanDevice& InDevice, FVulkanComma
 		FScopeLock ScopeLock(CommandBufferPool->GetCS());
 		AllocMemory();
 	}
-
-	Fence = Device.GetFenceManager().AllocateFence();
 }
 
 void FVulkanCommandBuffer::AllocMemory()
@@ -98,32 +73,10 @@ void FVulkanCommandBuffer::AllocMemory()
 
 FVulkanCommandBuffer::~FVulkanCommandBuffer()
 {
-	VulkanRHI::FFenceManager& FenceManager = Device.GetFenceManager();
-	if (State == EState::Submitted)
-	{
-		// Wait 33ms
-		uint64 WaitForCmdBufferInNanoSeconds = 33 * 1000 * 1000LL;
-		FenceManager.WaitAndReleaseFence(Fence, WaitForCmdBufferInNanoSeconds);
-	}
-	else
-	{
-		// Just free the fence, CmdBuffer was not submitted
-		FenceManager.ReleaseFence(Fence);
-	}
-
 	if (State != EState::NotAllocated)
 	{
 		FreeMemory();
 	}
-
-#if (RHI_NEW_GPU_PROFILER == 0)
-	if (Timing)
-	{
-		Timing->Release();
-		delete Timing;
-		Timing = nullptr;
-	}
-#endif
 }
 
 void FVulkanCommandBuffer::FreeMemory()
@@ -194,17 +147,6 @@ void FVulkanCommandBuffer::End()
 {
 	checkf(IsOutsideRenderPass(), TEXT("Can't End as we're inside a render pass! CmdBuffer 0x%p State=%d"), CommandBufferHandle, (int32)State);
 
-#if (RHI_NEW_GPU_PROFILER == 0)
-	if (GVulkanProfileCmdBuffers || GVulkanUseCmdBufferTimingForGPUTime)
-	{
-		if (Timing)
-		{
-			Timing->EndTiming(this);
-			LastValidTiming = FenceSignaledCounter;
-		}
-	}
-#endif
-
 	for (PendingQuery& Query : PendingTimestampQueries)
 	{
 		const uint64 Index = Query.Index;
@@ -227,38 +169,6 @@ void FVulkanCommandBuffer::End()
 	State = EState::HasEnded;
 }
 
-void FVulkanCommandBuffer::InitializeTimings(FVulkanContextCommon* InContext)
-{
-#if (RHI_NEW_GPU_PROFILER == 0)
-	if ((GVulkanProfileCmdBuffers || GVulkanUseCmdBufferTimingForGPUTime) && !Timing)
-	{
-		if (InContext)
-		{
-			Timing = new FVulkanGPUTiming(InContext, &Device);
-
-			// Upload cb's can be submitted multiple times in a single frame, so we use an expanded pool to catch timings
-			// Any overflow will wrap
-			const uint32 PoolSize = bIsUploadOnly ? 256 : 32;
-			Timing->Initialize(PoolSize);
-		}
-	}
-#endif
-}
-
-void FVulkanCommandBuffer::AddWaitSemaphore(VkPipelineStageFlags InWaitFlags, TArrayView<VulkanRHI::FSemaphore*> InWaitSemaphores)
-{
-	WaitFlags.Reserve(WaitFlags.Num() + InWaitSemaphores.Num());
-
-	for (VulkanRHI::FSemaphore* Sema : InWaitSemaphores)
-	{
-		WaitFlags.Add(InWaitFlags);
-		Sema->AddRef();
-		check(!WaitSemaphores.Contains(Sema));
-	}
-
-	WaitSemaphores.Append(InWaitSemaphores);
-}
-
 void FVulkanCommandBuffer::Begin()
 {
 	{
@@ -280,19 +190,9 @@ void FVulkanCommandBuffer::Begin()
 
 	VERIFYVULKANRESULT(VulkanRHI::vkBeginCommandBuffer(CommandBufferHandle, &CmdBufBeginInfo));
 
-	if (GVulkanProfileCmdBuffers || GVulkanUseCmdBufferTimingForGPUTime)
-	{
-		InitializeTimings(&CommandBufferPool->GetContext());
-#if (RHI_NEW_GPU_PROFILER == 0)
-		if (Timing)
-		{
-			Timing->StartTiming(this);
-		}
-#endif
-	}
 	check(!CurrentDescriptorPoolSetContainer);
 
-	if (!bIsUploadOnly && Device.SupportsBindless())
+	if (Device.SupportsBindless())
 	{
 		FVulkanBindlessDescriptorManager* BindlessDescriptorManager = Device.GetBindlessDescriptorManager();
 		FVulkanQueue& Queue = GetOwner()->GetContext().Queue;
@@ -335,55 +235,35 @@ bool FVulkanCommandBuffer::AcquirePoolSetAndDescriptorsIfNeeded(const class FVul
 	return false;
 }
 
-void FVulkanCommandBuffer::RefreshFenceStatus()
+void FVulkanCommandBuffer::ResetFenceStatus()
 {
 	if (State == EState::Submitted)
 	{
-		VulkanRHI::FFenceManager* FenceMgr = Fence->GetOwner();
-		if (FenceMgr->IsFenceSignaled(Fence))
-		{
-			bHasPipeline = false;
-			bHasViewport = false;
-			bHasScissor = false;
-			bHasStencilRef = false;
+		bHasPipeline = false;
+		bHasViewport = false;
+		bHasScissor = false;
+		bHasStencilRef = false;
 
-			for (VulkanRHI::FSemaphore* Semaphore : SubmittedWaitSemaphores)
-			{
-				Semaphore->Release();
-			}
-			SubmittedWaitSemaphores.Reset();
-
-			CurrentViewports.Empty();
-			CurrentScissors.Empty();
-			CurrentStencilRef = 0;
-#if VULKAN_REUSE_FENCES
-			Fence->GetOwner()->ResetFence(Fence);
-#else
-			VulkanRHI::FFence* PrevFence = Fence;
-			Fence = FenceMgr->AllocateFence();
-			FenceMgr->ReleaseFence(PrevFence);
-#endif
-			++FenceSignaledCounter;
+		CurrentViewports.Empty();
+		CurrentScissors.Empty();
+		CurrentStencilRef = 0;
 
-			if (CurrentDescriptorPoolSetContainer)
-			{
-				//#todo-rco: Reset here?
-				TypedDescriptorPoolSets.Reset();
-				Device.GetDescriptorPoolsManager().ReleasePoolSet(*CurrentDescriptorPoolSetContainer);
-				CurrentDescriptorPoolSetContainer = nullptr;
-			}
-			else
-			{
-				check(TypedDescriptorPoolSets.Num() == 0);
-			}
+		++FenceSignaledCounter;
 
-			// Change state at the end to be safe
-			State = EState::NeedReset;
+		if (CurrentDescriptorPoolSetContainer)
+		{
+			//#todo-rco: Reset here?
+			TypedDescriptorPoolSets.Reset();
+			Device.GetDescriptorPoolsManager().ReleasePoolSet(*CurrentDescriptorPoolSetContainer);
+			CurrentDescriptorPoolSetContainer = nullptr;
 		}
-	}
-	else
-	{
-		check(!Fence->IsSignaled());
+		else
+		{
+			check(TypedDescriptorPoolSets.Num() == 0);
+		}
+
+		// Change state at the end to be safe
+		State = EState::NeedReset;
 	}
 }
 
@@ -467,38 +347,20 @@ void FVulkanCommandBufferPool::FreeUnusedCmdBuffers(FVulkanQueue* InQueue, bool
 #endif
 }
 
-FVulkanCommandBuffer* FVulkanCommandBufferPool::Create(bool bIsUploadOnly)
+FVulkanCommandBuffer* FVulkanCommandBufferPool::Create()
 {
 	// Assumes we are inside a lock for the pool
-	for (int32 Index = FreeCmdBuffers.Num() - 1; Index >= 0; --Index)
+	if (FreeCmdBuffers.Num())
 	{
-		FVulkanCommandBuffer* CmdBuffer = FreeCmdBuffers[Index];
-#if VULKAN_USE_DIFFERENT_POOL_CMDBUFFERS
-		if (CmdBuffer->bIsUploadOnly == bIsUploadOnly)
-#endif
-		{
-			FreeCmdBuffers.RemoveAtSwap(Index);
-			CmdBuffer->AllocMemory();
-			CmdBuffers.Add(CmdBuffer);
-			return CmdBuffer;
-		}
+		FVulkanCommandBuffer* CmdBuffer = FreeCmdBuffers[0];
+		FreeCmdBuffers.RemoveAtSwap(0);
+		CmdBuffer->AllocMemory();
+		CmdBuffers.Add(CmdBuffer);
+		return CmdBuffer;
 	}
 
-	FVulkanCommandBuffer* CmdBuffer = new FVulkanCommandBuffer(Device, this, bIsUploadOnly);
+	FVulkanCommandBuffer* CmdBuffer = new FVulkanCommandBuffer(Device, this);
 	CmdBuffers.Add(CmdBuffer);
 	check(CmdBuffer);
 	return CmdBuffer;
 }
-
-void FVulkanCommandBufferPool::RefreshFenceStatus(FVulkanCommandBuffer* SkipCmdBuffer)
-{
-	FScopeLock ScopeLock(&CS);
-	for (int32 Index = 0; Index < CmdBuffers.Num(); ++Index)
-	{
-		FVulkanCommandBuffer* CmdBuffer = CmdBuffers[Index];
-		if (CmdBuffer != SkipCmdBuffer)
-		{
-			CmdBuffer->RefreshFenceStatus();
-		}
-	}
-}
diff --git a/Engine/Source/Runtime/VulkanRHI/Private/VulkanCommandBuffer.h b/Engine/Source/Runtime/VulkanRHI/Private/VulkanCommandBuffer.h
--- a/Engine/Source/Runtime/VulkanRHI/Private/VulkanCommandBuffer.h
+++ b/Engine/Source/Runtime/VulkanRHI/Private/VulkanCommandBuffer.h
@@ -33,7 +33,7 @@ protected:
 	friend class FVulkanCommandBufferPool;
 	friend class FVulkanQueue;
 
-	FVulkanCommandBuffer(FVulkanDevice& InDevice, FVulkanCommandBufferPool* InCommandBufferPool, bool bInIsUploadOnly);
+	FVulkanCommandBuffer(FVulkanDevice& InDevice, FVulkanCommandBufferPool* InCommandBufferPool);
 	~FVulkanCommandBuffer();
 
 public:
@@ -42,116 +42,56 @@ public:
 		return CommandBufferPool;
 	}
 
-	inline bool IsInsideRenderPass() const
+	bool IsInsideRenderPass() const
 	{
 		return State == EState::IsInsideRenderPass;
 	}
 
-	inline bool IsOutsideRenderPass() const
+	bool IsOutsideRenderPass() const
 	{
 		return State == EState::IsInsideBegin;
 	}
 
-	inline bool HasBegun() const
+	bool HasBegun() const
 	{
 		return State == EState::IsInsideBegin || State == EState::IsInsideRenderPass;
 	}
 
-	inline bool HasEnded() const
+	bool HasEnded() const
 	{
 		return State == EState::HasEnded;
 	}
 
-	inline bool IsSubmitted() const
+	bool IsSubmitted() const
 	{
 		return State == EState::Submitted;
 	}
 
-	inline bool IsAllocated() const
+	bool IsAllocated() const
 	{
 		return State != EState::NotAllocated;
 	}
 
-	inline VkCommandBuffer GetHandle()
+	VkCommandBuffer GetHandle()
 	{
 		return CommandBufferHandle;
 	}
 
-	inline volatile uint64 GetFenceSignaledCounter() const
+	volatile uint64 GetFenceSignaledCounter() const
 	{
 		return FenceSignaledCounter;
 	}
 
-	//#todo-rco: Temp to help find out where the crash is coming from!
-	inline volatile uint64 GetFenceSignaledCounterA() const
+	volatile uint64 GetSubmittedTimelineValue() const
 	{
-		return FenceSignaledCounter;
-	}
-
-	inline volatile uint64 GetFenceSignaledCounterB() const
-	{
-		return FenceSignaledCounter;
-	}
-
-	inline volatile uint64 GetFenceSignaledCounterC() const
-	{
-		return FenceSignaledCounter;
-	}
-
-	inline volatile uint64 GetFenceSignaledCounterD() const
-	{
-		return FenceSignaledCounter;
-	}
-
-	inline volatile uint64 GetFenceSignaledCounterE() const
-	{
-		return FenceSignaledCounter;
-	}
-
-	inline volatile uint64 GetFenceSignaledCounterF() const
-	{
-		return FenceSignaledCounter;
-	}
-
-	inline volatile uint64 GetFenceSignaledCounterG() const
-	{
-		return FenceSignaledCounter;
-	}
-
-	inline volatile uint64 GetFenceSignaledCounterH() const
-	{
-		return FenceSignaledCounter;
-	}
-
-	inline volatile uint64 GetFenceSignaledCounterI() const
-	{
-		return FenceSignaledCounter;
-	}
-
-	inline volatile uint64 GetSubmittedFenceCounter() const
-	{
-		return SubmittedFenceCounter;
+		return SubmittedTimelineValue;
 	}
 
-	inline FVulkanLayoutManager& GetLayoutManager()
+	FVulkanLayoutManager& GetLayoutManager()
 	{
 		return LayoutManager;
 	}
 
-#if (RHI_NEW_GPU_PROFILER == 0)
-	inline bool HasValidTiming() const
-	{
-		return (Timing != nullptr) && (FMath::Abs((int64)FenceSignaledCounter - (int64)LastValidTiming) < 3);
-	}
-#endif
-
-	void VULKANRHI_API AddWaitSemaphore(VkPipelineStageFlags InWaitFlags, VulkanRHI::FSemaphore* InWaitSemaphore)
-	{
-		AddWaitSemaphore(InWaitFlags, MakeArrayView<VulkanRHI::FSemaphore*>(&InWaitSemaphore, 1));
-	}
-
-	void VULKANRHI_API AddWaitSemaphore(VkPipelineStageFlags InWaitFlags, TArrayView<VulkanRHI::FSemaphore*> InWaitSemaphores);
-
 	void Begin();
 	void End();
 
@@ -175,7 +115,6 @@ public:
 	uint8 bHasViewport					: 1;
 	uint8 bHasScissor					: 1;
 	uint8 bHasStencilRef				: 1;
-	uint8 bIsUploadOnly					: 1;
 
 	// You never want to call Begin/EndRenderPass directly as it will mess up the layout manager.
 	void BeginRenderPass(const FVulkanRenderTargetLayout& Layout, class FVulkanRenderPass* RenderPass, class FVulkanFramebuffer* Framebuffer, const VkClearValue* AttachmentClearValues);
@@ -204,35 +143,15 @@ private:
 	VkCommandBuffer CommandBufferHandle;
 	double SubmittedTime = 0.0f;
 
-	TArray<VkPipelineStageFlags> WaitFlags;
-	TArray<VulkanRHI::FSemaphore*> WaitSemaphores;
-	TArray<VulkanRHI::FSemaphore*> SubmittedWaitSemaphores;
 	TArray<PendingQuery> PendingTimestampQueries;
 
-	void MarkSemaphoresAsSubmitted()
-	{
-		WaitFlags.Reset();
-		// Move to pending delete list
-		SubmittedWaitSemaphores = WaitSemaphores;
-		WaitSemaphores.Reset();
-	}
-
-	// Do not cache this pointer as it might change depending on VULKAN_REUSE_FENCES
-	VulkanRHI::FFence* Fence;
-
-	// Last value passed after the fence got signaled
-	volatile uint64 FenceSignaledCounter;
-	// Last value when we submitted the cmd buffer; useful to track down if something waiting for the fence has actually been submitted
-	volatile uint64 SubmittedFenceCounter;
+	volatile uint64 FenceSignaledCounter = 0;
+	volatile uint64 SubmittedTimelineValue = 0;
 
-	void RefreshFenceStatus();
-	void InitializeTimings(FVulkanContextCommon* InContext);
+	void ResetFenceStatus();
 
 	FVulkanCommandBufferPool* CommandBufferPool;
 
-#if (RHI_NEW_GPU_PROFILER == 0)
-	FVulkanGPUTiming* Timing;
-#endif
 	uint64 LastValidTiming;
 
 	void AcquirePoolSetContainer();
@@ -257,8 +176,6 @@ public:
 	FVulkanCommandBufferPool(FVulkanDevice& InDevice, FVulkanContextCommon& InContext);
 	~FVulkanCommandBufferPool();
 
-	void RefreshFenceStatus(FVulkanCommandBuffer* SkipCmdBuffer = nullptr);
-
 	VkCommandPool GetHandle() const
 	{
 		return Handle;
@@ -287,7 +204,7 @@ private:
 
 	FCriticalSection CS;
 
-	FVulkanCommandBuffer* Create(bool bIsUploadOnly);
+	FVulkanCommandBuffer* Create();
 
 	void Create(uint32 QueueFamilyIndex);
 
diff --git a/Engine/Source/Runtime/VulkanRHI/Private/VulkanCommands.cpp b/Engine/Source/Runtime/VulkanRHI/Private/VulkanCommands.cpp
--- a/Engine/Source/Runtime/VulkanRHI/Private/VulkanCommands.cpp
+++ b/Engine/Source/Runtime/VulkanRHI/Private/VulkanCommands.cpp
@@ -17,7 +17,6 @@
 #include "RHICoreShader.h"
 #include "RHIUniformBufferUtilities.h"
 
-
 // make sure what the hardware expects matches what we give it for indirect arguments
 static_assert(sizeof(FRHIDrawIndirectParameters) == sizeof(VkDrawIndirectCommand), "FRHIDrawIndirectParameters size is wrong.");
 static_assert(STRUCT_OFFSET(FRHIDrawIndirectParameters, VertexCountPerInstance) == STRUCT_OFFSET(VkDrawIndirectCommand, vertexCount), "Wrong offset of FRHIDrawIndirectParameters::VertexCountPerInstance.");
@@ -873,26 +872,6 @@ void FVulkanCommandListContext::RHISetDepthBounds(float MinDepth, float MaxDepth
 	VulkanRHI::vkCmdSetDepthBounds(CommandBuffer.GetHandle(), MinDepth, MaxDepth);
 }
 
-void FVulkanCommandListContext::RequestSubmitCurrentCommands()
-{
-	if (Device.GetComputeQueue() == &Queue)
-	{
-		bSubmitAtNextSafePoint = true;
-		SafePointSubmit();
-	}
-	else
-	{
-		ensure(IsImmediate());
-		bSubmitAtNextSafePoint = true;
-	}
-}
-
-void FVulkanCommandListContext::InternalSubmitActiveCmdBuffer()
-{
-	SubmitActiveCmdBuffer();
-	PrepareForNewActiveCommandBuffer();
-}
-
 void FVulkanCommandListContext::RHICopyToStagingBuffer(FRHIBuffer* SourceBufferRHI, FRHIStagingBuffer* StagingBufferRHI, uint32 Offset, uint32 NumBytes)
 {
 	FVulkanCommandBuffer& CommandBuffer = GetCommandBuffer();
@@ -932,108 +911,11 @@ void FVulkanCommandListContext::RHIWriteGPUFence(FRHIGPUFence* FenceRHI)
 	Fence->FenceSignaledCounter = CmdBuffer->GetFenceSignaledCounter();
 }
 
-
-
-
-struct FVulkanPlatformCommandList : public IRHIPlatformCommandList
-{
-	FVulkanCommandListContext* CmdContext = nullptr;
-	FVulkanUploadContext* UploadContext = nullptr;
-};
-
-template<>
-struct TVulkanResourceTraits<IRHIPlatformCommandList>
-{
-	typedef FVulkanPlatformCommandList TConcreteType;
-};
-
 IRHIComputeContext* FVulkanDynamicRHI::RHIGetCommandContext(ERHIPipeline Pipeline, FRHIGPUMask GPUMask)
 {
 	// @todo: RHI command list refactor - fix async compute
 	checkf(Pipeline == ERHIPipeline::Graphics, TEXT("Async compute command contexts not currently implemented."));
 
 	FVulkanCommandListContext* CmdContext = Device->AcquireDeferredContext();
-
-	FVulkanCommandBuffer* CmdBuffer = CmdContext->GetActiveCmdBuffer();
-	if (!CmdBuffer)
-	{
-		CmdContext->PrepareForNewActiveCommandBuffer();
-		CmdBuffer = CmdContext->GetActiveCmdBuffer();
-	}
-	else if (CmdBuffer->IsSubmitted())
-	{
-		CmdContext->PrepareForNewActiveCommandBuffer();
-		CmdBuffer = CmdContext->GetActiveCmdBuffer();
-	}
-	if (!CmdBuffer->HasBegun())
-	{
-		CmdBuffer->Begin();
-	}
-
 	return CmdContext;
 }
-
-void FVulkanDynamicRHI::RHIFinalizeContext(FRHIFinalizeContextArgs&& Args, TRHIPipelineArray<IRHIPlatformCommandList*>& Output)
-{
-	for (IRHIComputeContext* Context : Args.Contexts)
-	{
-		FVulkanPlatformCommandList* PlatformCmdList = new FVulkanPlatformCommandList();
-		PlatformCmdList->CmdContext = static_cast<FVulkanCommandListContext*>(Context);
-		Output[Context->GetPipeline()] = PlatformCmdList;
-	}
-
-	// :todo-jn: Temporarily just place it before the graphic pipeline
-	if (Args.UploadContext)
-	{
-		if (Output[ERHIPipeline::Graphics] == nullptr)
-		{
-			Output[ERHIPipeline::Graphics] = new FVulkanPlatformCommandList();
-		}
-		FVulkanUploadContext* UploadContext = static_cast<FVulkanUploadContext*>(Args.UploadContext);
-		ResourceCast(Output[ERHIPipeline::Graphics])->UploadContext = UploadContext;
-		Args.UploadContext = nullptr;
-	}
-}
-
-void FVulkanDynamicRHI::RHISubmitCommandLists(FRHISubmitCommandListsArgs&& Args)
-{
-	// Submit all the upload context up front
-	for (IRHIPlatformCommandList* Ptr : Args.CommandLists)
-	{
-		FVulkanPlatformCommandList* PlatformCmdList = ResourceCast(Ptr);
-		if (FVulkanUploadContext* UploadContext = PlatformCmdList->UploadContext)
-		{
-			UploadContext->SubmitActiveCmdBuffer();
-			UploadContext->RefreshFenceStatus();
-			FVulkanUploadContext::Pool.Push(UploadContext);
-			PlatformCmdList->UploadContext = nullptr;
-		}
-	}
-
-	for (IRHIPlatformCommandList* Ptr : Args.CommandLists)
-	{
-		FVulkanPlatformCommandList* PlatformCmdList = ResourceCast(Ptr);
-
-		FVulkanCommandListContext* CurrentContext = PlatformCmdList->CmdContext;
-		if (CurrentContext->IsImmediate())
-		{
-			CurrentContext->RequestSubmitCurrentCommands();
-			FVulkanCommandBuffer* CmdBuffer = CurrentContext->GetActiveCmdBuffer();
-			if (CmdBuffer && CmdBuffer->HasBegun() && CmdBuffer->IsOutsideRenderPass())
-			{
-				CurrentContext->SafePointSubmit();
-			}
-			CurrentContext->RefreshFenceStatus();
-		}
-		else
-		{
-			FVulkanCommandBuffer* CmdBuffer = CurrentContext->GetActiveCmdBuffer();
-			check(!CmdBuffer->IsInsideRenderPass());
-			CurrentContext->SubmitActiveCmdBuffer();
-
-			Device->ReleaseDeferredContext(CurrentContext);
-		}
-
-		delete PlatformCmdList;
-	}
-}
diff --git a/Engine/Source/Runtime/VulkanRHI/Private/VulkanContext.cpp b/Engine/Source/Runtime/VulkanRHI/Private/VulkanContext.cpp
--- a/Engine/Source/Runtime/VulkanRHI/Private/VulkanContext.cpp
+++ b/Engine/Source/Runtime/VulkanRHI/Private/VulkanContext.cpp
@@ -7,6 +7,7 @@
 #include "VulkanRHIPrivate.h"
 #include "VulkanDevice.h"
 #include "VulkanPendingState.h"
+#include "DynamicRHI.h"
 
 int32 GVulkanTempBlockSize = 4 * 1024 * 1024;
 static FAutoConsoleVariableRef GCVarVulkanTempBlockSize(
@@ -18,119 +19,89 @@ static FAutoConsoleVariableRef GCVarVulkanTempBlockSize(
 
 
 
-
 FVulkanContextCommon::FVulkanContextCommon(FVulkanDevice& InDevice, FVulkanQueue& InQueue)
 	: Device(InDevice)
 	, Queue(InQueue)
 	, Pool(InDevice, *this)
 {
 	Pool.Create(Queue.GetFamilyIndex());
-
-	ActiveCmdBuffer = Pool.Create(false);
-	ActiveCmdBuffer->InitializeTimings(this);
-	ActiveCmdBuffer->Begin();
 }
 
 FVulkanContextCommon::~FVulkanContextCommon()
 {
 }
 
-void FVulkanContextCommon::WaitForCmdBuffer(FVulkanCommandBuffer* CmdBuffer, float TimeInSecondsToWait)
+void FVulkanContextCommon::NewPayload()
 {
-	FScopeLock ScopeLock(&Pool.CS);
-	check(CmdBuffer->IsSubmitted());
-	bool bSuccess = Device.GetFenceManager().WaitForFence(CmdBuffer->Fence, (uint64)(TimeInSecondsToWait * 1e9));
-	check(bSuccess);
-	CmdBuffer->RefreshFenceStatus();
+	Payloads.Add(new FVulkanPayload(Queue));
+	CurrentPhase = EPhase::Wait;
 }
 
-void FVulkanContextCommon::AddQueryPoolForReset(VkQueryPool QueryPool, uint32 Size)
+// Complete recording of the current command list set, and appends the resulting
+// payloads to the given array. Resets the context so new commands can be recorded.
+void FVulkanContextCommon::Finalize(TArray<FVulkanPayload*>& OutPayloads)
 {
-	FScopeLock ScopeLock(&Pool.CS);
-	FQueryPoolReset QPR = { QueryPool, Size };
-	PoolResets.Add(QPR);
-}
+	if (FVulkanPayload* ResetPayload = FlushResetQueryPools())
+	{
+		OutPayloads.Add(ResetPayload);
+	}
 
-void FVulkanContextCommon::FlushResetQueryPools()
-{
-	FScopeLock ScopeLock(&Pool.CS);
-	if (PoolResets.Num())
+	for (FVulkanPayload* Payload : Payloads)
 	{
-		FVulkanCommandBuffer* PoolResetCommandBuffer = nullptr;
-		for (int32 Index = 0; Index < Pool.CmdBuffers.Num(); ++Index)
+		for (FVulkanCommandBuffer* CommandBuffer : Payload->CommandBuffers)
 		{
-			FVulkanCommandBuffer* CmdBuffer = Pool.CmdBuffers[Index];
-			CmdBuffer->RefreshFenceStatus();
-#if VULKAN_USE_DIFFERENT_POOL_CMDBUFFERS
-			if (!CmdBuffer->bIsUploadOnly)
-#endif
+			checkSlow(!CommandBuffer->IsSubmitted() && CommandBuffer->HasBegun());
+
+			if (!CommandBuffer->IsOutsideRenderPass())
 			{
-				if (CmdBuffer->State == FVulkanCommandBuffer::EState::ReadyForBegin)
-				{
-					PoolResetCommandBuffer = CmdBuffer;
-				}
+				UE_LOG(LogVulkanRHI, Warning, TEXT("Forcing EndRenderPass() for submission"));
+				CommandBuffer->EndRenderPass();
 			}
-		}
-		if (!PoolResetCommandBuffer)
-		{
-			PoolResetCommandBuffer = Pool.Create(false);
-		}
-		PoolResetCommandBuffer->Begin();
 
-		VkCommandBuffer CommandBuffer = PoolResetCommandBuffer->GetHandle();
-		for (FQueryPoolReset& QPR : PoolResets)
-		{
-			VulkanRHI::vkCmdResetQueryPool(CommandBuffer, QPR.Pool, 0, QPR.Size);
+			CommandBuffer->End();
 		}
-		PoolResetCommandBuffer->End();
-		Queue.Submit(PoolResetCommandBuffer);
-		PoolResets.Empty();
 	}
+
+	OutPayloads.Append(MoveTemp(Payloads));
 }
 
-void FVulkanContextCommon::SubmitActiveCmdBuffer(TArrayView<VulkanRHI::FSemaphore*> SignalSemaphores)
+void FVulkanContextCommon::AddQueryPoolForReset(VkQueryPool QueryPool, uint32 Size)
 {
 	FScopeLock ScopeLock(&Pool.CS);
-	FlushResetQueryPools();
-	check(ActiveCmdBuffer);
-
-	TArray<VkSemaphore, FConcurrentLinearArrayAllocator> SemaphoreHandles;
-	SemaphoreHandles.Reserve(SignalSemaphores.Num() + 1);
-	for (VulkanRHI::FSemaphore* Semaphore : SignalSemaphores)
-	{
-		SemaphoreHandles.Add(Semaphore->GetHandle());
-	}
+	FQueryPoolReset QPR = { QueryPool, Size };
+	PoolResets.Add(QPR);
+}
 
-	if (!ActiveCmdBuffer->IsSubmitted() && ActiveCmdBuffer->HasBegun())
+FVulkanPayload* FVulkanContextCommon::FlushResetQueryPools()
+{
+	FScopeLock ScopeLock(&Pool.CS);
+	if (PoolResets.Num())
 	{
-		if (!ActiveCmdBuffer->IsOutsideRenderPass())
+		FVulkanCommandBuffer* PoolResetCommandBuffer = PrepareNewCommandBuffer();
+		const VkCommandBuffer CommandBufferHandle = PoolResetCommandBuffer->GetHandle();
+		for (FQueryPoolReset& QPR : PoolResets)
 		{
-			UE_LOG(LogVulkanRHI, Warning, TEXT("Forcing EndRenderPass() for submission"));
-			ActiveCmdBuffer->EndRenderPass();
+			VulkanRHI::vkCmdResetQueryPool(CommandBufferHandle, QPR.Pool, 0, QPR.Size);
 		}
+		PoolResetCommandBuffer->End();
 
-		VulkanRHI::DebugHeavyWeightBarrier(ActiveCmdBuffer->GetHandle(), 8);
-
-		ActiveCmdBuffer->End();
-
-		Queue.Submit(ActiveCmdBuffer, SemaphoreHandles.Num(), SemaphoreHandles.GetData());
-
-		ActiveCmdBuffer->SubmittedTime = FPlatformTime::Seconds();
+		FVulkanPayload* ResetPayload = new FVulkanPayload(Queue);
+		ResetPayload->CommandBuffers.Add(PoolResetCommandBuffer);
+		PoolResets.Empty();
+		return ResetPayload;
 	}
-
-	ActiveCmdBuffer = nullptr;
+	return nullptr;
 }
 
-void FVulkanContextCommon::SubmitActiveCmdBufferFromPresent(VulkanRHI::FSemaphore* SignalSemaphore)
+// :todo-jn: deprecated, on its way to being removed.  Should not be called, bypasses potential UploadContext which could mess up ordering.
+void FVulkanContextCommon::SubmitActiveCmdBuffer(FVulkanDynamicRHI* DynamicRHI)
 {
-	if (SignalSemaphore)
-	{
-		Queue.Submit(ActiveCmdBuffer, SignalSemaphore->GetHandle());
-	}
-	else
-	{
-		Queue.Submit(ActiveCmdBuffer);
-	}
+	FVulkanPlatformCommandList* FinalizedPayloads = new FVulkanPlatformCommandList;
+	Finalize(*FinalizedPayloads);
+
+	FDynamicRHI::FRHISubmitCommandListsArgs Args;
+	Args.CommandLists.Add(FinalizedPayloads);
+	DynamicRHI->RHISubmitCommandLists(MoveTemp(Args));
 }
 
 #if VULKAN_DELETE_STALE_CMDBUFFERS
@@ -170,48 +141,37 @@ void FVulkanContextCommon::FreeUnusedCmdBuffers(bool bTrimMemory)
 #endif
 }
 
-void FVulkanContextCommon::PrepareForNewActiveCommandBuffer()
+FVulkanCommandBuffer* FVulkanContextCommon::PrepareNewCommandBuffer()
 {
 	FScopeLock ScopeLock(&Pool.CS);
 
+	FVulkanCommandBuffer* NewCommandBuffer = nullptr;
+
 	for (int32 Index = 0; Index < Pool.CmdBuffers.Num(); ++Index)
 	{
 		FVulkanCommandBuffer* CmdBuffer = Pool.CmdBuffers[Index];
-		CmdBuffer->RefreshFenceStatus();
-#if VULKAN_USE_DIFFERENT_POOL_CMDBUFFERS
-		if (!CmdBuffer->bIsUploadOnly)
-#endif
+		if (CmdBuffer->State == FVulkanCommandBuffer::EState::ReadyForBegin || CmdBuffer->State == FVulkanCommandBuffer::EState::NeedReset)
 		{
-			if (CmdBuffer->State == FVulkanCommandBuffer::EState::ReadyForBegin || CmdBuffer->State == FVulkanCommandBuffer::EState::NeedReset)
-			{
-				ActiveCmdBuffer = CmdBuffer;
-				ActiveCmdBuffer->Begin();
-				return;
-			}
-			else
-			{
-				check(CmdBuffer->State == FVulkanCommandBuffer::EState::Submitted);
-			}
+			NewCommandBuffer = CmdBuffer;
+			NewCommandBuffer->Begin();
+			return NewCommandBuffer;
+		}
+		else
+		{
+			check(CmdBuffer->State == FVulkanCommandBuffer::EState::Submitted);
 		}
 	}
 
 	// All cmd buffers are being executed still
-	ActiveCmdBuffer = Pool.Create(false);
-	ActiveCmdBuffer->Begin();
+	NewCommandBuffer = Pool.Create();
+	NewCommandBuffer->Begin();
+	return NewCommandBuffer;
 }
 
 #if (RHI_NEW_GPU_PROFILER == 0)
 uint32 FVulkanContextCommon::CalculateGPUTime()
 {
 	uint32 Time = 0;
-	for (int32 Index = 0; Index < Pool.CmdBuffers.Num(); ++Index)
-	{
-		FVulkanCommandBuffer* CmdBuffer = Pool.CmdBuffers[Index];
-		if (CmdBuffer->HasValidTiming())
-		{
-			Time += CmdBuffer->Timing->GetTiming(false);
-		}
-	}
 	return Time;
 }
 #endif // (RHI_NEW_GPU_PROFILER == 0)
@@ -221,7 +181,6 @@ FVulkanCommandListContext::FVulkanCommandListContext(FVulkanDynamicRHI* InRHI, F
 	: FVulkanContextCommon(InDevice, InQueue)
 	, RHI(InRHI)
 	, Immediate(InImmediate)
-	, bSubmitAtNextSafePoint(false)
 #if (RHI_NEW_GPU_PROFILER == 0)
 	, GpuProfiler(this, &InDevice)
 #endif
@@ -239,7 +198,6 @@ FVulkanCommandListContext::FVulkanCommandListContext(FVulkanDynamicRHI* InRHI, F
 		// Flush the cmd buffer immediately to ensure a valid
 		// 'Last submitted' cmd buffer exists at frame 0.
 		SubmitActiveCmdBuffer();
-		PrepareForNewActiveCommandBuffer();
 	}
 
 	// Create Pending state, contains pipeline states such as current shader and etc..
@@ -314,11 +272,7 @@ FVulkanUploadContext::~FVulkanUploadContext()
 IRHIUploadContext* FVulkanDynamicRHI::RHIGetUploadContext()
 {
 	FVulkanUploadContext* Context = FVulkanUploadContext::Pool.Pop();
-	if (Context)
-	{
-		Context->PrepareForNewActiveCommandBuffer();
-	}
-	else
+	if (!Context)
 	{
 		// :todo-jn: locked to graphics queue for now
 		Context = new FVulkanUploadContext(*Device, *Device->GetGraphicsQueue());
diff --git a/Engine/Source/Runtime/VulkanRHI/Private/VulkanContext.h b/Engine/Source/Runtime/VulkanRHI/Private/VulkanContext.h
--- a/Engine/Source/Runtime/VulkanRHI/Private/VulkanContext.h
+++ b/Engine/Source/Runtime/VulkanRHI/Private/VulkanContext.h
@@ -9,6 +9,7 @@
 #include "VulkanResources.h"
 #include "VulkanRHIPrivate.h"
 #include "VulkanGPUProfiler.h"
+#include "VulkanSubmission.h"
 
 class FVulkanDevice;
 class FVulkanDynamicRHI;
@@ -21,54 +22,90 @@ class FVulkanSwapChain;
 
 class FVulkanContextCommon
 {
+	// Phases are used to track where we are at in filling the current payload.
+	// Phases always move forward in the order declared here, never backwards.
+	// Always use GetPayload(EPhase) to fetch the payload for the phase you want,
+	// it will take care of creating a new payload when necessary.
+	enum class EPhase
+	{
+		Wait,
+		Execute,
+		Signal
+	};
+
 public:
 	FVulkanContextCommon(FVulkanDevice& InDevice, FVulkanQueue& InQueue);
 	virtual ~FVulkanContextCommon();
 
+	FVulkanPayload& GetPayload(EPhase Phase)
+	{
+		if (Payloads.Num() == 0 || Phase < CurrentPhase)
+		{
+			NewPayload();
+		}
+
+		CurrentPhase = Phase;
+		return *Payloads.Last();
+	}
+
+	// NOTE: This call is getting phased out, use GetCommandBuffer()
 	FVulkanCommandBuffer* GetActiveCmdBuffer()
 	{
-		return ActiveCmdBuffer;
+		return &GetCommandBuffer();
 	}
 
 	FVulkanCommandBuffer& GetCommandBuffer()
 	{
-		FVulkanCommandBuffer* CommandBuffer = GetActiveCmdBuffer();
-		check(CommandBuffer);
-		return *CommandBuffer;
+		FVulkanPayload& Payload = GetPayload(EPhase::Execute);
+		checkSlow(Payload.SignalSemaphores.Num() == 0);
+
+		if (Payload.CommandBuffers.Num() == 0)
+		{
+			Payload.CommandBuffers.Add(PrepareNewCommandBuffer());
+		}
+
+		return *Payload.CommandBuffers.Last();
 	}
 
-	bool HasPendingActiveCmdBuffer() const
+	void AddWaitSemaphore(VkPipelineStageFlags InWaitFlags, VulkanRHI::FSemaphore* InWaitSemaphore)
 	{
-		return ActiveCmdBuffer != nullptr;
+		AddWaitSemaphores(InWaitFlags, MakeArrayView<VulkanRHI::FSemaphore*>(&InWaitSemaphore, 1));
 	}
-
-	void SubmitActiveCmdBuffer(TArrayView<VulkanRHI::FSemaphore*> SignalSemaphores);
-
-	void SubmitActiveCmdBuffer()
+	void AddWaitSemaphores(VkPipelineStageFlags InWaitFlags, TArrayView<VulkanRHI::FSemaphore*> InWaitSemaphores)
 	{
-		SubmitActiveCmdBuffer(MakeArrayView<VulkanRHI::FSemaphore*>(nullptr, 0));
+		if (InWaitSemaphores.Num())
+		{
+			FVulkanPayload& Payload = GetPayload(EPhase::Wait);
+			checkSlow((Payload.CommandBuffers.Num() == 0) && (Payload.SignalSemaphores.Num() == 0));
+
+			Payload.WaitFlags.Reserve(Payload.WaitFlags.Num() + InWaitSemaphores.Num());
+			for (int32 Index = 0; Index < InWaitSemaphores.Num(); ++Index)
+			{
+				InWaitSemaphores[Index]->AddRef();
+				Payload.WaitFlags.Add(InWaitFlags);
+			}
+			Payload.WaitSemaphores.Append(InWaitSemaphores);
+		}
 	}
 
-	void SubmitActiveCmdBuffer(VulkanRHI::FSemaphore* SignalSemaphore)
+	void AddSignalSemaphore(VulkanRHI::FSemaphore* InSignalSemaphore)
 	{
-		SubmitActiveCmdBuffer(MakeArrayView<VulkanRHI::FSemaphore*>(&SignalSemaphore, 1));
+		AddSignalSemaphores(MakeArrayView<VulkanRHI::FSemaphore*>(&InSignalSemaphore, 1));
 	}
-
-	/** Regular SACB() expects not-ended and would rotate the command buffer immediately, but Present has a special logic */
-	void SubmitActiveCmdBufferFromPresent(VulkanRHI::FSemaphore* SignalSemaphore = nullptr);
-
-	void WaitForCmdBuffer(FVulkanCommandBuffer* CmdBuffer, float TimeInSecondsToWait = 10.0f);
-
-	void AddQueryPoolForReset(VkQueryPool Pool, uint32 Size);
-	void FlushResetQueryPools();
-
-	// Update the fences of all cmd buffers except SkipCmdBuffer
-	void RefreshFenceStatus(FVulkanCommandBuffer* SkipCmdBuffer = nullptr)
+	void AddSignalSemaphores(TArrayView<VulkanRHI::FSemaphore*> InSignalSemaphores)
 	{
-		Pool.RefreshFenceStatus(SkipCmdBuffer);
+		if (InSignalSemaphores.Num())
+		{
+			FVulkanPayload& Payload = GetPayload(EPhase::Signal);
+			Payload.SignalSemaphores.Append(InSignalSemaphores);
+		}
 	}
 
-	void PrepareForNewActiveCommandBuffer();
+	// Complete recording of the current command list set, and appends the resulting
+	// payloads to the given array. Resets the context so new commands can be recorded.
+	virtual void Finalize(TArray<FVulkanPayload*>& OutPayloads);
+
+	void AddQueryPoolForReset(VkQueryPool Pool, uint32 Size);
 
 	VkCommandPool GetHandle() const
 	{
@@ -81,13 +118,10 @@ public:
 
 	void FreeUnusedCmdBuffers(bool bTrimMemory);
 
+	void SubmitActiveCmdBuffer(FVulkanDynamicRHI* DynamicRHI);
+
 	void NotifyDeletedImage(VkImage Image)
 	{
-		if (ActiveCmdBuffer)
-		{
-			ActiveCmdBuffer->GetLayoutManager().NotifyDeletedImage(Image);
-		}
-
 		Queue.NotifyDeletedImage(Image);
 	}
 
@@ -105,6 +139,10 @@ protected:
 	uint64 FrameCounter = 0;
 
 private:
+	FVulkanCommandBuffer* PrepareNewCommandBuffer();
+	FVulkanPayload* FlushResetQueryPools();
+	void NewPayload();
+
 	FVulkanCommandBufferPool Pool;
 	struct FQueryPoolReset
 	{
@@ -113,10 +151,9 @@ private:
 	};
 	TArray<FQueryPoolReset> PoolResets;
 
-	FVulkanCommandBuffer* ActiveCmdBuffer = nullptr;
+	TArray<FVulkanPayload*> Payloads;
 
-	// Holds semaphores associated with the recent upload cmdbuf(s) - waiting to be added to the next graphics cmdbuf as WaitSemaphores.
-	TArray<VulkanRHI::FSemaphore*> RenderingCompletedSemaphores;
+	EPhase CurrentPhase = EPhase::Wait;
 };
 
 
@@ -279,15 +316,16 @@ public:
 		// Nothing to do
 	}
 
-	void EndRenderQueryInternal(FVulkanCommandBuffer* CmdBuffer, FVulkanRenderQuery* Query);
-
 	void ReleasePendingState();
 
+	void SubmitActiveCmdBuffer()
+	{
+		FVulkanContextCommon::SubmitActiveCmdBuffer(RHI);
+	}
+
 protected:
 	FVulkanDynamicRHI* RHI;
 	FVulkanCommandListContext* Immediate;
-	bool bSubmitAtNextSafePoint;
-	bool bUniformBufferUploadRenderPassDirty = true;
 
 	void BeginOcclusionQueryBatch(FVulkanCommandBuffer* CmdBuffer, uint32 NumQueriesInBatch);
 	void EndOcclusionQueryBatch(FVulkanCommandBuffer* CmdBuffer);
@@ -313,8 +351,6 @@ protected:
 	// Bit array to track which uniform buffers have changed since the last draw call.
 	uint16 DirtyUniformBuffers[SF_NumStandardFrequencies] = {};
 
-	void RequestSubmitCurrentCommands();
-
 	void InternalClearMRT(FVulkanCommandBuffer& CommandBuffer, bool bClearColor, int32 NumClearColors, const FLinearColor* ColorArray, bool bClearDepth, float Depth, bool bClearStencil, uint32 Stencil);
 
 public:
@@ -328,20 +364,6 @@ public:
 private:
 	void RHIClearMRT(bool bClearColor, int32 NumClearColors, const FLinearColor* ColorArray, bool bClearDepth, float Depth, bool bClearStencil, uint32 Stencil);
 
-	inline bool SafePointSubmit()
-	{
-		if (bSubmitAtNextSafePoint)
-		{
-			InternalSubmitActiveCmdBuffer();
-			bSubmitAtNextSafePoint = false;
-			return true;
-		}
-
-		return false;
-	}
-
-	void InternalSubmitActiveCmdBuffer();
-
 	friend class FVulkanDevice;
 	friend class FVulkanDynamicRHI;
 
@@ -360,8 +382,6 @@ private:
 	void ApplyStaticUniformBuffers(TRHIShader* Shader);
 
 	TArray<FRHIUniformBuffer*> GlobalUniformBuffers;
-
-	friend struct FVulkanCommandContextContainer;
 };
 
 class FVulkanCommandListContextImmediate : public FVulkanCommandListContext
@@ -391,6 +411,13 @@ public:
 	static void DestroyPool();
 };
 
+template<>
+struct TVulkanResourceTraits<IRHIUploadContext>
+{
+	typedef FVulkanUploadContext TConcreteType;
+};
+
+
 struct FVulkanContextArray : public TRHIPipelineArray<FVulkanCommandListContext*>
 {
 	FVulkanContextArray(FRHIContextArray const& Contexts)
diff --git a/Engine/Source/Runtime/VulkanRHI/Private/VulkanDevice.cpp b/Engine/Source/Runtime/VulkanRHI/Private/VulkanDevice.cpp
--- a/Engine/Source/Runtime/VulkanRHI/Private/VulkanDevice.cpp
+++ b/Engine/Source/Runtime/VulkanRHI/Private/VulkanDevice.cpp
@@ -320,7 +320,7 @@ void FVulkanDevice::CreateDevice(TArray<const ANSICHAR*>& DeviceLayers, FVulkanD
 		const VkQueueFamilyProperties& CurrProps = QueueFamilyProps[FamilyIndex];
 
 		bool bIsValidQueue = false;
-		if ((CurrProps.queueFlags & VK_QUEUE_GRAPHICS_BIT) == VK_QUEUE_GRAPHICS_BIT)
+		if (VKHasAllFlags(CurrProps.queueFlags, VK_QUEUE_GRAPHICS_BIT))
 		{
 			if (GfxQueueFamilyIndex == -1)
 			{
@@ -333,7 +333,7 @@ void FVulkanDevice::CreateDevice(TArray<const ANSICHAR*>& DeviceLayers, FVulkanD
 			}
 		}
 
-		if ((CurrProps.queueFlags & VK_QUEUE_COMPUTE_BIT) == VK_QUEUE_COMPUTE_BIT)
+		if (VKHasAllFlags(CurrProps.queueFlags, VK_QUEUE_COMPUTE_BIT))
 		{
 			if (ComputeQueueFamilyIndex == -1 &&
 				(GRHIAllowAsyncComputeCvar.GetValueOnAnyThread() != 0 || GAllowPresentOnComputeQueue.GetValueOnAnyThread() != 0) && GfxQueueFamilyIndex != FamilyIndex)
@@ -342,11 +342,11 @@ void FVulkanDevice::CreateDevice(TArray<const ANSICHAR*>& DeviceLayers, FVulkanD
 				bIsValidQueue = true;
 			}
 		}
-
-		if ((CurrProps.queueFlags & VK_QUEUE_TRANSFER_BIT) == VK_QUEUE_TRANSFER_BIT)
+		
+		if (VKHasAllFlags(CurrProps.queueFlags, VK_QUEUE_TRANSFER_BIT))
 		{
 			// Prefer a non-gfx transfer queue
-			if (TransferQueueFamilyIndex == -1 && (CurrProps.queueFlags & VK_QUEUE_GRAPHICS_BIT) != VK_QUEUE_GRAPHICS_BIT && (CurrProps.queueFlags & VK_QUEUE_COMPUTE_BIT) != VK_QUEUE_COMPUTE_BIT)
+			if (TransferQueueFamilyIndex == -1 && !VKHasAnyFlags(CurrProps.queueFlags, VK_QUEUE_GRAPHICS_BIT | VK_QUEUE_COMPUTE_BIT))
 			{
 				TransferQueueFamilyIndex = FamilyIndex;
 				bIsValidQueue = true;
@@ -399,7 +399,7 @@ void FVulkanDevice::CreateDevice(TArray<const ANSICHAR*>& DeviceLayers, FVulkanD
 	FVulkanPlatform::NotifyFoundDeviceLayersAndExtensions(Gpu, DeviceLayers, DeviceExtensions);
 
 	// Create Graphics Queue, here we submit command buffers for execution
-	GfxQueue = new FVulkanQueue(this, GfxQueueFamilyIndex);
+	GfxQueue = new FVulkanQueue(*this, GfxQueueFamilyIndex);
 	if (ComputeQueueFamilyIndex == -1)
 	{
 		// If we didn't find a dedicated Queue, use the default one
@@ -413,13 +413,13 @@ void FVulkanDevice::CreateDevice(TArray<const ANSICHAR*>& DeviceLayers, FVulkanD
 			bAsyncComputeQueue = true;
 		}
 	}
-	ComputeQueue = new FVulkanQueue(this, ComputeQueueFamilyIndex);
+	ComputeQueue = new FVulkanQueue(*this, ComputeQueueFamilyIndex);
 	if (TransferQueueFamilyIndex == -1)
 	{
 		// If we didn't find a dedicated Queue, use the default one
 		TransferQueueFamilyIndex = ComputeQueueFamilyIndex;
 	}
-	TransferQueue = new FVulkanQueue(this, TransferQueueFamilyIndex);
+	TransferQueue = new FVulkanQueue(*this, TransferQueueFamilyIndex);
 
 	uint64 NumBits = QueueFamilyProps[GfxQueueFamilyIndex].timestampValidBits;
 	if (NumBits > 0)
@@ -1529,8 +1529,11 @@ void FVulkanDevice::Destroy()
 	MemoryManager.Deinit();
 
 	delete TransferQueue;
+	TransferQueue = nullptr;
 	delete ComputeQueue;
+	ComputeQueue = nullptr;
 	delete GfxQueue;
+	GfxQueue = nullptr;
 
 	FenceManager.Deinit();
 	DeviceMemoryManager.Deinit();
@@ -1543,9 +1546,6 @@ void FVulkanDevice::Destroy()
 void FVulkanDevice::WaitUntilIdle()
 {
 	VERIFYVULKANRESULT(VulkanRHI::vkDeviceWaitIdle(Device));
-
-	//#todo-rco: Loop through all contexts!
-	GetImmediateContext().RefreshFenceStatus();
 }
 
 const VkComponentMapping& FVulkanDevice::GetFormatComponentMapping(EPixelFormat UEFormat) const
@@ -1570,11 +1570,7 @@ void FVulkanDevice::NotifyDeletedImage(VkImage Image, bool bRenderTarget)
 
 void FVulkanDevice::SubmitCommands(FVulkanCommandListContext* Context)
 {
-	if (Context->HasPendingActiveCmdBuffer())
-	{
-		Context->SubmitActiveCmdBuffer();
-	}
-	Context->PrepareForNewActiveCommandBuffer();
+	Context->SubmitActiveCmdBuffer();
 }
 
 void FVulkanDevice::SubmitCommandsAndFlushGPU()
@@ -1695,3 +1691,10 @@ FGPUTimingCalibrationTimestamp FVulkanDevice::GetCalibrationTimestamp()
 	return CalibrationTimestamp;
 }
 #endif // (RHI_NEW_GPU_PROFILER == 0)
+
+void FVulkanDevice::ForEachQueue(TFunctionRef<void(FVulkanQueue&)> Callback)
+{
+	Callback(*GfxQueue);
+	Callback(*ComputeQueue);
+	Callback(*TransferQueue);
+}
diff --git a/Engine/Source/Runtime/VulkanRHI/Private/VulkanDevice.h b/Engine/Source/Runtime/VulkanRHI/Private/VulkanDevice.h
--- a/Engine/Source/Runtime/VulkanRHI/Private/VulkanDevice.h
+++ b/Engine/Source/Runtime/VulkanRHI/Private/VulkanDevice.h
@@ -58,6 +58,7 @@ struct FOptionalVulkanDeviceExtensions
 			uint64 HasEXTMeshShader : 1;
 			uint64 HasEXTToolingInfo : 1;
 			uint64 HasEXTImageCompressionControl : 1;
+			uint64 HasKHRTimelineSemaphore : 1;
 
 			// Vendor specific
 			uint64 HasAMDBufferMarker : 1;
@@ -309,6 +310,8 @@ public:
 		return PresentQueue;
 	}
 
+	void ForEachQueue(TFunctionRef<void(FVulkanQueue&)> Callback);
+
 	inline VkPhysicalDevice GetPhysicalHandle() const
 	{
 		return Gpu;
diff --git a/Engine/Source/Runtime/VulkanRHI/Private/VulkanExtensions.cpp b/Engine/Source/Runtime/VulkanRHI/Private/VulkanExtensions.cpp
--- a/Engine/Source/Runtime/VulkanRHI/Private/VulkanExtensions.cpp
+++ b/Engine/Source/Runtime/VulkanRHI/Private/VulkanExtensions.cpp
@@ -1056,6 +1056,7 @@ public:
 	virtual void PostPhysicalDeviceFeatures(FOptionalVulkanDeviceExtensions& ExtensionFlags) override final
 	{
 		bRequirementsPassed = (TimelineSemaphoreFeatures.timelineSemaphore == VK_TRUE);
+		ExtensionFlags.HasKHRTimelineSemaphore = bRequirementsPassed;
 	}
 
 	virtual void PreCreateDevice(VkDeviceCreateInfo& DeviceCreateInfo) override final
diff --git a/Engine/Source/Runtime/VulkanRHI/Private/VulkanMemory.cpp b/Engine/Source/Runtime/VulkanRHI/Private/VulkanMemory.cpp
--- a/Engine/Source/Runtime/VulkanRHI/Private/VulkanMemory.cpp
+++ b/Engine/Source/Runtime/VulkanRHI/Private/VulkanMemory.cpp
@@ -4431,7 +4431,7 @@ namespace VulkanRHI
 		if (CmdBuffer)
 		{
 			FPendingItemsPerCmdBuffer* ItemsForCmdBuffer = FindOrAdd(CmdBuffer);
-			FPendingItemsPerCmdBuffer::FPendingItems* ItemsForFence = ItemsForCmdBuffer->FindOrAddItemsForFence(CmdBuffer->GetFenceSignaledCounterA());
+			FPendingItemsPerCmdBuffer::FPendingItems* ItemsForFence = ItemsForCmdBuffer->FindOrAddItemsForFence(CmdBuffer->GetFenceSignaledCounter());
 			check(StagingBuffer);
 			ItemsForFence->Resources.Add(StagingBuffer);
 		}
@@ -4523,7 +4523,7 @@ namespace VulkanRHI
 			for (int32 FenceIndex = EntriesPerCmdBuffer.PendingItems.Num() - 1; FenceIndex >= 0; --FenceIndex)
 			{
 				FPendingItemsPerCmdBuffer::FPendingItems& PendingItems = EntriesPerCmdBuffer.PendingItems[FenceIndex];
-				if (bImmediately || PendingItems.FenceCounter < EntriesPerCmdBuffer.CmdBuffer->GetFenceSignaledCounterB())
+				if (bImmediately || PendingItems.FenceCounter < EntriesPerCmdBuffer.CmdBuffer->GetFenceSignaledCounter())
 				{
 					for (int32 ResourceIndex = 0; ResourceIndex < PendingItems.Resources.Num(); ++ResourceIndex)
 					{
@@ -4828,7 +4828,7 @@ namespace VulkanRHI
 			// #todo-rco: Had to add this check, we were getting null CmdBuffers on the first frame, or before first frame maybe
 			if (bDeleteImmediately ||
 				(GVulkanRHIDeletionFrameNumber > Entry->FrameNumber + NUM_FRAMES_TO_WAIT_FOR_RESOURCE_DELETE &&
-					(Entry->CmdBuffer == nullptr || Entry->FenceCounter < Entry->CmdBuffer->GetFenceSignaledCounterC()))
+					(Entry->CmdBuffer == nullptr || Entry->FenceCounter < Entry->CmdBuffer->GetFenceSignaledCounter()))
 				)
 			{
 				switch (Entry->StructureType)
@@ -5111,32 +5111,74 @@ namespace VulkanRHI
 
 
 
-	FSemaphore::FSemaphore(FVulkanDevice& InDevice) :
-		Device(InDevice),
-		SemaphoreHandle(VK_NULL_HANDLE),
-		bExternallyOwned(false)
+	FSemaphore::FSemaphore(FVulkanDevice& InDevice, EVulkanSemaphoreFlags InFlags)
+		: Device(InDevice)
+		, SemaphoreHandle(VK_NULL_HANDLE)
+		, Flags(InFlags)
 	{
 		// Create semaphore
 		VkSemaphoreCreateInfo CreateInfo;
 		ZeroVulkanStruct(CreateInfo, VK_STRUCTURE_TYPE_SEMAPHORE_CREATE_INFO);
+
+		VkSemaphoreTypeCreateInfo TypeCreateInfo;
+		if (IsTimeline())
+		{
+			checkf(InDevice.GetOptionalExtensions().HasKHRTimelineSemaphore, TEXT("Trying to create timeline semaphore without device support!"));
+			ZeroVulkanStruct(TypeCreateInfo, VK_STRUCTURE_TYPE_SEMAPHORE_TYPE_CREATE_INFO);
+			TypeCreateInfo.semaphoreType = VK_SEMAPHORE_TYPE_TIMELINE;
+			TypeCreateInfo.initialValue = 0;
+			CreateInfo.pNext = &TypeCreateInfo;
+		}
+
 		VERIFYVULKANRESULT(VulkanRHI::vkCreateSemaphore(Device.GetInstanceHandle(), &CreateInfo, VULKAN_CPU_ALLOCATOR, &SemaphoreHandle));
 	}
 
-	FSemaphore::FSemaphore(FVulkanDevice& InDevice, const VkSemaphore& InExternalSemaphore) :
-		Device(InDevice),
-		SemaphoreHandle(InExternalSemaphore),
-		bExternallyOwned(true)
+	FSemaphore::FSemaphore(FVulkanDevice& InDevice, const VkSemaphore& InExternalSemaphore)
+		: Device(InDevice)
+		, SemaphoreHandle(InExternalSemaphore)
+		, Flags(EVulkanSemaphoreFlags::ExternallySignaled)
 	{}
 
 	FSemaphore::~FSemaphore()
 	{
 		check(SemaphoreHandle != VK_NULL_HANDLE);
-		if (!bExternallyOwned)
+		if (!IsExternallyOwned())
 		{
 			Device.GetDeferredDeletionQueue().EnqueueResource(VulkanRHI::FDeferredDeletionQueue2::EType::Semaphore, SemaphoreHandle);
 		}
 		SemaphoreHandle = VK_NULL_HANDLE;
 	}
+
+	bool FSemaphore::WaitForTimelineSemaphoreValue(uint64 Value, uint64 Timeout)
+	{
+		checkf(IsTimeline(), TEXT("WaitForTimelineSemaphoreValue not available!  Semaphore was not created with EVulkanSemaphoreFlags::Timeline."));
+		checkSlow(SemaphoreHandle != VK_NULL_HANDLE);
+		VkSemaphoreWaitInfo WaitInfo;
+		ZeroVulkanStruct(WaitInfo, VK_STRUCTURE_TYPE_SEMAPHORE_WAIT_INFO);
+		WaitInfo.semaphoreCount = 1;
+		WaitInfo.pSemaphores = &SemaphoreHandle;
+		WaitInfo.pValues = (uint64_t*)&Value;
+		const VkResult Result = VulkanRHI::vkWaitSemaphoresKHR(Device.GetInstanceHandle(), &WaitInfo, Timeout);
+		switch (Result)
+		{
+		case VK_SUCCESS:
+			return true;
+		case VK_TIMEOUT:
+			break;
+		default:
+			VERIFYVULKANRESULT(Result);
+			break;
+		}
+		return false;
+	}
+
+	uint64 FSemaphore::GetTimelineSemaphoreValue()
+	{
+		checkf(IsTimeline(), TEXT("GetTimelineSemaphoreValue not available!  Semaphore was not created with EVulkanSemaphoreFlags::Timeline."));
+		uint64 Value = 0;
+		VERIFYVULKANRESULT(VulkanRHI::vkGetSemaphoreCounterValueKHR(Device.GetInstanceHandle(), SemaphoreHandle, (uint64_t*)&Value));
+		return Value;
+	}
 }
 
 
diff --git a/Engine/Source/Runtime/VulkanRHI/Private/VulkanMemory.h b/Engine/Source/Runtime/VulkanRHI/Private/VulkanMemory.h
--- a/Engine/Source/Runtime/VulkanRHI/Private/VulkanMemory.h
+++ b/Engine/Source/Runtime/VulkanRHI/Private/VulkanMemory.h
@@ -1163,28 +1163,55 @@ namespace VulkanRHI
 		FRWLock RWLock;
 	};
 
+	enum class EVulkanSemaphoreFlags : uint8
+	{
+		None = 0,
+
+		// Will not delete handle on destruction
+		ExternallyOwned    = 1 << 1,    
+
+		// Inform submission pipeline that the signal will not come from a payload (acquired image semaphore for example)
+		ExternallySignaled = 1 << 2,
+
+		// Create a timeline semaphore (must be supported)
+		Timeline           = 1 << 3
+	};
+	ENUM_CLASS_FLAGS(EVulkanSemaphoreFlags);
 
 	class VULKANRHI_API FSemaphore : public FRefCount
 	{
 	public:
-		FSemaphore(FVulkanDevice& InDevice);
+		FSemaphore(FVulkanDevice& InDevice, EVulkanSemaphoreFlags InFlags=EVulkanSemaphoreFlags::None);
 		FSemaphore(FVulkanDevice& InDevice, const VkSemaphore& InExternalSemaphore);
 		virtual ~FSemaphore();
 
-		inline VkSemaphore GetHandle() const
+		VkSemaphore GetHandle() const
 		{
 			return SemaphoreHandle;
 		}
 
-		inline bool IsExternallyOwned() const
+		bool IsExternallyOwned() const
+		{
+			return EnumHasAnyFlags(Flags, EVulkanSemaphoreFlags::ExternallyOwned);
+		}
+
+		bool IsExternallySignaled() const
+		{
+			return EnumHasAnyFlags(Flags, EVulkanSemaphoreFlags::ExternallySignaled);
+		}
+
+		bool IsTimeline() const
 		{
-			return bExternallyOwned;
+			return EnumHasAnyFlags(Flags, EVulkanSemaphoreFlags::Timeline);
 		}
 
+		uint64 GetTimelineSemaphoreValue();
+		bool WaitForTimelineSemaphoreValue(uint64 Value, uint64 Timeout);
+
 	private:
 		FVulkanDevice& Device;
 		VkSemaphore SemaphoreHandle;
-		bool bExternallyOwned;
+		const EVulkanSemaphoreFlags Flags;
 	};
 }
 
diff --git a/Engine/Source/Runtime/VulkanRHI/Private/VulkanQuery.cpp b/Engine/Source/Runtime/VulkanRHI/Private/VulkanQuery.cpp
--- a/Engine/Source/Runtime/VulkanRHI/Private/VulkanQuery.cpp
+++ b/Engine/Source/Runtime/VulkanRHI/Private/VulkanQuery.cpp
@@ -393,18 +393,7 @@ void FVulkanCommandListContext::ReadAndCalculateGPUFrameTime()
 
 	if (GSupportsTimestampRenderQueries && FrameTiming)
 	{
-		uint64 Delta = 0;
-		
-		// If we support profile command buffers then use this timing for GPU
-		if (GVulkanUseCmdBufferTimingForGPUTime)
-		{
-			Delta = CalculateGPUTime();
-		}
-		else
-		{
-			Delta = FrameTiming->GetTiming(false);
-		}
-
+		const uint64 Delta = FrameTiming->GetTiming(false);
 		if (Delta > 0)
 		{
 			const double SecondsPerCycle = FPlatformTime::GetSecondsPerCycle();
@@ -542,7 +531,7 @@ bool FVulkanDynamicRHI::RHIGetRenderQueryResult(FRHIRenderQuery* QueryRHI, uint6
 					// CPU wait for query results to be ready.
 					if (bWaitForStart && StartQuerySyncPoint.FenceCounter == StartQuerySyncPoint.CmdBuffer->GetFenceSignaledCounter())
 					{
-						Device->GetImmediateContext().WaitForCmdBuffer(StartQuerySyncPoint.CmdBuffer);
+						WaitForCmdBuffer(StartQuerySyncPoint.CmdBuffer);
 					}
 				}
 
diff --git a/Engine/Source/Runtime/VulkanRHI/Private/VulkanQueue.cpp b/Engine/Source/Runtime/VulkanRHI/Private/VulkanQueue.cpp
--- a/Engine/Source/Runtime/VulkanRHI/Private/VulkanQueue.cpp
+++ b/Engine/Source/Runtime/VulkanRHI/Private/VulkanQueue.cpp
@@ -19,96 +19,349 @@ FAutoConsoleVariableRef CVarVulkanWaitForIdleOnSubmit(
 	ECVF_Default
 	);
 
-FVulkanQueue::FVulkanQueue(FVulkanDevice* InDevice, uint32 InFamilyIndex)
+static int32 GAllowTimelineSemaphores = 1;
+FAutoConsoleVariableRef CVarVulkanSubmissionAllowTimelineSemaphores(
+	TEXT("r.Vulkan.Submission.AllowTimelineSemaphores"),
+	GAllowTimelineSemaphores,
+	TEXT("If supported, use timeline semaphores for queue submission to reduce API calls.\n")
+	TEXT(" 0: Use timeline semaphores if available (default)\n")
+	TEXT(" 1: Use normal fences\n"),
+	ECVF_ReadOnly
+);
+
+// :todo-jn: Merge payloads
+//static int32 GVulkanMergePayloads = 1;
+//static FAutoConsoleVariableRef CVarVulkanSubmissionMergePayloads(
+//	TEXT("r.Vulkan.Submission.MergePayloads"),
+//	GVulkanMergePayloads,
+//	TEXT("0: Submit payloads individually\n")
+//	TEXT("1: Merge consecutive payloads without syncs going to same queue (default)\n"),
+//	ECVF_ReadOnly
+//);
+
+
+FVulkanQueue::FVulkanQueue(FVulkanDevice& InDevice, uint32 InFamilyIndex)
 	: Queue(VK_NULL_HANDLE)
 	, FamilyIndex(InFamilyIndex)
 	, QueueIndex(0)
 	, Device(InDevice)
+	, bUseTimelineSemaphores(InDevice.GetOptionalExtensions().HasKHRTimelineSemaphore && GAllowTimelineSemaphores)
 	, LastSubmittedCmdBuffer(nullptr)
 	, LastSubmittedCmdBufferFenceCounter(0)
-	, SubmitCounter(0)
 	, LayoutManager(false, nullptr)
 {
-	VulkanRHI::vkGetDeviceQueue(Device->GetInstanceHandle(), FamilyIndex, QueueIndex, &Queue);
+	VulkanRHI::vkGetDeviceQueue(Device.GetInstanceHandle(), FamilyIndex, QueueIndex, &Queue);
 
 	FillSupportedStageBits();
+
+	if (bUseTimelineSemaphores)
+	{
+		TimelineSempahore = new VulkanRHI::FSemaphore(InDevice, VulkanRHI::EVulkanSemaphoreFlags::Timeline);
+		TimelineSempahore->AddRef();
+	}
 }
 
 FVulkanQueue::~FVulkanQueue()
 {
-	check(Device);
+	if (TimelineSempahore)
+	{
+		TimelineSempahore->Release();
+	}
 }
 
-void FVulkanQueue::Submit(FVulkanCommandBuffer* CmdBuffer, uint32 NumSignalSemaphores, VkSemaphore* SignalSemaphores)
+int32 FVulkanQueue::SubmitQueuedPayloads(TSet<VkSemaphore>& SignaledSemas)
 {
-	check(CmdBuffer->HasEnded());
+	auto CanBeProcessed = [&SignaledSemas](const TArray<VulkanRHI::FSemaphore*>& WaitSemas)
+	{
+		for (VulkanRHI::FSemaphore* WaitSema : WaitSemas)
+		{
+			const VkSemaphore SemaphoreHandle = WaitSema->GetHandle();
+			if (!WaitSema->IsExternallySignaled() && !SignaledSemas.Contains(SemaphoreHandle))
+			{
+				return false;
+			}
+		}
 
-	VulkanRHI::FFence* Fence = CmdBuffer->Fence;
-	check(!Fence->IsSignaled());
+		// We can only remove them from the list if they are all present
+		for (VulkanRHI::FSemaphore* WaitSema : WaitSemas)
+		{
+			const VkSemaphore SemaphoreHandle = WaitSema->GetHandle();
+			const int32 NumRemoved = SignaledSemas.Remove(SemaphoreHandle);
+			checkSlow(NumRemoved > 0);
+		}
 
-	const VkCommandBuffer CmdBuffers[] = { CmdBuffer->GetHandle() };
+		return true;
+	};
 
-	VkSubmitInfo SubmitInfo;
-	ZeroVulkanStruct(SubmitInfo, VK_STRUCTURE_TYPE_SUBMIT_INFO);
-	SubmitInfo.commandBufferCount = 1;
-	SubmitInfo.pCommandBuffers = CmdBuffers;
-	SubmitInfo.signalSemaphoreCount = NumSignalSemaphores;
-	SubmitInfo.pSignalSemaphores = SignalSemaphores;
+	// Accumulate a list of the payloads we can submit
+	TArray<FVulkanPayload*> Payloads;
+	FVulkanPayload* Payload = nullptr;
+	while (PendingSubmission.Peek(Payload))
+	{
+		// We can only submit the payload if all its wait semas have been signaled
+		if (!CanBeProcessed(Payload->WaitSemaphores))
+		{
+			break;
+		}
 
-	TArray<VkSemaphore> WaitSemaphores;
-	if (CmdBuffer->WaitSemaphores.Num() > 0)
+		Payloads.Add(Payload);
+		PendingSubmission.Pop();
+	}
+
+	if (Payloads.Num())
 	{
-		WaitSemaphores.Empty((uint32)CmdBuffer->WaitSemaphores.Num());
-		for (VulkanRHI::FSemaphore* Semaphore : CmdBuffer->WaitSemaphores)
+		SubmitPayloads(Payloads, SignaledSemas);
+	}
+
+	return Payloads.Num();
+}
+
+void FVulkanQueue::SubmitPayloads(TArrayView<FVulkanPayload*> Payloads, TSet<VkSemaphore>& SignaledSemas)
+{
+	TArray<VkSemaphore> SemaphoreStorage;
+	TArray<VkCommandBuffer> CommandBufferStorage;
+	TArray<VkSubmitInfo> SubmitInfos;
+	TArray<VkTimelineSemaphoreSubmitInfo> TimelineInfos;
+	TArray<uint64> TimelineValues;
+
+	// Presize the arrays so that we don't reallocate
+	const uint32 NumPayloads = Payloads.Num();
+	uint32 NumWaitSemaphores = 0;
+	uint32 NumSignalSemaphores = 0;
+	uint32 NumCommandBuffers = 0;
+
+	for (FVulkanPayload* Payload : Payloads)
+	{
+		NumWaitSemaphores += Payload->WaitSemaphores.Num();
+		NumSignalSemaphores += Payload->SignalSemaphores.Num();
+		NumCommandBuffers += Payload->CommandBuffers.Num();
+	}
+
+	if (bUseTimelineSemaphores)
+	{
+		// We will add a timeline sempahore per payload
+		NumSignalSemaphores += NumPayloads;
+
+		TimelineInfos.Reserve(NumPayloads);
+		TimelineValues.Reserve(NumSignalSemaphores);
+	}
+
+	SemaphoreStorage.Reserve(NumWaitSemaphores + NumSignalSemaphores);
+	CommandBufferStorage.Reserve(NumCommandBuffers);
+	SubmitInfos.Reserve(NumPayloads);
+
+	FVulkanCommandBuffer* LastSubmittedCommandBuffer = nullptr;
+	for (FVulkanPayload* Payload : Payloads)
+	{
+		VkSubmitInfo& SubmitInfo = SubmitInfos.AddZeroed_GetRef();
+		SubmitInfo.sType = VK_STRUCTURE_TYPE_SUBMIT_INFO;
+
+		VkTimelineSemaphoreSubmitInfo* TimelineInfo = nullptr;
+		if (bUseTimelineSemaphores)
+		{
+			TimelineInfo = &TimelineInfos.AddZeroed_GetRef();
+			TimelineInfo->sType = VK_STRUCTURE_TYPE_TIMELINE_SEMAPHORE_SUBMIT_INFO_KHR;
+			SubmitInfo.pNext = TimelineInfo;
+		}
+
+		// Wait Sempahores
+		if (Payload->WaitSemaphores.Num())
 		{
-			WaitSemaphores.Add(Semaphore->GetHandle());
+			const int32 FirstSempahoreIndex = SemaphoreStorage.Num();
+			for (VulkanRHI::FSemaphore* Semaphore : Payload->WaitSemaphores)
+			{
+				SemaphoreStorage.Add(Semaphore->GetHandle());
+			}
+			SubmitInfo.waitSemaphoreCount = Payload->WaitSemaphores.Num();
+			SubmitInfo.pWaitSemaphores = &SemaphoreStorage[FirstSempahoreIndex];
+			SubmitInfo.pWaitDstStageMask = Payload->WaitFlags.GetData();
+		}
+
+		// Command buffers
+		if (Payload->CommandBuffers.Num())
+		{
+			const int32 FirstIndex = CommandBufferStorage.Num();
+			for (FVulkanCommandBuffer* CommandBuffer : Payload->CommandBuffers)
+			{
+				CommandBufferStorage.Add(CommandBuffer->GetHandle());
+
+				CommandBuffer->SubmittedTimelineValue = NextTimelineSemaphoreValue;
+				CommandBuffer->State = FVulkanCommandBuffer::EState::Submitted;
+				CommandBuffer->SubmittedTime = FPlatformTime::Seconds();
+				LastSubmittedCommandBuffer = CommandBuffer;
+
+				// If we're tracking layouts for the queue, merge in the changes recorded in this command buffer's context
+				CommandBuffer->GetLayoutManager().TransferTo(LayoutManager);
+			}
+			SubmitInfo.commandBufferCount = Payload->CommandBuffers.Num();
+			SubmitInfo.pCommandBuffers = &CommandBufferStorage[FirstIndex];
 		}
-		SubmitInfo.waitSemaphoreCount = (uint32)CmdBuffer->WaitSemaphores.Num();
-		SubmitInfo.pWaitSemaphores = WaitSemaphores.GetData();
-		SubmitInfo.pWaitDstStageMask = CmdBuffer->WaitFlags.GetData();
+
+		// Signal Semaphores (and timeline semaphore if supported)
+		if (Payload->SignalSemaphores.Num())
+		{
+			const int32 FirstSempahoreIndex = SemaphoreStorage.Num();
+			for (VulkanRHI::FSemaphore* Semaphore : Payload->SignalSemaphores)
+			{
+				SemaphoreStorage.Add(Semaphore->GetHandle());
+				SignaledSemas.Add(Semaphore->GetHandle());
+			}
+			SubmitInfo.signalSemaphoreCount = Payload->SignalSemaphores.Num();
+			SubmitInfo.pSignalSemaphores = &SemaphoreStorage[FirstSempahoreIndex];
+		}
+
+		if (bUseTimelineSemaphores)
+		{
+			SemaphoreStorage.Add(TimelineSempahore->GetHandle());
+			if (SubmitInfo.pSignalSemaphores)
+			{
+				checkSlow(SubmitInfo.signalSemaphoreCount > 0);
+				++SubmitInfo.signalSemaphoreCount;
+			}
+			else
+			{
+				SubmitInfo.pSignalSemaphores = &SemaphoreStorage[SemaphoreStorage.Num()-1];
+				SubmitInfo.signalSemaphoreCount = 1;
+			}
+
+			const int32 FirstValueIndex = TimelineValues.Num();
+			TimelineInfo->signalSemaphoreValueCount = SubmitInfo.signalSemaphoreCount;
+			TimelineValues.AddZeroed(SubmitInfo.signalSemaphoreCount);
+			TimelineInfo->pSignalSemaphoreValues = (uint64_t*)&TimelineValues[FirstValueIndex];
+
+			Payload->TimelineSemaphoreValue = NextTimelineSemaphoreValue;
+			TimelineValues.Last() = NextTimelineSemaphoreValue;
+		}
+		else
+		{
+			// If timeline semaphores aren't supported, we need to use Fences.
+			// Because there can only be a single Fence per call to QueueSubmit()
+			// we need to submit each payload individually.
+			Payload->TimelineSemaphoreValue = NextTimelineSemaphoreValue;
+			Payload->Fence = Device.GetFenceManager().AllocateFence();
+			Submit(SubmitInfos, Payload->Fence);
+			SubmitInfos.Reset();
+		}
+
+		NextTimelineSemaphoreValue++;
 	}
+
+	if (bUseTimelineSemaphores)
+	{
+		Submit(SubmitInfos, nullptr);
+	}
+
+	// Queue the submitted payloads in the interrupt queue
+	for (FVulkanPayload* Payload : Payloads)
+	{
+		PendingInterrupt.Enqueue(Payload);
+	}
+
+	if (LastSubmittedCmdBuffer)
 	{
-		SCOPE_CYCLE_COUNTER(STAT_VulkanQueueSubmit);
-		VERIFYVULKANRESULT(VulkanRHI::vkQueueSubmit(Queue, 1, &SubmitInfo, Fence->GetHandle()));
+		UpdateLastSubmittedCommandBuffer(LastSubmittedCmdBuffer);
 	}
+}
+
+
+void FVulkanQueue::Submit(TArrayView<VkSubmitInfo> InSubmitInfos, VulkanRHI::FFence* Fence)
+{
+	SCOPE_CYCLE_COUNTER(STAT_VulkanQueueSubmit);
 
-	CmdBuffer->State = FVulkanCommandBuffer::EState::Submitted;
-	CmdBuffer->MarkSemaphoresAsSubmitted();
-	CmdBuffer->SubmittedFenceCounter = CmdBuffer->FenceSignaledCounter;
+	const VkFence FenceHandle = Fence ? Fence->GetHandle() : VK_NULL_HANDLE;
+	VERIFYVULKANRESULT(VulkanRHI::vkQueueSubmit(Queue, InSubmitInfos.Num(), InSubmitInfos.GetData(), FenceHandle));
 
 	if (GWaitForIdleOnSubmit != 0)
 	{
-		VERIFYVULKANRESULT(VulkanRHI::vkDeviceWaitIdle(Device->GetInstanceHandle()));
-		const bool bSuccess = Device->GetFenceManager().WaitForFence(CmdBuffer->Fence, 200 * 1000 * 1000);
-		ensure(bSuccess);
-		ensure(Device->GetFenceManager().IsFenceSignaled(CmdBuffer->Fence));
-		CmdBuffer->GetOwner()->RefreshFenceStatus();
+		VERIFYVULKANRESULT(VulkanRHI::vkDeviceWaitIdle(Device.GetInstanceHandle()));
+
+		if (Fence)
+		{
+			const bool bSuccess = Device.GetFenceManager().WaitForFence(Fence, 500ULL * 1000 * 1000);
+			ensure(bSuccess);
+			ensure(Device.GetFenceManager().IsFenceSignaled(Fence));
+		}
 	}
+}
 
-	UpdateLastSubmittedCommandBuffer(CmdBuffer);
+int32 FVulkanQueue::ProcessInterruptQueue(uint64 Timeout, uint64 UpToValue)
+{
+	if (bUseTimelineSemaphores)
+	{
+		checkSlow(TimelineSempahore);
+		CompletedTimelineSemaphoreValue = TimelineSempahore->GetTimelineSemaphoreValue();
+	}
+
+	int32 NumPayloads = 0;
+	FVulkanPayload* Payload = nullptr;
+	bool bSuccess;
+	do
+	{
+		bSuccess = false;
+		if (PendingInterrupt.Peek(Payload))
+		{
+			if (bUseTimelineSemaphores)
+			{
+				checkSlow(Payload->TimelineSemaphoreValue > 0);
+				bSuccess = (Payload->TimelineSemaphoreValue <= CompletedTimelineSemaphoreValue);
 
-	CmdBuffer->GetOwner()->RefreshFenceStatus(CmdBuffer);
+				if (!bSuccess && (Timeout > 0))
+				{
+					bSuccess = TimelineSempahore->WaitForTimelineSemaphoreValue(Payload->TimelineSemaphoreValue, Timeout);
+					CompletedTimelineSemaphoreValue = TimelineSempahore->GetTimelineSemaphoreValue();
+				}
+			}
+			else
+			{
+				checkSlow(Payload->Fence);
+				bSuccess = (Timeout == 0) ?
+					Device.GetFenceManager().IsFenceSignaled(Payload->Fence) :
+					Device.GetFenceManager().WaitForFence(Payload->Fence, Timeout);
+
+				if (bSuccess)
+				{
+					checkSlow(Payload->TimelineSemaphoreValue > 0);
+					check(CompletedTimelineSemaphoreValue < Payload->TimelineSemaphoreValue);
+					CompletedTimelineSemaphoreValue = Payload->TimelineSemaphoreValue;
+				}
+			}
+		}
+
+		if (bSuccess)
+		{
+			++NumPayloads;
+			PendingInterrupt.Pop();
+
+			for (FVulkanCommandBuffer* CommandBuffer : Payload->CommandBuffers)
+			{
+				CommandBuffer->ResetFenceStatus();
+			}
+
+			delete Payload;
+		}
+
+		if (UpToValue && (UpToValue <= CompletedTimelineSemaphoreValue))
+		{
+			break;
+		}
 
-	Device->GetStagingManager().ProcessPendingFree(false, false);
+	} while (bSuccess);
 
-	// If we're tracking layouts for the queue, merge in the changes recorded in this command buffer's context
-	CmdBuffer->GetLayoutManager().TransferTo(LayoutManager);
+	return NumPayloads;
 }
 
 void FVulkanQueue::UpdateLastSubmittedCommandBuffer(FVulkanCommandBuffer* CmdBuffer)
 {
 	FScopeLock ScopeLock(&CS);
 	LastSubmittedCmdBuffer = CmdBuffer;
-	LastSubmittedCmdBufferFenceCounter = CmdBuffer->GetFenceSignaledCounterH();
-	++SubmitCounter;
+	LastSubmittedCmdBufferFenceCounter = CmdBuffer->GetFenceSignaledCounter();
 }
 
 void FVulkanQueue::FillSupportedStageBits()
 {
-	check(Device);
-	check((int32)FamilyIndex < Device->GetQueueFamilyProps().Num());
+	check((int32)FamilyIndex < Device.GetQueueFamilyProps().Num());
 
-	const VkQueueFamilyProperties& QueueProps = Device->GetQueueFamilyProps()[FamilyIndex];
+	const VkQueueFamilyProperties& QueueProps = Device.GetQueueFamilyProps()[FamilyIndex];
 
 	SupportedStages = 
 		VK_PIPELINE_STAGE_TOP_OF_PIPE_BIT | 
@@ -129,19 +382,19 @@ void FVulkanQueue::FillSupportedStageBits()
 			VK_PIPELINE_STAGE_TRANSFER_BIT |
 			VK_PIPELINE_STAGE_ALL_GRAPHICS_BIT;
 
-		if (Device->GetPhysicalDeviceFeatures().Core_1_0.geometryShader)
+		if (Device.GetPhysicalDeviceFeatures().Core_1_0.geometryShader)
 		{
 			SupportedStages |= VK_PIPELINE_STAGE_GEOMETRY_SHADER_BIT;
 		}
-		if (Device->GetOptionalExtensions().HasKHRFragmentShadingRate)
+		if (Device.GetOptionalExtensions().HasKHRFragmentShadingRate)
 		{
 			SupportedStages |= VK_PIPELINE_STAGE_FRAGMENT_SHADING_RATE_ATTACHMENT_BIT_KHR;
 		}
-		if (Device->GetOptionalExtensions().HasEXTFragmentDensityMap)
+		if (Device.GetOptionalExtensions().HasEXTFragmentDensityMap)
 		{
 			SupportedStages |= VK_PIPELINE_STAGE_FRAGMENT_DENSITY_PROCESS_BIT_EXT;
 		}
-		if (Device->GetOptionalExtensions().HasEXTMeshShader)
+		if (Device.GetOptionalExtensions().HasEXTMeshShader)
 		{
 			SupportedStages |= VK_PIPELINE_STAGE_TASK_SHADER_BIT_EXT | VK_PIPELINE_STAGE_MESH_SHADER_BIT_EXT;
 		}
@@ -154,12 +407,12 @@ void FVulkanQueue::FillSupportedStageBits()
 			VK_PIPELINE_STAGE_DRAW_INDIRECT_BIT |
 			VK_PIPELINE_STAGE_TRANSFER_BIT;
 
-		if (Device->GetOptionalExtensions().HasAccelerationStructure)
+		if (Device.GetOptionalExtensions().HasAccelerationStructure)
 		{
 			SupportedStages |= VK_PIPELINE_STAGE_ACCELERATION_STRUCTURE_BUILD_BIT_KHR;
 		}
 
-		if (Device->GetOptionalExtensions().HasRayTracingPipeline)
+		if (Device.GetOptionalExtensions().HasRayTracingPipeline)
 		{
 			SupportedStages |= VK_PIPELINE_STAGE_RAY_TRACING_SHADER_BIT_KHR;
 		}
diff --git a/Engine/Source/Runtime/VulkanRHI/Private/VulkanQueue.h b/Engine/Source/Runtime/VulkanRHI/Private/VulkanQueue.h
--- a/Engine/Source/Runtime/VulkanRHI/Private/VulkanQueue.h
+++ b/Engine/Source/Runtime/VulkanRHI/Private/VulkanQueue.h
@@ -13,6 +13,7 @@ class FVulkanDevice;
 class FVulkanCommandBuffer;
 class FVulkanCommandListContext;
 class FVulkanLayoutManager;
+class FVulkanPayload;
 
 namespace VulkanRHI
 {
@@ -23,7 +24,7 @@ namespace VulkanRHI
 class FVulkanQueue
 {
 public:
-	FVulkanQueue(FVulkanDevice* InDevice, uint32 InFamilyIndex);
+	FVulkanQueue(FVulkanDevice& InDevice, uint32 InFamilyIndex);
 	~FVulkanQueue();
 
 	inline uint32 GetFamilyIndex() const
@@ -36,13 +37,6 @@ public:
 		return QueueIndex;
 	}
 
-	void Submit(FVulkanCommandBuffer* CmdBuffer, uint32 NumSignalSemaphores = 0, VkSemaphore* SignalSemaphores = nullptr);
-
-	inline void Submit(FVulkanCommandBuffer* CmdBuffer, VkSemaphore SignalSemaphore)
-	{
-		Submit(CmdBuffer, 1, &SignalSemaphore);
-	}
-
 	inline VkQueue GetHandle() const
 	{
 		return Queue;
@@ -55,11 +49,6 @@ public:
 		OutFenceCounter = LastSubmittedCmdBufferFenceCounter;
 	}
 
-	inline uint64 GetSubmitCount() const
-	{
-		return SubmitCounter;
-	}
-
 	inline VkPipelineStageFlags GetSupportedStageBits() const
 	{
 		return SupportedStages;
@@ -72,11 +61,37 @@ public:
 
 	void NotifyDeletedImage(VkImage Image);
 
+	uint64 GetLastSubmittedTimelineSemaphoreValue()
+	{
+		return NextTimelineSemaphoreValue - 1;
+	}
+
+	uint64 GetCompletedTimelineSemaphoreValue()
+	{
+		return CompletedTimelineSemaphoreValue;
+	}
+
+	void EnqueuePayload(FVulkanPayload* InPayload)
+	{
+		PendingSubmission.Enqueue(InPayload);
+	}
+
+	int32 SubmitQueuedPayloads(TSet<VkSemaphore>& SignaledSemas);
+	int32 ProcessInterruptQueue(uint64 Timeout, uint64 UpToValue = 0);
+
 private:
+	void SubmitPayloads(TArrayView<FVulkanPayload*> Payloads, TSet<VkSemaphore>& InOutSignaledSemas);
+	void Submit(TArrayView<VkSubmitInfo> InSubmitInfos, VulkanRHI::FFence* Fence);
+
 	VkQueue Queue;
 	uint32 FamilyIndex;
 	uint32 QueueIndex;
-	FVulkanDevice* Device;
+	FVulkanDevice& Device;
+
+	const bool bUseTimelineSemaphores;
+	VulkanRHI::FSemaphore* TimelineSempahore = nullptr;
+	uint64 NextTimelineSemaphoreValue = 2;
+	uint64 CompletedTimelineSemaphoreValue = 1;
 
 	mutable FCriticalSection CS;
 	FVulkanCommandBuffer* LastSubmittedCmdBuffer;
@@ -84,6 +99,9 @@ private:
 	uint64 SubmitCounter;
 	VkPipelineStageFlags SupportedStages;
 
+	TQueue<FVulkanPayload*, EQueueMode::Mpsc> PendingSubmission;
+	TQueue<FVulkanPayload*, EQueueMode::Spsc> PendingInterrupt;
+
 	// Last known layouts submitted on this queue, used for defrag
 	FVulkanLayoutManager LayoutManager;
 
diff --git a/Engine/Source/Runtime/VulkanRHI/Private/VulkanRHI.cpp b/Engine/Source/Runtime/VulkanRHI/Private/VulkanRHI.cpp
--- a/Engine/Source/Runtime/VulkanRHI/Private/VulkanRHI.cpp
+++ b/Engine/Source/Runtime/VulkanRHI/Private/VulkanRHI.cpp
@@ -1148,7 +1148,7 @@ void FVulkanCommandListContext::RHIEndDrawingViewport(FRHIViewport* ViewportRHI,
 
 	WriteEndTimestamp(&CommandBuffer);
 
-	bool bNativePresent = Viewport->Present(*this, &CommandBuffer, &Queue, Device.GetPresentQueue(), bLockToVsync);
+	bool bNativePresent = Viewport->Present(*this, &Queue, Device.GetPresentQueue(), bLockToVsync);
 	if (bNativePresent)
 	{
 		//#todo-rco: Check for r.FinishCurrentFrame
diff --git a/Engine/Source/Runtime/VulkanRHI/Private/VulkanRenderTarget.cpp b/Engine/Source/Runtime/VulkanRHI/Private/VulkanRenderTarget.cpp
--- a/Engine/Source/Runtime/VulkanRHI/Private/VulkanRenderTarget.cpp
+++ b/Engine/Source/Runtime/VulkanRHI/Private/VulkanRenderTarget.cpp
@@ -268,9 +268,9 @@ void FVulkanDynamicRHI::RHIMapStagingSurface(FRHITexture* TextureRHI, FRHIGPUFen
 	}
 	else
 	{
-		if(GVulkanFlushOnMapStaging)
+		if (GVulkanFlushOnMapStaging)
 		{
-			FRHICommandListExecutor::GetImmediateCommandList().ImmediateFlush(EImmediateFlushType::FlushRHIThread);
+			FRHICommandListImmediate::Get().ImmediateFlush(EImmediateFlushType::FlushRHIThread);
 			Device->WaitUntilIdle();
 		}
 	}
@@ -548,7 +548,6 @@ bool FVulkanCommandListContext::IsSwapchainImage(FRHITexture* InTexture) const
 
 void FVulkanCommandListContext::RHIBeginRenderPass(const FRHIRenderPassInfo& InInfo, const TCHAR* InName)
 {
-	SafePointSubmit();
 	FVulkanCommandBuffer* CmdBuffer = GetActiveCmdBuffer();
 
 	RenderPassInfo = InInfo;
diff --git a/Engine/Source/Runtime/VulkanRHI/Private/VulkanSubmission.cpp b/Engine/Source/Runtime/VulkanRHI/Private/VulkanSubmission.cpp
new file mode 100644
--- /dev/null
+++ b/Engine/Source/Runtime/VulkanRHI/Private/VulkanSubmission.cpp
@@ -0,0 +1,132 @@
+// Copyright Epic Games, Inc. All Rights Reserved.
+
+/*=============================================================================
+	VulkanSubmission.cpp: Vulkan RHI submission implementation.
+=============================================================================*/
+
+#include "VulkanRHIPrivate.h"
+#include "VulkanDynamicRHI.h"
+#include "VulkanContext.h"
+#include "VulkanDevice.h"
+#include "VulkanSubmission.h"
+#include "VulkanQueue.h"
+
+
+FVulkanPayload::~FVulkanPayload()
+{
+	for (VulkanRHI::FSemaphore* Semaphore : WaitSemaphores)
+	{
+		Semaphore->Release();
+	}
+
+	if (Fence)
+	{
+		Fence->GetOwner()->ReleaseFence(Fence);
+	}
+}
+
+void FVulkanDynamicRHI::RHIFinalizeContext(FRHIFinalizeContextArgs&& Args, TRHIPipelineArray<IRHIPlatformCommandList*>& Output)
+{
+	auto GetPlatformCommandList = [&Output](ERHIPipeline Pipeline)
+	{
+		if (!Output[Pipeline])
+		{
+			Output[Pipeline] = new FVulkanPlatformCommandList();
+		}
+		return ResourceCast(Output[Pipeline]);
+	};
+
+	// :todo-jn: Place uploads first on the graphics pipe for now
+	if (Args.UploadContext)
+	{
+		FVulkanPlatformCommandList& PlatformCmdList = *GetPlatformCommandList(ERHIPipeline::Graphics);
+		FVulkanUploadContext* UploadContext = ResourceCast(Args.UploadContext);
+		UploadContext->Finalize(PlatformCmdList);
+		FVulkanUploadContext::Pool.Push(UploadContext);
+		Args.UploadContext = nullptr;
+	}
+
+	for (IRHIComputeContext* Context : Args.Contexts)
+	{
+		FVulkanPlatformCommandList& PlatformCmdList = *GetPlatformCommandList(Context->GetPipeline());
+
+		FVulkanCommandListContext* VulkanContext = static_cast<FVulkanCommandListContext*>(Context);
+		VulkanContext->Finalize(PlatformCmdList);
+
+		if (VulkanContext->IsImmediate())
+		{
+			// :todo-jn: clean up immediate context for next use?
+		}
+		else
+		{
+			Device->ReleaseDeferredContext(VulkanContext);
+		}
+	}
+}
+
+bool FVulkanDynamicRHI::ProcessSubmissionQueue()
+{
+	// Sort all the payloads into the queues
+	FVulkanPlatformCommandList* PlatformCmdList = nullptr;
+	while (PendingPayloadsForSubmission.Dequeue(PlatformCmdList))
+	{
+		for (FVulkanPayload* Payload : *PlatformCmdList)
+		{
+			Payload->Queue.EnqueuePayload(Payload);
+		}
+
+		delete PlatformCmdList;
+	}
+
+	bool bProgress = false;
+	Device->ForEachQueue([&bProgress, &SignaledSemaphores=SignaledSemaphores](FVulkanQueue& Queue)
+	{
+		bProgress = Queue.SubmitQueuedPayloads(SignaledSemaphores) > 0;
+	});
+	return bProgress;
+}
+
+void FVulkanDynamicRHI::ProcessInterruptQueue(uint64 Timeout)
+{
+	Device->ForEachQueue([Timeout](FVulkanQueue& Queue)
+	{
+		Queue.ProcessInterruptQueue(Timeout);
+	});
+}
+
+void FVulkanDynamicRHI::RHISubmitCommandLists(FRHISubmitCommandListsArgs&& Args)
+{
+	for (IRHIPlatformCommandList* CommandList : Args.CommandLists)
+	{
+		FVulkanPlatformCommandList* PlatformCmdList = ResourceCast(CommandList);
+		PendingPayloadsForSubmission.Enqueue(PlatformCmdList);
+	}
+
+	// :todo-jn: add submission thread
+	{
+		FScopeLock Lock(&SubmissionCS);
+		while (ProcessSubmissionQueue()) {};
+	}
+
+	// :todo-jn: add interrupt thread
+	{
+		if (InterruptCS.TryLock())
+		{
+			// Timeout==0 will complete all the payloads it can without waiting
+			ProcessInterruptQueue(0ULL);
+			InterruptCS.Unlock();
+		}
+	}
+}
+
+void FVulkanDynamicRHI::WaitForCmdBuffer(FVulkanCommandBuffer* CmdBuffer, float TimeInSecondsToWait)
+{
+	FVulkanQueue& Queue = CmdBuffer->GetOwner()->GetContext().Queue;
+	const uint64 Timeout = (uint64)(TimeInSecondsToWait * 1e9);
+	checkf(CmdBuffer->IsSubmitted(), TEXT("Attempting to wait on a command buffer that was not submitted!"));
+
+	{
+		FScopeLock Lock(&InterruptCS);
+		Queue.ProcessInterruptQueue(Timeout, CmdBuffer->GetSubmittedTimelineValue());
+	}
+}
diff --git a/Engine/Source/Runtime/VulkanRHI/Private/VulkanSubmission.h b/Engine/Source/Runtime/VulkanRHI/Private/VulkanSubmission.h
new file mode 100644
--- /dev/null
+++ b/Engine/Source/Runtime/VulkanRHI/Private/VulkanSubmission.h
@@ -0,0 +1,54 @@
+// Copyright Epic Games, Inc. All Rights Reserved.
+
+/*=============================================================================
+	VulkanContext.h: Class to generate Vulkan command buffers from RHI CommandLists
+=============================================================================*/
+
+#pragma once 
+
+#include "VulkanRHIPrivate.h"
+
+class FVulkanDevice;
+class FVulkanDynamicRHI;
+class FVulkanOcclusionQueryPool;
+class FVulkanPendingGfxState;
+class FVulkanPendingComputeState;
+class FVulkanQueue;
+class FVulkanSwapChain;
+
+
+class FVulkanPayload
+{
+	friend class FVulkanContextCommon;
+	friend class FVulkanDynamicRHI;
+	friend class FVulkanQueue;
+
+public:
+	FVulkanPayload(FVulkanQueue& InQueue)
+		: Queue(InQueue)
+	{}
+
+	~FVulkanPayload();
+
+protected:
+	FVulkanQueue& Queue;
+	TArray<VkPipelineStageFlags> WaitFlags;  // flags that match 1:1 with WaitSemaphores
+	TArray<VulkanRHI::FSemaphore*> WaitSemaphores; // wait before command buffers
+	TArray<FVulkanCommandBuffer*> CommandBuffers;
+	TArray<VulkanRHI::FSemaphore*> SignalSemaphores; // signaled after command buffers
+
+	// For completion tracking
+	uint64 TimelineSemaphoreValue = 0;
+	VulkanRHI::FFence* Fence = nullptr;
+};
+
+struct FVulkanPlatformCommandList : public IRHIPlatformCommandList, public TArray<FVulkanPayload*>
+{
+};
+
+template<>
+struct TVulkanResourceTraits<IRHIPlatformCommandList>
+{
+	typedef FVulkanPlatformCommandList TConcreteType;
+};
+
diff --git a/Engine/Source/Runtime/VulkanRHI/Private/VulkanSwapChain.cpp b/Engine/Source/Runtime/VulkanRHI/Private/VulkanSwapChain.cpp
--- a/Engine/Source/Runtime/VulkanRHI/Private/VulkanSwapChain.cpp
+++ b/Engine/Source/Runtime/VulkanRHI/Private/VulkanSwapChain.cpp
@@ -603,7 +603,7 @@ FVulkanSwapChain::FVulkanSwapChain(VkInstance InInstance, FVulkanDevice& InDevic
 	ImageAcquiredSemaphore.AddUninitialized(NumSwapChainImages);
 	for (uint32 BufferIndex = 0; BufferIndex < NumSwapChainImages; ++BufferIndex)
 	{
-		ImageAcquiredSemaphore[BufferIndex] = new VulkanRHI::FSemaphore(Device);
+		ImageAcquiredSemaphore[BufferIndex] = new VulkanRHI::FSemaphore(Device, VulkanRHI::EVulkanSemaphoreFlags::ExternallySignaled);
 		ImageAcquiredSemaphore[BufferIndex]->AddRef();
 	}
 
diff --git a/Engine/Source/Runtime/VulkanRHI/Private/VulkanUtil.cpp b/Engine/Source/Runtime/VulkanRHI/Private/VulkanUtil.cpp
--- a/Engine/Source/Runtime/VulkanRHI/Private/VulkanUtil.cpp
+++ b/Engine/Source/Runtime/VulkanRHI/Private/VulkanUtil.cpp
@@ -216,7 +216,6 @@ void FVulkanGPUTiming::DiscardExpiredQueries(bool bNeedSpace)
 			FVulkanTimingQueryPool::FCmdBufferFence& QuerySyncPoint = Pool->TimestampListHandles[QueryIndex];
 			if ((QuerySyncPoint.FenceCounter >= QuerySyncPoint.CmdBuffer->GetFenceSignaledCounter()) && QuerySyncPoint.CmdBuffer->IsSubmitted())
 			{
-				QuerySyncPoint.CmdBuffer->GetOwner()->RefreshFenceStatus();
 				if (QuerySyncPoint.FenceCounter < QuerySyncPoint.CmdBuffer->GetFenceSignaledCounter())
 				{
 					QueryExpired = CanBeDiscarded(QueryIndex);
@@ -369,17 +368,17 @@ uint64 FVulkanGPUTiming::GetTiming(bool bGetCurrentResultsAndBlock)
 					if (bWaitForEnd || bWaitForStart)
 					{
 						// Need to submit the open command lists.
-						Device->SubmitCommandsAndFlushGPU();
+						Device->SubmitCommandsAndFlushGPU();  // :todo-jn: deprecated
 					}
 
 					// CPU wait for query results to be ready.
 					if (bWaitForStart && StartQuerySyncPoint.FenceCounter >= StartQuerySyncPoint.CmdBuffer->GetFenceSignaledCounter())
 					{
-						Context->WaitForCmdBuffer(StartQuerySyncPoint.CmdBuffer);
+						FVulkanDynamicRHI::Get().WaitForCmdBuffer(StartQuerySyncPoint.CmdBuffer);
 					}
 					if (bWaitForEnd && EndQuerySyncPoint.FenceCounter >= EndQuerySyncPoint.CmdBuffer->GetFenceSignaledCounter())
 					{
-						Context->WaitForCmdBuffer(EndQuerySyncPoint.CmdBuffer);
+						FVulkanDynamicRHI::Get().WaitForCmdBuffer(EndQuerySyncPoint.CmdBuffer);
 					}
 				}
 
@@ -768,7 +767,7 @@ void FVulkanGPUFence::Wait(FRHICommandListImmediate& RHICmdList, FRHIGPUMask GPU
 {
 	if (!Poll())
 	{
-		FVulkanDynamicRHI::Get().GetDevice()->GetImmediateContext().WaitForCmdBuffer(CmdBuffer);
+		FVulkanDynamicRHI::Get().WaitForCmdBuffer(CmdBuffer);
 	}
 }
 
diff --git a/Engine/Source/Runtime/VulkanRHI/Private/VulkanViewport.cpp b/Engine/Source/Runtime/VulkanRHI/Private/VulkanViewport.cpp
--- a/Engine/Source/Runtime/VulkanRHI/Private/VulkanViewport.cpp
+++ b/Engine/Source/Runtime/VulkanRHI/Private/VulkanViewport.cpp
@@ -91,7 +91,7 @@ void FVulkanBackBuffer::AcquireBackBufferImage(FVulkanCommandListContext& Contex
 			LayoutMgr.SetFullLayout(Image, CustomLayout);
 
 			// Wait for semaphore signal before writing to backbuffer image
-			CommandBuffer.AddWaitSemaphore(VK_PIPELINE_STAGE_COLOR_ATTACHMENT_OUTPUT_BIT, Viewport->AcquiredSemaphore);
+			Context.AddWaitSemaphore(VK_PIPELINE_STAGE_COLOR_ATTACHMENT_OUTPUT_BIT, Viewport->AcquiredSemaphore);
 		}
 		else
 		{
@@ -289,7 +289,7 @@ void FVulkanViewport::WaitForFrameEventCompletion()
 				if (!GWaitForIdleOnSubmit)
 				{
 					// The wait has already happened if GWaitForIdleOnSubmit is set
-					LastFrameCommandBuffer->GetOwner()->GetContext().WaitForCmdBuffer(LastFrameCommandBuffer);
+					FVulkanDynamicRHI::Get().WaitForCmdBuffer(LastFrameCommandBuffer);
 				}
 			}
 		}
@@ -1021,79 +1021,82 @@ inline static void CopyImageToBackBuffer(FVulkanCommandListContext& Context, FVu
 	}
 }
 
-bool FVulkanViewport::Present(FVulkanCommandListContext& Context, FVulkanCommandBuffer* CmdBuffer, FVulkanQueue* Queue, FVulkanQueue* PresentQueue, bool bLockToVsync)
+bool FVulkanViewport::Present(FVulkanCommandListContext& Context, FVulkanQueue* Queue, FVulkanQueue* PresentQueue, bool bLockToVsync)
 {
 	check(Context.IsImmediate());
 
 	FPlatformAtomics::AtomicStore(&LockToVsync, bLockToVsync ? 1 : 0);
 	bool bFailedToDelayAcquireBackbuffer = false;
 
-	//Transition back buffer to presentable and submit that command
-	check(CmdBuffer->IsOutsideRenderPass());
-
-	if (SupportsStandardSwapchain())
 	{
-		if (GVulkanDelayAcquireImage == EDelayAcquireImageType::DelayAcquire && RenderingBackBuffer)
+		FVulkanCommandBuffer& CommandBuffer = Context.GetCommandBuffer();
+
+		//Transition back buffer to presentable and submit that command
+		check(CommandBuffer.IsOutsideRenderPass());
+
+		if (SupportsStandardSwapchain())
 		{
-			SCOPE_CYCLE_COUNTER(STAT_VulkanAcquireBackBuffer);
-			// swapchain can go out of date, do not crash at this point
-			if (LIKELY(TryAcquireImageIndex()))
+			if (GVulkanDelayAcquireImage == EDelayAcquireImageType::DelayAcquire && RenderingBackBuffer)
 			{
-				uint32 WindowSizeX = FMath::Min(SizeX, SwapChain->InternalWidth);
-				uint32 WindowSizeY = FMath::Min(SizeY, SwapChain->InternalHeight);
+				SCOPE_CYCLE_COUNTER(STAT_VulkanAcquireBackBuffer);
+				// swapchain can go out of date, do not crash at this point
+				if (LIKELY(TryAcquireImageIndex()))
+				{
+					uint32 WindowSizeX = FMath::Min(SizeX, SwapChain->InternalWidth);
+					uint32 WindowSizeY = FMath::Min(SizeY, SwapChain->InternalHeight);
 
-				CopyImageToBackBuffer(Context, CmdBuffer, *RenderingBackBuffer.GetReference(), *BackBufferImages[AcquiredImageIndex].GetReference(), SizeX, SizeY, WindowSizeX, WindowSizeY, SwapChain->GetCachedSurfaceTransform());
+					CopyImageToBackBuffer(Context, &CommandBuffer, *RenderingBackBuffer.GetReference(), *BackBufferImages[AcquiredImageIndex].GetReference(), SizeX, SizeY, WindowSizeX, WindowSizeY, SwapChain->GetCachedSurfaceTransform());
+				}
+				else
+				{
+					bFailedToDelayAcquireBackbuffer = true;
+				}
 			}
 			else
 			{
-				bFailedToDelayAcquireBackbuffer = true;
-			}
-		}
-		else
-		{
-			if (AcquiredImageIndex != -1)
-			{
-				check(RHIBackBuffer != nullptr && RHIBackBuffer->Image == BackBufferImages[AcquiredImageIndex]->Image);
+				if (AcquiredImageIndex != -1)
+				{
+					check(RHIBackBuffer != nullptr && RHIBackBuffer->Image == BackBufferImages[AcquiredImageIndex]->Image);
 
-				FVulkanLayoutManager& LayoutManager = CmdBuffer->GetLayoutManager();
-				const FVulkanImageLayout* TrackedLayout = LayoutManager.GetFullLayout(BackBufferImages[AcquiredImageIndex]->Image);
+					FVulkanLayoutManager& LayoutManager = CommandBuffer.GetLayoutManager();
+					const FVulkanImageLayout* TrackedLayout = LayoutManager.GetFullLayout(BackBufferImages[AcquiredImageIndex]->Image);
 
-				// One of the rare cases we'll let parallel rendering check the tracking (legacy path already checked as a fallback)
-				if (!TrackedLayout && Device->SupportsParallelRendering())
-				{
-					TrackedLayout = Context.Queue.GetLayoutManager().GetFullLayout(BackBufferImages[AcquiredImageIndex]->Image);
-				}
+					// One of the rare cases we'll let parallel rendering check the tracking (legacy path already checked as a fallback)
+					if (!TrackedLayout && Device->SupportsParallelRendering())
+					{
+						TrackedLayout = Context.Queue.GetLayoutManager().GetFullLayout(BackBufferImages[AcquiredImageIndex]->Image);
+					}
 
-				VkImageLayout LastLayout = TrackedLayout ? TrackedLayout->MainLayout : VK_IMAGE_LAYOUT_UNDEFINED;
-				VulkanSetImageLayout(CmdBuffer, BackBufferImages[AcquiredImageIndex]->Image, LastLayout, VK_IMAGE_LAYOUT_PRESENT_SRC_KHR, FVulkanPipelineBarrier::MakeSubresourceRange(VK_IMAGE_ASPECT_COLOR_BIT));
+					VkImageLayout LastLayout = TrackedLayout ? TrackedLayout->MainLayout : VK_IMAGE_LAYOUT_UNDEFINED;
+					VulkanSetImageLayout(&CommandBuffer, BackBufferImages[AcquiredImageIndex]->Image, LastLayout, VK_IMAGE_LAYOUT_PRESENT_SRC_KHR, FVulkanPipelineBarrier::MakeSubresourceRange(VK_IMAGE_ASPECT_COLOR_BIT));
 
-				const FVulkanImageLayout UndefinedLayout(VK_IMAGE_LAYOUT_UNDEFINED, 1, 1, VK_IMAGE_ASPECT_COLOR_BIT);
-				LayoutManager.SetFullLayout(BackBufferImages[AcquiredImageIndex]->Image, UndefinedLayout);
-			}
-			else
-			{
-				// When we have failed to acquire backbuffer image we fallback to using 'dummy' backbuffer
-				check(RHIBackBuffer != nullptr && RHIBackBuffer->Image == RenderingBackBuffer->Image);
+					const FVulkanImageLayout UndefinedLayout(VK_IMAGE_LAYOUT_UNDEFINED, 1, 1, VK_IMAGE_ASPECT_COLOR_BIT);
+					LayoutManager.SetFullLayout(BackBufferImages[AcquiredImageIndex]->Image, UndefinedLayout);
+				}
+				else
+				{
+					// When we have failed to acquire backbuffer image we fallback to using 'dummy' backbuffer
+					check(RHIBackBuffer != nullptr && RHIBackBuffer->Image == RenderingBackBuffer->Image);
+				}
 			}
 		}
 	}
 
-	CmdBuffer->End();
-
-	Context.FlushResetQueryPools();
-	checkf(Context.GetActiveCmdBuffer() == CmdBuffer, TEXT("Present() is submitting something else than the active command buffer"));
 	if (SupportsStandardSwapchain())
 	{
 		if (LIKELY(!bFailedToDelayAcquireBackbuffer))
 		{
 			if (GVulkanDelayAcquireImage == EDelayAcquireImageType::DelayAcquire)
 			{
-				CmdBuffer->AddWaitSemaphore(VK_PIPELINE_STAGE_BOTTOM_OF_PIPE_BIT, AcquiredSemaphore);
+				Context.AddWaitSemaphore(VK_PIPELINE_STAGE_BOTTOM_OF_PIPE_BIT, AcquiredSemaphore);
 			}
 			
-			VulkanRHI::FSemaphore* SignalSemaphore = (AcquiredImageIndex >= 0 ? RenderingDoneSemaphores[AcquiredImageIndex] : nullptr);
-			// submit through the CommandBufferManager as it will add the proper semaphore
-			Context.SubmitActiveCmdBufferFromPresent(SignalSemaphore);
+			if (AcquiredImageIndex >= 0)
+			{
+				Context.AddSignalSemaphore(RenderingDoneSemaphores[AcquiredImageIndex]);
+			}
+
+			Context.SubmitActiveCmdBuffer();  // :todo-jn: deprecated
 		}
 		else
 		{
@@ -1103,11 +1106,11 @@ bool FVulkanViewport::Present(FVulkanCommandListContext& Context, FVulkanCommand
 			UE_LOG(LogVulkanRHI, Log, TEXT("AcquireNextImage() failed due to the outdated swapchain, not even attempting to present."));
 
 			// cannot just throw out this command buffer (needs to be submitted or other checks fail)
-			Queue->Submit(CmdBuffer);
+			Context.SubmitActiveCmdBuffer();  // :todo-jn: deprecated
 			RecreateSwapchain(Context, WindowHandle);
 
 			// Swapchain creation pushes some commands - flush the command buffers now to begin with a fresh state
-			Device->SubmitCommandsAndFlushGPU();
+			Device->SubmitCommandsAndFlushGPU();  // :todo-jn: deprecated
 			Device->WaitUntilIdle();
 
 			// early exit
@@ -1117,12 +1120,9 @@ bool FVulkanViewport::Present(FVulkanCommandListContext& Context, FVulkanCommand
 	else
 	{
 		// Submit active command buffer if not supporting standard swapchain (e.g. XR devices).
-		Context.SubmitActiveCmdBufferFromPresent(nullptr);
+		Context.SubmitActiveCmdBuffer();  // :todo-jn: deprecated
 	}
 
-	//Flush all commands
-	//check(0);
-
 	//#todo-rco: Proper SyncInterval bLockToVsync ? RHIConsoleVariables::SyncInterval : 0
 	int32 SyncInterval = 0;
 	bool bNeedNativePresent = true;
@@ -1181,12 +1181,6 @@ bool FVulkanViewport::Present(FVulkanCommandListContext& Context, FVulkanCommand
 	//	GInputLatencyTimer.RenderThreadTrigger = false;
 	//}
 
-	// PrepareForNewActiveCommandBuffer might be called by swapchain re-creation routine. Skip prepare if we already have an open active buffer.
-	if (Context.GetActiveCmdBuffer() && !Context.GetActiveCmdBuffer()->HasBegun())
-	{
-		Context.PrepareForNewActiveCommandBuffer();
-	}
-
 	AcquiredImageIndex = -1;
 
 	++PresentCount;
diff --git a/Engine/Source/Runtime/VulkanRHI/Private/VulkanViewport.h b/Engine/Source/Runtime/VulkanRHI/Private/VulkanViewport.h
--- a/Engine/Source/Runtime/VulkanRHI/Private/VulkanViewport.h
+++ b/Engine/Source/Runtime/VulkanRHI/Private/VulkanViewport.h
@@ -73,7 +73,7 @@ public:
 
 	virtual void Tick(float DeltaTime) override final;
 	
-	bool Present(FVulkanCommandListContext& Context, FVulkanCommandBuffer* CmdBuffer, FVulkanQueue* Queue, FVulkanQueue* PresentQueue, bool bLockToVsync);
+	bool Present(FVulkanCommandListContext& Context, FVulkanQueue* Queue, FVulkanQueue* PresentQueue, bool bLockToVsync);
 
 	inline uint32 GetPresentCount() const
 	{
diff --git a/Engine/Source/Runtime/VulkanRHI/Public/VulkanDynamicRHI.h b/Engine/Source/Runtime/VulkanRHI/Public/VulkanDynamicRHI.h
--- a/Engine/Source/Runtime/VulkanRHI/Public/VulkanDynamicRHI.h
+++ b/Engine/Source/Runtime/VulkanRHI/Public/VulkanDynamicRHI.h
@@ -14,9 +14,10 @@ class FVulkanDevice;
 class FVulkanQueue;
 class FVulkanUniformBuffer;
 class FVulkanViewport;
+class FVulkanPayload;
 class IHeadMountedDisplayVulkanExtensions;
 struct FRHITransientHeapAllocation;
-
+struct FVulkanPlatformCommandList;
 struct FOptionalVulkanInstanceExtensions
 {
 	union
@@ -274,11 +275,17 @@ public:
 public:
 	static void SavePipelineCache();
 
-private:
-	void PooledUniformBuffersBeginFrame();
-	void ReleasePooledUniformBuffers();
+	void WaitForCmdBuffer(FVulkanCommandBuffer* CmdBuffer, float TimeInSecondsToWait=30.0f);
 
 protected:
+	bool ProcessSubmissionQueue();
+	void ProcessInterruptQueue(uint64 Timeout);
+
+	FCriticalSection SubmissionCS;
+	FCriticalSection InterruptCS;
+	TQueue<FVulkanPlatformCommandList*, EQueueMode::Mpsc> PendingPayloadsForSubmission;
+	TSet<VkSemaphore> SignaledSemaphores;
+
 	uint32 ApiVersion = 0;
 	VkInstance Instance;
 	TArray<const ANSICHAR*> InstanceExtensions;
