VulkanRHI:  fix static analysis warning and put back wait.  fix assert in debug.

#rnx

diff --git a/Engine/Source/Runtime/VulkanRHI/Private/VulkanUtil.cpp b/Engine/Source/Runtime/VulkanRHI/Private/VulkanUtil.cpp
--- a/Engine/Source/Runtime/VulkanRHI/Private/VulkanUtil.cpp
+++ b/Engine/Source/Runtime/VulkanRHI/Private/VulkanUtil.cpp
@@ -216,6 +216,7 @@ void FVulkanGPUTiming::DiscardExpiredQueries(bool bNeedSpace)
 			FVulkanTimingQueryPool::FCmdBufferFence& QuerySyncPoint = Pool->TimestampListHandles[QueryIndex];
 			if ((QuerySyncPoint.FenceCounter >= QuerySyncPoint.CmdBuffer->GetFenceSignaledCounter()) && QuerySyncPoint.CmdBuffer->IsSubmitted())
 			{
+				FVulkanDynamicRHI::Get().WaitForCmdBuffer(QuerySyncPoint.CmdBuffer, 0.0f);
 				if (QuerySyncPoint.FenceCounter < QuerySyncPoint.CmdBuffer->GetFenceSignaledCounter())
 				{
 					QueryExpired = CanBeDiscarded(QueryIndex);
