VulkanRHI: Fix debug check always triggering

#rnx

diff --git a/Engine/Source/Runtime/VulkanRHI/Private/VulkanQueue.cpp b/Engine/Source/Runtime/VulkanRHI/Private/VulkanQueue.cpp
--- a/Engine/Source/Runtime/VulkanRHI/Private/VulkanQueue.cpp
+++ b/Engine/Source/Runtime/VulkanRHI/Private/VulkanQueue.cpp
@@ -85,9 +85,12 @@ int32 FVulkanQueue::SubmitQueuedPayloads(TSet<VkSemaphore>& SignaledSemas)
 		// We can only remove them from the list if they are all present
 		for (VulkanRHI::FSemaphore* WaitSema : WaitSemas)
 		{
-			const VkSemaphore SemaphoreHandle = WaitSema->GetHandle();
-			const int32 NumRemoved = SignaledSemas.Remove(SemaphoreHandle);
-			checkSlow(NumRemoved > 0);
+			if (!WaitSema->IsExternallySignaled())
+			{
+				const VkSemaphore SemaphoreHandle = WaitSema->GetHandle();
+				const int32 NumRemoved = SignaledSemas.Remove(SemaphoreHandle);
+				checkSlow(NumRemoved > 0);
+			}
 		}
 
 		return true;
