VulkanRHI: Quick fix for unfreed timeline semaphores on shutdown.

#rnx

diff --git a/Engine/Source/Runtime/VulkanRHI/Private/VulkanMemory.cpp b/Engine/Source/Runtime/VulkanRHI/Private/VulkanMemory.cpp
--- a/Engine/Source/Runtime/VulkanRHI/Private/VulkanMemory.cpp
+++ b/Engine/Source/Runtime/VulkanRHI/Private/VulkanMemory.cpp
@@ -4739,7 +4739,7 @@ namespace VulkanRHI
 
 	FDeferredDeletionQueue2::~FDeferredDeletionQueue2()
 	{
-		check(Entries.Num() == 0);
+		checkf(Entries.Num() == 0, TEXT("There are %d entries remaining in the Deferred Deletion Queue!"), Entries.Num());
 	}
 
 	void FDeferredDeletionQueue2::EnqueueGenericResource(EType Type, uint64 Handle)
@@ -5142,7 +5142,11 @@ namespace VulkanRHI
 	FSemaphore::~FSemaphore()
 	{
 		check(SemaphoreHandle != VK_NULL_HANDLE);
-		if (!IsExternallyOwned())
+		if (EnumHasAnyFlags(Flags, EVulkanSemaphoreFlags::ImmediateDeletion))
+		{
+			VulkanRHI::vkDestroySemaphore(Device.GetInstanceHandle(), SemaphoreHandle, VULKAN_CPU_ALLOCATOR);
+		}
+		else if (!IsExternallyOwned())
 		{
 			Device.GetDeferredDeletionQueue().EnqueueResource(VulkanRHI::FDeferredDeletionQueue2::EType::Semaphore, SemaphoreHandle);
 		}
diff --git a/Engine/Source/Runtime/VulkanRHI/Private/VulkanMemory.h b/Engine/Source/Runtime/VulkanRHI/Private/VulkanMemory.h
--- a/Engine/Source/Runtime/VulkanRHI/Private/VulkanMemory.h
+++ b/Engine/Source/Runtime/VulkanRHI/Private/VulkanMemory.h
@@ -1174,7 +1174,10 @@ namespace VulkanRHI
 		ExternallySignaled = 1 << 2,
 
 		// Create a timeline semaphore (must be supported)
-		Timeline           = 1 << 3
+		Timeline           = 1 << 3,
+
+		// Will not be queued for deferred deletion
+		ImmediateDeletion  = 1 << 4
 	};
 	ENUM_CLASS_FLAGS(EVulkanSemaphoreFlags);
 
diff --git a/Engine/Source/Runtime/VulkanRHI/Private/VulkanQueue.cpp b/Engine/Source/Runtime/VulkanRHI/Private/VulkanQueue.cpp
--- a/Engine/Source/Runtime/VulkanRHI/Private/VulkanQueue.cpp
+++ b/Engine/Source/Runtime/VulkanRHI/Private/VulkanQueue.cpp
@@ -56,7 +56,8 @@ FVulkanQueue::FVulkanQueue(FVulkanDevice& InDevice, uint32 InFamilyIndex)
 
 	if (bUseTimelineSemaphores)
 	{
-		TimelineSempahore = new VulkanRHI::FSemaphore(InDevice, VulkanRHI::EVulkanSemaphoreFlags::Timeline);
+		// Use ImmediateDeletion for the timeline since it's deleted after the deferrered deletion queue
+		TimelineSempahore = new VulkanRHI::FSemaphore(InDevice, VulkanRHI::EVulkanSemaphoreFlags::Timeline | VulkanRHI::EVulkanSemaphoreFlags::ImmediateDeletion);
 		TimelineSempahore->AddRef();
 	}
 }
