[HWRT] Allow setting the minimum primitive count threshold below which ray tracing geometry skips the compaction (by default 128).
This avoid issues with compacting small ray tracing geometries on some of the drivers.

#rnx

diff --git a/Engine/Source/Runtime/D3D12RHI/Private/D3D12RayTracing.cpp b/Engine/Source/Runtime/D3D12RHI/Private/D3D12RayTracing.cpp
--- a/Engine/Source/Runtime/D3D12RHI/Private/D3D12RayTracing.cpp
+++ b/Engine/Source/Runtime/D3D12RHI/Private/D3D12RayTracing.cpp
@@ -64,6 +64,13 @@ static FAutoConsoleVariableRef CVarD3D12RayTracingMaxBatchedCompaction(
 	ECVF_ReadOnly
 );
 
+static int32 GD3D12RayTracingCompactionMinPrimitiveCount = 128;
+static FAutoConsoleVariableRef CVarD3D12RayTracingCompactionMinPrimitiveCount(
+	TEXT("r.D3D12.RayTracing.Compaction.MinPrimitiveCount"),
+	GD3D12RayTracingCompactionMinPrimitiveCount,
+	TEXT("Sets the minimum primitive count threshold below which geometry skips the compaction. (default = 128)\n")
+);
+
 static int32 GRayTracingSpecializeStateObjects = 0;
 static FAutoConsoleVariableRef CVarRayTracingSpecializeStateObjects(
 	TEXT("r.D3D12.RayTracing.SpecializeStateObjects"),
@@ -3043,7 +3050,7 @@ static ERayTracingAccelerationStructureFlags GetRayTracingAccelerationStructureB
 		EnumAddFlags(BuildFlags, ERayTracingAccelerationStructureFlags::AllowUpdate);
 	}
 
-	if (!Initializer.bFastBuild && !Initializer.bAllowUpdate && Initializer.bAllowCompaction && GD3D12RayTracingAllowCompaction)
+	if (!Initializer.bFastBuild && !Initializer.bAllowUpdate && Initializer.bAllowCompaction && GD3D12RayTracingAllowCompaction && (uint32(GD3D12RayTracingCompactionMinPrimitiveCount) < Initializer.TotalPrimitiveCount))
 	{
 		EnumAddFlags(BuildFlags, ERayTracingAccelerationStructureFlags::AllowCompaction);
 	}
diff --git a/Engine/Source/Runtime/VulkanRHI/Private/VulkanRayTracing.cpp b/Engine/Source/Runtime/VulkanRHI/Private/VulkanRayTracing.cpp
--- a/Engine/Source/Runtime/VulkanRHI/Private/VulkanRayTracing.cpp
+++ b/Engine/Source/Runtime/VulkanRHI/Private/VulkanRayTracing.cpp
@@ -26,6 +26,13 @@ static FAutoConsoleVariableRef CVarVulkanRayTracingMaxBatchedCompaction(
 	ECVF_ReadOnly
 );
 
+static int32 GVulkanRayTracingCompactionMinPrimitiveCount = 128;
+static FAutoConsoleVariableRef CVarVulkanRayTracingCompactionMinPrimitiveCount(
+	TEXT("r.Vulkan.RayTracing.Compaction.MinPrimitiveCount"),
+	GVulkanRayTracingCompactionMinPrimitiveCount,
+	TEXT("Sets the minimum primitive count threshold below which geometry skips the compaction. (default = 128)\n")
+);
+
 static int32 GVulkanRayTracingAllowDeferredOperation = 0;
 static FAutoConsoleVariableRef CVarVulkanRayTracingAllowDeferredOperation(
 	TEXT("r.Vulkan.RayTracing.AllowDeferredOperation"),
@@ -185,7 +192,7 @@ static ERayTracingAccelerationStructureFlags GetRayTracingAccelerationStructureB
 		EnumAddFlags(BuildFlags, ERayTracingAccelerationStructureFlags::AllowUpdate);
 	}
 
-	if (!Initializer.bFastBuild && !Initializer.bAllowUpdate && Initializer.bAllowCompaction && GVulkanRayTracingAllowCompaction)
+	if (!Initializer.bFastBuild && !Initializer.bAllowUpdate && Initializer.bAllowCompaction && GVulkanRayTracingAllowCompaction && (uint32(GVulkanRayTracingCompactionMinPrimitiveCount) < Initializer.TotalPrimitiveCount))
 	{
 		EnumAddFlags(BuildFlags, ERayTracingAccelerationStructureFlags::AllowCompaction);
 	}
