Fix vulkan SBT binding code which wasn't handling clear bindings correctly

diff --git a/Engine/Source/Runtime/Renderer/Private/RayTracing/RayTracingShaderBindingTable.cpp b/Engine/Source/Runtime/Renderer/Private/RayTracing/RayTracingShaderBindingTable.cpp
--- a/Engine/Source/Runtime/Renderer/Private/RayTracing/RayTracingShaderBindingTable.cpp
+++ b/Engine/Source/Runtime/Renderer/Private/RayTracing/RayTracingShaderBindingTable.cpp
@@ -507,6 +507,7 @@ void FRayTracingShaderBindingTable::FlushAllocationsToClear(FRHICommandList& RHI
 		for (uint32 RangeIndex = 0; RangeIndex < Allocation->NumRecords; ++RangeIndex)
 		{
 			FRayTracingLocalShaderBindings& RecordBindingToClear = LocalBindings[LocalBindingIndex++];
+			RecordBindingToClear = FRayTracingLocalShaderBindings();
 			RecordBindingToClear.RecordIndex = Allocation->BaseRecordIndex + RangeIndex;
 			RecordBindingToClear.BindingType = ERayTracingLocalShaderBindingType::Clear;
 		}
diff --git a/Engine/Source/Runtime/VulkanRHI/Private/VulkanRayTracing.cpp b/Engine/Source/Runtime/VulkanRHI/Private/VulkanRayTracing.cpp
--- a/Engine/Source/Runtime/VulkanRHI/Private/VulkanRayTracing.cpp
+++ b/Engine/Source/Runtime/VulkanRHI/Private/VulkanRayTracing.cpp
@@ -2335,16 +2335,24 @@ void FVulkanCommandListContext::RHISetBindingsOnShaderBindingTable(FRHIShaderBin
 		{
 			const FVulkanRayTracingGeometry* Geometry = ResourceCast(Binding.Geometry);
 
-			SetRayTracingHitGroup( 
-				ShaderTable, Binding.RecordIndex,
-				Pipeline, Binding.ShaderIndexInPipeline,
-				Geometry, Binding.SegmentIndex,
-				Binding.NumUniformBuffers,
-				Binding.UniformBuffers,
-				Binding.LooseParameterDataSize,
-				Binding.LooseParameterData,
-				Binding.UserData,
-				Context.WorkerIndex);
+			if (Binding.BindingType != ERayTracingLocalShaderBindingType::Clear)
+			{
+				SetRayTracingHitGroup(
+					ShaderTable, Binding.RecordIndex,
+					Pipeline, Binding.ShaderIndexInPipeline,
+					Geometry, Binding.SegmentIndex,
+					Binding.NumUniformBuffers,
+					Binding.UniformBuffers,
+					Binding.LooseParameterDataSize,
+					Binding.LooseParameterData,
+					Binding.UserData,
+					Context.WorkerIndex);
+			}
+			else
+			{
+				// Only transient SBT support for Vulkan right now (otherwise hit record data might need to be cleared)
+				check(ShaderTable->GetInitializer().Lifetime == ERayTracingShaderBindingTableLifetime::Transient);
+			}
 		}
 		else if (BindingType == ERayTracingBindingType::CallableShader)
 		{
