VulkanRHI: Silence validation error from SPIRV code generation until compiler fix is integrated.

#rnx

diff --git a/Engine/Source/Runtime/VulkanRHI/Private/VulkanDebug.cpp b/Engine/Source/Runtime/VulkanRHI/Private/VulkanDebug.cpp
--- a/Engine/Source/Runtime/VulkanRHI/Private/VulkanDebug.cpp
+++ b/Engine/Source/Runtime/VulkanRHI/Private/VulkanDebug.cpp
@@ -154,6 +154,12 @@ static VKAPI_ATTR VkBool32 VKAPI_CALL DebugUtilsCallback(VkDebugUtilsMessageSeve
 		// *** [Error:Validation(VUID-VkShaderModuleCreateInfo-pCode-04147)] vkCreateShaderModule(): The SPIR-V Extension (SPV_GOOGLE_user_type) was declared, but none of the requirements were met to use it.
 		return VK_FALSE;
 	}
+	else if (FCStringAnsi::Strstr(CallbackData->pMessage, "Cannot store to OpTypeImage, OpTypeSampler, OpTypeSampledImage, or OpTypeAccelerationStructure") != nullptr)
+	{
+		// *** [VUID-StandaloneSpirv-OpTypeImage-06924] Cannot store to OpTypeImage, OpTypeSampler, OpTypeSampledImage, or OpTypeAccelerationStructureKHR objects
+		// OpStore is called on OpTypeAccelerationStructure when they are used in bindless.  Currently no side effects from it, skip until compiler fix is integrated.
+		return VK_FALSE;
+	}
 
 	static TStaticArray<TSet<uint32>, (int32)EMsgBucket::Count> SeenCodes;
 	bool bUniqueMessages = GCVarUniqueValidationMessages->GetInt() != 0;
