VulkanRHI: Do not include USAGE_STORAGE_BIT in SRVUsage of the DefaultView when a different format is used for UAV views.

#rnx

diff --git a/Engine/Source/Runtime/VulkanRHI/Private/VulkanTexture.cpp b/Engine/Source/Runtime/VulkanRHI/Private/VulkanTexture.cpp
--- a/Engine/Source/Runtime/VulkanRHI/Private/VulkanTexture.cpp
+++ b/Engine/Source/Runtime/VulkanRHI/Private/VulkanTexture.cpp
@@ -1466,7 +1466,8 @@ FVulkanTexture::FVulkanTexture(FVulkanDevice& InDevice, const FRHITextureCreateD
 	}
 
 	const VkDescriptorType DescriptorType = SupportsSampling() ? VK_DESCRIPTOR_TYPE_SAMPLED_IMAGE : VK_DESCRIPTOR_TYPE_STORAGE_IMAGE;
-	const VkImageUsageFlags SRVUsage = bIsSRGB ? (ImageCreateInfo.ImageCreateInfo.usage & ~VK_IMAGE_USAGE_STORAGE_BIT) : ImageCreateInfo.ImageCreateInfo.usage;
+	const bool bHasUAVFormat = (InCreateDesc.UAVFormat != PF_Unknown && InCreateDesc.UAVFormat != InCreateDesc.Format);
+	const VkImageUsageFlags SRVUsage = (bIsSRGB || bHasUAVFormat) ? (ImageCreateInfo.ImageCreateInfo.usage & ~VK_IMAGE_USAGE_STORAGE_BIT) : ImageCreateInfo.ImageCreateInfo.usage;
 	if (ViewType != VK_IMAGE_VIEW_TYPE_MAX_ENUM) //-V547
 	{
 		DefaultView = (new FVulkanView(InDevice, DescriptorType))->InitAsTextureView(
