VulkanRHI: change compatibility code for RADV since RT driver fix has been identified for version 25.

#rnx

diff --git a/Engine/Source/Runtime/VulkanRHI/Private/Linux/VulkanLinuxPlatform.cpp b/Engine/Source/Runtime/VulkanRHI/Private/Linux/VulkanLinuxPlatform.cpp
--- a/Engine/Source/Runtime/VulkanRHI/Private/Linux/VulkanLinuxPlatform.cpp
+++ b/Engine/Source/Runtime/VulkanRHI/Private/Linux/VulkanLinuxPlatform.cpp
@@ -15,7 +15,6 @@
 ENUM_VK_ENTRYPOINTS_ALL(DEFINE_VK_ENTRYPOINTS)
 
 static bool GRenderOffScreen = false;
-extern int32 GVulkanAMDCompatibilityMode;
 
 void* FVulkanLinuxPlatform::VulkanLib = nullptr;
 bool FVulkanLinuxPlatform::bAttemptedLoad = false;
@@ -52,17 +51,6 @@ bool FVulkanLinuxPlatform::LoadVulkanLibrary()
 	}
 	bAttemptedLoad = true;
 
-	// Set regardless of GPU being used, it will simply get ignored on other vendors
-	if (GVulkanAMDCompatibilityMode && FPlatformMisc::GetEnvironmentVariable(TEXT("RADV_DEBUG")).IsEmpty())
-	{
-		// Force compiler backend to llvm for better compatibility with our ray tracing pipeline shaders at time of release 5.5 (with Mesa 24.0.9)
-		FPlatformMisc::SetEnvironmentVar(TEXT("RADV_DEBUG"), TEXT("llvm"));
-	}
-	else
-	{
-		UE_LOG(LogVulkanRHI, Display, TEXT("Found existing RADV_DEBUG, it will not be overwritten."));
-	}
-
 #if VULKAN_HAS_DEBUGGING_ENABLED
 	const FString VulkanSDK = FPlatformMisc::GetEnvironmentVariable(TEXT("VULKAN_SDK"));
 	UE_LOG(LogVulkanRHI, Display, TEXT("Found VULKAN_SDK=%s"), *VulkanSDK);
diff --git a/Engine/Source/Runtime/VulkanRHI/Private/VulkanDevice.h b/Engine/Source/Runtime/VulkanRHI/Private/VulkanDevice.h
--- a/Engine/Source/Runtime/VulkanRHI/Private/VulkanDevice.h
+++ b/Engine/Source/Runtime/VulkanRHI/Private/VulkanDevice.h
@@ -130,6 +130,8 @@ struct FOptionalVulkanDeviceExtensionProperties
 		FMemory::Memzero(*this);
 	}
 
+	VkPhysicalDeviceDriverPropertiesKHR PhysicalDeviceDriverProperties;
+
 	VkPhysicalDeviceDescriptorBufferPropertiesEXT DescriptorBufferProps;
 	VkPhysicalDeviceSubgroupSizeControlPropertiesEXT SubgroupSizeControlProperties;
 
@@ -164,6 +166,7 @@ private:
 	VkPhysicalDeviceVulkan13Features Core_1_3;
 
 	friend class FVulkanDevice;
+	friend FVulkanDynamicRHI;
 };
 
 
diff --git a/Engine/Source/Runtime/VulkanRHI/Private/VulkanExtensions.cpp b/Engine/Source/Runtime/VulkanRHI/Private/VulkanExtensions.cpp
--- a/Engine/Source/Runtime/VulkanRHI/Private/VulkanExtensions.cpp
+++ b/Engine/Source/Runtime/VulkanRHI/Private/VulkanExtensions.cpp
@@ -243,19 +243,18 @@ public:
 
 	virtual void PrePhysicalDeviceProperties(VkPhysicalDeviceProperties2KHR& PhysicalDeviceProperties2) override final
 	{
+		VkPhysicalDeviceDriverPropertiesKHR& PhysicalDeviceDriverProperties = GetDeviceExtensionProperties().PhysicalDeviceDriverProperties;
 		ZeroVulkanStruct(PhysicalDeviceDriverProperties, VK_STRUCTURE_TYPE_PHYSICAL_DEVICE_DRIVER_PROPERTIES_KHR);
 		AddToPNext(PhysicalDeviceProperties2, PhysicalDeviceDriverProperties);
 	}
 
 	virtual void PostPhysicalDeviceProperties() override final
 	{
+		const VkPhysicalDeviceDriverPropertiesKHR& PhysicalDeviceDriverProperties = GetDeviceExtensionProperties().PhysicalDeviceDriverProperties;
 		UE_LOG(LogVulkanRHI, Display, TEXT("Vulkan Physical Device Driver Properties:"));
 		UE_LOG(LogVulkanRHI, Display, TEXT("- driverName: %s"), ANSI_TO_TCHAR(PhysicalDeviceDriverProperties.driverName));
 		UE_LOG(LogVulkanRHI, Display, TEXT("- driverInfo: %s"), ANSI_TO_TCHAR(PhysicalDeviceDriverProperties.driverInfo));
 	}
-
-private:
-	VkPhysicalDeviceDriverPropertiesKHR PhysicalDeviceDriverProperties;
 };
 
 
diff --git a/Engine/Source/Runtime/VulkanRHI/Private/VulkanRHI.cpp b/Engine/Source/Runtime/VulkanRHI/Private/VulkanRHI.cpp
--- a/Engine/Source/Runtime/VulkanRHI/Private/VulkanRHI.cpp
+++ b/Engine/Source/Runtime/VulkanRHI/Private/VulkanRHI.cpp
@@ -856,8 +856,19 @@ void FVulkanDynamicRHI::InitInstance()
 		// Use this compatibility mode avoid known issues at launch time with latest drivers at the time of release 5.5.  This will:
 		// - disable inline ray tracing and use ray tracing pipelines everywhere (instead of a mix of both)
 		// - disable mesh shaders until issues can be resolved (holes in Nanite meshes)
-		// - force llvm compiler backend on Linux (see VulkanLinuxPlatform.cpp) to circumvent raytracing pipeline compilation crash
+		// - disable ray tracing for RADV driver prior to version 25 in Linux (raytracing pipeline compilation crash)
+
+		// :todo-jn: to be removed when the official minimum RADV version is set to 25.0.0 in BaseHardware.ini
 		const bool bUseAMDCompatibilityMode = GVulkanAMDCompatibilityMode && (Device->GetVendorId() == EGpuVendorId::Amd);
+		if (GVulkanAMDCompatibilityMode && 
+			(Device->GetOptionalExtensionProperties().PhysicalDeviceDriverProperties.driverID == VK_DRIVER_ID_MESA_RADV) &&
+			(Props.driverVersion < VK_MAKE_VERSION(25, 0, 0)))
+		{
+			GRHISupportsRayTracing = false;
+			GRHISupportsRayTracingShaders = false;
+			GRHISupportsInlineRayTracing = false;
+			UE_LOG(LogVulkanRHI, Warning, TEXT("Using MESA RADV version prior to 25.X, ray tracing disabled."));
+		}
 
 		if (GRHISupportsRayTracing)
 		{
@@ -877,6 +888,7 @@ void FVulkanDynamicRHI::InitInstance()
 			// Loose parameters are always placed in the shader record after the FVulkanHitGroupSystemParameters in Vulkan (see VulkanRayTracing.h and VulkanCommon.ush)
 			GRHIGlobals.RayTracing.SupportsLooseParamsInShaderRecord = true;
 		}
+
 #if VULKAN_ENABLE_DUMP_LAYER
 		// Disable RHI thread by default if the dump layer is enabled
 		GRHISupportsRHIThread = false;
