VulkanRHI: Add GraphEvent to trigger CPU work on payload completion.  Start adding submission and interrupt threads (disabled by default).  Replace all uses of GetFenceSignaledCounter() by GraphEvents.  Add EVulkanQueueType and store queues in an array.

#rnx

diff --git a/Engine/Source/Runtime/VulkanRHI/Private/VulkanBuffer.cpp b/Engine/Source/Runtime/VulkanRHI/Private/VulkanBuffer.cpp
--- a/Engine/Source/Runtime/VulkanRHI/Private/VulkanBuffer.cpp
+++ b/Engine/Source/Runtime/VulkanRHI/Private/VulkanBuffer.cpp
@@ -399,7 +399,7 @@ void FVulkanBuffer::Unlock(FRHICommandListBase& RHICmdList)
 				VkMemoryBarrier BarrierAfter = { VK_STRUCTURE_TYPE_MEMORY_BARRIER, nullptr, VK_ACCESS_TRANSFER_WRITE_BIT, VK_ACCESS_MEMORY_READ_BIT | VK_ACCESS_MEMORY_WRITE_BIT };
 				VulkanRHI::vkCmdPipelineBarrier(CommandBufferHandle, VK_PIPELINE_STAGE_TRANSFER_BIT, VK_PIPELINE_STAGE_ALL_COMMANDS_BIT, 0, 1, &BarrierAfter, 0, nullptr, 0, nullptr);
 
-				Buffer->GetParent()->GetStagingManager().ReleaseBuffer(&CommandBuffer, StagingBuffer);
+				Buffer->GetParent()->GetStagingManager().ReleaseBuffer(&CommandBuffer.GetOwner()->GetContext(), StagingBuffer);
 			}
 		});
 	}
diff --git a/Engine/Source/Runtime/VulkanRHI/Private/VulkanCommandBuffer.cpp b/Engine/Source/Runtime/VulkanRHI/Private/VulkanCommandBuffer.cpp
--- a/Engine/Source/Runtime/VulkanRHI/Private/VulkanCommandBuffer.cpp
+++ b/Engine/Source/Runtime/VulkanRHI/Private/VulkanCommandBuffer.cpp
@@ -248,8 +248,6 @@ void FVulkanCommandBuffer::ResetFenceStatus()
 		CurrentScissors.Empty();
 		CurrentStencilRef = 0;
 
-		++FenceSignaledCounter;
-
 		if (CurrentDescriptorPoolSetContainer)
 		{
 			//#todo-rco: Reset here?
@@ -316,29 +314,12 @@ void FVulkanCommandBufferPool::FreeUnusedCmdBuffers(FVulkanQueue* InQueue, bool
 
 	const double CurrentTime = FPlatformTime::Seconds();
 
-	// In case Queue stores pointer to a cmdbuffer, do not delete it
-	FVulkanCommandBuffer* LastSubmittedCmdBuffer = nullptr;
-	uint64 LastSubmittedFenceCounter = 0;
-	InQueue->GetLastSubmittedInfo(LastSubmittedCmdBuffer, LastSubmittedFenceCounter);
-
-	// Deferred deletion queue caches pointers to cmdbuffers
-	VulkanRHI::FDeferredDeletionQueue2& DeferredDeletionQueue = Device.GetDeferredDeletionQueue();
-
 	for (int32 Index = CmdBuffers.Num() - 1; Index >= 0; --Index)
 	{
 		FVulkanCommandBuffer* CmdBuffer = CmdBuffers[Index];
-		if (CmdBuffer != LastSubmittedCmdBuffer &&
-			(CmdBuffer->State == FVulkanCommandBuffer::EState::ReadyForBegin || CmdBuffer->State == FVulkanCommandBuffer::EState::NeedReset) &&
+		if ((CmdBuffer->State == FVulkanCommandBuffer::EState::ReadyForBegin || CmdBuffer->State == FVulkanCommandBuffer::EState::NeedReset) &&
 			((CurrentTime - CmdBuffer->SubmittedTime) > CMD_BUFFER_TIME_TO_WAIT_BEFORE_DELETING))
 		{
-			// Skip command buffer that contain unresolved compaction queries
-			if (Device.GetRayTracingCompactionRequestHandler() && Device.GetRayTracingCompactionRequestHandler()->IsUsingCmdBuffer(CmdBuffer))
-			{
-				continue;
-			}
-
-			DeferredDeletionQueue.OnCmdBufferDeleted(CmdBuffer);
-
 			CmdBuffer->FreeMemory();
 			CmdBuffers.RemoveAtSwap(Index, EAllowShrinking::No);
 			FreeCmdBuffers.Add(CmdBuffer);
diff --git a/Engine/Source/Runtime/VulkanRHI/Private/VulkanCommandBuffer.h b/Engine/Source/Runtime/VulkanRHI/Private/VulkanCommandBuffer.h
--- a/Engine/Source/Runtime/VulkanRHI/Private/VulkanCommandBuffer.h
+++ b/Engine/Source/Runtime/VulkanRHI/Private/VulkanCommandBuffer.h
@@ -77,16 +77,6 @@ public:
 		return CommandBufferHandle;
 	}
 
-	volatile uint64 GetFenceSignaledCounter() const
-	{
-		return FenceSignaledCounter;
-	}
-
-	volatile uint64 GetSubmittedTimelineValue() const
-	{
-		return SubmittedTimelineValue;
-	}
-
 	FVulkanLayoutManager& GetLayoutManager()
 	{
 		return LayoutManager;
@@ -145,9 +135,6 @@ private:
 
 	TArray<PendingQuery> PendingTimestampQueries;
 
-	volatile uint64 FenceSignaledCounter = 0;
-	volatile uint64 SubmittedTimelineValue = 0;
-
 	void ResetFenceStatus();
 
 	FVulkanCommandBufferPool* CommandBufferPool;
diff --git a/Engine/Source/Runtime/VulkanRHI/Private/VulkanCommands.cpp b/Engine/Source/Runtime/VulkanRHI/Private/VulkanCommands.cpp
--- a/Engine/Source/Runtime/VulkanRHI/Private/VulkanCommands.cpp
+++ b/Engine/Source/Runtime/VulkanRHI/Private/VulkanCommands.cpp
@@ -904,11 +904,8 @@ void FVulkanCommandListContext::RHICopyToStagingBuffer(FRHIBuffer* SourceBufferR
 
 void FVulkanCommandListContext::RHIWriteGPUFence(FRHIGPUFence* FenceRHI)
 {
-	FVulkanCommandBuffer* CmdBuffer = GetActiveCmdBuffer();
 	FVulkanGPUFence* Fence = ResourceCast(FenceRHI);
-
-	Fence->CmdBuffer = CmdBuffer;
-	Fence->FenceSignaledCounter = CmdBuffer->GetFenceSignaledCounter();
+	Fence->SyncPoint = AddSyncPoint();
 }
 
 IRHIComputeContext* FVulkanDynamicRHI::RHIGetCommandContext(ERHIPipeline Pipeline, FRHIGPUMask GPUMask)
diff --git a/Engine/Source/Runtime/VulkanRHI/Private/VulkanContext.cpp b/Engine/Source/Runtime/VulkanRHI/Private/VulkanContext.cpp
--- a/Engine/Source/Runtime/VulkanRHI/Private/VulkanContext.cpp
+++ b/Engine/Source/Runtime/VulkanRHI/Private/VulkanContext.cpp
@@ -33,36 +33,51 @@ FVulkanContextCommon::~FVulkanContextCommon()
 
 void FVulkanContextCommon::NewPayload()
 {
+	EndPayload();
 	Payloads.Add(new FVulkanPayload(Queue));
 	CurrentPhase = EPhase::Wait;
 }
 
+void FVulkanContextCommon::EndPayload()
+{
+	if ((Payloads.Num() > 0) && (Payloads.Last()->CommandBuffers.Num() > 0))
+	{
+		FVulkanCommandBuffer* CommandBuffer = Payloads.Last()->CommandBuffers.Last();
+		checkSlow(!CommandBuffer->IsSubmitted() && CommandBuffer->HasBegun());
+
+		if (!CommandBuffer->IsOutsideRenderPass())
+		{
+			UE_LOG(LogVulkanRHI, Warning, TEXT("Forcing EndRenderPass() for submission"));
+			CommandBuffer->EndRenderPass();
+		}
+
+		CommandBuffer->End();
+	}
+}
+
 // Complete recording of the current command list set, and appends the resulting
 // payloads to the given array. Resets the context so new commands can be recorded.
 void FVulkanContextCommon::Finalize(TArray<FVulkanPayload*>& OutPayloads)
 {
+	EndPayload();
+
 	if (FVulkanPayload* ResetPayload = FlushResetQueryPools())
 	{
 		OutPayloads.Add(ResetPayload);
 	}
 
-	for (FVulkanPayload* Payload : Payloads)
+	OutPayloads.Append(MoveTemp(Payloads));
+
+	if (ContextSyncPoint.IsValid())
 	{
-		for (FVulkanCommandBuffer* CommandBuffer : Payload->CommandBuffers)
+		// If the last payload has no space for the ContextSyncPoint, then add a dummy payload
+		if ((OutPayloads.Num() == 0) || OutPayloads.Last()->SyncPoint.IsValid())
 		{
-			checkSlow(!CommandBuffer->IsSubmitted() && CommandBuffer->HasBegun());
-
-			if (!CommandBuffer->IsOutsideRenderPass())
-			{
-				UE_LOG(LogVulkanRHI, Warning, TEXT("Forcing EndRenderPass() for submission"));
-				CommandBuffer->EndRenderPass();
-			}
-
-			CommandBuffer->End();
+			OutPayloads.Add(new FVulkanPayload(Queue));
 		}
+		OutPayloads.Last()->SyncPoint = ContextSyncPoint;
+		ContextSyncPoint = nullptr;
 	}
-
-	OutPayloads.Append(MoveTemp(Payloads));
 }
 
 void FVulkanContextCommon::AddQueryPoolForReset(VkQueryPool QueryPool, uint32 Size)
@@ -77,7 +92,7 @@ FVulkanPayload* FVulkanContextCommon::FlushResetQueryPools()
 	FScopeLock ScopeLock(&Pool.CS);
 	if (PoolResets.Num())
 	{
-		FVulkanCommandBuffer* PoolResetCommandBuffer = PrepareNewCommandBuffer();
+		FVulkanCommandBuffer* PoolResetCommandBuffer = PrepareNewCommandBuffer(false);
 		const VkCommandBuffer CommandBufferHandle = PoolResetCommandBuffer->GetHandle();
 		for (FQueryPoolReset& QPR : PoolResets)
 		{
@@ -141,7 +156,7 @@ void FVulkanContextCommon::FreeUnusedCmdBuffers(bool bTrimMemory)
 #endif
 }
 
-FVulkanCommandBuffer* FVulkanContextCommon::PrepareNewCommandBuffer()
+FVulkanCommandBuffer* FVulkanContextCommon::PrepareNewCommandBuffer(bool bTransferLayoutMgr)
 {
 	FScopeLock ScopeLock(&Pool.CS);
 
@@ -153,17 +168,40 @@ FVulkanCommandBuffer* FVulkanContextCommon::PrepareNewCommandBuffer()
 		if (CmdBuffer->State == FVulkanCommandBuffer::EState::ReadyForBegin || CmdBuffer->State == FVulkanCommandBuffer::EState::NeedReset)
 		{
 			NewCommandBuffer = CmdBuffer;
-			NewCommandBuffer->Begin();
-			return NewCommandBuffer;
+			break;
 		}
 		else
 		{
-			check(CmdBuffer->State == FVulkanCommandBuffer::EState::Submitted);
+			check(CmdBuffer->IsSubmitted() || CmdBuffer->HasEnded());
+		}
+	}
+
+	// All cmd buffers are being executed still, create a new one
+	if (!NewCommandBuffer)
+	{
+		NewCommandBuffer = Pool.Create();
+	}
+
+	// Move the layout mgr forward to the new commandbuffer
+	// NOTE: LayoutMgr should be in context, but it will get removed soon anyhow
+	if (bTransferLayoutMgr && Payloads.Num())
+	{
+		FVulkanCommandBuffer* LastCommandBuffer = nullptr;
+		for (int32 PayloadIndex = Payloads.Num() - 1; PayloadIndex >= 0; --PayloadIndex)
+		{
+			if (Payloads[PayloadIndex]->CommandBuffers.Num())
+			{
+				LastCommandBuffer = Payloads[PayloadIndex]->CommandBuffers.Last();
+				break;
+			}
+		}
+
+		if (LastCommandBuffer)
+		{
+			LastCommandBuffer->GetLayoutManager().TransferTo(NewCommandBuffer->GetLayoutManager());
 		}
 	}
 
-	// All cmd buffers are being executed still
-	NewCommandBuffer = Pool.Create();
 	NewCommandBuffer->Begin();
 	return NewCommandBuffer;
 }
diff --git a/Engine/Source/Runtime/VulkanRHI/Private/VulkanContext.h b/Engine/Source/Runtime/VulkanRHI/Private/VulkanContext.h
--- a/Engine/Source/Runtime/VulkanRHI/Private/VulkanContext.h
+++ b/Engine/Source/Runtime/VulkanRHI/Private/VulkanContext.h
@@ -61,7 +61,7 @@ public:
 
 		if (Payload.CommandBuffers.Num() == 0)
 		{
-			Payload.CommandBuffers.Add(PrepareNewCommandBuffer());
+			Payload.CommandBuffers.Add(PrepareNewCommandBuffer(true));
 		}
 
 		return *Payload.CommandBuffers.Last();
@@ -134,14 +134,36 @@ public:
 	FVulkanDevice& Device;
 	FVulkanQueue& Queue;
 
+	// Returns a sync point that will be triggered at the end of the current payload's accumulated work
+	FVulkanSyncPoint* AddSyncPoint()
+	{
+		FVulkanPayload& Payload = GetPayload(EPhase::Signal);
+		if (!Payload.SyncPoint.IsValid())
+		{
+			Payload.SyncPoint = CreateVulkanSyncPoint();
+		}
+		return Payload.SyncPoint;
+	}
+
+	// Returns a single sync point for the context that will be inserted when it is finalized
+	FVulkanSyncPoint* GetContextSyncPoint()
+	{
+		if (!ContextSyncPoint)
+		{
+			ContextSyncPoint = CreateVulkanSyncPoint();
+		}
+		return ContextSyncPoint;
+	}
+
 protected:
 	// Number of times EndFrame() has been called on this context
 	uint64 FrameCounter = 0;
 
 private:
-	FVulkanCommandBuffer* PrepareNewCommandBuffer();
+	FVulkanCommandBuffer* PrepareNewCommandBuffer(bool bTransferLayoutMgr);
 	FVulkanPayload* FlushResetQueryPools();
 	void NewPayload();
+	void EndPayload();
 
 	FVulkanCommandBufferPool Pool;
 	struct FQueryPoolReset
@@ -152,8 +174,9 @@ private:
 	TArray<FQueryPoolReset> PoolResets;
 
 	TArray<FVulkanPayload*> Payloads;
-
 	EPhase CurrentPhase = EPhase::Wait;
+
+	FVulkanSyncPointRef ContextSyncPoint;
 };
 
 
@@ -328,7 +351,7 @@ protected:
 	FVulkanCommandListContext* Immediate;
 
 	void BeginOcclusionQueryBatch(FVulkanCommandBuffer* CmdBuffer, uint32 NumQueriesInBatch);
-	void EndOcclusionQueryBatch(FVulkanCommandBuffer* CmdBuffer);
+	void EndOcclusionQueryBatch();
 
 	VulkanRHI::FTempBlockAllocator* TempBlockAllocator = nullptr;
 
diff --git a/Engine/Source/Runtime/VulkanRHI/Private/VulkanDescriptorSets.cpp b/Engine/Source/Runtime/VulkanRHI/Private/VulkanDescriptorSets.cpp
--- a/Engine/Source/Runtime/VulkanRHI/Private/VulkanDescriptorSets.cpp
+++ b/Engine/Source/Runtime/VulkanRHI/Private/VulkanDescriptorSets.cpp
@@ -621,7 +621,7 @@ void FVulkanBindlessDescriptorManager::RegisterUniformBuffers(FVulkanCommandBuff
 	VkDeviceSize BufferOffsets[VulkanBindless::MaxNumSets];
 	FMemory::Memzero(BufferOffsets);
 
-	uint8* MappedPointer = SingleUseUBAllocator->Alloc(BlockSize, &CommandBuffer, 
+	uint8* MappedPointer = SingleUseUBAllocator->Alloc(BlockSize, CommandBuffer.GetOwner()->GetContext(),
 		LocalBufferBindingInfo[VulkanBindless::BindlessSingleUseUniformBufferSet], BufferOffsets[VulkanBindless::BindlessSingleUseUniformBufferSet]);
 
 	for (int32 StageIndex = 0; StageIndex < ShaderStage::MaxNumStages; ++StageIndex)
@@ -754,7 +754,7 @@ void FVulkanBindlessDescriptorManager::UpdateDescriptor(FRHIDescriptorHandle Des
 				MemoryBarrier.dstAccessMask = VK_ACCESS_2_DESCRIPTOR_BUFFER_READ_BIT_EXT;
 				VulkanRHI::vkCmdPipelineBarrier2KHR(CmdBuffer->GetHandle(), &DependencyInfo);
 			}
-			Device->GetStagingManager().ReleaseBuffer(CmdBuffer, StagingBuffer);
+			Device->GetStagingManager().ReleaseBuffer(&Device->GetImmediateContext(), StagingBuffer);
 		}
 	}
 
diff --git a/Engine/Source/Runtime/VulkanRHI/Private/VulkanDevice.cpp b/Engine/Source/Runtime/VulkanRHI/Private/VulkanDevice.cpp
--- a/Engine/Source/Runtime/VulkanRHI/Private/VulkanDevice.cpp
+++ b/Engine/Source/Runtime/VulkanRHI/Private/VulkanDevice.cpp
@@ -214,15 +214,12 @@ FVulkanDevice::FVulkanDevice(FVulkanDynamicRHI* InRHI, VkPhysicalDevice InGpu)
 	, DeferredDeletionQueue(this)
 	, DefaultSampler(nullptr)
 	, Gpu(InGpu)
-	, GfxQueue(nullptr)
-	, ComputeQueue(nullptr)
-	, TransferQueue(nullptr)
-	, PresentQueue(nullptr)
 	, ImmediateContext(nullptr)
 	, ComputeContext(nullptr)
 	, PipelineStateCache(nullptr)
 {
 	RHI = InRHI;
+	FMemory::Memzero(Queues);
 	FMemory::Memzero(GpuProps);
 	FMemory::Memzero(FormatProperties);
 	FMemory::Memzero(PixelFormatComponentMapping);
@@ -399,7 +396,7 @@ void FVulkanDevice::CreateDevice(TArray<const ANSICHAR*>& DeviceLayers, FVulkanD
 	FVulkanPlatform::NotifyFoundDeviceLayersAndExtensions(Gpu, DeviceLayers, DeviceExtensions);
 
 	// Create Graphics Queue, here we submit command buffers for execution
-	GfxQueue = new FVulkanQueue(*this, GfxQueueFamilyIndex);
+	Queues[(int32)EVulkanQueueType::Graphics] = new FVulkanQueue(*this, GfxQueueFamilyIndex, EVulkanQueueType::Graphics);
 	if (ComputeQueueFamilyIndex == -1)
 	{
 		// If we didn't find a dedicated Queue, use the default one
@@ -413,13 +410,13 @@ void FVulkanDevice::CreateDevice(TArray<const ANSICHAR*>& DeviceLayers, FVulkanD
 			bAsyncComputeQueue = true;
 		}
 	}
-	ComputeQueue = new FVulkanQueue(*this, ComputeQueueFamilyIndex);
+	Queues[(int32)EVulkanQueueType::AsyncCompute] = new FVulkanQueue(*this, ComputeQueueFamilyIndex, EVulkanQueueType::AsyncCompute);
 	if (TransferQueueFamilyIndex == -1)
 	{
 		// If we didn't find a dedicated Queue, use the default one
 		TransferQueueFamilyIndex = ComputeQueueFamilyIndex;
 	}
-	TransferQueue = new FVulkanQueue(*this, TransferQueueFamilyIndex);
+	Queues[(int32)EVulkanQueueType::Transfer] = new FVulkanQueue(*this, TransferQueueFamilyIndex, EVulkanQueueType::Transfer);
 
 	uint64 NumBits = QueueFamilyProps[GfxQueueFamilyIndex].timestampValidBits;
 	if (NumBits > 0)
@@ -1376,11 +1373,11 @@ void FVulkanDevice::InitGPU()
 	// always look in the saved directory (for the cache from previous run that wasn't moved over to stage directory)
 	CacheFilenames.Add(VulkanRHI::GetPipelineCacheFilename());
 
-	ImmediateContext = new FVulkanCommandListContextImmediate(GVulkanRHI, *this, *GfxQueue);
+	ImmediateContext = new FVulkanCommandListContextImmediate(GVulkanRHI, *this, *GetGraphicsQueue());
 
-	if (GfxQueue->GetFamilyIndex() != ComputeQueue->GetFamilyIndex() && GRHIAllowAsyncComputeCvar.GetValueOnAnyThread() != 0)
+	if (GetGraphicsQueue()->GetFamilyIndex() != GetComputeQueue()->GetFamilyIndex() && GRHIAllowAsyncComputeCvar.GetValueOnAnyThread() != 0)
 	{
-		ComputeContext = new FVulkanCommandListContextImmediate(GVulkanRHI, *this, *ComputeQueue);
+		ComputeContext = new FVulkanCommandListContextImmediate(GVulkanRHI, *this, *GetComputeQueue());
 	}
 	else
 	{
@@ -1393,7 +1390,7 @@ void FVulkanDevice::InitGPU()
 		int32 Num = FTaskGraphInterface::Get().GetNumWorkerThreads();
 		for (int32 Index = 0; Index < Num; Index++)
 		{
-			FVulkanCommandListContext* CmdContext = new FVulkanCommandListContext(GVulkanRHI, *this, *GfxQueue, ImmediateContext);
+			FVulkanCommandListContext* CmdContext = new FVulkanCommandListContext(GVulkanRHI, *this, *GetGraphicsQueue(), ImmediateContext);
 			CommandContexts.Add(CmdContext);
 		}
 	}
@@ -1528,12 +1525,11 @@ void FVulkanDevice::Destroy()
 
 	MemoryManager.Deinit();
 
-	delete TransferQueue;
-	TransferQueue = nullptr;
-	delete ComputeQueue;
-	ComputeQueue = nullptr;
-	delete GfxQueue;
-	GfxQueue = nullptr;
+	for (int32 QueueIndex = 0; QueueIndex < (int32)EVulkanQueueType::Count; ++QueueIndex)
+	{
+		delete Queues[QueueIndex];
+		Queues[QueueIndex] = nullptr;
+	}
 
 	FenceManager.Deinit();
 	DeviceMemoryManager.Deinit();
@@ -1624,7 +1620,7 @@ FVulkanCommandListContext* FVulkanDevice::AcquireDeferredContext()
 	FScopeLock Lock(&GContextCS);
 	if (CommandContexts.Num() == 0)
 	{
-		return new FVulkanCommandListContext(GVulkanRHI, *this, *GfxQueue, ImmediateContext);
+		return new FVulkanCommandListContext(GVulkanRHI, *this, *GetGraphicsQueue(), ImmediateContext);
 	}
 	return CommandContexts.Pop(EAllowShrinking::No);
 }
@@ -1694,7 +1690,8 @@ FGPUTimingCalibrationTimestamp FVulkanDevice::GetCalibrationTimestamp()
 
 void FVulkanDevice::ForEachQueue(TFunctionRef<void(FVulkanQueue&)> Callback)
 {
-	Callback(*GfxQueue);
-	Callback(*ComputeQueue);
-	Callback(*TransferQueue);
+	for (int32 QueueIndex = 0; QueueIndex < (int32)EVulkanQueueType::Count; ++QueueIndex)
+	{
+		Callback(*Queues[QueueIndex]);
+	}
 }
diff --git a/Engine/Source/Runtime/VulkanRHI/Private/VulkanDevice.h b/Engine/Source/Runtime/VulkanRHI/Private/VulkanDevice.h
--- a/Engine/Source/Runtime/VulkanRHI/Private/VulkanDevice.h
+++ b/Engine/Source/Runtime/VulkanRHI/Private/VulkanDevice.h
@@ -8,6 +8,7 @@
 
 #include "VulkanMemory.h"
 #include "VulkanResources.h"
+#include "VulkanQueue.h"
 #include "GPUProfiler.h"
 
 class FVulkanSamplerState;
@@ -229,7 +230,6 @@ namespace VulkanRHI
 			ReleaseResources(true);
 		}
 
-		void OnCmdBufferDeleted(FVulkanCommandBuffer* CmdBuffer);
 	private:
 		void EnqueueGenericResource(EType Type, uint64 Handle);
 
@@ -237,13 +237,13 @@ namespace VulkanRHI
 		{
 			EType StructureType;
 			uint32 FrameNumber;
-			uint64 FenceCounter;
-			FVulkanCommandBuffer* CmdBuffer;
-
 			uint64 Handle;
 			FVulkanAllocation Allocation;
 			FDeviceMemoryAllocation* DeviceMemoryAllocation;
 		};
+
+		void ReleaseResourcesImmediately(const TArray<FEntry>& InEntries);
+
 		FCriticalSection CS;
 		TArray<FEntry> Entries;
 	};
@@ -295,17 +295,17 @@ public:
 
 	inline FVulkanQueue* GetGraphicsQueue()
 	{
-		return GfxQueue;
+		return Queues[(int32)EVulkanQueueType::Graphics];
 	}
 
 	inline FVulkanQueue* GetComputeQueue()
 	{
-		return ComputeQueue;
+		return Queues[(int32)EVulkanQueueType::AsyncCompute];
 	}
 
 	inline FVulkanQueue* GetTransferQueue()
 	{
-		return TransferQueue;
+		return Queues[(int32)EVulkanQueueType::Transfer];
 	}
 
 	inline FVulkanQueue* GetPresentQueue()
@@ -594,10 +594,8 @@ private:
 
 	uint64 TimestampValidBitsMask = 0;
 
-	FVulkanQueue* GfxQueue;
-	FVulkanQueue* ComputeQueue;
-	FVulkanQueue* TransferQueue;
-	FVulkanQueue* PresentQueue;
+	FVulkanQueue* Queues[(int32)EVulkanQueueType::Count];
+	FVulkanQueue* PresentQueue = nullptr;  // points to an existing queue
 	bool bAsyncComputeQueue = false;
 	bool bPresentOnComputeQueue = false;
 
diff --git a/Engine/Source/Runtime/VulkanRHI/Private/VulkanMemory.cpp b/Engine/Source/Runtime/VulkanRHI/Private/VulkanMemory.cpp
--- a/Engine/Source/Runtime/VulkanRHI/Private/VulkanMemory.cpp
+++ b/Engine/Source/Runtime/VulkanRHI/Private/VulkanMemory.cpp
@@ -25,7 +25,15 @@ THIRD_PARTY_INCLUDES_END
 
 // This 'frame number' should only be used for the deletion queue
 uint32 GVulkanRHIDeletionFrameNumber = 0;
-const uint32 NUM_FRAMES_TO_WAIT_FOR_RESOURCE_DELETE = 2;
+
+
+static int32 GVulkanNumFramesToWaitForResourceDelete = 2;
+static FAutoConsoleVariableRef CVarVulkanNumFramesToWaitForResourceDelete(
+	TEXT("r.Vulkan.NumFramesToWaitForResourceDelete"),
+	GVulkanNumFramesToWaitForResourceDelete,
+	TEXT("The number of frames to wait before deleting a resource. Used for debugging. (default:2)\n"),
+	ECVF_ReadOnly
+);
 
 #define UE_VK_MEMORY_MAX_SUB_ALLOCATION (64llu << 20llu) // set to 0 to disable
 
@@ -1992,7 +2000,8 @@ namespace VulkanRHI
 		}
 		if (!DeviceMemoryAllocation)
 		{
-			UE_LOG(LogVulkanRHI, Fatal, TEXT("Out of memory on Vulkan; MemoryTypeIndex=%d, AllocSize=%0.3fMB"), MemoryTypeIndex, (float)AllocationSize / 1048576.0f);
+			UE_LOG(LogVulkanRHI, Warning, TEXT("Out of memory on Vulkan; MemoryTypeIndex=%d, AllocSize=%0.3fMB"), MemoryTypeIndex, (float)AllocationSize / 1048576.0f);
+			return false;
 		}
 		if (bMapAllocation)
 		{
@@ -2972,6 +2981,7 @@ namespace VulkanRHI
 		SCOPED_NAMED_EVENT(FResourceHeapManager_AllocateBufferMemory, FColor::Cyan);
 		uint32 TypeIndex = 0;
 		VkResult Result = DeviceMemoryManager->GetMemoryTypeFromProperties(MemoryReqs.memoryTypeBits, MemoryPropertyFlags, &TypeIndex);
+		const bool bPreferBAR = VKHasAllFlags(MemoryPropertyFlags, VK_MEMORY_PROPERTY_HOST_VISIBLE_BIT | VK_MEMORY_PROPERTY_HOST_COHERENT_BIT | VK_MEMORY_PROPERTY_DEVICE_LOCAL_BIT);
 		bool bMapped = VKHasAllFlags(MemoryPropertyFlags, VK_MEMORY_PROPERTY_HOST_VISIBLE_BIT);
 		if ((Result != VK_SUCCESS) || !ResourceTypeHeaps[TypeIndex])
 		{
@@ -2987,7 +2997,7 @@ namespace VulkanRHI
 				MemoryPropertyFlags = MemoryPropertyFlags & ~VK_MEMORY_PROPERTY_LAZILY_ALLOCATED_BIT;
 			}
 
-			if (VKHasAllFlags(MemoryPropertyFlags, VK_MEMORY_PROPERTY_HOST_VISIBLE_BIT | VK_MEMORY_PROPERTY_HOST_COHERENT_BIT | VK_MEMORY_PROPERTY_DEVICE_LOCAL_BIT))
+			if (bPreferBAR)
 			{
 				// Try regular host memory if local+host is not available
 				MemoryPropertyFlags = MemoryPropertyFlags & ~VK_MEMORY_PROPERTY_DEVICE_LOCAL_BIT;
@@ -3012,6 +3022,11 @@ namespace VulkanRHI
 		if (!ResourceTypeHeaps[TypeIndex]->AllocateResource(OutAllocation, AllocationOwner, EType::Buffer, MemoryReqs.size, MemoryReqs.alignment, bMapped, bForceSeparateAllocation, MetaType, bExternal, File, Line))
 		{
 			// Try another memory type if the allocation failed
+			if (bPreferBAR)
+			{
+				// Try regular host memory if local+host is not available
+				MemoryPropertyFlags = MemoryPropertyFlags & ~VK_MEMORY_PROPERTY_DEVICE_LOCAL_BIT;
+			}
 			VERIFYVULKANRESULT(DeviceMemoryManager->GetMemoryTypeFromPropertiesExcluding(MemoryReqs.memoryTypeBits, MemoryPropertyFlags, TypeIndex, &TypeIndex));
 			bMapped = (MemoryPropertyFlags & VK_MEMORY_PROPERTY_HOST_VISIBLE_BIT) == VK_MEMORY_PROPERTY_HOST_VISIBLE_BIT;
 			if (!ResourceTypeHeaps[TypeIndex])
@@ -4350,6 +4365,7 @@ namespace VulkanRHI
 					FreeStagingBuffers.RemoveAtSwap(Index, EAllowShrinking::No);
 					UsedStagingBuffers.Add(Buffer);
 					VULKAN_FILL_TRACK_INFO(Buffer->Track, __FILE__, __LINE__);
+					check(Buffer->GetHandle());
 					return Buffer;
 				}
 			}
@@ -4389,37 +4405,7 @@ namespace VulkanRHI
 		return StagingBuffer;
 	}
 
-	inline FStagingManager::FPendingItemsPerCmdBuffer* FStagingManager::FindOrAdd(FVulkanCommandBuffer* CmdBuffer)
-	{
-		for (int32 Index = 0; Index < PendingFreeStagingBuffers.Num(); ++Index)
-		{
-			if (PendingFreeStagingBuffers[Index].CmdBuffer == CmdBuffer)
-			{
-				return &PendingFreeStagingBuffers[Index];
-			}
-		}
-
-		FPendingItemsPerCmdBuffer& New = PendingFreeStagingBuffers.AddDefaulted_GetRef();
-		New.CmdBuffer = CmdBuffer;
-		return &New;
-	}
-
-	inline FStagingManager::FPendingItemsPerCmdBuffer::FPendingItems* FStagingManager::FPendingItemsPerCmdBuffer::FindOrAddItemsForFence(uint64 Fence)
-	{
-		for (int32 Index = 0; Index < PendingItems.Num(); ++Index)
-		{
-			if (PendingItems[Index].FenceCounter == Fence)
-			{
-				return &PendingItems[Index];
-			}
-		}
-
-		FPendingItems& New = PendingItems.AddDefaulted_GetRef();
-		New.FenceCounter = Fence;
-		return &New;
-	}
-
-	void FStagingManager::ReleaseBuffer(FVulkanCommandBuffer* CmdBuffer, FStagingBuffer*& StagingBuffer)
+	void FStagingManager::ReleaseBuffer(FVulkanContextCommon* Context, FStagingBuffer*& StagingBuffer)
 	{
 #if VULKAN_ENABLE_AGGRESSIVE_STATS
 		SCOPE_CYCLE_COUNTER(STAT_VulkanStagingBuffer);
@@ -4428,12 +4414,10 @@ namespace VulkanRHI
 		FScopeLock Lock(&StagingLock);
 		UsedStagingBuffers.RemoveSingleSwap(StagingBuffer, EAllowShrinking::No);
 
-		if (CmdBuffer)
+		if (Context)
 		{
-			FPendingItemsPerCmdBuffer* ItemsForCmdBuffer = FindOrAdd(CmdBuffer);
-			FPendingItemsPerCmdBuffer::FPendingItems* ItemsForFence = ItemsForCmdBuffer->FindOrAddItemsForFence(CmdBuffer->GetFenceSignaledCounter());
-			check(StagingBuffer);
-			ItemsForFence->Resources.Add(StagingBuffer);
+			TArray<FStagingBuffer*>& BufferArray = PendingFreeStagingBuffers.FindOrAdd(Context->GetContextSyncPoint());
+			BufferArray.Add(StagingBuffer);
 		}
 		else
 		{
@@ -4452,22 +4436,18 @@ namespace VulkanRHI
 			Stats.TotalMemory += Buffer->BufferSize;
 			Stats.MemoryFlags |= Buffer->MemoryReadFlags|VK_MEMORY_PROPERTY_HOST_VISIBLE_BIT;
 		}
-		for (int32 Index = 0; Index < PendingFreeStagingBuffers.Num(); ++Index)
+
+		for (auto& Pair : PendingFreeStagingBuffers)
 		{
-			FPendingItemsPerCmdBuffer& ItemPerCmdBuffer = PendingFreeStagingBuffers[Index];
-			for (int32 FenceIndex = 0; FenceIndex < ItemPerCmdBuffer.PendingItems.Num(); ++FenceIndex)
+			for (FStagingBuffer* StagingBuffer : Pair.Value)
 			{
-				FPendingItemsPerCmdBuffer::FPendingItems& ItemsPerFence = ItemPerCmdBuffer.PendingItems[FenceIndex];
-				for (int32 BufferIndex = 0; BufferIndex < ItemsPerFence.Resources.Num(); ++BufferIndex)
-				{
-					FStagingBuffer* Buffer = ItemsPerFence.Resources[BufferIndex];
-					Stats.BufferAllocations += 1;
-					Stats.UsedBufferMemory += Buffer->BufferSize;
-					Stats.TotalMemory += Buffer->BufferSize;
-					Stats.MemoryFlags |= Buffer->MemoryReadFlags|VK_MEMORY_PROPERTY_HOST_VISIBLE_BIT;
-				}
+				Stats.BufferAllocations += 1;
+				Stats.UsedBufferMemory += StagingBuffer->BufferSize;
+				Stats.TotalMemory += StagingBuffer->BufferSize;
+				Stats.MemoryFlags |= StagingBuffer->MemoryReadFlags | VK_MEMORY_PROPERTY_HOST_VISIBLE_BIT;
 			}
 		}
+
 		for (int32 Index = 0; Index < FreeStagingBuffers.Num(); ++Index)
 		{
 			FFreeEntry& Entry = FreeStagingBuffers[Index];
@@ -4488,20 +4468,14 @@ namespace VulkanRHI
 			VULKAN_LOGMEMORY(TEXT("%6d 0x%016llx 0x%016llx %6d"), Index, (void*)Buffer->GetHandle(), (void*)Buffer->Allocation.GetBufferHandle(), Buffer->BufferSize);
 		}
 
-		VULKAN_LOGMEMORY(TEXT("Pending CmdBuffer   Fence   BufferHandle    ResourceAllocation Size"));
-		for (int32 Index = 0; Index < PendingFreeStagingBuffers.Num(); ++Index)
+		VULKAN_LOGMEMORY(TEXT("Pending Fence    BufferHandle       ResourceAllocation Size"));
+		int32 PendingFreeIndex = 0;
+		for (auto& Pair : PendingFreeStagingBuffers)
 		{
-			FPendingItemsPerCmdBuffer& ItemPerCmdBuffer = PendingFreeStagingBuffers[Index];
-			VULKAN_LOGMEMORY(TEXT("%6d %p"), Index, (void*)ItemPerCmdBuffer.CmdBuffer->GetHandle());
-			for (int32 FenceIndex = 0; FenceIndex < ItemPerCmdBuffer.PendingItems.Num(); ++FenceIndex)
+			for (FStagingBuffer* StagingBuffer : Pair.Value)
 			{
-				FPendingItemsPerCmdBuffer::FPendingItems& ItemsPerFence = ItemPerCmdBuffer.PendingItems[FenceIndex];
-				VULKAN_LOGMEMORY(TEXT("         Fence %p"), (void*)ItemsPerFence.FenceCounter);
-				for (int32 BufferIndex = 0; BufferIndex < ItemsPerFence.Resources.Num(); ++BufferIndex)
-				{
-					FStagingBuffer* Buffer = ItemsPerFence.Resources[BufferIndex];
-					VULKAN_LOGMEMORY(TEXT("                   0x%016llx 0x%016llx %6d"), (void*)Buffer->GetHandle(), (void*)Buffer->Allocation.GetBufferHandle(), Buffer->BufferSize);
-				}
+				VULKAN_LOGMEMORY(TEXT("%6d  %p 0x%016llx 0x%016llx %6d"), PendingFreeIndex++, Pair.Key.GetReference(),
+					(void*)StagingBuffer->GetHandle(), (void*)StagingBuffer->Allocation.GetBufferHandle(), StagingBuffer->BufferSize);
 			}
 		}
 
@@ -4516,29 +4490,23 @@ namespace VulkanRHI
 
 	void FStagingManager::ProcessPendingFreeNoLock(bool bImmediately, bool bFreeToOS)
 	{
-		int32 NumOriginalFreeBuffers = FreeStagingBuffers.Num();
-		for (int32 Index = PendingFreeStagingBuffers.Num() - 1; Index >= 0; --Index)
+		const int32 NumOriginalFreeBuffers = FreeStagingBuffers.Num();
+
+		TArray<FVulkanSyncPoint*> Completed;
+		for (auto& Pair : PendingFreeStagingBuffers)
 		{
-			FPendingItemsPerCmdBuffer& EntriesPerCmdBuffer = PendingFreeStagingBuffers[Index];
-			for (int32 FenceIndex = EntriesPerCmdBuffer.PendingItems.Num() - 1; FenceIndex >= 0; --FenceIndex)
+			if (bImmediately || Pair.Key->IsComplete())
 			{
-				FPendingItemsPerCmdBuffer::FPendingItems& PendingItems = EntriesPerCmdBuffer.PendingItems[FenceIndex];
-				if (bImmediately || PendingItems.FenceCounter < EntriesPerCmdBuffer.CmdBuffer->GetFenceSignaledCounter())
+				for (FStagingBuffer* StagingBuffer : Pair.Value)
 				{
-					for (int32 ResourceIndex = 0; ResourceIndex < PendingItems.Resources.Num(); ++ResourceIndex)
-					{
-						check(PendingItems.Resources[ResourceIndex]);
-						FreeStagingBuffers.Add({PendingItems.Resources[ResourceIndex], GFrameNumberRenderThread});
-					}
-
-					EntriesPerCmdBuffer.PendingItems.RemoveAtSwap(FenceIndex, EAllowShrinking::No);
+					FreeStagingBuffers.Add({ StagingBuffer, GFrameNumberRenderThread });
 				}
+				Completed.Add(Pair.Key);
 			}
-
-			if (EntriesPerCmdBuffer.PendingItems.Num() == 0)
-			{
-				PendingFreeStagingBuffers.RemoveAtSwap(Index, EAllowShrinking::No);
-			}
+		}
+		for (FVulkanSyncPoint* SyncPoint : Completed)
+		{
+			PendingFreeStagingBuffers.Remove(SyncPoint);
 		}
 
 		if (bFreeToOS)
@@ -4707,6 +4675,38 @@ namespace VulkanRHI
 		return false;
 	}
 
+	bool FFenceManager::WaitForAnyFence(TArrayView<FFence*> Fences, uint64 TimeInNanoseconds)
+	{
+		// Mostly used for waits on queues, inlinealloc to that size
+		TArray<VkFence, TInlineAllocator<(int32)EVulkanQueueType::Count>> FenceHandles;
+		for (FFence* Fence : Fences)
+		{
+			if (Fence && !Fence->IsSignaled())
+			{
+				check(UsedFences.Contains(Fence));
+				FenceHandles.Add(Fence->GetHandle());
+			}
+		}
+
+		const VkResult Result = VulkanRHI::vkWaitForFences(Device->GetInstanceHandle(), FenceHandles.Num(), FenceHandles.GetData(), false, TimeInNanoseconds);
+		switch (Result)
+		{
+		case VK_SUCCESS:
+			if (Fences.Num() == 1)
+			{
+				Fences[0]->State = FFence::EState::Signaled;
+			}
+			return true;
+		case VK_TIMEOUT:
+			break;
+		default:
+			VERIFYVULKANRESULT(Result);
+			break;
+		}
+
+		return false;
+	}
+
 	void FFenceManager::ResetFence(FFence* Fence)
 	{
 		if (Fence->State != FFence::EState::NotReady)
@@ -4748,7 +4748,6 @@ namespace VulkanRHI
 
 		FEntry Entry;
 		Entry.StructureType = Type;
-		Queue->GetLastSubmittedInfo(Entry.CmdBuffer, Entry.FenceCounter);
 		Entry.FrameNumber = GVulkanRHIDeletionFrameNumber;
 		Entry.DeviceMemoryAllocation = 0;
 		Entry.Handle = Handle;
@@ -4774,7 +4773,6 @@ namespace VulkanRHI
 
 		FEntry Entry;
 		Entry.StructureType = EType::DeviceMemoryAllocation;
-		Queue->GetLastSubmittedInfo(Entry.CmdBuffer, Entry.FenceCounter);
 		Entry.FrameNumber = GVulkanRHIDeletionFrameNumber;
 
 		Entry.Handle = 0;
@@ -4796,7 +4794,6 @@ namespace VulkanRHI
 
 		FEntry Entry;
 		Entry.StructureType = EType::ResourceAllocation;
-		Queue->GetLastSubmittedInfo(Entry.CmdBuffer, Entry.FenceCounter);
 		Entry.FrameNumber = GVulkanRHIDeletionFrameNumber;
 
 		Entry.Handle = 0;
@@ -4817,94 +4814,110 @@ namespace VulkanRHI
 #if VULKAN_ENABLE_AGGRESSIVE_STATS
 		SCOPE_CYCLE_COUNTER(STAT_VulkanDeletionQueue);
 #endif
-		FScopeLock ScopeLock(&CS);
 
-		VkDevice DeviceHandle = Device->GetInstanceHandle();
-
-		// Traverse list backwards so the swap switches to elements already tested
-		for (int32 Index = Entries.Num() - 1; Index >= 0; --Index)
+		if (bDeleteImmediately)
 		{
-			FEntry* Entry = &Entries[Index];
-			// #todo-rco: Had to add this check, we were getting null CmdBuffers on the first frame, or before first frame maybe
-			if (bDeleteImmediately ||
-				(GVulkanRHIDeletionFrameNumber > Entry->FrameNumber + NUM_FRAMES_TO_WAIT_FOR_RESOURCE_DELETE &&
-					(Entry->CmdBuffer == nullptr || Entry->FenceCounter < Entry->CmdBuffer->GetFenceSignaledCounter()))
-				)
+			FScopeLock ScopeLock(&CS);
+			ReleaseResourcesImmediately(Entries);
+			Entries.Reset();
+		}
+		else
+		{
+			TArray<FEntry> TempEntries;
 			{
-				switch (Entry->StructureType)
+				FScopeLock ScopeLock(&CS);
+				if (GVulkanNumFramesToWaitForResourceDelete)
 				{
-#define VKSWITCH(Type, ...)	case EType::Type: __VA_ARGS__; VulkanRHI::vkDestroy##Type(DeviceHandle, (Vk##Type)Entry->Handle, VULKAN_CPU_ALLOCATOR); break
-					VKSWITCH(RenderPass);
-					VKSWITCH(Buffer);
-					VKSWITCH(BufferView);
-					VKSWITCH(Image);
-					VKSWITCH(ImageView);
-					VKSWITCH(Pipeline, DEC_DWORD_STAT(STAT_VulkanNumPSOs));
-					VKSWITCH(PipelineLayout);
-					VKSWITCH(Framebuffer);
-					VKSWITCH(DescriptorSetLayout);
-					VKSWITCH(Sampler);
-					VKSWITCH(Semaphore);
-					VKSWITCH(ShaderModule);
-					VKSWITCH(Event);
-#undef VKSWITCH
-				case EType::ResourceAllocation:
-				{
-					FVulkanAllocation Allocation = Entry->Allocation;
-					Allocation.Own();
-					Device->GetMemoryManager().FreeVulkanAllocation(Allocation, EVulkanFreeFlag_DontDefer);
-					break;
-				}
-				case EType::DeviceMemoryAllocation:
-				{
-					check(Entry->DeviceMemoryAllocation);
-					Device->GetDeviceMemoryManager().Free(Entry->DeviceMemoryAllocation);
-					break;
-				}
-				case EType::AccelerationStructure:
-				{
-					VulkanDynamicAPI::vkDestroyAccelerationStructureKHR(DeviceHandle, (VkAccelerationStructureKHR)Entry->Handle, VULKAN_CPU_ALLOCATOR);
-					break;
+					TempEntries.Reserve(Entries.Num()/2);
+					for (int32 Index = Entries.Num() - 1; Index >= 0; --Index)
+					{
+						FEntry& Entry = Entries[Index];
+						if (GVulkanRHIDeletionFrameNumber > Entry.FrameNumber + GVulkanNumFramesToWaitForResourceDelete)
+						{
+							TempEntries.Add(Entry);
+							Entries.RemoveAtSwap(Index, EAllowShrinking::No);
+						}
+					}
 				}
-				case EType::BindlessHandle:
+				else
 				{
-					check(Device->SupportsBindless());
-					uint32 BindlessType = (uint32)(Entry->Handle >> 32);
-					check(BindlessType < MAX_uint8);
-					uint32 BindlessIndex = (uint32)(Entry->Handle & 0xFFFFFFFF);
-					FRHIDescriptorHandle DescriptorHandle((uint8)BindlessType, BindlessIndex);
-					Device->GetBindlessDescriptorManager()->Unregister(DescriptorHandle);
-					break;
+					TempEntries = MoveTemp(Entries);
 				}
+			}
 
-				default:
-					check(0);
-					break;
-				}
-				Entries.RemoveAtSwap(Index, EAllowShrinking::No);
+			if (TempEntries.Num())
+			{
+				FVulkanDynamicRHI::Get().EnqueueEndOfPipeTask([Device = Device, Array = MoveTemp(TempEntries)]()
+				{
+					Device->GetDeferredDeletionQueue().ReleaseResourcesImmediately(Array);
+				});
 			}
 		}
 	}
 
-
-
-
-	void FDeferredDeletionQueue2::OnCmdBufferDeleted(FVulkanCommandBuffer* DeletedCmdBuffer)
+	void FDeferredDeletionQueue2::ReleaseResourcesImmediately(const TArray<FEntry>& InEntries)
 	{
-		FScopeLock ScopeLock(&CS);
-		for (int32 Index = 0; Index < Entries.Num(); ++Index)
-		{
-			FEntry& Entry = Entries[Index];
-			if (Entry.CmdBuffer == DeletedCmdBuffer)
+		VkDevice DeviceHandle = Device->GetInstanceHandle();
+
+		// Traverse list backwards so the swap switches to elements already tested
+		for (const FEntry& Entry : InEntries)
+		{
+			switch (Entry.StructureType)
+			{
+#define VKSWITCH(Type, ...)	case EType::Type: __VA_ARGS__; VulkanRHI::vkDestroy##Type(DeviceHandle, (Vk##Type)Entry.Handle, VULKAN_CPU_ALLOCATOR); break
+				VKSWITCH(RenderPass);
+				VKSWITCH(Buffer);
+				VKSWITCH(BufferView);
+				VKSWITCH(Image);
+				VKSWITCH(ImageView);
+				VKSWITCH(Pipeline, DEC_DWORD_STAT(STAT_VulkanNumPSOs));
+				VKSWITCH(PipelineLayout);
+				VKSWITCH(Framebuffer);
+				VKSWITCH(DescriptorSetLayout);
+				VKSWITCH(Sampler);
+				VKSWITCH(Semaphore);
+				VKSWITCH(ShaderModule);
+				VKSWITCH(Event);
+#undef VKSWITCH
+			case EType::ResourceAllocation:
+			{
+				FVulkanAllocation Allocation = Entry.Allocation;
+				Allocation.Own();
+				Device->GetMemoryManager().FreeVulkanAllocation(Allocation, EVulkanFreeFlag_DontDefer);
+				break;
+			}
+			case EType::DeviceMemoryAllocation:
+			{
+				check(Entry.DeviceMemoryAllocation);
+				FDeviceMemoryAllocation* DeviceMemoryAllocation = Entry.DeviceMemoryAllocation;
+				Device->GetDeviceMemoryManager().Free(DeviceMemoryAllocation);
+				break;
+			}
+			case EType::AccelerationStructure:
+			{
+				VulkanDynamicAPI::vkDestroyAccelerationStructureKHR(DeviceHandle, (VkAccelerationStructureKHR)Entry.Handle, VULKAN_CPU_ALLOCATOR);
+				break;
+			}
+			case EType::BindlessHandle:
 			{
-				Entry.CmdBuffer = nullptr;
+				check(Device->SupportsBindless());
+				uint32 BindlessType = (uint32)(Entry.Handle >> 32);
+				check(BindlessType < MAX_uint8);
+				uint32 BindlessIndex = (uint32)(Entry.Handle & 0xFFFFFFFF);
+				FRHIDescriptorHandle DescriptorHandle((uint8)BindlessType, BindlessIndex);
+				Device->GetBindlessDescriptorManager()->Unregister(DescriptorHandle);
+				break;
+			}
+
+			default:
+				check(0);
+				break;
 			}
 		}
 	}
 
 
 
-
 	FTempBlockAllocator::FTempBlockAllocator(FVulkanDevice* InDevice, uint32 InBlockSize, uint32 InBlockAlignment, VkBufferUsageFlags InBufferUsage)
 		: FDeviceChild(InDevice)
 		, BlockSize(InBlockSize)
@@ -4980,7 +4993,7 @@ namespace VulkanRHI
 		return NewBlock;
 	}
 
-	FTempBlockAllocator::FInternalAlloc FTempBlockAllocator::InternalAlloc(uint32 InSize, FVulkanCommandBuffer* CmdBuffer)
+	FTempBlockAllocator::FInternalAlloc FTempBlockAllocator::InternalAlloc(uint32 InSize, FVulkanContextCommon& Context)
 	{
 		const uint32 AlignedSize = Align(InSize, BlockAlignment);
 		checkfSlow(BlockAlignment < BlockSize, TEXT("Requested size of %d (%d aligned) is too large for block size of %d"), InSize, AlignedSize, BlockSize);
@@ -4989,10 +5002,8 @@ namespace VulkanRHI
 		{
 			FReadScopeLock ScopeLock(RWLock);
 
-			if (CurrentBlock->Fences.Contains(CmdBuffer))
+			if (CurrentBlock->SyncPoints.Contains(Context.GetContextSyncPoint()))
 			{
-				CurrentBlock->Fences[CmdBuffer] = CmdBuffer->GetFenceSignaledCounter();
-
 				const uint32 AllocOffset = CurrentBlock->CurrentOffset.fetch_add(AlignedSize);
 				if (AllocOffset + InSize < BlockSize)
 				{
@@ -5009,7 +5020,7 @@ namespace VulkanRHI
 			uint32 AllocOffset = CurrentBlock->CurrentOffset.fetch_add(AlignedSize);
 			if (AllocOffset + InSize < BlockSize)
 			{
-				CurrentBlock->Fences.FindOrAdd(CmdBuffer) = CmdBuffer->GetFenceSignaledCounter();
+				CurrentBlock->SyncPoints.AddUnique(Context.GetContextSyncPoint());
 				return FInternalAlloc{ CurrentBlock, AllocOffset };
 			}
 
@@ -5024,7 +5035,7 @@ namespace VulkanRHI
 				CurrentBlock = AllocBlock();
 			}
 
-			CurrentBlock->Fences.FindOrAdd(CmdBuffer) = CmdBuffer->GetFenceSignaledCounter();
+			CurrentBlock->SyncPoints.AddUnique(Context.GetContextSyncPoint());
 
 			AllocOffset = CurrentBlock->CurrentOffset.fetch_add(AlignedSize);
 			checkSlow(AllocOffset == 0);
@@ -5033,9 +5044,9 @@ namespace VulkanRHI
 		}
 	}
 
-	uint8* FTempBlockAllocator::Alloc(uint32 InSize, FVulkanCommandBuffer* CmdBuffer, VkDescriptorBufferBindingInfoEXT& OutBindingInfo, VkDeviceSize& OutOffset)
+	uint8* FTempBlockAllocator::Alloc(uint32 InSize, FVulkanContextCommon& Context, VkDescriptorBufferBindingInfoEXT& OutBindingInfo, VkDeviceSize& OutOffset)
 	{
-		FInternalAlloc Alloc = InternalAlloc(InSize, CmdBuffer);
+		FInternalAlloc Alloc = InternalAlloc(InSize, Context);
 
 		ZeroVulkanStruct(OutBindingInfo, VK_STRUCTURE_TYPE_DESCRIPTOR_BUFFER_BINDING_INFO_EXT);
 		OutBindingInfo.address = Alloc.Block->BufferAddress;
@@ -5044,12 +5055,12 @@ namespace VulkanRHI
 		return Alloc.Block->MappedPointer + Alloc.Offset;
 	}
 
-	uint8* FTempBlockAllocator::Alloc(uint32 InSize, uint32 InAlignment, FVulkanCommandBuffer* CmdBuffer, FVulkanAllocation& OutAllocation, VkDescriptorAddressInfoEXT* OutDescriptorAddressInfo)
+	uint8* FTempBlockAllocator::Alloc(uint32 InSize, uint32 InAlignment, FVulkanContextCommon& Context, FVulkanAllocation& OutAllocation, VkDescriptorAddressInfoEXT* OutDescriptorAddressInfo)
 	{
 		const uint32 AlignedSize = Align(InSize, BlockAlignment);
 		checkfSlow(AlignedSize < BlockSize, TEXT("Requested size of %d (%d aligned) is too large for block size of %d"), InSize, AlignedSize, BlockSize);
 
-		FInternalAlloc Alloc = InternalAlloc(InSize, CmdBuffer);
+		FInternalAlloc Alloc = InternalAlloc(InSize, Context);
 
 		OutAllocation.Reference(Alloc.Block->Allocation);
 		OutAllocation.VulkanHandle = (uint64)Alloc.Block->Buffer;
@@ -5066,13 +5077,13 @@ namespace VulkanRHI
 		return Alloc.Block->MappedPointer + Alloc.Offset;
 	}
 
-	uint8* FTempBlockAllocator::Alloc(uint32 InSize, FVulkanCommandBuffer* CmdBuffer, VkStridedDeviceAddressRegionKHR* OutStridedDeviceAddressRegion)
+	uint8* FTempBlockAllocator::Alloc(uint32 InSize, FVulkanContextCommon& Context, VkStridedDeviceAddressRegionKHR* OutStridedDeviceAddressRegion)
 	{
 		const uint32 AlignedSize = Align(InSize, BlockAlignment);
 		checkfSlow(AlignedSize < BlockSize, TEXT("Requested size of %d (%d aligned) is too large for block size of %d"), InSize, AlignedSize, BlockSize);
 		checkSlow(OutStridedDeviceAddressRegion);
 
-		FInternalAlloc Alloc = InternalAlloc(InSize, CmdBuffer);
+		FInternalAlloc Alloc = InternalAlloc(InSize, Context);
 		OutStridedDeviceAddressRegion->deviceAddress = Alloc.Block->BufferAddress + Alloc.Offset;
 		OutStridedDeviceAddressRegion->size = InSize;
 		OutStridedDeviceAddressRegion->stride = OutStridedDeviceAddressRegion->size;
@@ -5088,10 +5099,9 @@ namespace VulkanRHI
 			FTempMemoryBlock* Block = BusyBlocks[Index];
 
 			bool BlockReady = true;
-			for (auto& Pair : Block->Fences)
+			for (FVulkanSyncPointRef& SyncPoint : Block->SyncPoints)
 			{
-				FVulkanCommandBuffer* CmdBuffer = Pair.Key;
-				if (Pair.Value >= CmdBuffer->GetFenceSignaledCounter())
+				if (!SyncPoint->IsComplete())
 				{
 					BlockReady = false;
 					break;
@@ -5100,7 +5110,7 @@ namespace VulkanRHI
 
 			if (BlockReady)
 			{
-				Block->Fences.Empty();
+				Block->SyncPoints.Empty();
 				Block->CurrentOffset = 0;
 				BusyBlocks.RemoveAtSwap(Index, EAllowShrinking::No);
 				AvailableBlocks.Add(Block);
diff --git a/Engine/Source/Runtime/VulkanRHI/Private/VulkanMemory.h b/Engine/Source/Runtime/VulkanRHI/Private/VulkanMemory.h
--- a/Engine/Source/Runtime/VulkanRHI/Private/VulkanMemory.h
+++ b/Engine/Source/Runtime/VulkanRHI/Private/VulkanMemory.h
@@ -24,12 +24,16 @@
 #define VULKAN_MEMORY_HIGHER_PRIORITY 0.75f
 #define VULKAN_MEMORY_HIGHEST_PRIORITY 1.f
 
+class FVulkanContextCommon;
 class FVulkanCommandListContext;
 class FVulkanDevice;
 class FVulkanQueue;
 class FVulkanCommandBuffer;
 class FVulkanTexture;
 
+using FVulkanSyncPoint = FGraphEvent;
+using FVulkanSyncPointRef = TRefCountPtr<FVulkanSyncPoint>;
+
 #if !VULKAN_OBJECT_TRACKING
 #define VULKAN_TRACK_OBJECT_CREATE(Type, Ptr) do{}while(0)
 #define VULKAN_TRACK_OBJECT_DELETE(Type, Ptr) do{}while(0)
@@ -968,7 +972,7 @@ namespace VulkanRHI
 		FStagingBuffer* AcquireBuffer(uint32 Size, VkBufferUsageFlags InUsageFlags = VK_BUFFER_USAGE_TRANSFER_SRC_BIT, VkMemoryPropertyFlagBits InMemoryReadFlags = VK_MEMORY_PROPERTY_HOST_COHERENT_BIT);
 
 		// Sets pointer to nullptr
-		void ReleaseBuffer(FVulkanCommandBuffer* CmdBuffer, FStagingBuffer*& StagingBuffer);
+		void ReleaseBuffer(FVulkanContextCommon* Context, FStagingBuffer*& StagingBuffer);
 
 		void ProcessPendingFree(bool bImmediately, bool bFreeToOS);
 
@@ -978,25 +982,11 @@ namespace VulkanRHI
 
 	protected:
 		friend class FMemoryManager;
-		struct FPendingItemsPerCmdBuffer
-		{
-			FVulkanCommandBuffer* CmdBuffer;
-			struct FPendingItems
-			{
-				uint64 FenceCounter;
-				TArray<FStagingBuffer*> Resources;
-			};
-
-
-			inline FPendingItems* FindOrAddItemsForFence(uint64 Fence);
-
-			TArray<FPendingItems> PendingItems;
-		};
 
 		FCriticalSection StagingLock;
 
 		TArray<FStagingBuffer*> UsedStagingBuffers;
-		TArray<FPendingItemsPerCmdBuffer> PendingFreeStagingBuffers;
+		TMap<FVulkanSyncPointRef, TArray<FStagingBuffer*>> PendingFreeStagingBuffers;
 		struct FFreeEntry
 		{
 			FStagingBuffer* StagingBuffer;
@@ -1007,8 +997,6 @@ namespace VulkanRHI
 		uint64 PeakUsedMemory;
 		uint64 UsedMemory;
 
-		FPendingItemsPerCmdBuffer* FindOrAdd(FVulkanCommandBuffer* CmdBuffer);
-
 		void ProcessPendingFreeNoLock(bool bImmediately, bool bFreeToOS);
 
 		FVulkanDevice* Device;
@@ -1080,6 +1068,7 @@ namespace VulkanRHI
 		}
 
 		// Returns false if it timed out
+		bool WaitForAnyFence(TArrayView<FFence*> Fences, uint64 TimeInNanoseconds);
 		bool WaitForFence(FFence* Fence, uint64 TimeInNanoseconds);
 
 		void ResetFence(FFence* Fence);
@@ -1130,7 +1119,7 @@ namespace VulkanRHI
 			VkDeviceAddress BufferAddress = 0;
 
 			// Used by owners to suballocate
-			TMap<FVulkanCommandBuffer*, uint64> Fences;
+			TArray<FVulkanSyncPointRef> SyncPoints;
 			std::atomic<uint32> CurrentOffset = 0;
 		};
 
@@ -1143,15 +1132,15 @@ namespace VulkanRHI
 		FTempBlockAllocator(FVulkanDevice* InDevice, uint32 InBlockSize, uint32 InBlockAlignment, VkBufferUsageFlags InBufferUsage);
 		virtual ~FTempBlockAllocator();
 
-		uint8* Alloc(uint32 InSize, FVulkanCommandBuffer* CmdBuffer, VkDescriptorBufferBindingInfoEXT& OutBindingInfo, VkDeviceSize& OutOffset);
-		uint8* Alloc(uint32 InSize, uint32 InAlignment, FVulkanCommandBuffer* CmdBuffer, FVulkanAllocation& OutAllocation, VkDescriptorAddressInfoEXT* OutDescriptorAddressInfo = nullptr);
-		uint8* Alloc(uint32 InSize, FVulkanCommandBuffer* CmdBuffer, VkStridedDeviceAddressRegionKHR* OutStridedDeviceAddressRegion);
+		uint8* Alloc(uint32 InSize, FVulkanContextCommon& Context, VkDescriptorBufferBindingInfoEXT& OutBindingInfo, VkDeviceSize& OutOffset);
+		uint8* Alloc(uint32 InSize, uint32 InAlignment, FVulkanContextCommon& Context, FVulkanAllocation& OutAllocation, VkDescriptorAddressInfoEXT* OutDescriptorAddressInfo = nullptr);
+		uint8* Alloc(uint32 InSize, FVulkanContextCommon& Context, VkStridedDeviceAddressRegionKHR* OutStridedDeviceAddressRegion);
 
 		void UpdateBlocks();
 
 	protected:
 		FTempMemoryBlock* AllocBlock();
-		FInternalAlloc InternalAlloc(uint32 InSize, FVulkanCommandBuffer* CmdBuffer);
+		FInternalAlloc InternalAlloc(uint32 InSize, FVulkanContextCommon& Context);
 
 		const uint32 BlockSize;
 		const uint32 BlockAlignment;
diff --git a/Engine/Source/Runtime/VulkanRHI/Private/VulkanPipelineState.cpp b/Engine/Source/Runtime/VulkanRHI/Private/VulkanPipelineState.cpp
--- a/Engine/Source/Runtime/VulkanRHI/Private/VulkanPipelineState.cpp
+++ b/Engine/Source/Runtime/VulkanRHI/Private/VulkanPipelineState.cpp
@@ -165,7 +165,7 @@ static inline VulkanRHI::FVulkanAllocation UpdatePackedUniformBuffers(const FPac
 	const uint32 UBAlign = (uint32)InCmdBuffer->Device.GetLimits().minUniformBufferOffsetAlignment;
 
 	VulkanRHI::FVulkanAllocation TempAllocation;
-	uint8* MappedPointer = InContext->GetTempBlockAllocator().Alloc(UBSize, UBAlign, InCmdBuffer, TempAllocation);
+	uint8* MappedPointer = InContext->GetTempBlockAllocator().Alloc(UBSize, UBAlign, *InContext, TempAllocation);
 
 	FMemory::Memcpy(MappedPointer, StagedUniformBuffer.GetData(), UBSize);
 
@@ -250,7 +250,7 @@ void FVulkanComputePipelineDescriptorState::UpdateBindlessDescriptors(FVulkanCom
 		const VkDeviceSize UBOffsetAlignment = Device->GetLimits().minUniformBufferOffsetAlignment;
 
 		VulkanRHI::FVulkanAllocation TempAllocation;
-		uint8* MappedPointer = CmdListContext->GetTempBlockAllocator().Alloc(UBSize, UBOffsetAlignment, &CommandBuffer, TempAllocation, &DescriptorAddressInfos[BindingIndex]);
+		uint8* MappedPointer = CmdListContext->GetTempBlockAllocator().Alloc(UBSize, UBOffsetAlignment, *CmdListContext, TempAllocation, &DescriptorAddressInfos[BindingIndex]);
 		FMemory::Memcpy(MappedPointer, StagedUniformBuffer.GetData(), UBSize);
 
 		PackedUniformBuffersDirty = 0;
@@ -476,7 +476,7 @@ void FVulkanGraphicsPipelineDescriptorState::UpdateBindlessDescriptors(FVulkanCo
 				const int32 BindingIndex = 0;
 
 				VulkanRHI::FVulkanAllocation TempAllocation;
-				uint8* MappedPointer = CmdListContext->GetTempBlockAllocator().Alloc(UBSize, UBOffsetAlignment, &CommandBuffer, TempAllocation, &DescriptorAddressInfos[BindingIndex]);
+				uint8* MappedPointer = CmdListContext->GetTempBlockAllocator().Alloc(UBSize, UBOffsetAlignment, *CmdListContext, TempAllocation, &DescriptorAddressInfos[BindingIndex]);
 				FMemory::Memcpy(MappedPointer, StagedUniformBuffer.GetData(), UBSize);
 
 				PackedUniformBuffersDirty[Stage] = 0;
diff --git a/Engine/Source/Runtime/VulkanRHI/Private/VulkanQuery.cpp b/Engine/Source/Runtime/VulkanRHI/Private/VulkanQuery.cpp
--- a/Engine/Source/Runtime/VulkanRHI/Private/VulkanQuery.cpp
+++ b/Engine/Source/Runtime/VulkanRHI/Private/VulkanQuery.cpp
@@ -79,7 +79,6 @@ FVulkanQueryPool::~FVulkanQueryPool()
 	ResetEvent = VK_NULL_HANDLE;
 }
 
-
 bool FVulkanOcclusionQueryPool::CanBeReused()
 {
 	const uint64 NumWords = NumUsedQueries / 64;
@@ -98,7 +97,6 @@ bool FVulkanOcclusionQueryPool::CanBeReused()
 
 bool FVulkanOcclusionQueryPool::InternalTryGetResults(bool bWait)
 {
-	check(CmdBuffer);
 	check(State == EState::RHIT_PostEndBatch);
 
 	VkResult Result = VK_NOT_READY;
@@ -187,24 +185,11 @@ bool FVulkanOcclusionQueryPool::InternalTryGetResults(bool bWait)
 	return false;
 }
 
-void FVulkanOcclusionQueryPool::SetFence(FVulkanCommandBuffer* InCmdBuffer)
-{
-	check(InCmdBuffer);
-	CmdBuffer = InCmdBuffer;
-	FenceCounter = InCmdBuffer->GetFenceSignaledCounter();
-}
-
 void FVulkanOcclusionQueryPool::Reset(FVulkanCommandBuffer* InCmdBuffer, uint32 InFrameNumber)
 {
-/*
-	check(!InCmdBuffer || InCmdBuffer->GetFenceSignaledCounter() > FenceCounter);
-*/
 	//ensure(State == EState::Undefined || State == EState::RT_PostGetResults);
 	FMemory::Memzero(AcquiredIndices.GetData(), AcquiredIndices.GetAllocatedSize());
-/*
-	CmdBuffer = nullptr;
-	FenceCounter = UINT32_MAX;
-*/
+
 	NumUsedQueries = 0;
 	FrameNumber = InFrameNumber;
 /*
@@ -232,9 +217,6 @@ void FVulkanOcclusionQueryPool::Reset(FVulkanCommandBuffer* InCmdBuffer, uint32
 void FVulkanCommandListContext::BeginOcclusionQueryBatch(FVulkanCommandBuffer* CmdBuffer, uint32 NumQueriesInBatch)
 {
 	ensure(IsImmediate());
-/*
-	FScopeLock ScopeLock(&GOcclusionQueryCS);
-*/
 	checkf(!CurrentOcclusionQueryPool, TEXT("BeginOcclusionQueryBatch called without corresponding EndOcclusionQueryBatch!"));
 	CurrentOcclusionQueryPool = Device.AcquireOcclusionQueryPool(this, NumQueriesInBatch);
 	ensure(CmdBuffer->IsOutsideRenderPass());
@@ -346,12 +328,11 @@ void FVulkanDevice::ReleaseUnusedOcclusionQueryPools()
 	}
 }
 
-void FVulkanCommandListContext::EndOcclusionQueryBatch(FVulkanCommandBuffer* CmdBuffer)
+void FVulkanCommandListContext::EndOcclusionQueryBatch()
 {
 	ensure(IsImmediate());
-	check(CmdBuffer);
 	checkf(CurrentOcclusionQueryPool, TEXT("EndOcclusionQueryBatch called without corresponding BeginOcclusionQueryBatch!"));
-	CurrentOcclusionQueryPool->EndBatch(CmdBuffer);
+	CurrentOcclusionQueryPool->EndBatch();
 	CurrentOcclusionQueryPool = nullptr;
 }
 
@@ -496,7 +477,7 @@ bool FVulkanDynamicRHI::RHIGetRenderQueryResult(FRHIRenderQuery* QueryRHI, uint6
 			for (uint32 IssueIndex = 1; IssueIndex < Query->Pool->NumIssuedTimestamps; ++IssueIndex)
 			{
 				const FVulkanTimingQueryPool::FCmdBufferFence& StartQuerySyncPoint = Query->Pool->TimestampListHandles[TimestampIndex];
-				if (StartQuerySyncPoint.FenceCounter < StartQuerySyncPoint.CmdBuffer->GetFenceSignaledCounter())
+				if (StartQuerySyncPoint.SyncPoint->IsComplete())
 				{
 					Query->Pool->ResultsBuffer->InvalidateMappedMemory();
 					const uint64* Data = Query->Pool->MappedPointer;
@@ -522,16 +503,11 @@ bool FVulkanDynamicRHI::RHIGetRenderQueryResult(FRHIRenderQuery* QueryRHI, uint6
 				if (bBlocking)
 				{
 					const FVulkanTimingQueryPool::FCmdBufferFence& StartQuerySyncPoint = Query->Pool->TimestampListHandles[TimestampIndex];
-					bool bWaitForStart = StartQuerySyncPoint.FenceCounter == StartQuerySyncPoint.CmdBuffer->GetFenceSignaledCounter();
-					if (bWaitForStart)
-					{
-						RHICmdList.ImmediateFlush(EImmediateFlushType::FlushRHIThread);
-					}
 
 					// CPU wait for query results to be ready.
-					if (bWaitForStart && StartQuerySyncPoint.FenceCounter == StartQuerySyncPoint.CmdBuffer->GetFenceSignaledCounter())
+					if (!StartQuerySyncPoint.SyncPoint->IsComplete())
 					{
-						WaitForCmdBuffer(StartQuerySyncPoint.CmdBuffer);
+						FVulkanDynamicRHI::Get().ProcessInterruptQueueUntil(StartQuerySyncPoint.SyncPoint);
 					}
 				}
 
@@ -617,8 +593,7 @@ void FVulkanCommandListContext::RHIEndRenderQuery(FRHIRenderQuery* QueryRHI)
 		FVulkanCommandBuffer& CommandBuffer = GetCommandBuffer();
 		VulkanRHI::vkCmdWriteTimestamp(CommandBuffer.GetHandle(), QueryPipelineStage, Query->Pool->GetHandle(), QueryEndIndex);
 		CommandBuffer.AddPendingTimestampQuery(QueryEndIndex, 1, Query->Pool->GetHandle(), Query->Pool->ResultsBuffer->GetHandle(), true);
-		Query->Pool->TimestampListHandles[QueryEndIndex].CmdBuffer = &CommandBuffer;
-		Query->Pool->TimestampListHandles[QueryEndIndex].FenceCounter = CommandBuffer.GetFenceSignaledCounter();
+		Query->Pool->TimestampListHandles[QueryEndIndex].SyncPoint = GetContextSyncPoint();
 		Query->Pool->TimestampListHandles[QueryEndIndex].FrameCount = GetFrameCounter();
 		Query->Pool->NumIssuedTimestamps = FMath::Min<uint32>(Query->Pool->NumIssuedTimestamps + 1, Query->Pool->BufferSize);
 	}
diff --git a/Engine/Source/Runtime/VulkanRHI/Private/VulkanQueue.cpp b/Engine/Source/Runtime/VulkanRHI/Private/VulkanQueue.cpp
--- a/Engine/Source/Runtime/VulkanRHI/Private/VulkanQueue.cpp
+++ b/Engine/Source/Runtime/VulkanRHI/Private/VulkanQueue.cpp
@@ -19,7 +19,7 @@ FAutoConsoleVariableRef CVarVulkanWaitForIdleOnSubmit(
 	ECVF_Default
 	);
 
-static int32 GAllowTimelineSemaphores = 1;
+int32 GAllowTimelineSemaphores = 1;
 FAutoConsoleVariableRef CVarVulkanSubmissionAllowTimelineSemaphores(
 	TEXT("r.Vulkan.Submission.AllowTimelineSemaphores"),
 	GAllowTimelineSemaphores,
@@ -40,14 +40,13 @@ FAutoConsoleVariableRef CVarVulkanSubmissionAllowTimelineSemaphores(
 //);
 
 
-FVulkanQueue::FVulkanQueue(FVulkanDevice& InDevice, uint32 InFamilyIndex)
+FVulkanQueue::FVulkanQueue(FVulkanDevice& InDevice, uint32 InFamilyIndex, EVulkanQueueType InQueueType)
 	: Queue(VK_NULL_HANDLE)
 	, FamilyIndex(InFamilyIndex)
 	, QueueIndex(0)
+	, QueueType(InQueueType)
 	, Device(InDevice)
 	, bUseTimelineSemaphores(InDevice.GetOptionalExtensions().HasKHRTimelineSemaphore && GAllowTimelineSemaphores)
-	, LastSubmittedCmdBuffer(nullptr)
-	, LastSubmittedCmdBufferFenceCounter(0)
 	, LayoutManager(false, nullptr)
 {
 	VulkanRHI::vkGetDeviceQueue(Device.GetInstanceHandle(), FamilyIndex, QueueIndex, &Queue);
@@ -154,7 +153,6 @@ void FVulkanQueue::SubmitPayloads(TArrayView<FVulkanPayload*> Payloads, TSet<VkS
 	CommandBufferStorage.Reserve(NumCommandBuffers);
 	SubmitInfos.Reserve(NumPayloads);
 
-	FVulkanCommandBuffer* LastSubmittedCommandBuffer = nullptr;
 	for (FVulkanPayload* Payload : Payloads)
 	{
 		VkSubmitInfo& SubmitInfo = SubmitInfos.AddZeroed_GetRef();
@@ -189,10 +187,8 @@ void FVulkanQueue::SubmitPayloads(TArrayView<FVulkanPayload*> Payloads, TSet<VkS
 			{
 				CommandBufferStorage.Add(CommandBuffer->GetHandle());
 
-				CommandBuffer->SubmittedTimelineValue = NextTimelineSemaphoreValue;
 				CommandBuffer->State = FVulkanCommandBuffer::EState::Submitted;
 				CommandBuffer->SubmittedTime = FPlatformTime::Seconds();
-				LastSubmittedCommandBuffer = CommandBuffer;
 
 				// If we're tracking layouts for the queue, merge in the changes recorded in this command buffer's context
 				CommandBuffer->GetLayoutManager().TransferTo(LayoutManager);
@@ -260,11 +256,6 @@ void FVulkanQueue::SubmitPayloads(TArrayView<FVulkanPayload*> Payloads, TSet<VkS
 	{
 		PendingInterrupt.Enqueue(Payload);
 	}
-
-	if (LastSubmittedCmdBuffer)
-	{
-		UpdateLastSubmittedCommandBuffer(LastSubmittedCmdBuffer);
-	}
 }
 
 
@@ -288,7 +279,7 @@ void FVulkanQueue::Submit(TArrayView<VkSubmitInfo> InSubmitInfos, VulkanRHI::FFe
 	}
 }
 
-int32 FVulkanQueue::ProcessInterruptQueue(uint64 Timeout, uint64 UpToValue)
+int32 FVulkanQueue::ProcessInterruptQueue(uint64 Timeout)
 {
 	if (bUseTimelineSemaphores)
 	{
@@ -341,12 +332,13 @@ int32 FVulkanQueue::ProcessInterruptQueue(uint64 Timeout, uint64 UpToValue)
 				CommandBuffer->ResetFenceStatus();
 			}
 
-			delete Payload;
-		}
+			if (Payload->SyncPoint.IsValid())
+			{
+				Payload->SyncPoint->DispatchSubsequents();
+				Payload->SyncPoint = nullptr;
+			}
 
-		if (UpToValue && (UpToValue <= CompletedTimelineSemaphoreValue))
-		{
-			break;
+			delete Payload;
 		}
 
 	} while (bSuccess);
@@ -354,13 +346,6 @@ int32 FVulkanQueue::ProcessInterruptQueue(uint64 Timeout, uint64 UpToValue)
 	return NumPayloads;
 }
 
-void FVulkanQueue::UpdateLastSubmittedCommandBuffer(FVulkanCommandBuffer* CmdBuffer)
-{
-	FScopeLock ScopeLock(&CS);
-	LastSubmittedCmdBuffer = CmdBuffer;
-	LastSubmittedCmdBufferFenceCounter = CmdBuffer->GetFenceSignaledCounter();
-}
-
 void FVulkanQueue::FillSupportedStageBits()
 {
 	check((int32)FamilyIndex < Device.GetQueueFamilyProps().Num());
diff --git a/Engine/Source/Runtime/VulkanRHI/Private/VulkanQueue.h b/Engine/Source/Runtime/VulkanRHI/Private/VulkanQueue.h
--- a/Engine/Source/Runtime/VulkanRHI/Private/VulkanQueue.h
+++ b/Engine/Source/Runtime/VulkanRHI/Private/VulkanQueue.h
@@ -9,9 +9,8 @@
 #include "VulkanConfiguration.h"
 #include "VulkanBarriers.h"
 
+class FVulkanDynamicRHI;
 class FVulkanDevice;
-class FVulkanCommandBuffer;
-class FVulkanCommandListContext;
 class FVulkanLayoutManager;
 class FVulkanPayload;
 
@@ -21,10 +20,19 @@ namespace VulkanRHI
 	class FSemaphore;
 }
 
+enum class EVulkanQueueType : uint8
+{
+	Graphics = 0,
+	AsyncCompute,
+	Transfer,
+
+	Count,
+};
+
 class FVulkanQueue
 {
 public:
-	FVulkanQueue(FVulkanDevice& InDevice, uint32 InFamilyIndex);
+	FVulkanQueue(FVulkanDevice& InDevice, uint32 InFamilyIndex, EVulkanQueueType InQueueType);
 	~FVulkanQueue();
 
 	inline uint32 GetFamilyIndex() const
@@ -42,13 +50,6 @@ public:
 		return Queue;
 	}
 
-	inline void GetLastSubmittedInfo(FVulkanCommandBuffer*& OutCmdBuffer, uint64& OutFenceCounter) const
-	{
-		FScopeLock ScopeLock(&CS);
-		OutCmdBuffer = LastSubmittedCmdBuffer;
-		OutFenceCounter = LastSubmittedCmdBufferFenceCounter;
-	}
-
 	inline VkPipelineStageFlags GetSupportedStageBits() const
 	{
 		return SupportedStages;
@@ -61,12 +62,17 @@ public:
 
 	void NotifyDeletedImage(VkImage Image);
 
-	uint64 GetLastSubmittedTimelineSemaphoreValue()
+	const VulkanRHI::FSemaphore* GetTimelineSemaphore() const
+	{
+		return TimelineSempahore;
+	}
+
+	uint64 GetLastSubmittedTimelineSemaphoreValue() const
 	{
 		return NextTimelineSemaphoreValue - 1;
 	}
 
-	uint64 GetCompletedTimelineSemaphoreValue()
+	uint64 GetCompletedTimelineSemaphoreValue() const
 	{
 		return CompletedTimelineSemaphoreValue;
 	}
@@ -77,15 +83,25 @@ public:
 	}
 
 	int32 SubmitQueuedPayloads(TSet<VkSemaphore>& SignaledSemas);
-	int32 ProcessInterruptQueue(uint64 Timeout, uint64 UpToValue = 0);
+	int32 ProcessInterruptQueue(uint64 Timeout);
 
 private:
 	void SubmitPayloads(TArrayView<FVulkanPayload*> Payloads, TSet<VkSemaphore>& InOutSignaledSemas);
 	void Submit(TArrayView<VkSubmitInfo> InSubmitInfos, VulkanRHI::FFence* Fence);
 
+	// Used by submission pipe which holds the proper locks to access this queue
+	FVulkanPayload* GetNextInterruptPayload()
+	{
+		FVulkanPayload* Payload = nullptr;
+		PendingInterrupt.Peek(Payload);
+		return Payload;
+	}
+	friend FVulkanDynamicRHI;
+
 	VkQueue Queue;
 	uint32 FamilyIndex;
 	uint32 QueueIndex;
+	EVulkanQueueType QueueType;
 	FVulkanDevice& Device;
 
 	const bool bUseTimelineSemaphores;
@@ -94,8 +110,6 @@ private:
 	uint64 CompletedTimelineSemaphoreValue = 1;
 
 	mutable FCriticalSection CS;
-	FVulkanCommandBuffer* LastSubmittedCmdBuffer;
-	uint64 LastSubmittedCmdBufferFenceCounter;
 	uint64 SubmitCounter;
 	VkPipelineStageFlags SupportedStages;
 
@@ -105,6 +119,5 @@ private:
 	// Last known layouts submitted on this queue, used for defrag
 	FVulkanLayoutManager LayoutManager;
 
-	void UpdateLastSubmittedCommandBuffer(FVulkanCommandBuffer* CmdBuffer);
 	void FillSupportedStageBits();
 };
diff --git a/Engine/Source/Runtime/VulkanRHI/Private/VulkanRHI.cpp b/Engine/Source/Runtime/VulkanRHI/Private/VulkanRHI.cpp
--- a/Engine/Source/Runtime/VulkanRHI/Private/VulkanRHI.cpp
+++ b/Engine/Source/Runtime/VulkanRHI/Private/VulkanRHI.cpp
@@ -550,6 +550,8 @@ void FVulkanDynamicRHI::Shutdown()
 
 		// Flush all pending deletes before destroying the device.
 		FRHICommandListImmediate::Get().ImmediateFlush(EImmediateFlushType::FlushRHIThreadFlushResources);
+
+		ShutdownSubmissionPipe();
 	}
 
 	Device->Destroy();
@@ -1064,6 +1066,8 @@ void FVulkanDynamicRHI::InitInstance()
 		UE_LOG(LogVulkanRHI, Log, TEXT("Integrated GPU (iGPU): %s"), GRHIDeviceIsIntegrated ? TEXT("true") : TEXT("false"));
 #endif
 
+		InitializeSubmissionPipe();
+
 		FRenderResource::InitPreRHIResources();
 		GIsRHIInitialized = true;
 	}
diff --git a/Engine/Source/Runtime/VulkanRHI/Private/VulkanRayTracing.cpp b/Engine/Source/Runtime/VulkanRHI/Private/VulkanRayTracing.cpp
--- a/Engine/Source/Runtime/VulkanRHI/Private/VulkanRayTracing.cpp
+++ b/Engine/Source/Runtime/VulkanRHI/Private/VulkanRayTracing.cpp
@@ -968,7 +968,7 @@ void FVulkanRayTracingShaderTable::SetSlot(EShaderFrequency Frequency, uint32 Ds
 VkStridedDeviceAddressRegionKHR FVulkanRayTracingShaderTable::CommitRayGenShader(FVulkanCommandListContext& Context, uint32 SrcHandleIndex, TConstArrayView<uint8> SrcHandleData)
 {
 	VkStridedDeviceAddressRegionKHR RayGenRegion;
-	uint8* pMappedMemory = Context.GetTempBlockAllocator().Alloc(HandleSize, Context.GetActiveCmdBuffer(), &RayGenRegion);
+	uint8* pMappedMemory = Context.GetTempBlockAllocator().Alloc(HandleSize, Context, &RayGenRegion);
 
 	FMemory::Memcpy(pMappedMemory, &SrcHandleData[SrcHandleIndex * HandleSize], HandleSize);
 
@@ -1024,7 +1024,7 @@ void FVulkanRayTracingShaderTable::Commit(FVulkanCommandListContext& Context)
 	VkMemoryBarrier BarrierBefore = { VK_STRUCTURE_TYPE_MEMORY_BARRIER, nullptr, VK_ACCESS_MEMORY_READ_BIT | VK_ACCESS_MEMORY_WRITE_BIT, VK_ACCESS_TRANSFER_READ_BIT | VK_ACCESS_TRANSFER_WRITE_BIT };
 	VulkanRHI::vkCmdPipelineBarrier(CommandBuffer.GetHandle(), VK_PIPELINE_STAGE_BOTTOM_OF_PIPE_BIT, VK_PIPELINE_STAGE_TRANSFER_BIT, 0, 1, &BarrierBefore, 0, nullptr, 0, nullptr);
 
-	auto CommitBuffer = [Device = Device, &CommandBuffer](FVulkanShaderTableAllocation& Alloc)
+	auto CommitBuffer = [Device = Device, &Context](FVulkanShaderTableAllocation& Alloc)
 	{
 		if (Alloc.bIsDirty)
 		{
@@ -1060,9 +1060,9 @@ void FVulkanRayTracingShaderTable::Commit(FVulkanCommandListContext& Context)
 				RegionInfo.srcOffset = 0;
 				RegionInfo.dstOffset = 0;
 				RegionInfo.size = Alloc.Region.size;
-				VulkanRHI::vkCmdCopyBuffer(CommandBuffer.GetHandle(), StagingBuffer->GetHandle(), Alloc.LocalBuffer, 1, &RegionInfo);
+				VulkanRHI::vkCmdCopyBuffer(Context.GetCommandBuffer().GetHandle(), StagingBuffer->GetHandle(), Alloc.LocalBuffer, 1, &RegionInfo);
 
-				Device->GetStagingManager().ReleaseBuffer(&CommandBuffer, StagingBuffer);
+				Device->GetStagingManager().ReleaseBuffer(&Context, StagingBuffer);
 			}
 			else
 			{
@@ -1752,31 +1752,22 @@ const TArray<uint8>& FVulkanRayTracingPipelineState::GetShaderHandles(EShaderFre
 	return GetShaderData(Frequency).ShaderHandles;
 }
 
-void FVulkanRayTracingCompactedSizeQueryPool::EndBatch(FVulkanCommandBuffer* InCmdBuffer)
+void FVulkanRayTracingCompactedSizeQueryPool::EndBatch(FVulkanCommandListContext& CommandContext)
 {
-	check(CmdBuffer == nullptr);
-	CmdBuffer = InCmdBuffer;
-	FenceSignaledCounter = InCmdBuffer->GetFenceSignaledCounter();
+	SyncPoint = CommandContext.GetContextSyncPoint();
 }
 
 void FVulkanRayTracingCompactedSizeQueryPool::Reset(FVulkanCommandBuffer* InCmdBuffer)
 {
 	VulkanRHI::vkCmdResetQueryPool(InCmdBuffer->GetHandle(), QueryPool, 0, MaxQueries);
-	FenceSignaledCounter = 0;
-	CmdBuffer = nullptr;
+	SyncPoint = nullptr;
 	check(QueryOutput.Num() == MaxQueries);
 	FMemory::Memzero(QueryOutput.GetData(), MaxQueries * sizeof(uint64));
 }
 
 bool FVulkanRayTracingCompactedSizeQueryPool::TryGetResults(uint32 NumResults)
 {
-	if (CmdBuffer == nullptr)
-	{
-		return false;
-	}
-
-	const uint64 FenceCurrentSignaledCounter = CmdBuffer->GetFenceSignaledCounter();
-	if (FenceSignaledCounter >= FenceCurrentSignaledCounter)
+	if (!SyncPoint.IsValid() || !SyncPoint->IsComplete())
 	{
 		return false;
 	}
@@ -1860,8 +1851,7 @@ void FVulkanRayTracingCompactionRequestHandler::Update(FVulkanCommandListContext
 
 			ActiveBLASes.Empty(GVulkanRayTracingMaxBatchedCompaction);
 
-			ActiveRequestsCmdBuffer = &CmdBuffer;
-			ActiveRequestsFenceCounter = CmdBuffer.GetFenceSignaledCounter();
+			ActiveRequestsSyncPoint = InCommandContext.GetContextSyncPoint();
 		}
 
 		// Only one active batch at a time (otherwise track the offset for when we launch queries)
@@ -1870,15 +1860,14 @@ void FVulkanRayTracingCompactionRequestHandler::Update(FVulkanCommandListContext
 	// If we have an active batch, wait until the compaction went through to launch another batch
 	else if (ActiveRequests.Num() > 0)
 	{
-		if (ActiveRequestsCmdBuffer)
+		if (ActiveRequestsSyncPoint.IsValid())
 		{
-			if (ActiveRequestsFenceCounter >= ActiveRequestsCmdBuffer->GetFenceSignaledCounter())
+			if (!ActiveRequestsSyncPoint->IsComplete())
 			{
 				return;
 			}
 
-			ActiveRequestsCmdBuffer = nullptr;
-			ActiveRequestsFenceCounter = MAX_uint64;
+			ActiveRequestsSyncPoint = nullptr;
 		}
 
 		ActiveRequests.Empty(GVulkanRayTracingMaxBatchedCompaction);
@@ -1906,30 +1895,22 @@ void FVulkanRayTracingCompactionRequestHandler::Update(FVulkanCommandListContext
 		// clear out all of the pending requests, don't allow the array to shrink
 		PendingRequests.RemoveAt(0, ActiveRequests.Num(), EAllowShrinking::No);
 
-		FVulkanCommandBuffer* const CmdBuffer = InCommandContext.GetActiveCmdBuffer();
-
 		// Barrier here is not stricly necessary as it is added after the build.
 		// AddAccelerationStructureBuildBarrier(CmdBuffer->GetHandle());
 
 		// Write compacted size info from the selected requests
 		VulkanDynamicAPI::vkCmdWriteAccelerationStructuresPropertiesKHR(
-			CmdBuffer->GetHandle(),
+			InCommandContext.GetCommandBuffer().GetHandle(),
 			ActiveBLASes.Num(), ActiveBLASes.GetData(),
 			VK_QUERY_TYPE_ACCELERATION_STRUCTURE_COMPACTED_SIZE_KHR,
 			QueryPool->GetHandle(),
 			0
 		);
 
-		QueryPool->EndBatch(CmdBuffer);
+		QueryPool->EndBatch(InCommandContext);
 	}
 }
 
-bool FVulkanRayTracingCompactionRequestHandler::IsUsingCmdBuffer(FVulkanCommandBuffer* CmdBuffer)
-{
-	FScopeLock Lock(&CS);
-	return (CmdBuffer == ActiveRequestsCmdBuffer) || (CmdBuffer == QueryPool->CmdBuffer);
-}
-
 
 static FVulkanPipelineBarrier SetRayGenResources(FVulkanDevice& Device, FVulkanCommandBuffer* const CmdBuffer, const FRayTracingShaderBindings& InGlobalResourceBindings, FVulkanRayTracingShaderTable* ShaderTable)
 {
@@ -2072,7 +2053,10 @@ static FVulkanPipelineBarrier SetRayGenResources(FVulkanDevice& Device, FVulkanC
 			// Because Sync2 is a prereq to ray tracing, use the conveniently generic layout
 			const VkImageLayout TargetLayout = bReadOnly ? VK_IMAGE_LAYOUT_READ_ONLY_OPTIMAL : VK_IMAGE_LAYOUT_GENERAL;
 			const FVulkanImageLayout* OriginalLayout = CmdBuffer->GetLayoutManager().GetFullLayout(Texture->Image);
-			check(OriginalLayout);
+			if (!OriginalLayout)
+			{
+				return;
+			}
 
 			// If all the subresource are already in a correct layout for the desired RendOnly state, then skip the barrier
 			if (!OriginalLayout->AreAllSubresourcesSameLayout() ||
diff --git a/Engine/Source/Runtime/VulkanRHI/Private/VulkanRayTracing.h b/Engine/Source/Runtime/VulkanRHI/Private/VulkanRayTracing.h
--- a/Engine/Source/Runtime/VulkanRHI/Private/VulkanRayTracing.h
+++ b/Engine/Source/Runtime/VulkanRHI/Private/VulkanRayTracing.h
@@ -10,6 +10,9 @@ class FVulkanCommandListContext;
 class FVulkanBuffer;
 struct FVulkanRayTracingGeometryParameters;
 
+using FVulkanSyncPoint = FGraphEvent;
+using FVulkanSyncPointRef = TRefCountPtr<FVulkanSyncPoint>;
+
 class FVulkanRayTracingPlatform
 {
 public:
@@ -317,12 +320,11 @@ public:
 		: FVulkanQueryPool(InDevice, nullptr, InMaxQueries, VK_QUERY_TYPE_ACCELERATION_STRUCTURE_COMPACTED_SIZE_KHR, false)
 	{}
 
-	void EndBatch(FVulkanCommandBuffer* InCmdBuffer);
+	void EndBatch(FVulkanCommandListContext& CommandContext);
 	bool TryGetResults(uint32 NumResults);
 	void Reset(FVulkanCommandBuffer* InCmdBuffer);
 
-	FVulkanCommandBuffer* CmdBuffer = nullptr;
-	uint64 FenceSignaledCounter = 0;
+	FVulkanSyncPointRef SyncPoint;
 };
 
 // Manages all the pending BLAS compaction requests
@@ -343,8 +345,6 @@ public:
 
 	void Update(FVulkanCommandListContext& InCommandContext);
 
-	bool IsUsingCmdBuffer(FVulkanCommandBuffer* CmdBuffer);
-
 private:
 
 	FCriticalSection CS;
@@ -353,8 +353,7 @@ private:
 
 	// Keep references on FVulkanRayTracingGeometry until lifetime issue is found (this prevents cancellation)
 	TArray<TRefCountPtr<FVulkanRayTracingGeometry>> ActiveRequests;
-	FVulkanCommandBuffer* ActiveRequestsCmdBuffer = nullptr;
-	uint64 ActiveRequestsFenceCounter = MAX_uint64;
+	FVulkanSyncPointRef ActiveRequestsSyncPoint;
 
 	FVulkanRayTracingCompactedSizeQueryPool* QueryPool = nullptr;
 };
diff --git a/Engine/Source/Runtime/VulkanRHI/Private/VulkanRenderTarget.cpp b/Engine/Source/Runtime/VulkanRHI/Private/VulkanRenderTarget.cpp
--- a/Engine/Source/Runtime/VulkanRHI/Private/VulkanRenderTarget.cpp
+++ b/Engine/Source/Runtime/VulkanRHI/Private/VulkanRenderTarget.cpp
@@ -652,14 +652,14 @@ void FVulkanCommandListContext::RHIBeginRenderPass(const FRHIRenderPassInfo& InI
 void FVulkanCommandListContext::RHIEndRenderPass()
 {
 	const bool bHasOcclusionQueries = (RenderPassInfo.NumOcclusionQueries > 0);
-	FVulkanCommandBuffer* CmdBuffer = GetActiveCmdBuffer();
 
 	if (bHasOcclusionQueries)
 	{
-		EndOcclusionQueryBatch(CmdBuffer);
+		EndOcclusionQueryBatch();
 	}
 	
-	Device.GetRenderPassManager().EndRenderPass(CmdBuffer);
+	FVulkanCommandBuffer& CommandBuffer = GetCommandBuffer();
+	Device.GetRenderPassManager().EndRenderPass(&CommandBuffer);
 
 	check(CurrentRenderPass);
 	CurrentRenderPass = nullptr;
diff --git a/Engine/Source/Runtime/VulkanRHI/Private/VulkanResources.h b/Engine/Source/Runtime/VulkanRHI/Private/VulkanResources.h
--- a/Engine/Source/Runtime/VulkanRHI/Private/VulkanResources.h
+++ b/Engine/Source/Runtime/VulkanRHI/Private/VulkanResources.h
@@ -27,12 +27,14 @@ class FVulkanLayout;
 class FVulkanOcclusionQuery;
 struct FRHITransientHeapAllocation;
 struct FVulkanPendingBufferLock;
-
 class FVulkanView;
 class FVulkanViewableResource;
 class FVulkanShaderResourceView;
 class FVulkanUnorderedAccessView;
 
+using FVulkanSyncPoint = FGraphEvent;
+using FVulkanSyncPointRef = TRefCountPtr<FVulkanSyncPoint>;
+
 namespace VulkanRHI
 {
 	class FDeviceMemoryAllocation;
@@ -842,11 +844,10 @@ public:
 		AllocatedQueries[Index] = nullptr;
 	}
 
-	inline void EndBatch(FVulkanCommandBuffer* InCmdBuffer)
+	void EndBatch()
 	{
 		ensure(State == EState::RHIT_PostBeginBatch);
 		State = EState::RHIT_PostEndBatch;
-		SetFence(InCmdBuffer);
 	}
 
 	bool CanBeReused();
@@ -888,10 +889,7 @@ protected:
 	TArray<FVulkanOcclusionQuery*> AllocatedQueries;
 	TArray<uint64> AcquiredIndices;
 	bool InternalTryGetResults(bool bWait);
-	void SetFence(FVulkanCommandBuffer* InCmdBuffer);
 
-	FVulkanCommandBuffer* CmdBuffer = nullptr;
-	uint64 FenceCounter = UINT64_MAX;
 	uint32 FrameNumber = UINT32_MAX;
 };
 
@@ -911,8 +909,7 @@ public:
 
 	struct FCmdBufferFence
 	{
-		FVulkanCommandBuffer* CmdBuffer;
-		uint64 FenceCounter;
+		FVulkanSyncPointRef SyncPoint;
 		uint64 FrameCount = UINT64_MAX;
 		uint32 Attempts = 0;
 	};
@@ -1261,8 +1258,7 @@ public:
 	virtual void Wait(FRHICommandListImmediate& RHICmdList, FRHIGPUMask GPUMask) const override;
 
 protected:
-	FVulkanCommandBuffer*	CmdBuffer = nullptr;
-	uint64				FenceSignaledCounter = MAX_uint64;
+	FVulkanSyncPointRef SyncPoint;
 
 	friend class FVulkanCommandListContext;
 };
diff --git a/Engine/Source/Runtime/VulkanRHI/Private/VulkanSubmission.cpp b/Engine/Source/Runtime/VulkanRHI/Private/VulkanSubmission.cpp
--- a/Engine/Source/Runtime/VulkanRHI/Private/VulkanSubmission.cpp
+++ b/Engine/Source/Runtime/VulkanRHI/Private/VulkanSubmission.cpp
@@ -11,6 +11,9 @@
 #include "VulkanSubmission.h"
 #include "VulkanQueue.h"
 
+#include "HAL/Runnable.h"
+#include "HAL/RunnableThread.h"
+
 
 FVulkanPayload::~FVulkanPayload()
 {
@@ -25,6 +28,101 @@ FVulkanPayload::~FVulkanPayload()
 	}
 }
 
+
+
+class FVulkanThread final : private FRunnable
+{
+	typedef bool(FVulkanDynamicRHI::* FQueueFunc)();
+
+public:
+	FVulkanThread(TCHAR const* InName, EThreadPriority InPriority, FVulkanDynamicRHI* InRHI, FQueueFunc InFunc)
+		: RHI(InRHI)
+		, Event(FPlatformProcess::GetSynchEventFromPool(false))
+		, Func(InFunc)
+		, Thread(FRunnableThread::Create(this, InName, 0, InPriority))
+	{}
+
+	virtual ~FVulkanThread()
+	{
+		bExit = true;
+		Event->Trigger();
+
+		Thread->WaitForCompletion();
+		delete Thread;
+		Thread = nullptr;
+
+		FPlatformProcess::ReturnSynchEventToPool(Event);
+		Event = nullptr;
+	}
+
+	void Kick() const
+	{
+		Event->Trigger();
+	}
+
+	void Join() const
+	{
+		Thread->WaitForCompletion();
+	}
+
+	uint32 GetThreadID() const
+	{
+		return Thread->GetThreadID();
+	}
+
+protected:
+
+	virtual uint32 Run() override
+	{
+		while (!bExit)
+		{
+			// Process the queue until no more progress is made
+			while ((RHI->*Func)());
+
+			Event->Wait();
+		}
+
+		// :todo-jn: Drain any remaining work in the queue
+
+		return 0;
+	}
+
+	FVulkanDynamicRHI* RHI;
+	std::atomic<bool> bExit = false;
+	FEvent* Event;
+	FQueueFunc Func;
+	FRunnableThread* Thread = nullptr;
+};
+
+
+
+void FVulkanDynamicRHI::InitializeSubmissionPipe()
+{
+	// :todo-jn: SubmissionThread and InterruptThread
+}
+
+void FVulkanDynamicRHI::ShutdownSubmissionPipe()
+{
+	if (SubmissionThread)
+	{
+		delete SubmissionThread;
+		SubmissionThread = nullptr;
+	}
+
+	if (InterruptThread)
+	{
+		delete InterruptThread;
+		InterruptThread = nullptr;
+	}
+
+	if (EopTask)
+	{
+		ProcessInterruptQueueUntil(EopTask);
+		EopTask = nullptr;
+	}
+}
+
+
 void FVulkanDynamicRHI::RHIFinalizeContext(FRHIFinalizeContextArgs&& Args, TRHIPipelineArray<IRHIPlatformCommandList*>& Output)
 {
 	auto GetPlatformCommandList = [&Output](ERHIPipeline Pipeline)
@@ -81,17 +179,137 @@ bool FVulkanDynamicRHI::ProcessSubmissionQueue()
 	bool bProgress = false;
 	Device->ForEachQueue([&bProgress, &SignaledSemaphores=SignaledSemaphores](FVulkanQueue& Queue)
 	{
-		bProgress = Queue.SubmitQueuedPayloads(SignaledSemaphores) > 0;
+		bProgress |= Queue.SubmitQueuedPayloads(SignaledSemaphores) > 0;
 	});
+
+	// Wake up the interrupt thread to go wait on these new payloads
+	if (InterruptThread)
+	{
+		InterruptThread->Kick();
+	}
+
 	return bProgress;
 }
 
-void FVulkanDynamicRHI::ProcessInterruptQueue(uint64 Timeout)
+bool FVulkanDynamicRHI::ProcessInterruptQueue(bool bWaitForProgress)
 {
-	Device->ForEachQueue([Timeout](FVulkanQueue& Queue)
+	if (bWaitForProgress)
+	{
+		const int32 NumQueues = (int32)EVulkanQueueType::Count;
+
+		// Pick up the next payload for each queue
+		// NOTE: we can hold on to these because we're holding InterruptCS or we're the interrupt thread.
+		FVulkanPayload* NextPayloads[NumQueues];
+		uint32 NumSyncs = 0;
+		Device->ForEachQueue([&NumSyncs, &NextPayloads](FVulkanQueue& Queue)
+		{
+			FVulkanPayload* Payload = Queue.GetNextInterruptPayload();
+			if (Payload)
+			{
+				NextPayloads[NumSyncs++] = Payload;
+			}
+		});
+
+		if (NumSyncs)
+		{
+			const uint64 Timeout = 10ULL * 1000 * 1000;
+
+			// Figure out if we wait on Fences or TimelineSemaphores
+			if (NextPayloads[0]->Fence)
+			{
+				VulkanRHI::FFence* Fences[NumQueues];
+				int32 NumFences = 0;
+				for (uint32 Index = 0; Index < NumSyncs; ++Index)
+				{
+					VulkanRHI::FFence* Fence = NextPayloads[Index]->Fence;
+					checkf(Fence, TEXT("Payloads should all use the same types of syncs!"));
+					if (!Fence->IsSignaled())
+					{
+						Fences[NumFences++] = Fence;
+					}
+				}
+
+				// if one of these fences is already signaled, then that guarantees progress (no wait necessary on this pass)
+				if (NumFences == NumSyncs)
+				{
+					Device->GetFenceManager().WaitForAnyFence(MakeArrayView<VulkanRHI::FFence*>(Fences, NumFences), Timeout);
+				}
+			}
+			else
+			{
+				VkSemaphore Semaphores[NumQueues];
+				uint64_t Values[NumQueues];
+				for (uint32 Index = 0; Index < NumSyncs; ++Index)
+				{
+					checkf(NextPayloads[Index]->TimelineSemaphoreValue > 0, TEXT("Payloads should all use the same types of syncs!"));
+					Semaphores[Index] = NextPayloads[Index]->Queue.GetTimelineSemaphore()->GetHandle();
+					Values[Index] = (uint64_t)NextPayloads[Index]->TimelineSemaphoreValue;
+				}
+
+				VkSemaphoreWaitInfo WaitInfo;
+				ZeroVulkanStruct(WaitInfo, VK_STRUCTURE_TYPE_SEMAPHORE_WAIT_INFO);
+				WaitInfo.flags = VK_SEMAPHORE_WAIT_ANY_BIT;
+				WaitInfo.semaphoreCount = NumSyncs;
+				WaitInfo.pSemaphores = Semaphores;
+				WaitInfo.pValues = Values;
+				VERIFYVULKANRESULT_EXPANDED(VulkanRHI::vkWaitSemaphoresKHR(Device->GetInstanceHandle(), &WaitInfo, (uint64_t)Timeout));
+			}
+		}
+	}
+
+	bool bProgress = false;
+	Device->ForEachQueue([&bProgress](FVulkanQueue& Queue)
 	{
-		Queue.ProcessInterruptQueue(Timeout);
+		bProgress |= Queue.ProcessInterruptQueue(0) > 0;
 	});
+	return bProgress;
+}
+
+void FVulkanDynamicRHI::ProcessInterruptQueueUntil(FGraphEvent* GraphEvent)
+{
+	if (InterruptThread)
+	{
+		InterruptThread->Kick();
+		if (GraphEvent && !GraphEvent->IsComplete())
+		{
+			GraphEvent->Wait();
+		}
+	}
+	else if (GraphEvent)
+	{
+		// If we're waiting for an event, accumulate the idle time
+		UE::Stats::FThreadIdleStats::FScopeIdle IdleScope;
+
+		// If we have a sync point, keep processing until the sync point is signaled.
+		while (!GraphEvent->IsComplete())
+		{
+			if (InterruptCS.TryLock())
+			{
+				ProcessInterruptQueue(true);
+				InterruptCS.Unlock();
+			}
+			else
+			{
+				// Failed to get the lock. Another thread is processing the interrupt queue. Try again...
+				FPlatformProcess::SleepNoStats(0);
+			}
+		}
+	}
+	else
+	{
+		// Process until no more progress is being made, or skip if another thread is processing already
+		bool bProgress;
+		do
+		{
+			bProgress = false;
+			if (InterruptCS.TryLock())
+			{
+				bProgress = ProcessInterruptQueue(false);
+				InterruptCS.Unlock();
+			}
+
+		} while (bProgress);
+	}
 }
 
 void FVulkanDynamicRHI::RHISubmitCommandLists(FRHISubmitCommandListsArgs&& Args)
@@ -102,31 +320,49 @@ void FVulkanDynamicRHI::RHISubmitCommandLists(FRHISubmitCommandListsArgs&& Args)
 		PendingPayloadsForSubmission.Enqueue(PlatformCmdList);
 	}
 
-	// :todo-jn: add submission thread
+	if (SubmissionThread)
+	{
+		SubmissionThread->Kick();
+	}
+	else
 	{
 		FScopeLock Lock(&SubmissionCS);
 		while (ProcessSubmissionQueue()) {};
 	}
 
-	// :todo-jn: add interrupt thread
-	{
-		if (InterruptCS.TryLock())
-		{
-			// Timeout==0 will complete all the payloads it can without waiting
-			ProcessInterruptQueue(0ULL);
-			InterruptCS.Unlock();
-		}
-	}
+	ProcessInterruptQueueUntil(nullptr);
 }
 
-void FVulkanDynamicRHI::WaitForCmdBuffer(FVulkanCommandBuffer* CmdBuffer, float TimeInSecondsToWait)
+void FVulkanDynamicRHI::EnqueueEndOfPipeTask(TUniqueFunction<void()> TaskFunc, TUniqueFunction<void(FVulkanPayload&)> ModifyPayloadCallback)
 {
-	FVulkanQueue& Queue = CmdBuffer->GetOwner()->GetContext().Queue;
-	const uint64 Timeout = (uint64)(TimeInSecondsToWait * 1e9);
-	checkf(CmdBuffer->IsSubmitted(), TEXT("Attempting to wait on a command buffer that was not submitted!"));
-
+	FGraphEventArray Prereqs;
+	Prereqs.Reserve((int32)EVulkanQueueType::Count + 1);
+	if (EopTask.IsValid())
 	{
-		FScopeLock Lock(&InterruptCS);
-		Queue.ProcessInterruptQueue(Timeout, CmdBuffer->GetSubmittedTimelineValue());
+		Prereqs.Add(EopTask);
 	}
-}
+
+	FVulkanPlatformCommandList* Payloads = new FVulkanPlatformCommandList;
+	Payloads->Reserve((int32)EVulkanQueueType::Count);
+
+	Device->ForEachQueue([&](FVulkanQueue& Queue)
+	{
+		FVulkanPayload* Payload = new FVulkanPayload(Queue);
+
+		Payload->SyncPoint = CreateVulkanSyncPoint();
+		Prereqs.Add(Payload->SyncPoint);
+
+		if (ModifyPayloadCallback)
+			ModifyPayloadCallback(*Payload);
+
+		Payloads->Add(Payload);
+	});
+
+	PendingPayloadsForSubmission.Enqueue(Payloads);
+
+	EopTask = FFunctionGraphTask::CreateAndDispatchWhenReady(
+		MoveTemp(TaskFunc),
+		QUICK_USE_CYCLE_STAT(FExecuteRHIThreadTask, STATGROUP_TaskGraphTasks),
+		&Prereqs
+	);
+}
\ No newline at end of file
diff --git a/Engine/Source/Runtime/VulkanRHI/Private/VulkanSubmission.h b/Engine/Source/Runtime/VulkanRHI/Private/VulkanSubmission.h
--- a/Engine/Source/Runtime/VulkanRHI/Private/VulkanSubmission.h
+++ b/Engine/Source/Runtime/VulkanRHI/Private/VulkanSubmission.h
@@ -16,6 +16,14 @@ class FVulkanPendingComputeState;
 class FVulkanQueue;
 class FVulkanSwapChain;
 
+// NOTE: typedef to FGraphEvent for now, but eventually move all syncs inside a FVulkanSyncPoint class
+using FVulkanSyncPoint = FGraphEvent;
+using FVulkanSyncPointRef = TRefCountPtr<FVulkanSyncPoint>;
+
+static FVulkanSyncPointRef CreateVulkanSyncPoint()
+{
+	return FGraphEvent::CreateGraphEvent();
+}
 
 class FVulkanPayload
 {
@@ -37,9 +45,12 @@ protected:
 	TArray<FVulkanCommandBuffer*> CommandBuffers;
 	TArray<VulkanRHI::FSemaphore*> SignalSemaphores; // signaled after command buffers
 
-	// For completion tracking
+	// For internal completion tracking of the payload
 	uint64 TimelineSemaphoreValue = 0;
 	VulkanRHI::FFence* Fence = nullptr;
+
+	// Used to sync other CPU work to payload completion
+	FVulkanSyncPointRef SyncPoint;
 };
 
 struct FVulkanPlatformCommandList : public IRHIPlatformCommandList, public TArray<FVulkanPayload*>
diff --git a/Engine/Source/Runtime/VulkanRHI/Private/VulkanSwapChain.cpp b/Engine/Source/Runtime/VulkanRHI/Private/VulkanSwapChain.cpp
--- a/Engine/Source/Runtime/VulkanRHI/Private/VulkanSwapChain.cpp
+++ b/Engine/Source/Runtime/VulkanRHI/Private/VulkanSwapChain.cpp
@@ -9,6 +9,7 @@
 #include "VulkanDevice.h"
 #include "VulkanPlatform.h"
 #include "VulkanRHIPrivate.h"
+#include "VulkanSubmission.h"
 #include "Engine/RendererSettings.h"
 #include "HAL/PlatformFramePacer.h"
 #include "IHeadMountedDisplayModule.h"
@@ -792,8 +793,6 @@ FVulkanSwapChain::EStatus FVulkanSwapChain::Present(FVulkanQueue* GfxQueue, FVul
 {
 	check(CurrentImageIndex != -1);
 
-	//ensure(GfxQueue == PresentQueue);
-
 	VkPresentInfoKHR Info;
 	ZeroVulkanStruct(Info, VK_STRUCTURE_TYPE_PRESENT_INFO_KHR);
 	VkSemaphore Semaphore = VK_NULL_HANDLE;
@@ -896,27 +895,27 @@ void FVulkanDevice::SetupPresentQueue(VkSurfaceKHR Surface)
 			return (bSupportsPresent == VK_TRUE);
 		};
 
-		bool bGfx = SupportsPresent(Gpu, GfxQueue);
+		bool bGfx = SupportsPresent(Gpu, GetGraphicsQueue());
 		if (!bGfx)
 		{
 			FPlatformMisc::MessageBoxExt(EAppMsgType::Ok, TEXT("Cannot find a compatible Vulkan device that supports surface presentation.\n\n"), TEXT("Vulkan device not available"));
 			FPlatformMisc::RequestExitWithStatus(true, 1);
 		}
 
-		bool bCompute = SupportsPresent(Gpu, ComputeQueue);
-		if (TransferQueue->GetFamilyIndex() != GfxQueue->GetFamilyIndex() && TransferQueue->GetFamilyIndex() != ComputeQueue->GetFamilyIndex())
+		bool bCompute = SupportsPresent(Gpu, GetComputeQueue());
+		if (GetTransferQueue()->GetFamilyIndex() != GetGraphicsQueue()->GetFamilyIndex() && GetTransferQueue()->GetFamilyIndex() != GetComputeQueue()->GetFamilyIndex())
 		{
-			SupportsPresent(Gpu, TransferQueue);
+			SupportsPresent(Gpu, GetTransferQueue());
 		}
-		if (GAllowPresentOnComputeQueue.GetValueOnAnyThread() != 0 && ComputeQueue->GetFamilyIndex() != GfxQueue->GetFamilyIndex() && bCompute)
+		if (GAllowPresentOnComputeQueue.GetValueOnAnyThread() != 0 && GetComputeQueue()->GetFamilyIndex() != GetGraphicsQueue()->GetFamilyIndex() && bCompute)
 		{
 			//#todo-rco: Do other IHVs have a fast path here?
 			bPresentOnComputeQueue = (VendorId == EGpuVendorId::Amd);
-			PresentQueue = ComputeQueue;
+			PresentQueue = GetComputeQueue();
 		}
 		else
 		{
-			PresentQueue = GfxQueue;
+			PresentQueue = GetGraphicsQueue();
 		}
 	}
 }
diff --git a/Engine/Source/Runtime/VulkanRHI/Private/VulkanTexture.cpp b/Engine/Source/Runtime/VulkanRHI/Private/VulkanTexture.cpp
--- a/Engine/Source/Runtime/VulkanRHI/Private/VulkanTexture.cpp
+++ b/Engine/Source/Runtime/VulkanRHI/Private/VulkanTexture.cpp
@@ -126,7 +126,7 @@ void FVulkanTexture::InternalLockWrite(FVulkanContextCommon& Context, FVulkanTex
 		Barrier.Execute(CmdBuffer);
 	}
 
-	Surface->Device->GetStagingManager().ReleaseBuffer(CmdBuffer, StagingBuffer);
+	Surface->Device->GetStagingManager().ReleaseBuffer(&Context, StagingBuffer);
 }
 
 void FVulkanTexture::ErrorInvalidViewType() const
diff --git a/Engine/Source/Runtime/VulkanRHI/Private/VulkanUniformBuffer.cpp b/Engine/Source/Runtime/VulkanRHI/Private/VulkanUniformBuffer.cpp
--- a/Engine/Source/Runtime/VulkanRHI/Private/VulkanUniformBuffer.cpp
+++ b/Engine/Source/Runtime/VulkanRHI/Private/VulkanUniformBuffer.cpp
@@ -50,7 +50,7 @@ static void UpdateUniformBufferHelper(FVulkanCommandListContext& Context, FVulka
 	const int32 DataAlignment = FMath::Max<uint32>(Device.GetLimits().minUniformBufferOffsetAlignment, 16u);
 
 	VulkanRHI::FVulkanAllocation TempAllocation;
-	void* DestinationData = Context.GetTempBlockAllocator().Alloc(DataSize, DataAlignment, CmdBuffer, TempAllocation);
+	void* DestinationData = Context.GetTempBlockAllocator().Alloc(DataSize, DataAlignment, Context, TempAllocation);
 
 	if (bUpdateConstants)
 	{
diff --git a/Engine/Source/Runtime/VulkanRHI/Private/VulkanUtil.cpp b/Engine/Source/Runtime/VulkanRHI/Private/VulkanUtil.cpp
--- a/Engine/Source/Runtime/VulkanRHI/Private/VulkanUtil.cpp
+++ b/Engine/Source/Runtime/VulkanRHI/Private/VulkanUtil.cpp
@@ -141,8 +141,7 @@ void FVulkanGPUTiming::StartTiming(FVulkanCommandBuffer* CmdBuffer)
 		}
 
 		VulkanRHI::vkCmdWriteTimestamp(CmdBuffer->GetHandle(), VK_PIPELINE_STAGE_BOTTOM_OF_PIPE_BIT, Pool->GetHandle(), QueryStartIndex);
-		Pool->TimestampListHandles[QueryStartIndex].CmdBuffer = CmdBuffer;
-		Pool->TimestampListHandles[QueryStartIndex].FenceCounter = CmdBuffer->GetFenceSignaledCounter();
+		Pool->TimestampListHandles[QueryStartIndex].SyncPoint = Context->GetContextSyncPoint();
 		Pool->TimestampListHandles[QueryStartIndex].FrameCount = Context->GetFrameCounter();
 		Pool->TimestampListHandles[QueryStartIndex].Attempts = 0;
 		bIsTiming = true;
@@ -172,10 +171,8 @@ void FVulkanGPUTiming::EndTiming(FVulkanCommandBuffer* CmdBuffer)
 
 		VulkanRHI::vkCmdWriteTimestamp(CmdBuffer->GetHandle(), VK_PIPELINE_STAGE_BOTTOM_OF_PIPE_BIT, Pool->GetHandle(), QueryEndIndex);
 		CmdBuffer->AddPendingTimestampQuery(QueryStartIndex, 2, Pool->GetHandle(), Pool->ResultsBuffer->GetHandle(), false);
-		Pool->TimestampListHandles[QueryStartIndex].CmdBuffer = CmdBuffer;
-		Pool->TimestampListHandles[QueryStartIndex].FenceCounter = CmdBuffer->GetFenceSignaledCounter();
-		Pool->TimestampListHandles[QueryEndIndex].CmdBuffer = CmdBuffer;
-		Pool->TimestampListHandles[QueryEndIndex].FenceCounter = CmdBuffer->GetFenceSignaledCounter();
+		Pool->TimestampListHandles[QueryStartIndex].SyncPoint = Context->GetContextSyncPoint();
+		Pool->TimestampListHandles[QueryEndIndex].SyncPoint = Context->GetContextSyncPoint();
 		Pool->TimestampListHandles[QueryEndIndex].FrameCount = Context->GetFrameCounter();
 		Pool->TimestampListHandles[QueryEndIndex].Attempts = 0;
 		Pool->NumIssuedTimestamps = FMath::Min<uint32>(Pool->NumIssuedTimestamps + 1, Pool->BufferSize);
@@ -192,11 +189,7 @@ void FVulkanGPUTiming::DiscardExpiredQueries(bool bNeedSpace)
 	{
 		FVulkanTimingQueryPool::FCmdBufferFence& QuerySyncPoint = Pool->TimestampListHandles[QueryIndex];
 
-		if (!QuerySyncPoint.CmdBuffer)
-		{
-			return true;
-		}
-		if (QuerySyncPoint.FenceCounter < QuerySyncPoint.CmdBuffer->GetFenceSignaledCounter())
+		if (QuerySyncPoint.SyncPoint->IsComplete())
 		{
 			const uint64 Availability = Pool->MappedPointer[QueryIndex * 2 + 1];
 			if (!Availability && (QuerySyncPoint.Attempts >= TIMING_QUERY_RETRIES))
@@ -214,10 +207,11 @@ void FVulkanGPUTiming::DiscardExpiredQueries(bool bNeedSpace)
 		{
 			// If the hold back for a query is the fence and the command buffer was submitted, take a shot at updating its status
 			FVulkanTimingQueryPool::FCmdBufferFence& QuerySyncPoint = Pool->TimestampListHandles[QueryIndex];
-			if ((QuerySyncPoint.FenceCounter >= QuerySyncPoint.CmdBuffer->GetFenceSignaledCounter()) && QuerySyncPoint.CmdBuffer->IsSubmitted())
+			if (!QuerySyncPoint.SyncPoint->IsComplete())
 			{
-				FVulkanDynamicRHI::Get().WaitForCmdBuffer(QuerySyncPoint.CmdBuffer, 0.0f);
-				if (QuerySyncPoint.FenceCounter < QuerySyncPoint.CmdBuffer->GetFenceSignaledCounter())
+				FVulkanDynamicRHI::Get().ProcessInterruptQueueUntil(nullptr);  // leave null, we don't want to force a wait on the SyncPoint
+
+				if (QuerySyncPoint.SyncPoint->IsComplete())
 				{
 					QueryExpired = CanBeDiscarded(QueryIndex);
 				}
@@ -256,15 +250,14 @@ void FVulkanGPUTiming::RefreshPendingQueries(FVulkanCommandBuffer& CmdBuffer)
 	{
 		FVulkanTimingQueryPool::FCmdBufferFence& QuerySyncPoint = Pool->TimestampListHandles[QueryIndex];
 
-		if (QuerySyncPoint.CmdBuffer && (QuerySyncPoint.FenceCounter < QuerySyncPoint.CmdBuffer->GetFenceSignaledCounter()))
+		if (QuerySyncPoint.SyncPoint.IsValid() && QuerySyncPoint.SyncPoint->IsComplete())
 		{
 			Pool->ResultsBuffer->InvalidateMappedMemory();
 			const uint64 Availability = Pool->MappedPointer[QueryIndex * 2 + 1];
 
 			if (!Availability && (QuerySyncPoint.Attempts < TIMING_QUERY_RETRIES))
 			{
-				QuerySyncPoint.CmdBuffer = &CmdBuffer;
-				QuerySyncPoint.FenceCounter = CmdBuffer.GetFenceSignaledCounter();
+				// Keep the same QuerySyncPoint.SyncPoint
 				QuerySyncPoint.FrameCount = Context->GetFrameCounter();
 				++QuerySyncPoint.Attempts;
 
@@ -332,8 +325,7 @@ uint64 FVulkanGPUTiming::GetTiming(bool bGetCurrentResultsAndBlock)
 				}
 			}
 
-			if (EndQuerySyncPoint.FenceCounter < EndQuerySyncPoint.CmdBuffer->GetFenceSignaledCounter() &&
-				StartQuerySyncPoint.FenceCounter < StartQuerySyncPoint.CmdBuffer->GetFenceSignaledCounter())
+			if (EndQuerySyncPoint.SyncPoint->IsComplete() && StartQuerySyncPoint.SyncPoint->IsComplete())
 			{
 				Pool->ResultsBuffer->InvalidateMappedMemory();
 				StartTimeAvailability = Pool->MappedPointer[QueryStartIndex * 2 + 1];
@@ -364,22 +356,14 @@ uint64 FVulkanGPUTiming::GetTiming(bool bGetCurrentResultsAndBlock)
 				{
 					FRenderThreadIdleScope IdleScope(ERenderThreadIdleTypes::WaitingForGPUQuery);
 
-					const bool bWaitForStart = StartQuerySyncPoint.FenceCounter >= StartQuerySyncPoint.CmdBuffer->GetFenceSignaledCounter();
-					const bool bWaitForEnd = EndQuerySyncPoint.FenceCounter >= EndQuerySyncPoint.CmdBuffer->GetFenceSignaledCounter();
-					if (bWaitForEnd || bWaitForStart)
-					{
-						// Need to submit the open command lists.
-						Device->SubmitCommandsAndFlushGPU();  // :todo-jn: deprecated
-					}
-
 					// CPU wait for query results to be ready.
-					if (bWaitForStart && StartQuerySyncPoint.FenceCounter >= StartQuerySyncPoint.CmdBuffer->GetFenceSignaledCounter())
+					if (!StartQuerySyncPoint.SyncPoint->IsComplete())
 					{
-						FVulkanDynamicRHI::Get().WaitForCmdBuffer(StartQuerySyncPoint.CmdBuffer);
+						FVulkanDynamicRHI::Get().ProcessInterruptQueueUntil(StartQuerySyncPoint.SyncPoint);
 					}
-					if (bWaitForEnd && EndQuerySyncPoint.FenceCounter >= EndQuerySyncPoint.CmdBuffer->GetFenceSignaledCounter())
+					if (!EndQuerySyncPoint.SyncPoint->IsComplete())
 					{
-						FVulkanDynamicRHI::Get().WaitForCmdBuffer(EndQuerySyncPoint.CmdBuffer);
+						FVulkanDynamicRHI::Get().ProcessInterruptQueueUntil(EndQuerySyncPoint.SyncPoint);
 					}
 				}
 
@@ -755,20 +739,19 @@ void FVulkanDynamicRHI::RHIUnlockStagingBuffer(FRHIStagingBuffer* StagingBufferR
 
 void FVulkanGPUFence::Clear()
 {
-	CmdBuffer = nullptr;
-	FenceSignaledCounter = MAX_uint64;
+	SyncPoint = nullptr;
 }
 
 bool FVulkanGPUFence::Poll() const
 {
-	return (CmdBuffer && (FenceSignaledCounter < CmdBuffer->GetFenceSignaledCounter()));
+	return (SyncPoint.IsValid() && SyncPoint->IsComplete());
 }
 
 void FVulkanGPUFence::Wait(FRHICommandListImmediate& RHICmdList, FRHIGPUMask GPUMask) const
 {
 	if (!Poll())
 	{
-		FVulkanDynamicRHI::Get().WaitForCmdBuffer(CmdBuffer);
+		FVulkanDynamicRHI::Get().ProcessInterruptQueueUntil(SyncPoint);
 	}
 }
 
diff --git a/Engine/Source/Runtime/VulkanRHI/Private/VulkanViewport.cpp b/Engine/Source/Runtime/VulkanRHI/Private/VulkanViewport.cpp
--- a/Engine/Source/Runtime/VulkanRHI/Private/VulkanViewport.cpp
+++ b/Engine/Source/Runtime/VulkanRHI/Private/VulkanViewport.cpp
@@ -281,16 +281,12 @@ void FVulkanViewport::WaitForFrameEventCompletion()
 	{
 		static FCriticalSection CS;
 		FScopeLock ScopeLock(&CS);
-		if (LastFrameCommandBuffer && LastFrameCommandBuffer->IsSubmitted())
+		if (LastFrameSyncPoint.IsValid())
 		{
 			// If last frame's fence hasn't been signaled already, wait for it here
-			if (LastFrameFenceCounter == LastFrameCommandBuffer->GetFenceSignaledCounter())
+			if (!LastFrameSyncPoint->IsComplete())
 			{
-				if (!GWaitForIdleOnSubmit)
-				{
-					// The wait has already happened if GWaitForIdleOnSubmit is set
-					FVulkanDynamicRHI::Get().WaitForCmdBuffer(LastFrameCommandBuffer);
-				}
+				FVulkanDynamicRHI::Get().ProcessInterruptQueueUntil(LastFrameSyncPoint);
 			}
 		}
 	}
@@ -300,11 +296,9 @@ void FVulkanViewport::IssueFrameEvent()
 {
 	if (FVulkanPlatform::RequiresWaitingForFrameCompletionEvent())
 	{
-		// The fence we need to wait on next frame is already there in the command buffer
-		// that was just submitted in this frame's Present. Just grab that command buffer's
-		// info to use next frame in WaitForFrameEventCompletion.
-		FVulkanQueue* Queue = Device->GetGraphicsQueue();
-		Queue->GetLastSubmittedInfo(LastFrameCommandBuffer, LastFrameFenceCounter);
+		FVulkanCommandListContextImmediate& ImmediateContext = Device->GetImmediateContext();
+		LastFrameSyncPoint = ImmediateContext.GetContextSyncPoint();
+		Device->SubmitCommandsAndFlushGPU();
 	}
 }
 
@@ -1146,6 +1140,7 @@ bool FVulkanViewport::Present(FVulkanCommandListContext& Context, FVulkanQueue*
 				// Skip present silently if image has not been acquired
 				return (int32)FVulkanSwapChain::EStatus::Healthy;
 			}
+		
 			return (int32)Viewport->SwapChain->Present(Queue, PresentQueue, Viewport->RenderingDoneSemaphores[Viewport->AcquiredImageIndex]);
 		};
 		if (SupportsStandardSwapchain() && !DoCheckedSwapChainJob(Context, SwapChainJob))
diff --git a/Engine/Source/Runtime/VulkanRHI/Private/VulkanViewport.h b/Engine/Source/Runtime/VulkanRHI/Private/VulkanViewport.h
--- a/Engine/Source/Runtime/VulkanRHI/Private/VulkanViewport.h
+++ b/Engine/Source/Runtime/VulkanRHI/Private/VulkanViewport.h
@@ -139,8 +139,7 @@ protected:
 
 	FCustomPresentRHIRef CustomPresent;
 
-	FVulkanCommandBuffer* LastFrameCommandBuffer = nullptr;
-	uint64 LastFrameFenceCounter = 0;
+	FVulkanSyncPointRef LastFrameSyncPoint;
 
 	EDeviceScreenOrientation CachedOrientation = EDeviceScreenOrientation::Unknown;
 	void OnSystemResolutionChanged(uint32 ResX, uint32 ResY);
diff --git a/Engine/Source/Runtime/VulkanRHI/Public/VulkanDynamicRHI.h b/Engine/Source/Runtime/VulkanRHI/Public/VulkanDynamicRHI.h
--- a/Engine/Source/Runtime/VulkanRHI/Public/VulkanDynamicRHI.h
+++ b/Engine/Source/Runtime/VulkanRHI/Public/VulkanDynamicRHI.h
@@ -16,6 +16,7 @@ class FVulkanUniformBuffer;
 class FVulkanViewport;
 class FVulkanPayload;
 class IHeadMountedDisplayVulkanExtensions;
+class FVulkanThread;
 struct FRHITransientHeapAllocation;
 struct FVulkanPlatformCommandList;
 struct FOptionalVulkanInstanceExtensions
@@ -275,16 +276,22 @@ public:
 public:
 	static void SavePipelineCache();
 
-	void WaitForCmdBuffer(FVulkanCommandBuffer* CmdBuffer, float TimeInSecondsToWait=30.0f);
+	void ProcessInterruptQueueUntil(FGraphEvent* GraphEvent);
+	void EnqueueEndOfPipeTask(TUniqueFunction<void()> TaskFunc, TUniqueFunction<void(FVulkanPayload&)> ModifyPayloadCallback = {});
 
 protected:
+	void InitializeSubmissionPipe();
+	void ShutdownSubmissionPipe();
 	bool ProcessSubmissionQueue();
-	void ProcessInterruptQueue(uint64 Timeout);
+	bool ProcessInterruptQueue(bool bWaitForProgress);
 
 	FCriticalSection SubmissionCS;
+	FVulkanThread* SubmissionThread = nullptr;
 	FCriticalSection InterruptCS;
+	FVulkanThread* InterruptThread = nullptr;
 	TQueue<FVulkanPlatformCommandList*, EQueueMode::Mpsc> PendingPayloadsForSubmission;
 	TSet<VkSemaphore> SignaledSemaphores;
+	FGraphEventRef EopTask;
 
 	uint32 ApiVersion = 0;
 	VkInstance Instance;
@@ -301,8 +308,6 @@ protected:
 
 	void CreateInstance();
 	void SelectDevice();
-	void InitGPU(FVulkanDevice* Device);
-	void InitDevice(FVulkanDevice* Device);
 
 	friend class FVulkanCommandListContext;
 
