VulkanRHI: Clean Up:  Remove unused parameters.

#rnx

diff --git a/Engine/Source/Runtime/VulkanRHI/Private/VulkanRHI.cpp b/Engine/Source/Runtime/VulkanRHI/Private/VulkanRHI.cpp
--- a/Engine/Source/Runtime/VulkanRHI/Private/VulkanRHI.cpp
+++ b/Engine/Source/Runtime/VulkanRHI/Private/VulkanRHI.cpp
@@ -1170,7 +1170,7 @@ void FVulkanCommandListContext::RHIEndDrawingViewport(FRHIViewport* ViewportRHI,
 
 	WriteEndTimestamp(&CommandBuffer);
 
-	bool bNativePresent = Viewport->Present(*this, &Queue, Device.GetPresentQueue(), bLockToVsync);
+	bool bNativePresent = Viewport->Present(*this, Device.GetPresentQueue(), bLockToVsync);
 	if (bNativePresent)
 	{
 		//#todo-rco: Check for r.FinishCurrentFrame
diff --git a/Engine/Source/Runtime/VulkanRHI/Private/VulkanSwapChain.cpp b/Engine/Source/Runtime/VulkanRHI/Private/VulkanSwapChain.cpp
--- a/Engine/Source/Runtime/VulkanRHI/Private/VulkanSwapChain.cpp
+++ b/Engine/Source/Runtime/VulkanRHI/Private/VulkanSwapChain.cpp
@@ -789,7 +789,7 @@ void FVulkanSwapChain::RenderThreadPacing()
 	}
 }
 
-FVulkanSwapChain::EStatus FVulkanSwapChain::Present(FVulkanQueue* GfxQueue, FVulkanQueue* PresentQueue, VulkanRHI::FSemaphore* BackBufferRenderingDoneSemaphore)
+FVulkanSwapChain::EStatus FVulkanSwapChain::Present(FVulkanQueue* PresentQueue, VulkanRHI::FSemaphore* BackBufferRenderingDoneSemaphore)
 {
 	check(CurrentImageIndex != -1);
 
diff --git a/Engine/Source/Runtime/VulkanRHI/Private/VulkanSwapChain.h b/Engine/Source/Runtime/VulkanRHI/Private/VulkanSwapChain.h
--- a/Engine/Source/Runtime/VulkanRHI/Private/VulkanSwapChain.h
+++ b/Engine/Source/Runtime/VulkanRHI/Private/VulkanSwapChain.h
@@ -45,7 +45,7 @@ public:
 		OutOfDate = -1,
 		SurfaceLost = -2,
 	};
-	EStatus Present(FVulkanQueue* GfxQueue, FVulkanQueue* PresentQueue, VulkanRHI::FSemaphore* BackBufferRenderingDoneSemaphore);
+	EStatus Present(FVulkanQueue* PresentQueue, VulkanRHI::FSemaphore* BackBufferRenderingDoneSemaphore);
 
 	void RenderThreadPacing();
 	inline int8 DoesLockToVsync() { return LockToVsync; }
diff --git a/Engine/Source/Runtime/VulkanRHI/Private/VulkanViewport.cpp b/Engine/Source/Runtime/VulkanRHI/Private/VulkanViewport.cpp
--- a/Engine/Source/Runtime/VulkanRHI/Private/VulkanViewport.cpp
+++ b/Engine/Source/Runtime/VulkanRHI/Private/VulkanViewport.cpp
@@ -888,11 +888,13 @@ void FVulkanViewport::DestroySwapchain(FVulkanSwapChainRecreateInfo* RecreateInf
 	AcquiredImageIndex = -1;
 }
 
-inline static void CopyImageToBackBuffer(FVulkanCommandListContext& Context, FVulkanCommandBuffer* CmdBuffer, FVulkanTexture& SrcSurface, FVulkanTexture& DstSurface, int32 SizeX, int32 SizeY, int32 WindowSizeX, int32 WindowSizeY, VkSurfaceTransformFlagBitsKHR CachedSurfaceTransform)
+inline static void CopyImageToBackBuffer(FVulkanCommandListContext& Context, FVulkanTexture& SrcSurface, FVulkanTexture& DstSurface, int32 SizeX, int32 SizeY, int32 WindowSizeX, int32 WindowSizeY, VkSurfaceTransformFlagBitsKHR CachedSurfaceTransform)
 {
 	RHI_BREADCRUMB_EVENT(Context, "CopyImageToBackBuffer");
 	const bool bNeedsVulkanPreTransform = CachedSurfaceTransform != VK_SURFACE_TRANSFORM_IDENTITY_BIT_KHR;
 
+	FVulkanCommandBuffer* CmdBuffer = &Context.GetCommandBuffer();
+	check(CmdBuffer->IsOutsideRenderPass());
 	FVulkanLayoutManager& LayoutManager = CmdBuffer->GetLayoutManager();
 	const FVulkanImageLayout* SrcImageLayout = LayoutManager.GetFullLayout(SrcSurface);
 
@@ -1015,63 +1017,60 @@ inline static void CopyImageToBackBuffer(FVulkanCommandListContext& Context, FVu
 	}
 }
 
-bool FVulkanViewport::Present(FVulkanCommandListContext& Context, FVulkanQueue* Queue, FVulkanQueue* PresentQueue, bool bLockToVsync)
+bool FVulkanViewport::Present(FVulkanCommandListContext& Context, FVulkanQueue* PresentQueue, bool bLockToVsync)
 {
 	check(Context.IsImmediate());
 
 	FPlatformAtomics::AtomicStore(&LockToVsync, bLockToVsync ? 1 : 0);
 	bool bFailedToDelayAcquireBackbuffer = false;
 
+	//Transition back buffer to presentable and submit that command
+	if (SupportsStandardSwapchain())
 	{
-		FVulkanCommandBuffer& CommandBuffer = Context.GetCommandBuffer();
-
-		//Transition back buffer to presentable and submit that command
-		check(CommandBuffer.IsOutsideRenderPass());
-
-		if (SupportsStandardSwapchain())
+		if (GVulkanDelayAcquireImage == EDelayAcquireImageType::DelayAcquire && RenderingBackBuffer)
 		{
-			if (GVulkanDelayAcquireImage == EDelayAcquireImageType::DelayAcquire && RenderingBackBuffer)
+			SCOPE_CYCLE_COUNTER(STAT_VulkanAcquireBackBuffer);
+			// swapchain can go out of date, do not crash at this point
+			if (LIKELY(TryAcquireImageIndex()))
 			{
-				SCOPE_CYCLE_COUNTER(STAT_VulkanAcquireBackBuffer);
-				// swapchain can go out of date, do not crash at this point
-				if (LIKELY(TryAcquireImageIndex()))
-				{
-					uint32 WindowSizeX = FMath::Min(SizeX, SwapChain->InternalWidth);
-					uint32 WindowSizeY = FMath::Min(SizeY, SwapChain->InternalHeight);
+				uint32 WindowSizeX = FMath::Min(SizeX, SwapChain->InternalWidth);
+				uint32 WindowSizeY = FMath::Min(SizeY, SwapChain->InternalHeight);
 
-					CopyImageToBackBuffer(Context, &CommandBuffer, *RenderingBackBuffer.GetReference(), *BackBufferImages[AcquiredImageIndex].GetReference(), SizeX, SizeY, WindowSizeX, WindowSizeY, SwapChain->GetCachedSurfaceTransform());
-				}
-				else
-				{
-					bFailedToDelayAcquireBackbuffer = true;
-				}
+				CopyImageToBackBuffer(Context, *RenderingBackBuffer.GetReference(), *BackBufferImages[AcquiredImageIndex].GetReference(), SizeX, SizeY, WindowSizeX, WindowSizeY, SwapChain->GetCachedSurfaceTransform());
 			}
 			else
 			{
-				if (AcquiredImageIndex != -1)
-				{
-					check(RHIBackBuffer != nullptr && RHIBackBuffer->Image == BackBufferImages[AcquiredImageIndex]->Image);
-
-					FVulkanLayoutManager& LayoutManager = CommandBuffer.GetLayoutManager();
-					const FVulkanImageLayout* TrackedLayout = LayoutManager.GetFullLayout(BackBufferImages[AcquiredImageIndex]->Image);
+				bFailedToDelayAcquireBackbuffer = true;
+			}
+		}
+		else
+		{
+			if (AcquiredImageIndex != -1)
+			{
+				FVulkanCommandBuffer& CommandBuffer = Context.GetCommandBuffer();
+				check(CommandBuffer.IsOutsideRenderPass());
 
-					// One of the rare cases we'll let parallel rendering check the tracking (legacy path already checked as a fallback)
-					if (!TrackedLayout && Device->SupportsParallelRendering())
-					{
-						TrackedLayout = Context.Queue.GetLayoutManager().GetFullLayout(BackBufferImages[AcquiredImageIndex]->Image);
-					}
+				check(RHIBackBuffer != nullptr && RHIBackBuffer->Image == BackBufferImages[AcquiredImageIndex]->Image);
 
-					VkImageLayout LastLayout = TrackedLayout ? TrackedLayout->MainLayout : VK_IMAGE_LAYOUT_UNDEFINED;
-					VulkanSetImageLayout(&CommandBuffer, BackBufferImages[AcquiredImageIndex]->Image, LastLayout, VK_IMAGE_LAYOUT_PRESENT_SRC_KHR, FVulkanPipelineBarrier::MakeSubresourceRange(VK_IMAGE_ASPECT_COLOR_BIT));
+				FVulkanLayoutManager& LayoutManager = CommandBuffer.GetLayoutManager();
+				const FVulkanImageLayout* TrackedLayout = LayoutManager.GetFullLayout(BackBufferImages[AcquiredImageIndex]->Image);
 
-					const FVulkanImageLayout UndefinedLayout(VK_IMAGE_LAYOUT_UNDEFINED, 1, 1, VK_IMAGE_ASPECT_COLOR_BIT);
-					LayoutManager.SetFullLayout(BackBufferImages[AcquiredImageIndex]->Image, UndefinedLayout);
-				}
-				else
+				// One of the rare cases we'll let parallel rendering check the tracking (legacy path already checked as a fallback)
+				if (!TrackedLayout && Device->SupportsParallelRendering())
 				{
-					// When we have failed to acquire backbuffer image we fallback to using 'dummy' backbuffer
-					check(RHIBackBuffer != nullptr && RHIBackBuffer->Image == RenderingBackBuffer->Image);
+					TrackedLayout = Context.Queue.GetLayoutManager().GetFullLayout(BackBufferImages[AcquiredImageIndex]->Image);
 				}
+
+				VkImageLayout LastLayout = TrackedLayout ? TrackedLayout->MainLayout : VK_IMAGE_LAYOUT_UNDEFINED;
+				VulkanSetImageLayout(&CommandBuffer, BackBufferImages[AcquiredImageIndex]->Image, LastLayout, VK_IMAGE_LAYOUT_PRESENT_SRC_KHR, FVulkanPipelineBarrier::MakeSubresourceRange(VK_IMAGE_ASPECT_COLOR_BIT));
+
+				const FVulkanImageLayout UndefinedLayout(VK_IMAGE_LAYOUT_UNDEFINED, 1, 1, VK_IMAGE_ASPECT_COLOR_BIT);
+				LayoutManager.SetFullLayout(BackBufferImages[AcquiredImageIndex]->Image, UndefinedLayout);
+			}
+			else
+			{
+				// When we have failed to acquire backbuffer image we fallback to using 'dummy' backbuffer
+				check(RHIBackBuffer != nullptr && RHIBackBuffer->Image == RenderingBackBuffer->Image);
 			}
 		}
 	}
@@ -1132,7 +1131,7 @@ bool FVulkanViewport::Present(FVulkanCommandListContext& Context, FVulkanQueue*
 	if (bNeedNativePresent && (!SupportsStandardSwapchain() || GVulkanDelayAcquireImage == EDelayAcquireImageType::DelayAcquire || RHIBackBuffer != nullptr))
 	{
 		// Present the back buffer to the viewport window.
-		auto SwapChainJob = [Queue, PresentQueue](FVulkanViewport* Viewport)
+		auto SwapChainJob = [PresentQueue](FVulkanViewport* Viewport)
 		{
 			// May happend if swapchain was recreated in DoCheckedSwapChainJob()
 			if (Viewport->AcquiredImageIndex == -1)
@@ -1141,7 +1140,7 @@ bool FVulkanViewport::Present(FVulkanCommandListContext& Context, FVulkanQueue*
 				return (int32)FVulkanSwapChain::EStatus::Healthy;
 			}
 		
-			return (int32)Viewport->SwapChain->Present(Queue, PresentQueue, Viewport->RenderingDoneSemaphores[Viewport->AcquiredImageIndex]);
+			return (int32)Viewport->SwapChain->Present(PresentQueue, Viewport->RenderingDoneSemaphores[Viewport->AcquiredImageIndex]);
 		};
 		if (SupportsStandardSwapchain() && !DoCheckedSwapChainJob(Context, SwapChainJob))
 		{
diff --git a/Engine/Source/Runtime/VulkanRHI/Private/VulkanViewport.h b/Engine/Source/Runtime/VulkanRHI/Private/VulkanViewport.h
--- a/Engine/Source/Runtime/VulkanRHI/Private/VulkanViewport.h
+++ b/Engine/Source/Runtime/VulkanRHI/Private/VulkanViewport.h
@@ -73,7 +73,7 @@ public:
 
 	virtual void Tick(float DeltaTime) override final;
 	
-	bool Present(FVulkanCommandListContext& Context, FVulkanQueue* Queue, FVulkanQueue* PresentQueue, bool bLockToVsync);
+	bool Present(FVulkanCommandListContext& Context, FVulkanQueue* PresentQueue, bool bLockToVsync);
 
 	inline uint32 GetPresentCount() const
 	{
