VulkanRHI: driver fix for radv has made it to version 24.3.2, change minimum version for ray tracing

#rnx

diff --git a/Engine/Source/Runtime/VulkanRHI/Private/VulkanRHI.cpp b/Engine/Source/Runtime/VulkanRHI/Private/VulkanRHI.cpp
--- a/Engine/Source/Runtime/VulkanRHI/Private/VulkanRHI.cpp
+++ b/Engine/Source/Runtime/VulkanRHI/Private/VulkanRHI.cpp
@@ -864,7 +864,7 @@ void FVulkanDynamicRHI::InitInstance()
 		const bool bUseAMDCompatibilityMode = GVulkanAMDCompatibilityMode && (Device->GetVendorId() == EGpuVendorId::Amd);
 		if (GVulkanAMDCompatibilityMode && 
 			(Device->GetOptionalExtensionProperties().PhysicalDeviceDriverProperties.driverID == VK_DRIVER_ID_MESA_RADV) &&
-			(Props.driverVersion < VK_MAKE_VERSION(25, 0, 0)))
+			(Props.driverVersion < VK_MAKE_VERSION(24, 3, 2)))
 		{
 			GRHISupportsRayTracing = false;
 			GRHISupportsRayTracingShaders = false;
