VulkanRHI: Add an event to know when a payload is submitted to the queue.

#rnx

diff --git a/Engine/Source/Runtime/VulkanRHI/Private/VulkanQueue.cpp b/Engine/Source/Runtime/VulkanRHI/Private/VulkanQueue.cpp
--- a/Engine/Source/Runtime/VulkanRHI/Private/VulkanQueue.cpp
+++ b/Engine/Source/Runtime/VulkanRHI/Private/VulkanQueue.cpp
@@ -254,6 +254,12 @@ void FVulkanQueue::SubmitPayloads(TArrayView<FVulkanPayload*> Payloads, TSet<VkS
 	// Queue the submitted payloads in the interrupt queue
 	for (FVulkanPayload* Payload : Payloads)
 	{
+		if (Payload->SubmissionEvent.IsValid())
+		{
+			Payload->SubmissionEvent->DispatchSubsequents();
+			Payload->SubmissionEvent = nullptr;
+		}
+
 		PendingInterrupt.Enqueue(Payload);
 	}
 }
diff --git a/Engine/Source/Runtime/VulkanRHI/Private/VulkanSubmission.h b/Engine/Source/Runtime/VulkanRHI/Private/VulkanSubmission.h
--- a/Engine/Source/Runtime/VulkanRHI/Private/VulkanSubmission.h
+++ b/Engine/Source/Runtime/VulkanRHI/Private/VulkanSubmission.h
@@ -45,6 +45,9 @@ protected:
 	TArray<FVulkanCommandBuffer*> CommandBuffers;
 	TArray<VulkanRHI::FSemaphore*> SignalSemaphores; // signaled after command buffers
 
+	// Signaled when the payload has been submitted to the GPU queue
+	FGraphEventRef SubmissionEvent;
+
 	// For internal completion tracking of the payload
 	uint64 TimelineSemaphoreValue = 0;
 	VulkanRHI::FFence* Fence = nullptr;
