Implement GetGPUSizeBytes for Metal and Vulkan staging buffers. Not implementing these falls back to inherited FRHIStagingBuffer::GetGPUSizeBytes which returns 0, which in turns triggers assertions when the size of these staging buffers are checked (in ReflectionEnvironmentRealTimeCapture.cpp in Movile_ReadBackSkyIrradianceEnvironmentMap).

diff --git a/Engine/Source/Runtime/Apple/MetalRHI/Private/MetalRHIStagingBuffer.cpp b/Engine/Source/Runtime/Apple/MetalRHI/Private/MetalRHIStagingBuffer.cpp
--- a/Engine/Source/Runtime/Apple/MetalRHI/Private/MetalRHIStagingBuffer.cpp
+++ b/Engine/Source/Runtime/Apple/MetalRHI/Private/MetalRHIStagingBuffer.cpp
@@ -47,3 +47,8 @@ void FMetalRHIStagingBuffer::Unlock()
 	check(bIsLocked);
 	bIsLocked = false;
 }
+
+uint64 FMetalRHIStagingBuffer::GetGPUSizeBytes() const
+{
+	return ShadowBuffer ? ShadowBuffer->GetLength() : 0;
+}
diff --git a/Engine/Source/Runtime/Apple/MetalRHI/Private/MetalRHIStagingBuffer.h b/Engine/Source/Runtime/Apple/MetalRHI/Private/MetalRHIStagingBuffer.h
--- a/Engine/Source/Runtime/Apple/MetalRHI/Private/MetalRHIStagingBuffer.h
+++ b/Engine/Source/Runtime/Apple/MetalRHI/Private/MetalRHIStagingBuffer.h
@@ -29,6 +29,8 @@ public:
 	virtual void* Lock(uint32 Offset, uint32 NumBytes) final override;
 
 	virtual void Unlock() final override;
+	
+	uint64 GetGPUSizeBytes() const override;
 
 private:
 	FMetalDevice& Device;
diff --git a/Engine/Source/Runtime/VulkanRHI/Private/VulkanResources.h b/Engine/Source/Runtime/VulkanRHI/Private/VulkanResources.h
--- a/Engine/Source/Runtime/VulkanRHI/Private/VulkanResources.h
+++ b/Engine/Source/Runtime/VulkanRHI/Private/VulkanResources.h
@@ -1238,6 +1238,8 @@ public:
 	virtual void* Lock(uint32 Offset, uint32 NumBytes) final override;
 	virtual void Unlock() final override;
 
+	uint64 GetGPUSizeBytes() const override { return StagingBuffer ? StagingBuffer->GetSize() : 0; }
+
 private:
 	VulkanRHI::FStagingBuffer* StagingBuffer = nullptr;
 	uint32 QueuedNumBytes = 0;
