VulkanRHI: Initialize FVulkanRHIExternalImageDeleteCallbackInfo pointers.

#rnx

diff --git a/Engine/Source/Runtime/VulkanRHI/Public/IVulkanDynamicRHI.h b/Engine/Source/Runtime/VulkanRHI/Public/IVulkanDynamicRHI.h
--- a/Engine/Source/Runtime/VulkanRHI/Public/IVulkanDynamicRHI.h
+++ b/Engine/Source/Runtime/VulkanRHI/Public/IVulkanDynamicRHI.h
@@ -26,8 +26,8 @@ struct FVulkanRHIImageViewInfo
 
 struct FVulkanRHIExternalImageDeleteCallbackInfo
 {
-	void* UserData;
-	void (*Function)(void* UserData);
+	void* UserData = nullptr;
+	void (*Function)(void* UserData) = nullptr;
 };
 
 struct IVulkanDynamicRHI : public FDynamicRHI
