VulkanRHI: Remove layout tracking inside the VulkanRHI and use RHICommandList state tracking instead.  Remove 'notifications' for deleted images to remove their tracking info.  Fix incorrect initial layout when DelayAcquire was used.  Fix incorrect context on texture unlocks.

#rnx

diff --git a/Engine/Source/Runtime/RHI/Public/RHIResources.h b/Engine/Source/Runtime/RHI/Public/RHIResources.h
--- a/Engine/Source/Runtime/RHI/Public/RHIResources.h
+++ b/Engine/Source/Runtime/RHI/Public/RHIResources.h
@@ -1331,6 +1331,11 @@ public:
 		return TrackedAccess.GetResourceState();
 	}
 
+	bool AreAllSubresourcesSame_Unsafe() const
+	{
+		return TrackedAccess.AreAllSubresourcesSame();
+	}
+
 	FName GetName() const
 	{
 		return Name;
diff --git a/Engine/Source/Runtime/VulkanRHI/Private/VulkanBarriers.cpp b/Engine/Source/Runtime/VulkanRHI/Private/VulkanBarriers.cpp
--- a/Engine/Source/Runtime/VulkanRHI/Private/VulkanBarriers.cpp
+++ b/Engine/Source/Runtime/VulkanRHI/Private/VulkanBarriers.cpp
@@ -22,28 +22,6 @@ static FAutoConsoleVariableRef CVarVulkanUseMemoryBarrierOpt(
 	ECVF_Default
 );
 
-int32 GVulkanAutoCorrectExpectedLayouts = 1;
-static FAutoConsoleVariableRef CVarVulkanAutoCorrectExpectedLayouts(
-	TEXT("r.Vulkan.AutoCorrectExpectedLayouts"),
-	GVulkanAutoCorrectExpectedLayouts,
-	TEXT("Will use layout tracking to correct mismatched layouts when setting expected layouts.  \n")
-	TEXT("This is unsafe for multi-threaded command buffer generations\n")
-	TEXT(" 0: Do not correct layouts.\n")
-	TEXT(" 1: Correct the layouts using layout tracking. (default)"),
-	ECVF_Default
-);
-
-int32 GVulkanAutoCorrectUnknownLayouts = 1;
-static FAutoConsoleVariableRef CVarVulkanAutoCorrectUnknownLayouts(
-	TEXT("r.Vulkan.AutoCorrectUnknownLayouts"),
-	GVulkanAutoCorrectUnknownLayouts,
-	TEXT("Will use layout tracking to correct unknown layouts.  \n")
-	TEXT("This is unsafe for multi-threaded command buffer generations\n")
-	TEXT(" 0: Do not correct layouts.\n")
-	TEXT(" 1: Correct the layouts using layout tracking. (default)"),
-	ECVF_Default
-);
-
 int32 GVulkanMaxBarriersPerBatch = -1;
 static FAutoConsoleVariableRef CVarVulkanMaxBarriersPerBatch(
 	TEXT("r.Vulkan.MaxBarriersPerBatch"),
@@ -255,9 +233,11 @@ static void GetVkStageAndAccessFlags(ERHIAccess RHIAccess, FRHITransitionInfo::E
 			return;
 
 		case ERHIAccess::Present:
-			StageFlags = bIsSourceState ? VK_PIPELINE_STAGE_TOP_OF_PIPE_BIT : VK_PIPELINE_STAGE_BOTTOM_OF_PIPE_BIT;
-			AccessFlags = VK_PIPELINE_STAGE_COLOR_ATTACHMENT_OUTPUT_BIT;
-			Layout = VK_IMAGE_LAYOUT_COLOR_ATTACHMENT_OPTIMAL;
+			StageFlags = bIsSourceState ? VK_PIPELINE_STAGE_COLOR_ATTACHMENT_OUTPUT_BIT : VK_PIPELINE_STAGE_BOTTOM_OF_PIPE_BIT;
+			// When transitionning out of present, the sema handles access.
+			// When transitioning into present, vkQueuePresentKHR guarantees visibility.
+			AccessFlags = 0;
+			Layout = VK_IMAGE_LAYOUT_PRESENT_SRC_KHR;
 			return;
 
 		case ERHIAccess::RTV:
@@ -482,7 +462,7 @@ static VkImageAspectFlags GetDepthStencilAspectMask(uint32 PlaneSlice)
 
 
 // Returns the VK_KHR_synchronization2 layout corresponding to an access type
-static VkImageLayout GetDepthOrStencilLayout(ERHIAccess Access)
+VkImageLayout FVulkanPipelineBarrier::GetDepthOrStencilLayout(ERHIAccess Access)
 {
 	VkImageLayout Layout;
 	if (Access == ERHIAccess::Unknown)
@@ -827,8 +807,8 @@ void FVulkanDynamicRHI::RHICreateTransition(FRHITransition* Transition, const FR
 		if (bIsDepthStencil)
 		{
 			// if we use separate transitions, then just feed them in as they are
-			SrcLayout = GetDepthOrStencilLayout(Info.AccessBefore);
-			DstLayout = GetDepthOrStencilLayout(Info.AccessAfter);
+			SrcLayout = FVulkanPipelineBarrier::GetDepthOrStencilLayout(Info.AccessBefore);
+			DstLayout = FVulkanPipelineBarrier::GetDepthOrStencilLayout(Info.AccessAfter);
 			GetDepthOrStencilStageAndAccessFlags(Info.AccessBefore, SrcStageMask, SrcAccessFlags);
 			GetDepthOrStencilStageAndAccessFlags(Info.AccessAfter, DstStageMask, DstAccessFlags);
 		}
@@ -890,6 +870,15 @@ void FVulkanDynamicRHI::RHICreateTransition(FRHITransition* Transition, const FR
 		FVulkanPipelineBarrier::ImageBarrierExtraData& ExtraData = Data->ImageBarrierExtras.AddDefaulted_GetRef();
 		ExtraData.BaseTexture = Texture;
 		ExtraData.IsAliasingBarrier = (Info.AccessBefore == ERHIAccess::Discard);
+		if (EnumHasAnyFlags(Info.Flags, EResourceTransitionFlags::FillOtherPlaneState))
+		{
+			checkSlow(bIsDepthStencil);
+			ExtraData.OtherPlaneLayout = FVulkanPipelineBarrier::GetDepthOrStencilLayout(Info.OtherPlaneAccess);
+		}
+		else
+		{
+			checkSlow(!bIsDepthStencil || (Info.PlaneSlice == FRHISubresourceRange::kAllSubresources) || AspectFlags == Texture->GetFullAspectMask());
+		}
 	}
 }
 
@@ -898,61 +887,6 @@ void FVulkanDynamicRHI::RHIReleaseTransition(FRHITransition* Transition)
 	Transition->GetPrivateData<FVulkanPipelineBarrier>()->~FVulkanPipelineBarrier();
 }
 
-static void AddSubresourceTransitions(TArray<VkImageMemoryBarrier>& Barriers, VkPipelineStageFlags& SrcStageMask, const VkImageMemoryBarrier TemplateBarrier, VkImage ImageHandle, FVulkanImageLayout& CurrentLayout, VkImageLayout DstLayout)
-{
-	const bool bIsDepthStencil = VKHasAnyFlags(TemplateBarrier.subresourceRange.aspectMask, VK_IMAGE_ASPECT_DEPTH_BIT | VK_IMAGE_ASPECT_STENCIL_BIT);
-
-	const uint32 FirstLayer = TemplateBarrier.subresourceRange.baseArrayLayer;
-	const uint32 LastLayer = FirstLayer + CurrentLayout.GetSubresRangeLayerCount(TemplateBarrier.subresourceRange);
-
-	const uint32 FirstMip = TemplateBarrier.subresourceRange.baseMipLevel;
-	const uint32 LastMip = FirstMip + CurrentLayout.GetSubresRangeMipCount(TemplateBarrier.subresourceRange);
-
-	for (uint32 LayerIdx = FirstLayer; LayerIdx < LastLayer; ++LayerIdx)
-	{
-		VkImageMemoryBarrier* PrevMipBarrier = nullptr;
-
-		for (uint32 MipIdx = FirstMip; MipIdx < LastMip; ++MipIdx)
-		{
-			VkImageLayout SrcLayout = CurrentLayout.GetSubresLayout(LayerIdx, MipIdx, 0);
-			if (bIsDepthStencil)
-			{
-				const VkImageLayout OtherLayout = (CurrentLayout.NumPlanes == 1) ? SrcLayout : CurrentLayout.GetSubresLayout(LayerIdx, MipIdx, 1);
-				SrcLayout = VulkanRHI::GetMergedDepthStencilLayout(SrcLayout, OtherLayout);
-			}
-
-			// Merge with the previous transition if the previous mip was in the same state as this mip.
-			if (PrevMipBarrier && PrevMipBarrier->oldLayout == SrcLayout)
-			{
-				PrevMipBarrier->subresourceRange.levelCount += 1;
-			}
-			else
-			{
-				if (SrcLayout == DstLayout)
-				{
-					continue;
-				}
-
-				SrcStageMask |= GetVkStageFlagsForLayout(SrcLayout);
-
-				VkImageMemoryBarrier& Barrier = Barriers.AddDefaulted_GetRef();
-				Barrier = TemplateBarrier;
-				Barrier.srcAccessMask = GetVkAccessMaskForLayout(SrcLayout);
-				Barrier.oldLayout = SrcLayout;
-				Barrier.newLayout = bIsDepthStencil ? VulkanRHI::GetMergedDepthStencilLayout(DstLayout, DstLayout) : DstLayout;
-				Barrier.image = ImageHandle;
-				Barrier.subresourceRange.baseMipLevel = MipIdx;
-				Barrier.subresourceRange.levelCount = 1;
-				Barrier.subresourceRange.baseArrayLayer = LayerIdx;
-				Barrier.subresourceRange.layerCount = 1;
-
-				PrevMipBarrier = &Barrier;
-			}
-		}
-	}
-	CurrentLayout.Set(DstLayout, TemplateBarrier.subresourceRange);
-}
-
 static void DowngradeBarrier(VkMemoryBarrier& OutBarrier, const VkMemoryBarrier2& InBarrier)
 {
 	OutBarrier.sType = VK_STRUCTURE_TYPE_MEMORY_BARRIER;
@@ -1222,136 +1156,80 @@ void FTransitionProcessor<VkMemoryBarrier, VkBufferMemoryBarrier, VkImageMemoryB
 template <>
 void FTransitionProcessor<VkMemoryBarrier, VkBufferMemoryBarrier, VkImageMemoryBarrier>::ApplyTracking(uint32_t TargetIndex, const FVulkanPipelineBarrier::ImageBarrierExtraData* RemainingExtras, int32 RemainingCount, bool IsCrossPipe)
 {
-	FVulkanLayoutManager& LayoutManager = Context.GetCommandBuffer().GetLayoutManager();
-
-	VkImageMemoryBarrier* ImageBarrier = &ImageBarriers[TargetIndex];
+	VkImageMemoryBarrier& ImageBarrier = ImageBarriers[TargetIndex];
 	FVulkanTexture* Texture = RemainingExtras->BaseTexture;
-	check(Texture->Image == ImageBarrier->image);
-
-	// Get a copy of the current layout, we will alter it and feed it back in the manager
-	FVulkanImageLayout Layout = *LayoutManager.GetFullLayout(*Texture, true);
-
-	check((Layout.NumMips == Texture->GetNumMips()) && (Layout.NumLayers == Texture->GetNumberOfArrayLevels()));
-	check((ImageBarrier->subresourceRange.baseArrayLayer + Layout.GetSubresRangeLayerCount(ImageBarrier->subresourceRange)) <= Layout.NumLayers);
-	check((ImageBarrier->subresourceRange.baseMipLevel + Layout.GetSubresRangeMipCount(ImageBarrier->subresourceRange)) <= Layout.NumMips);
-	check(ImageBarrier->newLayout != VK_IMAGE_LAYOUT_UNDEFINED);
+	check(Texture->Image == ImageBarrier.image);
+	check(ImageBarrier.newLayout != VK_IMAGE_LAYOUT_UNDEFINED);
+	check((ImageBarrier.oldLayout != VK_IMAGE_LAYOUT_UNDEFINED) || RemainingExtras->IsAliasingBarrier);
 
 	// For Depth/Stencil formats where only one of the aspects is transitioned, look ahead for other barriers on the same resource
-	if (Texture->IsDepthOrStencilAspect() && (Texture->GetFullAspectMask() != ImageBarrier->subresourceRange.aspectMask))
+	if (Texture->IsDepthOrStencilAspect())
 	{
-		check(VKHasAnyFlags(ImageBarrier->subresourceRange.aspectMask, VK_IMAGE_ASPECT_DEPTH_BIT |VK_IMAGE_ASPECT_STENCIL_BIT));
-		if ((ImageBarrier->oldLayout == VK_IMAGE_LAYOUT_UNDEFINED) && (!RemainingExtras->IsAliasingBarrier))
+		if (Texture->GetFullAspectMask() != ImageBarrier.subresourceRange.aspectMask)
 		{
-			ImageBarrier->oldLayout = Layout.GetSubresLayout(ImageBarrier->subresourceRange.baseArrayLayer, ImageBarrier->subresourceRange.baseMipLevel, (VkImageAspectFlagBits)ImageBarrier->subresourceRange.aspectMask);
-		}
+			check(VKHasAnyFlags(ImageBarrier.subresourceRange.aspectMask, VK_IMAGE_ASPECT_DEPTH_BIT | VK_IMAGE_ASPECT_STENCIL_BIT));
 
-		const VkImageAspectFlagBits OtherAspectMask = (VkImageAspectFlagBits)(Texture->GetFullAspectMask() ^ ImageBarrier->subresourceRange.aspectMask);
-		VkImageLayout OtherAspectOldLayout = Layout.GetSubresLayout(ImageBarrier->subresourceRange.baseArrayLayer, ImageBarrier->subresourceRange.baseMipLevel, OtherAspectMask);
-		VkImageLayout OtherAspectNewLayout = OtherAspectOldLayout;
-		for (int32 OtherBarrierIndex = 0; OtherBarrierIndex < RemainingCount; ++OtherBarrierIndex)
-		{
-			VkImageMemoryBarrier& OtherImageBarrier = ImageBarriers[TargetIndex+OtherBarrierIndex];
-			FVulkanTexture* OtherTexture = RemainingExtras[OtherBarrierIndex].BaseTexture;
-			if ((OtherTexture->Image == ImageBarrier->image) && (OtherImageBarrier.subresourceRange.aspectMask == OtherAspectMask))
+			const VkImageAspectFlagBits OtherAspectMask = (VkImageAspectFlagBits)(Texture->GetFullAspectMask() ^ ImageBarrier.subresourceRange.aspectMask);
+			VkImageMemoryBarrier* OtherAspectImageBarrier = nullptr;
+			for (int32 OtherBarrierIndex = 0; OtherBarrierIndex < RemainingCount; ++OtherBarrierIndex)
 			{
-				check(ImageBarrier->subresourceRange.baseArrayLayer == OtherImageBarrier.subresourceRange.baseArrayLayer);
-				check(ImageBarrier->subresourceRange.baseMipLevel == OtherImageBarrier.subresourceRange.baseMipLevel);
-
-				OtherAspectNewLayout = OtherImageBarrier.newLayout;
-				if ((OtherImageBarrier.oldLayout != VK_IMAGE_LAYOUT_UNDEFINED) || RemainingExtras[OtherBarrierIndex].IsAliasingBarrier)
+				VkImageMemoryBarrier& OtherImageBarrier = ImageBarriers[TargetIndex + OtherBarrierIndex];
+				FVulkanTexture* OtherTexture = RemainingExtras[OtherBarrierIndex].BaseTexture;
+				if ((OtherTexture->Image == ImageBarrier.image) && (OtherImageBarrier.subresourceRange.aspectMask == OtherAspectMask))
 				{
-					OtherAspectOldLayout = OtherImageBarrier.oldLayout;
+					check(ImageBarrier.subresourceRange.baseArrayLayer == OtherImageBarrier.subresourceRange.baseArrayLayer);
+					check(ImageBarrier.subresourceRange.baseMipLevel == OtherImageBarrier.subresourceRange.baseMipLevel);
+
+					OtherAspectImageBarrier = &OtherImageBarrier;
+
+					break;
 				}
+			}
 
-				Layout.Set(OtherImageBarrier.newLayout, OtherImageBarrier.subresourceRange);
+			VkImageLayout OtherAspectOldLayout, OtherAspectNewLayout;
+			if (OtherAspectImageBarrier)
+			{
+				OtherAspectOldLayout = OtherAspectImageBarrier->oldLayout;
+				OtherAspectNewLayout = OtherAspectImageBarrier->newLayout;
 
 				// Make it invalid so that it gets removed when we reach it
-				OtherImageBarrier.subresourceRange.aspectMask = 0;
-
-				break;
+				OtherAspectImageBarrier->subresourceRange.aspectMask = 0;
 			}
-		}
-
-		Layout.Set(ImageBarrier->newLayout, ImageBarrier->subresourceRange);
+			else
+			{
+				checkf((RemainingExtras->OtherPlaneLayout != VK_IMAGE_LAYOUT_MAX_ENUM), TEXT("This shouldn't happen with GRHIGlobals.NeedsAllPlanes set!"));
 
-		// Merge the layout with its other half and set it in the barrier
-		if (OtherAspectMask == VK_IMAGE_ASPECT_STENCIL_BIT)
-		{
-			ImageBarrier->oldLayout = VulkanRHI::GetMergedDepthStencilLayout(ImageBarrier->oldLayout, OtherAspectOldLayout);
-			ImageBarrier->newLayout = VulkanRHI::GetMergedDepthStencilLayout(ImageBarrier->newLayout, OtherAspectNewLayout);
-		}
-		else
-		{
-			ImageBarrier->oldLayout = VulkanRHI::GetMergedDepthStencilLayout(OtherAspectOldLayout, ImageBarrier->oldLayout);
-			ImageBarrier->newLayout = VulkanRHI::GetMergedDepthStencilLayout(OtherAspectNewLayout, ImageBarrier->newLayout);
-		}
-        ImageBarrier->subresourceRange.aspectMask |= OtherAspectMask;
+				OtherAspectOldLayout = OtherAspectNewLayout = RemainingExtras->OtherPlaneLayout;
+			}
 
-		// Once we're done downgrading the barrier, make sure there are no sync2 states left
-		check((ImageBarrier->oldLayout != VK_IMAGE_LAYOUT_READ_ONLY_OPTIMAL) && (ImageBarrier->oldLayout != VK_IMAGE_LAYOUT_ATTACHMENT_OPTIMAL));
-		check((ImageBarrier->newLayout != VK_IMAGE_LAYOUT_READ_ONLY_OPTIMAL) && (ImageBarrier->newLayout != VK_IMAGE_LAYOUT_ATTACHMENT_OPTIMAL));
-	}
-	else
-	{
-		if ((ImageBarrier->oldLayout == VK_IMAGE_LAYOUT_UNDEFINED) && (!RemainingExtras->IsAliasingBarrier))
-		{
-			if (Layout.AreAllSubresourcesSameLayout())
+			// Merge the layout with its other half and set it in the barrier
+			if (OtherAspectMask == VK_IMAGE_ASPECT_STENCIL_BIT)
 			{
-				ImageBarrier->oldLayout = Layout.MainLayout;
-				ImageBarrier->srcAccessMask = GetVkAccessMaskForLayout(ImageBarrier->oldLayout);
-				MergedSrcStageMask |= GetVkStageFlagsForLayout(ImageBarrier->oldLayout);
+				ImageBarrier.oldLayout = VulkanRHI::GetMergedDepthStencilLayout(ImageBarrier.oldLayout, OtherAspectOldLayout);
+				ImageBarrier.newLayout = VulkanRHI::GetMergedDepthStencilLayout(ImageBarrier.newLayout, OtherAspectNewLayout);
 			}
 			else
 			{
-				// Slow path, adds one transition per subresource.
-				check(!IsCrossPipe);
-				AddSubresourceTransitions(ImageBarriers, MergedSrcStageMask, *ImageBarrier, Texture->Image, Layout, ImageBarrier->newLayout);
-                // Update the pointer to the current barrier as AddSubresourceTransitions can invalidate it when adding
-                ImageBarrier = &ImageBarriers[TargetIndex];
+				ImageBarrier.oldLayout = VulkanRHI::GetMergedDepthStencilLayout(OtherAspectOldLayout, ImageBarrier.oldLayout);
+				ImageBarrier.newLayout = VulkanRHI::GetMergedDepthStencilLayout(OtherAspectNewLayout, ImageBarrier.newLayout);
 			}
-		}
-		else
-		{
-			checkSlow(RemainingExtras->IsAliasingBarrier || Layout.AreSubresourcesSameLayout(ImageBarrier->oldLayout, ImageBarrier->subresourceRange));
-		}
 
-		const VkImageLayout TrackedNewLayout = ImageBarrier->newLayout;
-		if (Texture->IsDepthOrStencilAspect())
-		{
-			// The only way we end up here is if the barrier transitions every aspect of the depth(-stencil) texture
-			check(Texture->GetFullAspectMask() == ImageBarrier->subresourceRange.aspectMask);
-			ImageBarrier->oldLayout = VulkanRHI::GetMergedDepthStencilLayout(ImageBarrier->oldLayout, ImageBarrier->oldLayout);
-			ImageBarrier->newLayout = VulkanRHI::GetMergedDepthStencilLayout(ImageBarrier->newLayout, ImageBarrier->newLayout);
+			ImageBarrier.subresourceRange.aspectMask |= OtherAspectMask;
 		}
-
-		MergedDstStageMask |= GetVkStageFlagsForLayout(ImageBarriers[TargetIndex].newLayout);
-
-		if (!IsCrossPipe)
+		else
 		{
-			if (ImageBarrier->oldLayout == ImageBarrier->newLayout)
-			{
-				// It turns out that we don't need a layout transition after all. We may still need a memory barrier if the
-				// previous access was writable.
-
-				if (MemoryBarriers.Num() == 0)
-				{
-					MemoryBarriers.AddDefaulted();
-					ZeroVulkanStruct(MemoryBarriers[0], VK_STRUCTURE_TYPE_MEMORY_BARRIER);
-				}
-
-				MemoryBarriers[0].srcAccessMask |= ImageBarrier->srcAccessMask;
-				MemoryBarriers[0].dstAccessMask |= ImageBarrier->dstAccessMask;
-			}
+			// Transitions every aspect of the depth(-stencil) texture
+			ImageBarrier.oldLayout = VulkanRHI::GetMergedDepthStencilLayout(ImageBarrier.oldLayout, ImageBarrier.oldLayout);
+			ImageBarrier.newLayout = VulkanRHI::GetMergedDepthStencilLayout(ImageBarrier.newLayout, ImageBarrier.newLayout);
 		}
-
-		// Make sure these new barriers are downgraded (since we're already past that step by the time we apply tracking)
-		check((ImageBarrier->oldLayout != VK_IMAGE_LAYOUT_READ_ONLY_OPTIMAL) && (ImageBarrier->oldLayout != VK_IMAGE_LAYOUT_ATTACHMENT_OPTIMAL));
-		check((ImageBarrier->newLayout != VK_IMAGE_LAYOUT_READ_ONLY_OPTIMAL) && (ImageBarrier->newLayout != VK_IMAGE_LAYOUT_ATTACHMENT_OPTIMAL));
-
-		Layout.Set(TrackedNewLayout, ImageBarrier->subresourceRange);
 	}
 
-	LayoutManager.SetFullLayout(*Texture, Layout);
+	MergedSrcStageMask |= GetVkStageFlagsForLayout(ImageBarrier.oldLayout);
+	MergedDstStageMask |= GetVkStageFlagsForLayout(ImageBarrier.newLayout);
+
+	// Once we're done with the barrier, make sure there are no sync2 states left
+	check((ImageBarrier.oldLayout != VK_IMAGE_LAYOUT_READ_ONLY_OPTIMAL) && (ImageBarrier.oldLayout != VK_IMAGE_LAYOUT_ATTACHMENT_OPTIMAL));
+	check((ImageBarrier.newLayout != VK_IMAGE_LAYOUT_READ_ONLY_OPTIMAL) && (ImageBarrier.newLayout != VK_IMAGE_LAYOUT_ATTACHMENT_OPTIMAL));
 }
 
 
@@ -1461,52 +1339,7 @@ void FTransitionProcessor<VkMemoryBarrier2, VkBufferMemoryBarrier2, VkImageMemor
 template <>
 void FTransitionProcessor<VkMemoryBarrier2, VkBufferMemoryBarrier2, VkImageMemoryBarrier2>::ApplyTracking(uint32 TargetIndex, const FVulkanPipelineBarrier::ImageBarrierExtraData* RemainingExtras, int32 RemainingCount, bool IsCrossPipe)
 {
-	if (GVulkanAutoCorrectUnknownLayouts)
-	{
-		VkImageMemoryBarrier2& ImageBarrier = ImageBarriers[TargetIndex];
-
-		if ((ImageBarrier.oldLayout == VK_IMAGE_LAYOUT_UNDEFINED) && (!RemainingExtras->IsAliasingBarrier))
-		{
-			FVulkanLayoutManager& LayoutManager = Context.GetCommandBuffer().GetLayoutManager();
-			FVulkanTexture* Texture = RemainingExtras->BaseTexture;
-			check(Texture->Image == ImageBarrier.image);
-			const FVulkanImageLayout* Layout = LayoutManager.GetFullLayout(*Texture, false);
-			if (Layout)
-			{
-				if (Layout->AreAllSubresourcesSameLayout())
-				{
-					ImageBarrier.oldLayout = Layout->MainLayout;
-				}
-				else
-				{
-					ForEachAspect(ImageBarrier.subresourceRange.aspectMask, [&](VkImageAspectFlagBits SingleAspect)
-						{
-							const uint32 LastLayer = ImageBarrier.subresourceRange.baseArrayLayer + ImageBarrier.subresourceRange.layerCount;
-							const uint32 LastMip = ImageBarrier.subresourceRange.baseMipLevel + ImageBarrier.subresourceRange.levelCount;
-							for (uint32 LayerIndex = ImageBarrier.subresourceRange.baseArrayLayer; LayerIndex < LastLayer; ++LayerIndex)
-							{
-								for (uint32 MipIndex = ImageBarrier.subresourceRange.baseMipLevel; MipIndex < LastMip; ++MipIndex)
-								{
-									const VkImageLayout SubresourceLayout = Layout->GetSubresLayout(LayerIndex, MipIndex, SingleAspect);
-									check((ImageBarrier.oldLayout == VK_IMAGE_LAYOUT_UNDEFINED) || (ImageBarrier.oldLayout == SubresourceLayout));
-									ImageBarrier.oldLayout = SubresourceLayout;
-								}
-							}
-						});
-				}
-
-				ImageBarrier.srcAccessMask = GetVkAccessMaskForLayout(ImageBarrier.oldLayout);
-				ImageBarrier.srcStageMask = GetVkStageFlagsForLayout(ImageBarrier.oldLayout);
-			}
-		}
-	}
-
-	// Sync2 only needs to inform tracking of the new layout
-	Context.GetCommandBuffer().GetLayoutManager().SetLayout(*RemainingExtras->BaseTexture, ImageBarriers[TargetIndex].subresourceRange, ImageBarriers[TargetIndex].newLayout);
-
-	// No needs to worry about 'surprise' subresource transition
-	// No need to worry about partial depth/stencil layout changes
-	// Identical layouts are chill ("Image memory barriers that do not perform an image layout transition can be specified by setting oldLayout equal to newLayout")
+	// Nothing to do when VK_KHR_separate_depth_stencil_layouts and VK_KHR_synchronization2 are supported (requirements for SM6).
 }
 
 
@@ -1645,111 +1478,6 @@ void FVulkanPipelineBarrier::AddImageLayoutTransition(VkImage Image, VkImageLayo
 	SetupImageBarrier(ImgBarrier, Image, SrcStageMask, DstStageMask, SrcAccessFlags, DstAccessFlags, SrcLayout, DstLayout, SubresourceRange);
 }
 
-void FVulkanPipelineBarrier::AddImageLayoutTransition(VkImage Image, VkImageAspectFlags AspectMask, const FVulkanImageLayout& SrcLayout, VkImageLayout DstLayout)
-{
-	if (SrcLayout.AreAllSubresourcesSameLayout())
-	{
-		AddImageLayoutTransition(Image, SrcLayout.MainLayout, DstLayout, MakeSubresourceRange(AspectMask));
-		return;
-	}
-
-	const VkPipelineStageFlags DstStageMask = GetVkStageFlagsForLayout(DstLayout);
-	const VkAccessFlags DstAccessFlags = GetVkAccessMaskForLayout(DstLayout);
-
-	ForEachAspect(AspectMask, [&](VkImageAspectFlagBits SingleAspect)
-		{
-			VkImageSubresourceRange SubresourceRange = MakeSubresourceRange(SingleAspect, 0, 1, 0, 1);
-			for (; SubresourceRange.baseArrayLayer < SrcLayout.NumLayers; ++SubresourceRange.baseArrayLayer)
-			{
-				for (SubresourceRange.baseMipLevel = 0; SubresourceRange.baseMipLevel < SrcLayout.NumMips; ++SubresourceRange.baseMipLevel)
-				{
-					const VkImageLayout SubresourceLayout = SrcLayout.GetSubresLayout(SubresourceRange.baseArrayLayer, SubresourceRange.baseMipLevel, SingleAspect);
-					if (SubresourceLayout != DstLayout)
-					{
-						const VkPipelineStageFlags SrcStageMask = GetVkStageFlagsForLayout(SubresourceLayout);
-						const VkAccessFlags SrcAccessFlags = GetVkAccessMaskForLayout(SubresourceLayout);
-
-						VkImageMemoryBarrier2& ImgBarrier = ImageBarriers.AddDefaulted_GetRef();
-						SetupImageBarrier(ImgBarrier, Image, SrcStageMask, DstStageMask, SrcAccessFlags, DstAccessFlags, SubresourceLayout, DstLayout, SubresourceRange);
-					}
-				}
-			}
-		});
-}
-
-void FVulkanPipelineBarrier::AddImageLayoutTransition(VkImage Image, VkImageAspectFlags AspectMask, VkImageLayout SrcLayout, const FVulkanImageLayout& DstLayout)
-{
-	if (DstLayout.AreAllSubresourcesSameLayout())
-	{
-		AddImageLayoutTransition(Image, SrcLayout, DstLayout.MainLayout, MakeSubresourceRange(AspectMask));
-		return;
-	}
-
-	const VkPipelineStageFlags SrcStageMask = GetVkStageFlagsForLayout(SrcLayout);
-	const VkAccessFlags SrcAccessFlags = GetVkAccessMaskForLayout(SrcLayout);
-
-	ForEachAspect(AspectMask, [&](VkImageAspectFlagBits SingleAspect)
-		{
-			VkImageSubresourceRange SubresourceRange = MakeSubresourceRange(SingleAspect, 0, 1, 0, 1);
-			for (; SubresourceRange.baseArrayLayer < DstLayout.NumLayers; ++SubresourceRange.baseArrayLayer)
-			{
-				for (SubresourceRange.baseMipLevel = 0; SubresourceRange.baseMipLevel < DstLayout.NumMips; ++SubresourceRange.baseMipLevel)
-				{
-					const VkImageLayout SubresourceLayout = DstLayout.GetSubresLayout(SubresourceRange.baseArrayLayer, SubresourceRange.baseMipLevel, SingleAspect);
-					if (SubresourceLayout != SrcLayout)
-					{
-						const VkPipelineStageFlags DstStageMask = GetVkStageFlagsForLayout(SubresourceLayout);
-						const VkAccessFlags DstAccessFlags = GetVkAccessMaskForLayout(SubresourceLayout);
-
-						VkImageMemoryBarrier2& ImgBarrier = ImageBarriers.AddDefaulted_GetRef();
-						SetupImageBarrier(ImgBarrier, Image, SrcStageMask, DstStageMask, SrcAccessFlags, DstAccessFlags, SrcLayout, SubresourceLayout, SubresourceRange);
-					}
-				}
-			}
-		});
-}
-
-void FVulkanPipelineBarrier::AddImageLayoutTransition(VkImage Image, VkImageAspectFlags AspectMask, const FVulkanImageLayout& SrcLayout, const FVulkanImageLayout& DstLayout)
-{
-	if (SrcLayout.AreAllSubresourcesSameLayout())
-	{
-		AddImageLayoutTransition(Image, AspectMask, SrcLayout.MainLayout, DstLayout);
-	}
-	else if (DstLayout.AreAllSubresourcesSameLayout())
-	{
-		AddImageLayoutTransition(Image, AspectMask, SrcLayout, DstLayout.MainLayout);
-	}
-	else
-	{
-		checkf(SrcLayout.NumLayers == DstLayout.NumLayers, TEXT("Source (%d) and Destination (%d) layer count mismatch!"), SrcLayout.NumLayers, DstLayout.NumLayers);
-		checkf(SrcLayout.NumMips == DstLayout.NumMips, TEXT("Source (%d) and Destination (%d) mip count mismatch!"), SrcLayout.NumMips, DstLayout.NumMips);
-
-		ForEachAspect(AspectMask, [&](VkImageAspectFlagBits SingleAspect)
-			{
-				VkImageSubresourceRange SubresourceRange = MakeSubresourceRange(SingleAspect, 0, 1, 0, 1);
-				for (; SubresourceRange.baseArrayLayer < DstLayout.NumLayers; ++SubresourceRange.baseArrayLayer)
-				{
-					for (SubresourceRange.baseMipLevel = 0; SubresourceRange.baseMipLevel < DstLayout.NumMips; ++SubresourceRange.baseMipLevel)
-					{
-						const VkImageLayout SrcSubresourceLayout = SrcLayout.GetSubresLayout(SubresourceRange.baseArrayLayer, SubresourceRange.baseMipLevel, SingleAspect);
-						const VkImageLayout DstSubresourceLayout = DstLayout.GetSubresLayout(SubresourceRange.baseArrayLayer, SubresourceRange.baseMipLevel, SingleAspect);
-						if (SrcSubresourceLayout != DstSubresourceLayout)
-						{
-							const VkPipelineStageFlags SrcStageMask = GetVkStageFlagsForLayout(SrcSubresourceLayout);
-							const VkAccessFlags SrcAccessFlags = GetVkAccessMaskForLayout(SrcSubresourceLayout);
-
-							const VkPipelineStageFlags DstStageMask = GetVkStageFlagsForLayout(DstSubresourceLayout);
-							const VkAccessFlags DstAccessFlags = GetVkAccessMaskForLayout(DstSubresourceLayout);
-
-							VkImageMemoryBarrier2& ImgBarrier = ImageBarriers.AddDefaulted_GetRef();
-							SetupImageBarrier(ImgBarrier, Image, SrcStageMask, DstStageMask, SrcAccessFlags, DstAccessFlags, SrcSubresourceLayout, DstSubresourceLayout, SubresourceRange);
-						}
-					}
-				}
-			});
-	}
-}
-
 void FVulkanPipelineBarrier::AddImageAccessTransition(const FVulkanTexture& Surface, ERHIAccess SrcAccess, ERHIAccess DstAccess, const VkImageSubresourceRange& SubresourceRange, VkImageLayout& InOutLayout)
 {
 	// This function should only be used for known states.
@@ -1859,115 +1587,7 @@ void VulkanSetImageLayout(FVulkanCommandBuffer* CmdBuffer, VkImage Image, VkImag
 	Barrier.Execute(CmdBuffer);
 }
 
-bool FVulkanImageLayout::AreSubresourcesSameLayout(VkImageLayout Layout, const VkImageSubresourceRange& SubresourceRange) const
-{
-	if (SubresLayouts.Num() == 0)
-	{
-		return MainLayout == Layout;
-	}
-
-	const uint32 FirstPlane = (SubresourceRange.aspectMask == VK_IMAGE_ASPECT_STENCIL_BIT) ? NumPlanes - 1 : 0;
-	const uint32 LastPlane = (SubresourceRange.aspectMask & VK_IMAGE_ASPECT_STENCIL_BIT) ? NumPlanes : 1;
-
-	const uint32 FirstLayer = SubresourceRange.baseArrayLayer;
-	const uint32 LastLayer = FirstLayer + GetSubresRangeLayerCount(SubresourceRange);
-
-	const uint32 FirstMip = SubresourceRange.baseMipLevel;
-	const uint32 LastMip = FirstMip + GetSubresRangeMipCount(SubresourceRange);
-
-	for (uint32 PlaneIdx = FirstPlane; PlaneIdx < LastPlane; ++PlaneIdx)
-	{
-		for (uint32 LayerIdx = FirstLayer; LayerIdx < LastLayer; ++LayerIdx)
-		{
-			for (uint32 MipIdx = FirstMip; MipIdx < LastMip; ++MipIdx)
-			{
-				if (SubresLayouts[(PlaneIdx * NumLayers * NumMips) + (LayerIdx * NumMips) + MipIdx] != Layout)
-				{
-					return false;
-				}
-			}
-		}
-	}
-
-	return true;
-}
-
-void FVulkanImageLayout::CollapseSubresLayoutsIfSame()
-{
-	if (SubresLayouts.Num() == 0)
-	{
-		return;
-	}
-
-	const VkImageLayout Layout = SubresLayouts[0];
-	for (uint32 i = 1; i < NumPlanes * NumLayers * NumMips; ++i)
-	{
-		if (SubresLayouts[i] != Layout)
-		{
-			return;
-		}
-	}
-
-	MainLayout = Layout;
-	SubresLayouts.Reset();
-}
-
-void FVulkanImageLayout::Set(VkImageLayout Layout, const VkImageSubresourceRange& SubresourceRange)
-{
-	checkf(
-		(Layout != VK_IMAGE_LAYOUT_DEPTH_READ_ONLY_STENCIL_ATTACHMENT_OPTIMAL) &&
-		(Layout != VK_IMAGE_LAYOUT_DEPTH_ATTACHMENT_STENCIL_READ_ONLY_OPTIMAL) &&
-		(Layout != VK_IMAGE_LAYOUT_DEPTH_STENCIL_ATTACHMENT_OPTIMAL) &&
-		(Layout != VK_IMAGE_LAYOUT_DEPTH_STENCIL_READ_ONLY_OPTIMAL),
-		TEXT("Layout tracking should always use separate depth and stencil layouts.")
-	);
-
-	const uint32 FirstPlane = (SubresourceRange.aspectMask == VK_IMAGE_ASPECT_STENCIL_BIT) ? NumPlanes - 1 : 0;
-	const uint32 LastPlane = (SubresourceRange.aspectMask & VK_IMAGE_ASPECT_STENCIL_BIT) ? NumPlanes : 1;
-
-	const uint32 FirstLayer = SubresourceRange.baseArrayLayer;
-	const uint32 LayerCount = GetSubresRangeLayerCount(SubresourceRange);
-
-	const uint32 FirstMip = SubresourceRange.baseMipLevel;
-	const uint32 MipCount = GetSubresRangeMipCount(SubresourceRange);
-
-	if (FirstPlane == 0 && LastPlane == NumPlanes &&
-		FirstLayer == 0 && LayerCount == NumLayers && 
-		FirstMip == 0 && MipCount == NumMips)
-	{
-		// We're setting the entire resource to the same layout.
-		MainLayout = Layout;
-		SubresLayouts.Reset();
-		return;
-	}
-
-	if (SubresLayouts.Num() == 0)
-	{
-		const uint32 SubresLayoutCount = NumPlanes * NumLayers * NumMips;
-		SubresLayouts.SetNum(SubresLayoutCount);
-		for (uint32 i = 0; i < SubresLayoutCount; ++i)
-		{
-			SubresLayouts[i] = MainLayout;
-		}
-	}
-
-	for (uint32 Plane = FirstPlane; Plane < LastPlane; ++Plane)
-	{
-		for (uint32 Layer = FirstLayer; Layer < FirstLayer + LayerCount; ++Layer)
-		{
-			for (uint32 Mip = FirstMip; Mip < FirstMip + MipCount; ++Mip)
-			{
-				SubresLayouts[Plane * (NumLayers * NumMips) + Layer * NumMips + Mip] = Layout;
-			}
-		}
-	}
-
-	// It's possible we've just set all the subresources to the same layout. If that's the case, get rid of the
-	// subresource info and set the main layout appropriatedly.
-	CollapseSubresLayoutsIfSame();
-}
-
-VkImageLayout FVulkanLayoutManager::GetDefaultLayout(FVulkanCommandBuffer* CmdBuffer, const FVulkanTexture& VulkanTexture, ERHIAccess DesiredAccess)
+VkImageLayout FVulkanPipelineBarrier::GetDefaultLayout(const FVulkanTexture& VulkanTexture, ERHIAccess DesiredAccess)
 {
 	switch (DesiredAccess)
 	{
@@ -1979,19 +1599,13 @@ VkImageLayout FVulkanLayoutManager::GetDefaultLayout(FVulkanCommandBuffer* CmdBu
 	{
 		if (VulkanTexture.IsDepthOrStencilAspect())
 		{
-			const bool bSupportsParallelRendering = CmdBuffer->Device.SupportsParallelRendering();
-			if (bSupportsParallelRendering)
+			if (VulkanTexture.Device->SupportsParallelRendering())
 			{
 				return VK_IMAGE_LAYOUT_READ_ONLY_OPTIMAL;
-
 			}
 			else
 			{
-				// The only layout that can't be hardcoded...  Even if we only use depth, we need to know stencil
-				FVulkanLayoutManager& LayoutMgr = CmdBuffer->GetLayoutManager();
-				const VkImageLayout DepthLayout = LayoutMgr.GetDepthStencilHint(VulkanTexture, VK_IMAGE_ASPECT_DEPTH_BIT);
-				const VkImageLayout StencilLayout = LayoutMgr.GetDepthStencilHint(VulkanTexture, VK_IMAGE_ASPECT_STENCIL_BIT);
-				return VulkanRHI::GetMergedDepthStencilLayout(DepthLayout, StencilLayout);
+				return VK_IMAGE_LAYOUT_DEPTH_STENCIL_READ_ONLY_OPTIMAL;
 			}
 		}
 		else
@@ -2027,73 +1641,3 @@ VkImageLayout FVulkanLayoutManager::GetDefaultLayout(FVulkanCommandBuffer* CmdBu
 		return VK_IMAGE_LAYOUT_UNDEFINED;
 	}
 }
-
-VkImageLayout FVulkanLayoutManager::SetExpectedLayout(FVulkanCommandBuffer* CmdBuffer, const FVulkanTexture& VulkanTexture, ERHIAccess DesiredAccess)
-{
-	const VkImageLayout ExpectedLayout = GetDefaultLayout(CmdBuffer, VulkanTexture, DesiredAccess);
-	VkImageLayout PreviousLayout = ExpectedLayout;
-
-	FVulkanLayoutManager& LayoutMgr = CmdBuffer->GetLayoutManager();
-
-	// This code path is not safe for multi-threaded code
-	if (!LayoutMgr.bWriteOnly || GVulkanAutoCorrectExpectedLayouts)
-	{
-		const FVulkanImageLayout* SrcLayout = LayoutMgr.Layouts.Find(VulkanTexture.Image);
-
-		if (!SrcLayout)
-		{
-			FVulkanLayoutManager& QueueLayoutMgr = CmdBuffer->GetOwner()->GetContext().Queue.GetLayoutManager();
-			SrcLayout = QueueLayoutMgr.Layouts.Find(VulkanTexture.Image);
-		}
-
-		if (SrcLayout && (!SrcLayout->AreAllSubresourcesSameLayout() || (SrcLayout->MainLayout != ExpectedLayout)))
-		{
-			FVulkanPipelineBarrier Barrier;
-			if (VulkanTexture.IsDepthOrStencilAspect() && (SrcLayout->NumPlanes > 1) && !LayoutMgr.bWriteOnly)
-			{
-				const VkImageLayout MergedLayout = VulkanRHI::GetMergedDepthStencilLayout(SrcLayout->GetSubresLayout(0, 0, VK_IMAGE_ASPECT_DEPTH_BIT), SrcLayout->GetSubresLayout(0, 0, VK_IMAGE_ASPECT_STENCIL_BIT));
-				Barrier.AddImageLayoutTransition(VulkanTexture.Image, MergedLayout, ExpectedLayout, FVulkanPipelineBarrier::MakeSubresourceRange(VulkanTexture.GetFullAspectMask()));
-				PreviousLayout = MergedLayout;
-			}
-			else
-			{
-				Barrier.AddImageLayoutTransition(VulkanTexture.Image, VulkanTexture.GetFullAspectMask(), *SrcLayout, ExpectedLayout);
-				check(SrcLayout->AreAllSubresourcesSameLayout());
-				PreviousLayout = SrcLayout->MainLayout;
-			}
-			Barrier.Execute(CmdBuffer);
-		}
-	}
-
-	return PreviousLayout;
-}
-
-// todo-jn: these function are called hints, but it more of a temporary hack to access depth/stencil information for render passes
-VkImageLayout FVulkanLayoutManager::GetDepthStencilHint(const FVulkanTexture& VulkanTexture, VkImageAspectFlagBits AspectBit)
-{
-	check(AspectBit == VK_IMAGE_ASPECT_DEPTH_BIT || AspectBit == VK_IMAGE_ASPECT_STENCIL_BIT);
-	FVulkanImageLayout* Layout = Layouts.Find(VulkanTexture.Image);
-	if (Layout)
-	{
-		return Layout->GetSubresLayout(0, 0, AspectBit);
-	}
-	else if (Fallback)
-	{
-		return Fallback->GetDepthStencilHint(VulkanTexture, AspectBit);
-	}
-	return VK_IMAGE_LAYOUT_UNDEFINED;
-}
-
-void FVulkanLayoutManager::NotifyDeletedImage(VkImage Image)
-{
-	Layouts.Remove(Image);
-}
-
-void FVulkanLayoutManager::TransferTo(FVulkanLayoutManager& Destination)
-{
-	if (Layouts.Num())
-	{
-		Destination.Layouts.Append(Layouts);
-		Layouts.Empty();
-	}
-}
diff --git a/Engine/Source/Runtime/VulkanRHI/Private/VulkanBarriers.h b/Engine/Source/Runtime/VulkanRHI/Private/VulkanBarriers.h
--- a/Engine/Source/Runtime/VulkanRHI/Private/VulkanBarriers.h
+++ b/Engine/Source/Runtime/VulkanRHI/Private/VulkanBarriers.h
@@ -29,6 +29,7 @@ struct FVulkanPipelineBarrier
 	struct ImageBarrierExtraData
 	{
 		FVulkanTexture* BaseTexture = nullptr;
+		VkImageLayout OtherPlaneLayout = VK_IMAGE_LAYOUT_MAX_ENUM;
 		bool IsAliasingBarrier = false;
 	};
 	TArray<ImageBarrierExtraData, TInlineAllocator<2>> ImageBarrierExtras;
@@ -36,220 +37,13 @@ struct FVulkanPipelineBarrier
 	void AddMemoryBarrier(VkAccessFlags SrcAccessFlags, VkAccessFlags DstAccessFlags, VkPipelineStageFlags SrcStageMask, VkPipelineStageFlags DstStageMask);
 	void AddFullImageLayoutTransition(const FVulkanTexture& Texture, VkImageLayout SrcLayout, VkImageLayout DstLayout);
 	void AddImageLayoutTransition(VkImage Image, VkImageLayout SrcLayout, VkImageLayout DstLayout, const VkImageSubresourceRange& SubresourceRange);
-	void AddImageLayoutTransition(VkImage Image, VkImageAspectFlags AspectMask, const struct FVulkanImageLayout& SrcLayout, VkImageLayout DstLayout);
-	void AddImageLayoutTransition(VkImage Image, VkImageAspectFlags AspectMask, VkImageLayout SrcLayout, const struct FVulkanImageLayout& DstLayout);
-	void AddImageLayoutTransition(VkImage Image, VkImageAspectFlags AspectMask, const struct FVulkanImageLayout& SrcLayout, const struct FVulkanImageLayout& DstLayout);
 	void AddImageAccessTransition(const FVulkanTexture& Surface, ERHIAccess SrcAccess, ERHIAccess DstAccess, const VkImageSubresourceRange& SubresourceRange, VkImageLayout& InOutLayout);
 	void Execute(VkCommandBuffer CmdBuffer);
 	void Execute(FVulkanCommandBuffer* CmdBuffer);
 
 	static VkImageSubresourceRange MakeSubresourceRange(VkImageAspectFlags AspectMask, uint32 FirstMip = 0, uint32 NumMips = VK_REMAINING_MIP_LEVELS, uint32 FirstLayer = 0, uint32 NumLayers = VK_REMAINING_ARRAY_LAYERS);
-};
-
-struct FVulkanImageLayout
-{
-	FVulkanImageLayout(VkImageLayout InitialLayout, uint32 InNumMips, uint32 InNumLayers, VkImageAspectFlags Aspect) :
-		NumMips(InNumMips),
-		NumLayers(InNumLayers),
-		NumPlanes((Aspect == (VK_IMAGE_ASPECT_DEPTH_BIT | VK_IMAGE_ASPECT_STENCIL_BIT)) ? 2 : 1),
-		MainLayout(InitialLayout)
-	{
-	}
-
-	uint32 NumMips;
-	uint32 NumLayers;
-	uint32 NumPlanes;
-
-	// The layout when all the subresources are in the same state.
-	VkImageLayout MainLayout;
-
-	// Explicit subresource layouts. Always NumLayers*NumMips elements.
-	TArray<VkImageLayout> SubresLayouts;
-
-	inline bool AreAllSubresourcesSameLayout() const
-	{
-		return SubresLayouts.Num() == 0;
-	}
-
-	VkImageLayout GetSubresLayout(uint32 Layer, uint32 Mip, VkImageAspectFlagBits Aspect) const
-	{
-		return GetSubresLayout(Layer, Mip, (Aspect==VK_IMAGE_ASPECT_STENCIL_BIT) ? NumPlanes - 1 : 0);
-	}
-
-	VkImageLayout GetSubresLayout(uint32 Layer, uint32 Mip, uint32 Plane) const
-	{
-		if (SubresLayouts.Num() == 0)
-		{
-			return MainLayout;
-		}
-
-		if (Layer == (uint32)-1)
-		{
-			Layer = 0;
-		}
-
-		check(Plane < NumPlanes && Layer < NumLayers && Mip < NumMips);
-		return SubresLayouts[(Plane * NumLayers * NumMips) + (Layer * NumMips) + Mip];
-	}
-
-	bool AreSubresourcesSameLayout(VkImageLayout Layout, const VkImageSubresourceRange& SubresourceRange) const;
-
-	inline uint32 GetSubresRangeLayerCount(const VkImageSubresourceRange& SubresourceRange) const
-	{
-		check(SubresourceRange.baseArrayLayer < NumLayers);
-		return (SubresourceRange.layerCount == VK_REMAINING_ARRAY_LAYERS) ? (NumLayers - SubresourceRange.baseArrayLayer) : SubresourceRange.layerCount;
-	}
-
-	inline uint32 GetSubresRangeMipCount(const VkImageSubresourceRange& SubresourceRange) const
-	{
-		check(SubresourceRange.baseMipLevel < NumMips);
-		return (SubresourceRange.levelCount == VK_REMAINING_MIP_LEVELS) ? (NumMips - SubresourceRange.baseMipLevel) : SubresourceRange.levelCount;
-	}
-
-	void CollapseSubresLayoutsIfSame();
-
-	void Set(VkImageLayout Layout, const VkImageSubresourceRange& SubresourceRange);
-};
-
-class FVulkanLayoutManager
-{
-public:
-	FVulkanLayoutManager(bool InWriteOnly, FVulkanLayoutManager* InFallback)
-		: bWriteOnly(InWriteOnly)
-		, Fallback(InFallback)
-	{
-	}
-
-	void NotifyDeletedImage(VkImage Image);
 
 	// Predetermined layouts for given RHIAccess
-	static VkImageLayout GetDefaultLayout(FVulkanCommandBuffer* CmdBuffer, const FVulkanTexture& VulkanTexture, ERHIAccess DesiredAccess);
-
-	// Expected layouts and Hints are workarounds until we can use 'hardcoded layouts' everywhere.
-	static VkImageLayout SetExpectedLayout(FVulkanCommandBuffer* CmdBuffer, const FVulkanTexture& VulkanTexture, ERHIAccess DesiredAccess);
-	VkImageLayout GetDepthStencilHint(const FVulkanTexture& VulkanTexture, VkImageAspectFlagBits AspectBit);
-
-	VULKANRHI_API const FVulkanImageLayout* GetFullLayout(VkImage Image) const
-	{
-		check(!bWriteOnly);
-		const FVulkanImageLayout* Layout = Layouts.Find(Image);
-		if (!Layout && Fallback)
-		{
-			return Fallback->GetFullLayout(Image);
-		}
-		return Layout;
-	}
-
-	VULKANRHI_API const FVulkanImageLayout* GetFullLayout(const FVulkanTexture& VulkanTexture, bool bAddIfNotFound = false, VkImageLayout LayoutIfNotFound = VK_IMAGE_LAYOUT_UNDEFINED)
-	{
-		check(!bWriteOnly);
-		const FVulkanImageLayout* Layout = Layouts.Find(VulkanTexture.Image);
-		
-		if (!Layout && Fallback)
-		{
-			Layout = Fallback->GetFullLayout(VulkanTexture, false);
-
-			// If the layout was found in the fallback, carry it forward to our current manager for future tracking
-			if (Layout)
-			{
-				return &Layouts.Add(VulkanTexture.Image, *Layout);
-			}
-		}
-		
-		if (Layout)
-		{
-			return Layout;
-		}
-		else if (!bAddIfNotFound)
-		{
-			return nullptr;
-		}
-
-		return &Layouts.Add(VulkanTexture.Image, FVulkanImageLayout(LayoutIfNotFound, VulkanTexture.GetNumMips(), VulkanTexture.GetNumberOfArrayLevels(), VulkanTexture.GetFullAspectMask()));
-	}
-
-	// Not the preferred path because we can't ensure Mip and Layer counts match, but still necessary for images like the backbuffer
-	VULKANRHI_API void SetFullLayout(VkImage Image, const FVulkanImageLayout& NewLayout)
-	{
-		FVulkanImageLayout* Layout = Layouts.Find(Image);
-		if (Layout)
-		{
-			*Layout = NewLayout;
-		}
-		else
-		{
-			Layouts.Add(Image, NewLayout);
-		}
-	}
-
-	VULKANRHI_API void SetFullLayout(const FVulkanTexture& VulkanTexture, const FVulkanImageLayout& NewLayout)
-	{
-		check((VulkanTexture.GetNumMips() == NewLayout.NumMips) && (VulkanTexture.GetNumberOfArrayLevels() == NewLayout.NumLayers));
-		SetFullLayout(VulkanTexture.Image, NewLayout);
-	}
-
-	VULKANRHI_API void SetFullLayout(const FVulkanTexture& VulkanTexture, VkImageLayout InLayout, bool bOnlyIfNotFound=false)
-	{
-		FVulkanImageLayout* Layout = Layouts.Find(VulkanTexture.Image);
-		if (Layout)
-		{
-			if (!bOnlyIfNotFound)
-			{
-				Layout->Set(InLayout, FVulkanPipelineBarrier::MakeSubresourceRange(VulkanTexture.GetFullAspectMask()));
-			}
-		}
-		else
-		{
-			Layouts.Add(VulkanTexture.Image, FVulkanImageLayout(InLayout, VulkanTexture.GetNumMips(), VulkanTexture.GetNumberOfArrayLevels(), VulkanTexture.GetFullAspectMask()));
-		}
-	}
-
-	static bool IsPartialResource(const FVulkanTexture& VulkanTexture, const VkImageSubresourceRange& InSubresourceRange)
-	{
-		return (VulkanTexture.GetFullAspectMask() != InSubresourceRange.aspectMask) ||
-			(InSubresourceRange.baseArrayLayer != 0) ||
-			(InSubresourceRange.baseMipLevel != 0) ||
-			((InSubresourceRange.levelCount != VK_REMAINING_MIP_LEVELS) && (InSubresourceRange.levelCount != VulkanTexture.GetNumMips())) ||
-			((InSubresourceRange.layerCount != VK_REMAINING_ARRAY_LAYERS) && (InSubresourceRange.layerCount != VulkanTexture.GetNumberOfArrayLevels()));
-	}
-
-	VULKANRHI_API void SetLayout(const FVulkanTexture& VulkanTexture, const VkImageSubresourceRange& InSubresourceRange, VkImageLayout InLayout)
-	{
-		FVulkanImageLayout* Layout = Layouts.Find(VulkanTexture.Image);
-
-		// If we're not going to overwrite the entire resource, start from its last known state on the queue
-		if (!Layout && IsPartialResource(VulkanTexture, InSubresourceRange) && Fallback)
-		{
-			const FVulkanImageLayout* FallbackLayout = Fallback->GetFullLayout(VulkanTexture, false);
-			if (FallbackLayout)
-			{
-				Layout = &Layouts.Add(VulkanTexture.Image, *FallbackLayout);
-			}
-		}
-
-		if (Layout)
-		{
-			Layout->Set(InLayout, InSubresourceRange);
-		}
-		else
-		{
-			FVulkanImageLayout NewLayout(VK_IMAGE_LAYOUT_UNDEFINED, VulkanTexture.GetNumMips(), VulkanTexture.GetNumberOfArrayLevels(), VulkanTexture.GetFullAspectMask());
-			NewLayout.Set(InLayout, InSubresourceRange);
-			Layouts.Add(VulkanTexture.Image, NewLayout);
-		}
-	}
-
-	// Transfers our layouts into the destination
-	void TransferTo(FVulkanLayoutManager& Destination);
-
-private:
-
-	TMap<VkImage, FVulkanImageLayout> Layouts;
-
-	// If we're WriteOnly, we should never read layout from this instance.  This is important for parallel rendering.
-	// When WriteOnly, this instance of the layout manager should only collect layouts to later feed them to the another central mgr.
-	const bool bWriteOnly;
-
-	// If parallel command list creation is NOT supported, then the queue's layout mgr can be used as a fallback to fetch previous layouts.
-	FVulkanLayoutManager* Fallback;
+	static VkImageLayout GetDefaultLayout(const FVulkanTexture& VulkanTexture, ERHIAccess DesiredAccess);
+	static VkImageLayout GetDepthOrStencilLayout(ERHIAccess Access);
 };
diff --git a/Engine/Source/Runtime/VulkanRHI/Private/VulkanCommandBuffer.cpp b/Engine/Source/Runtime/VulkanRHI/Private/VulkanCommandBuffer.cpp
--- a/Engine/Source/Runtime/VulkanRHI/Private/VulkanCommandBuffer.cpp
+++ b/Engine/Source/Runtime/VulkanRHI/Private/VulkanCommandBuffer.cpp
@@ -23,7 +23,6 @@ static FAutoConsoleVariableRef CVarVulkanUseSingleQueue(
 #define CMD_BUFFER_TIME_TO_WAIT_BEFORE_DELETING		10
 
 const uint32 GNumberOfFramesBeforeDeletingDescriptorPool = 300;
-extern int32 GVulkanAutoCorrectUnknownLayouts;
 
 FVulkanCommandBuffer::FVulkanCommandBuffer(FVulkanDevice& InDevice, FVulkanCommandBufferPool* InCommandBufferPool)
 	: CurrentStencilRef(0)
@@ -37,8 +36,6 @@ FVulkanCommandBuffer::FVulkanCommandBuffer(FVulkanDevice& InDevice, FVulkanComma
 	, CommandBufferHandle(VK_NULL_HANDLE)
 	, CommandBufferPool(InCommandBufferPool)
 	, LastValidTiming(0)
-	, LayoutManager(InDevice.SupportsParallelRendering() && !GVulkanAutoCorrectUnknownLayouts,
-		&InCommandBufferPool->GetContext().Queue.GetLayoutManager())
 {
 	{
 		FScopeLock ScopeLock(CommandBufferPool->GetCS());
diff --git a/Engine/Source/Runtime/VulkanRHI/Private/VulkanCommandBuffer.h b/Engine/Source/Runtime/VulkanRHI/Private/VulkanCommandBuffer.h
--- a/Engine/Source/Runtime/VulkanRHI/Private/VulkanCommandBuffer.h
+++ b/Engine/Source/Runtime/VulkanRHI/Private/VulkanCommandBuffer.h
@@ -77,11 +77,6 @@ public:
 		return CommandBufferHandle;
 	}
 
-	FVulkanLayoutManager& GetLayoutManager()
-	{
-		return LayoutManager;
-	}
-
 	void Begin();
 	void End();
 
@@ -146,14 +141,11 @@ private:
 	void AllocMemory();
 	void FreeMemory();
 
-	FVulkanLayoutManager LayoutManager;
-
 public:
 	//#todo-rco: Hide this
 	TMap<uint32, class FVulkanTypedDescriptorPoolSet*> TypedDescriptorPoolSets;
 
 	friend class FVulkanDynamicRHI;
-	friend FVulkanLayoutManager;
 	friend FVulkanContextCommon;
 };
 
diff --git a/Engine/Source/Runtime/VulkanRHI/Private/VulkanCommands.cpp b/Engine/Source/Runtime/VulkanRHI/Private/VulkanCommands.cpp
--- a/Engine/Source/Runtime/VulkanRHI/Private/VulkanCommands.cpp
+++ b/Engine/Source/Runtime/VulkanRHI/Private/VulkanCommands.cpp
@@ -166,7 +166,7 @@ struct FVulkanResourceBinder
 		{
 			FVulkanTexture* VulkanTexture = ResourceCast(TextureRHI);
 			const ERHIAccess RHIAccess = (Frequency == SF_Compute) ? ERHIAccess::SRVCompute : ERHIAccess::SRVGraphics;
-			const VkImageLayout ExpectedLayout = FVulkanLayoutManager::GetDefaultLayout(Context.GetActiveCmdBuffer(), *VulkanTexture, RHIAccess);
+			const VkImageLayout ExpectedLayout = FVulkanPipelineBarrier::GetDefaultLayout(*VulkanTexture, RHIAccess);
 			PendingState->SetTextureForUBResource(Stage, Index, VulkanTexture, ExpectedLayout);
 		}
 	}
@@ -351,7 +351,7 @@ void FVulkanCommandListContext::RHISetUAVParameter(FRHIComputeShader* ComputeSha
 void FVulkanCommandListContext::RHISetShaderTexture(FRHIGraphicsShader* ShaderRHI, uint32 TextureIndex, FRHITexture* NewTextureRHI)
 {
 	FVulkanTexture* VulkanTexture = ResourceCast(NewTextureRHI);
-	const VkImageLayout ExpectedLayout = FVulkanLayoutManager::GetDefaultLayout(GetActiveCmdBuffer(), *VulkanTexture, ERHIAccess::SRVGraphics);
+	const VkImageLayout ExpectedLayout = FVulkanPipelineBarrier::GetDefaultLayout(*VulkanTexture, ERHIAccess::SRVGraphics);
 
 	ShaderStage::EStage Stage = GetAndVerifyShaderStage(ShaderRHI, PendingGfxState);
 	PendingGfxState->SetTextureForStage(Stage, TextureIndex, VulkanTexture, ExpectedLayout);
@@ -364,7 +364,7 @@ void FVulkanCommandListContext::RHISetShaderTexture(FRHIComputeShader* ComputeSh
 	check(PendingComputeState->GetCurrentShader() == ComputeShader);
 
 	FVulkanTexture* VulkanTexture = ResourceCast(NewTextureRHI);
-	const VkImageLayout ExpectedLayout = FVulkanLayoutManager::GetDefaultLayout(GetActiveCmdBuffer(), *VulkanTexture, ERHIAccess::SRVCompute);
+	const VkImageLayout ExpectedLayout = FVulkanPipelineBarrier::GetDefaultLayout(*VulkanTexture, ERHIAccess::SRVCompute);
 	PendingComputeState->SetTextureForStage(TextureIndex, VulkanTexture, ExpectedLayout);
 	NewTextureRHI->SetLastRenderTime((float)FPlatformTime::Seconds());
 }
diff --git a/Engine/Source/Runtime/VulkanRHI/Private/VulkanContext.cpp b/Engine/Source/Runtime/VulkanRHI/Private/VulkanContext.cpp
--- a/Engine/Source/Runtime/VulkanRHI/Private/VulkanContext.cpp
+++ b/Engine/Source/Runtime/VulkanRHI/Private/VulkanContext.cpp
@@ -182,26 +182,6 @@ FVulkanCommandBuffer* FVulkanContextCommon::PrepareNewCommandBuffer(bool bTransf
 		NewCommandBuffer = Pool.Create();
 	}
 
-	// Move the layout mgr forward to the new commandbuffer
-	// NOTE: LayoutMgr should be in context, but it will get removed soon anyhow
-	if (bTransferLayoutMgr && Payloads.Num())
-	{
-		FVulkanCommandBuffer* LastCommandBuffer = nullptr;
-		for (int32 PayloadIndex = Payloads.Num() - 1; PayloadIndex >= 0; --PayloadIndex)
-		{
-			if (Payloads[PayloadIndex]->CommandBuffers.Num())
-			{
-				LastCommandBuffer = Payloads[PayloadIndex]->CommandBuffers.Last();
-				break;
-			}
-		}
-
-		if (LastCommandBuffer)
-		{
-			LastCommandBuffer->GetLayoutManager().TransferTo(NewCommandBuffer->GetLayoutManager());
-		}
-	}
-
 	NewCommandBuffer->Begin();
 	return NewCommandBuffer;
 }
diff --git a/Engine/Source/Runtime/VulkanRHI/Private/VulkanContext.h b/Engine/Source/Runtime/VulkanRHI/Private/VulkanContext.h
--- a/Engine/Source/Runtime/VulkanRHI/Private/VulkanContext.h
+++ b/Engine/Source/Runtime/VulkanRHI/Private/VulkanContext.h
@@ -120,11 +120,6 @@ public:
 
 	void SubmitActiveCmdBuffer(FVulkanDynamicRHI* DynamicRHI);
 
-	void NotifyDeletedImage(VkImage Image)
-	{
-		Queue.NotifyDeletedImage(Image);
-	}
-
 	uint64 GetFrameCounter() const
 	{
 		return FrameCounter;
diff --git a/Engine/Source/Runtime/VulkanRHI/Private/VulkanDevice.cpp b/Engine/Source/Runtime/VulkanRHI/Private/VulkanDevice.cpp
--- a/Engine/Source/Runtime/VulkanRHI/Private/VulkanDevice.cpp
+++ b/Engine/Source/Runtime/VulkanRHI/Private/VulkanDevice.cpp
@@ -1559,9 +1559,6 @@ void FVulkanDevice::NotifyDeletedImage(VkImage Image, bool bRenderTarget)
 		// Delete framebuffers using this image
 		GetRenderPassManager().NotifyDeletedRenderTarget(Image);
 	}
-
-	//#todo-jn: Loop through all contexts!  And all queues!
-	GetImmediateContext().NotifyDeletedImage(Image);
 }
 
 void FVulkanDevice::SubmitCommands(FVulkanCommandListContext* Context)
diff --git a/Engine/Source/Runtime/VulkanRHI/Private/VulkanMemory.cpp b/Engine/Source/Runtime/VulkanRHI/Private/VulkanMemory.cpp
--- a/Engine/Source/Runtime/VulkanRHI/Private/VulkanMemory.cpp
+++ b/Engine/Source/Runtime/VulkanRHI/Private/VulkanMemory.cpp
@@ -4182,10 +4182,8 @@ namespace VulkanRHI
 					{
 						FVulkanTexture* Texture = EvictableOwner->GetEvictableTexture();
 
-						// Defrag works from the queue for layouts
 						// Only work with straightforward targets that are in a single layout
-						const FVulkanImageLayout* OriginalLayout = Context.Queue.GetLayoutManager().GetFullLayout(Texture->Image);
-						if (OriginalLayout && OriginalLayout->AreAllSubresourcesSameLayout() && (OriginalLayout->MainLayout != VK_IMAGE_LAYOUT_UNDEFINED))
+						if (Texture)
 						{
 							FVulkanAllocation Allocation;
 							// The current SubAllocator is tagged as locked, this will never allocate in the current SubAllocator.
@@ -4246,7 +4244,10 @@ namespace VulkanRHI
 				case EVulkanAllocationMetaImageOther:
 				{
 					FVulkanEvictable* Texture = Alloc.AllocationOwner;
-					Texture->Evict(Device, Context);
+					if (Texture->CanEvict())
+					{
+						Texture->Evict(Device, Context);
+					}
 				}
 				break;
 				default:
diff --git a/Engine/Source/Runtime/VulkanRHI/Private/VulkanPendingState.cpp b/Engine/Source/Runtime/VulkanRHI/Private/VulkanPendingState.cpp
--- a/Engine/Source/Runtime/VulkanRHI/Private/VulkanPendingState.cpp
+++ b/Engine/Source/Runtime/VulkanRHI/Private/VulkanPendingState.cpp
@@ -329,13 +329,13 @@ FVulkanPendingComputeState::~FVulkanPendingComputeState()
 void FVulkanPendingComputeState::SetSRVForUBResource(uint32 DescriptorSet, uint32 BindingIndex, FVulkanShaderResourceView* SRV)
 {
 	check(SRV);
-	CurrentState->SetSRV(Context.GetActiveCmdBuffer(), true, DescriptorSet, BindingIndex, SRV);
+	CurrentState->SetSRV(true, DescriptorSet, BindingIndex, SRV);
 }
 
 void FVulkanPendingComputeState::SetUAVForUBResource(uint32 DescriptorSet, uint32 BindingIndex, FVulkanUnorderedAccessView* UAV)
 {
 	check(UAV);
-	CurrentState->SetUAV(Context.GetActiveCmdBuffer(), true, DescriptorSet, BindingIndex, UAV);
+	CurrentState->SetUAV(true, DescriptorSet, BindingIndex, UAV);
 }
 
 
@@ -574,13 +574,13 @@ void FVulkanPendingGfxState::UpdateInputAttachments(FVulkanFramebuffer* Framebuf
 void FVulkanPendingGfxState::SetSRVForUBResource(uint8 DescriptorSet, uint32 BindingIndex, FVulkanShaderResourceView* SRV)
 {
 	check(SRV);
-	CurrentState->SetSRV(Context.GetActiveCmdBuffer(), false, DescriptorSet, BindingIndex, SRV);
+	CurrentState->SetSRV(false, DescriptorSet, BindingIndex, SRV);
 }
 
 void FVulkanPendingGfxState::SetUAVForUBResource(uint8 DescriptorSet, uint32 BindingIndex, FVulkanUnorderedAccessView* UAV)
 {
 	check(UAV);
-	CurrentState->SetUAV(Context.GetActiveCmdBuffer(), false, DescriptorSet, BindingIndex, UAV);
+	CurrentState->SetUAV(false, DescriptorSet, BindingIndex, UAV);
 }
 
 int32 GDSetCacheTargetSetsPerPool = 4096;
diff --git a/Engine/Source/Runtime/VulkanRHI/Private/VulkanPipelineState.cpp b/Engine/Source/Runtime/VulkanRHI/Private/VulkanPipelineState.cpp
--- a/Engine/Source/Runtime/VulkanRHI/Private/VulkanPipelineState.cpp
+++ b/Engine/Source/Runtime/VulkanRHI/Private/VulkanPipelineState.cpp
@@ -703,14 +703,10 @@ void FVulkanDescriptorSetWriter::InitWrittenMasks(uint32 NumDescriptorWrites)
 #endif
 }
 
-void FVulkanCommonPipelineDescriptorState::SetSRV(FVulkanCommandBuffer* CmdBuffer, bool bCompute, uint8 DescriptorSet, uint32 BindingIndex, FVulkanShaderResourceView* SRV)
+void FVulkanCommonPipelineDescriptorState::SetSRV(bool bCompute, uint8 DescriptorSet, uint32 BindingIndex, FVulkanShaderResourceView* SRV)
 {
 	check(!bUseBindless);
 
-	ERHIAccess Access = bCompute
-		? ERHIAccess::SRVCompute
-		: ERHIAccess::SRVGraphics;
-
 	switch (SRV->GetViewType())
 	{
 	case FVulkanView::EType::Null:
@@ -723,8 +719,9 @@ void FVulkanCommonPipelineDescriptorState::SetSRV(FVulkanCommandBuffer* CmdBuffe
 
 	case FVulkanView::EType::Texture:
 		{
+			const ERHIAccess Access = bCompute ? ERHIAccess::SRVCompute : ERHIAccess::SRVGraphics;
 			const FVulkanTexture* VulkanTexture = ResourceCast(SRV->GetTexture());
-			const VkImageLayout Layout = FVulkanLayoutManager::GetDefaultLayout(CmdBuffer, *VulkanTexture, Access);
+			const VkImageLayout Layout = FVulkanPipelineBarrier::GetDefaultLayout(*VulkanTexture, Access);
 			MarkDirty(DSWriter[DescriptorSet].WriteImage(BindingIndex, SRV->GetTextureView(), Layout));
 		}
 		break;
@@ -740,14 +737,10 @@ void FVulkanCommonPipelineDescriptorState::SetSRV(FVulkanCommandBuffer* CmdBuffe
 	}
 }
 
-void FVulkanCommonPipelineDescriptorState::SetUAV(FVulkanCommandBuffer* CmdBuffer, bool bCompute, uint8 DescriptorSet, uint32 BindingIndex, FVulkanUnorderedAccessView* UAV)
+void FVulkanCommonPipelineDescriptorState::SetUAV(bool bCompute, uint8 DescriptorSet, uint32 BindingIndex, FVulkanUnorderedAccessView* UAV)
 {
 	check(!bUseBindless);
 
-	ERHIAccess Access = bCompute
-		? ERHIAccess::UAVCompute
-		: ERHIAccess::UAVGraphics;
-
 	switch (UAV->GetViewType())
 	{
 	case FVulkanView::EType::Null:
@@ -760,13 +753,13 @@ void FVulkanCommonPipelineDescriptorState::SetUAV(FVulkanCommandBuffer* CmdBuffe
 
 	case FVulkanView::EType::Texture:
 		{
+			const ERHIAccess Access = bCompute ? ERHIAccess::UAVCompute : ERHIAccess::UAVGraphics;
 			const FVulkanTexture* VulkanTexture = ResourceCast(UAV->GetTexture());
-			const VkImageLayout ExpectedLayout = FVulkanLayoutManager::GetDefaultLayout(CmdBuffer, *VulkanTexture, Access);
+			const VkImageLayout ExpectedLayout = FVulkanPipelineBarrier::GetDefaultLayout(*VulkanTexture, Access);
 			MarkDirty(DSWriter[DescriptorSet].WriteStorageImage(BindingIndex, UAV->GetTextureView(), ExpectedLayout));
 		}
 		break;
 
-
 	case FVulkanView::EType::StructuredBuffer:
 		check((ResourceCast(UAV->GetBuffer())->GetBufferUsageFlags() & VK_BUFFER_USAGE_STORAGE_BUFFER_BIT) == VK_BUFFER_USAGE_STORAGE_BUFFER_BIT);
 		MarkDirty(DSWriter[DescriptorSet].WriteStorageBuffer(BindingIndex, UAV->GetStructuredBufferView()));
diff --git a/Engine/Source/Runtime/VulkanRHI/Private/VulkanPipelineState.h b/Engine/Source/Runtime/VulkanRHI/Private/VulkanPipelineState.h
--- a/Engine/Source/Runtime/VulkanRHI/Private/VulkanPipelineState.h
+++ b/Engine/Source/Runtime/VulkanRHI/Private/VulkanPipelineState.h
@@ -61,8 +61,8 @@ public:
 		bIsDSetsKeyDirty |= bDirty;
 	}
 
-	void SetSRV(FVulkanCommandBuffer* CmdBuffer, bool bCompute, uint8 DescriptorSet, uint32 BindingIndex, FVulkanShaderResourceView* SRV);
-	void SetUAV(FVulkanCommandBuffer* CmdBuffer, bool bCompute, uint8 DescriptorSet, uint32 BindingIndex, FVulkanUnorderedAccessView* UAV);
+	void SetSRV(bool bCompute, uint8 DescriptorSet, uint32 BindingIndex, FVulkanShaderResourceView* SRV);
+	void SetUAV(bool bCompute, uint8 DescriptorSet, uint32 BindingIndex, FVulkanUnorderedAccessView* UAV);
 
 	inline void SetTexture(uint8 DescriptorSet, uint32 BindingIndex, const FVulkanTexture* Texture, VkImageLayout Layout)
 	{
diff --git a/Engine/Source/Runtime/VulkanRHI/Private/VulkanQueue.cpp b/Engine/Source/Runtime/VulkanRHI/Private/VulkanQueue.cpp
--- a/Engine/Source/Runtime/VulkanRHI/Private/VulkanQueue.cpp
+++ b/Engine/Source/Runtime/VulkanRHI/Private/VulkanQueue.cpp
@@ -47,7 +47,6 @@ FVulkanQueue::FVulkanQueue(FVulkanDevice& InDevice, uint32 InFamilyIndex, EVulka
 	, QueueType(InQueueType)
 	, Device(InDevice)
 	, bUseTimelineSemaphores(InDevice.GetOptionalExtensions().HasKHRTimelineSemaphore && GAllowTimelineSemaphores)
-	, LayoutManager(false, nullptr)
 {
 	VulkanRHI::vkGetDeviceQueue(Device.GetInstanceHandle(), FamilyIndex, QueueIndex, &Queue);
 
@@ -189,9 +188,6 @@ void FVulkanQueue::SubmitPayloads(TArrayView<FVulkanPayload*> Payloads, TSet<VkS
 
 				CommandBuffer->State = FVulkanCommandBuffer::EState::Submitted;
 				CommandBuffer->SubmittedTime = FPlatformTime::Seconds();
-
-				// If we're tracking layouts for the queue, merge in the changes recorded in this command buffer's context
-				CommandBuffer->GetLayoutManager().TransferTo(LayoutManager);
 			}
 			SubmitInfo.commandBufferCount = Payload->CommandBuffers.Num();
 			SubmitInfo.pCommandBuffers = &CommandBufferStorage[FirstIndex];
@@ -418,8 +414,3 @@ void FVulkanQueue::FillSupportedStageBits()
 		SupportedStages |= VK_PIPELINE_STAGE_TRANSFER_BIT;
 	}
 }
-
-void FVulkanQueue::NotifyDeletedImage(VkImage Image)
-{
-	LayoutManager.NotifyDeletedImage(Image);
-}
diff --git a/Engine/Source/Runtime/VulkanRHI/Private/VulkanQueue.h b/Engine/Source/Runtime/VulkanRHI/Private/VulkanQueue.h
--- a/Engine/Source/Runtime/VulkanRHI/Private/VulkanQueue.h
+++ b/Engine/Source/Runtime/VulkanRHI/Private/VulkanQueue.h
@@ -11,7 +11,6 @@
 
 class FVulkanDynamicRHI;
 class FVulkanDevice;
-class FVulkanLayoutManager;
 class FVulkanPayload;
 
 namespace VulkanRHI
@@ -55,13 +54,6 @@ public:
 		return SupportedStages;
 	}
 
-	inline FVulkanLayoutManager& GetLayoutManager()
-	{
-		return LayoutManager;
-	}
-
-	void NotifyDeletedImage(VkImage Image);
-
 	const VulkanRHI::FSemaphore* GetTimelineSemaphore() const
 	{
 		return TimelineSempahore;
@@ -116,8 +108,5 @@ private:
 	TQueue<FVulkanPayload*, EQueueMode::Mpsc> PendingSubmission;
 	TQueue<FVulkanPayload*, EQueueMode::Spsc> PendingInterrupt;
 
-	// Last known layouts submitted on this queue, used for defrag
-	FVulkanLayoutManager LayoutManager;
-
 	void FillSupportedStageBits();
 };
diff --git a/Engine/Source/Runtime/VulkanRHI/Private/VulkanRHI.cpp b/Engine/Source/Runtime/VulkanRHI/Private/VulkanRHI.cpp
--- a/Engine/Source/Runtime/VulkanRHI/Private/VulkanRHI.cpp
+++ b/Engine/Source/Runtime/VulkanRHI/Private/VulkanRHI.cpp
@@ -855,6 +855,10 @@ void FVulkanDynamicRHI::InitInstance()
 		GRHISupportsMSAAShaderResolve = Device->GetOptionalExtensions().HasQcomRenderPassShaderResolve ? true : false;
 		GRHISupportsRayTracing = RHISupportsRayTracing(GMaxRHIShaderPlatform) && Device->GetOptionalExtensions().HasRaytracingExtensions();
 
+		GRHIGlobals.EnableHighLevelTransitionPatching = true;
+		GRHIGlobals.NeedExtraTransitions = true;
+		GRHIGlobals.NeedsAllPlanes = true;
+
 		// Use this compatibility mode avoid known issues at launch time with latest drivers at the time of release 5.5.  This will:
 		// - disable inline ray tracing and use ray tracing pipelines everywhere (instead of a mix of both)
 		// - disable mesh shaders until issues can be resolved (holes in Nanite meshes)
diff --git a/Engine/Source/Runtime/VulkanRHI/Private/VulkanRayTracing.cpp b/Engine/Source/Runtime/VulkanRHI/Private/VulkanRayTracing.cpp
--- a/Engine/Source/Runtime/VulkanRHI/Private/VulkanRayTracing.cpp
+++ b/Engine/Source/Runtime/VulkanRHI/Private/VulkanRayTracing.cpp
@@ -1912,7 +1912,7 @@ void FVulkanRayTracingCompactionRequestHandler::Update(FVulkanCommandListContext
 }
 
 
-static FVulkanPipelineBarrier SetRayGenResources(FVulkanDevice& Device, FVulkanCommandBuffer* const CmdBuffer, const FRayTracingShaderBindings& InGlobalResourceBindings, FVulkanRayTracingShaderTable* ShaderTable)
+static void SetRayGenResources(FVulkanDevice& Device, FVulkanCommandBuffer* const CmdBuffer, const FRayTracingShaderBindings& InGlobalResourceBindings, FVulkanRayTracingShaderTable* ShaderTable)
 {
 	TArray<const FVulkanUniformBuffer*> UniformBuffers;
 	UniformBuffers.Reserve(UE_ARRAY_COUNT(InGlobalResourceBindings.UniformBuffers));
@@ -1963,151 +1963,6 @@ static FVulkanPipelineBarrier SetRayGenResources(FVulkanDevice& Device, FVulkanC
 		const FVulkanUniformBuffer* VulkanUniformBuffer = ResourceCast(UniformBuffer.GetReference());
 		UniformBuffers.AddUnique(VulkanUniformBuffer);
 	}
-
-	// Track all the missing transitions for the dispatch to be able to bring it back afterwards (will not touch tracking)
-	FVulkanPipelineBarrier PreDispatch, PostDispatch;
-	{
-		auto TransitionBuffer = [&PreDispatch, &PostDispatch](bool bReadOnly)
-		{
-			// :todo-jn: tighten these barriers
-			const VkAccessFlags RWAccessFlags = VK_ACCESS_MEMORY_READ_BIT | VK_ACCESS_MEMORY_WRITE_BIT;
-			const VkAccessFlags DesiredAccessFlags = bReadOnly ? VK_ACCESS_MEMORY_READ_BIT : VK_ACCESS_MEMORY_READ_BIT | VK_ACCESS_MEMORY_WRITE_BIT;
-			PreDispatch.AddMemoryBarrier(RWAccessFlags, DesiredAccessFlags, VK_PIPELINE_STAGE_ALL_COMMANDS_BIT, VK_PIPELINE_STAGE_RAY_TRACING_SHADER_BIT_KHR);
-			PostDispatch.AddMemoryBarrier(DesiredAccessFlags, RWAccessFlags, VK_PIPELINE_STAGE_RAY_TRACING_SHADER_BIT_KHR, VK_PIPELINE_STAGE_ALL_COMMANDS_BIT);
-		};
-
-		// Make sure we only transition textures once, accumulate them in sets
-		TSet<FRHITexture*> SRVTransitions;
-		TSet<FRHITexture*> UAVTransitions;
-		for (const FVulkanUniformBuffer* UniformBuffer : UniformBuffers)
-		{
-			const TArray<TRefCountPtr<FRHIResource>>& ResourceTable = UniformBuffer->GetResourceTable();
-			for (const TRefCountPtr<FRHIResource>& RHIResourceRef : ResourceTable)
-			{
-				const FRHIResource* RHIResource = RHIResourceRef.GetReference();
-				if (!RHIResource)
-				{
-					continue;
-				}
-
-				switch (RHIResource->GetType())
-				{
-				case RRT_Texture:
-				case RRT_Texture2D:
-				case RRT_Texture2DArray:
-				case RRT_Texture3D:
-				case RRT_TextureCube:
-					SRVTransitions.Add((FRHITexture*)RHIResource);
-					break;
-
-				case RRT_TextureReference:
-					SRVTransitions.Add(((FRHITextureReference*)RHIResource)->GetReferencedTexture());
-					break;
-
-				case RRT_UnorderedAccessView:
-				{
-					const FRHIUnorderedAccessView* RHIUnorderedAccessView = (FRHIUnorderedAccessView*)RHIResource;
-					if (RHIUnorderedAccessView->IsTexture())
-					{
-						UAVTransitions.Add(RHIUnorderedAccessView->GetTexture());
-					}
-					else
-					{
-						TransitionBuffer(false);
-					}
-					break;
-				}
-
-				case RRT_ShaderResourceView:
-				{
-					const FRHIShaderResourceView* RHIShaderResourceView = (FRHIShaderResourceView*)RHIResource;
-					if (RHIShaderResourceView->IsTexture())
-					{
-						SRVTransitions.Add(RHIShaderResourceView->GetTexture());
-					}
-					else
-					{
-						TransitionBuffer(true);
-					}
-					break;
-				}
-
-				case RRT_RayTracingAccelerationStructure:
-				case RRT_StagingBuffer:
-				case RRT_Buffer:
-					TransitionBuffer(true);
-					break;
-
-				case RRT_SamplerState: [[fallthrough]];
-				default:
-					// Do nothing
-					break;
-				};
-			}
-		}
-
-		auto TransitionTexture = [&PreDispatch, &PostDispatch, CmdBuffer](FRHITexture* RHITexture, bool bReadOnly)
-		{
-			const FVulkanTexture* Texture = ResourceCast(RHITexture);
-
-			// Because Sync2 is a prereq to ray tracing, use the conveniently generic layout
-			const VkImageLayout TargetLayout = bReadOnly ? VK_IMAGE_LAYOUT_READ_ONLY_OPTIMAL : VK_IMAGE_LAYOUT_GENERAL;
-			const FVulkanImageLayout* OriginalLayout = CmdBuffer->GetLayoutManager().GetFullLayout(Texture->Image);
-			if (!OriginalLayout)
-			{
-				return;
-			}
-
-			// If all the subresource are already in a correct layout for the desired RendOnly state, then skip the barrier
-			if (!OriginalLayout->AreAllSubresourcesSameLayout() ||
-				(
-					(bReadOnly && (OriginalLayout->MainLayout != VK_IMAGE_LAYOUT_READ_ONLY_OPTIMAL) && (OriginalLayout->MainLayout != VK_IMAGE_LAYOUT_SHADER_READ_ONLY_OPTIMAL))
-					|| (!bReadOnly && (OriginalLayout->MainLayout != VK_IMAGE_LAYOUT_GENERAL)) // :todo-jn: prevent overlap?
-					))
-			{
-				PreDispatch.AddImageLayoutTransition(Texture->Image, Texture->GetFullAspectMask(), *OriginalLayout, TargetLayout);
-
-				// Transition back to where it was, leaving any undefined transitions to whatever we set them to
-				{
-					FVulkanImageLayout FinalLayout = *OriginalLayout;
-					if (FinalLayout.AreAllSubresourcesSameLayout())
-					{
-						if (FinalLayout.MainLayout == VK_IMAGE_LAYOUT_UNDEFINED)
-						{
-							FinalLayout.MainLayout = TargetLayout;
-						}
-					}
-					else
-					{
-						for (int32 SubResIndex = 0; SubResIndex < FinalLayout.SubresLayouts.Num(); ++SubResIndex)
-						{
-							if (FinalLayout.SubresLayouts[SubResIndex] == VK_IMAGE_LAYOUT_UNDEFINED)
-							{
-								FinalLayout.SubresLayouts[SubResIndex] = TargetLayout;
-							}
-						}
-					}
-					PostDispatch.AddImageLayoutTransition(Texture->Image, Texture->GetFullAspectMask(), TargetLayout, FinalLayout);
-				}
-			}
-		};
-
-		for (FRHITexture* RHITexture : UAVTransitions)
-		{
-			TransitionTexture(RHITexture, false);
-
-			// If a resource shows up as both, use it in VK_IMAGE_LAYOUT_GENERAL
-			SRVTransitions.Remove(RHITexture);
-		}
-
-		for (FRHITexture* RHITexture : SRVTransitions)
-		{
-			TransitionTexture(RHITexture, true);
-		}
-	}
-
-	PreDispatch.Execute(CmdBuffer);
-	return PostDispatch;
 }
 
 void FVulkanCommandListContext::RHIRayTraceDispatch(
@@ -2127,7 +1982,7 @@ void FVulkanCommandListContext::RHIRayTraceDispatch(
 	checkSlow(InRayGenShader->GetFrequency() == SF_RayGen);
 	const VkStridedDeviceAddressRegionKHR RayGenRegion = ShaderTable->CommitRayGenShader(*this, Pipeline->GetShaderIndex(RayGenShader), Pipeline->GetShaderHandles(SF_RayGen));
 
-	FVulkanPipelineBarrier PostDispatch = SetRayGenResources(Device, CmdBuffer, InGlobalResourceBindings, ShaderTable);
+	SetRayGenResources(Device, CmdBuffer, InGlobalResourceBindings, ShaderTable);
 
 	VulkanRHI::vkCmdTraceRaysKHR(
 		CmdBuffer->GetHandle(),
@@ -2136,8 +1991,6 @@ void FVulkanCommandListContext::RHIRayTraceDispatch(
 		ShaderTable->GetRegion(SF_RayHitGroup),
 		ShaderTable->GetRegion(SF_RayCallable),
 		InWidth, InHeight, 1);
-
-	PostDispatch.Execute(CmdBuffer);
 }
 
 void FVulkanCommandListContext::RHIRayTraceDispatchIndirect(
@@ -2159,7 +2012,7 @@ void FVulkanCommandListContext::RHIRayTraceDispatchIndirect(
 	checkSlow(InRayGenShader->GetFrequency() == SF_RayGen);
 	const VkStridedDeviceAddressRegionKHR RayGenRegion = ShaderTable->CommitRayGenShader(*this, Pipeline->GetShaderIndex(RayGenShader), Pipeline->GetShaderHandles(SF_RayGen));
 
-	FVulkanPipelineBarrier PostDispatch = SetRayGenResources(Device, CmdBuffer, InGlobalResourceBindings, ShaderTable);
+	SetRayGenResources(Device, CmdBuffer, InGlobalResourceBindings, ShaderTable);
 
 	FVulkanBuffer* ArgumentBuffer = ResourceCast(InArgumentBuffer);
 	const VkDeviceAddress IndirectDeviceAddress = ArgumentBuffer->GetDeviceAddress() + InArgumentOffset;
@@ -2171,8 +2024,6 @@ void FVulkanCommandListContext::RHIRayTraceDispatchIndirect(
 		ShaderTable->GetRegion(SF_RayHitGroup),
 		ShaderTable->GetRegion(SF_RayCallable),
 		IndirectDeviceAddress);
-
-	PostDispatch.Execute(CmdBuffer);
 }
 
 
diff --git a/Engine/Source/Runtime/VulkanRHI/Private/VulkanRenderTarget.cpp b/Engine/Source/Runtime/VulkanRHI/Private/VulkanRenderTarget.cpp
--- a/Engine/Source/Runtime/VulkanRHI/Private/VulkanRenderTarget.cpp
+++ b/Engine/Source/Runtime/VulkanRHI/Private/VulkanRenderTarget.cpp
@@ -191,27 +191,15 @@ void FVulkanDynamicRHI::RHIReadSurfaceData(FRHITexture* TextureRHI, FIntRect Rec
 		CopyRegion.imageExtent.depth = 1;
 
 		RHICmdList.EnqueueLambda([CopyRegion, StagingBuffer, &Surface](FRHICommandListBase& ExecutingCmdList)
-			{
-				FVulkanCommandListContext& Context = FVulkanCommandListContext::Get(ExecutingCmdList);
-				FVulkanCommandBuffer& CommandBuffer = Context.GetCommandBuffer();
-
-				const VkImageLayout PreviousLayout = FVulkanLayoutManager::SetExpectedLayout(&CommandBuffer, Surface, ERHIAccess::CopySrc);
-
-				VulkanRHI::vkCmdCopyImageToBuffer(CommandBuffer.GetHandle(), Surface.Image, VK_IMAGE_LAYOUT_TRANSFER_SRC_OPTIMAL, StagingBuffer->GetHandle(), 1, &CopyRegion);
-
-				FVulkanPipelineBarrier AfterBarrier;
-				if ((PreviousLayout != VK_IMAGE_LAYOUT_UNDEFINED) && (PreviousLayout != VK_IMAGE_LAYOUT_TRANSFER_SRC_OPTIMAL))
-				{
-					AfterBarrier.AddImageLayoutTransition(Surface.Image, VK_IMAGE_LAYOUT_TRANSFER_SRC_OPTIMAL, PreviousLayout, FVulkanPipelineBarrier::MakeSubresourceRange(Surface.GetFullAspectMask()));
-				}
-				else
-				{
-					CommandBuffer.GetLayoutManager().SetFullLayout(Surface, VK_IMAGE_LAYOUT_TRANSFER_SRC_OPTIMAL);
-				}
-
-				AfterBarrier.AddMemoryBarrier(VK_ACCESS_TRANSFER_WRITE_BIT, VK_ACCESS_HOST_READ_BIT, VK_PIPELINE_STAGE_TRANSFER_BIT, VK_PIPELINE_STAGE_HOST_BIT);
-				AfterBarrier.Execute(&CommandBuffer);
-			});
+		{
+			FVulkanCommandListContext& Context = FVulkanCommandListContext::Get(ExecutingCmdList);
+			FVulkanCommandBuffer& CommandBuffer = Context.GetCommandBuffer();
+			VulkanRHI::vkCmdCopyImageToBuffer(CommandBuffer.GetHandle(), Surface.Image, VK_IMAGE_LAYOUT_TRANSFER_SRC_OPTIMAL, StagingBuffer->GetHandle(), 1, &CopyRegion);
+
+			FVulkanPipelineBarrier AfterBarrier;
+			AfterBarrier.AddMemoryBarrier(VK_ACCESS_TRANSFER_WRITE_BIT, VK_ACCESS_HOST_READ_BIT, VK_PIPELINE_STAGE_TRANSFER_BIT, VK_PIPELINE_STAGE_HOST_BIT);
+			AfterBarrier.Execute(&CommandBuffer);
+		});
 	}
 
 	// We need to execute the command list so we can read the data from the map below
@@ -317,27 +305,15 @@ void FVulkanDynamicRHI::RHIReadSurfaceFloatData(FRHITexture* TextureRHI, FIntRec
 
 		FRHICommandListImmediate& RHICmdList = FRHICommandListImmediate::Get();
 		RHICmdList.EnqueueLambda([CopyRegion, StagingBuffer, &VulkanTexture](FRHICommandListBase& ExecutingCmdList)
-			{
-				FVulkanCommandListContext& Context = FVulkanCommandListContext::Get(ExecutingCmdList);
-				FVulkanCommandBuffer& CommandBuffer = Context.GetCommandBuffer();
-				const VkImageLayout OriginalLayout = FVulkanLayoutManager::SetExpectedLayout(&CommandBuffer, VulkanTexture, ERHIAccess::CopySrc);
-
-				VulkanRHI::vkCmdCopyImageToBuffer(CommandBuffer.GetHandle(), VulkanTexture.Image, VK_IMAGE_LAYOUT_TRANSFER_SRC_OPTIMAL, StagingBuffer->GetHandle(), 1, &CopyRegion);
-
-				FVulkanPipelineBarrier AfterBarrier;
-				if ((OriginalLayout != VK_IMAGE_LAYOUT_UNDEFINED) && (OriginalLayout != VK_IMAGE_LAYOUT_TRANSFER_SRC_OPTIMAL))
-				{
-					const VkImageSubresourceRange ImageSubresource = FVulkanPipelineBarrier::MakeSubresourceRange(VulkanTexture.GetFullAspectMask());
-					AfterBarrier.AddImageLayoutTransition(VulkanTexture.Image, VK_IMAGE_LAYOUT_TRANSFER_SRC_OPTIMAL, OriginalLayout, ImageSubresource);
-				}
-				else
-				{
-					CommandBuffer.GetLayoutManager().SetFullLayout(VulkanTexture, VK_IMAGE_LAYOUT_TRANSFER_SRC_OPTIMAL);
-				}
-
-				AfterBarrier.AddMemoryBarrier(VK_ACCESS_TRANSFER_WRITE_BIT, VK_ACCESS_HOST_READ_BIT, VK_PIPELINE_STAGE_TRANSFER_BIT, VK_PIPELINE_STAGE_HOST_BIT);
-				AfterBarrier.Execute(&CommandBuffer);
-			});
+		{
+			FVulkanCommandListContext& Context = FVulkanCommandListContext::Get(ExecutingCmdList);
+			FVulkanCommandBuffer& CommandBuffer = Context.GetCommandBuffer();
+			VulkanRHI::vkCmdCopyImageToBuffer(CommandBuffer.GetHandle(), VulkanTexture.Image, VK_IMAGE_LAYOUT_TRANSFER_SRC_OPTIMAL, StagingBuffer->GetHandle(), 1, &CopyRegion);
+
+			FVulkanPipelineBarrier AfterBarrier;
+			AfterBarrier.AddMemoryBarrier(VK_ACCESS_TRANSFER_WRITE_BIT, VK_ACCESS_HOST_READ_BIT, VK_PIPELINE_STAGE_TRANSFER_BIT, VK_PIPELINE_STAGE_HOST_BIT);
+			AfterBarrier.Execute(&CommandBuffer);
+		});
 
 		// We need to execute the command list so we can read the data from the map below
 		RHICmdList.SubmitAndBlockUntilGPUIdle();
@@ -455,28 +431,15 @@ void FVulkanDynamicRHI::RHIRead3DSurfaceFloatData(FRHITexture* TextureRHI, FIntR
 
 	FRHICommandListImmediate& RHICmdList = FRHICommandListImmediate::Get();
 	RHICmdList.EnqueueLambda([CopyRegion, StagingBuffer, &Surface](FRHICommandListBase& ExecutingCmdList)
-		{
-			FVulkanCommandListContext& Context = FVulkanCommandListContext::Get(ExecutingCmdList);
-			FVulkanCommandBuffer& CommandBuffer = Context.GetCommandBuffer();
-
-			const VkImageLayout OriginalLayout = FVulkanLayoutManager::SetExpectedLayout(&CommandBuffer, Surface, ERHIAccess::CopySrc);
-
-			VulkanRHI::vkCmdCopyImageToBuffer(CommandBuffer.GetHandle(), Surface.Image, VK_IMAGE_LAYOUT_TRANSFER_SRC_OPTIMAL, StagingBuffer->GetHandle(), 1, &CopyRegion);
-
-			FVulkanPipelineBarrier AfterBarrier;
-			if ((OriginalLayout != VK_IMAGE_LAYOUT_UNDEFINED) && (OriginalLayout != VK_IMAGE_LAYOUT_TRANSFER_SRC_OPTIMAL))
-			{
-				AfterBarrier.AddImageLayoutTransition(Surface.Image, VK_IMAGE_LAYOUT_TRANSFER_SRC_OPTIMAL, OriginalLayout, FVulkanPipelineBarrier::MakeSubresourceRange(Surface.GetFullAspectMask()));
-			}
-			else
-			{
-				CommandBuffer.GetLayoutManager().SetFullLayout(Surface, VK_IMAGE_LAYOUT_TRANSFER_SRC_OPTIMAL);
-			}
+	{
+		FVulkanCommandListContext& Context = FVulkanCommandListContext::Get(ExecutingCmdList);
+		FVulkanCommandBuffer& CommandBuffer = Context.GetCommandBuffer();
+		VulkanRHI::vkCmdCopyImageToBuffer(CommandBuffer.GetHandle(), Surface.Image, VK_IMAGE_LAYOUT_TRANSFER_SRC_OPTIMAL, StagingBuffer->GetHandle(), 1, &CopyRegion);
 
-
-			AfterBarrier.AddMemoryBarrier(VK_ACCESS_TRANSFER_WRITE_BIT, VK_ACCESS_HOST_READ_BIT, VK_PIPELINE_STAGE_TRANSFER_BIT, VK_PIPELINE_STAGE_HOST_BIT);
-			AfterBarrier.Execute(&CommandBuffer);
-		});
+		FVulkanPipelineBarrier AfterBarrier;
+		AfterBarrier.AddMemoryBarrier(VK_ACCESS_TRANSFER_WRITE_BIT, VK_ACCESS_HOST_READ_BIT, VK_PIPELINE_STAGE_TRANSFER_BIT, VK_PIPELINE_STAGE_HOST_BIT);
+		AfterBarrier.Execute(&CommandBuffer);
+	});
 
 	// We need to execute the command list so we can read the data from the map below
 	RHICmdList.SubmitAndBlockUntilGPUIdle();
@@ -565,70 +528,36 @@ void FVulkanCommandListContext::RHIBeginRenderPass(const FRHIRenderPassInfo& InI
 		FVulkanTexture& VulkanTexture = *ResourceCast(DSTexture);
 		const VkImageAspectFlags AspectFlags = VulkanTexture.GetFullAspectMask();
 
-		if (Device.SupportsParallelRendering())
+		const FExclusiveDepthStencil ExclusiveDepthStencil = InInfo.DepthStencilRenderTarget.ExclusiveDepthStencil;
+		if (VKHasAnyFlags(AspectFlags, VK_IMAGE_ASPECT_DEPTH_BIT))
 		{
-			const FExclusiveDepthStencil ExclusiveDepthStencil = InInfo.DepthStencilRenderTarget.ExclusiveDepthStencil;
-			if (VKHasAnyFlags(AspectFlags, VK_IMAGE_ASPECT_DEPTH_BIT))
+			if (ExclusiveDepthStencil.IsDepthWrite())
 			{
-				if (ExclusiveDepthStencil.IsDepthWrite())
-				{
-					CurrentDepthLayout = VK_IMAGE_LAYOUT_ATTACHMENT_OPTIMAL;
-				}
-				else
-				{
-					// todo-jn: temporarily use tracking illegally because sometimes depth is left in ATTACHMENT_OPTIMAL for Read passes
-					CurrentDepthLayout = CmdBuffer->GetLayoutManager().GetDepthStencilHint(VulkanTexture, VK_IMAGE_ASPECT_DEPTH_BIT);
-					if (CurrentDepthLayout == VK_IMAGE_LAYOUT_UNDEFINED)
-					{
-						CurrentDepthLayout = Queue.GetLayoutManager().GetDepthStencilHint(VulkanTexture, VK_IMAGE_ASPECT_DEPTH_BIT);
-						if (CurrentDepthLayout == VK_IMAGE_LAYOUT_UNDEFINED)
-						{
-							if (ExclusiveDepthStencil.IsDepthRead())
-							{
-								CurrentDepthLayout = VK_IMAGE_LAYOUT_READ_ONLY_OPTIMAL;
-							}
-							else
-							{
-								UE_LOG(LogVulkanRHI, Warning, TEXT("Render pass [%s] used an unknown depth state!"), InName ? InName : TEXT("unknown"));
-							}
-						}
-					}
-				}
+				CurrentDepthLayout = VK_IMAGE_LAYOUT_ATTACHMENT_OPTIMAL;
 			}
-
-			if (VKHasAnyFlags(AspectFlags, VK_IMAGE_ASPECT_STENCIL_BIT))
+			else if (ExclusiveDepthStencil.IsDepthRead())
+			{
+				CurrentDepthLayout = VK_IMAGE_LAYOUT_READ_ONLY_OPTIMAL;
+			}
+			else
 			{
-				if (ExclusiveDepthStencil.IsStencilWrite())
-				{
-					CurrentStencilLayout = VK_IMAGE_LAYOUT_ATTACHMENT_OPTIMAL;
-				}
-				else if (ExclusiveDepthStencil.IsStencilRead())
-				{
-					CurrentStencilLayout = VK_IMAGE_LAYOUT_READ_ONLY_OPTIMAL;
-				}
-				else
-				{
-					// todo-jn: temporarily use tracking illegally because stencil is no-op so we can't know expected layout
-					CurrentStencilLayout = CmdBuffer->GetLayoutManager().GetDepthStencilHint(VulkanTexture, VK_IMAGE_ASPECT_STENCIL_BIT);
-					if (CurrentStencilLayout == VK_IMAGE_LAYOUT_UNDEFINED)
-					{
-						CurrentStencilLayout = Queue.GetLayoutManager().GetDepthStencilHint(VulkanTexture, VK_IMAGE_ASPECT_STENCIL_BIT);
-						// If the layout is still UNDEFINED, FVulkanRenderTargetLayout will force it into a known state
-					}
-				}
+				CurrentDepthLayout = FVulkanPipelineBarrier::GetDepthOrStencilLayout(InInfo.NopAccess);
 			}
 		}
-		else
+
+		if (VKHasAnyFlags(AspectFlags, VK_IMAGE_ASPECT_STENCIL_BIT))
 		{
-			const FVulkanImageLayout* FullLayout = CmdBuffer->GetLayoutManager().GetFullLayout(VulkanTexture);
-			check(FullLayout);
-			if (VKHasAnyFlags(AspectFlags, VK_IMAGE_ASPECT_DEPTH_BIT))
+			if (ExclusiveDepthStencil.IsStencilWrite())
 			{
-				CurrentDepthLayout = FullLayout->GetSubresLayout(0, 0, VK_IMAGE_ASPECT_DEPTH_BIT);
+				CurrentStencilLayout = VK_IMAGE_LAYOUT_ATTACHMENT_OPTIMAL;
 			}
-			if (VKHasAnyFlags(AspectFlags, VK_IMAGE_ASPECT_STENCIL_BIT))
+			else if (ExclusiveDepthStencil.IsStencilRead())
+			{
+				CurrentStencilLayout = VK_IMAGE_LAYOUT_READ_ONLY_OPTIMAL;
+			}
+			else
 			{
-				CurrentStencilLayout = FullLayout->GetSubresLayout(0, 0, VK_IMAGE_ASPECT_STENCIL_BIT);
+				CurrentStencilLayout = FVulkanPipelineBarrier::GetDepthOrStencilLayout(InInfo.NopAccess);
 			}
 		}
 	}
diff --git a/Engine/Source/Runtime/VulkanRHI/Private/VulkanResources.h b/Engine/Source/Runtime/VulkanRHI/Private/VulkanResources.h
--- a/Engine/Source/Runtime/VulkanRHI/Private/VulkanResources.h
+++ b/Engine/Source/Runtime/VulkanRHI/Private/VulkanResources.h
@@ -617,10 +617,10 @@ public:
 
 	// FVulkanEvictable interface.
 	bool CanMove() const override { return false; }
-	bool CanEvict() const override { return false; }
+	bool CanEvict() const override { return AreAllSubresourcesSame_Unsafe(); }
 	void Evict(FVulkanDevice& Device, FVulkanCommandListContext& Context) override; ///evict to system memory
 	void Move(FVulkanDevice& Device, FVulkanCommandListContext& Context, VulkanRHI::FVulkanAllocation& NewAllocation) override; //move to a full new allocation
-	FVulkanTexture* GetEvictableTexture() override { return this; }
+	FVulkanTexture* GetEvictableTexture() override { return CanEvict() ? this : nullptr; }
 
 	bool GetTextureResourceInfo(FRHIResourceInfo& OutResourceInfo) const;
 
@@ -759,7 +759,7 @@ public:
 
 private:
 	void SetInitialImageState(FRHICommandListBase& RHICmdList, VkImageLayout InitialLayout, bool bClear, const FClearValueBinding& ClearValueBinding, bool bIsTransientResource);
-	void InternalMoveSurface(FVulkanDevice& InDevice, FVulkanCommandListContext& Context, VulkanRHI::FVulkanAllocation& DestAllocation, VkImageLayout OriginalLayout);
+	void InternalMoveSurface(FVulkanDevice& InDevice, FVulkanCommandListContext& Context, VulkanRHI::FVulkanAllocation& DestAllocation);
 
 	VkImageTiling Tiling;
 	VulkanRHI::FVulkanAllocation Allocation;
diff --git a/Engine/Source/Runtime/VulkanRHI/Private/VulkanTexture.cpp b/Engine/Source/Runtime/VulkanRHI/Private/VulkanTexture.cpp
--- a/Engine/Source/Runtime/VulkanRHI/Private/VulkanTexture.cpp
+++ b/Engine/Source/Runtime/VulkanRHI/Private/VulkanTexture.cpp
@@ -120,6 +120,7 @@ void FVulkanTexture::InternalLockWrite(FVulkanContextCommon& Context, FVulkanTex
 
 	VulkanRHI::vkCmdCopyBufferToImage(StagingCommandBuffer, StagingBuffer->GetHandle(), Surface->Image, VK_IMAGE_LAYOUT_TRANSFER_DST_OPTIMAL, 1, &Region);
 
+	// :todo-jn: replace with cmdlist layout tracking (ideally would happen on UploadContext)
 	{
 		FVulkanPipelineBarrier Barrier;
 		Barrier.AddImageLayoutTransition(Surface->Image, VK_IMAGE_LAYOUT_TRANSFER_DST_OPTIMAL, Surface->GetDefaultLayout(), SubresourceRange);
@@ -505,11 +506,16 @@ static VkImageLayout ChooseVRSLayout()
 
 static VkImageLayout GetInitialLayoutFromRHIAccess(ERHIAccess RHIAccess, bool bIsDepthStencilTarget, bool bSupportReadOnlyOptimal)
 {
-	if (EnumHasAnyFlags(RHIAccess, ERHIAccess::RTV) || RHIAccess == ERHIAccess::Present)
+	if (EnumHasAnyFlags(RHIAccess, ERHIAccess::RTV))
 	{
 		return VK_IMAGE_LAYOUT_COLOR_ATTACHMENT_OPTIMAL;
 	}
 
+	if (RHIAccess == ERHIAccess::Present)
+	{
+		return VK_IMAGE_LAYOUT_PRESENT_SRC_KHR;
+	}
+
 	if (EnumHasAnyFlags(RHIAccess, ERHIAccess::DSVWrite))
 	{
 		return VK_IMAGE_LAYOUT_ATTACHMENT_OPTIMAL;
@@ -548,8 +554,9 @@ static VkImageLayout GetInitialLayoutFromRHIAccess(ERHIAccess RHIAccess, bool bI
 	return VK_IMAGE_LAYOUT_UNDEFINED;
 }
 
-void FVulkanTexture::InternalMoveSurface(FVulkanDevice& InDevice, FVulkanCommandListContext& Context, VulkanRHI::FVulkanAllocation& DestAllocation, VkImageLayout OriginalLayout)
+void FVulkanTexture::InternalMoveSurface(FVulkanDevice& InDevice, FVulkanCommandListContext& Context, VulkanRHI::FVulkanAllocation& DestAllocation)
 {
+
 	FImageCreateInfo ImageCreateInfo;
 	const FRHITextureDesc& Desc = GetDesc();
 	FVulkanTexture::GenerateImageCreateInfo(ImageCreateInfo, InDevice, Desc, &StorageFormat, &ViewFormat);
@@ -587,6 +594,9 @@ void FVulkanTexture::InternalMoveSurface(FVulkanDevice& InDevice, FVulkanCommand
 		const uint32 NumberOfArrayLevels = GetNumberOfArrayLevels();
 		const VkImageSubresourceRange FullSubresourceRange = FVulkanPipelineBarrier::MakeSubresourceRange(FullAspectMask);
 
+		const ERHIAccess OriginalAccess = GetTrackedAccess_Unsafe();
+		const VkImageLayout OriginalLayout = GetInitialLayoutFromRHIAccess(OriginalAccess, IsDepthOrStencilAspect(), SupportsSampling());
+
 		// Transition to copying layouts
 		{
 			FVulkanPipelineBarrier Barrier;
@@ -613,6 +623,7 @@ void FVulkanTexture::InternalMoveSurface(FVulkanDevice& InDevice, FVulkanCommand
 				Region.srcSubresource.mipLevel = i;
 				Region.dstSubresource.mipLevel = i;
 			}
+
 			VulkanRHI::vkCmdCopyImage(CommandBufferHandle,
 				Image, VK_IMAGE_LAYOUT_TRANSFER_SRC_OPTIMAL,
 				MovedImage, VK_IMAGE_LAYOUT_TRANSFER_DST_OPTIMAL,
@@ -640,9 +651,6 @@ void FVulkanTexture::InternalMoveSurface(FVulkanDevice& InDevice, FVulkanCommand
 	}
 
 	Image = MovedImage;
-
-	// Move is used for defrag, which uses layouts stored in the queue, update the layout stored there
-	Context.Queue.GetLayoutManager().SetFullLayout(*this, OriginalLayout);
 }
 
 void FVulkanTexture::DestroySurface()
@@ -838,12 +846,6 @@ void FVulkanTexture::SetInitialImageState(FRHICommandListBase& RHICmdList, VkIma
 				Barrier.AddFullImageLayoutTransition(*VulkanTexture, CurrentLayout, InitialLayout);
 				Barrier.Execute(&CommandBuffer);
 			}
-
-			if (ExecutingCmdList.GetPipeline() != ERHIPipeline::None)
-			{
-				FVulkanCommandBuffer& TrackingCommandBuffer = FVulkanCommandListContext::Get(ExecutingCmdList).GetCommandBuffer();
-				TrackingCommandBuffer.GetLayoutManager().SetFullLayout(*VulkanTexture, InitialLayout);
-			}
 		});
 }
 
@@ -1045,10 +1047,6 @@ static void DoAsyncReallocateTexture2D(FVulkanContextCommon& Context, FVulkanTex
 		Barrier.AddImageLayoutTransition(OldTexture->Image, VK_IMAGE_LAYOUT_TRANSFER_SRC_OPTIMAL, OldTexture->GetDefaultLayout(), SourceSubResourceRange);
 		Barrier.AddImageLayoutTransition(NewTexture->Image, VK_IMAGE_LAYOUT_TRANSFER_DST_OPTIMAL, NewTexture->GetDefaultLayout(), DestSubResourceRange);
 		Barrier.Execute(&CommandBuffer);
-
-		// Add tracking for the appropriate subresources (intentionally leave added mips in VK_IMAGE_LAYOUT_UNDEFINED)
-		// NOTE: Overwriting whatever is contained in the layout manager for the new texture (circumvents issue with stale layouts and pointer reuse)
-		CommandBuffer.GetLayoutManager().SetLayout(*NewTexture, DestSubResourceRange, NewTexture->GetDefaultLayout());
 	}
 
 	// request is now complete
@@ -1164,7 +1162,7 @@ void FVulkanDynamicRHI::RHIUnlockTexture(FRHICommandListImmediate& RHICmdList, c
 	RHICmdList.EnqueueLambda(TEXT("FVulkanTexture::InternalLockWrite"),
 		[Texture, Region, StagingBuffer](FRHICommandListBase& ExecutingCmdList)
 		{
-			FVulkanTexture::InternalLockWrite(FVulkanUploadContext::Get(ExecutingCmdList), Texture, Region, StagingBuffer);
+			FVulkanTexture::InternalLockWrite(FVulkanCommandListContext::Get(ExecutingCmdList), Texture, Region, StagingBuffer);
 		});
 }
 
@@ -1219,7 +1217,7 @@ void FVulkanDynamicRHI::InternalUpdateTexture2D(FRHICommandListBase& RHICmdList,
 	RHICmdList.EnqueueLambda(TEXT("FVulkanTexture::InternalLockWrite"),
 		[Texture, Region, StagingBuffer](FRHICommandListBase& ExecutingCmdList)
 		{
-			FVulkanTexture::InternalLockWrite(FVulkanUploadContext::Get(ExecutingCmdList), Texture, Region, StagingBuffer);
+			FVulkanTexture::InternalLockWrite(FVulkanCommandListContext::Get(ExecutingCmdList), Texture, Region, StagingBuffer);
 		});
 }
 
@@ -1313,7 +1311,7 @@ void FVulkanDynamicRHI::InternalUpdateTexture3D(FRHICommandListBase& RHICmdList,
 	RHICmdList.EnqueueLambda(TEXT("FVulkanTexture::InternalLockWrite"),
 		[Texture, Region, StagingBuffer](FRHICommandListBase& ExecutingCmdList)
 		{
-			FVulkanTexture::InternalLockWrite(FVulkanUploadContext::Get(ExecutingCmdList), Texture, Region, StagingBuffer);
+			FVulkanTexture::InternalLockWrite(FVulkanCommandListContext::Get(ExecutingCmdList), Texture, Region, StagingBuffer);
 		});
 }
 
@@ -2001,19 +1999,13 @@ void FVulkanTexture::Move(FVulkanDevice& InDevice, FVulkanCommandListContext& Co
 		UE_LOG(LogVulkanRHI, Display, TEXT("Moved %8.4fkb %8.4fkb   TB %p  :: IMG %p   %-40s\n"), Size / (1024.f), TotalSize / (1024.f), this, reinterpret_cast<const void*>(Image), *GetName().ToString());
 	}
 
-	// Move is used for defrag, which uses layouts stored in the queue
-	const FVulkanImageLayout* OriginalLayout = Context.Queue.GetLayoutManager().GetFullLayout(Image);
-	check(OriginalLayout && OriginalLayout->AreAllSubresourcesSameLayout());
-
 	const ETextureCreateFlags UEFlags = GetDesc().Flags;
 	const bool bRenderTarget = EnumHasAnyFlags(UEFlags, TexCreate_RenderTargetable | TexCreate_DepthStencilTargetable | TexCreate_ResolveTargetable);
 	const bool bUAV = EnumHasAnyFlags(UEFlags, TexCreate_UAV);
 	checkf(bRenderTarget || bUAV, TEXT("Surface must be a RenderTarget or a UAV in order to be moved.  UEFlags=0x%x"), (int32)UEFlags);
 	checkf(Tiling == VK_IMAGE_TILING_OPTIMAL, TEXT("Tiling [%s] is not supported for move, only VK_IMAGE_TILING_OPTIMAL"), VK_TYPE_TO_STRING(VkImageTiling, Tiling));
-	checkf((OriginalLayout->NumMips == GetNumMips()), TEXT("NumMips reported by LayoutManager (%d) differs from surface (%d)"), OriginalLayout->NumMips, GetNumMips());
-	checkf((OriginalLayout->NumLayers == GetNumberOfArrayLevels()), TEXT("NumLayers reported by LayoutManager (%d) differs from surface (%d)"), OriginalLayout->NumLayers, GetNumberOfArrayLevels());
 
-	InternalMoveSurface(InDevice, Context, NewAllocation, OriginalLayout->MainLayout);
+	InternalMoveSurface(InDevice, Context, NewAllocation);
 	
 	// Swap in the new allocation for this surface
 	Allocation.Swap(NewAllocation);
@@ -2032,9 +2024,6 @@ void FVulkanTexture::Evict(FVulkanDevice& InDevice, FVulkanCommandListContext& C
 		FGenericPlatformMisc::LowLevelOutputDebugStringf(TEXT("Evicted %8.4fkb %8.4fkb   TB %p  :: IMG %p   %-40s\n"), Size / (1024.f), TotalSize / (1024.f), this, Image, *GetName().ToString());
 	}
 
-	// Eviction layouts are read from the queue
-	const FVulkanImageLayout* OriginalLayout = Context.Queue.GetLayoutManager().GetFullLayout(Image);
-	if (OriginalLayout && OriginalLayout->AreAllSubresourcesSameLayout())
 	{
 		check(0 == CpuReadbackBuffer);
 		checkf(MemProps == VK_MEMORY_PROPERTY_DEVICE_LOCAL_BIT, TEXT("Can't evict surface that isn't device local.  MemoryProperties=%s"), VK_FLAGS_TO_STRING(VkMemoryPropertyFlags, MemProps));
@@ -2047,8 +2036,6 @@ void FVulkanTexture::Evict(FVulkanDevice& InDevice, FVulkanCommandListContext& C
 		//none of this is supported for eviction
 		checkf(!bRenderTarget, TEXT("RenderTargets do not support evict."));
 		checkf(!bUAV, TEXT("UAV do not support evict."));
-		checkf((OriginalLayout->NumMips == GetNumMips()), TEXT("NumMips reported by LayoutManager (%d) differs from surface (%d)"), OriginalLayout->NumMips, GetNumMips());
-		checkf((OriginalLayout->NumLayers == GetNumberOfArrayLevels()), TEXT("NumLayers reported by LayoutManager (%d) differs from surface (%d)"), OriginalLayout->NumLayers, GetNumberOfArrayLevels());
 
 		MemProps = InDevice.GetDeviceMemoryManager().GetEvictedMemoryProperties();
 
@@ -2061,7 +2048,7 @@ void FVulkanTexture::Evict(FVulkanDevice& InDevice, FVulkanCommandListContext& C
 			checkNoEntry();
 		}
 
-		InternalMoveSurface(InDevice, Context, HostAllocation, OriginalLayout->MainLayout);
+		InternalMoveSurface(InDevice, Context, HostAllocation);
 
 		// Delete the original allocation and swap in the new host allocation
 		Device->GetMemoryManager().FreeVulkanAllocation(Allocation);
@@ -2203,11 +2190,6 @@ void FVulkanCommandListContext::RHICopyTexture(FRHITexture* SourceTexture, FRHIT
 	FVulkanCommandBuffer& CommandBuffer = GetCommandBuffer();
 	check(CommandBuffer.IsOutsideRenderPass());
 
-	{
-		const VkImageLayout ExpectedSrcLayout = FVulkanLayoutManager::SetExpectedLayout(&CommandBuffer, *Source, ERHIAccess::CopySrc);
-		ensureMsgf((ExpectedSrcLayout == VK_IMAGE_LAYOUT_TRANSFER_SRC_OPTIMAL), TEXT("Expected source texture to be in VK_IMAGE_LAYOUT_TRANSFER_SRC_OPTIMAL, actual layout is %s"), VK_TYPE_TO_STRING(VkImageLayout, ExpectedSrcLayout));
-	}
-
 	const FPixelFormatInfo& PixelFormatInfo = GPixelFormats[DestTexture->GetDesc().Format];
 	const FRHITextureDesc& SourceDesc = SourceTexture->GetDesc();
 	const FRHITextureDesc& DestDesc = DestTexture->GetDesc();
@@ -2263,11 +2245,6 @@ void FVulkanCommandListContext::RHICopyTexture(FRHITexture* SourceTexture, FRHIT
 	}
 	else
 	{
-		{
-			const VkImageLayout ExpectedDstLayout = FVulkanLayoutManager::SetExpectedLayout(&CommandBuffer, *Dest, ERHIAccess::CopyDest);
-			ensureMsgf((ExpectedDstLayout == VK_IMAGE_LAYOUT_TRANSFER_DST_OPTIMAL), TEXT("Expected destination texture to be in VK_IMAGE_LAYOUT_TRANSFER_DST_OPTIMAL, actual layout is %s"), VK_TYPE_TO_STRING(VkImageLayout, ExpectedDstLayout));
-		}
-
 		VkImageCopy Region;
 		FMemory::Memzero(Region);
 		if (CopyInfo.Size == FIntVector::ZeroValue)
diff --git a/Engine/Source/Runtime/VulkanRHI/Private/VulkanViewport.cpp b/Engine/Source/Runtime/VulkanRHI/Private/VulkanViewport.cpp
--- a/Engine/Source/Runtime/VulkanRHI/Private/VulkanViewport.cpp
+++ b/Engine/Source/Runtime/VulkanRHI/Private/VulkanViewport.cpp
@@ -18,7 +18,7 @@
 #include "RHIStaticStates.h"
 
 FVulkanBackBuffer::FVulkanBackBuffer(FVulkanDevice& Device, FVulkanViewport* InViewport, EPixelFormat Format, uint32 SizeX, uint32 SizeY, ETextureCreateFlags UEFlags)
-	: FVulkanTexture(Device, FRHITextureCreateDesc::Create2D(TEXT("FVulkanBackBuffer"), SizeX, SizeY, Format).SetFlags(UEFlags).DetermineInititialState(), VK_NULL_HANDLE, {})
+	: FVulkanTexture(Device, FRHITextureCreateDesc::Create2D(TEXT("FVulkanBackBuffer"), SizeX, SizeY, Format).SetFlags(UEFlags).SetInitialState(ERHIAccess::Present), VK_NULL_HANDLE, {})
 	, Viewport(InViewport)
 {
 }
@@ -82,16 +82,10 @@ void FVulkanBackBuffer::AcquireBackBufferImage(FVulkanCommandListContext& Contex
 			DefaultView = &ImageView;
 			PartialView = &ImageView;
 
-			FVulkanCommandBuffer& CommandBuffer = Context.GetCommandBuffer();
-			check(!CommandBuffer.IsInsideRenderPass());
-
-			// right after acquiring image is in undefined state
-			FVulkanLayoutManager& LayoutMgr = CommandBuffer.GetLayoutManager();
-			const FVulkanImageLayout CustomLayout(VK_IMAGE_LAYOUT_UNDEFINED, 1, 1, VK_IMAGE_ASPECT_COLOR_BIT);
-			LayoutMgr.SetFullLayout(Image, CustomLayout);
-
 			// Wait for semaphore signal before writing to backbuffer image
 			Context.AddWaitSemaphore(VK_PIPELINE_STAGE_COLOR_ATTACHMENT_OUTPUT_BIT, Viewport->AcquiredSemaphore);
+
+			// :todo-jn: transition from unknown the first time
 		}
 		else
 		{
@@ -740,7 +734,6 @@ void FVulkanViewport::InitImages(FVulkanContextCommon& Context, TConstArrayView<
 	FMemory::Memzero(ClearColor);
 
 	const VkImageSubresourceRange Range = FVulkanPipelineBarrier::MakeSubresourceRange(VK_IMAGE_ASPECT_COLOR_BIT);
-	const FVulkanImageLayout InitialLayout(VK_IMAGE_LAYOUT_PRESENT_SRC_KHR, 1, 1, VK_IMAGE_ASPECT_COLOR_BIT);
 	for (int32 Index = 0; Index < Images.Num(); ++Index)
 	{
 		uint32 ImageSizeX = SizeX;
@@ -759,8 +752,7 @@ void FVulkanViewport::InitImages(FVulkanContextCommon& Context, TConstArrayView<
 		{
 			VulkanSetImageLayout(&CommandBuffer, Images[Index], VK_IMAGE_LAYOUT_UNDEFINED, VK_IMAGE_LAYOUT_TRANSFER_DST_OPTIMAL, Range);
 			VulkanRHI::vkCmdClearColorImage(CommandBuffer.GetHandle(), Images[Index], VK_IMAGE_LAYOUT_TRANSFER_DST_OPTIMAL, &ClearColor, 1, &Range);
-			VulkanSetImageLayout(&CommandBuffer, Images[Index], VK_IMAGE_LAYOUT_TRANSFER_DST_OPTIMAL, InitialLayout.MainLayout, Range);
-			CommandBuffer.GetLayoutManager().SetFullLayout(Images[Index], InitialLayout);
+			VulkanSetImageLayout(&CommandBuffer, Images[Index], VK_IMAGE_LAYOUT_TRANSFER_DST_OPTIMAL, VK_IMAGE_LAYOUT_PRESENT_SRC_KHR, Range);
 		}
 
 #if VULKAN_ENABLE_DRAW_MARKERS
@@ -895,27 +887,14 @@ inline static void CopyImageToBackBuffer(FVulkanCommandListContext& Context, FVu
 
 	FVulkanCommandBuffer* CmdBuffer = &Context.GetCommandBuffer();
 	check(CmdBuffer->IsOutsideRenderPass());
-	FVulkanLayoutManager& LayoutManager = CmdBuffer->GetLayoutManager();
-	const FVulkanImageLayout* SrcImageLayout = LayoutManager.GetFullLayout(SrcSurface);
-
-	VkImageLayout SrcSurfaceLayout = VK_IMAGE_LAYOUT_TRANSFER_SRC_OPTIMAL;
-	VkImageLayout DstSurfaceLayout = VK_IMAGE_LAYOUT_TRANSFER_DST_OPTIMAL;
 
-	if (bNeedsVulkanPreTransform)
-	{
-		SrcSurfaceLayout = VK_IMAGE_LAYOUT_SHADER_READ_ONLY_OPTIMAL;
-		DstSurfaceLayout = VK_IMAGE_LAYOUT_COLOR_ATTACHMENT_OPTIMAL;
-	}
+	const VkImageLayout SrcSurfaceLayout = bNeedsVulkanPreTransform ? VK_IMAGE_LAYOUT_SHADER_READ_ONLY_OPTIMAL : VK_IMAGE_LAYOUT_TRANSFER_SRC_OPTIMAL;
+	const VkImageLayout DstSurfaceLayout = bNeedsVulkanPreTransform ? VK_IMAGE_LAYOUT_COLOR_ATTACHMENT_OPTIMAL : VK_IMAGE_LAYOUT_TRANSFER_DST_OPTIMAL;
 
 	{
 		FVulkanPipelineBarrier Barrier;
-		const VkImageLayout PreviousSrcSurfaceLayout = SrcImageLayout ? SrcImageLayout->GetSubresLayout(0, 0, VK_IMAGE_ASPECT_COLOR_BIT) : VK_IMAGE_LAYOUT_UNDEFINED;
-		if (PreviousSrcSurfaceLayout != SrcSurfaceLayout)
-		{
-			Barrier.AddImageLayoutTransition(SrcSurface.Image, PreviousSrcSurfaceLayout, SrcSurfaceLayout, FVulkanPipelineBarrier::MakeSubresourceRange(VK_IMAGE_ASPECT_COLOR_BIT, 0, 1));
-			LayoutManager.SetFullLayout(SrcSurface, SrcSurfaceLayout);
-		}
-		Barrier.AddImageLayoutTransition(DstSurface.Image, VK_IMAGE_LAYOUT_PRESENT_SRC_KHR, DstSurfaceLayout, FVulkanPipelineBarrier::MakeSubresourceRange(VK_IMAGE_ASPECT_COLOR_BIT, 0, 1));
+		Barrier.AddImageLayoutTransition(SrcSurface.Image, VK_IMAGE_LAYOUT_PRESENT_SRC_KHR, SrcSurfaceLayout, FVulkanPipelineBarrier::MakeSubresourceRange(VK_IMAGE_ASPECT_COLOR_BIT, 0, 1));
+		Barrier.AddImageLayoutTransition(DstSurface.Image, VK_IMAGE_LAYOUT_UNDEFINED, DstSurfaceLayout, FVulkanPipelineBarrier::MakeSubresourceRange(VK_IMAGE_ASPECT_COLOR_BIT, 0, 1));
 		Barrier.Execute(CmdBuffer);
 	}
 
@@ -1011,7 +990,7 @@ inline static void CopyImageToBackBuffer(FVulkanCommandListContext& Context, FVu
 
 	{
 		FVulkanPipelineBarrier Barrier;
-		// We do not need to transition SrcSurface here, render passes that use it will do transition
+		Barrier.AddImageLayoutTransition(SrcSurface.Image, SrcSurfaceLayout, VK_IMAGE_LAYOUT_PRESENT_SRC_KHR, FVulkanPipelineBarrier::MakeSubresourceRange(VK_IMAGE_ASPECT_COLOR_BIT, 0, 1));
 		Barrier.AddImageLayoutTransition(DstSurface.Image, DstSurfaceLayout, VK_IMAGE_LAYOUT_PRESENT_SRC_KHR, FVulkanPipelineBarrier::MakeSubresourceRange(VK_IMAGE_ASPECT_COLOR_BIT, 0, 1));
 		Barrier.Execute(CmdBuffer);
 	}
@@ -1049,23 +1028,8 @@ bool FVulkanViewport::Present(FVulkanCommandListContext& Context, FVulkanQueue*
 			{
 				FVulkanCommandBuffer& CommandBuffer = Context.GetCommandBuffer();
 				check(CommandBuffer.IsOutsideRenderPass());
-
 				check(RHIBackBuffer != nullptr && RHIBackBuffer->Image == BackBufferImages[AcquiredImageIndex]->Image);
-
-				FVulkanLayoutManager& LayoutManager = CommandBuffer.GetLayoutManager();
-				const FVulkanImageLayout* TrackedLayout = LayoutManager.GetFullLayout(BackBufferImages[AcquiredImageIndex]->Image);
-
-				// One of the rare cases we'll let parallel rendering check the tracking (legacy path already checked as a fallback)
-				if (!TrackedLayout && Device->SupportsParallelRendering())
-				{
-					TrackedLayout = Context.Queue.GetLayoutManager().GetFullLayout(BackBufferImages[AcquiredImageIndex]->Image);
-				}
-
-				VkImageLayout LastLayout = TrackedLayout ? TrackedLayout->MainLayout : VK_IMAGE_LAYOUT_UNDEFINED;
-				VulkanSetImageLayout(&CommandBuffer, BackBufferImages[AcquiredImageIndex]->Image, LastLayout, VK_IMAGE_LAYOUT_PRESENT_SRC_KHR, FVulkanPipelineBarrier::MakeSubresourceRange(VK_IMAGE_ASPECT_COLOR_BIT));
-
-				const FVulkanImageLayout UndefinedLayout(VK_IMAGE_LAYOUT_UNDEFINED, 1, 1, VK_IMAGE_ASPECT_COLOR_BIT);
-				LayoutManager.SetFullLayout(BackBufferImages[AcquiredImageIndex]->Image, UndefinedLayout);
+				VulkanSetImageLayout(&CommandBuffer, BackBufferImages[AcquiredImageIndex]->Image, VK_IMAGE_LAYOUT_UNDEFINED, VK_IMAGE_LAYOUT_PRESENT_SRC_KHR, FVulkanPipelineBarrier::MakeSubresourceRange(VK_IMAGE_ASPECT_COLOR_BIT));
 			}
 			else
 			{
