VulkanRHI: Add submission and interrupt threads (disabled by default).  Add a timeline semaphore for the CPU to signal when the interrupt thread is waiting on the queues.  Remove all other context and device functions to submit to queues and replace it with FlushCommands that offers to option to wait until submitted or completed.

#rnx

diff --git a/Engine/Source/Runtime/VulkanRHI/Private/VulkanCommands.cpp b/Engine/Source/Runtime/VulkanRHI/Private/VulkanCommands.cpp
--- a/Engine/Source/Runtime/VulkanRHI/Private/VulkanCommands.cpp
+++ b/Engine/Source/Runtime/VulkanRHI/Private/VulkanCommands.cpp
@@ -855,11 +855,6 @@ bool FVulkanDynamicRHI::RHIIsRenderingSuspended()
 	return false;
 }
 
-void FVulkanDynamicRHI::RHIBlockUntilGPUIdle()
-{
-	Device->WaitUntilIdle();
-}
-
 uint32 FVulkanDynamicRHI::RHIGetGPUFrameCycles(uint32 GPUIndex)
 {
 	check(GPUIndex == 0);
diff --git a/Engine/Source/Runtime/VulkanRHI/Private/VulkanContext.cpp b/Engine/Source/Runtime/VulkanRHI/Private/VulkanContext.cpp
--- a/Engine/Source/Runtime/VulkanRHI/Private/VulkanContext.cpp
+++ b/Engine/Source/Runtime/VulkanRHI/Private/VulkanContext.cpp
@@ -108,15 +108,38 @@ FVulkanPayload* FVulkanContextCommon::FlushResetQueryPools()
 	return nullptr;
 }
 
-// :todo-jn: deprecated, on its way to being removed.  Should not be called, bypasses potential UploadContext which could mess up ordering.
-void FVulkanContextCommon::SubmitActiveCmdBuffer(FVulkanDynamicRHI* DynamicRHI)
+void FVulkanContextCommon::FlushCommands(EVulkanFlushFlags FlushFlags)
 {
+	FVulkanSyncPointRef SyncPoint;
+	if (EnumHasAnyFlags(FlushFlags, EVulkanFlushFlags::WaitForCompletion))
+	{
+		SyncPoint = GetContextSyncPoint();
+	}
+
+	FGraphEventRef SubmissionEvent;
+	if (EnumHasAnyFlags(FlushFlags, EVulkanFlushFlags::WaitForSubmission))
+	{
+		SubmissionEvent = FGraphEvent::CreateGraphEvent();
+		GetPayload(EPhase::Signal).SubmissionEvent = SubmissionEvent;
+	}
+
 	FVulkanPlatformCommandList* FinalizedPayloads = new FVulkanPlatformCommandList;
 	Finalize(*FinalizedPayloads);
 
 	FDynamicRHI::FRHISubmitCommandListsArgs Args;
 	Args.CommandLists.Add(FinalizedPayloads);
-	DynamicRHI->RHISubmitCommandLists(MoveTemp(Args));
+	FVulkanDynamicRHI::Get().RHISubmitCommandLists(MoveTemp(Args));
+
+	if (SyncPoint)
+	{
+		SyncPoint->Wait();
+	}
+
+	if (SubmissionEvent && !SubmissionEvent->IsComplete())
+	{
+		SCOPED_NAMED_EVENT_TEXT("Submission_Wait", FColor::Turquoise);
+		SubmissionEvent->Wait();
+	}
 }
 
 #if VULKAN_DELETE_STALE_CMDBUFFERS
@@ -212,10 +235,6 @@ FVulkanCommandListContext::FVulkanCommandListContext(FVulkanDynamicRHI* InRHI, F
 	{
 		// Insert the Begin frame timestamp query. On EndDrawingViewport() we'll insert the End and immediately after a new Begin()
 		WriteBeginTimestamp(GetActiveCmdBuffer());
-
-		// Flush the cmd buffer immediately to ensure a valid
-		// 'Last submitted' cmd buffer exists at frame 0.
-		SubmitActiveCmdBuffer();
 	}
 
 	// Create Pending state, contains pipeline states such as current shader and etc..
diff --git a/Engine/Source/Runtime/VulkanRHI/Private/VulkanContext.h b/Engine/Source/Runtime/VulkanRHI/Private/VulkanContext.h
--- a/Engine/Source/Runtime/VulkanRHI/Private/VulkanContext.h
+++ b/Engine/Source/Runtime/VulkanRHI/Private/VulkanContext.h
@@ -20,6 +20,21 @@ class FVulkanQueue;
 class FVulkanSwapChain;
 
 
+enum class EVulkanFlushFlags
+{
+	None = 0,
+
+	// Block the calling thread until the submission thread has dispatched all work.
+	WaitForSubmission = 1,
+
+	// Both the calling thread until the GPU has signaled completion of all dispatched work.
+	WaitForCompletion = 2
+};
+ENUM_CLASS_FLAGS(EVulkanFlushFlags)
+
+
+
+
 class FVulkanContextCommon
 {
 	// Phases are used to track where we are at in filling the current payload.
@@ -118,7 +133,8 @@ public:
 
 	void FreeUnusedCmdBuffers(bool bTrimMemory);
 
-	void SubmitActiveCmdBuffer(FVulkanDynamicRHI* DynamicRHI);
+	// Should only ne used when we are certain there are no other pending contexts (like UploadContext)
+	void FlushCommands(EVulkanFlushFlags FlushFlags = EVulkanFlushFlags::None);
 
 	uint64 GetFrameCounter() const
 	{
@@ -336,11 +352,6 @@ public:
 
 	void ReleasePendingState();
 
-	void SubmitActiveCmdBuffer()
-	{
-		FVulkanContextCommon::SubmitActiveCmdBuffer(RHI);
-	}
-
 protected:
 	FVulkanDynamicRHI* RHI;
 	FVulkanCommandListContext* Immediate;
diff --git a/Engine/Source/Runtime/VulkanRHI/Private/VulkanDevice.cpp b/Engine/Source/Runtime/VulkanRHI/Private/VulkanDevice.cpp
--- a/Engine/Source/Runtime/VulkanRHI/Private/VulkanDevice.cpp
+++ b/Engine/Source/Runtime/VulkanRHI/Private/VulkanDevice.cpp
@@ -1541,6 +1541,8 @@ void FVulkanDevice::Destroy()
 
 void FVulkanDevice::WaitUntilIdle()
 {
+	RHI->RHIBlockUntilGPUIdle();
+
 	VERIFYVULKANRESULT(VulkanRHI::vkDeviceWaitIdle(Device));
 }
 
@@ -1561,23 +1563,6 @@ void FVulkanDevice::NotifyDeletedImage(VkImage Image, bool bRenderTarget)
 	}
 }
 
-void FVulkanDevice::SubmitCommands(FVulkanCommandListContext* Context)
-{
-	Context->SubmitActiveCmdBuffer();
-}
-
-void FVulkanDevice::SubmitCommandsAndFlushGPU()
-{
-	if (ComputeContext != ImmediateContext)
-	{
-		SubmitCommands(ComputeContext);
-	}
-
-	SubmitCommands(ImmediateContext);
-
-	//#todo-rco: Process other contexts first!
-}
-
 void FVulkanDevice::NotifyDeletedGfxPipeline(class FVulkanRHIGraphicsPipelineState* Pipeline)
 {
 	if (ComputeContext != ImmediateContext)
diff --git a/Engine/Source/Runtime/VulkanRHI/Private/VulkanDevice.h b/Engine/Source/Runtime/VulkanRHI/Private/VulkanDevice.h
--- a/Engine/Source/Runtime/VulkanRHI/Private/VulkanDevice.h
+++ b/Engine/Source/Runtime/VulkanRHI/Private/VulkanDevice.h
@@ -472,8 +472,6 @@ public:
 	}
 #endif
 
-	void SubmitCommandsAndFlushGPU();
-
 	FVulkanOcclusionQueryPool* AcquireOcclusionQueryPool(FVulkanContextCommon* InCommonContext, uint32 NumQueries);
 	void ReleaseUnusedOcclusionQueryPools();
 
@@ -540,9 +538,6 @@ private:
 	void MapFormatSupport(EPixelFormat UEFormat, std::initializer_list<VkFormat> PrioritizedFormats, const VkComponentMapping& ComponentMapping, int32 BlockBytes);
 	void MapFormatSupport(EPixelFormat UEFormat, std::initializer_list<VkFormat> PrioritizedFormats, const VkComponentMapping& ComponentMapping, EPixelFormatCapabilities RequiredCapabilities);
 
-	void SubmitCommands(FVulkanCommandListContext* Context);
-
-
 	VkDevice Device;
 
 	VulkanRHI::FDeviceMemoryManager DeviceMemoryManager;
diff --git a/Engine/Source/Runtime/VulkanRHI/Private/VulkanRHI.cpp b/Engine/Source/Runtime/VulkanRHI/Private/VulkanRHI.cpp
--- a/Engine/Source/Runtime/VulkanRHI/Private/VulkanRHI.cpp
+++ b/Engine/Source/Runtime/VulkanRHI/Private/VulkanRHI.cpp
@@ -2008,7 +2008,9 @@ void FVulkanDynamicRHI::DestroySwapChain()
 		[Viewports](FRHICommandListImmediate& RHICmdList)
 	{
 		UE_LOG(LogVulkanRHI, Log, TEXT("Destroy swapchain ... "));
-		
+
+		RHICmdList.ImmediateFlush(EImmediateFlushType::FlushRHIThread);
+
 		for (FVulkanViewport* Viewport : Viewports)
 		{
 			Viewport->DestroySwapchain(nullptr);
@@ -2038,10 +2040,7 @@ void FVulkanDynamicRHI::RecreateSwapChain(void* NewNativeWindow)
 			
 			for (FVulkanViewport* Viewport : Viewports)
 			{
-				RHICmdList.EnqueueLambda(TEXT("RecreateSwapchain"), [Viewport, NewNativeWindow](FRHICommandListImmediate& ExecutingCmdList)
-				{
-					Viewport->RecreateSwapchain(FVulkanCommandListContext::Get(ExecutingCmdList), NewNativeWindow);
-				});
+				Viewport->RecreateSwapchainFromRT(RHICmdList, Viewport->PixelFormat, NewNativeWindow);
 			}
 
 			RHICmdList.SubmitAndBlockUntilGPUIdle();
diff --git a/Engine/Source/Runtime/VulkanRHI/Private/VulkanSubmission.cpp b/Engine/Source/Runtime/VulkanRHI/Private/VulkanSubmission.cpp
--- a/Engine/Source/Runtime/VulkanRHI/Private/VulkanSubmission.cpp
+++ b/Engine/Source/Runtime/VulkanRHI/Private/VulkanSubmission.cpp
@@ -96,9 +96,40 @@ protected:
 
 
 
+
+static TAutoConsoleVariable<int32> CVarVulkanUseInterruptThread(
+	TEXT("r.Vulkan.Submission.UseInterruptThread"),
+	0,
+	TEXT("  0: Process completed GPU work directly on the RHI thread. (default)\n")
+	TEXT("  1: Create a dedicated thread to process completed GPU work. \n"),
+	ECVF_ReadOnly);
+
+static TAutoConsoleVariable<int32> CVarVulkanUseSubmissionThread(
+	TEXT("r.Vulkan.Submission.UseSubmissionThread"),
+	0,
+	TEXT("  0: Submit GPU work directly from the RHI thread. (default)\n")
+	TEXT("  1: Create a dedicated thread to submit GPU work. \n"),
+	ECVF_ReadOnly);
+
+extern int32 GAllowTimelineSemaphores;
+
 void FVulkanDynamicRHI::InitializeSubmissionPipe()
 {
-	// :todo-jn: SubmissionThread and InterruptThread
+	const bool UseMultiThreading = FPlatformProcess::SupportsMultithreading() && FApp::ShouldUseThreadingForPerformance() && (FTaskGraphInterface::Get().GetNumWorkerThreads() > 6);
+	const bool UseTimelineSemaphores = GAllowTimelineSemaphores && Device->GetOptionalExtensions().HasKHRTimelineSemaphore;
+	if (UseTimelineSemaphores && UseMultiThreading)
+	{
+		if (CVarVulkanUseInterruptThread.GetValueOnAnyThread())
+		{
+			CPUTimelineSemaphore = new VulkanRHI::FSemaphore(*Device, VulkanRHI::EVulkanSemaphoreFlags::Timeline | VulkanRHI::EVulkanSemaphoreFlags::ImmediateDeletion);
+			InterruptThread = new FVulkanThread(TEXT("RHIInterruptThread"), TPri_Highest, this, &FVulkanDynamicRHI::WaitAndProcessInterruptQueue);
+		}
+
+		if (CVarVulkanUseSubmissionThread.GetValueOnAnyThread())
+		{
+			SubmissionThread = new FVulkanThread(TEXT("RHISubmissionThread"), TPri_Highest, this, &FVulkanDynamicRHI::ProcessSubmissionQueue);
+		}
+	}
 }
 
 void FVulkanDynamicRHI::ShutdownSubmissionPipe()
@@ -120,8 +151,40 @@ void FVulkanDynamicRHI::ShutdownSubmissionPipe()
 		ProcessInterruptQueueUntil(EopTask);
 		EopTask = nullptr;
 	}
+
+	if (CPUTimelineSemaphore)
+	{
+		delete CPUTimelineSemaphore;
+		CPUTimelineSemaphore = nullptr;
+	}
 }
 
+void FVulkanDynamicRHI::KickInterruptThread()
+{
+	if (InterruptThread)
+	{
+		checkSlow(CPUTimelineSemaphore);
+		VkSemaphoreSignalInfo SignalInfo;
+		ZeroVulkanStruct(SignalInfo, VK_STRUCTURE_TYPE_SEMAPHORE_SIGNAL_INFO);
+		SignalInfo.semaphore = CPUTimelineSemaphore->GetHandle();
+		SignalInfo.value = (uint64_t)CPUTimelineSemaphoreValue++;
+
+		InterruptThread->Kick();
+	}
+}
+
+void FVulkanDynamicRHI::KickSubmissionThread()
+{
+	if (SubmissionThread)
+	{
+		SubmissionThread->Kick();
+	}
+	else
+	{
+		FScopeLock Lock(&SubmissionCS);
+		while (ProcessSubmissionQueue()) {};
+	}
+}
 
 void FVulkanDynamicRHI::RHIFinalizeContext(FRHIFinalizeContextArgs&& Args, TRHIPipelineArray<IRHIPlatformCommandList*>& Output)
 {
@@ -164,6 +227,8 @@ void FVulkanDynamicRHI::RHIFinalizeContext(FRHIFinalizeContextArgs&& Args, TRHIP
 
 bool FVulkanDynamicRHI::ProcessSubmissionQueue()
 {
+	SCOPED_NAMED_EVENT_TEXT("ProcessSubmissionQueue", FColor::Orange);
+
 	// Sort all the payloads into the queues
 	FVulkanPlatformCommandList* PlatformCmdList = nullptr;
 	while (PendingPayloadsForSubmission.Dequeue(PlatformCmdList))
@@ -185,15 +250,17 @@ bool FVulkanDynamicRHI::ProcessSubmissionQueue()
 	// Wake up the interrupt thread to go wait on these new payloads
 	if (InterruptThread)
 	{
-		InterruptThread->Kick();
+		KickInterruptThread();
 	}
 
 	return bProgress;
 }
 
-bool FVulkanDynamicRHI::ProcessInterruptQueue(bool bWaitForProgress)
+bool FVulkanDynamicRHI::WaitAndProcessInterruptQueue()
 {
-	if (bWaitForProgress)
+	bool bProgress = false;
+
+	// Wait for progress
 	{
 		const int32 NumQueues = (int32)EVulkanQueueType::Count;
 
@@ -201,8 +268,12 @@ bool FVulkanDynamicRHI::ProcessInterruptQueue(bool bWaitForProgress)
 		// NOTE: we can hold on to these because we're holding InterruptCS or we're the interrupt thread.
 		FVulkanPayload* NextPayloads[NumQueues];
 		uint32 NumSyncs = 0;
-		Device->ForEachQueue([&NumSyncs, &NextPayloads](FVulkanQueue& Queue)
+		Device->ForEachQueue([&NumSyncs, &NextPayloads, &bProgress](FVulkanQueue& Queue)
 		{
+			// Clear any existing completed payloads
+			bProgress |= (Queue.ProcessInterruptQueue(0) > 0);
+
+			// Get the next payload to wait on
 			FVulkanPayload* Payload = Queue.GetNextInterruptPayload();
 			if (Payload)
 			{
@@ -237,8 +308,9 @@ bool FVulkanDynamicRHI::ProcessInterruptQueue(bool bWaitForProgress)
 			}
 			else
 			{
-				VkSemaphore Semaphores[NumQueues];
-				uint64_t Values[NumQueues];
+				const int32 NumSemas = NumQueues + 1;  // add a timeline sema for CPU wake
+				VkSemaphore Semaphores[NumSemas];
+				uint64_t Values[NumSemas];
 				for (uint32 Index = 0; Index < NumSyncs; ++Index)
 				{
 					checkf(NextPayloads[Index]->TimelineSemaphoreValue > 0, TEXT("Payloads should all use the same types of syncs!"));
@@ -246,6 +318,13 @@ bool FVulkanDynamicRHI::ProcessInterruptQueue(bool bWaitForProgress)
 					Values[Index] = (uint64_t)NextPayloads[Index]->TimelineSemaphoreValue;
 				}
 
+				if (CPUTimelineSemaphore)
+				{
+					Semaphores[NumSyncs] = CPUTimelineSemaphore->GetHandle();
+					Values[NumSyncs] = (uint64_t)CPUTimelineSemaphoreValue++;
+					NumSyncs++;
+				}
+
 				VkSemaphoreWaitInfo WaitInfo;
 				ZeroVulkanStruct(WaitInfo, VK_STRUCTURE_TYPE_SEMAPHORE_WAIT_INFO);
 				WaitInfo.flags = VK_SEMAPHORE_WAIT_ANY_BIT;
@@ -254,9 +333,20 @@ bool FVulkanDynamicRHI::ProcessInterruptQueue(bool bWaitForProgress)
 				WaitInfo.pValues = Values;
 				VERIFYVULKANRESULT_EXPANDED(VulkanRHI::vkWaitSemaphoresKHR(Device->GetInstanceHandle(), &WaitInfo, (uint64_t)Timeout));
 			}
+
+			// Some work completed and we made progress, or we stopped waiting because of CPUTimelineSemaphore
+			// Either way, we know there is work to wait on so make sure we loop back in here
+			bProgress = true;
 		}
 	}
 
+	return bProgress || ProcessInterruptQueue();
+}
+
+bool FVulkanDynamicRHI::ProcessInterruptQueue()
+{
+	SCOPED_NAMED_EVENT_TEXT("ProcessInterruptQueue", FColor::Magenta);
+
 	bool bProgress = false;
 	Device->ForEachQueue([&bProgress](FVulkanQueue& Queue)
 	{
@@ -269,7 +359,7 @@ void FVulkanDynamicRHI::ProcessInterruptQueueUntil(FGraphEvent* GraphEvent)
 {
 	if (InterruptThread)
 	{
-		InterruptThread->Kick();
+		KickInterruptThread();
 		if (GraphEvent && !GraphEvent->IsComplete())
 		{
 			GraphEvent->Wait();
@@ -285,7 +375,7 @@ void FVulkanDynamicRHI::ProcessInterruptQueueUntil(FGraphEvent* GraphEvent)
 		{
 			if (InterruptCS.TryLock())
 			{
-				ProcessInterruptQueue(true);
+				WaitAndProcessInterruptQueue();
 				InterruptCS.Unlock();
 			}
 			else
@@ -304,7 +394,7 @@ void FVulkanDynamicRHI::ProcessInterruptQueueUntil(FGraphEvent* GraphEvent)
 			bProgress = false;
 			if (InterruptCS.TryLock())
 			{
-				bProgress = ProcessInterruptQueue(false);
+				bProgress = ProcessInterruptQueue();
 				InterruptCS.Unlock();
 			}
 
@@ -320,15 +410,7 @@ void FVulkanDynamicRHI::RHISubmitCommandLists(FRHISubmitCommandListsArgs&& Args)
 		PendingPayloadsForSubmission.Enqueue(PlatformCmdList);
 	}
 
-	if (SubmissionThread)
-	{
-		SubmissionThread->Kick();
-	}
-	else
-	{
-		FScopeLock Lock(&SubmissionCS);
-		while (ProcessSubmissionQueue()) {};
-	}
+	KickSubmissionThread();
 
 	ProcessInterruptQueueUntil(nullptr);
 }
@@ -360,9 +442,42 @@ void FVulkanDynamicRHI::EnqueueEndOfPipeTask(TUniqueFunction<void()> TaskFunc, T
 
 	PendingPayloadsForSubmission.Enqueue(Payloads);
 
+	KickSubmissionThread();
+
 	EopTask = FFunctionGraphTask::CreateAndDispatchWhenReady(
 		MoveTemp(TaskFunc),
 		QUICK_USE_CYCLE_STAT(FExecuteRHIThreadTask, STATGROUP_TaskGraphTasks),
 		&Prereqs
 	);
+}
+
+
+void FVulkanDynamicRHI::RHIBlockUntilGPUIdle()
+{
+	TArray<FVulkanSyncPoint*, TInlineAllocator<(int32)EVulkanQueueType::Count>> EndOfQueueEvent;
+
+	// Create payloads with a signal for each queue
+	{
+		FVulkanPlatformCommandList* Payloads = new FVulkanPlatformCommandList;
+		Payloads->Reserve((int32)EVulkanQueueType::Count);
+
+		Device->ForEachQueue([&EndOfQueueEvent, &Payloads](FVulkanQueue& Queue)
+		{
+			FVulkanPayload* Payload = new FVulkanPayload(Queue);
+			Payload->SyncPoint = CreateVulkanSyncPoint();
+			Payloads->Add(Payload);
+
+			EndOfQueueEvent.Add(Payload->SyncPoint);
+		});
+
+		PendingPayloadsForSubmission.Enqueue(Payloads);
+	}
+
+	KickSubmissionThread();
+
+	// Wait on each event
+	for (FVulkanSyncPoint* Event : EndOfQueueEvent)
+	{
+		ProcessInterruptQueueUntil(Event);
+	}
 }
\ No newline at end of file
diff --git a/Engine/Source/Runtime/VulkanRHI/Private/VulkanUtil.cpp b/Engine/Source/Runtime/VulkanRHI/Private/VulkanUtil.cpp
--- a/Engine/Source/Runtime/VulkanRHI/Private/VulkanUtil.cpp
+++ b/Engine/Source/Runtime/VulkanRHI/Private/VulkanUtil.cpp
@@ -543,7 +543,7 @@ void FVulkanGPUProfiler::EndFrame()
 	{
 		if (bTrackingEvents)
 		{
-			CmdContext->Device.SubmitCommandsAndFlushGPU();
+			CmdContext->FlushCommands(EVulkanFlushFlags::None);
 
 			SetEmitDrawEvents(bOriginalGEmitDrawEvents);
 			UE_LOG(LogRHI, Warning, TEXT(""));
diff --git a/Engine/Source/Runtime/VulkanRHI/Private/VulkanViewport.cpp b/Engine/Source/Runtime/VulkanRHI/Private/VulkanViewport.cpp
--- a/Engine/Source/Runtime/VulkanRHI/Private/VulkanViewport.cpp
+++ b/Engine/Source/Runtime/VulkanRHI/Private/VulkanViewport.cpp
@@ -203,11 +203,10 @@ bool FVulkanViewport::DoCheckedSwapChainJob(FVulkanCommandListContext& Context,
 			check(0);
 		}
 
-		RecreateSwapchain(Context, WindowHandle);
+		RecreateSwapchain(Context);
 
 		// Swapchain creation pushes some commands - flush the command buffers now to begin with a fresh state
-		Device->SubmitCommandsAndFlushGPU();
-		Device->WaitUntilIdle();
+		Context.FlushCommands(EVulkanFlushFlags::WaitForCompletion);
 
 		Status = SwapChainJob(this);
 
@@ -292,7 +291,7 @@ void FVulkanViewport::IssueFrameEvent()
 	{
 		FVulkanCommandListContextImmediate& ImmediateContext = Device->GetImmediateContext();
 		LastFrameSyncPoint = ImmediateContext.GetContextSyncPoint();
-		Device->SubmitCommandsAndFlushGPU();
+		ImmediateContext.FlushCommands();
 	}
 }
 
@@ -672,15 +671,17 @@ bool FVulkanFramebuffer::Matches(const FRHISetRenderTargetsInfo& InRTInfo) const
 }
 
 // Tear down and recreate swapchain and related resources.
-void FVulkanViewport::RecreateSwapchain(FVulkanCommandListContext& Context, void* NewNativeWindow)
+void FVulkanViewport::RecreateSwapchain(FVulkanCommandListContext& Context)
 {
+	// Make sure everything is submitted and submission queue is idle
+	Context.FlushCommands(EVulkanFlushFlags::WaitForCompletion);
+
 	FScopeLock LockSwapchain(&RecreatingSwapchain);
 
 	TRHICommandList_RecursiveHazardous<FVulkanCommandListContext> RHICmdList(&Context);
 
 	FVulkanSwapChainRecreateInfo RecreateInfo = { VK_NULL_HANDLE, VK_NULL_HANDLE };
 	DestroySwapchain(&RecreateInfo);
-	WindowHandle = NewNativeWindow;
 	CreateSwapchain(RHICmdList, &RecreateInfo);
 	check(RecreateInfo.Surface == VK_NULL_HANDLE);
 	check(RecreateInfo.SwapChain == VK_NULL_HANDLE);
@@ -696,7 +697,7 @@ void FVulkanViewport::Tick(float DeltaTime)
 		ENQUEUE_RENDER_COMMAND(UpdateVsync)(
 			[this](FRHICommandListImmediate& RHICmdList)
 		{
-			RecreateSwapchainFromRT(RHICmdList, PixelFormat);
+			RecreateSwapchainFromRT(RHICmdList);
 		});
 		FlushRenderingCommands();
 	}
@@ -707,20 +708,22 @@ void FVulkanViewport::Resize(uint32 InSizeX, uint32 InSizeY, bool bInIsFullscree
 	SizeX = InSizeX;
 	SizeY = InSizeY;
 	bIsFullscreen = bInIsFullscreen;
+	PixelFormat = PreferredPixelFormat;
 
-	RecreateSwapchainFromRT(FRHICommandListImmediate::Get(), PreferredPixelFormat);
+	RecreateSwapchainFromRT(FRHICommandListImmediate::Get());
 }
 
-void FVulkanViewport::RecreateSwapchainFromRT(FRHICommandListImmediate& RHICmdList, EPixelFormat PreferredPixelFormat)
+void FVulkanViewport::RecreateSwapchainFromRT(FRHICommandListImmediate& RHICmdList, EPixelFormat PreferredPixelFormat, void* NewNativeWindow)
 {
 	check(IsInRenderingThread());
 
 	// TODO: should flush RHIT commands here?
 	RHICmdList.ImmediateFlush(EImmediateFlushType::FlushRHIThread);
-	
+
 	FVulkanSwapChainRecreateInfo RecreateInfo = { VK_NULL_HANDLE, VK_NULL_HANDLE};
 	DestroySwapchain(&RecreateInfo);
 	PixelFormat = PreferredPixelFormat;
+	WindowHandle = NewNativeWindow;
 	CreateSwapchain(RHICmdList, &RecreateInfo);
 	check(RecreateInfo.Surface == VK_NULL_HANDLE);
 	check(RecreateInfo.SwapChain == VK_NULL_HANDLE);
@@ -846,10 +849,6 @@ void FVulkanViewport::CreateSwapchain(FRHICommandListBase& RHICmdList, FVulkanSw
 
 void FVulkanViewport::DestroySwapchain(FVulkanSwapChainRecreateInfo* RecreateInfo)
 {
-	// Submit all command buffers here
-	Device->SubmitCommandsAndFlushGPU();
-	Device->WaitUntilIdle();
-	
 	// Intentionally leave RenderingBackBuffer alive, so it can be used a dummy backbuffer while we don't have swapchain images
 	// RenderingBackBuffer = nullptr;
 	
@@ -1001,11 +1000,12 @@ bool FVulkanViewport::Present(FVulkanCommandListContext& Context, FVulkanQueue*
 	check(Context.IsImmediate());
 
 	FPlatformAtomics::AtomicStore(&LockToVsync, bLockToVsync ? 1 : 0);
-	bool bFailedToDelayAcquireBackbuffer = false;
 
 	//Transition back buffer to presentable and submit that command
 	if (SupportsStandardSwapchain())
 	{
+		bool bFailedToDelayAcquireBackbuffer = false;
+
 		if (GVulkanDelayAcquireImage == EDelayAcquireImageType::DelayAcquire && RenderingBackBuffer)
 		{
 			SCOPE_CYCLE_COUNTER(STAT_VulkanAcquireBackBuffer);
@@ -1037,10 +1037,7 @@ bool FVulkanViewport::Present(FVulkanCommandListContext& Context, FVulkanQueue*
 				check(RHIBackBuffer != nullptr && RHIBackBuffer->Image == RenderingBackBuffer->Image);
 			}
 		}
-	}
 
-	if (SupportsStandardSwapchain())
-	{
 		if (LIKELY(!bFailedToDelayAcquireBackbuffer))
 		{
 			if (GVulkanDelayAcquireImage == EDelayAcquireImageType::DelayAcquire)
@@ -1052,8 +1049,6 @@ bool FVulkanViewport::Present(FVulkanCommandListContext& Context, FVulkanQueue*
 			{
 				Context.AddSignalSemaphore(RenderingDoneSemaphores[AcquiredImageIndex]);
 			}
-
-			Context.SubmitActiveCmdBuffer();  // :todo-jn: deprecated
 		}
 		else
 		{
@@ -1062,23 +1057,18 @@ bool FVulkanViewport::Present(FVulkanCommandListContext& Context, FVulkanQueue*
 
 			UE_LOG(LogVulkanRHI, Log, TEXT("AcquireNextImage() failed due to the outdated swapchain, not even attempting to present."));
 
-			// cannot just throw out this command buffer (needs to be submitted or other checks fail)
-			Context.SubmitActiveCmdBuffer();  // :todo-jn: deprecated
-			RecreateSwapchain(Context, WindowHandle);
+			RecreateSwapchain(Context);
 
 			// Swapchain creation pushes some commands - flush the command buffers now to begin with a fresh state
-			Device->SubmitCommandsAndFlushGPU();  // :todo-jn: deprecated
-			Device->WaitUntilIdle();
+			Context.FlushCommands(EVulkanFlushFlags::WaitForCompletion);
 
 			// early exit
 			return (int32)FVulkanSwapChain::EStatus::Healthy;
 		}
 	}
-	else
-	{
-		// Submit active command buffer if not supporting standard swapchain (e.g. XR devices).
-		Context.SubmitActiveCmdBuffer();  // :todo-jn: deprecated
-	}
+
+	// Submit any accumulated commands or syncs, wait until they hit the queue so that we can present
+	Context.FlushCommands(EVulkanFlushFlags::WaitForSubmission);
 
 	//#todo-rco: Proper SyncInterval bLockToVsync ? RHIConsoleVariables::SyncInterval : 0
 	int32 SyncInterval = 0;
@@ -1193,7 +1183,7 @@ void FVulkanViewport::OnSystemResolutionChanged(uint32 ResX, uint32 ResY)
 		ENQUEUE_RENDER_COMMAND(RecreateSwapchain)(
 			[this](FRHICommandListImmediate& RHICmdList)
 			{
-				RecreateSwapchainFromRT(RHICmdList, PixelFormat);
+				RecreateSwapchainFromRT(RHICmdList);
 			});
 		FlushRenderingCommands();
 	}
diff --git a/Engine/Source/Runtime/VulkanRHI/Private/VulkanViewport.h b/Engine/Source/Runtime/VulkanRHI/Private/VulkanViewport.h
--- a/Engine/Source/Runtime/VulkanRHI/Private/VulkanViewport.h
+++ b/Engine/Source/Runtime/VulkanRHI/Private/VulkanViewport.h
@@ -149,8 +149,12 @@ protected:
 	bool TryAcquireImageIndex();
 	void InitImages(FVulkanContextCommon& Context, TConstArrayView<VkImage> Images);
 
-	void RecreateSwapchain(FVulkanCommandListContext& Context, void* NewNativeWindow);
-	void RecreateSwapchainFromRT(FRHICommandListImmediate& RHICmdList, EPixelFormat PreferredPixelFormat);
+	void RecreateSwapchain(FVulkanCommandListContext& Context);
+	void RecreateSwapchainFromRT(FRHICommandListImmediate& RHICmdList, EPixelFormat PreferredPixelFormat, void* NewNativeWindow);
+	void RecreateSwapchainFromRT(FRHICommandListImmediate& RHICmdList)
+	{
+		RecreateSwapchainFromRT(RHICmdList, PixelFormat, WindowHandle);
+	}
 	void Resize(uint32 InSizeX, uint32 InSizeY, bool bIsFullscreen, EPixelFormat PreferredPixelFormat);
 
 	bool DoCheckedSwapChainJob(FVulkanCommandListContext& Context, TFunction<int32(FVulkanViewport*)> SwapChainJob);
diff --git a/Engine/Source/Runtime/VulkanRHI/Public/VulkanDynamicRHI.h b/Engine/Source/Runtime/VulkanRHI/Public/VulkanDynamicRHI.h
--- a/Engine/Source/Runtime/VulkanRHI/Public/VulkanDynamicRHI.h
+++ b/Engine/Source/Runtime/VulkanRHI/Public/VulkanDynamicRHI.h
@@ -19,6 +19,12 @@ class IHeadMountedDisplayVulkanExtensions;
 class FVulkanThread;
 struct FRHITransientHeapAllocation;
 struct FVulkanPlatformCommandList;
+
+namespace VulkanRHI
+{
+	class FSemaphore;
+};
+
 struct FOptionalVulkanInstanceExtensions
 {
 	union
@@ -283,12 +289,17 @@ protected:
 	void InitializeSubmissionPipe();
 	void ShutdownSubmissionPipe();
 	bool ProcessSubmissionQueue();
-	bool ProcessInterruptQueue(bool bWaitForProgress);
+	bool ProcessInterruptQueue();
+	bool WaitAndProcessInterruptQueue();
+	void KickSubmissionThread();
+	void KickInterruptThread();
 
 	FCriticalSection SubmissionCS;
 	FVulkanThread* SubmissionThread = nullptr;
 	FCriticalSection InterruptCS;
 	FVulkanThread* InterruptThread = nullptr;
+	VulkanRHI::FSemaphore* CPUTimelineSemaphore = nullptr;  // used to wake up interrupt thread from CPU
+	std::atomic<uint64> CPUTimelineSemaphoreValue = 0;
 	TQueue<FVulkanPlatformCommandList*, EQueueMode::Mpsc> PendingPayloadsForSubmission;
 	TSet<VkSemaphore> SignaledSemaphores;
 	FGraphEventRef EopTask;
