VulkanRHI: Process interrupt queue while waiting for event in FlushCommands.

#rnx

diff --git a/Engine/Source/Runtime/VulkanRHI/Private/VulkanContext.cpp b/Engine/Source/Runtime/VulkanRHI/Private/VulkanContext.cpp
--- a/Engine/Source/Runtime/VulkanRHI/Private/VulkanContext.cpp
+++ b/Engine/Source/Runtime/VulkanRHI/Private/VulkanContext.cpp
@@ -132,7 +132,7 @@ void FVulkanContextCommon::FlushCommands(EVulkanFlushFlags FlushFlags)
 
 	if (SyncPoint)
 	{
-		SyncPoint->Wait();
+		FVulkanDynamicRHI::Get().ProcessInterruptQueueUntil(SyncPoint);
 	}
 
 	if (SubmissionEvent && !SubmissionEvent->IsComplete())
diff --git a/Engine/Source/Runtime/VulkanRHI/Private/VulkanSubmission.cpp b/Engine/Source/Runtime/VulkanRHI/Private/VulkanSubmission.cpp
--- a/Engine/Source/Runtime/VulkanRHI/Private/VulkanSubmission.cpp
+++ b/Engine/Source/Runtime/VulkanRHI/Private/VulkanSubmission.cpp
@@ -168,6 +168,7 @@ void FVulkanDynamicRHI::KickInterruptThread()
 		ZeroVulkanStruct(SignalInfo, VK_STRUCTURE_TYPE_SEMAPHORE_SIGNAL_INFO);
 		SignalInfo.semaphore = CPUTimelineSemaphore->GetHandle();
 		SignalInfo.value = (uint64_t)CPUTimelineSemaphoreValue++;
+		VERIFYVULKANRESULT(VulkanRHI::vkSignalSemaphoreKHR(Device->GetInstanceHandle(), &SignalInfo));
 
 		InterruptThread->Kick();
 	}
