VulkanRHI: Add back the wait for GPU in DestroySwapchain before we force process the deletion queue.

#rnx

diff --git a/Engine/Source/Runtime/VulkanRHI/Private/VulkanViewport.cpp b/Engine/Source/Runtime/VulkanRHI/Private/VulkanViewport.cpp
--- a/Engine/Source/Runtime/VulkanRHI/Private/VulkanViewport.cpp
+++ b/Engine/Source/Runtime/VulkanRHI/Private/VulkanViewport.cpp
@@ -717,7 +717,6 @@ void FVulkanViewport::RecreateSwapchainFromRT(FRHICommandListImmediate& RHICmdLi
 {
 	check(IsInRenderingThread());
 
-	// TODO: should flush RHIT commands here?
 	RHICmdList.ImmediateFlush(EImmediateFlushType::FlushRHIThread);
 
 	FVulkanSwapChainRecreateInfo RecreateInfo = { VK_NULL_HANDLE, VK_NULL_HANDLE};
@@ -849,6 +848,8 @@ void FVulkanViewport::CreateSwapchain(FRHICommandListBase& RHICmdList, FVulkanSw
 
 void FVulkanViewport::DestroySwapchain(FVulkanSwapChainRecreateInfo* RecreateInfo)
 {
+	FVulkanDynamicRHI::Get().RHIBlockUntilGPUIdle();
+
 	// Intentionally leave RenderingBackBuffer alive, so it can be used a dummy backbuffer while we don't have swapchain images
 	// RenderingBackBuffer = nullptr;
 	
