VulkanRHI: fix validation warning with mesh shader extension when primitiveShadingRate is disabled.

#rnx

diff --git a/Engine/Source/Runtime/VulkanRHI/Private/VulkanExtensions.cpp b/Engine/Source/Runtime/VulkanRHI/Private/VulkanExtensions.cpp
--- a/Engine/Source/Runtime/VulkanRHI/Private/VulkanExtensions.cpp
+++ b/Engine/Source/Runtime/VulkanRHI/Private/VulkanExtensions.cpp
@@ -1104,6 +1104,11 @@ public:
 	{
 		if (bRequirementsPassed)
 		{
+			if (!GetDeviceExtensionProperties().FragmentShadingRateFeatures.primitiveFragmentShadingRate)
+			{
+				MeshShaderFeatures.primitiveFragmentShadingRateMeshShader = VK_FALSE;
+			}
+
 			AddToPNext(DeviceCreateInfo, MeshShaderFeatures);
 		}
 	}
