VulkanRHI: fix validation error about buffer sizes not being multiples of 4 when clearing.

#rnx

diff --git a/Engine/Source/Runtime/VulkanRHI/Private/VulkanBuffer.cpp b/Engine/Source/Runtime/VulkanRHI/Private/VulkanBuffer.cpp
--- a/Engine/Source/Runtime/VulkanRHI/Private/VulkanBuffer.cpp
+++ b/Engine/Source/Runtime/VulkanRHI/Private/VulkanBuffer.cpp
@@ -169,7 +169,7 @@ void FVulkanBuffer::AllocateMemory(FBufferAlloc& OutAlloc)
 	{
 		BufferMemFlags = VK_MEMORY_PROPERTY_DEVICE_LOCAL_BIT;
 	}
-	const uint32 BufferSize = GetSize();
+	const uint32 BufferSize = Align(GetSize(), 4);  // keep allocated size a multiple of 4 (for use with vkCmdFillBuffer)
 	const uint32 BufferAlignment = VulkanRHI::FMemoryManager::CalculateBufferAlignment(*Device, GetUsage(), (BufferSize == 0));
 
 	FBufferAlloc& NewBufferAlloc = OutAlloc;
