VulkanRHI: change compatibility code for RADV since RT driver fix has been identified for version 24.3.2.  (integrate 38784523 and 38852466)

#rnx

diff --git a/Engine/Source/Runtime/VulkanRHI/Private/VulkanRHI.cpp b/Engine/Source/Runtime/VulkanRHI/Private/VulkanRHI.cpp
--- a/Engine/Source/Runtime/VulkanRHI/Private/VulkanRHI.cpp
+++ b/Engine/Source/Runtime/VulkanRHI/Private/VulkanRHI.cpp
@@ -862,9 +862,9 @@ void FVulkanDynamicRHI::InitInstance()
 		// Use this compatibility mode avoid known issues at launch time with latest drivers at the time of release 5.5.  This will:
 		// - disable inline ray tracing and use ray tracing pipelines everywhere (instead of a mix of both)
 		// - disable mesh shaders until issues can be resolved (holes in Nanite meshes)
-		// - disable ray tracing for RADV driver prior to version 25 in Linux (raytracing pipeline compilation crash)
+		// - disable ray tracing for RADV driver prior to version 24.3.2 in Linux (raytracing pipeline compilation crash)
 
-		// :todo-jn: to be removed when the official minimum RADV version is set to 25.0.0 in BaseHardware.ini
+		// :todo-jn: to be removed when the official minimum RADV version is set to 24.3.2 in BaseHardware.ini
 		const bool bUseAMDCompatibilityMode = GVulkanAMDCompatibilityMode && (Device->GetVendorId() == EGpuVendorId::Amd);
 		if (GVulkanAMDCompatibilityMode && 
 			(Device->GetOptionalExtensionProperties().PhysicalDeviceDriverProperties.driverID == VK_DRIVER_ID_MESA_RADV) &&
@@ -873,7 +873,7 @@ void FVulkanDynamicRHI::InitInstance()
 			GRHISupportsRayTracing = false;
 			GRHISupportsRayTracingShaders = false;
 			GRHISupportsInlineRayTracing = false;
-			UE_LOG(LogVulkanRHI, Warning, TEXT("Using MESA RADV version prior to 25.X, ray tracing disabled."));
+			UE_LOG(LogVulkanRHI, Warning, TEXT("Using MESA RADV version prior to 24.3.2, ray tracing disabled."));
 		}
 
 		if (GRHISupportsRayTracing)
