VulkanRHI: Change SyncPoint in the payload to an array.

#rnx

diff --git a/Engine/Source/Runtime/VulkanRHI/Private/VulkanContext.cpp b/Engine/Source/Runtime/VulkanRHI/Private/VulkanContext.cpp
--- a/Engine/Source/Runtime/VulkanRHI/Private/VulkanContext.cpp
+++ b/Engine/Source/Runtime/VulkanRHI/Private/VulkanContext.cpp
@@ -59,6 +59,12 @@ void FVulkanContextCommon::EndPayload()
 // payloads to the given array. Resets the context so new commands can be recorded.
 void FVulkanContextCommon::Finalize(TArray<FVulkanPayload*>& OutPayloads)
 {
+	if (ContextSyncPoint.IsValid())
+	{
+		AddSyncPoints(MakeArrayView<FVulkanSyncPointRef>(&ContextSyncPoint, 1));
+		ContextSyncPoint = nullptr;
+	}
+
 	EndPayload();
 
 	if (FVulkanPayload* ResetPayload = FlushResetQueryPools())
@@ -67,17 +73,6 @@ void FVulkanContextCommon::Finalize(TArray<FVulkanPayload*>& OutPayloads)
 	}
 
 	OutPayloads.Append(MoveTemp(Payloads));
-
-	if (ContextSyncPoint.IsValid())
-	{
-		// If the last payload has no space for the ContextSyncPoint, then add a dummy payload
-		if ((OutPayloads.Num() == 0) || OutPayloads.Last()->SyncPoint.IsValid())
-		{
-			OutPayloads.Add(new FVulkanPayload(Queue));
-		}
-		OutPayloads.Last()->SyncPoint = ContextSyncPoint;
-		ContextSyncPoint = nullptr;
-	}
 }
 
 void FVulkanContextCommon::AddQueryPoolForReset(VkQueryPool QueryPool, uint32 Size)
diff --git a/Engine/Source/Runtime/VulkanRHI/Private/VulkanContext.h b/Engine/Source/Runtime/VulkanRHI/Private/VulkanContext.h
--- a/Engine/Source/Runtime/VulkanRHI/Private/VulkanContext.h
+++ b/Engine/Source/Runtime/VulkanRHI/Private/VulkanContext.h
@@ -149,11 +149,18 @@ public:
 	FVulkanSyncPoint* AddSyncPoint()
 	{
 		FVulkanPayload& Payload = GetPayload(EPhase::Signal);
-		if (!Payload.SyncPoint.IsValid())
+		if (Payload.SyncPoints.Num() == 0)
 		{
-			Payload.SyncPoint = CreateVulkanSyncPoint();
+			Payload.SyncPoints.Add(CreateVulkanSyncPoint());
 		}
-		return Payload.SyncPoint;
+		return Payload.SyncPoints[0];
+	}
+
+	// Add a provided sync point that will be triggered at the end of the current payload's accumulated work
+	void AddSyncPoints(TArrayView<FVulkanSyncPointRef> InSyncs)
+	{
+		FVulkanPayload& Payload = GetPayload(EPhase::Signal);
+		Payload.SyncPoints.Append(InSyncs);
 	}
 
 	// Returns a single sync point for the context that will be inserted when it is finalized
diff --git a/Engine/Source/Runtime/VulkanRHI/Private/VulkanQueue.cpp b/Engine/Source/Runtime/VulkanRHI/Private/VulkanQueue.cpp
--- a/Engine/Source/Runtime/VulkanRHI/Private/VulkanQueue.cpp
+++ b/Engine/Source/Runtime/VulkanRHI/Private/VulkanQueue.cpp
@@ -334,11 +334,11 @@ int32 FVulkanQueue::ProcessInterruptQueue(uint64 Timeout)
 				CommandBuffer->ResetFenceStatus();
 			}
 
-			if (Payload->SyncPoint.IsValid())
+			for (FVulkanSyncPointRef& SyncPoint : Payload->SyncPoints)
 			{
-				Payload->SyncPoint->DispatchSubsequents();
-				Payload->SyncPoint = nullptr;
+				SyncPoint->DispatchSubsequents();
 			}
+			Payload->SyncPoints.Empty();
 
 			delete Payload;
 		}
diff --git a/Engine/Source/Runtime/VulkanRHI/Private/VulkanSubmission.cpp b/Engine/Source/Runtime/VulkanRHI/Private/VulkanSubmission.cpp
--- a/Engine/Source/Runtime/VulkanRHI/Private/VulkanSubmission.cpp
+++ b/Engine/Source/Runtime/VulkanRHI/Private/VulkanSubmission.cpp
@@ -432,11 +432,14 @@ void FVulkanDynamicRHI::EnqueueEndOfPipeTask(TUniqueFunction<void()> TaskFunc, T
 	{
 		FVulkanPayload* Payload = new FVulkanPayload(Queue);
 
-		Payload->SyncPoint = CreateVulkanSyncPoint();
-		Prereqs.Add(Payload->SyncPoint);
+		FVulkanSyncPointRef SyncPoint = CreateVulkanSyncPoint();
+		Payload->SyncPoints.Add(SyncPoint);
+		Prereqs.Add(SyncPoint);
 
 		if (ModifyPayloadCallback)
+		{
 			ModifyPayloadCallback(*Payload);
+		}
 
 		Payloads->Add(Payload);
 	});
@@ -465,10 +468,11 @@ void FVulkanDynamicRHI::RHIBlockUntilGPUIdle()
 		Device->ForEachQueue([&EndOfQueueEvent, &Payloads](FVulkanQueue& Queue)
 		{
 			FVulkanPayload* Payload = new FVulkanPayload(Queue);
-			Payload->SyncPoint = CreateVulkanSyncPoint();
+			FVulkanSyncPointRef SyncPoint = CreateVulkanSyncPoint();
+			Payload->SyncPoints.Add(SyncPoint);
 			Payloads->Add(Payload);
 
-			EndOfQueueEvent.Add(Payload->SyncPoint);
+			EndOfQueueEvent.Add(SyncPoint);
 		});
 
 		PendingPayloadsForSubmission.Enqueue(Payloads);
diff --git a/Engine/Source/Runtime/VulkanRHI/Private/VulkanSubmission.h b/Engine/Source/Runtime/VulkanRHI/Private/VulkanSubmission.h
--- a/Engine/Source/Runtime/VulkanRHI/Private/VulkanSubmission.h
+++ b/Engine/Source/Runtime/VulkanRHI/Private/VulkanSubmission.h
@@ -53,7 +53,7 @@ protected:
 	VulkanRHI::FFence* Fence = nullptr;
 
 	// Used to sync other CPU work to payload completion
-	FVulkanSyncPointRef SyncPoint;
+	TArray<FVulkanSyncPointRef> SyncPoints;
 };
 
 struct FVulkanPlatformCommandList : public IRHIPlatformCommandList, public TArray<FVulkanPayload*>
