VulkanRHI: Add pending syncs array to accumulate syncs during render passes.

#rnx

diff --git a/Engine/Source/Runtime/VulkanRHI/Private/VulkanContext.cpp b/Engine/Source/Runtime/VulkanRHI/Private/VulkanContext.cpp
--- a/Engine/Source/Runtime/VulkanRHI/Private/VulkanContext.cpp
+++ b/Engine/Source/Runtime/VulkanRHI/Private/VulkanContext.cpp
@@ -40,18 +40,23 @@ void FVulkanContextCommon::NewPayload()
 
 void FVulkanContextCommon::EndPayload()
 {
-	if ((Payloads.Num() > 0) && (Payloads.Last()->CommandBuffers.Num() > 0))
+	if (Payloads.Num() > 0)
 	{
-		FVulkanCommandBuffer* CommandBuffer = Payloads.Last()->CommandBuffers.Last();
-		checkSlow(!CommandBuffer->IsSubmitted() && CommandBuffer->HasBegun());
+		FlushPendingSyncPoints();
 
-		if (!CommandBuffer->IsOutsideRenderPass())
+		if (Payloads.Last()->CommandBuffers.Num() > 0)
 		{
-			UE_LOG(LogVulkanRHI, Warning, TEXT("Forcing EndRenderPass() for submission"));
-			CommandBuffer->EndRenderPass();
-		}
+			FVulkanCommandBuffer* CommandBuffer = Payloads.Last()->CommandBuffers.Last();
+			checkSlow(!CommandBuffer->IsSubmitted() && CommandBuffer->HasBegun());
+
+			if (!CommandBuffer->IsOutsideRenderPass())
+			{
+				UE_LOG(LogVulkanRHI, Warning, TEXT("Forcing EndRenderPass() for submission"));
+				CommandBuffer->EndRenderPass();
+			}
 
-		CommandBuffer->End();
+			CommandBuffer->End();
+		}
 	}
 }
 
@@ -59,6 +64,8 @@ void FVulkanContextCommon::EndPayload()
 // payloads to the given array. Resets the context so new commands can be recorded.
 void FVulkanContextCommon::Finalize(TArray<FVulkanPayload*>& OutPayloads)
 {
+	FlushPendingSyncPoints();
+
 	if (ContextSyncPoint.IsValid())
 	{
 		AddSyncPoints(MakeArrayView<FVulkanSyncPointRef>(&ContextSyncPoint, 1));
diff --git a/Engine/Source/Runtime/VulkanRHI/Private/VulkanContext.h b/Engine/Source/Runtime/VulkanRHI/Private/VulkanContext.h
--- a/Engine/Source/Runtime/VulkanRHI/Private/VulkanContext.h
+++ b/Engine/Source/Runtime/VulkanRHI/Private/VulkanContext.h
@@ -163,6 +163,23 @@ public:
 		Payload.SyncPoints.Append(InSyncs);
 	}
 
+	// Add provided sync points that will be appended next time a payload reaches the Signal phase
+	void AddPendingSyncPoint(const FVulkanSyncPointRef& InSyncs)
+	{
+		PendingSyncPoints.Add(InSyncs);
+	}
+
+	// Force pending syncs to be sent to a payload
+	void FlushPendingSyncPoints()
+	{
+		if (PendingSyncPoints.Num())
+		{
+			FVulkanPayload& Payload = GetPayload(EPhase::Signal);
+			Payload.SyncPoints.Append(PendingSyncPoints);
+			PendingSyncPoints.Empty();
+		}
+	}
+
 	// Returns a single sync point for the context that will be inserted when it is finalized
 	FVulkanSyncPoint* GetContextSyncPoint()
 	{
@@ -194,6 +211,10 @@ private:
 	TArray<FVulkanPayload*> Payloads;
 	EPhase CurrentPhase = EPhase::Wait;
 
+	// Sync points signaled at the next Signal phase (will not force a phase change) or when context is finalized
+	TArray<FVulkanSyncPointRef> PendingSyncPoints;
+
+	// Sync point signaled when the current context is finalized
 	FVulkanSyncPointRef ContextSyncPoint;
 };
 
