VulkanRHI: Move temp block allocator to device for sharing between contexts when multithreaded.

#rnx

diff --git a/Engine/Source/Runtime/VulkanRHI/Private/VulkanContext.cpp b/Engine/Source/Runtime/VulkanRHI/Private/VulkanContext.cpp
--- a/Engine/Source/Runtime/VulkanRHI/Private/VulkanContext.cpp
+++ b/Engine/Source/Runtime/VulkanRHI/Private/VulkanContext.cpp
@@ -9,14 +9,6 @@
 #include "VulkanPendingState.h"
 #include "DynamicRHI.h"
 
-int32 GVulkanTempBlockSize = 4 * 1024 * 1024;
-static FAutoConsoleVariableRef GCVarVulkanTempBlockSize(
-	TEXT("r.Vulkan.TempBlockSize"),
-	GVulkanTempBlockSize,
-	TEXT("Size of the temporary blocks allocate by contexts, used for single use ub allocs and copies (default: 4MB)."),
-	ECVF_ReadOnly
-);
-
 
 
 FVulkanContextCommon::FVulkanContextCommon(FVulkanDevice& InDevice, FVulkanQueue& InQueue)
@@ -243,21 +235,6 @@ FVulkanCommandListContext::FVulkanCommandListContext(FVulkanDynamicRHI* InRHI, F
 	PendingGfxState = new FVulkanPendingGfxState(&Device, *this);
 	PendingComputeState = new FVulkanPendingComputeState(&Device, *this);
 
-	// Currently used for UB and copies
-	{
-		uint32 BlockAlignment = FMath::Max<uint32>(Device.GetLimits().minUniformBufferOffsetAlignment, 16u);
-		VkBufferUsageFlags BufferUsageFlags =
-			(Device.GetOptionalExtensions().HasBufferDeviceAddress ? VK_BUFFER_USAGE_SHADER_DEVICE_ADDRESS_BIT : 0) |
-			VK_BUFFER_USAGE_TRANSFER_SRC_BIT |
-			VK_BUFFER_USAGE_UNIFORM_BUFFER_BIT;
-		if (Device.GetOptionalExtensions().HasRayTracingPipeline)
-		{
-			BufferUsageFlags |= VK_BUFFER_USAGE_SHADER_BINDING_TABLE_BIT_KHR;
-			BlockAlignment = FMath::Max<uint32>(Device.GetOptionalExtensionProperties().RayTracingPipelineProps.shaderGroupBaseAlignment, BlockAlignment);
-		}
-		TempBlockAllocator = new VulkanRHI::FTempBlockAllocator(&InDevice, GVulkanTempBlockSize, BlockAlignment, BufferUsageFlags);
-	}
-
 	GlobalUniformBuffers.AddZeroed(FUniformBufferStaticSlotRegistry::Get().GetSlotCount());
 }
 
@@ -281,7 +258,6 @@ FVulkanCommandListContext::~FVulkanCommandListContext()
 #endif
 	}
 
-	delete TempBlockAllocator;
 	delete PendingGfxState;
 	delete PendingComputeState;
 }
diff --git a/Engine/Source/Runtime/VulkanRHI/Private/VulkanContext.h b/Engine/Source/Runtime/VulkanRHI/Private/VulkanContext.h
--- a/Engine/Source/Runtime/VulkanRHI/Private/VulkanContext.h
+++ b/Engine/Source/Runtime/VulkanRHI/Private/VulkanContext.h
@@ -328,11 +328,6 @@ public:
 		uint32 NumBindings, const FRayTracingLocalShaderBindings* Bindings,
 		ERayTracingBindingType BindingType) final override;
 
-	inline VulkanRHI::FTempBlockAllocator& GetTempBlockAllocator()
-	{
-		return *TempBlockAllocator;
-	}
-
 	inline FVulkanPendingGfxState* GetPendingGfxState()
 	{
 		return PendingGfxState;
@@ -387,8 +382,6 @@ protected:
 	void BeginOcclusionQueryBatch(FVulkanCommandBuffer* CmdBuffer, uint32 NumQueriesInBatch);
 	void EndOcclusionQueryBatch();
 
-	VulkanRHI::FTempBlockAllocator* TempBlockAllocator = nullptr;
-
 	TArray<FString> EventStack;
 
 	FVulkanRenderPass* CurrentRenderPass = nullptr;
diff --git a/Engine/Source/Runtime/VulkanRHI/Private/VulkanDevice.cpp b/Engine/Source/Runtime/VulkanRHI/Private/VulkanDevice.cpp
--- a/Engine/Source/Runtime/VulkanRHI/Private/VulkanDevice.cpp
+++ b/Engine/Source/Runtime/VulkanRHI/Private/VulkanDevice.cpp
@@ -50,6 +50,13 @@ static TAutoConsoleVariable<int32> CVarVulkanUseD24(
 	ECVF_ReadOnly
 );
 
+static TAutoConsoleVariable<int32> GCVarVulkanTempBlockSizeKB(
+	TEXT("r.Vulkan.TempBlockSizeKB"),
+	4096,
+	TEXT("Size in KB of the temporary blocks allocate by contexts, used for single use ub allocs and copies (default: 4096KB)."),
+	ECVF_ReadOnly
+);
+
 // Mirror GPixelFormats with format information for buffers
 VkFormat GVulkanBufferFormat[PF_MAX];
 
@@ -1329,6 +1336,23 @@ void FVulkanDevice::InitGPU()
 
 	StagingManager.Init(this);
 
+	// TempBlockAllocator is currently used for UB uploads and copies
+	{
+		uint32 BlockAlignment = FMath::Max<uint32>(GetLimits().minUniformBufferOffsetAlignment, 16u);
+		VkBufferUsageFlags BufferUsageFlags =
+			(GetOptionalExtensions().HasBufferDeviceAddress ? VK_BUFFER_USAGE_SHADER_DEVICE_ADDRESS_BIT : 0) |
+			VK_BUFFER_USAGE_TRANSFER_SRC_BIT |
+			VK_BUFFER_USAGE_UNIFORM_BUFFER_BIT;
+		if (GetOptionalExtensions().HasRayTracingPipeline)
+		{
+			BufferUsageFlags |= VK_BUFFER_USAGE_SHADER_BINDING_TABLE_BIT_KHR;
+			BlockAlignment = FMath::Max<uint32>(GetOptionalExtensionProperties().RayTracingPipelineProps.shaderGroupBaseAlignment, BlockAlignment);
+		}
+
+		const int32 VulkanTempBlockSize = GCVarVulkanTempBlockSizeKB.GetValueOnAnyThread() * 1024;
+		TempBlockAllocator = new VulkanRHI::FTempBlockAllocator(this, VulkanTempBlockSize, BlockAlignment, BufferUsageFlags);
+	}
+
 #if VULKAN_SUPPORTS_GPU_CRASH_DUMPS
 	if (UE::RHI::UseGPUCrashDebugging())
 	{
@@ -1523,6 +1547,9 @@ void FVulkanDevice::Destroy()
 
 	FVulkanUploadContext::DestroyPool();
 
+	delete TempBlockAllocator;
+	TempBlockAllocator = nullptr;
+
 	MemoryManager.Deinit();
 
 	for (int32 QueueIndex = 0; QueueIndex < (int32)EVulkanQueueType::Count; ++QueueIndex)
diff --git a/Engine/Source/Runtime/VulkanRHI/Private/VulkanDevice.h b/Engine/Source/Runtime/VulkanRHI/Private/VulkanDevice.h
--- a/Engine/Source/Runtime/VulkanRHI/Private/VulkanDevice.h
+++ b/Engine/Source/Runtime/VulkanRHI/Private/VulkanDevice.h
@@ -418,6 +418,11 @@ public:
 		return FenceManager;
 	}
 
+	inline VulkanRHI::FTempBlockAllocator& GetTempBlockAllocator()
+	{
+		return *TempBlockAllocator;
+	}
+
 	inline FVulkanRenderPassManager& GetRenderPassManager()
 	{
 		return *RenderPassManager;
@@ -550,6 +555,8 @@ private:
 
 	VulkanRHI::FFenceManager FenceManager;
 
+	VulkanRHI::FTempBlockAllocator* TempBlockAllocator = nullptr;
+
 	FVulkanRenderPassManager* RenderPassManager;
 
 	FVulkanTransientHeapCache* TransientHeapCache = nullptr;
diff --git a/Engine/Source/Runtime/VulkanRHI/Private/VulkanPipelineState.cpp b/Engine/Source/Runtime/VulkanRHI/Private/VulkanPipelineState.cpp
--- a/Engine/Source/Runtime/VulkanRHI/Private/VulkanPipelineState.cpp
+++ b/Engine/Source/Runtime/VulkanRHI/Private/VulkanPipelineState.cpp
@@ -165,7 +165,7 @@ static inline VulkanRHI::FVulkanAllocation UpdatePackedUniformBuffers(const FPac
 	const uint32 UBAlign = (uint32)InCmdBuffer->Device.GetLimits().minUniformBufferOffsetAlignment;
 
 	VulkanRHI::FVulkanAllocation TempAllocation;
-	uint8* MappedPointer = InContext->GetTempBlockAllocator().Alloc(UBSize, UBAlign, *InContext, TempAllocation);
+	uint8* MappedPointer = InContext->Device.GetTempBlockAllocator().Alloc(UBSize, UBAlign, *InContext, TempAllocation);
 
 	FMemory::Memcpy(MappedPointer, StagedUniformBuffer.GetData(), UBSize);
 
@@ -250,7 +250,7 @@ void FVulkanComputePipelineDescriptorState::UpdateBindlessDescriptors(FVulkanCom
 		const VkDeviceSize UBOffsetAlignment = Device->GetLimits().minUniformBufferOffsetAlignment;
 
 		VulkanRHI::FVulkanAllocation TempAllocation;
-		uint8* MappedPointer = CmdListContext->GetTempBlockAllocator().Alloc(UBSize, UBOffsetAlignment, *CmdListContext, TempAllocation, &DescriptorAddressInfos[BindingIndex]);
+		uint8* MappedPointer = Device->GetTempBlockAllocator().Alloc(UBSize, UBOffsetAlignment, *CmdListContext, TempAllocation, &DescriptorAddressInfos[BindingIndex]);
 		FMemory::Memcpy(MappedPointer, StagedUniformBuffer.GetData(), UBSize);
 
 		PackedUniformBuffersDirty = 0;
@@ -476,7 +476,7 @@ void FVulkanGraphicsPipelineDescriptorState::UpdateBindlessDescriptors(FVulkanCo
 				const int32 BindingIndex = 0;
 
 				VulkanRHI::FVulkanAllocation TempAllocation;
-				uint8* MappedPointer = CmdListContext->GetTempBlockAllocator().Alloc(UBSize, UBOffsetAlignment, *CmdListContext, TempAllocation, &DescriptorAddressInfos[BindingIndex]);
+				uint8* MappedPointer = Device->GetTempBlockAllocator().Alloc(UBSize, UBOffsetAlignment, *CmdListContext, TempAllocation, &DescriptorAddressInfos[BindingIndex]);
 				FMemory::Memcpy(MappedPointer, StagedUniformBuffer.GetData(), UBSize);
 
 				PackedUniformBuffersDirty[Stage] = 0;
diff --git a/Engine/Source/Runtime/VulkanRHI/Private/VulkanRHI.cpp b/Engine/Source/Runtime/VulkanRHI/Private/VulkanRHI.cpp
--- a/Engine/Source/Runtime/VulkanRHI/Private/VulkanRHI.cpp
+++ b/Engine/Source/Runtime/VulkanRHI/Private/VulkanRHI.cpp
@@ -1110,7 +1110,7 @@ void FVulkanDynamicRHI::RHIEndFrame_RenderThread(FRHICommandListImmediate& RHICm
 		Context.Device.GetPipelineStateCache()->TickLRU();
 
 		Context.Device.GetBindlessDescriptorManager()->UpdateUBAllocator();
-		Context.GetTempBlockAllocator().UpdateBlocks();
+		Context.Device.GetTempBlockAllocator().UpdateBlocks();
 
 		++Context.FrameCounter;
 	});
diff --git a/Engine/Source/Runtime/VulkanRHI/Private/VulkanRayTracing.cpp b/Engine/Source/Runtime/VulkanRHI/Private/VulkanRayTracing.cpp
--- a/Engine/Source/Runtime/VulkanRHI/Private/VulkanRayTracing.cpp
+++ b/Engine/Source/Runtime/VulkanRHI/Private/VulkanRayTracing.cpp
@@ -981,7 +981,7 @@ void FVulkanRayTracingShaderTable::SetSlot(EShaderFrequency Frequency, uint32 Ds
 VkStridedDeviceAddressRegionKHR FVulkanRayTracingShaderTable::CommitRayGenShader(FVulkanCommandListContext& Context, uint32 SrcHandleIndex, TConstArrayView<uint8> SrcHandleData)
 {
 	VkStridedDeviceAddressRegionKHR RayGenRegion;
-	uint8* pMappedMemory = Context.GetTempBlockAllocator().Alloc(HandleSize, Context, &RayGenRegion);
+	uint8* pMappedMemory = Device->GetTempBlockAllocator().Alloc(HandleSize, Context, &RayGenRegion);
 
 	FMemory::Memcpy(pMappedMemory, &SrcHandleData[SrcHandleIndex * HandleSize], HandleSize);
 
diff --git a/Engine/Source/Runtime/VulkanRHI/Private/VulkanUniformBuffer.cpp b/Engine/Source/Runtime/VulkanRHI/Private/VulkanUniformBuffer.cpp
--- a/Engine/Source/Runtime/VulkanRHI/Private/VulkanUniformBuffer.cpp
+++ b/Engine/Source/Runtime/VulkanRHI/Private/VulkanUniformBuffer.cpp
@@ -50,7 +50,7 @@ static void UpdateUniformBufferHelper(FVulkanCommandListContext& Context, FVulka
 	const int32 DataAlignment = FMath::Max<uint32>(Device.GetLimits().minUniformBufferOffsetAlignment, 16u);
 
 	VulkanRHI::FVulkanAllocation TempAllocation;
-	void* DestinationData = Context.GetTempBlockAllocator().Alloc(DataSize, DataAlignment, Context, TempAllocation);
+	void* DestinationData = Device.GetTempBlockAllocator().Alloc(DataSize, DataAlignment, Context, TempAllocation);
 
 	if (bUpdateConstants)
 	{
