Fix Android Vulkan SM5 build for blueprint projects

Vulkan RHI now onlty reports GRHISupportsRayTracing=1 if RHI_RAYTRACING=1, this mainly improves log output.
Blueprint projects with Android Vulkan SM5 support enabled now trigger a build with RHI_RAYTRACING=1.

diff --git a/Engine/Source/Programs/UnrealBuildTool/Platform/Android/UEBuildAndroid.cs b/Engine/Source/Programs/UnrealBuildTool/Platform/Android/UEBuildAndroid.cs
--- a/Engine/Source/Programs/UnrealBuildTool/Platform/Android/UEBuildAndroid.cs
+++ b/Engine/Source/Programs/UnrealBuildTool/Platform/Android/UEBuildAndroid.cs
@@ -295,7 +295,7 @@ namespace UnrealBuildTool
 		{
 			string[] BoolKeys = new string[] {
 				"bBuildForArm64", "bBuildForX8664",
-				"bBuildForES31", "bBuildWithHiddenSymbolVisibility", "bSaveSymbols"
+				"bBuildForES31", "bBuildWithHiddenSymbolVisibility", "bSaveSymbols", "bSupportsVulkanSM5"
 			};
 			string[] StringKeys = new string[] {
 				"NDKAPILevelOverride"
diff --git a/Engine/Source/Runtime/VulkanRHI/Private/VulkanRHI.cpp b/Engine/Source/Runtime/VulkanRHI/Private/VulkanRHI.cpp
--- a/Engine/Source/Runtime/VulkanRHI/Private/VulkanRHI.cpp
+++ b/Engine/Source/Runtime/VulkanRHI/Private/VulkanRHI.cpp
@@ -853,7 +853,7 @@ void FVulkanDynamicRHI::InitInstance()
 		GRHISupportsGPUTimestampBubblesRemoval = true;
 		GSupportsMobileMultiView = Device->GetOptionalExtensions().HasKHRMultiview ? true : false;
 		GRHISupportsMSAAShaderResolve = Device->GetOptionalExtensions().HasQcomRenderPassShaderResolve ? true : false;
-		GRHISupportsRayTracing = RHISupportsRayTracing(GMaxRHIShaderPlatform) && Device->GetOptionalExtensions().HasRaytracingExtensions();
+		GRHISupportsRayTracing = RHI_RAYTRACING && RHISupportsRayTracing(GMaxRHIShaderPlatform) && Device->GetOptionalExtensions().HasRaytracingExtensions();
 
 		GRHIGlobals.EnableHighLevelTransitionPatching = true;
 		GRHIGlobals.NeedExtraTransitions = true;
