VulkanRHI: Add missing RHIThreadFence when creating views and locking buffers

#rnx

diff --git a/Engine/Source/Runtime/VulkanRHI/Private/VulkanBuffer.cpp b/Engine/Source/Runtime/VulkanRHI/Private/VulkanBuffer.cpp
--- a/Engine/Source/Runtime/VulkanRHI/Private/VulkanBuffer.cpp
+++ b/Engine/Source/Runtime/VulkanRHI/Private/VulkanBuffer.cpp
@@ -307,6 +307,11 @@ void* FVulkanBuffer::Lock(FRHICommandListBase& RHICmdList, EResourceLockMode Loc
 				Buffer->UpdateLinkedViews();
 			});
 
+			if (RHICmdList.IsTopOfPipe())
+			{
+				RHICmdList.RHIThreadFence(true);
+			}
+
 			Data = NewAlloc.HostPtr;
 			DataOffset = Offset;
 			LockStatus = ELockStatus::PersistentMapping;
diff --git a/Engine/Source/Runtime/VulkanRHI/Private/VulkanUAV.cpp b/Engine/Source/Runtime/VulkanRHI/Private/VulkanUAV.cpp
--- a/Engine/Source/Runtime/VulkanRHI/Private/VulkanUAV.cpp
+++ b/Engine/Source/Runtime/VulkanRHI/Private/VulkanUAV.cpp
@@ -384,6 +384,8 @@ FVulkanShaderResourceView::FVulkanShaderResourceView(FRHICommandListBase& RHICmd
 		LinkHead(GetBaseResource()->LinkedViews);
 		UpdateView();
 	});
+
+	RHICmdList.RHIThreadFence(true);
 }
 
 FVulkanViewableResource* FVulkanShaderResourceView::GetBaseResource() const
@@ -474,6 +476,8 @@ FVulkanUnorderedAccessView::FVulkanUnorderedAccessView(FRHICommandListBase& RHIC
 		LinkHead(GetBaseResource()->LinkedViews);
 		UpdateView();
 	});
+
+	RHICmdList.RHIThreadFence(true);
 }
 
 FVulkanViewableResource* FVulkanUnorderedAccessView::GetBaseResource() const
