Apply depth-stencil workaround only if needed, based on driver version

PR #12804: Apply depth-stencil workaround only if needed, based on driver version

diff --git a/Engine/Config/BaseDeviceProfiles.ini b/Engine/Config/BaseDeviceProfiles.ini
--- a/Engine/Config/BaseDeviceProfiles.ini
+++ b/Engine/Config/BaseDeviceProfiles.ini
@@ -979,8 +979,6 @@ BaseProfileName=Android_Vulkan_SM5
 [Android_Mali_Vulkan_SM5 DeviceProfile]
 DeviceType=Android
 BaseProfileName=Android_Vulkan_SM5
-; DepthStencilForceStorageBit is a workaround that is only needed for Mali before the G710, so it's not needed based on the current matching rules for Android_Mali_Vulkan_SM5
-;+CVars=r.Vulkan.DepthStencilForceStorageBit=1
 
 [Android_Xclipse_Vulkan_SM5 DeviceProfile]
 DeviceType=Android
diff --git a/Engine/Source/Runtime/VulkanRHI/Private/Android/VulkanAndroidPlatform.cpp b/Engine/Source/Runtime/VulkanRHI/Private/Android/VulkanAndroidPlatform.cpp
--- a/Engine/Source/Runtime/VulkanRHI/Private/Android/VulkanAndroidPlatform.cpp
+++ b/Engine/Source/Runtime/VulkanRHI/Private/Android/VulkanAndroidPlatform.cpp
@@ -101,6 +101,7 @@ TArray<TArray<ANSICHAR>> FVulkanAndroidPlatform::DebugVulkanInstanceLayers;
 TArray<TArray<ANSICHAR>> FVulkanAndroidPlatform::SwappyRequiredExtensions;
 int32 FVulkanAndroidPlatform::AFBCWorkaroundOption = 0;
 int32 FVulkanAndroidPlatform::ASTCWorkaroundOption = 0;
+bool FVulkanAndroidPlatform::bRequiresDepthStencilFullWrite = false;
 
 #define CHECK_VK_ENTRYPOINTS(Type,Func) if (VulkanDynamicAPI::Func == NULL) { bFoundAllEntryPoints = false; UE_LOG(LogRHI, Warning, TEXT("Failed to find entry point for %s"), TEXT(#Func)); }
 
@@ -641,6 +642,20 @@ void FVulkanAndroidPlatform::OverridePlatformHandlers(bool bInit)
 	}
 }
 
+void FVulkanAndroidPlatform::SetupRequiresDepthStencilFullWriteWorkaround(const FVulkanDevice& Device)
+{
+	bRequiresDepthStencilFullWrite = false;
+	if (Device.GetVendorId() == EGpuVendorId::Arm)
+	{
+		const uint32_t DriverVersion = Device.GetDeviceProperties().driverVersion;
+		if (DriverVersion >= VK_MAKE_API_VERSION(0, 16, 0, 0) && DriverVersion < VK_MAKE_API_VERSION(0, 24, 0, 0))
+		{
+			bRequiresDepthStencilFullWrite = true;
+			UE_LOG(LogRHI, Display, TEXT("Enabling workaround that changes partial depth/stencil attachment stores into full depth/stencil stores."));
+		}
+	}
+}
+
 bool FVulkanAndroidPlatform::FramePace(FVulkanDevice& Device, void* WindowHandle, VkSwapchainKHR Swapchain, uint32 PresentID, VkPresentInfoKHR& Info)
 {
 	QUICK_SCOPE_CYCLE_COUNTER(STAT_VulkanAndroid_FramePace);
@@ -1732,6 +1747,7 @@ bool FAndroidVulkanFramePacer::SupportsFramePace(int32 QueryFramePace)
 void FVulkanAndroidPlatform::PostInitGPU(const FVulkanDevice& InDevice)
 {
 	SetupImageMemoryRequirementWorkaround(InDevice);
+	SetupRequiresDepthStencilFullWriteWorkaround(InDevice);
 
 	// start external compilers if precaching is specified.	
 	static const auto CVarChunkedPSOCache = IConsoleManager::Get().FindConsoleVariable(TEXT("r.Vulkan.UseChunkedPSOCache"));
diff --git a/Engine/Source/Runtime/VulkanRHI/Private/Android/VulkanAndroidPlatform.h b/Engine/Source/Runtime/VulkanRHI/Private/Android/VulkanAndroidPlatform.h
--- a/Engine/Source/Runtime/VulkanRHI/Private/Android/VulkanAndroidPlatform.h
+++ b/Engine/Source/Runtime/VulkanRHI/Private/Android/VulkanAndroidPlatform.h
@@ -135,8 +135,9 @@ public:
 	// Does the platform allow a nullptr Pixelshader on the pipeline
 	static bool SupportsNullPixelShader() { return false; }
 
-	//#todo-rco: Detect Mali? Does the platform require depth to be written on stencil clear
-	static bool RequiresDepthWriteOnStencilClear() { return true; }
+	// Does the platform require depth to be written on stencil clear
+	static bool RequiresDepthStencilFullWrite() { return bRequiresDepthStencilFullWrite; }
+	static void SetupRequiresDepthStencilFullWriteWorkaround(const FVulkanDevice& Device);
 
 	static bool FramePace(FVulkanDevice& Device, void* WindowHandle, VkSwapchainKHR Swapchain, uint32 PresentID, VkPresentInfoKHR& Info);
 
@@ -203,6 +204,8 @@ protected:
 
 	static int32 AFBCWorkaroundOption;
 	static int32 ASTCWorkaroundOption;
+
+	static bool bRequiresDepthStencilFullWrite;
 };
 
 #if VULKAN_SUPPORTS_GOOGLE_DISPLAY_TIMING
diff --git a/Engine/Source/Runtime/VulkanRHI/Private/VulkanGenericPlatform.h b/Engine/Source/Runtime/VulkanRHI/Private/VulkanGenericPlatform.h
--- a/Engine/Source/Runtime/VulkanRHI/Private/VulkanGenericPlatform.h
+++ b/Engine/Source/Runtime/VulkanRHI/Private/VulkanGenericPlatform.h
@@ -118,7 +118,7 @@ public:
 	static bool SupportsNullPixelShader() { return true; }
 
 	// Does the platform require depth to be written on stencil clear
-	static bool RequiresDepthWriteOnStencilClear() { return false; }
+	static bool RequiresDepthStencilFullWrite() { return false; }
 
 	// Checks if the PSO cache matches the expected vulkan device properties
 	static bool PSOBinaryCacheMatches(FVulkanDevice* Device, const TArray<uint8>& DeviceCache);
diff --git a/Engine/Source/Runtime/VulkanRHI/Private/VulkanRenderTarget.cpp b/Engine/Source/Runtime/VulkanRHI/Private/VulkanRenderTarget.cpp
--- a/Engine/Source/Runtime/VulkanRHI/Private/VulkanRenderTarget.cpp
+++ b/Engine/Source/Runtime/VulkanRHI/Private/VulkanRenderTarget.cpp
@@ -995,13 +995,13 @@ FVulkanRenderTargetLayout::FVulkanRenderTargetLayout(FVulkanDevice& InDevice, co
 			ensure(CurrDesc.stencilStoreOp == VK_ATTACHMENT_STORE_OP_DONT_CARE);
 		}
 
-		FExclusiveDepthStencil ExclusiveDepthStencil = RPInfo.DepthStencilRenderTarget.ExclusiveDepthStencil;
-		if (FVulkanPlatform::RequiresDepthWriteOnStencilClear() &&
-			RPInfo.DepthStencilRenderTarget.Action == EDepthStencilTargetActions::LoadDepthClearStencil_StoreDepthStencil)
+		if (FVulkanPlatform::RequiresDepthStencilFullWrite() &&
+			Texture->GetFullAspectMask() == (VK_IMAGE_ASPECT_DEPTH_BIT | VK_IMAGE_ASPECT_STENCIL_BIT) &&
+			(CurrDesc.storeOp == VK_ATTACHMENT_STORE_OP_STORE || CurrDesc.stencilStoreOp == VK_ATTACHMENT_STORE_OP_STORE))
 		{
-			ExclusiveDepthStencil = FExclusiveDepthStencil::DepthWrite_StencilWrite;
-			CurrentDepthLayout = VK_IMAGE_LAYOUT_ATTACHMENT_OPTIMAL;
-			CurrentStencilLayout = VK_IMAGE_LAYOUT_ATTACHMENT_OPTIMAL;
+			// Workaround for old mali drivers: writing not all of the image aspects to compressed render-target could cause gpu-hang
+			CurrDesc.storeOp = VK_ATTACHMENT_STORE_OP_STORE;
+			CurrDesc.stencilStoreOp = VK_ATTACHMENT_STORE_OP_STORE;
 		}
 
 		// If the initial != final we need to change the FullHashInfo and use FinalLayout
