VulkanRHI:  Add support for RLM_WriteOnly_NoOverwrite.  Fix assert in UpdateBoneData in Vulkan.

#rnx

diff --git a/Engine/Source/Runtime/VulkanRHI/Private/VulkanBuffer.cpp b/Engine/Source/Runtime/VulkanRHI/Private/VulkanBuffer.cpp
--- a/Engine/Source/Runtime/VulkanRHI/Private/VulkanBuffer.cpp
+++ b/Engine/Source/Runtime/VulkanRHI/Private/VulkanBuffer.cpp
@@ -267,11 +267,11 @@ void* FVulkanBuffer::Lock(FRHICommandListBase& RHICmdList, EResourceLockMode Loc
 	}
 	else
 	{
-		check(LockMode == RLM_WriteOnly);
+		check((LockMode == RLM_WriteOnly) || (LockMode == RLM_WriteOnly_NoOverwrite));
 
 		// If this is the first lock on host visible memory, then the memory is still untouched so use it directly
 		const bool bIsHostVisible = (bUnifiedMem || bDynamic);
-		if (bIsHostVisible && bIsFirstLock)
+		if (bIsHostVisible && (bIsFirstLock || (LockMode == RLM_WriteOnly_NoOverwrite)))
 		{
 			check(CurrentBufferAlloc.HostPtr);
 			Data = CurrentBufferAlloc.HostPtr;
diff --git a/Engine/Source/Runtime/VulkanRHI/Private/VulkanRHI.cpp b/Engine/Source/Runtime/VulkanRHI/Private/VulkanRHI.cpp
--- a/Engine/Source/Runtime/VulkanRHI/Private/VulkanRHI.cpp
+++ b/Engine/Source/Runtime/VulkanRHI/Private/VulkanRHI.cpp
@@ -854,6 +854,7 @@ void FVulkanDynamicRHI::InitInstance()
 		GSupportsMobileMultiView = Device->GetOptionalExtensions().HasKHRMultiview ? true : false;
 		GRHISupportsMSAAShaderResolve = Device->GetOptionalExtensions().HasQcomRenderPassShaderResolve ? true : false;
 		GRHISupportsRayTracing = RHI_RAYTRACING && RHISupportsRayTracing(GMaxRHIShaderPlatform) && Device->GetOptionalExtensions().HasRaytracingExtensions();
+		GRHIGlobals.SupportsMapWriteNoOverwrite = true;
 
 		GRHIGlobals.EnableHighLevelTransitionPatching = true;
 		GRHIGlobals.NeedExtraTransitions = true;
