VulkanRHI: Fix comment inverted

#rnx

diff --git a/Engine/Source/Runtime/VulkanRHI/Private/VulkanQueue.cpp b/Engine/Source/Runtime/VulkanRHI/Private/VulkanQueue.cpp
--- a/Engine/Source/Runtime/VulkanRHI/Private/VulkanQueue.cpp
+++ b/Engine/Source/Runtime/VulkanRHI/Private/VulkanQueue.cpp
@@ -24,8 +24,8 @@ FAutoConsoleVariableRef CVarVulkanSubmissionAllowTimelineSemaphores(
 	TEXT("r.Vulkan.Submission.AllowTimelineSemaphores"),
 	GAllowTimelineSemaphores,
 	TEXT("If supported, use timeline semaphores for queue submission to reduce API calls.\n")
-	TEXT(" 0: Use timeline semaphores if available (default)\n")
-	TEXT(" 1: Use normal fences\n"),
+	TEXT(" 0: Use normal fences\n")
+	TEXT(" 1: Use timeline semaphores if available (default)\n"),
 	ECVF_ReadOnly
 );
 
