VulkanRHI: use TopOfPipe RHI functions for queries.  Now that SyncPoints guarantee that the query has completed, remove other synchronization mechanisms (events and loops).  Add lock to acquire pools, remove the requirement of running on immediate context.

#rnx

diff --git a/Engine/Source/Runtime/VulkanRHI/Private/VulkanDevice.h b/Engine/Source/Runtime/VulkanRHI/Private/VulkanDevice.h
--- a/Engine/Source/Runtime/VulkanRHI/Private/VulkanDevice.h
+++ b/Engine/Source/Runtime/VulkanRHI/Private/VulkanDevice.h
@@ -591,6 +591,7 @@ private:
 	// Info for formats that are not in the core Vulkan spec (i.e. extensions)
 	mutable TMap<VkFormat, VkFormatProperties> ExtensionFormatProperties;
 
+	FCriticalSection OcclusionQueryPoolLock;
 	TArray<FVulkanOcclusionQueryPool*> UsedOcclusionQueryPools;
 	TArray<FVulkanOcclusionQueryPool*> FreeOcclusionQueryPools;
 
diff --git a/Engine/Source/Runtime/VulkanRHI/Private/VulkanQuery.cpp b/Engine/Source/Runtime/VulkanRHI/Private/VulkanQuery.cpp
--- a/Engine/Source/Runtime/VulkanRHI/Private/VulkanQuery.cpp
+++ b/Engine/Source/Runtime/VulkanRHI/Private/VulkanQuery.cpp
@@ -40,7 +40,6 @@ constexpr uint32 NUM_FRAMES_TO_WAIT_RELEASE_POOL = MAX_uint32; // never release
 FVulkanQueryPool::FVulkanQueryPool(FVulkanDevice& InDevice, FVulkanContextCommon* InCommonContext, uint32 InMaxQueries, VkQueryType InQueryType, bool bInShouldAddReset)
 	: Device(InDevice)
 	, QueryPool(VK_NULL_HANDLE)
-	, ResetEvent(VK_NULL_HANDLE)
 	, MaxQueries(InMaxQueries)
 	, QueryType(InQueryType)
 {
@@ -62,10 +61,6 @@ FVulkanQueryPool::FVulkanQueryPool(FVulkanDevice& InDevice, FVulkanContextCommon
 		InCommonContext->AddQueryPoolForReset(QueryPool, InMaxQueries);
 	}
 
-	VkEventCreateInfo EventCreateInfo;
-	ZeroVulkanStruct(EventCreateInfo, VK_STRUCTURE_TYPE_EVENT_CREATE_INFO);
-	VERIFYVULKANRESULT(VulkanRHI::vkCreateEvent(Device.GetInstanceHandle(), &EventCreateInfo, VULKAN_CPU_ALLOCATOR, &ResetEvent));
-	
 	QueryOutput.AddZeroed(MaxQueries);
 }
 
@@ -73,10 +68,7 @@ FVulkanQueryPool::~FVulkanQueryPool()
 {
 	DEC_DWORD_STAT(STAT_VulkanNumQueryPools);
 	VulkanRHI::vkDestroyQueryPool(Device.GetInstanceHandle(), QueryPool, VULKAN_CPU_ALLOCATOR);
-	VulkanRHI::vkDestroyEvent(Device.GetInstanceHandle(), ResetEvent, VULKAN_CPU_ALLOCATOR);
-	
 	QueryPool = VK_NULL_HANDLE;
-	ResetEvent = VK_NULL_HANDLE;
 }
 
 bool FVulkanOcclusionQueryPool::CanBeReused()
@@ -95,91 +87,21 @@ bool FVulkanOcclusionQueryPool::CanBeReused()
 	return Mask == 0 || (AcquiredIndices[NumWords] & Mask) == 0;
 }
 
-bool FVulkanOcclusionQueryPool::InternalTryGetResults(bool bWait)
+bool FVulkanOcclusionQueryPool::GetResults()
 {
-	check(State == EState::RHIT_PostEndBatch);
-
-	VkResult Result = VK_NOT_READY;
-	if (VulkanRHI::vkGetEventStatus(Device.GetInstanceHandle(), ResetEvent) == VK_EVENT_SET)
+	if (State == EState::RT_PostGetResults)
 	{
-		Result = VulkanRHI::vkGetQueryPoolResults(Device.GetInstanceHandle(), QueryPool, 0, NumUsedQueries, NumUsedQueries * sizeof(uint64), QueryOutput.GetData(), sizeof(uint64), VK_QUERY_RESULT_64_BIT);
-		if (Result == VK_SUCCESS)
-		{
-			State = EState::RT_PostGetResults;
-			return true;
-		}
+		return true;
 	}
 
-	if (Result == VK_NOT_READY)
-	{
-		if (bWait)
-		{
-			FRenderThreadIdleScope IdleScope(ERenderThreadIdleTypes::WaitingForGPUQuery);
-			SCOPE_CYCLE_COUNTER(STAT_VulkanWaitQuery);
-
-			// We'll do manual wait
-			double StartTime = FPlatformTime::Seconds();
-
-			ENamedThreads::Type RenderThread_Local = ENamedThreads::GetRenderThread_Local();
-			bool bSuccess = false;
-			int32 NumLoops = 0;
-			while (!bSuccess)
-			{
-				FPlatformProcess::SleepNoStats(0);
-
-				// pump RHIThread to make sure these queries have actually been submitted to the GPU.
-				if (IsInActualRenderingThread())
-				{
-					FTaskGraphInterface::Get().ProcessThreadUntilIdle(RenderThread_Local);
-				}
-
-				Result = VulkanRHI::vkGetEventStatus(Device.GetInstanceHandle(), ResetEvent);
-				if (Result == VK_EVENT_SET)
-				{
-					Result = VulkanRHI::vkGetQueryPoolResults(Device.GetInstanceHandle(), QueryPool, 0, NumUsedQueries, NumUsedQueries * sizeof(uint64), QueryOutput.GetData(), sizeof(uint64), VK_QUERY_RESULT_64_BIT);
-				}
-
-				if (Result == VK_SUCCESS)
-				{
-					bSuccess = true;
-					break;
-				}
-				else if ((Result == VK_NOT_READY) || (Result == VK_EVENT_RESET))
-				{
-					bSuccess = false;
-				}
-				else
-				{
-					bSuccess = false;
-					VERIFYVULKANRESULT(Result);
-				}
-
-				// timer queries are used for Benchmarks which can stall a bit more
-				const double TimeoutValue = (QueryType == VK_QUERY_TYPE_TIMESTAMP) ? 2.0 : 0.5;
-				// look for gpu stuck/crashed
-				if ((FPlatformTime::Seconds() - StartTime) > TimeoutValue)
-				{
-					if (QueryType == VK_QUERY_TYPE_OCCLUSION)
-					{
-						UE_LOG(LogRHI, Log, TEXT("Timed out while waiting for GPU to catch up on occlusion results. (%.1f s)"), TimeoutValue);
-					}
-					else
-					{
-						UE_LOG(LogRHI, Log, TEXT("Timed out while waiting for GPU to catch up on occlusion/timer results. (%.1f s)"), TimeoutValue);
-					}
-					return false;
-				}
-
-				++NumLoops;
-			}
-
-			State = EState::RT_PostGetResults;
-			return true;
-		}
-	}
-	else
+	if (State == RHIT_PostEndBatch)
 	{
+		check(State == EState::RHIT_PostEndBatch);
+		const VkQueryResultFlags QueryResultFlags = VK_QUERY_RESULT_64_BIT | VK_QUERY_RESULT_WAIT_BIT;
+		const VkResult Result = VulkanRHI::vkGetQueryPoolResults(Device.GetInstanceHandle(), QueryPool, 0, NumUsedQueries, NumUsedQueries * sizeof(uint64), QueryOutput.GetData(), sizeof(uint64), QueryResultFlags);
 		VERIFYVULKANRESULT(Result);
+		State = EState::RT_PostGetResults;
+		return true;
 	}
 
 	return false;
@@ -192,31 +114,31 @@ void FVulkanOcclusionQueryPool::Reset(FVulkanCommandBuffer* InCmdBuffer, uint32
 
 	NumUsedQueries = 0;
 	FrameNumber = InFrameNumber;
-/*
-	bHasResults = false;
-*/
-	VulkanRHI::vkResetEvent(Device.GetInstanceHandle(), ResetEvent);
-	VulkanRHI::vkCmdResetQueryPool(InCmdBuffer->GetHandle(), QueryPool, 0, MaxQueries);
-	
-	// workaround for apparent cache-flush bug in AMD driver implementation of vkCmdResetQueryPool
-	if (Device.GetVendorId() == EGpuVendorId::Amd)
+
+	if (Device.GetOptionalExtensions().HasEXTHostQueryReset)
 	{
-		VkMemoryBarrier barrier;
-		ZeroVulkanStruct(barrier, VK_STRUCTURE_TYPE_MEMORY_BARRIER);
-		barrier.srcAccessMask = VK_ACCESS_TRANSFER_WRITE_BIT;
-		barrier.dstAccessMask = VK_ACCESS_HOST_READ_BIT;
-		VulkanRHI::vkCmdPipelineBarrier(InCmdBuffer->GetHandle(), VK_PIPELINE_STAGE_TRANSFER_BIT, VK_PIPELINE_STAGE_HOST_BIT, 0, 1, &barrier, 0, 0, 0, 0);
+		VulkanRHI::vkResetQueryPoolEXT(GetDevice().GetInstanceHandle(), QueryPool, 0, MaxQueries);
 	}
-	
-	VulkanRHI::vkCmdSetEvent(InCmdBuffer->GetHandle(), ResetEvent, VK_PIPELINE_STAGE_BOTTOM_OF_PIPE_BIT);
-	//FVulkanQueryPool::ResetAll(InCmdBuffer);
+	else
+	{
+		VulkanRHI::vkCmdResetQueryPool(InCmdBuffer->GetHandle(), QueryPool, 0, MaxQueries);
 
+		// workaround for apparent cache-flush bug in AMD driver implementation of vkCmdResetQueryPool
+		if (Device.GetVendorId() == EGpuVendorId::Amd)
+		{
+			VkMemoryBarrier Barrier;
+			ZeroVulkanStruct(Barrier, VK_STRUCTURE_TYPE_MEMORY_BARRIER);
+			Barrier.srcAccessMask = VK_ACCESS_TRANSFER_WRITE_BIT;
+			Barrier.dstAccessMask = VK_ACCESS_HOST_READ_BIT;
+			VulkanRHI::vkCmdPipelineBarrier(InCmdBuffer->GetHandle(), VK_PIPELINE_STAGE_TRANSFER_BIT, VK_PIPELINE_STAGE_HOST_BIT, 0, 1, &Barrier, 0, 0, 0, 0);
+		}
+	}
+	
 	State = EState::RHIT_PostBeginBatch;
 }
 
 void FVulkanCommandListContext::BeginOcclusionQueryBatch(FVulkanCommandBuffer* CmdBuffer, uint32 NumQueriesInBatch)
 {
-	ensure(IsImmediate());
 	checkf(!CurrentOcclusionQueryPool, TEXT("BeginOcclusionQueryBatch called without corresponding EndOcclusionQueryBatch!"));
 	CurrentOcclusionQueryPool = Device.AcquireOcclusionQueryPool(this, NumQueriesInBatch);
 	ensure(CmdBuffer->IsOutsideRenderPass());
@@ -260,6 +182,8 @@ FVulkanOcclusionQueryPool* FVulkanDevice::AcquireOcclusionQueryPool(FVulkanConte
 	NumQueries = FMath::Max(1u, NumQueries);
 	NumQueries = AlignArbitrary(NumQueries, GMinNumberOfQueriesInPool);
 
+	FScopeLock Lock(&OcclusionQueryPoolLock);
+
 	bool bChanged = false;
 	for (int32 Index = UsedOcclusionQueryPools.Num() - 1; Index >= 0; --Index)
 	{
@@ -315,8 +239,9 @@ void FVulkanDevice::ReleaseUnusedOcclusionQueryPools()
 		return;
 	}
 	
-	uint32 ReleaseFrame = (GFrameNumberRenderThread - NUM_FRAMES_TO_WAIT_RELEASE_POOL);
+	const uint32 ReleaseFrame = (GFrameNumberRenderThread - NUM_FRAMES_TO_WAIT_RELEASE_POOL);
 
+	FScopeLock Lock(&OcclusionQueryPoolLock);
 	for (int32 Index = FreeOcclusionQueryPools.Num() - 1; Index >= 0; --Index)
 	{
 		FVulkanOcclusionQueryPool* Pool = FreeOcclusionQueryPools[Index];
@@ -330,7 +255,6 @@ void FVulkanDevice::ReleaseUnusedOcclusionQueryPools()
 
 void FVulkanCommandListContext::EndOcclusionQueryBatch()
 {
-	ensure(IsImmediate());
 	checkf(CurrentOcclusionQueryPool, TEXT("EndOcclusionQueryBatch called without corresponding BeginOcclusionQueryBatch!"));
 	CurrentOcclusionQueryPool->EndBatch();
 	CurrentOcclusionQueryPool = nullptr;
@@ -413,21 +337,27 @@ FRenderQueryRHIRef FVulkanDynamicRHI::RHICreateRenderQuery(ERenderQueryType Quer
 
 bool FVulkanDynamicRHI::RHIGetRenderQueryResult(FRHIRenderQuery* QueryRHI, uint64& OutNumPixels, bool bWait, uint32 GPUIndex)
 {
-	auto ToMicroseconds = [](uint64 Timestamp) -> uint64
+	FVulkanRenderQuery* BaseQuery = ResourceCast(QueryRHI);
+
+	if (!ensureMsgf(BaseQuery->SyncPoint, TEXT("Attempt to get result data for an FRHIRenderQuery that was never used in a command list.")))
 	{
-#if RHI_NEW_GPU_PROFILER
-		checkNoEntry(); // @todo - new gpu profiler
-		return 0;
-#else
-		const double Frequency = double(FVulkanGPUTiming::GetTimingFrequency());
-		uint64 Microseconds = (uint64)((double(Timestamp) / Frequency) * 1000.0 * 1000.0);
-		return Microseconds;
-#endif
-	};
+		OutNumPixels = 0;
+		return false;
+	}
 
-	FRHICommandListImmediate& RHICmdList = FRHICommandListImmediate::Get();
+	if (!BaseQuery->SyncPoint->IsComplete())
+	{
+		if (bWait)
+		{
+			FRenderThreadIdleScope IdleScope(ERenderThreadIdleTypes::WaitingForGPUQuery);
+			ProcessInterruptQueueUntil(BaseQuery->SyncPoint);
+		}
+		else
+		{
+			return false;
+		}
+	}
 
-	FVulkanRenderQuery* BaseQuery = ResourceCast(QueryRHI);
 	if (BaseQuery->QueryType == RQT_Occlusion)
 	{
 		FVulkanOcclusionQuery* Query = static_cast<FVulkanOcclusionQuery*>(BaseQuery);
@@ -444,7 +374,7 @@ bool FVulkanDynamicRHI::RHIGetRenderQueryResult(FRHIRenderQuery* QueryRHI, uint6
 		}
 
 		ensure(Query->State == FVulkanOcclusionQuery::EState::RHI_PostEnd);
-		if (Query->Pool->TryGetResults(bWait))
+		if (Query->Pool->GetResults())
 		{
 			Query->Result = Query->Pool->GetResultValue(Query->IndexInPool);
 			Query->ReleaseFromPool();
@@ -455,21 +385,27 @@ bool FVulkanDynamicRHI::RHIGetRenderQueryResult(FRHIRenderQuery* QueryRHI, uint6
 	}
 	else if (BaseQuery->QueryType == RQT_AbsoluteTime)
 	{
-		FVulkanTimingQuery* Query = static_cast<FVulkanTimingQuery*>(BaseQuery);
-		if (Query->Pool == nullptr)
+		auto ToMicroseconds = [](uint64 Timestamp) -> uint64
 		{
-			// RHIEndRenderQuery hasn't been executed yet for this query, so it's not ready yet. If it's a non-blocking wait, we can return false right away.
-			if (!bWait)
-			{
-				return false;
-			}
+#if RHI_NEW_GPU_PROFILER
+			checkNoEntry(); // @todo - new gpu profiler
+			return 0;
+#else
+			const double Frequency = double(FVulkanGPUTiming::GetTimingFrequency());
+			uint64 Microseconds = (uint64)((double(Timestamp) / Frequency) * 1000.0 * 1000.0);
+			return Microseconds;
+#endif
+		};
 
+		FVulkanTimingQuery* Query = static_cast<FVulkanTimingQuery*>(BaseQuery);
+		if ((Query->Pool == nullptr) && bWait)
+		{
 			// Blocking was requested, so, we need to wait for the RHI thread to catch up.
-			RHICmdList.ImmediateFlush(EImmediateFlushType::FlushRHIThread);
+			FRHICommandListImmediate::Get().ImmediateFlush(EImmediateFlushType::FlushRHIThread);
 			check(Query->Pool != nullptr);
 		}
 
-		check(Query->Pool->CurrentTimestamp < Query->Pool->BufferSize);
+		check(Query->Pool && (Query->Pool->CurrentTimestamp < Query->Pool->BufferSize));
 		int32 TimestampIndex = Query->Pool->CurrentTimestamp;
 		if (!bWait)
 		{
@@ -523,16 +459,75 @@ bool FVulkanDynamicRHI::RHIGetRenderQueryResult(FRHIRenderQuery* QueryRHI, uint6
 	return false;
 }
 
+void FVulkanDynamicRHI::RHIEndRenderQuery_TopOfPipe(FRHICommandListBase& RHICmdList, FRHIRenderQuery* RenderQuery)
+{
+	const uint32 GPUIndex = 0;
+
+	FVulkanRenderQuery* Query = ResourceCast(RenderQuery);
+	auto& QueryBatchData = RHICmdList.GetQueryBatchData(Query->QueryType);
+
+	if (QueryBatchData[GPUIndex])
+	{
+		// This query belongs to a batch. Use the sync point we created earlier
+		Query->SyncPoint = static_cast<FVulkanSyncPoint*>(QueryBatchData[GPUIndex]);
+	}
+	else
+	{
+		// Queries issued outside of a batch use one sync point per query.
+//		check(Query->SyncPoint == nullptr);
+		Query->SyncPoint = CreateVulkanSyncPoint();
+
+		RHICmdList.EnqueueLambda([SyncPoint = Query->SyncPoint](FRHICommandListBase& ExecutingCmdList)
+		{
+			FVulkanCommandListContext& Context = FVulkanCommandListContext::Get(ExecutingCmdList);
+			Context.AddPendingSyncPoint(SyncPoint);
+		});
+	}
+
+	// Enqueue the RHI command to record the EndQuery() call on the context.
+	FDynamicRHI::RHIEndRenderQuery_TopOfPipe(RHICmdList, RenderQuery);
+}
+
+void FVulkanDynamicRHI::RHIBeginRenderQueryBatch_TopOfPipe(FRHICommandListBase& RHICmdList, ERenderQueryType QueryType)
+{
+	// Each query batch uses a single sync point to signal when the results are ready.
+	const uint32 GPUIndex = 0;
+
+	auto& QueryBatchData = RHICmdList.GetQueryBatchData(QueryType);
+	checkf(QueryBatchData[GPUIndex] == nullptr, TEXT("A query batch for this type has already begun on this command list."));
+
+	FVulkanSyncPointRef SyncPoint = CreateVulkanSyncPoint();
+
+	// Keep a reference in the RHI command list, so we can retrieve it later in BeginQuery/EndQuery/EndBatch.
+	QueryBatchData[GPUIndex] = SyncPoint.GetReference();
+	SyncPoint->AddRef();
+}
+
+void FVulkanDynamicRHI::RHIEndRenderQueryBatch_TopOfPipe(FRHICommandListBase& RHICmdList, ERenderQueryType QueryType)
+{
+	const uint32 GPUIndex = 0;
+
+	auto& QueryBatchData = RHICmdList.GetQueryBatchData(QueryType);
+	checkf(QueryBatchData[GPUIndex], TEXT("A query batch for this type is not open on this command list."));
+
+	FVulkanSyncPointRef SyncPoint = static_cast<FVulkanSyncPoint*>(QueryBatchData[GPUIndex]);
+
+	// Clear the sync point reference on the RHI command list
+	SyncPoint->Release();
+	QueryBatchData[GPUIndex] = nullptr;
+
+	RHICmdList.EnqueueLambda([SyncPoint = MoveTemp(SyncPoint), QueryType](FRHICommandListBase& ExecutingCmdList)
+	{
+		FVulkanCommandListContext& Context = FVulkanCommandListContext::Get(ExecutingCmdList);
+		Context.AddPendingSyncPoint(SyncPoint);
+	});
+}
+
 void FVulkanCommandListContext::RHIBeginRenderQuery(FRHIRenderQuery* QueryRHI)
 {
 	FVulkanRenderQuery* BaseQuery = ResourceCast(QueryRHI);
 	if (BaseQuery->QueryType == RQT_Occlusion)
 	{
-		//#todo-rco: Temp until we get the merge straightened out
-		if (!CurrentOcclusionQueryPool)
-		{
-			return;
-		}
 		check(CurrentOcclusionQueryPool);
 		FVulkanOcclusionQuery* Query = static_cast<FVulkanOcclusionQuery*>(BaseQuery);
 		if (Query->State == FVulkanOcclusionQuery::EState::RHI_PostEnd)
@@ -565,12 +560,7 @@ void FVulkanCommandListContext::RHIEndRenderQuery(FRHIRenderQuery* QueryRHI)
 	FVulkanRenderQuery* BaseQuery = ResourceCast(QueryRHI);
 	if (BaseQuery->QueryType == RQT_Occlusion)
 	{
-		//#todo-rco: Temp until we get the merge straightened out
-		if (!CurrentOcclusionQueryPool)
-		{
-			return;
-		}
-
+		check(CurrentOcclusionQueryPool);
 		FVulkanOcclusionQuery* Query = static_cast<FVulkanOcclusionQuery*>(BaseQuery);
 		ensure(Query->State == FVulkanOcclusionQuery::EState::RHI_PostBegin);
 		Query->State = FVulkanOcclusionQuery::EState::RHI_PostEnd;
diff --git a/Engine/Source/Runtime/VulkanRHI/Private/VulkanRenderTarget.cpp b/Engine/Source/Runtime/VulkanRHI/Private/VulkanRenderTarget.cpp
--- a/Engine/Source/Runtime/VulkanRHI/Private/VulkanRenderTarget.cpp
+++ b/Engine/Source/Runtime/VulkanRHI/Private/VulkanRenderTarget.cpp
@@ -580,15 +580,17 @@ void FVulkanCommandListContext::RHIBeginRenderPass(const FRHIRenderPassInfo& InI
 
 void FVulkanCommandListContext::RHIEndRenderPass()
 {
-	const bool bHasOcclusionQueries = (RenderPassInfo.NumOcclusionQueries > 0);
+	FVulkanCommandBuffer& CommandBuffer = GetCommandBuffer();
+	Device.GetRenderPassManager().EndRenderPass(&CommandBuffer);
 
+	const bool bHasOcclusionQueries = (RenderPassInfo.NumOcclusionQueries > 0);
 	if (bHasOcclusionQueries)
 	{
 		EndOcclusionQueryBatch();
+
+		// Force the syncs points to be signaled right after the render pass containing the queries
+		FlushPendingSyncPoints();
 	}
-	
-	FVulkanCommandBuffer& CommandBuffer = GetCommandBuffer();
-	Device.GetRenderPassManager().EndRenderPass(&CommandBuffer);
 
 	check(CurrentRenderPass);
 	CurrentRenderPass = nullptr;
diff --git a/Engine/Source/Runtime/VulkanRHI/Private/VulkanResources.h b/Engine/Source/Runtime/VulkanRHI/Private/VulkanResources.h
--- a/Engine/Source/Runtime/VulkanRHI/Private/VulkanResources.h
+++ b/Engine/Source/Runtime/VulkanRHI/Private/VulkanResources.h
@@ -801,7 +801,6 @@ public:
 protected:
 	FVulkanDevice& Device;
 	VkQueryPool QueryPool;
-	VkEvent ResetEvent;
 	const uint32 MaxQueries;
 	const VkQueryType QueryType;
 	TArray<uint64> QueryOutput;
@@ -852,20 +851,7 @@ public:
 
 	bool CanBeReused();
 
-	inline bool TryGetResults(bool bWait)
-	{
-		if (State == RT_PostGetResults)
-		{
-			return true;
-		}
-
-		if (State == RHIT_PostEndBatch)
-		{
-			return InternalTryGetResults(bWait);
-		}
-
-		return false;
-	}
+	bool GetResults();
 
 	void Reset(FVulkanCommandBuffer* InCmdBuffer, uint32 InFrameNumber);
 
@@ -888,7 +874,6 @@ protected:
 	uint32 NumUsedQueries = 0;
 	TArray<FVulkanOcclusionQuery*> AllocatedQueries;
 	TArray<uint64> AcquiredIndices;
-	bool InternalTryGetResults(bool bWait);
 
 	uint32 FrameNumber = UINT32_MAX;
 };
@@ -933,6 +918,8 @@ public:
 	uint64 Result = 0;
 
 	uint32 IndexInPool = UINT32_MAX;
+
+	FVulkanSyncPointRef SyncPoint;
 };
 
 class FVulkanOcclusionQuery : public FVulkanRenderQuery
diff --git a/Engine/Source/Runtime/VulkanRHI/Public/VulkanDynamicRHI.h b/Engine/Source/Runtime/VulkanRHI/Public/VulkanDynamicRHI.h
--- a/Engine/Source/Runtime/VulkanRHI/Public/VulkanDynamicRHI.h
+++ b/Engine/Source/Runtime/VulkanRHI/Public/VulkanDynamicRHI.h
@@ -170,6 +170,9 @@ public:
 	virtual void RHIReadSurfaceFloatData(FRHITexture* Texture, FIntRect Rect, TArray<FFloat16Color>& OutData, ECubeFace CubeFace, int32 ArrayIndex, int32 MipIndex) final override;
 	virtual void RHIRead3DSurfaceFloatData(FRHITexture* Texture, FIntRect Rect, FIntPoint ZMinMax, TArray<FFloat16Color>& OutData) final override;
 	virtual FRenderQueryRHIRef RHICreateRenderQuery(ERenderQueryType QueryType) final override;
+	virtual void RHIEndRenderQuery_TopOfPipe(FRHICommandListBase& RHICmdList, FRHIRenderQuery* RenderQuery) final override;
+	virtual void RHIBeginRenderQueryBatch_TopOfPipe(FRHICommandListBase& RHICmdList, ERenderQueryType QueryType) final override;
+	virtual void RHIEndRenderQueryBatch_TopOfPipe(FRHICommandListBase& RHICmdList, ERenderQueryType QueryType) final override;
 	virtual bool RHIGetRenderQueryResult(FRHIRenderQuery* RenderQuery, uint64& OutResult, bool bWait, uint32 GPUIndex = INDEX_NONE) final override;
 	virtual FTextureRHIRef RHIGetViewportBackBuffer(FRHIViewport* Viewport) final override;
 	virtual void RHIAliasTextureResources(FTextureRHIRef& DestTexture, FTextureRHIRef& SrcTexture) final override;
