Fixed: Visual stuttering and graphical corruption on Android Vulkan devices
VulkanRHI was not correctly waiting for rendering to finish before presenting an image

diff --git a/Engine/Source/Runtime/VulkanRHI/Private/VulkanViewport.cpp b/Engine/Source/Runtime/VulkanRHI/Private/VulkanViewport.cpp
--- a/Engine/Source/Runtime/VulkanRHI/Private/VulkanViewport.cpp
+++ b/Engine/Source/Runtime/VulkanRHI/Private/VulkanViewport.cpp
@@ -1013,6 +1013,9 @@ bool FVulkanViewport::Present(FVulkanCommandListContext& Context, FVulkanQueue*
 			// swapchain can go out of date, do not crash at this point
 			if (LIKELY(TryAcquireImageIndex()))
 			{
+				// Wait for semaphore signal before writing to backbuffer image
+				Context.AddWaitSemaphore(VK_PIPELINE_STAGE_BOTTOM_OF_PIPE_BIT, AcquiredSemaphore);
+				
 				uint32 WindowSizeX = FMath::Min(SizeX, SwapChain->InternalWidth);
 				uint32 WindowSizeY = FMath::Min(SizeY, SwapChain->InternalHeight);
 
@@ -1041,11 +1044,6 @@ bool FVulkanViewport::Present(FVulkanCommandListContext& Context, FVulkanQueue*
 
 		if (LIKELY(!bFailedToDelayAcquireBackbuffer))
 		{
-			if (GVulkanDelayAcquireImage == EDelayAcquireImageType::DelayAcquire)
-			{
-				Context.AddWaitSemaphore(VK_PIPELINE_STAGE_BOTTOM_OF_PIPE_BIT, AcquiredSemaphore);
-			}
-			
 			if (AcquiredImageIndex >= 0)
 			{
 				Context.AddSignalSemaphore(RenderingDoneSemaphores[AcquiredImageIndex]);
