Fixing VulkanBuffer support of ByteAddressBuffer. If a buffer is created with *only* the ByteAddressBuffer flag, it would fail to calculate its size and offset since it didn't handle the conversion to the Vulkan flags.

diff --git a/Engine/Source/Runtime/VulkanRHI/Private/VulkanBuffer.cpp b/Engine/Source/Runtime/VulkanRHI/Private/VulkanBuffer.cpp
--- a/Engine/Source/Runtime/VulkanRHI/Private/VulkanBuffer.cpp
+++ b/Engine/Source/Runtime/VulkanRHI/Private/VulkanBuffer.cpp
@@ -84,7 +84,7 @@ VkBufferUsageFlags FVulkanBuffer::UEToVKBufferUsageFlags(FVulkanDevice* InDevice
 
 	TranslateFlag(BUF_VertexBuffer, VK_BUFFER_USAGE_VERTEX_BUFFER_BIT);
 	TranslateFlag(BUF_IndexBuffer, VK_BUFFER_USAGE_INDEX_BUFFER_BIT);
-	TranslateFlag(BUF_StructuredBuffer, VK_BUFFER_USAGE_STORAGE_BUFFER_BIT);
+	TranslateFlag(BUF_ByteAddressBuffer|BUF_StructuredBuffer, VK_BUFFER_USAGE_STORAGE_BUFFER_BIT);
 	TranslateFlag(BUF_UniformBuffer, VK_BUFFER_USAGE_UNIFORM_BUFFER_BIT);
 	TranslateFlag(BUF_AccelerationStructure, VK_BUFFER_USAGE_ACCELERATION_STRUCTURE_STORAGE_BIT_KHR);
 
