Vulkan RHI Texture cleanup
- Moving texture creation code to support creating the FRHITexture before the internal Vulkan texture. This also separates the FRHITexture creation from the RHICmdList usage.
- Moving duplicated constructor initializers to the headers.

diff --git a/Engine/Source/Runtime/VulkanRHI/Private/VulkanResources.h b/Engine/Source/Runtime/VulkanRHI/Private/VulkanResources.h
--- a/Engine/Source/Runtime/VulkanRHI/Private/VulkanResources.h
+++ b/Engine/Source/Runtime/VulkanRHI/Private/VulkanResources.h
@@ -575,15 +575,11 @@ enum class EImageOwnerType : uint8
 
 class FVulkanTexture : public FRHITexture, public FVulkanEvictable, public FVulkanViewableResource
 {
+public:
 	// Constructor without command list: targets will not be cleared, cannot contain initial data, will be left in VK_IMAGE_LAYOUT_UNDEFINED.
-	// Only used by transient texture allocator.
-	friend class FVulkanTransientResourceAllocator;
+	// Layout changes and initial data upload should be handled by FinalizeCreateTextureInternal
 	FVulkanTexture(FVulkanDevice& InDevice, const FRHITextureCreateDesc& InCreateDesc, const FRHITransientHeapAllocation* InTransientHeapAllocation);
 
-public:
-	// Regular constructor.
-	FVulkanTexture(FRHICommandListBase& RHICmdList, FVulkanDevice& InDevice, const FRHITextureCreateDesc& InCreateDesc, const FRHITransientHeapAllocation* InTransientHeapAllocation);
-
 	// Construct from external resource.
 	FVulkanTexture(FVulkanDevice& InDevice, const FRHITextureCreateDesc& InCreateDesc, VkImage InImage, const FVulkanRHIExternalImageDeleteCallbackInfo& InExternalImageDeleteCallbackInfo);
 
@@ -601,6 +597,7 @@ public:
 
 	// View with all mips/layers
 	FVulkanView* DefaultView = nullptr;
+
 	// View with all mips/layers, but if it's a Depth/Stencil, only the Depth view
 	FVulkanView* PartialView = nullptr;
 
@@ -662,7 +659,7 @@ public:
 	/**
 	 * Returns how much memory is used by the surface
 	 */
-	inline uint32 GetMemorySize() const
+	uint32 GetMemorySize() const
 	{
 		return MemoryRequirements.size;
 	}
@@ -677,14 +674,17 @@ public:
 	*/
 	void GetMipSize(uint32 MipIndex, uint64& MipBytes);
 
-	inline VkImageViewType GetViewType() const
+	VkImageViewType GetViewType() const
 	{
 		return UETextureDimensionToVkImageViewType(GetDesc().Dimension);
 	}
 
-	inline VkImageTiling GetTiling() const { return Tiling; }
+	VkImageTiling GetTiling() const
+	{
+		return Tiling;
+	}
 
-	inline uint32 GetNumberOfArrayLevels() const
+	uint32 GetNumberOfArrayLevels() const
 	{
 		switch (GetViewType())
 		{
@@ -706,38 +706,43 @@ public:
 	VULKANRHI_API void ErrorInvalidViewType() const;
 
 	// Full includes Depth+Stencil
-	inline VkImageAspectFlags GetFullAspectMask() const
+	VkImageAspectFlags GetFullAspectMask() const
 	{
 		return FullAspectMask;
 	}
 
 	// Only Depth or Stencil
-	inline VkImageAspectFlags GetPartialAspectMask() const
+	VkImageAspectFlags GetPartialAspectMask() const
 	{
 		return PartialAspectMask;
 	}
 
-	inline bool IsDepthOrStencilAspect() const
+	bool IsDepthOrStencilAspect() const
 	{
 		return (FullAspectMask & (VK_IMAGE_ASPECT_DEPTH_BIT | VK_IMAGE_ASPECT_STENCIL_BIT)) != 0;
 	}
 
-	inline bool IsImageOwner() const
+	bool IsImageOwner() const
 	{
 		return (ImageOwnerType == EImageOwnerType::LocalOwner);
 	}
 
-	inline bool SupportsSampling() const
+	bool SupportsSampling() const
 	{
 		return EnumHasAllFlags(GPixelFormats[GetDesc().Format].Capabilities, EPixelFormatCapabilities::TextureSample) &&
 			((ImageUsageFlags & VK_IMAGE_USAGE_SAMPLED_BIT) != 0);
 	}
 
-	inline VkImageLayout GetDefaultLayout() const
+	VkImageLayout GetDefaultLayout() const
 	{
 		return DefaultLayout;
 	}
 
+	void SetDefaultLayout(VkImageLayout InDefaultLayout)
+	{
+		DefaultLayout = InDefaultLayout;
+	}
+
 	VULKANRHI_API VkDeviceMemory GetAllocationHandle() const;
 	VULKANRHI_API uint64 GetAllocationOffset() const;
 
@@ -747,26 +752,32 @@ public:
 
 	virtual void UpdateLinkedViews() override;
 
-	FVulkanDevice* Device;
-	VkImage Image;
-	VkImageUsageFlags ImageUsageFlags;
-	VkFormat StorageFormat;  // Removes SRGB if requested, used to upload data
-	VkFormat ViewFormat;  // Format for SRVs, render targets
-	VkMemoryPropertyFlags MemProps;
+	FVulkanDevice* const Device;
+	VkImage Image = VK_NULL_HANDLE;
+	VkImageUsageFlags ImageUsageFlags = 0;
+	
+	// Removes SRGB if requested, used to upload data
+	VkFormat StorageFormat = VK_FORMAT_UNDEFINED;
+
+	// Format for SRVs, render targets
+	VkFormat ViewFormat = VK_FORMAT_UNDEFINED;
+
+	VkMemoryPropertyFlags MemProps = VK_MEMORY_PROPERTY_DEVICE_LOCAL_BIT;
 	VkMemoryRequirements MemoryRequirements;
 
 	FVulkanRHIExternalImageDeleteCallbackInfo ExternalImageDeleteCallbackInfo;
 
-private:
 	void SetInitialImageState(FRHICommandListBase& RHICmdList, VkImageLayout InitialLayout, bool bClear, const FClearValueBinding& ClearValueBinding, bool bIsTransientResource);
+
+private:
 	void InternalMoveSurface(FVulkanDevice& InDevice, FVulkanCommandListContext& Context, VulkanRHI::FVulkanAllocation& DestAllocation);
 
-	VkImageTiling Tiling;
+	VkImageTiling Tiling = VK_IMAGE_TILING_MAX_ENUM;
 	VulkanRHI::FVulkanAllocation Allocation;
-	VkImageAspectFlags FullAspectMask;
-	VkImageAspectFlags PartialAspectMask;
-	FVulkanCpuReadbackBuffer* CpuReadbackBuffer;
-	VkImageLayout DefaultLayout;
+	VkImageAspectFlags FullAspectMask = 0;
+	VkImageAspectFlags PartialAspectMask = 0;
+	FVulkanCpuReadbackBuffer* CpuReadbackBuffer = nullptr;
+	VkImageLayout DefaultLayout = VK_IMAGE_LAYOUT_UNDEFINED;
 
 protected:
 	EImageOwnerType ImageOwnerType;
diff --git a/Engine/Source/Runtime/VulkanRHI/Private/VulkanTexture.cpp b/Engine/Source/Runtime/VulkanRHI/Private/VulkanTexture.cpp
--- a/Engine/Source/Runtime/VulkanRHI/Private/VulkanTexture.cpp
+++ b/Engine/Source/Runtime/VulkanRHI/Private/VulkanTexture.cpp
@@ -976,10 +976,107 @@ void FVulkanDynamicRHI::RHIUpdateTextureReference(FRHICommandListBase& RHICmdLis
 	2D texture support.
 -----------------------------------------------------------------------------*/
 
+FVulkanDynamicRHI::FCreateTextureResult FVulkanDynamicRHI::BeginCreateTextureInternal(const FRHITextureCreateDesc& CreateDesc, const FRHITransientHeapAllocation* InTransientHeapAllocation)
+{
+	LLM_SCOPE_VULKAN(GetMemoryTagForTextureFlags(CreateDesc.Flags));
+
+	FVulkanTexture* Texture = new FVulkanTexture(*Device, CreateDesc, InTransientHeapAllocation);
+
+	const bool bIsTransientResource = (InTransientHeapAllocation != nullptr);
+	const bool bDoInitialClear = VKHasAnyFlags(Texture->ImageUsageFlags, VK_IMAGE_USAGE_SAMPLED_BIT) &&
+		EnumHasAnyFlags(CreateDesc.Flags, TexCreate_RenderTargetable | TexCreate_DepthStencilTargetable) &&
+		!bIsTransientResource;
+
+	FCreateTextureResult Result{};
+	Result.Texture = Texture;
+	Result.DefaultLayout = Texture->GetDefaultLayout();
+	Result.bTransientResource = bIsTransientResource;
+	Result.bClear = bDoInitialClear;
+
+	return Result;
+}
+
+FVulkanTexture* FVulkanDynamicRHI::FinalizeCreateTextureInternal(FRHICommandListBase& RHICmdList, FCreateTextureResult CreateResult, TConstArrayView<uint8> InitialData)
+{
+	FVulkanTexture* Texture = CreateResult.Texture;
+	const FRHITextureDesc& Desc = Texture->GetDesc();
+
+	if (!EnumHasAnyFlags(Desc.Flags, TexCreate_CPUReadback))
+	{
+		if (CreateResult.DefaultLayout != VK_IMAGE_LAYOUT_UNDEFINED || CreateResult.bClear)
+		{
+			Texture->SetInitialImageState(RHICmdList, CreateResult.DefaultLayout, CreateResult.bClear, Desc.ClearValue, CreateResult.bTransientResource);
+		}
+
+		if (InitialData.Num())
+		{
+			// InternalLockWrite leaves the image in VK_IMAGE_LAYOUT_SHADER_READ_ONLY_OPTIMAL, so make sure the requested resource state is SRV.
+			//check(EnumHasAnyFlags(InCreateDesc.InitialState, ERHIAccess::SRVMask));
+			Texture->SetDefaultLayout(VK_IMAGE_LAYOUT_SHADER_READ_ONLY_OPTIMAL);
+
+			// Transfer bulk data
+			VulkanRHI::FStagingBuffer* StagingBuffer = Device->GetStagingManager().AcquireBuffer(InitialData.NumBytes());
+			void* Data = StagingBuffer->GetMappedPointer();
+
+			// Do copy
+			FMemory::Memcpy(Data, InitialData.GetData(), InitialData.NumBytes());
+
+			VkBufferImageCopy Region;
+			FMemory::Memzero(Region);
+			//#todo-rco: Use real Buffer offset when switching to suballocations!
+			Region.bufferOffset = 0;
+			Region.bufferRowLength = Desc.Extent.X;
+			Region.bufferImageHeight = Desc.Extent.Y;
+
+			Region.imageSubresource.mipLevel = 0;
+			Region.imageSubresource.baseArrayLayer = 0;
+			Region.imageSubresource.layerCount = Texture->GetNumberOfArrayLevels();
+			Region.imageSubresource.aspectMask = Texture->GetFullAspectMask();
+
+			Region.imageExtent.width = Region.bufferRowLength;
+			Region.imageExtent.height = Region.bufferImageHeight;
+			Region.imageExtent.depth = Desc.Depth;
+
+			RHICmdList.EnqueueLambda(TEXT("FVulkanTexture::InternalLockWrite"),
+				[Texture, Region, StagingBuffer](FRHICommandListBase& ExecutingCmdList)
+				{
+					FVulkanTexture::InternalLockWrite(FVulkanUploadContext::Get(ExecutingCmdList), Texture, Region, StagingBuffer);
+				});
+		}
+	}
+
+	return Texture;
+}
+
+FVulkanTexture* FVulkanDynamicRHI::CreateTextureInternal(FRHICommandListBase& RHICmdList, const FRHITextureCreateDesc& CreateDesc, TConstArrayView<uint8> InitialData)
+{
+	FCreateTextureResult CreateResult = BeginCreateTextureInternal(CreateDesc, nullptr);
+	return FinalizeCreateTextureInternal(RHICmdList, CreateResult, InitialData);
+}
+
+FVulkanTexture* FVulkanDynamicRHI::CreateTextureInternal(const FRHITextureCreateDesc& CreateDesc, const FRHITransientHeapAllocation& InTransientHeapAllocation)
+{
+	FCreateTextureResult CreateResult = BeginCreateTextureInternal(CreateDesc, &InTransientHeapAllocation);
+	return CreateResult.Texture;
+}
+
 FTextureRHIRef FVulkanDynamicRHI::RHICreateTexture(FRHICommandListBase& RHICmdList, const FRHITextureCreateDesc& CreateDesc)
 {
 	LLM_SCOPE_VULKAN(GetMemoryTagForTextureFlags(CreateDesc.Flags));
-	return new FVulkanTexture(RHICmdList, *Device, CreateDesc, nullptr);
+	TConstArrayView<uint8> InitialData;
+	if (CreateDesc.BulkData)
+	{
+		InitialData = CreateDesc.BulkData->GetBulkDataView<uint8>();
+	}
+
+	FVulkanTexture* Texture = CreateTextureInternal(RHICmdList, CreateDesc, InitialData);
+
+	if (CreateDesc.BulkData)
+	{
+		CreateDesc.BulkData->Discard();
+	}
+
+	return Texture;
 }
 
 FTextureRHIRef FVulkanDynamicRHI::RHIAsyncCreateTexture2D(uint32 SizeX,uint32 SizeY,uint8 Format,uint32 NumMips,ETextureCreateFlags Flags, ERHIAccess InResourceState,void** InitialMipData,uint32 NumInitialMips, const TCHAR* DebugName, FGraphEventRef& OutCompletionEvent)
@@ -1063,7 +1160,7 @@ FTextureRHIRef FVulkanDynamicRHI::AsyncReallocateTexture2D_RenderThread(FRHIComm
 	FVulkanTexture* OldTexture = ResourceCast(OldTextureRHI);
 	const FRHITextureDesc& OldDesc = OldTexture->GetDesc();
 
-	const FRHITextureCreateDesc Desc =
+	const FRHITextureCreateDesc CreateDesc =
 		FRHITextureCreateDesc::Create2D(TEXT("AsyncReallocateTexture2D_RenderThread"), NewSizeX, NewSizeY, OldDesc.Format)
 		.SetClearValue(OldDesc.ClearValue)
 		.SetFlags(OldDesc.Flags)
@@ -1071,7 +1168,7 @@ FTextureRHIRef FVulkanDynamicRHI::AsyncReallocateTexture2D_RenderThread(FRHIComm
 		.SetNumSamples(OldDesc.NumSamples)
 		.DetermineInititialState();
 
-	FVulkanTexture* NewTexture = new FVulkanTexture(RHICmdList, *Device, Desc, nullptr);
+	FVulkanTexture* NewTexture = CreateTextureInternal(RHICmdList, CreateDesc, {});
 
 	RHICmdList.EnqueueLambda(TEXT("AsyncReallocateTexture2D"), [OldTexture, NewTexture, NewMipCount, NewSizeX, NewSizeY, RequestStatus](FRHICommandListImmediate& ImmCmdList)
 	{
@@ -1317,23 +1414,11 @@ void FVulkanDynamicRHI::InternalUpdateTexture3D(FRHICommandListBase& RHICmdList,
 
 FVulkanTexture::FVulkanTexture(FVulkanDevice& InDevice, const FRHITextureCreateDesc& InCreateDesc, const FRHITransientHeapAllocation* InTransientHeapAllocation)
 	: FRHITexture(InCreateDesc)
-	, PartialView(nullptr)
 	, Device(&InDevice)
-	, Image(VK_NULL_HANDLE)
-	, ImageUsageFlags(0)
-	, StorageFormat(VK_FORMAT_UNDEFINED)
-	, ViewFormat(VK_FORMAT_UNDEFINED)
-	, MemProps(VK_MEMORY_PROPERTY_DEVICE_LOCAL_BIT)
-	, Tiling(VK_IMAGE_TILING_MAX_ENUM)	// Can be expanded to a per-platform definition
-	, FullAspectMask(0)
-	, PartialAspectMask(0)
-	, CpuReadbackBuffer(nullptr) // for readback textures we use a staging buffer. this is because vulkan only requires implementations to support 1 mip level(which is useless), so we emulate using a buffer
-	, DefaultLayout(VK_IMAGE_LAYOUT_UNDEFINED)
+	, ImageOwnerType(EImageOwnerType::LocalOwner)
 {
 	VULKAN_TRACK_OBJECT_CREATE(FVulkanTexture, this);
 
-	checkf(InCreateDesc.BulkData == nullptr, TEXT("Must call the constructor with a command list to initialize a texture with data."));
-
 	if (EnumHasAnyFlags(InCreateDesc.Flags, TexCreate_CPUReadback))
 	{
 		check(InCreateDesc.NumSamples == 1); //not implemented
@@ -1372,8 +1457,6 @@ FVulkanTexture::FVulkanTexture(FVulkanDevice& InDevice, const FRHITextureCreateD
 		return;
 	}
 
-	ImageOwnerType = EImageOwnerType::LocalOwner;
-
 	FImageCreateInfo ImageCreateInfo;
 	FVulkanTexture::GenerateImageCreateInfo(ImageCreateInfo, InDevice, InCreateDesc, &StorageFormat, &ViewFormat);
 
@@ -1508,75 +1591,11 @@ FVulkanTexture::FVulkanTexture(FVulkanDevice& InDevice, const FRHITextureCreateD
 	}
 }
 
-FVulkanTexture::FVulkanTexture(FRHICommandListBase& RHICmdList, FVulkanDevice& InDevice, const FRHITextureCreateDesc& InCreateDesc, const FRHITransientHeapAllocation* InTransientHeapAllocation)
-	: FVulkanTexture(InDevice, FRHITextureCreateDesc(InCreateDesc, InCreateDesc.InitialState, InCreateDesc.DebugName), InTransientHeapAllocation)
-{
-	if (EnumHasAnyFlags(InCreateDesc.Flags, TexCreate_CPUReadback))
-	{
-		return;
-	}
-
-	const bool bIsTransientResource = (InTransientHeapAllocation != nullptr);
-	const bool bDoInitialClear = VKHasAnyFlags(ImageUsageFlags, VK_IMAGE_USAGE_SAMPLED_BIT) &&
-		EnumHasAnyFlags(InCreateDesc.Flags, TexCreate_RenderTargetable | TexCreate_DepthStencilTargetable) &&
-		!bIsTransientResource;
-	if (DefaultLayout != VK_IMAGE_LAYOUT_UNDEFINED || bDoInitialClear)
-	{
-		SetInitialImageState(RHICmdList, DefaultLayout, bDoInitialClear, InCreateDesc.ClearValue, bIsTransientResource);
-	}
-
-	if (InCreateDesc.BulkData)
-	{
-		// InternalLockWrite leaves the image in VK_IMAGE_LAYOUT_SHADER_READ_ONLY_OPTIMAL, so make sure the requested resource state is SRV.
-		check(EnumHasAnyFlags(InCreateDesc.InitialState, ERHIAccess::SRVMask));
-		DefaultLayout = VK_IMAGE_LAYOUT_SHADER_READ_ONLY_OPTIMAL;
-
-		// Transfer bulk data
-		VulkanRHI::FStagingBuffer* StagingBuffer = InDevice.GetStagingManager().AcquireBuffer(InCreateDesc.BulkData->GetResourceBulkDataSize());
-		void* Data = StagingBuffer->GetMappedPointer();
-
-		// Do copy
-		FMemory::Memcpy(Data, InCreateDesc.BulkData->GetResourceBulkData(), InCreateDesc.BulkData->GetResourceBulkDataSize());
-		InCreateDesc.BulkData->Discard();
-
-		VkBufferImageCopy Region;
-		FMemory::Memzero(Region);
-		//#todo-rco: Use real Buffer offset when switching to suballocations!
-		Region.bufferOffset = 0;
-		Region.bufferRowLength = InCreateDesc.Extent.X;
-		Region.bufferImageHeight = InCreateDesc.Extent.Y;
-
-		Region.imageSubresource.mipLevel = 0;
-		Region.imageSubresource.baseArrayLayer = 0;
-		Region.imageSubresource.layerCount = GetNumberOfArrayLevels();
-		Region.imageSubresource.aspectMask = GetFullAspectMask();
-
-		Region.imageExtent.width = Region.bufferRowLength;
-		Region.imageExtent.height = Region.bufferImageHeight;
-		Region.imageExtent.depth = InCreateDesc.Depth;
-
-		RHICmdList.EnqueueLambda(TEXT("FVulkanTexture::InternalLockWrite"),
-			[VulkanTexture = this, Region, StagingBuffer](FRHICommandListBase& ExecutingCmdList)
-			{
-				FVulkanTexture::InternalLockWrite(FVulkanUploadContext::Get(ExecutingCmdList), VulkanTexture, Region, StagingBuffer);
-			});
-	}
-}
-
 FVulkanTexture::FVulkanTexture(FVulkanDevice& InDevice, const FRHITextureCreateDesc& InCreateDesc, VkImage InImage, const FVulkanRHIExternalImageDeleteCallbackInfo& InExternalImageDeleteCallbackInfo)
 	: FRHITexture(InCreateDesc)
 	, Device(&InDevice)
 	, Image(InImage)
-	, ImageUsageFlags(0)
-	, StorageFormat(VK_FORMAT_UNDEFINED)
-	, ViewFormat(VK_FORMAT_UNDEFINED)
-	, MemProps(VK_MEMORY_PROPERTY_DEVICE_LOCAL_BIT)
 	, ExternalImageDeleteCallbackInfo(InExternalImageDeleteCallbackInfo)
-	, Tiling(VK_IMAGE_TILING_MAX_ENUM)	// Can be expanded to a per-platform definition
-	, FullAspectMask(0)
-	, PartialAspectMask(0)
-	, CpuReadbackBuffer(nullptr)
-	, DefaultLayout(VK_IMAGE_LAYOUT_UNDEFINED)
 	, ImageOwnerType(EImageOwnerType::ExternalOwner)
 {
 	VULKAN_TRACK_OBJECT_CREATE(FVulkanTexture, this);
@@ -1702,15 +1721,6 @@ static void CleanupVulkanAndroidTextureResources(void* UserData)
 FVulkanTexture::FVulkanTexture(FVulkanDevice& InDevice, const FRHITextureCreateDesc& InCreateDesc, const AHardwareBuffer_Desc& HardwareBufferDesc, AHardwareBuffer* HardwareBuffer)
 	: FRHITexture(InCreateDesc)
 	, Device(&InDevice)
-	, ImageUsageFlags(0)
-	, StorageFormat(VK_FORMAT_UNDEFINED)
-	, ViewFormat(VK_FORMAT_UNDEFINED)
-	, MemProps(VK_MEMORY_PROPERTY_DEVICE_LOCAL_BIT)
-	, Tiling(VK_IMAGE_TILING_MAX_ENUM)	// Can be expanded to a per-platform definition
-	, FullAspectMask(0)
-	, PartialAspectMask(0)
-	, CpuReadbackBuffer(nullptr)
-	, DefaultLayout(VK_IMAGE_LAYOUT_UNDEFINED)
 	, ImageOwnerType(EImageOwnerType::ExternalOwner)
 {
 	VULKAN_TRACK_OBJECT_CREATE(FVulkanTexture, this);
@@ -1902,16 +1912,6 @@ FVulkanTexture::FVulkanTexture(FVulkanDevice& InDevice, const FRHITextureCreateD
 FVulkanTexture::FVulkanTexture(FVulkanDevice& InDevice, const FRHITextureCreateDesc& InCreateDesc, FTextureRHIRef& SrcTextureRHI)
 	: FRHITexture(InCreateDesc)
 	, Device(&InDevice)
-	, Image(VK_NULL_HANDLE)
-	, ImageUsageFlags(0)
-	, StorageFormat(VK_FORMAT_UNDEFINED)
-	, ViewFormat(VK_FORMAT_UNDEFINED)
-	, MemProps(VK_MEMORY_PROPERTY_DEVICE_LOCAL_BIT)
-	, Tiling(VK_IMAGE_TILING_MAX_ENUM)	// Can be expanded to a per-platform definition
-	, FullAspectMask(0)
-	, PartialAspectMask(0)
-	, CpuReadbackBuffer(nullptr)
-	, DefaultLayout(VK_IMAGE_LAYOUT_UNDEFINED)
 	, ImageOwnerType(EImageOwnerType::Aliased)
 {
 	VULKAN_TRACK_OBJECT_CREATE(FVulkanTexture, this);
diff --git a/Engine/Source/Runtime/VulkanRHI/Private/VulkanTransientResourceAllocator.cpp b/Engine/Source/Runtime/VulkanRHI/Private/VulkanTransientResourceAllocator.cpp
--- a/Engine/Source/Runtime/VulkanRHI/Private/VulkanTransientResourceAllocator.cpp
+++ b/Engine/Source/Runtime/VulkanRHI/Private/VulkanTransientResourceAllocator.cpp
@@ -1,6 +1,8 @@
 // Copyright Epic Games, Inc. All Rights Reserved.
 
 #include "VulkanTransientResourceAllocator.h"
+
+#include "VulkanRHIPrivate.h"
 #include "VulkanCommandWrappers.h"
 #include "VulkanDevice.h"
 #include "VulkanMemory.h"
@@ -101,7 +103,8 @@ FRHITransientTexture* FVulkanTransientResourceAllocator::CreateTexture(const FRH
 		[Device=Device, InCreateInfo, InDebugName, MemReq](const FRHITransientHeap::FResourceInitializer& Initializer)
 	{
 		FRHITextureCreateDesc CreateDesc(InCreateInfo, ERHIAccess::Discard, InDebugName);
-		FRHITexture* Texture = new FVulkanTexture(*Device, CreateDesc, &Initializer.Allocation);
+
+		FRHITexture* Texture = FVulkanDynamicRHI::Get().CreateTextureInternal(CreateDesc, Initializer.Allocation);
 		return new FRHITransientTexture(Texture, 0/*GpuVirtualAddress*/, Initializer.Hash, MemReq.Size, ERHITransientAllocationType::Heap, InCreateInfo);
 	});
 }
diff --git a/Engine/Source/Runtime/VulkanRHI/Private/VulkanViewport.cpp b/Engine/Source/Runtime/VulkanRHI/Private/VulkanViewport.cpp
--- a/Engine/Source/Runtime/VulkanRHI/Private/VulkanViewport.cpp
+++ b/Engine/Source/Runtime/VulkanRHI/Private/VulkanViewport.cpp
@@ -827,14 +827,14 @@ void FVulkanViewport::CreateSwapchain(FRHICommandListBase& RHICmdList, FVulkanSw
 		const UE::StereoRenderUtils::FStereoShaderAspects Aspects(GMaxRHIShaderPlatform);
 
 		const int kMultiViewCount = 2; // TODO: number of subresources may change in the future
-		const FRHITextureCreateDesc Desc = (Aspects.IsMobileMultiViewEnabled() ?			
+		const FRHITextureCreateDesc CreateDesc = (Aspects.IsMobileMultiViewEnabled() ?			
 			FRHITextureCreateDesc::Create2DArray(TEXT("RenderingBackBufferArr"), BackBufferSizeX, BackBufferSizeY, kMultiViewCount, PixelFormat) :
 			FRHITextureCreateDesc::Create2D(TEXT("RenderingBackBuffer"), BackBufferSizeX, BackBufferSizeY, PixelFormat))
 			.SetClearValue(FClearValueBinding::None)
 			.SetFlags(ETextureCreateFlags::RenderTargetable | ETextureCreateFlags::ShaderResource | ETextureCreateFlags::ResolveTargetable)
 			.SetInitialState(ERHIAccess::Present);
 
-		RenderingBackBuffer = new FVulkanTexture(RHICmdList, *Device, Desc, nullptr);
+		RenderingBackBuffer = FVulkanDynamicRHI::Get().CreateTextureInternal(RHICmdList, CreateDesc, {});
 #if VULKAN_ENABLE_DRAW_MARKERS
 		if (Device->GetSetDebugName())
 		{
diff --git a/Engine/Source/Runtime/VulkanRHI/Public/VulkanDynamicRHI.h b/Engine/Source/Runtime/VulkanRHI/Public/VulkanDynamicRHI.h
--- a/Engine/Source/Runtime/VulkanRHI/Public/VulkanDynamicRHI.h
+++ b/Engine/Source/Runtime/VulkanRHI/Public/VulkanDynamicRHI.h
@@ -15,6 +15,7 @@ class FVulkanQueue;
 class FVulkanUniformBuffer;
 class FVulkanViewport;
 class FVulkanPayload;
+class FVulkanTexture;
 class IHeadMountedDisplayVulkanExtensions;
 class FVulkanThread;
 struct FRHITransientHeapAllocation;
@@ -144,6 +145,21 @@ public:
 	virtual FRHICalcTextureSizeResult RHICalcTexturePlatformSize(FRHITextureDesc const& Desc, uint32 FirstMipIndex) final override;
 	virtual void RHIGetTextureMemoryStats(FTextureMemoryStats& OutStats) final override;
 	virtual bool RHIGetTextureMemoryVisualizeData(FColor* TextureData, int32 SizeX, int32 SizeY, int32 Pitch, int32 PixelSize) final override;
+
+	struct FCreateTextureResult
+	{
+		FVulkanTexture* Texture;
+		VkImageLayout DefaultLayout;
+		bool bTransientResource;
+		bool bClear;
+	};
+
+	FCreateTextureResult BeginCreateTextureInternal(const FRHITextureCreateDesc& CreateDesc, const FRHITransientHeapAllocation* InTransientHeapAllocation);
+	FVulkanTexture* FinalizeCreateTextureInternal(FRHICommandListBase& RHICmdList, FCreateTextureResult CreateResult, TConstArrayView<uint8> InitialData);
+
+	FVulkanTexture* CreateTextureInternal(FRHICommandListBase& RHICmdList, const FRHITextureCreateDesc& CreateDesc, TConstArrayView<uint8> InitialData);
+	FVulkanTexture* CreateTextureInternal(const FRHITextureCreateDesc& CreateDesc, const FRHITransientHeapAllocation& InTransientHeapAllocation);
+
 	virtual FTextureRHIRef RHICreateTexture(FRHICommandListBase& RHICmdList, const FRHITextureCreateDesc& CreateDesc) final override;
 	virtual FTextureRHIRef RHIAsyncCreateTexture2D(uint32 SizeX, uint32 SizeY, uint8 Format, uint32 NumMips, ETextureCreateFlags Flags, ERHIAccess InResourceState, void** InitialMipData, uint32 NumInitialMips, const TCHAR* DebugName, FGraphEventRef& OutCompletionEvent) final override;
 	virtual uint32 RHIComputeMemorySize(FRHITexture* TextureRHI) final override;
