VulkanRHI: Simplify timing queries in preparation for new profiling system.  Read timing query results from interrupt thread before dispatching sync points.  Fix issue where GPU Time would sometimes not update for extended periods of time.

#rnx

diff --git a/Engine/Source/Runtime/VulkanRHI/Private/VulkanCommandBuffer.cpp b/Engine/Source/Runtime/VulkanRHI/Private/VulkanCommandBuffer.cpp
--- a/Engine/Source/Runtime/VulkanRHI/Private/VulkanCommandBuffer.cpp
+++ b/Engine/Source/Runtime/VulkanRHI/Private/VulkanCommandBuffer.cpp
@@ -129,39 +129,42 @@ void FVulkanCommandBuffer::BeginRenderPass(const FVulkanRenderTargetLayout& Layo
 }
 
 
-void FVulkanCommandBuffer::AddPendingTimestampQuery(uint64 Index, uint64 Count, VkQueryPool PoolHandle, VkBuffer BufferHandle, bool bBlocking)
+void FVulkanCommandBuffer::AddPendingTimestampQuery(uint64 Index, uint64 Count, VkQueryPool PoolHandle, FVulkanTimingQueryResult* Results, bool bBlocking)
 {
-	PendingQuery PQ;
-	PQ.Index = Index;
-	PQ.Count = Count;
-	PQ.PoolHandle = PoolHandle;
-	PQ.BufferHandle = BufferHandle;
-	PQ.bBlocking = bBlocking;
-	PendingTimestampQueries.Add(PQ);
+	FPendingTimestampQuery PendingQuery;
+	PendingQuery.Index = Index;
+	PendingQuery.Count = Count;
+	PendingQuery.PoolHandle = PoolHandle;
+	PendingQuery.Results = Results;
+	PendingQuery.bBlocking = bBlocking;
+	PendingTimestampQueries.Add(PendingQuery);
 }
 
-void FVulkanCommandBuffer::End()
-{
-	checkf(IsOutsideRenderPass(), TEXT("Can't End as we're inside a render pass! CmdBuffer 0x%p State=%d"), CommandBufferHandle, (int32)State);
 
-	for (PendingQuery& Query : PendingTimestampQueries)
+void FVulkanCommandBuffer::ResolvePendingTimestampQueries()
+{
+	for (const FPendingTimestampQuery& PendingQuery : PendingTimestampQueries)
 	{
-		const uint64 Index = Query.Index;
-		const VkBuffer BufferHandle = Query.BufferHandle;
-		const VkQueryPool PoolHandle = Query.PoolHandle;
-		const VkQueryResultFlags BlockingFlags = Query.bBlocking ?  VK_QUERY_RESULT_WAIT_BIT : VK_QUERY_RESULT_WITH_AVAILABILITY_BIT;
-		const uint32 Width = (Query.bBlocking ? 1 : 2);
-		const uint32 Stride = sizeof(uint64) * Width;
-
-		VulkanRHI::vkCmdCopyQueryPoolResults(GetHandle(), PoolHandle, Index, Query.Count, BufferHandle, Stride * Index, Stride, VK_QUERY_RESULT_64_BIT | BlockingFlags);
-		if (Query.bBlocking)
-		{
-			VulkanRHI::vkCmdResetQueryPool(GetHandle(), PoolHandle, Index, Query.Count);
-		}
+		const VkQueryResultFlags BlockingQueryResultFlags = PendingQuery.bBlocking ? VK_QUERY_RESULT_WAIT_BIT : 0;
+		const VkResult Result = VulkanRHI::vkGetQueryPoolResults(
+			Device.GetInstanceHandle(),
+			PendingQuery.PoolHandle,
+			PendingQuery.Index, 
+			PendingQuery.Count, 
+			PendingQuery.Count * sizeof(FVulkanTimingQueryResult),
+			(uint64_t*)PendingQuery.Results,
+			sizeof(FVulkanTimingQueryResult),
+			VK_QUERY_RESULT_64_BIT | VK_QUERY_RESULT_WITH_AVAILABILITY_BIT | BlockingQueryResultFlags);
+		VERIFYVULKANRESULT_EXPANDED(Result);
 	}
 
 	PendingTimestampQueries.Reset();
+}
+
 
+void FVulkanCommandBuffer::End()
+{
+	checkf(IsOutsideRenderPass(), TEXT("Can't End as we're inside a render pass! CmdBuffer 0x%p State=%d"), CommandBufferHandle, (int32)State);
 	VERIFYVULKANRESULT(VulkanRHI::vkEndCommandBuffer(GetHandle()));
 	State = EState::HasEnded;
 }
diff --git a/Engine/Source/Runtime/VulkanRHI/Private/VulkanCommandBuffer.h b/Engine/Source/Runtime/VulkanRHI/Private/VulkanCommandBuffer.h
--- a/Engine/Source/Runtime/VulkanRHI/Private/VulkanCommandBuffer.h
+++ b/Engine/Source/Runtime/VulkanRHI/Private/VulkanCommandBuffer.h
@@ -20,6 +20,7 @@ class FVulkanRenderTargetLayout;
 class FVulkanQueue;
 class FVulkanDescriptorPoolSetContainer;
 class FVulkanGPUTiming;
+struct FVulkanTimingQueryResult;
 
 namespace VulkanRHI
 {
@@ -110,16 +111,16 @@ public:
 
 	bool AcquirePoolSetAndDescriptorsIfNeeded(const class FVulkanDescriptorSetsLayout& Layout, bool bNeedDescriptors, VkDescriptorSet* OutDescriptors);
 
-
-	struct PendingQuery
+	struct FPendingTimestampQuery
 	{
 		uint64 Index;
 		uint64 Count;
-		VkBuffer BufferHandle;
+		FVulkanTimingQueryResult* Results;
 		VkQueryPool PoolHandle;
 		bool bBlocking;
 	};
-	void AddPendingTimestampQuery(uint64 Index, uint64 Count, VkQueryPool PoolHandle, VkBuffer BufferHandle, bool bBlocking);
+	void AddPendingTimestampQuery(uint64 Index, uint64 Count, VkQueryPool PoolHandle, FVulkanTimingQueryResult* Results, bool bBlocking);
+	void ResolvePendingTimestampQueries();
 
 	// Public immutable reference to be accessed directly
 	FVulkanDevice& Device;
@@ -128,7 +129,7 @@ private:
 	VkCommandBuffer CommandBufferHandle;
 	double SubmittedTime = 0.0f;
 
-	TArray<PendingQuery> PendingTimestampQueries;
+	TArray<FPendingTimestampQuery> PendingTimestampQueries;
 
 	void ResetFenceStatus();
 
diff --git a/Engine/Source/Runtime/VulkanRHI/Private/VulkanGPUProfiler.h b/Engine/Source/Runtime/VulkanRHI/Private/VulkanGPUProfiler.h
--- a/Engine/Source/Runtime/VulkanRHI/Private/VulkanGPUProfiler.h
+++ b/Engine/Source/Runtime/VulkanRHI/Private/VulkanGPUProfiler.h
@@ -69,15 +69,7 @@ public:
 
 private:
 
-	/**
-	 * If results weren't available when first queried, make sure we try again.
-	 */
-	void RefreshPendingQueries(FVulkanCommandBuffer& CmdBuffer);
-
-	/**
-	 * Remove expired queries to make space for new ones
-	 */
-	void DiscardExpiredQueries(bool bNeedSpace);
+	void DiscardOldestQuery();
 
 	/**
 	 * Initializes the static variables, if necessary.
@@ -90,8 +82,7 @@ private:
 	/** Whether we are currently timing the GPU: between StartTiming() and EndTiming(). */
 	bool bIsTiming;
 	bool bEndTimestampIssued;
-	uint64 PreviousFrame = UINT64_MAX;
-	uint64 PreviousTime = 0;
+	uint64 LastTime = 0;
 	uint32 NumPendingQueries = 0;
 	TQueue<uint32> PendingQueries;
 	FVulkanTimingQueryPool* Pool = nullptr;
diff --git a/Engine/Source/Runtime/VulkanRHI/Private/VulkanQuery.cpp b/Engine/Source/Runtime/VulkanRHI/Private/VulkanQuery.cpp
--- a/Engine/Source/Runtime/VulkanRHI/Private/VulkanQuery.cpp
+++ b/Engine/Source/Runtime/VulkanRHI/Private/VulkanQuery.cpp
@@ -71,6 +71,15 @@ FVulkanQueryPool::~FVulkanQueryPool()
 	QueryPool = VK_NULL_HANDLE;
 }
 
+FVulkanTimingQueryPool::~FVulkanTimingQueryPool()
+{
+	for (FTimestampData& QueryData : TimestampArray)
+	{
+		checkf(!QueryData.SyncPoint.IsValid() || QueryData.SyncPoint->IsComplete(), 
+			TEXT("Vulkan Query Pool being destroyed with queries still active!"));
+	}
+}
+
 bool FVulkanOcclusionQueryPool::CanBeReused()
 {
 	const uint64 NumWords = NumUsedQueries / 64;
@@ -335,13 +344,13 @@ FRenderQueryRHIRef FVulkanDynamicRHI::RHICreateRenderQuery(ERenderQueryType Quer
 	return Query;
 }
 
-bool FVulkanDynamicRHI::RHIGetRenderQueryResult(FRHIRenderQuery* QueryRHI, uint64& OutNumPixels, bool bWait, uint32 GPUIndex)
+bool FVulkanDynamicRHI::RHIGetRenderQueryResult(FRHIRenderQuery* QueryRHI, uint64& OutQueryResult, bool bWait, uint32 GPUIndex)
 {
 	FVulkanRenderQuery* BaseQuery = ResourceCast(QueryRHI);
 
 	if (!ensureMsgf(BaseQuery->SyncPoint, TEXT("Attempt to get result data for an FRHIRenderQuery that was never used in a command list.")))
 	{
-		OutNumPixels = 0;
+		OutQueryResult = 0;
 		return false;
 	}
 
@@ -363,7 +372,7 @@ bool FVulkanDynamicRHI::RHIGetRenderQueryResult(FRHIRenderQuery* QueryRHI, uint6
 		FVulkanOcclusionQuery* Query = static_cast<FVulkanOcclusionQuery*>(BaseQuery);
 		if (Query->State == FVulkanOcclusionQuery::EState::RT_GotResults || Query->State == FVulkanOcclusionQuery::EState::FlushedFromPoolHadResults)
 		{
-			OutNumPixels = Query->Result;
+			OutQueryResult = Query->Result;
 			return true;
 		}
 
@@ -379,7 +388,7 @@ bool FVulkanDynamicRHI::RHIGetRenderQueryResult(FRHIRenderQuery* QueryRHI, uint6
 			Query->Result = Query->Pool->GetResultValue(Query->IndexInPool);
 			Query->ReleaseFromPool();
 			Query->State = FVulkanOcclusionQuery::EState::RT_GotResults;
-			OutNumPixels = Query->Result;
+			OutQueryResult = Query->Result;
 			return true;
 		}
 	}
@@ -405,23 +414,21 @@ bool FVulkanDynamicRHI::RHIGetRenderQueryResult(FRHIRenderQuery* QueryRHI, uint6
 			check(Query->Pool != nullptr);
 		}
 
-		check(Query->Pool && (Query->Pool->CurrentTimestamp < Query->Pool->BufferSize));
-		int32 TimestampIndex = Query->Pool->CurrentTimestamp;
+		check(Query->Pool && (Query->Pool->CurrentTimestampIndex < Query->Pool->BufferedCount));
+		int32 TimestampIndex = Query->Pool->CurrentTimestampIndex;
 		if (!bWait)
 		{
 			// Quickly check the most recent measurements to see if any of them has been resolved.  Do not flush these queries.
 			for (uint32 IssueIndex = 1; IssueIndex < Query->Pool->NumIssuedTimestamps; ++IssueIndex)
 			{
-				const FVulkanTimingQueryPool::FCmdBufferFence& StartQuerySyncPoint = Query->Pool->TimestampListHandles[TimestampIndex];
-				if (StartQuerySyncPoint.SyncPoint->IsComplete())
+				const FVulkanTimingQueryPool::FTimestampData& QueryData = Query->Pool->TimestampArray[TimestampIndex];
+				if (QueryData.SyncPoint->IsComplete())
 				{
-					Query->Pool->ResultsBuffer->InvalidateMappedMemory();
-					const uint64* Data = Query->Pool->MappedPointer;
-					OutNumPixels = ToMicroseconds(Data[TimestampIndex]);
+					OutQueryResult = ToMicroseconds(QueryData.Result.Value);
 					return true;
 				}
 
-				TimestampIndex = (TimestampIndex + Query->Pool->BufferSize - 1) % Query->Pool->BufferSize;
+				TimestampIndex = (TimestampIndex + Query->Pool->BufferedCount - 1) % Query->Pool->BufferedCount;
 			}
 		}
 
@@ -431,27 +438,24 @@ bool FVulkanDynamicRHI::RHIGetRenderQueryResult(FRHIRenderQuery* QueryRHI, uint6
 			// so check the oldest measurement more thoroughly.
 			// This really only happens if occlusion and frame sync event queries are disabled, otherwise those will block until the GPU catches up to 1 frame behind
 
-			const bool bBlocking = (Query->Pool->NumIssuedTimestamps == Query->Pool->BufferSize) || bWait;
+			const FVulkanTimingQueryPool::FTimestampData& QueryData = Query->Pool->TimestampArray[TimestampIndex];
+
 			{
 				FRenderThreadIdleScope IdleScope(ERenderThreadIdleTypes::WaitingForGPUQuery);
 				SCOPE_CYCLE_COUNTER(STAT_RenderQueryResultTime);
 
+				const bool bBlocking = (Query->Pool->NumIssuedTimestamps == Query->Pool->BufferedCount) || bWait;
 				if (bBlocking)
 				{
-					const FVulkanTimingQueryPool::FCmdBufferFence& StartQuerySyncPoint = Query->Pool->TimestampListHandles[TimestampIndex];
-
 					// CPU wait for query results to be ready.
-					if (!StartQuerySyncPoint.SyncPoint->IsComplete())
+					if (!QueryData.SyncPoint->IsComplete())
 					{
-						FVulkanDynamicRHI::Get().ProcessInterruptQueueUntil(StartQuerySyncPoint.SyncPoint);
+						FVulkanDynamicRHI::Get().ProcessInterruptQueueUntil(QueryData.SyncPoint);
 					}
 				}
-
-				Query->Pool->ResultsBuffer->InvalidateMappedMemory();
 			}
 
-			const uint64* Data = Query->Pool->MappedPointer;
-			OutNumPixels = ToMicroseconds(Data[TimestampIndex]);
+			OutQueryResult = ToMicroseconds(QueryData.Result.Value);
 			return true;
 		}
 	}
@@ -573,19 +577,19 @@ void FVulkanCommandListContext::RHIEndRenderQuery(FRHIRenderQuery* QueryRHI)
 		if (Query->Pool == nullptr)
 		{
 			Query->Pool = new FVulkanTimingQueryPool(Device, this, 4);
-			Query->Pool->ResultsBuffer = Device.GetStagingManager().AcquireBuffer(Query->Pool->BufferSize * sizeof(uint64), VK_BUFFER_USAGE_TRANSFER_DST_BIT, VK_MEMORY_PROPERTY_HOST_COHERENT_BIT);
-			Query->Pool->MappedPointer = (uint64*)Query->Pool->ResultsBuffer->GetMappedPointer();
 		}
 
-		Query->Pool->CurrentTimestamp = (Query->Pool->CurrentTimestamp + 1) % Query->Pool->BufferSize;
-		const uint32 QueryEndIndex = Query->Pool->CurrentTimestamp;
+		Query->Pool->CurrentTimestampIndex = (Query->Pool->CurrentTimestampIndex + 1) % Query->Pool->BufferedCount;
+		const uint32 QueryEndIndex = Query->Pool->CurrentTimestampIndex;
+		FVulkanTimingQueryPool::FTimestampData& QueryData = Query->Pool->TimestampArray[QueryEndIndex];
+
 		const VkPipelineStageFlagBits QueryPipelineStage = GTimestampQueryStage ? VK_PIPELINE_STAGE_BOTTOM_OF_PIPE_BIT : VK_PIPELINE_STAGE_TOP_OF_PIPE_BIT;
 		FVulkanCommandBuffer& CommandBuffer = GetCommandBuffer();
 		VulkanRHI::vkCmdWriteTimestamp(CommandBuffer.GetHandle(), QueryPipelineStage, Query->Pool->GetHandle(), QueryEndIndex);
-		CommandBuffer.AddPendingTimestampQuery(QueryEndIndex, 1, Query->Pool->GetHandle(), Query->Pool->ResultsBuffer->GetHandle(), true);
-		Query->Pool->TimestampListHandles[QueryEndIndex].SyncPoint = GetContextSyncPoint();
-		Query->Pool->TimestampListHandles[QueryEndIndex].FrameCount = GetFrameCounter();
-		Query->Pool->NumIssuedTimestamps = FMath::Min<uint32>(Query->Pool->NumIssuedTimestamps + 1, Query->Pool->BufferSize);
+		Query->Pool->TimestampArray[QueryEndIndex].SyncPoint = GetContextSyncPoint();
+		Query->Pool->NumIssuedTimestamps = FMath::Min<uint32>(Query->Pool->NumIssuedTimestamps + 1, Query->Pool->BufferedCount);
+
+		CommandBuffer.AddPendingTimestampQuery(QueryEndIndex, 1, Query->Pool->GetHandle(), &QueryData.Result, true);
 	}
 }
 
@@ -624,7 +628,7 @@ FVulkanTimingQuery::~FVulkanTimingQuery()
 {
 	if (Pool != nullptr)
 	{
-		Pool->GetDevice().GetStagingManager().ReleaseBuffer(nullptr, Pool->ResultsBuffer);
 		delete Pool;
+		Pool = nullptr;
 	}
 }
diff --git a/Engine/Source/Runtime/VulkanRHI/Private/VulkanQueue.cpp b/Engine/Source/Runtime/VulkanRHI/Private/VulkanQueue.cpp
--- a/Engine/Source/Runtime/VulkanRHI/Private/VulkanQueue.cpp
+++ b/Engine/Source/Runtime/VulkanRHI/Private/VulkanQueue.cpp
@@ -331,6 +331,7 @@ int32 FVulkanQueue::ProcessInterruptQueue(uint64 Timeout)
 
 			for (FVulkanCommandBuffer* CommandBuffer : Payload->CommandBuffers)
 			{
+				CommandBuffer->ResolvePendingTimestampQueries();
 				CommandBuffer->ResetFenceStatus();
 			}
 
diff --git a/Engine/Source/Runtime/VulkanRHI/Private/VulkanResources.h b/Engine/Source/Runtime/VulkanRHI/Private/VulkanResources.h
--- a/Engine/Source/Runtime/VulkanRHI/Private/VulkanResources.h
+++ b/Engine/Source/Runtime/VulkanRHI/Private/VulkanResources.h
@@ -889,30 +889,34 @@ protected:
 	uint32 FrameNumber = UINT32_MAX;
 };
 
+struct FVulkanTimingQueryResult
+{
+	uint64 Value = 0;
+	uint64 Availability = 0;
+};
+
 class FVulkanTimingQueryPool : public FVulkanQueryPool
 {
 public:
-	FVulkanTimingQueryPool(FVulkanDevice& InDevice, FVulkanContextCommon* InContext, uint32 InBufferSize)
-		: FVulkanQueryPool(InDevice, InContext, InBufferSize * 2, VK_QUERY_TYPE_TIMESTAMP, false)
-		, BufferSize(InBufferSize)
+	FVulkanTimingQueryPool(FVulkanDevice& InDevice, FVulkanContextCommon* InContext, uint32 InBufferedCount)
+		: FVulkanQueryPool(InDevice, InContext, InBufferedCount * 2, VK_QUERY_TYPE_TIMESTAMP, false)
+		, BufferedCount(InBufferedCount)
 	{
-		TimestampListHandles.AddZeroed(InBufferSize * 2);
+		TimestampArray.AddDefaulted(GetMaxQueries());
 	}
 
-	uint32 CurrentTimestamp = 0;
+	virtual ~FVulkanTimingQueryPool();
+
+	uint32 CurrentTimestampIndex = 0;
 	uint32 NumIssuedTimestamps = 0;
-	const uint32 BufferSize;
+	const uint32 BufferedCount;
 
-	struct FCmdBufferFence
+	struct FTimestampData
 	{
 		FVulkanSyncPointRef SyncPoint;
-		uint64 FrameCount = UINT64_MAX;
-		uint32 Attempts = 0;
+		FVulkanTimingQueryResult Result;
 	};
-	TArray<FCmdBufferFence> TimestampListHandles;
-
-	VulkanRHI::FStagingBuffer* ResultsBuffer = nullptr;
-	uint64* MappedPointer = nullptr;
+	TArray<FTimestampData> TimestampArray;
 };
 
 class FVulkanRenderQuery : public FRHIRenderQuery
diff --git a/Engine/Source/Runtime/VulkanRHI/Private/VulkanUtil.cpp b/Engine/Source/Runtime/VulkanRHI/Private/VulkanUtil.cpp
--- a/Engine/Source/Runtime/VulkanRHI/Private/VulkanUtil.cpp
+++ b/Engine/Source/Runtime/VulkanRHI/Private/VulkanUtil.cpp
@@ -28,8 +28,6 @@ extern CORE_API bool GIsGPUCrashed;
 
 static FString		EventDeepString(TEXT("EventTooDeep"));
 static const uint32	EventDeepCRC = FCrc::StrCrc32<TCHAR>(*EventDeepString);
-static const uint32 BUFFERED_TIMING_QUERIES = 1;
-static const uint32 TIMING_QUERY_RETRIES = 1;
 
 /**
  * Initializes the static variables, if necessary.
@@ -85,8 +83,6 @@ void FVulkanGPUTiming::Initialize(uint32 PoolSize)
 	{
 		check(!Pool);
 		Pool = new FVulkanTimingQueryPool(*Device, Context, PoolSize);
-		Pool->ResultsBuffer = Device->GetStagingManager().AcquireBuffer(Pool->GetMaxQueries() * sizeof(uint64) * 2, VK_BUFFER_USAGE_TRANSFER_DST_BIT, VK_MEMORY_PROPERTY_HOST_COHERENT_BIT);
-		Pool->MappedPointer = (uint64*)Pool->ResultsBuffer->GetMappedPointer();
 	}
 }
 
@@ -98,7 +94,6 @@ void FVulkanGPUTiming::Release()
 	if (FVulkanPlatform::SupportsTimestampRenderQueries() && GIsSupported)
 	{
 		check(Pool);
-		Device->GetStagingManager().ReleaseBuffer(nullptr, Pool->ResultsBuffer);
 		delete Pool;
 		Pool = nullptr;
 	}
@@ -118,20 +113,18 @@ void FVulkanGPUTiming::StartTiming(FVulkanCommandBuffer* CmdBuffer)
 		}
 
 		// In case we aren't reading queries, remove oldest
-		if (NumPendingQueries >= Pool->BufferSize)
+		if (NumPendingQueries >= Pool->BufferedCount)
 		{
-			DiscardExpiredQueries(true);
-
-			if (NumPendingQueries >= Pool->BufferSize)
+			DiscardOldestQuery();
+			if (NumPendingQueries >= Pool->BufferedCount)
 			{
 				return;
 			}
 		}
 
-		RefreshPendingQueries(*CmdBuffer);
-
-		Pool->CurrentTimestamp = (Pool->CurrentTimestamp + 1) % Pool->BufferSize;
-		const uint32 QueryStartIndex = Pool->CurrentTimestamp * 2;
+		Pool->CurrentTimestampIndex = (Pool->CurrentTimestampIndex + 1) % Pool->BufferedCount;
+		const uint32 QueryStartIndex = Pool->CurrentTimestampIndex * 2;
+		FVulkanTimingQueryPool::FTimestampData& StartQueryData = Pool->TimestampArray[QueryStartIndex];
 
 		// Since we use non-blocking queries, we need to reset them before using them
 		// Use host query resets if supported since there's no need to consider host/GPU sync
@@ -145,10 +138,10 @@ void FVulkanGPUTiming::StartTiming(FVulkanCommandBuffer* CmdBuffer)
 		}
 
 		VulkanRHI::vkCmdWriteTimestamp(CmdBuffer->GetHandle(), VK_PIPELINE_STAGE_BOTTOM_OF_PIPE_BIT, Pool->GetHandle(), QueryStartIndex);
-		Pool->TimestampListHandles[QueryStartIndex].SyncPoint = Context->GetContextSyncPoint();
-		Pool->TimestampListHandles[QueryStartIndex].FrameCount = Context->GetFrameCounter();
-		Pool->TimestampListHandles[QueryStartIndex].Attempts = 0;
+		StartQueryData.SyncPoint = Context->GetContextSyncPoint();
 		bIsTiming = true;
+
+		CmdBuffer->AddPendingTimestampQuery(QueryStartIndex, 1, Pool->GetHandle(), &StartQueryData.Result, false);
 	}
 }
 
@@ -165,75 +158,47 @@ void FVulkanGPUTiming::EndTiming(FVulkanCommandBuffer* CmdBuffer)
 		{
 			CmdBuffer = Context->GetActiveCmdBuffer();
 		}
-		check(Pool->CurrentTimestamp < Pool->BufferSize);
-		const uint32 QueryStartIndex = Pool->CurrentTimestamp * 2;
-		// Keep Start and End contiguous to fetch them together with a single AddPendingTimestampQuery(QueryStartIndex,2,...)
-		const uint32 QueryEndIndex = QueryStartIndex + 1;   
+		check(Pool->CurrentTimestampIndex < Pool->BufferedCount);
+		const uint32 QueryStartIndex = Pool->CurrentTimestampIndex * 2;
+		const uint32 QueryEndIndex = QueryStartIndex + 1;
+
+		FVulkanTimingQueryPool::FTimestampData& EndQueryData = Pool->TimestampArray[QueryEndIndex];
 
-		PendingQueries.Enqueue(Pool->CurrentTimestamp);
+		PendingQueries.Enqueue(Pool->CurrentTimestampIndex);
 		NumPendingQueries++;
 
 		VulkanRHI::vkCmdWriteTimestamp(CmdBuffer->GetHandle(), VK_PIPELINE_STAGE_BOTTOM_OF_PIPE_BIT, Pool->GetHandle(), QueryEndIndex);
-		CmdBuffer->AddPendingTimestampQuery(QueryStartIndex, 2, Pool->GetHandle(), Pool->ResultsBuffer->GetHandle(), false);
-		Pool->TimestampListHandles[QueryStartIndex].SyncPoint = Context->GetContextSyncPoint();
-		Pool->TimestampListHandles[QueryEndIndex].SyncPoint = Context->GetContextSyncPoint();
-		Pool->TimestampListHandles[QueryEndIndex].FrameCount = Context->GetFrameCounter();
-		Pool->TimestampListHandles[QueryEndIndex].Attempts = 0;
-		Pool->NumIssuedTimestamps = FMath::Min<uint32>(Pool->NumIssuedTimestamps + 1, Pool->BufferSize);
+		Pool->TimestampArray[QueryEndIndex].SyncPoint = Context->GetContextSyncPoint();
+		Pool->NumIssuedTimestamps = FMath::Min<uint32>(Pool->NumIssuedTimestamps + 1, Pool->BufferedCount);
 
 		bIsTiming = false;
 		bEndTimestampIssued = true;
+
+		CmdBuffer->AddPendingTimestampQuery(QueryEndIndex, 1, Pool->GetHandle(), &EndQueryData.Result, false);
 	}
 }
 
-void FVulkanGPUTiming::DiscardExpiredQueries(bool bNeedSpace)
+void FVulkanGPUTiming::DiscardOldestQuery()
 {
-	// Returns true if we can safely throw away the query
-	auto CanBeDiscarded = [this](uint32 QueryIndex)
-	{
-		FVulkanTimingQueryPool::FCmdBufferFence& QuerySyncPoint = Pool->TimestampListHandles[QueryIndex];
-
-		if (QuerySyncPoint.SyncPoint->IsComplete())
-		{
-			const uint64 Availability = Pool->MappedPointer[QueryIndex * 2 + 1];
-			if (!Availability && (QuerySyncPoint.Attempts >= TIMING_QUERY_RETRIES))
-			{
-				return true;
-			}
-		}
-		return false;
-	};
-
-	auto IsQueryExpired = [&](uint32 QueryIndex)
+	auto CanBeDiscarded = [&](uint32 QueryIndex)
 	{
-		bool QueryExpired = CanBeDiscarded(QueryIndex);
-		if (!QueryExpired && bNeedSpace)
+		// It can only be competed if the sync point was signaled (queries hold a pointer they will write to when completed)
+		FVulkanTimingQueryPool::FTimestampData& QueryData = Pool->TimestampArray[QueryIndex];
+		if (!QueryData.SyncPoint->IsComplete())
 		{
-			// If the hold back for a query is the fence and the command buffer was submitted, take a shot at updating its status
-			FVulkanTimingQueryPool::FCmdBufferFence& QuerySyncPoint = Pool->TimestampListHandles[QueryIndex];
-			if (!QuerySyncPoint.SyncPoint->IsComplete())
-			{
-				FVulkanDynamicRHI::Get().ProcessInterruptQueueUntil(nullptr);  // leave null, we don't want to force a wait on the SyncPoint
-
-				if (QuerySyncPoint.SyncPoint->IsComplete())
-				{
-					QueryExpired = CanBeDiscarded(QueryIndex);
-				}
-			}
+			FVulkanDynamicRHI::Get().ProcessInterruptQueueUntil(nullptr);  // leave null, we don't want to force a wait on the SyncPoint
 		}
-		return QueryExpired;
+		return QueryData.SyncPoint->IsComplete();
 	};
 
-	Pool->ResultsBuffer->InvalidateMappedMemory();
-
 	uint32 OldestTimestamp;
-	while (PendingQueries.Peek(OldestTimestamp))
+	if (PendingQueries.Peek(OldestTimestamp))
 	{
 		const uint32 QueryStartIndex = OldestTimestamp * 2;
 		const uint32 QueryEndIndex = QueryStartIndex + 1;
 
-		const bool StartQueryExpired = IsQueryExpired(QueryStartIndex);
-		const bool EndQueryExpired = IsQueryExpired(QueryEndIndex);
+		const bool StartQueryExpired = CanBeDiscarded(QueryStartIndex);
+		const bool EndQueryExpired = CanBeDiscarded(QueryEndIndex);
 
 		if (StartQueryExpired || EndQueryExpired)
 		{
@@ -241,38 +206,9 @@ void FVulkanGPUTiming::DiscardExpiredQueries(bool bNeedSpace)
 			PendingQueries.Pop();
 			NumPendingQueries--;
 		}
-		else
-		{
-			break;
-		}
 	}
 }
 
-void FVulkanGPUTiming::RefreshPendingQueries(FVulkanCommandBuffer& CmdBuffer)
-{
-	for (int32 QueryIndex=0; QueryIndex < Pool->TimestampListHandles.Num(); ++QueryIndex)
-	{
-		FVulkanTimingQueryPool::FCmdBufferFence& QuerySyncPoint = Pool->TimestampListHandles[QueryIndex];
-
-		if (QuerySyncPoint.SyncPoint.IsValid() && QuerySyncPoint.SyncPoint->IsComplete())
-		{
-			Pool->ResultsBuffer->InvalidateMappedMemory();
-			const uint64 Availability = Pool->MappedPointer[QueryIndex * 2 + 1];
-
-			if (!Availability && (QuerySyncPoint.Attempts < TIMING_QUERY_RETRIES))
-			{
-				// Keep the same QuerySyncPoint.SyncPoint
-				QuerySyncPoint.FrameCount = Context->GetFrameCounter();
-				++QuerySyncPoint.Attempts;
-
-				CmdBuffer.AddPendingTimestampQuery(
-					QueryIndex, 1, Pool->GetHandle(), Pool->ResultsBuffer->GetHandle(), false);
-			}
-		}
-	}
-}
-
-
 /**
  * Retrieves the most recently resolved timing measurement.
  * The unit is the same as for FPlatformTime::Cycles(). Returns 0 if there are no resolved measurements.
@@ -283,112 +219,69 @@ uint64 FVulkanGPUTiming::GetTiming(bool bGetCurrentResultsAndBlock)
 {
 	if (GIsSupported)
 	{
-		check(Pool->CurrentTimestamp < Pool->BufferSize);
-		uint64 StartTime, EndTime, StartTimeAvailability, EndTimeAvailability;
+		checkf(Pool->CurrentTimestampIndex < Pool->BufferedCount, TEXT("Timestamp index out of bounds!"));
 		uint32 TimestampIndex;
-
-		uint64 TotalTime = 0;
-
-		uint64_t FrameCount = Context->GetFrameCounter();
-
-		// If these timings have been processed return the same time
-		if (PreviousFrame == FrameCount)
-		{
-			return PreviousTime;
-		}
-
-		PreviousFrame = FrameCount;
-
-		DiscardExpiredQueries(false);
-
 		while (PendingQueries.Peek(TimestampIndex))
 		{
 			const uint32 QueryStartIndex = TimestampIndex * 2;
 			const uint32 QueryEndIndex = QueryStartIndex + 1;
 
-			const FVulkanTimingQueryPool::FCmdBufferFence& StartQuerySyncPoint = Pool->TimestampListHandles[QueryStartIndex];
-			const FVulkanTimingQueryPool::FCmdBufferFence& EndQuerySyncPoint = Pool->TimestampListHandles[QueryEndIndex];
-
-			// Block if we require results or we are backed up
-			bool bBlocking = bGetCurrentResultsAndBlock;
-
-			if (!bBlocking)
-			{
-				const uint64_t QueryFrameCount = EndQuerySyncPoint.FrameCount;
-				
-				// Allow queries to back up if we are non-blocking
-				if (FrameCount < BUFFERED_TIMING_QUERIES || FrameCount - BUFFERED_TIMING_QUERIES < QueryFrameCount)
-				{
-					PreviousTime = TotalTime;
-					return TotalTime;
-				}
-				else
-				{
-					// Otherwise if we don't have a result we have to block
-					bBlocking = true;
-				}
-			}
+			const FVulkanTimingQueryPool::FTimestampData& StartQueryData = Pool->TimestampArray[QueryStartIndex];
+			const FVulkanTimingQueryPool::FTimestampData& EndQueryData = Pool->TimestampArray[QueryEndIndex];
 
-			if (EndQuerySyncPoint.SyncPoint->IsComplete() && StartQuerySyncPoint.SyncPoint->IsComplete())
-			{
-				Pool->ResultsBuffer->InvalidateMappedMemory();
-				StartTimeAvailability = Pool->MappedPointer[QueryStartIndex * 2 + 1];
-				EndTimeAvailability = Pool->MappedPointer[QueryEndIndex * 2 + 1];
-				StartTime = Pool->MappedPointer[QueryStartIndex * 2];
-				EndTime = Pool->MappedPointer[QueryEndIndex * 2];
-
-				if (!StartTimeAvailability || !EndTimeAvailability)
-				{
-					break;
-				}
-
-				PendingQueries.Pop();
-				NumPendingQueries--;
-
-				if (EndTime > StartTime)
-				{
-					// Only keep the most recent result
-					TotalTime = EndTime - StartTime;
-				}
-
-				continue;
-			}
-
-			if (bBlocking)
+			if (bGetCurrentResultsAndBlock && StartQueryData.SyncPoint.IsValid() && EndQueryData.SyncPoint.IsValid())
 			{
 				SCOPE_CYCLE_COUNTER(STAT_RenderQueryResultTime);
 				{
 					FRenderThreadIdleScope IdleScope(ERenderThreadIdleTypes::WaitingForGPUQuery);
 
 					// CPU wait for query results to be ready.
-					if (!StartQuerySyncPoint.SyncPoint->IsComplete())
+					if (!StartQueryData.SyncPoint->IsComplete())
 					{
-						FVulkanDynamicRHI::Get().ProcessInterruptQueueUntil(StartQuerySyncPoint.SyncPoint);
+						FVulkanDynamicRHI::Get().ProcessInterruptQueueUntil(StartQueryData.SyncPoint);
 					}
-					if (!EndQuerySyncPoint.SyncPoint->IsComplete())
+					if (!EndQueryData.SyncPoint->IsComplete())
 					{
-						FVulkanDynamicRHI::Get().ProcessInterruptQueueUntil(EndQuerySyncPoint.SyncPoint);
+						FVulkanDynamicRHI::Get().ProcessInterruptQueueUntil(EndQueryData.SyncPoint);
 					}
 				}
+			}
 
-				Pool->ResultsBuffer->InvalidateMappedMemory();
-				StartTimeAvailability = Pool->MappedPointer[QueryStartIndex * 2 + 1];
-				EndTimeAvailability = Pool->MappedPointer[QueryEndIndex * 2 + 1];
-				StartTime = Pool->MappedPointer[QueryStartIndex * 2];
-				EndTime = Pool->MappedPointer[QueryEndIndex * 2];
+			if (EndQueryData.SyncPoint->IsComplete() && StartQueryData.SyncPoint->IsComplete())
+			{
+				if (!StartQueryData.Result.Availability)
+				{
+					const VkQueryResultFlags QueryResultFlags = VK_QUERY_RESULT_64_BIT | VK_QUERY_RESULT_WAIT_BIT;
+					const VkResult Result = VulkanRHI::vkGetQueryPoolResults(
+						Device->GetInstanceHandle(), Pool->GetHandle(),
+						QueryStartIndex, 1, sizeof(uint64), (uint64_t*)&StartQueryData.Result.Value, sizeof(uint64), QueryResultFlags);
+					VERIFYVULKANRESULT(Result);
+				}
+				if (!EndQueryData.Result.Availability)
+				{
+					const VkQueryResultFlags QueryResultFlags = VK_QUERY_RESULT_64_BIT | VK_QUERY_RESULT_WAIT_BIT;
+					const VkResult Result = VulkanRHI::vkGetQueryPoolResults(
+						Device->GetInstanceHandle(), Pool->GetHandle(),
+						QueryStartIndex, 1, sizeof(uint64), (uint64_t*)&EndQueryData.Result.Value, sizeof(uint64), QueryResultFlags);
+					VERIFYVULKANRESULT(Result);
+				}
 
 				PendingQueries.Pop();
 				NumPendingQueries--;
 
-				if (EndTime > StartTime && StartTimeAvailability && EndTimeAvailability)
+				if (EndQueryData.Result.Value > StartQueryData.Result.Value)
 				{
-					TotalTime += EndTime - StartTime;
+					// Only keep the most recent result
+					LastTime = EndQueryData.Result.Value - StartQueryData.Result.Value;
 				}
 			}
+			else
+			{
+				break;
+			}
 		}
 
-		PreviousTime = TotalTime;
-		return TotalTime;
+		return LastTime;
 	}
 
 	return 0;
