VulkanRHI: Move queries into their own header

#rnx

diff --git a/Engine/Source/Runtime/VulkanRHI/Private/VulkanQuery.cpp b/Engine/Source/Runtime/VulkanRHI/Private/VulkanQuery.cpp
--- a/Engine/Source/Runtime/VulkanRHI/Private/VulkanQuery.cpp
+++ b/Engine/Source/Runtime/VulkanRHI/Private/VulkanQuery.cpp
@@ -9,6 +9,7 @@
 #include "VulkanResources.h"
 #include "VulkanContext.h"
 #include "VulkanCommandBuffer.h"
+#include "VulkanQuery.h"
 #include "EngineGlobals.h"
 #include "RenderCore.h"
 
diff --git a/Engine/Source/Runtime/VulkanRHI/Private/VulkanQuery.h b/Engine/Source/Runtime/VulkanRHI/Private/VulkanQuery.h
new file mode 100644
--- /dev/null
+++ b/Engine/Source/Runtime/VulkanRHI/Private/VulkanQuery.h
@@ -0,0 +1,204 @@
+// Copyright Epic Games, Inc. All Rights Reserved.
+
+/*=============================================================================
+	VulkanQuery.h: Private Vulkan RHI definitions.
+=============================================================================*/
+
+#pragma once
+
+#include "VulkanPlatform.h"
+
+class FVulkanDevice;
+class FVulkanQueue;
+class FVulkanContextCommon;
+class FVulkanCommandBuffer;
+class FVulkanOcclusionQuery;
+
+using FVulkanSyncPoint = FGraphEvent;
+using FVulkanSyncPointRef = TRefCountPtr<FVulkanSyncPoint>;
+
+class FVulkanQueryPool
+{
+public:
+	FVulkanQueryPool(FVulkanDevice& InDevice, FVulkanContextCommon* InCommonContext, uint32 InMaxQueries, VkQueryType InQueryType, bool bInShouldAddReset = true);
+	virtual ~FVulkanQueryPool();
+
+	uint32 GetMaxQueries() const
+	{
+		return MaxQueries;
+	}
+
+	VkQueryPool GetHandle() const
+	{
+		return QueryPool;
+	}
+
+	uint64 GetResultValue(uint32 Index) const
+	{
+		return QueryOutput[Index];
+	}
+
+	FVulkanDevice& GetDevice()
+	{
+		return Device;
+	}
+
+protected:
+	FVulkanDevice& Device;
+	VkQueryPool QueryPool;
+	const uint32 MaxQueries;
+	const VkQueryType QueryType;
+	TArray<uint64> QueryOutput;
+};
+
+class FVulkanOcclusionQueryPool : public FVulkanQueryPool
+{
+public:
+	FVulkanOcclusionQueryPool(FVulkanDevice& InDevice, FVulkanContextCommon* InContext, uint32 InMaxQueries)
+		: FVulkanQueryPool(InDevice, InContext, InMaxQueries, VK_QUERY_TYPE_OCCLUSION)
+	{
+		AcquiredIndices.AddZeroed(Align(InMaxQueries, 64) / 64);
+		AllocatedQueries.AddZeroed(InMaxQueries);
+	}
+
+	inline uint32 AcquireIndex(FVulkanOcclusionQuery* Query)
+	{
+		check(NumUsedQueries < MaxQueries);
+		const uint32 Index = NumUsedQueries;
+		const uint32 Word = Index / 64;
+		const uint32 Bit = Index % 64;
+		const uint64 Mask = (uint64)1 << (uint64)Bit;
+		const uint64& WordValue = AcquiredIndices[Word];
+		AcquiredIndices[Word] = WordValue | Mask;
+		++NumUsedQueries;
+		ensure(AllocatedQueries[Index] == nullptr);
+		AllocatedQueries[Index] = Query;
+		return Index;
+	}
+
+	inline void ReleaseIndex(uint32 Index)
+	{
+		check(Index < NumUsedQueries);
+		const uint32 Word = Index / 64;
+		const uint32 Bit = Index % 64;
+		const uint64 Mask = (uint64)1 << (uint64)Bit;
+		const uint64& WordValue = AcquiredIndices[Word];
+		ensure((WordValue & Mask) == Mask);
+		AcquiredIndices[Word] = WordValue & (~Mask);
+		AllocatedQueries[Index] = nullptr;
+	}
+
+	void EndBatch()
+	{
+		ensure(State == EState::RHIT_PostBeginBatch);
+		State = EState::RHIT_PostEndBatch;
+	}
+
+	bool CanBeReused();
+
+	bool GetResults();
+
+	void Reset(FVulkanCommandBuffer* InCmdBuffer, uint32 InFrameNumber);
+
+	bool IsStalePool() const;
+
+	void FlushAllocatedQueries();
+
+	enum EState
+	{
+		Undefined,
+		RHIT_PostBeginBatch,
+		RHIT_PostEndBatch,
+		RT_PostGetResults,
+	};
+	EState State = Undefined;
+
+	// frame number when pool was placed into free list
+	uint32 FreedFrameNumber = UINT32_MAX;
+protected:
+	uint32 NumUsedQueries = 0;
+	TArray<FVulkanOcclusionQuery*> AllocatedQueries;
+	TArray<uint64> AcquiredIndices;
+
+	uint32 FrameNumber = UINT32_MAX;
+};
+
+struct FVulkanTimingQueryResult
+{
+	uint64 Value = 0;
+	uint64 Availability = 0;
+};
+
+class FVulkanTimingQueryPool : public FVulkanQueryPool
+{
+public:
+	FVulkanTimingQueryPool(FVulkanDevice& InDevice, FVulkanContextCommon* InContext, uint32 InBufferedCount)
+		: FVulkanQueryPool(InDevice, InContext, InBufferedCount * 2, VK_QUERY_TYPE_TIMESTAMP, false)
+		, BufferedCount(InBufferedCount)
+	{
+		TimestampArray.AddDefaulted(GetMaxQueries());
+	}
+
+	virtual ~FVulkanTimingQueryPool();
+
+	uint32 CurrentTimestampIndex = 0;
+	uint32 NumIssuedTimestamps = 0;
+	const uint32 BufferedCount;
+
+	struct FTimestampData
+	{
+		FVulkanSyncPointRef SyncPoint;
+		FVulkanTimingQueryResult Result;
+	};
+	TArray<FTimestampData> TimestampArray;
+};
+
+class FVulkanRenderQuery : public FRHIRenderQuery
+{
+public:
+	FVulkanRenderQuery(ERenderQueryType InType)
+		: QueryType(InType)
+	{
+	}
+
+	virtual ~FVulkanRenderQuery() {}
+
+	const ERenderQueryType QueryType;
+	uint64 Result = 0;
+
+	uint32 IndexInPool = UINT32_MAX;
+
+	FVulkanSyncPointRef SyncPoint;
+};
+
+class FVulkanOcclusionQuery : public FVulkanRenderQuery
+{
+public:
+	FVulkanOcclusionQuery();
+	virtual ~FVulkanOcclusionQuery();
+
+	enum class EState
+	{
+		Undefined,
+		RHI_PostBegin,
+		RHI_PostEnd,
+		RT_GotResults,
+		FlushedFromPoolHadResults,
+	};
+
+	FVulkanOcclusionQueryPool* Pool = nullptr;
+
+	void ReleaseFromPool();
+
+	EState State = EState::Undefined;
+};
+
+class FVulkanTimingQuery : public FVulkanRenderQuery
+{
+public:
+	FVulkanTimingQuery();
+	virtual ~FVulkanTimingQuery();
+
+	FVulkanTimingQueryPool* Pool = nullptr;
+};
+
diff --git a/Engine/Source/Runtime/VulkanRHI/Private/VulkanRayTracing.h b/Engine/Source/Runtime/VulkanRHI/Private/VulkanRayTracing.h
--- a/Engine/Source/Runtime/VulkanRHI/Private/VulkanRayTracing.h
+++ b/Engine/Source/Runtime/VulkanRHI/Private/VulkanRayTracing.h
@@ -6,6 +6,7 @@
 
 #include "RayTracingBuiltInResources.h"
 #include "VulkanPipeline.h"
+#include "VulkanQuery.h"
 
 class FVulkanCommandListContext;
 class FVulkanBuffer;
diff --git a/Engine/Source/Runtime/VulkanRHI/Private/VulkanResources.h b/Engine/Source/Runtime/VulkanRHI/Private/VulkanResources.h
--- a/Engine/Source/Runtime/VulkanRHI/Private/VulkanResources.h
+++ b/Engine/Source/Runtime/VulkanRHI/Private/VulkanResources.h
@@ -31,6 +31,7 @@ class FVulkanView;
 class FVulkanViewableResource;
 class FVulkanShaderResourceView;
 class FVulkanUnorderedAccessView;
+class FVulkanRenderQuery;
 
 using FVulkanSyncPoint = FGraphEvent;
 using FVulkanSyncPointRef = TRefCountPtr<FVulkanSyncPoint>;
@@ -783,190 +784,6 @@ protected:
 	EImageOwnerType ImageOwnerType;
 };
 
-class FVulkanQueryPool
-{
-public:
-	FVulkanQueryPool(FVulkanDevice& InDevice, FVulkanContextCommon* InCommonContext, uint32 InMaxQueries, VkQueryType InQueryType, bool bInShouldAddReset = true);
-	virtual ~FVulkanQueryPool();
-
-	uint32 GetMaxQueries() const
-	{
-		return MaxQueries;
-	}
-
-	VkQueryPool GetHandle() const
-	{
-		return QueryPool;
-	}
-
-	uint64 GetResultValue(uint32 Index) const
-	{
-		return QueryOutput[Index];
-	}
-
-	FVulkanDevice& GetDevice()
-	{
-		return Device;
-	}
-
-protected:
-	FVulkanDevice& Device;
-	VkQueryPool QueryPool;
-	const uint32 MaxQueries;
-	const VkQueryType QueryType;
-	TArray<uint64> QueryOutput;
-};
-
-class FVulkanOcclusionQueryPool : public FVulkanQueryPool
-{
-public:
-	FVulkanOcclusionQueryPool(FVulkanDevice& InDevice, FVulkanContextCommon* InContext, uint32 InMaxQueries)
-		: FVulkanQueryPool(InDevice, InContext, InMaxQueries, VK_QUERY_TYPE_OCCLUSION)
-	{
-		AcquiredIndices.AddZeroed(Align(InMaxQueries, 64) / 64);
-		AllocatedQueries.AddZeroed(InMaxQueries);
-	}
-
-	inline uint32 AcquireIndex(FVulkanOcclusionQuery* Query)
-	{
-		check(NumUsedQueries < MaxQueries);
-		const uint32 Index = NumUsedQueries;
-		const uint32 Word = Index / 64;
-		const uint32 Bit = Index % 64;
-		const uint64 Mask = (uint64)1 << (uint64)Bit;
-		const uint64& WordValue = AcquiredIndices[Word];
-		AcquiredIndices[Word] = WordValue | Mask;
-		++NumUsedQueries;
-		ensure(AllocatedQueries[Index] == nullptr);
-		AllocatedQueries[Index] = Query;
-		return Index;
-	}
-
-	inline void ReleaseIndex(uint32 Index)
-	{
-		check(Index < NumUsedQueries);
-		const uint32 Word = Index / 64;
-		const uint32 Bit = Index % 64;
-		const uint64 Mask = (uint64)1 << (uint64)Bit;
-		const uint64& WordValue = AcquiredIndices[Word];
-		ensure((WordValue & Mask) == Mask);
-		AcquiredIndices[Word] = WordValue & (~Mask);
-		AllocatedQueries[Index] = nullptr;
-	}
-
-	void EndBatch()
-	{
-		ensure(State == EState::RHIT_PostBeginBatch);
-		State = EState::RHIT_PostEndBatch;
-	}
-
-	bool CanBeReused();
-
-	bool GetResults();
-
-	void Reset(FVulkanCommandBuffer* InCmdBuffer, uint32 InFrameNumber);
-
-	bool IsStalePool() const;
-
-	void FlushAllocatedQueries();
-
-	enum EState
-	{
-		Undefined,
-		RHIT_PostBeginBatch,
-		RHIT_PostEndBatch,
-		RT_PostGetResults,
-	};
-	EState State = Undefined;
-	
-	// frame number when pool was placed into free list
-	uint32 FreedFrameNumber = UINT32_MAX;
-protected:
-	uint32 NumUsedQueries = 0;
-	TArray<FVulkanOcclusionQuery*> AllocatedQueries;
-	TArray<uint64> AcquiredIndices;
-
-	uint32 FrameNumber = UINT32_MAX;
-};
-
-struct FVulkanTimingQueryResult
-{
-	uint64 Value = 0;
-	uint64 Availability = 0;
-};
-
-class FVulkanTimingQueryPool : public FVulkanQueryPool
-{
-public:
-	FVulkanTimingQueryPool(FVulkanDevice& InDevice, FVulkanContextCommon* InContext, uint32 InBufferedCount)
-		: FVulkanQueryPool(InDevice, InContext, InBufferedCount * 2, VK_QUERY_TYPE_TIMESTAMP, false)
-		, BufferedCount(InBufferedCount)
-	{
-		TimestampArray.AddDefaulted(GetMaxQueries());
-	}
-
-	virtual ~FVulkanTimingQueryPool();
-
-	uint32 CurrentTimestampIndex = 0;
-	uint32 NumIssuedTimestamps = 0;
-	const uint32 BufferedCount;
-
-	struct FTimestampData
-	{
-		FVulkanSyncPointRef SyncPoint;
-		FVulkanTimingQueryResult Result;
-	};
-	TArray<FTimestampData> TimestampArray;
-};
-
-class FVulkanRenderQuery : public FRHIRenderQuery
-{
-public:
-	FVulkanRenderQuery(ERenderQueryType InType)
-		: QueryType(InType)
-	{
-	}
-
-	virtual ~FVulkanRenderQuery() {}
-
-	const ERenderQueryType QueryType;
-	uint64 Result = 0;
-
-	uint32 IndexInPool = UINT32_MAX;
-
-	FVulkanSyncPointRef SyncPoint;
-};
-
-class FVulkanOcclusionQuery : public FVulkanRenderQuery
-{
-public:
-	FVulkanOcclusionQuery();
-	virtual ~FVulkanOcclusionQuery();
-
-	enum class EState
-	{
-		Undefined,
-		RHI_PostBegin,
-		RHI_PostEnd,
-		RT_GotResults,
-		FlushedFromPoolHadResults,
-	};
-
-	FVulkanOcclusionQueryPool* Pool = nullptr;
-
-	void ReleaseFromPool();
-
-	EState State = EState::Undefined;
-};
-
-class FVulkanTimingQuery : public FVulkanRenderQuery
-{
-public:
-	FVulkanTimingQuery();
-	virtual ~FVulkanTimingQuery();
-
-	FVulkanTimingQueryPool* Pool = nullptr;
-};
 
 class FVulkanBuffer : public FRHIBuffer, public VulkanRHI::FDeviceChild, public FVulkanViewableResource
 {
diff --git a/Engine/Source/Runtime/VulkanRHI/Private/VulkanUtil.cpp b/Engine/Source/Runtime/VulkanRHI/Private/VulkanUtil.cpp
--- a/Engine/Source/Runtime/VulkanRHI/Private/VulkanUtil.cpp
+++ b/Engine/Source/Runtime/VulkanRHI/Private/VulkanUtil.cpp
@@ -19,6 +19,7 @@
 #include "RenderCore.h"
 #include "RHICoreNvidiaAftermath.h"
 #include "VulkanRayTracing.h"
+#include "VulkanQuery.h"
 
 FVulkanDynamicRHI*	GVulkanRHI = nullptr;
 
