Vulkan: Expose function for global command addresses.

vkGetInstanceProcAddr's only exposure forces a non-null instance parameter, but some Vulkan functions are "global" and can only be retrieved with a null instance value.

#rn

diff --git a/Engine/Source/Runtime/VulkanRHI/Private/VulkanRHI.cpp b/Engine/Source/Runtime/VulkanRHI/Private/VulkanRHI.cpp
--- a/Engine/Source/Runtime/VulkanRHI/Private/VulkanRHI.cpp
+++ b/Engine/Source/Runtime/VulkanRHI/Private/VulkanRHI.cpp
@@ -1400,6 +1400,11 @@ void* FVulkanDynamicRHI::RHIGetVkInstanceProcAddr(const char* InName) const
 	return (void*)VulkanRHI::vkGetInstanceProcAddr(Instance, InName);
 }
 
+void* FVulkanDynamicRHI::RHIGetVkInstanceGlobalProcAddr(const char* InName) const
+{
+	return (void*)VulkanRHI::vkGetInstanceProcAddr(nullptr, InName);
+}
+
 VkFormat FVulkanDynamicRHI::RHIGetSwapChainVkFormat(EPixelFormat InFormat) const
 {
 	// UE renders a gamma-corrected image so we need to use an sRGB format if available
diff --git a/Engine/Source/Runtime/VulkanRHI/Public/IVulkanDynamicRHI.h b/Engine/Source/Runtime/VulkanRHI/Public/IVulkanDynamicRHI.h
--- a/Engine/Source/Runtime/VulkanRHI/Public/IVulkanDynamicRHI.h
+++ b/Engine/Source/Runtime/VulkanRHI/Public/IVulkanDynamicRHI.h
@@ -51,6 +51,11 @@ struct IVulkanDynamicRHI : public FDynamicRHI
 	virtual bool             RHIDoesAdapterMatchDevice(const void* InAdapterId) const = 0;
 	virtual void*            RHIGetVkDeviceProcAddr(const char* InName) const = 0;
 	virtual void*            RHIGetVkInstanceProcAddr(const char* InName) const = 0;
+	/**
+	 * Version of RHIGetVkInstanceProcAddr that uses nullptr as the instance argument.
+	 * See vkGetInstanceProcAddr manpage for distinction between "global" and non-global commands.
+	 */
+	virtual void*            RHIGetVkInstanceGlobalProcAddr(const char* InName) const = 0;
 	virtual VkFormat         RHIGetSwapChainVkFormat(EPixelFormat InFormat) const = 0;
 	virtual bool             RHISupportsEXTFragmentDensityMap2() const = 0;
 
diff --git a/Engine/Source/Runtime/VulkanRHI/Public/VulkanDynamicRHI.h b/Engine/Source/Runtime/VulkanRHI/Public/VulkanDynamicRHI.h
--- a/Engine/Source/Runtime/VulkanRHI/Public/VulkanDynamicRHI.h
+++ b/Engine/Source/Runtime/VulkanRHI/Public/VulkanDynamicRHI.h
@@ -78,6 +78,7 @@ public:
 	virtual bool RHIDoesAdapterMatchDevice(const void* InAdapterId) const final override;
 	virtual void* RHIGetVkDeviceProcAddr(const char* InName) const final override;
 	virtual void* RHIGetVkInstanceProcAddr(const char* InName) const final override;
+	virtual void* RHIGetVkInstanceGlobalProcAddr(const char* InName) const final override;
 	virtual VkFormat RHIGetSwapChainVkFormat(EPixelFormat InFormat) const final override;
 	virtual bool RHISupportsEXTFragmentDensityMap2() const final override;
 	virtual TArray<VkExtensionProperties> RHIGetAllInstanceExtensions() const final override;
