VulkanRHI: Fix missing View binding crash in Vulkan bindless.  Add check during binding with proper information for next time.

#rnx

diff --git a/Engine/Shaders/Private/VirtualShadowMaps/VirtualShadowMapCacheGPUInvalidation.usf b/Engine/Shaders/Private/VirtualShadowMaps/VirtualShadowMapCacheGPUInvalidation.usf
--- a/Engine/Shaders/Private/VirtualShadowMaps/VirtualShadowMapCacheGPUInvalidation.usf
+++ b/Engine/Shaders/Private/VirtualShadowMaps/VirtualShadowMapCacheGPUInvalidation.usf
@@ -1,5 +1,8 @@
 // Copyright Epic Games, Inc. All Rights Reserved.
 
+// Do not use shared samplers as it requires the View uniform buffer, which is not bound for this shader.
+#define USE_HZB_SHARED_SAMPLERS 0
+
 #include "../Common.ush"
 #include "../SceneData.ush"
 #include "../ViewData.ush"
diff --git a/Engine/Shaders/Private/VirtualShadowMaps/VirtualShadowMapCacheLoadBalancer.usf b/Engine/Shaders/Private/VirtualShadowMaps/VirtualShadowMapCacheLoadBalancer.usf
--- a/Engine/Shaders/Private/VirtualShadowMaps/VirtualShadowMapCacheLoadBalancer.usf
+++ b/Engine/Shaders/Private/VirtualShadowMaps/VirtualShadowMapCacheLoadBalancer.usf
@@ -1,5 +1,8 @@
 // Copyright Epic Games, Inc. All Rights Reserved.
 
+// Do not use shared samplers as it requires the View uniform buffer, which is not bound for this shader.
+#define USE_HZB_SHARED_SAMPLERS 0
+
 #include "../Common.ush"
 #include "../SceneData.ush"
 #include "VirtualShadowMapPageCacheCommon.ush"
diff --git a/Engine/Source/Runtime/VulkanRHI/Private/VulkanPipelineState.cpp b/Engine/Source/Runtime/VulkanRHI/Private/VulkanPipelineState.cpp
--- a/Engine/Source/Runtime/VulkanRHI/Private/VulkanPipelineState.cpp
+++ b/Engine/Source/Runtime/VulkanRHI/Private/VulkanPipelineState.cpp
@@ -267,7 +267,10 @@ void FVulkanComputePipelineDescriptorState::UpdateBindlessDescriptors(FVulkanCom
 		check(WriteDescriptorSet.dstArrayElement == 0);
 		check(WriteDescriptorSet.descriptorCount == 1);
 		check(WriteDescriptorSet.descriptorType == VK_DESCRIPTOR_TYPE_UNIFORM_BUFFER);
-		checkSlow(WriteDescriptorSet.pBufferInfo);
+		checkf(WriteDescriptorSet.pBufferInfo && WriteDescriptorSet.pBufferInfo->buffer, 
+			TEXT("Missing uniform buffer binding for [%s] at index [%d] of shader [%s]."),
+			*ComputePipeline->GetComputeShader()->GetUniformBufferName(UBIndex), UBIndex,
+			ComputePipeline->GetComputeShader()->GetShaderName());
 
 		VkBufferDeviceAddressInfo BufferInfo;
 		ZeroVulkanStruct(BufferInfo, VK_STRUCTURE_TYPE_BUFFER_DEVICE_ADDRESS_INFO);
