Vulkan RHI: Add missing public include path from Vulkan module.

The VulkanRHI module's Public/IVulkanDynamicRHI.h includes VulkanThirdParty.h which lives in the Vulkan module.  This change is necessary for modules that declare VulkanRHI as a dependency to be able to include this important public header without encountering compile errors, or having to themselves declare a Vulkan module include path dependency (which would be a case of "leaking implementation details").

diff --git a/Engine/Source/Runtime/VulkanRHI/VulkanRHI.Build.cs b/Engine/Source/Runtime/VulkanRHI/VulkanRHI.Build.cs
--- a/Engine/Source/Runtime/VulkanRHI/VulkanRHI.Build.cs
+++ b/Engine/Source/Runtime/VulkanRHI/VulkanRHI.Build.cs
@@ -72,6 +72,7 @@ public class VulkanRHI : ModuleRules
 			|| Target.IsInPlatformGroup(UnrealPlatformGroup.Unix) || Target.Platform == UnrealTargetPlatform.Mac)
 		{
             AddEngineThirdPartyPrivateStaticDependencies(Target, "Vulkan");
+			PublicIncludePathModuleNames.Add("Vulkan");
         }
 
 		if (Target.IsInPlatformGroup(UnrealPlatformGroup.Linux))
