Build fix for Vulkan raytracing inline binding data buffer

diff --git a/Engine/Source/Runtime/VulkanRHI/Private/VulkanRayTracing.cpp b/Engine/Source/Runtime/VulkanRHI/Private/VulkanRayTracing.cpp
--- a/Engine/Source/Runtime/VulkanRHI/Private/VulkanRayTracing.cpp
+++ b/Engine/Source/Runtime/VulkanRHI/Private/VulkanRayTracing.cpp
@@ -1019,7 +1019,7 @@ FRHISizeAndStride FVulkanRayTracingShaderTable::GetInlineBindingDataSizeAndStrid
 
 FRHISizeAndStride FVulkanRayTracingShaderTable::GetInlineBindingDataSizeAndStride() const
 {	
-	return FRHISizeAndStride { InlineGeometryParameterData.Num(), sizeof(FVulkanRayTracingGeometryParameters) };
+	return FRHISizeAndStride { (uint64)InlineGeometryParameterData.Num(), sizeof(FVulkanRayTracingGeometryParameters) };
 }
 
 #endif // VULKAN_SUPPORTS_RAY_TRACING_POSITION_FETCH
@@ -1090,10 +1090,11 @@ void FVulkanRayTracingShaderTable::Commit(FVulkanCommandListContext& Context, FR
 	{
 		TRHICommandList_RecursiveHazardous<FVulkanCommandListContext> RHICmdList(&Context);
 
-		const uint32 ParameterBufferSize = InlineGeometryParameterData.Num();		
-		void* MappedBuffer = InlineBindingDataBuffer->Lock(RHICmdList, RLM_WriteOnly, ParameterBufferSize, 0);
+		const uint32 ParameterBufferSize = InlineGeometryParameterData.Num();
+		FVulkanBuffer* VulkanBuffer = (FVulkanBuffer*)InlineBindingDataBuffer;
+		void* MappedBuffer = VulkanBuffer->Lock(RHICmdList, RLM_WriteOnly, ParameterBufferSize, 0);
 		FMemory::Memcpy(MappedBuffer, InlineGeometryParameterData.GetData(), ParameterBufferSize);
-		InlineBindingDataBuffer->Unlock(RHICmdList);
+		VulkanBuffer->Unlock(RHICmdList);
 	}
 #endif // !VULKAN_SUPPORTS_RAY_TRACING_POSITION_FETCH
 
