VulkanRHI: Move FVulkanDynamicRHI definition to Private.  Deprecate original header.  Remove deprecated VulkanRHIBridge.h.

#rnx

diff --git a/Engine/Source/Runtime/VulkanRHI/Private/VulkanDynamicRHI.h b/Engine/Source/Runtime/VulkanRHI/Private/VulkanDynamicRHI.h
new file mode 100644
--- /dev/null
+++ b/Engine/Source/Runtime/VulkanRHI/Private/VulkanDynamicRHI.h
@@ -0,0 +1,430 @@
+// Copyright Epic Games, Inc. All Rights Reserved.
+
+/*=============================================================================
+	VulkanDynamicRHI.h: Vulkan RHI definitions.
+=============================================================================*/
+
+#pragma once
+
+#include "IVulkanDynamicRHI.h"
+#include "VulkanConfiguration.h"
+
+class FVulkanFramebuffer;
+class FVulkanDevice;
+class FVulkanQueue;
+class FVulkanUniformBuffer;
+class FVulkanViewport;
+class FVulkanPayload;
+class FVulkanTexture;
+class IHeadMountedDisplayVulkanExtensions;
+class FVulkanThread;
+struct FRHITransientHeapAllocation;
+struct FVulkanPlatformCommandList;
+
+namespace VulkanRHI
+{
+	class FSemaphore;
+};
+
+struct FOptionalVulkanInstanceExtensions
+{
+	union
+	{
+		struct
+		{
+			UE_DEPRECATED(5.3, "Vulkan 1.1 is now a requirement so there is no need to check these capabilities.")
+			uint32 HasKHRExternalMemoryCapabilities : 1;
+
+			UE_DEPRECATED(5.3, "Vulkan 1.1 is now a requirement so there is no need to check these capabilities.")
+			uint32 HasKHRGetPhysicalDeviceProperties2 : 1;
+		};
+		uint32 Packed;
+	};
+
+	FOptionalVulkanInstanceExtensions()
+	{
+		static_assert(sizeof(Packed) == sizeof(FOptionalVulkanInstanceExtensions), "More bits needed!");
+		Packed = 0;
+	}
+};
+
+PRAGMA_DISABLE_DEPRECATION_WARNINGS
+
+/** The interface which is implemented by the dynamically bound RHI. */
+class FVulkanDynamicRHI : public IVulkanDynamicRHI
+{
+public:
+
+	static FVulkanDynamicRHI& Get() { return *GetDynamicRHI<FVulkanDynamicRHI>(); }
+
+	/** Initialization constructor. */
+	FVulkanDynamicRHI();
+
+	/** Destructor */
+	~FVulkanDynamicRHI();
+
+	// IVulkanDynamicRHI interface
+	virtual uint32 RHIGetVulkanVersion() const final override;
+	virtual VkInstance RHIGetVkInstance() const final override;
+	virtual VkDevice RHIGetVkDevice() const final override;
+	virtual const uint8* RHIGetVulkanDeviceUUID() const final override;
+	virtual VkPhysicalDevice RHIGetVkPhysicalDevice() const final override;
+	virtual const VkAllocationCallbacks* RHIGetVkAllocationCallbacks() final override;
+	virtual VkQueue RHIGetGraphicsVkQueue() const final override;
+	virtual uint32 RHIGetGraphicsQueueIndex() const final override;
+	virtual uint32 RHIGetGraphicsQueueFamilyIndex() const final override;
+	virtual VkCommandBuffer RHIGetActiveVkCommandBuffer() final override;
+	virtual uint64 RHIGetGraphicsAdapterLUID(VkPhysicalDevice InPhysicalDevice) const final override;
+	virtual bool RHIDoesAdapterMatchDevice(const void* InAdapterId) const final override;
+	virtual void* RHIGetVkDeviceProcAddr(const char* InName) const final override;
+	virtual void* RHIGetVkInstanceProcAddr(const char* InName) const final override;
+	virtual void* RHIGetVkInstanceGlobalProcAddr(const char* InName) const final override;
+	virtual VkFormat RHIGetSwapChainVkFormat(EPixelFormat InFormat) const final override;
+	virtual bool RHISupportsEXTFragmentDensityMap2() const final override;
+	virtual TArray<VkExtensionProperties> RHIGetAllInstanceExtensions() const final override;
+	virtual TArray<VkExtensionProperties> RHIGetAllDeviceExtensions(VkPhysicalDevice InPhysicalDevice) const final override;
+	virtual TArray<FAnsiString> RHIGetLoadedDeviceExtensions() const final override;
+	virtual FTextureRHIRef RHICreateTexture2DFromResource(EPixelFormat Format, uint32 SizeX, uint32 SizeY, uint32 NumMips, uint32 NumSamples, VkImage Resource, ETextureCreateFlags Flags, const FClearValueBinding& ClearValueBinding = FClearValueBinding::Transparent, const FVulkanRHIExternalImageDeleteCallbackInfo& ExternalImageDeleteCallbackInfo = {}) final override;
+#if PLATFORM_ANDROID
+	virtual FTextureRHIRef RHICreateTexture2DFromAndroidHardwareBuffer(AHardwareBuffer* HardwareBuffer) final override;
+#endif
+	virtual FTextureRHIRef RHICreateTexture2DArrayFromResource(EPixelFormat Format, uint32 SizeX, uint32 SizeY, uint32 ArraySize, uint32 NumMips, uint32 NumSamples, VkImage Resource, ETextureCreateFlags Flags, const FClearValueBinding& ClearValueBinding = FClearValueBinding::Transparent) final override;
+	virtual FTextureRHIRef RHICreateTextureCubeFromResource(EPixelFormat Format, uint32 Size, bool bArray, uint32 ArraySize, uint32 NumMips, VkImage Resource, ETextureCreateFlags Flags, const FClearValueBinding& ClearValueBinding = FClearValueBinding::Transparent) final override;
+	virtual VkImage RHIGetVkImage(FRHITexture* InTexture) const final override;
+	virtual VkFormat RHIGetViewVkFormat(FRHITexture* InTexture) const final override;
+	virtual FVulkanRHIAllocationInfo RHIGetAllocationInfo(FRHITexture* InTexture) const final override;
+	virtual FVulkanRHIImageViewInfo RHIGetImageViewInfo(FRHITexture* InTexture) const final override;
+	virtual FVulkanRHIAllocationInfo RHIGetAllocationInfo(FRHIBuffer* InBuffer) const final override;
+	virtual void RHISetImageLayout(VkImage Image, VkImageLayout OldLayout, VkImageLayout NewLayout, const VkImageSubresourceRange& SubresourceRange) final override;
+	virtual void RHISetUploadImageLayout(VkImage Image, VkImageLayout OldLayout, VkImageLayout NewLayout, const VkImageSubresourceRange& SubresourceRange) final override;
+	virtual void RHIFinishExternalComputeWork(VkCommandBuffer InCommandBuffer) final override;
+	virtual void RHIRegisterWork(uint32 NumPrimitives) final override;
+	virtual void RHISubmitUploadCommandBuffer() final override;
+	virtual void RHIVerifyResult(VkResult Result, const ANSICHAR* VkFuntion, const ANSICHAR* Filename, uint32 Line) final override;
+
+	// FDynamicRHI interface.
+	virtual void Init() final override;
+	virtual void PostInit() final override;
+	virtual void Shutdown() final override;
+	virtual const TCHAR* GetName() final override { return TEXT("Vulkan"); }
+
+	void InitInstance();
+
+	virtual void RHIEndFrame_RenderThread(FRHICommandListImmediate& RHICmdList) final override;
+	virtual void RHIEndFrame(const FRHIEndFrameArgs& Args) final override;
+
+	virtual FSamplerStateRHIRef RHICreateSamplerState(const FSamplerStateInitializerRHI& Initializer) final override;
+	virtual FRasterizerStateRHIRef RHICreateRasterizerState(const FRasterizerStateInitializerRHI& Initializer) final override;
+	virtual FDepthStencilStateRHIRef RHICreateDepthStencilState(const FDepthStencilStateInitializerRHI& Initializer) final override;
+	virtual FBlendStateRHIRef RHICreateBlendState(const FBlendStateInitializerRHI& Initializer) final override;
+	virtual FVertexDeclarationRHIRef RHICreateVertexDeclaration(const FVertexDeclarationElementList& Elements) final override;
+	virtual FPixelShaderRHIRef RHICreatePixelShader(TArrayView<const uint8> Code, const FSHAHash& Hash) final override;
+	virtual FVertexShaderRHIRef RHICreateVertexShader(TArrayView<const uint8> Code, const FSHAHash& Hash) final override;
+	virtual FMeshShaderRHIRef RHICreateMeshShader(TArrayView<const uint8> Code, const FSHAHash& Hash) final override;
+	virtual FAmplificationShaderRHIRef RHICreateAmplificationShader(TArrayView<const uint8> Code, const FSHAHash& Hash) final override;
+	virtual FGeometryShaderRHIRef RHICreateGeometryShader(TArrayView<const uint8> Code, const FSHAHash& Hash) final override;
+	virtual FComputeShaderRHIRef RHICreateComputeShader(TArrayView<const uint8> Code, const FSHAHash& Hash) final override;
+	virtual FGPUFenceRHIRef RHICreateGPUFence(const FName &Name) final override;
+	virtual void RHICreateTransition(FRHITransition* Transition, const FRHITransitionCreateInfo& CreateInfo) final override;
+	virtual void RHIReleaseTransition(FRHITransition* Transition) final override;
+	virtual FStagingBufferRHIRef RHICreateStagingBuffer() final override;
+	virtual FBoundShaderStateRHIRef RHICreateBoundShaderState(FRHIVertexDeclaration* VertexDeclaration, FRHIVertexShader* VertexShader, FRHIPixelShader* PixelShader, FRHIGeometryShader* GeometryShader) final override;
+	virtual FGraphicsPipelineStateRHIRef RHICreateGraphicsPipelineState(const FGraphicsPipelineStateInitializer& Initializer) final override;
+	virtual FComputePipelineStateRHIRef RHICreateComputePipelineState(FRHIComputeShader* ComputeShader) final override;
+	virtual FUniformBufferRHIRef RHICreateUniformBuffer(const void* Contents, const FRHIUniformBufferLayout* Layout, EUniformBufferUsage Usage, EUniformBufferValidation Validation) final override;
+	virtual void RHIUpdateUniformBuffer(FRHICommandListBase& RHICmdList, FRHIUniformBuffer* UniformBufferRHI, const void* Contents) final override;
+
+	[[nodiscard]] virtual FRHIBufferInitializer RHICreateBufferInitializer(FRHICommandListBase& RHICmdList, const FRHIBufferCreateDesc& CreateDesc) final override;
+
+	virtual void RHIReplaceResources(FRHICommandListBase& RHICmdList, TArray<FRHIResourceReplaceInfo>&& ReplaceInfos) final override;
+	virtual void* LockBuffer_BottomOfPipe(FRHICommandListBase& RHICmdList, FRHIBuffer* Buffer, uint32 Offset, uint32 Size, EResourceLockMode LockMode) final override;
+	virtual void UnlockBuffer_BottomOfPipe(FRHICommandListBase& RHICmdList, FRHIBuffer* Buffer) final override;
+	virtual void* RHILockBuffer(FRHICommandListBase& RHICmdList, FRHIBuffer* BufferRHI, uint32 Offset, uint32 Size, EResourceLockMode LockMode) final override;
+	virtual void RHIUnlockBuffer(FRHICommandListBase& RHICmdList, FRHIBuffer* Buffer) final override;
+	virtual FTextureReferenceRHIRef RHICreateTextureReference(FRHICommandListBase& RHICmdList, FRHITexture* InReferencedTexture) final override;
+	virtual void RHIUpdateTextureReference(FRHICommandListBase& RHICmdList, FRHITextureReference* TextureRef, FRHITexture* NewTexture) final override;
+	virtual FRHICalcTextureSizeResult RHICalcTexturePlatformSize(FRHITextureDesc const& Desc, uint32 FirstMipIndex) final override;
+	virtual void RHIGetTextureMemoryStats(FTextureMemoryStats& OutStats) final override;
+	virtual bool RHIGetTextureMemoryVisualizeData(FColor* TextureData, int32 SizeX, int32 SizeY, int32 Pitch, int32 PixelSize) final override;
+
+	struct FCreateTextureResult
+	{
+		FVulkanTexture* Texture;
+		VkImageLayout DefaultLayout;
+		bool bTransientResource;
+		bool bClear;
+	};
+
+	FCreateTextureResult BeginCreateTextureInternal(const FRHITextureCreateDesc& CreateDesc, const FRHITransientHeapAllocation* InTransientHeapAllocation);
+	FVulkanTexture* FinalizeCreateTextureInternal(FRHICommandListBase& RHICmdList, FCreateTextureResult CreateResult, TConstArrayView<uint8> InitialData);
+
+	FVulkanTexture* CreateTextureInternal(FRHICommandListBase& RHICmdList, const FRHITextureCreateDesc& CreateDesc, TConstArrayView<uint8> InitialData);
+	FVulkanTexture* CreateTextureInternal(const FRHITextureCreateDesc& CreateDesc, const FRHITransientHeapAllocation& InTransientHeapAllocation);
+
+	virtual FTextureRHIRef RHICreateTexture(FRHICommandListBase& RHICmdList, const FRHITextureCreateDesc& CreateDesc) final override;
+	virtual FTextureRHIRef RHIAsyncCreateTexture2D(uint32 SizeX, uint32 SizeY, uint8 Format, uint32 NumMips, ETextureCreateFlags Flags, ERHIAccess InResourceState, void** InitialMipData, uint32 NumInitialMips, const TCHAR* DebugName, FGraphEventRef& OutCompletionEvent) final override;
+	virtual uint32 RHIComputeMemorySize(FRHITexture* TextureRHI) final override;
+	virtual FTextureRHIRef RHIAsyncReallocateTexture2D(FRHITexture* Texture2D, int32 NewMipCount, int32 NewSizeX, int32 NewSizeY, FThreadSafeCounter* RequestStatus) final override;
+
+	virtual FRHILockTextureResult RHILockTexture(FRHICommandListImmediate& RHICmdList, const FRHILockTextureArgs& Arguments) final override;
+	virtual void RHIUnlockTexture(FRHICommandListImmediate& RHICmdList, const FRHILockTextureArgs& Arguments) final override;
+
+	virtual void RHIUpdateTexture2D(FRHICommandListBase& RHICmdList, FRHITexture* Texture, uint32 MipIndex, const struct FUpdateTextureRegion2D& UpdateRegion, uint32 SourcePitch, const uint8* SourceData) final override
+	{
+		InternalUpdateTexture2D(RHICmdList, Texture, MipIndex, UpdateRegion, SourcePitch, SourceData);
+	}
+	virtual void RHIUpdateTexture3D(FRHICommandListBase& RHICmdList, FRHITexture* Texture, uint32 MipIndex, const struct FUpdateTextureRegion3D& UpdateRegion, uint32 SourceRowPitch, uint32 SourceDepthPitch, const uint8* SourceData) final override
+	{
+		InternalUpdateTexture3D(RHICmdList, Texture, MipIndex, UpdateRegion, SourceRowPitch, SourceDepthPitch, SourceData);
+	}
+	virtual void RHIBindDebugLabelName(FRHICommandListBase& RHICmdList, FRHITexture* Texture, const TCHAR* Name) final override;
+	virtual void RHIBindDebugLabelName(FRHICommandListBase& RHICmdList, FRHIUnorderedAccessView* UnorderedAccessViewRHI, const TCHAR* Name) final override;
+	virtual void RHIReadSurfaceData(FRHITexture* Texture, FIntRect Rect, TArray<FColor>& OutData, FReadSurfaceDataFlags InFlags) final override;
+	virtual void RHIReadSurfaceData(FRHITexture* Texture, FIntRect Rect, TArray<FLinearColor>& OutData, FReadSurfaceDataFlags InFlags) final override;
+	virtual void RHIMapStagingSurface(FRHITexture* Texture, FRHIGPUFence* Fence, void*& OutData, int32& OutWidth, int32& OutHeight, uint32 GPUIndex = 0) final override;
+	virtual void RHIUnmapStagingSurface(FRHITexture* Texture, uint32 GPUIndex = 0) final override;
+
+	virtual void RHIReadSurfaceFloatData(FRHITexture* Texture, FIntRect Rect, TArray<FFloat16Color>& OutData, ECubeFace CubeFace, int32 ArrayIndex, int32 MipIndex) final override;
+	virtual void RHIRead3DSurfaceFloatData(FRHITexture* Texture, FIntRect Rect, FIntPoint ZMinMax, TArray<FFloat16Color>& OutData) final override;
+	virtual FRenderQueryRHIRef RHICreateRenderQuery(ERenderQueryType QueryType) final override;
+	virtual void RHIEndRenderQuery_TopOfPipe(FRHICommandListBase& RHICmdList, FRHIRenderQuery* RenderQuery) final override;
+	virtual void RHIBeginRenderQueryBatch_TopOfPipe(FRHICommandListBase& RHICmdList, ERenderQueryType QueryType) final override;
+	virtual void RHIEndRenderQueryBatch_TopOfPipe(FRHICommandListBase& RHICmdList, ERenderQueryType QueryType) final override;
+	virtual bool RHIGetRenderQueryResult(FRHIRenderQuery* RenderQuery, uint64& OutResult, bool bWait, uint32 GPUIndex = INDEX_NONE) final override;
+	virtual FTextureRHIRef RHIGetViewportBackBuffer(FRHIViewport* Viewport) final override;
+	virtual void RHIAliasTextureResources(FTextureRHIRef& DestTexture, FTextureRHIRef& SrcTexture) final override;
+	virtual FTextureRHIRef RHICreateAliasedTexture(FTextureRHIRef& SourceTexture) final override;
+	virtual void RHIAdvanceFrameForGetViewportBackBuffer(FRHIViewport* Viewport) final override;
+	virtual void RHIFlushResources() final override;
+	virtual uint32 RHIGetGPUFrameCycles(uint32 GPUIndex = 0) final override;
+	virtual FViewportRHIRef RHICreateViewport(void* WindowHandle, uint32 SizeX, uint32 SizeY, bool bIsFullscreen, EPixelFormat PreferredPixelFormat) final override;
+	virtual void RHIResizeViewport(FRHIViewport* Viewport, uint32 SizeX, uint32 SizeY, bool bIsFullscreen) final override;
+	virtual void RHIResizeViewport(FRHIViewport* Viewport, uint32 SizeX, uint32 SizeY, bool bIsFullscreen, EPixelFormat PreferredPixelFormat) final override;
+	virtual void RHITick(float DeltaTime) final override;
+	virtual void RHIBlockUntilGPUIdle() final override;
+	virtual void RHISuspendRendering() final override;
+	virtual void RHIResumeRendering() final override;
+	virtual bool RHIIsRenderingSuspended() final override;
+	virtual bool RHIGetAvailableResolutions(FScreenResolutionArray& Resolutions, bool bIgnoreRefreshRate) final override;
+	virtual void RHIGetSupportedResolution(uint32& Width, uint32& Height) final override;
+	virtual void* RHIGetNativeDevice() final override;
+	virtual void* RHIGetNativePhysicalDevice() final override;
+	virtual void* RHIGetNativeGraphicsQueue() final override;
+	virtual void* RHIGetNativeComputeQueue() final override;
+	virtual void* RHIGetNativeInstance() final override;
+	virtual class IRHICommandContext* RHIGetDefaultContext() final override;
+	virtual IRHIComputeContext* RHIGetCommandContext(ERHIPipeline Pipeline, FRHIGPUMask GPUMask) final override;
+	virtual IRHIUploadContext* RHIGetUploadContext() final override;
+	virtual void RHIFinalizeContext(FRHIFinalizeContextArgs&& Args, TRHIPipelineArray<IRHIPlatformCommandList*>& Output) final override;
+	virtual void RHISubmitCommandLists(FRHISubmitCommandListsArgs&& Args) final override;
+	virtual uint64 RHIGetMinimumAlignmentForBufferBackedSRV(EPixelFormat Format) final override;
+
+	// Transient Resource Functions and Helpers
+	virtual IRHITransientResourceAllocator* RHICreateTransientResourceAllocator() final override;
+
+	// version of the RHIComputePrecachePSOHash
+	static uint32 GetPrecachePSOHashVersion();
+	virtual uint64 RHIComputePrecachePSOHash(const FGraphicsPipelineStateInitializer& Initializer) final override;
+	virtual uint64 RHIComputeStatePrecachePSOHash(const FGraphicsPipelineStateInitializer& Initializer) final override;
+	virtual bool RHIMatchPrecachePSOInitializers(const FGraphicsPipelineStateInitializer& LHS, const FGraphicsPipelineStateInitializer& RHS) final override;
+
+	virtual FTextureRHIRef AsyncReallocateTexture2D_RenderThread(class FRHICommandListImmediate& RHICmdList, FRHITexture* Texture2D, int32 NewMipCount, int32 NewSizeX, int32 NewSizeY, FThreadSafeCounter* RequestStatus) override final;
+
+	virtual FUpdateTexture3DData RHIBeginUpdateTexture3D(FRHICommandListBase& RHICmdList, FRHITexture* Texture, uint32 MipIndex, const struct FUpdateTextureRegion3D& UpdateRegion) override final;
+	virtual void RHIEndUpdateTexture3D(FRHICommandListBase& RHICmdList, FUpdateTexture3DData& UpdateData) override final;
+
+	// SRV / UAV creation functions
+	virtual FShaderResourceViewRHIRef  RHICreateShaderResourceView (class FRHICommandListBase& RHICmdList, FRHIViewableResource* Resource, FRHIViewDesc const& ViewDesc) final override;
+	virtual FUnorderedAccessViewRHIRef RHICreateUnorderedAccessView(class FRHICommandListBase& RHICmdList, FRHIViewableResource* Resource, FRHIViewDesc const& ViewDesc) final override;
+
+#if PLATFORM_SUPPORTS_BINDLESS_RENDERING
+	virtual FRHIResourceCollectionRef RHICreateResourceCollection(FRHICommandListBase& RHICmdList, TConstArrayView<FRHIResourceCollectionMember> InMembers) final;
+#endif
+
+	virtual FRayTracingAccelerationStructureSize RHICalcRayTracingSceneSize(const FRayTracingSceneInitializer& Initializer) final override;
+	virtual FRayTracingAccelerationStructureSize RHICalcRayTracingGeometrySize(const FRayTracingGeometryInitializer& Initializer) final override;
+	virtual FRayTracingGeometryRHIRef RHICreateRayTracingGeometry(FRHICommandListBase& RHICmdList, const FRayTracingGeometryInitializer& Initializer) final override;
+	virtual FRayTracingSceneRHIRef RHICreateRayTracingScene(FRayTracingSceneInitializer Initializer) final override;
+	virtual FRayTracingShaderRHIRef RHICreateRayTracingShader(TArrayView<const uint8> Code, const FSHAHash& Hash, EShaderFrequency ShaderFrequency) final override;
+	virtual FRayTracingPipelineStateRHIRef RHICreateRayTracingPipelineState(const FRayTracingPipelineStateInitializer& Initializer) final override;
+
+	virtual FShaderBindingTableRHIRef RHICreateShaderBindingTable(FRHICommandListBase& RHICmdList, const FRayTracingShaderBindingTableInitializer& Initializer) final override;
+
+	// Historical number of times we've presented any and all viewports
+	uint32 TotalPresentCount = 0;
+
+	const TArray<const ANSICHAR*>& GetInstanceExtensions() const
+	{
+		return InstanceExtensions;
+	}
+
+	const TArray<const ANSICHAR*>& GetInstanceLayers() const
+	{
+		return InstanceLayers;
+	}
+
+	static void DestroySwapChain();
+	static void RecreateSwapChain(void* NewNativeWindow);
+
+	VkInstance GetInstance() const
+	{
+		return Instance;
+	}
+	
+	FVulkanDevice* GetDevice() const
+	{
+		return Device;
+	}
+
+	bool SupportsDebugUtilsExt() const
+	{
+		return (ActiveDebugLayerExtension == EActiveDebugLayerExtension::DebugUtilsExtension);
+	}
+
+	const FOptionalVulkanInstanceExtensions& GetOptionalExtensions() const
+	{
+		return OptionalInstanceExtensions;
+	}
+
+	void VulkanSetImageLayout( VkCommandBuffer CmdBuffer, VkImage Image, VkImageLayout OldLayout, VkImageLayout NewLayout, const VkImageSubresourceRange& SubresourceRange );
+
+	virtual void* RHILockStagingBuffer(FRHIStagingBuffer* StagingBuffer, FRHIGPUFence* Fence, uint32 Offset, uint32 SizeRHI) final override;
+	virtual void RHIUnlockStagingBuffer(FRHIStagingBuffer* StagingBuffer) final override;
+
+	TArray<FVulkanViewport*>& GetViewports()
+	{
+		return Viewports;
+	}
+
+	uint32 GetApiVersion() const
+	{
+		return ApiVersion;
+	}
+
+public:
+	static void SavePipelineCache();
+
+	void ProcessInterruptQueueUntil(FGraphEvent* GraphEvent);
+	void EnqueueEndOfPipeTask(TUniqueFunction<void()> TaskFunc, TUniqueFunction<void(FVulkanPayload&)> ModifyPayloadCallback = {});
+
+protected:
+	void InitializeSubmissionPipe();
+	void ShutdownSubmissionPipe();
+	bool ProcessSubmissionQueue();
+	bool ProcessInterruptQueue();
+	bool WaitAndProcessInterruptQueue();
+	void KickSubmissionThread();
+	void KickInterruptThread();
+
+	FCriticalSection SubmissionCS;
+	FVulkanThread* SubmissionThread = nullptr;
+	FCriticalSection InterruptCS;
+	FVulkanThread* InterruptThread = nullptr;
+	VulkanRHI::FSemaphore* CPUTimelineSemaphore = nullptr;  // used to wake up interrupt thread from CPU
+	std::atomic<uint64> CPUTimelineSemaphoreValue = 0;
+	TQueue<FVulkanPlatformCommandList*, EQueueMode::Mpsc> PendingPayloadsForSubmission;
+	TSet<VkSemaphore> SignaledSemaphores;
+	FGraphEventRef EopTask;
+
+	uint32 ApiVersion = 0;
+	VkInstance Instance;
+	TArray<const ANSICHAR*> InstanceExtensions;
+	TArray<const ANSICHAR*> InstanceLayers;
+
+	FVulkanDevice* Device;
+
+	/** A list of all viewport RHIs that have been created. */
+	TArray<FVulkanViewport*> Viewports;
+
+	/** The viewport which is currently being drawn. */
+	TRefCountPtr<FVulkanViewport> DrawingViewport;
+
+	void CreateInstance();
+	void SelectDevice();
+
+	friend class FVulkanCommandListContext;
+
+	friend class FVulkanViewport;
+
+	/*
+	static void WriteEndFrameTimestamp(void*);
+	*/
+
+	TArray<const ANSICHAR*> SetupInstanceLayers(FVulkanInstanceExtensionArray& UEExtensions);
+
+	IConsoleObject* SavePipelineCacheCmd = nullptr;
+	IConsoleObject* RebuildPipelineCacheCmd = nullptr;
+#if VULKAN_SUPPORTS_VALIDATION_CACHE
+	IConsoleObject* SaveValidationCacheCmd = nullptr;
+#endif
+
+	static void RebuildPipelineCache();
+#if VULKAN_SUPPORTS_VALIDATION_CACHE
+	static void SaveValidationCache();
+#endif
+
+#if UE_BUILD_DEBUG || UE_BUILD_DEVELOPMENT
+	IConsoleObject* DumpMemoryCmd = nullptr;
+	IConsoleObject* DumpMemoryFullCmd = nullptr;
+	IConsoleObject* DumpStagingMemoryCmd = nullptr;
+	IConsoleObject* DumpLRUCmd = nullptr;
+	IConsoleObject* TrimLRUCmd = nullptr;
+public:
+	static void DumpMemory();
+	static void DumpMemoryFull();
+	static void DumpStagingMemory();
+	static void DumpLRU();
+	static void TrimLRU();
+#endif
+
+public:
+	enum class EActiveDebugLayerExtension
+	{
+		None,
+		GfxReconstructLayer,
+		VkTraceLayer,
+		DebugUtilsExtension
+	};
+
+protected:
+	bool bIsStandaloneStereoDevice = false;
+
+	EActiveDebugLayerExtension ActiveDebugLayerExtension = EActiveDebugLayerExtension::None;
+
+#if VULKAN_HAS_DEBUGGING_ENABLED
+	VkDebugUtilsMessengerEXT Messenger = VK_NULL_HANDLE;
+
+	void SetupDebugLayerCallback();
+	void RemoveDebugLayerCallback();
+#endif
+
+	FCriticalSection LockBufferCS;
+
+	void InternalUpdateTexture2D(FRHICommandListBase& RHICmdList, FRHITexture* TextureRHI, uint32 MipIndex, const struct FUpdateTextureRegion2D& UpdateRegion, uint32 SourcePitch, const uint8* SourceData);
+	void InternalUpdateTexture3D(FRHICommandListBase& RHICmdList, FRHITexture* TextureRHI, uint32 MipIndex, const struct FUpdateTextureRegion3D& UpdateRegion, uint32 SourceRowPitch, uint32 SourceDepthPitch, const uint8* SourceData);
+
+	void UpdateUniformBuffer(FRHICommandListBase& RHICmdList, FVulkanUniformBuffer* UniformBuffer, const void* Contents);
+
+	static void SetupValidationRequests();
+
+	FOptionalVulkanInstanceExtensions OptionalInstanceExtensions;
+
+public:
+	static TSharedPtr< IHeadMountedDisplayVulkanExtensions, ESPMode::ThreadSafe > HMDVulkanExtensions;
+};
+
+PRAGMA_ENABLE_DEPRECATION_WARNINGS
+
+/** Implements the Vulkan module as a dynamic RHI providing module. */
+class FVulkanDynamicRHIModule : public IDynamicRHIModule
+{
+public:
+	// IModuleInterface
+	virtual void StartupModule() override;
+
+	// IDynamicRHIModule
+	virtual bool IsSupported() override;
+	virtual bool IsSupported(ERHIFeatureLevel::Type RequestedFeatureLevel) override;
+
+	virtual FDynamicRHI* CreateRHI(ERHIFeatureLevel::Type RequestedFeatureLevel = ERHIFeatureLevel::Num) override;
+};
diff --git a/Engine/Source/Runtime/VulkanRHI/Private/VulkanRHI.cpp b/Engine/Source/Runtime/VulkanRHI/Private/VulkanRHI.cpp
--- a/Engine/Source/Runtime/VulkanRHI/Private/VulkanRHI.cpp
+++ b/Engine/Source/Runtime/VulkanRHI/Private/VulkanRHI.cpp
@@ -1317,7 +1317,11 @@ VkInstance FVulkanDynamicRHI::RHIGetVkInstance() const
 
 VkDevice FVulkanDynamicRHI::RHIGetVkDevice() const
 {
-	return Device->GetInstanceHandle();
+	if (Device)
+	{
+		return Device->GetInstanceHandle();
+	}
+	return VK_NULL_HANDLE;
 }
 
 const uint8* FVulkanDynamicRHI::RHIGetVulkanDeviceUUID() const
diff --git a/Engine/Source/Runtime/VulkanRHI/Public/VulkanDynamicRHI.h b/Engine/Source/Runtime/VulkanRHI/Public/VulkanDynamicRHI.h
--- a/Engine/Source/Runtime/VulkanRHI/Public/VulkanDynamicRHI.h
+++ b/Engine/Source/Runtime/VulkanRHI/Public/VulkanDynamicRHI.h
@@ -1,430 +1,12 @@
 // Copyright Epic Games, Inc. All Rights Reserved.
 
 /*=============================================================================
-	VulkanRHI.h: Public Vulkan RHI definitions.
+	VulkanDynamicRHI.h: Public Vulkan RHI definitions.
 =============================================================================*/
 
 #pragma once
 
-#include "IVulkanDynamicRHI.h"
-#include "VulkanConfiguration.h"
-
-class FVulkanFramebuffer;
-class FVulkanDevice;
-class FVulkanQueue;
-class FVulkanUniformBuffer;
-class FVulkanViewport;
-class FVulkanPayload;
-class FVulkanTexture;
-class IHeadMountedDisplayVulkanExtensions;
-class FVulkanThread;
-struct FRHITransientHeapAllocation;
-struct FVulkanPlatformCommandList;
-
-namespace VulkanRHI
-{
-	class FSemaphore;
-};
-
-struct FOptionalVulkanInstanceExtensions
-{
-	union
-	{
-		struct
-		{
-			UE_DEPRECATED(5.3, "Vulkan 1.1 is now a requirement so there is no need to check these capabilities.")
-			uint32 HasKHRExternalMemoryCapabilities : 1;
-
-			UE_DEPRECATED(5.3, "Vulkan 1.1 is now a requirement so there is no need to check these capabilities.")
-			uint32 HasKHRGetPhysicalDeviceProperties2 : 1;
-		};
-		uint32 Packed;
-	};
-
-	FOptionalVulkanInstanceExtensions()
-	{
-		static_assert(sizeof(Packed) == sizeof(FOptionalVulkanInstanceExtensions), "More bits needed!");
-		Packed = 0;
-	}
-};
-
-PRAGMA_DISABLE_DEPRECATION_WARNINGS
-
-/** The interface which is implemented by the dynamically bound RHI. */
-class FVulkanDynamicRHI : public IVulkanDynamicRHI
-{
-public:
-
-	static FVulkanDynamicRHI& Get() { return *GetDynamicRHI<FVulkanDynamicRHI>(); }
-
-	/** Initialization constructor. */
-	FVulkanDynamicRHI();
-
-	/** Destructor */
-	~FVulkanDynamicRHI();
-
-	// IVulkanDynamicRHI interface
-	virtual uint32 RHIGetVulkanVersion() const final override;
-	virtual VkInstance RHIGetVkInstance() const final override;
-	virtual VkDevice RHIGetVkDevice() const final override;
-	virtual const uint8* RHIGetVulkanDeviceUUID() const final override;
-	virtual VkPhysicalDevice RHIGetVkPhysicalDevice() const final override;
-	virtual const VkAllocationCallbacks* RHIGetVkAllocationCallbacks() final override;
-	virtual VkQueue RHIGetGraphicsVkQueue() const final override;
-	virtual uint32 RHIGetGraphicsQueueIndex() const final override;
-	virtual uint32 RHIGetGraphicsQueueFamilyIndex() const final override;
-	virtual VkCommandBuffer RHIGetActiveVkCommandBuffer() final override;
-	virtual uint64 RHIGetGraphicsAdapterLUID(VkPhysicalDevice InPhysicalDevice) const final override;
-	virtual bool RHIDoesAdapterMatchDevice(const void* InAdapterId) const final override;
-	virtual void* RHIGetVkDeviceProcAddr(const char* InName) const final override;
-	virtual void* RHIGetVkInstanceProcAddr(const char* InName) const final override;
-	virtual void* RHIGetVkInstanceGlobalProcAddr(const char* InName) const final override;
-	virtual VkFormat RHIGetSwapChainVkFormat(EPixelFormat InFormat) const final override;
-	virtual bool RHISupportsEXTFragmentDensityMap2() const final override;
-	virtual TArray<VkExtensionProperties> RHIGetAllInstanceExtensions() const final override;
-	virtual TArray<VkExtensionProperties> RHIGetAllDeviceExtensions(VkPhysicalDevice InPhysicalDevice) const final override;
-	virtual TArray<FAnsiString> RHIGetLoadedDeviceExtensions() const final override;
-	virtual FTextureRHIRef RHICreateTexture2DFromResource(EPixelFormat Format, uint32 SizeX, uint32 SizeY, uint32 NumMips, uint32 NumSamples, VkImage Resource, ETextureCreateFlags Flags, const FClearValueBinding& ClearValueBinding = FClearValueBinding::Transparent, const FVulkanRHIExternalImageDeleteCallbackInfo& ExternalImageDeleteCallbackInfo = {}) final override;
-#if PLATFORM_ANDROID
-	virtual FTextureRHIRef RHICreateTexture2DFromAndroidHardwareBuffer(AHardwareBuffer* HardwareBuffer) final override;
-#endif
-	virtual FTextureRHIRef RHICreateTexture2DArrayFromResource(EPixelFormat Format, uint32 SizeX, uint32 SizeY, uint32 ArraySize, uint32 NumMips, uint32 NumSamples, VkImage Resource, ETextureCreateFlags Flags, const FClearValueBinding& ClearValueBinding = FClearValueBinding::Transparent) final override;
-	virtual FTextureRHIRef RHICreateTextureCubeFromResource(EPixelFormat Format, uint32 Size, bool bArray, uint32 ArraySize, uint32 NumMips, VkImage Resource, ETextureCreateFlags Flags, const FClearValueBinding& ClearValueBinding = FClearValueBinding::Transparent) final override;
-	virtual VkImage RHIGetVkImage(FRHITexture* InTexture) const final override;
-	virtual VkFormat RHIGetViewVkFormat(FRHITexture* InTexture) const final override;
-	virtual FVulkanRHIAllocationInfo RHIGetAllocationInfo(FRHITexture* InTexture) const final override;
-	virtual FVulkanRHIImageViewInfo RHIGetImageViewInfo(FRHITexture* InTexture) const final override;
-	virtual FVulkanRHIAllocationInfo RHIGetAllocationInfo(FRHIBuffer* InBuffer) const final override;
-	virtual void RHISetImageLayout(VkImage Image, VkImageLayout OldLayout, VkImageLayout NewLayout, const VkImageSubresourceRange& SubresourceRange) final override;
-	virtual void RHISetUploadImageLayout(VkImage Image, VkImageLayout OldLayout, VkImageLayout NewLayout, const VkImageSubresourceRange& SubresourceRange) final override;
-	virtual void RHIFinishExternalComputeWork(VkCommandBuffer InCommandBuffer) final override;
-	virtual void RHIRegisterWork(uint32 NumPrimitives) final override;
-	virtual void RHISubmitUploadCommandBuffer() final override;
-	virtual void RHIVerifyResult(VkResult Result, const ANSICHAR* VkFuntion, const ANSICHAR* Filename, uint32 Line) final override;
-
-	// FDynamicRHI interface.
-	virtual void Init() final override;
-	virtual void PostInit() final override;
-	virtual void Shutdown() final override;
-	virtual const TCHAR* GetName() final override { return TEXT("Vulkan"); }
-
-	void InitInstance();
-
-	virtual void RHIEndFrame_RenderThread(FRHICommandListImmediate& RHICmdList) final override;
-	virtual void RHIEndFrame(const FRHIEndFrameArgs& Args) final override;
-
-	virtual FSamplerStateRHIRef RHICreateSamplerState(const FSamplerStateInitializerRHI& Initializer) final override;
-	virtual FRasterizerStateRHIRef RHICreateRasterizerState(const FRasterizerStateInitializerRHI& Initializer) final override;
-	virtual FDepthStencilStateRHIRef RHICreateDepthStencilState(const FDepthStencilStateInitializerRHI& Initializer) final override;
-	virtual FBlendStateRHIRef RHICreateBlendState(const FBlendStateInitializerRHI& Initializer) final override;
-	virtual FVertexDeclarationRHIRef RHICreateVertexDeclaration(const FVertexDeclarationElementList& Elements) final override;
-	virtual FPixelShaderRHIRef RHICreatePixelShader(TArrayView<const uint8> Code, const FSHAHash& Hash) final override;
-	virtual FVertexShaderRHIRef RHICreateVertexShader(TArrayView<const uint8> Code, const FSHAHash& Hash) final override;
-	virtual FMeshShaderRHIRef RHICreateMeshShader(TArrayView<const uint8> Code, const FSHAHash& Hash) final override;
-	virtual FAmplificationShaderRHIRef RHICreateAmplificationShader(TArrayView<const uint8> Code, const FSHAHash& Hash) final override;
-	virtual FGeometryShaderRHIRef RHICreateGeometryShader(TArrayView<const uint8> Code, const FSHAHash& Hash) final override;
-	virtual FComputeShaderRHIRef RHICreateComputeShader(TArrayView<const uint8> Code, const FSHAHash& Hash) final override;
-	virtual FGPUFenceRHIRef RHICreateGPUFence(const FName &Name) final override;
-	virtual void RHICreateTransition(FRHITransition* Transition, const FRHITransitionCreateInfo& CreateInfo) final override;
-	virtual void RHIReleaseTransition(FRHITransition* Transition) final override;
-	virtual FStagingBufferRHIRef RHICreateStagingBuffer() final override;
-	virtual FBoundShaderStateRHIRef RHICreateBoundShaderState(FRHIVertexDeclaration* VertexDeclaration, FRHIVertexShader* VertexShader, FRHIPixelShader* PixelShader, FRHIGeometryShader* GeometryShader) final override;
-	virtual FGraphicsPipelineStateRHIRef RHICreateGraphicsPipelineState(const FGraphicsPipelineStateInitializer& Initializer) final override;
-	virtual FComputePipelineStateRHIRef RHICreateComputePipelineState(FRHIComputeShader* ComputeShader) final override;
-	virtual FUniformBufferRHIRef RHICreateUniformBuffer(const void* Contents, const FRHIUniformBufferLayout* Layout, EUniformBufferUsage Usage, EUniformBufferValidation Validation) final override;
-	virtual void RHIUpdateUniformBuffer(FRHICommandListBase& RHICmdList, FRHIUniformBuffer* UniformBufferRHI, const void* Contents) final override;
-
-	[[nodiscard]] virtual FRHIBufferInitializer RHICreateBufferInitializer(FRHICommandListBase& RHICmdList, const FRHIBufferCreateDesc& CreateDesc) final override;
-
-	virtual void RHIReplaceResources(FRHICommandListBase& RHICmdList, TArray<FRHIResourceReplaceInfo>&& ReplaceInfos) final override;
-	virtual void* LockBuffer_BottomOfPipe(FRHICommandListBase& RHICmdList, FRHIBuffer* Buffer, uint32 Offset, uint32 Size, EResourceLockMode LockMode) final override;
-	virtual void UnlockBuffer_BottomOfPipe(FRHICommandListBase& RHICmdList, FRHIBuffer* Buffer) final override;
-	virtual void* RHILockBuffer(FRHICommandListBase& RHICmdList, FRHIBuffer* BufferRHI, uint32 Offset, uint32 Size, EResourceLockMode LockMode) final override;
-	virtual void RHIUnlockBuffer(FRHICommandListBase& RHICmdList, FRHIBuffer* Buffer) final override;
-	virtual FTextureReferenceRHIRef RHICreateTextureReference(FRHICommandListBase& RHICmdList, FRHITexture* InReferencedTexture) final override;
-	virtual void RHIUpdateTextureReference(FRHICommandListBase& RHICmdList, FRHITextureReference* TextureRef, FRHITexture* NewTexture) final override;
-	virtual FRHICalcTextureSizeResult RHICalcTexturePlatformSize(FRHITextureDesc const& Desc, uint32 FirstMipIndex) final override;
-	virtual void RHIGetTextureMemoryStats(FTextureMemoryStats& OutStats) final override;
-	virtual bool RHIGetTextureMemoryVisualizeData(FColor* TextureData, int32 SizeX, int32 SizeY, int32 Pitch, int32 PixelSize) final override;
-
-	struct FCreateTextureResult
-	{
-		FVulkanTexture* Texture;
-		VkImageLayout DefaultLayout;
-		bool bTransientResource;
-		bool bClear;
-	};
-
-	FCreateTextureResult BeginCreateTextureInternal(const FRHITextureCreateDesc& CreateDesc, const FRHITransientHeapAllocation* InTransientHeapAllocation);
-	FVulkanTexture* FinalizeCreateTextureInternal(FRHICommandListBase& RHICmdList, FCreateTextureResult CreateResult, TConstArrayView<uint8> InitialData);
-
-	FVulkanTexture* CreateTextureInternal(FRHICommandListBase& RHICmdList, const FRHITextureCreateDesc& CreateDesc, TConstArrayView<uint8> InitialData);
-	FVulkanTexture* CreateTextureInternal(const FRHITextureCreateDesc& CreateDesc, const FRHITransientHeapAllocation& InTransientHeapAllocation);
-
-	virtual FTextureRHIRef RHICreateTexture(FRHICommandListBase& RHICmdList, const FRHITextureCreateDesc& CreateDesc) final override;
-	virtual FTextureRHIRef RHIAsyncCreateTexture2D(uint32 SizeX, uint32 SizeY, uint8 Format, uint32 NumMips, ETextureCreateFlags Flags, ERHIAccess InResourceState, void** InitialMipData, uint32 NumInitialMips, const TCHAR* DebugName, FGraphEventRef& OutCompletionEvent) final override;
-	virtual uint32 RHIComputeMemorySize(FRHITexture* TextureRHI) final override;
-	virtual FTextureRHIRef RHIAsyncReallocateTexture2D(FRHITexture* Texture2D, int32 NewMipCount, int32 NewSizeX, int32 NewSizeY, FThreadSafeCounter* RequestStatus) final override;
-
-	virtual FRHILockTextureResult RHILockTexture(FRHICommandListImmediate& RHICmdList, const FRHILockTextureArgs& Arguments) final override;
-	virtual void RHIUnlockTexture(FRHICommandListImmediate& RHICmdList, const FRHILockTextureArgs& Arguments) final override;
-
-	virtual void RHIUpdateTexture2D(FRHICommandListBase& RHICmdList, FRHITexture* Texture, uint32 MipIndex, const struct FUpdateTextureRegion2D& UpdateRegion, uint32 SourcePitch, const uint8* SourceData) final override
-	{
-		InternalUpdateTexture2D(RHICmdList, Texture, MipIndex, UpdateRegion, SourcePitch, SourceData);
-	}
-	virtual void RHIUpdateTexture3D(FRHICommandListBase& RHICmdList, FRHITexture* Texture, uint32 MipIndex, const struct FUpdateTextureRegion3D& UpdateRegion, uint32 SourceRowPitch, uint32 SourceDepthPitch, const uint8* SourceData) final override
-	{
-		InternalUpdateTexture3D(RHICmdList, Texture, MipIndex, UpdateRegion, SourceRowPitch, SourceDepthPitch, SourceData);
-	}
-	virtual void RHIBindDebugLabelName(FRHICommandListBase& RHICmdList, FRHITexture* Texture, const TCHAR* Name) final override;
-	virtual void RHIBindDebugLabelName(FRHICommandListBase& RHICmdList, FRHIUnorderedAccessView* UnorderedAccessViewRHI, const TCHAR* Name) final override;
-	virtual void RHIReadSurfaceData(FRHITexture* Texture, FIntRect Rect, TArray<FColor>& OutData, FReadSurfaceDataFlags InFlags) final override;
-	virtual void RHIReadSurfaceData(FRHITexture* Texture, FIntRect Rect, TArray<FLinearColor>& OutData, FReadSurfaceDataFlags InFlags) final override;
-	virtual void RHIMapStagingSurface(FRHITexture* Texture, FRHIGPUFence* Fence, void*& OutData, int32& OutWidth, int32& OutHeight, uint32 GPUIndex = 0) final override;
-	virtual void RHIUnmapStagingSurface(FRHITexture* Texture, uint32 GPUIndex = 0) final override;
-
-	virtual void RHIReadSurfaceFloatData(FRHITexture* Texture, FIntRect Rect, TArray<FFloat16Color>& OutData, ECubeFace CubeFace, int32 ArrayIndex, int32 MipIndex) final override;
-	virtual void RHIRead3DSurfaceFloatData(FRHITexture* Texture, FIntRect Rect, FIntPoint ZMinMax, TArray<FFloat16Color>& OutData) final override;
-	virtual FRenderQueryRHIRef RHICreateRenderQuery(ERenderQueryType QueryType) final override;
-	virtual void RHIEndRenderQuery_TopOfPipe(FRHICommandListBase& RHICmdList, FRHIRenderQuery* RenderQuery) final override;
-	virtual void RHIBeginRenderQueryBatch_TopOfPipe(FRHICommandListBase& RHICmdList, ERenderQueryType QueryType) final override;
-	virtual void RHIEndRenderQueryBatch_TopOfPipe(FRHICommandListBase& RHICmdList, ERenderQueryType QueryType) final override;
-	virtual bool RHIGetRenderQueryResult(FRHIRenderQuery* RenderQuery, uint64& OutResult, bool bWait, uint32 GPUIndex = INDEX_NONE) final override;
-	virtual FTextureRHIRef RHIGetViewportBackBuffer(FRHIViewport* Viewport) final override;
-	virtual void RHIAliasTextureResources(FTextureRHIRef& DestTexture, FTextureRHIRef& SrcTexture) final override;
-	virtual FTextureRHIRef RHICreateAliasedTexture(FTextureRHIRef& SourceTexture) final override;
-	virtual void RHIAdvanceFrameForGetViewportBackBuffer(FRHIViewport* Viewport) final override;
-	virtual void RHIFlushResources() final override;
-	virtual uint32 RHIGetGPUFrameCycles(uint32 GPUIndex = 0) final override;
-	virtual FViewportRHIRef RHICreateViewport(void* WindowHandle, uint32 SizeX, uint32 SizeY, bool bIsFullscreen, EPixelFormat PreferredPixelFormat) final override;
-	virtual void RHIResizeViewport(FRHIViewport* Viewport, uint32 SizeX, uint32 SizeY, bool bIsFullscreen) final override;
-	virtual void RHIResizeViewport(FRHIViewport* Viewport, uint32 SizeX, uint32 SizeY, bool bIsFullscreen, EPixelFormat PreferredPixelFormat) final override;
-	virtual void RHITick(float DeltaTime) final override;
-	virtual void RHIBlockUntilGPUIdle() final override;
-	virtual void RHISuspendRendering() final override;
-	virtual void RHIResumeRendering() final override;
-	virtual bool RHIIsRenderingSuspended() final override;
-	virtual bool RHIGetAvailableResolutions(FScreenResolutionArray& Resolutions, bool bIgnoreRefreshRate) final override;
-	virtual void RHIGetSupportedResolution(uint32& Width, uint32& Height) final override;
-	virtual void* RHIGetNativeDevice() final override;
-	virtual void* RHIGetNativePhysicalDevice() final override;
-	virtual void* RHIGetNativeGraphicsQueue() final override;
-	virtual void* RHIGetNativeComputeQueue() final override;
-	virtual void* RHIGetNativeInstance() final override;
-	virtual class IRHICommandContext* RHIGetDefaultContext() final override;
-	virtual IRHIComputeContext* RHIGetCommandContext(ERHIPipeline Pipeline, FRHIGPUMask GPUMask) final override;
-	virtual IRHIUploadContext* RHIGetUploadContext() final override;
-	virtual void RHIFinalizeContext(FRHIFinalizeContextArgs&& Args, TRHIPipelineArray<IRHIPlatformCommandList*>& Output) final override;
-	virtual void RHISubmitCommandLists(FRHISubmitCommandListsArgs&& Args) final override;
-	virtual uint64 RHIGetMinimumAlignmentForBufferBackedSRV(EPixelFormat Format) final override;
-
-	// Transient Resource Functions and Helpers
-	virtual IRHITransientResourceAllocator* RHICreateTransientResourceAllocator() final override;
-
-	// version of the RHIComputePrecachePSOHash
-	static uint32 GetPrecachePSOHashVersion();
-	virtual uint64 RHIComputePrecachePSOHash(const FGraphicsPipelineStateInitializer& Initializer) final override;
-	virtual uint64 RHIComputeStatePrecachePSOHash(const FGraphicsPipelineStateInitializer& Initializer) final override;
-	virtual bool RHIMatchPrecachePSOInitializers(const FGraphicsPipelineStateInitializer& LHS, const FGraphicsPipelineStateInitializer& RHS) final override;
-
-	virtual FTextureRHIRef AsyncReallocateTexture2D_RenderThread(class FRHICommandListImmediate& RHICmdList, FRHITexture* Texture2D, int32 NewMipCount, int32 NewSizeX, int32 NewSizeY, FThreadSafeCounter* RequestStatus) override final;
-
-	virtual FUpdateTexture3DData RHIBeginUpdateTexture3D(FRHICommandListBase& RHICmdList, FRHITexture* Texture, uint32 MipIndex, const struct FUpdateTextureRegion3D& UpdateRegion) override final;
-	virtual void RHIEndUpdateTexture3D(FRHICommandListBase& RHICmdList, FUpdateTexture3DData& UpdateData) override final;
-
-	// SRV / UAV creation functions
-	virtual FShaderResourceViewRHIRef  RHICreateShaderResourceView (class FRHICommandListBase& RHICmdList, FRHIViewableResource* Resource, FRHIViewDesc const& ViewDesc) final override;
-	virtual FUnorderedAccessViewRHIRef RHICreateUnorderedAccessView(class FRHICommandListBase& RHICmdList, FRHIViewableResource* Resource, FRHIViewDesc const& ViewDesc) final override;
-
-#if PLATFORM_SUPPORTS_BINDLESS_RENDERING
-	virtual FRHIResourceCollectionRef RHICreateResourceCollection(FRHICommandListBase& RHICmdList, TConstArrayView<FRHIResourceCollectionMember> InMembers) final;
-#endif
-
-	virtual FRayTracingAccelerationStructureSize RHICalcRayTracingSceneSize(const FRayTracingSceneInitializer& Initializer) final override;
-	virtual FRayTracingAccelerationStructureSize RHICalcRayTracingGeometrySize(const FRayTracingGeometryInitializer& Initializer) final override;
-	virtual FRayTracingGeometryRHIRef RHICreateRayTracingGeometry(FRHICommandListBase& RHICmdList, const FRayTracingGeometryInitializer& Initializer) final override;
-	virtual FRayTracingSceneRHIRef RHICreateRayTracingScene(FRayTracingSceneInitializer Initializer) final override;
-	virtual FRayTracingShaderRHIRef RHICreateRayTracingShader(TArrayView<const uint8> Code, const FSHAHash& Hash, EShaderFrequency ShaderFrequency) final override;
-	virtual FRayTracingPipelineStateRHIRef RHICreateRayTracingPipelineState(const FRayTracingPipelineStateInitializer& Initializer) final override;
-
-	virtual FShaderBindingTableRHIRef RHICreateShaderBindingTable(FRHICommandListBase& RHICmdList, const FRayTracingShaderBindingTableInitializer& Initializer) final override;
-
-	// Historical number of times we've presented any and all viewports
-	uint32 TotalPresentCount = 0;
-
-	inline const TArray<const ANSICHAR*>& GetInstanceExtensions() const
-	{
-		return InstanceExtensions;
-	}
-
-	inline const TArray<const ANSICHAR*>& GetInstanceLayers() const
-	{
-		return InstanceLayers;
-	}
-
-	static void DestroySwapChain();
-	static void RecreateSwapChain(void* NewNativeWindow);
-
-	inline VkInstance GetInstance() const
-	{
-		return Instance;
-	}
-	
-	inline FVulkanDevice* GetDevice() const
-	{
-		return Device;
-	}
+#include "Misc/CoreMiscDefines.h"
+UE_DEPRECATED_HEADER(5.6, "VulkanDynamicRHI.h in the Public folder is deprecated. Please switch to using IVulkanDynamicRHI.h")
 
-	inline bool SupportsDebugUtilsExt() const
-	{
-		return (ActiveDebugLayerExtension == EActiveDebugLayerExtension::DebugUtilsExtension);
-	}
-
-	inline const FOptionalVulkanInstanceExtensions& GetOptionalExtensions() const
-	{
-		return OptionalInstanceExtensions;
-	}
-
-	void VulkanSetImageLayout( VkCommandBuffer CmdBuffer, VkImage Image, VkImageLayout OldLayout, VkImageLayout NewLayout, const VkImageSubresourceRange& SubresourceRange );
-
-	virtual void* RHILockStagingBuffer(FRHIStagingBuffer* StagingBuffer, FRHIGPUFence* Fence, uint32 Offset, uint32 SizeRHI) final override;
-	virtual void RHIUnlockStagingBuffer(FRHIStagingBuffer* StagingBuffer) final override;
-
-	inline TArray<FVulkanViewport*>& GetViewports()
-	{
-		return Viewports;
-	}
-
-	inline uint32 GetApiVersion() const
-	{
-		return ApiVersion;
-	}
-
-public:
-	static void SavePipelineCache();
-
-	void ProcessInterruptQueueUntil(FGraphEvent* GraphEvent);
-	void EnqueueEndOfPipeTask(TUniqueFunction<void()> TaskFunc, TUniqueFunction<void(FVulkanPayload&)> ModifyPayloadCallback = {});
-
-protected:
-	void InitializeSubmissionPipe();
-	void ShutdownSubmissionPipe();
-	bool ProcessSubmissionQueue();
-	bool ProcessInterruptQueue();
-	bool WaitAndProcessInterruptQueue();
-	void KickSubmissionThread();
-	void KickInterruptThread();
-
-	FCriticalSection SubmissionCS;
-	FVulkanThread* SubmissionThread = nullptr;
-	FCriticalSection InterruptCS;
-	FVulkanThread* InterruptThread = nullptr;
-	VulkanRHI::FSemaphore* CPUTimelineSemaphore = nullptr;  // used to wake up interrupt thread from CPU
-	std::atomic<uint64> CPUTimelineSemaphoreValue = 0;
-	TQueue<FVulkanPlatformCommandList*, EQueueMode::Mpsc> PendingPayloadsForSubmission;
-	TSet<VkSemaphore> SignaledSemaphores;
-	FGraphEventRef EopTask;
-
-	uint32 ApiVersion = 0;
-	VkInstance Instance;
-	TArray<const ANSICHAR*> InstanceExtensions;
-	TArray<const ANSICHAR*> InstanceLayers;
-
-	FVulkanDevice* Device;
-
-	/** A list of all viewport RHIs that have been created. */
-	TArray<FVulkanViewport*> Viewports;
-
-	/** The viewport which is currently being drawn. */
-	TRefCountPtr<FVulkanViewport> DrawingViewport;
-
-	void CreateInstance();
-	void SelectDevice();
-
-	friend class FVulkanCommandListContext;
-
-	friend class FVulkanViewport;
-
-	/*
-	static void WriteEndFrameTimestamp(void*);
-	*/
-
-	TArray<const ANSICHAR*> SetupInstanceLayers(FVulkanInstanceExtensionArray& UEExtensions);
-
-	IConsoleObject* SavePipelineCacheCmd = nullptr;
-	IConsoleObject* RebuildPipelineCacheCmd = nullptr;
-#if VULKAN_SUPPORTS_VALIDATION_CACHE
-	IConsoleObject* SaveValidationCacheCmd = nullptr;
-#endif
-
-	static void RebuildPipelineCache();
-#if VULKAN_SUPPORTS_VALIDATION_CACHE
-	static void SaveValidationCache();
-#endif
-
-#if UE_BUILD_DEBUG || UE_BUILD_DEVELOPMENT
-	IConsoleObject* DumpMemoryCmd = nullptr;
-	IConsoleObject* DumpMemoryFullCmd = nullptr;
-	IConsoleObject* DumpStagingMemoryCmd = nullptr;
-	IConsoleObject* DumpLRUCmd = nullptr;
-	IConsoleObject* TrimLRUCmd = nullptr;
-public:
-	static void DumpMemory();
-	static void DumpMemoryFull();
-	static void DumpStagingMemory();
-	static void DumpLRU();
-	static void TrimLRU();
-#endif
-
-public:
-	enum class EActiveDebugLayerExtension
-	{
-		None,
-		GfxReconstructLayer,
-		VkTraceLayer,
-		DebugUtilsExtension
-	};
-
-protected:
-	bool bIsStandaloneStereoDevice = false;
-
-	EActiveDebugLayerExtension ActiveDebugLayerExtension = EActiveDebugLayerExtension::None;
-
-#if VULKAN_HAS_DEBUGGING_ENABLED
-	VkDebugUtilsMessengerEXT Messenger = VK_NULL_HANDLE;
-
-	void SetupDebugLayerCallback();
-	void RemoveDebugLayerCallback();
-#endif
-
-	FCriticalSection LockBufferCS;
-
-	void InternalUpdateTexture2D(FRHICommandListBase& RHICmdList, FRHITexture* TextureRHI, uint32 MipIndex, const struct FUpdateTextureRegion2D& UpdateRegion, uint32 SourcePitch, const uint8* SourceData);
-	void InternalUpdateTexture3D(FRHICommandListBase& RHICmdList, FRHITexture* TextureRHI, uint32 MipIndex, const struct FUpdateTextureRegion3D& UpdateRegion, uint32 SourceRowPitch, uint32 SourceDepthPitch, const uint8* SourceData);
-
-	void UpdateUniformBuffer(FRHICommandListBase& RHICmdList, FVulkanUniformBuffer* UniformBuffer, const void* Contents);
-
-	static void SetupValidationRequests();
-
-	FOptionalVulkanInstanceExtensions OptionalInstanceExtensions;
-
-public:
-	static TSharedPtr< IHeadMountedDisplayVulkanExtensions, ESPMode::ThreadSafe > HMDVulkanExtensions;
-};
-
-PRAGMA_ENABLE_DEPRECATION_WARNINGS
-
-/** Implements the Vulkan module as a dynamic RHI providing module. */
-class FVulkanDynamicRHIModule : public IDynamicRHIModule
-{
-public:
-	// IModuleInterface
-	virtual void StartupModule() override;
-
-	// IDynamicRHIModule
-	virtual bool IsSupported() override;
-	virtual bool IsSupported(ERHIFeatureLevel::Type RequestedFeatureLevel) override;
-
-	virtual FDynamicRHI* CreateRHI(ERHIFeatureLevel::Type RequestedFeatureLevel = ERHIFeatureLevel::Num) override;
-};
+#include "IVulkanDynamicRHI.h"
diff --git a/Engine/Source/Runtime/VulkanRHI/Public/VulkanRHIBridge.h b/Engine/Source/Runtime/VulkanRHI/Public/VulkanRHIBridge.h
deleted file mode 100644
