Fix VulkanRHI validation errors when masked materials, or any other material that sets bAllowVariableRateShading=false are present in the scene and VRS is enabled.

diff --git a/Engine/Source/Runtime/VulkanRHI/Private/VulkanPipeline.cpp b/Engine/Source/Runtime/VulkanRHI/Private/VulkanPipeline.cpp
--- a/Engine/Source/Runtime/VulkanRHI/Private/VulkanPipeline.cpp
+++ b/Engine/Source/Runtime/VulkanRHI/Private/VulkanPipeline.cpp
@@ -1365,7 +1365,7 @@ bool FVulkanPipelineStateCacheManager::CreateGfxPipelineFromEntry(FVulkanRHIGrap
 
 	PipelineInfo.pDynamicState = &DynamicState;
 
-	const bool bUsingVariableRateShading = PSO->Desc.ShadingRate != EVRSShadingRate::VRSSR_1x1 || PSO->RenderPass->Layout.bHasFragmentDensityAttachment;
+	const bool bUsingVariableRateShading = PSO->Desc.ShadingRate != EVRSShadingRate::VRSSR_1x1 || (PSO->RenderPass->Layout.bHasFragmentDensityAttachment && PSO->Desc.Combiner != EVRSRateCombiner::VRSRB_Passthrough);
 
 	VkPipelineFragmentShadingRateStateCreateInfoKHR PipelineFragmentShadingRate;
 	if (GRHISupportsPipelineVariableRateShading && GRHIVariableRateShadingImageDataType == VRSImage_Palette && bUsingVariableRateShading)
@@ -1943,11 +1943,7 @@ void FVulkanPipelineStateCacheManager::CreateGfxEntry(const FGraphicsPipelineSta
 
 	// Shading rate:
 	OutGfxEntry->ShadingRate = PSOInitializer.bAllowVariableRateShading ? PSOInitializer.ShadingRate : EVRSShadingRate::VRSSR_1x1;
-	OutGfxEntry->Combiner = EVRSRateCombiner::VRSRB_Max;		// @todo: This needs to be specified twice; from pipeline-to-primitive, and from primitive-to-attachment. 
-																// We don't have per-primitive VRS so that should just be hard-coded to "passthrough" until this is supported; but we should expose 
-																// this setting in the material properies, especially since there's some materials that don't play nicely with 
-																// shading rates other than 1x1, in which case we'll want to use VRSRB_Min to override e.g. the attachment shading rate.
-																// For now, just locked to "max".
+	OutGfxEntry->Combiner = PSOInitializer.bAllowVariableRateShading ? EVRSRateCombiner::VRSRB_Max : EVRSRateCombiner::VRSRB_Passthrough; // Forces using the 1x1 rate over any fragment density attachment when VRS is disallowed in material settings
 }
 
 
diff --git a/Engine/Source/Runtime/VulkanRHI/Private/VulkanRenderTarget.cpp b/Engine/Source/Runtime/VulkanRHI/Private/VulkanRenderTarget.cpp
--- a/Engine/Source/Runtime/VulkanRHI/Private/VulkanRenderTarget.cpp
+++ b/Engine/Source/Runtime/VulkanRHI/Private/VulkanRenderTarget.cpp
@@ -1276,7 +1276,7 @@ FVulkanRenderTargetLayout::FVulkanRenderTargetLayout(const FGraphicsPipelineStat
 		bHasDepthStencil = true;
 	}
 
-	if (Initializer.bHasFragmentDensityAttachment && Initializer.bAllowVariableRateShading)
+	if (Initializer.bHasFragmentDensityAttachment)
 	{
 		VkAttachmentDescription& CurrDesc = Desc[NumAttachmentDescriptions];
 		FMemory::Memzero(CurrDesc);
