VulkanRHI: Add support for new profiler to VulkanRHI (disabled by default):
- Simplify queries in VulkanRHI (remove different classes for occlusion and timestamp)
- Read query results from interrupt thread before signaling SyncPoints
- Rewrite calibration of timestamps and apply it regularly.  Adjust GPU timestamps to CPU's timeline
- Rework legacy FVulkanGPUProfiler to use new queries until new profiler is enabled.

#rnx

diff --git a/Engine/Source/Runtime/VulkanRHI/Private/Android/VulkanAndroidPlatform.h b/Engine/Source/Runtime/VulkanRHI/Private/Android/VulkanAndroidPlatform.h
--- a/Engine/Source/Runtime/VulkanRHI/Private/Android/VulkanAndroidPlatform.h
+++ b/Engine/Source/Runtime/VulkanRHI/Private/Android/VulkanAndroidPlatform.h
@@ -179,6 +179,8 @@ public:
 
 	static void WriteCrashMarker(const FOptionalVulkanDeviceExtensions& OptionalExtensions, FVulkanCommandBuffer* CmdBuffer, VkBuffer DestBuffer, const TArrayView<uint32>& Entries, bool bAdding);
 
+	static VkTimeDomainKHR GetTimeDomain() { return VK_TIME_DOMAIN_CLOCK_MONOTONIC_KHR; }
+
 #if USE_ANDROID_VULKAN_SWAPPY
 	static bool bSwappyEnabledAtRHIInit;
 #endif
diff --git a/Engine/Source/Runtime/VulkanRHI/Private/Linux/VulkanLinuxPlatform.h b/Engine/Source/Runtime/VulkanRHI/Private/Linux/VulkanLinuxPlatform.h
--- a/Engine/Source/Runtime/VulkanRHI/Private/Linux/VulkanLinuxPlatform.h
+++ b/Engine/Source/Runtime/VulkanRHI/Private/Linux/VulkanLinuxPlatform.h
@@ -49,6 +49,8 @@ public:
 
 	static void WriteCrashMarker(const FOptionalVulkanDeviceExtensions& OptionalExtensions, FVulkanCommandBuffer* CmdBuffer, VkBuffer DestBuffer, const TArrayView<uint32>& Entries, bool bAdding);
 
+	static VkTimeDomainKHR GetTimeDomain() { return VK_TIME_DOMAIN_CLOCK_MONOTONIC_KHR; }
+
 protected:
 	static void* VulkanLib;
 	static bool bAttemptedLoad;
diff --git a/Engine/Source/Runtime/VulkanRHI/Private/VulkanCommandBuffer.cpp b/Engine/Source/Runtime/VulkanRHI/Private/VulkanCommandBuffer.cpp
--- a/Engine/Source/Runtime/VulkanRHI/Private/VulkanCommandBuffer.cpp
+++ b/Engine/Source/Runtime/VulkanRHI/Private/VulkanCommandBuffer.cpp
@@ -25,22 +25,14 @@ static FAutoConsoleVariableRef CVarVulkanUseSingleQueue(
 const uint32 GNumberOfFramesBeforeDeletingDescriptorPool = 300;
 
 FVulkanCommandBuffer::FVulkanCommandBuffer(FVulkanDevice& InDevice, FVulkanCommandBufferPool* InCommandBufferPool)
-	: CurrentStencilRef(0)
-	, State(EState::NotAllocated)
-	, bNeedsDynamicStateSet(1)
-	, bHasPipeline(0)
-	, bHasViewport(0)
-	, bHasScissor(0)
-	, bHasStencilRef(0)
-	, Device(InDevice)
-	, CommandBufferHandle(VK_NULL_HANDLE)
+	: Device(InDevice)
 	, CommandBufferPool(InCommandBufferPool)
-	, LastValidTiming(0)
+#if RHI_NEW_GPU_PROFILER
+	, EventStream(InCommandBufferPool->GetContext().Queue.GetProfilerQueue())
+#endif
 {
-	{
-		FScopeLock ScopeLock(CommandBufferPool->GetCS());
-		AllocMemory();
-	}
+	FScopeLock ScopeLock(CommandBufferPool->GetCS());
+	AllocMemory();
 }
 
 void FVulkanCommandBuffer::AllocMemory()
@@ -128,52 +120,34 @@ void FVulkanCommandBuffer::BeginRenderPass(const FVulkanRenderTargetLayout& Layo
 	}
 }
 
-
-void FVulkanCommandBuffer::AddPendingTimestampQuery(uint64 Index, uint64 Count, VkQueryPool PoolHandle, FVulkanTimingQueryResult* Results, bool bBlocking)
+void FVulkanCommandBuffer::End(FVulkanTimingQueryPool* QueryPool)
 {
-	FPendingTimestampQuery PendingQuery;
-	PendingQuery.Index = Index;
-	PendingQuery.Count = Count;
-	PendingQuery.PoolHandle = PoolHandle;
-	PendingQuery.Results = Results;
-	PendingQuery.bBlocking = bBlocking;
-	PendingTimestampQueries.Add(PendingQuery);
-}
-
+	checkf(IsOutsideRenderPass(), TEXT("Can't End as we're inside a render pass! CmdBuffer 0x%p State=%d"), CommandBufferHandle, (int32)State);
 
-void FVulkanCommandBuffer::ResolvePendingTimestampQueries()
-{
-	for (const FPendingTimestampQuery& PendingQuery : PendingTimestampQueries)
+#if RHI_NEW_GPU_PROFILER
+	if (QueryPool)
 	{
-		const VkQueryResultFlags BlockingQueryResultFlags = PendingQuery.bBlocking ? VK_QUERY_RESULT_WAIT_BIT : 0;
-		const VkResult Result = VulkanRHI::vkGetQueryPoolResults(
-			Device.GetInstanceHandle(),
-			PendingQuery.PoolHandle,
-			PendingQuery.Index, 
-			PendingQuery.Count, 
-			PendingQuery.Count * sizeof(FVulkanTimingQueryResult),
-			(uint64_t*)PendingQuery.Results,
-			sizeof(FVulkanTimingQueryResult),
-			VK_QUERY_RESULT_64_BIT | VK_QUERY_RESULT_WITH_AVAILABILITY_BIT | BlockingQueryResultFlags);
-		VERIFYVULKANRESULT_EXPANDED(Result);
+		auto& Event = EmplaceProfilerEvent<UE::RHI::GPUProfiler::FEvent::FEndWork>(0);
+		const uint32 IndexInPool = QueryPool->ReserveQuery(&Event.GPUTimestampBOP);
+		VulkanRHI::vkCmdWriteTimestamp(GetHandle(), VK_PIPELINE_STAGE_BOTTOM_OF_PIPE_BIT, QueryPool->GetHandle(), IndexInPool);
 	}
+#else
+	if (QueryPool)
+	{
+		const uint32 IndexInPool = QueryPool->ReserveQuery(&EndTimestamp);
+		VulkanRHI::vkCmdWriteTimestamp(GetHandle(), VK_PIPELINE_STAGE_BOTTOM_OF_PIPE_BIT, QueryPool->GetHandle(), IndexInPool);
+	}
+#endif
 
-	PendingTimestampQueries.Reset();
-}
-
-
-void FVulkanCommandBuffer::End()
-{
-	checkf(IsOutsideRenderPass(), TEXT("Can't End as we're inside a render pass! CmdBuffer 0x%p State=%d"), CommandBufferHandle, (int32)State);
 	VERIFYVULKANRESULT(VulkanRHI::vkEndCommandBuffer(GetHandle()));
 	State = EState::HasEnded;
 }
 
-void FVulkanCommandBuffer::Begin()
+void FVulkanCommandBuffer::Begin(FVulkanTimingQueryPool* QueryPool)
 {
 	{
 		FScopeLock ScopeLock(CommandBufferPool->GetCS());
-		if(State == EState::NeedReset)
+		if (State == EState::NeedReset)
 		{
 			VulkanRHI::vkResetCommandBuffer(CommandBufferHandle, VK_COMMAND_BUFFER_RESET_RELEASE_RESOURCES_BIT);
 		}
@@ -201,6 +175,23 @@ void FVulkanCommandBuffer::Begin()
 	}
 
 	bNeedsDynamicStateSet = true;
+
+#if RHI_NEW_GPU_PROFILER
+	if (QueryPool)
+	{
+		auto& Event = EmplaceProfilerEvent<UE::RHI::GPUProfiler::FEvent::FBeginWork>(0, UINT64_MAX);
+		const uint32 IndexInPool = QueryPool->ReserveQuery(&Event.GPUTimestampTOP);
+		VulkanRHI::vkCmdWriteTimestamp(GetHandle(), VK_PIPELINE_STAGE_TOP_OF_PIPE_BIT, QueryPool->GetHandle(), IndexInPool);
+	}
+#else
+	if (QueryPool)
+	{
+		StartTimestamp = 0;
+		EndTimestamp = 0;
+		const uint32 IndexInPool = QueryPool->ReserveQuery(&StartTimestamp);
+		VulkanRHI::vkCmdWriteTimestamp(GetHandle(), VK_PIPELINE_STAGE_TOP_OF_PIPE_BIT, QueryPool->GetHandle(), IndexInPool);
+	}
+#endif
 }
 
 void FVulkanCommandBuffer::AcquirePoolSetContainer()
diff --git a/Engine/Source/Runtime/VulkanRHI/Private/VulkanCommandBuffer.h b/Engine/Source/Runtime/VulkanRHI/Private/VulkanCommandBuffer.h
--- a/Engine/Source/Runtime/VulkanRHI/Private/VulkanCommandBuffer.h
+++ b/Engine/Source/Runtime/VulkanRHI/Private/VulkanCommandBuffer.h
@@ -10,9 +10,6 @@
 #include "VulkanConfiguration.h"
 #include "VulkanBarriers.h"
 
-extern int32 GVulkanProfileCmdBuffers;
-extern int32 GVulkanUseCmdBufferTimingForGPUTime;
-
 class FVulkanDevice;
 class FVulkanCommandBufferPool;
 class FVulkanContextCommon;
@@ -20,7 +17,7 @@ class FVulkanRenderTargetLayout;
 class FVulkanQueue;
 class FVulkanDescriptorPoolSetContainer;
 class FVulkanGPUTiming;
-struct FVulkanTimingQueryResult;
+class FVulkanTimingQueryPool;
 
 namespace VulkanRHI
 {
@@ -78,8 +75,8 @@ public:
 		return CommandBufferHandle;
 	}
 
-	void Begin();
-	void End();
+	void Begin(FVulkanTimingQueryPool* QueryPool);
+	void End(FVulkanTimingQueryPool* QueryPool);
 
 	enum class EState : uint8
 	{
@@ -94,13 +91,13 @@ public:
 
 	TArray<VkViewport, TInlineAllocator<2>> CurrentViewports;
 	TArray<VkRect2D, TInlineAllocator<2>> CurrentScissors;
-	uint32 CurrentStencilRef;
-	EState State;
-	uint8 bNeedsDynamicStateSet			: 1;
-	uint8 bHasPipeline					: 1;
-	uint8 bHasViewport					: 1;
-	uint8 bHasScissor					: 1;
-	uint8 bHasStencilRef				: 1;
+	uint32 CurrentStencilRef = 0;
+	EState State = EState::NotAllocated;
+	uint8 bNeedsDynamicStateSet			: 1 = 1;
+	uint8 bHasPipeline					: 1 = 0;
+	uint8 bHasViewport					: 1 = 0;
+	uint8 bHasScissor					: 1 = 0;
+	uint8 bHasStencilRef				: 1 = 0;
 
 	// You never want to call Begin/EndRenderPass directly as it will mess up the layout manager.
 	void BeginRenderPass(const FVulkanRenderTargetLayout& Layout, class FVulkanRenderPass* RenderPass, class FVulkanFramebuffer* Framebuffer, const VkClearValue* AttachmentClearValues);
@@ -111,37 +108,67 @@ public:
 
 	bool AcquirePoolSetAndDescriptorsIfNeeded(const class FVulkanDescriptorSetsLayout& Layout, bool bNeedDescriptors, VkDescriptorSet* OutDescriptors);
 
-	struct FPendingTimestampQuery
+#if RHI_NEW_GPU_PROFILER
+	template <typename TEventType, typename... TArgs>
+	TEventType& EmplaceProfilerEvent(TArgs&&... Args)
 	{
-		uint64 Index;
-		uint64 Count;
-		FVulkanTimingQueryResult* Results;
-		VkQueryPool PoolHandle;
-		bool bBlocking;
-	};
-	void AddPendingTimestampQuery(uint64 Index, uint64 Count, VkQueryPool PoolHandle, FVulkanTimingQueryResult* Results, bool bBlocking);
-	void ResolvePendingTimestampQueries();
+		TEventType& Data = EventStream.Emplace<TEventType>(Forward<TArgs>(Args)...);
+
+		if constexpr (std::is_same_v<UE::RHI::GPUProfiler::FEvent::FBeginWork, TEventType>)
+		{
+			// Store BeginEvents in a separate array as the CPUTimestamp field needs updating at submit time.
+			BeginEvents.Add(&Data);
+		}
+
+		return Data;
+	}
+
+	void FlushProfilerEvents(UE::RHI::GPUProfiler::FEventStream& Destination, uint64 CPUTimestamp)
+	{
+		for (UE::RHI::GPUProfiler::FEvent::FBeginWork* BeginEvent : BeginEvents)
+		{
+			BeginEvent->CPUTimestamp = CPUTimestamp;
+		}
+
+		Destination.Append(MoveTemp(EventStream));
+	}
+#else
+	uint64 GetBusyCycles() const
+	{
+		if (EndTimestamp > StartTimestamp)
+		{
+			return (EndTimestamp - StartTimestamp);
+		}
+		return 0;
+	}
+#endif
 
 	// Public immutable reference to be accessed directly
 	FVulkanDevice& Device;
 
 private:
-	VkCommandBuffer CommandBufferHandle;
+	VkCommandBuffer CommandBufferHandle = VK_NULL_HANDLE;
 	double SubmittedTime = 0.0f;
 
-	TArray<FPendingTimestampQuery> PendingTimestampQueries;
-
 	void ResetFenceStatus();
 
-	FVulkanCommandBufferPool* CommandBufferPool;
+	FVulkanCommandBufferPool* CommandBufferPool = nullptr;
 
-	uint64 LastValidTiming;
+	uint64 LastValidTiming = 0;
 
 	void AcquirePoolSetContainer();
 
 	void AllocMemory();
 	void FreeMemory();
 
+#if RHI_NEW_GPU_PROFILER
+	UE::RHI::GPUProfiler::FEventStream EventStream;
+	TArray<UE::RHI::GPUProfiler::FEvent::FBeginWork*, TInlineAllocator<8>> BeginEvents;
+#else
+	uint64 StartTimestamp = 0;
+	uint64 EndTimestamp = 0;
+#endif
+
 public:
 	//#todo-rco: Hide this
 	TMap<uint32, class FVulkanTypedDescriptorPoolSet*> TypedDescriptorPoolSets;
diff --git a/Engine/Source/Runtime/VulkanRHI/Private/VulkanCommands.cpp b/Engine/Source/Runtime/VulkanRHI/Private/VulkanCommands.cpp
--- a/Engine/Source/Runtime/VulkanRHI/Private/VulkanCommands.cpp
+++ b/Engine/Source/Runtime/VulkanRHI/Private/VulkanCommands.cpp
@@ -290,10 +290,12 @@ void FVulkanCommandListContext::RHIDispatchComputeShader(uint32 ThreadGroupCount
 
 	VulkanRHI::vkCmdDispatch(CommandBufferHandle, ThreadGroupCountX, ThreadGroupCountY, ThreadGroupCountZ);
 
+#if (RHI_NEW_GPU_PROFILER == 0)
 	if (FVulkanPlatform::RegisterGPUWork() && IsImmediate())
 	{
 		RegisterGPUDispatch(FIntVector(ThreadGroupCountX, ThreadGroupCountY, ThreadGroupCountZ));	
 	}
+#endif
 
 	VulkanRHI::DebugHeavyWeightBarrier(CommandBufferHandle, 2);
 }
@@ -314,10 +316,12 @@ void FVulkanCommandListContext::RHIDispatchIndirectComputeShader(FRHIBuffer* Arg
 
 	VulkanRHI::vkCmdDispatchIndirect(CommandBufferHandle, ArgumentBuffer->GetHandle(), ArgumentBuffer->GetOffset() + ArgumentOffset);
 
+#if (RHI_NEW_GPU_PROFILER == 0)
 	if (FVulkanPlatform::RegisterGPUWork()/* && IsImmediate()*/)
 	{
 		RegisterGPUDispatch(FIntVector(1, 1, 1));	
 	}
+#endif
 
 	VulkanRHI::DebugHeavyWeightBarrier(CommandBufferHandle, 2);
 }
@@ -627,10 +631,12 @@ void FVulkanCommandListContext::RHIDrawPrimitive(uint32 BaseVertexIndex, uint32
 
 	VulkanRHI::vkCmdDraw(CommandBuffer.GetHandle(), NumVertices, NumInstances, BaseVertexIndex, 0);
 
+#if (RHI_NEW_GPU_PROFILER == 0)
 	if (FVulkanPlatform::RegisterGPUWork() && IsImmediate())
 	{
 		RegisterGPUWork(NumPrimitives * NumInstances, NumVertices * NumInstances);
 	}
+#endif
 }
 
 void FVulkanCommandListContext::RHIDrawPrimitiveIndirect(FRHIBuffer* ArgumentBufferRHI, uint32 ArgumentOffset)
@@ -652,10 +658,12 @@ void FVulkanCommandListContext::RHIDrawPrimitiveIndirect(FRHIBuffer* ArgumentBuf
 
 	VulkanRHI::vkCmdDrawIndirect(CommandBufferHandle, ArgumentBuffer->GetHandle(), ArgumentBuffer->GetOffset() + ArgumentOffset, 1, sizeof(VkDrawIndirectCommand));
 
+#if (RHI_NEW_GPU_PROFILER == 0)
 	if (FVulkanPlatform::RegisterGPUWork() && IsImmediate())
 	{
 		RegisterGPUWork(1);
 	}
+#endif
 }
 
 void FVulkanCommandListContext::RHIDrawIndexedPrimitive(FRHIBuffer* IndexBufferRHI, int32 BaseVertexIndex, uint32 FirstInstance,
@@ -679,10 +687,12 @@ void FVulkanCommandListContext::RHIDrawIndexedPrimitive(FRHIBuffer* IndexBufferR
 	uint32 NumIndices = GetVertexCountForPrimitiveCount(NumPrimitives, PendingGfxState->PrimitiveType);
 	VulkanRHI::vkCmdDrawIndexed(CommandBufferHandle, NumIndices, NumInstances, StartIndex, BaseVertexIndex, FirstInstance);
 
+#if (RHI_NEW_GPU_PROFILER == 0)
 	if (FVulkanPlatform::RegisterGPUWork() && IsImmediate())
 	{
 		RegisterGPUWork(NumPrimitives * NumInstances, NumVertices * NumInstances);
 	}
+#endif
 }
 
 void FVulkanCommandListContext::RHIDrawIndexedIndirect(FRHIBuffer* IndexBufferRHI, FRHIBuffer* ArgumentsBufferRHI, int32 DrawArgumentsIndex, uint32 /*NumInstances*/)
@@ -706,10 +716,12 @@ void FVulkanCommandListContext::RHIDrawIndexedIndirect(FRHIBuffer* IndexBufferRH
 
 	VulkanRHI::vkCmdDrawIndexedIndirect(CommandBufferHandle, ArgumentBuffer->GetHandle(), ArgumentBuffer->GetOffset() + ArgumentOffset, 1, sizeof(VkDrawIndexedIndirectCommand));
 
+#if (RHI_NEW_GPU_PROFILER == 0)
 	if (FVulkanPlatform::RegisterGPUWork() && IsImmediate())
 	{
 		RegisterGPUWork(1);
 	}
+#endif
 }
 
 void FVulkanCommandListContext::RHIDrawIndexedPrimitiveIndirect(FRHIBuffer* IndexBufferRHI, FRHIBuffer* ArgumentBufferRHI, uint32 ArgumentOffset)
@@ -732,10 +744,12 @@ void FVulkanCommandListContext::RHIDrawIndexedPrimitiveIndirect(FRHIBuffer* Inde
 
 	VulkanRHI::vkCmdDrawIndexedIndirect(CommandBufferHandle, ArgumentBuffer->GetHandle(), ArgumentBuffer->GetOffset() + ArgumentOffset, 1, sizeof(VkDrawIndexedIndirectCommand));
 
+#if (RHI_NEW_GPU_PROFILER == 0)
 	if (FVulkanPlatform::RegisterGPUWork() && IsImmediate())
 	{
 		RegisterGPUWork(1); 
 	}
+#endif
 }
 
 #if PLATFORM_SUPPORTS_MESH_SHADERS
@@ -751,10 +765,12 @@ void FVulkanCommandListContext::RHIDispatchMeshShader(uint32 ThreadGroupCountX,
 
 	VulkanRHI::vkCmdDrawMeshTasksEXT(CommandBuffer.GetHandle(), ThreadGroupCountX, ThreadGroupCountY, ThreadGroupCountZ);
 
+#if (RHI_NEW_GPU_PROFILER == 0)
 	if (FVulkanPlatform::RegisterGPUWork() && IsImmediate())
 	{
 		RegisterGPUWork(1);
 	}
+#endif
 }
 
 void FVulkanCommandListContext::RHIDispatchIndirectMeshShader(FRHIBuffer* ArgumentBufferRHI, uint32 ArgumentOffset)
@@ -771,10 +787,12 @@ void FVulkanCommandListContext::RHIDispatchIndirectMeshShader(FRHIBuffer* Argume
 
 	VulkanRHI::vkCmdDrawMeshTasksIndirectEXT(CommandBuffer.GetHandle(), ArgumentBuffer->GetHandle(), ArgumentBuffer->GetOffset() + ArgumentOffset, 1, sizeof(VkDrawMeshTasksIndirectCommandEXT));
 
+#if (RHI_NEW_GPU_PROFILER == 0)
 	if (FVulkanPlatform::RegisterGPUWork() && IsImmediate())
 	{
 		RegisterGPUWork(1);
 	}
+#endif
 }
 #endif
 
diff --git a/Engine/Source/Runtime/VulkanRHI/Private/VulkanContext.cpp b/Engine/Source/Runtime/VulkanRHI/Private/VulkanContext.cpp
--- a/Engine/Source/Runtime/VulkanRHI/Private/VulkanContext.cpp
+++ b/Engine/Source/Runtime/VulkanRHI/Private/VulkanContext.cpp
@@ -7,6 +7,7 @@
 #include "VulkanRHIPrivate.h"
 #include "VulkanDevice.h"
 #include "VulkanPendingState.h"
+#include "VulkanQuery.h"
 #include "DynamicRHI.h"
 
 
@@ -36,9 +37,10 @@ void FVulkanContextCommon::EndPayload()
 	{
 		FlushPendingSyncPoints();
 
-		if (Payloads.Last()->CommandBuffers.Num() > 0)
+		FVulkanPayload* Payload = Payloads.Last();
+		if (Payload->CommandBuffers.Num() > 0)
 		{
-			FVulkanCommandBuffer* CommandBuffer = Payloads.Last()->CommandBuffers.Last();
+			FVulkanCommandBuffer* CommandBuffer = Payload->CommandBuffers.Last();
 			checkSlow(!CommandBuffer->IsSubmitted() && CommandBuffer->HasBegun());
 
 			if (!CommandBuffer->IsOutsideRenderPass())
@@ -47,7 +49,7 @@ void FVulkanContextCommon::EndPayload()
 				CommandBuffer->EndRenderPass();
 			}
 
-			CommandBuffer->End();
+			CommandBuffer->End(GRHIGlobals.SupportsTimestampRenderQueries ? GetCurrentTimestampQueryPool(*Payload) : nullptr);
 		}
 	}
 }
@@ -66,42 +68,9 @@ void FVulkanContextCommon::Finalize(TArray<FVulkanPayload*>& OutPayloads)
 
 	EndPayload();
 
-	if (FVulkanPayload* ResetPayload = FlushResetQueryPools())
-	{
-		OutPayloads.Add(ResetPayload);
-	}
-
 	OutPayloads.Append(MoveTemp(Payloads));
 }
 
-void FVulkanContextCommon::AddQueryPoolForReset(VkQueryPool QueryPool, uint32 Size)
-{
-	FScopeLock ScopeLock(&Pool.CS);
-	FQueryPoolReset QPR = { QueryPool, Size };
-	PoolResets.Add(QPR);
-}
-
-FVulkanPayload* FVulkanContextCommon::FlushResetQueryPools()
-{
-	FScopeLock ScopeLock(&Pool.CS);
-	if (PoolResets.Num())
-	{
-		FVulkanCommandBuffer* PoolResetCommandBuffer = PrepareNewCommandBuffer(false);
-		const VkCommandBuffer CommandBufferHandle = PoolResetCommandBuffer->GetHandle();
-		for (FQueryPoolReset& QPR : PoolResets)
-		{
-			VulkanRHI::vkCmdResetQueryPool(CommandBufferHandle, QPR.Pool, 0, QPR.Size);
-		}
-		PoolResetCommandBuffer->End();
-
-		FVulkanPayload* ResetPayload = new FVulkanPayload(Queue);
-		ResetPayload->CommandBuffers.Add(PoolResetCommandBuffer);
-		PoolResets.Empty();
-		return ResetPayload;
-	}
-	return nullptr;
-}
-
 void FVulkanContextCommon::FlushCommands(EVulkanFlushFlags FlushFlags)
 {
 	FVulkanSyncPointRef SyncPoint;
@@ -173,7 +142,7 @@ void FVulkanContextCommon::FreeUnusedCmdBuffers(bool bTrimMemory)
 #endif
 }
 
-FVulkanCommandBuffer* FVulkanContextCommon::PrepareNewCommandBuffer(bool bTransferLayoutMgr)
+void FVulkanContextCommon::PrepareNewCommandBuffer(FVulkanPayload& Payload)
 {
 	FScopeLock ScopeLock(&Pool.CS);
 
@@ -199,17 +168,10 @@ FVulkanCommandBuffer* FVulkanContextCommon::PrepareNewCommandBuffer(bool bTransf
 		NewCommandBuffer = Pool.Create();
 	}
 
-	NewCommandBuffer->Begin();
-	return NewCommandBuffer;
-}
+	Payload.CommandBuffers.Add(NewCommandBuffer);
 
-#if (RHI_NEW_GPU_PROFILER == 0)
-uint32 FVulkanContextCommon::CalculateGPUTime()
-{
-	uint32 Time = 0;
-	return Time;
+	NewCommandBuffer->Begin(GRHIGlobals.SupportsTimestampRenderQueries ? GetCurrentTimestampQueryPool(Payload) : nullptr);
 }
-#endif // (RHI_NEW_GPU_PROFILER == 0)
 
 
 FVulkanCommandListContext::FVulkanCommandListContext(FVulkanDynamicRHI* InRHI, FVulkanDevice& InDevice, FVulkanQueue& InQueue, FVulkanCommandListContext* InImmediate)
@@ -225,12 +187,6 @@ FVulkanCommandListContext::FVulkanCommandListContext(FVulkanDynamicRHI* InRHI, F
 	FrameTiming->Initialize();
 #endif
 
-	if (IsImmediate())
-	{
-		// Insert the Begin frame timestamp query. On EndDrawingViewport() we'll insert the End and immediately after a new Begin()
-		WriteBeginTimestamp(GetActiveCmdBuffer());
-	}
-
 	// Create Pending state, contains pipeline states such as current shader and etc..
 	PendingGfxState = new FVulkanPendingGfxState(&Device, *this);
 	PendingComputeState = new FVulkanPendingComputeState(&Device, *this);
diff --git a/Engine/Source/Runtime/VulkanRHI/Private/VulkanContext.h b/Engine/Source/Runtime/VulkanRHI/Private/VulkanContext.h
--- a/Engine/Source/Runtime/VulkanRHI/Private/VulkanContext.h
+++ b/Engine/Source/Runtime/VulkanRHI/Private/VulkanContext.h
@@ -76,7 +76,7 @@ public:
 
 		if (Payload.CommandBuffers.Num() == 0)
 		{
-			Payload.CommandBuffers.Add(PrepareNewCommandBuffer(true));
+			PrepareNewCommandBuffer(Payload);
 		}
 
 		return *Payload.CommandBuffers.Last();
@@ -120,17 +120,11 @@ public:
 	// payloads to the given array. Resets the context so new commands can be recorded.
 	virtual void Finalize(TArray<FVulkanPayload*>& OutPayloads);
 
-	void AddQueryPoolForReset(VkQueryPool Pool, uint32 Size);
-
 	VkCommandPool GetHandle() const
 	{
 		return Pool.GetHandle();
 	}
 
-#if (RHI_NEW_GPU_PROFILER == 0)
-	uint32 CalculateGPUTime();
-#endif
-
 	void FreeUnusedCmdBuffers(bool bTrimMemory);
 
 	// Should only ne used when we are certain there are no other pending contexts (like UploadContext)
@@ -190,23 +184,28 @@ public:
 		return ContextSyncPoint;
 	}
 
+	TArray<FVulkanQueryPool*>& GetQueryPoolArray(EVulkanQueryPoolType Type)
+	{
+		FVulkanPayload& Payload = GetPayload(EPhase::Execute);
+		return Payload.QueryPools[(int32)Type];
+	}
+
+	FVulkanTimingQueryPool* GetCurrentTimestampQueryPool()
+	{
+		return GetCurrentTimestampQueryPool(GetPayload(EPhase::Execute));
+	}
+
 protected:
 	// Number of times EndFrame() has been called on this context
 	uint64 FrameCounter = 0;
 
 private:
-	FVulkanCommandBuffer* PrepareNewCommandBuffer(bool bTransferLayoutMgr);
-	FVulkanPayload* FlushResetQueryPools();
+	FVulkanTimingQueryPool* GetCurrentTimestampQueryPool(FVulkanPayload& Payload);
+	void PrepareNewCommandBuffer(FVulkanPayload& Payload);
 	void NewPayload();
 	void EndPayload();
 
 	FVulkanCommandBufferPool Pool;
-	struct FQueryPoolReset
-	{
-		VkQueryPool Pool;
-		uint32 Size;
-	};
-	TArray<FQueryPoolReset> PoolResets;
 
 	TArray<FVulkanPayload*> Payloads;
 	EPhase CurrentPhase = EPhase::Wait;
@@ -297,7 +296,17 @@ public:
 	virtual void RHIBeginRenderQuery(FRHIRenderQuery* RenderQuery) final override;
 	virtual void RHIEndRenderQuery(FRHIRenderQuery* RenderQuery) final override;
 
-#if (RHI_NEW_GPU_PROFILER == 0)
+#if RHI_NEW_GPU_PROFILER
+	void FlushProfilerStats()
+	{
+		// Flush accumulated draw stats
+		if (StatEvent)
+		{
+			GetCommandBuffer().EmplaceProfilerEvent<UE::RHI::GPUProfiler::FEvent::FStats>() = StatEvent;
+			StatEvent = {};
+		}
+	}
+#else
 	virtual void RHICalibrateTimers(FRHITimestampCalibrationQuery* CalibrationQuery) final override;
 #endif
 
@@ -356,12 +365,7 @@ public:
 		return CurrentFramebuffer;
 	}
 
-	void WriteBeginTimestamp(FVulkanCommandBuffer* CmdBuffer);
-	void WriteEndTimestamp(FVulkanCommandBuffer* CmdBuffer);
-
 #if (RHI_NEW_GPU_PROFILER == 0)
-	void ReadAndCalculateGPUFrameTime();
-	
 	inline FVulkanGPUProfiler& GetGPUProfiler()
 	{
 		return GpuProfiler;
@@ -375,20 +379,19 @@ public:
 
 	void ReleasePendingState();
 
+	FVulkanOcclusionQueryPool* GetCurrentOcclusionQueryPool();
+
 protected:
 	FVulkanDynamicRHI* RHI;
 	FVulkanCommandListContext* Immediate;
 
-	void BeginOcclusionQueryBatch(FVulkanCommandBuffer* CmdBuffer, uint32 NumQueriesInBatch);
-	void EndOcclusionQueryBatch();
+	void BeginOcclusionQueryBatch(uint32 NumQueriesInBatch);
 
 	TArray<FString> EventStack;
 
 	FVulkanRenderPass* CurrentRenderPass = nullptr;
 	FVulkanFramebuffer* CurrentFramebuffer = nullptr;
 
-	FVulkanOcclusionQueryPool* CurrentOcclusionQueryPool = nullptr;
-
 	FVulkanPendingGfxState* PendingGfxState = nullptr;
 	FVulkanPendingComputeState* PendingComputeState = nullptr;
 
@@ -417,10 +420,7 @@ private:
 	friend class FVulkanDevice;
 	friend class FVulkanDynamicRHI;
 
-#if RHI_NEW_GPU_PROFILER
-	void RegisterGPUWork(uint32 NumPrimitives = 0, uint32 NumVertices = 0)	{ checkNoEntry(); } // @todo - new gpu profiler
-	void RegisterGPUDispatch(FIntVector GroupCount)	                        { checkNoEntry(); } // @todo - new gpu profiler
-#else
+#if (RHI_NEW_GPU_PROFILER == 0)
 	void RegisterGPUWork(uint32 NumPrimitives = 0, uint32 NumVertices = 0)	{ GpuProfiler.RegisterGPUWork(NumPrimitives, NumVertices); }
 	void RegisterGPUDispatch(FIntVector GroupCount)	                        { GpuProfiler.RegisterGPUDispatch(GroupCount); }
 
diff --git a/Engine/Source/Runtime/VulkanRHI/Private/VulkanDevice.cpp b/Engine/Source/Runtime/VulkanRHI/Private/VulkanDevice.cpp
--- a/Engine/Source/Runtime/VulkanRHI/Private/VulkanDevice.cpp
+++ b/Engine/Source/Runtime/VulkanRHI/Private/VulkanDevice.cpp
@@ -1507,16 +1507,14 @@ void FVulkanDevice::Destroy()
 	delete RenderPassManager;
 	RenderPassManager = nullptr;
 
-	for (FVulkanOcclusionQueryPool* Pool : UsedOcclusionQueryPools)
+	for (TArray<FVulkanQueryPool*>& PoolArray : FreeQueryPools)
 	{
-		delete Pool;
-	}
-	UsedOcclusionQueryPools.SetNum(0, EAllowShrinking::No);
-	for (FVulkanOcclusionQueryPool* Pool : FreeOcclusionQueryPools)
-	{
-		delete Pool;
+		for (FVulkanQueryPool* Pool : PoolArray)
+		{
+			delete Pool;
+		}
+		PoolArray.Empty();
 	}
-	FreeOcclusionQueryPools.SetNum(0, EAllowShrinking::No);
 
 	delete PipelineStateCache;
 	PipelineStateCache = nullptr;
@@ -1668,7 +1666,39 @@ FVulkanTransientHeapCache& FVulkanDevice::GetOrCreateTransientHeapCache()
 	return *TransientHeapCache;
 }
 
-#if (RHI_NEW_GPU_PROFILER == 0)
+#if RHI_NEW_GPU_PROFILER
+void FVulkanDevice::GetCalibrationTimestamp(FVulkanTiming& InOutTiming)
+{
+	// TimestampPeriod is the number of nanoseconds required for a timestamp query to be incremented by 1
+	InOutTiming.GPUFrequency = (uint64)((1000.0 * 1000.0 * 1000.0) / (double)GetDeviceProperties().limits.timestampPeriod);
+	InOutTiming.CPUFrequency = uint64(1.0 / FPlatformTime::GetSecondsPerCycle64());
+
+	if (GetOptionalExtensions().HasEXTCalibratedTimestamps)
+	{
+		uint64_t Timestamps[2] = { 0, 0 };
+		uint64_t MaxDeviations[2] = { 0, 0 };
+		VkCalibratedTimestampInfoKHR CalibratedTimestampInfo[2];
+		ZeroVulkanStruct(CalibratedTimestampInfo[0], VK_STRUCTURE_TYPE_CALIBRATED_TIMESTAMP_INFO_KHR);
+		ZeroVulkanStruct(CalibratedTimestampInfo[1], VK_STRUCTURE_TYPE_CALIBRATED_TIMESTAMP_INFO_KHR);
+		CalibratedTimestampInfo[0].timeDomain = VK_TIME_DOMAIN_DEVICE_KHR;
+		CalibratedTimestampInfo[1].timeDomain = FVulkanPlatform::GetTimeDomain();
+
+		const bool bHasHostTimeDomain = (CalibratedTimestampInfo[1].timeDomain != VK_TIME_DOMAIN_DEVICE_KHR);
+		const uint32 TimestampCount = bHasHostTimeDomain ? 2 : 1;
+
+		VERIFYVULKANRESULT(VulkanDynamicAPI::vkGetCalibratedTimestampsEXT(GetInstanceHandle(), TimestampCount, CalibratedTimestampInfo, Timestamps, MaxDeviations));
+
+		InOutTiming.GPUTimestamp = Timestamps[0];
+		InOutTiming.CPUTimestamp = bHasHostTimeDomain ? Timestamps[1] : FPlatformTime::Cycles64();
+		check(InOutTiming.CPUTimestamp <= FPlatformTime::Cycles64());
+	}
+	else
+	{
+		InOutTiming.GPUTimestamp = 0;
+		InOutTiming.CPUTimestamp = 0;
+	}
+}
+#else
 FGPUTimingCalibrationTimestamp FVulkanDevice::GetCalibrationTimestamp()
 {
 	auto ToMicroseconds = [](uint64_t Timestamp)
diff --git a/Engine/Source/Runtime/VulkanRHI/Private/VulkanDevice.h b/Engine/Source/Runtime/VulkanRHI/Private/VulkanDevice.h
--- a/Engine/Source/Runtime/VulkanRHI/Private/VulkanDevice.h
+++ b/Engine/Source/Runtime/VulkanRHI/Private/VulkanDevice.h
@@ -9,6 +9,7 @@
 #include "VulkanMemory.h"
 #include "VulkanResources.h"
 #include "VulkanQueue.h"
+#include "VulkanQuery.h"
 #include "GPUProfiler.h"
 
 class FVulkanSamplerState;
@@ -20,7 +21,6 @@ class FVulkanBindlessDescriptorManager;
 class FVulkanCommandListContextImmediate;
 class FVulkanTransientHeapCache;
 class FVulkanDeviceExtension;
-class FVulkanOcclusionQueryPool;
 class FVulkanRenderPassManager;
 class FVulkanRayTracingCompactionRequestHandler;
 
@@ -267,22 +267,22 @@ public:
 
 	void WaitUntilIdle();
 
-	inline EGpuVendorId GetVendorId() const
+	EGpuVendorId GetVendorId() const
 	{
 		return VendorId;
 	}
 
-	inline bool HasAsyncComputeQueue() const
+	bool HasAsyncComputeQueue() const
 	{
 		return bAsyncComputeQueue;
 	}
 
-	inline bool CanPresentOnComputeQueue() const
+	bool CanPresentOnComputeQueue() const
 	{
 		return bPresentOnComputeQueue;
 	}
 
-	inline bool IsRealAsyncComputeContext(const FVulkanCommandListContext* InContext) const
+	bool IsRealAsyncComputeContext(const FVulkanCommandListContext* InContext) const
 	{
 		if (bAsyncComputeQueue)
 		{
@@ -293,54 +293,54 @@ public:
 		return false;
 	}
 
-	inline FVulkanQueue* GetGraphicsQueue()
+	FVulkanQueue* GetGraphicsQueue()
 	{
 		return Queues[(int32)EVulkanQueueType::Graphics];
 	}
 
-	inline FVulkanQueue* GetComputeQueue()
+	FVulkanQueue* GetComputeQueue()
 	{
 		return Queues[(int32)EVulkanQueueType::AsyncCompute];
 	}
 
-	inline FVulkanQueue* GetTransferQueue()
+	FVulkanQueue* GetTransferQueue()
 	{
 		return Queues[(int32)EVulkanQueueType::Transfer];
 	}
 
-	inline FVulkanQueue* GetPresentQueue()
+	FVulkanQueue* GetPresentQueue()
 	{
 		return PresentQueue;
 	}
 
 	void ForEachQueue(TFunctionRef<void(FVulkanQueue&)> Callback);
 
-	inline VkPhysicalDevice GetPhysicalHandle() const
+	VkPhysicalDevice GetPhysicalHandle() const
 	{
 		return Gpu;
 	}
 
-	inline const VkPhysicalDeviceProperties& GetDeviceProperties() const
+	const VkPhysicalDeviceProperties& GetDeviceProperties() const
 	{
 		return GpuProps;
 	}
 
-	inline VkExtent2D GetBestMatchedFragmentSize(EVRSShadingRate Rate) const
+	VkExtent2D GetBestMatchedFragmentSize(EVRSShadingRate Rate) const
 	{
 		return FragmentSizeMap[Rate];
 	}
 
-	inline const VkPhysicalDeviceLimits& GetLimits() const
+	const VkPhysicalDeviceLimits& GetLimits() const
 	{
 		return GpuProps.limits;
 	}
 
-	inline const VkPhysicalDeviceIDPropertiesKHR& GetDeviceIdProperties() const
+	const VkPhysicalDeviceIDPropertiesKHR& GetDeviceIdProperties() const
 	{
 		return GpuIdProps;
 	}
 
-	inline const VkPhysicalDeviceSubgroupProperties& GetDeviceSubgroupProperties() const
+	const VkPhysicalDeviceSubgroupProperties& GetDeviceSubgroupProperties() const
 	{
 		return GpuSubgroupProps;
 	}
@@ -351,104 +351,104 @@ public:
 	void CleanUpRayTracing();
 
 #if VULKAN_SUPPORTS_VALIDATION_CACHE
-	inline VkValidationCacheEXT GetValidationCache() const
+	VkValidationCacheEXT GetValidationCache() const
 	{
 		return ValidationCache;
 	}
 #endif
 
-	inline const FVulkanPhysicalDeviceFeatures& GetPhysicalDeviceFeatures() const
+	const FVulkanPhysicalDeviceFeatures& GetPhysicalDeviceFeatures() const
 	{
 		return PhysicalDeviceFeatures;
 	}
 
-	inline bool HasUnifiedMemory() const
+	bool HasUnifiedMemory() const
 	{
 		return DeviceMemoryManager.HasUnifiedMemory();
 	}
 
 	bool SupportsBindless() const;
 
-	inline uint64 GetTimestampValidBitsMask() const
+	uint64 GetTimestampValidBitsMask() const
 	{
 		return TimestampValidBitsMask;
 	}
 
 	const VkComponentMapping& GetFormatComponentMapping(EPixelFormat UEFormat) const;
 
-	inline VkDevice GetInstanceHandle() const
+	VkDevice GetInstanceHandle() const
 	{
 		return Device;
 	}
 
-	inline const FVulkanSamplerState& GetDefaultSampler() const
+	const FVulkanSamplerState& GetDefaultSampler() const
 	{
 		return *DefaultSampler;
 	}
 
 	const VkFormatProperties& GetFormatProperties(VkFormat InFormat) const;
 
-	inline VulkanRHI::FDeviceMemoryManager& GetDeviceMemoryManager()
+	VulkanRHI::FDeviceMemoryManager& GetDeviceMemoryManager()
 	{
 		return DeviceMemoryManager;
 	}
 
-	inline const VkPhysicalDeviceMemoryProperties& GetDeviceMemoryProperties() const
+	const VkPhysicalDeviceMemoryProperties& GetDeviceMemoryProperties() const
 	{
 		return DeviceMemoryManager.GetMemoryProperties();
 	}
 
-	inline VulkanRHI::FMemoryManager& GetMemoryManager()
+	VulkanRHI::FMemoryManager& GetMemoryManager()
 	{
 		return MemoryManager;
 	}
 
-	inline VulkanRHI::FDeferredDeletionQueue2& GetDeferredDeletionQueue()
+	VulkanRHI::FDeferredDeletionQueue2& GetDeferredDeletionQueue()
 	{
 		return DeferredDeletionQueue;
 	}
 
-	inline VulkanRHI::FStagingManager& GetStagingManager()
+	VulkanRHI::FStagingManager& GetStagingManager()
 	{
 		return StagingManager;
 	}
 
-	inline VulkanRHI::FFenceManager& GetFenceManager()
+	VulkanRHI::FFenceManager& GetFenceManager()
 	{
 		return FenceManager;
 	}
 
-	inline VulkanRHI::FTempBlockAllocator& GetTempBlockAllocator()
+	VulkanRHI::FTempBlockAllocator& GetTempBlockAllocator()
 	{
 		return *TempBlockAllocator;
 	}
 
-	inline FVulkanRenderPassManager& GetRenderPassManager()
+	FVulkanRenderPassManager& GetRenderPassManager()
 	{
 		return *RenderPassManager;
 	}
 
-	inline FVulkanDescriptorSetCache& GetDescriptorSetCache()
+	FVulkanDescriptorSetCache& GetDescriptorSetCache()
 	{
 		return *DescriptorSetCache;
 	}
 
-	inline FVulkanDescriptorPoolsManager& GetDescriptorPoolsManager()
+	FVulkanDescriptorPoolsManager& GetDescriptorPoolsManager()
 	{
 		return *DescriptorPoolsManager;
 	}
 
-	inline FVulkanBindlessDescriptorManager* GetBindlessDescriptorManager()
+	FVulkanBindlessDescriptorManager* GetBindlessDescriptorManager()
 	{
 		return BindlessDescriptorManager;
 	}
 
-	inline TMap<uint32, FSamplerStateRHIRef>& GetSamplerMap()
+	TMap<uint32, FSamplerStateRHIRef>& GetSamplerMap()
 	{
 		return SamplerMap;
 	}
 
-	inline FVulkanShaderFactory& GetShaderFactory()
+	FVulkanShaderFactory& GetShaderFactory()
 	{
 		return ShaderFactory;
 	}
@@ -461,24 +461,26 @@ public:
 	void NotifyDeletedImage(VkImage Image, bool bRenderTarget);
 
 #if VULKAN_ENABLE_DRAW_MARKERS
-	inline PFN_vkCmdBeginDebugUtilsLabelEXT GetCmdBeginDebugLabel() const
+	PFN_vkCmdBeginDebugUtilsLabelEXT GetCmdBeginDebugLabel() const
 	{
 		return DebugMarkers.CmdBeginDebugLabel;
 	}
 
-	inline PFN_vkCmdEndDebugUtilsLabelEXT GetCmdEndDebugLabel() const
+	PFN_vkCmdEndDebugUtilsLabelEXT GetCmdEndDebugLabel() const
 	{
 		return DebugMarkers.CmdEndDebugLabel;
 	}
 
-	inline PFN_vkSetDebugUtilsObjectNameEXT GetSetDebugName() const
+	PFN_vkSetDebugUtilsObjectNameEXT GetSetDebugName() const
 	{
 		return DebugMarkers.SetDebugName;
 	}
 #endif
 
-	FVulkanOcclusionQueryPool* AcquireOcclusionQueryPool(FVulkanContextCommon* InCommonContext, uint32 NumQueries);
-	void ReleaseUnusedOcclusionQueryPools();
+	FVulkanOcclusionQueryPool* AcquireOcclusionQueryPool(uint32 NumQueries);
+	FVulkanTimingQueryPool* AcquireTimingQueryPool();
+	void ReleaseQueryPool(FVulkanQueryPool* Pool);
+	void RemoveStaleQueryPools();
 
 	inline class FVulkanPipelineStateCacheManager* GetPipelineStateCache()
 	{
@@ -529,7 +531,9 @@ public:
 
 	const TArray<const ANSICHAR*>& GetDeviceExtensions() { return DeviceExtensions; }
 
-#if (RHI_NEW_GPU_PROFILER == 0)
+#if RHI_NEW_GPU_PROFILER
+	void GetCalibrationTimestamp(FVulkanTiming& InOutTiming);
+#else
 	// Performs a GPU and CPU timestamp at nearly the same time.
 	// This allows aligning GPU and CPU events on the same timeline in profile visualization.
 	FGPUTimingCalibrationTimestamp GetCalibrationTimestamp();
@@ -591,9 +595,10 @@ private:
 	// Info for formats that are not in the core Vulkan spec (i.e. extensions)
 	mutable TMap<VkFormat, VkFormatProperties> ExtensionFormatProperties;
 
-	FCriticalSection OcclusionQueryPoolLock;
-	TArray<FVulkanOcclusionQueryPool*> UsedOcclusionQueryPools;
-	TArray<FVulkanOcclusionQueryPool*> FreeOcclusionQueryPools;
+	// Reusable query pools
+	FCriticalSection QueryPoolLock;
+	uint32 OcclusionQueryPoolSize = 256;
+	TStaticArray<TArray<FVulkanQueryPool*>, (int32)EVulkanQueryPoolType::Count> FreeQueryPools;
 
 	uint64 TimestampValidBitsMask = 0;
 
diff --git a/Engine/Source/Runtime/VulkanRHI/Private/VulkanDynamicRHI.h b/Engine/Source/Runtime/VulkanRHI/Private/VulkanDynamicRHI.h
--- a/Engine/Source/Runtime/VulkanRHI/Private/VulkanDynamicRHI.h
+++ b/Engine/Source/Runtime/VulkanRHI/Private/VulkanDynamicRHI.h
@@ -8,10 +8,10 @@
 
 #include "IVulkanDynamicRHI.h"
 #include "VulkanConfiguration.h"
+#include "VulkanQueue.h"
 
 class FVulkanFramebuffer;
 class FVulkanDevice;
-class FVulkanQueue;
 class FVulkanUniformBuffer;
 class FVulkanViewport;
 class FVulkanPayload;
@@ -48,6 +48,23 @@ struct FOptionalVulkanInstanceExtensions
 	}
 };
 
+#if RHI_NEW_GPU_PROFILER
+// Encapsulates the state required for tracking GPU queue performance across a frame.
+class FVulkanTiming
+{
+public:
+	FVulkanTiming(FVulkanQueue& InQueue);
+
+	FVulkanQueue& Queue;
+
+	// Timer calibration data
+	uint64 GPUFrequency = 0, GPUTimestamp = 0;
+	uint64 CPUFrequency = 0, CPUTimestamp = 0;
+
+	UE::RHI::GPUProfiler::FEventStream EventStream;
+};
+#endif
+
 PRAGMA_DISABLE_DEPRECATION_WARNINGS
 
 /** The interface which is implemented by the dynamically bound RHI. */
@@ -306,6 +323,7 @@ public:
 
 	void ProcessInterruptQueueUntil(FGraphEvent* GraphEvent);
 	void EnqueueEndOfPipeTask(TUniqueFunction<void()> TaskFunc, TUniqueFunction<void(FVulkanPayload&)> ModifyPayloadCallback = {});
+	void CompletePayload(FVulkanPayload* Payload);
 
 protected:
 	void InitializeSubmissionPipe();
@@ -316,6 +334,17 @@ protected:
 	void KickSubmissionThread();
 	void KickInterruptThread();
 
+#if RHI_NEW_GPU_PROFILER
+	struct FVulkanTimingArray : public TArray<TUniquePtr<FVulkanTiming>, TInlineAllocator<EVulkanQueueType::Count>>
+	{
+		FVulkanTiming* CreateNew(FVulkanQueue& Queue)
+		{
+			return Emplace_GetRef(MakeUnique<FVulkanTiming>(Queue)).Get();
+		}
+	};
+	FVulkanTimingArray CurrentTimingPerQueue;
+#endif
+
 	FCriticalSection SubmissionCS;
 	FVulkanThread* SubmissionThread = nullptr;
 	FCriticalSection InterruptCS;
diff --git a/Engine/Source/Runtime/VulkanRHI/Private/VulkanGPUProfiler.h b/Engine/Source/Runtime/VulkanRHI/Private/VulkanGPUProfiler.h
--- a/Engine/Source/Runtime/VulkanRHI/Private/VulkanGPUProfiler.h
+++ b/Engine/Source/Runtime/VulkanRHI/Private/VulkanGPUProfiler.h
@@ -17,6 +17,9 @@ class FVulkanTimingQueryPool;
 
 #if (RHI_NEW_GPU_PROFILER == 0)
 
+using FVulkanSyncPoint = FGraphEvent;
+using FVulkanSyncPointRef = TRefCountPtr<FVulkanSyncPoint>;
+
 class FVulkanGPUTiming : public FGPUTiming
 {
 public:
@@ -33,13 +36,13 @@ public:
 	/**
 	 * Start a GPU timing measurement.
 	 */
-	void StartTiming(FVulkanCommandBuffer* CmdBuffer = nullptr);
+	void StartTiming(FVulkanContextCommon* InContext = nullptr);
 
 	/**
 	 * End a GPU timing measurement.
 	 * The timing for this particular measurement will be resolved at a later time by the GPU.
 	 */
-	void EndTiming(FVulkanCommandBuffer* CmdBuffer = nullptr);
+	void EndTiming(FVulkanContextCommon* InContext = nullptr);
 
 	/**
 	 * Retrieves the most recently resolved timing measurement.
@@ -65,6 +68,11 @@ public:
 		return true;
 	}
 
+	bool IsTiming() const
+	{
+		return bIsTiming;
+	}
+
 	static void CalibrateTimers(FVulkanDevice& Device);
 
 private:
@@ -83,9 +91,18 @@ private:
 	bool bIsTiming;
 	bool bEndTimestampIssued;
 	uint64 LastTime = 0;
+
+	struct FPendingQuery
+	{
+		FVulkanSyncPointRef StartSyncPoint;
+		uint64 StartResult = 0;
+		FVulkanSyncPointRef EndSyncPoint;
+		uint64 EndResult = 0;
+	};
+	const uint32 MaxPendingQueries = 4;
 	uint32 NumPendingQueries = 0;
-	TQueue<uint32> PendingQueries;
-	FVulkanTimingQueryPool* Pool = nullptr;
+	TQueue<FPendingQuery*> PendingQueries;
+	FPendingQuery* ActiveQuery = nullptr;
 };
 
 /** A single perf event node, which tracks information about a appBeginDrawEvent/appEndDrawEvent range. */
diff --git a/Engine/Source/Runtime/VulkanRHI/Private/VulkanGenericPlatform.h b/Engine/Source/Runtime/VulkanRHI/Private/VulkanGenericPlatform.h
--- a/Engine/Source/Runtime/VulkanRHI/Private/VulkanGenericPlatform.h
+++ b/Engine/Source/Runtime/VulkanRHI/Private/VulkanGenericPlatform.h
@@ -144,4 +144,6 @@ public:
 
 	// Returns the shader stages for which we need support for subgroup ops.
 	static VkShaderStageFlags RequiredWaveOpsShaderStageFlags(VkShaderStageFlags VulkanDeviceShaderStageFlags) { return VulkanDeviceShaderStageFlags; }
+
+	static VkTimeDomainKHR GetTimeDomain() { return VK_TIME_DOMAIN_DEVICE_KHR; }
 };
diff --git a/Engine/Source/Runtime/VulkanRHI/Private/VulkanQuery.cpp b/Engine/Source/Runtime/VulkanRHI/Private/VulkanQuery.cpp
--- a/Engine/Source/Runtime/VulkanRHI/Private/VulkanQuery.cpp
+++ b/Engine/Source/Runtime/VulkanRHI/Private/VulkanQuery.cpp
@@ -17,8 +17,8 @@
 #include "HAL/PlatformStackwalk.h"
 #endif
 
-static uint32 GTimestampQueryStage = 0;
-TAutoConsoleVariable<int32> GTimestampQueryStageCVar(
+static int32 GTimestampQueryStage = 0;
+FAutoConsoleVariableRef CVarTimestampQueryStage(
 	TEXT("r.Vulkan.TimestampQueryStage"),
 	GTimestampQueryStage,
 	TEXT("Defines which pipeline stage is used for timestamp queries.\n")
@@ -27,18 +27,24 @@ TAutoConsoleVariable<int32> GTimestampQueryStageCVar(
 	ECVF_RenderThreadSafe
 );
 
-constexpr uint32 GMinNumberOfQueriesInPool = 256;
+static int32 GTimingQueryPoolSize = 128;
+FAutoConsoleVariableRef CVarTimingQueryPoolSize(
+	TEXT("r.Vulkan.TimingQueryPoolSize"),
+	GTimingQueryPoolSize,
+	TEXT("Amount of timing queries per pool in reusable query pools. (Default: 128)\n"),
+	ECVF_ReadOnly
+);
 
-#if PLATFORM_ANDROID
-constexpr int32 NUM_FRAMES_TO_WAIT_REUSE_POOL = 5;
-// numbers of frame to wait before releasing QueryPool
-constexpr uint32 NUM_FRAMES_TO_WAIT_RELEASE_POOL = 10; 
-#else
-constexpr int32 NUM_FRAMES_TO_WAIT_REUSE_POOL = 10;
-constexpr uint32 NUM_FRAMES_TO_WAIT_RELEASE_POOL = MAX_uint32; // never release
-#endif
+static int32 GQueryPoolDeletionDelay = 10;
+FAutoConsoleVariableRef CVarQueryPoolDeletionDelay(
+	TEXT("r.Vulkan.QueryPoolDeletionDelay"),
+	GQueryPoolDeletionDelay,
+	TEXT("Amount of frames to wait before deleting an unused query pools. (Default: 10)\n"),
+	ECVF_ReadOnly
+);
 
-FVulkanQueryPool::FVulkanQueryPool(FVulkanDevice& InDevice, FVulkanContextCommon* InCommonContext, uint32 InMaxQueries, VkQueryType InQueryType, bool bInShouldAddReset)
+
+FVulkanQueryPool::FVulkanQueryPool(FVulkanDevice& InDevice, uint32 InMaxQueries, EVulkanQueryPoolType InQueryType)
 	: Device(InDevice)
 	, QueryPool(VK_NULL_HANDLE)
 	, MaxQueries(InMaxQueries)
@@ -47,22 +53,9 @@ FVulkanQueryPool::FVulkanQueryPool(FVulkanDevice& InDevice, FVulkanContextCommon
 	INC_DWORD_STAT(STAT_VulkanNumQueryPools);
 	VkQueryPoolCreateInfo PoolCreateInfo;
 	ZeroVulkanStruct(PoolCreateInfo, VK_STRUCTURE_TYPE_QUERY_POOL_CREATE_INFO);
-	PoolCreateInfo.queryType = QueryType;
+	PoolCreateInfo.queryType = GetVkQueryType(InQueryType);
 	PoolCreateInfo.queryCount = MaxQueries;
-
 	VERIFYVULKANRESULT(VulkanRHI::vkCreateQueryPool(Device.GetInstanceHandle(), &PoolCreateInfo, VULKAN_CPU_ALLOCATOR, &QueryPool));
-	
-	// If host query resets are supported, reset all new query pools on creation
-	if (Device.GetOptionalExtensions().HasEXTHostQueryReset)
-	{
-		VulkanRHI::vkResetQueryPoolEXT(Device.GetInstanceHandle(), QueryPool, 0, MaxQueries);
-	}
-	else if (bInShouldAddReset && InCommonContext)
-	{
-		InCommonContext->AddQueryPoolForReset(QueryPool, InMaxQueries);
-	}
-
-	QueryOutput.AddZeroed(MaxQueries);
 }
 
 FVulkanQueryPool::~FVulkanQueryPool()
@@ -72,58 +65,32 @@ FVulkanQueryPool::~FVulkanQueryPool()
 	QueryPool = VK_NULL_HANDLE;
 }
 
-FVulkanTimingQueryPool::~FVulkanTimingQueryPool()
+void FVulkanQueryPool::ReserveQuery(FVulkanRenderQuery* Query)
 {
-	for (FTimestampData& QueryData : TimestampArray)
-	{
-		checkf(!QueryData.SyncPoint.IsValid() || QueryData.SyncPoint->IsComplete(), 
-			TEXT("Vulkan Query Pool being destroyed with queries still active!"));
-	}
+	checkSlow(QueryType == Query->GetQueryPoolType());
+	Query->IndexInPool = CurrentQueryCount++;
+	QueryRefs[Query->IndexInPool] = Query;
+	QueryResults[Query->IndexInPool] = &Query->Result;
 }
 
-bool FVulkanOcclusionQueryPool::CanBeReused()
+uint32 FVulkanQueryPool::ReserveQuery(uint64* ResultPtr)
 {
-	const uint64 NumWords = NumUsedQueries / 64;
-	for (int32 Index = 0; Index < NumWords; ++Index)
-	{
-		if (AcquiredIndices[Index])
-		{
-			return false;
-		}
-	}
-
-	const uint64 Remaining = NumUsedQueries % 64;
-	const uint64 Mask = ((uint64)1 << (uint64)Remaining) - 1;
-	return Mask == 0 || (AcquiredIndices[NumWords] & Mask) == 0;
+	checkSlow(QueryType == EVulkanQueryPoolType::Timestamp);
+	const uint32 IndexInPool = CurrentQueryCount++;
+	QueryRefs[IndexInPool] = nullptr;
+	QueryResults[IndexInPool] = ResultPtr;
+	return IndexInPool;
 }
 
-bool FVulkanOcclusionQueryPool::GetResults()
+void FVulkanQueryPool::Reset(FVulkanCommandBuffer& InCmdBuffer)
 {
-	if (State == EState::RT_PostGetResults)
-	{
-		return true;
-	}
-
-	if (State == RHIT_PostEndBatch)
-	{
-		check(State == EState::RHIT_PostEndBatch);
-		const VkQueryResultFlags QueryResultFlags = VK_QUERY_RESULT_64_BIT | VK_QUERY_RESULT_WAIT_BIT;
-		const VkResult Result = VulkanRHI::vkGetQueryPoolResults(Device.GetInstanceHandle(), QueryPool, 0, NumUsedQueries, NumUsedQueries * sizeof(uint64), QueryOutput.GetData(), sizeof(uint64), QueryResultFlags);
-		VERIFYVULKANRESULT(Result);
-		State = EState::RT_PostGetResults;
-		return true;
-	}
-
-	return false;
-}
-
-void FVulkanOcclusionQueryPool::Reset(FVulkanCommandBuffer* InCmdBuffer, uint32 InFrameNumber)
-{
-	//ensure(State == EState::Undefined || State == EState::RT_PostGetResults);
-	FMemory::Memzero(AcquiredIndices.GetData(), AcquiredIndices.GetAllocatedSize());
+	CurrentQueryCount = 0;
+	UnusedFrameCount = 0;
 
-	NumUsedQueries = 0;
-	FrameNumber = InFrameNumber;
+	QueryRefs.Empty(MaxQueries);
+	QueryResults.Empty(MaxQueries);
+	QueryRefs.AddDefaulted(MaxQueries);
+	QueryResults.SetNumZeroed(MaxQueries);
 
 	if (Device.GetOptionalExtensions().HasEXTHostQueryReset)
 	{
@@ -131,236 +98,132 @@ void FVulkanOcclusionQueryPool::Reset(FVulkanCommandBuffer* InCmdBuffer, uint32
 	}
 	else
 	{
-		VulkanRHI::vkCmdResetQueryPool(InCmdBuffer->GetHandle(), QueryPool, 0, MaxQueries);
-
-		// workaround for apparent cache-flush bug in AMD driver implementation of vkCmdResetQueryPool
-		if (Device.GetVendorId() == EGpuVendorId::Amd)
-		{
-			VkMemoryBarrier Barrier;
-			ZeroVulkanStruct(Barrier, VK_STRUCTURE_TYPE_MEMORY_BARRIER);
-			Barrier.srcAccessMask = VK_ACCESS_TRANSFER_WRITE_BIT;
-			Barrier.dstAccessMask = VK_ACCESS_HOST_READ_BIT;
-			VulkanRHI::vkCmdPipelineBarrier(InCmdBuffer->GetHandle(), VK_PIPELINE_STAGE_TRANSFER_BIT, VK_PIPELINE_STAGE_HOST_BIT, 0, 1, &Barrier, 0, 0, 0, 0);
-		}
+		VulkanRHI::vkCmdResetQueryPool(InCmdBuffer.GetHandle(), QueryPool, 0, MaxQueries);
 	}
-	
-	State = EState::RHIT_PostBeginBatch;
 }
 
-void FVulkanCommandListContext::BeginOcclusionQueryBatch(FVulkanCommandBuffer* CmdBuffer, uint32 NumQueriesInBatch)
+bool FVulkanQueryPool::IsStale() const
 {
-	checkf(!CurrentOcclusionQueryPool, TEXT("BeginOcclusionQueryBatch called without corresponding EndOcclusionQueryBatch!"));
-	CurrentOcclusionQueryPool = Device.AcquireOcclusionQueryPool(this, NumQueriesInBatch);
-	ensure(CmdBuffer->IsOutsideRenderPass());
-	CurrentOcclusionQueryPool->Reset(CmdBuffer, GFrameNumberRenderThread);
+	return UnusedFrameCount >= GQueryPoolDeletionDelay;
 }
 
-
-inline bool FVulkanOcclusionQueryPool::IsStalePool() const
+void FVulkanCommandListContext::BeginOcclusionQueryBatch(uint32 NumQueriesInBatch)
 {
-	return FrameNumber + NUM_FRAMES_TO_WAIT_REUSE_POOL < GFrameNumberRenderThread;
+	TArray<FVulkanQueryPool*>& OcclusionPoolArray = GetQueryPoolArray(EVulkanQueryPoolType::Occlusion);
+	FVulkanQueryPool* NewOcclusionQueryPool = Device.AcquireOcclusionQueryPool(NumQueriesInBatch);
+	NewOcclusionQueryPool->Reset(GetCommandBuffer());
+	OcclusionPoolArray.Add(NewOcclusionQueryPool);
 }
 
-void FVulkanOcclusionQueryPool::FlushAllocatedQueries()
+FVulkanOcclusionQueryPool* FVulkanDevice::AcquireOcclusionQueryPool(uint32 NumQueries)
 {
-	for (int32 Index = 0; Index < AcquiredIndices.Num(); ++Index)
-	{
-		uint32 QueryIndex = 0;
-		uint64 Acquired = AcquiredIndices[Index];
-		while (Acquired)
-		{
-			if (Acquired & 1)
-			{
-				FVulkanOcclusionQuery* Query =AllocatedQueries[QueryIndex + (Index * 64)];
-				Query->State = Query->State == FVulkanOcclusionQuery::EState::RT_GotResults ? FVulkanOcclusionQuery::EState::FlushedFromPoolHadResults : FVulkanOcclusionQuery::EState::Undefined;
-				Query->Pool = nullptr;
-				Query->IndexInPool = UINT32_MAX;
-
-				AllocatedQueries[QueryIndex + (Index * 64)] = nullptr;
-			}
-			Acquired = Acquired >> (uint64)1;
-			QueryIndex++;
-		}
-
-		AcquiredIndices[Index] = Acquired;
-	}
-}
+	FScopeLock Lock(&QueryPoolLock);
 
-FVulkanOcclusionQueryPool* FVulkanDevice::AcquireOcclusionQueryPool(FVulkanContextCommon* InCommonContext, uint32 NumQueries)
-{
 	// At least add one query
-	NumQueries = FMath::Max(1u, NumQueries);
-	NumQueries = AlignArbitrary(NumQueries, GMinNumberOfQueriesInPool);
+	NumQueries = FMath::Max(1u, AlignArbitrary(NumQueries, 256));
+	OcclusionQueryPoolSize = FMath::Max(OcclusionQueryPoolSize, NumQueries);
 
-	FScopeLock Lock(&OcclusionQueryPoolLock);
+	TArray<FVulkanQueryPool*>& FreeOcclusionPools = FreeQueryPools[(int32)EVulkanQueryPoolType::Occlusion];
 
-	bool bChanged = false;
-	for (int32 Index = UsedOcclusionQueryPools.Num() - 1; Index >= 0; --Index)
+	// Destroy pools that can't accomodate our new minimum size
+	for (int32 Index = FreeOcclusionPools.Num() - 1; Index >= 0; --Index)
 	{
-		FVulkanOcclusionQueryPool* Pool = UsedOcclusionQueryPools[Index];
-		if (Pool->CanBeReused())
-		{
-			UsedOcclusionQueryPools.RemoveAtSwap(Index);
-			FreeOcclusionQueryPools.Add(Pool);
-			Pool->FreedFrameNumber = GFrameNumberRenderThread;
-			bChanged = true;
-		}
-		else if (Pool->IsStalePool())
+		FVulkanQueryPool* Pool = FreeOcclusionPools[Index];
+		checkSlow(Pool && (Pool->GetPoolType() == EVulkanQueryPoolType::Occlusion));
+		if (Pool->GetMaxQueries() < OcclusionQueryPoolSize)
 		{
-			Pool->FlushAllocatedQueries();
-			UsedOcclusionQueryPools.RemoveAtSwap(Index);
-			FreeOcclusionQueryPools.Add(Pool);
-			Pool->FreedFrameNumber = GFrameNumberRenderThread;
-			bChanged = true;
+			delete Pool;
+			FreeOcclusionPools.RemoveAtSwap(Index, EAllowShrinking::No);
 		}
 	}
 
-	if (FreeOcclusionQueryPools.Num() > 0)
+	if (FreeOcclusionPools.Num())
 	{
-		if (bChanged)
-		{
-			FreeOcclusionQueryPools.Sort([](FVulkanOcclusionQueryPool& A, FVulkanOcclusionQueryPool& B)
-			{
-				return A.GetMaxQueries() < B.GetMaxQueries();
-			});
-		}
-
-		for (int32 Index = 0; Index < FreeOcclusionQueryPools.Num(); ++Index)
-		{
-			if (NumQueries <= FreeOcclusionQueryPools[Index]->GetMaxQueries())
-			{
-				FVulkanOcclusionQueryPool* Pool = FreeOcclusionQueryPools[Index];
-				FreeOcclusionQueryPools.RemoveAt(Index);
-				UsedOcclusionQueryPools.Add(Pool);
-				return Pool;
-			}
-		}
+		FVulkanQueryPool* Pool = FreeOcclusionPools.Pop(EAllowShrinking::No);
+		checkSlow(Pool && (Pool->GetPoolType() == EVulkanQueryPoolType::Occlusion));
+		return (FVulkanOcclusionQueryPool*)Pool;
 	}
 
-	FVulkanOcclusionQueryPool* Pool = new FVulkanOcclusionQueryPool(*this, InCommonContext, NumQueries);
-	UsedOcclusionQueryPools.Add(Pool);
+	FVulkanOcclusionQueryPool* Pool = new FVulkanOcclusionQueryPool(*this, OcclusionQueryPoolSize);
 	return Pool;
 }
 
-void FVulkanDevice::ReleaseUnusedOcclusionQueryPools()
+FVulkanTimingQueryPool* FVulkanDevice::AcquireTimingQueryPool()
 {
-	if (GFrameNumberRenderThread < NUM_FRAMES_TO_WAIT_RELEASE_POOL)
-	{
-		return;
-	}
-	
-	const uint32 ReleaseFrame = (GFrameNumberRenderThread - NUM_FRAMES_TO_WAIT_RELEASE_POOL);
+	FScopeLock Lock(&QueryPoolLock);
 
-	FScopeLock Lock(&OcclusionQueryPoolLock);
-	for (int32 Index = FreeOcclusionQueryPools.Num() - 1; Index >= 0; --Index)
+	TArray<FVulkanQueryPool*>& FreeTimingPools = FreeQueryPools[(int32)EVulkanQueryPoolType::Timestamp];
+	if (FreeTimingPools.Num())
 	{
-		FVulkanOcclusionQueryPool* Pool = FreeOcclusionQueryPools[Index];
-		if (ReleaseFrame > Pool->FreedFrameNumber) //-V547
-		{
-			delete Pool;
-			FreeOcclusionQueryPools.RemoveAt(Index);
-		}
+		FVulkanQueryPool* Pool = FreeTimingPools.Pop(EAllowShrinking::No);
+		checkSlow(Pool && (Pool->GetPoolType() == EVulkanQueryPoolType::Timestamp));
+		return (FVulkanTimingQueryPool*)Pool;
 	}
+	return new FVulkanTimingQueryPool(*this, GTimingQueryPoolSize);
 }
 
-void FVulkanCommandListContext::EndOcclusionQueryBatch()
-{
-	checkf(CurrentOcclusionQueryPool, TEXT("EndOcclusionQueryBatch called without corresponding BeginOcclusionQueryBatch!"));
-	CurrentOcclusionQueryPool->EndBatch();
-	CurrentOcclusionQueryPool = nullptr;
-}
-
-
-FVulkanOcclusionQuery::FVulkanOcclusionQuery()
-	: FVulkanRenderQuery(RQT_Occlusion)
+void FVulkanDevice::ReleaseQueryPool(FVulkanQueryPool* Pool)
 {
-	INC_DWORD_STAT(STAT_VulkanNumQueries);
+	FScopeLock Lock(&QueryPoolLock);
+	FreeQueryPools[(int32)Pool->GetPoolType()].Add(Pool);
 }
 
-FVulkanOcclusionQuery::~FVulkanOcclusionQuery()
+void FVulkanDevice::RemoveStaleQueryPools()
 {
-	if (State != EState::Undefined)
+	FScopeLock Lock(&QueryPoolLock);
+	for (TArray<FVulkanQueryPool*>& PoolArray : FreeQueryPools)
 	{
-		if (IndexInPool != UINT32_MAX)
+		for (int32 Index = PoolArray.Num() - 1; Index >= 0; --Index)
 		{
-			ReleaseFromPool();
-		}
-		else
-		{
-			ensure(State == EState::RT_GotResults);
+			FVulkanQueryPool* Pool = PoolArray[Index];
+			checkSlow(Pool);
+			if (Pool->IsStale())
+			{
+				delete Pool;
+				PoolArray.RemoveAtSwap(Index, EAllowShrinking::No);
+			}
+			else
+			{
+				Pool->IncrementUnusedFrameCount();
+			}
 		}
 	}
-
-	DEC_DWORD_STAT(STAT_VulkanNumQueries);
 }
 
-void FVulkanOcclusionQuery::ReleaseFromPool()
+FVulkanRenderQuery::FVulkanRenderQuery(ERenderQueryType InType)
+	: QueryType(InType)
 {
-	Pool->ReleaseIndex(IndexInPool);
-	State = EState::Undefined;
-	IndexInPool = UINT32_MAX;
+	INC_DWORD_STAT(STAT_VulkanNumQueries);
 }
 
-#if (RHI_NEW_GPU_PROFILER == 0)
-void FVulkanCommandListContext::ReadAndCalculateGPUFrameTime()
+FVulkanRenderQuery::~FVulkanRenderQuery()
 {
-	check(IsImmediate());
-
-	if (GSupportsTimestampRenderQueries && FrameTiming)
-	{
-		const uint64 Delta = FrameTiming->GetTiming(false);
-		if (Delta > 0)
-		{
-			const double SecondsPerCycle = FPlatformTime::GetSecondsPerCycle();
-			const double Frequency = double(FVulkanGPUTiming::GetTimingFrequency());
-			GGPUFrameTime = FMath::TruncToInt(double(Delta) / Frequency / SecondsPerCycle);
-		}
-	}
-	else if(!FVulkanPlatform::HasCustomFrameTiming())
-	{
-		GGPUFrameTime = 0;
-	}
+	check(!SyncPoint.IsValid() || SyncPoint->IsComplete());
+	DEC_DWORD_STAT(STAT_VulkanNumQueries);
 }
-#endif // (RHI_NEW_GPU_PROFILER == 0)
-
 
 FRenderQueryRHIRef FVulkanDynamicRHI::RHICreateRenderQuery(ERenderQueryType QueryType)
 {
-	if (QueryType == RQT_Occlusion)
-	{
-		FVulkanRenderQuery* Query = new FVulkanOcclusionQuery();
-		return Query;
-	}
-	else if (QueryType == RQT_AbsoluteTime)
-	{
-		FVulkanTimingQuery* Query = new FVulkanTimingQuery();
-		return Query;
-	}
-	else
-	{
-		// Dummy!
-		ensureMsgf(0, TEXT("Unknown QueryType %d"), QueryType);
-	}
-	FVulkanRenderQuery* Query = new FVulkanRenderQuery(QueryType);
-	return Query;
+	ensureMsgf((QueryType == RQT_Occlusion) || (QueryType == RQT_AbsoluteTime), TEXT("Unknown QueryType %d"), QueryType);
+	return new FVulkanRenderQuery(QueryType);
 }
 
 bool FVulkanDynamicRHI::RHIGetRenderQueryResult(FRHIRenderQuery* QueryRHI, uint64& OutQueryResult, bool bWait, uint32 GPUIndex)
 {
-	FVulkanRenderQuery* BaseQuery = ResourceCast(QueryRHI);
+	FVulkanRenderQuery* Query = ResourceCast(QueryRHI);
 
-	if (!ensureMsgf(BaseQuery->SyncPoint, TEXT("Attempt to get result data for an FRHIRenderQuery that was never used in a command list.")))
+	if (!ensureMsgf(Query->SyncPoint, TEXT("Attempt to get result data for an FRHIRenderQuery that was never used in a command list.")))
 	{
 		OutQueryResult = 0;
 		return false;
 	}
 
-	if (!BaseQuery->SyncPoint->IsComplete())
+	if (!Query->SyncPoint->IsComplete())
 	{
 		if (bWait)
 		{
 			FRenderThreadIdleScope IdleScope(ERenderThreadIdleTypes::WaitingForGPUQuery);
-			ProcessInterruptQueueUntil(BaseQuery->SyncPoint);
+			ProcessInterruptQueueUntil(Query->SyncPoint);
 		}
 		else
 		{
@@ -368,97 +231,19 @@ bool FVulkanDynamicRHI::RHIGetRenderQueryResult(FRHIRenderQuery* QueryRHI, uint6
 		}
 	}
 
-	if (BaseQuery->QueryType == RQT_Occlusion)
-	{
-		FVulkanOcclusionQuery* Query = static_cast<FVulkanOcclusionQuery*>(BaseQuery);
-		if (Query->State == FVulkanOcclusionQuery::EState::RT_GotResults || Query->State == FVulkanOcclusionQuery::EState::FlushedFromPoolHadResults)
-		{
-			OutQueryResult = Query->Result;
-			return true;
-		}
+	checkSlow(Query->SyncPoint->IsComplete());
 
-		if (Query->State == FVulkanOcclusionQuery::EState::Undefined)
-		{
-			UE_LOG(LogVulkanRHI, Verbose, TEXT("Stale query asking for result!"));
-			return false;
-		}
-
-		ensure(Query->State == FVulkanOcclusionQuery::EState::RHI_PostEnd);
-		if (Query->Pool->GetResults())
-		{
-			Query->Result = Query->Pool->GetResultValue(Query->IndexInPool);
-			Query->ReleaseFromPool();
-			Query->State = FVulkanOcclusionQuery::EState::RT_GotResults;
-			OutQueryResult = Query->Result;
-			return true;
-		}
+	if (Query->QueryType == RQT_Occlusion)
+	{
+		OutQueryResult = Query->Result;
+		return true;
 	}
-	else if (BaseQuery->QueryType == RQT_AbsoluteTime)
+	else if (Query->QueryType == RQT_AbsoluteTime)
 	{
-		auto ToMicroseconds = [](uint64 Timestamp) -> uint64
-		{
-#if RHI_NEW_GPU_PROFILER
-			checkNoEntry(); // @todo - new gpu profiler
-			return 0;
-#else
-			const double Frequency = double(FVulkanGPUTiming::GetTimingFrequency());
-			uint64 Microseconds = (uint64)((double(Timestamp) / Frequency) * 1000.0 * 1000.0);
-			return Microseconds;
-#endif
-		};
-
-		FVulkanTimingQuery* Query = static_cast<FVulkanTimingQuery*>(BaseQuery);
-		if ((Query->Pool == nullptr) && bWait)
-		{
-			// Blocking was requested, so, we need to wait for the RHI thread to catch up.
-			FRHICommandListImmediate::Get().ImmediateFlush(EImmediateFlushType::FlushRHIThread);
-			check(Query->Pool != nullptr);
-		}
-
-		check(Query->Pool && (Query->Pool->CurrentTimestampIndex < Query->Pool->BufferedCount));
-		int32 TimestampIndex = Query->Pool->CurrentTimestampIndex;
-		if (!bWait)
-		{
-			// Quickly check the most recent measurements to see if any of them has been resolved.  Do not flush these queries.
-			for (uint32 IssueIndex = 1; IssueIndex < Query->Pool->NumIssuedTimestamps; ++IssueIndex)
-			{
-				const FVulkanTimingQueryPool::FTimestampData& QueryData = Query->Pool->TimestampArray[TimestampIndex];
-				if (QueryData.SyncPoint->IsComplete())
-				{
-					OutQueryResult = ToMicroseconds(QueryData.Result.Value);
-					return true;
-				}
-
-				TimestampIndex = (TimestampIndex + Query->Pool->BufferedCount - 1) % Query->Pool->BufferedCount;
-			}
-		}
-
-		if (Query->Pool->NumIssuedTimestamps > 0 || bWait)
-		{
-			// None of the (NumIssuedTimestamps - 1) measurements were ready yet,
-			// so check the oldest measurement more thoroughly.
-			// This really only happens if occlusion and frame sync event queries are disabled, otherwise those will block until the GPU catches up to 1 frame behind
-
-			const FVulkanTimingQueryPool::FTimestampData& QueryData = Query->Pool->TimestampArray[TimestampIndex];
-
-			{
-				FRenderThreadIdleScope IdleScope(ERenderThreadIdleTypes::WaitingForGPUQuery);
-				SCOPE_CYCLE_COUNTER(STAT_RenderQueryResultTime);
-
-				const bool bBlocking = (Query->Pool->NumIssuedTimestamps == Query->Pool->BufferedCount) || bWait;
-				if (bBlocking)
-				{
-					// CPU wait for query results to be ready.
-					if (!QueryData.SyncPoint->IsComplete())
-					{
-						FVulkanDynamicRHI::Get().ProcessInterruptQueueUntil(QueryData.SyncPoint);
-					}
-				}
-			}
-
-			OutQueryResult = ToMicroseconds(QueryData.Result.Value);
-			return true;
-		}
+		const VkPhysicalDeviceLimits& Limits = Device->GetDeviceProperties().limits;
+		const double TimingFrequency = (double)((uint64)((1000.0 * 1000.0 * 1000.0) / Limits.timestampPeriod));
+		OutQueryResult = (uint64)((double(Query->Result) / TimingFrequency) * 1000.0 * 1000.0);
+		return true;
 	}
 
 	return false;
@@ -521,7 +306,7 @@ void FVulkanDynamicRHI::RHIEndRenderQueryBatch_TopOfPipe(FRHICommandListBase& RH
 	SyncPoint->Release();
 	QueryBatchData[GPUIndex] = nullptr;
 
-	RHICmdList.EnqueueLambda([SyncPoint = MoveTemp(SyncPoint), QueryType](FRHICommandListBase& ExecutingCmdList)
+	RHICmdList.EnqueueLambda([SyncPoint = MoveTemp(SyncPoint)](FRHICommandListBase& ExecutingCmdList)
 	{
 		FVulkanCommandListContext& Context = FVulkanCommandListContext::Get(ExecutingCmdList);
 		Context.AddPendingSyncPoint(SyncPoint);
@@ -530,31 +315,14 @@ void FVulkanDynamicRHI::RHIEndRenderQueryBatch_TopOfPipe(FRHICommandListBase& RH
 
 void FVulkanCommandListContext::RHIBeginRenderQuery(FRHIRenderQuery* QueryRHI)
 {
-	FVulkanRenderQuery* BaseQuery = ResourceCast(QueryRHI);
-	if (BaseQuery->QueryType == RQT_Occlusion)
+	FVulkanRenderQuery* Query = ResourceCast(QueryRHI);
+	if (Query->QueryType == RQT_Occlusion)
 	{
-		check(CurrentOcclusionQueryPool);
-		FVulkanOcclusionQuery* Query = static_cast<FVulkanOcclusionQuery*>(BaseQuery);
-		if (Query->State == FVulkanOcclusionQuery::EState::RHI_PostEnd)
-		{
-			Query->ReleaseFromPool();
-		}
-		else if (Query->State == FVulkanOcclusionQuery::EState::RT_GotResults)
-		{
-			// Nothing to do here...
-		}
-		else
-		{
-			ensure(Query->State == FVulkanOcclusionQuery::EState::Undefined || Query->State == FVulkanOcclusionQuery::EState::FlushedFromPoolHadResults);
-		}
-		Query->State = FVulkanOcclusionQuery::EState::RHI_PostBegin;
-		const uint32 IndexInPool = CurrentOcclusionQueryPool->AcquireIndex(Query);
-		Query->Pool = CurrentOcclusionQueryPool;
-		Query->IndexInPool = IndexInPool;
-		FVulkanCommandBuffer& CommandBuffer = GetCommandBuffer();
-		VulkanRHI::vkCmdBeginQuery(CommandBuffer.GetHandle(), CurrentOcclusionQueryPool->GetHandle(), IndexInPool, VK_QUERY_CONTROL_PRECISE_BIT);
+		FVulkanOcclusionQueryPool* CurrentOcclusionQueryPool = GetCurrentOcclusionQueryPool();
+		CurrentOcclusionQueryPool->ReserveQuery(Query);
+		VulkanRHI::vkCmdBeginQuery(GetCommandBuffer().GetHandle(), CurrentOcclusionQueryPool->GetHandle(), Query->IndexInPool, VK_QUERY_CONTROL_PRECISE_BIT);
 	}
-	else if (BaseQuery->QueryType == RQT_AbsoluteTime)
+	else if (Query->QueryType == RQT_AbsoluteTime)
 	{
 		ensureMsgf(0, TEXT("Timing queries should NOT call RHIBeginRenderQuery()!"));
 	}
@@ -562,35 +330,20 @@ void FVulkanCommandListContext::RHIBeginRenderQuery(FRHIRenderQuery* QueryRHI)
 
 void FVulkanCommandListContext::RHIEndRenderQuery(FRHIRenderQuery* QueryRHI)
 {
-	FVulkanRenderQuery* BaseQuery = ResourceCast(QueryRHI);
-	if (BaseQuery->QueryType == RQT_Occlusion)
+	FVulkanRenderQuery* Query = ResourceCast(QueryRHI);
+	check(Query->SyncPoint.IsValid());
+	if (Query->QueryType == RQT_Occlusion)
 	{
-		check(CurrentOcclusionQueryPool);
-		FVulkanOcclusionQuery* Query = static_cast<FVulkanOcclusionQuery*>(BaseQuery);
-		ensure(Query->State == FVulkanOcclusionQuery::EState::RHI_PostBegin);
-		Query->State = FVulkanOcclusionQuery::EState::RHI_PostEnd;
-		VulkanRHI::vkCmdEndQuery(GetCommandBuffer().GetHandle(), CurrentOcclusionQueryPool->GetHandle(), Query->IndexInPool);
+		FVulkanOcclusionQueryPool* CurrentPool = GetCurrentOcclusionQueryPool();
+		VulkanRHI::vkCmdEndQuery(GetCommandBuffer().GetHandle(), CurrentPool->GetHandle(), Query->IndexInPool);
 	}
-	else if (BaseQuery->QueryType == RQT_AbsoluteTime)
+	else if (Query->QueryType == RQT_AbsoluteTime)
 	{
-		FVulkanTimingQuery* Query = static_cast<FVulkanTimingQuery*>(BaseQuery);
-
-		if (Query->Pool == nullptr)
-		{
-			Query->Pool = new FVulkanTimingQueryPool(Device, this, 4);
-		}
-
-		Query->Pool->CurrentTimestampIndex = (Query->Pool->CurrentTimestampIndex + 1) % Query->Pool->BufferedCount;
-		const uint32 QueryEndIndex = Query->Pool->CurrentTimestampIndex;
-		FVulkanTimingQueryPool::FTimestampData& QueryData = Query->Pool->TimestampArray[QueryEndIndex];
+		FVulkanTimingQueryPool* CurrentPool = GetCurrentTimestampQueryPool();
+		CurrentPool->ReserveQuery(Query);
 
 		const VkPipelineStageFlagBits QueryPipelineStage = GTimestampQueryStage ? VK_PIPELINE_STAGE_BOTTOM_OF_PIPE_BIT : VK_PIPELINE_STAGE_TOP_OF_PIPE_BIT;
-		FVulkanCommandBuffer& CommandBuffer = GetCommandBuffer();
-		VulkanRHI::vkCmdWriteTimestamp(CommandBuffer.GetHandle(), QueryPipelineStage, Query->Pool->GetHandle(), QueryEndIndex);
-		Query->Pool->TimestampArray[QueryEndIndex].SyncPoint = GetContextSyncPoint();
-		Query->Pool->NumIssuedTimestamps = FMath::Min<uint32>(Query->Pool->NumIssuedTimestamps + 1, Query->Pool->BufferedCount);
-
-		CommandBuffer.AddPendingTimestampQuery(QueryEndIndex, 1, Query->Pool->GetHandle(), &QueryData.Result, true);
+		VulkanRHI::vkCmdWriteTimestamp(GetCommandBuffer().GetHandle(), QueryPipelineStage, CurrentPool->GetHandle(), Query->IndexInPool);
 	}
 }
 
@@ -606,30 +359,27 @@ void FVulkanCommandListContext::RHICalibrateTimers(FRHITimestampCalibrationQuery
 }
 #endif
 
-void FVulkanCommandListContext::WriteBeginTimestamp(FVulkanCommandBuffer* CmdBuffer)
+FVulkanOcclusionQueryPool* FVulkanCommandListContext::GetCurrentOcclusionQueryPool()
 {
-#if (RHI_NEW_GPU_PROFILER == 0)
-	FrameTiming->StartTiming(CmdBuffer);
-#endif
+	TArray<FVulkanQueryPool*>& OcclusionPoolArray = GetQueryPoolArray(EVulkanQueryPoolType::Occlusion);
+	checkSlow(OcclusionPoolArray.Num() && !OcclusionPoolArray.Last()->IsFull());
+	return (FVulkanOcclusionQueryPool*)OcclusionPoolArray.Last();
 }
 
-void FVulkanCommandListContext::WriteEndTimestamp(FVulkanCommandBuffer* CmdBuffer)
+FVulkanTimingQueryPool* FVulkanContextCommon::GetCurrentTimestampQueryPool(FVulkanPayload& Payload)
 {
-#if (RHI_NEW_GPU_PROFILER == 0)
-	FrameTiming->EndTiming(CmdBuffer);
-#endif
-}
+	TArray<FVulkanQueryPool*>& TimestampPoolArray = Payload.QueryPools[(int32)EVulkanQueryPoolType::Timestamp];
+	if ((TimestampPoolArray.Num() == 0) || TimestampPoolArray.Last()->IsFull())
+	{
+		FVulkanTimingQueryPool* NewPool = Device.AcquireTimingQueryPool();
 
-FVulkanTimingQuery::FVulkanTimingQuery()
-	: FVulkanRenderQuery(RQT_AbsoluteTime)
-{
-}
+		if (Payload.CommandBuffers.Num() == 0)
+		{
+			PrepareNewCommandBuffer(Payload);
+		}
+		NewPool->Reset(*Payload.CommandBuffers.Last());
 
-FVulkanTimingQuery::~FVulkanTimingQuery()
-{
-	if (Pool != nullptr)
-	{
-		delete Pool;
-		Pool = nullptr;
+		TimestampPoolArray.Add(NewPool);
 	}
+	return (FVulkanTimingQueryPool*)TimestampPoolArray.Last();
 }
diff --git a/Engine/Source/Runtime/VulkanRHI/Private/VulkanQuery.h b/Engine/Source/Runtime/VulkanRHI/Private/VulkanQuery.h
--- a/Engine/Source/Runtime/VulkanRHI/Private/VulkanQuery.h
+++ b/Engine/Source/Runtime/VulkanRHI/Private/VulkanQuery.h
@@ -12,15 +12,43 @@ class FVulkanDevice;
 class FVulkanQueue;
 class FVulkanContextCommon;
 class FVulkanCommandBuffer;
-class FVulkanOcclusionQuery;
+class FVulkanRenderQuery;
+class FVulkanDynamicRHI;
 
 using FVulkanSyncPoint = FGraphEvent;
 using FVulkanSyncPointRef = TRefCountPtr<FVulkanSyncPoint>;
 
+
+enum class EVulkanQueryPoolType : uint8
+{
+	Occlusion,
+	PipelineStats,
+	Timestamp,
+	ASCompactedSize,
+
+	Count
+};
+
+constexpr VkQueryType GetVkQueryType(EVulkanQueryPoolType Type)
+{
+	switch (Type)
+	{
+	case EVulkanQueryPoolType::Occlusion:       return VK_QUERY_TYPE_OCCLUSION;
+	case EVulkanQueryPoolType::PipelineStats:   return VK_QUERY_TYPE_PIPELINE_STATISTICS;
+	case EVulkanQueryPoolType::Timestamp:       return VK_QUERY_TYPE_TIMESTAMP;
+	case EVulkanQueryPoolType::ASCompactedSize: return VK_QUERY_TYPE_ACCELERATION_STRUCTURE_COMPACTED_SIZE_KHR;
+
+	case EVulkanQueryPoolType::Count:
+	default:
+		checkNoEntry();
+		return VK_QUERY_TYPE_MAX_ENUM;
+	}
+}
+
 class FVulkanQueryPool
 {
 public:
-	FVulkanQueryPool(FVulkanDevice& InDevice, FVulkanContextCommon* InCommonContext, uint32 InMaxQueries, VkQueryType InQueryType, bool bInShouldAddReset = true);
+	FVulkanQueryPool(FVulkanDevice& InDevice, uint32 InMaxQueries, EVulkanQueryPoolType InQueryType);
 	virtual ~FVulkanQueryPool();
 
 	uint32 GetMaxQueries() const
@@ -33,9 +61,9 @@ public:
 		return QueryPool;
 	}
 
-	uint64 GetResultValue(uint32 Index) const
+	EVulkanQueryPoolType GetPoolType() const
 	{
-		return QueryOutput[Index];
+		return QueryType;
 	}
 
 	FVulkanDevice& GetDevice()
@@ -43,162 +71,74 @@ public:
 		return Device;
 	}
 
-protected:
-	FVulkanDevice& Device;
-	VkQueryPool QueryPool;
-	const uint32 MaxQueries;
-	const VkQueryType QueryType;
-	TArray<uint64> QueryOutput;
-};
-
-class FVulkanOcclusionQueryPool : public FVulkanQueryPool
-{
-public:
-	FVulkanOcclusionQueryPool(FVulkanDevice& InDevice, FVulkanContextCommon* InContext, uint32 InMaxQueries)
-		: FVulkanQueryPool(InDevice, InContext, InMaxQueries, VK_QUERY_TYPE_OCCLUSION)
+	bool IsFull() const
 	{
-		AcquiredIndices.AddZeroed(Align(InMaxQueries, 64) / 64);
-		AllocatedQueries.AddZeroed(InMaxQueries);
+		return (CurrentQueryCount >= MaxQueries);
 	}
 
-	inline uint32 AcquireIndex(FVulkanOcclusionQuery* Query)
-	{
-		check(NumUsedQueries < MaxQueries);
-		const uint32 Index = NumUsedQueries;
-		const uint32 Word = Index / 64;
-		const uint32 Bit = Index % 64;
-		const uint64 Mask = (uint64)1 << (uint64)Bit;
-		const uint64& WordValue = AcquiredIndices[Word];
-		AcquiredIndices[Word] = WordValue | Mask;
-		++NumUsedQueries;
-		ensure(AllocatedQueries[Index] == nullptr);
-		AllocatedQueries[Index] = Query;
-		return Index;
-	}
+	void ReserveQuery(FVulkanRenderQuery* Query);
+	uint32 ReserveQuery(uint64* ResultPtr);
 
-	inline void ReleaseIndex(uint32 Index)
-	{
-		check(Index < NumUsedQueries);
-		const uint32 Word = Index / 64;
-		const uint32 Bit = Index % 64;
-		const uint64 Mask = (uint64)1 << (uint64)Bit;
-		const uint64& WordValue = AcquiredIndices[Word];
-		ensure((WordValue & Mask) == Mask);
-		AcquiredIndices[Word] = WordValue & (~Mask);
-		AllocatedQueries[Index] = nullptr;
-	}
+	void Reset(FVulkanCommandBuffer& InCmdBuffer);
+	bool IsStale() const;
 
-	void EndBatch()
+	void IncrementUnusedFrameCount()
 	{
-		ensure(State == EState::RHIT_PostBeginBatch);
-		State = EState::RHIT_PostEndBatch;
+		UnusedFrameCount++;
 	}
 
-	bool CanBeReused();
-
-	bool GetResults();
-
-	void Reset(FVulkanCommandBuffer* InCmdBuffer, uint32 InFrameNumber);
-
-	bool IsStalePool() const;
-
-	void FlushAllocatedQueries();
-
-	enum EState
-	{
-		Undefined,
-		RHIT_PostBeginBatch,
-		RHIT_PostEndBatch,
-		RT_PostGetResults,
-	};
-	EState State = Undefined;
-
-	// frame number when pool was placed into free list
-	uint32 FreedFrameNumber = UINT32_MAX;
 protected:
-	uint32 NumUsedQueries = 0;
-	TArray<FVulkanOcclusionQuery*> AllocatedQueries;
-	TArray<uint64> AcquiredIndices;
+	FVulkanDevice& Device;
+	VkQueryPool QueryPool;
+	const uint32 MaxQueries;
+	const EVulkanQueryPoolType QueryType;
+	TArray<TRefCountPtr<FVulkanRenderQuery>> QueryRefs;
+	TArray<uint64*> QueryResults;
+	uint32 CurrentQueryCount = 0;
+	int32 UnusedFrameCount = 0;
 
-	uint32 FrameNumber = UINT32_MAX;
+	friend FVulkanDynamicRHI;
 };
 
-struct FVulkanTimingQueryResult
+class FVulkanOcclusionQueryPool : public FVulkanQueryPool
 {
-	uint64 Value = 0;
-	uint64 Availability = 0;
+public:
+	FVulkanOcclusionQueryPool(FVulkanDevice& InDevice, uint32 InMaxQueries)
+		: FVulkanQueryPool(InDevice, InMaxQueries, EVulkanQueryPoolType::Occlusion)
+	{
+	}
 };
 
 class FVulkanTimingQueryPool : public FVulkanQueryPool
 {
 public:
-	FVulkanTimingQueryPool(FVulkanDevice& InDevice, FVulkanContextCommon* InContext, uint32 InBufferedCount)
-		: FVulkanQueryPool(InDevice, InContext, InBufferedCount * 2, VK_QUERY_TYPE_TIMESTAMP, false)
-		, BufferedCount(InBufferedCount)
+	FVulkanTimingQueryPool(FVulkanDevice& InDevice, uint32 InBufferedCount)
+		: FVulkanQueryPool(InDevice, InBufferedCount * 2, EVulkanQueryPoolType::Timestamp)
 	{
-		TimestampArray.AddDefaulted(GetMaxQueries());
 	}
-
-	virtual ~FVulkanTimingQueryPool();
-
-	uint32 CurrentTimestampIndex = 0;
-	uint32 NumIssuedTimestamps = 0;
-	const uint32 BufferedCount;
-
-	struct FTimestampData
-	{
-		FVulkanSyncPointRef SyncPoint;
-		FVulkanTimingQueryResult Result;
-	};
-	TArray<FTimestampData> TimestampArray;
 };
 
 class FVulkanRenderQuery : public FRHIRenderQuery
 {
 public:
-	FVulkanRenderQuery(ERenderQueryType InType)
-		: QueryType(InType)
+	FVulkanRenderQuery(ERenderQueryType InType);
+	virtual ~FVulkanRenderQuery();
+
+	EVulkanQueryPoolType GetQueryPoolType()
 	{
+		if (QueryType == RQT_Occlusion)
+		{
+			return EVulkanQueryPoolType::Occlusion;
+		}
+		else 
+		{
+			check(QueryType == RQT_AbsoluteTime);
+			return EVulkanQueryPoolType::Timestamp;
+		}
 	}
 
-	virtual ~FVulkanRenderQuery() {}
-
 	const ERenderQueryType QueryType;
 	uint64 Result = 0;
-
 	uint32 IndexInPool = UINT32_MAX;
-
 	FVulkanSyncPointRef SyncPoint;
 };
-
-class FVulkanOcclusionQuery : public FVulkanRenderQuery
-{
-public:
-	FVulkanOcclusionQuery();
-	virtual ~FVulkanOcclusionQuery();
-
-	enum class EState
-	{
-		Undefined,
-		RHI_PostBegin,
-		RHI_PostEnd,
-		RT_GotResults,
-		FlushedFromPoolHadResults,
-	};
-
-	FVulkanOcclusionQueryPool* Pool = nullptr;
-
-	void ReleaseFromPool();
-
-	EState State = EState::Undefined;
-};
-
-class FVulkanTimingQuery : public FVulkanRenderQuery
-{
-public:
-	FVulkanTimingQuery();
-	virtual ~FVulkanTimingQuery();
-
-	FVulkanTimingQueryPool* Pool = nullptr;
-};
-
diff --git a/Engine/Source/Runtime/VulkanRHI/Private/VulkanQueue.cpp b/Engine/Source/Runtime/VulkanRHI/Private/VulkanQueue.cpp
--- a/Engine/Source/Runtime/VulkanRHI/Private/VulkanQueue.cpp
+++ b/Engine/Source/Runtime/VulkanRHI/Private/VulkanQueue.cpp
@@ -152,6 +152,8 @@ void FVulkanQueue::SubmitPayloads(TArrayView<FVulkanPayload*> Payloads, TSet<VkS
 	CommandBufferStorage.Reserve(NumCommandBuffers);
 	SubmitInfos.Reserve(NumPayloads);
 
+	const uint64 SubmitTime = FPlatformTime::Cycles64();
+
 	for (FVulkanPayload* Payload : Payloads)
 	{
 		VkSubmitInfo& SubmitInfo = SubmitInfos.AddZeroed_GetRef();
@@ -188,6 +190,10 @@ void FVulkanQueue::SubmitPayloads(TArrayView<FVulkanPayload*> Payloads, TSet<VkS
 
 				CommandBuffer->State = FVulkanCommandBuffer::EState::Submitted;
 				CommandBuffer->SubmittedTime = FPlatformTime::Seconds();
+
+#if RHI_NEW_GPU_PROFILER
+				CommandBuffer->FlushProfilerEvents(Payload->EventStream, SubmitTime);
+#endif // RHI_NEW_GPU_PROFILER
 			}
 			SubmitInfo.commandBufferCount = Payload->CommandBuffers.Num();
 			SubmitInfo.pCommandBuffers = &CommandBufferStorage[FirstIndex];
@@ -206,6 +212,23 @@ void FVulkanQueue::SubmitPayloads(TArrayView<FVulkanPayload*> Payloads, TSet<VkS
 			SubmitInfo.pSignalSemaphores = &SemaphoreStorage[FirstSempahoreIndex];
 		}
 
+#if RHI_NEW_GPU_PROFILER
+		if (Payload->EndFrameEvent.IsSet())
+		{
+			Payload->EndFrameEvent->CPUTimestamp = SubmitTime;
+			Payload->EventStream.Emplace<UE::RHI::GPUProfiler::FEvent::FFrameBoundary>(*Payload->EndFrameEvent);
+		}
+
+		if (Payload->Timing.IsSet())
+		{
+			if (FVulkanTiming* LocalTiming = *Payload->Timing)
+			{
+				SCOPED_NAMED_EVENT(CalibrateClocks, FColor::Red);
+				Device.GetCalibrationTimestamp(*LocalTiming);
+			}
+		}
+#endif // RHI_NEW_GPU_PROFILER
+
 		if (bUseTimelineSemaphores)
 		{
 			SemaphoreStorage.Add(TimelineSempahore->GetHandle());
@@ -329,19 +352,8 @@ int32 FVulkanQueue::ProcessInterruptQueue(uint64 Timeout)
 			++NumPayloads;
 			PendingInterrupt.Pop();
 
-			for (FVulkanCommandBuffer* CommandBuffer : Payload->CommandBuffers)
-			{
-				CommandBuffer->ResolvePendingTimestampQueries();
-				CommandBuffer->ResetFenceStatus();
-			}
-
-			for (FVulkanSyncPointRef& SyncPoint : Payload->SyncPoints)
-			{
-				SyncPoint->DispatchSubsequents();
-			}
-			Payload->SyncPoints.Empty();
-
-			delete Payload;
+			// Resolve any pending actions and delete the payload
+			FVulkanDynamicRHI::Get().CompletePayload(Payload);
 		}
 
 	} while (bSuccess);
@@ -415,3 +427,22 @@ void FVulkanQueue::FillSupportedStageBits()
 		SupportedStages |= VK_PIPELINE_STAGE_TRANSFER_BIT;
 	}
 }
+
+#if RHI_NEW_GPU_PROFILER
+UE::RHI::GPUProfiler::FQueue FVulkanQueue::GetProfilerQueue() const
+{
+	UE::RHI::GPUProfiler::FQueue ProfilerQueue;
+	ProfilerQueue.GPU = 0;
+	ProfilerQueue.Index = 0;
+
+	switch (QueueType)
+	{
+	default: checkNoEntry(); [[fallthrough]];
+	case EVulkanQueueType::Graphics:     ProfilerQueue.Type = UE::RHI::GPUProfiler::FQueue::EType::Graphics; break;
+	case EVulkanQueueType::AsyncCompute: ProfilerQueue.Type = UE::RHI::GPUProfiler::FQueue::EType::Compute;  break;
+	case EVulkanQueueType::Transfer:     ProfilerQueue.Type = UE::RHI::GPUProfiler::FQueue::EType::Copy;     break;
+	}
+
+	return ProfilerQueue;
+}
+#endif // RHI_NEW_GPU_PROFILER
\ No newline at end of file
diff --git a/Engine/Source/Runtime/VulkanRHI/Private/VulkanQueue.h b/Engine/Source/Runtime/VulkanRHI/Private/VulkanQueue.h
--- a/Engine/Source/Runtime/VulkanRHI/Private/VulkanQueue.h
+++ b/Engine/Source/Runtime/VulkanRHI/Private/VulkanQueue.h
@@ -12,6 +12,7 @@
 class FVulkanDynamicRHI;
 class FVulkanDevice;
 class FVulkanPayload;
+class FVulkanTiming;
 
 namespace VulkanRHI
 {
@@ -77,6 +78,10 @@ public:
 	int32 SubmitQueuedPayloads(TSet<VkSemaphore>& SignaledSemas);
 	int32 ProcessInterruptQueue(uint64 Timeout);
 
+#if RHI_NEW_GPU_PROFILER
+	UE::RHI::GPUProfiler::FQueue GetProfilerQueue() const;
+#endif
+
 private:
 	void SubmitPayloads(TArrayView<FVulkanPayload*> Payloads, TSet<VkSemaphore>& InOutSignaledSemas);
 	void Submit(TArrayView<VkSubmitInfo> InSubmitInfos, VulkanRHI::FFence* Fence);
@@ -108,5 +113,10 @@ private:
 	TQueue<FVulkanPayload*, EQueueMode::Mpsc> PendingSubmission;
 	TQueue<FVulkanPayload*, EQueueMode::Spsc> PendingInterrupt;
 
+#if RHI_NEW_GPU_PROFILER
+	// The active timing struct on this queue. Updated / accessed by the interrupt thread.
+	FVulkanTiming* Timing = nullptr;
+#endif
+
 	void FillSupportedStageBits();
 };
diff --git a/Engine/Source/Runtime/VulkanRHI/Private/VulkanRHI.cpp b/Engine/Source/Runtime/VulkanRHI/Private/VulkanRHI.cpp
--- a/Engine/Source/Runtime/VulkanRHI/Private/VulkanRHI.cpp
+++ b/Engine/Source/Runtime/VulkanRHI/Private/VulkanRHI.cpp
@@ -815,6 +815,8 @@ void FVulkanDynamicRHI::SelectDevice()
 	}
 
 	GRHIPersistentThreadGroupCount = 1440; // TODO: Revisit based on vendor/adapter/perf query
+
+	GRHIGlobals.SupportsTimestampRenderQueries = FVulkanPlatform::SupportsTimestampRenderQueries() && (Device->GetLimits().timestampPeriod > 0.0f);
 }
 
 void FVulkanDynamicRHI::InitInstance()
@@ -849,7 +851,6 @@ void FVulkanDynamicRHI::InitInstance()
 		GSupportsDepthBoundsTest = Device->GetPhysicalDeviceFeatures().Core_1_0.depthBounds != 0;
 		GSupportsRenderTargetFormat_PF_G8 = false;	// #todo-rco
 		GRHISupportsTextureStreaming = true;
-		GSupportsTimestampRenderQueries = FVulkanPlatform::SupportsTimestampRenderQueries();
 		GRHISupportsGPUTimestampBubblesRemoval = true;
 		GSupportsMobileMultiView = Device->GetOptionalExtensions().HasKHRMultiview ? true : false;
 		GRHISupportsMSAAShaderResolve = Device->GetOptionalExtensions().HasQcomRenderPassShaderResolve ? true : false;
@@ -1082,6 +1083,8 @@ void FVulkanDynamicRHI::InitInstance()
 
 void FVulkanDynamicRHI::RHIEndFrame_RenderThread(FRHICommandListImmediate& RHICmdList)
 {
+	// @todo dev-pr - refactor RHIEndFrame_RenderThread to reduce use of the immediate command list, and move cleanup work to RHIEndFrame() below.
+
 	RHICmdList.EnqueueLambdaMultiPipe(ERHIPipeline::Graphics, FRHICommandListBase::EThreadFence::Enabled, TEXT("Vulkan EndFrame"),
 		[this](FVulkanContextArray const& Contexts)
 	{
@@ -1090,7 +1093,6 @@ void FVulkanDynamicRHI::RHIEndFrame_RenderThread(FRHICommandListImmediate& RHICm
 		check(Context.IsImmediate());
 
 #if (RHI_NEW_GPU_PROFILER == 0)
-		Context.ReadAndCalculateGPUFrameTime();
 		Context.GpuProfiler.EndFrame();
 #endif
 
@@ -1107,7 +1109,7 @@ void FVulkanDynamicRHI::RHIEndFrame_RenderThread(FRHICommandListImmediate& RHICm
 		}
 		Context.Device.GetDescriptorPoolsManager().GC();
 
-		Context.Device.ReleaseUnusedOcclusionQueryPools();
+		Context.Device.RemoveStaleQueryPools();
 
 		Context.Device.GetPipelineStateCache()->TickLRU();
 
@@ -1142,8 +1144,70 @@ void FVulkanDynamicRHI::RHIEndFrame_RenderThread(FRHICommandListImmediate& RHICm
 
 void FVulkanDynamicRHI::RHIEndFrame(const FRHIEndFrameArgs& Args)
 {
-	// @todo dev-pr - refactor RHIEndFrame_RenderThread to reduce use of the immediate command list, and move cleanup work to here.
+#if RHI_NEW_GPU_PROFILER
+	// Close the previous frame's timing and start a new one
+	auto Lambda = [this, OldTiming = MoveTemp(CurrentTimingPerQueue)]()
+	{
+		TArray<UE::RHI::GPUProfiler::FEventStream, TInlineAllocator<(int32)EVulkanQueueType::Count>> Streams;
+		for (auto const& Timing : OldTiming)
+		{
+			Streams.Add(MoveTemp(Timing->EventStream));
+		}
+		UE::RHI::GPUProfiler::ProcessEvents(Streams);
+	};
+
+	EnqueueEndOfPipeTask(MoveTemp(Lambda), [&](FVulkanPayload& Payload)
+		{
+			// Modify the payloads the EOP task will submit to include
+			// a new timing struct and a frame boundary event.
+
+			Payload.Timing = CurrentTimingPerQueue.CreateNew(Payload.Queue);
+
+			ERHIPipeline Pipeline;
+			switch (Payload.Queue.QueueType)
+			{
+			default: checkNoEntry(); [[fallthrough]];
+			case EVulkanQueueType::Graphics:     Pipeline = ERHIPipeline::Graphics; break;
+			case EVulkanQueueType::AsyncCompute: Pipeline = ERHIPipeline::AsyncCompute; break;
+
+			case EVulkanQueueType::Transfer:
+				// There is currently no high level RHI copy queue support
+				Pipeline = ERHIPipeline::None;
+				break;
+			}
+
+			// CPU timestamp for the frame boundary event is filled in by the submission thread
+			Payload.EndFrameEvent = UE::RHI::GPUProfiler::FEvent::FFrameBoundary(0, Args.FrameNumber
+#if WITH_RHI_BREADCRUMBS
+				, (Pipeline != ERHIPipeline::None) ? Args.GPUBreadcrumbs[Pipeline] : nullptr
+#endif
+#if STATS
+				, Args.StatsFrame
+#endif
+			);
+		});
+#else
+	{
+		FVulkanPlatformCommandList* Payloads = new FVulkanPlatformCommandList;
+		FVulkanPayload* Payload = new FVulkanPayload(*Device->GetGraphicsQueue());
+		Payload->bEndFrame = true;
+		Payloads->Add(Payload);
+		PendingPayloadsForSubmission.Enqueue(Payloads);
+	}
+#endif
+
+	// Pump the interrupt queue to gather completed events
+	// (required if we're not using an interrupt thread).
+	ProcessInterruptQueueUntil(nullptr);
+}
+
+#if RHI_NEW_GPU_PROFILER
+FVulkanTiming::FVulkanTiming(FVulkanQueue& InQueue)
+	: Queue(InQueue)
+	, EventStream(InQueue.GetProfilerQueue())
+{
 }
+#endif
 
 void FVulkanCommandListContext::RHIBeginDrawingViewport(FRHIViewport* ViewportRHI, FRHITexture* RenderTargetRHI)
 {
@@ -1175,8 +1239,6 @@ void FVulkanCommandListContext::RHIEndDrawingViewport(FRHIViewport* ViewportRHI,
 	FVulkanCommandBuffer& CommandBuffer = GetCommandBuffer();
 	check(!CommandBuffer.HasEnded() && !CommandBuffer.IsInsideRenderPass());
 
-	WriteEndTimestamp(&CommandBuffer);
-
 	bool bNativePresent = Viewport->Present(*this, Device.GetPresentQueue(), bLockToVsync);
 	if (bNativePresent)
 	{
@@ -1184,8 +1246,6 @@ void FVulkanCommandListContext::RHIEndDrawingViewport(FRHIViewport* ViewportRHI,
 	}
 
 	RHI->DrawingViewport = nullptr;
-
-	WriteBeginTimestamp(GetActiveCmdBuffer());
 }
 
 #if WITH_RHI_BREADCRUMBS
@@ -1218,7 +1278,7 @@ void FVulkanCommandListContext::RHIEndDrawingViewport(FRHIViewport* ViewportRHI,
 				Label.color[1] = LColor.G;
 				Label.color[2] = LColor.B;
 				Label.color[3] = LColor.A;
-				CmdBeginLabel(GetActiveCmdBuffer()->GetHandle(), &Label);
+				CmdBeginLabel(GetCommandBuffer().GetHandle(), &Label);
 			}
 		#endif
 
@@ -1231,7 +1291,19 @@ void FVulkanCommandListContext::RHIEndDrawingViewport(FRHIViewport* ViewportRHI,
 		#endif
 		}
 		
-	#if (RHI_NEW_GPU_PROFILER == 0)
+#if RHI_NEW_GPU_PROFILER
+		{
+			FlushProfilerStats();
+
+			if (GRHIGlobals.SupportsTimestampRenderQueries)
+			{
+				auto& Event = GetCommandBuffer().EmplaceProfilerEvent<UE::RHI::GPUProfiler::FEvent::FBeginBreadcrumb>(Breadcrumb);
+				FVulkanTimingQueryPool* CurrentPool = GetCurrentTimestampQueryPool();
+				const uint32 IndexInPool = CurrentPool->ReserveQuery(&Event.GPUTimestampTOP);
+				VulkanRHI::vkCmdWriteTimestamp(GetCommandBuffer().GetHandle(), VK_PIPELINE_STAGE_TOP_OF_PIPE_BIT, CurrentPool->GetHandle(), IndexInPool);
+			}
+		}
+#else
 		if (IsImmediate())
 		{
 		#if VULKAN_SUPPORTS_GPU_CRASH_DUMPS
@@ -1245,12 +1317,24 @@ void FVulkanCommandListContext::RHIEndDrawingViewport(FRHIViewport* ViewportRHI,
 				GpuProfiler.PushEvent(GetNameStr(), Color);
 			}
 		}
-	#endif // (RHI_NEW_GPU_PROFILER == 0)
+#endif // RHI_NEW_GPU_PROFILER
 	}
 
 	void FVulkanCommandListContext::RHIEndBreadcrumbGPU(FRHIBreadcrumbNode* Breadcrumb)
 	{
-	#if (RHI_NEW_GPU_PROFILER == 0)
+#if RHI_NEW_GPU_PROFILER
+		{
+			FlushProfilerStats();
+
+			if (GRHIGlobals.SupportsTimestampRenderQueries)
+			{
+				auto& Event = GetCommandBuffer().EmplaceProfilerEvent<UE::RHI::GPUProfiler::FEvent::FEndBreadcrumb>(Breadcrumb);
+				FVulkanTimingQueryPool* CurrentPool = GetCurrentTimestampQueryPool();
+				const uint32 IndexInPool = CurrentPool->ReserveQuery(&Event.GPUTimestampBOP);
+				VulkanRHI::vkCmdWriteTimestamp(GetCommandBuffer().GetHandle(), VK_PIPELINE_STAGE_BOTTOM_OF_PIPE_BIT, CurrentPool->GetHandle(), IndexInPool);
+			}
+		}
+#else
 		//only valid on immediate context currently.  needs to be fixed for parallel rhi execute
 		if (IsImmediate())
 		{
@@ -1266,7 +1350,7 @@ void FVulkanCommandListContext::RHIEndDrawingViewport(FRHIViewport* ViewportRHI,
 			}
 		#endif
 		}
-	#endif // (RHI_NEW_GPU_PROFILER == 0)
+#endif // RHI_NEW_GPU_PROFILER
 
 		if (ShouldEmitBreadcrumbs())
 		{
@@ -1280,7 +1364,7 @@ void FVulkanCommandListContext::RHIEndDrawingViewport(FRHIViewport* ViewportRHI,
 		#if VULKAN_ENABLE_DRAW_MARKERS
 			if (auto CmdEndLabel = Device.GetCmdEndDebugLabel())
 			{
-				CmdEndLabel(GetActiveCmdBuffer()->GetHandle());
+				CmdEndLabel(GetCommandBuffer().GetHandle());
 			}
 		#endif
 		}
@@ -1547,11 +1631,13 @@ void FVulkanDynamicRHI::RHIFinishExternalComputeWork(VkCommandBuffer InCommandBu
 
 void FVulkanDynamicRHI::RHIRegisterWork(uint32 NumPrimitives)
 {
+#if (RHI_NEW_GPU_PROFILER == 0)
 	FVulkanCommandListContextImmediate& ImmediateContext = GetDevice()->GetImmediateContext();
 	if (FVulkanPlatform::RegisterGPUWork() && ImmediateContext.IsImmediate())
 	{
 		ImmediateContext.RegisterGPUWork(NumPrimitives);
 	}
+#endif
 }
 
 void FVulkanDynamicRHI::RHISubmitUploadCommandBuffer()
diff --git a/Engine/Source/Runtime/VulkanRHI/Private/VulkanRayTracing.cpp b/Engine/Source/Runtime/VulkanRHI/Private/VulkanRayTracing.cpp
--- a/Engine/Source/Runtime/VulkanRHI/Private/VulkanRayTracing.cpp
+++ b/Engine/Source/Runtime/VulkanRHI/Private/VulkanRayTracing.cpp
@@ -1768,9 +1768,9 @@ void FVulkanRayTracingCompactedSizeQueryPool::EndBatch(FVulkanCommandListContext
 	SyncPoint = CommandContext.GetContextSyncPoint();
 }
 
-void FVulkanRayTracingCompactedSizeQueryPool::Reset(FVulkanCommandBuffer* InCmdBuffer)
+void FVulkanRayTracingCompactedSizeQueryPool::Reset(FVulkanCommandBuffer& InCmdBuffer)
 {
-	VulkanRHI::vkCmdResetQueryPool(InCmdBuffer->GetHandle(), QueryPool, 0, MaxQueries);
+	VulkanRHI::vkCmdResetQueryPool(InCmdBuffer.GetHandle(), QueryPool, 0, MaxQueries);
 	SyncPoint = nullptr;
 	check(QueryOutput.Num() == MaxQueries);
 	FMemory::Memzero(QueryOutput.GetData(), MaxQueries * sizeof(uint64));
@@ -1854,11 +1854,11 @@ void FVulkanRayTracingCompactionRequestHandler::Update(FVulkanCommandListContext
 			{
 				if (ActiveRequests[BLASIndex] != nullptr)
 				{
-					ActiveRequests[BLASIndex]->CompactAccelerationStructure(CmdBuffer, QueryPool->GetResultValue(BLASIndex));
+					ActiveRequests[BLASIndex]->CompactAccelerationStructure(CmdBuffer, QueryPool->QueryOutput[BLASIndex]);
 				}
 			}
 
-			QueryPool->Reset(&CmdBuffer);
+			QueryPool->Reset(CmdBuffer);
 
 			ActiveBLASes.Empty(GVulkanRayTracingMaxBatchedCompaction);
 
diff --git a/Engine/Source/Runtime/VulkanRHI/Private/VulkanRayTracing.h b/Engine/Source/Runtime/VulkanRHI/Private/VulkanRayTracing.h
--- a/Engine/Source/Runtime/VulkanRHI/Private/VulkanRayTracing.h
+++ b/Engine/Source/Runtime/VulkanRHI/Private/VulkanRayTracing.h
@@ -318,13 +318,16 @@ class FVulkanRayTracingCompactedSizeQueryPool : public FVulkanQueryPool
 {
 public:
 	FVulkanRayTracingCompactedSizeQueryPool(FVulkanDevice& InDevice, uint32 InMaxQueries)
-		: FVulkanQueryPool(InDevice, nullptr, InMaxQueries, VK_QUERY_TYPE_ACCELERATION_STRUCTURE_COMPACTED_SIZE_KHR, false)
-	{}
+		: FVulkanQueryPool(InDevice, InMaxQueries, EVulkanQueryPoolType::ASCompactedSize)
+	{
+		QueryOutput.SetNumZeroed(InMaxQueries);
+	}
 
 	void EndBatch(FVulkanCommandListContext& CommandContext);
 	bool TryGetResults(uint32 NumResults);
-	void Reset(FVulkanCommandBuffer* InCmdBuffer);
+	void Reset(FVulkanCommandBuffer& InCmdBuffer);
 
+	TArray<uint64> QueryOutput;
 	FVulkanSyncPointRef SyncPoint;
 };
 
diff --git a/Engine/Source/Runtime/VulkanRHI/Private/VulkanRenderTarget.cpp b/Engine/Source/Runtime/VulkanRHI/Private/VulkanRenderTarget.cpp
--- a/Engine/Source/Runtime/VulkanRHI/Private/VulkanRenderTarget.cpp
+++ b/Engine/Source/Runtime/VulkanRHI/Private/VulkanRenderTarget.cpp
@@ -511,13 +511,11 @@ bool FVulkanCommandListContext::IsSwapchainImage(FRHITexture* InTexture) const
 
 void FVulkanCommandListContext::RHIBeginRenderPass(const FRHIRenderPassInfo& InInfo, const TCHAR* InName)
 {
-	FVulkanCommandBuffer* CmdBuffer = GetActiveCmdBuffer();
-
 	RenderPassInfo = InInfo;
 
 	if (InInfo.NumOcclusionQueries > 0)
 	{
-		BeginOcclusionQueryBatch(CmdBuffer, InInfo.NumOcclusionQueries);
+		BeginOcclusionQueryBatch(InInfo.NumOcclusionQueries);
 	}
 
 	FRHITexture* DSTexture = InInfo.DepthStencilRenderTarget.DepthStencilTarget;
@@ -571,7 +569,7 @@ void FVulkanCommandListContext::RHIBeginRenderPass(const FRHIRenderPassInfo& InI
 
 	FVulkanFramebuffer* Framebuffer = Device.GetRenderPassManager().GetOrCreateFramebuffer(RTInfo, RTLayout, RenderPass);
 	checkf(RenderPass != nullptr && Framebuffer != nullptr, TEXT("RenderPass not started! Bad combination of values? Depth %p #Color %d Color0 %p"), (void*)InInfo.DepthStencilRenderTarget.DepthStencilTarget, InInfo.GetNumColorRenderTargets(), (void*)InInfo.ColorRenderTargets[0].RenderTarget);
-	Device.GetRenderPassManager().BeginRenderPass(*this, Device, CmdBuffer, InInfo, RTLayout, RenderPass, Framebuffer);
+	Device.GetRenderPassManager().BeginRenderPass(*this, Device, &GetCommandBuffer(), InInfo, RTLayout, RenderPass, Framebuffer);
 
 	check(!CurrentRenderPass);
 	CurrentRenderPass = RenderPass;
@@ -586,8 +584,6 @@ void FVulkanCommandListContext::RHIEndRenderPass()
 	const bool bHasOcclusionQueries = (RenderPassInfo.NumOcclusionQueries > 0);
 	if (bHasOcclusionQueries)
 	{
-		EndOcclusionQueryBatch();
-
 		// Force the syncs points to be signaled right after the render pass containing the queries
 		FlushPendingSyncPoints();
 	}
diff --git a/Engine/Source/Runtime/VulkanRHI/Private/VulkanSubmission.cpp b/Engine/Source/Runtime/VulkanRHI/Private/VulkanSubmission.cpp
--- a/Engine/Source/Runtime/VulkanRHI/Private/VulkanSubmission.cpp
+++ b/Engine/Source/Runtime/VulkanRHI/Private/VulkanSubmission.cpp
@@ -130,6 +130,24 @@ void FVulkanDynamicRHI::InitializeSubmissionPipe()
 			SubmissionThread = new FVulkanThread(TEXT("RHISubmissionThread"), TPri_Highest, this, &FVulkanDynamicRHI::ProcessSubmissionQueue);
 		}
 	}
+
+#if RHI_NEW_GPU_PROFILER
+	// Initialize the timing structs in each queue, and the engine GPU profilers
+	{
+		TArray<UE::RHI::GPUProfiler::FQueue> ProfilerQueues;
+		FVulkanPlatformCommandList* Payloads = new FVulkanPlatformCommandList;
+		Payloads->Reserve((int32)EVulkanQueueType::Count);
+		Device->ForEachQueue([&](FVulkanQueue& Queue)
+		{
+			ProfilerQueues.Add(Queue.GetProfilerQueue());
+			FVulkanPayload* Payload = Payloads->Emplace_GetRef(new FVulkanPayload(Queue));
+			Payload->Timing = CurrentTimingPerQueue.CreateNew(Queue);
+		});
+
+		UE::RHI::GPUProfiler::InitializeQueues(ProfilerQueues);
+		PendingPayloadsForSubmission.Enqueue(Payloads);
+	}
+#endif
 }
 
 void FVulkanDynamicRHI::ShutdownSubmissionPipe()
@@ -234,10 +252,31 @@ bool FVulkanDynamicRHI::ProcessSubmissionQueue()
 	FVulkanPlatformCommandList* PlatformCmdList = nullptr;
 	while (PendingPayloadsForSubmission.Dequeue(PlatformCmdList))
 	{
+
+#if WITH_RHI_BREADCRUMBS && RHI_NEW_GPU_PROFILER
+		TSharedPtr<FRHIBreadcrumbAllocatorArray> BreadcrumbAllocators{};
+		if (PlatformCmdList->BreadcrumbAllocators.Num())
+		{
+			BreadcrumbAllocators = MakeShared<FRHIBreadcrumbAllocatorArray>(MoveTemp(PlatformCmdList->BreadcrumbAllocators));
+		}
+
+		for (FVulkanPayload* Payload : *PlatformCmdList)
+		{
+			Payload->BreadcrumbRange = PlatformCmdList->BreadcrumbRange;
+			if (BreadcrumbAllocators.IsValid())
+			{
+				check(!Payload->BreadcrumbAllocators.IsValid());
+				Payload->BreadcrumbAllocators = BreadcrumbAllocators;
+			}
+
+			Payload->Queue.EnqueuePayload(Payload);
+		}
+#else
 		for (FVulkanPayload* Payload : *PlatformCmdList)
 		{
 			Payload->Queue.EnqueuePayload(Payload);
 		}
+#endif
 
 		delete PlatformCmdList;
 	}
@@ -485,4 +524,123 @@ void FVulkanDynamicRHI::RHIBlockUntilGPUIdle()
 	{
 		ProcessInterruptQueueUntil(Event);
 	}
-}
\ No newline at end of file
+}
+
+// Resolve any pending actions and delete the payload
+void FVulkanDynamicRHI::CompletePayload(FVulkanPayload* Payload)
+{
+	for (FVulkanCommandBuffer* CommandBuffer : Payload->CommandBuffers)
+	{
+		CommandBuffer->ResetFenceStatus();
+	}
+
+	struct FQueryResult
+	{
+		uint64 Result;
+		uint64 Availability;
+	};
+	static TArray<FQueryResult> TempResults;
+	for (TArray<FVulkanQueryPool*>& QueryPoolArray : Payload->QueryPools)
+	{
+		for (FVulkanQueryPool* QueryPool : QueryPoolArray)
+		{
+			check(QueryPool->CurrentQueryCount);
+
+			// We need one for the result and one for availablility
+			if (TempResults.Num() < (int32)QueryPool->GetMaxQueries())
+			{
+				TempResults.SetNumZeroed(QueryPool->GetMaxQueries());
+			}
+
+			const VkQueryResultFlags QueryFlags = VK_QUERY_RESULT_64_BIT | VK_QUERY_RESULT_WITH_AVAILABILITY_BIT | VK_QUERY_RESULT_WAIT_BIT;
+			const VkResult Result = VulkanRHI::vkGetQueryPoolResults(
+				Device->GetInstanceHandle(),
+				QueryPool->GetHandle(),
+				0,
+				QueryPool->CurrentQueryCount,
+				TempResults.Num() * TempResults.GetTypeSize(),
+				(uint64_t*)TempResults.GetData(),
+				sizeof(FQueryResult),
+				QueryFlags);
+			VERIFYVULKANRESULT_EXPANDED(Result);
+
+			for (uint32 QueryIndex = 0; QueryIndex < QueryPool->CurrentQueryCount; ++QueryIndex)
+			{
+				if (TempResults[QueryIndex].Availability)
+				{
+					uint64* ResultPtr = QueryPool->QueryResults[QueryIndex];
+					if (ResultPtr)
+					{
+#if RHI_NEW_GPU_PROFILER
+						if (QueryPool->GetPoolType() == EVulkanQueryPoolType::Timestamp)
+						{
+							// Convert from GPU timestamp to CPU timestamp (relative to FPlatformTime::Cycles64())
+							checkSlow(Payload->Queue.Timing);
+							const FVulkanTiming* CurrentTiming = Payload->Queue.Timing;
+							checkf(TempResults[QueryIndex].Result >= CurrentTiming->GPUTimestamp,
+								TEXT("Invalid timestamp can't be moved to CPU timestamp (%lld < %lld)"),
+								TempResults[QueryIndex].Result, CurrentTiming->GPUTimestamp);
+							const uint64 GPUDelta = TempResults[QueryIndex].Result - CurrentTiming->GPUTimestamp;
+							const uint64 CPUDelta = (GPUDelta * CurrentTiming->CPUFrequency) / CurrentTiming->GPUFrequency;
+							(*ResultPtr) = CPUDelta + CurrentTiming->CPUTimestamp;
+						}
+						else
+#endif
+						{
+							(*ResultPtr) = TempResults[QueryIndex].Result;
+						}
+					}
+				}
+				else
+				{
+					UE_LOG(LogVulkanRHI, Warning, TEXT("Could not read back query: [PoolType=%d, Index=%d]"), (int32)QueryPool->GetPoolType(), QueryIndex);
+				}
+			}
+
+			QueryPool->QueryRefs.Empty(QueryPool->GetMaxQueries());
+			Device->ReleaseQueryPool(QueryPool);
+		}
+		QueryPoolArray.Reset();
+	}
+
+#if RHI_NEW_GPU_PROFILER
+	if (!Payload->EventStream.IsEmpty())
+	{
+		check(Payload->Queue.Timing);
+		Payload->Queue.Timing->EventStream.Append(MoveTemp(Payload->EventStream));
+	}
+
+	if (Payload->Timing.IsSet())
+	{
+		// Switch the new timing struct into the queue. This redirects timestamp results to separate each frame's work.
+		Payload->Queue.Timing = Payload->Timing.GetValue();
+	}
+#else
+	if (GRHIGlobals.SupportsTimestampRenderQueries && !FVulkanPlatform::HasCustomFrameTiming())
+	{
+		static uint64 BusyCycles = 0;
+		if (Payload->Queue.QueueType == EVulkanQueueType::Graphics)
+		{
+			for (FVulkanCommandBuffer* CommandBuffer : Payload->CommandBuffers)
+			{
+				BusyCycles += CommandBuffer->GetBusyCycles();
+			}
+		}
+		if (Payload->bEndFrame && (BusyCycles > 0))
+		{
+			const double SecondsPerCycle = FPlatformTime::GetSecondsPerCycle();
+			const double Frequency = double((uint64)((1000.0 * 1000.0 * 1000.0) / Device->GetLimits().timestampPeriod));
+			GGPUFrameTime = FMath::TruncToInt(double(BusyCycles) / Frequency / SecondsPerCycle);
+			BusyCycles = 0;
+		}
+	}
+#endif
+
+	for (FVulkanSyncPointRef& SyncPoint : Payload->SyncPoints)
+	{
+		SyncPoint->DispatchSubsequents();
+	}
+	Payload->SyncPoints.Empty();
+
+	delete Payload;
+}
diff --git a/Engine/Source/Runtime/VulkanRHI/Private/VulkanSubmission.h b/Engine/Source/Runtime/VulkanRHI/Private/VulkanSubmission.h
--- a/Engine/Source/Runtime/VulkanRHI/Private/VulkanSubmission.h
+++ b/Engine/Source/Runtime/VulkanRHI/Private/VulkanSubmission.h
@@ -7,14 +7,15 @@
 #pragma once 
 
 #include "VulkanRHIPrivate.h"
+#include "VulkanQuery.h"
 
 class FVulkanDevice;
 class FVulkanDynamicRHI;
-class FVulkanOcclusionQueryPool;
 class FVulkanPendingGfxState;
 class FVulkanPendingComputeState;
 class FVulkanQueue;
 class FVulkanSwapChain;
+class FVulkanTiming;
 
 // NOTE: typedef to FGraphEvent for now, but eventually move all syncs inside a FVulkanSyncPoint class
 using FVulkanSyncPoint = FGraphEvent;
@@ -34,11 +35,15 @@ class FVulkanPayload
 public:
 	FVulkanPayload(FVulkanQueue& InQueue)
 		: Queue(InQueue)
+#if RHI_NEW_GPU_PROFILER
+		, EventStream(InQueue.GetProfilerQueue())
+#endif
 	{}
 
 	~FVulkanPayload();
 
 protected:
+
 	FVulkanQueue& Queue;
 	TArray<VkPipelineStageFlags> WaitFlags;  // flags that match 1:1 with WaitSemaphores
 	TArray<VulkanRHI::FSemaphore*> WaitSemaphores; // wait before command buffers
@@ -54,6 +59,22 @@ protected:
 
 	// Used to sync other CPU work to payload completion
 	TArray<FVulkanSyncPointRef> SyncPoints;
+
+	// Queries used in the command lists of this payload
+	TStaticArray<TArray<FVulkanQueryPool*>, (int32)EVulkanQueryPoolType::Count> QueryPools;
+
+#if WITH_RHI_BREADCRUMBS
+	FRHIBreadcrumbRange BreadcrumbRange{};
+	TSharedPtr<FRHIBreadcrumbAllocatorArray> BreadcrumbAllocators{};
+#endif
+
+#if RHI_NEW_GPU_PROFILER
+	TOptional<FVulkanTiming*> Timing;
+	TOptional<UE::RHI::GPUProfiler::FEvent::FFrameBoundary> EndFrameEvent;
+	UE::RHI::GPUProfiler::FEventStream EventStream;
+#else
+	bool bEndFrame = false;
+#endif
 };
 
 struct FVulkanPlatformCommandList : public IRHIPlatformCommandList, public TArray<FVulkanPayload*>
diff --git a/Engine/Source/Runtime/VulkanRHI/Private/VulkanUtil.cpp b/Engine/Source/Runtime/VulkanRHI/Private/VulkanUtil.cpp
--- a/Engine/Source/Runtime/VulkanRHI/Private/VulkanUtil.cpp
+++ b/Engine/Source/Runtime/VulkanRHI/Private/VulkanUtil.cpp
@@ -68,7 +68,6 @@ void FVulkanGPUTiming::CalibrateTimers(FVulkanDevice& Device)
 
 FVulkanGPUTiming::~FVulkanGPUTiming()
 {
-	check(!Pool);
 }
 
 /**
@@ -79,12 +78,6 @@ void FVulkanGPUTiming::Initialize(uint32 PoolSize)
 	StaticInitialize(this, PlatformStaticInitialize);
 
 	bIsTiming = false;
-
-	if (FVulkanPlatform::SupportsTimestampRenderQueries() && GIsSupported)
-	{
-		check(!Pool);
-		Pool = new FVulkanTimingQueryPool(*Device, Context, PoolSize);
-	}
 }
 
 /**
@@ -92,57 +85,40 @@ void FVulkanGPUTiming::Initialize(uint32 PoolSize)
  */
 void FVulkanGPUTiming::Release()
 {
-	if (FVulkanPlatform::SupportsTimestampRenderQueries() && GIsSupported)
-	{
-		check(Pool);
-		delete Pool;
-		Pool = nullptr;
-	}
 }
 
 /**
  * Start a GPU timing measurement.
  */
-void FVulkanGPUTiming::StartTiming(FVulkanCommandBuffer* CmdBuffer)
+void FVulkanGPUTiming::StartTiming(FVulkanContextCommon* InContext)
 {
 	// Issue a timestamp query for the 'start' time.
 	if (GIsSupported && !bIsTiming)
 	{
-		if (CmdBuffer == nullptr)
+		if (InContext == nullptr)
 		{
-			CmdBuffer = Context->GetActiveCmdBuffer();
+			InContext = Context;
 		}
 
 		// In case we aren't reading queries, remove oldest
-		if (NumPendingQueries >= Pool->BufferedCount)
+		if (NumPendingQueries >= MaxPendingQueries)
 		{
 			DiscardOldestQuery();
-			if (NumPendingQueries >= Pool->BufferedCount)
+			if (NumPendingQueries >= MaxPendingQueries)
 			{
 				return;
 			}
 		}
 
-		Pool->CurrentTimestampIndex = (Pool->CurrentTimestampIndex + 1) % Pool->BufferedCount;
-		const uint32 QueryStartIndex = Pool->CurrentTimestampIndex * 2;
-		FVulkanTimingQueryPool::FTimestampData& StartQueryData = Pool->TimestampArray[QueryStartIndex];
+		check(!ActiveQuery);
+		ActiveQuery = new FPendingQuery;
 
-		// Since we use non-blocking queries, we need to reset them before using them
-		// Use host query resets if supported since there's no need to consider host/GPU sync
-		if (Device->GetOptionalExtensions().HasEXTHostQueryReset)
-		{
-			VulkanRHI::vkResetQueryPoolEXT(Device->GetInstanceHandle(), Pool->GetHandle(), QueryStartIndex, 2);
-		}
-		else
-		{
-			VulkanRHI::vkCmdResetQueryPool(CmdBuffer->GetHandle(), Pool->GetHandle(), QueryStartIndex, 2);
-		}
+		FVulkanTimingQueryPool* CurrentPool = InContext->GetCurrentTimestampQueryPool();
+		const uint32 IndexInPool = CurrentPool->ReserveQuery(&ActiveQuery->StartResult);
+		VulkanRHI::vkCmdWriteTimestamp(InContext->GetCommandBuffer().GetHandle(), VK_PIPELINE_STAGE_BOTTOM_OF_PIPE_BIT, CurrentPool->GetHandle(), IndexInPool);
+		ActiveQuery->StartSyncPoint = InContext->GetContextSyncPoint();
 
-		VulkanRHI::vkCmdWriteTimestamp(CmdBuffer->GetHandle(), VK_PIPELINE_STAGE_BOTTOM_OF_PIPE_BIT, Pool->GetHandle(), QueryStartIndex);
-		StartQueryData.SyncPoint = Context->GetContextSyncPoint();
 		bIsTiming = true;
-
-		CmdBuffer->AddPendingTimestampQuery(QueryStartIndex, 1, Pool->GetHandle(), &StartQueryData.Result, false);
 	}
 }
 
@@ -150,61 +126,46 @@ void FVulkanGPUTiming::StartTiming(FVulkanCommandBuffer* CmdBuffer)
  * End a GPU timing measurement.
  * The timing for this particular measurement will be resolved at a later time by the GPU.
  */
-void FVulkanGPUTiming::EndTiming(FVulkanCommandBuffer* CmdBuffer)
+void FVulkanGPUTiming::EndTiming(FVulkanContextCommon* InContext)
 {
 	// Issue a timestamp query for the 'end' time.
 	if (GIsSupported && bIsTiming)
 	{
-		if (CmdBuffer == nullptr)
+		if (InContext == nullptr)
 		{
-			CmdBuffer = Context->GetActiveCmdBuffer();
+			InContext = Context;
 		}
-		check(Pool->CurrentTimestampIndex < Pool->BufferedCount);
-		const uint32 QueryStartIndex = Pool->CurrentTimestampIndex * 2;
-		const uint32 QueryEndIndex = QueryStartIndex + 1;
 
-		FVulkanTimingQueryPool::FTimestampData& EndQueryData = Pool->TimestampArray[QueryEndIndex];
+		check(ActiveQuery);
 
-		PendingQueries.Enqueue(Pool->CurrentTimestampIndex);
-		NumPendingQueries++;
+		FVulkanTimingQueryPool* CurrentPool = InContext->GetCurrentTimestampQueryPool();
+		const uint32 IndexInPool = CurrentPool->ReserveQuery(&ActiveQuery->EndResult);
+		VulkanRHI::vkCmdWriteTimestamp(InContext->GetCommandBuffer().GetHandle(), VK_PIPELINE_STAGE_BOTTOM_OF_PIPE_BIT, CurrentPool->GetHandle(), IndexInPool);
+		ActiveQuery->EndSyncPoint = InContext->GetContextSyncPoint();
 
-		VulkanRHI::vkCmdWriteTimestamp(CmdBuffer->GetHandle(), VK_PIPELINE_STAGE_BOTTOM_OF_PIPE_BIT, Pool->GetHandle(), QueryEndIndex);
-		Pool->TimestampArray[QueryEndIndex].SyncPoint = Context->GetContextSyncPoint();
-		Pool->NumIssuedTimestamps = FMath::Min<uint32>(Pool->NumIssuedTimestamps + 1, Pool->BufferedCount);
+		PendingQueries.Enqueue(ActiveQuery);
+		NumPendingQueries++;
+		ActiveQuery = nullptr;
 
 		bIsTiming = false;
 		bEndTimestampIssued = true;
-
-		CmdBuffer->AddPendingTimestampQuery(QueryEndIndex, 1, Pool->GetHandle(), &EndQueryData.Result, false);
 	}
 }
 
 void FVulkanGPUTiming::DiscardOldestQuery()
 {
-	auto CanBeDiscarded = [&](uint32 QueryIndex)
+	FPendingQuery* PendingQuery = nullptr;
+	if (PendingQueries.Peek(PendingQuery))
 	{
-		// It can only be competed if the sync point was signaled (queries hold a pointer they will write to when completed)
-		FVulkanTimingQueryPool::FTimestampData& QueryData = Pool->TimestampArray[QueryIndex];
-		if (!QueryData.SyncPoint->IsComplete())
+		if (!PendingQuery->StartSyncPoint->IsComplete() || !PendingQuery->EndSyncPoint->IsComplete())
 		{
 			FVulkanDynamicRHI::Get().ProcessInterruptQueueUntil(nullptr);  // leave null, we don't want to force a wait on the SyncPoint
 		}
-		return QueryData.SyncPoint->IsComplete();
-	};
-
-	uint32 OldestTimestamp;
-	if (PendingQueries.Peek(OldestTimestamp))
-	{
-		const uint32 QueryStartIndex = OldestTimestamp * 2;
-		const uint32 QueryEndIndex = QueryStartIndex + 1;
-
-		const bool StartQueryExpired = CanBeDiscarded(QueryStartIndex);
-		const bool EndQueryExpired = CanBeDiscarded(QueryEndIndex);
 
-		if (StartQueryExpired || EndQueryExpired)
+		if (PendingQuery->StartSyncPoint->IsComplete() && PendingQuery->EndSyncPoint->IsComplete())
 		{
-			// If the queries is no good anymore, throw it away and try the next one
 			PendingQueries.Pop();
+			delete PendingQuery;
 			NumPendingQueries--;
 		}
 	}
@@ -220,61 +181,38 @@ uint64 FVulkanGPUTiming::GetTiming(bool bGetCurrentResultsAndBlock)
 {
 	if (GIsSupported)
 	{
-		checkf(Pool->CurrentTimestampIndex < Pool->BufferedCount, TEXT("Timestamp index out of bounds!"));
-		uint32 TimestampIndex;
-		while (PendingQueries.Peek(TimestampIndex))
+		FPendingQuery* PendingQuery = nullptr;
+		while (PendingQueries.Peek(PendingQuery))
 		{
-			const uint32 QueryStartIndex = TimestampIndex * 2;
-			const uint32 QueryEndIndex = QueryStartIndex + 1;
-
-			const FVulkanTimingQueryPool::FTimestampData& StartQueryData = Pool->TimestampArray[QueryStartIndex];
-			const FVulkanTimingQueryPool::FTimestampData& EndQueryData = Pool->TimestampArray[QueryEndIndex];
-
-			if (bGetCurrentResultsAndBlock && StartQueryData.SyncPoint.IsValid() && EndQueryData.SyncPoint.IsValid())
+			if (bGetCurrentResultsAndBlock && PendingQuery->StartSyncPoint.IsValid() && PendingQuery->EndSyncPoint.IsValid())
 			{
 				SCOPE_CYCLE_COUNTER(STAT_RenderQueryResultTime);
 				{
 					FRenderThreadIdleScope IdleScope(ERenderThreadIdleTypes::WaitingForGPUQuery);
 
 					// CPU wait for query results to be ready.
-					if (!StartQueryData.SyncPoint->IsComplete())
+					if (!PendingQuery->StartSyncPoint->IsComplete())
 					{
-						FVulkanDynamicRHI::Get().ProcessInterruptQueueUntil(StartQueryData.SyncPoint);
+						FVulkanDynamicRHI::Get().ProcessInterruptQueueUntil(PendingQuery->StartSyncPoint);
 					}
-					if (!EndQueryData.SyncPoint->IsComplete())
+					if (!PendingQuery->EndSyncPoint->IsComplete())
 					{
-						FVulkanDynamicRHI::Get().ProcessInterruptQueueUntil(EndQueryData.SyncPoint);
+						FVulkanDynamicRHI::Get().ProcessInterruptQueueUntil(PendingQuery->EndSyncPoint);
 					}
 				}
 			}
 
-			if (EndQueryData.SyncPoint->IsComplete() && StartQueryData.SyncPoint->IsComplete())
+			if (PendingQuery->EndSyncPoint->IsComplete() && PendingQuery->StartSyncPoint->IsComplete())
 			{
-				if (!StartQueryData.Result.Availability)
+				if (PendingQuery->EndResult > PendingQuery->StartResult)
 				{
-					const VkQueryResultFlags QueryResultFlags = VK_QUERY_RESULT_64_BIT | VK_QUERY_RESULT_WAIT_BIT;
-					const VkResult Result = VulkanRHI::vkGetQueryPoolResults(
-						Device->GetInstanceHandle(), Pool->GetHandle(),
-						QueryStartIndex, 1, sizeof(uint64), (uint64_t*)&StartQueryData.Result.Value, sizeof(uint64), QueryResultFlags);
-					VERIFYVULKANRESULT(Result);
-				}
-				if (!EndQueryData.Result.Availability)
-				{
-					const VkQueryResultFlags QueryResultFlags = VK_QUERY_RESULT_64_BIT | VK_QUERY_RESULT_WAIT_BIT;
-					const VkResult Result = VulkanRHI::vkGetQueryPoolResults(
-						Device->GetInstanceHandle(), Pool->GetHandle(),
-						QueryStartIndex, 1, sizeof(uint64), (uint64_t*)&EndQueryData.Result.Value, sizeof(uint64), QueryResultFlags);
-					VERIFYVULKANRESULT(Result);
+					// Only keep the most recent result
+					LastTime = PendingQuery->EndResult - PendingQuery->StartResult;
 				}
 
 				PendingQueries.Pop();
+				delete PendingQuery;
 				NumPendingQueries--;
-
-				if (EndQueryData.Result.Value > StartQueryData.Result.Value)
-				{
-					// Only keep the most recent result
-					LastTime = EndQueryData.Result.Value - StartQueryData.Result.Value;
-				}
 			}
 			else
 			{
diff --git a/Engine/Source/Runtime/VulkanRHI/Private/Windows/VulkanWindowsPlatform.h b/Engine/Source/Runtime/VulkanRHI/Private/Windows/VulkanWindowsPlatform.h
--- a/Engine/Source/Runtime/VulkanRHI/Private/Windows/VulkanWindowsPlatform.h
+++ b/Engine/Source/Runtime/VulkanRHI/Private/Windows/VulkanWindowsPlatform.h
@@ -58,6 +58,8 @@ public:
 
 	static void WriteCrashMarker(const FOptionalVulkanDeviceExtensions& OptionalExtensions, FVulkanCommandBuffer* CmdBuffer, VkBuffer DestBuffer, const TArrayView<uint32>& Entries, bool bAdding);
 
+	static VkTimeDomainKHR GetTimeDomain() { return VK_TIME_DOMAIN_QUERY_PERFORMANCE_COUNTER_KHR; }
+
 private:
 	static bool bAttemptedLoad;
 };
