VulkanRHI: Add option to specify initial value for timeline semaphores.

#rnx

diff --git a/Engine/Source/Runtime/VulkanRHI/Private/VulkanMemory.cpp b/Engine/Source/Runtime/VulkanRHI/Private/VulkanMemory.cpp
--- a/Engine/Source/Runtime/VulkanRHI/Private/VulkanMemory.cpp
+++ b/Engine/Source/Runtime/VulkanRHI/Private/VulkanMemory.cpp
@@ -5122,7 +5122,7 @@ namespace VulkanRHI
 
 
 
-	FSemaphore::FSemaphore(FVulkanDevice& InDevice, EVulkanSemaphoreFlags InFlags)
+	FSemaphore::FSemaphore(FVulkanDevice& InDevice, EVulkanSemaphoreFlags InFlags, uint64 InInitialTimelineValue)
 		: Device(InDevice)
 		, SemaphoreHandle(VK_NULL_HANDLE)
 		, Flags(InFlags)
@@ -5137,7 +5137,7 @@ namespace VulkanRHI
 			checkf(InDevice.GetOptionalExtensions().HasKHRTimelineSemaphore, TEXT("Trying to create timeline semaphore without device support!"));
 			ZeroVulkanStruct(TypeCreateInfo, VK_STRUCTURE_TYPE_SEMAPHORE_TYPE_CREATE_INFO);
 			TypeCreateInfo.semaphoreType = VK_SEMAPHORE_TYPE_TIMELINE;
-			TypeCreateInfo.initialValue = 0;
+			TypeCreateInfo.initialValue = (uint64_t)InInitialTimelineValue;
 			CreateInfo.pNext = &TypeCreateInfo;
 		}
 
diff --git a/Engine/Source/Runtime/VulkanRHI/Private/VulkanMemory.h b/Engine/Source/Runtime/VulkanRHI/Private/VulkanMemory.h
--- a/Engine/Source/Runtime/VulkanRHI/Private/VulkanMemory.h
+++ b/Engine/Source/Runtime/VulkanRHI/Private/VulkanMemory.h
@@ -1173,7 +1173,7 @@ namespace VulkanRHI
 	class VULKANRHI_API FSemaphore : public FRefCount
 	{
 	public:
-		FSemaphore(FVulkanDevice& InDevice, EVulkanSemaphoreFlags InFlags=EVulkanSemaphoreFlags::None);
+		FSemaphore(FVulkanDevice& InDevice, EVulkanSemaphoreFlags InFlags=EVulkanSemaphoreFlags::None, uint64 InInitialTimelineValue=0);
 		FSemaphore(FVulkanDevice& InDevice, const VkSemaphore& InExternalSemaphore);
 		virtual ~FSemaphore();
 
diff --git a/Engine/Source/Runtime/VulkanRHI/Private/VulkanQueue.cpp b/Engine/Source/Runtime/VulkanRHI/Private/VulkanQueue.cpp
--- a/Engine/Source/Runtime/VulkanRHI/Private/VulkanQueue.cpp
+++ b/Engine/Source/Runtime/VulkanRHI/Private/VulkanQueue.cpp
@@ -55,7 +55,8 @@ FVulkanQueue::FVulkanQueue(FVulkanDevice& InDevice, uint32 InFamilyIndex, EVulka
 	if (bUseTimelineSemaphores)
 	{
 		// Use ImmediateDeletion for the timeline since it's deleted after the deferrered deletion queue
-		TimelineSempahore = new VulkanRHI::FSemaphore(InDevice, VulkanRHI::EVulkanSemaphoreFlags::Timeline | VulkanRHI::EVulkanSemaphoreFlags::ImmediateDeletion);
+		const VulkanRHI::EVulkanSemaphoreFlags SemaphoreFlags = VulkanRHI::EVulkanSemaphoreFlags::Timeline | VulkanRHI::EVulkanSemaphoreFlags::ImmediateDeletion;
+		TimelineSempahore = new VulkanRHI::FSemaphore(InDevice, SemaphoreFlags, CompletedTimelineSemaphoreValue);
 		TimelineSempahore->AddRef();
 	}
 }
