VulkanRHI: Remove empty submits for payloads without command buffers.

#rnx

diff --git a/Engine/Source/Runtime/VulkanRHI/Private/VulkanQueue.cpp b/Engine/Source/Runtime/VulkanRHI/Private/VulkanQueue.cpp
--- a/Engine/Source/Runtime/VulkanRHI/Private/VulkanQueue.cpp
+++ b/Engine/Source/Runtime/VulkanRHI/Private/VulkanQueue.cpp
@@ -157,6 +157,14 @@ void FVulkanQueue::SubmitPayloads(TArrayView<FVulkanPayload*> Payloads, TSet<VkS
 
 	for (FVulkanPayload* Payload : Payloads)
 	{
+		// Some payloads have nothing to submit because they are only used to trigger CPU events
+		if (!Payload->WaitSemaphores.Num() && !Payload->CommandBuffers.Num() && !Payload->SignalSemaphores.Num())
+		{
+			// Consider complete when previous workload is done
+			Payload->TimelineSemaphoreValue = NextTimelineSemaphoreValue - 1;
+			continue;
+		}
+
 		VkSubmitInfo& SubmitInfo = SubmitInfos.AddZeroed_GetRef();
 		SubmitInfo.sType = VK_STRUCTURE_TYPE_SUBMIT_INFO;
 
@@ -266,7 +274,7 @@ void FVulkanQueue::SubmitPayloads(TArrayView<FVulkanPayload*> Payloads, TSet<VkS
 		NextTimelineSemaphoreValue++;
 	}
 
-	if (bUseTimelineSemaphores)
+	if (bUseTimelineSemaphores && SubmitInfos.Num())
 	{
 		Submit(SubmitInfos, nullptr);
 	}
@@ -321,12 +329,14 @@ int32 FVulkanQueue::ProcessInterruptQueue(uint64 Timeout)
 		bSuccess = false;
 		if (PendingInterrupt.Peek(Payload))
 		{
-			if (bUseTimelineSemaphores)
+			if (Payload->TimelineSemaphoreValue <= CompletedTimelineSemaphoreValue)
+			{
+				bSuccess = true;
+			}
+			else if (bUseTimelineSemaphores)
 			{
 				checkSlow(Payload->TimelineSemaphoreValue > 0);
-				bSuccess = (Payload->TimelineSemaphoreValue <= CompletedTimelineSemaphoreValue);
-
-				if (!bSuccess && (Timeout > 0))
+				if (Timeout > 0)
 				{
 					bSuccess = TimelineSempahore->WaitForTimelineSemaphoreValue(Payload->TimelineSemaphoreValue, Timeout);
 					CompletedTimelineSemaphoreValue = TimelineSempahore->GetTimelineSemaphoreValue();
