VulkanRHI: allow payloads to handle more than a single submission event

#rnx

diff --git a/Engine/Source/Runtime/VulkanRHI/Private/VulkanContext.cpp b/Engine/Source/Runtime/VulkanRHI/Private/VulkanContext.cpp
--- a/Engine/Source/Runtime/VulkanRHI/Private/VulkanContext.cpp
+++ b/Engine/Source/Runtime/VulkanRHI/Private/VulkanContext.cpp
@@ -83,7 +83,7 @@ void FVulkanContextCommon::FlushCommands(EVulkanFlushFlags FlushFlags)
 	if (EnumHasAnyFlags(FlushFlags, EVulkanFlushFlags::WaitForSubmission))
 	{
 		SubmissionEvent = FGraphEvent::CreateGraphEvent();
-		GetPayload(EPhase::Signal).SubmissionEvent = SubmissionEvent;
+		AddSubmissionEvent(SubmissionEvent);
 	}
 
 	FVulkanPlatformCommandList* FinalizedPayloads = new FVulkanPlatformCommandList;
diff --git a/Engine/Source/Runtime/VulkanRHI/Private/VulkanContext.h b/Engine/Source/Runtime/VulkanRHI/Private/VulkanContext.h
--- a/Engine/Source/Runtime/VulkanRHI/Private/VulkanContext.h
+++ b/Engine/Source/Runtime/VulkanRHI/Private/VulkanContext.h
@@ -163,6 +163,13 @@ public:
 		PendingSyncPoints.Add(InSyncs);
 	}
 
+	// Add an event to be signaled with the current payload is submitted (does not alter phase)
+	void AddSubmissionEvent(FGraphEventRef& InEvent)
+	{
+		FVulkanPayload& Payload = GetPayload(CurrentPhase);
+		Payload.SubmissionEvents.Add(InEvent);
+	}
+
 	// Force pending syncs to be sent to a payload
 	void FlushPendingSyncPoints()
 	{
diff --git a/Engine/Source/Runtime/VulkanRHI/Private/VulkanQueue.cpp b/Engine/Source/Runtime/VulkanRHI/Private/VulkanQueue.cpp
--- a/Engine/Source/Runtime/VulkanRHI/Private/VulkanQueue.cpp
+++ b/Engine/Source/Runtime/VulkanRHI/Private/VulkanQueue.cpp
@@ -282,11 +282,11 @@ void FVulkanQueue::SubmitPayloads(TArrayView<FVulkanPayload*> Payloads, TSet<VkS
 	// Queue the submitted payloads in the interrupt queue
 	for (FVulkanPayload* Payload : Payloads)
 	{
-		if (Payload->SubmissionEvent.IsValid())
+		for (FGraphEventRef& SubmissionEvent : Payload->SubmissionEvents)
 		{
-			Payload->SubmissionEvent->DispatchSubsequents();
-			Payload->SubmissionEvent = nullptr;
+			SubmissionEvent->DispatchSubsequents();
 		}
+		Payload->SubmissionEvents.Reset();
 
 		PendingInterrupt.Enqueue(Payload);
 	}
diff --git a/Engine/Source/Runtime/VulkanRHI/Private/VulkanSubmission.h b/Engine/Source/Runtime/VulkanRHI/Private/VulkanSubmission.h
--- a/Engine/Source/Runtime/VulkanRHI/Private/VulkanSubmission.h
+++ b/Engine/Source/Runtime/VulkanRHI/Private/VulkanSubmission.h
@@ -51,7 +51,7 @@ protected:
 	TArray<VulkanRHI::FSemaphore*> SignalSemaphores; // signaled after command buffers
 
 	// Signaled when the payload has been submitted to the GPU queue
-	FGraphEventRef SubmissionEvent;
+	TArray<FGraphEventRef> SubmissionEvents;
 
 	// For internal completion tracking of the payload
 	uint64 TimelineSemaphoreValue = 0;
