VulkanRHI: Handle RHIWriteGPUFence at top of pipe.

#rnx

diff --git a/Engine/Source/Runtime/VulkanRHI/Private/VulkanCommands.cpp b/Engine/Source/Runtime/VulkanRHI/Private/VulkanCommands.cpp
--- a/Engine/Source/Runtime/VulkanRHI/Private/VulkanCommands.cpp
+++ b/Engine/Source/Runtime/VulkanRHI/Private/VulkanCommands.cpp
@@ -917,8 +917,7 @@ void FVulkanCommandListContext::RHICopyToStagingBuffer(FRHIBuffer* SourceBufferR
 
 void FVulkanCommandListContext::RHIWriteGPUFence(FRHIGPUFence* FenceRHI)
 {
-	FVulkanGPUFence* Fence = ResourceCast(FenceRHI);
-	Fence->SyncPoint = AddSyncPoint();
+	checkNoEntry(); // Should never be called
 }
 
 IRHIComputeContext* FVulkanDynamicRHI::RHIGetCommandContext(ERHIPipeline Pipeline, FRHIGPUMask GPUMask)
diff --git a/Engine/Source/Runtime/VulkanRHI/Private/VulkanContext.cpp b/Engine/Source/Runtime/VulkanRHI/Private/VulkanContext.cpp
--- a/Engine/Source/Runtime/VulkanRHI/Private/VulkanContext.cpp
+++ b/Engine/Source/Runtime/VulkanRHI/Private/VulkanContext.cpp
@@ -62,7 +62,7 @@ void FVulkanContextCommon::Finalize(TArray<FVulkanPayload*>& OutPayloads)
 
 	if (ContextSyncPoint.IsValid())
 	{
-		AddSyncPoints(MakeArrayView<FVulkanSyncPointRef>(&ContextSyncPoint, 1));
+		SignalSyncPoint(ContextSyncPoint);
 		ContextSyncPoint = nullptr;
 	}
 
diff --git a/Engine/Source/Runtime/VulkanRHI/Private/VulkanContext.h b/Engine/Source/Runtime/VulkanRHI/Private/VulkanContext.h
--- a/Engine/Source/Runtime/VulkanRHI/Private/VulkanContext.h
+++ b/Engine/Source/Runtime/VulkanRHI/Private/VulkanContext.h
@@ -139,22 +139,11 @@ public:
 	FVulkanDevice& Device;
 	FVulkanQueue& Queue;
 
-	// Returns a sync point that will be triggered at the end of the current payload's accumulated work
-	FVulkanSyncPoint* AddSyncPoint()
+	// Add a provided sync point that will be signaled at the end of the current payload's accumulated work
+	void SignalSyncPoint(const FVulkanSyncPointRef& InSync)
 	{
 		FVulkanPayload& Payload = GetPayload(EPhase::Signal);
-		if (Payload.SyncPoints.Num() == 0)
-		{
-			Payload.SyncPoints.Add(CreateVulkanSyncPoint());
-		}
-		return Payload.SyncPoints[0];
-	}
-
-	// Add a provided sync point that will be triggered at the end of the current payload's accumulated work
-	void AddSyncPoints(TArrayView<FVulkanSyncPointRef> InSyncs)
-	{
-		FVulkanPayload& Payload = GetPayload(EPhase::Signal);
-		Payload.SyncPoints.Append(InSyncs);
+		Payload.SyncPoints.Add(InSync);
 	}
 
 	// Add provided sync points that will be appended next time a payload reaches the Signal phase
diff --git a/Engine/Source/Runtime/VulkanRHI/Private/VulkanDynamicRHI.h b/Engine/Source/Runtime/VulkanRHI/Private/VulkanDynamicRHI.h
--- a/Engine/Source/Runtime/VulkanRHI/Private/VulkanDynamicRHI.h
+++ b/Engine/Source/Runtime/VulkanRHI/Private/VulkanDynamicRHI.h
@@ -142,6 +142,7 @@ public:
 	virtual FGeometryShaderRHIRef RHICreateGeometryShader(TArrayView<const uint8> Code, const FSHAHash& Hash) final override;
 	virtual FComputeShaderRHIRef RHICreateComputeShader(TArrayView<const uint8> Code, const FSHAHash& Hash) final override;
 	virtual FGPUFenceRHIRef RHICreateGPUFence(const FName &Name) final override;
+	virtual void RHIWriteGPUFence_TopOfPipe(FRHICommandListBase& RHICmdList, FRHIGPUFence* FenceRHI) final override;
 	virtual void RHICreateTransition(FRHITransition* Transition, const FRHITransitionCreateInfo& CreateInfo) final override;
 	virtual void RHIReleaseTransition(FRHITransition* Transition) final override;
 	virtual FStagingBufferRHIRef RHICreateStagingBuffer() final override;
diff --git a/Engine/Source/Runtime/VulkanRHI/Private/VulkanResources.h b/Engine/Source/Runtime/VulkanRHI/Private/VulkanResources.h
--- a/Engine/Source/Runtime/VulkanRHI/Private/VulkanResources.h
+++ b/Engine/Source/Runtime/VulkanRHI/Private/VulkanResources.h
@@ -32,6 +32,7 @@ class FVulkanViewableResource;
 class FVulkanShaderResourceView;
 class FVulkanUnorderedAccessView;
 class FVulkanRenderQuery;
+class FVulkanDynamicRHI;
 
 using FVulkanSyncPoint = FGraphEvent;
 using FVulkanSyncPointRef = TRefCountPtr<FVulkanSyncPoint>;
@@ -1079,19 +1080,19 @@ private:
 class FVulkanGPUFence final : public FRHIGPUFence
 {
 public:
-	FVulkanGPUFence(FName InName)
-		: FRHIGPUFence(InName)
-	{
-	}
+	FVulkanGPUFence(FVulkanDevice& InDevice, FName InName);
+	virtual ~FVulkanGPUFence();
 
 	virtual void Clear() override;
 	virtual bool Poll() const override;
 	virtual void Wait(FRHICommandListImmediate& RHICmdList, FRHIGPUMask GPUMask) const override;
 
 protected:
+	FVulkanDevice& Device;
 	FVulkanSyncPointRef SyncPoint;
 
 	friend class FVulkanCommandListContext;
+	friend class FVulkanDynamicRHI;
 };
 
 template<class T>
diff --git a/Engine/Source/Runtime/VulkanRHI/Private/VulkanUtil.cpp b/Engine/Source/Runtime/VulkanRHI/Private/VulkanUtil.cpp
--- a/Engine/Source/Runtime/VulkanRHI/Private/VulkanUtil.cpp
+++ b/Engine/Source/Runtime/VulkanRHI/Private/VulkanUtil.cpp
@@ -693,17 +693,7 @@ FStagingBufferRHIRef FVulkanDynamicRHI::RHICreateStagingBuffer()
 void* FVulkanDynamicRHI::RHILockStagingBuffer(FRHIStagingBuffer* StagingBufferRHI, FRHIGPUFence* FenceRHI, uint32 Offset, uint32 NumBytes)
 {
 	FVulkanStagingBuffer* StagingBuffer = ResourceCast(StagingBufferRHI);
-
-	if (FenceRHI && !FenceRHI->Poll())
-	{
-		FRHICommandListImmediate& RHICmdList = FRHICommandListImmediate::Get();
-		RHICmdList.ImmediateFlush(EImmediateFlushType::FlushRHIThread);
-
-		// SubmitCommandsAndFlushGPU might update fence state if it was tied to a previously submitted command buffer.
-		// Its state will have been updated from Submitted to NeedReset, and would assert in WaitForCmdBuffer (which is not needed in such a case)
-		FenceRHI->Wait(RHICmdList, FRHIGPUMask::All());
-	}
-
+	checkSlow(!FenceRHI || FenceRHI->Poll());
 	return StagingBuffer->Lock(Offset, NumBytes);
 }
 
@@ -713,6 +703,16 @@ void FVulkanDynamicRHI::RHIUnlockStagingBuffer(FRHIStagingBuffer* StagingBufferR
 	StagingBuffer->Unlock();
 }
 
+FVulkanGPUFence::FVulkanGPUFence(FVulkanDevice& InDevice, FName InName)
+	: FRHIGPUFence(InName)
+	, Device(InDevice)
+{
+}
+
+FVulkanGPUFence::~FVulkanGPUFence()
+{
+}
+
 void FVulkanGPUFence::Clear()
 {
 	SyncPoint = nullptr;
@@ -733,7 +733,24 @@ void FVulkanGPUFence::Wait(FRHICommandListImmediate& RHICmdList, FRHIGPUMask GPU
 
 FGPUFenceRHIRef FVulkanDynamicRHI::RHICreateGPUFence(const FName& Name)
 {
-	return new FVulkanGPUFence(Name);
+	return new FVulkanGPUFence(*Device, Name);
+}
+
+void FVulkanDynamicRHI::RHIWriteGPUFence_TopOfPipe(FRHICommandListBase& RHICmdList, FRHIGPUFence* FenceRHI)
+{
+	FVulkanGPUFence* Fence = ResourceCast(FenceRHI);
+	check(Fence);
+
+	checkf(!Fence->SyncPoint.IsValid(), TEXT("The fence for the current GPU node has already been issued."));
+	Fence->SyncPoint = CreateVulkanSyncPoint();
+
+	Fence->NumPendingWriteCommands.Increment();
+	RHICmdList.EnqueueLambda([Fence, SyncPoint = Fence->SyncPoint](FRHICommandListBase& CmdList)
+	{
+		FVulkanCommandListContext& Context = FVulkanCommandListContext::Get(CmdList);
+		Context.SignalSyncPoint(SyncPoint);
+		Fence->NumPendingWriteCommands.Decrement();
+	});
 }
 
 namespace VulkanRHI
