VulkanRHI: Reduce submit count when timelines aren't supported by using a VkEvent for FGPUFence to avoid splitting the payload for syncs.

#rnx

diff --git a/Engine/Source/Runtime/VulkanRHI/Private/VulkanResources.h b/Engine/Source/Runtime/VulkanRHI/Private/VulkanResources.h
--- a/Engine/Source/Runtime/VulkanRHI/Private/VulkanResources.h
+++ b/Engine/Source/Runtime/VulkanRHI/Private/VulkanResources.h
@@ -1089,9 +1089,10 @@ public:
 
 protected:
 	FVulkanDevice& Device;
-	FVulkanSyncPointRef SyncPoint;
+	VkEvent Event = VK_NULL_HANDLE;
+	FGraphEventRef SubmittedSyncPoint;
+	FVulkanSyncPointRef CompletedSyncPoint;
 
-	friend class FVulkanCommandListContext;
 	friend class FVulkanDynamicRHI;
 };
 
diff --git a/Engine/Source/Runtime/VulkanRHI/Private/VulkanUtil.cpp b/Engine/Source/Runtime/VulkanRHI/Private/VulkanUtil.cpp
--- a/Engine/Source/Runtime/VulkanRHI/Private/VulkanUtil.cpp
+++ b/Engine/Source/Runtime/VulkanRHI/Private/VulkanUtil.cpp
@@ -703,31 +703,62 @@ void FVulkanDynamicRHI::RHIUnlockStagingBuffer(FRHIStagingBuffer* StagingBufferR
 	StagingBuffer->Unlock();
 }
 
+extern int32 GAllowTimelineSemaphores;
+
 FVulkanGPUFence::FVulkanGPUFence(FVulkanDevice& InDevice, FName InName)
 	: FRHIGPUFence(InName)
 	, Device(InDevice)
 {
+	if (!GAllowTimelineSemaphores)
+	{
+		VkEventCreateInfo EventCreateInfo;
+		ZeroVulkanStruct(EventCreateInfo, VK_STRUCTURE_TYPE_EVENT_CREATE_INFO);
+		VERIFYVULKANRESULT(VulkanRHI::vkCreateEvent(InDevice.GetInstanceHandle(), &EventCreateInfo, VULKAN_CPU_ALLOCATOR, &Event));
+		VERIFYVULKANRESULT(VulkanRHI::vkResetEvent(InDevice.GetInstanceHandle(), Event));
+	}
 }
 
 FVulkanGPUFence::~FVulkanGPUFence()
 {
+	if (Event)
+	{
+		Device.GetDeferredDeletionQueue().EnqueueResource(VulkanRHI::FDeferredDeletionQueue2::EType::Event, Event);
+		Event = VK_NULL_HANDLE;
+	}
 }
 
 void FVulkanGPUFence::Clear()
 {
-	SyncPoint = nullptr;
+	CompletedSyncPoint = nullptr;
+
+	if (Event)
+	{
+		VERIFYVULKANRESULT(VulkanRHI::vkResetEvent(Device.GetInstanceHandle(), Event));
+		SubmittedSyncPoint = nullptr;
+	}
 }
 
 bool FVulkanGPUFence::Poll() const
 {
-	return (SyncPoint.IsValid() && SyncPoint->IsComplete());
+	if (CompletedSyncPoint.IsValid() && CompletedSyncPoint->IsComplete())
+	{
+		return true;
+	}
+
+	if (Event && SubmittedSyncPoint.IsValid() && SubmittedSyncPoint->IsComplete())
+	{
+		return (VulkanRHI::vkGetEventStatus(Device.GetInstanceHandle(), Event) == VK_EVENT_SET);
+	}
+
+	return false;
 }
 
 void FVulkanGPUFence::Wait(FRHICommandListImmediate& RHICmdList, FRHIGPUMask GPUMask) const
 {
 	if (!Poll())
 	{
-		FVulkanDynamicRHI::Get().ProcessInterruptQueueUntil(SyncPoint);
+		SCOPED_NAMED_EVENT_TEXT("FVulkanGPUFence_Wait", FColor::Turquoise);
+		FVulkanDynamicRHI::Get().ProcessInterruptQueueUntil(CompletedSyncPoint);
 	}
 }
 
@@ -741,14 +772,31 @@ void FVulkanDynamicRHI::RHIWriteGPUFence_TopOfPipe(FRHICommandListBase& RHICmdLi
 	FVulkanGPUFence* Fence = ResourceCast(FenceRHI);
 	check(Fence);
 
-	checkf(!Fence->SyncPoint.IsValid(), TEXT("The fence for the current GPU node has already been issued."));
-	Fence->SyncPoint = CreateVulkanSyncPoint();
+	checkf(!Fence->SubmittedSyncPoint.IsValid() && !Fence->CompletedSyncPoint.IsValid(), TEXT("The fence for the current GPU node has already been issued."));
+	Fence->CompletedSyncPoint = CreateVulkanSyncPoint();
+
+	if (Fence->Event)
+	{
+		Fence->SubmittedSyncPoint = FGraphEvent::CreateGraphEvent();
+	}
 
 	Fence->NumPendingWriteCommands.Increment();
-	RHICmdList.EnqueueLambda([Fence, SyncPoint = Fence->SyncPoint](FRHICommandListBase& CmdList)
+	RHICmdList.EnqueueLambda([Fence](FRHICommandListBase& CmdList)
 	{
 		FVulkanCommandListContext& Context = FVulkanCommandListContext::Get(CmdList);
-		Context.SignalSyncPoint(SyncPoint);
+
+		VulkanRHI::vkCmdSetEvent(Context.GetCommandBuffer().GetHandle(), Fence->Event, VK_PIPELINE_STAGE_BOTTOM_OF_PIPE_BIT);
+
+		if (Fence->Event)
+		{
+			Context.AddSubmissionEvent(Fence->SubmittedSyncPoint);
+			Context.AddPendingSyncPoint(Fence->CompletedSyncPoint);
+		}
+		else
+		{
+			Context.SignalSyncPoint(Fence->CompletedSyncPoint);
+		}
+
 		Fence->NumPendingWriteCommands.Decrement();
 	});
 }
