VulkanRHI: Add missing reset on ASCompactedSize query pools.

#rnx

diff --git a/Engine/Source/Runtime/VulkanRHI/Private/VulkanRayTracing.cpp b/Engine/Source/Runtime/VulkanRHI/Private/VulkanRayTracing.cpp
--- a/Engine/Source/Runtime/VulkanRHI/Private/VulkanRayTracing.cpp
+++ b/Engine/Source/Runtime/VulkanRHI/Private/VulkanRayTracing.cpp
@@ -1763,6 +1763,13 @@ const TArray<uint8>& FVulkanRayTracingPipelineState::GetShaderHandles(EShaderFre
 	return GetShaderData(Frequency).ShaderHandles;
 }
 
+FVulkanRayTracingCompactedSizeQueryPool::FVulkanRayTracingCompactedSizeQueryPool(FVulkanDevice& InDevice, uint32 InMaxQueries)
+	: FVulkanQueryPool(InDevice, InMaxQueries, EVulkanQueryPoolType::ASCompactedSize)
+{
+	QueryOutput.SetNumZeroed(InMaxQueries);
+	VulkanRHI::vkResetQueryPoolEXT(InDevice.GetInstanceHandle(), QueryPool, 0, MaxQueries);
+}
+
 void FVulkanRayTracingCompactedSizeQueryPool::EndBatch(FVulkanCommandListContext& CommandContext)
 {
 	SyncPoint = CommandContext.GetContextSyncPoint();
diff --git a/Engine/Source/Runtime/VulkanRHI/Private/VulkanRayTracing.h b/Engine/Source/Runtime/VulkanRHI/Private/VulkanRayTracing.h
--- a/Engine/Source/Runtime/VulkanRHI/Private/VulkanRayTracing.h
+++ b/Engine/Source/Runtime/VulkanRHI/Private/VulkanRayTracing.h
@@ -317,11 +317,7 @@ public:
 class FVulkanRayTracingCompactedSizeQueryPool : public FVulkanQueryPool
 {
 public:
-	FVulkanRayTracingCompactedSizeQueryPool(FVulkanDevice& InDevice, uint32 InMaxQueries)
-		: FVulkanQueryPool(InDevice, InMaxQueries, EVulkanQueryPoolType::ASCompactedSize)
-	{
-		QueryOutput.SetNumZeroed(InMaxQueries);
-	}
+	FVulkanRayTracingCompactedSizeQueryPool(FVulkanDevice& InDevice, uint32 InMaxQueries);
 
 	void EndBatch(FVulkanCommandListContext& CommandContext);
 	bool TryGetResults(uint32 NumResults);
