VulkanRHI: Change VK_NV_compute_shader_derivatives to use KHR extension.  Add detection for VK_EXT_pipeline_library and VK_EXT_graphics_pipeline_library extensions.

#rnx

diff --git a/Engine/Source/Runtime/VulkanRHI/Private/VulkanDevice.h b/Engine/Source/Runtime/VulkanRHI/Private/VulkanDevice.h
--- a/Engine/Source/Runtime/VulkanRHI/Private/VulkanDevice.h
+++ b/Engine/Source/Runtime/VulkanRHI/Private/VulkanDevice.h
@@ -51,6 +51,7 @@ struct FOptionalVulkanDeviceExtensions
 			uint64 HasAccelerationStructure : 1;
 			uint64 HasRayTracingPipeline : 1;
 			uint64 HasRayQuery : 1;
+			uint64 HasKHRPipelineLibrary : 1;
 			uint64 HasDeferredHostOperations : 1;
 			uint64 HasEXTCalibratedTimestamps : 1;
 			uint64 HasEXTDescriptorBuffer : 1;
@@ -82,7 +83,6 @@ struct FOptionalVulkanDeviceExtensions
 			uint64 HasShaderFloatControls : 1;
 			uint64 HasKHRShaderFloat16 : 1;
 			uint64 HasEXTDescriptorIndexing : 1;
-			uint64 HasEXTShaderViewportIndexLayer : 1;
 			uint64 HasSeparateDepthStencilLayouts : 1;
 			uint64 HasEXTHostQueryReset : 1;
 			uint64 HasQcomRenderPassShaderResolve : 1;
@@ -114,7 +114,7 @@ struct FOptionalVulkanDeviceExtensions
 	{
 		return 
 			HasAccelerationStructure && 
-			(HasRayTracingPipeline || HasRayQuery) &&
+			((HasRayTracingPipeline && HasKHRPipelineLibrary) || HasRayQuery) &&
 			HasEXTDescriptorIndexing &&
 			HasBufferDeviceAddress && 
 			HasDeferredHostOperations && 
@@ -144,7 +144,8 @@ struct FOptionalVulkanDeviceExtensionProperties
 	VkPhysicalDeviceFragmentDensityMap2FeaturesEXT FragmentDensityMap2Features;
 
 	VkPhysicalDeviceFragmentShaderBarycentricPropertiesKHR FragmentShaderBarycentricProps;
-	VkPhysicalDeviceComputeShaderDerivativesFeaturesNV ComputeShaderDerivativesFeatures;
+	VkPhysicalDeviceComputeShaderDerivativesFeaturesKHR ComputeShaderDerivativesFeatures;
+	VkPhysicalDeviceGraphicsPipelineLibraryPropertiesEXT GraphicsPipelineLibraryProperties;
 
 	VkPhysicalDeviceMeshShaderPropertiesEXT MeshShaderProperties;
 };
@@ -293,19 +294,24 @@ public:
 		return false;
 	}
 
+	FVulkanQueue* GetQueue(EVulkanQueueType QueueType)
+	{
+		return Queues[(int32)QueueType];
+	}
+
 	FVulkanQueue* GetGraphicsQueue()
 	{
-		return Queues[(int32)EVulkanQueueType::Graphics];
+		return GetQueue(EVulkanQueueType::Graphics);
 	}
 
 	FVulkanQueue* GetComputeQueue()
 	{
-		return Queues[(int32)EVulkanQueueType::AsyncCompute];
+		return GetQueue(EVulkanQueueType::AsyncCompute);
 	}
 
 	FVulkanQueue* GetTransferQueue()
 	{
-		return Queues[(int32)EVulkanQueueType::Transfer];
+		return GetQueue(EVulkanQueueType::Transfer);
 	}
 
 	FVulkanQueue* GetPresentQueue()
@@ -602,7 +608,7 @@ private:
 
 	uint64 TimestampValidBitsMask = 0;
 
-	FVulkanQueue* Queues[(int32)EVulkanQueueType::Count];
+	TStaticArray<FVulkanQueue*, (int32)EVulkanQueueType::Count> Queues;
 	FVulkanQueue* PresentQueue = nullptr;  // points to an existing queue
 	bool bAsyncComputeQueue = false;
 	bool bPresentOnComputeQueue = false;
diff --git a/Engine/Source/Runtime/VulkanRHI/Private/VulkanExtensions.cpp b/Engine/Source/Runtime/VulkanRHI/Private/VulkanExtensions.cpp
--- a/Engine/Source/Runtime/VulkanRHI/Private/VulkanExtensions.cpp
+++ b/Engine/Source/Runtime/VulkanRHI/Private/VulkanExtensions.cpp
@@ -346,8 +346,6 @@ public:
 
 	virtual void PostPhysicalDeviceFeatures(FOptionalVulkanDeviceExtensions& ExtensionFlags) override final
 	{
-		ExtensionFlags.HasEXTShaderViewportIndexLayer = 1; 
-
 		GRHISupportsArrayIndexFromAnyShader = true;
 	}
 };
@@ -789,28 +787,27 @@ public:
 
 
 
-// ***** VK_NV_compute_shader_derivatives
-class FVulkanNVComputeShaderDerivatives : public FVulkanDeviceExtension
+// ***** VK_KHR_compute_shader_derivatives
+class FVulkanKHRComputeShaderDerivatives : public FVulkanDeviceExtension
 {
 public:
 
-	FVulkanNVComputeShaderDerivatives(FVulkanDevice* InDevice)
-		: FVulkanDeviceExtension(InDevice, VK_NV_COMPUTE_SHADER_DERIVATIVES_EXTENSION_NAME, VULKAN_EXTENSION_ENABLED)
+	FVulkanKHRComputeShaderDerivatives(FVulkanDevice* InDevice)
+		: FVulkanDeviceExtension(InDevice, VK_KHR_COMPUTE_SHADER_DERIVATIVES_EXTENSION_NAME, VULKAN_EXTENSION_ENABLED)
 	{
 	}
 
 	virtual void PrePhysicalDeviceFeatures(VkPhysicalDeviceFeatures2KHR& PhysicalDeviceFeatures2) override final
 	{
-		VkPhysicalDeviceComputeShaderDerivativesFeaturesNV& ComputeShaderDerivativesFeatures = GetDeviceExtensionProperties().ComputeShaderDerivativesFeatures;
-		ZeroVulkanStruct(ComputeShaderDerivativesFeatures, VK_STRUCTURE_TYPE_PHYSICAL_DEVICE_COMPUTE_SHADER_DERIVATIVES_FEATURES_NV);
+		VkPhysicalDeviceComputeShaderDerivativesFeaturesKHR& ComputeShaderDerivativesFeatures = GetDeviceExtensionProperties().ComputeShaderDerivativesFeatures;
+		ZeroVulkanStruct(ComputeShaderDerivativesFeatures, VK_STRUCTURE_TYPE_PHYSICAL_DEVICE_COMPUTE_SHADER_DERIVATIVES_FEATURES_KHR);
 		AddToPNext(PhysicalDeviceFeatures2, ComputeShaderDerivativesFeatures);
 	}
 
 	virtual void PostPhysicalDeviceFeatures(FOptionalVulkanDeviceExtensions& ExtensionFlags) override final
 	{
-		VkPhysicalDeviceComputeShaderDerivativesFeaturesNV& ComputeShaderDerivativesFeatures = GetDeviceExtensionProperties().ComputeShaderDerivativesFeatures;
+		VkPhysicalDeviceComputeShaderDerivativesFeaturesKHR& ComputeShaderDerivativesFeatures = GetDeviceExtensionProperties().ComputeShaderDerivativesFeatures;
 		bRequirementsPassed = (ComputeShaderDerivativesFeatures.computeDerivativeGroupLinear == VK_TRUE);
-		ComputeShaderDerivativesFeatures.computeDerivativeGroupQuads = VK_FALSE;  // disable the unused quad mode
 		// Should be runtime guaranteed through SM6 profile
 	}
 
@@ -818,7 +815,7 @@ public:
 	{
 		if (bRequirementsPassed)
 		{
-			VkPhysicalDeviceComputeShaderDerivativesFeaturesNV& ComputeShaderDerivativesFeatures = GetDeviceExtensionProperties().ComputeShaderDerivativesFeatures;
+			VkPhysicalDeviceComputeShaderDerivativesFeaturesKHR& ComputeShaderDerivativesFeatures = GetDeviceExtensionProperties().ComputeShaderDerivativesFeatures;
 			AddToPNext(DeviceCreateInfo, ComputeShaderDerivativesFeatures);
 		}
 	}
@@ -1652,7 +1649,7 @@ private:
 	VkPhysicalDeviceSamplerYcbcrConversionFeatures PhysicalDeviceSamplerYcbcrConversionFeatures;
 };
 
-//VK_EXT_image_compression_control
+// ***** VK_EXT_image_compression_control
 class FVulkanEXTImageCompressionControlExtension : public FVulkanDeviceExtension
 {
 public:
@@ -1690,6 +1687,49 @@ private:
 	VkPhysicalDeviceImageCompressionControlFeaturesEXT PhysicalDeviceImageCompressionControlFeatures;
 };
 
+// ***** VK_KHR_graphics_pipeline_library
+class FVulkanKHRGraphicsPipelineLibraryExtension : public FVulkanDeviceExtension
+{
+public:
+
+	FVulkanKHRGraphicsPipelineLibraryExtension(FVulkanDevice* InDevice)
+		: FVulkanDeviceExtension(InDevice, VK_EXT_GRAPHICS_PIPELINE_LIBRARY_EXTENSION_NAME, VULKAN_EXTENSION_ENABLED)
+	{
+	}
+
+	virtual void PrePhysicalDeviceFeatures(VkPhysicalDeviceFeatures2KHR& PhysicalDeviceFeatures2) override final
+	{
+		ZeroVulkanStruct(GraphicsPipelineLibraryFeatures, VK_STRUCTURE_TYPE_PHYSICAL_DEVICE_GRAPHICS_PIPELINE_LIBRARY_FEATURES_EXT);
+		AddToPNext(PhysicalDeviceFeatures2, GraphicsPipelineLibraryFeatures);
+	}
+
+	virtual void PostPhysicalDeviceFeatures(FOptionalVulkanDeviceExtensions& ExtensionFlags) override final
+	{
+		bRequirementsPassed = (GraphicsPipelineLibraryFeatures.graphicsPipelineLibrary == VK_TRUE);
+	}
+
+	virtual void PrePhysicalDeviceProperties(VkPhysicalDeviceProperties2KHR& PhysicalDeviceProperties2) override final
+	{
+		if (bRequirementsPassed)
+		{
+			VkPhysicalDeviceGraphicsPipelineLibraryPropertiesEXT& GraphicsPipelineLibraryProperties = GetDeviceExtensionProperties().GraphicsPipelineLibraryProperties;
+			ZeroVulkanStruct(GraphicsPipelineLibraryProperties, VK_STRUCTURE_TYPE_PHYSICAL_DEVICE_GRAPHICS_PIPELINE_LIBRARY_PROPERTIES_EXT);
+			AddToPNext(PhysicalDeviceProperties2, GraphicsPipelineLibraryProperties);
+		}
+	}
+
+	virtual void PreCreateDevice(VkDeviceCreateInfo& DeviceCreateInfo) override final
+	{
+		if (bRequirementsPassed)
+		{
+			AddToPNext(DeviceCreateInfo, GraphicsPipelineLibraryFeatures);
+		}
+	}
+
+	VkPhysicalDeviceGraphicsPipelineLibraryFeaturesEXT GraphicsPipelineLibraryFeatures;
+};
+
+
 template <typename ExtensionType>
 static void FlagExtensionSupport(const TArray<VkExtensionProperties>& ExtensionProperties, TArray<TUniquePtr<ExtensionType>>& UEExtensions, uint32 ApiVersion, const TCHAR* ExtensionTypeName)
 {
@@ -1741,8 +1781,9 @@ FVulkanDeviceExtensionArray FVulkanDeviceExtension::GetUESupportedDeviceExtensio
 	ADD_SIMPLE_EXTENSION(VK_EXT_VALIDATION_CACHE_EXTENSION_NAME,         VULKAN_SUPPORTS_VALIDATION_CACHE,     VULKAN_EXTENSION_NOT_PROMOTED, DEVICE_EXT_FLAG_SETTER(HasEXTValidationCache));
 	ADD_SIMPLE_EXTENSION(VK_EXT_LOAD_STORE_OP_NONE_EXTENSION_NAME,       VULKAN_EXTENSION_ENABLED,             VULKAN_EXTENSION_NOT_PROMOTED, DEVICE_EXT_FLAG_SETTER(HasEXTLoadStoreOpNone));
 	ADD_SIMPLE_EXTENSION(VK_QCOM_RENDER_PASS_SHADER_RESOLVE_EXTENSION_NAME, VULKAN_SUPPORTS_QCOM_RENDERPASS_SHADER_RESOLVE, VULKAN_EXTENSION_NOT_PROMOTED, DEVICE_EXT_FLAG_SETTER(HasQcomRenderPassShaderResolve));
-	ADD_SIMPLE_EXTENSION(VK_EXT_TOOLING_INFO_EXTENSION_NAME, VULKAN_HAS_DEBUGGING_ENABLED || VULKAN_ENABLE_DRAW_MARKERS, VK_API_VERSION_1_3, DEVICE_EXT_FLAG_SETTER(HasEXTToolingInfo));
+	ADD_SIMPLE_EXTENSION(VK_EXT_TOOLING_INFO_EXTENSION_NAME, VULKAN_HAS_DEBUGGING_ENABLED || VULKAN_ENABLE_DRAW_MARKERS, VK_API_VERSION_1_3,  DEVICE_EXT_FLAG_SETTER(HasEXTToolingInfo));
 	ADD_SIMPLE_EXTENSION(VK_KHR_DYNAMIC_RENDERING_EXTENSION_NAME,        VULKAN_EXTENSION_ENABLED,             VK_API_VERSION_1_3,            nullptr); // Only enabled because maintenance5 depends on it
+	ADD_SIMPLE_EXTENSION(VK_KHR_PIPELINE_LIBRARY_EXTENSION_NAME,         VULKAN_EXTENSION_ENABLED,             VULKAN_EXTENSION_NOT_PROMOTED, DEVICE_EXT_FLAG_SETTER(HasKHRPipelineLibrary));
 
 	// Externally activated extensions (supported by the engine, but enabled externally by plugin or other) :
 	ADD_EXTERNAL_EXTENSION(VK_EXT_DEBUG_MARKER_EXTENSION_NAME, VULKAN_HAS_DEBUGGING_ENABLED || VULKAN_ENABLE_DRAW_MARKERS, VULKAN_EXTENSION_NOT_PROMOTED, nullptr);
@@ -1775,11 +1816,12 @@ FVulkanDeviceExtensionArray FVulkanDeviceExtension::GetUESupportedDeviceExtensio
 	ADD_CUSTOM_EXTENSION(FVulkanKHRDepthStencilResolveExtension);
 	ADD_CUSTOM_EXTENSION(FVulkanEXTPipelineCreationCacheControlExtension);
 	ADD_CUSTOM_EXTENSION(FVulkanKHRFragmentShaderBarycentricExtension);
-	ADD_CUSTOM_EXTENSION(FVulkanNVComputeShaderDerivatives);
+	ADD_CUSTOM_EXTENSION(FVulkanKHRComputeShaderDerivatives);
 	ADD_CUSTOM_EXTENSION(FVulkanKHRSamplerYcbcrConversionExtension);
 	ADD_CUSTOM_EXTENSION(FVulkanKHRTimelineSemaphoreExtension);
 	ADD_CUSTOM_EXTENSION(FVulkanEXTMeshShaderExtension);
 	ADD_CUSTOM_EXTENSION(FVulkanEXTImageCompressionControlExtension);
+	ADD_CUSTOM_EXTENSION(FVulkanKHRGraphicsPipelineLibraryExtension);
 
 	// Needed for Raytracing
 	ADD_CUSTOM_EXTENSION(FVulkanKHRBufferDeviceAddressExtension);
