VulkanRHI: Add FSignalFence/FWaitFence events for new profiler.

#rnx

diff --git a/Engine/Source/Runtime/VulkanRHI/Private/VulkanDynamicRHI.h b/Engine/Source/Runtime/VulkanRHI/Private/VulkanDynamicRHI.h
--- a/Engine/Source/Runtime/VulkanRHI/Private/VulkanDynamicRHI.h
+++ b/Engine/Source/Runtime/VulkanRHI/Private/VulkanDynamicRHI.h
@@ -336,7 +336,7 @@ protected:
 	void KickInterruptThread();
 
 #if RHI_NEW_GPU_PROFILER
-	struct FVulkanTimingArray : public TArray<TUniquePtr<FVulkanTiming>, TInlineAllocator<EVulkanQueueType::Count>>
+	struct FVulkanTimingArray : public TArray<TUniquePtr<FVulkanTiming>, TInlineAllocator<(int32)EVulkanQueueType::Count>>
 	{
 		FVulkanTiming* CreateNew(FVulkanQueue& Queue)
 		{
@@ -353,7 +353,7 @@ protected:
 	VulkanRHI::FSemaphore* CPUTimelineSemaphore = nullptr;  // used to wake up interrupt thread from CPU
 	std::atomic<uint64> CPUTimelineSemaphoreValue = 0;
 	TQueue<FVulkanPlatformCommandList*, EQueueMode::Mpsc> PendingPayloadsForSubmission;
-	TSet<VkSemaphore> SignaledSemaphores;
+	TMap<VkSemaphore, FBinarySemaphoreSignalInfo> SignaledSemaphores;
 	FGraphEventRef EopTask;
 
 	uint32 ApiVersion = 0;
diff --git a/Engine/Source/Runtime/VulkanRHI/Private/VulkanQueue.cpp b/Engine/Source/Runtime/VulkanRHI/Private/VulkanQueue.cpp
--- a/Engine/Source/Runtime/VulkanRHI/Private/VulkanQueue.cpp
+++ b/Engine/Source/Runtime/VulkanRHI/Private/VulkanQueue.cpp
@@ -69,10 +69,11 @@ FVulkanQueue::~FVulkanQueue()
 	}
 }
 
-int32 FVulkanQueue::SubmitQueuedPayloads(TSet<VkSemaphore>& SignaledSemas)
+int32 FVulkanQueue::SubmitQueuedPayloads(TMap<VkSemaphore, FBinarySemaphoreSignalInfo>& SignaledSemas)
 {
-	auto CanBeProcessed = [&SignaledSemas](const TArray<VulkanRHI::FSemaphore*>& WaitSemas)
+	auto CanBeProcessed = [&SignaledSemas, &Device=Device](FVulkanPayload* Payload)
 	{
+		const TArray<VulkanRHI::FSemaphore*>& WaitSemas = Payload->WaitSemaphores;
 		for (VulkanRHI::FSemaphore* WaitSema : WaitSemas)
 		{
 			const VkSemaphore SemaphoreHandle = WaitSema->GetHandle();
@@ -82,6 +83,35 @@ int32 FVulkanQueue::SubmitQueuedPayloads(TSet<VkSemaphore>& SignaledSemas)
 			}
 		}
 
+#if RHI_NEW_GPU_PROFILER
+		// Find the maximum timeline value to wait on for each queue
+		uint64 MaxTimelineWaitValue[(int32)EVulkanQueueType::Count] = {};
+		for (VulkanRHI::FSemaphore* WaitSema : WaitSemas)
+		{
+			if (!WaitSema->IsExternallySignaled())
+			{
+				const VkSemaphore SemaphoreHandle = WaitSema->GetHandle();
+				FBinarySemaphoreSignalInfo& SignalInfo = SignaledSemas[SemaphoreHandle];
+
+				const int32 QueueType = (int32)SignalInfo.QueueType;
+				MaxTimelineWaitValue[QueueType] = FMath::Max(MaxTimelineWaitValue[QueueType], SignalInfo.TimelineValue);
+			}
+		}
+
+		// Create the profiler events
+		const uint64 SubmitTime = FPlatformTime::Cycles64();
+		for (int32 QueueType = 0; QueueType < (int32)EVulkanQueueType::Count; ++QueueType)
+		{
+			if (MaxTimelineWaitValue[QueueType] > 0)
+			{
+				FVulkanQueue* SignalQueue = Device.GetQueue((EVulkanQueueType)QueueType);
+				check(SignalQueue);
+				Payload->EventStream.Emplace<UE::RHI::GPUProfiler::FEvent::FWaitFence>(
+					SubmitTime, MaxTimelineWaitValue[QueueType], SignalQueue->GetProfilerQueue());
+			}
+		}
+#endif
+
 		// We can only remove them from the list if they are all present
 		for (VulkanRHI::FSemaphore* WaitSema : WaitSemas)
 		{
@@ -102,7 +132,7 @@ int32 FVulkanQueue::SubmitQueuedPayloads(TSet<VkSemaphore>& SignaledSemas)
 	while (PendingSubmission.Peek(Payload))
 	{
 		// We can only submit the payload if all its wait semas have been signaled
-		if (!CanBeProcessed(Payload->WaitSemaphores))
+		if (!CanBeProcessed(Payload))
 		{
 			break;
 		}
@@ -119,7 +149,7 @@ int32 FVulkanQueue::SubmitQueuedPayloads(TSet<VkSemaphore>& SignaledSemas)
 	return Payloads.Num();
 }
 
-void FVulkanQueue::SubmitPayloads(TArrayView<FVulkanPayload*> Payloads, TSet<VkSemaphore>& SignaledSemas)
+void FVulkanQueue::SubmitPayloads(TArrayView<FVulkanPayload*> Payloads, TMap<VkSemaphore, FBinarySemaphoreSignalInfo>& SignaledSemas)
 {
 	TArray<VkSemaphore> SemaphoreStorage;
 	TArray<VkCommandBuffer> CommandBufferStorage;
@@ -157,6 +187,23 @@ void FVulkanQueue::SubmitPayloads(TArrayView<FVulkanPayload*> Payloads, TSet<VkS
 
 	for (FVulkanPayload* Payload : Payloads)
 	{
+#if RHI_NEW_GPU_PROFILER
+		if (Payload->EndFrameEvent.IsSet())
+		{
+			Payload->EndFrameEvent->CPUTimestamp = SubmitTime;
+			Payload->EventStream.Emplace<UE::RHI::GPUProfiler::FEvent::FFrameBoundary>(*Payload->EndFrameEvent);
+		}
+
+		if (Payload->Timing.IsSet())
+		{
+			if (FVulkanTiming* LocalTiming = *Payload->Timing)
+			{
+				SCOPED_NAMED_EVENT(CalibrateClocks, FColor::Red);
+				Device.GetCalibrationTimestamp(*LocalTiming);
+			}
+		}
+#endif // RHI_NEW_GPU_PROFILER
+
 		// Some payloads have nothing to submit because they are only used to trigger CPU events
 		if (!Payload->WaitSemaphores.Num() && !Payload->CommandBuffers.Num() && !Payload->SignalSemaphores.Num())
 		{
@@ -215,28 +262,16 @@ void FVulkanQueue::SubmitPayloads(TArrayView<FVulkanPayload*> Payloads, TSet<VkS
 			for (VulkanRHI::FSemaphore* Semaphore : Payload->SignalSemaphores)
 			{
 				SemaphoreStorage.Add(Semaphore->GetHandle());
-				SignaledSemas.Add(Semaphore->GetHandle());
+				SignaledSemas.Add(Semaphore->GetHandle(), { NextTimelineSemaphoreValue, QueueType });
 			}
 			SubmitInfo.signalSemaphoreCount = Payload->SignalSemaphores.Num();
 			SubmitInfo.pSignalSemaphores = &SemaphoreStorage[FirstSempahoreIndex];
 		}
 
 #if RHI_NEW_GPU_PROFILER
-		if (Payload->EndFrameEvent.IsSet())
-		{
-			Payload->EndFrameEvent->CPUTimestamp = SubmitTime;
-			Payload->EventStream.Emplace<UE::RHI::GPUProfiler::FEvent::FFrameBoundary>(*Payload->EndFrameEvent);
-		}
-
-		if (Payload->Timing.IsSet())
-		{
-			if (FVulkanTiming* LocalTiming = *Payload->Timing)
-			{
-				SCOPED_NAMED_EVENT(CalibrateClocks, FColor::Red);
-				Device.GetCalibrationTimestamp(*LocalTiming);
-			}
-		}
-#endif // RHI_NEW_GPU_PROFILER
+		Payload->EventStream.Emplace<UE::RHI::GPUProfiler::FEvent::FSignalFence>(
+			SubmitTime, NextTimelineSemaphoreValue);
+#endif
 
 		if (bUseTimelineSemaphores)
 		{
diff --git a/Engine/Source/Runtime/VulkanRHI/Private/VulkanQueue.h b/Engine/Source/Runtime/VulkanRHI/Private/VulkanQueue.h
--- a/Engine/Source/Runtime/VulkanRHI/Private/VulkanQueue.h
+++ b/Engine/Source/Runtime/VulkanRHI/Private/VulkanQueue.h
@@ -29,6 +29,12 @@ enum class EVulkanQueueType : uint8
 	Count,
 };
 
+struct FBinarySemaphoreSignalInfo
+{
+	uint64 TimelineValue;
+	EVulkanQueueType QueueType;
+};
+
 class FVulkanQueue
 {
 public:
@@ -75,7 +81,7 @@ public:
 		PendingSubmission.Enqueue(InPayload);
 	}
 
-	int32 SubmitQueuedPayloads(TSet<VkSemaphore>& SignaledSemas);
+	int32 SubmitQueuedPayloads(TMap<VkSemaphore, FBinarySemaphoreSignalInfo>& SignaledSemas);
 	int32 ProcessInterruptQueue(uint64 Timeout);
 
 #if RHI_NEW_GPU_PROFILER
@@ -83,7 +89,7 @@ public:
 #endif
 
 private:
-	void SubmitPayloads(TArrayView<FVulkanPayload*> Payloads, TSet<VkSemaphore>& InOutSignaledSemas);
+	void SubmitPayloads(TArrayView<FVulkanPayload*> Payloads, TMap<VkSemaphore, FBinarySemaphoreSignalInfo>& InOutSignaledSemas);
 	void Submit(TArrayView<VkSubmitInfo> InSubmitInfos, VulkanRHI::FFence* Fence);
 
 	// Used by submission pipe which holds the proper locks to access this queue
