VulkanRHI: Use pipeline libraries and pipeline cache for ray tracing pipelines.

#rnx

diff --git a/Engine/Source/Runtime/VulkanRHI/Private/VulkanDevice.h b/Engine/Source/Runtime/VulkanRHI/Private/VulkanDevice.h
--- a/Engine/Source/Runtime/VulkanRHI/Private/VulkanDevice.h
+++ b/Engine/Source/Runtime/VulkanRHI/Private/VulkanDevice.h
@@ -23,6 +23,8 @@ class FVulkanTransientHeapCache;
 class FVulkanDeviceExtension;
 class FVulkanRenderPassManager;
 class FVulkanRayTracingCompactionRequestHandler;
+class FVulkanRayTracingPipelineLibraryCache;
+
 
 #define VULKAN_USE_DEBUG_NAMES 1
 #if VULKAN_USE_DEBUG_NAMES
@@ -351,7 +353,15 @@ public:
 		return GpuSubgroupProps;
 	}
 
-	FVulkanRayTracingCompactionRequestHandler* GetRayTracingCompactionRequestHandler() { return RayTracingCompactionRequestHandler; }
+	FVulkanRayTracingCompactionRequestHandler* GetRayTracingCompactionRequestHandler()
+	{
+		return RayTracingCompactionRequestHandler; 
+	}
+
+	FVulkanRayTracingPipelineLibraryCache* GetRayTracingPipelineLibraryCache()
+	{
+		return RayTracingPipelineLibraryCache;
+	}
 
 	void InitializeRayTracing();
 	void CleanUpRayTracing();
@@ -593,6 +603,7 @@ private:
 	VkPhysicalDeviceSubgroupProperties GpuSubgroupProps;
 
 	FVulkanRayTracingCompactionRequestHandler* RayTracingCompactionRequestHandler = nullptr;
+	FVulkanRayTracingPipelineLibraryCache* RayTracingPipelineLibraryCache = nullptr;
 
 	FVulkanPhysicalDeviceFeatures PhysicalDeviceFeatures;
 
diff --git a/Engine/Source/Runtime/VulkanRHI/Private/VulkanPipeline.cpp b/Engine/Source/Runtime/VulkanRHI/Private/VulkanPipeline.cpp
--- a/Engine/Source/Runtime/VulkanRHI/Private/VulkanPipeline.cpp
+++ b/Engine/Source/Runtime/VulkanRHI/Private/VulkanPipeline.cpp
@@ -1670,6 +1670,35 @@ VkResult FVulkanPipelineStateCacheManager::CreateVKPipeline(FVulkanRHIGraphicsPi
  	return Result;
 }
 
+VkResult FVulkanPipelineStateCacheManager::CreateRayTracingPipeline(const VkRayTracingPipelineCreateInfoKHR& CreateInfo, bool bIsPartial, VkPipeline& OutPipeline)
+{
+	if (bIsPartial)
+	{
+		FScopedPipelineCache PipelineCacheShared = CurrentPrecompilingPSOCache.Get(EPipelineCacheAccess::Shared);
+		if (PipelineCacheShared.Get() != VK_NULL_HANDLE)
+		{
+			return VulkanDynamicAPI::vkCreateRayTracingPipelinesKHR(
+				Device->GetInstanceHandle(),
+				VK_NULL_HANDLE, // Deferred Operation
+				PipelineCacheShared.Get(), // Pipeline Cache 
+				1,
+				&CreateInfo,
+				VULKAN_CPU_ALLOCATOR,
+				&OutPipeline);
+		}
+	}
+
+	FScopedPipelineCache PipelineCacheShared = GlobalPSOCache.Get(EPipelineCacheAccess::Shared);
+	return VulkanDynamicAPI::vkCreateRayTracingPipelinesKHR(
+		Device->GetInstanceHandle(),
+		VK_NULL_HANDLE, // Deferred Operation
+		PipelineCacheShared.Get(), // Pipeline Cache 
+		1,
+		&CreateInfo,
+		VULKAN_CPU_ALLOCATOR,
+		&OutPipeline);
+}
+
 void FVulkanPipelineStateCacheManager::DestroyCache()
 {
 	VkDevice DeviceHandle = Device->GetInstanceHandle();
diff --git a/Engine/Source/Runtime/VulkanRHI/Private/VulkanPipeline.h b/Engine/Source/Runtime/VulkanRHI/Private/VulkanPipeline.h
--- a/Engine/Source/Runtime/VulkanRHI/Private/VulkanPipeline.h
+++ b/Engine/Source/Runtime/VulkanRHI/Private/VulkanPipeline.h
@@ -503,6 +503,8 @@ public:
 	FVulkanComputePipeline* GetOrCreateComputePipeline(FVulkanComputeShader* ComputeShader);
 	void NotifyDeletedComputePipeline(FVulkanComputePipeline* Pipeline);
 
+	VkResult CreateRayTracingPipeline(const VkRayTracingPipelineCreateInfoKHR& CreateInfo, bool bIsPartial, VkPipeline& OutPipeline);
+
 private:
 	class FPipelineCache;
 
diff --git a/Engine/Source/Runtime/VulkanRHI/Private/VulkanRayTracing.cpp b/Engine/Source/Runtime/VulkanRHI/Private/VulkanRayTracing.cpp
--- a/Engine/Source/Runtime/VulkanRHI/Private/VulkanRayTracing.cpp
+++ b/Engine/Source/Runtime/VulkanRHI/Private/VulkanRayTracing.cpp
@@ -33,17 +33,6 @@ static FAutoConsoleVariableRef CVarVulkanRayTracingCompactionMinPrimitiveCount(
 	TEXT("Sets the minimum primitive count threshold below which geometry skips the compaction. (default = 128)\n")
 );
 
-static int32 GVulkanRayTracingAllowDeferredOperation = 0;
-static FAutoConsoleVariableRef CVarVulkanRayTracingAllowDeferredOperation(
-	TEXT("r.Vulkan.RayTracing.AllowDeferredOperation"),
-	GVulkanRayTracingAllowDeferredOperation,
-	TEXT("Whether to use Vulkan Deferred Operation for RT pipeline creation. (default = -1)\n")
-	TEXT(" <0: Disabled\n")
-	TEXT(" 0: Enabled, auto detect the maximum number of threads")
-	TEXT(" >0: Enabled, use the specified number of threads"),
-	ECVF_ReadOnly
-);
-
 static int32 GVulkanRayTracingMaxShaderGroupStride = 4096;
 static FAutoConsoleVariableRef GCVarVulkanRayTracingMaxShaderGroupStride(
 	TEXT("r.Vulkan.RayTracing.MaxShaderGroupStride"),
@@ -1376,6 +1365,186 @@ static FRHIRayTracingShader* GetBuiltInRayTracingShader()
 	return static_cast<FRHIRayTracingShader*>(Shader.GetRayTracingShader());
 }
 
+
+
+class FVulkanRayTracingPipelineLibraryCache
+{
+public:
+
+	static uint64 GetShaderHash64(FRHIRayTracingShader* ShaderRHI)
+	{
+		uint64 ShaderHash; // 64 bits from the shader SHA1
+		FMemory::Memcpy(&ShaderHash, ShaderRHI->GetHash().Hash, sizeof(ShaderHash));
+		return ShaderHash;
+	}
+
+	struct FKey
+	{
+		FKey(FVulkanRayTracingShader* InShader, uint32 InMaxAttributeSizeInBytes, uint32 InMaxPayloadSizeInBytes)
+			: ShaderHash(GetShaderHash64(InShader))
+			, ShaderHandle(InShader->GetOrCreateHandle(FVulkanRayTracingShader::MainModuleIdentifier)->GetVkShaderModule())
+			, MaxAttributeSizeInBytes(InMaxAttributeSizeInBytes)
+			, MaxPayloadSizeInBytes(InMaxPayloadSizeInBytes)
+		{
+		}
+
+		const uint64 ShaderHash;
+		const VkShaderModule ShaderHandle;
+		const uint32 MaxAttributeSizeInBytes;
+		const uint32 MaxPayloadSizeInBytes;
+
+		bool operator == (const FKey& Other) const
+		{
+			return ShaderHash == Other.ShaderHash
+				&& ShaderHandle == Other.ShaderHandle
+				&& MaxAttributeSizeInBytes == Other.MaxAttributeSizeInBytes
+				&& MaxPayloadSizeInBytes == Other.MaxPayloadSizeInBytes;
+		}
+
+		inline friend uint32 GetTypeHash(const FKey& Key)
+		{
+			return Key.ShaderHash;
+		}
+	};
+
+	struct FPipelineLibrary
+	{
+		VkPipeline PipelineHandle = VK_NULL_HANDLE;
+		FGraphEventRef GraphEvent;
+	};
+
+	FVulkanRayTracingPipelineLibraryCache(FVulkanDevice& InDevice)
+		: Device(InDevice)
+	{
+	}
+
+	~FVulkanRayTracingPipelineLibraryCache()
+	{
+		for (auto& Pair : PipelineLibraryMap)
+		{
+			VulkanRHI::vkDestroyPipeline(Device.GetInstanceHandle(), Pair.Value->PipelineHandle, VULKAN_CPU_ALLOCATOR);
+			Pair.Value->PipelineHandle = VK_NULL_HANDLE;
+		}
+		PipelineLibraryMap.Empty();
+	}
+
+	FPipelineLibrary* GetOrAddLibrary(FVulkanRayTracingShader* Shader, const FRayTracingPipelineStateInitializer& Initializer)
+	{
+		FScopeLock Lock(&PipelineLibraryMapCS);
+
+		const FKey Key(Shader, Initializer.MaxAttributeSizeInBytes, Initializer.MaxPayloadSizeInBytes);
+		FPipelineLibrary** ExistingPipeline = PipelineLibraryMap.Find(Key);
+		if (ExistingPipeline)
+		{
+			return *ExistingPipeline;
+		}
+
+		FPipelineLibrary* PipelineLibrary = new FPipelineLibrary();
+
+		auto CreatePipelineLibraryTask = [this, PipelineLibrary, Shader, Key, bPartial=Initializer.bPartial]()
+			{
+				ANSICHAR EntryPoint[MaxEntryPointNameLength];
+				Shader->GetEntryPoint(EntryPoint, MaxEntryPointNameLength);
+
+				uint32 ShaderCount = 0;
+				TStaticArray<VkPipelineShaderStageCreateInfo, MaxHitGroupShaderCount> ShaderStages;
+				ZeroVulkanStruct(ShaderStages[ShaderCount], VK_STRUCTURE_TYPE_PIPELINE_SHADER_STAGE_CREATE_INFO);
+				ShaderStages[ShaderCount].module = Key.ShaderHandle;
+				ShaderStages[ShaderCount].stage = UEFrequencyToVKStageBit(Shader->GetFrequency());  // will default to closestHitShader for SF_RayHitGroup
+				ShaderStages[ShaderCount].pName = EntryPoint;
+
+				VkRayTracingShaderGroupCreateInfoKHR ShaderGroup;
+				ZeroVulkanStruct(ShaderGroup, VK_STRUCTURE_TYPE_RAY_TRACING_SHADER_GROUP_CREATE_INFO_KHR);
+
+				if (Shader->GetFrequency() != SF_RayHitGroup)
+				{
+					ShaderGroup.type = VK_RAY_TRACING_SHADER_GROUP_TYPE_GENERAL_KHR;
+					ShaderGroup.generalShader = ShaderCount++;
+					ShaderGroup.closestHitShader = VK_SHADER_UNUSED_KHR;
+					ShaderGroup.anyHitShader = VK_SHADER_UNUSED_KHR;
+					ShaderGroup.intersectionShader = VK_SHADER_UNUSED_KHR;
+				}
+				else
+				{
+					ShaderGroup.type = VK_RAY_TRACING_SHADER_GROUP_TYPE_TRIANGLES_HIT_GROUP_KHR;
+					ShaderGroup.generalShader = VK_SHADER_UNUSED_KHR;
+
+					// Closest Hit, always present
+					ShaderGroup.closestHitShader = ShaderCount++;
+
+					// Any Hit, optional
+					if (Shader->GetCodeHeader().RayGroupAnyHit != FVulkanShaderHeader::ERayHitGroupEntrypoint::NotPresent)
+					{
+						ZeroVulkanStruct(ShaderStages[ShaderCount], VK_STRUCTURE_TYPE_PIPELINE_SHADER_STAGE_CREATE_INFO);
+						ShaderStages[ShaderCount].module = Shader->GetOrCreateHandle(FVulkanRayTracingShader::AnyHitModuleIdentifier)->GetVkShaderModule();
+						ShaderStages[ShaderCount].stage = VK_SHADER_STAGE_ANY_HIT_BIT_KHR;
+						ShaderStages[ShaderCount].pName = "main_00000000_00000000"; // :todo-jn: patch in the size_crc
+						ShaderGroup.anyHitShader = ShaderCount++;
+					}
+					else
+					{
+						ShaderGroup.anyHitShader = VK_SHADER_UNUSED_KHR;
+					}
+
+					// Intersection, optional
+					if (Shader->GetCodeHeader().RayGroupIntersection != FVulkanShaderHeader::ERayHitGroupEntrypoint::NotPresent)
+					{
+						ZeroVulkanStruct(ShaderStages[ShaderCount], VK_STRUCTURE_TYPE_PIPELINE_SHADER_STAGE_CREATE_INFO);
+						ShaderStages[ShaderCount].module = Shader->GetOrCreateHandle(FVulkanRayTracingShader::IntersectionModuleIdentifier)->GetVkShaderModule();
+						ShaderStages[ShaderCount].stage = VK_SHADER_STAGE_INTERSECTION_BIT_KHR;
+						ShaderStages[ShaderCount].pName = "main_00000000_00000000"; // :todo-jn: patch in the size_crc
+						ShaderGroup.intersectionShader = ShaderCount++;
+
+						// Switch the shader group type given the presence of an intersection shader
+						ShaderGroup.type = VK_RAY_TRACING_SHADER_GROUP_TYPE_PROCEDURAL_HIT_GROUP_KHR;
+					}
+					else
+					{
+						ShaderGroup.intersectionShader = VK_SHADER_UNUSED_KHR;
+					}
+				}
+
+				VkRayTracingPipelineInterfaceCreateInfoKHR RayTracingPipelineInterfaceCreateInfo;
+				ZeroVulkanStruct(RayTracingPipelineInterfaceCreateInfo, VK_STRUCTURE_TYPE_RAY_TRACING_PIPELINE_INTERFACE_CREATE_INFO_KHR);
+				RayTracingPipelineInterfaceCreateInfo.maxPipelineRayHitAttributeSize = Key.MaxAttributeSizeInBytes;
+				RayTracingPipelineInterfaceCreateInfo.maxPipelineRayPayloadSize = Key.MaxPayloadSizeInBytes;
+
+				VkRayTracingPipelineCreateInfoKHR RayTracingPipelineCreateInfo;
+				ZeroVulkanStruct(RayTracingPipelineCreateInfo, VK_STRUCTURE_TYPE_RAY_TRACING_PIPELINE_CREATE_INFO_KHR);
+				RayTracingPipelineCreateInfo.pLibraryInterface = &RayTracingPipelineInterfaceCreateInfo;
+				RayTracingPipelineCreateInfo.stageCount = ShaderCount;
+				RayTracingPipelineCreateInfo.pStages = ShaderStages.GetData();
+				RayTracingPipelineCreateInfo.groupCount = 1;
+				RayTracingPipelineCreateInfo.pGroups = &ShaderGroup;
+				RayTracingPipelineCreateInfo.maxPipelineRayRecursionDepth = 1;
+				RayTracingPipelineCreateInfo.layout = Device.GetBindlessDescriptorManager()->GetPipelineLayout();
+				RayTracingPipelineCreateInfo.flags = VK_PIPELINE_CREATE_LIBRARY_BIT_KHR | VK_PIPELINE_CREATE_DESCRIPTOR_BUFFER_BIT_EXT;
+
+				const VkResult Result = Device.GetPipelineStateCache()->CreateRayTracingPipeline(RayTracingPipelineCreateInfo, bPartial, PipelineLibrary->PipelineHandle);
+				if (!bPartial)
+				{
+					VERIFYVULKANRESULT_EXPANDED(Result);
+				}
+			};
+
+		PipelineLibrary->GraphEvent = FFunctionGraphTask::CreateAndDispatchWhenReady(MoveTemp(CreatePipelineLibraryTask), QUICK_USE_CYCLE_STAT(FExecuteRHIThreadTask, STATGROUP_TaskGraphTasks));
+		PipelineLibraryMap.Add(Key, PipelineLibrary);
+		return PipelineLibrary;
+	}
+
+private:
+	static const uint32 MaxEntryPointNameLength = 24;
+	static const uint32 MaxHitGroupShaderCount = 3;
+
+	FVulkanDevice& Device;
+
+	FCriticalSection PipelineLibraryMapCS;
+	TMap<FKey, FPipelineLibrary*> PipelineLibraryMap;
+};
+
+
+
+
 void FVulkanDevice::InitializeRayTracing()
 {
 	if (GRHISupportsRayTracingShaders)
@@ -1386,6 +1555,8 @@ void FVulkanDevice::InitializeRayTracing()
 			UE_LOG(LogRHI, Warning, TEXT("Specified value for r.Vulkan.RayTracing.MaxShaderGroupStride is too large for this device! It will be capped."));
 		}
 		GVulkanRayTracingMaxShaderGroupStride = FMath::Min<VkDeviceSize>(RayTracingPipelineProps.maxShaderGroupStride, GVulkanRayTracingMaxShaderGroupStride);
+
+		RayTracingPipelineLibraryCache = new FVulkanRayTracingPipelineLibraryCache(*this);
 	}
 }
 
@@ -1422,245 +1593,100 @@ static FVulkanUniformBuffer* GetDummyUB(FVulkanDevice* Device, uint32 UBLayoutHa
 
 void FVulkanDevice::CleanUpRayTracing()
 {
+	if (RayTracingPipelineLibraryCache)
+	{
+		delete RayTracingPipelineLibraryCache;
+		RayTracingPipelineLibraryCache = nullptr;
+	}
+
 	DummyUBs.Empty();
 }
 
 
 FVulkanRayTracingPipelineState::FVulkanRayTracingPipelineState(FVulkanDevice* const InDevice, const FRayTracingPipelineStateInitializer& Initializer)
 	: FRHIRayTracingPipelineState(Initializer), FDeviceChild(InDevice)
+	, bIsPartialPipeline(Initializer.bPartial)
 {
 	checkf(InDevice->SupportsBindless(), TEXT("Vulkan ray tracing pipelines are only supported in bindless."));
+	checkf(Initializer.MaxAttributeSizeInBytes <= InDevice->GetOptionalExtensionProperties().RayTracingPipelineProps.maxRayHitAttributeSize, 
+		TEXT("Required attribute size (%u) too large for current device (%u)."), 
+		Initializer.MaxAttributeSizeInBytes, 
+		InDevice->GetOptionalExtensionProperties().RayTracingPipelineProps.maxRayHitAttributeSize);
 
 	TArrayView<FRHIRayTracingShader*> InitializerRayGenShaders = Initializer.GetRayGenTable();
 	TArrayView<FRHIRayTracingShader*> InitializerMissShaders = Initializer.GetMissTable();
 	TArrayView<FRHIRayTracingShader*> InitializerHitGroupShaders = Initializer.GetHitGroupTable();
 	TArrayView<FRHIRayTracingShader*> InitializerCallableShaders = Initializer.GetCallableTable();
 
-	TArray<VkPipelineShaderStageCreateInfo> ShaderStages;
-	TArray<VkRayTracingShaderGroupCreateInfoKHR> ShaderGroups;
-	TArray<ANSICHAR*> EntryPointNames;
-	const uint32 EntryPointNameMaxLength = 24;
-
-	RayGen.Shaders.Reserve(InitializerRayGenShaders.Num());
-	for (FRHIRayTracingShader* const RayGenShaderRHI : InitializerRayGenShaders)
-	{
-		checkSlow(RayGenShaderRHI->GetFrequency() == SF_RayGen);
-		FVulkanRayTracingShader* const RayGenShader = ResourceCast(RayGenShaderRHI);
-
-		VkPipelineShaderStageCreateInfo ShaderStage;
-		ZeroVulkanStruct(ShaderStage, VK_STRUCTURE_TYPE_PIPELINE_SHADER_STAGE_CREATE_INFO);
-		ShaderStage.module = RayGenShader->GetOrCreateHandle(FVulkanRayTracingShader::MainModuleIdentifier)->GetVkShaderModule();
-		ShaderStage.stage = VK_SHADER_STAGE_RAYGEN_BIT_KHR;
-			
-		ANSICHAR* const EntryPoint = new ANSICHAR[EntryPointNameMaxLength];
-		RayGenShader->GetEntryPoint(EntryPoint, EntryPointNameMaxLength);
-		EntryPointNames.Add(EntryPoint);
-		ShaderStage.pName = EntryPoint;
-		ShaderStages.Add(ShaderStage);
-
-		VkRayTracingShaderGroupCreateInfoKHR ShaderGroup;
-		ZeroVulkanStruct(ShaderGroup, VK_STRUCTURE_TYPE_RAY_TRACING_SHADER_GROUP_CREATE_INFO_KHR);
-		ShaderGroup.type = VK_RAY_TRACING_SHADER_GROUP_TYPE_GENERAL_KHR;
-		ShaderGroup.generalShader = ShaderStages.Num() - 1;
-		ShaderGroup.closestHitShader = VK_SHADER_UNUSED_KHR;
-		ShaderGroup.anyHitShader = VK_SHADER_UNUSED_KHR;
-		ShaderGroup.intersectionShader = VK_SHADER_UNUSED_KHR;
-		ShaderGroups.Add(ShaderGroup);
-
-		RayGen.Shaders.Add(RayGenShader);
-	}
-
-	Miss.Shaders.Reserve(InitializerMissShaders.Num());
-	for (FRHIRayTracingShader* const MissShaderRHI : InitializerMissShaders)
-	{
-		checkSlow(MissShaderRHI->GetFrequency() == SF_RayMiss);
-		FVulkanRayTracingShader* const MissShader = ResourceCast(MissShaderRHI);
-
-		VkPipelineShaderStageCreateInfo ShaderStage;
-		ZeroVulkanStruct(ShaderStage, VK_STRUCTURE_TYPE_PIPELINE_SHADER_STAGE_CREATE_INFO);
-		ShaderStage.module = MissShader->GetOrCreateHandle(FVulkanRayTracingShader::MainModuleIdentifier)->GetVkShaderModule();
-		ShaderStage.stage = VK_SHADER_STAGE_MISS_BIT_KHR;
-
-		ANSICHAR* const EntryPoint = new char[EntryPointNameMaxLength];
-		MissShader->GetEntryPoint(EntryPoint, EntryPointNameMaxLength);
-		EntryPointNames.Add(EntryPoint);
-		ShaderStage.pName = EntryPoint;
-		ShaderStages.Add(ShaderStage);
-
-		VkRayTracingShaderGroupCreateInfoKHR ShaderGroup;
-		ZeroVulkanStruct(ShaderGroup, VK_STRUCTURE_TYPE_RAY_TRACING_SHADER_GROUP_CREATE_INFO_KHR);
-		ShaderGroup.type = VK_RAY_TRACING_SHADER_GROUP_TYPE_GENERAL_KHR;
-		ShaderGroup.generalShader = ShaderStages.Num() - 1;
-		ShaderGroup.closestHitShader = VK_SHADER_UNUSED_KHR;
-		ShaderGroup.anyHitShader = VK_SHADER_UNUSED_KHR;
-		ShaderGroup.intersectionShader = VK_SHADER_UNUSED_KHR;
-		ShaderGroups.Add(ShaderGroup);
-
-		Miss.Shaders.Add(MissShader);
-	}
-
-	HitGroup.Shaders.Reserve(InitializerHitGroupShaders.Num());
-	for (FRHIRayTracingShader* const HitGroupShaderRHI : InitializerHitGroupShaders)
-	{
-		VkRayTracingShaderGroupCreateInfoKHR ShaderGroup;
-		ZeroVulkanStruct(ShaderGroup, VK_STRUCTURE_TYPE_RAY_TRACING_SHADER_GROUP_CREATE_INFO_KHR);
-		ShaderGroup.type = VK_RAY_TRACING_SHADER_GROUP_TYPE_TRIANGLES_HIT_GROUP_KHR;
-		ShaderGroup.generalShader = VK_SHADER_UNUSED_KHR;
-
-		checkSlow(HitGroupShaderRHI->GetFrequency() == SF_RayHitGroup);
-		FVulkanRayTracingShader* const HitGroupShader = ResourceCast(HitGroupShaderRHI);
-
-		// Closest Hit, always present
-		{
-			ANSICHAR* const EntryPoint = new char[EntryPointNameMaxLength];
-			HitGroupShader->GetEntryPoint(EntryPoint, EntryPointNameMaxLength);
-			EntryPointNames.Add(EntryPoint);
-
-			VkPipelineShaderStageCreateInfo ShaderStage;
-			ZeroVulkanStruct(ShaderStage, VK_STRUCTURE_TYPE_PIPELINE_SHADER_STAGE_CREATE_INFO);
-			ShaderStage.module = HitGroupShader->GetOrCreateHandle(FVulkanRayTracingShader::ClosestHitModuleIdentifier)->GetVkShaderModule();
-			ShaderStage.stage = VK_SHADER_STAGE_CLOSEST_HIT_BIT_KHR;
-			ShaderStage.pName = EntryPoint;
-			ShaderGroup.closestHitShader = ShaderStages.Add(ShaderStage);
-		}
+	FVulkanRayTracingPipelineLibraryCache* LibraryCache = InDevice->GetRayTracingPipelineLibraryCache();
 
-		// Any Hit, optional
-		if (HitGroupShader->GetCodeHeader().RayGroupAnyHit != FVulkanShaderHeader::ERayHitGroupEntrypoint::NotPresent)
-		{
-			VkPipelineShaderStageCreateInfo ShaderStage;
-			ZeroVulkanStruct(ShaderStage, VK_STRUCTURE_TYPE_PIPELINE_SHADER_STAGE_CREATE_INFO);
-			ShaderStage.module = HitGroupShader->GetOrCreateHandle(FVulkanRayTracingShader::AnyHitModuleIdentifier)->GetVkShaderModule();
-			ShaderStage.stage = VK_SHADER_STAGE_ANY_HIT_BIT_KHR;
-			ShaderStage.pName = "main_00000000_00000000"; // :todo-jn: patch in the size_crc
-			ShaderGroup.anyHitShader = ShaderStages.Add(ShaderStage);
-		}
-		else
-		{
-			ShaderGroup.anyHitShader = VK_SHADER_UNUSED_KHR;
-		}
+	TArray<FVulkanRayTracingPipelineLibraryCache::FPipelineLibrary*> PipelineLibraries;
+	FGraphEventArray PendingLibraryTasks;
 
-		// Intersection, optional
-		if (HitGroupShader->GetCodeHeader().RayGroupIntersection != FVulkanShaderHeader::ERayHitGroupEntrypoint::NotPresent)
-		{
-			VkPipelineShaderStageCreateInfo ShaderStage;
-			ZeroVulkanStruct(ShaderStage, VK_STRUCTURE_TYPE_PIPELINE_SHADER_STAGE_CREATE_INFO);
-			ShaderStage.module = HitGroupShader->GetOrCreateHandle(FVulkanRayTracingShader::IntersectionModuleIdentifier)->GetVkShaderModule();
-			ShaderStage.stage = VK_SHADER_STAGE_INTERSECTION_BIT_KHR;
-			ShaderStage.pName = "main_00000000_00000000"; // :todo-jn: patch in the size_crc
-			ShaderGroup.intersectionShader = ShaderStages.Add(ShaderStage);
-
-			// Switch the shader group type given the presence of an intersection shader
-			ShaderGroup.type = VK_RAY_TRACING_SHADER_GROUP_TYPE_PROCEDURAL_HIT_GROUP_KHR;
-		}
-		else
+	const uint32 NumShaders = InitializerRayGenShaders.Num() + InitializerMissShaders.Num() + InitializerHitGroupShaders.Num() + InitializerCallableShaders.Num();
+	PipelineLibraries.Reserve(NumShaders);
+	PendingLibraryTasks.Reserve(NumShaders);
+
+	auto ProcessShaderArray = [&PipelineLibraries , &PendingLibraryTasks, LibraryCache, &Initializer](TArrayView<FRHIRayTracingShader*> InitializerShaders, TArray<TRefCountPtr<FVulkanRayTracingShader>>& OutShaders)
 		{
-			ShaderGroup.intersectionShader = VK_SHADER_UNUSED_KHR;
-		}
+			OutShaders.Reserve(InitializerShaders.Num());
+			for (FRHIRayTracingShader* const ShaderRHI : InitializerShaders)
+			{
+				FVulkanRayTracingShader* const Shader = ResourceCast(ShaderRHI);
+				FVulkanRayTracingPipelineLibraryCache::FPipelineLibrary* PipelineLibrary = LibraryCache->GetOrAddLibrary(Shader, Initializer);
+				if (PipelineLibrary->GraphEvent.IsValid() && !PipelineLibrary->GraphEvent->IsComplete())
+				{
+					PendingLibraryTasks.Add(PipelineLibrary->GraphEvent);
+				}
+				
+				PipelineLibraries.Add(PipelineLibrary);
+				OutShaders.Add(Shader);
+			}
+		};
 
-		ShaderGroups.Add(ShaderGroup);
+	ProcessShaderArray(InitializerRayGenShaders, RayGen.Shaders);
+	ProcessShaderArray(InitializerMissShaders, Miss.Shaders);
+	ProcessShaderArray(InitializerHitGroupShaders, HitGroup.Shaders);
+	ProcessShaderArray(InitializerCallableShaders, Callable.Shaders);
+	INC_DWORD_STAT_BY(STAT_VulkanRayTracingCompiledShaders, NumShaders);
 
-		HitGroup.Shaders.Add(HitGroupShader);
-	}
-
-	Callable.Shaders.Reserve(InitializerCallableShaders.Num());
-	for (FRHIRayTracingShader* const CallableShaderRHI : InitializerCallableShaders)
+	// No need to continue, partial pipelines are only used to kick off early shader compilation (they can't be used for rendering)
+	if (bIsPartialPipeline)
 	{
-		checkSlow(CallableShaderRHI->GetFrequency() == SF_RayCallable);
-		FVulkanRayTracingShader* const CallableShader = ResourceCast(CallableShaderRHI);
+		return;
+	}
 
-		VkPipelineShaderStageCreateInfo ShaderStage;
-		ZeroVulkanStruct(ShaderStage, VK_STRUCTURE_TYPE_PIPELINE_SHADER_STAGE_CREATE_INFO);
-		ShaderStage.module = CallableShader->GetOrCreateHandle(FVulkanRayTracingShader::MainModuleIdentifier)->GetVkShaderModule();
-		ShaderStage.stage = VK_SHADER_STAGE_CALLABLE_BIT_KHR;
+	// Wait for all the libraries to be ready before fetching the handles
+	FTaskGraphInterface::Get().WaitUntilTasksComplete(PendingLibraryTasks);
 
-		ANSICHAR* const EntryPoint = new char[EntryPointNameMaxLength];
-		CallableShader->GetEntryPoint(EntryPoint, EntryPointNameMaxLength);
-		EntryPointNames.Add(EntryPoint);
-		ShaderStage.pName = EntryPoint;
-		ShaderStages.Add(ShaderStage);
+	// Pull out the pipeline library handles to use them to create the pipeline
+	TArray<VkPipeline> Pipelines;
+	for (FVulkanRayTracingPipelineLibraryCache::FPipelineLibrary* PipelineLibrary : PipelineLibraries)
+	{
+		checkSlow(!PipelineLibrary->GraphEvent.IsValid() || PipelineLibrary->GraphEvent->IsComplete());
+		checkf(PipelineLibrary->PipelineHandle, TEXT("Invalid pipeline library handle while building ray tracing pipeline."));
+		Pipelines.Add(PipelineLibrary->PipelineHandle);
+	}
 
-		VkRayTracingShaderGroupCreateInfoKHR ShaderGroup;
-		ZeroVulkanStruct(ShaderGroup, VK_STRUCTURE_TYPE_RAY_TRACING_SHADER_GROUP_CREATE_INFO_KHR);
-		ShaderGroup.type = VK_RAY_TRACING_SHADER_GROUP_TYPE_GENERAL_KHR;
-		ShaderGroup.generalShader = ShaderStages.Num() - 1;
-		ShaderGroup.closestHitShader = VK_SHADER_UNUSED_KHR;
-		ShaderGroup.anyHitShader = VK_SHADER_UNUSED_KHR;
-		ShaderGroup.intersectionShader = VK_SHADER_UNUSED_KHR;
-		ShaderGroups.Add(ShaderGroup);
+	VkRayTracingPipelineInterfaceCreateInfoKHR RayTracingPipelineInterfaceCreateInfo;
+	ZeroVulkanStruct(RayTracingPipelineInterfaceCreateInfo, VK_STRUCTURE_TYPE_RAY_TRACING_PIPELINE_INTERFACE_CREATE_INFO_KHR);
+	RayTracingPipelineInterfaceCreateInfo.maxPipelineRayPayloadSize = Initializer.MaxPayloadSizeInBytes;
+	RayTracingPipelineInterfaceCreateInfo.maxPipelineRayHitAttributeSize = Initializer.MaxAttributeSizeInBytes;
 
-		Callable.Shaders.Add(CallableShader);
-	}
+	VkPipelineLibraryCreateInfoKHR PipelineLibraryCreateInfo;
+	ZeroVulkanStruct(PipelineLibraryCreateInfo, VK_STRUCTURE_TYPE_PIPELINE_LIBRARY_CREATE_INFO_KHR);
+	PipelineLibraryCreateInfo.libraryCount = Pipelines.Num();
+	PipelineLibraryCreateInfo.pLibraries = Pipelines.GetData();
 
 	VkRayTracingPipelineCreateInfoKHR RayTracingPipelineCreateInfo;
 	ZeroVulkanStruct(RayTracingPipelineCreateInfo, VK_STRUCTURE_TYPE_RAY_TRACING_PIPELINE_CREATE_INFO_KHR);
-	RayTracingPipelineCreateInfo.stageCount = ShaderStages.Num();
-	RayTracingPipelineCreateInfo.pStages = ShaderStages.GetData();
-	RayTracingPipelineCreateInfo.groupCount = ShaderGroups.Num();
-	RayTracingPipelineCreateInfo.pGroups = ShaderGroups.GetData();
+	RayTracingPipelineCreateInfo.pLibraryInfo = &PipelineLibraryCreateInfo;
+	RayTracingPipelineCreateInfo.pLibraryInterface = &RayTracingPipelineInterfaceCreateInfo;
 	RayTracingPipelineCreateInfo.maxPipelineRayRecursionDepth = 1;
 	RayTracingPipelineCreateInfo.layout = InDevice->GetBindlessDescriptorManager()->GetPipelineLayout();
 	RayTracingPipelineCreateInfo.flags = VK_PIPELINE_CREATE_DESCRIPTOR_BUFFER_BIT_EXT;
 
-	VkDeferredOperationKHR DeferredOp = VK_NULL_HANDLE; // :todo-jn: more speed
-	if (GVulkanRayTracingAllowDeferredOperation >= 0)
-	{
-		VERIFYVULKANRESULT(VulkanRHI::vkCreateDeferredOperationKHR(
-			InDevice->GetInstanceHandle(),
-			VULKAN_CPU_ALLOCATOR,
-			&DeferredOp));
-	}
-
-	VERIFYVULKANRESULT_EXPANDED(VulkanDynamicAPI::vkCreateRayTracingPipelinesKHR(
-		InDevice->GetInstanceHandle(),
-		DeferredOp,
-		VK_NULL_HANDLE, // Pipeline Cache 
-		1,
-		&RayTracingPipelineCreateInfo,
-		VULKAN_CPU_ALLOCATOR,
-		&Pipeline));
-
-	if (DeferredOp != VK_NULL_HANDLE)
-	{
-		int32 MaxConcurrency = FMath::Min(
-			(int32)VulkanRHI::vkGetDeferredOperationMaxConcurrencyKHR(InDevice->GetInstanceHandle(), DeferredOp),
-			FTaskGraphInterface::Get().GetNumWorkerThreads());
-
-		if (GVulkanRayTracingAllowDeferredOperation > 0)
-		{
-			MaxConcurrency = FMath::Min(MaxConcurrency, GVulkanRayTracingAllowDeferredOperation);
-		}
-		else if (MaxConcurrency > 3)
-		{
-			// Prevent automatic detection from completely clogging the machine
-			MaxConcurrency -= 2;
-		}
-
-		bool bCompleted = false;
-		ParallelFor(MaxConcurrency, [DeferredOp, InDevice, &bCompleted](int32 Unused)
-			{
-				VkResult Result = VulkanRHI::vkDeferredOperationJoinKHR(InDevice->GetInstanceHandle(), DeferredOp);
-				while (Result == VK_THREAD_IDLE_KHR)
-				{
-					FPlatformProcess::Sleep(0.01f);
-					Result = VulkanRHI::vkDeferredOperationJoinKHR(InDevice->GetInstanceHandle(), DeferredOp);
-				}
-
-				if (Result == VK_SUCCESS)
-				{
-					bCompleted = true;
-				}
-			});
-		checkf(bCompleted, TEXT("ParallelFor returned but Deferred Operation not complete!"));
-
-		VERIFYVULKANRESULT(VulkanRHI::vkGetDeferredOperationResultKHR(InDevice->GetInstanceHandle(), DeferredOp));
-
-		VulkanRHI::vkDestroyDeferredOperationKHR(InDevice->GetInstanceHandle(), DeferredOp, VULKAN_CPU_ALLOCATOR);
-	}
-
-	for (ANSICHAR* const EntryPoint : EntryPointNames)
-	{
-		delete[] EntryPoint;
-	}
+	VERIFYVULKANRESULT_EXPANDED(InDevice->GetPipelineStateCache()->CreateRayTracingPipeline(RayTracingPipelineCreateInfo, bIsPartialPipeline, Pipeline));
 
 	// Grab all shader handles for each stage
 	{
@@ -1698,7 +1724,6 @@ FVulkanRayTracingPipelineState::FVulkanRayTracingPipelineState(FVulkanDevice* co
 	PRAGMA_ENABLE_DEPRECATION_WARNINGS
 
 	INC_DWORD_STAT(STAT_VulkanRayTracingCreatedPipelines);
-	INC_DWORD_STAT_BY(STAT_VulkanRayTracingCompiledShaders, 1);
 }
 
 FVulkanRayTracingPipelineState::~FVulkanRayTracingPipelineState()
@@ -1930,7 +1955,7 @@ void FVulkanRayTracingCompactionRequestHandler::Update(FVulkanCommandListContext
 }
 
 
-static void SetRayGenResources(FVulkanDevice& Device, FVulkanCommandBuffer* const CmdBuffer, const FRayTracingShaderBindings& InGlobalResourceBindings, FVulkanRayTracingShaderTable* ShaderTable)
+static void SetRayGenResources(FVulkanDevice& Device, FVulkanCommandBuffer& CmdBuffer, const FRayTracingShaderBindings& InGlobalResourceBindings, FVulkanRayTracingShaderTable* ShaderTable)
 {
 	TArray<const FVulkanUniformBuffer*> UniformBuffers;
 	UniformBuffers.Reserve(UE_ARRAY_COUNT(InGlobalResourceBindings.UniformBuffers));
@@ -1971,7 +1996,7 @@ static void SetRayGenResources(FVulkanDevice& Device, FVulkanCommandBuffer* cons
 				NumSkippedSlots++;
 			}
 		}
-		Device.GetBindlessDescriptorManager()->RegisterUniformBuffers(*CmdBuffer, VK_PIPELINE_BIND_POINT_RAY_TRACING_KHR, StageUBs);
+		Device.GetBindlessDescriptorManager()->RegisterUniformBuffers(CmdBuffer, VK_PIPELINE_BIND_POINT_RAY_TRACING_KHR, StageUBs);
 	}
 
 	// Add all the UBs references by the shader table
@@ -1996,16 +2021,17 @@ void FVulkanCommandListContext::RHIRayTraceDispatch(
 	
 	Pipeline->FrameCounter.Set(GFrameNumberRenderThread);
 
-	FVulkanCommandBuffer* const CmdBuffer = GetActiveCmdBuffer();
-	VulkanRHI::vkCmdBindPipeline(CmdBuffer->GetHandle(), VK_PIPELINE_BIND_POINT_RAY_TRACING_KHR, Pipeline->GetPipeline());
+	FVulkanCommandBuffer& CommandBuffer = GetCommandBuffer();
+	checkf(!Pipeline->IsPartialPipeline(), TEXT("Attempting to bind partial pipeline, these can't be used for rendering."));
+	VulkanRHI::vkCmdBindPipeline(CommandBuffer.GetHandle(), VK_PIPELINE_BIND_POINT_RAY_TRACING_KHR, Pipeline->GetPipeline());
 
 	checkSlow(InRayGenShader->GetFrequency() == SF_RayGen);
 	const VkStridedDeviceAddressRegionKHR RayGenRegion = ShaderTable->CommitRayGenShader(*this, Pipeline->GetShaderIndex(RayGenShader), Pipeline->GetShaderHandles(SF_RayGen));
 
-	SetRayGenResources(Device, CmdBuffer, InGlobalResourceBindings, ShaderTable);
+	SetRayGenResources(Device, CommandBuffer, InGlobalResourceBindings, ShaderTable);
 
 	VulkanRHI::vkCmdTraceRaysKHR(
-		CmdBuffer->GetHandle(),
+		CommandBuffer.GetHandle(),
 		&RayGenRegion,
 		ShaderTable->GetRegion(SF_RayMiss),
 		ShaderTable->GetRegion(SF_RayHitGroup),
@@ -2026,19 +2052,20 @@ void FVulkanCommandListContext::RHIRayTraceDispatchIndirect(
 	FVulkanRayTracingShader* RayGenShader = ResourceCast(InRayGenShader);
 	FVulkanRayTracingShaderTable* ShaderTable = ResourceCast(InSBT);
 
-	FVulkanCommandBuffer* const CmdBuffer = GetActiveCmdBuffer();
-	VulkanRHI::vkCmdBindPipeline(CmdBuffer->GetHandle(), VK_PIPELINE_BIND_POINT_RAY_TRACING_KHR, Pipeline->GetPipeline());
+	FVulkanCommandBuffer& CommandBuffer = GetCommandBuffer();
+	checkf(!Pipeline->IsPartialPipeline(), TEXT("Attempting to bind partial pipeline, these can't be used for rendering."));
+	VulkanRHI::vkCmdBindPipeline(CommandBuffer.GetHandle(), VK_PIPELINE_BIND_POINT_RAY_TRACING_KHR, Pipeline->GetPipeline());
 
 	checkSlow(InRayGenShader->GetFrequency() == SF_RayGen);
 	const VkStridedDeviceAddressRegionKHR RayGenRegion = ShaderTable->CommitRayGenShader(*this, Pipeline->GetShaderIndex(RayGenShader), Pipeline->GetShaderHandles(SF_RayGen));
 
-	SetRayGenResources(Device, CmdBuffer, InGlobalResourceBindings, ShaderTable);
+	SetRayGenResources(Device, CommandBuffer, InGlobalResourceBindings, ShaderTable);
 
 	FVulkanBuffer* ArgumentBuffer = ResourceCast(InArgumentBuffer);
 	const VkDeviceAddress IndirectDeviceAddress = ArgumentBuffer->GetDeviceAddress() + InArgumentOffset;
 
 	VulkanRHI::vkCmdTraceRaysIndirectKHR(
-		CmdBuffer->GetHandle(),
+		CommandBuffer.GetHandle(),
 		&RayGenRegion,
 		ShaderTable->GetRegion(SF_RayMiss),
 		ShaderTable->GetRegion(SF_RayHitGroup),
diff --git a/Engine/Source/Runtime/VulkanRHI/Private/VulkanRayTracing.h b/Engine/Source/Runtime/VulkanRHI/Private/VulkanRayTracing.h
--- a/Engine/Source/Runtime/VulkanRHI/Private/VulkanRayTracing.h
+++ b/Engine/Source/Runtime/VulkanRHI/Private/VulkanRayTracing.h
@@ -277,11 +277,16 @@ public:
 	FVulkanRayTracingPipelineState(FVulkanDevice* const InDevice, const FRayTracingPipelineStateInitializer& Initializer);
 	~FVulkanRayTracingPipelineState();
 
-	inline VkPipeline GetPipeline() const
+	VkPipeline GetPipeline() const
 	{
 		return Pipeline;
 	}
 
+	bool IsPartialPipeline() const
+	{
+		return bIsPartialPipeline;
+	}
+
 	int32 GetShaderIndex(const FVulkanRayTracingShader* Shader) const;
 	const FVulkanRayTracingShader* GetVulkanShader(EShaderFrequency Frequency, int32 ShaderIndex) const;
 	int32 GetVulkanShaderNum(EShaderFrequency Frequency) const;
@@ -304,6 +309,8 @@ private:
 
 	VkPipeline Pipeline = VK_NULL_HANDLE;
 
+	const bool bIsPartialPipeline;
+
 public:
 	UE_DEPRECATED(5.5, "bAllowHitGroupIndexing is now stored in the ShaderBindingTable.")
 	bool bAllowHitGroupIndexing = true;
