VulkanRHI: Remove FVulkanTimingQueryPool and FVulkanOcclusionQueryPool classes that aren't needed anymore

#rnx

diff --git a/Engine/Source/Runtime/VulkanRHI/Private/VulkanCommandBuffer.cpp b/Engine/Source/Runtime/VulkanRHI/Private/VulkanCommandBuffer.cpp
--- a/Engine/Source/Runtime/VulkanRHI/Private/VulkanCommandBuffer.cpp
+++ b/Engine/Source/Runtime/VulkanRHI/Private/VulkanCommandBuffer.cpp
@@ -120,7 +120,7 @@ void FVulkanCommandBuffer::BeginRenderPass(const FVulkanRenderTargetLayout& Layo
 	}
 }
 
-void FVulkanCommandBuffer::End(FVulkanTimingQueryPool* QueryPool)
+void FVulkanCommandBuffer::End(FVulkanQueryPool* QueryPool)
 {
 	checkf(IsOutsideRenderPass(), TEXT("Can't End as we're inside a render pass! CmdBuffer 0x%p State=%d"), CommandBufferHandle, (int32)State);
 
@@ -143,7 +143,7 @@ void FVulkanCommandBuffer::End(FVulkanTimingQueryPool* QueryPool)
 	State = EState::HasEnded;
 }
 
-void FVulkanCommandBuffer::Begin(FVulkanTimingQueryPool* QueryPool)
+void FVulkanCommandBuffer::Begin(FVulkanQueryPool* QueryPool)
 {
 	{
 		FScopeLock ScopeLock(CommandBufferPool->GetCS());
diff --git a/Engine/Source/Runtime/VulkanRHI/Private/VulkanCommandBuffer.h b/Engine/Source/Runtime/VulkanRHI/Private/VulkanCommandBuffer.h
--- a/Engine/Source/Runtime/VulkanRHI/Private/VulkanCommandBuffer.h
+++ b/Engine/Source/Runtime/VulkanRHI/Private/VulkanCommandBuffer.h
@@ -16,8 +16,7 @@ class FVulkanContextCommon;
 class FVulkanRenderTargetLayout;
 class FVulkanQueue;
 class FVulkanDescriptorPoolSetContainer;
-class FVulkanGPUTiming;
-class FVulkanTimingQueryPool;
+class FVulkanQueryPool;
 
 namespace VulkanRHI
 {
@@ -75,8 +74,8 @@ public:
 		return CommandBufferHandle;
 	}
 
-	void Begin(FVulkanTimingQueryPool* QueryPool);
-	void End(FVulkanTimingQueryPool* QueryPool);
+	void Begin(FVulkanQueryPool* QueryPool);
+	void End(FVulkanQueryPool* QueryPool);
 
 	enum class EState : uint8
 	{
diff --git a/Engine/Source/Runtime/VulkanRHI/Private/VulkanContext.h b/Engine/Source/Runtime/VulkanRHI/Private/VulkanContext.h
--- a/Engine/Source/Runtime/VulkanRHI/Private/VulkanContext.h
+++ b/Engine/Source/Runtime/VulkanRHI/Private/VulkanContext.h
@@ -13,7 +13,6 @@
 
 class FVulkanDevice;
 class FVulkanDynamicRHI;
-class FVulkanOcclusionQueryPool;
 class FVulkanPendingGfxState;
 class FVulkanPendingComputeState;
 class FVulkanQueue;
@@ -186,7 +185,7 @@ public:
 		return Payload.QueryPools[(int32)Type];
 	}
 
-	FVulkanTimingQueryPool* GetCurrentTimestampQueryPool()
+	FVulkanQueryPool* GetCurrentTimestampQueryPool()
 	{
 		return GetCurrentTimestampQueryPool(GetPayload(EPhase::Execute));
 	}
@@ -196,7 +195,7 @@ protected:
 	uint64 FrameCounter = 0;
 
 private:
-	FVulkanTimingQueryPool* GetCurrentTimestampQueryPool(FVulkanPayload& Payload);
+	FVulkanQueryPool* GetCurrentTimestampQueryPool(FVulkanPayload& Payload);
 	void PrepareNewCommandBuffer(FVulkanPayload& Payload);
 	void NewPayload();
 	void EndPayload();
@@ -375,7 +374,7 @@ public:
 
 	void ReleasePendingState();
 
-	FVulkanOcclusionQueryPool* GetCurrentOcclusionQueryPool();
+	FVulkanQueryPool* GetCurrentOcclusionQueryPool();
 
 protected:
 	FVulkanDynamicRHI* RHI;
diff --git a/Engine/Source/Runtime/VulkanRHI/Private/VulkanDevice.h b/Engine/Source/Runtime/VulkanRHI/Private/VulkanDevice.h
--- a/Engine/Source/Runtime/VulkanRHI/Private/VulkanDevice.h
+++ b/Engine/Source/Runtime/VulkanRHI/Private/VulkanDevice.h
@@ -493,8 +493,8 @@ public:
 	}
 #endif
 
-	FVulkanOcclusionQueryPool* AcquireOcclusionQueryPool(uint32 NumQueries);
-	FVulkanTimingQueryPool* AcquireTimingQueryPool();
+	FVulkanQueryPool* AcquireOcclusionQueryPool(uint32 NumQueries);
+	FVulkanQueryPool* AcquireTimingQueryPool();
 	void ReleaseQueryPool(FVulkanQueryPool* Pool);
 	void RemoveStaleQueryPools();
 
diff --git a/Engine/Source/Runtime/VulkanRHI/Private/VulkanGPUProfiler.h b/Engine/Source/Runtime/VulkanRHI/Private/VulkanGPUProfiler.h
--- a/Engine/Source/Runtime/VulkanRHI/Private/VulkanGPUProfiler.h
+++ b/Engine/Source/Runtime/VulkanRHI/Private/VulkanGPUProfiler.h
@@ -13,7 +13,6 @@
 class FVulkanCommandBuffer;
 class FVulkanContextCommon;
 class FVulkanDevice;
-class FVulkanTimingQueryPool;
 
 #if (RHI_NEW_GPU_PROFILER == 0)
 
diff --git a/Engine/Source/Runtime/VulkanRHI/Private/VulkanQuery.cpp b/Engine/Source/Runtime/VulkanRHI/Private/VulkanQuery.cpp
--- a/Engine/Source/Runtime/VulkanRHI/Private/VulkanQuery.cpp
+++ b/Engine/Source/Runtime/VulkanRHI/Private/VulkanQuery.cpp
@@ -115,7 +115,7 @@ void FVulkanCommandListContext::BeginOcclusionQueryBatch(uint32 NumQueriesInBatc
 	OcclusionPoolArray.Add(NewOcclusionQueryPool);
 }
 
-FVulkanOcclusionQueryPool* FVulkanDevice::AcquireOcclusionQueryPool(uint32 NumQueries)
+FVulkanQueryPool* FVulkanDevice::AcquireOcclusionQueryPool(uint32 NumQueries)
 {
 	FScopeLock Lock(&QueryPoolLock);
 
@@ -141,14 +141,14 @@ FVulkanOcclusionQueryPool* FVulkanDevice::AcquireOcclusionQueryPool(uint32 NumQu
 	{
 		FVulkanQueryPool* Pool = FreeOcclusionPools.Pop(EAllowShrinking::No);
 		checkSlow(Pool && (Pool->GetPoolType() == EVulkanQueryPoolType::Occlusion));
-		return (FVulkanOcclusionQueryPool*)Pool;
+		return Pool;
 	}
 
-	FVulkanOcclusionQueryPool* Pool = new FVulkanOcclusionQueryPool(*this, OcclusionQueryPoolSize);
+	FVulkanQueryPool* Pool = new FVulkanQueryPool(*this, OcclusionQueryPoolSize, EVulkanQueryPoolType::Occlusion);
 	return Pool;
 }
 
-FVulkanTimingQueryPool* FVulkanDevice::AcquireTimingQueryPool()
+FVulkanQueryPool* FVulkanDevice::AcquireTimingQueryPool()
 {
 	FScopeLock Lock(&QueryPoolLock);
 
@@ -157,9 +157,9 @@ FVulkanTimingQueryPool* FVulkanDevice::AcquireTimingQueryPool()
 	{
 		FVulkanQueryPool* Pool = FreeTimingPools.Pop(EAllowShrinking::No);
 		checkSlow(Pool && (Pool->GetPoolType() == EVulkanQueryPoolType::Timestamp));
-		return (FVulkanTimingQueryPool*)Pool;
+		return Pool;
 	}
-	return new FVulkanTimingQueryPool(*this, GTimingQueryPoolSize);
+	return new FVulkanQueryPool(*this, GTimingQueryPoolSize, EVulkanQueryPoolType::Timestamp);
 }
 
 void FVulkanDevice::ReleaseQueryPool(FVulkanQueryPool* Pool)
@@ -318,7 +318,7 @@ void FVulkanCommandListContext::RHIBeginRenderQuery(FRHIRenderQuery* QueryRHI)
 	FVulkanRenderQuery* Query = ResourceCast(QueryRHI);
 	if (Query->QueryType == RQT_Occlusion)
 	{
-		FVulkanOcclusionQueryPool* CurrentOcclusionQueryPool = GetCurrentOcclusionQueryPool();
+		FVulkanQueryPool* CurrentOcclusionQueryPool = GetCurrentOcclusionQueryPool();
 		CurrentOcclusionQueryPool->ReserveQuery(Query);
 		VulkanRHI::vkCmdBeginQuery(GetCommandBuffer().GetHandle(), CurrentOcclusionQueryPool->GetHandle(), Query->IndexInPool, VK_QUERY_CONTROL_PRECISE_BIT);
 	}
@@ -334,12 +334,12 @@ void FVulkanCommandListContext::RHIEndRenderQuery(FRHIRenderQuery* QueryRHI)
 	check(Query->SyncPoint.IsValid());
 	if (Query->QueryType == RQT_Occlusion)
 	{
-		FVulkanOcclusionQueryPool* CurrentPool = GetCurrentOcclusionQueryPool();
+		FVulkanQueryPool* CurrentPool = GetCurrentOcclusionQueryPool();
 		VulkanRHI::vkCmdEndQuery(GetCommandBuffer().GetHandle(), CurrentPool->GetHandle(), Query->IndexInPool);
 	}
 	else if (Query->QueryType == RQT_AbsoluteTime)
 	{
-		FVulkanTimingQueryPool* CurrentPool = GetCurrentTimestampQueryPool();
+		FVulkanQueryPool* CurrentPool = GetCurrentTimestampQueryPool();
 		CurrentPool->ReserveQuery(Query);
 
 		const VkPipelineStageFlagBits QueryPipelineStage = GTimestampQueryStage ? VK_PIPELINE_STAGE_BOTTOM_OF_PIPE_BIT : VK_PIPELINE_STAGE_TOP_OF_PIPE_BIT;
@@ -359,19 +359,19 @@ void FVulkanCommandListContext::RHICalibrateTimers(FRHITimestampCalibrationQuery
 }
 #endif
 
-FVulkanOcclusionQueryPool* FVulkanCommandListContext::GetCurrentOcclusionQueryPool()
+FVulkanQueryPool* FVulkanCommandListContext::GetCurrentOcclusionQueryPool()
 {
 	TArray<FVulkanQueryPool*>& OcclusionPoolArray = GetQueryPoolArray(EVulkanQueryPoolType::Occlusion);
 	checkSlow(OcclusionPoolArray.Num() && !OcclusionPoolArray.Last()->IsFull());
-	return (FVulkanOcclusionQueryPool*)OcclusionPoolArray.Last();
+	return OcclusionPoolArray.Last();
 }
 
-FVulkanTimingQueryPool* FVulkanContextCommon::GetCurrentTimestampQueryPool(FVulkanPayload& Payload)
+FVulkanQueryPool* FVulkanContextCommon::GetCurrentTimestampQueryPool(FVulkanPayload& Payload)
 {
 	TArray<FVulkanQueryPool*>& TimestampPoolArray = Payload.QueryPools[(int32)EVulkanQueryPoolType::Timestamp];
 	if ((TimestampPoolArray.Num() == 0) || TimestampPoolArray.Last()->IsFull())
 	{
-		FVulkanTimingQueryPool* NewPool = Device.AcquireTimingQueryPool();
+		FVulkanQueryPool* NewPool = Device.AcquireTimingQueryPool();
 
 		if (Payload.CommandBuffers.Num() == 0)
 		{
@@ -381,5 +381,5 @@ FVulkanTimingQueryPool* FVulkanContextCommon::GetCurrentTimestampQueryPool(FVulk
 
 		TimestampPoolArray.Add(NewPool);
 	}
-	return (FVulkanTimingQueryPool*)TimestampPoolArray.Last();
+	return TimestampPoolArray.Last();
 }
diff --git a/Engine/Source/Runtime/VulkanRHI/Private/VulkanQuery.h b/Engine/Source/Runtime/VulkanRHI/Private/VulkanQuery.h
--- a/Engine/Source/Runtime/VulkanRHI/Private/VulkanQuery.h
+++ b/Engine/Source/Runtime/VulkanRHI/Private/VulkanQuery.h
@@ -100,24 +100,6 @@ protected:
 	friend FVulkanDynamicRHI;
 };
 
-class FVulkanOcclusionQueryPool : public FVulkanQueryPool
-{
-public:
-	FVulkanOcclusionQueryPool(FVulkanDevice& InDevice, uint32 InMaxQueries)
-		: FVulkanQueryPool(InDevice, InMaxQueries, EVulkanQueryPoolType::Occlusion)
-	{
-	}
-};
-
-class FVulkanTimingQueryPool : public FVulkanQueryPool
-{
-public:
-	FVulkanTimingQueryPool(FVulkanDevice& InDevice, uint32 InBufferedCount)
-		: FVulkanQueryPool(InDevice, InBufferedCount * 2, EVulkanQueryPoolType::Timestamp)
-	{
-	}
-};
-
 class FVulkanRenderQuery : public FRHIRenderQuery
 {
 public:
diff --git a/Engine/Source/Runtime/VulkanRHI/Private/VulkanRHI.cpp b/Engine/Source/Runtime/VulkanRHI/Private/VulkanRHI.cpp
--- a/Engine/Source/Runtime/VulkanRHI/Private/VulkanRHI.cpp
+++ b/Engine/Source/Runtime/VulkanRHI/Private/VulkanRHI.cpp
@@ -1298,7 +1298,7 @@ void FVulkanCommandListContext::RHIEndDrawingViewport(FRHIViewport* ViewportRHI,
 			if (GRHIGlobals.SupportsTimestampRenderQueries)
 			{
 				auto& Event = GetCommandBuffer().EmplaceProfilerEvent<UE::RHI::GPUProfiler::FEvent::FBeginBreadcrumb>(Breadcrumb);
-				FVulkanTimingQueryPool* CurrentPool = GetCurrentTimestampQueryPool();
+				FVulkanQueryPool* CurrentPool = GetCurrentTimestampQueryPool();
 				const uint32 IndexInPool = CurrentPool->ReserveQuery(&Event.GPUTimestampTOP);
 				VulkanRHI::vkCmdWriteTimestamp(GetCommandBuffer().GetHandle(), VK_PIPELINE_STAGE_TOP_OF_PIPE_BIT, CurrentPool->GetHandle(), IndexInPool);
 			}
@@ -1329,7 +1329,7 @@ void FVulkanCommandListContext::RHIEndDrawingViewport(FRHIViewport* ViewportRHI,
 			if (GRHIGlobals.SupportsTimestampRenderQueries)
 			{
 				auto& Event = GetCommandBuffer().EmplaceProfilerEvent<UE::RHI::GPUProfiler::FEvent::FEndBreadcrumb>(Breadcrumb);
-				FVulkanTimingQueryPool* CurrentPool = GetCurrentTimestampQueryPool();
+				FVulkanQueryPool* CurrentPool = GetCurrentTimestampQueryPool();
 				const uint32 IndexInPool = CurrentPool->ReserveQuery(&Event.GPUTimestampBOP);
 				VulkanRHI::vkCmdWriteTimestamp(GetCommandBuffer().GetHandle(), VK_PIPELINE_STAGE_BOTTOM_OF_PIPE_BIT, CurrentPool->GetHandle(), IndexInPool);
 			}
diff --git a/Engine/Source/Runtime/VulkanRHI/Private/VulkanUtil.cpp b/Engine/Source/Runtime/VulkanRHI/Private/VulkanUtil.cpp
--- a/Engine/Source/Runtime/VulkanRHI/Private/VulkanUtil.cpp
+++ b/Engine/Source/Runtime/VulkanRHI/Private/VulkanUtil.cpp
@@ -113,7 +113,7 @@ void FVulkanGPUTiming::StartTiming(FVulkanContextCommon* InContext)
 		check(!ActiveQuery);
 		ActiveQuery = new FPendingQuery;
 
-		FVulkanTimingQueryPool* CurrentPool = InContext->GetCurrentTimestampQueryPool();
+		FVulkanQueryPool* CurrentPool = InContext->GetCurrentTimestampQueryPool();
 		const uint32 IndexInPool = CurrentPool->ReserveQuery(&ActiveQuery->StartResult);
 		VulkanRHI::vkCmdWriteTimestamp(InContext->GetCommandBuffer().GetHandle(), VK_PIPELINE_STAGE_BOTTOM_OF_PIPE_BIT, CurrentPool->GetHandle(), IndexInPool);
 		ActiveQuery->StartSyncPoint = InContext->GetContextSyncPoint();
@@ -138,7 +138,7 @@ void FVulkanGPUTiming::EndTiming(FVulkanContextCommon* InContext)
 
 		check(ActiveQuery);
 
-		FVulkanTimingQueryPool* CurrentPool = InContext->GetCurrentTimestampQueryPool();
+		FVulkanQueryPool* CurrentPool = InContext->GetCurrentTimestampQueryPool();
 		const uint32 IndexInPool = CurrentPool->ReserveQuery(&ActiveQuery->EndResult);
 		VulkanRHI::vkCmdWriteTimestamp(InContext->GetCommandBuffer().GetHandle(), VK_PIPELINE_STAGE_BOTTOM_OF_PIPE_BIT, CurrentPool->GetHandle(), IndexInPool);
 		ActiveQuery->EndSyncPoint = InContext->GetContextSyncPoint();
