VulkanRHI: Remove duplicate

#rnx

diff --git a/Engine/Source/Runtime/VulkanRHI/Private/VulkanRHI.cpp b/Engine/Source/Runtime/VulkanRHI/Private/VulkanRHI.cpp
--- a/Engine/Source/Runtime/VulkanRHI/Private/VulkanRHI.cpp
+++ b/Engine/Source/Runtime/VulkanRHI/Private/VulkanRHI.cpp
@@ -440,7 +440,6 @@ FVulkanDynamicRHI::FVulkanDynamicRHI()
 	GTexturePoolSize = 0;
 	GRHISupportsMultithreading = true;
 	GRHISupportsMultithreadedResources = true;
-	GRHISupportsPipelineFileCache = true;
 	GRHITransitionPrivateData_SizeInBytes = sizeof(FVulkanPipelineBarrier);
 	GRHITransitionPrivateData_AlignInBytes = alignof(FVulkanPipelineBarrier);
 	GConfig->GetInt(TEXT("TextureStreaming"), TEXT("PoolSizeVRAMPercentage"), GPoolSizeVRAMPercentage, GEngineIni);
