VulkanRHI: Set bSupportsRealTypes to RuntimeDependent for Vulkan SM6.  Default r.Vulkan.Allow16bitOps to 1 but only set runtime GRHIGlobals.SupportsNative16BitOps to true if DDPI for the current RHI allows it.
Keep it disabled for TSR until shader compilation issue is resolved.

#rnx

diff --git a/Engine/Config/VulkanPC/DataDrivenPlatformInfo.ini b/Engine/Config/VulkanPC/DataDrivenPlatformInfo.ini
--- a/Engine/Config/VulkanPC/DataDrivenPlatformInfo.ini
+++ b/Engine/Config/VulkanPC/DataDrivenPlatformInfo.ini
@@ -258,7 +258,7 @@ FriendlyName=LOCTEXT("FriendlyShaderPlatformName_VulkanPC_SM6", "PC Vulkan SM6")
 BindlessSupport=AllShaderTypes
 bSupportsNNEShaders=false
 bSupportsUniformBufferObjects = true
-bSupportsRealTypes=Unsupported
+bSupportsRealTypes=RuntimeDependent
 bSupportsIndependentSamplers=true
 
 SupportsBarycentricsSemantic=RuntimeGuaranteed
diff --git a/Engine/Source/Runtime/Renderer/Private/PostProcess/TemporalSuperResolution.cpp b/Engine/Source/Runtime/Renderer/Private/PostProcess/TemporalSuperResolution.cpp
--- a/Engine/Source/Runtime/Renderer/Private/PostProcess/TemporalSuperResolution.cpp
+++ b/Engine/Source/Runtime/Renderer/Private/PostProcess/TemporalSuperResolution.cpp
@@ -431,6 +431,12 @@ public:
 
 	static ERHIFeatureSupport Supports16BitVALU(EShaderPlatform Platform)
 	{
+		// UE-254365
+		if (IsVulkanPlatform(Platform))
+		{
+			return ERHIFeatureSupport::Unsupported;
+		}
+
 		return FDataDrivenShaderPlatformInfo::GetSupportsRealTypes(Platform);
 	}
 
diff --git a/Engine/Source/Runtime/VulkanRHI/Private/VulkanExtensions.cpp b/Engine/Source/Runtime/VulkanRHI/Private/VulkanExtensions.cpp
--- a/Engine/Source/Runtime/VulkanRHI/Private/VulkanExtensions.cpp
+++ b/Engine/Source/Runtime/VulkanRHI/Private/VulkanExtensions.cpp
@@ -39,10 +39,10 @@ TAutoConsoleVariable<int32> GRHIAllow64bitShaderAtomicsCvar(
 
 TAutoConsoleVariable<int32> GRHIAllow16bitOps(
 	TEXT("r.Vulkan.Allow16bitOps"),
-	0,
+	1,
 	TEXT("Whether to enable 16bit ops to speeds up TSR\n")
-	TEXT("0 to disable (default)\n")
-	TEXT("1 to enable"),
+	TEXT("0 to disable\n")
+	TEXT("1 to enable (default)"),
 	ECVF_ReadOnly
 );
 
@@ -1554,7 +1554,8 @@ public:
 	{
 		// The PreCreateDevice() call is after all extension have gone through PostPhysicalDeviceFeatures(), so ExtensionFlags will be filled for both.
 		const FOptionalVulkanDeviceExtensions& ExtensionFlags = Device->GetOptionalExtensions();
-		GRHIGlobals.SupportsNative16BitOps = ExtensionFlags.HasKHR16bitStorage && ExtensionFlags.HasKHRShaderFloat16;
+		GRHIGlobals.SupportsNative16BitOps = ExtensionFlags.HasKHR16bitStorage && ExtensionFlags.HasKHRShaderFloat16 &&
+			(FDataDrivenShaderPlatformInfo::GetSupportsRealTypes(GMaxRHIShaderPlatform) != ERHIFeatureSupport::Unsupported);
 
 		if (bRequirementsPassed)
 		{
