VulkanRHI: Load up VK_EXT_mutable_descriptor_type extension.

#rnx

diff --git a/Engine/Source/Runtime/VulkanRHI/Private/VulkanDevice.h b/Engine/Source/Runtime/VulkanRHI/Private/VulkanDevice.h
--- a/Engine/Source/Runtime/VulkanRHI/Private/VulkanDevice.h
+++ b/Engine/Source/Runtime/VulkanRHI/Private/VulkanDevice.h
@@ -62,7 +62,7 @@ struct FOptionalVulkanDeviceExtensions
 			uint64 HasEXTMeshShader : 1;
 			uint64 HasEXTToolingInfo : 1;
 			uint64 HasEXTImageCompressionControl : 1;
-			uint64 HasKHRTimelineSemaphore : 1;
+			uint64 HasEXTMutableDescriptorType : 1;
 
 			// Vendor specific
 			uint64 HasAMDBufferMarker : 1;
@@ -89,6 +89,7 @@ struct FOptionalVulkanDeviceExtensions
 			uint64 HasEXTHostQueryReset : 1;
 			uint64 HasQcomRenderPassShaderResolve : 1;
 			uint64 HasKHRDepthStencilResolve : 1;
+			uint64 HasKHRTimelineSemaphore : 1;
 
 			// Promoted to 1.3
 			uint64 HasEXTTextureCompressionASTCHDR : 1;
diff --git a/Engine/Source/Runtime/VulkanRHI/Private/VulkanExtensions.cpp b/Engine/Source/Runtime/VulkanRHI/Private/VulkanExtensions.cpp
--- a/Engine/Source/Runtime/VulkanRHI/Private/VulkanExtensions.cpp
+++ b/Engine/Source/Runtime/VulkanRHI/Private/VulkanExtensions.cpp
@@ -1039,7 +1039,7 @@ class FVulkanKHRTimelineSemaphoreExtension : public FVulkanDeviceExtension
 public:
 
 	FVulkanKHRTimelineSemaphoreExtension(FVulkanDevice* InDevice)
-		: FVulkanDeviceExtension(InDevice, VK_KHR_TIMELINE_SEMAPHORE_EXTENSION_NAME, VULKAN_EXTENSION_ENABLED)
+		: FVulkanDeviceExtension(InDevice, VK_KHR_TIMELINE_SEMAPHORE_EXTENSION_NAME, VULKAN_EXTENSION_ENABLED, VK_API_VERSION_1_2)
 	{
 	}
 
@@ -1730,6 +1730,38 @@ public:
 	VkPhysicalDeviceGraphicsPipelineLibraryFeaturesEXT GraphicsPipelineLibraryFeatures;
 };
 
+// ***** VK_EXT_mutable_descriptor_type
+class FVulkanEXTMutableDescriptorTypeExtension : public FVulkanDeviceExtension
+{
+public:
+
+	FVulkanEXTMutableDescriptorTypeExtension(FVulkanDevice* InDevice)
+		: FVulkanDeviceExtension(InDevice, VK_EXT_MUTABLE_DESCRIPTOR_TYPE_EXTENSION_NAME, VULKAN_EXTENSION_ENABLED)
+	{
+	}
+
+	virtual void PrePhysicalDeviceFeatures(VkPhysicalDeviceFeatures2KHR& PhysicalDeviceFeatures2) override final
+	{
+		ZeroVulkanStruct(MutableDescriptorTypeFeatures, VK_STRUCTURE_TYPE_PHYSICAL_DEVICE_MUTABLE_DESCRIPTOR_TYPE_FEATURES_EXT);
+		AddToPNext(PhysicalDeviceFeatures2, MutableDescriptorTypeFeatures);
+	}
+
+	virtual void PostPhysicalDeviceFeatures(FOptionalVulkanDeviceExtensions& ExtensionFlags) override final
+	{
+		bRequirementsPassed = (MutableDescriptorTypeFeatures.mutableDescriptorType == VK_TRUE);
+	}
+
+	virtual void PreCreateDevice(VkDeviceCreateInfo& DeviceCreateInfo) override final
+	{
+		if (bRequirementsPassed)
+		{
+			AddToPNext(DeviceCreateInfo, MutableDescriptorTypeFeatures);
+		}
+	}
+
+	VkPhysicalDeviceMutableDescriptorTypeFeaturesEXT MutableDescriptorTypeFeatures;
+};
+
 
 template <typename ExtensionType>
 static void FlagExtensionSupport(const TArray<VkExtensionProperties>& ExtensionProperties, TArray<TUniquePtr<ExtensionType>>& UEExtensions, uint32 ApiVersion, const TCHAR* ExtensionTypeName)
@@ -1823,6 +1855,7 @@ FVulkanDeviceExtensionArray FVulkanDeviceExtension::GetUESupportedDeviceExtensio
 	ADD_CUSTOM_EXTENSION(FVulkanEXTMeshShaderExtension);
 	ADD_CUSTOM_EXTENSION(FVulkanEXTImageCompressionControlExtension);
 	ADD_CUSTOM_EXTENSION(FVulkanKHRGraphicsPipelineLibraryExtension);
+	ADD_CUSTOM_EXTENSION(FVulkanEXTMutableDescriptorTypeExtension);
 
 	// Needed for Raytracing
 	ADD_CUSTOM_EXTENSION(FVulkanKHRBufferDeviceAddressExtension);
