VulkanRHI: Enable Submission/Interrupt threads by default on desktop.  Enable parallel translate on desktop.

#rnx

diff --git a/Engine/Source/Runtime/VulkanRHI/Private/VulkanDynamicRHI.h b/Engine/Source/Runtime/VulkanRHI/Private/VulkanDynamicRHI.h
--- a/Engine/Source/Runtime/VulkanRHI/Private/VulkanDynamicRHI.h
+++ b/Engine/Source/Runtime/VulkanRHI/Private/VulkanDynamicRHI.h
@@ -351,7 +351,7 @@ protected:
 	FCriticalSection InterruptCS;
 	FVulkanThread* InterruptThread = nullptr;
 	VulkanRHI::FSemaphore* CPUTimelineSemaphore = nullptr;  // used to wake up interrupt thread from CPU
-	std::atomic<uint64> CPUTimelineSemaphoreValue = 0;
+	std::atomic<uint64> CPUTimelineSemaphoreValue = 1;
 	TQueue<FVulkanPlatformCommandList*, EQueueMode::Mpsc> PendingPayloadsForSubmission;
 	TMap<VkSemaphore, FBinarySemaphoreSignalInfo> SignaledSemaphores;
 	FGraphEventRef EopTask;
diff --git a/Engine/Source/Runtime/VulkanRHI/Private/VulkanRHI.cpp b/Engine/Source/Runtime/VulkanRHI/Private/VulkanRHI.cpp
--- a/Engine/Source/Runtime/VulkanRHI/Private/VulkanRHI.cpp
+++ b/Engine/Source/Runtime/VulkanRHI/Private/VulkanRHI.cpp
@@ -63,7 +63,7 @@ TAtomic<uint64> GVulkanDSetLayoutHandleIdCounter{ 0 };
 
 TAutoConsoleVariable<int32> GRHIThreadCvar(
 	TEXT("r.Vulkan.RHIThread"),
-	1,
+	2,
 	TEXT("0 to only use Render Thread\n")
 	TEXT("1 to use ONE RHI Thread\n")
 	TEXT("2 to use multiple RHI Thread\n")
diff --git a/Engine/Source/Runtime/VulkanRHI/Private/VulkanSubmission.cpp b/Engine/Source/Runtime/VulkanRHI/Private/VulkanSubmission.cpp
--- a/Engine/Source/Runtime/VulkanRHI/Private/VulkanSubmission.cpp
+++ b/Engine/Source/Runtime/VulkanRHI/Private/VulkanSubmission.cpp
@@ -99,16 +99,16 @@ protected:
 
 static TAutoConsoleVariable<int32> CVarVulkanUseInterruptThread(
 	TEXT("r.Vulkan.Submission.UseInterruptThread"),
-	0,
-	TEXT("  0: Process completed GPU work directly on the RHI thread. (default)\n")
-	TEXT("  1: Create a dedicated thread to process completed GPU work. \n"),
+	1,
+	TEXT("  0: Process completed GPU work directly on the RHI thread.\n")
+	TEXT("  1: Create a dedicated thread to process completed GPU work. (default)\n"),
 	ECVF_ReadOnly);
 
 static TAutoConsoleVariable<int32> CVarVulkanUseSubmissionThread(
 	TEXT("r.Vulkan.Submission.UseSubmissionThread"),
-	0,
-	TEXT("  0: Submit GPU work directly from the RHI thread. (default)\n")
-	TEXT("  1: Create a dedicated thread to submit GPU work. \n"),
+	1,
+	TEXT("  0: Submit GPU work directly from the RHI thread.\n")
+	TEXT("  1: Create a dedicated thread to submit GPU work. (default)\n"),
 	ECVF_ReadOnly);
 
 extern int32 GAllowTimelineSemaphores;
@@ -121,7 +121,7 @@ void FVulkanDynamicRHI::InitializeSubmissionPipe()
 	{
 		if (CVarVulkanUseInterruptThread.GetValueOnAnyThread())
 		{
-			CPUTimelineSemaphore = new VulkanRHI::FSemaphore(*Device, VulkanRHI::EVulkanSemaphoreFlags::Timeline | VulkanRHI::EVulkanSemaphoreFlags::ImmediateDeletion);
+			CPUTimelineSemaphore = new VulkanRHI::FSemaphore(*Device, VulkanRHI::EVulkanSemaphoreFlags::Timeline | VulkanRHI::EVulkanSemaphoreFlags::ImmediateDeletion, CPUTimelineSemaphoreValue);
 			InterruptThread = new FVulkanThread(TEXT("RHIInterruptThread"), TPri_Highest, this, &FVulkanDynamicRHI::WaitAndProcessInterruptQueue);
 		}
 
@@ -361,7 +361,7 @@ bool FVulkanDynamicRHI::WaitAndProcessInterruptQueue()
 				if (CPUTimelineSemaphore)
 				{
 					Semaphores[NumSyncs] = CPUTimelineSemaphore->GetHandle();
-					Values[NumSyncs] = (uint64_t)CPUTimelineSemaphoreValue++;
+					Values[NumSyncs] = (uint64_t)(CPUTimelineSemaphoreValue);
 					NumSyncs++;
 				}
 
@@ -399,7 +399,6 @@ void FVulkanDynamicRHI::ProcessInterruptQueueUntil(FGraphEvent* GraphEvent)
 {
 	if (InterruptThread)
 	{
-		KickInterruptThread();
 		if (GraphEvent && !GraphEvent->IsComplete())
 		{
 			GraphEvent->Wait();
