VulkanRHI: Disable RHI multithreading until issues are investigated.

#rnx

diff --git a/Engine/Source/Runtime/VulkanRHI/Private/VulkanRHI.cpp b/Engine/Source/Runtime/VulkanRHI/Private/VulkanRHI.cpp
--- a/Engine/Source/Runtime/VulkanRHI/Private/VulkanRHI.cpp
+++ b/Engine/Source/Runtime/VulkanRHI/Private/VulkanRHI.cpp
@@ -63,7 +63,7 @@ TAtomic<uint64> GVulkanDSetLayoutHandleIdCounter{ 0 };
 
 TAutoConsoleVariable<int32> GRHIThreadCvar(
 	TEXT("r.Vulkan.RHIThread"),
-	2,
+	1,
 	TEXT("0 to only use Render Thread\n")
 	TEXT("1 to use ONE RHI Thread\n")
 	TEXT("2 to use multiple RHI Thread\n")
diff --git a/Engine/Source/Runtime/VulkanRHI/Private/VulkanSubmission.cpp b/Engine/Source/Runtime/VulkanRHI/Private/VulkanSubmission.cpp
--- a/Engine/Source/Runtime/VulkanRHI/Private/VulkanSubmission.cpp
+++ b/Engine/Source/Runtime/VulkanRHI/Private/VulkanSubmission.cpp
@@ -99,16 +99,16 @@ protected:
 
 static TAutoConsoleVariable<int32> CVarVulkanUseInterruptThread(
 	TEXT("r.Vulkan.Submission.UseInterruptThread"),
-	1,
+	0,
 	TEXT("  0: Process completed GPU work directly on the RHI thread.\n")
-	TEXT("  1: Create a dedicated thread to process completed GPU work. (default)\n"),
+	TEXT("  1: Create a dedicated thread to process completed GPU work.\n"),
 	ECVF_ReadOnly);
 
 static TAutoConsoleVariable<int32> CVarVulkanUseSubmissionThread(
 	TEXT("r.Vulkan.Submission.UseSubmissionThread"),
-	1,
+	0,
 	TEXT("  0: Submit GPU work directly from the RHI thread.\n")
-	TEXT("  1: Create a dedicated thread to submit GPU work. (default)\n"),
+	TEXT("  1: Create a dedicated thread to submit GPU work.\n"),
 	ECVF_ReadOnly);
 
 extern int32 GAllowTimelineSemaphores;
