Vulkan: Don't recycle swapchain & surface when binding to new native window.

#rnx

diff --git a/Engine/Source/Runtime/VulkanRHI/Private/VulkanViewport.cpp b/Engine/Source/Runtime/VulkanRHI/Private/VulkanViewport.cpp
--- a/Engine/Source/Runtime/VulkanRHI/Private/VulkanViewport.cpp
+++ b/Engine/Source/Runtime/VulkanRHI/Private/VulkanViewport.cpp
@@ -720,7 +720,7 @@ void FVulkanViewport::RecreateSwapchainFromRT(FRHICommandListImmediate& RHICmdLi
 	RHICmdList.ImmediateFlush(EImmediateFlushType::FlushRHIThread);
 
 	FVulkanSwapChainRecreateInfo RecreateInfo = { VK_NULL_HANDLE, VK_NULL_HANDLE};
-	DestroySwapchain(&RecreateInfo);
+	DestroySwapchain(WindowHandle == NewNativeWindow ? &RecreateInfo : nullptr);
 	PixelFormat = PreferredPixelFormat;
 	WindowHandle = NewNativeWindow;
 	CreateSwapchain(RHICmdList, &RecreateInfo);
