VulkanRHI: add experimental (disabled) option to use mutable descriptors to reduce overall descriptor set count in Vulkan SM6 bindless.

#rnx

diff --git a/Engine/Source/Runtime/VulkanRHI/Private/VulkanDescriptorSets.cpp b/Engine/Source/Runtime/VulkanRHI/Private/VulkanDescriptorSets.cpp
--- a/Engine/Source/Runtime/VulkanRHI/Private/VulkanDescriptorSets.cpp
+++ b/Engine/Source/Runtime/VulkanRHI/Private/VulkanDescriptorSets.cpp
@@ -126,30 +126,11 @@ static constexpr uint8 GetIndexForDescriptorType(VkDescriptorType DescriptorType
 }
 
 
-static constexpr VkDescriptorType GetDescriptorTypeForSetIndex(uint8 SetIndex)
-{
-	switch (SetIndex)
-	{
-	case VulkanBindless::BindlessSamplerSet:                return VK_DESCRIPTOR_TYPE_SAMPLER;
-	case VulkanBindless::BindlessSampledImageSet:           return VK_DESCRIPTOR_TYPE_SAMPLED_IMAGE;
-	case VulkanBindless::BindlessStorageImageSet:           return VK_DESCRIPTOR_TYPE_STORAGE_IMAGE;
-	case VulkanBindless::BindlessUniformTexelBufferSet:     return VK_DESCRIPTOR_TYPE_UNIFORM_TEXEL_BUFFER;
-	case VulkanBindless::BindlessStorageTexelBufferSet:     return VK_DESCRIPTOR_TYPE_STORAGE_TEXEL_BUFFER;
-	case VulkanBindless::BindlessStorageBufferSet:          return VK_DESCRIPTOR_TYPE_STORAGE_BUFFER;
-	case VulkanBindless::BindlessUniformBufferSet:          return VK_DESCRIPTOR_TYPE_UNIFORM_BUFFER;
-	case VulkanBindless::BindlessSingleUseUniformBufferSet: return VK_DESCRIPTOR_TYPE_UNIFORM_BUFFER;
-	case VulkanBindless::BindlessAccelerationStructureSet:  return VK_DESCRIPTOR_TYPE_ACCELERATION_STRUCTURE_KHR;
-	default: checkNoEntry();
-	}
-	return VK_DESCRIPTOR_TYPE_MAX_ENUM;
-}
-
-
 extern TAutoConsoleVariable<int32> GCVarRobustBufferAccess;
-static inline uint32 GetDescriptorTypeSize(FVulkanDevice* Device, VkDescriptorType DescriptorType)
+static inline uint32 GetDescriptorTypeSize(FVulkanDevice& Device, VkDescriptorType DescriptorType)
 {
 	const bool bRobustBufferAccess = (GCVarRobustBufferAccess.GetValueOnAnyThread() > 0);
-	const VkPhysicalDeviceDescriptorBufferPropertiesEXT& DescriptorBufferProperties = Device->GetOptionalExtensionProperties().DescriptorBufferProps;
+	const VkPhysicalDeviceDescriptorBufferPropertiesEXT& DescriptorBufferProperties = Device.GetOptionalExtensionProperties().DescriptorBufferProps;
 
 	switch (DescriptorType)
 	{
@@ -197,10 +178,10 @@ static inline uint32 GetInitialDescriptorCount(VkDescriptorType DescriptorType)
 }
 
 
-static inline VkMemoryPropertyFlags GetDescriptorBufferMemoryType(FVulkanDevice* Device)
+static inline VkMemoryPropertyFlags GetDescriptorBufferMemoryType(FVulkanDevice& Device)
 {
-	if (Device->HasUnifiedMemory() || (FVulkanPlatform::SupportsDeviceLocalHostVisibleWithNoPenalty(Device->GetVendorId()) &&
-		Device->GetDeviceMemoryManager().SupportsMemoryType(VK_MEMORY_PROPERTY_HOST_VISIBLE_BIT | VK_MEMORY_PROPERTY_HOST_COHERENT_BIT | VK_MEMORY_PROPERTY_DEVICE_LOCAL_BIT)))
+	if (Device.HasUnifiedMemory() || (FVulkanPlatform::SupportsDeviceLocalHostVisibleWithNoPenalty(Device.GetVendorId()) &&
+		Device.GetDeviceMemoryManager().SupportsMemoryType(VK_MEMORY_PROPERTY_HOST_VISIBLE_BIT | VK_MEMORY_PROPERTY_HOST_COHERENT_BIT | VK_MEMORY_PROPERTY_DEVICE_LOCAL_BIT)))
 	{
 		return VK_MEMORY_PROPERTY_DEVICE_LOCAL_BIT | VK_MEMORY_PROPERTY_HOST_VISIBLE_BIT | VK_MEMORY_PROPERTY_HOST_COHERENT_BIT;
 	}
@@ -213,7 +194,7 @@ static inline VkMemoryPropertyFlags GetDescriptorBufferMemoryType(FVulkanDevice*
 
 
 // Check all the requirements to be running in Bindless using Descriptor Buffers
-bool FVulkanBindlessDescriptorManager::VerifySupport(FVulkanDevice* Device)
+bool FVulkanBindlessDescriptorManager::VerifySupport(FVulkanDevice& Device)
 {
 	const bool bFullyDisabled =
 		(RHIGetRuntimeBindlessResourcesConfiguration(GMaxRHIShaderPlatform) == ERHIBindlessConfiguration::Disabled) &&
@@ -230,9 +211,9 @@ bool FVulkanBindlessDescriptorManager::VerifySupport(FVulkanDevice* Device)
 
 	if (bFullyEnabled)
 	{
-		const VkPhysicalDeviceProperties& GpuProps = Device->GetDeviceProperties();
-		const FOptionalVulkanDeviceExtensions& OptionalDeviceExtensions = Device->GetOptionalExtensions();
-		const VkPhysicalDeviceDescriptorBufferPropertiesEXT& DescriptorBufferProperties = Device->GetOptionalExtensionProperties().DescriptorBufferProps;
+		const VkPhysicalDeviceProperties& GpuProps = Device.GetDeviceProperties();
+		const FOptionalVulkanDeviceExtensions& OptionalDeviceExtensions = Device.GetOptionalExtensions();
+		const VkPhysicalDeviceDescriptorBufferPropertiesEXT& DescriptorBufferProperties = Device.GetOptionalExtensionProperties().DescriptorBufferProps;
 
 		const bool bMeetsExtensionsRequirements =
 			OptionalDeviceExtensions.HasEXTDescriptorIndexing &&
@@ -246,7 +227,7 @@ bool FVulkanBindlessDescriptorManager::VerifySupport(FVulkanDevice* Device)
 				(DescriptorBufferProperties.maxDescriptorBufferBindings >= VulkanBindless::MaxNumSets) &&
 				(DescriptorBufferProperties.maxResourceDescriptorBufferBindings >= VulkanBindless::NumBindlessSets) &&
 				(DescriptorBufferProperties.maxSamplerDescriptorBufferBindings >= 1) &&
-				Device->GetDeviceMemoryManager().SupportsMemoryType(GetDescriptorBufferMemoryType(Device));
+				Device.GetDeviceMemoryManager().SupportsMemoryType(GetDescriptorBufferMemoryType(Device));
 
 			if (bMeetsPropertiesRequirements)
 			{
@@ -284,8 +265,8 @@ bool FVulkanBindlessDescriptorManager::VerifySupport(FVulkanDevice* Device)
 }
 
 
-FVulkanBindlessDescriptorManager::FVulkanBindlessDescriptorManager(FVulkanDevice* InDevice)
-	: VulkanRHI::FDeviceChild(InDevice)
+FVulkanBindlessDescriptorManager::FVulkanBindlessDescriptorManager(FVulkanDevice& InDevice)
+	: Device(InDevice)
 	, bIsSupported(VerifySupport(InDevice))
 {
 	FMemory::Memzero(BufferBindingInfo);
@@ -302,7 +283,7 @@ FVulkanBindlessDescriptorManager::~FVulkanBindlessDescriptorManager()
 
 void FVulkanBindlessDescriptorManager::Deinit()
 {
-	const VkDevice DeviceHandle = Device->GetInstanceHandle();
+	const VkDevice DeviceHandle = Device.GetInstanceHandle();
 
 	if (bIsSupported)
 	{
@@ -325,7 +306,7 @@ void FVulkanBindlessDescriptorManager::Deinit()
 		for (uint32 SetIndex = 0; SetIndex < VulkanBindless::NumBindlessSets; ++SetIndex)
 		{
 			BindlessSetState& State = BindlessSetStates[SetIndex];
-			if (State.DescriptorType != VK_DESCRIPTOR_TYPE_MAX_ENUM)
+			if (!State.DescriptorTypes.IsEmpty())
 			{
 				DestroyBindlessState(State);
 			}
@@ -349,11 +330,11 @@ void FVulkanBindlessDescriptorManager::Init()
 		return;
 	}
 
-	const VkDevice DeviceHandle = Device->GetInstanceHandle();
-	const VkPhysicalDeviceDescriptorBufferPropertiesEXT& DescriptorBufferProperties = Device->GetOptionalExtensionProperties().DescriptorBufferProps;
+	const VkDevice DeviceHandle = Device.GetInstanceHandle();
+	const VkPhysicalDeviceDescriptorBufferPropertiesEXT& DescriptorBufferProperties = Device.GetOptionalExtensionProperties().DescriptorBufferProps;
 
 	const VkBufferUsageFlags BufferUsageFlags = VK_BUFFER_USAGE_SHADER_DEVICE_ADDRESS_BIT | VK_BUFFER_USAGE_RESOURCE_DESCRIPTOR_BUFFER_BIT_EXT;
-	SingleUseUBAllocator = new VulkanRHI::FTempBlockAllocator(Device, GVulkanBindlessBlockSize, DescriptorBufferProperties.descriptorBufferOffsetAlignment, BufferUsageFlags);
+	SingleUseUBAllocator = new VulkanRHI::FTempBlockAllocator(&Device, GVulkanBindlessBlockSize, DescriptorBufferProperties.descriptorBufferOffsetAlignment, BufferUsageFlags);
 
 	// Create the dummy layout for unsupported descriptor types
 	{
@@ -364,21 +345,47 @@ void FVulkanBindlessDescriptorManager::Init()
 	}
 
 	{
-		auto InitBindlessSetState = [&](VkDescriptorType DescriptorType, BindlessSetState& OutState) {
+		auto InitBindlessSetState = [&BindlessSetStates=BindlessSetStates, &Device=Device](VkDescriptorType DescriptorType) {
+
+			const uint8 StateIndex = GetIndexForDescriptorType(DescriptorType);
+			BindlessSetState& OutState = BindlessSetStates[StateIndex];
 
-			OutState.DescriptorType = DescriptorType;
+			OutState.DescriptorTypes.Add(DescriptorType);
 
-			OutState.DescriptorSize = GetDescriptorTypeSize(Device, DescriptorType);
+			OutState.DescriptorSize = FMath::Max(OutState.DescriptorSize, GetDescriptorTypeSize(Device, DescriptorType));
 			checkf((OutState.DescriptorSize > 0), TEXT("Descriptor Type [%s] returned an invalid descriptor size!"), VK_TYPE_TO_STRING(VkDescriptorType, DescriptorType));
 
-			OutState.MaxDescriptorCount = GetInitialDescriptorCount(DescriptorType);
+			OutState.MaxDescriptorCount += GetInitialDescriptorCount(DescriptorType);
 			checkf((OutState.MaxDescriptorCount > 0), TEXT("Descriptor Type [%s] returned an invalid descriptor count!"), VK_TYPE_TO_STRING(VkDescriptorType, DescriptorType));
 		};
 
+		// Go through all the supported descriptor types in bindless
+		InitBindlessSetState(VK_DESCRIPTOR_TYPE_SAMPLER);
+		InitBindlessSetState(VK_DESCRIPTOR_TYPE_SAMPLED_IMAGE);
+		InitBindlessSetState(VK_DESCRIPTOR_TYPE_STORAGE_IMAGE);
+		InitBindlessSetState(VK_DESCRIPTOR_TYPE_UNIFORM_TEXEL_BUFFER);
+		InitBindlessSetState(VK_DESCRIPTOR_TYPE_STORAGE_TEXEL_BUFFER);
+		InitBindlessSetState(VK_DESCRIPTOR_TYPE_STORAGE_BUFFER);
+		InitBindlessSetState(VK_DESCRIPTOR_TYPE_UNIFORM_BUFFER);
+
+		// Create a dummy buffer for acceleration structures when they aren't supported (or ray tracing is disabled)
+		if (Device.GetOptionalExtensions().HasRaytracingExtensions())
+		{
+			InitBindlessSetState(VK_DESCRIPTOR_TYPE_ACCELERATION_STRUCTURE_KHR);
+		}
+		else
+		{
+			const uint8 StateIndex = GetIndexForDescriptorType(VK_DESCRIPTOR_TYPE_ACCELERATION_STRUCTURE_KHR);
+			BindlessSetState& State = BindlessSetStates[StateIndex];
+			State.DescriptorTypes.Add(VK_DESCRIPTOR_TYPE_STORAGE_BUFFER);
+			State.DescriptorSize = GetDescriptorTypeSize(Device, VK_DESCRIPTOR_TYPE_STORAGE_BUFFER);
+			State.MaxDescriptorCount = 16;
+		}
+
 		// Fill the DescriptorSetLayout for a BindlessSetState
 		auto CreateDescriptorSetLayout = [&](const BindlessSetState& State) {
 
-			if (State.DescriptorType == VK_DESCRIPTOR_TYPE_MAX_ENUM)
+			if (State.DescriptorTypes.Num() == 0)
 			{
 				return EmptyDescriptorSetLayout;
 			}
@@ -386,7 +393,6 @@ void FVulkanBindlessDescriptorManager::Init()
 			{
 				VkDescriptorSetLayoutBinding DescriptorSetLayoutBinding;
 				DescriptorSetLayoutBinding.binding = 0;
-				DescriptorSetLayoutBinding.descriptorType = State.DescriptorType;
 				DescriptorSetLayoutBinding.descriptorCount = State.MaxDescriptorCount;  // todo-jn: resizable
 				DescriptorSetLayoutBinding.stageFlags = VK_SHADER_STAGE_ALL;
 				DescriptorSetLayoutBinding.pImmutableSamplers = nullptr; // todo-jn: ImmutableSamplers
@@ -407,50 +413,43 @@ void FVulkanBindlessDescriptorManager::Init()
 				DescriptorSetLayoutCreateInfo.flags = VK_DESCRIPTOR_SET_LAYOUT_CREATE_DESCRIPTOR_BUFFER_BIT_EXT;
 				DescriptorSetLayoutCreateInfo.pNext = &DescriptorSetLayoutBindingFlagsCreateInfo;
 
-				VkDescriptorSetLayout DescriptorSetLayout = VK_NULL_HANDLE;
-				VERIFYVULKANRESULT(VulkanRHI::vkCreateDescriptorSetLayout(DeviceHandle, &DescriptorSetLayoutCreateInfo, VULKAN_CPU_ALLOCATOR, &DescriptorSetLayout));
-				return DescriptorSetLayout;
-			}
-		};
-
-		// Uniform buffer descriptor set layout differ from the other resources, we reserve a fixed number of descriptors per stage for each draw/dispatch
-		// todo-jn: this could be compacted..
-		auto CreateShaderStageUniformBufferLayout = [DeviceHandle]() {
+				VkMutableDescriptorTypeCreateInfoEXT MutableDescriptorTypeCreateInfo;
+				VkMutableDescriptorTypeListEXT MutableDescriptorTypeList;
+				if (State.DescriptorTypes.Num() == 1)
+				{
+					DescriptorSetLayoutBinding.descriptorType = State.DescriptorTypes[0];
+				}
+				else
+				{
+					DescriptorSetLayoutBinding.descriptorType = VK_DESCRIPTOR_TYPE_MUTABLE_EXT;
 
-			const uint32 NumTotalBindings = VulkanBindless::MaxUniformBuffersPerStage * ShaderStage::MaxNumStages;
+					MutableDescriptorTypeList.descriptorTypeCount = State.DescriptorTypes.Num();
+					MutableDescriptorTypeList.pDescriptorTypes = State.DescriptorTypes.GetData();
 
-			TArray<VkDescriptorSetLayoutBinding> DescriptorSetLayoutBindings;
-			DescriptorSetLayoutBindings.SetNumZeroed(NumTotalBindings);
-			for (uint32 BindingIndex = 0; BindingIndex < NumTotalBindings; ++BindingIndex)
-			{
-				DescriptorSetLayoutBindings[BindingIndex].binding = BindingIndex;
-				DescriptorSetLayoutBindings[BindingIndex].descriptorType = VK_DESCRIPTOR_TYPE_UNIFORM_BUFFER;
-				DescriptorSetLayoutBindings[BindingIndex].descriptorCount = 1;
-				DescriptorSetLayoutBindings[BindingIndex].stageFlags = VK_SHADER_STAGE_ALL;
-			}
+					ZeroVulkanStruct(MutableDescriptorTypeCreateInfo, VK_STRUCTURE_TYPE_MUTABLE_DESCRIPTOR_TYPE_CREATE_INFO_EXT);
+					MutableDescriptorTypeCreateInfo.mutableDescriptorTypeListCount = 1;
+					MutableDescriptorTypeCreateInfo.pMutableDescriptorTypeLists = &MutableDescriptorTypeList;
 
-			VkDescriptorSetLayoutCreateInfo DescriptorSetLayoutCreateInfo;
-			ZeroVulkanStruct(DescriptorSetLayoutCreateInfo, VK_STRUCTURE_TYPE_DESCRIPTOR_SET_LAYOUT_CREATE_INFO);
-			DescriptorSetLayoutCreateInfo.pBindings = DescriptorSetLayoutBindings.GetData();
-			DescriptorSetLayoutCreateInfo.bindingCount = NumTotalBindings;
-			DescriptorSetLayoutCreateInfo.flags = VK_DESCRIPTOR_SET_LAYOUT_CREATE_DESCRIPTOR_BUFFER_BIT_EXT;
-			DescriptorSetLayoutCreateInfo.pNext = nullptr;
+					MutableDescriptorTypeCreateInfo.pNext = DescriptorSetLayoutCreateInfo.pNext;
+					DescriptorSetLayoutCreateInfo.pNext = &MutableDescriptorTypeCreateInfo;
+				}
 
-			VkDescriptorSetLayout DescriptorSetLayout = VK_NULL_HANDLE;
-			VERIFYVULKANRESULT(VulkanRHI::vkCreateDescriptorSetLayout(DeviceHandle, &DescriptorSetLayoutCreateInfo, VULKAN_CPU_ALLOCATOR, &DescriptorSetLayout));
-			return DescriptorSetLayout;
+				VkDescriptorSetLayout DescriptorSetLayout = VK_NULL_HANDLE;
+				VERIFYVULKANRESULT(VulkanRHI::vkCreateDescriptorSetLayout(DeviceHandle, &DescriptorSetLayoutCreateInfo, VULKAN_CPU_ALLOCATOR, &DescriptorSetLayout));
+				return DescriptorSetLayout;
+			}
 		};
 
 		// Create the descriptor buffer for a BindlessSetState
-		auto CreateDescriptorBuffer = [&](BindlessSetState& InOutState, VkDescriptorBufferBindingInfoEXT& OutBufferBindingInfo, bool IsSingleUseUniformBufferSet) {
+		auto CreateDescriptorBuffer = [&](BindlessSetState& InOutState, VkDescriptorBufferBindingInfoEXT& OutBufferBindingInfo) {
 
 			// Skip unsupported descriptors
-			if (InOutState.DescriptorType == VK_DESCRIPTOR_TYPE_MAX_ENUM)
+			if (InOutState.DescriptorTypes.Num() == 0)
 			{
 				return 0u;
 			}
 
-			const bool IsSamplerSet = (InOutState.DescriptorType == VK_DESCRIPTOR_TYPE_SAMPLER);
+			const bool IsSamplerSet = (InOutState.DescriptorTypes[0] == VK_DESCRIPTOR_TYPE_SAMPLER);
 			const VkBufferUsageFlags BufferUsageFlags =
 				VK_BUFFER_USAGE_SHADER_DEVICE_ADDRESS_BIT |
 				VK_BUFFER_USAGE_TRANSFER_SRC_BIT |
@@ -462,18 +461,8 @@ void FVulkanBindlessDescriptorManager::Init()
 
 			VkDeviceSize LayoutSizeInBytes = 0;
 			VulkanRHI::vkGetDescriptorSetLayoutSizeEXT(DeviceHandle, InOutState.DescriptorSetLayout, &LayoutSizeInBytes);
-			if (IsSingleUseUniformBufferSet)
-			{
-				// todo-jn: We're picky about uniform buffers values for now to allow for shortcuts...
-				check(LayoutSizeInBytes == (ShaderStage::MaxNumStages * VulkanBindless::MaxUniformBuffersPerStage * InOutState.DescriptorSize));
-				check((LayoutSizeInBytes % DescriptorBufferProperties.descriptorBufferOffsetAlignment) == 0);
-				check((InOutState.MaxDescriptorCount % (ShaderStage::MaxNumStages * VulkanBindless::MaxUniformBuffersPerStage)) == 0);
-			}
-			else
-			{
-				// Double check that the layout follows the rules for a single binding with an array of descriptors that are tightly packed
-				check(LayoutSizeInBytes == (InOutState.MaxDescriptorCount * InOutState.DescriptorSize));
-			}
+			// Double check that the layout follows the rules for a single binding with an array of descriptors that are tightly packed
+			check(LayoutSizeInBytes == (InOutState.MaxDescriptorCount * InOutState.DescriptorSize));
 
 			if (IsSamplerSet)
 			{
@@ -499,7 +488,7 @@ void FVulkanBindlessDescriptorManager::Init()
 				check(BufferMemoryReqs.size >= DescriptorBufferSize);
 
 				uint32 MemoryTypeIndex = 0;
-				VERIFYVULKANRESULT(Device->GetDeviceMemoryManager().GetMemoryTypeFromProperties(BufferMemoryReqs.memoryTypeBits, GetDescriptorBufferMemoryType(Device), &MemoryTypeIndex));
+				VERIFYVULKANRESULT(Device.GetDeviceMemoryManager().GetMemoryTypeFromProperties(BufferMemoryReqs.memoryTypeBits, GetDescriptorBufferMemoryType(Device), &MemoryTypeIndex));
 
 				VkMemoryAllocateFlagsInfo FlagsInfo;
 				ZeroVulkanStruct(FlagsInfo, VK_STRUCTURE_TYPE_MEMORY_ALLOCATE_FLAGS_INFO);
@@ -536,28 +525,35 @@ void FVulkanBindlessDescriptorManager::Init()
 		for (uint32 SetIndex = 0; SetIndex < VulkanBindless::NumBindlessSets; ++SetIndex)
 		{
 			BindlessSetState& State = BindlessSetStates[SetIndex];
+			State.DescriptorSetLayout = CreateDescriptorSetLayout(State);
+			TotalResourceDescriptorBufferSize += CreateDescriptorBuffer(State, BufferBindingInfo[SetIndex]);
+		}
 
-			// Create a dummy buffer for acceleration structures when they aren't supported (or ray tracing is disabled)
-			const bool bSupported = (SetIndex != VulkanBindless::BindlessAccelerationStructureSet) || Device->GetOptionalExtensions().HasRaytracingExtensions();
-			if (bSupported)
-			{
-				InitBindlessSetState(GetDescriptorTypeForSetIndex(SetIndex), State);
-			}
-			else
+		// Fill in the state for single-use UB
+		// Uniform buffer descriptor set layout differ from the other resources, we reserve a fixed number of descriptors per stage for each draw/dispatch
+		// todo-jn: this could be compacted..
+		{
+			const uint32 NumTotalBindings = VulkanBindless::MaxUniformBuffersPerStage * ShaderStage::MaxNumStages;
+
+			TArray<VkDescriptorSetLayoutBinding> DescriptorSetLayoutBindings;
+			DescriptorSetLayoutBindings.SetNumZeroed(NumTotalBindings);
+			for (uint32 BindingIndex = 0; BindingIndex < NumTotalBindings; ++BindingIndex)
 			{
-				State.DescriptorType = VK_DESCRIPTOR_TYPE_STORAGE_BUFFER;
-				State.DescriptorSize = GetDescriptorTypeSize(Device, VK_DESCRIPTOR_TYPE_STORAGE_BUFFER);
-				State.MaxDescriptorCount = 16;
+				DescriptorSetLayoutBindings[BindingIndex].binding = BindingIndex;
+				DescriptorSetLayoutBindings[BindingIndex].descriptorType = VK_DESCRIPTOR_TYPE_UNIFORM_BUFFER;
+				DescriptorSetLayoutBindings[BindingIndex].descriptorCount = 1;
+				DescriptorSetLayoutBindings[BindingIndex].stageFlags = VK_SHADER_STAGE_ALL;
 			}
 
-			const bool IsSingleUseUniformBufferSet = (SetIndex == VulkanBindless::BindlessSingleUseUniformBufferSet);
-			State.DescriptorSetLayout = IsSingleUseUniformBufferSet ? CreateShaderStageUniformBufferLayout() : CreateDescriptorSetLayout(State);
-			TotalResourceDescriptorBufferSize += CreateDescriptorBuffer(State, BufferBindingInfo[SetIndex], IsSingleUseUniformBufferSet);
-		}
+			VkDescriptorSetLayoutCreateInfo DescriptorSetLayoutCreateInfo;
+			ZeroVulkanStruct(DescriptorSetLayoutCreateInfo, VK_STRUCTURE_TYPE_DESCRIPTOR_SET_LAYOUT_CREATE_INFO);
+			DescriptorSetLayoutCreateInfo.pBindings = DescriptorSetLayoutBindings.GetData();
+			DescriptorSetLayoutCreateInfo.bindingCount = NumTotalBindings;
+			DescriptorSetLayoutCreateInfo.flags = VK_DESCRIPTOR_SET_LAYOUT_CREATE_DESCRIPTOR_BUFFER_BIT_EXT;
+			DescriptorSetLayoutCreateInfo.pNext = nullptr;
 
-		// Fill in the state for single-use UB
-		{
-			SingleUseUBDescriptorSetLayout = CreateShaderStageUniformBufferLayout();
+			checkSlow(SingleUseUBDescriptorSetLayout == VK_NULL_HANDLE);
+			VERIFYVULKANRESULT(VulkanRHI::vkCreateDescriptorSetLayout(DeviceHandle, &DescriptorSetLayoutCreateInfo, VULKAN_CPU_ALLOCATOR, &SingleUseUBDescriptorSetLayout));
 		}
 
 		checkf(TotalResourceDescriptorBufferSize < DescriptorBufferProperties.resourceDescriptorBufferAddressSpaceSize,
@@ -580,7 +576,7 @@ void FVulkanBindlessDescriptorManager::Init()
 		PipelineLayoutCreateInfo.setLayoutCount = VulkanBindless::MaxNumSets;
 		PipelineLayoutCreateInfo.pSetLayouts = DescriptorSetLayouts;
 		VERIFYVULKANRESULT(VulkanRHI::vkCreatePipelineLayout(DeviceHandle, &PipelineLayoutCreateInfo, VULKAN_CPU_ALLOCATOR, &BindlessPipelineLayout));
-		VULKAN_SET_DEBUG_NAME((*Device), VK_OBJECT_TYPE_PIPELINE_LAYOUT, BindlessPipelineLayout, TEXT("BindlessPipelineLayout(SetCount=%d)"), VulkanBindless::MaxNumSets);
+		VULKAN_SET_DEBUG_NAME(Device, VK_OBJECT_TYPE_PIPELINE_LAYOUT, BindlessPipelineLayout, TEXT("BindlessPipelineLayout(SetCount=%d)"), VulkanBindless::MaxNumSets);
 	}
 }
 
@@ -643,7 +639,7 @@ void FVulkanBindlessDescriptorManager::RegisterUniformBuffers(FVulkanCommandBuff
 				const int32 BindingIndex = StageOffset + DescriptorAddressInfoIndex;
 				VkDeviceSize BindingByteOffset = 0;
 #if UE_BUILD_DEBUG
-				VulkanRHI::vkGetDescriptorSetLayoutBindingOffsetEXT(Device->GetInstanceHandle(), SingleUseUBDescriptorSetLayout, BindingIndex, &BindingByteOffset);
+				VulkanRHI::vkGetDescriptorSetLayoutBindingOffsetEXT(Device.GetInstanceHandle(), SingleUseUBDescriptorSetLayout, BindingIndex, &BindingByteOffset);
 				check(BindingByteOffset == BindingIndex * UBDescriptorSize);
 #else
 				BindingByteOffset = (BindingIndex * UBDescriptorSize);
@@ -653,7 +649,7 @@ void FVulkanBindlessDescriptorManager::RegisterUniformBuffers(FVulkanCommandBuff
 				ZeroVulkanStruct(Info, VK_STRUCTURE_TYPE_DESCRIPTOR_GET_INFO_EXT);
 				Info.type = VK_DESCRIPTOR_TYPE_UNIFORM_BUFFER;
 				Info.data.pUniformBuffer = &DescriptorAddressInfo;
-				VulkanRHI::vkGetDescriptorEXT(Device->GetInstanceHandle(), &Info, UBDescriptorSize, &MappedPointer[BindingByteOffset]);
+				VulkanRHI::vkGetDescriptorEXT(Device.GetInstanceHandle(), &Info, UBDescriptorSize, &MappedPointer[BindingByteOffset]);
 			}
 		}
 	}
@@ -663,24 +659,23 @@ void FVulkanBindlessDescriptorManager::RegisterUniformBuffers(FVulkanCommandBuff
 	VulkanRHI::vkCmdSetDescriptorBufferOffsetsEXT(CommandBuffer.GetHandle(), BindPoint, BindlessPipelineLayout, 0u, VulkanBindless::MaxNumSets, BufferIndices, BufferOffsets);
 }
 
-void FVulkanBindlessDescriptorManager::UpdateStatsForHandle(FRHIDescriptorHandle DescriptorHandle)
+void FVulkanBindlessDescriptorManager::UpdateStatsForHandle(VkDescriptorType DescriptorType)
 {
-	const uint8 SetIndex = DescriptorHandle.GetRawType();
+	const uint8 SetIndex = GetIndexForDescriptorType(DescriptorType);
 	const BindlessSetState& State = BindlessSetStates[SetIndex];
 
-	switch (DescriptorHandle.GetRawType())
+	switch (DescriptorType)
 	{
-	case VulkanBindless::BindlessSamplerSet:                SET_DWORD_STAT(STAT_VulkanBindlessPeakSampler, State.PeakDescriptorCount); break;
-	case VulkanBindless::BindlessSampledImageSet:           SET_DWORD_STAT(STAT_VulkanBindlessPeakSampledImage, State.PeakDescriptorCount); break;
-	case VulkanBindless::BindlessStorageImageSet:           SET_DWORD_STAT(STAT_VulkanBindlessPeakStorageImage, State.PeakDescriptorCount); break;
-	case VulkanBindless::BindlessUniformTexelBufferSet:     SET_DWORD_STAT(STAT_VulkanBindlessPeakUniformTexelBuffer, State.PeakDescriptorCount); break;
-	case VulkanBindless::BindlessStorageTexelBufferSet:     SET_DWORD_STAT(STAT_VulkanBindlessPeakStorageTexelBuffer, State.PeakDescriptorCount); break;
-	case VulkanBindless::BindlessUniformBufferSet:          SET_DWORD_STAT(STAT_VulkanBindlessPeakUniformBuffer, State.PeakDescriptorCount); break;
-	case VulkanBindless::BindlessStorageBufferSet:          SET_DWORD_STAT(STAT_VulkanBindlessPeakStorageBuffer, State.PeakDescriptorCount); break;
-	case VulkanBindless::BindlessAccelerationStructureSet:  SET_DWORD_STAT(STAT_VulkanBindlessPeakAccelerationStructure, State.PeakDescriptorCount); break;
-
-	case VulkanBindless::BindlessSingleUseUniformBufferSet:
-	default: checkNoEntry();
+	case VK_DESCRIPTOR_TYPE_SAMPLER:                   SET_DWORD_STAT(STAT_VulkanBindlessPeakSampler, State.PeakDescriptorCount); break;
+	case VK_DESCRIPTOR_TYPE_SAMPLED_IMAGE:             SET_DWORD_STAT(STAT_VulkanBindlessPeakSampledImage, State.PeakDescriptorCount); break;
+	case VK_DESCRIPTOR_TYPE_STORAGE_IMAGE:             SET_DWORD_STAT(STAT_VulkanBindlessPeakStorageImage, State.PeakDescriptorCount); break;
+	case VK_DESCRIPTOR_TYPE_UNIFORM_TEXEL_BUFFER:      SET_DWORD_STAT(STAT_VulkanBindlessPeakUniformTexelBuffer, State.PeakDescriptorCount); break;
+	case VK_DESCRIPTOR_TYPE_STORAGE_TEXEL_BUFFER:      SET_DWORD_STAT(STAT_VulkanBindlessPeakStorageTexelBuffer, State.PeakDescriptorCount); break;
+	case VK_DESCRIPTOR_TYPE_STORAGE_BUFFER:            SET_DWORD_STAT(STAT_VulkanBindlessPeakStorageBuffer, State.PeakDescriptorCount); break;
+	case VK_DESCRIPTOR_TYPE_UNIFORM_BUFFER:            SET_DWORD_STAT(STAT_VulkanBindlessPeakUniformBuffer, State.PeakDescriptorCount); break;
+	case VK_DESCRIPTOR_TYPE_ACCELERATION_STRUCTURE_KHR:SET_DWORD_STAT(STAT_VulkanBindlessPeakAccelerationStructure, State.PeakDescriptorCount); break;
+	default: 
+		checkf(false, TEXT("Unknown descriptor type for stat update: %s"), VK_TYPE_TO_STRING(VkDescriptorType, DescriptorType));
 	}
 }
 
@@ -696,19 +691,21 @@ FRHIDescriptorHandle FVulkanBindlessDescriptorManager::ReserveDescriptor(VkDescr
 	return FRHIDescriptorHandle();
 }
 
-void FVulkanBindlessDescriptorManager::UpdateDescriptor(FRHIDescriptorHandle DescriptorHandle, VkDescriptorDataEXT DescriptorData, bool bImmediateUpdate)
+void FVulkanBindlessDescriptorManager::UpdateDescriptor(FRHIDescriptorHandle DescriptorHandle, VkDescriptorType DescriptorType, VkDescriptorDataEXT DescriptorData, bool bImmediateUpdate)
 {
 	checkf(DescriptorHandle.IsValid(), TEXT("Attemping to update invalid descriptor handle!"));
 
 	const uint8 SetIndex = DescriptorHandle.GetRawType();
+	check(SetIndex == GetIndexForDescriptorType(DescriptorType));
 	BindlessSetState& State = BindlessSetStates[SetIndex];
 	const uint32 ByteOffset = DescriptorHandle.GetIndex() * State.DescriptorSize;
+	checkSlow(State.DescriptorTypes.Contains(DescriptorType));
 
 	VkDescriptorGetInfoEXT Info;
 	ZeroVulkanStruct(Info, VK_STRUCTURE_TYPE_DESCRIPTOR_GET_INFO_EXT);
-	Info.type = State.DescriptorType;
+	Info.type = DescriptorType;
 	Info.data = DescriptorData;
-	VulkanRHI::vkGetDescriptorEXT(Device->GetInstanceHandle(), &Info, State.DescriptorSize, &State.DebugDescriptors[ByteOffset]);
+	VulkanRHI::vkGetDescriptorEXT(Device.GetInstanceHandle(), &Info, State.DescriptorSize, &State.DebugDescriptors[ByteOffset]);
 
 	if (bImmediateUpdate)
 	{
@@ -717,7 +714,7 @@ void FVulkanBindlessDescriptorManager::UpdateDescriptor(FRHIDescriptorHandle Des
 	else
 	{
 		check(!IsInRenderingThread() || FRHICommandListExecutor::GetImmediateCommandList().Bypass() || !IsRunningRHIInSeparateThread());
-		FVulkanCommandBuffer* CmdBuffer = Device->GetImmediateContext().GetActiveCmdBuffer();
+		FVulkanCommandBuffer* CmdBuffer = Device.GetImmediateContext().GetActiveCmdBuffer();
 
 		// :todo-jn:  Hack to avoid barriers/copies in renderpasses
 		if (CmdBuffer->IsInsideRenderPass())
@@ -726,7 +723,7 @@ void FVulkanBindlessDescriptorManager::UpdateDescriptor(FRHIDescriptorHandle Des
 		}
 		else
 		{
-			VulkanRHI::FStagingBuffer* StagingBuffer = Device->GetStagingManager().AcquireBuffer(State.DescriptorSize);
+			VulkanRHI::FStagingBuffer* StagingBuffer = Device.GetStagingManager().AcquireBuffer(State.DescriptorSize);
 			FMemory::Memcpy(StagingBuffer->GetMappedPointer(), &State.DebugDescriptors[ByteOffset], State.DescriptorSize);
 			{
 				VkMemoryBarrier2 MemoryBarrier;
@@ -754,11 +751,11 @@ void FVulkanBindlessDescriptorManager::UpdateDescriptor(FRHIDescriptorHandle Des
 				MemoryBarrier.dstAccessMask = VK_ACCESS_2_DESCRIPTOR_BUFFER_READ_BIT_EXT;
 				VulkanRHI::vkCmdPipelineBarrier2KHR(CmdBuffer->GetHandle(), &DependencyInfo);
 			}
-			Device->GetStagingManager().ReleaseBuffer(&Device->GetImmediateContext(), StagingBuffer);
+			Device.GetStagingManager().ReleaseBuffer(&Device.GetImmediateContext(), StagingBuffer);
 		}
 	}
 
-	UpdateStatsForHandle(DescriptorHandle);
+	UpdateStatsForHandle(DescriptorType);
 }
 
 void FVulkanBindlessDescriptorManager::UpdateSampler(FRHIDescriptorHandle DescriptorHandle, VkSampler VulkanSampler)
@@ -767,15 +764,14 @@ void FVulkanBindlessDescriptorManager::UpdateSampler(FRHIDescriptorHandle Descri
 	{
 		VkDescriptorDataEXT DescriptorData;
 		DescriptorData.pSampler = &VulkanSampler;
-		UpdateDescriptor(DescriptorHandle, DescriptorData, true);
+		UpdateDescriptor(DescriptorHandle, VK_DESCRIPTOR_TYPE_SAMPLER, DescriptorData, true);
 	}
 }
 
-void FVulkanBindlessDescriptorManager::UpdateImage(FRHIDescriptorHandle DescriptorHandle, VkImageView ImageView, bool bIsDepthStencil, bool bImmediateUpdate)
+void FVulkanBindlessDescriptorManager::UpdateImage(FRHIDescriptorHandle DescriptorHandle, VkDescriptorType DescriptorType, VkImageView ImageView, bool bIsDepthStencil, bool bImmediateUpdate)
 {
 	if (bIsSupported)
 	{
-		const VkDescriptorType DescriptorType = GetDescriptorTypeForSetIndex(DescriptorHandle.GetRawType());
 		check((DescriptorType == VK_DESCRIPTOR_TYPE_SAMPLED_IMAGE) || (DescriptorType == VK_DESCRIPTOR_TYPE_STORAGE_IMAGE));
 
 		VkDescriptorImageInfo DescriptorImageInfo;
@@ -786,31 +782,29 @@ void FVulkanBindlessDescriptorManager::UpdateImage(FRHIDescriptorHandle Descript
 
 		VkDescriptorDataEXT DescriptorData;
 		DescriptorData.pSampledImage = &DescriptorImageInfo;  // same pointer for storage, it's a union
-		UpdateDescriptor(DescriptorHandle, DescriptorData, bImmediateUpdate);
+		UpdateDescriptor(DescriptorHandle, DescriptorType, DescriptorData, bImmediateUpdate);
 	}
 }
 
-void FVulkanBindlessDescriptorManager::UpdateBuffer(FRHIDescriptorHandle DescriptorHandle, VkBuffer VulkanBuffer, VkDeviceSize BufferOffset, VkDeviceSize BufferSize, bool bImmediateUpdate)
+void FVulkanBindlessDescriptorManager::UpdateBuffer(FRHIDescriptorHandle DescriptorHandle, VkDescriptorType DescriptorType, VkBuffer VulkanBuffer, VkDeviceSize BufferOffset, VkDeviceSize BufferSize, bool bImmediateUpdate)
 {
 	if (bIsSupported)
 	{
-		const VkDescriptorType DescriptorType = GetDescriptorTypeForSetIndex(DescriptorHandle.GetRawType());
 		check((DescriptorType == VK_DESCRIPTOR_TYPE_UNIFORM_BUFFER) || (DescriptorType == VK_DESCRIPTOR_TYPE_STORAGE_BUFFER));
 
 		VkBufferDeviceAddressInfo BufferInfo;
 		ZeroVulkanStruct(BufferInfo, VK_STRUCTURE_TYPE_BUFFER_DEVICE_ADDRESS_INFO);
 		BufferInfo.buffer = VulkanBuffer;
-		const VkDeviceAddress BufferAddress = VulkanRHI::vkGetBufferDeviceAddressKHR(Device->GetInstanceHandle(), &BufferInfo);
+		const VkDeviceAddress BufferAddress = VulkanRHI::vkGetBufferDeviceAddressKHR(Device.GetInstanceHandle(), &BufferInfo);
 
-		UpdateBuffer(DescriptorHandle, BufferAddress + BufferOffset, BufferSize, bImmediateUpdate);
+		UpdateBuffer(DescriptorHandle, DescriptorType, BufferAddress + BufferOffset, BufferSize, bImmediateUpdate);
 	}
 }
 
-void FVulkanBindlessDescriptorManager::UpdateBuffer(FRHIDescriptorHandle DescriptorHandle, VkDeviceAddress BufferAddress, VkDeviceSize BufferSize, bool bImmediateUpdate)
+void FVulkanBindlessDescriptorManager::UpdateBuffer(FRHIDescriptorHandle DescriptorHandle, VkDescriptorType DescriptorType, VkDeviceAddress BufferAddress, VkDeviceSize BufferSize, bool bImmediateUpdate)
 {
 	if (bIsSupported)
 	{
-		const VkDescriptorType DescriptorType = GetDescriptorTypeForSetIndex(DescriptorHandle.GetRawType());
 		check((DescriptorType == VK_DESCRIPTOR_TYPE_UNIFORM_BUFFER) || (DescriptorType == VK_DESCRIPTOR_TYPE_STORAGE_BUFFER));
 
 		VkDescriptorAddressInfoEXT AddressInfo;
@@ -820,22 +814,21 @@ void FVulkanBindlessDescriptorManager::UpdateBuffer(FRHIDescriptorHandle Descrip
 
 		VkDescriptorDataEXT DescriptorData;
 		DescriptorData.pStorageBuffer = &AddressInfo;  // same pointer for uniform, it's a union
-		UpdateDescriptor(DescriptorHandle, DescriptorData, bImmediateUpdate);
+		UpdateDescriptor(DescriptorHandle, DescriptorType, DescriptorData, bImmediateUpdate);
 	}
 }
 
-void FVulkanBindlessDescriptorManager::UpdateTexelBuffer(FRHIDescriptorHandle DescriptorHandle, const VkBufferViewCreateInfo& ViewInfo, bool bImmediateUpdate)
+void FVulkanBindlessDescriptorManager::UpdateTexelBuffer(FRHIDescriptorHandle DescriptorHandle, VkDescriptorType DescriptorType, const VkBufferViewCreateInfo& ViewInfo, bool bImmediateUpdate)
 {
 	if (bIsSupported)
 	{
-		const VkDescriptorType DescriptorType = GetDescriptorTypeForSetIndex(DescriptorHandle.GetRawType());
 		check((DescriptorType == VK_DESCRIPTOR_TYPE_UNIFORM_TEXEL_BUFFER) || (DescriptorType == VK_DESCRIPTOR_TYPE_STORAGE_TEXEL_BUFFER));
 
 		// :todo-jn: start caching buffer addresses in resources to avoid the extra call
 		VkBufferDeviceAddressInfo BufferInfo;
 		ZeroVulkanStruct(BufferInfo, VK_STRUCTURE_TYPE_BUFFER_DEVICE_ADDRESS_INFO);
 		BufferInfo.buffer = ViewInfo.buffer;
-		const VkDeviceAddress BufferAddress = VulkanRHI::vkGetBufferDeviceAddressKHR(Device->GetInstanceHandle(), &BufferInfo);
+		const VkDeviceAddress BufferAddress = VulkanRHI::vkGetBufferDeviceAddressKHR(Device.GetInstanceHandle(), &BufferInfo);
 
 		VkDescriptorAddressInfoEXT AddressInfo;
 		ZeroVulkanStruct(AddressInfo, VK_STRUCTURE_TYPE_DESCRIPTOR_ADDRESS_INFO_EXT);
@@ -845,7 +838,7 @@ void FVulkanBindlessDescriptorManager::UpdateTexelBuffer(FRHIDescriptorHandle De
 
 		VkDescriptorDataEXT DescriptorData;
 		DescriptorData.pUniformTexelBuffer = &AddressInfo;  // same pointer for storage, it's a union
-		UpdateDescriptor(DescriptorHandle, DescriptorData, bImmediateUpdate);
+		UpdateDescriptor(DescriptorHandle, DescriptorType, DescriptorData, bImmediateUpdate);
 	}
 }
 
@@ -853,17 +846,15 @@ void FVulkanBindlessDescriptorManager::UpdateAccelerationStructure(FRHIDescripto
 {
 	if (bIsSupported)
 	{
-		check(GetDescriptorTypeForSetIndex(DescriptorHandle.GetRawType()) == VK_DESCRIPTOR_TYPE_ACCELERATION_STRUCTURE_KHR);
-
 		// :todo-jn: start caching AccelerationStructure in resources to avoid the extra call
 		VkAccelerationStructureDeviceAddressInfoKHR AccelerationStructureDeviceAddressInfo;
 		ZeroVulkanStruct(AccelerationStructureDeviceAddressInfo, VK_STRUCTURE_TYPE_ACCELERATION_STRUCTURE_DEVICE_ADDRESS_INFO_KHR);
 		AccelerationStructureDeviceAddressInfo.accelerationStructure = AccelerationStructure;
-		const VkDeviceAddress BufferAddress = VulkanRHI::vkGetAccelerationStructureDeviceAddressKHR(Device->GetInstanceHandle(), &AccelerationStructureDeviceAddressInfo);
+		const VkDeviceAddress BufferAddress = VulkanRHI::vkGetAccelerationStructureDeviceAddressKHR(Device.GetInstanceHandle(), &AccelerationStructureDeviceAddressInfo);
 
 		VkDescriptorDataEXT DescriptorData;
 		DescriptorData.accelerationStructure = BufferAddress;
-		UpdateDescriptor(DescriptorHandle, DescriptorData, bImmediateUpdate);
+		UpdateDescriptor(DescriptorHandle, VK_DESCRIPTOR_TYPE_ACCELERATION_STRUCTURE_KHR, DescriptorData, bImmediateUpdate);
 	}
 }
 
@@ -887,7 +878,7 @@ uint32 FVulkanBindlessDescriptorManager::GetFreeResourceIndex(FVulkanBindlessDes
 	INC_DWORD_STAT(STAT_VulkanBindlessPeakDescriptorCount);
 
 	const uint32 ResourceIndex = State.PeakDescriptorCount++;
-	checkf(ResourceIndex < State.MaxDescriptorCount, TEXT("You need to grow the resource array size for [%s]!"), VK_TYPE_TO_STRING(VkDescriptorType, State.DescriptorType));
+	checkf(ResourceIndex < State.MaxDescriptorCount, TEXT("You need to grow the resource array size for [%s]!"), VK_TYPE_TO_STRING(VkDescriptorType, State.DescriptorTypes[0]));
 	return ResourceIndex;
 }
 
diff --git a/Engine/Source/Runtime/VulkanRHI/Private/VulkanDescriptorSets.h b/Engine/Source/Runtime/VulkanRHI/Private/VulkanDescriptorSets.h
--- a/Engine/Source/Runtime/VulkanRHI/Private/VulkanDescriptorSets.h
+++ b/Engine/Source/Runtime/VulkanRHI/Private/VulkanDescriptorSets.h
@@ -846,10 +846,10 @@ protected:
 
 
 // Manager for resource descriptors used in bindless rendering.
-class FVulkanBindlessDescriptorManager : public VulkanRHI::FDeviceChild
+class FVulkanBindlessDescriptorManager
 {
 public:
-	FVulkanBindlessDescriptorManager(FVulkanDevice* InDevice);
+	FVulkanBindlessDescriptorManager(FVulkanDevice& InDevice);
 	~FVulkanBindlessDescriptorManager();
 
 	typedef TStaticArray<TArray<VkDescriptorAddressInfoEXT>, ShaderStage::MaxNumStages> FUniformBufferDescriptorArrays;
@@ -857,7 +857,7 @@ public:
 	void Init();
 	void Deinit();
 
-	static bool VerifySupport(FVulkanDevice* InDevice);
+	static bool VerifySupport(FVulkanDevice& InDevice);
 	inline bool IsSupported()
 	{
 		return bIsSupported;
@@ -873,10 +873,10 @@ public:
 	FRHIDescriptorHandle ReserveDescriptor(VkDescriptorType DescriptorType);
 
 	void UpdateSampler(FRHIDescriptorHandle DescriptorHandle, VkSampler VulkanSampler);
-	void UpdateImage(FRHIDescriptorHandle DescriptorHandle, VkImageView VulkanImage, bool bIsDepthStencil, bool bImmediateUpdate = true);
-	void UpdateBuffer(FRHIDescriptorHandle DescriptorHandle, VkBuffer VulkanBuffer, VkDeviceSize BufferOffset, VkDeviceSize BufferSize, bool bImmediateUpdate = true);
-	void UpdateBuffer(FRHIDescriptorHandle DescriptorHandle, VkDeviceAddress BufferAddress, VkDeviceSize BufferSize, bool bImmediateUpdate = true);
-	void UpdateTexelBuffer(FRHIDescriptorHandle DescriptorHandle, const VkBufferViewCreateInfo& ViewInfo, bool bImmediateUpdate = true);
+	void UpdateImage(FRHIDescriptorHandle DescriptorHandle, VkDescriptorType DescriptorType, VkImageView VulkanImage, bool bIsDepthStencil, bool bImmediateUpdate = true);
+	void UpdateBuffer(FRHIDescriptorHandle DescriptorHandle, VkDescriptorType DescriptorType, VkBuffer VulkanBuffer, VkDeviceSize BufferOffset, VkDeviceSize BufferSize, bool bImmediateUpdate = true);
+	void UpdateBuffer(FRHIDescriptorHandle DescriptorHandle, VkDescriptorType DescriptorType, VkDeviceAddress BufferAddress, VkDeviceSize BufferSize, bool bImmediateUpdate = true);
+	void UpdateTexelBuffer(FRHIDescriptorHandle DescriptorHandle, VkDescriptorType DescriptorType, const VkBufferViewCreateInfo& ViewInfo, bool bImmediateUpdate = true);
 	void UpdateAccelerationStructure(FRHIDescriptorHandle DescriptorHandle, VkAccelerationStructureKHR AccelerationStructure, bool bImmediateUpdate = true);
 
 	void RegisterUniformBuffers(FVulkanCommandBuffer& CommandBuffer, VkPipelineBindPoint BindPoint, const FUniformBufferDescriptorArrays& StageUBs);
@@ -886,13 +886,14 @@ public:
 	void UpdateUBAllocator();
 
 private:
+	FVulkanDevice& Device;
 	const bool bIsSupported;
 
 	VkDescriptorSetLayout EmptyDescriptorSetLayout = VK_NULL_HANDLE;
 
 	struct BindlessSetState
 	{
-		VkDescriptorType DescriptorType = VK_DESCRIPTOR_TYPE_MAX_ENUM;
+		TArray<VkDescriptorType> DescriptorTypes;
 		uint32 MaxDescriptorCount = 0;
 		std::atomic<uint32> PeakDescriptorCount = 1;  // always keep a null descriptor in slot 0
 		VkDescriptorSetLayout DescriptorSetLayout = VK_NULL_HANDLE;
@@ -918,8 +919,8 @@ private:
 	VkPipelineLayout BindlessPipelineLayout = VK_NULL_HANDLE;
 
 	uint32 GetFreeResourceIndex(BindlessSetState& Desc);
-	void UpdateStatsForHandle(FRHIDescriptorHandle DescriptorHandle);
-	void UpdateDescriptor(FRHIDescriptorHandle DescriptorHandle, VkDescriptorDataEXT DescriptorData, bool bImmediateUpdate);
+	void UpdateStatsForHandle(VkDescriptorType DescriptorType);
+	void UpdateDescriptor(FRHIDescriptorHandle DescriptorHandle, VkDescriptorType DescriptorType, VkDescriptorDataEXT DescriptorData, bool bImmediateUpdate);
 };
 
 
diff --git a/Engine/Source/Runtime/VulkanRHI/Private/VulkanDevice.cpp b/Engine/Source/Runtime/VulkanRHI/Private/VulkanDevice.cpp
--- a/Engine/Source/Runtime/VulkanRHI/Private/VulkanDevice.cpp
+++ b/Engine/Source/Runtime/VulkanRHI/Private/VulkanDevice.cpp
@@ -1388,7 +1388,7 @@ void FVulkanDevice::InitGPU()
 	DescriptorPoolsManager = new FVulkanDescriptorPoolsManager();
 	DescriptorPoolsManager->Init(this);
 
-	BindlessDescriptorManager = new FVulkanBindlessDescriptorManager(this);
+	BindlessDescriptorManager = new FVulkanBindlessDescriptorManager(*this);
 	BindlessDescriptorManager->Init();
 
 	PipelineStateCache = new FVulkanPipelineStateCacheManager(this);
diff --git a/Engine/Source/Runtime/VulkanRHI/Private/VulkanRayTracing.cpp b/Engine/Source/Runtime/VulkanRHI/Private/VulkanRayTracing.cpp
--- a/Engine/Source/Runtime/VulkanRHI/Private/VulkanRayTracing.cpp
+++ b/Engine/Source/Runtime/VulkanRHI/Private/VulkanRayTracing.cpp
@@ -554,8 +554,9 @@ void FVulkanRayTracingGeometry::SetupHitGroupSystemParameters()
 	{
 		if (Buffer)
 		{
-			FRHIDescriptorHandle BindlessHandle = BindlessDescriptorManager->ReserveDescriptor(VK_DESCRIPTOR_TYPE_STORAGE_BUFFER);
-			BindlessDescriptorManager->UpdateBuffer(BindlessHandle, Buffer->GetHandle(), Buffer->GetOffset() + ExtraOffset, Buffer->GetCurrentSize() - ExtraOffset);
+			const VkDescriptorType DescriptorType = VK_DESCRIPTOR_TYPE_STORAGE_BUFFER;
+			FRHIDescriptorHandle BindlessHandle = BindlessDescriptorManager->ReserveDescriptor(DescriptorType);
+			BindlessDescriptorManager->UpdateBuffer(BindlessHandle, DescriptorType, Buffer->GetHandle(), Buffer->GetOffset() + ExtraOffset, Buffer->GetCurrentSize() - ExtraOffset);
 			return BindlessHandle;
 		}
 		return FRHIDescriptorHandle();
diff --git a/Engine/Source/Runtime/VulkanRHI/Private/VulkanResources.h b/Engine/Source/Runtime/VulkanRHI/Private/VulkanResources.h
--- a/Engine/Source/Runtime/VulkanRHI/Private/VulkanResources.h
+++ b/Engine/Source/Runtime/VulkanRHI/Private/VulkanResources.h
@@ -522,9 +522,15 @@ public:
 		return BindlessHandle;
 	}
 
+	VkDescriptorType GetDescriptorType() const
+	{
+		return DescriptorType;
+	}
+
 private:
 	FVulkanDevice& Device;
 	FRHIDescriptorHandle BindlessHandle;
+	const VkDescriptorType DescriptorType;
 	TStorage Storage;
 };
 
diff --git a/Engine/Source/Runtime/VulkanRHI/Private/VulkanUAV.cpp b/Engine/Source/Runtime/VulkanRHI/Private/VulkanUAV.cpp
--- a/Engine/Source/Runtime/VulkanRHI/Private/VulkanUAV.cpp
+++ b/Engine/Source/Runtime/VulkanRHI/Private/VulkanUAV.cpp
@@ -9,6 +9,7 @@
 
 FVulkanView::FVulkanView(FVulkanDevice& InDevice, VkDescriptorType InDescriptorType)
 	: Device(InDevice)
+	, DescriptorType(InDescriptorType)
 {
 	BindlessHandle = Device.GetBindlessDescriptorManager()->ReserveDescriptor(InDescriptorType);
 }
@@ -106,7 +107,7 @@ FVulkanView* FVulkanView::InitAsTypedBufferView(FVulkanBuffer* Buffer, EPixelFor
 	INC_DWORD_STAT(STAT_VulkanNumBufferViews);
 	// :todo-jn: the buffer view is actually not needed in bindless anymore
 
-	Device.GetBindlessDescriptorManager()->UpdateTexelBuffer(BindlessHandle, ViewInfo, bImmediateUpdate);
+	Device.GetBindlessDescriptorManager()->UpdateTexelBuffer(BindlessHandle, DescriptorType, ViewInfo, bImmediateUpdate);
 
 	return this;
 }
@@ -213,7 +214,7 @@ FVulkanView* FVulkanView::InitAsTextureView(
 	}
 
 	const bool bDepthOrStencilAspect = (AspectFlags & (VK_IMAGE_ASPECT_DEPTH_BIT | VK_IMAGE_ASPECT_STENCIL_BIT)) != 0;
-	Device.GetBindlessDescriptorManager()->UpdateImage(BindlessHandle, TV.View, bDepthOrStencilAspect, bImmediateUpdate);
+	Device.GetBindlessDescriptorManager()->UpdateImage(BindlessHandle, DescriptorType, TV.View, bDepthOrStencilAspect, bImmediateUpdate);
 
 	return this;
 }
@@ -241,7 +242,7 @@ FVulkanView* FVulkanView::InitAsStructuredBufferView(FVulkanBuffer* Buffer, uint
 
 	SBV.Size = InSize;
 
-	Device.GetBindlessDescriptorManager()->UpdateBuffer(BindlessHandle, Buffer->GetHandle(), TotalOffset, InSize, bImmediateUpdate);
+	Device.GetBindlessDescriptorManager()->UpdateBuffer(BindlessHandle, DescriptorType, Buffer->GetHandle(), TotalOffset, InSize, bImmediateUpdate);
 
 	return this;
 }
diff --git a/Engine/Source/Runtime/VulkanRHI/Private/VulkanUniformBuffer.cpp b/Engine/Source/Runtime/VulkanRHI/Private/VulkanUniformBuffer.cpp
--- a/Engine/Source/Runtime/VulkanRHI/Private/VulkanUniformBuffer.cpp
+++ b/Engine/Source/Runtime/VulkanRHI/Private/VulkanUniformBuffer.cpp
@@ -246,7 +246,7 @@ FRHIDescriptorHandle FVulkanUniformBuffer::GetBindlessHandle()
 		}
 
 		BindlessHandle = Device->GetBindlessDescriptorManager()->ReserveDescriptor(VK_DESCRIPTOR_TYPE_UNIFORM_BUFFER);
-		Device->GetBindlessDescriptorManager()->UpdateBuffer(BindlessHandle, CurrentAddress, GetSize(), true);
+		Device->GetBindlessDescriptorManager()->UpdateBuffer(BindlessHandle, VK_DESCRIPTOR_TYPE_UNIFORM_BUFFER, CurrentAddress, GetSize(), true);
 
 		CachedDeviceAddress = CurrentAddress;
 	}
diff --git a/Engine/Source/Runtime/VulkanRHI/Public/VulkanCommon.h b/Engine/Source/Runtime/VulkanRHI/Public/VulkanCommon.h
--- a/Engine/Source/Runtime/VulkanRHI/Public/VulkanCommon.h
+++ b/Engine/Source/Runtime/VulkanRHI/Public/VulkanCommon.h
@@ -18,6 +18,13 @@
 // ONLY used for debugging binding table/descriptor set bugs/mismatches.
 #define VULKAN_ENABLE_BINDING_DEBUG_NAMES						0
 
+// Bindless uses one descriptor set per resource type plus one for single use UBs, for a total of 9.
+// On platforms that have a maximum below 9, use mutable descriptors to lower total count to 6 by
+// creating a descriptor set per resource type rather than per descriptor type.  On most platforms,
+// these descriptors have the same size and should still be packed tightly.
+// NOTE: Should invalidate Vulkan shaders when changing this value (for example, update VulkanCommon GUID).
+#define VULKAN_REDUCE_BINDLESS_DESCRIPTOR_SET_COUNT				0
+
 namespace ShaderStage
 {
 	// There should be one value for each value in EShaderFrequency.
@@ -100,6 +107,16 @@ namespace VulkanBindless
 	{
 		BindlessSamplerSet = 0,
 
+#if VULKAN_REDUCE_BINDLESS_DESCRIPTOR_SET_COUNT
+		BindlessStorageBufferSet = 1,
+		BindlessUniformBufferSet = 1,
+
+		BindlessStorageImageSet = 2,
+		BindlessSampledImageSet = 2,
+
+		BindlessStorageTexelBufferSet = 3,
+		BindlessUniformTexelBufferSet = 3,
+#else
 		BindlessStorageBufferSet,
 		BindlessUniformBufferSet,
 
@@ -108,6 +125,7 @@ namespace VulkanBindless
 
 		BindlessStorageTexelBufferSet,
 		BindlessUniformTexelBufferSet,
+#endif
 
 		BindlessAccelerationStructureSet,
 
